Fix crash on very early shutdown caused by GDynamicRHI still being NULL.

diff --git a/Engine/Source/Runtime/RHI/Private/DynamicRHI.cpp b/Engine/Source/Runtime/RHI/Private/DynamicRHI.cpp
--- a/Engine/Source/Runtime/RHI/Private/DynamicRHI.cpp
+++ b/Engine/Source/Runtime/RHI/Private/DynamicRHI.cpp
@@ -73,7 +73,7 @@ void RHIInit(bool bHasEditorToken)
 
 void RHIExit()
 {
-	if ( !GUsingNullRHI )
+	if ( !GUsingNullRHI && GDynamicRHI != NULL )
 	{
 		// Destruct the dynamic RHI.
 		GDynamicRHI->Shutdown();
