[GitHub] Added plugin extensibility for texture resources

You can now access raw platform-specific texture resources for D3D11 or OpenGL by calling GetNativeResource() on a texture object.

Thanks to Coherent for submitting this change.  https://github.com/EpicGames/UnrealEngine/pull/138

diff --git a/Engine/Source/Runtime/OpenGLDrv/Public/OpenGLResources.h b/Engine/Source/Runtime/OpenGLDrv/Public/OpenGLResources.h
--- a/Engine/Source/Runtime/OpenGLDrv/Public/OpenGLResources.h
+++ b/Engine/Source/Runtime/OpenGLDrv/Public/OpenGLResources.h
@@ -955,6 +955,13 @@ public:
 	bool IsCubemap() const { return bCubemap != 0; }
 	bool IsStaging() const { return (this->GetFlags() & TexCreate_CPUReadback) != 0; }
 
+
+	/** FRHITexture override.  See FRHITexture::GetNativeResource() */
+	virtual void* GetNativeResource() const OVERRIDE
+	{ 
+		return (void*)Resource;
+	}
+
 	/**
 	 * Accessors to mark whether or not we have allocated storage for each mip/face.
 	 * For non-cubemaps FaceIndex should always be zero.
diff --git a/Engine/Source/Runtime/RHI/Public/RHIResources.h b/Engine/Source/Runtime/RHI/Public/RHIResources.h
--- a/Engine/Source/Runtime/RHI/Public/RHIResources.h
+++ b/Engine/Source/Runtime/RHI/Public/RHIResources.h
@@ -154,6 +154,18 @@ public:
 	virtual class FRHITexture3D* GetTexture3D() { return NULL; }
 	virtual class FRHITextureCube* GetTextureCube() { return NULL; }
 	
+	/**
+	 * Returns access to the platform-specific native resource pointer.  This is designed to be used to provide plugins with access
+	 * to the underlying resource and should be used very carefully or not at all.
+	 *
+	 * @return	The pointer to the native resource or NULL if it not initialized or not supported for this resource type for some reason
+	 */
+	virtual void* GetNativeResource() const
+	{
+		// Override this in derived classes to expose access to the native texture resource
+		return nullptr;
+	}
+
 	/** @return The number of mip-maps in the texture. */
 	uint32 GetNumMips() const { return NumMips; }
 
diff --git a/Engine/Source/Runtime/Windows/D3D11RHI/Public/D3D11Resources.h b/Engine/Source/Runtime/Windows/D3D11RHI/Public/D3D11Resources.h
--- a/Engine/Source/Runtime/Windows/D3D11RHI/Public/D3D11Resources.h
+++ b/Engine/Source/Runtime/Windows/D3D11RHI/Public/D3D11Resources.h
@@ -374,6 +374,12 @@ public:
 	ID3D11Texture2D* GetResource() const { return (ID3D11Texture2D*)FD3D11TextureBase::GetResource(); }
 	bool IsCubemap() const { return bCubemap; }
 
+	/** FRHITexture override.  See FRHITexture::GetNativeResource() */
+	virtual void* GetNativeResource() const OVERRIDE
+	{ 
+		return GetResource();
+	}
+
 	// IRefCountedObject interface.
 	virtual uint32 AddRef() const
 	{
