* added CreateInfo to RHICreateVertexBuffer and RHICreateStructuredBuffer and RHICreateIndexBuffer calls for upcoming changes

diff --git a/Engine/Plugins/Experimental/Paper2D/Source/Paper2D/Private/PaperRenderSceneProxy.cpp b/Engine/Plugins/Experimental/Paper2D/Source/Paper2D/Private/PaperRenderSceneProxy.cpp
--- a/Engine/Plugins/Experimental/Paper2D/Source/Paper2D/Private/PaperRenderSceneProxy.cpp
+++ b/Engine/Plugins/Experimental/Paper2D/Source/Paper2D/Private/PaperRenderSceneProxy.cpp
@@ -35,7 +35,8 @@ class FPaperSpriteVertexBuffer : public FVertexBuffer
 public:
 	virtual void InitRHI() OVERRIDE
 	{
-		VertexBufferRHI = RHICreateVertexBuffer(sizeof(FPaperSpriteVertex), NULL, BUF_Static);
+		FRHIResourceCreateInfo CreateInfo;
+		VertexBufferRHI = RHICreateVertexBuffer(sizeof(FPaperSpriteVertex), BUF_Static, CreateInfo);
 	}
 };
 static TGlobalResource<FPaperSpriteVertexBuffer> GDummyMaterialSpriteVertexBuffer;
diff --git a/Engine/Plugins/Runtime/CableComponent/Source/CableComponent/Private/CableComponent.cpp b/Engine/Plugins/Runtime/CableComponent/Source/CableComponent/Private/CableComponent.cpp
--- a/Engine/Plugins/Runtime/CableComponent/Source/CableComponent/Private/CableComponent.cpp
+++ b/Engine/Plugins/Runtime/CableComponent/Source/CableComponent/Private/CableComponent.cpp
@@ -9,7 +9,8 @@ class FCableVertexBuffer : public FVertexBuffer
 public:
 	virtual void InitRHI()
 	{
-		VertexBufferRHI = RHICreateVertexBuffer(NumVerts * sizeof(FDynamicMeshVertex), NULL, BUF_Dynamic);
+		FRHIResourceCreateInfo CreateInfo;
+		VertexBufferRHI = RHICreateVertexBuffer(NumVerts * sizeof(FDynamicMeshVertex), BUF_Dynamic, CreateInfo);
 	}
 
 	int32 NumVerts;
@@ -21,7 +22,8 @@ class FCableIndexBuffer : public FIndexBuffer
 public:
 	virtual void InitRHI()
 	{
-		IndexBufferRHI = RHICreateIndexBuffer(sizeof(int32), NumIndices * sizeof(int32), NULL, BUF_Dynamic);
+		FRHIResourceCreateInfo CreateInfo;
+		IndexBufferRHI = RHICreateIndexBuffer(sizeof(int32), NumIndices * sizeof(int32), BUF_Dynamic, CreateInfo);
 	}
 
 	int32 NumIndices;
diff --git a/Engine/Plugins/Runtime/CustomMeshComponent/Source/CustomMeshComponent/Private/CustomMeshComponent.cpp b/Engine/Plugins/Runtime/CustomMeshComponent/Source/CustomMeshComponent/Private/CustomMeshComponent.cpp
--- a/Engine/Plugins/Runtime/CustomMeshComponent/Source/CustomMeshComponent/Private/CustomMeshComponent.cpp
+++ b/Engine/Plugins/Runtime/CustomMeshComponent/Source/CustomMeshComponent/Private/CustomMeshComponent.cpp
@@ -12,7 +12,8 @@ public:
 
 	virtual void InitRHI()
 	{
-		VertexBufferRHI = RHICreateVertexBuffer(Vertices.Num() * sizeof(FDynamicMeshVertex),NULL,BUF_Static);
+		FRHIResourceCreateInfo CreateInfo;
+		VertexBufferRHI = RHICreateVertexBuffer(Vertices.Num() * sizeof(FDynamicMeshVertex),BUF_Static,CreateInfo);
 
 		// Copy the vertex data into the vertex buffer.
 		void* VertexBufferData = RHILockVertexBuffer(VertexBufferRHI,0,Vertices.Num() * sizeof(FDynamicMeshVertex), RLM_WriteOnly);
@@ -30,7 +31,8 @@ public:
 
 	virtual void InitRHI()
 	{
-		IndexBufferRHI = RHICreateIndexBuffer(sizeof(int32),Indices.Num() * sizeof(int32),NULL,BUF_Static);
+		FRHIResourceCreateInfo CreateInfo;
+		IndexBufferRHI = RHICreateIndexBuffer(sizeof(int32),Indices.Num() * sizeof(int32),BUF_Static, CreateInfo);
 
 		// Write the indices to the index buffer.
 		void* Buffer = RHILockIndexBuffer(IndexBufferRHI,0,Indices.Num() * sizeof(int32),RLM_WriteOnly);
diff --git a/Engine/Source/Runtime/EmptyRHI/Private/EmptyIndexBuffer.cpp b/Engine/Source/Runtime/EmptyRHI/Private/EmptyIndexBuffer.cpp
--- a/Engine/Source/Runtime/EmptyRHI/Private/EmptyIndexBuffer.cpp
+++ b/Engine/Source/Runtime/EmptyRHI/Private/EmptyIndexBuffer.cpp
@@ -23,25 +23,25 @@ void FEmptyIndexBuffer::Unlock()
 
 }
 
-FIndexBufferRHIRef FEmptyDynamicRHI::RHICreateIndexBuffer(uint32 Stride, uint32 Size, FResourceArrayInterface* ResourceArray, uint32 InUsage)
+FIndexBufferRHIRef FEmptyDynamicRHI::RHICreateIndexBuffer(uint32 Stride, uint32 Size, uint32 InUsage, FRHIResourceCreateInfo& CreateInfo)
 {
 	// make the RHI object, which will allocate memory
 	FEmptyIndexBuffer* IndexBuffer = new FEmptyIndexBuffer(Stride, Size, InUsage);
 	
-	if(ResourceArray)
+	if(CreateInfo.ResourceArray)
 	{
-		check(Size == ResourceArray->GetResourceDataSize());
+		check(Size == CreateInfo.ResourceArray->GetResourceDataSize());
 
 		// make a buffer usable by CPU
 		void* Buffer = RHILockIndexBuffer(IndexBuffer, 0, Size, RLM_WriteOnly);
 
 		// copy the contents of the given data into the buffer
-		FMemory::Memcpy(Buffer, ResourceArray->GetResourceData(), Size);
+		FMemory::Memcpy(Buffer, CreateInfo.ResourceArray->GetResourceData(), Size);
 
 		RHIUnlockIndexBuffer(IndexBuffer);
 
 		// Discard the resource array's contents.
-		ResourceArray->Discard();
+		CreateInfo.ResourceArray->Discard();
 	}
 
 	return IndexBuffer;
diff --git a/Engine/Source/Runtime/EmptyRHI/Private/EmptyVertexBuffer.cpp b/Engine/Source/Runtime/EmptyRHI/Private/EmptyVertexBuffer.cpp
--- a/Engine/Source/Runtime/EmptyRHI/Private/EmptyVertexBuffer.cpp
+++ b/Engine/Source/Runtime/EmptyRHI/Private/EmptyVertexBuffer.cpp
@@ -22,25 +22,25 @@ void FEmptyVertexBuffer::Unlock()
 	
 }
 
-FVertexBufferRHIRef FEmptyDynamicRHI::RHICreateVertexBuffer(uint32 Size, FResourceArrayInterface* ResourceArray, uint32 InUsage)
+FVertexBufferRHIRef FEmptyDynamicRHI::RHICreateVertexBuffer(uint32 Size, uint32 InUsage, FRHIResourceCreateInfo& CreateInfo)
 {
 	// make the RHI object, which will allocate memory
 	FEmptyVertexBuffer* VertexBuffer = new FEmptyVertexBuffer(Size, InUsage);
 
-	if(ResourceArray)
+	if(CreateInfo.ResourceArray)
 	{
-		check(Size == ResourceArray->GetResourceDataSize());
+		check(Size == CreateInfo.ResourceArray->GetResourceDataSize());
 
 		// make a buffer usable by CPU
 		void* Buffer = RHILockVertexBuffer(VertexBuffer, 0, Size, RLM_WriteOnly);
 
 		// copy the contents of the given data into the buffer
-		FMemory::Memcpy(Buffer, ResourceArray->GetResourceData(), Size);
+		FMemory::Memcpy(Buffer, CreateInfo.ResourceArray->GetResourceData(), Size);
 
 		RHIUnlockVertexBuffer(VertexBuffer);
 
 		// Discard the resource array's contents.
-		ResourceArray->Discard();
+		CreateInfo.ResourceArray->Discard();
 	}
 
 	return VertexBuffer;
diff --git a/Engine/Source/Runtime/Engine/Private/BrushComponent.cpp b/Engine/Source/Runtime/Engine/Private/BrushComponent.cpp
--- a/Engine/Source/Runtime/Engine/Private/BrushComponent.cpp
+++ b/Engine/Source/Runtime/Engine/Private/BrushComponent.cpp
@@ -41,7 +41,8 @@ public:
 	{
 		if(NumVertices)
 		{
-			VertexBufferRHI = RHICreateVertexBuffer(NumVertices * sizeof(FModelWireVertex),NULL,BUF_Static);
+			FRHIResourceCreateInfo CreateInfo;
+			VertexBufferRHI = RHICreateVertexBuffer(NumVertices * sizeof(FModelWireVertex),BUF_Static, CreateInfo);
 
 			FModelWireVertex* DestVertex = (FModelWireVertex*)RHILockVertexBuffer(VertexBufferRHI,0,NumVertices * sizeof(FModelWireVertex),RLM_WriteOnly);
 			for(int32 PolyIndex = 0;PolyIndex < Polys.Num();PolyIndex++)
@@ -91,7 +92,8 @@ public:
 	{
 		if(NumEdges)
 		{
-			IndexBufferRHI = RHICreateIndexBuffer(sizeof(uint16),NumEdges * 2 * sizeof(uint16),NULL,BUF_Static);
+			FRHIResourceCreateInfo CreateInfo;
+			IndexBufferRHI = RHICreateIndexBuffer(sizeof(uint16),NumEdges * 2 * sizeof(uint16),BUF_Static, CreateInfo);
 
 			uint16* DestIndex = (uint16*)RHILockIndexBuffer(IndexBufferRHI,0,NumEdges * 2 * sizeof(uint16),RLM_WriteOnly);
 			uint16 BaseIndex = 0;
diff --git a/Engine/Source/Runtime/Engine/Private/ColorVertexBuffer.cpp b/Engine/Source/Runtime/Engine/Private/ColorVertexBuffer.cpp
--- a/Engine/Source/Runtime/Engine/Private/ColorVertexBuffer.cpp
+++ b/Engine/Source/Runtime/Engine/Private/ColorVertexBuffer.cpp
@@ -357,7 +357,8 @@ void FColorVertexBuffer::InitRHI()
 		if(ResourceArray->GetResourceDataSize())
 		{
 			// Create the vertex buffer.
-			VertexBufferRHI = RHICreateVertexBuffer(ResourceArray->GetResourceDataSize(),ResourceArray,BUF_Static);
+			FRHIResourceCreateInfo CreateInfo(ResourceArray);
+			VertexBufferRHI = RHICreateVertexBuffer(ResourceArray->GetResourceDataSize(),BUF_Static,CreateInfo);
 		}
 	}
 }
diff --git a/Engine/Source/Runtime/Engine/Private/DynamicMeshBuilder.cpp b/Engine/Source/Runtime/Engine/Private/DynamicMeshBuilder.cpp
--- a/Engine/Source/Runtime/Engine/Private/DynamicMeshBuilder.cpp
+++ b/Engine/Source/Runtime/Engine/Private/DynamicMeshBuilder.cpp
@@ -16,7 +16,8 @@ public:
 	// FRenderResource interface.
 	virtual void InitRHI()
 	{
-		IndexBufferRHI = RHICreateIndexBuffer(sizeof(int32),Indices.Num() * sizeof(int32),NULL,BUF_Static);
+		FRHIResourceCreateInfo CreateInfo;
+		IndexBufferRHI = RHICreateIndexBuffer(sizeof(int32),Indices.Num() * sizeof(int32),BUF_Static,CreateInfo);
 
 		// Write the indices to the index buffer.
 		void* Buffer = RHILockIndexBuffer(IndexBufferRHI,0,Indices.Num() * sizeof(int32),RLM_WriteOnly);
@@ -46,7 +47,8 @@ public:
 	// FResourceResource interface.
 	virtual void InitRHI()
 	{
-		VertexBufferRHI = RHICreateVertexBuffer(Vertices.Num() * sizeof(FDynamicMeshVertex),NULL,BUF_Static);
+		FRHIResourceCreateInfo CreateInfo;
+		VertexBufferRHI = RHICreateVertexBuffer(Vertices.Num() * sizeof(FDynamicMeshVertex),BUF_Static, CreateInfo);
 
 		// Copy the vertex data into the vertex buffer.
 		void* VertexBufferData = RHILockVertexBuffer(VertexBufferRHI,0,Vertices.Num() * sizeof(FDynamicMeshVertex),RLM_WriteOnly);
diff --git a/Engine/Source/Runtime/Engine/Private/GPUSkinVertexFactory.cpp b/Engine/Source/Runtime/Engine/Private/GPUSkinVertexFactory.cpp
--- a/Engine/Source/Runtime/Engine/Private/GPUSkinVertexFactory.cpp
+++ b/Engine/Source/Runtime/Engine/Private/GPUSkinVertexFactory.cpp
@@ -86,7 +86,8 @@ FBoneBufferTypeRef FBoneBufferPoolPolicy::CreateResource(CreationArguments Args)
 {
 	uint32 BufferSize = GetPoolBucketSize(GetPoolBucketIndex(Args));
     FBoneBuffer Buffer;
-    Buffer.VertexBufferRHI = RHICreateVertexBuffer( BufferSize, NULL, (BUF_Dynamic | BUF_ShaderResource) );
+	FRHIResourceCreateInfo CreateInfo;
+    Buffer.VertexBufferRHI = RHICreateVertexBuffer( BufferSize, (BUF_Dynamic | BUF_ShaderResource), CreateInfo );
     Buffer.VertexBufferSRV = RHICreateShaderResourceView( Buffer.VertexBufferRHI, sizeof(FVector4), PF_A32B32G32R32F );
 	return Buffer;
 }
diff --git a/Engine/Source/Runtime/Engine/Private/GPUSort.cpp b/Engine/Source/Runtime/Engine/Private/GPUSort.cpp
--- a/Engine/Source/Runtime/Engine/Private/GPUSort.cpp
+++ b/Engine/Source/Runtime/Engine/Private/GPUSort.cpp
@@ -87,10 +87,11 @@ public:
 		{
 			for (int32 BufferIndex = 0; BufferIndex < 2; ++BufferIndex)
 			{
+				FRHIResourceCreateInfo CreateInfo;
 				Buffers[BufferIndex] = RHICreateVertexBuffer(
 					OffsetsBufferSize,
-					/*ResourceArray=*/ NULL,
-					BUF_Static | BUF_ShaderResource | BUF_UnorderedAccess );
+					BUF_Static | BUF_ShaderResource | BUF_UnorderedAccess,
+					CreateInfo);
 				BufferSRVs[BufferIndex] = RHICreateShaderResourceView(
 					Buffers[BufferIndex],
 					/*Stride=*/ sizeof(uint32),
@@ -182,11 +183,11 @@ public:
 	{
 		if (GRHIFeatureLevel == ERHIFeatureLevel::SM5)
 		{
+			FRHIResourceCreateInfo CreateInfo;
 			SortParametersBufferRHI = RHICreateVertexBuffer(
 				/*Size=*/ sizeof(FRadixSortParameters),
-				/*ResourceArray=*/ NULL,
-				/*Usage=*/ BUF_Volatile | BUF_ShaderResource
-				);
+				/*Usage=*/ BUF_Volatile | BUF_ShaderResource,
+				CreateInfo);
 			SortParametersBufferSRV = RHICreateShaderResourceView(
 				SortParametersBufferRHI, /*Stride=*/ sizeof(uint32), PF_R32_UINT 
 				);
@@ -863,10 +864,11 @@ static bool RunGPUSortTest(int32 TestSize)
 	// Allocate GPU resources.
 	for (int32 BufferIndex = 0; BufferIndex < 2; ++BufferIndex)
 	{
-		KeysBufferRHI[BufferIndex] = RHICreateVertexBuffer(BufferSize, /*ResourceArray=*/ NULL, BUF_Static|BUF_ShaderResource|BUF_UnorderedAccess);
+		FRHIResourceCreateInfo CreateInfo;
+		KeysBufferRHI[BufferIndex] = RHICreateVertexBuffer(BufferSize, BUF_Static|BUF_ShaderResource|BUF_UnorderedAccess, CreateInfo);
 		KeysBufferSRV[BufferIndex] = RHICreateShaderResourceView(KeysBufferRHI[BufferIndex], /*Stride=*/ sizeof(uint32), PF_R32_UINT);
 		KeysBufferUAV[BufferIndex] = RHICreateUnorderedAccessView(KeysBufferRHI[BufferIndex], PF_R32_UINT);
-		ValuesBufferRHI[BufferIndex] = RHICreateVertexBuffer(BufferSize, /*ResourceArray=*/ NULL, BUF_Static|BUF_ShaderResource|BUF_UnorderedAccess);
+		ValuesBufferRHI[BufferIndex] = RHICreateVertexBuffer(BufferSize, BUF_Static|BUF_ShaderResource|BUF_UnorderedAccess, CreateInfo);
 		ValuesBufferSRV[BufferIndex] = RHICreateShaderResourceView(ValuesBufferRHI[BufferIndex], /*Stride=*/ sizeof(uint32), PF_R32_UINT);
 		ValuesBufferUAV[BufferIndex] = RHICreateUnorderedAccessView(ValuesBufferRHI[BufferIndex], PF_R32_UINT);
 	}
diff --git a/Engine/Source/Runtime/Engine/Private/InstancedStaticMesh.cpp b/Engine/Source/Runtime/Engine/Private/InstancedStaticMesh.cpp
--- a/Engine/Source/Runtime/Engine/Private/InstancedStaticMesh.cpp
+++ b/Engine/Source/Runtime/Engine/Private/InstancedStaticMesh.cpp
@@ -279,7 +279,8 @@ void FStaticMeshInstanceBuffer::InitRHI()
 	if(ResourceArray->GetResourceDataSize())
 	{
 		// Create the vertex buffer.
-		VertexBufferRHI = RHICreateVertexBuffer(ResourceArray->GetResourceDataSize(),ResourceArray,BUF_Static);
+		FRHIResourceCreateInfo CreateInfo(ResourceArray);
+		VertexBufferRHI = RHICreateVertexBuffer(ResourceArray->GetResourceDataSize(),BUF_Static, CreateInfo);
 	}
 }
 
diff --git a/Engine/Source/Runtime/Engine/Private/Landscape/LandscapeRender.cpp b/Engine/Source/Runtime/Engine/Private/Landscape/LandscapeRender.cpp
--- a/Engine/Source/Runtime/Engine/Private/Landscape/LandscapeRender.cpp
+++ b/Engine/Source/Runtime/Engine/Private/Landscape/LandscapeRender.cpp
@@ -1573,7 +1573,8 @@ void FLandscapeComponentSceneProxy::DrawDynamicElements(FPrimitiveDrawInterface*
 void FLandscapeVertexBuffer::InitRHI()
 {
 	// create a static vertex buffer
-	VertexBufferRHI = RHICreateVertexBuffer(FMath::Square(SubsectionSizeVerts) * FMath::Square(NumSubsections) * sizeof(FLandscapeVertex), NULL, BUF_Static);
+	FRHIResourceCreateInfo CreateInfo;
+	VertexBufferRHI = RHICreateVertexBuffer(FMath::Square(SubsectionSizeVerts) * FMath::Square(NumSubsections) * sizeof(FLandscapeVertex), BUF_Static, CreateInfo);
 	FLandscapeVertex* Vertex = (FLandscapeVertex*)RHILockVertexBuffer(VertexBufferRHI, 0, FMath::Square(SubsectionSizeVerts) * FMath::Square(NumSubsections) * sizeof(FLandscapeVertex),RLM_WriteOnly);
 
 	for( int32 SubY=0;SubY<NumSubsections;SubY++ )
diff --git a/Engine/Source/Runtime/Engine/Private/Landscape/LandscapeRenderMobile.cpp b/Engine/Source/Runtime/Engine/Private/Landscape/LandscapeRenderMobile.cpp
--- a/Engine/Source/Runtime/Engine/Private/Landscape/LandscapeRenderMobile.cpp
+++ b/Engine/Source/Runtime/Engine/Private/Landscape/LandscapeRenderMobile.cpp
@@ -157,7 +157,8 @@ IMPLEMENT_VERTEX_FACTORY_TYPE(FLandscapeVertexFactoryMobile, "LandscapeVertexFac
 void FLandscapeVertexBufferMobile::InitRHI()
 {
 	// create a static vertex buffer
-	VertexBufferRHI = RHICreateVertexBuffer(DataSize, NULL, BUF_Static);
+	FRHIResourceCreateInfo CreateInfo;
+	VertexBufferRHI = RHICreateVertexBuffer(DataSize, BUF_Static, CreateInfo);
 	void* VertexData = RHILockVertexBuffer(VertexBufferRHI, 0, DataSize, RLM_WriteOnly);
 	// Copy stored platform data
 	FMemory::Memcpy(VertexData, (uint8*)Data, DataSize);
diff --git a/Engine/Source/Runtime/Engine/Private/MaterialBillboardComponent.cpp b/Engine/Source/Runtime/Engine/Private/MaterialBillboardComponent.cpp
--- a/Engine/Source/Runtime/Engine/Private/MaterialBillboardComponent.cpp
+++ b/Engine/Source/Runtime/Engine/Private/MaterialBillboardComponent.cpp
@@ -21,7 +21,8 @@ public:
 
 	virtual void InitRHI()
 	{
-		VertexBufferRHI = RHICreateVertexBuffer(sizeof(FMaterialSpriteVertex),NULL,BUF_Static);
+		FRHIResourceCreateInfo CreateInfo;
+		VertexBufferRHI = RHICreateVertexBuffer(sizeof(FMaterialSpriteVertex),BUF_Static,CreateInfo);
 	}
 };
 static TGlobalResource<FMaterialSpriteVertexBuffer> GDummyMaterialSpriteVertexBuffer;
diff --git a/Engine/Source/Runtime/Engine/Private/ModelRender.cpp b/Engine/Source/Runtime/Engine/Private/ModelRender.cpp
--- a/Engine/Source/Runtime/Engine/Private/ModelRender.cpp
+++ b/Engine/Source/Runtime/Engine/Private/ModelRender.cpp
@@ -41,7 +41,8 @@ void FModelVertexBuffer::InitRHI()
 	if( Size > 0 )
 	{
 		// Create the buffer.
-		VertexBufferRHI = RHICreateVertexBuffer(Size,&Vertices,BUF_Static);
+		FRHIResourceCreateInfo CreateInfo(&Vertices);
+		VertexBufferRHI = RHICreateVertexBuffer(Size, BUF_Static, CreateInfo);
 	}
 }
 
diff --git a/Engine/Source/Runtime/Engine/Private/Particles/ParticleGpuSimulation.cpp b/Engine/Source/Runtime/Engine/Private/Particles/ParticleGpuSimulation.cpp
--- a/Engine/Source/Runtime/Engine/Private/Particles/ParticleGpuSimulation.cpp
+++ b/Engine/Source/Runtime/Engine/Private/Particles/ParticleGpuSimulation.cpp
@@ -1993,10 +1993,11 @@ static FBox ComputeParticleBounds(
 
 		// Create a buffer for storing bounds.
 		const int32 BufferSize = GroupCount * 2 * sizeof(FVector4);
+		FRHIResourceCreateInfo CreateInfo;
 		FVertexBufferRHIRef BoundsVertexBufferRHI = RHICreateVertexBuffer(
 			BufferSize,
-			/*ResourceArray=*/ NULL,
-			BUF_Static | BUF_UnorderedAccess );
+			BUF_Static | BUF_UnorderedAccess,
+			CreateInfo);
 		FUnorderedAccessViewRHIRef BoundsVertexBufferUAV = RHICreateUnorderedAccessView(
 			BoundsVertexBufferRHI,
 			PF_A32B32G32R32F );
@@ -2126,7 +2127,8 @@ public:
 		{
 			const int32 TileBufferSize = AlignedTileCount * sizeof(FVector2D);
 			check(TileBufferSize > 0);
-			VertexBufferRHI = RHICreateVertexBuffer( TileBufferSize, /*ResourceArray=*/ NULL, BUF_Static | BUF_KeepCPUAccessible | BUF_ShaderResource );
+			FRHIResourceCreateInfo CreateInfo;
+			VertexBufferRHI = RHICreateVertexBuffer( TileBufferSize, BUF_Static | BUF_KeepCPUAccessible | BUF_ShaderResource, CreateInfo );
 			VertexBufferSRV = RHICreateShaderResourceView( VertexBufferRHI, /*Stride=*/ sizeof(FVector2D), PF_G32R32F );
 		}
 	}
@@ -2181,7 +2183,8 @@ public:
 			const int32 BufferStride = sizeof(FParticleIndex);
 			const int32 BufferSize = ParticleCount * BufferStride;
 			uint32 Flags = BUF_Static | /*BUF_KeepCPUAccessible | */BUF_ShaderResource;
-			VertexBufferRHI = RHICreateVertexBuffer(BufferSize, /*ResourceArray=*/ NULL, Flags);
+			FRHIResourceCreateInfo CreateInfo;
+			VertexBufferRHI = RHICreateVertexBuffer(BufferSize, Flags, CreateInfo);
 			VertexBufferSRV = RHICreateShaderResourceView(VertexBufferRHI, BufferStride, PF_G16R16F);
 		}
 	}
diff --git a/Engine/Source/Runtime/Engine/Private/Particles/ParticleResources.cpp b/Engine/Source/Runtime/Engine/Private/Particles/ParticleResources.cpp
--- a/Engine/Source/Runtime/Engine/Private/Particles/ParticleResources.cpp
+++ b/Engine/Source/Runtime/Engine/Private/Particles/ParticleResources.cpp
@@ -16,7 +16,8 @@ const int32 GParticleScratchVertexBufferSize = 64 * (1 << 10); // 64KB
 void FParticleTexCoordVertexBuffer::InitRHI()
 {
 	const uint32 Size = sizeof(FVector2D) * 4 * MAX_PARTICLES_PER_INSTANCE;
-	VertexBufferRHI = RHICreateVertexBuffer( Size, NULL, BUF_Static );
+	FRHIResourceCreateInfo CreateInfo;
+	VertexBufferRHI = RHICreateVertexBuffer( Size, BUF_Static, CreateInfo);
 	FVector2D* Vertices = (FVector2D*)RHILockVertexBuffer( VertexBufferRHI, 0, Size, RLM_WriteOnly );
 	for (uint32 SpriteIndex = 0; SpriteIndex < MAX_PARTICLES_PER_INSTANCE; ++SpriteIndex)
 	{
@@ -42,7 +43,8 @@ void FParticleIndexBuffer::InitRHI()
 	const uint32 MaxParticles = 65536 / 4;
 	const uint32 Size = sizeof(uint16) * 6 * MaxParticles;
 	const uint32 Stride = sizeof(uint16);
-	IndexBufferRHI = RHICreateIndexBuffer( Stride, Size, NULL, BUF_Static );
+	FRHIResourceCreateInfo CreateInfo;
+	IndexBufferRHI = RHICreateIndexBuffer( Stride, Size, BUF_Static, CreateInfo);
 	uint16* Indices = (uint16*)RHILockIndexBuffer( IndexBufferRHI, 0, Size, RLM_WriteOnly );
 	for (uint32 SpriteIndex = 0; SpriteIndex < MaxParticles; ++SpriteIndex)
 	{
@@ -70,7 +72,8 @@ void FParticleScratchVertexBuffer::InitRHI()
 	{
 		Flags |= BUF_ShaderResource;
 	}
-	VertexBufferRHI = RHICreateVertexBuffer(GParticleScratchVertexBufferSize, NULL, Flags);
+	FRHIResourceCreateInfo CreateInfo;
+	VertexBufferRHI = RHICreateVertexBuffer(GParticleScratchVertexBufferSize, Flags, CreateInfo);
 	if (GRHIFeatureLevel >= ERHIFeatureLevel::SM4)
 	{
 		VertexBufferSRV_G32R32F = RHICreateShaderResourceView( VertexBufferRHI, /*Stride=*/ sizeof(FVector2D), PF_G32R32F );
diff --git a/Engine/Source/Runtime/Engine/Private/Particles/ParticleSortingGPU.cpp b/Engine/Source/Runtime/Engine/Private/Particles/ParticleSortingGPU.cpp
--- a/Engine/Source/Runtime/Engine/Private/Particles/ParticleSortingGPU.cpp
+++ b/Engine/Source/Runtime/Engine/Private/Particles/ParticleSortingGPU.cpp
@@ -252,11 +252,13 @@ void FParticleSortBuffers::InitRHI()
 	{
 		for (int32 BufferIndex = 0; BufferIndex < 2; ++BufferIndex)
 		{
-			KeyBuffers[BufferIndex] = RHICreateVertexBuffer( BufferSize * sizeof(uint32), /*ResourceArray=*/ NULL, BUF_Static | BUF_ShaderResource | BUF_UnorderedAccess );
+			FRHIResourceCreateInfo CreateInfo;
+
+			KeyBuffers[BufferIndex] = RHICreateVertexBuffer( BufferSize * sizeof(uint32), BUF_Static | BUF_ShaderResource | BUF_UnorderedAccess, CreateInfo);
 			KeyBufferSRVs[BufferIndex] = RHICreateShaderResourceView( KeyBuffers[BufferIndex], /*Stride=*/ sizeof(uint32), PF_R32_UINT );
 			KeyBufferUAVs[BufferIndex] = RHICreateUnorderedAccessView( KeyBuffers[BufferIndex], PF_R32_UINT );
 
-			VertexBuffers[BufferIndex] = RHICreateVertexBuffer( BufferSize * sizeof(uint32), /*ResourceArray=*/ NULL, BUF_Static | BUF_ShaderResource | BUF_UnorderedAccess );
+			VertexBuffers[BufferIndex] = RHICreateVertexBuffer( BufferSize * sizeof(uint32), BUF_Static | BUF_ShaderResource | BUF_UnorderedAccess, CreateInfo);
 			VertexBufferSRVs[BufferIndex] = RHICreateShaderResourceView( VertexBuffers[BufferIndex], /*Stride=*/ sizeof(FFloat16)*2, PF_G16R16F );
 			VertexBufferUAVs[BufferIndex] = RHICreateUnorderedAccessView( VertexBuffers[BufferIndex], PF_G16R16F );
 			VertexBufferSortSRVs[BufferIndex] = RHICreateShaderResourceView( VertexBuffers[BufferIndex], /*Stride=*/ sizeof(uint32), PF_R32_UINT );
diff --git a/Engine/Source/Runtime/Engine/Private/PhysicsEngine/PhysDrawing.cpp b/Engine/Source/Runtime/Engine/Private/PhysicsEngine/PhysDrawing.cpp
--- a/Engine/Source/Runtime/Engine/Private/PhysicsEngine/PhysDrawing.cpp
+++ b/Engine/Source/Runtime/Engine/Private/PhysicsEngine/PhysDrawing.cpp
@@ -359,7 +359,8 @@ void FKConvexElem::AddCachedSolidConvexGeom(TArray<FDynamicMeshVertex>& VertexBu
 
 void FConvexCollisionVertexBuffer::InitRHI()
 {
-	VertexBufferRHI = RHICreateVertexBuffer(Vertices.Num() * sizeof(FDynamicMeshVertex),NULL,BUF_Static);
+	FRHIResourceCreateInfo CreateInfo;
+	VertexBufferRHI = RHICreateVertexBuffer(Vertices.Num() * sizeof(FDynamicMeshVertex),BUF_Static, CreateInfo);
 
 	// Copy the vertex data into the vertex buffer.
 	void* VertexBufferData = RHILockVertexBuffer(VertexBufferRHI,0,Vertices.Num() * sizeof(FDynamicMeshVertex), RLM_WriteOnly);
@@ -369,7 +370,8 @@ void FConvexCollisionVertexBuffer::InitRHI()
 
 void FConvexCollisionIndexBuffer::InitRHI()
 {
-	IndexBufferRHI = RHICreateIndexBuffer(sizeof(int32),Indices.Num() * sizeof(int32),NULL,BUF_Static);
+	FRHIResourceCreateInfo CreateInfo;
+	IndexBufferRHI = RHICreateIndexBuffer(sizeof(int32),Indices.Num() * sizeof(int32),BUF_Static,CreateInfo);
 
 	// Write the indices to the index buffer.
 	void* Buffer = RHILockIndexBuffer(IndexBufferRHI,0,Indices.Num() * sizeof(int32),RLM_WriteOnly);
diff --git a/Engine/Source/Runtime/Engine/Private/PositionVertexBuffer.cpp b/Engine/Source/Runtime/Engine/Private/PositionVertexBuffer.cpp
--- a/Engine/Source/Runtime/Engine/Private/PositionVertexBuffer.cpp
+++ b/Engine/Source/Runtime/Engine/Private/PositionVertexBuffer.cpp
@@ -153,7 +153,8 @@ void FPositionVertexBuffer::InitRHI()
 	if(ResourceArray->GetResourceDataSize())
 	{
 		// Create the vertex buffer.
-		VertexBufferRHI = RHICreateVertexBuffer(ResourceArray->GetResourceDataSize(),ResourceArray,BUF_Static);
+		FRHIResourceCreateInfo CreateInfo(ResourceArray);
+		VertexBufferRHI = RHICreateVertexBuffer(ResourceArray->GetResourceDataSize(),BUF_Static,CreateInfo);
 	}
 }
 
diff --git a/Engine/Source/Runtime/Engine/Private/RawIndexBuffer.cpp b/Engine/Source/Runtime/Engine/Private/RawIndexBuffer.cpp
--- a/Engine/Source/Runtime/Engine/Private/RawIndexBuffer.cpp
+++ b/Engine/Source/Runtime/Engine/Private/RawIndexBuffer.cpp
@@ -38,7 +38,8 @@ void FRawIndexBuffer::InitRHI()
 	if( Size > 0 )
 	{
 		// Create the index buffer.
-		IndexBufferRHI = RHICreateIndexBuffer(sizeof(uint16),Size,NULL,BUF_Static);
+		FRHIResourceCreateInfo CreateInfo;
+		IndexBufferRHI = RHICreateIndexBuffer(sizeof(uint16),Size,BUF_Static,CreateInfo);
 
 		// Initialize the buffer.
 		void* Buffer = RHILockIndexBuffer(IndexBufferRHI,0,Size,RLM_WriteOnly);
@@ -76,7 +77,8 @@ void FRawIndexBuffer16or32::InitRHI()
 	if( Size > 0 )
 	{
 		// Create the index buffer.
-		IndexBufferRHI = RHICreateIndexBuffer(sizeof(uint32),Size,NULL,BUF_Static);
+		FRHIResourceCreateInfo CreateInfo;
+		IndexBufferRHI = RHICreateIndexBuffer(sizeof(uint32),Size,BUF_Static,CreateInfo);
 
 		// Initialize the buffer.
 		void* Buffer = RHILockIndexBuffer(IndexBufferRHI,0,Size,RLM_WriteOnly);
@@ -197,7 +199,8 @@ void FRawStaticIndexBuffer::InitRHI()
 	if (SizeInBytes > 0)
 	{
 		// Create the index buffer.
-		IndexBufferRHI = RHICreateIndexBuffer(IndexStride,SizeInBytes,&IndexStorage,BUF_Static);
+		FRHIResourceCreateInfo CreateInfo(&IndexStorage);
+		IndexBufferRHI = RHICreateIndexBuffer(IndexStride,SizeInBytes,BUF_Static,CreateInfo);
 	}    
 }
 
diff --git a/Engine/Source/Runtime/Engine/Private/SkeletalMesh.cpp b/Engine/Source/Runtime/Engine/Private/SkeletalMesh.cpp
--- a/Engine/Source/Runtime/Engine/Private/SkeletalMesh.cpp
+++ b/Engine/Source/Runtime/Engine/Private/SkeletalMesh.cpp
@@ -189,7 +189,8 @@ void FSkeletalMeshVertexBuffer::InitRHI()
 	if( ResourceArray->GetResourceDataSize() > 0 )
 	{
 		// Create the vertex buffer.
-		VertexBufferRHI = RHICreateVertexBuffer( ResourceArray->GetResourceDataSize(), ResourceArray, BUF_Static|BUF_ShaderResource);
+		FRHIResourceCreateInfo CreateInfo(ResourceArray);
+		VertexBufferRHI = RHICreateVertexBuffer( ResourceArray->GetResourceDataSize(), BUF_Static|BUF_ShaderResource, CreateInfo);
 	}
 }
 
@@ -468,7 +469,8 @@ void FSkeletalMeshVertexColorBuffer::InitRHI()
 	FResourceArrayInterface* ResourceArray = VertexData->GetResourceArray();
 	if( ResourceArray->GetResourceDataSize() > 0 )
 	{
-		VertexBufferRHI = RHICreateVertexBuffer( ResourceArray->GetResourceDataSize(), ResourceArray, BUF_Static );
+		FRHIResourceCreateInfo CreateInfo(ResourceArray);
+		VertexBufferRHI = RHICreateVertexBuffer( ResourceArray->GetResourceDataSize(), BUF_Static, CreateInfo);
 	}
 }
 
@@ -621,7 +623,8 @@ void FSkeletalMeshVertexAPEXClothBuffer::InitRHI()
 	FResourceArrayInterface* ResourceArray = VertexData->GetResourceArray();
 	if( ResourceArray->GetResourceDataSize() > 0 )
 	{
-		VertexBufferRHI = RHICreateVertexBuffer( ResourceArray->GetResourceDataSize(), ResourceArray, BUF_Static );
+		FRHIResourceCreateInfo CreateInfo(ResourceArray);
+		VertexBufferRHI = RHICreateVertexBuffer( ResourceArray->GetResourceDataSize(), BUF_Static, CreateInfo);
 	}
 }
 
diff --git a/Engine/Source/Runtime/Engine/Private/SkeletalRenderCPUSkin.cpp b/Engine/Source/Runtime/Engine/Private/SkeletalRenderCPUSkin.cpp
--- a/Engine/Source/Runtime/Engine/Private/SkeletalRenderCPUSkin.cpp
+++ b/Engine/Source/Runtime/Engine/Private/SkeletalRenderCPUSkin.cpp
@@ -71,7 +71,8 @@ void FFinalSkinVertexBuffer::InitVertexData(FStaticLODModel& LodModel)
 	// Create the buffer rendering resource
 	uint32 Size = LodModel.NumVertices * sizeof(FFinalSkinVertex);
 
-	VertexBufferRHI = RHICreateVertexBuffer(Size,NULL,BUF_Dynamic);
+	FRHIResourceCreateInfo CreateInfo;
+	VertexBufferRHI = RHICreateVertexBuffer(Size,BUF_Dynamic, CreateInfo);
 
 	// Lock the buffer.
 	void* Buffer = RHILockVertexBuffer(VertexBufferRHI,0,Size,RLM_WriteOnly);
diff --git a/Engine/Source/Runtime/Engine/Private/SkeletalRenderGPUSkin.cpp b/Engine/Source/Runtime/Engine/Private/SkeletalRenderGPUSkin.cpp
--- a/Engine/Source/Runtime/Engine/Private/SkeletalRenderGPUSkin.cpp
+++ b/Engine/Source/Runtime/Engine/Private/SkeletalRenderGPUSkin.cpp
@@ -52,7 +52,8 @@ void FMorphVertexBuffer::InitDynamicRHI()
 
 	// Create the buffer rendering resource
 	uint32 Size = LodModel.NumVertices * sizeof(FMorphGPUSkinVertex);
-	VertexBufferRHI = RHICreateVertexBuffer(Size,NULL,BUF_Dynamic);
+	FRHIResourceCreateInfo CreateInfo;
+	VertexBufferRHI = RHICreateVertexBuffer(Size,BUF_Dynamic,CreateInfo);
 
 	// Lock the buffer.
 	FMorphGPUSkinVertex* Buffer = (FMorphGPUSkinVertex*) RHILockVertexBuffer(VertexBufferRHI,0,Size,RLM_WriteOnly);
diff --git a/Engine/Source/Runtime/Engine/Private/StaticMesh.cpp b/Engine/Source/Runtime/Engine/Private/StaticMesh.cpp
--- a/Engine/Source/Runtime/Engine/Private/StaticMesh.cpp
+++ b/Engine/Source/Runtime/Engine/Private/StaticMesh.cpp
@@ -204,7 +204,8 @@ void FStaticMeshVertexBuffer::InitRHI()
 	if(ResourceArray->GetResourceDataSize())
 	{
 		// Create the vertex buffer.
-		VertexBufferRHI = RHICreateVertexBuffer(ResourceArray->GetResourceDataSize(),ResourceArray,BUF_Static);
+		FRHIResourceCreateInfo CreateInfo(ResourceArray);
+		VertexBufferRHI = RHICreateVertexBuffer(ResourceArray->GetResourceDataSize(),BUF_Static,CreateInfo);
 	}
 }
 
diff --git a/Engine/Source/Runtime/Engine/Private/TextRenderComponent.cpp b/Engine/Source/Runtime/Engine/Private/TextRenderComponent.cpp
--- a/Engine/Source/Runtime/Engine/Private/TextRenderComponent.cpp
+++ b/Engine/Source/Runtime/Engine/Private/TextRenderComponent.cpp
@@ -104,7 +104,8 @@ public:
 
 	void InitRHI()
 	{
-		VertexBufferRHI = RHICreateVertexBuffer(Vertices.Num() * sizeof(FDynamicMeshVertex),NULL,BUF_Static);
+		FRHIResourceCreateInfo CreateInfo;
+		VertexBufferRHI = RHICreateVertexBuffer(Vertices.Num() * sizeof(FDynamicMeshVertex),BUF_Static,CreateInfo);
 
 		// Copy the vertex data into the vertex buffer.
 		void* VertexBufferData = RHILockVertexBuffer(VertexBufferRHI,0,Vertices.Num() * sizeof(FDynamicMeshVertex), RLM_WriteOnly);
@@ -121,7 +122,8 @@ public:
 
 	void InitRHI()
 	{
-		IndexBufferRHI = RHICreateIndexBuffer(sizeof(uint16), Indices.Num() * sizeof(uint16), NULL, BUF_Static);
+		FRHIResourceCreateInfo CreateInfo;
+		IndexBufferRHI = RHICreateIndexBuffer(sizeof(uint16), Indices.Num() * sizeof(uint16), BUF_Static, CreateInfo);
 
 		// Copy the index data into the index buffer.
 		void* Buffer = RHILockIndexBuffer(IndexBufferRHI, 0, Indices.Num() * sizeof(uint16), RLM_WriteOnly);
diff --git a/Engine/Source/Runtime/Engine/Private/TileRendering.cpp b/Engine/Source/Runtime/Engine/Private/TileRendering.cpp
--- a/Engine/Source/Runtime/Engine/Private/TileRendering.cpp
+++ b/Engine/Source/Runtime/Engine/Private/TileRendering.cpp
@@ -46,7 +46,8 @@ public:
 		// used with a tristrip, so only 4 vertices are needed
 		uint32 Size = 4 * sizeof(FMaterialTileVertex);
 		// create vertex buffer
-		VertexBufferRHI = RHICreateVertexBuffer(Size,NULL,BUF_Static);
+		FRHIResourceCreateInfo CreateInfo;
+		VertexBufferRHI = RHICreateVertexBuffer(Size,BUF_Static,CreateInfo);
 		// lock it
 		void* Buffer = RHILockVertexBuffer(VertexBufferRHI,0,Size,RLM_WriteOnly);
         	// first vertex element
diff --git a/Engine/Source/Runtime/Engine/Private/VectorFieldVisualization.cpp b/Engine/Source/Runtime/Engine/Private/VectorFieldVisualization.cpp
--- a/Engine/Source/Runtime/Engine/Private/VectorFieldVisualization.cpp
+++ b/Engine/Source/Runtime/Engine/Private/VectorFieldVisualization.cpp
@@ -83,7 +83,8 @@ public:
 
 	virtual void InitRHI()
 	{
-		VertexBufferRHI = RHICreateVertexBuffer(sizeof(FVector4) * 2, NULL, BUF_Static);
+		FRHIResourceCreateInfo CreateInfo;
+		VertexBufferRHI = RHICreateVertexBuffer(sizeof(FVector4) * 2, BUF_Static, CreateInfo);
 		FVector4* DummyContents = (FVector4*)RHILockVertexBuffer(VertexBufferRHI,0,sizeof(FVector4)*2,RLM_WriteOnly);
 		DummyContents[0] = FVector4(0.0f, 0.0f, 0.0f, 0.0f);
 		DummyContents[1] = FVector4(1.0f, 1.0f, 1.0f, 1.0f);
diff --git a/Engine/Source/Runtime/Engine/Public/GPUSkinVertexFactory.h b/Engine/Source/Runtime/Engine/Public/GPUSkinVertexFactory.h
--- a/Engine/Source/Runtime/Engine/Public/GPUSkinVertexFactory.h
+++ b/Engine/Source/Runtime/Engine/Public/GPUSkinVertexFactory.h
@@ -298,7 +298,8 @@ public:
 #if GPUSKIN_USE_DATA_BUFFERS
 			{
 				const int32 TileBufferSize = ComputeMemorySize();
-				BoneBuffer.VertexBufferRHI = RHICreateVertexBuffer( TileBufferSize, NULL, BUF_Volatile | BUF_ShaderResource );
+				FRHIResourceCreateInfo CreateInfo;
+				BoneBuffer.VertexBufferRHI = RHICreateVertexBuffer( TileBufferSize, BUF_Volatile | BUF_ShaderResource, CreateInfo );
 				BoneBuffer.VertexBufferSRV = RHICreateShaderResourceView( BoneBuffer.VertexBufferRHI, sizeof(FVector4), PF_A32B32G32R32F );
 			}
 #else
diff --git a/Engine/Source/Runtime/Engine/Public/ParticleHelper.h b/Engine/Source/Runtime/Engine/Public/ParticleHelper.h
--- a/Engine/Source/Runtime/Engine/Public/ParticleHelper.h
+++ b/Engine/Source/Runtime/Engine/Public/ParticleHelper.h
@@ -2857,7 +2857,8 @@ public:
 	virtual void InitRHI()
 	{
 		// create a static vertex buffer
-		VertexBufferRHI = RHICreateVertexBuffer(sizeof(FParticleVertexDynamicParameter), NULL, BUF_Static|BUF_ZeroStride);
+		FRHIResourceCreateInfo CreateInfo;
+		VertexBufferRHI = RHICreateVertexBuffer(sizeof(FParticleVertexDynamicParameter), BUF_Static|BUF_ZeroStride, CreateInfo);
 		FParticleVertexDynamicParameter* Vertices = (FParticleVertexDynamicParameter*)RHILockVertexBuffer(VertexBufferRHI, 0, sizeof(FParticleVertexDynamicParameter), RLM_WriteOnly);
 		Vertices[0].DynamicValue[0] = Vertices[0].DynamicValue[1] = Vertices[0].DynamicValue[2] = Vertices[0].DynamicValue[3] = 1.0f;
 		RHIUnlockVertexBuffer(VertexBufferRHI);
diff --git a/Engine/Source/Runtime/Engine/Public/RawIndexBuffer.h b/Engine/Source/Runtime/Engine/Public/RawIndexBuffer.h
--- a/Engine/Source/Runtime/Engine/Public/RawIndexBuffer.h
+++ b/Engine/Source/Runtime/Engine/Public/RawIndexBuffer.h
@@ -227,7 +227,8 @@ public:
 		if(Indices.Num())
 		{
 			// Create the index buffer.
-			IndexBufferRHI = RHICreateIndexBuffer(sizeof(INDEX_TYPE),Size,&Indices,BUF_Static);
+			FRHIResourceCreateInfo CreateInfo(&Indices);
+			IndexBufferRHI = RHICreateIndexBuffer(sizeof(INDEX_TYPE),Size,BUF_Static,CreateInfo);
 		}    
 	}
 
diff --git a/Engine/Source/Runtime/OpenGLDrv/Private/OpenGLIndexBuffer.cpp b/Engine/Source/Runtime/OpenGLDrv/Private/OpenGLIndexBuffer.cpp
--- a/Engine/Source/Runtime/OpenGLDrv/Private/OpenGLIndexBuffer.cpp
+++ b/Engine/Source/Runtime/OpenGLDrv/Private/OpenGLIndexBuffer.cpp
@@ -6,17 +6,17 @@
 
 #include "OpenGLDrvPrivate.h"
 
-FIndexBufferRHIRef FOpenGLDynamicRHI::RHICreateIndexBuffer(uint32 Stride,uint32 Size,FResourceArrayInterface* ResourceArray,uint32 InUsage)
+FIndexBufferRHIRef FOpenGLDynamicRHI::RHICreateIndexBuffer(uint32 Stride,uint32 Size, uint32 InUsage, FRHIResourceCreateInfo& CreateInfo)
 {
 	VERIFY_GL_SCOPE();
 
 	const void *Data = NULL;
 
 	// If a resource array was provided for the resource, create the resource pre-populated
-	if(ResourceArray)
+	if(CreateInfo.ResourceArray)
 	{
-		check(Size == ResourceArray->GetResourceDataSize());
-		Data = ResourceArray->GetResourceData();
+		check(Size == CreateInfo.ResourceArray->GetResourceDataSize());
+		Data = CreateInfo.ResourceArray->GetResourceData();
 	}
 
 	TRefCountPtr<FOpenGLIndexBuffer> IndexBuffer = new FOpenGLIndexBuffer(Stride, Size, InUsage, Data);
diff --git a/Engine/Source/Runtime/OpenGLDrv/Private/OpenGLStructuredBuffer.cpp b/Engine/Source/Runtime/OpenGLDrv/Private/OpenGLStructuredBuffer.cpp
--- a/Engine/Source/Runtime/OpenGLDrv/Private/OpenGLStructuredBuffer.cpp
+++ b/Engine/Source/Runtime/OpenGLDrv/Private/OpenGLStructuredBuffer.cpp
@@ -6,17 +6,17 @@
 
 #include "OpenGLDrvPrivate.h"
 
-FStructuredBufferRHIRef FOpenGLDynamicRHI::RHICreateStructuredBuffer(uint32 Stride,uint32 Size,FResourceArrayInterface* ResourceArray,uint32 InUsage)
+FStructuredBufferRHIRef FOpenGLDynamicRHI::RHICreateStructuredBuffer(uint32 Stride,uint32 Size,uint32 InUsage, FRHIResourceCreateInfo& CreateInfo)
 {
 	VERIFY_GL_SCOPE();
 
 	const void *Data = NULL;
 
 	// If a resource array was provided for the resource, create the resource pre-populated
-	if(ResourceArray)
+	if(CreateInfo.ResourceArray)
 	{
-		check(Size == ResourceArray->GetResourceDataSize());
-		Data = ResourceArray->GetResourceData();
+		check(Size == CreateInfo.ResourceArray->GetResourceDataSize());
+		Data = CreateInfo.ResourceArray->GetResourceData();
 	}
 
 	TRefCountPtr<FOpenGLStructuredBuffer> StructuredBuffer = new FOpenGLStructuredBuffer(Stride, Size, InUsage & BUF_AnyDynamic, Data);
diff --git a/Engine/Source/Runtime/OpenGLDrv/Private/OpenGLVertexBuffer.cpp b/Engine/Source/Runtime/OpenGLDrv/Private/OpenGLVertexBuffer.cpp
--- a/Engine/Source/Runtime/OpenGLDrv/Private/OpenGLVertexBuffer.cpp
+++ b/Engine/Source/Runtime/OpenGLDrv/Private/OpenGLVertexBuffer.cpp
@@ -6,17 +6,17 @@
 
 #include "OpenGLDrvPrivate.h"
 
-FVertexBufferRHIRef FOpenGLDynamicRHI::RHICreateVertexBuffer(uint32 Size,FResourceArrayInterface* ResourceArray,uint32 InUsage)
+FVertexBufferRHIRef FOpenGLDynamicRHI::RHICreateVertexBuffer(uint32 Size, uint32 InUsage, FRHIResourceCreateInfo& CreateInfo)
 {
 	VERIFY_GL_SCOPE();
 
 	const void *Data = NULL;
 
 	// If a resource array was provided for the resource, create the resource pre-populated
-	if(ResourceArray)
+	if(CreateInfo.ResourceArray)
 	{
-		check(Size == ResourceArray->GetResourceDataSize());
-		Data = ResourceArray->GetResourceData();
+		check(Size == CreateInfo.ResourceArray->GetResourceDataSize());
+		Data = CreateInfo.ResourceArray->GetResourceData();
 	}
 
 	TRefCountPtr<FOpenGLVertexBuffer> VertexBuffer = new FOpenGLVertexBuffer(0, Size, InUsage, Data);
diff --git a/Engine/Source/Runtime/RHI/Public/RHI.h b/Engine/Source/Runtime/RHI/Public/RHI.h
--- a/Engine/Source/Runtime/RHI/Public/RHI.h
+++ b/Engine/Source/Runtime/RHI/Public/RHI.h
@@ -1242,11 +1242,25 @@ struct FRHIResourceInfo
 
 struct FRHIResourceCreateInfo
 {
-	FRHIResourceCreateInfo(FResourceBulkDataInterface* InBulkData = 0)
+	FRHIResourceCreateInfo()
+		: BulkData(0)
+		, ResourceArray(0)
+	{}
+
+	// for CreateTexture calls
+	FRHIResourceCreateInfo(FResourceBulkDataInterface* InBulkData)
 		: BulkData(InBulkData)
 	{}
 
+	// for CreateVertexBuffer/CreateStructuredBuffer calls
+	FRHIResourceCreateInfo(FResourceArrayInterface* InResourceArray)
+		: ResourceArray(InResourceArray)
+	{}
+
+	// for CreateTexture calls
 	FResourceBulkDataInterface* BulkData;
+	// for CreateVertexBuffer/CreateStructuredBuffer calls
+	FResourceArrayInterface* ResourceArray;
 };
 
 // Forward-declaration.
diff --git a/Engine/Source/Runtime/RHI/Public/RHIMethods.h b/Engine/Source/Runtime/RHI/Public/RHIMethods.h
--- a/Engine/Source/Runtime/RHI/Public/RHIMethods.h
+++ b/Engine/Source/Runtime/RHI/Public/RHIMethods.h
@@ -238,9 +238,9 @@ DEFINE_RHIMETHOD_4(
 	FIndexBufferRHIRef,RHICreateIndexBuffer,
 	uint32,Stride,
 	uint32,Size,
-	FResourceArrayInterface*,ResourceArray,
 	uint32,InUsage,
-	return,if (ResourceArray) { ResourceArray->Discard(); } return new FRHIIndexBuffer(Stride,Size,InUsage);
+	FRHIResourceCreateInfo&,CreateInfo,
+	return,if (CreateInfo.ResourceArray) { CreateInfo.ResourceArray->Discard(); } return new FRHIIndexBuffer(Stride,Size,InUsage);
 	);
 
 DEFINE_RHIMETHOD_4(
@@ -263,9 +263,9 @@ DEFINE_RHIMETHOD_1(
 DEFINE_RHIMETHOD_3(
 	FVertexBufferRHIRef,RHICreateVertexBuffer,
 	uint32,Size,
-	FResourceArrayInterface*,ResourceArray,
 	uint32,InUsage,
-	return,if (ResourceArray) { ResourceArray->Discard(); } return new FRHIVertexBuffer(Size,InUsage);
+	FRHIResourceCreateInfo&,CreateInfo,
+	return,if (CreateInfo.ResourceArray) { CreateInfo.ResourceArray->Discard(); } return new FRHIVertexBuffer(Size,InUsage);
 	);
 
 DEFINE_RHIMETHOD_4(
@@ -297,9 +297,9 @@ DEFINE_RHIMETHOD_4(
 	FStructuredBufferRHIRef,RHICreateStructuredBuffer,
 	uint32,Stride,
 	uint32,Size,
-	FResourceArrayInterface*,ResourceArray,
 	uint32,InUsage,
-	return,if (ResourceArray) { ResourceArray->Discard(); } return new FRHIStructuredBuffer(Stride,Size,InUsage);
+	FRHIResourceCreateInfo&,CreateInfo,
+	return,if (CreateInfo.ResourceArray) { CreateInfo.ResourceArray->Discard(); } return new FRHIStructuredBuffer(Stride,Size,InUsage);
 	);
 
 DEFINE_RHIMETHOD_4(
@@ -472,7 +472,7 @@ DEFINE_RHIMETHOD_7(
 	uint32,NumMips,
 	uint32,NumSamples,
 	uint32,Flags,
-	FRHIResourceCreateInfo&,Info,
+	FRHIResourceCreateInfo&,CreateInfo,
 	return,return new FRHITexture2D(SizeX,SizeY,NumMips,NumSamples,(EPixelFormat)Format,Flags);
 	);
 
@@ -532,7 +532,7 @@ DEFINE_RHIMETHOD_7(
 	uint8,Format,
 	uint32,NumMips,
 	uint32,Flags,
-	FRHIResourceCreateInfo&,Info,
+	FRHIResourceCreateInfo&,CreateInfo,
 	return,return new FRHITexture2DArray(SizeX,SizeY,SizeZ,NumMips,(EPixelFormat)Format,Flags);
 	);
 
@@ -553,7 +553,7 @@ DEFINE_RHIMETHOD_7(
 	uint8,Format,
 	uint32,NumMips,
 	uint32,Flags,
-	FRHIResourceCreateInfo&,Info,
+	FRHIResourceCreateInfo&,CreateInfo,
 	return,return new FRHITexture3D(SizeX,SizeY,SizeZ,NumMips,(EPixelFormat)Format,Flags);
 	);
 
@@ -789,7 +789,7 @@ DEFINE_RHIMETHOD_5(
 	uint8,Format,
 	uint32,NumMips,
 	uint32,Flags,
-	FRHIResourceCreateInfo&,Info,
+	FRHIResourceCreateInfo&,CreateInfo,
 	return,return new FRHITextureCube(Size,NumMips,(EPixelFormat)Format,Flags);
 	);
 
@@ -809,7 +809,7 @@ DEFINE_RHIMETHOD_6(
 	uint8,Format,
 	uint32,NumMips,
 	uint32,Flags,
-	FRHIResourceCreateInfo&,Info,
+	FRHIResourceCreateInfo&,CreateInfo,
 	return,return new FRHITextureCube(Size,NumMips,(EPixelFormat)Format,Flags);
 );
 
diff --git a/Engine/Source/Runtime/RHI/Public/RHIUtilities.h b/Engine/Source/Runtime/RHI/Public/RHIUtilities.h
--- a/Engine/Source/Runtime/RHI/Public/RHIUtilities.h
+++ b/Engine/Source/Runtime/RHI/Public/RHIUtilities.h
@@ -19,7 +19,8 @@ struct FRWBuffer
 	{
 		check(GRHIFeatureLevel == ERHIFeatureLevel::SM5);
 		NumBytes = BytesPerElement * NumElements;
-		Buffer = RHICreateVertexBuffer(NumBytes,NULL,BUF_UnorderedAccess | BUF_ShaderResource | AdditionalUsage);
+		FRHIResourceCreateInfo CreateInfo;
+		Buffer = RHICreateVertexBuffer(NumBytes,BUF_UnorderedAccess | BUF_ShaderResource | AdditionalUsage, CreateInfo);
 		UAV = RHICreateUnorderedAccessView(Buffer,Format);
 		SRV = RHICreateShaderResourceView(Buffer,BytesPerElement,Format);
 	}
@@ -46,7 +47,8 @@ struct FReadBuffer
 	{
 		check(GRHIFeatureLevel >= ERHIFeatureLevel::SM4);
 		NumBytes = BytesPerElement * NumElements;
-		Buffer = RHICreateVertexBuffer(NumBytes,NULL,BUF_ShaderResource | AdditionalUsage);
+		FRHIResourceCreateInfo CreateInfo;
+		Buffer = RHICreateVertexBuffer(NumBytes,BUF_ShaderResource | AdditionalUsage,CreateInfo);
 		SRV = RHICreateShaderResourceView(Buffer,BytesPerElement,Format);
 	}
 
@@ -72,7 +74,8 @@ struct FRWBufferStructured
 	{
 		check(GRHIFeatureLevel == ERHIFeatureLevel::SM5);
 		NumBytes = BytesPerElement * NumElements;
-		Buffer = RHICreateStructuredBuffer(BytesPerElement, NumBytes, NULL, BUF_UnorderedAccess | BUF_ShaderResource | AdditionalUsage );
+		FRHIResourceCreateInfo CreateInfo;
+		Buffer = RHICreateStructuredBuffer(BytesPerElement, NumBytes, BUF_UnorderedAccess | BUF_ShaderResource | AdditionalUsage, CreateInfo);
 		UAV = RHICreateUnorderedAccessView(Buffer,bUseUavCounter, bAppendBuffer);
 		SRV = RHICreateShaderResourceView(Buffer);
 	}
@@ -101,7 +104,8 @@ struct FRWBufferByteAddress
 		NumBytes = InNumBytes;
 		check(GRHIFeatureLevel == ERHIFeatureLevel::SM5);
 		check( NumBytes % 4 == 0 );
-		Buffer = RHICreateStructuredBuffer(4, NumBytes, NULL, BUF_UnorderedAccess | BUF_ShaderResource | BUF_ByteAddressBuffer | AdditionalUsage );
+		FRHIResourceCreateInfo CreateInfo;
+		Buffer = RHICreateStructuredBuffer(4, NumBytes, BUF_UnorderedAccess | BUF_ShaderResource | BUF_ByteAddressBuffer | AdditionalUsage, CreateInfo);
 		UAV = RHICreateUnorderedAccessView(Buffer, false, false );
 		SRV = RHICreateShaderResourceView(Buffer);
 	}
diff --git a/Engine/Source/Runtime/RenderCore/Private/RenderResource.cpp b/Engine/Source/Runtime/RenderCore/Private/RenderResource.cpp
--- a/Engine/Source/Runtime/RenderCore/Private/RenderResource.cpp
+++ b/Engine/Source/Runtime/RenderCore/Private/RenderResource.cpp
@@ -171,7 +171,8 @@ public:
 	virtual void InitRHI() OVERRIDE
 	{
 		check(!IsValidRef(VertexBufferRHI));
-		VertexBufferRHI = RHICreateVertexBuffer(BufferSize, /*ResourceArray=*/ NULL, BUF_Volatile);
+		FRHIResourceCreateInfo CreateInfo;
+		VertexBufferRHI = RHICreateVertexBuffer(BufferSize, BUF_Volatile, CreateInfo);
 		MappedBuffer = NULL;
 		AllocatedByteCount = 0;
 	}
@@ -352,7 +353,8 @@ public:
 	virtual void InitRHI() OVERRIDE
 	{
 		check(!IsValidRef(IndexBufferRHI));
-		IndexBufferRHI = RHICreateIndexBuffer(Stride, BufferSize, /*ResourceArray=*/ NULL, BUF_Volatile);
+		FRHIResourceCreateInfo CreateInfo;
+		IndexBufferRHI = RHICreateIndexBuffer(Stride, BufferSize, BUF_Volatile, CreateInfo);
 		MappedBuffer = NULL;
 		AllocatedByteCount = 0;
 	}
diff --git a/Engine/Source/Runtime/RenderCore/Public/RenderResource.h b/Engine/Source/Runtime/RenderCore/Public/RenderResource.h
--- a/Engine/Source/Runtime/RenderCore/Public/RenderResource.h
+++ b/Engine/Source/Runtime/RenderCore/Public/RenderResource.h
@@ -363,7 +363,8 @@ public:
 	virtual void InitRHI()
 	{
 		// create a static vertex buffer
-		VertexBufferRHI = RHICreateVertexBuffer(sizeof(uint32), NULL, BUF_Static|BUF_ZeroStride);
+		FRHIResourceCreateInfo CreateInfo;
+		VertexBufferRHI = RHICreateVertexBuffer(sizeof(uint32), BUF_Static|BUF_ZeroStride,CreateInfo);
 		uint32* Vertices = (uint32*)RHILockVertexBuffer(VertexBufferRHI, 0, sizeof(uint32), RLM_WriteOnly);
 		Vertices[0] = FColor(255, 255, 255, 255).DWColor();
 		RHIUnlockVertexBuffer(VertexBufferRHI);
diff --git a/Engine/Source/Runtime/Renderer/Private/CompositionLighting/PostProcessDeferredDecals.cpp b/Engine/Source/Runtime/Renderer/Private/CompositionLighting/PostProcessDeferredDecals.cpp
--- a/Engine/Source/Runtime/Renderer/Private/CompositionLighting/PostProcessDeferredDecals.cpp
+++ b/Engine/Source/Runtime/Renderer/Private/CompositionLighting/PostProcessDeferredDecals.cpp
@@ -554,7 +554,8 @@ public:
 		uint32 Size = Verts.GetResourceDataSize();
 
 		// Create vertex buffer. Fill buffer with initial data upon creation
-		VertexBufferRHI = RHICreateVertexBuffer(Size, &Verts, BUF_Static);
+		FRHIResourceCreateInfo CreateInfo(&Verts);
+		VertexBufferRHI = RHICreateVertexBuffer(Size, BUF_Static, CreateInfo);
 	}
 };
 
@@ -579,7 +580,8 @@ public:
 		const uint32 Stride = sizeof(uint16);
 
 		// Create index buffer. Fill buffer with initial data upon creation
-		IndexBufferRHI = RHICreateIndexBuffer(Stride, Size, &Indices, BUF_Static);
+		FRHIResourceCreateInfo CreateInfo(&Indices);
+		IndexBufferRHI = RHICreateIndexBuffer(Stride, Size, BUF_Static, CreateInfo);
 	}
 
 	int32 GetIndexCount() const
diff --git a/Engine/Source/Runtime/Renderer/Private/PostProcess/PostProcessBokehDOF.cpp b/Engine/Source/Runtime/Renderer/Private/PostProcess/PostProcessBokehDOF.cpp
--- a/Engine/Source/Runtime/Renderer/Private/PostProcess/PostProcessBokehDOF.cpp
+++ b/Engine/Source/Runtime/Renderer/Private/PostProcess/PostProcessBokehDOF.cpp
@@ -31,7 +31,8 @@ public:
 	{
 		const uint32 Size = sizeof(uint16) * 6 * 8;
 		const uint32 Stride = sizeof(uint16);
-		IndexBufferRHI = RHICreateIndexBuffer( Stride, Size, NULL, BUF_Static );
+		FRHIResourceCreateInfo CreateInfo;
+		IndexBufferRHI = RHICreateIndexBuffer( Stride, Size, BUF_Static, CreateInfo);
 		uint16* Indices = (uint16*)RHILockIndexBuffer( IndexBufferRHI, 0, Size, RLM_WriteOnly );
 		if(DOFIndexStyle == BIS_Fast)
 		{
diff --git a/Engine/Source/Runtime/Renderer/Private/PostProcess/SceneFilterRendering.cpp b/Engine/Source/Runtime/Renderer/Private/PostProcess/SceneFilterRendering.cpp
--- a/Engine/Source/Runtime/Renderer/Private/PostProcess/SceneFilterRendering.cpp
+++ b/Engine/Source/Runtime/Renderer/Private/PostProcess/SceneFilterRendering.cpp
@@ -41,7 +41,8 @@ public:
 		Vertices[5].UV = FVector2D(1, -1);
 
 		// Create vertex buffer. Fill buffer with initial data upon creation
-		VertexBufferRHI = RHICreateVertexBuffer(Vertices.GetResourceDataSize(), &Vertices, BUF_Static);
+		FRHIResourceCreateInfo CreateInfo(&Vertices);
+		VertexBufferRHI = RHICreateVertexBuffer(Vertices.GetResourceDataSize(), BUF_Static, CreateInfo);
 	}
 };
 
@@ -60,7 +61,8 @@ public:
 		FMemory::Memcpy(IndexBuffer.GetData(), Indices, NumIndices * sizeof(uint16));
 
 		// Create index buffer. Fill buffer with initial data upon creation
-		IndexBufferRHI = RHICreateIndexBuffer(sizeof(uint16), IndexBuffer.GetResourceDataSize(), &IndexBuffer, BUF_Static);
+		FRHIResourceCreateInfo CreateInfo(&IndexBuffer);
+		IndexBufferRHI = RHICreateIndexBuffer(sizeof(uint16), IndexBuffer.GetResourceDataSize(), BUF_Static, CreateInfo);
 	}
 };
 
diff --git a/Engine/Source/Runtime/Renderer/Private/SceneOcclusion.cpp b/Engine/Source/Runtime/Renderer/Private/SceneOcclusion.cpp
--- a/Engine/Source/Runtime/Renderer/Private/SceneOcclusion.cpp
+++ b/Engine/Source/Runtime/Renderer/Private/SceneOcclusion.cpp
@@ -309,7 +309,8 @@ public:
 		const uint32 Stride = sizeof(uint16);
 		const uint32 SizeInBytes = MaxBatchedPrimitives * NUM_CUBE_VERTICES * Stride;
 
-		IndexBufferRHI = RHICreateIndexBuffer(Stride,SizeInBytes,NULL,BUF_Static);
+		FRHIResourceCreateInfo CreateInfo;
+		IndexBufferRHI = RHICreateIndexBuffer(Stride,SizeInBytes,BUF_Static,CreateInfo);
 		uint16* RESTRICT Indices = (uint16*)RHILockIndexBuffer(IndexBufferRHI,0,SizeInBytes,RLM_WriteOnly);
 		for(uint32 PrimitiveIndex = 0;PrimitiveIndex < MaxBatchedPrimitives;PrimitiveIndex++)
 		{
diff --git a/Engine/Source/Runtime/Renderer/Private/ShadowRendering.h b/Engine/Source/Runtime/Renderer/Private/ShadowRendering.h
--- a/Engine/Source/Runtime/Renderer/Private/ShadowRendering.h
+++ b/Engine/Source/Runtime/Renderer/Private/ShadowRendering.h
@@ -81,7 +81,8 @@ namespace StencilingGeometry
 			uint32 Size = Verts.GetResourceDataSize();
 
 			// Create vertex buffer. Fill buffer with initial data upon creation
-			VertexBufferRHI = RHICreateVertexBuffer(Size,&Verts,BUF_Static);
+			FRHIResourceCreateInfo CreateInfo(&Verts);
+			VertexBufferRHI = RHICreateVertexBuffer(Size,BUF_Static,CreateInfo);
 		}
 
 		int32 GetVertexCount() const { return NumSphereVerts; }
@@ -127,7 +128,8 @@ namespace StencilingGeometry
 			const uint32 Stride = sizeof(uint16);
 
 			// Create index buffer. Fill buffer with initial data upon creation
-			IndexBufferRHI = RHICreateIndexBuffer(Stride, Size, &Indices,BUF_Static);
+			FRHIResourceCreateInfo CreateInfo(&Indices);
+			IndexBufferRHI = RHICreateIndexBuffer(Stride, Size, BUF_Static, CreateInfo);
 		}
 
 		int32 GetIndexCount() const { return NumIndices; }; 
@@ -197,7 +199,8 @@ namespace StencilingGeometry
 			NumIndices = Indices.Num();
 
 			// Create index buffer. Fill buffer with initial data upon creation
-			IndexBufferRHI = RHICreateIndexBuffer(Stride, Size, &Indices,BUF_Static);
+			FRHIResourceCreateInfo CreateInfo(&Indices);
+			IndexBufferRHI = RHICreateIndexBuffer(Stride, Size, BUF_Static,CreateInfo);
 		}
 
 		int32 GetIndexCount() const { return NumIndices; } 
@@ -229,7 +232,8 @@ namespace StencilingGeometry
 			uint32 Size = Verts.GetResourceDataSize();
 
 			// Create vertex buffer. Fill buffer with initial data upon creation
-			VertexBufferRHI = RHICreateVertexBuffer(Size, &Verts, BUF_Static);
+			FRHIResourceCreateInfo CreateInfo(&Verts);
+			VertexBufferRHI = RHICreateVertexBuffer(Size, BUF_Static, CreateInfo);
 		}
 
 		int32 GetVertexCount() const { return NumVerts; }
diff --git a/Engine/Source/Runtime/Renderer/Private/TranslucentLighting.cpp b/Engine/Source/Runtime/Renderer/Private/TranslucentLighting.cpp
--- a/Engine/Source/Runtime/Renderer/Private/TranslucentLighting.cpp
+++ b/Engine/Source/Runtime/Renderer/Private/TranslucentLighting.cpp
@@ -976,7 +976,8 @@ public:
 	{
 		// Used as a non-indexed triangle strip, so 4 vertices per quad
 		const uint32 Size = 4 * sizeof(FScreenVertex);
-		VertexBufferRHI = RHICreateVertexBuffer(Size, NULL, BUF_Static);
+		FRHIResourceCreateInfo CreateInfo;
+		VertexBufferRHI = RHICreateVertexBuffer(Size, BUF_Static, CreateInfo);
 		void* Buffer = RHILockVertexBuffer(VertexBufferRHI, 0, Size, RLM_WriteOnly);
 		FScreenVertex* DestVertex = (FScreenVertex*)Buffer;
 
diff --git a/Engine/Source/Runtime/SlateRHIRenderer/Private/SlateRHIRenderingPolicy.cpp b/Engine/Source/Runtime/SlateRHIRenderer/Private/SlateRHIRenderingPolicy.cpp
--- a/Engine/Source/Runtime/SlateRHIRenderer/Private/SlateRHIRenderingPolicy.cpp
+++ b/Engine/Source/Runtime/SlateRHIRenderer/Private/SlateRHIRenderingPolicy.cpp
@@ -41,7 +41,8 @@ void FSlateElementVertexBuffer::InitDynamicRHI()
 			BufferSize = 200 * sizeof(FSlateVertex);
 		}
 
-		VertexBufferRHI = RHICreateVertexBuffer( BufferSize, NULL, BUF_Dynamic );
+		FRHIResourceCreateInfo CreateInfo;
+		VertexBufferRHI = RHICreateVertexBuffer( BufferSize, BUF_Dynamic, CreateInfo);
 
 		// Ensure the vertex buffer could be created
 		check(IsValidRef(VertexBufferRHI));
@@ -95,7 +96,8 @@ void FSlateElementVertexBuffer::ResizeBuffer( uint32 NewSizeBytes )
 	if( NewSizeBytes != 0 && NewSizeBytes != BufferSize)
 	{
 		VertexBufferRHI.SafeRelease();
-		VertexBufferRHI = RHICreateVertexBuffer( NewSizeBytes, NULL, BUF_Dynamic );
+		FRHIResourceCreateInfo CreateInfo;
+		VertexBufferRHI = RHICreateVertexBuffer( NewSizeBytes, BUF_Dynamic, CreateInfo);
 		check(IsValidRef(VertexBufferRHI));
 
 		BufferSize = NewSizeBytes;
@@ -125,7 +127,8 @@ void FSlateElementIndexBuffer::InitDynamicRHI()
 		BufferSize = 200 * sizeof(SlateIndex);
 	}
 
-	IndexBufferRHI = RHICreateIndexBuffer( sizeof(SlateIndex), BufferSize, NULL, BUF_Dynamic );
+	FRHIResourceCreateInfo CreateInfo;
+	IndexBufferRHI = RHICreateIndexBuffer( sizeof(SlateIndex), BufferSize, BUF_Dynamic, CreateInfo );
 	check( IsValidRef(IndexBufferRHI) );
 }
 
@@ -137,7 +140,8 @@ void FSlateElementIndexBuffer::ResizeBuffer( uint32 NewSizeBytes )
 	if( NewSizeBytes != 0 && NewSizeBytes != BufferSize)
 	{
 		IndexBufferRHI.SafeRelease();
-		IndexBufferRHI = RHICreateIndexBuffer( sizeof(SlateIndex), NewSizeBytes, NULL, BUF_Dynamic );
+		FRHIResourceCreateInfo CreateInfo;
+		IndexBufferRHI = RHICreateIndexBuffer( sizeof(SlateIndex), NewSizeBytes, BUF_Dynamic, CreateInfo );
 		check(IsValidRef(IndexBufferRHI));
 
 		BufferSize = NewSizeBytes;
diff --git a/Engine/Source/Runtime/Windows/D3D11RHI/Private/D3D11IndexBuffer.cpp b/Engine/Source/Runtime/Windows/D3D11RHI/Private/D3D11IndexBuffer.cpp
--- a/Engine/Source/Runtime/Windows/D3D11RHI/Private/D3D11IndexBuffer.cpp
+++ b/Engine/Source/Runtime/Windows/D3D11RHI/Private/D3D11IndexBuffer.cpp
@@ -6,7 +6,7 @@
 
 #include "D3D11RHIPrivate.h"
 
-FIndexBufferRHIRef FD3D11DynamicRHI::RHICreateIndexBuffer(uint32 Stride,uint32 Size,FResourceArrayInterface* ResourceArray,uint32 InUsage)
+FIndexBufferRHIRef FD3D11DynamicRHI::RHICreateIndexBuffer(uint32 Stride,uint32 Size,uint32 InUsage, FRHIResourceCreateInfo& CreateInfo)
 {
 	// Explicitly check that the size is nonzero before allowing CreateIndexBuffer to opaquely fail.
 	check(Size > 0);
@@ -38,10 +38,10 @@ FIndexBufferRHIRef FD3D11DynamicRHI::RHICreateIndexBuffer(uint32 Stride,uint32 S
 	// If a resource array was provided for the resource, create the resource pre-populated
 	D3D11_SUBRESOURCE_DATA InitData;
 	D3D11_SUBRESOURCE_DATA* pInitData = NULL;
-	if(ResourceArray)
+	if(CreateInfo.ResourceArray)
 	{
-		check(Size == ResourceArray->GetResourceDataSize());
-		InitData.pSysMem = ResourceArray->GetResourceData();
+		check(Size == CreateInfo.ResourceArray->GetResourceDataSize());
+		InitData.pSysMem = CreateInfo.ResourceArray->GetResourceData();
 		InitData.SysMemPitch = Size;
 		InitData.SysMemSlicePitch = 0;
 		pInitData = &InitData;
@@ -52,10 +52,10 @@ FIndexBufferRHIRef FD3D11DynamicRHI::RHICreateIndexBuffer(uint32 Stride,uint32 S
 
 	UpdateBufferStats(IndexBufferResource, true);
 
-	if(ResourceArray)
+	if(CreateInfo.ResourceArray)
 	{
 		// Discard the resource array's contents.
-		ResourceArray->Discard();
+		CreateInfo.ResourceArray->Discard();
 	}
 
 	return new FD3D11IndexBuffer(IndexBufferResource, Stride, Size, InUsage);
diff --git a/Engine/Source/Runtime/Windows/D3D11RHI/Private/D3D11StructuredBuffer.cpp b/Engine/Source/Runtime/Windows/D3D11RHI/Private/D3D11StructuredBuffer.cpp
--- a/Engine/Source/Runtime/Windows/D3D11RHI/Private/D3D11StructuredBuffer.cpp
+++ b/Engine/Source/Runtime/Windows/D3D11RHI/Private/D3D11StructuredBuffer.cpp
@@ -3,7 +3,7 @@
 
 #include "D3D11RHIPrivate.h"
 
-FStructuredBufferRHIRef FD3D11DynamicRHI::RHICreateStructuredBuffer(uint32 Stride,uint32 Size,FResourceArrayInterface* ResourceArray,uint32 InUsage)
+FStructuredBufferRHIRef FD3D11DynamicRHI::RHICreateStructuredBuffer(uint32 Stride,uint32 Size,uint32 InUsage, FRHIResourceCreateInfo& CreateInfo)
 {
 	// Explicitly check that the size is nonzero before allowing CreateStructuredBuffer to opaquely fail.
 	check(Size > 0);
@@ -65,10 +65,10 @@ FStructuredBufferRHIRef FD3D11DynamicRHI::RHICreateStructuredBuffer(uint32 Strid
 	// If a resource array was provided for the resource, create the resource pre-populated
 	D3D11_SUBRESOURCE_DATA InitData;
 	D3D11_SUBRESOURCE_DATA* pInitData = NULL;
-	if(ResourceArray)
+	if(CreateInfo.ResourceArray)
 	{
-		check(Size == ResourceArray->GetResourceDataSize());
-		InitData.pSysMem = ResourceArray->GetResourceData();
+		check(Size == CreateInfo.ResourceArray->GetResourceDataSize());
+		InitData.pSysMem = CreateInfo.ResourceArray->GetResourceData();
 		InitData.SysMemPitch = Size;
 		InitData.SysMemSlicePitch = 0;
 		pInitData = &InitData;
@@ -79,10 +79,10 @@ FStructuredBufferRHIRef FD3D11DynamicRHI::RHICreateStructuredBuffer(uint32 Strid
 
 	UpdateBufferStats(StructuredBufferResource, true);
 
-	if(ResourceArray)
+	if(CreateInfo.ResourceArray)
 	{
 		// Discard the resource array's contents.
-		ResourceArray->Discard();
+		CreateInfo.ResourceArray->Discard();
 	}
 
 	return new FD3D11StructuredBuffer(StructuredBufferResource,Stride,Size,InUsage);
diff --git a/Engine/Source/Runtime/Windows/D3D11RHI/Private/D3D11VertexBuffer.cpp b/Engine/Source/Runtime/Windows/D3D11RHI/Private/D3D11VertexBuffer.cpp
--- a/Engine/Source/Runtime/Windows/D3D11RHI/Private/D3D11VertexBuffer.cpp
+++ b/Engine/Source/Runtime/Windows/D3D11RHI/Private/D3D11VertexBuffer.cpp
@@ -6,7 +6,7 @@
 
 #include "D3D11RHIPrivate.h"
 
-FVertexBufferRHIRef FD3D11DynamicRHI::RHICreateVertexBuffer(uint32 Size,FResourceArrayInterface* ResourceArray,uint32 InUsage)
+FVertexBufferRHIRef FD3D11DynamicRHI::RHICreateVertexBuffer(uint32 Size,uint32 InUsage, FRHIResourceCreateInfo& CreateInfo)
 {
 	// Explicitly check that the size is nonzero before allowing CreateVertexBuffer to opaquely fail.
 	check(Size > 0);
@@ -57,10 +57,10 @@ FVertexBufferRHIRef FD3D11DynamicRHI::RHICreateVertexBuffer(uint32 Size,FResourc
 	// If a resource array was provided for the resource, create the resource pre-populated
 	D3D11_SUBRESOURCE_DATA InitData;
 	D3D11_SUBRESOURCE_DATA* pInitData = NULL;
-	if(ResourceArray)
+	if(CreateInfo.ResourceArray)
 	{
-		check(Size == ResourceArray->GetResourceDataSize());
-		InitData.pSysMem = ResourceArray->GetResourceData();
+		check(Size == CreateInfo.ResourceArray->GetResourceDataSize());
+		InitData.pSysMem = CreateInfo.ResourceArray->GetResourceData();
 		InitData.SysMemPitch = Size;
 		InitData.SysMemSlicePitch = 0;
 		pInitData = &InitData;
@@ -71,10 +71,10 @@ FVertexBufferRHIRef FD3D11DynamicRHI::RHICreateVertexBuffer(uint32 Size,FResourc
 
 	UpdateBufferStats(VertexBufferResource, true);
 
-	if(ResourceArray)
+	if(CreateInfo.ResourceArray)
 	{
 		// Discard the resource array's contents.
-		ResourceArray->Discard();
+		CreateInfo.ResourceArray->Discard();
 	}
 
 	return new FD3D11VertexBuffer(VertexBufferResource,Size,InUsage);
