UE4 - merge first pass at parallel rendering to main

diff --git a/Engine/Source/Runtime/OpenGLDrv/Private/OpenGLUniformBuffer.cpp b/Engine/Source/Runtime/OpenGLDrv/Private/OpenGLUniformBuffer.cpp
--- a/Engine/Source/Runtime/OpenGLDrv/Private/OpenGLUniformBuffer.cpp
+++ b/Engine/Source/Runtime/OpenGLDrv/Private/OpenGLUniformBuffer.cpp
@@ -148,10 +148,6 @@ static inline uint32 GetPoolBucketSize(uint32 NumBytes)
 	return UniformBufferSizeBuckets[GetPoolBucketIndex(NumBytes)];
 }
 
-
-uint32 FOpenGLEUniformBuffer::UniqueIDCounter = 1;
-
-
 struct FUniformBufferDataFactory
 {
 	TMap<GLuint, FOpenGLEUniformBufferDataRef> Entries;
@@ -181,23 +177,6 @@ struct FUniformBufferDataFactory
 
 static FUniformBufferDataFactory UniformBufferDataFactory;
 
-FOpenGLEUniformBufferDataRef AllocateOpenGLEUniformBufferData(uint32 Size, GLuint& Resource)
-{
-	SCOPE_CYCLE_COUNTER(STAT_OpenGLEmulatedUniformBufferTime);
-	if (Resource == 0)
-	{
-		return UniformBufferDataFactory.Create(Size, Resource);
-	}
-
-	return UniformBufferDataFactory.Get(Resource);
-}
-
-void FreeOpenGLEUniformBufferData(GLuint Resource)
-{
-	SCOPE_CYCLE_COUNTER(STAT_OpenGLEmulatedUniformBufferTime);
-	UniformBufferDataFactory.Destroy(Resource);
-}
-
 // Describes a uniform buffer in the free pool.
 struct FPooledGLUniformBuffer
 {
@@ -205,14 +184,14 @@ struct FPooledGLUniformBuffer
 	uint32 CreatedSize;
 	uint32 Offset;
 	uint32 FrameFreed;
-	uint8* Pointer;
+	uint8* PersistentlyMappedBuffer;
 };
 
 // Pool of free uniform buffers, indexed by bucket for constant size search time.
-TArray<FPooledGLUniformBuffer> GLUniformBufferPool[NUM_POOL_BUCKETS][2];
+static TArray<FPooledGLUniformBuffer> GLUniformBufferPool[NUM_POOL_BUCKETS][2];
 
 // Uniform buffers that have been freed more recently than NumSafeFrames ago.
-TArray<FPooledGLUniformBuffer> SafeGLUniformBufferPools[NUM_SAFE_FRAMES][NUM_POOL_BUCKETS][2];
+static TArray<FPooledGLUniformBuffer> SafeGLUniformBufferPools[NUM_SAFE_FRAMES][NUM_POOL_BUCKETS][2];
 
 // Does per-frame global updating for the uniform buffer pool.
 void BeginFrame_UniformBufferPoolCleanup()
@@ -225,44 +204,39 @@ void BeginFrame_UniformBufferPoolCleanup()
 
 	if (!IsSuballocatingUBOs())
 	{
-		// Clean a limited number of old entries to reduce hitching when leaving a large level
-		const bool bIsES2 = IsES2Platform(GRHIShaderPlatform);
-		for( int32 StreamedIndex = 0; StreamedIndex < 2; ++StreamedIndex)
-		{
+	// Clean a limited number of old entries to reduce hitching when leaving a large level
+	const bool bIsES2 = IsES2Platform(GRHIShaderPlatform);
+	for( int32 StreamedIndex = 0; StreamedIndex < 2; ++StreamedIndex)
+	{
 			for (int32 BucketIndex = 0; BucketIndex < UniformBufferSizeBuckets.Num(); BucketIndex++)
+		{
+			for (int32 EntryIndex = GLUniformBufferPool[BucketIndex][StreamedIndex].Num() - 1; EntryIndex >= 0; EntryIndex--)
 			{
-				for (int32 EntryIndex = GLUniformBufferPool[BucketIndex][StreamedIndex].Num() - 1; EntryIndex >= 0; EntryIndex--)
-				{
-					FPooledGLUniformBuffer& PoolEntry = GLUniformBufferPool[BucketIndex][StreamedIndex][EntryIndex];
+				FPooledGLUniformBuffer& PoolEntry = GLUniformBufferPool[BucketIndex][StreamedIndex][EntryIndex];
 
-					check(PoolEntry.Buffer);
+				check(PoolEntry.Buffer);
 
-					// Clean entries that are unlikely to be reused
-					if (GFrameNumberRenderThread - PoolEntry.FrameFreed > 30)
+				// Clean entries that are unlikely to be reused
+				if (GFrameNumberRenderThread - PoolEntry.FrameFreed > 30)
+				{
+					DEC_DWORD_STAT(STAT_OpenGLNumFreeUniformBuffers);
+					DEC_MEMORY_STAT_BY(STAT_OpenGLFreeUniformBufferMemory, PoolEntry.CreatedSize);
+					DecrementBufferMemory(GL_UNIFORM_BUFFER, /*bIsStructuredBuffer=*/ false, PoolEntry.CreatedSize);
+					if (bIsES2)
 					{
-						DEC_DWORD_STAT(STAT_OpenGLNumFreeUniformBuffers);
-						DEC_MEMORY_STAT_BY(STAT_OpenGLFreeUniformBufferMemory, PoolEntry.CreatedSize);
-						if (bIsES2)
-						{
-							UniformBufferDataFactory.Destroy(PoolEntry.Buffer);
-						}
-						else
-						{
-							glDeleteBuffers( 1, &PoolEntry.Buffer );
-						}
-						GLUniformBufferPool[BucketIndex][StreamedIndex].RemoveAtSwap(EntryIndex);
-
-						--NumToCleanThisFrame;
-						if (NumToCleanThisFrame == 0)
-						{
-							break;
-						}
+						UniformBufferDataFactory.Destroy(PoolEntry.Buffer);
 					}
-				}
+					else
+					{
+						glDeleteBuffers( 1, &PoolEntry.Buffer );
+					}
+					GLUniformBufferPool[BucketIndex][StreamedIndex].RemoveAtSwap(EntryIndex);
 
-				if (NumToCleanThisFrame == 0)
-				{
-					break;
+					--NumToCleanThisFrame;
+					if (NumToCleanThisFrame == 0)
+					{
+						break;
+					}
 				}
 			}
 
@@ -271,6 +245,12 @@ void BeginFrame_UniformBufferPoolCleanup()
 				break;
 			}
 		}
+
+		if (NumToCleanThisFrame == 0)
+		{
+			break;
+		}
+	}
 	}
 
 	// Index of the bucket that is now old enough to be reused
@@ -354,23 +334,119 @@ static void SuballocateUBO( uint32 Size, GLuint& Resource, uint32& Offset, uint8
 	UE_LOG(LogRHI,Log,TEXT("Allocated new buffer for uniform Pool %d buffers with %d bytes"),UBOPool.Num(), UBOPool.Num()*GetUBOPoolSize());
 }
 
-static FUniformBufferRHIRef CreateUniformBuffer(uint32 InStride,uint32 InSize,uint32 InUsage, const void *InData = NULL, bool bStreamedDraw = false, GLuint ResourceToUse = 0, uint32 ResourceSize = 0, uint32 Offset = 0, uint8 *Pointer = 0)
+static uint32 GUniqueUniformBufferID = 0;
+
+FOpenGLUniformBuffer::FOpenGLUniformBuffer(const FRHIUniformBufferLayout& InLayout, GLuint InResource, uint32 InOffset, uint8* InPersistentlyMappedBuffer, uint32 InAllocatedSize, FOpenGLEUniformBufferDataRef& InEmulatedBuffer, bool bInStreamDraw)
+	: FRHIUniformBuffer(InLayout)
+	, Resource(InResource)
+	, Offset(InOffset)
+	, PersistentlyMappedBuffer(InPersistentlyMappedBuffer)
+	, UniqueID(++GUniqueUniformBufferID)
+	, LastCachedFrame(INDEX_NONE)
+	, EmulatedBufferData(InEmulatedBuffer)
+	, AllocatedSize(InAllocatedSize)
+	, bStreamDraw(bInStreamDraw)
 {
-	if (IsES2Platform(GRHIShaderPlatform))
+}
+
+FOpenGLUniformBuffer::~FOpenGLUniformBuffer()
+{
+	if (Resource != 0)
 	{
-		return new FOpenGLEUniformBuffer(InStride, InSize, InUsage, InData, bStreamedDraw, ResourceToUse, ResourceSize);
+		if (IsPoolingEnabled())
+		{
+			FPooledGLUniformBuffer NewEntry;
+			NewEntry.Buffer = Resource;
+			NewEntry.Offset = Offset;
+			NewEntry.FrameFreed = GFrameNumberRenderThread;
+			NewEntry.CreatedSize = AllocatedSize;
+			NewEntry.PersistentlyMappedBuffer = PersistentlyMappedBuffer;
+
+			int StreamedIndex = bStreamDraw ? 1 : 0;
+
+			// Add to this frame's array of free uniform buffers
+			const int32 SafeFrameIndex = GFrameNumberRenderThread % NUM_SAFE_FRAMES;
+			const uint32 BucketIndex = GetPoolBucketIndex(AllocatedSize);
+
+			check(AllocatedSize == UniformBufferSizeBuckets[BucketIndex]);
+			// this might fail with sizes > 65536; handle it then by extending the range? sizes > 65536 are presently unsupported on Mac OS X.
+
+			SafeGLUniformBufferPools[SafeFrameIndex][BucketIndex][StreamedIndex].Add(NewEntry);
+			INC_DWORD_STAT(STAT_OpenGLNumFreeUniformBuffers);
+			INC_MEMORY_STAT_BY(STAT_OpenGLFreeUniformBufferMemory, AllocatedSize);
+		}
+		else if (IsValidRef(EmulatedBufferData))
+		{
+			UniformBufferDataFactory.Destroy(Resource);
+		}
+		else
+		{
+			glDeleteBuffers(1, &Resource);
+			DecrementBufferMemory(GL_UNIFORM_BUFFER, /*bIsStructuredBuffer=*/ false, AllocatedSize);
+		}
 	}
-	else
+}
+
+void FOpenGLUniformBuffer::CacheResourcesInternal()
+{
+	const FRHIUniformBufferLayout& Layout = GetLayout();
+	int32 NumResources = Layout.Resources.Num();
+	const uint8* RESTRICT ResourceTypes = Layout.Resources.GetData();
+	const TRefCountPtr<FRHIResource>* RESTRICT Resources = ResourceTable.GetData();
+	void** RESTRICT RawResources = RawResourceTable.GetData();
+	float CurrentTime = FApp::GetCurrentTime();
+
+	// todo: Immutable resources, i.e. not textures, can be safely cached across frames.
+	// Texture streaming makes textures complicated :)
+	for (int32 i = 0; i < NumResources; ++i)
 	{
-#if !SUBALLOCATED_CONSTANT_BUFFER
-		return new FOpenGLUniformBuffer(InStride, InSize, InUsage, InData, bStreamedDraw, ResourceToUse, ResourceSize);
-#else
-		return new FOpenGLUniformBuffer(InStride, InSize, InUsage, InData, bStreamedDraw, ResourceToUse, ResourceSize, Offset, Pointer);
-#endif
+		switch (ResourceTypes[i])
+		{
+		case UBMT_SRV:
+			RawResources[i] = (FOpenGLShaderResourceView*)Resources[i].GetReference();
+			break;
+
+		case UBMT_TEXTURE:
+			{
+				FRHITexture* TextureRHI = (FRHITexture*)Resources[i].GetReference();
+				TextureRHI->SetLastCachedTime(CurrentTime);
+				// todo: this does multiple virtual function calls to find the right type to cast to
+				// this is due to multiple inheritance nastiness, NEEDS CLEANUP
+				FOpenGLTextureBase* TextureGL = GetOpenGLTextureFromRHITexture(TextureRHI);
+				// todo: for now we just store the GL texture wrapper. we could also store the GLuint alongside it, saves an indirection later on for non-ES2 platforms
+				// todo: could also store an FTextureStage?
+				RawResources[i] = TextureGL;
+			}
+			break;
+
+		case UBMT_UAV:
+			RawResources[i] = 0;
+			check(0);
+			break;
+
+		case UBMT_SAMPLER:
+			{
+				FOpenGLSamplerState* SamplerState = (FOpenGLSamplerState*)Resources[i].GetReference();
+				if (FOpenGL::SupportsSamplerObjects())
+				{
+					PTRINT SamplerAsInt = (PTRINT)SamplerState->Resource;
+					RawResources[i] = (void*)SamplerAsInt;
+				}
+				else
+				{
+					RawResources[i] = SamplerState;
+				}
+			}
+			break;
+
+		default:
+			check(0);
+			break;
+		}
 	}
 }
 
-FUniformBufferRHIRef FOpenGLDynamicRHI::RHICreateUniformBuffer(const void* Contents,uint32 NumBytes,EUniformBufferUsage Usage)
+FUniformBufferRHIRef FOpenGLDynamicRHI::RHICreateUniformBuffer(const void* Contents, const FRHIUniformBufferLayout& Layout, EUniformBufferUsage Usage)
 {
 	check(IsInRenderingThread());
 
@@ -378,75 +454,112 @@ FUniformBufferRHIRef FOpenGLDynamicRHI::RHICreateUniformBuffer(const void* Conte
 	// are only used for drawing/shader usage, not for loading resources or framebuffer blitting, so no synchronization primitives for now.
 
 	// Explicitly check that the size is nonzero before allowing CreateBuffer to opaquely fail.
-	check(NumBytes > 0);
+	check(Layout.Resources.Num() > 0 || Layout.ConstantBufferSize > 0);
 
 	VERIFY_GL_SCOPE();
 
-	bool bStreamDraw = (Usage==UniformBuffer_SingleUse);
+	bool bStreamDraw = (Usage==UniformBuffer_SingleDraw || Usage == UniformBuffer_SingleFrame);
+	GLuint AllocatedResource = 0;
+	uint32 OffsetInBuffer = 0;
+	uint8* PersistentlyMappedBuffer = NULL;
+	uint32 AllocatedSize = 0;
+	FOpenGLEUniformBufferDataRef EmulatedUniformDataRef;
 
-	if(IsPoolingEnabled())
+	// If the uniform buffer contains constants, allocate a uniform buffer resource from GL.
+	if (Layout.ConstantBufferSize > 0)
 	{
-		int StreamedIndex = bStreamDraw ? 1 : 0;
-
-		// Find the appropriate bucket based on size
-		const uint32 BucketIndex = GetPoolBucketIndex(NumBytes);
-		TArray<FPooledGLUniformBuffer>& PoolBucket = GLUniformBufferPool[BucketIndex][StreamedIndex];
-		if (PoolBucket.Num() > 0)
+		uint32 SizeOfBufferToAllocate = 0;
+		if(IsPoolingEnabled())
 		{
-			// Reuse the last entry in this size bucket
-			FPooledGLUniformBuffer FreeBufferEntry = PoolBucket.Pop();
-			DEC_DWORD_STAT(STAT_OpenGLNumFreeUniformBuffers);
-			DEC_MEMORY_STAT_BY(STAT_OpenGLFreeUniformBufferMemory, FreeBufferEntry.CreatedSize);
+			// Find the appropriate bucket based on size
+			const uint32 BucketIndex = GetPoolBucketIndex(Layout.ConstantBufferSize);
+			int StreamedIndex = bStreamDraw ? 1 : 0;
+			TArray<FPooledGLUniformBuffer>& PoolBucket = GLUniformBufferPool[BucketIndex][StreamedIndex];
+			if (PoolBucket.Num() > 0)
+			{
+				// Reuse the last entry in this size bucket
+				FPooledGLUniformBuffer FreeBufferEntry = PoolBucket.Pop();
+				DEC_DWORD_STAT(STAT_OpenGLNumFreeUniformBuffers);
+				DEC_MEMORY_STAT_BY(STAT_OpenGLFreeUniformBufferMemory, FreeBufferEntry.CreatedSize);
+
+				AllocatedResource = FreeBufferEntry.Buffer;
+				AllocatedSize = FreeBufferEntry.CreatedSize;
 
-			return CreateUniformBuffer(0,NumBytes,0,Contents,bStreamDraw,FreeBufferEntry.Buffer,FreeBufferEntry.CreatedSize,FreeBufferEntry.Offset,FreeBufferEntry.Pointer);
+				if (IsES2Platform(GRHIShaderPlatform))
+				{
+					EmulatedUniformDataRef = UniformBufferDataFactory.Get(AllocatedResource);
+				}
+				else
+				{
+					::CachedBindUniformBuffer(AllocatedResource);
+				}
+			}
+			else
+			{
+				SizeOfBufferToAllocate = UniformBufferSizeBuckets[BucketIndex];
+			}
 		}
-		else
+		
+		if (AllocatedSize == 0)
 		{
-			// Nothing usable was found in the free pool, create a new uniform buffer
-			uint32 BufferSize = UniformBufferSizeBuckets[BucketIndex];
-			if (IsSuballocatingUBOs())
+			// When pooling is enabled we allocate a buffer large enough for the given bucket.
+			// Otherwise we just allocate the number of bytes needed for the constant buffer we've been given.
+			AllocatedSize = (SizeOfBufferToAllocate > 0) ? SizeOfBufferToAllocate : Layout.ConstantBufferSize;
+
+			// Nothing usable was found in the free pool, or we're not pooling, so create a new uniform buffer
+			if (IsES2Platform(GRHIShaderPlatform))
 			{
-				GLuint Resource = 0;
-				uint32 Offset = 0;
-				uint8* Pointer = 0;
-				SuballocateUBO( BufferSize, Resource, Offset, Pointer);
-				return CreateUniformBuffer(0,NumBytes,0,Contents,bStreamDraw,Resource,BufferSize,Offset,Pointer);
+				EmulatedUniformDataRef = UniformBufferDataFactory.Create(AllocatedSize, AllocatedResource);
+			}
+			else if (IsSuballocatingUBOs())
+			{
+				//gilmerge was SuballocateUBO(BufferSize, AllocatedResource, OffsetInBuffer, PersistentlyMappedBuffer);
+				SuballocateUBO(AllocatedSize, AllocatedResource, OffsetInBuffer, PersistentlyMappedBuffer);
 			}
 			else
 			{
-				return CreateUniformBuffer(0,NumBytes,0,Contents,bStreamDraw,0,BufferSize);
+				FOpenGL::GenBuffers(1, &AllocatedResource);
+				::CachedBindUniformBuffer(AllocatedResource);
+				glBufferData(GL_UNIFORM_BUFFER, AllocatedSize, NULL, bStreamDraw ? GL_STREAM_DRAW : GL_STATIC_DRAW);
 			}
-		}
-	}
-	else
-	{
-		return CreateUniformBuffer(0, NumBytes, 0, Contents, bStreamDraw);
-	}
-}
 
-void AddNewlyFreedBufferToUniformBufferPool( GLuint Buffer, uint32 BufferSize, bool bStreamDraw, uint32 Offset, uint8* Pointer )
-{
-	if(IsPoolingEnabled())
-	{
-		check(Buffer);
-		FPooledGLUniformBuffer NewEntry;
-		NewEntry.Buffer = Buffer;
-		NewEntry.FrameFreed = GFrameNumberRenderThread;
-		NewEntry.CreatedSize = BufferSize;
-		NewEntry.Offset = Offset;
-		NewEntry.Pointer = Pointer;
+			IncrementBufferMemory(GL_UNIFORM_BUFFER, /*bIsStructuredBuffer=*/ false, AllocatedSize);
+		}
 
-		int StreamedIndex = bStreamDraw ? 1 : 0;
+		check(!IsES2Platform(GRHIShaderPlatform) || (IsValidRef(EmulatedUniformDataRef) && (EmulatedUniformDataRef->Data.Num() * EmulatedUniformDataRef->Data.GetTypeSize() == AllocatedSize)));
 
-		// Add to this frame's array of free uniform buffers
-		const int32 SafeFrameIndex = (GFrameNumberRenderThread - 1) % NUM_SAFE_FRAMES;
-		const uint32 BucketIndex = GetPoolBucketIndex(BufferSize);
+		// Copy the contents of the uniform buffer.
+		if (IsValidRef(EmulatedUniformDataRef))
+		{
+			FMemory::Memcpy(EmulatedUniformDataRef->Data.GetData(), Contents, Layout.ConstantBufferSize);
+		}
+		else if (PersistentlyMappedBuffer)
+		{
+			FMemory::Memcpy(PersistentlyMappedBuffer, Contents, Layout.ConstantBufferSize);
+		}
+		else
+		{
+			glBufferSubData(GL_UNIFORM_BUFFER, 0, Layout.ConstantBufferSize, Contents);
+		}
+	}
 
-		check(BufferSize == UniformBufferSizeBuckets[BucketIndex]);
-		// this might fail with sizes > 65536; handle it then by extending the range? sizes > 65536 are presently unsupported on Mac OS X.
+	FOpenGLUniformBuffer* NewUniformBuffer = new FOpenGLUniformBuffer(Layout, AllocatedResource, OffsetInBuffer, PersistentlyMappedBuffer, AllocatedSize, EmulatedUniformDataRef, bStreamDraw);
 
-		SafeGLUniformBufferPools[SafeFrameIndex][BucketIndex][StreamedIndex].Add(NewEntry);
-		INC_DWORD_STAT(STAT_OpenGLNumFreeUniformBuffers);
-		INC_MEMORY_STAT_BY(STAT_OpenGLFreeUniformBufferMemory, BufferSize);
+	// Initialize the resource table for this uniform buffer.
+	if (Layout.Resources.Num())
+	{
+		int32 NumResources = Layout.Resources.Num();
+		FRHIResource** InResources = (FRHIResource**)((uint8*)Contents + Layout.ResourceOffset);
+		NewUniformBuffer->ResourceTable.Empty(NumResources);
+		NewUniformBuffer->ResourceTable.AddZeroed(NumResources);
+		for (int32 i = 0; i < NumResources; ++i)
+		{
+			check(InResources[i]);
+			NewUniformBuffer->ResourceTable[i] = InResources[i];
+		}
+		NewUniformBuffer->RawResourceTable.Empty(NumResources);
+		NewUniformBuffer->RawResourceTable.AddZeroed(NumResources);
 	}
-}
+
+	return NewUniformBuffer;
+}	
diff --git a/Engine/Source/Runtime/OpenGLDrv/Public/OpenGLDrv.h b/Engine/Source/Runtime/OpenGLDrv/Public/OpenGLDrv.h
--- a/Engine/Source/Runtime/OpenGLDrv/Public/OpenGLDrv.h
+++ b/Engine/Source/Runtime/OpenGLDrv/Public/OpenGLDrv.h
@@ -314,7 +314,7 @@ public:
 	void OnVertexBufferDeletion(GLuint VertexBufferResource);
 	void OnIndexBufferDeletion(GLuint IndexBufferResource);
 	void OnPixelBufferDeletion(GLuint PixelBufferResource);
-	void OnUniformBufferDeletion(GLuint UniformBufferResource,uint32 AllocatedSize,bool bStreamDraw,uint32 Offset,uint8* Pointer);
+	void OnUniformBufferDeletion(GLuint UniformBufferResource,uint32 AllocatedSize,bool bStreamDraw);
 	void OnProgramDeletion(GLint ProgramResource);
 	void InvalidateTextureResourceInCache(GLuint Resource);
 	void InvalidateUAVResourceInCache(GLuint Resource);
@@ -379,8 +379,16 @@ public:
 	/** Inform all queries about the need to recreate themselves after OpenGL context they're in gets deleted. */
 	void InvalidateQueries();
 
+	FOpenGLSamplerState* GetPointSamplerState() const { return (FOpenGLSamplerState*)PointSamplerState.GetReference(); }
+
 private:
 
+	/** Counter incremented each time RHIBeginScene is called. */
+	uint32 SceneFrameCounter;
+
+	/** Value used to detect when resource tables need to be recached. INDEX_NONE means always recache. */
+	uint32 ResourceTableFrameCounter;
+
 	/** RHI device state, independent of underlying OpenGL context used */
 	FOpenGLRHIState						PendingState;
 	FOpenGLStreamedVertexBufferArray	DynamicVertexBuffers;
@@ -445,6 +453,9 @@ private:
 	void UpdateScissorRectInOpenGLContext( FOpenGLContextState& ContextState );
 	void UpdateViewportInOpenGLContext( FOpenGLContextState& ContextState );
 	
+	template <class ShaderType> void SetResourcesFromTables(const ShaderType* RESTRICT);
+	void CommitGraphicsResourceTables();
+	void CommitComputeResourceTables(class FOpenGLComputeShader* ComputeShader);
 	void CommitNonComputeShaderConstants();
 	void CommitComputeShaderConstants(FComputeShaderRHIParamRef ComputeShaderRHI);
 	void SetPendingBlendStateForActiveRenderTargets( FOpenGLContextState& ContextState );
@@ -455,18 +466,18 @@ private:
 
 	void SetupUAVsForDraw(FOpenGLContextState& ContextState, const TRefCountPtr<FOpenGLComputeShader> &ComputeShader, int32 MaxUAVsNeeded);
 
+public:
 	/** Remember what RHI user wants set on a specific OpenGL texture stage, translating from Stage and TextureIndex for stage pair. */
 	void InternalSetShaderTexture(FOpenGLTextureBase* Texture, GLint TextureIndex, GLenum Target, GLuint Resource, int NumMips, int LimitMip);
 	void InternalSetShaderUAV(GLint UAVIndex, GLenum Format, GLuint Resource);
 	void InternalSetSamplerStates(GLint TextureIndex, FOpenGLSamplerState* SamplerState);
 
+private:
 	void ApplyTextureStage(FOpenGLContextState& ContextState, GLint TextureIndex, const FTextureStage& TextureStage, FOpenGLSamplerState* SamplerState);
 
 	void ReadSurfaceDataRaw(FOpenGLContextState& ContextState, FTextureRHIParamRef TextureRHI,FIntRect Rect,TArray<uint8>& OutData, FReadSurfaceDataFlags InFlags);
 
-	void BindUniformBufferBase(FOpenGLContextState& ContextState, int32 NumUniformBuffers, const TArray< TRefCountPtr<FRHIUniformBuffer> >& BoundUniformBuffers, uint32 FirstUniformBuffer, bool ForceUpdate);
-
-	void ApplyUniformBuffersParametersToPacked(FOpenGLContextState& ContextState, int32 NumUniformBuffers, const TArray< TRefCountPtr<FRHIUniformBuffer> >& BoundUniformBuffers, uint32 FirstUniformBuffer, bool ForceUpdate, const TArray<FOpenGLUniformBufferCopyInfo>& UniformBufferCopyInfo, FOpenGLShaderParameterCache& ParameterCache);
+	void BindUniformBufferBase(FOpenGLContextState& ContextState, int32 NumUniformBuffers, FUniformBufferRHIRef* BoundUniformBuffers, uint32 FirstUniformBuffer, bool ForceUpdate);
 
 	void ClearCurrentFramebufferWithCurrentScissor(FOpenGLContextState& ContextState, int8 ClearType, int32 NumClearColors, const FLinearColor* ClearColorArray, float Depth, uint32 Stencil);
 
diff --git a/Engine/Source/Runtime/OpenGLDrv/Public/OpenGLResources.h b/Engine/Source/Runtime/OpenGLDrv/Public/OpenGLResources.h
--- a/Engine/Source/Runtime/OpenGLDrv/Public/OpenGLResources.h
+++ b/Engine/Source/Runtime/OpenGLDrv/Public/OpenGLResources.h
@@ -442,66 +442,58 @@ struct FOpenGLEUniformBufferData : public FRefCountedObject
 };
 typedef TRefCountPtr<FOpenGLEUniformBufferData> FOpenGLEUniformBufferDataRef;
 
-extern FOpenGLEUniformBufferDataRef AllocateOpenGLEUniformBufferData(uint32 Size, GLuint& Resource);
-extern void FreeOpenGLEUniformBufferData(GLuint Resource);
-
-// Emulated Uniform Buffer
-class FOpenGLEUniformBuffer : public FRHIUniformBuffer
+class FOpenGLUniformBuffer : public FRHIUniformBuffer
 {
 public:
-	const uint32 UniqueID;
+	/** The GL resource for this uniform buffer. */
 	GLuint Resource;
-	bool bStreamDraw;
-	uint32 RealSize;
 
-	FOpenGLEUniformBuffer(uint32 InStride, uint32 InSize, uint32 InUsage,
-		const void *InData, bool bInStreamDraw, GLuint ResourceToUse, uint32 ResourceSize) :
-		FRHIUniformBuffer(InSize), UniqueID(UniqueIDCounter), Resource(ResourceToUse), bStreamDraw(bInStreamDraw)
-	{
-		RealSize = ResourceSize ? ResourceSize : InSize;
-		Buffer = AllocateOpenGLEUniformBufferData(InSize, Resource);
-		if (InData)
-		{
-			FMemory::Memcpy(Buffer->Data.GetData(), InData, Buffer->Data.Num() * 4);
-		}
+	/** The offset of the uniform buffer's contents in the resource. */
+	uint32 Offset;
 
-		++UniqueIDCounter;
-	}
+	/** When using a persistently mapped buffer this is a pointer to the CPU accessible data. */
+	uint8* PersistentlyMappedBuffer;
 
-	~FOpenGLEUniformBuffer()
-	{
-		OnUniformBufferDeletion(Resource, RealSize, bStreamDraw, 0, 0);
-	}
+	/** Unique ID for state shadowing purposes. */
+	const uint32 UniqueID;
 
-	FOpenGLEUniformBufferDataRef Buffer;
+	/** Resource table containing RHI references. */
+	TArray<TRefCountPtr<FRHIResource> > ResourceTable;
 
-protected:
-	static uint32 UniqueIDCounter;
-};
+	/** Raw resource table, cached once per frame. */
+	TArray<void*> RawResourceTable;
 
+	/** The frame in which RawResourceTable was last cached. */
+	uint32 LastCachedFrame;
 
-class FOpenGLBaseUniformBuffer : public FRHIUniformBuffer
-{
-public:
-	FOpenGLBaseUniformBuffer(uint32 InStride,uint32 InSize,uint32 InUsage): FRHIUniformBuffer(InSize) {}
-	static bool OnDelete(GLuint Resource,uint32 Size,bool bStreamDraw,uint32 Offset,uint8* Pointer)
-	{
-		OnUniformBufferDeletion(Resource,Size,bStreamDraw,Offset,Pointer);
-		return false;
-	}
-	uint32 GetUsage() const { return 0; }
+	/** Emulated uniform data for ES2. */
+	FOpenGLEUniformBufferDataRef EmulatedBufferData;
 
-	static FORCEINLINE bool GLSupportsType()
-	{
-		return FOpenGL::SupportsUniformBuffers();
-	}
+	/** The size of the buffer allocated to hold the uniform buffer contents. May be larger than necessary. */
+	uint32 AllocatedSize;
 
-	static void CreateType(GLuint& Resource, const void* InData, uint32 InSize)
+	/** True if the uniform buffer is not used across frames. */
+	bool bStreamDraw;
+
+	/** Initialization constructor. */
+	FOpenGLUniformBuffer(const FRHIUniformBufferLayout& InLayout, GLuint InResource, uint32 InOffset, uint8* InPersistentlyMappedBuffer, uint32 InAllocatedSize, FOpenGLEUniformBufferDataRef& InEmulatedBuffer, bool bInStreamDraw);
+
+	/** Destructor. */
+	~FOpenGLUniformBuffer();
+
+	/** Cache resources if needed. */
+	inline void CacheResources(uint32 InFrameCounter)
 	{
-		// @todo-mobile
+		if (InFrameCounter == INDEX_NONE || LastCachedFrame != InFrameCounter)
+		{
+			CacheResourcesInternal();
+			LastCachedFrame = InFrameCounter;
+		}
 	}
 
-	static bool IsStructuredBuffer() { return false; }
+private:
+	/** Actually cache resources. */
+	void CacheResourcesInternal();
 };
 
 
@@ -556,85 +548,6 @@ typedef TOpenGLBuffer<FOpenGLBaseVertexBuffer, GL_ARRAY_BUFFER, CachedBindArrayB
 typedef TOpenGLBuffer<FOpenGLBaseIndexBuffer,GL_ELEMENT_ARRAY_BUFFER,CachedBindElementArrayBuffer> FOpenGLIndexBuffer;
 typedef TOpenGLBuffer<FOpenGLBaseStructuredBuffer,GL_ARRAY_BUFFER,CachedBindArrayBuffer> FOpenGLStructuredBuffer;
 
-#define SUBALLOCATED_CONSTANT_BUFFER 1
-
-#if !SUBALLOCATED_CONSTANT_BUFFER
-typedef TOpenGLBuffer<FOpenGLBaseUniformBuffer, GL_UNIFORM_BUFFER, CachedBindUniformBuffer> FOpenGLUniformBuffer;
-#else
-class FOpenGLUniformBuffer : public FOpenGLBaseUniformBuffer
-{
-public:
-	GLuint Resource;
-	bool bStreamDraw;
-	uint32 RealSize;
-	uint32 Offset;
-	uint8* Pointer;
-
-	FOpenGLUniformBuffer(uint32 InStride,uint32 InSize,uint32 InUsage,
-		const void *InData = NULL, bool bStreamedDraw = false, GLuint ResourceToUse = 0, uint32 ResourceSize = 0, uint32 InOffset = 0, uint8* InPointer = 0)
-	: FOpenGLBaseUniformBuffer(InStride,InSize,InUsage)
-	, Resource(0)
-	, bStreamDraw(bStreamedDraw)
-	, RealSize(InSize)
-	, Offset(InOffset)
-	, Pointer(InPointer)
-	{
-		VERIFY_GL_SCOPE();
-		RealSize = ResourceSize ? ResourceSize : InSize;
-		if( ResourceToUse )
-		{
-			Resource = ResourceToUse;
-			if ( Pointer)
-			{
-				//Want to just use memcpy, no need to bind, etc
-				FMemory::Memcpy( Pointer, InData, InSize);
-			}
-			else
-			{
-				CachedBindUniformBuffer(Resource);
-				glBufferSubData(GL_UNIFORM_BUFFER, Offset, InSize, InData);
-			}
-		}
-		else
-		{
-			check( Offset == 0);
-			check( Pointer == 0);
-			if (GLSupportsType())
-			{
-				FOpenGL::GenBuffers(1, &Resource);
-				CachedBindUniformBuffer(Resource);
-				if( InData == NULL || RealSize <= InSize )
-				{
-					glBufferData(GL_UNIFORM_BUFFER, RealSize, InData, bStreamDraw ? GL_STREAM_DRAW : (IsDynamic() ? GL_DYNAMIC_DRAW : GL_STATIC_DRAW));
-				}
-				else
-				{
-					glBufferData(GL_UNIFORM_BUFFER, RealSize, NULL, bStreamDraw ? GL_STREAM_DRAW : (IsDynamic() ? GL_DYNAMIC_DRAW : GL_STATIC_DRAW));
-					glBufferSubData(GL_UNIFORM_BUFFER, 0, InSize, InData);
-				}
-				IncrementBufferMemory(GL_UNIFORM_BUFFER, IsStructuredBuffer(), RealSize);
-			}
-			else
-			{
-				CreateType(Resource, InData, InSize);
-			}
-		}
-	}
-
-	~FOpenGLUniformBuffer()
-	{
-		VERIFY_GL_SCOPE();
-		if (Resource != 0 && OnDelete(Resource,RealSize,bStreamDraw,Offset,Pointer))
-		{
-			glDeleteBuffers(1, &Resource);
-			DecrementBufferMemory(GL_UNIFORM_BUFFER, IsStructuredBuffer(), RealSize);
-		}
-	}
-
-	bool IsDynamic() const { return (this->GetUsage() & BUF_AnyDynamic) != 0; }
-};
-#endif
-
 #define MAX_STREAMED_BUFFERS_IN_ARRAY 2	// must be > 1!
 #define MIN_DRAWS_IN_SINGLE_BUFFER 16
 
@@ -837,6 +750,9 @@ public:
 	/** The OpenGL texture target. */
 	GLenum Target;
 
+	/** The number of mips in the texture. */
+	uint32 NumMips;
+
 	/** The OpenGL attachment point. This should always be GL_COLOR_ATTACHMENT0 in case of color buffer, but the actual texture may be attached on other color attachments. */
 	GLenum Attachment;
 
@@ -845,12 +761,14 @@ public:
 		FOpenGLDynamicRHI* InOpenGLRHI,
 		GLuint InResource,
 		GLenum InTarget,
+		uint32 InNumMips,
 		GLenum InAttachment
 		)
 	: OpenGLRHI(InOpenGLRHI)
 	, SamplerState(nullptr)
 	, Resource(InResource)
 	, Target(InTarget)
+	, NumMips(InNumMips)
 	, Attachment(InAttachment)
 	, MemorySize( 0 )
 	, bIsPowerOfTwo(false)
@@ -921,6 +839,7 @@ public:
 		InOpenGLRHI,
 		InResource,
 		InTarget,
+		InNumMips,
 		InAttachment
 		)
 	, BaseLevel(0)
@@ -1076,6 +995,19 @@ typedef TOpenGLTexture<FOpenGLBaseTexture2DArray>		FOpenGLTexture2DArray;
 typedef TOpenGLTexture<FOpenGLBaseTexture3D>			FOpenGLTexture3D;
 typedef TOpenGLTexture<FOpenGLBaseTextureCube>			FOpenGLTextureCube;
 
+class FOpenGLTextureReference : public FRHITextureReference
+{
+	FOpenGLTextureBase* TexturePtr;
+
+public:
+	FOpenGLTextureReference()
+		: TexturePtr(NULL)
+	{}
+
+	void SetReferencedTexture(FRHITexture* InTexture);
+	FOpenGLTextureBase* GetTexturePtr() const { return TexturePtr; }
+};
+
 /** Given a pointer to a RHI texture that was created by the OpenGL RHI, returns a pointer to the FOpenGLTextureBase it encapsulates. */
 inline FOpenGLTextureBase* GetOpenGLTextureFromRHITexture(FRHITexture* Texture)
 {
@@ -1083,6 +1015,10 @@ inline FOpenGLTextureBase* GetOpenGLTextureFromRHITexture(FRHITexture* Texture)
 	{
 		return NULL;
 	}
+	else if(Texture->GetTextureReference())
+	{
+		return ((FOpenGLTextureReference*)Texture)->GetTexturePtr();
+	}
 	else if(Texture->GetTexture2D())
 	{
 		return (FOpenGLTexture2D*)Texture;
diff --git a/Engine/Source/Runtime/OpenGLDrv/Public/OpenGLShaderResources.h b/Engine/Source/Runtime/OpenGLDrv/Public/OpenGLShaderResources.h
--- a/Engine/Source/Runtime/OpenGLDrv/Public/OpenGLShaderResources.h
+++ b/Engine/Source/Runtime/OpenGLDrv/Public/OpenGLShaderResources.h
@@ -90,6 +90,35 @@ inline FArchive& operator<<(FArchive& Ar, FOpenGLPackedArrayInfo& Info)
 	return Ar;
 }
 
+struct FOpenGLShaderResourceTable : public FBaseShaderResourceTable
+{
+	/** Mapping of bound Textures to their location in resource tables. */
+	TArray<uint32> TextureMap;
+	friend bool operator==(const FOpenGLShaderResourceTable &A, const FOpenGLShaderResourceTable& B)
+	{
+		if (!(((FBaseShaderResourceTable&)A) == ((FBaseShaderResourceTable&)B)))
+		{
+			return false;
+		}
+		if (A.TextureMap.Num() != B.TextureMap.Num())
+		{
+			return false;
+		}
+		if (FMemory::Memcmp(A.TextureMap.GetTypedData(), B.TextureMap.GetTypedData(), A.TextureMap.GetTypeSize()*A.TextureMap.Num()) != 0)
+		{
+			return false;
+		}
+		return true;
+	}
+};
+
+inline FArchive& operator<<(FArchive& Ar, FOpenGLShaderResourceTable& SRT)
+{
+	Ar << ((FBaseShaderResourceTable&)SRT);
+	Ar << SRT.TextureMap;
+	return Ar;
+}
+
 /**
  * Shader binding information.
  */
@@ -97,6 +126,7 @@ struct FOpenGLShaderBindings
 {
 	TArray<TArray<FOpenGLPackedArrayInfo>>	PackedUniformBuffers;
 	TArray<FOpenGLPackedArrayInfo>			PackedGlobalArrays;
+	FOpenGLShaderResourceTable				ShaderResourceTable;
 
 	uint16	InOutMask;
 	uint8	NumSamplers;
@@ -114,7 +144,7 @@ struct FOpenGLShaderBindings
 	{
 	}
 
-	friend bool operator==( const FOpenGLShaderBindings &A, const FOpenGLShaderBindings &B)
+	friend bool operator==( const FOpenGLShaderBindings &A, const FOpenGLShaderBindings& B)
 	{
 		bool bEqual = true;
 
@@ -125,6 +155,7 @@ struct FOpenGLShaderBindings
 		bEqual &= A.bFlattenUB == B.bFlattenUB;
 		bEqual &= A.PackedGlobalArrays.Num() == B.PackedGlobalArrays.Num();
 		bEqual &= A.PackedUniformBuffers.Num() == B.PackedUniformBuffers.Num();
+		bEqual &= A.ShaderResourceTable == B.ShaderResourceTable;
 
 		if ( !bEqual )
 		{
@@ -154,6 +185,8 @@ struct FOpenGLShaderBindings
 		Hash ^= Binding.bFlattenUB << 8;
 		Hash ^= FCrc::MemCrc_DEPRECATED( Binding.PackedGlobalArrays.GetTypedData(), Binding.PackedGlobalArrays.GetTypeSize()*Binding.PackedGlobalArrays.Num());
 
+		//@todo-rco: Do we need to calc Binding.ShaderResourceTable.GetTypeHash()?
+
 		for (int32 Item = 0; Item < Binding.PackedUniformBuffers.Num(); Item++)
 		{
 			const TArray<FOpenGLPackedArrayInfo> &Array = Binding.PackedUniformBuffers[Item];
@@ -167,6 +200,7 @@ inline FArchive& operator<<(FArchive& Ar, FOpenGLShaderBindings& Bindings)
 {
 	Ar << Bindings.PackedUniformBuffers;
 	Ar << Bindings.PackedGlobalArrays;
+	Ar << Bindings.ShaderResourceTable;
 	Ar << Bindings.InOutMask;
 	Ar << Bindings.NumSamplers;
 	Ar << Bindings.NumUniformBuffers;
@@ -245,20 +279,24 @@ class FOpenGLLinkedProgram;
 /**
  * OpenGL shader resource.
  */
-template <typename RHIResourceType, GLenum GLTypeEnum>
+template <typename RHIResourceType, GLenum GLTypeEnum, EShaderFrequency FrequencyT>
 class TOpenGLShader : public RHIResourceType
 {
 public:
-
+	enum
+	{
+		StaticFrequency = FrequencyT
+	};
 	static const GLenum TypeEnum = GLTypeEnum;
 
 	/** The OpenGL resource ID. */
 	GLuint Resource;
 	/** true if the shader has compiled successfully. */
 	bool bSuccessfullyCompiled;
+
 	/** External bindings for this shader. */
 	FOpenGLShaderBindings Bindings;
-	TArray< TRefCountPtr<FRHIUniformBuffer> > BoundUniformBuffers;
+
 	// List of memory copies from RHIUniformBuffer to packed uniforms
 	TArray<FOpenGLUniformBufferCopyInfo> UniformBuffersCopyInfo;
 
@@ -286,14 +324,14 @@ public:
 };
 
 
-typedef TOpenGLShader<FRHIVertexShader, GL_VERTEX_SHADER> FOpenGLVertexShader;
-typedef TOpenGLShader<FRHIPixelShader, GL_FRAGMENT_SHADER> FOpenGLPixelShader;
-typedef TOpenGLShader<FRHIGeometryShader, GL_GEOMETRY_SHADER> FOpenGLGeometryShader;
-typedef TOpenGLShader<FRHIHullShader, GL_TESS_CONTROL_SHADER> FOpenGLHullShader;
-typedef TOpenGLShader<FRHIDomainShader, GL_TESS_EVALUATION_SHADER> FOpenGLDomainShader;
+typedef TOpenGLShader<FRHIVertexShader, GL_VERTEX_SHADER, SF_Vertex> FOpenGLVertexShader;
+typedef TOpenGLShader<FRHIPixelShader, GL_FRAGMENT_SHADER, SF_Pixel> FOpenGLPixelShader;
+typedef TOpenGLShader<FRHIGeometryShader, GL_GEOMETRY_SHADER, SF_Geometry> FOpenGLGeometryShader;
+typedef TOpenGLShader<FRHIHullShader, GL_TESS_CONTROL_SHADER, SF_Hull> FOpenGLHullShader;
+typedef TOpenGLShader<FRHIDomainShader, GL_TESS_EVALUATION_SHADER, SF_Domain> FOpenGLDomainShader;
 
 
-class FOpenGLComputeShader : public TOpenGLShader<FRHIComputeShader, GL_COMPUTE_SHADER>
+class FOpenGLComputeShader : public TOpenGLShader<FRHIComputeShader, GL_COMPUTE_SHADER, SF_Compute>
 {
 public:
 	FOpenGLComputeShader():
@@ -340,7 +378,7 @@ public:
 	 */
 	void CommitPackedGlobals(const FOpenGLLinkedProgram* LinkedProgram, int32 Stage);
 
-	void CommitPackedUniformBuffers(FOpenGLLinkedProgram* LinkedProgram, int32 Stage, const TArray< TRefCountPtr<FRHIUniformBuffer> >& UniformBuffers, const TArray<FOpenGLUniformBufferCopyInfo>& UniformBuffersCopyInfo);
+	void CommitPackedUniformBuffers(FOpenGLLinkedProgram* LinkedProgram, int32 Stage, FUniformBufferRHIRef* UniformBuffers, const TArray<FOpenGLUniformBufferCopyInfo>& UniformBuffersCopyInfo);
 
 private:
 	/** CPU memory block for storing uniform values. */
diff --git a/Engine/Source/Runtime/OpenGLDrv/Public/OpenGLState.h b/Engine/Source/Runtime/OpenGLDrv/Public/OpenGLState.h
--- a/Engine/Source/Runtime/OpenGLDrv/Public/OpenGLState.h
+++ b/Engine/Source/Runtime/OpenGLDrv/Public/OpenGLState.h
@@ -413,6 +413,15 @@ struct FOpenGLRHIState : public FOpenGLCommonState
 	TRefCountPtr<FOpenGLBoundShaderState>	BoundShaderState;
 	FComputeShaderRHIRef					CurrentComputeShader;	
 
+	/** The RHI does not allow more than 14 constant buffers per shader stage due to D3D11 limits. */
+	enum { MAX_UNIFORM_BUFFERS_PER_SHADER_STAGE = 14 };
+
+	/** Track the currently bound uniform buffers. */
+	FUniformBufferRHIRef BoundUniformBuffers[SF_NumFrequencies][MAX_UNIFORM_BUFFERS_PER_SHADER_STAGE];
+
+	/** Bit array to track which uniform buffers have changed since the last draw call. */
+	uint16 DirtyUniformBuffers[SF_NumFrequencies];
+
 	// Used for if(!FOpenGL::SupportsFastBufferData())
 	uint32 UpVertexBufferBytes;
 	uint32 UpIndexBufferBytes;
@@ -479,6 +488,10 @@ struct FOpenGLRHIState : public FOpenGLCommonState
 			ShaderParameters[OGL_SHADER_STAGE_COMPUTE].InitializeResources(FOpenGL::GetMaxComputeUniformComponents() * sizeof(float));
 		}
 
+		for (int32 Frequency = 0; Frequency < SF_NumFrequencies; ++Frequency)
+		{
+			DirtyUniformBuffers[Frequency] = MAX_uint16;
+		}
 	}
 
 	virtual void CleanupResources() OVERRIDE
@@ -486,6 +499,15 @@ struct FOpenGLRHIState : public FOpenGLCommonState
 		delete [] ShaderParameters;
 		ShaderParameters = NULL;
 
+		// Release references to bound uniform buffers.
+		for (int32 Frequency = 0; Frequency < SF_NumFrequencies; ++Frequency)
+		{
+			for (int32 BindIndex = 0; BindIndex < MAX_UNIFORM_BUFFERS_PER_SHADER_STAGE; ++BindIndex)
+			{
+				BoundUniformBuffers[Frequency][BindIndex].SafeRelease();
+			}
+		}
+
 		FOpenGLCommonState::CleanupResources();
 	}
 };
diff --git a/Engine/Source/Runtime/OpenGLDrv/Public/OpenGLUtil.h b/Engine/Source/Runtime/OpenGLDrv/Public/OpenGLUtil.h
--- a/Engine/Source/Runtime/OpenGLDrv/Public/OpenGLUtil.h
+++ b/Engine/Source/Runtime/OpenGLDrv/Public/OpenGLUtil.h
@@ -14,7 +14,7 @@
 #define ENABLE_VERIFY_GL_TRACE 0
 
 /** Set to 1 to verify that the the engine side uniform buffer layout matches the driver side of the GLSL shader*/
-#define ENABLE_UNIFORM_BUFFER_LAYOUT_VERIFICATION (UE_BUILD_DEBUG & (OPENGL_GL3 | OPENGL_GL4))
+#define ENABLE_UNIFORM_BUFFER_LAYOUT_VERIFICATION (0 && UE_BUILD_DEBUG && OPENGL_GL3)
 
 /** Set to 1 to additinally dump uniform buffer layout at shader link time, this assumes ENABLE_UNIFORM_BUFFER_LAYOUT_VERIFICATION == 1 */
 #define ENABLE_UNIFORM_BUFFER_LAYOUT_DUMP 0
diff --git a/Engine/Source/Runtime/RHI/Private/RHI.cpp b/Engine/Source/Runtime/RHI/Private/RHI.cpp
--- a/Engine/Source/Runtime/RHI/Private/RHI.cpp
+++ b/Engine/Source/Runtime/RHI/Private/RHI.cpp
@@ -64,6 +64,11 @@ static FAutoConsoleCommandWithOutputDevice GDumpRHIMemoryCmd(
  * RHI configuration settings.
  */
 
+static TAutoConsoleVariable<int32> ResourceTableCachingCvar(
+	TEXT("rhi.ResourceTableCaching"),
+	1,
+	TEXT("If 1, the RHI will cache resource table contents within a frame. Otherwise resource tables are rebuilt for every draw call.")
+	);
 static TAutoConsoleVariable<int32> GSaveScreenshotAfterProfilingGPUCVar(
 	TEXT("RHI.SaveScreenshotAfterProfilingGPU"),
 	1,
diff --git a/Engine/Source/Runtime/RHI/Private/RHICommandList.cpp b/Engine/Source/Runtime/RHI/Private/RHICommandList.cpp
new file mode 100644
--- /dev/null
+++ b/Engine/Source/Runtime/RHI/Private/RHICommandList.cpp
@@ -0,0 +1,287 @@
+// Copyright 1998-2014 Epic Games, Inc. All Rights Reserved.
+
+/*=============================================================================
+	RHI.cpp: Render Hardware Interface implementation.
+=============================================================================*/
+
+#include "RHI.h"
+#include "RHICommandList.h"
+
+DEFINE_STAT(STAT_RHICmdListExecuteTime);
+DEFINE_STAT(STAT_RHICmdListEnqueueTime);
+DEFINE_STAT(STAT_RHICmdListMemory);
+DEFINE_STAT(STAT_RHICmdListCount);
+DEFINE_STAT(STAT_RHICounterTEMP);
+
+static int32 GRHICmd = 0;
+static FAutoConsoleVariableRef CVarRHICmd(
+	TEXT("r.RHICmd"),
+	GRHICmd,
+	TEXT("0: Disables the RHI Command List mode, running regular RHI calls (default)\n")
+	TEXT("1: Enables the RHI Command List mode")
+	);
+
+static int32 GRHISkip = 0;
+static FAutoConsoleVariableRef CVarRHICmdSkip(
+	TEXT("r.RHISkip"),
+	GRHISkip,
+	TEXT("0: No effect (default)\n")
+	TEXT("1: Skips executing RHI Command List")
+	);
+
+static int32 GRHICmdExec = 0;
+static FAutoConsoleVariableRef CVarRHICmdExec(
+	TEXT("r.RHIExec"),
+	GRHICmdExec,
+	TEXT("0: Executes RHI Command List using RHI calls (default)\n")
+	TEXT("1: Executes RHI Command List using platform specific APIs")
+	);
+
+static int32 GRHICmdMem = 0;
+static FAutoConsoleVariableRef CVarRHICmdMem(
+	TEXT("r.RHIMem"),
+	GRHICmdMem,
+	TEXT("0: No effect(default)\n")
+	TEXT("1: Processes the RHI Command List but doesn't execute commands")
+	);
+
+RHI_API FRHICommandListExecutor GRHICommandList;
+
+//@todo-rco: Temp hack!
+static FRHICommandList GCommandList;
+
+
+FRHICommandListExecutor::FRHICommandListExecutor() :
+	bLock(false)
+{
+}
+
+void FRHICommandListExecutor::ExecuteAndFreeList(FRHICommandList* CmdList)
+{
+	if (!CmdList)
+	{
+		// HACK!
+		return;
+	}
+
+	SCOPE_CYCLE_COUNTER(STAT_RHICmdListExecuteTime);
+
+	check(CmdList);
+	check(CmdList->CanAddCommand());
+
+	FScopeLock Lock(&CriticalSection);
+	check(bLock);
+	bLock = false;
+
+	CmdList->State = FRHICommandList::Executing;
+
+	auto* CmdPtr = &CmdList->Mem[0];
+	auto* CmdTail = CmdList->GetTail();
+	INC_MEMORY_STAT_BY(STAT_RHICmdListMemory, (CmdTail - CmdPtr));
+
+	static IConsoleVariable* RHICmdListCVar = IConsoleManager::Get().FindConsoleVariable(TEXT("r.RHISkip"));
+	bool bSkipRHICmdList = (RHICmdListCVar && RHICmdListCVar->GetInt() != 0);
+
+	if (!bSkipRHICmdList)
+	{
+		static IConsoleVariable* RHIExecCVar = IConsoleManager::Get().FindConsoleVariable(TEXT("r.RHIExec"));
+		if (RHIExecCVar && RHIExecCVar->GetInt() != 0)
+		{
+			RHIExecuteCommandList(CmdList);
+		}
+		else
+		{
+			static IConsoleVariable* RHIMemCVar = IConsoleManager::Get().FindConsoleVariable(TEXT("r.RHIMem"));
+			bool bOnlySkipAroundMem = (RHIMemCVar && RHIMemCVar->GetInt() != 0);
+
+			uint32 NumCommands = 0;
+			while (CmdPtr < CmdTail)
+			{
+				auto* Cmd = (FRHICommand*)CmdPtr;
+				switch (Cmd->Type)
+				{
+				case ERCT_NopBlob:
+					{
+						// Nop
+						auto* RHICmd = (FRHICommandNopBlob*)Cmd;
+						CmdPtr += sizeof(FRHICommandNopBlob) + RHICmd->Size;
+					}
+					break;
+				case ERCT_SetRasterizerState:
+					{
+						auto* RHICmd = (FRHICommandSetRasterizerState*)Cmd;
+						if (!bOnlySkipAroundMem)
+						{
+							RHISetRasterizerState(RHICmd->State);
+						}
+						CmdPtr += sizeof(FRHICommandSetRasterizerState);
+						RHICmd->State->Release();
+					}
+					break;
+				case ERCT_SetShaderParameter:
+					{
+						auto* RHICmd = (FRHICommandSetShaderParameter*)Cmd;
+						if (!bOnlySkipAroundMem)
+						{
+							switch (RHICmd->ShaderFrequency)
+							{
+							case SF_Vertex: RHISetShaderParameter((FVertexShaderRHIParamRef)RHICmd->Shader, RHICmd->BufferIndex, RHICmd->BaseIndex, RHICmd->NumBytes, RHICmd->NewValue); break;
+							case SF_Hull: RHISetShaderParameter((FHullShaderRHIParamRef)RHICmd->Shader, RHICmd->BufferIndex, RHICmd->BaseIndex, RHICmd->NumBytes, RHICmd->NewValue); break;
+							case SF_Domain: RHISetShaderParameter((FDomainShaderRHIParamRef)RHICmd->Shader, RHICmd->BufferIndex, RHICmd->BaseIndex, RHICmd->NumBytes, RHICmd->NewValue); break;
+							case SF_Geometry: RHISetShaderParameter((FGeometryShaderRHIParamRef)RHICmd->Shader, RHICmd->BufferIndex, RHICmd->BaseIndex, RHICmd->NumBytes, RHICmd->NewValue); break;
+							case SF_Pixel: RHISetShaderParameter((FPixelShaderRHIParamRef)RHICmd->Shader, RHICmd->BufferIndex, RHICmd->BaseIndex, RHICmd->NumBytes, RHICmd->NewValue); break;
+							case SF_Compute: RHISetShaderParameter((FComputeShaderRHIParamRef)RHICmd->Shader, RHICmd->BufferIndex, RHICmd->BaseIndex, RHICmd->NumBytes, RHICmd->NewValue); break;
+							default: check(0); break;
+							}
+						}
+						CmdPtr += sizeof(FRHICommandSetShaderParameter);
+					}
+					break;
+				case ERCT_SetShaderUniformBuffer:
+					{
+						auto* RHICmd = (FRHICommandSetShaderUniformBuffer*)Cmd;
+						if (!bOnlySkipAroundMem)
+						{
+							switch (RHICmd->ShaderFrequency)
+							{
+							case SF_Vertex: RHISetShaderUniformBuffer((FVertexShaderRHIParamRef)RHICmd->Shader, RHICmd->BaseIndex, RHICmd->UniformBuffer); break;
+							case SF_Hull: RHISetShaderUniformBuffer((FHullShaderRHIParamRef)RHICmd->Shader, RHICmd->BaseIndex, RHICmd->UniformBuffer); break;
+							case SF_Domain: RHISetShaderUniformBuffer((FDomainShaderRHIParamRef)RHICmd->Shader, RHICmd->BaseIndex, RHICmd->UniformBuffer); break;
+							case SF_Geometry: RHISetShaderUniformBuffer((FGeometryShaderRHIParamRef)RHICmd->Shader, RHICmd->BaseIndex, RHICmd->UniformBuffer); break;
+							case SF_Pixel: RHISetShaderUniformBuffer((FPixelShaderRHIParamRef)RHICmd->Shader, RHICmd->BaseIndex, RHICmd->UniformBuffer); break;
+							case SF_Compute: RHISetShaderUniformBuffer((FComputeShaderRHIParamRef)RHICmd->Shader, RHICmd->BaseIndex, RHICmd->UniformBuffer); break;
+							default: check(0); break;
+							}
+						}
+						CmdPtr += sizeof(FRHICommandSetShaderUniformBuffer);
+						RHICmd->Shader->Release(); 
+						RHICmd->UniformBuffer->Release();
+					}
+					break;
+				case ERCT_SetShaderSampler:
+					{
+						auto* RHICmd = (FRHICommandSetShaderSampler*)Cmd;
+						if (!bOnlySkipAroundMem)
+						{
+							switch (RHICmd->ShaderFrequency)
+							{
+							case SF_Vertex: RHISetShaderSampler((FVertexShaderRHIParamRef)RHICmd->Shader, RHICmd->SamplerIndex, RHICmd->Sampler); break;
+							case SF_Hull: RHISetShaderSampler((FHullShaderRHIParamRef)RHICmd->Shader, RHICmd->SamplerIndex, RHICmd->Sampler); break;
+							case SF_Domain: RHISetShaderSampler((FDomainShaderRHIParamRef)RHICmd->Shader, RHICmd->SamplerIndex, RHICmd->Sampler); break;
+							case SF_Geometry: RHISetShaderSampler((FGeometryShaderRHIParamRef)RHICmd->Shader, RHICmd->SamplerIndex, RHICmd->Sampler); break;
+							case SF_Pixel: RHISetShaderSampler((FPixelShaderRHIParamRef)RHICmd->Shader, RHICmd->SamplerIndex, RHICmd->Sampler); break;
+							case SF_Compute: RHISetShaderSampler((FComputeShaderRHIParamRef)RHICmd->Shader, RHICmd->SamplerIndex, RHICmd->Sampler); break;
+							default: check(0); break;
+							}
+						}
+						CmdPtr += sizeof(FRHICommandSetShaderSampler);
+						RHICmd->Shader->Release(); 
+						RHICmd->Sampler->Release();
+					}
+					break;
+				case ERCT_SetShaderTexture:
+					{
+						auto* RHICmd = (FRHICommandSetShaderTexture*)Cmd;
+						if (!bOnlySkipAroundMem)
+						{
+							switch (RHICmd->ShaderFrequency)
+							{
+							case SF_Vertex: RHISetShaderTexture((FVertexShaderRHIParamRef)RHICmd->Shader, RHICmd->TextureIndex, RHICmd->Texture); break;
+							case SF_Hull: RHISetShaderTexture((FHullShaderRHIParamRef)RHICmd->Shader, RHICmd->TextureIndex, RHICmd->Texture); break;
+							case SF_Domain: RHISetShaderTexture((FDomainShaderRHIParamRef)RHICmd->Shader, RHICmd->TextureIndex, RHICmd->Texture); break;
+							case SF_Geometry: RHISetShaderTexture((FGeometryShaderRHIParamRef)RHICmd->Shader, RHICmd->TextureIndex, RHICmd->Texture); break;
+							case SF_Pixel: RHISetShaderTexture((FPixelShaderRHIParamRef)RHICmd->Shader, RHICmd->TextureIndex, RHICmd->Texture); break;
+							case SF_Compute: RHISetShaderTexture((FComputeShaderRHIParamRef)RHICmd->Shader, RHICmd->TextureIndex, RHICmd->Texture); break;
+							default: check(0); break;
+							}
+						}
+						CmdPtr += sizeof(FRHICommandSetShaderTexture);
+						RHICmd->Shader->Release(); 
+						RHICmd->Texture->Release();
+					}
+					break;
+				case ERCT_DrawPrimitive:
+					{
+						auto* RHICmd = (FRHICommandDrawPrimitive*)Cmd;
+						if (!bOnlySkipAroundMem)
+						{
+							RHIDrawPrimitive(RHICmd->PrimitiveType, RHICmd->BaseVertexIndex, RHICmd->NumPrimitives, RHICmd->NumInstances);
+						}
+						CmdPtr += sizeof(FRHICommandDrawPrimitive);
+					}
+					break;
+				case ERCT_DrawIndexedPrimitive:
+					{
+						auto* RHICmd = (FRHICommandDrawIndexedPrimitive*)Cmd;
+						if (!bOnlySkipAroundMem)
+						{
+							RHIDrawIndexedPrimitive(RHICmd->IndexBuffer, RHICmd->PrimitiveType, RHICmd->BaseVertexIndex, RHICmd->MinIndex, RHICmd->NumVertices, RHICmd->StartIndex, RHICmd->NumPrimitives, RHICmd->NumInstances);
+						}
+						CmdPtr += sizeof(FRHICommandDrawIndexedPrimitive);
+						RHICmd->IndexBuffer->Release();
+					}
+					break;
+				case ERCT_SetBoundShaderState:
+					{
+						auto* RHICmd = (FRHICommandSetBoundShaderState*)Cmd;
+						if (!bOnlySkipAroundMem)
+						{
+							RHISetBoundShaderState(RHICmd->BoundShaderState);
+						}
+						CmdPtr += sizeof(FRHICommandSetBoundShaderState);
+						RHICmd->BoundShaderState->Release();
+					}
+					break;
+				case ERCT_SetBlendState:
+					{
+						auto* RHICmd = (FRHICommandSetBlendState*)Cmd;
+						if (!bOnlySkipAroundMem)
+						{
+							RHISetBlendState(RHICmd->State, RHICmd->BlendFactor);
+						}
+						CmdPtr += sizeof(FRHICommandSetBlendState);
+						RHICmd->State->Release();
+					}
+					break;
+				case ERCT_SetStreamSource:
+					{
+						auto* RHICmd = (FRHICommandSetStreamSource*)Cmd;
+						if (!bOnlySkipAroundMem)
+						{
+							RHISetStreamSource(RHICmd->StreamIndex, RHICmd->VertexBuffer, RHICmd->Stride, RHICmd->Offset);
+						}
+						CmdPtr += sizeof(FRHICommandSetStreamSource);
+						RHICmd->VertexBuffer->Release();
+					}
+					break;
+				default:
+					check(0);
+				}
+
+				CmdPtr = Align(CmdPtr, FRHICommandList::Alignment);
+				++NumCommands;
+			}
+
+			check(CmdList->NumCommands == NumCommands);
+		}
+	}
+
+	INC_DWORD_STAT_BY(STAT_RHICmdListCount, CmdList->NumCommands);
+	CmdList->Reset();
+}
+
+FRHICommandList* FRHICommandListExecutor::CreateList()
+{
+	FScopeLock Lock(&CriticalSection);
+	check(!bLock);
+	bLock = true;
+	check(GCommandList.State == FRHICommandList::Kicked);
+	check(GCommandList.NumCommands == 0);
+	GCommandList.State = FRHICommandList::Ready;
+	return &GCommandList;
+}
+
+void FRHICommandListExecutor::Verify()
+{
+	check(!bLock);
+}
diff --git a/Engine/Source/Runtime/RHI/Public/BoundShaderStateCache.h b/Engine/Source/Runtime/RHI/Public/BoundShaderStateCache.h
--- a/Engine/Source/Runtime/RHI/Public/BoundShaderStateCache.h
+++ b/Engine/Source/Runtime/RHI/Public/BoundShaderStateCache.h
@@ -63,6 +63,15 @@ public:
 			GetTypeHash(Key.GeometryShader);
 	}
 
+	/**
+	 * Get the RHI shader for the given frequency.
+	 */
+	FORCEINLINE FVertexShaderRHIParamRef   GetVertexShader() const   { return VertexShader; }
+	FORCEINLINE FPixelShaderRHIParamRef    GetPixelShader() const    { return PixelShader; }
+	FORCEINLINE FHullShaderRHIParamRef     GetHullShader() const     { return HullShader; }
+	FORCEINLINE FDomainShaderRHIParamRef   GetDomainShader() const   { return DomainShader; }
+	FORCEINLINE FGeometryShaderRHIParamRef GetGeometryShader() const { return GeometryShader; }
+
 private:
 	/**
 	 * Note: We intentionally do use ...Ref, not ...ParamRef to get 
@@ -124,6 +133,15 @@ public:
 	/** Destructor.  Removes the bound shader state from the cache. */
 	~FCachedBoundShaderStateLink();
 
+	/**
+	 * Get the RHI shader for the given frequency.
+	 */
+	FORCEINLINE FVertexShaderRHIParamRef   GetVertexShader() const   { return Key.GetVertexShader(); }
+	FORCEINLINE FPixelShaderRHIParamRef    GetPixelShader() const    { return Key.GetPixelShader(); }
+	FORCEINLINE FHullShaderRHIParamRef     GetHullShader() const     { return Key.GetHullShader(); }
+	FORCEINLINE FDomainShaderRHIParamRef   GetDomainShader() const   { return Key.GetDomainShader(); }
+	FORCEINLINE FGeometryShaderRHIParamRef GetGeometryShader() const { return Key.GetGeometryShader(); }
+
 private:
 
 	FBoundShaderStateKey Key;
diff --git a/Engine/Source/Runtime/RHI/Public/RHI.h b/Engine/Source/Runtime/RHI/Public/RHI.h
--- a/Engine/Source/Runtime/RHI/Public/RHI.h
+++ b/Engine/Source/Runtime/RHI/Public/RHI.h
@@ -8,8 +8,8 @@
 #define __RHI_h__
 
 #include "Core.h"
+#include "RHIDefinitions.h"
 #include "StaticArray.h"
-#include "Runtime/Engine/Public/PixelFormat.h" // for EPixelFormat
 
 // Forward declarations.
 class FSceneView;
@@ -36,25 +36,12 @@ namespace RHIConfig
 /** True if the render hardware has been initialized. */
 extern RHI_API bool GIsRHIInitialized;
 
+class RHI_API FRHICommandList;
+
 /**
  * RHI capabilities.
  */
 
-/** Maximum number of miplevels in a texture. */
-enum { MAX_TEXTURE_MIP_COUNT = 14 };
-
-/** The maximum number of vertex elements which can be used by a vertex declaration. */
-enum { MaxVertexElementCount = 16 };
-
-/** The alignment in bytes between elements of array shader parameters. */
-enum { ShaderArrayElementAlignBytes = 16 };
-
-/** The number of render-targets that may be simultaneously written to. */
-enum { MaxSimultaneousRenderTargets = 8 };
-
-/** The number of UAVs that may be simultaneously bound to a shader. */
-enum { MaxSimultaneousUAVs = 8 };
-
 
 /** The maximum number of mip-maps that a texture can contain. 	*/
 extern	RHI_API int32		GMaxTextureMipCount;
@@ -214,59 +201,6 @@ extern RHI_API int32 GNumPrimitivesDrawnRHI;
 /** Called once per frame only from within an RHI. */
 extern RHI_API void RHIPrivateBeginFrame();
 
-/** @warning: update *LegacyShaderPlatform* when the below changes */
-enum EShaderPlatform
-{
-	SP_PCD3D_SM5		= 0,
-	SP_OPENGL_SM4		= 1,
-	SP_PS4				= 2,
-	/** Used when running in Feature Level ES2 in OpenGL. */
-	SP_OPENGL_PCES2		= 3,
-	SP_XBOXONE			= 4,
-	SP_PCD3D_SM4		= 5,
-	SP_OPENGL_SM5		= 6,
-	/** Used when running in Feature Level ES2 in D3D11. */
-	SP_PCD3D_ES2		= 7,
-	SP_OPENGL_ES2		= 8,
-	SP_OPENGL_ES2_WEBGL = 9, 
-	SP_OPENGL_ES2_IOS	= 10,
-	
-	SP_NumPlatforms		= 11,
-	SP_NumBits			= 4,
-};
-checkAtCompileTime(SP_NumPlatforms <= (1 << SP_NumBits), SP_NumBits_TooSmall);
-
-enum ERenderQueryType
-{
-	// e.g. WaitForFrameEventCompletion()
-	RQT_Undefined,
-	// Result is the number of samples that are not culled (divide by MSAACount to get pixels)
-	RQT_Occlusion,
-	// Result is time in micro seconds = 1/1000 ms = 1/1000000 sec
-	RQT_AbsoluteTime,
-};
-
-inline bool IsPCPlatform(const EShaderPlatform Platform)
-{
-	return Platform == SP_PCD3D_SM5 || Platform == SP_PCD3D_SM4 || Platform == SP_PCD3D_ES2 || Platform ==  SP_OPENGL_SM4 || Platform == SP_OPENGL_SM5 || Platform == SP_OPENGL_PCES2;
-}
-
-/** Whether the shader platform corresponds to the ES2 feature level. */
-inline bool IsES2Platform(const EShaderPlatform Platform)
-{
-	return Platform == SP_PCD3D_ES2 || Platform == SP_OPENGL_PCES2 || Platform == SP_OPENGL_ES2 || Platform == SP_OPENGL_ES2_WEBGL || Platform == SP_OPENGL_ES2_IOS; 
-}
-
-inline bool IsOpenGLPlatform(const EShaderPlatform Platform)
-{
-	return Platform == SP_OPENGL_SM4 || Platform == SP_OPENGL_SM5 || Platform == SP_OPENGL_PCES2 || Platform == SP_OPENGL_ES2 || Platform == SP_OPENGL_ES2_WEBGL || Platform == SP_OPENGL_ES2_IOS;
-}
-
-inline bool IsConsolePlatform(const EShaderPlatform Platform)
-{
-	return Platform == SP_PS4 || Platform == SP_XBOXONE;
-}
-
 RHI_API FName LegacyShaderPlatformToShaderFormat(EShaderPlatform Platform);
 RHI_API EShaderPlatform ShaderFormatToLegacyShaderPlatform(FName ShaderFormat);
 
@@ -286,27 +220,6 @@ inline FMatrix AdjustProjectionMatrixForRHI(const FMatrix& InProjectionMatrix)
 /** Current shader platform. */
 extern RHI_API EShaderPlatform GRHIShaderPlatform;
 
-/**
- * The RHI's feature level indicates what level of support can be relied upon.
- * Note: these are named after graphics API's like ES2 but a feature level can be used with a different API (eg ERHIFeatureLevel::ES2 on D3D11)
- * As long as the graphics API supports all the features of the feature level (eg no ERHIFeatureLevel::SM5 on OpenGL ES2)
- */
-namespace ERHIFeatureLevel
-{
-	enum Type
-	{
-		/** Feature level defined by the core capabilities of OpenGL ES2. */
-		ES2,
-		/** Feature level defined by the capabilities of DX9 Shader Model 3. */
-		SM3,
-		/** Feature level defined by the capabilities of DX10 Shader Model 4. */
-		SM4,
-		/** Feature level defined by the capabilities of DX11 Shader Model 5. */
-		SM5,
-		Num
-	};
-};
-
 /** Finds a corresponding ERHIFeatureLevel::Type given an FName, or returns false if one could not be found. */
 extern RHI_API bool GetFeatureLevelFromName(FName Name, ERHIFeatureLevel::Type& OutFeatureLevel);
 
@@ -330,249 +243,6 @@ extern RHI_API void SetMaxRHIFeatureLevel(ERHIFeatureLevel::Type InType);
 /** Table for finding out which shader platform corresponds to a given feature level for this RHI. */
 extern RHI_API EShaderPlatform GShaderPlatformForFeatureLevel[ERHIFeatureLevel::Num];
 
-inline ERHIFeatureLevel::Type GetMaxSupportedFeatureLevel(EShaderPlatform InShaderPlatform)
-{
-	switch(InShaderPlatform)
-	{
-	case SP_PCD3D_SM5:
-	case SP_OPENGL_SM5:
-	case SP_PS4:
-	case SP_XBOXONE:
-		return ERHIFeatureLevel::SM5;
-	case SP_PCD3D_SM4:
-	case SP_OPENGL_SM4:
-		return ERHIFeatureLevel::SM4;
-	case SP_PCD3D_ES2:
-	case SP_OPENGL_PCES2:
-	case SP_OPENGL_ES2:
-	case SP_OPENGL_ES2_WEBGL:
-	case  SP_OPENGL_ES2_IOS:
-		return ERHIFeatureLevel::ES2;
-	default:
-		check(0);
-		return ERHIFeatureLevel::Num;
-	}	
-}
-
-/** Returns true if the feature level is supported by the shader platform. */
-inline bool IsFeatureLevelSupported(EShaderPlatform InShaderPlatform, ERHIFeatureLevel::Type InFeatureLevel)
-{
-	switch(InShaderPlatform)
-	{
-	case SP_PCD3D_SM5:
-		return InFeatureLevel <= ERHIFeatureLevel::SM5;
-	case SP_PCD3D_SM4:
-		return InFeatureLevel <= ERHIFeatureLevel::SM4;
-	case SP_PCD3D_ES2:
-		return InFeatureLevel <= ERHIFeatureLevel::ES2;
-	case SP_OPENGL_PCES2:
-		return InFeatureLevel <= ERHIFeatureLevel::ES2;
-	case SP_OPENGL_ES2:
-		return InFeatureLevel <= ERHIFeatureLevel::ES2;
-	case SP_OPENGL_ES2_WEBGL: 
-		return InFeatureLevel <= ERHIFeatureLevel::ES2;
-	case SP_OPENGL_ES2_IOS: 
-		return InFeatureLevel <= ERHIFeatureLevel::ES2;
-	case SP_OPENGL_SM4:
-		return InFeatureLevel <= ERHIFeatureLevel::SM4;
-	case SP_OPENGL_SM5:
-		return InFeatureLevel <= ERHIFeatureLevel::SM5;
-	case SP_PS4:
-		return InFeatureLevel <= ERHIFeatureLevel::SM5;
-	case SP_XBOXONE:
-		return InFeatureLevel <= ERHIFeatureLevel::SM5;
-	default:
-		return false;
-	}	
-}
-
-inline bool RHISupportsTessellation(const EShaderPlatform Platform)
-{
-	if (IsFeatureLevelSupported(Platform, ERHIFeatureLevel::SM5))
-	{
-		return ((Platform == SP_PCD3D_SM5) || (Platform == SP_XBOXONE) || (Platform == SP_OPENGL_SM5));
-	}
-	return false;
-}
-
-inline uint32 GetFeatureLevelMaxTextureSamplers(ERHIFeatureLevel::Type FeatureLevel)
-{
-	if (FeatureLevel == ERHIFeatureLevel::ES2)
-	{
-		return 8;
-	}
-	else
-	{
-		return 16;
-	}
-}
-
-inline int32 GetFeatureLevelMaxNumberOfBones(ERHIFeatureLevel::Type FeatureLevel)
-{
-	switch (FeatureLevel)
-	{
-	case ERHIFeatureLevel::ES2:
-	case ERHIFeatureLevel::SM3:
-		return 75;
-	case ERHIFeatureLevel::SM4:
-	case ERHIFeatureLevel::SM5:
-		return 256;
-	default:
-		check(0);
-	}
-
-	return 0;
-}
-
-//
-// Common RHI definitions.
-//
-
-enum ESamplerFilter
-{
-	SF_Point,
-	SF_Bilinear,
-	SF_Trilinear,
-	SF_AnisotropicPoint,
-	SF_AnisotropicLinear,
-};
-
-enum ESamplerAddressMode
-{
-	AM_Wrap,
-	AM_Clamp,
-	AM_Mirror,
-	/** Not supported on all platforms */
-	AM_Border
-};
-
-enum ESamplerCompareFunction
-{
-	SCF_Never,
-	SCF_Less
-};
-
-enum ERasterizerFillMode
-{
-	FM_Point,
-	FM_Wireframe,
-	FM_Solid
-};
-
-enum ERasterizerCullMode
-{
-	CM_None,
-	CM_CW,
-	CM_CCW
-};
-
-enum EColorWriteMask
-{
-	CW_RED   = 0x01,
-	CW_GREEN = 0x02,
-	CW_BLUE  = 0x04,
-	CW_ALPHA = 0x08,
-
-	CW_NONE  = 0,
-	CW_RGB   = CW_RED | CW_GREEN | CW_BLUE,
-	CW_RGBA  = CW_RED | CW_GREEN | CW_BLUE | CW_ALPHA,
-	CW_RG    = CW_RED | CW_GREEN,
-	CW_BA    = CW_BLUE | CW_ALPHA,
-};
-
-enum ECompareFunction
-{
-	CF_Less,
-	CF_LessEqual,
-	CF_Greater,
-	CF_GreaterEqual,
-	CF_Equal,
-	CF_NotEqual,
-	CF_Never,
-	CF_Always
-};
-
-enum EStencilOp
-{
-	SO_Keep,
-	SO_Zero,
-	SO_Replace,
-	SO_SaturatedIncrement,
-	SO_SaturatedDecrement,
-	SO_Invert,
-	SO_Increment,
-	SO_Decrement
-};
-
-enum EBlendOperation
-{
-	BO_Add,
-	BO_Subtract,
-	BO_Min,
-	BO_Max,
-    BO_ReverseSubtract,
-};
-
-enum EBlendFactor
-{
-	BF_Zero,
-	BF_One,
-	BF_SourceColor,
-	BF_InverseSourceColor,
-	BF_SourceAlpha,
-	BF_InverseSourceAlpha,
-	BF_DestAlpha,
-	BF_InverseDestAlpha,
-	BF_DestColor,
-	BF_InverseDestColor,
-	BF_ConstantBlendFactor,
-	BF_InverseConstantBlendFactor
-};
-
-enum EVertexElementType
-{
-	VET_None,
-	VET_Float1,
-	VET_Float2,
-	VET_Float3,
-	VET_Float4,
-	VET_PackedNormal,	// FPackedNormal
-	VET_UByte4,
-	VET_UByte4N,
-	VET_Color,
-	VET_Short2,
-	VET_Short4,
-	VET_Short2N,		// 16 bit word normalized to (value/32767.0,value/32767.0,0,0,1)
-	VET_Half2,			// 16 bit float using 1 bit sign, 5 bit exponent, 10 bit mantissa 
-	VET_Half4,
-	VET_MAX
-};
-
-enum ECubeFace
-{
-	CubeFace_PosX=0,
-	CubeFace_NegX,
-	CubeFace_PosY,
-	CubeFace_NegY,
-	CubeFace_PosZ,
-	CubeFace_NegZ,
-	CubeFace_MAX
-};
-
-enum EUniformBufferUsage
-{
-	// the uniform buffer is temporary, used for a single draw call then discarded
-	UniformBuffer_SingleUse = 0,
-	// the uniform buffer is used for multiple draw calls, possibly across multiple frames
-	UniformBuffer_MultiUse,
-};
-
-enum EResourceLockMode
-{
-	RLM_ReadOnly,
-	RLM_WriteOnly,
-	RLM_Num
-};
 
 inline FArchive& operator <<(FArchive& Ar, EResourceLockMode& LockMode)
 {
@@ -582,19 +252,6 @@ inline FArchive& operator <<(FArchive& Ar, EResourceLockMode& LockMode)
 	return Ar;
 }
 
-/** limited to 8 types in FReadSurfaceDataFlags */
-enum ERangeCompressionMode
-{
-	// 0 .. 1
-	RCM_UNorm,
-	// -1 .. 1
-	RCM_SNorm,
-	// 0 .. 1 unless there are smaller values than 0 or bigger values than 1, then the range is extended to the minimum or the maximum of the values
-	RCM_MinMaxNorm,
-	// minimum .. maximum (each channel independent)
-	RCM_MinMax,
-};
-
 /** to customize the RHIReadSurfaceData() output */
 class FReadSurfaceDataFlags
 {
@@ -987,98 +644,6 @@ public:
 	}
 };
 
-enum EPrimitiveType
-{
-	PT_TriangleList,
-	PT_TriangleStrip,
-	PT_LineList,
-	PT_QuadList,
-	PT_PointList,
-	PT_1_ControlPointPatchList,
-	PT_2_ControlPointPatchList,
-	PT_3_ControlPointPatchList,
-	PT_4_ControlPointPatchList,
-	PT_5_ControlPointPatchList,
-	PT_6_ControlPointPatchList,
-	PT_7_ControlPointPatchList,
-	PT_8_ControlPointPatchList,
-	PT_9_ControlPointPatchList,
-	PT_10_ControlPointPatchList,
-	PT_11_ControlPointPatchList,
-	PT_12_ControlPointPatchList,
-	PT_13_ControlPointPatchList,
-	PT_14_ControlPointPatchList,
-	PT_15_ControlPointPatchList,
-	PT_16_ControlPointPatchList,
-	PT_17_ControlPointPatchList,
-	PT_18_ControlPointPatchList,
-	PT_19_ControlPointPatchList,
-	PT_20_ControlPointPatchList,
-	PT_21_ControlPointPatchList,
-	PT_22_ControlPointPatchList,
-	PT_23_ControlPointPatchList,
-	PT_24_ControlPointPatchList,
-	PT_25_ControlPointPatchList,
-	PT_26_ControlPointPatchList,
-	PT_27_ControlPointPatchList,
-	PT_28_ControlPointPatchList,
-	PT_29_ControlPointPatchList,
-	PT_30_ControlPointPatchList,
-	PT_31_ControlPointPatchList,
-	PT_32_ControlPointPatchList,
-	PT_Num,
-	PT_NumBits = 6
-};
-
-checkAtCompileTime( PT_Num <= (1 << 8), EPrimitiveType_DoesntFitInAByte );
-checkAtCompileTime( PT_Num <= (1 << PT_NumBits), PT_NumBits_TooSmall );
-
-/**
- *	Resource usage flags - for vertex and index buffers.
- */
-enum EBufferUsageFlags
-{
-	// Mutually exclusive write-frequency flags
-	BUF_Static            = 0x0001, // The buffer will be written to once.
-	BUF_Dynamic           = 0x0002, // The buffer will be written to occasionally.
-	BUF_Volatile          = 0x0004, // The buffer will be written to frequently.
-
-	// Mutually exclusive bind flags.
-	BUF_UnorderedAccess   = 0x0008, // Allows an unordered access view to be created for the buffer.
-
-	/** Create a byte address buffer, which is basically a structured buffer with a uint32 type. */
-	BUF_ByteAddressBuffer = 0x0020,
-	/** Create a structured buffer with an atomic UAV counter. */
-	BUF_UAVCounter        = 0x0040,
-	/** Create a buffer that can be bound as a stream output target. */
-	BUF_StreamOutput      = 0x0080,
-	/** Create a buffer which contains the arguments used by DispatchIndirect or DrawIndirect. */
-	BUF_DrawIndirect      = 0x0100,
-	/** 
-	 * Create a buffer that can be bound as a shader resource. 
-	 * This is only needed for buffer types which wouldn't ordinarily be used as a shader resource, like a vertex buffer.
-	 */
-	BUF_ShaderResource    = 0x0200,
-
-	/**
-	 * Request that this buffer is directly CPU accessible
-	 * (@todo josh: this is probably temporary and will go away in a few months)
-	 */
-	BUF_KeepCPUAccessible = 0x0400,
-
-	/**
-	 * Provide information that this buffer will contain only one vertex, which should be delivered to every primitive drawn.
-	 * This is necessary for OpenGL implementations, which need to handle this case very differently (and can't handle GL_HALF_FLOAT in such vertices at all).
-	 */
-	BUF_ZeroStride        = 0x0800,
-
-	/** Buffer should go in fast vram (hint only) */
-	BUF_FastVRAM          = 0x1000,
-
-	// Helper bit-masks
-	BUF_AnyDynamic      = (BUF_Dynamic|BUF_Volatile),
-};
-
 /**
  *	Screen Resolution
  */
@@ -1124,102 +689,6 @@ struct FViewportBounds
 
 typedef TArray<FScreenResolutionRHI>	FScreenResolutionArray;
 
-#define ENUM_RHI_RESOURCE_TYPES(EnumerationMacro) \
-	EnumerationMacro(SamplerState,None) \
-	EnumerationMacro(RasterizerState,None) \
-	EnumerationMacro(DepthStencilState,None) \
-	EnumerationMacro(BlendState,None) \
-	EnumerationMacro(VertexDeclaration,None) \
-	EnumerationMacro(VertexShader,None) \
-	EnumerationMacro(HullShader,None) \
-	EnumerationMacro(DomainShader,None) \
-	EnumerationMacro(PixelShader,None) \
-	EnumerationMacro(GeometryShader,None) \
-	EnumerationMacro(ComputeShader,None) \
-	EnumerationMacro(BoundShaderState,None) \
-	EnumerationMacro(UniformBuffer,None) \
-	EnumerationMacro(IndexBuffer,None) \
-	EnumerationMacro(VertexBuffer,None) \
-	EnumerationMacro(StructuredBuffer,None) \
-	EnumerationMacro(Texture,None) \
-	EnumerationMacro(Texture2D,Texture) \
-	EnumerationMacro(Texture2DArray,Texture) \
-	EnumerationMacro(Texture3D,Texture) \
-	EnumerationMacro(TextureCube,Texture) \
-	EnumerationMacro(RenderQuery,None) \
-	EnumerationMacro(Viewport,None) \
-	EnumerationMacro(UnorderedAccessView,None) \
-	EnumerationMacro(ShaderResourceView,None)
-
-/** An enumeration of the different RHI reference types. */
-enum ERHIResourceType
-{
-	RRT_None,
-
-#define DECLARE_RESOURCETYPE_ENUM(Type,ParentType) RRT_##Type,
-	ENUM_RHI_RESOURCE_TYPES(DECLARE_RESOURCETYPE_ENUM)
-#undef DECLARE_RESOURCETYPE_ENUM
-
-	RRT_Num
-};
-
-/** Flags used for texture creation */
-enum ETextureCreateFlags
-{
-	TexCreate_None					= 0,
-
-	// Texture can be used as a render target
-	TexCreate_RenderTargetable		= 1<<0,
-	// Texture can be used as a resolve target
-	TexCreate_ResolveTargetable		= 1<<1,
-	// Texture can be used as a depth-stencil target.
-	TexCreate_DepthStencilTargetable= 1<<2,
-	// Texture can be used as a shader resource.
-	TexCreate_ShaderResource		= 1<<3,
-	
-	// Texture is encoded in sRGB gamma space
-	TexCreate_SRGB					= 1<<4,
-	// Texture will be created without a packed miptail
-	TexCreate_NoMipTail				= 1<<5,
-	// Texture will be created with an un-tiled format
-	TexCreate_NoTiling				= 1<<6,
-	// Texture that may be updated every frame
-	TexCreate_Dynamic				= 1<<8,
-	// Allow silent texture creation failure
-	// @warning:	When you update this, you must update FTextureAllocations::FindTextureType() in Core/Private/UObject/TextureAllocations.cpp
-	TexCreate_AllowFailure			= 1<<9,
-	// Disable automatic defragmentation if the initial texture memory allocation fails.
-	// @warning:	When you update this, you must update FTextureAllocations::FindTextureType() in Core/Private/UObject/TextureAllocations.cpp
-	TexCreate_DisableAutoDefrag		= 1<<10,
-	// Create the texture with automatic -1..1 biasing
-	TexCreate_BiasNormalMap			= 1<<11,
-	// Create the texture with the flag that allows mip generation later, only applicable to D3D11
-	TexCreate_GenerateMipCapable	= 1<<12,
-	// UnorderedAccessView (DX11 only)
-	// Warning: Causes additional synchronization between draw calls when using a render target allocated with this flag, use sparingly
-	// See: GCNPerformanceTweets.pdf Tip 37
-	TexCreate_UAV					= 1<<16,
-	// Render target texture that will be displayed on screen (back buffer)
-	TexCreate_Presentable			= 1<<17,
-	// Texture data is accessible by the CPU
-	TexCreate_CPUReadback			= 1<<18,
-	// Texture was processed offline (via a texture conversion process for the current platform)
-	TexCreate_OfflineProcessed		= 1<<19,
-	// Texture needs to go in fast VRAM if available (HINT only)
-	TexCreate_FastVRAM				= 1<<20,
-	// by default the texture is not showing up in the list - this is to reduce clutter, using the FULL option this can be ignored
-	TexCreate_HideInVisualizeTexture= 1<<21,
-	// Texture should be created in virtual memory, with no physical memory allocation made
-	// You must make further calls to RHIVirtualTextureSetFirstMipInMemory to allocate physical memory
-	// and RHIVirtualTextureSetFirstMipVisible to map the first mip visible to the GPU
-	TexCreate_Virtual				= 1<<22,
-	// Creates a RenderTargetView for each array slice of the texture
-	// Warning: if this was specified when the resource was created, you can't use SV_RenderTargetArrayIndex to route to other slices!
-	TexCreate_TargetArraySlicesIndependently	= 1<<23,
-	// Texture that may be shared with DX9 or other devices
-	TexCreate_Shared = 1 << 24,
-};
-
 struct FVRamAllocation
 {
 	FVRamAllocation(uint32 InAllocationStart = 0, uint32 InAllocationSize = 0)
@@ -1267,16 +736,6 @@ struct FRHIResourceCreateInfo
 // Forward-declaration.
 struct FResolveParams;
 
-/**
- * Async texture reallocation status, returned by RHIGetReallocateTexture2DStatus().
- */
-enum ETextureReallocationStatus
-{
-	TexRealloc_Succeeded = 0,
-	TexRealloc_Failed,
-	TexRealloc_InProgress,
-};
-
 struct FResolveRect
 {
 	int32 X1;
diff --git a/Engine/Source/Runtime/RHI/Public/RHICommandList.h b/Engine/Source/Runtime/RHI/Public/RHICommandList.h
new file mode 100644
--- /dev/null
+++ b/Engine/Source/Runtime/RHI/Public/RHICommandList.h
@@ -0,0 +1,312 @@
+// Copyright 1998-2014 Epic Games, Inc. All Rights Reserved.
+
+/*=============================================================================
+	RHICommandList.h: RHI Command List definitions for queueing up & executing later.
+=============================================================================*/
+
+#pragma once
+
+enum ERHICommandType
+{
+	// Start the enum at non-zero to help debug
+	ERCT_NopBlob = 1000,	// NoOp, maybe has a buffer to be auto alloc/freed
+	ERCT_SetRasterizerState,
+	ERCT_SetShaderParameter,
+	ERCT_SetShaderUniformBuffer,
+	ERCT_SetShaderTexture,
+	ERCT_SetShaderSampler,
+	ERCT_DrawPrimitive,
+	ERCT_DrawIndexedPrimitive,
+	ERCT_SetBoundShaderState,
+	ERCT_SetBlendState,
+	ERCT_SetStreamSource,
+
+	ERCT_,
+};
+
+class FRHICommandList
+{
+public:
+	enum
+	{
+		MemSize = 32 * 1024 * 1024,
+		Alignment = sizeof(void*),	// Use SIZE_T?
+	};
+
+	FRHICommandList()
+	{
+		Reset();
+	}
+
+	template <typename TShaderRHIParamRef>
+	FORCEINLINE static void SetShaderUniformBuffer(FRHICommandList* RHICmdList, TShaderRHIParamRef Shader, uint32 BaseIndex, FUniformBufferRHIParamRef UniformBufferRHI)
+	{
+		if (RHICmdList)
+		{
+			RHICmdList->SetShaderUniformBuffer(Shader, BaseIndex, UniformBufferRHI);
+		}
+		else
+		{
+			RHISetShaderUniformBuffer(Shader, BaseIndex, UniformBufferRHI);
+		}
+	}
+	template <typename TShaderRHIParamRef>
+	inline void SetShaderUniformBuffer(TShaderRHIParamRef Shader, uint32 BaseIndex, FUniformBufferRHIParamRef UniformBufferRHI);
+
+	template <typename TShaderRHIParamRef>
+	FORCEINLINE static void SetShaderResourceViewParameter(FRHICommandList* RHICmdList, TShaderRHIParamRef Shader, uint32 SamplerIndex, FShaderResourceViewRHIParamRef SRV)
+	{
+		if (RHICmdList)
+		{
+			RHICmdList->SetShaderResourceViewParameter(Shader, SamplerIndex, SRV);
+		}
+		else
+		{
+			RHISetShaderResourceViewParameter(Shader, SamplerIndex, SRV);
+		}
+	}
+	template <typename TShaderRHIParamRef>
+	inline void SetShaderResourceViewParameter(TShaderRHIParamRef Shader, uint32 SamplerIndex, FShaderResourceViewRHIParamRef SRV);
+
+	template <typename TShaderRHIParamRef>
+	FORCEINLINE static void SetShaderParameter(FRHICommandList* RHICmdList, TShaderRHIParamRef Shader, uint32 BufferIndex, uint32 BaseIndex, uint32 NumBytes, const void* Value, bool bValueInStack = false)
+	{
+		if (RHICmdList)
+		{
+			RHICmdList->SetShaderParameter(Shader, BufferIndex, BaseIndex, NumBytes, Value, bValueInStack);
+		}
+		else
+		{
+			RHISetShaderParameter(Shader, BufferIndex, BaseIndex, NumBytes, Value);
+		}
+	}
+	template <typename TShaderRHIParamRef>
+	inline void SetShaderParameter(TShaderRHIParamRef Shader, uint32 BufferIndex, uint32 BaseIndex, uint32 NumBytes, const void* Value, bool bValueInStack = false);
+
+	template <typename TShaderRHIParamRef>
+	FORCEINLINE static void SetShaderTexture(FRHICommandList* RHICmdList, TShaderRHIParamRef Shader, uint32 TextureIndex, FTextureRHIParamRef Texture)
+	{
+		if (RHICmdList)
+		{
+			RHICmdList->SetShaderTexture(Shader, TextureIndex, Texture);
+		}
+		else
+		{
+			RHISetShaderTexture(Shader, TextureIndex, Texture);
+		}
+	}
+	template <typename TShaderRHIParamRef>
+	inline void SetShaderTexture(TShaderRHIParamRef Shader, uint32 TextureIndex, FTextureRHIParamRef Texture);
+
+	template <typename TShaderRHIParamRef>
+	FORCEINLINE static void SetShaderSampler(FRHICommandList* RHICmdList, TShaderRHIParamRef Shader, uint32 SamplerIndex, FSamplerStateRHIParamRef State)
+	{
+		if (RHICmdList)
+		{
+			RHICmdList->SetShaderSampler(Shader, SamplerIndex, State);
+		}
+		else
+		{
+			RHISetShaderSampler(Shader, SamplerIndex, State);
+		}
+	}
+	template <typename TShaderRHIParamRef>
+	inline void SetShaderSampler(TShaderRHIParamRef Shader, uint32 SamplerIndex, FSamplerStateRHIParamRef State);
+
+	FORCEINLINE static void SetUAVParameter(FRHICommandList* RHICmdList, FComputeShaderRHIParamRef Shader, uint32 UAVIndex, FUnorderedAccessViewRHIParamRef UAV)
+	{
+		if (RHICmdList)
+		{
+			RHICmdList->SetUAVParameter(Shader, UAVIndex, UAV);
+		}
+		else
+		{
+			RHISetUAVParameter(Shader, UAVIndex, UAV);
+		}
+	}
+	FORCEINLINE static void SetUAVParameter(FRHICommandList* RHICmdList, FComputeShaderRHIParamRef Shader, uint32 UAVIndex, FUnorderedAccessViewRHIParamRef UAV, uint32 InitialCount)
+	{
+		if (RHICmdList)
+		{
+			RHICmdList->SetUAVParameter(Shader, UAVIndex, UAV, InitialCount);
+		}
+		else
+		{
+			RHISetUAVParameter(Shader, UAVIndex, UAV, InitialCount);
+		}
+	}
+	inline void SetUAVParameter(FComputeShaderRHIParamRef Shader, uint32 UAVIndex, FUnorderedAccessViewRHIParamRef UAV);
+	inline void SetUAVParameter(FComputeShaderRHIParamRef Shader, uint32 UAVIndex, FUnorderedAccessViewRHIParamRef UAV, uint32 InitialCount);
+
+	FORCEINLINE static void SetBoundShaderState(FRHICommandList* RHICmdList, FBoundShaderStateRHIParamRef BoundShaderState)
+	{
+		if (RHICmdList)
+		{
+			RHICmdList->SetBoundShaderState(BoundShaderState);
+		}
+		else
+		{
+			RHISetBoundShaderState(BoundShaderState);
+		}
+	}
+	inline void SetBoundShaderState(FBoundShaderStateRHIParamRef BoundShaderState);
+
+	FORCEINLINE static void SetRasterizerState(FRHICommandList* RHICmdList, FRasterizerStateRHIParamRef State)
+	{
+		if (RHICmdList)
+		{
+			RHICmdList->SetRasterizerState(State);
+		}
+		else
+		{
+			RHISetRasterizerState(State);
+		}
+	}
+	inline void SetRasterizerState(FRasterizerStateRHIParamRef State);
+
+	FORCEINLINE static void SetBlendState(FRHICommandList* RHICmdList, FBlendStateRHIParamRef State, const FLinearColor& BlendFactor)
+	{
+		if (RHICmdList)
+		{
+			RHICmdList->SetBlendState(State, BlendFactor);
+		}
+		else
+		{
+			RHISetBlendState(State, BlendFactor);
+		}
+	}
+	FORCEINLINE static void SetBlendState(FRHICommandList* RHICmdList, FBlendStateRHIParamRef State)
+	{
+		SetBlendState(RHICmdList, State, FLinearColor::White);
+	}
+	inline void SetBlendState(FBlendStateRHIParamRef State, const FLinearColor& BlendFactor);
+	inline void SetBlendState(FBlendStateRHIParamRef State)
+	{
+		SetBlendState(State, FLinearColor::White);
+	}
+
+	FORCEINLINE static void DrawPrimitive(FRHICommandList* RHICmdList, uint32 PrimitiveType, uint32 BaseVertexIndex, uint32 NumPrimitives, uint32 NumInstances)
+	{
+		if (RHICmdList)
+		{
+			RHICmdList->DrawPrimitive(PrimitiveType, BaseVertexIndex, NumPrimitives, NumInstances);
+		}
+		else
+		{
+			RHIDrawPrimitive(PrimitiveType, BaseVertexIndex, NumPrimitives, NumInstances);
+		}
+	}
+	inline void DrawPrimitive(uint32 PrimitiveType, uint32 BaseVertexIndex, uint32 NumPrimitives, uint32 NumInstances);
+
+	FORCEINLINE static void DrawIndexedPrimitive(FRHICommandList* RHICmdList, FIndexBufferRHIParamRef IndexBuffer, uint32 PrimitiveType, int32 BaseVertexIndex, uint32 MinIndex, uint32 NumVertices, uint32 StartIndex, uint32 NumPrimitives, uint32 NumInstances)
+	{
+		if (RHICmdList)
+		{
+			RHICmdList->DrawIndexedPrimitive(IndexBuffer, PrimitiveType, BaseVertexIndex, MinIndex, NumVertices, StartIndex, NumPrimitives, NumInstances);
+		}
+		else
+		{
+			RHIDrawIndexedPrimitive(IndexBuffer, PrimitiveType, BaseVertexIndex, MinIndex, NumVertices, StartIndex, NumPrimitives, NumInstances);
+		}
+	}
+	inline void DrawIndexedPrimitive(FIndexBufferRHIParamRef IndexBuffer, uint32 PrimitiveType, int32 BaseVertexIndex, uint32 MinIndex, uint32 NumVertices, uint32 StartIndex, uint32 NumPrimitives, uint32 NumInstances);
+
+	FORCEINLINE static void SetStreamSource(FRHICommandList* RHICmdList, uint32 StreamIndex, FVertexBufferRHIParamRef VertexBuffer, uint32 Stride, uint32 Offset)
+	{
+		if (RHICmdList)
+		{
+			RHICmdList->SetStreamSource(StreamIndex, VertexBuffer, Stride, Offset);
+		}
+		else
+		{
+			RHISetStreamSource(StreamIndex, VertexBuffer, Stride, Offset);
+		}
+	}
+	inline void SetStreamSource(uint32 StreamIndex, FVertexBufferRHIParamRef VertexBuffer, uint32 Stride, uint32 Offset);
+
+
+	uint8* Alloc(SIZE_T InSize)
+	{
+		InSize = (InSize + (Alignment - 1)) & ~(Alignment - 1);
+		check(Tail + InSize + Alignment <= &Mem[MemSize]);
+		uint8* New = Tail;
+		Tail += InSize;
+		return New;
+	}
+
+	const uint8* GetHead() const
+	{
+		return &Mem[0];
+	}
+
+	const uint8* GetTail() const
+	{
+		return Tail;
+	}
+
+protected:
+	uint8 Mem[MemSize];
+
+	enum EState
+	{
+		Ready,
+		Executing,
+		Kicked,
+	};
+	EState State;
+	uint8* Tail;
+	uint32 NumCommands;
+
+	inline bool CanAddCommand() const
+	{
+		return (State == Ready);
+	}
+
+	bool CanFitCommand(SIZE_T Size) const
+	{
+		return (Size + (Tail - Mem) <= MemSize);
+	}
+
+	template <typename TCmd>
+	FORCEINLINE TCmd* AddCommand()
+	{
+		check(CanAddCommand());
+		check(CanFitCommand(sizeof(TCmd)));
+		TCmd* Cmd = (TCmd*)Alloc(sizeof(TCmd));
+		++NumCommands;
+		return Cmd;
+	}
+
+	void Reset()
+	{
+		State = Kicked;
+		NumCommands = 0;
+		Tail = &Mem[0];
+		Tail = Align(Tail, Alignment);
+	}
+
+	friend class FRHICommandListExecutor;
+};
+
+class RHI_API FRHICommandListExecutor
+{
+public:
+	FRHICommandListExecutor();
+	FRHICommandList* CreateList();
+	void ExecuteAndFreeList(FRHICommandList* CmdList);
+
+	void Verify();
+
+private:
+	FCriticalSection CriticalSection;
+	bool bLock;
+};
+extern RHI_API FRHICommandListExecutor GRHICommandList;
+
+DECLARE_CYCLE_STAT_EXTERN(TEXT("RHI Command List Execute"),STAT_RHICmdListExecuteTime,STATGROUP_RHI, );
+DECLARE_CYCLE_STAT_EXTERN(TEXT("RHI Command List Enqueue"),STAT_RHICmdListEnqueueTime,STATGROUP_RHI, );
+DECLARE_DWORD_COUNTER_STAT_EXTERN(TEXT("RHI Command List memory"),STAT_RHICmdListMemory,STATGROUP_RHI,RHI_API);
+DECLARE_DWORD_COUNTER_STAT_EXTERN(TEXT("RHI Command List count"),STAT_RHICmdListCount,STATGROUP_RHI,RHI_API);
+DECLARE_CYCLE_STAT_EXTERN(TEXT("RHI Counter TEMP"), STAT_RHICounterTEMP,STATGROUP_RHI,RHI_API);
+
+#include "RHICommandList.inl"
diff --git a/Engine/Source/Runtime/RHI/Public/RHICommandList.inl b/Engine/Source/Runtime/RHI/Public/RHICommandList.inl
new file mode 100644
--- /dev/null
+++ b/Engine/Source/Runtime/RHI/Public/RHICommandList.inl
@@ -0,0 +1,286 @@
+// Copyright 1998-2014 Epic Games, Inc. All Rights Reserved.
+
+/*=============================================================================
+	RHICommandList.inl: RHI Command List inline definitions.
+=============================================================================*/
+
+#ifndef __RHICOMMANDLIST_INL__
+#define __RHICOMMANDLIST_INL__
+#include "RHICommandList.h"
+
+struct FRHICommand
+{
+	union
+	{
+		SIZE_T			RawType;
+		ERHICommandType	Type;
+	};
+};
+
+struct FRHICommandPerShader : public FRHICommand
+{
+	FRHIResource* Shader;
+	EShaderFrequency ShaderFrequency;
+
+	template <typename TShaderRHIParamRef>
+	FORCEINLINE void SetShader(TShaderRHIParamRef InShader)
+	{
+		Shader = InShader;
+		ShaderFrequency = (EShaderFrequency)TRHIShaderToEnum<TShaderRHIParamRef>::ShaderFrequency;
+	}
+};
+
+struct FRHICommandNopBlob : public FRHICommand
+{
+	void* Data;
+	uint32 Size;
+	FORCEINLINE void Set(FRHICommandList* List, uint32 InSize, const void* InData)
+	{
+		Type = ERCT_NopBlob;
+		Size = InSize;
+		Data = List->Alloc(InSize);
+		FMemory::Memcpy(Data, InData, InSize);
+	}
+};
+
+struct FRHICommandSetRasterizerState : public FRHICommand
+{
+	FRasterizerStateRHIParamRef State;
+	FORCEINLINE void Set(FRasterizerStateRHIParamRef InState)
+	{
+		Type = ERCT_SetRasterizerState;
+		State = InState;
+	}
+};
+
+struct FRHICommandSetShaderParameter : public FRHICommandPerShader
+{
+	const void* NewValue;
+	uint32 BufferIndex;
+	uint32 BaseIndex;
+	uint32 NumBytes;
+	template <typename TShaderRHIParamRef>
+	FORCEINLINE void Set(TShaderRHIParamRef InShader, uint32 InBufferIndex, uint32 InBaseIndex, uint32 InNumBytes, const void* InNewValue)
+	{
+		FRHICommandPerShader::SetShader<TShaderRHIParamRef>(InShader);
+		Type = ERCT_SetShaderParameter;
+		Shader = InShader;
+		BufferIndex = InBufferIndex;
+		BaseIndex = InBaseIndex;
+		NumBytes = InNumBytes;
+		NewValue = InNewValue;
+	}
+};
+
+struct FRHICommandSetShaderUniformBuffer : public FRHICommandPerShader
+{
+	uint32 BaseIndex;
+	FUniformBufferRHIParamRef UniformBuffer;
+	template <typename TShaderRHIParamRef>
+	FORCEINLINE void Set(TShaderRHIParamRef InShader, uint32 InBaseIndex, FUniformBufferRHIParamRef InUniformBuffer)
+	{
+		FRHICommandPerShader::SetShader<TShaderRHIParamRef>(InShader);
+		Type = ERCT_SetShaderUniformBuffer;
+		Shader = InShader;
+		BaseIndex = InBaseIndex;
+		UniformBuffer = InUniformBuffer;
+	}
+};
+
+struct FRHICommandSetShaderTexture : public FRHICommandPerShader
+{
+	uint32 TextureIndex;
+	FTextureRHIParamRef Texture;
+	template <typename TShaderRHIParamRef>
+	FORCEINLINE void Set(TShaderRHIParamRef InShader, uint32 InTextureIndex, FTextureRHIParamRef InTexture)
+	{
+		FRHICommandPerShader::SetShader<TShaderRHIParamRef>(InShader);
+		Type = ERCT_SetShaderTexture;
+		Shader = InShader;
+		TextureIndex = InTextureIndex;
+		Texture = InTexture;
+	}
+};
+
+struct FRHICommandSetShaderSampler : public FRHICommandPerShader
+{
+	uint32 SamplerIndex;
+	FSamplerStateRHIParamRef Sampler;
+	template <typename TShaderRHIParamRef>
+	FORCEINLINE void Set(TShaderRHIParamRef InShader, uint32 InSamplerIndex, FSamplerStateRHIParamRef InSampler)
+	{
+		FRHICommandPerShader::SetShader<TShaderRHIParamRef>(InShader);
+		Type = ERCT_SetShaderSampler;
+		Shader = InShader;
+		SamplerIndex = InSamplerIndex;
+		Sampler = InSampler;
+	}
+};
+
+struct FRHICommandDrawPrimitive : public FRHICommand
+{
+	uint32 PrimitiveType;
+	uint32 BaseVertexIndex;
+	uint32 NumPrimitives;
+	uint32 NumInstances;
+	FORCEINLINE void Set(uint32 InPrimitiveType, uint32 InBaseVertexIndex, uint32 InNumPrimitives, uint32 InNumInstances)
+	{
+		Type = ERCT_DrawPrimitive;
+		PrimitiveType = InPrimitiveType;
+		BaseVertexIndex = InBaseVertexIndex;
+		NumPrimitives = InNumPrimitives;
+		NumInstances = InNumInstances;
+	}
+};
+
+struct FRHICommandDrawIndexedPrimitive : public FRHICommand
+{
+	FIndexBufferRHIParamRef IndexBuffer;
+	uint32 PrimitiveType;
+	int32 BaseVertexIndex;
+	uint32 MinIndex;
+	uint32 NumVertices;
+	uint32 StartIndex;
+	uint32 NumPrimitives;
+	uint32 NumInstances;
+	FORCEINLINE void Set(FIndexBufferRHIParamRef InIndexBuffer, uint32 InPrimitiveType, int32 InBaseVertexIndex, uint32 InMinIndex, uint32 InNumVertices, uint32 InStartIndex, uint32 InNumPrimitives, uint32 InNumInstances)
+	{
+		Type = ERCT_DrawIndexedPrimitive;
+		IndexBuffer = InIndexBuffer;
+		PrimitiveType = InPrimitiveType;
+		BaseVertexIndex = InBaseVertexIndex;
+		MinIndex = InMinIndex;
+		NumVertices = InNumVertices;
+		StartIndex = InStartIndex;
+		NumPrimitives = InNumPrimitives;
+		NumInstances = InNumInstances;
+	}
+};
+
+struct FRHICommandSetBoundShaderState : public FRHICommand
+{
+	FBoundShaderStateRHIParamRef BoundShaderState;
+	FORCEINLINE void Set(FBoundShaderStateRHIParamRef InBoundShaderState)
+	{
+		Type = ERCT_SetBoundShaderState;
+		BoundShaderState = InBoundShaderState;
+	}
+};
+
+struct FRHICommandSetBlendState : public FRHICommand
+{
+	FBlendStateRHIParamRef State;
+	FLinearColor BlendFactor;
+	FORCEINLINE void Set(FBlendStateRHIParamRef InState, const FLinearColor& InBlendFactor)
+	{
+		Type = ERCT_SetBlendState;
+		State = InState;
+		BlendFactor = InBlendFactor;
+	}
+};
+
+struct FRHICommandSetStreamSource : public FRHICommand
+{
+	uint32 StreamIndex;
+	FVertexBufferRHIParamRef VertexBuffer;
+	uint32 Stride;
+	uint32 Offset;
+	FORCEINLINE void Set(uint32 InStreamIndex, FVertexBufferRHIParamRef InVertexBuffer, uint32 InStride, uint32 InOffset)
+	{
+		Type = ERCT_SetStreamSource;
+		StreamIndex = InStreamIndex;
+		VertexBuffer = InVertexBuffer;
+		Stride = InStride;
+		Offset = InOffset;
+	}
+};
+
+template <typename TShaderRHIParamRef>
+FORCEINLINE void FRHICommandList::SetShaderUniformBuffer(TShaderRHIParamRef Shader, uint32 BaseIndex, FUniformBufferRHIParamRef UniformBuffer)
+{
+	auto* Cmd = AddCommand<FRHICommandSetShaderUniformBuffer>();
+	Cmd->Set<TShaderRHIParamRef>(Shader, BaseIndex, UniformBuffer);
+	Shader->AddRef();
+	UniformBuffer->AddRef();
+}
+
+template <typename TShaderRHIParamRef>
+FORCEINLINE void FRHICommandList::SetShaderResourceViewParameter(TShaderRHIParamRef Shader, uint32 SamplerIndex, FShaderResourceViewRHIParamRef SRV){check(0);}
+
+template <typename TShaderRHIParamRef>
+FORCEINLINE void FRHICommandList::SetShaderParameter(TShaderRHIParamRef Shader, uint32 BufferIndex, uint32 BaseIndex, uint32 NumBytes, const void* NewValue, bool bValueInStack)
+{
+	if (bValueInStack)
+	{
+		auto* Blob = AddCommand<FRHICommandNopBlob>();
+		Blob->Set(this, NumBytes, NewValue);
+		NewValue = Blob->Data;
+	}
+	auto* Cmd = AddCommand<FRHICommandSetShaderParameter>();
+	Cmd->Set<TShaderRHIParamRef>(Shader, BufferIndex, BaseIndex, NumBytes, NewValue);
+	Shader->AddRef();
+}
+
+template <typename TShaderRHIParamRef>
+FORCEINLINE void FRHICommandList::SetShaderTexture(TShaderRHIParamRef Shader, uint32 TextureIndex, FTextureRHIParamRef Texture)
+{
+	auto* Cmd = AddCommand<FRHICommandSetShaderTexture>();
+	Cmd->Set<TShaderRHIParamRef>(Shader, TextureIndex, Texture);
+	Shader->AddRef();
+	Texture->AddRef();
+}
+
+template <typename TShaderRHIParamRef>
+FORCEINLINE void FRHICommandList::SetShaderSampler(TShaderRHIParamRef Shader, uint32 SamplerIndex, FSamplerStateRHIParamRef State)
+{
+	auto* Cmd = AddCommand<FRHICommandSetShaderSampler>();
+	Cmd->Set<TShaderRHIParamRef>(Shader, SamplerIndex, State);
+	Shader->AddRef();
+	State->AddRef();
+}
+
+FORCEINLINE void FRHICommandList::SetUAVParameter(FComputeShaderRHIParamRef Shader, uint32 UAVIndex, FUnorderedAccessViewRHIParamRef UAV){check(0);}
+FORCEINLINE void FRHICommandList::SetUAVParameter(FComputeShaderRHIParamRef Shader, uint32 UAVIndex, FUnorderedAccessViewRHIParamRef UAV, uint32 InitialCount){check(0);}
+
+FORCEINLINE void FRHICommandList::SetRasterizerState(FRasterizerStateRHIParamRef State)
+{
+	auto* Cmd = AddCommand<FRHICommandSetRasterizerState>();
+	Cmd->Set(State);
+	State->AddRef();
+}
+
+FORCEINLINE void FRHICommandList::DrawPrimitive(uint32 PrimitiveType, uint32 BaseVertexIndex, uint32 NumPrimitives, uint32 NumInstances)
+{
+	auto* Cmd = AddCommand<FRHICommandDrawPrimitive>();
+	Cmd->Set(PrimitiveType, BaseVertexIndex, NumPrimitives, NumInstances);
+}
+
+FORCEINLINE void FRHICommandList::DrawIndexedPrimitive(FIndexBufferRHIParamRef IndexBuffer, uint32 PrimitiveType, int32 BaseVertexIndex, uint32 MinIndex, uint32 NumVertices, uint32 StartIndex, uint32 NumPrimitives, uint32 NumInstances)
+{
+	auto* Cmd = AddCommand<FRHICommandDrawIndexedPrimitive>();
+	Cmd->Set(IndexBuffer, PrimitiveType, BaseVertexIndex, MinIndex, NumVertices, StartIndex, NumPrimitives, NumInstances);
+	IndexBuffer->AddRef();
+}
+
+FORCEINLINE void FRHICommandList::SetBoundShaderState(FBoundShaderStateRHIParamRef BoundShaderState)
+{
+	auto* Cmd = AddCommand<FRHICommandSetBoundShaderState>();
+	Cmd->Set(BoundShaderState);
+	BoundShaderState->AddRef();
+}
+
+FORCEINLINE void FRHICommandList::SetBlendState(FBlendStateRHIParamRef State, const FLinearColor& BlendFactor)
+{
+	auto* Cmd = AddCommand<FRHICommandSetBlendState>();
+	Cmd->Set(State, BlendFactor);
+	State->AddRef();
+}
+
+FORCEINLINE void FRHICommandList::SetStreamSource(uint32 StreamIndex, FVertexBufferRHIParamRef VertexBuffer, uint32 Stride, uint32 Offset)
+{
+	auto* Cmd = AddCommand<FRHICommandSetStreamSource>();
+	Cmd->Set(StreamIndex, VertexBuffer, Stride, Offset);
+	VertexBuffer->AddRef();
+}
+
+#endif	// __RHICOMMANDLIST_INL__
