UE4 - fix signed / unsigned mistmatch in win32

diff --git a/Engine/Source/Runtime/RHI/Public/RHICommandList.h b/Engine/Source/Runtime/RHI/Public/RHICommandList.h
--- a/Engine/Source/Runtime/RHI/Public/RHICommandList.h
+++ b/Engine/Source/Runtime/RHI/Public/RHICommandList.h
@@ -144,7 +144,7 @@ protected:
 
 	bool CanFitCommand(SIZE_T Size) const
 	{
-		return (Size + (Tail - Memory) <= MemSize);
+		return (Size + (Tail - Memory) <= (SIZE_T)MemSize);
 	}
 
 	template <typename TCmd>
