UE4 - Reduce rhi cmd list size

diff --git a/Engine/Source/Runtime/RHI/Public/RHICommandList.h b/Engine/Source/Runtime/RHI/Public/RHICommandList.h
--- a/Engine/Source/Runtime/RHI/Public/RHICommandList.h
+++ b/Engine/Source/Runtime/RHI/Public/RHICommandList.h
@@ -30,7 +30,7 @@ class RHI_API FRHICommandList
 public:
 	enum
 	{
-		DefaultMemSize = 32 * 1024 * 1024,
+		DefaultMemSize = 256 * 1024,
 		Alignment = MIN_ALIGNMENT,
 	};
 
