UE4 - Very minor prep for RHICmdList mem refactor

diff --git a/Engine/Source/Runtime/RHI/Private/RHICommandList.cpp b/Engine/Source/Runtime/RHI/Private/RHICommandList.cpp
--- a/Engine/Source/Runtime/RHI/Private/RHICommandList.cpp
+++ b/Engine/Source/Runtime/RHI/Private/RHICommandList.cpp
@@ -78,7 +78,7 @@ void FRHICommandListExecutor::ExecuteAndFreeList(FRHICommandList& CmdList)
 
 	auto* CmdPtr = CmdList.Memory;
 	auto* CmdTail = CmdList.GetTail();
-	INC_MEMORY_STAT_BY(STAT_RHICmdListMemory, (CmdTail - CmdPtr));
+	INC_MEMORY_STAT_BY(STAT_RHICmdListMemory, CmdList.GetUsedMemory());
 
 	static IConsoleVariable* RHICmdListCVar = IConsoleManager::Get().FindConsoleVariable(TEXT("r.RHISkip"));
 	bool bSkipRHICmdList = (RHICmdListCVar && RHICmdListCVar->GetInt() != 0);
diff --git a/Engine/Source/Runtime/RHI/Public/RHICommandList.h b/Engine/Source/Runtime/RHI/Public/RHICommandList.h
--- a/Engine/Source/Runtime/RHI/Public/RHICommandList.h
+++ b/Engine/Source/Runtime/RHI/Public/RHICommandList.h
@@ -9,7 +9,7 @@
 enum ERHICommandType
 {
 	// Start the enum at non-zero to help debug
-	ERCT_NopBlob = 1000,	// NoOp, maybe has a buffer to be auto alloc/freed
+	ERCT_NopBlob				= 512,	// NoOp, maybe has a buffer to be auto alloc/freed (512=0x200)
 	ERCT_SetRasterizerState,
 	ERCT_SetShaderParameter,
 	ERCT_SetShaderUniformBuffer,
@@ -49,6 +49,7 @@ public:
 		}
 		Reset();
 	}
+
 	~FRHICommandList()
 	{
 		if (!IsNull())
@@ -179,18 +180,16 @@ public:
 		RHIDrawPrimitiveIndirect(PrimitiveType, ArgumentBufferRHI, ArgumentOffset);
 	}
 
-
-
 	FORCEINLINE bool IsNull() const
 	{
 		return !Memory;
 	}
+
 	FORCEINLINE void CheckIsNull() const
 	{
 		check(IsNull());
 	}
 
-
 	uint8* Alloc(SIZE_T InSize)
 	{
 		checkSlow(!IsNull());
@@ -213,6 +212,11 @@ public:
 		return Tail;
 	}
 
+	const SIZE_T GetUsedMemory() const
+	{
+		return GetTail() - GetHead();
+	}
+
 protected:
 	uint8* Memory;
 	int32 MemSize;
diff --git a/Engine/Source/Runtime/RHI/Public/RHICommandList.inl b/Engine/Source/Runtime/RHI/Public/RHICommandList.inl
--- a/Engine/Source/Runtime/RHI/Public/RHICommandList.inl
+++ b/Engine/Source/Runtime/RHI/Public/RHICommandList.inl
@@ -10,6 +10,7 @@
 
 struct FRHICommand
 {
+	// This is a union only to force alignment to SIZE_T
 	union
 	{
 		SIZE_T			RawType;
@@ -17,6 +18,9 @@ struct FRHICommand
 	};
 };
 
+// No constructors are given for classes derived from FRHICommand as they are not created using
+// a classical 'new' for perf reasons.
+
 struct FRHICommandPerShader : public FRHICommand
 {
 	FRHIResource* Shader;
@@ -217,7 +221,7 @@ FORCEINLINE void FRHICommandList::SetShaderResourceViewParameter(TShaderRHIParam
 		RHISetShaderResourceViewParameter(Shader, SamplerIndex, SRV);
 		return;
 	}
-	check(0);
+	checkf(0, TEXT("TODO: Non-bypass version"));
 }
 
 template <typename TShaderRHIParamRef>
@@ -277,9 +281,9 @@ FORCEINLINE void FRHICommandList::SetUAVParameter(FComputeShaderRHIParamRef Shad
 		RHISetUAVParameter(Shader, UAVIndex, UAV);
 		return;
 	}
-
-	check(0);
+	checkf(0, TEXT("TODO: Non-bypass version"));
 }
+
 FORCEINLINE void FRHICommandList::SetUAVParameter(FComputeShaderRHIParamRef Shader, uint32 UAVIndex, FUnorderedAccessViewRHIParamRef UAV, uint32 InitialCount)
 {
 	if (Bypass())
@@ -287,7 +291,7 @@ FORCEINLINE void FRHICommandList::SetUAVParameter(FComputeShaderRHIParamRef Shad
 		RHISetUAVParameter(Shader, UAVIndex, UAV, InitialCount);
 		return;
 	}
-	check(0);
+	checkf(0, TEXT("TODO: Non-bypass version"));
 }
 
 FORCEINLINE void FRHICommandList::SetRasterizerState(FRasterizerStateRHIParamRef State)
