UE4 - RHICmdList command iterator, hide mem allocation

diff --git a/Engine/Source/Runtime/RHI/Private/RHICommandList.cpp b/Engine/Source/Runtime/RHI/Private/RHICommandList.cpp
--- a/Engine/Source/Runtime/RHI/Private/RHICommandList.cpp
+++ b/Engine/Source/Runtime/RHI/Private/RHICommandList.cpp
@@ -58,6 +58,167 @@ FRHICommandListExecutor::FRHICommandListExecutor() :
 {
 }
 
+template <bool bOnlyTestMemAccess>
+void FRHICommandListExecutor::ExecuteList(FRHICommandList& CmdList)
+{
+	FRHICommandListIterator Iter(CmdList);
+	while (Iter.HasCommandsLeft())
+	{
+		auto* Cmd = *Iter;
+		switch (Cmd->Type)
+		{
+		case ERCT_NopBlob:
+			{
+				// Nop
+				auto* RHICmd = Iter.NextCommand<FRHICommandNopBlob>();
+			}
+			break;
+		case ERCT_SetRasterizerState:
+			{
+				auto* RHICmd = Iter.NextCommand<FRHICommandSetRasterizerState>();
+				if (!bOnlyTestMemAccess)
+				{
+					RHISetRasterizerState(RHICmd->State);
+				}
+				RHICmd->State->Release();
+			}
+			break;
+		case ERCT_SetShaderParameter:
+			{
+				auto* RHICmd = Iter.NextCommand<FRHICommandSetShaderParameter>();
+				if (!bOnlyTestMemAccess)
+				{
+					switch (RHICmd->ShaderFrequency)
+					{
+					case SF_Vertex:		RHISetShaderParameter((FVertexShaderRHIParamRef)RHICmd->Shader, RHICmd->BufferIndex, RHICmd->BaseIndex, RHICmd->NumBytes, RHICmd->NewValue); break;
+					case SF_Hull:		RHISetShaderParameter((FHullShaderRHIParamRef)RHICmd->Shader, RHICmd->BufferIndex, RHICmd->BaseIndex, RHICmd->NumBytes, RHICmd->NewValue); break;
+					case SF_Domain:		RHISetShaderParameter((FDomainShaderRHIParamRef)RHICmd->Shader, RHICmd->BufferIndex, RHICmd->BaseIndex, RHICmd->NumBytes, RHICmd->NewValue); break;
+					case SF_Geometry:	RHISetShaderParameter((FGeometryShaderRHIParamRef)RHICmd->Shader, RHICmd->BufferIndex, RHICmd->BaseIndex, RHICmd->NumBytes, RHICmd->NewValue); break;
+					case SF_Pixel:		RHISetShaderParameter((FPixelShaderRHIParamRef)RHICmd->Shader, RHICmd->BufferIndex, RHICmd->BaseIndex, RHICmd->NumBytes, RHICmd->NewValue); break;
+					case SF_Compute:	RHISetShaderParameter((FComputeShaderRHIParamRef)RHICmd->Shader, RHICmd->BufferIndex, RHICmd->BaseIndex, RHICmd->NumBytes, RHICmd->NewValue); break;
+					default: check(0); break;
+					}
+				}
+			}
+			break;
+		case ERCT_SetShaderUniformBuffer:
+			{
+				auto* RHICmd = Iter.NextCommand<FRHICommandSetShaderUniformBuffer>();
+				if (!bOnlyTestMemAccess)
+				{
+					switch (RHICmd->ShaderFrequency)
+					{
+					case SF_Vertex:		RHISetShaderUniformBuffer((FVertexShaderRHIParamRef)RHICmd->Shader, RHICmd->BaseIndex, RHICmd->UniformBuffer); break;
+					case SF_Hull:		RHISetShaderUniformBuffer((FHullShaderRHIParamRef)RHICmd->Shader, RHICmd->BaseIndex, RHICmd->UniformBuffer); break;
+					case SF_Domain:		RHISetShaderUniformBuffer((FDomainShaderRHIParamRef)RHICmd->Shader, RHICmd->BaseIndex, RHICmd->UniformBuffer); break;
+					case SF_Geometry:	RHISetShaderUniformBuffer((FGeometryShaderRHIParamRef)RHICmd->Shader, RHICmd->BaseIndex, RHICmd->UniformBuffer); break;
+					case SF_Pixel:		RHISetShaderUniformBuffer((FPixelShaderRHIParamRef)RHICmd->Shader, RHICmd->BaseIndex, RHICmd->UniformBuffer); break;
+					case SF_Compute:	RHISetShaderUniformBuffer((FComputeShaderRHIParamRef)RHICmd->Shader, RHICmd->BaseIndex, RHICmd->UniformBuffer); break;
+					default: check(0); break;
+					}
+				}
+				RHICmd->Shader->Release(); 
+				RHICmd->UniformBuffer->Release();
+			}
+			break;
+		case ERCT_SetShaderSampler:
+			{
+				auto* RHICmd = Iter.NextCommand<FRHICommandSetShaderSampler>();
+				if (!bOnlyTestMemAccess)
+				{
+					switch (RHICmd->ShaderFrequency)
+					{
+					case SF_Vertex:		RHISetShaderSampler((FVertexShaderRHIParamRef)RHICmd->Shader, RHICmd->SamplerIndex, RHICmd->Sampler); break;
+					case SF_Hull:		RHISetShaderSampler((FHullShaderRHIParamRef)RHICmd->Shader, RHICmd->SamplerIndex, RHICmd->Sampler); break;
+					case SF_Domain:		RHISetShaderSampler((FDomainShaderRHIParamRef)RHICmd->Shader, RHICmd->SamplerIndex, RHICmd->Sampler); break;
+					case SF_Geometry:	RHISetShaderSampler((FGeometryShaderRHIParamRef)RHICmd->Shader, RHICmd->SamplerIndex, RHICmd->Sampler); break;
+					case SF_Pixel:		RHISetShaderSampler((FPixelShaderRHIParamRef)RHICmd->Shader, RHICmd->SamplerIndex, RHICmd->Sampler); break;
+					case SF_Compute:	RHISetShaderSampler((FComputeShaderRHIParamRef)RHICmd->Shader, RHICmd->SamplerIndex, RHICmd->Sampler); break;
+					default: check(0); break;
+					}
+				}
+				RHICmd->Shader->Release(); 
+				RHICmd->Sampler->Release();
+			}
+			break;
+		case ERCT_SetShaderTexture:
+			{
+				auto* RHICmd = Iter.NextCommand<FRHICommandSetShaderTexture>();
+				if (!bOnlyTestMemAccess)
+				{
+					switch (RHICmd->ShaderFrequency)
+					{
+					case SF_Vertex:		RHISetShaderTexture((FVertexShaderRHIParamRef)RHICmd->Shader, RHICmd->TextureIndex, RHICmd->Texture); break;
+					case SF_Hull:		RHISetShaderTexture((FHullShaderRHIParamRef)RHICmd->Shader, RHICmd->TextureIndex, RHICmd->Texture); break;
+					case SF_Domain:		RHISetShaderTexture((FDomainShaderRHIParamRef)RHICmd->Shader, RHICmd->TextureIndex, RHICmd->Texture); break;
+					case SF_Geometry:	RHISetShaderTexture((FGeometryShaderRHIParamRef)RHICmd->Shader, RHICmd->TextureIndex, RHICmd->Texture); break;
+					case SF_Pixel:		RHISetShaderTexture((FPixelShaderRHIParamRef)RHICmd->Shader, RHICmd->TextureIndex, RHICmd->Texture); break;
+					case SF_Compute:	RHISetShaderTexture((FComputeShaderRHIParamRef)RHICmd->Shader, RHICmd->TextureIndex, RHICmd->Texture); break;
+					default: check(0); break;
+					}
+				}
+				RHICmd->Shader->Release(); 
+				RHICmd->Texture->Release();
+			}
+			break;
+		case ERCT_DrawPrimitive:
+			{
+				auto* RHICmd = Iter.NextCommand<FRHICommandDrawPrimitive>();
+				if (!bOnlyTestMemAccess)
+				{
+					RHIDrawPrimitive(RHICmd->PrimitiveType, RHICmd->BaseVertexIndex, RHICmd->NumPrimitives, RHICmd->NumInstances);
+				}
+			}
+			break;
+		case ERCT_DrawIndexedPrimitive:
+			{
+				auto* RHICmd = Iter.NextCommand<FRHICommandDrawIndexedPrimitive>();
+				if (!bOnlyTestMemAccess)
+				{
+					RHIDrawIndexedPrimitive(RHICmd->IndexBuffer, RHICmd->PrimitiveType, RHICmd->BaseVertexIndex, RHICmd->MinIndex, RHICmd->NumVertices, RHICmd->StartIndex, RHICmd->NumPrimitives, RHICmd->NumInstances);
+				}
+				RHICmd->IndexBuffer->Release();
+			}
+			break;
+		case ERCT_SetBoundShaderState:
+			{
+				auto* RHICmd = Iter.NextCommand<FRHICommandSetBoundShaderState>();
+				if (!bOnlyTestMemAccess)
+				{
+					RHISetBoundShaderState(RHICmd->BoundShaderState);
+				}
+				RHICmd->BoundShaderState->Release();
+			}
+			break;
+		case ERCT_SetBlendState:
+			{
+				auto* RHICmd = Iter.NextCommand<FRHICommandSetBlendState>();
+				if (!bOnlyTestMemAccess)
+				{
+					RHISetBlendState(RHICmd->State, RHICmd->BlendFactor);
+				}
+				RHICmd->State->Release();
+			}
+			break;
+		case ERCT_SetStreamSource:
+			{
+				auto* RHICmd = Iter.NextCommand<FRHICommandSetStreamSource>();
+				if (!bOnlyTestMemAccess)
+				{
+					RHISetStreamSource(RHICmd->StreamIndex, RHICmd->VertexBuffer, RHICmd->Stride, RHICmd->Offset);
+				}
+				RHICmd->VertexBuffer->Release();
+			}
+			break;
+		default:
+			check(0);
+		}
+	}
+
+	Iter.CheckNumCommands();
+}
+
+
 void FRHICommandListExecutor::ExecuteAndFreeList(FRHICommandList& CmdList)
 {
 	if (CmdList.IsNull())
@@ -76,8 +237,6 @@ void FRHICommandListExecutor::ExecuteAndFreeList(FRHICommandList& CmdList)
 
 	CmdList.State = FRHICommandList::Executing;
 
-	auto* CmdPtr = CmdList.Memory;
-	auto* CmdTail = CmdList.GetTail();
 	INC_MEMORY_STAT_BY(STAT_RHICmdListMemory, CmdList.GetUsedMemory());
 
 	static IConsoleVariable* RHICmdListCVar = IConsoleManager::Get().FindConsoleVariable(TEXT("r.RHISkip"));
@@ -94,176 +253,14 @@ void FRHICommandListExecutor::ExecuteAndFreeList(FRHICommandList& CmdList)
 		{
 			static IConsoleVariable* RHIMemCVar = IConsoleManager::Get().FindConsoleVariable(TEXT("r.RHIMem"));
 			bool bOnlySkipAroundMem = (RHIMemCVar && RHIMemCVar->GetInt() != 0);
-
-			uint32 NumCommands = 0;
-			while (CmdPtr < CmdTail)
+			if (bOnlySkipAroundMem)
 			{
-				auto* Cmd = (FRHICommand*)CmdPtr;
-				switch (Cmd->Type)
-				{
-				case ERCT_NopBlob:
-					{
-						// Nop
-						auto* RHICmd = (FRHICommandNopBlob*)Cmd;
-						CmdPtr += sizeof(FRHICommandNopBlob) + RHICmd->Size;
-					}
-					break;
-				case ERCT_SetRasterizerState:
-					{
-						auto* RHICmd = (FRHICommandSetRasterizerState*)Cmd;
-						if (!bOnlySkipAroundMem)
-						{
-							RHISetRasterizerState(RHICmd->State);
-						}
-						CmdPtr += sizeof(FRHICommandSetRasterizerState);
-						RHICmd->State->Release();
-					}
-					break;
-				case ERCT_SetShaderParameter:
-					{
-						auto* RHICmd = (FRHICommandSetShaderParameter*)Cmd;
-						if (!bOnlySkipAroundMem)
-						{
-							switch (RHICmd->ShaderFrequency)
-							{
-							case SF_Vertex: RHISetShaderParameter((FVertexShaderRHIParamRef)RHICmd->Shader, RHICmd->BufferIndex, RHICmd->BaseIndex, RHICmd->NumBytes, RHICmd->NewValue); break;
-							case SF_Hull: RHISetShaderParameter((FHullShaderRHIParamRef)RHICmd->Shader, RHICmd->BufferIndex, RHICmd->BaseIndex, RHICmd->NumBytes, RHICmd->NewValue); break;
-							case SF_Domain: RHISetShaderParameter((FDomainShaderRHIParamRef)RHICmd->Shader, RHICmd->BufferIndex, RHICmd->BaseIndex, RHICmd->NumBytes, RHICmd->NewValue); break;
-							case SF_Geometry: RHISetShaderParameter((FGeometryShaderRHIParamRef)RHICmd->Shader, RHICmd->BufferIndex, RHICmd->BaseIndex, RHICmd->NumBytes, RHICmd->NewValue); break;
-							case SF_Pixel: RHISetShaderParameter((FPixelShaderRHIParamRef)RHICmd->Shader, RHICmd->BufferIndex, RHICmd->BaseIndex, RHICmd->NumBytes, RHICmd->NewValue); break;
-							case SF_Compute: RHISetShaderParameter((FComputeShaderRHIParamRef)RHICmd->Shader, RHICmd->BufferIndex, RHICmd->BaseIndex, RHICmd->NumBytes, RHICmd->NewValue); break;
-							default: check(0); break;
-							}
-						}
-						CmdPtr += sizeof(FRHICommandSetShaderParameter);
-					}
-					break;
-				case ERCT_SetShaderUniformBuffer:
-					{
-						auto* RHICmd = (FRHICommandSetShaderUniformBuffer*)Cmd;
-						if (!bOnlySkipAroundMem)
-						{
-							switch (RHICmd->ShaderFrequency)
-							{
-							case SF_Vertex: RHISetShaderUniformBuffer((FVertexShaderRHIParamRef)RHICmd->Shader, RHICmd->BaseIndex, RHICmd->UniformBuffer); break;
-							case SF_Hull: RHISetShaderUniformBuffer((FHullShaderRHIParamRef)RHICmd->Shader, RHICmd->BaseIndex, RHICmd->UniformBuffer); break;
-							case SF_Domain: RHISetShaderUniformBuffer((FDomainShaderRHIParamRef)RHICmd->Shader, RHICmd->BaseIndex, RHICmd->UniformBuffer); break;
-							case SF_Geometry: RHISetShaderUniformBuffer((FGeometryShaderRHIParamRef)RHICmd->Shader, RHICmd->BaseIndex, RHICmd->UniformBuffer); break;
-							case SF_Pixel: RHISetShaderUniformBuffer((FPixelShaderRHIParamRef)RHICmd->Shader, RHICmd->BaseIndex, RHICmd->UniformBuffer); break;
-							case SF_Compute: RHISetShaderUniformBuffer((FComputeShaderRHIParamRef)RHICmd->Shader, RHICmd->BaseIndex, RHICmd->UniformBuffer); break;
-							default: check(0); break;
-							}
-						}
-						CmdPtr += sizeof(FRHICommandSetShaderUniformBuffer);
-						RHICmd->Shader->Release(); 
-						RHICmd->UniformBuffer->Release();
-					}
-					break;
-				case ERCT_SetShaderSampler:
-					{
-						auto* RHICmd = (FRHICommandSetShaderSampler*)Cmd;
-						if (!bOnlySkipAroundMem)
-						{
-							switch (RHICmd->ShaderFrequency)
-							{
-							case SF_Vertex: RHISetShaderSampler((FVertexShaderRHIParamRef)RHICmd->Shader, RHICmd->SamplerIndex, RHICmd->Sampler); break;
-							case SF_Hull: RHISetShaderSampler((FHullShaderRHIParamRef)RHICmd->Shader, RHICmd->SamplerIndex, RHICmd->Sampler); break;
-							case SF_Domain: RHISetShaderSampler((FDomainShaderRHIParamRef)RHICmd->Shader, RHICmd->SamplerIndex, RHICmd->Sampler); break;
-							case SF_Geometry: RHISetShaderSampler((FGeometryShaderRHIParamRef)RHICmd->Shader, RHICmd->SamplerIndex, RHICmd->Sampler); break;
-							case SF_Pixel: RHISetShaderSampler((FPixelShaderRHIParamRef)RHICmd->Shader, RHICmd->SamplerIndex, RHICmd->Sampler); break;
-							case SF_Compute: RHISetShaderSampler((FComputeShaderRHIParamRef)RHICmd->Shader, RHICmd->SamplerIndex, RHICmd->Sampler); break;
-							default: check(0); break;
-							}
-						}
-						CmdPtr += sizeof(FRHICommandSetShaderSampler);
-						RHICmd->Shader->Release(); 
-						RHICmd->Sampler->Release();
-					}
-					break;
-				case ERCT_SetShaderTexture:
-					{
-						auto* RHICmd = (FRHICommandSetShaderTexture*)Cmd;
-						if (!bOnlySkipAroundMem)
-						{
-							switch (RHICmd->ShaderFrequency)
-							{
-							case SF_Vertex: RHISetShaderTexture((FVertexShaderRHIParamRef)RHICmd->Shader, RHICmd->TextureIndex, RHICmd->Texture); break;
-							case SF_Hull: RHISetShaderTexture((FHullShaderRHIParamRef)RHICmd->Shader, RHICmd->TextureIndex, RHICmd->Texture); break;
-							case SF_Domain: RHISetShaderTexture((FDomainShaderRHIParamRef)RHICmd->Shader, RHICmd->TextureIndex, RHICmd->Texture); break;
-							case SF_Geometry: RHISetShaderTexture((FGeometryShaderRHIParamRef)RHICmd->Shader, RHICmd->TextureIndex, RHICmd->Texture); break;
-							case SF_Pixel: RHISetShaderTexture((FPixelShaderRHIParamRef)RHICmd->Shader, RHICmd->TextureIndex, RHICmd->Texture); break;
-							case SF_Compute: RHISetShaderTexture((FComputeShaderRHIParamRef)RHICmd->Shader, RHICmd->TextureIndex, RHICmd->Texture); break;
-							default: check(0); break;
-							}
-						}
-						CmdPtr += sizeof(FRHICommandSetShaderTexture);
-						RHICmd->Shader->Release(); 
-						RHICmd->Texture->Release();
-					}
-					break;
-				case ERCT_DrawPrimitive:
-					{
-						auto* RHICmd = (FRHICommandDrawPrimitive*)Cmd;
-						if (!bOnlySkipAroundMem)
-						{
-							RHIDrawPrimitive(RHICmd->PrimitiveType, RHICmd->BaseVertexIndex, RHICmd->NumPrimitives, RHICmd->NumInstances);
-						}
-						CmdPtr += sizeof(FRHICommandDrawPrimitive);
-					}
-					break;
-				case ERCT_DrawIndexedPrimitive:
-					{
-						auto* RHICmd = (FRHICommandDrawIndexedPrimitive*)Cmd;
-						if (!bOnlySkipAroundMem)
-						{
-							RHIDrawIndexedPrimitive(RHICmd->IndexBuffer, RHICmd->PrimitiveType, RHICmd->BaseVertexIndex, RHICmd->MinIndex, RHICmd->NumVertices, RHICmd->StartIndex, RHICmd->NumPrimitives, RHICmd->NumInstances);
-						}
-						CmdPtr += sizeof(FRHICommandDrawIndexedPrimitive);
-						RHICmd->IndexBuffer->Release();
-					}
-					break;
-				case ERCT_SetBoundShaderState:
-					{
-						auto* RHICmd = (FRHICommandSetBoundShaderState*)Cmd;
-						if (!bOnlySkipAroundMem)
-						{
-							RHISetBoundShaderState(RHICmd->BoundShaderState);
-						}
-						CmdPtr += sizeof(FRHICommandSetBoundShaderState);
-						RHICmd->BoundShaderState->Release();
-					}
-					break;
-				case ERCT_SetBlendState:
-					{
-						auto* RHICmd = (FRHICommandSetBlendState*)Cmd;
-						if (!bOnlySkipAroundMem)
-						{
-							RHISetBlendState(RHICmd->State, RHICmd->BlendFactor);
-						}
-						CmdPtr += sizeof(FRHICommandSetBlendState);
-						RHICmd->State->Release();
-					}
-					break;
-				case ERCT_SetStreamSource:
-					{
-						auto* RHICmd = (FRHICommandSetStreamSource*)Cmd;
-						if (!bOnlySkipAroundMem)
-						{
-							RHISetStreamSource(RHICmd->StreamIndex, RHICmd->VertexBuffer, RHICmd->Stride, RHICmd->Offset);
-						}
-						CmdPtr += sizeof(FRHICommandSetStreamSource);
-						RHICmd->VertexBuffer->Release();
-					}
-					break;
-				default:
-					check(0);
-				}
-
-				CmdPtr = Align(CmdPtr, FRHICommandList::Alignment);
-				++NumCommands;
+				ExecuteList<true>(CmdList);
+			}
+			else
+			{
+				ExecuteList<false>(CmdList);
 			}
-
-			check(CmdList.NumCommands == NumCommands);
 		}
 	}
 
diff --git a/Engine/Source/Runtime/RHI/Public/RHICommandList.h b/Engine/Source/Runtime/RHI/Public/RHICommandList.h
--- a/Engine/Source/Runtime/RHI/Public/RHICommandList.h
+++ b/Engine/Source/Runtime/RHI/Public/RHICommandList.h
@@ -26,12 +26,85 @@ enum ERHICommandType
 
 class RHI_API FRHICommandList
 {
+private:
 	static FRHICommandList NullRHICommandList;
+
+	struct FMemManager
+	{
+		FMemManager(int32 Size = DefaultMemSize)
+		{
+			MemSize = Size;
+			Memory = nullptr;
+			if (Size)
+			{
+				Memory = (uint8*)FMemory::Malloc(MemSize, Alignment);
+			}
+		}
+
+		FORCEINLINE bool IsNull() const
+		{
+			return !Memory;
+		}
+
+		~FMemManager()
+		{
+			if (IsNull())
+			{
+				FMemory::Free(Memory);
+			}
+			MemSize = 0;
+			Memory = (uint8*)0x1; // we want this to not bypass the cmd list; please just crash.
+		}
+
+		uint8* Alloc(SIZE_T InSize)
+		{
+			checkSlow(!IsNull());
+			InSize = (InSize + (Alignment - 1)) & ~(Alignment - 1);
+			check(Tail + InSize + Alignment <= &Memory[MemSize]);
+			uint8* New = Tail;
+			Tail += InSize;
+			return New;
+		}
+
+		const uint8* GetHead() const
+		{
+			checkSlow(!IsNull());
+			return &Memory[0];
+		}
+
+		const uint8* GetTail() const
+		{
+			checkSlow(!IsNull());
+			return Tail;
+		}
+
+		const SIZE_T GetUsedMemory() const
+		{
+			return GetTail() - GetHead();
+		}
+
+		bool CanFitCommand(SIZE_T Size) const
+		{
+			return (Size + (Tail - Memory) <= (SIZE_T)MemSize);
+		}
+
+		void ResetTail()
+		{
+			Tail = Memory;
+		}
+
+		uint8* Memory;
+		uint8* Tail;
+		int32 MemSize;
+	};
+
+	FMemManager MemManager;
+
 public:
 	enum
 	{
 		DefaultMemSize = 256 * 1024,
-		Alignment = MIN_ALIGNMENT,
+		Alignment = sizeof(SIZE_T),
 	};
 
 	static FORCEINLINE FRHICommandList& GetNullRef()
@@ -39,25 +112,14 @@ public:
 		return NullRHICommandList;
 	}
 
-	FRHICommandList(int32 Size = DefaultMemSize)
+	FRHICommandList(int32 Size = DefaultMemSize) :
+		MemManager(Size)
 	{
-		MemSize = Size;
-		Memory = nullptr;
-		if (Size)
-		{
-			Memory = (uint8*)FMemory::Malloc(MemSize, MIN_ALIGNMENT);
-		}
 		Reset();
 	}
 
 	~FRHICommandList()
 	{
-		if (!IsNull())
-		{
-			FMemory::Free(Memory);
-		}
-		MemSize = 0;
-		Memory = (uint8*)0x1; // we want this to not bypass the cmd list; please just crash.
 	}
 
 	template <typename TShaderRHIParamRef>
@@ -182,7 +244,7 @@ public:
 
 	FORCEINLINE bool IsNull() const
 	{
-		return !Memory;
+		return MemManager.IsNull();
 	}
 
 	FORCEINLINE void CheckIsNull() const
@@ -190,37 +252,12 @@ public:
 		check(IsNull());
 	}
 
-	uint8* Alloc(SIZE_T InSize)
-	{
-		checkSlow(!IsNull());
-		InSize = (InSize + (Alignment - 1)) & ~(Alignment - 1);
-		check(Tail + InSize + Alignment <= &Memory[MemSize]);
-		uint8* New = Tail;
-		Tail += InSize;
-		return New;
-	}
-
-	const uint8* GetHead() const
-	{
-		checkSlow(!IsNull());
-		return &Memory[0];
-	}
-
-	const uint8* GetTail() const
-	{
-		checkSlow(!IsNull());
-		return Tail;
-	}
-
-	const SIZE_T GetUsedMemory() const
+	FORCEINLINE const SIZE_T GetUsedMemory() const
 	{
-		return GetTail() - GetHead();
+		return MemManager.GetUsedMemory();
 	}
 
 protected:
-	uint8* Memory;
-	int32 MemSize;
-
 	enum EState
 	{
 		Ready,
@@ -228,17 +265,31 @@ protected:
 		Kicked,
 	};
 	EState State;
-	uint8* Tail;
 	uint32 NumCommands;
 
+	FORCEINLINE uint8* Alloc(SIZE_T InSize)
+	{
+		return MemManager.Alloc(InSize);
+	}
+
+	FORCEINLINE const uint8* GetHead() const
+	{
+		return MemManager.GetHead();
+	}
+
+	FORCEINLINE const uint8* GetTail() const
+	{
+		return MemManager.GetTail();
+	}
+
 	inline bool CanAddCommand() const
 	{
 		return (State == Ready);
 	}
 
-	bool CanFitCommand(SIZE_T Size) const
+	FORCEINLINE bool CanFitCommand(SIZE_T Size) const
 	{
-		return (Size + (Tail - Memory) <= (SIZE_T)MemSize);
+		return MemManager.CanFitCommand(Size);
 	}
 
 	template <typename TCmd>
@@ -256,7 +307,7 @@ protected:
 	{
 		State = Kicked;
 		NumCommands = 0;
-		Tail = Memory;
+		MemManager.ResetTail();
 	}
 
 	FORCEINLINE bool Bypass()
@@ -265,6 +316,10 @@ protected:
 	}
 
 	friend class FRHICommandListExecutor;
+	friend class FRHICommandListIterator;
+
+	// Used here as we really want Alloc() to only work for this class, nobody else should use it
+	friend struct FRHICommandNopBlob;
 };
 
 class RHI_API FRHICommandListExecutor
@@ -279,6 +334,9 @@ public:
 private:
 	FCriticalSection CriticalSection;
 	bool bLock;
+
+	template <bool bOnlyTestMemAccess>
+	void ExecuteList(FRHICommandList& CmdList);
 };
 extern RHI_API FRHICommandListExecutor GRHICommandList;
 
diff --git a/Engine/Source/Runtime/RHI/Public/RHICommandList.inl b/Engine/Source/Runtime/RHI/Public/RHICommandList.inl
--- a/Engine/Source/Runtime/RHI/Public/RHICommandList.inl
+++ b/Engine/Source/Runtime/RHI/Public/RHICommandList.inl
@@ -13,9 +13,14 @@ struct FRHICommand
 	// This is a union only to force alignment to SIZE_T
 	union
 	{
-		SIZE_T			RawType;
 		ERHICommandType	Type;
+		SIZE_T			RawType;
 	};
+
+	uint32 ExtraSize() const
+	{
+		return 0;
+	}
 };
 
 // No constructors are given for classes derived from FRHICommand as they are not created using
@@ -45,6 +50,11 @@ struct FRHICommandNopBlob : public FRHICommand
 		Data = List->Alloc(InSize);
 		FMemory::Memcpy(Data, InData, InSize);
 	}
+
+	uint32 ExtraSize() const
+	{
+		return Size;
+	}
 };
 
 struct FRHICommandSetRasterizerState : public FRHICommand
@@ -371,4 +381,51 @@ FORCEINLINE void FRHICommandList::SetStreamSource(uint32 StreamIndex, FVertexBuf
 	VertexBuffer->AddRef();
 }
 
+// Helper class for traversing a FRHICommandList
+class FRHICommandListIterator
+{
+public:
+	FRHICommandListIterator(FRHICommandList& CmdList)
+	{
+		NumCommands = 0;
+		CmdListNumCommands = CmdList.NumCommands;
+		CmdPtr = CmdList.GetHead();
+		CmdTail = CmdList.GetTail();
+	}
+
+	FORCEINLINE bool HasCommandsLeft() const
+	{
+		return (CmdPtr < CmdTail);
+	}
+
+	// Current command
+	FORCEINLINE FRHICommand* operator * ()
+	{
+		return (FRHICommand*)CmdPtr;
+	}
+
+	// Get the next RHICommand and advance the iterator
+	template <typename T>
+	FORCEINLINE T* NextCommand()
+	{
+		T* RHICmd = (T*)CmdPtr;
+		CmdPtr += sizeof(T) + RHICmd->ExtraSize();
+		CmdPtr = Align(CmdPtr, FRHICommandList::Alignment);
+		++NumCommands;
+		return RHICmd;
+	}
+
+	void CheckNumCommands()
+	{
+		checkf(CmdListNumCommands == NumCommands, TEXT("Missed %d Commands!"), CmdListNumCommands - NumCommands);
+		checkf(CmdPtr == CmdTail, TEXT("Mismatched Tail Pointer!"));
+	}
+
+protected:
+	uint32 NumCommands;
+	const uint8* CmdPtr;
+	const uint8* CmdTail;
+	uint32 CmdListNumCommands;
+};
+
 #endif	// __RHICOMMANDLIST_INL__
diff --git a/Engine/Source/Runtime/Windows/D3D11RHI/Private/D3D11Commands.cpp b/Engine/Source/Runtime/Windows/D3D11RHI/Private/D3D11Commands.cpp
--- a/Engine/Source/Runtime/Windows/D3D11RHI/Private/D3D11Commands.cpp
+++ b/Engine/Source/Runtime/Windows/D3D11RHI/Private/D3D11Commands.cpp
@@ -2129,35 +2129,32 @@ uint32 FD3D11DynamicRHI::RHIGetGPUFrameCycles()
 
 void FD3D11DynamicRHI::RHIExecuteCommandList(FRHICommandList* CmdList)
 {
-	auto* CmdPtr = CmdList->GetHead();
-	auto* CmdTail = CmdList->GetTail();
-
-	while (CmdPtr < CmdTail)
+	FRHICommandListIterator Iter(*CmdList);
+	while (Iter.HasCommandsLeft())
 	{
-		auto* Cmd = (FRHICommand*)CmdPtr;
+		auto* Cmd = *Iter;
 		switch (Cmd->Type)
 		{
 		case ERCT_NopBlob:
 			{
 				// Nop
-				auto* RHICmd = (FRHICommandNopBlob*)Cmd;
-				CmdPtr += sizeof(FRHICommandNopBlob) + RHICmd->Size;
+				auto* RHICmd = Iter.NextCommand<FRHICommandNopBlob>();// (FRHICommandNopBlob*)Cmd;
+				(void)RHICmd;	// Unused
 			}
 			break;
 		case ERCT_SetRasterizerState:
 			{
-				auto* RHICmd = (FRHICommandSetRasterizerState*)Cmd;
+				auto* RHICmd = Iter.NextCommand<FRHICommandSetRasterizerState>();
 				{
 					auto* State = (FD3D11RasterizerState*)RHICmd->State;
 					StateCache.SetRasterizerState(State->Resource);
 				}
-				CmdPtr += sizeof(FRHICommandSetRasterizerState);
 				RHICmd->State->Release();
 			}
 			break;
 		case ERCT_SetShaderParameter:
 			{
-				auto* RHICmd = (FRHICommandSetShaderParameter*)Cmd;
+				auto* RHICmd = Iter.NextCommand<FRHICommandSetShaderParameter>();
 				{
 					switch (RHICmd->ShaderFrequency)
 					{
@@ -2194,12 +2191,11 @@ void FD3D11DynamicRHI::RHIExecuteCommandList(FRHICommandList* CmdList)
 					default: check(0); break;
 					}
 				}
-				CmdPtr += sizeof(FRHICommandSetShaderParameter);
 			}
 			break;
 		case ERCT_SetShaderUniformBuffer:
 			{
-				auto* RHICmd = (FRHICommandSetShaderUniformBuffer*)Cmd;
+				auto* RHICmd = Iter.NextCommand<FRHICommandSetShaderUniformBuffer>();
 				{
 					FD3D11UniformBuffer* Buffer = (FD3D11UniformBuffer*)RHICmd->UniformBuffer;
 					ID3D11Buffer* ConstantBuffer = Buffer ? Buffer->Resource : nullptr;
@@ -2244,14 +2240,13 @@ void FD3D11DynamicRHI::RHIExecuteCommandList(FRHICommandList* CmdList)
 					default: check(0); break;
 					}
 				}
-				CmdPtr += sizeof(FRHICommandSetShaderUniformBuffer);
 				RHICmd->Shader->Release(); 
 				RHICmd->UniformBuffer->Release();
 			}
 			break;
 		case ERCT_SetShaderSampler:
 			{
-				auto* RHICmd = (FRHICommandSetShaderSampler*)Cmd;
+				auto* RHICmd = Iter.NextCommand<FRHICommandSetShaderSampler>();
 				{
 					FD3D11SamplerState* SamplerState = (FD3D11SamplerState*)RHICmd->Sampler;
 					ID3D11SamplerState* StateResource = SamplerState->Resource;
@@ -2284,14 +2279,13 @@ void FD3D11DynamicRHI::RHIExecuteCommandList(FRHICommandList* CmdList)
 					default: check(0); break;
 					}
 				}
-				CmdPtr += sizeof(FRHICommandSetShaderSampler);
 				RHICmd->Shader->Release(); 
 				RHICmd->Sampler->Release();
 			}
 			break;
 		case ERCT_SetShaderTexture:
 			{
-				auto* RHICmd = (FRHICommandSetShaderTexture*)Cmd;
+				auto* RHICmd = Iter.NextCommand<FRHICommandSetShaderTexture>();
 				{
 					uint32 Start = FPlatformTime::Cycles();
 					FD3D11TextureBase* Texture = GetD3D11TextureFromRHITexture(RHICmd->Texture);
@@ -2328,14 +2322,13 @@ void FD3D11DynamicRHI::RHIExecuteCommandList(FRHICommandList* CmdList)
 					SetShaderTextureCycles += (FPlatformTime::Cycles() - Start);
 					SetShaderTextureCalls++;
 				}
-				CmdPtr += sizeof(FRHICommandSetShaderTexture);
-				RHICmd->Shader->Release(); 
+				RHICmd->Shader->Release();
 				RHICmd->Texture->Release();
 			}
 			break;
 		case ERCT_DrawPrimitive:
 			{
-				auto* RHICmd = (FRHICommandDrawPrimitive*)Cmd;
+				auto* RHICmd = Iter.NextCommand<FRHICommandDrawPrimitive>();
 				{
 					RHI_DRAW_CALL_STATS(RHICmd->PrimitiveType, RHICmd->NumInstances * RHICmd->NumPrimitives);
 
@@ -2354,12 +2347,11 @@ void FD3D11DynamicRHI::RHIExecuteCommandList(FRHICommandList* CmdList)
 						Direct3DDeviceIMContext->Draw(VertexCount, RHICmd->BaseVertexIndex);
 					}
 				}
-				CmdPtr += sizeof(FRHICommandDrawPrimitive);
 			}
 			break;
 		case ERCT_DrawIndexedPrimitive:
 			{
-				auto* RHICmd = (FRHICommandDrawIndexedPrimitive*)Cmd;
+				auto* RHICmd = Iter.NextCommand<FRHICommandDrawIndexedPrimitive>();
 				{
 					auto* IndexBuffer = (FD3D11IndexBuffer*)RHICmd->IndexBuffer;
 
@@ -2395,13 +2387,12 @@ void FD3D11DynamicRHI::RHIExecuteCommandList(FRHICommandList* CmdList)
 						Direct3DDeviceIMContext->DrawIndexed(IndexCount, RHICmd->StartIndex, RHICmd->BaseVertexIndex);
 					}
 				}
-				CmdPtr += sizeof(FRHICommandDrawIndexedPrimitive);
 				RHICmd->IndexBuffer->Release();
 			}
 			break;
 		case ERCT_SetBoundShaderState:
 			{
-				auto* RHICmd = (FRHICommandSetBoundShaderState*)Cmd;
+				auto* RHICmd = Iter.NextCommand<FRHICommandSetBoundShaderState>();
 				{
 					auto* BoundShaderState = (FD3D11BoundShaderState*)RHICmd->BoundShaderState;
 
@@ -2436,39 +2427,36 @@ void FD3D11DynamicRHI::RHIExecuteCommandList(FRHICommandList* CmdList)
 					DirtyUniformBuffers[SF_Domain] = 0xffff;
 					DirtyUniformBuffers[SF_Geometry] = 0xffff;
 				}
-				CmdPtr += sizeof(FRHICommandSetBoundShaderState);
 				RHICmd->BoundShaderState->Release();
 			}
 			break;
 		case ERCT_SetBlendState:
 			{
-				auto* RHICmd = (FRHICommandSetBlendState*)Cmd;
+				auto* RHICmd = Iter.NextCommand<FRHICommandSetBlendState>();
 				{
 					auto* State = (FD3D11BlendState*)RHICmd->State;
 					StateCache.SetBlendState(State->Resource, (const float*)&RHICmd->BlendFactor, 0xffffffff);
 				}
-				CmdPtr += sizeof(FRHICommandSetBlendState);
 				RHICmd->State->Release();
 			}
 			break;
 		case ERCT_SetStreamSource:
 			{
-				auto* RHICmd = (FRHICommandSetStreamSource*)Cmd;
+				auto* RHICmd = Iter.NextCommand<FRHICommandSetStreamSource>();
 				{
 					auto* VertexBuffer = (FD3D11VertexBuffer*)RHICmd->VertexBuffer;
 					ID3D11Buffer* D3DBuffer = VertexBuffer ? VertexBuffer->Resource : nullptr;
 					StateCache.SetStreamSource(D3DBuffer, RHICmd->StreamIndex, RHICmd->Stride, RHICmd->Offset);
 				}
-				CmdPtr += sizeof(FRHICommandSetStreamSource);
 				RHICmd->VertexBuffer->Release();
 			}
 			break;
 		default:
 			checkf(0, TEXT("Unimplemented RHICmd %d!"), Cmd->Type);
 		}
-
-		CmdPtr = Align(CmdPtr, FRHICommandList::Alignment);
 	}
+
+	Iter.CheckNumCommands();
 }
 
 // NVIDIA Depth Bounds Test interface
