UE4 - RHICommandList page allocator

diff --git a/Engine/Source/Runtime/RHI/Private/RHICommandList.cpp b/Engine/Source/Runtime/RHI/Private/RHICommandList.cpp
--- a/Engine/Source/Runtime/RHI/Private/RHICommandList.cpp
+++ b/Engine/Source/Runtime/RHI/Private/RHICommandList.cpp
@@ -46,8 +46,9 @@ static FAutoConsoleVariableRef CVarRHICmdMem(
 	);
 
 RHI_API FRHICommandListExecutor GRHICommandList;
-RHI_API FRHICommandList FRHICommandList::NullRHICommandList(0);
+RHI_API FRHICommandList FRHICommandList::NullRHICommandList;
 
+static_assert(sizeof(FRHICommand) == sizeof(FRHICommandNopEndOfPage), "These should match for filling in the end of a page when a new allocation won't fit");
 
 //@todo-rco: Temp hack!
 static FRHICommandList GCommandList;
@@ -71,6 +72,14 @@ void FRHICommandListExecutor::ExecuteList(FRHICommandList& CmdList)
 			{
 				// Nop
 				auto* RHICmd = Iter.NextCommand<FRHICommandNopBlob>();
+				(void)RHICmd;	// Unused
+			}
+			break;
+		case ERCT_NopEndOfPage:
+			{
+				// Nop
+				auto* RHICmd = Iter.NextCommand<FRHICommandNopEndOfPage>();
+				(void)RHICmd;	// Unused
 			}
 			break;
 		case ERCT_SetRasterizerState:
@@ -211,7 +220,7 @@ void FRHICommandListExecutor::ExecuteList(FRHICommandList& CmdList)
 			}
 			break;
 		default:
-			check(0);
+			checkf(0, TEXT("Unknown RHI Command %d!"), Cmd->Type);
 		}
 	}
 
@@ -283,3 +292,56 @@ void FRHICommandListExecutor::Verify()
 {
 	check(!bLock);
 }
+
+
+FRHICommandList::FMemManager::FMemManager() :
+	FirstPage(nullptr),
+	LastPage(nullptr)
+{
+}
+
+FRHICommandList::FMemManager::~FMemManager()
+{
+	while (FirstPage)
+	{
+		auto* Page = FirstPage->NextPage;
+		delete FirstPage;
+		FirstPage = Page;
+	}
+
+	FirstPage = LastPage = (FPage*)0x1; // we want this to not bypass the cmd list; please just crash.
+}
+
+const SIZE_T FRHICommandList::FMemManager::GetUsedMemory() const
+{
+	SIZE_T TotalUsedMemory = 0;
+	auto* Page = FirstPage;
+	do
+	{
+		SIZE_T UsedPageMemory = Page->MemUsed();
+		if (!UsedPageMemory)
+		{
+			return TotalUsedMemory;
+		}
+		TotalUsedMemory += UsedPageMemory;
+		Page = Page->NextPage;
+	}
+	while (Page);
+	return TotalUsedMemory;
+}
+
+void FRHICommandList::FMemManager::Reset()
+{
+	auto* Page = FirstPage;
+	while (Page)
+	{
+		Page->Reset();
+		Page = Page->NextPage;
+	}
+	LastPage = FirstPage;
+}
+
+const SIZE_T FRHICommandList::GetUsedMemory() const
+{
+	return MemManager.GetUsedMemory();
+}
diff --git a/Engine/Source/Runtime/RHI/Public/RHICommandList.h b/Engine/Source/Runtime/RHI/Public/RHICommandList.h
--- a/Engine/Source/Runtime/RHI/Public/RHICommandList.h
+++ b/Engine/Source/Runtime/RHI/Public/RHICommandList.h
@@ -8,8 +8,13 @@
 
 enum ERHICommandType
 {
-	// Start the enum at non-zero to help debug
+	// Start the enum at non-zero to help debugging
+
+	// List control
 	ERCT_NopBlob				= 512,	// NoOp, maybe has a buffer to be auto alloc/freed (512=0x200)
+	ERCT_NopEndOfPage,					// NoOp, stop traversing the current page
+
+	// RHI Commands
 	ERCT_SetRasterizerState,
 	ERCT_SetShaderParameter,
 	ERCT_SetShaderUniformBuffer,
@@ -24,6 +29,21 @@ enum ERHICommandType
 	ERCT_,
 };
 
+struct FRHICommand
+{
+	// This is a union only to force alignment to SIZE_T
+	union
+	{
+		ERHICommandType	Type;
+		SIZE_T			RawType;
+	};
+
+	uint32 ExtraSize() const
+	{
+		return 0;
+	}
+};
+
 class RHI_API FRHICommandList
 {
 private:
@@ -31,71 +51,114 @@ private:
 
 	struct FMemManager
 	{
-		FMemManager(int32 Size = DefaultMemSize)
+		struct FPage
 		{
-			MemSize = Size;
-			Memory = nullptr;
-			if (Size)
+			enum
+			{
+				PageSize = 64 * 1024
+			};
+
+			FPage() :
+				NextPage(nullptr)
 			{
-				Memory = (uint8*)FMemory::Malloc(MemSize, Alignment);
+				Head = (uint8*)FMemory::Malloc(PageSize, Alignment);
+				Current = Head;
+				Tail = Head + PageSize;
 			}
-		}
 
-		FORCEINLINE bool IsNull() const
-		{
-			return !Memory;
-		}
+			~FPage()
+			{
+				FMemory::Free(Head);
+				Head = Current = Tail = nullptr;
 
-		~FMemManager()
-		{
-			if (IsNull())
+				// Caller/owner is responsible for freeing the linked list
+				NextPage = nullptr;
+			}
+
+			void Reset()
 			{
-				FMemory::Free(Memory);
+				Current = Head;
 			}
-			MemSize = 0;
-			Memory = (uint8*)0x1; // we want this to not bypass the cmd list; please just crash.
-		}
 
-		uint8* Alloc(SIZE_T InSize)
-		{
-			checkSlow(!IsNull());
-			InSize = (InSize + (Alignment - 1)) & ~(Alignment - 1);
-			check(Tail + InSize + Alignment <= &Memory[MemSize]);
-			uint8* New = Tail;
-			Tail += InSize;
-			return New;
-		}
+			FORCEINLINE SIZE_T MemUsed() const
+			{
+				return Current - Head;
+			}
 
-		const uint8* GetHead() const
+			FORCEINLINE SIZE_T MemAvailable() const
+			{
+				return Tail - Current;
+			}
+
+			FORCEINLINE bool CanFitAllocation(SIZE_T AlignedSize) const
+			{
+				return (AlignedSize + Current <= Tail);
+			}
+
+			FORCEINLINE uint8* Alloc(SIZE_T AlignedSize)
+			{
+				auto* Ptr = Current;
+				Current += AlignedSize;
+				return Ptr;
+			}
+
+			uint8* Current;
+			uint8* Tail;
+			uint8* Head;
+
+			FPage* NextPage;
+		};
+
+		FMemManager();
+		~FMemManager();
+
+		FORCEINLINE bool IsNull() const
 		{
-			checkSlow(!IsNull());
-			return &Memory[0];
+			return (this == &NullRHICommandList.MemManager);
 		}
 
-		const uint8* GetTail() const
+		FORCEINLINE uint8* Alloc(SIZE_T InSize)
 		{
 			checkSlow(!IsNull());
-			return Tail;
-		}
+			InSize = (InSize + (Alignment - 1)) & ~(Alignment - 1);
+			bool bTryAgain = false;
 
-		const SIZE_T GetUsedMemory() const
-		{
-			return GetTail() - GetHead();
-		}
+			if (LastPage)
+			{
+				if (LastPage->CanFitAllocation(InSize))
+				{
+					return LastPage->Alloc(InSize);
+				}
+
+				// Insert End Of Page if there's trailing space
+				if (LastPage->MemAvailable() >= sizeof(FRHICommand))
+				{
+					auto* EndOfPageCmd = (FRHICommand*)LastPage->Alloc(sizeof(FRHICommand));
+					EndOfPageCmd->Type = ERCT_NopEndOfPage;
+				}
+			}
 
-		bool CanFitCommand(SIZE_T Size) const
-		{
-			return (Size + (Tail - Memory) <= (SIZE_T)MemSize);
+			checkfSlow(InSize <= FPage::PageSize, TEXT("Can't allocate an FPage bigger than %d!"), FPage::PageSize);
+			if (!FirstPage)
+			{
+				FirstPage = new FPage();
+				check(!LastPage);
+				LastPage = FirstPage;
+			}
+			else
+			{
+				LastPage->NextPage = new FPage();
+				LastPage = LastPage->NextPage;
+			}
+			return LastPage->Alloc(InSize);
 		}
 
-		void ResetTail()
-		{
-			Tail = Memory;
-		}
+		const SIZE_T GetUsedMemory() const;
+
+		void Reset();
 
-		uint8* Memory;
-		uint8* Tail;
-		int32 MemSize;
+		FPage* FirstPage;
+		FPage* LastPage;
 	};
 
 	FMemManager MemManager;
@@ -103,7 +166,6 @@ private:
 public:
 	enum
 	{
-		DefaultMemSize = 256 * 1024,
 		Alignment = sizeof(SIZE_T),
 	};
 
@@ -112,16 +174,11 @@ public:
 		return NullRHICommandList;
 	}
 
-	FRHICommandList(int32 Size = DefaultMemSize) :
-		MemManager(Size)
+	FRHICommandList()
 	{
 		Reset();
 	}
 
-	~FRHICommandList()
-	{
-	}
-
 	template <typename TShaderRHIParamRef>
 	inline void SetShaderUniformBuffer(TShaderRHIParamRef Shader, uint32 BaseIndex, FUniformBufferRHIParamRef UniformBufferRHI);
 
@@ -252,10 +309,7 @@ public:
 		check(IsNull());
 	}
 
-	FORCEINLINE const SIZE_T GetUsedMemory() const
-	{
-		return MemManager.GetUsedMemory();
-	}
+	const SIZE_T GetUsedMemory() const;
 
 protected:
 	enum EState
@@ -272,33 +326,17 @@ protected:
 		return MemManager.Alloc(InSize);
 	}
 
-	FORCEINLINE const uint8* GetHead() const
-	{
-		return MemManager.GetHead();
-	}
-
-	FORCEINLINE const uint8* GetTail() const
-	{
-		return MemManager.GetTail();
-	}
-
-	inline bool CanAddCommand() const
+	FORCEINLINE bool CanAddCommand() const
 	{
 		return (State == Ready);
 	}
 
-	FORCEINLINE bool CanFitCommand(SIZE_T Size) const
-	{
-		return MemManager.CanFitCommand(Size);
-	}
-
 	template <typename TCmd>
-	FORCEINLINE TCmd* AddCommand()
+	FORCEINLINE TCmd* AddCommand(uint32 ExtraData = 0)
 	{
 		checkSlow(!IsNull());
 		check(CanAddCommand());
-		check(CanFitCommand(sizeof(TCmd)));
-		TCmd* Cmd = (TCmd*)Alloc(sizeof(TCmd));
+		TCmd* Cmd = (TCmd*)Alloc(sizeof(TCmd) + ExtraData);
 		++NumCommands;
 		return Cmd;
 	}
@@ -307,7 +345,7 @@ protected:
 	{
 		State = Kicked;
 		NumCommands = 0;
-		MemManager.ResetTail();
+		MemManager.Reset();
 	}
 
 	FORCEINLINE bool Bypass()
diff --git a/Engine/Source/Runtime/RHI/Public/RHICommandList.inl b/Engine/Source/Runtime/RHI/Public/RHICommandList.inl
--- a/Engine/Source/Runtime/RHI/Public/RHICommandList.inl
+++ b/Engine/Source/Runtime/RHI/Public/RHICommandList.inl
@@ -8,21 +8,6 @@
 #define __RHICOMMANDLIST_INL__
 #include "RHICommandList.h"
 
-struct FRHICommand
-{
-	// This is a union only to force alignment to SIZE_T
-	union
-	{
-		ERHICommandType	Type;
-		SIZE_T			RawType;
-	};
-
-	uint32 ExtraSize() const
-	{
-		return 0;
-	}
-};
-
 // No constructors are given for classes derived from FRHICommand as they are not created using
 // a classical 'new' for perf reasons.
 
@@ -39,6 +24,10 @@ struct FRHICommandPerShader : public FRHICommand
 	}
 };
 
+struct FRHICommandNopEndOfPage : public FRHICommand
+{
+};
+
 struct FRHICommandNopBlob : public FRHICommand
 {
 	void* Data;
@@ -47,7 +36,9 @@ struct FRHICommandNopBlob : public FRHICommand
 	{
 		Type = ERCT_NopBlob;
 		Size = InSize;
-		Data = List->Alloc(InSize);
+
+		// Data is stored after 'this'
+		Data = &this[1];
 		FMemory::Memcpy(Data, InData, InSize);
 	}
 
@@ -245,7 +236,7 @@ FORCEINLINE void FRHICommandList::SetShaderParameter(TShaderRHIParamRef Shader,
 
 	if (bValueInStack)
 	{
-		auto* Blob = AddCommand<FRHICommandNopBlob>();
+		auto* Blob = AddCommand<FRHICommandNopBlob>(NumBytes);
 		Blob->Set(this, NumBytes, NewValue);
 		NewValue = Blob->Data;
 	}
@@ -381,6 +372,7 @@ FORCEINLINE void FRHICommandList::SetStreamSource(uint32 StreamIndex, FVertexBuf
 	VertexBuffer->AddRef();
 }
 
+
 // Helper class for traversing a FRHICommandList
 class FRHICommandListIterator
 {
@@ -389,8 +381,9 @@ public:
 	{
 		NumCommands = 0;
 		CmdListNumCommands = CmdList.NumCommands;
-		CmdPtr = CmdList.GetHead();
-		CmdTail = CmdList.GetTail();
+		Page = CmdList.MemManager.FirstPage;
+		CmdPtr = Page->Head;
+		CmdTail = Page->Current;
 	}
 
 	FORCEINLINE bool HasCommandsLeft() const
@@ -412,6 +405,25 @@ public:
 		CmdPtr += sizeof(T) + RHICmd->ExtraSize();
 		CmdPtr = Align(CmdPtr, FRHICommandList::Alignment);
 		++NumCommands;
+		if (CmdPtr >= CmdTail)
+		{
+			Page = Page->NextPage;
+			CmdPtr = Page ? Page->Head : nullptr;
+			CmdTail = Page ? Page->Current : nullptr;
+		}
+
+		return RHICmd;
+	}
+
+	// Specialization for EndOfPage
+	template <>
+	FORCEINLINE FRHICommandNopEndOfPage* NextCommand<FRHICommandNopEndOfPage>()
+	{
+		auto* RHICmd = (FRHICommandNopEndOfPage*)CmdPtr;
+		Page = Page->NextPage;
+		CmdPtr = Page ? Page->Head : nullptr;
+		CmdTail = Page ? Page->Current : nullptr;
+
 		return RHICmd;
 	}
 
@@ -422,6 +434,7 @@ public:
 	}
 
 protected:
+	FRHICommandList::FMemManager::FPage* Page;
 	uint32 NumCommands;
 	const uint8* CmdPtr;
 	const uint8* CmdTail;
diff --git a/Engine/Source/Runtime/Windows/D3D11RHI/Private/D3D11Commands.cpp b/Engine/Source/Runtime/Windows/D3D11RHI/Private/D3D11Commands.cpp
--- a/Engine/Source/Runtime/Windows/D3D11RHI/Private/D3D11Commands.cpp
+++ b/Engine/Source/Runtime/Windows/D3D11RHI/Private/D3D11Commands.cpp
@@ -2142,6 +2142,13 @@ void FD3D11DynamicRHI::RHIExecuteCommandList(FRHICommandList* CmdList)
 				(void)RHICmd;	// Unused
 			}
 			break;
+		case ERCT_NopEndOfPage:
+			{
+				// Nop
+				auto* RHICmd = Iter.NextCommand<FRHICommandNopEndOfPage>();
+				(void)RHICmd;	// Unused
+			}
+			break;
 		case ERCT_SetRasterizerState:
 			{
 				auto* RHICmd = Iter.NextCommand<FRHICommandSetRasterizerState>();
