UE4 - Fix for clang

diff --git a/Engine/Source/Runtime/RHI/Public/RHICommandList.h b/Engine/Source/Runtime/RHI/Public/RHICommandList.h
--- a/Engine/Source/Runtime/RHI/Public/RHICommandList.h
+++ b/Engine/Source/Runtime/RHI/Public/RHICommandList.h
@@ -42,6 +42,11 @@ struct FRHICommand
 	{
 		return 0;
 	}
+
+	enum
+	{
+		IsEndOfPage = 0,
+	};
 };
 
 class RHI_API FRHICommandList
diff --git a/Engine/Source/Runtime/RHI/Public/RHICommandList.inl b/Engine/Source/Runtime/RHI/Public/RHICommandList.inl
--- a/Engine/Source/Runtime/RHI/Public/RHICommandList.inl
+++ b/Engine/Source/Runtime/RHI/Public/RHICommandList.inl
@@ -26,6 +26,10 @@ struct FRHICommandPerShader : public FRHICommand
 
 struct FRHICommandNopEndOfPage : public FRHICommand
 {
+	enum
+	{
+		IsEndOfPage = 1
+	};
 };
 
 struct FRHICommandNopBlob : public FRHICommand
@@ -405,7 +409,7 @@ public:
 		CmdPtr += sizeof(T) + RHICmd->ExtraSize();
 		CmdPtr = Align(CmdPtr, FRHICommandList::Alignment);
 		++NumCommands;
-		if (CmdPtr >= CmdTail)
+		if (T::IsEndOfPage || CmdPtr >= CmdTail)
 		{
 			Page = Page->NextPage;
 			CmdPtr = Page ? Page->Head : nullptr;
@@ -415,18 +419,6 @@ public:
 		return RHICmd;
 	}
 
-	// Specialization for EndOfPage
-	template <>
-	FORCEINLINE FRHICommandNopEndOfPage* NextCommand<FRHICommandNopEndOfPage>()
-	{
-		auto* RHICmd = (FRHICommandNopEndOfPage*)CmdPtr;
-		Page = Page->NextPage;
-		CmdPtr = Page ? Page->Head : nullptr;
-		CmdTail = Page ? Page->Current : nullptr;
-
-		return RHICmd;
-	}
-
 	void CheckNumCommands()
 	{
 		checkf(CmdListNumCommands == NumCommands, TEXT("Missed %d Commands!"), CmdListNumCommands - NumCommands);
