UE4 - Fix RHICmdList assert on Elemental

diff --git a/Engine/Source/Runtime/RHI/Public/RHICommandList.h b/Engine/Source/Runtime/RHI/Public/RHICommandList.h
--- a/Engine/Source/Runtime/RHI/Public/RHICommandList.h
+++ b/Engine/Source/Runtime/RHI/Public/RHICommandList.h
@@ -69,6 +69,7 @@ private:
 				Head = (uint8*)FMemory::Malloc(PageSize, Alignment);
 				Current = Head;
 				Tail = Head + PageSize;
+				//::OutputDebugStringW(*FString::Printf(TEXT("PAGE %p, %p - %p\n"), this, Head, Tail));
 			}
 
 			~FPage()
@@ -138,6 +139,8 @@ private:
 				// Insert End Of Page if there's trailing space
 				if (LastPage->MemAvailable() >= sizeof(FRHICommand))
 				{
+					//@todo-rco: Fix me!
+					// This is an allocator command, so it should not be counted during regular execution...
 					auto* EndOfPageCmd = (FRHICommand*)LastPage->Alloc(sizeof(FRHICommand));
 					EndOfPageCmd->Type = ERCT_NopEndOfPage;
 				}
@@ -342,6 +345,7 @@ protected:
 		checkSlow(!IsNull());
 		check(CanAddCommand());
 		TCmd* Cmd = (TCmd*)Alloc(sizeof(TCmd) + ExtraData);
+		//::OutputDebugStringW(*FString::Printf(TEXT("Add %d: @ 0x%p, %d bytes\n"), NumCommands, (void*)Cmd, (sizeof(TCmd) + ExtraData)));
 		++NumCommands;
 		return Cmd;
 	}
diff --git a/Engine/Source/Runtime/RHI/Public/RHICommandList.inl b/Engine/Source/Runtime/RHI/Public/RHICommandList.inl
--- a/Engine/Source/Runtime/RHI/Public/RHICommandList.inl
+++ b/Engine/Source/Runtime/RHI/Public/RHICommandList.inl
@@ -402,14 +402,22 @@ public:
 	}
 
 	// Get the next RHICommand and advance the iterator
-	template <typename T>
-	FORCEINLINE T* NextCommand()
+	template <typename TCmd>
+	FORCEINLINE TCmd* NextCommand()
 	{
-		T* RHICmd = (T*)CmdPtr;
-		CmdPtr += sizeof(T) + RHICmd->ExtraSize();
+		TCmd* RHICmd = (TCmd*)CmdPtr;
+		//::OutputDebugStringW(*FString::Printf(TEXT("Exec %d: %d @ 0x%p, %d bytes\n"), NumCommands, (int32)RHICmd->Type, (void*)RHICmd, sizeof(TCmd) + RHICmd->ExtraSize()));
+		CmdPtr += sizeof(TCmd) + RHICmd->ExtraSize();
 		CmdPtr = Align(CmdPtr, FRHICommandList::Alignment);
-		++NumCommands;
-		if (T::IsEndOfPage || CmdPtr >= CmdTail)
+
+		//@todo-rco: Fix me!
+		if (!TCmd::IsEndOfPage)
+		{
+			// Don't count EOP as that is an allocator construct
+			++NumCommands;
+		}
+
+		if (TCmd::IsEndOfPage || CmdPtr >= CmdTail)
 		{
 			Page = Page->NextPage;
 			CmdPtr = Page ? Page->Head : nullptr;
