Provide native device access from the RHI, useful for 3rd party plugins that wish to directly render using the underlying device. Currently only exposed for D3D11.

Based on GitHub pull request 199 by CoherentUE4

diff --git a/Engine/Source/Runtime/EmptyRHI/Private/EmptyRHI.cpp b/Engine/Source/Runtime/EmptyRHI/Private/EmptyRHI.cpp
--- a/Engine/Source/Runtime/EmptyRHI/Private/EmptyRHI.cpp
+++ b/Engine/Source/Runtime/EmptyRHI/Private/EmptyRHI.cpp
@@ -177,3 +177,8 @@ void FEmptyDynamicRHI::RHIReleaseThreadOwnership()
 {
 
 }
+
+void* FEmptyDynamicRHI::RHIGetNativeDevice()
+{
+	return NULL;
+}
diff --git a/Engine/Source/Runtime/IOS/MetalRHI/Private/MetalRHI.cpp b/Engine/Source/Runtime/IOS/MetalRHI/Private/MetalRHI.cpp
--- a/Engine/Source/Runtime/IOS/MetalRHI/Private/MetalRHI.cpp
+++ b/Engine/Source/Runtime/IOS/MetalRHI/Private/MetalRHI.cpp
@@ -219,3 +219,8 @@ void FMetalDynamicRHI::RHIGpuTimeEnd(uint32 Hash, bool bCompute)
 {
 
 }
+
+void* FMetalDynamicRHI::RHIGetNativeDevice()
+{
+	return NULL;
+}
diff --git a/Engine/Source/Runtime/OpenGLDrv/Private/OpenGLDevice.cpp b/Engine/Source/Runtime/OpenGLDrv/Private/OpenGLDevice.cpp
--- a/Engine/Source/Runtime/OpenGLDrv/Private/OpenGLDevice.cpp
+++ b/Engine/Source/Runtime/OpenGLDrv/Private/OpenGLDevice.cpp
@@ -1159,3 +1159,8 @@ bool FOpenGLDynamicRHIModule::IsSupported()
 {
 	return PlatformInitOpenGL();
 }
+
+void* FOpenGLDynamicRHI::RHIGetNativeDevice()
+{
+	return NULL;
+}
diff --git a/Engine/Source/Runtime/RHI/Public/RHIMethods.h b/Engine/Source/Runtime/RHI/Public/RHIMethods.h
--- a/Engine/Source/Runtime/RHI/Public/RHIMethods.h
+++ b/Engine/Source/Runtime/RHI/Public/RHIMethods.h
@@ -1722,3 +1722,11 @@ DEFINE_RHIMETHOD_3(
 	float,MaxDepth,
 	,
 	);
+
+/**
+ * Provides access to the native device. Generally this should be avoided but is useful for third party plugins.
+ */
+DEFINE_RHIMETHOD_0(
+	void*,RHIGetNativeDevice,
+	return,return NULL
+	);
diff --git a/Engine/Source/Runtime/Windows/D3D11RHI/Private/D3D11Device.cpp b/Engine/Source/Runtime/Windows/D3D11RHI/Private/D3D11Device.cpp
--- a/Engine/Source/Runtime/Windows/D3D11RHI/Private/D3D11Device.cpp
+++ b/Engine/Source/Runtime/Windows/D3D11RHI/Private/D3D11Device.cpp
@@ -453,3 +453,8 @@ void FD3D11DynamicRHI::RHIFlushComputeShaderCache()
 {
 	// Nothing to do
 }
+
+void* FD3D11DynamicRHI::RHIGetNativeDevice()
+{
+	return Direct3DDevice.GetReference();
+}
