UE4 - rework RHI user-facing APIs to be safer and incorporate the desired sematics

diff --git a/Engine/Source/Runtime/RHI/Private/DynamicRHI.cpp b/Engine/Source/Runtime/RHI/Private/DynamicRHI.cpp
--- a/Engine/Source/Runtime/RHI/Private/DynamicRHI.cpp
+++ b/Engine/Source/Runtime/RHI/Private/DynamicRHI.cpp
@@ -76,14 +76,35 @@ void RHIExit()
 }
 
 // Implement the static RHI methods that call the dynamic RHI.
+#define DEFINE_RHIMETHOD_CMDLIST(Type,Name,ParameterTypesAndNames,ParameterNames,ReturnStatement,NullImplementation) \
+	RHI_API Type Name##_Internal ParameterTypesAndNames \
+	{ \
+		check(GDynamicRHI); \
+		ReturnStatement GDynamicRHI->RHI##Name ParameterNames; \
+	}
+#define DEFINE_RHIMETHOD_GLOBAL(Type,Name,ParameterTypesAndNames,ParameterNames,ReturnStatement,NullImplementation) \
+	RHI_API Type RHI##Name ParameterTypesAndNames \
+	{ \
+		check(GDynamicRHI); \
+		ReturnStatement GDynamicRHI->RHI##Name ParameterNames; \
+	}
+#define DEFINE_RHIMETHOD_GLOBALFLUSH(Type,Name,ParameterTypesAndNames,ParameterNames,ReturnStatement,NullImplementation) \
+	RHI_API Type Name##_Internal ParameterTypesAndNames \
+	{ \
+		check(GDynamicRHI); \
+		ReturnStatement GDynamicRHI->RHI##Name ParameterNames; \
+	}
 #define DEFINE_RHIMETHOD(Type,Name,ParameterTypesAndNames,ParameterNames,ReturnStatement,NullImplementation) \
-	RHI_API Type Name ParameterTypesAndNames \
+	RHI_API Type Name##_Internal ParameterTypesAndNames \
 	{ \
 		check(GDynamicRHI); \
-		ReturnStatement GDynamicRHI->Name ParameterNames; \
+		ReturnStatement GDynamicRHI->RHI##Name ParameterNames; \
 	}
 #include "RHIMethods.h"
 #undef DEFINE_RHIMETHOD
+#undef DEFINE_RHIMETHOD_CMDLIST
+#undef DEFINE_RHIMETHOD_GLOBAL
+#undef DEFINE_RHIMETHOD_GLOBALFLUSH
 
 #else
 
diff --git a/Engine/Source/Runtime/RHI/Private/RHICommandList.cpp b/Engine/Source/Runtime/RHI/Private/RHICommandList.cpp
--- a/Engine/Source/Runtime/RHI/Private/RHICommandList.cpp
+++ b/Engine/Source/Runtime/RHI/Private/RHICommandList.cpp
@@ -1,11 +1,9 @@
 // Copyright 1998-2014 Epic Games, Inc. All Rights Reserved.
 
-/*=============================================================================
-	RHI.cpp: Render Hardware Interface implementation.
-=============================================================================*/
 
 #include "RHI.h"
 #include "RHICommandList.h"
+//#include "../../RenderCore/Public/RenderingThread.h"
 
 DEFINE_STAT(STAT_RHICmdListExecuteTime);
 DEFINE_STAT(STAT_RHICmdListEnqueueTime);
@@ -21,6 +19,13 @@ static FAutoConsoleVariableRef CVarRHICmd(
 	TEXT("1: Enables the RHI Command List mode")
 	);
 
+static TAutoConsoleVariable<int32> CVarRHICmdBypass(
+	TEXT("r.RHICmdBypass"),
+	FRHICommandListExecutor::DefaultBypass,
+	TEXT("Whether to bypass the rhi command list and send the rhi commands immediately.\n")
+	TEXT("0: Disable, 1: Enable"),
+	ECVF_Cheat);
+
 static int32 GRHISkip = 0;
 static FAutoConsoleVariableRef CVarRHICmdSkip(
 	TEXT("r.RHISkip"),
@@ -46,19 +51,10 @@ static FAutoConsoleVariableRef CVarRHICmdMem(
 	);
 
 RHI_API FRHICommandListExecutor GRHICommandList;
-RHI_API FRHICommandList FRHICommandList::NullRHICommandList;
+//RHI_API FGlobalRHI GRHI;
 
 static_assert(sizeof(FRHICommand) == sizeof(FRHICommandNopEndOfPage), "These should match for filling in the end of a page when a new allocation won't fit");
 
-//@todo-rco: Temp hack!
-static FRHICommandList GCommandList;
-
-
-FRHICommandListExecutor::FRHICommandListExecutor() :
-	bLock(false)
-{
-}
-
 template <bool bOnlyTestMemAccess>
 void FRHICommandListExecutor::ExecuteList(FRHICommandList& CmdList)
 {
@@ -87,7 +83,7 @@ void FRHICommandListExecutor::ExecuteList(FRHICommandList& CmdList)
 				auto* RHICmd = Iter.NextCommand<FRHICommandSetRasterizerState>();
 				if (!bOnlyTestMemAccess)
 				{
-					RHISetRasterizerState(RHICmd->State);
+					SetRasterizerState_Internal(RHICmd->State);
 				}
 				RHICmd->State->Release();
 			}
@@ -99,12 +95,12 @@ void FRHICommandListExecutor::ExecuteList(FRHICommandList& CmdList)
 				{
 					switch (RHICmd->ShaderFrequency)
 					{
-					case SF_Vertex:		RHISetShaderParameter((FVertexShaderRHIParamRef)RHICmd->Shader, RHICmd->BufferIndex, RHICmd->BaseIndex, RHICmd->NumBytes, RHICmd->NewValue); break;
-					case SF_Hull:		RHISetShaderParameter((FHullShaderRHIParamRef)RHICmd->Shader, RHICmd->BufferIndex, RHICmd->BaseIndex, RHICmd->NumBytes, RHICmd->NewValue); break;
-					case SF_Domain:		RHISetShaderParameter((FDomainShaderRHIParamRef)RHICmd->Shader, RHICmd->BufferIndex, RHICmd->BaseIndex, RHICmd->NumBytes, RHICmd->NewValue); break;
-					case SF_Geometry:	RHISetShaderParameter((FGeometryShaderRHIParamRef)RHICmd->Shader, RHICmd->BufferIndex, RHICmd->BaseIndex, RHICmd->NumBytes, RHICmd->NewValue); break;
-					case SF_Pixel:		RHISetShaderParameter((FPixelShaderRHIParamRef)RHICmd->Shader, RHICmd->BufferIndex, RHICmd->BaseIndex, RHICmd->NumBytes, RHICmd->NewValue); break;
-					case SF_Compute:	RHISetShaderParameter((FComputeShaderRHIParamRef)RHICmd->Shader, RHICmd->BufferIndex, RHICmd->BaseIndex, RHICmd->NumBytes, RHICmd->NewValue); break;
+					case SF_Vertex:		SetShaderParameter_Internal((FVertexShaderRHIParamRef)RHICmd->Shader, RHICmd->BufferIndex, RHICmd->BaseIndex, RHICmd->NumBytes, RHICmd->NewValue); break;
+					case SF_Hull:		SetShaderParameter_Internal((FHullShaderRHIParamRef)RHICmd->Shader, RHICmd->BufferIndex, RHICmd->BaseIndex, RHICmd->NumBytes, RHICmd->NewValue); break;
+					case SF_Domain:		SetShaderParameter_Internal((FDomainShaderRHIParamRef)RHICmd->Shader, RHICmd->BufferIndex, RHICmd->BaseIndex, RHICmd->NumBytes, RHICmd->NewValue); break;
+					case SF_Geometry:	SetShaderParameter_Internal((FGeometryShaderRHIParamRef)RHICmd->Shader, RHICmd->BufferIndex, RHICmd->BaseIndex, RHICmd->NumBytes, RHICmd->NewValue); break;
+					case SF_Pixel:		SetShaderParameter_Internal((FPixelShaderRHIParamRef)RHICmd->Shader, RHICmd->BufferIndex, RHICmd->BaseIndex, RHICmd->NumBytes, RHICmd->NewValue); break;
+					case SF_Compute:	SetShaderParameter_Internal((FComputeShaderRHIParamRef)RHICmd->Shader, RHICmd->BufferIndex, RHICmd->BaseIndex, RHICmd->NumBytes, RHICmd->NewValue); break;
 					default: check(0); break;
 					}
 				}
@@ -117,12 +113,12 @@ void FRHICommandListExecutor::ExecuteList(FRHICommandList& CmdList)
 				{
 					switch (RHICmd->ShaderFrequency)
 					{
-					case SF_Vertex:		RHISetShaderUniformBuffer((FVertexShaderRHIParamRef)RHICmd->Shader, RHICmd->BaseIndex, RHICmd->UniformBuffer); break;
-					case SF_Hull:		RHISetShaderUniformBuffer((FHullShaderRHIParamRef)RHICmd->Shader, RHICmd->BaseIndex, RHICmd->UniformBuffer); break;
-					case SF_Domain:		RHISetShaderUniformBuffer((FDomainShaderRHIParamRef)RHICmd->Shader, RHICmd->BaseIndex, RHICmd->UniformBuffer); break;
-					case SF_Geometry:	RHISetShaderUniformBuffer((FGeometryShaderRHIParamRef)RHICmd->Shader, RHICmd->BaseIndex, RHICmd->UniformBuffer); break;
-					case SF_Pixel:		RHISetShaderUniformBuffer((FPixelShaderRHIParamRef)RHICmd->Shader, RHICmd->BaseIndex, RHICmd->UniformBuffer); break;
-					case SF_Compute:	RHISetShaderUniformBuffer((FComputeShaderRHIParamRef)RHICmd->Shader, RHICmd->BaseIndex, RHICmd->UniformBuffer); break;
+					case SF_Vertex:		SetShaderUniformBuffer_Internal((FVertexShaderRHIParamRef)RHICmd->Shader, RHICmd->BaseIndex, RHICmd->UniformBuffer); break;
+					case SF_Hull:		SetShaderUniformBuffer_Internal((FHullShaderRHIParamRef)RHICmd->Shader, RHICmd->BaseIndex, RHICmd->UniformBuffer); break;
+					case SF_Domain:		SetShaderUniformBuffer_Internal((FDomainShaderRHIParamRef)RHICmd->Shader, RHICmd->BaseIndex, RHICmd->UniformBuffer); break;
+					case SF_Geometry:	SetShaderUniformBuffer_Internal((FGeometryShaderRHIParamRef)RHICmd->Shader, RHICmd->BaseIndex, RHICmd->UniformBuffer); break;
+					case SF_Pixel:		SetShaderUniformBuffer_Internal((FPixelShaderRHIParamRef)RHICmd->Shader, RHICmd->BaseIndex, RHICmd->UniformBuffer); break;
+					case SF_Compute:	SetShaderUniformBuffer_Internal((FComputeShaderRHIParamRef)RHICmd->Shader, RHICmd->BaseIndex, RHICmd->UniformBuffer); break;
 					default: check(0); break;
 					}
 				}
@@ -137,12 +133,12 @@ void FRHICommandListExecutor::ExecuteList(FRHICommandList& CmdList)
 				{
 					switch (RHICmd->ShaderFrequency)
 					{
-					case SF_Vertex:		RHISetShaderSampler((FVertexShaderRHIParamRef)RHICmd->Shader, RHICmd->SamplerIndex, RHICmd->Sampler); break;
-					case SF_Hull:		RHISetShaderSampler((FHullShaderRHIParamRef)RHICmd->Shader, RHICmd->SamplerIndex, RHICmd->Sampler); break;
-					case SF_Domain:		RHISetShaderSampler((FDomainShaderRHIParamRef)RHICmd->Shader, RHICmd->SamplerIndex, RHICmd->Sampler); break;
-					case SF_Geometry:	RHISetShaderSampler((FGeometryShaderRHIParamRef)RHICmd->Shader, RHICmd->SamplerIndex, RHICmd->Sampler); break;
-					case SF_Pixel:		RHISetShaderSampler((FPixelShaderRHIParamRef)RHICmd->Shader, RHICmd->SamplerIndex, RHICmd->Sampler); break;
-					case SF_Compute:	RHISetShaderSampler((FComputeShaderRHIParamRef)RHICmd->Shader, RHICmd->SamplerIndex, RHICmd->Sampler); break;
+					case SF_Vertex:		SetShaderSampler_Internal((FVertexShaderRHIParamRef)RHICmd->Shader, RHICmd->SamplerIndex, RHICmd->Sampler); break;
+					case SF_Hull:		SetShaderSampler_Internal((FHullShaderRHIParamRef)RHICmd->Shader, RHICmd->SamplerIndex, RHICmd->Sampler); break;
+					case SF_Domain:		SetShaderSampler_Internal((FDomainShaderRHIParamRef)RHICmd->Shader, RHICmd->SamplerIndex, RHICmd->Sampler); break;
+					case SF_Geometry:	SetShaderSampler_Internal((FGeometryShaderRHIParamRef)RHICmd->Shader, RHICmd->SamplerIndex, RHICmd->Sampler); break;
+					case SF_Pixel:		SetShaderSampler_Internal((FPixelShaderRHIParamRef)RHICmd->Shader, RHICmd->SamplerIndex, RHICmd->Sampler); break;
+					case SF_Compute:	SetShaderSampler_Internal((FComputeShaderRHIParamRef)RHICmd->Shader, RHICmd->SamplerIndex, RHICmd->Sampler); break;
 					default: check(0); break;
 					}
 				}
@@ -157,12 +153,12 @@ void FRHICommandListExecutor::ExecuteList(FRHICommandList& CmdList)
 				{
 					switch (RHICmd->ShaderFrequency)
 					{
-					case SF_Vertex:		RHISetShaderTexture((FVertexShaderRHIParamRef)RHICmd->Shader, RHICmd->TextureIndex, RHICmd->Texture); break;
-					case SF_Hull:		RHISetShaderTexture((FHullShaderRHIParamRef)RHICmd->Shader, RHICmd->TextureIndex, RHICmd->Texture); break;
-					case SF_Domain:		RHISetShaderTexture((FDomainShaderRHIParamRef)RHICmd->Shader, RHICmd->TextureIndex, RHICmd->Texture); break;
-					case SF_Geometry:	RHISetShaderTexture((FGeometryShaderRHIParamRef)RHICmd->Shader, RHICmd->TextureIndex, RHICmd->Texture); break;
-					case SF_Pixel:		RHISetShaderTexture((FPixelShaderRHIParamRef)RHICmd->Shader, RHICmd->TextureIndex, RHICmd->Texture); break;
-					case SF_Compute:	RHISetShaderTexture((FComputeShaderRHIParamRef)RHICmd->Shader, RHICmd->TextureIndex, RHICmd->Texture); break;
+					case SF_Vertex:		SetShaderTexture_Internal((FVertexShaderRHIParamRef)RHICmd->Shader, RHICmd->TextureIndex, RHICmd->Texture); break;
+					case SF_Hull:		SetShaderTexture_Internal((FHullShaderRHIParamRef)RHICmd->Shader, RHICmd->TextureIndex, RHICmd->Texture); break;
+					case SF_Domain:		SetShaderTexture_Internal((FDomainShaderRHIParamRef)RHICmd->Shader, RHICmd->TextureIndex, RHICmd->Texture); break;
+					case SF_Geometry:	SetShaderTexture_Internal((FGeometryShaderRHIParamRef)RHICmd->Shader, RHICmd->TextureIndex, RHICmd->Texture); break;
+					case SF_Pixel:		SetShaderTexture_Internal((FPixelShaderRHIParamRef)RHICmd->Shader, RHICmd->TextureIndex, RHICmd->Texture); break;
+					case SF_Compute:	SetShaderTexture_Internal((FComputeShaderRHIParamRef)RHICmd->Shader, RHICmd->TextureIndex, RHICmd->Texture); break;
 					default: check(0); break;
 					}
 				}
@@ -170,12 +166,62 @@ void FRHICommandListExecutor::ExecuteList(FRHICommandList& CmdList)
 				RHICmd->Texture->Release();
 			}
 			break;
+		case ERCT_SetShaderResourceViewParameter:
+			{
+				auto* RHICmd = Iter.NextCommand<FRHICommandSetShaderResourceViewParameter>();
+				if (!bOnlyTestMemAccess)
+				{
+					switch (RHICmd->ShaderFrequency)
+					{
+					case SF_Vertex:		SetShaderResourceViewParameter_Internal((FVertexShaderRHIParamRef)RHICmd->Shader, RHICmd->SamplerIndex, RHICmd->SRV); break;
+					case SF_Hull:		SetShaderResourceViewParameter_Internal((FHullShaderRHIParamRef)RHICmd->Shader, RHICmd->SamplerIndex, RHICmd->SRV); break;
+					case SF_Domain:		SetShaderResourceViewParameter_Internal((FDomainShaderRHIParamRef)RHICmd->Shader, RHICmd->SamplerIndex, RHICmd->SRV); break;
+					case SF_Geometry:	SetShaderResourceViewParameter_Internal((FGeometryShaderRHIParamRef)RHICmd->Shader, RHICmd->SamplerIndex, RHICmd->SRV); break;
+					case SF_Pixel:		SetShaderResourceViewParameter_Internal((FPixelShaderRHIParamRef)RHICmd->Shader, RHICmd->SamplerIndex, RHICmd->SRV); break;
+					case SF_Compute:	SetShaderResourceViewParameter_Internal((FComputeShaderRHIParamRef)RHICmd->Shader, RHICmd->SamplerIndex, RHICmd->SRV); break;
+					default: check(0); break;
+					}
+				}
+				RHICmd->Shader->Release();
+				RHICmd->SRV->Release();
+			}
+			break;
+		case ERCT_SetUAVParameter:
+			{
+				auto* RHICmd = Iter.NextCommand<FRHICommandSetUAVParameter>();
+				if (!bOnlyTestMemAccess)
+				{
+					switch (RHICmd->ShaderFrequency)
+					{
+					case SF_Compute:	SetUAVParameter_Internal((FComputeShaderRHIParamRef)RHICmd->Shader, RHICmd->UAVIndex, RHICmd->UAV); break;
+					default: check(0); break;
+					}
+				}
+				RHICmd->Shader->Release();
+				RHICmd->UAV->Release();
+			}
+			break;
+		case ERCT_SetUAVParameter_IntialCount:
+			{
+				auto* RHICmd = Iter.NextCommand<FRHICommandSetUAVParameter_IntialCount>();
+				if (!bOnlyTestMemAccess)
+				{
+					switch (RHICmd->ShaderFrequency)
+					{
+					case SF_Compute:	SetUAVParameter_Internal((FComputeShaderRHIParamRef)RHICmd->Shader, RHICmd->UAVIndex, RHICmd->UAV, RHICmd->InitialCount); break;
+					default: check(0); break;
+					}
+				}
+				RHICmd->Shader->Release();
+				RHICmd->UAV->Release();
+			}
+			break;
 		case ERCT_DrawPrimitive:
 			{
 				auto* RHICmd = Iter.NextCommand<FRHICommandDrawPrimitive>();
 				if (!bOnlyTestMemAccess)
 				{
-					RHIDrawPrimitive(RHICmd->PrimitiveType, RHICmd->BaseVertexIndex, RHICmd->NumPrimitives, RHICmd->NumInstances);
+					DrawPrimitive_Internal(RHICmd->PrimitiveType, RHICmd->BaseVertexIndex, RHICmd->NumPrimitives, RHICmd->NumInstances);
 				}
 			}
 			break;
@@ -184,7 +230,7 @@ void FRHICommandListExecutor::ExecuteList(FRHICommandList& CmdList)
 				auto* RHICmd = Iter.NextCommand<FRHICommandDrawIndexedPrimitive>();
 				if (!bOnlyTestMemAccess)
 				{
-					RHIDrawIndexedPrimitive(RHICmd->IndexBuffer, RHICmd->PrimitiveType, RHICmd->BaseVertexIndex, RHICmd->MinIndex, RHICmd->NumVertices, RHICmd->StartIndex, RHICmd->NumPrimitives, RHICmd->NumInstances);
+					DrawIndexedPrimitive_Internal(RHICmd->IndexBuffer, RHICmd->PrimitiveType, RHICmd->BaseVertexIndex, RHICmd->MinIndex, RHICmd->NumVertices, RHICmd->StartIndex, RHICmd->NumPrimitives, RHICmd->NumInstances);
 				}
 				RHICmd->IndexBuffer->Release();
 			}
@@ -194,7 +240,7 @@ void FRHICommandListExecutor::ExecuteList(FRHICommandList& CmdList)
 				auto* RHICmd = Iter.NextCommand<FRHICommandSetBoundShaderState>();
 				if (!bOnlyTestMemAccess)
 				{
-					RHISetBoundShaderState(RHICmd->BoundShaderState);
+					SetBoundShaderState_Internal(RHICmd->BoundShaderState);
 				}
 				RHICmd->BoundShaderState->Release();
 			}
@@ -204,7 +250,7 @@ void FRHICommandListExecutor::ExecuteList(FRHICommandList& CmdList)
 				auto* RHICmd = Iter.NextCommand<FRHICommandSetBlendState>();
 				if (!bOnlyTestMemAccess)
 				{
-					RHISetBlendState(RHICmd->State, RHICmd->BlendFactor);
+					SetBlendState_Internal(RHICmd->State, RHICmd->BlendFactor);
 				}
 				RHICmd->State->Release();
 			}
@@ -214,7 +260,7 @@ void FRHICommandListExecutor::ExecuteList(FRHICommandList& CmdList)
 				auto* RHICmd = Iter.NextCommand<FRHICommandSetStreamSource>();
 				if (!bOnlyTestMemAccess)
 				{
-					RHISetStreamSource(RHICmd->StreamIndex, RHICmd->VertexBuffer, RHICmd->Stride, RHICmd->Offset);
+					SetStreamSource_Internal(RHICmd->StreamIndex, RHICmd->VertexBuffer, RHICmd->Stride, RHICmd->Offset);
 				}
 				RHICmd->VertexBuffer->Release();
 			}
@@ -224,7 +270,7 @@ void FRHICommandListExecutor::ExecuteList(FRHICommandList& CmdList)
 				auto* RHICmd = Iter.NextCommand<FRHICommandSetDepthStencilState>();
 				if (!bOnlyTestMemAccess)
 				{
-					RHISetDepthStencilState(RHICmd->State, RHICmd->StencilRef);
+					SetDepthStencilState_Internal(RHICmd->State, RHICmd->StencilRef);
 				}
 				RHICmd->State->Release();
 			}
@@ -238,23 +284,12 @@ void FRHICommandListExecutor::ExecuteList(FRHICommandList& CmdList)
 }
 
 
-void FRHICommandListExecutor::ExecuteAndFreeList(FRHICommandList& CmdList)
+void FRHICommandListExecutor::ExecuteList(FRHICommandList& CmdList)
 {
-	if (CmdList.IsNull())
-	{
-		// HACK!
-		return;
-	}
-
+	CmdList.bExecuting = true;
 	SCOPE_CYCLE_COUNTER(STAT_RHICmdListExecuteTime);
 
-	check(CmdList.CanAddCommand());
-
-	FScopeLock Lock(&CriticalSection);
-	check(bLock);
-	bLock = false;
-
-	CmdList.State = FRHICommandList::Executing;
+	check(IsInRenderingThread());
 
 	INC_MEMORY_STAT_BY(STAT_RHICmdListMemory, CmdList.GetUsedMemory());
 
@@ -263,6 +298,10 @@ void FRHICommandListExecutor::ExecuteAndFreeList(FRHICommandList& CmdList)
 
 	if (!bSkipRHICmdList)
 	{
+		if (&CmdList != &GetImmediateCommandList())
+		{
+			GetImmediateCommandList().Flush();
+		}
 		static IConsoleVariable* RHIExecCVar = IConsoleManager::Get().FindConsoleVariable(TEXT("r.RHIExec"));
 		if (RHIExecCVar && RHIExecCVar->GetInt() != 0)
 		{
@@ -284,26 +323,15 @@ void FRHICommandListExecutor::ExecuteAndFreeList(FRHICommandList& CmdList)
 	}
 
 	INC_DWORD_STAT_BY(STAT_RHICmdListCount, CmdList.NumCommands);
+#if !(UE_BUILD_SHIPPING || UE_BUILD_TEST)
+	if (&CmdList == &GetImmediateCommandList())
+	{
+		bLatchedBypass = (CVarRHICmdBypass.GetValueOnRenderThread() >= 1);
+	}
+#endif
 	CmdList.Reset();
 }
 
-FRHICommandList& FRHICommandListExecutor::CreateList()
-{
-	FScopeLock Lock(&CriticalSection);
-	check(!bLock);
-	bLock = true;
-	check(GCommandList.State == FRHICommandList::Kicked);
-	check(GCommandList.NumCommands == 0);
-	GCommandList.State = FRHICommandList::Ready;
-	return GCommandList;
-}
-
-void FRHICommandListExecutor::Verify()
-{
-	check(!bLock);
-}
-
-
 FRHICommandList::FMemManager::FMemManager() :
 	FirstPage(nullptr),
 	LastPage(nullptr)
@@ -319,7 +347,7 @@ FRHICommandList::FMemManager::~FMemManager()
 		FirstPage = Page;
 	}
 
-	FirstPage = LastPage = (FPage*)0x1; // we want this to not bypass the cmd list; please just crash.
+	FirstPage = LastPage = nullptr;
 }
 
 const SIZE_T FRHICommandList::FMemManager::GetUsedMemory() const
diff --git a/Engine/Source/Runtime/RHI/Public/DynamicRHI.h b/Engine/Source/Runtime/RHI/Public/DynamicRHI.h
--- a/Engine/Source/Runtime/RHI/Public/DynamicRHI.h
+++ b/Engine/Source/Runtime/RHI/Public/DynamicRHI.h
@@ -34,7 +34,7 @@ public:
 	virtual void PushEvent(const TCHAR* Name) {}
 	virtual void PopEvent() {}
 
-	#define DEFINE_RHIMETHOD(Type,Name,ParameterTypesAndNames,ParameterNames,ReturnStatement,NullImplementation) virtual Type Name ParameterTypesAndNames = 0
+	#define DEFINE_RHIMETHOD(Type,Name,ParameterTypesAndNames,ParameterNames,ReturnStatement,NullImplementation) virtual Type RHI##Name ParameterTypesAndNames = 0
 	#include "RHIMethods.h"
 	#undef DEFINE_RHIMETHOD
 };
@@ -43,10 +43,19 @@ public:
 extern RHI_API FDynamicRHI* GDynamicRHI;
 
 // Implement the statically bound RHI methods to simply call the dynamic RHI.
+#define DEFINE_RHIMETHOD_CMDLIST(Type,Name,ParameterTypesAndNames,ParameterNames,ReturnStatement,NullImplementation) \
+	extern RHI_API Type Name##_Internal ParameterTypesAndNames;
+#define DEFINE_RHIMETHOD_GLOBAL(Type,Name,ParameterTypesAndNames,ParameterNames,ReturnStatement,NullImplementation) \
+	extern RHI_API Type RHI##Name ParameterTypesAndNames;
+#define DEFINE_RHIMETHOD_GLOBALFLUSH(Type,Name,ParameterTypesAndNames,ParameterNames,ReturnStatement,NullImplementation) \
+	extern RHI_API Type Name##_Internal ParameterTypesAndNames;
 #define DEFINE_RHIMETHOD(Type,Name,ParameterTypesAndNames,ParameterNames,ReturnStatement,NullImplementation) \
-	extern RHI_API Type Name ParameterTypesAndNames;
+	extern RHI_API Type Name##_Internal ParameterTypesAndNames;
 #include "RHIMethods.h"
 #undef DEFINE_RHIMETHOD
+#undef DEFINE_RHIMETHOD_CMDLIST
+#undef DEFINE_RHIMETHOD_GLOBAL
+#undef DEFINE_RHIMETHOD_GLOBALFLUSH
 
 #endif	//#if USE_DYNAMIC_RHI
 
diff --git a/Engine/Source/Runtime/RHI/Public/RHI.h b/Engine/Source/Runtime/RHI/Public/RHI.h
--- a/Engine/Source/Runtime/RHI/Public/RHI.h
+++ b/Engine/Source/Runtime/RHI/Public/RHI.h
@@ -1025,9 +1025,21 @@ DECLARE_MEMORY_STAT_POOL_EXTERN(TEXT("Pixel buffer memory"),STAT_PixelBufferMemo
 
 #if !USE_DYNAMIC_RHI
 	// Define the statically bound RHI methods with the RHI name prefix.
-	#define DEFINE_RHIMETHOD(Type,Name,ParameterTypesAndNames,ParameterNames,ReturnStatement,NullImplementation) extern Type Name ParameterTypesAndNames
-	#include "RHIMethods.h"
-	#undef DEFINE_RHIMETHOD
+#define DEFINE_RHIMETHOD_CMDLIST(Type,Name,ParameterTypesAndNames,ParameterNames,ReturnStatement,NullImplementation) \
+	extern Type Name##_Internal ParameterTypesAndNames
+#define DEFINE_RHIMETHOD_GLOBAL(Type,Name,ParameterTypesAndNames,ParameterNames,ReturnStatement,NullImplementation) \
+	extern Type RHI##Name ParameterTypesAndNames
+#define DEFINE_RHIMETHOD_GLOBALFLUSH(Type,Name,ParameterTypesAndNames,ParameterNames,ReturnStatement,NullImplementation) \
+	extern Type Name##_Internal ParameterTypesAndNames
+#define DEFINE_RHIMETHOD(Type,Name,ParameterTypesAndNames,ParameterNames,ReturnStatement,NullImplementation) \
+	extern Type Name##_Internal ParameterTypesAndNames
+
+#include "RHIMethods.h"
+#undef DEFINE_RHIMETHOD
+#undef DEFINE_RHIMETHOD_CMDLIST
+#undef DEFINE_RHIMETHOD_GLOBAL
+#undef DEFINE_RHIMETHOD_GLOBALFLUSH
+
 #endif
 
 /** Initializes the RHI. */
diff --git a/Engine/Source/Runtime/RHI/Public/RHICommandList.h b/Engine/Source/Runtime/RHI/Public/RHICommandList.h
--- a/Engine/Source/Runtime/RHI/Public/RHICommandList.h
+++ b/Engine/Source/Runtime/RHI/Public/RHICommandList.h
@@ -26,7 +26,9 @@ enum ERHICommandType
 	ERCT_SetBlendState,
 	ERCT_SetStreamSource,
 	ERCT_SetDepthStencilState,
-
+	ERCT_SetShaderResourceViewParameter,
+	ERCT_SetUAVParameter,
+	ERCT_SetUAVParameter_IntialCount,
 	ERCT_,
 };
 
@@ -50,10 +52,9 @@ struct FRHICommand
 	};
 };
 
-class RHI_API FRHICommandList
+class RHI_API FRHICommandList : public FNoncopyable
 {
 private:
-	static FRHICommandList NullRHICommandList;
 
 	struct FMemManager
 	{
@@ -119,14 +120,8 @@ private:
 		FMemManager();
 		~FMemManager();
 
-		FORCEINLINE bool IsNull() const
-		{
-			return (this == &NullRHICommandList.MemManager);
-		}
-
 		FORCEINLINE uint8* Alloc(SIZE_T InSize)
 		{
-			checkSlow(!IsNull());
 			InSize = (InSize + (Alignment - 1)) & ~(Alignment - 1);
 			bool bTryAgain = false;
 
@@ -178,28 +173,28 @@ public:
 		Alignment = sizeof(SIZE_T),
 	};
 
-	static FORCEINLINE FRHICommandList& GetNullRef()
-	{
-		return NullRHICommandList;
-	}
-
 	FRHICommandList()
 	{
 		Reset();
 	}
+	~FRHICommandList()
+	{
+		Flush();
+	}
+	inline void Flush();
 
 	template <typename TShaderRHIParamRef>
 	inline void SetShaderUniformBuffer(TShaderRHIParamRef Shader, uint32 BaseIndex, FUniformBufferRHIParamRef UniformBufferRHI);
 
-	template <typename TShaderRHIParamRef>
-	inline void SetShaderResourceViewParameter(TShaderRHIParamRef Shader, uint32 SamplerIndex, FShaderResourceViewRHIParamRef SRV);
-
 	template <typename TShaderRHIParamRef>
 	inline void SetShaderParameter(TShaderRHIParamRef Shader, uint32 BufferIndex, uint32 BaseIndex, uint32 NumBytes, const void* Value, bool bValueInStack = false);
 
 	template <typename TShaderRHIParamRef>
 	inline void SetShaderTexture(TShaderRHIParamRef Shader, uint32 TextureIndex, FTextureRHIParamRef Texture);
 
+	template <typename TShaderRHIParamRef>
+	inline void SetShaderResourceViewParameter(TShaderRHIParamRef Shader, uint32 SamplerIndex, FShaderResourceViewRHIParamRef SRV);
+
 	template <typename TShaderRHIParamRef>
 	inline void SetShaderSampler(TShaderRHIParamRef Shader, uint32 SamplerIndex, FSamplerStateRHIParamRef State);
 
@@ -215,115 +210,12 @@ public:
 	inline void DrawPrimitive(uint32 PrimitiveType, uint32 BaseVertexIndex, uint32 NumPrimitives, uint32 NumInstances);
 	inline void DrawIndexedPrimitive(FIndexBufferRHIParamRef IndexBuffer, uint32 PrimitiveType, int32 BaseVertexIndex, uint32 MinIndex, uint32 NumVertices, uint32 StartIndex, uint32 NumPrimitives, uint32 NumInstances);
 	inline void SetStreamSource(uint32 StreamIndex, FVertexBufferRHIParamRef VertexBuffer, uint32 Stride, uint32 Offset);
-
-	// need cmd list implementations
-
-	FORCEINLINE void SetViewport(uint32 MinX, uint32 MinY, float MinZ, uint32 MaxX, uint32 MaxY, float MaxZ)
-	{
-		CheckIsNull();
-		RHISetViewport(MinX, MinY, MinZ, MaxX, MaxY, MaxZ);
-	}
-
 	inline void SetDepthStencilState(FDepthStencilStateRHIParamRef NewStateRHI, uint32 StencilRef = 0);
 
-	FORCEINLINE void BeginDrawPrimitiveUP(uint32 PrimitiveType, uint32 NumPrimitives, uint32 NumVertices, uint32 VertexDataStride, void*& OutVertexData)
-	{
-		CheckIsNull();
-		RHIBeginDrawPrimitiveUP(PrimitiveType, NumPrimitives, NumVertices, VertexDataStride, OutVertexData);
-	}
-	
-	FORCEINLINE void EndDrawPrimitiveUP()
-	{
-		CheckIsNull();
-		RHIEndDrawPrimitiveUP();
-	}
-
-	FORCEINLINE void CopyToResolveTarget(FTextureRHIParamRef SourceTextureRHI, FTextureRHIParamRef DestTextureRHI, bool bKeepOriginalSurface, const FResolveParams& ResolveParams)
-	{
-		CheckIsNull();
-		RHICopyToResolveTarget(SourceTextureRHI, DestTextureRHI, bKeepOriginalSurface, ResolveParams);
-	}
-
-	FORCEINLINE void BeginDrawIndexedPrimitiveUP(uint32 PrimitiveType, uint32 NumPrimitives, uint32 NumVertices, uint32 VertexDataStride, void*& OutVertexData, uint32 MinVertexIndex, uint32 NumIndices, uint32 IndexDataStride, void*& OutIndexData)
-	{
-		CheckIsNull();
-		RHIBeginDrawIndexedPrimitiveUP(PrimitiveType, NumPrimitives, NumVertices, VertexDataStride, OutVertexData, MinVertexIndex, NumIndices, IndexDataStride, OutIndexData);
-	}
-
-	FORCEINLINE void EndDrawIndexedPrimitiveUP()
-	{
-		CheckIsNull();
-		RHIEndDrawIndexedPrimitiveUP();
-	}
-
-	FORCEINLINE void Clear(bool bClearColor, const FLinearColor& Color, bool bClearDepth, float Depth, bool bClearStencil, uint32 Stencil, FIntRect ExcludeRect)
-	{
-		CheckIsNull();
-		RHIClear(bClearColor, Color, bClearDepth, Depth, bClearStencil, Stencil, ExcludeRect);
-	}
-
-	FORCEINLINE void ClearMRT(bool bClearColor, int32 NumClearColors, const FLinearColor* ClearColorArray, bool bClearDepth, float Depth, bool bClearStencil, uint32 Stencil, FIntRect ExcludeRect)
-	{
-		CheckIsNull();
-		RHIClearMRT(bClearColor, NumClearColors, ClearColorArray, bClearDepth, Depth, bClearStencil, Stencil, ExcludeRect);
-	}
-
-	FORCEINLINE void SetRenderTargets(uint32 NewNumSimultaneousRenderTargets, const FRHIRenderTargetView* NewRenderTargetsRHI, FTextureRHIParamRef NewDepthStencilTargetRHI, uint32 NewNumUAVs, const FUnorderedAccessViewRHIParamRef* UAVs)
-	{
-		CheckIsNull();
-		RHISetRenderTargets(NewNumSimultaneousRenderTargets, NewRenderTargetsRHI, NewDepthStencilTargetRHI, NewNumUAVs, UAVs);
-	}
-
-	FORCEINLINE void EnableDepthBoundsTest(bool bEnable, float MinDepth, float MaxDepth)
-	{
-		CheckIsNull();
-		RHIEnableDepthBoundsTest(bEnable, MinDepth, MaxDepth);
-	}
-
-	FORCEINLINE void SetScissorRect(bool bEnable, uint32 MinX, uint32 MinY, uint32 MaxX, uint32 MaxY)
-	{
-		CheckIsNull();
-		RHISetScissorRect(bEnable, MinX, MinY, MaxX, MaxY);
-	}
-
-	FORCEINLINE void ClearUAV(FUnorderedAccessViewRHIParamRef UnorderedAccessViewRHI, const uint32* Values)
-	{
-		CheckIsNull();
-		RHIClearUAV(UnorderedAccessViewRHI, Values);
-	}
-
-	FORCEINLINE void SetComputeShader(FComputeShaderRHIParamRef ComputeShaderRHI)
-	{
-		CheckIsNull();
-		RHISetComputeShader(ComputeShaderRHI);
-	}
-
-	FORCEINLINE void DrawPrimitiveIndirect(uint32 PrimitiveType, FVertexBufferRHIParamRef ArgumentBufferRHI, uint32 ArgumentOffset)
-	{
-		CheckIsNull();
-		RHIDrawPrimitiveIndirect(PrimitiveType, ArgumentBufferRHI, ArgumentOffset);
-	}
-
-	FORCEINLINE bool IsNull() const
-	{
-		return MemManager.IsNull();
-	}
-
-	FORCEINLINE void CheckIsNull() const
-	{
-		check(IsNull());
-	}
-
 	const SIZE_T GetUsedMemory() const;
 
-protected:
-	enum EState
-	{
-		Ready,
-		Executing,
-		Kicked,
-	};
-	EState State;
+private:
+	bool bExecuting;
 	uint32 NumCommands;
 
 	FORCEINLINE uint8* Alloc(SIZE_T InSize)
@@ -331,16 +223,15 @@ protected:
 		return MemManager.Alloc(InSize);
 	}
 
-	FORCEINLINE bool CanAddCommand() const
+	FORCEINLINE bool HasCommands() const
 	{
-		return (State == Ready);
+		return (NumCommands > 0);
 	}
 
 	template <typename TCmd>
 	FORCEINLINE TCmd* AddCommand(uint32 ExtraData = 0)
 	{
-		checkSlow(!IsNull());
-		check(CanAddCommand());
+		checkSlow(!bExecuting);
 		TCmd* Cmd = (TCmd*)Alloc(sizeof(TCmd) + ExtraData);
 		//::OutputDebugStringW(*FString::Printf(TEXT("Add %d: @ 0x%p, %d bytes\n"), NumCommands, (void*)Cmd, (sizeof(TCmd) + ExtraData)));
 		++NumCommands;
@@ -349,15 +240,12 @@ protected:
 
 	void Reset()
 	{
-		State = Kicked;
+		bExecuting = false;
 		NumCommands = 0;
 		MemManager.Reset();
 	}
 
-	FORCEINLINE bool Bypass()
-	{
-		return IsNull();
-	}
+	inline bool Bypass();
 
 	friend class FRHICommandListExecutor;
 	friend class FRHICommandListIterator;
@@ -366,24 +254,144 @@ protected:
 	friend struct FRHICommandNopBlob;
 };
 
+class RHI_API FRHICommandListImmediate : public FRHICommandList
+{
+	friend class FRHICommandListExecutor;
+	FRHICommandListImmediate()
+	{
+	}
+public:
+
+	#define DEFINE_RHIMETHOD_CMDLIST(Type,Name,ParameterTypesAndNames,ParameterNames,ReturnStatement,NullImplementation)
+	#define DEFINE_RHIMETHOD(Type, Name, ParameterTypesAndNames, ParameterNames, ReturnStatement, NullImplementation) \
+		FORCEINLINE Type Name ParameterTypesAndNames \
+		{ \
+			Flush(); \
+			ReturnStatement Name##_Internal ParameterNames; \
+		}
+	#define DEFINE_RHIMETHOD_GLOBAL(Type, Name, ParameterTypesAndNames, ParameterNames, ReturnStatement, NullImplementation) \
+		FORCEINLINE Type Name ParameterTypesAndNames \
+		{ \
+			ReturnStatement RHI##Name ParameterNames; \
+		}
+	#define DEFINE_RHIMETHOD_GLOBALFLUSH(Type, Name, ParameterTypesAndNames, ParameterNames, ReturnStatement, NullImplementation) \
+		FORCEINLINE Type Name ParameterTypesAndNames \
+		{ \
+			Flush(); \
+			ReturnStatement Name##_Internal ParameterNames; \
+		}
+	#include "RHIMethods.h"
+	#undef DEFINE_RHIMETHOD
+	#undef DEFINE_RHIMETHOD_CMDLIST
+	#undef DEFINE_RHIMETHOD_GLOBAL
+	#undef DEFINE_RHIMETHOD_GLOBALFLUSH
+};
+
+#if 0
+class RHI_API FRHICommandListBypass
+{
+	friend class FRHICommandListExecutor;
+	FRHICommandListBypass()
+	{
+	}
+public:
+#define DEFINE_RHIMETHOD_CMDLIST(Type,Name,ParameterTypesAndNames,ParameterNames,ReturnStatement,NullImplementation) \
+	FORCEINLINE Type Name ParameterTypesAndNames \
+	{ \
+		ReturnStatement Name##_Internal ParameterNames; \
+	}
+#define DEFINE_RHIMETHOD(Type, Name, ParameterTypesAndNames, ParameterNames, ReturnStatement, NullImplementation) \
+	FORCEINLINE Type Name ParameterTypesAndNames \
+	{ \
+		ReturnStatement Name##_Internal ParameterNames; \
+	}
+#define DEFINE_RHIMETHOD_GLOBAL(Type, Name, ParameterTypesAndNames, ParameterNames, ReturnStatement, NullImplementation) \
+	FORCEINLINE Type Name ParameterTypesAndNames \
+	{ \
+		ReturnStatement RHI##Name ParameterNames; \
+	}
+#define DEFINE_RHIMETHOD_GLOBALFLUSH(Type, Name, ParameterTypesAndNames, ParameterNames, ReturnStatement, NullImplementation) \
+	FORCEINLINE Type Name ParameterTypesAndNames \
+	{ \
+		ReturnStatement Name##_Internal ParameterNames; \
+	}
+#include "RHIMethods.h"
+#undef DEFINE_RHIMETHOD
+#undef DEFINE_RHIMETHOD_CMDLIST
+#undef DEFINE_RHIMETHOD_GLOBAL
+#undef DEFINE_RHIMETHOD_GLOBALFLUSH
+};
+#endif
+
 class RHI_API FRHICommandListExecutor
 {
 public:
-	FRHICommandListExecutor();
-	FRHICommandList& CreateList();
-	void ExecuteAndFreeList(FRHICommandList& CmdList);
+	enum
+	{
+		DefaultBypass = 1
+	};
+	FRHICommandListExecutor()
+		: bLatchedBypass(!!DefaultBypass)
+	{
+	}
+	static inline FRHICommandListImmediate& GetImmediateCommandList();
+	static inline FRHICommandListImmediate& GetRecursiveRHICommandList(); // Only for use by RHI implementations
 
-	void Verify();
+	void ExecuteList(FRHICommandList& CmdList);
+
+	FORCEINLINE void Verify()
+	{
+		check(!CommandListImmediate.HasCommands());
+	}
+	FORCEINLINE bool Bypass()
+	{
+#if !(UE_BUILD_SHIPPING || UE_BUILD_TEST)
+		return bLatchedBypass;
+#else
+		return !!DefaultBypass;
+#endif
+	}
 
 private:
-	FCriticalSection CriticalSection;
-	bool bLock;
+	bool bLatchedBypass;
+	FRHICommandListImmediate CommandListImmediate;
 
 	template <bool bOnlyTestMemAccess>
 	void ExecuteList(FRHICommandList& CmdList);
 };
 extern RHI_API FRHICommandListExecutor GRHICommandList;
 
+FORCEINLINE bool FRHICommandList::Bypass()
+{
+	return GRHICommandList.Bypass();
+}
+
+FORCEINLINE FRHICommandListImmediate& FRHICommandListExecutor::GetImmediateCommandList()
+{
+	return GRHICommandList.CommandListImmediate;
+}
+FORCEINLINE FRHICommandListImmediate& FRHICommandListExecutor::GetRecursiveRHICommandList()
+{
+	check(!GRHICommandList.CommandListImmediate.bExecuting); // queued RHI commands must not be implemented by calling other queued RHI commands
+	GRHICommandList.CommandListImmediate.Flush(); // if for some reason this is called non-recursively, this flushes everything to get us in a reasonable state for the recursive, hazardous calls
+	return GRHICommandList.CommandListImmediate;
+}
+
+#define DEFINE_RHIMETHOD_CMDLIST(Type,Name,ParameterTypesAndNames,ParameterNames,ReturnStatement,NullImplementation)
+#define DEFINE_RHIMETHOD(Type, Name, ParameterTypesAndNames, ParameterNames, ReturnStatement, NullImplementation)
+#define DEFINE_RHIMETHOD_GLOBAL(Type, Name, ParameterTypesAndNames, ParameterNames, ReturnStatement, NullImplementation)
+#define DEFINE_RHIMETHOD_GLOBALFLUSH(Type, Name, ParameterTypesAndNames, ParameterNames, ReturnStatement, NullImplementation) \
+	FORCEINLINE Type RHI##Name ParameterTypesAndNames \
+	{ \
+		ReturnStatement FRHICommandListExecutor::GetImmediateCommandList().Name ParameterNames; \
+	}
+#include "RHIMethods.h"
+#undef DEFINE_RHIMETHOD
+#undef DEFINE_RHIMETHOD_CMDLIST
+#undef DEFINE_RHIMETHOD_GLOBAL
+#undef DEFINE_RHIMETHOD_GLOBALFLUSH
+
+
 DECLARE_CYCLE_STAT_EXTERN(TEXT("RHI Command List Execute"),STAT_RHICmdListExecuteTime,STATGROUP_RHI, );
 DECLARE_CYCLE_STAT_EXTERN(TEXT("RHI Command List Enqueue"),STAT_RHICmdListEnqueueTime,STATGROUP_RHI, );
 DECLARE_DWORD_COUNTER_STAT_EXTERN(TEXT("RHI Command List memory"),STAT_RHICmdListMemory,STATGROUP_RHI,RHI_API);
diff --git a/Engine/Source/Runtime/RHI/Public/RHICommandList.inl b/Engine/Source/Runtime/RHI/Public/RHICommandList.inl
--- a/Engine/Source/Runtime/RHI/Public/RHICommandList.inl
+++ b/Engine/Source/Runtime/RHI/Public/RHICommandList.inl
@@ -123,6 +123,53 @@ struct FRHICommandSetShaderTexture : public FRHICommandPerShader
 	}
 };
 
+struct FRHICommandSetShaderResourceViewParameter : public FRHICommandPerShader
+{
+	uint32 SamplerIndex;
+	FShaderResourceViewRHIParamRef SRV;
+	template <typename TShaderRHIParamRef>
+	FORCEINLINE void Set(TShaderRHIParamRef InShader, uint32 InSamplerIndex, FShaderResourceViewRHIParamRef InSRV)
+	{
+		FRHICommandPerShader::SetShader<TShaderRHIParamRef>(InShader);
+		Type = ERCT_SetShaderResourceViewParameter;
+		Shader = InShader;
+		SamplerIndex = InSamplerIndex;
+		SRV = InSRV;
+	}
+};
+
+struct FRHICommandSetUAVParameter : public FRHICommandPerShader
+{
+	uint32 UAVIndex;
+	FUnorderedAccessViewRHIParamRef UAV;
+	template <typename TShaderRHIParamRef>
+	FORCEINLINE void Set(TShaderRHIParamRef InShader, uint32 InUAVIndex, FUnorderedAccessViewRHIParamRef InUAV)
+	{
+		FRHICommandPerShader::SetShader<TShaderRHIParamRef>(InShader);
+		Type = ERCT_SetUAVParameter;
+		Shader = InShader;
+		UAVIndex = InUAVIndex;
+		UAV = InUAV;
+	}
+};
+
+struct FRHICommandSetUAVParameter_IntialCount : public FRHICommandPerShader
+{
+	uint32 UAVIndex;
+	FUnorderedAccessViewRHIParamRef UAV;
+	uint32 InitialCount;
+	template <typename TShaderRHIParamRef>
+	FORCEINLINE void Set(TShaderRHIParamRef InShader, uint32 InUAVIndex, FUnorderedAccessViewRHIParamRef InUAV, uint32 InInitialCount)
+	{
+		FRHICommandPerShader::SetShader<TShaderRHIParamRef>(InShader);
+		Type = ERCT_SetUAVParameter_IntialCount;
+		Shader = InShader;
+		UAVIndex = InUAVIndex;
+		UAV = InUAV;
+		InitialCount = InInitialCount;
+	}
+};
+
 struct FRHICommandSetShaderSampler : public FRHICommandPerShader
 {
 	uint32 SamplerIndex;
@@ -221,7 +268,7 @@ FORCEINLINE void FRHICommandList::SetShaderUniformBuffer(TShaderRHIParamRef Shad
 {
 	if (Bypass())
 	{
-		RHISetShaderUniformBuffer(Shader, BaseIndex, UniformBuffer);
+		SetShaderUniformBuffer_Internal(Shader, BaseIndex, UniformBuffer);
 		return;
 	}
 	auto* Cmd = AddCommand<FRHICommandSetShaderUniformBuffer>();
@@ -230,23 +277,12 @@ FORCEINLINE void FRHICommandList::SetShaderUniformBuffer(TShaderRHIParamRef Shad
 	UniformBuffer->AddRef();
 }
 
-template <typename TShaderRHIParamRef>
-FORCEINLINE void FRHICommandList::SetShaderResourceViewParameter(TShaderRHIParamRef Shader, uint32 SamplerIndex, FShaderResourceViewRHIParamRef SRV)
-{
-	if (Bypass())
-	{
-		RHISetShaderResourceViewParameter(Shader, SamplerIndex, SRV);
-		return;
-	}
-	checkf(0, TEXT("TODO: Non-bypass version"));
-}
-
 template <typename TShaderRHIParamRef>
 FORCEINLINE void FRHICommandList::SetShaderParameter(TShaderRHIParamRef Shader, uint32 BufferIndex, uint32 BaseIndex, uint32 NumBytes, const void* NewValue, bool bValueInStack)
 {
 	if (Bypass())
 	{
-		RHISetShaderParameter(Shader, BufferIndex, BaseIndex, NumBytes, NewValue);
+		SetShaderParameter_Internal(Shader, BufferIndex, BaseIndex, NumBytes, NewValue);
 		return;
 	}
 
@@ -266,7 +302,7 @@ FORCEINLINE void FRHICommandList::SetShaderTexture(TShaderRHIParamRef Shader, ui
 {
 	if (Bypass())
 	{
-		RHISetShaderTexture(Shader, TextureIndex, Texture);
+		SetShaderTexture_Internal(Shader, TextureIndex, Texture);
 		return;
 	}
 
@@ -276,12 +312,27 @@ FORCEINLINE void FRHICommandList::SetShaderTexture(TShaderRHIParamRef Shader, ui
 	Texture->AddRef();
 }
 
+template <typename TShaderRHIParamRef>
+FORCEINLINE void FRHICommandList::SetShaderResourceViewParameter(TShaderRHIParamRef Shader, uint32 SamplerIndex, FShaderResourceViewRHIParamRef SRV)
+{
+	if (Bypass())
+	{
+		SetShaderResourceViewParameter_Internal(Shader, SamplerIndex, SRV);
+		return;
+	}
+	auto* Cmd = AddCommand<FRHICommandSetShaderResourceViewParameter>();
+	Cmd->Set<TShaderRHIParamRef>(Shader, SamplerIndex, SRV);
+	Shader->AddRef();
+	SRV->AddRef();
+}
+
+
 template <typename TShaderRHIParamRef>
 FORCEINLINE void FRHICommandList::SetShaderSampler(TShaderRHIParamRef Shader, uint32 SamplerIndex, FSamplerStateRHIParamRef State)
 {
 	if (Bypass())
 	{
-		RHISetShaderSampler(Shader, SamplerIndex, State);
+		SetShaderSampler_Internal(Shader, SamplerIndex, State);
 		return;
 	}
 
@@ -295,27 +346,33 @@ FORCEINLINE void FRHICommandList::SetUAVParameter(FComputeShaderRHIParamRef Shad
 {
 	if (Bypass())
 	{
-		RHISetUAVParameter(Shader, UAVIndex, UAV);
+		SetUAVParameter_Internal(Shader, UAVIndex, UAV);
 		return;
 	}
-	checkf(0, TEXT("TODO: Non-bypass version"));
+	auto* Cmd = AddCommand<FRHICommandSetUAVParameter>();
+	Cmd->Set<FComputeShaderRHIParamRef>(Shader, UAVIndex, UAV);
+	Shader->AddRef();
+	UAV->AddRef();
 }
 
 FORCEINLINE void FRHICommandList::SetUAVParameter(FComputeShaderRHIParamRef Shader, uint32 UAVIndex, FUnorderedAccessViewRHIParamRef UAV, uint32 InitialCount)
 {
 	if (Bypass())
 	{
-		RHISetUAVParameter(Shader, UAVIndex, UAV, InitialCount);
+		SetUAVParameter_Internal(Shader, UAVIndex, UAV, InitialCount);
 		return;
 	}
-	checkf(0, TEXT("TODO: Non-bypass version"));
+	auto* Cmd = AddCommand<FRHICommandSetUAVParameter_IntialCount>();
+	Cmd->Set<FComputeShaderRHIParamRef>(Shader, UAVIndex, UAV, InitialCount);
+	Shader->AddRef();
+	UAV->AddRef();
 }
 
 FORCEINLINE void FRHICommandList::SetRasterizerState(FRasterizerStateRHIParamRef State)
 {
 	if (Bypass())
 	{
-		RHISetRasterizerState(State);
+		SetRasterizerState_Internal(State);
 		return;
 	}
 
@@ -328,7 +385,7 @@ FORCEINLINE void FRHICommandList::SetDepthStencilState(FDepthStencilStateRHIPara
 {
 	if (Bypass())
 	{
-		RHISetDepthStencilState(State, StencilRef);
+		SetDepthStencilState_Internal(State, StencilRef);
 		return;
 	}
 
@@ -341,7 +398,7 @@ FORCEINLINE void FRHICommandList::DrawPrimitive(uint32 PrimitiveType, uint32 Bas
 {
 	if (Bypass())
 	{
-		RHIDrawPrimitive(PrimitiveType, BaseVertexIndex, NumPrimitives, NumInstances);
+		DrawPrimitive_Internal(PrimitiveType, BaseVertexIndex, NumPrimitives, NumInstances);
 		return;
 	}
 
@@ -353,7 +410,7 @@ FORCEINLINE void FRHICommandList::DrawIndexedPrimitive(FIndexBufferRHIParamRef I
 {
 	if (Bypass())
 	{
-		RHIDrawIndexedPrimitive(IndexBuffer, PrimitiveType, BaseVertexIndex, MinIndex, NumVertices, StartIndex, NumPrimitives, NumInstances);
+		DrawIndexedPrimitive_Internal(IndexBuffer, PrimitiveType, BaseVertexIndex, MinIndex, NumVertices, StartIndex, NumPrimitives, NumInstances);
 		return;
 	}
 
@@ -366,7 +423,7 @@ FORCEINLINE void FRHICommandList::SetBoundShaderState(FBoundShaderStateRHIParamR
 {
 	if (Bypass())
 	{
-		RHISetBoundShaderState(BoundShaderState);
+		SetBoundShaderState_Internal(BoundShaderState);
 		return;
 	}
 
@@ -379,7 +436,7 @@ FORCEINLINE void FRHICommandList::SetBlendState(FBlendStateRHIParamRef State, co
 {
 	if (Bypass())
 	{
-		RHISetBlendState(State, BlendFactor);
+		SetBlendState_Internal(State, BlendFactor);
 		return;
 	}
 
@@ -392,7 +449,7 @@ FORCEINLINE void FRHICommandList::SetStreamSource(uint32 StreamIndex, FVertexBuf
 {
 	if (Bypass())
 	{
-		RHISetStreamSource(StreamIndex, VertexBuffer, Stride, Offset);
+		SetStreamSource_Internal(StreamIndex, VertexBuffer, Stride, Offset);
 		return;
 	}
 
@@ -401,6 +458,16 @@ FORCEINLINE void FRHICommandList::SetStreamSource(uint32 StreamIndex, FVertexBuf
 	VertexBuffer->AddRef();
 }
 
+//--
+
+FORCEINLINE void FRHICommandList::Flush()
+{
+	if (HasCommands())
+	{
+		GRHICommandList.ExecuteList(*this);
+	}
+}
+
 
 // Helper class for traversing a FRHICommandList
 class FRHICommandListIterator
