UE4 - rework RHI user-facing APIs to be safer and incorporate the desired sematics

diff --git a/Engine/Source/Runtime/RHI/Public/RHIMethods.h b/Engine/Source/Runtime/RHI/Public/RHIMethods.h
--- a/Engine/Source/Runtime/RHI/Public/RHIMethods.h
+++ b/Engine/Source/Runtime/RHI/Public/RHIMethods.h
@@ -13,26 +13,92 @@
 #endif
 
 #ifdef GENERATE_VAX_FUNCTION_DECLARATIONS
-	#define DEFINE_RHIMETHOD_0(ReturnType,MethodName,ReturnStatement,NullImplementation) 
-		ReturnType MethodName();
+
+	#define DEFINE_RHIMETHOD_0(ReturnType,MethodName,ReturnStatement,NullImplementation) \
+		ReturnType MethodName##_Internal();
 	#define DEFINE_RHIMETHOD_1(ReturnType,MethodName,ParameterTypeA,ParameterNameA,ReturnStatement,NullImplementation) \
-		ReturnType MethodName(ParameterTypeA ParameterNameA);
+		ReturnType MethodName##_Internal(ParameterTypeA ParameterNameA);
 	#define DEFINE_RHIMETHOD_2(ReturnType,MethodName,ParameterTypeA,ParameterNameA,ParameterTypeB,ParameterNameB,ReturnStatement,NullImplementation) \
-		ReturnType MethodName(ParameterTypeA ParameterNameA,ParameterTypeB ParameterNameB);
+		ReturnType MethodName##_Internal(ParameterTypeA ParameterNameA, ParameterTypeB ParameterNameB);
 	#define DEFINE_RHIMETHOD_3(ReturnType,MethodName,ParameterTypeA,ParameterNameA,ParameterTypeB,ParameterNameB,ParameterTypeC,ParameterNameC,ReturnStatement,NullImplementation) \
-		ReturnType MethodName(ParameterTypeA ParameterNameA,ParameterTypeB ParameterNameB,ParameterTypeC ParameterNameC);
+		ReturnType MethodName##_Internal(ParameterTypeA ParameterNameA, ParameterTypeB ParameterNameB, ParameterTypeC ParameterNameC);
 	#define DEFINE_RHIMETHOD_4(ReturnType,MethodName,ParameterTypeA,ParameterNameA,ParameterTypeB,ParameterNameB,ParameterTypeC,ParameterNameC,ParameterTypeD,ParameterNameD,ReturnStatement,NullImplementation) \
-		ReturnType MethodName(ParameterTypeA ParameterNameA,ParameterTypeB ParameterNameB,ParameterTypeC ParameterNameC,ParameterTypeD ParameterNameD);
+		ReturnType MethodName##_Internal(ParameterTypeA ParameterNameA, ParameterTypeB ParameterNameB, ParameterTypeC ParameterNameC, ParameterTypeD ParameterNameD);
 	#define DEFINE_RHIMETHOD_5(ReturnType,MethodName,ParameterTypeA,ParameterNameA,ParameterTypeB,ParameterNameB,ParameterTypeC,ParameterNameC,ParameterTypeD,ParameterNameD,ParameterTypeE,ParameterNameE,ReturnStatement,NullImplementation) \
-		ReturnType MethodName(ParameterTypeA ParameterNameA,ParameterTypeB ParameterNameB,ParameterTypeC ParameterNameC,ParameterTypeD ParameterNameD,ParameterTypeE ParameterNameE);
+		ReturnType MethodName##_Internal(ParameterTypeA ParameterNameA, ParameterTypeB ParameterNameB, ParameterTypeC ParameterNameC, ParameterTypeD ParameterNameD, ParameterTypeE ParameterNameE);
 	#define DEFINE_RHIMETHOD_6(ReturnType,MethodName,ParameterTypeA,ParameterNameA,ParameterTypeB,ParameterNameB,ParameterTypeC,ParameterNameC,ParameterTypeD,ParameterNameD,ParameterTypeE,ParameterNameE,ParameterTypeF,ParameterNameF,ReturnStatement,NullImplementation) \
-		ReturnType MethodName(ParameterTypeA ParameterNameA,ParameterTypeB ParameterNameB,ParameterTypeC ParameterNameC,ParameterTypeD ParameterNameD,ParameterTypeE ParameterNameE,ParameterTypeF ParameterNameF);
+		ReturnType MethodName##_Internal(ParameterTypeA ParameterNameA, ParameterTypeB ParameterNameB, ParameterTypeC ParameterNameC, ParameterTypeD ParameterNameD, ParameterTypeE ParameterNameE, ParameterTypeF ParameterNameF);
 	#define DEFINE_RHIMETHOD_7(ReturnType,MethodName,ParameterTypeA,ParameterNameA,ParameterTypeB,ParameterNameB,ParameterTypeC,ParameterNameC,ParameterTypeD,ParameterNameD,ParameterTypeE,ParameterNameE,ParameterTypeF,ParameterNameF,ParameterTypeG,ParameterNameG,ReturnStatement,NullImplementation) \
-		ReturnType MethodName(ParameterTypeA ParameterNameA,ParameterTypeB ParameterNameB,ParameterTypeC ParameterNameC,ParameterTypeD ParameterNameD,ParameterTypeE ParameterNameE,ParameterTypeF ParameterNameF,ParameterTypeG ParameterNameG);
+		ReturnType MethodName##_Internal(ParameterTypeA ParameterNameA, ParameterTypeB ParameterNameB, ParameterTypeC ParameterNameC, ParameterTypeD ParameterNameD, ParameterTypeE ParameterNameE, ParameterTypeF ParameterNameF, ParameterTypeG ParameterNameG);
 	#define DEFINE_RHIMETHOD_8(ReturnType,MethodName,ParameterTypeA,ParameterNameA,ParameterTypeB,ParameterNameB,ParameterTypeC,ParameterNameC,ParameterTypeD,ParameterNameD,ParameterTypeE,ParameterNameE,ParameterTypeF,ParameterNameF,ParameterTypeG,ParameterNameG,ParameterTypeH,ParameterNameH,ReturnStatement,NullImplementation) \
-		ReturnType MethodName(ParameterTypeA ParameterNameA,ParameterTypeB ParameterNameB,ParameterTypeC ParameterNameC,ParameterTypeD ParameterNameD,ParameterTypeE ParameterNameE,ParameterTypeF ParameterNameF,ParameterTypeG ParameterNameG,ParameterTypeH ParameterNameH);
+		ReturnType MethodName##_Internal(ParameterTypeA ParameterNameA, ParameterTypeB ParameterNameB, ParameterTypeC ParameterNameC, ParameterTypeD ParameterNameD, ParameterTypeE ParameterNameE, ParameterTypeF ParameterNameF, ParameterTypeG ParameterNameG, ParameterTypeH ParameterNameH);
 	#define DEFINE_RHIMETHOD_9(ReturnType,MethodName,ParameterTypeA,ParameterNameA,ParameterTypeB,ParameterNameB,ParameterTypeC,ParameterNameC,ParameterTypeD,ParameterNameD,ParameterTypeE,ParameterNameE,ParameterTypeF,ParameterNameF,ParameterTypeG,ParameterNameG,ParameterTypeH,ParameterNameH,ParameterTypeI,ParameterNameI,ReturnStatement,NullImplementation) \
-		ReturnType MethodName(ParameterTypeA ParameterNameA,ParameterTypeB ParameterNameB,ParameterTypeC ParameterNameC,ParameterTypeD ParameterNameD,ParameterTypeE ParameterNameE,ParameterTypeF ParameterNameF,ParameterTypeG ParameterNameG,ParameterTypeH ParameterNameH,ParameterTypeI ParameterNameI);
+		ReturnType MethodName##_Internal(ParameterTypeA ParameterNameA, ParameterTypeB ParameterNameB, ParameterTypeC ParameterNameC, ParameterTypeD ParameterNameD, ParameterTypeE ParameterNameE, ParameterTypeF ParameterNameF, ParameterTypeG ParameterNameG, ParameterTypeH ParameterNameH, ParameterTypeI ParameterNameI);
+
+	#define DEFINE_RHIMETHOD_CMDLIST_0(ReturnType,MethodName,ReturnStatement,NullImplementation) \
+		ReturnType MethodName##_Internal();
+	#define DEFINE_RHIMETHOD_CMDLIST_1(ReturnType,MethodName,ParameterTypeA,ParameterNameA,ReturnStatement,NullImplementation) \
+		ReturnType MethodName##_Internal(ParameterTypeA ParameterNameA);
+	#define DEFINE_RHIMETHOD_CMDLIST_2(ReturnType,MethodName,ParameterTypeA,ParameterNameA,ParameterTypeB,ParameterNameB,ReturnStatement,NullImplementation) \
+		ReturnType MethodName##_Internal(ParameterTypeA ParameterNameA, ParameterTypeB ParameterNameB);
+	#define DEFINE_RHIMETHOD_CMDLIST_3(ReturnType,MethodName,ParameterTypeA,ParameterNameA,ParameterTypeB,ParameterNameB,ParameterTypeC,ParameterNameC,ReturnStatement,NullImplementation) \
+		ReturnType MethodName##_Internal(ParameterTypeA ParameterNameA, ParameterTypeB ParameterNameB, ParameterTypeC ParameterNameC);
+	#define DEFINE_RHIMETHOD_CMDLIST_4(ReturnType,MethodName,ParameterTypeA,ParameterNameA,ParameterTypeB,ParameterNameB,ParameterTypeC,ParameterNameC,ParameterTypeD,ParameterNameD,ReturnStatement,NullImplementation) \
+		ReturnType MethodName##_Internal(ParameterTypeA ParameterNameA, ParameterTypeB ParameterNameB, ParameterTypeC ParameterNameC, ParameterTypeD ParameterNameD);
+	#define DEFINE_RHIMETHOD_CMDLIST_5(ReturnType,MethodName,ParameterTypeA,ParameterNameA,ParameterTypeB,ParameterNameB,ParameterTypeC,ParameterNameC,ParameterTypeD,ParameterNameD,ParameterTypeE,ParameterNameE,ReturnStatement,NullImplementation) \
+		ReturnType MethodName##_Internal(ParameterTypeA ParameterNameA, ParameterTypeB ParameterNameB, ParameterTypeC ParameterNameC, ParameterTypeD ParameterNameD, ParameterTypeE ParameterNameE);
+	#define DEFINE_RHIMETHOD_CMDLIST_6(ReturnType,MethodName,ParameterTypeA,ParameterNameA,ParameterTypeB,ParameterNameB,ParameterTypeC,ParameterNameC,ParameterTypeD,ParameterNameD,ParameterTypeE,ParameterNameE,ParameterTypeF,ParameterNameF,ReturnStatement,NullImplementation) \
+		ReturnType MethodName##_Internal(ParameterTypeA ParameterNameA, ParameterTypeB ParameterNameB, ParameterTypeC ParameterNameC, ParameterTypeD ParameterNameD, ParameterTypeE ParameterNameE, ParameterTypeF ParameterNameF);
+	#define DEFINE_RHIMETHOD_CMDLIST_7(ReturnType,MethodName,ParameterTypeA,ParameterNameA,ParameterTypeB,ParameterNameB,ParameterTypeC,ParameterNameC,ParameterTypeD,ParameterNameD,ParameterTypeE,ParameterNameE,ParameterTypeF,ParameterNameF,ParameterTypeG,ParameterNameG,ReturnStatement,NullImplementation) \
+		ReturnType MethodName##_Internal(ParameterTypeA ParameterNameA, ParameterTypeB ParameterNameB, ParameterTypeC ParameterNameC, ParameterTypeD ParameterNameD, ParameterTypeE ParameterNameE, ParameterTypeF ParameterNameF, ParameterTypeG ParameterNameG);
+	#define DEFINE_RHIMETHOD_CMDLIST_8(ReturnType,MethodName,ParameterTypeA,ParameterNameA,ParameterTypeB,ParameterNameB,ParameterTypeC,ParameterNameC,ParameterTypeD,ParameterNameD,ParameterTypeE,ParameterNameE,ParameterTypeF,ParameterNameF,ParameterTypeG,ParameterNameG,ParameterTypeH,ParameterNameH,ReturnStatement,NullImplementation) \
+		ReturnType MethodName##_Internal(ParameterTypeA ParameterNameA, ParameterTypeB ParameterNameB, ParameterTypeC ParameterNameC, ParameterTypeD ParameterNameD, ParameterTypeE ParameterNameE, ParameterTypeF ParameterNameF, ParameterTypeG ParameterNameG, ParameterTypeH ParameterNameH);
+	#define DEFINE_RHIMETHOD_CMDLIST_9(ReturnType,MethodName,ParameterTypeA,ParameterNameA,ParameterTypeB,ParameterNameB,ParameterTypeC,ParameterNameC,ParameterTypeD,ParameterNameD,ParameterTypeE,ParameterNameE,ParameterTypeF,ParameterNameF,ParameterTypeG,ParameterNameG,ParameterTypeH,ParameterNameH,ParameterTypeI,ParameterNameI,ReturnStatement,NullImplementation) \
+		ReturnType MethodName##_Internal(ParameterTypeA ParameterNameA, ParameterTypeB ParameterNameB, ParameterTypeC ParameterNameC, ParameterTypeD ParameterNameD, ParameterTypeE ParameterNameE, ParameterTypeF ParameterNameF, ParameterTypeG ParameterNameG, ParameterTypeH ParameterNameH, ParameterTypeI ParameterNameI);
+
+	#define DEFINE_RHIMETHOD_GLOBAL_0(ReturnType,MethodName,ReturnStatement,NullImplementation) \
+		ReturnType RHI##MethodName();
+	#define DEFINE_RHIMETHOD_GLOBAL_1(ReturnType,MethodName,ParameterTypeA,ParameterNameA,ReturnStatement,NullImplementation) \
+		ReturnType RHI##MethodName(ParameterTypeA ParameterNameA);
+	#define DEFINE_RHIMETHOD_GLOBAL_2(ReturnType,MethodName,ParameterTypeA,ParameterNameA,ParameterTypeB,ParameterNameB,ReturnStatement,NullImplementation) \
+		ReturnType RHI##MethodName(ParameterTypeA ParameterNameA, ParameterTypeB ParameterNameB);
+	#define DEFINE_RHIMETHOD_GLOBAL_3(ReturnType,MethodName,ParameterTypeA,ParameterNameA,ParameterTypeB,ParameterNameB,ParameterTypeC,ParameterNameC,ReturnStatement,NullImplementation) \
+		ReturnType RHI##MethodName(ParameterTypeA ParameterNameA, ParameterTypeB ParameterNameB, ParameterTypeC ParameterNameC);
+	#define DEFINE_RHIMETHOD_GLOBAL_4(ReturnType,MethodName,ParameterTypeA,ParameterNameA,ParameterTypeB,ParameterNameB,ParameterTypeC,ParameterNameC,ParameterTypeD,ParameterNameD,ReturnStatement,NullImplementation) \
+		ReturnType RHI##MethodName(ParameterTypeA ParameterNameA, ParameterTypeB ParameterNameB, ParameterTypeC ParameterNameC, ParameterTypeD ParameterNameD);
+	#define DEFINE_RHIMETHOD_GLOBAL_5(ReturnType,MethodName,ParameterTypeA,ParameterNameA,ParameterTypeB,ParameterNameB,ParameterTypeC,ParameterNameC,ParameterTypeD,ParameterNameD,ParameterTypeE,ParameterNameE,ReturnStatement,NullImplementation) \
+		ReturnType RHI##MethodName(ParameterTypeA ParameterNameA, ParameterTypeB ParameterNameB, ParameterTypeC ParameterNameC, ParameterTypeD ParameterNameD, ParameterTypeE ParameterNameE);
+	#define DEFINE_RHIMETHOD_GLOBAL_6(ReturnType,MethodName,ParameterTypeA,ParameterNameA,ParameterTypeB,ParameterNameB,ParameterTypeC,ParameterNameC,ParameterTypeD,ParameterNameD,ParameterTypeE,ParameterNameE,ParameterTypeF,ParameterNameF,ReturnStatement,NullImplementation) \
+		ReturnType RHI##MethodName(ParameterTypeA ParameterNameA, ParameterTypeB ParameterNameB, ParameterTypeC ParameterNameC, ParameterTypeD ParameterNameD, ParameterTypeE ParameterNameE, ParameterTypeF ParameterNameF);
+	#define DEFINE_RHIMETHOD_GLOBAL_7(ReturnType,MethodName,ParameterTypeA,ParameterNameA,ParameterTypeB,ParameterNameB,ParameterTypeC,ParameterNameC,ParameterTypeD,ParameterNameD,ParameterTypeE,ParameterNameE,ParameterTypeF,ParameterNameF,ParameterTypeG,ParameterNameG,ReturnStatement,NullImplementation) \
+		ReturnType RHI##MethodName(ParameterTypeA ParameterNameA, ParameterTypeB ParameterNameB, ParameterTypeC ParameterNameC, ParameterTypeD ParameterNameD, ParameterTypeE ParameterNameE, ParameterTypeF ParameterNameF, ParameterTypeG ParameterNameG);
+	#define DEFINE_RHIMETHOD_GLOBAL_8(ReturnType,MethodName,ParameterTypeA,ParameterNameA,ParameterTypeB,ParameterNameB,ParameterTypeC,ParameterNameC,ParameterTypeD,ParameterNameD,ParameterTypeE,ParameterNameE,ParameterTypeF,ParameterNameF,ParameterTypeG,ParameterNameG,ParameterTypeH,ParameterNameH,ReturnStatement,NullImplementation) \
+		ReturnType RHI##MethodName(ParameterTypeA ParameterNameA, ParameterTypeB ParameterNameB, ParameterTypeC ParameterNameC, ParameterTypeD ParameterNameD, ParameterTypeE ParameterNameE, ParameterTypeF ParameterNameF, ParameterTypeG ParameterNameG, ParameterTypeH ParameterNameH);
+	#define DEFINE_RHIMETHOD_GLOBAL_9(ReturnType,MethodName,ParameterTypeA,ParameterNameA,ParameterTypeB,ParameterNameB,ParameterTypeC,ParameterNameC,ParameterTypeD,ParameterNameD,ParameterTypeE,ParameterNameE,ParameterTypeF,ParameterNameF,ParameterTypeG,ParameterNameG,ParameterTypeH,ParameterNameH,ParameterTypeI,ParameterNameI,ReturnStatement,NullImplementation) \
+		ReturnType RHI##MethodName(ParameterTypeA ParameterNameA, ParameterTypeB ParameterNameB, ParameterTypeC ParameterNameC, ParameterTypeD ParameterNameD, ParameterTypeE ParameterNameE, ParameterTypeF ParameterNameF, ParameterTypeG ParameterNameG, ParameterTypeH ParameterNameH, ParameterTypeI ParameterNameI);
+
+	#define DEFINE_RHIMETHOD_GLOBALFLUSH_0(ReturnType,MethodName,ReturnStatement,NullImplementation) \
+		ReturnType RHI##MethodName();
+	#define DEFINE_RHIMETHOD_GLOBALFLUSH_1(ReturnType,MethodName,ParameterTypeA,ParameterNameA,ReturnStatement,NullImplementation) \
+		ReturnType RHI##MethodName(ParameterTypeA ParameterNameA);
+	#define DEFINE_RHIMETHOD_GLOBALFLUSH_2(ReturnType,MethodName,ParameterTypeA,ParameterNameA,ParameterTypeB,ParameterNameB,ReturnStatement,NullImplementation) \
+		ReturnType RHI##MethodName(ParameterTypeA ParameterNameA, ParameterTypeB ParameterNameB);
+	#define DEFINE_RHIMETHOD_GLOBALFLUSH_3(ReturnType,MethodName,ParameterTypeA,ParameterNameA,ParameterTypeB,ParameterNameB,ParameterTypeC,ParameterNameC,ReturnStatement,NullImplementation) \
+		ReturnType RHI##MethodName(ParameterTypeA ParameterNameA, ParameterTypeB ParameterNameB, ParameterTypeC ParameterNameC);
+	#define DEFINE_RHIMETHOD_GLOBALFLUSH_4(ReturnType,MethodName,ParameterTypeA,ParameterNameA,ParameterTypeB,ParameterNameB,ParameterTypeC,ParameterNameC,ParameterTypeD,ParameterNameD,ReturnStatement,NullImplementation) \
+		ReturnType RHI##MethodName(ParameterTypeA ParameterNameA, ParameterTypeB ParameterNameB, ParameterTypeC ParameterNameC, ParameterTypeD ParameterNameD);
+	#define DEFINE_RHIMETHOD_GLOBALFLUSH_5(ReturnType,MethodName,ParameterTypeA,ParameterNameA,ParameterTypeB,ParameterNameB,ParameterTypeC,ParameterNameC,ParameterTypeD,ParameterNameD,ParameterTypeE,ParameterNameE,ReturnStatement,NullImplementation) \
+		ReturnType RHI##MethodName(ParameterTypeA ParameterNameA, ParameterTypeB ParameterNameB, ParameterTypeC ParameterNameC, ParameterTypeD ParameterNameD, ParameterTypeE ParameterNameE);
+	#define DEFINE_RHIMETHOD_GLOBALFLUSH_6(ReturnType,MethodName,ParameterTypeA,ParameterNameA,ParameterTypeB,ParameterNameB,ParameterTypeC,ParameterNameC,ParameterTypeD,ParameterNameD,ParameterTypeE,ParameterNameE,ParameterTypeF,ParameterNameF,ReturnStatement,NullImplementation) \
+		ReturnType RHI##MethodName(ParameterTypeA ParameterNameA, ParameterTypeB ParameterNameB, ParameterTypeC ParameterNameC, ParameterTypeD ParameterNameD, ParameterTypeE ParameterNameE, ParameterTypeF ParameterNameF);
+	#define DEFINE_RHIMETHOD_GLOBALFLUSH_7(ReturnType,MethodName,ParameterTypeA,ParameterNameA,ParameterTypeB,ParameterNameB,ParameterTypeC,ParameterNameC,ParameterTypeD,ParameterNameD,ParameterTypeE,ParameterNameE,ParameterTypeF,ParameterNameF,ParameterTypeG,ParameterNameG,ReturnStatement,NullImplementation) \
+		ReturnType RHI##MethodName(ParameterTypeA ParameterNameA, ParameterTypeB ParameterNameB, ParameterTypeC ParameterNameC, ParameterTypeD ParameterNameD, ParameterTypeE ParameterNameE, ParameterTypeF ParameterNameF, ParameterTypeG ParameterNameG);
+	#define DEFINE_RHIMETHOD_GLOBALFLUSH_8(ReturnType,MethodName,ParameterTypeA,ParameterNameA,ParameterTypeB,ParameterNameB,ParameterTypeC,ParameterNameC,ParameterTypeD,ParameterNameD,ParameterTypeE,ParameterNameE,ParameterTypeF,ParameterNameF,ParameterTypeG,ParameterNameG,ParameterTypeH,ParameterNameH,ReturnStatement,NullImplementation) \
+		ReturnType RHI##MethodName(ParameterTypeA ParameterNameA, ParameterTypeB ParameterNameB, ParameterTypeC ParameterNameC, ParameterTypeD ParameterNameD, ParameterTypeE ParameterNameE, ParameterTypeF ParameterNameF, ParameterTypeG ParameterNameG, ParameterTypeH ParameterNameH);
+	#define DEFINE_RHIMETHOD_GLOBALFLUSH_9(ReturnType,MethodName,ParameterTypeA,ParameterNameA,ParameterTypeB,ParameterNameB,ParameterTypeC,ParameterNameC,ParameterTypeD,ParameterNameD,ParameterTypeE,ParameterNameE,ParameterTypeF,ParameterNameF,ParameterTypeG,ParameterNameG,ParameterTypeH,ParameterNameH,ParameterTypeI,ParameterNameI,ReturnStatement,NullImplementation) \
+		ReturnType RHI##MethodName(ParameterTypeA ParameterNameA, ParameterTypeB ParameterNameB, ParameterTypeC ParameterNameC, ParameterTypeD ParameterNameD, ParameterTypeE ParameterNameE, ParameterTypeF ParameterNameF, ParameterTypeG ParameterNameG, ParameterTypeH ParameterNameH, ParameterTypeI ParameterNameI);
+
+
 #else
 	#define DEFINE_RHIMETHOD_0(ReturnType,MethodName,ReturnStatement,NullImplementation) \
 		DEFINE_RHIMETHOD(ReturnType,MethodName,(),(),ReturnStatement,NullImplementation)
@@ -54,6 +120,139 @@
 		DEFINE_RHIMETHOD(ReturnType,MethodName,(ParameterTypeA ParameterNameA,ParameterTypeB ParameterNameB,ParameterTypeC ParameterNameC,ParameterTypeD ParameterNameD,ParameterTypeE ParameterNameE,ParameterTypeF ParameterNameF,ParameterTypeG ParameterNameG,ParameterTypeH ParameterNameH),(ParameterNameA,ParameterNameB,ParameterNameC,ParameterNameD,ParameterNameE,ParameterNameF,ParameterNameG,ParameterNameH),ReturnStatement,NullImplementation)
 	#define DEFINE_RHIMETHOD_9(ReturnType,MethodName,ParameterTypeA,ParameterNameA,ParameterTypeB,ParameterNameB,ParameterTypeC,ParameterNameC,ParameterTypeD,ParameterNameD,ParameterTypeE,ParameterNameE,ParameterTypeF,ParameterNameF,ParameterTypeG,ParameterNameG,ParameterTypeH,ParameterNameH,ParameterTypeI,ParameterNameI,ReturnStatement,NullImplementation) \
 		DEFINE_RHIMETHOD(ReturnType,MethodName,(ParameterTypeA ParameterNameA,ParameterTypeB ParameterNameB,ParameterTypeC ParameterNameC,ParameterTypeD ParameterNameD,ParameterTypeE ParameterNameE,ParameterTypeF ParameterNameF,ParameterTypeG ParameterNameG,ParameterTypeH ParameterNameH,ParameterTypeI ParameterNameI),(ParameterNameA,ParameterNameB,ParameterNameC,ParameterNameD,ParameterNameE,ParameterNameF,ParameterNameG,ParameterNameH,ParameterNameI),ReturnStatement,NullImplementation)
+
+	#if defined(DEFINE_RHIMETHOD_CMDLIST)
+		#define DEFINE_RHIMETHOD_CMDLIST_0(ReturnType,MethodName,ReturnStatement,NullImplementation) \
+			DEFINE_RHIMETHOD_CMDLIST(ReturnType, MethodName, (), (), ReturnStatement, NullImplementation)
+		#define DEFINE_RHIMETHOD_CMDLIST_1(ReturnType,MethodName,ParameterTypeA,ParameterNameA,ReturnStatement,NullImplementation) \
+			DEFINE_RHIMETHOD_CMDLIST(ReturnType, MethodName, (ParameterTypeA ParameterNameA), (ParameterNameA), ReturnStatement, NullImplementation)
+		#define DEFINE_RHIMETHOD_CMDLIST_2(ReturnType,MethodName,ParameterTypeA,ParameterNameA,ParameterTypeB,ParameterNameB,ReturnStatement,NullImplementation) \
+			DEFINE_RHIMETHOD_CMDLIST(ReturnType, MethodName, (ParameterTypeA ParameterNameA, ParameterTypeB ParameterNameB), (ParameterNameA, ParameterNameB), ReturnStatement, NullImplementation)
+		#define DEFINE_RHIMETHOD_CMDLIST_3(ReturnType,MethodName,ParameterTypeA,ParameterNameA,ParameterTypeB,ParameterNameB,ParameterTypeC,ParameterNameC,ReturnStatement,NullImplementation) \
+			DEFINE_RHIMETHOD_CMDLIST(ReturnType, MethodName, (ParameterTypeA ParameterNameA, ParameterTypeB ParameterNameB, ParameterTypeC ParameterNameC), (ParameterNameA, ParameterNameB, ParameterNameC), ReturnStatement, NullImplementation)
+		#define DEFINE_RHIMETHOD_CMDLIST_4(ReturnType,MethodName,ParameterTypeA,ParameterNameA,ParameterTypeB,ParameterNameB,ParameterTypeC,ParameterNameC,ParameterTypeD,ParameterNameD,ReturnStatement,NullImplementation) \
+			DEFINE_RHIMETHOD_CMDLIST(ReturnType, MethodName, (ParameterTypeA ParameterNameA, ParameterTypeB ParameterNameB, ParameterTypeC ParameterNameC, ParameterTypeD ParameterNameD), (ParameterNameA, ParameterNameB, ParameterNameC, ParameterNameD), ReturnStatement, NullImplementation)
+		#define DEFINE_RHIMETHOD_CMDLIST_5(ReturnType,MethodName,ParameterTypeA,ParameterNameA,ParameterTypeB,ParameterNameB,ParameterTypeC,ParameterNameC,ParameterTypeD,ParameterNameD,ParameterTypeE,ParameterNameE,ReturnStatement,NullImplementation) \
+			DEFINE_RHIMETHOD_CMDLIST(ReturnType, MethodName, (ParameterTypeA ParameterNameA, ParameterTypeB ParameterNameB, ParameterTypeC ParameterNameC, ParameterTypeD ParameterNameD, ParameterTypeE ParameterNameE), (ParameterNameA, ParameterNameB, ParameterNameC, ParameterNameD, ParameterNameE), ReturnStatement, NullImplementation)
+		#define DEFINE_RHIMETHOD_CMDLIST_6(ReturnType,MethodName,ParameterTypeA,ParameterNameA,ParameterTypeB,ParameterNameB,ParameterTypeC,ParameterNameC,ParameterTypeD,ParameterNameD,ParameterTypeE,ParameterNameE,ParameterTypeF,ParameterNameF,ReturnStatement,NullImplementation) \
+			DEFINE_RHIMETHOD_CMDLIST(ReturnType, MethodName, (ParameterTypeA ParameterNameA, ParameterTypeB ParameterNameB, ParameterTypeC ParameterNameC, ParameterTypeD ParameterNameD, ParameterTypeE ParameterNameE, ParameterTypeF ParameterNameF), (ParameterNameA, ParameterNameB, ParameterNameC, ParameterNameD, ParameterNameE, ParameterNameF), ReturnStatement, NullImplementation)
+		#define DEFINE_RHIMETHOD_CMDLIST_7(ReturnType,MethodName,ParameterTypeA,ParameterNameA,ParameterTypeB,ParameterNameB,ParameterTypeC,ParameterNameC,ParameterTypeD,ParameterNameD,ParameterTypeE,ParameterNameE,ParameterTypeF,ParameterNameF,ParameterTypeG,ParameterNameG,ReturnStatement,NullImplementation) \
+			DEFINE_RHIMETHOD_CMDLIST(ReturnType, MethodName, (ParameterTypeA ParameterNameA, ParameterTypeB ParameterNameB, ParameterTypeC ParameterNameC, ParameterTypeD ParameterNameD, ParameterTypeE ParameterNameE, ParameterTypeF ParameterNameF, ParameterTypeG ParameterNameG), (ParameterNameA, ParameterNameB, ParameterNameC, ParameterNameD, ParameterNameE, ParameterNameF, ParameterNameG), ReturnStatement, NullImplementation)
+		#define DEFINE_RHIMETHOD_CMDLIST_8(ReturnType,MethodName,ParameterTypeA,ParameterNameA,ParameterTypeB,ParameterNameB,ParameterTypeC,ParameterNameC,ParameterTypeD,ParameterNameD,ParameterTypeE,ParameterNameE,ParameterTypeF,ParameterNameF,ParameterTypeG,ParameterNameG,ParameterTypeH,ParameterNameH,ReturnStatement,NullImplementation) \
+			DEFINE_RHIMETHOD_CMDLIST(ReturnType, MethodName, (ParameterTypeA ParameterNameA, ParameterTypeB ParameterNameB, ParameterTypeC ParameterNameC, ParameterTypeD ParameterNameD, ParameterTypeE ParameterNameE, ParameterTypeF ParameterNameF, ParameterTypeG ParameterNameG, ParameterTypeH ParameterNameH), (ParameterNameA, ParameterNameB, ParameterNameC, ParameterNameD, ParameterNameE, ParameterNameF, ParameterNameG, ParameterNameH), ReturnStatement, NullImplementation)
+		#define DEFINE_RHIMETHOD_CMDLIST_9(ReturnType,MethodName,ParameterTypeA,ParameterNameA,ParameterTypeB,ParameterNameB,ParameterTypeC,ParameterNameC,ParameterTypeD,ParameterNameD,ParameterTypeE,ParameterNameE,ParameterTypeF,ParameterNameF,ParameterTypeG,ParameterNameG,ParameterTypeH,ParameterNameH,ParameterTypeI,ParameterNameI,ReturnStatement,NullImplementation) \
+			DEFINE_RHIMETHOD_CMDLIST(ReturnType, MethodName, (ParameterTypeA ParameterNameA, ParameterTypeB ParameterNameB, ParameterTypeC ParameterNameC, ParameterTypeD ParameterNameD, ParameterTypeE ParameterNameE, ParameterTypeF ParameterNameF, ParameterTypeG ParameterNameG, ParameterTypeH ParameterNameH, ParameterTypeI ParameterNameI), (ParameterNameA, ParameterNameB, ParameterNameC, ParameterNameD, ParameterNameE, ParameterNameF, ParameterNameG, ParameterNameH, ParameterNameI), ReturnStatement, NullImplementation)
+	#else
+		#define DEFINE_RHIMETHOD_CMDLIST_0(ReturnType,MethodName,ReturnStatement,NullImplementation) \
+			DEFINE_RHIMETHOD(ReturnType, MethodName, (), (), ReturnStatement, NullImplementation)
+		#define DEFINE_RHIMETHOD_CMDLIST_1(ReturnType,MethodName,ParameterTypeA,ParameterNameA,ReturnStatement,NullImplementation) \
+			DEFINE_RHIMETHOD(ReturnType, MethodName, (ParameterTypeA ParameterNameA), (ParameterNameA), ReturnStatement, NullImplementation)
+		#define DEFINE_RHIMETHOD_CMDLIST_2(ReturnType,MethodName,ParameterTypeA,ParameterNameA,ParameterTypeB,ParameterNameB,ReturnStatement,NullImplementation) \
+			DEFINE_RHIMETHOD(ReturnType, MethodName, (ParameterTypeA ParameterNameA, ParameterTypeB ParameterNameB), (ParameterNameA, ParameterNameB), ReturnStatement, NullImplementation)
+		#define DEFINE_RHIMETHOD_CMDLIST_3(ReturnType,MethodName,ParameterTypeA,ParameterNameA,ParameterTypeB,ParameterNameB,ParameterTypeC,ParameterNameC,ReturnStatement,NullImplementation) \
+			DEFINE_RHIMETHOD(ReturnType, MethodName, (ParameterTypeA ParameterNameA, ParameterTypeB ParameterNameB, ParameterTypeC ParameterNameC), (ParameterNameA, ParameterNameB, ParameterNameC), ReturnStatement, NullImplementation)
+		#define DEFINE_RHIMETHOD_CMDLIST_4(ReturnType,MethodName,ParameterTypeA,ParameterNameA,ParameterTypeB,ParameterNameB,ParameterTypeC,ParameterNameC,ParameterTypeD,ParameterNameD,ReturnStatement,NullImplementation) \
+			DEFINE_RHIMETHOD(ReturnType, MethodName, (ParameterTypeA ParameterNameA, ParameterTypeB ParameterNameB, ParameterTypeC ParameterNameC, ParameterTypeD ParameterNameD), (ParameterNameA, ParameterNameB, ParameterNameC, ParameterNameD), ReturnStatement, NullImplementation)
+		#define DEFINE_RHIMETHOD_CMDLIST_5(ReturnType,MethodName,ParameterTypeA,ParameterNameA,ParameterTypeB,ParameterNameB,ParameterTypeC,ParameterNameC,ParameterTypeD,ParameterNameD,ParameterTypeE,ParameterNameE,ReturnStatement,NullImplementation) \
+			DEFINE_RHIMETHOD(ReturnType, MethodName, (ParameterTypeA ParameterNameA, ParameterTypeB ParameterNameB, ParameterTypeC ParameterNameC, ParameterTypeD ParameterNameD, ParameterTypeE ParameterNameE), (ParameterNameA, ParameterNameB, ParameterNameC, ParameterNameD, ParameterNameE), ReturnStatement, NullImplementation)
+		#define DEFINE_RHIMETHOD_CMDLIST_6(ReturnType,MethodName,ParameterTypeA,ParameterNameA,ParameterTypeB,ParameterNameB,ParameterTypeC,ParameterNameC,ParameterTypeD,ParameterNameD,ParameterTypeE,ParameterNameE,ParameterTypeF,ParameterNameF,ReturnStatement,NullImplementation) \
+			DEFINE_RHIMETHOD(ReturnType, MethodName, (ParameterTypeA ParameterNameA, ParameterTypeB ParameterNameB, ParameterTypeC ParameterNameC, ParameterTypeD ParameterNameD, ParameterTypeE ParameterNameE, ParameterTypeF ParameterNameF), (ParameterNameA, ParameterNameB, ParameterNameC, ParameterNameD, ParameterNameE, ParameterNameF), ReturnStatement, NullImplementation)
+		#define DEFINE_RHIMETHOD_CMDLIST_7(ReturnType,MethodName,ParameterTypeA,ParameterNameA,ParameterTypeB,ParameterNameB,ParameterTypeC,ParameterNameC,ParameterTypeD,ParameterNameD,ParameterTypeE,ParameterNameE,ParameterTypeF,ParameterNameF,ParameterTypeG,ParameterNameG,ReturnStatement,NullImplementation) \
+			DEFINE_RHIMETHOD(ReturnType, MethodName, (ParameterTypeA ParameterNameA, ParameterTypeB ParameterNameB, ParameterTypeC ParameterNameC, ParameterTypeD ParameterNameD, ParameterTypeE ParameterNameE, ParameterTypeF ParameterNameF, ParameterTypeG ParameterNameG), (ParameterNameA, ParameterNameB, ParameterNameC, ParameterNameD, ParameterNameE, ParameterNameF, ParameterNameG), ReturnStatement, NullImplementation)
+		#define DEFINE_RHIMETHOD_CMDLIST_8(ReturnType,MethodName,ParameterTypeA,ParameterNameA,ParameterTypeB,ParameterNameB,ParameterTypeC,ParameterNameC,ParameterTypeD,ParameterNameD,ParameterTypeE,ParameterNameE,ParameterTypeF,ParameterNameF,ParameterTypeG,ParameterNameG,ParameterTypeH,ParameterNameH,ReturnStatement,NullImplementation) \
+			DEFINE_RHIMETHOD(ReturnType, MethodName, (ParameterTypeA ParameterNameA, ParameterTypeB ParameterNameB, ParameterTypeC ParameterNameC, ParameterTypeD ParameterNameD, ParameterTypeE ParameterNameE, ParameterTypeF ParameterNameF, ParameterTypeG ParameterNameG, ParameterTypeH ParameterNameH), (ParameterNameA, ParameterNameB, ParameterNameC, ParameterNameD, ParameterNameE, ParameterNameF, ParameterNameG, ParameterNameH), ReturnStatement, NullImplementation)
+		#define DEFINE_RHIMETHOD_CMDLIST_9(ReturnType,MethodName,ParameterTypeA,ParameterNameA,ParameterTypeB,ParameterNameB,ParameterTypeC,ParameterNameC,ParameterTypeD,ParameterNameD,ParameterTypeE,ParameterNameE,ParameterTypeF,ParameterNameF,ParameterTypeG,ParameterNameG,ParameterTypeH,ParameterNameH,ParameterTypeI,ParameterNameI,ReturnStatement,NullImplementation) \
+			DEFINE_RHIMETHOD(ReturnType, MethodName, (ParameterTypeA ParameterNameA, ParameterTypeB ParameterNameB, ParameterTypeC ParameterNameC, ParameterTypeD ParameterNameD, ParameterTypeE ParameterNameE, ParameterTypeF ParameterNameF, ParameterTypeG ParameterNameG, ParameterTypeH ParameterNameH, ParameterTypeI ParameterNameI), (ParameterNameA, ParameterNameB, ParameterNameC, ParameterNameD, ParameterNameE, ParameterNameF, ParameterNameG, ParameterNameH, ParameterNameI), ReturnStatement, NullImplementation)
+	#endif
+
+	#if defined(DEFINE_RHIMETHOD_GLOBAL)
+		#define DEFINE_RHIMETHOD_GLOBAL_0(ReturnType,MethodName,ReturnStatement,NullImplementation) \
+			DEFINE_RHIMETHOD_GLOBAL(ReturnType, MethodName, (), (), ReturnStatement, NullImplementation)
+		#define DEFINE_RHIMETHOD_GLOBAL_1(ReturnType,MethodName,ParameterTypeA,ParameterNameA,ReturnStatement,NullImplementation) \
+			DEFINE_RHIMETHOD_GLOBAL(ReturnType, MethodName, (ParameterTypeA ParameterNameA), (ParameterNameA), ReturnStatement, NullImplementation)
+		#define DEFINE_RHIMETHOD_GLOBAL_2(ReturnType,MethodName,ParameterTypeA,ParameterNameA,ParameterTypeB,ParameterNameB,ReturnStatement,NullImplementation) \
+			DEFINE_RHIMETHOD_GLOBAL(ReturnType, MethodName, (ParameterTypeA ParameterNameA, ParameterTypeB ParameterNameB), (ParameterNameA, ParameterNameB), ReturnStatement, NullImplementation)
+		#define DEFINE_RHIMETHOD_GLOBAL_3(ReturnType,MethodName,ParameterTypeA,ParameterNameA,ParameterTypeB,ParameterNameB,ParameterTypeC,ParameterNameC,ReturnStatement,NullImplementation) \
+			DEFINE_RHIMETHOD_GLOBAL(ReturnType, MethodName, (ParameterTypeA ParameterNameA, ParameterTypeB ParameterNameB, ParameterTypeC ParameterNameC), (ParameterNameA, ParameterNameB, ParameterNameC), ReturnStatement, NullImplementation)
+		#define DEFINE_RHIMETHOD_GLOBAL_4(ReturnType,MethodName,ParameterTypeA,ParameterNameA,ParameterTypeB,ParameterNameB,ParameterTypeC,ParameterNameC,ParameterTypeD,ParameterNameD,ReturnStatement,NullImplementation) \
+			DEFINE_RHIMETHOD_GLOBAL(ReturnType, MethodName, (ParameterTypeA ParameterNameA, ParameterTypeB ParameterNameB, ParameterTypeC ParameterNameC, ParameterTypeD ParameterNameD), (ParameterNameA, ParameterNameB, ParameterNameC, ParameterNameD), ReturnStatement, NullImplementation)
+		#define DEFINE_RHIMETHOD_GLOBAL_5(ReturnType,MethodName,ParameterTypeA,ParameterNameA,ParameterTypeB,ParameterNameB,ParameterTypeC,ParameterNameC,ParameterTypeD,ParameterNameD,ParameterTypeE,ParameterNameE,ReturnStatement,NullImplementation) \
+			DEFINE_RHIMETHOD_GLOBAL(ReturnType, MethodName, (ParameterTypeA ParameterNameA, ParameterTypeB ParameterNameB, ParameterTypeC ParameterNameC, ParameterTypeD ParameterNameD, ParameterTypeE ParameterNameE), (ParameterNameA, ParameterNameB, ParameterNameC, ParameterNameD, ParameterNameE), ReturnStatement, NullImplementation)
+		#define DEFINE_RHIMETHOD_GLOBAL_6(ReturnType,MethodName,ParameterTypeA,ParameterNameA,ParameterTypeB,ParameterNameB,ParameterTypeC,ParameterNameC,ParameterTypeD,ParameterNameD,ParameterTypeE,ParameterNameE,ParameterTypeF,ParameterNameF,ReturnStatement,NullImplementation) \
+			DEFINE_RHIMETHOD_GLOBAL(ReturnType, MethodName, (ParameterTypeA ParameterNameA, ParameterTypeB ParameterNameB, ParameterTypeC ParameterNameC, ParameterTypeD ParameterNameD, ParameterTypeE ParameterNameE, ParameterTypeF ParameterNameF), (ParameterNameA, ParameterNameB, ParameterNameC, ParameterNameD, ParameterNameE, ParameterNameF), ReturnStatement, NullImplementation)
+		#define DEFINE_RHIMETHOD_GLOBAL_7(ReturnType,MethodName,ParameterTypeA,ParameterNameA,ParameterTypeB,ParameterNameB,ParameterTypeC,ParameterNameC,ParameterTypeD,ParameterNameD,ParameterTypeE,ParameterNameE,ParameterTypeF,ParameterNameF,ParameterTypeG,ParameterNameG,ReturnStatement,NullImplementation) \
+			DEFINE_RHIMETHOD_GLOBAL(ReturnType, MethodName, (ParameterTypeA ParameterNameA, ParameterTypeB ParameterNameB, ParameterTypeC ParameterNameC, ParameterTypeD ParameterNameD, ParameterTypeE ParameterNameE, ParameterTypeF ParameterNameF, ParameterTypeG ParameterNameG), (ParameterNameA, ParameterNameB, ParameterNameC, ParameterNameD, ParameterNameE, ParameterNameF, ParameterNameG), ReturnStatement, NullImplementation)
+		#define DEFINE_RHIMETHOD_GLOBAL_8(ReturnType,MethodName,ParameterTypeA,ParameterNameA,ParameterTypeB,ParameterNameB,ParameterTypeC,ParameterNameC,ParameterTypeD,ParameterNameD,ParameterTypeE,ParameterNameE,ParameterTypeF,ParameterNameF,ParameterTypeG,ParameterNameG,ParameterTypeH,ParameterNameH,ReturnStatement,NullImplementation) \
+			DEFINE_RHIMETHOD_GLOBAL(ReturnType, MethodName, (ParameterTypeA ParameterNameA, ParameterTypeB ParameterNameB, ParameterTypeC ParameterNameC, ParameterTypeD ParameterNameD, ParameterTypeE ParameterNameE, ParameterTypeF ParameterNameF, ParameterTypeG ParameterNameG, ParameterTypeH ParameterNameH), (ParameterNameA, ParameterNameB, ParameterNameC, ParameterNameD, ParameterNameE, ParameterNameF, ParameterNameG, ParameterNameH), ReturnStatement, NullImplementation)
+		#define DEFINE_RHIMETHOD_GLOBAL_9(ReturnType,MethodName,ParameterTypeA,ParameterNameA,ParameterTypeB,ParameterNameB,ParameterTypeC,ParameterNameC,ParameterTypeD,ParameterNameD,ParameterTypeE,ParameterNameE,ParameterTypeF,ParameterNameF,ParameterTypeG,ParameterNameG,ParameterTypeH,ParameterNameH,ParameterTypeI,ParameterNameI,ReturnStatement,NullImplementation) \
+			DEFINE_RHIMETHOD_GLOBAL(ReturnType, MethodName, (ParameterTypeA ParameterNameA, ParameterTypeB ParameterNameB, ParameterTypeC ParameterNameC, ParameterTypeD ParameterNameD, ParameterTypeE ParameterNameE, ParameterTypeF ParameterNameF, ParameterTypeG ParameterNameG, ParameterTypeH ParameterNameH, ParameterTypeI ParameterNameI), (ParameterNameA, ParameterNameB, ParameterNameC, ParameterNameD, ParameterNameE, ParameterNameF, ParameterNameG, ParameterNameH, ParameterNameI), ReturnStatement, NullImplementation)
+	#else
+		#define DEFINE_RHIMETHOD_GLOBAL_0(ReturnType,MethodName,ReturnStatement,NullImplementation) \
+			DEFINE_RHIMETHOD(ReturnType, MethodName, (), (), ReturnStatement, NullImplementation)
+		#define DEFINE_RHIMETHOD_GLOBAL_1(ReturnType,MethodName,ParameterTypeA,ParameterNameA,ReturnStatement,NullImplementation) \
+			DEFINE_RHIMETHOD(ReturnType, MethodName, (ParameterTypeA ParameterNameA), (ParameterNameA), ReturnStatement, NullImplementation)
+		#define DEFINE_RHIMETHOD_GLOBAL_2(ReturnType,MethodName,ParameterTypeA,ParameterNameA,ParameterTypeB,ParameterNameB,ReturnStatement,NullImplementation) \
+			DEFINE_RHIMETHOD(ReturnType, MethodName, (ParameterTypeA ParameterNameA, ParameterTypeB ParameterNameB), (ParameterNameA, ParameterNameB), ReturnStatement, NullImplementation)
+		#define DEFINE_RHIMETHOD_GLOBAL_3(ReturnType,MethodName,ParameterTypeA,ParameterNameA,ParameterTypeB,ParameterNameB,ParameterTypeC,ParameterNameC,ReturnStatement,NullImplementation) \
+			DEFINE_RHIMETHOD(ReturnType, MethodName, (ParameterTypeA ParameterNameA, ParameterTypeB ParameterNameB, ParameterTypeC ParameterNameC), (ParameterNameA, ParameterNameB, ParameterNameC), ReturnStatement, NullImplementation)
+		#define DEFINE_RHIMETHOD_GLOBAL_4(ReturnType,MethodName,ParameterTypeA,ParameterNameA,ParameterTypeB,ParameterNameB,ParameterTypeC,ParameterNameC,ParameterTypeD,ParameterNameD,ReturnStatement,NullImplementation) \
+			DEFINE_RHIMETHOD(ReturnType, MethodName, (ParameterTypeA ParameterNameA, ParameterTypeB ParameterNameB, ParameterTypeC ParameterNameC, ParameterTypeD ParameterNameD), (ParameterNameA, ParameterNameB, ParameterNameC, ParameterNameD), ReturnStatement, NullImplementation)
+		#define DEFINE_RHIMETHOD_GLOBAL_5(ReturnType,MethodName,ParameterTypeA,ParameterNameA,ParameterTypeB,ParameterNameB,ParameterTypeC,ParameterNameC,ParameterTypeD,ParameterNameD,ParameterTypeE,ParameterNameE,ReturnStatement,NullImplementation) \
+			DEFINE_RHIMETHOD(ReturnType, MethodName, (ParameterTypeA ParameterNameA, ParameterTypeB ParameterNameB, ParameterTypeC ParameterNameC, ParameterTypeD ParameterNameD, ParameterTypeE ParameterNameE), (ParameterNameA, ParameterNameB, ParameterNameC, ParameterNameD, ParameterNameE), ReturnStatement, NullImplementation)
+		#define DEFINE_RHIMETHOD_GLOBAL_6(ReturnType,MethodName,ParameterTypeA,ParameterNameA,ParameterTypeB,ParameterNameB,ParameterTypeC,ParameterNameC,ParameterTypeD,ParameterNameD,ParameterTypeE,ParameterNameE,ParameterTypeF,ParameterNameF,ReturnStatement,NullImplementation) \
+			DEFINE_RHIMETHOD(ReturnType, MethodName, (ParameterTypeA ParameterNameA, ParameterTypeB ParameterNameB, ParameterTypeC ParameterNameC, ParameterTypeD ParameterNameD, ParameterTypeE ParameterNameE, ParameterTypeF ParameterNameF), (ParameterNameA, ParameterNameB, ParameterNameC, ParameterNameD, ParameterNameE, ParameterNameF), ReturnStatement, NullImplementation)
+		#define DEFINE_RHIMETHOD_GLOBAL_7(ReturnType,MethodName,ParameterTypeA,ParameterNameA,ParameterTypeB,ParameterNameB,ParameterTypeC,ParameterNameC,ParameterTypeD,ParameterNameD,ParameterTypeE,ParameterNameE,ParameterTypeF,ParameterNameF,ParameterTypeG,ParameterNameG,ReturnStatement,NullImplementation) \
+			DEFINE_RHIMETHOD(ReturnType, MethodName, (ParameterTypeA ParameterNameA, ParameterTypeB ParameterNameB, ParameterTypeC ParameterNameC, ParameterTypeD ParameterNameD, ParameterTypeE ParameterNameE, ParameterTypeF ParameterNameF, ParameterTypeG ParameterNameG), (ParameterNameA, ParameterNameB, ParameterNameC, ParameterNameD, ParameterNameE, ParameterNameF, ParameterNameG), ReturnStatement, NullImplementation)
+		#define DEFINE_RHIMETHOD_GLOBAL_8(ReturnType,MethodName,ParameterTypeA,ParameterNameA,ParameterTypeB,ParameterNameB,ParameterTypeC,ParameterNameC,ParameterTypeD,ParameterNameD,ParameterTypeE,ParameterNameE,ParameterTypeF,ParameterNameF,ParameterTypeG,ParameterNameG,ParameterTypeH,ParameterNameH,ReturnStatement,NullImplementation) \
+			DEFINE_RHIMETHOD(ReturnType, MethodName, (ParameterTypeA ParameterNameA, ParameterTypeB ParameterNameB, ParameterTypeC ParameterNameC, ParameterTypeD ParameterNameD, ParameterTypeE ParameterNameE, ParameterTypeF ParameterNameF, ParameterTypeG ParameterNameG, ParameterTypeH ParameterNameH), (ParameterNameA, ParameterNameB, ParameterNameC, ParameterNameD, ParameterNameE, ParameterNameF, ParameterNameG, ParameterNameH), ReturnStatement, NullImplementation)
+		#define DEFINE_RHIMETHOD_GLOBAL_9(ReturnType,MethodName,ParameterTypeA,ParameterNameA,ParameterTypeB,ParameterNameB,ParameterTypeC,ParameterNameC,ParameterTypeD,ParameterNameD,ParameterTypeE,ParameterNameE,ParameterTypeF,ParameterNameF,ParameterTypeG,ParameterNameG,ParameterTypeH,ParameterNameH,ParameterTypeI,ParameterNameI,ReturnStatement,NullImplementation) \
+			DEFINE_RHIMETHOD(ReturnType, MethodName, (ParameterTypeA ParameterNameA, ParameterTypeB ParameterNameB, ParameterTypeC ParameterNameC, ParameterTypeD ParameterNameD, ParameterTypeE ParameterNameE, ParameterTypeF ParameterNameF, ParameterTypeG ParameterNameG, ParameterTypeH ParameterNameH, ParameterTypeI ParameterNameI), (ParameterNameA, ParameterNameB, ParameterNameC, ParameterNameD, ParameterNameE, ParameterNameF, ParameterNameG, ParameterNameH, ParameterNameI), ReturnStatement, NullImplementation)
+	#endif
+
+	#if defined(DEFINE_RHIMETHOD_GLOBALFLUSH)
+		#define DEFINE_RHIMETHOD_GLOBALFLUSH_0(ReturnType,MethodName,ReturnStatement,NullImplementation) \
+			DEFINE_RHIMETHOD_GLOBALFLUSH(ReturnType, MethodName, (), (), ReturnStatement, NullImplementation)
+		#define DEFINE_RHIMETHOD_GLOBALFLUSH_1(ReturnType,MethodName,ParameterTypeA,ParameterNameA,ReturnStatement,NullImplementation) \
+			DEFINE_RHIMETHOD_GLOBALFLUSH(ReturnType, MethodName, (ParameterTypeA ParameterNameA), (ParameterNameA), ReturnStatement, NullImplementation)
+		#define DEFINE_RHIMETHOD_GLOBALFLUSH_2(ReturnType,MethodName,ParameterTypeA,ParameterNameA,ParameterTypeB,ParameterNameB,ReturnStatement,NullImplementation) \
+			DEFINE_RHIMETHOD_GLOBALFLUSH(ReturnType, MethodName, (ParameterTypeA ParameterNameA, ParameterTypeB ParameterNameB), (ParameterNameA, ParameterNameB), ReturnStatement, NullImplementation)
+		#define DEFINE_RHIMETHOD_GLOBALFLUSH_3(ReturnType,MethodName,ParameterTypeA,ParameterNameA,ParameterTypeB,ParameterNameB,ParameterTypeC,ParameterNameC,ReturnStatement,NullImplementation) \
+			DEFINE_RHIMETHOD_GLOBALFLUSH(ReturnType, MethodName, (ParameterTypeA ParameterNameA, ParameterTypeB ParameterNameB, ParameterTypeC ParameterNameC), (ParameterNameA, ParameterNameB, ParameterNameC), ReturnStatement, NullImplementation)
+		#define DEFINE_RHIMETHOD_GLOBALFLUSH_4(ReturnType,MethodName,ParameterTypeA,ParameterNameA,ParameterTypeB,ParameterNameB,ParameterTypeC,ParameterNameC,ParameterTypeD,ParameterNameD,ReturnStatement,NullImplementation) \
+			DEFINE_RHIMETHOD_GLOBALFLUSH(ReturnType, MethodName, (ParameterTypeA ParameterNameA, ParameterTypeB ParameterNameB, ParameterTypeC ParameterNameC, ParameterTypeD ParameterNameD), (ParameterNameA, ParameterNameB, ParameterNameC, ParameterNameD), ReturnStatement, NullImplementation)
+		#define DEFINE_RHIMETHOD_GLOBALFLUSH_5(ReturnType,MethodName,ParameterTypeA,ParameterNameA,ParameterTypeB,ParameterNameB,ParameterTypeC,ParameterNameC,ParameterTypeD,ParameterNameD,ParameterTypeE,ParameterNameE,ReturnStatement,NullImplementation) \
+			DEFINE_RHIMETHOD_GLOBALFLUSH(ReturnType, MethodName, (ParameterTypeA ParameterNameA, ParameterTypeB ParameterNameB, ParameterTypeC ParameterNameC, ParameterTypeD ParameterNameD, ParameterTypeE ParameterNameE), (ParameterNameA, ParameterNameB, ParameterNameC, ParameterNameD, ParameterNameE), ReturnStatement, NullImplementation)
+		#define DEFINE_RHIMETHOD_GLOBALFLUSH_6(ReturnType,MethodName,ParameterTypeA,ParameterNameA,ParameterTypeB,ParameterNameB,ParameterTypeC,ParameterNameC,ParameterTypeD,ParameterNameD,ParameterTypeE,ParameterNameE,ParameterTypeF,ParameterNameF,ReturnStatement,NullImplementation) \
+			DEFINE_RHIMETHOD_GLOBALFLUSH(ReturnType, MethodName, (ParameterTypeA ParameterNameA, ParameterTypeB ParameterNameB, ParameterTypeC ParameterNameC, ParameterTypeD ParameterNameD, ParameterTypeE ParameterNameE, ParameterTypeF ParameterNameF), (ParameterNameA, ParameterNameB, ParameterNameC, ParameterNameD, ParameterNameE, ParameterNameF), ReturnStatement, NullImplementation)
+		#define DEFINE_RHIMETHOD_GLOBALFLUSH_7(ReturnType,MethodName,ParameterTypeA,ParameterNameA,ParameterTypeB,ParameterNameB,ParameterTypeC,ParameterNameC,ParameterTypeD,ParameterNameD,ParameterTypeE,ParameterNameE,ParameterTypeF,ParameterNameF,ParameterTypeG,ParameterNameG,ReturnStatement,NullImplementation) \
+			DEFINE_RHIMETHOD_GLOBALFLUSH(ReturnType, MethodName, (ParameterTypeA ParameterNameA, ParameterTypeB ParameterNameB, ParameterTypeC ParameterNameC, ParameterTypeD ParameterNameD, ParameterTypeE ParameterNameE, ParameterTypeF ParameterNameF, ParameterTypeG ParameterNameG), (ParameterNameA, ParameterNameB, ParameterNameC, ParameterNameD, ParameterNameE, ParameterNameF, ParameterNameG), ReturnStatement, NullImplementation)
+		#define DEFINE_RHIMETHOD_GLOBALFLUSH_8(ReturnType,MethodName,ParameterTypeA,ParameterNameA,ParameterTypeB,ParameterNameB,ParameterTypeC,ParameterNameC,ParameterTypeD,ParameterNameD,ParameterTypeE,ParameterNameE,ParameterTypeF,ParameterNameF,ParameterTypeG,ParameterNameG,ParameterTypeH,ParameterNameH,ReturnStatement,NullImplementation) \
+			DEFINE_RHIMETHOD_GLOBALFLUSH(ReturnType, MethodName, (ParameterTypeA ParameterNameA, ParameterTypeB ParameterNameB, ParameterTypeC ParameterNameC, ParameterTypeD ParameterNameD, ParameterTypeE ParameterNameE, ParameterTypeF ParameterNameF, ParameterTypeG ParameterNameG, ParameterTypeH ParameterNameH), (ParameterNameA, ParameterNameB, ParameterNameC, ParameterNameD, ParameterNameE, ParameterNameF, ParameterNameG, ParameterNameH), ReturnStatement, NullImplementation)
+		#define DEFINE_RHIMETHOD_GLOBALFLUSH_9(ReturnType,MethodName,ParameterTypeA,ParameterNameA,ParameterTypeB,ParameterNameB,ParameterTypeC,ParameterNameC,ParameterTypeD,ParameterNameD,ParameterTypeE,ParameterNameE,ParameterTypeF,ParameterNameF,ParameterTypeG,ParameterNameG,ParameterTypeH,ParameterNameH,ParameterTypeI,ParameterNameI,ReturnStatement,NullImplementation) \
+			DEFINE_RHIMETHOD_GLOBALFLUSH(ReturnType, MethodName, (ParameterTypeA ParameterNameA, ParameterTypeB ParameterNameB, ParameterTypeC ParameterNameC, ParameterTypeD ParameterNameD, ParameterTypeE ParameterNameE, ParameterTypeF ParameterNameF, ParameterTypeG ParameterNameG, ParameterTypeH ParameterNameH, ParameterTypeI ParameterNameI), (ParameterNameA, ParameterNameB, ParameterNameC, ParameterNameD, ParameterNameE, ParameterNameF, ParameterNameG, ParameterNameH, ParameterNameI), ReturnStatement, NullImplementation)
+	#else
+		#define DEFINE_RHIMETHOD_GLOBALFLUSH_0(ReturnType,MethodName,ReturnStatement,NullImplementation) \
+			DEFINE_RHIMETHOD(ReturnType, MethodName, (), (), ReturnStatement, NullImplementation)
+		#define DEFINE_RHIMETHOD_GLOBALFLUSH_1(ReturnType,MethodName,ParameterTypeA,ParameterNameA,ReturnStatement,NullImplementation) \
+			DEFINE_RHIMETHOD(ReturnType, MethodName, (ParameterTypeA ParameterNameA), (ParameterNameA), ReturnStatement, NullImplementation)
+		#define DEFINE_RHIMETHOD_GLOBALFLUSH_2(ReturnType,MethodName,ParameterTypeA,ParameterNameA,ParameterTypeB,ParameterNameB,ReturnStatement,NullImplementation) \
+			DEFINE_RHIMETHOD(ReturnType, MethodName, (ParameterTypeA ParameterNameA, ParameterTypeB ParameterNameB), (ParameterNameA, ParameterNameB), ReturnStatement, NullImplementation)
+		#define DEFINE_RHIMETHOD_GLOBALFLUSH_3(ReturnType,MethodName,ParameterTypeA,ParameterNameA,ParameterTypeB,ParameterNameB,ParameterTypeC,ParameterNameC,ReturnStatement,NullImplementation) \
+			DEFINE_RHIMETHOD(ReturnType, MethodName, (ParameterTypeA ParameterNameA, ParameterTypeB ParameterNameB, ParameterTypeC ParameterNameC), (ParameterNameA, ParameterNameB, ParameterNameC), ReturnStatement, NullImplementation)
+		#define DEFINE_RHIMETHOD_GLOBALFLUSH_4(ReturnType,MethodName,ParameterTypeA,ParameterNameA,ParameterTypeB,ParameterNameB,ParameterTypeC,ParameterNameC,ParameterTypeD,ParameterNameD,ReturnStatement,NullImplementation) \
+			DEFINE_RHIMETHOD(ReturnType, MethodName, (ParameterTypeA ParameterNameA, ParameterTypeB ParameterNameB, ParameterTypeC ParameterNameC, ParameterTypeD ParameterNameD), (ParameterNameA, ParameterNameB, ParameterNameC, ParameterNameD), ReturnStatement, NullImplementation)
+		#define DEFINE_RHIMETHOD_GLOBALFLUSH_5(ReturnType,MethodName,ParameterTypeA,ParameterNameA,ParameterTypeB,ParameterNameB,ParameterTypeC,ParameterNameC,ParameterTypeD,ParameterNameD,ParameterTypeE,ParameterNameE,ReturnStatement,NullImplementation) \
+			DEFINE_RHIMETHOD(ReturnType, MethodName, (ParameterTypeA ParameterNameA, ParameterTypeB ParameterNameB, ParameterTypeC ParameterNameC, ParameterTypeD ParameterNameD, ParameterTypeE ParameterNameE), (ParameterNameA, ParameterNameB, ParameterNameC, ParameterNameD, ParameterNameE), ReturnStatement, NullImplementation)
+		#define DEFINE_RHIMETHOD_GLOBALFLUSH_6(ReturnType,MethodName,ParameterTypeA,ParameterNameA,ParameterTypeB,ParameterNameB,ParameterTypeC,ParameterNameC,ParameterTypeD,ParameterNameD,ParameterTypeE,ParameterNameE,ParameterTypeF,ParameterNameF,ReturnStatement,NullImplementation) \
+			DEFINE_RHIMETHOD(ReturnType, MethodName, (ParameterTypeA ParameterNameA, ParameterTypeB ParameterNameB, ParameterTypeC ParameterNameC, ParameterTypeD ParameterNameD, ParameterTypeE ParameterNameE, ParameterTypeF ParameterNameF), (ParameterNameA, ParameterNameB, ParameterNameC, ParameterNameD, ParameterNameE, ParameterNameF), ReturnStatement, NullImplementation)
+		#define DEFINE_RHIMETHOD_GLOBALFLUSH_7(ReturnType,MethodName,ParameterTypeA,ParameterNameA,ParameterTypeB,ParameterNameB,ParameterTypeC,ParameterNameC,ParameterTypeD,ParameterNameD,ParameterTypeE,ParameterNameE,ParameterTypeF,ParameterNameF,ParameterTypeG,ParameterNameG,ReturnStatement,NullImplementation) \
+			DEFINE_RHIMETHOD(ReturnType, MethodName, (ParameterTypeA ParameterNameA, ParameterTypeB ParameterNameB, ParameterTypeC ParameterNameC, ParameterTypeD ParameterNameD, ParameterTypeE ParameterNameE, ParameterTypeF ParameterNameF, ParameterTypeG ParameterNameG), (ParameterNameA, ParameterNameB, ParameterNameC, ParameterNameD, ParameterNameE, ParameterNameF, ParameterNameG), ReturnStatement, NullImplementation)
+		#define DEFINE_RHIMETHOD_GLOBALFLUSH_8(ReturnType,MethodName,ParameterTypeA,ParameterNameA,ParameterTypeB,ParameterNameB,ParameterTypeC,ParameterNameC,ParameterTypeD,ParameterNameD,ParameterTypeE,ParameterNameE,ParameterTypeF,ParameterNameF,ParameterTypeG,ParameterNameG,ParameterTypeH,ParameterNameH,ReturnStatement,NullImplementation) \
+			DEFINE_RHIMETHOD(ReturnType, MethodName, (ParameterTypeA ParameterNameA, ParameterTypeB ParameterNameB, ParameterTypeC ParameterNameC, ParameterTypeD ParameterNameD, ParameterTypeE ParameterNameE, ParameterTypeF ParameterNameF, ParameterTypeG ParameterNameG, ParameterTypeH ParameterNameH), (ParameterNameA, ParameterNameB, ParameterNameC, ParameterNameD, ParameterNameE, ParameterNameF, ParameterNameG, ParameterNameH), ReturnStatement, NullImplementation)
+		#define DEFINE_RHIMETHOD_GLOBALFLUSH_9(ReturnType,MethodName,ParameterTypeA,ParameterNameA,ParameterTypeB,ParameterNameB,ParameterTypeC,ParameterNameC,ParameterTypeD,ParameterNameD,ParameterTypeE,ParameterNameE,ParameterTypeF,ParameterNameF,ParameterTypeG,ParameterNameG,ParameterTypeH,ParameterNameH,ParameterTypeI,ParameterNameI,ReturnStatement,NullImplementation) \
+			DEFINE_RHIMETHOD(ReturnType, MethodName, (ParameterTypeA ParameterNameA, ParameterTypeB ParameterNameB, ParameterTypeC ParameterNameC, ParameterTypeD ParameterNameD, ParameterTypeE ParameterNameE, ParameterTypeF ParameterNameF, ParameterTypeG ParameterNameG, ParameterTypeH ParameterNameH, ParameterTypeI ParameterNameI), (ParameterNameA, ParameterNameB, ParameterNameC, ParameterNameD, ParameterNameE, ParameterNameF, ParameterNameG, ParameterNameH, ParameterNameI), ReturnStatement, NullImplementation)
+	#endif
+
 #endif
 
 //
@@ -61,14 +260,14 @@
 //
 
 DEFINE_RHIMETHOD_2(
-	void,RHIGpuTimeBegin,
+	void,GpuTimeBegin,
 	uint32,Hash,
 	bool,bCompute,
 	return,return;
 	);
 
 DEFINE_RHIMETHOD_2(
-	void,RHIGpuTimeEnd,
+	void,GpuTimeEnd,
 	uint32,Hash,
 	bool,bCompute,
 	return,return;
@@ -79,66 +278,66 @@ DEFINE_RHIMETHOD_2(
 // RHI resource management functions.
 //
 
-DEFINE_RHIMETHOD_1(
-	FSamplerStateRHIRef,RHICreateSamplerState,
+DEFINE_RHIMETHOD_GLOBAL_1(
+	FSamplerStateRHIRef,CreateSamplerState,
 	const FSamplerStateInitializerRHI&,Initializer,
 	return,return new FRHISamplerState();
 	);
-DEFINE_RHIMETHOD_1(
-	FRasterizerStateRHIRef,RHICreateRasterizerState,
+DEFINE_RHIMETHOD_GLOBAL_1(
+	FRasterizerStateRHIRef,CreateRasterizerState,
 	const FRasterizerStateInitializerRHI&,Initializer,
 	return,return new FRHIRasterizerState();
 	);
-DEFINE_RHIMETHOD_1(
-	FDepthStencilStateRHIRef,RHICreateDepthStencilState,
+DEFINE_RHIMETHOD_GLOBAL_1(
+	FDepthStencilStateRHIRef,CreateDepthStencilState,
 	const FDepthStencilStateInitializerRHI&,Initializer,
 	return,return new FRHIDepthStencilState();
 	);
-DEFINE_RHIMETHOD_1(
-	FBlendStateRHIRef,RHICreateBlendState,
+DEFINE_RHIMETHOD_GLOBAL_1(
+	FBlendStateRHIRef,CreateBlendState,
 	const FBlendStateInitializerRHI&,Initializer,
 	return,return new FRHIBlendState();
 	);
 
-DEFINE_RHIMETHOD_1(
-	FVertexDeclarationRHIRef,RHICreateVertexDeclaration,
+DEFINE_RHIMETHOD_GLOBAL_1(
+	FVertexDeclarationRHIRef,CreateVertexDeclaration,
 	const FVertexDeclarationElementList&,Elements,
 	return,return new FRHIVertexDeclaration();
 	);
 
-DEFINE_RHIMETHOD_1(
-	FPixelShaderRHIRef,RHICreatePixelShader,
+DEFINE_RHIMETHOD_GLOBAL_1(
+	FPixelShaderRHIRef,CreatePixelShader,
 	const TArray<uint8>&,Code,
 	return,return new FRHIPixelShader();
 	);
 
-DEFINE_RHIMETHOD_1(
-	FVertexShaderRHIRef,RHICreateVertexShader,
+DEFINE_RHIMETHOD_GLOBAL_1(
+	FVertexShaderRHIRef,CreateVertexShader,
 	const TArray<uint8>&,Code,
 	return,return new FRHIVertexShader();
 	);
 
-DEFINE_RHIMETHOD_1(
-	FHullShaderRHIRef,RHICreateHullShader,
+DEFINE_RHIMETHOD_GLOBAL_1(
+	FHullShaderRHIRef,CreateHullShader,
 	const TArray<uint8>&,Code,
 	return,return new FRHIHullShader();
 	);
 
-DEFINE_RHIMETHOD_1(
-	FDomainShaderRHIRef,RHICreateDomainShader,
+DEFINE_RHIMETHOD_GLOBAL_1(
+	FDomainShaderRHIRef,CreateDomainShader,
 	const TArray<uint8>&,Code,
 	return,return new FRHIDomainShader();
 	);
 	
-DEFINE_RHIMETHOD_1(
-	FGeometryShaderRHIRef,RHICreateGeometryShader,
+DEFINE_RHIMETHOD_GLOBAL_1(
+	FGeometryShaderRHIRef,CreateGeometryShader,
 	const TArray<uint8>&,Code,
 	return,return new FRHIGeometryShader();
 	);
 
 /** Creates a geometry shader with stream output ability, defined by ElementList. */
-DEFINE_RHIMETHOD_5(
-	FGeometryShaderRHIRef,RHICreateGeometryShaderWithStreamOutput,
+DEFINE_RHIMETHOD_GLOBAL_5(
+	FGeometryShaderRHIRef,CreateGeometryShaderWithStreamOutput,
 	const TArray<uint8>&,Code,
 	const FStreamOutElementList&,ElementList,
 	uint32,NumStrides,
@@ -147,8 +346,8 @@ DEFINE_RHIMETHOD_5(
 	return,return new FRHIGeometryShader();
 	);
 
-DEFINE_RHIMETHOD_1(
-	FComputeShaderRHIRef,RHICreateComputeShader,
+DEFINE_RHIMETHOD_GLOBAL_1(
+	FComputeShaderRHIRef,CreateComputeShader,
 	const TArray<uint8>&,Code,
 	return,return new FRHIComputeShader();
 	);
@@ -162,8 +361,8 @@ DEFINE_RHIMETHOD_1(
  * @param GeometryShader - existing geometry shader
  * @param PixelShader - existing pixel shader
  */
-DEFINE_RHIMETHOD_6(
-	FBoundShaderStateRHIRef,RHICreateBoundShaderState,
+DEFINE_RHIMETHOD_GLOBAL_6(
+	FBoundShaderStateRHIRef,CreateBoundShaderState,
 	FVertexDeclarationRHIParamRef,VertexDeclaration,
 	FVertexShaderRHIParamRef,VertexShader,
 	FHullShaderRHIParamRef,HullShader,
@@ -178,13 +377,13 @@ DEFINE_RHIMETHOD_6(
  *that require shader setting before resource binding.
  */
 DEFINE_RHIMETHOD_1(
-	void,RHISetComputeShader,
+	void,SetComputeShader,
 	FComputeShaderRHIParamRef,ComputeShader,
 	,
 	);
 
 DEFINE_RHIMETHOD_3(
-	void,RHIDispatchComputeShader,	
+	void,DispatchComputeShader,	
 	uint32,ThreadGroupCountX,
 	uint32,ThreadGroupCountY,
 	uint32,ThreadGroupCountZ,
@@ -192,20 +391,20 @@ DEFINE_RHIMETHOD_3(
 	);
 
 DEFINE_RHIMETHOD_2(
-	void,RHIDispatchIndirectComputeShader,	
+	void,DispatchIndirectComputeShader,	
 	FVertexBufferRHIParamRef,ArgumentBuffer,
 	uint32,ArgumentOffset,
 	,
 	);
 
 DEFINE_RHIMETHOD_1(
-	void,RHIAutomaticCacheFlushAfterComputeShader,
+	void,AutomaticCacheFlushAfterComputeShader,
 	bool,bEnable,
 	,
 	);
 
 DEFINE_RHIMETHOD_0(
-	void,RHIFlushComputeShaderCache,
+	void,FlushComputeShaderCache,
 	,
 	);
 
@@ -213,7 +412,7 @@ DEFINE_RHIMETHOD_0(
 // @param Count >0
 // @param Data must not be 0
 DEFINE_RHIMETHOD_2(
-	void,RHISetMultipleViewports,
+	void,SetMultipleViewports,
 	uint32,Count,
 	const FViewportBounds*,Data,
 	,
@@ -225,8 +424,8 @@ DEFINE_RHIMETHOD_2(
  * @param NumBytes - The number of bytes the uniform buffer should contain.
  * @return The new uniform buffer.
  */
-DEFINE_RHIMETHOD_3(
-	FUniformBufferRHIRef,RHICreateUniformBuffer,
+DEFINE_RHIMETHOD_GLOBAL_3(
+	FUniformBufferRHIRef,CreateUniformBuffer,
 	const void*,Contents,
 	const FRHIUniformBufferLayout&,Layout,
 	EUniformBufferUsage,Usage,
@@ -234,8 +433,8 @@ DEFINE_RHIMETHOD_3(
 	);
 
 
-DEFINE_RHIMETHOD_4(
-	FIndexBufferRHIRef,RHICreateIndexBuffer,
+DEFINE_RHIMETHOD_GLOBAL_4(
+	FIndexBufferRHIRef,CreateIndexBuffer,
 	uint32,Stride,
 	uint32,Size,
 	uint32,InUsage,
@@ -243,16 +442,16 @@ DEFINE_RHIMETHOD_4(
 	return,if (CreateInfo.ResourceArray) { CreateInfo.ResourceArray->Discard(); } return new FRHIIndexBuffer(Stride,Size,InUsage);
 	);
 
-DEFINE_RHIMETHOD_4(
-	void*,RHILockIndexBuffer,
+DEFINE_RHIMETHOD_GLOBALFLUSH_4(
+	void*,LockIndexBuffer,
 	FIndexBufferRHIParamRef,IndexBuffer,
 	uint32,Offset,
 	uint32,Size,
 	EResourceLockMode,LockMode,
 	return,return GetStaticBuffer();
 	);
-DEFINE_RHIMETHOD_1(
-	void,RHIUnlockIndexBuffer,
+DEFINE_RHIMETHOD_GLOBALFLUSH_1(
+	void,UnlockIndexBuffer,
 	FIndexBufferRHIParamRef,IndexBuffer,
 	,
 	);
@@ -260,31 +459,31 @@ DEFINE_RHIMETHOD_1(
 /**
  * @param ResourceArray - An optional pointer to a resource array containing the resource's data.
  */
-DEFINE_RHIMETHOD_3(
-	FVertexBufferRHIRef,RHICreateVertexBuffer,
+DEFINE_RHIMETHOD_GLOBAL_3(
+	FVertexBufferRHIRef,CreateVertexBuffer,
 	uint32,Size,
 	uint32,InUsage,
 	FRHIResourceCreateInfo&,CreateInfo,
 	return,if (CreateInfo.ResourceArray) { CreateInfo.ResourceArray->Discard(); } return new FRHIVertexBuffer(Size,InUsage);
 	);
 
-DEFINE_RHIMETHOD_4(
-	void*,RHILockVertexBuffer,
+DEFINE_RHIMETHOD_GLOBALFLUSH_4(
+	void*,LockVertexBuffer,
 	FVertexBufferRHIParamRef,VertexBuffer,
 	uint32,Offset,
 	uint32,SizeRHI,
 	EResourceLockMode,LockMode,
 	return,return GetStaticBuffer();
 	);
-DEFINE_RHIMETHOD_1(
-	void,RHIUnlockVertexBuffer,
+DEFINE_RHIMETHOD_GLOBALFLUSH_1(
+	void,UnlockVertexBuffer,
 	FVertexBufferRHIParamRef,VertexBuffer,
 	,
 	);
 
 /** Copies the contents of one vertex buffer to another vertex buffer.  They must have identical sizes. */
 DEFINE_RHIMETHOD_2(
-	void,RHICopyVertexBuffer,
+	void,CopyVertexBuffer,
 	FVertexBufferRHIParamRef,SourceBuffer,
 	FVertexBufferRHIParamRef,DestBuffer,
 	,
@@ -293,8 +492,8 @@ DEFINE_RHIMETHOD_2(
 /**
  * @param ResourceArray - An optional pointer to a resource array containing the resource's data.
  */
-DEFINE_RHIMETHOD_4(
-	FStructuredBufferRHIRef,RHICreateStructuredBuffer,
+DEFINE_RHIMETHOD_GLOBAL_4(
+	FStructuredBufferRHIRef,CreateStructuredBuffer,
 	uint32,Stride,
 	uint32,Size,
 	uint32,InUsage,
@@ -302,23 +501,23 @@ DEFINE_RHIMETHOD_4(
 	return,if (CreateInfo.ResourceArray) { CreateInfo.ResourceArray->Discard(); } return new FRHIStructuredBuffer(Stride,Size,InUsage);
 	);
 
-DEFINE_RHIMETHOD_4(
-	void*,RHILockStructuredBuffer,
+DEFINE_RHIMETHOD_GLOBALFLUSH_4(
+	void*,LockStructuredBuffer,
 	FStructuredBufferRHIParamRef,StructuredBuffer,
 	uint32,Offset,
 	uint32,SizeRHI,
 	EResourceLockMode,LockMode,
 	return,return GetStaticBuffer();
 	);
-DEFINE_RHIMETHOD_1(
-	void,RHIUnlockStructuredBuffer,
+DEFINE_RHIMETHOD_GLOBALFLUSH_1(
+	void,UnlockStructuredBuffer,
 	FStructuredBufferRHIParamRef,StructuredBuffer,
 	,
 	);
 
 /** Creates an unordered access view of the given structured buffer. */
-DEFINE_RHIMETHOD_3(
-	FUnorderedAccessViewRHIRef,RHICreateUnorderedAccessView,
+DEFINE_RHIMETHOD_GLOBAL_3(
+	FUnorderedAccessViewRHIRef,CreateUnorderedAccessView,
 	FStructuredBufferRHIParamRef,StructuredBuffer,
 	bool,bUseUAVCounter,
 	bool,bAppendBuffer,
@@ -326,30 +525,30 @@ DEFINE_RHIMETHOD_3(
 	);
 
 /** Creates an unordered access view of the given texture. */
-DEFINE_RHIMETHOD_1(
-	FUnorderedAccessViewRHIRef,RHICreateUnorderedAccessView,
+	DEFINE_RHIMETHOD_GLOBAL_1(
+	FUnorderedAccessViewRHIRef,CreateUnorderedAccessView,
 	FTextureRHIParamRef,Texture,
 	return,return new FRHIUnorderedAccessView();
 	);
 	
 /** Creates an unordered access view of the given texture. */
-DEFINE_RHIMETHOD_2(
-	FUnorderedAccessViewRHIRef,RHICreateUnorderedAccessView,
+	DEFINE_RHIMETHOD_GLOBAL_2(
+	FUnorderedAccessViewRHIRef,CreateUnorderedAccessView,
 	FVertexBufferRHIParamRef,VertexBuffer,
 	uint8,Format,
 	return,return new FRHIUnorderedAccessView();
 	);
 
 /** Creates a shader resource view of the given structured buffer. */
-DEFINE_RHIMETHOD_1(
-	FShaderResourceViewRHIRef,RHICreateShaderResourceView,
+	DEFINE_RHIMETHOD_GLOBAL_1(
+	FShaderResourceViewRHIRef,CreateShaderResourceView,
 	FStructuredBufferRHIParamRef,StructuredBuffer,
 	return,return new FRHIShaderResourceView();
 	);
 
 /** Creates a shader resource view of the given vertex buffer. */
-DEFINE_RHIMETHOD_3(
-	FShaderResourceViewRHIRef,RHICreateShaderResourceView,
+	DEFINE_RHIMETHOD_GLOBAL_3(
+	FShaderResourceViewRHIRef,CreateShaderResourceView,
 	FVertexBufferRHIParamRef,VertexBuffer,
 	uint32,Stride,
 	uint8,Format,
@@ -358,7 +557,7 @@ DEFINE_RHIMETHOD_3(
 
 /** Clears a UAV to the multi-component value provided. */
 DEFINE_RHIMETHOD_2(
-	void,RHIClearUAV,
+	void,ClearUAV,
 	FUnorderedAccessViewRHIParamRef,UnorderedAccessViewRHI,
 	const uint32*,Values,
 	,
@@ -375,8 +574,8 @@ DEFINE_RHIMETHOD_2(
  * @param Flags - ETextureCreateFlags creation flags
  * @param OutAlign - Alignment required for this texture.  Output parameter.
  */
-DEFINE_RHIMETHOD_7(
-	uint64,	RHICalcTexture2DPlatformSize,
+DEFINE_RHIMETHOD_GLOBAL_7(
+	uint64,CalcTexture2DPlatformSize,
 	uint32, SizeX,
 	uint32, SizeY,	
 	uint8, Format,
@@ -396,8 +595,8 @@ DEFINE_RHIMETHOD_7(
 * @param Flags - ETextureCreateFlags creation flags
 * @param OutAlign - Alignment required for this texture.  Output parameter.
 */
-DEFINE_RHIMETHOD_7(
-	uint64,RHICalcTexture3DPlatformSize,
+DEFINE_RHIMETHOD_GLOBAL_7(
+	uint64,CalcTexture3DPlatformSize,
 	uint32,SizeX,
 	uint32,SizeY,
 	uint32,SizeZ,
@@ -415,8 +614,8 @@ DEFINE_RHIMETHOD_7(
 * @param Flags - ETextureCreateFlags creation flags
 * @param OutAlign - Alignment required for this texture.  Output parameter.
 */
-DEFINE_RHIMETHOD_5(
-	uint64,RHICalcTextureCubePlatformSize,
+DEFINE_RHIMETHOD_GLOBAL_5(
+	uint64,CalcTextureCubePlatformSize,
 	uint32,Size,
 	uint8,Format,
 	uint32,NumMips,
@@ -428,8 +627,8 @@ DEFINE_RHIMETHOD_5(
  * Retrieves texture memory stats.
  * safe to call on the main thread
  */
-DEFINE_RHIMETHOD_1(
-	void,RHIGetTextureMemoryStats,
+DEFINE_RHIMETHOD_GLOBAL_1(
+	void,GetTextureMemoryStats,
 	FTextureMemoryStats&,OutStats,
 	,
 	);
@@ -446,7 +645,7 @@ DEFINE_RHIMETHOD_1(
  * @return true if successful, false otherwise
  */
 DEFINE_RHIMETHOD_5(
-	bool,RHIGetTextureMemoryVisualizeData,
+	bool,GetTextureMemoryVisualizeData,
 	FColor*,TextureData,
 	int32,SizeX,
 	int32,SizeY,
@@ -456,12 +655,12 @@ DEFINE_RHIMETHOD_5(
 	);
 
 DEFINE_RHIMETHOD_0(
-	FTextureReferenceRHIRef,RHICreateTextureReference,
+	FTextureReferenceRHIRef,CreateTextureReference,
 	return, return new FRHITextureReferenceNullImpl();
 	);
 
-DEFINE_RHIMETHOD_2(
-	void,RHIUpdateTextureReference,
+DEFINE_RHIMETHOD_GLOBALFLUSH_2(
+	void,UpdateTextureReference,
 	FTextureReferenceRHIParamRef,TextureRef,
 	FTextureRHIParamRef,NewTexture,
 	,if (TextureRef) { ((FRHITextureReferenceNullImpl*)TextureRef)->SetReferencedTexture(NewTexture); }
@@ -476,8 +675,8 @@ DEFINE_RHIMETHOD_2(
 * @param NumSamples - number of MSAA samples, usually 1
 * @param Flags - ETextureCreateFlags creation flags
 */
-DEFINE_RHIMETHOD_7(
-	FTexture2DRHIRef,RHICreateTexture2D,
+DEFINE_RHIMETHOD_GLOBAL_7(
+	FTexture2DRHIRef,CreateTexture2D,
 	uint32,SizeX,
 	uint32,SizeY,
 	uint8,Format,
@@ -501,8 +700,8 @@ DEFINE_RHIMETHOD_7(
  * @param NumInitialMips - how many mips are provided in InitialMipData
  * @returns a reference to a 2D texture resource
  */
-DEFINE_RHIMETHOD_7(
-	FTexture2DRHIRef,RHIAsyncCreateTexture2D,
+DEFINE_RHIMETHOD_GLOBAL_7(
+	FTexture2DRHIRef,AsyncCreateTexture2D,
 	uint32,SizeX,
 	uint32,SizeY,
 	uint8,Format,
@@ -520,8 +719,8 @@ DEFINE_RHIMETHOD_7(
  * @param DestTexture2D - destination texture
  * @param SrcTexture2D - source texture
  */
-DEFINE_RHIMETHOD_2(
-	void,RHICopySharedMips,
+DEFINE_RHIMETHOD_GLOBALFLUSH_2(
+	void,CopySharedMips,
 	FTexture2DRHIParamRef,DestTexture2D,
 	FTexture2DRHIParamRef,SrcTexture2D,
 	return,return;
@@ -536,8 +735,8 @@ DEFINE_RHIMETHOD_2(
 * @param NumMips - number of mips to generate or 0 for full mip pyramid
 * @param Flags - ETextureCreateFlags creation flags
 */
-DEFINE_RHIMETHOD_7(
-	FTexture2DArrayRHIRef,RHICreateTexture2DArray,
+DEFINE_RHIMETHOD_GLOBAL_7(
+	FTexture2DArrayRHIRef,CreateTexture2DArray,
 	uint32,SizeX,
 	uint32,SizeY,
 	uint32,SizeZ,
@@ -557,8 +756,8 @@ DEFINE_RHIMETHOD_7(
 * @param NumMips - number of mips to generate or 0 for full mip pyramid
 * @param Flags - ETextureCreateFlags creation flags
 */
-DEFINE_RHIMETHOD_7(
-	FTexture3DRHIRef,RHICreateTexture3D,
+DEFINE_RHIMETHOD_GLOBAL_7(
+	FTexture3DRHIRef,CreateTexture3D,
 	uint32,SizeX,
 	uint32,SizeY,
 	uint32,SizeZ,
@@ -572,8 +771,8 @@ DEFINE_RHIMETHOD_7(
 /**
  * @param Ref may be 0
  */
-DEFINE_RHIMETHOD_2(
-	void,RHIGetResourceInfo,
+DEFINE_RHIMETHOD_GLOBAL_2(
+	void,GetResourceInfo,
 	FTextureRHIParamRef,Ref,
 	FRHIResourceInfo&,OutInfo,
 	return,return;
@@ -583,8 +782,8 @@ DEFINE_RHIMETHOD_2(
 * Creates a shader resource view for a 2d texture, viewing only a single
 * mip level. Useful when rendering to one mip while sampling from another.
 */
-DEFINE_RHIMETHOD_2(
-	FShaderResourceViewRHIRef,RHICreateShaderResourceView,
+DEFINE_RHIMETHOD_GLOBAL_2(
+	FShaderResourceViewRHIRef,CreateShaderResourceView,
 	FTexture2DRHIParamRef,Texture2DRHI,
 	uint8, MipLevel,
 	return,return new FRHIShaderResourceView();
@@ -594,8 +793,8 @@ DEFINE_RHIMETHOD_2(
 * Creates a shader resource view for a 2d texture, with a different
 * format from the original.  Useful when sampling stencil.
 */
-DEFINE_RHIMETHOD_4(
-	FShaderResourceViewRHIRef,RHICreateShaderResourceView,
+DEFINE_RHIMETHOD_GLOBAL_4(
+	FShaderResourceViewRHIRef,CreateShaderResourceView,
 	FTexture2DRHIParamRef,Texture2DRHI,
 	uint8, MipLevel,
 	uint8, NumMipLevels,
@@ -607,7 +806,7 @@ DEFINE_RHIMETHOD_4(
 * Generates mip maps for a texture.
 */
 DEFINE_RHIMETHOD_1(
-	void,RHIGenerateMips,
+	void,GenerateMips,
 	FTextureRHIParamRef,Texture,
 	return,return;
 	);
@@ -618,8 +817,8 @@ DEFINE_RHIMETHOD_1(
  * @param	TextureRHI		- Texture we want to know the size of, 0 is safely ignored
  * @return					- Size in Bytes
  */
-DEFINE_RHIMETHOD_1(
-	uint32,RHIComputeMemorySize,
+DEFINE_RHIMETHOD_GLOBAL_1(
+	uint32,ComputeMemorySize,
 	FTextureRHIParamRef,TextureRHI,
 	return,return 0;
 	);
@@ -639,8 +838,8 @@ DEFINE_RHIMETHOD_1(
  * @param RequestStatus	- Will be decremented by 1 when the reallocation is complete (success or failure).
  * @return				- New reference to the texture, or an invalid reference upon failure
  */
-DEFINE_RHIMETHOD_5(
-	FTexture2DRHIRef,RHIAsyncReallocateTexture2D,
+DEFINE_RHIMETHOD_GLOBAL_5(
+	FTexture2DRHIRef,AsyncReallocateTexture2D,
 	FTexture2DRHIParamRef,Texture2D,
 	int32,NewMipCount,
 	int32,NewSizeX,
@@ -660,8 +859,8 @@ DEFINE_RHIMETHOD_5(
  *	TexRealloc_Failed		Reallocation failed
  *	TexRealloc_InProgress	Reallocation is still in progress, try again later
  */
-DEFINE_RHIMETHOD_2(
-	ETextureReallocationStatus,RHIFinalizeAsyncReallocateTexture2D,
+DEFINE_RHIMETHOD_GLOBAL_2(
+	ETextureReallocationStatus,FinalizeAsyncReallocateTexture2D,
 	FTexture2DRHIParamRef,Texture2D,
 	bool,bBlockUntilCompleted,
 	return,return TexRealloc_Succeeded;
@@ -675,8 +874,8 @@ DEFINE_RHIMETHOD_2(
  * @param bBlockUntilCompleted	If true, blocks until the cancellation is fully completed
  * @return						Reallocation status
  */
-DEFINE_RHIMETHOD_2(
-	ETextureReallocationStatus,RHICancelAsyncReallocateTexture2D,
+DEFINE_RHIMETHOD_GLOBAL_2(
+	ETextureReallocationStatus,CancelAsyncReallocateTexture2D,
 	FTexture2DRHIParamRef,Texture2D,
 	bool,bBlockUntilCompleted,
 	return,return TexRealloc_Succeeded;
@@ -691,8 +890,8 @@ DEFINE_RHIMETHOD_2(
 * @param bLockWithinMiptail - for platforms that support packed miptails allow locking of individual mip levels within the miptail
 * @return pointer to the CPU accessible resource data
 */
-DEFINE_RHIMETHOD_5(
-	void*,RHILockTexture2D,
+DEFINE_RHIMETHOD_GLOBALFLUSH_5(
+	void*,LockTexture2D,
 	FTexture2DRHIParamRef,Texture,
 	uint32,MipIndex,
 	EResourceLockMode,LockMode,
@@ -707,8 +906,8 @@ DEFINE_RHIMETHOD_5(
 * @param MipIndex - index of the mip level to unlock
 * @param bLockWithinMiptail - for platforms that support packed miptails allow locking of individual mip levels within the miptail
 */
-DEFINE_RHIMETHOD_3(
-	void,RHIUnlockTexture2D,
+DEFINE_RHIMETHOD_GLOBALFLUSH_3(
+	void,UnlockTexture2D,
 	FTexture2DRHIParamRef,Texture,
 	uint32,MipIndex,
 	bool,bLockWithinMiptail,
@@ -724,8 +923,8 @@ DEFINE_RHIMETHOD_3(
 * @param bLockWithinMiptail - for platforms that support packed miptails allow locking of individual mip levels within the miptail
 * @return pointer to the CPU accessible resource data
 */
-DEFINE_RHIMETHOD_6(
-	void*,RHILockTexture2DArray,
+DEFINE_RHIMETHOD_GLOBALFLUSH_6(
+	void*,LockTexture2DArray,
 	FTexture2DArrayRHIParamRef,Texture,
 	uint32,TextureIndex,
 	uint32,MipIndex,
@@ -741,8 +940,8 @@ DEFINE_RHIMETHOD_6(
 * @param MipIndex - index of the mip level to unlock
 * @param bLockWithinMiptail - for platforms that support packed miptails allow locking of individual mip levels within the miptail
 */
-DEFINE_RHIMETHOD_4(
-	void,RHIUnlockTexture2DArray,
+DEFINE_RHIMETHOD_GLOBALFLUSH_4(
+	void,UnlockTexture2DArray,
 	FTexture2DArrayRHIParamRef,Texture,
 	uint32,TextureIndex,
 	uint32,MipIndex,
@@ -758,8 +957,8 @@ DEFINE_RHIMETHOD_4(
 * @param SourcePitch - size in bytes of each row of the source image
 * @param SourceData - source image data, starting at the upper left corner of the source rectangle (in same pixel format as texture)
 */
-DEFINE_RHIMETHOD_5(
-	void,RHIUpdateTexture2D,
+DEFINE_RHIMETHOD_GLOBALFLUSH_5(
+	void,UpdateTexture2D,
 	FTexture2DRHIParamRef,Texture,
 	uint32,MipIndex,
 	const struct FUpdateTextureRegion2D&,UpdateRegion,
@@ -777,8 +976,8 @@ DEFINE_RHIMETHOD_5(
 * @param SourceDepthPitch - size in bytes of each depth slice of the source image, usually Bpp * SizeX * SizeY
 * @param SourceData - source image data, starting at the upper left corner of the source rectangle (in same pixel format as texture)
 */
-DEFINE_RHIMETHOD_6(
-	void,RHIUpdateTexture3D,
+DEFINE_RHIMETHOD_GLOBALFLUSH_6(
+	void,UpdateTexture3D,
 	FTexture3DRHIParamRef,Texture,
 	uint32,MipIndex,
 	const struct FUpdateTextureRegion3D&,UpdateRegion,
@@ -795,8 +994,8 @@ DEFINE_RHIMETHOD_6(
 * @param NumMips - number of mips to generate or 0 for full mip pyramid
 * @param Flags - ETextureCreateFlags creation flags
 */
-DEFINE_RHIMETHOD_5(
-	FTextureCubeRHIRef,RHICreateTextureCube,
+DEFINE_RHIMETHOD_GLOBAL_5(
+	FTextureCubeRHIRef,CreateTextureCube,
 	uint32,Size,
 	uint8,Format,
 	uint32,NumMips,
@@ -806,24 +1005,24 @@ DEFINE_RHIMETHOD_5(
 	);
 
 
-	/**
-	* Creates a Cube Array RHI texture resource
-	* @param Size - width/height of the texture to create
-	* @param ArraySize - number of array elements of the texture to create
-	* @param Format - EPixelFormat texture format
-	* @param NumMips - number of mips to generate or 0 for full mip pyramid
-	* @param Flags - ETextureCreateFlags creation flags
-	*/
-	DEFINE_RHIMETHOD_6(
-		FTextureCubeRHIRef,RHICreateTextureCubeArray,
-		uint32,Size,
-		uint32,ArraySize,
-		uint8,Format,
-		uint32,NumMips,
-		uint32,Flags,
+/**
+* Creates a Cube Array RHI texture resource
+* @param Size - width/height of the texture to create
+* @param ArraySize - number of array elements of the texture to create
+* @param Format - EPixelFormat texture format
+* @param NumMips - number of mips to generate or 0 for full mip pyramid
+* @param Flags - ETextureCreateFlags creation flags
+*/
+DEFINE_RHIMETHOD_GLOBAL_6(
+	FTextureCubeRHIRef,CreateTextureCubeArray,
+	uint32,Size,
+	uint32,ArraySize,
+	uint8,Format,
+	uint32,NumMips,
+	uint32,Flags,
 	FRHIResourceCreateInfo&,CreateInfo,
-		return,return new FRHITextureCube(Size,NumMips,(EPixelFormat)Format,Flags);
-	);
+	return,return new FRHITextureCube(Size,NumMips,(EPixelFormat)Format,Flags);
+);
 
 
 /**
@@ -835,8 +1034,8 @@ DEFINE_RHIMETHOD_5(
 * @param bLockWithinMiptail - for platforms that support packed miptails allow locking of individual mip levels within the miptail
 * @return pointer to the CPU accessible resource data
 */
-DEFINE_RHIMETHOD_7(
-	void*,RHILockTextureCubeFace,
+DEFINE_RHIMETHOD_GLOBALFLUSH_7(
+	void*,LockTextureCubeFace,
 	FTextureCubeRHIParamRef,Texture,
 	uint32,FaceIndex,
 	uint32,ArrayIndex,
@@ -853,8 +1052,8 @@ DEFINE_RHIMETHOD_7(
 * @param MipIndex - index of the mip level to unlock
 * @param bLockWithinMiptail - for platforms that support packed miptails allow locking of individual mip levels within the miptail
 */
-DEFINE_RHIMETHOD_5(
-	void,RHIUnlockTextureCubeFace,
+DEFINE_RHIMETHOD_GLOBALFLUSH_5(
+	void,UnlockTextureCubeFace,
 	FTextureCubeRHIParamRef,Texture,
 	uint32,FaceIndex,
 	uint32,ArrayIndex,
@@ -871,7 +1070,7 @@ DEFINE_RHIMETHOD_5(
 * @param ResolveParams - optional resolve params
 */
 DEFINE_RHIMETHOD_4(
-	void,RHICopyToResolveTarget,
+	void,CopyToResolveTarget,
 	FTextureRHIParamRef,SourceTexture,
 	FTextureRHIParamRef,DestTexture,
 	bool,bKeepOriginalSurface,
@@ -879,8 +1078,8 @@ DEFINE_RHIMETHOD_4(
 	,
 	);
 
-DEFINE_RHIMETHOD_2(
-	void, RHIBindDebugLabelName,
+DEFINE_RHIMETHOD_GLOBAL_2(
+	void,BindDebugLabelName,
 	FTextureRHIParamRef, Texture,
 	const TCHAR*, Name,
 	,
@@ -891,7 +1090,7 @@ DEFINE_RHIMETHOD_2(
  * If the format or texture type is unsupported the OutData array will be size 0
  */
 DEFINE_RHIMETHOD_4(
-	void,RHIReadSurfaceData,
+	void,ReadSurfaceData,
 	FTextureRHIParamRef,Texture,
 	FIntRect,Rect,
 	TArray<FColor>&,OutData,
@@ -901,7 +1100,7 @@ DEFINE_RHIMETHOD_4(
 
 /** Watch out for OutData to be 0 (can happen on DXGI_ERROR_DEVICE_REMOVED), don't call RHIUnmapStagingSurface in that case. */
 DEFINE_RHIMETHOD_4(
-	void,RHIMapStagingSurface,
+	void,MapStagingSurface,
 	FTextureRHIParamRef,Texture,
 	void*&,OutData,
 	int32&,OutWidth,
@@ -911,13 +1110,13 @@ DEFINE_RHIMETHOD_4(
 
 /** call after a succesful RHIMapStagingSurface() call */
 DEFINE_RHIMETHOD_1(
-	void,RHIUnmapStagingSurface,
+	void,UnmapStagingSurface,
 	FTextureRHIParamRef,Texture,
 	,
 	);
 
 DEFINE_RHIMETHOD_6(
-	void,RHIReadSurfaceFloatData,
+	void,ReadSurfaceFloatData,
 	FTextureRHIParamRef,Texture,
 	FIntRect,Rect,
 	TArray<FFloat16Color>&,OutData,
@@ -928,7 +1127,7 @@ DEFINE_RHIMETHOD_6(
 	);
 
 DEFINE_RHIMETHOD_4(
-	void,RHIRead3DSurfaceFloatData,
+	void,Read3DSurfaceFloatData,
 	FTextureRHIParamRef,Texture,
 	FIntRect,Rect,
 	FIntPoint,ZMinMax,
@@ -938,43 +1137,43 @@ DEFINE_RHIMETHOD_4(
 
 
 // Occlusion/Timer queries.
-DEFINE_RHIMETHOD_1(
-	FRenderQueryRHIRef,RHICreateRenderQuery,
+DEFINE_RHIMETHOD_GLOBAL_1(
+	FRenderQueryRHIRef,CreateRenderQuery,
 	ERenderQueryType, QueryType,
 	return,return new FRHIRenderQuery();
 	);
 DEFINE_RHIMETHOD_1(
-	void,RHIResetRenderQuery,
+	void,ResetRenderQuery,
 	FRenderQueryRHIParamRef,RenderQuery,
 	,
 	);
 DEFINE_RHIMETHOD_3(
-	bool,RHIGetRenderQueryResult,
+	bool,GetRenderQueryResult,
 	FRenderQueryRHIParamRef,RenderQuery,
 	uint64&,OutResult,
 	bool,bWait,
 	return,return true;
 	);
 DEFINE_RHIMETHOD_1(
-	void,RHIBeginRenderQuery,
+	void,BeginRenderQuery,
 	FRenderQueryRHIParamRef,RenderQuery,
 	,
 	);
 DEFINE_RHIMETHOD_1(
-	void,RHIEndRenderQuery,
+	void,EndRenderQuery,
 	FRenderQueryRHIParamRef,RenderQuery,
 	,
 	);
 
 
 DEFINE_RHIMETHOD_2(
-	void,RHIBeginDrawingViewport,
+	void,BeginDrawingViewport,
 	FViewportRHIParamRef,Viewport,
 	FTextureRHIParamRef,RenderTargetRHI,
 	,
 	);
 DEFINE_RHIMETHOD_3(
-	void,RHIEndDrawingViewport,
+	void,EndDrawingViewport,
 	FViewportRHIParamRef,Viewport,
 	bool,bPresent,
 	bool,bLockToVsync,
@@ -986,21 +1185,21 @@ DEFINE_RHIMETHOD_3(
  * @return true if currently within a BeginDrawingViewport/EndDrawingViewport block
  */
 DEFINE_RHIMETHOD_0(
-	bool,RHIIsDrawingViewport,
+	bool,IsDrawingViewport,
 	return,return false;
 	);
 DEFINE_RHIMETHOD_1(
-	FTexture2DRHIRef,RHIGetViewportBackBuffer,
+	FTexture2DRHIRef,GetViewportBackBuffer,
 	FViewportRHIParamRef,Viewport,
 	return,return new FRHITexture2D(1,1,1,1,PF_B8G8R8A8,TexCreate_RenderTargetable);
 	);
 
 DEFINE_RHIMETHOD_0(
-	void,RHIBeginFrame,
+	void,BeginFrame,
 	return,return;
 	);
 DEFINE_RHIMETHOD_0(
-	void,RHIEndFrame,
+	void,EndFrame,
 	return,return;
 	);
 
@@ -1010,7 +1209,7 @@ DEFINE_RHIMETHOD_0(
  * references.
  */
 DEFINE_RHIMETHOD_0(
-	void,RHIBeginScene,
+	void,BeginScene,
 	return,return;
 	);
 
@@ -1018,56 +1217,56 @@ DEFINE_RHIMETHOD_0(
  * Signals the end of scene rendering. See RHIBeginScene.
  */
 DEFINE_RHIMETHOD_0(
-	void,RHIEndScene,
+	void,EndScene,
 	return,return;
 	);
 
 /*
  * Acquires or releases ownership of the platform-specific rendering context for the calling thread
  */
-DEFINE_RHIMETHOD_0(
-	void,RHIAcquireThreadOwnership,
+DEFINE_RHIMETHOD_GLOBALFLUSH_0(
+	void,AcquireThreadOwnership,
 	return,return;
 	);
-DEFINE_RHIMETHOD_0(
-	void,RHIReleaseThreadOwnership,
+DEFINE_RHIMETHOD_GLOBALFLUSH_0(
+	void,ReleaseThreadOwnership,
 	return,return;
 	);
 
 // Flush driver resources. Typically called when switching contexts/threads
-DEFINE_RHIMETHOD_0(
-	void,RHIFlushResources,
+DEFINE_RHIMETHOD_GLOBALFLUSH_0(
+	void,FlushResources,
 	return,return;
 	);
 /*
  * Returns the total GPU time taken to render the last frame. Same metric as FPlatformTime::Cycles().
  */
-DEFINE_RHIMETHOD_0(
-	uint32,RHIGetGPUFrameCycles,
+DEFINE_RHIMETHOD_GLOBAL_0(
+	uint32,GetGPUFrameCycles,
 	return,return 0;
 	);
 
 /**
  * The following RHI functions must be called from the main thread.
  */
-DEFINE_RHIMETHOD_4(
-	FViewportRHIRef,RHICreateViewport,
+DEFINE_RHIMETHOD_GLOBAL_4(
+	FViewportRHIRef,CreateViewport,
 	void*,WindowHandle,
 	uint32,SizeX,
 	uint32,SizeY,
 	bool,bIsFullscreen,
 	return,return new FRHIViewport();
 	);
-DEFINE_RHIMETHOD_4(
-	void,RHIResizeViewport,
+DEFINE_RHIMETHOD_GLOBAL_4(
+	void,ResizeViewport,
 	FViewportRHIParamRef,Viewport,
 	uint32,SizeX,
 	uint32,SizeY,
 	bool,bIsFullscreen,
 	,
 	);
-DEFINE_RHIMETHOD_1(
-	void,RHITick,
+DEFINE_RHIMETHOD_GLOBAL_1(
+	void,Tick,
 	float,DeltaTime,
 	,
 	);
@@ -1077,8 +1276,8 @@ DEFINE_RHIMETHOD_1(
 //
 
 // Vertex state.
-DEFINE_RHIMETHOD_4(
-	void,RHISetStreamSource,
+DEFINE_RHIMETHOD_CMDLIST_4(
+	void,SetStreamSource,
 	uint32,StreamIndex,
 	FVertexBufferRHIParamRef,VertexBuffer,
 	uint32,Stride,
@@ -1088,7 +1287,7 @@ DEFINE_RHIMETHOD_4(
 
 /** Sets stream output targets, for use with a geometry shader created with RHICreateGeometryShaderWithStreamOutput. */
 DEFINE_RHIMETHOD_3(
-	void,RHISetStreamOutTargets,
+	void,SetStreamOutTargets,
 	uint32,NumTargets,
 	const FVertexBufferRHIParamRef*,VertexBuffers,
 	const uint32*,Offsets,
@@ -1096,8 +1295,8 @@ DEFINE_RHIMETHOD_3(
 	);
 
 // Rasterizer state.
-DEFINE_RHIMETHOD_1(
-	void,RHISetRasterizerState,
+DEFINE_RHIMETHOD_CMDLIST_1(
+	void,SetRasterizerState,
 	FRasterizerStateRHIParamRef,NewState,
 	,
 	);
@@ -1106,7 +1305,7 @@ DEFINE_RHIMETHOD_1(
 // @param MaxX excluding like Win32 RECT
 // @param MaxY excluding like Win32 RECT
 DEFINE_RHIMETHOD_6(
-	void,RHISetViewport,
+	void,SetViewport,
 	uint32,MinX,
 	uint32,MinY,
 	float,MinZ,
@@ -1120,7 +1319,7 @@ DEFINE_RHIMETHOD_6(
 // @param MaxX excluding like Win32 RECT
 // @param MaxY excluding like Win32 RECT
 DEFINE_RHIMETHOD_5(
-	void,RHISetScissorRect,
+	void,SetScissorRect,
 	bool,bEnable,
 	uint32,MinX,
 	uint32,MinY,
@@ -1134,15 +1333,15 @@ DEFINE_RHIMETHOD_5(
  * Set bound shader state. This will set the vertex decl/shader, and pixel shader
  * @param BoundShaderState - state resource
  */
-DEFINE_RHIMETHOD_1(
-	void,RHISetBoundShaderState,
+DEFINE_RHIMETHOD_CMDLIST_1(
+	void,SetBoundShaderState,
 	FBoundShaderStateRHIParamRef,BoundShaderState,
 	,
 	);
 
 /** Set the shader resource view of a surface.  This is used for binding TextureMS parameter types that need a multi sampled view. */
-DEFINE_RHIMETHOD_3(
-	void,RHISetShaderTexture,
+DEFINE_RHIMETHOD_CMDLIST_3(
+	void,SetShaderTexture,
 	FVertexShaderRHIParamRef,VertexShader,
 	uint32,TextureIndex,
 	FTextureRHIParamRef,NewTexture,
@@ -1150,8 +1349,8 @@ DEFINE_RHIMETHOD_3(
 	);
 
 /** Set the shader resource view of a surface.  This is used for binding TextureMS parameter types that need a multi sampled view. */
-DEFINE_RHIMETHOD_3(
-	void,RHISetShaderTexture,
+DEFINE_RHIMETHOD_CMDLIST_3(
+	void,SetShaderTexture,
 	FHullShaderRHIParamRef,HullShader,
 	uint32,TextureIndex,
 	FTextureRHIParamRef,NewTexture,
@@ -1159,8 +1358,8 @@ DEFINE_RHIMETHOD_3(
 	);
 
 /** Set the shader resource view of a surface.  This is used for binding TextureMS parameter types that need a multi sampled view. */
-DEFINE_RHIMETHOD_3(
-	void,RHISetShaderTexture,
+DEFINE_RHIMETHOD_CMDLIST_3(
+	void,SetShaderTexture,
 	FDomainShaderRHIParamRef,DomainShader,
 	uint32,TextureIndex,
 	FTextureRHIParamRef,NewTexture,
@@ -1168,8 +1367,8 @@ DEFINE_RHIMETHOD_3(
 	);
 
 /** Set the shader resource view of a surface.  This is used for binding TextureMS parameter types that need a multi sampled view. */
-DEFINE_RHIMETHOD_3(
-	void,RHISetShaderTexture,
+DEFINE_RHIMETHOD_CMDLIST_3(
+	void,SetShaderTexture,
 	FGeometryShaderRHIParamRef,GeometryShader,
 	uint32,TextureIndex,
 	FTextureRHIParamRef,NewTexture,
@@ -1177,8 +1376,8 @@ DEFINE_RHIMETHOD_3(
 	);
 
 /** Set the shader resource view of a surface.  This is used for binding TextureMS parameter types that need a multi sampled view. */
-DEFINE_RHIMETHOD_3(
-	void,RHISetShaderTexture,
+DEFINE_RHIMETHOD_CMDLIST_3(
+	void,SetShaderTexture,
 	FPixelShaderRHIParamRef,PixelShader,
 	uint32,TextureIndex,
 	FTextureRHIParamRef,NewTexture,
@@ -1186,8 +1385,8 @@ DEFINE_RHIMETHOD_3(
 	);
 
 /** Set the shader resource view of a surface.  This is used for binding TextureMS parameter types that need a multi sampled view. */
-DEFINE_RHIMETHOD_3(
-	void,RHISetShaderTexture,
+DEFINE_RHIMETHOD_CMDLIST_3(
+	void,SetShaderTexture,
 	FComputeShaderRHIParamRef,PixelShader,
 	uint32,TextureIndex,
 	FTextureRHIParamRef,NewTexture,
@@ -1200,8 +1399,8 @@ DEFINE_RHIMETHOD_3(
  * @param SamplerIndex		The index of the sampler.
  * @param NewState			The new sampler state.
  */
-DEFINE_RHIMETHOD_3(
-	void,RHISetShaderSampler,
+DEFINE_RHIMETHOD_CMDLIST_3(
+	void,SetShaderSampler,
 	FComputeShaderRHIParamRef,ComputeShader,
 	uint32,SamplerIndex,
 	FSamplerStateRHIParamRef,NewState,
@@ -1214,8 +1413,8 @@ DEFINE_RHIMETHOD_3(
  * @param SamplerIndex		The index of the sampler.
  * @param NewState			The new sampler state.
  */
-DEFINE_RHIMETHOD_3(
-	void,RHISetShaderSampler,
+DEFINE_RHIMETHOD_CMDLIST_3(
+	void,SetShaderSampler,
 	FVertexShaderRHIParamRef,VertexShader,
 	uint32,SamplerIndex,
 	FSamplerStateRHIParamRef,NewState,
@@ -1228,8 +1427,8 @@ DEFINE_RHIMETHOD_3(
  * @param SamplerIndex		The index of the sampler.
  * @param NewState			The new sampler state.
  */
-DEFINE_RHIMETHOD_3(
-	void,RHISetShaderSampler,
+DEFINE_RHIMETHOD_CMDLIST_3(
+	void,SetShaderSampler,
 	FGeometryShaderRHIParamRef,GeometryShader,
 	uint32,SamplerIndex,
 	FSamplerStateRHIParamRef,NewState,
@@ -1242,8 +1441,8 @@ DEFINE_RHIMETHOD_3(
  * @param SamplerIndex		The index of the sampler.
  * @param NewState			The new sampler state.
  */
-DEFINE_RHIMETHOD_3(
-	void,RHISetShaderSampler,
+DEFINE_RHIMETHOD_CMDLIST_3(
+	void,SetShaderSampler,
 	FDomainShaderRHIParamRef,DomainShader,
 	uint32,SamplerIndex,
 	FSamplerStateRHIParamRef,NewState,
@@ -1256,8 +1455,8 @@ DEFINE_RHIMETHOD_3(
  * @param SamplerIndex		The index of the sampler.
  * @param NewState			The new sampler state.
  */
-DEFINE_RHIMETHOD_3(
-	void,RHISetShaderSampler,
+DEFINE_RHIMETHOD_CMDLIST_3(
+	void,SetShaderSampler,
 	FHullShaderRHIParamRef,HullShader,
 	uint32,SamplerIndex,
 	FSamplerStateRHIParamRef,NewState,
@@ -1270,8 +1469,8 @@ DEFINE_RHIMETHOD_3(
  * @param SamplerIndex		The index of the sampler.
  * @param NewState			The new sampler state.
  */
-DEFINE_RHIMETHOD_3(
-	void,RHISetShaderSampler,
+DEFINE_RHIMETHOD_CMDLIST_3(
+	void,SetShaderSampler,
 	FPixelShaderRHIParamRef,PixelShader,
 	uint32,SamplerIndex,
 	FSamplerStateRHIParamRef,NewState,
@@ -1279,8 +1478,8 @@ DEFINE_RHIMETHOD_3(
 	);
 
 /** Sets a compute shader UAV parameter. */
-DEFINE_RHIMETHOD_3(
-	 void,RHISetUAVParameter,
+DEFINE_RHIMETHOD_CMDLIST_3(
+	 void,SetUAVParameter,
 	 FComputeShaderRHIParamRef,ComputeShader,
 	 uint32,UAVIndex,
 	 FUnorderedAccessViewRHIParamRef,UAV,
@@ -1288,8 +1487,8 @@ DEFINE_RHIMETHOD_3(
 	 );
 
 /** Sets a compute shader UAV parameter and initial count */
-DEFINE_RHIMETHOD_4(
-	 void,RHISetUAVParameter,
+DEFINE_RHIMETHOD_CMDLIST_4(
+	 void,SetUAVParameter,
 	 FComputeShaderRHIParamRef,ComputeShader,
 	 uint32,UAVIndex,
 	 FUnorderedAccessViewRHIParamRef,UAV,
@@ -1298,8 +1497,8 @@ DEFINE_RHIMETHOD_4(
 	 );
 
 /** Sets a pixel shader resource view parameter. */
-DEFINE_RHIMETHOD_3(
-	void,RHISetShaderResourceViewParameter,
+DEFINE_RHIMETHOD_CMDLIST_3(
+	void,SetShaderResourceViewParameter,
 	FPixelShaderRHIParamRef,PixelShader,
 	uint32,SamplerIndex,
 	FShaderResourceViewRHIParamRef,SRV,
@@ -1307,8 +1506,8 @@ DEFINE_RHIMETHOD_3(
 	);
 
 /** Sets a vertex shader resource view parameter. */
-DEFINE_RHIMETHOD_3(
-	void,RHISetShaderResourceViewParameter,
+DEFINE_RHIMETHOD_CMDLIST_3(
+	void,SetShaderResourceViewParameter,
 	FVertexShaderRHIParamRef,VertexShader,
 	uint32,SamplerIndex,
 	FShaderResourceViewRHIParamRef,SRV,
@@ -1316,8 +1515,8 @@ DEFINE_RHIMETHOD_3(
 	);
 
 /** Sets a compute shader resource view parameter. */
-DEFINE_RHIMETHOD_3(
-	void,RHISetShaderResourceViewParameter,
+DEFINE_RHIMETHOD_CMDLIST_3(
+	void,SetShaderResourceViewParameter,
 	FComputeShaderRHIParamRef,ComputeShader,
 	uint32,SamplerIndex,
 	FShaderResourceViewRHIParamRef,SRV,
@@ -1325,8 +1524,8 @@ DEFINE_RHIMETHOD_3(
 	);
 
 /** Sets a hull shader resource view parameter. */
-DEFINE_RHIMETHOD_3(
-	void,RHISetShaderResourceViewParameter,
+DEFINE_RHIMETHOD_CMDLIST_3(
+	void,SetShaderResourceViewParameter,
 	FHullShaderRHIParamRef,HullShader,
 	uint32,SamplerIndex,
 	FShaderResourceViewRHIParamRef,SRV,
@@ -1334,8 +1533,8 @@ DEFINE_RHIMETHOD_3(
 	);
 
 /** Sets a domain shader resource view parameter. */
-DEFINE_RHIMETHOD_3(
-	void,RHISetShaderResourceViewParameter,
+DEFINE_RHIMETHOD_CMDLIST_3(
+	void,SetShaderResourceViewParameter,
 	FDomainShaderRHIParamRef,DomainShader,
 	uint32,SamplerIndex,
 	FShaderResourceViewRHIParamRef,SRV,
@@ -1343,64 +1542,64 @@ DEFINE_RHIMETHOD_3(
 	);
 
 /** Sets a geometry shader resource view parameter. */
-DEFINE_RHIMETHOD_3(
-	void,RHISetShaderResourceViewParameter,
+DEFINE_RHIMETHOD_CMDLIST_3(
+	void,SetShaderResourceViewParameter,
 	FGeometryShaderRHIParamRef,GeometryShader,
 	uint32,SamplerIndex,
 	FShaderResourceViewRHIParamRef,SRV,
 	,
 	);
 
-DEFINE_RHIMETHOD_3(
-	 void,RHISetShaderUniformBuffer,
+DEFINE_RHIMETHOD_CMDLIST_3(
+	 void,SetShaderUniformBuffer,
 	 FVertexShaderRHIParamRef,VertexShader,
 	 uint32,BufferIndex,
 	 FUniformBufferRHIParamRef,Buffer,
 	 ,
 	 );
 
-DEFINE_RHIMETHOD_3(
-	 void,RHISetShaderUniformBuffer,
+DEFINE_RHIMETHOD_CMDLIST_3(
+	 void,SetShaderUniformBuffer,
 	 FHullShaderRHIParamRef,HullShader,
 	 uint32,BufferIndex,
 	 FUniformBufferRHIParamRef,Buffer,
 	 ,
 	 );
 
-DEFINE_RHIMETHOD_3(
-	 void,RHISetShaderUniformBuffer,
+DEFINE_RHIMETHOD_CMDLIST_3(
+	 void,SetShaderUniformBuffer,
 	 FDomainShaderRHIParamRef,DomainShader,
 	 uint32,BufferIndex,
 	 FUniformBufferRHIParamRef,Buffer,
 	 ,
 	 );
 
-DEFINE_RHIMETHOD_3(
-	 void,RHISetShaderUniformBuffer,
+DEFINE_RHIMETHOD_CMDLIST_3(
+	 void,SetShaderUniformBuffer,
 	 FGeometryShaderRHIParamRef,GeometryShader,
 	 uint32,BufferIndex,
 	 FUniformBufferRHIParamRef,Buffer,
 	 ,
 	 );
 
-DEFINE_RHIMETHOD_3(
-	 void,RHISetShaderUniformBuffer,
+DEFINE_RHIMETHOD_CMDLIST_3(
+	 void,SetShaderUniformBuffer,
 	 FPixelShaderRHIParamRef,PixelShader,
 	 uint32,BufferIndex,
 	 FUniformBufferRHIParamRef,Buffer,
 	 ,
 	 );
 
-DEFINE_RHIMETHOD_3(
-	void,RHISetShaderUniformBuffer,
+DEFINE_RHIMETHOD_CMDLIST_3(
+	void,SetShaderUniformBuffer,
 	FComputeShaderRHIParamRef,ComputeShader,
 	uint32,BufferIndex,
 	FUniformBufferRHIParamRef,Buffer,
 	,
 	);
 
-DEFINE_RHIMETHOD_5(
-	void,RHISetShaderParameter,
+DEFINE_RHIMETHOD_CMDLIST_5(
+	void,SetShaderParameter,
 	FVertexShaderRHIParamRef,VertexShader,
 	uint32,BufferIndex,
 	uint32,BaseIndex,
@@ -1408,8 +1607,8 @@ DEFINE_RHIMETHOD_5(
 	const void*,NewValue,
 	,
 	);
-DEFINE_RHIMETHOD_5(
-	 void,RHISetShaderParameter,
+DEFINE_RHIMETHOD_CMDLIST_5(
+	 void,SetShaderParameter,
 	 FPixelShaderRHIParamRef,PixelShader,
 	 uint32,BufferIndex,
 	 uint32,BaseIndex,
@@ -1418,8 +1617,8 @@ DEFINE_RHIMETHOD_5(
 	 ,
 	 );
 
-DEFINE_RHIMETHOD_5(
-	 void,RHISetShaderParameter,
+DEFINE_RHIMETHOD_CMDLIST_5(
+	 void,SetShaderParameter,
 	 FHullShaderRHIParamRef,HullShader,
 	 uint32,BufferIndex,
 	 uint32,BaseIndex,
@@ -1427,8 +1626,8 @@ DEFINE_RHIMETHOD_5(
 	 const void*,NewValue,
 	 ,
 	 );
-DEFINE_RHIMETHOD_5(
-	 void,RHISetShaderParameter,
+DEFINE_RHIMETHOD_CMDLIST_5(
+	 void,SetShaderParameter,
 	 FDomainShaderRHIParamRef,DomainShader,
 	 uint32,BufferIndex,
 	 uint32,BaseIndex,
@@ -1436,8 +1635,8 @@ DEFINE_RHIMETHOD_5(
 	 const void*,NewValue,
 	 ,
 	 );
-DEFINE_RHIMETHOD_5(
-	 void,RHISetShaderParameter,
+DEFINE_RHIMETHOD_CMDLIST_5(
+	 void,SetShaderParameter,
 	 FGeometryShaderRHIParamRef,GeometryShader,
 	 uint32,BufferIndex,
 	 uint32,BaseIndex,
@@ -1445,8 +1644,8 @@ DEFINE_RHIMETHOD_5(
 	 const void*,NewValue,
 	 ,
 	 );
-DEFINE_RHIMETHOD_5(
-	void,RHISetShaderParameter,
+DEFINE_RHIMETHOD_CMDLIST_5(
+	void,SetShaderParameter,
 	FComputeShaderRHIParamRef,ComputeShader,
 	uint32,BufferIndex,
 	uint32,BaseIndex,
@@ -1456,21 +1655,21 @@ DEFINE_RHIMETHOD_5(
 	);
 
 // Output state.
-DEFINE_RHIMETHOD_2(
-	void,RHISetDepthStencilState,
+DEFINE_RHIMETHOD_CMDLIST_2(
+	void,SetDepthStencilState,
 	FDepthStencilStateRHIParamRef,NewState,
 	uint32,StencilRef,
 	,
 	);
 // Allows to set the blend state, parameter can be created with RHICreateBlendState()
-DEFINE_RHIMETHOD_2(
-	void,RHISetBlendState,
+DEFINE_RHIMETHOD_CMDLIST_2(
+	void,SetBlendState,
 	FBlendStateRHIParamRef,NewState,
 	const FLinearColor&,BlendFactor,
 	,
 	);
 DEFINE_RHIMETHOD_5(
-	void,RHISetRenderTargets,
+	void,SetRenderTargets,
 	uint32,NumSimultaneousRenderTargets,
 	const FRHIRenderTargetView*,NewRenderTargets,
 	FTextureRHIParamRef,NewDepthStencilTarget,
@@ -1479,7 +1678,7 @@ DEFINE_RHIMETHOD_5(
 	,
 	);
 DEFINE_RHIMETHOD_3(
-	void,RHIDiscardRenderTargets,
+	void,DiscardRenderTargets,
 	bool,Depth,
 	bool,Stencil,
 	uint32,ColorBitMask,
@@ -1487,8 +1686,8 @@ DEFINE_RHIMETHOD_3(
 	);
 
 // Primitive drawing.
-DEFINE_RHIMETHOD_4(
-	void,RHIDrawPrimitive,
+DEFINE_RHIMETHOD_CMDLIST_4(
+	void,DrawPrimitive,
 	uint32,PrimitiveType,
 	uint32,BaseVertexIndex,
 	uint32,NumPrimitives,
@@ -1496,7 +1695,7 @@ DEFINE_RHIMETHOD_4(
 	,
 	);
 DEFINE_RHIMETHOD_3(
-	void,RHIDrawPrimitiveIndirect,
+	void,DrawPrimitiveIndirect,
 	uint32,PrimitiveType,
 	FVertexBufferRHIParamRef,ArgumentBuffer,
 	uint32,ArgumentOffset,
@@ -1504,7 +1703,7 @@ DEFINE_RHIMETHOD_3(
 	);
 
 DEFINE_RHIMETHOD_5(
-	void,RHIDrawIndexedIndirect,
+	void,DrawIndexedIndirect,
 	FIndexBufferRHIParamRef,IndexBufferRHI,
 	uint32,PrimitiveType,
 	FStructuredBufferRHIParamRef,ArgumentsBufferRHI,
@@ -1514,8 +1713,8 @@ DEFINE_RHIMETHOD_5(
 	);
 
 // @param NumPrimitives need to be >0 
-DEFINE_RHIMETHOD_8(
-	void,RHIDrawIndexedPrimitive,
+DEFINE_RHIMETHOD_CMDLIST_8(
+	void,DrawIndexedPrimitive,
 	FIndexBufferRHIParamRef,IndexBuffer,
 	uint32,PrimitiveType,
 	int32,BaseVertexIndex,
@@ -1527,7 +1726,7 @@ DEFINE_RHIMETHOD_8(
 	,
 	);
 DEFINE_RHIMETHOD_4(
-	void,RHIDrawIndexedPrimitiveIndirect,
+	void,DrawIndexedPrimitiveIndirect,
 	uint32,PrimitiveType,
 	FIndexBufferRHIParamRef,IndexBuffer,
 	FVertexBufferRHIParamRef,ArgumentBuffer,
@@ -1545,7 +1744,7 @@ DEFINE_RHIMETHOD_4(
  * @param OutVertexData Reference to the allocated vertex memory
  */
 DEFINE_RHIMETHOD_5(
-	void,RHIBeginDrawPrimitiveUP,
+	void,BeginDrawPrimitiveUP,
 	uint32,PrimitiveType,
 	uint32,NumPrimitives,
 	uint32,NumVertices,
@@ -1558,7 +1757,7 @@ DEFINE_RHIMETHOD_5(
  * Draw a primitive using the vertex data populated since RHIBeginDrawPrimitiveUP and clean up any memory as needed
  */
 DEFINE_RHIMETHOD_0(
-	void,RHIEndDrawPrimitiveUP,
+	void,EndDrawPrimitiveUP,
 	,
 	);
 
@@ -1575,7 +1774,7 @@ DEFINE_RHIMETHOD_0(
  * @param OutIndexData Reference to the allocated index memory
  */
 DEFINE_RHIMETHOD_9(
-	void,RHIBeginDrawIndexedPrimitiveUP,
+	void,BeginDrawIndexedPrimitiveUP,
 	uint32,PrimitiveType,
 	uint32,NumPrimitives,
 	uint32,NumVertices,
@@ -1592,7 +1791,7 @@ DEFINE_RHIMETHOD_9(
  * Draw a primitive using the vertex and index data populated since RHIBeginDrawIndexedPrimitiveUP and clean up any memory as needed
  */
 DEFINE_RHIMETHOD_0(
-	void,RHIEndDrawIndexedPrimitiveUP,
+	void,EndDrawIndexedPrimitiveUP,
 	,
 	);
 
@@ -1602,7 +1801,7 @@ DEFINE_RHIMETHOD_0(
  * @param ExcludeRect within the viewport in pixels, is only a hint to optimize - if a fast clear can be done this is preferred
  */
 DEFINE_RHIMETHOD_7(
-	void,RHIClear,
+	void,Clear,
 	bool,bClearColor,
 	const FLinearColor&,Color,
 	bool,bClearDepth,
@@ -1618,7 +1817,7 @@ DEFINE_RHIMETHOD_7(
  * @param ExcludeRect within the viewport in pixels, is only a hint to optimize - if a fast clear can be done this is preferred
  */
 DEFINE_RHIMETHOD_8(
-	void,RHIClearMRT,
+	void,ClearMRT,
 	bool,bClearColor,
 	int32,NumClearColors,
 	const FLinearColor*,ColorArray,
@@ -1632,21 +1831,21 @@ DEFINE_RHIMETHOD_8(
 
 // Blocks the CPU until the GPU catches up and goes idle.
 DEFINE_RHIMETHOD_0(
-	void,RHIBlockUntilGPUIdle,
+	void,BlockUntilGPUIdle,
 	,
 	);
 
 // Operations to suspend title rendering and yield control to the system
 DEFINE_RHIMETHOD_0(
-	void,RHISuspendRendering,
+	void,SuspendRendering,
 	,
 	);
 DEFINE_RHIMETHOD_0(
-	void,RHIResumeRendering,
+	void,ResumeRendering,
 	,
 	);
 DEFINE_RHIMETHOD_0(
-	bool,RHIIsRenderingSuspended,
+	bool,IsRenderingSuspended,
 	return,return false;
 	);
 
@@ -1658,8 +1857,8 @@ DEFINE_RHIMETHOD_0(
  *
  *	@return	bool				true if successfully filled the array
  */
-DEFINE_RHIMETHOD_2(
-	bool,RHIGetAvailableResolutions,
+DEFINE_RHIMETHOD_GLOBAL_2(
+	bool,GetAvailableResolutions,
 	FScreenResolutionArray&,Resolutions,
 	bool,bIgnoreRefreshRate,
 	return,return false;
@@ -1670,8 +1869,8 @@ DEFINE_RHIMETHOD_2(
  * @param Width - Input: Desired resolution width in pixels. Output: A width that the platform supports.
  * @param Height - Input: Desired resolution height in pixels. Output: A height that the platform supports.
  */
-DEFINE_RHIMETHOD_2(
-	void,RHIGetSupportedResolution,
+DEFINE_RHIMETHOD_GLOBAL_2(
+	void,GetSupportedResolution,
 	uint32&,Width,
 	uint32&,Height,
 	,
@@ -1683,8 +1882,8 @@ DEFINE_RHIMETHOD_2(
  * @param Texture - the texture to update, must have been created with TexCreate_Virtual
  * @param FirstMip - the first mip that should be in memory
  */
-DEFINE_RHIMETHOD_2(
-	void,RHIVirtualTextureSetFirstMipInMemory,
+DEFINE_RHIMETHOD_GLOBAL_2(
+	void,VirtualTextureSetFirstMipInMemory,
 	FTexture2DRHIParamRef,Texture,
 	uint32,FirstMip,
 	,
@@ -1695,16 +1894,16 @@ DEFINE_RHIMETHOD_2(
  * @param Texture - the texture to update, must have been created with TexCreate_Virtual
  * @param FirstMip - the first mip that should be visible to the GPU
  */
-DEFINE_RHIMETHOD_2(
-	void,RHIVirtualTextureSetFirstMipVisible,
+DEFINE_RHIMETHOD_GLOBAL_2(
+	void,VirtualTextureSetFirstMipVisible,
 	FTexture2DRHIParamRef,Texture,
 	uint32,FirstMip,
 	,
 	);
 
-DEFINE_RHIMETHOD_1(
-	void,RHIExecuteCommandList,
-	FRHICommandList*,RHICmdList,
+DEFINE_RHIMETHOD_GLOBAL_1(
+	void,ExecuteCommandList,
+	FRHICommandList*,CmdList,
 	,
 	);
 
@@ -1716,7 +1915,7 @@ DEFINE_RHIMETHOD_1(
  *					The valid values for fMinDepth and fMaxDepth are such that 0 <= fMinDepth <= fMaxDepth <= 1
  */
 DEFINE_RHIMETHOD_3(
-	void,RHIEnableDepthBoundsTest,
+	void,EnableDepthBoundsTest,
 	bool,bEnable,
 	float,MinDepth,
 	float,MaxDepth,
@@ -1727,6 +1926,6 @@ DEFINE_RHIMETHOD_3(
  * Provides access to the native device. Generally this should be avoided but is useful for third party plugins.
  */
 DEFINE_RHIMETHOD_0(
-	void*,RHIGetNativeDevice,
+	void*,GetNativeDevice,
 	return,return NULL
 	);
