UE4 - follow up to user facing rhi refactor, mac

diff --git a/Engine/Source/Runtime/RHI/Public/RHIMethods.h b/Engine/Source/Runtime/RHI/Public/RHIMethods.h
--- a/Engine/Source/Runtime/RHI/Public/RHIMethods.h
+++ b/Engine/Source/Runtime/RHI/Public/RHIMethods.h
@@ -12,6 +12,47 @@
 #define GENERATE_VAX_FUNCTION_DECLARATIONS
 #endif
 
+#undef DEFINE_RHIMETHOD_0
+#undef DEFINE_RHIMETHOD_1
+#undef DEFINE_RHIMETHOD_2
+#undef DEFINE_RHIMETHOD_3
+#undef DEFINE_RHIMETHOD_4
+#undef DEFINE_RHIMETHOD_5
+#undef DEFINE_RHIMETHOD_6
+#undef DEFINE_RHIMETHOD_7
+#undef DEFINE_RHIMETHOD_8
+#undef DEFINE_RHIMETHOD_9
+#undef DEFINE_RHIMETHOD_CMDLIST_0
+#undef DEFINE_RHIMETHOD_CMDLIST_1
+#undef DEFINE_RHIMETHOD_CMDLIST_2
+#undef DEFINE_RHIMETHOD_CMDLIST_3
+#undef DEFINE_RHIMETHOD_CMDLIST_4
+#undef DEFINE_RHIMETHOD_CMDLIST_5
+#undef DEFINE_RHIMETHOD_CMDLIST_6
+#undef DEFINE_RHIMETHOD_CMDLIST_7
+#undef DEFINE_RHIMETHOD_CMDLIST_8
+#undef DEFINE_RHIMETHOD_CMDLIST_9
+#undef DEFINE_RHIMETHOD_GLOBAL_0
+#undef DEFINE_RHIMETHOD_GLOBAL_1
+#undef DEFINE_RHIMETHOD_GLOBAL_2
+#undef DEFINE_RHIMETHOD_GLOBAL_3
+#undef DEFINE_RHIMETHOD_GLOBAL_4
+#undef DEFINE_RHIMETHOD_GLOBAL_5
+#undef DEFINE_RHIMETHOD_GLOBAL_6
+#undef DEFINE_RHIMETHOD_GLOBAL_7
+#undef DEFINE_RHIMETHOD_GLOBAL_8
+#undef DEFINE_RHIMETHOD_GLOBAL_9
+#undef DEFINE_RHIMETHOD_GLOBALFLUSH_0
+#undef DEFINE_RHIMETHOD_GLOBALFLUSH_1
+#undef DEFINE_RHIMETHOD_GLOBALFLUSH_2
+#undef DEFINE_RHIMETHOD_GLOBALFLUSH_3
+#undef DEFINE_RHIMETHOD_GLOBALFLUSH_4
+#undef DEFINE_RHIMETHOD_GLOBALFLUSH_5
+#undef DEFINE_RHIMETHOD_GLOBALFLUSH_6
+#undef DEFINE_RHIMETHOD_GLOBALFLUSH_7
+#undef DEFINE_RHIMETHOD_GLOBALFLUSH_8
+#undef DEFINE_RHIMETHOD_GLOBALFLUSH_9
+
 #ifdef GENERATE_VAX_FUNCTION_DECLARATIONS
 
 	#define DEFINE_RHIMETHOD_0(ReturnType,MethodName,ReturnStatement,NullImplementation) \
@@ -98,8 +139,8 @@
 	#define DEFINE_RHIMETHOD_GLOBALFLUSH_9(ReturnType,MethodName,ParameterTypeA,ParameterNameA,ParameterTypeB,ParameterNameB,ParameterTypeC,ParameterNameC,ParameterTypeD,ParameterNameD,ParameterTypeE,ParameterNameE,ParameterTypeF,ParameterNameF,ParameterTypeG,ParameterNameG,ParameterTypeH,ParameterNameH,ParameterTypeI,ParameterNameI,ReturnStatement,NullImplementation) \
 		ReturnType RHI##MethodName(ParameterTypeA ParameterNameA, ParameterTypeB ParameterNameB, ParameterTypeC ParameterNameC, ParameterTypeD ParameterNameD, ParameterTypeE ParameterNameE, ParameterTypeF ParameterNameF, ParameterTypeG ParameterNameG, ParameterTypeH ParameterNameH, ParameterTypeI ParameterNameI);
 
-
 #else
+
 	#define DEFINE_RHIMETHOD_0(ReturnType,MethodName,ReturnStatement,NullImplementation) \
 		DEFINE_RHIMETHOD(ReturnType,MethodName,(),(),ReturnStatement,NullImplementation)
 	#define DEFINE_RHIMETHOD_1(ReturnType,MethodName,ParameterTypeA,ParameterNameA,ReturnStatement,NullImplementation) \
