UE4 - follow up to user-facing RHI refactor, xbox

diff --git a/Engine/Source/Runtime/RHI/Public/RHIMethods.h b/Engine/Source/Runtime/RHI/Public/RHIMethods.h
--- a/Engine/Source/Runtime/RHI/Public/RHIMethods.h
+++ b/Engine/Source/Runtime/RHI/Public/RHIMethods.h
@@ -1877,11 +1877,11 @@ DEFINE_RHIMETHOD_0(
 	);
 
 // Operations to suspend title rendering and yield control to the system
-DEFINE_RHIMETHOD_0(
+DEFINE_RHIMETHOD_GLOBAL_0(
 	void,SuspendRendering,
 	,
 	);
-DEFINE_RHIMETHOD_0(
+DEFINE_RHIMETHOD_GLOBAL_0(
 	void,ResumeRendering,
 	,
 	);
