UE4 - add ERCT_SetViewport, ERCT_SetScissorRect, ERCT_SetRenderTargets

diff --git a/Engine/Source/Runtime/EmptyRHI/Private/EmptyViewport.cpp b/Engine/Source/Runtime/EmptyRHI/Private/EmptyViewport.cpp
--- a/Engine/Source/Runtime/EmptyRHI/Private/EmptyViewport.cpp
+++ b/Engine/Source/Runtime/EmptyRHI/Private/EmptyViewport.cpp
@@ -47,7 +47,7 @@ void FEmptyDynamicRHI::RHIBeginDrawingViewport(FViewportRHIParamRef ViewportRHI,
 {
 	DYNAMIC_CAST_EMPTYRESOURCE(Viewport,Viewport);
 
-	SetRenderTarget(RHICmdList, RHIGetViewportBackBuffer(ViewportRHI), NULL);
+	RHISetRenderTarget(RHIGetViewportBackBuffer(ViewportRHI), NULL);
 }
 
 void FEmptyDynamicRHI::RHIEndDrawingViewport(FViewportRHIParamRef ViewportRHI,bool bPresent,bool bLockToVsync)
diff --git a/Engine/Source/Runtime/Engine/Private/PointLightSceneProxy.h b/Engine/Source/Runtime/Engine/Private/PointLightSceneProxy.h
--- a/Engine/Source/Runtime/Engine/Private/PointLightSceneProxy.h
+++ b/Engine/Source/Runtime/Engine/Private/PointLightSceneProxy.h
@@ -148,7 +148,7 @@ public:
 		return true;
 	}
 
-	virtual void SetScissorRect(FRHICommandListImmediate& RHICmdList, const FSceneView& View) const
+	virtual void SetScissorRect(FRHICommandList& RHICmdList, const FSceneView& View) const override
 	{
 		// Calculate a scissor rectangle for the light's radius.
 		if((GetLightToWorld().GetOrigin() - View.ViewMatrices.ViewOrigin).Size() > Radius)
diff --git a/Engine/Source/Runtime/Engine/Public/SceneManagement.h b/Engine/Source/Runtime/Engine/Public/SceneManagement.h
--- a/Engine/Source/Runtime/Engine/Public/SceneManagement.h
+++ b/Engine/Source/Runtime/Engine/Public/SceneManagement.h
@@ -693,7 +693,7 @@ public:
 	// @param OutCascadeSettings can be 0
 	virtual FSphere GetShadowSplitBounds(const class FSceneView& View, int32 SplitIndex, FShadowCascadeSettings* OutCascadeSettings) const { return FSphere(FVector::ZeroVector, 0); }
 
-	virtual void SetScissorRect(FRHICommandListImmediate& RHICmdList, const FSceneView& View) const
+	virtual void SetScissorRect(FRHICommandList& RHICmdList, const FSceneView& View) const
 	{
 	}
 
diff --git a/Engine/Source/Runtime/RHI/Private/RHICommandList.cpp b/Engine/Source/Runtime/RHI/Private/RHICommandList.cpp
--- a/Engine/Source/Runtime/RHI/Private/RHICommandList.cpp
+++ b/Engine/Source/Runtime/RHI/Private/RHICommandList.cpp
@@ -275,6 +275,47 @@ void FRHICommandListExecutor::ExecuteList(FRHICommandList& CmdList)
 				RHICmd->State->Release();
 			}
 			break;
+		case ERCT_SetViewport:
+			{
+				auto* RHICmd = Iter.NextCommand<FRHICommandSetViewport>();
+				if (!bOnlyTestMemAccess)
+				{
+					SetViewport_Internal(RHICmd->MinX, RHICmd->MinY, RHICmd->MinZ, RHICmd->MaxX, RHICmd->MaxY, RHICmd->MaxZ);
+				}
+			}
+			break;
+		case ERCT_SetScissorRect:
+			{
+				auto* RHICmd = Iter.NextCommand<FRHICommandSetScissorRect>();
+				if (!bOnlyTestMemAccess)
+				{
+					SetScissorRect_Internal(RHICmd->bEnable, RHICmd->MinX, RHICmd->MinY, RHICmd->MaxX, RHICmd->MaxY);
+				}
+			}
+			break;
+		case ERCT_SetRenderTargets:
+			{
+				auto* RHICmd = Iter.NextCommand<FRHICommandSetRenderTargets>();
+				if (!bOnlyTestMemAccess)
+				{
+					SetRenderTargets_Internal(
+						RHICmd->NewNumSimultaneousRenderTargets,
+						RHICmd->NewRenderTargetsRHI,
+						RHICmd->NewDepthStencilTargetRHI,
+						RHICmd->NewNumUAVs,
+						RHICmd->UAVs);
+				}
+				for (uint32 Index = 0; Index < RHICmd->NewNumSimultaneousRenderTargets; Index++)
+				{
+					RHICmd->NewRenderTargetsRHI[Index].Texture->Release();
+				}
+				RHICmd->NewDepthStencilTargetRHI->Release();
+				for (uint32 Index = 0; Index < RHICmd->NewNumUAVs; Index++)
+				{
+					RHICmd->UAVs[Index]->Release();
+				}
+			}
+			break;
 		default:
 			checkf(0, TEXT("Unknown RHI Command %d!"), Cmd->Type);
 		}
diff --git a/Engine/Source/Runtime/RHI/Public/RHICommandList.h b/Engine/Source/Runtime/RHI/Public/RHICommandList.h
--- a/Engine/Source/Runtime/RHI/Public/RHICommandList.h
+++ b/Engine/Source/Runtime/RHI/Public/RHICommandList.h
@@ -29,6 +29,9 @@ enum ERHICommandType
 	ERCT_SetShaderResourceViewParameter,
 	ERCT_SetUAVParameter,
 	ERCT_SetUAVParameter_IntialCount,
+	ERCT_SetViewport,
+	ERCT_SetScissorRect,
+	ERCT_SetRenderTargets,
 	ERCT_,
 };
 
@@ -52,6 +55,332 @@ struct FRHICommand
 	};
 };
 
+
+struct FRHICommandPerShader : public FRHICommand
+{
+	FRHIResource* Shader;
+	EShaderFrequency ShaderFrequency;
+
+	template <typename TShaderRHIParamRef>
+	FORCEINLINE void SetShader(TShaderRHIParamRef InShader)
+	{
+		Shader = InShader;
+		ShaderFrequency = (EShaderFrequency)TRHIShaderToEnum<TShaderRHIParamRef>::ShaderFrequency;
+	}
+};
+
+struct FRHICommandNopEndOfPage : public FRHICommand
+{
+	enum
+	{
+		IsEndOfPage = 1
+	};
+};
+
+struct FRHICommandNopBlob : public FRHICommand
+{
+	void* Data;
+	uint32 Size;
+	FORCEINLINE void Set(FRHICommandList* List, uint32 InSize, const void* InData)
+	{
+		Type = ERCT_NopBlob;
+		Size = InSize;
+
+		// Data is stored after 'this'
+		Data = &this[1];
+		FMemory::Memcpy(Data, InData, InSize);
+	}
+
+	uint32 ExtraSize() const
+	{
+		return Size;
+	}
+};
+
+struct FRHICommandSetRasterizerState : public FRHICommand
+{
+	FRasterizerStateRHIParamRef State;
+	FORCEINLINE void Set(FRasterizerStateRHIParamRef InState)
+	{
+		Type = ERCT_SetRasterizerState;
+		State = InState;
+	}
+};
+
+struct FRHICommandSetDepthStencilState : public FRHICommand
+{
+	FDepthStencilStateRHIParamRef State;
+	uint32 StencilRef;
+	FORCEINLINE void Set(FDepthStencilStateRHIParamRef InState, uint32 InStencilRef)
+	{
+		Type = ERCT_SetDepthStencilState;
+		State = InState;
+		StencilRef = InStencilRef;
+	}
+};
+
+struct FRHICommandSetShaderParameter : public FRHICommandPerShader
+{
+	const void* NewValue;
+	uint32 BufferIndex;
+	uint32 BaseIndex;
+	uint32 NumBytes;
+	template <typename TShaderRHIParamRef>
+	FORCEINLINE void Set(TShaderRHIParamRef InShader, uint32 InBufferIndex, uint32 InBaseIndex, uint32 InNumBytes, const void* InNewValue)
+	{
+		FRHICommandPerShader::SetShader<TShaderRHIParamRef>(InShader);
+		Type = ERCT_SetShaderParameter;
+		Shader = InShader;
+		BufferIndex = InBufferIndex;
+		BaseIndex = InBaseIndex;
+		NumBytes = InNumBytes;
+		NewValue = InNewValue;
+	}
+};
+
+struct FRHICommandSetShaderUniformBuffer : public FRHICommandPerShader
+{
+	uint32 BaseIndex;
+	FUniformBufferRHIParamRef UniformBuffer;
+	template <typename TShaderRHIParamRef>
+	FORCEINLINE void Set(TShaderRHIParamRef InShader, uint32 InBaseIndex, FUniformBufferRHIParamRef InUniformBuffer)
+	{
+		FRHICommandPerShader::SetShader<TShaderRHIParamRef>(InShader);
+		Type = ERCT_SetShaderUniformBuffer;
+		Shader = InShader;
+		BaseIndex = InBaseIndex;
+		UniformBuffer = InUniformBuffer;
+	}
+};
+
+struct FRHICommandSetShaderTexture : public FRHICommandPerShader
+{
+	uint32 TextureIndex;
+	FTextureRHIParamRef Texture;
+	template <typename TShaderRHIParamRef>
+	FORCEINLINE void Set(TShaderRHIParamRef InShader, uint32 InTextureIndex, FTextureRHIParamRef InTexture)
+	{
+		FRHICommandPerShader::SetShader<TShaderRHIParamRef>(InShader);
+		Type = ERCT_SetShaderTexture;
+		Shader = InShader;
+		TextureIndex = InTextureIndex;
+		Texture = InTexture;
+	}
+};
+
+struct FRHICommandSetShaderResourceViewParameter : public FRHICommandPerShader
+{
+	uint32 SamplerIndex;
+	FShaderResourceViewRHIParamRef SRV;
+	template <typename TShaderRHIParamRef>
+	FORCEINLINE void Set(TShaderRHIParamRef InShader, uint32 InSamplerIndex, FShaderResourceViewRHIParamRef InSRV)
+	{
+		FRHICommandPerShader::SetShader<TShaderRHIParamRef>(InShader);
+		Type = ERCT_SetShaderResourceViewParameter;
+		Shader = InShader;
+		SamplerIndex = InSamplerIndex;
+		SRV = InSRV;
+	}
+};
+
+struct FRHICommandSetUAVParameter : public FRHICommandPerShader
+{
+	uint32 UAVIndex;
+	FUnorderedAccessViewRHIParamRef UAV;
+	template <typename TShaderRHIParamRef>
+	FORCEINLINE void Set(TShaderRHIParamRef InShader, uint32 InUAVIndex, FUnorderedAccessViewRHIParamRef InUAV)
+	{
+		FRHICommandPerShader::SetShader<TShaderRHIParamRef>(InShader);
+		Type = ERCT_SetUAVParameter;
+		Shader = InShader;
+		UAVIndex = InUAVIndex;
+		UAV = InUAV;
+	}
+};
+
+struct FRHICommandSetUAVParameter_IntialCount : public FRHICommandPerShader
+{
+	uint32 UAVIndex;
+	FUnorderedAccessViewRHIParamRef UAV;
+	uint32 InitialCount;
+	template <typename TShaderRHIParamRef>
+	FORCEINLINE void Set(TShaderRHIParamRef InShader, uint32 InUAVIndex, FUnorderedAccessViewRHIParamRef InUAV, uint32 InInitialCount)
+	{
+		FRHICommandPerShader::SetShader<TShaderRHIParamRef>(InShader);
+		Type = ERCT_SetUAVParameter_IntialCount;
+		Shader = InShader;
+		UAVIndex = InUAVIndex;
+		UAV = InUAV;
+		InitialCount = InInitialCount;
+	}
+};
+
+struct FRHICommandSetShaderSampler : public FRHICommandPerShader
+{
+	uint32 SamplerIndex;
+	FSamplerStateRHIParamRef Sampler;
+	template <typename TShaderRHIParamRef>
+	FORCEINLINE void Set(TShaderRHIParamRef InShader, uint32 InSamplerIndex, FSamplerStateRHIParamRef InSampler)
+	{
+		FRHICommandPerShader::SetShader<TShaderRHIParamRef>(InShader);
+		Type = ERCT_SetShaderSampler;
+		Shader = InShader;
+		SamplerIndex = InSamplerIndex;
+		Sampler = InSampler;
+	}
+};
+
+struct FRHICommandDrawPrimitive : public FRHICommand
+{
+	uint32 PrimitiveType;
+	uint32 BaseVertexIndex;
+	uint32 NumPrimitives;
+	uint32 NumInstances;
+	FORCEINLINE void Set(uint32 InPrimitiveType, uint32 InBaseVertexIndex, uint32 InNumPrimitives, uint32 InNumInstances)
+	{
+		Type = ERCT_DrawPrimitive;
+		PrimitiveType = InPrimitiveType;
+		BaseVertexIndex = InBaseVertexIndex;
+		NumPrimitives = InNumPrimitives;
+		NumInstances = InNumInstances;
+	}
+};
+
+struct FRHICommandDrawIndexedPrimitive : public FRHICommand
+{
+	FIndexBufferRHIParamRef IndexBuffer;
+	uint32 PrimitiveType;
+	int32 BaseVertexIndex;
+	uint32 MinIndex;
+	uint32 NumVertices;
+	uint32 StartIndex;
+	uint32 NumPrimitives;
+	uint32 NumInstances;
+	FORCEINLINE void Set(FIndexBufferRHIParamRef InIndexBuffer, uint32 InPrimitiveType, int32 InBaseVertexIndex, uint32 InMinIndex, uint32 InNumVertices, uint32 InStartIndex, uint32 InNumPrimitives, uint32 InNumInstances)
+	{
+		Type = ERCT_DrawIndexedPrimitive;
+		IndexBuffer = InIndexBuffer;
+		PrimitiveType = InPrimitiveType;
+		BaseVertexIndex = InBaseVertexIndex;
+		MinIndex = InMinIndex;
+		NumVertices = InNumVertices;
+		StartIndex = InStartIndex;
+		NumPrimitives = InNumPrimitives;
+		NumInstances = InNumInstances;
+	}
+};
+
+struct FRHICommandSetBoundShaderState : public FRHICommand
+{
+	FBoundShaderStateRHIParamRef BoundShaderState;
+	FORCEINLINE void Set(FBoundShaderStateRHIParamRef InBoundShaderState)
+	{
+		Type = ERCT_SetBoundShaderState;
+		BoundShaderState = InBoundShaderState;
+	}
+};
+
+struct FRHICommandSetBlendState : public FRHICommand
+{
+	FBlendStateRHIParamRef State;
+	FLinearColor BlendFactor;
+	FORCEINLINE void Set(FBlendStateRHIParamRef InState, const FLinearColor& InBlendFactor)
+	{
+		Type = ERCT_SetBlendState;
+		State = InState;
+		BlendFactor = InBlendFactor;
+	}
+};
+
+struct FRHICommandSetStreamSource : public FRHICommand
+{
+	uint32 StreamIndex;
+	FVertexBufferRHIParamRef VertexBuffer;
+	uint32 Stride;
+	uint32 Offset;
+	FORCEINLINE void Set(uint32 InStreamIndex, FVertexBufferRHIParamRef InVertexBuffer, uint32 InStride, uint32 InOffset)
+	{
+		Type = ERCT_SetStreamSource;
+		StreamIndex = InStreamIndex;
+		VertexBuffer = InVertexBuffer;
+		Stride = InStride;
+		Offset = InOffset;
+	}
+};
+
+struct FRHICommandSetViewport : public FRHICommand
+{
+	uint32 MinX;
+	uint32 MinY;
+	float MinZ;
+	uint32 MaxX;
+	uint32 MaxY;
+	float MaxZ;
+	FORCEINLINE void Set(uint32 InMinX, uint32 InMinY, float InMinZ, uint32 InMaxX, uint32 InMaxY, float InMaxZ)
+	{
+		Type = ERCT_SetViewport;
+		MinX = InMinX;
+		MinY = InMinY;
+		MinZ = InMinZ;
+		MaxX = InMaxX;
+		MaxY = InMaxY;
+		MaxZ = InMaxZ;
+	}
+};
+
+struct FRHICommandSetScissorRect : public FRHICommand
+{
+	bool bEnable;
+	uint32 MinX;
+	uint32 MinY;
+	uint32 MaxX;
+	uint32 MaxY;
+	FORCEINLINE void Set(bool InbEnable, uint32 InMinX, uint32 InMinY, uint32 InMaxX, uint32 InMaxY)
+	{
+		Type = ERCT_SetScissorRect;
+		bEnable = InbEnable;
+		MinX = InMinX;
+		MinY = InMinY;
+		MaxX = InMaxX;
+		MaxY = InMaxY;
+	}
+};
+
+struct FRHICommandSetRenderTargets : public FRHICommand
+{
+	uint32 NewNumSimultaneousRenderTargets;
+	FRHIRenderTargetView NewRenderTargetsRHI[MaxSimultaneousRenderTargets];
+	FTextureRHIParamRef NewDepthStencilTargetRHI;
+	uint32 NewNumUAVs;
+	FUnorderedAccessViewRHIParamRef UAVs[MaxSimultaneousUAVs];
+
+	FORCEINLINE void Set(
+		uint32 InNewNumSimultaneousRenderTargets,
+		const FRHIRenderTargetView* InNewRenderTargetsRHI,
+		FTextureRHIParamRef InNewDepthStencilTargetRHI,
+		uint32 InNewNumUAVs,
+		const FUnorderedAccessViewRHIParamRef* InUAVs
+		)
+
+	{
+		Type = ERCT_SetRenderTargets;
+		check(NewNumSimultaneousRenderTargets <= MaxSimultaneousRenderTargets && NewNumUAVs <= MaxSimultaneousUAVs);
+		NewNumSimultaneousRenderTargets = InNewNumSimultaneousRenderTargets;
+		for (uint32 Index = 0; Index < NewNumSimultaneousRenderTargets; Index++)
+		{
+			NewRenderTargetsRHI[Index] = InNewRenderTargetsRHI[Index];
+		}
+		NewDepthStencilTargetRHI = InNewDepthStencilTargetRHI;
+		NewNumUAVs = InNewNumUAVs;
+		for (uint32 Index = 0; Index < NewNumUAVs; Index++)
+		{
+			UAVs[Index] = InUAVs[Index];
+		}
+	}
+};
+
+
+
 class RHI_API FRHICommandList : public FNoncopyable
 {
 private:
@@ -184,33 +513,257 @@ public:
 	inline void Flush();
 
 	template <typename TShaderRHIParamRef>
-	inline void SetShaderUniformBuffer(TShaderRHIParamRef Shader, uint32 BaseIndex, FUniformBufferRHIParamRef UniformBufferRHI);
+	FORCEINLINE void SetShaderUniformBuffer(TShaderRHIParamRef Shader, uint32 BaseIndex, FUniformBufferRHIParamRef UniformBuffer)
+	{
+		if (Bypass())
+		{
+			SetShaderUniformBuffer_Internal(Shader, BaseIndex, UniformBuffer);
+			return;
+		}
+		auto* Cmd = AddCommand<FRHICommandSetShaderUniformBuffer>();
+		Cmd->Set<TShaderRHIParamRef>(Shader, BaseIndex, UniformBuffer);
+		Shader->AddRef();
+		UniformBuffer->AddRef();
+	}
 
 	template <typename TShaderRHIParamRef>
-	inline void SetShaderParameter(TShaderRHIParamRef Shader, uint32 BufferIndex, uint32 BaseIndex, uint32 NumBytes, const void* Value, bool bValueInStack = false);
+	FORCEINLINE void SetShaderParameter(TShaderRHIParamRef Shader, uint32 BufferIndex, uint32 BaseIndex, uint32 NumBytes, const void* NewValue, bool bValueInStack = false)
+	{
+		if (Bypass())
+		{
+			SetShaderParameter_Internal(Shader, BufferIndex, BaseIndex, NumBytes, NewValue);
+			return;
+		}
+
+		if (bValueInStack)
+		{
+			auto* Blob = AddCommand<FRHICommandNopBlob>(NumBytes);
+			Blob->Set(this, NumBytes, NewValue);
+			NewValue = Blob->Data;
+		}
+		auto* Cmd = AddCommand<FRHICommandSetShaderParameter>();
+		Cmd->Set<TShaderRHIParamRef>(Shader, BufferIndex, BaseIndex, NumBytes, NewValue);
+		Shader->AddRef();
+	}
 
 	template <typename TShaderRHIParamRef>
-	inline void SetShaderTexture(TShaderRHIParamRef Shader, uint32 TextureIndex, FTextureRHIParamRef Texture);
+	FORCEINLINE void SetShaderTexture(TShaderRHIParamRef Shader, uint32 TextureIndex, FTextureRHIParamRef Texture)
+	{
+		if (Bypass())
+		{
+			SetShaderTexture_Internal(Shader, TextureIndex, Texture);
+			return;
+		}
+
+		auto* Cmd = AddCommand<FRHICommandSetShaderTexture>();
+		Cmd->Set<TShaderRHIParamRef>(Shader, TextureIndex, Texture);
+		Shader->AddRef();
+		Texture->AddRef();
+	}
 
 	template <typename TShaderRHIParamRef>
-	inline void SetShaderResourceViewParameter(TShaderRHIParamRef Shader, uint32 SamplerIndex, FShaderResourceViewRHIParamRef SRV);
+	FORCEINLINE void SetShaderResourceViewParameter(TShaderRHIParamRef Shader, uint32 SamplerIndex, FShaderResourceViewRHIParamRef SRV)
+	{
+		if (Bypass())
+		{
+			SetShaderResourceViewParameter_Internal(Shader, SamplerIndex, SRV);
+			return;
+		}
+		auto* Cmd = AddCommand<FRHICommandSetShaderResourceViewParameter>();
+		Cmd->Set<TShaderRHIParamRef>(Shader, SamplerIndex, SRV);
+		Shader->AddRef();
+		SRV->AddRef();
+	}
 
 	template <typename TShaderRHIParamRef>
-	inline void SetShaderSampler(TShaderRHIParamRef Shader, uint32 SamplerIndex, FSamplerStateRHIParamRef State);
+	FORCEINLINE void SetShaderSampler(TShaderRHIParamRef Shader, uint32 SamplerIndex, FSamplerStateRHIParamRef State)
+	{
+		if (Bypass())
+		{
+			SetShaderSampler_Internal(Shader, SamplerIndex, State);
+			return;
+		}
 
-	inline void SetUAVParameter(FComputeShaderRHIParamRef Shader, uint32 UAVIndex, FUnorderedAccessViewRHIParamRef UAV);
-	inline void SetUAVParameter(FComputeShaderRHIParamRef Shader, uint32 UAVIndex, FUnorderedAccessViewRHIParamRef UAV, uint32 InitialCount);
+		auto* Cmd = AddCommand<FRHICommandSetShaderSampler>();
+		Cmd->Set<TShaderRHIParamRef>(Shader, SamplerIndex, State);
+		Shader->AddRef();
+		State->AddRef();
+	}
 
-	inline void SetBoundShaderState(FBoundShaderStateRHIParamRef BoundShaderState);
+	FORCEINLINE void SetUAVParameter(FComputeShaderRHIParamRef Shader, uint32 UAVIndex, FUnorderedAccessViewRHIParamRef UAV)
+	{
+		if (Bypass())
+		{
+			SetUAVParameter_Internal(Shader, UAVIndex, UAV);
+			return;
+		}
+		auto* Cmd = AddCommand<FRHICommandSetUAVParameter>();
+		Cmd->Set<FComputeShaderRHIParamRef>(Shader, UAVIndex, UAV);
+		Shader->AddRef();
+		UAV->AddRef();
+	}
 
-	inline void SetRasterizerState(FRasterizerStateRHIParamRef State);
+	FORCEINLINE void SetUAVParameter(FComputeShaderRHIParamRef Shader, uint32 UAVIndex, FUnorderedAccessViewRHIParamRef UAV, uint32 InitialCount)
+	{
+		if (Bypass())
+		{
+			SetUAVParameter_Internal(Shader, UAVIndex, UAV, InitialCount);
+			return;
+		}
+		auto* Cmd = AddCommand<FRHICommandSetUAVParameter_IntialCount>();
+		Cmd->Set<FComputeShaderRHIParamRef>(Shader, UAVIndex, UAV, InitialCount);
+		Shader->AddRef();
+		UAV->AddRef();
+	}
 
-	inline void SetBlendState(FBlendStateRHIParamRef State, const FLinearColor& BlendFactor = FLinearColor::White);
+	FORCEINLINE void SetBoundShaderState(FBoundShaderStateRHIParamRef BoundShaderState)
+	{
+		if (Bypass())
+		{
+			SetBoundShaderState_Internal(BoundShaderState);
+			return;
+		}
+
+		auto* Cmd = AddCommand<FRHICommandSetBoundShaderState>();
+		Cmd->Set(BoundShaderState);
+		BoundShaderState->AddRef();
+	}
 
-	inline void DrawPrimitive(uint32 PrimitiveType, uint32 BaseVertexIndex, uint32 NumPrimitives, uint32 NumInstances);
-	inline void DrawIndexedPrimitive(FIndexBufferRHIParamRef IndexBuffer, uint32 PrimitiveType, int32 BaseVertexIndex, uint32 MinIndex, uint32 NumVertices, uint32 StartIndex, uint32 NumPrimitives, uint32 NumInstances);
-	inline void SetStreamSource(uint32 StreamIndex, FVertexBufferRHIParamRef VertexBuffer, uint32 Stride, uint32 Offset);
-	inline void SetDepthStencilState(FDepthStencilStateRHIParamRef NewStateRHI, uint32 StencilRef = 0);
+	FORCEINLINE void SetRasterizerState(FRasterizerStateRHIParamRef State)
+	{
+		if (Bypass())
+		{
+			SetRasterizerState_Internal(State);
+			return;
+		}
+
+		auto* Cmd = AddCommand<FRHICommandSetRasterizerState>();
+		Cmd->Set(State);
+		State->AddRef();
+	}
+
+	FORCEINLINE void SetBlendState(FBlendStateRHIParamRef State, const FLinearColor& BlendFactor = FLinearColor::White)
+	{
+		if (Bypass())
+		{
+			SetBlendState_Internal(State, BlendFactor);
+			return;
+		}
+
+		auto* Cmd = AddCommand<FRHICommandSetBlendState>();
+		Cmd->Set(State, BlendFactor);
+		State->AddRef();
+	}
+
+	FORCEINLINE void DrawPrimitive(uint32 PrimitiveType, uint32 BaseVertexIndex, uint32 NumPrimitives, uint32 NumInstances)
+	{
+		if (Bypass())
+		{
+			DrawPrimitive_Internal(PrimitiveType, BaseVertexIndex, NumPrimitives, NumInstances);
+			return;
+		}
+
+		auto* Cmd = AddCommand<FRHICommandDrawPrimitive>();
+		Cmd->Set(PrimitiveType, BaseVertexIndex, NumPrimitives, NumInstances);
+	}
+
+	FORCEINLINE void DrawIndexedPrimitive(FIndexBufferRHIParamRef IndexBuffer, uint32 PrimitiveType, int32 BaseVertexIndex, uint32 MinIndex, uint32 NumVertices, uint32 StartIndex, uint32 NumPrimitives, uint32 NumInstances)
+	{
+		if (Bypass())
+		{
+			DrawIndexedPrimitive_Internal(IndexBuffer, PrimitiveType, BaseVertexIndex, MinIndex, NumVertices, StartIndex, NumPrimitives, NumInstances);
+			return;
+		}
+
+		auto* Cmd = AddCommand<FRHICommandDrawIndexedPrimitive>();
+		Cmd->Set(IndexBuffer, PrimitiveType, BaseVertexIndex, MinIndex, NumVertices, StartIndex, NumPrimitives, NumInstances);
+		IndexBuffer->AddRef();
+	}
+
+	FORCEINLINE void SetStreamSource(uint32 StreamIndex, FVertexBufferRHIParamRef VertexBuffer, uint32 Stride, uint32 Offset)
+	{
+		if (Bypass())
+		{
+			SetStreamSource_Internal(StreamIndex, VertexBuffer, Stride, Offset);
+			return;
+		}
+
+		auto* Cmd = AddCommand<FRHICommandSetStreamSource>();
+		Cmd->Set(StreamIndex, VertexBuffer, Stride, Offset);
+		VertexBuffer->AddRef();
+	}
+
+	FORCEINLINE void SetDepthStencilState(FDepthStencilStateRHIParamRef NewStateRHI, uint32 StencilRef = 0)
+	{
+		if (Bypass())
+		{
+			SetDepthStencilState_Internal(NewStateRHI, StencilRef);
+			return;
+		}
+
+		auto* Cmd = AddCommand<FRHICommandSetDepthStencilState>();
+		Cmd->Set(NewStateRHI, StencilRef);
+		NewStateRHI->AddRef();
+	}
+
+	FORCEINLINE void SetViewport(uint32 MinX, uint32 MinY, float MinZ, uint32 MaxX, uint32 MaxY, float MaxZ)
+	{
+		if (Bypass())
+		{
+			SetViewport_Internal(MinX, MinY, MinZ, MaxX, MaxY, MaxZ);
+			return;
+		}
+		auto* Cmd = AddCommand<FRHICommandSetViewport>();
+		Cmd->Set(MinX, MinY, MinZ, MaxX, MaxY, MaxZ);
+	}
+
+	FORCEINLINE void SetScissorRect(bool bEnable, uint32 MinX, uint32 MinY, uint32 MaxX, uint32 MaxY)
+	{
+		if (Bypass())
+		{
+			SetScissorRect_Internal(bEnable, MinX, MinY, MaxX, MaxY);
+			return;
+		}
+		auto* Cmd = AddCommand<FRHICommandSetScissorRect>();
+		Cmd->Set(bEnable,  MinX, MinY, MaxX, MaxY);
+	}
+
+	FORCEINLINE void SetRenderTargets(
+		uint32 NewNumSimultaneousRenderTargets,
+		const FRHIRenderTargetView* NewRenderTargetsRHI,
+		FTextureRHIParamRef NewDepthStencilTargetRHI,
+		uint32 NewNumUAVs,
+		const FUnorderedAccessViewRHIParamRef* UAVs
+		)
+	{
+		if (Bypass())
+		{
+			SetRenderTargets_Internal(
+				NewNumSimultaneousRenderTargets,
+				NewRenderTargetsRHI,
+				NewDepthStencilTargetRHI,
+				NewNumUAVs,
+				UAVs);
+			return;
+		}
+		auto* Cmd = AddCommand<FRHICommandSetRenderTargets>();
+		Cmd->Set(
+			NewNumSimultaneousRenderTargets,
+			NewRenderTargetsRHI,
+			NewDepthStencilTargetRHI,
+			NewNumUAVs,
+			UAVs);
+		for (uint32 Index = 0; Index < NewNumSimultaneousRenderTargets; Index++)
+		{
+			NewRenderTargetsRHI[Index].Texture->AddRef();
+		}
+		NewDepthStencilTargetRHI->AddRef();
+		for (uint32 Index = 0; Index < NewNumUAVs; Index++)
+		{
+			UAVs[Index]->AddRef();
+		}
+
+	}
 
 	const SIZE_T GetUsedMemory() const;
 
diff --git a/Engine/Source/Runtime/RHI/Public/RHICommandList.inl b/Engine/Source/Runtime/RHI/Public/RHICommandList.inl
--- a/Engine/Source/Runtime/RHI/Public/RHICommandList.inl
+++ b/Engine/Source/Runtime/RHI/Public/RHICommandList.inl
@@ -8,458 +8,6 @@
 #define __RHICOMMANDLIST_INL__
 #include "RHICommandList.h"
 
-// No constructors are given for classes derived from FRHICommand as they are not created using
-// a classical 'new' for perf reasons.
-
-struct FRHICommandPerShader : public FRHICommand
-{
-	FRHIResource* Shader;
-	EShaderFrequency ShaderFrequency;
-
-	template <typename TShaderRHIParamRef>
-	FORCEINLINE void SetShader(TShaderRHIParamRef InShader)
-	{
-		Shader = InShader;
-		ShaderFrequency = (EShaderFrequency)TRHIShaderToEnum<TShaderRHIParamRef>::ShaderFrequency;
-	}
-};
-
-struct FRHICommandNopEndOfPage : public FRHICommand
-{
-	enum
-	{
-		IsEndOfPage = 1
-	};
-};
-
-struct FRHICommandNopBlob : public FRHICommand
-{
-	void* Data;
-	uint32 Size;
-	FORCEINLINE void Set(FRHICommandList* List, uint32 InSize, const void* InData)
-	{
-		Type = ERCT_NopBlob;
-		Size = InSize;
-
-		// Data is stored after 'this'
-		Data = &this[1];
-		FMemory::Memcpy(Data, InData, InSize);
-	}
-
-	uint32 ExtraSize() const
-	{
-		return Size;
-	}
-};
-
-struct FRHICommandSetRasterizerState : public FRHICommand
-{
-	FRasterizerStateRHIParamRef State;
-	FORCEINLINE void Set(FRasterizerStateRHIParamRef InState)
-	{
-		Type = ERCT_SetRasterizerState;
-		State = InState;
-	}
-};
-
-struct FRHICommandSetDepthStencilState : public FRHICommand
-{
-	FDepthStencilStateRHIParamRef State;
-	uint32 StencilRef;
-	FORCEINLINE void Set(FDepthStencilStateRHIParamRef InState, uint32 InStencilRef)
-	{
-		Type = ERCT_SetDepthStencilState;
-		State = InState;
-		StencilRef = InStencilRef;
-	}
-};
-
-struct FRHICommandSetShaderParameter : public FRHICommandPerShader
-{
-	const void* NewValue;
-	uint32 BufferIndex;
-	uint32 BaseIndex;
-	uint32 NumBytes;
-	template <typename TShaderRHIParamRef>
-	FORCEINLINE void Set(TShaderRHIParamRef InShader, uint32 InBufferIndex, uint32 InBaseIndex, uint32 InNumBytes, const void* InNewValue)
-	{
-		FRHICommandPerShader::SetShader<TShaderRHIParamRef>(InShader);
-		Type = ERCT_SetShaderParameter;
-		Shader = InShader;
-		BufferIndex = InBufferIndex;
-		BaseIndex = InBaseIndex;
-		NumBytes = InNumBytes;
-		NewValue = InNewValue;
-	}
-};
-
-struct FRHICommandSetShaderUniformBuffer : public FRHICommandPerShader
-{
-	uint32 BaseIndex;
-	FUniformBufferRHIParamRef UniformBuffer;
-	template <typename TShaderRHIParamRef>
-	FORCEINLINE void Set(TShaderRHIParamRef InShader, uint32 InBaseIndex, FUniformBufferRHIParamRef InUniformBuffer)
-	{
-		FRHICommandPerShader::SetShader<TShaderRHIParamRef>(InShader);
-		Type = ERCT_SetShaderUniformBuffer;
-		Shader = InShader;
-		BaseIndex = InBaseIndex;
-		UniformBuffer = InUniformBuffer;
-	}
-};
-
-struct FRHICommandSetShaderTexture : public FRHICommandPerShader
-{
-	uint32 TextureIndex;
-	FTextureRHIParamRef Texture;
-	template <typename TShaderRHIParamRef>
-	FORCEINLINE void Set(TShaderRHIParamRef InShader, uint32 InTextureIndex, FTextureRHIParamRef InTexture)
-	{
-		FRHICommandPerShader::SetShader<TShaderRHIParamRef>(InShader);
-		Type = ERCT_SetShaderTexture;
-		Shader = InShader;
-		TextureIndex = InTextureIndex;
-		Texture = InTexture;
-	}
-};
-
-struct FRHICommandSetShaderResourceViewParameter : public FRHICommandPerShader
-{
-	uint32 SamplerIndex;
-	FShaderResourceViewRHIParamRef SRV;
-	template <typename TShaderRHIParamRef>
-	FORCEINLINE void Set(TShaderRHIParamRef InShader, uint32 InSamplerIndex, FShaderResourceViewRHIParamRef InSRV)
-	{
-		FRHICommandPerShader::SetShader<TShaderRHIParamRef>(InShader);
-		Type = ERCT_SetShaderResourceViewParameter;
-		Shader = InShader;
-		SamplerIndex = InSamplerIndex;
-		SRV = InSRV;
-	}
-};
-
-struct FRHICommandSetUAVParameter : public FRHICommandPerShader
-{
-	uint32 UAVIndex;
-	FUnorderedAccessViewRHIParamRef UAV;
-	template <typename TShaderRHIParamRef>
-	FORCEINLINE void Set(TShaderRHIParamRef InShader, uint32 InUAVIndex, FUnorderedAccessViewRHIParamRef InUAV)
-	{
-		FRHICommandPerShader::SetShader<TShaderRHIParamRef>(InShader);
-		Type = ERCT_SetUAVParameter;
-		Shader = InShader;
-		UAVIndex = InUAVIndex;
-		UAV = InUAV;
-	}
-};
-
-struct FRHICommandSetUAVParameter_IntialCount : public FRHICommandPerShader
-{
-	uint32 UAVIndex;
-	FUnorderedAccessViewRHIParamRef UAV;
-	uint32 InitialCount;
-	template <typename TShaderRHIParamRef>
-	FORCEINLINE void Set(TShaderRHIParamRef InShader, uint32 InUAVIndex, FUnorderedAccessViewRHIParamRef InUAV, uint32 InInitialCount)
-	{
-		FRHICommandPerShader::SetShader<TShaderRHIParamRef>(InShader);
-		Type = ERCT_SetUAVParameter_IntialCount;
-		Shader = InShader;
-		UAVIndex = InUAVIndex;
-		UAV = InUAV;
-		InitialCount = InInitialCount;
-	}
-};
-
-struct FRHICommandSetShaderSampler : public FRHICommandPerShader
-{
-	uint32 SamplerIndex;
-	FSamplerStateRHIParamRef Sampler;
-	template <typename TShaderRHIParamRef>
-	FORCEINLINE void Set(TShaderRHIParamRef InShader, uint32 InSamplerIndex, FSamplerStateRHIParamRef InSampler)
-	{
-		FRHICommandPerShader::SetShader<TShaderRHIParamRef>(InShader);
-		Type = ERCT_SetShaderSampler;
-		Shader = InShader;
-		SamplerIndex = InSamplerIndex;
-		Sampler = InSampler;
-	}
-};
-
-struct FRHICommandDrawPrimitive : public FRHICommand
-{
-	uint32 PrimitiveType;
-	uint32 BaseVertexIndex;
-	uint32 NumPrimitives;
-	uint32 NumInstances;
-	FORCEINLINE void Set(uint32 InPrimitiveType, uint32 InBaseVertexIndex, uint32 InNumPrimitives, uint32 InNumInstances)
-	{
-		Type = ERCT_DrawPrimitive;
-		PrimitiveType = InPrimitiveType;
-		BaseVertexIndex = InBaseVertexIndex;
-		NumPrimitives = InNumPrimitives;
-		NumInstances = InNumInstances;
-	}
-};
-
-struct FRHICommandDrawIndexedPrimitive : public FRHICommand
-{
-	FIndexBufferRHIParamRef IndexBuffer;
-	uint32 PrimitiveType;
-	int32 BaseVertexIndex;
-	uint32 MinIndex;
-	uint32 NumVertices;
-	uint32 StartIndex;
-	uint32 NumPrimitives;
-	uint32 NumInstances;
-	FORCEINLINE void Set(FIndexBufferRHIParamRef InIndexBuffer, uint32 InPrimitiveType, int32 InBaseVertexIndex, uint32 InMinIndex, uint32 InNumVertices, uint32 InStartIndex, uint32 InNumPrimitives, uint32 InNumInstances)
-	{
-		Type = ERCT_DrawIndexedPrimitive;
-		IndexBuffer = InIndexBuffer;
-		PrimitiveType = InPrimitiveType;
-		BaseVertexIndex = InBaseVertexIndex;
-		MinIndex = InMinIndex;
-		NumVertices = InNumVertices;
-		StartIndex = InStartIndex;
-		NumPrimitives = InNumPrimitives;
-		NumInstances = InNumInstances;
-	}
-};
-
-struct FRHICommandSetBoundShaderState : public FRHICommand
-{
-	FBoundShaderStateRHIParamRef BoundShaderState;
-	FORCEINLINE void Set(FBoundShaderStateRHIParamRef InBoundShaderState)
-	{
-		Type = ERCT_SetBoundShaderState;
-		BoundShaderState = InBoundShaderState;
-	}
-};
-
-struct FRHICommandSetBlendState : public FRHICommand
-{
-	FBlendStateRHIParamRef State;
-	FLinearColor BlendFactor;
-	FORCEINLINE void Set(FBlendStateRHIParamRef InState, const FLinearColor& InBlendFactor)
-	{
-		Type = ERCT_SetBlendState;
-		State = InState;
-		BlendFactor = InBlendFactor;
-	}
-};
-
-struct FRHICommandSetStreamSource : public FRHICommand
-{
-	uint32 StreamIndex;
-	FVertexBufferRHIParamRef VertexBuffer;
-	uint32 Stride;
-	uint32 Offset;
-	FORCEINLINE void Set(uint32 InStreamIndex, FVertexBufferRHIParamRef InVertexBuffer, uint32 InStride, uint32 InOffset)
-	{
-		Type = ERCT_SetStreamSource;
-		StreamIndex = InStreamIndex;
-		VertexBuffer = InVertexBuffer;
-		Stride = InStride;
-		Offset = InOffset;
-	}
-};
-
-template <typename TShaderRHIParamRef>
-FORCEINLINE void FRHICommandList::SetShaderUniformBuffer(TShaderRHIParamRef Shader, uint32 BaseIndex, FUniformBufferRHIParamRef UniformBuffer)
-{
-	if (Bypass())
-	{
-		SetShaderUniformBuffer_Internal(Shader, BaseIndex, UniformBuffer);
-		return;
-	}
-	auto* Cmd = AddCommand<FRHICommandSetShaderUniformBuffer>();
-	Cmd->Set<TShaderRHIParamRef>(Shader, BaseIndex, UniformBuffer);
-	Shader->AddRef();
-	UniformBuffer->AddRef();
-}
-
-template <typename TShaderRHIParamRef>
-FORCEINLINE void FRHICommandList::SetShaderParameter(TShaderRHIParamRef Shader, uint32 BufferIndex, uint32 BaseIndex, uint32 NumBytes, const void* NewValue, bool bValueInStack)
-{
-	if (Bypass())
-	{
-		SetShaderParameter_Internal(Shader, BufferIndex, BaseIndex, NumBytes, NewValue);
-		return;
-	}
-
-	if (bValueInStack)
-	{
-		auto* Blob = AddCommand<FRHICommandNopBlob>(NumBytes);
-		Blob->Set(this, NumBytes, NewValue);
-		NewValue = Blob->Data;
-	}
-	auto* Cmd = AddCommand<FRHICommandSetShaderParameter>();
-	Cmd->Set<TShaderRHIParamRef>(Shader, BufferIndex, BaseIndex, NumBytes, NewValue);
-	Shader->AddRef();
-}
-
-template <typename TShaderRHIParamRef>
-FORCEINLINE void FRHICommandList::SetShaderTexture(TShaderRHIParamRef Shader, uint32 TextureIndex, FTextureRHIParamRef Texture)
-{
-	if (Bypass())
-	{
-		SetShaderTexture_Internal(Shader, TextureIndex, Texture);
-		return;
-	}
-
-	auto* Cmd = AddCommand<FRHICommandSetShaderTexture>();
-	Cmd->Set<TShaderRHIParamRef>(Shader, TextureIndex, Texture);
-	Shader->AddRef();
-	Texture->AddRef();
-}
-
-template <typename TShaderRHIParamRef>
-FORCEINLINE void FRHICommandList::SetShaderResourceViewParameter(TShaderRHIParamRef Shader, uint32 SamplerIndex, FShaderResourceViewRHIParamRef SRV)
-{
-	if (Bypass())
-	{
-		SetShaderResourceViewParameter_Internal(Shader, SamplerIndex, SRV);
-		return;
-	}
-	auto* Cmd = AddCommand<FRHICommandSetShaderResourceViewParameter>();
-	Cmd->Set<TShaderRHIParamRef>(Shader, SamplerIndex, SRV);
-	Shader->AddRef();
-	SRV->AddRef();
-}
-
-
-template <typename TShaderRHIParamRef>
-FORCEINLINE void FRHICommandList::SetShaderSampler(TShaderRHIParamRef Shader, uint32 SamplerIndex, FSamplerStateRHIParamRef State)
-{
-	if (Bypass())
-	{
-		SetShaderSampler_Internal(Shader, SamplerIndex, State);
-		return;
-	}
-
-	auto* Cmd = AddCommand<FRHICommandSetShaderSampler>();
-	Cmd->Set<TShaderRHIParamRef>(Shader, SamplerIndex, State);
-	Shader->AddRef();
-	State->AddRef();
-}
-
-FORCEINLINE void FRHICommandList::SetUAVParameter(FComputeShaderRHIParamRef Shader, uint32 UAVIndex, FUnorderedAccessViewRHIParamRef UAV)
-{
-	if (Bypass())
-	{
-		SetUAVParameter_Internal(Shader, UAVIndex, UAV);
-		return;
-	}
-	auto* Cmd = AddCommand<FRHICommandSetUAVParameter>();
-	Cmd->Set<FComputeShaderRHIParamRef>(Shader, UAVIndex, UAV);
-	Shader->AddRef();
-	UAV->AddRef();
-}
-
-FORCEINLINE void FRHICommandList::SetUAVParameter(FComputeShaderRHIParamRef Shader, uint32 UAVIndex, FUnorderedAccessViewRHIParamRef UAV, uint32 InitialCount)
-{
-	if (Bypass())
-	{
-		SetUAVParameter_Internal(Shader, UAVIndex, UAV, InitialCount);
-		return;
-	}
-	auto* Cmd = AddCommand<FRHICommandSetUAVParameter_IntialCount>();
-	Cmd->Set<FComputeShaderRHIParamRef>(Shader, UAVIndex, UAV, InitialCount);
-	Shader->AddRef();
-	UAV->AddRef();
-}
-
-FORCEINLINE void FRHICommandList::SetRasterizerState(FRasterizerStateRHIParamRef State)
-{
-	if (Bypass())
-	{
-		SetRasterizerState_Internal(State);
-		return;
-	}
-
-	auto* Cmd = AddCommand<FRHICommandSetRasterizerState>();
-	Cmd->Set(State);
-	State->AddRef();
-}
-
-FORCEINLINE void FRHICommandList::SetDepthStencilState(FDepthStencilStateRHIParamRef State, uint32 StencilRef)
-{
-	if (Bypass())
-	{
-		SetDepthStencilState_Internal(State, StencilRef);
-		return;
-	}
-
-	auto* Cmd = AddCommand<FRHICommandSetDepthStencilState>();
-	Cmd->Set(State, StencilRef);
-	State->AddRef();
-}
-
-FORCEINLINE void FRHICommandList::DrawPrimitive(uint32 PrimitiveType, uint32 BaseVertexIndex, uint32 NumPrimitives, uint32 NumInstances)
-{
-	if (Bypass())
-	{
-		DrawPrimitive_Internal(PrimitiveType, BaseVertexIndex, NumPrimitives, NumInstances);
-		return;
-	}
-
-	auto* Cmd = AddCommand<FRHICommandDrawPrimitive>();
-	Cmd->Set(PrimitiveType, BaseVertexIndex, NumPrimitives, NumInstances);
-}
-
-FORCEINLINE void FRHICommandList::DrawIndexedPrimitive(FIndexBufferRHIParamRef IndexBuffer, uint32 PrimitiveType, int32 BaseVertexIndex, uint32 MinIndex, uint32 NumVertices, uint32 StartIndex, uint32 NumPrimitives, uint32 NumInstances)
-{
-	if (Bypass())
-	{
-		DrawIndexedPrimitive_Internal(IndexBuffer, PrimitiveType, BaseVertexIndex, MinIndex, NumVertices, StartIndex, NumPrimitives, NumInstances);
-		return;
-	}
-
-	auto* Cmd = AddCommand<FRHICommandDrawIndexedPrimitive>();
-	Cmd->Set(IndexBuffer, PrimitiveType, BaseVertexIndex, MinIndex, NumVertices, StartIndex, NumPrimitives, NumInstances);
-	IndexBuffer->AddRef();
-}
-
-FORCEINLINE void FRHICommandList::SetBoundShaderState(FBoundShaderStateRHIParamRef BoundShaderState)
-{
-	if (Bypass())
-	{
-		SetBoundShaderState_Internal(BoundShaderState);
-		return;
-	}
-
-	auto* Cmd = AddCommand<FRHICommandSetBoundShaderState>();
-	Cmd->Set(BoundShaderState);
-	BoundShaderState->AddRef();
-}
-
-FORCEINLINE void FRHICommandList::SetBlendState(FBlendStateRHIParamRef State, const FLinearColor& BlendFactor)
-{
-	if (Bypass())
-	{
-		SetBlendState_Internal(State, BlendFactor);
-		return;
-	}
-
-	auto* Cmd = AddCommand<FRHICommandSetBlendState>();
-	Cmd->Set(State, BlendFactor);
-	State->AddRef();
-}
-
-FORCEINLINE void FRHICommandList::SetStreamSource(uint32 StreamIndex, FVertexBufferRHIParamRef VertexBuffer, uint32 Stride, uint32 Offset)
-{
-	if (Bypass())
-	{
-		SetStreamSource_Internal(StreamIndex, VertexBuffer, Stride, Offset);
-		return;
-	}
-
-	auto* Cmd = AddCommand<FRHICommandSetStreamSource>();
-	Cmd->Set(StreamIndex, VertexBuffer, Stride, Offset);
-	VertexBuffer->AddRef();
-}
-
-//--
-
 FORCEINLINE void FRHICommandList::Flush()
 {
 	if (HasCommands())
@@ -468,7 +16,6 @@ FORCEINLINE void FRHICommandList::Flush()
 	}
 }
 
-
 // Helper class for traversing a FRHICommandList
 class FRHICommandListIterator
 {
diff --git a/Engine/Source/Runtime/RHI/Public/RHIMethods.h b/Engine/Source/Runtime/RHI/Public/RHIMethods.h
--- a/Engine/Source/Runtime/RHI/Public/RHIMethods.h
+++ b/Engine/Source/Runtime/RHI/Public/RHIMethods.h
@@ -1345,7 +1345,7 @@ DEFINE_RHIMETHOD_CMDLIST_1(
 // @param MinY including like Win32 RECT
 // @param MaxX excluding like Win32 RECT
 // @param MaxY excluding like Win32 RECT
-DEFINE_RHIMETHOD_6(
+DEFINE_RHIMETHOD_CMDLIST_6(
 	void,SetViewport,
 	uint32,MinX,
 	uint32,MinY,
@@ -1359,7 +1359,7 @@ DEFINE_RHIMETHOD_6(
 // @param MinY including like Win32 RECT
 // @param MaxX excluding like Win32 RECT
 // @param MaxY excluding like Win32 RECT
-DEFINE_RHIMETHOD_5(
+DEFINE_RHIMETHOD_CMDLIST_5(
 	void,SetScissorRect,
 	bool,bEnable,
 	uint32,MinX,
@@ -1709,7 +1709,7 @@ DEFINE_RHIMETHOD_CMDLIST_2(
 	const FLinearColor&,BlendFactor,
 	,
 	);
-DEFINE_RHIMETHOD_5(
+DEFINE_RHIMETHOD_CMDLIST_5(
 	void,SetRenderTargets,
 	uint32,NumSimultaneousRenderTargets,
 	const FRHIRenderTargetView*,NewRenderTargets,
diff --git a/Engine/Source/Runtime/RHI/Public/RHIResources.h b/Engine/Source/Runtime/RHI/Public/RHIResources.h
--- a/Engine/Source/Runtime/RHI/Public/RHIResources.h
+++ b/Engine/Source/Runtime/RHI/Public/RHIResources.h
@@ -437,6 +437,12 @@ public:
 		ArraySliceIndex(-1)
 	{}
 
+	FRHIRenderTargetView(const FRHIRenderTargetView& Other) :
+		Texture(Other.Texture),
+		MipIndex(Other.MipIndex),
+		ArraySliceIndex(Other.ArraySliceIndex)
+	{}
+
 	FRHIRenderTargetView(FTextureRHIParamRef InTexture) :
 		Texture(InTexture),
 		MipIndex(0),
diff --git a/Engine/Source/Runtime/RHI/Public/RHIUtilities.h b/Engine/Source/Runtime/RHI/Public/RHIUtilities.h
--- a/Engine/Source/Runtime/RHI/Public/RHIUtilities.h
+++ b/Engine/Source/Runtime/RHI/Public/RHIUtilities.h
@@ -121,21 +121,21 @@ struct FRWBufferByteAddress
 };
 
 /** Helper for the common case of using a single color and depth render target. */
-inline void SetRenderTarget(FRHICommandListImmediate& RHICmdList, FTextureRHIParamRef NewRenderTarget, FTextureRHIParamRef NewDepthStencilTarget)
+inline void SetRenderTarget(FRHICommandList& RHICmdList, FTextureRHIParamRef NewRenderTarget, FTextureRHIParamRef NewDepthStencilTarget)
 {
 	FRHIRenderTargetView RTV(NewRenderTarget);
 	RHICmdList.SetRenderTargets(1, &RTV, NewDepthStencilTarget, 0, NULL);
 }
 
 /** Helper for the common case of using a single color and depth render target, with a mip index for the color target. */
-inline void SetRenderTarget(FRHICommandListImmediate& RHICmdList, FTextureRHIParamRef NewRenderTarget, int32 MipIndex, FTextureRHIParamRef NewDepthStencilTarget)
+inline void SetRenderTarget(FRHICommandList& RHICmdList, FTextureRHIParamRef NewRenderTarget, int32 MipIndex, FTextureRHIParamRef NewDepthStencilTarget)
 {
 	FRHIRenderTargetView RTV(NewRenderTarget, MipIndex, -1);
 	RHICmdList.SetRenderTargets(1, &RTV, NewDepthStencilTarget, 0, NULL);
 }
 
 /** Helper for the common case of using a single color and depth render target, with a mip index for the color target. */
-inline void SetRenderTarget(FRHICommandListImmediate& RHICmdList, FTextureRHIParamRef NewRenderTarget, int32 MipIndex, int32 ArraySliceIndex, FTextureRHIParamRef NewDepthStencilTarget)
+inline void SetRenderTarget(FRHICommandList& RHICmdList, FTextureRHIParamRef NewRenderTarget, int32 MipIndex, int32 ArraySliceIndex, FTextureRHIParamRef NewDepthStencilTarget)
 {
 	FRHIRenderTargetView RTV(NewRenderTarget, MipIndex, ArraySliceIndex);
 	RHICmdList.SetRenderTargets(1, &RTV, NewDepthStencilTarget, 0, NULL);
@@ -143,7 +143,7 @@ inline void SetRenderTarget(FRHICommandListImmediate& RHICmdList, FTextureRHIPar
 
 /** Helper that converts FTextureRHIParamRef's into FRHIRenderTargetView's. */
 inline void SetRenderTargets(
-	FRHICommandListImmediate& RHICmdList,
+	FRHICommandList& RHICmdList,
 	uint32 NewNumSimultaneousRenderTargets, 
 	const FTextureRHIParamRef* NewRenderTargetsRHI,
 	FTextureRHIParamRef NewDepthStencilTargetRHI,
diff --git a/Engine/Source/Runtime/Renderer/Private/TranslucentLighting.cpp b/Engine/Source/Runtime/Renderer/Private/TranslucentLighting.cpp
--- a/Engine/Source/Runtime/Renderer/Private/TranslucentLighting.cpp
+++ b/Engine/Source/Runtime/Renderer/Private/TranslucentLighting.cpp
@@ -1042,7 +1042,7 @@ public:
 TGlobalResource<FVolumeRasterizeVertexBuffer> GVolumeRasterizeVertexBuffer;
 
 /** Draws a quad per volume texture slice to the subregion of the volume texture specified. */
-void RasterizeToVolumeTexture(FRHICommandListImmediate& RHICmdList, FVolumeBounds VolumeBounds)
+void RasterizeToVolumeTexture(FRHICommandList& RHICmdList, FVolumeBounds VolumeBounds)
 {
 	// Setup the viewport to only render to the given bounds
 	RHICmdList.SetViewport(VolumeBounds.MinX, VolumeBounds.MinY, 0, VolumeBounds.MaxX, VolumeBounds.MaxY, 0);
diff --git a/Engine/Source/Runtime/Renderer/Private/TranslucentRendering.h b/Engine/Source/Runtime/Renderer/Private/TranslucentRendering.h
--- a/Engine/Source/Runtime/Renderer/Private/TranslucentRendering.h
+++ b/Engine/Source/Runtime/Renderer/Private/TranslucentRendering.h
@@ -240,4 +240,4 @@ private:
 	FShaderParameter MinZ;
 };
 
-extern void RasterizeToVolumeTexture(FRHICommandListImmediate& RHICmdList, FVolumeBounds VolumeBounds);
\ No newline at end of file
+extern void RasterizeToVolumeTexture(FRHICommandList& RHICmdList, FVolumeBounds VolumeBounds);
\ No newline at end of file
diff --git a/Engine/Source/Runtime/Windows/D3D11RHI/Private/D3D11Commands.cpp b/Engine/Source/Runtime/Windows/D3D11RHI/Private/D3D11Commands.cpp
--- a/Engine/Source/Runtime/Windows/D3D11RHI/Private/D3D11Commands.cpp
+++ b/Engine/Source/Runtime/Windows/D3D11RHI/Private/D3D11Commands.cpp
@@ -2082,6 +2082,7 @@ uint32 FD3D11DynamicRHI::RHIGetGPUFrameCycles()
 
 void FD3D11DynamicRHI::RHIExecuteCommandList(FRHICommandList* CmdList)
 {
+	check(0); // this path has gone stale and needs updated methods, starting at ERCT_SetScissorRect
 	FRHICommandListIterator Iter(*CmdList);
 	while (Iter.HasCommandsLeft())
 	{
@@ -2533,6 +2534,26 @@ void FD3D11DynamicRHI::RHIExecuteCommandList(FRHICommandList* CmdList)
 				RHICmd->State->Release();
 			}
 			break;
+		case ERCT_SetViewport:
+			{
+				auto* RHICmd = Iter.NextCommand<FRHICommandSetViewport>();
+				{
+					// These are the maximum viewport extents for D3D11. Exceeding them leads to badness.
+					check(RHICmd->MinX <= (uint32)D3D11_VIEWPORT_BOUNDS_MAX);
+					check(RHICmd->MinY <= (uint32)D3D11_VIEWPORT_BOUNDS_MAX);
+					check(RHICmd->MaxX <= (uint32)D3D11_VIEWPORT_BOUNDS_MAX);
+					check(RHICmd->MaxY <= (uint32)D3D11_VIEWPORT_BOUNDS_MAX);
+
+					D3D11_VIEWPORT Viewport = { RHICmd->MinX, RHICmd->MinY, RHICmd->MaxX - RHICmd->MinX, RHICmd->MaxY - RHICmd->MinY, RHICmd->MinZ, RHICmd->MaxZ };
+					//avoid setting a 0 extent viewport, which the debug runtime doesn't like
+					if (Viewport.Width > 0 && Viewport.Height > 0)
+					{
+						StateCache.SetViewport(Viewport);
+						SetScissorRectIfRequiredWhenSettingViewport(RHICmd->MinX, RHICmd->MinY, RHICmd->MaxX, RHICmd->MaxY);
+					}
+				}
+			}
+			break;
 		default:
 			checkf(0, TEXT("Unimplemented RHICmd %d!"), Cmd->Type);
 		}
