UE4 - more rhi command list work, adding Draw*UP and fixing bugs

diff --git a/Engine/Source/Runtime/Launch/Private/LaunchEngineLoop.cpp b/Engine/Source/Runtime/Launch/Private/LaunchEngineLoop.cpp
--- a/Engine/Source/Runtime/Launch/Private/LaunchEngineLoop.cpp
+++ b/Engine/Source/Runtime/Launch/Private/LaunchEngineLoop.cpp
@@ -2006,6 +2006,7 @@ void FEngineLoop::Tick()
 		ENQUEUE_UNIQUE_RENDER_COMMAND(
 				BeginFrame,
 			{
+				GRHICommandList.LatchBypass();
 				GFrameNumberRenderThread++;
 				GDynamicRHI->PushEvent(*FString::Printf(TEXT("Frame%d"),GFrameNumberRenderThread));
 				RHICmdList.BeginFrame();
diff --git a/Engine/Source/Runtime/RHI/Private/RHICommandList.cpp b/Engine/Source/Runtime/RHI/Private/RHICommandList.cpp
--- a/Engine/Source/Runtime/RHI/Private/RHICommandList.cpp
+++ b/Engine/Source/Runtime/RHI/Private/RHICommandList.cpp
@@ -198,7 +198,10 @@ void FRHICommandListExecutor::ExecuteList(FRHICommandList& CmdList)
 					}
 				}
 				RHICmd->Shader->Release();
-				RHICmd->UAV->Release();
+				if (RHICmd->UAV)
+				{
+					RHICmd->UAV->Release();
+				}
 			}
 			break;
 		case ERCT_SetUAVParameter_IntialCount:
@@ -213,7 +216,10 @@ void FRHICommandListExecutor::ExecuteList(FRHICommandList& CmdList)
 					}
 				}
 				RHICmd->Shader->Release();
-				RHICmd->UAV->Release();
+				if (RHICmd->UAV)
+				{
+					RHICmd->UAV->Release();
+				}
 			}
 			break;
 		case ERCT_DrawPrimitive:
@@ -307,15 +313,64 @@ void FRHICommandListExecutor::ExecuteList(FRHICommandList& CmdList)
 				}
 				for (uint32 Index = 0; Index < RHICmd->NewNumSimultaneousRenderTargets; Index++)
 				{
-					RHICmd->NewRenderTargetsRHI[Index].Texture->Release();
+					if (RHICmd->NewRenderTargetsRHI[Index].Texture)
+					{
+						RHICmd->NewRenderTargetsRHI[Index].Texture->Release();
+					}
+				}
+				if (RHICmd->NewDepthStencilTargetRHI)
+				{
+					RHICmd->NewDepthStencilTargetRHI->Release();
 				}
-				RHICmd->NewDepthStencilTargetRHI->Release();
 				for (uint32 Index = 0; Index < RHICmd->NewNumUAVs; Index++)
 				{
-					RHICmd->UAVs[Index]->Release();
+					if (RHICmd->UAVs[Index])
+					{
+						RHICmd->UAVs[Index]->Release();
+					}
+				}
+			}
+			break;
+		case ERCT_EndDrawPrimitiveUP:
+			{
+				auto* RHICmd = Iter.NextCommand<FRHICommandEndDrawPrimitiveUP>();
+				if (!bOnlyTestMemAccess)
+				{
+					void* Buffer = NULL;
+					BeginDrawPrimitiveUP_Internal(RHICmd->PrimitiveType, RHICmd->NumPrimitives, RHICmd->NumVertices, RHICmd->VertexDataStride, Buffer);
+					FMemory::Memcpy(Buffer, RHICmd->OutVertexData, RHICmd->NumVertices * RHICmd->VertexDataStride);
+					EndDrawPrimitiveUP_Internal();
+				}
+				FMemory::Free(RHICmd->OutVertexData);
+			}
+			break;
+		case ERCT_EndDrawIndexedPrimitiveUP:
+			{
+				auto* RHICmd = Iter.NextCommand<FRHICommandEndDrawIndexedPrimitiveUP>();
+				if (!bOnlyTestMemAccess)
+				{
+					void* VertexBuffer = NULL;
+					void* IndexBuffer = NULL;
+					BeginDrawIndexedPrimitiveUP_Internal(
+						RHICmd->PrimitiveType,
+						RHICmd->NumPrimitives,
+						RHICmd->NumVertices,
+						RHICmd->VertexDataStride,
+						VertexBuffer,
+						RHICmd->MinVertexIndex,
+						RHICmd->NumIndices,
+						RHICmd->IndexDataStride,
+						IndexBuffer);
+					FMemory::Memcpy(VertexBuffer, RHICmd->OutVertexData, RHICmd->NumVertices * RHICmd->VertexDataStride);
+					FMemory::Memcpy(IndexBuffer, RHICmd->OutIndexData, RHICmd->NumIndices * RHICmd->IndexDataStride);
+					EndDrawIndexedPrimitiveUP_Internal();
+
 				}
+				FMemory::Free(RHICmd->OutVertexData);
+				FMemory::Free(RHICmd->OutIndexData);
 			}
 			break;
+
 		default:
 			checkf(0, TEXT("Unknown RHI Command %d!"), Cmd->Type);
 		}
@@ -339,7 +394,8 @@ void FRHICommandListExecutor::ExecuteList(FRHICommandList& CmdList)
 
 	if (!bSkipRHICmdList)
 	{
-		if (&CmdList != &GetImmediateCommandList())
+		if (&CmdList != &GetImmediateCommandList() 
+			&& !GetImmediateCommandList().bExecuting) // don't flush if this is a recursive call and we are already executing the immediate command list
 		{
 			GetImmediateCommandList().Flush();
 		}
@@ -362,17 +418,23 @@ void FRHICommandListExecutor::ExecuteList(FRHICommandList& CmdList)
 			}
 		}
 	}
-
 	INC_DWORD_STAT_BY(STAT_RHICmdListCount, CmdList.NumCommands);
 #if !(UE_BUILD_SHIPPING || UE_BUILD_TEST)
 	if (&CmdList == &GetImmediateCommandList())
 	{
-		bLatchedBypass = (CVarRHICmdBypass.GetValueOnRenderThread() >= 1);
+		LatchBypass();
 	}
 #endif
 	CmdList.Reset();
 }
 
+void FRHICommandListExecutor::LatchBypass()
+{
+#if !(UE_BUILD_SHIPPING || UE_BUILD_TEST)
+	bLatchedBypass = (CVarRHICmdBypass.GetValueOnRenderThread() >= 1);
+#endif
+}
+
 FRHICommandList::FMemManager::FMemManager() :
 	FirstPage(nullptr),
 	LastPage(nullptr)
diff --git a/Engine/Source/Runtime/RHI/Public/RHICommandList.h b/Engine/Source/Runtime/RHI/Public/RHICommandList.h
--- a/Engine/Source/Runtime/RHI/Public/RHICommandList.h
+++ b/Engine/Source/Runtime/RHI/Public/RHICommandList.h
@@ -32,6 +32,8 @@ enum ERHICommandType
 	ERCT_SetViewport,
 	ERCT_SetScissorRect,
 	ERCT_SetRenderTargets,
+	ERCT_EndDrawPrimitiveUP,
+	ERCT_EndDrawIndexedPrimitiveUP,
 	ERCT_,
 };
 
@@ -62,7 +64,7 @@ struct FRHICommandPerShader : public FRHICommand
 	EShaderFrequency ShaderFrequency;
 
 	template <typename TShaderRHIParamRef>
-	FORCEINLINE void SetShader(TShaderRHIParamRef InShader)
+	FORCEINLINE_DEBUGGABLE void SetShader(TShaderRHIParamRef InShader)
 	{
 		Shader = InShader;
 		ShaderFrequency = (EShaderFrequency)TRHIShaderToEnum<TShaderRHIParamRef>::ShaderFrequency;
@@ -81,7 +83,7 @@ struct FRHICommandNopBlob : public FRHICommand
 {
 	void* Data;
 	uint32 Size;
-	FORCEINLINE void Set(FRHICommandList* List, uint32 InSize, const void* InData)
+	FORCEINLINE_DEBUGGABLE void Set(FRHICommandList* List, uint32 InSize, const void* InData)
 	{
 		Type = ERCT_NopBlob;
 		Size = InSize;
@@ -100,7 +102,7 @@ struct FRHICommandNopBlob : public FRHICommand
 struct FRHICommandSetRasterizerState : public FRHICommand
 {
 	FRasterizerStateRHIParamRef State;
-	FORCEINLINE void Set(FRasterizerStateRHIParamRef InState)
+	FORCEINLINE_DEBUGGABLE void Set(FRasterizerStateRHIParamRef InState)
 	{
 		Type = ERCT_SetRasterizerState;
 		State = InState;
@@ -111,7 +113,7 @@ struct FRHICommandSetDepthStencilState : public FRHICommand
 {
 	FDepthStencilStateRHIParamRef State;
 	uint32 StencilRef;
-	FORCEINLINE void Set(FDepthStencilStateRHIParamRef InState, uint32 InStencilRef)
+	FORCEINLINE_DEBUGGABLE void Set(FDepthStencilStateRHIParamRef InState, uint32 InStencilRef)
 	{
 		Type = ERCT_SetDepthStencilState;
 		State = InState;
@@ -126,7 +128,7 @@ struct FRHICommandSetShaderParameter : public FRHICommandPerShader
 	uint32 BaseIndex;
 	uint32 NumBytes;
 	template <typename TShaderRHIParamRef>
-	FORCEINLINE void Set(TShaderRHIParamRef InShader, uint32 InBufferIndex, uint32 InBaseIndex, uint32 InNumBytes, const void* InNewValue)
+	FORCEINLINE_DEBUGGABLE void Set(TShaderRHIParamRef InShader, uint32 InBufferIndex, uint32 InBaseIndex, uint32 InNumBytes, const void* InNewValue)
 	{
 		FRHICommandPerShader::SetShader<TShaderRHIParamRef>(InShader);
 		Type = ERCT_SetShaderParameter;
@@ -143,7 +145,7 @@ struct FRHICommandSetShaderUniformBuffer : public FRHICommandPerShader
 	uint32 BaseIndex;
 	FUniformBufferRHIParamRef UniformBuffer;
 	template <typename TShaderRHIParamRef>
-	FORCEINLINE void Set(TShaderRHIParamRef InShader, uint32 InBaseIndex, FUniformBufferRHIParamRef InUniformBuffer)
+	FORCEINLINE_DEBUGGABLE void Set(TShaderRHIParamRef InShader, uint32 InBaseIndex, FUniformBufferRHIParamRef InUniformBuffer)
 	{
 		FRHICommandPerShader::SetShader<TShaderRHIParamRef>(InShader);
 		Type = ERCT_SetShaderUniformBuffer;
@@ -158,7 +160,7 @@ struct FRHICommandSetShaderTexture : public FRHICommandPerShader
 	uint32 TextureIndex;
 	FTextureRHIParamRef Texture;
 	template <typename TShaderRHIParamRef>
-	FORCEINLINE void Set(TShaderRHIParamRef InShader, uint32 InTextureIndex, FTextureRHIParamRef InTexture)
+	FORCEINLINE_DEBUGGABLE void Set(TShaderRHIParamRef InShader, uint32 InTextureIndex, FTextureRHIParamRef InTexture)
 	{
 		FRHICommandPerShader::SetShader<TShaderRHIParamRef>(InShader);
 		Type = ERCT_SetShaderTexture;
@@ -173,7 +175,7 @@ struct FRHICommandSetShaderResourceViewParameter : public FRHICommandPerShader
 	uint32 SamplerIndex;
 	FShaderResourceViewRHIParamRef SRV;
 	template <typename TShaderRHIParamRef>
-	FORCEINLINE void Set(TShaderRHIParamRef InShader, uint32 InSamplerIndex, FShaderResourceViewRHIParamRef InSRV)
+	FORCEINLINE_DEBUGGABLE void Set(TShaderRHIParamRef InShader, uint32 InSamplerIndex, FShaderResourceViewRHIParamRef InSRV)
 	{
 		FRHICommandPerShader::SetShader<TShaderRHIParamRef>(InShader);
 		Type = ERCT_SetShaderResourceViewParameter;
@@ -188,7 +190,7 @@ struct FRHICommandSetUAVParameter : public FRHICommandPerShader
 	uint32 UAVIndex;
 	FUnorderedAccessViewRHIParamRef UAV;
 	template <typename TShaderRHIParamRef>
-	FORCEINLINE void Set(TShaderRHIParamRef InShader, uint32 InUAVIndex, FUnorderedAccessViewRHIParamRef InUAV)
+	FORCEINLINE_DEBUGGABLE void Set(TShaderRHIParamRef InShader, uint32 InUAVIndex, FUnorderedAccessViewRHIParamRef InUAV)
 	{
 		FRHICommandPerShader::SetShader<TShaderRHIParamRef>(InShader);
 		Type = ERCT_SetUAVParameter;
@@ -204,7 +206,7 @@ struct FRHICommandSetUAVParameter_IntialCount : public FRHICommandPerShader
 	FUnorderedAccessViewRHIParamRef UAV;
 	uint32 InitialCount;
 	template <typename TShaderRHIParamRef>
-	FORCEINLINE void Set(TShaderRHIParamRef InShader, uint32 InUAVIndex, FUnorderedAccessViewRHIParamRef InUAV, uint32 InInitialCount)
+	FORCEINLINE_DEBUGGABLE void Set(TShaderRHIParamRef InShader, uint32 InUAVIndex, FUnorderedAccessViewRHIParamRef InUAV, uint32 InInitialCount)
 	{
 		FRHICommandPerShader::SetShader<TShaderRHIParamRef>(InShader);
 		Type = ERCT_SetUAVParameter_IntialCount;
@@ -220,7 +222,7 @@ struct FRHICommandSetShaderSampler : public FRHICommandPerShader
 	uint32 SamplerIndex;
 	FSamplerStateRHIParamRef Sampler;
 	template <typename TShaderRHIParamRef>
-	FORCEINLINE void Set(TShaderRHIParamRef InShader, uint32 InSamplerIndex, FSamplerStateRHIParamRef InSampler)
+	FORCEINLINE_DEBUGGABLE void Set(TShaderRHIParamRef InShader, uint32 InSamplerIndex, FSamplerStateRHIParamRef InSampler)
 	{
 		FRHICommandPerShader::SetShader<TShaderRHIParamRef>(InShader);
 		Type = ERCT_SetShaderSampler;
@@ -236,7 +238,7 @@ struct FRHICommandDrawPrimitive : public FRHICommand
 	uint32 BaseVertexIndex;
 	uint32 NumPrimitives;
 	uint32 NumInstances;
-	FORCEINLINE void Set(uint32 InPrimitiveType, uint32 InBaseVertexIndex, uint32 InNumPrimitives, uint32 InNumInstances)
+	FORCEINLINE_DEBUGGABLE void Set(uint32 InPrimitiveType, uint32 InBaseVertexIndex, uint32 InNumPrimitives, uint32 InNumInstances)
 	{
 		Type = ERCT_DrawPrimitive;
 		PrimitiveType = InPrimitiveType;
@@ -256,7 +258,7 @@ struct FRHICommandDrawIndexedPrimitive : public FRHICommand
 	uint32 StartIndex;
 	uint32 NumPrimitives;
 	uint32 NumInstances;
-	FORCEINLINE void Set(FIndexBufferRHIParamRef InIndexBuffer, uint32 InPrimitiveType, int32 InBaseVertexIndex, uint32 InMinIndex, uint32 InNumVertices, uint32 InStartIndex, uint32 InNumPrimitives, uint32 InNumInstances)
+	FORCEINLINE_DEBUGGABLE void Set(FIndexBufferRHIParamRef InIndexBuffer, uint32 InPrimitiveType, int32 InBaseVertexIndex, uint32 InMinIndex, uint32 InNumVertices, uint32 InStartIndex, uint32 InNumPrimitives, uint32 InNumInstances)
 	{
 		Type = ERCT_DrawIndexedPrimitive;
 		IndexBuffer = InIndexBuffer;
@@ -273,7 +275,7 @@ struct FRHICommandDrawIndexedPrimitive : public FRHICommand
 struct FRHICommandSetBoundShaderState : public FRHICommand
 {
 	FBoundShaderStateRHIParamRef BoundShaderState;
-	FORCEINLINE void Set(FBoundShaderStateRHIParamRef InBoundShaderState)
+	FORCEINLINE_DEBUGGABLE void Set(FBoundShaderStateRHIParamRef InBoundShaderState)
 	{
 		Type = ERCT_SetBoundShaderState;
 		BoundShaderState = InBoundShaderState;
@@ -284,7 +286,7 @@ struct FRHICommandSetBlendState : public FRHICommand
 {
 	FBlendStateRHIParamRef State;
 	FLinearColor BlendFactor;
-	FORCEINLINE void Set(FBlendStateRHIParamRef InState, const FLinearColor& InBlendFactor)
+	FORCEINLINE_DEBUGGABLE void Set(FBlendStateRHIParamRef InState, const FLinearColor& InBlendFactor)
 	{
 		Type = ERCT_SetBlendState;
 		State = InState;
@@ -298,7 +300,7 @@ struct FRHICommandSetStreamSource : public FRHICommand
 	FVertexBufferRHIParamRef VertexBuffer;
 	uint32 Stride;
 	uint32 Offset;
-	FORCEINLINE void Set(uint32 InStreamIndex, FVertexBufferRHIParamRef InVertexBuffer, uint32 InStride, uint32 InOffset)
+	FORCEINLINE_DEBUGGABLE void Set(uint32 InStreamIndex, FVertexBufferRHIParamRef InVertexBuffer, uint32 InStride, uint32 InOffset)
 	{
 		Type = ERCT_SetStreamSource;
 		StreamIndex = InStreamIndex;
@@ -316,7 +318,7 @@ struct FRHICommandSetViewport : public FRHICommand
 	uint32 MaxX;
 	uint32 MaxY;
 	float MaxZ;
-	FORCEINLINE void Set(uint32 InMinX, uint32 InMinY, float InMinZ, uint32 InMaxX, uint32 InMaxY, float InMaxZ)
+	FORCEINLINE_DEBUGGABLE void Set(uint32 InMinX, uint32 InMinY, float InMinZ, uint32 InMaxX, uint32 InMaxY, float InMaxZ)
 	{
 		Type = ERCT_SetViewport;
 		MinX = InMinX;
@@ -335,7 +337,7 @@ struct FRHICommandSetScissorRect : public FRHICommand
 	uint32 MinY;
 	uint32 MaxX;
 	uint32 MaxY;
-	FORCEINLINE void Set(bool InbEnable, uint32 InMinX, uint32 InMinY, uint32 InMaxX, uint32 InMaxY)
+	FORCEINLINE_DEBUGGABLE void Set(bool InbEnable, uint32 InMinX, uint32 InMinY, uint32 InMaxX, uint32 InMaxY)
 	{
 		Type = ERCT_SetScissorRect;
 		bEnable = InbEnable;
@@ -354,7 +356,7 @@ struct FRHICommandSetRenderTargets : public FRHICommand
 	uint32 NewNumUAVs;
 	FUnorderedAccessViewRHIParamRef UAVs[MaxSimultaneousUAVs];
 
-	FORCEINLINE void Set(
+	FORCEINLINE_DEBUGGABLE void Set(
 		uint32 InNewNumSimultaneousRenderTargets,
 		const FRHIRenderTargetView* InNewRenderTargetsRHI,
 		FTextureRHIParamRef InNewDepthStencilTargetRHI,
@@ -364,7 +366,7 @@ struct FRHICommandSetRenderTargets : public FRHICommand
 
 	{
 		Type = ERCT_SetRenderTargets;
-		check(NewNumSimultaneousRenderTargets <= MaxSimultaneousRenderTargets && NewNumUAVs <= MaxSimultaneousUAVs);
+		check(InNewNumSimultaneousRenderTargets <= MaxSimultaneousRenderTargets && InNewNumUAVs <= MaxSimultaneousUAVs);
 		NewNumSimultaneousRenderTargets = InNewNumSimultaneousRenderTargets;
 		for (uint32 Index = 0; Index < NewNumSimultaneousRenderTargets; Index++)
 		{
@@ -379,8 +381,53 @@ struct FRHICommandSetRenderTargets : public FRHICommand
 	}
 };
 
+struct FRHICommandEndDrawPrimitiveUP : public FRHICommand
+{
+	uint32 PrimitiveType;
+	uint32 NumPrimitives;
+	uint32 NumVertices;
+	uint32 VertexDataStride;
+	void* OutVertexData;
+
+	FORCEINLINE_DEBUGGABLE void Set(uint32 InPrimitiveType, uint32 InNumPrimitives, uint32 InNumVertices, uint32 InVertexDataStride, void* InOutVertexData)
+	{
+		Type = ERCT_EndDrawPrimitiveUP;
+		PrimitiveType = InPrimitiveType;
+		NumPrimitives = InNumPrimitives;
+		NumVertices = InNumVertices;
+		VertexDataStride = InVertexDataStride;
+		OutVertexData = InOutVertexData;
+	}
+};
 
 
+struct FRHICommandEndDrawIndexedPrimitiveUP : public FRHICommand
+{
+	uint32 PrimitiveType;
+	uint32 NumPrimitives;
+	uint32 NumVertices;
+	uint32 VertexDataStride;
+	void* OutVertexData;
+	uint32 MinVertexIndex;
+	uint32 NumIndices;
+	uint32 IndexDataStride;
+	void* OutIndexData;
+
+	FORCEINLINE_DEBUGGABLE void Set(uint32 InPrimitiveType, uint32 InNumPrimitives, uint32 InNumVertices, uint32 InVertexDataStride, void* InOutVertexData, uint32 InMinVertexIndex, uint32 InNumIndices, uint32 InIndexDataStride, void* InOutIndexData)
+	{
+		Type = ERCT_EndDrawIndexedPrimitiveUP;
+		PrimitiveType = InPrimitiveType;
+		NumPrimitives = InNumPrimitives;
+		NumVertices = InNumVertices;
+		VertexDataStride = InVertexDataStride;
+		OutVertexData = InOutVertexData;
+		MinVertexIndex = InMinVertexIndex;
+		NumIndices = InNumIndices;
+		IndexDataStride = InIndexDataStride;
+		OutIndexData = InOutIndexData;
+	}
+};
+
 class RHI_API FRHICommandList : public FNoncopyable
 {
 private:
@@ -417,22 +464,22 @@ private:
 				Current = Head;
 			}
 
-			FORCEINLINE SIZE_T MemUsed() const
+			FORCEINLINE_DEBUGGABLE SIZE_T MemUsed() const
 			{
 				return Current - Head;
 			}
 
-			FORCEINLINE SIZE_T MemAvailable() const
+			FORCEINLINE_DEBUGGABLE SIZE_T MemAvailable() const
 			{
 				return Tail - Current;
 			}
 
-			FORCEINLINE bool CanFitAllocation(SIZE_T AlignedSize) const
+			FORCEINLINE_DEBUGGABLE bool CanFitAllocation(SIZE_T AlignedSize) const
 			{
 				return (AlignedSize + Current <= Tail);
 			}
 
-			FORCEINLINE uint8* Alloc(SIZE_T AlignedSize)
+			FORCEINLINE_DEBUGGABLE uint8* Alloc(SIZE_T AlignedSize)
 			{
 				auto* Ptr = Current;
 				Current += AlignedSize;
@@ -449,7 +496,7 @@ private:
 		FMemManager();
 		~FMemManager();
 
-		FORCEINLINE uint8* Alloc(SIZE_T InSize)
+		FORCEINLINE_DEBUGGABLE uint8* Alloc(SIZE_T InSize)
 		{
 			InSize = (InSize + (Alignment - 1)) & ~(Alignment - 1);
 			bool bTryAgain = false;
@@ -496,6 +543,30 @@ private:
 
 	FMemManager MemManager;
 
+	struct FDrawUpData
+	{
+
+		uint32 PrimitiveType;
+		uint32 NumPrimitives;
+		uint32 NumVertices;
+		uint32 VertexDataStride;
+		void* OutVertexData;
+		uint32 MinVertexIndex;
+		uint32 NumIndices;
+		uint32 IndexDataStride;
+		void* OutIndexData;
+
+		FDrawUpData()
+			: PrimitiveType(PT_Num)
+			, OutVertexData(nullptr)
+			, OutIndexData(nullptr)
+		{
+		}
+
+	};
+
+	FDrawUpData DrawUPData;
+
 public:
 	enum
 	{
@@ -513,7 +584,7 @@ public:
 	inline void Flush();
 
 	template <typename TShaderRHIParamRef>
-	FORCEINLINE void SetShaderUniformBuffer(TShaderRHIParamRef Shader, uint32 BaseIndex, FUniformBufferRHIParamRef UniformBuffer)
+	FORCEINLINE_DEBUGGABLE void SetShaderUniformBuffer(TShaderRHIParamRef Shader, uint32 BaseIndex, FUniformBufferRHIParamRef UniformBuffer)
 	{
 		if (Bypass())
 		{
@@ -527,7 +598,7 @@ public:
 	}
 
 	template <typename TShaderRHIParamRef>
-	FORCEINLINE void SetShaderParameter(TShaderRHIParamRef Shader, uint32 BufferIndex, uint32 BaseIndex, uint32 NumBytes, const void* NewValue, bool bValueInStack = false)
+	FORCEINLINE_DEBUGGABLE void SetShaderParameter(TShaderRHIParamRef Shader, uint32 BufferIndex, uint32 BaseIndex, uint32 NumBytes, const void* NewValue, bool bValueInStack = false)
 	{
 		if (Bypass())
 		{
@@ -547,7 +618,7 @@ public:
 	}
 
 	template <typename TShaderRHIParamRef>
-	FORCEINLINE void SetShaderTexture(TShaderRHIParamRef Shader, uint32 TextureIndex, FTextureRHIParamRef Texture)
+	FORCEINLINE_DEBUGGABLE void SetShaderTexture(TShaderRHIParamRef Shader, uint32 TextureIndex, FTextureRHIParamRef Texture)
 	{
 		if (Bypass())
 		{
@@ -562,7 +633,7 @@ public:
 	}
 
 	template <typename TShaderRHIParamRef>
-	FORCEINLINE void SetShaderResourceViewParameter(TShaderRHIParamRef Shader, uint32 SamplerIndex, FShaderResourceViewRHIParamRef SRV)
+	FORCEINLINE_DEBUGGABLE void SetShaderResourceViewParameter(TShaderRHIParamRef Shader, uint32 SamplerIndex, FShaderResourceViewRHIParamRef SRV)
 	{
 		if (Bypass())
 		{
@@ -576,7 +647,7 @@ public:
 	}
 
 	template <typename TShaderRHIParamRef>
-	FORCEINLINE void SetShaderSampler(TShaderRHIParamRef Shader, uint32 SamplerIndex, FSamplerStateRHIParamRef State)
+	FORCEINLINE_DEBUGGABLE void SetShaderSampler(TShaderRHIParamRef Shader, uint32 SamplerIndex, FSamplerStateRHIParamRef State)
 	{
 		if (Bypass())
 		{
@@ -590,7 +661,7 @@ public:
 		State->AddRef();
 	}
 
-	FORCEINLINE void SetUAVParameter(FComputeShaderRHIParamRef Shader, uint32 UAVIndex, FUnorderedAccessViewRHIParamRef UAV)
+	FORCEINLINE_DEBUGGABLE void SetUAVParameter(FComputeShaderRHIParamRef Shader, uint32 UAVIndex, FUnorderedAccessViewRHIParamRef UAV)
 	{
 		if (Bypass())
 		{
@@ -600,10 +671,13 @@ public:
 		auto* Cmd = AddCommand<FRHICommandSetUAVParameter>();
 		Cmd->Set<FComputeShaderRHIParamRef>(Shader, UAVIndex, UAV);
 		Shader->AddRef();
-		UAV->AddRef();
+		if (UAV)
+		{
+			UAV->AddRef();
+		}
 	}
 
-	FORCEINLINE void SetUAVParameter(FComputeShaderRHIParamRef Shader, uint32 UAVIndex, FUnorderedAccessViewRHIParamRef UAV, uint32 InitialCount)
+	FORCEINLINE_DEBUGGABLE void SetUAVParameter(FComputeShaderRHIParamRef Shader, uint32 UAVIndex, FUnorderedAccessViewRHIParamRef UAV, uint32 InitialCount)
 	{
 		if (Bypass())
 		{
@@ -613,10 +687,13 @@ public:
 		auto* Cmd = AddCommand<FRHICommandSetUAVParameter_IntialCount>();
 		Cmd->Set<FComputeShaderRHIParamRef>(Shader, UAVIndex, UAV, InitialCount);
 		Shader->AddRef();
-		UAV->AddRef();
+		if (UAV)
+		{
+			UAV->AddRef();
+		}
 	}
 
-	FORCEINLINE void SetBoundShaderState(FBoundShaderStateRHIParamRef BoundShaderState)
+	FORCEINLINE_DEBUGGABLE void SetBoundShaderState(FBoundShaderStateRHIParamRef BoundShaderState)
 	{
 		if (Bypass())
 		{
@@ -629,7 +706,7 @@ public:
 		BoundShaderState->AddRef();
 	}
 
-	FORCEINLINE void SetRasterizerState(FRasterizerStateRHIParamRef State)
+	FORCEINLINE_DEBUGGABLE void SetRasterizerState(FRasterizerStateRHIParamRef State)
 	{
 		if (Bypass())
 		{
@@ -642,7 +719,7 @@ public:
 		State->AddRef();
 	}
 
-	FORCEINLINE void SetBlendState(FBlendStateRHIParamRef State, const FLinearColor& BlendFactor = FLinearColor::White)
+	FORCEINLINE_DEBUGGABLE void SetBlendState(FBlendStateRHIParamRef State, const FLinearColor& BlendFactor = FLinearColor::White)
 	{
 		if (Bypass())
 		{
@@ -655,7 +732,7 @@ public:
 		State->AddRef();
 	}
 
-	FORCEINLINE void DrawPrimitive(uint32 PrimitiveType, uint32 BaseVertexIndex, uint32 NumPrimitives, uint32 NumInstances)
+	FORCEINLINE_DEBUGGABLE void DrawPrimitive(uint32 PrimitiveType, uint32 BaseVertexIndex, uint32 NumPrimitives, uint32 NumInstances)
 	{
 		if (Bypass())
 		{
@@ -667,7 +744,7 @@ public:
 		Cmd->Set(PrimitiveType, BaseVertexIndex, NumPrimitives, NumInstances);
 	}
 
-	FORCEINLINE void DrawIndexedPrimitive(FIndexBufferRHIParamRef IndexBuffer, uint32 PrimitiveType, int32 BaseVertexIndex, uint32 MinIndex, uint32 NumVertices, uint32 StartIndex, uint32 NumPrimitives, uint32 NumInstances)
+	FORCEINLINE_DEBUGGABLE void DrawIndexedPrimitive(FIndexBufferRHIParamRef IndexBuffer, uint32 PrimitiveType, int32 BaseVertexIndex, uint32 MinIndex, uint32 NumVertices, uint32 StartIndex, uint32 NumPrimitives, uint32 NumInstances)
 	{
 		if (Bypass())
 		{
@@ -680,7 +757,7 @@ public:
 		IndexBuffer->AddRef();
 	}
 
-	FORCEINLINE void SetStreamSource(uint32 StreamIndex, FVertexBufferRHIParamRef VertexBuffer, uint32 Stride, uint32 Offset)
+	FORCEINLINE_DEBUGGABLE void SetStreamSource(uint32 StreamIndex, FVertexBufferRHIParamRef VertexBuffer, uint32 Stride, uint32 Offset)
 	{
 		if (Bypass())
 		{
@@ -693,7 +770,7 @@ public:
 		VertexBuffer->AddRef();
 	}
 
-	FORCEINLINE void SetDepthStencilState(FDepthStencilStateRHIParamRef NewStateRHI, uint32 StencilRef = 0)
+	FORCEINLINE_DEBUGGABLE void SetDepthStencilState(FDepthStencilStateRHIParamRef NewStateRHI, uint32 StencilRef = 0)
 	{
 		if (Bypass())
 		{
@@ -706,7 +783,7 @@ public:
 		NewStateRHI->AddRef();
 	}
 
-	FORCEINLINE void SetViewport(uint32 MinX, uint32 MinY, float MinZ, uint32 MaxX, uint32 MaxY, float MaxZ)
+	FORCEINLINE_DEBUGGABLE void SetViewport(uint32 MinX, uint32 MinY, float MinZ, uint32 MaxX, uint32 MaxY, float MaxZ)
 	{
 		if (Bypass())
 		{
@@ -717,7 +794,7 @@ public:
 		Cmd->Set(MinX, MinY, MinZ, MaxX, MaxY, MaxZ);
 	}
 
-	FORCEINLINE void SetScissorRect(bool bEnable, uint32 MinX, uint32 MinY, uint32 MaxX, uint32 MaxY)
+	FORCEINLINE_DEBUGGABLE void SetScissorRect(bool bEnable, uint32 MinX, uint32 MinY, uint32 MaxX, uint32 MaxY)
 	{
 		if (Bypass())
 		{
@@ -728,7 +805,7 @@ public:
 		Cmd->Set(bEnable,  MinX, MinY, MaxX, MaxY);
 	}
 
-	FORCEINLINE void SetRenderTargets(
+	FORCEINLINE_DEBUGGABLE void SetRenderTargets(
 		uint32 NewNumSimultaneousRenderTargets,
 		const FRHIRenderTargetView* NewRenderTargetsRHI,
 		FTextureRHIParamRef NewDepthStencilTargetRHI,
@@ -755,15 +832,100 @@ public:
 			UAVs);
 		for (uint32 Index = 0; Index < NewNumSimultaneousRenderTargets; Index++)
 		{
-			NewRenderTargetsRHI[Index].Texture->AddRef();
+			if (NewRenderTargetsRHI[Index].Texture)
+			{
+				NewRenderTargetsRHI[Index].Texture->AddRef();
+			}
+		}
+		if (NewDepthStencilTargetRHI)
+		{
+			NewDepthStencilTargetRHI->AddRef();
 		}
-		NewDepthStencilTargetRHI->AddRef();
 		for (uint32 Index = 0; Index < NewNumUAVs; Index++)
 		{
-			UAVs[Index]->AddRef();
+			if (UAVs[Index])
+			{
+				UAVs[Index]->AddRef();
+			}
+		}
+
+	}
+
+	void BeginDrawPrimitiveUP(uint32 PrimitiveType, uint32 NumPrimitives, uint32 NumVertices, uint32 VertexDataStride, void*& OutVertexData)
+	{
+		if (Bypass())
+		{
+			BeginDrawPrimitiveUP_Internal(PrimitiveType, NumPrimitives, NumVertices, VertexDataStride, OutVertexData);
+			return;
+		}
+		check(!DrawUPData.OutVertexData && NumVertices * VertexDataStride > 0);
+
+		OutVertexData = FMemory::Malloc(NumVertices * VertexDataStride);
+
+		DrawUPData.PrimitiveType = PrimitiveType;
+		DrawUPData.NumPrimitives = NumPrimitives;
+		DrawUPData.NumVertices = NumVertices;
+		DrawUPData.VertexDataStride = VertexDataStride;
+		DrawUPData.OutVertexData = OutVertexData;
+	}
+
+	void EndDrawPrimitiveUP()
+	{
+		if (Bypass())
+		{
+			EndDrawPrimitiveUP_Internal();
+			return;
+		}
+		check(DrawUPData.OutVertexData && DrawUPData.NumVertices);
+		auto* Cmd = AddCommand<FRHICommandEndDrawPrimitiveUP>();
+		Cmd->Set(DrawUPData.PrimitiveType, DrawUPData.NumPrimitives, DrawUPData.NumVertices, DrawUPData.VertexDataStride, DrawUPData.OutVertexData);
+
+		DrawUPData.OutVertexData = nullptr;
+		DrawUPData.NumVertices = 0;
+	}
+
+	void BeginDrawIndexedPrimitiveUP(uint32 PrimitiveType, uint32 NumPrimitives, uint32 NumVertices, uint32 VertexDataStride, void*& OutVertexData, uint32 MinVertexIndex, uint32 NumIndices, uint32 IndexDataStride, void*& OutIndexData)
+	{
+		if (Bypass())
+		{
+			BeginDrawIndexedPrimitiveUP_Internal(PrimitiveType, NumPrimitives, NumVertices, VertexDataStride, OutVertexData, MinVertexIndex, NumIndices, IndexDataStride, OutIndexData);
+			return;
 		}
+		check(!DrawUPData.OutVertexData && !DrawUPData.OutIndexData && NumVertices * VertexDataStride > 0 && NumIndices * IndexDataStride > 0);
+
+		OutVertexData = FMemory::Malloc(NumVertices * VertexDataStride);
+		OutIndexData = FMemory::Malloc(NumIndices * IndexDataStride);
+
+		DrawUPData.PrimitiveType = PrimitiveType;
+		DrawUPData.NumPrimitives = NumPrimitives;
+		DrawUPData.NumVertices = NumVertices;
+		DrawUPData.VertexDataStride = VertexDataStride;
+		DrawUPData.OutVertexData = OutVertexData;
+		DrawUPData.MinVertexIndex = MinVertexIndex;
+		DrawUPData.NumIndices = NumIndices;
+		DrawUPData.IndexDataStride = IndexDataStride;
+		DrawUPData.OutIndexData = OutIndexData;
 
 	}
+	void EndDrawIndexedPrimitiveUP()
+	{
+		if (Bypass())
+		{
+			EndDrawIndexedPrimitiveUP_Internal();
+			return;
+		}
+		check(DrawUPData.OutVertexData && DrawUPData.OutIndexData && DrawUPData.NumIndices && DrawUPData.NumVertices);
+
+		auto* Cmd = AddCommand<FRHICommandEndDrawIndexedPrimitiveUP>();
+		Cmd->Set(DrawUPData.PrimitiveType, DrawUPData.NumPrimitives, DrawUPData.NumVertices, DrawUPData.VertexDataStride, DrawUPData.OutVertexData, DrawUPData.MinVertexIndex, DrawUPData.NumIndices, DrawUPData.IndexDataStride, DrawUPData.OutIndexData);
+
+		DrawUPData.OutVertexData = nullptr;
+		DrawUPData.OutIndexData = nullptr;
+		DrawUPData.NumIndices = 0;
+		DrawUPData.NumVertices = 0;
+	}
+
+
 
 	const SIZE_T GetUsedMemory() const;
 
@@ -771,18 +933,18 @@ private:
 	bool bExecuting;
 	uint32 NumCommands;
 
-	FORCEINLINE uint8* Alloc(SIZE_T InSize)
+	FORCEINLINE_DEBUGGABLE uint8* Alloc(SIZE_T InSize)
 	{
 		return MemManager.Alloc(InSize);
 	}
 
-	FORCEINLINE bool HasCommands() const
+	FORCEINLINE_DEBUGGABLE bool HasCommands() const
 	{
 		return (NumCommands > 0);
 	}
 
 	template <typename TCmd>
-	FORCEINLINE TCmd* AddCommand(uint32 ExtraData = 0)
+	FORCEINLINE_DEBUGGABLE TCmd* AddCommand(uint32 ExtraData = 0)
 	{
 		checkSlow(!bExecuting);
 		TCmd* Cmd = (TCmd*)Alloc(sizeof(TCmd) + ExtraData);
@@ -817,18 +979,18 @@ public:
 
 	#define DEFINE_RHIMETHOD_CMDLIST(Type,Name,ParameterTypesAndNames,ParameterNames,ReturnStatement,NullImplementation)
 	#define DEFINE_RHIMETHOD(Type, Name, ParameterTypesAndNames, ParameterNames, ReturnStatement, NullImplementation) \
-		FORCEINLINE Type Name ParameterTypesAndNames \
+		FORCEINLINE_DEBUGGABLE Type Name ParameterTypesAndNames \
 		{ \
 			Flush(); \
 			ReturnStatement Name##_Internal ParameterNames; \
 		}
 	#define DEFINE_RHIMETHOD_GLOBAL(Type, Name, ParameterTypesAndNames, ParameterNames, ReturnStatement, NullImplementation) \
-		FORCEINLINE Type Name ParameterTypesAndNames \
+		FORCEINLINE_DEBUGGABLE Type Name ParameterTypesAndNames \
 		{ \
 			ReturnStatement RHI##Name ParameterNames; \
 		}
 	#define DEFINE_RHIMETHOD_GLOBALFLUSH(Type, Name, ParameterTypesAndNames, ParameterNames, ReturnStatement, NullImplementation) \
-		FORCEINLINE Type Name ParameterTypesAndNames \
+		FORCEINLINE_DEBUGGABLE Type Name ParameterTypesAndNames \
 		{ \
 			Flush(); \
 			ReturnStatement Name##_Internal ParameterNames; \
@@ -849,22 +1011,22 @@ class RHI_API FRHICommandListBypass
 	}
 public:
 #define DEFINE_RHIMETHOD_CMDLIST(Type,Name,ParameterTypesAndNames,ParameterNames,ReturnStatement,NullImplementation) \
-	FORCEINLINE Type Name ParameterTypesAndNames \
+	FORCEINLINE_DEBUGGABLE Type Name ParameterTypesAndNames \
 	{ \
 		ReturnStatement Name##_Internal ParameterNames; \
 	}
 #define DEFINE_RHIMETHOD(Type, Name, ParameterTypesAndNames, ParameterNames, ReturnStatement, NullImplementation) \
-	FORCEINLINE Type Name ParameterTypesAndNames \
+	FORCEINLINE_DEBUGGABLE Type Name ParameterTypesAndNames \
 	{ \
 		ReturnStatement Name##_Internal ParameterNames; \
 	}
 #define DEFINE_RHIMETHOD_GLOBAL(Type, Name, ParameterTypesAndNames, ParameterNames, ReturnStatement, NullImplementation) \
-	FORCEINLINE Type Name ParameterTypesAndNames \
+	FORCEINLINE_DEBUGGABLE Type Name ParameterTypesAndNames \
 	{ \
 		ReturnStatement RHI##Name ParameterNames; \
 	}
 #define DEFINE_RHIMETHOD_GLOBALFLUSH(Type, Name, ParameterTypesAndNames, ParameterNames, ReturnStatement, NullImplementation) \
-	FORCEINLINE Type Name ParameterTypesAndNames \
+	FORCEINLINE_DEBUGGABLE Type Name ParameterTypesAndNames \
 	{ \
 		ReturnStatement Name##_Internal ParameterNames; \
 	}
@@ -891,12 +1053,13 @@ public:
 	static inline FRHICommandListImmediate& GetRecursiveRHICommandList(); // Only for use by RHI implementations
 
 	void ExecuteList(FRHICommandList& CmdList);
+	void LatchBypass();
 
-	FORCEINLINE void Verify()
+	FORCEINLINE_DEBUGGABLE void Verify()
 	{
-		check(!CommandListImmediate.HasCommands());
+		check(CommandListImmediate.bExecuting || !CommandListImmediate.HasCommands());
 	}
-	FORCEINLINE bool Bypass()
+	FORCEINLINE_DEBUGGABLE bool Bypass()
 	{
 #if !(UE_BUILD_SHIPPING || UE_BUILD_TEST)
 		return bLatchedBypass;
@@ -914,16 +1077,16 @@ private:
 };
 extern RHI_API FRHICommandListExecutor GRHICommandList;
 
-FORCEINLINE bool FRHICommandList::Bypass()
+FORCEINLINE_DEBUGGABLE bool FRHICommandList::Bypass()
 {
 	return GRHICommandList.Bypass();
 }
 
-FORCEINLINE FRHICommandListImmediate& FRHICommandListExecutor::GetImmediateCommandList()
+FORCEINLINE_DEBUGGABLE FRHICommandListImmediate& FRHICommandListExecutor::GetImmediateCommandList()
 {
 	return GRHICommandList.CommandListImmediate;
 }
-FORCEINLINE FRHICommandListImmediate& FRHICommandListExecutor::GetRecursiveRHICommandList()
+FORCEINLINE_DEBUGGABLE FRHICommandListImmediate& FRHICommandListExecutor::GetRecursiveRHICommandList()
 {
 	check(!GRHICommandList.CommandListImmediate.bExecuting); // queued RHI commands must not be implemented by calling other queued RHI commands
 	GRHICommandList.CommandListImmediate.Flush(); // if for some reason this is called non-recursively, this flushes everything to get us in a reasonable state for the recursive, hazardous calls
@@ -934,7 +1097,7 @@ FORCEINLINE FRHICommandListImmediate& FRHICommandListExecutor::GetRecursiveRHICo
 #define DEFINE_RHIMETHOD(Type, Name, ParameterTypesAndNames, ParameterNames, ReturnStatement, NullImplementation)
 #define DEFINE_RHIMETHOD_GLOBAL(Type, Name, ParameterTypesAndNames, ParameterNames, ReturnStatement, NullImplementation)
 #define DEFINE_RHIMETHOD_GLOBALFLUSH(Type, Name, ParameterTypesAndNames, ParameterNames, ReturnStatement, NullImplementation) \
-	FORCEINLINE Type RHI##Name ParameterTypesAndNames \
+	FORCEINLINE_DEBUGGABLE Type RHI##Name ParameterTypesAndNames \
 	{ \
 		ReturnStatement FRHICommandListExecutor::GetImmediateCommandList().Name ParameterNames; \
 	}
diff --git a/Engine/Source/Runtime/RHI/Public/RHICommandList.inl b/Engine/Source/Runtime/RHI/Public/RHICommandList.inl
--- a/Engine/Source/Runtime/RHI/Public/RHICommandList.inl
+++ b/Engine/Source/Runtime/RHI/Public/RHICommandList.inl
@@ -8,10 +8,11 @@
 #define __RHICOMMANDLIST_INL__
 #include "RHICommandList.h"
 
-FORCEINLINE void FRHICommandList::Flush()
+FORCEINLINE_DEBUGGABLE void FRHICommandList::Flush()
 {
 	if (HasCommands())
 	{
+		check(!DrawUPData.OutVertexData && !DrawUPData.OutIndexData);
 		GRHICommandList.ExecuteList(*this);
 	}
 }
@@ -29,20 +30,20 @@ public:
 		CmdTail = Page->Current;
 	}
 
-	FORCEINLINE bool HasCommandsLeft() const
+	FORCEINLINE_DEBUGGABLE bool HasCommandsLeft() const
 	{
 		return (CmdPtr < CmdTail);
 	}
 
 	// Current command
-	FORCEINLINE FRHICommand* operator * ()
+	FORCEINLINE_DEBUGGABLE FRHICommand* operator * ()
 	{
 		return (FRHICommand*)CmdPtr;
 	}
 
 	// Get the next RHICommand and advance the iterator
 	template <typename TCmd>
-	FORCEINLINE TCmd* NextCommand()
+	FORCEINLINE_DEBUGGABLE TCmd* NextCommand()
 	{
 		TCmd* RHICmd = (TCmd*)CmdPtr;
 		//::OutputDebugStringW(*FString::Printf(TEXT("Exec %d: %d @ 0x%p, %d bytes\n"), NumCommands, (int32)RHICmd->Type, (void*)RHICmd, sizeof(TCmd) + RHICmd->ExtraSize()));
diff --git a/Engine/Source/Runtime/RHI/Public/RHIMethods.h b/Engine/Source/Runtime/RHI/Public/RHIMethods.h
--- a/Engine/Source/Runtime/RHI/Public/RHIMethods.h
+++ b/Engine/Source/Runtime/RHI/Public/RHIMethods.h
@@ -1784,7 +1784,7 @@ DEFINE_RHIMETHOD_4(
  * @param VertexDataStride Size of each vertex 
  * @param OutVertexData Reference to the allocated vertex memory
  */
-DEFINE_RHIMETHOD_5(
+DEFINE_RHIMETHOD_CMDLIST_5(
 	void,BeginDrawPrimitiveUP,
 	uint32,PrimitiveType,
 	uint32,NumPrimitives,
@@ -1797,7 +1797,7 @@ DEFINE_RHIMETHOD_5(
 /**
  * Draw a primitive using the vertex data populated since RHIBeginDrawPrimitiveUP and clean up any memory as needed
  */
-DEFINE_RHIMETHOD_0(
+DEFINE_RHIMETHOD_CMDLIST_0(
 	void,EndDrawPrimitiveUP,
 	,
 	);
@@ -1814,7 +1814,7 @@ DEFINE_RHIMETHOD_0(
  * @param IndexDataStride Size of each index (either 2 or 4 bytes)
  * @param OutIndexData Reference to the allocated index memory
  */
-DEFINE_RHIMETHOD_9(
+DEFINE_RHIMETHOD_CMDLIST_9(
 	void,BeginDrawIndexedPrimitiveUP,
 	uint32,PrimitiveType,
 	uint32,NumPrimitives,
@@ -1831,7 +1831,7 @@ DEFINE_RHIMETHOD_9(
 /**
  * Draw a primitive using the vertex and index data populated since RHIBeginDrawIndexedPrimitiveUP and clean up any memory as needed
  */
-DEFINE_RHIMETHOD_0(
+DEFINE_RHIMETHOD_CMDLIST_0(
 	void,EndDrawIndexedPrimitiveUP,
 	,
 	);
diff --git a/Engine/Source/Runtime/RHI/Public/RHIUtilities.h b/Engine/Source/Runtime/RHI/Public/RHIUtilities.h
--- a/Engine/Source/Runtime/RHI/Public/RHIUtilities.h
+++ b/Engine/Source/Runtime/RHI/Public/RHIUtilities.h
@@ -368,7 +368,7 @@ inline uint32 GetVertexCountForPrimitiveCount(uint32 NumPrimitives, uint32 Primi
  * @param VertexData A reference to memory preallocate in RHIBeginDrawPrimitiveUP
  * @param VertexDataStride Size of each vertex
  */
-inline void DrawPrimitiveUP(FRHICommandListImmediate& RHICmdList, uint32 PrimitiveType, uint32 NumPrimitives, const void* VertexData, uint32 VertexDataStride)
+inline void DrawPrimitiveUP(FRHICommandList& RHICmdList, uint32 PrimitiveType, uint32 NumPrimitives, const void* VertexData, uint32 VertexDataStride)
 {
 	void* Buffer = NULL;
 	const uint32 VertexCount = GetVertexCountForPrimitiveCount( NumPrimitives, PrimitiveType );
@@ -389,7 +389,7 @@ inline void DrawPrimitiveUP(FRHICommandListImmediate& RHICmdList, uint32 Primiti
  * @param VertexDataStride The size of one vertex
  */
 inline void DrawIndexedPrimitiveUP(
-	FRHICommandListImmediate& RHICmdList,
+	FRHICommandList& RHICmdList,
 	uint32 PrimitiveType,
 	uint32 MinVertexIndex,
 	uint32 NumVertices,
diff --git a/Engine/Source/Runtime/Renderer/Private/DrawingPolicy.cpp b/Engine/Source/Runtime/Renderer/Private/DrawingPolicy.cpp
--- a/Engine/Source/Runtime/Renderer/Private/DrawingPolicy.cpp
+++ b/Engine/Source/Runtime/Renderer/Private/DrawingPolicy.cpp
@@ -71,14 +71,10 @@ void FMeshDrawingPolicy::DrawMesh(FRHICommandList& RHICmdList, const FMeshBatch&
 	{
 		check(Mesh.DynamicVertexData);
 
-		//hack here to get the code checked in, but this needs to be fixed. It was not sensible to wind this all the way out of the callstack
-		FRHICommandListImmediate& RHICmdListIm = FRHICommandListExecutor::GetImmediateCommandList();
-		check(&RHICmdList == &RHICmdListIm);
-
 		if (BatchElement.DynamicIndexData)
 		{
 			DrawIndexedPrimitiveUP(
-				RHICmdListIm,
+				RHICmdList,
 				Mesh.Type,
 				BatchElement.MinVertexIndex,
 				BatchElement.MaxVertexIndex - BatchElement.MinVertexIndex + 1,
@@ -92,7 +88,7 @@ void FMeshDrawingPolicy::DrawMesh(FRHICommandList& RHICmdList, const FMeshBatch&
 		else
 		{
 			DrawPrimitiveUP(
-				RHICmdListIm,
+				RHICmdList,
 				Mesh.Type,
 				BatchElement.NumPrimitives,
 				Mesh.DynamicVertexData,
