UE4 - further command list work. BoundShaderState, RHISetComputeShader, RHIDispatchComputeShader, RHIDispatchIndirectComputeShader, RHIAutomaticCacheFlushAfterComputeShader, RHIDrawPrimitiveIndirect, RHIDrawIndexedIndirect, RHIDrawIndexedPrimitiveIndirect, RHIEnableDepthBoundsTest

diff --git a/Engine/Source/Runtime/Engine/Private/Particles/ParticleGpuSimulation.cpp b/Engine/Source/Runtime/Engine/Private/Particles/ParticleGpuSimulation.cpp
--- a/Engine/Source/Runtime/Engine/Private/Particles/ParticleGpuSimulation.cpp
+++ b/Engine/Source/Runtime/Engine/Private/Particles/ParticleGpuSimulation.cpp
@@ -1208,7 +1208,7 @@ struct FSimulationCommandGPU
  */
 template <bool bUseDepthBufferCollision>
 void ExecuteSimulationCommands(
-	FRHICommandListImmediate& RHICmdList,
+	FRHICommandList& RHICmdList,
 	const TArray<FSimulationCommandGPU>& SimulationCommands,
 	const FParticleStateTextures& TextureResources,
 	const FParticleAttributesTexture& AttributeTexture,
@@ -1259,7 +1259,7 @@ void ExecuteSimulationCommands(
  * Invokes the clear simulation shader for each particle in each tile.
  * @param Tiles - The list of tiles to clear.
  */
-void ClearTiles(FRHICommandListImmediate& RHICmdList, const TArray<uint32>& Tiles)
+void ClearTiles(FRHICommandList& RHICmdList, const TArray<uint32>& Tiles)
 {
 	SCOPED_DRAW_EVENT(ClearTiles, DEC_PARTICLE);
 
@@ -1500,7 +1500,7 @@ TGlobalResource<FParticleInjectionVertexDeclaration> GParticleInjectionVertexDec
  * Injects new particles in to the GPU simulation.
  * @param NewParticles - A list of particles to inject in to the simulation.
  */
-void InjectNewParticles(FRHICommandListImmediate& RHICmdList, const TArray<FNewParticle>& NewParticles)
+void InjectNewParticles(FRHICommandList& RHICmdList, const TArray<FNewParticle>& NewParticles)
 {
 	const int32 MaxParticlesPerDrawCall = GParticleScratchVertexBufferSize / sizeof(FNewParticle);
 	FVertexBufferRHIParamRef ScratchVertexBufferRHI = GParticleScratchVertexBuffer.VertexBufferRHI;
@@ -1723,7 +1723,7 @@ TGlobalResource<FParticleSimVisualizeVertexDeclaration> GParticleSimVisualizeVer
  * @param RenderTarget - The render target on which to draw the visualization.
  */
 static void VisualizeGPUSimulation(
-	FRHICommandListImmediate& RHICmdList,
+	FRHICommandList& RHICmdList,
 	int32 VisualizationMode,
 	FRenderTarget* RenderTarget,
 	const FParticleStateTextures& StateTextures,
diff --git a/Engine/Source/Runtime/Engine/Private/StaticBoundShaderState.cpp b/Engine/Source/Runtime/Engine/Private/StaticBoundShaderState.cpp
--- a/Engine/Source/Runtime/Engine/Private/StaticBoundShaderState.cpp
+++ b/Engine/Source/Runtime/Engine/Private/StaticBoundShaderState.cpp
@@ -56,7 +56,7 @@ FBoundShaderStateRHIParamRef FGlobalBoundShaderStateResource::GetInitializedRHI(
 	// Create the bound shader state if it hasn't been cached yet.
 	if(!IsValidRef(BoundShaderState))
 	{
-		BoundShaderState = RHICreateBoundShaderState(
+		BoundShaderState = CreateBoundShaderState_Internal(
 			VertexDeclaration,
 			VertexShader,
 			FHullShaderRHIRef(), 
@@ -68,7 +68,7 @@ FBoundShaderStateRHIParamRef FGlobalBoundShaderStateResource::GetInitializedRHI(
 	// Only people working on shaders (and therefore have LogShaders unsuppressed) will want to see these errors
 	else if (!GUsingNullRHI && UE_LOG_ACTIVE(LogShaders, Warning))
 	{
-		FBoundShaderStateRHIRef TempBoundShaderState = RHICreateBoundShaderState(
+		FBoundShaderStateRHIRef TempBoundShaderState = CreateBoundShaderState_Internal(
 			VertexDeclaration,
 			VertexShader,
 			FHullShaderRHIRef(),
@@ -92,26 +92,26 @@ void FGlobalBoundShaderStateResource::ReleaseRHI()
 	BoundShaderState.SafeRelease();
 }
 
-void SetGlobalBoundShaderState(
-	FRHICommandListImmediate& RHICmdList,
-	FGlobalBoundShaderState& GlobalBoundShaderState,
-	FVertexDeclarationRHIParamRef VertexDeclaration,
-	FShader* VertexShader,
-	FShader* PixelShader,
-	FShader* GeometryShader
-	)
+void SetGlobalBoundShaderState_Internal(FGlobalBoundShaderState GlobalBoundShaderState)
 {
 	// Check for unset uniform buffer parameters
 	// Technically you can set uniform buffer parameters after calling RHISetBoundShaderState, but this is the most global place to check for unset parameters
-	VertexShader->VerifyBoundUniformBufferParameters();
-	PixelShader->VerifyBoundUniformBufferParameters();
-	GeometryShader->VerifyBoundUniformBufferParameters();
+	GlobalBoundShaderState->Args.VertexShader->VerifyBoundUniformBufferParameters();
+	GlobalBoundShaderState->Args.PixelShader->VerifyBoundUniformBufferParameters();
+	GlobalBoundShaderState->Args.GeometryShader->VerifyBoundUniformBufferParameters();
 
-	RHICmdList.SetBoundShaderState(
-		GlobalBoundShaderState.GetInitializedRHI(
-			VertexDeclaration,
-			GETSAFERHISHADER_VERTEX(VertexShader),
-			GETSAFERHISHADER_PIXEL(PixelShader),
-			(FGeometryShaderRHIParamRef)GETSAFERHISHADER_GEOMETRY(GeometryShader))
+	FGlobalBoundShaderState_Internal* BSS = GlobalBoundShaderState->BSS;
+	if (!BSS)
+	{
+		BSS = new FGlobalBoundShaderState_Internal(); // these are simply leaked and never freed
+		GlobalBoundShaderState->BSS = BSS;
+	}
+
+	SetBoundShaderState_Internal(
+		BSS->GetInitializedRHI(
+			GlobalBoundShaderState->Args.VertexDeclarationRHI,
+			GETSAFERHISHADER_VERTEX(GlobalBoundShaderState->Args.VertexShader),
+			GETSAFERHISHADER_PIXEL(GlobalBoundShaderState->Args.PixelShader),
+			(FGeometryShaderRHIParamRef)GETSAFERHISHADER_GEOMETRY(GlobalBoundShaderState->Args.GeometryShader))
 		);
 }
diff --git a/Engine/Source/Runtime/Engine/Public/StaticBoundShaderState.h b/Engine/Source/Runtime/Engine/Public/StaticBoundShaderState.h
--- a/Engine/Source/Runtime/Engine/Public/StaticBoundShaderState.h
+++ b/Engine/Source/Runtime/Engine/Public/StaticBoundShaderState.h
@@ -48,7 +48,17 @@ private:
 	ENGINE_API virtual void ReleaseRHI();
 };
 
-typedef TGlobalResource<FGlobalBoundShaderStateResource> FGlobalBoundShaderState;
+typedef TGlobalResource<FGlobalBoundShaderStateResource> FGlobalBoundShaderState_Internal;
+
+
+/**
+* SetGlobalBoundShaderState - sets the global bound shader state, also creates and caches it if necessary
+*
+* @param BoundShaderState			- global bound shader state to initialize and set
+*/
+
+extern ENGINE_API void SetGlobalBoundShaderState_Internal(FGlobalBoundShaderState BoundShaderState);
+
 
 /**
  * SetGlobalBoundShaderState - sets the global bound shader state, also creates and caches it if necessary
@@ -61,13 +71,16 @@ typedef TGlobalResource<FGlobalBoundShaderStateResource> FGlobalBoundShaderState
  * @param GeometryShader			- the geometry shader to use in creating the new bound shader state (0 if not used)
  */
 
-extern ENGINE_API void SetGlobalBoundShaderState(
-	FRHICommandListImmediate& RHICmdList,
+FORCEINLINE void SetGlobalBoundShaderState(
+	FRHICommandList& RHICmdList,
 	FGlobalBoundShaderState& BoundShaderState,
 	FVertexDeclarationRHIParamRef VertexDeclaration,
 	FShader* VertexShader,
 	FShader* PixelShader,
-	FShader* GeometryShader = 0
-	);
+	FShader* GeometryShader = nullptr
+	)
+{
+	RHICmdList.BuildAndSetGlobalBoundShaderState(&SetGlobalBoundShaderState_Internal, BoundShaderState, VertexDeclaration, VertexShader, PixelShader, GeometryShader);
+}
 
 #endif
diff --git a/Engine/Source/Runtime/RHI/Private/RHICommandList.cpp b/Engine/Source/Runtime/RHI/Private/RHICommandList.cpp
--- a/Engine/Source/Runtime/RHI/Private/RHICommandList.cpp
+++ b/Engine/Source/Runtime/RHI/Private/RHICommandList.cpp
@@ -51,6 +51,7 @@ static FAutoConsoleVariableRef CVarRHICmdMem(
 	);
 
 RHI_API FRHICommandListExecutor GRHICommandList;
+RHI_API FThreadSafeCounter FRHICommandList::UIDCounter;
 //RHI_API FGlobalRHI GRHI;
 
 static_assert(sizeof(FRHICommand) == sizeof(FRHICommandNopEndOfPage), "These should match for filling in the end of a page when a new allocation won't fit");
@@ -183,7 +184,10 @@ void FRHICommandListExecutor::ExecuteList(FRHICommandList& CmdList)
 					}
 				}
 				RHICmd->Shader->Release();
-				RHICmd->SRV->Release();
+				if (RHICmd->SRV)
+				{
+					RHICmd->SRV->Release();
+				}
 			}
 			break;
 		case ERCT_SetUAVParameter:
@@ -268,7 +272,10 @@ void FRHICommandListExecutor::ExecuteList(FRHICommandList& CmdList)
 				{
 					SetStreamSource_Internal(RHICmd->StreamIndex, RHICmd->VertexBuffer, RHICmd->Stride, RHICmd->Offset);
 				}
-				RHICmd->VertexBuffer->Release();
+				if (RHICmd->VertexBuffer)
+				{
+					RHICmd->VertexBuffer->Release();
+				}
 			}
 			break;
 		case ERCT_SetDepthStencilState:
@@ -370,6 +377,160 @@ void FRHICommandListExecutor::ExecuteList(FRHICommandList& CmdList)
 				FMemory::Free(RHICmd->OutIndexData);
 			}
 			break;
+		case ERCT_BuildLocalBoundShaderState:
+			{
+				auto* RHICmd = Iter.NextCommand<FRHICommandBuildLocalBoundShaderState>();
+				if (!bOnlyTestMemAccess)
+				{
+					check(RHICmd->WorkArea.CheckCmdList == &CmdList && RHICmd->WorkArea.UID == CmdList.UID // this BSS was not built for this particular commandlist
+						&& !IsValidRef(RHICmd->WorkArea.BSS)); // should not already have been created
+					if (RHICmd->WorkArea.UseCount)
+					{
+						RHICmd->WorkArea.BSS = CreateBoundShaderState_Internal(RHICmd->WorkArea.Args.VertexDeclarationRHI, RHICmd->WorkArea.Args.VertexShaderRHI, RHICmd->WorkArea.Args.HullShaderRHI, RHICmd->WorkArea.Args.DomainShaderRHI, RHICmd->WorkArea.Args.PixelShaderRHI, RHICmd->WorkArea.Args.GeometryShaderRHI);
+					}
+				}
+			}
+			break;
+		case ERCT_SetLocalBoundShaderState:
+			{
+				auto* RHICmd = Iter.NextCommand<FRHICommandSetLocalBoundShaderState>();
+				if (!bOnlyTestMemAccess)
+				{
+					check(RHICmd->LocalBoundShaderState.WorkArea->CheckCmdList == &CmdList && RHICmd->LocalBoundShaderState.WorkArea->UID == CmdList.UID // this BSS was not built for this particular commandlist
+						&& RHICmd->LocalBoundShaderState.WorkArea->UseCount > 0 && IsValidRef(RHICmd->LocalBoundShaderState.WorkArea->BSS)); // this should have been created and should have uses outstanding
+
+					SetBoundShaderState_Internal(RHICmd->LocalBoundShaderState.WorkArea->BSS);
+
+					if (--RHICmd->LocalBoundShaderState.WorkArea->UseCount == 0)
+					{
+						RHICmd->LocalBoundShaderState.WorkArea->BSS = nullptr; // this also releases the ref, which wouldn't otherwise happen because RHICommands are unsafe and don't do constructors or destructors
+						if (RHICmd->LocalBoundShaderState.WorkArea->Args.VertexDeclarationRHI)
+						{
+							RHICmd->LocalBoundShaderState.WorkArea->Args.VertexDeclarationRHI->Release();
+						}
+						if (RHICmd->LocalBoundShaderState.WorkArea->Args.VertexShaderRHI)
+						{
+							RHICmd->LocalBoundShaderState.WorkArea->Args.VertexShaderRHI->Release();
+						}
+						if (RHICmd->LocalBoundShaderState.WorkArea->Args.HullShaderRHI)
+						{
+							RHICmd->LocalBoundShaderState.WorkArea->Args.HullShaderRHI->Release();
+						}
+						if (RHICmd->LocalBoundShaderState.WorkArea->Args.DomainShaderRHI)
+						{
+							RHICmd->LocalBoundShaderState.WorkArea->Args.DomainShaderRHI->Release();
+						}
+						if (RHICmd->LocalBoundShaderState.WorkArea->Args.PixelShaderRHI)
+						{
+							RHICmd->LocalBoundShaderState.WorkArea->Args.PixelShaderRHI->Release();
+						}
+						if (RHICmd->LocalBoundShaderState.WorkArea->Args.GeometryShaderRHI)
+						{
+							RHICmd->LocalBoundShaderState.WorkArea->Args.GeometryShaderRHI->Release();
+						}
+					}
+				}
+			}
+			break;
+		case ERCT_SetGlobalBoundShaderState:
+			{
+				auto* RHICmd = Iter.NextCommand<FRHICommandSetGlobalBoundShaderState>();
+				if (!bOnlyTestMemAccess)
+				{
+					CmdList.SetGlobalBoundShaderState_InternalPtr(RHICmd->GlobalBoundShaderState);
+				}
+			}
+			break;
+
+		case ERCT_SetComputeShader:
+			{
+				auto* RHICmd = Iter.NextCommand<FRHICommandSetComputeShader>();
+				if (!bOnlyTestMemAccess)
+				{
+					SetComputeShader_Internal(RHICmd->ComputeShader);
+				}
+				RHICmd->ComputeShader->Release();
+			}
+			break;
+		case ERCT_DispatchComputeShader:
+			{
+				auto* RHICmd = Iter.NextCommand<FRHICommandDispatchComputeShader>();
+				if (!bOnlyTestMemAccess)
+				{
+					DispatchComputeShader_Internal(RHICmd->ThreadGroupCountX, RHICmd->ThreadGroupCountY, RHICmd->ThreadGroupCountZ);
+				}
+			}
+			break;
+		case ERCT_DispatchIndirectComputeShader:
+			{
+				auto* RHICmd = Iter.NextCommand<FRHICommandDispatchIndirectComputeShader>();
+				if (!bOnlyTestMemAccess)
+				{
+					DispatchIndirectComputeShader_Internal(RHICmd->ArgumentBuffer, RHICmd->ArgumentOffset);
+				}
+				RHICmd->ArgumentBuffer->Release();
+			}
+			break;
+		case ERCT_AutomaticCacheFlushAfterComputeShader:
+			{
+				auto* RHICmd = Iter.NextCommand<FRHICommandAutomaticCacheFlushAfterComputeShader>();
+				if (!bOnlyTestMemAccess)
+				{
+					AutomaticCacheFlushAfterComputeShader_Internal(RHICmd->bEnable);
+				}
+			}
+			break;
+		case ERCT_FlushComputeShaderCache:
+			{
+				auto* RHICmd = Iter.NextCommand<FRHICommandFlushComputeShaderCache>();
+				if (!bOnlyTestMemAccess)
+				{
+					FlushComputeShaderCache_Internal();
+				}
+			}
+			break;
+		case ERCT_DrawPrimitiveIndirect:
+			{
+				auto* RHICmd = Iter.NextCommand<FRHICommandDrawPrimitiveIndirect>();
+				if (!bOnlyTestMemAccess)
+				{
+					DrawPrimitiveIndirect_Internal(RHICmd->PrimitiveType, RHICmd->ArgumentBuffer, RHICmd->ArgumentOffset);
+				}
+				RHICmd->ArgumentBuffer->Release();
+			}
+			break;
+		case ERCT_DrawIndexedIndirect:
+			{
+				auto* RHICmd = Iter.NextCommand<FRHICommandDrawIndexedIndirect>();
+				if (!bOnlyTestMemAccess)
+				{
+					DrawIndexedIndirect_Internal(RHICmd->IndexBufferRHI, RHICmd->PrimitiveType, RHICmd->ArgumentsBufferRHI, RHICmd->DrawArgumentsIndex, RHICmd->NumInstances);
+				}
+				RHICmd->IndexBufferRHI->Release();
+				RHICmd->ArgumentsBufferRHI->Release();
+			}
+			break;
+		case ERCT_DrawIndexedPrimitiveIndirect:
+			{
+				auto* RHICmd = Iter.NextCommand<FRHICommandDrawIndexedPrimitiveIndirect>();
+				if (!bOnlyTestMemAccess)
+				{
+					DrawIndexedPrimitiveIndirect_Internal(RHICmd->PrimitiveType, RHICmd->IndexBuffer, RHICmd->ArgumentsBuffer, RHICmd->ArgumentOffset);
+				}
+				RHICmd->IndexBuffer->Release();
+				RHICmd->ArgumentsBuffer->Release();
+			}
+			break;
+		case ERCT_EnableDepthBoundsTest:
+			{
+				auto* RHICmd = Iter.NextCommand<FRHICommandEnableDepthBoundsTest>();
+				if (!bOnlyTestMemAccess)
+				{
+					EnableDepthBoundsTest_Internal(RHICmd->bEnable, RHICmd->MinDepth, RHICmd->MaxDepth);
+				}
+			}
+			break;
+
 
 		default:
 			checkf(0, TEXT("Unknown RHI Command %d!"), Cmd->Type);
@@ -486,3 +647,7 @@ const SIZE_T FRHICommandList::GetUsedMemory() const
 {
 	return MemManager.GetUsedMemory();
 }
+
+
+
+
diff --git a/Engine/Source/Runtime/RHI/Public/RHICommandList.h b/Engine/Source/Runtime/RHI/Public/RHICommandList.h
--- a/Engine/Source/Runtime/RHI/Public/RHICommandList.h
+++ b/Engine/Source/Runtime/RHI/Public/RHICommandList.h
@@ -34,6 +34,19 @@ enum ERHICommandType
 	ERCT_SetRenderTargets,
 	ERCT_EndDrawPrimitiveUP,
 	ERCT_EndDrawIndexedPrimitiveUP,
+	ERCT_BuildLocalBoundShaderState,
+	ERCT_SetLocalBoundShaderState,
+	ERCT_SetGlobalBoundShaderState,
+	ERCT_SetComputeShader,
+	ERCT_DispatchComputeShader,
+	ERCT_DispatchIndirectComputeShader,
+	ERCT_AutomaticCacheFlushAfterComputeShader,
+	ERCT_FlushComputeShaderCache,
+	ERCT_DrawPrimitiveIndirect,
+	ERCT_DrawIndexedIndirect,
+	ERCT_DrawIndexedPrimitiveIndirect,
+	ERCT_EnableDepthBoundsTest,
+
 	ERCT_,
 };
 
@@ -428,6 +441,261 @@ struct FRHICommandEndDrawIndexedPrimitiveUP : public FRHICommand
 	}
 };
 
+struct FRHICommandSetComputeShader : public FRHICommand
+{
+	FComputeShaderRHIParamRef ComputeShader;
+	FORCEINLINE_DEBUGGABLE void Set(FComputeShaderRHIParamRef InComputeShader)
+	{
+		Type = ERCT_SetComputeShader;
+		ComputeShader = InComputeShader;
+	}
+};
+
+struct FRHICommandDispatchComputeShader : public FRHICommand
+{
+	uint32 ThreadGroupCountX;
+	uint32 ThreadGroupCountY;
+	uint32 ThreadGroupCountZ;
+	FORCEINLINE_DEBUGGABLE void Set(uint32 InThreadGroupCountX, uint32 InThreadGroupCountY, uint32 InThreadGroupCountZ)
+	{
+		Type = ERCT_DispatchComputeShader;
+		ThreadGroupCountX = InThreadGroupCountX;
+		ThreadGroupCountY = InThreadGroupCountY;
+		ThreadGroupCountZ = InThreadGroupCountZ;
+	}
+};
+
+struct FRHICommandDispatchIndirectComputeShader : public FRHICommand
+{
+	FVertexBufferRHIParamRef ArgumentBuffer;
+	uint32 ArgumentOffset;
+	FORCEINLINE_DEBUGGABLE void Set(FVertexBufferRHIParamRef InArgumentBuffer, uint32 InArgumentOffset)
+	{
+		Type = ERCT_DispatchIndirectComputeShader;
+		ArgumentBuffer = InArgumentBuffer;
+		ArgumentOffset = InArgumentOffset;
+
+	}
+};
+
+struct FRHICommandAutomaticCacheFlushAfterComputeShader : public FRHICommand
+{
+	bool bEnable;
+	FORCEINLINE_DEBUGGABLE void Set(bool InbEnable)
+	{
+		Type = ERCT_AutomaticCacheFlushAfterComputeShader;
+		bEnable = InbEnable;
+	}
+};
+
+struct FRHICommandFlushComputeShaderCache : public FRHICommand
+{
+	FORCEINLINE_DEBUGGABLE void Set()
+	{
+		Type = ERCT_FlushComputeShaderCache;
+	}
+};
+
+struct FRHICommandDrawPrimitiveIndirect : public FRHICommand
+{
+	uint32 PrimitiveType;
+	FVertexBufferRHIParamRef ArgumentBuffer;
+	uint32 ArgumentOffset;
+	FORCEINLINE_DEBUGGABLE void Set(uint32 InPrimitiveType, FVertexBufferRHIParamRef InArgumentBuffer, uint32 InArgumentOffset)
+	{
+		Type = ERCT_DrawPrimitiveIndirect;
+		PrimitiveType = InPrimitiveType;
+		ArgumentBuffer = InArgumentBuffer;
+		ArgumentOffset = InArgumentOffset;
+
+	}
+};
+
+struct FRHICommandDrawIndexedIndirect : public FRHICommand
+{
+	FIndexBufferRHIParamRef IndexBufferRHI;
+	uint32 PrimitiveType;
+	FStructuredBufferRHIParamRef ArgumentsBufferRHI;
+	uint32 DrawArgumentsIndex;
+	uint32 NumInstances;
+
+	FORCEINLINE_DEBUGGABLE void Set(FIndexBufferRHIParamRef InIndexBufferRHI, uint32 InPrimitiveType, FStructuredBufferRHIParamRef InArgumentsBufferRHI, uint32 InDrawArgumentsIndex, uint32 InNumInstances)
+	{
+		Type = ERCT_DrawIndexedIndirect;
+		IndexBufferRHI = InIndexBufferRHI;
+		PrimitiveType = InPrimitiveType;
+		ArgumentsBufferRHI = InArgumentsBufferRHI;
+		DrawArgumentsIndex = InDrawArgumentsIndex;
+		NumInstances = InNumInstances;
+	}
+};
+
+struct FRHICommandDrawIndexedPrimitiveIndirect : public FRHICommand
+{
+	uint32 PrimitiveType;
+	FIndexBufferRHIParamRef IndexBuffer;
+	FVertexBufferRHIParamRef ArgumentsBuffer;
+	uint32 ArgumentOffset;
+
+	FORCEINLINE_DEBUGGABLE void Set(uint32 InPrimitiveType, FIndexBufferRHIParamRef InIndexBuffer, FVertexBufferRHIParamRef InArgumentsBuffer, uint32 InArgumentOffset)
+	{
+		Type = ERCT_DrawIndexedPrimitiveIndirect;
+		PrimitiveType = InPrimitiveType;
+		IndexBuffer = InIndexBuffer;
+		ArgumentsBuffer = InArgumentsBuffer;
+		ArgumentOffset = InArgumentOffset;
+	}
+};
+
+struct FRHICommandEnableDepthBoundsTest : public FRHICommand
+{
+	bool bEnable;
+	float MinDepth;
+	float MaxDepth;
+
+	FORCEINLINE_DEBUGGABLE void Set(bool InbEnable, float InMinDepth, float InMaxDepth)
+	{
+		Type = ERCT_EnableDepthBoundsTest;
+		bEnable = InbEnable;
+		MinDepth = InMinDepth;
+		MaxDepth = InMaxDepth;
+	}
+};
+
+struct FBoundShaderStateInput
+{
+	FVertexDeclarationRHIParamRef VertexDeclarationRHI;
+	FVertexShaderRHIParamRef VertexShaderRHI;
+	FHullShaderRHIParamRef HullShaderRHI;
+	FDomainShaderRHIParamRef DomainShaderRHI;
+	FPixelShaderRHIParamRef PixelShaderRHI;
+	FGeometryShaderRHIParamRef GeometryShaderRHI;
+
+	FORCEINLINE FBoundShaderStateInput()
+	{
+	}
+
+	FORCEINLINE FBoundShaderStateInput(FVertexDeclarationRHIParamRef InVertexDeclarationRHI,
+		FVertexShaderRHIParamRef InVertexShaderRHI,
+		FHullShaderRHIParamRef InHullShaderRHI,
+		FDomainShaderRHIParamRef InDomainShaderRHI,
+		FPixelShaderRHIParamRef InPixelShaderRHI,
+		FGeometryShaderRHIParamRef InGeometryShaderRHI
+		)
+		: VertexDeclarationRHI(InVertexDeclarationRHI)
+		, VertexShaderRHI(InVertexShaderRHI)
+		, HullShaderRHI(InHullShaderRHI)
+		, DomainShaderRHI(InDomainShaderRHI)
+		, PixelShaderRHI(InPixelShaderRHI)
+		, GeometryShaderRHI(InGeometryShaderRHI)
+	{
+	}
+};
+
+
+struct FLocalBoundShaderStateWorkArea
+{
+	FBoundShaderStateInput Args;
+	FBoundShaderStateRHIRef BSS;
+	int32 UseCount;
+#if !(UE_BUILD_SHIPPING || UE_BUILD_TEST)
+	class FRHICommandList* CheckCmdList;
+	int32 UID;
+#endif
+};
+
+struct FLocalBoundShaderState
+{
+	FLocalBoundShaderStateWorkArea* WorkArea;
+	FBoundShaderStateRHIRef BypassBSS; // this is only used in the case of Bypass, should eventually be deleted
+	FLocalBoundShaderState()
+		: WorkArea(nullptr)
+	{
+	}
+};
+
+struct FRHICommandBuildLocalBoundShaderState : public FRHICommand
+{
+	FLocalBoundShaderStateWorkArea WorkArea;
+	FORCEINLINE_DEBUGGABLE void Set(
+		class FRHICommandList* CheckCmdList,
+		int32 UID,
+		FVertexDeclarationRHIParamRef VertexDeclarationRHI,
+		FVertexShaderRHIParamRef VertexShaderRHI,
+		FHullShaderRHIParamRef HullShaderRHI,
+		FDomainShaderRHIParamRef DomainShaderRHI,
+		FPixelShaderRHIParamRef PixelShaderRHI,
+		FGeometryShaderRHIParamRef GeometryShaderRHI
+		)
+	{
+		Type = ERCT_BuildLocalBoundShaderState;
+		WorkArea.Args.VertexDeclarationRHI = VertexDeclarationRHI;
+		WorkArea.Args.VertexShaderRHI = VertexShaderRHI;
+		WorkArea.Args.HullShaderRHI = HullShaderRHI;
+		WorkArea.Args.DomainShaderRHI = DomainShaderRHI;
+		WorkArea.Args.PixelShaderRHI = PixelShaderRHI;
+		WorkArea.Args.GeometryShaderRHI = GeometryShaderRHI;
+
+		// these RHI commands are bad, no constructor! we construct this in place
+		new ((void*)&WorkArea.BSS) FBoundShaderStateRHIRef();
+		WorkArea.UseCount = 0;
+#if !(UE_BUILD_SHIPPING || UE_BUILD_TEST)
+		WorkArea.CheckCmdList = CheckCmdList;
+		WorkArea.UID = UID;
+#endif
+	}
+};
+
+struct FRHICommandSetLocalBoundShaderState : public FRHICommand
+{
+	FLocalBoundShaderState LocalBoundShaderState;
+	FORCEINLINE_DEBUGGABLE void Set(class FRHICommandList* CheckCmdList, int32 UID, FLocalBoundShaderState& InLocalBoundShaderState)
+	{
+		Type = ERCT_SetLocalBoundShaderState;
+		LocalBoundShaderState.WorkArea = InLocalBoundShaderState.WorkArea;
+		check(CheckCmdList == LocalBoundShaderState.WorkArea->CheckCmdList && UID == LocalBoundShaderState.WorkArea->UID); // this BSS was not built for this particular commandlist
+		LocalBoundShaderState.WorkArea->UseCount++;
+	}
+};
+
+
+class FShader;
+
+struct FGlobalBoundShaderStateArgs
+{
+	FVertexDeclarationRHIParamRef VertexDeclarationRHI;
+	FShader* VertexShader;
+	FShader* PixelShader;
+	FShader* GeometryShader;
+};
+
+class FGlobalBoundShaderStateResource;
+template<class ResourceType> class TGlobalResource;
+
+
+struct FGlobalBoundShaderStateWorkArea
+{
+	FGlobalBoundShaderStateArgs Args;
+	TGlobalResource<FGlobalBoundShaderStateResource>* BSS; //ideally this would be part of this memory block and not a separate allocation...that is doable, if a little tedious. The point is we need to delay the construction until we get back to the render thread.
+
+	FGlobalBoundShaderStateWorkArea()
+		: BSS(nullptr)
+	{
+	}
+};
+
+typedef FGlobalBoundShaderStateWorkArea* FGlobalBoundShaderState;
+
+struct FRHICommandSetGlobalBoundShaderState : public FRHICommand
+{
+	FGlobalBoundShaderState GlobalBoundShaderState;
+	FORCEINLINE_DEBUGGABLE void Set(FGlobalBoundShaderState InGlobalBoundShaderState)
+	{
+		Type = ERCT_SetGlobalBoundShaderState;
+		GlobalBoundShaderState = InGlobalBoundShaderState;
+	}
+};
+
 class RHI_API FRHICommandList : public FNoncopyable
 {
 private:
@@ -583,6 +851,143 @@ public:
 	}
 	inline void Flush();
 
+	FORCEINLINE_DEBUGGABLE FLocalBoundShaderState BuildLocalBoundShaderState(const FBoundShaderStateInput& BoundShaderStateInput)
+	{
+		return BuildLocalBoundShaderState(
+			BoundShaderStateInput.VertexDeclarationRHI,
+			BoundShaderStateInput.VertexShaderRHI,
+			BoundShaderStateInput.HullShaderRHI,
+			BoundShaderStateInput.DomainShaderRHI,
+			BoundShaderStateInput.PixelShaderRHI,
+			BoundShaderStateInput.GeometryShaderRHI
+			);
+	}
+
+	FORCEINLINE_DEBUGGABLE void BuildAndSetLocalBoundShaderState(const FBoundShaderStateInput& BoundShaderStateInput)
+	{
+		SetLocalBoundShaderState(BuildLocalBoundShaderState(
+			BoundShaderStateInput.VertexDeclarationRHI,
+			BoundShaderStateInput.VertexShaderRHI,
+			BoundShaderStateInput.HullShaderRHI,
+			BoundShaderStateInput.DomainShaderRHI,
+			BoundShaderStateInput.PixelShaderRHI,
+			BoundShaderStateInput.GeometryShaderRHI
+			));
+	}
+
+	FORCEINLINE_DEBUGGABLE FLocalBoundShaderState BuildLocalBoundShaderState(
+		FVertexDeclarationRHIParamRef VertexDeclarationRHI,
+		FVertexShaderRHIParamRef VertexShaderRHI,
+		FHullShaderRHIParamRef HullShaderRHI,
+		FDomainShaderRHIParamRef DomainShaderRHI,
+		FPixelShaderRHIParamRef PixelShaderRHI,
+		FGeometryShaderRHIParamRef GeometryShaderRHI
+		)
+	{
+		FLocalBoundShaderState Result;
+		if (Bypass())
+		{
+			Result.BypassBSS = CreateBoundShaderState_Internal(VertexDeclarationRHI, VertexShaderRHI, HullShaderRHI, DomainShaderRHI, PixelShaderRHI, GeometryShaderRHI);
+		}
+		else
+		{
+			auto* Cmd = AddCommand<FRHICommandBuildLocalBoundShaderState>();
+			Cmd->Set(this, this->UID, VertexDeclarationRHI, VertexShaderRHI, HullShaderRHI, DomainShaderRHI, PixelShaderRHI, GeometryShaderRHI);
+			Result.WorkArea = &Cmd->WorkArea;
+
+			if (VertexDeclarationRHI)
+			{
+				VertexDeclarationRHI->AddRef();
+			}
+			if (VertexShaderRHI)
+			{
+				VertexShaderRHI->AddRef();
+			}
+			if (HullShaderRHI)
+			{
+				HullShaderRHI->AddRef();
+			}
+			if (DomainShaderRHI)
+			{
+				DomainShaderRHI->AddRef();
+			}
+			if (PixelShaderRHI)
+			{
+				PixelShaderRHI->AddRef();
+			}
+			if (GeometryShaderRHI)
+			{
+				GeometryShaderRHI->AddRef();
+			}
+		}
+		return Result;
+	}
+
+	FORCEINLINE_DEBUGGABLE void SetLocalBoundShaderState(FLocalBoundShaderState LocalBoundShaderState)
+	{
+		if (Bypass())
+		{
+			SetBoundShaderState_Internal(LocalBoundShaderState.BypassBSS);
+			return;
+		}
+		auto* Cmd = AddCommand<FRHICommandSetLocalBoundShaderState>();
+		Cmd->Set(this, this->UID, LocalBoundShaderState);
+	}
+
+	FORCEINLINE_DEBUGGABLE void BuildAndSetGlobalBoundShaderState(
+		void(*InSetGlobalBoundShaderState_InternalPtr)(FGlobalBoundShaderState BoundShaderState),
+		FGlobalBoundShaderState& GlobalBoundShaderState,
+		FVertexDeclarationRHIParamRef VertexDeclarationRHI,
+		FShader* VertexShader,
+		FShader* PixelShader,
+		FShader* GeometryShader
+		)
+	{
+		SetGlobalBoundShaderState_InternalPtr = InSetGlobalBoundShaderState_InternalPtr; // total hack to work around module issues
+		if (!GlobalBoundShaderState)
+		{
+			FGlobalBoundShaderStateWorkArea* NewGlobalBoundShaderState = new FGlobalBoundShaderStateWorkArea();
+			NewGlobalBoundShaderState->Args.VertexDeclarationRHI = VertexDeclarationRHI;
+			NewGlobalBoundShaderState->Args.VertexShader = VertexShader;
+			NewGlobalBoundShaderState->Args.PixelShader = PixelShader;
+			NewGlobalBoundShaderState->Args.GeometryShader = GeometryShader;
+			FPlatformMisc::MemoryBarrier();
+
+			FGlobalBoundShaderStateWorkArea* OldGlobalBoundShaderState = (FGlobalBoundShaderStateWorkArea*)FPlatformAtomics::InterlockedCompareExchangePointer((void**)&GlobalBoundShaderState, NewGlobalBoundShaderState, nullptr);
+
+			if (OldGlobalBoundShaderState != nullptr)
+			{
+				//we lost
+				delete NewGlobalBoundShaderState;
+				check(OldGlobalBoundShaderState == GlobalBoundShaderState);
+			}
+			else
+			{
+				// we won
+
+				// this is not even slightly threadsafe, a "loser" might destroy the resource before we get here, and besides this is global; we are never going to destory these anyway.
+
+				if (VertexDeclarationRHI)
+				{
+					VertexDeclarationRHI->AddRef();
+				}
+			}
+		}
+		check(
+			VertexDeclarationRHI == GlobalBoundShaderState->Args.VertexDeclarationRHI &&
+			VertexShader == GlobalBoundShaderState->Args.VertexShader &&
+			PixelShader == GlobalBoundShaderState->Args.PixelShader &&
+			GeometryShader == GlobalBoundShaderState->Args.GeometryShader
+			);
+		if (Bypass())
+		{
+			InSetGlobalBoundShaderState_InternalPtr(GlobalBoundShaderState);
+			return;
+		}
+		auto* Cmd = AddCommand<FRHICommandSetGlobalBoundShaderState>();
+		Cmd->Set(GlobalBoundShaderState);
+	}
+
 	template <typename TShaderRHIParamRef>
 	FORCEINLINE_DEBUGGABLE void SetShaderUniformBuffer(TShaderRHIParamRef Shader, uint32 BaseIndex, FUniformBufferRHIParamRef UniformBuffer)
 	{
@@ -643,7 +1048,10 @@ public:
 		auto* Cmd = AddCommand<FRHICommandSetShaderResourceViewParameter>();
 		Cmd->Set<TShaderRHIParamRef>(Shader, SamplerIndex, SRV);
 		Shader->AddRef();
-		SRV->AddRef();
+		if (SRV)
+		{
+			SRV->AddRef();
+		}
 	}
 
 	template <typename TShaderRHIParamRef>
@@ -767,7 +1175,10 @@ public:
 
 		auto* Cmd = AddCommand<FRHICommandSetStreamSource>();
 		Cmd->Set(StreamIndex, VertexBuffer, Stride, Offset);
-		VertexBuffer->AddRef();
+		if (VertexBuffer)
+		{
+			VertexBuffer->AddRef();
+		}
 	}
 
 	FORCEINLINE_DEBUGGABLE void SetDepthStencilState(FDepthStencilStateRHIParamRef NewStateRHI, uint32 StencilRef = 0)
@@ -925,6 +1336,113 @@ public:
 		DrawUPData.NumVertices = 0;
 	}
 
+	void SetComputeShader(FComputeShaderRHIParamRef ComputeShader)
+	{
+		if (Bypass())
+		{
+			SetComputeShader_Internal(ComputeShader);
+			return;
+		}
+		auto* Cmd = AddCommand<FRHICommandSetComputeShader>();
+		Cmd->Set(ComputeShader);
+		ComputeShader->AddRef();
+	}
+
+	void DispatchComputeShader(uint32 ThreadGroupCountX, uint32 ThreadGroupCountY, uint32 ThreadGroupCountZ)
+	{
+		if (Bypass())
+		{
+			DispatchComputeShader_Internal(ThreadGroupCountX, ThreadGroupCountY, ThreadGroupCountZ);
+			return;
+		}
+		auto* Cmd = AddCommand<FRHICommandDispatchComputeShader>();
+		Cmd->Set(ThreadGroupCountX, ThreadGroupCountY, ThreadGroupCountZ);
+	}
+
+	void DispatchIndirectComputeShader(FVertexBufferRHIParamRef ArgumentBuffer, uint32 ArgumentOffset)
+	{
+		if (Bypass())
+		{
+			DispatchIndirectComputeShader_Internal(ArgumentBuffer, ArgumentOffset);
+			return;
+		}
+		auto* Cmd = AddCommand<FRHICommandDispatchIndirectComputeShader>();
+		Cmd->Set(ArgumentBuffer, ArgumentOffset);
+		ArgumentBuffer->AddRef();
+	}
+
+	void AutomaticCacheFlushAfterComputeShader(bool bEnable)
+	{
+		if (Bypass())
+		{
+			AutomaticCacheFlushAfterComputeShader_Internal(bEnable);
+			return;
+		}
+		auto* Cmd = AddCommand<FRHICommandAutomaticCacheFlushAfterComputeShader>();
+		Cmd->Set(bEnable);
+	}
+
+	void FlushComputeShaderCache()
+	{
+		if (Bypass())
+		{
+			FlushComputeShaderCache_Internal();
+			return;
+		}
+		auto* Cmd = AddCommand<FRHICommandFlushComputeShaderCache>();
+		Cmd->Set();
+	}
+
+	void DrawPrimitiveIndirect(uint32 PrimitiveType, FVertexBufferRHIParamRef ArgumentBuffer, uint32 ArgumentOffset)
+	{
+		if (Bypass())
+		{
+			DrawPrimitiveIndirect_Internal(PrimitiveType, ArgumentBuffer, ArgumentOffset);
+			return;
+		}
+		auto* Cmd = AddCommand<FRHICommandDrawPrimitiveIndirect>();
+		Cmd->Set(PrimitiveType, ArgumentBuffer, ArgumentOffset);
+		ArgumentBuffer->AddRef();
+	}
+
+	void DrawIndexedIndirect(FIndexBufferRHIParamRef IndexBufferRHI, uint32 PrimitiveType, FStructuredBufferRHIParamRef ArgumentsBufferRHI, uint32 DrawArgumentsIndex, uint32 NumInstances)
+	{
+		if (Bypass())
+		{
+			DrawIndexedIndirect_Internal(IndexBufferRHI, PrimitiveType, ArgumentsBufferRHI, DrawArgumentsIndex, NumInstances);
+			return;
+		}
+		auto* Cmd = AddCommand<FRHICommandDrawIndexedIndirect>();
+		Cmd->Set(IndexBufferRHI, PrimitiveType, ArgumentsBufferRHI, DrawArgumentsIndex, NumInstances);
+		IndexBufferRHI->AddRef();
+		ArgumentsBufferRHI->AddRef();
+	}
+
+	void DrawIndexedPrimitiveIndirect(uint32 PrimitiveType, FIndexBufferRHIParamRef IndexBuffer, FVertexBufferRHIParamRef ArgumentsBuffer, uint32 ArgumentOffset)
+	{
+		if (Bypass())
+		{
+			DrawIndexedPrimitiveIndirect_Internal(PrimitiveType, IndexBuffer, ArgumentsBuffer, ArgumentOffset);
+			return;
+		}
+		auto* Cmd = AddCommand<FRHICommandDrawIndexedPrimitiveIndirect>();
+		Cmd->Set(PrimitiveType, IndexBuffer, ArgumentsBuffer, ArgumentOffset);
+		IndexBuffer->AddRef();
+		ArgumentsBuffer->AddRef();
+	}
+
+	void EnableDepthBoundsTest(bool bEnable, float MinDepth, float MaxDepth)
+	{
+		if (Bypass())
+		{
+			EnableDepthBoundsTest_Internal(bEnable, MinDepth, MaxDepth);
+			return;
+		}
+		auto* Cmd = AddCommand<FRHICommandEnableDepthBoundsTest>();
+		Cmd->Set(bEnable, MinDepth, MaxDepth);
+
+	}
+
 
 
 	const SIZE_T GetUsedMemory() const;
@@ -932,6 +1450,10 @@ public:
 private:
 	bool bExecuting;
 	uint32 NumCommands;
+	uint32 UID;
+	static FThreadSafeCounter UIDCounter;
+	void(*SetGlobalBoundShaderState_InternalPtr)(FGlobalBoundShaderState BoundShaderState);
+
 
 	FORCEINLINE_DEBUGGABLE uint8* Alloc(SIZE_T InSize)
 	{
@@ -958,6 +1480,7 @@ private:
 		bExecuting = false;
 		NumCommands = 0;
 		MemManager.Reset();
+		UID = UIDCounter.Increment();
 	}
 
 	inline bool Bypass();
diff --git a/Engine/Source/Runtime/RHI/Public/RHIMethods.h b/Engine/Source/Runtime/RHI/Public/RHIMethods.h
--- a/Engine/Source/Runtime/RHI/Public/RHIMethods.h
+++ b/Engine/Source/Runtime/RHI/Public/RHIMethods.h
@@ -402,7 +402,7 @@ DEFINE_RHIMETHOD_GLOBAL_1(
  * @param GeometryShader - existing geometry shader
  * @param PixelShader - existing pixel shader
  */
-DEFINE_RHIMETHOD_GLOBAL_6(
+DEFINE_RHIMETHOD_CMDLIST_6(
 	FBoundShaderStateRHIRef,CreateBoundShaderState,
 	FVertexDeclarationRHIParamRef,VertexDeclaration,
 	FVertexShaderRHIParamRef,VertexShader,
@@ -417,13 +417,13 @@ DEFINE_RHIMETHOD_GLOBAL_6(
  *Sets the current compute shader.  Mostly for compliance with platforms
  *that require shader setting before resource binding.
  */
-DEFINE_RHIMETHOD_1(
+DEFINE_RHIMETHOD_CMDLIST_1(
 	void,SetComputeShader,
 	FComputeShaderRHIParamRef,ComputeShader,
 	,
 	);
 
-DEFINE_RHIMETHOD_3(
+DEFINE_RHIMETHOD_CMDLIST_3(
 	void,DispatchComputeShader,	
 	uint32,ThreadGroupCountX,
 	uint32,ThreadGroupCountY,
@@ -431,20 +431,20 @@ DEFINE_RHIMETHOD_3(
 	,
 	);
 
-DEFINE_RHIMETHOD_2(
+DEFINE_RHIMETHOD_CMDLIST_2(
 	void,DispatchIndirectComputeShader,	
 	FVertexBufferRHIParamRef,ArgumentBuffer,
 	uint32,ArgumentOffset,
 	,
 	);
 
-DEFINE_RHIMETHOD_1(
+DEFINE_RHIMETHOD_CMDLIST_1(
 	void,AutomaticCacheFlushAfterComputeShader,
 	bool,bEnable,
 	,
 	);
 
-DEFINE_RHIMETHOD_0(
+DEFINE_RHIMETHOD_CMDLIST_0(
 	void,FlushComputeShaderCache,
 	,
 	);
@@ -1736,7 +1736,7 @@ DEFINE_RHIMETHOD_CMDLIST_4(
 	uint32,NumInstances,
 	,
 	);
-DEFINE_RHIMETHOD_3(
+DEFINE_RHIMETHOD_CMDLIST_3(
 	void,DrawPrimitiveIndirect,
 	uint32,PrimitiveType,
 	FVertexBufferRHIParamRef,ArgumentBuffer,
@@ -1744,7 +1744,7 @@ DEFINE_RHIMETHOD_3(
 	,
 	);
 
-DEFINE_RHIMETHOD_5(
+DEFINE_RHIMETHOD_CMDLIST_5(
 	void,DrawIndexedIndirect,
 	FIndexBufferRHIParamRef,IndexBufferRHI,
 	uint32,PrimitiveType,
@@ -1767,7 +1767,7 @@ DEFINE_RHIMETHOD_CMDLIST_8(
 	uint32,NumInstances,
 	,
 	);
-DEFINE_RHIMETHOD_4(
+DEFINE_RHIMETHOD_CMDLIST_4(
 	void,DrawIndexedPrimitiveIndirect,
 	uint32,PrimitiveType,
 	FIndexBufferRHIParamRef,IndexBuffer,
@@ -1956,7 +1956,7 @@ DEFINE_RHIMETHOD_GLOBAL_1(
  * @param MaxDepth	The maximum depth for depth bounds test.
  *					The valid values for fMinDepth and fMaxDepth are such that 0 <= fMinDepth <= fMaxDepth <= 1
  */
-DEFINE_RHIMETHOD_3(
+DEFINE_RHIMETHOD_CMDLIST_3(
 	void,EnableDepthBoundsTest,
 	bool,bEnable,
 	float,MinDepth,
diff --git a/Engine/Source/Runtime/Renderer/Private/AtmosphereRendering.cpp b/Engine/Source/Runtime/Renderer/Private/AtmosphereRendering.cpp
--- a/Engine/Source/Runtime/Renderer/Private/AtmosphereRendering.cpp
+++ b/Engine/Source/Runtime/Renderer/Private/AtmosphereRendering.cpp
@@ -351,7 +351,7 @@ void FSceneRenderer::InitAtmosphereConstants()
 
 FGlobalBoundShaderState AtmosphereBoundShaderState[EAtmosphereRenderFlag::E_RenderFlagMax];
 
-void SetAtmosphericFogShaders(FRHICommandListImmediate& RHICmdList, FScene* Scene, const FViewInfo& View, TRefCountPtr<IPooledRenderTarget>& LightShaftOcclusion)
+void SetAtmosphericFogShaders(FRHICommandList& RHICmdList, FScene* Scene, const FViewInfo& View, TRefCountPtr<IPooledRenderTarget>& LightShaftOcclusion)
 {
 	TShaderMapRef<FAtmosphericVS> VertexShader(GetGlobalShaderMap());
 	FAtmosphericFogPS* PixelShader = NULL;
diff --git a/Engine/Source/Runtime/Renderer/Private/BasePassRendering.cpp b/Engine/Source/Runtime/Renderer/Private/BasePassRendering.cpp
--- a/Engine/Source/Runtime/Renderer/Private/BasePassRendering.cpp
+++ b/Engine/Source/Runtime/Renderer/Private/BasePassRendering.cpp
@@ -295,11 +295,8 @@ public:
 			false,
 			Parameters.bEditorCompositeDepthTest
 			);
-		DrawingPolicy.DrawShared(
-			RHICmdList, 
-			&View,
-			DrawingPolicy.CreateBoundShaderState(View.GetFeatureLevel())
-			);
+		RHICmdList.BuildAndSetLocalBoundShaderState(DrawingPolicy.GetBoundShaderStateInput(View.GetFeatureLevel()));
+		DrawingPolicy.SetSharedState(RHICmdList, &View);
 
 		for( int32 BatchElementIndex = 0, Num = Parameters.Mesh.Elements.Num(); BatchElementIndex < Num; BatchElementIndex++ )
 		{
diff --git a/Engine/Source/Runtime/Renderer/Private/BasePassRendering.h b/Engine/Source/Runtime/Renderer/Private/BasePassRendering.h
--- a/Engine/Source/Runtime/Renderer/Private/BasePassRendering.h
+++ b/Engine/Source/Runtime/Renderer/Private/BasePassRendering.h
@@ -549,10 +549,8 @@ public:
 			LightMapPolicy == Other.LightMapPolicy;
 	}
 
-	void DrawShared(FRHICommandList& RHICmdList, const FSceneView* View,FBoundShaderStateRHIParamRef BoundShaderState) const
+	void SetSharedState(FRHICommandList& RHICmdList, const FSceneView* View) const
 	{
-		RHICmdList.SetBoundShaderState(BoundShaderState);
-
 		// Set the light-map policy.
 		LightMapPolicy.Set(RHICmdList, VertexShader,bOverrideWithShaderComplexity ? NULL : PixelShader,VertexShader,PixelShader,VertexFactory,MaterialRenderProxy,View);
 
@@ -620,7 +618,7 @@ public:
 	* @param DynamicStride - optional stride for dynamic vertex data
 	* @return new bound shader state object
 	*/
-	FBoundShaderStateRHIRef CreateBoundShaderState(ERHIFeatureLevel::Type InFeatureLevel)
+	FBoundShaderStateInput GetBoundShaderStateInput(ERHIFeatureLevel::Type InFeatureLevel)
 	{
 		FPixelShaderRHIParamRef PixelShaderRHIRef = PixelShader->GetPixelShader();
 
@@ -631,17 +629,14 @@ public:
 			PixelShaderRHIRef = ShaderComplexityAccumulatePixelShader->GetPixelShader();
 		}
 #endif
-		FBoundShaderStateRHIRef BoundShaderState;
 
-		BoundShaderState = RHICreateBoundShaderState(
+		return FBoundShaderStateInput(
 			FMeshDrawingPolicy::GetVertexDeclaration(), 
 			VertexShader->GetVertexShader(),
 			GETSAFERHISHADER_HULL(HullShader), 
 			GETSAFERHISHADER_DOMAIN(DomainShader), 
 			PixelShaderRHIRef,
 			FGeometryShaderRHIRef());
-
-		return BoundShaderState;
 	}
 
 	void SetMeshRenderState(
diff --git a/Engine/Source/Runtime/Renderer/Private/ClearQuad.cpp b/Engine/Source/Runtime/Renderer/Private/ClearQuad.cpp
--- a/Engine/Source/Runtime/Renderer/Private/ClearQuad.cpp
+++ b/Engine/Source/Runtime/Renderer/Private/ClearQuad.cpp
@@ -6,7 +6,7 @@
 FGlobalBoundShaderState GClearMRTBoundShaderState[8];
 
 // TODO support ExcludeRect
-void DrawClearQuadMRT(FRHICommandListImmediate& RHICmdList, bool bClearColor, int32 NumClearColors, const FLinearColor* ClearColorArray, bool bClearDepth, float Depth, bool bClearStencil, uint32 Stencil)
+void DrawClearQuadMRT(FRHICommandList& RHICmdList, bool bClearColor, int32 NumClearColors, const FLinearColor* ClearColorArray, bool bClearDepth, float Depth, bool bClearStencil, uint32 Stencil)
 {
 	// Set new states
 	FBlendStateRHIParamRef BlendStateRHI;
diff --git a/Engine/Source/Runtime/Renderer/Private/ClearQuad.h b/Engine/Source/Runtime/Renderer/Private/ClearQuad.h
--- a/Engine/Source/Runtime/Renderer/Private/ClearQuad.h
+++ b/Engine/Source/Runtime/Renderer/Private/ClearQuad.h
@@ -3,9 +3,9 @@
 #pragma once
 
 // TODO support ExcludeRect
-void DrawClearQuadMRT(FRHICommandListImmediate& RHICmdList, bool bClearColor, int32 NumClearColors, const FLinearColor* ClearColorArray, bool bClearDepth, float Depth, bool bClearStencil, uint32 Stencil);
+void DrawClearQuadMRT(FRHICommandList& RHICmdList, bool bClearColor, int32 NumClearColors, const FLinearColor* ClearColorArray, bool bClearDepth, float Depth, bool bClearStencil, uint32 Stencil);
 
-inline void DrawClearQuad(FRHICommandListImmediate& RHICmdList, bool bClearColor, const FLinearColor& Color, bool bClearDepth, float Depth, bool bClearStencil, uint32 Stencil)
+inline void DrawClearQuad(FRHICommandList& RHICmdList, bool bClearColor, const FLinearColor& Color, bool bClearDepth, float Depth, bool bClearStencil, uint32 Stencil)
 {
 	DrawClearQuadMRT(RHICmdList, bClearColor, 1, &Color, bClearDepth, Depth, bClearStencil, Stencil);
 }
\ No newline at end of file
diff --git a/Engine/Source/Runtime/Renderer/Private/CompositionLighting/PostProcessDeferredDecals.cpp b/Engine/Source/Runtime/Renderer/Private/CompositionLighting/PostProcessDeferredDecals.cpp
--- a/Engine/Source/Runtime/Renderer/Private/CompositionLighting/PostProcessDeferredDecals.cpp
+++ b/Engine/Source/Runtime/Renderer/Private/CompositionLighting/PostProcessDeferredDecals.cpp
@@ -609,18 +609,7 @@ void SetShader(const FRenderingCompositePassContext& Context, const FTransientDe
 	const FMaterialShaderMap* MaterialShaderMap = DecalData.MaterialResource->GetRenderingThreadShaderMap();
 	FDeferredDecalPS* PixelShader = MaterialShaderMap->GetShader<FDeferredDecalPS>();
 
-	// This was cached but when changing the material (e.g. editor) it wasn't updated.
-	// This will change with upcoming multi threaded rendering changes.
-	FBoundShaderStateRHIRef BoundShaderState;
-	{
-		uint32 Strides[MaxVertexElementCount];
-		FMemory::Memzero(Strides, sizeof(Strides));
-		Strides[0] = sizeof(FVector);
-
-		BoundShaderState = Context.RHICmdList.CreateBoundShaderState(GetVertexDeclarationFVector3(), VertexShader->GetVertexShader(), FHullShaderRHIRef(), FDomainShaderRHIRef(), PixelShader->GetPixelShader(), FGeometryShaderRHIRef());
-	}
-
-	Context.RHICmdList.SetBoundShaderState(BoundShaderState);
+	Context.RHICmdList.SetLocalBoundShaderState(Context.RHICmdList.BuildLocalBoundShaderState(GetVertexDeclarationFVector3(), VertexShader->GetVertexShader(), FHullShaderRHIRef(), FDomainShaderRHIRef(), PixelShader->GetPixelShader(), FGeometryShaderRHIRef()));
 
 	PixelShader->SetParameters(Context.RHICmdList, View, DecalData.MaterialProxy, *DecalData.DecalProxy);
 }
@@ -650,14 +639,7 @@ bool RenderPreStencil(FRenderingCompositePassContext& Context, const FMaterialSh
 
 	FDeferredDecalVS* VertexShader = MaterialShaderMap->GetShader<FDeferredDecalVS>();
 	
-	// This was cached but when changing the material (e.g. editor) it wasn't updated.
-	// This will change with upcoming multi threaded rendering changes.
-	FBoundShaderStateRHIRef BoundShaderState;
-	{
-		BoundShaderState = Context.RHICmdList.CreateBoundShaderState(GetVertexDeclarationFVector3(), VertexShader->GetVertexShader(), FHullShaderRHIRef(), FDomainShaderRHIRef(), NULL, FGeometryShaderRHIRef());
-	}
-
-	Context.RHICmdList.SetBoundShaderState(BoundShaderState);
+	Context.RHICmdList.SetLocalBoundShaderState(Context.RHICmdList.BuildLocalBoundShaderState(GetVertexDeclarationFVector3(), VertexShader->GetVertexShader(), FHullShaderRHIRef(), FDomainShaderRHIRef(), NULL, FGeometryShaderRHIRef()));
 
 	VertexShader->SetParameters(Context.RHICmdList, View, FrustumComponentToClip);
 
diff --git a/Engine/Source/Runtime/Renderer/Private/DeferredShadingRenderer.cpp b/Engine/Source/Runtime/Renderer/Private/DeferredShadingRenderer.cpp
--- a/Engine/Source/Runtime/Renderer/Private/DeferredShadingRenderer.cpp
+++ b/Engine/Source/Runtime/Renderer/Private/DeferredShadingRenderer.cpp
@@ -123,7 +123,7 @@ extern FGlobalBoundShaderState GClearMRTBoundShaderState[8];
 /** 
 * Clears view where Z is still at the maximum value (ie no geometry rendered)
 */
-void FDeferredShadingSceneRenderer::ClearGBufferAtMaxZ(FRHICommandListImmediate& RHICmdList)
+void FDeferredShadingSceneRenderer::ClearGBufferAtMaxZ(FRHICommandList& RHICmdList)
 {
 	// Assumes BeginRenderingSceneColor() has been called before this function
 	SCOPED_DRAW_EVENT(ClearGBufferAtMaxZ, DEC_SCENE_ITEMS);
@@ -1055,7 +1055,7 @@ IMPLEMENT_SHADER_TYPE(,FDownsampleSceneDepthPS,TEXT("DownsampleDepthPixelShader"
 FGlobalBoundShaderState DownsampleDepthBoundShaderState;
 
 /** Updates the downsized depth buffer with the current full resolution depth buffer. */
-void FDeferredShadingSceneRenderer::UpdateDownsampledDepthSurface(FRHICommandListImmediate& RHICmdList)
+void FDeferredShadingSceneRenderer::UpdateDownsampledDepthSurface(FRHICommandList& RHICmdList)
 {
 	if (GSceneRenderTargets.UseDownsizedOcclusionQueries() && IsFeatureLevelSupported(GRHIShaderPlatform, ERHIFeatureLevel::SM3))
 	{
diff --git a/Engine/Source/Runtime/Renderer/Private/DeferredShadingRenderer.h b/Engine/Source/Runtime/Renderer/Private/DeferredShadingRenderer.h
--- a/Engine/Source/Runtime/Renderer/Private/DeferredShadingRenderer.h
+++ b/Engine/Source/Runtime/Renderer/Private/DeferredShadingRenderer.h
@@ -42,7 +42,7 @@ public:
 	void ClearView(FRHICommandListImmediate& RHICmdList);
 
 	/** Clears gbuffer where Z is still at the maximum value (ie no geometry rendered) */
-	void ClearGBufferAtMaxZ(FRHICommandListImmediate& RHICmdList);
+	void ClearGBufferAtMaxZ(FRHICommandList& RHICmdList);
 
 	/** Clears LPVs for all views */
 	void ClearLPVs(FRHICommandListImmediate& RHICmdList);
@@ -222,7 +222,7 @@ private:
 	bool RenderLightMapDensities(FRHICommandListImmediate& RHICmdList);
 
 	/** Updates the downsized depth buffer with the current full resolution depth buffer. */
-	void UpdateDownsampledDepthSurface(FRHICommandListImmediate& RHICmdList);
+	void UpdateDownsampledDepthSurface(FRHICommandList& RHICmdList);
 
 	/**
 	 * Finish rendering a view, writing the contents to ViewFamily.RenderTarget.
diff --git a/Engine/Source/Runtime/Renderer/Private/DepthRendering.cpp b/Engine/Source/Runtime/Renderer/Private/DepthRendering.cpp
--- a/Engine/Source/Runtime/Renderer/Private/DepthRendering.cpp
+++ b/Engine/Source/Runtime/Renderer/Private/DepthRendering.cpp
@@ -169,11 +169,8 @@ FDepthDrawingPolicy::FDepthDrawingPolicy(
 
 }
 
-void FDepthDrawingPolicy::DrawShared(FRHICommandList& RHICmdList, const FSceneView* View,FBoundShaderStateRHIParamRef BoundShaderState) const
+void FDepthDrawingPolicy::SetSharedState(FRHICommandList& RHICmdList, const FSceneView* View) const
 {
-	// Set the actual shader & vertex declaration state
-	RHICmdList.SetBoundShaderState(BoundShaderState);
-
 	// Set the depth-only shader parameters for the material.
 	VertexShader->SetParameters(RHICmdList, MaterialRenderProxy,*MaterialResource,*View);
 	if(HullShader && DomainShader)
@@ -196,9 +193,9 @@ void FDepthDrawingPolicy::DrawShared(FRHICommandList& RHICmdList, const FSceneVi
 * as well as the shaders needed to draw the mesh
 * @return new bound shader state object
 */
-FBoundShaderStateRHIRef FDepthDrawingPolicy::CreateBoundShaderState(ERHIFeatureLevel::Type InFeatureLevel)
+FBoundShaderStateInput FDepthDrawingPolicy::GetBoundShaderStateInput(ERHIFeatureLevel::Type InFeatureLevel)
 {
-	return RHICreateBoundShaderState(
+	return FBoundShaderStateInput(
 		FMeshDrawingPolicy::GetVertexDeclaration(), 
 		VertexShader->GetVertexShader(),
 		GETSAFERHISHADER_HULL(HullShader), 
@@ -257,11 +254,8 @@ FPositionOnlyDepthDrawingPolicy::FPositionOnlyDepthDrawingPolicy(
 	VertexShader = InMaterialResource.GetShader<TDepthOnlyVS<true> >(InVertexFactory->GetType());
 }
 
-void FPositionOnlyDepthDrawingPolicy::DrawShared(FRHICommandList& RHICmdList, const FSceneView* View,FBoundShaderStateRHIParamRef BoundShaderState) const
+void FPositionOnlyDepthDrawingPolicy::SetSharedState(FRHICommandList& RHICmdList, const FSceneView* View) const
 {
-	// Set the actual shader & vertex declaration state
-	RHICmdList.SetBoundShaderState( BoundShaderState);
-
 	// Set the depth-only shader parameters for the material.
 	VertexShader->SetParameters(RHICmdList, MaterialRenderProxy,*MaterialResource,*View);
 
@@ -274,13 +268,13 @@ void FPositionOnlyDepthDrawingPolicy::DrawShared(FRHICommandList& RHICmdList, co
 * as well as the shaders needed to draw the mesh
 * @return new bound shader state object
 */
-FBoundShaderStateRHIRef FPositionOnlyDepthDrawingPolicy::CreateBoundShaderState(ERHIFeatureLevel::Type InFeatureLevel)
+FBoundShaderStateInput FPositionOnlyDepthDrawingPolicy::GetBoundShaderStateInput(ERHIFeatureLevel::Type InFeatureLevel)
 {
 	FVertexDeclarationRHIParamRef VertexDeclaration;
 	VertexDeclaration = VertexFactory->GetPositionDeclaration();
 
 	checkSlow(MaterialRenderProxy->GetMaterial(InFeatureLevel)->GetBlendMode() == BLEND_Opaque);
-	return RHICreateBoundShaderState(VertexDeclaration, VertexShader->GetVertexShader(), FHullShaderRHIRef(), FDomainShaderRHIRef(), FPixelShaderRHIRef(), FGeometryShaderRHIRef());
+	return FBoundShaderStateInput(VertexDeclaration, VertexShader->GetVertexShader(), FHullShaderRHIRef(), FDomainShaderRHIRef(), FPixelShaderRHIRef(), FGeometryShaderRHIRef());
 }
 
 void FPositionOnlyDepthDrawingPolicy::SetMeshRenderState(
@@ -404,7 +398,9 @@ bool FDepthDrawingPolicyFactory::DrawMesh(
 			//render opaque primitives that support a separate position-only vertex buffer
 			const FMaterialRenderProxy* DefaultProxy = UMaterial::GetDefaultMaterial(MD_Surface)->GetRenderProxy(false);
 			FPositionOnlyDepthDrawingPolicy DrawingPolicy(Mesh.VertexFactory, DefaultProxy, *DefaultProxy->GetMaterial(View.GetFeatureLevel()), Material->IsTwoSided(), Material->IsWireframe());
-			DrawingPolicy.DrawShared(RHICmdList, &View,DrawingPolicy.CreateBoundShaderState(View.GetFeatureLevel()));
+			RHICmdList.BuildAndSetLocalBoundShaderState(DrawingPolicy.GetBoundShaderStateInput(View.GetFeatureLevel()));
+			DrawingPolicy.SetSharedState(RHICmdList, &View);
+
 			int32 BatchElementIndex = 0;
 			uint64 Mask = BatchElementMask;
 			do
@@ -446,7 +442,8 @@ bool FDepthDrawingPolicyFactory::DrawMesh(
 				}
 
 				FDepthDrawingPolicy DrawingPolicy(Mesh.VertexFactory, MaterialRenderProxy, *MaterialRenderProxy->GetMaterial(View.GetFeatureLevel()), Material->IsTwoSided());
-				DrawingPolicy.DrawShared(RHICmdList, &View,DrawingPolicy.CreateBoundShaderState(View.GetFeatureLevel()));
+				RHICmdList.BuildAndSetLocalBoundShaderState(DrawingPolicy.GetBoundShaderStateInput(View.GetFeatureLevel()));
+				DrawingPolicy.SetSharedState(RHICmdList, &View);
 
 				int32 BatchElementIndex = 0;
 				uint64 Mask = BatchElementMask;
diff --git a/Engine/Source/Runtime/Renderer/Private/DepthRendering.h b/Engine/Source/Runtime/Renderer/Private/DepthRendering.h
--- a/Engine/Source/Runtime/Renderer/Private/DepthRendering.h
+++ b/Engine/Source/Runtime/Renderer/Private/DepthRendering.h
@@ -44,7 +44,7 @@ public:
 			&& PixelShader == Other.PixelShader;
 	}
 
-	void DrawShared(FRHICommandList& RHICmdList, const FSceneView* View,FBoundShaderStateRHIParamRef BoundShaderState) const;
+	void SetSharedState(FRHICommandList& RHICmdList, const FSceneView* View) const;
 
 	/** 
 	* Create bound shader state using the vertex decl from the mesh draw policy
@@ -52,7 +52,7 @@ public:
 	* @param DynamicStride - optional stride for dynamic vertex data
 	* @return new bound shader state object
 	*/
-	FBoundShaderStateRHIRef CreateBoundShaderState(ERHIFeatureLevel::Type InFeatureLevel);
+	FBoundShaderStateInput GetBoundShaderStateInput(ERHIFeatureLevel::Type InFeatureLevel);
 
 	void SetMeshRenderState(
 		FRHICommandList& RHICmdList, 
@@ -97,14 +97,14 @@ public:
 		return FMeshDrawingPolicy::Matches(Other) && VertexShader == Other.VertexShader;
 	}
 
-	void DrawShared(FRHICommandList& RHICmdList, const FSceneView* View,FBoundShaderStateRHIParamRef BoundShaderState) const;
+	void SetSharedState(FRHICommandList& RHICmdList, const FSceneView* View) const;
 
 	/** 
 	* Create bound shader state using the vertex decl from the mesh draw policy
 	* as well as the shaders needed to draw the mesh
 	* @return new bound shader state object
 	*/
-	FBoundShaderStateRHIRef CreateBoundShaderState(ERHIFeatureLevel::Type InFeatureLevel);
+	FBoundShaderStateInput GetBoundShaderStateInput(ERHIFeatureLevel::Type InFeatureLevel);
 
 	void SetMeshRenderState(
 		FRHICommandList& RHICmdList, 
diff --git a/Engine/Source/Runtime/Renderer/Private/DistanceFieldSurfaceCacheLighting.cpp b/Engine/Source/Runtime/Renderer/Private/DistanceFieldSurfaceCacheLighting.cpp
--- a/Engine/Source/Runtime/Renderer/Private/DistanceFieldSurfaceCacheLighting.cpp
+++ b/Engine/Source/Runtime/Renderer/Private/DistanceFieldSurfaceCacheLighting.cpp
@@ -2054,7 +2054,7 @@ void AllocateOrReuseAORenderTarget(TRefCountPtr<IPooledRenderTarget>& Target, co
 }
 
 void UpdateHistory(
-	FRHICommandListImmediate& RHICmdList,
+	FRHICommandList& RHICmdList,
 	FViewInfo& View, 
 	const TCHAR* HistoryRTName,
 	/** Contains last frame's history, if non-NULL.  This will be updated with the new frame's history. */
@@ -2120,7 +2120,7 @@ void UpdateHistory(
 	}
 }
 
-void PostProcessBentNormalAO(FRHICommandListImmediate& RHICmdList, TArray<FViewInfo>& Views, FSceneRenderTargetItem& IrradianceCacheInterpolation, TRefCountPtr<IPooledRenderTarget>& AOOutput)
+void PostProcessBentNormalAO(FRHICommandList& RHICmdList, TArray<FViewInfo>& Views, FSceneRenderTargetItem& IrradianceCacheInterpolation, TRefCountPtr<IPooledRenderTarget>& AOOutput)
 {
 	TRefCountPtr<IPooledRenderTarget> DistanceFieldAOBentNormal;
 	AllocateOrReuseAORenderTarget(DistanceFieldAOBentNormal, TEXT("DistanceFieldBentNormalAO"));
@@ -2279,7 +2279,7 @@ private:
 
 IMPLEMENT_SHADER_TYPE(,FDistanceFieldAOUpsamplePS,TEXT("DistanceFieldSurfaceCacheLighting"),TEXT("AOUpsamplePS"),SF_Pixel);
 
-void UpsampleBentNormalAO(FRHICommandListImmediate& RHICmdList, const TArray<FViewInfo>& Views, TRefCountPtr<IPooledRenderTarget>& DistanceFieldAOBentNormal)
+void UpsampleBentNormalAO(FRHICommandList& RHICmdList, const TArray<FViewInfo>& Views, TRefCountPtr<IPooledRenderTarget>& DistanceFieldAOBentNormal)
 {
 	for (int32 ViewIndex = 0; ViewIndex < Views.Num(); ViewIndex++)
 	{
diff --git a/Engine/Source/Runtime/Renderer/Private/DistortionRendering.cpp b/Engine/Source/Runtime/Renderer/Private/DistortionRendering.cpp
--- a/Engine/Source/Runtime/Renderer/Private/DistortionRendering.cpp
+++ b/Engine/Source/Runtime/Renderer/Private/DistortionRendering.cpp
@@ -295,14 +295,14 @@ public:
 	* @param CI - The command interface to execute the draw commands on.
 	* @param View - The view of the scene being drawn.
 	*/
-	void DrawShared(FRHICommandList& RHICmdList, const FSceneView* View,FBoundShaderStateRHIParamRef BoundShaderState) const;
+	void SetSharedState(FRHICommandList& RHICmdList, const FSceneView* View) const;
 	
 	/** 
 	* Create bound shader state using the vertex decl from the mesh draw policy
 	* as well as the shaders needed to draw the mesh
 	* @return new bound shader state object
 	*/
-	FBoundShaderStateRHIRef CreateBoundShaderState(ERHIFeatureLevel::Type InFeatureLevel);
+	FBoundShaderStateInput GetBoundShaderStateInput(ERHIFeatureLevel::Type InFeatureLevel);
 
 	/**
 	* Sets the render states for drawing a mesh.
@@ -403,18 +403,14 @@ bool TDistortionMeshDrawingPolicy<DistortMeshPolicy>::Matches(
 * @param View - The view of the scene being drawn.
 */
 template<class DistortMeshPolicy>
-void TDistortionMeshDrawingPolicy<DistortMeshPolicy>::DrawShared(
+void TDistortionMeshDrawingPolicy<DistortMeshPolicy>::SetSharedState(
 	FRHICommandList& RHICmdList, 
-	const FSceneView* View,
-	FBoundShaderStateRHIParamRef BoundShaderState
+	const FSceneView* View
 	) const
 {
 	// Set shared mesh resources
 	FMeshDrawingPolicy::DrawShared(RHICmdList, View);
 
-	// Set the actual shader & vertex declaration state
-	RHICmdList.SetBoundShaderState(BoundShaderState);
-
 	// Set the translucent shader parameters for the material instance
 	VertexShader->SetParameters(RHICmdList, VertexFactory,MaterialRenderProxy,View);
 
@@ -447,7 +443,7 @@ void TDistortionMeshDrawingPolicy<DistortMeshPolicy>::DrawShared(
 * @return new bound shader state object
 */
 template<class DistortMeshPolicy>
-FBoundShaderStateRHIRef TDistortionMeshDrawingPolicy<DistortMeshPolicy>::CreateBoundShaderState(ERHIFeatureLevel::Type InFeatureLevel)
+FBoundShaderStateInput TDistortionMeshDrawingPolicy<DistortMeshPolicy>::GetBoundShaderStateInput(ERHIFeatureLevel::Type InFeatureLevel)
 {
 	FPixelShaderRHIParamRef PixelShaderRHIRef = NULL;
 
@@ -467,9 +463,7 @@ FBoundShaderStateRHIRef TDistortionMeshDrawingPolicy<DistortMeshPolicy>::CreateB
 		PixelShaderRHIRef = DistortPixelShader->GetPixelShader();
 	}
 
-	FBoundShaderStateRHIRef BoundShaderState;
-
-	BoundShaderState = RHICreateBoundShaderState(
+	return FBoundShaderStateInput(
 		FMeshDrawingPolicy::GetVertexDeclaration(), 
 		VertexShader->GetVertexShader(),
 		GETSAFERHISHADER_HULL(HullShader), 
@@ -477,7 +471,6 @@ FBoundShaderStateRHIRef TDistortionMeshDrawingPolicy<DistortMeshPolicy>::CreateB
 		PixelShaderRHIRef,
 		FGeometryShaderRHIRef());
 
-	return BoundShaderState;
 }
 
 /**
@@ -606,8 +599,9 @@ bool TDistortionMeshDrawingPolicyFactory<DistortMeshPolicy>::DrawDynamicMesh(
 			bInitializeOffsets,
 			View.Family->EngineShowFlags.ShaderComplexity
 			);
-		DrawingPolicy.DrawShared(RHICmdList, &View,DrawingPolicy.CreateBoundShaderState(View.GetFeatureLevel()));
-		for( int32 BatchElementIndex=0;BatchElementIndex < Mesh.Elements.Num();BatchElementIndex++ )
+		RHICmdList.BuildAndSetLocalBoundShaderState(DrawingPolicy.GetBoundShaderStateInput(View.GetFeatureLevel()));
+		DrawingPolicy.SetSharedState(RHICmdList, &View);
+		for (int32 BatchElementIndex = 0; BatchElementIndex < Mesh.Elements.Num(); BatchElementIndex++)
 		{
 			DrawingPolicy.SetMeshRenderState(RHICmdList, View,PrimitiveSceneProxy,Mesh,BatchElementIndex,bBackFace,typename TDistortionMeshDrawingPolicy<DistortMeshPolicy>::ElementDataType());
 			DrawingPolicy.DrawMesh(RHICmdList, Mesh,BatchElementIndex);
@@ -650,7 +644,8 @@ bool TDistortionMeshDrawingPolicyFactory<DistortMeshPolicy>::DrawStaticMesh(
 			bInitializeOffsets,
 			View->Family->EngineShowFlags.ShaderComplexity
 			);
-		DrawingPolicy.DrawShared(RHICmdList, View,DrawingPolicy.CreateBoundShaderState(FeatureLevel));
+		RHICmdList.BuildAndSetLocalBoundShaderState(DrawingPolicy.GetBoundShaderStateInput(View->GetFeatureLevel()));
+		DrawingPolicy.SetSharedState(RHICmdList, View);
 		int32 BatchElementIndex = 0;
 		do
 		{
diff --git a/Engine/Source/Runtime/Renderer/Private/FogRendering.cpp b/Engine/Source/Runtime/Renderer/Private/FogRendering.cpp
--- a/Engine/Source/Runtime/Renderer/Private/FogRendering.cpp
+++ b/Engine/Source/Runtime/Renderer/Private/FogRendering.cpp
@@ -275,7 +275,7 @@ void FSceneRenderer::InitFogConstants()
 FGlobalBoundShaderState ExponentialBoundShaderState;
 
 /** Sets the bound shader state for either the per-pixel or per-sample fog pass. */
-void SetFogShaders(FRHICommandListImmediate& RHICmdList, FScene* Scene, const FViewInfo& View, FLightShaftsOutput LightShaftsOutput)
+void SetFogShaders(FRHICommandList& RHICmdList, FScene* Scene, const FViewInfo& View, FLightShaftsOutput LightShaftsOutput)
 {
 	if (Scene->ExponentialFogs.Num() > 0)
 	{
diff --git a/Engine/Source/Runtime/Renderer/Private/ForwardBasePassRendering.cpp b/Engine/Source/Runtime/Renderer/Private/ForwardBasePassRendering.cpp
--- a/Engine/Source/Runtime/Renderer/Private/ForwardBasePassRendering.cpp
+++ b/Engine/Source/Runtime/Renderer/Private/ForwardBasePassRendering.cpp
@@ -171,11 +171,8 @@ public:
 			Parameters.TextureMode,
 			View.Family->EngineShowFlags.ShaderComplexity
 			);
-		DrawingPolicy.DrawShared(
-			RHICmdList, 
-			&View,
-			DrawingPolicy.CreateBoundShaderState(View.GetFeatureLevel())
-			);
+		RHICmdList.BuildAndSetLocalBoundShaderState(DrawingPolicy.GetBoundShaderStateInput(View.GetFeatureLevel()));
+		DrawingPolicy.SetSharedState(RHICmdList, &View);
 
 		for( int32 BatchElementIndex=0;BatchElementIndex<Parameters.Mesh.Elements.Num();BatchElementIndex++ )
 		{
diff --git a/Engine/Source/Runtime/Renderer/Private/ForwardBasePassRendering.h b/Engine/Source/Runtime/Renderer/Private/ForwardBasePassRendering.h
--- a/Engine/Source/Runtime/Renderer/Private/ForwardBasePassRendering.h
+++ b/Engine/Source/Runtime/Renderer/Private/ForwardBasePassRendering.h
@@ -264,11 +264,8 @@ public:
 			SceneTextureMode == Other.SceneTextureMode;
 	}
 
-	void DrawShared(FRHICommandList& RHICmdList, const FSceneView* View,FBoundShaderStateRHIParamRef BoundShaderState) const
+	void SetSharedState(FRHICommandList& RHICmdList, const FSceneView* View) const
 	{
-		// Set the actual shader & vertex declaration state
-		RHICmdList.SetBoundShaderState(BoundShaderState);
-
 		VertexShader->SetParameters(RHICmdList, MaterialRenderProxy, VertexFactory, *MaterialResource, *View, SceneTextureMode);
 
 #if !(UE_BUILD_SHIPPING || UE_BUILD_TEST)
@@ -326,7 +323,7 @@ public:
 	* as well as the shaders needed to draw the mesh
 	* @return new bound shader state object
 	*/
-	FBoundShaderStateRHIRef CreateBoundShaderState(ERHIFeatureLevel::Type InFeatureLevel)
+	FBoundShaderStateInput GetBoundShaderStateInput(ERHIFeatureLevel::Type InFeatureLevel)
 	{
 		FPixelShaderRHIParamRef PixelShaderRHIRef = PixelShader->GetPixelShader();
 
@@ -338,15 +335,13 @@ public:
 		}
 #endif
 
-		FBoundShaderStateRHIRef BoundShaderState = RHICreateBoundShaderState(
+		return FBoundShaderStateInput(
 			FMeshDrawingPolicy::GetVertexDeclaration(), 
 			VertexShader->GetVertexShader(),
 			FHullShaderRHIRef(), 
 			FDomainShaderRHIRef(), 
 			PixelShaderRHIRef,
 			FGeometryShaderRHIRef());
-
-		return BoundShaderState;
 	}
 
 	void SetMeshRenderState(
diff --git a/Engine/Source/Runtime/Renderer/Private/ForwardTranslucentRendering.cpp b/Engine/Source/Runtime/Renderer/Private/ForwardTranslucentRendering.cpp
--- a/Engine/Source/Runtime/Renderer/Private/ForwardTranslucentRendering.cpp
+++ b/Engine/Source/Runtime/Renderer/Private/ForwardTranslucentRendering.cpp
@@ -139,11 +139,8 @@ public:
 			View.Family->EngineShowFlags.ShaderComplexity
 			);
 
-		DrawingPolicy.DrawShared(
-			RHICmdList, 
-			&View,
-			DrawingPolicy.CreateBoundShaderState(View.GetFeatureLevel())
-			);
+		RHICmdList.BuildAndSetLocalBoundShaderState(DrawingPolicy.GetBoundShaderStateInput(View.GetFeatureLevel()));
+		DrawingPolicy.SetSharedState(RHICmdList, &View);
 
 		for (int32 BatchElementIndex = 0; BatchElementIndex<Parameters.Mesh.Elements.Num(); BatchElementIndex++)
 		{
diff --git a/Engine/Source/Runtime/Renderer/Private/GPUBenchmark.cpp b/Engine/Source/Runtime/Renderer/Private/GPUBenchmark.cpp
--- a/Engine/Source/Runtime/Renderer/Private/GPUBenchmark.cpp
+++ b/Engine/Source/Runtime/Renderer/Private/GPUBenchmark.cpp
@@ -125,7 +125,7 @@ IMPLEMENT_SHADER_TYPE(,FPostProcessBenchmarkVS,TEXT("GPUBenchmark"),TEXT("MainBe
 
 
 template <uint32 Method>
-void RunBenchmarkShader(FRHICommandListImmediate& RHICmdList, const FSceneView& View, TRefCountPtr<IPooledRenderTarget>& Src, uint32 Count)
+void RunBenchmarkShader(FRHICommandList& RHICmdList, const FSceneView& View, TRefCountPtr<IPooledRenderTarget>& Src, uint32 Count)
 {
 	TShaderMapRef<FPostProcessBenchmarkVS> VertexShader(GetGlobalShaderMap());
 	TShaderMapRef<FPostProcessBenchmarkPS<Method> > PixelShader(GetGlobalShaderMap());
diff --git a/Engine/Source/Runtime/Renderer/Private/LightFunctionRendering.cpp b/Engine/Source/Runtime/Renderer/Private/LightFunctionRendering.cpp
--- a/Engine/Source/Runtime/Renderer/Private/LightFunctionRendering.cpp
+++ b/Engine/Source/Runtime/Renderer/Private/LightFunctionRendering.cpp
@@ -249,12 +249,7 @@ bool FDeferredShadingSceneRenderer::RenderLightFunctionForMaterial(FRHICommandLi
 		FLightFunctionVS* VertexShader = MaterialShaderMap->GetShader<FLightFunctionVS>();
 		FLightFunctionPS* PixelShader = MaterialShaderMap->GetShader<FLightFunctionPS>();
 
-		// This was cached but when changing the material (e.g. editor) it wasn't updated.
-		// This will change with upcoming multi threaded rendering changes.
-		FBoundShaderStateRHIRef LightFunctionBoundShaderState;
-		{
-			LightFunctionBoundShaderState = RHICreateBoundShaderState(GetVertexDeclarationFVector4(), VertexShader->GetVertexShader(), FHullShaderRHIRef(), FDomainShaderRHIRef(), PixelShader->GetPixelShader(), FGeometryShaderRHIRef());
-		}
+		FLocalBoundShaderState LightFunctionBoundShaderState = RHICmdList.BuildLocalBoundShaderState(GetVertexDeclarationFVector4(), VertexShader->GetVertexShader(), FHullShaderRHIRef(), FDomainShaderRHIRef(), PixelShader->GetPixelShader(), FGeometryShaderRHIRef());
 
 		FSphere LightBounds = LightSceneInfo->Proxy->GetBoundingSphere();
 
@@ -327,7 +322,7 @@ bool FDeferredShadingSceneRenderer::RenderLightFunctionForMaterial(FRHICommandLi
 					LightSceneInfo->Proxy->SetScissorRect(RHICmdList, View);
 
 					// Render a bounding light sphere.
-					RHICmdList.SetBoundShaderState(LightFunctionBoundShaderState);
+					RHICmdList.SetLocalBoundShaderState(LightFunctionBoundShaderState);
 					VertexShader->SetParameters(RHICmdList, &View, LightSceneInfo);
 					PixelShader->SetParameters(RHICmdList, &View, LightSceneInfo, MaterialProxy, bRenderingPreviewShadowsIndicator, FadeAlpha);
 
diff --git a/Engine/Source/Runtime/Renderer/Private/LightMapDensityRendering.cpp b/Engine/Source/Runtime/Renderer/Private/LightMapDensityRendering.cpp
--- a/Engine/Source/Runtime/Renderer/Private/LightMapDensityRendering.cpp
+++ b/Engine/Source/Runtime/Renderer/Private/LightMapDensityRendering.cpp
@@ -136,8 +136,9 @@ bool FLightMapDensityDrawingPolicyFactory::DrawDynamicMesh(
 			if (bAllowHighQualityLightMaps)
 			{
 				TLightMapDensityDrawingPolicy< TLightMapPolicy<HQ_LIGHTMAP> > DrawingPolicy(View, Mesh.VertexFactory, MaterialRenderProxy, TLightMapPolicy<HQ_LIGHTMAP>(), BlendMode);
-				DrawingPolicy.DrawShared(RHICmdList, &View,DrawingPolicy.CreateBoundShaderState(View.GetFeatureLevel()));
-				for( int32 BatchElementIndex = 0; BatchElementIndex < Mesh.Elements.Num(); BatchElementIndex++)
+				RHICmdList.BuildAndSetLocalBoundShaderState(DrawingPolicy.GetBoundShaderStateInput(View.GetFeatureLevel()));
+				DrawingPolicy.SetSharedState(RHICmdList, &View);
+				for (int32 BatchElementIndex = 0; BatchElementIndex < Mesh.Elements.Num(); BatchElementIndex++)
 				{
 					DrawingPolicy.SetMeshRenderState(RHICmdList, View,PrimitiveSceneProxy,Mesh,BatchElementIndex,bBackFace,
 						TLightMapDensityDrawingPolicy< TLightMapPolicy<HQ_LIGHTMAP> >::ElementDataType(LightMapInteraction));
@@ -148,8 +149,9 @@ bool FLightMapDensityDrawingPolicyFactory::DrawDynamicMesh(
 			else
 			{
 				TLightMapDensityDrawingPolicy< TLightMapPolicy<LQ_LIGHTMAP> > DrawingPolicy(View, Mesh.VertexFactory, MaterialRenderProxy, TLightMapPolicy<LQ_LIGHTMAP>(), BlendMode);
-				DrawingPolicy.DrawShared(RHICmdList, &View, DrawingPolicy.CreateBoundShaderState(View.GetFeatureLevel()));
-				for( int32 BatchElementIndex = 0; BatchElementIndex < Mesh.Elements.Num(); BatchElementIndex++)
+				RHICmdList.BuildAndSetLocalBoundShaderState(DrawingPolicy.GetBoundShaderStateInput(View.GetFeatureLevel()));
+				DrawingPolicy.SetSharedState(RHICmdList, &View);
+				for (int32 BatchElementIndex = 0; BatchElementIndex < Mesh.Elements.Num(); BatchElementIndex++)
 				{
 					DrawingPolicy.SetMeshRenderState(RHICmdList, View,PrimitiveSceneProxy,Mesh,BatchElementIndex,bBackFace,
 						TLightMapDensityDrawingPolicy< TLightMapPolicy<LQ_LIGHTMAP> >::ElementDataType(LightMapInteraction));
@@ -161,8 +163,9 @@ bool FLightMapDensityDrawingPolicyFactory::DrawDynamicMesh(
 		else
 		{
 			TLightMapDensityDrawingPolicy<FDummyLightMapPolicy> DrawingPolicy(View, Mesh.VertexFactory, MaterialRenderProxy, FDummyLightMapPolicy(), BlendMode);
-			DrawingPolicy.DrawShared(RHICmdList, &View, DrawingPolicy.CreateBoundShaderState(View.GetFeatureLevel()));
-			for( int32 BatchElementIndex = 0; BatchElementIndex < Mesh.Elements.Num(); BatchElementIndex++)
+			RHICmdList.BuildAndSetLocalBoundShaderState(DrawingPolicy.GetBoundShaderStateInput(View.GetFeatureLevel()));
+			DrawingPolicy.SetSharedState(RHICmdList, &View);
+			for (int32 BatchElementIndex = 0; BatchElementIndex < Mesh.Elements.Num(); BatchElementIndex++)
 			{
 				DrawingPolicy.SetMeshRenderState(RHICmdList, View,PrimitiveSceneProxy,Mesh,BatchElementIndex,bBackFace,
 					TLightMapDensityDrawingPolicy<FDummyLightMapPolicy>::ElementDataType(LightMapInteraction));
@@ -174,8 +177,9 @@ bool FLightMapDensityDrawingPolicyFactory::DrawDynamicMesh(
 	else
 	{
 		TLightMapDensityDrawingPolicy<FNoLightMapPolicy> DrawingPolicy(View, Mesh.VertexFactory, MaterialRenderProxy, FNoLightMapPolicy(), BlendMode);
-		DrawingPolicy.DrawShared(RHICmdList, &View, DrawingPolicy.CreateBoundShaderState(View.GetFeatureLevel()));
-		for( int32 BatchElementIndex = 0; BatchElementIndex < Mesh.Elements.Num(); BatchElementIndex++)
+		RHICmdList.BuildAndSetLocalBoundShaderState(DrawingPolicy.GetBoundShaderStateInput(View.GetFeatureLevel()));
+		DrawingPolicy.SetSharedState(RHICmdList, &View);
+		for (int32 BatchElementIndex = 0; BatchElementIndex < Mesh.Elements.Num(); BatchElementIndex++)
 		{
 			DrawingPolicy.SetMeshRenderState(RHICmdList, View,PrimitiveSceneProxy,Mesh,BatchElementIndex,bBackFace,TLightMapDensityDrawingPolicy<FNoLightMapPolicy>::ElementDataType());
 			DrawingPolicy.DrawMesh(RHICmdList, Mesh,BatchElementIndex);
diff --git a/Engine/Source/Runtime/Renderer/Private/LightMapDensityRendering.h b/Engine/Source/Runtime/Renderer/Private/LightMapDensityRendering.h
--- a/Engine/Source/Runtime/Renderer/Private/LightMapDensityRendering.h
+++ b/Engine/Source/Runtime/Renderer/Private/LightMapDensityRendering.h
@@ -296,11 +296,8 @@ public:
 			LightMapPolicy == Other.LightMapPolicy;
 	}
 
-	void DrawShared(FRHICommandList& RHICmdList, const FSceneView* View, FBoundShaderStateRHIRef ShaderState ) const
+	void SetSharedState(FRHICommandList& RHICmdList, const FSceneView* View) const
 	{
-		// Set the actual shader & vertex declaration state
-		RHICmdList.SetBoundShaderState(ShaderState);
-
 		// Set the base pass shader parameters for the material.
 		VertexShader->SetParameters(RHICmdList, MaterialRenderProxy,*View);
 		PixelShader->SetParameters(RHICmdList, MaterialRenderProxy,View);
@@ -399,19 +396,15 @@ public:
 	 * as well as the shaders needed to draw the mesh
 	 * @return new bound shader state object
 	 */
-	FBoundShaderStateRHIRef CreateBoundShaderState(ERHIFeatureLevel::Type InFeatureLevel)
+	FBoundShaderStateInput GetBoundShaderStateInput(ERHIFeatureLevel::Type InFeatureLevel)
 	{
-		FBoundShaderStateRHIRef BoundShaderState;
-
-		BoundShaderState = RHICreateBoundShaderState(
+		return FBoundShaderStateInput(
 			FMeshDrawingPolicy::GetVertexDeclaration(), 
 			VertexShader->GetVertexShader(),
 			GETSAFERHISHADER_HULL(HullShader), 
 			GETSAFERHISHADER_DOMAIN(DomainShader),
 			PixelShader->GetPixelShader(),
 			FGeometryShaderRHIRef());
-
-		return BoundShaderState;
 	}
 
 	friend int32 CompareDrawingPolicy(const TLightMapDensityDrawingPolicy& A,const TLightMapDensityDrawingPolicy& B)
diff --git a/Engine/Source/Runtime/Renderer/Private/LightPropagationVolume.h b/Engine/Source/Runtime/Renderer/Private/LightPropagationVolume.h
--- a/Engine/Source/Runtime/Renderer/Private/LightPropagationVolume.h
+++ b/Engine/Source/Runtime/Renderer/Private/LightPropagationVolume.h
@@ -144,7 +144,7 @@ public:
 
 	void Propagate(FRHICommandListImmediate& RHICmdList);
 
-	void Visualise(FRHICommandListImmediate& RHICmdList, const FSceneView& View) const;
+	void Visualise(FRHICommandList& RHICmdList, const FSceneView& View) const;
 
 	const FIntVector& GetGridOffset() const { return mGridOffset; }
 
diff --git a/Engine/Source/Runtime/Renderer/Private/LightPropagationVolumeVisualisation.cpp b/Engine/Source/Runtime/Renderer/Private/LightPropagationVolumeVisualisation.cpp
--- a/Engine/Source/Runtime/Renderer/Private/LightPropagationVolumeVisualisation.cpp
+++ b/Engine/Source/Runtime/Renderer/Private/LightPropagationVolumeVisualisation.cpp
@@ -247,7 +247,7 @@ IMPLEMENT_SHADER_TYPE(,FLpvVisualiseVS,TEXT("LPVVisualise"),TEXT("VShader"),SF_V
 IMPLEMENT_SHADER_TYPE(,FLpvVisualisePS,TEXT("LPVVisualise"),TEXT("PShader"),SF_Pixel);
 
 
-void FLightPropagationVolume::Visualise(FRHICommandListImmediate& RHICmdList, const FSceneView& View) const
+void FLightPropagationVolume::Visualise(FRHICommandList& RHICmdList, const FSceneView& View) const
 {
 	SCOPED_DRAW_EVENT(LpvVisualise, DEC_LIGHT);
 	check(View.GetFeatureLevel() == ERHIFeatureLevel::SM5);
diff --git a/Engine/Source/Runtime/Renderer/Private/LightRendering.cpp b/Engine/Source/Runtime/Renderer/Private/LightRendering.cpp
--- a/Engine/Source/Runtime/Renderer/Private/LightRendering.cpp
+++ b/Engine/Source/Runtime/Renderer/Private/LightRendering.cpp
@@ -664,7 +664,7 @@ static FVertexDeclarationRHIParamRef GetDeferredLightingVertexDeclaration()
 
 template<bool bUseIESProfile, bool bRadialAttenuation, bool bInverseSquaredFalloff>
 static void SetShaderTemplLighting(
-	FRHICommandListImmediate& RHICmdList,
+	FRHICommandList& RHICmdList,
 	const FSceneView& View, 
 	FShader* VertexShader,
 	const FLightSceneInfo* LightSceneInfo)
@@ -685,7 +685,7 @@ static void SetShaderTemplLighting(
 
 template<bool bUseIESProfile, bool bRadialAttenuation, bool bInverseSquaredFalloff>
 static void SetShaderTemplLightingSimple(
-	FRHICommandListImmediate& RHICmdList,
+	FRHICommandList& RHICmdList,
 	const FSceneView& View, 
 	FShader* VertexShader,
 	const FSimpleLightEntry& SimpleLight,
diff --git a/Engine/Source/Runtime/Renderer/Private/PostProcess/PostProcessMaterial.cpp b/Engine/Source/Runtime/Renderer/Private/PostProcess/PostProcessMaterial.cpp
--- a/Engine/Source/Runtime/Renderer/Private/PostProcess/PostProcessMaterial.cpp
+++ b/Engine/Source/Runtime/Renderer/Private/PostProcess/PostProcessMaterial.cpp
@@ -151,8 +151,7 @@ void FRCPassPostProcessMaterial::Process(FRenderingCompositePassContext& Context
 	FPostProcessMaterialPS* PixelShader = MaterialShaderMap->GetShader<FPostProcessMaterialPS>();
 	FPostProcessMaterialVS* VertexShader = MaterialShaderMap->GetShader<FPostProcessMaterialVS>();
 
-	FBoundShaderStateRHIRef BoundShaderState = RHICreateBoundShaderState(GetVertexDeclarationFVector4(), VertexShader->GetVertexShader(), FHullShaderRHIRef(), FDomainShaderRHIRef(), PixelShader->GetPixelShader(), FGeometryShaderRHIRef());
-	Context.RHICmdList.SetBoundShaderState(BoundShaderState);
+	Context.RHICmdList.SetLocalBoundShaderState(Context.RHICmdList.BuildLocalBoundShaderState(GetVertexDeclarationFVector4(), VertexShader->GetVertexShader(), FHullShaderRHIRef(), FDomainShaderRHIRef(), PixelShader->GetPixelShader(), FGeometryShaderRHIRef()));
 
 	VertexShader->SetParameters(Context.RHICmdList, Context);
 	PixelShader->SetParameters(Context.RHICmdList, Context, MaterialInterface->GetRenderProxy(false));
diff --git a/Engine/Source/Runtime/Renderer/Private/PostProcess/VisualizeTexture.cpp b/Engine/Source/Runtime/Renderer/Private/PostProcess/VisualizeTexture.cpp
--- a/Engine/Source/Runtime/Renderer/Private/PostProcess/VisualizeTexture.cpp
+++ b/Engine/Source/Runtime/Renderer/Private/PostProcess/VisualizeTexture.cpp
@@ -191,7 +191,7 @@ public:
 IMPLEMENT_SHADER_TYPE(,FVisualizeTexturePresentPS,TEXT("VisualizeTexture"),TEXT("PresentPS"),SF_Pixel);
 
 
-template<uint32 TextureType> void VisualizeTextureForTextureType(FRHICommandListImmediate& RHICmdList, const FVisualizeTextureData& Data)
+template<uint32 TextureType> void VisualizeTextureForTextureType(FRHICommandList& RHICmdList, const FVisualizeTextureData& Data)
 {
 	TShaderMapRef<FScreenVS> VertexShader(GetGlobalShaderMap());
 	TShaderMapRef<VisualizeTexturePS<TextureType> > PixelShader(GetGlobalShaderMap());
diff --git a/Engine/Source/Runtime/Renderer/Private/ReflectionEnvironmentCapture.cpp b/Engine/Source/Runtime/Renderer/Private/ReflectionEnvironmentCapture.cpp
--- a/Engine/Source/Runtime/Renderer/Private/ReflectionEnvironmentCapture.cpp
+++ b/Engine/Source/Runtime/Renderer/Private/ReflectionEnvironmentCapture.cpp
@@ -274,7 +274,7 @@ private:
 IMPLEMENT_SHADER_TYPE(,FComputeBrightnessPS,TEXT("ReflectionEnvironmentShaders"),TEXT("ComputeBrightnessMain"),SF_Pixel);
 
 /** Computes the average brightness of the given reflection capture and stores it in the scene. */
-void ComputeAverageBrightness(FRHICommandListImmediate& RHICmdList)
+void ComputeAverageBrightness(FRHICommandList& RHICmdList)
 {
 	SetRenderTarget(RHICmdList, GSceneRenderTargets.ReflectionBrightness->GetRenderTargetItem().TargetableTexture, NULL);
 	RHICmdList.SetRasterizerState(TStaticRasterizerState<FM_Solid, CM_None>::GetRHI());
diff --git a/Engine/Source/Runtime/Renderer/Private/SceneHitProxyRendering.cpp b/Engine/Source/Runtime/Renderer/Private/SceneHitProxyRendering.cpp
--- a/Engine/Source/Runtime/Renderer/Private/SceneHitProxyRendering.cpp
+++ b/Engine/Source/Runtime/Renderer/Private/SceneHitProxyRendering.cpp
@@ -173,11 +173,8 @@ FHitProxyDrawingPolicy::FHitProxyDrawingPolicy(
 	PixelShader = MaterialResource->GetShader<FHitProxyPS>(InVertexFactory->GetType());
 }
 
-void FHitProxyDrawingPolicy::DrawShared(FRHICommandList& RHICmdList, const FSceneView* View, FBoundShaderStateRHIParamRef BoundShaderState) const
+void FHitProxyDrawingPolicy::SetSharedState(FRHICommandList& RHICmdList, const FSceneView* View) const
 {
-	// Set the actual shader & vertex declaration state
-	RHICmdList.SetBoundShaderState(BoundShaderState);
-
 	// Set the depth-only shader parameters for the material.
 	VertexShader->SetParameters(RHICmdList, MaterialRenderProxy,*View);
 	PixelShader->SetParameters(RHICmdList, MaterialRenderProxy,*View);
@@ -197,19 +194,15 @@ void FHitProxyDrawingPolicy::DrawShared(FRHICommandList& RHICmdList, const FScen
 * as well as the shaders needed to draw the mesh
 * @return new bound shader state object
 */
-FBoundShaderStateRHIRef FHitProxyDrawingPolicy::CreateBoundShaderState(ERHIFeatureLevel::Type InFeatureLevel)
+FBoundShaderStateInput FHitProxyDrawingPolicy::GetBoundShaderStateInput(ERHIFeatureLevel::Type InFeatureLevel)
 {
-	FBoundShaderStateRHIRef BoundShaderState;
-
-	BoundShaderState = RHICreateBoundShaderState(
+	return FBoundShaderStateInput(
 		FMeshDrawingPolicy::GetVertexDeclaration(), 
 		VertexShader->GetVertexShader(),
 		GETSAFERHISHADER_HULL(HullShader), 
 		GETSAFERHISHADER_DOMAIN(DomainShader), 
 		PixelShader->GetPixelShader(),
 		FGeometryShaderRHIRef());
-
-	return BoundShaderState;
 }
 
 void FHitProxyDrawingPolicy::SetMeshRenderState(
@@ -316,8 +309,9 @@ bool FHitProxyDrawingPolicyFactory::DrawDynamicMesh(
 				MaterialRenderProxy = UMaterial::GetDefaultMaterial(MD_Surface)->GetRenderProxy(false);
 			}
 			FHitProxyDrawingPolicy DrawingPolicy( Mesh.VertexFactory, MaterialRenderProxy, View.GetFeatureLevel() );
-			DrawingPolicy.DrawShared(RHICmdList, &View, DrawingPolicy.CreateBoundShaderState(View.GetFeatureLevel()));
-			for( int32 BatchElementIndex = 0; BatchElementIndex < Mesh.Elements.Num(); BatchElementIndex++ )
+			RHICmdList.BuildAndSetLocalBoundShaderState(DrawingPolicy.GetBoundShaderStateInput(View.GetFeatureLevel()));
+			DrawingPolicy.SetSharedState(RHICmdList, &View);
+			for (int32 BatchElementIndex = 0; BatchElementIndex < Mesh.Elements.Num(); BatchElementIndex++)
 			{
 				DrawingPolicy.SetMeshRenderState(RHICmdList, View, PrimitiveSceneProxy, Mesh, BatchElementIndex, bBackFace, HitProxyId);
 				DrawingPolicy.DrawMesh(RHICmdList, Mesh,BatchElementIndex);
diff --git a/Engine/Source/Runtime/Renderer/Private/SceneHitProxyRendering.h b/Engine/Source/Runtime/Renderer/Private/SceneHitProxyRendering.h
--- a/Engine/Source/Runtime/Renderer/Private/SceneHitProxyRendering.h
+++ b/Engine/Source/Runtime/Renderer/Private/SceneHitProxyRendering.h
@@ -30,14 +30,14 @@ public:
 			VertexShader == Other.VertexShader &&
 			PixelShader == Other.PixelShader;
 	}
-	void DrawShared(FRHICommandList& RHICmdList, const FSceneView* View,FBoundShaderStateRHIParamRef BoundShaderState) const;
+	void SetSharedState(FRHICommandList& RHICmdList, const FSceneView* View) const;
 
 	/** 
 	* Create bound shader state using the vertex decl from the mesh draw policy
 	* as well as the shaders needed to draw the mesh
 	* @return new bound shader state object
 	*/
-	FBoundShaderStateRHIRef CreateBoundShaderState(ERHIFeatureLevel::Type InFeatureLevel);
+	FBoundShaderStateInput GetBoundShaderStateInput(ERHIFeatureLevel::Type InFeatureLevel);
 
 	void SetMeshRenderState(
 		FRHICommandList& RHICmdList, 
diff --git a/Engine/Source/Runtime/Renderer/Private/ShadowRendering.cpp b/Engine/Source/Runtime/Renderer/Private/ShadowRendering.cpp
--- a/Engine/Source/Runtime/Renderer/Private/ShadowRendering.cpp
+++ b/Engine/Source/Runtime/Renderer/Private/ShadowRendering.cpp
@@ -846,13 +846,10 @@ FShadowDepthDrawingPolicy<bRenderingReflectiveShadowMaps>::FShadowDepthDrawingPo
 }
 
 template <bool bRenderingReflectiveShadowMaps>
-void FShadowDepthDrawingPolicy<bRenderingReflectiveShadowMaps>::DrawShared(FRHICommandList& RHICmdList, const FSceneView* View,FBoundShaderStateRHIParamRef BoundShaderState) const
+void FShadowDepthDrawingPolicy<bRenderingReflectiveShadowMaps>::SetSharedState(FRHICommandList& RHICmdList, const FSceneView* View) const
 {
 	checkSlow(bDirectionalLight == PolicyShadowInfo->bDirectionalLight && bPreShadow == PolicyShadowInfo->bPreShadow);
 
-	// Set the actual shader & vertex declaration state
-	RHICmdList.SetBoundShaderState(BoundShaderState);
-
 	VertexShader->SetParameters(RHICmdList, MaterialRenderProxy,*MaterialResource,*View,PolicyShadowInfo);
 
 	if (GeometryShader)
@@ -903,7 +900,7 @@ void FShadowDepthDrawingPolicy<bRenderingReflectiveShadowMaps>::DrawShared(FRHIC
  * @return new bound shader state object
  */
 template <bool bRenderingReflectiveShadowMaps>
-FBoundShaderStateRHIRef FShadowDepthDrawingPolicy<bRenderingReflectiveShadowMaps>::CreateBoundShaderState(ERHIFeatureLevel::Type InFeatureLevel)
+FBoundShaderStateInput FShadowDepthDrawingPolicy<bRenderingReflectiveShadowMaps>::GetBoundShaderStateInput(ERHIFeatureLevel::Type InFeatureLevel)
 {
 	FVertexDeclarationRHIRef VertexDeclaration;
 	if (bUsePositionOnlyVS)
@@ -915,7 +912,7 @@ FBoundShaderStateRHIRef FShadowDepthDrawingPolicy<bRenderingReflectiveShadowMaps
 		VertexDeclaration = FMeshDrawingPolicy::GetVertexDeclaration();
 	}
 
-	return RHICreateBoundShaderState(
+	return FBoundShaderStateInput(
 		VertexDeclaration, 
 		VertexShader->GetVertexShader(),
 		GETSAFERHISHADER_HULL(HullShader), 
@@ -1078,8 +1075,9 @@ bool FShadowDepthDrawingPolicyFactory::DrawDynamicMesh(
 					Mesh.ReverseCulling
 					);
 
-				DrawingPolicy.DrawShared(RHICmdList, &View,DrawingPolicy.CreateBoundShaderState(View.GetFeatureLevel()));
-				for( int32 BatchElementIndex = 0, Num = Mesh.Elements.Num(); BatchElementIndex < Num; BatchElementIndex++)
+				RHICmdList.BuildAndSetLocalBoundShaderState(DrawingPolicy.GetBoundShaderStateInput(View.GetFeatureLevel()));
+				DrawingPolicy.SetSharedState(RHICmdList, &View);
+				for (int32 BatchElementIndex = 0, Num = Mesh.Elements.Num(); BatchElementIndex < Num; BatchElementIndex++)
 				{
 					DrawingPolicy.SetMeshRenderState(RHICmdList, View,PrimitiveSceneProxy,Mesh,BatchElementIndex,bBackFace,FMeshDrawingPolicy::ElementDataType());
 					DrawingPolicy.DrawMesh(RHICmdList, Mesh,BatchElementIndex);
@@ -1098,8 +1096,9 @@ bool FShadowDepthDrawingPolicyFactory::DrawDynamicMesh(
 					Mesh.ReverseCulling
 					);
 
-				DrawingPolicy.DrawShared(RHICmdList, &View, DrawingPolicy.CreateBoundShaderState(View.GetFeatureLevel()));
-				for( int32 BatchElementIndex = 0; BatchElementIndex < Mesh.Elements.Num(); BatchElementIndex++)
+				RHICmdList.BuildAndSetLocalBoundShaderState(DrawingPolicy.GetBoundShaderStateInput(View.GetFeatureLevel()));
+				DrawingPolicy.SetSharedState(RHICmdList, &View);
+				for (int32 BatchElementIndex = 0; BatchElementIndex < Mesh.Elements.Num(); BatchElementIndex++)
 				{
 					DrawingPolicy.SetMeshRenderState(RHICmdList, View,PrimitiveSceneProxy,Mesh,BatchElementIndex,bBackFace,FMeshDrawingPolicy::ElementDataType());
 					DrawingPolicy.DrawMesh(RHICmdList, Mesh,BatchElementIndex);
@@ -1256,7 +1255,8 @@ void DrawShadowMeshElements(FRHICommandList& RHICmdList, const FViewInfo& View,
 			OldState = CurrentState;
 
 			SharedDrawingPolicy.UpdateElementState(CurrentState);
-			SharedDrawingPolicy.DrawShared(RHICmdList, &View, SharedDrawingPolicy.CreateBoundShaderState(View.GetFeatureLevel()));
+			RHICmdList.BuildAndSetLocalBoundShaderState(SharedDrawingPolicy.GetBoundShaderStateInput(View.GetFeatureLevel()));
+			SharedDrawingPolicy.SetSharedState(RHICmdList, &View);
 		}
 
 		DrawMeshElements(RHICmdList, SharedDrawingPolicy, OldState, View, ShadowMesh.Mesh);
@@ -1471,7 +1471,7 @@ void StencilingGeometry::DrawCone(FRHICommandList& RHICmdList)
 static FGlobalBoundShaderState MaskBoundShaderState[2];
 
 template <uint32 Quality>
-static void SetShadowProjectionShaderTemplNew(FRHICommandListImmediate& RHICmdList, int32 ViewIndex, const FViewInfo& View, const FProjectedShadowInfo* ShadowInfo)
+static void SetShadowProjectionShaderTemplNew(FRHICommandList& RHICmdList, int32 ViewIndex, const FViewInfo& View, const FProjectedShadowInfo* ShadowInfo)
 {
 	if (ShadowInfo->bTranslucentShadow)
 	{
@@ -1948,7 +1948,7 @@ void FProjectedShadowInfo::RenderProjection(FRHICommandListImmediate& RHICmdList
 
 
 template <uint32 Quality>
-static void SetPointLightShaderTempl(FRHICommandListImmediate& RHICmdList, int32 ViewIndex, const FViewInfo& View, const FProjectedShadowInfo* ShadowInfo)
+static void SetPointLightShaderTempl(FRHICommandList& RHICmdList, int32 ViewIndex, const FViewInfo& View, const FProjectedShadowInfo* ShadowInfo)
 {
 	TShaderMapRef<FShadowProjectionVS> VertexShader(GetGlobalShaderMap());
 	TShaderMapRef<TOnePassPointShadowProjectionPS<Quality> > PixelShader(GetGlobalShaderMap());
diff --git a/Engine/Source/Runtime/Renderer/Private/ShadowRendering.h b/Engine/Source/Runtime/Renderer/Private/ShadowRendering.h
--- a/Engine/Source/Runtime/Renderer/Private/ShadowRendering.h
+++ b/Engine/Source/Runtime/Renderer/Private/ShadowRendering.h
@@ -333,14 +333,14 @@ public:
 			&& bUsePositionOnlyVS == Other.bUsePositionOnlyVS
 			&& bPreShadow == Other.bPreShadow;
 	}
-	void DrawShared(FRHICommandList& RHICmdList, const FSceneView* View,FBoundShaderStateRHIParamRef BoundShaderState) const;
+	void SetSharedState(FRHICommandList& RHICmdList, const FSceneView* View) const;
 
 	/** 
 	 * Create bound shader state using the vertex decl from the mesh draw policy
 	 * as well as the shaders needed to draw the mesh
 	 * @return new bound shader state object
 	 */
-	FBoundShaderStateRHIRef CreateBoundShaderState(ERHIFeatureLevel::Type InFeatureLevel);
+	FBoundShaderStateInput GetBoundShaderStateInput(ERHIFeatureLevel::Type InFeatureLevel);
 
 	void SetMeshRenderState(
 		FRHICommandList& RHICmdList, 
diff --git a/Engine/Source/Runtime/Renderer/Private/StaticMeshDrawList.h b/Engine/Source/Runtime/Renderer/Private/StaticMeshDrawList.h
--- a/Engine/Source/Runtime/Renderer/Private/StaticMeshDrawList.h
+++ b/Engine/Source/Runtime/Renderer/Private/StaticMeshDrawList.h
@@ -197,7 +197,15 @@ private:
 
 		void CreateBoundShaderState()
 		{
-			BoundShaderState = DrawingPolicy.CreateBoundShaderState(FeatureLevel);
+			check(IsInRenderingThread());
+			FBoundShaderStateInput BoundShaderStateInput = DrawingPolicy.GetBoundShaderStateInput(FeatureLevel);
+			BoundShaderState = CreateBoundShaderState_Internal(
+				BoundShaderStateInput.VertexDeclarationRHI,
+				BoundShaderStateInput.VertexShaderRHI,
+				BoundShaderStateInput.HullShaderRHI,
+				BoundShaderStateInput.DomainShaderRHI,
+				BoundShaderStateInput.PixelShaderRHI,
+				BoundShaderStateInput.GeometryShaderRHI);
 		}
 	};
 
diff --git a/Engine/Source/Runtime/Renderer/Private/StaticMeshDrawList.inl b/Engine/Source/Runtime/Renderer/Private/StaticMeshDrawList.inl
--- a/Engine/Source/Runtime/Renderer/Private/StaticMeshDrawList.inl
+++ b/Engine/Source/Runtime/Renderer/Private/StaticMeshDrawList.inl
@@ -72,11 +72,15 @@ void TStaticMeshDrawList<DrawingPolicyType>::DrawElement(
 
 	if (!bDrawnShared)
 	{
-		if (!IsValidRef(DrawingPolicyLink->BoundShaderState))
+		if (ensure(IsValidRef(DrawingPolicyLink->BoundShaderState)))
 		{
-			DrawingPolicyLink->CreateBoundShaderState();
+			RHICmdList.SetBoundShaderState(DrawingPolicyLink->BoundShaderState);
 		}
-		DrawingPolicyLink->DrawingPolicy.DrawShared(RHICmdList, &View,DrawingPolicyLink->BoundShaderState);
+		else
+		{
+			RHICmdList.BuildAndSetLocalBoundShaderState(DrawingPolicyLink->DrawingPolicy.GetBoundShaderStateInput(View.GetFeatureLevel()));
+		}
+		DrawingPolicyLink->DrawingPolicy.SetSharedState(RHICmdList, &View);
 		bDrawnShared = true;
 	}
 	
diff --git a/Engine/Source/Runtime/Renderer/Private/TranslucentLighting.cpp b/Engine/Source/Runtime/Renderer/Private/TranslucentLighting.cpp
--- a/Engine/Source/Runtime/Renderer/Private/TranslucentLighting.cpp
+++ b/Engine/Source/Runtime/Renderer/Private/TranslucentLighting.cpp
@@ -348,14 +348,11 @@ public:
 		}
 	}
 
-	void DrawShared(FRHICommandList& RHICmdList, const FSceneView* View, FBoundShaderStateRHIParamRef BoundShaderState) const
+	void SetSharedState(FRHICommandList& RHICmdList, const FSceneView* View) const
 	{
 		// Set the shared mesh resources.
 		FMeshDrawingPolicy::DrawShared(RHICmdList, View);
 
-		// Set the actual shader & vertex declaration state
-		RHICmdList.SetBoundShaderState(BoundShaderState);
-
 		VertexShader->SetParameters(RHICmdList, MaterialRenderProxy, *View, PolicyShadowInfo);
 		PixelShader->SetParameters(RHICmdList, MaterialRenderProxy, *View, PolicyShadowInfo);
 	}
@@ -365,9 +362,9 @@ public:
 	* as well as the shaders needed to draw the mesh
 	* @return new bound shader state object
 	*/
-	FBoundShaderStateRHIRef CreateBoundShaderState(ERHIFeatureLevel::Type InFeatureLevel)
+	FBoundShaderStateInput GetBoundShaderStateInput(ERHIFeatureLevel::Type InFeatureLevel)
 	{
-		return RHICreateBoundShaderState(
+		return FBoundShaderStateInput(
 			FMeshDrawingPolicy::GetVertexDeclaration(), 
 			VertexShader->GetVertexShader(),
 			NULL, 
@@ -439,7 +436,8 @@ public:
 			if (IsTranslucentBlendMode(BlendMode))
 			{
 				FTranslucencyShadowDepthDrawingPolicy DrawingPolicy(Mesh.VertexFactory, MaterialRenderProxy, *MaterialRenderProxy->GetMaterial(FeatureLevel), DrawingContext.bDirectionalLight);
-				DrawingPolicy.DrawShared(RHICmdList, &View, DrawingPolicy.CreateBoundShaderState(FeatureLevel));
+				RHICmdList.BuildAndSetLocalBoundShaderState(DrawingPolicy.GetBoundShaderStateInput(View.GetFeatureLevel()));
+				DrawingPolicy.SetSharedState(RHICmdList, &View);
 
 				for (int32 BatchElementIndex = 0; BatchElementIndex < Mesh.Elements.Num(); BatchElementIndex++)
 				{
@@ -1557,7 +1555,8 @@ void SetInjectionShader(
 	if (!IsValidRef(BoundShaderState) || CachedShaderMap != MaterialShaderMap)
 	{
 		CachedShaderMap = MaterialShaderMap;
-		BoundShaderState = RHICreateBoundShaderState(GScreenVertexDeclaration.VertexDeclarationRHI, VertexShader->GetVertexShader(), FHullShaderRHIRef(), FDomainShaderRHIRef(), PixelShader->GetPixelShader(), GeometryShader->GetGeometryShader());
+		check(IsInRenderingThread()); // I didn't know quite how to deal with this caching. It won't work with threads.
+		BoundShaderState = CreateBoundShaderState_Internal(GScreenVertexDeclaration.VertexDeclarationRHI, VertexShader->GetVertexShader(), FHullShaderRHIRef(), FDomainShaderRHIRef(), PixelShader->GetPixelShader(), GeometryShader->GetGeometryShader());
 	}
 
 	RHICmdList.SetBoundShaderState(BoundShaderState);
diff --git a/Engine/Source/Runtime/Renderer/Private/TranslucentRendering.cpp b/Engine/Source/Runtime/Renderer/Private/TranslucentRendering.cpp
--- a/Engine/Source/Runtime/Renderer/Private/TranslucentRendering.cpp
+++ b/Engine/Source/Runtime/Renderer/Private/TranslucentRendering.cpp
@@ -273,11 +273,8 @@ public:
 			View.Family->EngineShowFlags.ShaderComplexity,
 			Parameters.bAllowFog
 			);
-		DrawingPolicy.DrawShared(
-			RHICmdList, 
-			&View,
-			DrawingPolicy.CreateBoundShaderState(View.GetFeatureLevel())
-			);
+		RHICmdList.BuildAndSetLocalBoundShaderState(DrawingPolicy.GetBoundShaderStateInput(View.GetFeatureLevel()));
+		DrawingPolicy.SetSharedState(RHICmdList, &View);
 
 		int32 BatchElementIndex = 0;
 		uint64 BatchElementMask = Parameters.BatchElementMask;
diff --git a/Engine/Source/Runtime/Renderer/Private/VelocityRendering.cpp b/Engine/Source/Runtime/Renderer/Private/VelocityRendering.cpp
--- a/Engine/Source/Runtime/Renderer/Private/VelocityRendering.cpp
+++ b/Engine/Source/Runtime/Renderer/Private/VelocityRendering.cpp
@@ -206,14 +206,11 @@ bool FVelocityDrawingPolicy::SupportsVelocity() const
 	return (VertexShader && PixelShader) ? VertexShader->SupportsVelocity() : false;
 }
 
-void FVelocityDrawingPolicy::DrawShared(FRHICommandList& RHICmdList, const FSceneView* SceneView, FBoundShaderStateRHIRef ShaderState) const
+void FVelocityDrawingPolicy::SetSharedState(FRHICommandList& RHICmdList, const FSceneView* SceneView) const
 {
 	// NOTE: Assuming this cast is always safe!
 	FViewInfo* View = (FViewInfo*)SceneView;
 
-	// Set the depth-only shader parameters for the material.
-	RHICmdList.SetBoundShaderState(ShaderState);
-
 	VertexShader->SetParameters(RHICmdList, VertexFactory, MaterialRenderProxy, *View);
 	PixelShader->SetParameters(RHICmdList, VertexFactory, MaterialRenderProxy, *View);
 
@@ -314,19 +311,15 @@ bool FVelocityDrawingPolicy::HasVelocity(const FViewInfo& View, const FPrimitive
 	return true;
 }
 
-FBoundShaderStateRHIRef FVelocityDrawingPolicy::CreateBoundShaderState(ERHIFeatureLevel::Type InFeatureLevel)
+FBoundShaderStateInput FVelocityDrawingPolicy::GetBoundShaderStateInput(ERHIFeatureLevel::Type InFeatureLevel)
 {
-	FBoundShaderStateRHIRef BoundShaderState;
-
-	BoundShaderState = RHICreateBoundShaderState(
+	return FBoundShaderStateInput(
 		FMeshDrawingPolicy::GetVertexDeclaration(), 
 		VertexShader->GetVertexShader(),
 		GETSAFERHISHADER_HULL(HullShader), 
 		GETSAFERHISHADER_DOMAIN(DomainShader),
 		PixelShader->GetPixelShader(),
 		FGeometryShaderRHIRef());
-
-	return BoundShaderState;
 }
 
 int32 Compare(const FVelocityDrawingPolicy& A,const FVelocityDrawingPolicy& B)
@@ -401,8 +394,9 @@ bool FVelocityDrawingPolicyFactory::DrawDynamicMesh(
 		FVelocityDrawingPolicy DrawingPolicy(Mesh.VertexFactory, MaterialRenderProxy, *MaterialRenderProxy->GetMaterial(FeatureLevel));
 		if(DrawingPolicy.SupportsVelocity())
 		{			
-			DrawingPolicy.DrawShared(RHICmdList, &View,DrawingPolicy.CreateBoundShaderState(FeatureLevel));
-			for(int32 BatchElementIndex = 0; BatchElementIndex < Mesh.Elements.Num(); ++BatchElementIndex)
+			RHICmdList.BuildAndSetLocalBoundShaderState(DrawingPolicy.GetBoundShaderStateInput(View.GetFeatureLevel()));
+			DrawingPolicy.SetSharedState(RHICmdList, &View);
+			for (int32 BatchElementIndex = 0; BatchElementIndex < Mesh.Elements.Num(); ++BatchElementIndex)
 			{
 				DrawingPolicy.SetMeshRenderState(RHICmdList, View, PrimitiveSceneProxy, Mesh, BatchElementIndex, bBackFace, FMeshDrawingPolicy::ElementDataType());
 				DrawingPolicy.DrawMesh(RHICmdList, Mesh, BatchElementIndex);
diff --git a/Engine/Source/Runtime/Renderer/Private/VelocityRendering.h b/Engine/Source/Runtime/Renderer/Private/VelocityRendering.h
--- a/Engine/Source/Runtime/Renderer/Private/VelocityRendering.h
+++ b/Engine/Source/Runtime/Renderer/Private/VelocityRendering.h
@@ -27,7 +27,7 @@ public:
 			&& VertexShader == Other.VertexShader
 			&& PixelShader == Other.PixelShader;
 	}
-	void DrawShared(FRHICommandList& RHICmdList, const FSceneView* View, FBoundShaderStateRHIRef ShaderState) const;
+	void SetSharedState(FRHICommandList& RHICmdList, const FSceneView* View) const;
 	void SetMeshRenderState(
 		FRHICommandList& RHICmdList, 
 		const FSceneView& View,
@@ -38,7 +38,7 @@ public:
 		const ElementDataType& ElementData
 		) const;
 
-	FBoundShaderStateRHIRef CreateBoundShaderState(ERHIFeatureLevel::Type InFeatureLevel);
+	FBoundShaderStateInput GetBoundShaderStateInput(ERHIFeatureLevel::Type InFeatureLevel);
 
 	friend int32 Compare(const FVelocityDrawingPolicy& A, const FVelocityDrawingPolicy& B);
 
diff --git a/Engine/Source/Runtime/SlateRHIRenderer/Private/SlateRHIRenderingPolicy.cpp b/Engine/Source/Runtime/SlateRHIRenderer/Private/SlateRHIRenderingPolicy.cpp
--- a/Engine/Source/Runtime/SlateRHIRenderer/Private/SlateRHIRenderingPolicy.cpp
+++ b/Engine/Source/Runtime/SlateRHIRenderer/Private/SlateRHIRenderingPolicy.cpp
@@ -324,17 +324,15 @@ void FSlateRHIRenderingPolicy::DrawElements(FRHICommandListImmediate& RHICmdList
 			if( !ShaderResource || ShaderResource->GetType() == ESlateShaderResource::Texture ) 
 			{
 				FSlateElementPS* PixelShader = GetTexturePixelShader(ShaderType, DrawEffects);
-				FBoundShaderStateRHIRef BoundShaderState;
 
-				BoundShaderState = RHICreateBoundShaderState(
+				RHICmdList.SetLocalBoundShaderState(RHICmdList.BuildLocalBoundShaderState(
 					GSlateVertexDeclaration.VertexDeclarationRHI,
 					VertexShader->GetVertexShader(),
 					nullptr,
 					nullptr,
 					PixelShader->GetPixelShader(),
-					FGeometryShaderRHIRef());
+					FGeometryShaderRHIRef()));
 
-				RHICmdList.SetBoundShaderState(BoundShaderState);
 
 				VertexShader->SetViewProjection(RHICmdList, ViewProjectionMatrix);
 				VertexShader->SetShaderParameters(RHICmdList, ShaderParams.VertexParams);
@@ -449,17 +447,13 @@ void FSlateRHIRenderingPolicy::DrawElements(FRHICommandListImmediate& RHICmdList
 
 				if( PixelShader )
 				{
-					FBoundShaderStateRHIRef BoundShaderState;
-
-					BoundShaderState = RHICreateBoundShaderState(
+					RHICmdList.SetLocalBoundShaderState(RHICmdList.BuildLocalBoundShaderState(
 						GSlateVertexDeclaration.VertexDeclarationRHI,
 						VertexShader->GetVertexShader(),
 						nullptr,
 						nullptr,
 						PixelShader->GetPixelShader(),
-						FGeometryShaderRHIRef());
-
-					RHICmdList.SetBoundShaderState(BoundShaderState);
+						FGeometryShaderRHIRef()));
 
 					VertexShader->SetShaderParameters(RHICmdList, ShaderParams.VertexParams);
 					PixelShader->SetParameters(RHICmdList, *SceneView, MaterialRenderProxy, Material, 1.0f / DisplayGamma, ShaderParams.PixelParams);
