UE4 - cleaned up the recursive RHI usage for command lists and added these RHI mthods as queuable RHICopyToResolveTarget, RHICLear, RHIClearMRT, RHIClearUAV

diff --git a/Engine/Source/Runtime/OpenGLDrv/Private/OpenGLDrv.cpp b/Engine/Source/Runtime/OpenGLDrv/Private/OpenGLDrv.cpp
--- a/Engine/Source/Runtime/OpenGLDrv/Private/OpenGLDrv.cpp
+++ b/Engine/Source/Runtime/OpenGLDrv/Private/OpenGLDrv.cpp
@@ -386,7 +386,6 @@ static FGlobalBoundShaderState LongGPUTaskBoundShaderState;
 
 void FOpenGLDynamicRHI::IssueLongGPUTask()
 {
-	FRHICommandListImmediate& RHICmdList = FRHICommandListExecutor::GetRecursiveRHICommandList();
 	int32 LargestViewportIndex = INDEX_NONE;
 	int32 LargestViewportPixels = 0;
 
@@ -405,6 +404,7 @@ void FOpenGLDynamicRHI::IssueLongGPUTask()
 	{
 		FOpenGLViewport* Viewport = Viewports[LargestViewportIndex];
 
+		FRHICommandList_RecursiveHazardous RHICmdList;
 		SetRenderTarget(RHICmdList, Viewport->GetBackBuffer(), FTextureRHIRef());
 		RHICmdList.SetBlendState(TStaticBlendState<CW_RGBA, BO_Add, BF_One, BF_One>::GetRHI(), FLinearColor::Black);
 		RHICmdList.SetDepthStencilState(TStaticDepthStencilState<false, CF_Always>::GetRHI(), 0);
@@ -422,8 +422,9 @@ void FOpenGLDynamicRHI::IssueLongGPUTask()
 		Vertices[2].Set( -1.0f, -1.0f, 0, 1.0f );
 		Vertices[3].Set(  1.0f, -1.0f, 0, 1.0f );
 		DrawPrimitiveUP(RHICmdList, PT_TriangleStrip, 2, Vertices, sizeof(Vertices[0]));
+
+		// RHICmdList flushes on destruction
 	}
-	RHICmdList.Flush(); // always call flush with GetRecursiveRHICommandList, recursive use of the RHI is hazardous
 }
 
 void FOpenGLDynamicRHI::InitializeStateResources()
diff --git a/Engine/Source/Runtime/OpenGLDrv/Private/OpenGLViewport.cpp b/Engine/Source/Runtime/OpenGLDrv/Private/OpenGLViewport.cpp
--- a/Engine/Source/Runtime/OpenGLDrv/Private/OpenGLViewport.cpp
+++ b/Engine/Source/Runtime/OpenGLDrv/Private/OpenGLViewport.cpp
@@ -69,8 +69,6 @@ void FOpenGLDynamicRHI::RHITick( float DeltaTime )
 
 void FOpenGLDynamicRHI::RHIBeginDrawingViewport(FViewportRHIParamRef ViewportRHI, FTextureRHIParamRef RenderTarget)
 {
-	FRHICommandListImmediate& RHICmdList = FRHICommandListExecutor::GetRecursiveRHICommandList();
-
 	VERIFY_GL_SCOPE();
 
 	DYNAMIC_CAST_OPENGLRESOURCE(Viewport,Viewport);
@@ -98,13 +96,14 @@ void FOpenGLDynamicRHI::RHIBeginDrawingViewport(FViewportRHIParamRef ViewportRHI
 	// Set the render target and viewport.
 	if( RenderTarget )
 	{
-		SetRenderTarget(RHICmdList, RenderTarget, FTextureRHIRef());
+		FRHIRenderTargetView RTV(RenderTarget);
+		RHISetRenderTargets(1, &RTV, FTextureRHIRef(), 0, NULL);
 	}
 	else
 	{
-		SetRenderTarget(RHICmdList, DrawingViewport->GetBackBuffer(), FTextureRHIRef());
+		FRHIRenderTargetView RTV(DrawingViewport->GetBackBuffer());
+		RHISetRenderTargets(1, &RTV, FTextureRHIRef(), 0, NULL);
 	}
-	RHICmdList.Flush(); // always call flush with GetRecursiveRHICommandList, recursive use of the RHI is hazardous
 }
 
 void FOpenGLDynamicRHI::RHIEndDrawingViewport(FViewportRHIParamRef ViewportRHI,bool bPresent,bool bLockToVsync)
diff --git a/Engine/Source/Runtime/RHI/Private/RHICommandList.cpp b/Engine/Source/Runtime/RHI/Private/RHICommandList.cpp
--- a/Engine/Source/Runtime/RHI/Private/RHICommandList.cpp
+++ b/Engine/Source/Runtime/RHI/Private/RHICommandList.cpp
@@ -52,6 +52,7 @@ static FAutoConsoleVariableRef CVarRHICmdMem(
 
 RHI_API FRHICommandListExecutor GRHICommandList;
 RHI_API FThreadSafeCounter FRHICommandList::UIDCounter;
+RHI_API TLockFreeFixedSizeAllocator<FRHICommandList::FMemManager::FPage::PageSize> FRHICommandList::FMemManager::FPage::TheAllocator;
 //RHI_API FGlobalRHI GRHI;
 
 static_assert(sizeof(FRHICommand) == sizeof(FRHICommandNopEndOfPage), "These should match for filling in the end of a page when a new allocation won't fit");
@@ -530,6 +531,45 @@ void FRHICommandListExecutor::ExecuteList(FRHICommandList& CmdList)
 				}
 			}
 			break;
+		case ERCT_ClearUAV:
+			{
+				auto* RHICmd = Iter.NextCommand<FRHICommandClearUAV>();
+				if (!bOnlyTestMemAccess)
+				{
+					ClearUAV_Internal(RHICmd->UnorderedAccessViewRHI, RHICmd->Values);
+				}
+				RHICmd->UnorderedAccessViewRHI->Release();
+			}
+			break;
+		case ERCT_CopyToResolveTarget:
+			{
+				auto* RHICmd = Iter.NextCommand<FRHICommandCopyToResolveTarget>();
+				if (!bOnlyTestMemAccess)
+				{
+					CopyToResolveTarget_Internal(RHICmd->SourceTexture, RHICmd->DestTexture, RHICmd->bKeepOriginalSurface, RHICmd->ResolveParams);
+				}
+				RHICmd->SourceTexture->Release();
+				RHICmd->DestTexture->Release();
+			}
+			break;
+		case ERCT_Clear:
+			{
+				auto* RHICmd = Iter.NextCommand<FRHICommandClear>();
+				if (!bOnlyTestMemAccess)
+				{
+					Clear_Internal(RHICmd->bClearColor, RHICmd->Color, RHICmd->bClearDepth, RHICmd->Depth, RHICmd->bClearStencil, RHICmd->Stencil, RHICmd->ExcludeRect);
+				}
+			}
+			break;
+		case ERCT_ClearMRT:
+			{
+				auto* RHICmd = Iter.NextCommand<FRHICommandClearMRT>();
+				if (!bOnlyTestMemAccess)
+				{
+					ClearMRT_Internal(RHICmd->bClearColor, RHICmd->NumClearColors, RHICmd->ColorArray, RHICmd->bClearDepth, RHICmd->Depth, RHICmd->bClearStencil, RHICmd->Stencil, RHICmd->ExcludeRect);
+				}
+			}
+			break;
 
 
 		default:
diff --git a/Engine/Source/Runtime/RHI/Public/RHI.h b/Engine/Source/Runtime/RHI/Public/RHI.h
--- a/Engine/Source/Runtime/RHI/Public/RHI.h
+++ b/Engine/Source/Runtime/RHI/Public/RHI.h
@@ -749,13 +749,20 @@ struct FResolveRect
 	int32 Y2;
 	// e.g. for a a full 256 x 256 area starting at (0, 0) it would be 
 	// the values would be 0, 0, 256, 256
-	FResolveRect(int32 InX1=-1, int32 InY1=-1, int32 InX2=-1, int32 InY2=-1)
+	FORCEINLINE FResolveRect(int32 InX1=-1, int32 InY1=-1, int32 InX2=-1, int32 InY2=-1)
 	:	X1(InX1)
 	,	Y1(InY1)
 	,	X2(InX2)
 	,	Y2(InY2)
 	{}
 
+	FORCEINLINE FResolveRect(const FResolveRect& Other)
+		: X1(Other.X1)
+		, Y1(Other.Y1)
+		, X2(Other.X2)
+		, Y2(Other.Y2)
+	{}
+
 	bool IsValid() const
 	{
 		return X1 >= 0 && Y1 >= 0 && X2 - X1 > 0 && Y2 - Y1 > 0;
@@ -798,6 +805,14 @@ struct FResolveParams
 		,	DestArrayIndex(InDestArrayIndex)
 	{}
 
+	FORCEINLINE FResolveParams(const FResolveParams& Other)
+		: CubeFace(Other.CubeFace)
+		, Rect(Other.Rect)
+		, MipIndex(Other.MipIndex)
+		, SourceArrayIndex(Other.SourceArrayIndex)
+		, DestArrayIndex(Other.DestArrayIndex)
+	{}
+
 	friend FArchive& operator<<(FArchive& Ar,FResolveParams& ResolveParams)
 	{
 		uint8 CubeFaceInt = ResolveParams.CubeFace;
diff --git a/Engine/Source/Runtime/RHI/Public/RHICommandList.h b/Engine/Source/Runtime/RHI/Public/RHICommandList.h
--- a/Engine/Source/Runtime/RHI/Public/RHICommandList.h
+++ b/Engine/Source/Runtime/RHI/Public/RHICommandList.h
@@ -5,6 +5,7 @@
 =============================================================================*/
 
 #pragma once
+#include "LockFreeFixedSizeAllocator.h"
 
 enum ERHICommandType
 {
@@ -46,6 +47,12 @@ enum ERHICommandType
 	ERCT_DrawIndexedIndirect,
 	ERCT_DrawIndexedPrimitiveIndirect,
 	ERCT_EnableDepthBoundsTest,
+	ERCT_ClearUAV,
+	ERCT_CopyToResolveTarget,
+	ERCT_ClearMRT,
+	ERCT_Clear,
+
+
 
 	ERCT_,
 };
@@ -562,6 +569,107 @@ struct FRHICommandEnableDepthBoundsTest : public FRHICommand
 	}
 };
 
+struct FRHICommandClearUAV : public FRHICommand
+{
+	FUnorderedAccessViewRHIParamRef UnorderedAccessViewRHI;
+	uint32 Values[4];
+
+	FORCEINLINE_DEBUGGABLE void Set(FUnorderedAccessViewRHIParamRef InUnorderedAccessViewRHI, const uint32* InValues)
+	{
+		Type = ERCT_ClearUAV;
+		UnorderedAccessViewRHI = InUnorderedAccessViewRHI;
+		Values[0] = InValues[0];
+		Values[1] = InValues[1];
+		Values[2] = InValues[2];
+		Values[3] = InValues[3];
+	}
+};
+
+struct FRHICommandCopyToResolveTarget : public FRHICommand
+{
+	FTextureRHIParamRef SourceTexture;
+	FTextureRHIParamRef DestTexture;
+	bool bKeepOriginalSurface;
+	FResolveParams ResolveParams;
+
+	FORCEINLINE_DEBUGGABLE void Set(FTextureRHIParamRef InSourceTexture, FTextureRHIParamRef InDestTexture, bool InbKeepOriginalSurface, const FResolveParams& InResolveParams)
+	{
+		Type = ERCT_CopyToResolveTarget;
+		SourceTexture = InSourceTexture;
+		DestTexture = InDestTexture;
+		bKeepOriginalSurface = InbKeepOriginalSurface;
+		ResolveParams = InResolveParams;
+	}
+};
+
+struct FRHICommandClear : public FRHICommand
+{
+	FLinearColor Color;
+	FIntRect ExcludeRect;
+	float Depth;
+	uint32 Stencil;
+	bool bClearColor;
+	bool bClearDepth;
+	bool bClearStencil;
+	FORCEINLINE_DEBUGGABLE void Set(
+		bool InbClearColor,
+		const FLinearColor& InColor,
+		bool InbClearDepth,
+		float InDepth,
+		bool InbClearStencil,
+		uint32 InStencil,
+		FIntRect InExcludeRect
+		)
+	{
+		Type = ERCT_Clear;
+		Color = InColor;
+		ExcludeRect = InExcludeRect;
+		Depth = InDepth;
+		Stencil = InStencil;
+		bClearColor = InbClearColor;
+		bClearDepth = InbClearDepth;
+		bClearStencil = InbClearStencil;
+	}
+};
+
+struct FRHICommandClearMRT : public FRHICommand
+{
+	FLinearColor ColorArray[MaxSimultaneousRenderTargets];
+	FIntRect ExcludeRect;
+	float Depth;
+	uint32 Stencil;
+	int32 NumClearColors;
+	bool bClearColor;
+	bool bClearDepth;
+	bool bClearStencil;
+	FORCEINLINE_DEBUGGABLE void Set(
+		bool InbClearColor,
+		int32 InNumClearColors,
+		const FLinearColor* InColorArray,
+		bool InbClearDepth,
+		float InDepth,
+		bool InbClearStencil,
+		uint32 InStencil,
+		FIntRect InExcludeRect
+		)
+	{
+		Type = ERCT_ClearMRT;
+		check(InNumClearColors < MaxSimultaneousRenderTargets);
+		for (int32 Index = 0; Index < InNumClearColors; Index++)
+		{
+			ColorArray[Index] = InColorArray[Index];
+		}
+		ExcludeRect = InExcludeRect;
+		Depth = InDepth;
+		Stencil = InStencil;
+		NumClearColors = InNumClearColors;
+		bClearColor = InbClearColor;
+		bClearDepth = InbClearDepth;
+		bClearStencil = InbClearStencil;
+	}
+};
+
+
 struct FBoundShaderStateInput
 {
 	FVertexDeclarationRHIParamRef VertexDeclarationRHI;
@@ -709,10 +817,13 @@ private:
 				PageSize = 64 * 1024
 			};
 
+			RHI_API static TLockFreeFixedSizeAllocator<PageSize> TheAllocator;
+
 			FPage() :
 				NextPage(nullptr)
 			{
-				Head = (uint8*)FMemory::Malloc(PageSize, Alignment);
+				Head = (uint8*)TheAllocator.Allocate();
+				check((UPTRINT(Head) & (Alignment - 1)) == 0);
 				Current = Head;
 				Tail = Head + PageSize;
 				//::OutputDebugStringW(*FString::Printf(TEXT("PAGE %p, %p - %p\n"), this, Head, Tail));
@@ -720,7 +831,7 @@ private:
 
 			~FPage()
 			{
-				FMemory::Free(Head);
+				TheAllocator.Free(Head);
 				Head = Current = Tail = nullptr;
 
 				// Caller/owner is responsible for freeing the linked list
@@ -1262,7 +1373,7 @@ public:
 
 	}
 
-	void BeginDrawPrimitiveUP(uint32 PrimitiveType, uint32 NumPrimitives, uint32 NumVertices, uint32 VertexDataStride, void*& OutVertexData)
+	FORCEINLINE_DEBUGGABLE void BeginDrawPrimitiveUP(uint32 PrimitiveType, uint32 NumPrimitives, uint32 NumVertices, uint32 VertexDataStride, void*& OutVertexData)
 	{
 		if (Bypass())
 		{
@@ -1280,7 +1391,7 @@ public:
 		DrawUPData.OutVertexData = OutVertexData;
 	}
 
-	void EndDrawPrimitiveUP()
+	FORCEINLINE_DEBUGGABLE void EndDrawPrimitiveUP()
 	{
 		if (Bypass())
 		{
@@ -1295,7 +1406,7 @@ public:
 		DrawUPData.NumVertices = 0;
 	}
 
-	void BeginDrawIndexedPrimitiveUP(uint32 PrimitiveType, uint32 NumPrimitives, uint32 NumVertices, uint32 VertexDataStride, void*& OutVertexData, uint32 MinVertexIndex, uint32 NumIndices, uint32 IndexDataStride, void*& OutIndexData)
+	FORCEINLINE_DEBUGGABLE void BeginDrawIndexedPrimitiveUP(uint32 PrimitiveType, uint32 NumPrimitives, uint32 NumVertices, uint32 VertexDataStride, void*& OutVertexData, uint32 MinVertexIndex, uint32 NumIndices, uint32 IndexDataStride, void*& OutIndexData)
 	{
 		if (Bypass())
 		{
@@ -1318,7 +1429,8 @@ public:
 		DrawUPData.OutIndexData = OutIndexData;
 
 	}
-	void EndDrawIndexedPrimitiveUP()
+
+	FORCEINLINE_DEBUGGABLE void EndDrawIndexedPrimitiveUP()
 	{
 		if (Bypass())
 		{
@@ -1336,7 +1448,7 @@ public:
 		DrawUPData.NumVertices = 0;
 	}
 
-	void SetComputeShader(FComputeShaderRHIParamRef ComputeShader)
+	FORCEINLINE_DEBUGGABLE void SetComputeShader(FComputeShaderRHIParamRef ComputeShader)
 	{
 		if (Bypass())
 		{
@@ -1348,7 +1460,7 @@ public:
 		ComputeShader->AddRef();
 	}
 
-	void DispatchComputeShader(uint32 ThreadGroupCountX, uint32 ThreadGroupCountY, uint32 ThreadGroupCountZ)
+	FORCEINLINE_DEBUGGABLE void DispatchComputeShader(uint32 ThreadGroupCountX, uint32 ThreadGroupCountY, uint32 ThreadGroupCountZ)
 	{
 		if (Bypass())
 		{
@@ -1359,7 +1471,7 @@ public:
 		Cmd->Set(ThreadGroupCountX, ThreadGroupCountY, ThreadGroupCountZ);
 	}
 
-	void DispatchIndirectComputeShader(FVertexBufferRHIParamRef ArgumentBuffer, uint32 ArgumentOffset)
+	FORCEINLINE_DEBUGGABLE void DispatchIndirectComputeShader(FVertexBufferRHIParamRef ArgumentBuffer, uint32 ArgumentOffset)
 	{
 		if (Bypass())
 		{
@@ -1371,7 +1483,7 @@ public:
 		ArgumentBuffer->AddRef();
 	}
 
-	void AutomaticCacheFlushAfterComputeShader(bool bEnable)
+	FORCEINLINE_DEBUGGABLE void AutomaticCacheFlushAfterComputeShader(bool bEnable)
 	{
 		if (Bypass())
 		{
@@ -1382,7 +1494,7 @@ public:
 		Cmd->Set(bEnable);
 	}
 
-	void FlushComputeShaderCache()
+	FORCEINLINE_DEBUGGABLE void FlushComputeShaderCache()
 	{
 		if (Bypass())
 		{
@@ -1393,7 +1505,7 @@ public:
 		Cmd->Set();
 	}
 
-	void DrawPrimitiveIndirect(uint32 PrimitiveType, FVertexBufferRHIParamRef ArgumentBuffer, uint32 ArgumentOffset)
+	FORCEINLINE_DEBUGGABLE void DrawPrimitiveIndirect(uint32 PrimitiveType, FVertexBufferRHIParamRef ArgumentBuffer, uint32 ArgumentOffset)
 	{
 		if (Bypass())
 		{
@@ -1405,7 +1517,7 @@ public:
 		ArgumentBuffer->AddRef();
 	}
 
-	void DrawIndexedIndirect(FIndexBufferRHIParamRef IndexBufferRHI, uint32 PrimitiveType, FStructuredBufferRHIParamRef ArgumentsBufferRHI, uint32 DrawArgumentsIndex, uint32 NumInstances)
+	FORCEINLINE_DEBUGGABLE void DrawIndexedIndirect(FIndexBufferRHIParamRef IndexBufferRHI, uint32 PrimitiveType, FStructuredBufferRHIParamRef ArgumentsBufferRHI, uint32 DrawArgumentsIndex, uint32 NumInstances)
 	{
 		if (Bypass())
 		{
@@ -1418,7 +1530,7 @@ public:
 		ArgumentsBufferRHI->AddRef();
 	}
 
-	void DrawIndexedPrimitiveIndirect(uint32 PrimitiveType, FIndexBufferRHIParamRef IndexBuffer, FVertexBufferRHIParamRef ArgumentsBuffer, uint32 ArgumentOffset)
+	FORCEINLINE_DEBUGGABLE void DrawIndexedPrimitiveIndirect(uint32 PrimitiveType, FIndexBufferRHIParamRef IndexBuffer, FVertexBufferRHIParamRef ArgumentsBuffer, uint32 ArgumentOffset)
 	{
 		if (Bypass())
 		{
@@ -1431,7 +1543,7 @@ public:
 		ArgumentsBuffer->AddRef();
 	}
 
-	void EnableDepthBoundsTest(bool bEnable, float MinDepth, float MaxDepth)
+	FORCEINLINE_DEBUGGABLE void EnableDepthBoundsTest(bool bEnable, float MinDepth, float MaxDepth)
 	{
 		if (Bypass())
 		{
@@ -1443,7 +1555,52 @@ public:
 
 	}
 
+	FORCEINLINE_DEBUGGABLE void ClearUAV(FUnorderedAccessViewRHIParamRef UnorderedAccessViewRHI, const uint32(&Values)[4])
+	{
+		if (Bypass())
+		{
+			ClearUAV_Internal(UnorderedAccessViewRHI, Values);
+			return;
+		}
+		auto* Cmd = AddCommand<FRHICommandClearUAV>();
+		Cmd->Set(UnorderedAccessViewRHI, Values);
+		UnorderedAccessViewRHI->AddRef();
+	}
 
+	FORCEINLINE_DEBUGGABLE void CopyToResolveTarget(FTextureRHIParamRef SourceTextureRHI, FTextureRHIParamRef DestTextureRHI, bool bKeepOriginalSurface, const FResolveParams& ResolveParams)
+	{
+		if (Bypass())
+		{
+			CopyToResolveTarget_Internal(SourceTextureRHI, DestTextureRHI, bKeepOriginalSurface, ResolveParams);
+			return;
+		}
+		auto* Cmd = AddCommand<FRHICommandCopyToResolveTarget>();
+		Cmd->Set(SourceTextureRHI, DestTextureRHI, bKeepOriginalSurface, ResolveParams);
+		SourceTextureRHI->AddRef();
+		DestTextureRHI->AddRef();
+	}
+
+	FORCEINLINE_DEBUGGABLE void Clear(bool bClearColor, const FLinearColor& Color, bool bClearDepth, float Depth, bool bClearStencil, uint32 Stencil, FIntRect ExcludeRect)
+	{
+		if (Bypass())
+		{
+			Clear_Internal(bClearColor, Color, bClearDepth, Depth, bClearStencil, Stencil, ExcludeRect);
+			return;
+		}
+		auto* Cmd = AddCommand<FRHICommandClear>();
+		Cmd->Set(bClearColor, Color, bClearDepth, Depth, bClearStencil, Stencil, ExcludeRect);
+	}
+
+	FORCEINLINE_DEBUGGABLE void ClearMRT(bool bClearColor, int32 NumClearColors, const FLinearColor* ClearColorArray, bool bClearDepth, float Depth, bool bClearStencil, uint32 Stencil, FIntRect ExcludeRect)
+	{
+		if (Bypass())
+		{
+			ClearMRT_Internal(bClearColor, NumClearColors, ClearColorArray, bClearDepth, Depth, bClearStencil, Stencil, ExcludeRect);
+			return;
+		}
+		auto* Cmd = AddCommand<FRHICommandClearMRT>();
+		Cmd->Set(bClearColor, NumClearColors, ClearColorArray, bClearDepth, Depth, bClearStencil, Stencil, ExcludeRect);
+	}
 
 	const SIZE_T GetUsedMemory() const;
 
@@ -1561,6 +1718,9 @@ public:
 };
 #endif
 
+// typedef to mark the recursive use of commandlists in the RHI implmentations
+typedef FRHICommandList FRHICommandList_RecursiveHazardous;
+
 class RHI_API FRHICommandListExecutor
 {
 public:
@@ -1573,7 +1733,6 @@ public:
 	{
 	}
 	static inline FRHICommandListImmediate& GetImmediateCommandList();
-	static inline FRHICommandListImmediate& GetRecursiveRHICommandList(); // Only for use by RHI implementations
 
 	void ExecuteList(FRHICommandList& CmdList);
 	void LatchBypass();
@@ -1609,12 +1768,6 @@ FORCEINLINE_DEBUGGABLE FRHICommandListImmediate& FRHICommandListExecutor::GetImm
 {
 	return GRHICommandList.CommandListImmediate;
 }
-FORCEINLINE_DEBUGGABLE FRHICommandListImmediate& FRHICommandListExecutor::GetRecursiveRHICommandList()
-{
-	check(!GRHICommandList.CommandListImmediate.bExecuting); // queued RHI commands must not be implemented by calling other queued RHI commands
-	GRHICommandList.CommandListImmediate.Flush(); // if for some reason this is called non-recursively, this flushes everything to get us in a reasonable state for the recursive, hazardous calls
-	return GRHICommandList.CommandListImmediate;
-}
 
 #define DEFINE_RHIMETHOD_CMDLIST(Type,Name,ParameterTypesAndNames,ParameterNames,ReturnStatement,NullImplementation)
 #define DEFINE_RHIMETHOD(Type, Name, ParameterTypesAndNames, ParameterNames, ReturnStatement, NullImplementation)
diff --git a/Engine/Source/Runtime/RHI/Public/RHIMethods.h b/Engine/Source/Runtime/RHI/Public/RHIMethods.h
--- a/Engine/Source/Runtime/RHI/Public/RHIMethods.h
+++ b/Engine/Source/Runtime/RHI/Public/RHIMethods.h
@@ -597,7 +597,7 @@ DEFINE_RHIMETHOD_GLOBAL_3(
 	);
 
 /** Clears a UAV to the multi-component value provided. */
-DEFINE_RHIMETHOD_2(
+	DEFINE_RHIMETHOD_CMDLIST_2(
 	void,ClearUAV,
 	FUnorderedAccessViewRHIParamRef,UnorderedAccessViewRHI,
 	const uint32*,Values,
@@ -1111,7 +1111,7 @@ DEFINE_RHIMETHOD_GLOBALFLUSH_5(
 * @param bKeepOriginalSurface - true if the original surface will still be used after this function so must remain valid
 * @param ResolveParams - optional resolve params
 */
-DEFINE_RHIMETHOD_4(
+DEFINE_RHIMETHOD_CMDLIST_4(
 	void,CopyToResolveTarget,
 	FTextureRHIParamRef,SourceTexture,
 	FTextureRHIParamRef,DestTexture,
@@ -1842,7 +1842,7 @@ DEFINE_RHIMETHOD_CMDLIST_0(
  * This method clears all MRT's, but to only one color value
  * @param ExcludeRect within the viewport in pixels, is only a hint to optimize - if a fast clear can be done this is preferred
  */
-DEFINE_RHIMETHOD_7(
+DEFINE_RHIMETHOD_CMDLIST_7(
 	void,Clear,
 	bool,bClearColor,
 	const FLinearColor&,Color,
@@ -1858,7 +1858,7 @@ DEFINE_RHIMETHOD_7(
  * This method clears all MRT's to potentially different color values
  * @param ExcludeRect within the viewport in pixels, is only a hint to optimize - if a fast clear can be done this is preferred
  */
-DEFINE_RHIMETHOD_8(
+DEFINE_RHIMETHOD_CMDLIST_8(
 	void,ClearMRT,
 	bool,bClearColor,
 	int32,NumClearColors,
diff --git a/Engine/Source/Runtime/Renderer/Private/AtmosphereRendering.cpp b/Engine/Source/Runtime/Renderer/Private/AtmosphereRendering.cpp
--- a/Engine/Source/Runtime/Renderer/Private/AtmosphereRendering.cpp
+++ b/Engine/Source/Runtime/Renderer/Private/AtmosphereRendering.cpp
@@ -1081,7 +1081,7 @@ void FAtmosphericFogSceneInfo::GetLayerValue(int Layer, float& AtmosphereR, FVec
 	DhdH = FVector4(DMin, DMax, DMinP, DMaxP);
 }
 
-void FAtmosphericFogSceneInfo::RenderAtmosphereShaders(FRHICommandListImmediate& RHICmdList, const FViewInfo& View, const FIntRect& ViewRect)
+void FAtmosphericFogSceneInfo::RenderAtmosphereShaders(FRHICommandList& RHICmdList, const FViewInfo& View, const FIntRect& ViewRect)
 {
 	check(Component != NULL);
 	switch(AtmospherePhase)
diff --git a/Engine/Source/Runtime/Renderer/Private/AtmosphereRendering.h b/Engine/Source/Runtime/Renderer/Private/AtmosphereRendering.h
--- a/Engine/Source/Runtime/Renderer/Private/AtmosphereRendering.h
+++ b/Engine/Source/Runtime/Renderer/Private/AtmosphereRendering.h
@@ -72,7 +72,7 @@ private:
 	FIntPoint GetTextureSize();
 	inline void DrawQuad(FRHICommandList& RHICmdList, const FIntRect& ViewRect, FShader* VertexShader);
 	void GetLayerValue(int Layer, float& AtmosphereR, FVector4& DhdH);
-	void RenderAtmosphereShaders(FRHICommandListImmediate& RHICmdList, const FViewInfo& View, const FIntRect& ViewRect);
+	void RenderAtmosphereShaders(FRHICommandList& RHICmdList, const FViewInfo& View, const FIntRect& ViewRect);
 	void PrecomputeAtmosphereData(FRHICommandListImmediate& RHICmdList, const FViewInfo* View, FSceneViewFamily& ViewFamily);
 
 	void ReadPixelsPtr(FRHICommandListImmediate& RHICmdList, TRefCountPtr<IPooledRenderTarget> RenderTarget, FColor* OutData, FIntRect InRect);
diff --git a/Engine/Source/Runtime/Renderer/Private/CompositionLighting/PostProcessDeferredDecals.cpp b/Engine/Source/Runtime/Renderer/Private/CompositionLighting/PostProcessDeferredDecals.cpp
--- a/Engine/Source/Runtime/Renderer/Private/CompositionLighting/PostProcessDeferredDecals.cpp
+++ b/Engine/Source/Runtime/Renderer/Private/CompositionLighting/PostProcessDeferredDecals.cpp
@@ -350,7 +350,7 @@ IMPLEMENT_SHADER_TYPE(,FStencilDecalMaskPS,TEXT("DeferredDecal"),TEXT("StencilDe
 FGlobalBoundShaderState StencilDecalMaskBoundShaderState;
 
 /** Draws a full view quad that sets stencil to 1 anywhere that decals should not be projected. */
-void StencilDecalMask(FRHICommandListImmediate& RHICmdList, const FSceneView& View)
+void StencilDecalMask(FRHICommandList& RHICmdList, const FSceneView& View)
 {
 	SCOPED_DRAW_EVENT(StencilDecalMask, DEC_SCENE_ITEMS);
 	RHICmdList.SetRasterizerState(TStaticRasterizerState<FM_Solid, CM_None>::GetRHI());
diff --git a/Engine/Source/Runtime/Renderer/Private/DeferredShadingRenderer.cpp b/Engine/Source/Runtime/Renderer/Private/DeferredShadingRenderer.cpp
--- a/Engine/Source/Runtime/Renderer/Private/DeferredShadingRenderer.cpp
+++ b/Engine/Source/Runtime/Renderer/Private/DeferredShadingRenderer.cpp
@@ -963,7 +963,7 @@ void FDeferredShadingSceneRenderer::ClearLPVs(FRHICommandListImmediate& RHICmdLi
 			{
 				SCOPED_DRAW_EVENT(ClearLPVs, DEC_SCENE_ITEMS);
 				SCOPE_CYCLE_COUNTER(STAT_UpdateLPVs);
-				LightPropagationVolume->InitSettings(Views[ViewIndex]);
+				LightPropagationVolume->InitSettings(RHICmdList, Views[ViewIndex]);
 				LightPropagationVolume->Clear(RHICmdList);
 			}
 		}
diff --git a/Engine/Source/Runtime/Renderer/Private/DeferredShadingRenderer.h b/Engine/Source/Runtime/Renderer/Private/DeferredShadingRenderer.h
--- a/Engine/Source/Runtime/Renderer/Private/DeferredShadingRenderer.h
+++ b/Engine/Source/Runtime/Renderer/Private/DeferredShadingRenderer.h
@@ -294,7 +294,7 @@ private:
 	  * @param LightIndex The light's index into FScene::Lights
 	  * @return true if anything got rendered
 	  */
-	void RenderLight(FRHICommandListImmediate& RHICmdList, const FLightSceneInfo* LightSceneInfo, bool bRenderOverlap, bool bIssueDrawEvent);
+	void RenderLight(FRHICommandList& RHICmdList, const FLightSceneInfo* LightSceneInfo, bool bRenderOverlap, bool bIssueDrawEvent);
 
 	/** Renders an array of simple lights using standard deferred shading. */
 	void RenderSimpleLightsStandardDeferred(FRHICommandListImmediate& RHICmdList, const FSimpleLightArray& SimpleLights);
@@ -303,16 +303,16 @@ private:
 	void ClearTranslucentVolumeLighting(FRHICommandListImmediate& RHICmdList);
 
 	/** Add AmbientCubemap to the lighting volumes. */
-	void InjectAmbientCubemapTranslucentVolumeLighting(FRHICommandListImmediate& RHICmdList);
+	void InjectAmbientCubemapTranslucentVolumeLighting(FRHICommandList& RHICmdList);
 
 	/** Renders indirect lighting into the translucent lighting volumes. */
-	void CompositeIndirectTranslucentVolumeLighting(FRHICommandListImmediate& RHICmdList);
+	void CompositeIndirectTranslucentVolumeLighting(FRHICommandList& RHICmdList);
 
 	/** Clears the volume texture used to accumulate per object shadows for translucency. */
-	void ClearTranslucentVolumePerObjectShadowing(FRHICommandListImmediate& RHICmdList);
+	void ClearTranslucentVolumePerObjectShadowing(FRHICommandList& RHICmdList);
 
 	/** Accumulates the per object shadow's contribution for translucency. */
-	void AccumulateTranslucentVolumeObjectShadowing(FRHICommandListImmediate& RHICmdList, const FProjectedShadowInfo* InProjectedShadowInfo, bool bClearVolume);
+	void AccumulateTranslucentVolumeObjectShadowing(FRHICommandList& RHICmdList, const FProjectedShadowInfo* InProjectedShadowInfo, bool bClearVolume);
 
 	/** Accumulates direct lighting for the given light.  InProjectedShadowInfo can be NULL in which case the light will be unshadowed. */
 	void InjectTranslucentVolumeLighting(FRHICommandListImmediate& RHICmdList, const FLightSceneInfo& LightSceneInfo, const FProjectedShadowInfo* InProjectedShadowInfo);
diff --git a/Engine/Source/Runtime/Renderer/Private/LightPropagationVolume.cpp b/Engine/Source/Runtime/Renderer/Private/LightPropagationVolume.cpp
--- a/Engine/Source/Runtime/Renderer/Private/LightPropagationVolume.cpp
+++ b/Engine/Source/Runtime/Renderer/Private/LightPropagationVolume.cpp
@@ -804,7 +804,7 @@ FLightPropagationVolume::~FLightPropagationVolume()
 /**
 * Sets up the LPV at the beginning of the frame
 */
-void FLightPropagationVolume::InitSettings( const FSceneView& View )
+void FLightPropagationVolume::InitSettings(FRHICommandList& RHICmdList, const FSceneView& View)
 	{
 	check(View.GetFeatureLevel() >= ERHIFeatureLevel::SM5);
 
@@ -879,8 +879,8 @@ void FLightPropagationVolume::InitSettings( const FSceneView& View )
 	if ( bNeedsBufferClear )
 	{
 #if !LPV_VOLUME_TEXTURE
-		uint32 Zero = 0;
-		RHIClearUAV( mVplListHeadBuffer->UAV, &Zero );
+		uint32 ClearValues[4] = {0};
+		RHICmdList.ClearUAV( mVplListHeadBuffer->UAV, ClearValues );
 #endif
 
 		ClearMultiplier = 0.0f;
diff --git a/Engine/Source/Runtime/Renderer/Private/LightPropagationVolume.h b/Engine/Source/Runtime/Renderer/Private/LightPropagationVolume.h
--- a/Engine/Source/Runtime/Renderer/Private/LightPropagationVolume.h
+++ b/Engine/Source/Runtime/Renderer/Private/LightPropagationVolume.h
@@ -124,7 +124,7 @@ public:
 	FLightPropagationVolume();
 	virtual ~FLightPropagationVolume();
 
-	void InitSettings( const FSceneView& View );
+	void InitSettings(FRHICommandList& RHICmdList, const FSceneView& View);
 
 	void Clear(FRHICommandListImmediate& RHICmdList);
 
diff --git a/Engine/Source/Runtime/Renderer/Private/LightRendering.cpp b/Engine/Source/Runtime/Renderer/Private/LightRendering.cpp
--- a/Engine/Source/Runtime/Renderer/Private/LightRendering.cpp
+++ b/Engine/Source/Runtime/Renderer/Private/LightRendering.cpp
@@ -712,7 +712,7 @@ static void SetShaderTemplLightingSimple(
  * @param LightIndex The light's index into FScene::Lights
  * @return true if anything got rendered
  */
-void FDeferredShadingSceneRenderer::RenderLight(FRHICommandListImmediate& RHICmdList, const FLightSceneInfo* LightSceneInfo, bool bRenderOverlap, bool bIssueDrawEvent)
+void FDeferredShadingSceneRenderer::RenderLight(FRHICommandList& RHICmdList, const FLightSceneInfo* LightSceneInfo, bool bRenderOverlap, bool bIssueDrawEvent)
 {
 	SCOPE_CYCLE_COUNTER(STAT_DirectLightRenderingTime);
 	INC_DWORD_STAT(STAT_NumLightsUsingStandardDeferred);
diff --git a/Engine/Source/Runtime/Renderer/Private/LightShaftRendering.cpp b/Engine/Source/Runtime/Renderer/Private/LightShaftRendering.cpp
--- a/Engine/Source/Runtime/Renderer/Private/LightShaftRendering.cpp
+++ b/Engine/Source/Runtime/Renderer/Private/LightShaftRendering.cpp
@@ -643,7 +643,7 @@ void ApplyRadialBlurPasses(
 	}
 }
 
-void FinishOcclusionTerm(FRHICommandListImmediate& RHICmdList, const FViewInfo& View, const FLightSceneInfo* const LightSceneInfo, TRefCountPtr<IPooledRenderTarget>& LightShaftsSource, TRefCountPtr<IPooledRenderTarget>& LightShaftsDest)
+void FinishOcclusionTerm(FRHICommandList& RHICmdList, const FViewInfo& View, const FLightSceneInfo* const LightSceneInfo, TRefCountPtr<IPooledRenderTarget>& LightShaftsSource, TRefCountPtr<IPooledRenderTarget>& LightShaftsDest)
 {
 	TShaderMapRef<FScreenVS> ScreenVertexShader(GetGlobalShaderMap());
 
diff --git a/Engine/Source/Runtime/Renderer/Private/PostProcess/SceneRenderTargets.cpp b/Engine/Source/Runtime/Renderer/Private/PostProcess/SceneRenderTargets.cpp
--- a/Engine/Source/Runtime/Renderer/Private/PostProcess/SceneRenderTargets.cpp
+++ b/Engine/Source/Runtime/Renderer/Private/PostProcess/SceneRenderTargets.cpp
@@ -625,7 +625,7 @@ void FSceneRenderTargets::FinishRenderingCustomDepth(FRHICommandListImmediate& R
 /**
 * Saves a previously rendered scene color target
 */
-void FSceneRenderTargets::ResolveSceneColor(FRHICommandListImmediate& RHICmdList, const FResolveRect& ResolveRect)
+void FSceneRenderTargets::ResolveSceneColor(FRHICommandList& RHICmdList, const FResolveRect& ResolveRect)
 {
     SCOPED_DRAW_EVENT(ResolveSceneColor, DEC_SCENE_ITEMS);
 
@@ -633,7 +633,7 @@ void FSceneRenderTargets::ResolveSceneColor(FRHICommandListImmediate& RHICmdList
 }
 
 /** Resolves the GBuffer targets so that their resolved textures can be sampled. */
-void FSceneRenderTargets::ResolveGBufferSurfaces(FRHICommandListImmediate& RHICmdList, const FResolveRect& ResolveRect)
+void FSceneRenderTargets::ResolveGBufferSurfaces(FRHICommandList& RHICmdList, const FResolveRect& ResolveRect)
 {
 	if (GRHIFeatureLevel >= ERHIFeatureLevel::SM4)
 	{
@@ -679,7 +679,7 @@ void FSceneRenderTargets::BeginRenderingCubeShadowDepth(FRHICommandListImmediate
 	SetRenderTarget(RHICmdList, FTextureRHIRef(), GetCubeShadowDepthZSurface(ShadowResolution));
 }
 
-void FSceneRenderTargets::FinishRenderingShadowDepth(FRHICommandListImmediate& RHICmdList, const FResolveRect& ResolveRect)
+void FSceneRenderTargets::FinishRenderingShadowDepth(FRHICommandList& RHICmdList, const FResolveRect& ResolveRect)
 {
 	// Resolve the shadow depth z surface.
 	RHICmdList.CopyToResolveTarget(GetShadowDepthZSurface(), GetShadowDepthZTexture(), false, FResolveParams(ResolveRect));
@@ -701,7 +701,7 @@ void FSceneRenderTargets::BeginRenderingReflectiveShadowMap(FRHICommandListImmed
 	SetRenderTargets(RHICmdList, ARRAY_COUNT(RenderTargets), RenderTargets, GetReflectiveShadowMapDepthSurface(), 4, Uavs);
 }
 
-void FSceneRenderTargets::FinishRenderingReflectiveShadowMap(FRHICommandListImmediate& RHICmdList, const FResolveRect& ResolveRect)
+void FSceneRenderTargets::FinishRenderingReflectiveShadowMap(FRHICommandList& RHICmdList, const FResolveRect& ResolveRect)
 {
 	// Resolve the shadow depth z surface.
 	RHICmdList.CopyToResolveTarget(GetReflectiveShadowMapDepthSurface(), GetReflectiveShadowMapDepthTexture(), false, FResolveParams(ResolveRect));
@@ -714,7 +714,7 @@ void FSceneRenderTargets::FinishRenderingReflectiveShadowMap(FRHICommandListImme
 	SetRenderTargets(RHICmdList, ARRAY_COUNT(RenderTargets), RenderTargets, FTextureRHIParamRef(), 2, Uavs);
 }
 
-void FSceneRenderTargets::FinishRenderingCubeShadowDepth(FRHICommandListImmediate& RHICmdList, int32 ShadowResolution, const FResolveParams& ResolveParams)
+void FSceneRenderTargets::FinishRenderingCubeShadowDepth(FRHICommandList& RHICmdList, int32 ShadowResolution, const FResolveParams& ResolveParams)
 {
 	SCOPED_DRAW_EVENT(FinishRenderingCubeShadowDepth, DEC_SCENE_ITEMS);
 	RHICmdList.CopyToResolveTarget(GetCubeShadowDepthZSurface(ShadowResolution), GetCubeShadowDepthZTexture(ShadowResolution), false, ResolveParams);
@@ -799,7 +799,7 @@ bool FSceneRenderTargets::BeginRenderingSeparateTranslucency(FRHICommandListImme
 	return false;
 }
 
-void FSceneRenderTargets::FinishRenderingSeparateTranslucency(FRHICommandListImmediate& RHICmdList, const FViewInfo& View)
+void FSceneRenderTargets::FinishRenderingSeparateTranslucency(FRHICommandList& RHICmdList, const FViewInfo& View)
 {
 	if(IsSeparateTranslucencyActive(View))
 	{
@@ -810,14 +810,14 @@ void FSceneRenderTargets::FinishRenderingSeparateTranslucency(FRHICommandListImm
 	}
 }
 
-void FSceneRenderTargets::ResolveSceneDepthTexture(FRHICommandListImmediate& RHICmdList)
+void FSceneRenderTargets::ResolveSceneDepthTexture(FRHICommandList& RHICmdList)
 {
 	SCOPED_DRAW_EVENT(ResolveSceneDepthTexture, DEC_SCENE_ITEMS);
 
 	RHICmdList.CopyToResolveTarget(GetSceneDepthSurface(), GetSceneDepthTexture(), true, FResolveParams());
 }
 
-void FSceneRenderTargets::ResolveSceneDepthToAuxiliaryTexture(FRHICommandListImmediate& RHICmdList)
+void FSceneRenderTargets::ResolveSceneDepthToAuxiliaryTexture(FRHICommandList& RHICmdList)
 {
 	// Resolve the scene depth to an auxiliary texture when SM3/SM4 is in use. This needs to happen so the auxiliary texture can be bound as a shader parameter
 	// while the primary scene depth texture can be bound as the target. Simultaneously binding a single DepthStencil resource as a parameter and target
diff --git a/Engine/Source/Runtime/Renderer/Private/PostProcess/SceneRenderTargets.h b/Engine/Source/Runtime/Renderer/Private/PostProcess/SceneRenderTargets.h
--- a/Engine/Source/Runtime/Renderer/Private/PostProcess/SceneRenderTargets.h
+++ b/Engine/Source/Runtime/Renderer/Private/PostProcess/SceneRenderTargets.h
@@ -162,10 +162,10 @@ public:
 	/**
 	 * Resolve a previously rendered scene color surface.
 	 */
-	void ResolveSceneColor(FRHICommandListImmediate& RHICmdList, const FResolveRect& ResolveRect = FResolveRect());
+	void ResolveSceneColor(FRHICommandList& RHICmdList, const FResolveRect& ResolveRect = FResolveRect());
 
 	/** Resolves the GBuffer targets so that their resolved textures can be sampled. */
-	void ResolveGBufferSurfaces(FRHICommandListImmediate& RHICmdList, const FResolveRect& ResolveRect = FResolveRect());
+	void ResolveGBufferSurfaces(FRHICommandList& RHICmdList, const FResolveRect& ResolveRect = FResolveRect());
 
 	void BeginRenderingShadowDepth(FRHICommandListImmediate& RHICmdList);
 
@@ -176,22 +176,22 @@ public:
 	 * Called when finished rendering to the subject shadow depths so the surface can be copied to texture
 	 * @param ResolveParams - optional resolve params
 	 */
-	void FinishRenderingShadowDepth(FRHICommandListImmediate& RHICmdList, const FResolveRect& ResolveRect = FResolveRect());
+	void FinishRenderingShadowDepth(FRHICommandList& RHICmdList, const FResolveRect& ResolveRect = FResolveRect());
 
 	void BeginRenderingReflectiveShadowMap(FRHICommandListImmediate& RHICmdList, class FLightPropagationVolume* Lpv);
-	void FinishRenderingReflectiveShadowMap(FRHICommandListImmediate& RHICmdList, const FResolveRect& ResolveRect = FResolveRect());
+	void FinishRenderingReflectiveShadowMap(FRHICommandList& RHICmdList, const FResolveRect& ResolveRect = FResolveRect());
 
 	/** Resolves the appropriate shadow depth cube map and restores default state. */
-	void FinishRenderingCubeShadowDepth(FRHICommandListImmediate& RHICmdList, int32 ShadowResolution, const FResolveParams& ResolveParams = FResolveParams());
+	void FinishRenderingCubeShadowDepth(FRHICommandList& RHICmdList, int32 ShadowResolution, const FResolveParams& ResolveParams = FResolveParams());
 	
 	void BeginRenderingTranslucency(FRHICommandListImmediate& RHICmdList, const class FViewInfo& View);
 
 	bool BeginRenderingSeparateTranslucency(FRHICommandListImmediate& RHICmdList, const FViewInfo& View, bool bFirstTimeThisFrame);
-	void FinishRenderingSeparateTranslucency(FRHICommandListImmediate& RHICmdList, const FViewInfo& View);
+	void FinishRenderingSeparateTranslucency(FRHICommandList& RHICmdList, const FViewInfo& View);
 	void FreeSeparateTranslucency();
 
-	void ResolveSceneDepthTexture(FRHICommandListImmediate& RHICmdList);
-	void ResolveSceneDepthToAuxiliaryTexture(FRHICommandListImmediate& RHICmdList);
+	void ResolveSceneDepthTexture(FRHICommandList& RHICmdList);
+	void ResolveSceneDepthToAuxiliaryTexture(FRHICommandList& RHICmdList);
 
 	void BeginRenderingPrePass(FRHICommandListImmediate& RHICmdList);
 	void FinishRenderingPrePass(FRHICommandListImmediate& RHICmdList);
diff --git a/Engine/Source/Runtime/Renderer/Private/ReflectionEnvironmentCapture.cpp b/Engine/Source/Runtime/Renderer/Private/ReflectionEnvironmentCapture.cpp
--- a/Engine/Source/Runtime/Renderer/Private/ReflectionEnvironmentCapture.cpp
+++ b/Engine/Source/Runtime/Renderer/Private/ReflectionEnvironmentCapture.cpp
@@ -774,7 +774,7 @@ int32 FindOrAllocateCubemapIndex(FScene* Scene, const UReflectionCaptureComponen
 	return CaptureIndex;
 }
 
-void ClearScratchCubemaps(FRHICommandListImmediate& RHICmdList)
+void ClearScratchCubemaps(FRHICommandList& RHICmdList)
 {
 	// Clear scratch render targets to a consistent but noticeable value
 	// This makes debugging capture issues much easier, otherwise the random contents from previous captures is shown
@@ -915,7 +915,7 @@ void CaptureSceneToScratchCubemap(FRHICommandListImmediate& RHICmdList, FSceneRe
 	delete SceneRenderer;
 }
 
-void CopyCubemapToScratchCubemap(FRHICommandListImmediate& RHICmdList, UTextureCube* SourceCubemap, bool bIsSkyLight, bool bLowerHemisphereIsBlack)
+void CopyCubemapToScratchCubemap(FRHICommandList& RHICmdList, UTextureCube* SourceCubemap, bool bIsSkyLight, bool bLowerHemisphereIsBlack)
 {
 	check(SourceCubemap);
 	
@@ -1438,7 +1438,7 @@ void CopyToSceneArray(FRHICommandListImmediate& RHICmdList, FScene* Scene, FRefl
 	}
 }
 
-void CopyToComponentTexture(FRHICommandListImmediate& RHICmdList, FScene* Scene, FReflectionCaptureProxy* ReflectionProxy)
+void CopyToComponentTexture(FRHICommandList& RHICmdList, FScene* Scene, FReflectionCaptureProxy* ReflectionProxy)
 {
 	check(ReflectionProxy->SM4FullHDRCubemap);
 	const int32 EffectiveTopMipSize = GReflectionCaptureSize;
@@ -1520,7 +1520,7 @@ void FScene::UpdateReflectionCaptureContents(UReflectionCaptureComponent* Captur
 	}
 }
 
-void CopyToSkyTexture(FRHICommandListImmediate& RHICmdList, FScene* Scene, FTexture* ProcessedTexture)
+void CopyToSkyTexture(FRHICommandList& RHICmdList, FScene* Scene, FTexture* ProcessedTexture)
 {
 	const int32 EffectiveTopMipSize = ProcessedTexture->GetSizeX();
 	const int32 NumMips = FMath::CeilLogTwo(EffectiveTopMipSize) + 1;
diff --git a/Engine/Source/Runtime/Renderer/Private/ScenePostProcessing.cpp b/Engine/Source/Runtime/Renderer/Private/ScenePostProcessing.cpp
--- a/Engine/Source/Runtime/Renderer/Private/ScenePostProcessing.cpp
+++ b/Engine/Source/Runtime/Renderer/Private/ScenePostProcessing.cpp
@@ -98,7 +98,7 @@ void FDeferredShadingSceneRenderer::FinishRenderViewTarget(FRHICommandListImmedi
 }
 
 // TODO: REMOVE if no longer needed:
-void FSceneRenderer::GammaCorrectToViewportRenderTarget(FRHICommandListImmediate& RHICmdList, const FViewInfo* View, float OverrideGamma)
+void FSceneRenderer::GammaCorrectToViewportRenderTarget(FRHICommandList& RHICmdList, const FViewInfo* View, float OverrideGamma)
 {
 	// Set the view family's render target/viewport.
 	SetRenderTarget(RHICmdList, ViewFamily.RenderTarget->GetRenderTargetTexture(), FTextureRHIRef());
diff --git a/Engine/Source/Runtime/Renderer/Private/SceneRendering.h b/Engine/Source/Runtime/Renderer/Private/SceneRendering.h
--- a/Engine/Source/Runtime/Renderer/Private/SceneRendering.h
+++ b/Engine/Source/Runtime/Renderer/Private/SceneRendering.h
@@ -551,7 +551,7 @@ protected:
 	bool ShouldRenderTranslucency() const;
 
 	/** TODO: REMOVE if no longer needed: Copies scene color to the viewport's render target after applying gamma correction. */
-	void GammaCorrectToViewportRenderTarget(FRHICommandListImmediate& RHICmdList, const FViewInfo* View, float OverrideGamma);
+	void GammaCorrectToViewportRenderTarget(FRHICommandList& RHICmdList, const FViewInfo* View, float OverrideGamma);
 
 	/** Updates state for the end of the frame. */
 	void RenderFinish(FRHICommandListImmediate& RHICmdList);
diff --git a/Engine/Source/Runtime/Renderer/Private/ShadowRendering.cpp b/Engine/Source/Runtime/Renderer/Private/ShadowRendering.cpp
--- a/Engine/Source/Runtime/Renderer/Private/ShadowRendering.cpp
+++ b/Engine/Source/Runtime/Renderer/Private/ShadowRendering.cpp
@@ -1126,7 +1126,7 @@ static void CheckShadowDepthMaterials(const FMaterialRenderProxy* InRenderProxy,
 	check(Material == InMaterial);
 }
 
-void FProjectedShadowInfo::ClearDepth(FRHICommandListImmediate& RHICmdList, FDeferredShadingSceneRenderer* SceneRenderer)
+void FProjectedShadowInfo::ClearDepth(FRHICommandList& RHICmdList, FDeferredShadingSceneRenderer* SceneRenderer)
 {
 	if (bOnePassPointLightShadow)
 	{
diff --git a/Engine/Source/Runtime/Renderer/Private/ShadowRendering.h b/Engine/Source/Runtime/Renderer/Private/ShadowRendering.h
--- a/Engine/Source/Runtime/Renderer/Private/ShadowRendering.h
+++ b/Engine/Source/Runtime/Renderer/Private/ShadowRendering.h
@@ -664,7 +664,7 @@ public:
 	 */
 	void RenderDepth(FRHICommandListImmediate& RHICmdList, class FDeferredShadingSceneRenderer* SceneRenderer);
 
-	void ClearDepth(FRHICommandListImmediate& RHICmdList, class FDeferredShadingSceneRenderer* SceneRenderer);
+	void ClearDepth(FRHICommandList& RHICmdList, class FDeferredShadingSceneRenderer* SceneRenderer);
 
 	/** Renders shadow maps for translucent primitives. */
 	void RenderTranslucencyDepths(FRHICommandListImmediate& RHICmdList, class FDeferredShadingSceneRenderer* SceneRenderer);
diff --git a/Engine/Source/Runtime/Renderer/Private/TranslucentLighting.cpp b/Engine/Source/Runtime/Renderer/Private/TranslucentLighting.cpp
--- a/Engine/Source/Runtime/Renderer/Private/TranslucentLighting.cpp
+++ b/Engine/Source/Runtime/Renderer/Private/TranslucentLighting.cpp
@@ -1052,7 +1052,7 @@ void RasterizeToVolumeTexture(FRHICommandList& RHICmdList, FVolumeBounds VolumeB
 
 /** Helper function that clears the given volume texture render targets. */
 template<int32 NumRenderTargets>
-void ClearVolumeTextures(FRHICommandListImmediate& RHICmdList, const FTextureRHIParamRef* RenderTargets, const FLinearColor* ClearColors)
+void ClearVolumeTextures(FRHICommandList& RHICmdList, const FTextureRHIParamRef* RenderTargets, const FLinearColor* ClearColors)
 {
 	SetRenderTargets(RHICmdList, NumRenderTargets, RenderTargets, FTextureRHIRef(), 0, NULL);
 
@@ -1163,7 +1163,7 @@ public:
 
 IMPLEMENT_SHADER_TYPE(,FInjectAmbientCubemapPS,TEXT("TranslucentLightingShaders"),TEXT("InjectAmbientCubemapMainPS"),SF_Pixel);
 
-void FDeferredShadingSceneRenderer::InjectAmbientCubemapTranslucentVolumeLighting(FRHICommandListImmediate& RHICmdList)
+void FDeferredShadingSceneRenderer::InjectAmbientCubemapTranslucentVolumeLighting(FRHICommandList& RHICmdList)
 {
 	//@todo - support multiple views
 	const FViewInfo& View = Views[0];
@@ -1316,7 +1316,7 @@ private:
 
 IMPLEMENT_SHADER_TYPE(,FCompositeGIForTranslucencyPS,TEXT("TranslucentLightingShaders"),TEXT("CompositeGIMainPS"),SF_Pixel);
 
-void FDeferredShadingSceneRenderer::CompositeIndirectTranslucentVolumeLighting(FRHICommandListImmediate& RHICmdList)
+void FDeferredShadingSceneRenderer::CompositeIndirectTranslucentVolumeLighting(FRHICommandList& RHICmdList)
 {
 	bool bAnyViewAllowsIndirectLightingCache = false;
 
@@ -1377,7 +1377,7 @@ void FDeferredShadingSceneRenderer::CompositeIndirectTranslucentVolumeLighting(F
 	}
 }
 
-void FDeferredShadingSceneRenderer::ClearTranslucentVolumePerObjectShadowing(FRHICommandListImmediate& RHICmdList)
+void FDeferredShadingSceneRenderer::ClearTranslucentVolumePerObjectShadowing(FRHICommandList& RHICmdList)
 {
 	if (GUseTranslucentLightingVolumes && GSupportsVolumeTextureRendering)
 	{
@@ -1425,7 +1425,7 @@ FVolumeBounds CalculateLightVolumeBounds(const FSphere& LightBounds, const FView
 
 FGlobalBoundShaderState ObjectShadowingBoundShaderState;
 
-void FDeferredShadingSceneRenderer::AccumulateTranslucentVolumeObjectShadowing(FRHICommandListImmediate& RHICmdList, const FProjectedShadowInfo* InProjectedShadowInfo, bool bClearVolume)
+void FDeferredShadingSceneRenderer::AccumulateTranslucentVolumeObjectShadowing(FRHICommandList& RHICmdList, const FProjectedShadowInfo* InProjectedShadowInfo, bool bClearVolume)
 {
 	const FLightSceneInfo* LightSceneInfo = InProjectedShadowInfo->LightSceneInfo;
 
diff --git a/Engine/Source/Runtime/ShaderCore/Private/Shader.cpp b/Engine/Source/Runtime/ShaderCore/Private/Shader.cpp
--- a/Engine/Source/Runtime/ShaderCore/Private/Shader.cpp
+++ b/Engine/Source/Runtime/ShaderCore/Private/Shader.cpp
@@ -1034,7 +1034,7 @@ FShaderType* FindShaderTypeByName(const TCHAR* ShaderTypeName)
 
 
 void DispatchComputeShader(
-	FRHICommandListImmediate& RHICmdList,
+	FRHICommandList& RHICmdList,
 	FShader* Shader,
 	uint32 ThreadGroupCountX,
 	uint32 ThreadGroupCountY,
@@ -1046,7 +1046,7 @@ void DispatchComputeShader(
 
 
 void DispatchIndirectComputeShader(
-	FRHICommandListImmediate& RHICmdList,
+	FRHICommandList& RHICmdList,
 	FShader* Shader,
 	FVertexBufferRHIParamRef ArgumentBuffer,
 	uint32 ArgumentOffset)
diff --git a/Engine/Source/Runtime/ShaderCore/Public/Shader.h b/Engine/Source/Runtime/ShaderCore/Public/Shader.h
--- a/Engine/Source/Runtime/ShaderCore/Public/Shader.h
+++ b/Engine/Source/Runtime/ShaderCore/Public/Shader.h
@@ -1320,7 +1320,7 @@ extern SHADERCORE_API FShaderType* FindShaderTypeByName(const TCHAR* ShaderTypeN
 
 /** Helper function to dispatch a compute shader while checking that parameters have been set correctly. */
 extern SHADERCORE_API void DispatchComputeShader(
-	FRHICommandListImmediate& RHICmdList,
+	FRHICommandList& RHICmdList,
 	FShader* Shader,
 	uint32 ThreadGroupCountX,
 	uint32 ThreadGroupCountY,
@@ -1328,7 +1328,7 @@ extern SHADERCORE_API void DispatchComputeShader(
 
 /** Helper function to dispatch a compute shader indirectly while checking that parameters have been set correctly. */
 extern SHADERCORE_API void DispatchIndirectComputeShader(
-	FRHICommandListImmediate& RHICmdList,
+	FRHICommandList& RHICmdList,
 	FShader* Shader,
 	FVertexBufferRHIParamRef ArgumentBuffer,
 	uint32 ArgumentOffset);
diff --git a/Engine/Source/Runtime/Windows/D3D11RHI/Private/D3D11Commands.cpp b/Engine/Source/Runtime/Windows/D3D11RHI/Private/D3D11Commands.cpp
--- a/Engine/Source/Runtime/Windows/D3D11RHI/Private/D3D11Commands.cpp
+++ b/Engine/Source/Runtime/Windows/D3D11RHI/Private/D3D11Commands.cpp
@@ -1620,8 +1620,6 @@ void FD3D11DynamicRHI::RHIClear(bool bClearColor,const FLinearColor& Color,bool
 
 void FD3D11DynamicRHI::RHIClearMRT(bool bClearColor,int32 NumClearColors,const FLinearColor* ClearColorArray,bool bClearDepth,float Depth,bool bClearStencil,uint32 Stencil, FIntRect ExcludeRect)
 {
-	FRHICommandListImmediate& RHICmdList = FRHICommandListExecutor::GetRecursiveRHICommandList();
-
 	GRHICommandList.Verify();
 
 	// Helper struct to record and restore device states RHIClearMRT modifies.
@@ -1941,65 +1939,68 @@ void FD3D11DynamicRHI::RHIClearMRT(bool bClearColor,int32 NumClearColors,const F
 			TShaderMapRef<TOneColorPixelShaderMRT<8> > MRTPixelShader(GetGlobalShaderMap());
 			PixelShader = *MRTPixelShader;
 		}
-		
-		SetGlobalBoundShaderState(RHICmdList, GD3D11ClearMRTBoundShaderState[FMath::Max(BoundRenderTargets.GetNumActiveTargets() - 1, 0)], GD3D11Vector4VertexDeclaration.VertexDeclarationRHI, *VertexShader, PixelShader);
-		FLinearColor ShaderClearColors[MaxSimultaneousRenderTargets];
-		FMemory::MemZero(ShaderClearColors);
 
-		for (int32 i = 0; i < NumClearColors; i++)
 		{
-			ShaderClearColors[i] = ClearColorArray[i];
-		}
+			FRHICommandList_RecursiveHazardous RHICmdList;
+			SetGlobalBoundShaderState(RHICmdList, GD3D11ClearMRTBoundShaderState[FMath::Max(BoundRenderTargets.GetNumActiveTargets() - 1, 0)], GD3D11Vector4VertexDeclaration.VertexDeclarationRHI, *VertexShader, PixelShader);
+			FLinearColor ShaderClearColors[MaxSimultaneousRenderTargets];
+			FMemory::MemZero(ShaderClearColors);
 
-		SetShaderValueArray(RHICmdList, PixelShader->GetPixelShader(),PixelShader->ColorParameter,ShaderClearColors,NumClearColors);
-		
-		{
-			// Draw a fullscreen quad
-			if(ExcludeRect.Width() > 0 && ExcludeRect.Height() > 0)
+			for (int32 i = 0; i < NumClearColors; i++)
 			{
-				// with a hole in it (optimization in case the hardware has non constant clear performance)
-				FVector4 OuterVertices[4];
-				OuterVertices[0].Set( -1.0f,  1.0f, Depth, 1.0f );
-				OuterVertices[1].Set(  1.0f,  1.0f, Depth, 1.0f );
-				OuterVertices[2].Set(  1.0f, -1.0f, Depth, 1.0f );
-				OuterVertices[3].Set( -1.0f, -1.0f, Depth, 1.0f );
-
-				float InvViewWidth = 1.0f / Viewport.Width;
-				float InvViewHeight = 1.0f / Viewport.Height;
-				FVector4 FractionRect = FVector4(ExcludeRect.Min.X * InvViewWidth, ExcludeRect.Min.Y * InvViewHeight, (ExcludeRect.Max.X - 1) * InvViewWidth, (ExcludeRect.Max.Y - 1) * InvViewHeight);
-
-				FVector4 InnerVertices[4];
-				InnerVertices[0].Set( FMath::Lerp(-1.0f,  1.0f, FractionRect.X), FMath::Lerp(1.0f, -1.0f, FractionRect.Y), Depth, 1.0f );
-				InnerVertices[1].Set( FMath::Lerp(-1.0f,  1.0f, FractionRect.Z), FMath::Lerp(1.0f, -1.0f, FractionRect.Y), Depth, 1.0f );
-				InnerVertices[2].Set( FMath::Lerp(-1.0f,  1.0f, FractionRect.Z), FMath::Lerp(1.0f, -1.0f, FractionRect.W), Depth, 1.0f );
-				InnerVertices[3].Set( FMath::Lerp(-1.0f,  1.0f, FractionRect.X), FMath::Lerp(1.0f, -1.0f, FractionRect.W), Depth, 1.0f );
-				
-				FVector4 Vertices[10];
-				Vertices[0] = OuterVertices[0];
-				Vertices[1] = InnerVertices[0];
-				Vertices[2] = OuterVertices[1];
-				Vertices[3] = InnerVertices[1];
-				Vertices[4] = OuterVertices[2];
-				Vertices[5] = InnerVertices[2];
-				Vertices[6] = OuterVertices[3];
-				Vertices[7] = InnerVertices[3];
-				Vertices[8] = OuterVertices[0];
-				Vertices[9] = InnerVertices[0];
-
-				DrawPrimitiveUP(RHICmdList, PT_TriangleStrip, 8, Vertices, sizeof(Vertices[0]));
+				ShaderClearColors[i] = ClearColorArray[i];
 			}
-			else
+
+			SetShaderValueArray(RHICmdList, PixelShader->GetPixelShader(), PixelShader->ColorParameter, ShaderClearColors, NumClearColors);
+
 			{
-				// without a hole
-				FVector4 Vertices[4];
-				Vertices[0].Set( -1.0f,  1.0f, Depth, 1.0f );
-				Vertices[1].Set(  1.0f,  1.0f, Depth, 1.0f );
-				Vertices[2].Set( -1.0f, -1.0f, Depth, 1.0f );
-				Vertices[3].Set(  1.0f, -1.0f, Depth, 1.0f );
-				DrawPrimitiveUP(RHICmdList, PT_TriangleStrip, 2, Vertices, sizeof(Vertices[0]));
+				// Draw a fullscreen quad
+				if (ExcludeRect.Width() > 0 && ExcludeRect.Height() > 0)
+				{
+					// with a hole in it (optimization in case the hardware has non constant clear performance)
+					FVector4 OuterVertices[4];
+					OuterVertices[0].Set(-1.0f, 1.0f, Depth, 1.0f);
+					OuterVertices[1].Set(1.0f, 1.0f, Depth, 1.0f);
+					OuterVertices[2].Set(1.0f, -1.0f, Depth, 1.0f);
+					OuterVertices[3].Set(-1.0f, -1.0f, Depth, 1.0f);
+
+					float InvViewWidth = 1.0f / Viewport.Width;
+					float InvViewHeight = 1.0f / Viewport.Height;
+					FVector4 FractionRect = FVector4(ExcludeRect.Min.X * InvViewWidth, ExcludeRect.Min.Y * InvViewHeight, (ExcludeRect.Max.X - 1) * InvViewWidth, (ExcludeRect.Max.Y - 1) * InvViewHeight);
+
+					FVector4 InnerVertices[4];
+					InnerVertices[0].Set(FMath::Lerp(-1.0f, 1.0f, FractionRect.X), FMath::Lerp(1.0f, -1.0f, FractionRect.Y), Depth, 1.0f);
+					InnerVertices[1].Set(FMath::Lerp(-1.0f, 1.0f, FractionRect.Z), FMath::Lerp(1.0f, -1.0f, FractionRect.Y), Depth, 1.0f);
+					InnerVertices[2].Set(FMath::Lerp(-1.0f, 1.0f, FractionRect.Z), FMath::Lerp(1.0f, -1.0f, FractionRect.W), Depth, 1.0f);
+					InnerVertices[3].Set(FMath::Lerp(-1.0f, 1.0f, FractionRect.X), FMath::Lerp(1.0f, -1.0f, FractionRect.W), Depth, 1.0f);
+
+					FVector4 Vertices[10];
+					Vertices[0] = OuterVertices[0];
+					Vertices[1] = InnerVertices[0];
+					Vertices[2] = OuterVertices[1];
+					Vertices[3] = InnerVertices[1];
+					Vertices[4] = OuterVertices[2];
+					Vertices[5] = InnerVertices[2];
+					Vertices[6] = OuterVertices[3];
+					Vertices[7] = InnerVertices[3];
+					Vertices[8] = OuterVertices[0];
+					Vertices[9] = InnerVertices[0];
+
+					DrawPrimitiveUP(RHICmdList, PT_TriangleStrip, 8, Vertices, sizeof(Vertices[0]));
+				}
+				else
+				{
+					// without a hole
+					FVector4 Vertices[4];
+					Vertices[0].Set(-1.0f, 1.0f, Depth, 1.0f);
+					Vertices[1].Set(1.0f, 1.0f, Depth, 1.0f);
+					Vertices[2].Set(-1.0f, -1.0f, Depth, 1.0f);
+					Vertices[3].Set(1.0f, -1.0f, Depth, 1.0f);
+					DrawPrimitiveUP(RHICmdList, PT_TriangleStrip, 2, Vertices, sizeof(Vertices[0]));
+				}
 			}
+			// Implicit flush. Always call flush when using a command list in RHI implementations before doing anything else. This is super hazardous.
 		}
-		RHICmdList.Flush(); // always call flush with GetRecursiveRHICommandList, recursive use of the RHI is hazardous
 
 		if (bChangedDepthStencilTarget)
 		{
diff --git a/Engine/Source/Runtime/Windows/D3D11RHI/Private/D3D11RHI.cpp b/Engine/Source/Runtime/Windows/D3D11RHI/Private/D3D11RHI.cpp
--- a/Engine/Source/Runtime/Windows/D3D11RHI/Private/D3D11RHI.cpp
+++ b/Engine/Source/Runtime/Windows/D3D11RHI/Private/D3D11RHI.cpp
@@ -182,9 +182,6 @@ FGlobalBoundShaderState LongGPUTaskBoundShaderState;
 
 void FD3D11DynamicRHI::IssueLongGPUTask()
 {
-	
-	FRHICommandListImmediate& RHICmdList = FRHICommandListExecutor::GetRecursiveRHICommandList();
-
 	if (GRHIFeatureLevel >= ERHIFeatureLevel::SM3)
 	{
 		int32 LargestViewportIndex = INDEX_NONE;
@@ -205,6 +202,8 @@ void FD3D11DynamicRHI::IssueLongGPUTask()
 		{
 			FD3D11Viewport* Viewport = Viewports[LargestViewportIndex];
 
+			FRHICommandList_RecursiveHazardous RHICmdList;
+
 			SetRenderTarget(RHICmdList, Viewport->GetBackBuffer(), FTextureRHIRef());
 			RHICmdList.SetBlendState(TStaticBlendState<CW_RGBA, BO_Add, BF_One, BF_One>::GetRHI(), FLinearColor::Black);
 			RHICmdList.SetDepthStencilState(TStaticDepthStencilState<false, CF_Always>::GetRHI(), 0);
@@ -222,9 +221,9 @@ void FD3D11DynamicRHI::IssueLongGPUTask()
 			Vertices[2].Set( -1.0f, -1.0f, 0, 1.0f );
 			Vertices[3].Set(  1.0f, -1.0f, 0, 1.0f );
 			DrawPrimitiveUP(RHICmdList, PT_TriangleStrip, 2, Vertices, sizeof(Vertices[0]));
+			// Implicit fluish. Always call flush when using a command list in RHI implementations before doing anything else. This is super hazardous.
 		}
 	}
-	RHICmdList.Flush(); // always call flush with GetRecursiveRHICommandList, recursive use of the RHI is hazardous
 }
 
 void FD3DGPUProfiler::BeginFrame(FD3D11DynamicRHI* InRHI)
diff --git a/Engine/Source/Runtime/Windows/D3D11RHI/Private/D3D11RHIPrivate.h b/Engine/Source/Runtime/Windows/D3D11RHI/Private/D3D11RHIPrivate.h
--- a/Engine/Source/Runtime/Windows/D3D11RHI/Private/D3D11RHIPrivate.h
+++ b/Engine/Source/Runtime/Windows/D3D11RHI/Private/D3D11RHIPrivate.h
@@ -564,6 +564,7 @@ protected:
 
 	template<typename TPixelShader>
 	void ResolveTextureUsingShader(
+		FRHICommandList_RecursiveHazardous& RHICmdList,
 		FD3D11Texture2D* SourceTexture,
 		FD3D11Texture2D* DestTexture,
 		ID3D11RenderTargetView* DestSurfaceRTV,
@@ -593,7 +594,7 @@ protected:
 
 	void ReadSurfaceDataNoMSAARaw(FTextureRHIParamRef TextureRHI,FIntRect Rect,TArray<uint8>& OutData, FReadSurfaceDataFlags InFlags);
 
-	void ReadSurfaceDataMSAARaw(FTextureRHIParamRef TextureRHI,FIntRect Rect,TArray<uint8>& OutData, FReadSurfaceDataFlags InFlags);
+	void ReadSurfaceDataMSAARaw(FRHICommandList_RecursiveHazardous& RHICmdList, FTextureRHIParamRef TextureRHI, FIntRect Rect, TArray<uint8>& OutData, FReadSurfaceDataFlags InFlags);
 
 	friend struct FD3DGPUProfiler;
 };
diff --git a/Engine/Source/Runtime/Windows/D3D11RHI/Private/D3D11RenderTarget.cpp b/Engine/Source/Runtime/Windows/D3D11RHI/Private/D3D11RenderTarget.cpp
--- a/Engine/Source/Runtime/Windows/D3D11RHI/Private/D3D11RenderTarget.cpp
+++ b/Engine/Source/Runtime/Windows/D3D11RHI/Private/D3D11RenderTarget.cpp
@@ -161,6 +161,7 @@ static FResolveRect GetDefaultRect(const FResolveRect& Rect,uint32 DefaultWidth,
 
 template<typename TPixelShader>
 void FD3D11DynamicRHI::ResolveTextureUsingShader(
+	FRHICommandList_RecursiveHazardous& RHICmdList,
 	FD3D11Texture2D* SourceTexture,
 	FD3D11Texture2D* DestTexture,
 	ID3D11RenderTargetView* DestTextureRTV,
@@ -172,9 +173,6 @@ void FD3D11DynamicRHI::ResolveTextureUsingShader(
 	typename TPixelShader::FParameter PixelShaderParameter
 	)
 {
-	
-	FRHICommandListImmediate& RHICmdList = FRHICommandListExecutor::GetRecursiveRHICommandList();
-
 	// Save the current viewport so that it can be restored
 	D3D11_VIEWPORT SavedViewport;
 	uint32 NumSavedViewports = 1;
@@ -183,7 +181,7 @@ void FD3D11DynamicRHI::ResolveTextureUsingShader(
 	// No alpha blending, no depth tests or writes, no stencil tests or writes, no backface culling.
 	RHICmdList.SetBlendState(TStaticBlendState<>::GetRHI(),FLinearColor::White);
 	RHICmdList.SetRasterizerState(TStaticRasterizerState<FM_Solid, CM_None>::GetRHI());
-	RHICmdList.Flush(); // always call flush with GetRecursiveRHICommandList, recursive use of the RHI is hazardous
+	RHICmdList.Flush(); // always call flush when using a command list in RHI implementations before doing anything else. This is super hazardous.
 
 	// Make sure the destination is not bound as a shader resource.
 	if (DestTexture)
@@ -211,7 +209,7 @@ void FD3D11DynamicRHI::ResolveTextureUsingShader(
 		}
 
 		RHICmdList.SetDepthStencilState(TStaticDepthStencilState<true,CF_Always>::GetRHI(),0);
-		RHICmdList.Flush(); // always call flush with GetRecursiveRHICommandList, recursive use of the RHI is hazardous
+		RHICmdList.Flush(); // always call flush when using a command list in RHI implementations before doing anything else. This is super hazardous.
 
 		// Write to the dest texture as a depth-stencil target.
 		ID3D11RenderTargetView* NullRTV = NULL;
@@ -229,7 +227,7 @@ void FD3D11DynamicRHI::ResolveTextureUsingShader(
 		}
 
 		RHICmdList.SetDepthStencilState(TStaticDepthStencilState<false,CF_Always>::GetRHI(),0);
-		RHICmdList.Flush(); // always call flush with GetRecursiveRHICommandList, recursive use of the RHI is hazardous
+		RHICmdList.Flush(); // always call flush when using a command list in RHI implementations before doing anything else. This is super hazardous.
 
 		// Write to the dest surface as a render target.
 		Direct3DDeviceContext->OMSetRenderTargets(1,&DestTextureRTV,NULL);
@@ -256,7 +254,7 @@ void FD3D11DynamicRHI::ResolveTextureUsingShader(
 	SetGlobalBoundShaderState(RHICmdList, ResolveBoundShaderState, GScreenVertexDeclaration.VertexDeclarationRHI, *ResolveVertexShader, *ResolvePixelShader);
 
 	ResolvePixelShader->SetParameters(RHICmdList, Direct3DDeviceContext,PixelShaderParameter);
-	RHICmdList.Flush(); // always call flush with GetRecursiveRHICommandList, recursive use of the RHI is hazardous
+	RHICmdList.Flush(); // always call flush when using a command list in RHI implementations before doing anything else. This is super hazardous.
 
 	// Set the source texture.
 	const uint32 TextureIndex = ResolvePixelShader->UnresolvedSurface.GetBaseIndex();
@@ -286,7 +284,7 @@ void FD3D11DynamicRHI::ResolveTextureUsingShader(
 	Vertices[3].UV.Y       = MaxV;
 
 	DrawPrimitiveUP(RHICmdList, PT_TriangleStrip, 2, Vertices, sizeof(Vertices[0]));
-	RHICmdList.Flush(); // always call flush with GetRecursiveRHICommandList, recursive use of the RHI is hazardous
+	RHICmdList.Flush(); // always call flush when using a command list in RHI implementations before doing anything else. This is super hazardous.
 
 	if (SourceTexture)
 	{
@@ -314,6 +312,9 @@ void FD3D11DynamicRHI::RHICopyToResolveTarget(FTextureRHIParamRef SourceTextureR
 		return;
 	}
 
+	FRHICommandList_RecursiveHazardous RHICmdList;
+
+
 	FD3D11Texture2D* SourceTexture2D = static_cast<FD3D11Texture2D*>(SourceTextureRHI->GetTexture2D());
 	FD3D11Texture2D* DestTexture2D = static_cast<FD3D11Texture2D*>(DestTextureRHI->GetTexture2D());
 
@@ -338,6 +339,7 @@ void FD3D11DynamicRHI::RHICopyToResolveTarget(FTextureRHIParamRef SourceTextureR
 				DestTexture2D->GetResource()->GetDesc(&ResolveTargetDesc);
 
 				ResolveTextureUsingShader<FResolveDepthPS>(
+					RHICmdList,
 					SourceTexture2D,
 					DestTexture2D,
 					DestTexture2D->GetRenderTargetView(0, -1),
@@ -357,6 +359,7 @@ void FD3D11DynamicRHI::RHICopyToResolveTarget(FTextureRHIParamRef SourceTextureR
 				DestTexture2D->GetResource()->GetDesc(&ResolveTargetDesc);
 
 				ResolveTextureUsingShader<FResolveDepthNonMSPS>(
+					RHICmdList,
 					SourceTexture2D,
 					DestTexture2D,
 					NULL,
@@ -994,7 +997,8 @@ void FD3D11DynamicRHI::RHIReadSurfaceData(FTextureRHIParamRef TextureRHI,FIntRec
 	}
 	else
 	{
-		ReadSurfaceDataMSAARaw(TextureRHI, InRect, OutDataRaw, InFlags);
+		FRHICommandList_RecursiveHazardous RHICmdList;
+		ReadSurfaceDataMSAARaw(RHICmdList, TextureRHI, InRect, OutDataRaw, InFlags);
 	}
 
 	const uint32 SizeX = InRect.Width() * TextureDesc.SampleDesc.Count;
@@ -1010,7 +1014,7 @@ void FD3D11DynamicRHI::RHIReadSurfaceData(FTextureRHIParamRef TextureRHI,FIntRec
 	ConvertRAWSurfaceDataToFColor(TextureDesc.Format, SizeX, SizeY, OutDataRaw.GetTypedData(), SrcPitch, OutData.GetTypedData(), InFlags);
 }
 
-void FD3D11DynamicRHI::ReadSurfaceDataMSAARaw(FTextureRHIParamRef TextureRHI,FIntRect InRect,TArray<uint8>& OutData, FReadSurfaceDataFlags InFlags)
+void FD3D11DynamicRHI::ReadSurfaceDataMSAARaw(FRHICommandList_RecursiveHazardous& RHICmdList, FTextureRHIParamRef TextureRHI,FIntRect InRect,TArray<uint8>& OutData, FReadSurfaceDataFlags InFlags)
 {
 	FD3D11TextureBase* Texture = GetD3D11TextureFromRHITexture(TextureRHI);
 
@@ -1096,6 +1100,7 @@ void FD3D11DynamicRHI::ReadSurfaceDataMSAARaw(FTextureRHIParamRef TextureRHI,FIn
 	{
 		// Resolve the sample to the non-MSAA render target.
 		ResolveTextureUsingShader<FResolveSingleSamplePS>(
+			RHICmdList,
 			(FD3D11Texture2D*)TextureRHI->GetTexture2D(),
 			NULL,
 			NonMSAARTV,
