UE4 - RHISupportsInstancing

diff --git a/Engine/Source/Runtime/RHI/Public/RHIDefinitions.h b/Engine/Source/Runtime/RHI/Public/RHIDefinitions.h
--- a/Engine/Source/Runtime/RHI/Public/RHIDefinitions.h
+++ b/Engine/Source/Runtime/RHI/Public/RHIDefinitions.h
@@ -605,6 +605,11 @@ inline bool RHISupportsTessellation(const EShaderPlatform Platform)
 	return false;
 }
 
+inline bool RHISupportsInstancing(const EShaderPlatform Platform)
+{
+	return IsFeatureLevelSupported(Platform, ERHIFeatureLevel::SM3) || (Platform == SP_METAL);
+}
+
 inline uint32 GetFeatureLevelMaxTextureSamplers(ERHIFeatureLevel::Type FeatureLevel)
 {
 	if (FeatureLevel == ERHIFeatureLevel::ES2)
