Mobile Preview - Fix crash when switching preview mode.
 - Global bound shader states now contain cached resources per-feature level, so that the same state being used in multiple feature levels can use different FShader objects

Fix a crash when switching featurelevel on maps with multiple worlds, where some of the worlds don't have scenes (shootergame default map for example)

diff --git a/Engine/Source/Runtime/Engine/Private/StaticBoundShaderState.cpp b/Engine/Source/Runtime/Engine/Private/StaticBoundShaderState.cpp
--- a/Engine/Source/Runtime/Engine/Private/StaticBoundShaderState.cpp
+++ b/Engine/Source/Runtime/Engine/Private/StaticBoundShaderState.cpp
@@ -92,26 +92,26 @@ void FGlobalBoundShaderStateResource::ReleaseRHI()
 	BoundShaderState.SafeRelease();
 }
 
-void SetGlobalBoundShaderState_Internal(FGlobalBoundShaderState GlobalBoundShaderState)
+void SetGlobalBoundShaderState_Internal(FGlobalBoundShaderState& GlobalBoundShaderState)
 {
 	// Check for unset uniform buffer parameters
 	// Technically you can set uniform buffer parameters after calling RHISetBoundShaderState, but this is the most global place to check for unset parameters
-	GlobalBoundShaderState->Args.VertexShader->VerifyBoundUniformBufferParameters();
-	GlobalBoundShaderState->Args.PixelShader->VerifyBoundUniformBufferParameters();
-	GlobalBoundShaderState->Args.GeometryShader->VerifyBoundUniformBufferParameters();
+	GlobalBoundShaderState.Get()->Args.VertexShader->VerifyBoundUniformBufferParameters();
+	GlobalBoundShaderState.Get()->Args.PixelShader->VerifyBoundUniformBufferParameters();
+	GlobalBoundShaderState.Get()->Args.GeometryShader->VerifyBoundUniformBufferParameters();
 
-	FGlobalBoundShaderState_Internal* BSS = GlobalBoundShaderState->BSS;
+	FGlobalBoundShaderState_Internal* BSS = GlobalBoundShaderState.Get()->BSS;
 	if (!BSS)
 	{
 		BSS = new FGlobalBoundShaderState_Internal(); // these are simply leaked and never freed
-		GlobalBoundShaderState->BSS = BSS;
+		GlobalBoundShaderState.Get()->BSS = BSS;
 	}
 
 	SetBoundShaderState_Internal(
 		BSS->GetInitializedRHI(
-			GlobalBoundShaderState->Args.VertexDeclarationRHI,
-			GETSAFERHISHADER_VERTEX(GlobalBoundShaderState->Args.VertexShader),
-			GETSAFERHISHADER_PIXEL(GlobalBoundShaderState->Args.PixelShader),
-			(FGeometryShaderRHIParamRef)GETSAFERHISHADER_GEOMETRY(GlobalBoundShaderState->Args.GeometryShader))
+			GlobalBoundShaderState.Get()->Args.VertexDeclarationRHI,
+			GETSAFERHISHADER_VERTEX(GlobalBoundShaderState.Get()->Args.VertexShader),
+			GETSAFERHISHADER_PIXEL(GlobalBoundShaderState.Get()->Args.PixelShader),
+			(FGeometryShaderRHIParamRef)GETSAFERHISHADER_GEOMETRY(GlobalBoundShaderState.Get()->Args.GeometryShader))
 		);
 }
diff --git a/Engine/Source/Runtime/Engine/Private/World.cpp b/Engine/Source/Runtime/Engine/Private/World.cpp
--- a/Engine/Source/Runtime/Engine/Private/World.cpp
+++ b/Engine/Source/Runtime/Engine/Private/World.cpp
@@ -5018,16 +5018,19 @@ void UWorld::ChangeFeatureLevel(ERHIFeatureLevel::Type InFeatureLevel)
 {
 	if (InFeatureLevel != FeatureLevel)
 	{
-		Scene->Release();
-		GetRendererModule().RemoveScene(Scene);
-
 		FeatureLevel = InFeatureLevel;
 
-		Scene = GetRendererModule().AllocateScene(this, bRequiresHitProxies, InFeatureLevel);
+		if (Scene)
+		{
+			Scene->Release();
+			GetRendererModule().RemoveScene(Scene);
+
+			Scene = GetRendererModule().AllocateScene(this, bRequiresHitProxies, InFeatureLevel);
 
- 		FFXSystemInterface::Destroy(FXSystem);
- 		FXSystem = FFXSystemInterface::Create(InFeatureLevel);
- 		Scene->SetFXSystem(FXSystem);
+			FFXSystemInterface::Destroy(FXSystem);
+			FXSystem = FFXSystemInterface::Create(InFeatureLevel);
+			Scene->SetFXSystem(FXSystem);
+		}
 	}
 }
 
diff --git a/Engine/Source/Runtime/Engine/Public/StaticBoundShaderState.h b/Engine/Source/Runtime/Engine/Public/StaticBoundShaderState.h
--- a/Engine/Source/Runtime/Engine/Public/StaticBoundShaderState.h
+++ b/Engine/Source/Runtime/Engine/Public/StaticBoundShaderState.h
@@ -57,7 +57,7 @@ typedef TGlobalResource<FGlobalBoundShaderStateResource> FGlobalBoundShaderState
 * @param BoundShaderState			- global bound shader state to initialize and set
 */
 
-extern ENGINE_API void SetGlobalBoundShaderState_Internal(FGlobalBoundShaderState BoundShaderState);
+extern ENGINE_API void SetGlobalBoundShaderState_Internal(FGlobalBoundShaderState& BoundShaderState);
 
 
 /**
diff --git a/Engine/Source/Runtime/RHI/Public/RHICommandList.h b/Engine/Source/Runtime/RHI/Public/RHICommandList.h
--- a/Engine/Source/Runtime/RHI/Public/RHICommandList.h
+++ b/Engine/Source/Runtime/RHI/Public/RHICommandList.h
@@ -792,12 +792,22 @@ struct FGlobalBoundShaderStateWorkArea
 	}
 };
 
-typedef FGlobalBoundShaderStateWorkArea* FGlobalBoundShaderState;
+struct FGlobalBoundShaderState
+{
+public:
+
+	FGlobalBoundShaderStateWorkArea* Get(ERHIFeatureLevel::Type InFeatureLevel = GRHIFeatureLevel)  { return WorkAreas[InFeatureLevel]; }
+	FGlobalBoundShaderStateWorkArea** GetPtr(ERHIFeatureLevel::Type InFeatureLevel = GRHIFeatureLevel)  { return &WorkAreas[InFeatureLevel]; }
+
+private:
+
+	FGlobalBoundShaderStateWorkArea* WorkAreas[ERHIFeatureLevel::Num];
+};
 
 struct FRHICommandSetGlobalBoundShaderState : public FRHICommand
 {
 	FGlobalBoundShaderState GlobalBoundShaderState;
-	FORCEINLINE_DEBUGGABLE void Set(FGlobalBoundShaderState InGlobalBoundShaderState)
+	FORCEINLINE_DEBUGGABLE void Set(FGlobalBoundShaderState& InGlobalBoundShaderState)
 	{
 		Type = ERCT_SetGlobalBoundShaderState;
 		GlobalBoundShaderState = InGlobalBoundShaderState;
@@ -1046,7 +1056,7 @@ public:
 	}
 
 	FORCEINLINE_DEBUGGABLE void BuildAndSetGlobalBoundShaderState(
-		void(*InSetGlobalBoundShaderState_InternalPtr)(FGlobalBoundShaderState BoundShaderState),
+		void(*InSetGlobalBoundShaderState_InternalPtr)(FGlobalBoundShaderState& BoundShaderState),
 		FGlobalBoundShaderState& GlobalBoundShaderState,
 		FVertexDeclarationRHIParamRef VertexDeclarationRHI,
 		FShader* VertexShader,
@@ -1055,7 +1065,7 @@ public:
 		)
 	{
 		SetGlobalBoundShaderState_InternalPtr = InSetGlobalBoundShaderState_InternalPtr; // total hack to work around module issues
-		if (!GlobalBoundShaderState)
+		if (!GlobalBoundShaderState.Get())
 		{
 			FGlobalBoundShaderStateWorkArea* NewGlobalBoundShaderState = new FGlobalBoundShaderStateWorkArea();
 			NewGlobalBoundShaderState->Args.VertexDeclarationRHI = VertexDeclarationRHI;
@@ -1064,13 +1074,13 @@ public:
 			NewGlobalBoundShaderState->Args.GeometryShader = GeometryShader;
 			FPlatformMisc::MemoryBarrier();
 
-			FGlobalBoundShaderStateWorkArea* OldGlobalBoundShaderState = (FGlobalBoundShaderStateWorkArea*)FPlatformAtomics::InterlockedCompareExchangePointer((void**)&GlobalBoundShaderState, NewGlobalBoundShaderState, nullptr);
+			FGlobalBoundShaderStateWorkArea* OldGlobalBoundShaderState = (FGlobalBoundShaderStateWorkArea*)FPlatformAtomics::InterlockedCompareExchangePointer((void**)GlobalBoundShaderState.GetPtr(), NewGlobalBoundShaderState, nullptr);
 
 			if (OldGlobalBoundShaderState != nullptr)
 			{
 				//we lost
 				delete NewGlobalBoundShaderState;
-				check(OldGlobalBoundShaderState == GlobalBoundShaderState);
+				check(OldGlobalBoundShaderState == GlobalBoundShaderState.Get());
 			}
 			else
 			{
@@ -1085,10 +1095,10 @@ public:
 			}
 		}
 		check(
-			VertexDeclarationRHI == GlobalBoundShaderState->Args.VertexDeclarationRHI &&
-			VertexShader == GlobalBoundShaderState->Args.VertexShader &&
-			PixelShader == GlobalBoundShaderState->Args.PixelShader &&
-			GeometryShader == GlobalBoundShaderState->Args.GeometryShader
+			VertexDeclarationRHI == GlobalBoundShaderState.Get()->Args.VertexDeclarationRHI &&
+			VertexShader == GlobalBoundShaderState.Get()->Args.VertexShader &&
+			PixelShader == GlobalBoundShaderState.Get()->Args.PixelShader &&
+			GeometryShader == GlobalBoundShaderState.Get()->Args.GeometryShader
 			);
 		if (Bypass())
 		{
@@ -1609,7 +1619,7 @@ private:
 	uint32 NumCommands;
 	uint32 UID;
 	static FThreadSafeCounter UIDCounter;
-	void(*SetGlobalBoundShaderState_InternalPtr)(FGlobalBoundShaderState BoundShaderState);
+	void(*SetGlobalBoundShaderState_InternalPtr)(FGlobalBoundShaderState& BoundShaderState);
 
 
 	FORCEINLINE_DEBUGGABLE uint8* Alloc(SIZE_T InSize)
