UE4 - defer deletes of RHI resources and some parallel prototyping

diff --git a/Engine/Plugins/Experimental/Paper2D/Source/Paper2D/Paper2D.Build.cs b/Engine/Plugins/Experimental/Paper2D/Source/Paper2D/Paper2D.Build.cs
--- a/Engine/Plugins/Experimental/Paper2D/Source/Paper2D/Paper2D.Build.cs
+++ b/Engine/Plugins/Experimental/Paper2D/Source/Paper2D/Paper2D.Build.cs
@@ -21,7 +21,8 @@ public class Paper2D : ModuleRules
 
 		PrivateDependencyModuleNames.AddRange(
 			new string[] {
-				"Slate"
+				"Slate",
+				"Renderer",
 			}
 			);
 
diff --git a/Engine/Source/Editor/NewLevelDialog/NewLevelDialog.Build.cs b/Engine/Source/Editor/NewLevelDialog/NewLevelDialog.Build.cs
--- a/Engine/Source/Editor/NewLevelDialog/NewLevelDialog.Build.cs
+++ b/Engine/Source/Editor/NewLevelDialog/NewLevelDialog.Build.cs
@@ -12,6 +12,7 @@ public class NewLevelDialog : ModuleRules
 				"CoreUObject",
                 "InputCore",
 				"RenderCore",
+				"RHI",
 				"Engine",
 				"Slate",
 				"SlateCore",
diff --git a/Engine/Source/Runtime/Core/Private/HAL/ThreadingBase.cpp b/Engine/Source/Runtime/Core/Private/HAL/ThreadingBase.cpp
--- a/Engine/Source/Runtime/Core/Private/HAL/ThreadingBase.cpp
+++ b/Engine/Source/Runtime/Core/Private/HAL/ThreadingBase.cpp
@@ -35,6 +35,11 @@ CORE_API bool IsInRenderingThread()
 	return !GRenderingThread || GIsRenderingThreadSuspended || (FPlatformTLS::GetCurrentThreadId() == GRenderingThread->GetThreadID());
 }
 
+CORE_API bool IsInParallelRenderingThread()
+{
+	return !GRenderingThread || GIsRenderingThreadSuspended || (FPlatformTLS::GetCurrentThreadId() != GGameThreadId);
+}
+
 
 // Fake threads
 
diff --git a/Engine/Source/Runtime/Core/Public/Async/TaskGraphInterfaces.h b/Engine/Source/Runtime/Core/Public/Async/TaskGraphInterfaces.h
--- a/Engine/Source/Runtime/Core/Public/Async/TaskGraphInterfaces.h
+++ b/Engine/Source/Runtime/Core/Public/Async/TaskGraphInterfaces.h
@@ -555,6 +555,27 @@ public:
 			new ((void *)&Owner->TaskStorage) TTask(Arg1,Arg2,Arg3,Arg4,Arg5);
 			return Owner->Setup(Prerequisites, CurrentThreadIfKnown);
 		}
+		/** Passthrough internal task constructor and dispatch. Note! Generally speaking references will not pass through; use pointers */
+		template<typename T1, typename T2, typename T3, typename T4, typename T5, typename T6>
+		FGraphEventRef ConstructAndDispatchWhenReady(const T1& Arg1, const T2& Arg2, const T3& Arg3, const T4& Arg4, const T5& Arg5, const T6& Arg6)
+		{
+			new ((void *)&Owner->TaskStorage) TTask(Arg1, Arg2, Arg3, Arg4, Arg5, Arg6);
+			return Owner->Setup(Prerequisites, CurrentThreadIfKnown);
+		}
+		/** Passthrough internal task constructor and dispatch. Note! Generally speaking references will not pass through; use pointers */
+		template<typename T1, typename T2, typename T3, typename T4, typename T5, typename T6, typename T7>
+		FGraphEventRef ConstructAndDispatchWhenReady(const T1& Arg1, const T2& Arg2, const T3& Arg3, const T4& Arg4, const T5& Arg5, const T6& Arg6, const T7& Arg7)
+		{
+			new ((void *)&Owner->TaskStorage) TTask(Arg1, Arg2, Arg3, Arg4, Arg5, Arg6, Arg7);
+			return Owner->Setup(Prerequisites, CurrentThreadIfKnown);
+		}
+		/** Passthrough internal task constructor and dispatch. Note! Generally speaking references will not pass through; use pointers */
+		template<typename T1, typename T2, typename T3, typename T4, typename T5, typename T6, typename T7, typename T8>
+		FGraphEventRef ConstructAndDispatchWhenReady(const T1& Arg1, const T2& Arg2, const T3& Arg3, const T4& Arg4, const T5& Arg5, const T6& Arg6, const T7& Arg7, const T8& Arg8)
+		{
+			new ((void *)&Owner->TaskStorage) TTask(Arg1, Arg2, Arg3, Arg4, Arg5, Arg6, Arg7, Arg8);
+			return Owner->Setup(Prerequisites, CurrentThreadIfKnown);
+		}
 	private:
 		friend class TGraphTask;
 
diff --git a/Engine/Source/Runtime/Core/Public/HAL/IConsoleManager.h b/Engine/Source/Runtime/Core/Public/HAL/IConsoleManager.h
--- a/Engine/Source/Runtime/Core/Public/HAL/IConsoleManager.h
+++ b/Engine/Source/Runtime/Core/Public/HAL/IConsoleManager.h
@@ -691,7 +691,7 @@ public:
 #if !defined(__clang__) // @todo Mac: figure out how to make this compile
 		// compiled out in shipping for performance (we can change in development later), if this get triggered you need to call GetValueOnGameThread() or GetValueOnAnyThread(), the last one is a bit slower
 	#if DO_CHECK
-		ensure(IsInRenderingThread());	// ensure to not block content creators, #if to optimize in shipping
+		ensure(IsInParallelRenderingThread());	// ensure to not block content creators, #if to optimize in shipping
 	#endif
 
 #endif
diff --git a/Engine/Source/Runtime/Core/Public/HAL/ThreadingBase.h b/Engine/Source/Runtime/Core/Public/HAL/ThreadingBase.h
--- a/Engine/Source/Runtime/Core/Public/HAL/ThreadingBase.h
+++ b/Engine/Source/Runtime/Core/Public/HAL/ThreadingBase.h
@@ -875,6 +875,9 @@ extern CORE_API bool IsInSlateThread();
 /** @return True if called from the rendering thread, or if called from ANY thread during single threaded rendering */
 extern CORE_API bool IsInRenderingThread();
 
+/** @return True if called from the rendering thread, or if called from ANY thread that isn't the game thread, except that during single threaded rendering the game thread is ok too.*/
+extern CORE_API bool IsInParallelRenderingThread();
+
 /** @return True if called from the rendering thread. */
 // Unlike IsInRenderingThread, this will always return false if we are running single threaded. It only returns true if this is actually a separate rendering thread. Mostly useful for checks
 extern CORE_API bool IsInActualRenderingThread();
diff --git a/Engine/Source/Runtime/EmptyRHI/Public/EmptyResources.h b/Engine/Source/Runtime/EmptyRHI/Public/EmptyResources.h
--- a/Engine/Source/Runtime/EmptyRHI/Public/EmptyResources.h
+++ b/Engine/Source/Runtime/EmptyRHI/Public/EmptyResources.h
@@ -42,15 +42,15 @@ public:
 	// IRefCountedObject interface.
 	virtual uint32 AddRef() const
 	{
-		return FRefCountedObject::AddRef();
+		return FRHIResource::AddRef();
 	}
 	virtual uint32 Release() const
 	{
-		return FRefCountedObject::Release();
+		return FRHIResource::Release();
 	}
 	virtual uint32 GetRefCount() const
 	{
-		return FRefCountedObject::GetRefCount();
+		return FRHIResource::GetRefCount();
 	}
 };
 
diff --git a/Engine/Source/Runtime/Engine/Private/Materials/MaterialShared.cpp b/Engine/Source/Runtime/Engine/Private/Materials/MaterialShared.cpp
--- a/Engine/Source/Runtime/Engine/Private/Materials/MaterialShared.cpp
+++ b/Engine/Source/Runtime/Engine/Private/Materials/MaterialShared.cpp
@@ -436,7 +436,7 @@ bool FMaterial::MaterialMayModifyMeshPosition() const
 
 FMaterialShaderMap* FMaterial::GetRenderingThreadShaderMap() const 
 { 
-	checkSlow(IsInRenderingThread());
+	checkSlow(IsInParallelRenderingThread());
 	return RenderingThreadShaderMap; 
 }
 
@@ -1439,7 +1439,7 @@ FMaterialRenderContext::FMaterialRenderContext(
 
 void FMaterialRenderProxy::EvaluateUniformExpressions(FUniformExpressionCache& OutUniformExpressionCache, const FMaterialRenderContext& Context) const
 {
-	check(IsInRenderingThread());
+	check(IsInParallelRenderingThread());
 
 	SCOPE_CYCLE_COUNTER(STAT_CacheUniformExpressions);
 	
diff --git a/Engine/Source/Runtime/Engine/Private/UnrealEngine.cpp b/Engine/Source/Runtime/Engine/Private/UnrealEngine.cpp
--- a/Engine/Source/Runtime/Engine/Private/UnrealEngine.cpp
+++ b/Engine/Source/Runtime/Engine/Private/UnrealEngine.cpp
@@ -8248,7 +8248,9 @@ bool UEngine::LoadMap( FWorldContext& WorldContext, FURL URL, class UPendingNetG
 	// may be unavailable to load the next one.
 	ENQUEUE_UNIQUE_RENDER_COMMAND(FlushCommand, 
 		{
+			FlushPendingDeleteRHIResources_RenderThread();
 			RHIFlushResources();
+			FlushPendingDeleteRHIResources_RenderThread();
 		}
 	);
 	FlushRenderingCommands();	  
diff --git a/Engine/Source/Runtime/IOS/MetalRHI/Private/MetalTexture.cpp b/Engine/Source/Runtime/IOS/MetalRHI/Private/MetalTexture.cpp
--- a/Engine/Source/Runtime/IOS/MetalRHI/Private/MetalTexture.cpp
+++ b/Engine/Source/Runtime/IOS/MetalRHI/Private/MetalTexture.cpp
@@ -31,15 +31,15 @@ public:
 	// IRefCountedObject interface.
 	virtual uint32 AddRef() const
 	{
-		return FRefCountedObject::AddRef();
+		return FRHIResource::AddRef();
 	}
 	virtual uint32 Release() const
 	{
-		return FRefCountedObject::Release();
+		return FRHIResource::Release();
 	}
 	virtual uint32 GetRefCount() const
 	{
-		return FRefCountedObject::GetRefCount();
+		return FRHIResource::GetRefCount();
 	}
 
 	void SetReferencedTexture(FRHITexture* InTexture)
diff --git a/Engine/Source/Runtime/IOS/MetalRHI/Public/MetalResources.h b/Engine/Source/Runtime/IOS/MetalRHI/Public/MetalResources.h
--- a/Engine/Source/Runtime/IOS/MetalRHI/Public/MetalResources.h
+++ b/Engine/Source/Runtime/IOS/MetalRHI/Public/MetalResources.h
@@ -50,15 +50,15 @@ public:
 	// IRefCountedObject interface.
 	virtual uint32 AddRef() const
 	{
-		return FRefCountedObject::AddRef();
+		return FRHIResource::AddRef();
 	}
 	virtual uint32 Release() const
 	{
-		return FRefCountedObject::Release();
+		return FRHIResource::Release();
 	}
 	virtual uint32 GetRefCount() const
 	{
-		return FRefCountedObject::GetRefCount();
+		return FRHIResource::GetRefCount();
 	}
 	
 	// this is the compiler shader
diff --git a/Engine/Source/Runtime/Launch/Private/LaunchEngineLoop.cpp b/Engine/Source/Runtime/Launch/Private/LaunchEngineLoop.cpp
--- a/Engine/Source/Runtime/Launch/Private/LaunchEngineLoop.cpp
+++ b/Engine/Source/Runtime/Launch/Private/LaunchEngineLoop.cpp
@@ -1839,6 +1839,7 @@ void FEngineLoop::InitTime()
 //called via FCoreDelegates::StarvedGameLoop
 void GameLoopIsStarved()
 {
+	FlushPendingDeleteRHIResources_GameThread();
 	FStats::AdvanceFrame( true, FStats::FOnAdvanceRenderingThreadStats::CreateStatic( &AdvanceRenderingThreadStatsGT ) );
 }
 
@@ -2076,6 +2077,7 @@ void FEngineLoop::Tick()
 			ResetDeferredUpdates,
 			{
 				FDeferredUpdateResource::ResetNeedsUpdate();
+				FlushPendingDeleteRHIResources_RenderThread();
 			});
 		
 		{
diff --git a/Engine/Source/Runtime/RHI/Private/RHI.cpp b/Engine/Source/Runtime/RHI/Private/RHI.cpp
--- a/Engine/Source/Runtime/RHI/Private/RHI.cpp
+++ b/Engine/Source/Runtime/RHI/Private/RHI.cpp
@@ -60,6 +60,46 @@ static FAutoConsoleCommandWithOutputDevice GDumpRHIMemoryCmd(
 	);
 #endif
 
+
+TLockFreePointerList<FRHIResource> FRHIResource::PendingDeletes;
+FRHIResource* FRHIResource::CurrentlyDeleting = nullptr;
+
+#if !DISABLE_RHI_DEFFERED_DELETE
+bool FRHIResource::Bypass()
+{
+	return GRHICommandList.Bypass();
+}
+#endif
+
+void FRHIResource::FlushPendingDeletes()
+{
+	check(IsInRenderingThread());
+	FRHICommandListExecutor::GetImmediateCommandList().Flush();
+	FRHICommandListExecutor::CheckNoOutstandingCmdLists();
+
+	while (1)
+	{
+		TArray<FRHIResource*> ToDelete;
+		PendingDeletes.PopAll(ToDelete);
+		if (!ToDelete.Num())
+		{
+			break;
+		}
+		for (int32 Index = 0; Index < ToDelete.Num(); Index++)
+		{
+			FRHIResource* Ref = ToDelete[Index];
+			if (Ref->GetRefCount() == 0) // caches can bring dead objects back to life
+			{
+				CurrentlyDeleting = Ref;
+				delete Ref;
+				CurrentlyDeleting = nullptr;
+			}
+		}
+	}
+}
+
+
+
 /**
  * RHI configuration settings.
  */
diff --git a/Engine/Source/Runtime/RHI/Private/RHICommandList.cpp b/Engine/Source/Runtime/RHI/Private/RHICommandList.cpp
--- a/Engine/Source/Runtime/RHI/Private/RHICommandList.cpp
+++ b/Engine/Source/Runtime/RHI/Private/RHICommandList.cpp
@@ -3,7 +3,6 @@
 
 #include "RHI.h"
 #include "RHICommandList.h"
-//#include "../../RenderCore/Public/RenderingThread.h"
 
 DEFINE_STAT(STAT_RHICmdListExecuteTime);
 DEFINE_STAT(STAT_RHICmdListEnqueueTime);
@@ -51,7 +50,6 @@ static FAutoConsoleVariableRef CVarRHICmdMem(
 	);
 
 RHI_API FRHICommandListExecutor GRHICommandList;
-RHI_API FThreadSafeCounter FRHICommandList::UIDCounter;
 RHI_API TLockFreeFixedSizeAllocator<FRHICommandList::FMemManager::FPage::PageSize> FRHICommandList::FMemManager::FPage::TheAllocator;
 //RHI_API FGlobalRHI GRHI;
 
@@ -87,7 +85,6 @@ void FRHICommandListExecutor::ExecuteList(FRHICommandList& CmdList)
 				{
 					SetRasterizerState_Internal(RHICmd->State);
 				}
-				RHICmd->State->Release();
 			}
 			break;
 		case ERCT_SetShaderParameter:
@@ -124,8 +121,6 @@ void FRHICommandListExecutor::ExecuteList(FRHICommandList& CmdList)
 					default: check(0); break;
 					}
 				}
-				RHICmd->Shader->Release(); 
-				RHICmd->UniformBuffer->Release();
 			}
 			break;
 		case ERCT_SetShaderSampler:
@@ -144,8 +139,6 @@ void FRHICommandListExecutor::ExecuteList(FRHICommandList& CmdList)
 					default: check(0); break;
 					}
 				}
-				RHICmd->Shader->Release(); 
-				RHICmd->Sampler->Release();
 			}
 			break;
 		case ERCT_SetShaderTexture:
@@ -164,8 +157,6 @@ void FRHICommandListExecutor::ExecuteList(FRHICommandList& CmdList)
 					default: check(0); break;
 					}
 				}
-				RHICmd->Shader->Release(); 
-				RHICmd->Texture->Release();
 			}
 			break;
 		case ERCT_SetShaderResourceViewParameter:
@@ -184,11 +175,6 @@ void FRHICommandListExecutor::ExecuteList(FRHICommandList& CmdList)
 					default: check(0); break;
 					}
 				}
-				RHICmd->Shader->Release();
-				if (RHICmd->SRV)
-				{
-					RHICmd->SRV->Release();
-				}
 			}
 			break;
 		case ERCT_SetUAVParameter:
@@ -202,11 +188,6 @@ void FRHICommandListExecutor::ExecuteList(FRHICommandList& CmdList)
 					default: check(0); break;
 					}
 				}
-				RHICmd->Shader->Release();
-				if (RHICmd->UAV)
-				{
-					RHICmd->UAV->Release();
-				}
 			}
 			break;
 		case ERCT_SetUAVParameter_IntialCount:
@@ -220,11 +201,6 @@ void FRHICommandListExecutor::ExecuteList(FRHICommandList& CmdList)
 					default: check(0); break;
 					}
 				}
-				RHICmd->Shader->Release();
-				if (RHICmd->UAV)
-				{
-					RHICmd->UAV->Release();
-				}
 			}
 			break;
 		case ERCT_DrawPrimitive:
@@ -243,7 +219,6 @@ void FRHICommandListExecutor::ExecuteList(FRHICommandList& CmdList)
 				{
 					DrawIndexedPrimitive_Internal(RHICmd->IndexBuffer, RHICmd->PrimitiveType, RHICmd->BaseVertexIndex, RHICmd->MinIndex, RHICmd->NumVertices, RHICmd->StartIndex, RHICmd->NumPrimitives, RHICmd->NumInstances);
 				}
-				RHICmd->IndexBuffer->Release();
 			}
 			break;
 		case ERCT_SetBoundShaderState:
@@ -253,7 +228,6 @@ void FRHICommandListExecutor::ExecuteList(FRHICommandList& CmdList)
 				{
 					SetBoundShaderState_Internal(RHICmd->BoundShaderState);
 				}
-				RHICmd->BoundShaderState->Release();
 			}
 			break;
 		case ERCT_SetBlendState:
@@ -263,7 +237,6 @@ void FRHICommandListExecutor::ExecuteList(FRHICommandList& CmdList)
 				{
 					SetBlendState_Internal(RHICmd->State, RHICmd->BlendFactor);
 				}
-				RHICmd->State->Release();
 			}
 			break;
 		case ERCT_SetStreamSource:
@@ -273,10 +246,6 @@ void FRHICommandListExecutor::ExecuteList(FRHICommandList& CmdList)
 				{
 					SetStreamSource_Internal(RHICmd->StreamIndex, RHICmd->VertexBuffer, RHICmd->Stride, RHICmd->Offset);
 				}
-				if (RHICmd->VertexBuffer)
-				{
-					RHICmd->VertexBuffer->Release();
-				}
 			}
 			break;
 		case ERCT_SetDepthStencilState:
@@ -286,7 +255,6 @@ void FRHICommandListExecutor::ExecuteList(FRHICommandList& CmdList)
 				{
 					SetDepthStencilState_Internal(RHICmd->State, RHICmd->StencilRef);
 				}
-				RHICmd->State->Release();
 			}
 			break;
 		case ERCT_SetViewport:
@@ -319,24 +287,6 @@ void FRHICommandListExecutor::ExecuteList(FRHICommandList& CmdList)
 						RHICmd->NewNumUAVs,
 						RHICmd->UAVs);
 				}
-				for (uint32 Index = 0; Index < RHICmd->NewNumSimultaneousRenderTargets; Index++)
-				{
-					if (RHICmd->NewRenderTargetsRHI[Index].Texture)
-					{
-						RHICmd->NewRenderTargetsRHI[Index].Texture->Release();
-					}
-				}
-				if (RHICmd->NewDepthStencilTargetRHI)
-				{
-					RHICmd->NewDepthStencilTargetRHI->Release();
-				}
-				for (uint32 Index = 0; Index < RHICmd->NewNumUAVs; Index++)
-				{
-					if (RHICmd->UAVs[Index])
-					{
-						RHICmd->UAVs[Index]->Release();
-					}
-				}
 			}
 			break;
 		case ERCT_EndDrawPrimitiveUP:
@@ -405,30 +355,6 @@ void FRHICommandListExecutor::ExecuteList(FRHICommandList& CmdList)
 					if (--RHICmd->LocalBoundShaderState.WorkArea->UseCount == 0)
 					{
 						RHICmd->LocalBoundShaderState.WorkArea->BSS = nullptr; // this also releases the ref, which wouldn't otherwise happen because RHICommands are unsafe and don't do constructors or destructors
-						if (RHICmd->LocalBoundShaderState.WorkArea->Args.VertexDeclarationRHI)
-						{
-							RHICmd->LocalBoundShaderState.WorkArea->Args.VertexDeclarationRHI->Release();
-						}
-						if (RHICmd->LocalBoundShaderState.WorkArea->Args.VertexShaderRHI)
-						{
-							RHICmd->LocalBoundShaderState.WorkArea->Args.VertexShaderRHI->Release();
-						}
-						if (RHICmd->LocalBoundShaderState.WorkArea->Args.HullShaderRHI)
-						{
-							RHICmd->LocalBoundShaderState.WorkArea->Args.HullShaderRHI->Release();
-						}
-						if (RHICmd->LocalBoundShaderState.WorkArea->Args.DomainShaderRHI)
-						{
-							RHICmd->LocalBoundShaderState.WorkArea->Args.DomainShaderRHI->Release();
-						}
-						if (RHICmd->LocalBoundShaderState.WorkArea->Args.PixelShaderRHI)
-						{
-							RHICmd->LocalBoundShaderState.WorkArea->Args.PixelShaderRHI->Release();
-						}
-						if (RHICmd->LocalBoundShaderState.WorkArea->Args.GeometryShaderRHI)
-						{
-							RHICmd->LocalBoundShaderState.WorkArea->Args.GeometryShaderRHI->Release();
-						}
 					}
 				}
 			}
@@ -450,7 +376,6 @@ void FRHICommandListExecutor::ExecuteList(FRHICommandList& CmdList)
 				{
 					SetComputeShader_Internal(RHICmd->ComputeShader);
 				}
-				RHICmd->ComputeShader->Release();
 			}
 			break;
 		case ERCT_DispatchComputeShader:
@@ -469,7 +394,6 @@ void FRHICommandListExecutor::ExecuteList(FRHICommandList& CmdList)
 				{
 					DispatchIndirectComputeShader_Internal(RHICmd->ArgumentBuffer, RHICmd->ArgumentOffset);
 				}
-				RHICmd->ArgumentBuffer->Release();
 			}
 			break;
 		case ERCT_AutomaticCacheFlushAfterComputeShader:
@@ -497,7 +421,6 @@ void FRHICommandListExecutor::ExecuteList(FRHICommandList& CmdList)
 				{
 					DrawPrimitiveIndirect_Internal(RHICmd->PrimitiveType, RHICmd->ArgumentBuffer, RHICmd->ArgumentOffset);
 				}
-				RHICmd->ArgumentBuffer->Release();
 			}
 			break;
 		case ERCT_DrawIndexedIndirect:
@@ -507,8 +430,6 @@ void FRHICommandListExecutor::ExecuteList(FRHICommandList& CmdList)
 				{
 					DrawIndexedIndirect_Internal(RHICmd->IndexBufferRHI, RHICmd->PrimitiveType, RHICmd->ArgumentsBufferRHI, RHICmd->DrawArgumentsIndex, RHICmd->NumInstances);
 				}
-				RHICmd->IndexBufferRHI->Release();
-				RHICmd->ArgumentsBufferRHI->Release();
 			}
 			break;
 		case ERCT_DrawIndexedPrimitiveIndirect:
@@ -518,8 +439,6 @@ void FRHICommandListExecutor::ExecuteList(FRHICommandList& CmdList)
 				{
 					DrawIndexedPrimitiveIndirect_Internal(RHICmd->PrimitiveType, RHICmd->IndexBuffer, RHICmd->ArgumentsBuffer, RHICmd->ArgumentOffset);
 				}
-				RHICmd->IndexBuffer->Release();
-				RHICmd->ArgumentsBuffer->Release();
 			}
 			break;
 		case ERCT_EnableDepthBoundsTest:
@@ -538,7 +457,6 @@ void FRHICommandListExecutor::ExecuteList(FRHICommandList& CmdList)
 				{
 					ClearUAV_Internal(RHICmd->UnorderedAccessViewRHI, RHICmd->Values);
 				}
-				RHICmd->UnorderedAccessViewRHI->Release();
 			}
 			break;
 		case ERCT_CopyToResolveTarget:
@@ -548,8 +466,6 @@ void FRHICommandListExecutor::ExecuteList(FRHICommandList& CmdList)
 				{
 					CopyToResolveTarget_Internal(RHICmd->SourceTexture, RHICmd->DestTexture, RHICmd->bKeepOriginalSurface, RHICmd->ResolveParams);
 				}
-				RHICmd->SourceTexture->Release();
-				RHICmd->DestTexture->Release();
 			}
 			break;
 		case ERCT_Clear:
@@ -621,7 +537,7 @@ void FRHICommandListExecutor::ExecuteList(FRHICommandList& CmdList)
 	}
 	INC_DWORD_STAT_BY(STAT_RHICmdListCount, CmdList.NumCommands);
 #if !(UE_BUILD_SHIPPING || UE_BUILD_TEST)
-	if (&CmdList == &GetImmediateCommandList())
+	if (&CmdList == &GetImmediateCommandList() && GRHICommandList.OutstandingCmdListCount.GetValue() == 1)
 	{
 		LatchBypass();
 	}
@@ -632,10 +548,21 @@ void FRHICommandListExecutor::ExecuteList(FRHICommandList& CmdList)
 void FRHICommandListExecutor::LatchBypass()
 {
 #if !(UE_BUILD_SHIPPING || UE_BUILD_TEST)
-	bLatchedBypass = (CVarRHICmdBypass.GetValueOnRenderThread() >= 1);
+	check(GRHICommandList.OutstandingCmdListCount.GetValue() == 1 && !GRHICommandList.GetImmediateCommandList().HasCommands());
+	bool NewBypass = (CVarRHICmdBypass.GetValueOnRenderThread() >= 1);
+	if (NewBypass)
+	{
+		FRHIResource::FlushPendingDeletes();
+	}
+	bLatchedBypass = NewBypass;
 #endif
 }
 
+void FRHICommandListExecutor::CheckNoOutstandingCmdLists()
+{
+	check(GRHICommandList.OutstandingCmdListCount.GetValue() == 1); // else we are attempting to delete resources while there is still a live cmdlist (other than the immediate cmd list) somewhere.
+}
+
 FRHICommandList::FMemManager::FMemManager() :
 	FirstPage(nullptr),
 	LastPage(nullptr)
@@ -683,11 +610,33 @@ void FRHICommandList::FMemManager::Reset()
 	LastPage = FirstPage;
 }
 
+FRHICommandList::FRHICommandList()
+{
+	GRHICommandList.OutstandingCmdListCount.Increment();
+	Reset();
+}
+
+FRHICommandList::~FRHICommandList()
+{
+	Flush();
+	GRHICommandList.OutstandingCmdListCount.Decrement();
+}
+
 const SIZE_T FRHICommandList::GetUsedMemory() const
 {
 	return MemManager.GetUsedMemory();
 }
 
-
+void FRHICommandList::Reset()
+{
+	bExecuting = false;
+	NumCommands = 0;
+	MemManager.Reset();
+	UID = GRHICommandList.UIDCounter.Increment();
+	for (int32 Index = 0; Index < ERCT_NUM; Index++)
+	{
+		StateReducer[Index] = nullptr;
+	}
+}
 
 
diff --git a/Engine/Source/Runtime/RHI/Public/RHICommandList.h b/Engine/Source/Runtime/RHI/Public/RHICommandList.h
--- a/Engine/Source/Runtime/RHI/Public/RHICommandList.h
+++ b/Engine/Source/Runtime/RHI/Public/RHICommandList.h
@@ -54,7 +54,7 @@ enum ERHICommandType
 
 
 
-	ERCT_,
+	ERCT_NUM,
 };
 
 struct FRHICommand
@@ -874,6 +874,10 @@ private:
 				Current += AlignedSize;
 				return Ptr;
 			}
+			FORCEINLINE_DEBUGGABLE void Dealloc(SIZE_T AlignedSize)
+			{
+				Current -= AlignedSize;
+			}
 
 			uint8* Current;
 			uint8* Tail;
@@ -885,6 +889,11 @@ private:
 		FMemManager();
 		~FMemManager();
 
+		FORCEINLINE_DEBUGGABLE void Dealloc(SIZE_T InSize)
+		{
+			checkSlow(LastPage);
+			LastPage->Dealloc(InSize);
+		}
 		FORCEINLINE_DEBUGGABLE uint8* Alloc(SIZE_T InSize)
 		{
 			InSize = (InSize + (Alignment - 1)) & ~(Alignment - 1);
@@ -962,14 +971,9 @@ public:
 		Alignment = sizeof(SIZE_T),
 	};
 
-	FRHICommandList()
-	{
-		Reset();
-	}
-	~FRHICommandList()
-	{
-		Flush();
-	}
+	FRHICommandList();
+	~FRHICommandList();
+
 	inline void Flush();
 
 	FORCEINLINE_DEBUGGABLE FLocalBoundShaderState BuildLocalBoundShaderState(const FBoundShaderStateInput& BoundShaderStateInput)
@@ -1015,31 +1019,6 @@ public:
 			auto* Cmd = AddCommand<FRHICommandBuildLocalBoundShaderState>();
 			Cmd->Set(this, this->UID, VertexDeclarationRHI, VertexShaderRHI, HullShaderRHI, DomainShaderRHI, PixelShaderRHI, GeometryShaderRHI);
 			Result.WorkArea = &Cmd->WorkArea;
-
-			if (VertexDeclarationRHI)
-			{
-				VertexDeclarationRHI->AddRef();
-			}
-			if (VertexShaderRHI)
-			{
-				VertexShaderRHI->AddRef();
-			}
-			if (HullShaderRHI)
-			{
-				HullShaderRHI->AddRef();
-			}
-			if (DomainShaderRHI)
-			{
-				DomainShaderRHI->AddRef();
-			}
-			if (PixelShaderRHI)
-			{
-				PixelShaderRHI->AddRef();
-			}
-			if (GeometryShaderRHI)
-			{
-				GeometryShaderRHI->AddRef();
-			}
 		}
 		return Result;
 	}
@@ -1082,17 +1061,6 @@ public:
 				delete NewGlobalBoundShaderState;
 				check(OldGlobalBoundShaderState == GlobalBoundShaderState.Get());
 			}
-			else
-			{
-				// we won
-
-				// this is not even slightly threadsafe, a "loser" might destroy the resource before we get here, and besides this is global; we are never going to destory these anyway.
-
-				if (VertexDeclarationRHI)
-				{
-					VertexDeclarationRHI->AddRef();
-				}
-			}
 		}
 		check(
 			VertexDeclarationRHI == GlobalBoundShaderState.Get()->Args.VertexDeclarationRHI &&
@@ -1119,8 +1087,7 @@ public:
 		}
 		auto* Cmd = AddCommand<FRHICommandSetShaderUniformBuffer>();
 		Cmd->Set<TShaderRHIParamRef>(Shader, BaseIndex, UniformBuffer);
-		Shader->AddRef();
-		UniformBuffer->AddRef();
+		StateReduce<FRHICommandSetShaderUniformBuffer>(Cmd);
 	}
 
 	template <typename TShaderRHIParamRef>
@@ -1140,7 +1107,7 @@ public:
 		}
 		auto* Cmd = AddCommand<FRHICommandSetShaderParameter>();
 		Cmd->Set<TShaderRHIParamRef>(Shader, BufferIndex, BaseIndex, NumBytes, NewValue);
-		Shader->AddRef();
+		StateReduce<FRHICommandSetShaderParameter>(Cmd);
 	}
 
 	template <typename TShaderRHIParamRef>
@@ -1154,8 +1121,7 @@ public:
 
 		auto* Cmd = AddCommand<FRHICommandSetShaderTexture>();
 		Cmd->Set<TShaderRHIParamRef>(Shader, TextureIndex, Texture);
-		Shader->AddRef();
-		Texture->AddRef();
+		StateReduce<FRHICommandSetShaderTexture>(Cmd);
 	}
 
 	template <typename TShaderRHIParamRef>
@@ -1168,11 +1134,7 @@ public:
 		}
 		auto* Cmd = AddCommand<FRHICommandSetShaderResourceViewParameter>();
 		Cmd->Set<TShaderRHIParamRef>(Shader, SamplerIndex, SRV);
-		Shader->AddRef();
-		if (SRV)
-		{
-			SRV->AddRef();
-		}
+		StateReduce<FRHICommandSetShaderResourceViewParameter>(Cmd);
 	}
 
 	template <typename TShaderRHIParamRef>
@@ -1186,8 +1148,7 @@ public:
 
 		auto* Cmd = AddCommand<FRHICommandSetShaderSampler>();
 		Cmd->Set<TShaderRHIParamRef>(Shader, SamplerIndex, State);
-		Shader->AddRef();
-		State->AddRef();
+		StateReduce<FRHICommandSetShaderSampler>(Cmd);
 	}
 
 	FORCEINLINE_DEBUGGABLE void SetUAVParameter(FComputeShaderRHIParamRef Shader, uint32 UAVIndex, FUnorderedAccessViewRHIParamRef UAV)
@@ -1199,11 +1160,7 @@ public:
 		}
 		auto* Cmd = AddCommand<FRHICommandSetUAVParameter>();
 		Cmd->Set<FComputeShaderRHIParamRef>(Shader, UAVIndex, UAV);
-		Shader->AddRef();
-		if (UAV)
-		{
-			UAV->AddRef();
-		}
+		// StateReduce<FRHICommandSetUAVParameter>(Cmd); conflicts with next one
 	}
 
 	FORCEINLINE_DEBUGGABLE void SetUAVParameter(FComputeShaderRHIParamRef Shader, uint32 UAVIndex, FUnorderedAccessViewRHIParamRef UAV, uint32 InitialCount)
@@ -1215,11 +1172,7 @@ public:
 		}
 		auto* Cmd = AddCommand<FRHICommandSetUAVParameter_IntialCount>();
 		Cmd->Set<FComputeShaderRHIParamRef>(Shader, UAVIndex, UAV, InitialCount);
-		Shader->AddRef();
-		if (UAV)
-		{
-			UAV->AddRef();
-		}
+		// StateReduce<FRHICommandSetUAVParameter_IntialCount>(Cmd); conflicts with previous one
 	}
 
 	FORCEINLINE_DEBUGGABLE void SetBoundShaderState(FBoundShaderStateRHIParamRef BoundShaderState)
@@ -1232,7 +1185,7 @@ public:
 
 		auto* Cmd = AddCommand<FRHICommandSetBoundShaderState>();
 		Cmd->Set(BoundShaderState);
-		BoundShaderState->AddRef();
+		StateReduce<FRHICommandSetBoundShaderState>(Cmd);
 	}
 
 	FORCEINLINE_DEBUGGABLE void SetRasterizerState(FRasterizerStateRHIParamRef State)
@@ -1245,7 +1198,7 @@ public:
 
 		auto* Cmd = AddCommand<FRHICommandSetRasterizerState>();
 		Cmd->Set(State);
-		State->AddRef();
+		StateReduce<FRHICommandSetRasterizerState>(Cmd);
 	}
 
 	FORCEINLINE_DEBUGGABLE void SetBlendState(FBlendStateRHIParamRef State, const FLinearColor& BlendFactor = FLinearColor::White)
@@ -1258,7 +1211,7 @@ public:
 
 		auto* Cmd = AddCommand<FRHICommandSetBlendState>();
 		Cmd->Set(State, BlendFactor);
-		State->AddRef();
+		StateReduce<FRHICommandSetBlendState>(Cmd);
 	}
 
 	FORCEINLINE_DEBUGGABLE void DrawPrimitive(uint32 PrimitiveType, uint32 BaseVertexIndex, uint32 NumPrimitives, uint32 NumInstances)
@@ -1283,7 +1236,6 @@ public:
 
 		auto* Cmd = AddCommand<FRHICommandDrawIndexedPrimitive>();
 		Cmd->Set(IndexBuffer, PrimitiveType, BaseVertexIndex, MinIndex, NumVertices, StartIndex, NumPrimitives, NumInstances);
-		IndexBuffer->AddRef();
 	}
 
 	FORCEINLINE_DEBUGGABLE void SetStreamSource(uint32 StreamIndex, FVertexBufferRHIParamRef VertexBuffer, uint32 Stride, uint32 Offset)
@@ -1296,10 +1248,7 @@ public:
 
 		auto* Cmd = AddCommand<FRHICommandSetStreamSource>();
 		Cmd->Set(StreamIndex, VertexBuffer, Stride, Offset);
-		if (VertexBuffer)
-		{
-			VertexBuffer->AddRef();
-		}
+		StateReduce<FRHICommandSetStreamSource>(Cmd);
 	}
 
 	FORCEINLINE_DEBUGGABLE void SetDepthStencilState(FDepthStencilStateRHIParamRef NewStateRHI, uint32 StencilRef = 0)
@@ -1312,7 +1261,7 @@ public:
 
 		auto* Cmd = AddCommand<FRHICommandSetDepthStencilState>();
 		Cmd->Set(NewStateRHI, StencilRef);
-		NewStateRHI->AddRef();
+		StateReduce<FRHICommandSetDepthStencilState>(Cmd);
 	}
 
 	FORCEINLINE_DEBUGGABLE void SetViewport(uint32 MinX, uint32 MinY, float MinZ, uint32 MaxX, uint32 MaxY, float MaxZ)
@@ -1324,6 +1273,7 @@ public:
 		}
 		auto* Cmd = AddCommand<FRHICommandSetViewport>();
 		Cmd->Set(MinX, MinY, MinZ, MaxX, MaxY, MaxZ);
+		StateReduce<FRHICommandSetViewport>(Cmd);
 	}
 
 	FORCEINLINE_DEBUGGABLE void SetScissorRect(bool bEnable, uint32 MinX, uint32 MinY, uint32 MaxX, uint32 MaxY)
@@ -1362,25 +1312,7 @@ public:
 			NewDepthStencilTargetRHI,
 			NewNumUAVs,
 			UAVs);
-		for (uint32 Index = 0; Index < NewNumSimultaneousRenderTargets; Index++)
-		{
-			if (NewRenderTargetsRHI[Index].Texture)
-			{
-				NewRenderTargetsRHI[Index].Texture->AddRef();
-			}
-		}
-		if (NewDepthStencilTargetRHI)
-		{
-			NewDepthStencilTargetRHI->AddRef();
-		}
-		for (uint32 Index = 0; Index < NewNumUAVs; Index++)
-		{
-			if (UAVs[Index])
-			{
-				UAVs[Index]->AddRef();
-			}
-		}
-
+		StateReduce<FRHICommandSetRenderTargets>(Cmd);
 	}
 
 	FORCEINLINE_DEBUGGABLE void BeginDrawPrimitiveUP(uint32 PrimitiveType, uint32 NumPrimitives, uint32 NumVertices, uint32 VertexDataStride, void*& OutVertexData)
@@ -1467,7 +1399,7 @@ public:
 		}
 		auto* Cmd = AddCommand<FRHICommandSetComputeShader>();
 		Cmd->Set(ComputeShader);
-		ComputeShader->AddRef();
+		StateReduce<FRHICommandSetComputeShader>(Cmd);
 	}
 
 	FORCEINLINE_DEBUGGABLE void DispatchComputeShader(uint32 ThreadGroupCountX, uint32 ThreadGroupCountY, uint32 ThreadGroupCountZ)
@@ -1490,7 +1422,6 @@ public:
 		}
 		auto* Cmd = AddCommand<FRHICommandDispatchIndirectComputeShader>();
 		Cmd->Set(ArgumentBuffer, ArgumentOffset);
-		ArgumentBuffer->AddRef();
 	}
 
 	FORCEINLINE_DEBUGGABLE void AutomaticCacheFlushAfterComputeShader(bool bEnable)
@@ -1524,7 +1455,6 @@ public:
 		}
 		auto* Cmd = AddCommand<FRHICommandDrawPrimitiveIndirect>();
 		Cmd->Set(PrimitiveType, ArgumentBuffer, ArgumentOffset);
-		ArgumentBuffer->AddRef();
 	}
 
 	FORCEINLINE_DEBUGGABLE void DrawIndexedIndirect(FIndexBufferRHIParamRef IndexBufferRHI, uint32 PrimitiveType, FStructuredBufferRHIParamRef ArgumentsBufferRHI, uint32 DrawArgumentsIndex, uint32 NumInstances)
@@ -1536,8 +1466,6 @@ public:
 		}
 		auto* Cmd = AddCommand<FRHICommandDrawIndexedIndirect>();
 		Cmd->Set(IndexBufferRHI, PrimitiveType, ArgumentsBufferRHI, DrawArgumentsIndex, NumInstances);
-		IndexBufferRHI->AddRef();
-		ArgumentsBufferRHI->AddRef();
 	}
 
 	FORCEINLINE_DEBUGGABLE void DrawIndexedPrimitiveIndirect(uint32 PrimitiveType, FIndexBufferRHIParamRef IndexBuffer, FVertexBufferRHIParamRef ArgumentsBuffer, uint32 ArgumentOffset)
@@ -1549,8 +1477,6 @@ public:
 		}
 		auto* Cmd = AddCommand<FRHICommandDrawIndexedPrimitiveIndirect>();
 		Cmd->Set(PrimitiveType, IndexBuffer, ArgumentsBuffer, ArgumentOffset);
-		IndexBuffer->AddRef();
-		ArgumentsBuffer->AddRef();
 	}
 
 	FORCEINLINE_DEBUGGABLE void EnableDepthBoundsTest(bool bEnable, float MinDepth, float MaxDepth)
@@ -1562,7 +1488,7 @@ public:
 		}
 		auto* Cmd = AddCommand<FRHICommandEnableDepthBoundsTest>();
 		Cmd->Set(bEnable, MinDepth, MaxDepth);
-
+		StateReduce<FRHICommandEnableDepthBoundsTest>(Cmd);
 	}
 
 	FORCEINLINE_DEBUGGABLE void ClearUAV(FUnorderedAccessViewRHIParamRef UnorderedAccessViewRHI, const uint32(&Values)[4])
@@ -1574,7 +1500,6 @@ public:
 		}
 		auto* Cmd = AddCommand<FRHICommandClearUAV>();
 		Cmd->Set(UnorderedAccessViewRHI, Values);
-		UnorderedAccessViewRHI->AddRef();
 	}
 
 	FORCEINLINE_DEBUGGABLE void CopyToResolveTarget(FTextureRHIParamRef SourceTextureRHI, FTextureRHIParamRef DestTextureRHI, bool bKeepOriginalSurface, const FResolveParams& ResolveParams)
@@ -1586,8 +1511,6 @@ public:
 		}
 		auto* Cmd = AddCommand<FRHICommandCopyToResolveTarget>();
 		Cmd->Set(SourceTextureRHI, DestTextureRHI, bKeepOriginalSurface, ResolveParams);
-		SourceTextureRHI->AddRef();
-		DestTextureRHI->AddRef();
 	}
 
 	FORCEINLINE_DEBUGGABLE void Clear(bool bClearColor, const FLinearColor& Color, bool bClearDepth, float Depth, bool bClearStencil, uint32 Stencil, FIntRect ExcludeRect)
@@ -1618,7 +1541,7 @@ private:
 	bool bExecuting;
 	uint32 NumCommands;
 	uint32 UID;
-	static FThreadSafeCounter UIDCounter;
+	FRHICommand* StateReducer[ERCT_NUM];
 	void(*SetGlobalBoundShaderState_InternalPtr)(FGlobalBoundShaderState& BoundShaderState);
 
 
@@ -1642,14 +1565,25 @@ private:
 		return Cmd;
 	}
 
-	void Reset()
+	template <typename TCmd>
+	FORCEINLINE_DEBUGGABLE void StateReduce(TCmd* Cmd)
 	{
-		bExecuting = false;
-		NumCommands = 0;
-		MemManager.Reset();
-		UID = UIDCounter.Increment();
+#if 0
+		if (StateReducer[Cmd->Type])
+		{
+			if (FMemory::Memcmp(StateReducer[Cmd->Type], Cmd, sizeof(TCmd)) == 0)
+			{
+				--NumCommands;
+				MemManager.Dealloc(sizeof(TCmd));
+				return;
+			}
+		}
+		StateReducer[Cmd->Type] = Cmd;
+#endif
 	}
 
+	void Reset();
+
 	inline bool Bypass();
 
 	friend class FRHICommandListExecutor;
@@ -1759,9 +1693,13 @@ public:
 		return !!DefaultBypass;
 #endif
 	}
+	static void CheckNoOutstandingCmdLists();
 
 private:
 	bool bLatchedBypass;
+	friend class FRHICommandList;
+	FThreadSafeCounter UIDCounter;
+	FThreadSafeCounter OutstandingCmdListCount;
 	FRHICommandListImmediate CommandListImmediate;
 
 	template <bool bOnlyTestMemAccess>
@@ -1779,6 +1717,7 @@ FORCEINLINE_DEBUGGABLE FRHICommandListImmediate& FRHICommandListExecutor::GetImm
 	return GRHICommandList.CommandListImmediate;
 }
 
+
 #define DEFINE_RHIMETHOD_CMDLIST(Type,Name,ParameterTypesAndNames,ParameterNames,ReturnStatement,NullImplementation)
 #define DEFINE_RHIMETHOD(Type, Name, ParameterTypesAndNames, ParameterNames, ReturnStatement, NullImplementation)
 #define DEFINE_RHIMETHOD_GLOBAL(Type, Name, ParameterTypesAndNames, ParameterNames, ReturnStatement, NullImplementation)
diff --git a/Engine/Source/Runtime/RHI/Public/RHIResources.h b/Engine/Source/Runtime/RHI/Public/RHIResources.h
--- a/Engine/Source/Runtime/RHI/Public/RHIResources.h
+++ b/Engine/Source/Runtime/RHI/Public/RHIResources.h
@@ -6,14 +6,82 @@
 #include "RHIDefinitions.h"
 #include "RefCounting.h"
 #include "Runtime/Engine/Public/PixelFormat.h" // for EPixelFormat
+#include "LockFreeList.h"
+
+#define DISABLE_RHI_DEFFERED_DELETE 0
 
 /** The base type of RHI resources. */
-class RHI_API FRHIResource : public FRefCountedObject
+class RHI_API FRHIResource
 {
 public:
-	virtual ~FRHIResource() {}
+	FRHIResource(bool InbDoNotDeferDelete = false)
+		: bDoNotDeferDelete(InbDoNotDeferDelete)
+	{
+	}
+	virtual ~FRHIResource() 
+	{ 
+		check(NumRefs.GetValue() == 0 && (CurrentlyDeleting == this || bDoNotDeferDelete || Bypass())); // this should not have any outstanding refs
+	}
+	uint32 AddRef() const
+	{
+		int32 NewValue = NumRefs.Increment();
+		check(NewValue > 0); 
+		return uint32(NewValue);
+	}
+	uint32 Release() const
+	{
+		int32 NewValue = NumRefs.Decrement();
+		if (NewValue == 0)
+		{
+			if (bDoNotDeferDelete || Bypass())
+			{ 
+				delete this;
+			}
+			else
+			{
+				if (MarkedForDelete.Increment() == 1)
+				{
+					PendingDeletes.Push(const_cast<FRHIResource*>(this));
+				}
+			}
+		}
+		check(NewValue >= 0);
+		return uint32(NewValue);
+	}
+	uint32 GetRefCount() const
+	{
+		int32 CurrentValue = NumRefs.GetValue();
+		check(CurrentValue >= 0); 
+		return uint32(CurrentValue);
+	}
+	void DoNoDeferDelete()
+	{
+		check(!MarkedForDelete.GetValue());
+		bDoNotDeferDelete = true;
+		FPlatformMisc::MemoryBarrier();
+		check(!MarkedForDelete.GetValue());
+	}
+
+	static void FlushPendingDeletes();
+
+#if DISABLE_RHI_DEFFERED_DELETE
+	FORCEINLINE static bool Bypass()
+	{
+		return true;
+	}
+#else
+	static bool Bypass();
+#endif
+
+private:
+	mutable FThreadSafeCounter NumRefs;
+	mutable FThreadSafeCounter MarkedForDelete;
+	bool bDoNotDeferDelete;
+	static TLockFreePointerList<FRHIResource> PendingDeletes;
+	static FRHIResource* CurrentlyDeleting;
 };
 
+
 //
 // State blocks
 //
@@ -444,6 +512,10 @@ class FRHIRenderQuery : public FRHIResource {};
 class FRHIViewport : public FRHIResource 
 {
 public:
+	FRHIViewport()
+		: FRHIResource(true)
+	{
+	}
 	/**
 	 * Returns access to the platform-specific native resource pointer.  This is designed to be used to provide plugins with access
 	 * to the underlying resource and should be used very carefully or not at all.
@@ -531,7 +603,11 @@ public:
 class FRHICustomPresent : public FRHIResource
 {
 public:
-	explicit FRHICustomPresent(FRHIViewport* InViewport) : ViewportRHI(InViewport) {}
+	explicit FRHICustomPresent(FRHIViewport* InViewport) 
+		: FRHIResource(true)
+		, ViewportRHI(InViewport) 
+	{
+	}
 	
 	virtual ~FRHICustomPresent() {} // should release any references to D3D resources.
 	
diff --git a/Engine/Source/Runtime/RenderCore/Private/RenderingThread.cpp b/Engine/Source/Runtime/RenderCore/Private/RenderingThread.cpp
--- a/Engine/Source/Runtime/RenderCore/Private/RenderingThread.cpp
+++ b/Engine/Source/Runtime/RenderCore/Private/RenderingThread.cpp
@@ -748,6 +748,8 @@ void AdvanceFrameRenderPrerequisite()
  */
 void FlushRenderingCommands()
 {
+	FlushPendingDeleteRHIResources_GameThread(); // we do this here because with modal dialogs or blueprint debugging, the engine does not tick
+
 	AdvanceFrameRenderPrerequisite();
 
 	// Find the objects which may be cleaned up once the rendering thread command queue has been flushed.
@@ -762,6 +764,21 @@ void FlushRenderingCommands()
 	delete PendingCleanupObjects;
 }
 
+void FlushPendingDeleteRHIResources_GameThread()
+{
+	ENQUEUE_UNIQUE_RENDER_COMMAND(
+		FlushPendingDeleteRHIResources,
+		{
+			FlushPendingDeleteRHIResources_RenderThread();
+		}
+	);
+}
+void FlushPendingDeleteRHIResources_RenderThread()
+{
+	FRHIResource::FlushPendingDeletes();
+}
+
+
 FRHICommandListImmediate& GetImmediateCommandList_ForRenderCommand()
 {
 	return FRHICommandListExecutor::GetImmediateCommandList();
diff --git a/Engine/Source/Runtime/RenderCore/Public/RenderingThread.h b/Engine/Source/Runtime/RenderCore/Public/RenderingThread.h
--- a/Engine/Source/Runtime/RenderCore/Public/RenderingThread.h
+++ b/Engine/Source/Runtime/RenderCore/Public/RenderingThread.h
@@ -75,6 +75,9 @@ extern RENDERCORE_API void AdvanceFrameRenderPrerequisite();
  */
 extern RENDERCORE_API void FlushRenderingCommands();
 
+extern RENDERCORE_API void FlushPendingDeleteRHIResources_GameThread();
+extern RENDERCORE_API void FlushPendingDeleteRHIResources_RenderThread();
+
 
 ////////////////////////////////////
 // Render thread suspension
diff --git a/Engine/Source/Runtime/Renderer/Private/DeferredShadingRenderer.cpp b/Engine/Source/Runtime/Renderer/Private/DeferredShadingRenderer.cpp
--- a/Engine/Source/Runtime/Renderer/Private/DeferredShadingRenderer.cpp
+++ b/Engine/Source/Runtime/Renderer/Private/DeferredShadingRenderer.cpp
@@ -214,10 +214,44 @@ bool FDeferredShadingSceneRenderer::RenderBasePassStaticDataMasked(FRHICommandLi
 	return bDirty;
 }
 
+static TAutoConsoleVariable<int32> CVarRHICmdWidth(
+	TEXT("r.RHICmdWidth"), 8,
+	TEXT("Number of threads."));
+
 bool FDeferredShadingSceneRenderer::RenderBasePassStaticDataDefault(FRHICommandList& RHICmdList, FViewInfo& View)
 {
+	QUICK_SCOPE_CYCLE_COUNTER(AA_RenderBasePassStaticDataDefault);
 	bool bDirty = false;
+	if (!GRHICommandList.Bypass())
+	{
+		RHICmdList.Flush(); // we assume this is the immediate command list, regardless it makes sense to empty it before we start using more command lists
+		FGraphEventRef SubmitChain;
+		int32 Width = CVarRHICmdWidth.GetValueOnRenderThread(); // we use a few more than needed to cover non-equal jobs
+
+		{
+			QUICK_SCOPE_CYCLE_COUNTER(AA_CreateTasks);
+
+			FScene::EBasePassDrawListType OpaqueDrawType = FScene::EBasePass_Default;
+			{
+				SubmitChain = Scene->BasePassNoLightMapDrawList[OpaqueDrawType].DrawVisibleParallel(View, View.StaticMeshVisibilityMap, View.StaticMeshBatchVisibility, Width, SubmitChain, bDirty);
+				SubmitChain = Scene->BasePassSimpleDynamicLightingDrawList[OpaqueDrawType].DrawVisibleParallel(View, View.StaticMeshVisibilityMap, View.StaticMeshBatchVisibility, Width, SubmitChain, bDirty);
+				SubmitChain = Scene->BasePassCachedVolumeIndirectLightingDrawList[OpaqueDrawType].DrawVisibleParallel(View, View.StaticMeshVisibilityMap, View.StaticMeshBatchVisibility, Width, SubmitChain, bDirty);
+				SubmitChain = Scene->BasePassCachedPointIndirectLightingDrawList[OpaqueDrawType].DrawVisibleParallel(View, View.StaticMeshVisibilityMap, View.StaticMeshBatchVisibility, Width, SubmitChain, bDirty);
+			}
+			{
+				SubmitChain = Scene->BasePassHighQualityLightMapDrawList[OpaqueDrawType].DrawVisibleParallel(View, View.StaticMeshVisibilityMap, View.StaticMeshBatchVisibility, Width, SubmitChain, bDirty);
+				SubmitChain = Scene->BasePassDistanceFieldShadowMapLightMapDrawList[OpaqueDrawType].DrawVisibleParallel(View, View.StaticMeshVisibilityMap, View.StaticMeshBatchVisibility, Width, SubmitChain, bDirty);
+				SubmitChain = Scene->BasePassLowQualityLightMapDrawList[OpaqueDrawType].DrawVisibleParallel(View, View.StaticMeshVisibilityMap, View.StaticMeshBatchVisibility, Width, SubmitChain, bDirty);
+			}
+		}
+		if (SubmitChain.GetReference())
+		{
+			QUICK_SCOPE_CYCLE_COUNTER(AA_WaitAndSubmit);
+			FTaskGraphInterface::Get().WaitUntilTaskCompletes(SubmitChain, ENamedThreads::RenderThread_Local);
+		}
 
+	}
+	else
 	{
 		FScene::EBasePassDrawListType OpaqueDrawType = FScene::EBasePass_Default;
 		{
diff --git a/Engine/Source/Runtime/Renderer/Private/SceneRendering.cpp b/Engine/Source/Runtime/Renderer/Private/SceneRendering.cpp
--- a/Engine/Source/Runtime/Renderer/Private/SceneRendering.cpp
+++ b/Engine/Source/Runtime/Renderer/Private/SceneRendering.cpp
@@ -1024,6 +1024,7 @@ void FRendererModule::BeginRenderingViewFamily(FCanvas* Canvas,const FSceneViewF
 			FSceneRenderer*,SceneRenderer,SceneRenderer,
 		{
 			RenderViewFamily_RenderThread(RHICmdList, SceneRenderer);
+			FlushPendingDeleteRHIResources_RenderThread();
 		});
 	}
 }
diff --git a/Engine/Source/Runtime/Renderer/Private/ShaderBaseClasses.cpp b/Engine/Source/Runtime/Renderer/Private/ShaderBaseClasses.cpp
--- a/Engine/Source/Runtime/Renderer/Private/ShaderBaseClasses.cpp
+++ b/Engine/Source/Runtime/Renderer/Private/ShaderBaseClasses.cpp
@@ -80,9 +80,17 @@ void FMaterialShader::SetParameters(
 
 	if (!bAllowCachedUniformExpressions || !UniformExpressionCache->bUpToDate || bOverrideSelection)
 	{
-		FMaterialRenderContext MaterialRenderContext(MaterialRenderProxy, Material, &View);
-		MaterialRenderProxy->EvaluateUniformExpressions(TempUniformExpressionCache, MaterialRenderContext);
-		UniformExpressionCache = &TempUniformExpressionCache;
+		// temp solution hopefully, this is unsavory to have to send a task back to the render thread and wait for it.
+		if (!IsInRenderingThread())
+		{
+			check(0);
+		}
+		else
+		{
+			FMaterialRenderContext MaterialRenderContext(MaterialRenderProxy, Material, &View);
+			MaterialRenderProxy->EvaluateUniformExpressions(TempUniformExpressionCache, MaterialRenderContext);
+			UniformExpressionCache = &TempUniformExpressionCache;
+		}
 	}
 
 	check(Material.GetRenderingThreadShaderMap());
diff --git a/Engine/Source/Runtime/Renderer/Private/StaticMeshDrawList.h b/Engine/Source/Runtime/Renderer/Private/StaticMeshDrawList.h
--- a/Engine/Source/Runtime/Renderer/Private/StaticMeshDrawList.h
+++ b/Engine/Source/Runtime/Renderer/Private/StaticMeshDrawList.h
@@ -71,10 +71,15 @@ inline FDrawListSortKey GetSortKey(bool bBackground, float BoundsRadius, float D
  * @param DrawingPolicyType - The drawing policy type used to draw mesh in this draw list.
  * @param HashSize - The number of buckets to use in the drawing policy hash.
  */
+template<typename DrawingPolicyType>
+class FDrawVisibleAnyThreadTask;
+
 template<typename DrawingPolicyType>
 class TStaticMeshDrawList : public FStaticMeshDrawListBase, public FRenderResource
 {
 public:
+	friend class FDrawVisibleAnyThreadTask<DrawingPolicyType>;
+
 	typedef typename DrawingPolicyType::ElementDataType ElementPolicyDataType;
 
 private:
@@ -271,6 +276,15 @@ public:
 	 */
 	bool DrawVisible(FRHICommandList& RHICmdList, const FViewInfo& View, const TBitArray<SceneRenderingBitArrayAllocator>& StaticMeshVisibilityMap, const TArray<uint64,SceneRenderingAllocator>& BatchVisibilityArray);
 
+	/**
+	* Draws only the static meshes which are in the visibility map.
+	* @param View - The view of the meshes to render.
+	* @param StaticMeshVisibilityMap - An map from FStaticMesh::Id to visibility state.
+	* @param BatchVisibilityArray - An array of batch element visibility bitmasks.
+	* @return True if any static meshes were drawn.
+	*/
+	FGraphEventRef DrawVisibleParallel(const FViewInfo& View, const TBitArray<SceneRenderingBitArrayAllocator>& StaticMeshVisibilityMap, const TArray<uint64, SceneRenderingAllocator>& BatchVisibilityArray, int32 Width, FGraphEventRef SubmitChain, bool& OutDirty);
+
 	/**
 	 * Draws only the static meshes which are in the visibility map, sorted front-to-back.
 	 * @param View - The view of the meshes to render.
diff --git a/Engine/Source/Runtime/Renderer/Private/StaticMeshDrawList.inl b/Engine/Source/Runtime/Renderer/Private/StaticMeshDrawList.inl
--- a/Engine/Source/Runtime/Renderer/Private/StaticMeshDrawList.inl
+++ b/Engine/Source/Runtime/Renderer/Private/StaticMeshDrawList.inl
@@ -273,6 +273,159 @@ bool TStaticMeshDrawList<DrawingPolicyType>::DrawVisible(
 	return bDirty;
 }
 
+
+template<typename DrawingPolicyType>
+class FDrawVisibleAnyThreadTask
+{
+	TStaticMeshDrawList<DrawingPolicyType>& Caller;
+	FRHICommandList& RHICmdList;
+	const FViewInfo& View;
+	const TBitArray<SceneRenderingBitArrayAllocator>& StaticMeshVisibilityMap;
+	const TArray<uint64, SceneRenderingAllocator>& BatchVisibilityArray;
+
+	const int32 FirstPolicy;
+	const int32 LastPolicy;
+
+	bool& OutDirty;
+
+public:
+
+	FDrawVisibleAnyThreadTask(
+		TStaticMeshDrawList<DrawingPolicyType>* InCaller,
+		FRHICommandList* InRHICmdList,
+		const FViewInfo* InView,
+		const TBitArray<SceneRenderingBitArrayAllocator>* InStaticMeshVisibilityMap,
+		const TArray<uint64, SceneRenderingAllocator>* InBatchVisibilityArray,
+		int32 InFirstPolicy,
+		int32 InLastPolicy,
+		bool* InOutDirty
+		)
+		: Caller(*InCaller)
+		, RHICmdList(*InRHICmdList)
+		, View(*InView)
+		, StaticMeshVisibilityMap(*InStaticMeshVisibilityMap)
+		, BatchVisibilityArray(*InBatchVisibilityArray)
+		, FirstPolicy(InFirstPolicy)
+		, LastPolicy(InLastPolicy)
+		, OutDirty(*InOutDirty)
+	{
+	}
+
+	FORCEINLINE TStatId GetStatId() const
+	{
+		RETURN_QUICK_DECLARE_CYCLE_STAT(FDrawVisibleAnyThreadTask, STATGROUP_TaskGraphTasks);
+	}
+
+	ENamedThreads::Type GetDesiredThread()
+	{
+		return ENamedThreads::AnyThread;
+	}
+
+	static ESubsequentsMode::Type GetSubsequentsMode() { return ESubsequentsMode::TrackSubsequents; }
+
+	void DoTask(ENamedThreads::Type CurrentThread, const FGraphEventRef& MyCompletionGraphEvent)
+	{
+		QUICK_SCOPE_CYCLE_COUNTER(AA_FDrawVisibleAnyThreadTask);
+
+		for (int32 Index = this->FirstPolicy; Index <= this->LastPolicy; Index++)
+		{
+			typename TStaticMeshDrawList<DrawingPolicyType>::FDrawingPolicyLink* DrawingPolicyLink = &this->Caller.DrawingPolicySet[this->Caller.OrderedDrawingPolicies[Index]];
+			bool bDrawnShared = false;
+			FPlatformMisc::Prefetch(DrawingPolicyLink->CompactElements.GetTypedData());
+			const int32 NumElements = DrawingPolicyLink->Elements.Num();
+			FPlatformMisc::Prefetch(&DrawingPolicyLink->CompactElements.GetTypedData()->MeshId);
+			const typename TStaticMeshDrawList<DrawingPolicyType>::FElementCompact* CompactElementPtr = DrawingPolicyLink->CompactElements.GetTypedData();
+			for (int32 ElementIndex = 0; ElementIndex < NumElements; ElementIndex++, CompactElementPtr++)
+			{
+				if (this->StaticMeshVisibilityMap.AccessCorrespondingBit(FRelativeBitReference(CompactElementPtr->MeshId)))
+				{
+					const typename TStaticMeshDrawList<DrawingPolicyType>::FElement& Element = DrawingPolicyLink->Elements[ElementIndex];
+					INC_DWORD_STAT_BY(STAT_StaticMeshTriangles, Element.Mesh->GetNumPrimitives());
+					// Avoid the cache miss looking up batch visibility if there is only one element.
+					uint64 BatchElementMask = Element.Mesh->Elements.Num() == 1 ? 1 : this->BatchVisibilityArray[Element.Mesh->Id];
+					this->Caller.DrawElement(this->RHICmdList, this->View, Element, BatchElementMask, DrawingPolicyLink, bDrawnShared);
+					this->OutDirty = true;
+				}
+			}
+		}
+	}
+};
+
+class FDrawVisibleRenderThreadTask
+{
+	FRHICommandList* RHICmdList;
+public:
+
+	FDrawVisibleRenderThreadTask(FRHICommandList* InRHICmdList)
+		: RHICmdList(InRHICmdList)
+	{
+	}
+
+	FORCEINLINE TStatId GetStatId() const
+	{
+		RETURN_QUICK_DECLARE_CYCLE_STAT(FDrawVisibleRenderThreadTask, STATGROUP_TaskGraphTasks);
+	}
+
+	ENamedThreads::Type GetDesiredThread()
+	{
+		return ENamedThreads::RenderThread_Local;
+	}
+
+	static ESubsequentsMode::Type GetSubsequentsMode() { return ESubsequentsMode::TrackSubsequents; }
+
+	void DoTask(ENamedThreads::Type CurrentThread, const FGraphEventRef& MyCompletionGraphEvent)
+	{
+		QUICK_SCOPE_CYCLE_COUNTER(AA_FDrawVisibleRenderThreadTask);
+		delete RHICmdList;
+	}
+};
+
+template<typename DrawingPolicyType>
+FGraphEventRef TStaticMeshDrawList<DrawingPolicyType>::DrawVisibleParallel(
+	const FViewInfo& View,
+	const TBitArray<SceneRenderingBitArrayAllocator>& StaticMeshVisibilityMap,
+	const TArray<uint64, SceneRenderingAllocator>& BatchVisibilityArray,
+	int32 Width, FGraphEventRef SubmitChain, bool& OutDirty
+	)
+{
+
+	int32 EffectiveThreads = FMath::Min<int32>(OrderedDrawingPolicies.Num(), Width);
+
+	int32 Start = 0;
+	if (EffectiveThreads)
+	{
+
+		int32 NumPer = OrderedDrawingPolicies.Num() / EffectiveThreads;
+		int32 Extra = OrderedDrawingPolicies.Num() - NumPer * EffectiveThreads;
+
+
+		for (int32 ThreadIndex = 0; ThreadIndex < EffectiveThreads; ThreadIndex++)
+		{
+			int32 Last = Start + (NumPer - 1) + (ThreadIndex < Extra);
+			check(Last >= Start);
+
+			FRHICommandList* CmdList = new FRHICommandList;
+
+			FGraphEventRef AnyThreadCompletionEvent = TGraphTask<FDrawVisibleAnyThreadTask<DrawingPolicyType> >::CreateTask(nullptr, ENamedThreads::RenderThread)
+				.ConstructAndDispatchWhenReady(this, CmdList, &View, &StaticMeshVisibilityMap, &BatchVisibilityArray, Start, Last, &OutDirty);
+			Start = Last + 1;
+
+			FGraphEventArray Prereqs;
+			Prereqs.Add(AnyThreadCompletionEvent);
+			if (SubmitChain.GetReference())
+			{
+				Prereqs.Add(SubmitChain);
+			}
+
+			SubmitChain = TGraphTask<FDrawVisibleRenderThreadTask>::CreateTask(&Prereqs, ENamedThreads::RenderThread).ConstructAndDispatchWhenReady(CmdList);
+
+		}
+	}
+	check(Start == OrderedDrawingPolicies.Num());
+
+	return SubmitChain;
+}
+
 template<typename DrawingPolicyType>
 int32 TStaticMeshDrawList<DrawingPolicyType>::DrawVisibleFrontToBack(
 	FRHICommandList& RHICmdList,
diff --git a/Engine/Source/Runtime/Windows/D3D11RHI/Private/D3D11Viewport.cpp b/Engine/Source/Runtime/Windows/D3D11RHI/Private/D3D11Viewport.cpp
--- a/Engine/Source/Runtime/Windows/D3D11RHI/Private/D3D11Viewport.cpp
+++ b/Engine/Source/Runtime/Windows/D3D11RHI/Private/D3D11Viewport.cpp
@@ -137,6 +137,8 @@ FD3D11Texture2D* GetSwapChainSurface(FD3D11DynamicRHI* D3DRHI,IDXGISwapChain* Sw
 
 	D3D11TextureAllocated2D(*NewTexture);
 
+	NewTexture->DoNoDeferDelete();
+
 	return NewTexture;
 }
 
diff --git a/Engine/Source/Runtime/Windows/D3D11RHI/Public/D3D11Resources.h b/Engine/Source/Runtime/Windows/D3D11RHI/Public/D3D11Resources.h
--- a/Engine/Source/Runtime/Windows/D3D11RHI/Public/D3D11Resources.h
+++ b/Engine/Source/Runtime/Windows/D3D11RHI/Public/D3D11Resources.h
@@ -373,15 +373,15 @@ public:
 	// IRefCountedObject interface.
 	virtual uint32 AddRef() const
 	{
-		return FRefCountedObject::AddRef();
+		return FRHIResource::AddRef();
 	}
 	virtual uint32 Release() const
 	{
-		return FRefCountedObject::Release();
+		return FRHIResource::Release();
 	}
 	virtual uint32 GetRefCount() const
 	{
-		return FRefCountedObject::GetRefCount();
+		return FRHIResource::GetRefCount();
 	}
 #if PLATFORM_SUPPORTS_VIRTUAL_TEXTURES
 	void* GetRawTextureMemory() const
@@ -440,15 +440,15 @@ public:
 	// IRefCountedObject interface.
 	virtual uint32 AddRef() const
 	{
-		return FRefCountedObject::AddRef();
+		return FRHIResource::AddRef();
 	}
 	virtual uint32 Release() const
 	{
-		return FRefCountedObject::Release();
+		return FRHIResource::Release();
 	}
 	virtual uint32 GetRefCount() const
 	{
-		return FRefCountedObject::GetRefCount();
+		return FRHIResource::GetRefCount();
 	}
 };
 
@@ -506,15 +506,15 @@ public:
 	// IRefCountedObject interface.
 	virtual uint32 AddRef() const
 	{
-		return FRefCountedObject::AddRef();
+		return FRHIResource::AddRef();
 	}
 	virtual uint32 Release() const
 	{
-		return FRefCountedObject::Release();
+		return FRHIResource::Release();
 	}
 	virtual uint32 GetRefCount() const
 	{
-		return FRefCountedObject::GetRefCount();
+		return FRHIResource::GetRefCount();
 	}
 };
 
@@ -671,15 +671,15 @@ public:
 	// IRefCountedObject interface.
 	virtual uint32 AddRef() const
 	{
-		return FRefCountedObject::AddRef();
+		return FRHIResource::AddRef();
 	}
 	virtual uint32 Release() const
 	{
-		return FRefCountedObject::Release();
+		return FRHIResource::Release();
 	}
 	virtual uint32 GetRefCount() const
 	{
-		return FRefCountedObject::GetRefCount();
+		return FRHIResource::GetRefCount();
 	}
 };
 
@@ -703,15 +703,15 @@ public:
 	// IRefCountedObject interface.
 	virtual uint32 AddRef() const
 	{
-		return FRefCountedObject::AddRef();
+		return FRHIResource::AddRef();
 	}
 	virtual uint32 Release() const
 	{
-		return FRefCountedObject::Release();
+		return FRHIResource::Release();
 	}
 	virtual uint32 GetRefCount() const
 	{
-		return FRefCountedObject::GetRefCount();
+		return FRHIResource::GetRefCount();
 	}
 };
 
@@ -735,15 +735,15 @@ public:
 	// IRefCountedObject interface.
 	virtual uint32 AddRef() const
 	{
-		return FRefCountedObject::AddRef();
+		return FRHIResource::AddRef();
 	}
 	virtual uint32 Release() const
 	{
-		return FRefCountedObject::Release();
+		return FRHIResource::Release();
 	}
 	virtual uint32 GetRefCount() const
 	{
-		return FRefCountedObject::GetRefCount();
+		return FRHIResource::GetRefCount();
 	}
 };
 
