UE4 - local uniformbuffers for command lists

diff --git a/Engine/Source/Runtime/Engine/Private/Materials/MaterialShared.cpp b/Engine/Source/Runtime/Engine/Private/Materials/MaterialShared.cpp
--- a/Engine/Source/Runtime/Engine/Private/Materials/MaterialShared.cpp
+++ b/Engine/Source/Runtime/Engine/Private/Materials/MaterialShared.cpp
@@ -1437,7 +1437,7 @@ FMaterialRenderContext::FMaterialRenderContext(
 	FMaterialRenderProxy
 -----------------------------------------------------------------------------*/
 
-void FMaterialRenderProxy::EvaluateUniformExpressions(FUniformExpressionCache& OutUniformExpressionCache, const FMaterialRenderContext& Context) const
+void FMaterialRenderProxy::EvaluateUniformExpressions(FUniformExpressionCache& OutUniformExpressionCache, const FMaterialRenderContext& Context, FRHICommandList* CommandListIfLocalMode) const
 {
 	check(IsInParallelRenderingThread());
 
@@ -1449,7 +1449,7 @@ void FMaterialRenderProxy::EvaluateUniformExpressions(FUniformExpressionCache& O
 	OutUniformExpressionCache.CachedUniformExpressionShaderMap = Context.Material.GetRenderingThreadShaderMap();
 
 	// Create and cache the material's uniform buffer.
-	OutUniformExpressionCache.UniformBuffer = UniformExpressionSet.CreateUniformBuffer(Context);
+	OutUniformExpressionCache.UniformBuffer = UniformExpressionSet.CreateUniformBuffer(Context, CommandListIfLocalMode, &OutUniformExpressionCache.LocalUniformBuffer);
 
 	// Cache 2D texture uniform expressions.
 	OutUniformExpressionCache.Textures.Empty(UniformExpressionSet.Uniform2DTextureExpressions.Num());
diff --git a/Engine/Source/Runtime/Engine/Private/Materials/MaterialUniformExpressions.cpp b/Engine/Source/Runtime/Engine/Private/Materials/MaterialUniformExpressions.cpp
--- a/Engine/Source/Runtime/Engine/Private/Materials/MaterialUniformExpressions.cpp
+++ b/Engine/Source/Runtime/Engine/Private/Materials/MaterialUniformExpressions.cpp
@@ -257,10 +257,10 @@ const FUniformBufferStruct& FUniformExpressionSet::GetUniformBufferStruct() cons
 	return *UniformBufferStruct;
 }
 
-FUniformBufferRHIRef FUniformExpressionSet::CreateUniformBuffer(const FMaterialRenderContext& MaterialRenderContext) const
+FUniformBufferRHIRef FUniformExpressionSet::CreateUniformBuffer(const FMaterialRenderContext& MaterialRenderContext, FRHICommandList* CommandListIfLocalMode, struct FLocalUniformBuffer* OutLocalUniformBuffer) const
 {
 	check(UniformBufferStruct);
-	check(IsInRenderingThread());
+	check(IsInParallelRenderingThread());
 	
 	FUniformBufferRHIRef UniformBuffer;
 
@@ -320,8 +320,17 @@ FUniformBufferRHIRef FUniformExpressionSet::CreateUniformBuffer(const FMaterialR
 				*ResourceTable++ = GWhiteTexture->SamplerStateRHI;
 			}
 		}
-
-		UniformBuffer = RHICreateUniformBuffer(TempBuffer,UniformBufferStruct->GetLayout(),UniformBuffer_MultiFrame);
+		if (CommandListIfLocalMode)
+		{
+			check(OutLocalUniformBuffer);
+			*OutLocalUniformBuffer = CommandListIfLocalMode->BuildLocalUniformBuffer(TempBuffer, UniformBufferStruct->GetSize(), UniformBufferStruct->GetLayout());
+			check(OutLocalUniformBuffer->IsValid());
+		}
+		else
+		{
+			UniformBuffer = RHICreateUniformBuffer(TempBuffer, UniformBufferStruct->GetLayout(), UniformBuffer_MultiFrame);
+			check(!OutLocalUniformBuffer->IsValid());
+		}
 	}
 
 	return UniformBuffer;
@@ -358,7 +367,7 @@ void FMaterialUniformExpressionTexture::SetTransientOverrideTextureValue( UTextu
 
 void FMaterialUniformExpressionTexture::GetTextureValue(const FMaterialRenderContext& Context,const FMaterial& Material,const UTexture*& OutValue) const
 {
-	check(IsInRenderingThread());
+	check(IsInParallelRenderingThread());
 	if( TransientOverrideValue_RenderThread != NULL )
 	{
 		OutValue = TransientOverrideValue_RenderThread;
diff --git a/Engine/Source/Runtime/Engine/Public/MaterialShared.h b/Engine/Source/Runtime/Engine/Public/MaterialShared.h
--- a/Engine/Source/Runtime/Engine/Public/MaterialShared.h
+++ b/Engine/Source/Runtime/Engine/Public/MaterialShared.h
@@ -214,7 +214,7 @@ public:
 	void SetParameterCollections(const TArray<class UMaterialParameterCollection*>& Collections);
 	void CreateBufferStruct();
 	const FUniformBufferStruct& GetUniformBufferStruct() const;
-	ENGINE_API FUniformBufferRHIRef CreateUniformBuffer(const FMaterialRenderContext& MaterialRenderContext) const;
+	ENGINE_API FUniformBufferRHIRef CreateUniformBuffer(const FMaterialRenderContext& MaterialRenderContext, FRHICommandList* CommandListIfLocalMode, struct FLocalUniformBuffer* OutLocalUniformBuffer) const;
 
 	uint32 GetAllocatedSize() const
 	{
@@ -1193,6 +1193,8 @@ struct FUniformExpressionCache
 {
 	/** Material uniform buffer. */
 	FUniformBufferRHIRef UniformBuffer;
+	/** Material uniform buffer. */
+	FLocalUniformBuffer LocalUniformBuffer;
 	/** Textures referenced by uniform expressions. */
 	TArray<const UTexture*> Textures;
 	/** Cube textures referenced by uniform expressions. */
@@ -1241,7 +1243,7 @@ public:
 	 * @param OutUniformExpressionCache - The uniform expression cache to build.
 	 * @param MaterialRenderContext - The context for which to cache expressions.
 	 */
-	void ENGINE_API EvaluateUniformExpressions(FUniformExpressionCache& OutUniformExpressionCache, const FMaterialRenderContext& Context) const;
+	void ENGINE_API EvaluateUniformExpressions(FUniformExpressionCache& OutUniformExpressionCache, const FMaterialRenderContext& Context, class FRHICommandList* CommandListIfLocalMode = nullptr) const;
 
 	/**
 	 * Caches uniform expressions for efficient runtime evaluation.
diff --git a/Engine/Source/Runtime/RHI/Private/RHICommandList.cpp b/Engine/Source/Runtime/RHI/Private/RHICommandList.cpp
--- a/Engine/Source/Runtime/RHI/Private/RHICommandList.cpp
+++ b/Engine/Source/Runtime/RHI/Private/RHICommandList.cpp
@@ -486,7 +486,47 @@ void FRHICommandListExecutor::ExecuteList(FRHICommandList& CmdList)
 				}
 			}
 			break;
-
+		case ERCT_BuildLocalUniformBuffer:
+			{
+				auto* RHICmd = Iter.NextCommand<FRHICommandBuildLocalUniformBuffer>();
+				if (!bOnlyTestMemAccess)
+				{
+					check(RHICmd->WorkArea.CheckCmdList == &CmdList && RHICmd->WorkArea.UID == CmdList.UID // this uniform buffer was not built for this particular commandlist
+						&& !IsValidRef(RHICmd->WorkArea.UniformBuffer) && RHICmd->WorkArea.Layout && RHICmd->WorkArea.Contents); // should not already have been created
+					if (RHICmd->WorkArea.UseCount)
+					{
+						RHICmd->WorkArea.UniformBuffer = RHICreateUniformBuffer(RHICmd->WorkArea.Contents, *RHICmd->WorkArea.Layout, UniformBuffer_SingleFrame);
+					}
+					FMemory::Free(RHICmd->WorkArea.Contents);
+					RHICmd->WorkArea.Contents = nullptr;
+				}
+			}
+			break;
+		case ERCT_SetLocalUniformBuffer:
+			{
+				auto* RHICmd = Iter.NextCommand<FRHICommandSetLocalUniformBuffer>();
+				if (!bOnlyTestMemAccess)
+				{
+					check(RHICmd->LocalUniformBuffer.WorkArea->CheckCmdList == &CmdList && RHICmd->LocalUniformBuffer.WorkArea->UID == CmdList.UID // this BSS was not built for this particular commandlist
+						&& RHICmd->LocalUniformBuffer.WorkArea->UseCount > 0 && IsValidRef(RHICmd->LocalUniformBuffer.WorkArea->UniformBuffer)); // this should have been created and should have uses outstanding
+					switch (RHICmd->ShaderFrequency)
+					{
+						case SF_Vertex:		SetShaderUniformBuffer_Internal((FVertexShaderRHIParamRef)RHICmd->Shader, RHICmd->BaseIndex, RHICmd->LocalUniformBuffer.WorkArea->UniformBuffer); break;
+						case SF_Hull:		SetShaderUniformBuffer_Internal((FHullShaderRHIParamRef)RHICmd->Shader, RHICmd->BaseIndex, RHICmd->LocalUniformBuffer.WorkArea->UniformBuffer); break;
+						case SF_Domain:		SetShaderUniformBuffer_Internal((FDomainShaderRHIParamRef)RHICmd->Shader, RHICmd->BaseIndex, RHICmd->LocalUniformBuffer.WorkArea->UniformBuffer); break;
+						case SF_Geometry:	SetShaderUniformBuffer_Internal((FGeometryShaderRHIParamRef)RHICmd->Shader, RHICmd->BaseIndex, RHICmd->LocalUniformBuffer.WorkArea->UniformBuffer); break;
+						case SF_Pixel:		SetShaderUniformBuffer_Internal((FPixelShaderRHIParamRef)RHICmd->Shader, RHICmd->BaseIndex, RHICmd->LocalUniformBuffer.WorkArea->UniformBuffer); break;
+						case SF_Compute:	SetShaderUniformBuffer_Internal((FComputeShaderRHIParamRef)RHICmd->Shader, RHICmd->BaseIndex, RHICmd->LocalUniformBuffer.WorkArea->UniformBuffer); break;
+						default: check(0); break;
+					}
+					if (--RHICmd->LocalUniformBuffer.WorkArea->UseCount == 0)
+					{
+						RHICmd->LocalUniformBuffer.WorkArea->UniformBuffer = nullptr; // this also releases the ref, which wouldn't otherwise happen because RHICommands are unsafe and don't do constructors or destructors
+						RHICmd->LocalUniformBuffer.WorkArea->Layout = nullptr;
+					}
+				}
+			}
+			break;
 
 		default:
 			checkf(0, TEXT("Unknown RHI Command %d!"), Cmd->Type);
@@ -536,13 +576,13 @@ void FRHICommandListExecutor::ExecuteList(FRHICommandList& CmdList)
 		}
 	}
 	INC_DWORD_STAT_BY(STAT_RHICmdListCount, CmdList.NumCommands);
+	CmdList.Reset();
 #if !(UE_BUILD_SHIPPING || UE_BUILD_TEST)
 	if (&CmdList == &GetImmediateCommandList() && GRHICommandList.OutstandingCmdListCount.GetValue() == 1)
 	{
 		LatchBypass();
 	}
 #endif
-	CmdList.Reset();
 }
 
 void FRHICommandListExecutor::LatchBypass()
@@ -550,10 +590,6 @@ void FRHICommandListExecutor::LatchBypass()
 #if !(UE_BUILD_SHIPPING || UE_BUILD_TEST)
 	check(GRHICommandList.OutstandingCmdListCount.GetValue() == 1 && !GRHICommandList.GetImmediateCommandList().HasCommands());
 	bool NewBypass = (CVarRHICmdBypass.GetValueOnRenderThread() >= 1);
-	if (NewBypass)
-	{
-		FRHIResource::FlushPendingDeletes();
-	}
 	bLatchedBypass = NewBypass;
 #endif
 }
diff --git a/Engine/Source/Runtime/RHI/Public/RHI.h b/Engine/Source/Runtime/RHI/Public/RHI.h
--- a/Engine/Source/Runtime/RHI/Public/RHI.h
+++ b/Engine/Source/Runtime/RHI/Public/RHI.h
@@ -11,6 +11,9 @@
 #include "RHIDefinitions.h"
 #include "StaticArray.h"
 
+/** Uniform buffer structs must be aligned to 16-byte boundaries. */
+#define UNIFORM_BUFFER_STRUCT_ALIGNMENT 16
+
 // Forward declarations.
 class FSceneView;
 struct FMeshBatch;
diff --git a/Engine/Source/Runtime/RHI/Public/RHICommandList.h b/Engine/Source/Runtime/RHI/Public/RHICommandList.h
--- a/Engine/Source/Runtime/RHI/Public/RHICommandList.h
+++ b/Engine/Source/Runtime/RHI/Public/RHICommandList.h
@@ -51,6 +51,9 @@ enum ERHICommandType
 	ERCT_CopyToResolveTarget,
 	ERCT_ClearMRT,
 	ERCT_Clear,
+	ERCT_BuildLocalUniformBuffer,
+	ERCT_SetLocalUniformBuffer,
+
 
 
 
@@ -152,7 +155,6 @@ struct FRHICommandSetShaderParameter : public FRHICommandPerShader
 	{
 		FRHICommandPerShader::SetShader<TShaderRHIParamRef>(InShader);
 		Type = ERCT_SetShaderParameter;
-		Shader = InShader;
 		BufferIndex = InBufferIndex;
 		BaseIndex = InBaseIndex;
 		NumBytes = InNumBytes;
@@ -169,7 +171,6 @@ struct FRHICommandSetShaderUniformBuffer : public FRHICommandPerShader
 	{
 		FRHICommandPerShader::SetShader<TShaderRHIParamRef>(InShader);
 		Type = ERCT_SetShaderUniformBuffer;
-		Shader = InShader;
 		BaseIndex = InBaseIndex;
 		UniformBuffer = InUniformBuffer;
 	}
@@ -184,7 +185,6 @@ struct FRHICommandSetShaderTexture : public FRHICommandPerShader
 	{
 		FRHICommandPerShader::SetShader<TShaderRHIParamRef>(InShader);
 		Type = ERCT_SetShaderTexture;
-		Shader = InShader;
 		TextureIndex = InTextureIndex;
 		Texture = InTexture;
 	}
@@ -199,7 +199,6 @@ struct FRHICommandSetShaderResourceViewParameter : public FRHICommandPerShader
 	{
 		FRHICommandPerShader::SetShader<TShaderRHIParamRef>(InShader);
 		Type = ERCT_SetShaderResourceViewParameter;
-		Shader = InShader;
 		SamplerIndex = InSamplerIndex;
 		SRV = InSRV;
 	}
@@ -214,7 +213,6 @@ struct FRHICommandSetUAVParameter : public FRHICommandPerShader
 	{
 		FRHICommandPerShader::SetShader<TShaderRHIParamRef>(InShader);
 		Type = ERCT_SetUAVParameter;
-		Shader = InShader;
 		UAVIndex = InUAVIndex;
 		UAV = InUAV;
 	}
@@ -230,7 +228,6 @@ struct FRHICommandSetUAVParameter_IntialCount : public FRHICommandPerShader
 	{
 		FRHICommandPerShader::SetShader<TShaderRHIParamRef>(InShader);
 		Type = ERCT_SetUAVParameter_IntialCount;
-		Shader = InShader;
 		UAVIndex = InUAVIndex;
 		UAV = InUAV;
 		InitialCount = InInitialCount;
@@ -246,7 +243,6 @@ struct FRHICommandSetShaderSampler : public FRHICommandPerShader
 	{
 		FRHICommandPerShader::SetShader<TShaderRHIParamRef>(InShader);
 		Type = ERCT_SetShaderSampler;
-		Shader = InShader;
 		SamplerIndex = InSamplerIndex;
 		Sampler = InSampler;
 	}
@@ -814,6 +810,75 @@ struct FRHICommandSetGlobalBoundShaderState : public FRHICommand
 	}
 };
 
+//---
+
+struct FLocalUniformBufferWorkArea
+{
+	void* Contents;
+	const FRHIUniformBufferLayout* Layout;
+	FUniformBufferRHIRef UniformBuffer;
+	mutable int32 UseCount;
+#if !(UE_BUILD_SHIPPING || UE_BUILD_TEST)
+	class FRHICommandList* CheckCmdList;
+	int32 UID;
+#endif
+};
+
+struct FLocalUniformBuffer
+{
+	FLocalUniformBufferWorkArea* WorkArea;
+	FUniformBufferRHIRef BypassUniform; // this is only used in the case of Bypass, should eventually be deleted
+	FLocalUniformBuffer()
+		: WorkArea(nullptr)
+	{
+	}
+	FORCEINLINE_DEBUGGABLE bool IsValid() const
+	{
+		return WorkArea || IsValidRef(BypassUniform);
+	}
+};
+
+struct FRHICommandBuildLocalUniformBuffer : public FRHICommand
+{
+	FLocalUniformBufferWorkArea WorkArea;
+	FORCEINLINE_DEBUGGABLE void Set(
+	class FRHICommandList* CheckCmdList,
+		int32 UID,
+		const void* Contents,
+		uint32 ContentsSize,
+		const FRHIUniformBufferLayout& Layout
+		)
+	{
+		Type = ERCT_BuildLocalUniformBuffer;
+		// these RHI commands are bad, no constructor! we construct this in place
+		WorkArea.Layout = &Layout;
+		WorkArea.Contents = FMemory::Malloc(ContentsSize, UNIFORM_BUFFER_STRUCT_ALIGNMENT);
+		FMemory::Memcpy(WorkArea.Contents, Contents, ContentsSize);
+		new ((void*)&WorkArea.UniformBuffer) FUniformBufferRHIRef();
+		WorkArea.UseCount = 0;
+#if !(UE_BUILD_SHIPPING || UE_BUILD_TEST)
+		WorkArea.CheckCmdList = CheckCmdList;
+		WorkArea.UID = UID;
+#endif
+	}
+};
+
+struct FRHICommandSetLocalUniformBuffer : public FRHICommandPerShader
+{
+	uint32 BaseIndex;
+	FLocalUniformBuffer LocalUniformBuffer;
+	template <typename TShaderRHIParamRef>
+	FORCEINLINE_DEBUGGABLE void Set(class FRHICommandList* CheckCmdList, int32 UID, TShaderRHIParamRef InShader, uint32 InBaseIndex, const FLocalUniformBuffer& InLocalUniformBuffer)
+	{
+		FRHICommandPerShader::SetShader<TShaderRHIParamRef>(InShader);
+		Type = ERCT_SetLocalUniformBuffer;
+		BaseIndex = InBaseIndex;
+		LocalUniformBuffer.WorkArea = InLocalUniformBuffer.WorkArea;
+		check(CheckCmdList == LocalUniformBuffer.WorkArea->CheckCmdList && UID == LocalUniformBuffer.WorkArea->UID); // this uniform buffer was not built for this particular commandlist
+		LocalUniformBuffer.WorkArea->UseCount++;
+	}
+};
+
 class RHI_API FRHICommandList : public FNoncopyable
 {
 private:
@@ -1077,6 +1142,35 @@ public:
 		Cmd->Set(GlobalBoundShaderState);
 	}
 
+	FORCEINLINE_DEBUGGABLE FLocalUniformBuffer BuildLocalUniformBuffer(const void* Contents, uint32 ContentsSize, const FRHIUniformBufferLayout& Layout)
+	{
+		FLocalUniformBuffer Result;
+		if (Bypass())
+		{
+			Result.BypassUniform = RHICreateUniformBuffer(Contents, Layout, UniformBuffer_SingleFrame);
+		}
+		else
+		{
+			auto* Cmd = AddCommand<FRHICommandBuildLocalUniformBuffer>();
+			check(ContentsSize);
+			Cmd->Set(this, this->UID, Contents, ContentsSize, Layout);
+			Result.WorkArea = &Cmd->WorkArea;
+		}
+		return Result;
+	}
+
+	template <typename TShaderRHIParamRef>
+	FORCEINLINE_DEBUGGABLE void SetLocalShaderUniformBuffer(TShaderRHIParamRef Shader, uint32 BaseIndex, const FLocalUniformBuffer& UniformBuffer)
+	{
+		if (Bypass())
+		{
+			SetShaderUniformBuffer_Internal(Shader, BaseIndex, UniformBuffer.BypassUniform);
+			return;
+		}
+		auto* Cmd = AddCommand<FRHICommandSetLocalUniformBuffer>();
+		Cmd->Set<TShaderRHIParamRef>(this, this->UID, Shader, BaseIndex, UniformBuffer);
+	}
+
 	template <typename TShaderRHIParamRef>
 	FORCEINLINE_DEBUGGABLE void SetShaderUniformBuffer(TShaderRHIParamRef Shader, uint32 BaseIndex, FUniformBufferRHIParamRef UniformBuffer)
 	{
diff --git a/Engine/Source/Runtime/RenderCore/Public/UniformBuffer.h b/Engine/Source/Runtime/RenderCore/Public/UniformBuffer.h
--- a/Engine/Source/Runtime/RenderCore/Public/UniformBuffer.h
+++ b/Engine/Source/Runtime/RenderCore/Public/UniformBuffer.h
@@ -10,9 +10,6 @@
 #include "RenderResource.h"
 #include "RHI.h"
 
-/** Uniform buffer structs must be aligned to 16-byte boundaries. */
-#define UNIFORM_BUFFER_STRUCT_ALIGNMENT 16
-
 namespace EShaderPrecisionModifier
 {
 	enum Type
diff --git a/Engine/Source/Runtime/Renderer/Private/ShaderBaseClasses.cpp b/Engine/Source/Runtime/Renderer/Private/ShaderBaseClasses.cpp
--- a/Engine/Source/Runtime/Renderer/Private/ShaderBaseClasses.cpp
+++ b/Engine/Source/Runtime/Renderer/Private/ShaderBaseClasses.cpp
@@ -80,17 +80,9 @@ void FMaterialShader::SetParameters(
 
 	if (!bAllowCachedUniformExpressions || !UniformExpressionCache->bUpToDate || bOverrideSelection)
 	{
-		// temp solution hopefully, this is unsavory to have to send a task back to the render thread and wait for it.
-		if (!IsInRenderingThread())
-		{
-			check(0);
-		}
-		else
-		{
-			FMaterialRenderContext MaterialRenderContext(MaterialRenderProxy, Material, &View);
-			MaterialRenderProxy->EvaluateUniformExpressions(TempUniformExpressionCache, MaterialRenderContext);
-			UniformExpressionCache = &TempUniformExpressionCache;
-		}
+		FMaterialRenderContext MaterialRenderContext(MaterialRenderProxy, Material, &View);
+		MaterialRenderProxy->EvaluateUniformExpressions(TempUniformExpressionCache, MaterialRenderContext, &RHICmdList);
+		UniformExpressionCache = &TempUniformExpressionCache;
 	}
 
 	check(Material.GetRenderingThreadShaderMap());
@@ -132,8 +124,16 @@ void FMaterialShader::SetParameters(
 	}
 #endif
 
-	// Set the material uniform buffer.
-	SetUniformBufferParameter(RHICmdList, ShaderRHI,MaterialUniformBuffer,UniformExpressionCache->UniformBuffer);
+	if (UniformExpressionCache->LocalUniformBuffer.IsValid())
+	{
+		// Set the material uniform buffer.
+		SetLocalUniformBufferParameter(RHICmdList, ShaderRHI, MaterialUniformBuffer, UniformExpressionCache->LocalUniformBuffer);
+	}
+	else
+	{
+		// Set the material uniform buffer.
+		SetUniformBufferParameter(RHICmdList, ShaderRHI, MaterialUniformBuffer, UniformExpressionCache->UniformBuffer);
+	}
 
 	{
 		const TArray<FGuid>& ParameterCollections = UniformExpressionCache->ParameterCollections;
diff --git a/Engine/Source/Runtime/ShaderCore/Public/ShaderParameterUtils.h b/Engine/Source/Runtime/ShaderCore/Public/ShaderParameterUtils.h
--- a/Engine/Source/Runtime/ShaderCore/Public/ShaderParameterUtils.h
+++ b/Engine/Source/Runtime/ShaderCore/Public/ShaderParameterUtils.h
@@ -337,6 +337,23 @@ inline void FRWShaderParameter::UnsetUAV(FRHICommandList& RHICmdList, FComputeSh
 }
 
 
+/** Sets the value of a shader uniform buffer parameter to a uniform buffer containing the struct. */
+template<typename TShaderRHIRef>
+inline void SetLocalUniformBufferParameter(
+	FRHICommandList& RHICmdList,
+	TShaderRHIRef Shader,
+	const FShaderUniformBufferParameter& Parameter,
+	const FLocalUniformBuffer& LocalUniformBuffer
+	)
+{
+	// This will trigger if the parameter was not serialized
+	checkSlow(Parameter.IsInitialized());
+	if(Parameter.IsBound())
+	{
+		RHICmdList.SetLocalShaderUniformBuffer(Shader, Parameter.GetBaseIndex(), LocalUniformBuffer);
+	}
+}
+
 /** Sets the value of a shader uniform buffer parameter to a uniform buffer containing the struct. */
 template<typename TShaderRHIRef>
 inline void SetUniformBufferParameter(
