UE4 - some cleanup on parallel rendering

diff --git a/Engine/Source/Runtime/Engine/Private/StaticBoundShaderState.cpp b/Engine/Source/Runtime/Engine/Private/StaticBoundShaderState.cpp
--- a/Engine/Source/Runtime/Engine/Private/StaticBoundShaderState.cpp
+++ b/Engine/Source/Runtime/Engine/Private/StaticBoundShaderState.cpp
@@ -7,8 +7,23 @@
 #include "EnginePrivate.h"
 #include "StaticBoundShaderState.h"
 
+static struct FStartupSetBSSPtr
+{
+	int DontDeadStripMeInt;
+	FStartupSetBSSPtr()
+	{
+		DontDeadStripMeInt = 0;
+		SetGlobalBoundShaderState_InternalPtr = &SetGlobalBoundShaderState_Internal;
+	}
+	void DontDeadStripMe()
+	{
+		DontDeadStripMeInt++;
+	}
+} StartupSetBSSPtr;
+
 TLinkedList<FGlobalBoundShaderStateResource*>*& FGlobalBoundShaderStateResource::GetGlobalBoundShaderStateList()
 {
+	StartupSetBSSPtr.DontDeadStripMe();
 	static TLinkedList<FGlobalBoundShaderStateResource*>* List = NULL;
 	return List;
 }
@@ -115,3 +130,5 @@ void SetGlobalBoundShaderState_Internal(FGlobalBoundShaderState& GlobalBoundShad
 			(FGeometryShaderRHIParamRef)GETSAFERHISHADER_GEOMETRY(GlobalBoundShaderState.Get()->Args.GeometryShader))
 		);
 }
+
+
diff --git a/Engine/Source/Runtime/Engine/Public/StaticBoundShaderState.h b/Engine/Source/Runtime/Engine/Public/StaticBoundShaderState.h
--- a/Engine/Source/Runtime/Engine/Public/StaticBoundShaderState.h
+++ b/Engine/Source/Runtime/Engine/Public/StaticBoundShaderState.h
@@ -80,7 +80,7 @@ FORCEINLINE void SetGlobalBoundShaderState(
 	FShader* GeometryShader = nullptr
 	)
 {
-	RHICmdList.BuildAndSetGlobalBoundShaderState(&SetGlobalBoundShaderState_Internal, BoundShaderState, VertexDeclaration, VertexShader, PixelShader, GeometryShader);
+	RHICmdList.BuildAndSetGlobalBoundShaderState(BoundShaderState, VertexDeclaration, VertexShader, PixelShader, GeometryShader);
 }
 
 #endif
diff --git a/Engine/Source/Runtime/RHI/Private/RHICommandList.cpp b/Engine/Source/Runtime/RHI/Private/RHICommandList.cpp
--- a/Engine/Source/Runtime/RHI/Private/RHICommandList.cpp
+++ b/Engine/Source/Runtime/RHI/Private/RHICommandList.cpp
@@ -4,19 +4,13 @@
 #include "RHI.h"
 #include "RHICommandList.h"
 
-DEFINE_STAT(STAT_RHICmdListExecuteTime);
-DEFINE_STAT(STAT_RHICmdListEnqueueTime);
-DEFINE_STAT(STAT_RHICmdListMemory);
-DEFINE_STAT(STAT_RHICmdListCount);
-DEFINE_STAT(STAT_RHICounterTEMP);
-
-static int32 GRHICmd = 0;
-static FAutoConsoleVariableRef CVarRHICmd(
-	TEXT("r.RHICmd"),
-	GRHICmd,
-	TEXT("0: Disables the RHI Command List mode, running regular RHI calls (default)\n")
-	TEXT("1: Enables the RHI Command List mode")
-	);
+DEFINE_STAT(STAT_ImmedCmdListExecuteTime);
+DEFINE_STAT(STAT_ImmedCmdListMemory);
+DEFINE_STAT(STAT_ImmedCmdListCount);
+
+DEFINE_STAT(STAT_NonImmedCmdListExecuteTime);
+DEFINE_STAT(STAT_NonImmedCmdListMemory);
+DEFINE_STAT(STAT_NonImmedCmdListCount);
 
 static TAutoConsoleVariable<int32> CVarRHICmdBypass(
 	TEXT("r.RHICmdBypass"),
@@ -25,14 +19,6 @@ static TAutoConsoleVariable<int32> CVarRHICmdBypass(
 	TEXT("0: Disable, 1: Enable"),
 	ECVF_Cheat);
 
-static int32 GRHISkip = 0;
-static FAutoConsoleVariableRef CVarRHICmdSkip(
-	TEXT("r.RHISkip"),
-	GRHISkip,
-	TEXT("0: No effect (default)\n")
-	TEXT("1: Skips executing RHI Command List")
-	);
-
 static int32 GRHICmdExec = 0;
 static FAutoConsoleVariableRef CVarRHICmdExec(
 	TEXT("r.RHIExec"),
@@ -41,22 +27,17 @@ static FAutoConsoleVariableRef CVarRHICmdExec(
 	TEXT("1: Executes RHI Command List using platform specific APIs")
 	);
 
-static int32 GRHICmdMem = 0;
-static FAutoConsoleVariableRef CVarRHICmdMem(
-	TEXT("r.RHIMem"),
-	GRHICmdMem,
-	TEXT("0: No effect(default)\n")
-	TEXT("1: Processes the RHI Command List but doesn't execute commands")
-	);
 
 RHI_API FRHICommandListExecutor GRHICommandList;
 RHI_API TLockFreeFixedSizeAllocator<FRHICommandList::FMemManager::FPage::PageSize> FRHICommandList::FMemManager::FPage::TheAllocator;
+RHI_API void(*SetGlobalBoundShaderState_InternalPtr)(FGlobalBoundShaderState& BoundShaderState);
+
+
 //RHI_API FGlobalRHI GRHI;
 
 static_assert(sizeof(FRHICommand) == sizeof(FRHICommandNopEndOfPage), "These should match for filling in the end of a page when a new allocation won't fit");
 
-template <bool bOnlyTestMemAccess>
-void FRHICommandListExecutor::ExecuteList(FRHICommandList& CmdList)
+void FRHICommandListExecutor::ExecuteListGenericRHI(FRHICommandList& CmdList)
 {
 	FRHICommandListIterator Iter(CmdList);
 	while (Iter.HasCommandsLeft())
@@ -81,249 +62,191 @@ void FRHICommandListExecutor::ExecuteList(FRHICommandList& CmdList)
 		case ERCT_SetRasterizerState:
 			{
 				auto* RHICmd = Iter.NextCommand<FRHICommandSetRasterizerState>();
-				if (!bOnlyTestMemAccess)
-				{
-					SetRasterizerState_Internal(RHICmd->State);
-				}
+				SetRasterizerState_Internal(RHICmd->State);
 			}
 			break;
 		case ERCT_SetShaderParameter:
 			{
 				auto* RHICmd = Iter.NextCommand<FRHICommandSetShaderParameter>();
-				if (!bOnlyTestMemAccess)
+				switch (RHICmd->ShaderFrequency)
 				{
-					switch (RHICmd->ShaderFrequency)
-					{
-					case SF_Vertex:		SetShaderParameter_Internal((FVertexShaderRHIParamRef)RHICmd->Shader, RHICmd->BufferIndex, RHICmd->BaseIndex, RHICmd->NumBytes, RHICmd->NewValue); break;
-					case SF_Hull:		SetShaderParameter_Internal((FHullShaderRHIParamRef)RHICmd->Shader, RHICmd->BufferIndex, RHICmd->BaseIndex, RHICmd->NumBytes, RHICmd->NewValue); break;
-					case SF_Domain:		SetShaderParameter_Internal((FDomainShaderRHIParamRef)RHICmd->Shader, RHICmd->BufferIndex, RHICmd->BaseIndex, RHICmd->NumBytes, RHICmd->NewValue); break;
-					case SF_Geometry:	SetShaderParameter_Internal((FGeometryShaderRHIParamRef)RHICmd->Shader, RHICmd->BufferIndex, RHICmd->BaseIndex, RHICmd->NumBytes, RHICmd->NewValue); break;
-					case SF_Pixel:		SetShaderParameter_Internal((FPixelShaderRHIParamRef)RHICmd->Shader, RHICmd->BufferIndex, RHICmd->BaseIndex, RHICmd->NumBytes, RHICmd->NewValue); break;
-					case SF_Compute:	SetShaderParameter_Internal((FComputeShaderRHIParamRef)RHICmd->Shader, RHICmd->BufferIndex, RHICmd->BaseIndex, RHICmd->NumBytes, RHICmd->NewValue); break;
-					default: check(0); break;
-					}
+				case SF_Vertex:		SetShaderParameter_Internal((FVertexShaderRHIParamRef)RHICmd->Shader, RHICmd->BufferIndex, RHICmd->BaseIndex, RHICmd->NumBytes, RHICmd->NewValue); break;
+				case SF_Hull:		SetShaderParameter_Internal((FHullShaderRHIParamRef)RHICmd->Shader, RHICmd->BufferIndex, RHICmd->BaseIndex, RHICmd->NumBytes, RHICmd->NewValue); break;
+				case SF_Domain:		SetShaderParameter_Internal((FDomainShaderRHIParamRef)RHICmd->Shader, RHICmd->BufferIndex, RHICmd->BaseIndex, RHICmd->NumBytes, RHICmd->NewValue); break;
+				case SF_Geometry:	SetShaderParameter_Internal((FGeometryShaderRHIParamRef)RHICmd->Shader, RHICmd->BufferIndex, RHICmd->BaseIndex, RHICmd->NumBytes, RHICmd->NewValue); break;
+				case SF_Pixel:		SetShaderParameter_Internal((FPixelShaderRHIParamRef)RHICmd->Shader, RHICmd->BufferIndex, RHICmd->BaseIndex, RHICmd->NumBytes, RHICmd->NewValue); break;
+				case SF_Compute:	SetShaderParameter_Internal((FComputeShaderRHIParamRef)RHICmd->Shader, RHICmd->BufferIndex, RHICmd->BaseIndex, RHICmd->NumBytes, RHICmd->NewValue); break;
+				default: check(0); break;
 				}
 			}
 			break;
 		case ERCT_SetShaderUniformBuffer:
 			{
 				auto* RHICmd = Iter.NextCommand<FRHICommandSetShaderUniformBuffer>();
-				if (!bOnlyTestMemAccess)
+				switch (RHICmd->ShaderFrequency)
 				{
-					switch (RHICmd->ShaderFrequency)
-					{
-					case SF_Vertex:		SetShaderUniformBuffer_Internal((FVertexShaderRHIParamRef)RHICmd->Shader, RHICmd->BaseIndex, RHICmd->UniformBuffer); break;
-					case SF_Hull:		SetShaderUniformBuffer_Internal((FHullShaderRHIParamRef)RHICmd->Shader, RHICmd->BaseIndex, RHICmd->UniformBuffer); break;
-					case SF_Domain:		SetShaderUniformBuffer_Internal((FDomainShaderRHIParamRef)RHICmd->Shader, RHICmd->BaseIndex, RHICmd->UniformBuffer); break;
-					case SF_Geometry:	SetShaderUniformBuffer_Internal((FGeometryShaderRHIParamRef)RHICmd->Shader, RHICmd->BaseIndex, RHICmd->UniformBuffer); break;
-					case SF_Pixel:		SetShaderUniformBuffer_Internal((FPixelShaderRHIParamRef)RHICmd->Shader, RHICmd->BaseIndex, RHICmd->UniformBuffer); break;
-					case SF_Compute:	SetShaderUniformBuffer_Internal((FComputeShaderRHIParamRef)RHICmd->Shader, RHICmd->BaseIndex, RHICmd->UniformBuffer); break;
-					default: check(0); break;
-					}
+				case SF_Vertex:		SetShaderUniformBuffer_Internal((FVertexShaderRHIParamRef)RHICmd->Shader, RHICmd->BaseIndex, RHICmd->UniformBuffer); break;
+				case SF_Hull:		SetShaderUniformBuffer_Internal((FHullShaderRHIParamRef)RHICmd->Shader, RHICmd->BaseIndex, RHICmd->UniformBuffer); break;
+				case SF_Domain:		SetShaderUniformBuffer_Internal((FDomainShaderRHIParamRef)RHICmd->Shader, RHICmd->BaseIndex, RHICmd->UniformBuffer); break;
+				case SF_Geometry:	SetShaderUniformBuffer_Internal((FGeometryShaderRHIParamRef)RHICmd->Shader, RHICmd->BaseIndex, RHICmd->UniformBuffer); break;
+				case SF_Pixel:		SetShaderUniformBuffer_Internal((FPixelShaderRHIParamRef)RHICmd->Shader, RHICmd->BaseIndex, RHICmd->UniformBuffer); break;
+				case SF_Compute:	SetShaderUniformBuffer_Internal((FComputeShaderRHIParamRef)RHICmd->Shader, RHICmd->BaseIndex, RHICmd->UniformBuffer); break;
+				default: check(0); break;
 				}
 			}
 			break;
 		case ERCT_SetShaderSampler:
 			{
 				auto* RHICmd = Iter.NextCommand<FRHICommandSetShaderSampler>();
-				if (!bOnlyTestMemAccess)
+				switch (RHICmd->ShaderFrequency)
 				{
-					switch (RHICmd->ShaderFrequency)
-					{
-					case SF_Vertex:		SetShaderSampler_Internal((FVertexShaderRHIParamRef)RHICmd->Shader, RHICmd->SamplerIndex, RHICmd->Sampler); break;
-					case SF_Hull:		SetShaderSampler_Internal((FHullShaderRHIParamRef)RHICmd->Shader, RHICmd->SamplerIndex, RHICmd->Sampler); break;
-					case SF_Domain:		SetShaderSampler_Internal((FDomainShaderRHIParamRef)RHICmd->Shader, RHICmd->SamplerIndex, RHICmd->Sampler); break;
-					case SF_Geometry:	SetShaderSampler_Internal((FGeometryShaderRHIParamRef)RHICmd->Shader, RHICmd->SamplerIndex, RHICmd->Sampler); break;
-					case SF_Pixel:		SetShaderSampler_Internal((FPixelShaderRHIParamRef)RHICmd->Shader, RHICmd->SamplerIndex, RHICmd->Sampler); break;
-					case SF_Compute:	SetShaderSampler_Internal((FComputeShaderRHIParamRef)RHICmd->Shader, RHICmd->SamplerIndex, RHICmd->Sampler); break;
-					default: check(0); break;
-					}
+				case SF_Vertex:		SetShaderSampler_Internal((FVertexShaderRHIParamRef)RHICmd->Shader, RHICmd->SamplerIndex, RHICmd->Sampler); break;
+				case SF_Hull:		SetShaderSampler_Internal((FHullShaderRHIParamRef)RHICmd->Shader, RHICmd->SamplerIndex, RHICmd->Sampler); break;
+				case SF_Domain:		SetShaderSampler_Internal((FDomainShaderRHIParamRef)RHICmd->Shader, RHICmd->SamplerIndex, RHICmd->Sampler); break;
+				case SF_Geometry:	SetShaderSampler_Internal((FGeometryShaderRHIParamRef)RHICmd->Shader, RHICmd->SamplerIndex, RHICmd->Sampler); break;
+				case SF_Pixel:		SetShaderSampler_Internal((FPixelShaderRHIParamRef)RHICmd->Shader, RHICmd->SamplerIndex, RHICmd->Sampler); break;
+				case SF_Compute:	SetShaderSampler_Internal((FComputeShaderRHIParamRef)RHICmd->Shader, RHICmd->SamplerIndex, RHICmd->Sampler); break;
+				default: check(0); break;
 				}
 			}
 			break;
 		case ERCT_SetShaderTexture:
 			{
 				auto* RHICmd = Iter.NextCommand<FRHICommandSetShaderTexture>();
-				if (!bOnlyTestMemAccess)
+				switch (RHICmd->ShaderFrequency)
 				{
-					switch (RHICmd->ShaderFrequency)
-					{
-					case SF_Vertex:		SetShaderTexture_Internal((FVertexShaderRHIParamRef)RHICmd->Shader, RHICmd->TextureIndex, RHICmd->Texture); break;
-					case SF_Hull:		SetShaderTexture_Internal((FHullShaderRHIParamRef)RHICmd->Shader, RHICmd->TextureIndex, RHICmd->Texture); break;
-					case SF_Domain:		SetShaderTexture_Internal((FDomainShaderRHIParamRef)RHICmd->Shader, RHICmd->TextureIndex, RHICmd->Texture); break;
-					case SF_Geometry:	SetShaderTexture_Internal((FGeometryShaderRHIParamRef)RHICmd->Shader, RHICmd->TextureIndex, RHICmd->Texture); break;
-					case SF_Pixel:		SetShaderTexture_Internal((FPixelShaderRHIParamRef)RHICmd->Shader, RHICmd->TextureIndex, RHICmd->Texture); break;
-					case SF_Compute:	SetShaderTexture_Internal((FComputeShaderRHIParamRef)RHICmd->Shader, RHICmd->TextureIndex, RHICmd->Texture); break;
-					default: check(0); break;
-					}
+				case SF_Vertex:		SetShaderTexture_Internal((FVertexShaderRHIParamRef)RHICmd->Shader, RHICmd->TextureIndex, RHICmd->Texture); break;
+				case SF_Hull:		SetShaderTexture_Internal((FHullShaderRHIParamRef)RHICmd->Shader, RHICmd->TextureIndex, RHICmd->Texture); break;
+				case SF_Domain:		SetShaderTexture_Internal((FDomainShaderRHIParamRef)RHICmd->Shader, RHICmd->TextureIndex, RHICmd->Texture); break;
+				case SF_Geometry:	SetShaderTexture_Internal((FGeometryShaderRHIParamRef)RHICmd->Shader, RHICmd->TextureIndex, RHICmd->Texture); break;
+				case SF_Pixel:		SetShaderTexture_Internal((FPixelShaderRHIParamRef)RHICmd->Shader, RHICmd->TextureIndex, RHICmd->Texture); break;
+				case SF_Compute:	SetShaderTexture_Internal((FComputeShaderRHIParamRef)RHICmd->Shader, RHICmd->TextureIndex, RHICmd->Texture); break;
+				default: check(0); break;
 				}
 			}
 			break;
 		case ERCT_SetShaderResourceViewParameter:
 			{
 				auto* RHICmd = Iter.NextCommand<FRHICommandSetShaderResourceViewParameter>();
-				if (!bOnlyTestMemAccess)
+				switch (RHICmd->ShaderFrequency)
 				{
-					switch (RHICmd->ShaderFrequency)
-					{
-					case SF_Vertex:		SetShaderResourceViewParameter_Internal((FVertexShaderRHIParamRef)RHICmd->Shader, RHICmd->SamplerIndex, RHICmd->SRV); break;
-					case SF_Hull:		SetShaderResourceViewParameter_Internal((FHullShaderRHIParamRef)RHICmd->Shader, RHICmd->SamplerIndex, RHICmd->SRV); break;
-					case SF_Domain:		SetShaderResourceViewParameter_Internal((FDomainShaderRHIParamRef)RHICmd->Shader, RHICmd->SamplerIndex, RHICmd->SRV); break;
-					case SF_Geometry:	SetShaderResourceViewParameter_Internal((FGeometryShaderRHIParamRef)RHICmd->Shader, RHICmd->SamplerIndex, RHICmd->SRV); break;
-					case SF_Pixel:		SetShaderResourceViewParameter_Internal((FPixelShaderRHIParamRef)RHICmd->Shader, RHICmd->SamplerIndex, RHICmd->SRV); break;
-					case SF_Compute:	SetShaderResourceViewParameter_Internal((FComputeShaderRHIParamRef)RHICmd->Shader, RHICmd->SamplerIndex, RHICmd->SRV); break;
-					default: check(0); break;
-					}
+				case SF_Vertex:		SetShaderResourceViewParameter_Internal((FVertexShaderRHIParamRef)RHICmd->Shader, RHICmd->SamplerIndex, RHICmd->SRV); break;
+				case SF_Hull:		SetShaderResourceViewParameter_Internal((FHullShaderRHIParamRef)RHICmd->Shader, RHICmd->SamplerIndex, RHICmd->SRV); break;
+				case SF_Domain:		SetShaderResourceViewParameter_Internal((FDomainShaderRHIParamRef)RHICmd->Shader, RHICmd->SamplerIndex, RHICmd->SRV); break;
+				case SF_Geometry:	SetShaderResourceViewParameter_Internal((FGeometryShaderRHIParamRef)RHICmd->Shader, RHICmd->SamplerIndex, RHICmd->SRV); break;
+				case SF_Pixel:		SetShaderResourceViewParameter_Internal((FPixelShaderRHIParamRef)RHICmd->Shader, RHICmd->SamplerIndex, RHICmd->SRV); break;
+				case SF_Compute:	SetShaderResourceViewParameter_Internal((FComputeShaderRHIParamRef)RHICmd->Shader, RHICmd->SamplerIndex, RHICmd->SRV); break;
+				default: check(0); break;
 				}
 			}
 			break;
 		case ERCT_SetUAVParameter:
 			{
 				auto* RHICmd = Iter.NextCommand<FRHICommandSetUAVParameter>();
-				if (!bOnlyTestMemAccess)
+				switch (RHICmd->ShaderFrequency)
 				{
-					switch (RHICmd->ShaderFrequency)
-					{
-					case SF_Compute:	SetUAVParameter_Internal((FComputeShaderRHIParamRef)RHICmd->Shader, RHICmd->UAVIndex, RHICmd->UAV); break;
-					default: check(0); break;
-					}
+				case SF_Compute:	SetUAVParameter_Internal((FComputeShaderRHIParamRef)RHICmd->Shader, RHICmd->UAVIndex, RHICmd->UAV); break;
+				default: check(0); break;
 				}
 			}
 			break;
 		case ERCT_SetUAVParameter_IntialCount:
 			{
 				auto* RHICmd = Iter.NextCommand<FRHICommandSetUAVParameter_IntialCount>();
-				if (!bOnlyTestMemAccess)
+				switch (RHICmd->ShaderFrequency)
 				{
-					switch (RHICmd->ShaderFrequency)
-					{
-					case SF_Compute:	SetUAVParameter_Internal((FComputeShaderRHIParamRef)RHICmd->Shader, RHICmd->UAVIndex, RHICmd->UAV, RHICmd->InitialCount); break;
-					default: check(0); break;
-					}
+				case SF_Compute:	SetUAVParameter_Internal((FComputeShaderRHIParamRef)RHICmd->Shader, RHICmd->UAVIndex, RHICmd->UAV, RHICmd->InitialCount); break;
+				default: check(0); break;
 				}
 			}
 			break;
 		case ERCT_DrawPrimitive:
 			{
 				auto* RHICmd = Iter.NextCommand<FRHICommandDrawPrimitive>();
-				if (!bOnlyTestMemAccess)
-				{
-					DrawPrimitive_Internal(RHICmd->PrimitiveType, RHICmd->BaseVertexIndex, RHICmd->NumPrimitives, RHICmd->NumInstances);
-				}
+				DrawPrimitive_Internal(RHICmd->PrimitiveType, RHICmd->BaseVertexIndex, RHICmd->NumPrimitives, RHICmd->NumInstances);
 			}
 			break;
 		case ERCT_DrawIndexedPrimitive:
 			{
 				auto* RHICmd = Iter.NextCommand<FRHICommandDrawIndexedPrimitive>();
-				if (!bOnlyTestMemAccess)
-				{
-					DrawIndexedPrimitive_Internal(RHICmd->IndexBuffer, RHICmd->PrimitiveType, RHICmd->BaseVertexIndex, RHICmd->MinIndex, RHICmd->NumVertices, RHICmd->StartIndex, RHICmd->NumPrimitives, RHICmd->NumInstances);
-				}
+				DrawIndexedPrimitive_Internal(RHICmd->IndexBuffer, RHICmd->PrimitiveType, RHICmd->BaseVertexIndex, RHICmd->MinIndex, RHICmd->NumVertices, RHICmd->StartIndex, RHICmd->NumPrimitives, RHICmd->NumInstances);
 			}
 			break;
 		case ERCT_SetBoundShaderState:
 			{
 				auto* RHICmd = Iter.NextCommand<FRHICommandSetBoundShaderState>();
-				if (!bOnlyTestMemAccess)
-				{
-					SetBoundShaderState_Internal(RHICmd->BoundShaderState);
-				}
+				SetBoundShaderState_Internal(RHICmd->BoundShaderState);
 			}
 			break;
 		case ERCT_SetBlendState:
 			{
 				auto* RHICmd = Iter.NextCommand<FRHICommandSetBlendState>();
-				if (!bOnlyTestMemAccess)
-				{
-					SetBlendState_Internal(RHICmd->State, RHICmd->BlendFactor);
-				}
+				SetBlendState_Internal(RHICmd->State, RHICmd->BlendFactor);
 			}
 			break;
 		case ERCT_SetStreamSource:
 			{
 				auto* RHICmd = Iter.NextCommand<FRHICommandSetStreamSource>();
-				if (!bOnlyTestMemAccess)
-				{
-					SetStreamSource_Internal(RHICmd->StreamIndex, RHICmd->VertexBuffer, RHICmd->Stride, RHICmd->Offset);
-				}
+				SetStreamSource_Internal(RHICmd->StreamIndex, RHICmd->VertexBuffer, RHICmd->Stride, RHICmd->Offset);
 			}
 			break;
 		case ERCT_SetDepthStencilState:
 			{
 				auto* RHICmd = Iter.NextCommand<FRHICommandSetDepthStencilState>();
-				if (!bOnlyTestMemAccess)
-				{
-					SetDepthStencilState_Internal(RHICmd->State, RHICmd->StencilRef);
-				}
+				SetDepthStencilState_Internal(RHICmd->State, RHICmd->StencilRef);
 			}
 			break;
 		case ERCT_SetViewport:
 			{
 				auto* RHICmd = Iter.NextCommand<FRHICommandSetViewport>();
-				if (!bOnlyTestMemAccess)
-				{
-					SetViewport_Internal(RHICmd->MinX, RHICmd->MinY, RHICmd->MinZ, RHICmd->MaxX, RHICmd->MaxY, RHICmd->MaxZ);
-				}
+				SetViewport_Internal(RHICmd->MinX, RHICmd->MinY, RHICmd->MinZ, RHICmd->MaxX, RHICmd->MaxY, RHICmd->MaxZ);
 			}
 			break;
 		case ERCT_SetScissorRect:
 			{
 				auto* RHICmd = Iter.NextCommand<FRHICommandSetScissorRect>();
-				if (!bOnlyTestMemAccess)
-				{
-					SetScissorRect_Internal(RHICmd->bEnable, RHICmd->MinX, RHICmd->MinY, RHICmd->MaxX, RHICmd->MaxY);
-				}
+				SetScissorRect_Internal(RHICmd->bEnable, RHICmd->MinX, RHICmd->MinY, RHICmd->MaxX, RHICmd->MaxY);
 			}
 			break;
 		case ERCT_SetRenderTargets:
 			{
 				auto* RHICmd = Iter.NextCommand<FRHICommandSetRenderTargets>();
-				if (!bOnlyTestMemAccess)
-				{
-					SetRenderTargets_Internal(
-						RHICmd->NewNumSimultaneousRenderTargets,
-						RHICmd->NewRenderTargetsRHI,
-						RHICmd->NewDepthStencilTargetRHI,
-						RHICmd->NewNumUAVs,
-						RHICmd->UAVs);
-				}
+				SetRenderTargets_Internal(
+					RHICmd->NewNumSimultaneousRenderTargets,
+					RHICmd->NewRenderTargetsRHI,
+					RHICmd->NewDepthStencilTargetRHI,
+					RHICmd->NewNumUAVs,
+					RHICmd->UAVs);
 			}
 			break;
 		case ERCT_EndDrawPrimitiveUP:
 			{
 				auto* RHICmd = Iter.NextCommand<FRHICommandEndDrawPrimitiveUP>();
-				if (!bOnlyTestMemAccess)
-				{
-					void* Buffer = NULL;
-					BeginDrawPrimitiveUP_Internal(RHICmd->PrimitiveType, RHICmd->NumPrimitives, RHICmd->NumVertices, RHICmd->VertexDataStride, Buffer);
-					FMemory::Memcpy(Buffer, RHICmd->OutVertexData, RHICmd->NumVertices * RHICmd->VertexDataStride);
-					EndDrawPrimitiveUP_Internal();
-				}
+				void* Buffer = NULL;
+				BeginDrawPrimitiveUP_Internal(RHICmd->PrimitiveType, RHICmd->NumPrimitives, RHICmd->NumVertices, RHICmd->VertexDataStride, Buffer);
+				FMemory::Memcpy(Buffer, RHICmd->OutVertexData, RHICmd->NumVertices * RHICmd->VertexDataStride);
+				EndDrawPrimitiveUP_Internal();
 				FMemory::Free(RHICmd->OutVertexData);
 			}
 			break;
 		case ERCT_EndDrawIndexedPrimitiveUP:
 			{
 				auto* RHICmd = Iter.NextCommand<FRHICommandEndDrawIndexedPrimitiveUP>();
-				if (!bOnlyTestMemAccess)
-				{
-					void* VertexBuffer = NULL;
-					void* IndexBuffer = NULL;
-					BeginDrawIndexedPrimitiveUP_Internal(
-						RHICmd->PrimitiveType,
-						RHICmd->NumPrimitives,
-						RHICmd->NumVertices,
-						RHICmd->VertexDataStride,
-						VertexBuffer,
-						RHICmd->MinVertexIndex,
-						RHICmd->NumIndices,
-						RHICmd->IndexDataStride,
-						IndexBuffer);
-					FMemory::Memcpy(VertexBuffer, RHICmd->OutVertexData, RHICmd->NumVertices * RHICmd->VertexDataStride);
-					FMemory::Memcpy(IndexBuffer, RHICmd->OutIndexData, RHICmd->NumIndices * RHICmd->IndexDataStride);
-					EndDrawIndexedPrimitiveUP_Internal();
-
-				}
+				void* VertexBuffer = NULL;
+				void* IndexBuffer = NULL;
+				BeginDrawIndexedPrimitiveUP_Internal(
+					RHICmd->PrimitiveType,
+					RHICmd->NumPrimitives,
+					RHICmd->NumVertices,
+					RHICmd->VertexDataStride,
+					VertexBuffer,
+					RHICmd->MinVertexIndex,
+					RHICmd->NumIndices,
+					RHICmd->IndexDataStride,
+					IndexBuffer);
+				FMemory::Memcpy(VertexBuffer, RHICmd->OutVertexData, RHICmd->NumVertices * RHICmd->VertexDataStride);
+				FMemory::Memcpy(IndexBuffer, RHICmd->OutIndexData, RHICmd->NumIndices * RHICmd->IndexDataStride);
+				EndDrawIndexedPrimitiveUP_Internal();
 				FMemory::Free(RHICmd->OutVertexData);
 				FMemory::Free(RHICmd->OutIndexData);
 			}
@@ -331,199 +254,145 @@ void FRHICommandListExecutor::ExecuteList(FRHICommandList& CmdList)
 		case ERCT_BuildLocalBoundShaderState:
 			{
 				auto* RHICmd = Iter.NextCommand<FRHICommandBuildLocalBoundShaderState>();
-				if (!bOnlyTestMemAccess)
+				check(RHICmd->WorkArea.CheckCmdList == &CmdList && RHICmd->WorkArea.UID == CmdList.UID // this BSS was not built for this particular commandlist
+					&& !IsValidRef(RHICmd->WorkArea.BSS)); // should not already have been created
+				if (RHICmd->WorkArea.UseCount)
 				{
-					check(RHICmd->WorkArea.CheckCmdList == &CmdList && RHICmd->WorkArea.UID == CmdList.UID // this BSS was not built for this particular commandlist
-						&& !IsValidRef(RHICmd->WorkArea.BSS)); // should not already have been created
-					if (RHICmd->WorkArea.UseCount)
-					{
-						RHICmd->WorkArea.BSS = CreateBoundShaderState_Internal(RHICmd->WorkArea.Args.VertexDeclarationRHI, RHICmd->WorkArea.Args.VertexShaderRHI, RHICmd->WorkArea.Args.HullShaderRHI, RHICmd->WorkArea.Args.DomainShaderRHI, RHICmd->WorkArea.Args.PixelShaderRHI, RHICmd->WorkArea.Args.GeometryShaderRHI);
-					}
+					RHICmd->WorkArea.BSS = CreateBoundShaderState_Internal(RHICmd->WorkArea.Args.VertexDeclarationRHI, RHICmd->WorkArea.Args.VertexShaderRHI, RHICmd->WorkArea.Args.HullShaderRHI, RHICmd->WorkArea.Args.DomainShaderRHI, RHICmd->WorkArea.Args.PixelShaderRHI, RHICmd->WorkArea.Args.GeometryShaderRHI);
 				}
 			}
 			break;
 		case ERCT_SetLocalBoundShaderState:
 			{
 				auto* RHICmd = Iter.NextCommand<FRHICommandSetLocalBoundShaderState>();
-				if (!bOnlyTestMemAccess)
-				{
-					check(RHICmd->LocalBoundShaderState.WorkArea->CheckCmdList == &CmdList && RHICmd->LocalBoundShaderState.WorkArea->UID == CmdList.UID // this BSS was not built for this particular commandlist
-						&& RHICmd->LocalBoundShaderState.WorkArea->UseCount > 0 && IsValidRef(RHICmd->LocalBoundShaderState.WorkArea->BSS)); // this should have been created and should have uses outstanding
+				check(RHICmd->LocalBoundShaderState.WorkArea->CheckCmdList == &CmdList && RHICmd->LocalBoundShaderState.WorkArea->UID == CmdList.UID // this BSS was not built for this particular commandlist
+					&& RHICmd->LocalBoundShaderState.WorkArea->UseCount > 0 && IsValidRef(RHICmd->LocalBoundShaderState.WorkArea->BSS)); // this should have been created and should have uses outstanding
 
-					SetBoundShaderState_Internal(RHICmd->LocalBoundShaderState.WorkArea->BSS);
+				SetBoundShaderState_Internal(RHICmd->LocalBoundShaderState.WorkArea->BSS);
 
-					if (--RHICmd->LocalBoundShaderState.WorkArea->UseCount == 0)
-					{
-						RHICmd->LocalBoundShaderState.WorkArea->BSS = nullptr; // this also releases the ref, which wouldn't otherwise happen because RHICommands are unsafe and don't do constructors or destructors
-					}
+				if (--RHICmd->LocalBoundShaderState.WorkArea->UseCount == 0)
+				{
+					RHICmd->LocalBoundShaderState.WorkArea->BSS = nullptr; // this also releases the ref, which wouldn't otherwise happen because RHICommands are unsafe and don't do constructors or destructors
 				}
 			}
 			break;
 		case ERCT_SetGlobalBoundShaderState:
 			{
 				auto* RHICmd = Iter.NextCommand<FRHICommandSetGlobalBoundShaderState>();
-				if (!bOnlyTestMemAccess)
-				{
-					CmdList.SetGlobalBoundShaderState_InternalPtr(RHICmd->GlobalBoundShaderState);
-				}
+				SetGlobalBoundShaderState_InternalPtr(RHICmd->GlobalBoundShaderState);
 			}
 			break;
 
 		case ERCT_SetComputeShader:
 			{
 				auto* RHICmd = Iter.NextCommand<FRHICommandSetComputeShader>();
-				if (!bOnlyTestMemAccess)
-				{
-					SetComputeShader_Internal(RHICmd->ComputeShader);
-				}
+				SetComputeShader_Internal(RHICmd->ComputeShader);
 			}
 			break;
 		case ERCT_DispatchComputeShader:
 			{
 				auto* RHICmd = Iter.NextCommand<FRHICommandDispatchComputeShader>();
-				if (!bOnlyTestMemAccess)
-				{
-					DispatchComputeShader_Internal(RHICmd->ThreadGroupCountX, RHICmd->ThreadGroupCountY, RHICmd->ThreadGroupCountZ);
-				}
+				DispatchComputeShader_Internal(RHICmd->ThreadGroupCountX, RHICmd->ThreadGroupCountY, RHICmd->ThreadGroupCountZ);
 			}
 			break;
 		case ERCT_DispatchIndirectComputeShader:
 			{
 				auto* RHICmd = Iter.NextCommand<FRHICommandDispatchIndirectComputeShader>();
-				if (!bOnlyTestMemAccess)
-				{
-					DispatchIndirectComputeShader_Internal(RHICmd->ArgumentBuffer, RHICmd->ArgumentOffset);
-				}
+				DispatchIndirectComputeShader_Internal(RHICmd->ArgumentBuffer, RHICmd->ArgumentOffset);
 			}
 			break;
 		case ERCT_AutomaticCacheFlushAfterComputeShader:
 			{
 				auto* RHICmd = Iter.NextCommand<FRHICommandAutomaticCacheFlushAfterComputeShader>();
-				if (!bOnlyTestMemAccess)
-				{
-					AutomaticCacheFlushAfterComputeShader_Internal(RHICmd->bEnable);
-				}
+				AutomaticCacheFlushAfterComputeShader_Internal(RHICmd->bEnable);
 			}
 			break;
 		case ERCT_FlushComputeShaderCache:
 			{
 				auto* RHICmd = Iter.NextCommand<FRHICommandFlushComputeShaderCache>();
-				if (!bOnlyTestMemAccess)
-				{
-					FlushComputeShaderCache_Internal();
-				}
+				FlushComputeShaderCache_Internal();
 			}
 			break;
 		case ERCT_DrawPrimitiveIndirect:
 			{
 				auto* RHICmd = Iter.NextCommand<FRHICommandDrawPrimitiveIndirect>();
-				if (!bOnlyTestMemAccess)
-				{
-					DrawPrimitiveIndirect_Internal(RHICmd->PrimitiveType, RHICmd->ArgumentBuffer, RHICmd->ArgumentOffset);
-				}
+				DrawPrimitiveIndirect_Internal(RHICmd->PrimitiveType, RHICmd->ArgumentBuffer, RHICmd->ArgumentOffset);
 			}
 			break;
 		case ERCT_DrawIndexedIndirect:
 			{
 				auto* RHICmd = Iter.NextCommand<FRHICommandDrawIndexedIndirect>();
-				if (!bOnlyTestMemAccess)
-				{
-					DrawIndexedIndirect_Internal(RHICmd->IndexBufferRHI, RHICmd->PrimitiveType, RHICmd->ArgumentsBufferRHI, RHICmd->DrawArgumentsIndex, RHICmd->NumInstances);
-				}
+				DrawIndexedIndirect_Internal(RHICmd->IndexBufferRHI, RHICmd->PrimitiveType, RHICmd->ArgumentsBufferRHI, RHICmd->DrawArgumentsIndex, RHICmd->NumInstances);
 			}
 			break;
 		case ERCT_DrawIndexedPrimitiveIndirect:
 			{
 				auto* RHICmd = Iter.NextCommand<FRHICommandDrawIndexedPrimitiveIndirect>();
-				if (!bOnlyTestMemAccess)
-				{
-					DrawIndexedPrimitiveIndirect_Internal(RHICmd->PrimitiveType, RHICmd->IndexBuffer, RHICmd->ArgumentsBuffer, RHICmd->ArgumentOffset);
-				}
+				DrawIndexedPrimitiveIndirect_Internal(RHICmd->PrimitiveType, RHICmd->IndexBuffer, RHICmd->ArgumentsBuffer, RHICmd->ArgumentOffset);
 			}
 			break;
 		case ERCT_EnableDepthBoundsTest:
 			{
 				auto* RHICmd = Iter.NextCommand<FRHICommandEnableDepthBoundsTest>();
-				if (!bOnlyTestMemAccess)
-				{
-					EnableDepthBoundsTest_Internal(RHICmd->bEnable, RHICmd->MinDepth, RHICmd->MaxDepth);
-				}
+				EnableDepthBoundsTest_Internal(RHICmd->bEnable, RHICmd->MinDepth, RHICmd->MaxDepth);
 			}
 			break;
 		case ERCT_ClearUAV:
 			{
 				auto* RHICmd = Iter.NextCommand<FRHICommandClearUAV>();
-				if (!bOnlyTestMemAccess)
-				{
-					ClearUAV_Internal(RHICmd->UnorderedAccessViewRHI, RHICmd->Values);
-				}
+				ClearUAV_Internal(RHICmd->UnorderedAccessViewRHI, RHICmd->Values);
 			}
 			break;
 		case ERCT_CopyToResolveTarget:
 			{
 				auto* RHICmd = Iter.NextCommand<FRHICommandCopyToResolveTarget>();
-				if (!bOnlyTestMemAccess)
-				{
-					CopyToResolveTarget_Internal(RHICmd->SourceTexture, RHICmd->DestTexture, RHICmd->bKeepOriginalSurface, RHICmd->ResolveParams);
-				}
+				CopyToResolveTarget_Internal(RHICmd->SourceTexture, RHICmd->DestTexture, RHICmd->bKeepOriginalSurface, RHICmd->ResolveParams);
 			}
 			break;
 		case ERCT_Clear:
 			{
 				auto* RHICmd = Iter.NextCommand<FRHICommandClear>();
-				if (!bOnlyTestMemAccess)
-				{
-					Clear_Internal(RHICmd->bClearColor, RHICmd->Color, RHICmd->bClearDepth, RHICmd->Depth, RHICmd->bClearStencil, RHICmd->Stencil, RHICmd->ExcludeRect);
-				}
+				Clear_Internal(RHICmd->bClearColor, RHICmd->Color, RHICmd->bClearDepth, RHICmd->Depth, RHICmd->bClearStencil, RHICmd->Stencil, RHICmd->ExcludeRect);
 			}
 			break;
 		case ERCT_ClearMRT:
 			{
 				auto* RHICmd = Iter.NextCommand<FRHICommandClearMRT>();
-				if (!bOnlyTestMemAccess)
-				{
-					ClearMRT_Internal(RHICmd->bClearColor, RHICmd->NumClearColors, RHICmd->ColorArray, RHICmd->bClearDepth, RHICmd->Depth, RHICmd->bClearStencil, RHICmd->Stencil, RHICmd->ExcludeRect);
-				}
+				ClearMRT_Internal(RHICmd->bClearColor, RHICmd->NumClearColors, RHICmd->ColorArray, RHICmd->bClearDepth, RHICmd->Depth, RHICmd->bClearStencil, RHICmd->Stencil, RHICmd->ExcludeRect);
 			}
 			break;
 		case ERCT_BuildLocalUniformBuffer:
 			{
 				auto* RHICmd = Iter.NextCommand<FRHICommandBuildLocalUniformBuffer>();
-				if (!bOnlyTestMemAccess)
+				check(RHICmd->WorkArea.CheckCmdList == &CmdList && RHICmd->WorkArea.UID == CmdList.UID // this uniform buffer was not built for this particular commandlist
+					&& !IsValidRef(RHICmd->WorkArea.UniformBuffer) && RHICmd->WorkArea.Layout && RHICmd->WorkArea.Contents); // should not already have been created
+				if (RHICmd->WorkArea.UseCount)
 				{
-					check(RHICmd->WorkArea.CheckCmdList == &CmdList && RHICmd->WorkArea.UID == CmdList.UID // this uniform buffer was not built for this particular commandlist
-						&& !IsValidRef(RHICmd->WorkArea.UniformBuffer) && RHICmd->WorkArea.Layout && RHICmd->WorkArea.Contents); // should not already have been created
-					if (RHICmd->WorkArea.UseCount)
-					{
-						RHICmd->WorkArea.UniformBuffer = RHICreateUniformBuffer(RHICmd->WorkArea.Contents, *RHICmd->WorkArea.Layout, UniformBuffer_SingleFrame);
-					}
-					FMemory::Free(RHICmd->WorkArea.Contents);
-					RHICmd->WorkArea.Contents = nullptr;
+					RHICmd->WorkArea.UniformBuffer = RHICreateUniformBuffer(RHICmd->WorkArea.Contents, *RHICmd->WorkArea.Layout, UniformBuffer_SingleFrame);
 				}
+				FMemory::Free(RHICmd->WorkArea.Contents);
+				RHICmd->WorkArea.Contents = nullptr;
 			}
 			break;
 		case ERCT_SetLocalUniformBuffer:
 			{
 				auto* RHICmd = Iter.NextCommand<FRHICommandSetLocalUniformBuffer>();
-				if (!bOnlyTestMemAccess)
+				check(RHICmd->LocalUniformBuffer.WorkArea->CheckCmdList == &CmdList && RHICmd->LocalUniformBuffer.WorkArea->UID == CmdList.UID // this BSS was not built for this particular commandlist
+					&& RHICmd->LocalUniformBuffer.WorkArea->UseCount > 0 && IsValidRef(RHICmd->LocalUniformBuffer.WorkArea->UniformBuffer)); // this should have been created and should have uses outstanding
+				switch (RHICmd->ShaderFrequency)
+				{
+					case SF_Vertex:		SetShaderUniformBuffer_Internal((FVertexShaderRHIParamRef)RHICmd->Shader, RHICmd->BaseIndex, RHICmd->LocalUniformBuffer.WorkArea->UniformBuffer); break;
+					case SF_Hull:		SetShaderUniformBuffer_Internal((FHullShaderRHIParamRef)RHICmd->Shader, RHICmd->BaseIndex, RHICmd->LocalUniformBuffer.WorkArea->UniformBuffer); break;
+					case SF_Domain:		SetShaderUniformBuffer_Internal((FDomainShaderRHIParamRef)RHICmd->Shader, RHICmd->BaseIndex, RHICmd->LocalUniformBuffer.WorkArea->UniformBuffer); break;
+					case SF_Geometry:	SetShaderUniformBuffer_Internal((FGeometryShaderRHIParamRef)RHICmd->Shader, RHICmd->BaseIndex, RHICmd->LocalUniformBuffer.WorkArea->UniformBuffer); break;
+					case SF_Pixel:		SetShaderUniformBuffer_Internal((FPixelShaderRHIParamRef)RHICmd->Shader, RHICmd->BaseIndex, RHICmd->LocalUniformBuffer.WorkArea->UniformBuffer); break;
+					case SF_Compute:	SetShaderUniformBuffer_Internal((FComputeShaderRHIParamRef)RHICmd->Shader, RHICmd->BaseIndex, RHICmd->LocalUniformBuffer.WorkArea->UniformBuffer); break;
+					default: check(0); break;
+				}
+				if (--RHICmd->LocalUniformBuffer.WorkArea->UseCount == 0)
 				{
-					check(RHICmd->LocalUniformBuffer.WorkArea->CheckCmdList == &CmdList && RHICmd->LocalUniformBuffer.WorkArea->UID == CmdList.UID // this BSS was not built for this particular commandlist
-						&& RHICmd->LocalUniformBuffer.WorkArea->UseCount > 0 && IsValidRef(RHICmd->LocalUniformBuffer.WorkArea->UniformBuffer)); // this should have been created and should have uses outstanding
-					switch (RHICmd->ShaderFrequency)
-					{
-						case SF_Vertex:		SetShaderUniformBuffer_Internal((FVertexShaderRHIParamRef)RHICmd->Shader, RHICmd->BaseIndex, RHICmd->LocalUniformBuffer.WorkArea->UniformBuffer); break;
-						case SF_Hull:		SetShaderUniformBuffer_Internal((FHullShaderRHIParamRef)RHICmd->Shader, RHICmd->BaseIndex, RHICmd->LocalUniformBuffer.WorkArea->UniformBuffer); break;
-						case SF_Domain:		SetShaderUniformBuffer_Internal((FDomainShaderRHIParamRef)RHICmd->Shader, RHICmd->BaseIndex, RHICmd->LocalUniformBuffer.WorkArea->UniformBuffer); break;
-						case SF_Geometry:	SetShaderUniformBuffer_Internal((FGeometryShaderRHIParamRef)RHICmd->Shader, RHICmd->BaseIndex, RHICmd->LocalUniformBuffer.WorkArea->UniformBuffer); break;
-						case SF_Pixel:		SetShaderUniformBuffer_Internal((FPixelShaderRHIParamRef)RHICmd->Shader, RHICmd->BaseIndex, RHICmd->LocalUniformBuffer.WorkArea->UniformBuffer); break;
-						case SF_Compute:	SetShaderUniformBuffer_Internal((FComputeShaderRHIParamRef)RHICmd->Shader, RHICmd->BaseIndex, RHICmd->LocalUniformBuffer.WorkArea->UniformBuffer); break;
-						default: check(0); break;
-					}
-					if (--RHICmd->LocalUniformBuffer.WorkArea->UseCount == 0)
-					{
-						RHICmd->LocalUniformBuffer.WorkArea->UniformBuffer = nullptr; // this also releases the ref, which wouldn't otherwise happen because RHICommands are unsafe and don't do constructors or destructors
-						RHICmd->LocalUniformBuffer.WorkArea->Layout = nullptr;
-					}
+					RHICmd->LocalUniformBuffer.WorkArea->UniformBuffer = nullptr; // this also releases the ref, which wouldn't otherwise happen because RHICommands are unsafe and don't do constructors or destructors
+					RHICmd->LocalUniformBuffer.WorkArea->Layout = nullptr;
 				}
 			}
 			break;
@@ -536,55 +405,58 @@ void FRHICommandListExecutor::ExecuteList(FRHICommandList& CmdList)
 	Iter.CheckNumCommands();
 }
 
+void FRHICommandListExecutor::ExecuteInner(FRHICommandList& CmdList)
+{
+	static IConsoleVariable* RHIExecCVar = IConsoleManager::Get().FindConsoleVariable(TEXT("r.RHIExec"));
+	CmdList.bExecuting = true;
+	if (RHIExecCVar && RHIExecCVar->GetInt() != 0)
+	{
+		RHIExecuteCommandList(&CmdList);
+	}
+	else
+	{
+		ExecuteListGenericRHI(CmdList);
+	}
+	CmdList.Reset();
+}
 
 void FRHICommandListExecutor::ExecuteList(FRHICommandList& CmdList)
 {
-	CmdList.bExecuting = true;
-	SCOPE_CYCLE_COUNTER(STAT_RHICmdListExecuteTime);
+	check(IsInRenderingThread() && &CmdList != &GetImmediateCommandList());
 
-	check(IsInRenderingThread());
+	if (!GetImmediateCommandList().bExecuting) // don't flush if this is a recursive call and we are already executing the immediate command list
+	{
+		GetImmediateCommandList().Flush();
+	}
 
-	INC_MEMORY_STAT_BY(STAT_RHICmdListMemory, CmdList.GetUsedMemory());
+	INC_MEMORY_STAT_BY(STAT_NonImmedCmdListMemory, CmdList.GetUsedMemory());
+	INC_DWORD_STAT_BY(STAT_NonImmedCmdListCount, CmdList.NumCommands);
 
-	static IConsoleVariable* RHICmdListCVar = IConsoleManager::Get().FindConsoleVariable(TEXT("r.RHISkip"));
-	bool bSkipRHICmdList = (RHICmdListCVar && RHICmdListCVar->GetInt() != 0);
+	SCOPE_CYCLE_COUNTER(STAT_NonImmedCmdListExecuteTime);
+	ExecuteInner(CmdList);
+}
+
+void FRHICommandListExecutor::ExecuteList(FRHICommandListImmediate& CmdList)
+{
+	check(IsInRenderingThread() && &CmdList == &GetImmediateCommandList());
+
+	INC_MEMORY_STAT_BY(STAT_ImmedCmdListMemory, CmdList.GetUsedMemory());
+	INC_DWORD_STAT_BY(STAT_ImmedCmdListCount, CmdList.NumCommands);
 
-	if (!bSkipRHICmdList)
 	{
-		if (&CmdList != &GetImmediateCommandList() 
-			&& !GetImmediateCommandList().bExecuting) // don't flush if this is a recursive call and we are already executing the immediate command list
-		{
-			GetImmediateCommandList().Flush();
-		}
-		static IConsoleVariable* RHIExecCVar = IConsoleManager::Get().FindConsoleVariable(TEXT("r.RHIExec"));
-		if (RHIExecCVar && RHIExecCVar->GetInt() != 0)
-		{
-			RHIExecuteCommandList(&CmdList);
-		}
-		else
-		{
-			static IConsoleVariable* RHIMemCVar = IConsoleManager::Get().FindConsoleVariable(TEXT("r.RHIMem"));
-			bool bOnlySkipAroundMem = (RHIMemCVar && RHIMemCVar->GetInt() != 0);
-			if (bOnlySkipAroundMem)
-			{
-				ExecuteList<true>(CmdList);
-			}
-			else
-			{
-				ExecuteList<false>(CmdList);
-			}
-		}
+		SCOPE_CYCLE_COUNTER(STAT_ImmedCmdListExecuteTime);
+		ExecuteInner(CmdList);
 	}
-	INC_DWORD_STAT_BY(STAT_RHICmdListCount, CmdList.NumCommands);
-	CmdList.Reset();
+
 #if !(UE_BUILD_SHIPPING || UE_BUILD_TEST)
-	if (&CmdList == &GetImmediateCommandList() && GRHICommandList.OutstandingCmdListCount.GetValue() == 1)
+	if (GRHICommandList.OutstandingCmdListCount.GetValue() == 1)
 	{
 		LatchBypass();
 	}
 #endif
 }
 
+
 void FRHICommandListExecutor::LatchBypass()
 {
 #if !(UE_BUILD_SHIPPING || UE_BUILD_TEST)
@@ -669,10 +541,6 @@ void FRHICommandList::Reset()
 	NumCommands = 0;
 	MemManager.Reset();
 	UID = GRHICommandList.UIDCounter.Increment();
-	for (int32 Index = 0; Index < ERCT_NUM; Index++)
-	{
-		StateReducer[Index] = nullptr;
-	}
 }
 
 
diff --git a/Engine/Source/Runtime/RHI/Public/RHICommandList.h b/Engine/Source/Runtime/RHI/Public/RHICommandList.h
--- a/Engine/Source/Runtime/RHI/Public/RHICommandList.h
+++ b/Engine/Source/Runtime/RHI/Public/RHICommandList.h
@@ -7,6 +7,17 @@
 #pragma once
 #include "LockFreeFixedSizeAllocator.h"
 
+
+DECLARE_STATS_GROUP(TEXT("RHICmdList"), STATGROUP_RHICMDLIST, STATCAT_Advanced);
+
+DECLARE_CYCLE_STAT_EXTERN(TEXT("Nonimmed. Command List Execute"), STAT_NonImmedCmdListExecuteTime, STATGROUP_RHICMDLIST, );
+DECLARE_DWORD_COUNTER_STAT_EXTERN(TEXT("Nonimmed. Command List memory"), STAT_NonImmedCmdListMemory, STATGROUP_RHICMDLIST, RHI_API);
+DECLARE_DWORD_COUNTER_STAT_EXTERN(TEXT("Nonimmed. Command count"), STAT_NonImmedCmdListCount, STATGROUP_RHICMDLIST, RHI_API);
+
+DECLARE_CYCLE_STAT_EXTERN(TEXT("Immed. Command List Execute"), STAT_ImmedCmdListExecuteTime, STATGROUP_RHICMDLIST, );
+DECLARE_DWORD_COUNTER_STAT_EXTERN(TEXT("Immed. Command List memory"), STAT_ImmedCmdListMemory, STATGROUP_RHICMDLIST, RHI_API);
+DECLARE_DWORD_COUNTER_STAT_EXTERN(TEXT("Immed. Command count"), STAT_ImmedCmdListCount, STATGROUP_RHICMDLIST, RHI_API);
+
 enum ERHICommandType
 {
 	// Start the enum at non-zero to help debugging
@@ -60,6 +71,7 @@ enum ERHICommandType
 	ERCT_NUM,
 };
 
+
 struct FRHICommand
 {
 	// This is a union only to force alignment to SIZE_T
@@ -810,7 +822,8 @@ struct FRHICommandSetGlobalBoundShaderState : public FRHICommand
 	}
 };
 
-//---
+// there is a layering problem here, we use a function pointer to avoid making RHI depend on engine
+extern RHI_API void(*SetGlobalBoundShaderState_InternalPtr)(FGlobalBoundShaderState& BoundShaderState);
 
 struct FLocalUniformBufferWorkArea
 {
@@ -1006,30 +1019,6 @@ private:
 
 	FMemManager MemManager;
 
-	struct FDrawUpData
-	{
-
-		uint32 PrimitiveType;
-		uint32 NumPrimitives;
-		uint32 NumVertices;
-		uint32 VertexDataStride;
-		void* OutVertexData;
-		uint32 MinVertexIndex;
-		uint32 NumIndices;
-		uint32 IndexDataStride;
-		void* OutIndexData;
-
-		FDrawUpData()
-			: PrimitiveType(PT_Num)
-			, OutVertexData(nullptr)
-			, OutIndexData(nullptr)
-		{
-		}
-
-	};
-
-	FDrawUpData DrawUPData;
-
 public:
 	enum
 	{
@@ -1100,7 +1089,6 @@ public:
 	}
 
 	FORCEINLINE_DEBUGGABLE void BuildAndSetGlobalBoundShaderState(
-		void(*InSetGlobalBoundShaderState_InternalPtr)(FGlobalBoundShaderState& BoundShaderState),
 		FGlobalBoundShaderState& GlobalBoundShaderState,
 		FVertexDeclarationRHIParamRef VertexDeclarationRHI,
 		FShader* VertexShader,
@@ -1108,7 +1096,6 @@ public:
 		FShader* GeometryShader
 		)
 	{
-		SetGlobalBoundShaderState_InternalPtr = InSetGlobalBoundShaderState_InternalPtr; // total hack to work around module issues
 		if (!GlobalBoundShaderState.Get())
 		{
 			FGlobalBoundShaderStateWorkArea* NewGlobalBoundShaderState = new FGlobalBoundShaderStateWorkArea();
@@ -1135,7 +1122,7 @@ public:
 			);
 		if (Bypass())
 		{
-			InSetGlobalBoundShaderState_InternalPtr(GlobalBoundShaderState);
+			SetGlobalBoundShaderState_InternalPtr(GlobalBoundShaderState);
 			return;
 		}
 		auto* Cmd = AddCommand<FRHICommandSetGlobalBoundShaderState>();
@@ -1181,7 +1168,6 @@ public:
 		}
 		auto* Cmd = AddCommand<FRHICommandSetShaderUniformBuffer>();
 		Cmd->Set<TShaderRHIParamRef>(Shader, BaseIndex, UniformBuffer);
-		StateReduce<FRHICommandSetShaderUniformBuffer>(Cmd);
 	}
 
 	template <typename TShaderRHIParamRef>
@@ -1201,7 +1187,6 @@ public:
 		}
 		auto* Cmd = AddCommand<FRHICommandSetShaderParameter>();
 		Cmd->Set<TShaderRHIParamRef>(Shader, BufferIndex, BaseIndex, NumBytes, NewValue);
-		StateReduce<FRHICommandSetShaderParameter>(Cmd);
 	}
 
 	template <typename TShaderRHIParamRef>
@@ -1215,7 +1200,6 @@ public:
 
 		auto* Cmd = AddCommand<FRHICommandSetShaderTexture>();
 		Cmd->Set<TShaderRHIParamRef>(Shader, TextureIndex, Texture);
-		StateReduce<FRHICommandSetShaderTexture>(Cmd);
 	}
 
 	template <typename TShaderRHIParamRef>
@@ -1228,7 +1212,6 @@ public:
 		}
 		auto* Cmd = AddCommand<FRHICommandSetShaderResourceViewParameter>();
 		Cmd->Set<TShaderRHIParamRef>(Shader, SamplerIndex, SRV);
-		StateReduce<FRHICommandSetShaderResourceViewParameter>(Cmd);
 	}
 
 	template <typename TShaderRHIParamRef>
@@ -1242,7 +1225,6 @@ public:
 
 		auto* Cmd = AddCommand<FRHICommandSetShaderSampler>();
 		Cmd->Set<TShaderRHIParamRef>(Shader, SamplerIndex, State);
-		StateReduce<FRHICommandSetShaderSampler>(Cmd);
 	}
 
 	FORCEINLINE_DEBUGGABLE void SetUAVParameter(FComputeShaderRHIParamRef Shader, uint32 UAVIndex, FUnorderedAccessViewRHIParamRef UAV)
@@ -1254,7 +1236,6 @@ public:
 		}
 		auto* Cmd = AddCommand<FRHICommandSetUAVParameter>();
 		Cmd->Set<FComputeShaderRHIParamRef>(Shader, UAVIndex, UAV);
-		// StateReduce<FRHICommandSetUAVParameter>(Cmd); conflicts with next one
 	}
 
 	FORCEINLINE_DEBUGGABLE void SetUAVParameter(FComputeShaderRHIParamRef Shader, uint32 UAVIndex, FUnorderedAccessViewRHIParamRef UAV, uint32 InitialCount)
@@ -1266,7 +1247,6 @@ public:
 		}
 		auto* Cmd = AddCommand<FRHICommandSetUAVParameter_IntialCount>();
 		Cmd->Set<FComputeShaderRHIParamRef>(Shader, UAVIndex, UAV, InitialCount);
-		// StateReduce<FRHICommandSetUAVParameter_IntialCount>(Cmd); conflicts with previous one
 	}
 
 	FORCEINLINE_DEBUGGABLE void SetBoundShaderState(FBoundShaderStateRHIParamRef BoundShaderState)
@@ -1279,7 +1259,6 @@ public:
 
 		auto* Cmd = AddCommand<FRHICommandSetBoundShaderState>();
 		Cmd->Set(BoundShaderState);
-		StateReduce<FRHICommandSetBoundShaderState>(Cmd);
 	}
 
 	FORCEINLINE_DEBUGGABLE void SetRasterizerState(FRasterizerStateRHIParamRef State)
@@ -1292,7 +1271,6 @@ public:
 
 		auto* Cmd = AddCommand<FRHICommandSetRasterizerState>();
 		Cmd->Set(State);
-		StateReduce<FRHICommandSetRasterizerState>(Cmd);
 	}
 
 	FORCEINLINE_DEBUGGABLE void SetBlendState(FBlendStateRHIParamRef State, const FLinearColor& BlendFactor = FLinearColor::White)
@@ -1305,7 +1283,6 @@ public:
 
 		auto* Cmd = AddCommand<FRHICommandSetBlendState>();
 		Cmd->Set(State, BlendFactor);
-		StateReduce<FRHICommandSetBlendState>(Cmd);
 	}
 
 	FORCEINLINE_DEBUGGABLE void DrawPrimitive(uint32 PrimitiveType, uint32 BaseVertexIndex, uint32 NumPrimitives, uint32 NumInstances)
@@ -1342,7 +1319,6 @@ public:
 
 		auto* Cmd = AddCommand<FRHICommandSetStreamSource>();
 		Cmd->Set(StreamIndex, VertexBuffer, Stride, Offset);
-		StateReduce<FRHICommandSetStreamSource>(Cmd);
 	}
 
 	FORCEINLINE_DEBUGGABLE void SetDepthStencilState(FDepthStencilStateRHIParamRef NewStateRHI, uint32 StencilRef = 0)
@@ -1355,7 +1331,6 @@ public:
 
 		auto* Cmd = AddCommand<FRHICommandSetDepthStencilState>();
 		Cmd->Set(NewStateRHI, StencilRef);
-		StateReduce<FRHICommandSetDepthStencilState>(Cmd);
 	}
 
 	FORCEINLINE_DEBUGGABLE void SetViewport(uint32 MinX, uint32 MinY, float MinZ, uint32 MaxX, uint32 MaxY, float MaxZ)
@@ -1367,7 +1342,6 @@ public:
 		}
 		auto* Cmd = AddCommand<FRHICommandSetViewport>();
 		Cmd->Set(MinX, MinY, MinZ, MaxX, MaxY, MaxZ);
-		StateReduce<FRHICommandSetViewport>(Cmd);
 	}
 
 	FORCEINLINE_DEBUGGABLE void SetScissorRect(bool bEnable, uint32 MinX, uint32 MinY, uint32 MaxX, uint32 MaxY)
@@ -1406,7 +1380,6 @@ public:
 			NewDepthStencilTargetRHI,
 			NewNumUAVs,
 			UAVs);
-		StateReduce<FRHICommandSetRenderTargets>(Cmd);
 	}
 
 	FORCEINLINE_DEBUGGABLE void BeginDrawPrimitiveUP(uint32 PrimitiveType, uint32 NumPrimitives, uint32 NumVertices, uint32 VertexDataStride, void*& OutVertexData)
@@ -1493,7 +1466,6 @@ public:
 		}
 		auto* Cmd = AddCommand<FRHICommandSetComputeShader>();
 		Cmd->Set(ComputeShader);
-		StateReduce<FRHICommandSetComputeShader>(Cmd);
 	}
 
 	FORCEINLINE_DEBUGGABLE void DispatchComputeShader(uint32 ThreadGroupCountX, uint32 ThreadGroupCountY, uint32 ThreadGroupCountZ)
@@ -1582,7 +1554,6 @@ public:
 		}
 		auto* Cmd = AddCommand<FRHICommandEnableDepthBoundsTest>();
 		Cmd->Set(bEnable, MinDepth, MaxDepth);
-		StateReduce<FRHICommandEnableDepthBoundsTest>(Cmd);
 	}
 
 	FORCEINLINE_DEBUGGABLE void ClearUAV(FUnorderedAccessViewRHIParamRef UnorderedAccessViewRHI, const uint32(&Values)[4])
@@ -1631,24 +1602,48 @@ public:
 
 	const SIZE_T GetUsedMemory() const;
 
+protected:
+
+	struct FDrawUpData
+	{
+
+		uint32 PrimitiveType;
+		uint32 NumPrimitives;
+		uint32 NumVertices;
+		uint32 VertexDataStride;
+		void* OutVertexData;
+		uint32 MinVertexIndex;
+		uint32 NumIndices;
+		uint32 IndexDataStride;
+		void* OutIndexData;
+
+		FDrawUpData()
+			: PrimitiveType(PT_Num)
+			, OutVertexData(nullptr)
+			, OutIndexData(nullptr)
+		{
+		}
+
+	};
+
+	FDrawUpData DrawUPData;
+
+	FORCEINLINE_DEBUGGABLE bool HasCommands() const
+	{
+		return (NumCommands > 0);
+	}
+
+
 private:
 	bool bExecuting;
 	uint32 NumCommands;
 	uint32 UID;
-	FRHICommand* StateReducer[ERCT_NUM];
-	void(*SetGlobalBoundShaderState_InternalPtr)(FGlobalBoundShaderState& BoundShaderState);
-
 
 	FORCEINLINE_DEBUGGABLE uint8* Alloc(SIZE_T InSize)
 	{
 		return MemManager.Alloc(InSize);
 	}
 
-	FORCEINLINE_DEBUGGABLE bool HasCommands() const
-	{
-		return (NumCommands > 0);
-	}
-
 	template <typename TCmd>
 	FORCEINLINE_DEBUGGABLE TCmd* AddCommand(uint32 ExtraData = 0)
 	{
@@ -1659,23 +1654,6 @@ private:
 		return Cmd;
 	}
 
-	template <typename TCmd>
-	FORCEINLINE_DEBUGGABLE void StateReduce(TCmd* Cmd)
-	{
-#if 0
-		if (StateReducer[Cmd->Type])
-		{
-			if (FMemory::Memcmp(StateReducer[Cmd->Type], Cmd, sizeof(TCmd)) == 0)
-			{
-				--NumCommands;
-				MemManager.Dealloc(sizeof(TCmd));
-				return;
-			}
-		}
-		StateReducer[Cmd->Type] = Cmd;
-#endif
-	}
-
 	void Reset();
 
 	inline bool Bypass();
@@ -1693,8 +1671,14 @@ class RHI_API FRHICommandListImmediate : public FRHICommandList
 	FRHICommandListImmediate()
 	{
 	}
+	~FRHICommandListImmediate()
+	{
+		Flush(); // this probably never happens, but we want to be sure there are no commands, otherwise it will be executed like a non-immediate command list
+	}
 public:
 
+	inline void Flush();
+
 	#define DEFINE_RHIMETHOD_CMDLIST(Type,Name,ParameterTypesAndNames,ParameterNames,ReturnStatement,NullImplementation)
 	#define DEFINE_RHIMETHOD(Type, Name, ParameterTypesAndNames, ParameterNames, ReturnStatement, NullImplementation) \
 		FORCEINLINE_DEBUGGABLE Type Name ParameterTypesAndNames \
@@ -1773,6 +1757,7 @@ public:
 	static inline FRHICommandListImmediate& GetImmediateCommandList();
 
 	void ExecuteList(FRHICommandList& CmdList);
+	void ExecuteList(FRHICommandListImmediate& CmdList);
 	void LatchBypass();
 
 	FORCEINLINE_DEBUGGABLE void Verify()
@@ -1790,14 +1775,16 @@ public:
 	static void CheckNoOutstandingCmdLists();
 
 private:
+
+	void ExecuteInner(FRHICommandList& CmdList);
+
 	bool bLatchedBypass;
 	friend class FRHICommandList;
 	FThreadSafeCounter UIDCounter;
 	FThreadSafeCounter OutstandingCmdListCount;
 	FRHICommandListImmediate CommandListImmediate;
 
-	template <bool bOnlyTestMemAccess>
-	void ExecuteList(FRHICommandList& CmdList);
+	void ExecuteListGenericRHI(FRHICommandList& CmdList);
 };
 extern RHI_API FRHICommandListExecutor GRHICommandList;
 
@@ -1826,11 +1813,4 @@ FORCEINLINE_DEBUGGABLE FRHICommandListImmediate& FRHICommandListExecutor::GetImm
 #undef DEFINE_RHIMETHOD_GLOBAL
 #undef DEFINE_RHIMETHOD_GLOBALFLUSH
 
-
-DECLARE_CYCLE_STAT_EXTERN(TEXT("RHI Command List Execute"),STAT_RHICmdListExecuteTime,STATGROUP_RHI, );
-DECLARE_CYCLE_STAT_EXTERN(TEXT("RHI Command List Enqueue"),STAT_RHICmdListEnqueueTime,STATGROUP_RHI, );
-DECLARE_DWORD_COUNTER_STAT_EXTERN(TEXT("RHI Command List memory"),STAT_RHICmdListMemory,STATGROUP_RHI,RHI_API);
-DECLARE_DWORD_COUNTER_STAT_EXTERN(TEXT("RHI Command List count"),STAT_RHICmdListCount,STATGROUP_RHI,RHI_API);
-DECLARE_CYCLE_STAT_EXTERN(TEXT("RHI Counter TEMP"), STAT_RHICounterTEMP,STATGROUP_RHI,RHI_API);
-
 #include "RHICommandList.inl"
diff --git a/Engine/Source/Runtime/RHI/Public/RHICommandList.inl b/Engine/Source/Runtime/RHI/Public/RHICommandList.inl
--- a/Engine/Source/Runtime/RHI/Public/RHICommandList.inl
+++ b/Engine/Source/Runtime/RHI/Public/RHICommandList.inl
@@ -17,6 +17,16 @@ FORCEINLINE_DEBUGGABLE void FRHICommandList::Flush()
 	}
 }
 
+FORCEINLINE_DEBUGGABLE void FRHICommandListImmediate::Flush()
+{
+	if (HasCommands())
+	{
+		check(!DrawUPData.OutVertexData && !DrawUPData.OutIndexData);
+		GRHICommandList.ExecuteList(*this);
+	}
+}
+
+
 // Helper class for traversing a FRHICommandList
 class FRHICommandListIterator
 {
diff --git a/Engine/Source/Runtime/Renderer/Private/DeferredShadingRenderer.cpp b/Engine/Source/Runtime/Renderer/Private/DeferredShadingRenderer.cpp
--- a/Engine/Source/Runtime/Renderer/Private/DeferredShadingRenderer.cpp
+++ b/Engine/Source/Runtime/Renderer/Private/DeferredShadingRenderer.cpp
@@ -222,9 +222,8 @@ bool FDeferredShadingSceneRenderer::RenderBasePassStaticDataDefault(FRHICommandL
 {
 	QUICK_SCOPE_CYCLE_COUNTER(AA_RenderBasePassStaticDataDefault);
 	bool bDirty = false;
-	if (!GRHICommandList.Bypass())
+	if (!GRHICommandList.Bypass() && FApp::ShouldUseThreadingForPerformance())
 	{
-		RHICmdList.Flush(); // we assume this is the immediate command list, regardless it makes sense to empty it before we start using more command lists
 		FGraphEventRef SubmitChain;
 		int32 Width = CVarRHICmdWidth.GetValueOnRenderThread(); // we use a few more than needed to cover non-equal jobs
 
@@ -367,14 +366,7 @@ bool FDeferredShadingSceneRenderer::RenderBasePass(FRHICommandListImmediate& RHI
 
 	// Render the base pass static data
 	{
-		//@todo-rco: Very temp code!!!
-		SCOPE_CYCLE_COUNTER(STAT_RHICounterTEMP);
-		static IConsoleVariable* RHICmdListCVar = IConsoleManager::Get().FindConsoleVariable(TEXT("r.RHICmd"));
-		bool bUseRHICmdList = (RHICmdListCVar->GetInt() >= 1);
-
-		FRHICommandList LocalRHICmdList;
-		FRHICommandList& UseRHICmdList = bUseRHICmdList ? LocalRHICmdList : RHICmdList;
-		bDirty |= RenderBasePassStaticData(UseRHICmdList, View);
+		bDirty |= RenderBasePassStaticData(RHICmdList, View);
 	}
 
 	{
diff --git a/Engine/Source/Runtime/Renderer/Private/ForwardBasePassRendering.cpp b/Engine/Source/Runtime/Renderer/Private/ForwardBasePassRendering.cpp
--- a/Engine/Source/Runtime/Renderer/Private/ForwardBasePassRendering.cpp
+++ b/Engine/Source/Runtime/Renderer/Private/ForwardBasePassRendering.cpp
@@ -327,29 +327,22 @@ void FForwardShadingSceneRenderer::RenderForwardShadingBasePass(FRHICommandListI
 		RHICmdList.SetViewport(View.ViewRect.Min.X, View.ViewRect.Min.Y, 0, View.ViewRect.Max.X, View.ViewRect.Max.Y, 1);
 
 		{
-			SCOPE_CYCLE_COUNTER(STAT_RHICounterTEMP);
-			//@todo-rco: Very temp code!!!
-			static IConsoleVariable* RHICmdListCVar = IConsoleManager::Get().FindConsoleVariable(TEXT("r.RHICmd"));
-			bool bUseRHICmdList = RHICmdListCVar->GetInt() != 0;
-			FRHICommandList LocalRHICmdList;
-			FRHICommandList& UseRHICmdList = bUseRHICmdList ? LocalRHICmdList : RHICmdList;
-
 			// Render the base pass static data
 			if (SortMode == EBasePassSort::SortPerMesh)
 			{
 				SCOPE_CYCLE_COUNTER(STAT_StaticDrawListDrawTime);
-				MaxDraws -= Scene->BasePassForForwardShadingLowQualityLightMapDrawList[FScene::EBasePass_Default].DrawVisibleFrontToBack(UseRHICmdList, View, View.StaticMeshVisibilityMap, View.StaticMeshBatchVisibility, MaxDraws);
-				MaxDraws -= Scene->BasePassForForwardShadingDistanceFieldShadowMapLightMapDrawList[FScene::EBasePass_Default].DrawVisibleFrontToBack(UseRHICmdList, View, View.StaticMeshVisibilityMap, View.StaticMeshBatchVisibility, MaxDraws);
-				MaxDraws -= Scene->BasePassForForwardShadingDirectionalLightAndSHIndirectDrawList[FScene::EBasePass_Default].DrawVisibleFrontToBack(UseRHICmdList, View, View.StaticMeshVisibilityMap, View.StaticMeshBatchVisibility, MaxDraws);
-				MaxDraws -= Scene->BasePassForForwardShadingNoLightMapDrawList[FScene::EBasePass_Default].DrawVisibleFrontToBack(UseRHICmdList, View, View.StaticMeshVisibilityMap, View.StaticMeshBatchVisibility, MaxDraws);
+				MaxDraws -= Scene->BasePassForForwardShadingLowQualityLightMapDrawList[FScene::EBasePass_Default].DrawVisibleFrontToBack(RHICmdList, View, View.StaticMeshVisibilityMap, View.StaticMeshBatchVisibility, MaxDraws);
+				MaxDraws -= Scene->BasePassForForwardShadingDistanceFieldShadowMapLightMapDrawList[FScene::EBasePass_Default].DrawVisibleFrontToBack(RHICmdList, View, View.StaticMeshVisibilityMap, View.StaticMeshBatchVisibility, MaxDraws);
+				MaxDraws -= Scene->BasePassForForwardShadingDirectionalLightAndSHIndirectDrawList[FScene::EBasePass_Default].DrawVisibleFrontToBack(RHICmdList, View, View.StaticMeshVisibilityMap, View.StaticMeshBatchVisibility, MaxDraws);
+				MaxDraws -= Scene->BasePassForForwardShadingNoLightMapDrawList[FScene::EBasePass_Default].DrawVisibleFrontToBack(RHICmdList, View, View.StaticMeshVisibilityMap, View.StaticMeshBatchVisibility, MaxDraws);
 			}
 			else
 			{
 				SCOPE_CYCLE_COUNTER(STAT_StaticDrawListDrawTime);
-				Scene->BasePassForForwardShadingLowQualityLightMapDrawList[FScene::EBasePass_Default].DrawVisible(UseRHICmdList, View, View.StaticMeshVisibilityMap, View.StaticMeshBatchVisibility);
-				Scene->BasePassForForwardShadingDistanceFieldShadowMapLightMapDrawList[FScene::EBasePass_Default].DrawVisible(UseRHICmdList, View, View.StaticMeshVisibilityMap, View.StaticMeshBatchVisibility);
-				Scene->BasePassForForwardShadingDirectionalLightAndSHIndirectDrawList[FScene::EBasePass_Default].DrawVisible(UseRHICmdList, View, View.StaticMeshVisibilityMap, View.StaticMeshBatchVisibility);
-				Scene->BasePassForForwardShadingNoLightMapDrawList[FScene::EBasePass_Default].DrawVisible(UseRHICmdList, View, View.StaticMeshVisibilityMap, View.StaticMeshBatchVisibility);
+				Scene->BasePassForForwardShadingLowQualityLightMapDrawList[FScene::EBasePass_Default].DrawVisible(RHICmdList, View, View.StaticMeshVisibilityMap, View.StaticMeshBatchVisibility);
+				Scene->BasePassForForwardShadingDistanceFieldShadowMapLightMapDrawList[FScene::EBasePass_Default].DrawVisible(RHICmdList, View, View.StaticMeshVisibilityMap, View.StaticMeshBatchVisibility);
+				Scene->BasePassForForwardShadingDirectionalLightAndSHIndirectDrawList[FScene::EBasePass_Default].DrawVisible(RHICmdList, View, View.StaticMeshVisibilityMap, View.StaticMeshBatchVisibility);
+				Scene->BasePassForForwardShadingNoLightMapDrawList[FScene::EBasePass_Default].DrawVisible(RHICmdList, View, View.StaticMeshVisibilityMap, View.StaticMeshBatchVisibility);
 			}
 		}
 
diff --git a/Engine/Source/Runtime/Renderer/Private/ShadowRendering.cpp b/Engine/Source/Runtime/Renderer/Private/ShadowRendering.cpp
--- a/Engine/Source/Runtime/Renderer/Private/ShadowRendering.cpp
+++ b/Engine/Source/Runtime/Renderer/Private/ShadowRendering.cpp
@@ -1364,13 +1364,6 @@ void FProjectedShadowInfo::RenderDepth(FRHICommandListImmediate& RHICmdList, FDe
 	FShadowDepthDrawingPolicy<true>::PolicyShadowInfo = this;
 
 	{
-	//@todo-rco: Very temp code!!!
-		SCOPE_CYCLE_COUNTER(STAT_RHICounterTEMP);
-		static IConsoleVariable* RHICmdListCVar = IConsoleManager::Get().FindConsoleVariable(TEXT("r.RHICmd"));
-		bool bUseRHICmdList = (RHICmdListCVar->GetInt() >= 2);
-		FRHICommandList LocalRHICmdList;
-		FRHICommandList& UseRHICmdList = bUseRHICmdList ? LocalRHICmdList : RHICmdList;
-
 		// Draw the subject's static elements using static draw lists
 		if (IsWholeSceneDirectionalShadow())
 		{
@@ -1378,12 +1371,12 @@ void FProjectedShadowInfo::RenderDepth(FRHICommandListImmediate& RHICmdList, FDe
 
 			if (bReflectiveShadowmap)
 			{
-				SceneRenderer->Scene->WholeSceneReflectiveShadowMapDrawList.DrawVisible(UseRHICmdList, *FoundView, StaticMeshWholeSceneShadowDepthMap, StaticMeshWholeSceneShadowBatchVisibility);
+				SceneRenderer->Scene->WholeSceneReflectiveShadowMapDrawList.DrawVisible(RHICmdList, *FoundView, StaticMeshWholeSceneShadowDepthMap, StaticMeshWholeSceneShadowBatchVisibility);
 			}
 			else
 			{
 				// Use the scene's shadow depth draw list with this shadow's visibility map
-				SceneRenderer->Scene->WholeSceneShadowDepthDrawList.DrawVisible(UseRHICmdList, *FoundView, StaticMeshWholeSceneShadowDepthMap, StaticMeshWholeSceneShadowBatchVisibility);
+				SceneRenderer->Scene->WholeSceneShadowDepthDrawList.DrawVisible(RHICmdList, *FoundView, StaticMeshWholeSceneShadowDepthMap, StaticMeshWholeSceneShadowBatchVisibility);
 			}
 		}
 		// Draw the subject's static elements using manual state filtering
@@ -1394,12 +1387,12 @@ void FProjectedShadowInfo::RenderDepth(FRHICommandListImmediate& RHICmdList, FDe
 			if(bReflectiveShadowmap && !bOnePassPointLightShadow)
 			{
 				// reflective shadow map
-				DrawShadowMeshElements<true>(UseRHICmdList, *FoundView, *this);
+				DrawShadowMeshElements<true>(RHICmdList, *FoundView, *this);
 			}
 			else
 			{
 				// normal shadow map
-				DrawShadowMeshElements<false>(UseRHICmdList, *FoundView, *this);
+				DrawShadowMeshElements<false>(RHICmdList, *FoundView, *this);
 			}
 		}
 
diff --git a/Engine/Source/Runtime/Renderer/Private/TranslucentRendering.cpp b/Engine/Source/Runtime/Renderer/Private/TranslucentRendering.cpp
--- a/Engine/Source/Runtime/Renderer/Private/TranslucentRendering.cpp
+++ b/Engine/Source/Runtime/Renderer/Private/TranslucentRendering.cpp
@@ -531,13 +531,6 @@ void FTranslucentPrimSet::RenderPrimitive(
 		// Render static scene prim
 		if( ViewRelevance.bStaticRelevance )
 		{
-			//@todo-rco: Very temp code!!!
-			SCOPE_CYCLE_COUNTER(STAT_RHICounterTEMP);
-			static IConsoleVariable* RHICmdListCVar = IConsoleManager::Get().FindConsoleVariable(TEXT("r.RHICmd"));
-			bool bUseRHICmdList = (RHICmdListCVar->GetInt() >= 3);
-			FRHICommandList LocalRHICmdList;
-			FRHICommandList& UseRHICmdList = bUseRHICmdList ? LocalRHICmdList : RHICmdList;
-
 			// Render static meshes from static scene prim
 			for( int32 StaticMeshIdx=0; StaticMeshIdx < PrimitiveSceneInfo->StaticMeshes.Num(); StaticMeshIdx++ )
 			{
@@ -548,7 +541,7 @@ void FTranslucentPrimSet::RenderPrimitive(
 					&& (StaticMesh.MaterialRenderProxy->GetMaterial(FeatureLevel)->IsSeparateTranslucencyEnabled() == bSeparateTranslucencyPass))
 				{
 					FTranslucencyDrawingPolicyFactory::DrawStaticMesh(
-						UseRHICmdList,
+						RHICmdList,
 						View,
 						FTranslucencyDrawingPolicyFactory::ContextType( TranslucentSelfShadow, bSeparateTranslucencyPass),
 						StaticMesh,
