UE4 - fix recompile shaders

diff --git a/Engine/Source/Runtime/RHI/Public/RHICommandList.h b/Engine/Source/Runtime/RHI/Public/RHICommandList.h
--- a/Engine/Source/Runtime/RHI/Public/RHICommandList.h
+++ b/Engine/Source/Runtime/RHI/Public/RHICommandList.h
@@ -1096,7 +1096,8 @@ public:
 		FShader* GeometryShader
 		)
 	{
-		if (!GlobalBoundShaderState.Get())
+		FGlobalBoundShaderStateWorkArea* ExistingGlobalBoundShaderState = GlobalBoundShaderState.Get();
+		if (!ExistingGlobalBoundShaderState)
 		{
 			FGlobalBoundShaderStateWorkArea* NewGlobalBoundShaderState = new FGlobalBoundShaderStateWorkArea();
 			NewGlobalBoundShaderState->Args.VertexDeclarationRHI = VertexDeclarationRHI;
@@ -1111,15 +1112,21 @@ public:
 			{
 				//we lost
 				delete NewGlobalBoundShaderState;
-				check(OldGlobalBoundShaderState == GlobalBoundShaderState.Get());
+				check(OldGlobalBoundShaderState == ExistingGlobalBoundShaderState);
 			}
 		}
-		check(
-			VertexDeclarationRHI == GlobalBoundShaderState.Get()->Args.VertexDeclarationRHI &&
-			VertexShader == GlobalBoundShaderState.Get()->Args.VertexShader &&
-			PixelShader == GlobalBoundShaderState.Get()->Args.PixelShader &&
-			GeometryShader == GlobalBoundShaderState.Get()->Args.GeometryShader
-			);
+		else if (!(VertexDeclarationRHI == ExistingGlobalBoundShaderState->Args.VertexDeclarationRHI &&
+			VertexShader == ExistingGlobalBoundShaderState->Args.VertexShader &&
+			PixelShader == ExistingGlobalBoundShaderState->Args.PixelShader &&
+			GeometryShader == ExistingGlobalBoundShaderState->Args.GeometryShader))
+		{
+			// this is sketchy from a parallel perspective, but assuming the writes are atomic, and probably even if they aren't, this should be ok
+			ExistingGlobalBoundShaderState->Args.VertexDeclarationRHI = VertexDeclarationRHI;
+			ExistingGlobalBoundShaderState->Args.VertexShader = VertexShader;
+			ExistingGlobalBoundShaderState->Args.PixelShader = PixelShader;
+			ExistingGlobalBoundShaderState->Args.GeometryShader = GeometryShader;
+		}
+
 		if (Bypass())
 		{
 			SetGlobalBoundShaderState_InternalPtr(GlobalBoundShaderState);
