Actually ES2 does support MSAA resolve on PC emulation.

diff --git a/Engine/Source/Runtime/RHI/Public/RHIUtilities.h b/Engine/Source/Runtime/RHI/Public/RHIUtilities.h
--- a/Engine/Source/Runtime/RHI/Public/RHIUtilities.h
+++ b/Engine/Source/Runtime/RHI/Public/RHIUtilities.h
@@ -199,9 +199,6 @@ inline void RHICreateTargetableShaderResource2D(
 		bForceSeparateTargetAndShaderResource = true;
 	}
 
-	// ES2 doesn't support resolve operations.
-    bForceSeparateTargetAndShaderResource &= (GRHIFeatureLevel > ERHIFeatureLevel::ES2);
-
 	if (!bForceSeparateTargetAndShaderResource/* && GSupportsRenderDepthTargetableShaderResources*/)
 	{
 		// Create a single texture that has both TargetableTextureFlags and TexCreate_ShaderResource set.
