UE4 - parallel rendering, reworks command lists

diff --git a/Engine/Source/Runtime/Core/Public/Misc/MemStack.h b/Engine/Source/Runtime/Core/Public/Misc/MemStack.h
--- a/Engine/Source/Runtime/Core/Public/Misc/MemStack.h
+++ b/Engine/Source/Runtime/Core/Public/Misc/MemStack.h
@@ -55,12 +55,13 @@ class CORE_API FMemStackBase
 {
 public:
 
-	FMemStackBase()
+	FMemStackBase(int32 InMinMarksToAlloc = 1)
 		: Top(nullptr)
 		, End(nullptr)
 		, TopChunk(nullptr)
 		, TopMark(nullptr)
 		, NumMarks(0)
+		, MinMarksToAlloc(InMinMarksToAlloc)
 	{
 	}
 
@@ -70,19 +71,19 @@ public:
 		FreeChunks(nullptr);
 	}
 
-
-	// Get bytes.
 	FORCEINLINE uint8* PushBytes(int32 AllocSize, int32 Alignment)
 	{
 		return (uint8*)Alloc(AllocSize, FMath::Max(AllocSize >= 16 ? (int32)16 : (int32)8, Alignment));
 	}
+
 	FORCEINLINE void* Alloc(int32 AllocSize, int32 Alignment)
 	{
 		// Debug checks.
 		checkSlow(AllocSize>=0);
 		checkSlow((Alignment&(Alignment-1))==0);
 		checkSlow(Top<=End);
-		checkSlow(NumMarks > 0);
+		checkSlow(NumMarks >= MinMarksToAlloc);
+
 
 		// Try to get memory from the current chunk.
 		uint8* Result = Align( Top, Alignment );
@@ -110,6 +111,11 @@ public:
 		return TopChunk == nullptr;
 	}
 
+	FORCEINLINE void Flush()
+	{
+		check(!NumMarks && !MinMarksToAlloc);
+		FreeChunks(nullptr);
+	}
 	/** @return the number of bytes allocated for this FMemStack that are currently in use. */
 	int32 GetByteCount() const;
 
@@ -135,6 +141,15 @@ public:
 
 private:
 
+	/**
+	 * Allocate a new chunk of memory of at least MinSize size,
+	 * updates the memory stack's Chunks table and ActiveChunks counter.
+	 */
+	void AllocateNewChunk( int32 MinSize );
+
+	/** Frees the chunks above the specified chunk on the stack. */
+	void FreeChunks( FTaggedMemory* NewTopChunk );
+
 	// Variables.
 	uint8*			Top;				// Top of current chunk (Top<=End).
 	uint8*			End;				// End of current chunk.
@@ -146,19 +161,12 @@ private:
 	/** The number of marks on this stack. */
 	int32 NumMarks;
 
-	/**
-	 * Allocate a new chunk of memory of at least MinSize size,
-	 * updates the memory stack's Chunks table and ActiveChunks counter.
-	 */
-	void AllocateNewChunk( int32 MinSize );
-
-	/** Frees the chunks above the specified chunk on the stack. */
-	void FreeChunks( FTaggedMemory* NewTopChunk );
+	/** Used for a checkSlow. Most stacks require a mark to allocate. Command lists don't because they never mark, only flush*/
+	int32 MinMarksToAlloc;
 };
 
 class CORE_API FMemStack : public TThreadSingleton<FMemStack>, public FMemStackBase
 {
-
 };
 
 /*-----------------------------------------------------------------------------
diff --git a/Engine/Source/Runtime/Engine/Private/StaticBoundShaderState.cpp b/Engine/Source/Runtime/Engine/Private/StaticBoundShaderState.cpp
--- a/Engine/Source/Runtime/Engine/Private/StaticBoundShaderState.cpp
+++ b/Engine/Source/Runtime/Engine/Private/StaticBoundShaderState.cpp
@@ -7,23 +7,8 @@
 #include "EnginePrivate.h"
 #include "StaticBoundShaderState.h"
 
-static struct FStartupSetBSSPtr
-{
-	int DontDeadStripMeInt;
-	FStartupSetBSSPtr()
-	{
-		DontDeadStripMeInt = 0;
-		SetGlobalBoundShaderState_InternalPtr = &SetGlobalBoundShaderState_Internal;
-	}
-	void DontDeadStripMe()
-	{
-		DontDeadStripMeInt++;
-	}
-} StartupSetBSSPtr;
-
 TLinkedList<FGlobalBoundShaderStateResource*>*& FGlobalBoundShaderStateResource::GetGlobalBoundShaderStateList()
 {
-	StartupSetBSSPtr.DontDeadStripMe();
 	static TLinkedList<FGlobalBoundShaderStateResource*>* List = NULL;
 	return List;
 }
@@ -107,7 +92,7 @@ void FGlobalBoundShaderStateResource::ReleaseRHI()
 	BoundShaderState.SafeRelease();
 }
 
-void SetGlobalBoundShaderState_Internal(FGlobalBoundShaderState& GlobalBoundShaderState)
+static void SetGlobalBoundShaderState_Internal(FGlobalBoundShaderState& GlobalBoundShaderState)
 {
 	// Check for unset uniform buffer parameters
 	// Technically you can set uniform buffer parameters after calling RHISetBoundShaderState, but this is the most global place to check for unset parameters
@@ -132,3 +117,66 @@ void SetGlobalBoundShaderState_Internal(FGlobalBoundShaderState& GlobalBoundShad
 }
 
 
+struct FRHICommandSetGlobalBoundShaderState : public FRHICommand<FRHICommandSetGlobalBoundShaderState>
+{
+	FGlobalBoundShaderState GlobalBoundShaderState;
+	FORCEINLINE_DEBUGGABLE FRHICommandSetGlobalBoundShaderState(FGlobalBoundShaderState& InGlobalBoundShaderState)
+		: GlobalBoundShaderState(InGlobalBoundShaderState)
+	{
+	}
+	void Execute()
+	{
+		SetGlobalBoundShaderState_Internal(GlobalBoundShaderState);
+	}
+};
+
+
+void SetGlobalBoundShaderState(
+	FRHICommandList& RHICmdList,
+	FGlobalBoundShaderState& GlobalBoundShaderState,
+	FVertexDeclarationRHIParamRef VertexDeclarationRHI,
+	FShader* VertexShader,
+	FShader* PixelShader,
+	FShader* GeometryShader
+	)
+{
+	FGlobalBoundShaderStateWorkArea* ExistingGlobalBoundShaderState = GlobalBoundShaderState.Get();
+	if (!ExistingGlobalBoundShaderState)
+	{
+		FGlobalBoundShaderStateWorkArea* NewGlobalBoundShaderState = new FGlobalBoundShaderStateWorkArea();
+		NewGlobalBoundShaderState->Args.VertexDeclarationRHI = VertexDeclarationRHI;
+		NewGlobalBoundShaderState->Args.VertexShader = VertexShader;
+		NewGlobalBoundShaderState->Args.PixelShader = PixelShader;
+		NewGlobalBoundShaderState->Args.GeometryShader = GeometryShader;
+		FPlatformMisc::MemoryBarrier();
+
+		FGlobalBoundShaderStateWorkArea* OldGlobalBoundShaderState = (FGlobalBoundShaderStateWorkArea*)FPlatformAtomics::InterlockedCompareExchangePointer((void**)GlobalBoundShaderState.GetPtr(), NewGlobalBoundShaderState, nullptr);
+
+		if (OldGlobalBoundShaderState != nullptr)
+		{
+			//we lost
+			delete NewGlobalBoundShaderState;
+			check(OldGlobalBoundShaderState == ExistingGlobalBoundShaderState);
+		}
+	}
+	else if (!(VertexDeclarationRHI == ExistingGlobalBoundShaderState->Args.VertexDeclarationRHI &&
+		VertexShader == ExistingGlobalBoundShaderState->Args.VertexShader &&
+		PixelShader == ExistingGlobalBoundShaderState->Args.PixelShader &&
+		GeometryShader == ExistingGlobalBoundShaderState->Args.GeometryShader))
+	{
+		// this is sketchy from a parallel perspective, but assuming the writes are atomic, and probably even if they aren't, this should be ok
+		ExistingGlobalBoundShaderState->Args.VertexDeclarationRHI = VertexDeclarationRHI;
+		ExistingGlobalBoundShaderState->Args.VertexShader = VertexShader;
+		ExistingGlobalBoundShaderState->Args.PixelShader = PixelShader;
+		ExistingGlobalBoundShaderState->Args.GeometryShader = GeometryShader;
+	}
+
+	if (RHICmdList.Bypass())
+	{
+		SetGlobalBoundShaderState_Internal(GlobalBoundShaderState);
+		return;
+	}
+	new (RHICmdList.AllocCommand<FRHICommandSetGlobalBoundShaderState>()) FRHICommandSetGlobalBoundShaderState(GlobalBoundShaderState);
+}
+
+
diff --git a/Engine/Source/Runtime/Engine/Public/StaticBoundShaderState.h b/Engine/Source/Runtime/Engine/Public/StaticBoundShaderState.h
--- a/Engine/Source/Runtime/Engine/Public/StaticBoundShaderState.h
+++ b/Engine/Source/Runtime/Engine/Public/StaticBoundShaderState.h
@@ -51,13 +51,37 @@ private:
 typedef TGlobalResource<FGlobalBoundShaderStateResource> FGlobalBoundShaderState_Internal;
 
 
-/**
-* SetGlobalBoundShaderState - sets the global bound shader state, also creates and caches it if necessary
-*
-* @param BoundShaderState			- global bound shader state to initialize and set
-*/
+struct FGlobalBoundShaderStateArgs
+{
+	FVertexDeclarationRHIParamRef VertexDeclarationRHI;
+	FShader* VertexShader;
+	FShader* PixelShader;
+	FShader* GeometryShader;
+};
+
+struct FGlobalBoundShaderStateWorkArea
+{
+	FGlobalBoundShaderStateArgs Args;
+	FGlobalBoundShaderState_Internal* BSS; //ideally this would be part of this memory block and not a separate allocation...that is doable, if a little tedious. The point is we need to delay the construction until we get back to the render thread.
+
+	FGlobalBoundShaderStateWorkArea()
+		: BSS(nullptr)
+	{
+	}
+};
+
+struct FGlobalBoundShaderState
+{
+public:
+
+	FGlobalBoundShaderStateWorkArea* Get(ERHIFeatureLevel::Type InFeatureLevel = GRHIFeatureLevel)  { return WorkAreas[InFeatureLevel]; }
+	FGlobalBoundShaderStateWorkArea** GetPtr(ERHIFeatureLevel::Type InFeatureLevel = GRHIFeatureLevel)  { return &WorkAreas[InFeatureLevel]; }
+
+private:
+
+	FGlobalBoundShaderStateWorkArea* WorkAreas[ERHIFeatureLevel::Num];
+};
 
-extern ENGINE_API void SetGlobalBoundShaderState_Internal(FGlobalBoundShaderState& BoundShaderState);
 
 
 /**
@@ -71,16 +95,13 @@ extern ENGINE_API void SetGlobalBoundShaderState_Internal(FGlobalBoundShaderStat
  * @param GeometryShader			- the geometry shader to use in creating the new bound shader state (0 if not used)
  */
 
-FORCEINLINE void SetGlobalBoundShaderState(
+ENGINE_API void SetGlobalBoundShaderState(
 	FRHICommandList& RHICmdList,
 	FGlobalBoundShaderState& BoundShaderState,
 	FVertexDeclarationRHIParamRef VertexDeclaration,
 	FShader* VertexShader,
 	FShader* PixelShader,
 	FShader* GeometryShader = nullptr
-	)
-{
-	RHICmdList.BuildAndSetGlobalBoundShaderState(BoundShaderState, VertexDeclaration, VertexShader, PixelShader, GeometryShader);
-}
+	);
 
 #endif
diff --git a/Engine/Source/Runtime/RHI/Private/RHI.cpp b/Engine/Source/Runtime/RHI/Private/RHI.cpp
--- a/Engine/Source/Runtime/RHI/Private/RHI.cpp
+++ b/Engine/Source/Runtime/RHI/Private/RHI.cpp
@@ -94,6 +94,10 @@ void FRHIResource::FlushPendingDeletes()
 				delete Ref;
 				CurrentlyDeleting = nullptr;
 			}
+			else
+			{
+				verify(Ref->MarkedForDelete.Decrement() == 0);
+			}
 		}
 	}
 }
diff --git a/Engine/Source/Runtime/RHI/Private/RHICommandList.cpp b/Engine/Source/Runtime/RHI/Private/RHICommandList.cpp
--- a/Engine/Source/Runtime/RHI/Private/RHICommandList.cpp
+++ b/Engine/Source/Runtime/RHI/Private/RHICommandList.cpp
@@ -19,13 +19,6 @@ static TAutoConsoleVariable<int32> CVarRHICmdBypass(
 	TEXT("0: Disable, 1: Enable"),
 	ECVF_Cheat);
 
-static TAutoConsoleVariable<int32> CVarRHICmdExec(
-	TEXT("r.RHIExec"),
-	0,
-	TEXT("0: Executes RHI Command List using RHI calls (default)\n")
-	TEXT("1: Executes RHI Command List using platform specific APIs")
-	);
-
 TAutoConsoleVariable<int32> CVarRHICmdWidth(
 	TEXT("r.RHICmdWidth"), 
 	8,
@@ -33,397 +26,25 @@ TAutoConsoleVariable<int32> CVarRHICmdWidth(
 
 
 RHI_API FRHICommandListExecutor GRHICommandList;
-RHI_API TLockFreeFixedSizeAllocator<FRHICommandList::FMemManager::FPage::PageSize> FRHICommandList::FMemManager::FPage::TheAllocator;
-RHI_API void(*SetGlobalBoundShaderState_InternalPtr)(FGlobalBoundShaderState& BoundShaderState);
 
 
 //RHI_API FGlobalRHI GRHI;
 
-static_assert(sizeof(FRHICommand) == sizeof(FRHICommandNopEndOfPage), "These should match for filling in the end of a page when a new allocation won't fit");
-
-void FRHICommandListExecutor::ExecuteListGenericRHI(FRHICommandList& CmdList)
+void FRHICommandListExecutor::ExecuteInner(FRHICommandListBase& CmdList)
 {
-	FRHICommandListIterator Iter(CmdList);
-	while (Iter.HasCommandsLeft())
+	CmdList.bExecuting = true;
 	{
-		auto* Cmd = *Iter;
-		switch (Cmd->Type)
+		FRHICommandListIterator Iter(CmdList);
+		while (Iter.HasCommandsLeft())
 		{
-		case ERCT_NopBlob:
-			{
-				// Nop
-				auto* RHICmd = Iter.NextCommand<FRHICommandNopBlob>();
-				(void)RHICmd;	// Unused
-			}
-			break;
-		case ERCT_NopEndOfPage:
-			{
-				// Nop
-				auto* RHICmd = Iter.NextCommand<FRHICommandNopEndOfPage>();
-				(void)RHICmd;	// Unused
-			}
-			break;
-		case ERCT_SetRasterizerState:
-			{
-				auto* RHICmd = Iter.NextCommand<FRHICommandSetRasterizerState>();
-				SetRasterizerState_Internal(RHICmd->State);
-			}
-			break;
-		case ERCT_SetShaderParameter:
-			{
-				auto* RHICmd = Iter.NextCommand<FRHICommandSetShaderParameter>();
-				switch (RHICmd->ShaderFrequency)
-				{
-				case SF_Vertex:		SetShaderParameter_Internal((FVertexShaderRHIParamRef)RHICmd->Shader, RHICmd->BufferIndex, RHICmd->BaseIndex, RHICmd->NumBytes, RHICmd->NewValue); break;
-				case SF_Hull:		SetShaderParameter_Internal((FHullShaderRHIParamRef)RHICmd->Shader, RHICmd->BufferIndex, RHICmd->BaseIndex, RHICmd->NumBytes, RHICmd->NewValue); break;
-				case SF_Domain:		SetShaderParameter_Internal((FDomainShaderRHIParamRef)RHICmd->Shader, RHICmd->BufferIndex, RHICmd->BaseIndex, RHICmd->NumBytes, RHICmd->NewValue); break;
-				case SF_Geometry:	SetShaderParameter_Internal((FGeometryShaderRHIParamRef)RHICmd->Shader, RHICmd->BufferIndex, RHICmd->BaseIndex, RHICmd->NumBytes, RHICmd->NewValue); break;
-				case SF_Pixel:		SetShaderParameter_Internal((FPixelShaderRHIParamRef)RHICmd->Shader, RHICmd->BufferIndex, RHICmd->BaseIndex, RHICmd->NumBytes, RHICmd->NewValue); break;
-				case SF_Compute:	SetShaderParameter_Internal((FComputeShaderRHIParamRef)RHICmd->Shader, RHICmd->BufferIndex, RHICmd->BaseIndex, RHICmd->NumBytes, RHICmd->NewValue); break;
-				default: check(0); break;
-				}
-			}
-			break;
-		case ERCT_SetShaderUniformBuffer:
-			{
-				auto* RHICmd = Iter.NextCommand<FRHICommandSetShaderUniformBuffer>();
-				switch (RHICmd->ShaderFrequency)
-				{
-				case SF_Vertex:		SetShaderUniformBuffer_Internal((FVertexShaderRHIParamRef)RHICmd->Shader, RHICmd->BaseIndex, RHICmd->UniformBuffer); break;
-				case SF_Hull:		SetShaderUniformBuffer_Internal((FHullShaderRHIParamRef)RHICmd->Shader, RHICmd->BaseIndex, RHICmd->UniformBuffer); break;
-				case SF_Domain:		SetShaderUniformBuffer_Internal((FDomainShaderRHIParamRef)RHICmd->Shader, RHICmd->BaseIndex, RHICmd->UniformBuffer); break;
-				case SF_Geometry:	SetShaderUniformBuffer_Internal((FGeometryShaderRHIParamRef)RHICmd->Shader, RHICmd->BaseIndex, RHICmd->UniformBuffer); break;
-				case SF_Pixel:		SetShaderUniformBuffer_Internal((FPixelShaderRHIParamRef)RHICmd->Shader, RHICmd->BaseIndex, RHICmd->UniformBuffer); break;
-				case SF_Compute:	SetShaderUniformBuffer_Internal((FComputeShaderRHIParamRef)RHICmd->Shader, RHICmd->BaseIndex, RHICmd->UniformBuffer); break;
-				default: check(0); break;
-				}
-			}
-			break;
-		case ERCT_SetShaderSampler:
-			{
-				auto* RHICmd = Iter.NextCommand<FRHICommandSetShaderSampler>();
-				switch (RHICmd->ShaderFrequency)
-				{
-				case SF_Vertex:		SetShaderSampler_Internal((FVertexShaderRHIParamRef)RHICmd->Shader, RHICmd->SamplerIndex, RHICmd->Sampler); break;
-				case SF_Hull:		SetShaderSampler_Internal((FHullShaderRHIParamRef)RHICmd->Shader, RHICmd->SamplerIndex, RHICmd->Sampler); break;
-				case SF_Domain:		SetShaderSampler_Internal((FDomainShaderRHIParamRef)RHICmd->Shader, RHICmd->SamplerIndex, RHICmd->Sampler); break;
-				case SF_Geometry:	SetShaderSampler_Internal((FGeometryShaderRHIParamRef)RHICmd->Shader, RHICmd->SamplerIndex, RHICmd->Sampler); break;
-				case SF_Pixel:		SetShaderSampler_Internal((FPixelShaderRHIParamRef)RHICmd->Shader, RHICmd->SamplerIndex, RHICmd->Sampler); break;
-				case SF_Compute:	SetShaderSampler_Internal((FComputeShaderRHIParamRef)RHICmd->Shader, RHICmd->SamplerIndex, RHICmd->Sampler); break;
-				default: check(0); break;
-				}
-			}
-			break;
-		case ERCT_SetShaderTexture:
-			{
-				auto* RHICmd = Iter.NextCommand<FRHICommandSetShaderTexture>();
-				switch (RHICmd->ShaderFrequency)
-				{
-				case SF_Vertex:		SetShaderTexture_Internal((FVertexShaderRHIParamRef)RHICmd->Shader, RHICmd->TextureIndex, RHICmd->Texture); break;
-				case SF_Hull:		SetShaderTexture_Internal((FHullShaderRHIParamRef)RHICmd->Shader, RHICmd->TextureIndex, RHICmd->Texture); break;
-				case SF_Domain:		SetShaderTexture_Internal((FDomainShaderRHIParamRef)RHICmd->Shader, RHICmd->TextureIndex, RHICmd->Texture); break;
-				case SF_Geometry:	SetShaderTexture_Internal((FGeometryShaderRHIParamRef)RHICmd->Shader, RHICmd->TextureIndex, RHICmd->Texture); break;
-				case SF_Pixel:		SetShaderTexture_Internal((FPixelShaderRHIParamRef)RHICmd->Shader, RHICmd->TextureIndex, RHICmd->Texture); break;
-				case SF_Compute:	SetShaderTexture_Internal((FComputeShaderRHIParamRef)RHICmd->Shader, RHICmd->TextureIndex, RHICmd->Texture); break;
-				default: check(0); break;
-				}
-			}
-			break;
-		case ERCT_SetShaderResourceViewParameter:
-			{
-				auto* RHICmd = Iter.NextCommand<FRHICommandSetShaderResourceViewParameter>();
-				switch (RHICmd->ShaderFrequency)
-				{
-				case SF_Vertex:		SetShaderResourceViewParameter_Internal((FVertexShaderRHIParamRef)RHICmd->Shader, RHICmd->SamplerIndex, RHICmd->SRV); break;
-				case SF_Hull:		SetShaderResourceViewParameter_Internal((FHullShaderRHIParamRef)RHICmd->Shader, RHICmd->SamplerIndex, RHICmd->SRV); break;
-				case SF_Domain:		SetShaderResourceViewParameter_Internal((FDomainShaderRHIParamRef)RHICmd->Shader, RHICmd->SamplerIndex, RHICmd->SRV); break;
-				case SF_Geometry:	SetShaderResourceViewParameter_Internal((FGeometryShaderRHIParamRef)RHICmd->Shader, RHICmd->SamplerIndex, RHICmd->SRV); break;
-				case SF_Pixel:		SetShaderResourceViewParameter_Internal((FPixelShaderRHIParamRef)RHICmd->Shader, RHICmd->SamplerIndex, RHICmd->SRV); break;
-				case SF_Compute:	SetShaderResourceViewParameter_Internal((FComputeShaderRHIParamRef)RHICmd->Shader, RHICmd->SamplerIndex, RHICmd->SRV); break;
-				default: check(0); break;
-				}
-			}
-			break;
-		case ERCT_SetUAVParameter:
-			{
-				auto* RHICmd = Iter.NextCommand<FRHICommandSetUAVParameter>();
-				switch (RHICmd->ShaderFrequency)
-				{
-				case SF_Compute:	SetUAVParameter_Internal((FComputeShaderRHIParamRef)RHICmd->Shader, RHICmd->UAVIndex, RHICmd->UAV); break;
-				default: check(0); break;
-				}
-			}
-			break;
-		case ERCT_SetUAVParameter_IntialCount:
-			{
-				auto* RHICmd = Iter.NextCommand<FRHICommandSetUAVParameter_IntialCount>();
-				switch (RHICmd->ShaderFrequency)
-				{
-				case SF_Compute:	SetUAVParameter_Internal((FComputeShaderRHIParamRef)RHICmd->Shader, RHICmd->UAVIndex, RHICmd->UAV, RHICmd->InitialCount); break;
-				default: check(0); break;
-				}
-			}
-			break;
-		case ERCT_DrawPrimitive:
-			{
-				auto* RHICmd = Iter.NextCommand<FRHICommandDrawPrimitive>();
-				DrawPrimitive_Internal(RHICmd->PrimitiveType, RHICmd->BaseVertexIndex, RHICmd->NumPrimitives, RHICmd->NumInstances);
-			}
-			break;
-		case ERCT_DrawIndexedPrimitive:
-			{
-				auto* RHICmd = Iter.NextCommand<FRHICommandDrawIndexedPrimitive>();
-				DrawIndexedPrimitive_Internal(RHICmd->IndexBuffer, RHICmd->PrimitiveType, RHICmd->BaseVertexIndex, RHICmd->MinIndex, RHICmd->NumVertices, RHICmd->StartIndex, RHICmd->NumPrimitives, RHICmd->NumInstances);
-			}
-			break;
-		case ERCT_SetBoundShaderState:
-			{
-				auto* RHICmd = Iter.NextCommand<FRHICommandSetBoundShaderState>();
-				SetBoundShaderState_Internal(RHICmd->BoundShaderState);
-			}
-			break;
-		case ERCT_SetBlendState:
-			{
-				auto* RHICmd = Iter.NextCommand<FRHICommandSetBlendState>();
-				SetBlendState_Internal(RHICmd->State, RHICmd->BlendFactor);
-			}
-			break;
-		case ERCT_SetStreamSource:
-			{
-				auto* RHICmd = Iter.NextCommand<FRHICommandSetStreamSource>();
-				SetStreamSource_Internal(RHICmd->StreamIndex, RHICmd->VertexBuffer, RHICmd->Stride, RHICmd->Offset);
-			}
-			break;
-		case ERCT_SetDepthStencilState:
-			{
-				auto* RHICmd = Iter.NextCommand<FRHICommandSetDepthStencilState>();
-				SetDepthStencilState_Internal(RHICmd->State, RHICmd->StencilRef);
-			}
-			break;
-		case ERCT_SetViewport:
-			{
-				auto* RHICmd = Iter.NextCommand<FRHICommandSetViewport>();
-				SetViewport_Internal(RHICmd->MinX, RHICmd->MinY, RHICmd->MinZ, RHICmd->MaxX, RHICmd->MaxY, RHICmd->MaxZ);
-			}
-			break;
-		case ERCT_SetScissorRect:
-			{
-				auto* RHICmd = Iter.NextCommand<FRHICommandSetScissorRect>();
-				SetScissorRect_Internal(RHICmd->bEnable, RHICmd->MinX, RHICmd->MinY, RHICmd->MaxX, RHICmd->MaxY);
-			}
-			break;
-		case ERCT_SetRenderTargets:
-			{
-				auto* RHICmd = Iter.NextCommand<FRHICommandSetRenderTargets>();
-				SetRenderTargets_Internal(
-					RHICmd->NewNumSimultaneousRenderTargets,
-					RHICmd->NewRenderTargetsRHI,
-					RHICmd->NewDepthStencilTargetRHI,
-					RHICmd->NewNumUAVs,
-					RHICmd->UAVs);
-			}
-			break;
-		case ERCT_EndDrawPrimitiveUP:
-			{
-				auto* RHICmd = Iter.NextCommand<FRHICommandEndDrawPrimitiveUP>();
-				void* Buffer = NULL;
-				BeginDrawPrimitiveUP_Internal(RHICmd->PrimitiveType, RHICmd->NumPrimitives, RHICmd->NumVertices, RHICmd->VertexDataStride, Buffer);
-				FMemory::Memcpy(Buffer, RHICmd->OutVertexData, RHICmd->NumVertices * RHICmd->VertexDataStride);
-				EndDrawPrimitiveUP_Internal();
-				FMemory::Free(RHICmd->OutVertexData);
-			}
-			break;
-		case ERCT_EndDrawIndexedPrimitiveUP:
-			{
-				auto* RHICmd = Iter.NextCommand<FRHICommandEndDrawIndexedPrimitiveUP>();
-				void* VertexBuffer = NULL;
-				void* IndexBuffer = NULL;
-				BeginDrawIndexedPrimitiveUP_Internal(
-					RHICmd->PrimitiveType,
-					RHICmd->NumPrimitives,
-					RHICmd->NumVertices,
-					RHICmd->VertexDataStride,
-					VertexBuffer,
-					RHICmd->MinVertexIndex,
-					RHICmd->NumIndices,
-					RHICmd->IndexDataStride,
-					IndexBuffer);
-				FMemory::Memcpy(VertexBuffer, RHICmd->OutVertexData, RHICmd->NumVertices * RHICmd->VertexDataStride);
-				FMemory::Memcpy(IndexBuffer, RHICmd->OutIndexData, RHICmd->NumIndices * RHICmd->IndexDataStride);
-				EndDrawIndexedPrimitiveUP_Internal();
-				FMemory::Free(RHICmd->OutVertexData);
-				FMemory::Free(RHICmd->OutIndexData);
-			}
-			break;
-		case ERCT_BuildLocalBoundShaderState:
-			{
-				auto* RHICmd = Iter.NextCommand<FRHICommandBuildLocalBoundShaderState>();
-				check(RHICmd->WorkArea.CheckCmdList == &CmdList && RHICmd->WorkArea.UID == CmdList.UID // this BSS was not built for this particular commandlist
-					&& !IsValidRef(RHICmd->WorkArea.BSS)); // should not already have been created
-				if (RHICmd->WorkArea.UseCount)
-				{
-					RHICmd->WorkArea.BSS = CreateBoundShaderState_Internal(RHICmd->WorkArea.Args.VertexDeclarationRHI, RHICmd->WorkArea.Args.VertexShaderRHI, RHICmd->WorkArea.Args.HullShaderRHI, RHICmd->WorkArea.Args.DomainShaderRHI, RHICmd->WorkArea.Args.PixelShaderRHI, RHICmd->WorkArea.Args.GeometryShaderRHI);
-				}
-			}
-			break;
-		case ERCT_SetLocalBoundShaderState:
-			{
-				auto* RHICmd = Iter.NextCommand<FRHICommandSetLocalBoundShaderState>();
-				check(RHICmd->LocalBoundShaderState.WorkArea->CheckCmdList == &CmdList && RHICmd->LocalBoundShaderState.WorkArea->UID == CmdList.UID // this BSS was not built for this particular commandlist
-					&& RHICmd->LocalBoundShaderState.WorkArea->UseCount > 0 && IsValidRef(RHICmd->LocalBoundShaderState.WorkArea->BSS)); // this should have been created and should have uses outstanding
-
-				SetBoundShaderState_Internal(RHICmd->LocalBoundShaderState.WorkArea->BSS);
-
-				if (--RHICmd->LocalBoundShaderState.WorkArea->UseCount == 0)
-				{
-					RHICmd->LocalBoundShaderState.WorkArea->BSS = nullptr; // this also releases the ref, which wouldn't otherwise happen because RHICommands are unsafe and don't do constructors or destructors
-				}
-			}
-			break;
-		case ERCT_SetGlobalBoundShaderState:
-			{
-				auto* RHICmd = Iter.NextCommand<FRHICommandSetGlobalBoundShaderState>();
-				SetGlobalBoundShaderState_InternalPtr(RHICmd->GlobalBoundShaderState);
-			}
-			break;
-
-		case ERCT_SetComputeShader:
-			{
-				auto* RHICmd = Iter.NextCommand<FRHICommandSetComputeShader>();
-				SetComputeShader_Internal(RHICmd->ComputeShader);
-			}
-			break;
-		case ERCT_DispatchComputeShader:
-			{
-				auto* RHICmd = Iter.NextCommand<FRHICommandDispatchComputeShader>();
-				DispatchComputeShader_Internal(RHICmd->ThreadGroupCountX, RHICmd->ThreadGroupCountY, RHICmd->ThreadGroupCountZ);
-			}
-			break;
-		case ERCT_DispatchIndirectComputeShader:
-			{
-				auto* RHICmd = Iter.NextCommand<FRHICommandDispatchIndirectComputeShader>();
-				DispatchIndirectComputeShader_Internal(RHICmd->ArgumentBuffer, RHICmd->ArgumentOffset);
-			}
-			break;
-		case ERCT_AutomaticCacheFlushAfterComputeShader:
-			{
-				auto* RHICmd = Iter.NextCommand<FRHICommandAutomaticCacheFlushAfterComputeShader>();
-				AutomaticCacheFlushAfterComputeShader_Internal(RHICmd->bEnable);
-			}
-			break;
-		case ERCT_FlushComputeShaderCache:
-			{
-				auto* RHICmd = Iter.NextCommand<FRHICommandFlushComputeShaderCache>();
-				FlushComputeShaderCache_Internal();
-			}
-			break;
-		case ERCT_DrawPrimitiveIndirect:
-			{
-				auto* RHICmd = Iter.NextCommand<FRHICommandDrawPrimitiveIndirect>();
-				DrawPrimitiveIndirect_Internal(RHICmd->PrimitiveType, RHICmd->ArgumentBuffer, RHICmd->ArgumentOffset);
-			}
-			break;
-		case ERCT_DrawIndexedIndirect:
-			{
-				auto* RHICmd = Iter.NextCommand<FRHICommandDrawIndexedIndirect>();
-				DrawIndexedIndirect_Internal(RHICmd->IndexBufferRHI, RHICmd->PrimitiveType, RHICmd->ArgumentsBufferRHI, RHICmd->DrawArgumentsIndex, RHICmd->NumInstances);
-			}
-			break;
-		case ERCT_DrawIndexedPrimitiveIndirect:
-			{
-				auto* RHICmd = Iter.NextCommand<FRHICommandDrawIndexedPrimitiveIndirect>();
-				DrawIndexedPrimitiveIndirect_Internal(RHICmd->PrimitiveType, RHICmd->IndexBuffer, RHICmd->ArgumentsBuffer, RHICmd->ArgumentOffset);
-			}
-			break;
-		case ERCT_EnableDepthBoundsTest:
-			{
-				auto* RHICmd = Iter.NextCommand<FRHICommandEnableDepthBoundsTest>();
-				EnableDepthBoundsTest_Internal(RHICmd->bEnable, RHICmd->MinDepth, RHICmd->MaxDepth);
-			}
-			break;
-		case ERCT_ClearUAV:
-			{
-				auto* RHICmd = Iter.NextCommand<FRHICommandClearUAV>();
-				ClearUAV_Internal(RHICmd->UnorderedAccessViewRHI, RHICmd->Values);
-			}
-			break;
-		case ERCT_CopyToResolveTarget:
-			{
-				auto* RHICmd = Iter.NextCommand<FRHICommandCopyToResolveTarget>();
-				CopyToResolveTarget_Internal(RHICmd->SourceTexture, RHICmd->DestTexture, RHICmd->bKeepOriginalSurface, RHICmd->ResolveParams);
-			}
-			break;
-		case ERCT_Clear:
-			{
-				auto* RHICmd = Iter.NextCommand<FRHICommandClear>();
-				Clear_Internal(RHICmd->bClearColor, RHICmd->Color, RHICmd->bClearDepth, RHICmd->Depth, RHICmd->bClearStencil, RHICmd->Stencil, RHICmd->ExcludeRect);
-			}
-			break;
-		case ERCT_ClearMRT:
-			{
-				auto* RHICmd = Iter.NextCommand<FRHICommandClearMRT>();
-				ClearMRT_Internal(RHICmd->bClearColor, RHICmd->NumClearColors, RHICmd->ColorArray, RHICmd->bClearDepth, RHICmd->Depth, RHICmd->bClearStencil, RHICmd->Stencil, RHICmd->ExcludeRect);
-			}
-			break;
-		case ERCT_BuildLocalUniformBuffer:
-			{
-				auto* RHICmd = Iter.NextCommand<FRHICommandBuildLocalUniformBuffer>();
-				check(RHICmd->WorkArea.CheckCmdList == &CmdList && RHICmd->WorkArea.UID == CmdList.UID // this uniform buffer was not built for this particular commandlist
-					&& !IsValidRef(RHICmd->WorkArea.UniformBuffer) && RHICmd->WorkArea.Layout && RHICmd->WorkArea.Contents); // should not already have been created
-				if (RHICmd->WorkArea.UseCount)
-				{
-					RHICmd->WorkArea.UniformBuffer = RHICreateUniformBuffer(RHICmd->WorkArea.Contents, *RHICmd->WorkArea.Layout, UniformBuffer_SingleFrame);
-				}
-				FMemory::Free(RHICmd->WorkArea.Contents);
-				RHICmd->WorkArea.Contents = nullptr;
-			}
-			break;
-		case ERCT_SetLocalUniformBuffer:
-			{
-				auto* RHICmd = Iter.NextCommand<FRHICommandSetLocalUniformBuffer>();
-				check(RHICmd->LocalUniformBuffer.WorkArea->CheckCmdList == &CmdList && RHICmd->LocalUniformBuffer.WorkArea->UID == CmdList.UID // this BSS was not built for this particular commandlist
-					&& RHICmd->LocalUniformBuffer.WorkArea->UseCount > 0 && IsValidRef(RHICmd->LocalUniformBuffer.WorkArea->UniformBuffer)); // this should have been created and should have uses outstanding
-				switch (RHICmd->ShaderFrequency)
-				{
-					case SF_Vertex:		SetShaderUniformBuffer_Internal((FVertexShaderRHIParamRef)RHICmd->Shader, RHICmd->BaseIndex, RHICmd->LocalUniformBuffer.WorkArea->UniformBuffer); break;
-					case SF_Hull:		SetShaderUniformBuffer_Internal((FHullShaderRHIParamRef)RHICmd->Shader, RHICmd->BaseIndex, RHICmd->LocalUniformBuffer.WorkArea->UniformBuffer); break;
-					case SF_Domain:		SetShaderUniformBuffer_Internal((FDomainShaderRHIParamRef)RHICmd->Shader, RHICmd->BaseIndex, RHICmd->LocalUniformBuffer.WorkArea->UniformBuffer); break;
-					case SF_Geometry:	SetShaderUniformBuffer_Internal((FGeometryShaderRHIParamRef)RHICmd->Shader, RHICmd->BaseIndex, RHICmd->LocalUniformBuffer.WorkArea->UniformBuffer); break;
-					case SF_Pixel:		SetShaderUniformBuffer_Internal((FPixelShaderRHIParamRef)RHICmd->Shader, RHICmd->BaseIndex, RHICmd->LocalUniformBuffer.WorkArea->UniformBuffer); break;
-					case SF_Compute:	SetShaderUniformBuffer_Internal((FComputeShaderRHIParamRef)RHICmd->Shader, RHICmd->BaseIndex, RHICmd->LocalUniformBuffer.WorkArea->UniformBuffer); break;
-					default: check(0); break;
-				}
-				if (--RHICmd->LocalUniformBuffer.WorkArea->UseCount == 0)
-				{
-					RHICmd->LocalUniformBuffer.WorkArea->UniformBuffer = nullptr; // this also releases the ref, which wouldn't otherwise happen because RHICommands are unsafe and don't do constructors or destructors
-					RHICmd->LocalUniformBuffer.WorkArea->Layout = nullptr;
-				}
-			}
-			break;
-
-		default:
-			checkf(0, TEXT("Unknown RHI Command %d!"), Cmd->Type);
+			FRHICommandBase* Cmd = Iter.NextCommand();
+			Cmd->CallExecuteAndDestruct();
 		}
 	}
-
-	Iter.CheckNumCommands();
-}
-
-void FRHICommandListExecutor::ExecuteInner(FRHICommandList& CmdList)
-{
-	CmdList.bExecuting = true;
-	if (CVarRHICmdExec.GetValueOnRenderThread() != 0)
-	{
-		RHIExecuteCommandList(&CmdList);
-	}
-	else
-	{
-		ExecuteListGenericRHI(CmdList);
-	}
 	CmdList.Reset();
 }
 
-void FRHICommandListExecutor::ExecuteList(FRHICommandList& CmdList)
+void FRHICommandListExecutor::ExecuteList(FRHICommandListBase& CmdList)
 {
 	check(IsInRenderingThread() && &CmdList != &GetImmediateCommandList());
 
@@ -474,75 +95,31 @@ void FRHICommandListExecutor::CheckNoOutstandingCmdLists()
 	check(GRHICommandList.OutstandingCmdListCount.GetValue() == 1); // else we are attempting to delete resources while there is still a live cmdlist (other than the immediate cmd list) somewhere.
 }
 
-FRHICommandList::FMemManager::FMemManager() :
-	FirstPage(nullptr),
-	LastPage(nullptr)
-{
-}
-
-FRHICommandList::FMemManager::~FMemManager()
-{
-	while (FirstPage)
-	{
-		auto* Page = FirstPage->NextPage;
-		delete FirstPage;
-		FirstPage = Page;
-	}
-
-	FirstPage = LastPage = nullptr;
-}
-
-const SIZE_T FRHICommandList::FMemManager::GetUsedMemory() const
-{
-	SIZE_T TotalUsedMemory = 0;
-	auto* Page = FirstPage;
-	do
-	{
-		SIZE_T UsedPageMemory = Page->MemUsed();
-		if (!UsedPageMemory)
-		{
-			return TotalUsedMemory;
-		}
-		TotalUsedMemory += UsedPageMemory;
-		Page = Page->NextPage;
-	}
-	while (Page);
-	return TotalUsedMemory;
-}
-
-void FRHICommandList::FMemManager::Reset()
-{
-	auto* Page = FirstPage;
-	while (Page)
-	{
-		Page->Reset();
-		Page = Page->NextPage;
-	}
-	LastPage = FirstPage;
-}
-
-FRHICommandList::FRHICommandList()
+FRHICommandListBase::FRHICommandListBase()
+	: MemManager(0)
 {
 	GRHICommandList.OutstandingCmdListCount.Increment();
 	Reset();
 }
 
-FRHICommandList::~FRHICommandList()
+FRHICommandListBase::~FRHICommandListBase()
 {
 	Flush();
 	GRHICommandList.OutstandingCmdListCount.Decrement();
 }
 
-const SIZE_T FRHICommandList::GetUsedMemory() const
+const int32 FRHICommandListBase::GetUsedMemory() const
 {
-	return MemManager.GetUsedMemory();
+	return MemManager.GetByteCount();
 }
 
-void FRHICommandList::Reset()
+void FRHICommandListBase::Reset()
 {
 	bExecuting = false;
+	MemManager.Flush();
 	NumCommands = 0;
-	MemManager.Reset();
+	Root = nullptr;
+	CommandLink = &Root;
 	UID = GRHICommandList.UIDCounter.Increment();
 }
 
diff --git a/Engine/Source/Runtime/RHI/Public/RHICommandList.h b/Engine/Source/Runtime/RHI/Public/RHICommandList.h
--- a/Engine/Source/Runtime/RHI/Public/RHICommandList.h
+++ b/Engine/Source/Runtime/RHI/Public/RHICommandList.h
@@ -19,266 +19,303 @@ DECLARE_DWORD_COUNTER_STAT_EXTERN(TEXT("Immed. Command List memory"), STAT_Immed
 DECLARE_DWORD_COUNTER_STAT_EXTERN(TEXT("Immed. Command count"), STAT_ImmedCmdListCount, STATGROUP_RHICMDLIST, RHI_API);
 
 
-enum ERHICommandType
-{
-	// Start the enum at non-zero to help debugging
-
-	// List control
-	ERCT_NopBlob				= 512,	// NoOp, maybe has a buffer to be auto alloc/freed (512=0x200)
-	ERCT_NopEndOfPage,					// NoOp, stop traversing the current page
-
-	// RHI Commands
-	ERCT_SetRasterizerState,
-	ERCT_SetShaderParameter,
-	ERCT_SetShaderUniformBuffer,
-	ERCT_SetShaderTexture,
-	ERCT_SetShaderSampler,
-	ERCT_DrawPrimitive,
-	ERCT_DrawIndexedPrimitive,
-	ERCT_SetBoundShaderState,
-	ERCT_SetBlendState,
-	ERCT_SetStreamSource,
-	ERCT_SetDepthStencilState,
-	ERCT_SetShaderResourceViewParameter,
-	ERCT_SetUAVParameter,
-	ERCT_SetUAVParameter_IntialCount,
-	ERCT_SetViewport,
-	ERCT_SetScissorRect,
-	ERCT_SetRenderTargets,
-	ERCT_EndDrawPrimitiveUP,
-	ERCT_EndDrawIndexedPrimitiveUP,
-	ERCT_BuildLocalBoundShaderState,
-	ERCT_SetLocalBoundShaderState,
-	ERCT_SetGlobalBoundShaderState,
-	ERCT_SetComputeShader,
-	ERCT_DispatchComputeShader,
-	ERCT_DispatchIndirectComputeShader,
-	ERCT_AutomaticCacheFlushAfterComputeShader,
-	ERCT_FlushComputeShaderCache,
-	ERCT_DrawPrimitiveIndirect,
-	ERCT_DrawIndexedIndirect,
-	ERCT_DrawIndexedPrimitiveIndirect,
-	ERCT_EnableDepthBoundsTest,
-	ERCT_ClearUAV,
-	ERCT_CopyToResolveTarget,
-	ERCT_ClearMRT,
-	ERCT_Clear,
-	ERCT_BuildLocalUniformBuffer,
-	ERCT_SetLocalUniformBuffer,
-
-
-
-
-	ERCT_NUM,
-};
 
 extern RHI_API TAutoConsoleVariable<int32> CVarRHICmdWidth;
 
-struct FRHICommand
+struct FRHICommandBase
 {
-	// This is a union only to force alignment to SIZE_T
-	union
-	{
-		ERHICommandType	Type;
-		SIZE_T			RawType;
-	};
-
-	uint32 ExtraSize() const
+	FRHICommandBase* Next;
+	void(*ExecuteAndDestructPtr)(FRHICommandBase *Cmd);
+	FORCEINLINE FRHICommandBase(void(*InExecuteAndDestructPtr)(FRHICommandBase *Cmd))
+		: Next(nullptr)
+		, ExecuteAndDestructPtr(InExecuteAndDestructPtr)
 	{
-		return 0;
 	}
-
-	enum
+	FORCEINLINE void CallExecuteAndDestruct()
 	{
-		IsEndOfPage = 0,
-	};
+		ExecuteAndDestructPtr(this);
+	}
 };
 
-
-struct FRHICommandPerShader : public FRHICommand
+class RHI_API FRHICommandListBase : public FNoncopyable
 {
-	FRHIResource* Shader;
-	EShaderFrequency ShaderFrequency;
+public:
 
-	template <typename TShaderRHIParamRef>
-	FORCEINLINE_DEBUGGABLE void SetShader(TShaderRHIParamRef InShader)
+	FRHICommandListBase();
+	~FRHICommandListBase();
+
+	inline void Flush();
+	const int32 GetUsedMemory() const;
+
+	FORCEINLINE_DEBUGGABLE void* Alloc(int32 AllocSize, int32 Alignment)
 	{
-		Shader = InShader;
-		ShaderFrequency = (EShaderFrequency)TRHIShaderToEnum<TShaderRHIParamRef>::ShaderFrequency;
+		return MemManager.Alloc(AllocSize, Alignment);
 	}
-};
 
-struct FRHICommandNopEndOfPage : public FRHICommand
-{
-	enum
+	template <typename T>
+	FORCEINLINE_DEBUGGABLE void* Alloc()
 	{
-		IsEndOfPage = 1
+		return MemManager.Alloc(sizeof(T), ALIGNOF(T));
+	}
+
+	template <typename TCmd>
+	FORCEINLINE_DEBUGGABLE void* AllocCommand()
+	{
+		checkSlow(!bExecuting);
+		TCmd* Result = (TCmd*)Alloc<TCmd>();
+		++NumCommands;
+		*CommandLink = Result;
+		CommandLink = &Result->Next;
+		return Result;
+	}
+	FORCEINLINE uint32 GetUID()
+	{
+		return UID;
+	}
+	FORCEINLINE bool HasCommands() const
+	{
+		return (NumCommands > 0);
+	}
+
+	inline bool Bypass();
+
+	struct FDrawUpData
+	{
+
+		uint32 PrimitiveType;
+		uint32 NumPrimitives;
+		uint32 NumVertices;
+		uint32 VertexDataStride;
+		void* OutVertexData;
+		uint32 MinVertexIndex;
+		uint32 NumIndices;
+		uint32 IndexDataStride;
+		void* OutIndexData;
+
+		FDrawUpData()
+			: PrimitiveType(PT_Num)
+			, OutVertexData(nullptr)
+			, OutIndexData(nullptr)
+		{
+		}
+
 	};
+
+	FDrawUpData DrawUPData;
+
+private:
+	FMemStackBase MemManager;
+	FRHICommandBase* Root;
+	FRHICommandBase** CommandLink;
+	bool bExecuting;
+	uint32 NumCommands;
+	uint32 UID;
+
+
+	void Reset();
+
+	friend class FRHICommandListExecutor;
+	friend class FRHICommandListIterator;
+
 };
 
-struct FRHICommandNopBlob : public FRHICommand
+template<typename TCmd>
+struct FRHICommand : public FRHICommandBase
 {
-	void* Data;
-	uint32 Size;
-	FORCEINLINE_DEBUGGABLE void Set(FRHICommandList* List, uint32 InSize, const void* InData)
+	FORCEINLINE FRHICommand()
+		: FRHICommandBase(&ExecuteAndDestruct)
 	{
-		Type = ERCT_NopBlob;
-		Size = InSize;
 
-		// Data is stored after 'this'
-		Data = &this[1];
-		FMemory::Memcpy(Data, InData, InSize);
 	}
-
-	uint32 ExtraSize() const
+	static FORCEINLINE void ExecuteAndDestruct(FRHICommandBase *Cmd)
 	{
-		return Size;
+		TCmd *ThisCmd = (TCmd*)Cmd;
+		ThisCmd->Execute();
+		ThisCmd->~TCmd();
 	}
 };
 
-struct FRHICommandSetRasterizerState : public FRHICommand
+struct FRHICommandSetRasterizerState : public FRHICommand<FRHICommandSetRasterizerState>
 {
 	FRasterizerStateRHIParamRef State;
-	FORCEINLINE_DEBUGGABLE void Set(FRasterizerStateRHIParamRef InState)
+	FORCEINLINE_DEBUGGABLE FRHICommandSetRasterizerState(FRasterizerStateRHIParamRef InState)
+		: State(InState)
 	{
-		Type = ERCT_SetRasterizerState;
-		State = InState;
+	}
+	void Execute()
+	{
+		SetRasterizerState_Internal(State);
 	}
 };
 
-struct FRHICommandSetDepthStencilState : public FRHICommand
+struct FRHICommandSetDepthStencilState : public FRHICommand<FRHICommandSetDepthStencilState>
 {
 	FDepthStencilStateRHIParamRef State;
 	uint32 StencilRef;
-	FORCEINLINE_DEBUGGABLE void Set(FDepthStencilStateRHIParamRef InState, uint32 InStencilRef)
+	FORCEINLINE_DEBUGGABLE FRHICommandSetDepthStencilState(FDepthStencilStateRHIParamRef InState, uint32 InStencilRef)
+		: State(InState)
+		, StencilRef(InStencilRef)
+	{
+	}
+	void Execute()
 	{
-		Type = ERCT_SetDepthStencilState;
-		State = InState;
-		StencilRef = InStencilRef;
+		SetDepthStencilState_Internal(State, StencilRef);
 	}
 };
 
-struct FRHICommandSetShaderParameter : public FRHICommandPerShader
+template <typename TShaderRHIParamRef>
+struct FRHICommandSetShaderParameter : public FRHICommand<FRHICommandSetShaderParameter<TShaderRHIParamRef> >
 {
+	TShaderRHIParamRef Shader;
 	const void* NewValue;
 	uint32 BufferIndex;
 	uint32 BaseIndex;
 	uint32 NumBytes;
-	template <typename TShaderRHIParamRef>
-	FORCEINLINE_DEBUGGABLE void Set(TShaderRHIParamRef InShader, uint32 InBufferIndex, uint32 InBaseIndex, uint32 InNumBytes, const void* InNewValue)
+	FORCEINLINE_DEBUGGABLE FRHICommandSetShaderParameter(TShaderRHIParamRef InShader, uint32 InBufferIndex, uint32 InBaseIndex, uint32 InNumBytes, const void* InNewValue)
+		: Shader(InShader)
+		, NewValue(InNewValue)
+		, BufferIndex(InBufferIndex)
+		, BaseIndex(InBaseIndex)
+		, NumBytes(InNumBytes)
 	{
-		FRHICommandPerShader::SetShader<TShaderRHIParamRef>(InShader);
-		Type = ERCT_SetShaderParameter;
-		BufferIndex = InBufferIndex;
-		BaseIndex = InBaseIndex;
-		NumBytes = InNumBytes;
-		NewValue = InNewValue;
+	}
+	void Execute()
+	{
+		SetShaderParameter_Internal(Shader, BufferIndex, BaseIndex, NumBytes, NewValue); 
 	}
 };
 
-struct FRHICommandSetShaderUniformBuffer : public FRHICommandPerShader
+template <typename TShaderRHIParamRef>
+struct FRHICommandSetShaderUniformBuffer : public FRHICommand<FRHICommandSetShaderUniformBuffer<TShaderRHIParamRef> >
 {
+	TShaderRHIParamRef Shader;
 	uint32 BaseIndex;
 	FUniformBufferRHIParamRef UniformBuffer;
-	template <typename TShaderRHIParamRef>
-	FORCEINLINE_DEBUGGABLE void Set(TShaderRHIParamRef InShader, uint32 InBaseIndex, FUniformBufferRHIParamRef InUniformBuffer)
+	FORCEINLINE_DEBUGGABLE FRHICommandSetShaderUniformBuffer(TShaderRHIParamRef InShader, uint32 InBaseIndex, FUniformBufferRHIParamRef InUniformBuffer)
+		: Shader(InShader)
+		, BaseIndex(InBaseIndex)
+		, UniformBuffer(InUniformBuffer)
 	{
-		FRHICommandPerShader::SetShader<TShaderRHIParamRef>(InShader);
-		Type = ERCT_SetShaderUniformBuffer;
-		BaseIndex = InBaseIndex;
-		UniformBuffer = InUniformBuffer;
+	}
+	void Execute()
+	{
+		SetShaderUniformBuffer_Internal(Shader, BaseIndex, UniformBuffer);
 	}
 };
 
-struct FRHICommandSetShaderTexture : public FRHICommandPerShader
+template <typename TShaderRHIParamRef>
+struct FRHICommandSetShaderTexture : public FRHICommand<FRHICommandSetShaderTexture<TShaderRHIParamRef> >
 {
+	TShaderRHIParamRef Shader;
 	uint32 TextureIndex;
 	FTextureRHIParamRef Texture;
-	template <typename TShaderRHIParamRef>
-	FORCEINLINE_DEBUGGABLE void Set(TShaderRHIParamRef InShader, uint32 InTextureIndex, FTextureRHIParamRef InTexture)
+	FORCEINLINE_DEBUGGABLE FRHICommandSetShaderTexture(TShaderRHIParamRef InShader, uint32 InTextureIndex, FTextureRHIParamRef InTexture)
+		: Shader(InShader)
+		, TextureIndex(InTextureIndex)
+		, Texture(InTexture)
+	{
+	}
+	void Execute()
 	{
-		FRHICommandPerShader::SetShader<TShaderRHIParamRef>(InShader);
-		Type = ERCT_SetShaderTexture;
-		TextureIndex = InTextureIndex;
-		Texture = InTexture;
+		SetShaderTexture_Internal(Shader, TextureIndex, Texture);
 	}
 };
 
-struct FRHICommandSetShaderResourceViewParameter : public FRHICommandPerShader
+template <typename TShaderRHIParamRef>
+struct FRHICommandSetShaderResourceViewParameter : public FRHICommand<FRHICommandSetShaderResourceViewParameter<TShaderRHIParamRef> >
 {
+	TShaderRHIParamRef Shader;
 	uint32 SamplerIndex;
 	FShaderResourceViewRHIParamRef SRV;
-	template <typename TShaderRHIParamRef>
-	FORCEINLINE_DEBUGGABLE void Set(TShaderRHIParamRef InShader, uint32 InSamplerIndex, FShaderResourceViewRHIParamRef InSRV)
+	FORCEINLINE_DEBUGGABLE FRHICommandSetShaderResourceViewParameter(TShaderRHIParamRef InShader, uint32 InSamplerIndex, FShaderResourceViewRHIParamRef InSRV)
+		: Shader(InShader)
+		, SamplerIndex(InSamplerIndex)
+		, SRV(InSRV)
+	{
+	}
+	void Execute()
 	{
-		FRHICommandPerShader::SetShader<TShaderRHIParamRef>(InShader);
-		Type = ERCT_SetShaderResourceViewParameter;
-		SamplerIndex = InSamplerIndex;
-		SRV = InSRV;
+		SetShaderResourceViewParameter_Internal(Shader, SamplerIndex, SRV);
 	}
 };
 
-struct FRHICommandSetUAVParameter : public FRHICommandPerShader
+template <typename TShaderRHIParamRef>
+struct FRHICommandSetUAVParameter : public FRHICommand<FRHICommandSetUAVParameter<TShaderRHIParamRef> >
 {
+	TShaderRHIParamRef Shader;
 	uint32 UAVIndex;
 	FUnorderedAccessViewRHIParamRef UAV;
-	template <typename TShaderRHIParamRef>
-	FORCEINLINE_DEBUGGABLE void Set(TShaderRHIParamRef InShader, uint32 InUAVIndex, FUnorderedAccessViewRHIParamRef InUAV)
+	FORCEINLINE_DEBUGGABLE FRHICommandSetUAVParameter(TShaderRHIParamRef InShader, uint32 InUAVIndex, FUnorderedAccessViewRHIParamRef InUAV)
+		: Shader(InShader)
+		, UAVIndex(InUAVIndex)
+		, UAV(InUAV)
+	{
+	}
+	void Execute()
 	{
-		FRHICommandPerShader::SetShader<TShaderRHIParamRef>(InShader);
-		Type = ERCT_SetUAVParameter;
-		UAVIndex = InUAVIndex;
-		UAV = InUAV;
+		SetUAVParameter_Internal(Shader, UAVIndex, UAV);
 	}
 };
 
-struct FRHICommandSetUAVParameter_IntialCount : public FRHICommandPerShader
+template <typename TShaderRHIParamRef>
+struct FRHICommandSetUAVParameter_IntialCount : public FRHICommand<FRHICommandSetUAVParameter_IntialCount<TShaderRHIParamRef> >
 {
+	TShaderRHIParamRef Shader;
 	uint32 UAVIndex;
 	FUnorderedAccessViewRHIParamRef UAV;
 	uint32 InitialCount;
-	template <typename TShaderRHIParamRef>
-	FORCEINLINE_DEBUGGABLE void Set(TShaderRHIParamRef InShader, uint32 InUAVIndex, FUnorderedAccessViewRHIParamRef InUAV, uint32 InInitialCount)
+	FORCEINLINE_DEBUGGABLE FRHICommandSetUAVParameter_IntialCount(TShaderRHIParamRef InShader, uint32 InUAVIndex, FUnorderedAccessViewRHIParamRef InUAV, uint32 InInitialCount)
+		: Shader(InShader)
+		, UAVIndex(InUAVIndex)
+		, UAV(InUAV)
+		, InitialCount(InInitialCount)
+	{
+	}
+	void Execute()
 	{
-		FRHICommandPerShader::SetShader<TShaderRHIParamRef>(InShader);
-		Type = ERCT_SetUAVParameter_IntialCount;
-		UAVIndex = InUAVIndex;
-		UAV = InUAV;
-		InitialCount = InInitialCount;
+		SetUAVParameter_Internal(Shader, UAVIndex, UAV, InitialCount);
 	}
 };
 
-struct FRHICommandSetShaderSampler : public FRHICommandPerShader
+template <typename TShaderRHIParamRef>
+struct FRHICommandSetShaderSampler : public FRHICommand<FRHICommandSetShaderSampler<TShaderRHIParamRef> >
 {
+	TShaderRHIParamRef Shader;
 	uint32 SamplerIndex;
 	FSamplerStateRHIParamRef Sampler;
-	template <typename TShaderRHIParamRef>
-	FORCEINLINE_DEBUGGABLE void Set(TShaderRHIParamRef InShader, uint32 InSamplerIndex, FSamplerStateRHIParamRef InSampler)
+	FORCEINLINE_DEBUGGABLE FRHICommandSetShaderSampler(TShaderRHIParamRef InShader, uint32 InSamplerIndex, FSamplerStateRHIParamRef InSampler)
+		: Shader(InShader)
+		, SamplerIndex(InSamplerIndex)
+		, Sampler(InSampler)
 	{
-		FRHICommandPerShader::SetShader<TShaderRHIParamRef>(InShader);
-		Type = ERCT_SetShaderSampler;
-		SamplerIndex = InSamplerIndex;
-		Sampler = InSampler;
+	}
+	void Execute()
+	{
+		SetShaderSampler_Internal(Shader, SamplerIndex, Sampler);
 	}
 };
 
-struct FRHICommandDrawPrimitive : public FRHICommand
+struct FRHICommandDrawPrimitive : public FRHICommand<FRHICommandDrawPrimitive>
 {
 	uint32 PrimitiveType;
 	uint32 BaseVertexIndex;
 	uint32 NumPrimitives;
 	uint32 NumInstances;
-	FORCEINLINE_DEBUGGABLE void Set(uint32 InPrimitiveType, uint32 InBaseVertexIndex, uint32 InNumPrimitives, uint32 InNumInstances)
+	FORCEINLINE_DEBUGGABLE FRHICommandDrawPrimitive(uint32 InPrimitiveType, uint32 InBaseVertexIndex, uint32 InNumPrimitives, uint32 InNumInstances)
+		: PrimitiveType(InPrimitiveType)
+		, BaseVertexIndex(InBaseVertexIndex)
+		, NumPrimitives(InNumPrimitives)
+		, NumInstances(InNumInstances)
+	{
+	}
+	void Execute()
 	{
-		Type = ERCT_DrawPrimitive;
-		PrimitiveType = InPrimitiveType;
-		BaseVertexIndex = InBaseVertexIndex;
-		NumPrimitives = InNumPrimitives;
-		NumInstances = InNumInstances;
+		DrawPrimitive_Internal(PrimitiveType, BaseVertexIndex, NumPrimitives, NumInstances);
 	}
+
 };
 
-struct FRHICommandDrawIndexedPrimitive : public FRHICommand
+struct FRHICommandDrawIndexedPrimitive : public FRHICommand<FRHICommandDrawIndexedPrimitive>
 {
 	FIndexBufferRHIParamRef IndexBuffer;
 	uint32 PrimitiveType;
@@ -288,59 +325,71 @@ struct FRHICommandDrawIndexedPrimitive : public FRHICommand
 	uint32 StartIndex;
 	uint32 NumPrimitives;
 	uint32 NumInstances;
-	FORCEINLINE_DEBUGGABLE void Set(FIndexBufferRHIParamRef InIndexBuffer, uint32 InPrimitiveType, int32 InBaseVertexIndex, uint32 InMinIndex, uint32 InNumVertices, uint32 InStartIndex, uint32 InNumPrimitives, uint32 InNumInstances)
-	{
-		Type = ERCT_DrawIndexedPrimitive;
-		IndexBuffer = InIndexBuffer;
-		PrimitiveType = InPrimitiveType;
-		BaseVertexIndex = InBaseVertexIndex;
-		MinIndex = InMinIndex;
-		NumVertices = InNumVertices;
-		StartIndex = InStartIndex;
-		NumPrimitives = InNumPrimitives;
-		NumInstances = InNumInstances;
+	FORCEINLINE_DEBUGGABLE FRHICommandDrawIndexedPrimitive(FIndexBufferRHIParamRef InIndexBuffer, uint32 InPrimitiveType, int32 InBaseVertexIndex, uint32 InMinIndex, uint32 InNumVertices, uint32 InStartIndex, uint32 InNumPrimitives, uint32 InNumInstances)
+		: IndexBuffer(InIndexBuffer)
+		, PrimitiveType(InPrimitiveType)
+		, BaseVertexIndex(InBaseVertexIndex)
+		, MinIndex(InMinIndex)
+		, NumVertices(InNumVertices)
+		, StartIndex(InStartIndex)
+		, NumPrimitives(InNumPrimitives)
+		, NumInstances(InNumInstances)
+	{
+	}
+	void Execute()
+	{
+		DrawIndexedPrimitive_Internal(IndexBuffer, PrimitiveType, BaseVertexIndex, MinIndex, NumVertices, StartIndex, NumPrimitives, NumInstances);
 	}
 };
 
-struct FRHICommandSetBoundShaderState : public FRHICommand
+struct FRHICommandSetBoundShaderState : public FRHICommand<FRHICommandSetBoundShaderState>
 {
 	FBoundShaderStateRHIParamRef BoundShaderState;
-	FORCEINLINE_DEBUGGABLE void Set(FBoundShaderStateRHIParamRef InBoundShaderState)
+	FORCEINLINE_DEBUGGABLE FRHICommandSetBoundShaderState(FBoundShaderStateRHIParamRef InBoundShaderState)
+		: BoundShaderState(InBoundShaderState)
 	{
-		Type = ERCT_SetBoundShaderState;
-		BoundShaderState = InBoundShaderState;
+	}
+	void Execute()
+	{
+		SetBoundShaderState_Internal(BoundShaderState);
 	}
 };
 
-struct FRHICommandSetBlendState : public FRHICommand
+struct FRHICommandSetBlendState : public FRHICommand<FRHICommandSetBlendState>
 {
 	FBlendStateRHIParamRef State;
 	FLinearColor BlendFactor;
-	FORCEINLINE_DEBUGGABLE void Set(FBlendStateRHIParamRef InState, const FLinearColor& InBlendFactor)
+	FORCEINLINE_DEBUGGABLE FRHICommandSetBlendState(FBlendStateRHIParamRef InState, const FLinearColor& InBlendFactor)
+		: State(InState)
+		, BlendFactor(InBlendFactor)
 	{
-		Type = ERCT_SetBlendState;
-		State = InState;
-		BlendFactor = InBlendFactor;
+	}
+	void Execute()
+	{
+		SetBlendState_Internal(State, BlendFactor);
 	}
 };
 
-struct FRHICommandSetStreamSource : public FRHICommand
+struct FRHICommandSetStreamSource : public FRHICommand<FRHICommandSetStreamSource>
 {
 	uint32 StreamIndex;
 	FVertexBufferRHIParamRef VertexBuffer;
 	uint32 Stride;
 	uint32 Offset;
-	FORCEINLINE_DEBUGGABLE void Set(uint32 InStreamIndex, FVertexBufferRHIParamRef InVertexBuffer, uint32 InStride, uint32 InOffset)
+	FORCEINLINE_DEBUGGABLE FRHICommandSetStreamSource(uint32 InStreamIndex, FVertexBufferRHIParamRef InVertexBuffer, uint32 InStride, uint32 InOffset)
+		: StreamIndex(InStreamIndex)
+		, VertexBuffer(InVertexBuffer)
+		, Stride(InStride)
+		, Offset(InOffset)
+	{
+	}
+	void Execute()
 	{
-		Type = ERCT_SetStreamSource;
-		StreamIndex = InStreamIndex;
-		VertexBuffer = InVertexBuffer;
-		Stride = InStride;
-		Offset = InOffset;
+		SetStreamSource_Internal(StreamIndex, VertexBuffer, Stride, Offset);
 	}
 };
 
-struct FRHICommandSetViewport : public FRHICommand
+struct FRHICommandSetViewport : public FRHICommand<FRHICommandSetViewport>
 {
 	uint32 MinX;
 	uint32 MinY;
@@ -348,37 +397,43 @@ struct FRHICommandSetViewport : public FRHICommand
 	uint32 MaxX;
 	uint32 MaxY;
 	float MaxZ;
-	FORCEINLINE_DEBUGGABLE void Set(uint32 InMinX, uint32 InMinY, float InMinZ, uint32 InMaxX, uint32 InMaxY, float InMaxZ)
+	FORCEINLINE_DEBUGGABLE FRHICommandSetViewport(uint32 InMinX, uint32 InMinY, float InMinZ, uint32 InMaxX, uint32 InMaxY, float InMaxZ)
+		: MinX(InMinX)
+		, MinY(InMinY)
+		, MinZ(InMinZ)
+		, MaxX(InMaxX)
+		, MaxY(InMaxY)
+		, MaxZ(InMaxZ)
+	{
+	}
+	void Execute()
 	{
-		Type = ERCT_SetViewport;
-		MinX = InMinX;
-		MinY = InMinY;
-		MinZ = InMinZ;
-		MaxX = InMaxX;
-		MaxY = InMaxY;
-		MaxZ = InMaxZ;
+		SetViewport_Internal(MinX, MinY, MinZ, MaxX, MaxY, MaxZ);
 	}
 };
 
-struct FRHICommandSetScissorRect : public FRHICommand
+struct FRHICommandSetScissorRect : public FRHICommand<FRHICommandSetScissorRect>
 {
 	bool bEnable;
 	uint32 MinX;
 	uint32 MinY;
 	uint32 MaxX;
 	uint32 MaxY;
-	FORCEINLINE_DEBUGGABLE void Set(bool InbEnable, uint32 InMinX, uint32 InMinY, uint32 InMaxX, uint32 InMaxY)
+	FORCEINLINE_DEBUGGABLE FRHICommandSetScissorRect(bool InbEnable, uint32 InMinX, uint32 InMinY, uint32 InMaxX, uint32 InMaxY)
+		: bEnable(InbEnable)
+		, MinX(InMinX)
+		, MinY(InMinY)
+		, MaxX(InMaxX)
+		, MaxY(InMaxY)
 	{
-		Type = ERCT_SetScissorRect;
-		bEnable = InbEnable;
-		MinX = InMinX;
-		MinY = InMinY;
-		MaxX = InMaxX;
-		MaxY = InMaxY;
+	}
+	void Execute()
+	{
+		SetScissorRect_Internal(bEnable, MinX, MinY, MaxX, MaxY);
 	}
 };
 
-struct FRHICommandSetRenderTargets : public FRHICommand
+struct FRHICommandSetRenderTargets : public FRHICommand<FRHICommandSetRenderTargets>
 {
 	uint32 NewNumSimultaneousRenderTargets;
 	FRHIRenderTargetView NewRenderTargetsRHI[MaxSimultaneousRenderTargets];
@@ -386,32 +441,40 @@ struct FRHICommandSetRenderTargets : public FRHICommand
 	uint32 NewNumUAVs;
 	FUnorderedAccessViewRHIParamRef UAVs[MaxSimultaneousUAVs];
 
-	FORCEINLINE_DEBUGGABLE void Set(
+	FORCEINLINE_DEBUGGABLE FRHICommandSetRenderTargets(
 		uint32 InNewNumSimultaneousRenderTargets,
 		const FRHIRenderTargetView* InNewRenderTargetsRHI,
 		FTextureRHIParamRef InNewDepthStencilTargetRHI,
 		uint32 InNewNumUAVs,
 		const FUnorderedAccessViewRHIParamRef* InUAVs
 		)
+		: NewNumSimultaneousRenderTargets(InNewNumSimultaneousRenderTargets)
+		, NewDepthStencilTargetRHI(InNewDepthStencilTargetRHI)
+		, NewNumUAVs(InNewNumUAVs)
 
 	{
-		Type = ERCT_SetRenderTargets;
 		check(InNewNumSimultaneousRenderTargets <= MaxSimultaneousRenderTargets && InNewNumUAVs <= MaxSimultaneousUAVs);
-		NewNumSimultaneousRenderTargets = InNewNumSimultaneousRenderTargets;
 		for (uint32 Index = 0; Index < NewNumSimultaneousRenderTargets; Index++)
 		{
 			NewRenderTargetsRHI[Index] = InNewRenderTargetsRHI[Index];
 		}
-		NewDepthStencilTargetRHI = InNewDepthStencilTargetRHI;
-		NewNumUAVs = InNewNumUAVs;
 		for (uint32 Index = 0; Index < NewNumUAVs; Index++)
 		{
 			UAVs[Index] = InUAVs[Index];
 		}
 	}
+	void Execute()
+	{
+		SetRenderTargets_Internal(
+			NewNumSimultaneousRenderTargets,
+			NewRenderTargetsRHI,
+			NewDepthStencilTargetRHI,
+			NewNumUAVs,
+			UAVs);
+	}
 };
 
-struct FRHICommandEndDrawPrimitiveUP : public FRHICommand
+struct FRHICommandEndDrawPrimitiveUP : public FRHICommand<FRHICommandEndDrawPrimitiveUP>
 {
 	uint32 PrimitiveType;
 	uint32 NumPrimitives;
@@ -419,19 +482,25 @@ struct FRHICommandEndDrawPrimitiveUP : public FRHICommand
 	uint32 VertexDataStride;
 	void* OutVertexData;
 
-	FORCEINLINE_DEBUGGABLE void Set(uint32 InPrimitiveType, uint32 InNumPrimitives, uint32 InNumVertices, uint32 InVertexDataStride, void* InOutVertexData)
+	FORCEINLINE_DEBUGGABLE FRHICommandEndDrawPrimitiveUP(uint32 InPrimitiveType, uint32 InNumPrimitives, uint32 InNumVertices, uint32 InVertexDataStride, void* InOutVertexData)
+		: PrimitiveType(InPrimitiveType)
+		, NumPrimitives(InNumPrimitives)
+		, NumVertices(InNumVertices)
+		, VertexDataStride(InVertexDataStride)
+		, OutVertexData(InOutVertexData)
+	{
+	}
+	void Execute()
 	{
-		Type = ERCT_EndDrawPrimitiveUP;
-		PrimitiveType = InPrimitiveType;
-		NumPrimitives = InNumPrimitives;
-		NumVertices = InNumVertices;
-		VertexDataStride = InVertexDataStride;
-		OutVertexData = InOutVertexData;
+		void* Buffer = NULL;
+		BeginDrawPrimitiveUP_Internal(PrimitiveType, NumPrimitives, NumVertices, VertexDataStride, Buffer);
+		FMemory::Memcpy(Buffer, OutVertexData, NumVertices * VertexDataStride);
+		EndDrawPrimitiveUP_Internal();
 	}
 };
 
 
-struct FRHICommandEndDrawIndexedPrimitiveUP : public FRHICommand
+struct FRHICommandEndDrawIndexedPrimitiveUP : public FRHICommand<FRHICommandEndDrawIndexedPrimitiveUP>
 {
 	uint32 PrimitiveType;
 	uint32 NumPrimitives;
@@ -443,92 +512,123 @@ struct FRHICommandEndDrawIndexedPrimitiveUP : public FRHICommand
 	uint32 IndexDataStride;
 	void* OutIndexData;
 
-	FORCEINLINE_DEBUGGABLE void Set(uint32 InPrimitiveType, uint32 InNumPrimitives, uint32 InNumVertices, uint32 InVertexDataStride, void* InOutVertexData, uint32 InMinVertexIndex, uint32 InNumIndices, uint32 InIndexDataStride, void* InOutIndexData)
-	{
-		Type = ERCT_EndDrawIndexedPrimitiveUP;
-		PrimitiveType = InPrimitiveType;
-		NumPrimitives = InNumPrimitives;
-		NumVertices = InNumVertices;
-		VertexDataStride = InVertexDataStride;
-		OutVertexData = InOutVertexData;
-		MinVertexIndex = InMinVertexIndex;
-		NumIndices = InNumIndices;
-		IndexDataStride = InIndexDataStride;
-		OutIndexData = InOutIndexData;
+	FORCEINLINE_DEBUGGABLE FRHICommandEndDrawIndexedPrimitiveUP(uint32 InPrimitiveType, uint32 InNumPrimitives, uint32 InNumVertices, uint32 InVertexDataStride, void* InOutVertexData, uint32 InMinVertexIndex, uint32 InNumIndices, uint32 InIndexDataStride, void* InOutIndexData)
+		: PrimitiveType(InPrimitiveType)
+		, NumPrimitives(InNumPrimitives)
+		, NumVertices(InNumVertices)
+		, VertexDataStride(InVertexDataStride)
+		, OutVertexData(InOutVertexData)
+		, MinVertexIndex(InMinVertexIndex)
+		, NumIndices(InNumIndices)
+		, IndexDataStride(InIndexDataStride)
+		, OutIndexData(InOutIndexData)
+	{
+	}
+	void Execute()
+	{
+		void* VertexBuffer = nullptr;
+		void* IndexBuffer = nullptr;
+		BeginDrawIndexedPrimitiveUP_Internal(
+			PrimitiveType,
+			NumPrimitives,
+			NumVertices,
+			VertexDataStride,
+			VertexBuffer,
+			MinVertexIndex,
+			NumIndices,
+			IndexDataStride,
+			IndexBuffer);
+		FMemory::Memcpy(VertexBuffer, OutVertexData, NumVertices * VertexDataStride);
+		FMemory::Memcpy(IndexBuffer, OutIndexData, NumIndices * IndexDataStride);
+		EndDrawIndexedPrimitiveUP_Internal();
 	}
 };
 
-struct FRHICommandSetComputeShader : public FRHICommand
+struct FRHICommandSetComputeShader : public FRHICommand<FRHICommandSetComputeShader>
 {
 	FComputeShaderRHIParamRef ComputeShader;
-	FORCEINLINE_DEBUGGABLE void Set(FComputeShaderRHIParamRef InComputeShader)
+	FORCEINLINE_DEBUGGABLE FRHICommandSetComputeShader(FComputeShaderRHIParamRef InComputeShader)
+		: ComputeShader(InComputeShader)
 	{
-		Type = ERCT_SetComputeShader;
-		ComputeShader = InComputeShader;
+	}
+	void Execute()
+	{
+		SetComputeShader_Internal(ComputeShader);
 	}
 };
 
-struct FRHICommandDispatchComputeShader : public FRHICommand
+struct FRHICommandDispatchComputeShader : public FRHICommand<FRHICommandDispatchComputeShader>
 {
 	uint32 ThreadGroupCountX;
 	uint32 ThreadGroupCountY;
 	uint32 ThreadGroupCountZ;
-	FORCEINLINE_DEBUGGABLE void Set(uint32 InThreadGroupCountX, uint32 InThreadGroupCountY, uint32 InThreadGroupCountZ)
+	FORCEINLINE_DEBUGGABLE FRHICommandDispatchComputeShader(uint32 InThreadGroupCountX, uint32 InThreadGroupCountY, uint32 InThreadGroupCountZ)
+		: ThreadGroupCountX(InThreadGroupCountX)
+		, ThreadGroupCountY(InThreadGroupCountY)
+		, ThreadGroupCountZ(InThreadGroupCountZ)
+	{
+	}
+	void Execute()
 	{
-		Type = ERCT_DispatchComputeShader;
-		ThreadGroupCountX = InThreadGroupCountX;
-		ThreadGroupCountY = InThreadGroupCountY;
-		ThreadGroupCountZ = InThreadGroupCountZ;
+		DispatchComputeShader_Internal(ThreadGroupCountX, ThreadGroupCountY, ThreadGroupCountZ);
 	}
+
 };
 
-struct FRHICommandDispatchIndirectComputeShader : public FRHICommand
+struct FRHICommandDispatchIndirectComputeShader : public FRHICommand<FRHICommandDispatchIndirectComputeShader>
 {
 	FVertexBufferRHIParamRef ArgumentBuffer;
 	uint32 ArgumentOffset;
-	FORCEINLINE_DEBUGGABLE void Set(FVertexBufferRHIParamRef InArgumentBuffer, uint32 InArgumentOffset)
+	FORCEINLINE_DEBUGGABLE FRHICommandDispatchIndirectComputeShader(FVertexBufferRHIParamRef InArgumentBuffer, uint32 InArgumentOffset)
+		: ArgumentBuffer(InArgumentBuffer)
+		, ArgumentOffset(InArgumentOffset)
 	{
-		Type = ERCT_DispatchIndirectComputeShader;
-		ArgumentBuffer = InArgumentBuffer;
-		ArgumentOffset = InArgumentOffset;
-
+	}
+	void Execute()
+	{
+		DispatchIndirectComputeShader_Internal(ArgumentBuffer, ArgumentOffset);
 	}
 };
 
-struct FRHICommandAutomaticCacheFlushAfterComputeShader : public FRHICommand
+struct FRHICommandAutomaticCacheFlushAfterComputeShader : public FRHICommand<FRHICommandAutomaticCacheFlushAfterComputeShader>
 {
 	bool bEnable;
-	FORCEINLINE_DEBUGGABLE void Set(bool InbEnable)
+	FORCEINLINE_DEBUGGABLE FRHICommandAutomaticCacheFlushAfterComputeShader(bool InbEnable)
+		: bEnable(InbEnable)
 	{
-		Type = ERCT_AutomaticCacheFlushAfterComputeShader;
-		bEnable = InbEnable;
+	}
+	void Execute()
+	{
+		AutomaticCacheFlushAfterComputeShader_Internal(bEnable);
 	}
 };
 
-struct FRHICommandFlushComputeShaderCache : public FRHICommand
+struct FRHICommandFlushComputeShaderCache : public FRHICommand<FRHICommandFlushComputeShaderCache>
 {
-	FORCEINLINE_DEBUGGABLE void Set()
+	void Execute()
 	{
-		Type = ERCT_FlushComputeShaderCache;
+		FlushComputeShaderCache_Internal();
 	}
 };
 
-struct FRHICommandDrawPrimitiveIndirect : public FRHICommand
+struct FRHICommandDrawPrimitiveIndirect : public FRHICommand<FRHICommandDrawPrimitiveIndirect>
 {
 	uint32 PrimitiveType;
 	FVertexBufferRHIParamRef ArgumentBuffer;
 	uint32 ArgumentOffset;
-	FORCEINLINE_DEBUGGABLE void Set(uint32 InPrimitiveType, FVertexBufferRHIParamRef InArgumentBuffer, uint32 InArgumentOffset)
+	FORCEINLINE_DEBUGGABLE FRHICommandDrawPrimitiveIndirect(uint32 InPrimitiveType, FVertexBufferRHIParamRef InArgumentBuffer, uint32 InArgumentOffset)
+		: PrimitiveType(InPrimitiveType)
+		, ArgumentBuffer(InArgumentBuffer)
+		, ArgumentOffset(InArgumentOffset)
 	{
-		Type = ERCT_DrawPrimitiveIndirect;
-		PrimitiveType = InPrimitiveType;
-		ArgumentBuffer = InArgumentBuffer;
-		ArgumentOffset = InArgumentOffset;
-
+	}
+	void Execute()
+	{
+		DrawPrimitiveIndirect_Internal(PrimitiveType, ArgumentBuffer, ArgumentOffset);
 	}
 };
 
-struct FRHICommandDrawIndexedIndirect : public FRHICommand
+struct FRHICommandDrawIndexedIndirect : public FRHICommand<FRHICommandDrawIndexedIndirect>
 {
 	FIndexBufferRHIParamRef IndexBufferRHI;
 	uint32 PrimitiveType;
@@ -536,83 +636,99 @@ struct FRHICommandDrawIndexedIndirect : public FRHICommand
 	uint32 DrawArgumentsIndex;
 	uint32 NumInstances;
 
-	FORCEINLINE_DEBUGGABLE void Set(FIndexBufferRHIParamRef InIndexBufferRHI, uint32 InPrimitiveType, FStructuredBufferRHIParamRef InArgumentsBufferRHI, uint32 InDrawArgumentsIndex, uint32 InNumInstances)
+	FORCEINLINE_DEBUGGABLE FRHICommandDrawIndexedIndirect(FIndexBufferRHIParamRef InIndexBufferRHI, uint32 InPrimitiveType, FStructuredBufferRHIParamRef InArgumentsBufferRHI, uint32 InDrawArgumentsIndex, uint32 InNumInstances)
+		: IndexBufferRHI(InIndexBufferRHI)
+		, PrimitiveType(InPrimitiveType)
+		, ArgumentsBufferRHI(InArgumentsBufferRHI)
+		, DrawArgumentsIndex(InDrawArgumentsIndex)
+		, NumInstances(InNumInstances)
 	{
-		Type = ERCT_DrawIndexedIndirect;
-		IndexBufferRHI = InIndexBufferRHI;
-		PrimitiveType = InPrimitiveType;
-		ArgumentsBufferRHI = InArgumentsBufferRHI;
-		DrawArgumentsIndex = InDrawArgumentsIndex;
-		NumInstances = InNumInstances;
+	}
+	void Execute()
+	{
+		DrawIndexedIndirect_Internal(IndexBufferRHI, PrimitiveType, ArgumentsBufferRHI, DrawArgumentsIndex, NumInstances);
 	}
 };
 
-struct FRHICommandDrawIndexedPrimitiveIndirect : public FRHICommand
+struct FRHICommandDrawIndexedPrimitiveIndirect : public FRHICommand<FRHICommandDrawIndexedPrimitiveIndirect>
 {
 	uint32 PrimitiveType;
 	FIndexBufferRHIParamRef IndexBuffer;
 	FVertexBufferRHIParamRef ArgumentsBuffer;
 	uint32 ArgumentOffset;
 
-	FORCEINLINE_DEBUGGABLE void Set(uint32 InPrimitiveType, FIndexBufferRHIParamRef InIndexBuffer, FVertexBufferRHIParamRef InArgumentsBuffer, uint32 InArgumentOffset)
+	FORCEINLINE_DEBUGGABLE FRHICommandDrawIndexedPrimitiveIndirect(uint32 InPrimitiveType, FIndexBufferRHIParamRef InIndexBuffer, FVertexBufferRHIParamRef InArgumentsBuffer, uint32 InArgumentOffset)
+		: PrimitiveType(InPrimitiveType)
+		, IndexBuffer(InIndexBuffer)
+		, ArgumentsBuffer(InArgumentsBuffer)
+		, ArgumentOffset(InArgumentOffset)
+	{
+	}
+	void Execute()
 	{
-		Type = ERCT_DrawIndexedPrimitiveIndirect;
-		PrimitiveType = InPrimitiveType;
-		IndexBuffer = InIndexBuffer;
-		ArgumentsBuffer = InArgumentsBuffer;
-		ArgumentOffset = InArgumentOffset;
+		DrawIndexedPrimitiveIndirect_Internal(PrimitiveType, IndexBuffer, ArgumentsBuffer, ArgumentOffset);
 	}
 };
 
-struct FRHICommandEnableDepthBoundsTest : public FRHICommand
+struct FRHICommandEnableDepthBoundsTest : public FRHICommand<FRHICommandEnableDepthBoundsTest>
 {
 	bool bEnable;
 	float MinDepth;
 	float MaxDepth;
 
-	FORCEINLINE_DEBUGGABLE void Set(bool InbEnable, float InMinDepth, float InMaxDepth)
+	FORCEINLINE_DEBUGGABLE FRHICommandEnableDepthBoundsTest(bool InbEnable, float InMinDepth, float InMaxDepth)
+		: bEnable(InbEnable)
+		, MinDepth(InMinDepth)
+		, MaxDepth(InMaxDepth)
+	{
+	}
+	void Execute()
 	{
-		Type = ERCT_EnableDepthBoundsTest;
-		bEnable = InbEnable;
-		MinDepth = InMinDepth;
-		MaxDepth = InMaxDepth;
+		EnableDepthBoundsTest_Internal(bEnable, MinDepth, MaxDepth);
 	}
 };
 
-struct FRHICommandClearUAV : public FRHICommand
+struct FRHICommandClearUAV : public FRHICommand<FRHICommandClearUAV>
 {
 	FUnorderedAccessViewRHIParamRef UnorderedAccessViewRHI;
 	uint32 Values[4];
 
-	FORCEINLINE_DEBUGGABLE void Set(FUnorderedAccessViewRHIParamRef InUnorderedAccessViewRHI, const uint32* InValues)
+	FORCEINLINE_DEBUGGABLE FRHICommandClearUAV(FUnorderedAccessViewRHIParamRef InUnorderedAccessViewRHI, const uint32* InValues)
+		: UnorderedAccessViewRHI(InUnorderedAccessViewRHI)
 	{
-		Type = ERCT_ClearUAV;
-		UnorderedAccessViewRHI = InUnorderedAccessViewRHI;
 		Values[0] = InValues[0];
 		Values[1] = InValues[1];
 		Values[2] = InValues[2];
 		Values[3] = InValues[3];
 	}
+	void Execute()
+	{
+		ClearUAV_Internal(UnorderedAccessViewRHI, Values);
+	}
 };
 
-struct FRHICommandCopyToResolveTarget : public FRHICommand
+struct FRHICommandCopyToResolveTarget : public FRHICommand<FRHICommandCopyToResolveTarget>
 {
 	FTextureRHIParamRef SourceTexture;
 	FTextureRHIParamRef DestTexture;
 	bool bKeepOriginalSurface;
 	FResolveParams ResolveParams;
 
-	FORCEINLINE_DEBUGGABLE void Set(FTextureRHIParamRef InSourceTexture, FTextureRHIParamRef InDestTexture, bool InbKeepOriginalSurface, const FResolveParams& InResolveParams)
+	FORCEINLINE_DEBUGGABLE FRHICommandCopyToResolveTarget(FTextureRHIParamRef InSourceTexture, FTextureRHIParamRef InDestTexture, bool InbKeepOriginalSurface, const FResolveParams& InResolveParams)
+		: SourceTexture(InSourceTexture)
+		, DestTexture(InDestTexture)
+		, bKeepOriginalSurface(InbKeepOriginalSurface)
+		, ResolveParams(InResolveParams)
 	{
-		Type = ERCT_CopyToResolveTarget;
-		SourceTexture = InSourceTexture;
-		DestTexture = InDestTexture;
-		bKeepOriginalSurface = InbKeepOriginalSurface;
-		ResolveParams = InResolveParams;
 	}
+	void Execute()
+	{
+		CopyToResolveTarget_Internal(SourceTexture, DestTexture, bKeepOriginalSurface, ResolveParams);
+	}
+
 };
 
-struct FRHICommandClear : public FRHICommand
+struct FRHICommandClear : public FRHICommand<FRHICommandClear>
 {
 	FLinearColor Color;
 	FIntRect ExcludeRect;
@@ -621,7 +737,8 @@ struct FRHICommandClear : public FRHICommand
 	bool bClearColor;
 	bool bClearDepth;
 	bool bClearStencil;
-	FORCEINLINE_DEBUGGABLE void Set(
+
+	FORCEINLINE_DEBUGGABLE FRHICommandClear(
 		bool InbClearColor,
 		const FLinearColor& InColor,
 		bool InbClearDepth,
@@ -630,19 +747,22 @@ struct FRHICommandClear : public FRHICommand
 		uint32 InStencil,
 		FIntRect InExcludeRect
 		)
+		: Color(InColor)
+		, ExcludeRect(InExcludeRect)
+		, Depth(InDepth)
+		, Stencil(InStencil)
+		, bClearColor(InbClearColor)
+		, bClearDepth(InbClearDepth)
+		, bClearStencil(InbClearStencil)
+	{
+	}
+	void Execute()
 	{
-		Type = ERCT_Clear;
-		Color = InColor;
-		ExcludeRect = InExcludeRect;
-		Depth = InDepth;
-		Stencil = InStencil;
-		bClearColor = InbClearColor;
-		bClearDepth = InbClearDepth;
-		bClearStencil = InbClearStencil;
+		Clear_Internal(bClearColor, Color, bClearDepth, Depth, bClearStencil, Stencil, ExcludeRect);
 	}
 };
 
-struct FRHICommandClearMRT : public FRHICommand
+struct FRHICommandClearMRT : public FRHICommand<FRHICommandClearMRT>
 {
 	FLinearColor ColorArray[MaxSimultaneousRenderTargets];
 	FIntRect ExcludeRect;
@@ -652,7 +772,8 @@ struct FRHICommandClearMRT : public FRHICommand
 	bool bClearColor;
 	bool bClearDepth;
 	bool bClearStencil;
-	FORCEINLINE_DEBUGGABLE void Set(
+
+	FORCEINLINE_DEBUGGABLE FRHICommandClearMRT(
 		bool InbClearColor,
 		int32 InNumClearColors,
 		const FLinearColor* InColorArray,
@@ -662,20 +783,23 @@ struct FRHICommandClearMRT : public FRHICommand
 		uint32 InStencil,
 		FIntRect InExcludeRect
 		)
+		: ExcludeRect(InExcludeRect)
+		, Depth(InDepth)
+		, Stencil(InStencil)
+		, NumClearColors(InNumClearColors)
+		, bClearColor(InbClearColor)
+		, bClearDepth(InbClearDepth)
+		, bClearStencil(InbClearStencil)
 	{
-		Type = ERCT_ClearMRT;
 		check(InNumClearColors < MaxSimultaneousRenderTargets);
 		for (int32 Index = 0; Index < InNumClearColors; Index++)
 		{
 			ColorArray[Index] = InColorArray[Index];
 		}
-		ExcludeRect = InExcludeRect;
-		Depth = InDepth;
-		Stencil = InStencil;
-		NumClearColors = InNumClearColors;
-		bClearColor = InbClearColor;
-		bClearDepth = InbClearDepth;
-		bClearStencil = InbClearStencil;
+	}
+	void Execute()
+	{
+		ClearMRT_Internal(bClearColor, NumClearColors, ColorArray, bClearDepth, Depth, bClearStencil, Stencil, ExcludeRect);
 	}
 };
 
@@ -710,16 +834,42 @@ struct FBoundShaderStateInput
 	}
 };
 
+struct FComputedBSS
+{
+	FBoundShaderStateRHIRef BSS;
+	int32 UseCount;
+	FComputedBSS()
+		: UseCount(0)
+	{
+	}
+};
 
 struct FLocalBoundShaderStateWorkArea
 {
 	FBoundShaderStateInput Args;
-	FBoundShaderStateRHIRef BSS;
-	int32 UseCount;
+	FComputedBSS* ComputedBSS;
 #if !(UE_BUILD_SHIPPING || UE_BUILD_TEST)
-	class FRHICommandList* CheckCmdList;
+	FRHICommandListBase* CheckCmdList;
 	int32 UID;
 #endif
+
+	FORCEINLINE_DEBUGGABLE FLocalBoundShaderStateWorkArea(
+		FRHICommandListBase* InCheckCmdList,
+		FVertexDeclarationRHIParamRef VertexDeclarationRHI,
+		FVertexShaderRHIParamRef VertexShaderRHI,
+		FHullShaderRHIParamRef HullShaderRHI,
+		FDomainShaderRHIParamRef DomainShaderRHI,
+		FPixelShaderRHIParamRef PixelShaderRHI,
+		FGeometryShaderRHIParamRef GeometryShaderRHI
+		)
+		: Args(VertexDeclarationRHI, VertexShaderRHI, HullShaderRHI, DomainShaderRHI, PixelShaderRHI, GeometryShaderRHI)
+#if !(UE_BUILD_SHIPPING || UE_BUILD_TEST)
+		, CheckCmdList(InCheckCmdList)
+		, UID(InCheckCmdList->GetUID())
+#endif
+	{
+		ComputedBSS = new (InCheckCmdList->Alloc<FComputedBSS>()) FComputedBSS;
+	}
 };
 
 struct FLocalBoundShaderState
@@ -730,14 +880,19 @@ struct FLocalBoundShaderState
 		: WorkArea(nullptr)
 	{
 	}
+	FLocalBoundShaderState(const FLocalBoundShaderState& Other)
+		: WorkArea(Other.WorkArea)
+		, BypassBSS(Other.BypassBSS)
+	{
+	}
 };
 
-struct FRHICommandBuildLocalBoundShaderState : public FRHICommand
+struct FRHICommandBuildLocalBoundShaderState : public FRHICommand<FRHICommandBuildLocalBoundShaderState>
 {
 	FLocalBoundShaderStateWorkArea WorkArea;
-	FORCEINLINE_DEBUGGABLE void Set(
-		class FRHICommandList* CheckCmdList,
-		int32 UID,
+
+	FORCEINLINE_DEBUGGABLE FRHICommandBuildLocalBoundShaderState(
+		FRHICommandListBase* CheckCmdList,
 		FVertexDeclarationRHIParamRef VertexDeclarationRHI,
 		FVertexShaderRHIParamRef VertexShaderRHI,
 		FHullShaderRHIParamRef HullShaderRHI,
@@ -745,98 +900,78 @@ struct FRHICommandBuildLocalBoundShaderState : public FRHICommand
 		FPixelShaderRHIParamRef PixelShaderRHI,
 		FGeometryShaderRHIParamRef GeometryShaderRHI
 		)
+		: WorkArea(CheckCmdList, VertexDeclarationRHI, VertexShaderRHI, HullShaderRHI, DomainShaderRHI, PixelShaderRHI, GeometryShaderRHI)
+
 	{
-		Type = ERCT_BuildLocalBoundShaderState;
-		WorkArea.Args.VertexDeclarationRHI = VertexDeclarationRHI;
-		WorkArea.Args.VertexShaderRHI = VertexShaderRHI;
-		WorkArea.Args.HullShaderRHI = HullShaderRHI;
-		WorkArea.Args.DomainShaderRHI = DomainShaderRHI;
-		WorkArea.Args.PixelShaderRHI = PixelShaderRHI;
-		WorkArea.Args.GeometryShaderRHI = GeometryShaderRHI;
-
-		// these RHI commands are bad, no constructor! we construct this in place
-		new ((void*)&WorkArea.BSS) FBoundShaderStateRHIRef();
-		WorkArea.UseCount = 0;
-#if !(UE_BUILD_SHIPPING || UE_BUILD_TEST)
-		WorkArea.CheckCmdList = CheckCmdList;
-		WorkArea.UID = UID;
-#endif
 	}
-};
-
-struct FRHICommandSetLocalBoundShaderState : public FRHICommand
-{
-	FLocalBoundShaderState LocalBoundShaderState;
-	FORCEINLINE_DEBUGGABLE void Set(class FRHICommandList* CheckCmdList, int32 UID, FLocalBoundShaderState& InLocalBoundShaderState)
+	void Execute()
 	{
-		Type = ERCT_SetLocalBoundShaderState;
-		LocalBoundShaderState.WorkArea = InLocalBoundShaderState.WorkArea;
-		check(CheckCmdList == LocalBoundShaderState.WorkArea->CheckCmdList && UID == LocalBoundShaderState.WorkArea->UID); // this BSS was not built for this particular commandlist
-		LocalBoundShaderState.WorkArea->UseCount++;
+		check(!IsValidRef(WorkArea.ComputedBSS->BSS)); // should not already have been created
+		if (WorkArea.ComputedBSS->UseCount)
+		{
+			WorkArea.ComputedBSS->BSS = CreateBoundShaderState_Internal(WorkArea.Args.VertexDeclarationRHI, WorkArea.Args.VertexShaderRHI, WorkArea.Args.HullShaderRHI, WorkArea.Args.DomainShaderRHI, WorkArea.Args.PixelShaderRHI, WorkArea.Args.GeometryShaderRHI);
+		}
 	}
-};
-
-
-class FShader;
 
-struct FGlobalBoundShaderStateArgs
-{
-	FVertexDeclarationRHIParamRef VertexDeclarationRHI;
-	FShader* VertexShader;
-	FShader* PixelShader;
-	FShader* GeometryShader;
 };
 
-class FGlobalBoundShaderStateResource;
-template<class ResourceType> class TGlobalResource;
-
-
-struct FGlobalBoundShaderStateWorkArea
+struct FRHICommandSetLocalBoundShaderState : public FRHICommand<FRHICommandSetLocalBoundShaderState>
 {
-	FGlobalBoundShaderStateArgs Args;
-	TGlobalResource<FGlobalBoundShaderStateResource>* BSS; //ideally this would be part of this memory block and not a separate allocation...that is doable, if a little tedious. The point is we need to delay the construction until we get back to the render thread.
-
-	FGlobalBoundShaderStateWorkArea()
-		: BSS(nullptr)
+	FLocalBoundShaderState LocalBoundShaderState;
+	FORCEINLINE_DEBUGGABLE FRHICommandSetLocalBoundShaderState(FRHICommandListBase* CheckCmdList, FLocalBoundShaderState& InLocalBoundShaderState)
+		: LocalBoundShaderState(InLocalBoundShaderState)
 	{
+		check(CheckCmdList == LocalBoundShaderState.WorkArea->CheckCmdList && CheckCmdList->GetUID() == LocalBoundShaderState.WorkArea->UID); // this BSS was not built for this particular commandlist
+		LocalBoundShaderState.WorkArea->ComputedBSS->UseCount++;
 	}
-};
-
-struct FGlobalBoundShaderState
-{
-public:
+	void Execute()
+	{
+		check(LocalBoundShaderState.WorkArea->ComputedBSS->UseCount > 0 && IsValidRef(LocalBoundShaderState.WorkArea->ComputedBSS->BSS)); // this should have been created and should have uses outstanding
 
-	FGlobalBoundShaderStateWorkArea* Get(ERHIFeatureLevel::Type InFeatureLevel = GRHIFeatureLevel)  { return WorkAreas[InFeatureLevel]; }
-	FGlobalBoundShaderStateWorkArea** GetPtr(ERHIFeatureLevel::Type InFeatureLevel = GRHIFeatureLevel)  { return &WorkAreas[InFeatureLevel]; }
+		SetBoundShaderState_Internal(LocalBoundShaderState.WorkArea->ComputedBSS->BSS);
 
-private:
+		if (--LocalBoundShaderState.WorkArea->ComputedBSS->UseCount == 0)
+		{
+			LocalBoundShaderState.WorkArea->ComputedBSS->~FComputedBSS();
+		}
+	}
 
-	FGlobalBoundShaderStateWorkArea* WorkAreas[ERHIFeatureLevel::Num];
 };
 
-struct FRHICommandSetGlobalBoundShaderState : public FRHICommand
+
+struct FComputedUniformBuffer
 {
-	FGlobalBoundShaderState GlobalBoundShaderState;
-	FORCEINLINE_DEBUGGABLE void Set(FGlobalBoundShaderState& InGlobalBoundShaderState)
+	FUniformBufferRHIRef UniformBuffer;
+	mutable int32 UseCount;
+	FComputedUniformBuffer()
+		: UseCount(0)
 	{
-		Type = ERCT_SetGlobalBoundShaderState;
-		GlobalBoundShaderState = InGlobalBoundShaderState;
 	}
 };
 
-// there is a layering problem here, we use a function pointer to avoid making RHI depend on engine
-extern RHI_API void(*SetGlobalBoundShaderState_InternalPtr)(FGlobalBoundShaderState& BoundShaderState);
 
 struct FLocalUniformBufferWorkArea
 {
 	void* Contents;
 	const FRHIUniformBufferLayout* Layout;
-	FUniformBufferRHIRef UniformBuffer;
-	mutable int32 UseCount;
+	FComputedUniformBuffer* ComputedUniformBuffer;
 #if !(UE_BUILD_SHIPPING || UE_BUILD_TEST)
-	class FRHICommandList* CheckCmdList;
+	FRHICommandListBase* CheckCmdList;
 	int32 UID;
 #endif
+
+	FLocalUniformBufferWorkArea(FRHICommandListBase* InCheckCmdList, const void* InContents, uint32 ContentsSize, const FRHIUniformBufferLayout* InLayout)
+		: Layout(InLayout)
+#if !(UE_BUILD_SHIPPING || UE_BUILD_TEST)
+		, CheckCmdList(InCheckCmdList)
+		, UID(InCheckCmdList->GetUID())
+#endif
+	{
+		check(ContentsSize);
+		Contents = InCheckCmdList->Alloc(ContentsSize, UNIFORM_BUFFER_STRUCT_ALIGNMENT);
+		FMemory::Memcpy(Contents, InContents, ContentsSize);
+		ComputedUniformBuffer = new (InCheckCmdList->Alloc<FComputedUniformBuffer>()) FComputedUniformBuffer;
+	}
 };
 
 struct FLocalUniformBuffer
@@ -847,191 +982,73 @@ struct FLocalUniformBuffer
 		: WorkArea(nullptr)
 	{
 	}
+	FLocalUniformBuffer(const FLocalUniformBuffer& Other)
+		: WorkArea(Other.WorkArea)
+		, BypassUniform(Other.BypassUniform)
+	{
+	}
 	FORCEINLINE_DEBUGGABLE bool IsValid() const
 	{
 		return WorkArea || IsValidRef(BypassUniform);
 	}
 };
 
-struct FRHICommandBuildLocalUniformBuffer : public FRHICommand
+struct FRHICommandBuildLocalUniformBuffer : public FRHICommand<FRHICommandBuildLocalUniformBuffer>
 {
 	FLocalUniformBufferWorkArea WorkArea;
-	FORCEINLINE_DEBUGGABLE void Set(
-	class FRHICommandList* CheckCmdList,
-		int32 UID,
+	FORCEINLINE_DEBUGGABLE FRHICommandBuildLocalUniformBuffer(
+		FRHICommandListBase* CheckCmdList,
 		const void* Contents,
 		uint32 ContentsSize,
 		const FRHIUniformBufferLayout& Layout
 		)
+		: WorkArea(CheckCmdList, Contents, ContentsSize, &Layout)
+
 	{
-		Type = ERCT_BuildLocalUniformBuffer;
-		// these RHI commands are bad, no constructor! we construct this in place
-		WorkArea.Layout = &Layout;
-		WorkArea.Contents = FMemory::Malloc(ContentsSize, UNIFORM_BUFFER_STRUCT_ALIGNMENT);
-		FMemory::Memcpy(WorkArea.Contents, Contents, ContentsSize);
-		new ((void*)&WorkArea.UniformBuffer) FUniformBufferRHIRef();
-		WorkArea.UseCount = 0;
-#if !(UE_BUILD_SHIPPING || UE_BUILD_TEST)
-		WorkArea.CheckCmdList = CheckCmdList;
-		WorkArea.UID = UID;
-#endif
 	}
+	void Execute()
+	{
+		check(!IsValidRef(WorkArea.ComputedUniformBuffer->UniformBuffer) && WorkArea.Layout && WorkArea.Contents); // should not already have been created
+		if (WorkArea.ComputedUniformBuffer->UseCount)
+		{
+			WorkArea.ComputedUniformBuffer->UniformBuffer = RHICreateUniformBuffer(WorkArea.Contents, *WorkArea.Layout, UniformBuffer_SingleFrame);
+		}
+		WorkArea.Layout = nullptr;
+		WorkArea.Contents = nullptr;
+	}
+
 };
 
-struct FRHICommandSetLocalUniformBuffer : public FRHICommandPerShader
+template <typename TShaderRHIParamRef>
+struct FRHICommandSetLocalUniformBuffer : public FRHICommand<FRHICommandSetLocalUniformBuffer<TShaderRHIParamRef> >
 {
+	TShaderRHIParamRef Shader;
 	uint32 BaseIndex;
 	FLocalUniformBuffer LocalUniformBuffer;
-	template <typename TShaderRHIParamRef>
-	FORCEINLINE_DEBUGGABLE void Set(class FRHICommandList* CheckCmdList, int32 UID, TShaderRHIParamRef InShader, uint32 InBaseIndex, const FLocalUniformBuffer& InLocalUniformBuffer)
+	FORCEINLINE_DEBUGGABLE FRHICommandSetLocalUniformBuffer(FRHICommandListBase* CheckCmdList, TShaderRHIParamRef InShader, uint32 InBaseIndex, const FLocalUniformBuffer& InLocalUniformBuffer)
+		: Shader(InShader)
+		, BaseIndex(InBaseIndex)
+		, LocalUniformBuffer(InLocalUniformBuffer)
+
 	{
-		FRHICommandPerShader::SetShader<TShaderRHIParamRef>(InShader);
-		Type = ERCT_SetLocalUniformBuffer;
-		BaseIndex = InBaseIndex;
-		LocalUniformBuffer.WorkArea = InLocalUniformBuffer.WorkArea;
-		check(CheckCmdList == LocalUniformBuffer.WorkArea->CheckCmdList && UID == LocalUniformBuffer.WorkArea->UID); // this uniform buffer was not built for this particular commandlist
-		LocalUniformBuffer.WorkArea->UseCount++;
+		check(CheckCmdList == LocalUniformBuffer.WorkArea->CheckCmdList && CheckCmdList->GetUID() == LocalUniformBuffer.WorkArea->UID); // this uniform buffer was not built for this particular commandlist
+		LocalUniformBuffer.WorkArea->ComputedUniformBuffer->UseCount++;
 	}
-};
-
-class RHI_API FRHICommandList : public FNoncopyable
-{
-private:
-
-	struct FMemManager
+	void Execute()
 	{
-		struct FPage
-		{
-			enum
-			{
-				PageSize = 64 * 1024
-			};
-
-			RHI_API static TLockFreeFixedSizeAllocator<PageSize> TheAllocator;
-
-			FPage() :
-				NextPage(nullptr)
-			{
-				Head = (uint8*)TheAllocator.Allocate();
-				check((UPTRINT(Head) & (Alignment - 1)) == 0);
-				Current = Head;
-				Tail = Head + PageSize;
-				//::OutputDebugStringW(*FString::Printf(TEXT("PAGE %p, %p - %p\n"), this, Head, Tail));
-			}
-
-			~FPage()
-			{
-				TheAllocator.Free(Head);
-				Head = Current = Tail = nullptr;
-
-				// Caller/owner is responsible for freeing the linked list
-				NextPage = nullptr;
-			}
-
-			void Reset()
-			{
-				Current = Head;
-			}
-
-			FORCEINLINE_DEBUGGABLE SIZE_T MemUsed() const
-			{
-				return Current - Head;
-			}
-
-			FORCEINLINE_DEBUGGABLE SIZE_T MemAvailable() const
-			{
-				return Tail - Current;
-			}
-
-			FORCEINLINE_DEBUGGABLE bool CanFitAllocation(SIZE_T AlignedSize) const
-			{
-				return (AlignedSize + Current <= Tail);
-			}
-
-			FORCEINLINE_DEBUGGABLE uint8* Alloc(SIZE_T AlignedSize)
-			{
-				auto* Ptr = Current;
-				Current += AlignedSize;
-				return Ptr;
-			}
-			FORCEINLINE_DEBUGGABLE void Dealloc(SIZE_T AlignedSize)
-			{
-				Current -= AlignedSize;
-			}
-
-			uint8* Current;
-			uint8* Tail;
-			uint8* Head;
-
-			FPage* NextPage;
-		};
-
-		FMemManager();
-		~FMemManager();
-
-		FORCEINLINE_DEBUGGABLE void Dealloc(SIZE_T InSize)
+		check(LocalUniformBuffer.WorkArea->ComputedUniformBuffer->UseCount > 0 && IsValidRef(LocalUniformBuffer.WorkArea->ComputedUniformBuffer->UniformBuffer)); // this should have been created and should have uses outstanding
+		SetShaderUniformBuffer_Internal(Shader, BaseIndex, LocalUniformBuffer.WorkArea->ComputedUniformBuffer->UniformBuffer);
+		if (--LocalUniformBuffer.WorkArea->ComputedUniformBuffer->UseCount == 0)
 		{
-			checkSlow(LastPage);
-			LastPage->Dealloc(InSize);
+			LocalUniformBuffer.WorkArea->ComputedUniformBuffer->~FComputedUniformBuffer();
 		}
-		FORCEINLINE_DEBUGGABLE uint8* Alloc(SIZE_T InSize)
-		{
-			InSize = (InSize + (Alignment - 1)) & ~(Alignment - 1);
-			bool bTryAgain = false;
-
-			if (LastPage)
-			{
-				if (LastPage->CanFitAllocation(InSize))
-				{
-					return LastPage->Alloc(InSize);
-				}
-
-				// Insert End Of Page if there's trailing space
-				if (LastPage->MemAvailable() >= sizeof(FRHICommand))
-				{
-					//@todo-rco: Fix me!
-					// This is an allocator command, so it should not be counted during regular execution...
-					auto* EndOfPageCmd = (FRHICommand*)LastPage->Alloc(sizeof(FRHICommand));
-					EndOfPageCmd->Type = ERCT_NopEndOfPage;
-				}
-			}
-
-			checkfSlow(InSize <= FPage::PageSize, TEXT("Can't allocate an FPage bigger than %d!"), FPage::PageSize);
-			if (!FirstPage)
-			{
-				FirstPage = new FPage();
-				check(!LastPage);
-				LastPage = FirstPage;
-			}
-			else
-			{
-				LastPage->NextPage = new FPage();
-				LastPage = LastPage->NextPage;
-			}
-			return LastPage->Alloc(InSize);
-		}
-
-		const SIZE_T GetUsedMemory() const;
-
-		void Reset();
-
-		FPage* FirstPage;
-		FPage* LastPage;
-	};
+	}
 
-	FMemManager MemManager;
+};
 
+class RHI_API FRHICommandList : public FRHICommandListBase
+{
 public:
-	enum
-	{
-		Alignment = sizeof(SIZE_T),
-	};
-
-	FRHICommandList();
-	~FRHICommandList();
-
-	inline void Flush();
-
 	FORCEINLINE_DEBUGGABLE FLocalBoundShaderState BuildLocalBoundShaderState(const FBoundShaderStateInput& BoundShaderStateInput)
 	{
 		return BuildLocalBoundShaderState(
@@ -1072,8 +1089,7 @@ public:
 		}
 		else
 		{
-			auto* Cmd = AddCommand<FRHICommandBuildLocalBoundShaderState>();
-			Cmd->Set(this, this->UID, VertexDeclarationRHI, VertexShaderRHI, HullShaderRHI, DomainShaderRHI, PixelShaderRHI, GeometryShaderRHI);
+			auto* Cmd = new (AllocCommand<FRHICommandBuildLocalBoundShaderState>()) FRHICommandBuildLocalBoundShaderState(this, VertexDeclarationRHI, VertexShaderRHI, HullShaderRHI, DomainShaderRHI, PixelShaderRHI, GeometryShaderRHI);
 			Result.WorkArea = &Cmd->WorkArea;
 		}
 		return Result;
@@ -1086,56 +1102,7 @@ public:
 			SetBoundShaderState_Internal(LocalBoundShaderState.BypassBSS);
 			return;
 		}
-		auto* Cmd = AddCommand<FRHICommandSetLocalBoundShaderState>();
-		Cmd->Set(this, this->UID, LocalBoundShaderState);
-	}
-
-	FORCEINLINE_DEBUGGABLE void BuildAndSetGlobalBoundShaderState(
-		FGlobalBoundShaderState& GlobalBoundShaderState,
-		FVertexDeclarationRHIParamRef VertexDeclarationRHI,
-		FShader* VertexShader,
-		FShader* PixelShader,
-		FShader* GeometryShader
-		)
-	{
-		FGlobalBoundShaderStateWorkArea* ExistingGlobalBoundShaderState = GlobalBoundShaderState.Get();
-		if (!ExistingGlobalBoundShaderState)
-		{
-			FGlobalBoundShaderStateWorkArea* NewGlobalBoundShaderState = new FGlobalBoundShaderStateWorkArea();
-			NewGlobalBoundShaderState->Args.VertexDeclarationRHI = VertexDeclarationRHI;
-			NewGlobalBoundShaderState->Args.VertexShader = VertexShader;
-			NewGlobalBoundShaderState->Args.PixelShader = PixelShader;
-			NewGlobalBoundShaderState->Args.GeometryShader = GeometryShader;
-			FPlatformMisc::MemoryBarrier();
-
-			FGlobalBoundShaderStateWorkArea* OldGlobalBoundShaderState = (FGlobalBoundShaderStateWorkArea*)FPlatformAtomics::InterlockedCompareExchangePointer((void**)GlobalBoundShaderState.GetPtr(), NewGlobalBoundShaderState, nullptr);
-
-			if (OldGlobalBoundShaderState != nullptr)
-			{
-				//we lost
-				delete NewGlobalBoundShaderState;
-				check(OldGlobalBoundShaderState == ExistingGlobalBoundShaderState);
-			}
-		}
-		else if (!(VertexDeclarationRHI == ExistingGlobalBoundShaderState->Args.VertexDeclarationRHI &&
-			VertexShader == ExistingGlobalBoundShaderState->Args.VertexShader &&
-			PixelShader == ExistingGlobalBoundShaderState->Args.PixelShader &&
-			GeometryShader == ExistingGlobalBoundShaderState->Args.GeometryShader))
-		{
-			// this is sketchy from a parallel perspective, but assuming the writes are atomic, and probably even if they aren't, this should be ok
-			ExistingGlobalBoundShaderState->Args.VertexDeclarationRHI = VertexDeclarationRHI;
-			ExistingGlobalBoundShaderState->Args.VertexShader = VertexShader;
-			ExistingGlobalBoundShaderState->Args.PixelShader = PixelShader;
-			ExistingGlobalBoundShaderState->Args.GeometryShader = GeometryShader;
-		}
-
-		if (Bypass())
-		{
-			SetGlobalBoundShaderState_InternalPtr(GlobalBoundShaderState);
-			return;
-		}
-		auto* Cmd = AddCommand<FRHICommandSetGlobalBoundShaderState>();
-		Cmd->Set(GlobalBoundShaderState);
+		new (AllocCommand<FRHICommandSetLocalBoundShaderState>()) FRHICommandSetLocalBoundShaderState(this, LocalBoundShaderState);
 	}
 
 	FORCEINLINE_DEBUGGABLE FLocalUniformBuffer BuildLocalUniformBuffer(const void* Contents, uint32 ContentsSize, const FRHIUniformBufferLayout& Layout)
@@ -1147,9 +1114,7 @@ public:
 		}
 		else
 		{
-			auto* Cmd = AddCommand<FRHICommandBuildLocalUniformBuffer>();
-			check(ContentsSize);
-			Cmd->Set(this, this->UID, Contents, ContentsSize, Layout);
+			auto* Cmd = new (AllocCommand<FRHICommandBuildLocalUniformBuffer>()) FRHICommandBuildLocalUniformBuffer(this, Contents, ContentsSize, Layout);
 			Result.WorkArea = &Cmd->WorkArea;
 		}
 		return Result;
@@ -1163,8 +1128,7 @@ public:
 			SetShaderUniformBuffer_Internal(Shader, BaseIndex, UniformBuffer.BypassUniform);
 			return;
 		}
-		auto* Cmd = AddCommand<FRHICommandSetLocalUniformBuffer>();
-		Cmd->Set<TShaderRHIParamRef>(this, this->UID, Shader, BaseIndex, UniformBuffer);
+		new (AllocCommand<FRHICommandSetLocalUniformBuffer<TShaderRHIParamRef> >()) FRHICommandSetLocalUniformBuffer<TShaderRHIParamRef>(this, Shader, BaseIndex, UniformBuffer);
 	}
 
 	template <typename TShaderRHIParamRef>
@@ -1175,27 +1139,20 @@ public:
 			SetShaderUniformBuffer_Internal(Shader, BaseIndex, UniformBuffer);
 			return;
 		}
-		auto* Cmd = AddCommand<FRHICommandSetShaderUniformBuffer>();
-		Cmd->Set<TShaderRHIParamRef>(Shader, BaseIndex, UniformBuffer);
+		new (AllocCommand<FRHICommandSetShaderUniformBuffer<TShaderRHIParamRef> >()) FRHICommandSetShaderUniformBuffer<TShaderRHIParamRef>(Shader, BaseIndex, UniformBuffer);
 	}
 
 	template <typename TShaderRHIParamRef>
-	FORCEINLINE_DEBUGGABLE void SetShaderParameter(TShaderRHIParamRef Shader, uint32 BufferIndex, uint32 BaseIndex, uint32 NumBytes, const void* NewValue, bool bValueInStack = false)
+	FORCEINLINE_DEBUGGABLE void SetShaderParameter(TShaderRHIParamRef Shader, uint32 BufferIndex, uint32 BaseIndex, uint32 NumBytes, const void* NewValue)
 	{
 		if (Bypass())
 		{
 			SetShaderParameter_Internal(Shader, BufferIndex, BaseIndex, NumBytes, NewValue);
 			return;
 		}
-
-		if (bValueInStack)
-		{
-			auto* Blob = AddCommand<FRHICommandNopBlob>(NumBytes);
-			Blob->Set(this, NumBytes, NewValue);
-			NewValue = Blob->Data;
-		}
-		auto* Cmd = AddCommand<FRHICommandSetShaderParameter>();
-		Cmd->Set<TShaderRHIParamRef>(Shader, BufferIndex, BaseIndex, NumBytes, NewValue);
+		void* UseValue = Alloc(NumBytes, 16);
+		FMemory::Memcpy(UseValue, NewValue, NumBytes);
+		new (AllocCommand<FRHICommandSetShaderParameter<TShaderRHIParamRef> >()) FRHICommandSetShaderParameter<TShaderRHIParamRef>(Shader, BufferIndex, BaseIndex, NumBytes, UseValue);
 	}
 
 	template <typename TShaderRHIParamRef>
@@ -1206,9 +1163,7 @@ public:
 			SetShaderTexture_Internal(Shader, TextureIndex, Texture);
 			return;
 		}
-
-		auto* Cmd = AddCommand<FRHICommandSetShaderTexture>();
-		Cmd->Set<TShaderRHIParamRef>(Shader, TextureIndex, Texture);
+		new (AllocCommand<FRHICommandSetShaderTexture<TShaderRHIParamRef> >()) FRHICommandSetShaderTexture<TShaderRHIParamRef>(Shader, TextureIndex, Texture);
 	}
 
 	template <typename TShaderRHIParamRef>
@@ -1219,8 +1174,7 @@ public:
 			SetShaderResourceViewParameter_Internal(Shader, SamplerIndex, SRV);
 			return;
 		}
-		auto* Cmd = AddCommand<FRHICommandSetShaderResourceViewParameter>();
-		Cmd->Set<TShaderRHIParamRef>(Shader, SamplerIndex, SRV);
+		new (AllocCommand<FRHICommandSetShaderResourceViewParameter<TShaderRHIParamRef> >()) FRHICommandSetShaderResourceViewParameter<TShaderRHIParamRef>(Shader, SamplerIndex, SRV);
 	}
 
 	template <typename TShaderRHIParamRef>
@@ -1231,9 +1185,7 @@ public:
 			SetShaderSampler_Internal(Shader, SamplerIndex, State);
 			return;
 		}
-
-		auto* Cmd = AddCommand<FRHICommandSetShaderSampler>();
-		Cmd->Set<TShaderRHIParamRef>(Shader, SamplerIndex, State);
+		new (AllocCommand<FRHICommandSetShaderSampler<TShaderRHIParamRef> >()) FRHICommandSetShaderSampler<TShaderRHIParamRef>(Shader, SamplerIndex, State);
 	}
 
 	FORCEINLINE_DEBUGGABLE void SetUAVParameter(FComputeShaderRHIParamRef Shader, uint32 UAVIndex, FUnorderedAccessViewRHIParamRef UAV)
@@ -1243,8 +1195,7 @@ public:
 			SetUAVParameter_Internal(Shader, UAVIndex, UAV);
 			return;
 		}
-		auto* Cmd = AddCommand<FRHICommandSetUAVParameter>();
-		Cmd->Set<FComputeShaderRHIParamRef>(Shader, UAVIndex, UAV);
+		new (AllocCommand<FRHICommandSetUAVParameter<FComputeShaderRHIParamRef> >()) FRHICommandSetUAVParameter<FComputeShaderRHIParamRef>(Shader, UAVIndex, UAV);
 	}
 
 	FORCEINLINE_DEBUGGABLE void SetUAVParameter(FComputeShaderRHIParamRef Shader, uint32 UAVIndex, FUnorderedAccessViewRHIParamRef UAV, uint32 InitialCount)
@@ -1254,8 +1205,7 @@ public:
 			SetUAVParameter_Internal(Shader, UAVIndex, UAV, InitialCount);
 			return;
 		}
-		auto* Cmd = AddCommand<FRHICommandSetUAVParameter_IntialCount>();
-		Cmd->Set<FComputeShaderRHIParamRef>(Shader, UAVIndex, UAV, InitialCount);
+		new (AllocCommand<FRHICommandSetUAVParameter_IntialCount<FComputeShaderRHIParamRef> >()) FRHICommandSetUAVParameter_IntialCount<FComputeShaderRHIParamRef>(Shader, UAVIndex, UAV, InitialCount);
 	}
 
 	FORCEINLINE_DEBUGGABLE void SetBoundShaderState(FBoundShaderStateRHIParamRef BoundShaderState)
@@ -1265,9 +1215,7 @@ public:
 			SetBoundShaderState_Internal(BoundShaderState);
 			return;
 		}
-
-		auto* Cmd = AddCommand<FRHICommandSetBoundShaderState>();
-		Cmd->Set(BoundShaderState);
+		new (AllocCommand<FRHICommandSetBoundShaderState>()) FRHICommandSetBoundShaderState(BoundShaderState);
 	}
 
 	FORCEINLINE_DEBUGGABLE void SetRasterizerState(FRasterizerStateRHIParamRef State)
@@ -1277,9 +1225,7 @@ public:
 			SetRasterizerState_Internal(State);
 			return;
 		}
-
-		auto* Cmd = AddCommand<FRHICommandSetRasterizerState>();
-		Cmd->Set(State);
+		new (AllocCommand<FRHICommandSetRasterizerState>()) FRHICommandSetRasterizerState(State);
 	}
 
 	FORCEINLINE_DEBUGGABLE void SetBlendState(FBlendStateRHIParamRef State, const FLinearColor& BlendFactor = FLinearColor::White)
@@ -1289,9 +1235,7 @@ public:
 			SetBlendState_Internal(State, BlendFactor);
 			return;
 		}
-
-		auto* Cmd = AddCommand<FRHICommandSetBlendState>();
-		Cmd->Set(State, BlendFactor);
+		new (AllocCommand<FRHICommandSetBlendState>()) FRHICommandSetBlendState(State, BlendFactor);
 	}
 
 	FORCEINLINE_DEBUGGABLE void DrawPrimitive(uint32 PrimitiveType, uint32 BaseVertexIndex, uint32 NumPrimitives, uint32 NumInstances)
@@ -1301,9 +1245,7 @@ public:
 			DrawPrimitive_Internal(PrimitiveType, BaseVertexIndex, NumPrimitives, NumInstances);
 			return;
 		}
-
-		auto* Cmd = AddCommand<FRHICommandDrawPrimitive>();
-		Cmd->Set(PrimitiveType, BaseVertexIndex, NumPrimitives, NumInstances);
+		new (AllocCommand<FRHICommandDrawPrimitive>()) FRHICommandDrawPrimitive(PrimitiveType, BaseVertexIndex, NumPrimitives, NumInstances);
 	}
 
 	FORCEINLINE_DEBUGGABLE void DrawIndexedPrimitive(FIndexBufferRHIParamRef IndexBuffer, uint32 PrimitiveType, int32 BaseVertexIndex, uint32 MinIndex, uint32 NumVertices, uint32 StartIndex, uint32 NumPrimitives, uint32 NumInstances)
@@ -1313,9 +1255,7 @@ public:
 			DrawIndexedPrimitive_Internal(IndexBuffer, PrimitiveType, BaseVertexIndex, MinIndex, NumVertices, StartIndex, NumPrimitives, NumInstances);
 			return;
 		}
-
-		auto* Cmd = AddCommand<FRHICommandDrawIndexedPrimitive>();
-		Cmd->Set(IndexBuffer, PrimitiveType, BaseVertexIndex, MinIndex, NumVertices, StartIndex, NumPrimitives, NumInstances);
+		new (AllocCommand<FRHICommandDrawIndexedPrimitive>()) FRHICommandDrawIndexedPrimitive(IndexBuffer, PrimitiveType, BaseVertexIndex, MinIndex, NumVertices, StartIndex, NumPrimitives, NumInstances);
 	}
 
 	FORCEINLINE_DEBUGGABLE void SetStreamSource(uint32 StreamIndex, FVertexBufferRHIParamRef VertexBuffer, uint32 Stride, uint32 Offset)
@@ -1325,9 +1265,7 @@ public:
 			SetStreamSource_Internal(StreamIndex, VertexBuffer, Stride, Offset);
 			return;
 		}
-
-		auto* Cmd = AddCommand<FRHICommandSetStreamSource>();
-		Cmd->Set(StreamIndex, VertexBuffer, Stride, Offset);
+		new (AllocCommand<FRHICommandSetStreamSource>()) FRHICommandSetStreamSource(StreamIndex, VertexBuffer, Stride, Offset);
 	}
 
 	FORCEINLINE_DEBUGGABLE void SetDepthStencilState(FDepthStencilStateRHIParamRef NewStateRHI, uint32 StencilRef = 0)
@@ -1337,9 +1275,7 @@ public:
 			SetDepthStencilState_Internal(NewStateRHI, StencilRef);
 			return;
 		}
-
-		auto* Cmd = AddCommand<FRHICommandSetDepthStencilState>();
-		Cmd->Set(NewStateRHI, StencilRef);
+		new (AllocCommand<FRHICommandSetDepthStencilState>()) FRHICommandSetDepthStencilState(NewStateRHI, StencilRef);
 	}
 
 	FORCEINLINE_DEBUGGABLE void SetViewport(uint32 MinX, uint32 MinY, float MinZ, uint32 MaxX, uint32 MaxY, float MaxZ)
@@ -1349,8 +1285,7 @@ public:
 			SetViewport_Internal(MinX, MinY, MinZ, MaxX, MaxY, MaxZ);
 			return;
 		}
-		auto* Cmd = AddCommand<FRHICommandSetViewport>();
-		Cmd->Set(MinX, MinY, MinZ, MaxX, MaxY, MaxZ);
+		new (AllocCommand<FRHICommandSetViewport>()) FRHICommandSetViewport(MinX, MinY, MinZ, MaxX, MaxY, MaxZ);
 	}
 
 	FORCEINLINE_DEBUGGABLE void SetScissorRect(bool bEnable, uint32 MinX, uint32 MinY, uint32 MaxX, uint32 MaxY)
@@ -1360,8 +1295,7 @@ public:
 			SetScissorRect_Internal(bEnable, MinX, MinY, MaxX, MaxY);
 			return;
 		}
-		auto* Cmd = AddCommand<FRHICommandSetScissorRect>();
-		Cmd->Set(bEnable,  MinX, MinY, MaxX, MaxY);
+		new (AllocCommand<FRHICommandSetScissorRect>()) FRHICommandSetScissorRect(bEnable, MinX, MinY, MaxX, MaxY);
 	}
 
 	FORCEINLINE_DEBUGGABLE void SetRenderTargets(
@@ -1382,8 +1316,7 @@ public:
 				UAVs);
 			return;
 		}
-		auto* Cmd = AddCommand<FRHICommandSetRenderTargets>();
-		Cmd->Set(
+		new (AllocCommand<FRHICommandSetRenderTargets>()) FRHICommandSetRenderTargets(
 			NewNumSimultaneousRenderTargets,
 			NewRenderTargetsRHI,
 			NewDepthStencilTargetRHI,
@@ -1400,7 +1333,7 @@ public:
 		}
 		check(!DrawUPData.OutVertexData && NumVertices * VertexDataStride > 0);
 
-		OutVertexData = FMemory::Malloc(NumVertices * VertexDataStride);
+		OutVertexData = Alloc(NumVertices * VertexDataStride, 16);
 
 		DrawUPData.PrimitiveType = PrimitiveType;
 		DrawUPData.NumPrimitives = NumPrimitives;
@@ -1417,8 +1350,7 @@ public:
 			return;
 		}
 		check(DrawUPData.OutVertexData && DrawUPData.NumVertices);
-		auto* Cmd = AddCommand<FRHICommandEndDrawPrimitiveUP>();
-		Cmd->Set(DrawUPData.PrimitiveType, DrawUPData.NumPrimitives, DrawUPData.NumVertices, DrawUPData.VertexDataStride, DrawUPData.OutVertexData);
+		new (AllocCommand<FRHICommandEndDrawPrimitiveUP>()) FRHICommandEndDrawPrimitiveUP(DrawUPData.PrimitiveType, DrawUPData.NumPrimitives, DrawUPData.NumVertices, DrawUPData.VertexDataStride, DrawUPData.OutVertexData);
 
 		DrawUPData.OutVertexData = nullptr;
 		DrawUPData.NumVertices = 0;
@@ -1433,8 +1365,8 @@ public:
 		}
 		check(!DrawUPData.OutVertexData && !DrawUPData.OutIndexData && NumVertices * VertexDataStride > 0 && NumIndices * IndexDataStride > 0);
 
-		OutVertexData = FMemory::Malloc(NumVertices * VertexDataStride);
-		OutIndexData = FMemory::Malloc(NumIndices * IndexDataStride);
+		OutVertexData = Alloc(NumVertices * VertexDataStride, 16);
+		OutIndexData = Alloc(NumIndices * IndexDataStride, 16);
 
 		DrawUPData.PrimitiveType = PrimitiveType;
 		DrawUPData.NumPrimitives = NumPrimitives;
@@ -1457,8 +1389,7 @@ public:
 		}
 		check(DrawUPData.OutVertexData && DrawUPData.OutIndexData && DrawUPData.NumIndices && DrawUPData.NumVertices);
 
-		auto* Cmd = AddCommand<FRHICommandEndDrawIndexedPrimitiveUP>();
-		Cmd->Set(DrawUPData.PrimitiveType, DrawUPData.NumPrimitives, DrawUPData.NumVertices, DrawUPData.VertexDataStride, DrawUPData.OutVertexData, DrawUPData.MinVertexIndex, DrawUPData.NumIndices, DrawUPData.IndexDataStride, DrawUPData.OutIndexData);
+		new (AllocCommand<FRHICommandEndDrawIndexedPrimitiveUP>()) FRHICommandEndDrawIndexedPrimitiveUP(DrawUPData.PrimitiveType, DrawUPData.NumPrimitives, DrawUPData.NumVertices, DrawUPData.VertexDataStride, DrawUPData.OutVertexData, DrawUPData.MinVertexIndex, DrawUPData.NumIndices, DrawUPData.IndexDataStride, DrawUPData.OutIndexData);
 
 		DrawUPData.OutVertexData = nullptr;
 		DrawUPData.OutIndexData = nullptr;
@@ -1473,8 +1404,7 @@ public:
 			SetComputeShader_Internal(ComputeShader);
 			return;
 		}
-		auto* Cmd = AddCommand<FRHICommandSetComputeShader>();
-		Cmd->Set(ComputeShader);
+		new (AllocCommand<FRHICommandSetComputeShader>()) FRHICommandSetComputeShader(ComputeShader);
 	}
 
 	FORCEINLINE_DEBUGGABLE void DispatchComputeShader(uint32 ThreadGroupCountX, uint32 ThreadGroupCountY, uint32 ThreadGroupCountZ)
@@ -1484,8 +1414,7 @@ public:
 			DispatchComputeShader_Internal(ThreadGroupCountX, ThreadGroupCountY, ThreadGroupCountZ);
 			return;
 		}
-		auto* Cmd = AddCommand<FRHICommandDispatchComputeShader>();
-		Cmd->Set(ThreadGroupCountX, ThreadGroupCountY, ThreadGroupCountZ);
+		new (AllocCommand<FRHICommandDispatchComputeShader>()) FRHICommandDispatchComputeShader(ThreadGroupCountX, ThreadGroupCountY, ThreadGroupCountZ);
 	}
 
 	FORCEINLINE_DEBUGGABLE void DispatchIndirectComputeShader(FVertexBufferRHIParamRef ArgumentBuffer, uint32 ArgumentOffset)
@@ -1495,8 +1424,7 @@ public:
 			DispatchIndirectComputeShader_Internal(ArgumentBuffer, ArgumentOffset);
 			return;
 		}
-		auto* Cmd = AddCommand<FRHICommandDispatchIndirectComputeShader>();
-		Cmd->Set(ArgumentBuffer, ArgumentOffset);
+		new (AllocCommand<FRHICommandDispatchIndirectComputeShader>()) FRHICommandDispatchIndirectComputeShader(ArgumentBuffer, ArgumentOffset);
 	}
 
 	FORCEINLINE_DEBUGGABLE void AutomaticCacheFlushAfterComputeShader(bool bEnable)
@@ -1506,8 +1434,7 @@ public:
 			AutomaticCacheFlushAfterComputeShader_Internal(bEnable);
 			return;
 		}
-		auto* Cmd = AddCommand<FRHICommandAutomaticCacheFlushAfterComputeShader>();
-		Cmd->Set(bEnable);
+		new (AllocCommand<FRHICommandAutomaticCacheFlushAfterComputeShader>()) FRHICommandAutomaticCacheFlushAfterComputeShader(bEnable);
 	}
 
 	FORCEINLINE_DEBUGGABLE void FlushComputeShaderCache()
@@ -1517,8 +1444,7 @@ public:
 			FlushComputeShaderCache_Internal();
 			return;
 		}
-		auto* Cmd = AddCommand<FRHICommandFlushComputeShaderCache>();
-		Cmd->Set();
+		new (AllocCommand<FRHICommandFlushComputeShaderCache>()) FRHICommandFlushComputeShaderCache();
 	}
 
 	FORCEINLINE_DEBUGGABLE void DrawPrimitiveIndirect(uint32 PrimitiveType, FVertexBufferRHIParamRef ArgumentBuffer, uint32 ArgumentOffset)
@@ -1528,8 +1454,7 @@ public:
 			DrawPrimitiveIndirect_Internal(PrimitiveType, ArgumentBuffer, ArgumentOffset);
 			return;
 		}
-		auto* Cmd = AddCommand<FRHICommandDrawPrimitiveIndirect>();
-		Cmd->Set(PrimitiveType, ArgumentBuffer, ArgumentOffset);
+		new (AllocCommand<FRHICommandDrawPrimitiveIndirect>()) FRHICommandDrawPrimitiveIndirect(PrimitiveType, ArgumentBuffer, ArgumentOffset);
 	}
 
 	FORCEINLINE_DEBUGGABLE void DrawIndexedIndirect(FIndexBufferRHIParamRef IndexBufferRHI, uint32 PrimitiveType, FStructuredBufferRHIParamRef ArgumentsBufferRHI, uint32 DrawArgumentsIndex, uint32 NumInstances)
@@ -1539,8 +1464,7 @@ public:
 			DrawIndexedIndirect_Internal(IndexBufferRHI, PrimitiveType, ArgumentsBufferRHI, DrawArgumentsIndex, NumInstances);
 			return;
 		}
-		auto* Cmd = AddCommand<FRHICommandDrawIndexedIndirect>();
-		Cmd->Set(IndexBufferRHI, PrimitiveType, ArgumentsBufferRHI, DrawArgumentsIndex, NumInstances);
+		new (AllocCommand<FRHICommandDrawIndexedIndirect>()) FRHICommandDrawIndexedIndirect(IndexBufferRHI, PrimitiveType, ArgumentsBufferRHI, DrawArgumentsIndex, NumInstances);
 	}
 
 	FORCEINLINE_DEBUGGABLE void DrawIndexedPrimitiveIndirect(uint32 PrimitiveType, FIndexBufferRHIParamRef IndexBuffer, FVertexBufferRHIParamRef ArgumentsBuffer, uint32 ArgumentOffset)
@@ -1550,8 +1474,7 @@ public:
 			DrawIndexedPrimitiveIndirect_Internal(PrimitiveType, IndexBuffer, ArgumentsBuffer, ArgumentOffset);
 			return;
 		}
-		auto* Cmd = AddCommand<FRHICommandDrawIndexedPrimitiveIndirect>();
-		Cmd->Set(PrimitiveType, IndexBuffer, ArgumentsBuffer, ArgumentOffset);
+		new (AllocCommand<FRHICommandDrawIndexedPrimitiveIndirect>()) FRHICommandDrawIndexedPrimitiveIndirect(PrimitiveType, IndexBuffer, ArgumentsBuffer, ArgumentOffset);
 	}
 
 	FORCEINLINE_DEBUGGABLE void EnableDepthBoundsTest(bool bEnable, float MinDepth, float MaxDepth)
@@ -1561,8 +1484,7 @@ public:
 			EnableDepthBoundsTest_Internal(bEnable, MinDepth, MaxDepth);
 			return;
 		}
-		auto* Cmd = AddCommand<FRHICommandEnableDepthBoundsTest>();
-		Cmd->Set(bEnable, MinDepth, MaxDepth);
+		new (AllocCommand<FRHICommandEnableDepthBoundsTest>()) FRHICommandEnableDepthBoundsTest(bEnable, MinDepth, MaxDepth);
 	}
 
 	FORCEINLINE_DEBUGGABLE void ClearUAV(FUnorderedAccessViewRHIParamRef UnorderedAccessViewRHI, const uint32(&Values)[4])
@@ -1572,8 +1494,7 @@ public:
 			ClearUAV_Internal(UnorderedAccessViewRHI, Values);
 			return;
 		}
-		auto* Cmd = AddCommand<FRHICommandClearUAV>();
-		Cmd->Set(UnorderedAccessViewRHI, Values);
+		new (AllocCommand<FRHICommandClearUAV>()) FRHICommandClearUAV(UnorderedAccessViewRHI, Values);
 	}
 
 	FORCEINLINE_DEBUGGABLE void CopyToResolveTarget(FTextureRHIParamRef SourceTextureRHI, FTextureRHIParamRef DestTextureRHI, bool bKeepOriginalSurface, const FResolveParams& ResolveParams)
@@ -1583,8 +1504,7 @@ public:
 			CopyToResolveTarget_Internal(SourceTextureRHI, DestTextureRHI, bKeepOriginalSurface, ResolveParams);
 			return;
 		}
-		auto* Cmd = AddCommand<FRHICommandCopyToResolveTarget>();
-		Cmd->Set(SourceTextureRHI, DestTextureRHI, bKeepOriginalSurface, ResolveParams);
+		new (AllocCommand<FRHICommandCopyToResolveTarget>()) FRHICommandCopyToResolveTarget(SourceTextureRHI, DestTextureRHI, bKeepOriginalSurface, ResolveParams);
 	}
 
 	FORCEINLINE_DEBUGGABLE void Clear(bool bClearColor, const FLinearColor& Color, bool bClearDepth, float Depth, bool bClearStencil, uint32 Stencil, FIntRect ExcludeRect)
@@ -1594,8 +1514,7 @@ public:
 			Clear_Internal(bClearColor, Color, bClearDepth, Depth, bClearStencil, Stencil, ExcludeRect);
 			return;
 		}
-		auto* Cmd = AddCommand<FRHICommandClear>();
-		Cmd->Set(bClearColor, Color, bClearDepth, Depth, bClearStencil, Stencil, ExcludeRect);
+		new (AllocCommand<FRHICommandClear>()) FRHICommandClear(bClearColor, Color, bClearDepth, Depth, bClearStencil, Stencil, ExcludeRect);
 	}
 
 	FORCEINLINE_DEBUGGABLE void ClearMRT(bool bClearColor, int32 NumClearColors, const FLinearColor* ClearColorArray, bool bClearDepth, float Depth, bool bClearStencil, uint32 Stencil, FIntRect ExcludeRect)
@@ -1605,73 +1524,8 @@ public:
 			ClearMRT_Internal(bClearColor, NumClearColors, ClearColorArray, bClearDepth, Depth, bClearStencil, Stencil, ExcludeRect);
 			return;
 		}
-		auto* Cmd = AddCommand<FRHICommandClearMRT>();
-		Cmd->Set(bClearColor, NumClearColors, ClearColorArray, bClearDepth, Depth, bClearStencil, Stencil, ExcludeRect);
-	}
-
-	const SIZE_T GetUsedMemory() const;
-
-protected:
-
-	struct FDrawUpData
-	{
-
-		uint32 PrimitiveType;
-		uint32 NumPrimitives;
-		uint32 NumVertices;
-		uint32 VertexDataStride;
-		void* OutVertexData;
-		uint32 MinVertexIndex;
-		uint32 NumIndices;
-		uint32 IndexDataStride;
-		void* OutIndexData;
-
-		FDrawUpData()
-			: PrimitiveType(PT_Num)
-			, OutVertexData(nullptr)
-			, OutIndexData(nullptr)
-		{
-		}
-
-	};
-
-	FDrawUpData DrawUPData;
-
-	FORCEINLINE_DEBUGGABLE bool HasCommands() const
-	{
-		return (NumCommands > 0);
+		new (AllocCommand<FRHICommandClearMRT>()) FRHICommandClearMRT(bClearColor, NumClearColors, ClearColorArray, bClearDepth, Depth, bClearStencil, Stencil, ExcludeRect);
 	}
-
-
-private:
-	bool bExecuting;
-	uint32 NumCommands;
-	uint32 UID;
-
-	FORCEINLINE_DEBUGGABLE uint8* Alloc(SIZE_T InSize)
-	{
-		return MemManager.Alloc(InSize);
-	}
-
-	template <typename TCmd>
-	FORCEINLINE_DEBUGGABLE TCmd* AddCommand(uint32 ExtraData = 0)
-	{
-		checkSlow(!bExecuting);
-		TCmd* Cmd = (TCmd*)Alloc(sizeof(TCmd) + ExtraData);
-		//::OutputDebugStringW(*FString::Printf(TEXT("Add %d: @ 0x%p, %d bytes\n"), NumCommands, (void*)Cmd, (sizeof(TCmd) + ExtraData)));
-		++NumCommands;
-		return Cmd;
-	}
-
-	void Reset();
-
-	inline bool Bypass();
-
-	friend class FRHICommandListExecutor;
-	friend class FRHICommandListIterator;
-
-	// Used here as we really want Alloc() to only work for this class, nobody else should use it
-	friend struct FRHICommandNopBlob;
 };
 
 class RHI_API FRHICommandListImmediate : public FRHICommandList
@@ -1765,7 +1619,7 @@ public:
 	}
 	static inline FRHICommandListImmediate& GetImmediateCommandList();
 
-	void ExecuteList(FRHICommandList& CmdList);
+	void ExecuteList(FRHICommandListBase& CmdList);
 	void ExecuteList(FRHICommandListImmediate& CmdList);
 	void LatchBypass();
 
@@ -1785,22 +1639,16 @@ public:
 
 private:
 
-	void ExecuteInner(FRHICommandList& CmdList);
+	void ExecuteInner(FRHICommandListBase& CmdList);
 
 	bool bLatchedBypass;
-	friend class FRHICommandList;
+	friend class FRHICommandListBase;
 	FThreadSafeCounter UIDCounter;
 	FThreadSafeCounter OutstandingCmdListCount;
 	FRHICommandListImmediate CommandListImmediate;
-
-	void ExecuteListGenericRHI(FRHICommandList& CmdList);
 };
-extern RHI_API FRHICommandListExecutor GRHICommandList;
 
-FORCEINLINE_DEBUGGABLE bool FRHICommandList::Bypass()
-{
-	return GRHICommandList.Bypass();
-}
+extern RHI_API FRHICommandListExecutor GRHICommandList;
 
 FORCEINLINE_DEBUGGABLE FRHICommandListImmediate& FRHICommandListExecutor::GetImmediateCommandList()
 {
diff --git a/Engine/Source/Runtime/RHI/Public/RHICommandList.inl b/Engine/Source/Runtime/RHI/Public/RHICommandList.inl
--- a/Engine/Source/Runtime/RHI/Public/RHICommandList.inl
+++ b/Engine/Source/Runtime/RHI/Public/RHICommandList.inl
@@ -8,7 +8,7 @@
 #define __RHICOMMANDLIST_INL__
 #include "RHICommandList.h"
 
-FORCEINLINE_DEBUGGABLE void FRHICommandList::Flush()
+FORCEINLINE_DEBUGGABLE void FRHICommandListBase::Flush()
 {
 	if (HasCommands())
 	{
@@ -17,6 +17,12 @@ FORCEINLINE_DEBUGGABLE void FRHICommandList::Flush()
 	}
 }
 
+FORCEINLINE_DEBUGGABLE bool FRHICommandListBase::Bypass()
+{
+	return GRHICommandList.Bypass();
+}
+
+
 FORCEINLINE_DEBUGGABLE void FRHICommandListImmediate::Flush()
 {
 	if (HasCommands())
@@ -31,63 +37,33 @@ FORCEINLINE_DEBUGGABLE void FRHICommandListImmediate::Flush()
 class FRHICommandListIterator
 {
 public:
-	FRHICommandListIterator(FRHICommandList& CmdList)
+	FRHICommandListIterator(FRHICommandListBase& CmdList)
 	{
+		CmdPtr = CmdList.Root;
 		NumCommands = 0;
 		CmdListNumCommands = CmdList.NumCommands;
-		Page = CmdList.MemManager.FirstPage;
-		CmdPtr = Page->Head;
-		CmdTail = Page->Current;
 	}
-
-	FORCEINLINE_DEBUGGABLE bool HasCommandsLeft() const
+	~FRHICommandListIterator()
 	{
-		return (CmdPtr < CmdTail);
+		checkf(CmdListNumCommands == NumCommands, TEXT("Missed %d Commands!"), CmdListNumCommands - NumCommands);
 	}
 
-	// Current command
-	FORCEINLINE_DEBUGGABLE FRHICommand* operator * ()
+	FORCEINLINE_DEBUGGABLE bool HasCommandsLeft() const
 	{
-		return (FRHICommand*)CmdPtr;
+		return !!CmdPtr;
 	}
 
-	// Get the next RHICommand and advance the iterator
-	template <typename TCmd>
-	FORCEINLINE_DEBUGGABLE TCmd* NextCommand()
+	FORCEINLINE_DEBUGGABLE FRHICommandBase* NextCommand()
 	{
-		TCmd* RHICmd = (TCmd*)CmdPtr;
-		//::OutputDebugStringW(*FString::Printf(TEXT("Exec %d: %d @ 0x%p, %d bytes\n"), NumCommands, (int32)RHICmd->Type, (void*)RHICmd, sizeof(TCmd) + RHICmd->ExtraSize()));
-		CmdPtr += sizeof(TCmd) + RHICmd->ExtraSize();
-		CmdPtr = Align(CmdPtr, FRHICommandList::Alignment);
-
-		//@todo-rco: Fix me!
-		if (!TCmd::IsEndOfPage)
-		{
-			// Don't count EOP as that is an allocator construct
-			++NumCommands;
-		}
-
-		if (TCmd::IsEndOfPage || CmdPtr >= CmdTail)
-		{
-			Page = Page->NextPage;
-			CmdPtr = Page ? Page->Head : nullptr;
-			CmdTail = Page ? Page->Current : nullptr;
-		}
-
+		FRHICommandBase* RHICmd = CmdPtr;
+		CmdPtr = RHICmd->Next;
+		NumCommands++;
 		return RHICmd;
 	}
 
-	void CheckNumCommands()
-	{
-		checkf(CmdListNumCommands == NumCommands, TEXT("Missed %d Commands!"), CmdListNumCommands - NumCommands);
-		checkf(CmdPtr == CmdTail, TEXT("Mismatched Tail Pointer!"));
-	}
-
-protected:
-	FRHICommandList::FMemManager::FPage* Page;
+private:
+	FRHICommandBase* CmdPtr;
 	uint32 NumCommands;
-	const uint8* CmdPtr;
-	const uint8* CmdTail;
 	uint32 CmdListNumCommands;
 };
 
diff --git a/Engine/Source/Runtime/RHI/Public/RHIResources.h b/Engine/Source/Runtime/RHI/Public/RHIResources.h
--- a/Engine/Source/Runtime/RHI/Public/RHIResources.h
+++ b/Engine/Source/Runtime/RHI/Public/RHIResources.h
@@ -43,6 +43,10 @@ public:
 				{
 					PendingDeletes.Push(const_cast<FRHIResource*>(this));
 				}
+				else
+				{
+					MarkedForDelete.Decrement();
+				}
 			}
 		}
 		check(NewValue >= 0);
diff --git a/Engine/Source/Runtime/ShaderCore/Public/ShaderParameterUtils.h b/Engine/Source/Runtime/ShaderCore/Public/ShaderParameterUtils.h
--- a/Engine/Source/Runtime/ShaderCore/Public/ShaderParameterUtils.h
+++ b/Engine/Source/Runtime/ShaderCore/Public/ShaderParameterUtils.h
@@ -40,8 +40,7 @@ void SetShaderValue(
 			Parameter.GetBufferIndex(),
 			Parameter.GetBaseIndex() + ElementIndex * AlignedTypeSize,
 			(uint32)NumBytesToSet,
-			&Value,
-			true
+			&Value
 			);
 	}
 }
@@ -89,8 +88,7 @@ void SetShaderValueArray(
 			Parameter.GetBufferIndex(),
 			Parameter.GetBaseIndex() + BaseElementIndex * AlignedTypeSize,
 			(uint32)NumBytesToSet,
-			Values,
-			true
+			Values
 			);
 	}
 }
@@ -131,8 +129,7 @@ inline void SetPixelShaderBool(
 			Parameter.GetBufferIndex(),
 			Parameter.GetBaseIndex(),
 			sizeof(BoolValue),
-			&BoolValue,
-			true
+			&BoolValue
 			);
 	}
 }
diff --git a/Engine/Source/Runtime/Windows/D3D11RHI/Private/D3D11Commands.cpp b/Engine/Source/Runtime/Windows/D3D11RHI/Private/D3D11Commands.cpp
--- a/Engine/Source/Runtime/Windows/D3D11RHI/Private/D3D11Commands.cpp
+++ b/Engine/Source/Runtime/Windows/D3D11RHI/Private/D3D11Commands.cpp
@@ -2084,483 +2084,6 @@ uint32 FD3D11DynamicRHI::RHIGetGPUFrameCycles()
 void FD3D11DynamicRHI::RHIExecuteCommandList(FRHICommandList* CmdList)
 {
 	check(0); // this path has gone stale and needs updated methods, starting at ERCT_SetScissorRect
-	FRHICommandListIterator Iter(*CmdList);
-	while (Iter.HasCommandsLeft())
-	{
-		auto* Cmd = *Iter;
-		switch (Cmd->Type)
-		{
-		case ERCT_NopBlob:
-			{
-				// Nop
-				auto* RHICmd = Iter.NextCommand<FRHICommandNopBlob>();// (FRHICommandNopBlob*)Cmd;
-				(void)RHICmd;	// Unused
-			}
-			break;
-		case ERCT_NopEndOfPage:
-			{
-				// Nop
-				auto* RHICmd = Iter.NextCommand<FRHICommandNopEndOfPage>();
-				(void)RHICmd;	// Unused
-			}
-			break;
-		case ERCT_SetRasterizerState:
-			{
-				auto* RHICmd = Iter.NextCommand<FRHICommandSetRasterizerState>();
-				{
-					auto* State = (FD3D11RasterizerState*)RHICmd->State;
-					StateCache.SetRasterizerState(State->Resource);
-				}
-				RHICmd->State->Release();
-			}
-			break;
-		case ERCT_SetShaderParameter:
-			{
-				auto* RHICmd = Iter.NextCommand<FRHICommandSetShaderParameter>();
-				{
-					switch (RHICmd->ShaderFrequency)
-					{
-					case SF_Vertex: 
-						VALIDATE_BOUND_SHADER((FVertexShaderRHIParamRef)RHICmd->Shader);
-						checkSlow(VSConstantBuffers[RHICmd->BufferIndex]);
-						VSConstantBuffers[RHICmd->BufferIndex]->UpdateConstant((const uint8*)RHICmd->NewValue, RHICmd->BaseIndex, RHICmd->NumBytes);
-						break;
-					case SF_Hull:
-						VALIDATE_BOUND_SHADER((FHullShaderRHIParamRef)RHICmd->Shader);
-						checkSlow(HSConstantBuffers[RHICmd->BufferIndex]);
-						HSConstantBuffers[RHICmd->BufferIndex]->UpdateConstant((const uint8*)RHICmd->NewValue, RHICmd->BaseIndex, RHICmd->NumBytes);
-						break;
-					case SF_Domain:
-						VALIDATE_BOUND_SHADER((FDomainShaderRHIParamRef)RHICmd->Shader);
-						checkSlow(DSConstantBuffers[RHICmd->BufferIndex]);
-						DSConstantBuffers[RHICmd->BufferIndex]->UpdateConstant((const uint8*)RHICmd->NewValue, RHICmd->BaseIndex, RHICmd->NumBytes);
-						break;
-					case SF_Geometry:
-						VALIDATE_BOUND_SHADER((FGeometryShaderRHIParamRef)RHICmd->Shader);
-						checkSlow(GSConstantBuffers[RHICmd->BufferIndex]);
-						GSConstantBuffers[RHICmd->BufferIndex]->UpdateConstant((const uint8*)RHICmd->NewValue, RHICmd->BaseIndex, RHICmd->NumBytes);
-						break;
-					case SF_Pixel:
-						VALIDATE_BOUND_SHADER((FPixelShaderRHIParamRef)RHICmd->Shader);
-						checkSlow(PSConstantBuffers[RHICmd->BufferIndex]);
-						PSConstantBuffers[RHICmd->BufferIndex]->UpdateConstant((const uint8*)RHICmd->NewValue, RHICmd->BaseIndex, RHICmd->NumBytes);
-						break;
-					case SF_Compute:
-						VALIDATE_BOUND_SHADER((FComputeShaderRHIParamRef)RHICmd->Shader);
-						checkSlow(CSConstantBuffers[RHICmd->BufferIndex]);
-						CSConstantBuffers[RHICmd->BufferIndex]->UpdateConstant((const uint8*)RHICmd->NewValue, RHICmd->BaseIndex, RHICmd->NumBytes);
-						break;
-					default: check(0); break;
-					}
-				}
-			}
-			break;
-		case ERCT_SetShaderUniformBuffer:
-			{
-				auto* RHICmd = Iter.NextCommand<FRHICommandSetShaderUniformBuffer>();
-				{
-					FD3D11UniformBuffer* Buffer = (FD3D11UniformBuffer*)RHICmd->UniformBuffer;
-					ID3D11Buffer* ConstantBuffer = Buffer ? Buffer->Resource : nullptr;
-					switch (RHICmd->ShaderFrequency)
-					{
-					case SF_Vertex:
-						VALIDATE_BOUND_SHADER((FVertexShaderRHIParamRef)RHICmd->Shader);
-						StateCache.SetConstantBuffer<SF_Vertex>(ConstantBuffer, RHICmd->BaseIndex);
-						BoundUniformBuffers[SF_Vertex][RHICmd->BaseIndex] = Buffer;
-						DirtyUniformBuffers[SF_Vertex] |= (1 << RHICmd->BaseIndex);
-						break;
-					case SF_Hull:
-						VALIDATE_BOUND_SHADER((FHullShaderRHIParamRef)RHICmd->Shader);
-						StateCache.SetConstantBuffer<SF_Hull>(ConstantBuffer, RHICmd->BaseIndex);
-						BoundUniformBuffers[SF_Hull][RHICmd->BaseIndex] = Buffer;
-						DirtyUniformBuffers[SF_Hull] |= (1 << RHICmd->BaseIndex);
-						break;
-					case SF_Domain:
-						VALIDATE_BOUND_SHADER((FDomainShaderRHIParamRef)RHICmd->Shader);
-						StateCache.SetConstantBuffer<SF_Domain>(ConstantBuffer, RHICmd->BaseIndex);
-						BoundUniformBuffers[SF_Domain][RHICmd->BaseIndex] = Buffer;
-						DirtyUniformBuffers[SF_Domain] |= (1 << RHICmd->BaseIndex);
-						break;
-					case SF_Geometry:
-						VALIDATE_BOUND_SHADER((FGeometryShaderRHIParamRef)RHICmd->Shader);
-						StateCache.SetConstantBuffer<SF_Geometry>(ConstantBuffer, RHICmd->BaseIndex);
-						BoundUniformBuffers[SF_Geometry][RHICmd->BaseIndex] = Buffer;
-						DirtyUniformBuffers[SF_Geometry] |= (1 << RHICmd->BaseIndex);
-						break;
-					case SF_Pixel:
-						VALIDATE_BOUND_SHADER((FPixelShaderRHIParamRef)RHICmd->Shader);
-						StateCache.SetConstantBuffer<SF_Pixel>(ConstantBuffer, RHICmd->BaseIndex);
-						BoundUniformBuffers[SF_Pixel][RHICmd->BaseIndex] = Buffer;
-						DirtyUniformBuffers[SF_Pixel] |= (1 << RHICmd->BaseIndex);
-						break;
-					case SF_Compute:
-						VALIDATE_BOUND_SHADER((FComputeShaderRHIParamRef)RHICmd->Shader);
-						StateCache.SetConstantBuffer<SF_Compute>(ConstantBuffer, RHICmd->BaseIndex);
-						BoundUniformBuffers[SF_Compute][RHICmd->BaseIndex] = Buffer;
-						DirtyUniformBuffers[SF_Compute] |= (1 << RHICmd->BaseIndex);
-						break;
-					default: check(0); break;
-					}
-				}
-				RHICmd->Shader->Release(); 
-				RHICmd->UniformBuffer->Release();
-			}
-			break;
-		case ERCT_SetShaderSampler:
-			{
-				auto* RHICmd = Iter.NextCommand<FRHICommandSetShaderSampler>();
-				{
-					FD3D11SamplerState* SamplerState = (FD3D11SamplerState*)RHICmd->Sampler;
-					ID3D11SamplerState* StateResource = SamplerState->Resource;
-					switch (RHICmd->ShaderFrequency)
-					{
-					case SF_Vertex:
-						VALIDATE_BOUND_SHADER((FVertexShaderRHIParamRef)RHICmd->Shader);
-						StateCache.SetSamplerState<SF_Vertex>(StateResource, RHICmd->SamplerIndex);
-						break;
-					case SF_Hull:
-						VALIDATE_BOUND_SHADER((FHullShaderRHIParamRef)RHICmd->Shader);
-						StateCache.SetSamplerState<SF_Hull>(StateResource, RHICmd->SamplerIndex);
-						break;
-					case SF_Domain:
-						VALIDATE_BOUND_SHADER((FDomainShaderRHIParamRef)RHICmd->Shader);
-						StateCache.SetSamplerState<SF_Domain>(StateResource, RHICmd->SamplerIndex);
-						break;
-					case SF_Geometry:
-						VALIDATE_BOUND_SHADER((FGeometryShaderRHIParamRef)RHICmd->Shader);
-						StateCache.SetSamplerState<SF_Geometry>(StateResource, RHICmd->SamplerIndex);
-						break;
-					case SF_Pixel:
-						VALIDATE_BOUND_SHADER((FPixelShaderRHIParamRef)RHICmd->Shader);
-						StateCache.SetSamplerState<SF_Pixel>(StateResource, RHICmd->SamplerIndex);
-						break;
-					case SF_Compute:
-						VALIDATE_BOUND_SHADER((FComputeShaderRHIParamRef)RHICmd->Shader);
-						StateCache.SetSamplerState<SF_Compute>(StateResource, RHICmd->SamplerIndex);
-						break;
-					default: check(0); break;
-					}
-				}
-				RHICmd->Shader->Release(); 
-				RHICmd->Sampler->Release();
-			}
-			break;
-		case ERCT_SetShaderTexture:
-			{
-				auto* RHICmd = Iter.NextCommand<FRHICommandSetShaderTexture>();
-				{
-					uint32 Start = FPlatformTime::Cycles();
-					FD3D11TextureBase* Texture = GetD3D11TextureFromRHITexture(RHICmd->Texture);
-					ID3D11ShaderResourceView* ShaderResourceView = Texture ? Texture->GetShaderResourceView() : nullptr;
-					auto SRVType = (Texture == nullptr || Texture->GetRenderTargetView(0, 0) != nullptr || Texture->HasDepthStencilView()) ? FD3D11StateCache::SRV_Dynamic : FD3D11StateCache::SRV_Static;
-					switch (RHICmd->ShaderFrequency)
-					{
-					case SF_Vertex:
-						VALIDATE_BOUND_SHADER((FVertexShaderRHIParamRef)RHICmd->Shader);
-						SetShaderResourceView<SF_Vertex>(Texture, ShaderResourceView, RHICmd->TextureIndex, SRVType);
-						break;
-					case SF_Hull:
-						VALIDATE_BOUND_SHADER((FHullShaderRHIParamRef)RHICmd->Shader);
-						SetShaderResourceView<SF_Hull>(Texture, ShaderResourceView, RHICmd->TextureIndex, SRVType);
-						break;
-					case SF_Domain:
-						VALIDATE_BOUND_SHADER((FDomainShaderRHIParamRef)RHICmd->Shader);
-						SetShaderResourceView<SF_Domain>(Texture, ShaderResourceView, RHICmd->TextureIndex, SRVType);
-						break;
-					case SF_Geometry:
-						VALIDATE_BOUND_SHADER((FGeometryShaderRHIParamRef)RHICmd->Shader);
-						SetShaderResourceView<SF_Geometry>(Texture, ShaderResourceView, RHICmd->TextureIndex, SRVType);
-						break;
-					case SF_Pixel:
-						VALIDATE_BOUND_SHADER((FPixelShaderRHIParamRef)RHICmd->Shader);
-						SetShaderResourceView<SF_Pixel>(Texture, ShaderResourceView, RHICmd->TextureIndex, SRVType);
-						break;
-					case SF_Compute:
-						VALIDATE_BOUND_SHADER((FComputeShaderRHIParamRef)RHICmd->Shader);
-						SetShaderResourceView<SF_Compute>(Texture, ShaderResourceView, RHICmd->TextureIndex, SRVType);
-						break;
-					default: check(0); break;
-					}
-					SetShaderTextureCycles += (FPlatformTime::Cycles() - Start);
-					SetShaderTextureCalls++;
-				}
-				RHICmd->Shader->Release();
-				RHICmd->Texture->Release();
-			}
-			break;
-		case ERCT_SetShaderResourceViewParameter:
-			{
-				auto* RHICmd = Iter.NextCommand<FRHICommandSetShaderResourceViewParameter>();
-				{
-					uint32 Start = FPlatformTime::Cycles();
-
-					FD3D11ShaderResourceView*  SRV = (FD3D11ShaderResourceView*)RHICmd->SRV;
-
-					FD3D11BaseShaderResource* Resource = SRV->Resource;
-					ID3D11ShaderResourceView* D3D11SRV = SRV->View;
-
-					switch (RHICmd->ShaderFrequency)
-					{
-					case SF_Vertex:
-						VALIDATE_BOUND_SHADER((FVertexShaderRHIParamRef)RHICmd->Shader);
-						SetShaderResourceView<SF_Vertex>(Resource, D3D11SRV, RHICmd->SamplerIndex);
-						break;
-					case SF_Hull:
-						VALIDATE_BOUND_SHADER((FHullShaderRHIParamRef)RHICmd->Shader);
-						SetShaderResourceView<SF_Hull>(Resource, D3D11SRV, RHICmd->SamplerIndex);
-						break;
-					case SF_Domain:
-						VALIDATE_BOUND_SHADER((FDomainShaderRHIParamRef)RHICmd->Shader);
-						SetShaderResourceView<SF_Domain>(Resource, D3D11SRV, RHICmd->SamplerIndex);
-						break;
-					case SF_Geometry:
-						VALIDATE_BOUND_SHADER((FGeometryShaderRHIParamRef)RHICmd->Shader);
-						SetShaderResourceView<SF_Geometry>(Resource, D3D11SRV, RHICmd->SamplerIndex);
-						break;
-					case SF_Pixel:
-						VALIDATE_BOUND_SHADER((FPixelShaderRHIParamRef)RHICmd->Shader);
-						SetShaderResourceView<SF_Pixel>(Resource, D3D11SRV, RHICmd->SamplerIndex);
-						break;
-					case SF_Compute:
-						VALIDATE_BOUND_SHADER((FComputeShaderRHIParamRef)RHICmd->Shader);
-						SetShaderResourceView<SF_Compute>(Resource, D3D11SRV, RHICmd->SamplerIndex);
-						break;
-					default: check(0); break;
-					}
-				}
-				RHICmd->Shader->Release();
-				RHICmd->SRV->Release();
-			}
-			break;
-		case ERCT_SetUAVParameter:
-			{
-				auto* RHICmd = Iter.NextCommand<FRHICommandSetUAVParameter>();
-				{
-					uint32 Start = FPlatformTime::Cycles();
-
-					FD3D11UnorderedAccessView*  UAV = (FD3D11UnorderedAccessView*)RHICmd->UAV;
-
-					FD3D11BaseShaderResource* Resource = UAV->Resource;
-					ID3D11UnorderedAccessView* D3D11UAV = UAV->View;
-
-
-					ConditionalClearShaderResource(Resource);
-
-					uint32 InitialCount = -1;
-
-					switch (RHICmd->ShaderFrequency)
-					{
-					case SF_Compute:
-						//VALIDATE_BOUND_SHADER((FComputeShaderRHIParamRef)RHICmd->Shader);
-						Direct3DDeviceIMContext->CSSetUnorderedAccessViews(RHICmd->UAVIndex, 1, &D3D11UAV, &InitialCount);
-						break;
-					default: check(0); break;
-					}
-				}
-				RHICmd->Shader->Release();
-				RHICmd->UAV->Release();
-			}
-			break;
-		case ERCT_SetUAVParameter_IntialCount:
-			{
-				auto* RHICmd = Iter.NextCommand<FRHICommandSetUAVParameter_IntialCount>();
-				{
-					uint32 Start = FPlatformTime::Cycles();
-
-					FD3D11UnorderedAccessView*  UAV = (FD3D11UnorderedAccessView*)RHICmd->UAV;
-
-					FD3D11BaseShaderResource* Resource = UAV->Resource;
-					ID3D11UnorderedAccessView* D3D11UAV = UAV->View;
-
-
-					ConditionalClearShaderResource(Resource);
-
-					switch (RHICmd->ShaderFrequency)
-					{
-					case SF_Compute:
-						//VALIDATE_BOUND_SHADER((FComputeShaderRHIParamRef)RHICmd->Shader);
-						Direct3DDeviceIMContext->CSSetUnorderedAccessViews(RHICmd->UAVIndex, 1, &D3D11UAV, &RHICmd->InitialCount);
-						break;
-					default: check(0); break;
-					}
-				}
-				RHICmd->Shader->Release();
-				RHICmd->UAV->Release();
-			}
-			break;
-		case ERCT_DrawPrimitive:
-			{
-				auto* RHICmd = Iter.NextCommand<FRHICommandDrawPrimitive>();
-				{
-					RHI_DRAW_CALL_STATS(RHICmd->PrimitiveType, RHICmd->NumInstances * RHICmd->NumPrimitives);
-
-					CommitGraphicsResourceTables();
-					CommitNonComputeShaderConstants();
-					const uint32 VertexCount = GetVertexCountForPrimitiveCount(RHICmd->NumPrimitives, RHICmd->PrimitiveType);
-
-					GPUProfilingData.RegisterGPUWork(RHICmd->NumPrimitives * RHICmd->NumInstances, VertexCount * RHICmd->NumInstances);
-					StateCache.SetPrimitiveTopology(GetD3D11PrimitiveType(RHICmd->PrimitiveType, bUsingTessellation));
-					if (RHICmd->NumInstances > 1)
-					{
-						Direct3DDeviceIMContext->DrawInstanced(VertexCount, RHICmd->NumInstances, RHICmd->BaseVertexIndex,0);
-					}
-					else
-					{
-						Direct3DDeviceIMContext->Draw(VertexCount, RHICmd->BaseVertexIndex);
-					}
-				}
-			}
-			break;
-		case ERCT_DrawIndexedPrimitive:
-			{
-				auto* RHICmd = Iter.NextCommand<FRHICommandDrawIndexedPrimitive>();
-				{
-					auto* IndexBuffer = (FD3D11IndexBuffer*)RHICmd->IndexBuffer;
-
-					// Caller should make sure the input is valid, this avoid hidden bugs
-					ensure(RHICmd->NumPrimitives > 0);
-
-					RHI_DRAW_CALL_STATS(RHICmd->PrimitiveType, RHICmd->NumInstances* RHICmd->NumPrimitives);
-
-					GPUProfilingData.RegisterGPUWork(RHICmd->NumPrimitives * RHICmd->NumInstances, RHICmd->NumVertices * RHICmd->NumInstances);
-
-					CommitGraphicsResourceTables();
-					CommitNonComputeShaderConstants();
-
-					// determine 16bit vs 32bit indices
-					uint32 SizeFormat = sizeof(DXGI_FORMAT);
-					const DXGI_FORMAT Format = (IndexBuffer->GetStride() == sizeof(uint16) ? DXGI_FORMAT_R16_UINT : DXGI_FORMAT_R32_UINT);
-
-					uint32 IndexCount = GetVertexCountForPrimitiveCount(RHICmd->NumPrimitives, RHICmd->PrimitiveType);
-
-					// Verify that we are not trying to read outside the index buffer range
-					checkf(RHICmd->StartIndex + IndexCount <= IndexBuffer->GetSize() / IndexBuffer->GetStride(), 
-						TEXT("Start %u, Count %u, Type %u, Buffer Size %u, Buffer stride %u"), RHICmd->StartIndex, IndexCount, RHICmd->PrimitiveType, IndexBuffer->GetSize(), IndexBuffer->GetStride());
-
-					StateCache.SetIndexBuffer(IndexBuffer->Resource, Format, 0);
-					StateCache.SetPrimitiveTopology(GetD3D11PrimitiveType(RHICmd->PrimitiveType, bUsingTessellation));
-
-					if (RHICmd->NumInstances > 1)
-					{
-						Direct3DDeviceIMContext->DrawIndexedInstanced(IndexCount, RHICmd->NumInstances, RHICmd->StartIndex, RHICmd->BaseVertexIndex,0);
-					}
-					else
-					{
-						Direct3DDeviceIMContext->DrawIndexed(IndexCount, RHICmd->StartIndex, RHICmd->BaseVertexIndex);
-					}
-				}
-				RHICmd->IndexBuffer->Release();
-			}
-			break;
-		case ERCT_SetBoundShaderState:
-			{
-				auto* RHICmd = Iter.NextCommand<FRHICommandSetBoundShaderState>();
-				{
-					auto* BoundShaderState = (FD3D11BoundShaderState*)RHICmd->BoundShaderState;
-
-					StateCache.SetInputLayout(BoundShaderState->InputLayout);
-					StateCache.SetVertexShader(BoundShaderState->VertexShader);
-					StateCache.SetPixelShader(BoundShaderState->PixelShader);
-
-					StateCache.SetHullShader(BoundShaderState->HullShader);
-					StateCache.SetDomainShader(BoundShaderState->DomainShader);
-					StateCache.SetGeometryShader(BoundShaderState->GeometryShader);
-
-					if(BoundShaderState->HullShader != NULL && BoundShaderState->DomainShader != NULL)
-					{
-						bUsingTessellation = true;
-					}
-					else
-					{
-						bUsingTessellation = false;
-					}
-
-					// @TODO : really should only discard the constants if the shader state has actually changed.
-					bDiscardSharedConstants = true;
-
-					// Prevent transient bound shader states from being recreated for each use by keeping a history of the most recently used bound shader states.
-					// The history keeps them alive, and the bound shader state cache allows them to be reused if needed.
-					BoundShaderStateHistory.Add(BoundShaderState);
-
-					// Shader changed so all resource tables are dirty
-					DirtyUniformBuffers[SF_Vertex] = 0xffff;
-					DirtyUniformBuffers[SF_Pixel] = 0xffff;
-					DirtyUniformBuffers[SF_Hull] = 0xffff;
-					DirtyUniformBuffers[SF_Domain] = 0xffff;
-					DirtyUniformBuffers[SF_Geometry] = 0xffff;
-				}
-				RHICmd->BoundShaderState->Release();
-			}
-			break;
-		case ERCT_SetBlendState:
-			{
-				auto* RHICmd = Iter.NextCommand<FRHICommandSetBlendState>();
-				{
-					auto* State = (FD3D11BlendState*)RHICmd->State;
-					StateCache.SetBlendState(State->Resource, (const float*)&RHICmd->BlendFactor, 0xffffffff);
-				}
-				RHICmd->State->Release();
-			}
-			break;
-		case ERCT_SetStreamSource:
-			{
-				auto* RHICmd = Iter.NextCommand<FRHICommandSetStreamSource>();
-				{
-					auto* VertexBuffer = (FD3D11VertexBuffer*)RHICmd->VertexBuffer;
-					ID3D11Buffer* D3DBuffer = VertexBuffer ? VertexBuffer->Resource : nullptr;
-					StateCache.SetStreamSource(D3DBuffer, RHICmd->StreamIndex, RHICmd->Stride, RHICmd->Offset);
-				}
-				RHICmd->VertexBuffer->Release();
-			}
-			break;
-		case ERCT_SetDepthStencilState:
-			{
-				auto* RHICmd = Iter.NextCommand<FRHICommandSetDepthStencilState>();
-				{
-					auto* State = (FD3D11DepthStencilState*)RHICmd->State;
-
-					if (CurrentDepthTexture && State->AccessType != CurrentDSVAccessType)
-					{
-						CurrentDSVAccessType = State->AccessType;
-						CurrentDepthStencilTarget = CurrentDepthTexture->GetDepthStencilView(CurrentDSVAccessType);
-
-						// Unbind any shader views of the depth stencil target that are bound.
-						ConditionalClearShaderResource(CurrentDepthTexture);
-
-						CommitRenderTargetsAndUAVs();
-					}
-
-					StateCache.SetDepthStencilState(State->Resource, RHICmd->StencilRef);
-				}
-				RHICmd->State->Release();
-			}
-			break;
-		case ERCT_SetViewport:
-			{
-				auto* RHICmd = Iter.NextCommand<FRHICommandSetViewport>();
-				{
-					// These are the maximum viewport extents for D3D11. Exceeding them leads to badness.
-					check(RHICmd->MinX <= (uint32)D3D11_VIEWPORT_BOUNDS_MAX);
-					check(RHICmd->MinY <= (uint32)D3D11_VIEWPORT_BOUNDS_MAX);
-					check(RHICmd->MaxX <= (uint32)D3D11_VIEWPORT_BOUNDS_MAX);
-					check(RHICmd->MaxY <= (uint32)D3D11_VIEWPORT_BOUNDS_MAX);
-
-					D3D11_VIEWPORT Viewport = { RHICmd->MinX, RHICmd->MinY, RHICmd->MaxX - RHICmd->MinX, RHICmd->MaxY - RHICmd->MinY, RHICmd->MinZ, RHICmd->MaxZ };
-					//avoid setting a 0 extent viewport, which the debug runtime doesn't like
-					if (Viewport.Width > 0 && Viewport.Height > 0)
-					{
-						StateCache.SetViewport(Viewport);
-						SetScissorRectIfRequiredWhenSettingViewport(RHICmd->MinX, RHICmd->MinY, RHICmd->MaxX, RHICmd->MaxY);
-					}
-				}
-			}
-			break;
-		default:
-			checkf(0, TEXT("Unimplemented RHICmd %d!"), Cmd->Type);
-		}
-	}
-
-	Iter.CheckNumCommands();
 }
 
 // NVIDIA Depth Bounds Test interface
diff --git a/Engine/Source/Runtime/Windows/D3D11RHI/Private/D3D11RHIPrivate.h b/Engine/Source/Runtime/Windows/D3D11RHI/Private/D3D11RHIPrivate.h
--- a/Engine/Source/Runtime/Windows/D3D11RHI/Private/D3D11RHIPrivate.h
+++ b/Engine/Source/Runtime/Windows/D3D11RHI/Private/D3D11RHIPrivate.h
@@ -52,7 +52,8 @@ DECLARE_CYCLE_STAT_EXTERN(TEXT("New bound shader state time"),STAT_D3D11NewBound
 DECLARE_CYCLE_STAT_EXTERN(TEXT("Clean uniform buffer pool"),STAT_D3D11CleanUniformBufferTime,STATGROUP_D3D11RHI, );
 DECLARE_CYCLE_STAT_EXTERN(TEXT("Clear shader resources"),STAT_D3D11ClearShaderResourceTime,STATGROUP_D3D11RHI, );
 DECLARE_DWORD_ACCUMULATOR_STAT_EXTERN(TEXT("Uniform buffer pool num free"),STAT_D3D11NumFreeUniformBuffers,STATGROUP_D3D11RHI, );
-DECLARE_MEMORY_STAT_EXTERN(TEXT("Uniform buffer pool memory"),STAT_D3D11FreeUniformBufferMemory,STATGROUP_D3D11RHI, );
+DECLARE_DWORD_ACCUMULATOR_STAT_EXTERN(TEXT("Num Bound Shader State"), STAT_D3D11NumBoundShaderState, STATGROUP_D3D11RHI, );
+DECLARE_MEMORY_STAT_EXTERN(TEXT("Uniform buffer pool memory"), STAT_D3D11FreeUniformBufferMemory, STATGROUP_D3D11RHI, );
 DECLARE_CYCLE_STAT_EXTERN(TEXT("Update uniform buffer"),STAT_D3D11UpdateUniformBufferTime,STATGROUP_D3D11RHI, );
 DECLARE_DWORD_COUNTER_STAT_EXTERN(TEXT("Textures Allocated"),STAT_D3D11TexturesAllocated,STATGROUP_D3D11RHI, );
 DECLARE_DWORD_COUNTER_STAT_EXTERN(TEXT("Textures Released"),STAT_D3D11TexturesReleased,STATGROUP_D3D11RHI, );
diff --git a/Engine/Source/Runtime/Windows/D3D11RHI/Private/D3D11Shaders.cpp b/Engine/Source/Runtime/Windows/D3D11RHI/Private/D3D11Shaders.cpp
--- a/Engine/Source/Runtime/Windows/D3D11RHI/Private/D3D11Shaders.cpp
+++ b/Engine/Source/Runtime/Windows/D3D11RHI/Private/D3D11Shaders.cpp
@@ -201,6 +201,8 @@ FD3D11BoundShaderState::FD3D11BoundShaderState(
 	):
 	CacheLink(InVertexDeclarationRHI,InVertexShaderRHI,InPixelShaderRHI,InHullShaderRHI,InDomainShaderRHI,InGeometryShaderRHI,this)
 {
+	INC_DWORD_STAT(STAT_D3D11NumBoundShaderState);
+
 	DYNAMIC_CAST_D3D11RESOURCE(VertexDeclaration,InVertexDeclaration);
 	DYNAMIC_CAST_D3D11RESOURCE(VertexShader,InVertexShader);
 	DYNAMIC_CAST_D3D11RESOURCE(PixelShader,InPixelShader);
@@ -237,6 +239,12 @@ FD3D11BoundShaderState::FD3D11BoundShaderState(
 	static_assert(ARRAY_COUNT(bShaderNeedsGlobalConstantBuffer) == SF_NumFrequencies, "EShaderFrequency size should match with array count of bShaderNeedsGlobalConstantBuffer.");
 }
 
+FD3D11BoundShaderState::~FD3D11BoundShaderState()
+{
+	DEC_DWORD_STAT(STAT_D3D11NumBoundShaderState);
+}
+
+
 /**
 * Creates a bound shader state instance which encapsulates a decl, vertex shader, and pixel shader
 * @param VertexDeclaration - existing vertex decl
diff --git a/Engine/Source/Runtime/Windows/D3D11RHI/Private/D3D11Util.cpp b/Engine/Source/Runtime/Windows/D3D11RHI/Private/D3D11Util.cpp
--- a/Engine/Source/Runtime/Windows/D3D11RHI/Private/D3D11Util.cpp
+++ b/Engine/Source/Runtime/Windows/D3D11RHI/Private/D3D11Util.cpp
@@ -382,6 +382,7 @@ DEFINE_STAT(STAT_D3D11UpdateUniformBufferTime);
 DEFINE_STAT(STAT_D3D11TexturePoolMemory);
 DEFINE_STAT(STAT_D3D11FreeUniformBufferMemory);
 DEFINE_STAT(STAT_D3D11NumFreeUniformBuffers);
+DEFINE_STAT(STAT_D3D11NumBoundShaderState);
 
 DEFINE_STAT(STAT_D3D11CommitResourceTables);
 DEFINE_STAT(STAT_D3D11CacheResourceTables);
diff --git a/Engine/Source/Runtime/Windows/D3D11RHI/Public/D3D11Resources.h b/Engine/Source/Runtime/Windows/D3D11RHI/Public/D3D11Resources.h
--- a/Engine/Source/Runtime/Windows/D3D11RHI/Public/D3D11Resources.h
+++ b/Engine/Source/Runtime/Windows/D3D11RHI/Public/D3D11Resources.h
@@ -148,6 +148,8 @@ public:
 		ID3D11Device* Direct3DDevice
 		);
 
+	~FD3D11BoundShaderState();
+
 	/**
 	 * Get the shader for the given frequency.
 	 */
