Add a specific SP_OPENGL_SM4_MAC so that we can compile shaders from Windows for all Macs, even the older GL 3.3 devices that don't support volume-texture rendering.

diff --git a/Engine/Source/Developer/CrashTracker/Private/CrashTrackerModule.cpp b/Engine/Source/Developer/CrashTracker/Private/CrashTrackerModule.cpp
--- a/Engine/Source/Developer/CrashTracker/Private/CrashTrackerModule.cpp
+++ b/Engine/Source/Developer/CrashTracker/Private/CrashTrackerModule.cpp
@@ -112,7 +112,7 @@ void FCrashTrackerModule::StartupModule()
 	VideoCapture.Reset();
 #if CRASH_TRACKER_SUPPORTED
 	if (GIsEditor &&
-		(GRHIShaderPlatform == SP_PCD3D_SM4 || GRHIShaderPlatform == SP_PCD3D_SM5 || GRHIShaderPlatform == SP_OPENGL_SM4 ))
+		(GRHIShaderPlatform == SP_PCD3D_SM4 || GRHIShaderPlatform == SP_PCD3D_SM5 || GRHIShaderPlatform == SP_OPENGL_SM4 || GRHIShaderPlatform == SP_OPENGL_SM4_MAC ))
 	{
 		bool bCrashTrackerShouldBeEnabled = false;
 #if UE_BUILD_DEVELOPMENT
diff --git a/Engine/Source/Runtime/OpenGLDrv/Private/Mac/OpenGLMac.h b/Engine/Source/Runtime/OpenGLDrv/Private/Mac/OpenGLMac.h
--- a/Engine/Source/Runtime/OpenGLDrv/Private/Mac/OpenGLMac.h
+++ b/Engine/Source/Runtime/OpenGLDrv/Private/Mac/OpenGLMac.h
@@ -212,6 +212,26 @@ struct FMacOpenGL : public FOpenGL3
 	
 	static FORCEINLINE bool SupportsSeamlessCubeMap()					{ return true; }
 	
+	static FORCEINLINE EShaderPlatform GetShaderPlatform()
+	{
+		static bool bForceFeatureLevelES2 = FParse::Param(FCommandLine::Get(), TEXT("FeatureLevelES2"));
+		if (bForceFeatureLevelES2)
+		{
+			return SP_OPENGL_PCES2;
+		}
+		
+		// Shader platform
+		switch(GetMajorVersion())
+		{
+			case 3:
+				return SP_OPENGL_SM4_MAC;
+			case 4:
+				return GetMinorVersion() > 2 ? SP_OPENGL_SM5 : SP_OPENGL_SM4_MAC;
+			default:
+				return SP_OPENGL_SM4_MAC;
+		}
+	}
+	
 	static uint64 GetVideoMemorySize();
 	
 	static void ProcessQueryGLInt();
diff --git a/Engine/Source/Runtime/OpenGLDrv/Private/OpenGLDevice.cpp b/Engine/Source/Runtime/OpenGLDrv/Private/OpenGLDevice.cpp
--- a/Engine/Source/Runtime/OpenGLDrv/Private/OpenGLDevice.cpp
+++ b/Engine/Source/Runtime/OpenGLDrv/Private/OpenGLDevice.cpp
@@ -643,7 +643,7 @@ static void InitRHICapabilitiesForGL()
 
 	GShaderPlatformForFeatureLevel[ERHIFeatureLevel::ES2] = (GMaxRHIFeatureLevel == ERHIFeatureLevel::ES2) ? GRHIShaderPlatform : SP_OPENGL_PCES2;
 	GShaderPlatformForFeatureLevel[ERHIFeatureLevel::SM3] = SP_NumPlatforms;
-	GShaderPlatformForFeatureLevel[ERHIFeatureLevel::SM4] = SP_OPENGL_SM4;
+	GShaderPlatformForFeatureLevel[ERHIFeatureLevel::SM4] = PLATFORM_MAC ? SP_OPENGL_SM4_MAC : SP_OPENGL_SM4;
 	GShaderPlatformForFeatureLevel[ERHIFeatureLevel::SM5] = SP_OPENGL_SM5;
 
 	// Set to same values as in DX11, as for the time being clip space adjustment are done entirely
diff --git a/Engine/Source/Runtime/RHI/Private/RHI.cpp b/Engine/Source/Runtime/RHI/Private/RHI.cpp
--- a/Engine/Source/Runtime/RHI/Private/RHI.cpp
+++ b/Engine/Source/Runtime/RHI/Private/RHI.cpp
@@ -292,11 +292,9 @@ FName LegacyShaderPlatformToShaderFormat(EShaderPlatform Platform)
 	case SP_PCD3D_ES2:
 		return NAME_PCD3D_ES2;
 	case SP_OPENGL_SM4:
-#if PLATFORM_MAC
-		return NAME_GLSL_150_MAC;
-#else
 		return NAME_GLSL_150;
-#endif
+	case SP_OPENGL_SM4_MAC:
+		return NAME_GLSL_150_MAC;
 	case SP_PS4:
 		return NAME_SF_PS4;
 	case SP_XBOXONE:
@@ -325,7 +323,7 @@ EShaderPlatform ShaderFormatToLegacyShaderPlatform(FName ShaderFormat)
 	if (ShaderFormat == NAME_PCD3D_SM4)			return SP_PCD3D_SM4;
 	if (ShaderFormat == NAME_PCD3D_ES2)			return SP_PCD3D_ES2;
 	if (ShaderFormat == NAME_GLSL_150)			return SP_OPENGL_SM4;
-	if (ShaderFormat == NAME_GLSL_150_MAC)		return SP_OPENGL_SM4;
+	if (ShaderFormat == NAME_GLSL_150_MAC)		return SP_OPENGL_SM4_MAC;
 	if (ShaderFormat == NAME_SF_PS4)			return SP_PS4;
 	if (ShaderFormat == NAME_SF_XBOXONE)		return SP_XBOXONE;
 	if (ShaderFormat == NAME_GLSL_430)			return SP_OPENGL_SM5;
diff --git a/Engine/Source/Runtime/RHI/Public/RHIDefinitions.h b/Engine/Source/Runtime/RHI/Public/RHIDefinitions.h
--- a/Engine/Source/Runtime/RHI/Public/RHIDefinitions.h
+++ b/Engine/Source/Runtime/RHI/Public/RHIDefinitions.h
@@ -41,9 +41,10 @@ enum EShaderPlatform
 	SP_OPENGL_ES2_WEBGL = 9, 
 	SP_OPENGL_ES2_IOS	= 10,
 	SP_METAL			= 11,
+	SP_OPENGL_SM4_MAC	= 12,
 
-	SP_NumPlatforms		= 12,
-	SP_NumBits			= 4,
+	SP_NumPlatforms		= 13,
+	SP_NumBits			= 5,
 };
 static_assert(SP_NumPlatforms <= (1 << SP_NumBits), "SP_NumPlatforms will not fit on SP_NumBits");
 
@@ -518,7 +519,7 @@ enum ETextureReallocationStatus
 
 inline bool IsPCPlatform(const EShaderPlatform Platform)
 {
-	return Platform == SP_PCD3D_SM5 || Platform == SP_PCD3D_SM4 || Platform == SP_PCD3D_ES2 || Platform ==  SP_OPENGL_SM4 || Platform == SP_OPENGL_SM5 || Platform == SP_OPENGL_PCES2;
+	return Platform == SP_PCD3D_SM5 || Platform == SP_PCD3D_SM4 || Platform == SP_PCD3D_ES2 || Platform ==  SP_OPENGL_SM4 || Platform == SP_OPENGL_SM4_MAC || Platform == SP_OPENGL_SM5 || Platform == SP_OPENGL_PCES2;
 }
 
 /** Whether the shader platform corresponds to the ES2 feature level. */
@@ -529,7 +530,7 @@ inline bool IsES2Platform(const EShaderPlatform Platform)
 
 inline bool IsOpenGLPlatform(const EShaderPlatform Platform)
 {
-	return Platform == SP_OPENGL_SM4 || Platform == SP_OPENGL_SM5 || Platform == SP_OPENGL_PCES2 || Platform == SP_OPENGL_ES2 || Platform == SP_OPENGL_ES2_WEBGL || Platform == SP_OPENGL_ES2_IOS;
+	return Platform == SP_OPENGL_SM4 || Platform == SP_OPENGL_SM4_MAC || Platform == SP_OPENGL_SM5 || Platform == SP_OPENGL_PCES2 || Platform == SP_OPENGL_ES2 || Platform == SP_OPENGL_ES2_WEBGL || Platform == SP_OPENGL_ES2_IOS;
 }
 
 inline bool IsConsolePlatform(const EShaderPlatform Platform)
@@ -548,6 +549,7 @@ inline ERHIFeatureLevel::Type GetMaxSupportedFeatureLevel(EShaderPlatform InShad
 		return ERHIFeatureLevel::SM5;
 	case SP_PCD3D_SM4:
 	case SP_OPENGL_SM4:
+	case SP_OPENGL_SM4_MAC:
 		return ERHIFeatureLevel::SM4;
 	case SP_PCD3D_ES2:
 	case SP_OPENGL_PCES2:
@@ -583,6 +585,8 @@ inline bool IsFeatureLevelSupported(EShaderPlatform InShaderPlatform, ERHIFeatur
 		return InFeatureLevel <= ERHIFeatureLevel::ES2;
 	case SP_OPENGL_SM4:
 		return InFeatureLevel <= ERHIFeatureLevel::SM4;
+	case SP_OPENGL_SM4_MAC:
+		return InFeatureLevel <= ERHIFeatureLevel::SM4;
 	case SP_OPENGL_SM5:
 		return InFeatureLevel <= ERHIFeatureLevel::SM5;
 	case SP_PS4:
diff --git a/Engine/Source/Runtime/Renderer/Private/LightMapRendering.h b/Engine/Source/Runtime/Renderer/Private/LightMapRendering.h
--- a/Engine/Source/Runtime/Renderer/Private/LightMapRendering.h
+++ b/Engine/Source/Runtime/Renderer/Private/LightMapRendering.h
@@ -733,9 +733,9 @@ public:
 		static const auto AllowStaticLightingVar = IConsoleManager::Get().FindTConsoleVariableDataInt(TEXT("r.AllowStaticLighting"));
 		
 		// @todo Mac OS X: For GL 3.3 devices which don't support volume-texture rendering we need to cache the simpler point indirect lighting shaders.
-		return Material->GetShadingModel() != MSM_Unlit 
+		return Material->GetShadingModel() != MSM_Unlit
 			&& (IsTranslucentBlendMode(Material->GetBlendMode()) || GetMaxSupportedFeatureLevel(Platform) == ERHIFeatureLevel::ES2
-				|| (PLATFORM_MAC && GetMaxSupportedFeatureLevel(Platform) == ERHIFeatureLevel::SM4))
+				|| (Platform == SP_OPENGL_SM4_MAC))
 			&& (!AllowStaticLightingVar || AllowStaticLightingVar->GetValueOnAnyThread() != 0);
 	}
 
diff --git a/Engine/Source/Runtime/ShaderCore/Private/ShaderParameters.cpp b/Engine/Source/Runtime/ShaderCore/Private/ShaderParameters.cpp
--- a/Engine/Source/Runtime/ShaderCore/Private/ShaderParameters.cpp
+++ b/Engine/Source/Runtime/ShaderCore/Private/ShaderParameters.cpp
@@ -319,6 +319,7 @@ FString CreateUniformBufferShaderDeclaration(const TCHAR* Name,const FUniformBuf
 	case SP_PCD3D_SM5:
 		return CreateHLSLUniformBufferDeclaration(Name,UniformBufferStruct, true);
 	case SP_OPENGL_SM4:
+	case SP_OPENGL_SM4_MAC:
 	case SP_OPENGL_SM5:
 		return CreateHLSLUniformBufferDeclaration(Name,UniformBufferStruct, false);
 	}
