- Added RHISupportsSeparateMSAAAndResolveTexture() to allow MSAA on Metal

diff --git a/Engine/Source/Runtime/RHI/Public/RHIDefinitions.h b/Engine/Source/Runtime/RHI/Public/RHIDefinitions.h
--- a/Engine/Source/Runtime/RHI/Public/RHIDefinitions.h
+++ b/Engine/Source/Runtime/RHI/Public/RHIDefinitions.h
@@ -615,6 +615,12 @@ inline bool RHISupportsInstancing(const EShaderPlatform Platform)
 	return IsFeatureLevelSupported(Platform, ERHIFeatureLevel::SM3);// || (Platform == SP_METAL);
 }
 
+inline bool RHISupportsSeparateMSAAAndResolveTextures(const EShaderPlatform Platform)
+{
+	// Metal needs to handle MSAA and resolve textures internally (unless RHICreateTexture2D was changed to take an optional resolve target)
+	return Platform != SP_METAL;
+}
+
 inline uint32 GetFeatureLevelMaxTextureSamplers(ERHIFeatureLevel::Type FeatureLevel)
 {
 	if (FeatureLevel == ERHIFeatureLevel::ES2)
diff --git a/Engine/Source/Runtime/RHI/Public/RHIUtilities.h b/Engine/Source/Runtime/RHI/Public/RHIUtilities.h
--- a/Engine/Source/Runtime/RHI/Public/RHIUtilities.h
+++ b/Engine/Source/Runtime/RHI/Public/RHIUtilities.h
@@ -196,7 +196,7 @@ inline void RHICreateTargetableShaderResource2D(
 
 	if (NumSamples > 1)
 	{
-		bForceSeparateTargetAndShaderResource = true;
+		bForceSeparateTargetAndShaderResource = RHISupportsSeparateMSAAAndResolveTextures(GRHIShaderPlatform);
 	}
 
 	if (!bForceSeparateTargetAndShaderResource/* && GSupportsRenderDepthTargetableShaderResources*/)
