Fix for dedservers crashing on start (TTP #343382).

- Result of me not testing well github changes.

diff --git a/Engine/Source/Runtime/RHI/RHI.Build.cs b/Engine/Source/Runtime/RHI/RHI.Build.cs
--- a/Engine/Source/Runtime/RHI/RHI.Build.cs
+++ b/Engine/Source/Runtime/RHI/RHI.Build.cs
@@ -8,7 +8,7 @@ public class RHI : ModuleRules
 	{
 		PrivateDependencyModuleNames.Add("Core");
 
-		if (UEBuildConfiguration.bCompileAgainstEngine && !UEBuildConfiguration.bBuildDedicatedServer)
+		if (UEBuildConfiguration.bCompileAgainstEngine)
 		{
             DynamicallyLoadedModuleNames.Add("NullDrv");
 
@@ -20,7 +20,7 @@ public class RHI : ModuleRules
 			if ((Target.Platform == UnrealTargetPlatform.Win32) ||
 				(Target.Platform == UnrealTargetPlatform.Win64) ||
 				(Target.Platform == UnrealTargetPlatform.Mac)   ||
-				(Target.Platform == UnrealTargetPlatform.Linux) ||
+                (Target.Platform == UnrealTargetPlatform.Linux && Target.Type != TargetRules.TargetType.Server) ||  // @todo should servers on all platforms skip this?
                 (Target.Platform == UnrealTargetPlatform.HTML5))
 			{
 				DynamicallyLoadedModuleNames.Add("OpenGLDrv");
