oops

diff --git a/Engine/Source/Runtime/RHI/Public/RHICommandList.inl b/Engine/Source/Runtime/RHI/Public/RHICommandList.inl
--- a/Engine/Source/Runtime/RHI/Public/RHICommandList.inl
+++ b/Engine/Source/Runtime/RHI/Public/RHICommandList.inl
@@ -17,6 +17,11 @@ FORCEINLINE_DEBUGGABLE void FRHICommandListBase::Flush()
 	}
 }
 
+FORCEINLINE_DEBUGGABLE bool FRHICommandListBase::IsImmediate()
+{
+	return this == &FRHICommandListExecutor::GetImmediateCommandList();
+}
+
 FORCEINLINE_DEBUGGABLE bool FRHICommandListBase::Bypass()
 {
 	return GRHICommandList.Bypass();
