Match condition for variables with check()s.

diff --git a/Engine/Source/Runtime/RHI/Public/RHICommandList.h b/Engine/Source/Runtime/RHI/Public/RHICommandList.h
--- a/Engine/Source/Runtime/RHI/Public/RHICommandList.h
+++ b/Engine/Source/Runtime/RHI/Public/RHICommandList.h
@@ -844,7 +844,7 @@ struct FLocalBoundShaderStateWorkArea
 {
 	FBoundShaderStateInput Args;
 	FComputedBSS* ComputedBSS;
-#if !(UE_BUILD_SHIPPING || UE_BUILD_TEST)
+#if DO_CHECK // the below variables are used in check(), which can be enabled in Shipping builds (see Build.h)
 	FRHICommandListBase* CheckCmdList;
 	int32 UID;
 #endif
@@ -859,7 +859,7 @@ struct FLocalBoundShaderStateWorkArea
 		FGeometryShaderRHIParamRef GeometryShaderRHI
 		)
 		: Args(VertexDeclarationRHI, VertexShaderRHI, HullShaderRHI, DomainShaderRHI, PixelShaderRHI, GeometryShaderRHI)
-#if !(UE_BUILD_SHIPPING || UE_BUILD_TEST)
+#if DO_CHECK
 		, CheckCmdList(InCheckCmdList)
 		, UID(InCheckCmdList->GetUID())
 #endif
@@ -951,14 +951,14 @@ struct FLocalUniformBufferWorkArea
 	void* Contents;
 	const FRHIUniformBufferLayout* Layout;
 	FComputedUniformBuffer* ComputedUniformBuffer;
-#if !(UE_BUILD_SHIPPING || UE_BUILD_TEST)
+#if DO_CHECK // the below variables are used in check(), which can be enabled in Shipping builds (see Build.h)
 	FRHICommandListBase* CheckCmdList;
 	int32 UID;
 #endif
 
 	FLocalUniformBufferWorkArea(FRHICommandListBase* InCheckCmdList, const void* InContents, uint32 ContentsSize, const FRHIUniformBufferLayout* InLayout)
 		: Layout(InLayout)
-#if !(UE_BUILD_SHIPPING || UE_BUILD_TEST)
+#if DO_CHECK
 		, CheckCmdList(InCheckCmdList)
 		, UID(InCheckCmdList->GetUID())
 #endif
