UE4 - fix possible crashes on parallel rendering toggle

diff --git a/Engine/Source/Runtime/RHI/Private/RHICommandList.cpp b/Engine/Source/Runtime/RHI/Private/RHICommandList.cpp
--- a/Engine/Source/Runtime/RHI/Private/RHICommandList.cpp
+++ b/Engine/Source/Runtime/RHI/Private/RHICommandList.cpp
@@ -98,13 +98,6 @@ void FRHICommandListExecutor::ExecuteList(FRHICommandListImmediate& CmdList)
 		//FTaskGraphInterface::Get().ProcessThreadUntilIdle(ENamedThreads::RenderThread_Local); 
 		ExecuteInner(CmdList);
 	}
-
-#if !UE_BUILD_SHIPPING
-	if (GRHICommandList.OutstandingCmdListCount.GetValue() == 1)
-	{
-		LatchBypass();
-	}
-#endif
 }
 
 
