UE4 - reorganize command list data structure for better cache coherency

diff --git a/Engine/Source/Runtime/RHI/Public/RHICommandList.h b/Engine/Source/Runtime/RHI/Public/RHICommandList.h
--- a/Engine/Source/Runtime/RHI/Public/RHICommandList.h
+++ b/Engine/Source/Runtime/RHI/Public/RHICommandList.h
@@ -99,22 +99,21 @@ public:
 
 	};
 
-	FDrawUpData DrawUPData;
-
 private:
-	FMemStackBase MemManager;
 	FRHICommandBase* Root;
 	FRHICommandBase** CommandLink;
 	bool bExecuting;
 	uint32 NumCommands;
 	uint32 UID;
-
+	FMemStackBase MemManager; // this has 1k inline storage, so we want to small members before it
 
 	void Reset();
 
 	friend class FRHICommandListExecutor;
 	friend class FRHICommandListIterator;
 
+public:
+	FDrawUpData DrawUPData; // this isn't used very much so we put it at the cold end here
 };
 
 template<typename TCmd>
