UE4 - RHI thread prototype (PS4 only so far)

diff --git a/Engine/Source/Runtime/Core/Private/HAL/ThreadingBase.cpp b/Engine/Source/Runtime/Core/Private/HAL/ThreadingBase.cpp
--- a/Engine/Source/Runtime/Core/Private/HAL/ThreadingBase.cpp
+++ b/Engine/Source/Runtime/Core/Private/HAL/ThreadingBase.cpp
@@ -40,7 +40,11 @@ CORE_API bool IsInParallelRenderingThread()
 	return !GRenderingThread || GIsRenderingThreadSuspended || (FPlatformTLS::GetCurrentThreadId() != GGameThreadId);
 }
 
-
+CORE_API bool IsInRHIThread()
+{
+	return GRHIThread && FPlatformTLS::GetCurrentThreadId() == GRHIThread->GetThreadID();
+}
+CORE_API FRunnableThread* GRHIThread = NULL;
 // Fake threads
 
 /**
diff --git a/Engine/Source/Runtime/Core/Private/Misc/MemStack.cpp b/Engine/Source/Runtime/Core/Private/Misc/MemStack.cpp
--- a/Engine/Source/Runtime/Core/Private/Misc/MemStack.cpp
+++ b/Engine/Source/Runtime/Core/Private/Misc/MemStack.cpp
@@ -14,6 +14,7 @@ DECLARE_MEMORY_STAT(TEXT("PageAllocator Used"), STAT_PageAllocatorUsed, STATGROU
 
 
 TLockFreeFixedSizeAllocator<FPageAllocator::PageSize, FThreadSafeCounter> FPageAllocator::TheAllocator;
+TLockFreeFixedSizeAllocator<FPageAllocator::SmallPageSize, FThreadSafeCounter> FPageAllocator::TheSmallAllocator;
 
 #if STATS
 void FPageAllocator::UpdateStats()
@@ -49,7 +50,7 @@ void FMemStackBase::AllocateNewChunk(int32 MinSize)
 	// Create new chunk.
 	int32 TotalSize = MinSize + (int32)sizeof(FTaggedMemory);
 	uint32 AllocSize;
-	if (TopChunk || TotalSize > SMALL_START_SIZE)
+	if (TopChunk || TotalSize > FPageAllocator::SmallPageSize)
 	{
 		AllocSize = AlignArbitrary<int32>(TotalSize, FPageAllocator::PageSize);
 		if (AllocSize == FPageAllocator::PageSize)
@@ -61,11 +62,12 @@ void FMemStackBase::AllocateNewChunk(int32 MinSize)
 			Chunk = (FTaggedMemory*)FMemory::Malloc(AllocSize);
 			INC_MEMORY_STAT_BY(STAT_MemStackLargeBLock, AllocSize);
 		}
+		check(AllocSize != FPageAllocator::SmallPageSize);
 	}
 	else
 	{
-		AllocSize = SMALL_START_SIZE;
-		Chunk = (FTaggedMemory*)SmallStart;
+		AllocSize = FPageAllocator::SmallPageSize;
+		Chunk = (FTaggedMemory*)FPageAllocator::AllocSmall();
 	}
 	Chunk->DataSize = AllocSize - sizeof(FTaggedMemory);
 
@@ -81,17 +83,18 @@ void FMemStackBase::FreeChunks(FTaggedMemory* NewTopChunk)
 	{
 		FTaggedMemory* RemoveChunk = TopChunk;
 		TopChunk                   = TopChunk->Next;
-		if (RemoveChunk != (FTaggedMemory*)SmallStart)
+		if (RemoveChunk->DataSize + sizeof(FTaggedMemory) == FPageAllocator::PageSize)
 		{
-			if (RemoveChunk->DataSize + sizeof(FTaggedMemory) == FPageAllocator::PageSize)
-			{
-				FPageAllocator::Free(RemoveChunk);
-			}
-			else
-			{
-				DEC_MEMORY_STAT_BY(STAT_MemStackLargeBLock, RemoveChunk->DataSize + sizeof(FTaggedMemory));
-				FMemory::Free(RemoveChunk);
-			}
+			FPageAllocator::Free(RemoveChunk);
+		}
+		else if (RemoveChunk->DataSize + sizeof(FTaggedMemory) == FPageAllocator::SmallPageSize)
+		{
+			FPageAllocator::FreeSmall(RemoveChunk);
+		}
+		else
+		{
+			DEC_MEMORY_STAT_BY(STAT_MemStackLargeBLock, RemoveChunk->DataSize + sizeof(FTaggedMemory));
+			FMemory::Free(RemoveChunk);
 		}
 	}
 	Top = nullptr;
diff --git a/Engine/Source/Runtime/Core/Private/Stats/Stats2.cpp b/Engine/Source/Runtime/Core/Private/Stats/Stats2.cpp
--- a/Engine/Source/Runtime/Core/Private/Stats/Stats2.cpp
+++ b/Engine/Source/Runtime/Core/Private/Stats/Stats2.cpp
@@ -868,7 +868,7 @@ void FThreadStats::Flush(bool bHasBrokenCallstacks)
 void FThreadStats::ExplicitFlush(bool DiscardCallstack)
 {
 	FThreadStats* ThreadStats = GetThreadStats();
-	check(ThreadStats->Packet.ThreadType != EThreadType::Other);
+	//check(ThreadStats->Packet.ThreadType != EThreadType::Other);
 	if (ThreadStats->bWaitForExplicitFlush)
 	{
 		ThreadStats->ScopeCount--; // the main thread pre-incremented this to prevent stats from being sent. we send them at the next available opportunity
diff --git a/Engine/Source/Runtime/Core/Public/Async/TaskGraphInterfaces.h b/Engine/Source/Runtime/Core/Public/Async/TaskGraphInterfaces.h
--- a/Engine/Source/Runtime/Core/Public/Async/TaskGraphInterfaces.h
+++ b/Engine/Source/Runtime/Core/Public/Async/TaskGraphInterfaces.h
@@ -32,6 +32,9 @@ namespace ENamedThreads
 		/** The always-present, named threads are listed next **/
 #if STATS
 		StatsThread, 
+#endif
+#if PLATFORM_SUPPORTS_RHI_THREAD
+		RHIThread,
 #endif
 		GameThread,
 		// The render thread is sometimes the game thread and is sometimes the actual rendering thread
diff --git a/Engine/Source/Runtime/Core/Public/Containers/LockFreeFixedSizeAllocator.h b/Engine/Source/Runtime/Core/Public/Containers/LockFreeFixedSizeAllocator.h
--- a/Engine/Source/Runtime/Core/Public/Containers/LockFreeFixedSizeAllocator.h
+++ b/Engine/Source/Runtime/Core/Public/Containers/LockFreeFixedSizeAllocator.h
@@ -27,7 +27,7 @@ public:
 		void *Memory = FreeList.Pop();
 		if (Memory)
 		{
-			verify(NumFree.Decrement() >= 0);
+			NumFree.Decrement();
 		}
 		else
 		{
diff --git a/Engine/Source/Runtime/Core/Public/HAL/Platform.h b/Engine/Source/Runtime/Core/Public/HAL/Platform.h
--- a/Engine/Source/Runtime/Core/Public/HAL/Platform.h
+++ b/Engine/Source/Runtime/Core/Public/HAL/Platform.h
@@ -238,6 +238,10 @@
 	#define PLATFORM_SUPPORTS_NAMED_PIPES		0
 #endif
 
+#ifndef PLATFORM_SUPPORTS_RHI_THREAD
+	#define PLATFORM_SUPPORTS_RHI_THREAD		0
+#endif
+
 #ifndef PLATFORM_SUPPORTS_MULTIPLE_NATIVE_WINDOWS
 	#define PLATFORM_SUPPORTS_MULTIPLE_NATIVE_WINDOWS	1
 #endif
diff --git a/Engine/Source/Runtime/Core/Public/HAL/ThreadingBase.h b/Engine/Source/Runtime/Core/Public/HAL/ThreadingBase.h
--- a/Engine/Source/Runtime/Core/Public/HAL/ThreadingBase.h
+++ b/Engine/Source/Runtime/Core/Public/HAL/ThreadingBase.h
@@ -888,6 +888,12 @@ extern CORE_API FRunnableThread* GRenderingThread;
 /** Whether the rendering thread is suspended (not even processing the tickables) */
 extern CORE_API int32 GIsRenderingThreadSuspended;
 
+/** @return True if called from the RHI thread, or if called from ANY thread during single threaded rendering */
+extern CORE_API bool IsInRHIThread();
+/** Thread used for RHI */
+extern CORE_API FRunnableThread* GRHIThread;
+
+
 /** Minimal base class for the thread singleton. */
 struct FThreadSingleton
 {
diff --git a/Engine/Source/Runtime/Core/Public/Misc/MemStack.h b/Engine/Source/Runtime/Core/Public/Misc/MemStack.h
--- a/Engine/Source/Runtime/Core/Public/Misc/MemStack.h
+++ b/Engine/Source/Runtime/Core/Public/Misc/MemStack.h
@@ -17,7 +17,8 @@ class CORE_API FPageAllocator
 public:
 	enum
 	{
-		PageSize = 64 * 1024
+		PageSize = 64 * 1024,
+		SmallPageSize = 1024
 	};
 	FORCEINLINE static void *Alloc()
 	{
@@ -30,13 +31,24 @@ public:
 		TheAllocator.Free(Mem);
 		STAT(UpdateStats());
 	}
+	FORCEINLINE static void *AllocSmall()
+	{
+		void *Result = TheSmallAllocator.Allocate();
+		STAT(UpdateStats());
+		return Result;
+	}
+	FORCEINLINE static void FreeSmall(void *Mem)
+	{
+		TheSmallAllocator.Free(Mem);
+		STAT(UpdateStats());
+	}
 	static uint64 BytesUsed()
 	{
-		return uint64(TheAllocator.GetNumUsed().GetValue()) * PageSize;
+		return uint64(TheAllocator.GetNumUsed().GetValue()) * PageSize + uint64(TheSmallAllocator.GetNumUsed().GetValue()) * SmallPageSize;
 	}
 	static uint64 BytesFree()
 	{
-		return uint64(TheAllocator.GetNumFree().GetValue()) * PageSize;
+		return uint64(TheAllocator.GetNumFree().GetValue()) * PageSize + uint64(TheSmallAllocator.GetNumFree().GetValue()) * SmallPageSize;
 	}
 private:
 
@@ -44,6 +56,7 @@ private:
 	static void UpdateStats();
 #endif
 	static TLockFreeFixedSizeAllocator<PageSize, FThreadSafeCounter> TheAllocator;
+	static TLockFreeFixedSizeAllocator<SmallPageSize, FThreadSafeCounter> TheSmallAllocator;
 };
 
 /**
@@ -167,13 +180,6 @@ private:
 
 	/** Used for a checkSlow. Most stacks require a mark to allocate. Command lists don't because they never mark, only flush*/
 	int32 MinMarksToAlloc;
-
-	// we are going to use an inline allocator for the first chunk to support things that are small and never get big
-	enum 
-	{
-		SMALL_START_SIZE = 1024 - sizeof(FTaggedMemory)
-	};
-	MS_ALIGN(16) uint8 SmallStart[SMALL_START_SIZE] GCC_ALIGN(16); 
 };
 
 class CORE_API FMemStack : public TThreadSingleton<FMemStack>, public FMemStackBase
diff --git a/Engine/Source/Runtime/Engine/Classes/Engine/Engine.h b/Engine/Source/Runtime/Engine/Classes/Engine/Engine.h
--- a/Engine/Source/Runtime/Engine/Classes/Engine/Engine.h
+++ b/Engine/Source/Runtime/Engine/Classes/Engine/Engine.h
@@ -1641,6 +1641,7 @@ public:
 	bool HandleFreezeAllCommand( const TCHAR* Cmd, FOutputDevice& Ar, UWorld* InWorld );			// Smedis
 	bool HandleFlushIOManagerCommand( const TCHAR* Cmd, FOutputDevice& Ar );						// Smedis
 	bool HandleToggleRenderingThreadCommand( const TCHAR* Cmd, FOutputDevice& Ar );
+	bool HandleToggleRHIThreadCommand( const TCHAR* Cmd, FOutputDevice& Ar );
 	bool HandleRecompileShadersCommand( const TCHAR* Cmd, FOutputDevice& Ar );
 	bool HandleRecompileGlobalShadersCommand( const TCHAR* Cmd, FOutputDevice& Ar );
 	bool HandleDumpShaderStatsCommand( const TCHAR* Cmd, FOutputDevice& Ar );
diff --git a/Engine/Source/Runtime/Engine/Private/GPUSkinVertexFactory.cpp b/Engine/Source/Runtime/Engine/Private/GPUSkinVertexFactory.cpp
--- a/Engine/Source/Runtime/Engine/Private/GPUSkinVertexFactory.cpp
+++ b/Engine/Source/Runtime/Engine/Private/GPUSkinVertexFactory.cpp
@@ -155,25 +155,31 @@ void FGPUBaseSkinVertexFactory::ShaderDataType::UpdateBoneData()
 
 FBoneDataVertexBuffer::FBoneDataVertexBuffer()
 	: SizeX(80 * 1024)		// todo: we will replace this fixed size using FGlobalDynamicVertexBuffer
+	, AllocSize(0)
+	, AllocBlock(nullptr)
 {
 }
 
 float* FBoneDataVertexBuffer::LockData()
 {
-	checkSlow(IsInRenderingThread());
-	checkSlow(GetSizeX());
-	checkSlow(IsValidRef(BoneBuffer));
-
-	float* Data = (float*)RHILockVertexBuffer(BoneBuffer.VertexBufferRHI, 0, ComputeMemorySize(), RLM_WriteOnly);
-	checkSlow(Data);
-
+	check(IsInRenderingThread() && GetSizeX() && IsValidRef(BoneBuffer));
+	check(!AllocSize && !AllocBlock);
+	AllocSize = ComputeMemorySize();
+	AllocBlock = FMemory::Malloc(((AllocSize + 15) / 16) * 16, 16);
+	float* Data = (float*)AllocBlock;
+	check(AllocSize && AllocBlock);
+	check(Data);
 	return Data;
 }
 
 void FBoneDataVertexBuffer::UnlockData()
 {
-	checkSlow(IsValidRef(BoneBuffer));
-	RHIUnlockVertexBuffer(BoneBuffer.VertexBufferRHI);
+	check(IsValidRef(BoneBuffer));
+	check(IsInRenderingThread() && AllocBlock && AllocSize);
+	FRHICommandListExecutor::GetImmediateCommandList().UpdateVertexBuffer(BoneBuffer.VertexBufferRHI, AllocBlock, AllocSize);
+	// UpdateVertexBuffer will free the block after the update
+	AllocBlock = nullptr;
+	AllocSize = 0;
 }
 
 uint32 FBoneDataVertexBuffer::GetSizeX() const
diff --git a/Engine/Source/Runtime/Engine/Private/SkeletalRenderGPUSkin.cpp b/Engine/Source/Runtime/Engine/Private/SkeletalRenderGPUSkin.cpp
--- a/Engine/Source/Runtime/Engine/Private/SkeletalRenderGPUSkin.cpp
+++ b/Engine/Source/Runtime/Engine/Private/SkeletalRenderGPUSkin.cpp
@@ -1053,7 +1053,7 @@ FPreviousPerBoneMotionBlur
 -----------------------------------------------------------------------------*/
 
 FPreviousPerBoneMotionBlur::FPreviousPerBoneMotionBlur()
-	:BufferIndex(0), LockedData(0), LockedTexelPosition(0), LockedTexelCount(0), bWarningBufferSizeExceeded(false)
+	:BufferIndex(0), LockedData(0), LockedTexelCount(0), bWarningBufferSizeExceeded(false)
 {
 }
 
@@ -1110,8 +1110,8 @@ void FPreviousPerBoneMotionBlur::InitIfNeeded()
 }
 void FPreviousPerBoneMotionBlur::LockData()
 {
-	checkSlow(!LockedData);
-	checkSlow(IsInRenderingThread());
+	check(!LockedData);
+	check(IsInRenderingThread());
 
 	InitIfNeeded();
 
@@ -1120,24 +1120,23 @@ void FPreviousPerBoneMotionBlur::LockData()
 	if(WriteTexture.IsValid())
 	{
 		LockedData = WriteTexture.LockData();
-		LockedTexelPosition = 0;
+		check(LockedTexelPosition.GetValue() == 0); //otherwise it wasn't unlocked or it was unlocked before it was done being filled by async stuff
+		LockedTexelPosition.Set(0);
 		LockedTexelCount = WriteTexture.GetSizeX();
 	}
 }
 
 uint32 FPreviousPerBoneMotionBlur::AppendData(FBoneSkinning *DataStart, uint32 BoneCount)
 {
-	checkSlow(LockedData);
+	check(LockedData);
 	checkSlow(DataStart);
 	checkSlow(BoneCount);
 
 	uint32 TexelCount = BoneCount * sizeof(FBoneSkinning) / sizeof(float) / 4;
 
-	uint32 OldLockedTexelPosition = LockedTexelPosition;
+	uint32 OldLockedTexelPosition = (uint32)LockedTexelPosition.Add(TexelCount);
 	
-	LockedTexelPosition += TexelCount;
-
-	if(LockedTexelPosition <= LockedTexelCount)
+	if(OldLockedTexelPosition + TexelCount <= LockedTexelCount)
 	{
 		FMemory::Memcpy(&LockedData[OldLockedTexelPosition * 4], DataStart, BoneCount * sizeof(FBoneSkinning));
 
@@ -1157,7 +1156,8 @@ void FPreviousPerBoneMotionBlur::UnlockData()
 {
 	if(IsLocked())
 	{
-		LockedTexelPosition = 0;
+		check(IsInRenderingThread());
+		LockedTexelPosition.Set(0);
 		LockedTexelCount = 0;
 		LockedData = 0;
 
diff --git a/Engine/Source/Runtime/Engine/Private/SkeletalRenderGPUSkin.h b/Engine/Source/Runtime/Engine/Private/SkeletalRenderGPUSkin.h
--- a/Engine/Source/Runtime/Engine/Private/SkeletalRenderGPUSkin.h
+++ b/Engine/Source/Runtime/Engine/Private/SkeletalRenderGPUSkin.h
@@ -216,7 +216,7 @@ private:
 	/* !=0 if data is locked, does not change during the lock */
 	float* LockedData;
 	/** only valid if LockedData != 0, advances with every Append() */
-	uint32 LockedTexelPosition;
+	FThreadSafeCounter LockedTexelPosition;
 	/** only valid if LockedData != 0 */
 	uint32 LockedTexelCount;
 
diff --git a/Engine/Source/Runtime/Engine/Private/UnrealEngine.cpp b/Engine/Source/Runtime/Engine/Private/UnrealEngine.cpp
--- a/Engine/Source/Runtime/Engine/Private/UnrealEngine.cpp
+++ b/Engine/Source/Runtime/Engine/Private/UnrealEngine.cpp
@@ -2259,6 +2259,10 @@ bool UEngine::Exec( UWorld* InWorld, const TCHAR* Cmd, FOutputDevice& Ar )
 	{
 		return HandleToggleRenderingThreadCommand( Cmd, Ar );
 	}
+	else if( FParse::Command(&Cmd,TEXT("ToggleRHIThread")) )
+	{
+		return HandleToggleRHIThreadCommand( Cmd, Ar );
+	}
 	else if( FParse::Command(&Cmd,TEXT("RecompileShaders")) )				    
 	{
 		return HandleRecompileShadersCommand( Cmd, Ar );
@@ -2922,6 +2926,28 @@ bool UEngine::HandleToggleRenderingThreadCommand( const TCHAR* Cmd, FOutputDevic
 	return true;
 }
 
+bool UEngine::HandleToggleRHIThreadCommand( const TCHAR* Cmd, FOutputDevice& Ar )
+{
+#if PLATFORM_SUPPORTS_RHI_THREAD
+	if (!GIsThreadedRendering)
+	{
+		check(!GRHIThread);
+		Ar.Logf( TEXT("Can't switch to RHI thread mode when we are not running a multithreaded renderer."));
+	}
+	else
+	{
+		bool bWasRHIThread = !!GRHIThread;
+		StopRenderingThread();
+		GUseRHIThread = !bWasRHIThread;
+		StartRenderingThread();
+	}
+	Ar.Logf( TEXT("RHIThread is now %s."), GRHIThread ? TEXT("active") : TEXT("inactive"));
+#else
+	Ar.Logf( TEXT("This platform does not support the RHI thread."));
+#endif
+	return true;
+}
+
 bool UEngine::HandleRecompileShadersCommand( const TCHAR* Cmd, FOutputDevice& Ar )
 {
 	return RecompileShaders(Cmd, Ar);
diff --git a/Engine/Source/Runtime/Engine/Public/GPUSkinVertexFactory.h b/Engine/Source/Runtime/Engine/Public/GPUSkinVertexFactory.h
--- a/Engine/Source/Runtime/Engine/Public/GPUSkinVertexFactory.h
+++ b/Engine/Source/Runtime/Engine/Public/GPUSkinVertexFactory.h
@@ -255,6 +255,12 @@ private: // -------------------------------------------------------
 	/** Buffer size in texels */
 	uint32 SizeX;
 
+	/** Buffer size, as allocated */
+	uint32 AllocSize;
+
+	/** Allocated Buffer */
+	void *AllocBlock;
+
 	/** @return in bytes */
 	uint32 ComputeMemorySize();
 };
diff --git a/Engine/Source/Runtime/Engine/Public/StaticRHI.h b/Engine/Source/Runtime/Engine/Public/StaticRHI.h
--- a/Engine/Source/Runtime/Engine/Public/StaticRHI.h
+++ b/Engine/Source/Runtime/Engine/Public/StaticRHI.h
@@ -192,6 +192,11 @@ private:
 		ReturnStatement GStaticRHI->RHI##Name ParameterNames; \
 	}
 #define DEFINE_RHIMETHOD_GLOBAL(Type,Name,ParameterTypesAndNames,ParameterNames,ReturnStatement,NullImplementation) \
+	FORCEINLINE Type Name##_Internal ParameterTypesAndNames \
+	{ \
+		ReturnStatement GStaticRHI->RHI##Name ParameterNames; \
+	}
+#define DEFINE_RHIMETHOD_GLOBALTHREADSAFE(Type,Name,ParameterTypesAndNames,ParameterNames,ReturnStatement,NullImplementation) \
 	FORCEINLINE Type RHI##Name ParameterTypesAndNames \
 	{ \
 		ReturnStatement GStaticRHI->RHI##Name ParameterNames; \
@@ -211,6 +216,7 @@ private:
 #undef DEFINE_RHIMETHOD_CMDLIST
 #undef DEFINE_RHIMETHOD_GLOBAL
 #undef DEFINE_RHIMETHOD_GLOBALFLUSH
+#undef DEFINE_RHIMETHOD_GLOBALTHREADSAFE
 
 #endif	//#if USE_STATIC_RHI
 
diff --git a/Engine/Source/Runtime/Launch/Private/LaunchEngineLoop.cpp b/Engine/Source/Runtime/Launch/Private/LaunchEngineLoop.cpp
--- a/Engine/Source/Runtime/Launch/Private/LaunchEngineLoop.cpp
+++ b/Engine/Source/Runtime/Launch/Private/LaunchEngineLoop.cpp
@@ -2012,14 +2012,16 @@ void FEngineLoop::Tick()
 		SCOPE_CYCLE_COUNTER( STAT_FrameTime );	
 
 		ENQUEUE_UNIQUE_RENDER_COMMAND(
-				BeginFrame,
+			BeginFrame,
+		{
+			GRHICommandList.LatchBypass();
+			GFrameNumberRenderThread++;
+			if (GRHICommandList.Bypass())
 			{
-				GRHICommandList.GetImmediateCommandList().Flush();
-				GRHICommandList.LatchBypass();
-				GFrameNumberRenderThread++;
 				GDynamicRHI->PushEvent(*FString::Printf(TEXT("Frame%d"),GFrameNumberRenderThread));
-				RHICmdList.BeginFrame();
-			});
+			}
+			RHICmdList.BeginFrame();
+		});
 
 		// Flush debug output which has been buffered by other threads.
 		GLog->FlushThreadedLogs();
@@ -2190,7 +2192,10 @@ void FEngineLoop::Tick()
 			EndFrame,
 		{
 			RHICmdList.EndFrame();
-			GDynamicRHI->PopEvent();
+			if (GRHICommandList.Bypass())
+			{
+				GDynamicRHI->PopEvent();
+			}
 		});
 	} 
 
diff --git a/Engine/Source/Runtime/RHI/Private/DynamicRHI.cpp b/Engine/Source/Runtime/RHI/Private/DynamicRHI.cpp
--- a/Engine/Source/Runtime/RHI/Private/DynamicRHI.cpp
+++ b/Engine/Source/Runtime/RHI/Private/DynamicRHI.cpp
@@ -82,6 +82,12 @@ void RHIExit()
 		ReturnStatement GDynamicRHI->RHI##Name ParameterNames; \
 	}
 #define DEFINE_RHIMETHOD_GLOBAL(Type,Name,ParameterTypesAndNames,ParameterNames,ReturnStatement,NullImplementation) \
+	RHI_API Type Name##_Internal ParameterTypesAndNames \
+	{ \
+		check(GDynamicRHI); \
+		ReturnStatement GDynamicRHI->RHI##Name ParameterNames; \
+	}
+#define DEFINE_RHIMETHOD_GLOBALTHREADSAFE(Type,Name,ParameterTypesAndNames,ParameterNames,ReturnStatement,NullImplementation) \
 	RHI_API Type RHI##Name ParameterTypesAndNames \
 	{ \
 		check(GDynamicRHI); \
@@ -104,6 +110,8 @@ void RHIExit()
 #undef DEFINE_RHIMETHOD_CMDLIST
 #undef DEFINE_RHIMETHOD_GLOBAL
 #undef DEFINE_RHIMETHOD_GLOBALFLUSH
+#undef DEFINE_RHIMETHOD_GLOBALTHREADSAFE
+
 
 #else
 
diff --git a/Engine/Source/Runtime/RHI/Private/RHI.cpp b/Engine/Source/Runtime/RHI/Private/RHI.cpp
--- a/Engine/Source/Runtime/RHI/Private/RHI.cpp
+++ b/Engine/Source/Runtime/RHI/Private/RHI.cpp
@@ -70,12 +70,15 @@ bool FRHIResource::Bypass()
 }
 #endif
 
+DECLARE_CYCLE_STAT(TEXT("Delete Resources"), STAT_DeleteResources, STATGROUP_RHICMDLIST);
+
 void FRHIResource::FlushPendingDeletes()
 {
+	SCOPE_CYCLE_COUNTER(STAT_DeleteResources);
+
 	check(IsInRenderingThread());
-	FRHICommandListExecutor::GetImmediateCommandList().Flush();
 	FRHICommandListExecutor::CheckNoOutstandingCmdLists();
-
+	FRHICommandListExecutor::GetImmediateCommandList().ImmediateFlush(EImmediateFlushType::FlushRHIThread);
 	while (1)
 	{
 		TArray<FRHIResource*> ToDelete;
diff --git a/Engine/Source/Runtime/RHI/Private/RHICommandList.cpp b/Engine/Source/Runtime/RHI/Private/RHICommandList.cpp
--- a/Engine/Source/Runtime/RHI/Private/RHICommandList.cpp
+++ b/Engine/Source/Runtime/RHI/Private/RHICommandList.cpp
@@ -12,6 +12,16 @@ DECLARE_CYCLE_STAT(TEXT("All Command List Execute"), STAT_ImmedCmdListExecuteTim
 DECLARE_DWORD_COUNTER_STAT(TEXT("Immed. Command List memory"), STAT_ImmedCmdListMemory, STATGROUP_RHICMDLIST);
 DECLARE_DWORD_COUNTER_STAT(TEXT("Immed. Command count"), STAT_ImmedCmdListCount, STATGROUP_RHICMDLIST);
 
+#if PLATFORM_SUPPORTS_RHI_THREAD
+/***
+Requirements for RHI thread
+* Microresources (those in RHIStaticStates.h) need to be able to be created by any thread at any time and be able to work with a radically simplified rhi resource lifecycle. CreateSamplerState, CreateRasterizerState, CreateDepthStencilState, CreateBlendState
+* CreateUniformBuffer needs to be threadsafe
+* GetRenderQueryResult needs to be threadsafe, but perhaps this can be supported if it can't be
+* AdvanceFrameForGetViewportBackBuffer needs be added as an RHI method and this needs to work with GetViewportBackBuffer to give the render thread the right back buffer even though many commands relating to the beginning and end of the frame are queued.
+
+***/
+#endif
 
 static TAutoConsoleVariable<int32> CVarRHICmdBypass(
 	TEXT("r.RHICmdBypass"),
@@ -21,14 +31,8 @@ static TAutoConsoleVariable<int32> CVarRHICmdBypass(
 
 static TAutoConsoleVariable<int32> CVarRHICmdUseParallelAlgorithms(
 	TEXT("r.RHICmdUseParallelAlgorithms"),
-	FRHICommandListExecutor::DefaultBypass,
-	TEXT("True to use parallel algorithms. Ignored if r.RHICmdBypass is 1.\n"));
-
-static TAutoConsoleVariable<int32> CVarRHICmdFlushBeforeWait(
-	TEXT("r.RHICmdFlushBeforeWait"),
 	1,
-	TEXT("Whether to flush the immediate command list before waiting for a command list chain.\n")
-	TEXT("0: Disable, 1: Enable"));
+	TEXT("True to use parallel algorithms. Ignored if r.RHICmdBypass is 1.\n"));
 
 TAutoConsoleVariable<int32> CVarRHICmdWidth(
 	TEXT("r.RHICmdWidth"), 
@@ -38,10 +42,14 @@ TAutoConsoleVariable<int32> CVarRHICmdWidth(
 
 RHI_API FRHICommandListExecutor GRHICommandList;
 
+static FGraphEventArray OutstandingTasks;
+static FGraphEventRef RHIThreadTask;
+static FGraphEventRef RenderThreadSublistDispatchTask;
 
-//RHI_API FGlobalRHI GRHI;
 
-void FRHICommandListExecutor::ExecuteInner(FRHICommandListBase& CmdList)
+DECLARE_CYCLE_STAT(TEXT("RHI Thread Execute"), STAT_RHIThreadExecute, STATGROUP_RHICMDLIST);
+
+void FRHICommandListExecutor::ExecuteInner_DoExecute(FRHICommandListBase& CmdList)
 {
 	CmdList.bExecuting = true;
 	{
@@ -55,11 +63,197 @@ void FRHICommandListExecutor::ExecuteInner(FRHICommandListBase& CmdList)
 	CmdList.Reset();
 }
 
-void FRHICommandListExecutor::ExecuteList(FRHICommandListBase& CmdList)
+class FExecuteRHIThreadTask
+{
+	FRHICommandListBase* RHICmdList;
+
+public:
+
+	FExecuteRHIThreadTask(FRHICommandListBase* InRHICmdList)
+		: RHICmdList(InRHICmdList)
+	{
+	}
+
+	FORCEINLINE TStatId GetStatId() const
+	{
+		RETURN_QUICK_DECLARE_CYCLE_STAT(FExecuteRHIThreadTask, STATGROUP_TaskGraphTasks);
+	}
+
+	ENamedThreads::Type GetDesiredThread()
+	{
+#if PLATFORM_SUPPORTS_RHI_THREAD
+		return ENamedThreads::RHIThread;
+#else
+		check(0); // this should never be used on a platform that doesn't support the RHI thread
+		return ENamedThreads::AnyThread;
+#endif
+	}
+
+	static ESubsequentsMode::Type GetSubsequentsMode() { return ESubsequentsMode::TrackSubsequents; }
+
+	void DoTask(ENamedThreads::Type CurrentThread, const FGraphEventRef& MyCompletionGraphEvent)
+	{
+		SCOPE_CYCLE_COUNTER(STAT_RHIThreadExecute);
+		FRHICommandListExecutor::ExecuteInner_DoExecute(*RHICmdList);
+		delete RHICmdList;
+	}
+};
+
+class FDispatchRHIThreadTask
 {
-	check(IsInRenderingThread() && &CmdList != &GetImmediateCommandList());
+	FRHICommandListBase* RHICmdList;
+
+public:
 
-	if (!GetImmediateCommandList().IsExecuting()) // don't flush if this is a recursive call and we are already executing the immediate command list
+	FDispatchRHIThreadTask(FRHICommandListBase* InRHICmdList)
+		: RHICmdList(InRHICmdList)
+	{
+	}
+
+	FORCEINLINE TStatId GetStatId() const
+	{
+		RETURN_QUICK_DECLARE_CYCLE_STAT(FDispatchRHIThreadTask, STATGROUP_TaskGraphTasks);
+	}
+
+	ENamedThreads::Type GetDesiredThread()
+	{
+#if !PLATFORM_SUPPORTS_RHI_THREAD
+		check(0); // this should never be used on a platform that doesn't support the RHI thread
+#endif
+		return ENamedThreads::RenderThread_Local;
+	}
+
+	static ESubsequentsMode::Type GetSubsequentsMode() { return ESubsequentsMode::TrackSubsequents; }
+
+	void DoTask(ENamedThreads::Type CurrentThread, const FGraphEventRef& MyCompletionGraphEvent)
+	{
+		check(IsInRenderingThread());
+		FGraphEventArray Prereq;
+		if (RHIThreadTask.GetReference())
+		{
+			Prereq.Add(RHIThreadTask);
+		}
+		RHIThreadTask = TGraphTask<FExecuteRHIThreadTask>::CreateTask(&Prereq, ENamedThreads::RenderThread).ConstructAndDispatchWhenReady(RHICmdList);
+	}
+};
+
+void FRHICommandListExecutor::ExecuteInner(FRHICommandListBase& CmdList)
+{
+	static TAutoConsoleVariable<int32> CVarRHICmdUseThread(
+		TEXT("r.RHICmdUseThread"),
+		1,
+		TEXT("Uses the RHI thread.\n")
+		TEXT("0: Disable, 1: Enable"),
+		ECVF_Cheat);
+	static TAutoConsoleVariable<int32> CVarRHICmdForceRHIFlush(
+		TEXT("r.RHICmdForceRHIFlush"),
+		0,
+		TEXT("Force a flush for every task sent to the RHI thread.\n")
+		TEXT("0: Disable, 1: Enable"),
+		ECVF_Cheat);
+	check(CmdList.RTTasks.Num() == 0 && CmdList.HasCommands()); 
+
+	if (GRHIThread)
+	{
+		bool bIsInRenderingThread = IsInRenderingThread();
+
+		if (bIsInRenderingThread)
+		{
+			if (RHIThreadTask.GetReference() && RHIThreadTask->IsComplete())
+			{
+				RHIThreadTask = nullptr;
+			}
+			if (RenderThreadSublistDispatchTask.GetReference() && RenderThreadSublistDispatchTask->IsComplete())
+			{
+				RenderThreadSublistDispatchTask = nullptr;
+			}
+		}
+		if (CVarRHICmdUseThread.GetValueOnRenderThread() > 0 && bIsInRenderingThread && !IsInGameThread())
+		{
+			FRHICommandList* SwapCmdList = new FRHICommandList;
+
+			// Super scary stuff here, but we just want the swap command list to inherit everything and leave the immediate command list wiped.
+			// we should make command lists virtual and transfer ownership rather than this devious approach
+			static_assert(sizeof(FRHICommandList) == sizeof(FRHICommandListImmediate), "We are memswapping FRHICommandList and FRHICommandListImmediate; they need to be swappable.");
+			SwapCmdList->ExchangeCmdList(CmdList);
+
+			if (OutstandingTasks.Num() || RenderThreadSublistDispatchTask.GetReference())
+			{
+				FGraphEventArray Prereq(OutstandingTasks);
+				OutstandingTasks.Reset();
+				if (RenderThreadSublistDispatchTask.GetReference())
+				{
+					Prereq.Add(RenderThreadSublistDispatchTask);
+				}
+				RenderThreadSublistDispatchTask = TGraphTask<FDispatchRHIThreadTask>::CreateTask(&Prereq, ENamedThreads::RenderThread).ConstructAndDispatchWhenReady(SwapCmdList);
+			}
+			else
+			{
+				FGraphEventArray Prereq;
+				if (RHIThreadTask.GetReference())
+				{
+					Prereq.Add(RHIThreadTask);
+				}
+				RHIThreadTask = TGraphTask<FExecuteRHIThreadTask>::CreateTask(&Prereq, ENamedThreads::RenderThread).ConstructAndDispatchWhenReady(SwapCmdList);
+			}
+			if (CVarRHICmdForceRHIFlush.GetValueOnRenderThread() > 0 )
+			{
+				if (FTaskGraphInterface::Get().IsThreadProcessingTasks(ENamedThreads::RenderThread_Local))
+				{
+					// this is a deadlock. RT tasks must be done by now or they won't be done. We could add a third queue...
+					UE_LOG(LogRHI, Fatal, TEXT("Deadlock in FRHICommandListExecutor::ExecuteInner 2."));
+				}
+				if (RenderThreadSublistDispatchTask.GetReference())
+				{
+					FTaskGraphInterface::Get().WaitUntilTaskCompletes(RenderThreadSublistDispatchTask, ENamedThreads::RenderThread_Local);
+					RenderThreadSublistDispatchTask = nullptr;
+				}
+				if (RHIThreadTask.GetReference())
+				{
+					FTaskGraphInterface::Get().WaitUntilTaskCompletes(RHIThreadTask, ENamedThreads::RenderThread_Local);
+					RHIThreadTask = nullptr;
+				}
+			}
+			return;
+		}
+		if (bIsInRenderingThread)
+		{
+			if (RenderThreadSublistDispatchTask.GetReference())
+			{
+				if (FTaskGraphInterface::Get().IsThreadProcessingTasks(ENamedThreads::RenderThread_Local))
+				{
+					// this is a deadlock. RT tasks must be done by now or they won't be done. We could add a third queue...
+					UE_LOG(LogRHI, Fatal, TEXT("Deadlock in FRHICommandListExecutor::ExecuteInner 3."));
+				}
+				FTaskGraphInterface::Get().WaitUntilTaskCompletes(RenderThreadSublistDispatchTask, ENamedThreads::RenderThread_Local);
+				RenderThreadSublistDispatchTask = nullptr;
+			}
+			if (RHIThreadTask.GetReference())
+			{
+				if (FTaskGraphInterface::Get().IsThreadProcessingTasks(ENamedThreads::RenderThread_Local))
+				{
+					// this is a deadlock. RT tasks must be done by now or they won't be done. We could add a third queue...
+					UE_LOG(LogRHI, Fatal, TEXT("Deadlock in FRHICommandListExecutor::ExecuteInner 3."));
+				}
+				FTaskGraphInterface::Get().WaitUntilTaskCompletes(RHIThreadTask, ENamedThreads::RenderThread_Local);
+				RHIThreadTask = nullptr;
+			}
+		}
+	}
+
+	ExecuteInner_DoExecute(CmdList);
+}
+
+
+static FORCEINLINE bool IsInRenderingOrRHIThread()
+{
+	return IsInRenderingThread() || IsInRHIThread();
+}
+
+void FRHICommandListExecutor::ExecuteList(FRHICommandListBase& CmdList)
+{
+	check(IsInRenderingOrRHIThread() && &CmdList != &GetImmediateCommandList());
+	if (IsInRenderingThread() && !GetImmediateCommandList().IsExecuting()) // don't flush if this is a recursive call and we are already executing the immediate command list
 	{
 		GetImmediateCommandList().Flush();
 	}
@@ -73,7 +267,7 @@ void FRHICommandListExecutor::ExecuteList(FRHICommandListBase& CmdList)
 
 void FRHICommandListExecutor::ExecuteList(FRHICommandListImmediate& CmdList)
 {
-	check(IsInRenderingThread() && &CmdList == &GetImmediateCommandList());
+	check(IsInRenderingOrRHIThread() && &CmdList == &GetImmediateCommandList());
 
 	INC_MEMORY_STAT_BY(STAT_ImmedCmdListMemory, CmdList.GetUsedMemory());
 	INC_DWORD_STAT_BY(STAT_ImmedCmdListCount, CmdList.NumCommands);
@@ -95,7 +289,6 @@ void FRHICommandListExecutor::ExecuteList(FRHICommandListImmediate& CmdList)
 #endif
 	{
 		SCOPE_CYCLE_COUNTER(STAT_ImmedCmdListExecuteTime);
-		//FTaskGraphInterface::Get().ProcessThreadUntilIdle(ENamedThreads::RenderThread_Local); 
 		ExecuteInner(CmdList);
 	}
 }
@@ -104,34 +297,50 @@ void FRHICommandListExecutor::ExecuteList(FRHICommandListImmediate& CmdList)
 void FRHICommandListExecutor::LatchBypass()
 {
 #if !UE_BUILD_SHIPPING
-
-	static bool bOnce = false;
-	if (!bOnce)
+	if (GRHIThread)
 	{
-		bOnce = true;
-		if (FParse::Param(FCommandLine::Get(),TEXT("parallelrendering")) && CVarRHICmdBypass.GetValueOnRenderThread() >= 1)
+		if (bLatchedBypass)
 		{
-			IConsoleVariable* BypassVar = IConsoleManager::Get().FindConsoleVariable(TEXT("r.RHICmdBypass"));
-			BypassVar->Set(0);
+			check((GRHICommandList.OutstandingCmdListCount.GetValue() == 1 && !GRHICommandList.GetImmediateCommandList().HasCommands()));
+			bLatchedBypass = false;
 		}
 	}
-
-	check(GRHICommandList.OutstandingCmdListCount.GetValue() == 1 && !GRHICommandList.GetImmediateCommandList().HasCommands());
-	bool NewBypass = (CVarRHICmdBypass.GetValueOnAnyThread() >= 1);
-	if (NewBypass && !bLatchedBypass)
+	else
 	{
-		FRHIResource::FlushPendingDeletes();
+		GRHICommandList.GetImmediateCommandList().ImmediateFlush(EImmediateFlushType::DispatchToRHIThread);
+
+		static bool bOnce = false;
+		if (!bOnce)
+		{
+			bOnce = true;
+			if (FParse::Param(FCommandLine::Get(),TEXT("parallelrendering")) && CVarRHICmdBypass.GetValueOnRenderThread() >= 1)
+			{
+				IConsoleVariable* BypassVar = IConsoleManager::Get().FindConsoleVariable(TEXT("r.RHICmdBypass"));
+				BypassVar->Set(0);
+			}
+		}
+
+		check((GRHICommandList.OutstandingCmdListCount.GetValue() == 1 && !GRHICommandList.GetImmediateCommandList().HasCommands()));
+		bool NewBypass = (CVarRHICmdBypass.GetValueOnAnyThread() >= 1);
+		if (NewBypass && !bLatchedBypass)
+		{
+			FRHIResource::FlushPendingDeletes();
+		}
+		bLatchedBypass = NewBypass;
 	}
-	bLatchedBypass = NewBypass;
+#endif
 	if (!bLatchedBypass)
 	{
-		bLatchedUseParallelAlgorithms = (CVarRHICmdUseParallelAlgorithms.GetValueOnAnyThread() >= 1);
+		bLatchedUseParallelAlgorithms = FApp::ShouldUseThreadingForPerformance() 
+#if !UE_BUILD_SHIPPING
+			&& (CVarRHICmdUseParallelAlgorithms.GetValueOnAnyThread() >= 1)
+#endif
+			;
 	}
 	else
 	{
 		bLatchedUseParallelAlgorithms = false;
 	}
-#endif
 }
 
 void FRHICommandListExecutor::CheckNoOutstandingCmdLists()
@@ -139,6 +348,35 @@ void FRHICommandListExecutor::CheckNoOutstandingCmdLists()
 	check(GRHICommandList.OutstandingCmdListCount.GetValue() == 1); // else we are attempting to delete resources while there is still a live cmdlist (other than the immediate cmd list) somewhere.
 }
 
+FGraphEventRef FRHICommandListExecutor::RHIThreadFence()
+{
+	check(IsInRenderingThread() && GRHIThread);
+	{
+		QUICK_SCOPE_CYCLE_COUNTER(STAT_RHIThreadFence_Dispatch);
+		FRHICommandListExecutor::GetImmediateCommandList().ImmediateFlush(EImmediateFlushType::DispatchToRHIThread);
+	}
+	if (RHIThreadTask.GetReference() && RHIThreadTask->IsComplete())
+	{
+		RHIThreadTask = nullptr;
+	}
+	return RHIThreadTask;
+}
+void FRHICommandListExecutor::WaitOnRHIThreadFence(FGraphEventRef& Fence)
+{
+	check(IsInRenderingThread() && GRHIThread);
+	if (Fence.GetReference() && !Fence->IsComplete())
+	{
+		QUICK_SCOPE_CYCLE_COUNTER(STAT_WaitOnRHIThreadFence_Wait);
+		if (FTaskGraphInterface::Get().IsThreadProcessingTasks(ENamedThreads::RenderThread_Local))
+		{
+			// this is a deadlock. RT tasks must be done by now or they won't be done. We could add a third queue...
+			UE_LOG(LogRHI, Fatal, TEXT("Deadlock in WaitOnRHIThreadFence."));
+		}
+		FTaskGraphInterface::Get().WaitUntilTaskCompletes(Fence, ENamedThreads::RenderThread_Local);
+	}
+}
+
+
 FRHICommandListBase::FRHICommandListBase()
 	: MemManager(0)
 {
@@ -160,6 +398,7 @@ const int32 FRHICommandListBase::GetUsedMemory() const
 void FRHICommandListBase::Reset()
 {
 	bExecuting = false;
+	check(!RTTasks.Num());
 	MemManager.Flush();
 	NumCommands = 0;
 	Root = nullptr;
@@ -183,9 +422,18 @@ struct FRHICommandWaitForAndSubmitSubList : public FRHICommand<FRHICommandWaitFo
 	void Execute()
 	{
 		INC_DWORD_STAT_BY(STAT_ChainLinkCount, 1);
+		if (!EventToWaitFor->IsComplete())
 		{
+			check(!GRHIThread || !IsInRHIThread()); // things should not be dispatched if they can't complete without further waits
 			SCOPE_CYCLE_COUNTER(STAT_ChainWait);
-			FTaskGraphInterface::Get().WaitUntilTaskCompletes(EventToWaitFor, ENamedThreads::RenderThread_Local);
+			if (IsInRenderingThread())
+			{
+				FTaskGraphInterface::Get().WaitUntilTaskCompletes(EventToWaitFor, ENamedThreads::RenderThread_Local);
+			}
+			else
+			{
+				check(0);
+			}
 		}
 		{
 			SCOPE_CYCLE_COUNTER(STAT_ChainExecute);
@@ -196,7 +444,17 @@ struct FRHICommandWaitForAndSubmitSubList : public FRHICommand<FRHICommandWaitFo
 
 void FRHICommandListBase::QueueAsyncCommandListSubmit(FGraphEventRef& AnyThreadCompletionEvent, class FRHICommandList* CmdList)
 {
+	check(IsInRenderingThread() && IsImmediate());
+	OutstandingTasks.Add(AnyThreadCompletionEvent);
+	if (GRHIThread)
+	{
+		FRHICommandListExecutor::GetImmediateCommandList().ImmediateFlush(EImmediateFlushType::DispatchToRHIThread); // we should start on the stuff before this async list
+	}
 	new (AllocCommand<FRHICommandWaitForAndSubmitSubList>()) FRHICommandWaitForAndSubmitSubList(AnyThreadCompletionEvent, CmdList);
+	if (GRHIThread)
+	{
+		FRHICommandListExecutor::GetImmediateCommandList().ImmediateFlush(EImmediateFlushType::DispatchToRHIThread); // we don't want stuff after the async cmd list to be bundled with it
+	}
 }
 
 DECLARE_DWORD_COUNTER_STAT(TEXT("Num RT Chains Links"), STAT_RTChainLinkCount, STATGROUP_RHICMDLIST);
@@ -217,28 +475,22 @@ struct FRHICommandWaitForAndSubmitRTSubList : public FRHICommand<FRHICommandWait
 		INC_DWORD_STAT_BY(STAT_RTChainLinkCount, 1);
 		{
 			SCOPE_CYCLE_COUNTER(STAT_RTChainWait);
-			if (IsInRenderingThread())
+			if (!EventToWaitFor->IsComplete())
 			{
-				if (!EventToWaitFor->IsComplete())
+				check(!GRHIThread || !IsInRHIThread()); // things should not be dispatched if they can't complete without further waits
+				if (IsInRenderingThread())
 				{
 					if (FTaskGraphInterface::Get().IsThreadProcessingTasks(ENamedThreads::RenderThread_Local))
 					{
 						// this is a deadlock. RT tasks must be done by now or they won't be done. We could add a third queue...
 						UE_LOG(LogRHI, Fatal, TEXT("Deadlock in command list processing."));
 					}
-#if 0
-					while (!EventToWaitFor->IsComplete())
-					{
-						FTaskGraphInterface::Get().ProcessThreadUntilIdle(ENamedThreads::RenderThread_Local);
-					}
-#endif
 					FTaskGraphInterface::Get().WaitUntilTaskCompletes(EventToWaitFor, ENamedThreads::RenderThread_Local);
 				}
-			}
-			else
-			{
-				FScopedEvent Waiter;
-				FTaskGraphInterface::Get().TriggerEventWhenTaskCompletes(Waiter.Get(), EventToWaitFor);
+				else
+				{
+					FTaskGraphInterface::Get().WaitUntilTaskCompletes(EventToWaitFor);
+				}
 			}
 		}
 		{
@@ -250,6 +502,8 @@ struct FRHICommandWaitForAndSubmitRTSubList : public FRHICommand<FRHICommandWait
 
 void FRHICommandListBase::QueueRenderThreadCommandListSubmit(FGraphEventRef& RenderThreadCompletionEvent, class FRHICommandList* CmdList)
 {
+	check(!IsInRenderingThread() && !IsImmediate() && !IsInRHIThread());
+	RTTasks.Add(RenderThreadCompletionEvent);
 	new (AllocCommand<FRHICommandWaitForAndSubmitRTSubList>()) FRHICommandWaitForAndSubmitRTSubList(RenderThreadCompletionEvent, CmdList);
 }
 
@@ -273,6 +527,202 @@ void FRHICommandListBase::QueueCommandListSubmit(class FRHICommandList* CmdList)
 	new (AllocCommand<FRHICommandSubmitSubList>()) FRHICommandSubmitSubList(CmdList);
 }
 
+#if PLATFORM_SUPPORTS_RHI_THREAD
+struct FRHICommandBeginDrawingViewport : public FRHICommand<FRHICommandBeginDrawingViewport>
+{
+	FViewportRHIParamRef Viewport;
+	FTextureRHIParamRef RenderTargetRHI;
+
+	FORCEINLINE_DEBUGGABLE FRHICommandBeginDrawingViewport(FViewportRHIParamRef InViewport, FTextureRHIParamRef InRenderTargetRHI)
+		: Viewport(InViewport)
+		, RenderTargetRHI(InRenderTargetRHI)
+	{
+	}
+	void Execute()
+	{
+		BeginDrawingViewport_Internal(Viewport, RenderTargetRHI);
+	}
+};
+
+void FRHICommandList::BeginDrawingViewport(FViewportRHIParamRef Viewport, FTextureRHIParamRef RenderTargetRHI)
+{
+	check(IsImmediate() && IsInRenderingThread());
+	if (Bypass())
+	{
+		BeginDrawingViewport_Internal(Viewport, RenderTargetRHI);
+		return;
+	}
+	new (AllocCommand<FRHICommandBeginDrawingViewport>()) FRHICommandBeginDrawingViewport(Viewport, RenderTargetRHI);
+}
+struct FRHICommandEndDrawingViewport : public FRHICommand<FRHICommandEndDrawingViewport>
+{
+	FViewportRHIParamRef Viewport;
+	bool bPresent;
+	bool bLockToVsync;
+
+	FORCEINLINE_DEBUGGABLE FRHICommandEndDrawingViewport(FViewportRHIParamRef InViewport, bool InbPresent, bool InbLockToVsync)
+		: Viewport(InViewport)
+		, bPresent(InbPresent)
+		, bLockToVsync(InbLockToVsync)
+	{
+	}
+	void Execute()
+	{
+		EndDrawingViewport_Internal(Viewport, bPresent, bLockToVsync);
+	}
+};
+
+void FRHICommandList::EndDrawingViewport(FViewportRHIParamRef Viewport, bool bPresent, bool bLockToVsync)
+{
+	check(IsImmediate() && IsInRenderingThread());
+	if (Bypass())
+	{
+		EndDrawingViewport_Internal(Viewport, bPresent, bLockToVsync);
+	}
+	else
+	{
+		new (AllocCommand<FRHICommandEndDrawingViewport>()) FRHICommandEndDrawingViewport(Viewport, bPresent, bLockToVsync);
+		// this should be started asap
+		{
+			QUICK_SCOPE_CYCLE_COUNTER(STAT_EndDrawingViewport_Dispatch);
+			FRHICommandListExecutor::GetImmediateCommandList().ImmediateFlush(EImmediateFlushType::DispatchToRHIThread);
+		}
+	}
+	RHIAdvanceFrameForGetViewportBackBuffer();
+}
+
+void FRHICommandList::BeginFrame()
+{
+	if (Bypass())
+	{
+		BeginFrame_Internal();
+		return;
+	}
+	new (AllocCommand<FRHICommandBeginFrame>()) FRHICommandBeginFrame();
+}
+
+void FRHICommandList::EndFrame()
+{
+	if (Bypass())
+	{
+		EndFrame_Internal();
+		return;
+	}
+	new (AllocCommand<FRHICommandEndFrame>()) FRHICommandEndFrame();
+}
+
+#endif
+
+DECLARE_CYCLE_STAT(TEXT("Explicit wait for tasks"), STAT_ExplicitWait, STATGROUP_RHICMDLIST);
+DECLARE_CYCLE_STAT(TEXT("Prewait dispatch"), STAT_PrewaitDispatch, STATGROUP_RHICMDLIST);
+void FRHICommandListBase::WaitForTasks(bool bKnownToBeComplete)
+{
+	check(IsImmediate() && IsInRenderingThread());
+	if (GRHIThread)
+	{
+		if (bKnownToBeComplete)
+		{
+			check(!OutstandingTasks.Num());
+		}
+		else if (OutstandingTasks.Num() || RenderThreadSublistDispatchTask.GetReference())
+		{
+			{
+				SCOPE_CYCLE_COUNTER(STAT_PrewaitDispatch);
+				// as long as we are sitting around, make sure the rhi thread is not
+				FRHICommandListExecutor::GetImmediateCommandList().ImmediateFlush(EImmediateFlushType::DispatchToRHIThread);
+			}
+			check(!OutstandingTasks.Num());
+			if (RenderThreadSublistDispatchTask.GetReference())
+			{
+				if (FTaskGraphInterface::Get().IsThreadProcessingTasks(ENamedThreads::RenderThread_Local))
+				{
+					// this is a deadlock. RT tasks must be done by now or they won't be done. We could add a third queue...
+					UE_LOG(LogRHI, Fatal, TEXT("Deadlock in FRHICommandListBase::WaitForRHIThreadTasks."));
+				}
+				FTaskGraphInterface::Get().WaitUntilTaskCompletes(RenderThreadSublistDispatchTask, ENamedThreads::RenderThread_Local);
+				RenderThreadSublistDispatchTask = nullptr;
+			}
+		}
+	}	
+	else
+	{
+		if (OutstandingTasks.Num())
+		{
+			if (bKnownToBeComplete)
+			{
+				for (int32 Index = 0; Index < OutstandingTasks.Num(); Index++)
+				{
+					check(OutstandingTasks[Index]->IsComplete());
+				}
+			}
+			else
+			{
+				SCOPE_CYCLE_COUNTER(STAT_ExplicitWait);
+				FTaskGraphInterface::Get().WaitUntilTasksComplete(OutstandingTasks, ENamedThreads::RenderThread_Local);
+			}
+			OutstandingTasks.Reset();
+		}
+	}
+}
+
+DECLARE_CYCLE_STAT(TEXT("Explicit wait for RHI thread"), STAT_ExplicitWaitRHIThread, STATGROUP_RHICMDLIST);
+DECLARE_CYCLE_STAT(TEXT("Deep spin for stray resource init."), STAT_SpinWaitRHIThread, STATGROUP_RHICMDLIST);
+void FRHICommandListBase::WaitForRHIThreadTasks()
+{
+	check(IsImmediate() && IsInRenderingThread() && GRHIThread);
+	if (RHIThreadTask.GetReference() && RHIThreadTask->IsComplete())
+	{
+		RHIThreadTask = nullptr;
+	}
+	if (RHIThreadTask.GetReference())
+	{
+		SCOPE_CYCLE_COUNTER(STAT_ExplicitWaitRHIThread);
+		if (FTaskGraphInterface::Get().IsThreadProcessingTasks(ENamedThreads::RenderThread_Local))
+		{
+			// we have to spin here because all task threads might be stalled, meaning the fire event anythread task might not be hit.
+			// todo, add a third queue
+			SCOPE_CYCLE_COUNTER(STAT_SpinWaitRHIThread);
+			while (!RHIThreadTask->IsComplete())
+			{
+				FPlatformProcess::Sleep(0);
+			}
+		}
+		else
+		{
+			FTaskGraphInterface::Get().WaitUntilTaskCompletes(RHIThreadTask, ENamedThreads::RenderThread_Local);
+		}
+		RHIThreadTask = nullptr;
+	}
+}
+
+DECLARE_CYCLE_STAT(TEXT("RTTask completion join"), STAT_HandleRTThreadTaskCompletion_Join, STATGROUP_RHICMDLIST);
+void FRHICommandListBase::HandleRTThreadTaskCompletion(const FGraphEventRef& MyCompletionGraphEvent)
+{
+	check(!IsImmediate() && !IsInRenderingThread() && !IsInRHIThread())
+	for (int32 Index = 0; Index < RTTasks.Num(); Index++)
+	{
+		if (RTTasks[Index]->IsComplete())
+		{
+			RTTasks.RemoveAtSwap(Index, 1, false);
+			Index--;
+		}
+	}
+	if (RTTasks.Num())
+	{
+		FGraphEventRef WaitFor;
+		if (RTTasks.Num() > 1)
+		{
+			WaitFor = TGraphTask<FNullGraphTask>::CreateTask(&RTTasks).ConstructAndDispatchWhenReady(GET_STATID(STAT_HandleRTThreadTaskCompletion_Join));
+		}
+		else
+		{
+			WaitFor = RTTasks[0];
+		}
+		MyCompletionGraphEvent->DontCompleteUntil(WaitFor);
+	}
+	RTTasks.Empty();
+}
+
 static TLockFreeFixedSizeAllocator<sizeof(FRHICommandList), FThreadSafeCounter> RHICommandListAllocator;
 
 void* FRHICommandList::operator new(size_t Size)
@@ -288,6 +738,7 @@ void* FRHICommandList::operator new(size_t Size)
  */
 void FRHICommandList::operator delete(void *RawMemory)
 {
+	check(RawMemory != (void*) &GRHICommandList.GetImmediateCommandList());
 	RHICommandListAllocator.Free(RawMemory);
 	//FMemory::Free(RawMemory);
 }	
diff --git a/Engine/Source/Runtime/RHI/Public/DynamicRHI.h b/Engine/Source/Runtime/RHI/Public/DynamicRHI.h
--- a/Engine/Source/Runtime/RHI/Public/DynamicRHI.h
+++ b/Engine/Source/Runtime/RHI/Public/DynamicRHI.h
@@ -46,6 +46,8 @@ extern RHI_API FDynamicRHI* GDynamicRHI;
 #define DEFINE_RHIMETHOD_CMDLIST(Type,Name,ParameterTypesAndNames,ParameterNames,ReturnStatement,NullImplementation) \
 	extern RHI_API Type Name##_Internal ParameterTypesAndNames;
 #define DEFINE_RHIMETHOD_GLOBAL(Type,Name,ParameterTypesAndNames,ParameterNames,ReturnStatement,NullImplementation) \
+	extern RHI_API Type Name##_Internal ParameterTypesAndNames;
+#define DEFINE_RHIMETHOD_GLOBALTHREADSAFE(Type,Name,ParameterTypesAndNames,ParameterNames,ReturnStatement,NullImplementation) \
 	extern RHI_API Type RHI##Name ParameterTypesAndNames;
 #define DEFINE_RHIMETHOD_GLOBALFLUSH(Type,Name,ParameterTypesAndNames,ParameterNames,ReturnStatement,NullImplementation) \
 	extern RHI_API Type Name##_Internal ParameterTypesAndNames;
@@ -56,6 +58,8 @@ extern RHI_API FDynamicRHI* GDynamicRHI;
 #undef DEFINE_RHIMETHOD_CMDLIST
 #undef DEFINE_RHIMETHOD_GLOBAL
 #undef DEFINE_RHIMETHOD_GLOBALFLUSH
+#undef DEFINE_RHIMETHOD_GLOBALTHREADSAFE
+
 
 #endif	//#if USE_DYNAMIC_RHI
 
diff --git a/Engine/Source/Runtime/RHI/Public/RHI.h b/Engine/Source/Runtime/RHI/Public/RHI.h
--- a/Engine/Source/Runtime/RHI/Public/RHI.h
+++ b/Engine/Source/Runtime/RHI/Public/RHI.h
@@ -1064,6 +1064,8 @@ DECLARE_MEMORY_STAT_POOL_EXTERN(TEXT("Pixel buffer memory"),STAT_PixelBufferMemo
 #define DEFINE_RHIMETHOD_CMDLIST(Type,Name,ParameterTypesAndNames,ParameterNames,ReturnStatement,NullImplementation) \
 	extern Type Name##_Internal ParameterTypesAndNames
 #define DEFINE_RHIMETHOD_GLOBAL(Type,Name,ParameterTypesAndNames,ParameterNames,ReturnStatement,NullImplementation) \
+	extern Type Name##_Internal ParameterTypesAndNames
+#define DEFINE_RHIMETHOD_GLOBALTHREADSAFE(Type,Name,ParameterTypesAndNames,ParameterNames,ReturnStatement,NullImplementation) \
 	extern Type RHI##Name ParameterTypesAndNames
 #define DEFINE_RHIMETHOD_GLOBALFLUSH(Type,Name,ParameterTypesAndNames,ParameterNames,ReturnStatement,NullImplementation) \
 	extern Type Name##_Internal ParameterTypesAndNames
@@ -1075,6 +1077,8 @@ DECLARE_MEMORY_STAT_POOL_EXTERN(TEXT("Pixel buffer memory"),STAT_PixelBufferMemo
 #undef DEFINE_RHIMETHOD_CMDLIST
 #undef DEFINE_RHIMETHOD_GLOBAL
 #undef DEFINE_RHIMETHOD_GLOBALFLUSH
+#undef DEFINE_RHIMETHOD_GLOBALTHREADSAFE
+
 
 #endif
 
diff --git a/Engine/Source/Runtime/RHI/Public/RHICommandList.h b/Engine/Source/Runtime/RHI/Public/RHICommandList.h
--- a/Engine/Source/Runtime/RHI/Public/RHICommandList.h
+++ b/Engine/Source/Runtime/RHI/Public/RHICommandList.h
@@ -40,6 +40,9 @@ public:
 	void QueueAsyncCommandListSubmit(FGraphEventRef& AnyThreadCompletionEvent, class FRHICommandList* CmdList);
 	void QueueRenderThreadCommandListSubmit(FGraphEventRef& RenderThreadCompletionEvent, class FRHICommandList* CmdList);
 	void QueueCommandListSubmit(class FRHICommandList* CmdList);
+	void WaitForTasks(bool bKnownToBeComplete = false);
+	void WaitForRHIThreadTasks();
+	void HandleRTThreadTaskCompletion(const FGraphEventRef& MyCompletionGraphEvent);
 
 	FORCEINLINE_DEBUGGABLE void* Alloc(int32 AllocSize, int32 Alignment)
 	{
@@ -77,6 +80,21 @@ public:
 
 	inline bool Bypass();
 
+	FORCEINLINE void ExchangeCmdList(FRHICommandListBase& Other)
+	{
+		check(!RTTasks.Num() && !Other.RTTasks.Num());
+		FMemory::Memswap(this, &Other, sizeof(FRHICommandListBase));
+		if (CommandLink == &Other.Root)
+		{
+			CommandLink = &Root;
+		}
+		if (Other.CommandLink == &Root)
+		{
+			Other.CommandLink = &Other.Root;
+		}
+	}
+
+
 	struct FDrawUpData
 	{
 
@@ -103,6 +121,7 @@ private:
 	FRHICommandBase* Root;
 	FRHICommandBase** CommandLink;
 	bool bExecuting;
+	FGraphEventArray RTTasks;
 	uint32 NumCommands;
 	uint32 UID;
 	FMemStackBase MemManager; // this has 1k inline storage, so we want to small members before it
@@ -1008,7 +1027,11 @@ struct FRHICommandBuildLocalUniformBuffer : public FRHICommand<FRHICommandBuildL
 		check(!IsValidRef(WorkArea.ComputedUniformBuffer->UniformBuffer) && WorkArea.Layout && WorkArea.Contents); // should not already have been created
 		if (WorkArea.ComputedUniformBuffer->UseCount)
 		{
+#if PLATFORM_SUPPORTS_RHI_THREAD
 			WorkArea.ComputedUniformBuffer->UniformBuffer = RHICreateUniformBuffer(WorkArea.Contents, *WorkArea.Layout, UniformBuffer_SingleFrame);
+#else
+			WorkArea.ComputedUniformBuffer->UniformBuffer = CreateUniformBuffer_Internal(WorkArea.Contents, *WorkArea.Layout, UniformBuffer_SingleFrame);
+#endif
 		}
 		WorkArea.Layout = nullptr;
 		WorkArea.Contents = nullptr;
@@ -1043,6 +1066,115 @@ struct FRHICommandSetLocalUniformBuffer : public FRHICommand<FRHICommandSetLocal
 
 };
 
+struct FRHICommandBeginRenderQuery : public FRHICommand<FRHICommandBeginRenderQuery>
+{
+	FRenderQueryRHIParamRef RenderQuery;
+
+	FORCEINLINE_DEBUGGABLE FRHICommandBeginRenderQuery(FRenderQueryRHIParamRef InRenderQuery)
+		: RenderQuery(InRenderQuery)
+	{
+	}
+	void Execute()
+	{
+		BeginRenderQuery_Internal(RenderQuery);
+	}
+};
+
+struct FRHICommandEndRenderQuery : public FRHICommand<FRHICommandEndRenderQuery>
+{
+	FRenderQueryRHIParamRef RenderQuery;
+
+	FORCEINLINE_DEBUGGABLE FRHICommandEndRenderQuery(FRenderQueryRHIParamRef InRenderQuery)
+		: RenderQuery(InRenderQuery)
+	{
+	}
+	void Execute()
+	{
+		EndRenderQuery_Internal(RenderQuery);
+	}
+};
+
+struct FRHICommandResetRenderQuery : public FRHICommand<FRHICommandResetRenderQuery>
+{
+	FRenderQueryRHIParamRef RenderQuery;
+
+	FORCEINLINE_DEBUGGABLE FRHICommandResetRenderQuery(FRenderQueryRHIParamRef InRenderQuery)
+		: RenderQuery(InRenderQuery)
+	{
+	}
+	void Execute()
+	{
+		ResetRenderQuery_Internal(RenderQuery);
+	}
+};
+
+struct FRHICommandBeginScene : public FRHICommand<FRHICommandBeginScene>
+{
+	FORCEINLINE_DEBUGGABLE FRHICommandBeginScene()
+	{
+	}
+	void Execute()
+	{
+		BeginScene_Internal();
+	}
+};
+
+struct FRHICommandEndScene : public FRHICommand<FRHICommandEndScene>
+{
+	FORCEINLINE_DEBUGGABLE FRHICommandEndScene()
+	{
+	}
+	void Execute()
+	{
+		EndScene_Internal();
+	}
+};
+
+struct FRHICommandUpdateVertexBuffer : public FRHICommand<FRHICommandUpdateVertexBuffer>
+{
+	FVertexBufferRHIParamRef VertexBuffer;
+	void const* Buffer;
+	int32 BufferSize;
+
+	FORCEINLINE_DEBUGGABLE FRHICommandUpdateVertexBuffer(FVertexBufferRHIParamRef InVertexBuffer, void const* InBuffer, int32 InBufferSize)
+		: VertexBuffer(InVertexBuffer)
+		, Buffer(InBuffer)
+		, BufferSize(InBufferSize)
+	{
+	}
+	void Execute()
+	{
+		void* Data = LockVertexBuffer_Internal(VertexBuffer, 0, BufferSize, RLM_WriteOnly);
+		FMemory::Memcpy(Data, Buffer, BufferSize);
+		UnlockVertexBuffer_Internal(VertexBuffer);
+		FMemory::Free((void*)Buffer);
+	}
+};
+
+struct FRHICommandBeginFrame : public FRHICommand<FRHICommandBeginFrame>
+{
+	FORCEINLINE_DEBUGGABLE FRHICommandBeginFrame()
+	{
+	}
+	void Execute()
+	{
+		BeginFrame_Internal();
+	}
+};
+
+struct FRHICommandEndFrame : public FRHICommand<FRHICommandEndFrame>
+{
+	FORCEINLINE_DEBUGGABLE FRHICommandEndFrame()
+	{
+	}
+	void Execute()
+	{
+		EndFrame_Internal();
+	}
+};
+
+
+
 class RHI_API FRHICommandList : public FRHICommandListBase
 {
 public:
@@ -1112,7 +1244,11 @@ public:
 		FLocalUniformBuffer Result;
 		if (Bypass())
 		{
+#if PLATFORM_SUPPORTS_RHI_THREAD
 			Result.BypassUniform = RHICreateUniformBuffer(Contents, Layout, UniformBuffer_SingleFrame);
+#else
+			Result.BypassUniform = CreateUniformBuffer_Internal(Contents, Layout, UniformBuffer_SingleFrame);
+#endif
 		}
 		else
 		{
@@ -1528,6 +1664,83 @@ public:
 		}
 		new (AllocCommand<FRHICommandClearMRT>()) FRHICommandClearMRT(bClearColor, NumClearColors, ClearColorArray, bClearDepth, Depth, bClearStencil, Stencil, ExcludeRect);
 	}
+
+	FORCEINLINE_DEBUGGABLE void UpdateVertexBuffer(FVertexBufferRHIParamRef VertexBuffer, void const* Buffer, int32 BufferSize)
+	{
+		if (Bypass())
+		{
+			void* Data = LockVertexBuffer_Internal(VertexBuffer, 0, BufferSize, RLM_WriteOnly);
+			FMemory::Memcpy(Data, Buffer, BufferSize);
+			UnlockVertexBuffer_Internal(VertexBuffer);
+			FMemory::Free((void*)Buffer);
+			return;
+		}
+		new (AllocCommand<FRHICommandUpdateVertexBuffer>()) FRHICommandUpdateVertexBuffer(VertexBuffer, Buffer, BufferSize);
+	}
+	FORCEINLINE_DEBUGGABLE void BeginRenderQuery(FRenderQueryRHIParamRef RenderQuery)
+	{
+		if (Bypass())
+		{
+			BeginRenderQuery_Internal(RenderQuery);
+			return;
+		}
+		new (AllocCommand<FRHICommandBeginRenderQuery>()) FRHICommandBeginRenderQuery(RenderQuery);
+	}
+	FORCEINLINE_DEBUGGABLE void EndRenderQuery(FRenderQueryRHIParamRef RenderQuery)
+	{
+		if (Bypass())
+		{
+			EndRenderQuery_Internal(RenderQuery);
+			return;
+		}
+		new (AllocCommand<FRHICommandEndRenderQuery>()) FRHICommandEndRenderQuery(RenderQuery);
+	}
+	FORCEINLINE_DEBUGGABLE void ResetRenderQuery(FRenderQueryRHIParamRef RenderQuery)
+	{
+		if (Bypass())
+		{
+			ResetRenderQuery_Internal(RenderQuery);
+			return;
+		}
+		new (AllocCommand<FRHICommandResetRenderQuery>()) FRHICommandResetRenderQuery(RenderQuery);
+	}
+
+#if PLATFORM_SUPPORTS_RHI_THREAD
+	FORCEINLINE_DEBUGGABLE void BeginScene()
+	{
+		if (Bypass())
+		{
+			BeginScene_Internal();
+			return;
+		}
+		new (AllocCommand<FRHICommandBeginScene>()) FRHICommandBeginScene();
+	}
+	FORCEINLINE_DEBUGGABLE void EndScene()
+	{
+		if (Bypass())
+		{
+			EndScene_Internal();
+			return;
+		}
+		new (AllocCommand<FRHICommandEndScene>()) FRHICommandEndScene();
+	}
+	void BeginDrawingViewport(FViewportRHIParamRef Viewport, FTextureRHIParamRef RenderTargetRHI);
+	void EndDrawingViewport(FViewportRHIParamRef Viewport, bool bPresent, bool bLockToVsync);
+	void BeginFrame();
+	void EndFrame();
+#endif
+};
+
+namespace EImmediateFlushType
+{
+	enum Type
+	{ 
+		WaitForOutstandingTasksOnly = 0, 
+		DispatchToRHIThread, 
+		WaitForRHIThread, 
+		FlushRHIThread,
+		FlushRHIThreadFlushResources
+	};
 };
 
 class RHI_API FRHICommandListImmediate : public FRHICommandList
@@ -1538,20 +1751,31 @@ class RHI_API FRHICommandListImmediate : public FRHICommandList
 	}
 	~FRHICommandListImmediate()
 	{
-		Flush(); // this probably never happens, but we want to be sure there are no commands, otherwise it will be executed like a non-immediate command list
+		check(!HasCommands());
 	}
 public:
 
-	inline void Flush();
+	inline void ImmediateFlush(EImmediateFlushType::Type FlushType);
 
 	#define DEFINE_RHIMETHOD_CMDLIST(Type,Name,ParameterTypesAndNames,ParameterNames,ReturnStatement,NullImplementation)
 	#define DEFINE_RHIMETHOD(Type, Name, ParameterTypesAndNames, ParameterNames, ReturnStatement, NullImplementation) \
 		FORCEINLINE_DEBUGGABLE Type Name ParameterTypesAndNames \
 		{ \
-			Flush(); \
+			QUICK_SCOPE_CYCLE_COUNTER(STAT_RHIMETHOD_##Name##_Flush); \
+			ImmediateFlush(EImmediateFlushType::FlushRHIThread); \
 			ReturnStatement Name##_Internal ParameterNames; \
 		}
 	#define DEFINE_RHIMETHOD_GLOBAL(Type, Name, ParameterTypesAndNames, ParameterNames, ReturnStatement, NullImplementation) \
+		FORCEINLINE_DEBUGGABLE Type Name ParameterTypesAndNames \
+		{ \
+			if (GRHIThread) \
+			{ \
+				QUICK_SCOPE_CYCLE_COUNTER(STAT_RHIMETHOD_##Name##_WaitRHI); \
+				ImmediateFlush(EImmediateFlushType::WaitForRHIThread); \
+			} \
+			ReturnStatement Name##_Internal ParameterNames; \
+		}
+	#define DEFINE_RHIMETHOD_GLOBALTHREADSAFE(Type, Name, ParameterTypesAndNames, ParameterNames, ReturnStatement, NullImplementation) \
 		FORCEINLINE_DEBUGGABLE Type Name ParameterTypesAndNames \
 		{ \
 			ReturnStatement RHI##Name ParameterNames; \
@@ -1559,7 +1783,8 @@ public:
 	#define DEFINE_RHIMETHOD_GLOBALFLUSH(Type, Name, ParameterTypesAndNames, ParameterNames, ReturnStatement, NullImplementation) \
 		FORCEINLINE_DEBUGGABLE Type Name ParameterTypesAndNames \
 		{ \
-			Flush(); \
+			QUICK_SCOPE_CYCLE_COUNTER(STAT_RHIMETHOD_##Name##_Flush); \
+			ImmediateFlush(EImmediateFlushType::FlushRHIThread); \
 			ReturnStatement Name##_Internal ParameterNames; \
 		}
 	#include "RHIMethods.h"
@@ -1567,43 +1792,10 @@ public:
 	#undef DEFINE_RHIMETHOD_CMDLIST
 	#undef DEFINE_RHIMETHOD_GLOBAL
 	#undef DEFINE_RHIMETHOD_GLOBALFLUSH
-};
+	#undef DEFINE_RHIMETHOD_GLOBALTHREADSAFE
 
-#if 0
-class RHI_API FRHICommandListBypass
-{
-	friend class FRHICommandListExecutor;
-	FRHICommandListBypass()
-	{
-	}
-public:
-#define DEFINE_RHIMETHOD_CMDLIST(Type,Name,ParameterTypesAndNames,ParameterNames,ReturnStatement,NullImplementation) \
-	FORCEINLINE_DEBUGGABLE Type Name ParameterTypesAndNames \
-	{ \
-		ReturnStatement Name##_Internal ParameterNames; \
-	}
-#define DEFINE_RHIMETHOD(Type, Name, ParameterTypesAndNames, ParameterNames, ReturnStatement, NullImplementation) \
-	FORCEINLINE_DEBUGGABLE Type Name ParameterTypesAndNames \
-	{ \
-		ReturnStatement Name##_Internal ParameterNames; \
-	}
-#define DEFINE_RHIMETHOD_GLOBAL(Type, Name, ParameterTypesAndNames, ParameterNames, ReturnStatement, NullImplementation) \
-	FORCEINLINE_DEBUGGABLE Type Name ParameterTypesAndNames \
-	{ \
-		ReturnStatement RHI##Name ParameterNames; \
-	}
-#define DEFINE_RHIMETHOD_GLOBALFLUSH(Type, Name, ParameterTypesAndNames, ParameterNames, ReturnStatement, NullImplementation) \
-	FORCEINLINE_DEBUGGABLE Type Name ParameterTypesAndNames \
-	{ \
-		ReturnStatement Name##_Internal ParameterNames; \
-	}
-#include "RHIMethods.h"
-#undef DEFINE_RHIMETHOD
-#undef DEFINE_RHIMETHOD_CMDLIST
-#undef DEFINE_RHIMETHOD_GLOBAL
-#undef DEFINE_RHIMETHOD_GLOBALFLUSH
 };
-#endif
+
 
 // typedef to mark the recursive use of commandlists in the RHI implmentations
 typedef FRHICommandList FRHICommandList_RecursiveHazardous;
@@ -1617,7 +1809,7 @@ public:
 	};
 	FRHICommandListExecutor()
 		: bLatchedBypass(!!DefaultBypass)
-		, bLatchedUseParallelAlgorithms(!DefaultBypass && FApp::ShouldUseThreadingForPerformance())
+		, bLatchedUseParallelAlgorithms(false)
 	{
 	}
 	static inline FRHICommandListImmediate& GetImmediateCommandList();
@@ -1626,9 +1818,11 @@ public:
 	void ExecuteList(FRHICommandListImmediate& CmdList);
 	void LatchBypass();
 
+	static FGraphEventRef RHIThreadFence();
+	static void WaitOnRHIThreadFence(FGraphEventRef& Fence);
+
 	FORCEINLINE_DEBUGGABLE void Verify()
 	{
-		check(CommandListImmediate.IsExecuting() || !CommandListImmediate.HasCommands());
 	}
 	FORCEINLINE_DEBUGGABLE bool Bypass()
 	{
@@ -1651,6 +1845,8 @@ public:
 private:
 
 	void ExecuteInner(FRHICommandListBase& CmdList);
+	friend class FExecuteRHIThreadTask;
+	static void ExecuteInner_DoExecute(FRHICommandListBase& CmdList);
 
 	bool bLatchedBypass;
 	bool bLatchedUseParallelAlgorithms;
@@ -1667,23 +1863,46 @@ FORCEINLINE_DEBUGGABLE FRHICommandListImmediate& FRHICommandListExecutor::GetImm
 	return GRHICommandList.CommandListImmediate;
 }
 
-struct FScopedCommandListFlush
+struct FScopedCommandListWaitForTasks
 {
 	FRHICommandListImmediate& RHICmdList;
 
-	FScopedCommandListFlush(FRHICommandListImmediate& InRHICmdList = FRHICommandListExecutor::GetImmediateCommandList())
+	FScopedCommandListWaitForTasks(FRHICommandListImmediate& InRHICmdList = FRHICommandListExecutor::GetImmediateCommandList())
 		: RHICmdList(InRHICmdList)
 	{
 	}
-	~FScopedCommandListFlush()
+	~FScopedCommandListWaitForTasks()
 	{
-		RHICmdList.Flush();
+		check(IsInRenderingThread());
+		if (GRHIThread)
+		{
+			{
+				QUICK_SCOPE_CYCLE_COUNTER(STAT_FScopedCommandListWaitForTasks_Dispatch);
+				RHICmdList.ImmediateFlush(EImmediateFlushType::DispatchToRHIThread);
+			}
+			{
+				QUICK_SCOPE_CYCLE_COUNTER(STAT_FScopedCommandListWaitForTasks_WaitAsync);
+				RHICmdList.ImmediateFlush(EImmediateFlushType::WaitForOutstandingTasksOnly);
+			}
+		}
+		else
+		{
+			QUICK_SCOPE_CYCLE_COUNTER(STAT_FScopedCommandListWaitForTasks_Flush);
+			RHICmdList.ImmediateFlush(EImmediateFlushType::FlushRHIThread);
+		}
 	}
 };
 
 #define DEFINE_RHIMETHOD_CMDLIST(Type,Name,ParameterTypesAndNames,ParameterNames,ReturnStatement,NullImplementation)
 #define DEFINE_RHIMETHOD(Type, Name, ParameterTypesAndNames, ParameterNames, ReturnStatement, NullImplementation)
-#define DEFINE_RHIMETHOD_GLOBAL(Type, Name, ParameterTypesAndNames, ParameterNames, ReturnStatement, NullImplementation)
+#define DEFINE_RHIMETHOD_GLOBALTHREADSAFE(Type, Name, ParameterTypesAndNames, ParameterNames, ReturnStatement, NullImplementation)
+
+#define DEFINE_RHIMETHOD_GLOBAL(Type, Name, ParameterTypesAndNames, ParameterNames, ReturnStatement, NullImplementation) \
+	FORCEINLINE_DEBUGGABLE Type RHI##Name ParameterTypesAndNames \
+	{ \
+		ReturnStatement FRHICommandListExecutor::GetImmediateCommandList().Name ParameterNames; \
+	}
+
 #define DEFINE_RHIMETHOD_GLOBALFLUSH(Type, Name, ParameterTypesAndNames, ParameterNames, ReturnStatement, NullImplementation) \
 	FORCEINLINE_DEBUGGABLE Type RHI##Name ParameterTypesAndNames \
 	{ \
@@ -1694,5 +1913,7 @@ struct FScopedCommandListFlush
 #undef DEFINE_RHIMETHOD_CMDLIST
 #undef DEFINE_RHIMETHOD_GLOBAL
 #undef DEFINE_RHIMETHOD_GLOBALFLUSH
+#undef DEFINE_RHIMETHOD_GLOBALTHREADSAFE
+
 
 #include "RHICommandList.inl"
diff --git a/Engine/Source/Runtime/RHI/Public/RHICommandList.inl b/Engine/Source/Runtime/RHI/Public/RHICommandList.inl
--- a/Engine/Source/Runtime/RHI/Public/RHICommandList.inl
+++ b/Engine/Source/Runtime/RHI/Public/RHICommandList.inl
@@ -13,6 +13,7 @@ FORCEINLINE_DEBUGGABLE void FRHICommandListBase::Flush()
 	if (HasCommands())
 	{
 		check(!DrawUPData.OutVertexData && !DrawUPData.OutIndexData);
+		check(!IsImmediate());
 		GRHICommandList.ExecuteList(*this);
 	}
 }
@@ -28,12 +29,57 @@ FORCEINLINE_DEBUGGABLE bool FRHICommandListBase::Bypass()
 }
 
 
-FORCEINLINE_DEBUGGABLE void FRHICommandListImmediate::Flush()
+FORCEINLINE_DEBUGGABLE void FRHICommandListImmediate::ImmediateFlush(EImmediateFlushType::Type FlushType)
 {
-	if (HasCommands())
+	check(!DrawUPData.OutVertexData && !DrawUPData.OutIndexData);
+	switch (FlushType)
 	{
-		check(!DrawUPData.OutVertexData && !DrawUPData.OutIndexData);
-		GRHICommandList.ExecuteList(*this);
+	case EImmediateFlushType::WaitForOutstandingTasksOnly:
+		{
+			WaitForTasks();
+		}
+		break;
+	case EImmediateFlushType::DispatchToRHIThread:
+		{
+			if (HasCommands())
+			{
+				GRHICommandList.ExecuteList(*this);
+			}
+		}
+		break;
+	case EImmediateFlushType::WaitForRHIThread:
+		{
+			check(GRHIThread);
+			WaitForRHIThreadTasks();
+		}
+		break;
+	case EImmediateFlushType::FlushRHIThread:
+		{
+			if (HasCommands())
+			{
+				GRHICommandList.ExecuteList(*this);
+			}
+			if (GRHIThread)
+			{
+				WaitForRHIThreadTasks();
+			}
+			WaitForTasks(true); // these are already done, but this resets the outstanding array
+		}
+		break;
+	case EImmediateFlushType::FlushRHIThreadFlushResources:
+		{
+			check(GRHIThread);
+			if (HasCommands())
+			{
+				GRHICommandList.ExecuteList(*this);
+			}
+			WaitForRHIThreadTasks();
+			WaitForTasks(true); // these are already done, but this resets the outstanding array
+			FRHIResource::FlushPendingDeletes();
+		}
+		break;
+	default:
+		check(0);
 	}
 }
 
diff --git a/Engine/Source/Runtime/RHI/Public/RHIMethods.h b/Engine/Source/Runtime/RHI/Public/RHIMethods.h
--- a/Engine/Source/Runtime/RHI/Public/RHIMethods.h
+++ b/Engine/Source/Runtime/RHI/Public/RHIMethods.h
@@ -52,6 +52,16 @@
 #undef DEFINE_RHIMETHOD_GLOBALFLUSH_7
 #undef DEFINE_RHIMETHOD_GLOBALFLUSH_8
 #undef DEFINE_RHIMETHOD_GLOBALFLUSH_9
+#undef DEFINE_RHIMETHOD_GLOBALTHREADSAFE_0
+#undef DEFINE_RHIMETHOD_GLOBALTHREADSAFE_1
+#undef DEFINE_RHIMETHOD_GLOBALTHREADSAFE_2
+#undef DEFINE_RHIMETHOD_GLOBALTHREADSAFE_3
+#undef DEFINE_RHIMETHOD_GLOBALTHREADSAFE_4
+#undef DEFINE_RHIMETHOD_GLOBALTHREADSAFE_5
+#undef DEFINE_RHIMETHOD_GLOBALTHREADSAFE_6
+#undef DEFINE_RHIMETHOD_GLOBALTHREADSAFE_7
+#undef DEFINE_RHIMETHOD_GLOBALTHREADSAFE_8
+#undef DEFINE_RHIMETHOD_GLOBALTHREADSAFE_9
 
 #ifdef GENERATE_VAX_FUNCTION_DECLARATIONS
 
@@ -139,6 +149,27 @@
 	#define DEFINE_RHIMETHOD_GLOBALFLUSH_9(ReturnType,MethodName,ParameterTypeA,ParameterNameA,ParameterTypeB,ParameterNameB,ParameterTypeC,ParameterNameC,ParameterTypeD,ParameterNameD,ParameterTypeE,ParameterNameE,ParameterTypeF,ParameterNameF,ParameterTypeG,ParameterNameG,ParameterTypeH,ParameterNameH,ParameterTypeI,ParameterNameI,ReturnStatement,NullImplementation) \
 		ReturnType RHI##MethodName(ParameterTypeA ParameterNameA, ParameterTypeB ParameterNameB, ParameterTypeC ParameterNameC, ParameterTypeD ParameterNameD, ParameterTypeE ParameterNameE, ParameterTypeF ParameterNameF, ParameterTypeG ParameterNameG, ParameterTypeH ParameterNameH, ParameterTypeI ParameterNameI);
 
+	#define DEFINE_RHIMETHOD_GLOBALTHREADSAFE_0(ReturnType,MethodName,ReturnStatement,NullImplementation) \
+		ReturnType RHI##MethodName();
+	#define DEFINE_RHIMETHOD_GLOBALTHREADSAFE_1(ReturnType,MethodName,ParameterTypeA,ParameterNameA,ReturnStatement,NullImplementation) \
+		ReturnType RHI##MethodName(ParameterTypeA ParameterNameA);
+	#define DEFINE_RHIMETHOD_GLOBALTHREADSAFE_2(ReturnType,MethodName,ParameterTypeA,ParameterNameA,ParameterTypeB,ParameterNameB,ReturnStatement,NullImplementation) \
+		ReturnType RHI##MethodName(ParameterTypeA ParameterNameA, ParameterTypeB ParameterNameB);
+	#define DEFINE_RHIMETHOD_GLOBALTHREADSAFE_3(ReturnType,MethodName,ParameterTypeA,ParameterNameA,ParameterTypeB,ParameterNameB,ParameterTypeC,ParameterNameC,ReturnStatement,NullImplementation) \
+		ReturnType RHI##MethodName(ParameterTypeA ParameterNameA, ParameterTypeB ParameterNameB, ParameterTypeC ParameterNameC);
+	#define DEFINE_RHIMETHOD_GLOBALTHREADSAFE_4(ReturnType,MethodName,ParameterTypeA,ParameterNameA,ParameterTypeB,ParameterNameB,ParameterTypeC,ParameterNameC,ParameterTypeD,ParameterNameD,ReturnStatement,NullImplementation) \
+		ReturnType RHI##MethodName(ParameterTypeA ParameterNameA, ParameterTypeB ParameterNameB, ParameterTypeC ParameterNameC, ParameterTypeD ParameterNameD);
+	#define DEFINE_RHIMETHOD_GLOBALTHREADSAFE_5(ReturnType,MethodName,ParameterTypeA,ParameterNameA,ParameterTypeB,ParameterNameB,ParameterTypeC,ParameterNameC,ParameterTypeD,ParameterNameD,ParameterTypeE,ParameterNameE,ReturnStatement,NullImplementation) \
+		ReturnType RHI##MethodName(ParameterTypeA ParameterNameA, ParameterTypeB ParameterNameB, ParameterTypeC ParameterNameC, ParameterTypeD ParameterNameD, ParameterTypeE ParameterNameE);
+	#define DEFINE_RHIMETHOD_GLOBALTHREADSAFE_6(ReturnType,MethodName,ParameterTypeA,ParameterNameA,ParameterTypeB,ParameterNameB,ParameterTypeC,ParameterNameC,ParameterTypeD,ParameterNameD,ParameterTypeE,ParameterNameE,ParameterTypeF,ParameterNameF,ReturnStatement,NullImplementation) \
+		ReturnType RHI##MethodName(ParameterTypeA ParameterNameA, ParameterTypeB ParameterNameB, ParameterTypeC ParameterNameC, ParameterTypeD ParameterNameD, ParameterTypeE ParameterNameE, ParameterTypeF ParameterNameF);
+	#define DEFINE_RHIMETHOD_GLOBALTHREADSAFE_7(ReturnType,MethodName,ParameterTypeA,ParameterNameA,ParameterTypeB,ParameterNameB,ParameterTypeC,ParameterNameC,ParameterTypeD,ParameterNameD,ParameterTypeE,ParameterNameE,ParameterTypeF,ParameterNameF,ParameterTypeG,ParameterNameG,ReturnStatement,NullImplementation) \
+		ReturnType RHI##MethodName(ParameterTypeA ParameterNameA, ParameterTypeB ParameterNameB, ParameterTypeC ParameterNameC, ParameterTypeD ParameterNameD, ParameterTypeE ParameterNameE, ParameterTypeF ParameterNameF, ParameterTypeG ParameterNameG);
+	#define DEFINE_RHIMETHOD_GLOBALTHREADSAFE_8(ReturnType,MethodName,ParameterTypeA,ParameterNameA,ParameterTypeB,ParameterNameB,ParameterTypeC,ParameterNameC,ParameterTypeD,ParameterNameD,ParameterTypeE,ParameterNameE,ParameterTypeF,ParameterNameF,ParameterTypeG,ParameterNameG,ParameterTypeH,ParameterNameH,ReturnStatement,NullImplementation) \
+		ReturnType RHI##MethodName(ParameterTypeA ParameterNameA, ParameterTypeB ParameterNameB, ParameterTypeC ParameterNameC, ParameterTypeD ParameterNameD, ParameterTypeE ParameterNameE, ParameterTypeF ParameterNameF, ParameterTypeG ParameterNameG, ParameterTypeH ParameterNameH);
+	#define DEFINE_RHIMETHOD_GLOBALTHREADSAFE_9(ReturnType,MethodName,ParameterTypeA,ParameterNameA,ParameterTypeB,ParameterNameB,ParameterTypeC,ParameterNameC,ParameterTypeD,ParameterNameD,ParameterTypeE,ParameterNameE,ParameterTypeF,ParameterNameF,ParameterTypeG,ParameterNameG,ParameterTypeH,ParameterNameH,ParameterTypeI,ParameterNameI,ReturnStatement,NullImplementation) \
+		ReturnType RHI##MethodName(ParameterTypeA ParameterNameA, ParameterTypeB ParameterNameB, ParameterTypeC ParameterNameC, ParameterTypeD ParameterNameD, ParameterTypeE ParameterNameE, ParameterTypeF ParameterNameF, ParameterTypeG ParameterNameG, ParameterTypeH ParameterNameH, ParameterTypeI ParameterNameI);
+
 #else
 
 	#define DEFINE_RHIMETHOD_0(ReturnType,MethodName,ReturnStatement,NullImplementation) \
@@ -294,6 +325,50 @@
 			DEFINE_RHIMETHOD(ReturnType, MethodName, (ParameterTypeA ParameterNameA, ParameterTypeB ParameterNameB, ParameterTypeC ParameterNameC, ParameterTypeD ParameterNameD, ParameterTypeE ParameterNameE, ParameterTypeF ParameterNameF, ParameterTypeG ParameterNameG, ParameterTypeH ParameterNameH, ParameterTypeI ParameterNameI), (ParameterNameA, ParameterNameB, ParameterNameC, ParameterNameD, ParameterNameE, ParameterNameF, ParameterNameG, ParameterNameH, ParameterNameI), ReturnStatement, NullImplementation)
 	#endif
 
+#if defined(DEFINE_RHIMETHOD_GLOBALTHREADSAFE)
+	#define DEFINE_RHIMETHOD_GLOBALTHREADSAFE_0(ReturnType,MethodName,ReturnStatement,NullImplementation) \
+		DEFINE_RHIMETHOD_GLOBALTHREADSAFE(ReturnType, MethodName, (), (), ReturnStatement, NullImplementation)
+	#define DEFINE_RHIMETHOD_GLOBALTHREADSAFE_1(ReturnType,MethodName,ParameterTypeA,ParameterNameA,ReturnStatement,NullImplementation) \
+		DEFINE_RHIMETHOD_GLOBALTHREADSAFE(ReturnType, MethodName, (ParameterTypeA ParameterNameA), (ParameterNameA), ReturnStatement, NullImplementation)
+	#define DEFINE_RHIMETHOD_GLOBALTHREADSAFE_2(ReturnType,MethodName,ParameterTypeA,ParameterNameA,ParameterTypeB,ParameterNameB,ReturnStatement,NullImplementation) \
+		DEFINE_RHIMETHOD_GLOBALTHREADSAFE(ReturnType, MethodName, (ParameterTypeA ParameterNameA, ParameterTypeB ParameterNameB), (ParameterNameA, ParameterNameB), ReturnStatement, NullImplementation)
+	#define DEFINE_RHIMETHOD_GLOBALTHREADSAFE_3(ReturnType,MethodName,ParameterTypeA,ParameterNameA,ParameterTypeB,ParameterNameB,ParameterTypeC,ParameterNameC,ReturnStatement,NullImplementation) \
+		DEFINE_RHIMETHOD_GLOBALTHREADSAFE(ReturnType, MethodName, (ParameterTypeA ParameterNameA, ParameterTypeB ParameterNameB, ParameterTypeC ParameterNameC), (ParameterNameA, ParameterNameB, ParameterNameC), ReturnStatement, NullImplementation)
+	#define DEFINE_RHIMETHOD_GLOBALTHREADSAFE_4(ReturnType,MethodName,ParameterTypeA,ParameterNameA,ParameterTypeB,ParameterNameB,ParameterTypeC,ParameterNameC,ParameterTypeD,ParameterNameD,ReturnStatement,NullImplementation) \
+		DEFINE_RHIMETHOD_GLOBALTHREADSAFE(ReturnType, MethodName, (ParameterTypeA ParameterNameA, ParameterTypeB ParameterNameB, ParameterTypeC ParameterNameC, ParameterTypeD ParameterNameD), (ParameterNameA, ParameterNameB, ParameterNameC, ParameterNameD), ReturnStatement, NullImplementation)
+	#define DEFINE_RHIMETHOD_GLOBALTHREADSAFE_5(ReturnType,MethodName,ParameterTypeA,ParameterNameA,ParameterTypeB,ParameterNameB,ParameterTypeC,ParameterNameC,ParameterTypeD,ParameterNameD,ParameterTypeE,ParameterNameE,ReturnStatement,NullImplementation) \
+		DEFINE_RHIMETHOD_GLOBALTHREADSAFE(ReturnType, MethodName, (ParameterTypeA ParameterNameA, ParameterTypeB ParameterNameB, ParameterTypeC ParameterNameC, ParameterTypeD ParameterNameD, ParameterTypeE ParameterNameE), (ParameterNameA, ParameterNameB, ParameterNameC, ParameterNameD, ParameterNameE), ReturnStatement, NullImplementation)
+	#define DEFINE_RHIMETHOD_GLOBALTHREADSAFE_6(ReturnType,MethodName,ParameterTypeA,ParameterNameA,ParameterTypeB,ParameterNameB,ParameterTypeC,ParameterNameC,ParameterTypeD,ParameterNameD,ParameterTypeE,ParameterNameE,ParameterTypeF,ParameterNameF,ReturnStatement,NullImplementation) \
+		DEFINE_RHIMETHOD_GLOBALTHREADSAFE(ReturnType, MethodName, (ParameterTypeA ParameterNameA, ParameterTypeB ParameterNameB, ParameterTypeC ParameterNameC, ParameterTypeD ParameterNameD, ParameterTypeE ParameterNameE, ParameterTypeF ParameterNameF), (ParameterNameA, ParameterNameB, ParameterNameC, ParameterNameD, ParameterNameE, ParameterNameF), ReturnStatement, NullImplementation)
+	#define DEFINE_RHIMETHOD_GLOBALTHREADSAFE_7(ReturnType,MethodName,ParameterTypeA,ParameterNameA,ParameterTypeB,ParameterNameB,ParameterTypeC,ParameterNameC,ParameterTypeD,ParameterNameD,ParameterTypeE,ParameterNameE,ParameterTypeF,ParameterNameF,ParameterTypeG,ParameterNameG,ReturnStatement,NullImplementation) \
+		DEFINE_RHIMETHOD_GLOBALTHREADSAFE(ReturnType, MethodName, (ParameterTypeA ParameterNameA, ParameterTypeB ParameterNameB, ParameterTypeC ParameterNameC, ParameterTypeD ParameterNameD, ParameterTypeE ParameterNameE, ParameterTypeF ParameterNameF, ParameterTypeG ParameterNameG), (ParameterNameA, ParameterNameB, ParameterNameC, ParameterNameD, ParameterNameE, ParameterNameF, ParameterNameG), ReturnStatement, NullImplementation)
+	#define DEFINE_RHIMETHOD_GLOBALTHREADSAFE_8(ReturnType,MethodName,ParameterTypeA,ParameterNameA,ParameterTypeB,ParameterNameB,ParameterTypeC,ParameterNameC,ParameterTypeD,ParameterNameD,ParameterTypeE,ParameterNameE,ParameterTypeF,ParameterNameF,ParameterTypeG,ParameterNameG,ParameterTypeH,ParameterNameH,ReturnStatement,NullImplementation) \
+		DEFINE_RHIMETHOD_GLOBALTHREADSAFE(ReturnType, MethodName, (ParameterTypeA ParameterNameA, ParameterTypeB ParameterNameB, ParameterTypeC ParameterNameC, ParameterTypeD ParameterNameD, ParameterTypeE ParameterNameE, ParameterTypeF ParameterNameF, ParameterTypeG ParameterNameG, ParameterTypeH ParameterNameH), (ParameterNameA, ParameterNameB, ParameterNameC, ParameterNameD, ParameterNameE, ParameterNameF, ParameterNameG, ParameterNameH), ReturnStatement, NullImplementation)
+	#define DEFINE_RHIMETHOD_GLOBALTHREADSAFE_9(ReturnType,MethodName,ParameterTypeA,ParameterNameA,ParameterTypeB,ParameterNameB,ParameterTypeC,ParameterNameC,ParameterTypeD,ParameterNameD,ParameterTypeE,ParameterNameE,ParameterTypeF,ParameterNameF,ParameterTypeG,ParameterNameG,ParameterTypeH,ParameterNameH,ParameterTypeI,ParameterNameI,ReturnStatement,NullImplementation) \
+		DEFINE_RHIMETHOD_GLOBALTHREADSAFE(ReturnType, MethodName, (ParameterTypeA ParameterNameA, ParameterTypeB ParameterNameB, ParameterTypeC ParameterNameC, ParameterTypeD ParameterNameD, ParameterTypeE ParameterNameE, ParameterTypeF ParameterNameF, ParameterTypeG ParameterNameG, ParameterTypeH ParameterNameH, ParameterTypeI ParameterNameI), (ParameterNameA, ParameterNameB, ParameterNameC, ParameterNameD, ParameterNameE, ParameterNameF, ParameterNameG, ParameterNameH, ParameterNameI), ReturnStatement, NullImplementation)
+#else
+	#define DEFINE_RHIMETHOD_GLOBALTHREADSAFE_0(ReturnType,MethodName,ReturnStatement,NullImplementation) \
+		DEFINE_RHIMETHOD(ReturnType, MethodName, (), (), ReturnStatement, NullImplementation)
+	#define DEFINE_RHIMETHOD_GLOBALTHREADSAFE_1(ReturnType,MethodName,ParameterTypeA,ParameterNameA,ReturnStatement,NullImplementation) \
+		DEFINE_RHIMETHOD(ReturnType, MethodName, (ParameterTypeA ParameterNameA), (ParameterNameA), ReturnStatement, NullImplementation)
+	#define DEFINE_RHIMETHOD_GLOBALTHREADSAFE_2(ReturnType,MethodName,ParameterTypeA,ParameterNameA,ParameterTypeB,ParameterNameB,ReturnStatement,NullImplementation) \
+		DEFINE_RHIMETHOD(ReturnType, MethodName, (ParameterTypeA ParameterNameA, ParameterTypeB ParameterNameB), (ParameterNameA, ParameterNameB), ReturnStatement, NullImplementation)
+	#define DEFINE_RHIMETHOD_GLOBALTHREADSAFE_3(ReturnType,MethodName,ParameterTypeA,ParameterNameA,ParameterTypeB,ParameterNameB,ParameterTypeC,ParameterNameC,ReturnStatement,NullImplementation) \
+		DEFINE_RHIMETHOD(ReturnType, MethodName, (ParameterTypeA ParameterNameA, ParameterTypeB ParameterNameB, ParameterTypeC ParameterNameC), (ParameterNameA, ParameterNameB, ParameterNameC), ReturnStatement, NullImplementation)
+	#define DEFINE_RHIMETHOD_GLOBALTHREADSAFE_4(ReturnType,MethodName,ParameterTypeA,ParameterNameA,ParameterTypeB,ParameterNameB,ParameterTypeC,ParameterNameC,ParameterTypeD,ParameterNameD,ReturnStatement,NullImplementation) \
+		DEFINE_RHIMETHOD(ReturnType, MethodName, (ParameterTypeA ParameterNameA, ParameterTypeB ParameterNameB, ParameterTypeC ParameterNameC, ParameterTypeD ParameterNameD), (ParameterNameA, ParameterNameB, ParameterNameC, ParameterNameD), ReturnStatement, NullImplementation)
+	#define DEFINE_RHIMETHOD_GLOBALTHREADSAFE_5(ReturnType,MethodName,ParameterTypeA,ParameterNameA,ParameterTypeB,ParameterNameB,ParameterTypeC,ParameterNameC,ParameterTypeD,ParameterNameD,ParameterTypeE,ParameterNameE,ReturnStatement,NullImplementation) \
+		DEFINE_RHIMETHOD(ReturnType, MethodName, (ParameterTypeA ParameterNameA, ParameterTypeB ParameterNameB, ParameterTypeC ParameterNameC, ParameterTypeD ParameterNameD, ParameterTypeE ParameterNameE), (ParameterNameA, ParameterNameB, ParameterNameC, ParameterNameD, ParameterNameE), ReturnStatement, NullImplementation)
+	#define DEFINE_RHIMETHOD_GLOBALTHREADSAFE_6(ReturnType,MethodName,ParameterTypeA,ParameterNameA,ParameterTypeB,ParameterNameB,ParameterTypeC,ParameterNameC,ParameterTypeD,ParameterNameD,ParameterTypeE,ParameterNameE,ParameterTypeF,ParameterNameF,ReturnStatement,NullImplementation) \
+		DEFINE_RHIMETHOD(ReturnType, MethodName, (ParameterTypeA ParameterNameA, ParameterTypeB ParameterNameB, ParameterTypeC ParameterNameC, ParameterTypeD ParameterNameD, ParameterTypeE ParameterNameE, ParameterTypeF ParameterNameF), (ParameterNameA, ParameterNameB, ParameterNameC, ParameterNameD, ParameterNameE, ParameterNameF), ReturnStatement, NullImplementation)
+	#define DEFINE_RHIMETHOD_GLOBALTHREADSAFE_7(ReturnType,MethodName,ParameterTypeA,ParameterNameA,ParameterTypeB,ParameterNameB,ParameterTypeC,ParameterNameC,ParameterTypeD,ParameterNameD,ParameterTypeE,ParameterNameE,ParameterTypeF,ParameterNameF,ParameterTypeG,ParameterNameG,ReturnStatement,NullImplementation) \
+		DEFINE_RHIMETHOD(ReturnType, MethodName, (ParameterTypeA ParameterNameA, ParameterTypeB ParameterNameB, ParameterTypeC ParameterNameC, ParameterTypeD ParameterNameD, ParameterTypeE ParameterNameE, ParameterTypeF ParameterNameF, ParameterTypeG ParameterNameG), (ParameterNameA, ParameterNameB, ParameterNameC, ParameterNameD, ParameterNameE, ParameterNameF, ParameterNameG), ReturnStatement, NullImplementation)
+	#define DEFINE_RHIMETHOD_GLOBALTHREADSAFE_8(ReturnType,MethodName,ParameterTypeA,ParameterNameA,ParameterTypeB,ParameterNameB,ParameterTypeC,ParameterNameC,ParameterTypeD,ParameterNameD,ParameterTypeE,ParameterNameE,ParameterTypeF,ParameterNameF,ParameterTypeG,ParameterNameG,ParameterTypeH,ParameterNameH,ReturnStatement,NullImplementation) \
+		DEFINE_RHIMETHOD(ReturnType, MethodName, (ParameterTypeA ParameterNameA, ParameterTypeB ParameterNameB, ParameterTypeC ParameterNameC, ParameterTypeD ParameterNameD, ParameterTypeE ParameterNameE, ParameterTypeF ParameterNameF, ParameterTypeG ParameterNameG, ParameterTypeH ParameterNameH), (ParameterNameA, ParameterNameB, ParameterNameC, ParameterNameD, ParameterNameE, ParameterNameF, ParameterNameG, ParameterNameH), ReturnStatement, NullImplementation)
+	#define DEFINE_RHIMETHOD_GLOBALTHREADSAFE_9(ReturnType,MethodName,ParameterTypeA,ParameterNameA,ParameterTypeB,ParameterNameB,ParameterTypeC,ParameterNameC,ParameterTypeD,ParameterNameD,ParameterTypeE,ParameterNameE,ParameterTypeF,ParameterNameF,ParameterTypeG,ParameterNameG,ParameterTypeH,ParameterNameH,ParameterTypeI,ParameterNameI,ReturnStatement,NullImplementation) \
+		DEFINE_RHIMETHOD(ReturnType, MethodName, (ParameterTypeA ParameterNameA, ParameterTypeB ParameterNameB, ParameterTypeC ParameterNameC, ParameterTypeD ParameterNameD, ParameterTypeE ParameterNameE, ParameterTypeF ParameterNameF, ParameterTypeG ParameterNameG, ParameterTypeH ParameterNameH, ParameterTypeI ParameterNameI), (ParameterNameA, ParameterNameB, ParameterNameC, ParameterNameD, ParameterNameE, ParameterNameF, ParameterNameG, ParameterNameH, ParameterNameI), ReturnStatement, NullImplementation)
+#endif
+
 #endif
 
 //
@@ -319,22 +394,22 @@ DEFINE_RHIMETHOD_2(
 // RHI resource management functions.
 //
 
-DEFINE_RHIMETHOD_GLOBAL_1(
+DEFINE_RHIMETHOD_GLOBALTHREADSAFE_1(
 	FSamplerStateRHIRef,CreateSamplerState,
 	const FSamplerStateInitializerRHI&,Initializer,
 	return,return new FRHISamplerState();
 	);
-DEFINE_RHIMETHOD_GLOBAL_1(
+DEFINE_RHIMETHOD_GLOBALTHREADSAFE_1(
 	FRasterizerStateRHIRef,CreateRasterizerState,
 	const FRasterizerStateInitializerRHI&,Initializer,
 	return,return new FRHIRasterizerState();
 	);
-DEFINE_RHIMETHOD_GLOBAL_1(
+DEFINE_RHIMETHOD_GLOBALTHREADSAFE_1(
 	FDepthStencilStateRHIRef,CreateDepthStencilState,
 	const FDepthStencilStateInitializerRHI&,Initializer,
 	return,return new FRHIDepthStencilState();
 	);
-DEFINE_RHIMETHOD_GLOBAL_1(
+DEFINE_RHIMETHOD_GLOBALTHREADSAFE_1(
 	FBlendStateRHIRef,CreateBlendState,
 	const FBlendStateInitializerRHI&,Initializer,
 	return,return new FRHIBlendState();
@@ -465,13 +540,23 @@ DEFINE_RHIMETHOD_2(
  * @param NumBytes - The number of bytes the uniform buffer should contain.
  * @return The new uniform buffer.
  */
-DEFINE_RHIMETHOD_GLOBAL_3(
-	FUniformBufferRHIRef,CreateUniformBuffer,
-	const void*,Contents,
-	const FRHIUniformBufferLayout&,Layout,
-	EUniformBufferUsage,Usage,
-	return,return new FRHIUniformBuffer(Layout);
+#if PLATFORM_SUPPORTS_RHI_THREAD
+	DEFINE_RHIMETHOD_GLOBALTHREADSAFE_3(
+		FUniformBufferRHIRef,CreateUniformBuffer,
+		const void*,Contents,
+		const FRHIUniformBufferLayout&,Layout,
+		EUniformBufferUsage,Usage,
+		return,return new FRHIUniformBuffer(Layout);
 	);
+#else
+	DEFINE_RHIMETHOD_GLOBAL_3(
+		FUniformBufferRHIRef,CreateUniformBuffer,
+		const void*,Contents,
+		const FRHIUniformBufferLayout&,Layout,
+		EUniformBufferUsage,Usage,
+		return,return new FRHIUniformBuffer(Layout);
+	);
+#endif
 
 
 DEFINE_RHIMETHOD_GLOBAL_4(
@@ -615,7 +700,7 @@ DEFINE_RHIMETHOD_GLOBAL_3(
  * @param Flags - ETextureCreateFlags creation flags
  * @param OutAlign - Alignment required for this texture.  Output parameter.
  */
-DEFINE_RHIMETHOD_GLOBAL_7(
+DEFINE_RHIMETHOD_GLOBALTHREADSAFE_7(
 	uint64,CalcTexture2DPlatformSize,
 	uint32, SizeX,
 	uint32, SizeY,	
@@ -636,7 +721,7 @@ DEFINE_RHIMETHOD_GLOBAL_7(
 * @param Flags - ETextureCreateFlags creation flags
 * @param OutAlign - Alignment required for this texture.  Output parameter.
 */
-DEFINE_RHIMETHOD_GLOBAL_7(
+DEFINE_RHIMETHOD_GLOBALTHREADSAFE_7(
 	uint64,CalcTexture3DPlatformSize,
 	uint32,SizeX,
 	uint32,SizeY,
@@ -655,7 +740,7 @@ DEFINE_RHIMETHOD_GLOBAL_7(
 * @param Flags - ETextureCreateFlags creation flags
 * @param OutAlign - Alignment required for this texture.  Output parameter.
 */
-DEFINE_RHIMETHOD_GLOBAL_5(
+DEFINE_RHIMETHOD_GLOBALTHREADSAFE_5(
 	uint64,CalcTextureCubePlatformSize,
 	uint32,Size,
 	uint8,Format,
@@ -668,7 +753,7 @@ DEFINE_RHIMETHOD_GLOBAL_5(
  * Retrieves texture memory stats.
  * safe to call on the main thread
  */
-DEFINE_RHIMETHOD_GLOBAL_1(
+DEFINE_RHIMETHOD_GLOBALTHREADSAFE_1(
 	void,GetTextureMemoryStats,
 	FTextureMemoryStats&,OutStats,
 	,
@@ -813,7 +898,7 @@ DEFINE_RHIMETHOD_GLOBAL_7(
 /**
  * @param Ref may be 0
  */
-DEFINE_RHIMETHOD_GLOBAL_2(
+DEFINE_RHIMETHOD_GLOBALTHREADSAFE_2(
 	void,GetResourceInfo,
 	FTextureRHIParamRef,Ref,
 	FRHIResourceInfo&,OutInfo,
@@ -859,7 +944,7 @@ DEFINE_RHIMETHOD_1(
  * @param	TextureRHI		- Texture we want to know the size of, 0 is safely ignored
  * @return					- Size in Bytes
  */
-DEFINE_RHIMETHOD_GLOBAL_1(
+DEFINE_RHIMETHOD_GLOBALTHREADSAFE_1(
 	uint32,ComputeMemorySize,
 	FTextureRHIParamRef,TextureRHI,
 	return,return 0;
@@ -1120,7 +1205,7 @@ DEFINE_RHIMETHOD_CMDLIST_4(
 	,
 	);
 
-DEFINE_RHIMETHOD_GLOBAL_2(
+DEFINE_RHIMETHOD_GLOBALTHREADSAFE_2(
 	void,BindDebugLabelName,
 	FTextureRHIParamRef, Texture,
 	const TCHAR*, Name,
@@ -1184,45 +1269,138 @@ DEFINE_RHIMETHOD_GLOBAL_1(
 	ERenderQueryType, QueryType,
 	return,return new FRHIRenderQuery();
 	);
-DEFINE_RHIMETHOD_1(
+
+DEFINE_RHIMETHOD_CMDLIST_1(
 	void,ResetRenderQuery,
 	FRenderQueryRHIParamRef,RenderQuery,
 	,
 	);
-DEFINE_RHIMETHOD_3(
-	bool,GetRenderQueryResult,
-	FRenderQueryRHIParamRef,RenderQuery,
-	uint64&,OutResult,
-	bool,bWait,
-	return,return true;
-	);
-DEFINE_RHIMETHOD_1(
+
+DEFINE_RHIMETHOD_CMDLIST_1(
 	void,BeginRenderQuery,
 	FRenderQueryRHIParamRef,RenderQuery,
 	,
 	);
-DEFINE_RHIMETHOD_1(
+DEFINE_RHIMETHOD_CMDLIST_1(
 	void,EndRenderQuery,
 	FRenderQueryRHIParamRef,RenderQuery,
 	,
 	);
 
-
-DEFINE_RHIMETHOD_2(
-	void,BeginDrawingViewport,
-	FViewportRHIParamRef,Viewport,
-	FTextureRHIParamRef,RenderTargetRHI,
-	,
+#if PLATFORM_SUPPORTS_RHI_THREAD
+	DEFINE_RHIMETHOD_GLOBALTHREADSAFE_3(
+		bool,GetRenderQueryResult,
+		FRenderQueryRHIParamRef,RenderQuery,
+		uint64&,OutResult,
+		bool,bWait,
+		return,return true;
 	);
-DEFINE_RHIMETHOD_3(
-	void,EndDrawingViewport,
-	FViewportRHIParamRef,Viewport,
-	bool,bPresent,
-	bool,bLockToVsync,
-	,
+	DEFINE_RHIMETHOD_CMDLIST_2(
+		void,BeginDrawingViewport,
+		FViewportRHIParamRef,Viewport,
+		FTextureRHIParamRef,RenderTargetRHI,
+		,
+		);
+	DEFINE_RHIMETHOD_CMDLIST_3(
+		void,EndDrawingViewport,
+		FViewportRHIParamRef,Viewport,
+		bool,bPresent,
+		bool,bLockToVsync,
+		,
+		);
+	DEFINE_RHIMETHOD_GLOBALTHREADSAFE_1(
+		FTexture2DRHIRef,GetViewportBackBuffer,
+		FViewportRHIParamRef,Viewport,
+		return,return new FRHITexture2D(1,1,1,1,PF_B8G8R8A8,TexCreate_RenderTargetable);
+	);
+	DEFINE_RHIMETHOD_CMDLIST_0(
+		void,BeginFrame,
+		return,return;
+	);
+
+	DEFINE_RHIMETHOD_CMDLIST_0(
+		void,EndFrame,
+		return,return;
+	);
+	/**
+	 * Signals the beginning of scene rendering. The RHI makes certain caching assumptions between
+	 * calls to BeginScene/EndScene. Currently the only restriction is that you can't update texture
+	 * references.
+	 */
+	DEFINE_RHIMETHOD_CMDLIST_0(
+		void,BeginScene,
+		return,return;
+		);
+
+	/**
+	 * Signals the end of scene rendering. See RHIBeginScene.
+	 */
+	DEFINE_RHIMETHOD_CMDLIST_0(
+		void,EndScene,
+		return,return;
+		);
+
+	DEFINE_RHIMETHOD_GLOBALTHREADSAFE_0(
+		void,AdvanceFrameForGetViewportBackBuffer,
+		return,return;
+		);
+
+#else
+	DEFINE_RHIMETHOD_3(
+		bool,GetRenderQueryResult,
+		FRenderQueryRHIParamRef,RenderQuery,
+		uint64&,OutResult,
+		bool,bWait,
+		return,return true;
+	);
+	DEFINE_RHIMETHOD_2(
+		void,BeginDrawingViewport,
+		FViewportRHIParamRef,Viewport,
+		FTextureRHIParamRef,RenderTargetRHI,
+		,
+		);
+	DEFINE_RHIMETHOD_3(
+		void,EndDrawingViewport,
+		FViewportRHIParamRef,Viewport,
+		bool,bPresent,
+		bool,bLockToVsync,
+		,
+		);
+	DEFINE_RHIMETHOD_GLOBALFLUSH_1(
+		FTexture2DRHIRef,GetViewportBackBuffer,
+		FViewportRHIParamRef,Viewport,
+		return,return new FRHITexture2D(1,1,1,1,PF_B8G8R8A8,TexCreate_RenderTargetable);
+	);
+	DEFINE_RHIMETHOD_0(
+		void,BeginFrame,
+		return,return;
+	);
+	DEFINE_RHIMETHOD_0(
+		void,EndFrame,
+		return,return;
 	);
+	/**
+	 * Signals the beginning of scene rendering. The RHI makes certain caching assumptions between
+	 * calls to BeginScene/EndScene. Currently the only restriction is that you can't update texture
+	 * references.
+	 */
+	DEFINE_RHIMETHOD_0(
+		void,BeginScene,
+		return,return;
+		);
+
+	/**
+	 * Signals the end of scene rendering. See RHIBeginScene.
+	 */
+	DEFINE_RHIMETHOD_0(
+		void,EndScene,
+		return,return;
+		);
+#endif
+
 /**
  * Determine if currently drawing the viewport
+ * REMOVE THIS If you are updating the RHI for all platforms
  *
  * @return true if currently within a BeginDrawingViewport/EndDrawingViewport block
  */
@@ -1230,39 +1408,6 @@ DEFINE_RHIMETHOD_0(
 	bool,IsDrawingViewport,
 	return,return false;
 	);
-DEFINE_RHIMETHOD_GLOBALFLUSH_1(
-	FTexture2DRHIRef,GetViewportBackBuffer,
-	FViewportRHIParamRef,Viewport,
-	return,return new FRHITexture2D(1,1,1,1,PF_B8G8R8A8,TexCreate_RenderTargetable);
-	);
-
-DEFINE_RHIMETHOD_0(
-	void,BeginFrame,
-	return,return;
-	);
-DEFINE_RHIMETHOD_0(
-	void,EndFrame,
-	return,return;
-	);
-
-/**
- * Signals the beginning of scene rendering. The RHI makes certain caching assumptions between
- * calls to BeginScene/EndScene. Currently the only restriction is that you can't update texture
- * references.
- */
-DEFINE_RHIMETHOD_0(
-	void,BeginScene,
-	return,return;
-	);
-
-/**
- * Signals the end of scene rendering. See RHIBeginScene.
- */
-DEFINE_RHIMETHOD_0(
-	void,EndScene,
-	return,return;
-	);
-
 /*
  * Acquires or releases ownership of the platform-specific rendering context for the calling thread
  */
@@ -1283,7 +1428,7 @@ DEFINE_RHIMETHOD_GLOBALFLUSH_0(
 /*
  * Returns the total GPU time taken to render the last frame. Same metric as FPlatformTime::Cycles().
  */
-DEFINE_RHIMETHOD_GLOBAL_0(
+DEFINE_RHIMETHOD_GLOBALTHREADSAFE_0(
 	uint32,GetGPUFrameCycles,
 	return,return 0;
 	);
@@ -1291,7 +1436,7 @@ DEFINE_RHIMETHOD_GLOBAL_0(
 /**
  * The following RHI functions must be called from the main thread.
  */
-DEFINE_RHIMETHOD_GLOBAL_4(
+DEFINE_RHIMETHOD_GLOBALTHREADSAFE_4(
 	FViewportRHIRef,CreateViewport,
 	void*,WindowHandle,
 	uint32,SizeX,
@@ -1299,7 +1444,7 @@ DEFINE_RHIMETHOD_GLOBAL_4(
 	bool,bIsFullscreen,
 	return,return new FRHIViewport();
 	);
-DEFINE_RHIMETHOD_GLOBAL_4(
+DEFINE_RHIMETHOD_GLOBALTHREADSAFE_4(
 	void,ResizeViewport,
 	FViewportRHIParamRef,Viewport,
 	uint32,SizeX,
@@ -1307,7 +1452,7 @@ DEFINE_RHIMETHOD_GLOBAL_4(
 	bool,bIsFullscreen,
 	,
 	);
-DEFINE_RHIMETHOD_GLOBAL_1(
+DEFINE_RHIMETHOD_GLOBALTHREADSAFE_1(
 	void,Tick,
 	float,DeltaTime,
 	,
@@ -1899,7 +2044,7 @@ DEFINE_RHIMETHOD_0(
  *
  *	@return	bool				true if successfully filled the array
  */
-DEFINE_RHIMETHOD_GLOBAL_2(
+DEFINE_RHIMETHOD_GLOBALTHREADSAFE_2(
 	bool,GetAvailableResolutions,
 	FScreenResolutionArray&,Resolutions,
 	bool,bIgnoreRefreshRate,
@@ -1911,7 +2056,7 @@ DEFINE_RHIMETHOD_GLOBAL_2(
  * @param Width - Input: Desired resolution width in pixels. Output: A width that the platform supports.
  * @param Height - Input: Desired resolution height in pixels. Output: A height that the platform supports.
  */
-DEFINE_RHIMETHOD_GLOBAL_2(
+DEFINE_RHIMETHOD_GLOBALTHREADSAFE_2(
 	void,GetSupportedResolution,
 	uint32&,Width,
 	uint32&,Height,
@@ -1967,7 +2112,7 @@ DEFINE_RHIMETHOD_CMDLIST_3(
 /**
  * Provides access to the native device. Generally this should be avoided but is useful for third party plugins.
  */
-DEFINE_RHIMETHOD_GLOBAL_0(
+DEFINE_RHIMETHOD_GLOBALFLUSH_0(
 	void*,GetNativeDevice,
 	return,return NULL
 	);
diff --git a/Engine/Source/Runtime/RHI/Public/RHIStaticStates.h b/Engine/Source/Runtime/RHI/Public/RHIStaticStates.h
--- a/Engine/Source/Runtime/RHI/Public/RHIStaticStates.h
+++ b/Engine/Source/Runtime/RHI/Public/RHIStaticStates.h
@@ -64,6 +64,14 @@ public:
 
 		if (!StaticResource)
 		{
+#if PLATFORM_SUPPORTS_RHI_THREAD
+			FStaticStateResource* NewStaticResource = new FStaticStateResource();
+			FStaticStateResource* ValueWas = (FStaticStateResource*)FPlatformAtomics::InterlockedCompareExchangePointer((void**)&StaticResource, NewStaticResource, nullptr);
+			if (ValueWas)
+			{
+				// we made a redundant one...leak it
+			}
+#else
 			if (!IsInRenderingThread())
 			{
 				check(IsInParallelRenderingThread());
@@ -76,6 +84,7 @@ public:
 			{
 				StaticResource = new FStaticStateResource();
 			}
+#endif
 		}
 		return StaticResource->StateRHI;
 	};
@@ -87,7 +96,27 @@ private:
 	{
 	public:
 		RHIRefType StateRHI;
+#if PLATFORM_SUPPORTS_RHI_THREAD
+		FStaticStateResource()
+		{
+			StateRHI = InitializerType::CreateRHI();
+		}
 
+		// FRenderResource interface.
+		virtual void InitRHI() override
+		{
+			check(0);
+		}
+		virtual void ReleaseRHI() override
+		{
+			check(0);
+		}
+
+		~FStaticStateResource()
+		{
+			check(0);
+		}
+#else
 		FStaticStateResource()
 		{
 			InitResource();
@@ -107,6 +136,7 @@ private:
 		{
 			ReleaseResource();
 		}
+#endif
 	};
 };
 
diff --git a/Engine/Source/Runtime/RenderCore/Private/RenderResource.cpp b/Engine/Source/Runtime/RenderCore/Private/RenderResource.cpp
--- a/Engine/Source/Runtime/RenderCore/Private/RenderResource.cpp
+++ b/Engine/Source/Runtime/RenderCore/Private/RenderResource.cpp
@@ -220,6 +220,8 @@ public:
 	enum { ALIGNMENT = (1 << 16) }; // 64KB
 	/** Pointer to the vertex buffer mapped in main memory. */
 	uint8* MappedBuffer;
+	/** Pointer to a buffer to map and copy later. */
+	uint8* DeferredBuffer;
 	/** Size of the vertex buffer in bytes. */
 	uint32 BufferSize;
 	/** Number of bytes currently allocated from the buffer. */
@@ -228,18 +230,42 @@ public:
 	/** Default constructor. */
 	explicit FDynamicVertexBuffer(uint32 InMinBufferSize)
 		: MappedBuffer(NULL)
+		, DeferredBuffer(NULL)
 		, BufferSize(FMath::Max<uint32>(Align(InMinBufferSize,ALIGNMENT),ALIGNMENT))
 		, AllocatedByteCount(0)
 	{
 	}
 
+	/**
+	 * Instead of locking, allocates a deferred memory block
+	 */
+	void AllocDeferred()
+	{
+		check(MappedBuffer == NULL && DeferredBuffer == NULL);
+		check(AllocatedByteCount == 0);
+		check(IsValidRef(VertexBufferRHI));
+		DeferredBuffer = (uint8*)FMemory::Malloc(BufferSize, 16);
+	}
+
+	/**
+	 * Lock, copy, free, unlock the deferred into place
+	 */
+	void CopyDeferred()
+	{
+		Lock();
+		FMemory::Memcpy(MappedBuffer, DeferredBuffer, AllocatedByteCount);
+		FMemory::Free(DeferredBuffer);
+		DeferredBuffer = nullptr;
+		Unlock();
+	}
+
 	/**
 	 * Locks the vertex buffer so it may be written to.
 	 */
 	void Lock()
 	{
 		check(MappedBuffer == NULL);
-		check(AllocatedByteCount == 0);
+		check(AllocatedByteCount == 0 || DeferredBuffer != NULL);
 		check(IsValidRef(VertexBufferRHI));
 		MappedBuffer = (uint8*)RHILockVertexBuffer(VertexBufferRHI, 0, BufferSize, RLM_WriteOnly);
 	}
@@ -249,7 +275,7 @@ public:
 	 */
 	void Unlock()
 	{
-		check(MappedBuffer != NULL);
+		check(MappedBuffer != NULL && DeferredBuffer == NULL);
 		check(IsValidRef(VertexBufferRHI));
 		RHIUnlockVertexBuffer(VertexBufferRHI);
 		MappedBuffer = NULL;
@@ -263,6 +289,7 @@ public:
 		FRHIResourceCreateInfo CreateInfo;
 		VertexBufferRHI = RHICreateVertexBuffer(BufferSize, BUF_Volatile, CreateInfo);
 		MappedBuffer = NULL;
+		DeferredBuffer = nullptr;
 		AllocatedByteCount = 0;
 	}
 
@@ -270,6 +297,7 @@ public:
 	{
 		FVertexBuffer::ReleaseRHI();
 		MappedBuffer = NULL;
+		DeferredBuffer = nullptr;
 		AllocatedByteCount = 0;
 	}
 
@@ -317,19 +345,21 @@ FGlobalDynamicVertexBuffer::~FGlobalDynamicVertexBuffer()
 	Pool = NULL;
 }
 
-FGlobalDynamicVertexBuffer::FAllocation FGlobalDynamicVertexBuffer::Allocate(uint32 SizeInBytes)
+FGlobalDynamicVertexBuffer::FAllocation FGlobalDynamicVertexBuffer::Allocate(uint32 SizeInBytes, bool bDeferLock)
 {
 	FAllocation Allocation;
 
 	FDynamicVertexBuffer* VertexBuffer = Pool->CurrentVertexBuffer;
-	if (VertexBuffer == NULL || VertexBuffer->AllocatedByteCount + SizeInBytes > VertexBuffer->BufferSize)
+	if (VertexBuffer == NULL || VertexBuffer->AllocatedByteCount + SizeInBytes > VertexBuffer->BufferSize
+		|| (bDeferLock && VertexBuffer->MappedBuffer) || (!bDeferLock && VertexBuffer->DeferredBuffer)) // we won't mix these styles of allocation
 	{
 		// Find a buffer in the pool big enough to service the request.
 		VertexBuffer = NULL;
 		for (int32 BufferIndex = 0, NumBuffers = Pool->VertexBuffers.Num(); BufferIndex < NumBuffers; ++BufferIndex)
 		{
 			FDynamicVertexBuffer& VertexBufferToCheck = Pool->VertexBuffers[BufferIndex];
-			if (VertexBufferToCheck.AllocatedByteCount + SizeInBytes <= VertexBufferToCheck.BufferSize)
+			bool bMismatch = (bDeferLock && VertexBufferToCheck.MappedBuffer) || (!bDeferLock && VertexBufferToCheck.DeferredBuffer);
+			if (VertexBufferToCheck.AllocatedByteCount + SizeInBytes <= VertexBufferToCheck.BufferSize && !bMismatch)
 			{
 				VertexBuffer = &VertexBufferToCheck;
 				break;
@@ -344,10 +374,16 @@ FGlobalDynamicVertexBuffer::FAllocation FGlobalDynamicVertexBuffer::Allocate(uin
 		}
 
 		// Lock the buffer if needed.
-		if (VertexBuffer->MappedBuffer == NULL)
+		if (!bDeferLock && VertexBuffer->MappedBuffer == NULL)
 		{
+			check(!VertexBuffer->DeferredBuffer);
 			VertexBuffer->Lock();
 		}
+		else if (bDeferLock && VertexBuffer->DeferredBuffer == NULL)
+		{
+			check(!VertexBuffer->MappedBuffer);
+			VertexBuffer->AllocDeferred();
+		}
 
 		// Remember this buffer, we'll try to allocate out of it in the future.
 		Pool->CurrentVertexBuffer = VertexBuffer;
@@ -355,8 +391,9 @@ FGlobalDynamicVertexBuffer::FAllocation FGlobalDynamicVertexBuffer::Allocate(uin
 
 	check(VertexBuffer != NULL);
 	checkf(VertexBuffer->AllocatedByteCount + SizeInBytes <= VertexBuffer->BufferSize, TEXT("Global vertex buffer allocation failed: BufferSize=%d AllocatedByteCount=%d SizeInBytes=%d"), VertexBuffer->BufferSize, VertexBuffer->AllocatedByteCount, SizeInBytes);
-
-	Allocation.Buffer = VertexBuffer->MappedBuffer + VertexBuffer->AllocatedByteCount;
+	uint8* Base = bDeferLock ? VertexBuffer->DeferredBuffer : VertexBuffer->MappedBuffer;
+	check(Base);
+	Allocation.Buffer = Base + VertexBuffer->AllocatedByteCount;
 	Allocation.VertexBuffer = VertexBuffer;
 	Allocation.VertexOffset = VertexBuffer->AllocatedByteCount;
 	VertexBuffer->AllocatedByteCount += SizeInBytes;
@@ -373,6 +410,10 @@ void FGlobalDynamicVertexBuffer::Commit()
 		{
 			VertexBuffer.Unlock();
 		}
+		else if (VertexBuffer.DeferredBuffer != NULL)
+		{
+			VertexBuffer.CopyDeferred();
+		}
 	}
 	Pool->CurrentVertexBuffer = NULL;
 }
diff --git a/Engine/Source/Runtime/RenderCore/Private/RenderingThread.cpp b/Engine/Source/Runtime/RenderCore/Private/RenderingThread.cpp
--- a/Engine/Source/Runtime/RenderCore/Private/RenderingThread.cpp
+++ b/Engine/Source/Runtime/RenderCore/Private/RenderingThread.cpp
@@ -20,6 +20,7 @@
 
 RENDERCORE_API bool GIsThreadedRendering = false;
 RENDERCORE_API bool GUseThreadedRendering = false;
+RENDERCORE_API bool GUseRHIThread = false;
 
 #if !(UE_BUILD_SHIPPING || UE_BUILD_TEST)
 	RENDERCORE_API bool GMainThreadBlockedOnRenderThread = false;
@@ -205,6 +206,46 @@ uint32 GRenderThreadNumIdle[ERenderThreadIdleTypes::Num] = {0};
 /** How many cycles the renderthread used (excluding idle time). It's set once per frame in FViewport::Draw. */
 uint32 GRenderThreadTime = 0;
 
+
+
+#if PLATFORM_SUPPORTS_RHI_THREAD
+/** The RHI thread runnable object. */
+class FRHIThread : public FRunnable
+{
+public:
+	FRunnableThread* Thread;
+
+	FRHIThread()
+		: Thread(nullptr)
+	{
+		check(IsInGameThread());
+	}
+
+	virtual uint32 Run()
+	{
+		FTaskGraphInterface::Get().AttachToThread(ENamedThreads::RHIThread);
+		FTaskGraphInterface::Get().ProcessThreadUntilRequestReturn(ENamedThreads::RHIThread);
+		return 0;
+	}
+
+	static FRHIThread& Get()
+	{
+		static FRHIThread Singleton;
+		return Singleton;
+	}
+
+	void Start()
+	{
+		Thread = FRunnableThread::Create(this, TEXT("RHIThread"), 512 * 1024, TPri_Normal, 
+			//FPlatformAffinity::GetRHIThreadMask()
+			0xFFFFFFFFFFFFFFFF
+			);
+		check(Thread);
+	}
+};
+
+#endif
+
 /** The rendering thread main loop */
 void RenderingThreadMain( FEvent* TaskGraphBoundSyncEvent )
 {
@@ -459,6 +500,23 @@ void StartRenderingThread()
 	static uint32 ThreadCount = 0;
 	check(!GIsThreadedRendering && GUseThreadedRendering);
 
+	check(!GRHIThread)
+#if PLATFORM_SUPPORTS_RHI_THREAD
+	if (GUseRHIThread)
+	{
+		if (!FTaskGraphInterface::Get().IsThreadProcessingTasks(ENamedThreads::RHIThread))
+		{
+			FRHIThread::Get().Start();
+		}
+		DECLARE_CYCLE_STAT(TEXT("Wait For RHIThread"), STAT_WaitForRHIThread, STATGROUP_TaskGraphTasks);
+
+		FGraphEventRef CompletionEvent = TGraphTask<FNullGraphTask>::CreateTask(NULL, ENamedThreads::GameThread).ConstructAndDispatchWhenReady(GET_STATID(STAT_WaitForRHIThread), ENamedThreads::RHIThread);
+		FTaskGraphInterface::Get().WaitUntilTaskCompletes(CompletionEvent, ENamedThreads::GameThread_Local);
+		GRHIThread = FRHIThread::Get().Thread;
+		check(GRHIThread);
+	}
+#endif
+
 	// Turn on the threaded rendering flag.
 	GIsThreadedRendering = true;
 
@@ -520,27 +578,40 @@ void StopRenderingThread()
 		// The rendering thread may have already been stopped during the call to GFlushStreamingFunc or FlushRenderingCommands.
 		if ( GIsThreadedRendering )
 		{
+#if PLATFORM_SUPPORTS_RHI_THREAD
+			if (GRHIThread)
+			{
+				DECLARE_CYCLE_STAT(TEXT("Wait For RHIThread Finish"), STAT_WaitForRHIThreadFinish, STATGROUP_TaskGraphTasks);
+				FGraphEventRef FlushTask = TGraphTask<FNullGraphTask>::CreateTask(NULL, ENamedThreads::GameThread).ConstructAndDispatchWhenReady(GET_STATID(STAT_WaitForRHIThreadFinish), ENamedThreads::RHIThread);
+				FTaskGraphInterface::Get().WaitUntilTaskCompletes(FlushTask, ENamedThreads::GameThread_Local);
+				GRHIThread = nullptr;
+			}
+#endif
+
+
 			check( GRenderingThread );
 			check(!GIsRenderingThreadSuspended);
 
 			// Turn off the threaded rendering flag.
 			GIsThreadedRendering = false;
 
-			FGraphEventRef QuitTask = TGraphTask<FReturnGraphTask>::CreateTask(NULL, ENamedThreads::GameThread).ConstructAndDispatchWhenReady(ENamedThreads::RenderThread);
-
-			// Busy wait while BP debugging, to avoid opportunistic execution of game thread tasks
-			// If the game thread is already executing tasks, then we have no choice but to spin
-			if (GIntraFrameDebuggingGameThread || FTaskGraphInterface::Get().IsThreadProcessingTasks(ENamedThreads::GameThread) ) 
 			{
-				while ((QuitTask.GetReference() != nullptr) && !QuitTask->IsComplete())
+				FGraphEventRef QuitTask = TGraphTask<FReturnGraphTask>::CreateTask(NULL, ENamedThreads::GameThread).ConstructAndDispatchWhenReady(ENamedThreads::RenderThread);
+
+				// Busy wait while BP debugging, to avoid opportunistic execution of game thread tasks
+				// If the game thread is already executing tasks, then we have no choice but to spin
+				if (GIntraFrameDebuggingGameThread || FTaskGraphInterface::Get().IsThreadProcessingTasks(ENamedThreads::GameThread) ) 
+				{
+					while ((QuitTask.GetReference() != nullptr) && !QuitTask->IsComplete())
+					{
+						FPlatformProcess::Sleep(0.0f);
+					}
+				}
+				else
 				{
-					FPlatformProcess::Sleep(0.0f);
+					FTaskGraphInterface::Get().WaitUntilTaskCompletes(QuitTask, ENamedThreads::GameThread_Local);
 				}
 			}
-			else
-			{
-				FTaskGraphInterface::Get().WaitUntilTaskCompletes(QuitTask, ENamedThreads::GameThread_Local);
-			}
 
 			// Wait for the rendering thread to return.
 			GRenderingThread->WaitForCompletion();
@@ -556,6 +627,7 @@ void StopRenderingThread()
 		// Delete the pending cleanup objects which were in use by the rendering thread.
 		delete PendingCleanupObjects;
 	}
+	check(!GRHIThread);
 }
 
 void CheckRenderingThreadHealth()
@@ -779,16 +851,22 @@ void FlushRenderingCommands()
 
 void FlushPendingDeleteRHIResources_GameThread()
 {
-	ENQUEUE_UNIQUE_RENDER_COMMAND(
-		FlushPendingDeleteRHIResources,
+	if (!GRHIThread)
+	{
+		ENQUEUE_UNIQUE_RENDER_COMMAND(
+			FlushPendingDeleteRHIResources,
 		{
 			FlushPendingDeleteRHIResources_RenderThread();
 		}
-	);
+		);
+	}
 }
 void FlushPendingDeleteRHIResources_RenderThread()
 {
-	FRHIResource::FlushPendingDeletes();
+	if (!GRHIThread)
+	{
+		FRHIResource::FlushPendingDeletes();
+	}
 }
 
 
diff --git a/Engine/Source/Runtime/RenderCore/Public/RenderResource.h b/Engine/Source/Runtime/RenderCore/Public/RenderResource.h
--- a/Engine/Source/Runtime/RenderCore/Public/RenderResource.h
+++ b/Engine/Source/Runtime/RenderCore/Public/RenderResource.h
@@ -508,7 +508,7 @@ public:
 	 * @param SizeInBytes - The amount of memory to allocate in bytes.
 	 * @returns An FAllocation with information regarding the allocated memory.
 	 */
-	FAllocation Allocate(uint32 SizeInBytes);
+	FAllocation Allocate(uint32 SizeInBytes, bool bDeferLock = false);
 
 	/**
 	 * Commits allocated memory to the GPU.
diff --git a/Engine/Source/Runtime/RenderCore/Public/RenderingThread.h b/Engine/Source/Runtime/RenderCore/Public/RenderingThread.h
--- a/Engine/Source/Runtime/RenderCore/Public/RenderingThread.h
+++ b/Engine/Source/Runtime/RenderCore/Public/RenderingThread.h
@@ -26,6 +26,11 @@ extern RENDERCORE_API bool GIsThreadedRendering;
  * Currently set by command line parameter and by the ToggleRenderingThread console command.
  */
 extern RENDERCORE_API bool GUseThreadedRendering;
+/**
+ * Whether the RHI thread should be created or not, requires rendering thread
+ * Currently set by command line parameter and by the ToggleRHIThread console command.
+ */
+extern RENDERCORE_API bool GUseRHIThread;
 
 #if (UE_BUILD_SHIPPING || UE_BUILD_TEST)
 	static FORCEINLINE void CheckNotBlockedOnRenderThread() {}
diff --git a/Engine/Source/Runtime/Renderer/Private/DeferredShadingRenderer.cpp b/Engine/Source/Runtime/Renderer/Private/DeferredShadingRenderer.cpp
--- a/Engine/Source/Runtime/Renderer/Private/DeferredShadingRenderer.cpp
+++ b/Engine/Source/Runtime/Renderer/Private/DeferredShadingRenderer.cpp
@@ -479,6 +479,7 @@ public:
 	void DoTask(ENamedThreads::Type CurrentThread, const FGraphEventRef& MyCompletionGraphEvent)
 	{
 		ThisRenderer.RenderBasePassDynamicData(RHICmdList, View, OutDirty);
+		RHICmdList.HandleRTThreadTaskCompletion(MyCompletionGraphEvent);
 	}
 };
 
@@ -599,15 +600,22 @@ void FDeferredShadingSceneRenderer::RenderFinish(FRHICommandListImmediate& RHICm
 /** 
 * Renders the view family. 
 */
+
 void FDeferredShadingSceneRenderer::Render(FRHICommandListImmediate& RHICmdList)
 {
+	static FGraphEventRef OcclusionSubmittedFence;
+	if (GRHIThread)
+	{
+		QUICK_SCOPE_CYCLE_COUNTER(STAT_OcclusionSubmittedFence_Wait);
+		FRHICommandListExecutor::WaitOnRHIThreadFence(OcclusionSubmittedFence);
+	}
+
 	if(!ViewFamily.EngineShowFlags.Rendering)
 	{
 		return;
 	}
-
 	SCOPED_DRAW_EVENT(Scene,DEC_SCENE_ITEMS);
-
+	
 	// Initialize global system textures (pass-through if already initialized).
 	GSystemTextures.InitializeTextures(RHICmdList, FeatureLevel);
 
@@ -617,6 +625,16 @@ void FDeferredShadingSceneRenderer::Render(FRHICommandListImmediate& RHICmdList)
 	// Find the visible primitives.
 	InitViews(RHICmdList);
 
+	if (GRHIThread)
+	{
+		// we will probably stall on occlusion queries, so might as well have the RHI thread and GPU work while we wait.
+		//QUICK_SCOPE_CYCLE_COUNTER(STAT_OcclusionCull_Dispatch);
+		//FRHICommandListExecutor::GetImmediateCommandList().ImmediateFlush(EImmediateFlushType::DispatchToRHIThread); // we want to make sure this all gets to the GPU this frame and doesn't hang around
+		QUICK_SCOPE_CYCLE_COUNTER(STAT_PostInitViews_FlushDel);
+		FRHICommandListExecutor::GetImmediateCommandList().ImmediateFlush(EImmediateFlushType::FlushRHIThreadFlushResources);
+	}
+
+
 	const bool bIsWireframe = ViewFamily.EngineShowFlags.Wireframe;
 	static const auto ClearMethodCVar = IConsoleManager::Get().FindTConsoleVariableDataInt(TEXT("r.ClearSceneMethod"));
 	bool bRequiresRHIClear = true;
@@ -791,8 +809,15 @@ void FDeferredShadingSceneRenderer::Render(FRHICommandListImmediate& RHICmdList)
 	// This is done after the downsampled depth buffer is created so that it can be used for issuing queries
 	if ( bIsOcclusionTesting )
 	{
+		//FRHICommandListExecutor::GetImmediateCommandList().ImmediateFlush(EImmediateFlushType::FlushRHIThread);
 		BeginOcclusionTests(RHICmdList);
 	}
+
+	if (GRHIThread)
+	{
+		QUICK_SCOPE_CYCLE_COUNTER(STAT_OcclusionSubmittedFence_Dispatch);
+		OcclusionSubmittedFence = FRHICommandListExecutor::RHIThreadFence();
+	}
 	
 	// Render lighting.
 	if (ViewFamily.EngineShowFlags.Lighting
@@ -1131,6 +1156,7 @@ public:
 		{
 			OutDirty = true;
 		}
+		RHICmdList.HandleRTThreadTaskCompletion(MyCompletionGraphEvent);
 	}
 };
 
@@ -1190,7 +1216,7 @@ bool FDeferredShadingSceneRenderer::RenderPrePass(FRHICommandListImmediate& RHIC
 	{
 		if (GRHICommandList.UseParallelAlgorithms())
 		{
-			FScopedCommandListFlush Flusher(RHICmdList);
+			FScopedCommandListWaitForTasks Flusher(RHICmdList);
 			int32 Width = CVarRHICmdWidth.GetValueOnRenderThread(); // we use a few more than needed to cover non-equal jobs
 
 			for(int32 ViewIndex = 0;ViewIndex < Views.Num();ViewIndex++)
@@ -1236,7 +1262,7 @@ bool FDeferredShadingSceneRenderer::RenderBasePass(FRHICommandListImmediate& RHI
 
 		if (GRHICommandList.UseParallelAlgorithms())
 		{
-			FScopedCommandListFlush Flusher(RHICmdList);
+			FScopedCommandListWaitForTasks Flusher(RHICmdList);
 			int32 Width = CVarRHICmdWidth.GetValueOnRenderThread(); // we use a few more than needed to cover non-equal jobs
 
 			for(int32 ViewIndex = 0;ViewIndex < Views.Num();ViewIndex++)
diff --git a/Engine/Source/Runtime/Renderer/Private/DynamicPrimitiveDrawing.inl b/Engine/Source/Runtime/Renderer/Private/DynamicPrimitiveDrawing.inl
--- a/Engine/Source/Runtime/Renderer/Private/DynamicPrimitiveDrawing.inl
+++ b/Engine/Source/Runtime/Renderer/Private/DynamicPrimitiveDrawing.inl
@@ -303,6 +303,7 @@ public:
 	void DoTask(ENamedThreads::Type CurrentThread, const FGraphEventRef& MyCompletionGraphEvent)
 	{
 		DrawViewElementsInner<DrawingPolicyFactoryType>(RHICmdList, View, DrawingContext, DPGIndex, bPreFog, FirstIndex, LastIndex);
+		RHICmdList.HandleRTThreadTaskCompletion(MyCompletionGraphEvent);
 	}
 };
 
diff --git a/Engine/Source/Runtime/Renderer/Private/ReflectionEnvironmentCapture.cpp b/Engine/Source/Runtime/Renderer/Private/ReflectionEnvironmentCapture.cpp
--- a/Engine/Source/Runtime/Renderer/Private/ReflectionEnvironmentCapture.cpp
+++ b/Engine/Source/Runtime/Renderer/Private/ReflectionEnvironmentCapture.cpp
@@ -841,8 +841,12 @@ void CaptureSceneToScratchCubemap(FRHICommandListImmediate& RHICmdList, FSceneRe
 #if PLATFORM_PS4 // @todo ps4 - this should be done a different way
 		// PS4 needs some code here to process the scene
 		extern void TEMP_PostReflectionCaptureRender();
-		//@todo-rco: Fix this assert
 		check(&RHICmdList == &FRHICommandListExecutor::GetImmediateCommandList());
+		check(IsInRenderingThread());
+		{
+			QUICK_SCOPE_CYCLE_COUNTER(STAT_CaptureSceneToScratchCubemap_Flush);
+			FRHICommandListExecutor::GetImmediateCommandList().ImmediateFlush(EImmediateFlushType::FlushRHIThread);
+		}
 		TEMP_PostReflectionCaptureRender();
 #endif
 
diff --git a/Engine/Source/Runtime/Renderer/Private/SceneOcclusion.cpp b/Engine/Source/Runtime/Renderer/Private/SceneOcclusion.cpp
--- a/Engine/Source/Runtime/Renderer/Private/SceneOcclusion.cpp
+++ b/Engine/Source/Runtime/Renderer/Private/SceneOcclusion.cpp
@@ -377,7 +377,7 @@ FRenderQueryRHIParamRef FOcclusionQueryBatcher::BatchPrimitive(const FVector& Bo
 		check(OcclusionQueryPool);
 		CurrentBatchOcclusionQuery = new(BatchOcclusionQueries) FOcclusionBatch;
 		CurrentBatchOcclusionQuery->Query = OcclusionQueryPool->AllocateQuery();
-		CurrentBatchOcclusionQuery->VertexAllocation = FGlobalDynamicVertexBuffer::Get().Allocate(MaxBatchedPrimitives * 8 * sizeof(FVector));
+		CurrentBatchOcclusionQuery->VertexAllocation = FGlobalDynamicVertexBuffer::Get().Allocate(MaxBatchedPrimitives * 8 * sizeof(FVector), true);
 		check(CurrentBatchOcclusionQuery->VertexAllocation.IsValid());
 		NumBatchedPrimitives = 0;
 	}
diff --git a/Engine/Source/Runtime/Renderer/Private/SceneRendering.cpp b/Engine/Source/Runtime/Renderer/Private/SceneRendering.cpp
--- a/Engine/Source/Runtime/Renderer/Private/SceneRendering.cpp
+++ b/Engine/Source/Runtime/Renderer/Private/SceneRendering.cpp
@@ -979,7 +979,10 @@ static void RenderViewFamily_RenderThread(FRHICommandListImmediate& RHICmdList,
 
         // Delete the scene renderer.
 		delete SceneRenderer;
-		FRHICommandListExecutor::GetImmediateCommandList().Flush(); // we want to make sure this all gets to the GPU this frame and doesn't hang around
+		{
+			QUICK_SCOPE_CYCLE_COUNTER(STAT_DeleteSceneRenderer_Dispatch);
+			FRHICommandListExecutor::GetImmediateCommandList().ImmediateFlush(EImmediateFlushType::DispatchToRHIThread); // we want to make sure this all gets to the GPU this frame and doesn't hang around
+		}
 	}
 
 #if STATS
diff --git a/Engine/Source/Runtime/Renderer/Private/ShadowRendering.cpp b/Engine/Source/Runtime/Renderer/Private/ShadowRendering.cpp
--- a/Engine/Source/Runtime/Renderer/Private/ShadowRendering.cpp
+++ b/Engine/Source/Runtime/Renderer/Private/ShadowRendering.cpp
@@ -1385,6 +1385,7 @@ public:
 			// normal shadow map
 			DrawShadowMeshElements<false>(RHICmdList, View, ThisRenderer);
 		}
+		RHICmdList.HandleRTThreadTaskCompletion(MyCompletionGraphEvent);
 	}
 };
 
@@ -1425,6 +1426,7 @@ public:
 	void DoTask(ENamedThreads::Type CurrentThread, const FGraphEventRef& MyCompletionGraphEvent)
 	{
 		ThisRenderer.RenderDepthDynamic(RHICmdList, SceneRenderer, &View);
+		RHICmdList.HandleRTThreadTaskCompletion(MyCompletionGraphEvent);
 	}
 };
 
@@ -1443,7 +1445,7 @@ void FProjectedShadowInfo::RenderDepthInner(FRHICommandList& RHICmdList, FDeferr
 		GRHICommandList.UseParallelAlgorithms() && CVarParallelShadows.GetValueOnRenderThread())
 	{
 		// parallel version
-		FScopedCommandListFlush Flusher;
+		FScopedCommandListWaitForTasks Flusher;
 
 		int32 Width = CVarRHICmdWidth.GetValueOnRenderThread(); // we use a few more than needed to cover non-equal jobs
 		bool OutDirty = false; // unused
diff --git a/Engine/Source/Runtime/Renderer/Private/StaticMeshDrawList.inl b/Engine/Source/Runtime/Renderer/Private/StaticMeshDrawList.inl
--- a/Engine/Source/Runtime/Renderer/Private/StaticMeshDrawList.inl
+++ b/Engine/Source/Runtime/Renderer/Private/StaticMeshDrawList.inl
@@ -349,6 +349,7 @@ public:
 		{
 			this->OutDirty = true;
 		}
+		this->RHICmdList.HandleRTThreadTaskCompletion(MyCompletionGraphEvent);
 	}
 };
 
diff --git a/Engine/Source/Runtime/Renderer/Private/TranslucentRendering.cpp b/Engine/Source/Runtime/Renderer/Private/TranslucentRendering.cpp
--- a/Engine/Source/Runtime/Renderer/Private/TranslucentRendering.cpp
+++ b/Engine/Source/Runtime/Renderer/Private/TranslucentRendering.cpp
@@ -599,6 +599,7 @@ void FTranslucentPrimSet::DrawPrimitivesParallel(
 
 		if (PrimitiveSceneInfo->Proxy && PrimitiveSceneInfo->Proxy->CastsVolumetricTranslucentShadow())
 		{
+			check(!IsInRenderingThread());
 			// can't do this in parallel, defer
 			FRHICommandList* CmdList = new FRHICommandList;
 			FGraphEventRef RenderThreadCompletionEvent = TGraphTask<FVolumetricTranslucentShadowRenderThreadTask>::CreateTask().ConstructAndDispatchWhenReady(CmdList, this, &View, &Renderer, bSeparateTranslucencyPass, PrimIdx);
@@ -827,6 +828,7 @@ public:
 	void DoTask(ENamedThreads::Type CurrentThread, const FGraphEventRef& MyCompletionGraphEvent)
 	{
 		View.TranslucentPrimSet.DrawPrimitivesParallel(RHICmdList, View, Renderer, bSeparateTranslucencyPass, FirstIndex, LastIndex);
+		RHICmdList.HandleRTThreadTaskCompletion(MyCompletionGraphEvent);
 	}
 };
 
@@ -836,7 +838,7 @@ void FDeferredShadingSceneRenderer::RenderTranslucencyParallel(FRHICommandListIm
 	GSceneRenderTargets.AllocLightAttenuation(); // materials will attempt to get this texture before the deferred command to set it up executes
 	int32 Width = CVarRHICmdWidth.GetValueOnRenderThread(); // we use a few more than needed to cover non-equal jobs
 
-	FScopedCommandListFlush Flusher(RHICmdList);
+	FScopedCommandListWaitForTasks Flusher(RHICmdList);
 	FTranslucencyDrawingPolicyFactory::ContextType ThisContext;
 
 	check(IsInRenderingThread());
diff --git a/Engine/Source/Runtime/Renderer/Private/VelocityRendering.cpp b/Engine/Source/Runtime/Renderer/Private/VelocityRendering.cpp
--- a/Engine/Source/Runtime/Renderer/Private/VelocityRendering.cpp
+++ b/Engine/Source/Runtime/Renderer/Private/VelocityRendering.cpp
@@ -523,13 +523,14 @@ public:
 	void DoTask(ENamedThreads::Type CurrentThread, const FGraphEventRef& MyCompletionGraphEvent)
 	{
 		ThisRenderer.RenderDynamicVelocitiesInner(RHICmdList, View, FirstIndex, LastIndex);
+		RHICmdList.HandleRTThreadTaskCompletion(MyCompletionGraphEvent);
 	}
 };
 
 
 static TAutoConsoleVariable<int32> CVarParallelVelocity(
 	TEXT("r.ParallelVelocity"),
-	0,  // not possible because of per-bone motion blur
+	1,  
 	TEXT("Toggles parallel velocity rendering. Parallel rendering must be enabled for this to have an effect.\n"),
 	ECVF_RenderThreadSafe
 	);
@@ -541,7 +542,7 @@ void FDeferredShadingSceneRenderer::RenderVelocitiesInner(FRHICommandListImmedia
 	if (GRHICommandList.UseParallelAlgorithms() && CVarParallelVelocity.GetValueOnRenderThread())
 	{
 		// parallel version
-		FScopedCommandListFlush Flusher(RHICmdList);
+		FScopedCommandListWaitForTasks Flusher(RHICmdList);
 
 		int32 Width = CVarRHICmdWidth.GetValueOnRenderThread(); // we use a few more than needed to cover non-equal jobs
 		bool OutDirty = false; // unused
diff --git a/Engine/Source/Runtime/SlateCore/Private/SlateCoreModule.cpp b/Engine/Source/Runtime/SlateCore/Private/SlateCoreModule.cpp
--- a/Engine/Source/Runtime/SlateCore/Private/SlateCoreModule.cpp
+++ b/Engine/Source/Runtime/SlateCore/Private/SlateCoreModule.cpp
@@ -9,6 +9,7 @@ DEFINE_LOG_CATEGORY(LogSlateStyles);
 
 DEFINE_STAT(STAT_SlateMeasureStringTime);
 DEFINE_STAT(STAT_SlateRenderingRTTime);
+DEFINE_STAT(STAT_SlatePresentRTTime);
 DEFINE_STAT(STAT_SlateUpdateBufferRTTime);
 DEFINE_STAT(STAT_SlateDrawTime);
 DEFINE_STAT(STAT_SlateUpdateBufferGTTime);
diff --git a/Engine/Source/Runtime/SlateCore/Public/SlateCore.h b/Engine/Source/Runtime/SlateCore/Public/SlateCore.h
--- a/Engine/Source/Runtime/SlateCore/Public/SlateCore.h
+++ b/Engine/Source/Runtime/SlateCore/Public/SlateCore.h
@@ -19,6 +19,7 @@ SLATECORE_API DECLARE_LOG_CATEGORY_EXTERN(LogSlateStyles, Log, All);
 
 DECLARE_CYCLE_STAT_EXTERN(TEXT("Measure String"), STAT_SlateMeasureStringTime, STATGROUP_Slate,);
 DECLARE_CYCLE_STAT_EXTERN(TEXT("Slate Rendering RT Time"), STAT_SlateRenderingRTTime, STATGROUP_Slate , SLATECORE_API);
+DECLARE_CYCLE_STAT_EXTERN(TEXT("Slate RT Present Time"), STAT_SlatePresentRTTime, STATGROUP_Slate , SLATECORE_API);
 DECLARE_CYCLE_STAT_EXTERN(TEXT("Update Buffers RT"), STAT_SlateUpdateBufferRTTime, STATGROUP_Slate , SLATECORE_API);
 DECLARE_CYCLE_STAT_EXTERN(TEXT("Draw Time"), STAT_SlateDrawTime, STATGROUP_Slate, SLATECORE_API);
 DECLARE_CYCLE_STAT_EXTERN(TEXT("Update Buffers GT"), STAT_SlateUpdateBufferGTTime, STATGROUP_Slate,);
diff --git a/Engine/Source/Runtime/SlateRHIRenderer/Private/SlateRHIRenderer.cpp b/Engine/Source/Runtime/SlateRHIRenderer/Private/SlateRHIRenderer.cpp
--- a/Engine/Source/Runtime/SlateRHIRenderer/Private/SlateRHIRenderer.cpp
+++ b/Engine/Source/Runtime/SlateRHIRenderer/Private/SlateRHIRenderer.cpp
@@ -450,7 +450,10 @@ void FSlateRHIRenderer::DrawWindow_RenderThread(FRHICommandListImmediate& RHICmd
 	uint32 StartTime		= FPlatformTime::Cycles();
 		
 	// Note - We do not include present time in the slate render thread stat
-	RHICmdList.EndDrawingViewport(ViewportInfo.ViewportRHI, true, bLockToVsync);
+	{
+		SCOPE_CYCLE_COUNTER(STAT_SlatePresentRTTime);
+		RHICmdList.EndDrawingViewport(ViewportInfo.ViewportRHI, true, bLockToVsync);
+	}
 
 	if (bNeedCallFinishFrameForStereo)
 	{
