UE4 - Load/Store actions for render targets
- Fix on OpenGL for read-back into CPU
- Some ES2 -> ES3.1 feature parity

diff --git a/Engine/Source/Runtime/Engine/Private/GPUSkinVertexFactory.cpp b/Engine/Source/Runtime/Engine/Private/GPUSkinVertexFactory.cpp
--- a/Engine/Source/Runtime/Engine/Private/GPUSkinVertexFactory.cpp
+++ b/Engine/Source/Runtime/Engine/Private/GPUSkinVertexFactory.cpp
@@ -202,7 +202,7 @@ template <bool bExtraBoneInfluencesT>
 bool TGPUSkinVertexFactory<bExtraBoneInfluencesT>::ShouldCache(EShaderPlatform Platform, const class FMaterial* Material, const FShaderType* ShaderType)
 {
 	// Skip trying to use extra bone influences on < SM4
-	if (bExtraBoneInfluencesT && GetMaxSupportedFeatureLevel(Platform) < ERHIFeatureLevel::SM4)
+	if (bExtraBoneInfluencesT && GetMaxSupportedFeatureLevel(Platform) < ERHIFeatureLevel::ES3_1)
 	{
 		return false;
 	}
diff --git a/Engine/Source/Runtime/Engine/Private/Landscape/LandscapeRender.cpp b/Engine/Source/Runtime/Engine/Private/Landscape/LandscapeRender.cpp
--- a/Engine/Source/Runtime/Engine/Private/Landscape/LandscapeRender.cpp
+++ b/Engine/Source/Runtime/Engine/Private/Landscape/LandscapeRender.cpp
@@ -555,7 +555,7 @@ FLandscapeComponentSceneProxy::FLandscapeComponentSceneProxy(ULandscapeComponent
 	LevelColor = FLinearColor(1.f, 1.f, 1.f);
 
 	const auto FeatureLevel = GetScene()->GetFeatureLevel();
-	if (FeatureLevel == ERHIFeatureLevel::ES2)
+	if (FeatureLevel <= ERHIFeatureLevel::ES3_1)
 	{
 		HeightmapTexture = NULL;
 		HeightmapSubsectionOffsetU = 0;
@@ -1946,7 +1946,7 @@ void FLandscapeVertexBuffer::InitRHI()
 template <typename INDEX_TYPE>
 void FLandscapeSharedBuffers::CreateIndexBuffers(ERHIFeatureLevel::Type InFeatureLevel, bool bRequiresAdjacencyInformation)
 {
-	if (InFeatureLevel == ERHIFeatureLevel::ES2)
+	if (InFeatureLevel <= ERHIFeatureLevel::ES3_1)
 	{
 		if (!bVertexScoresComputed)
 		{
@@ -1973,7 +1973,7 @@ void FLandscapeSharedBuffers::CreateIndexBuffers(ERHIFeatureLevel::Type InFeatur
 		MaxIndexFull = 0;
 		MinIndexFull = MAX_int32;
 
-		if (InFeatureLevel == ERHIFeatureLevel::ES2)
+		if (InFeatureLevel <= ERHIFeatureLevel::ES3_1)
 		{
 			// ES2 version
 			float MipRatio = (float)SubsectionSizeQuads / (float)LodSubsectionSizeQuads; // Morph current MIP to base MIP
diff --git a/Engine/Source/Runtime/OpenGLDrv/Private/OpenGLRenderTarget.cpp b/Engine/Source/Runtime/OpenGLDrv/Private/OpenGLRenderTarget.cpp
--- a/Engine/Source/Runtime/OpenGLDrv/Private/OpenGLRenderTarget.cpp
+++ b/Engine/Source/Runtime/OpenGLDrv/Private/OpenGLRenderTarget.cpp
@@ -708,11 +708,10 @@ void FOpenGLDynamicRHI::RHIUnmapStagingSurface(FTextureRHIParamRef TextureRHI)
 
 void FOpenGLDynamicRHI::RHIReadSurfaceFloatData(FTextureRHIParamRef TextureRHI,FIntRect Rect,TArray<FFloat16Color>& OutData,ECubeFace CubeFace,int32 ArrayIndex,int32 MipIndex)
 {
-	VERIFY_GL_SCOPE();
-	// Not supported on older APIs
-	check(GMaxRHIFeatureLevel >= ERHIFeatureLevel::SM5 || ArrayIndex == 0);
+	VERIFY_GL_SCOPE();	
 
-	check( FOpenGL::SupportsFloatReadSurface() );
+	//reading from arrays only supported on SM5 and up.
+	check(FOpenGL::SupportsFloatReadSurface() && (ArrayIndex == 0 || GRHIFeatureLevel >= ERHIFeatureLevel::SM5));	
 	FOpenGLTextureBase* Texture = GetOpenGLTextureFromRHITexture(TextureRHI);
 	check(TextureRHI->GetFormat() == PF_FloatRGBA);
 
@@ -747,8 +746,8 @@ void FOpenGLDynamicRHI::RHIReadSurfaceFloatData(FTextureRHIParamRef TextureRHI,F
 
 	glBindFramebuffer(UGL_READ_FRAMEBUFFER, SourceFramebuffer);
 	FOpenGL::ReadBuffer(SourceFramebuffer == 0 ? GL_BACK : GL_COLOR_ATTACHMENT0);
-	glPixelStorei(GL_PACK_ALIGNMENT, 1);
-	glReadPixels(Rect.Min.X, Rect.Min.Y, SizeX, SizeY, GL_RGBA, GL_HALF_FLOAT, OutData.GetData());
+	glPixelStorei(GL_PACK_ALIGNMENT, 1);	
+	glReadPixels(Rect.Min.X, Rect.Min.Y, SizeX, SizeY, GL_RGBA, FOpenGL::GetReadHalfFloatPixelsEnum(), OutData.GetData());
 	glPixelStorei(GL_PACK_ALIGNMENT, 4);
 
 	if (bTempFBO)
diff --git a/Engine/Source/Runtime/RHI/Public/RHIDefinitions.h b/Engine/Source/Runtime/RHI/Public/RHIDefinitions.h
--- a/Engine/Source/Runtime/RHI/Public/RHIDefinitions.h
+++ b/Engine/Source/Runtime/RHI/Public/RHIDefinitions.h
@@ -517,6 +517,26 @@ enum ETextureReallocationStatus
 	TexRealloc_InProgress,
 };
 
+/**
+ * Action to take when a rendertarget is set.
+ */
+enum class ERenderTargetLoadAction
+{
+	ENoAction,
+	ELoad,
+	//EClear clear should go here but isn't implemented yet.
+};
+
+/**
+ * Action to take when a rendertarget is unset or at the end of a pass. 
+ */
+enum class ERenderTargetStoreAction
+{
+	ENoAction,
+	EStore,
+	EMultisampleResolve,
+};
+
 inline bool IsPCPlatform(const EShaderPlatform Platform)
 {
 	return Platform == SP_PCD3D_SM5 || Platform == SP_PCD3D_SM4 || Platform == SP_PCD3D_ES2 || Platform ==  SP_OPENGL_SM4 || Platform == SP_OPENGL_SM4_MAC || Platform == SP_OPENGL_SM5 || Platform == SP_OPENGL_PCES2;
diff --git a/Engine/Source/Runtime/RHI/Public/RHIResources.h b/Engine/Source/Runtime/RHI/Public/RHIResources.h
--- a/Engine/Source/Runtime/RHI/Public/RHIResources.h
+++ b/Engine/Source/Runtime/RHI/Public/RHIResources.h
@@ -576,29 +576,48 @@ public:
 
 	/** Array slice or texture cube face.  Only valid if texture resource was created with TexCreate_TargetArraySlicesIndependently! */
 	uint32 ArraySliceIndex;
+	
+	ERenderTargetLoadAction LoadAction;
+	ERenderTargetStoreAction StoreAction;
 
 	FRHIRenderTargetView() : 
 		Texture(NULL),
 		MipIndex(0),
-		ArraySliceIndex(-1)
+		ArraySliceIndex(-1),
+		LoadAction(ERenderTargetLoadAction::ENoAction),
+		StoreAction(ERenderTargetStoreAction::EStore)
 	{}
 
 	FRHIRenderTargetView(const FRHIRenderTargetView& Other) :
 		Texture(Other.Texture),
 		MipIndex(Other.MipIndex),
-		ArraySliceIndex(Other.ArraySliceIndex)
+		ArraySliceIndex(Other.ArraySliceIndex),
+		LoadAction(Other.LoadAction),
+		StoreAction(Other.StoreAction)
 	{}
 
 	FRHIRenderTargetView(FTextureRHIParamRef InTexture) :
 		Texture(InTexture),
 		MipIndex(0),
-		ArraySliceIndex(-1)
+		ArraySliceIndex(-1),
+		LoadAction(ERenderTargetLoadAction::ENoAction),
+		StoreAction(ERenderTargetStoreAction::EStore)
 	{}
 
 	FRHIRenderTargetView(FTextureRHIParamRef InTexture, uint32 InMipIndex, uint32 InArraySliceIndex) :
 		Texture(InTexture),
 		MipIndex(InMipIndex),
-		ArraySliceIndex(InArraySliceIndex)
+		ArraySliceIndex(InArraySliceIndex),
+		LoadAction(ERenderTargetLoadAction::ENoAction),
+		StoreAction(ERenderTargetStoreAction::EStore)
+	{}
+	
+	FRHIRenderTargetView(FTextureRHIParamRef InTexture, uint32 InMipIndex, uint32 InArraySliceIndex, ERenderTargetLoadAction InLoadAction, ERenderTargetStoreAction InStoreAction) :
+		Texture(InTexture),
+		MipIndex(InMipIndex),
+		ArraySliceIndex(InArraySliceIndex),
+		LoadAction(InLoadAction),
+		StoreAction(InStoreAction)
 	{}
 };
 
