UE4 - parallel rendering tweaks. XB1 pseduo-threadsafe query readback. virtual RHI call elimination, optional (and disabled) stats on RHI commands. Disabled state reduction.

diff --git a/Engine/Source/Runtime/Core/Public/HAL/Platform.h b/Engine/Source/Runtime/Core/Public/HAL/Platform.h
--- a/Engine/Source/Runtime/Core/Public/HAL/Platform.h
+++ b/Engine/Source/Runtime/Core/Public/HAL/Platform.h
@@ -246,6 +246,10 @@
 	#define PLATFORM_HAS_THREADSAFE_RHIGetRenderQueryResult		0
 #endif
 
+#ifndef PLATFORM_USES_FIXED_RHI_CLASS
+	#define PLATFORM_USES_FIXED_RHI_CLASS		0
+#endif
+
 #ifndef PLATFORM_SUPPORTS_MULTIPLE_NATIVE_WINDOWS
 	#define PLATFORM_SUPPORTS_MULTIPLE_NATIVE_WINDOWS	1
 #endif
diff --git a/Engine/Source/Runtime/IOS/MetalRHI/Private/MetalManager.cpp b/Engine/Source/Runtime/IOS/MetalRHI/Private/MetalManager.cpp
--- a/Engine/Source/Runtime/IOS/MetalRHI/Private/MetalManager.cpp
+++ b/Engine/Source/Runtime/IOS/MetalRHI/Private/MetalManager.cpp
@@ -815,8 +815,6 @@ void FMetalManager::CommitGraphicsResourceTables()
 {
 	uint32 Start = FPlatformTime::Cycles();
 
-// 	GRHICommandList.Verify();
-
 	check(CurrentBoundShaderState);
 
 	SetResourcesFromTables(CurrentBoundShaderState->VertexShader, CrossCompiler::SHADER_STAGE_VERTEX, ResourceTableFrameCounter);
diff --git a/Engine/Source/Runtime/RHI/Private/RHI.cpp b/Engine/Source/Runtime/RHI/Private/RHI.cpp
--- a/Engine/Source/Runtime/RHI/Private/RHI.cpp
+++ b/Engine/Source/Runtime/RHI/Private/RHI.cpp
@@ -74,6 +74,11 @@ DECLARE_CYCLE_STAT(TEXT("Delete Resources"), STAT_DeleteResources, STATGROUP_RHI
 
 void FRHIResource::FlushPendingDeletes()
 {
+	if (GIsRequestingExit) 
+	{ 
+		// Leak these because we're terrible people. 
+		return; 
+	}
 	SCOPE_CYCLE_COUNTER(STAT_DeleteResources);
 
 	check(IsInRenderingThread());
diff --git a/Engine/Source/Runtime/RHI/Private/RHICommandList.cpp b/Engine/Source/Runtime/RHI/Private/RHICommandList.cpp
--- a/Engine/Source/Runtime/RHI/Private/RHICommandList.cpp
+++ b/Engine/Source/Runtime/RHI/Private/RHICommandList.cpp
@@ -24,6 +24,11 @@ Requirements for RHI thread
 ***/
 #endif
 
+#if !PLATFORM_USES_FIXED_RHI_CLASS
+#include "RHICommandListCommandExecutes.inl"
+#endif
+
+
 static TAutoConsoleVariable<int32> CVarRHICmdBypass(
 	TEXT("r.RHICmdBypass"),
 	FRHICommandListExecutor::DefaultBypass,
@@ -58,6 +63,7 @@ void FRHICommandListExecutor::ExecuteInner_DoExecute(FRHICommandListBase& CmdLis
 		while (Iter.HasCommandsLeft())
 		{
 			FRHICommandBase* Cmd = Iter.NextCommand();
+			//FPlatformMisc::Prefetch(Cmd->Next);
 			Cmd->CallExecuteAndDestruct();
 		}
 	}
@@ -415,6 +421,9 @@ void FRHICommandListBase::Reset()
 	NumCommands = 0;
 	Root = nullptr;
 	CommandLink = &Root;
+#if USE_RHICOMMAND_STATE_REDUCTION
+	StateCache = nullptr;
+#endif
 	UID = GRHICommandList.UIDCounter.Increment();
 }
 
@@ -540,22 +549,6 @@ void FRHICommandListBase::QueueCommandListSubmit(class FRHICommandList* CmdList)
 }
 
 #if PLATFORM_SUPPORTS_RHI_THREAD
-struct FRHICommandBeginDrawingViewport : public FRHICommand<FRHICommandBeginDrawingViewport>
-{
-	FViewportRHIParamRef Viewport;
-	FTextureRHIParamRef RenderTargetRHI;
-
-	FORCEINLINE_DEBUGGABLE FRHICommandBeginDrawingViewport(FViewportRHIParamRef InViewport, FTextureRHIParamRef InRenderTargetRHI)
-		: Viewport(InViewport)
-		, RenderTargetRHI(InRenderTargetRHI)
-	{
-	}
-	void Execute()
-	{
-		BeginDrawingViewport_Internal(Viewport, RenderTargetRHI);
-	}
-};
-
 void FRHICommandList::BeginDrawingViewport(FViewportRHIParamRef Viewport, FTextureRHIParamRef RenderTargetRHI)
 {
 	check(IsImmediate() && IsInRenderingThread());
@@ -566,23 +559,6 @@ void FRHICommandList::BeginDrawingViewport(FViewportRHIParamRef Viewport, FTextu
 	}
 	new (AllocCommand<FRHICommandBeginDrawingViewport>()) FRHICommandBeginDrawingViewport(Viewport, RenderTargetRHI);
 }
-struct FRHICommandEndDrawingViewport : public FRHICommand<FRHICommandEndDrawingViewport>
-{
-	FViewportRHIParamRef Viewport;
-	bool bPresent;
-	bool bLockToVsync;
-
-	FORCEINLINE_DEBUGGABLE FRHICommandEndDrawingViewport(FViewportRHIParamRef InViewport, bool InbPresent, bool InbLockToVsync)
-		: Viewport(InViewport)
-		, bPresent(InbPresent)
-		, bLockToVsync(InbLockToVsync)
-	{
-	}
-	void Execute()
-	{
-		EndDrawingViewport_Internal(Viewport, bPresent, bLockToVsync);
-	}
-};
 
 void FRHICommandList::EndDrawingViewport(FViewportRHIParamRef Viewport, bool bPresent, bool bLockToVsync)
 {
diff --git a/Engine/Source/Runtime/RHI/Public/RHICommandList.h b/Engine/Source/Runtime/RHI/Public/RHICommandList.h
--- a/Engine/Source/Runtime/RHI/Public/RHICommandList.h
+++ b/Engine/Source/Runtime/RHI/Public/RHICommandList.h
@@ -10,6 +10,8 @@
 
 DECLARE_STATS_GROUP(TEXT("RHICmdList"), STATGROUP_RHICMDLIST, STATCAT_Advanced);
 
+#define USE_RHICOMMAND_STATE_REDUCTION (0)
+
 extern RHI_API TAutoConsoleVariable<int32> CVarRHICmdWidth;
 
 struct FRHICommandBase
@@ -94,7 +96,6 @@ public:
 		}
 	}
 
-
 	struct FDrawUpData
 	{
 
@@ -114,17 +115,21 @@ public:
 			, OutIndexData(nullptr)
 		{
 		}
-
 	};
 
 private:
 	FRHICommandBase* Root;
 	FRHICommandBase** CommandLink;
 	bool bExecuting;
-	FGraphEventArray RTTasks;
 	uint32 NumCommands;
 	uint32 UID;
-	FMemStackBase MemManager; // this has 1k inline storage, so we want to small members before it
+	FMemStackBase MemManager; 
+#if USE_RHICOMMAND_STATE_REDUCTION
+protected:
+	struct FRHICommandListStateCache* StateCache;
+private:
+#endif
+	FGraphEventArray RTTasks;
 
 	void Reset();
 
@@ -132,9 +137,51 @@ private:
 	friend class FRHICommandListIterator;
 
 public:
-	FDrawUpData DrawUPData; // this isn't used very much so we put it at the cold end here
+	FDrawUpData DrawUPData; 
 };
 
+#if USE_RHICOMMAND_STATE_REDUCTION
+struct FRHICommandListStateCache
+{
+	// warning, no destructor is ever called for this struct
+
+	FRHICommandListStateCache()
+	{
+		FlushShaderState();
+		FlushSamplerState();
+	}
+	enum 
+	{
+		MAX_SAMPLERS_PER_SHADER_STAGE = 32, // it is okish if this is too small, those buffers simply won't get state reduction
+		MAX_UNIFORM_BUFFERS_PER_SHADER_STAGE = 14, // it is okish if this is too small, those buffers simply won't get state reduction
+	};
+	FSamplerStateRHIParamRef Samplers[SF_NumFrequencies][MAX_SAMPLERS_PER_SHADER_STAGE];
+	FUniformBufferRHIParamRef BoundUniformBuffers[SF_NumFrequencies][MAX_UNIFORM_BUFFERS_PER_SHADER_STAGE];
+
+	FORCEINLINE void FlushSamplerState()
+	{
+		for (int32 Index = 0; Index < SF_NumFrequencies; Index++)
+		{
+			for (int32 IndexInner = 0; IndexInner < MAX_SAMPLERS_PER_SHADER_STAGE; IndexInner++)
+			{
+				Samplers[Index][IndexInner] = nullptr;
+			}
+		}
+	}
+
+	FORCEINLINE void FlushShaderState()
+	{
+		for (int32 Index = 0; Index < SF_NumFrequencies; Index++)
+		{
+			for (int32 IndexInner = 0; IndexInner < MAX_UNIFORM_BUFFERS_PER_SHADER_STAGE; IndexInner++)
+			{
+				BoundUniformBuffers[Index][IndexInner] = nullptr;
+			}
+		}
+	}
+};
+#endif
+
 template<typename TCmd>
 struct FRHICommand : public FRHICommandBase
 {
@@ -158,10 +205,7 @@ struct FRHICommandSetRasterizerState : public FRHICommand<FRHICommandSetRasteriz
 		: State(InState)
 	{
 	}
-	void Execute()
-	{
-		SetRasterizerState_Internal(State);
-	}
+	RHI_API void Execute();
 };
 
 struct FRHICommandSetDepthStencilState : public FRHICommand<FRHICommandSetDepthStencilState>
@@ -173,10 +217,7 @@ struct FRHICommandSetDepthStencilState : public FRHICommand<FRHICommandSetDepthS
 		, StencilRef(InStencilRef)
 	{
 	}
-	void Execute()
-	{
-		SetDepthStencilState_Internal(State, StencilRef);
-	}
+	RHI_API void Execute();
 };
 
 template <typename TShaderRHIParamRef>
@@ -195,10 +236,7 @@ struct FRHICommandSetShaderParameter : public FRHICommand<FRHICommandSetShaderPa
 		, NumBytes(InNumBytes)
 	{
 	}
-	void Execute()
-	{
-		SetShaderParameter_Internal(Shader, BufferIndex, BaseIndex, NumBytes, NewValue); 
-	}
+	RHI_API void Execute();
 };
 
 template <typename TShaderRHIParamRef>
@@ -213,10 +251,7 @@ struct FRHICommandSetShaderUniformBuffer : public FRHICommand<FRHICommandSetShad
 		, UniformBuffer(InUniformBuffer)
 	{
 	}
-	void Execute()
-	{
-		SetShaderUniformBuffer_Internal(Shader, BaseIndex, UniformBuffer);
-	}
+	RHI_API void Execute();
 };
 
 template <typename TShaderRHIParamRef>
@@ -231,10 +266,7 @@ struct FRHICommandSetShaderTexture : public FRHICommand<FRHICommandSetShaderText
 		, Texture(InTexture)
 	{
 	}
-	void Execute()
-	{
-		SetShaderTexture_Internal(Shader, TextureIndex, Texture);
-	}
+	RHI_API void Execute();
 };
 
 template <typename TShaderRHIParamRef>
@@ -249,10 +281,7 @@ struct FRHICommandSetShaderResourceViewParameter : public FRHICommand<FRHIComman
 		, SRV(InSRV)
 	{
 	}
-	void Execute()
-	{
-		SetShaderResourceViewParameter_Internal(Shader, SamplerIndex, SRV);
-	}
+	RHI_API void Execute();
 };
 
 template <typename TShaderRHIParamRef>
@@ -267,10 +296,7 @@ struct FRHICommandSetUAVParameter : public FRHICommand<FRHICommandSetUAVParamete
 		, UAV(InUAV)
 	{
 	}
-	void Execute()
-	{
-		SetUAVParameter_Internal(Shader, UAVIndex, UAV);
-	}
+	RHI_API void Execute();
 };
 
 template <typename TShaderRHIParamRef>
@@ -287,10 +313,7 @@ struct FRHICommandSetUAVParameter_IntialCount : public FRHICommand<FRHICommandSe
 		, InitialCount(InInitialCount)
 	{
 	}
-	void Execute()
-	{
-		SetUAVParameter_Internal(Shader, UAVIndex, UAV, InitialCount);
-	}
+	RHI_API void Execute();
 };
 
 template <typename TShaderRHIParamRef>
@@ -305,10 +328,7 @@ struct FRHICommandSetShaderSampler : public FRHICommand<FRHICommandSetShaderSamp
 		, Sampler(InSampler)
 	{
 	}
-	void Execute()
-	{
-		SetShaderSampler_Internal(Shader, SamplerIndex, Sampler);
-	}
+	RHI_API void Execute();
 };
 
 struct FRHICommandDrawPrimitive : public FRHICommand<FRHICommandDrawPrimitive>
@@ -324,11 +344,7 @@ struct FRHICommandDrawPrimitive : public FRHICommand<FRHICommandDrawPrimitive>
 		, NumInstances(InNumInstances)
 	{
 	}
-	void Execute()
-	{
-		DrawPrimitive_Internal(PrimitiveType, BaseVertexIndex, NumPrimitives, NumInstances);
-	}
-
+	RHI_API void Execute();
 };
 
 struct FRHICommandDrawIndexedPrimitive : public FRHICommand<FRHICommandDrawIndexedPrimitive>
@@ -352,10 +368,7 @@ struct FRHICommandDrawIndexedPrimitive : public FRHICommand<FRHICommandDrawIndex
 		, NumInstances(InNumInstances)
 	{
 	}
-	void Execute()
-	{
-		DrawIndexedPrimitive_Internal(IndexBuffer, PrimitiveType, BaseVertexIndex, MinIndex, NumVertices, StartIndex, NumPrimitives, NumInstances);
-	}
+	RHI_API void Execute();
 };
 
 struct FRHICommandSetBoundShaderState : public FRHICommand<FRHICommandSetBoundShaderState>
@@ -365,10 +378,7 @@ struct FRHICommandSetBoundShaderState : public FRHICommand<FRHICommandSetBoundSh
 		: BoundShaderState(InBoundShaderState)
 	{
 	}
-	void Execute()
-	{
-		SetBoundShaderState_Internal(BoundShaderState);
-	}
+	RHI_API void Execute();
 };
 
 struct FRHICommandSetBlendState : public FRHICommand<FRHICommandSetBlendState>
@@ -380,10 +390,7 @@ struct FRHICommandSetBlendState : public FRHICommand<FRHICommandSetBlendState>
 		, BlendFactor(InBlendFactor)
 	{
 	}
-	void Execute()
-	{
-		SetBlendState_Internal(State, BlendFactor);
-	}
+	RHI_API void Execute();
 };
 
 struct FRHICommandSetStreamSource : public FRHICommand<FRHICommandSetStreamSource>
@@ -399,10 +406,7 @@ struct FRHICommandSetStreamSource : public FRHICommand<FRHICommandSetStreamSourc
 		, Offset(InOffset)
 	{
 	}
-	void Execute()
-	{
-		SetStreamSource_Internal(StreamIndex, VertexBuffer, Stride, Offset);
-	}
+	RHI_API void Execute();
 };
 
 struct FRHICommandSetViewport : public FRHICommand<FRHICommandSetViewport>
@@ -422,10 +426,7 @@ struct FRHICommandSetViewport : public FRHICommand<FRHICommandSetViewport>
 		, MaxZ(InMaxZ)
 	{
 	}
-	void Execute()
-	{
-		SetViewport_Internal(MinX, MinY, MinZ, MaxX, MaxY, MaxZ);
-	}
+	RHI_API void Execute();
 };
 
 struct FRHICommandSetScissorRect : public FRHICommand<FRHICommandSetScissorRect>
@@ -443,10 +444,7 @@ struct FRHICommandSetScissorRect : public FRHICommand<FRHICommandSetScissorRect>
 		, MaxY(InMaxY)
 	{
 	}
-	void Execute()
-	{
-		SetScissorRect_Internal(bEnable, MinX, MinY, MaxX, MaxY);
-	}
+	RHI_API void Execute();
 };
 
 struct FRHICommandSetRenderTargets : public FRHICommand<FRHICommandSetRenderTargets>
@@ -479,15 +477,7 @@ struct FRHICommandSetRenderTargets : public FRHICommand<FRHICommandSetRenderTarg
 			UAVs[Index] = InUAVs[Index];
 		}
 	}
-	void Execute()
-	{
-		SetRenderTargets_Internal(
-			NewNumSimultaneousRenderTargets,
-			NewRenderTargetsRHI,
-			NewDepthStencilTargetRHI,
-			NewNumUAVs,
-			UAVs);
-	}
+	RHI_API void Execute();
 };
 
 struct FRHICommandEndDrawPrimitiveUP : public FRHICommand<FRHICommandEndDrawPrimitiveUP>
@@ -506,13 +496,7 @@ struct FRHICommandEndDrawPrimitiveUP : public FRHICommand<FRHICommandEndDrawPrim
 		, OutVertexData(InOutVertexData)
 	{
 	}
-	void Execute()
-	{
-		void* Buffer = NULL;
-		BeginDrawPrimitiveUP_Internal(PrimitiveType, NumPrimitives, NumVertices, VertexDataStride, Buffer);
-		FMemory::Memcpy(Buffer, OutVertexData, NumVertices * VertexDataStride);
-		EndDrawPrimitiveUP_Internal();
-	}
+	RHI_API void Execute();
 };
 
 
@@ -540,24 +524,7 @@ struct FRHICommandEndDrawIndexedPrimitiveUP : public FRHICommand<FRHICommandEndD
 		, OutIndexData(InOutIndexData)
 	{
 	}
-	void Execute()
-	{
-		void* VertexBuffer = nullptr;
-		void* IndexBuffer = nullptr;
-		BeginDrawIndexedPrimitiveUP_Internal(
-			PrimitiveType,
-			NumPrimitives,
-			NumVertices,
-			VertexDataStride,
-			VertexBuffer,
-			MinVertexIndex,
-			NumIndices,
-			IndexDataStride,
-			IndexBuffer);
-		FMemory::Memcpy(VertexBuffer, OutVertexData, NumVertices * VertexDataStride);
-		FMemory::Memcpy(IndexBuffer, OutIndexData, NumIndices * IndexDataStride);
-		EndDrawIndexedPrimitiveUP_Internal();
-	}
+	RHI_API void Execute();
 };
 
 struct FRHICommandSetComputeShader : public FRHICommand<FRHICommandSetComputeShader>
@@ -567,10 +534,7 @@ struct FRHICommandSetComputeShader : public FRHICommand<FRHICommandSetComputeSha
 		: ComputeShader(InComputeShader)
 	{
 	}
-	void Execute()
-	{
-		SetComputeShader_Internal(ComputeShader);
-	}
+	RHI_API void Execute();
 };
 
 struct FRHICommandDispatchComputeShader : public FRHICommand<FRHICommandDispatchComputeShader>
@@ -584,11 +548,7 @@ struct FRHICommandDispatchComputeShader : public FRHICommand<FRHICommandDispatch
 		, ThreadGroupCountZ(InThreadGroupCountZ)
 	{
 	}
-	void Execute()
-	{
-		DispatchComputeShader_Internal(ThreadGroupCountX, ThreadGroupCountY, ThreadGroupCountZ);
-	}
-
+	RHI_API void Execute();
 };
 
 struct FRHICommandDispatchIndirectComputeShader : public FRHICommand<FRHICommandDispatchIndirectComputeShader>
@@ -600,10 +560,7 @@ struct FRHICommandDispatchIndirectComputeShader : public FRHICommand<FRHICommand
 		, ArgumentOffset(InArgumentOffset)
 	{
 	}
-	void Execute()
-	{
-		DispatchIndirectComputeShader_Internal(ArgumentBuffer, ArgumentOffset);
-	}
+	RHI_API void Execute();
 };
 
 struct FRHICommandAutomaticCacheFlushAfterComputeShader : public FRHICommand<FRHICommandAutomaticCacheFlushAfterComputeShader>
@@ -613,18 +570,12 @@ struct FRHICommandAutomaticCacheFlushAfterComputeShader : public FRHICommand<FRH
 		: bEnable(InbEnable)
 	{
 	}
-	void Execute()
-	{
-		AutomaticCacheFlushAfterComputeShader_Internal(bEnable);
-	}
+	RHI_API void Execute();
 };
 
 struct FRHICommandFlushComputeShaderCache : public FRHICommand<FRHICommandFlushComputeShaderCache>
 {
-	void Execute()
-	{
-		FlushComputeShaderCache_Internal();
-	}
+	RHI_API void Execute();
 };
 
 struct FRHICommandDrawPrimitiveIndirect : public FRHICommand<FRHICommandDrawPrimitiveIndirect>
@@ -638,10 +589,7 @@ struct FRHICommandDrawPrimitiveIndirect : public FRHICommand<FRHICommandDrawPrim
 		, ArgumentOffset(InArgumentOffset)
 	{
 	}
-	void Execute()
-	{
-		DrawPrimitiveIndirect_Internal(PrimitiveType, ArgumentBuffer, ArgumentOffset);
-	}
+	RHI_API void Execute();
 };
 
 struct FRHICommandDrawIndexedIndirect : public FRHICommand<FRHICommandDrawIndexedIndirect>
@@ -660,10 +608,7 @@ struct FRHICommandDrawIndexedIndirect : public FRHICommand<FRHICommandDrawIndexe
 		, NumInstances(InNumInstances)
 	{
 	}
-	void Execute()
-	{
-		DrawIndexedIndirect_Internal(IndexBufferRHI, PrimitiveType, ArgumentsBufferRHI, DrawArgumentsIndex, NumInstances);
-	}
+	RHI_API void Execute();
 };
 
 struct FRHICommandDrawIndexedPrimitiveIndirect : public FRHICommand<FRHICommandDrawIndexedPrimitiveIndirect>
@@ -680,10 +625,7 @@ struct FRHICommandDrawIndexedPrimitiveIndirect : public FRHICommand<FRHICommandD
 		, ArgumentOffset(InArgumentOffset)
 	{
 	}
-	void Execute()
-	{
-		DrawIndexedPrimitiveIndirect_Internal(PrimitiveType, IndexBuffer, ArgumentsBuffer, ArgumentOffset);
-	}
+	RHI_API void Execute();
 };
 
 struct FRHICommandEnableDepthBoundsTest : public FRHICommand<FRHICommandEnableDepthBoundsTest>
@@ -698,10 +640,7 @@ struct FRHICommandEnableDepthBoundsTest : public FRHICommand<FRHICommandEnableDe
 		, MaxDepth(InMaxDepth)
 	{
 	}
-	void Execute()
-	{
-		EnableDepthBoundsTest_Internal(bEnable, MinDepth, MaxDepth);
-	}
+	RHI_API void Execute();
 };
 
 struct FRHICommandClearUAV : public FRHICommand<FRHICommandClearUAV>
@@ -717,10 +656,7 @@ struct FRHICommandClearUAV : public FRHICommand<FRHICommandClearUAV>
 		Values[2] = InValues[2];
 		Values[3] = InValues[3];
 	}
-	void Execute()
-	{
-		ClearUAV_Internal(UnorderedAccessViewRHI, Values);
-	}
+	RHI_API void Execute();
 };
 
 struct FRHICommandCopyToResolveTarget : public FRHICommand<FRHICommandCopyToResolveTarget>
@@ -737,11 +673,7 @@ struct FRHICommandCopyToResolveTarget : public FRHICommand<FRHICommandCopyToReso
 		, ResolveParams(InResolveParams)
 	{
 	}
-	void Execute()
-	{
-		CopyToResolveTarget_Internal(SourceTexture, DestTexture, bKeepOriginalSurface, ResolveParams);
-	}
-
+	RHI_API void Execute();
 };
 
 struct FRHICommandClear : public FRHICommand<FRHICommandClear>
@@ -772,10 +704,7 @@ struct FRHICommandClear : public FRHICommand<FRHICommandClear>
 		, bClearStencil(InbClearStencil)
 	{
 	}
-	void Execute()
-	{
-		Clear_Internal(bClearColor, Color, bClearDepth, Depth, bClearStencil, Stencil, ExcludeRect);
-	}
+	RHI_API void Execute();
 };
 
 struct FRHICommandClearMRT : public FRHICommand<FRHICommandClearMRT>
@@ -813,10 +742,7 @@ struct FRHICommandClearMRT : public FRHICommand<FRHICommandClearMRT>
 			ColorArray[Index] = InColorArray[Index];
 		}
 	}
-	void Execute()
-	{
-		ClearMRT_Internal(bClearColor, NumClearColors, ColorArray, bClearDepth, Depth, bClearStencil, Stencil, ExcludeRect);
-	}
+	RHI_API void Execute();
 };
 
 
@@ -920,15 +846,7 @@ struct FRHICommandBuildLocalBoundShaderState : public FRHICommand<FRHICommandBui
 
 	{
 	}
-	void Execute()
-	{
-		check(!IsValidRef(WorkArea.ComputedBSS->BSS)); // should not already have been created
-		if (WorkArea.ComputedBSS->UseCount)
-		{
-			WorkArea.ComputedBSS->BSS = CreateBoundShaderState_Internal(WorkArea.Args.VertexDeclarationRHI, WorkArea.Args.VertexShaderRHI, WorkArea.Args.HullShaderRHI, WorkArea.Args.DomainShaderRHI, WorkArea.Args.PixelShaderRHI, WorkArea.Args.GeometryShaderRHI);
-		}
-	}
-
+	RHI_API void Execute();
 };
 
 struct FRHICommandSetLocalBoundShaderState : public FRHICommand<FRHICommandSetLocalBoundShaderState>
@@ -940,18 +858,7 @@ struct FRHICommandSetLocalBoundShaderState : public FRHICommand<FRHICommandSetLo
 		check(CheckCmdList == LocalBoundShaderState.WorkArea->CheckCmdList && CheckCmdList->GetUID() == LocalBoundShaderState.WorkArea->UID); // this BSS was not built for this particular commandlist
 		LocalBoundShaderState.WorkArea->ComputedBSS->UseCount++;
 	}
-	void Execute()
-	{
-		check(LocalBoundShaderState.WorkArea->ComputedBSS->UseCount > 0 && IsValidRef(LocalBoundShaderState.WorkArea->ComputedBSS->BSS)); // this should have been created and should have uses outstanding
-
-		SetBoundShaderState_Internal(LocalBoundShaderState.WorkArea->ComputedBSS->BSS);
-
-		if (--LocalBoundShaderState.WorkArea->ComputedBSS->UseCount == 0)
-		{
-			LocalBoundShaderState.WorkArea->ComputedBSS->~FComputedBSS();
-		}
-	}
-
+	RHI_API void Execute();
 };
 
 
@@ -1022,21 +929,7 @@ struct FRHICommandBuildLocalUniformBuffer : public FRHICommand<FRHICommandBuildL
 
 	{
 	}
-	void Execute()
-	{
-		check(!IsValidRef(WorkArea.ComputedUniformBuffer->UniformBuffer) && WorkArea.Layout && WorkArea.Contents); // should not already have been created
-		if (WorkArea.ComputedUniformBuffer->UseCount)
-		{
-#if PLATFORM_SUPPORTS_RHI_THREAD
-			WorkArea.ComputedUniformBuffer->UniformBuffer = RHICreateUniformBuffer(WorkArea.Contents, *WorkArea.Layout, UniformBuffer_SingleFrame);
-#else
-			WorkArea.ComputedUniformBuffer->UniformBuffer = CreateUniformBuffer_Internal(WorkArea.Contents, *WorkArea.Layout, UniformBuffer_SingleFrame);
-#endif
-		}
-		WorkArea.Layout = nullptr;
-		WorkArea.Contents = nullptr;
-	}
-
+	RHI_API void Execute();
 };
 
 template <typename TShaderRHIParamRef>
@@ -1054,16 +947,7 @@ struct FRHICommandSetLocalUniformBuffer : public FRHICommand<FRHICommandSetLocal
 		check(CheckCmdList == LocalUniformBuffer.WorkArea->CheckCmdList && CheckCmdList->GetUID() == LocalUniformBuffer.WorkArea->UID); // this uniform buffer was not built for this particular commandlist
 		LocalUniformBuffer.WorkArea->ComputedUniformBuffer->UseCount++;
 	}
-	void Execute()
-	{
-		check(LocalUniformBuffer.WorkArea->ComputedUniformBuffer->UseCount > 0 && IsValidRef(LocalUniformBuffer.WorkArea->ComputedUniformBuffer->UniformBuffer)); // this should have been created and should have uses outstanding
-		SetShaderUniformBuffer_Internal(Shader, BaseIndex, LocalUniformBuffer.WorkArea->ComputedUniformBuffer->UniformBuffer);
-		if (--LocalUniformBuffer.WorkArea->ComputedUniformBuffer->UseCount == 0)
-		{
-			LocalUniformBuffer.WorkArea->ComputedUniformBuffer->~FComputedUniformBuffer();
-		}
-	}
-
+	RHI_API void Execute();
 };
 
 struct FRHICommandBeginRenderQuery : public FRHICommand<FRHICommandBeginRenderQuery>
@@ -1074,10 +958,7 @@ struct FRHICommandBeginRenderQuery : public FRHICommand<FRHICommandBeginRenderQu
 		: RenderQuery(InRenderQuery)
 	{
 	}
-	void Execute()
-	{
-		BeginRenderQuery_Internal(RenderQuery);
-	}
+	RHI_API void Execute();
 };
 
 struct FRHICommandEndRenderQuery : public FRHICommand<FRHICommandEndRenderQuery>
@@ -1088,10 +969,7 @@ struct FRHICommandEndRenderQuery : public FRHICommand<FRHICommandEndRenderQuery>
 		: RenderQuery(InRenderQuery)
 	{
 	}
-	void Execute()
-	{
-		EndRenderQuery_Internal(RenderQuery);
-	}
+	RHI_API void Execute();
 };
 
 #if !PLATFORM_SUPPORTS_RHI_THREAD
@@ -1103,10 +981,7 @@ struct FRHICommandResetRenderQuery : public FRHICommand<FRHICommandResetRenderQu
 		: RenderQuery(InRenderQuery)
 	{
 	}
-	void Execute()
-	{
-		ResetRenderQuery_Internal(RenderQuery);
-	}
+	RHI_API void Execute();
 };
 #endif
 
@@ -1115,10 +990,7 @@ struct FRHICommandBeginScene : public FRHICommand<FRHICommandBeginScene>
 	FORCEINLINE_DEBUGGABLE FRHICommandBeginScene()
 	{
 	}
-	void Execute()
-	{
-		BeginScene_Internal();
-	}
+	RHI_API void Execute();
 };
 
 struct FRHICommandEndScene : public FRHICommand<FRHICommandEndScene>
@@ -1126,10 +998,7 @@ struct FRHICommandEndScene : public FRHICommand<FRHICommandEndScene>
 	FORCEINLINE_DEBUGGABLE FRHICommandEndScene()
 	{
 	}
-	void Execute()
-	{
-		EndScene_Internal();
-	}
+	RHI_API void Execute();
 };
 
 struct FRHICommandUpdateVertexBuffer : public FRHICommand<FRHICommandUpdateVertexBuffer>
@@ -1144,13 +1013,7 @@ struct FRHICommandUpdateVertexBuffer : public FRHICommand<FRHICommandUpdateVerte
 		, BufferSize(InBufferSize)
 	{
 	}
-	void Execute()
-	{
-		void* Data = LockVertexBuffer_Internal(VertexBuffer, 0, BufferSize, RLM_WriteOnly);
-		FMemory::Memcpy(Data, Buffer, BufferSize);
-		UnlockVertexBuffer_Internal(VertexBuffer);
-		FMemory::Free((void*)Buffer);
-	}
+	RHI_API void Execute();
 };
 
 struct FRHICommandBeginFrame : public FRHICommand<FRHICommandBeginFrame>
@@ -1158,10 +1021,7 @@ struct FRHICommandBeginFrame : public FRHICommand<FRHICommandBeginFrame>
 	FORCEINLINE_DEBUGGABLE FRHICommandBeginFrame()
 	{
 	}
-	void Execute()
-	{
-		BeginFrame_Internal();
-	}
+	RHI_API void Execute();
 };
 
 struct FRHICommandEndFrame : public FRHICommand<FRHICommandEndFrame>
@@ -1169,12 +1029,38 @@ struct FRHICommandEndFrame : public FRHICommand<FRHICommandEndFrame>
 	FORCEINLINE_DEBUGGABLE FRHICommandEndFrame()
 	{
 	}
-	void Execute()
+	RHI_API void Execute();
+};
+
+#if PLATFORM_SUPPORTS_RHI_THREAD
+struct FRHICommandBeginDrawingViewport : public FRHICommand<FRHICommandBeginDrawingViewport>
+{
+	FViewportRHIParamRef Viewport;
+	FTextureRHIParamRef RenderTargetRHI;
+
+	FORCEINLINE_DEBUGGABLE FRHICommandBeginDrawingViewport(FViewportRHIParamRef InViewport, FTextureRHIParamRef InRenderTargetRHI)
+		: Viewport(InViewport)
+		, RenderTargetRHI(InRenderTargetRHI)
 	{
-		EndFrame_Internal();
 	}
+	RHI_API void Execute();
 };
 
+struct FRHICommandEndDrawingViewport : public FRHICommand<FRHICommandEndDrawingViewport>
+{
+	FViewportRHIParamRef Viewport;
+	bool bPresent;
+	bool bLockToVsync;
+
+	FORCEINLINE_DEBUGGABLE FRHICommandEndDrawingViewport(FViewportRHIParamRef InViewport, bool InbPresent, bool InbLockToVsync)
+		: Viewport(InViewport)
+		, bPresent(InbPresent)
+		, bLockToVsync(InbLockToVsync)
+	{
+	}
+	RHI_API void Execute();
+};
+#endif
 
 
 class RHI_API FRHICommandList : public FRHICommandListBase
@@ -1185,6 +1071,17 @@ public:
 	void* operator new(size_t Size);
 	void operator delete(void *RawMemory);
 
+#if USE_RHICOMMAND_STATE_REDUCTION
+	FORCEINLINE_DEBUGGABLE bool AllocStateCache()
+	{
+		if (!StateCache)
+		{
+			StateCache = new (Alloc<FRHICommandListStateCache>()) FRHICommandListStateCache();
+			return false;
+		}
+		return true;
+	}
+#endif
 	FORCEINLINE_DEBUGGABLE FLocalBoundShaderState BuildLocalBoundShaderState(const FBoundShaderStateInput& BoundShaderStateInput)
 	{
 		return BuildLocalBoundShaderState(
@@ -1238,6 +1135,12 @@ public:
 			SetBoundShaderState_Internal(LocalBoundShaderState.BypassBSS);
 			return;
 		}
+#if USE_RHICOMMAND_STATE_REDUCTION
+		if (StateCache)
+		{
+			StateCache->FlushShaderState();
+		}
+#endif
 		new (AllocCommand<FRHICommandSetLocalBoundShaderState>()) FRHICommandSetLocalBoundShaderState(this, LocalBoundShaderState);
 	}
 
@@ -1268,22 +1171,44 @@ public:
 			SetShaderUniformBuffer_Internal(Shader, BaseIndex, UniformBuffer.BypassUniform);
 			return;
 		}
+#if USE_RHICOMMAND_STATE_REDUCTION
+		if (StateCache)
+		{
+			//local uniform buffers are rare, so we will just flush
+			StateCache->BoundUniformBuffers[TRHIShaderToEnum<TShaderRHIParamRef>::ShaderFrequency][BaseIndex] = nullptr;
+		}
+#endif
 		new (AllocCommand<FRHICommandSetLocalUniformBuffer<TShaderRHIParamRef> >()) FRHICommandSetLocalUniformBuffer<TShaderRHIParamRef>(this, Shader, BaseIndex, UniformBuffer);
 	}
 
-	template <typename TShaderRHIParamRef>
-	FORCEINLINE_DEBUGGABLE void SetShaderUniformBuffer(TShaderRHIParamRef Shader, uint32 BaseIndex, FUniformBufferRHIParamRef UniformBuffer)
+	template <typename TShaderRHI>
+	FORCEINLINE_DEBUGGABLE void SetShaderUniformBuffer(TShaderRHI* Shader, uint32 BaseIndex, FUniformBufferRHIParamRef UniformBuffer)
 	{
 		if (Bypass())
 		{
 			SetShaderUniformBuffer_Internal(Shader, BaseIndex, UniformBuffer);
 			return;
 		}
-		new (AllocCommand<FRHICommandSetShaderUniformBuffer<TShaderRHIParamRef> >()) FRHICommandSetShaderUniformBuffer<TShaderRHIParamRef>(Shader, BaseIndex, UniformBuffer);
+#if USE_RHICOMMAND_STATE_REDUCTION
+		if (BaseIndex < FRHICommandListStateCache::MAX_UNIFORM_BUFFERS_PER_SHADER_STAGE)
+		{
+			if (AllocStateCache() && StateCache->BoundUniformBuffers[TRHIShaderToEnum<TShaderRHI*>::ShaderFrequency][BaseIndex] == UniformBuffer)
+			{
+				return;
+			}
+			StateCache->BoundUniformBuffers[TRHIShaderToEnum<TShaderRHI*>::ShaderFrequency][BaseIndex] = UniformBuffer;
+		}
+#endif
+		new (AllocCommand<FRHICommandSetShaderUniformBuffer<TShaderRHI*> >()) FRHICommandSetShaderUniformBuffer<TShaderRHI*>(Shader, BaseIndex, UniformBuffer);
+	}
+	template <typename TShaderRHI>
+	FORCEINLINE void SetShaderUniformBuffer(TRefCountPtr<TShaderRHI>& Shader, uint32 BaseIndex, FUniformBufferRHIParamRef UniformBuffer)
+	{
+		SetShaderUniformBuffer(Shader.GetReference(), BaseIndex, UniformBuffer);
 	}
 
-	template <typename TShaderRHIParamRef>
-	FORCEINLINE_DEBUGGABLE void SetShaderParameter(TShaderRHIParamRef Shader, uint32 BufferIndex, uint32 BaseIndex, uint32 NumBytes, const void* NewValue)
+	template <typename TShaderRHI>
+	FORCEINLINE_DEBUGGABLE void SetShaderParameter(TShaderRHI* Shader, uint32 BufferIndex, uint32 BaseIndex, uint32 NumBytes, const void* NewValue)
 	{
 		if (Bypass())
 		{
@@ -1292,7 +1217,12 @@ public:
 		}
 		void* UseValue = Alloc(NumBytes, 16);
 		FMemory::Memcpy(UseValue, NewValue, NumBytes);
-		new (AllocCommand<FRHICommandSetShaderParameter<TShaderRHIParamRef> >()) FRHICommandSetShaderParameter<TShaderRHIParamRef>(Shader, BufferIndex, BaseIndex, NumBytes, UseValue);
+		new (AllocCommand<FRHICommandSetShaderParameter<TShaderRHI*> >()) FRHICommandSetShaderParameter<TShaderRHI*>(Shader, BufferIndex, BaseIndex, NumBytes, UseValue);
+	}
+	template <typename TShaderRHI>
+	FORCEINLINE void SetShaderParameter(TRefCountPtr<TShaderRHI>& Shader, uint32 BufferIndex, uint32 BaseIndex, uint32 NumBytes, const void* NewValue)
+	{
+		SetShaderParameter(Shader.GetReference(), BufferIndex, BaseIndex, NumBytes, NewValue);
 	}
 
 	template <typename TShaderRHIParamRef>
@@ -1325,6 +1255,16 @@ public:
 			SetShaderSampler_Internal(Shader, SamplerIndex, State);
 			return;
 		}
+#if USE_RHICOMMAND_STATE_REDUCTION
+		if (SamplerIndex < FRHICommandListStateCache::MAX_SAMPLERS_PER_SHADER_STAGE)
+		{
+			if (AllocStateCache() && StateCache->Samplers[TRHIShaderToEnum<TShaderRHIParamRef>::ShaderFrequency][SamplerIndex] == State)
+			{
+				return;
+			}
+			StateCache->Samplers[TRHIShaderToEnum<TShaderRHIParamRef>::ShaderFrequency][SamplerIndex] = State;
+		}
+#endif
 		new (AllocCommand<FRHICommandSetShaderSampler<TShaderRHIParamRef> >()) FRHICommandSetShaderSampler<TShaderRHIParamRef>(Shader, SamplerIndex, State);
 	}
 
@@ -1355,6 +1295,12 @@ public:
 			SetBoundShaderState_Internal(BoundShaderState);
 			return;
 		}
+#if USE_RHICOMMAND_STATE_REDUCTION
+		if (StateCache)
+		{
+			StateCache->FlushShaderState();
+		}
+#endif
 		new (AllocCommand<FRHICommandSetBoundShaderState>()) FRHICommandSetBoundShaderState(BoundShaderState);
 	}
 
@@ -1544,6 +1490,12 @@ public:
 			SetComputeShader_Internal(ComputeShader);
 			return;
 		}
+#if USE_RHICOMMAND_STATE_REDUCTION
+		if (StateCache)
+		{
+			StateCache->FlushShaderState();
+		}
+#endif
 		new (AllocCommand<FRHICommandSetComputeShader>()) FRHICommandSetComputeShader(ComputeShader);
 	}
 
@@ -1825,9 +1777,6 @@ public:
 	static FGraphEventRef RHIThreadFence();
 	static void WaitOnRHIThreadFence(FGraphEventRef& Fence);
 
-	FORCEINLINE_DEBUGGABLE void Verify()
-	{
-	}
 	FORCEINLINE_DEBUGGABLE bool Bypass()
 	{
 #if !UE_BUILD_SHIPPING
diff --git a/Engine/Source/Runtime/RHI/Public/RHICommandListCommandExecutes.inl b/Engine/Source/Runtime/RHI/Public/RHICommandListCommandExecutes.inl
new file mode 100644
--- /dev/null
+++ b/Engine/Source/Runtime/RHI/Public/RHICommandListCommandExecutes.inl
@@ -0,0 +1,397 @@
+// Copyright 1998-2014 Epic Games, Inc. All Rights Reserved.
+
+/*=============================================================================
+	RHICommandListCommandExecutes.inl: RHI Command List execute functions.
+=============================================================================*/
+
+#if !defined(INTERNAL_DECORATOR)
+	#define INTERNAL_DECORATOR(Method) Method##_Internal
+#endif
+
+// set this one to get a stat for each RHI command 
+#define RHI_STATS 0
+
+#if RHI_STATS
+	DECLARE_STATS_GROUP(TEXT("RHICommands"),STATGROUP_RHI_COMMANDS, STATCAT_Advanced);
+	#define RHISTAT(Method)	DECLARE_SCOPE_CYCLE_COUNTER(TEXT(#Method), STAT_RHI##Method, STATGROUP_RHI_COMMANDS)
+#else
+	#define RHISTAT(Method)
+#endif
+
+void FRHICommandSetRasterizerState::Execute()
+{
+	RHISTAT(SetRasterizerState);
+	INTERNAL_DECORATOR(SetRasterizerState)(State);
+}
+
+void FRHICommandSetDepthStencilState::Execute()
+{
+	RHISTAT(SetDepthStencilState);
+	INTERNAL_DECORATOR(SetDepthStencilState)(State, StencilRef);
+}
+
+template <typename TShaderRHIParamRef>
+void FRHICommandSetShaderParameter<TShaderRHIParamRef>::Execute()
+{
+	RHISTAT(SetShaderParameter);
+	INTERNAL_DECORATOR(SetShaderParameter)(Shader, BufferIndex, BaseIndex, NumBytes, NewValue); 
+}
+template struct FRHICommandSetShaderParameter<FVertexShaderRHIParamRef>;
+template struct FRHICommandSetShaderParameter<FHullShaderRHIParamRef>;
+template struct FRHICommandSetShaderParameter<FDomainShaderRHIParamRef>;
+template struct FRHICommandSetShaderParameter<FGeometryShaderRHIParamRef>;
+template struct FRHICommandSetShaderParameter<FPixelShaderRHIParamRef>;
+template struct FRHICommandSetShaderParameter<FComputeShaderRHIParamRef>;
+
+template <typename TShaderRHIParamRef>
+void FRHICommandSetShaderUniformBuffer<TShaderRHIParamRef>::Execute()
+{
+	RHISTAT(SetShaderUniformBuffer);
+	INTERNAL_DECORATOR(SetShaderUniformBuffer)(Shader, BaseIndex, UniformBuffer);
+}
+template struct FRHICommandSetShaderUniformBuffer<FVertexShaderRHIParamRef>;
+template struct FRHICommandSetShaderUniformBuffer<FHullShaderRHIParamRef>;
+template struct FRHICommandSetShaderUniformBuffer<FDomainShaderRHIParamRef>;
+template struct FRHICommandSetShaderUniformBuffer<FGeometryShaderRHIParamRef>;
+template struct FRHICommandSetShaderUniformBuffer<FPixelShaderRHIParamRef>;
+template struct FRHICommandSetShaderUniformBuffer<FComputeShaderRHIParamRef>;
+
+template <typename TShaderRHIParamRef>
+void FRHICommandSetShaderTexture<TShaderRHIParamRef>::Execute()
+{
+	RHISTAT(SetShaderTexture);
+	INTERNAL_DECORATOR(SetShaderTexture)(Shader, TextureIndex, Texture);
+}
+template struct FRHICommandSetShaderTexture<FVertexShaderRHIParamRef>;
+template struct FRHICommandSetShaderTexture<FHullShaderRHIParamRef>;
+template struct FRHICommandSetShaderTexture<FDomainShaderRHIParamRef>;
+template struct FRHICommandSetShaderTexture<FGeometryShaderRHIParamRef>;
+template struct FRHICommandSetShaderTexture<FPixelShaderRHIParamRef>;
+template struct FRHICommandSetShaderTexture<FComputeShaderRHIParamRef>;
+
+template <typename TShaderRHIParamRef>
+void FRHICommandSetShaderResourceViewParameter<TShaderRHIParamRef>::Execute()
+{
+	RHISTAT(SetShaderResourceViewParameter);
+	INTERNAL_DECORATOR(SetShaderResourceViewParameter)(Shader, SamplerIndex, SRV);
+}
+template struct FRHICommandSetShaderResourceViewParameter<FVertexShaderRHIParamRef>;
+template struct FRHICommandSetShaderResourceViewParameter<FHullShaderRHIParamRef>;
+template struct FRHICommandSetShaderResourceViewParameter<FDomainShaderRHIParamRef>;
+template struct FRHICommandSetShaderResourceViewParameter<FGeometryShaderRHIParamRef>;
+template struct FRHICommandSetShaderResourceViewParameter<FPixelShaderRHIParamRef>;
+template struct FRHICommandSetShaderResourceViewParameter<FComputeShaderRHIParamRef>;
+
+template <typename TShaderRHIParamRef>
+void FRHICommandSetUAVParameter<TShaderRHIParamRef>::Execute()
+{
+	RHISTAT(SetUAVParameter);
+	INTERNAL_DECORATOR(SetUAVParameter)(Shader, UAVIndex, UAV);
+}
+template struct FRHICommandSetUAVParameter<FComputeShaderRHIParamRef>;
+
+template <typename TShaderRHIParamRef>
+void FRHICommandSetUAVParameter_IntialCount<TShaderRHIParamRef>::Execute()
+{
+	RHISTAT(SetUAVParameter);
+	INTERNAL_DECORATOR(SetUAVParameter)(Shader, UAVIndex, UAV, InitialCount);
+}
+template struct FRHICommandSetUAVParameter_IntialCount<FComputeShaderRHIParamRef>;
+
+template <typename TShaderRHIParamRef>
+void FRHICommandSetShaderSampler<TShaderRHIParamRef>::Execute()
+{
+	RHISTAT(SetShaderSampler);
+	INTERNAL_DECORATOR(SetShaderSampler)(Shader, SamplerIndex, Sampler);
+}
+template struct FRHICommandSetShaderSampler<FVertexShaderRHIParamRef>;
+template struct FRHICommandSetShaderSampler<FHullShaderRHIParamRef>;
+template struct FRHICommandSetShaderSampler<FDomainShaderRHIParamRef>;
+template struct FRHICommandSetShaderSampler<FGeometryShaderRHIParamRef>;
+template struct FRHICommandSetShaderSampler<FPixelShaderRHIParamRef>;
+template struct FRHICommandSetShaderSampler<FComputeShaderRHIParamRef>;
+
+void FRHICommandDrawPrimitive::Execute()
+{
+	RHISTAT(DrawPrimitive);
+	INTERNAL_DECORATOR(DrawPrimitive)(PrimitiveType, BaseVertexIndex, NumPrimitives, NumInstances);
+}
+
+void FRHICommandDrawIndexedPrimitive::Execute()
+{
+	RHISTAT(DrawIndexedPrimitive);
+	INTERNAL_DECORATOR(DrawIndexedPrimitive)(IndexBuffer, PrimitiveType, BaseVertexIndex, MinIndex, NumVertices, StartIndex, NumPrimitives, NumInstances);
+}
+
+void FRHICommandSetBoundShaderState::Execute()
+{
+	RHISTAT(SetBoundShaderState);
+	INTERNAL_DECORATOR(SetBoundShaderState)(BoundShaderState);
+}
+
+void FRHICommandSetBlendState::Execute()
+{
+	RHISTAT(SetBlendState);
+	INTERNAL_DECORATOR(SetBlendState)(State, BlendFactor);
+}
+
+void FRHICommandSetStreamSource::Execute()
+{
+	RHISTAT(SetStreamSource);
+	INTERNAL_DECORATOR(SetStreamSource)(StreamIndex, VertexBuffer, Stride, Offset);
+}
+
+void FRHICommandSetViewport::Execute()
+{
+	RHISTAT(SetViewport);
+	INTERNAL_DECORATOR(SetViewport)(MinX, MinY, MinZ, MaxX, MaxY, MaxZ);
+}
+
+void FRHICommandSetScissorRect::Execute()
+{
+	RHISTAT(SetScissorRect);
+	INTERNAL_DECORATOR(SetScissorRect)(bEnable, MinX, MinY, MaxX, MaxY);
+}
+
+void FRHICommandSetRenderTargets::Execute()
+{
+	RHISTAT(SetRenderTargets);
+	INTERNAL_DECORATOR(SetRenderTargets)(
+		NewNumSimultaneousRenderTargets,
+		NewRenderTargetsRHI,
+		NewDepthStencilTargetRHI,
+		NewNumUAVs,
+		UAVs);
+}
+
+void FRHICommandEndDrawPrimitiveUP::Execute()
+{
+	RHISTAT(EndDrawPrimitiveUP);
+	void* Buffer = NULL;
+	INTERNAL_DECORATOR(BeginDrawPrimitiveUP)(PrimitiveType, NumPrimitives, NumVertices, VertexDataStride, Buffer);
+	FMemory::Memcpy(Buffer, OutVertexData, NumVertices * VertexDataStride);
+	INTERNAL_DECORATOR(EndDrawPrimitiveUP)();
+}
+
+void FRHICommandEndDrawIndexedPrimitiveUP::Execute()
+{
+	RHISTAT(EndDrawIndexedPrimitiveUP);
+	void* VertexBuffer = nullptr;
+	void* IndexBuffer = nullptr;
+	INTERNAL_DECORATOR(BeginDrawIndexedPrimitiveUP)(
+		PrimitiveType,
+		NumPrimitives,
+		NumVertices,
+		VertexDataStride,
+		VertexBuffer,
+		MinVertexIndex,
+		NumIndices,
+		IndexDataStride,
+		IndexBuffer);
+	FMemory::Memcpy(VertexBuffer, OutVertexData, NumVertices * VertexDataStride);
+	FMemory::Memcpy(IndexBuffer, OutIndexData, NumIndices * IndexDataStride);
+	INTERNAL_DECORATOR(EndDrawIndexedPrimitiveUP)();
+}
+
+void FRHICommandSetComputeShader::Execute()
+{
+	RHISTAT(SetComputeShader);
+	INTERNAL_DECORATOR(SetComputeShader)(ComputeShader);
+}
+
+void FRHICommandDispatchComputeShader::Execute()
+{
+	RHISTAT(DispatchComputeShader);
+	INTERNAL_DECORATOR(DispatchComputeShader)(ThreadGroupCountX, ThreadGroupCountY, ThreadGroupCountZ);
+}
+
+void FRHICommandDispatchIndirectComputeShader::Execute()
+{
+	RHISTAT(DispatchIndirectComputeShader);
+	INTERNAL_DECORATOR(DispatchIndirectComputeShader)(ArgumentBuffer, ArgumentOffset);
+}
+
+void FRHICommandAutomaticCacheFlushAfterComputeShader::Execute()
+{
+	RHISTAT(AutomaticCacheFlushAfterComputeShader);
+	INTERNAL_DECORATOR(AutomaticCacheFlushAfterComputeShader)(bEnable);
+}
+
+void FRHICommandFlushComputeShaderCache::Execute()
+{
+	RHISTAT(FlushComputeShaderCache);
+	INTERNAL_DECORATOR(FlushComputeShaderCache)();
+}
+
+void FRHICommandDrawPrimitiveIndirect::Execute()
+{
+	RHISTAT(DrawPrimitiveIndirect);
+	INTERNAL_DECORATOR(DrawPrimitiveIndirect)(PrimitiveType, ArgumentBuffer, ArgumentOffset);
+}
+
+void FRHICommandDrawIndexedIndirect::Execute()
+{
+	RHISTAT(DrawIndexedIndirect);
+	INTERNAL_DECORATOR(DrawIndexedIndirect)(IndexBufferRHI, PrimitiveType, ArgumentsBufferRHI, DrawArgumentsIndex, NumInstances);
+}
+
+void FRHICommandDrawIndexedPrimitiveIndirect::Execute()
+{
+	RHISTAT(DrawIndexedPrimitiveIndirect);
+	INTERNAL_DECORATOR(DrawIndexedPrimitiveIndirect)(PrimitiveType, IndexBuffer, ArgumentsBuffer, ArgumentOffset);
+}
+
+void FRHICommandEnableDepthBoundsTest::Execute()
+{
+	RHISTAT(EnableDepthBoundsTest);
+	INTERNAL_DECORATOR(EnableDepthBoundsTest)(bEnable, MinDepth, MaxDepth);
+}
+
+void FRHICommandClearUAV::Execute()
+{
+	RHISTAT(ClearUAV);
+	INTERNAL_DECORATOR(ClearUAV)(UnorderedAccessViewRHI, Values);
+}
+
+void FRHICommandCopyToResolveTarget::Execute()
+{
+	RHISTAT(CopyToResolveTarget);
+	INTERNAL_DECORATOR(CopyToResolveTarget)(SourceTexture, DestTexture, bKeepOriginalSurface, ResolveParams);
+}
+
+void FRHICommandClear::Execute()
+{
+	RHISTAT(Clear);
+	INTERNAL_DECORATOR(Clear)(bClearColor, Color, bClearDepth, Depth, bClearStencil, Stencil, ExcludeRect);
+}
+
+void FRHICommandClearMRT::Execute()
+{
+	RHISTAT(ClearMRT);
+	INTERNAL_DECORATOR(ClearMRT)(bClearColor, NumClearColors, ColorArray, bClearDepth, Depth, bClearStencil, Stencil, ExcludeRect);
+}
+
+void FRHICommandBuildLocalBoundShaderState::Execute()
+{
+	RHISTAT(BuildLocalBoundShaderState);
+	check(!IsValidRef(WorkArea.ComputedBSS->BSS)); // should not already have been created
+	if (WorkArea.ComputedBSS->UseCount)
+	{
+		WorkArea.ComputedBSS->BSS = 	
+			INTERNAL_DECORATOR(CreateBoundShaderState)(WorkArea.Args.VertexDeclarationRHI, WorkArea.Args.VertexShaderRHI, WorkArea.Args.HullShaderRHI, WorkArea.Args.DomainShaderRHI, WorkArea.Args.PixelShaderRHI, WorkArea.Args.GeometryShaderRHI);
+	}
+}
+
+void FRHICommandSetLocalBoundShaderState::Execute()
+{
+	RHISTAT(SetLocalBoundShaderState);
+	check(LocalBoundShaderState.WorkArea->ComputedBSS->UseCount > 0 && IsValidRef(LocalBoundShaderState.WorkArea->ComputedBSS->BSS)); // this should have been created and should have uses outstanding
+
+	INTERNAL_DECORATOR(SetBoundShaderState)(LocalBoundShaderState.WorkArea->ComputedBSS->BSS);
+
+	if (--LocalBoundShaderState.WorkArea->ComputedBSS->UseCount == 0)
+	{
+		LocalBoundShaderState.WorkArea->ComputedBSS->~FComputedBSS();
+	}
+}
+
+void FRHICommandBuildLocalUniformBuffer::Execute()
+{
+	RHISTAT(BuildLocalUniformBuffer);
+	check(!IsValidRef(WorkArea.ComputedUniformBuffer->UniformBuffer) && WorkArea.Layout && WorkArea.Contents); // should not already have been created
+	if (WorkArea.ComputedUniformBuffer->UseCount)
+	{
+#if PLATFORM_SUPPORTS_RHI_THREAD
+		WorkArea.ComputedUniformBuffer->UniformBuffer = RHICreateUniformBuffer(WorkArea.Contents, *WorkArea.Layout, UniformBuffer_SingleFrame); 
+#else
+		WorkArea.ComputedUniformBuffer->UniformBuffer = INTERNAL_DECORATOR(CreateUniformBuffer)(WorkArea.Contents, *WorkArea.Layout, UniformBuffer_SingleFrame);
+#endif
+	}
+	WorkArea.Layout = nullptr;
+	WorkArea.Contents = nullptr;
+}
+
+template <typename TShaderRHIParamRef>
+void FRHICommandSetLocalUniformBuffer<TShaderRHIParamRef>::Execute()
+{
+	RHISTAT(SetLocalUniformBuffer);
+	check(LocalUniformBuffer.WorkArea->ComputedUniformBuffer->UseCount > 0 && IsValidRef(LocalUniformBuffer.WorkArea->ComputedUniformBuffer->UniformBuffer)); // this should have been created and should have uses outstanding
+	INTERNAL_DECORATOR(SetShaderUniformBuffer)(Shader, BaseIndex, LocalUniformBuffer.WorkArea->ComputedUniformBuffer->UniformBuffer);
+	if (--LocalUniformBuffer.WorkArea->ComputedUniformBuffer->UseCount == 0)
+	{
+		LocalUniformBuffer.WorkArea->ComputedUniformBuffer->~FComputedUniformBuffer();
+	}
+}
+template struct FRHICommandSetLocalUniformBuffer<FVertexShaderRHIParamRef>;
+template struct FRHICommandSetLocalUniformBuffer<FHullShaderRHIParamRef>;
+template struct FRHICommandSetLocalUniformBuffer<FDomainShaderRHIParamRef>;
+template struct FRHICommandSetLocalUniformBuffer<FGeometryShaderRHIParamRef>;
+template struct FRHICommandSetLocalUniformBuffer<FPixelShaderRHIParamRef>;
+template struct FRHICommandSetLocalUniformBuffer<FComputeShaderRHIParamRef>;
+
+void FRHICommandBeginRenderQuery::Execute()
+{
+	RHISTAT(BeginRenderQuery);
+	INTERNAL_DECORATOR(BeginRenderQuery)(RenderQuery);
+}
+
+void FRHICommandEndRenderQuery::Execute()
+{
+	RHISTAT(EndRenderQuery);
+	INTERNAL_DECORATOR(EndRenderQuery)(RenderQuery);
+}
+
+#if !PLATFORM_SUPPORTS_RHI_THREAD
+void FRHICommandResetRenderQuery::Execute()
+{
+	RHISTAT(ResetRenderQuery);
+	INTERNAL_DECORATOR(ResetRenderQuery)(RenderQuery);
+}
+#endif
+
+void FRHICommandBeginScene::Execute()
+{
+	RHISTAT(BeginScene);
+	INTERNAL_DECORATOR(BeginScene)();
+}
+
+void FRHICommandEndScene::Execute()
+{
+	RHISTAT(EndScene);
+	INTERNAL_DECORATOR(EndScene)();
+}
+
+void FRHICommandUpdateVertexBuffer::Execute()
+{
+	RHISTAT(UpdateVertexBuffer);
+	void* Data = INTERNAL_DECORATOR(LockVertexBuffer)(VertexBuffer, 0, BufferSize, RLM_WriteOnly);
+	FMemory::Memcpy(Data, Buffer, BufferSize);
+	INTERNAL_DECORATOR(UnlockVertexBuffer)(VertexBuffer);
+	FMemory::Free((void*)Buffer);
+}
+
+void FRHICommandBeginFrame::Execute()
+{
+	RHISTAT(BeginFrame);
+	INTERNAL_DECORATOR(BeginFrame)();
+}
+
+void FRHICommandEndFrame::Execute()
+{
+	RHISTAT(EndFrame);
+	INTERNAL_DECORATOR(EndFrame)();
+}
+
+#if PLATFORM_SUPPORTS_RHI_THREAD
+void FRHICommandBeginDrawingViewport::Execute()
+{
+	RHISTAT(BeginDrawingViewport);
+	INTERNAL_DECORATOR(BeginDrawingViewport)(Viewport, RenderTargetRHI);
+}
+
+void FRHICommandEndDrawingViewport::Execute()
+{
+	RHISTAT(EndDrawingViewport);
+	INTERNAL_DECORATOR(EndDrawingViewport)(Viewport, bPresent, bLockToVsync);
+}
+#endif
\ No newline at end of file
diff --git a/Engine/Source/Runtime/Windows/D3D11RHI/Private/D3D11Commands.cpp b/Engine/Source/Runtime/Windows/D3D11RHI/Private/D3D11Commands.cpp
--- a/Engine/Source/Runtime/Windows/D3D11RHI/Private/D3D11Commands.cpp
+++ b/Engine/Source/Runtime/Windows/D3D11RHI/Private/D3D11Commands.cpp
@@ -139,8 +139,6 @@ void FD3D11DynamicRHI::RHIGpuTimeEnd(uint32 Hash, bool bCompute)
 // Vertex state.
 void FD3D11DynamicRHI::RHISetStreamSource(uint32 StreamIndex,FVertexBufferRHIParamRef VertexBufferRHI,uint32 Stride,uint32 Offset)
 {
-	GRHICommandList.Verify();
-
 	DYNAMIC_CAST_D3D11RESOURCE(VertexBuffer,VertexBuffer);
 
 	ID3D11Buffer* D3DBuffer = VertexBuffer ? VertexBuffer->Resource : NULL;
@@ -149,8 +147,6 @@ void FD3D11DynamicRHI::RHISetStreamSource(uint32 StreamIndex,FVertexBufferRHIPar
 
 void FD3D11DynamicRHI::RHISetStreamOutTargets(uint32 NumTargets, const FVertexBufferRHIParamRef* VertexBuffers, const uint32* Offsets)
 {
-	GRHICommandList.Verify();
-
 	ID3D11Buffer* D3DVertexBuffers[D3D11_SO_BUFFER_SLOT_COUNT] = {0};
 
 	if (VertexBuffers)
@@ -167,24 +163,18 @@ void FD3D11DynamicRHI::RHISetStreamOutTargets(uint32 NumTargets, const FVertexBu
 // Rasterizer state.
 void FD3D11DynamicRHI::RHISetRasterizerState(FRasterizerStateRHIParamRef NewStateRHI)
 {
-	GRHICommandList.Verify();
-
 	DYNAMIC_CAST_D3D11RESOURCE(RasterizerState,NewState);
 	StateCache.SetRasterizerState(NewState->Resource);
 }
 
 void FD3D11DynamicRHI::RHISetComputeShader(FComputeShaderRHIParamRef ComputeShaderRHI)
 {
-	GRHICommandList.Verify();
-
 	DYNAMIC_CAST_D3D11RESOURCE(ComputeShader,ComputeShader);
 	SetCurrentComputeShader(ComputeShaderRHI);
 }
 
 void FD3D11DynamicRHI::RHIDispatchComputeShader(uint32 ThreadGroupCountX, uint32 ThreadGroupCountY, uint32 ThreadGroupCountZ) 
 { 
-	GRHICommandList.Verify();
-
 	FComputeShaderRHIParamRef ComputeShaderRHI = GetCurrentComputeShader();
 	DYNAMIC_CAST_D3D11RESOURCE(ComputeShader,ComputeShader);
 
@@ -204,8 +194,6 @@ void FD3D11DynamicRHI::RHIDispatchComputeShader(uint32 ThreadGroupCountX, uint32
 
 void FD3D11DynamicRHI::RHIDispatchIndirectComputeShader(FVertexBufferRHIParamRef ArgumentBufferRHI, uint32 ArgumentOffset) 
 { 
-	GRHICommandList.Verify();
-
 	FComputeShaderRHIParamRef ComputeShaderRHI = GetCurrentComputeShader();
 	DYNAMIC_CAST_D3D11RESOURCE(ComputeShader,ComputeShader);
 	DYNAMIC_CAST_D3D11RESOURCE(VertexBuffer,ArgumentBuffer);
@@ -226,8 +214,6 @@ void FD3D11DynamicRHI::RHIDispatchIndirectComputeShader(FVertexBufferRHIParamRef
 
 void FD3D11DynamicRHI::RHISetViewport(uint32 MinX,uint32 MinY,float MinZ,uint32 MaxX,uint32 MaxY,float MaxZ)
 {
-	GRHICommandList.Verify();
-
 	// These are the maximum viewport extents for D3D11. Exceeding them leads to badness.
 	check(MinX <= (uint32)D3D11_VIEWPORT_BOUNDS_MAX);
 	check(MinY <= (uint32)D3D11_VIEWPORT_BOUNDS_MAX);
@@ -245,8 +231,6 @@ void FD3D11DynamicRHI::RHISetViewport(uint32 MinX,uint32 MinY,float MinZ,uint32
 
 void FD3D11DynamicRHI::RHISetScissorRect(bool bEnable,uint32 MinX,uint32 MinY,uint32 MaxX,uint32 MaxY)
 {
-	GRHICommandList.Verify();
-
 	if(bEnable)
 	{
 		D3D11_RECT ScissorRect;
@@ -273,8 +257,6 @@ void FD3D11DynamicRHI::RHISetScissorRect(bool bEnable,uint32 MinX,uint32 MinY,ui
 */
 void FD3D11DynamicRHI::RHISetBoundShaderState( FBoundShaderStateRHIParamRef BoundShaderStateRHI)
 {
-	GRHICommandList.Verify();
-
 	DYNAMIC_CAST_D3D11RESOURCE(BoundShaderState,BoundShaderState);
 
 	StateCache.SetInputLayout(BoundShaderState->InputLayout);
@@ -311,8 +293,6 @@ void FD3D11DynamicRHI::RHISetBoundShaderState( FBoundShaderStateRHIParamRef Boun
 
 void FD3D11DynamicRHI::RHISetShaderTexture(FVertexShaderRHIParamRef VertexShaderRHI,uint32 TextureIndex,FTextureRHIParamRef NewTextureRHI)
 {
-	GRHICommandList.Verify();
-
 	VALIDATE_BOUND_SHADER(VertexShaderRHI);
 
 	FD3D11TextureBase* NewTexture = GetD3D11TextureFromRHITexture(NewTextureRHI);
@@ -326,8 +306,6 @@ void FD3D11DynamicRHI::RHISetShaderTexture(FVertexShaderRHIParamRef VertexShader
 
 void FD3D11DynamicRHI::RHISetShaderTexture(FHullShaderRHIParamRef HullShaderRHI,uint32 TextureIndex,FTextureRHIParamRef NewTextureRHI)
 {
-	GRHICommandList.Verify();
-
 	VALIDATE_BOUND_SHADER(HullShaderRHI);
 
 	FD3D11TextureBase* NewTexture = GetD3D11TextureFromRHITexture(NewTextureRHI);
@@ -341,8 +319,6 @@ void FD3D11DynamicRHI::RHISetShaderTexture(FHullShaderRHIParamRef HullShaderRHI,
 
 void FD3D11DynamicRHI::RHISetShaderTexture(FDomainShaderRHIParamRef DomainShaderRHI,uint32 TextureIndex,FTextureRHIParamRef NewTextureRHI)
 {
-	GRHICommandList.Verify();
-
 	VALIDATE_BOUND_SHADER(DomainShaderRHI);
 
 	FD3D11TextureBase* NewTexture = GetD3D11TextureFromRHITexture(NewTextureRHI);
@@ -356,8 +332,6 @@ void FD3D11DynamicRHI::RHISetShaderTexture(FDomainShaderRHIParamRef DomainShader
 
 void FD3D11DynamicRHI::RHISetShaderTexture(FGeometryShaderRHIParamRef GeometryShaderRHI,uint32 TextureIndex,FTextureRHIParamRef NewTextureRHI)
 {
-	GRHICommandList.Verify();
-
 	VALIDATE_BOUND_SHADER(GeometryShaderRHI);
 
 	FD3D11TextureBase* NewTexture = GetD3D11TextureFromRHITexture(NewTextureRHI);
@@ -372,8 +346,6 @@ void FD3D11DynamicRHI::RHISetShaderTexture(FGeometryShaderRHIParamRef GeometrySh
 void FD3D11DynamicRHI::RHISetShaderTexture(FPixelShaderRHIParamRef PixelShaderRHI,uint32 TextureIndex,FTextureRHIParamRef NewTextureRHI)
 {
 
-	GRHICommandList.Verify();
-
 	VALIDATE_BOUND_SHADER(PixelShaderRHI);
 
 	FD3D11TextureBase* NewTexture = GetD3D11TextureFromRHITexture(NewTextureRHI);
@@ -386,7 +358,6 @@ void FD3D11DynamicRHI::RHISetShaderTexture(FPixelShaderRHIParamRef PixelShaderRH
 
 void FD3D11DynamicRHI::RHISetShaderTexture(FComputeShaderRHIParamRef ComputeShaderRHI,uint32 TextureIndex,FTextureRHIParamRef NewTextureRHI)
 {
-	GRHICommandList.Verify();
 	//VALIDATE_BOUND_SHADER(ComputeShaderRHI);
 
 	FD3D11TextureBase* NewTexture = GetD3D11TextureFromRHITexture(NewTextureRHI);
@@ -400,7 +371,6 @@ void FD3D11DynamicRHI::RHISetShaderTexture(FComputeShaderRHIParamRef ComputeShad
 
 void FD3D11DynamicRHI::RHISetUAVParameter(FComputeShaderRHIParamRef ComputeShaderRHI,uint32 UAVIndex,FUnorderedAccessViewRHIParamRef UAVRHI)
 {
-	GRHICommandList.Verify();
 	//VALIDATE_BOUND_SHADER(ComputeShaderRHI);
 
 	DYNAMIC_CAST_D3D11RESOURCE(UnorderedAccessView,UAV);
@@ -418,7 +388,6 @@ void FD3D11DynamicRHI::RHISetUAVParameter(FComputeShaderRHIParamRef ComputeShade
 
 void FD3D11DynamicRHI::RHISetUAVParameter(FComputeShaderRHIParamRef ComputeShaderRHI,uint32 UAVIndex,FUnorderedAccessViewRHIParamRef UAVRHI, uint32 InitialCount )
 {
-	GRHICommandList.Verify();
 	//VALIDATE_BOUND_SHADER(ComputeShaderRHI);
 
 	DYNAMIC_CAST_D3D11RESOURCE(UnorderedAccessView,UAV);
@@ -434,8 +403,6 @@ void FD3D11DynamicRHI::RHISetUAVParameter(FComputeShaderRHIParamRef ComputeShade
 
 void FD3D11DynamicRHI::RHISetShaderResourceViewParameter(FPixelShaderRHIParamRef PixelShaderRHI,uint32 TextureIndex,FShaderResourceViewRHIParamRef SRVRHI)
 {
-	GRHICommandList.Verify();
-
 	VALIDATE_BOUND_SHADER(PixelShaderRHI);
 
 	DYNAMIC_CAST_D3D11RESOURCE(ShaderResourceView,SRV);
@@ -454,8 +421,6 @@ void FD3D11DynamicRHI::RHISetShaderResourceViewParameter(FPixelShaderRHIParamRef
 
 void FD3D11DynamicRHI::RHISetShaderResourceViewParameter(FVertexShaderRHIParamRef VertexShaderRHI,uint32 TextureIndex,FShaderResourceViewRHIParamRef SRVRHI)
 {
-	GRHICommandList.Verify();
-
 	VALIDATE_BOUND_SHADER(VertexShaderRHI);
 
 	DYNAMIC_CAST_D3D11RESOURCE(ShaderResourceView,SRV);
@@ -474,8 +439,6 @@ void FD3D11DynamicRHI::RHISetShaderResourceViewParameter(FVertexShaderRHIParamRe
 
 void FD3D11DynamicRHI::RHISetShaderResourceViewParameter(FComputeShaderRHIParamRef ComputeShaderRHI,uint32 TextureIndex,FShaderResourceViewRHIParamRef SRVRHI)
 {
-	GRHICommandList.Verify();
-
 	//VALIDATE_BOUND_SHADER(ComputeShaderRHI);
 
 	DYNAMIC_CAST_D3D11RESOURCE(ShaderResourceView,SRV);
@@ -494,8 +457,6 @@ void FD3D11DynamicRHI::RHISetShaderResourceViewParameter(FComputeShaderRHIParamR
 
 void FD3D11DynamicRHI::RHISetShaderResourceViewParameter(FHullShaderRHIParamRef HullShaderRHI,uint32 TextureIndex,FShaderResourceViewRHIParamRef SRVRHI)
 {
-	GRHICommandList.Verify();
-
 	VALIDATE_BOUND_SHADER(HullShaderRHI);
 
 	DYNAMIC_CAST_D3D11RESOURCE(ShaderResourceView,SRV);
@@ -514,8 +475,6 @@ void FD3D11DynamicRHI::RHISetShaderResourceViewParameter(FHullShaderRHIParamRef
 
 void FD3D11DynamicRHI::RHISetShaderResourceViewParameter(FDomainShaderRHIParamRef DomainShaderRHI,uint32 TextureIndex,FShaderResourceViewRHIParamRef SRVRHI)
 {
-	GRHICommandList.Verify();
-
 	VALIDATE_BOUND_SHADER(DomainShaderRHI);
 
 	DYNAMIC_CAST_D3D11RESOURCE(ShaderResourceView,SRV);
@@ -534,8 +493,6 @@ void FD3D11DynamicRHI::RHISetShaderResourceViewParameter(FDomainShaderRHIParamRe
 
 void FD3D11DynamicRHI::RHISetShaderResourceViewParameter(FGeometryShaderRHIParamRef GeometryShaderRHI,uint32 TextureIndex,FShaderResourceViewRHIParamRef SRVRHI)
 {
-	GRHICommandList.Verify();
-
 	VALIDATE_BOUND_SHADER(GeometryShaderRHI);
 
 	DYNAMIC_CAST_D3D11RESOURCE(ShaderResourceView,SRV);
@@ -554,8 +511,6 @@ void FD3D11DynamicRHI::RHISetShaderResourceViewParameter(FGeometryShaderRHIParam
 
 void FD3D11DynamicRHI::RHISetShaderSampler(FVertexShaderRHIParamRef VertexShaderRHI,uint32 SamplerIndex,FSamplerStateRHIParamRef NewStateRHI)
 {
-	GRHICommandList.Verify();
-
 	VALIDATE_BOUND_SHADER(VertexShaderRHI);
 
 	DYNAMIC_CAST_D3D11RESOURCE(VertexShader,VertexShader);
@@ -567,8 +522,6 @@ void FD3D11DynamicRHI::RHISetShaderSampler(FVertexShaderRHIParamRef VertexShader
 
 void FD3D11DynamicRHI::RHISetShaderSampler(FHullShaderRHIParamRef HullShaderRHI,uint32 SamplerIndex,FSamplerStateRHIParamRef NewStateRHI)
 {
-	GRHICommandList.Verify();
-
 	VALIDATE_BOUND_SHADER(HullShaderRHI);
 
 	DYNAMIC_CAST_D3D11RESOURCE(HullShader,HullShader);
@@ -580,8 +533,6 @@ void FD3D11DynamicRHI::RHISetShaderSampler(FHullShaderRHIParamRef HullShaderRHI,
 
 void FD3D11DynamicRHI::RHISetShaderSampler(FDomainShaderRHIParamRef DomainShaderRHI,uint32 SamplerIndex,FSamplerStateRHIParamRef NewStateRHI)
 {
-	GRHICommandList.Verify();
-
 	VALIDATE_BOUND_SHADER(DomainShaderRHI);
 
 	DYNAMIC_CAST_D3D11RESOURCE(DomainShader,DomainShader);
@@ -593,8 +544,6 @@ void FD3D11DynamicRHI::RHISetShaderSampler(FDomainShaderRHIParamRef DomainShader
 
 void FD3D11DynamicRHI::RHISetShaderSampler(FGeometryShaderRHIParamRef GeometryShaderRHI,uint32 SamplerIndex,FSamplerStateRHIParamRef NewStateRHI)
 {
-	GRHICommandList.Verify();
-
 	VALIDATE_BOUND_SHADER(GeometryShaderRHI);
 
 	DYNAMIC_CAST_D3D11RESOURCE(GeometryShader,GeometryShader);
@@ -606,8 +555,6 @@ void FD3D11DynamicRHI::RHISetShaderSampler(FGeometryShaderRHIParamRef GeometrySh
 
 void FD3D11DynamicRHI::RHISetShaderSampler(FPixelShaderRHIParamRef PixelShaderRHI,uint32 SamplerIndex,FSamplerStateRHIParamRef NewStateRHI)
 {
-	GRHICommandList.Verify();
-
 	VALIDATE_BOUND_SHADER(PixelShaderRHI);
 
 	DYNAMIC_CAST_D3D11RESOURCE(PixelShader,PixelShader);
@@ -619,8 +566,6 @@ void FD3D11DynamicRHI::RHISetShaderSampler(FPixelShaderRHIParamRef PixelShaderRH
 
 void FD3D11DynamicRHI::RHISetShaderSampler(FComputeShaderRHIParamRef ComputeShaderRHI,uint32 SamplerIndex,FSamplerStateRHIParamRef NewStateRHI)
 {
-	GRHICommandList.Verify();
-
 	//VALIDATE_BOUND_SHADER(ComputeShaderRHI);
 	DYNAMIC_CAST_D3D11RESOURCE(ComputeShader,ComputeShader);
 	DYNAMIC_CAST_D3D11RESOURCE(SamplerState,NewState);
@@ -631,8 +576,6 @@ void FD3D11DynamicRHI::RHISetShaderSampler(FComputeShaderRHIParamRef ComputeShad
 
 void FD3D11DynamicRHI::RHISetShaderUniformBuffer(FVertexShaderRHIParamRef VertexShader,uint32 BufferIndex,FUniformBufferRHIParamRef BufferRHI)
 {
-	GRHICommandList.Verify();
-
 	VALIDATE_BOUND_SHADER(VertexShader);
 	DYNAMIC_CAST_D3D11RESOURCE(UniformBuffer,Buffer);
 #if PLATFORM_XBOXONE
@@ -653,8 +596,6 @@ void FD3D11DynamicRHI::RHISetShaderUniformBuffer(FVertexShaderRHIParamRef Vertex
 
 void FD3D11DynamicRHI::RHISetShaderUniformBuffer(FHullShaderRHIParamRef HullShader,uint32 BufferIndex,FUniformBufferRHIParamRef BufferRHI)
 {
-	GRHICommandList.Verify();
-
 	VALIDATE_BOUND_SHADER(HullShader);
 	DYNAMIC_CAST_D3D11RESOURCE(UniformBuffer,Buffer);
 #if PLATFORM_XBOXONE
@@ -675,8 +616,6 @@ void FD3D11DynamicRHI::RHISetShaderUniformBuffer(FHullShaderRHIParamRef HullShad
 
 void FD3D11DynamicRHI::RHISetShaderUniformBuffer(FDomainShaderRHIParamRef DomainShader,uint32 BufferIndex,FUniformBufferRHIParamRef BufferRHI)
 {
-	GRHICommandList.Verify();
-
 	VALIDATE_BOUND_SHADER(DomainShader);
 	DYNAMIC_CAST_D3D11RESOURCE(UniformBuffer,Buffer);
 #if PLATFORM_XBOXONE
@@ -697,8 +636,6 @@ void FD3D11DynamicRHI::RHISetShaderUniformBuffer(FDomainShaderRHIParamRef Domain
 
 void FD3D11DynamicRHI::RHISetShaderUniformBuffer(FGeometryShaderRHIParamRef GeometryShader,uint32 BufferIndex,FUniformBufferRHIParamRef BufferRHI)
 {
-	GRHICommandList.Verify();
-
 	VALIDATE_BOUND_SHADER(GeometryShader);
 	DYNAMIC_CAST_D3D11RESOURCE(UniformBuffer,Buffer);
 #if PLATFORM_XBOXONE
@@ -719,8 +656,6 @@ void FD3D11DynamicRHI::RHISetShaderUniformBuffer(FGeometryShaderRHIParamRef Geom
 
 void FD3D11DynamicRHI::RHISetShaderUniformBuffer(FPixelShaderRHIParamRef PixelShader,uint32 BufferIndex,FUniformBufferRHIParamRef BufferRHI)
 {
-	GRHICommandList.Verify();
-
 	VALIDATE_BOUND_SHADER(PixelShader);
 	DYNAMIC_CAST_D3D11RESOURCE(UniformBuffer,Buffer);
 
@@ -742,8 +677,6 @@ void FD3D11DynamicRHI::RHISetShaderUniformBuffer(FPixelShaderRHIParamRef PixelSh
 
 void FD3D11DynamicRHI::RHISetShaderUniformBuffer(FComputeShaderRHIParamRef ComputeShader,uint32 BufferIndex,FUniformBufferRHIParamRef BufferRHI)
 {
-	GRHICommandList.Verify();
-
 	//VALIDATE_BOUND_SHADER(ComputeShader);
 	DYNAMIC_CAST_D3D11RESOURCE(UniformBuffer,Buffer);
 #if PLATFORM_XBOXONE
@@ -764,8 +697,6 @@ void FD3D11DynamicRHI::RHISetShaderUniformBuffer(FComputeShaderRHIParamRef Compu
 
 void FD3D11DynamicRHI::RHISetShaderParameter(FHullShaderRHIParamRef HullShaderRHI,uint32 BufferIndex,uint32 BaseIndex,uint32 NumBytes,const void* NewValue)
 {
-	GRHICommandList.Verify();
-
 	VALIDATE_BOUND_SHADER(HullShaderRHI);
 	checkSlow(HSConstantBuffers[BufferIndex]);
 	HSConstantBuffers[BufferIndex]->UpdateConstant((const uint8*)NewValue,BaseIndex,NumBytes);
@@ -773,8 +704,6 @@ void FD3D11DynamicRHI::RHISetShaderParameter(FHullShaderRHIParamRef HullShaderRH
 
 void FD3D11DynamicRHI::RHISetShaderParameter(FDomainShaderRHIParamRef DomainShaderRHI,uint32 BufferIndex,uint32 BaseIndex,uint32 NumBytes,const void* NewValue)
 {
-	GRHICommandList.Verify();
-
 	VALIDATE_BOUND_SHADER(DomainShaderRHI);
 	checkSlow(DSConstantBuffers[BufferIndex]);
 	DSConstantBuffers[BufferIndex]->UpdateConstant((const uint8*)NewValue,BaseIndex,NumBytes);
@@ -782,8 +711,6 @@ void FD3D11DynamicRHI::RHISetShaderParameter(FDomainShaderRHIParamRef DomainShad
 
 void FD3D11DynamicRHI::RHISetShaderParameter(FVertexShaderRHIParamRef VertexShaderRHI,uint32 BufferIndex,uint32 BaseIndex,uint32 NumBytes,const void* NewValue)
 {
-	GRHICommandList.Verify();
-
 	VALIDATE_BOUND_SHADER(VertexShaderRHI);
 	checkSlow(VSConstantBuffers[BufferIndex]);
 	VSConstantBuffers[BufferIndex]->UpdateConstant((const uint8*)NewValue,BaseIndex,NumBytes);
@@ -791,8 +718,6 @@ void FD3D11DynamicRHI::RHISetShaderParameter(FVertexShaderRHIParamRef VertexShad
 
 void FD3D11DynamicRHI::RHISetShaderParameter(FPixelShaderRHIParamRef PixelShaderRHI,uint32 BufferIndex,uint32 BaseIndex,uint32 NumBytes,const void* NewValue)
 {
-	GRHICommandList.Verify();
-
 	VALIDATE_BOUND_SHADER(PixelShaderRHI);
 	checkSlow(PSConstantBuffers[BufferIndex]);
 	PSConstantBuffers[BufferIndex]->UpdateConstant((const uint8*)NewValue,BaseIndex,NumBytes);
@@ -800,8 +725,6 @@ void FD3D11DynamicRHI::RHISetShaderParameter(FPixelShaderRHIParamRef PixelShader
 
 void FD3D11DynamicRHI::RHISetShaderParameter(FGeometryShaderRHIParamRef GeometryShaderRHI,uint32 BufferIndex,uint32 BaseIndex,uint32 NumBytes,const void* NewValue)
 {
-	GRHICommandList.Verify();
-
 	VALIDATE_BOUND_SHADER(GeometryShaderRHI);
 	checkSlow(GSConstantBuffers[BufferIndex]);
 	GSConstantBuffers[BufferIndex]->UpdateConstant((const uint8*)NewValue,BaseIndex,NumBytes);
@@ -809,8 +732,6 @@ void FD3D11DynamicRHI::RHISetShaderParameter(FGeometryShaderRHIParamRef Geometry
 
 void FD3D11DynamicRHI::RHISetShaderParameter(FComputeShaderRHIParamRef ComputeShaderRHI,uint32 BufferIndex,uint32 BaseIndex,uint32 NumBytes,const void* NewValue)
 {
-	GRHICommandList.Verify();
-
 	//VALIDATE_BOUND_SHADER(ComputeShaderRHI);
 	checkSlow(CSConstantBuffers[BufferIndex]);
 	CSConstantBuffers[BufferIndex]->UpdateConstant((const uint8*)NewValue,BaseIndex,NumBytes);
@@ -818,8 +739,6 @@ void FD3D11DynamicRHI::RHISetShaderParameter(FComputeShaderRHIParamRef ComputeSh
 
 void FD3D11DynamicRHI::RHISetDepthStencilState(FDepthStencilStateRHIParamRef NewStateRHI,uint32 StencilRef)
 {
-	GRHICommandList.Verify();
-
 	DYNAMIC_CAST_D3D11RESOURCE(DepthStencilState,NewState);
 
 	if (CurrentDepthTexture && NewState->AccessType != CurrentDSVAccessType)
@@ -838,16 +757,12 @@ void FD3D11DynamicRHI::RHISetDepthStencilState(FDepthStencilStateRHIParamRef New
 
 void FD3D11DynamicRHI::RHISetBlendState(FBlendStateRHIParamRef NewStateRHI,const FLinearColor& BlendFactor)
 {
-	GRHICommandList.Verify();
-
 	DYNAMIC_CAST_D3D11RESOURCE(BlendState,NewState);
 	StateCache.SetBlendState(NewState->Resource, (const float*)&BlendFactor, 0xffffffff);
 }
 
 void FD3D11DynamicRHI::CommitRenderTargetsAndUAVs()
 {
-	GRHICommandList.Verify();
-
 	ID3D11RenderTargetView* RTArray[D3D11_SIMULTANEOUS_RENDER_TARGET_COUNT];
 	for(uint32 RenderTargetIndex = 0;RenderTargetIndex < NumSimultaneousRenderTargets;++RenderTargetIndex)
 	{
@@ -954,8 +869,6 @@ void FD3D11DynamicRHI::RHISetRenderTargets(
 	const FUnorderedAccessViewRHIParamRef* UAVs
 	)
 {
-	GRHICommandList.Verify();
-
 	FD3D11TextureBase* NewDepthStencilTarget = GetD3D11TextureFromRHITexture(NewDepthStencilTargetRHI);
 
 	check(NewNumSimultaneousRenderTargets + NewNumUAVs <= MaxSimultaneousRenderTargets);
@@ -1087,44 +1000,9 @@ void FD3D11DynamicRHI::RHISetRenderTargets(
 
 void FD3D11DynamicRHI::RHIDiscardRenderTargets(bool Depth, bool Stencil, uint32 ColorBitMask)
 {
-	GRHICommandList.Verify();
-
 	// Could support in DX11.1 via ID3D11DeviceContext1::Discard*() functions.
 }
 
-// Occlusion/Timer queries.
-void FD3D11DynamicRHI::RHIBeginRenderQuery(FRenderQueryRHIParamRef QueryRHI)
-{
-	GRHICommandList.Verify();
-
-	DYNAMIC_CAST_D3D11RESOURCE(OcclusionQuery,Query);
-
-	if(Query->QueryType == RQT_Occlusion)
-	{
-#if PLATFORM_SUPPORTS_RHI_THREAD
-		Query->bResultIsCached = false;
-#endif
-		Direct3DDeviceIMContext->Begin(Query->Resource);
-	}
-	else
-	{
-		// not supported/needed for RQT_AbsoluteTime
-		check(0);
-	}
-}
-void FD3D11DynamicRHI::RHIEndRenderQuery(FRenderQueryRHIParamRef QueryRHI)
-{
-	GRHICommandList.Verify();
-
-	DYNAMIC_CAST_D3D11RESOURCE(OcclusionQuery,Query);
-#if PLATFORM_SUPPORTS_RHI_THREAD
-	Query->bResultIsCached = false; // for occlusion queries, this is redundant with the one in begin
-#endif
-	Direct3DDeviceIMContext->End(Query->Resource);
-
-	//@todo - d3d debug spews warnings about OQ's that are being issued but not polled, need to investigate
-}
-
 // Primitive drawing.
 
 static D3D11_PRIMITIVE_TOPOLOGY GetD3D11PrimitiveType(uint32 PrimitiveType, bool bUsingTessellation)
@@ -1202,8 +1080,6 @@ static D3D11_PRIMITIVE_TOPOLOGY GetD3D11PrimitiveType(uint32 PrimitiveType, bool
 
 void FD3D11DynamicRHI::CommitNonComputeShaderConstants()
 {
-	GRHICommandList.Verify();
-
 	FD3D11BoundShaderState* CurrentBoundShaderState = (FD3D11BoundShaderState*)BoundShaderStateHistory.GetLast();
 	check(CurrentBoundShaderState);
 
@@ -1271,8 +1147,6 @@ void FD3D11DynamicRHI::CommitNonComputeShaderConstants()
 
 void FD3D11DynamicRHI::CommitComputeShaderConstants()
 {
-	GRHICommandList.Verify();
-
 	bool bLocalDiscardSharedConstants = true;
 
 	// Commit and bind compute shader constants
@@ -1350,10 +1224,18 @@ void FD3D11DynamicRHI::SetResourcesFromTables(const ShaderType* RESTRICT Shader)
 	DirtyUniformBuffers[ShaderType::StaticFrequency] = 0;
 }
 
+static int32 PeriodicCheck = 0;
+
 void FD3D11DynamicRHI::CommitGraphicsResourceTables()
 {
-	GRHICommandList.Verify();
-
+#if PLATFORM_HAS_THREADSAFE_RHIGetRenderQueryResult
+	if (GRHIThread && ++PeriodicCheck >= 10)
+	{
+		QUICK_SCOPE_CYCLE_COUNTER(STAT_CommitGraphicsResourceTables_CheckThreadsafeQueries);
+		PeriodicCheck = 0;
+		CheckThreadsafeQueries();
+	}
+#endif
 	FD3D11BoundShaderState* RESTRICT CurrentBoundShaderState = (FD3D11BoundShaderState*)BoundShaderStateHistory.GetLast();
 	check(CurrentBoundShaderState);
 
@@ -1388,8 +1270,6 @@ void FD3D11DynamicRHI::CommitComputeResourceTables(FD3D11ComputeShader* InComput
 
 void FD3D11DynamicRHI::RHIDrawPrimitive(uint32 PrimitiveType,uint32 BaseVertexIndex,uint32 NumPrimitives,uint32 NumInstances)
 {
-	GRHICommandList.Verify();
-
 	RHI_DRAW_CALL_STATS(PrimitiveType,NumInstances*NumPrimitives);
 
 	CommitGraphicsResourceTables();
@@ -1412,8 +1292,6 @@ void FD3D11DynamicRHI::RHIDrawPrimitive(uint32 PrimitiveType,uint32 BaseVertexIn
 
 void FD3D11DynamicRHI::RHIDrawPrimitiveIndirect(uint32 PrimitiveType,FVertexBufferRHIParamRef ArgumentBufferRHI,uint32 ArgumentOffset)
 {
-	GRHICommandList.Verify();
-
 	DYNAMIC_CAST_D3D11RESOURCE(VertexBuffer,ArgumentBuffer);
 
 	RHI_DRAW_CALL_INC();
@@ -1429,8 +1307,6 @@ void FD3D11DynamicRHI::RHIDrawPrimitiveIndirect(uint32 PrimitiveType,FVertexBuff
 
 void FD3D11DynamicRHI::RHIDrawIndexedIndirect(FIndexBufferRHIParamRef IndexBufferRHI, uint32 PrimitiveType, FStructuredBufferRHIParamRef ArgumentsBufferRHI, int32 DrawArgumentsIndex, uint32 NumInstances)
 {
-	GRHICommandList.Verify();
-
 	DYNAMIC_CAST_D3D11RESOURCE(IndexBuffer,IndexBuffer);
 	DYNAMIC_CAST_D3D11RESOURCE(StructuredBuffer,ArgumentsBuffer);
 
@@ -1460,8 +1336,6 @@ void FD3D11DynamicRHI::RHIDrawIndexedIndirect(FIndexBufferRHIParamRef IndexBuffe
 
 void FD3D11DynamicRHI::RHIDrawIndexedPrimitive(FIndexBufferRHIParamRef IndexBufferRHI,uint32 PrimitiveType,int32 BaseVertexIndex,uint32 MinIndex,uint32 NumVertices,uint32 StartIndex,uint32 NumPrimitives,uint32 NumInstances)
 {
-	GRHICommandList.Verify();
-
 	DYNAMIC_CAST_D3D11RESOURCE(IndexBuffer,IndexBuffer);
 
 	// called should make sure the input is valid, this avoid hidden bugs
@@ -1500,8 +1374,6 @@ void FD3D11DynamicRHI::RHIDrawIndexedPrimitive(FIndexBufferRHIParamRef IndexBuff
 
 void FD3D11DynamicRHI::RHIDrawIndexedPrimitiveIndirect(uint32 PrimitiveType,FIndexBufferRHIParamRef IndexBufferRHI,FVertexBufferRHIParamRef ArgumentBufferRHI,uint32 ArgumentOffset)
 {
-	GRHICommandList.Verify();
-
 	DYNAMIC_CAST_D3D11RESOURCE(IndexBuffer,IndexBuffer);
 	DYNAMIC_CAST_D3D11RESOURCE(VertexBuffer,ArgumentBuffer);
 
@@ -1530,8 +1402,6 @@ void FD3D11DynamicRHI::RHIDrawIndexedPrimitiveIndirect(uint32 PrimitiveType,FInd
  */
 void FD3D11DynamicRHI::RHIBeginDrawPrimitiveUP( uint32 PrimitiveType, uint32 NumPrimitives, uint32 NumVertices, uint32 VertexDataStride, void*& OutVertexData)
 {
-	GRHICommandList.Verify();
-
 	checkSlow( PendingNumVertices == 0 );
 
 	// Remember the parameters for the draw call.
@@ -1549,8 +1419,6 @@ void FD3D11DynamicRHI::RHIBeginDrawPrimitiveUP( uint32 PrimitiveType, uint32 Num
  */
 void FD3D11DynamicRHI::RHIEndDrawPrimitiveUP()
 {
-	GRHICommandList.Verify();
-
 	RHI_DRAW_CALL_STATS(PendingPrimitiveType,PendingNumPrimitives);
 
 	checkSlow(!bUsingTessellation || PendingPrimitiveType == PT_TriangleList);
@@ -1589,8 +1457,6 @@ void FD3D11DynamicRHI::RHIEndDrawPrimitiveUP()
  */
 void FD3D11DynamicRHI::RHIBeginDrawIndexedPrimitiveUP( uint32 PrimitiveType, uint32 NumPrimitives, uint32 NumVertices, uint32 VertexDataStride, void*& OutVertexData, uint32 MinVertexIndex, uint32 NumIndices, uint32 IndexDataStride, void*& OutIndexData)
 {
-	GRHICommandList.Verify();
-
 	checkSlow((sizeof(uint16) == IndexDataStride) || (sizeof(uint32) == IndexDataStride));
 
 	// Store off information needed for the draw call.
@@ -1612,8 +1478,6 @@ void FD3D11DynamicRHI::RHIBeginDrawIndexedPrimitiveUP( uint32 PrimitiveType, uin
  */
 void FD3D11DynamicRHI::RHIEndDrawIndexedPrimitiveUP()
 {
-	GRHICommandList.Verify();
-
 	// tessellation only supports trilists
 	checkSlow(!bUsingTessellation || PendingPrimitiveType == PT_TriangleList);
 
@@ -1647,15 +1511,11 @@ void FD3D11DynamicRHI::RHIEndDrawIndexedPrimitiveUP()
 // Raster operations.
 void FD3D11DynamicRHI::RHIClear(bool bClearColor,const FLinearColor& Color,bool bClearDepth,float Depth,bool bClearStencil,uint32 Stencil, FIntRect ExcludeRect)
 {
-	GRHICommandList.Verify();
-
 	FD3D11DynamicRHI::RHIClearMRT(bClearColor, 1, &Color, bClearDepth, Depth, bClearStencil, Stencil, ExcludeRect);
 }
 
 void FD3D11DynamicRHI::RHIClearMRT(bool bClearColor,int32 NumClearColors,const FLinearColor* ClearColorArray,bool bClearDepth,float Depth,bool bClearStencil,uint32 Stencil, FIntRect ExcludeRect)
 {
-	GRHICommandList.Verify();
-
 	// Helper struct to record and restore device states RHIClearMRT modifies.
 	class FDeviceStateHelper
 	{
@@ -2080,22 +1940,16 @@ void FD3D11DynamicRHI::RHIClearMRT(bool bClearColor,int32 NumClearColors,const F
 
 void FD3D11DynamicRHI::RHISuspendRendering()
 {
-	GRHICommandList.Verify();
-
 	// Not supported
 }
 
 void FD3D11DynamicRHI::RHIResumeRendering()
 {
-	GRHICommandList.Verify();
-
 	// Not supported
 }
 
 bool FD3D11DynamicRHI::RHIIsRenderingSuspended()
 {
-	GRHICommandList.Verify();
-
 	// Not supported
 	return false;
 }
@@ -2103,8 +1957,6 @@ bool FD3D11DynamicRHI::RHIIsRenderingSuspended()
 // Blocks the CPU until the GPU catches up and goes idle.
 void FD3D11DynamicRHI::RHIBlockUntilGPUIdle()
 {
-	GRHICommandList.Verify();
-
 	// Not really supported
 }
 
diff --git a/Engine/Source/Runtime/Windows/D3D11RHI/Private/D3D11Query.cpp b/Engine/Source/Runtime/Windows/D3D11RHI/Private/D3D11Query.cpp
--- a/Engine/Source/Runtime/Windows/D3D11RHI/Private/D3D11Query.cpp
+++ b/Engine/Source/Runtime/Windows/D3D11RHI/Private/D3D11Query.cpp
@@ -6,6 +6,96 @@
 
 #include "D3D11RHIPrivate.h"
 
+#if PLATFORM_HAS_THREADSAFE_RHIGetRenderQueryResult
+struct FThreadsafeQueryEmulator
+{
+	FThreadSafeCounter IsWaiting;
+	FCriticalSection OutstandingQueriesLock;
+	TArray<FRenderQueryRHIParamRef> OutstandingQueries;
+
+
+	FORCEINLINE_DEBUGGABLE void CheckThreadsafeQueries(FD3D11DynamicRHI &CallingRHI, bool bWait = false)
+	{
+		if (!IsWaiting.GetValue())
+		{
+			return;
+		}
+		QUICK_SCOPE_CYCLE_COUNTER(STAT_CheckThreadsafeQueries);
+		FScopeLock Lock(&OutstandingQueriesLock);
+		int32 LastOutstanding = OutstandingQueries.Num() - 1;
+		for (int32 Index = LastOutstanding; Index >= 0; Index--)
+		{
+			FRenderQueryRHIParamRef QueryRHI = OutstandingQueries[Index];
+
+			DYNAMIC_CAST_D3D11RESOURCE(OcclusionQuery,Query); // I think this macro is the worst thing in the whole world -- Gil
+
+			if(!Query->bResultIsCached)
+			{
+				bool bSuccess = CallingRHI.GetQueryData(Query->Resource,&Query->Result,sizeof(Query->Result),bWait, Query->QueryType);
+				FPlatformMisc::MemoryBarrier();
+				Query->bResultIsCached = bSuccess;
+				check(!bWait || Query->bResultIsCached);
+			}
+			if (Query->bResultIsCached)
+			{
+				LastOutstanding = Index - 1;
+			}
+			else
+			{
+				break;
+			}
+		}
+		int32 NumDeleted = OutstandingQueries.Num() - LastOutstanding - 1;
+		OutstandingQueries.RemoveAt(LastOutstanding + 1, NumDeleted, false);
+		check(!bWait || !OutstandingQueries.Num());
+	}
+	FORCEINLINE_DEBUGGABLE void Remove(FRenderQueryRHIParamRef QueryRHI)
+	{
+		FScopeLock Lock(&OutstandingQueriesLock);
+		OutstandingQueries.Remove(QueryRHI);
+	}
+	FORCEINLINE_DEBUGGABLE void Add(FRenderQueryRHIParamRef QueryRHI)
+	{
+		FScopeLock Lock(&OutstandingQueriesLock); // I am not sure this is actually required, but there should be no contention anyway
+		OutstandingQueries.Add(QueryRHI);
+	}
+	FORCEINLINE_DEBUGGABLE void WaitForQuery(FD3D11DynamicRHI &CallingRHI,FRenderQueryRHIParamRef QueryRHI)
+	{
+		DYNAMIC_CAST_D3D11RESOURCE(OcclusionQuery,Query);
+		verify(IsWaiting.Increment() == 1);
+		int32 Iteration = 0;
+		while (!Query->bResultIsCached)
+		{
+			if (Iteration++ > 0)
+			{
+				FPlatformProcess::Sleep(0.0001f);
+			}
+			if (!GRHICommandList.IsRHIThreadActive())
+			{
+				CheckThreadsafeQueries(CallingRHI, true);
+				break;
+			}
+			FPlatformMisc::MemoryBarrier();
+		}
+		verify(IsWaiting.Decrement() == 0);
+	}
+};
+
+static FThreadsafeQueryEmulator ThreadsafeQueryEmulator;
+
+void FD3D11DynamicRHI::CheckThreadsafeQueries(bool bWait)
+{
+	ThreadsafeQueryEmulator.CheckThreadsafeQueries(*this, bWait);
+}
+
+FD3D11OcclusionQuery::~FD3D11OcclusionQuery()
+{
+	// this more or less never happens until shutdown, but we will do a slow cleanup here for completeness.
+	ThreadsafeQueryEmulator.Remove(this);
+}
+
+#endif
+
 FRenderQueryRHIRef FD3D11DynamicRHI::RHICreateRenderQuery(ERenderQueryType QueryType)
 {
 	TRefCountPtr<ID3D11Query> Query;
@@ -42,6 +132,42 @@ bool FD3D11DynamicRHI::RHIGetRenderQueryResult(FRenderQueryRHIParamRef QueryRHI,
 {
 	DYNAMIC_CAST_D3D11RESOURCE(OcclusionQuery,Query);
 
+#if PLATFORM_HAS_THREADSAFE_RHIGetRenderQueryResult
+	if (GRHIThread && !IsInRHIThread())
+	{
+		check(IsInRenderingThread());
+		if (!Query->bResultIsCached && !bWait)
+		{
+			return false;
+		}
+		if (!Query->bResultIsCached && bWait)
+		{
+			QUICK_SCOPE_CYCLE_COUNTER(STAT_GetRenderQueryResult_Spin);
+			ThreadsafeQueryEmulator.WaitForQuery(*this, QueryRHI);
+		}
+		if (Query->bResultIsCached)
+		{
+			if(Query->QueryType == RQT_AbsoluteTime)
+			{
+				// GetTimingFrequency is the number of ticks per second
+				uint64 Div = FGPUTiming::GetTimingFrequency() / (1000 * 1000);
+
+				// convert from GPU specific timestamp to micro sec (1 / 1 000 000 s) which seems a reasonable resolution
+				OutResult = Query->Result / Div;
+			}
+			else
+			{
+				OutResult = Query->Result;
+			}
+			return true;
+		}
+		if (!bWait)
+		{
+			return false;
+		}
+		 // we need to fall through in the case that we don't have outstanding queries...this happens when we transition into rhi thread because the last frame queries were not queued up
+	}
+#endif
 	bool bSuccess = true;
 	if(!Query->bResultIsCached)
 	{
@@ -62,10 +188,54 @@ bool FD3D11DynamicRHI::RHIGetRenderQueryResult(FRenderQueryRHIParamRef QueryRHI,
 	{
 		OutResult = Query->Result;
 	}
-
+#if PLATFORM_HAS_THREADSAFE_RHIGetRenderQueryResult
+	check(!bWait || bSuccess);
+#endif
 	return bSuccess;
 }
 
+
+// Occlusion/Timer queries.
+void FD3D11DynamicRHI::RHIBeginRenderQuery(FRenderQueryRHIParamRef QueryRHI)
+{
+	DYNAMIC_CAST_D3D11RESOURCE(OcclusionQuery,Query);
+
+	if(Query->QueryType == RQT_Occlusion)
+	{
+#if PLATFORM_SUPPORTS_RHI_THREAD
+		Query->bResultIsCached = false;
+#endif
+		Direct3DDeviceIMContext->Begin(Query->Resource);
+	}
+	else
+	{
+		// not supported/needed for RQT_AbsoluteTime
+		check(0);
+	}
+}
+
+void FD3D11DynamicRHI::RHIEndRenderQuery(FRenderQueryRHIParamRef QueryRHI)
+{
+	DYNAMIC_CAST_D3D11RESOURCE(OcclusionQuery,Query);
+#if PLATFORM_SUPPORTS_RHI_THREAD
+	Query->bResultIsCached = false; // for occlusion queries, this is redundant with the one in begin
+
+	#if PLATFORM_HAS_THREADSAFE_RHIGetRenderQueryResult
+		if (GRHIThread)
+		{
+			ThreadsafeQueryEmulator.Add(QueryRHI);
+		}
+	#endif
+
+#elif PLATFORM_HAS_THREADSAFE_RHIGetRenderQueryResult
+	#error "PLATFORM_HAS_THREADSAFE_RHIGetRenderQueryResult requires PLATFORM_SUPPORTS_RHI_THREAD"
+#endif
+	Direct3DDeviceIMContext->End(Query->Resource);
+
+	//@todo - d3d debug spews warnings about OQ's that are being issued but not polled, need to investigate
+}
+
+
 bool FD3D11DynamicRHI::GetQueryData(ID3D11Query* Query,void* Data,SIZE_T DataSize,bool bWait, ERenderQueryType QueryType)
 {
 	// Request the data from the query.
diff --git a/Engine/Source/Runtime/Windows/D3D11RHI/Private/D3D11RHIPrivate.h b/Engine/Source/Runtime/Windows/D3D11RHI/Private/D3D11RHIPrivate.h
--- a/Engine/Source/Runtime/Windows/D3D11RHI/Private/D3D11RHIPrivate.h
+++ b/Engine/Source/Runtime/Windows/D3D11RHI/Private/D3D11RHIPrivate.h
@@ -317,10 +317,6 @@ public:
 	/** If it hasn't been initialized yet, initializes the D3D device. */
 	virtual void InitD3DDevice();
 
-#if PLATFORM_SUPPORTS_RHI_THREAD
-	void SetupRecursiveResources();
-#endif
-
 	// FDynamicRHI interface.
 	virtual void Init() override;
 	virtual void Shutdown() override;
@@ -585,6 +581,13 @@ protected:
 
 	void ReadSurfaceDataMSAARaw(FRHICommandList_RecursiveHazardous& RHICmdList, FTextureRHIParamRef TextureRHI, FIntRect Rect, TArray<uint8>& OutData, FReadSurfaceDataFlags InFlags);
 
+#if PLATFORM_SUPPORTS_RHI_THREAD
+	void SetupRecursiveResources();
+#endif
+#if PLATFORM_HAS_THREADSAFE_RHIGetRenderQueryResult
+	void CheckThreadsafeQueries(bool bWait = false);
+#endif
+
 	friend struct FD3DGPUProfiler;
 };
 
diff --git a/Engine/Source/Runtime/Windows/D3D11RHI/Public/D3D11Resources.h b/Engine/Source/Runtime/Windows/D3D11RHI/Public/D3D11Resources.h
--- a/Engine/Source/Runtime/Windows/D3D11RHI/Public/D3D11Resources.h
+++ b/Engine/Source/Runtime/Windows/D3D11RHI/Public/D3D11Resources.h
@@ -578,6 +578,10 @@ public:
 		bResultIsCached(false),
 		QueryType(InQueryType)
 	{}
+
+#if PLATFORM_HAS_THREADSAFE_RHIGetRenderQueryResult
+	~FD3D11OcclusionQuery();
+#endif
 };
 
 /** Updates tracked stats for a buffer. */
