UE4 - Add GSupportsDepthRenderTargetWithoutColorRenderTarget in prep for mobile cascade shadows

diff --git a/Engine/Source/Runtime/OpenGLDrv/Private/OpenGLES2.cpp b/Engine/Source/Runtime/OpenGLDrv/Private/OpenGLES2.cpp
--- a/Engine/Source/Runtime/OpenGLDrv/Private/OpenGLES2.cpp
+++ b/Engine/Source/Runtime/OpenGLDrv/Private/OpenGLES2.cpp
@@ -193,6 +193,9 @@ void FOpenGLES2::ProcessExtensions( const FString& ExtensionsString )
 		GEmitMeshDrawEvent = 1;
 		GEmitDrawEvents = true;
 	}
+
+	// ES2 requires a color attachment when rendering to depth-only.
+	GSupportsDepthRenderTargetWithoutColorRenderTarget = false;
 }
 
 #endif
diff --git a/Engine/Source/Runtime/RHI/Private/RHI.cpp b/Engine/Source/Runtime/RHI/Private/RHI.cpp
--- a/Engine/Source/Runtime/RHI/Private/RHI.cpp
+++ b/Engine/Source/Runtime/RHI/Private/RHI.cpp
@@ -176,6 +176,7 @@ bool GSupportsQuads = false;
 bool GSupportsVolumeTextureRendering = true;
 bool GSupportsGSRenderTargetLayerSwitchingToMips = true;
 bool GSupportsSeparateRenderTargetBlendState = false;
+bool GSupportsDepthRenderTargetWithoutColorRenderTarget = true;
 float GPixelCenterOffset = 0;
 float GMinClipZ = 0.0f;
 float GProjectionSignY = 1.0f;
diff --git a/Engine/Source/Runtime/RHI/Public/RHI.h b/Engine/Source/Runtime/RHI/Public/RHI.h
--- a/Engine/Source/Runtime/RHI/Public/RHI.h
+++ b/Engine/Source/Runtime/RHI/Public/RHI.h
@@ -105,6 +105,9 @@ extern RHI_API bool GSupportsGSRenderTargetLayerSwitchingToMips;
 /** True if the RHI supports separate blend states per render target. */
 extern RHI_API bool GSupportsSeparateRenderTargetBlendState;
 
+/** True if the RHI can render to a depth-only render target with no additional color render target. */
+extern RHI_API bool GSupportsDepthRenderTargetWithoutColorRenderTarget;
+
 /** True if the RHI supports depth bounds testing */
 extern RHI_API bool GSupportsDepthBoundsTest;
 
@@ -963,11 +966,11 @@ struct FTextureMemoryStats
 	// Hardware state (never change after device creation):
 
 	// -1 if unknown, in bytes
-    int64 DedicatedVideoMemory;
+	int64 DedicatedVideoMemory;
 	// -1 if unknown, in bytes
-    int64 DedicatedSystemMemory;
+	int64 DedicatedSystemMemory;
 	// -1 if unknown, in bytes
-    int64 SharedSystemMemory;
+	int64 SharedSystemMemory;
 	// Total amount of "graphics memory" that we think we can use for all our graphics resources, in bytes. -1 if unknown.
 	int64 TotalGraphicsMemory;
 
