UE4 - new RHISetRenderTargets API

diff --git a/Engine/Source/Runtime/EmptyRHI/Private/EmptyCommands.cpp b/Engine/Source/Runtime/EmptyRHI/Private/EmptyCommands.cpp
--- a/Engine/Source/Runtime/EmptyRHI/Private/EmptyCommands.cpp
+++ b/Engine/Source/Runtime/EmptyRHI/Private/EmptyCommands.cpp
@@ -273,6 +273,9 @@ void FEmptyDynamicRHI::RHIDiscardRenderTargets(bool Depth, bool Stencil, uint32
 {
 }
 
+void FEmptyDynamicRHI::RHISetRenderTargetsAndClear(const FRHISetRenderTargetsInfo& RenderTargetsInfo)
+{
+}
 
 // Occlusion/Timer queries.
 void FEmptyDynamicRHI::RHIBeginRenderQuery(FRenderQueryRHIParamRef QueryRHI)
diff --git a/Engine/Source/Runtime/IOS/MetalRHI/Private/MetalCommands.cpp b/Engine/Source/Runtime/IOS/MetalRHI/Private/MetalCommands.cpp
--- a/Engine/Source/Runtime/IOS/MetalRHI/Private/MetalCommands.cpp
+++ b/Engine/Source/Runtime/IOS/MetalRHI/Private/MetalCommands.cpp
@@ -364,6 +364,19 @@ void FMetalDynamicRHI::RHIDiscardRenderTargets(bool Depth, bool Stencil, uint32
 {
 }
 
+void FMetalDynamicRHI::RHISetRenderTargetsAndClear(const FRHISetRenderTargetsInfo& RenderTargetsInfo)
+{
+//@todo-rco: TEMP
+	RHISetRenderTargets(RenderTargetsInfo.NumColorRenderTargets,
+		RenderTargetsInfo.ColorRenderTarget,
+		RenderTargetsInfo.DepthStencilTarget,
+		0,
+		nullptr);
+	if (RenderTargetsInfo.bClearColor || RenderTargetsInfo.bClearStencil || RenderTargetsInfo.bClearDepth)
+	{
+		RHIClearMRT(RenderTargetsInfo.bClearColor, RenderTargetsInfo.NumColorRenderTargets, RenderTargetsInfo.ClearColors, RenderTargetsInfo.bClearDepth, RenderTargetsInfo.DepthClearValue, RenderTargetsInfo.bClearStencil, RenderTargetsInfo.StencilClearValue, FIntRect());
+	}
+}
 
 // Occlusion/Timer queries.
 void FMetalDynamicRHI::RHIBeginRenderQuery(FRenderQueryRHIParamRef QueryRHI)
diff --git a/Engine/Source/Runtime/OpenGLDrv/Private/OpenGLCommands.cpp b/Engine/Source/Runtime/OpenGLDrv/Private/OpenGLCommands.cpp
--- a/Engine/Source/Runtime/OpenGLDrv/Private/OpenGLCommands.cpp
+++ b/Engine/Source/Runtime/OpenGLDrv/Private/OpenGLCommands.cpp
@@ -1633,6 +1633,19 @@ void FOpenGLDynamicRHI::RHIDiscardRenderTargets(bool Depth, bool Stencil, uint32
 	}
 }
 
+void FOpenGLDynamicRHI::RHISetRenderTargetsAndClear(const FRHISetRenderTargetsInfo& RenderTargetsInfo)
+{
+	RHISetRenderTargets(RenderTargetsInfo.NumColorRenderTargets,
+		RenderTargetsInfo.ColorRenderTarget,
+		RenderTargetsInfo.DepthStencilTarget,
+		0,
+		nullptr);
+	if (RenderTargetsInfo.bClearColor || RenderTargetsInfo.bClearStencil || RenderTargetsInfo.bClearDepth)
+	{
+		RHIClearMRT(RenderTargetsInfo.bClearColor, RenderTargetsInfo.NumColorRenderTargets, RenderTargetsInfo.ClearColors, RenderTargetsInfo.bClearDepth, RenderTargetsInfo.DepthClearValue, RenderTargetsInfo.bClearStencil, RenderTargetsInfo.StencilClearValue, FIntRect());
+	}
+}
+
 // Primitive drawing.
 
 void FOpenGLDynamicRHI::EnableVertexElementCached(
diff --git a/Engine/Source/Runtime/RHI/Public/RHICommandList.h b/Engine/Source/Runtime/RHI/Public/RHICommandList.h
--- a/Engine/Source/Runtime/RHI/Public/RHICommandList.h
+++ b/Engine/Source/Runtime/RHI/Public/RHICommandList.h
@@ -480,6 +480,18 @@ struct FRHICommandSetRenderTargets : public FRHICommand<FRHICommandSetRenderTarg
 	RHI_API void Execute();
 };
 
+struct FRHICommandSetRenderTargetsAndClear : public FRHICommand<FRHICommandSetRenderTargetsAndClear>
+{
+	FRHISetRenderTargetsInfo RenderTargetsInfo;
+
+	FORCEINLINE_DEBUGGABLE FRHICommandSetRenderTargetsAndClear(const FRHISetRenderTargetsInfo& InRenderTargetsInfo) :
+		RenderTargetsInfo(InRenderTargetsInfo)
+	{
+	}
+
+	RHI_API void Execute();
+};
+
 struct FRHICommandEndDrawPrimitiveUP : public FRHICommand<FRHICommandEndDrawPrimitiveUP>
 {
 	uint32 PrimitiveType;
@@ -1425,6 +1437,16 @@ public:
 			UAVs);
 	}
 
+	FORCEINLINE_DEBUGGABLE void SetRenderTargetsAndClear(const FRHISetRenderTargetsInfo& RenderTargetsInfo)
+	{
+		if (Bypass())
+		{
+			SetRenderTargetsAndClear_Internal(RenderTargetsInfo);
+			return;
+		}
+		new (AllocCommand<FRHICommandSetRenderTargetsAndClear>()) FRHICommandSetRenderTargetsAndClear(RenderTargetsInfo);
+	}
+
 	FORCEINLINE_DEBUGGABLE void BeginDrawPrimitiveUP(uint32 PrimitiveType, uint32 NumPrimitives, uint32 NumVertices, uint32 VertexDataStride, void*& OutVertexData)
 	{
 		if (Bypass())
diff --git a/Engine/Source/Runtime/RHI/Public/RHICommandListCommandExecutes.inl b/Engine/Source/Runtime/RHI/Public/RHICommandListCommandExecutes.inl
--- a/Engine/Source/Runtime/RHI/Public/RHICommandListCommandExecutes.inl
+++ b/Engine/Source/Runtime/RHI/Public/RHICommandListCommandExecutes.inl
@@ -164,6 +164,12 @@ void FRHICommandSetRenderTargets::Execute()
 		UAVs);
 }
 
+void FRHICommandSetRenderTargetsAndClear::Execute()
+{
+	RHISTAT(SetRenderTargetsAndClear);
+	INTERNAL_DECORATOR(SetRenderTargetsAndClear)(RenderTargetsInfo);
+}
+
 void FRHICommandEndDrawPrimitiveUP::Execute()
 {
 	RHISTAT(EndDrawPrimitiveUP);
diff --git a/Engine/Source/Runtime/RHI/Public/RHIDefinitions.h b/Engine/Source/Runtime/RHI/Public/RHIDefinitions.h
--- a/Engine/Source/Runtime/RHI/Public/RHIDefinitions.h
+++ b/Engine/Source/Runtime/RHI/Public/RHIDefinitions.h
@@ -524,7 +524,7 @@ enum class ERenderTargetLoadAction
 {
 	ENoAction,
 	ELoad,
-	//EClear clear should go here but isn't implemented yet.
+	EClear,
 };
 
 /**
diff --git a/Engine/Source/Runtime/RHI/Public/RHIMethods.h b/Engine/Source/Runtime/RHI/Public/RHIMethods.h
--- a/Engine/Source/Runtime/RHI/Public/RHIMethods.h
+++ b/Engine/Source/Runtime/RHI/Public/RHIMethods.h
@@ -1876,6 +1876,11 @@ DEFINE_RHIMETHOD_3(
 	uint32,ColorBitMask,
 	,
 	);
+DEFINE_RHIMETHOD_1(
+	void,SetRenderTargetsAndClear,
+	const FRHISetRenderTargetsInfo&, RenderTargetsInfo,
+	,
+	);
 
 // Primitive drawing.
 DEFINE_RHIMETHOD_CMDLIST_4(
diff --git a/Engine/Source/Runtime/RHI/Public/RHIResources.h b/Engine/Source/Runtime/RHI/Public/RHIResources.h
--- a/Engine/Source/Runtime/RHI/Public/RHIResources.h
+++ b/Engine/Source/Runtime/RHI/Public/RHIResources.h
@@ -621,6 +621,38 @@ public:
 	{}
 };
 
+class FRHISetRenderTargetsInfo
+{
+public:
+	// Color Render Targets Info
+	int32 NumColorRenderTargets;
+	FLinearColor ClearColors[MaxSimultaneousRenderTargets];
+	bool bClearColor;
+	FRHIRenderTargetView ColorRenderTarget[MaxSimultaneousRenderTargets];
+
+	// Depth/Stencil Render Target Info
+	ERenderTargetLoadAction DepthLoadAction;
+	ERenderTargetStoreAction DepthStoreAction;
+	float DepthClearValue;
+	bool bClearDepth;
+	uint32 StencilClearValue;
+	bool bClearStencil;
+	FTextureRHIParamRef DepthStencilTarget;
+
+	FRHISetRenderTargetsInfo() :
+		NumColorRenderTargets(0),
+		bClearColor(false),
+		DepthLoadAction(ERenderTargetLoadAction::EClear),
+		DepthStoreAction(ERenderTargetStoreAction::ENoAction),
+		DepthClearValue(0.0f),
+		bClearDepth(false),
+		StencilClearValue(0),
+		bClearStencil(false),
+		DepthStencilTarget(nullptr)
+	{
+	}
+};
+
 class FRHICustomPresent : public FRHIResource
 {
 public:
diff --git a/Engine/Source/Runtime/Windows/D3D11RHI/Private/D3D11Commands.cpp b/Engine/Source/Runtime/Windows/D3D11RHI/Private/D3D11Commands.cpp
--- a/Engine/Source/Runtime/Windows/D3D11RHI/Private/D3D11Commands.cpp
+++ b/Engine/Source/Runtime/Windows/D3D11RHI/Private/D3D11Commands.cpp
@@ -1003,6 +1003,19 @@ void FD3D11DynamicRHI::RHIDiscardRenderTargets(bool Depth, bool Stencil, uint32
 	// Could support in DX11.1 via ID3D11DeviceContext1::Discard*() functions.
 }
 
+void FD3D11DynamicRHI::RHISetRenderTargetsAndClear(const FRHISetRenderTargetsInfo& RenderTargetsInfo)
+{
+	RHISetRenderTargets(RenderTargetsInfo.NumColorRenderTargets,
+		RenderTargetsInfo.ColorRenderTarget,
+		RenderTargetsInfo.DepthStencilTarget,
+		0,
+		nullptr);
+	if (RenderTargetsInfo.bClearColor || RenderTargetsInfo.bClearStencil || RenderTargetsInfo.bClearDepth)
+	{
+		RHIClearMRT(RenderTargetsInfo.bClearColor, RenderTargetsInfo.NumColorRenderTargets, RenderTargetsInfo.ClearColors, RenderTargetsInfo.bClearDepth, RenderTargetsInfo.DepthClearValue, RenderTargetsInfo.bClearStencil, RenderTargetsInfo.StencilClearValue, FIntRect());
+	}
+}
+
 // Primitive drawing.
 
 static D3D11_PRIMITIVE_TOPOLOGY GetD3D11PrimitiveType(uint32 PrimitiveType, bool bUsingTessellation)
