added r.ProfileGPU.Pattern which allows to filter the entries when using ProfileGPU. Renames some other cvars

diff --git a/Engine/Source/Runtime/RHI/Private/GPUProfiler.cpp b/Engine/Source/Runtime/RHI/Private/GPUProfiler.cpp
--- a/Engine/Source/Runtime/RHI/Private/GPUProfiler.cpp
+++ b/Engine/Source/Runtime/RHI/Private/GPUProfiler.cpp
@@ -13,6 +13,16 @@
 #include "TaskGraphInterfaces.h"
 #endif
 
+static TAutoConsoleVariable<FString> GProfileGPUPatternCVar(
+	TEXT("r.ProfileGPU.Pattern"),
+	TEXT("*"),
+	TEXT("Allows to filter the entries when using ProfileGPU, the pattern match is case sensitive.\n")
+	TEXT("'*' can be used in the end to get all entries starting with the string.\n")
+	TEXT("    '*' without any leading characters disables the pattern matching and uses a time threshold instead (default).\n")
+	TEXT("'?' allows to ignore one character.\n")
+	TEXT("e.g. AmbientOcclusionSetup, AmbientOcclusion*, Ambient???lusion*, *"),
+	ECVF_Default);
+
 struct FNodeStatsCompare
 {
 	/** Sorts nodes by descending durations. */
@@ -196,11 +206,45 @@ void FGPUProfilerEventNodeFrame::DumpEventTree()
 		// Log stats about the node histogram
 		UE_LOG(LogRHI, Warning, TEXT("Node histogram %u buckets"), EventHistogram.Num());
 
+		static IConsoleVariable* CVar = IConsoleManager::Get().FindConsoleVariable(TEXT("r.ProfileGPU.Pattern"));
+
+		// bad: reading on render thread but we don't support ECVF_RenderThreadSafe on strings yet
+		// It's very unlikely to cause a problem as the cvar is only changes by the user.
+		FString WildcardString = CVar->GetString(); 
+
+		const float ThresholdInMS = 5.0f;
+
+		if(WildcardString == FString(TEXT("*")))
+		{
+			// disable Wildcard functionality
+			WildcardString.Empty();
+		}
+
+		if(WildcardString.IsEmpty())
+		{
+			UE_LOG(LogRHI, Warning, TEXT(" r.ProfileGPU.Pattern = '*' (using threshold of %g ms)"), ThresholdInMS);
+		}
+		else
+		{
+			UE_LOG(LogRHI, Warning, TEXT(" r.ProfileGPU.Pattern = '%s' (not using time threshold)"), *WildcardString);
+		}
+
+		FWildcardString Wildcard(WildcardString);
+
 		int32 NumNotShown = 0;
 		for (TMap<FString, FGPUProfilerEventNodeStats>::TIterator It(EventHistogram); It; ++It)
 		{
 			const FGPUProfilerEventNodeStats& NodeStats = It.Value();
-			if (NodeStats.TimingResult > RootResult * 1000.0f * .005f)
+
+			bool bDump = NodeStats.TimingResult > RootResult * ThresholdInMS;
+
+			if(!Wildcard.IsEmpty())
+			{
+				// if a Wildcard string was specified, we want to always dump all entries
+				bDump = Wildcard.IsMatch(*It.Key());
+			}
+
+			if (bDump)
 			{
 				UE_LOG(LogRHI, Warning, TEXT("   %.2fms   %s   Events %u   Draws %u"), NodeStats.TimingResult, *It.Key(), NodeStats.NumEvents, NodeStats.NumDraws);
 			}
diff --git a/Engine/Source/Runtime/RHI/Private/RHI.cpp b/Engine/Source/Runtime/RHI/Private/RHI.cpp
--- a/Engine/Source/Runtime/RHI/Private/RHI.cpp
+++ b/Engine/Source/Runtime/RHI/Private/RHI.cpp
@@ -118,14 +118,16 @@ static TAutoConsoleVariable<int32> ResourceTableCachingCvar(
 	TEXT("If 1, the RHI will cache resource table contents within a frame. Otherwise resource tables are rebuilt for every draw call.")
 	);
 static TAutoConsoleVariable<int32> GSaveScreenshotAfterProfilingGPUCVar(
-	TEXT("RHI.SaveScreenshotAfterProfilingGPU"),
+	TEXT("r.ProfileGPU.Screenshot"),
 	1,
-	TEXT("Whether a screenshot should be taken when profiling the GPU."),
+	TEXT("Whether a screenshot should be taken when profiling the GPU. 0:off, 1:on (default)"),
 	ECVF_RenderThreadSafe);
 static TAutoConsoleVariable<int32> GShowProfilerAfterProfilingGPUCVar(
-	TEXT("RHI.ShowProfilerAfterProfilingGPU"),
+	TEXT("r.ProfileGPU.ShowUI"),
 	1,
-	TEXT("Whether the profiler should be displayed after profiling the GPU."),
+	TEXT("Whether the user interface profiler should be displayed after profiling the GPU.\n")
+	TEXT("The results will always go to the log/console\n")
+	TEXT("0:off, 1:on (default)"),
 	ECVF_RenderThreadSafe);
 static TAutoConsoleVariable<float> GGPUHitchThresholdCVar(
 	TEXT("RHI.GPUHitchThreshold"),
