Fixed IsOpenGLPlatform

diff --git a/Engine/Source/Runtime/RHI/Public/RHIDefinitions.h b/Engine/Source/Runtime/RHI/Public/RHIDefinitions.h
--- a/Engine/Source/Runtime/RHI/Public/RHIDefinitions.h
+++ b/Engine/Source/Runtime/RHI/Public/RHIDefinitions.h
@@ -557,8 +557,8 @@ inline bool IsMobilePlatform(const EShaderPlatform Platform)
 
 inline bool IsOpenGLPlatform(const EShaderPlatform Platform)
 {
-	return Platform == SP_OPENGL_SM4 || Platform == SP_OPENGL_SM4_MAC || Platform == SP_OPENGL_SM5 || Platform == SP_OPENGL_PCES2 || Platform == SP_OPENGL_ES2 || Platform == SP_OPENGL_ES2_WEBGL || Platform == SP_OPENGL_ES2_IOS
-		|| SP_OPENGL_ES31_EXT;
+	return Platform == SP_OPENGL_SM4 || Platform == SP_OPENGL_SM4_MAC || Platform == SP_OPENGL_SM5 || Platform == SP_OPENGL_PCES2
+		|| Platform == SP_OPENGL_ES2 || Platform == SP_OPENGL_ES2_WEBGL || Platform == SP_OPENGL_ES2_IOS || Platform == SP_OPENGL_ES31_EXT;
 }
 
 inline bool IsConsolePlatform(const EShaderPlatform Platform)
