Added 'r.ProfileGPU.Root', which can be used to filter the hierarchy to a specified root

diff --git a/Engine/Source/Runtime/RHI/Private/GPUProfiler.cpp b/Engine/Source/Runtime/RHI/Private/GPUProfiler.cpp
--- a/Engine/Source/Runtime/RHI/Private/GPUProfiler.cpp
+++ b/Engine/Source/Runtime/RHI/Private/GPUProfiler.cpp
@@ -23,6 +23,12 @@ static TAutoConsoleVariable<FString> GProfileGPUPatternCVar(
 	TEXT("e.g. AmbientOcclusionSetup, AmbientOcclusion*, Ambient???lusion*, *"),
 	ECVF_Default);
 
+static TAutoConsoleVariable<FString> GProfileGPURootCVar(
+	TEXT("r.ProfileGPU.Root"),
+	TEXT("*"),
+	TEXT("Allows to filter the tree when using ProfileGPU, the pattern match is case sensitive."),
+	ECVF_Default);
+
 struct FNodeStatsCompare
 {
 	/** Sorts nodes by descending durations. */
@@ -69,7 +75,7 @@ static void GatherStatsEventNode(FGPUProfilerEventNode* Node, int32 Depth, TMap<
 }
 
 /** Recursively dumps stats for each node with a depth first traversal. */
-static void DumpStatsEventNode(FGPUProfilerEventNode* Node, float RootResult, int32 Depth, int32& NumNodes, int32& NumDraws)
+static void DumpStatsEventNode(FGPUProfilerEventNode* Node, float RootResult, int32 Depth, const FWildcardString& WildcardFilter, bool bParentMatchedFilter, int32& NumNodes, int32& NumDraws)
 {
 	NumNodes++;
 	if (Node->NumDraws > 0 || Node->Children.Num() > 0)
@@ -80,16 +86,21 @@ static void DumpStatsEventNode(FGPUProfilerEventNode* Node, float RootResult, in
 
 		const int32 EffectiveDepth = FMath::Max(Depth - 1, 0);
 
-		// Print information about this node, padded to its depth in the tree
-		UE_LOG(LogRHI, Warning, TEXT("%s%4.1f%%%5.2fms   %s %u draws %u prims %u verts"), 
-			*FString(TEXT("")).LeftPad(EffectiveDepth * 3), 
-			Percent,
-			Node->TimingResult,
-			*Node->Name,
-			Node->NumDraws,
-			Node->NumPrimitives,
-			Node->NumVertices
-			);
+		const bool bMatchesFilter = bParentMatchedFilter || WildcardFilter.IsMatch(Node->Name);
+
+		if (bMatchesFilter)
+		{
+			// Print information about this node, padded to its depth in the tree
+			UE_LOG(LogRHI, Warning, TEXT("%s%4.1f%%%5.2fms   %s %u draws %u prims %u verts"), 
+				*FString(TEXT("")).LeftPad(EffectiveDepth * 3), 
+				Percent,
+				Node->TimingResult,
+				*Node->Name,
+				Node->NumDraws,
+				Node->NumPrimitives,
+				Node->NumVertices
+				);
+		}
 
 		float TotalChildTime = 0;
 		uint32 TotalChildDraws = 0;
@@ -99,7 +110,7 @@ static void DumpStatsEventNode(FGPUProfilerEventNode* Node, float RootResult, in
 
 			int32 NumChildDraws = 0;
 			// Traverse children
-			DumpStatsEventNode(Node->Children[ChildIndex], RootResult, Depth + 1, NumNodes, NumChildDraws);
+			DumpStatsEventNode(Node->Children[ChildIndex], RootResult, Depth + 1, WildcardFilter, bMatchesFilter, NumNodes, NumChildDraws);
 			NumDraws += NumChildDraws;
 
 			TotalChildTime += ChildNode->TimingResult;
@@ -110,7 +121,7 @@ static void DumpStatsEventNode(FGPUProfilerEventNode* Node, float RootResult, in
 		const float UnaccountedPercent = UnaccountedTime * 100.0f / (RootResult * 1000.0f);
 
 		// Add an 'Unaccounted' node if necessary to show time spent in the current node that is not in any of its children
-		if (Node->Children.Num() > 0 && TotalChildDraws > 0 && (UnaccountedPercent > 2.0f || UnaccountedTime > .2f))
+		if (bMatchesFilter && Node->Children.Num() > 0 && TotalChildDraws > 0 && (UnaccountedPercent > 2.0f || UnaccountedTime > .2f))
 		{
 			UE_LOG(LogRHI, Warning, TEXT("%s%4.1f%%%5.2fms Unaccounted"), 
 				*FString(TEXT("")).LeftPad((EffectiveDepth + 1) * 3), 
@@ -186,11 +197,15 @@ void FGPUProfilerEventNodeFrame::DumpEventTree()
 			GatherStatsEventNode(EventTree[BaseNodeIndex], 0, EventHistogram);
 		}
 
+		static IConsoleVariable* CVar2 = IConsoleManager::Get().FindConsoleVariable(TEXT("r.ProfileGPU.Root"));
+		FString RootWildcardString = CVar2->GetString(); 
+		FWildcardString RootWildcard(RootWildcardString);
+
 		int32 NumNodes = 0;
 		int32 NumDraws = 0;
 		for (int32 BaseNodeIndex = 0; BaseNodeIndex < EventTree.Num(); BaseNodeIndex++)
 		{
-			DumpStatsEventNode(EventTree[BaseNodeIndex], RootResult, 0, NumNodes, NumDraws);
+			DumpStatsEventNode(EventTree[BaseNodeIndex], RootResult, 0, RootWildcard, false, NumNodes, NumDraws);
 		}
 
 		//@todo - calculate overhead instead of hardcoding
@@ -200,61 +215,64 @@ void FGPUProfilerEventNodeFrame::DumpEventTree()
 		UE_LOG(LogRHI, Warning, TEXT(""));
 		UE_LOG(LogRHI, Warning, TEXT(""));
 
-		// Sort descending based on node duration
-		EventHistogram.ValueSort( FNodeStatsCompare() );
-
-		// Log stats about the node histogram
-		UE_LOG(LogRHI, Warning, TEXT("Node histogram %u buckets"), EventHistogram.Num());
-
-		static IConsoleVariable* CVar = IConsoleManager::Get().FindConsoleVariable(TEXT("r.ProfileGPU.Pattern"));
-
-		// bad: reading on render thread but we don't support ECVF_RenderThreadSafe on strings yet
-		// It's very unlikely to cause a problem as the cvar is only changes by the user.
-		FString WildcardString = CVar->GetString(); 
-
-		const float ThresholdInMS = 5.0f;
-
-		if(WildcardString == FString(TEXT("*")))
+		if (RootWildcardString == TEXT("*"))
 		{
-			// disable Wildcard functionality
-			WildcardString.Empty();
-		}
+			// Sort descending based on node duration
+			EventHistogram.ValueSort( FNodeStatsCompare() );
 
-		if(WildcardString.IsEmpty())
-		{
-			UE_LOG(LogRHI, Warning, TEXT(" r.ProfileGPU.Pattern = '*' (using threshold of %g ms)"), ThresholdInMS);
-		}
-		else
-		{
-			UE_LOG(LogRHI, Warning, TEXT(" r.ProfileGPU.Pattern = '%s' (not using time threshold)"), *WildcardString);
-		}
+			// Log stats about the node histogram
+			UE_LOG(LogRHI, Warning, TEXT("Node histogram %u buckets"), EventHistogram.Num());
 
-		FWildcardString Wildcard(WildcardString);
+			static IConsoleVariable* CVar = IConsoleManager::Get().FindConsoleVariable(TEXT("r.ProfileGPU.Pattern"));
 
-		int32 NumNotShown = 0;
-		for (TMap<FString, FGPUProfilerEventNodeStats>::TIterator It(EventHistogram); It; ++It)
-		{
-			const FGPUProfilerEventNodeStats& NodeStats = It.Value();
+			// bad: reading on render thread but we don't support ECVF_RenderThreadSafe on strings yet
+			// It's very unlikely to cause a problem as the cvar is only changes by the user.
+			FString WildcardString = CVar->GetString(); 
 
-			bool bDump = NodeStats.TimingResult > RootResult * ThresholdInMS;
+			const float ThresholdInMS = 5.0f;
 
-			if(!Wildcard.IsEmpty())
+			if(WildcardString == FString(TEXT("*")))
 			{
-				// if a Wildcard string was specified, we want to always dump all entries
-				bDump = Wildcard.IsMatch(*It.Key());
+				// disable Wildcard functionality
+				WildcardString.Empty();
 			}
 
-			if (bDump)
+			if(WildcardString.IsEmpty())
 			{
-				UE_LOG(LogRHI, Warning, TEXT("   %.2fms   %s   Events %u   Draws %u"), NodeStats.TimingResult, *It.Key(), NodeStats.NumEvents, NodeStats.NumDraws);
+				UE_LOG(LogRHI, Warning, TEXT(" r.ProfileGPU.Pattern = '*' (using threshold of %g ms)"), ThresholdInMS);
 			}
 			else
 			{
-				NumNotShown++;
+				UE_LOG(LogRHI, Warning, TEXT(" r.ProfileGPU.Pattern = '%s' (not using time threshold)"), *WildcardString);
 			}
-		}
 
-		UE_LOG(LogRHI, Warning, TEXT("   %u buckets not shown"), NumNotShown);
+			FWildcardString Wildcard(WildcardString);
+
+			int32 NumNotShown = 0;
+			for (TMap<FString, FGPUProfilerEventNodeStats>::TIterator It(EventHistogram); It; ++It)
+			{
+				const FGPUProfilerEventNodeStats& NodeStats = It.Value();
+
+				bool bDump = NodeStats.TimingResult > RootResult * ThresholdInMS;
+
+				if(!Wildcard.IsEmpty())
+				{
+					// if a Wildcard string was specified, we want to always dump all entries
+					bDump = Wildcard.IsMatch(*It.Key());
+				}
+
+				if (bDump)
+				{
+					UE_LOG(LogRHI, Warning, TEXT("   %.2fms   %s   Events %u   Draws %u"), NodeStats.TimingResult, *It.Key(), NodeStats.NumEvents, NodeStats.NumDraws);
+				}
+				else
+				{
+					NumNotShown++;
+				}
+			}
+
+			UE_LOG(LogRHI, Warning, TEXT("   %u buckets not shown"), NumNotShown);
+		}
 
 #if !UE_BUILD_SHIPPING
 		// Create and display profile visualizer data
