UE4 - Split Xbox RHI off. Deffered contexts for XB1, parallel velocity rendering fixed and reenabled.RHIThread on by default on platforms that support it.

diff --git a/Engine/Config/BaseInput.ini b/Engine/Config/BaseInput.ini
--- a/Engine/Config/BaseInput.ini
+++ b/Engine/Config/BaseInput.ini
@@ -136,6 +136,8 @@ MaxScrollbackSize=1024
 +ManualAutoCompleteList=(Command="DisableAllScreenMessages",Desc="Disables all on-screen warnings/messages")
 +ManualAutoCompleteList=(Command="EnableAllScreenMessages",Desc="Enables all on-screen warnings/messages")
 +ManualAutoCompleteList=(Command="ToggleAllScreenMessages",Desc="Toggles display state of all on-screen warnings/messages")
++ManualAutoCompleteList=(Command="ToggleRHIThread",Desc="Toggles the RHI thread for platforms that have it")
++ManualAutoCompleteList=(Command="ToggleRenderingThread",Desc="Toggles the rendering thread for platforms that have it")
 +ManualAutoCompleteList=(Command="CaptureMode",Desc="Toggles display state of all on-screen warnings/messages")
 +ManualAutoCompleteList=(Command="ShowDebug None",Desc="ShowDebug None (Toggles ShowDebug w/ current debug type selection)")
 +ManualAutoCompleteList=(Command="ShowDebug Reset",Desc="ShowDebug Reset (Turns off ShowDebug, and clears debug type selection)")
diff --git a/Engine/Source/Runtime/Core/Private/Async/TaskGraph.cpp b/Engine/Source/Runtime/Core/Private/Async/TaskGraph.cpp
--- a/Engine/Source/Runtime/Core/Private/Async/TaskGraph.cpp
+++ b/Engine/Source/Runtime/Core/Private/Async/TaskGraph.cpp
@@ -670,7 +670,7 @@ public:
 		}
 		
 		NumNamedThreads = LastExternalThread + 1;
-		NumThreads = FMath::Max<int32>(FMath::Min<int32>(InNumThreads,MAX_THREADS),NumNamedThreads + 1);
+		NumThreads = FMath::Max<int32>(FMath::Min<int32>(InNumThreads + NumNamedThreads,MAX_THREADS),NumNamedThreads + 1);
 		// Cap number of extra threads to the platform worker thread count
 		NumThreads = FMath::Min(NumThreads, NumNamedThreads + FPlatformMisc::NumberOfWorkerThreadsToSpawn());
 		UE_LOG(LogTaskGraph, Log, TEXT("Started task graph with %d named threads and %d total threads."), NumNamedThreads, NumThreads);
diff --git a/Engine/Source/Runtime/Core/Public/HAL/Platform.h b/Engine/Source/Runtime/Core/Public/HAL/Platform.h
--- a/Engine/Source/Runtime/Core/Public/HAL/Platform.h
+++ b/Engine/Source/Runtime/Core/Public/HAL/Platform.h
@@ -272,6 +272,19 @@
 	#define PLATFORM_USES_ANSI_STRING_FOR_EXTERNAL_PROFILING 1
 #endif
 
+#ifndef PLATFORM_RHI_USES_CONTEXT_OBJECT
+	#define PLATFORM_RHI_USES_CONTEXT_OBJECT 0
+#endif
+
+#ifndef PLATFORM_SUPPORTS_PARALLEL_RHI_EXECUTE
+	#define PLATFORM_SUPPORTS_PARALLEL_RHI_EXECUTE		0
+#endif
+
+
+#if PLATFORM_SUPPORTS_PARALLEL_RHI_EXECUTE && !PLATFORM_RHI_USES_CONTEXT_OBJECT
+	#error "PLATFORM_SUPPORTS_PARALLEL_RHI_EXECUTE requires PLATFORM_RHI_USES_CONTEXT_OBJECT"
+#endif
+
 //These are deprecated old defines that we want to make sure are not used
 #define CONSOLE (#)
 #define MOBILE (#)
diff --git a/Engine/Source/Runtime/Engine/Private/GPUSkinVertexFactory.cpp b/Engine/Source/Runtime/Engine/Private/GPUSkinVertexFactory.cpp
--- a/Engine/Source/Runtime/Engine/Private/GPUSkinVertexFactory.cpp
+++ b/Engine/Source/Runtime/Engine/Private/GPUSkinVertexFactory.cpp
@@ -437,14 +437,13 @@ public:
 						BoneIndexOffsetValue
 						);
 				}
-
+				FScopeLock Lock(&ShaderData.OldBoneDataLock);
 				// if we haven't copied the data yet we skip the update (e.g. split screen)
 				if(ShaderData.IsOldBoneDataUpdateNeeded(View.FrameNumber))
 				{
 					const FGPUBaseSkinVertexFactory* GPUVertexFactory = (const FGPUBaseSkinVertexFactory*)VertexFactory;
 
 					// copy the bone data and tell the instance where it can pick it up next frame
-					check(IsInRenderingThread());
 					// append data to a buffer we bind next frame to read old matrix data for motion blur
 					uint32 OldBoneDataStartIndex = GPrevPerBoneMotionBlur.AppendData(ShaderData.BoneMatrices.GetData(), ShaderData.BoneMatrices.Num());
 					GPUVertexFactory->SetOldBoneDataStartIndex(View.FrameNumber, OldBoneDataStartIndex);
diff --git a/Engine/Source/Runtime/Engine/Private/StaticBoundShaderState.cpp b/Engine/Source/Runtime/Engine/Private/StaticBoundShaderState.cpp
--- a/Engine/Source/Runtime/Engine/Private/StaticBoundShaderState.cpp
+++ b/Engine/Source/Runtime/Engine/Private/StaticBoundShaderState.cpp
@@ -56,7 +56,12 @@ FBoundShaderStateRHIParamRef FGlobalBoundShaderStateResource::GetInitializedRHI(
 	// Create the bound shader state if it hasn't been cached yet.
 	if(!IsValidRef(BoundShaderState))
 	{
-		BoundShaderState = CreateBoundShaderState_Internal(
+		BoundShaderState = 
+#if PLATFORM_SUPPORTS_PARALLEL_RHI_EXECUTE
+			RHICreateBoundShaderState(
+#else
+			CreateBoundShaderState_Internal(
+#endif
 			VertexDeclaration,
 			VertexShader,
 			FHullShaderRHIRef(), 
@@ -68,7 +73,12 @@ FBoundShaderStateRHIParamRef FGlobalBoundShaderStateResource::GetInitializedRHI(
 	// Only people working on shaders (and therefore have LogShaders unsuppressed) will want to see these errors
 	else if (!GUsingNullRHI && UE_LOG_ACTIVE(LogShaders, Warning))
 	{
-		FBoundShaderStateRHIRef TempBoundShaderState = CreateBoundShaderState_Internal(
+		FBoundShaderStateRHIRef TempBoundShaderState = 
+#if PLATFORM_SUPPORTS_PARALLEL_RHI_EXECUTE
+			RHICreateBoundShaderState(
+#else
+			CreateBoundShaderState_Internal(
+#endif
 			VertexDeclaration,
 			VertexShader,
 			FHullShaderRHIRef(),
diff --git a/Engine/Source/Runtime/Engine/Public/GPUSkinVertexFactory.h b/Engine/Source/Runtime/Engine/Public/GPUSkinVertexFactory.h
--- a/Engine/Source/Runtime/Engine/Public/GPUSkinVertexFactory.h
+++ b/Engine/Source/Runtime/Engine/Public/GPUSkinVertexFactory.h
@@ -329,13 +329,13 @@ public:
 		void SetOldBoneData(uint32 FrameNumber, uint32 Index) const
 		{
 			// keep the one from last from, someone might read want to read that
-			if(OldBoneFrameNumber[0] + 1 == FrameNumber)
+			if (OldBoneFrameNumber[0] + 1 == FrameNumber)
 			{
 				// [1] has the right data
 				OldBoneFrameNumber[1] = FrameNumber;
 				OldBoneDataStartIndex[1] = Index;
 			}
-			else
+			else 
 			{
 				// [0] has the right data
 				OldBoneFrameNumber[0] = FrameNumber;
@@ -377,6 +377,7 @@ public:
 			return BoneBuffer;
 		}
 
+		mutable FCriticalSection OldBoneDataLock;
 	private:
 
 		// the following members can be stored in less bytes (after some adjustments)
diff --git a/Engine/Source/Runtime/NullDrv/Public/NullRHI.h b/Engine/Source/Runtime/NullDrv/Public/NullRHI.h
--- a/Engine/Source/Runtime/NullDrv/Public/NullRHI.h
+++ b/Engine/Source/Runtime/NullDrv/Public/NullRHI.h
@@ -6,6 +6,9 @@
 /** A null implementation of the dynamically bound RHI. */
 class FNullDynamicRHI
 	: public FDynamicRHI
+#if PLATFORM_RHI_USES_CONTEXT_OBJECT
+	, public IRHICommandContext
+#endif
 {
 public:
 
diff --git a/Engine/Source/Runtime/RHI/Private/BoundShaderStateCache.cpp b/Engine/Source/Runtime/RHI/Private/BoundShaderStateCache.cpp
--- a/Engine/Source/Runtime/RHI/Private/BoundShaderStateCache.cpp
+++ b/Engine/Source/Runtime/RHI/Private/BoundShaderStateCache.cpp
@@ -16,6 +16,11 @@ static FBoundShaderStateCache& GetBoundShaderStateCache()
 	return BoundShaderStateCache;
 }
 
+#if PLATFORM_SUPPORTS_PARALLEL_RHI_EXECUTE
+static FCriticalSection BoundShaderStateCacheLock;
+#endif
+
+
 FCachedBoundShaderStateLink::FCachedBoundShaderStateLink(
 	FVertexDeclarationRHIParamRef VertexDeclaration,
 	FVertexShaderRHIParamRef VertexShader,
@@ -29,7 +34,9 @@ FCachedBoundShaderStateLink::FCachedBoundShaderStateLink(
 	Key(VertexDeclaration,VertexShader,PixelShader,HullShader,DomainShader,GeometryShader)
 {
 	// Add this bound shader state to the cache.
+#if !PLATFORM_SUPPORTS_PARALLEL_RHI_EXECUTE
 	GetBoundShaderStateCache().Add(Key,this);
+#endif
 }
 
 FCachedBoundShaderStateLink::FCachedBoundShaderStateLink(
@@ -42,14 +49,57 @@ FCachedBoundShaderStateLink::FCachedBoundShaderStateLink(
 	Key(VertexDeclaration,VertexShader,PixelShader)
 {
 	// Add this bound shader state to the cache.
+#if !PLATFORM_SUPPORTS_PARALLEL_RHI_EXECUTE
 	GetBoundShaderStateCache().Add(Key,this);
+#endif
 }
 
 FCachedBoundShaderStateLink::~FCachedBoundShaderStateLink()
 {
+#if !PLATFORM_SUPPORTS_PARALLEL_RHI_EXECUTE
 	GetBoundShaderStateCache().Remove(Key);
+#endif
 }
 
+#if PLATFORM_SUPPORTS_PARALLEL_RHI_EXECUTE
+void FCachedBoundShaderStateLink::AddToCache()
+{
+	FScopeLock Lock(&BoundShaderStateCacheLock);
+	GetBoundShaderStateCache().Add(Key,this);
+}
+void FCachedBoundShaderStateLink::RemoveFromCache()
+{
+	FScopeLock Lock(&BoundShaderStateCacheLock);
+	GetBoundShaderStateCache().Remove(Key);
+}
+#endif
+
+
+#if PLATFORM_SUPPORTS_PARALLEL_RHI_EXECUTE
+FBoundShaderStateRHIRef GetCachedBoundShaderState_Threadsafe(
+	FVertexDeclarationRHIParamRef VertexDeclaration,
+	FVertexShaderRHIParamRef VertexShader,
+	FPixelShaderRHIParamRef PixelShader,
+	FHullShaderRHIParamRef HullShader,
+	FDomainShaderRHIParamRef DomainShader,
+	FGeometryShaderRHIParamRef GeometryShader
+	)
+{
+	FScopeLock Lock(&BoundShaderStateCacheLock);
+	// Find the existing bound shader state in the cache.
+	FCachedBoundShaderStateLink* CachedBoundShaderStateLink = GetBoundShaderStateCache().FindRef(
+		FBoundShaderStateKey(VertexDeclaration,VertexShader,PixelShader,HullShader,DomainShader,GeometryShader)
+		);
+	if(CachedBoundShaderStateLink)
+	{
+		// If we've already created a bound shader state with these parameters, reuse it.
+		return CachedBoundShaderStateLink->BoundShaderState;
+	}
+	return FBoundShaderStateRHIRef();
+}
+
+#else
+
 FCachedBoundShaderStateLink* GetCachedBoundShaderState(
 	FVertexDeclarationRHIParamRef VertexDeclaration,
 	FVertexShaderRHIParamRef VertexShader,
@@ -64,3 +114,5 @@ FCachedBoundShaderStateLink* GetCachedBoundShaderState(
 		FBoundShaderStateKey(VertexDeclaration,VertexShader,PixelShader,HullShader,DomainShader,GeometryShader)
 		);
 }
+
+#endif
diff --git a/Engine/Source/Runtime/RHI/Private/DynamicRHI.cpp b/Engine/Source/Runtime/RHI/Private/DynamicRHI.cpp
--- a/Engine/Source/Runtime/RHI/Private/DynamicRHI.cpp
+++ b/Engine/Source/Runtime/RHI/Private/DynamicRHI.cpp
@@ -29,6 +29,9 @@ void InitNullRHI()
 
 	GDynamicRHI = DynamicRHIModule->CreateRHI();
 	GDynamicRHI->Init();
+#if PLATFORM_RHI_USES_CONTEXT_OBJECT
+	GRHICommandList.GetImmediateCommandList().SetContext(GDynamicRHI->RHIGetDefaultContext());
+#endif
 	GUsingNullRHI = true;
 	GRHISupportsTextureStreaming = false;
 }
@@ -49,6 +52,9 @@ void RHIInit(bool bHasEditorToken)
 			if (GDynamicRHI)
 			{
 				GDynamicRHI->Init();
+#if PLATFORM_RHI_USES_CONTEXT_OBJECT
+				GRHICommandList.GetImmediateCommandList().SetContext(GDynamicRHI->RHIGetDefaultContext());
+#endif
 			}
 #if PLATFORM_ALLOW_NULL_RHI
 			else
@@ -74,13 +80,18 @@ void RHIExit()
 	}
 }
 
-// Implement the static RHI methods that call the dynamic RHI.
-#define DEFINE_RHIMETHOD_CMDLIST(Type,Name,ParameterTypesAndNames,ParameterNames,ReturnStatement,NullImplementation) \
-	RHI_API Type Name##_Internal ParameterTypesAndNames \
+#if PLATFORM_RHI_USES_CONTEXT_OBJECT
+	#define DEFINE_RHIMETHOD_CMDLIST(Type,Name,ParameterTypesAndNames,ParameterNames,ReturnStatement,NullImplementation)
+#else
+	// Implement the static RHI methods that call the dynamic RHI.
+	#define DEFINE_RHIMETHOD_CMDLIST(Type,Name,ParameterTypesAndNames,ParameterNames,ReturnStatement,NullImplementation) \
+		RHI_API Type Name##_Internal ParameterTypesAndNames \
 	{ \
 		check(GDynamicRHI); \
 		ReturnStatement GDynamicRHI->RHI##Name ParameterNames; \
 	}
+#endif
+
 #define DEFINE_RHIMETHOD_GLOBAL(Type,Name,ParameterTypesAndNames,ParameterNames,ReturnStatement,NullImplementation) \
 	RHI_API Type Name##_Internal ParameterTypesAndNames \
 	{ \
diff --git a/Engine/Source/Runtime/RHI/Private/GPUProfiler.cpp b/Engine/Source/Runtime/RHI/Private/GPUProfiler.cpp
--- a/Engine/Source/Runtime/RHI/Private/GPUProfiler.cpp
+++ b/Engine/Source/Runtime/RHI/Private/GPUProfiler.cpp
@@ -316,6 +316,7 @@ void FGPUProfiler::PushEvent(const TCHAR* Name)
 {
 	if (bTrackingEvents)
 	{
+		check(IsInRenderingThread() || IsInRHIThread());
 		if (CurrentEventNode)
 		{
 			// Add to the current node's children
@@ -339,7 +340,7 @@ void FGPUProfiler::PopEvent()
 {
 	if (bTrackingEvents)
 	{
-		check(CurrentEventNode);
+		check(CurrentEventNode && (IsInRenderingThread() || IsInRHIThread()));
 		// Stop timing the current node and move one level up the tree
 		CurrentEventNode->StopTiming();
 		CurrentEventNode = CurrentEventNode->Parent;
diff --git a/Engine/Source/Runtime/RHI/Private/RHICommandList.cpp b/Engine/Source/Runtime/RHI/Private/RHICommandList.cpp
--- a/Engine/Source/Runtime/RHI/Private/RHICommandList.cpp
+++ b/Engine/Source/Runtime/RHI/Private/RHICommandList.cpp
@@ -45,6 +45,16 @@ TAutoConsoleVariable<int32> CVarRHICmdWidth(
 	8,
 	TEXT("Number of threads."));
 
+#if PLATFORM_SUPPORTS_PARALLEL_RHI_EXECUTE
+static TAutoConsoleVariable<int32> CVarRHICmdUseDeferredContexts(
+	TEXT("r.RHICmdUseDeferredContexts"),
+	1,
+	TEXT("True to use deferred contexts to parallelize command list execution.\n"));
+static TAutoConsoleVariable<int32> CVarRHIDeferredContextWidth(
+	TEXT("r.RHIDeferredContextWidth"),
+	99999,
+	TEXT("Number of pieces to break a parallel translate task into. A very large number is one task per original cmd list.\n"));
+#endif
 
 RHI_API FRHICommandListExecutor GRHICommandList;
 
@@ -55,21 +65,60 @@ static FGraphEventRef RenderThreadSublistDispatchTask;
 
 DECLARE_CYCLE_STAT(TEXT("RHI Thread Execute"), STAT_RHIThreadExecute, STATGROUP_RHICMDLIST);
 
+static TStatId GCurrentExecuteStat;
+
+struct FRHICommandStat : public FRHICommand<FRHICommandStat>
+{
+	TStatId CurrentExecuteStat;
+	FORCEINLINE_DEBUGGABLE FRHICommandStat(TStatId InCurrentExecuteStat)
+		: CurrentExecuteStat(InCurrentExecuteStat)
+	{
+	}
+	void Execute(FRHICommandListBase& CmdList)
+	{
+		GCurrentExecuteStat = CurrentExecuteStat;
+	}
+};
+
+void FRHICommandListImmediate::SetCurrentStat(TStatId Stat)
+{
+	new (AllocCommand<FRHICommandStat>()) FRHICommandStat(Stat);
+}
+
 void FRHICommandListExecutor::ExecuteInner_DoExecute(FRHICommandListBase& CmdList)
 {
 	CmdList.bExecuting = true;
+#if PLATFORM_RHI_USES_CONTEXT_OBJECT
+	check(CmdList.Context);
+#endif
+
+	bool bDoStats = false;
+	STAT(bDoStats = FThreadStats::IsCollectingData() && (IsInRenderingThread() || IsInRHIThread());)
+	FRHICommandListIterator Iter(CmdList);
+	if (bDoStats)
 	{
-		FRHICommandListIterator Iter(CmdList);
 		while (Iter.HasCommandsLeft())
 		{
 			FRHICommandBase* Cmd = Iter.NextCommand();
 			//FPlatformMisc::Prefetch(Cmd->Next);
-			Cmd->CallExecuteAndDestruct();
+			FScopeCycleCounter Scope(GCurrentExecuteStat);
+			Cmd->CallExecuteAndDestruct(CmdList);
+		}
+	}
+	else
+	{
+		while (Iter.HasCommandsLeft())
+		{
+			FRHICommandBase* Cmd = Iter.NextCommand();
+			//FPlatformMisc::Prefetch(Cmd->Next);
+			Cmd->CallExecuteAndDestruct(CmdList);
 		}
 	}
 	CmdList.Reset();
 }
 
+
+
 class FExecuteRHIThreadTask
 {
 	FRHICommandListBase* RHICmdList;
@@ -259,7 +308,10 @@ static FORCEINLINE bool IsInRenderingOrRHIThread()
 
 void FRHICommandListExecutor::ExecuteList(FRHICommandListBase& CmdList)
 {
-	check(IsInRenderingOrRHIThread() && &CmdList != &GetImmediateCommandList());
+	check(&CmdList != &GetImmediateCommandList() &&
+		(PLATFORM_SUPPORTS_PARALLEL_RHI_EXECUTE || IsInRenderingOrRHIThread())
+		);
+
 	if (IsInRenderingThread() && !GetImmediateCommandList().IsExecuting()) // don't flush if this is a recursive call and we are already executing the immediate command list
 	{
 		GetImmediateCommandList().Flush();
@@ -421,12 +473,118 @@ void FRHICommandListBase::Reset()
 	NumCommands = 0;
 	Root = nullptr;
 	CommandLink = &Root;
+#if PLATFORM_RHI_USES_CONTEXT_OBJECT
+	static_assert(USE_DYNAMIC_RHI, "PLATFORM_RHI_USES_CONTEXT_OBJECT current requires a dynamic RHI");
+	Context = GDynamicRHI ? RHIGetDefaultContext() : nullptr;
+#else
+	Context = nullptr;
+#endif
 #if USE_RHICOMMAND_STATE_REDUCTION
 	StateCache = nullptr;
 #endif
 	UID = GRHICommandList.UIDCounter.Increment();
 }
 
+
+#if PLATFORM_SUPPORTS_PARALLEL_RHI_EXECUTE
+DECLARE_CYCLE_STAT(TEXT("Parallel Async Chain Translate"), STAT_ParallelChainTranslate, STATGROUP_RHICMDLIST);
+
+class FParallelTranslateCommandList
+{
+	FRHICommandListBase** RHICmdLists;
+	int32 NumCommandLists;
+	IRHICommandContextContainer* ContextContainer;
+public:
+
+	FParallelTranslateCommandList(FRHICommandListBase** InRHICmdLists, int32 InNumCommandLists, IRHICommandContextContainer* InContextContainer)
+		: RHICmdLists(InRHICmdLists)
+		, NumCommandLists(InNumCommandLists)
+		, ContextContainer(InContextContainer)
+	{
+		check(RHICmdLists && ContextContainer && NumCommandLists);
+	}
+
+	FORCEINLINE TStatId GetStatId() const
+	{
+		RETURN_QUICK_DECLARE_CYCLE_STAT(FParallelTranslateCommandList, STATGROUP_TaskGraphTasks);
+	}
+
+	ENamedThreads::Type GetDesiredThread()
+	{
+		return ENamedThreads::AnyThread;
+	}
+
+	static ESubsequentsMode::Type GetSubsequentsMode() { return ESubsequentsMode::TrackSubsequents; }
+
+	void DoTask(ENamedThreads::Type CurrentThread, const FGraphEventRef& MyCompletionGraphEvent)
+	{
+		SCOPE_CYCLE_COUNTER(STAT_ParallelChainTranslate);
+		check(ContextContainer && RHICmdLists);
+		IRHICommandContext* Context = ContextContainer->GetContext();
+		check(Context);
+		for (int32 Index = 0; Index < NumCommandLists; Index++)
+		{
+			RHICmdLists[Index]->SetContext(Context);
+			delete RHICmdLists[Index];
+		}
+		ContextContainer->FinishContext();
+	}
+};
+
+DECLARE_DWORD_COUNTER_STAT(TEXT("Num Parallel Async Chains Links"), STAT_ParallelChainLinkCount, STATGROUP_RHICMDLIST);
+DECLARE_CYCLE_STAT(TEXT("Wait for Parallel Async CmdList"), STAT_ParallelChainWait, STATGROUP_RHICMDLIST);
+DECLARE_CYCLE_STAT(TEXT("Parallel Async Chain Execute"), STAT_ParallelChainExecute, STATGROUP_RHICMDLIST);
+
+struct FRHICommandWaitForAndSubmitSubListParallel : public FRHICommand<FRHICommandWaitForAndSubmitSubListParallel>
+{
+	FGraphEventRef* TranslateCompletionEvents;
+	IRHICommandContextContainer** ContextContainers;
+	int32 Num;
+
+	FORCEINLINE_DEBUGGABLE FRHICommandWaitForAndSubmitSubListParallel(FGraphEventRef* InTranslateCompletionEvents, IRHICommandContextContainer** InContextContainers, int32 InNum)
+		: TranslateCompletionEvents(InTranslateCompletionEvents)
+		, ContextContainers(InContextContainers)
+		, Num(InNum)
+	{
+		check(TranslateCompletionEvents && ContextContainers && Num);
+	}
+	void Execute(FRHICommandListBase& CmdList)
+	{
+		check(TranslateCompletionEvents && ContextContainers && Num);
+		INC_DWORD_STAT_BY(STAT_ParallelChainLinkCount, Num);
+		for (int32 Index = 0; Index < Num; Index++)
+		{
+			FGraphEventRef& EventToWaitFor = TranslateCompletionEvents[Index];
+			IRHICommandContextContainer* ContextContainer = ContextContainers[Index];
+
+			if (EventToWaitFor.GetReference() && !EventToWaitFor->IsComplete())
+			{
+				SCOPE_CYCLE_COUNTER(STAT_ParallelChainWait);
+				if (IsInRenderingThread())
+				{
+					FTaskGraphInterface::Get().WaitUntilTaskCompletes(EventToWaitFor, ENamedThreads::RenderThread_Local);
+				}
+				else if (IsInRHIThread())
+				{
+					FTaskGraphInterface::Get().WaitUntilTaskCompletes(EventToWaitFor, ENamedThreads::RHIThread);
+				}
+				else
+				{
+					check(0);
+				}
+			}
+			EventToWaitFor.~FGraphEventRef();
+			{
+				SCOPE_CYCLE_COUNTER(STAT_ParallelChainExecute);
+				ContextContainer->SubmitAndFreeContextContainer(Index, Num);
+			}
+		}
+	}
+};
+#endif
+
+
+
 DECLARE_DWORD_COUNTER_STAT(TEXT("Num Async Chains Links"), STAT_ChainLinkCount, STATGROUP_RHICMDLIST);
 DECLARE_CYCLE_STAT(TEXT("Wait for Async CmdList"), STAT_ChainWait, STATGROUP_RHICMDLIST);
 DECLARE_CYCLE_STAT(TEXT("Async Chain Execute"), STAT_ChainExecute, STATGROUP_RHICMDLIST);
@@ -440,10 +598,10 @@ struct FRHICommandWaitForAndSubmitSubList : public FRHICommand<FRHICommandWaitFo
 		, RHICmdList(InRHICmdList)
 	{
 	}
-	void Execute()
+	void Execute(FRHICommandListBase& CmdList)
 	{
 		INC_DWORD_STAT_BY(STAT_ChainLinkCount, 1);
-		if (!EventToWaitFor->IsComplete())
+		if (EventToWaitFor.GetReference() && !EventToWaitFor->IsComplete())
 		{
 			check(!GRHIThread || !IsInRHIThread()); // things should not be dispatched if they can't complete without further waits
 			SCOPE_CYCLE_COUNTER(STAT_ChainWait);
@@ -458,19 +616,109 @@ struct FRHICommandWaitForAndSubmitSubList : public FRHICommand<FRHICommandWaitFo
 		}
 		{
 			SCOPE_CYCLE_COUNTER(STAT_ChainExecute);
+			RHICmdList->CopyContext(CmdList);
 			delete RHICmdList;
 		}
 	}
 };
 
+void FRHICommandListBase::QueueParallelAsyncCommandListSubmit(FGraphEventRef* AnyThreadCompletionEvents, FRHICommandList** CmdLists, int32 Num)
+{
+	check(IsInRenderingThread() && IsImmediate() && Num);
+	if (GRHIThread)
+	{
+		FRHICommandListExecutor::GetImmediateCommandList().ImmediateFlush(EImmediateFlushType::DispatchToRHIThread); // we should start on the stuff before this async list
+	}
+#if PLATFORM_SUPPORTS_PARALLEL_RHI_EXECUTE
+	if (Num)
+	{
+		IRHICommandContextContainer* ContextContainer = nullptr;
+		if (CVarRHICmdUseDeferredContexts.GetValueOnAnyThread() >= 1)
+		{
+			ContextContainer = RHIGetCommandContextContainer(0);
+		}
+		if (ContextContainer)
+		{
+			int32 EffectiveThreads = FMath::Min<int32>(Num, CVarRHIDeferredContextWidth.GetValueOnRenderThread());
+
+			int32 Start = 0;
+			if (EffectiveThreads)
+			{
+
+				int32 NumPer = Num / EffectiveThreads;
+				int32 Extra = Num - NumPer * EffectiveThreads;
+
+				FGraphEventRef* OutTranslateCompletionEvents = (FGraphEventRef*)Alloc(sizeof(FGraphEventRef) * EffectiveThreads, ALIGNOF(FGraphEventRef));
+				IRHICommandContextContainer** ContextContainers = (IRHICommandContextContainer**)Alloc(sizeof(IRHICommandContextContainer*) * EffectiveThreads, ALIGNOF(IRHICommandContextContainer*));
+				FRHICommandListBase** RHICmdLists = (FRHICommandListBase**)Alloc(sizeof(FRHICommandListBase*) * Num, ALIGNOF(FRHICommandListBase*));
+
+				for (int32 ThreadIndex = 0; ThreadIndex < EffectiveThreads; ThreadIndex++)
+				{
+					int32 Last = Start + (NumPer - 1) + (ThreadIndex < Extra);
+					check(Last >= Start);
+
+					if (!ContextContainer)
+					{
+						ContextContainer = RHIGetCommandContextContainer(0);
+					}
+					check(ContextContainer);
+
+					FGraphEventArray Prereq;
+					for (int32 Index = Start; Index <= Last; Index++)
+					{
+						FGraphEventRef& AnyThreadCompletionEvent = AnyThreadCompletionEvents[Index];
+						FRHICommandList* CmdList = CmdLists[Index];
+						RHICmdLists[Index] = CmdList;
+						if (AnyThreadCompletionEvent.GetReference())
+						{
+							Prereq.Add(AnyThreadCompletionEvent);
+							OutstandingTasks.Add(AnyThreadCompletionEvent);
+						}
+					}
+					FGraphEventRef TranslateCompletionEvent = TGraphTask<FParallelTranslateCommandList>::CreateTask(&Prereq, ENamedThreads::RenderThread).ConstructAndDispatchWhenReady(&RHICmdLists[Start], 1 + Last - Start, ContextContainer);
+
+					new ((void*)&OutTranslateCompletionEvents[ThreadIndex]) FGraphEventRef(TranslateCompletionEvent);
+					ContextContainers[ThreadIndex] = ContextContainer;
+					ContextContainer = nullptr;
+					Start = Last + 1;
+				}
+				new (AllocCommand<FRHICommandWaitForAndSubmitSubListParallel>()) FRHICommandWaitForAndSubmitSubListParallel(OutTranslateCompletionEvents, ContextContainers, EffectiveThreads);
+				if (GRHIThread)
+				{
+					FRHICommandListExecutor::GetImmediateCommandList().ImmediateFlush(EImmediateFlushType::DispatchToRHIThread); // we don't want stuff after the async cmd list to be bundled with it
+				}
+				return;
+			}
+		}
+	}
+#endif
+	for (int32 Index = 0; Index < Num; Index++)
+	{
+		FGraphEventRef& AnyThreadCompletionEvent = AnyThreadCompletionEvents[Index];
+		FRHICommandList* CmdList = CmdLists[Index];
+		if (AnyThreadCompletionEvent.GetReference())
+		{
+			OutstandingTasks.Add(AnyThreadCompletionEvent);
+		}
+		new (AllocCommand<FRHICommandWaitForAndSubmitSubList>()) FRHICommandWaitForAndSubmitSubList(AnyThreadCompletionEvent, CmdList);
+	}
+	if (GRHIThread)
+	{
+		FRHICommandListExecutor::GetImmediateCommandList().ImmediateFlush(EImmediateFlushType::DispatchToRHIThread); // we don't want stuff after the async cmd list to be bundled with it
+	}
+}
+
 void FRHICommandListBase::QueueAsyncCommandListSubmit(FGraphEventRef& AnyThreadCompletionEvent, class FRHICommandList* CmdList)
 {
 	check(IsInRenderingThread() && IsImmediate());
-	OutstandingTasks.Add(AnyThreadCompletionEvent);
 	if (GRHIThread)
 	{
 		FRHICommandListExecutor::GetImmediateCommandList().ImmediateFlush(EImmediateFlushType::DispatchToRHIThread); // we should start on the stuff before this async list
 	}
+	if (AnyThreadCompletionEvent.GetReference())
+	{
+		OutstandingTasks.Add(AnyThreadCompletionEvent);
+	}
 	new (AllocCommand<FRHICommandWaitForAndSubmitSubList>()) FRHICommandWaitForAndSubmitSubList(AnyThreadCompletionEvent, CmdList);
 	if (GRHIThread)
 	{
@@ -491,7 +739,7 @@ struct FRHICommandWaitForAndSubmitRTSubList : public FRHICommand<FRHICommandWait
 		, RHICmdList(InRHICmdList)
 	{
 	}
-	void Execute()
+	void Execute(FRHICommandListBase& CmdList)
 	{
 		INC_DWORD_STAT_BY(STAT_RTChainLinkCount, 1);
 		{
@@ -516,6 +764,7 @@ struct FRHICommandWaitForAndSubmitRTSubList : public FRHICommand<FRHICommandWait
 		}
 		{
 			SCOPE_CYCLE_COUNTER(STAT_RTChainExecute);
+			RHICmdList->CopyContext(CmdList);
 			delete RHICmdList;
 		}
 	}
@@ -535,10 +784,11 @@ struct FRHICommandSubmitSubList : public FRHICommand<FRHICommandSubmitSubList>
 		: RHICmdList(InRHICmdList)
 	{
 	}
-	void Execute()
+	void Execute(FRHICommandListBase& CmdList)
 	{
 		INC_DWORD_STAT_BY(STAT_ChainLinkCount, 1);
 		SCOPE_CYCLE_COUNTER(STAT_ChainExecute);
+		RHICmdList->CopyContext(CmdList);
 		delete RHICmdList;
 	}
 };
@@ -554,7 +804,7 @@ void FRHICommandList::BeginDrawingViewport(FViewportRHIParamRef Viewport, FTextu
 	check(IsImmediate() && IsInRenderingThread());
 	if (Bypass())
 	{
-		BeginDrawingViewport_Internal(Viewport, RenderTargetRHI);
+		CMD_CONTEXT(BeginDrawingViewport)(Viewport, RenderTargetRHI);
 		return;
 	}
 	new (AllocCommand<FRHICommandBeginDrawingViewport>()) FRHICommandBeginDrawingViewport(Viewport, RenderTargetRHI);
@@ -565,7 +815,7 @@ void FRHICommandList::EndDrawingViewport(FViewportRHIParamRef Viewport, bool bPr
 	check(IsImmediate() && IsInRenderingThread());
 	if (Bypass())
 	{
-		EndDrawingViewport_Internal(Viewport, bPresent, bLockToVsync);
+		CMD_CONTEXT(EndDrawingViewport)(Viewport, bPresent, bLockToVsync);
 	}
 	else
 	{
@@ -583,7 +833,7 @@ void FRHICommandList::BeginFrame()
 {
 	if (Bypass())
 	{
-		BeginFrame_Internal();
+		CMD_CONTEXT(BeginFrame)();
 		return;
 	}
 	new (AllocCommand<FRHICommandBeginFrame>()) FRHICommandBeginFrame();
@@ -593,7 +843,7 @@ void FRHICommandList::EndFrame()
 {
 	if (Bypass())
 	{
-		EndFrame_Internal();
+		CMD_CONTEXT(EndFrame)();
 		return;
 	}
 	new (AllocCommand<FRHICommandEndFrame>()) FRHICommandEndFrame();
@@ -689,24 +939,10 @@ void FRHICommandListBase::HandleRTThreadTaskCompletion(const FGraphEventRef& MyC
 	check(!IsImmediate() && !IsInRenderingThread() && !IsInRHIThread())
 	for (int32 Index = 0; Index < RTTasks.Num(); Index++)
 	{
-		if (RTTasks[Index]->IsComplete())
-		{
-			RTTasks.RemoveAtSwap(Index, 1, false);
-			Index--;
-		}
-	}
-	if (RTTasks.Num())
-	{
-		FGraphEventRef WaitFor;
-		if (RTTasks.Num() > 1)
-		{
-			WaitFor = TGraphTask<FNullGraphTask>::CreateTask(&RTTasks).ConstructAndDispatchWhenReady(GET_STATID(STAT_HandleRTThreadTaskCompletion_Join));
-		}
-		else
+		if (!RTTasks[Index]->IsComplete())
 		{
-			WaitFor = RTTasks[0];
+			MyCompletionGraphEvent->DontCompleteUntil(RTTasks[Index]);
 		}
-		MyCompletionGraphEvent->DontCompleteUntil(WaitFor);
 	}
 	RTTasks.Empty();
 }
diff --git a/Engine/Source/Runtime/RHI/Public/BoundShaderStateCache.h b/Engine/Source/Runtime/RHI/Public/BoundShaderStateCache.h
--- a/Engine/Source/Runtime/RHI/Public/BoundShaderStateCache.h
+++ b/Engine/Source/Runtime/RHI/Public/BoundShaderStateCache.h
@@ -4,9 +4,7 @@
 	BoundShaderStateCache.h: Bound shader state cache definition.
 =============================================================================*/
 
-#ifndef __BOUNDSHADERCACHE_H__
-#define __BOUNDSHADERCACHE_H__
-
+#pragma once
 #include "RHI.h"
 
 /**
@@ -141,6 +139,10 @@ public:
 	FORCEINLINE FHullShaderRHIParamRef     GetHullShader() const     { return Key.GetHullShader(); }
 	FORCEINLINE FDomainShaderRHIParamRef   GetDomainShader() const   { return Key.GetDomainShader(); }
 	FORCEINLINE FGeometryShaderRHIParamRef GetGeometryShader() const { return Key.GetGeometryShader(); }
+#if PLATFORM_SUPPORTS_PARALLEL_RHI_EXECUTE
+	void AddToCache();
+	void RemoveFromCache();
+#endif
 
 private:
 
@@ -151,7 +153,8 @@ private:
  * Searches for a cached bound shader state.
  * @return If a bound shader state matching the parameters is cached, it is returned; otherwise NULL is returned.
  */
-extern RHI_API FCachedBoundShaderStateLink* GetCachedBoundShaderState(
+#if PLATFORM_SUPPORTS_PARALLEL_RHI_EXECUTE
+extern RHI_API FBoundShaderStateRHIRef GetCachedBoundShaderState_Threadsafe(
 	FVertexDeclarationRHIParamRef VertexDeclaration,
 	FVertexShaderRHIParamRef VertexShader,
 	FPixelShaderRHIParamRef PixelShader,
@@ -160,4 +163,13 @@ extern RHI_API FCachedBoundShaderStateLink* GetCachedBoundShaderState(
 	FGeometryShaderRHIParamRef GeometryShader = NULL
 	);
 
+#else
+extern RHI_API FCachedBoundShaderStateLink* GetCachedBoundShaderState(
+	FVertexDeclarationRHIParamRef VertexDeclaration,
+	FVertexShaderRHIParamRef VertexShader,
+	FPixelShaderRHIParamRef PixelShader,
+	FHullShaderRHIParamRef HullShader = NULL,
+	FDomainShaderRHIParamRef DomainShader = NULL,
+	FGeometryShaderRHIParamRef GeometryShader = NULL
+	);
 #endif
diff --git a/Engine/Source/Runtime/RHI/Public/DynamicRHI.h b/Engine/Source/Runtime/RHI/Public/DynamicRHI.h
--- a/Engine/Source/Runtime/RHI/Public/DynamicRHI.h
+++ b/Engine/Source/Runtime/RHI/Public/DynamicRHI.h
@@ -31,17 +31,33 @@ public:
 	/** Shutdown the RHI; handle shutdown and resource destruction before the RHI's actual destructor is called (so that all resources of the RHI are still available for shutdown). */
 	virtual void Shutdown() = 0;
 
+#if PLATFORM_RHI_USES_CONTEXT_OBJECT
 	#define DEFINE_RHIMETHOD(Type,Name,ParameterTypesAndNames,ParameterNames,ReturnStatement,NullImplementation) virtual Type RHI##Name ParameterTypesAndNames = 0
+	#define DEFINE_RHIMETHOD_CMDLIST(Type,Name,ParameterTypesAndNames,ParameterNames,ReturnStatement,NullImplementation)
+	#define DEFINE_RHIMETHOD_GLOBAL(Type,Name,ParameterTypesAndNames,ParameterNames,ReturnStatement,NullImplementation) virtual Type RHI##Name ParameterTypesAndNames = 0
+	#define DEFINE_RHIMETHOD_GLOBALFLUSH(Type,Name,ParameterTypesAndNames,ParameterNames,ReturnStatement,NullImplementation) virtual Type RHI##Name ParameterTypesAndNames = 0
+	#define DEFINE_RHIMETHOD_GLOBALTHREADSAFE(Type,Name,ParameterTypesAndNames,ParameterNames,ReturnStatement,NullImplementation) virtual Type RHI##Name ParameterTypesAndNames = 0
+#else
+	#define DEFINE_RHIMETHOD(Type,Name,ParameterTypesAndNames,ParameterNames,ReturnStatement,NullImplementation) virtual Type RHI##Name ParameterTypesAndNames = 0
+#endif
 	#include "RHIMethods.h"
 	#undef DEFINE_RHIMETHOD
+	#undef DEFINE_RHIMETHOD_CMDLIST
+	#undef DEFINE_RHIMETHOD_GLOBAL
+	#undef DEFINE_RHIMETHOD_GLOBALFLUSH
+	#undef DEFINE_RHIMETHOD_GLOBALTHREADSAFE
 };
 
 /** A global pointer to the dynamically bound RHI implementation. */
 extern RHI_API FDynamicRHI* GDynamicRHI;
 
 // Implement the statically bound RHI methods to simply call the dynamic RHI.
-#define DEFINE_RHIMETHOD_CMDLIST(Type,Name,ParameterTypesAndNames,ParameterNames,ReturnStatement,NullImplementation) \
-	extern RHI_API Type Name##_Internal ParameterTypesAndNames;
+#if PLATFORM_RHI_USES_CONTEXT_OBJECT
+	#define DEFINE_RHIMETHOD_CMDLIST(Type,Name,ParameterTypesAndNames,ParameterNames,ReturnStatement,NullImplementation)
+#else
+	#define DEFINE_RHIMETHOD_CMDLIST(Type,Name,ParameterTypesAndNames,ParameterNames,ReturnStatement,NullImplementation) \
+		extern RHI_API Type Name##_Internal ParameterTypesAndNames;
+#endif
 #define DEFINE_RHIMETHOD_GLOBAL(Type,Name,ParameterTypesAndNames,ParameterNames,ReturnStatement,NullImplementation) \
 	extern RHI_API Type Name##_Internal ParameterTypesAndNames;
 #define DEFINE_RHIMETHOD_GLOBALTHREADSAFE(Type,Name,ParameterTypesAndNames,ParameterNames,ReturnStatement,NullImplementation) \
diff --git a/Engine/Source/Runtime/RHI/Public/GPUProfiler.h b/Engine/Source/Runtime/RHI/Public/GPUProfiler.h
--- a/Engine/Source/Runtime/RHI/Public/GPUProfiler.h
+++ b/Engine/Source/Runtime/RHI/Public/GPUProfiler.h
@@ -188,6 +188,7 @@ struct RHI_API FGPUProfiler
 	{
 		if (bTrackingEvents && CurrentEventNode)
 		{
+			check(IsInRenderingThread() || IsInRHIThread());
 			CurrentEventNode->NumDraws++;
 			CurrentEventNode->NumPrimitives += NumPrimitives;
 			CurrentEventNode->NumVertices += NumVertices;
diff --git a/Engine/Source/Runtime/RHI/Public/RHICommandList.h b/Engine/Source/Runtime/RHI/Public/RHICommandList.h
--- a/Engine/Source/Runtime/RHI/Public/RHICommandList.h
+++ b/Engine/Source/Runtime/RHI/Public/RHICommandList.h
@@ -13,19 +13,52 @@ DECLARE_STATS_GROUP(TEXT("RHICmdList"), STATGROUP_RHICMDLIST, STATCAT_Advanced);
 #define USE_RHICOMMAND_STATE_REDUCTION (0)
 
 extern RHI_API TAutoConsoleVariable<int32> CVarRHICmdWidth;
+class FRHICommandListBase;
+
+
+class IRHICommandContext
+{
+public:
+	virtual ~IRHICommandContext() {}
+
+#define DEFINE_RHIMETHOD(Type,Name,ParameterTypesAndNames,ParameterNames,ReturnStatement,NullImplementation)
+#define DEFINE_RHIMETHOD_CMDLIST(Type,Name,ParameterTypesAndNames,ParameterNames,ReturnStatement,NullImplementation) virtual Type RHI##Name ParameterTypesAndNames = 0
+#include "RHIMethods.h"
+#undef DEFINE_RHIMETHOD
+#undef DEFINE_RHIMETHOD_CMDLIST
+};
+
+class IRHICommandContextContainer
+{
+public:
+	virtual ~IRHICommandContextContainer() {}
+
+	virtual IRHICommandContext* GetContext()
+	{
+		return nullptr;
+	}
+	virtual void FinishContext()
+	{
+		check(0);
+	}
+	virtual void SubmitAndFreeContextContainer(int32 Index, int32 Num)
+	{
+		check(0);
+	}
+};
 
 struct FRHICommandBase
 {
 	FRHICommandBase* Next;
-	void(*ExecuteAndDestructPtr)(FRHICommandBase *Cmd);
-	FORCEINLINE FRHICommandBase(void(*InExecuteAndDestructPtr)(FRHICommandBase *Cmd))
+	void(*ExecuteAndDestructPtr)(FRHICommandListBase& CmdList, FRHICommandBase *Cmd);
+	FORCEINLINE FRHICommandBase(void(*InExecuteAndDestructPtr)(FRHICommandListBase& CmdList, FRHICommandBase *Cmd))
 		: Next(nullptr)
 		, ExecuteAndDestructPtr(InExecuteAndDestructPtr)
 	{
 	}
-	FORCEINLINE void CallExecuteAndDestruct()
+	FORCEINLINE void CallExecuteAndDestruct(FRHICommandListBase& CmdList)
 	{
-		ExecuteAndDestructPtr(this);
+		ExecuteAndDestructPtr(CmdList, this);
 	}
 };
 
@@ -40,6 +73,7 @@ public:
 	inline bool IsImmediate();
 	const int32 GetUsedMemory() const;
 	void QueueAsyncCommandListSubmit(FGraphEventRef& AnyThreadCompletionEvent, class FRHICommandList* CmdList);
+	void QueueParallelAsyncCommandListSubmit(FGraphEventRef* AnyThreadCompletionEvents, class FRHICommandList** CmdLists, int32 Num);
 	void QueueRenderThreadCommandListSubmit(FGraphEventRef& RenderThreadCompletionEvent, class FRHICommandList* CmdList);
 	void QueueCommandListSubmit(class FRHICommandList* CmdList);
 	void WaitForTasks(bool bKnownToBeComplete = false);
@@ -96,6 +130,25 @@ public:
 		}
 	}
 
+	void SetContext(IRHICommandContext* InContext)
+	{
+#if PLATFORM_RHI_USES_CONTEXT_OBJECT
+		check(InContext);
+#endif
+		Context = InContext;
+	}
+
+	FORCEINLINE IRHICommandContext& GetContext()
+	{
+		checkSlow(Context);
+		return *Context;
+	}
+
+	void CopyContext(FRHICommandListBase& ParentCommandList)
+	{
+		Context = ParentCommandList.Context;
+	}
+
 	struct FDrawUpData
 	{
 
@@ -123,6 +176,8 @@ private:
 	bool bExecuting;
 	uint32 NumCommands;
 	uint32 UID;
+	IRHICommandContext* Context;
+
 	FMemStackBase MemManager; 
 #if USE_RHICOMMAND_STATE_REDUCTION
 protected:
@@ -182,6 +237,7 @@ struct FRHICommandListStateCache
 };
 #endif
 
+
 template<typename TCmd>
 struct FRHICommand : public FRHICommandBase
 {
@@ -190,10 +246,10 @@ struct FRHICommand : public FRHICommandBase
 	{
 
 	}
-	static FORCEINLINE void ExecuteAndDestruct(FRHICommandBase *Cmd)
+	static FORCEINLINE void ExecuteAndDestruct(FRHICommandListBase& CmdList, FRHICommandBase *Cmd)
 	{
 		TCmd *ThisCmd = (TCmd*)Cmd;
-		ThisCmd->Execute();
+		ThisCmd->Execute(CmdList);
 		ThisCmd->~TCmd();
 	}
 };
@@ -205,7 +261,7 @@ struct FRHICommandSetRasterizerState : public FRHICommand<FRHICommandSetRasteriz
 		: State(InState)
 	{
 	}
-	RHI_API void Execute();
+	RHI_API void Execute(FRHICommandListBase& CmdList);
 };
 
 struct FRHICommandSetDepthStencilState : public FRHICommand<FRHICommandSetDepthStencilState>
@@ -217,7 +273,7 @@ struct FRHICommandSetDepthStencilState : public FRHICommand<FRHICommandSetDepthS
 		, StencilRef(InStencilRef)
 	{
 	}
-	RHI_API void Execute();
+	RHI_API void Execute(FRHICommandListBase& CmdList);
 };
 
 template <typename TShaderRHIParamRef>
@@ -236,7 +292,7 @@ struct FRHICommandSetShaderParameter : public FRHICommand<FRHICommandSetShaderPa
 		, NumBytes(InNumBytes)
 	{
 	}
-	RHI_API void Execute();
+	RHI_API void Execute(FRHICommandListBase& CmdList);
 };
 
 template <typename TShaderRHIParamRef>
@@ -251,7 +307,7 @@ struct FRHICommandSetShaderUniformBuffer : public FRHICommand<FRHICommandSetShad
 		, UniformBuffer(InUniformBuffer)
 	{
 	}
-	RHI_API void Execute();
+	RHI_API void Execute(FRHICommandListBase& CmdList);
 };
 
 template <typename TShaderRHIParamRef>
@@ -266,7 +322,7 @@ struct FRHICommandSetShaderTexture : public FRHICommand<FRHICommandSetShaderText
 		, Texture(InTexture)
 	{
 	}
-	RHI_API void Execute();
+	RHI_API void Execute(FRHICommandListBase& CmdList);
 };
 
 template <typename TShaderRHIParamRef>
@@ -281,7 +337,7 @@ struct FRHICommandSetShaderResourceViewParameter : public FRHICommand<FRHIComman
 		, SRV(InSRV)
 	{
 	}
-	RHI_API void Execute();
+	RHI_API void Execute(FRHICommandListBase& CmdList);
 };
 
 template <typename TShaderRHIParamRef>
@@ -296,7 +352,7 @@ struct FRHICommandSetUAVParameter : public FRHICommand<FRHICommandSetUAVParamete
 		, UAV(InUAV)
 	{
 	}
-	RHI_API void Execute();
+	RHI_API void Execute(FRHICommandListBase& CmdList);
 };
 
 template <typename TShaderRHIParamRef>
@@ -313,7 +369,7 @@ struct FRHICommandSetUAVParameter_IntialCount : public FRHICommand<FRHICommandSe
 		, InitialCount(InInitialCount)
 	{
 	}
-	RHI_API void Execute();
+	RHI_API void Execute(FRHICommandListBase& CmdList);
 };
 
 template <typename TShaderRHIParamRef>
@@ -328,7 +384,7 @@ struct FRHICommandSetShaderSampler : public FRHICommand<FRHICommandSetShaderSamp
 		, Sampler(InSampler)
 	{
 	}
-	RHI_API void Execute();
+	RHI_API void Execute(FRHICommandListBase& CmdList);
 };
 
 struct FRHICommandDrawPrimitive : public FRHICommand<FRHICommandDrawPrimitive>
@@ -344,7 +400,7 @@ struct FRHICommandDrawPrimitive : public FRHICommand<FRHICommandDrawPrimitive>
 		, NumInstances(InNumInstances)
 	{
 	}
-	RHI_API void Execute();
+	RHI_API void Execute(FRHICommandListBase& CmdList);
 };
 
 struct FRHICommandDrawIndexedPrimitive : public FRHICommand<FRHICommandDrawIndexedPrimitive>
@@ -368,7 +424,7 @@ struct FRHICommandDrawIndexedPrimitive : public FRHICommand<FRHICommandDrawIndex
 		, NumInstances(InNumInstances)
 	{
 	}
-	RHI_API void Execute();
+	RHI_API void Execute(FRHICommandListBase& CmdList);
 };
 
 struct FRHICommandSetBoundShaderState : public FRHICommand<FRHICommandSetBoundShaderState>
@@ -378,7 +434,7 @@ struct FRHICommandSetBoundShaderState : public FRHICommand<FRHICommandSetBoundSh
 		: BoundShaderState(InBoundShaderState)
 	{
 	}
-	RHI_API void Execute();
+	RHI_API void Execute(FRHICommandListBase& CmdList);
 };
 
 struct FRHICommandSetBlendState : public FRHICommand<FRHICommandSetBlendState>
@@ -390,7 +446,7 @@ struct FRHICommandSetBlendState : public FRHICommand<FRHICommandSetBlendState>
 		, BlendFactor(InBlendFactor)
 	{
 	}
-	RHI_API void Execute();
+	RHI_API void Execute(FRHICommandListBase& CmdList);
 };
 
 struct FRHICommandSetStreamSource : public FRHICommand<FRHICommandSetStreamSource>
@@ -406,7 +462,7 @@ struct FRHICommandSetStreamSource : public FRHICommand<FRHICommandSetStreamSourc
 		, Offset(InOffset)
 	{
 	}
-	RHI_API void Execute();
+	RHI_API void Execute(FRHICommandListBase& CmdList);
 };
 
 struct FRHICommandSetViewport : public FRHICommand<FRHICommandSetViewport>
@@ -426,7 +482,7 @@ struct FRHICommandSetViewport : public FRHICommand<FRHICommandSetViewport>
 		, MaxZ(InMaxZ)
 	{
 	}
-	RHI_API void Execute();
+	RHI_API void Execute(FRHICommandListBase& CmdList);
 };
 
 struct FRHICommandSetScissorRect : public FRHICommand<FRHICommandSetScissorRect>
@@ -444,7 +500,7 @@ struct FRHICommandSetScissorRect : public FRHICommand<FRHICommandSetScissorRect>
 		, MaxY(InMaxY)
 	{
 	}
-	RHI_API void Execute();
+	RHI_API void Execute(FRHICommandListBase& CmdList);
 };
 
 struct FRHICommandSetRenderTargets : public FRHICommand<FRHICommandSetRenderTargets>
@@ -477,7 +533,7 @@ struct FRHICommandSetRenderTargets : public FRHICommand<FRHICommandSetRenderTarg
 			UAVs[Index] = InUAVs[Index];
 		}
 	}
-	RHI_API void Execute();
+	RHI_API void Execute(FRHICommandListBase& CmdList);
 };
 
 struct FRHICommandSetRenderTargetsAndClear : public FRHICommand<FRHICommandSetRenderTargetsAndClear>
@@ -489,7 +545,7 @@ struct FRHICommandSetRenderTargetsAndClear : public FRHICommand<FRHICommandSetRe
 	{
 	}
 
-	RHI_API void Execute();
+	RHI_API void Execute(FRHICommandListBase& CmdList);
 };
 
 struct FRHICommandEndDrawPrimitiveUP : public FRHICommand<FRHICommandEndDrawPrimitiveUP>
@@ -508,7 +564,7 @@ struct FRHICommandEndDrawPrimitiveUP : public FRHICommand<FRHICommandEndDrawPrim
 		, OutVertexData(InOutVertexData)
 	{
 	}
-	RHI_API void Execute();
+	RHI_API void Execute(FRHICommandListBase& CmdList);
 };
 
 
@@ -536,7 +592,7 @@ struct FRHICommandEndDrawIndexedPrimitiveUP : public FRHICommand<FRHICommandEndD
 		, OutIndexData(InOutIndexData)
 	{
 	}
-	RHI_API void Execute();
+	RHI_API void Execute(FRHICommandListBase& CmdList);
 };
 
 struct FRHICommandSetComputeShader : public FRHICommand<FRHICommandSetComputeShader>
@@ -546,7 +602,7 @@ struct FRHICommandSetComputeShader : public FRHICommand<FRHICommandSetComputeSha
 		: ComputeShader(InComputeShader)
 	{
 	}
-	RHI_API void Execute();
+	RHI_API void Execute(FRHICommandListBase& CmdList);
 };
 
 struct FRHICommandDispatchComputeShader : public FRHICommand<FRHICommandDispatchComputeShader>
@@ -560,7 +616,7 @@ struct FRHICommandDispatchComputeShader : public FRHICommand<FRHICommandDispatch
 		, ThreadGroupCountZ(InThreadGroupCountZ)
 	{
 	}
-	RHI_API void Execute();
+	RHI_API void Execute(FRHICommandListBase& CmdList);
 };
 
 struct FRHICommandDispatchIndirectComputeShader : public FRHICommand<FRHICommandDispatchIndirectComputeShader>
@@ -572,7 +628,7 @@ struct FRHICommandDispatchIndirectComputeShader : public FRHICommand<FRHICommand
 		, ArgumentOffset(InArgumentOffset)
 	{
 	}
-	RHI_API void Execute();
+	RHI_API void Execute(FRHICommandListBase& CmdList);
 };
 
 struct FRHICommandAutomaticCacheFlushAfterComputeShader : public FRHICommand<FRHICommandAutomaticCacheFlushAfterComputeShader>
@@ -582,12 +638,12 @@ struct FRHICommandAutomaticCacheFlushAfterComputeShader : public FRHICommand<FRH
 		: bEnable(InbEnable)
 	{
 	}
-	RHI_API void Execute();
+	RHI_API void Execute(FRHICommandListBase& CmdList);
 };
 
 struct FRHICommandFlushComputeShaderCache : public FRHICommand<FRHICommandFlushComputeShaderCache>
 {
-	RHI_API void Execute();
+	RHI_API void Execute(FRHICommandListBase& CmdList);
 };
 
 struct FRHICommandDrawPrimitiveIndirect : public FRHICommand<FRHICommandDrawPrimitiveIndirect>
@@ -601,7 +657,7 @@ struct FRHICommandDrawPrimitiveIndirect : public FRHICommand<FRHICommandDrawPrim
 		, ArgumentOffset(InArgumentOffset)
 	{
 	}
-	RHI_API void Execute();
+	RHI_API void Execute(FRHICommandListBase& CmdList);
 };
 
 struct FRHICommandDrawIndexedIndirect : public FRHICommand<FRHICommandDrawIndexedIndirect>
@@ -620,7 +676,7 @@ struct FRHICommandDrawIndexedIndirect : public FRHICommand<FRHICommandDrawIndexe
 		, NumInstances(InNumInstances)
 	{
 	}
-	RHI_API void Execute();
+	RHI_API void Execute(FRHICommandListBase& CmdList);
 };
 
 struct FRHICommandDrawIndexedPrimitiveIndirect : public FRHICommand<FRHICommandDrawIndexedPrimitiveIndirect>
@@ -637,7 +693,7 @@ struct FRHICommandDrawIndexedPrimitiveIndirect : public FRHICommand<FRHICommandD
 		, ArgumentOffset(InArgumentOffset)
 	{
 	}
-	RHI_API void Execute();
+	RHI_API void Execute(FRHICommandListBase& CmdList);
 };
 
 struct FRHICommandEnableDepthBoundsTest : public FRHICommand<FRHICommandEnableDepthBoundsTest>
@@ -652,7 +708,7 @@ struct FRHICommandEnableDepthBoundsTest : public FRHICommand<FRHICommandEnableDe
 		, MaxDepth(InMaxDepth)
 	{
 	}
-	RHI_API void Execute();
+	RHI_API void Execute(FRHICommandListBase& CmdList);
 };
 
 struct FRHICommandClearUAV : public FRHICommand<FRHICommandClearUAV>
@@ -668,7 +724,7 @@ struct FRHICommandClearUAV : public FRHICommand<FRHICommandClearUAV>
 		Values[2] = InValues[2];
 		Values[3] = InValues[3];
 	}
-	RHI_API void Execute();
+	RHI_API void Execute(FRHICommandListBase& CmdList);
 };
 
 struct FRHICommandCopyToResolveTarget : public FRHICommand<FRHICommandCopyToResolveTarget>
@@ -685,7 +741,7 @@ struct FRHICommandCopyToResolveTarget : public FRHICommand<FRHICommandCopyToReso
 		, ResolveParams(InResolveParams)
 	{
 	}
-	RHI_API void Execute();
+	RHI_API void Execute(FRHICommandListBase& CmdList);
 };
 
 struct FRHICommandClear : public FRHICommand<FRHICommandClear>
@@ -716,7 +772,7 @@ struct FRHICommandClear : public FRHICommand<FRHICommandClear>
 		, bClearStencil(InbClearStencil)
 	{
 	}
-	RHI_API void Execute();
+	RHI_API void Execute(FRHICommandListBase& CmdList);
 };
 
 struct FRHICommandClearMRT : public FRHICommand<FRHICommandClearMRT>
@@ -754,7 +810,7 @@ struct FRHICommandClearMRT : public FRHICommand<FRHICommandClearMRT>
 			ColorArray[Index] = InColorArray[Index];
 		}
 	}
-	RHI_API void Execute();
+	RHI_API void Execute(FRHICommandListBase& CmdList);
 };
 
 
@@ -858,7 +914,7 @@ struct FRHICommandBuildLocalBoundShaderState : public FRHICommand<FRHICommandBui
 
 	{
 	}
-	RHI_API void Execute();
+	RHI_API void Execute(FRHICommandListBase& CmdList);
 };
 
 struct FRHICommandSetLocalBoundShaderState : public FRHICommand<FRHICommandSetLocalBoundShaderState>
@@ -870,7 +926,7 @@ struct FRHICommandSetLocalBoundShaderState : public FRHICommand<FRHICommandSetLo
 		check(CheckCmdList == LocalBoundShaderState.WorkArea->CheckCmdList && CheckCmdList->GetUID() == LocalBoundShaderState.WorkArea->UID); // this BSS was not built for this particular commandlist
 		LocalBoundShaderState.WorkArea->ComputedBSS->UseCount++;
 	}
-	RHI_API void Execute();
+	RHI_API void Execute(FRHICommandListBase& CmdList);
 };
 
 
@@ -941,7 +997,7 @@ struct FRHICommandBuildLocalUniformBuffer : public FRHICommand<FRHICommandBuildL
 
 	{
 	}
-	RHI_API void Execute();
+	RHI_API void Execute(FRHICommandListBase& CmdList);
 };
 
 template <typename TShaderRHIParamRef>
@@ -959,7 +1015,7 @@ struct FRHICommandSetLocalUniformBuffer : public FRHICommand<FRHICommandSetLocal
 		check(CheckCmdList == LocalUniformBuffer.WorkArea->CheckCmdList && CheckCmdList->GetUID() == LocalUniformBuffer.WorkArea->UID); // this uniform buffer was not built for this particular commandlist
 		LocalUniformBuffer.WorkArea->ComputedUniformBuffer->UseCount++;
 	}
-	RHI_API void Execute();
+	RHI_API void Execute(FRHICommandListBase& CmdList);
 };
 
 struct FRHICommandBeginRenderQuery : public FRHICommand<FRHICommandBeginRenderQuery>
@@ -970,7 +1026,7 @@ struct FRHICommandBeginRenderQuery : public FRHICommand<FRHICommandBeginRenderQu
 		: RenderQuery(InRenderQuery)
 	{
 	}
-	RHI_API void Execute();
+	RHI_API void Execute(FRHICommandListBase& CmdList);
 };
 
 struct FRHICommandEndRenderQuery : public FRHICommand<FRHICommandEndRenderQuery>
@@ -981,7 +1037,7 @@ struct FRHICommandEndRenderQuery : public FRHICommand<FRHICommandEndRenderQuery>
 		: RenderQuery(InRenderQuery)
 	{
 	}
-	RHI_API void Execute();
+	RHI_API void Execute(FRHICommandListBase& CmdList);
 };
 
 #if !PLATFORM_SUPPORTS_RHI_THREAD
@@ -993,7 +1049,7 @@ struct FRHICommandResetRenderQuery : public FRHICommand<FRHICommandResetRenderQu
 		: RenderQuery(InRenderQuery)
 	{
 	}
-	RHI_API void Execute();
+	RHI_API void Execute(FRHICommandListBase& CmdList);
 };
 #endif
 
@@ -1002,7 +1058,7 @@ struct FRHICommandBeginScene : public FRHICommand<FRHICommandBeginScene>
 	FORCEINLINE_DEBUGGABLE FRHICommandBeginScene()
 	{
 	}
-	RHI_API void Execute();
+	RHI_API void Execute(FRHICommandListBase& CmdList);
 };
 
 struct FRHICommandEndScene : public FRHICommand<FRHICommandEndScene>
@@ -1010,7 +1066,7 @@ struct FRHICommandEndScene : public FRHICommand<FRHICommandEndScene>
 	FORCEINLINE_DEBUGGABLE FRHICommandEndScene()
 	{
 	}
-	RHI_API void Execute();
+	RHI_API void Execute(FRHICommandListBase& CmdList);
 };
 
 struct FRHICommandUpdateVertexBuffer : public FRHICommand<FRHICommandUpdateVertexBuffer>
@@ -1025,7 +1081,7 @@ struct FRHICommandUpdateVertexBuffer : public FRHICommand<FRHICommandUpdateVerte
 		, BufferSize(InBufferSize)
 	{
 	}
-	RHI_API void Execute();
+	RHI_API void Execute(FRHICommandListBase& CmdList);
 };
 
 struct FRHICommandBeginFrame : public FRHICommand<FRHICommandBeginFrame>
@@ -1033,7 +1089,7 @@ struct FRHICommandBeginFrame : public FRHICommand<FRHICommandBeginFrame>
 	FORCEINLINE_DEBUGGABLE FRHICommandBeginFrame()
 	{
 	}
-	RHI_API void Execute();
+	RHI_API void Execute(FRHICommandListBase& CmdList);
 };
 
 struct FRHICommandEndFrame : public FRHICommand<FRHICommandEndFrame>
@@ -1041,7 +1097,7 @@ struct FRHICommandEndFrame : public FRHICommand<FRHICommandEndFrame>
 	FORCEINLINE_DEBUGGABLE FRHICommandEndFrame()
 	{
 	}
-	RHI_API void Execute();
+	RHI_API void Execute(FRHICommandListBase& CmdList);
 };
 
 #if PLATFORM_SUPPORTS_RHI_THREAD
@@ -1055,7 +1111,7 @@ struct FRHICommandBeginDrawingViewport : public FRHICommand<FRHICommandBeginDraw
 		, RenderTargetRHI(InRenderTargetRHI)
 	{
 	}
-	RHI_API void Execute();
+	RHI_API void Execute(FRHICommandListBase& CmdList);
 };
 
 struct FRHICommandEndDrawingViewport : public FRHICommand<FRHICommandEndDrawingViewport>
@@ -1070,7 +1126,7 @@ struct FRHICommandEndDrawingViewport : public FRHICommand<FRHICommandEndDrawingV
 		, bLockToVsync(InbLockToVsync)
 	{
 	}
-	RHI_API void Execute();
+	RHI_API void Execute(FRHICommandListBase& CmdList);
 };
 #endif
 
@@ -1082,14 +1138,32 @@ struct FRHICommandPushEvent : public FRHICommand<FRHICommandPushEvent>
 		: Name(InName)
 	{
 	}
-	RHI_API void Execute();
+	RHI_API void Execute(FRHICommandListBase& CmdList);
 };
 
 struct FRHICommandPopEvent : public FRHICommand<FRHICommandPopEvent>
 {
-	RHI_API void Execute();
+	RHI_API void Execute(FRHICommandListBase& CmdList);
+};
+
+struct FRHICommandDebugBreak : public FRHICommand<FRHICommandDebugBreak>
+{
+	void Execute(FRHICommandListBase& CmdList)
+	{
+		if (FPlatformMisc::IsDebuggerPresent())
+		{
+			FPlatformMisc::DebugBreak();
+		}
+	}
 };
 
+#if PLATFORM_RHI_USES_CONTEXT_OBJECT
+#define CMD_CONTEXT(Method) GetContext().RHI##Method
+#else
+#define CMD_CONTEXT(Method) Method##_Internal
+#endif
+
+
 class RHI_API FRHICommandList : public FRHICommandListBase
 {
 public:
@@ -1145,7 +1219,11 @@ public:
 		FLocalBoundShaderState Result;
 		if (Bypass())
 		{
+#if PLATFORM_SUPPORTS_PARALLEL_RHI_EXECUTE
+			Result.BypassBSS = RHICreateBoundShaderState(VertexDeclarationRHI, VertexShaderRHI, HullShaderRHI, DomainShaderRHI, PixelShaderRHI, GeometryShaderRHI);
+#else
 			Result.BypassBSS = CreateBoundShaderState_Internal(VertexDeclarationRHI, VertexShaderRHI, HullShaderRHI, DomainShaderRHI, PixelShaderRHI, GeometryShaderRHI);
+#endif
 		}
 		else
 		{
@@ -1159,7 +1237,7 @@ public:
 	{
 		if (Bypass())
 		{
-			SetBoundShaderState_Internal(LocalBoundShaderState.BypassBSS);
+			CMD_CONTEXT(SetBoundShaderState)(LocalBoundShaderState.BypassBSS);
 			return;
 		}
 #if USE_RHICOMMAND_STATE_REDUCTION
@@ -1195,7 +1273,7 @@ public:
 	{
 		if (Bypass())
 		{
-			SetShaderUniformBuffer_Internal(Shader, BaseIndex, UniformBuffer.BypassUniform);
+			CMD_CONTEXT(SetShaderUniformBuffer)(Shader, BaseIndex, UniformBuffer.BypassUniform);
 			return;
 		}
 #if USE_RHICOMMAND_STATE_REDUCTION
@@ -1213,7 +1291,7 @@ public:
 	{
 		if (Bypass())
 		{
-			SetShaderUniformBuffer_Internal(Shader, BaseIndex, UniformBuffer);
+			CMD_CONTEXT(SetShaderUniformBuffer)(Shader, BaseIndex, UniformBuffer);
 			return;
 		}
 #if USE_RHICOMMAND_STATE_REDUCTION
@@ -1239,7 +1317,7 @@ public:
 	{
 		if (Bypass())
 		{
-			SetShaderParameter_Internal(Shader, BufferIndex, BaseIndex, NumBytes, NewValue);
+			CMD_CONTEXT(SetShaderParameter)(Shader, BufferIndex, BaseIndex, NumBytes, NewValue);
 			return;
 		}
 		void* UseValue = Alloc(NumBytes, 16);
@@ -1257,7 +1335,7 @@ public:
 	{
 		if (Bypass())
 		{
-			SetShaderTexture_Internal(Shader, TextureIndex, Texture);
+			CMD_CONTEXT(SetShaderTexture)(Shader, TextureIndex, Texture);
 			return;
 		}
 		new (AllocCommand<FRHICommandSetShaderTexture<TShaderRHIParamRef> >()) FRHICommandSetShaderTexture<TShaderRHIParamRef>(Shader, TextureIndex, Texture);
@@ -1268,7 +1346,7 @@ public:
 	{
 		if (Bypass())
 		{
-			SetShaderResourceViewParameter_Internal(Shader, SamplerIndex, SRV);
+			CMD_CONTEXT(SetShaderResourceViewParameter)(Shader, SamplerIndex, SRV);
 			return;
 		}
 		new (AllocCommand<FRHICommandSetShaderResourceViewParameter<TShaderRHIParamRef> >()) FRHICommandSetShaderResourceViewParameter<TShaderRHIParamRef>(Shader, SamplerIndex, SRV);
@@ -1279,7 +1357,7 @@ public:
 	{
 		if (Bypass())
 		{
-			SetShaderSampler_Internal(Shader, SamplerIndex, State);
+			CMD_CONTEXT(SetShaderSampler)(Shader, SamplerIndex, State);
 			return;
 		}
 #if USE_RHICOMMAND_STATE_REDUCTION
@@ -1299,7 +1377,7 @@ public:
 	{
 		if (Bypass())
 		{
-			SetUAVParameter_Internal(Shader, UAVIndex, UAV);
+			CMD_CONTEXT(SetUAVParameter)(Shader, UAVIndex, UAV);
 			return;
 		}
 		new (AllocCommand<FRHICommandSetUAVParameter<FComputeShaderRHIParamRef> >()) FRHICommandSetUAVParameter<FComputeShaderRHIParamRef>(Shader, UAVIndex, UAV);
@@ -1309,7 +1387,7 @@ public:
 	{
 		if (Bypass())
 		{
-			SetUAVParameter_Internal(Shader, UAVIndex, UAV, InitialCount);
+			CMD_CONTEXT(SetUAVParameter)(Shader, UAVIndex, UAV, InitialCount);
 			return;
 		}
 		new (AllocCommand<FRHICommandSetUAVParameter_IntialCount<FComputeShaderRHIParamRef> >()) FRHICommandSetUAVParameter_IntialCount<FComputeShaderRHIParamRef>(Shader, UAVIndex, UAV, InitialCount);
@@ -1319,7 +1397,7 @@ public:
 	{
 		if (Bypass())
 		{
-			SetBoundShaderState_Internal(BoundShaderState);
+			CMD_CONTEXT(SetBoundShaderState)(BoundShaderState);
 			return;
 		}
 #if USE_RHICOMMAND_STATE_REDUCTION
@@ -1335,7 +1413,7 @@ public:
 	{
 		if (Bypass())
 		{
-			SetRasterizerState_Internal(State);
+			CMD_CONTEXT(SetRasterizerState)(State);
 			return;
 		}
 		new (AllocCommand<FRHICommandSetRasterizerState>()) FRHICommandSetRasterizerState(State);
@@ -1345,7 +1423,7 @@ public:
 	{
 		if (Bypass())
 		{
-			SetBlendState_Internal(State, BlendFactor);
+			CMD_CONTEXT(SetBlendState)(State, BlendFactor);
 			return;
 		}
 		new (AllocCommand<FRHICommandSetBlendState>()) FRHICommandSetBlendState(State, BlendFactor);
@@ -1355,7 +1433,7 @@ public:
 	{
 		if (Bypass())
 		{
-			DrawPrimitive_Internal(PrimitiveType, BaseVertexIndex, NumPrimitives, NumInstances);
+			CMD_CONTEXT(DrawPrimitive)(PrimitiveType, BaseVertexIndex, NumPrimitives, NumInstances);
 			return;
 		}
 		new (AllocCommand<FRHICommandDrawPrimitive>()) FRHICommandDrawPrimitive(PrimitiveType, BaseVertexIndex, NumPrimitives, NumInstances);
@@ -1365,7 +1443,7 @@ public:
 	{
 		if (Bypass())
 		{
-			DrawIndexedPrimitive_Internal(IndexBuffer, PrimitiveType, BaseVertexIndex, MinIndex, NumVertices, StartIndex, NumPrimitives, NumInstances);
+			CMD_CONTEXT(DrawIndexedPrimitive)(IndexBuffer, PrimitiveType, BaseVertexIndex, MinIndex, NumVertices, StartIndex, NumPrimitives, NumInstances);
 			return;
 		}
 		new (AllocCommand<FRHICommandDrawIndexedPrimitive>()) FRHICommandDrawIndexedPrimitive(IndexBuffer, PrimitiveType, BaseVertexIndex, MinIndex, NumVertices, StartIndex, NumPrimitives, NumInstances);
@@ -1375,7 +1453,7 @@ public:
 	{
 		if (Bypass())
 		{
-			SetStreamSource_Internal(StreamIndex, VertexBuffer, Stride, Offset);
+			CMD_CONTEXT(SetStreamSource)(StreamIndex, VertexBuffer, Stride, Offset);
 			return;
 		}
 		new (AllocCommand<FRHICommandSetStreamSource>()) FRHICommandSetStreamSource(StreamIndex, VertexBuffer, Stride, Offset);
@@ -1385,7 +1463,7 @@ public:
 	{
 		if (Bypass())
 		{
-			SetDepthStencilState_Internal(NewStateRHI, StencilRef);
+			CMD_CONTEXT(SetDepthStencilState)(NewStateRHI, StencilRef);
 			return;
 		}
 		new (AllocCommand<FRHICommandSetDepthStencilState>()) FRHICommandSetDepthStencilState(NewStateRHI, StencilRef);
@@ -1395,7 +1473,7 @@ public:
 	{
 		if (Bypass())
 		{
-			SetViewport_Internal(MinX, MinY, MinZ, MaxX, MaxY, MaxZ);
+			CMD_CONTEXT(SetViewport)(MinX, MinY, MinZ, MaxX, MaxY, MaxZ);
 			return;
 		}
 		new (AllocCommand<FRHICommandSetViewport>()) FRHICommandSetViewport(MinX, MinY, MinZ, MaxX, MaxY, MaxZ);
@@ -1405,7 +1483,7 @@ public:
 	{
 		if (Bypass())
 		{
-			SetScissorRect_Internal(bEnable, MinX, MinY, MaxX, MaxY);
+			CMD_CONTEXT(SetScissorRect)(bEnable, MinX, MinY, MaxX, MaxY);
 			return;
 		}
 		new (AllocCommand<FRHICommandSetScissorRect>()) FRHICommandSetScissorRect(bEnable, MinX, MinY, MaxX, MaxY);
@@ -1421,7 +1499,7 @@ public:
 	{
 		if (Bypass())
 		{
-			SetRenderTargets_Internal(
+			CMD_CONTEXT(SetRenderTargets)(
 				NewNumSimultaneousRenderTargets,
 				NewRenderTargetsRHI,
 				NewDepthStencilTargetRHI,
@@ -1441,7 +1519,7 @@ public:
 	{
 		if (Bypass())
 		{
-			SetRenderTargetsAndClear_Internal(RenderTargetsInfo);
+			CMD_CONTEXT(SetRenderTargetsAndClear)(RenderTargetsInfo);
 			return;
 		}
 		new (AllocCommand<FRHICommandSetRenderTargetsAndClear>()) FRHICommandSetRenderTargetsAndClear(RenderTargetsInfo);
@@ -1451,7 +1529,7 @@ public:
 	{
 		if (Bypass())
 		{
-			BeginDrawPrimitiveUP_Internal(PrimitiveType, NumPrimitives, NumVertices, VertexDataStride, OutVertexData);
+			CMD_CONTEXT(BeginDrawPrimitiveUP)(PrimitiveType, NumPrimitives, NumVertices, VertexDataStride, OutVertexData);
 			return;
 		}
 		check(!DrawUPData.OutVertexData && NumVertices * VertexDataStride > 0);
@@ -1469,7 +1547,7 @@ public:
 	{
 		if (Bypass())
 		{
-			EndDrawPrimitiveUP_Internal();
+			CMD_CONTEXT(EndDrawPrimitiveUP)();
 			return;
 		}
 		check(DrawUPData.OutVertexData && DrawUPData.NumVertices);
@@ -1483,7 +1561,7 @@ public:
 	{
 		if (Bypass())
 		{
-			BeginDrawIndexedPrimitiveUP_Internal(PrimitiveType, NumPrimitives, NumVertices, VertexDataStride, OutVertexData, MinVertexIndex, NumIndices, IndexDataStride, OutIndexData);
+			CMD_CONTEXT(BeginDrawIndexedPrimitiveUP)(PrimitiveType, NumPrimitives, NumVertices, VertexDataStride, OutVertexData, MinVertexIndex, NumIndices, IndexDataStride, OutIndexData);
 			return;
 		}
 		check(!DrawUPData.OutVertexData && !DrawUPData.OutIndexData && NumVertices * VertexDataStride > 0 && NumIndices * IndexDataStride > 0);
@@ -1507,7 +1585,7 @@ public:
 	{
 		if (Bypass())
 		{
-			EndDrawIndexedPrimitiveUP_Internal();
+			CMD_CONTEXT(EndDrawIndexedPrimitiveUP)();
 			return;
 		}
 		check(DrawUPData.OutVertexData && DrawUPData.OutIndexData && DrawUPData.NumIndices && DrawUPData.NumVertices);
@@ -1524,7 +1602,7 @@ public:
 	{
 		if (Bypass())
 		{
-			SetComputeShader_Internal(ComputeShader);
+			CMD_CONTEXT(SetComputeShader)(ComputeShader);
 			return;
 		}
 #if USE_RHICOMMAND_STATE_REDUCTION
@@ -1540,7 +1618,7 @@ public:
 	{
 		if (Bypass())
 		{
-			DispatchComputeShader_Internal(ThreadGroupCountX, ThreadGroupCountY, ThreadGroupCountZ);
+			CMD_CONTEXT(DispatchComputeShader)(ThreadGroupCountX, ThreadGroupCountY, ThreadGroupCountZ);
 			return;
 		}
 		new (AllocCommand<FRHICommandDispatchComputeShader>()) FRHICommandDispatchComputeShader(ThreadGroupCountX, ThreadGroupCountY, ThreadGroupCountZ);
@@ -1550,7 +1628,7 @@ public:
 	{
 		if (Bypass())
 		{
-			DispatchIndirectComputeShader_Internal(ArgumentBuffer, ArgumentOffset);
+			CMD_CONTEXT(DispatchIndirectComputeShader)(ArgumentBuffer, ArgumentOffset);
 			return;
 		}
 		new (AllocCommand<FRHICommandDispatchIndirectComputeShader>()) FRHICommandDispatchIndirectComputeShader(ArgumentBuffer, ArgumentOffset);
@@ -1560,7 +1638,7 @@ public:
 	{
 		if (Bypass())
 		{
-			AutomaticCacheFlushAfterComputeShader_Internal(bEnable);
+			CMD_CONTEXT(AutomaticCacheFlushAfterComputeShader)(bEnable);
 			return;
 		}
 		new (AllocCommand<FRHICommandAutomaticCacheFlushAfterComputeShader>()) FRHICommandAutomaticCacheFlushAfterComputeShader(bEnable);
@@ -1570,7 +1648,7 @@ public:
 	{
 		if (Bypass())
 		{
-			FlushComputeShaderCache_Internal();
+			CMD_CONTEXT(FlushComputeShaderCache)();
 			return;
 		}
 		new (AllocCommand<FRHICommandFlushComputeShaderCache>()) FRHICommandFlushComputeShaderCache();
@@ -1580,7 +1658,7 @@ public:
 	{
 		if (Bypass())
 		{
-			DrawPrimitiveIndirect_Internal(PrimitiveType, ArgumentBuffer, ArgumentOffset);
+			CMD_CONTEXT(DrawPrimitiveIndirect)(PrimitiveType, ArgumentBuffer, ArgumentOffset);
 			return;
 		}
 		new (AllocCommand<FRHICommandDrawPrimitiveIndirect>()) FRHICommandDrawPrimitiveIndirect(PrimitiveType, ArgumentBuffer, ArgumentOffset);
@@ -1590,7 +1668,7 @@ public:
 	{
 		if (Bypass())
 		{
-			DrawIndexedIndirect_Internal(IndexBufferRHI, PrimitiveType, ArgumentsBufferRHI, DrawArgumentsIndex, NumInstances);
+			CMD_CONTEXT(DrawIndexedIndirect)(IndexBufferRHI, PrimitiveType, ArgumentsBufferRHI, DrawArgumentsIndex, NumInstances);
 			return;
 		}
 		new (AllocCommand<FRHICommandDrawIndexedIndirect>()) FRHICommandDrawIndexedIndirect(IndexBufferRHI, PrimitiveType, ArgumentsBufferRHI, DrawArgumentsIndex, NumInstances);
@@ -1600,7 +1678,7 @@ public:
 	{
 		if (Bypass())
 		{
-			DrawIndexedPrimitiveIndirect_Internal(PrimitiveType, IndexBuffer, ArgumentsBuffer, ArgumentOffset);
+			CMD_CONTEXT(DrawIndexedPrimitiveIndirect)(PrimitiveType, IndexBuffer, ArgumentsBuffer, ArgumentOffset);
 			return;
 		}
 		new (AllocCommand<FRHICommandDrawIndexedPrimitiveIndirect>()) FRHICommandDrawIndexedPrimitiveIndirect(PrimitiveType, IndexBuffer, ArgumentsBuffer, ArgumentOffset);
@@ -1610,7 +1688,7 @@ public:
 	{
 		if (Bypass())
 		{
-			EnableDepthBoundsTest_Internal(bEnable, MinDepth, MaxDepth);
+			CMD_CONTEXT(EnableDepthBoundsTest)(bEnable, MinDepth, MaxDepth);
 			return;
 		}
 		new (AllocCommand<FRHICommandEnableDepthBoundsTest>()) FRHICommandEnableDepthBoundsTest(bEnable, MinDepth, MaxDepth);
@@ -1620,7 +1698,7 @@ public:
 	{
 		if (Bypass())
 		{
-			ClearUAV_Internal(UnorderedAccessViewRHI, Values);
+			CMD_CONTEXT(ClearUAV)(UnorderedAccessViewRHI, Values);
 			return;
 		}
 		new (AllocCommand<FRHICommandClearUAV>()) FRHICommandClearUAV(UnorderedAccessViewRHI, Values);
@@ -1630,7 +1708,7 @@ public:
 	{
 		if (Bypass())
 		{
-			CopyToResolveTarget_Internal(SourceTextureRHI, DestTextureRHI, bKeepOriginalSurface, ResolveParams);
+			CMD_CONTEXT(CopyToResolveTarget)(SourceTextureRHI, DestTextureRHI, bKeepOriginalSurface, ResolveParams);
 			return;
 		}
 		new (AllocCommand<FRHICommandCopyToResolveTarget>()) FRHICommandCopyToResolveTarget(SourceTextureRHI, DestTextureRHI, bKeepOriginalSurface, ResolveParams);
@@ -1640,7 +1718,7 @@ public:
 	{
 		if (Bypass())
 		{
-			Clear_Internal(bClearColor, Color, bClearDepth, Depth, bClearStencil, Stencil, ExcludeRect);
+			CMD_CONTEXT(Clear)(bClearColor, Color, bClearDepth, Depth, bClearStencil, Stencil, ExcludeRect);
 			return;
 		}
 		new (AllocCommand<FRHICommandClear>()) FRHICommandClear(bClearColor, Color, bClearDepth, Depth, bClearStencil, Stencil, ExcludeRect);
@@ -1650,7 +1728,7 @@ public:
 	{
 		if (Bypass())
 		{
-			ClearMRT_Internal(bClearColor, NumClearColors, ClearColorArray, bClearDepth, Depth, bClearStencil, Stencil, ExcludeRect);
+			CMD_CONTEXT(ClearMRT)(bClearColor, NumClearColors, ClearColorArray, bClearDepth, Depth, bClearStencil, Stencil, ExcludeRect);
 			return;
 		}
 		new (AllocCommand<FRHICommandClearMRT>()) FRHICommandClearMRT(bClearColor, NumClearColors, ClearColorArray, bClearDepth, Depth, bClearStencil, Stencil, ExcludeRect);
@@ -1673,7 +1751,7 @@ public:
 	{
 		if (Bypass())
 		{
-			BeginRenderQuery_Internal(RenderQuery);
+			CMD_CONTEXT(BeginRenderQuery)(RenderQuery);
 			return;
 		}
 		new (AllocCommand<FRHICommandBeginRenderQuery>()) FRHICommandBeginRenderQuery(RenderQuery);
@@ -1682,7 +1760,7 @@ public:
 	{
 		if (Bypass())
 		{
-			EndRenderQuery_Internal(RenderQuery);
+			CMD_CONTEXT(EndRenderQuery)(RenderQuery);
 			return;
 		}
 		new (AllocCommand<FRHICommandEndRenderQuery>()) FRHICommandEndRenderQuery(RenderQuery);
@@ -1692,7 +1770,7 @@ public:
 	{
 		if (Bypass())
 		{
-			ResetRenderQuery_Internal(RenderQuery);
+			CMD_CONTEXT(ResetRenderQuery)(RenderQuery);
 			return;
 		}
 		new (AllocCommand<FRHICommandResetRenderQuery>()) FRHICommandResetRenderQuery(RenderQuery);
@@ -1704,7 +1782,7 @@ public:
 	{
 		if (Bypass())
 		{
-			BeginScene_Internal();
+			CMD_CONTEXT(BeginScene)();
 			return;
 		}
 		new (AllocCommand<FRHICommandBeginScene>()) FRHICommandBeginScene();
@@ -1713,7 +1791,7 @@ public:
 	{
 		if (Bypass())
 		{
-			EndScene_Internal();
+			CMD_CONTEXT(EndScene)();
 			return;
 		}
 		new (AllocCommand<FRHICommandEndScene>()) FRHICommandEndScene();
@@ -1728,7 +1806,7 @@ public:
 	{
 		if (Bypass())
 		{
-			PushEvent_Internal(Name);
+			CMD_CONTEXT(PushEvent)(Name);
 			return;
 		}
 		int32 Len = FCString::Strlen(Name) + 1;
@@ -1741,12 +1819,26 @@ public:
 	{
 		if (Bypass())
 		{
-			PopEvent_Internal();
+			CMD_CONTEXT(PopEvent)();
 			return;
 		}
 		new (AllocCommand<FRHICommandPopEvent>()) FRHICommandPopEvent();
 	}
 
+	FORCEINLINE_DEBUGGABLE void BreakPoint()
+	{
+#if !UE_BUILD_SHIPPING
+		if (Bypass())
+		{
+			if (FPlatformMisc::IsDebuggerPresent())
+			{
+				FPlatformMisc::DebugBreak();
+			}
+			return;
+		}
+		new (AllocCommand<FRHICommandDebugBreak>()) FRHICommandDebugBreak();
+#endif
+	}
 };
 
 namespace EImmediateFlushType
@@ -1775,6 +1867,8 @@ public:
 
 	inline void ImmediateFlush(EImmediateFlushType::Type FlushType);
 
+	void SetCurrentStat(TStatId Stat);
+
 	#define DEFINE_RHIMETHOD_CMDLIST(Type,Name,ParameterTypesAndNames,ParameterNames,ReturnStatement,NullImplementation)
 	#define DEFINE_RHIMETHOD(Type, Name, ParameterTypesAndNames, ParameterNames, ReturnStatement, NullImplementation) \
 		FORCEINLINE_DEBUGGABLE Type Name ParameterTypesAndNames \
@@ -1815,8 +1909,25 @@ public:
 };
 
 
-// typedef to mark the recursive use of commandlists in the RHI implmentations
+// typedef to mark the recursive use of commandlists in the RHI implementations
+
+#if PLATFORM_RHI_USES_CONTEXT_OBJECT
+class RHI_API FRHICommandList_RecursiveHazardous : public FRHICommandList
+{
+	FRHICommandList_RecursiveHazardous()
+	{
+
+	}
+public:
+	FRHICommandList_RecursiveHazardous(IRHICommandContext *Context)
+	{
+		SetContext(Context);
+	}
+};
+#else
 typedef FRHICommandList FRHICommandList_RecursiveHazardous;
+#endif
+
 
 class RHI_API FRHICommandListExecutor
 {
diff --git a/Engine/Source/Runtime/RHI/Public/RHICommandListCommandExecutes.inl b/Engine/Source/Runtime/RHI/Public/RHICommandListCommandExecutes.inl
--- a/Engine/Source/Runtime/RHI/Public/RHICommandListCommandExecutes.inl
+++ b/Engine/Source/Runtime/RHI/Public/RHICommandListCommandExecutes.inl
@@ -5,7 +5,11 @@
 =============================================================================*/
 
 #if !defined(INTERNAL_DECORATOR)
-	#define INTERNAL_DECORATOR(Method) Method##_Internal
+	#if PLATFORM_RHI_USES_CONTEXT_OBJECT
+		#define INTERNAL_DECORATOR(Method) CmdList.GetContext().RHI##Method
+	#else
+		#define INTERNAL_DECORATOR(Method) Method##_Internal
+	#endif
 #endif
 
 // set this one to get a stat for each RHI command 
@@ -18,20 +22,20 @@
 	#define RHISTAT(Method)
 #endif
 
-void FRHICommandSetRasterizerState::Execute()
+void FRHICommandSetRasterizerState::Execute(FRHICommandListBase& CmdList)
 {
 	RHISTAT(SetRasterizerState);
 	INTERNAL_DECORATOR(SetRasterizerState)(State);
 }
 
-void FRHICommandSetDepthStencilState::Execute()
+void FRHICommandSetDepthStencilState::Execute(FRHICommandListBase& CmdList)
 {
 	RHISTAT(SetDepthStencilState);
 	INTERNAL_DECORATOR(SetDepthStencilState)(State, StencilRef);
 }
 
 template <typename TShaderRHIParamRef>
-void FRHICommandSetShaderParameter<TShaderRHIParamRef>::Execute()
+void FRHICommandSetShaderParameter<TShaderRHIParamRef>::Execute(FRHICommandListBase& CmdList)
 {
 	RHISTAT(SetShaderParameter);
 	INTERNAL_DECORATOR(SetShaderParameter)(Shader, BufferIndex, BaseIndex, NumBytes, NewValue); 
@@ -44,7 +48,7 @@ template struct FRHICommandSetShaderParameter<FPixelShaderRHIParamRef>;
 template struct FRHICommandSetShaderParameter<FComputeShaderRHIParamRef>;
 
 template <typename TShaderRHIParamRef>
-void FRHICommandSetShaderUniformBuffer<TShaderRHIParamRef>::Execute()
+void FRHICommandSetShaderUniformBuffer<TShaderRHIParamRef>::Execute(FRHICommandListBase& CmdList)
 {
 	RHISTAT(SetShaderUniformBuffer);
 	INTERNAL_DECORATOR(SetShaderUniformBuffer)(Shader, BaseIndex, UniformBuffer);
@@ -57,7 +61,7 @@ template struct FRHICommandSetShaderUniformBuffer<FPixelShaderRHIParamRef>;
 template struct FRHICommandSetShaderUniformBuffer<FComputeShaderRHIParamRef>;
 
 template <typename TShaderRHIParamRef>
-void FRHICommandSetShaderTexture<TShaderRHIParamRef>::Execute()
+void FRHICommandSetShaderTexture<TShaderRHIParamRef>::Execute(FRHICommandListBase& CmdList)
 {
 	RHISTAT(SetShaderTexture);
 	INTERNAL_DECORATOR(SetShaderTexture)(Shader, TextureIndex, Texture);
@@ -70,7 +74,7 @@ template struct FRHICommandSetShaderTexture<FPixelShaderRHIParamRef>;
 template struct FRHICommandSetShaderTexture<FComputeShaderRHIParamRef>;
 
 template <typename TShaderRHIParamRef>
-void FRHICommandSetShaderResourceViewParameter<TShaderRHIParamRef>::Execute()
+void FRHICommandSetShaderResourceViewParameter<TShaderRHIParamRef>::Execute(FRHICommandListBase& CmdList)
 {
 	RHISTAT(SetShaderResourceViewParameter);
 	INTERNAL_DECORATOR(SetShaderResourceViewParameter)(Shader, SamplerIndex, SRV);
@@ -83,7 +87,7 @@ template struct FRHICommandSetShaderResourceViewParameter<FPixelShaderRHIParamRe
 template struct FRHICommandSetShaderResourceViewParameter<FComputeShaderRHIParamRef>;
 
 template <typename TShaderRHIParamRef>
-void FRHICommandSetUAVParameter<TShaderRHIParamRef>::Execute()
+void FRHICommandSetUAVParameter<TShaderRHIParamRef>::Execute(FRHICommandListBase& CmdList)
 {
 	RHISTAT(SetUAVParameter);
 	INTERNAL_DECORATOR(SetUAVParameter)(Shader, UAVIndex, UAV);
@@ -91,7 +95,7 @@ void FRHICommandSetUAVParameter<TShaderRHIParamRef>::Execute()
 template struct FRHICommandSetUAVParameter<FComputeShaderRHIParamRef>;
 
 template <typename TShaderRHIParamRef>
-void FRHICommandSetUAVParameter_IntialCount<TShaderRHIParamRef>::Execute()
+void FRHICommandSetUAVParameter_IntialCount<TShaderRHIParamRef>::Execute(FRHICommandListBase& CmdList)
 {
 	RHISTAT(SetUAVParameter);
 	INTERNAL_DECORATOR(SetUAVParameter)(Shader, UAVIndex, UAV, InitialCount);
@@ -99,7 +103,7 @@ void FRHICommandSetUAVParameter_IntialCount<TShaderRHIParamRef>::Execute()
 template struct FRHICommandSetUAVParameter_IntialCount<FComputeShaderRHIParamRef>;
 
 template <typename TShaderRHIParamRef>
-void FRHICommandSetShaderSampler<TShaderRHIParamRef>::Execute()
+void FRHICommandSetShaderSampler<TShaderRHIParamRef>::Execute(FRHICommandListBase& CmdList)
 {
 	RHISTAT(SetShaderSampler);
 	INTERNAL_DECORATOR(SetShaderSampler)(Shader, SamplerIndex, Sampler);
@@ -111,49 +115,49 @@ template struct FRHICommandSetShaderSampler<FGeometryShaderRHIParamRef>;
 template struct FRHICommandSetShaderSampler<FPixelShaderRHIParamRef>;
 template struct FRHICommandSetShaderSampler<FComputeShaderRHIParamRef>;
 
-void FRHICommandDrawPrimitive::Execute()
+void FRHICommandDrawPrimitive::Execute(FRHICommandListBase& CmdList)
 {
 	RHISTAT(DrawPrimitive);
 	INTERNAL_DECORATOR(DrawPrimitive)(PrimitiveType, BaseVertexIndex, NumPrimitives, NumInstances);
 }
 
-void FRHICommandDrawIndexedPrimitive::Execute()
+void FRHICommandDrawIndexedPrimitive::Execute(FRHICommandListBase& CmdList)
 {
 	RHISTAT(DrawIndexedPrimitive);
 	INTERNAL_DECORATOR(DrawIndexedPrimitive)(IndexBuffer, PrimitiveType, BaseVertexIndex, MinIndex, NumVertices, StartIndex, NumPrimitives, NumInstances);
 }
 
-void FRHICommandSetBoundShaderState::Execute()
+void FRHICommandSetBoundShaderState::Execute(FRHICommandListBase& CmdList)
 {
 	RHISTAT(SetBoundShaderState);
 	INTERNAL_DECORATOR(SetBoundShaderState)(BoundShaderState);
 }
 
-void FRHICommandSetBlendState::Execute()
+void FRHICommandSetBlendState::Execute(FRHICommandListBase& CmdList)
 {
 	RHISTAT(SetBlendState);
 	INTERNAL_DECORATOR(SetBlendState)(State, BlendFactor);
 }
 
-void FRHICommandSetStreamSource::Execute()
+void FRHICommandSetStreamSource::Execute(FRHICommandListBase& CmdList)
 {
 	RHISTAT(SetStreamSource);
 	INTERNAL_DECORATOR(SetStreamSource)(StreamIndex, VertexBuffer, Stride, Offset);
 }
 
-void FRHICommandSetViewport::Execute()
+void FRHICommandSetViewport::Execute(FRHICommandListBase& CmdList)
 {
 	RHISTAT(SetViewport);
 	INTERNAL_DECORATOR(SetViewport)(MinX, MinY, MinZ, MaxX, MaxY, MaxZ);
 }
 
-void FRHICommandSetScissorRect::Execute()
+void FRHICommandSetScissorRect::Execute(FRHICommandListBase& CmdList)
 {
 	RHISTAT(SetScissorRect);
 	INTERNAL_DECORATOR(SetScissorRect)(bEnable, MinX, MinY, MaxX, MaxY);
 }
 
-void FRHICommandSetRenderTargets::Execute()
+void FRHICommandSetRenderTargets::Execute(FRHICommandListBase& CmdList)
 {
 	RHISTAT(SetRenderTargets);
 	INTERNAL_DECORATOR(SetRenderTargets)(
@@ -164,13 +168,13 @@ void FRHICommandSetRenderTargets::Execute()
 		UAVs);
 }
 
-void FRHICommandSetRenderTargetsAndClear::Execute()
+void FRHICommandSetRenderTargetsAndClear::Execute(FRHICommandListBase& CmdList)
 {
 	RHISTAT(SetRenderTargetsAndClear);
 	INTERNAL_DECORATOR(SetRenderTargetsAndClear)(RenderTargetsInfo);
 }
 
-void FRHICommandEndDrawPrimitiveUP::Execute()
+void FRHICommandEndDrawPrimitiveUP::Execute(FRHICommandListBase& CmdList)
 {
 	RHISTAT(EndDrawPrimitiveUP);
 	void* Buffer = NULL;
@@ -179,7 +183,7 @@ void FRHICommandEndDrawPrimitiveUP::Execute()
 	INTERNAL_DECORATOR(EndDrawPrimitiveUP)();
 }
 
-void FRHICommandEndDrawIndexedPrimitiveUP::Execute()
+void FRHICommandEndDrawIndexedPrimitiveUP::Execute(FRHICommandListBase& CmdList)
 {
 	RHISTAT(EndDrawIndexedPrimitiveUP);
 	void* VertexBuffer = nullptr;
@@ -199,96 +203,100 @@ void FRHICommandEndDrawIndexedPrimitiveUP::Execute()
 	INTERNAL_DECORATOR(EndDrawIndexedPrimitiveUP)();
 }
 
-void FRHICommandSetComputeShader::Execute()
+void FRHICommandSetComputeShader::Execute(FRHICommandListBase& CmdList)
 {
 	RHISTAT(SetComputeShader);
 	INTERNAL_DECORATOR(SetComputeShader)(ComputeShader);
 }
 
-void FRHICommandDispatchComputeShader::Execute()
+void FRHICommandDispatchComputeShader::Execute(FRHICommandListBase& CmdList)
 {
 	RHISTAT(DispatchComputeShader);
 	INTERNAL_DECORATOR(DispatchComputeShader)(ThreadGroupCountX, ThreadGroupCountY, ThreadGroupCountZ);
 }
 
-void FRHICommandDispatchIndirectComputeShader::Execute()
+void FRHICommandDispatchIndirectComputeShader::Execute(FRHICommandListBase& CmdList)
 {
 	RHISTAT(DispatchIndirectComputeShader);
 	INTERNAL_DECORATOR(DispatchIndirectComputeShader)(ArgumentBuffer, ArgumentOffset);
 }
 
-void FRHICommandAutomaticCacheFlushAfterComputeShader::Execute()
+void FRHICommandAutomaticCacheFlushAfterComputeShader::Execute(FRHICommandListBase& CmdList)
 {
 	RHISTAT(AutomaticCacheFlushAfterComputeShader);
 	INTERNAL_DECORATOR(AutomaticCacheFlushAfterComputeShader)(bEnable);
 }
 
-void FRHICommandFlushComputeShaderCache::Execute()
+void FRHICommandFlushComputeShaderCache::Execute(FRHICommandListBase& CmdList)
 {
 	RHISTAT(FlushComputeShaderCache);
 	INTERNAL_DECORATOR(FlushComputeShaderCache)();
 }
 
-void FRHICommandDrawPrimitiveIndirect::Execute()
+void FRHICommandDrawPrimitiveIndirect::Execute(FRHICommandListBase& CmdList)
 {
 	RHISTAT(DrawPrimitiveIndirect);
 	INTERNAL_DECORATOR(DrawPrimitiveIndirect)(PrimitiveType, ArgumentBuffer, ArgumentOffset);
 }
 
-void FRHICommandDrawIndexedIndirect::Execute()
+void FRHICommandDrawIndexedIndirect::Execute(FRHICommandListBase& CmdList)
 {
 	RHISTAT(DrawIndexedIndirect);
 	INTERNAL_DECORATOR(DrawIndexedIndirect)(IndexBufferRHI, PrimitiveType, ArgumentsBufferRHI, DrawArgumentsIndex, NumInstances);
 }
 
-void FRHICommandDrawIndexedPrimitiveIndirect::Execute()
+void FRHICommandDrawIndexedPrimitiveIndirect::Execute(FRHICommandListBase& CmdList)
 {
 	RHISTAT(DrawIndexedPrimitiveIndirect);
 	INTERNAL_DECORATOR(DrawIndexedPrimitiveIndirect)(PrimitiveType, IndexBuffer, ArgumentsBuffer, ArgumentOffset);
 }
 
-void FRHICommandEnableDepthBoundsTest::Execute()
+void FRHICommandEnableDepthBoundsTest::Execute(FRHICommandListBase& CmdList)
 {
 	RHISTAT(EnableDepthBoundsTest);
 	INTERNAL_DECORATOR(EnableDepthBoundsTest)(bEnable, MinDepth, MaxDepth);
 }
 
-void FRHICommandClearUAV::Execute()
+void FRHICommandClearUAV::Execute(FRHICommandListBase& CmdList)
 {
 	RHISTAT(ClearUAV);
 	INTERNAL_DECORATOR(ClearUAV)(UnorderedAccessViewRHI, Values);
 }
 
-void FRHICommandCopyToResolveTarget::Execute()
+void FRHICommandCopyToResolveTarget::Execute(FRHICommandListBase& CmdList)
 {
 	RHISTAT(CopyToResolveTarget);
 	INTERNAL_DECORATOR(CopyToResolveTarget)(SourceTexture, DestTexture, bKeepOriginalSurface, ResolveParams);
 }
 
-void FRHICommandClear::Execute()
+void FRHICommandClear::Execute(FRHICommandListBase& CmdList)
 {
 	RHISTAT(Clear);
 	INTERNAL_DECORATOR(Clear)(bClearColor, Color, bClearDepth, Depth, bClearStencil, Stencil, ExcludeRect);
 }
 
-void FRHICommandClearMRT::Execute()
+void FRHICommandClearMRT::Execute(FRHICommandListBase& CmdList)
 {
 	RHISTAT(ClearMRT);
 	INTERNAL_DECORATOR(ClearMRT)(bClearColor, NumClearColors, ColorArray, bClearDepth, Depth, bClearStencil, Stencil, ExcludeRect);
 }
 
-void FRHICommandBuildLocalBoundShaderState::Execute()
+void FRHICommandBuildLocalBoundShaderState::Execute(FRHICommandListBase& CmdList)
 {
 	RHISTAT(BuildLocalBoundShaderState);
 	check(!IsValidRef(WorkArea.ComputedBSS->BSS)); // should not already have been created
 	if (WorkArea.ComputedBSS->UseCount)
 	{
 		WorkArea.ComputedBSS->BSS = 	
-			INTERNAL_DECORATOR(CreateBoundShaderState)(WorkArea.Args.VertexDeclarationRHI, WorkArea.Args.VertexShaderRHI, WorkArea.Args.HullShaderRHI, WorkArea.Args.DomainShaderRHI, WorkArea.Args.PixelShaderRHI, WorkArea.Args.GeometryShaderRHI);
+#if PLATFORM_SUPPORTS_PARALLEL_RHI_EXECUTE
+			RHICreateBoundShaderState(WorkArea.Args.VertexDeclarationRHI, WorkArea.Args.VertexShaderRHI, WorkArea.Args.HullShaderRHI, WorkArea.Args.DomainShaderRHI, WorkArea.Args.PixelShaderRHI, WorkArea.Args.GeometryShaderRHI);
+#else
+			CreateBoundShaderState_Internal(WorkArea.Args.VertexDeclarationRHI, WorkArea.Args.VertexShaderRHI, WorkArea.Args.HullShaderRHI, WorkArea.Args.DomainShaderRHI, WorkArea.Args.PixelShaderRHI, WorkArea.Args.GeometryShaderRHI);
+#endif
 	}
 }
 
-void FRHICommandSetLocalBoundShaderState::Execute()
+void FRHICommandSetLocalBoundShaderState::Execute(FRHICommandListBase& CmdList)
 {
 	RHISTAT(SetLocalBoundShaderState);
 	check(LocalBoundShaderState.WorkArea->ComputedBSS->UseCount > 0 && IsValidRef(LocalBoundShaderState.WorkArea->ComputedBSS->BSS)); // this should have been created and should have uses outstanding
@@ -301,7 +309,7 @@ void FRHICommandSetLocalBoundShaderState::Execute()
 	}
 }
 
-void FRHICommandBuildLocalUniformBuffer::Execute()
+void FRHICommandBuildLocalUniformBuffer::Execute(FRHICommandListBase& CmdList)
 {
 	RHISTAT(BuildLocalUniformBuffer);
 	check(!IsValidRef(WorkArea.ComputedUniformBuffer->UniformBuffer) && WorkArea.Layout && WorkArea.Contents); // should not already have been created
@@ -310,7 +318,7 @@ void FRHICommandBuildLocalUniformBuffer::Execute()
 #if PLATFORM_SUPPORTS_RHI_THREAD
 		WorkArea.ComputedUniformBuffer->UniformBuffer = RHICreateUniformBuffer(WorkArea.Contents, *WorkArea.Layout, UniformBuffer_SingleFrame); 
 #else
-		WorkArea.ComputedUniformBuffer->UniformBuffer = INTERNAL_DECORATOR(CreateUniformBuffer)(WorkArea.Contents, *WorkArea.Layout, UniformBuffer_SingleFrame);
+		WorkArea.ComputedUniformBuffer->UniformBuffer = CreateUniformBuffer_Internal(WorkArea.Contents, *WorkArea.Layout, UniformBuffer_SingleFrame);
 #endif
 	}
 	WorkArea.Layout = nullptr;
@@ -318,7 +326,7 @@ void FRHICommandBuildLocalUniformBuffer::Execute()
 }
 
 template <typename TShaderRHIParamRef>
-void FRHICommandSetLocalUniformBuffer<TShaderRHIParamRef>::Execute()
+void FRHICommandSetLocalUniformBuffer<TShaderRHIParamRef>::Execute(FRHICommandListBase& CmdList)
 {
 	RHISTAT(SetLocalUniformBuffer);
 	check(LocalUniformBuffer.WorkArea->ComputedUniformBuffer->UseCount > 0 && IsValidRef(LocalUniformBuffer.WorkArea->ComputedUniformBuffer->UniformBuffer)); // this should have been created and should have uses outstanding
@@ -335,79 +343,95 @@ template struct FRHICommandSetLocalUniformBuffer<FGeometryShaderRHIParamRef>;
 template struct FRHICommandSetLocalUniformBuffer<FPixelShaderRHIParamRef>;
 template struct FRHICommandSetLocalUniformBuffer<FComputeShaderRHIParamRef>;
 
-void FRHICommandBeginRenderQuery::Execute()
+void FRHICommandBeginRenderQuery::Execute(FRHICommandListBase& CmdList)
 {
 	RHISTAT(BeginRenderQuery);
 	INTERNAL_DECORATOR(BeginRenderQuery)(RenderQuery);
 }
 
-void FRHICommandEndRenderQuery::Execute()
+void FRHICommandEndRenderQuery::Execute(FRHICommandListBase& CmdList)
 {
 	RHISTAT(EndRenderQuery);
 	INTERNAL_DECORATOR(EndRenderQuery)(RenderQuery);
 }
 
 #if !PLATFORM_SUPPORTS_RHI_THREAD
-void FRHICommandResetRenderQuery::Execute()
+void FRHICommandResetRenderQuery::Execute(FRHICommandListBase& CmdList)
 {
 	RHISTAT(ResetRenderQuery);
 	INTERNAL_DECORATOR(ResetRenderQuery)(RenderQuery);
 }
 #endif
 
-void FRHICommandBeginScene::Execute()
+void FRHICommandBeginScene::Execute(FRHICommandListBase& CmdList)
 {
 	RHISTAT(BeginScene);
+#if PLATFORM_SUPPORTS_RHI_THREAD
 	INTERNAL_DECORATOR(BeginScene)();
+#else
+	check(0);
+#endif
 }
 
-void FRHICommandEndScene::Execute()
+void FRHICommandEndScene::Execute(FRHICommandListBase& CmdList)
 {
 	RHISTAT(EndScene);
+#if PLATFORM_SUPPORTS_RHI_THREAD
 	INTERNAL_DECORATOR(EndScene)();
+#else
+	check(0);
+#endif
 }
 
-void FRHICommandUpdateVertexBuffer::Execute()
+void FRHICommandUpdateVertexBuffer::Execute(FRHICommandListBase& CmdList)
 {
 	RHISTAT(UpdateVertexBuffer);
-	void* Data = INTERNAL_DECORATOR(LockVertexBuffer)(VertexBuffer, 0, BufferSize, RLM_WriteOnly);
+	void* Data = LockVertexBuffer_Internal(VertexBuffer, 0, BufferSize, RLM_WriteOnly);
 	FMemory::Memcpy(Data, Buffer, BufferSize);
-	INTERNAL_DECORATOR(UnlockVertexBuffer)(VertexBuffer);
+	UnlockVertexBuffer_Internal(VertexBuffer);
 }
 
-void FRHICommandBeginFrame::Execute()
+void FRHICommandBeginFrame::Execute(FRHICommandListBase& CmdList)
 {
 	RHISTAT(BeginFrame);
+#if PLATFORM_SUPPORTS_RHI_THREAD
 	INTERNAL_DECORATOR(BeginFrame)();
+#else
+	check(0);
+#endif
 }
 
-void FRHICommandEndFrame::Execute()
+void FRHICommandEndFrame::Execute(FRHICommandListBase& CmdList)
 {
 	RHISTAT(EndFrame);
+#if PLATFORM_SUPPORTS_RHI_THREAD
 	INTERNAL_DECORATOR(EndFrame)();
+#else
+	check(0);
+#endif
 }
 
 #if PLATFORM_SUPPORTS_RHI_THREAD
-void FRHICommandBeginDrawingViewport::Execute()
+void FRHICommandBeginDrawingViewport::Execute(FRHICommandListBase& CmdList)
 {
 	RHISTAT(BeginDrawingViewport);
 	INTERNAL_DECORATOR(BeginDrawingViewport)(Viewport, RenderTargetRHI);
 }
 
-void FRHICommandEndDrawingViewport::Execute()
+void FRHICommandEndDrawingViewport::Execute(FRHICommandListBase& CmdList)
 {
 	RHISTAT(EndDrawingViewport);
 	INTERNAL_DECORATOR(EndDrawingViewport)(Viewport, bPresent, bLockToVsync);
 }
 #endif
 
-void FRHICommandPushEvent::Execute()
+void FRHICommandPushEvent::Execute(FRHICommandListBase& CmdList)
 {
 	RHISTAT(PushEvent);
 	INTERNAL_DECORATOR(PushEvent)(Name);
 }
 
-void FRHICommandPopEvent::Execute()
+void FRHICommandPopEvent::Execute(FRHICommandListBase& CmdList)
 {
 	RHISTAT(PopEvent);
 	INTERNAL_DECORATOR(PopEvent)();
diff --git a/Engine/Source/Runtime/RHI/Public/RHIMethods.h b/Engine/Source/Runtime/RHI/Public/RHIMethods.h
--- a/Engine/Source/Runtime/RHI/Public/RHIMethods.h
+++ b/Engine/Source/Runtime/RHI/Public/RHIMethods.h
@@ -477,16 +477,29 @@ DEFINE_RHIMETHOD_GLOBAL_1(
  * @param GeometryShader - existing geometry shader
  * @param PixelShader - existing pixel shader
  */
-DEFINE_RHIMETHOD_CMDLIST_6(
-	FBoundShaderStateRHIRef,CreateBoundShaderState,
-	FVertexDeclarationRHIParamRef,VertexDeclaration,
-	FVertexShaderRHIParamRef,VertexShader,
-	FHullShaderRHIParamRef,HullShader,
-	FDomainShaderRHIParamRef,DomainShader,
-	FPixelShaderRHIParamRef,PixelShader,
-	FGeometryShaderRHIParamRef,GeometryShader,
-	return,return new FRHIBoundShaderState();
+#if PLATFORM_SUPPORTS_PARALLEL_RHI_EXECUTE
+	DEFINE_RHIMETHOD_GLOBALTHREADSAFE_6(
+		FBoundShaderStateRHIRef,CreateBoundShaderState,
+		FVertexDeclarationRHIParamRef,VertexDeclaration,
+		FVertexShaderRHIParamRef,VertexShader,
+		FHullShaderRHIParamRef,HullShader,
+		FDomainShaderRHIParamRef,DomainShader,
+		FPixelShaderRHIParamRef,PixelShader,
+		FGeometryShaderRHIParamRef,GeometryShader,
+		return,return new FRHIBoundShaderState();
 	);
+#else
+	DEFINE_RHIMETHOD_GLOBAL_6(
+		FBoundShaderStateRHIRef,CreateBoundShaderState,
+		FVertexDeclarationRHIParamRef,VertexDeclaration,
+		FVertexShaderRHIParamRef,VertexShader,
+		FHullShaderRHIParamRef,HullShader,
+		FDomainShaderRHIParamRef,DomainShader,
+		FPixelShaderRHIParamRef,PixelShader,
+		FGeometryShaderRHIParamRef,GeometryShader,
+		return,return new FRHIBoundShaderState();
+		);
+#endif
 
 /**
  *Sets the current compute shader.  Mostly for compliance with platforms
@@ -1478,6 +1491,7 @@ DEFINE_RHIMETHOD_CMDLIST_4(
 	);
 
 /** Sets stream output targets, for use with a geometry shader created with RHICreateGeometryShaderWithStreamOutput. */
+//@todo this should be a CMDLIST method
 DEFINE_RHIMETHOD_3(
 	void,SetStreamOutTargets,
 	uint32,NumTargets,
@@ -2136,3 +2150,20 @@ DEFINE_RHIMETHOD_CMDLIST_0(
 	void,PopEvent,
 	,
 	);
+
+#if PLATFORM_RHI_USES_CONTEXT_OBJECT
+	DEFINE_RHIMETHOD_GLOBALTHREADSAFE_0(
+		class IRHICommandContext*,GetDefaultContext,
+		return,return this
+	);
+
+	#if PLATFORM_SUPPORTS_PARALLEL_RHI_EXECUTE
+			DEFINE_RHIMETHOD_GLOBALTHREADSAFE_1(
+				class IRHICommandContextContainer*,GetCommandContextContainer,
+				uint32, ClearTag,
+				return,return nullptr
+				);
+	#endif
+
+#endif
+
diff --git a/Engine/Source/Runtime/RenderCore/Public/RenderResource.h b/Engine/Source/Runtime/RenderCore/Public/RenderResource.h
--- a/Engine/Source/Runtime/RenderCore/Public/RenderResource.h
+++ b/Engine/Source/Runtime/RenderCore/Public/RenderResource.h
@@ -619,12 +619,16 @@ public:
 	/** Adds a bound shader state to the history. */
 	void Add(FBoundShaderStateRHIParamRef BoundShaderState)
 	{
+#if PLATFORM_SUPPORTS_PARALLEL_RHI_EXECUTE
+		FScopeLock Lock(&BoundShaderStateHistoryLock);
+#endif
 		BoundShaderStates[NextBoundShaderStateIndex] = BoundShaderState;
 		NextBoundShaderStateIndex = (NextBoundShaderStateIndex + 1) % Size;
 	}
 
 	FBoundShaderStateRHIParamRef GetLast()
 	{
+		check(!PLATFORM_SUPPORTS_PARALLEL_RHI_EXECUTE);
 		// % doesn't work as we want on negative numbers, so handle the wraparound manually
 		uint32 LastIndex = NextBoundShaderStateIndex == 0 ? Size - 1 : NextBoundShaderStateIndex - 1;
 		return BoundShaderStates[LastIndex];
@@ -643,4 +647,7 @@ private:
 
 	FBoundShaderStateRHIRef BoundShaderStates[Size];
 	uint32 NextBoundShaderStateIndex;
+#if PLATFORM_SUPPORTS_PARALLEL_RHI_EXECUTE
+	FCriticalSection BoundShaderStateHistoryLock;
+#endif
 };
diff --git a/Engine/Source/Runtime/Renderer/Private/DeferredShadingRenderer.cpp b/Engine/Source/Runtime/Renderer/Private/DeferredShadingRenderer.cpp
--- a/Engine/Source/Runtime/Renderer/Private/DeferredShadingRenderer.cpp
+++ b/Engine/Source/Runtime/Renderer/Private/DeferredShadingRenderer.cpp
@@ -218,23 +218,25 @@ bool FDeferredShadingSceneRenderer::RenderBasePassStaticDataMasked(FRHICommandLi
 	return bDirty;
 }
 
-void FDeferredShadingSceneRenderer::RenderBasePassStaticDataMaskedParallel(FViewInfo& View, int32 Width, FRHICommandList& ParentCmdList, bool& OutDirty)
+void FDeferredShadingSceneRenderer::RenderBasePassStaticDataMaskedParallel(FParallelCommandListSet& ParallelCommandListSet)
 {
 	// Draw the scene's base pass draw lists.
 	FScene::EBasePassDrawListType MaskedDrawType = FScene::EBasePass_Masked;
 	{
-		SCOPED_DRAW_EVENT(ParentCmdList, StaticMaskedNoLightmap, DEC_SCENE_ITEMS);
-		Scene->BasePassNoLightMapDrawList[MaskedDrawType].DrawVisibleParallel(View, View.StaticMeshVisibilityMap, View.StaticMeshBatchVisibility, Width, ParentCmdList, OutDirty);
-		Scene->BasePassSimpleDynamicLightingDrawList[MaskedDrawType].DrawVisibleParallel(View, View.StaticMeshVisibilityMap, View.StaticMeshBatchVisibility, Width, ParentCmdList, OutDirty);
-		Scene->BasePassCachedVolumeIndirectLightingDrawList[MaskedDrawType].DrawVisibleParallel(View, View.StaticMeshVisibilityMap, View.StaticMeshBatchVisibility, Width, ParentCmdList, OutDirty);
-		Scene->BasePassCachedPointIndirectLightingDrawList[MaskedDrawType].DrawVisibleParallel(View, View.StaticMeshVisibilityMap, View.StaticMeshBatchVisibility, Width, ParentCmdList, OutDirty);
+		// we can't insert this event on the parent command list; need to pass it along to the ParallelCommandListSet
+		//SCOPED_DRAW_EVENT(ParentCmdList, StaticMaskedNoLightmap, DEC_SCENE_ITEMS);
+		Scene->BasePassNoLightMapDrawList[MaskedDrawType].DrawVisibleParallel(ParallelCommandListSet.View.StaticMeshVisibilityMap, ParallelCommandListSet.View.StaticMeshBatchVisibility, ParallelCommandListSet);
+		Scene->BasePassSimpleDynamicLightingDrawList[MaskedDrawType].DrawVisibleParallel(ParallelCommandListSet.View.StaticMeshVisibilityMap, ParallelCommandListSet.View.StaticMeshBatchVisibility, ParallelCommandListSet);
+		Scene->BasePassCachedVolumeIndirectLightingDrawList[MaskedDrawType].DrawVisibleParallel(ParallelCommandListSet.View.StaticMeshVisibilityMap, ParallelCommandListSet.View.StaticMeshBatchVisibility, ParallelCommandListSet);
+		Scene->BasePassCachedPointIndirectLightingDrawList[MaskedDrawType].DrawVisibleParallel(ParallelCommandListSet.View.StaticMeshVisibilityMap, ParallelCommandListSet.View.StaticMeshBatchVisibility, ParallelCommandListSet);
 	}
 
 	{
-		SCOPED_DRAW_EVENT(ParentCmdList, StaticMaskedLightmapped, DEC_SCENE_ITEMS);
-		Scene->BasePassHighQualityLightMapDrawList[MaskedDrawType].DrawVisibleParallel(View, View.StaticMeshVisibilityMap, View.StaticMeshBatchVisibility, Width, ParentCmdList, OutDirty);
-		Scene->BasePassDistanceFieldShadowMapLightMapDrawList[MaskedDrawType].DrawVisibleParallel(View, View.StaticMeshVisibilityMap, View.StaticMeshBatchVisibility, Width, ParentCmdList, OutDirty);
-		Scene->BasePassLowQualityLightMapDrawList[MaskedDrawType].DrawVisibleParallel(View, View.StaticMeshVisibilityMap, View.StaticMeshBatchVisibility, Width, ParentCmdList, OutDirty);
+		// we can't insert this event on the parent command list; need to pass it along to the ParallelCommandListSet
+		//SCOPED_DRAW_EVENT(ParentCmdList, StaticMaskedLightmapped, DEC_SCENE_ITEMS);
+		Scene->BasePassHighQualityLightMapDrawList[MaskedDrawType].DrawVisibleParallel(ParallelCommandListSet.View.StaticMeshVisibilityMap, ParallelCommandListSet.View.StaticMeshBatchVisibility, ParallelCommandListSet);
+		Scene->BasePassDistanceFieldShadowMapLightMapDrawList[MaskedDrawType].DrawVisibleParallel(ParallelCommandListSet.View.StaticMeshVisibilityMap, ParallelCommandListSet.View.StaticMeshBatchVisibility, ParallelCommandListSet);
+		Scene->BasePassLowQualityLightMapDrawList[MaskedDrawType].DrawVisibleParallel(ParallelCommandListSet.View.StaticMeshVisibilityMap, ParallelCommandListSet.View.StaticMeshBatchVisibility, ParallelCommandListSet);
 	}
 }
 
@@ -261,22 +263,24 @@ bool FDeferredShadingSceneRenderer::RenderBasePassStaticDataDefault(FRHICommandL
 	return bDirty;
 }
 
-void FDeferredShadingSceneRenderer::RenderBasePassStaticDataDefaultParallel(FViewInfo& View, int32 Width, FRHICommandList& ParentCmdList, bool& OutDirty)
+void FDeferredShadingSceneRenderer::RenderBasePassStaticDataDefaultParallel(FParallelCommandListSet& ParallelCommandListSet)
 {
 	FScene::EBasePassDrawListType OpaqueDrawType = FScene::EBasePass_Default;
 	{
-		SCOPED_DRAW_EVENT(ParentCmdList, StaticOpaqueNoLightmap, DEC_SCENE_ITEMS);
-		Scene->BasePassNoLightMapDrawList[OpaqueDrawType].DrawVisibleParallel(View, View.StaticMeshVisibilityMap, View.StaticMeshBatchVisibility, Width, ParentCmdList, OutDirty);
-		Scene->BasePassSimpleDynamicLightingDrawList[OpaqueDrawType].DrawVisibleParallel(View, View.StaticMeshVisibilityMap, View.StaticMeshBatchVisibility, Width, ParentCmdList, OutDirty);
-		Scene->BasePassCachedVolumeIndirectLightingDrawList[OpaqueDrawType].DrawVisibleParallel(View, View.StaticMeshVisibilityMap, View.StaticMeshBatchVisibility, Width, ParentCmdList, OutDirty);
-		Scene->BasePassCachedPointIndirectLightingDrawList[OpaqueDrawType].DrawVisibleParallel(View, View.StaticMeshVisibilityMap, View.StaticMeshBatchVisibility, Width, ParentCmdList, OutDirty);
+		// we can't insert this event on the parent command list; need to pass it along to the ParallelCommandListSet
+		//SCOPED_DRAW_EVENT(ParentCmdList, StaticOpaqueNoLightmap, DEC_SCENE_ITEMS);
+		Scene->BasePassNoLightMapDrawList[OpaqueDrawType].DrawVisibleParallel(ParallelCommandListSet.View.StaticMeshVisibilityMap, ParallelCommandListSet.View.StaticMeshBatchVisibility, ParallelCommandListSet);
+		Scene->BasePassSimpleDynamicLightingDrawList[OpaqueDrawType].DrawVisibleParallel(ParallelCommandListSet.View.StaticMeshVisibilityMap, ParallelCommandListSet.View.StaticMeshBatchVisibility, ParallelCommandListSet);
+		Scene->BasePassCachedVolumeIndirectLightingDrawList[OpaqueDrawType].DrawVisibleParallel(ParallelCommandListSet.View.StaticMeshVisibilityMap, ParallelCommandListSet.View.StaticMeshBatchVisibility, ParallelCommandListSet);
+		Scene->BasePassCachedPointIndirectLightingDrawList[OpaqueDrawType].DrawVisibleParallel(ParallelCommandListSet.View.StaticMeshVisibilityMap, ParallelCommandListSet.View.StaticMeshBatchVisibility, ParallelCommandListSet);
 	}
 
 	{
-		SCOPED_DRAW_EVENT(ParentCmdList, StaticOpaqueLightmapped, DEC_SCENE_ITEMS);
-		Scene->BasePassHighQualityLightMapDrawList[OpaqueDrawType].DrawVisibleParallel(View, View.StaticMeshVisibilityMap, View.StaticMeshBatchVisibility, Width, ParentCmdList, OutDirty);
-		Scene->BasePassDistanceFieldShadowMapLightMapDrawList[OpaqueDrawType].DrawVisibleParallel(View, View.StaticMeshVisibilityMap, View.StaticMeshBatchVisibility, Width, ParentCmdList, OutDirty);
-		Scene->BasePassLowQualityLightMapDrawList[OpaqueDrawType].DrawVisibleParallel(View, View.StaticMeshVisibilityMap, View.StaticMeshBatchVisibility, Width, ParentCmdList, OutDirty);
+		// we can't insert this event on the parent command list; need to pass it along to the ParallelCommandListSet
+		//SCOPED_DRAW_EVENT(ParentCmdList, StaticOpaqueLightmapped, DEC_SCENE_ITEMS);
+		Scene->BasePassHighQualityLightMapDrawList[OpaqueDrawType].DrawVisibleParallel(ParallelCommandListSet.View.StaticMeshVisibilityMap, ParallelCommandListSet.View.StaticMeshBatchVisibility, ParallelCommandListSet);
+		Scene->BasePassDistanceFieldShadowMapLightMapDrawList[OpaqueDrawType].DrawVisibleParallel(ParallelCommandListSet.View.StaticMeshVisibilityMap, ParallelCommandListSet.View.StaticMeshBatchVisibility, ParallelCommandListSet);
+		Scene->BasePassLowQualityLightMapDrawList[OpaqueDrawType].DrawVisibleParallel(ParallelCommandListSet.View.StaticMeshVisibilityMap, ParallelCommandListSet.View.StaticMeshBatchVisibility, ParallelCommandListSet);
 	}
 }
 
@@ -332,7 +336,7 @@ bool FDeferredShadingSceneRenderer::RenderBasePassStaticData(FRHICommandList& RH
 	return bDirty;
 }
 
-void FDeferredShadingSceneRenderer::RenderBasePassStaticDataParallel(FViewInfo& View, int32 Width, FRHICommandList& ParentCmdList, bool& OutDirty)
+void FDeferredShadingSceneRenderer::RenderBasePassStaticDataParallel(FParallelCommandListSet& ParallelCommandListSet)
 {
 	SCOPE_CYCLE_COUNTER(STAT_StaticDrawListDrawTime);
 
@@ -342,16 +346,16 @@ void FDeferredShadingSceneRenderer::RenderBasePassStaticDataParallel(FViewInfo&
 	// rendering the default opaque geometry afterward.
 	if (EarlyZPassMode != DDM_None)
 	{
-		RenderBasePassStaticDataMaskedParallel(View, Width, ParentCmdList, OutDirty);
-		RenderBasePassStaticDataDefaultParallel(View, Width, ParentCmdList, OutDirty);
+		RenderBasePassStaticDataMaskedParallel(ParallelCommandListSet);
+		RenderBasePassStaticDataDefaultParallel(ParallelCommandListSet);
 	}
 	else
 	{
 		// Otherwise, in the case where we're not using a depth-only pre-pass, there
 		// is an advantage to rendering default opaque first to help cull the more
 		// expensive masked geometry.
-		RenderBasePassStaticDataDefaultParallel(View, Width, ParentCmdList, OutDirty);
-		RenderBasePassStaticDataMaskedParallel(View, Width, ParentCmdList, OutDirty);
+		RenderBasePassStaticDataDefaultParallel(ParallelCommandListSet);
+		RenderBasePassStaticDataMaskedParallel(ParallelCommandListSet);
 	}
 }
 
@@ -487,14 +491,13 @@ public:
 	}
 };
 
-void FDeferredShadingSceneRenderer::RenderBasePassDynamicDataParallel(FViewInfo& View, FRHICommandList& ParentCmdList, bool& bOutDirty)
+void FDeferredShadingSceneRenderer::RenderBasePassDynamicDataParallel(FParallelCommandListSet& ParallelCommandListSet)
 {
-	FRHICommandList* CmdList = new FRHICommandList;
-
+	FRHICommandList* CmdList = ParallelCommandListSet.NewParallelCommandList();
 	FGraphEventRef AnyThreadCompletionEvent = TGraphTask<FRenderBasePassDynamicDataThreadTask>::CreateTask(nullptr, ENamedThreads::RenderThread)
-		.ConstructAndDispatchWhenReady(this, CmdList, &View, &bOutDirty);
+		.ConstructAndDispatchWhenReady(this, CmdList, &ParallelCommandListSet.View, &ParallelCommandListSet.OutDirty);
 
-	ParentCmdList.QueueAsyncCommandListSubmit(AnyThreadCompletionEvent, CmdList);
+	ParallelCommandListSet.AddParallelCommandList(CmdList, AnyThreadCompletionEvent);
 }
 
 static void SetupBasePassView(FRHICommandList& RHICmdList, const FIntRect& ViewRect, bool bShaderComplexity)
@@ -514,15 +517,39 @@ static void SetupBasePassView(FRHICommandList& RHICmdList, const FIntRect& ViewR
 		RHICmdList.SetDepthStencilState(TStaticDepthStencilState<true,CF_GreaterEqual>::GetRHI());
 	}
 	RHICmdList.SetViewport(ViewRect.Min.X, ViewRect.Min.Y, 0, ViewRect.Max.X, ViewRect.Max.Y, 1);
+	RHICmdList.SetRasterizerState(TStaticRasterizerState<FM_Solid, CM_None>::GetRHI());
+	RHICmdList.SetScissorRect(false, 0, 0, 0, 0);
 }
 
-void FDeferredShadingSceneRenderer::RenderBasePassViewParallel(FViewInfo& View, int32 Width, FRHICommandList& ParentCmdList, bool& OutDirty)
+class FBasePassParallelCommandListSet : public FParallelCommandListSet
+{
+public:
+	const FSceneViewFamily& ViewFamily;
+
+	FBasePassParallelCommandListSet(const FViewInfo& InView, FRHICommandList& InParentCmdList, bool* InOutDirty, bool bInParallelExecute, const FSceneViewFamily& InViewFamily)
+		: FParallelCommandListSet(InView, InParentCmdList, InOutDirty, bInParallelExecute)
+		, ViewFamily(InViewFamily)
+	{
+		SetStateOnCommandList(ParentCmdList);
+	}
+	virtual void SetStateOnCommandList(FRHICommandList& CmdList) override
+	{
+		GSceneRenderTargets.BeginRenderingSceneColor(CmdList, true);
+		SetupBasePassView(CmdList, View.ViewRect, !!ViewFamily.EngineShowFlags.ShaderComplexity);
+	}
+};
+
+static TAutoConsoleVariable<int32> CVarRHICmdBasePassDeferredContexts(
+	TEXT("r.RHICmdBasePassDeferredContexts"),
+	1,
+	TEXT("True to use deferred contexts to parallelize base pass command list execution.\n"));
+
+void FDeferredShadingSceneRenderer::RenderBasePassViewParallel(FViewInfo& View, FRHICommandList& ParentCmdList, bool& OutDirty)
 {
-	// this is probably only needed on the first view
-	SetupBasePassView(ParentCmdList, View.ViewRect, !!ViewFamily.EngineShowFlags.ShaderComplexity);
+	FBasePassParallelCommandListSet ParallelSet(View, ParentCmdList, &OutDirty, CVarRHICmdBasePassDeferredContexts.GetValueOnRenderThread() > 0, ViewFamily);
 
-	RenderBasePassStaticDataParallel(View, Width, ParentCmdList, OutDirty);
-	RenderBasePassDynamicDataParallel(View, ParentCmdList, OutDirty);
+	RenderBasePassStaticDataParallel(ParallelSet);
+	RenderBasePassDynamicDataParallel(ParallelSet);
 }
 
 bool FDeferredShadingSceneRenderer::RenderBasePassView(FRHICommandListImmediate& RHICmdList, FViewInfo& View)
@@ -605,6 +632,21 @@ void FDeferredShadingSceneRenderer::RenderFinish(FRHICommandListImmediate& RHICm
 * Renders the view family. 
 */
 
+DECLARE_STATS_GROUP(TEXT("Command List Markers"), STATGROUP_CommandListMarkers, STATCAT_Advanced);
+
+DECLARE_CYCLE_STAT(TEXT("PrePass"), STAT_CLM_PrePass, STATGROUP_CommandListMarkers);
+DECLARE_CYCLE_STAT(TEXT("AfterPrePass"), STAT_CLM_AfterPrePass, STATGROUP_CommandListMarkers);
+DECLARE_CYCLE_STAT(TEXT("BasePass"), STAT_CLM_BasePass, STATGROUP_CommandListMarkers);
+DECLARE_CYCLE_STAT(TEXT("AfterBasePass"), STAT_CLM_AfterBasePass, STATGROUP_CommandListMarkers);
+DECLARE_CYCLE_STAT(TEXT("Lighting"), STAT_CLM_Lighting, STATGROUP_CommandListMarkers);
+DECLARE_CYCLE_STAT(TEXT("AfterLighting"), STAT_CLM_AfterLighting, STATGROUP_CommandListMarkers);
+DECLARE_CYCLE_STAT(TEXT("Translucency"), STAT_CLM_Translucency, STATGROUP_CommandListMarkers);
+DECLARE_CYCLE_STAT(TEXT("AfterTranslucency"), STAT_CLM_AfterTranslucency, STATGROUP_CommandListMarkers);
+DECLARE_CYCLE_STAT(TEXT("Velocity"), STAT_CLM_Velocity, STATGROUP_CommandListMarkers);
+DECLARE_CYCLE_STAT(TEXT("AfterVelocity"), STAT_CLM_AfterVelocity, STATGROUP_CommandListMarkers);
+DECLARE_CYCLE_STAT(TEXT("AfterFrame"), STAT_CLM_AfterFrame, STATGROUP_CommandListMarkers);
+
+
 void FDeferredShadingSceneRenderer::Render(FRHICommandListImmediate& RHICmdList)
 {
 	bool bDBuffer = IsDBufferEnabled();
@@ -711,7 +753,9 @@ void FDeferredShadingSceneRenderer::Render(FRHICommandListImmediate& RHICmdList)
 	GRenderTargetPool.AddPhaseEvent(TEXT("EarlyZPass"));
 
 	// Draw the scene pre-pass / early z pass, populating the scene depth buffer and HiZ
+	RHICmdList.SetCurrentStat(GET_STATID(STAT_CLM_PrePass));
 	RenderPrePass(RHICmdList);
+	RHICmdList.SetCurrentStat(GET_STATID(STAT_CLM_AfterPrePass));
 	
 	GSceneRenderTargets.AllocGBufferTargets();
 	
@@ -762,7 +806,9 @@ void FDeferredShadingSceneRenderer::Render(FRHICommandListImmediate& RHICmdList)
 
 	GRenderTargetPool.AddPhaseEvent(TEXT("BasePass"));
 
+	RHICmdList.SetCurrentStat(GET_STATID(STAT_CLM_BasePass));
 	RenderBasePass(RHICmdList);
+	RHICmdList.SetCurrentStat(GET_STATID(STAT_CLM_AfterBasePass));
 
 	if(ViewFamily.EngineShowFlags.VisualizeLightCulling)
 	{
@@ -844,7 +890,9 @@ void FDeferredShadingSceneRenderer::Render(FRHICommandListImmediate& RHICmdList)
 		// Clear the translucent lighting volumes before we accumulate
 		ClearTranslucentVolumeLighting(RHICmdList);
 
+		RHICmdList.SetCurrentStat(GET_STATID(STAT_CLM_Lighting));
 		RenderLights(RHICmdList);
+		RHICmdList.SetCurrentStat(GET_STATID(STAT_CLM_AfterLighting));
 
 		GRenderTargetPool.AddPhaseEvent(TEXT("AfterRenderLights"));
 
@@ -931,7 +979,9 @@ void FDeferredShadingSceneRenderer::Render(FRHICommandListImmediate& RHICmdList)
 	{
 		SCOPE_CYCLE_COUNTER(STAT_TranslucencyDrawTime);
 
+		RHICmdList.SetCurrentStat(GET_STATID(STAT_CLM_Translucency));
 		RenderTranslucency(RHICmdList);
+		RHICmdList.SetCurrentStat(GET_STATID(STAT_CLM_AfterTranslucency));
 
 		if (ViewFamily.EngineShowFlags.Refraction)
 		{
@@ -963,7 +1013,9 @@ void FDeferredShadingSceneRenderer::Render(FRHICommandListImmediate& RHICmdList)
 	TRefCountPtr<IPooledRenderTarget> VelocityRT;
 
 	// Render the velocities of movable objects
+	RHICmdList.SetCurrentStat(GET_STATID(STAT_CLM_Velocity));
 	RenderVelocities(RHICmdList, VelocityRT);
+	RHICmdList.SetCurrentStat(GET_STATID(STAT_CLM_AfterVelocity));
 
 	// Finish rendering for each view.
 	if(ViewFamily.bResolveScene)
@@ -991,6 +1043,7 @@ void FDeferredShadingSceneRenderer::Render(FRHICommandListImmediate& RHICmdList)
 	VelocityRT.SafeRelease();
 
 	RenderFinish(RHICmdList);
+	RHICmdList.SetCurrentStat(GET_STATID(STAT_CLM_AfterFrame));
 }
 
 bool FDeferredShadingSceneRenderer::RenderPrePassViewDynamic(FRHICommandList& RHICmdList, const FViewInfo& View)
@@ -1081,6 +1134,8 @@ static void SetupPrePassView(FRHICommandList& RHICmdList, const FIntRect& ViewRe
 	// Note, this is a reversed Z depth surface, using CF_GreaterEqual.
 	RHICmdList.SetDepthStencilState(TStaticDepthStencilState<true,CF_GreaterEqual>::GetRHI());
 	RHICmdList.SetViewport(ViewRect.Min.X, ViewRect.Min.Y, 0, ViewRect.Max.X, ViewRect.Max.Y, 1);
+	RHICmdList.SetRasterizerState(TStaticRasterizerState<FM_Solid, CM_None>::GetRHI());
+	RHICmdList.SetScissorRect(false, 0, 0, 0, 0);
 }
 
 bool FDeferredShadingSceneRenderer::RenderPrePassView(FRHICommandList& RHICmdList, const FViewInfo& View)
@@ -1158,42 +1213,64 @@ public:
 	}
 };
 
+class FPrePassParallelCommandListSet : public FParallelCommandListSet
+{
+public:
+	FPrePassParallelCommandListSet(const FViewInfo& InView, FRHICommandList& InParentCmdList, bool* InOutDirty, bool bInParallelExecute)
+		: FParallelCommandListSet(InView, InParentCmdList, InOutDirty, bInParallelExecute)
+	{
+		SetStateOnCommandList(ParentCmdList);
+	}
+	virtual void SetStateOnCommandList(FRHICommandList& CmdList) override
+	{
+		GSceneRenderTargets.BeginRenderingPrePass(CmdList);
+		SetupPrePassView(CmdList, View.ViewRect);
+	}
+};
+
+static TAutoConsoleVariable<int32> CVarRHICmdPrePassDeferredContexts(
+	TEXT("r.RHICmdPrePassDeferredContexts"),
+	1,
+	TEXT("True to use deferred contexts to parallelize prepass command list execution.\n"));
 
-void FDeferredShadingSceneRenderer::RenderPrePassViewParallel(const FViewInfo& View, int32 Width, FRHICommandList& ParentCmdList, bool& OutDirty)
+void FDeferredShadingSceneRenderer::RenderPrePassViewParallel(const FViewInfo& View, FRHICommandList& ParentCmdList, bool& OutDirty)
 {
-	SetupPrePassView(ParentCmdList, View.ViewRect);
+	FPrePassParallelCommandListSet ParallelCommandListSet(View, ParentCmdList, &OutDirty, CVarRHICmdPrePassDeferredContexts.GetValueOnRenderThread() > 0);
+
 
 	// Draw the static occluder primitives using a depth drawing policy.
 	{
 		// Draw opaque occluders which support a separate position-only
 		// vertex buffer to minimize vertex fetch bandwidth, which is
 		// often the bottleneck during the depth only pass.
-		SCOPED_DRAW_EVENT(ParentCmdList, PosOnlyOpaque, DEC_SCENE_ITEMS);
-		Scene->PositionOnlyDepthDrawList.DrawVisibleParallel(View, View.StaticMeshOccluderMap, View.StaticMeshBatchVisibility, Width, ParentCmdList, OutDirty);
+		// we can't insert this event on the parent command list; need to pass it along to the ParallelCommandListSet
+		//SCOPED_DRAW_EVENT(ParentCmdList, PosOnlyOpaque, DEC_SCENE_ITEMS);
+		Scene->PositionOnlyDepthDrawList.DrawVisibleParallel(View.StaticMeshOccluderMap, View.StaticMeshBatchVisibility, ParallelCommandListSet);
 	}
 	{
 		// Draw opaque occluders, using double speed z where supported.
-		SCOPED_DRAW_EVENT(ParentCmdList, Opaque, DEC_SCENE_ITEMS);
-		Scene->DepthDrawList.DrawVisibleParallel(View,View.StaticMeshOccluderMap,View.StaticMeshBatchVisibility, Width, ParentCmdList, OutDirty);
+		// we can't insert this event on the parent command list; need to pass it along to the ParallelCommandListSet
+		//SCOPED_DRAW_EVENT(ParentCmdList, Opaque, DEC_SCENE_ITEMS);
+		Scene->DepthDrawList.DrawVisibleParallel(View.StaticMeshOccluderMap,View.StaticMeshBatchVisibility, ParallelCommandListSet);
 	}
 
 	if(EarlyZPassMode >= DDM_AllOccluders)
 	{
 		// Draw opaque occluders with masked materials
-		SCOPED_DRAW_EVENT(ParentCmdList, Opaque, DEC_SCENE_ITEMS);
-		Scene->MaskedDepthDrawList.DrawVisibleParallel(View,View.StaticMeshOccluderMap,View.StaticMeshBatchVisibility, Width, ParentCmdList, OutDirty);
+		// we can't insert this event on the parent command list; need to pass it along to the ParallelCommandListSet
+		//SCOPED_DRAW_EVENT(ParentCmdList, Opaque, DEC_SCENE_ITEMS);
+		Scene->MaskedDepthDrawList.DrawVisibleParallel(View.StaticMeshOccluderMap,View.StaticMeshBatchVisibility, ParallelCommandListSet);
 	}
 	{
-		FRHICommandList* CmdList = new FRHICommandList;
+		FRHICommandList* CmdList = ParallelCommandListSet.NewParallelCommandList();
 
 		FGraphEventRef AnyThreadCompletionEvent = TGraphTask<FRenderPrepassDynamicDataThreadTask>::CreateTask(nullptr, ENamedThreads::RenderThread)
 			.ConstructAndDispatchWhenReady(this, CmdList, &View, &OutDirty);
 
-		ParentCmdList.QueueAsyncCommandListSubmit(AnyThreadCompletionEvent, CmdList);
+		ParallelCommandListSet.AddParallelCommandList(CmdList, AnyThreadCompletionEvent);
 	}
 }
 
-
 /** Renders the scene's prepass and occlusion queries */
 bool FDeferredShadingSceneRenderer::RenderPrePass(FRHICommandListImmediate& RHICmdList)
 {
@@ -1204,7 +1281,6 @@ bool FDeferredShadingSceneRenderer::RenderPrePass(FRHICommandListImmediate& RHIC
 
 	GSceneRenderTargets.BeginRenderingPrePass(RHICmdList);
 
-	
 	// Clear the depth buffer.
 	// Note, this is a reversed Z depth surface, so 0.0f is the far plane.
 	RHICmdList.Clear(false,FLinearColor::Black,true,0.0f,true,0, FIntRect());
@@ -1215,13 +1291,12 @@ bool FDeferredShadingSceneRenderer::RenderPrePass(FRHICommandListImmediate& RHIC
 		if (GRHICommandList.UseParallelAlgorithms())
 		{
 			FScopedCommandListWaitForTasks Flusher(RHICmdList);
-			int32 Width = CVarRHICmdWidth.GetValueOnRenderThread(); // we use a few more than needed to cover non-equal jobs
 
 			for(int32 ViewIndex = 0;ViewIndex < Views.Num();ViewIndex++)
 			{
 				SCOPED_CONDITIONAL_DRAW_EVENTF(RHICmdList, EventView, Views.Num() > 1, DEC_SCENE_ITEMS, TEXT("View%d"), ViewIndex);
 				const FViewInfo& View = Views[ViewIndex];
-				RenderPrePassViewParallel(View, Width, RHICmdList, bDirty);
+				RenderPrePassViewParallel(View, RHICmdList, bDirty);
 			}
 		}
 		else
@@ -1261,13 +1336,11 @@ bool FDeferredShadingSceneRenderer::RenderBasePass(FRHICommandListImmediate& RHI
 		if (GRHICommandList.UseParallelAlgorithms())
 		{
 			FScopedCommandListWaitForTasks Flusher(RHICmdList);
-			int32 Width = CVarRHICmdWidth.GetValueOnRenderThread(); // we use a few more than needed to cover non-equal jobs
-
 			for(int32 ViewIndex = 0;ViewIndex < Views.Num();ViewIndex++)
 			{
 				SCOPED_CONDITIONAL_DRAW_EVENTF(RHICmdList, EventView, Views.Num() > 1, DEC_SCENE_ITEMS, TEXT("View%d"), ViewIndex);
 				FViewInfo& View = Views[ViewIndex];
-				RenderBasePassViewParallel(View, Width, RHICmdList, bDirty);
+				RenderBasePassViewParallel(View, RHICmdList, bDirty);
 			}
 		}
 		else
diff --git a/Engine/Source/Runtime/Renderer/Private/DeferredShadingRenderer.h b/Engine/Source/Runtime/Renderer/Private/DeferredShadingRenderer.h
--- a/Engine/Source/Runtime/Renderer/Private/DeferredShadingRenderer.h
+++ b/Engine/Source/Runtime/Renderer/Private/DeferredShadingRenderer.h
@@ -65,7 +65,7 @@ public:
 	/**
 	 * Renders the scene's prepass for a particular view in parallel
 	 */
-	void RenderPrePassViewParallel(const FViewInfo& View, int32 Width, FRHICommandList& ParentCmdList, bool& OutDirty);
+	void RenderPrePassViewParallel(const FViewInfo& View, FRHICommandList& ParentCmdList, bool& OutDirty);
 
     /** Renders the basepass for the static data of a given View. */
     bool RenderBasePassStaticData(FRHICommandList& RHICmdList, FViewInfo& View);
@@ -73,9 +73,9 @@ public:
 	bool RenderBasePassStaticDataDefault(FRHICommandList& RHICmdList, FViewInfo& View);
 
 	/** Renders the basepass for the static data of a given View. Parallel versions.*/
-	void RenderBasePassStaticDataParallel(FViewInfo& View, int32 Width, FRHICommandList& ParentCmdList, bool& OutDirty);
-	void RenderBasePassStaticDataMaskedParallel(FViewInfo& View, int32 Width, FRHICommandList& ParentCmdList, bool& OutDirty);
-	void RenderBasePassStaticDataDefaultParallel(FViewInfo& View, int32 Width, FRHICommandList& ParentCmdList, bool& OutDirty);
+	void RenderBasePassStaticDataParallel(FParallelCommandListSet& ParallelCommandListSet);
+	void RenderBasePassStaticDataMaskedParallel(FParallelCommandListSet& ParallelCommandListSet);
+	void RenderBasePassStaticDataDefaultParallel(FParallelCommandListSet& ParallelCommandListSet);
 
 	/** Sorts base pass draw lists front to back for improved GPU culling. */
 	void SortBasePassStaticData(FVector ViewPosition);
@@ -84,10 +84,10 @@ public:
 	void RenderBasePassDynamicData(FRHICommandList& RHICmdList, const FViewInfo& View, bool& bOutDirty);
 
 	/** Renders the basepass for the dynamic data of a given View, in parallel. */
-	void RenderBasePassDynamicDataParallel(FViewInfo& View, FRHICommandList& ParentCmdList, bool& bOutDirty);
+	void RenderBasePassDynamicDataParallel(FParallelCommandListSet& ParallelCommandListSet);
 
     /** Renders the basepass for a given View, in parallel */
-	void RenderBasePassViewParallel(FViewInfo& View, int32 Width, FRHICommandList& ParentCmdList, bool& OutDirty);
+	void RenderBasePassViewParallel(FViewInfo& View, FRHICommandList& ParentCmdList, bool& OutDirty);
 
 	/** Renders the basepass for a given View. */
 	bool RenderBasePassView(FRHICommandListImmediate& RHICmdList, FViewInfo& View);
@@ -241,10 +241,11 @@ private:
 	/** Renders the velocities for a subset of movable objects for the motion blur effect. */
 	friend class FRenderVelocityDynamicThreadTask;
 	void RenderDynamicVelocitiesInner(FRHICommandList& RHICmdList, const FViewInfo& View, int32 FirstIndex, int32 LastIndex);
-	void RenderDynamicVelocitiesMeshElementsInner(FRHICommandList& RHICmdList, const FViewInfo& View);
+	void RenderDynamicVelocitiesMeshElementsInner(FRHICommandList& RHICmdList, const FViewInfo& View, int32 FirstIndex, int32 LastIndex);
 
 	/** Renders the velocities of movable objects for the motion blur effect. */
-	void RenderVelocitiesInner(FRHICommandListImmediate& RHICmdList, const FViewInfo& View);
+	void RenderVelocitiesInner(FRHICommandListImmediate& RHICmdList, TRefCountPtr<IPooledRenderTarget>& VelocityRT);
+	void RenderVelocitiesInnerParallel(FRHICommandListImmediate& RHICmdList, TRefCountPtr<IPooledRenderTarget>& VelocityRT);
 
 	/** Renders world-space lightmap density instead of the normal color. */
 	bool RenderLightMapDensities(FRHICommandListImmediate& RHICmdList);
diff --git a/Engine/Source/Runtime/Renderer/Private/DynamicPrimitiveDrawing.h b/Engine/Source/Runtime/Renderer/Private/DynamicPrimitiveDrawing.h
--- a/Engine/Source/Runtime/Renderer/Private/DynamicPrimitiveDrawing.h
+++ b/Engine/Source/Runtime/Renderer/Private/DynamicPrimitiveDrawing.h
@@ -181,12 +181,10 @@ bool DrawViewElements(
 */
 template<class DrawingPolicyFactoryType>
 void DrawViewElementsParallel(
-	const FViewInfo& View,
 	const typename DrawingPolicyFactoryType::ContextType& DrawingContext,
 	uint8 DPGIndex,
 	bool bPreFog,
-	FRHICommandList& ParentCmdList,
-	int32 Width
+	FParallelCommandListSet& ParallelCommandListSet
 	);
 
 /**
diff --git a/Engine/Source/Runtime/Renderer/Private/DynamicPrimitiveDrawing.inl b/Engine/Source/Runtime/Renderer/Private/DynamicPrimitiveDrawing.inl
--- a/Engine/Source/Runtime/Renderer/Private/DynamicPrimitiveDrawing.inl
+++ b/Engine/Source/Runtime/Renderer/Private/DynamicPrimitiveDrawing.inl
@@ -309,20 +309,18 @@ public:
 
 template<class DrawingPolicyFactoryType>
 void DrawViewElementsParallel(
-	const FViewInfo& View,
 	const typename DrawingPolicyFactoryType::ContextType& DrawingContext,
 	uint8 DPGIndex,
 	bool bPreFog,
-	FRHICommandList& ParentCmdList,
-	int32 Width
+	FParallelCommandListSet& ParallelCommandListSet
 	)
 {
 	// Get the correct element list based on dpg index
-	const TIndirectArray<FMeshBatch>& ViewMeshElementList = (DPGIndex == SDPG_Foreground ? View.TopViewMeshElements : View.ViewMeshElements);
+	const TIndirectArray<FMeshBatch>& ViewMeshElementList = (DPGIndex == SDPG_Foreground ? ParallelCommandListSet.View.TopViewMeshElements : ParallelCommandListSet.View.ViewMeshElements);
 
 	{
 		int32 NumPrims = ViewMeshElementList.Num();
-		int32 EffectiveThreads = FMath::Min<int32>(NumPrims, Width);
+		int32 EffectiveThreads = FMath::Min<int32>(NumPrims, ParallelCommandListSet.Width);
 
 		int32 Start = 0;
 		if (EffectiveThreads)
@@ -338,12 +336,12 @@ void DrawViewElementsParallel(
 				check(Last >= Start);
 
 				{
-					FRHICommandList* CmdList = new FRHICommandList;
+					FRHICommandList* CmdList = ParallelCommandListSet.NewParallelCommandList();
 
 					FGraphEventRef AnyThreadCompletionEvent = TGraphTask<FDrawViewElementsAnyThreadTask<DrawingPolicyFactoryType> >::CreateTask(nullptr, ENamedThreads::RenderThread)
-						.ConstructAndDispatchWhenReady(CmdList, &View, DrawingContext, DPGIndex, bPreFog, Start, Last);
+						.ConstructAndDispatchWhenReady(CmdList, &ParallelCommandListSet.View, DrawingContext, DPGIndex, bPreFog, Start, Last);
 
-					ParentCmdList.QueueAsyncCommandListSubmit(AnyThreadCompletionEvent, CmdList);
+					ParallelCommandListSet.AddParallelCommandList(CmdList, AnyThreadCompletionEvent);
 				}
 				Start = Last + 1;
 			}
diff --git a/Engine/Source/Runtime/Renderer/Private/PostProcess/SceneRenderTargets.cpp b/Engine/Source/Runtime/Renderer/Private/PostProcess/SceneRenderTargets.cpp
--- a/Engine/Source/Runtime/Renderer/Private/PostProcess/SceneRenderTargets.cpp
+++ b/Engine/Source/Runtime/Renderer/Private/PostProcess/SceneRenderTargets.cpp
@@ -762,7 +762,7 @@ void FSceneRenderTargets::ResolveGBufferSurfaces(FRHICommandList& RHICmdList, co
 	}
 }
 
-void FSceneRenderTargets::BeginRenderingPrePass(FRHICommandListImmediate& RHICmdList)
+void FSceneRenderTargets::BeginRenderingPrePass(FRHICommandList& RHICmdList)
 {
 	SCOPED_DRAW_EVENT(RHICmdList, BeginRenderingPrePass, DEC_SCENE_ITEMS);
 
@@ -778,11 +778,14 @@ void FSceneRenderTargets::FinishRenderingPrePass(FRHICommandListImmediate& RHICm
 	GRenderTargetPool.VisualizeTexture.SetCheckPoint(RHICmdList, SceneDepthZ);
 }
 
-void FSceneRenderTargets::BeginRenderingShadowDepth(FRHICommandListImmediate& RHICmdList)
+void FSceneRenderTargets::BeginRenderingShadowDepth(FRHICommandList& RHICmdList, bool bClear)
 {
 	GRenderTargetPool.VisualizeTexture.SetCheckPoint(RHICmdList, ShadowDepthZ);
 	FRHISetRenderTargetsInfo Info(0, nullptr, FRHIDepthRenderTargetView(GetShadowDepthZSurface(), ERenderTargetLoadAction::EClear, ERenderTargetStoreAction::EStore));
-	Info.SetClearDepthStencil(true, 1.0f);
+	if (bClear)
+	{
+		Info.SetClearDepthStencil(true, 1.0f);
+	}
 	Info.ColorRenderTarget[0].StoreAction = ERenderTargetStoreAction::ENoAction;
 
 	if (!GSupportsDepthRenderTargetWithoutColorRenderTarget)
@@ -794,7 +797,7 @@ void FSceneRenderTargets::BeginRenderingShadowDepth(FRHICommandListImmediate& RH
 	RHICmdList.SetRenderTargetsAndClear(Info);
 }
 
-void FSceneRenderTargets::BeginRenderingCubeShadowDepth(FRHICommandListImmediate& RHICmdList, int32 ShadowResolution)
+void FSceneRenderTargets::BeginRenderingCubeShadowDepth(FRHICommandList& RHICmdList, int32 ShadowResolution)
 {
 	SCOPED_DRAW_EVENT(RHICmdList, BeginRenderingCubeShadowDepth, DEC_SCENE_ITEMS);
 	SetRenderTarget(RHICmdList, FTextureRHIRef(), GetCubeShadowDepthZSurface(ShadowResolution));
@@ -806,7 +809,7 @@ void FSceneRenderTargets::FinishRenderingShadowDepth(FRHICommandList& RHICmdList
 	RHICmdList.CopyToResolveTarget(GetShadowDepthZSurface(), GetShadowDepthZTexture(), false, FResolveParams(ResolveRect));
 }
 
-void FSceneRenderTargets::BeginRenderingReflectiveShadowMap(FRHICommandListImmediate& RHICmdList, FLightPropagationVolume* Lpv)
+void FSceneRenderTargets::BeginRenderingReflectiveShadowMap(FRHICommandList& RHICmdList, FLightPropagationVolume* Lpv)
 {
 	FTextureRHIParamRef RenderTargets[2];
 	RenderTargets[0] = GetReflectiveShadowMapNormalSurface();
diff --git a/Engine/Source/Runtime/Renderer/Private/PostProcess/SceneRenderTargets.h b/Engine/Source/Runtime/Renderer/Private/PostProcess/SceneRenderTargets.h
--- a/Engine/Source/Runtime/Renderer/Private/PostProcess/SceneRenderTargets.h
+++ b/Engine/Source/Runtime/Renderer/Private/PostProcess/SceneRenderTargets.h
@@ -183,10 +183,10 @@ public:
 	/** Resolves the GBuffer targets so that their resolved textures can be sampled. */
 	void ResolveGBufferSurfaces(FRHICommandList& RHICmdList, const FResolveRect& ResolveRect = FResolveRect());
 
-	void BeginRenderingShadowDepth(FRHICommandListImmediate& RHICmdList);
+	void BeginRenderingShadowDepth(FRHICommandList& RHICmdList, bool bClear);
 
 	/** Binds the appropriate shadow depth cube map for rendering. */
-	void BeginRenderingCubeShadowDepth(FRHICommandListImmediate& RHICmdList, int32 ShadowResolution);
+	void BeginRenderingCubeShadowDepth(FRHICommandList& RHICmdList, int32 ShadowResolution);
 
 	/**
 	 * Called when finished rendering to the subject shadow depths so the surface can be copied to texture
@@ -194,7 +194,7 @@ public:
 	 */
 	void FinishRenderingShadowDepth(FRHICommandList& RHICmdList, const FResolveRect& ResolveRect = FResolveRect());
 
-	void BeginRenderingReflectiveShadowMap(FRHICommandListImmediate& RHICmdList, class FLightPropagationVolume* Lpv);
+	void BeginRenderingReflectiveShadowMap(FRHICommandList& RHICmdList, class FLightPropagationVolume* Lpv);
 	void FinishRenderingReflectiveShadowMap(FRHICommandList& RHICmdList, const FResolveRect& ResolveRect = FResolveRect());
 
 	/** Resolves the appropriate shadow depth cube map and restores default state. */
@@ -209,7 +209,7 @@ public:
 	void ResolveSceneDepthTexture(FRHICommandList& RHICmdList);
 	void ResolveSceneDepthToAuxiliaryTexture(FRHICommandList& RHICmdList);
 
-	void BeginRenderingPrePass(FRHICommandListImmediate& RHICmdList);
+	void BeginRenderingPrePass(FRHICommandList& RHICmdList);
 	void FinishRenderingPrePass(FRHICommandListImmediate& RHICmdList);
 
 	void BeginRenderingSceneAlphaCopy(FRHICommandListImmediate& RHICmdList);
diff --git a/Engine/Source/Runtime/Renderer/Private/SceneRendering.cpp b/Engine/Source/Runtime/Renderer/Private/SceneRendering.cpp
--- a/Engine/Source/Runtime/Renderer/Private/SceneRendering.cpp
+++ b/Engine/Source/Runtime/Renderer/Private/SceneRendering.cpp
@@ -69,6 +69,69 @@ static TAutoConsoleVariable<float> CVarTessellationAdaptivePixelsPerTriangle(
 	TEXT("Global tessellation factor multiplier"),
 	ECVF_RenderThreadSafe);
 
+/*-----------------------------------------------------------------------------
+	FParallelCommandListSet
+-----------------------------------------------------------------------------*/
+
+
+FRHICommandList* FParallelCommandListSet::AllocCommandList()
+{
+	return new FRHICommandList;
+}
+
+FParallelCommandListSet::FParallelCommandListSet(const FViewInfo& InView, FRHICommandList& InParentCmdList, bool* InOutDirty, bool bInParallelExecute)
+	: View(InView)
+	, ParentCmdList(InParentCmdList)
+	, OutDirtyIfIgnored(false)
+	, OutDirty(InOutDirty ? *InOutDirty : OutDirtyIfIgnored)
+	, bParallelExecute(bInParallelExecute)
+{
+	Width = CVarRHICmdWidth.GetValueOnRenderThread();
+	CommandLists.Reserve(Width * 8);
+	Events.Reserve(Width * 8);
+
+}
+
+FParallelCommandListSet::~FParallelCommandListSet()
+{
+	check(CommandLists.Num() == Events.Num());
+#if PLATFORM_SUPPORTS_PARALLEL_RHI_EXECUTE
+	if (bParallelExecute && CommandLists.Num())
+	{
+		ParentCmdList.QueueParallelAsyncCommandListSubmit(&Events[0], &CommandLists[0], CommandLists.Num());
+		SetStateOnCommandList(ParentCmdList);
+	}
+	else
+#endif
+	{
+		for (int32 Index = 0; Index < CommandLists.Num(); Index++)
+		{
+			ParentCmdList.QueueAsyncCommandListSubmit(Events[Index], CommandLists[Index]);
+		}
+	}
+	CommandLists.Reset();
+	Events.Reset();
+}
+
+FRHICommandList* FParallelCommandListSet::NewParallelCommandList()
+{
+	FRHICommandList* Result = AllocCommandList();
+//	if (bParallelExecute)
+	{
+		SetStateOnCommandList(*Result); 
+	}
+	return Result;
+}
+
+void FParallelCommandListSet::AddParallelCommandList(FRHICommandList* CmdList, FGraphEventRef& CompletionEvent)
+{
+	check(CommandLists.Num() == Events.Num());
+	CommandLists.Add(CmdList);
+	Events.Add(CompletionEvent);
+}
+
+
+
 /*-----------------------------------------------------------------------------
 	FViewInfo
 -----------------------------------------------------------------------------*/
diff --git a/Engine/Source/Runtime/Renderer/Private/SceneRendering.h b/Engine/Source/Runtime/Renderer/Private/SceneRendering.h
--- a/Engine/Source/Runtime/Renderer/Private/SceneRendering.h
+++ b/Engine/Source/Runtime/Renderer/Private/SceneRendering.h
@@ -306,6 +306,33 @@ private:
 	uint32 ValidFrameNumber;
 };
 
+class FParallelCommandListSet
+{
+public:
+	const FViewInfo& View;
+	int32 Width;
+	FRHICommandList& ParentCmdList;
+private:
+	bool OutDirtyIfIgnored;
+public:
+	bool& OutDirty;
+	TArray<FRHICommandList*,SceneRenderingAllocator> CommandLists;
+	TArray<FGraphEventRef,SceneRenderingAllocator> Events;
+protected:
+	FRHICommandList* AllocCommandList();
+	bool bParallelExecute;
+public:
+	FParallelCommandListSet(const FViewInfo& InView, FRHICommandList& InParentCmdList, bool* InOutDirty, bool bInParallelExecute);
+	~FParallelCommandListSet();
+	FRHICommandList* NewParallelCommandList();
+	void AddParallelCommandList(FRHICommandList* CmdList, FGraphEventRef& CompletionEvent);
+
+	virtual void SetStateOnCommandList(FRHICommandList& CmdList)
+	{
+
+	}
+};
+
 /** A FSceneView with additional state used by the scene renderer. */
 class FViewInfo : public FSceneView
 {
diff --git a/Engine/Source/Runtime/Renderer/Private/ShadowRendering.cpp b/Engine/Source/Runtime/Renderer/Private/ShadowRendering.cpp
--- a/Engine/Source/Runtime/Renderer/Private/ShadowRendering.cpp
+++ b/Engine/Source/Runtime/Renderer/Private/ShadowRendering.cpp
@@ -1444,6 +1444,64 @@ public:
 	}
 };
 
+void FProjectedShadowInfo::SetStateForDepth(FRHICommandList& RHICmdList)
+{
+	check(bAllocated);
+
+	if (bOnePassPointLightShadow)
+	{
+		// Set the viewport to the whole render target since it's a cube map, don't leave any border space
+		RHICmdList.SetViewport(
+			0,
+			0,
+			0.0f,
+			ResolutionX,
+			ResolutionY,
+			1.0f
+			);
+	}
+	else
+	{
+		// Do not use a shadow border for RSMs.
+		if(bReflectiveShadowmap)
+		{
+			RHICmdList.SetViewport(
+				X,
+				Y,
+				0.0f,
+				X + ResolutionX,
+				Y + ResolutionY,
+				1.0f
+				);
+		}
+		else
+		{
+			RHICmdList.SetViewport(
+				X + SHADOW_BORDER,
+				Y + SHADOW_BORDER,
+				0.0f,
+				X + SHADOW_BORDER + ResolutionX,
+				Y + SHADOW_BORDER + ResolutionY,
+				1.0f
+				);
+		}
+	}
+
+	if (bReflectiveShadowmap && !bOnePassPointLightShadow)
+	{
+		// Enable color writes to the reflective shadow map targets with opaque blending
+		RHICmdList.SetBlendState(TStaticBlendStateWriteMask<CW_RGBA, CW_RGBA>::GetRHI());
+	}
+	else
+	{
+		// Disable color writes
+		RHICmdList.SetBlendState(TStaticBlendState<CW_NONE>::GetRHI());
+	}
+
+	RHICmdList.SetDepthStencilState(TStaticDepthStencilState<true, CF_LessEqual>::GetRHI());
+}
+
+
 static TAutoConsoleVariable<int32> CVarParallelShadows(
 	TEXT("r.ParallelShadows"),
 	1,
@@ -1451,7 +1509,31 @@ static TAutoConsoleVariable<int32> CVarParallelShadows(
 	ECVF_RenderThreadSafe
 	);
 
-void FProjectedShadowInfo::RenderDepthInner(FRHICommandList& RHICmdList, FSceneRenderer* SceneRenderer, const FViewInfo* FoundView)
+static TAutoConsoleVariable<int32> CVarRHICmdShadowDeferredContexts(
+	TEXT("r.RHICmdShadowDeferredContexts"),
+	1,
+	TEXT("True to use deferred contexts to parallelize shadow command list execution.\n"));
+
+class FShadowParallelCommandListSet : public FParallelCommandListSet
+{
+	FProjectedShadowInfo& ProjectedShadowInfo;
+	ISetShadowRenderTargets& SetShadowRenderTargets;
+public:
+	FShadowParallelCommandListSet(const FViewInfo& InView, FRHICommandList& InParentCmdList, bool* InOutDirty, bool bInParallelExecute, FProjectedShadowInfo& InProjectedShadowInfo, ISetShadowRenderTargets& InSetShadowRenderTargets)
+		: FParallelCommandListSet(InView, InParentCmdList, InOutDirty, bInParallelExecute)
+		, ProjectedShadowInfo(InProjectedShadowInfo)
+		, SetShadowRenderTargets(InSetShadowRenderTargets)
+	{
+		SetStateOnCommandList(ParentCmdList);
+	}
+	virtual void SetStateOnCommandList(FRHICommandList& CmdList) override
+	{
+		SetShadowRenderTargets.SetRenderTargetsForShadows(CmdList);
+		ProjectedShadowInfo.SetStateForDepth(CmdList);
+	}
+};
+
+void FProjectedShadowInfo::RenderDepthInner(FRHICommandList& RHICmdList, FSceneRenderer* SceneRenderer, const FViewInfo* FoundView, ISetShadowRenderTargets& SetShadowRenderTargets)
 {
 	FShadowDepthDrawingPolicyContext PolicyContext(this);
 
@@ -1461,46 +1543,48 @@ void FProjectedShadowInfo::RenderDepthInner(FRHICommandList& RHICmdList, FSceneR
 		// parallel version
 		FScopedCommandListWaitForTasks Flusher;
 
-		int32 Width = CVarRHICmdWidth.GetValueOnRenderThread(); // we use a few more than needed to cover non-equal jobs
-		bool OutDirty = false; // unused
-
-		// Draw the subject's static elements using static draw lists
-		if (IsWholeSceneDirectionalShadow())
 		{
-			SCOPE_CYCLE_COUNTER(STAT_WholeSceneStaticDrawListShadowDepthsTime);
+			FShadowParallelCommandListSet ParallelCommandListSet(*FoundView, RHICmdList, nullptr, CVarRHICmdShadowDeferredContexts.GetValueOnRenderThread() > 0, *this, SetShadowRenderTargets);
 
-			if (bReflectiveShadowmap)
+			// Draw the subject's static elements using static draw lists
+			if (IsWholeSceneDirectionalShadow())
 			{
-				SceneRenderer->Scene->WholeSceneReflectiveShadowMapDrawList.DrawVisibleParallel(*FoundView, PolicyContext, StaticMeshWholeSceneShadowDepthMap, StaticMeshWholeSceneShadowBatchVisibility, Width, RHICmdList, OutDirty);
+				SCOPE_CYCLE_COUNTER(STAT_WholeSceneStaticDrawListShadowDepthsTime);
+
+				if (bReflectiveShadowmap)
+				{
+					SceneRenderer->Scene->WholeSceneReflectiveShadowMapDrawList.DrawVisibleParallel(PolicyContext, StaticMeshWholeSceneShadowDepthMap, StaticMeshWholeSceneShadowBatchVisibility, ParallelCommandListSet);
+				}
+				else
+				{
+					// Use the scene's shadow depth draw list with this shadow's visibility map
+					SceneRenderer->Scene->WholeSceneShadowDepthDrawList.DrawVisibleParallel(PolicyContext, StaticMeshWholeSceneShadowDepthMap, StaticMeshWholeSceneShadowBatchVisibility, ParallelCommandListSet);
+				}
 			}
-			else
+			// Draw the subject's static elements using manual state filtering
+			else if (SubjectMeshElements.Num() > 0)
 			{
-				// Use the scene's shadow depth draw list with this shadow's visibility map
-				SceneRenderer->Scene->WholeSceneShadowDepthDrawList.DrawVisibleParallel(*FoundView, PolicyContext, StaticMeshWholeSceneShadowDepthMap, StaticMeshWholeSceneShadowBatchVisibility, Width, RHICmdList, OutDirty);
-			}
-		}
-		// Draw the subject's static elements using manual state filtering
-		else if (SubjectMeshElements.Num() > 0)
-		{
-			FRHICommandList* CmdList = new FRHICommandList;
+				FRHICommandList* CmdList = ParallelCommandListSet.NewParallelCommandList();
 
-			FGraphEventRef AnyThreadCompletionEvent = TGraphTask<FDrawShadowMeshElementsThreadTask>::CreateTask(nullptr, ENamedThreads::RenderThread)
-				.ConstructAndDispatchWhenReady(this, CmdList, FoundView, bReflectiveShadowmap && !bOnePassPointLightShadow);
+				FGraphEventRef AnyThreadCompletionEvent = TGraphTask<FDrawShadowMeshElementsThreadTask>::CreateTask(nullptr, ENamedThreads::RenderThread)
+					.ConstructAndDispatchWhenReady(this, CmdList, FoundView, bReflectiveShadowmap && !bOnePassPointLightShadow);
 
-			RHICmdList.QueueAsyncCommandListSubmit(AnyThreadCompletionEvent, CmdList);
-		}
-		{
-			FRHICommandList* CmdList = new FRHICommandList;
+				ParallelCommandListSet.AddParallelCommandList(CmdList, AnyThreadCompletionEvent);
+			}
+			{
+				FRHICommandList* CmdList = ParallelCommandListSet.NewParallelCommandList();
 
-			FGraphEventRef AnyThreadCompletionEvent = TGraphTask<FRenderDepthDynamicThreadTask>::CreateTask(nullptr, ENamedThreads::RenderThread)
-				.ConstructAndDispatchWhenReady(this, CmdList, FoundView, SceneRenderer);
+				FGraphEventRef AnyThreadCompletionEvent = TGraphTask<FRenderDepthDynamicThreadTask>::CreateTask(nullptr, ENamedThreads::RenderThread)
+					.ConstructAndDispatchWhenReady(this, CmdList, FoundView, SceneRenderer);
 
-			RHICmdList.QueueAsyncCommandListSubmit(AnyThreadCompletionEvent, CmdList);
+				ParallelCommandListSet.AddParallelCommandList(CmdList, AnyThreadCompletionEvent);
+			}
 		}
 	}
 	else
 	{
 		// single threaded version
+		SetStateForDepth(RHICmdList);
 
 		// Draw the subject's static elements using static draw lists
 		if (IsWholeSceneDirectionalShadow())
@@ -1537,7 +1621,7 @@ void FProjectedShadowInfo::RenderDepthInner(FRHICommandList& RHICmdList, FSceneR
 	}
 }
 
-void FProjectedShadowInfo::RenderDepth(FRHICommandList& RHICmdList, FSceneRenderer* SceneRenderer)
+void FProjectedShadowInfo::RenderDepth(FRHICommandList& RHICmdList, FSceneRenderer* SceneRenderer, ISetShadowRenderTargets& SetShadowRenderTargets)
 {
 #if WANTS_DRAW_MESH_EVENTS
 	FString EventName;
@@ -1548,60 +1632,6 @@ void FProjectedShadowInfo::RenderDepth(FRHICommandList& RHICmdList, FSceneRender
 	CONDITIONAL_SCOPE_CYCLE_COUNTER(STAT_RenderWholeSceneShadowDepthsTime, bWholeSceneShadow);
 	CONDITIONAL_SCOPE_CYCLE_COUNTER(STAT_RenderPerObjectShadowDepthsTime, !bWholeSceneShadow);
 
-	check(bAllocated);
-
-	if (bOnePassPointLightShadow)
-	{
-		// Set the viewport to the whole render target since it's a cube map, don't leave any border space
-		RHICmdList.SetViewport(
-			0,
-			0,
-			0.0f,
-			ResolutionX,
-			ResolutionY,
-			1.0f
-			);
-	}
-	else
-	{
-		// Do not use a shadow border for RSMs.
-		if(bReflectiveShadowmap)
-		{
-			RHICmdList.SetViewport(
-				X,
-				Y,
-				0.0f,
-				X + ResolutionX,
-				Y + ResolutionY,
-				1.0f
-				);
-		}
-		else
-		{
-			RHICmdList.SetViewport(
-				X + SHADOW_BORDER,
-				Y + SHADOW_BORDER,
-				0.0f,
-				X + SHADOW_BORDER + ResolutionX,
-				Y + SHADOW_BORDER + ResolutionY,
-				1.0f
-				);
-		}
-	}
-
-	if (bReflectiveShadowmap && !bOnePassPointLightShadow)
-	{
-		// Enable color writes to the reflective shadow map targets with opaque blending
-		RHICmdList.SetBlendState(TStaticBlendStateWriteMask<CW_RGBA, CW_RGBA>::GetRHI());
-	}
-	else
-	{
-		// Disable color writes
-		RHICmdList.SetBlendState(TStaticBlendState<CW_NONE>::GetRHI());
-	}
-	
-	RHICmdList.SetDepthStencilState(TStaticDepthStencilState<true, CF_LessEqual>::GetRHI());
-
 	// Choose an arbitrary view where this shadow's subject is relevant.
 	FViewInfo* FoundView = NULL;
 	for(int32 ViewIndex = 0;ViewIndex < SceneRenderer->Views.Num();ViewIndex++)
@@ -1636,7 +1666,7 @@ void FProjectedShadowInfo::RenderDepth(FRHICommandList& RHICmdList, FSceneRender
 	// Lighting only should only affect the material used with direct lighting, not the indirect lighting
 	FoundView->bForceShowMaterials = true;
 
-	RenderDepthInner(RHICmdList, SceneRenderer, FoundView);
+	RenderDepthInner(RHICmdList, SceneRenderer, FoundView, SetShadowRenderTargets);
 
 	FoundView->bForceShowMaterials = false;
 	FoundView->UniformBuffer = OriginalUniformBuffer;
@@ -2589,9 +2619,17 @@ bool FDeferredShadingSceneRenderer::RenderOnePassPointLightShadows(FRHICommandLi
 				SCOPED_DRAW_EVENT(RHICmdList, ShadowDepthsFromOpaque, DEC_SCENE_ITEMS);
 				GSceneRenderTargets.BeginRenderingCubeShadowDepth(RHICmdList, ProjectedShadowInfo->ResolutionX);
 				ProjectedShadowInfo->bAllocated = true;
-				//gilmerge, if clear depth doesn't go into a cmdlist, then how can renderdepth? This occurs in 4 spots
 				ProjectedShadowInfo->ClearDepth(RHICmdList, this);
-				ProjectedShadowInfo->RenderDepth(RHICmdList, this);
+				struct FLocalShadowTargets : public ISetShadowRenderTargets
+				{
+					uint32 ResolutionX;
+					virtual void SetRenderTargetsForShadows(FRHICommandList& RHICmdList)
+					{
+						GSceneRenderTargets.BeginRenderingCubeShadowDepth(RHICmdList, ResolutionX);
+					}
+				} LocalShadowTargets;
+				LocalShadowTargets.ResolutionX = ProjectedShadowInfo->ResolutionX;
+				ProjectedShadowInfo->RenderDepth(RHICmdList, this, LocalShadowTargets);
 				GSceneRenderTargets.FinishRenderingCubeShadowDepth(RHICmdList, ProjectedShadowInfo->ResolutionX);
 			}
 
@@ -2999,7 +3037,16 @@ bool FDeferredShadingSceneRenderer::RenderReflectiveShadowMaps(FRHICommandListIm
 				if (ProjectedShadowInfo->bAllocated && !ProjectedShadowInfo->bTranslucentShadow)
 				{
 					ProjectedShadowInfo->ClearDepth(RHICmdList, this);
-					ProjectedShadowInfo->RenderDepth(RHICmdList, this);
+					struct FLocalShadowTargets : public ISetShadowRenderTargets
+					{
+						FLightPropagationVolume* LightPropagationVolume;
+						virtual void SetRenderTargetsForShadows(FRHICommandList& RHICmdList)
+						{
+							GSceneRenderTargets.BeginRenderingReflectiveShadowMap(RHICmdList, LightPropagationVolume);
+						}
+					} LocalShadowTargets;
+					LocalShadowTargets.LightPropagationVolume = LightPropagationVolume;
+					ProjectedShadowInfo->RenderDepth(RHICmdList, this, LocalShadowTargets);
 				}
 
 				GSceneRenderTargets.FinishRenderingReflectiveShadowMap(RHICmdList);
@@ -3163,9 +3210,14 @@ bool FDeferredShadingSceneRenderer::RenderProjectedShadows(FRHICommandListImmedi
 			// Render the shadow depths.
 			SCOPED_DRAW_EVENT(RHICmdList, ShadowDepthsFromOpaque, DEC_SCENE_ITEMS);
 
-			GSceneRenderTargets.BeginRenderingShadowDepth(RHICmdList);
-
-			RHICmdList.Clear(false, FColor(255, 255, 255), true, 1.0f, false, 0, FIntRect());
+			GSceneRenderTargets.BeginRenderingShadowDepth(RHICmdList, true);
+			struct FLocalShadowTargets : public ISetShadowRenderTargets
+			{
+				virtual void SetRenderTargetsForShadows(FRHICommandList& RHICmdList)
+				{
+					GSceneRenderTargets.BeginRenderingShadowDepth(RHICmdList, false);
+				}
+			} LocalShadowTargets;
 
 			// render depth for each shadow
 			for (int32 ShadowIndex = 0; ShadowIndex < Shadows.Num(); ShadowIndex++)
@@ -3173,7 +3225,7 @@ bool FDeferredShadingSceneRenderer::RenderProjectedShadows(FRHICommandListImmedi
 				FProjectedShadowInfo* ProjectedShadowInfo = Shadows[ShadowIndex];
 				if (ProjectedShadowInfo->bAllocated && !ProjectedShadowInfo->bTranslucentShadow && !ProjectedShadowInfo->bRayTracedDistanceFieldShadow)
 				{
-					ProjectedShadowInfo->RenderDepth(RHICmdList, this);
+					ProjectedShadowInfo->RenderDepth(RHICmdList, this, LocalShadowTargets);
 				}
 			}
 
@@ -3307,7 +3359,15 @@ bool FDeferredShadingSceneRenderer::RenderCachedPreshadows(FRHICommandListImmedi
 				{
 					check(ProjectedShadowInfo->bAllocated);
 					ProjectedShadowInfo->ClearDepth(RHICmdList, this);
-					ProjectedShadowInfo->RenderDepth(RHICmdList, this);
+
+					struct FLocalShadowTargets : public ISetShadowRenderTargets
+					{
+						virtual void SetRenderTargetsForShadows(FRHICommandList& RHICmdList)
+						{
+							SetRenderTarget(RHICmdList, FTextureRHIRef(), GSceneRenderTargets.PreShadowCacheDepthZ->GetRenderTargetItem().TargetableTexture);
+						}
+					} LocalShadowTargets;
+					ProjectedShadowInfo->RenderDepth(RHICmdList, this, LocalShadowTargets);
 					ProjectedShadowInfo->bDepthsCached = true;
 
 					const FResolveParams ResolveParams(FResolveRect(
@@ -3412,7 +3472,7 @@ bool FForwardShadingSceneRenderer::RenderShadowDepthMap(FRHICommandListImmediate
 		SCOPED_DRAW_EVENT(RHICmdList, ShadowDepthsFromOpaque, DEC_SCENE_ITEMS);
 
 		// Clearin the render target happens inside the Begin function
-		GSceneRenderTargets.BeginRenderingShadowDepth(RHICmdList);
+		GSceneRenderTargets.BeginRenderingShadowDepth(RHICmdList, true);
 
 		// render depth for each shadow
 		for (int32 ShadowIndex = 0; ShadowIndex < Shadows.Num(); ShadowIndex++)
@@ -3421,7 +3481,14 @@ bool FForwardShadingSceneRenderer::RenderShadowDepthMap(FRHICommandListImmediate
 			check(!ProjectedShadowInfo->bTranslucentShadow);
 			ProjectedShadowInfo->bRendered = false;
 			ProjectedShadowInfo->bAllocated = true;
-			ProjectedShadowInfo->RenderDepth(RHICmdList, this);
+			struct FLocalShadowTargets : public ISetShadowRenderTargets
+			{
+				virtual void SetRenderTargetsForShadows(FRHICommandList& RHICmdList)
+				{
+					GSceneRenderTargets.BeginRenderingShadowDepth(RHICmdList, false);
+				}
+			} LocalShadowTargets;
+			ProjectedShadowInfo->RenderDepth(RHICmdList, this, LocalShadowTargets);
 			bAttenuationBufferDirty = true;
 		}
 
diff --git a/Engine/Source/Runtime/Renderer/Private/ShadowRendering.h b/Engine/Source/Runtime/Renderer/Private/ShadowRendering.h
--- a/Engine/Source/Runtime/Renderer/Private/ShadowRendering.h
+++ b/Engine/Source/Runtime/Renderer/Private/ShadowRendering.h
@@ -490,6 +490,13 @@ public:
 	bool bIsTwoSided;
 };
 
+
+struct ISetShadowRenderTargets
+{
+	virtual ~ISetShadowRenderTargets() {}
+	virtual void SetRenderTargetsForShadows(FRHICommandList& RHICmdList) = 0;
+};
+
 /**
  * Information about a projected shadow.
  */
@@ -671,7 +678,10 @@ public:
 	/**
 	 * Renders the shadow subject depth.
 	 */
-	void RenderDepth(FRHICommandList& RHICmdList, class FSceneRenderer* SceneRenderer);
+	void RenderDepth(FRHICommandList& RHICmdList, class FSceneRenderer* SceneRenderer, ISetShadowRenderTargets& SetShadowRenderTargets);
+
+	/** Set state for depth rendering */
+	void SetStateForDepth(FRHICommandList& RHICmdList);
 
 	void ClearDepth(FRHICommandList& RHICmdList, class FDeferredShadingSceneRenderer* SceneRenderer);
 
@@ -756,7 +766,7 @@ private:
 	/**
 	* Renders the shadow subject depth, to a particular hacked view
 	*/
-	void RenderDepthInner(FRHICommandList& RHICmdList, class FSceneRenderer* SceneRenderer, const FViewInfo* FoundView);
+	void RenderDepthInner(FRHICommandList& RHICmdList, class FSceneRenderer* SceneRenderer, const FViewInfo* FoundView, ISetShadowRenderTargets& SetShadowRenderTargets);
 
 	/**
 	* Renders the dynamic shadow subject depth, to a particular hacked view
diff --git a/Engine/Source/Runtime/Renderer/Private/StaticMeshDrawList.h b/Engine/Source/Runtime/Renderer/Private/StaticMeshDrawList.h
--- a/Engine/Source/Runtime/Renderer/Private/StaticMeshDrawList.h
+++ b/Engine/Source/Runtime/Renderer/Private/StaticMeshDrawList.h
@@ -204,7 +204,12 @@ private:
 		{
 			check(IsInRenderingThread());
 			FBoundShaderStateInput BoundShaderStateInput = DrawingPolicy.GetBoundShaderStateInput(FeatureLevel);
-			BoundShaderState = CreateBoundShaderState_Internal(
+			BoundShaderState = 
+#if PLATFORM_SUPPORTS_PARALLEL_RHI_EXECUTE
+				RHICreateBoundShaderState(
+#else
+				CreateBoundShaderState_Internal(
+#endif
 				BoundShaderStateInput.VertexDeclarationRHI,
 				BoundShaderStateInput.VertexShaderRHI,
 				BoundShaderStateInput.HullShaderRHI,
@@ -289,14 +294,12 @@ public:
 
 	/**
 	* Draws only the static meshes which are in the visibility map.
-	* @param View - The view of the meshes to render.
 	* @param StaticMeshVisibilityMap - An map from FStaticMesh::Id to visibility state.
 	* @param BatchVisibilityArray - An array of batch element visibility bitmasks.
-	* @param Width - parallel width
-	* @param ParentCmdList - cmdlist to put the wait and execute task on
-	* @param OutDirty - set to true if anything is drawn
+	* @param BatchVisibilityArray - An array of batch element visibility bitmasks.
+	* @param ParallelCommandListSet - holds information on how to get a fresh command list and deal with submits, etc
 	*/
-	void DrawVisibleParallel(const FViewInfo& View, const typename DrawingPolicyType::ContextDataType PolicyContext, const TBitArray<SceneRenderingBitArrayAllocator>& StaticMeshVisibilityMap, const TArray<uint64, SceneRenderingAllocator>& BatchVisibilityArray, int32 Width, FRHICommandList& ParentCmdList, bool& OutDirty);
+	void DrawVisibleParallel(const typename DrawingPolicyType::ContextDataType PolicyContext, const TBitArray<SceneRenderingBitArrayAllocator>& StaticMeshVisibilityMap, const TArray<uint64, SceneRenderingAllocator>& BatchVisibilityArray, FParallelCommandListSet& ParallelCommandListSet);
 
 	/**
 	 * Draws only the static meshes which are in the visibility map, sorted front-to-back.
@@ -321,9 +324,9 @@ public:
 		return DrawVisible(RHICmdList, View, typename DrawingPolicyType::ContextDataType(), StaticMeshVisibilityMap, BatchVisibilityArray);
 	}
 
-	inline void DrawVisibleParallel(const FViewInfo& View, const TBitArray<SceneRenderingBitArrayAllocator>& StaticMeshVisibilityMap, const TArray<uint64, SceneRenderingAllocator>& BatchVisibilityArray, int32 Width, FRHICommandList& ParentCmdList, bool& OutDirty)
+	inline void DrawVisibleParallel(const TBitArray<SceneRenderingBitArrayAllocator>& StaticMeshVisibilityMap, const TArray<uint64, SceneRenderingAllocator>& BatchVisibilityArray, FParallelCommandListSet& ParallelCommandListSet)
 	{
-		DrawVisibleParallel(View, typename DrawingPolicyType::ContextDataType(), StaticMeshVisibilityMap, BatchVisibilityArray, Width, ParentCmdList, OutDirty);
+		DrawVisibleParallel(typename DrawingPolicyType::ContextDataType(), StaticMeshVisibilityMap, BatchVisibilityArray, ParallelCommandListSet);
 	}
 
 	inline int32 DrawVisibleFrontToBack(FRHICommandList& RHICmdList, const FViewInfo& View, const TBitArray<SceneRenderingBitArrayAllocator>& StaticMeshVisibilityMap, const TArray<uint64,SceneRenderingAllocator>& BatchVisibilityArray, int32 MaxToDraw)
diff --git a/Engine/Source/Runtime/Renderer/Private/StaticMeshDrawList.inl b/Engine/Source/Runtime/Renderer/Private/StaticMeshDrawList.inl
--- a/Engine/Source/Runtime/Renderer/Private/StaticMeshDrawList.inl
+++ b/Engine/Source/Runtime/Renderer/Private/StaticMeshDrawList.inl
@@ -355,15 +355,14 @@ public:
 
 template<typename DrawingPolicyType>
 void TStaticMeshDrawList<DrawingPolicyType>::DrawVisibleParallel(
-	const FViewInfo& View,
 	const typename DrawingPolicyType::ContextDataType PolicyContext,
 	const TBitArray<SceneRenderingBitArrayAllocator>& StaticMeshVisibilityMap,
 	const TArray<uint64, SceneRenderingAllocator>& BatchVisibilityArray,
-	int32 Width, FRHICommandList& ParentCmdList, bool& OutDirty
+	FParallelCommandListSet& ParallelCommandListSet
 	)
 {
 
-	int32 EffectiveThreads = FMath::Min<int32>(OrderedDrawingPolicies.Num(), Width);
+	int32 EffectiveThreads = FMath::Min<int32>(OrderedDrawingPolicies.Num(), ParallelCommandListSet.Width);
 
 	int32 Start = 0;
 	if (EffectiveThreads)
@@ -379,12 +378,12 @@ void TStaticMeshDrawList<DrawingPolicyType>::DrawVisibleParallel(
 			check(Last >= Start);
 
 			{
-				FRHICommandList* CmdList = new FRHICommandList;
+				FRHICommandList* CmdList = ParallelCommandListSet.NewParallelCommandList();
 
 				FGraphEventRef AnyThreadCompletionEvent = TGraphTask<FDrawVisibleAnyThreadTask<DrawingPolicyType> >::CreateTask(nullptr, ENamedThreads::RenderThread)
-					.ConstructAndDispatchWhenReady(this, CmdList, &View, PolicyContext, &StaticMeshVisibilityMap, &BatchVisibilityArray, Start, Last, &OutDirty);
+					.ConstructAndDispatchWhenReady(this, CmdList, &ParallelCommandListSet.View, PolicyContext, &StaticMeshVisibilityMap, &BatchVisibilityArray, Start, Last, &ParallelCommandListSet.OutDirty);
 
-				ParentCmdList.QueueAsyncCommandListSubmit(AnyThreadCompletionEvent, CmdList);
+				ParallelCommandListSet.AddParallelCommandList(CmdList, AnyThreadCompletionEvent);
 			}
 
 			Start = Last + 1;
diff --git a/Engine/Source/Runtime/Renderer/Private/TranslucentLighting.cpp b/Engine/Source/Runtime/Renderer/Private/TranslucentLighting.cpp
--- a/Engine/Source/Runtime/Renderer/Private/TranslucentLighting.cpp
+++ b/Engine/Source/Runtime/Renderer/Private/TranslucentLighting.cpp
@@ -1410,7 +1410,13 @@ void SetInjectionShader(
 	{
 		CachedShaderMap = MaterialShaderMap;
 		check(IsInRenderingThread()); // I didn't know quite how to deal with this caching. It won't work with threads.
-		BoundShaderState = CreateBoundShaderState_Internal(GScreenVertexDeclaration.VertexDeclarationRHI, VertexShader->GetVertexShader(), FHullShaderRHIRef(), FDomainShaderRHIRef(), PixelShader->GetPixelShader(), GeometryShader->GetGeometryShader());
+		BoundShaderState = 
+#if PLATFORM_SUPPORTS_PARALLEL_RHI_EXECUTE
+			RHICreateBoundShaderState(
+#else
+			CreateBoundShaderState_Internal(
+#endif
+			GScreenVertexDeclaration.VertexDeclarationRHI, VertexShader->GetVertexShader(), FHullShaderRHIRef(), FDomainShaderRHIRef(), PixelShader->GetPixelShader(), GeometryShader->GetGeometryShader());
 	}
 
 	RHICmdList.SetBoundShaderState(BoundShaderState);
diff --git a/Engine/Source/Runtime/Renderer/Private/TranslucentRendering.cpp b/Engine/Source/Runtime/Renderer/Private/TranslucentRendering.cpp
--- a/Engine/Source/Runtime/Renderer/Private/TranslucentRendering.cpp
+++ b/Engine/Source/Runtime/Renderer/Private/TranslucentRendering.cpp
@@ -11,13 +11,13 @@
 #include "LightPropagationVolume.h"
 #include "SceneUtils.h"
 
-static void SetTranslucentRenderTargetAndState(FRHICommandList& RHICmdList, const FViewInfo& View, bool bSeparateTranslucencyPass)
+static void SetTranslucentRenderTargetAndState(FRHICommandList& RHICmdList, const FViewInfo& View, bool bSeparateTranslucencyPass, bool bFirstTimeThisFrame = false)
 {
 	bool bSetupTranslucentState = true;
 
 	if (bSeparateTranslucencyPass && GSceneRenderTargets.IsSeparateTranslucencyActive(View))
 	{
-		bSetupTranslucentState = GSceneRenderTargets.BeginRenderingSeparateTranslucency(RHICmdList, View, false);
+		bSetupTranslucentState = GSceneRenderTargets.BeginRenderingSeparateTranslucency(RHICmdList, View, bFirstTimeThisFrame);
 	}
 	else
 	{
@@ -328,7 +328,7 @@ public:
 
 	FORCEINLINE TStatId GetStatId() const
 	{
-		RETURN_QUICK_DECLARE_CYCLE_STAT(FSetGlobalBoundShaderStateRenderThreadTask, STATGROUP_TaskGraphTasks);
+		RETURN_QUICK_DECLARE_CYCLE_STAT(FCopySceneColorAndRestoreRenderThreadTask, STATGROUP_TaskGraphTasks);
 	}
 
 	ENamedThreads::Type GetDesiredThread()
@@ -865,11 +865,33 @@ public:
 	}
 };
 
+class FTranslucencyPassParallelCommandListSet : public FParallelCommandListSet
+{
+	bool bSeparateTranslucency;
+	bool bFirstTimeThisFrame;
+public:
+	FTranslucencyPassParallelCommandListSet(const FViewInfo& InView, FRHICommandList& InParentCmdList, bool* InOutDirty, bool bInParallelExecute, bool bInSeparateTranslucency)
+		: FParallelCommandListSet(InView, InParentCmdList, InOutDirty, bInParallelExecute)
+		, bSeparateTranslucency(bInSeparateTranslucency)
+		, bFirstTimeThisFrame(bInSeparateTranslucency)
+	{
+		SetStateOnCommandList(ParentCmdList);
+	}
+	virtual void SetStateOnCommandList(FRHICommandList& CmdList) override
+	{
+		SetTranslucentRenderTargetAndState(CmdList, View, bSeparateTranslucency, bFirstTimeThisFrame);
+		bFirstTimeThisFrame = false;
+	}
+};
+
+static TAutoConsoleVariable<int32> CVarRHICmdTranslucencyPassDeferredContexts(
+	TEXT("r.RHICmdTranslucencyPassDeferredContexts"),
+	1,
+	TEXT("True to use deferred contexts to parallelize base pass command list execution.\n"));
 
 void FDeferredShadingSceneRenderer::RenderTranslucencyParallel(FRHICommandListImmediate& RHICmdList)
 {
 	GSceneRenderTargets.AllocLightAttenuation(); // materials will attempt to get this texture before the deferred command to set it up executes
-	int32 Width = CVarRHICmdWidth.GetValueOnRenderThread(); // we use a few more than needed to cover non-equal jobs
 
 	FScopedCommandListWaitForTasks Flusher(RHICmdList);
 	FTranslucencyDrawingPolicyFactory::ContextType ThisContext;
@@ -882,41 +904,45 @@ void FDeferredShadingSceneRenderer::RenderTranslucencyParallel(FRHICommandListIm
 
 		const FViewInfo& View = Views[ViewIndex];
 
-		SetTranslucentRenderTargetAndState(RHICmdList, View, false);
-
 		{
-			int32 NumPrims = View.TranslucentPrimSet.NumPrims() - View.TranslucentPrimSet.NumSeparateTranslucencyPrims();
-			int32 EffectiveThreads = FMath::Min<int32>(NumPrims, Width);
+			//@todo multiple views would be better moving this outside of the loop to batch all views together
+			FTranslucencyPassParallelCommandListSet ParallelCommandListSet(View, RHICmdList, nullptr, CVarRHICmdTranslucencyPassDeferredContexts.GetValueOnRenderThread() > 0, false);
 
-			int32 Start = 0;
-			if (EffectiveThreads)
 			{
+				int32 NumPrims = View.TranslucentPrimSet.NumPrims() - View.TranslucentPrimSet.NumSeparateTranslucencyPrims();
+				int32 EffectiveThreads = FMath::Min<int32>(NumPrims, ParallelCommandListSet.Width);
 
-				int32 NumPer = NumPrims / EffectiveThreads;
-				int32 Extra = NumPrims - NumPer * EffectiveThreads;
+				int32 Start = 0;
+				if (EffectiveThreads)
+				{
 
+					int32 NumPer = NumPrims / EffectiveThreads;
+					int32 Extra = NumPrims - NumPer * EffectiveThreads;
 
-				for (int32 ThreadIndex = 0; ThreadIndex < EffectiveThreads; ThreadIndex++)
-				{
-					int32 Last = Start + (NumPer - 1) + (ThreadIndex < Extra);
-					check(Last >= Start);
 
+					for (int32 ThreadIndex = 0; ThreadIndex < EffectiveThreads; ThreadIndex++)
 					{
-						FRHICommandList* CmdList = new FRHICommandList;
+						int32 Last = Start + (NumPer - 1) + (ThreadIndex < Extra);
+						check(Last >= Start);
 
-						FGraphEventRef AnyThreadCompletionEvent = TGraphTask<FDrawSortedTransAnyThreadTask>::CreateTask(nullptr, ENamedThreads::RenderThread)
-							.ConstructAndDispatchWhenReady(this, CmdList, &View, false, Start, Last);
+						{
+							FRHICommandList* CmdList = ParallelCommandListSet.NewParallelCommandList();
+
+							FGraphEventRef AnyThreadCompletionEvent = TGraphTask<FDrawSortedTransAnyThreadTask>::CreateTask(nullptr, ENamedThreads::RenderThread)
+								.ConstructAndDispatchWhenReady(this, CmdList, &View, false, Start, Last);
 
-						RHICmdList.QueueAsyncCommandListSubmit(AnyThreadCompletionEvent, CmdList); 
+							ParallelCommandListSet.AddParallelCommandList(CmdList, AnyThreadCompletionEvent);
+						}
+						Start = Last + 1;
 					}
-					Start = Last + 1;
 				}
 			}
+			// Draw the view's mesh elements with the translucent drawing policy.
+			DrawViewElementsParallel<FTranslucencyDrawingPolicyFactory>(ThisContext, SDPG_World, false, ParallelCommandListSet);
+			// Draw the view's mesh elements with the translucent drawing policy.
+			DrawViewElementsParallel<FTranslucencyDrawingPolicyFactory>(ThisContext, SDPG_Foreground, false, ParallelCommandListSet);
+
 		}
-		// Draw the view's mesh elements with the translucent drawing policy.
-		DrawViewElementsParallel<FTranslucencyDrawingPolicyFactory>(View, ThisContext, SDPG_World, false, RHICmdList, Width);
-		// Draw the view's mesh elements with the translucent drawing policy.
-		DrawViewElementsParallel<FTranslucencyDrawingPolicyFactory>(View, ThisContext, SDPG_Foreground, false, RHICmdList, Width);
 
 #if 0 // unsupported visualization in the parallel case
 		const FSceneViewState* ViewState = (const FSceneViewState*)View.State;
@@ -931,44 +957,38 @@ void FDeferredShadingSceneRenderer::RenderTranslucencyParallel(FRHICommandListIm
 		}
 #endif
 		{
-			bool bRenderSeparateTranslucency = View.TranslucentPrimSet.NumSeparateTranslucencyPrims() > 0;
-
 			{
 				// always call BeginRenderingSeparateTranslucency() even if there are no primitives to we keep the RT allocated
-				bool bSetupTranslucency = GSceneRenderTargets.BeginRenderingSeparateTranslucency(RHICmdList, View, true);
-				if (bRenderSeparateTranslucency && bSetupTranslucency)
-				{
-					RHICmdList.SetDepthStencilState(TStaticDepthStencilState<false, CF_GreaterEqual>::GetRHI());
-				}
-			}
+				FTranslucencyPassParallelCommandListSet ParallelCommandListSet(View, RHICmdList, nullptr, CVarRHICmdTranslucencyPassDeferredContexts.GetValueOnRenderThread() > 0, true);
 
-			// Draw only translucent prims that are in the SeparateTranslucency pass
-			if (bRenderSeparateTranslucency)
-			{
-				int32 NumPrims = View.TranslucentPrimSet.NumSeparateTranslucencyPrims();
-				int32 EffectiveThreads = FMath::Min<int32>(NumPrims, Width);
-
-				int32 Start = 0;
-				check(EffectiveThreads);
+				// Draw only translucent prims that are in the SeparateTranslucency pass
+				if (View.TranslucentPrimSet.NumSeparateTranslucencyPrims() > 0)
 				{
-					int32 NumPer = NumPrims / EffectiveThreads;
-					int32 Extra = NumPrims - NumPer * EffectiveThreads;
+					int32 NumPrims = View.TranslucentPrimSet.NumSeparateTranslucencyPrims();
+					int32 EffectiveThreads = FMath::Min<int32>(NumPrims, ParallelCommandListSet.Width);
 
-
-					for (int32 ThreadIndex = 0; ThreadIndex < EffectiveThreads; ThreadIndex++)
+					int32 Start = 0;
+					check(EffectiveThreads);
 					{
-						int32 Last = Start + (NumPer - 1) + (ThreadIndex < Extra);
-						check(Last >= Start);
+						int32 NumPer = NumPrims / EffectiveThreads;
+						int32 Extra = NumPrims - NumPer * EffectiveThreads;
+
 
+						for (int32 ThreadIndex = 0; ThreadIndex < EffectiveThreads; ThreadIndex++)
 						{
-							FRHICommandList* CmdList = new FRHICommandList;
+							int32 Last = Start + (NumPer - 1) + (ThreadIndex < Extra);
+							check(Last >= Start);
 
-							FGraphEventRef AnyThreadCompletionEvent = TGraphTask<FDrawSortedTransAnyThreadTask>::CreateTask(nullptr, ENamedThreads::RenderThread)
-								.ConstructAndDispatchWhenReady(this, CmdList, &View, true, Start, Last);
+							{
+								FRHICommandList* CmdList = ParallelCommandListSet.NewParallelCommandList();
 
-							RHICmdList.QueueAsyncCommandListSubmit(AnyThreadCompletionEvent, CmdList); 
+								FGraphEventRef AnyThreadCompletionEvent = TGraphTask<FDrawSortedTransAnyThreadTask>::CreateTask(nullptr, ENamedThreads::RenderThread)
+									.ConstructAndDispatchWhenReady(this, CmdList, &View, true, Start, Last);
+
+								ParallelCommandListSet.AddParallelCommandList(CmdList, AnyThreadCompletionEvent);
+							}
+							Start = Last + 1;
 						}
-						Start = Last + 1;
 					}
 				}
 			}
diff --git a/Engine/Source/Runtime/Renderer/Private/VelocityRendering.cpp b/Engine/Source/Runtime/Renderer/Private/VelocityRendering.cpp
--- a/Engine/Source/Runtime/Renderer/Private/VelocityRendering.cpp
+++ b/Engine/Source/Runtime/Renderer/Private/VelocityRendering.cpp
@@ -473,11 +473,11 @@ void FDeferredShadingSceneRenderer::RenderDynamicVelocitiesInner(FRHICommandList
 	}
 }
 
-void FDeferredShadingSceneRenderer::RenderDynamicVelocitiesMeshElementsInner(FRHICommandList& RHICmdList, const FViewInfo& View)
+void FDeferredShadingSceneRenderer::RenderDynamicVelocitiesMeshElementsInner(FRHICommandList& RHICmdList, const FViewInfo& View, int32 FirstIndex, int32 LastIndex)
 {
 	FVelocityDrawingPolicyFactory::ContextType Context(DDM_AllOccluders);
 
-	for (int32 MeshBatchIndex = 0; MeshBatchIndex < View.DynamicMeshElements.Num(); MeshBatchIndex++)
+	for (int32 MeshBatchIndex = FirstIndex; MeshBatchIndex < LastIndex; MeshBatchIndex++)
 	{
 		const FMeshBatchAndRelevance& MeshBatchAndRelevance = View.DynamicMeshElements[MeshBatchIndex];
 
@@ -529,77 +529,110 @@ public:
 
 	void DoTask(ENamedThreads::Type CurrentThread, const FGraphEventRef& MyCompletionGraphEvent)
 	{
-		ThisRenderer.RenderDynamicVelocitiesInner(RHICmdList, View, FirstIndex, LastIndex);
+		ThisRenderer.RenderDynamicVelocitiesMeshElementsInner(RHICmdList, View, FirstIndex, LastIndex);
 		RHICmdList.HandleRTThreadTaskCompletion(MyCompletionGraphEvent);
 	}
 };
 
+static void SetVelocitiesState(FRHICommandList& RHICmdList, const FViewInfo& View, TRefCountPtr<IPooledRenderTarget>& VelocityRT)
+{
+	SetRenderTarget(RHICmdList, VelocityRT->GetRenderTargetItem().TargetableTexture, GSceneRenderTargets.GetSceneDepthTexture());
 
-static TAutoConsoleVariable<int32> CVarParallelVelocity(
-	TEXT("r.ParallelVelocity"),
-	1,  
-	TEXT("Toggles parallel velocity rendering. Parallel rendering must be enabled for this to have an effect.\n"),
-	ECVF_RenderThreadSafe
-	);
+	const FIntPoint BufferSize = GSceneRenderTargets.GetBufferSizeXY();
+	const FIntPoint VelocityBufferSize = BufferSize;		// full resolution so we can reuse the existing full res z buffer
+
+	const uint32 MinX = View.ViewRect.Min.X * VelocityBufferSize.X / BufferSize.X;
+	const uint32 MinY = View.ViewRect.Min.Y * VelocityBufferSize.Y / BufferSize.Y;
+	const uint32 MaxX = View.ViewRect.Max.X * VelocityBufferSize.X / BufferSize.X;
+	const uint32 MaxY = View.ViewRect.Max.Y * VelocityBufferSize.Y / BufferSize.Y;
+
+	RHICmdList.SetViewport(MinX, MinY, 0.0f, MaxX, MaxY, 1.0f);
 
-void FDeferredShadingSceneRenderer::RenderVelocitiesInner(FRHICommandListImmediate& RHICmdList, const FViewInfo& View)
+	RHICmdList.SetBlendState(TStaticBlendState<CW_RGBA>::GetRHI());
+	// Note, this is a reversed Z depth surface, using CF_GreaterEqual.
+	RHICmdList.SetDepthStencilState(TStaticDepthStencilState<false,CF_GreaterEqual>::GetRHI());
+	RHICmdList.SetRasterizerState(GetStaticRasterizerState<true>(FM_Solid, CM_CW));
+}
+
+class FVelocityPassParallelCommandListSet : public FParallelCommandListSet
 {
-	const bool bUseGetMeshElements = ShouldUseGetDynamicMeshElements();
+	TRefCountPtr<IPooledRenderTarget>& VelocityRT;
+public:
+	FVelocityPassParallelCommandListSet(const FViewInfo& InView, FRHICommandList& InParentCmdList, bool* InOutDirty, bool bInParallelExecute, TRefCountPtr<IPooledRenderTarget>& InVelocityRT)
+		: FParallelCommandListSet(InView, InParentCmdList, InOutDirty, bInParallelExecute)
+		, VelocityRT(InVelocityRT)
+	{
+		SetStateOnCommandList(ParentCmdList);
+	}
+	virtual void SetStateOnCommandList(FRHICommandList& CmdList) override
+	{
+		SetVelocitiesState(CmdList, View, VelocityRT);
+	}
+};
 
-	if (GRHICommandList.UseParallelAlgorithms() && CVarParallelVelocity.GetValueOnRenderThread())
+static TAutoConsoleVariable<int32> CVarRHICmdVelocityPassDeferredContexts(
+	TEXT("r.RHICmdVelocityPassDeferredContexts"),
+	1,
+	TEXT("True to use deferred contexts to parallelize velocity pass command list execution.\n"));
+
+void FDeferredShadingSceneRenderer::RenderVelocitiesInnerParallel(FRHICommandListImmediate& RHICmdList, TRefCountPtr<IPooledRenderTarget>& VelocityRT)
+{
+	// parallel version
+	FScopedCommandListWaitForTasks Flusher(RHICmdList);
+
+	for(int32 ViewIndex = 0; ViewIndex < Views.Num(); ViewIndex++)
 	{
-		// parallel version
-		FScopedCommandListWaitForTasks Flusher(RHICmdList);
+		const FViewInfo& View = Views[ViewIndex];
 
-		int32 Width = CVarRHICmdWidth.GetValueOnRenderThread(); // we use a few more than needed to cover non-equal jobs
-		bool OutDirty = false; // unused
+		FVelocityPassParallelCommandListSet ParallelCommandListSet(View, RHICmdList, nullptr, CVarRHICmdVelocityPassDeferredContexts.GetValueOnRenderThread() > 0, VelocityRT);
 
-		Scene->VelocityDrawList.DrawVisibleParallel(View, View.StaticMeshVelocityMap, View.StaticMeshBatchVisibility, Width, RHICmdList, OutDirty);
+		Scene->VelocityDrawList.DrawVisibleParallel(View.StaticMeshVelocityMap, View.StaticMeshBatchVisibility, ParallelCommandListSet);
 
-		if (bUseGetMeshElements)
-		{
-			RenderDynamicVelocitiesMeshElementsInner(RHICmdList, View);
-		}
-		else
-		{
-			int32 NumPrims = View.VisibleDynamicPrimitives.Num();
-			int32 EffectiveThreads = FMath::Min<int32>(NumPrims, Width);
+		int32 NumPrims = View.DynamicMeshElements.Num();
+		int32 EffectiveThreads = FMath::Min<int32>(NumPrims, ParallelCommandListSet.Width);
 
-			int32 Start = 0;
-			if (EffectiveThreads)
-			{
-				check(IsInRenderingThread());
+		int32 Start = 0;
+		if (EffectiveThreads)
+		{
+			check(IsInRenderingThread());
 
-				int32 NumPer = NumPrims / EffectiveThreads;
-				int32 Extra = NumPrims - NumPer * EffectiveThreads;
+			int32 NumPer = NumPrims / EffectiveThreads;
+			int32 Extra = NumPrims - NumPer * EffectiveThreads;
 
 
-				for (int32 ThreadIndex = 0; ThreadIndex < EffectiveThreads; ThreadIndex++)
-				{
-					int32 Last = Start + (NumPer - 1) + (ThreadIndex < Extra);
-					check(Last >= Start);
+			for (int32 ThreadIndex = 0; ThreadIndex < EffectiveThreads; ThreadIndex++)
+			{
+				int32 Last = Start + (NumPer - 1) + (ThreadIndex < Extra);
+				check(Last >= Start);
 
-					FRHICommandList* CmdList = new FRHICommandList;
+				FRHICommandList* CmdList = ParallelCommandListSet.NewParallelCommandList();
 
-					FGraphEventRef AnyThreadCompletionEvent = TGraphTask<FRenderVelocityDynamicThreadTask>::CreateTask(nullptr, ENamedThreads::RenderThread)
-						.ConstructAndDispatchWhenReady(this, CmdList, &View, Start, Last);
+				FGraphEventRef AnyThreadCompletionEvent = TGraphTask<FRenderVelocityDynamicThreadTask>::CreateTask(nullptr, ENamedThreads::RenderThread)
+					.ConstructAndDispatchWhenReady(this, CmdList, &View, Start, Last);
 
-					RHICmdList.QueueAsyncCommandListSubmit(AnyThreadCompletionEvent, CmdList);
+				ParallelCommandListSet.AddParallelCommandList(CmdList, AnyThreadCompletionEvent);
 
-					Start = Last + 1;
-				}
+				Start = Last + 1;
 			}
 		}
 	}
-	else
+}
+
+void FDeferredShadingSceneRenderer::RenderVelocitiesInner(FRHICommandListImmediate& RHICmdList, TRefCountPtr<IPooledRenderTarget>& VelocityRT)
+{
+	const bool bUseGetMeshElements = ShouldUseGetDynamicMeshElements();
+
+	for(int32 ViewIndex = 0; ViewIndex < Views.Num(); ViewIndex++)
 	{
-		// single threaded version
+		const FViewInfo& View = Views[ViewIndex];
+
+		SetVelocitiesState(RHICmdList, View, VelocityRT);
 		// Draw velocities for movable static meshes.
 		Scene->VelocityDrawList.DrawVisible(RHICmdList, View, View.StaticMeshVelocityMap, View.StaticMeshBatchVisibility);
 
 		if (bUseGetMeshElements)
 		{
-			RenderDynamicVelocitiesMeshElementsInner(RHICmdList, View);
+			RenderDynamicVelocitiesMeshElementsInner(RHICmdList, View, 0, View.DynamicMeshElements.Num() - 1);
 		}
 		else
 		{
@@ -608,6 +641,16 @@ void FDeferredShadingSceneRenderer::RenderVelocitiesInner(FRHICommandListImmedia
 	}
 }
 
+
+
+static TAutoConsoleVariable<int32> CVarParallelVelocity(
+	TEXT("r.ParallelVelocity"),
+	1,  
+	TEXT("Toggles parallel velocity rendering. Parallel rendering must be enabled for this to have an effect.\n"),
+	ECVF_RenderThreadSafe
+	);
+
+
 void FDeferredShadingSceneRenderer::RenderVelocities(FRHICommandListImmediate& RHICmdList, TRefCountPtr<IPooledRenderTarget>& VelocityRT)
 {
 	check(FeatureLevel >= ERHIFeatureLevel::SM4);
@@ -643,23 +686,13 @@ void FDeferredShadingSceneRenderer::RenderVelocities(FRHICommandListImmediate& R
 	// Clear the velocity buffer (0.0f means "use static background velocity").
 	RHICmdList.Clear(true, FLinearColor::Black, false, 1.0f, false, 0, FIntRect());
 
-	for(int32 ViewIndex = 0; ViewIndex < Views.Num(); ViewIndex++)
+	if (GRHICommandList.UseParallelAlgorithms() && CVarParallelVelocity.GetValueOnRenderThread())
 	{
-		const FViewInfo& View = Views[ViewIndex];
-
-		const uint32 MinX = View.ViewRect.Min.X * VelocityBufferSize.X / BufferSize.X;
-		const uint32 MinY = View.ViewRect.Min.Y * VelocityBufferSize.Y / BufferSize.Y;
-		const uint32 MaxX = View.ViewRect.Max.X * VelocityBufferSize.X / BufferSize.X;
-		const uint32 MaxY = View.ViewRect.Max.Y * VelocityBufferSize.Y / BufferSize.Y;
-
-		RHICmdList.SetViewport(MinX, MinY, 0.0f, MaxX, MaxY, 1.0f);
-
-		RHICmdList.SetBlendState(TStaticBlendState<CW_RGBA>::GetRHI());
-		// Note, this is a reversed Z depth surface, using CF_GreaterEqual.
-		RHICmdList.SetDepthStencilState(TStaticDepthStencilState<false,CF_GreaterEqual>::GetRHI());
-		RHICmdList.SetRasterizerState(GetStaticRasterizerState<true>(FM_Solid, CM_CW));
-
-		RenderVelocitiesInner(RHICmdList, View);
+		RenderVelocitiesInnerParallel(RHICmdList, VelocityRT);
+	}
+	else
+	{
+		RenderVelocitiesInner(RHICmdList, VelocityRT);
 	}
 
 	RHICmdList.CopyToResolveTarget(VelocityRT->GetRenderTargetItem().TargetableTexture, VelocityRT->GetRenderTargetItem().ShaderResourceTexture, false, FResolveParams());
