Remove dead serialization code in FResolveParams

diff --git a/Engine/Source/Runtime/RHI/Public/RHI.h b/Engine/Source/Runtime/RHI/Public/RHI.h
--- a/Engine/Source/Runtime/RHI/Public/RHI.h
+++ b/Engine/Source/Runtime/RHI/Public/RHI.h
@@ -833,21 +833,6 @@ struct FResolveParams
 		, SourceArrayIndex(Other.SourceArrayIndex)
 		, DestArrayIndex(Other.DestArrayIndex)
 	{}
-
-	friend FArchive& operator<<(FArchive& Ar,FResolveParams& ResolveParams)
-	{
-		uint8 CubeFaceInt = ResolveParams.CubeFace;
-		Ar << CubeFaceInt;
-		if(Ar.IsLoading())
-		{
-			ResolveParams.CubeFace = (ECubeFace)ResolveParams.CubeFace;
-		}
-		Ar << ResolveParams.Rect;
-		Ar << ResolveParams.MipIndex;
-		Ar << ResolveParams.SourceArrayIndex;
-		Ar << ResolveParams.DestArrayIndex;
-		return Ar;
-	}
 };
 
 /** specifies an update region for a texture */
