UE4 - disable deferred contexts and fix some issues with bound shader state

diff --git a/Engine/Source/Runtime/Engine/Private/StaticBoundShaderState.cpp b/Engine/Source/Runtime/Engine/Private/StaticBoundShaderState.cpp
--- a/Engine/Source/Runtime/Engine/Private/StaticBoundShaderState.cpp
+++ b/Engine/Source/Runtime/Engine/Private/StaticBoundShaderState.cpp
@@ -57,11 +57,7 @@ FBoundShaderStateRHIParamRef FGlobalBoundShaderStateResource::GetInitializedRHI(
 	if(!IsValidRef(BoundShaderState))
 	{
 		BoundShaderState = 
-#if PLATFORM_SUPPORTS_PARALLEL_RHI_EXECUTE
 			RHICreateBoundShaderState(
-#else
-			CreateBoundShaderState_Internal(
-#endif
 			VertexDeclaration,
 			VertexShader,
 			FHullShaderRHIRef(), 
@@ -74,11 +70,7 @@ FBoundShaderStateRHIParamRef FGlobalBoundShaderStateResource::GetInitializedRHI(
 	else if (!GUsingNullRHI && UE_LOG_ACTIVE(LogShaders, Warning))
 	{
 		FBoundShaderStateRHIRef TempBoundShaderState = 
-#if PLATFORM_SUPPORTS_PARALLEL_RHI_EXECUTE
 			RHICreateBoundShaderState(
-#else
-			CreateBoundShaderState_Internal(
-#endif
 			VertexDeclaration,
 			VertexShader,
 			FHullShaderRHIRef(),
diff --git a/Engine/Source/Runtime/RHI/Private/BoundShaderStateCache.cpp b/Engine/Source/Runtime/RHI/Private/BoundShaderStateCache.cpp
--- a/Engine/Source/Runtime/RHI/Private/BoundShaderStateCache.cpp
+++ b/Engine/Source/Runtime/RHI/Private/BoundShaderStateCache.cpp
@@ -7,6 +7,11 @@
 #include "RHI.h"
 #include "BoundShaderStateCache.h"
 
+#if !defined(HAS_THREADSAFE_CreateBoundShaderState)
+#error "HAS_THREADSAFE_CreateBoundShaderState must be defined"
+#endif
+
+
 typedef TMap<FBoundShaderStateKey,FCachedBoundShaderStateLink*> FBoundShaderStateCache;
 
 /** Lazily initialized bound shader state cache singleton. */
@@ -16,7 +21,7 @@ static FBoundShaderStateCache& GetBoundShaderStateCache()
 	return BoundShaderStateCache;
 }
 
-#if PLATFORM_SUPPORTS_PARALLEL_RHI_EXECUTE
+#if HAS_THREADSAFE_CreateBoundShaderState
 static FCriticalSection BoundShaderStateCacheLock;
 #endif
 
@@ -34,7 +39,7 @@ FCachedBoundShaderStateLink::FCachedBoundShaderStateLink(
 	Key(VertexDeclaration,VertexShader,PixelShader,HullShader,DomainShader,GeometryShader)
 {
 	// Add this bound shader state to the cache.
-#if !PLATFORM_SUPPORTS_PARALLEL_RHI_EXECUTE
+#if !HAS_THREADSAFE_CreateBoundShaderState
 	GetBoundShaderStateCache().Add(Key,this);
 #endif
 }
@@ -49,19 +54,19 @@ FCachedBoundShaderStateLink::FCachedBoundShaderStateLink(
 	Key(VertexDeclaration,VertexShader,PixelShader)
 {
 	// Add this bound shader state to the cache.
-#if !PLATFORM_SUPPORTS_PARALLEL_RHI_EXECUTE
+#if !HAS_THREADSAFE_CreateBoundShaderState
 	GetBoundShaderStateCache().Add(Key,this);
 #endif
 }
 
 FCachedBoundShaderStateLink::~FCachedBoundShaderStateLink()
 {
-#if !PLATFORM_SUPPORTS_PARALLEL_RHI_EXECUTE
+#if !HAS_THREADSAFE_CreateBoundShaderState
 	GetBoundShaderStateCache().Remove(Key);
 #endif
 }
 
-#if PLATFORM_SUPPORTS_PARALLEL_RHI_EXECUTE
+#if HAS_THREADSAFE_CreateBoundShaderState
 void FCachedBoundShaderStateLink::AddToCache()
 {
 	FScopeLock Lock(&BoundShaderStateCacheLock);
@@ -75,7 +80,7 @@ void FCachedBoundShaderStateLink::RemoveFromCache()
 #endif
 
 
-#if PLATFORM_SUPPORTS_PARALLEL_RHI_EXECUTE
+#if HAS_THREADSAFE_CreateBoundShaderState
 FBoundShaderStateRHIRef GetCachedBoundShaderState_Threadsafe(
 	FVertexDeclarationRHIParamRef VertexDeclaration,
 	FVertexShaderRHIParamRef VertexShader,
diff --git a/Engine/Source/Runtime/RHI/Private/RHICommandList.cpp b/Engine/Source/Runtime/RHI/Private/RHICommandList.cpp
--- a/Engine/Source/Runtime/RHI/Private/RHICommandList.cpp
+++ b/Engine/Source/Runtime/RHI/Private/RHICommandList.cpp
@@ -48,7 +48,7 @@ TAutoConsoleVariable<int32> CVarRHICmdWidth(
 #if PLATFORM_SUPPORTS_PARALLEL_RHI_EXECUTE
 static TAutoConsoleVariable<int32> CVarRHICmdUseDeferredContexts(
 	TEXT("r.RHICmdUseDeferredContexts"),
-	1,
+	0,
 	TEXT("True to use deferred contexts to parallelize command list execution.\n"));
 static TAutoConsoleVariable<int32> CVarRHIDeferredContextWidth(
 	TEXT("r.RHIDeferredContextWidth"),
diff --git a/Engine/Source/Runtime/RHI/Public/RHICommandList.h b/Engine/Source/Runtime/RHI/Public/RHICommandList.h
--- a/Engine/Source/Runtime/RHI/Public/RHICommandList.h
+++ b/Engine/Source/Runtime/RHI/Public/RHICommandList.h
@@ -1219,7 +1219,7 @@ public:
 		FLocalBoundShaderState Result;
 		if (Bypass())
 		{
-#if PLATFORM_SUPPORTS_PARALLEL_RHI_EXECUTE
+#if HAS_THREADSAFE_CreateBoundShaderState
 			Result.BypassBSS = RHICreateBoundShaderState(VertexDeclarationRHI, VertexShaderRHI, HullShaderRHI, DomainShaderRHI, PixelShaderRHI, GeometryShaderRHI);
 #else
 			Result.BypassBSS = CreateBoundShaderState_Internal(VertexDeclarationRHI, VertexShaderRHI, HullShaderRHI, DomainShaderRHI, PixelShaderRHI, GeometryShaderRHI);
diff --git a/Engine/Source/Runtime/RHI/Public/RHICommandListCommandExecutes.inl b/Engine/Source/Runtime/RHI/Public/RHICommandListCommandExecutes.inl
--- a/Engine/Source/Runtime/RHI/Public/RHICommandListCommandExecutes.inl
+++ b/Engine/Source/Runtime/RHI/Public/RHICommandListCommandExecutes.inl
@@ -288,7 +288,7 @@ void FRHICommandBuildLocalBoundShaderState::Execute(FRHICommandListBase& CmdList
 	if (WorkArea.ComputedBSS->UseCount)
 	{
 		WorkArea.ComputedBSS->BSS = 	
-#if PLATFORM_SUPPORTS_PARALLEL_RHI_EXECUTE
+#if HAS_THREADSAFE_CreateBoundShaderState
 			RHICreateBoundShaderState(WorkArea.Args.VertexDeclarationRHI, WorkArea.Args.VertexShaderRHI, WorkArea.Args.HullShaderRHI, WorkArea.Args.DomainShaderRHI, WorkArea.Args.PixelShaderRHI, WorkArea.Args.GeometryShaderRHI);
 #else
 			CreateBoundShaderState_Internal(WorkArea.Args.VertexDeclarationRHI, WorkArea.Args.VertexShaderRHI, WorkArea.Args.HullShaderRHI, WorkArea.Args.DomainShaderRHI, WorkArea.Args.PixelShaderRHI, WorkArea.Args.GeometryShaderRHI);
diff --git a/Engine/Source/Runtime/RHI/Public/RHIMethods.h b/Engine/Source/Runtime/RHI/Public/RHIMethods.h
--- a/Engine/Source/Runtime/RHI/Public/RHIMethods.h
+++ b/Engine/Source/Runtime/RHI/Public/RHIMethods.h
@@ -477,7 +477,16 @@ DEFINE_RHIMETHOD_GLOBAL_1(
  * @param GeometryShader - existing geometry shader
  * @param PixelShader - existing pixel shader
  */
-#if PLATFORM_SUPPORTS_PARALLEL_RHI_EXECUTE
+
+#if !defined(HAS_THREADSAFE_CreateBoundShaderState)
+	#define HAS_THREADSAFE_CreateBoundShaderState PLATFORM_SUPPORTS_PARALLEL_RHI_EXECUTE
+#endif
+
+#if HAS_THREADSAFE_CreateBoundShaderState != PLATFORM_SUPPORTS_PARALLEL_RHI_EXECUTE
+	#error "HAS_THREADSAFE_CreateBoundShaderState changed meaning"
+#endif
+
+#if HAS_THREADSAFE_CreateBoundShaderState
 	DEFINE_RHIMETHOD_GLOBALTHREADSAFE_6(
 		FBoundShaderStateRHIRef,CreateBoundShaderState,
 		FVertexDeclarationRHIParamRef,VertexDeclaration,
diff --git a/Engine/Source/Runtime/RenderCore/Public/RenderResource.h b/Engine/Source/Runtime/RenderCore/Public/RenderResource.h
--- a/Engine/Source/Runtime/RenderCore/Public/RenderResource.h
+++ b/Engine/Source/Runtime/RenderCore/Public/RenderResource.h
@@ -606,6 +606,10 @@ private:
  * A list of the most recently used bound shader states.
  * This is used to keep bound shader states that have been used recently from being freed, as they're likely to be used again soon.
  */
+#if !defined(HAS_THREADSAFE_CreateBoundShaderState)
+#error "HAS_THREADSAFE_CreateBoundShaderState must be defined"
+#endif
+
 template<uint32 Size>
 class TBoundShaderStateHistory : public FRenderResource
 {
@@ -619,7 +623,7 @@ public:
 	/** Adds a bound shader state to the history. */
 	void Add(FBoundShaderStateRHIParamRef BoundShaderState)
 	{
-#if PLATFORM_SUPPORTS_PARALLEL_RHI_EXECUTE
+#if HAS_THREADSAFE_CreateBoundShaderState
 		FScopeLock Lock(&BoundShaderStateHistoryLock);
 #endif
 		BoundShaderStates[NextBoundShaderStateIndex] = BoundShaderState;
@@ -628,7 +632,7 @@ public:
 
 	FBoundShaderStateRHIParamRef GetLast()
 	{
-		check(!PLATFORM_SUPPORTS_PARALLEL_RHI_EXECUTE);
+		check(!HAS_THREADSAFE_CreateBoundShaderState);
 		// % doesn't work as we want on negative numbers, so handle the wraparound manually
 		uint32 LastIndex = NextBoundShaderStateIndex == 0 ? Size - 1 : NextBoundShaderStateIndex - 1;
 		return BoundShaderStates[LastIndex];
@@ -647,7 +651,7 @@ private:
 
 	FBoundShaderStateRHIRef BoundShaderStates[Size];
 	uint32 NextBoundShaderStateIndex;
-#if PLATFORM_SUPPORTS_PARALLEL_RHI_EXECUTE
+#if HAS_THREADSAFE_CreateBoundShaderState
 	FCriticalSection BoundShaderStateHistoryLock;
 #endif
 };
diff --git a/Engine/Source/Runtime/Renderer/Private/StaticMeshDrawList.h b/Engine/Source/Runtime/Renderer/Private/StaticMeshDrawList.h
--- a/Engine/Source/Runtime/Renderer/Private/StaticMeshDrawList.h
+++ b/Engine/Source/Runtime/Renderer/Private/StaticMeshDrawList.h
@@ -204,12 +204,7 @@ private:
 		{
 			check(IsInRenderingThread());
 			FBoundShaderStateInput BoundShaderStateInput = DrawingPolicy.GetBoundShaderStateInput(FeatureLevel);
-			BoundShaderState = 
-#if PLATFORM_SUPPORTS_PARALLEL_RHI_EXECUTE
-				RHICreateBoundShaderState(
-#else
-				CreateBoundShaderState_Internal(
-#endif
+			BoundShaderState = RHICreateBoundShaderState(
 				BoundShaderStateInput.VertexDeclarationRHI,
 				BoundShaderStateInput.VertexShaderRHI,
 				BoundShaderStateInput.HullShaderRHI,
diff --git a/Engine/Source/Runtime/Renderer/Private/TranslucentLighting.cpp b/Engine/Source/Runtime/Renderer/Private/TranslucentLighting.cpp
--- a/Engine/Source/Runtime/Renderer/Private/TranslucentLighting.cpp
+++ b/Engine/Source/Runtime/Renderer/Private/TranslucentLighting.cpp
@@ -1411,11 +1411,7 @@ void SetInjectionShader(
 		CachedShaderMap = MaterialShaderMap;
 		check(IsInRenderingThread()); // I didn't know quite how to deal with this caching. It won't work with threads.
 		BoundShaderState = 
-#if PLATFORM_SUPPORTS_PARALLEL_RHI_EXECUTE
 			RHICreateBoundShaderState(
-#else
-			CreateBoundShaderState_Internal(
-#endif
 			GScreenVertexDeclaration.VertexDeclarationRHI, VertexShader->GetVertexShader(), FHullShaderRHIRef(), FDomainShaderRHIRef(), PixelShader->GetPixelShader(), GeometryShader->GetGeometryShader());
 	}
 
