UE4 - fix two serious race conditions on the PS4 and XB1 that resulted in GPU hangs etc

diff --git a/Engine/Source/Runtime/RHI/Private/RHICommandList.cpp b/Engine/Source/Runtime/RHI/Private/RHICommandList.cpp
--- a/Engine/Source/Runtime/RHI/Private/RHICommandList.cpp
+++ b/Engine/Source/Runtime/RHI/Private/RHICommandList.cpp
@@ -58,7 +58,8 @@ static TAutoConsoleVariable<int32> CVarRHIDeferredContextWidth(
 
 RHI_API FRHICommandListExecutor GRHICommandList;
 
-static FGraphEventArray OutstandingTasks;
+static FGraphEventArray AllOutstandingTasks;
+static FGraphEventArray WaitOutstandingTasks;
 static FGraphEventRef RHIThreadTask;
 static FGraphEventRef RenderThreadSublistDispatchTask;
 
@@ -93,7 +94,7 @@ void FRHICommandListExecutor::ExecuteInner_DoExecute(FRHICommandListBase& CmdLis
 #endif
 
 	bool bDoStats = false;
-	STAT(bDoStats = FThreadStats::IsCollectingData() && (IsInRenderingThread() || IsInRHIThread());)
+	//STAT(bDoStats = FThreadStats::IsCollectingData() && (IsInRenderingThread() || IsInRHIThread());)
 	FRHICommandListIterator Iter(CmdList);
 	if (bDoStats)
 	{
@@ -233,10 +234,10 @@ void FRHICommandListExecutor::ExecuteInner(FRHICommandListBase& CmdList)
 			static_assert(sizeof(FRHICommandList) == sizeof(FRHICommandListImmediate), "We are memswapping FRHICommandList and FRHICommandListImmediate; they need to be swappable.");
 			SwapCmdList->ExchangeCmdList(CmdList);
 
-			if (OutstandingTasks.Num() || RenderThreadSublistDispatchTask.GetReference())
+			if (AllOutstandingTasks.Num() || RenderThreadSublistDispatchTask.GetReference())
 			{
-				FGraphEventArray Prereq(OutstandingTasks);
-				OutstandingTasks.Reset();
+				FGraphEventArray Prereq(AllOutstandingTasks);
+				AllOutstandingTasks.Reset();
 				if (RenderThreadSublistDispatchTask.GetReference())
 				{
 					Prereq.Add(RenderThreadSublistDispatchTask);
@@ -264,10 +265,13 @@ void FRHICommandListExecutor::ExecuteInner(FRHICommandListBase& CmdList)
 					FTaskGraphInterface::Get().WaitUntilTaskCompletes(RenderThreadSublistDispatchTask, ENamedThreads::RenderThread_Local);
 					RenderThreadSublistDispatchTask = nullptr;
 				}
-				if (RHIThreadTask.GetReference())
+				while (RHIThreadTask.GetReference())
 				{
 					FTaskGraphInterface::Get().WaitUntilTaskCompletes(RHIThreadTask, ENamedThreads::RenderThread_Local);
-					RHIThreadTask = nullptr;
+					if (RHIThreadTask.GetReference() && RHIThreadTask->IsComplete())
+					{
+						RHIThreadTask = nullptr;
+					}
 				}
 			}
 			return;
@@ -284,7 +288,7 @@ void FRHICommandListExecutor::ExecuteInner(FRHICommandListBase& CmdList)
 				FTaskGraphInterface::Get().WaitUntilTaskCompletes(RenderThreadSublistDispatchTask, ENamedThreads::RenderThread_Local);
 				RenderThreadSublistDispatchTask = nullptr;
 			}
-			if (RHIThreadTask.GetReference())
+			while (RHIThreadTask.GetReference())
 			{
 				if (FTaskGraphInterface::Get().IsThreadProcessingTasks(ENamedThreads::RenderThread_Local))
 				{
@@ -292,7 +296,10 @@ void FRHICommandListExecutor::ExecuteInner(FRHICommandListBase& CmdList)
 					UE_LOG(LogRHI, Fatal, TEXT("Deadlock in FRHICommandListExecutor::ExecuteInner 3."));
 				}
 				FTaskGraphInterface::Get().WaitUntilTaskCompletes(RHIThreadTask, ENamedThreads::RenderThread_Local);
-				RHIThreadTask = nullptr;
+				if (RHIThreadTask.GetReference() && RHIThreadTask->IsComplete())
+				{
+					RHIThreadTask = nullptr;
+				}
 			}
 		}
 	}
@@ -417,6 +424,8 @@ bool FRHICommandListExecutor::IsRHIThreadActive()
 	return !!RHIThreadTask.GetReference();
 }
 
+DECLARE_CYCLE_STAT(TEXT("FNullGraphTask.RHIThreadFence"), STAT_RHIThreadFence, STATGROUP_TaskGraphTasks);
+
 FGraphEventRef FRHICommandListExecutor::RHIThreadFence()
 {
 	check(IsInRenderingThread() && GRHIThread);
@@ -424,12 +433,31 @@ FGraphEventRef FRHICommandListExecutor::RHIThreadFence()
 		QUICK_SCOPE_CYCLE_COUNTER(STAT_RHIThreadFence_Dispatch);
 		FRHICommandListExecutor::GetImmediateCommandList().ImmediateFlush(EImmediateFlushType::DispatchToRHIThread);
 	}
-	// this isn't quite right, but is probably fine for current uses. If there is a RenderThreadSublistDispatchTask, then we need that to be the fence, however, then that RT thread task needs to "do not complete until" or something so we can track the actual completion.
+#if PLATFORM_SUPPORTS_RHI_THREAD
+	FGraphEventArray Finish;
 	if (RHIThreadTask.GetReference() && RHIThreadTask->IsComplete())
 	{
 		RHIThreadTask = nullptr;
 	}
-	return RHIThreadTask;
+	if (RHIThreadTask.GetReference())
+	{
+		Finish.Add(RHIThreadTask);
+	}
+	if (RenderThreadSublistDispatchTask.GetReference() && RenderThreadSublistDispatchTask->IsComplete())
+	{
+		RenderThreadSublistDispatchTask = nullptr;
+	}
+	if (RenderThreadSublistDispatchTask.GetReference())
+	{
+		Finish.Add(RenderThreadSublistDispatchTask);
+	}
+	if (Finish.Num())
+	{
+		return TGraphTask<FNullGraphTask>::CreateTask( &Finish, ENamedThreads::RenderThread )
+			.ConstructAndDispatchWhenReady(GET_STATID(STAT_RHIThreadFence), ENamedThreads::RHIThread);
+	}
+#endif
+	return FGraphEventRef();
 }
 void FRHICommandListExecutor::WaitOnRHIThreadFence(FGraphEventRef& Fence)
 {
@@ -537,48 +565,44 @@ DECLARE_CYCLE_STAT(TEXT("Parallel Async Chain Execute"), STAT_ParallelChainExecu
 
 struct FRHICommandWaitForAndSubmitSubListParallel : public FRHICommand<FRHICommandWaitForAndSubmitSubListParallel>
 {
-	FGraphEventRef* TranslateCompletionEvents;
-	IRHICommandContextContainer** ContextContainers;
+	FGraphEventRef TranslateCompletionEvent;
+	IRHICommandContextContainer* ContextContainer;
 	int32 Num;
+	int32 Index;
 
-	FORCEINLINE_DEBUGGABLE FRHICommandWaitForAndSubmitSubListParallel(FGraphEventRef* InTranslateCompletionEvents, IRHICommandContextContainer** InContextContainers, int32 InNum)
-		: TranslateCompletionEvents(InTranslateCompletionEvents)
-		, ContextContainers(InContextContainers)
+	FORCEINLINE_DEBUGGABLE FRHICommandWaitForAndSubmitSubListParallel(FGraphEventRef& InTranslateCompletionEvent, IRHICommandContextContainer* InContextContainer, int32 InNum, int32 InIndex)
+		: TranslateCompletionEvent(InTranslateCompletionEvent)
+		, ContextContainer(InContextContainer)
 		, Num(InNum)
+		, Index(InIndex)
 	{
-		check(TranslateCompletionEvents && ContextContainers && Num);
+		check(ContextContainer && Num);
 	}
 	void Execute(FRHICommandListBase& CmdList)
 	{
-		check(TranslateCompletionEvents && ContextContainers && Num);
-		INC_DWORD_STAT_BY(STAT_ParallelChainLinkCount, Num);
-		for (int32 Index = 0; Index < Num; Index++)
-		{
-			FGraphEventRef& EventToWaitFor = TranslateCompletionEvents[Index];
-			IRHICommandContextContainer* ContextContainer = ContextContainers[Index];
+		check(ContextContainer && Num && IsInRHIThread());
+		INC_DWORD_STAT_BY(STAT_ParallelChainLinkCount, 1);
 
-			if (EventToWaitFor.GetReference() && !EventToWaitFor->IsComplete())
+		if (TranslateCompletionEvent.GetReference() && !TranslateCompletionEvent->IsComplete())
+		{
+			SCOPE_CYCLE_COUNTER(STAT_ParallelChainWait);
+			if (IsInRenderingThread())
 			{
-				SCOPE_CYCLE_COUNTER(STAT_ParallelChainWait);
-				if (IsInRenderingThread())
-				{
-					FTaskGraphInterface::Get().WaitUntilTaskCompletes(EventToWaitFor, ENamedThreads::RenderThread_Local);
-				}
-				else if (IsInRHIThread())
-				{
-					FTaskGraphInterface::Get().WaitUntilTaskCompletes(EventToWaitFor, ENamedThreads::RHIThread);
-				}
-				else
-				{
-					check(0);
-				}
+				FTaskGraphInterface::Get().WaitUntilTaskCompletes(TranslateCompletionEvent, ENamedThreads::RenderThread_Local);
 			}
-			EventToWaitFor.~FGraphEventRef();
+			else if (IsInRHIThread())
+			{
+				FTaskGraphInterface::Get().WaitUntilTaskCompletes(TranslateCompletionEvent, ENamedThreads::RHIThread);
+			}
+			else
 			{
-				SCOPE_CYCLE_COUNTER(STAT_ParallelChainExecute);
-				ContextContainer->SubmitAndFreeContextContainer(Index, Num);
+				check(0);
 			}
 		}
+		{
+			SCOPE_CYCLE_COUNTER(STAT_ParallelChainExecute);
+			ContextContainer->SubmitAndFreeContextContainer(Index, Num);
+		}
 	}
 };
 #endif
@@ -648,10 +672,6 @@ void FRHICommandListBase::QueueParallelAsyncCommandListSubmit(FGraphEventRef* An
 				int32 NumPer = Num / EffectiveThreads;
 				int32 Extra = Num - NumPer * EffectiveThreads;
 
-				FGraphEventRef* OutTranslateCompletionEvents = (FGraphEventRef*)Alloc(sizeof(FGraphEventRef) * EffectiveThreads, ALIGNOF(FGraphEventRef));
-				IRHICommandContextContainer** ContextContainers = (IRHICommandContextContainer**)Alloc(sizeof(IRHICommandContextContainer*) * EffectiveThreads, ALIGNOF(IRHICommandContextContainer*));
-				FRHICommandListBase** RHICmdLists = (FRHICommandListBase**)Alloc(sizeof(FRHICommandListBase*) * Num, ALIGNOF(FRHICommandListBase*));
-
 				for (int32 ThreadIndex = 0; ThreadIndex < EffectiveThreads; ThreadIndex++)
 				{
 					int32 Last = Start + (NumPer - 1) + (ThreadIndex < Extra);
@@ -664,29 +684,31 @@ void FRHICommandListBase::QueueParallelAsyncCommandListSubmit(FGraphEventRef* An
 					check(ContextContainer);
 
 					FGraphEventArray Prereq;
+					FRHICommandListBase** RHICmdLists = (FRHICommandListBase**)Alloc(sizeof(FRHICommandListBase*) * (1 + Last - Start), ALIGNOF(FRHICommandListBase*));
 					for (int32 Index = Start; Index <= Last; Index++)
 					{
 						FGraphEventRef& AnyThreadCompletionEvent = AnyThreadCompletionEvents[Index];
 						FRHICommandList* CmdList = CmdLists[Index];
-						RHICmdLists[Index] = CmdList;
+						RHICmdLists[Index - Start] = CmdList;
 						if (AnyThreadCompletionEvent.GetReference())
 						{
 							Prereq.Add(AnyThreadCompletionEvent);
-							OutstandingTasks.Add(AnyThreadCompletionEvent);
+							AllOutstandingTasks.Add(AnyThreadCompletionEvent);
+							WaitOutstandingTasks.Add(AnyThreadCompletionEvent);
 						}
 					}
-					FGraphEventRef TranslateCompletionEvent = TGraphTask<FParallelTranslateCommandList>::CreateTask(&Prereq, ENamedThreads::RenderThread).ConstructAndDispatchWhenReady(&RHICmdLists[Start], 1 + Last - Start, ContextContainer);
+					FGraphEventRef TranslateCompletionEvent = TGraphTask<FParallelTranslateCommandList>::CreateTask(&Prereq, ENamedThreads::RenderThread).ConstructAndDispatchWhenReady(&RHICmdLists[0], 1 + Last - Start, ContextContainer);
+
+					AllOutstandingTasks.Add(TranslateCompletionEvent);
+					new (AllocCommand<FRHICommandWaitForAndSubmitSubListParallel>()) FRHICommandWaitForAndSubmitSubListParallel(TranslateCompletionEvent, ContextContainer, EffectiveThreads, ThreadIndex);
+					if (GRHIThread)
+					{
+						FRHICommandListExecutor::GetImmediateCommandList().ImmediateFlush(EImmediateFlushType::DispatchToRHIThread); // we don't want stuff after the async cmd list to be bundled with it
+					}
 
-					new ((void*)&OutTranslateCompletionEvents[ThreadIndex]) FGraphEventRef(TranslateCompletionEvent);
-					ContextContainers[ThreadIndex] = ContextContainer;
 					ContextContainer = nullptr;
 					Start = Last + 1;
 				}
-				new (AllocCommand<FRHICommandWaitForAndSubmitSubListParallel>()) FRHICommandWaitForAndSubmitSubListParallel(OutTranslateCompletionEvents, ContextContainers, EffectiveThreads);
-				if (GRHIThread)
-				{
-					FRHICommandListExecutor::GetImmediateCommandList().ImmediateFlush(EImmediateFlushType::DispatchToRHIThread); // we don't want stuff after the async cmd list to be bundled with it
-				}
 				return;
 			}
 		}
@@ -698,7 +720,8 @@ void FRHICommandListBase::QueueParallelAsyncCommandListSubmit(FGraphEventRef* An
 		FRHICommandList* CmdList = CmdLists[Index];
 		if (AnyThreadCompletionEvent.GetReference())
 		{
-			OutstandingTasks.Add(AnyThreadCompletionEvent);
+			AllOutstandingTasks.Add(AnyThreadCompletionEvent);
+			WaitOutstandingTasks.Add(AnyThreadCompletionEvent);
 		}
 		new (AllocCommand<FRHICommandWaitForAndSubmitSubList>()) FRHICommandWaitForAndSubmitSubList(AnyThreadCompletionEvent, CmdList);
 	}
@@ -717,7 +740,8 @@ void FRHICommandListBase::QueueAsyncCommandListSubmit(FGraphEventRef& AnyThreadC
 	}
 	if (AnyThreadCompletionEvent.GetReference())
 	{
-		OutstandingTasks.Add(AnyThreadCompletionEvent);
+		AllOutstandingTasks.Add(AnyThreadCompletionEvent);
+		WaitOutstandingTasks.Add(AnyThreadCompletionEvent);
 	}
 	new (AllocCommand<FRHICommandWaitForAndSubmitSubList>()) FRHICommandWaitForAndSubmitSubList(AnyThreadCompletionEvent, CmdList);
 	if (GRHIThread)
@@ -856,49 +880,48 @@ DECLARE_CYCLE_STAT(TEXT("Prewait dispatch"), STAT_PrewaitDispatch, STATGROUP_RHI
 void FRHICommandListBase::WaitForTasks(bool bKnownToBeComplete)
 {
 	check(IsImmediate() && IsInRenderingThread());
-	if (GRHIThread)
+	if (WaitOutstandingTasks.Num())
 	{
-		if (bKnownToBeComplete)
-		{
-			check(!OutstandingTasks.Num());
-		}
-		else if (OutstandingTasks.Num() || RenderThreadSublistDispatchTask.GetReference())
+		bool bAny = false;
+		for (int32 Index = 0; Index < WaitOutstandingTasks.Num(); Index++)
 		{
+			if (!WaitOutstandingTasks[Index]->IsComplete())
 			{
-				SCOPE_CYCLE_COUNTER(STAT_PrewaitDispatch);
-				// as long as we are sitting around, make sure the rhi thread is not
-				FRHICommandListExecutor::GetImmediateCommandList().ImmediateFlush(EImmediateFlushType::DispatchToRHIThread);
-			}
-			check(!OutstandingTasks.Num());
-			if (RenderThreadSublistDispatchTask.GetReference())
-			{
-				if (FTaskGraphInterface::Get().IsThreadProcessingTasks(ENamedThreads::RenderThread_Local))
-				{
-					// this is a deadlock. RT tasks must be done by now or they won't be done. We could add a third queue...
-					UE_LOG(LogRHI, Fatal, TEXT("Deadlock in FRHICommandListBase::WaitForRHIThreadTasks."));
-				}
-				FTaskGraphInterface::Get().WaitUntilTaskCompletes(RenderThreadSublistDispatchTask, ENamedThreads::RenderThread_Local);
-				RenderThreadSublistDispatchTask = nullptr;
+				check(!bKnownToBeComplete);
+				bAny = true;
+				break;
 			}
 		}
-	}	
-	else
+		if (bAny)
+		{
+			SCOPE_CYCLE_COUNTER(STAT_ExplicitWait);
+			FTaskGraphInterface::Get().WaitUntilTasksComplete(WaitOutstandingTasks, ENamedThreads::RenderThread_Local);
+		}
+		WaitOutstandingTasks.Reset();
+	}
+}
+
+DECLARE_CYCLE_STAT(TEXT("Explicit wait for dispatch"), STAT_ExplicitWaitDispatch, STATGROUP_RHICMDLIST);
+void FRHICommandListBase::WaitForDispatch()
+{
+	check(IsImmediate() && IsInRenderingThread());
+	check(!AllOutstandingTasks.Num()); // dispatch before you get here
+	if (RenderThreadSublistDispatchTask.GetReference() && RenderThreadSublistDispatchTask->IsComplete())
+	{
+		RenderThreadSublistDispatchTask = nullptr;
+	}
+	while (RenderThreadSublistDispatchTask.GetReference())
 	{
-		if (OutstandingTasks.Num())
+		SCOPE_CYCLE_COUNTER(STAT_ExplicitWaitDispatch);
+		if (FTaskGraphInterface::Get().IsThreadProcessingTasks(ENamedThreads::RenderThread_Local))
 		{
-			if (bKnownToBeComplete)
-			{
-				for (int32 Index = 0; Index < OutstandingTasks.Num(); Index++)
-				{
-					check(OutstandingTasks[Index]->IsComplete());
-				}
-			}
-			else
-			{
-				SCOPE_CYCLE_COUNTER(STAT_ExplicitWait);
-				FTaskGraphInterface::Get().WaitUntilTasksComplete(OutstandingTasks, ENamedThreads::RenderThread_Local);
-			}
-			OutstandingTasks.Reset();
+			// this is a deadlock. RT tasks must be done by now or they won't be done. We could add a third queue...
+			UE_LOG(LogRHI, Fatal, TEXT("Deadlock in FRHICommandListBase::WaitForDispatch."));
+		}
+		FTaskGraphInterface::Get().WaitUntilTaskCompletes(RenderThreadSublistDispatchTask, ENamedThreads::RenderThread_Local);
+		if (RenderThreadSublistDispatchTask.GetReference() && RenderThreadSublistDispatchTask->IsComplete())
+		{
+			RenderThreadSublistDispatchTask = nullptr;
 		}
 	}
 }
@@ -912,7 +935,7 @@ void FRHICommandListBase::WaitForRHIThreadTasks()
 	{
 		RHIThreadTask = nullptr;
 	}
-	if (RHIThreadTask.GetReference())
+	while (RHIThreadTask.GetReference())
 	{
 		SCOPE_CYCLE_COUNTER(STAT_ExplicitWaitRHIThread);
 		if (FTaskGraphInterface::Get().IsThreadProcessingTasks(ENamedThreads::RenderThread_Local))
@@ -929,7 +952,10 @@ void FRHICommandListBase::WaitForRHIThreadTasks()
 		{
 			FTaskGraphInterface::Get().WaitUntilTaskCompletes(RHIThreadTask, ENamedThreads::RenderThread_Local);
 		}
-		RHIThreadTask = nullptr;
+		if (RHIThreadTask.GetReference() && RHIThreadTask->IsComplete())
+		{
+			RHIThreadTask = nullptr;
+		}
 	}
 }
 
diff --git a/Engine/Source/Runtime/RHI/Public/RHICommandList.h b/Engine/Source/Runtime/RHI/Public/RHICommandList.h
--- a/Engine/Source/Runtime/RHI/Public/RHICommandList.h
+++ b/Engine/Source/Runtime/RHI/Public/RHICommandList.h
@@ -77,6 +77,7 @@ public:
 	void QueueRenderThreadCommandListSubmit(FGraphEventRef& RenderThreadCompletionEvent, class FRHICommandList* CmdList);
 	void QueueCommandListSubmit(class FRHICommandList* CmdList);
 	void WaitForTasks(bool bKnownToBeComplete = false);
+	void WaitForDispatch();
 	void WaitForRHIThreadTasks();
 	void HandleRTThreadTaskCompletion(const FGraphEventRef& MyCompletionGraphEvent);
 
diff --git a/Engine/Source/Runtime/RHI/Public/RHICommandList.inl b/Engine/Source/Runtime/RHI/Public/RHICommandList.inl
--- a/Engine/Source/Runtime/RHI/Public/RHICommandList.inl
+++ b/Engine/Source/Runtime/RHI/Public/RHICommandList.inl
@@ -59,6 +59,7 @@ FORCEINLINE_DEBUGGABLE void FRHICommandListImmediate::ImmediateFlush(EImmediateF
 			{
 				GRHICommandList.ExecuteList(*this);
 			}
+			WaitForDispatch();
 			if (GRHIThread)
 			{
 				WaitForRHIThreadTasks();
@@ -72,6 +73,7 @@ FORCEINLINE_DEBUGGABLE void FRHICommandListImmediate::ImmediateFlush(EImmediateF
 			{
 				GRHICommandList.ExecuteList(*this);
 			}
+			WaitForDispatch();
 			WaitForRHIThreadTasks();
 			WaitForTasks(true); // these are already done, but this resets the outstanding array
 			FRHIResource::FlushPendingDeletes();
