UE4 - fix togglerhithread

diff --git a/Engine/Source/Runtime/RHI/Private/RHICommandList.cpp b/Engine/Source/Runtime/RHI/Private/RHICommandList.cpp
--- a/Engine/Source/Runtime/RHI/Private/RHICommandList.cpp
+++ b/Engine/Source/Runtime/RHI/Private/RHICommandList.cpp
@@ -720,7 +720,10 @@ void FRHICommandListBase::QueueParallelAsyncCommandListSubmit(FGraphEventRef* An
 		FRHICommandList* CmdList = CmdLists[Index];
 		if (AnyThreadCompletionEvent.GetReference())
 		{
-			AllOutstandingTasks.Add(AnyThreadCompletionEvent);
+			if (GRHIThread)
+			{
+				AllOutstandingTasks.Add(AnyThreadCompletionEvent);
+			}
 			WaitOutstandingTasks.Add(AnyThreadCompletionEvent);
 		}
 		new (AllocCommand<FRHICommandWaitForAndSubmitSubList>()) FRHICommandWaitForAndSubmitSubList(AnyThreadCompletionEvent, CmdList);
@@ -740,7 +743,10 @@ void FRHICommandListBase::QueueAsyncCommandListSubmit(FGraphEventRef& AnyThreadC
 	}
 	if (AnyThreadCompletionEvent.GetReference())
 	{
-		AllOutstandingTasks.Add(AnyThreadCompletionEvent);
+		if (GRHIThread)
+		{
+			AllOutstandingTasks.Add(AnyThreadCompletionEvent);
+		}
 		WaitOutstandingTasks.Add(AnyThreadCompletionEvent);
 	}
 	new (AllocCommand<FRHICommandWaitForAndSubmitSubList>()) FRHICommandWaitForAndSubmitSubList(AnyThreadCompletionEvent, CmdList);
