UE4 - Fix uninitialized vars (not clearing depth on Metal)

diff --git a/Engine/Source/Runtime/RHI/Public/RHIResources.h b/Engine/Source/Runtime/RHI/Public/RHIResources.h
--- a/Engine/Source/Runtime/RHI/Public/RHIResources.h
+++ b/Engine/Source/Runtime/RHI/Public/RHIResources.h
@@ -713,6 +713,8 @@ public:
 		NumColorRenderTargets(InNumColorRenderTargets),
 		bClearColor(InNumColorRenderTargets > 0 && InColorRenderTargets[0].LoadAction == ERenderTargetLoadAction::EClear),
 		DepthStencilRenderTarget(InDepthStencilRenderTarget),
+		DepthClearValue(0.0f),
+		StencilClearValue(0),
 		bClearDepth(InDepthStencilRenderTarget.DepthLoadAction == ERenderTargetLoadAction::EClear),
 		bClearStencil(InDepthStencilRenderTarget.StencilLoadAction == ERenderTargetLoadAction::EClear)
 	{
