UE4 - minor parallel rendering fixes. Improved constructors for Xbox RHI. delete comments and cruft. Turn off parallel shadows. Flush before reflection capture on all platforms.

diff --git a/Engine/Source/Runtime/RHI/Private/RHICommandList.cpp b/Engine/Source/Runtime/RHI/Private/RHICommandList.cpp
--- a/Engine/Source/Runtime/RHI/Private/RHICommandList.cpp
+++ b/Engine/Source/Runtime/RHI/Private/RHICommandList.cpp
@@ -48,7 +48,7 @@ TAutoConsoleVariable<int32> CVarRHICmdWidth(
 #if PLATFORM_SUPPORTS_PARALLEL_RHI_EXECUTE
 static TAutoConsoleVariable<int32> CVarRHICmdUseDeferredContexts(
 	TEXT("r.RHICmdUseDeferredContexts"),
-	0,
+	1,
 	TEXT("True to use deferred contexts to parallelize command list execution.\n"));
 static TAutoConsoleVariable<int32> CVarRHIDeferredContextWidth(
 	TEXT("r.RHIDeferredContextWidth"),
@@ -86,6 +86,8 @@ void FRHICommandListImmediate::SetCurrentStat(TStatId Stat)
 	new (AllocCommand<FRHICommandStat>()) FRHICommandStat(Stat);
 }
 
+FRHICommandBase* GCurrentCommand = nullptr;
+
 void FRHICommandListExecutor::ExecuteInner_DoExecute(FRHICommandListBase& CmdList)
 {
 	CmdList.bExecuting = true;
@@ -111,6 +113,7 @@ void FRHICommandListExecutor::ExecuteInner_DoExecute(FRHICommandListBase& CmdLis
 		while (Iter.HasCommandsLeft())
 		{
 			FRHICommandBase* Cmd = Iter.NextCommand();
+			GCurrentCommand = Cmd;
 			//FPlatformMisc::Prefetch(Cmd->Next);
 			Cmd->CallExecuteAndDestruct(CmdList);
 		}
@@ -654,7 +657,7 @@ void FRHICommandListBase::QueueParallelAsyncCommandListSubmit(FGraphEventRef* An
 		FRHICommandListExecutor::GetImmediateCommandList().ImmediateFlush(EImmediateFlushType::DispatchToRHIThread); // we should start on the stuff before this async list
 	}
 #if PLATFORM_SUPPORTS_PARALLEL_RHI_EXECUTE
-	if (Num)
+	if (Num && GRHIThread)
 	{
 		IRHICommandContextContainer* ContextContainer = nullptr;
 		if (CVarRHICmdUseDeferredContexts.GetValueOnAnyThread() >= 1)
diff --git a/Engine/Source/Runtime/Renderer/Private/DeferredShadingRenderer.cpp b/Engine/Source/Runtime/Renderer/Private/DeferredShadingRenderer.cpp
--- a/Engine/Source/Runtime/Renderer/Private/DeferredShadingRenderer.cpp
+++ b/Engine/Source/Runtime/Renderer/Private/DeferredShadingRenderer.cpp
@@ -858,7 +858,6 @@ void FDeferredShadingSceneRenderer::Render(FRHICommandListImmediate& RHICmdList)
 	// This is done after the downsampled depth buffer is created so that it can be used for issuing queries
 	if ( bIsOcclusionTesting )
 	{
-		//FRHICommandListExecutor::GetImmediateCommandList().ImmediateFlush(EImmediateFlushType::FlushRHIThread);
 		BeginOcclusionTests(RHICmdList);
 	}
 
diff --git a/Engine/Source/Runtime/Renderer/Private/ShadowRendering.cpp b/Engine/Source/Runtime/Renderer/Private/ShadowRendering.cpp
--- a/Engine/Source/Runtime/Renderer/Private/ShadowRendering.cpp
+++ b/Engine/Source/Runtime/Renderer/Private/ShadowRendering.cpp
@@ -1504,14 +1504,14 @@ void FProjectedShadowInfo::SetStateForDepth(FRHICommandList& RHICmdList)
 
 static TAutoConsoleVariable<int32> CVarParallelShadows(
 	TEXT("r.ParallelShadows"),
-	1,
+	0,
 	TEXT("Toggles parallel shadow rendering. Parallel rendering must be enabled for this to have an effect.\n"),
 	ECVF_RenderThreadSafe
 	);
 
 static TAutoConsoleVariable<int32> CVarRHICmdShadowDeferredContexts(
 	TEXT("r.RHICmdShadowDeferredContexts"),
-	1,
+	0,
 	TEXT("True to use deferred contexts to parallelize shadow command list execution.\n"));
 
 class FShadowParallelCommandListSet : public FParallelCommandListSet
