Removal of GRHIShaderPlatform_DEPRECATED define.

diff --git a/Engine/Source/Runtime/RHI/Public/RHI.h b/Engine/Source/Runtime/RHI/Public/RHI.h
--- a/Engine/Source/Runtime/RHI/Public/RHI.h
+++ b/Engine/Source/Runtime/RHI/Public/RHI.h
@@ -238,14 +238,11 @@ extern RHI_API void GetFeatureLevelName(ERHIFeatureLevel::Type InFeatureLevel, F
 extern RHI_API void GetFeatureLevelName(ERHIFeatureLevel::Type InFeatureLevel, FName& OutName);
 
 // The maximum feature level and shader platform available on this system
+// GRHIFeatureLevel and GRHIShaderPlatform have been deprecated. There is no longer a current featurelevel/shaderplatform that
+// should be used for all rendering, rather a specific set for each view.
 extern RHI_API ERHIFeatureLevel::Type GMaxRHIFeatureLevel;
 extern RHI_API EShaderPlatform GMaxRHIShaderPlatform;
 
-// GRHIFeatureLevel and GRHIShaderPlatform have been deprecated. There is no longer a current featurelevel/shaderplatform that
-// should be used for all rendering, rather a specific set for each view. The deprecated names are being kept around temporarily 
-// while the last instances are removed.
-#define GRHIShaderPlatform_DEPRECATED GMaxRHIShaderPlatform
-
 /** Table for finding out which shader platform corresponds to a given feature level for this RHI. */
 extern RHI_API EShaderPlatform GShaderPlatformForFeatureLevel[ERHIFeatureLevel::Num];
 
