- Changed default Load/Store actions in FRHIRenderTargetView to Load/Store instead of NoAction

diff --git a/Engine/Source/Runtime/RHI/Public/RHIResources.h b/Engine/Source/Runtime/RHI/Public/RHIResources.h
--- a/Engine/Source/Runtime/RHI/Public/RHIResources.h
+++ b/Engine/Source/Runtime/RHI/Public/RHIResources.h
@@ -584,7 +584,7 @@ public:
 		Texture(NULL),
 		MipIndex(0),
 		ArraySliceIndex(-1),
-		LoadAction(ERenderTargetLoadAction::ENoAction),
+		LoadAction(ERenderTargetLoadAction::ELoad),
 		StoreAction(ERenderTargetStoreAction::EStore)
 	{}
 
@@ -600,7 +600,7 @@ public:
 		Texture(InTexture),
 		MipIndex(0),
 		ArraySliceIndex(-1),
-		LoadAction(ERenderTargetLoadAction::ENoAction),
+		LoadAction(ERenderTargetLoadAction::ELoad),
 		StoreAction(ERenderTargetStoreAction::EStore)
 	{}
 
@@ -608,7 +608,7 @@ public:
 		Texture(InTexture),
 		MipIndex(InMipIndex),
 		ArraySliceIndex(InArraySliceIndex),
-		LoadAction(ERenderTargetLoadAction::ENoAction),
+		LoadAction(ERenderTargetLoadAction::ELoad),
 		StoreAction(ERenderTargetStoreAction::EStore)
 	{}
 	
@@ -644,17 +644,17 @@ public:
 	FRHIDepthRenderTargetView() :
 		Texture(nullptr),
 		DepthLoadAction(ERenderTargetLoadAction::EClear),
-		DepthStoreAction(ERenderTargetStoreAction::ENoAction),
+		DepthStoreAction(ERenderTargetStoreAction::EStore),
 		StencilLoadAction(ERenderTargetLoadAction::EClear),
-		StencilStoreAction(ERenderTargetStoreAction::ENoAction)
+		StencilStoreAction(ERenderTargetStoreAction::EStore)
 	{}
 
 	FRHIDepthRenderTargetView(FTextureRHIParamRef InTexture) :
 		Texture(InTexture),
 		DepthLoadAction(ERenderTargetLoadAction::EClear),
-		DepthStoreAction(ERenderTargetStoreAction::ENoAction),
+		DepthStoreAction(ERenderTargetStoreAction::EStore),
 		StencilLoadAction(ERenderTargetLoadAction::EClear),
-		StencilStoreAction(ERenderTargetStoreAction::ENoAction)
+		StencilStoreAction(ERenderTargetStoreAction::EStore)
 	{}
 
 	FRHIDepthRenderTargetView(FTextureRHIParamRef InTexture, ERenderTargetLoadAction InLoadAction, ERenderTargetStoreAction InStoreAction) :
