New simpler occlusion query system for Xbox One which should be more robust with parallel rendering enabled.

diff --git a/Engine/Source/Runtime/Core/Public/HAL/Platform.h b/Engine/Source/Runtime/Core/Public/HAL/Platform.h
--- a/Engine/Source/Runtime/Core/Public/HAL/Platform.h
+++ b/Engine/Source/Runtime/Core/Public/HAL/Platform.h
@@ -261,10 +261,6 @@
 	#define PLATFORM_SUPPORTS_RHI_THREAD		0
 #endif
 
-#ifndef PLATFORM_HAS_THREADSAFE_RHIGetRenderQueryResult
-	#define PLATFORM_HAS_THREADSAFE_RHIGetRenderQueryResult		0
-#endif
-
 #ifndef PLATFORM_USES_FIXED_RHI_CLASS
 	#define PLATFORM_USES_FIXED_RHI_CLASS		0
 #endif
diff --git a/Engine/Source/Runtime/RHI/Private/RHICommandList.cpp b/Engine/Source/Runtime/RHI/Private/RHICommandList.cpp
--- a/Engine/Source/Runtime/RHI/Private/RHICommandList.cpp
+++ b/Engine/Source/Runtime/RHI/Private/RHICommandList.cpp
@@ -17,7 +17,7 @@ DECLARE_DWORD_COUNTER_STAT(TEXT("Immed. Command count"), STAT_ImmedCmdListCount,
 Requirements for RHI thread
 * Microresources (those in RHIStaticStates.h) need to be able to be created by any thread at any time and be able to work with a radically simplified rhi resource lifecycle. CreateSamplerState, CreateRasterizerState, CreateDepthStencilState, CreateBlendState
 * CreateUniformBuffer needs to be threadsafe
-* GetRenderQueryResult should be threadsafe, if not, then PLATFORM_HAS_THREADSAFE_RHIGetRenderQueryResult should be 1
+* GetRenderQueryResult should be threadsafe
 * AdvanceFrameForGetViewportBackBuffer needs be added as an RHI method and this needs to work with GetViewportBackBuffer to give the render thread the right back buffer even though many commands relating to the beginning and end of the frame are queued.
 * ResetRenderQuery does not exist; this stuff should be done in BeginQuery
 
diff --git a/Engine/Source/Runtime/RHI/Public/RHIMethods.h b/Engine/Source/Runtime/RHI/Public/RHIMethods.h
--- a/Engine/Source/Runtime/RHI/Public/RHIMethods.h
+++ b/Engine/Source/Runtime/RHI/Public/RHIMethods.h
@@ -1310,23 +1310,13 @@ DEFINE_RHIMETHOD_CMDLIST_1(
 	FRenderQueryRHIParamRef,RenderQuery,
 	,
 	);
-#if PLATFORM_HAS_THREADSAFE_RHIGetRenderQueryResult
-	DEFINE_RHIMETHOD_GLOBALTHREADSAFE_3(
-		bool,GetRenderQueryResult,
-		FRenderQueryRHIParamRef,RenderQuery,
-		uint64&,OutResult,
-		bool,bWait,
-		return,return true;
-	);
-#else
-	DEFINE_RHIMETHOD_3(
-		bool,GetRenderQueryResult,
-		FRenderQueryRHIParamRef,RenderQuery,
-		uint64&,OutResult,
-		bool,bWait,
-		return,return true;
+DEFINE_RHIMETHOD_3(
+	bool,GetRenderQueryResult,
+	FRenderQueryRHIParamRef,RenderQuery,
+	uint64&,OutResult,
+	bool,bWait,
+	return,return true;
 	);
-#endif
 
 #if PLATFORM_SUPPORTS_RHI_THREAD
 	DEFINE_RHIMETHOD_CMDLIST_2(
diff --git a/Engine/Source/Runtime/Windows/D3D11RHI/Private/D3D11Query.cpp b/Engine/Source/Runtime/Windows/D3D11RHI/Private/D3D11Query.cpp
--- a/Engine/Source/Runtime/Windows/D3D11RHI/Private/D3D11Query.cpp
+++ b/Engine/Source/Runtime/Windows/D3D11RHI/Private/D3D11Query.cpp
@@ -91,8 +91,6 @@ void FD3D11DynamicRHI::RHIEndRenderQuery(FRenderQueryRHIParamRef QueryRHI)
 	DYNAMIC_CAST_D3D11RESOURCE(OcclusionQuery,Query);
 #if PLATFORM_SUPPORTS_RHI_THREAD
 	Query->bResultIsCached = false; // for occlusion queries, this is redundant with the one in begin
-#elif PLATFORM_HAS_THREADSAFE_RHIGetRenderQueryResult
-	#error "PLATFORM_HAS_THREADSAFE_RHIGetRenderQueryResult requires PLATFORM_SUPPORTS_RHI_THREAD"
 #endif
 	Direct3DDeviceIMContext->End(Query->Resource);
 
