Create context class for ps4.
First pass at parallel rhi execution for ps4.  Must manually define PLATFORM_SUPPORTS_PARALLEL_RHI_EXECUTE in PS4Platform.h
Parallel RHI execution does not currently support CMASK and HTILE and needs optimizations for memory usage, but should function for most projects.

diff --git a/Engine/Source/Runtime/RHI/Public/RHIMethods.h b/Engine/Source/Runtime/RHI/Public/RHIMethods.h
--- a/Engine/Source/Runtime/RHI/Public/RHIMethods.h
+++ b/Engine/Source/Runtime/RHI/Public/RHIMethods.h
@@ -549,7 +549,7 @@ DEFINE_RHIMETHOD_CMDLIST_0(
 // Useful when used with geometry shader (emit polygons to different viewports), otherwise SetViewPort() is simpler
 // @param Count >0
 // @param Data must not be 0
-DEFINE_RHIMETHOD_2(
+DEFINE_RHIMETHOD_CMDLIST_2(
 	void,SetMultipleViewports,
 	uint32,Count,
 	const FViewportBounds*,Data,
@@ -987,7 +987,7 @@ DEFINE_RHIMETHOD_GLOBALTHREADSAFE_1(
  * @param RequestStatus	- Will be decremented by 1 when the reallocation is complete (success or failure).
  * @return				- New reference to the texture, or an invalid reference upon failure
  */
-DEFINE_RHIMETHOD_GLOBAL_5(
+DEFINE_RHIMETHOD_GLOBALFLUSH_5(
 	FTexture2DRHIRef,AsyncReallocateTexture2D,
 	FTexture2DRHIParamRef,Texture2D,
 	int32,NewMipCount,
