UE4 - Fix RHI clear color check() range

diff --git a/Engine/Source/Runtime/RHI/Public/RHICommandList.h b/Engine/Source/Runtime/RHI/Public/RHICommandList.h
--- a/Engine/Source/Runtime/RHI/Public/RHICommandList.h
+++ b/Engine/Source/Runtime/RHI/Public/RHICommandList.h
@@ -805,7 +805,7 @@ struct FRHICommandClearMRT : public FRHICommand<FRHICommandClearMRT>
 		, bClearDepth(InbClearDepth)
 		, bClearStencil(InbClearStencil)
 	{
-		check(InNumClearColors < MaxSimultaneousRenderTargets);
+		check(InNumClearColors <= MaxSimultaneousRenderTargets);
 		for (int32 Index = 0; Index < InNumClearColors; Index++)
 		{
 			ColorArray[Index] = InColorArray[Index];
