Optimize new Xbox One query system to use a single GPU fence per batch, rather than a fence per query.

diff --git a/Engine/Source/Runtime/OpenGLDrv/Private/OpenGLQuery.cpp b/Engine/Source/Runtime/OpenGLDrv/Private/OpenGLQuery.cpp
--- a/Engine/Source/Runtime/OpenGLDrv/Private/OpenGLQuery.cpp
+++ b/Engine/Source/Runtime/OpenGLDrv/Private/OpenGLQuery.cpp
@@ -10,6 +10,14 @@
 // Ignore functions from RHIMethods.h when parsing documentation; Doxygen's preprocessor can't parse the declaration, so spews warnings for the definitions.
 #if !UE_BUILD_DOCS
 
+void FOpenGLDynamicRHI::RHIBeginOcclusionQueryBatch()
+{
+}
+
+void FOpenGLDynamicRHI::RHIEndOcclusionQueryBatch()
+{
+}
+
 FRenderQueryRHIRef FOpenGLDynamicRHI::RHICreateRenderQuery(ERenderQueryType QueryType)
 {
 	VERIFY_GL_SCOPE();
diff --git a/Engine/Source/Runtime/RHI/Public/RHICommandList.h b/Engine/Source/Runtime/RHI/Public/RHICommandList.h
--- a/Engine/Source/Runtime/RHI/Public/RHICommandList.h
+++ b/Engine/Source/Runtime/RHI/Public/RHICommandList.h
@@ -1054,6 +1054,22 @@ struct FRHICommandResetRenderQuery : public FRHICommand<FRHICommandResetRenderQu
 };
 #endif
 
+struct FRHICommandBeginOcclusionQueryBatch : public FRHICommand<FRHICommandBeginOcclusionQueryBatch>
+{
+	FORCEINLINE_DEBUGGABLE FRHICommandBeginOcclusionQueryBatch()
+	{
+	}
+	RHI_API void Execute(FRHICommandListBase& CmdList);
+};
+
+struct FRHICommandEndOcclusionQueryBatch : public FRHICommand<FRHICommandEndOcclusionQueryBatch>
+{
+	FORCEINLINE_DEBUGGABLE FRHICommandEndOcclusionQueryBatch()
+	{
+	}
+	RHI_API void Execute(FRHICommandListBase& CmdList);
+};
+
 struct FRHICommandBeginScene : public FRHICommand<FRHICommandBeginScene>
 {
 	FORCEINLINE_DEBUGGABLE FRHICommandBeginScene()
@@ -1778,6 +1794,26 @@ public:
 	}
 #endif
 
+	FORCEINLINE_DEBUGGABLE void BeginOcclusionQueryBatch()
+	{
+		if (Bypass())
+		{
+			CMD_CONTEXT(BeginOcclusionQueryBatch)();
+			return;
+		}
+		new (AllocCommand<FRHICommandBeginOcclusionQueryBatch>()) FRHICommandBeginOcclusionQueryBatch();
+	}
+
+	FORCEINLINE_DEBUGGABLE void EndOcclusionQueryBatch()
+	{
+		if (Bypass())
+		{
+			CMD_CONTEXT(EndOcclusionQueryBatch)();
+			return;
+		}
+		new (AllocCommand<FRHICommandEndOcclusionQueryBatch>()) FRHICommandEndOcclusionQueryBatch();
+	}
+
 #if PLATFORM_SUPPORTS_RHI_THREAD
 	FORCEINLINE_DEBUGGABLE void BeginScene()
 	{
diff --git a/Engine/Source/Runtime/RHI/Public/RHICommandListCommandExecutes.inl b/Engine/Source/Runtime/RHI/Public/RHICommandListCommandExecutes.inl
--- a/Engine/Source/Runtime/RHI/Public/RHICommandListCommandExecutes.inl
+++ b/Engine/Source/Runtime/RHI/Public/RHICommandListCommandExecutes.inl
@@ -363,6 +363,18 @@ void FRHICommandResetRenderQuery::Execute(FRHICommandListBase& CmdList)
 }
 #endif
 
+void FRHICommandBeginOcclusionQueryBatch::Execute(FRHICommandListBase& CmdList)
+{
+	RHISTAT(BeginOcclusionQueryBatch);
+	INTERNAL_DECORATOR(BeginOcclusionQueryBatch)();
+}
+
+void FRHICommandEndOcclusionQueryBatch::Execute(FRHICommandListBase& CmdList)
+{
+	RHISTAT(EndOcclusionQueryBatch);
+	INTERNAL_DECORATOR(EndOcclusionQueryBatch)();
+}
+
 void FRHICommandBeginScene::Execute(FRHICommandListBase& CmdList)
 {
 	RHISTAT(BeginScene);
diff --git a/Engine/Source/Runtime/RHI/Public/RHIMethods.h b/Engine/Source/Runtime/RHI/Public/RHIMethods.h
--- a/Engine/Source/Runtime/RHI/Public/RHIMethods.h
+++ b/Engine/Source/Runtime/RHI/Public/RHIMethods.h
@@ -1310,12 +1310,31 @@ DEFINE_RHIMETHOD_CMDLIST_1(
 	FRenderQueryRHIParamRef,RenderQuery,
 	,
 	);
-DEFINE_RHIMETHOD_3(
-	bool,GetRenderQueryResult,
-	FRenderQueryRHIParamRef,RenderQuery,
-	uint64&,OutResult,
-	bool,bWait,
-	return,return true;
+#if PLATFORM_HAS_THREADSAFE_RHIGetRenderQueryResult
+	DEFINE_RHIMETHOD_GLOBALTHREADSAFE_3(
+		bool,GetRenderQueryResult,
+		FRenderQueryRHIParamRef,RenderQuery,
+		uint64&,OutResult,
+		bool,bWait,
+		return,return true;
+	);
+#else
+	DEFINE_RHIMETHOD_3(
+		bool,GetRenderQueryResult,
+		FRenderQueryRHIParamRef,RenderQuery,
+		uint64&,OutResult,
+		bool,bWait,
+		return,return true;
+	);
+#endif
+
+DEFINE_RHIMETHOD_CMDLIST_0(
+	void,BeginOcclusionQueryBatch,
+	,
+	);
+DEFINE_RHIMETHOD_CMDLIST_0(
+	void,EndOcclusionQueryBatch,
+	,
 	);
 
 #if PLATFORM_SUPPORTS_RHI_THREAD
diff --git a/Engine/Source/Runtime/Renderer/Private/SceneOcclusion.cpp b/Engine/Source/Runtime/Renderer/Private/SceneOcclusion.cpp
--- a/Engine/Source/Runtime/Renderer/Private/SceneOcclusion.cpp
+++ b/Engine/Source/Runtime/Renderer/Private/SceneOcclusion.cpp
@@ -1038,6 +1038,8 @@ void FDeferredShadingSceneRenderer::BeginOcclusionTests(FRHICommandListImmediate
 	SCOPE_CYCLE_COUNTER(STAT_BeginOcclusionTestsTime);
 	const bool bUseDownsampledDepth = IsValidRef(GSceneRenderTargets.GetSmallDepthSurface()) && GSceneRenderTargets.UseDownsizedOcclusionQueries();
 
+	RHICmdList.BeginOcclusionQueryBatch();
+
 	if (bUseDownsampledDepth)
 	{
 		SetRenderTarget(RHICmdList, NULL, GSceneRenderTargets.GetSmallDepthSurface());
@@ -1187,6 +1189,8 @@ void FDeferredShadingSceneRenderer::BeginOcclusionTests(FRHICommandListImmediate
 		}
 	}
 
+	RHICmdList.EndOcclusionQueryBatch();
+
 	for( int32 ViewIndex = 0; ViewIndex < Views.Num(); ViewIndex++ )
 	{
 		FViewInfo& View = Views[ViewIndex];
diff --git a/Engine/Source/Runtime/Windows/D3D11RHI/Private/D3D11Query.cpp b/Engine/Source/Runtime/Windows/D3D11RHI/Private/D3D11Query.cpp
--- a/Engine/Source/Runtime/Windows/D3D11RHI/Private/D3D11Query.cpp
+++ b/Engine/Source/Runtime/Windows/D3D11RHI/Private/D3D11Query.cpp
@@ -6,6 +6,13 @@
 
 #include "D3D11RHIPrivate.h"
 
+void FD3D11DynamicRHI::RHIBeginOcclusionQueryBatch()
+{
+}
+
+void FD3D11DynamicRHI::RHIEndOcclusionQueryBatch()
+{
+}
 
 FRenderQueryRHIRef FD3D11DynamicRHI::RHICreateRenderQuery(ERenderQueryType QueryType)
 {
