GPU Profiler: Display a warning if GPU profile is captured with v-sync enabled.

diff --git a/Engine/Source/Developer/TaskGraph/Private/SProfileVisualizer.cpp b/Engine/Source/Developer/TaskGraph/Private/SProfileVisualizer.cpp
--- a/Engine/Source/Developer/TaskGraph/Private/SProfileVisualizer.cpp
+++ b/Engine/Source/Developer/TaskGraph/Private/SProfileVisualizer.cpp
@@ -8,44 +8,63 @@
 #include "SGraphBar.h"
 #include "SBarVisualizer.h"
 #include "SEventsTree.h"
-
 	
 void SProfileVisualizer::Construct( const FArguments& InArgs )
 {
 	ProfileData = InArgs._ProfileData;
 	ProfilerType = InArgs._ProfilerType;
+	HeaderMessageText = InArgs._HeaderMessageText;
+	HeaderMessageTextColor = InArgs._HeaderMessageTextColor;
 
 	const FSlateBrush* ContentAreaBrush = FEditorStyle::GetBrush( "Docking.Tab", ".ContentAreaBrush" );
 
 	ChildSlot
 	[
-		SAssignNew( MainSplitter, SSplitter )
-		.Orientation( Orient_Vertical )
-		+ SSplitter::Slot()
-		.Value(1.0f)
+		SNew(SVerticalBox)
+		+ SVerticalBox::Slot()
+		.AutoHeight()
 		[
-			SNew( SBorder )
-			.Visibility( EVisibility::Visible )
-			.BorderImage( ContentAreaBrush )
+			SNew(SBorder)
+			.Visibility(EVisibility::Visible)
+			.BorderImage(ContentAreaBrush)
 			[
-				SAssignNew( BarVisualizer, SBarVisualizer )
-				.ProfileData( ProfileData )
-				.OnBarGraphSelectionChanged( this, &SProfileVisualizer::RouteBarGraphSelectionChanged )
-				.OnBarGraphExpansionChanged( this, &SProfileVisualizer::RouteBarGraphExpansionChanged )
-				.OnBarEventSelectionChanged( this, &SProfileVisualizer::RouteBarEventSelectionChanged )
-				.OnBarGraphContextMenu( this, &SProfileVisualizer::OnBarGraphContextMenu )
+				SNew(STextBlock)
+				.Visibility( HeaderMessageText.IsEmpty() ? EVisibility::Collapsed : EVisibility::Visible )
+				.Text(HeaderMessageText)
+				.ColorAndOpacity(HeaderMessageTextColor)
+				.AutoWrapText(true)
 			]
 		]
-		+ SSplitter::Slot()
-		.Value(1.0f)
+		+ SVerticalBox::Slot()
 		[
-			SNew( SBorder )
-			.Visibility( EVisibility::Visible )
-			.BorderImage( ContentAreaBrush )
+			SAssignNew( MainSplitter, SSplitter )
+			.Orientation( Orient_Vertical )
+			+ SSplitter::Slot()
+			.Value(1.0f)
+			[
+				SNew( SBorder )
+				.Visibility( EVisibility::Visible )
+				.BorderImage( ContentAreaBrush )
+				[
+					SAssignNew( BarVisualizer, SBarVisualizer )
+					.ProfileData( ProfileData )
+					.OnBarGraphSelectionChanged( this, &SProfileVisualizer::RouteBarGraphSelectionChanged )
+					.OnBarGraphExpansionChanged( this, &SProfileVisualizer::RouteBarGraphExpansionChanged )
+					.OnBarEventSelectionChanged( this, &SProfileVisualizer::RouteBarEventSelectionChanged )
+					.OnBarGraphContextMenu( this, &SProfileVisualizer::OnBarGraphContextMenu )
+				]
+			]
+			+ SSplitter::Slot()
+			.Value(1.0f)
 			[
-				SAssignNew( EventsTree, SEventsTree )
-				.ProfileData( ProfileData )
-				.OnEventSelectionChanged( this, &SProfileVisualizer::RouteEventSelectionChanged )
+				SNew( SBorder )
+				.Visibility( EVisibility::Visible )
+				.BorderImage( ContentAreaBrush )
+				[
+					SAssignNew( EventsTree, SEventsTree )
+					.ProfileData( ProfileData )
+					.OnEventSelectionChanged( this, &SProfileVisualizer::RouteEventSelectionChanged )
+				]
 			]
 		]
 	];
diff --git a/Engine/Source/Developer/TaskGraph/Private/SProfileVisualizer.h b/Engine/Source/Developer/TaskGraph/Private/SProfileVisualizer.h
--- a/Engine/Source/Developer/TaskGraph/Private/SProfileVisualizer.h
+++ b/Engine/Source/Developer/TaskGraph/Private/SProfileVisualizer.h
@@ -21,6 +21,10 @@ public:
 
 		SLATE_ARGUMENT( FText, ProfilerType )
 
+		SLATE_ARGUMENT( FText, HeaderMessageText )
+
+		SLATE_ARGUMENT( FLinearColor, HeaderMessageTextColor )
+
 	SLATE_END_ARGS()
 
 	
@@ -83,5 +87,11 @@ private:
 
 	/** Profiler name */
 	FText ProfilerType;
+
+	/** Optional header message to display at the top of the profile window */
+	FText HeaderMessageText;
+	 
+	/** Optional header message text color */
+	FLinearColor HeaderMessageTextColor;
 };
 
diff --git a/Engine/Source/Developer/TaskGraph/Private/STaskGraph.cpp b/Engine/Source/Developer/TaskGraph/Private/STaskGraph.cpp
--- a/Engine/Source/Developer/TaskGraph/Private/STaskGraph.cpp
+++ b/Engine/Source/Developer/TaskGraph/Private/STaskGraph.cpp
@@ -15,7 +15,7 @@
  * @param ProfileData Visualizer data
  * @return Visualizer window
  */
-void MakeTaskGraphVisualizerWindow( TSharedPtr< FVisualizerEvent > ProfileData, const FText& WindowTitle, const FText& ProfilerType )
+void MakeTaskGraphVisualizerWindow( TSharedPtr< FVisualizerEvent > ProfileData, const FText& WindowTitle, const FText& ProfilerType, const FText& HeaderMessageText = FText::GetEmpty(), const FLinearColor& HeaderMessageTextColor = FLinearColor::White )
 {
 	FGlobalTabmanager::Get()->InsertNewDocumentTab
 		(
@@ -27,6 +27,8 @@ void MakeTaskGraphVisualizerWindow( TSharedPtr< FVisualizerEvent > ProfileData,
 				SNew( SProfileVisualizer )
 				.ProfileData( ProfileData )
 				.ProfilerType( ProfilerType )
+				.HeaderMessageText( HeaderMessageText )
+				.HeaderMessageTextColor( HeaderMessageTextColor )
 			]
 		);
 }
@@ -125,7 +127,7 @@ void ShutdownProfileVisualizer()
 
 static bool GHasRegisteredVisualizerLayout = false;
 
-void DisplayProfileVisualizer( TSharedPtr< FVisualizerEvent > InProfileData, const TCHAR* InProfilerType )
+void DisplayProfileVisualizer(TSharedPtr< FVisualizerEvent > InProfileData, const TCHAR* InProfilerType, const FText& HeaderMessageText = FText::GetEmpty(), const FLinearColor& HeaderMessageTextColor = FLinearColor::White)
 {
 	check( IsInGameThread() );
 
@@ -152,8 +154,7 @@ void DisplayProfileVisualizer( TSharedPtr< FVisualizerEvent > InProfileData, con
 	const FText WindowTitle = FText::Format( NSLOCTEXT("TaskGraph", "WindowTitle", "{ProfilerType} Visualizer"), Args );
 	const FText ProfilerType = FText::Format( NSLOCTEXT("TaskGraph", "ProfilerType", "{ProfilerType} Profile"), Args );
 
-
-	MakeTaskGraphVisualizerWindow( InProfileData, WindowTitle, ProfilerType );
+	MakeTaskGraphVisualizerWindow( InProfileData, WindowTitle, ProfilerType, HeaderMessageText, HeaderMessageTextColor );
 	
 }
 
@@ -171,10 +172,10 @@ public:
 	{
 		::ShutdownProfileVisualizer();
 	}
-	virtual void DisplayProfileVisualizer( TSharedPtr< FVisualizerEvent > InProfileData, const TCHAR* InProfilerType ) override
+	virtual void DisplayProfileVisualizer(TSharedPtr< FVisualizerEvent > InProfileData, const TCHAR* InProfilerType, const FText& HeaderMessageText, const FLinearColor& HeaderMessageTextColor) override
 	{
 #if	WITH_EDITOR
-		::DisplayProfileVisualizer( InProfileData, InProfilerType );
+		::DisplayProfileVisualizer( InProfileData, InProfilerType, HeaderMessageText, HeaderMessageTextColor );
 #endif // WITH_EDITOR
 	}
 };
diff --git a/Engine/Source/Developer/TaskGraph/Public/STaskGraph.h b/Engine/Source/Developer/TaskGraph/Public/STaskGraph.h
--- a/Engine/Source/Developer/TaskGraph/Public/STaskGraph.h
+++ b/Engine/Source/Developer/TaskGraph/Public/STaskGraph.h
@@ -8,5 +8,5 @@ class IProfileVisualizerModule
 	: public IModuleInterface
 {
 public:
-	virtual void DisplayProfileVisualizer( TSharedPtr< FVisualizerEvent > InProfileData, const TCHAR* InProfilerType ) = 0;
+	virtual void DisplayProfileVisualizer( TSharedPtr< FVisualizerEvent > InProfileData, const TCHAR* InProfilerType, const FText& HeaderMessageText = FText::GetEmpty(), const FLinearColor& HeaderMessageTextColor = FLinearColor::White ) = 0;
 };
diff --git a/Engine/Source/Runtime/RHI/Private/GPUProfiler.cpp b/Engine/Source/Runtime/RHI/Private/GPUProfiler.cpp
--- a/Engine/Source/Runtime/RHI/Private/GPUProfiler.cpp
+++ b/Engine/Source/Runtime/RHI/Private/GPUProfiler.cpp
@@ -13,6 +13,8 @@
 #include "TaskGraphInterfaces.h"
 #endif
 
+#define LOCTEXT_NAMESPACE "GpuProfiler"
+
 static TAutoConsoleVariable<FString> GProfileGPUPatternCVar(
 	TEXT("r.ProfileGPU.Pattern"),
 	TEXT("*"),
@@ -189,6 +191,16 @@ void FGPUProfilerEventNodeFrame::DumpEventTree()
 
 		UE_LOG(LogRHI, Warning, TEXT("Perf marker hierarchy, total GPU time %.2fms"), RootResult * 1000.0f);
 
+		// Display a warning if this is a GPU profile and the GPU was profiled with v-sync enabled
+		FText VsyncEnabledWarningText = FText::GetEmpty();
+		static IConsoleVariable* CVSyncVar = IConsoleManager::Get().FindConsoleVariable(TEXT("r.VSync"));
+		if (CVSyncVar->GetInt() != 0)
+		{
+			VsyncEnabledWarningText = LOCTEXT("GpuProfileVsyncEnabledWarning",
+				"WARNING: This GPU profile was captured with v-sync enabled.  V-sync wait time may show up in any bucket, and as a result the data in this profile may be skewed. Please profile with v-sync disabled to obtain the most accurate data.");
+			UE_LOG(LogRHI, Warning, TEXT("%s"), *(VsyncEnabledWarningText.ToString()));
+		}
+
 		LogDisjointQuery();
 
 		TMap<FString, FGPUProfilerEventNodeStats> EventHistogram;
@@ -283,13 +295,14 @@ void FGPUProfilerEventNodeFrame::DumpEventTree()
 			{
 				struct FDisplayProfilerVisualizer
 				{
-					void Thread( TSharedPtr<FVisualizerEvent> InVisualizerData )
+					void Thread( TSharedPtr<FVisualizerEvent> InVisualizerData, const FText InVsyncEnabledWarningText )
 					{
 						static FName TaskGraphModule(TEXT("TaskGraph"));			
 						if (FModuleManager::Get().IsModuleLoaded(TaskGraphModule))
 						{
 							IProfileVisualizerModule& ProfileVisualizer = FModuleManager::GetModuleChecked<IProfileVisualizerModule>(TaskGraphModule);
-							ProfileVisualizer.DisplayProfileVisualizer( InVisualizerData, TEXT("GPU") );
+							// Display a warning if this is a GPU profile and the GPU was profiled with v-sync enabled (otherwise InVsyncEnabledWarningText is empty)
+							ProfileVisualizer.DisplayProfileVisualizer( InVisualizerData, TEXT("GPU"), InVsyncEnabledWarningText, FLinearColor::Red );
 						}
 					}
 				} DisplayProfilerVisualizer;
@@ -301,7 +314,7 @@ void FGPUProfilerEventNodeFrame::DumpEventTree()
 					STATGROUP_TaskGraphTasks);
 
 				FSimpleDelegateGraphTask::CreateAndDispatchWhenReady(
-					FSimpleDelegateGraphTask::FDelegate::CreateRaw(&DisplayProfilerVisualizer, &FDisplayProfilerVisualizer::Thread, VisualizerData),
+					FSimpleDelegateGraphTask::FDelegate::CreateRaw(&DisplayProfilerVisualizer, &FDisplayProfilerVisualizer::Thread, VisualizerData, VsyncEnabledWarningText),
 					GET_STATID(STAT_FSimpleDelegateGraphTask_DisplayProfilerVisualizer), nullptr, ENamedThreads::GameThread
 				);
 			}
@@ -355,3 +368,5 @@ uint64 FGPUTiming::GTimingFrequency = 0;
 
 /** Whether the static variables have been initialized. */
 bool FGPUTiming::GAreGlobalsInitialized = false;
+
+#undef LOCTEXT_NAMESPACE
\ No newline at end of file
