Fix flickering from corrupt uniform buffers on PS4.

diff --git a/Engine/Source/Runtime/RHI/Private/RHICommandList.cpp b/Engine/Source/Runtime/RHI/Private/RHICommandList.cpp
--- a/Engine/Source/Runtime/RHI/Private/RHICommandList.cpp
+++ b/Engine/Source/Runtime/RHI/Private/RHICommandList.cpp
@@ -81,6 +81,7 @@ struct FRHICommandStat : public FRHICommand<FRHICommandStat>
 	}
 };
 
+bool FRHICommandListImmediate::bFlushedGlobal = false;
 void FRHICommandListImmediate::SetCurrentStat(TStatId Stat)
 {
 	new (AllocCommand<FRHICommandStat>()) FRHICommandStat(Stat);
diff --git a/Engine/Source/Runtime/RHI/Public/RHICommandList.h b/Engine/Source/Runtime/RHI/Public/RHICommandList.h
--- a/Engine/Source/Runtime/RHI/Public/RHICommandList.h
+++ b/Engine/Source/Runtime/RHI/Public/RHICommandList.h
@@ -1900,12 +1900,20 @@ class RHI_API FRHICommandListImmediate : public FRHICommandList
 	{
 		check(!HasCommands());
 	}
+
+	static bool bFlushedGlobal;
+
 public:
 
 	inline void ImmediateFlush(EImmediateFlushType::Type FlushType);
 
 	void SetCurrentStat(TStatId Stat);
 
+	static bool IsFullyFlushed()
+	{
+		return bFlushedGlobal;
+	};
+
 	#define DEFINE_RHIMETHOD_CMDLIST(Type,Name,ParameterTypesAndNames,ParameterNames,ReturnStatement,NullImplementation)
 	#define DEFINE_RHIMETHOD(Type, Name, ParameterTypesAndNames, ParameterNames, ReturnStatement, NullImplementation) \
 		FORCEINLINE_DEBUGGABLE Type Name ParameterTypesAndNames \
@@ -1934,7 +1942,9 @@ public:
 		{ \
 			QUICK_SCOPE_CYCLE_COUNTER(STAT_RHIMETHOD_##Name##_Flush); \
 			ImmediateFlush(EImmediateFlushType::FlushRHIThread); \
+			bFlushedGlobal = true;\
 			ReturnStatement Name##_Internal ParameterNames; \
+			bFlushedGlobal = false;\
 		}
 	#include "RHIMethods.h"
 	#undef DEFINE_RHIMETHOD
