comments on cvars and make them globals so autocomplete and help can find them

diff --git a/Engine/Source/Runtime/RHI/Private/RHICommandList.cpp b/Engine/Source/Runtime/RHI/Private/RHICommandList.cpp
--- a/Engine/Source/Runtime/RHI/Private/RHICommandList.cpp
+++ b/Engine/Source/Runtime/RHI/Private/RHICommandList.cpp
@@ -261,20 +261,21 @@ public:
 	}
 };
 
+static TAutoConsoleVariable<int32> CVarRHICmdUseThread(
+	TEXT("r.RHICmdUseThread"),
+	1,
+	TEXT("Uses the RHI thread. Just for testing, should be removed. (GilG)\n")
+	TEXT("0: Disable, 1: Enable"),
+	ECVF_Cheat);
+static TAutoConsoleVariable<int32> CVarRHICmdForceRHIFlush(
+	TEXT("r.RHICmdForceRHIFlush"),
+	0,
+	TEXT("Force a flush for every task sent to the RHI thread. Just for testing, should be removed. (GilG)\n")
+	TEXT("0: Disable, 1: Enable"),
+	ECVF_Cheat);
+
 void FRHICommandListExecutor::ExecuteInner(FRHICommandListBase& CmdList)
 {
-	static TAutoConsoleVariable<int32> CVarRHICmdUseThread(
-		TEXT("r.RHICmdUseThread"),
-		1,
-		TEXT("Uses the RHI thread.\n")
-		TEXT("0: Disable, 1: Enable"),
-		ECVF_Cheat);
-	static TAutoConsoleVariable<int32> CVarRHICmdForceRHIFlush(
-		TEXT("r.RHICmdForceRHIFlush"),
-		0,
-		TEXT("Force a flush for every task sent to the RHI thread.\n")
-		TEXT("0: Disable, 1: Enable"),
-		ECVF_Cheat);
 	check(CmdList.RTTasks.Num() == 0 && CmdList.HasCommands()); 
 
 	bool bIsInRenderingThread = IsInRenderingThread();
