UE4 - remove several ifdefs related to parallel rendering making it available for DX12 and any PC RHI

diff --git a/Engine/Source/Runtime/Core/Public/Async/TaskGraphInterfaces.h b/Engine/Source/Runtime/Core/Public/Async/TaskGraphInterfaces.h
--- a/Engine/Source/Runtime/Core/Public/Async/TaskGraphInterfaces.h
+++ b/Engine/Source/Runtime/Core/Public/Async/TaskGraphInterfaces.h
@@ -35,9 +35,7 @@ namespace ENamedThreads
 #if STATS
 		StatsThread, 
 #endif
-#if PLATFORM_SUPPORTS_RHI_THREAD
 		RHIThread,
-#endif
 		GameThread,
 		// The render thread is sometimes the game thread and is sometimes the actual rendering thread
 		ActualRenderingThread = GameThread + 1,
diff --git a/Engine/Source/Runtime/Core/Public/HAL/Platform.h b/Engine/Source/Runtime/Core/Public/HAL/Platform.h
--- a/Engine/Source/Runtime/Core/Public/HAL/Platform.h
+++ b/Engine/Source/Runtime/Core/Public/HAL/Platform.h
@@ -253,14 +253,6 @@
 	#define PLATFORM_SUPPORTS_NAMED_PIPES		0
 #endif
 
-#ifndef PLATFORM_SUPPORTS_RHI_THREAD
-	#define PLATFORM_SUPPORTS_RHI_THREAD		0
-#endif
-
-#ifndef PLATFORM_HAS_THREADSAFE_RHIGetRenderQueryResult
-	#define PLATFORM_HAS_THREADSAFE_RHIGetRenderQueryResult		0
-#endif
-
 #ifndef PLATFORM_USES_FIXED_RHI_CLASS
 	#define PLATFORM_USES_FIXED_RHI_CLASS		0
 #endif
@@ -281,18 +273,12 @@
 	#define PLATFORM_USES_ANSI_STRING_FOR_EXTERNAL_PROFILING 1
 #endif
 
-#ifndef PLATFORM_RHI_USES_CONTEXT_OBJECT
-	#define PLATFORM_RHI_USES_CONTEXT_OBJECT 0
-#endif
+// deprecated, do not use
+#define PLATFORM_HAS_THREADSAFE_RHIGetRenderQueryResult	#
+#define PLATFORM_SUPPORTS_RHI_THREAD #
+#define PLATFORM_RHI_USES_CONTEXT_OBJECT # // deprecated, do not use; all platforms must use a context object
+#define PLATFORM_SUPPORTS_PARALLEL_RHI_EXECUTE # // deprecated, do not use; see GRHISupportsParallelRHIExecute
 
-#ifndef PLATFORM_SUPPORTS_PARALLEL_RHI_EXECUTE
-	#define PLATFORM_SUPPORTS_PARALLEL_RHI_EXECUTE		0
-#endif
-
-
-#if PLATFORM_SUPPORTS_PARALLEL_RHI_EXECUTE && !PLATFORM_RHI_USES_CONTEXT_OBJECT
-	#error "PLATFORM_SUPPORTS_PARALLEL_RHI_EXECUTE requires PLATFORM_RHI_USES_CONTEXT_OBJECT"
-#endif
 
 //These are deprecated old defines that we want to make sure are not used
 #define CONSOLE (#)
diff --git a/Engine/Source/Runtime/EmptyRHI/Private/EmptyQuery.cpp b/Engine/Source/Runtime/EmptyRHI/Private/EmptyQuery.cpp
--- a/Engine/Source/Runtime/EmptyRHI/Private/EmptyQuery.cpp
+++ b/Engine/Source/Runtime/EmptyRHI/Private/EmptyQuery.cpp
@@ -37,14 +37,10 @@ FRenderQueryRHIRef FEmptyDynamicRHI::RHICreateRenderQuery(ERenderQueryType Query
 	return new FEmptyRenderQuery(QueryType);
 }
 
-void FEmptyDynamicRHI::RHIResetRenderQuery(FRenderQueryRHIParamRef QueryRHI)
-{
-	DYNAMIC_CAST_EMPTYRESOURCE(RenderQuery,Query);
-
-}
-
 bool FEmptyDynamicRHI::RHIGetRenderQueryResult(FRenderQueryRHIParamRef QueryRHI,uint64& OutNumPixels,bool bWait)
 {
+	check(IsInRenderingThread());
+
 	DYNAMIC_CAST_EMPTYRESOURCE(RenderQuery,Query);
 
 	return false;
diff --git a/Engine/Source/Runtime/EmptyRHI/Private/EmptyRHI.cpp b/Engine/Source/Runtime/EmptyRHI/Private/EmptyRHI.cpp
--- a/Engine/Source/Runtime/EmptyRHI/Private/EmptyRHI.cpp
+++ b/Engine/Source/Runtime/EmptyRHI/Private/EmptyRHI.cpp
@@ -176,3 +176,14 @@ void* FEmptyDynamicRHI::RHIGetNativeDevice()
 {
 	return NULL;
 }
+
+IRHICommandContext* FEmptyDynamicRHI::RHIGetDefaultContext()
+{
+	return this;
+}
+
+IRHICommandContextContainer* FEmptyDynamicRHI::RHIGetCommandContextContainer()
+{
+	return nullptr;
+}
+
diff --git a/Engine/Source/Runtime/EmptyRHI/Private/EmptyShaders.cpp b/Engine/Source/Runtime/EmptyRHI/Private/EmptyShaders.cpp
--- a/Engine/Source/Runtime/EmptyRHI/Private/EmptyShaders.cpp
+++ b/Engine/Source/Runtime/EmptyRHI/Private/EmptyShaders.cpp
@@ -106,6 +106,7 @@ FBoundShaderStateRHIRef FEmptyDynamicRHI::RHICreateBoundShaderState(
 	FGeometryShaderRHIParamRef GeometryShaderRHI
 	)
 {
+	check(IsInRenderingThread());
 	// Check for an existing bound shader state which matches the parameters
 	FCachedBoundShaderStateLink* CachedBoundShaderStateLink = GetCachedBoundShaderState(
 		VertexDeclarationRHI,
diff --git a/Engine/Source/Runtime/EmptyRHI/Private/EmptyUniformBuffer.cpp b/Engine/Source/Runtime/EmptyRHI/Private/EmptyUniformBuffer.cpp
--- a/Engine/Source/Runtime/EmptyRHI/Private/EmptyUniformBuffer.cpp
+++ b/Engine/Source/Runtime/EmptyRHI/Private/EmptyUniformBuffer.cpp
@@ -65,5 +65,6 @@ FEmptyUniformBuffer::~FEmptyUniformBuffer()
 
 FUniformBufferRHIRef FEmptyDynamicRHI::RHICreateUniformBuffer(const void* Contents, uint32 NumBytes, EUniformBufferUsage Usage)
 {
+	check(IsInRenderingThread());
 	return new FEmptyUniformBuffer(Contents, NumBytes, Usage);
 }
diff --git a/Engine/Source/Runtime/EmptyRHI/Private/EmptyViewport.cpp b/Engine/Source/Runtime/EmptyRHI/Private/EmptyViewport.cpp
--- a/Engine/Source/Runtime/EmptyRHI/Private/EmptyViewport.cpp
+++ b/Engine/Source/Runtime/EmptyRHI/Private/EmptyViewport.cpp
@@ -66,3 +66,7 @@ FTexture2DRHIRef FEmptyDynamicRHI::RHIGetViewportBackBuffer(FViewportRHIParamRef
 
 	return FTexture2DRHIRef();
 }
+
+void FEmptyDynamicRHI::RHIAdvanceFrameForGetViewportBackBuffer()
+{
+}
diff --git a/Engine/Source/Runtime/EmptyRHI/Public/EmptyRHI.h b/Engine/Source/Runtime/EmptyRHI/Public/EmptyRHI.h
--- a/Engine/Source/Runtime/EmptyRHI/Public/EmptyRHI.h
+++ b/Engine/Source/Runtime/EmptyRHI/Public/EmptyRHI.h
@@ -18,7 +18,7 @@ DECLARE_LOG_CATEGORY_EXTERN(LogEmpty, Display, All)
 #include "EmptyViewport.h"
 
 /** The interface which is implemented by the dynamically bound RHI. */
-class FEmptyDynamicRHI : public FDynamicRHI
+class FEmptyDynamicRHI : public FDynamicRHI, public IRHICommandContext
 {
 public:
 
diff --git a/Engine/Source/Runtime/Engine/Private/GPUSkinVertexFactory.cpp b/Engine/Source/Runtime/Engine/Private/GPUSkinVertexFactory.cpp
--- a/Engine/Source/Runtime/Engine/Private/GPUSkinVertexFactory.cpp
+++ b/Engine/Source/Runtime/Engine/Private/GPUSkinVertexFactory.cpp
@@ -134,14 +134,12 @@ void FGPUBaseSkinVertexFactory::ShaderDataType::UpdateBoneData(ERHIFeatureLevel:
 		}
 		if(NumBones)
 		{
-#if PLATFORM_SUPPORTS_RHI_THREAD
-			check(VectorArraySize == NumBones * sizeof(BoneMatrices[0]));
 			if (GRHIThread)
 			{
+				check(VectorArraySize == NumBones * sizeof(BoneMatrices[0]));
 				GRHICommandList.GetImmediateCommandList().UpdateVertexBuffer(BoneBuffer.VertexBufferRHI, BoneMatrices.GetData(), NumBones * sizeof(BoneMatrices[0]));
 			}
 			else
-#endif
 			{
 				float* Data = (float*)RHILockVertexBuffer(BoneBuffer.VertexBufferRHI, 0, VectorArraySize, RLM_WriteOnly);
 				checkSlow(Data);
diff --git a/Engine/Source/Runtime/Engine/Private/Particles/ParticleGpuSimulation.cpp b/Engine/Source/Runtime/Engine/Private/Particles/ParticleGpuSimulation.cpp
--- a/Engine/Source/Runtime/Engine/Private/Particles/ParticleGpuSimulation.cpp
+++ b/Engine/Source/Runtime/Engine/Private/Particles/ParticleGpuSimulation.cpp
@@ -1558,13 +1558,11 @@ void InjectNewParticles(FRHICommandList& RHICmdList, ERHIFeatureLevel::Type Feat
 		// Copy new particles in to the vertex buffer.
 		const int32 ParticlesThisDrawCall = FMath::Min<int32>( ParticleCount, MaxParticlesPerDrawCall );
 		const void* Src = NewParticles.GetData() + FirstParticle;
-#if PLATFORM_SUPPORTS_RHI_THREAD
 		if (GRHIThread)
 		{
 			RHICmdList.UpdateVertexBuffer(ScratchVertexBufferRHI, Src, ParticlesThisDrawCall * sizeof(FNewParticle));
 		}
 		else
-#endif
 		{
 			void* Dest = RHILockVertexBuffer( ScratchVertexBufferRHI, 0, ParticlesThisDrawCall * sizeof(FNewParticle), RLM_WriteOnly );
 			FMemory::Memcpy( Dest, Src, ParticlesThisDrawCall * sizeof(FNewParticle) );
diff --git a/Engine/Source/Runtime/Engine/Private/UnrealEngine.cpp b/Engine/Source/Runtime/Engine/Private/UnrealEngine.cpp
--- a/Engine/Source/Runtime/Engine/Private/UnrealEngine.cpp
+++ b/Engine/Source/Runtime/Engine/Private/UnrealEngine.cpp
@@ -3098,23 +3098,26 @@ bool UEngine::HandleToggleRenderingThreadCommand( const TCHAR* Cmd, FOutputDevic
 
 bool UEngine::HandleToggleRHIThreadCommand( const TCHAR* Cmd, FOutputDevice& Ar )
 {
-#if PLATFORM_SUPPORTS_RHI_THREAD
-	if (!GIsThreadedRendering)
+	if (GRHISupportsRHIThread)
 	{
-		check(!GRHIThread);
-		Ar.Logf( TEXT("Can't switch to RHI thread mode when we are not running a multithreaded renderer."));
+		if (!GIsThreadedRendering)
+		{
+			check(!GRHIThread);
+			Ar.Logf( TEXT("Can't switch to RHI thread mode when we are not running a multithreaded renderer."));
+		}
+		else
+		{
+			bool bWasRHIThread = !!GRHIThread;
+			StopRenderingThread();
+			GUseRHIThread = !bWasRHIThread;
+			StartRenderingThread();
+		}
+		Ar.Logf( TEXT("RHIThread is now %s."), GRHIThread ? TEXT("active") : TEXT("inactive"));
 	}
 	else
 	{
-		bool bWasRHIThread = !!GRHIThread;
-		StopRenderingThread();
-		GUseRHIThread = !bWasRHIThread;
-		StartRenderingThread();
+		Ar.Logf( TEXT("This RHI does not support the RHI thread."));
 	}
-	Ar.Logf( TEXT("RHIThread is now %s."), GRHIThread ? TEXT("active") : TEXT("inactive"));
-#else
-	Ar.Logf( TEXT("This platform does not support the RHI thread."));
-#endif
 	return true;
 }
 
diff --git a/Engine/Source/Runtime/Engine/Public/StaticRHI.h b/Engine/Source/Runtime/Engine/Public/StaticRHI.h
deleted file mode 100644
diff --git a/Engine/Source/Runtime/IOS/MetalRHI/Private/MetalCommands.cpp b/Engine/Source/Runtime/IOS/MetalRHI/Private/MetalCommands.cpp
--- a/Engine/Source/Runtime/IOS/MetalRHI/Private/MetalCommands.cpp
+++ b/Engine/Source/Runtime/IOS/MetalRHI/Private/MetalCommands.cpp
@@ -776,7 +776,7 @@ void FMetalDynamicRHI::RHIClearMRT(bool bClearColor,int32 NumClearColors,const F
 		PixelShader = *MRTPixelShader;
 
 		{
-			FRHICommandList_RecursiveHazardous RHICmdList;
+			FRHICommandList_RecursiveHazardous RHICmdList(this);
 			SetGlobalBoundShaderState(RHICmdList, GMaxRHIFeatureLevel, GClearMRTBoundShaderState[0], GVector4VertexDeclaration.VertexDeclarationRHI, *VertexShader, PixelShader);
 			FLinearColor ShaderClearColors[MaxSimultaneousRenderTargets];
 			FMemory::Memzero(ShaderClearColors);
@@ -893,3 +893,15 @@ void FMetalDynamicRHI::RHIEnableDepthBoundsTest(bool bEnable, float MinDepth, fl
 	// not supported
 }
 
+IRHICommandContext* FMetalDynamicRHI::RHIGetDefaultContext()
+{
+	return this;
+}
+
+IRHICommandContextContainer* FMetalDynamicRHI::RHIGetCommandContextContainer()
+{
+	return nullptr;
+}
+
+
+
diff --git a/Engine/Source/Runtime/IOS/MetalRHI/Private/MetalQuery.cpp b/Engine/Source/Runtime/IOS/MetalRHI/Private/MetalQuery.cpp
--- a/Engine/Source/Runtime/IOS/MetalRHI/Private/MetalQuery.cpp
+++ b/Engine/Source/Runtime/IOS/MetalRHI/Private/MetalQuery.cpp
@@ -46,15 +46,9 @@ FRenderQueryRHIRef FMetalDynamicRHI::RHICreateRenderQuery(ERenderQueryType Query
 	return new FMetalRenderQuery(QueryType);
 }
 
-void FMetalDynamicRHI::RHIResetRenderQuery(FRenderQueryRHIParamRef QueryRHI)
-{
-	DYNAMIC_CAST_METALRESOURCE(RenderQuery,Query);
-
-//	NOT_SUPPORTED("RHIResetRenderQuery");
-}
-
 bool FMetalDynamicRHI::RHIGetRenderQueryResult(FRenderQueryRHIParamRef QueryRHI,uint64& OutNumPixels,bool bWait)
 {
+	check(IsInRenderingThread());
 	DYNAMIC_CAST_METALRESOURCE(RenderQuery, Query);
 
 	uint32* QueryMemory = (uint32*)((uint8*)[FMetalManager::Get()->GetQueryBuffer() contents] + Query->Offset);
diff --git a/Engine/Source/Runtime/IOS/MetalRHI/Private/MetalShaders.cpp b/Engine/Source/Runtime/IOS/MetalRHI/Private/MetalShaders.cpp
--- a/Engine/Source/Runtime/IOS/MetalRHI/Private/MetalShaders.cpp
+++ b/Engine/Source/Runtime/IOS/MetalRHI/Private/MetalShaders.cpp
@@ -334,6 +334,7 @@ FBoundShaderStateRHIRef FMetalDynamicRHI::RHICreateBoundShaderState(
 	FGeometryShaderRHIParamRef GeometryShaderRHI
 	)
 {
+	check(IsInRenderingThread());
 	// Check for an existing bound shader state which matches the parameters
 	FCachedBoundShaderStateLink* CachedBoundShaderStateLink = GetCachedBoundShaderState(
 		VertexDeclarationRHI,
diff --git a/Engine/Source/Runtime/IOS/MetalRHI/Private/MetalUniformBuffer.cpp b/Engine/Source/Runtime/IOS/MetalRHI/Private/MetalUniformBuffer.cpp
--- a/Engine/Source/Runtime/IOS/MetalRHI/Private/MetalUniformBuffer.cpp
+++ b/Engine/Source/Runtime/IOS/MetalRHI/Private/MetalUniformBuffer.cpp
@@ -308,5 +308,6 @@ void FMetalUniformBuffer::CacheResourcesInternal()
 
 FUniformBufferRHIRef FMetalDynamicRHI::RHICreateUniformBuffer(const void* Contents, const FRHIUniformBufferLayout& Layout, EUniformBufferUsage Usage)
 {
+	check(IsInRenderingThread());
 	return new FMetalUniformBuffer(Contents, Layout, Usage);
 }
diff --git a/Engine/Source/Runtime/IOS/MetalRHI/Private/MetalViewport.cpp b/Engine/Source/Runtime/IOS/MetalRHI/Private/MetalViewport.cpp
--- a/Engine/Source/Runtime/IOS/MetalRHI/Private/MetalViewport.cpp
+++ b/Engine/Source/Runtime/IOS/MetalRHI/Private/MetalViewport.cpp
@@ -62,3 +62,7 @@ FTexture2DRHIRef FMetalDynamicRHI::RHIGetViewportBackBuffer(FViewportRHIParamRef
 {
 	return FMetalManager::Get()->GetBackBuffer();
 }
+
+void FMetalDynamicRHI::RHIAdvanceFrameForGetViewportBackBuffer()
+{
+}
diff --git a/Engine/Source/Runtime/IOS/MetalRHI/Public/MetalRHI.h b/Engine/Source/Runtime/IOS/MetalRHI/Public/MetalRHI.h
--- a/Engine/Source/Runtime/IOS/MetalRHI/Public/MetalRHI.h
+++ b/Engine/Source/Runtime/IOS/MetalRHI/Public/MetalRHI.h
@@ -21,7 +21,7 @@ DECLARE_LOG_CATEGORY_EXTERN(LogMetal, Display, All);
 class FMetalManager;
 
 /** The interface which is implemented by the dynamically bound RHI. */
-class FMetalDynamicRHI : public FDynamicRHI
+class FMetalDynamicRHI : public FDynamicRHI, public IRHICommandContext
 {
 public:
 
diff --git a/Engine/Source/Runtime/Launch/Private/LaunchEngineLoop.cpp b/Engine/Source/Runtime/Launch/Private/LaunchEngineLoop.cpp
--- a/Engine/Source/Runtime/Launch/Private/LaunchEngineLoop.cpp
+++ b/Engine/Source/Runtime/Launch/Private/LaunchEngineLoop.cpp
@@ -1409,18 +1409,19 @@ int32 FEngineLoop::PreInit( const TCHAR* CmdLine )
 
 	if (GUseThreadedRendering)
 	{
-#if PLATFORM_SUPPORTS_RHI_THREAD
-		const bool DefaultUseRHIThread = true;
-		GUseRHIThread = DefaultUseRHIThread;
-		if (FParse::Param(FCommandLine::Get(),TEXT("rhithread")))
+		if (GRHISupportsRHIThread)
 		{
-			GUseRHIThread = true;
-		}
-		else if (FParse::Param(FCommandLine::Get(),TEXT("norhithread")))
-		{
-			GUseRHIThread = false;
+			const bool DefaultUseRHIThread = true;
+			GUseRHIThread = DefaultUseRHIThread;
+			if (FParse::Param(FCommandLine::Get(),TEXT("rhithread")))
+			{
+				GUseRHIThread = true;
+			}
+			else if (FParse::Param(FCommandLine::Get(),TEXT("norhithread")))
+			{
+				GUseRHIThread = false;
+			}
 		}
-#endif
 		StartRenderingThread();
 	}
 	
diff --git a/Engine/Source/Runtime/NullDrv/Public/NullRHI.h b/Engine/Source/Runtime/NullDrv/Public/NullRHI.h
--- a/Engine/Source/Runtime/NullDrv/Public/NullRHI.h
+++ b/Engine/Source/Runtime/NullDrv/Public/NullRHI.h
@@ -4,11 +4,7 @@
 
 
 /** A null implementation of the dynamically bound RHI. */
-class FNullDynamicRHI
-	: public FDynamicRHI
-#if PLATFORM_RHI_USES_CONTEXT_OBJECT
-	, public IRHICommandContext
-#endif
+class FNullDynamicRHI : public FDynamicRHI , public IRHICommandContext
 {
 public:
 
diff --git a/Engine/Source/Runtime/OpenGLDrv/Private/OpenGLCommands.cpp b/Engine/Source/Runtime/OpenGLDrv/Private/OpenGLCommands.cpp
--- a/Engine/Source/Runtime/OpenGLDrv/Private/OpenGLCommands.cpp
+++ b/Engine/Source/Runtime/OpenGLDrv/Private/OpenGLCommands.cpp
@@ -3541,4 +3541,15 @@ void FOpenGLDynamicRHI::RHIEnableDepthBoundsTest(bool bEnable,float MinDepth,flo
 	}
 }
 
+IRHICommandContext* FOpenGLDynamicRHI::RHIGetDefaultContext()
+{
+	return this;
+}
+
+IRHICommandContextContainer* FOpenGLDynamicRHI::RHIGetCommandContextContainer()
+{
+	return nullptr;
+}
+
+
 #endif
\ No newline at end of file
diff --git a/Engine/Source/Runtime/OpenGLDrv/Private/OpenGLDrv.cpp b/Engine/Source/Runtime/OpenGLDrv/Private/OpenGLDrv.cpp
--- a/Engine/Source/Runtime/OpenGLDrv/Private/OpenGLDrv.cpp
+++ b/Engine/Source/Runtime/OpenGLDrv/Private/OpenGLDrv.cpp
@@ -387,7 +387,7 @@ void FOpenGLDynamicRHI::IssueLongGPUTask()
 
 		const auto FeatureLevel = GMaxRHIFeatureLevel;
 
-		FRHICommandList_RecursiveHazardous RHICmdList;
+		FRHICommandList_RecursiveHazardous RHICmdList(this);
 		SetRenderTarget(RHICmdList, Viewport->GetBackBuffer(), FTextureRHIRef());
 		RHICmdList.SetBlendState(TStaticBlendState<CW_RGBA, BO_Add, BF_One, BF_One>::GetRHI(), FLinearColor::Black);
 		RHICmdList.SetDepthStencilState(TStaticDepthStencilState<false, CF_Always>::GetRHI(), 0);
diff --git a/Engine/Source/Runtime/OpenGLDrv/Private/OpenGLQuery.cpp b/Engine/Source/Runtime/OpenGLDrv/Private/OpenGLQuery.cpp
--- a/Engine/Source/Runtime/OpenGLDrv/Private/OpenGLQuery.cpp
+++ b/Engine/Source/Runtime/OpenGLDrv/Private/OpenGLQuery.cpp
@@ -32,22 +32,12 @@ FRenderQueryRHIRef FOpenGLDynamicRHI::RHICreateRenderQuery(ERenderQueryType Quer
 	return new FOpenGLRenderQuery(QueryType);
 }
 
-void FOpenGLDynamicRHI::RHIResetRenderQuery(FRenderQueryRHIParamRef QueryRHI)
-{
-	DYNAMIC_CAST_OPENGLRESOURCE(RenderQuery,Query);
-
-	if (Query)
-	{
-		Query->bResultIsCached = false;
-	}
-}
-
 void FOpenGLDynamicRHI::RHIBeginRenderQuery(FRenderQueryRHIParamRef QueryRHI)
 {
 	VERIFY_GL_SCOPE();
 
 	DYNAMIC_CAST_OPENGLRESOURCE(RenderQuery,Query);
-
+	Query->bResultIsCached = false;
 	if(Query->QueryType == RQT_Occlusion)
 	{
 		check(PendingState.RunningOcclusionQuery == 0);
@@ -108,6 +98,7 @@ void FOpenGLDynamicRHI::RHIEndRenderQuery(FRenderQueryRHIParamRef QueryRHI)
 
 bool FOpenGLDynamicRHI::RHIGetRenderQueryResult(FRenderQueryRHIParamRef QueryRHI,uint64& OutResult,bool bWait)
 {
+	check(IsInRenderingThread());
 	VERIFY_GL_SCOPE();
 
 	DYNAMIC_CAST_OPENGLRESOURCE(RenderQuery,Query);
diff --git a/Engine/Source/Runtime/OpenGLDrv/Private/OpenGLShaders.cpp b/Engine/Source/Runtime/OpenGLDrv/Private/OpenGLShaders.cpp
--- a/Engine/Source/Runtime/OpenGLDrv/Private/OpenGLShaders.cpp
+++ b/Engine/Source/Runtime/OpenGLDrv/Private/OpenGLShaders.cpp
@@ -2021,6 +2021,8 @@ FBoundShaderStateRHIRef FOpenGLDynamicRHI::RHICreateBoundShaderState(
 	FGeometryShaderRHIParamRef GeometryShaderRHI
 	)
 { 
+	check(IsInRenderingThread());
+
 	VERIFY_GL_SCOPE();
 
 	SCOPE_CYCLE_COUNTER(STAT_OpenGLCreateBoundShaderStateTime);
diff --git a/Engine/Source/Runtime/OpenGLDrv/Private/OpenGLViewport.cpp b/Engine/Source/Runtime/OpenGLDrv/Private/OpenGLViewport.cpp
--- a/Engine/Source/Runtime/OpenGLDrv/Private/OpenGLViewport.cpp
+++ b/Engine/Source/Runtime/OpenGLDrv/Private/OpenGLViewport.cpp
@@ -216,6 +216,11 @@ FTexture2DRHIRef FOpenGLDynamicRHI::RHIGetViewportBackBuffer(FViewportRHIParamRe
 	return Viewport->GetBackBuffer();
 }
 
+void FOpenGLDynamicRHI::RHIAdvanceFrameForGetViewportBackBuffer()
+{
+}
+
+
 FOpenGLViewport::FOpenGLViewport(FOpenGLDynamicRHI* InOpenGLRHI,void* InWindowHandle,uint32 InSizeX,uint32 InSizeY,bool bInIsFullscreen,EPixelFormat PreferredPixelFormat)
 	: OpenGLRHI(InOpenGLRHI)
 	, OpenGLContext(NULL)
diff --git a/Engine/Source/Runtime/OpenGLDrv/Public/OpenGLDrv.h b/Engine/Source/Runtime/OpenGLDrv/Public/OpenGLDrv.h
--- a/Engine/Source/Runtime/OpenGLDrv/Public/OpenGLDrv.h
+++ b/Engine/Source/Runtime/OpenGLDrv/Public/OpenGLDrv.h
@@ -286,7 +286,7 @@ struct FOpenGLGPUProfiler : public FGPUProfiler
 
 
 /** The interface which is implemented by the dynamically bound RHI. */
-class FOpenGLDynamicRHI : public FDynamicRHI
+class FOpenGLDynamicRHI : public FDynamicRHI, public IRHICommandContext
 {
 public:
 
diff --git a/Engine/Source/Runtime/RHI/Private/BoundShaderStateCache.cpp b/Engine/Source/Runtime/RHI/Private/BoundShaderStateCache.cpp
--- a/Engine/Source/Runtime/RHI/Private/BoundShaderStateCache.cpp
+++ b/Engine/Source/Runtime/RHI/Private/BoundShaderStateCache.cpp
@@ -7,12 +7,9 @@
 #include "RHI.h"
 #include "BoundShaderStateCache.h"
 
-#if !defined(HAS_THREADSAFE_CreateBoundShaderState)
-#error "HAS_THREADSAFE_CreateBoundShaderState must be defined"
-#endif
-
 
 typedef TMap<FBoundShaderStateKey,FCachedBoundShaderStateLink*> FBoundShaderStateCache;
+typedef TMap<FBoundShaderStateKey,FCachedBoundShaderStateLink_Threadsafe*> FBoundShaderStateCache_Threadsafe;
 
 /** Lazily initialized bound shader state cache singleton. */
 static FBoundShaderStateCache& GetBoundShaderStateCache()
@@ -21,9 +18,14 @@ static FBoundShaderStateCache& GetBoundShaderStateCache()
 	return BoundShaderStateCache;
 }
 
-#if HAS_THREADSAFE_CreateBoundShaderState
+/** Lazily initialized bound shader state cache singleton. */
+static FBoundShaderStateCache_Threadsafe& GetBoundShaderStateCache_Threadsafe()
+{
+	static FBoundShaderStateCache_Threadsafe BoundShaderStateCache;
+	return BoundShaderStateCache;
+}
+
 static FCriticalSection BoundShaderStateCacheLock;
-#endif
 
 
 FCachedBoundShaderStateLink::FCachedBoundShaderStateLink(
@@ -33,54 +35,73 @@ FCachedBoundShaderStateLink::FCachedBoundShaderStateLink(
 	FHullShaderRHIParamRef HullShader,
 	FDomainShaderRHIParamRef DomainShader,
 	FGeometryShaderRHIParamRef GeometryShader,
-	FBoundShaderStateRHIParamRef InBoundShaderState
+	FBoundShaderStateRHIParamRef InBoundShaderState,
+	bool bAddToSingleThreadedCache
 	):
 	BoundShaderState(InBoundShaderState),
-	Key(VertexDeclaration,VertexShader,PixelShader,HullShader,DomainShader,GeometryShader)
+	Key(VertexDeclaration,VertexShader,PixelShader,HullShader,DomainShader,GeometryShader),
+	bAddedToSingleThreadedCache(bAddToSingleThreadedCache)
 {
-	// Add this bound shader state to the cache.
-#if !HAS_THREADSAFE_CreateBoundShaderState
-	GetBoundShaderStateCache().Add(Key,this);
-#endif
+	if (bAddToSingleThreadedCache)
+	{
+		GetBoundShaderStateCache().Add(Key,this);
+	}
 }
 
 FCachedBoundShaderStateLink::FCachedBoundShaderStateLink(
 	FVertexDeclarationRHIParamRef VertexDeclaration,
 	FVertexShaderRHIParamRef VertexShader,
 	FPixelShaderRHIParamRef PixelShader,
-	FBoundShaderStateRHIParamRef InBoundShaderState
+	FBoundShaderStateRHIParamRef InBoundShaderState,
+	bool bAddToSingleThreadedCache
 	):
 	BoundShaderState(InBoundShaderState),
-	Key(VertexDeclaration,VertexShader,PixelShader)
+	Key(VertexDeclaration,VertexShader,PixelShader),
+	bAddedToSingleThreadedCache(bAddToSingleThreadedCache)
 {
-	// Add this bound shader state to the cache.
-#if !HAS_THREADSAFE_CreateBoundShaderState
-	GetBoundShaderStateCache().Add(Key,this);
-#endif
+	if (bAddToSingleThreadedCache)
+	{
+		GetBoundShaderStateCache().Add(Key,this);
+	}
 }
 
 FCachedBoundShaderStateLink::~FCachedBoundShaderStateLink()
 {
-#if !HAS_THREADSAFE_CreateBoundShaderState
-	GetBoundShaderStateCache().Remove(Key);
-#endif
+	if (bAddedToSingleThreadedCache)
+	{
+		GetBoundShaderStateCache().Remove(Key);
+		bAddedToSingleThreadedCache = false;
+	}
+}
+
+FCachedBoundShaderStateLink* GetCachedBoundShaderState(
+	FVertexDeclarationRHIParamRef VertexDeclaration,
+	FVertexShaderRHIParamRef VertexShader,
+	FPixelShaderRHIParamRef PixelShader,
+	FHullShaderRHIParamRef HullShader,
+	FDomainShaderRHIParamRef DomainShader,
+	FGeometryShaderRHIParamRef GeometryShader
+	)
+{
+	// Find the existing bound shader state in the cache.
+	return GetBoundShaderStateCache().FindRef(
+		FBoundShaderStateKey(VertexDeclaration,VertexShader,PixelShader,HullShader,DomainShader,GeometryShader)
+		);
 }
 
-#if HAS_THREADSAFE_CreateBoundShaderState
-void FCachedBoundShaderStateLink::AddToCache()
+
+void FCachedBoundShaderStateLink_Threadsafe::AddToCache()
 {
 	FScopeLock Lock(&BoundShaderStateCacheLock);
-	GetBoundShaderStateCache().Add(Key,this);
+	GetBoundShaderStateCache_Threadsafe().Add(Key,this);
 }
-void FCachedBoundShaderStateLink::RemoveFromCache()
+void FCachedBoundShaderStateLink_Threadsafe::RemoveFromCache()
 {
 	FScopeLock Lock(&BoundShaderStateCacheLock);
-	GetBoundShaderStateCache().Remove(Key);
+	GetBoundShaderStateCache_Threadsafe().Remove(Key);
 }
-#endif
 
 
-#if HAS_THREADSAFE_CreateBoundShaderState
 FBoundShaderStateRHIRef GetCachedBoundShaderState_Threadsafe(
 	FVertexDeclarationRHIParamRef VertexDeclaration,
 	FVertexShaderRHIParamRef VertexShader,
@@ -92,7 +113,7 @@ FBoundShaderStateRHIRef GetCachedBoundShaderState_Threadsafe(
 {
 	FScopeLock Lock(&BoundShaderStateCacheLock);
 	// Find the existing bound shader state in the cache.
-	FCachedBoundShaderStateLink* CachedBoundShaderStateLink = GetBoundShaderStateCache().FindRef(
+	FCachedBoundShaderStateLink_Threadsafe* CachedBoundShaderStateLink = GetBoundShaderStateCache_Threadsafe().FindRef(
 		FBoundShaderStateKey(VertexDeclaration,VertexShader,PixelShader,HullShader,DomainShader,GeometryShader)
 		);
 	if(CachedBoundShaderStateLink)
@@ -103,21 +124,3 @@ FBoundShaderStateRHIRef GetCachedBoundShaderState_Threadsafe(
 	return FBoundShaderStateRHIRef();
 }
 
-#else
-
-FCachedBoundShaderStateLink* GetCachedBoundShaderState(
-	FVertexDeclarationRHIParamRef VertexDeclaration,
-	FVertexShaderRHIParamRef VertexShader,
-	FPixelShaderRHIParamRef PixelShader,
-	FHullShaderRHIParamRef HullShader,
-	FDomainShaderRHIParamRef DomainShader,
-	FGeometryShaderRHIParamRef GeometryShader
-	)
-{
-	// Find the existing bound shader state in the cache.
-	return GetBoundShaderStateCache().FindRef(
-		FBoundShaderStateKey(VertexDeclaration,VertexShader,PixelShader,HullShader,DomainShader,GeometryShader)
-		);
-}
-
-#endif
diff --git a/Engine/Source/Runtime/RHI/Private/DynamicRHI.cpp b/Engine/Source/Runtime/RHI/Private/DynamicRHI.cpp
--- a/Engine/Source/Runtime/RHI/Private/DynamicRHI.cpp
+++ b/Engine/Source/Runtime/RHI/Private/DynamicRHI.cpp
@@ -29,9 +29,7 @@ void InitNullRHI()
 
 	GDynamicRHI = DynamicRHIModule->CreateRHI();
 	GDynamicRHI->Init();
-#if PLATFORM_RHI_USES_CONTEXT_OBJECT
 	GRHICommandList.GetImmediateCommandList().SetContext(GDynamicRHI->RHIGetDefaultContext());
-#endif
 	GUsingNullRHI = true;
 	GRHISupportsTextureStreaming = false;
 }
@@ -52,9 +50,7 @@ void RHIInit(bool bHasEditorToken)
 			if (GDynamicRHI)
 			{
 				GDynamicRHI->Init();
-#if PLATFORM_RHI_USES_CONTEXT_OBJECT
 				GRHICommandList.GetImmediateCommandList().SetContext(GDynamicRHI->RHIGetDefaultContext());
-#endif
 			}
 #if PLATFORM_ALLOW_NULL_RHI
 			else
@@ -80,17 +76,7 @@ void RHIExit()
 	}
 }
 
-#if PLATFORM_RHI_USES_CONTEXT_OBJECT
-	#define DEFINE_RHIMETHOD_CMDLIST(Type,Name,ParameterTypesAndNames,ParameterNames,ReturnStatement,NullImplementation)
-#else
-	// Implement the static RHI methods that call the dynamic RHI.
-	#define DEFINE_RHIMETHOD_CMDLIST(Type,Name,ParameterTypesAndNames,ParameterNames,ReturnStatement,NullImplementation) \
-		RHI_API Type Name##_Internal ParameterTypesAndNames \
-	{ \
-		check(GDynamicRHI); \
-		ReturnStatement GDynamicRHI->RHI##Name ParameterNames; \
-	}
-#endif
+#define DEFINE_RHIMETHOD_CMDLIST(Type,Name,ParameterTypesAndNames,ParameterNames,ReturnStatement,NullImplementation)
 
 #define DEFINE_RHIMETHOD_GLOBAL(Type,Name,ParameterTypesAndNames,ParameterNames,ReturnStatement,NullImplementation) \
 	RHI_API Type Name##_Internal ParameterTypesAndNames \
diff --git a/Engine/Source/Runtime/RHI/Private/RHI.cpp b/Engine/Source/Runtime/RHI/Private/RHI.cpp
--- a/Engine/Source/Runtime/RHI/Private/RHI.cpp
+++ b/Engine/Source/Runtime/RHI/Private/RHI.cpp
@@ -189,6 +189,8 @@ bool GSupportsDepthBoundsTest = false;
 bool GRHISupportsBaseVertexIndex = true;
 bool GRHISupportsFirstInstance = false;
 bool GRHIRequiresEarlyBackBufferRenderTarget = true;
+bool GRHISupportsRHIThread = false;
+bool GRHISupportsParallelRHIExecute = false;
 
 /** Whether we are profiling GPU hitches. */
 bool GTriggerGPUHitchProfile = false;
diff --git a/Engine/Source/Runtime/RHI/Private/RHICommandList.cpp b/Engine/Source/Runtime/RHI/Private/RHICommandList.cpp
--- a/Engine/Source/Runtime/RHI/Private/RHICommandList.cpp
+++ b/Engine/Source/Runtime/RHI/Private/RHICommandList.cpp
@@ -12,18 +12,6 @@ DECLARE_CYCLE_STAT(TEXT("All Command List Execute"), STAT_ImmedCmdListExecuteTim
 DECLARE_DWORD_COUNTER_STAT(TEXT("Immed. Command List memory"), STAT_ImmedCmdListMemory, STATGROUP_RHICMDLIST);
 DECLARE_DWORD_COUNTER_STAT(TEXT("Immed. Command count"), STAT_ImmedCmdListCount, STATGROUP_RHICMDLIST);
 
-#if PLATFORM_SUPPORTS_RHI_THREAD
-/***
-Requirements for RHI thread
-* Microresources (those in RHIStaticStates.h) need to be able to be created by any thread at any time and be able to work with a radically simplified rhi resource lifecycle. CreateSamplerState, CreateRasterizerState, CreateDepthStencilState, CreateBlendState
-* CreateUniformBuffer needs to be threadsafe
-* GetRenderQueryResult should be threadsafe, if not, then PLATFORM_HAS_THREADSAFE_RHIGetRenderQueryResult should be 1
-* AdvanceFrameForGetViewportBackBuffer needs be added as an RHI method and this needs to work with GetViewportBackBuffer to give the render thread the right back buffer even though many commands relating to the beginning and end of the frame are queued.
-* ResetRenderQuery does not exist; this stuff should be done in BeginQuery
-
-***/
-#endif
-
 #if !PLATFORM_USES_FIXED_RHI_CLASS
 #include "RHICommandListCommandExecutes.inl"
 #endif
@@ -45,21 +33,19 @@ TAutoConsoleVariable<int32> CVarRHICmdWidth(
 	8,
 	TEXT("Number of threads."));
 
-#if PLATFORM_SUPPORTS_PARALLEL_RHI_EXECUTE
 static TAutoConsoleVariable<int32> CVarRHICmdUseDeferredContexts(
 	TEXT("r.RHICmdUseDeferredContexts"),
 	1,
-	TEXT("True to use deferred contexts to parallelize command list execution."));
+	TEXT("True to use deferred contexts to parallelize command list execution. Only available on some RHIs."));
 static TAutoConsoleVariable<int32> CVarRHIDeferredContextWidth(
 	TEXT("r.RHIDeferredContextWidth"),
 	99999,
-	TEXT("Number of pieces to break a parallel translate task into. A very large number is one task per original cmd list."));
+	TEXT("Number of pieces to break a parallel translate task into. A very large number is one task per original cmd list. Only available on some RHIs."));
 
 static TAutoConsoleVariable<int32> CVarRHICmdFlushOnQueueParallelSubmit(
 	TEXT("r.RHICmdFlushOnQueueParallelSubmit"),
 	0,
-	TEXT("Wait for completion of parallel commandlists immediately after submitting.  Debug feature."));
-#endif
+	TEXT("Wait for completion of parallel commandlists immediately after submitting.  Debug feature. Only available on some RHIs."));
 
 RHI_API FRHICommandListExecutor GRHICommandList;
 
@@ -97,9 +83,7 @@ FRHICommandBase* GCurrentCommand = nullptr;
 void FRHICommandListExecutor::ExecuteInner_DoExecute(FRHICommandListBase& CmdList)
 {
 	CmdList.bExecuting = true;
-#if PLATFORM_RHI_USES_CONTEXT_OBJECT
 	check(CmdList.Context);
-#endif
 
 	bool bDoStats = false;
 	//STAT(bDoStats = FThreadStats::IsCollectingData() && (IsInRenderingThread() || IsInRHIThread());)
@@ -155,12 +139,8 @@ public:
 
 	ENamedThreads::Type GetDesiredThread()
 	{
-#if PLATFORM_SUPPORTS_RHI_THREAD
+		check(GRHIThread); // this should never be used on a platform that doesn't support the RHI thread
 		return ENamedThreads::RHIThread;
-#else
-		check(0); // this should never be used on a platform that doesn't support the RHI thread
-		return ENamedThreads::AnyThread;
-#endif
 	}
 
 	static ESubsequentsMode::Type GetSubsequentsMode() { return ESubsequentsMode::TrackSubsequents; }
@@ -197,12 +177,8 @@ public:
 
 	ENamedThreads::Type GetDesiredThread()
 	{
-#if PLATFORM_SUPPORTS_RHI_THREAD
+		check(GRHIThread); // this should never be used on a platform that doesn't support the RHI thread
 		return ENamedThreads::RHIThread;
-#else
-		check(0); // this should never be used on a platform that doesn't support the RHI thread
-		return ENamedThreads::AnyThread;
-#endif
 	}
 
 	static ESubsequentsMode::Type GetSubsequentsMode() { return ESubsequentsMode::TrackSubsequents; }
@@ -240,9 +216,7 @@ public:
 
 	ENamedThreads::Type GetDesiredThread()
 	{
-#if !PLATFORM_SUPPORTS_RHI_THREAD
-		check(0); // this should never be used on a platform that doesn't support the RHI thread
-#endif
+		check(GRHIThread); // this should never be used on a platform that doesn't support the RHI thread
 		return ENamedThreads::RenderThread_Local;
 	}
 
@@ -390,9 +364,7 @@ static FORCEINLINE bool IsInRenderingOrRHIThread()
 
 void FRHICommandListExecutor::ExecuteList(FRHICommandListBase& CmdList)
 {
-	check(&CmdList != &GetImmediateCommandList() &&
-		(PLATFORM_SUPPORTS_PARALLEL_RHI_EXECUTE || IsInRenderingOrRHIThread())
-		);
+	check(&CmdList != &GetImmediateCommandList() && (GRHISupportsParallelRHIExecute || IsInRenderingOrRHIThread()));
 
 	if (IsInRenderingThread() && !GetImmediateCommandList().IsExecuting()) // don't flush if this is a recursive call and we are already executing the immediate command list
 	{
@@ -509,14 +481,12 @@ FGraphEventRef FRHICommandListExecutor::RHIThreadFence()
 		FRHICommandListExecutor::GetImmediateCommandList().ImmediateFlush(EImmediateFlushType::DispatchToRHIThread);
 	}
 
-#if PLATFORM_SUPPORTS_RHI_THREAD
 	// every task dispatched or executed through the flush creates a new RenderThreadExecutionTask that represents the 'execution' of the last dispatched
 	// set of RHI commands.
 	if (RenderThreadExecutionTask.GetReference() && RenderThreadExecutionTask->IsComplete())
 	{
 		RenderThreadExecutionTask = nullptr;
 	}		
-#endif
 	
 	return RenderThreadExecutionTask;
 }
@@ -562,12 +532,8 @@ void FRHICommandListBase::Reset()
 	NumCommands = 0;
 	Root = nullptr;
 	CommandLink = &Root;
-#if PLATFORM_RHI_USES_CONTEXT_OBJECT
-	static_assert(USE_DYNAMIC_RHI, "PLATFORM_RHI_USES_CONTEXT_OBJECT current requires a dynamic RHI");
+	static_assert(USE_DYNAMIC_RHI, "static RHIs are not supported.");
 	Context = GDynamicRHI ? RHIGetDefaultContext() : nullptr;
-#else
-	Context = nullptr;
-#endif
 #if USE_RHICOMMAND_STATE_REDUCTION
 	StateCache = nullptr;
 #endif
@@ -575,7 +541,6 @@ void FRHICommandListBase::Reset()
 }
 
 
-#if PLATFORM_SUPPORTS_PARALLEL_RHI_EXECUTE
 DECLARE_CYCLE_STAT(TEXT("Parallel Async Chain Translate"), STAT_ParallelChainTranslate, STATGROUP_RHICMDLIST);
 
 class FParallelTranslateCommandList
@@ -666,7 +631,6 @@ struct FRHICommandWaitForAndSubmitSubListParallel : public FRHICommand<FRHIComma
 		}
 	}
 };
-#endif
 
 
 
@@ -714,8 +678,6 @@ void FRHICommandListBase::QueueParallelAsyncCommandListSubmit(FGraphEventRef* An
 	{
 		FRHICommandListExecutor::GetImmediateCommandList().ImmediateFlush(EImmediateFlushType::DispatchToRHIThread); // we should start on the stuff before this async list
 	}
-#if PLATFORM_SUPPORTS_PARALLEL_RHI_EXECUTE
-
 #if !UE_BUILD_SHIPPING
 	// do a flush before hand so we can tell if it was this parallel set that broke something, or what came before.
 	if (CVarRHICmdFlushOnQueueParallelSubmit.GetValueOnRenderThread())
@@ -727,9 +689,9 @@ void FRHICommandListBase::QueueParallelAsyncCommandListSubmit(FGraphEventRef* An
 	if (Num && GRHIThread)
 	{
 		IRHICommandContextContainer* ContextContainer = nullptr;
-		if (CVarRHICmdUseDeferredContexts.GetValueOnAnyThread() >= 1)
+		if (GRHISupportsParallelRHIExecute && CVarRHICmdUseDeferredContexts.GetValueOnAnyThread() >= 1)
 		{
-			ContextContainer = RHIGetCommandContextContainer(0);
+			ContextContainer = RHIGetCommandContextContainer();
 		}
 		if (ContextContainer)
 		{
@@ -749,7 +711,7 @@ void FRHICommandListBase::QueueParallelAsyncCommandListSubmit(FGraphEventRef* An
 
 					if (!ContextContainer)
 					{
-						ContextContainer = RHIGetCommandContextContainer(0);
+						ContextContainer = RHIGetCommandContextContainer();
 					}
 					check(ContextContainer);
 
@@ -790,7 +752,6 @@ void FRHICommandListBase::QueueParallelAsyncCommandListSubmit(FGraphEventRef* An
 			}
 		}
 	}
-#endif
 	for (int32 Index = 0; Index < Num; Index++)
 	{
 		FGraphEventRef& AnyThreadCompletionEvent = AnyThreadCompletionEvents[Index];
@@ -905,7 +866,41 @@ void FRHICommandListBase::QueueCommandListSubmit(class FRHICommandList* CmdList)
 	new (AllocCommand<FRHICommandSubmitSubList>()) FRHICommandSubmitSubList(CmdList);
 }
 
-#if PLATFORM_SUPPORTS_RHI_THREAD
+
+void FRHICommandList::BeginScene()
+{
+	check(IsImmediate() && IsInRenderingThread());
+	if (Bypass())
+	{
+		CMD_CONTEXT(BeginScene)();
+		return;
+	}
+	new (AllocCommand<FRHICommandBeginScene>()) FRHICommandBeginScene();
+	if (!GRHIThread)
+	{
+		// if we aren't running an RHIThread, there is no good reason to buffer this frame advance stuff and that complicates state management, so flush everything out now
+		QUICK_SCOPE_CYCLE_COUNTER(BeginScene_Flush);
+		FRHICommandListExecutor::GetImmediateCommandList().ImmediateFlush(EImmediateFlushType::FlushRHIThread);
+	}
+}
+void FRHICommandList::EndScene()
+{
+	check(IsImmediate() && IsInRenderingThread());
+	if (Bypass())
+	{
+		CMD_CONTEXT(EndScene)();
+		return;
+	}
+	new (AllocCommand<FRHICommandEndScene>()) FRHICommandEndScene();
+	if (!GRHIThread)
+	{
+		// if we aren't running an RHIThread, there is no good reason to buffer this frame advance stuff and that complicates state management, so flush everything out now
+		QUICK_SCOPE_CYCLE_COUNTER(EndScene_Flush);
+		FRHICommandListExecutor::GetImmediateCommandList().ImmediateFlush(EImmediateFlushType::FlushRHIThread);
+	}
+}
+
+
 void FRHICommandList::BeginDrawingViewport(FViewportRHIParamRef Viewport, FTextureRHIParamRef RenderTargetRHI)
 {
 	check(IsImmediate() && IsInRenderingThread());
@@ -915,6 +910,12 @@ void FRHICommandList::BeginDrawingViewport(FViewportRHIParamRef Viewport, FTextu
 		return;
 	}
 	new (AllocCommand<FRHICommandBeginDrawingViewport>()) FRHICommandBeginDrawingViewport(Viewport, RenderTargetRHI);
+	if (!GRHIThread)
+	{
+		// if we aren't running an RHIThread, there is no good reason to buffer this frame advance stuff and that complicates state management, so flush everything out now
+		QUICK_SCOPE_CYCLE_COUNTER(BeginDrawingViewport_Flush);
+		FRHICommandListExecutor::GetImmediateCommandList().ImmediateFlush(EImmediateFlushType::FlushRHIThread);
+	}
 }
 
 void FRHICommandList::EndDrawingViewport(FViewportRHIParamRef Viewport, bool bPresent, bool bLockToVsync)
@@ -927,7 +928,7 @@ void FRHICommandList::EndDrawingViewport(FViewportRHIParamRef Viewport, bool bPr
 	else
 	{
 		new (AllocCommand<FRHICommandEndDrawingViewport>()) FRHICommandEndDrawingViewport(Viewport, bPresent, bLockToVsync);
-		// this should be started asap
+		// if we aren't running an RHIThread, there is no good reason to buffer this frame advance stuff and that complicates state management, so flush everything out now
 		{
 			QUICK_SCOPE_CYCLE_COUNTER(STAT_EndDrawingViewport_Dispatch);
 			FRHICommandListExecutor::GetImmediateCommandList().ImmediateFlush(EImmediateFlushType::DispatchToRHIThread);
@@ -938,26 +939,38 @@ void FRHICommandList::EndDrawingViewport(FViewportRHIParamRef Viewport, bool bPr
 
 void FRHICommandList::BeginFrame()
 {
+	check(IsImmediate() && IsInRenderingThread());
 	if (Bypass())
 	{
 		CMD_CONTEXT(BeginFrame)();
 		return;
 	}
 	new (AllocCommand<FRHICommandBeginFrame>()) FRHICommandBeginFrame();
+	if (!GRHIThread)
+	{
+		// if we aren't running an RHIThread, there is no good reason to buffer this frame advance stuff and that complicates state management, so flush everything out now
+		QUICK_SCOPE_CYCLE_COUNTER(BeginFrame_Flush);
+		FRHICommandListExecutor::GetImmediateCommandList().ImmediateFlush(EImmediateFlushType::FlushRHIThread);
+	}
 }
 
 void FRHICommandList::EndFrame()
 {
+	check(IsImmediate() && IsInRenderingThread());
 	if (Bypass())
 	{
 		CMD_CONTEXT(EndFrame)();
 		return;
 	}
 	new (AllocCommand<FRHICommandEndFrame>()) FRHICommandEndFrame();
+	if (!GRHIThread)
+	{
+		// if we aren't running an RHIThread, there is no good reason to buffer this frame advance stuff and that complicates state management, so flush everything out now
+		QUICK_SCOPE_CYCLE_COUNTER(EndFrame_Flush);
+		FRHICommandListExecutor::GetImmediateCommandList().ImmediateFlush(EImmediateFlushType::FlushRHIThread);
+	}
 }
 
-#endif
-
 DECLARE_CYCLE_STAT(TEXT("Explicit wait for tasks"), STAT_ExplicitWait, STATGROUP_RHICMDLIST);
 DECLARE_CYCLE_STAT(TEXT("Prewait dispatch"), STAT_PrewaitDispatch, STATGROUP_RHICMDLIST);
 void FRHICommandListBase::WaitForTasks(bool bKnownToBeComplete)
diff --git a/Engine/Source/Runtime/RHI/Private/StaticRHI.cpp b/Engine/Source/Runtime/RHI/Private/StaticRHI.cpp
deleted file mode 100644
diff --git a/Engine/Source/Runtime/RHI/Public/BoundShaderStateCache.h b/Engine/Source/Runtime/RHI/Public/BoundShaderStateCache.h
--- a/Engine/Source/Runtime/RHI/Public/BoundShaderStateCache.h
+++ b/Engine/Source/Runtime/RHI/Public/BoundShaderStateCache.h
@@ -114,7 +114,8 @@ public:
 		FVertexDeclarationRHIParamRef VertexDeclaration,
 		FVertexShaderRHIParamRef VertexShader,
 		FPixelShaderRHIParamRef PixelShader,
-		FBoundShaderStateRHIParamRef InBoundShaderState
+		FBoundShaderStateRHIParamRef InBoundShaderState,
+		bool bAddToSingleThreadedCache = true
 		);
 
 	/** Adds the bound shader state to the cache. */
@@ -125,7 +126,8 @@ public:
 		FHullShaderRHIParamRef HullShader,
 		FDomainShaderRHIParamRef DomainShader,
 		FGeometryShaderRHIParamRef GeometryShader,
-		FBoundShaderStateRHIParamRef InBoundShaderState
+		FBoundShaderStateRHIParamRef InBoundShaderState,
+		bool bAddToSingleThreadedCache = true
 		);
 
 	/** Destructor.  Removes the bound shader state from the cache. */
@@ -139,22 +141,17 @@ public:
 	FORCEINLINE FHullShaderRHIParamRef     GetHullShader() const     { return Key.GetHullShader(); }
 	FORCEINLINE FDomainShaderRHIParamRef   GetDomainShader() const   { return Key.GetDomainShader(); }
 	FORCEINLINE FGeometryShaderRHIParamRef GetGeometryShader() const { return Key.GetGeometryShader(); }
-#if PLATFORM_SUPPORTS_PARALLEL_RHI_EXECUTE
-	void AddToCache();
-	void RemoveFromCache();
-#endif
-
-private:
-
+protected:
 	FBoundShaderStateKey Key;
+	bool bAddedToSingleThreadedCache;
 };
 
+
 /**
  * Searches for a cached bound shader state.
  * @return If a bound shader state matching the parameters is cached, it is returned; otherwise NULL is returned.
  */
-#if PLATFORM_SUPPORTS_PARALLEL_RHI_EXECUTE
-extern RHI_API FBoundShaderStateRHIRef GetCachedBoundShaderState_Threadsafe(
+extern RHI_API FCachedBoundShaderStateLink* GetCachedBoundShaderState(
 	FVertexDeclarationRHIParamRef VertexDeclaration,
 	FVertexShaderRHIParamRef VertexShader,
 	FPixelShaderRHIParamRef PixelShader,
@@ -163,8 +160,43 @@ extern RHI_API FBoundShaderStateRHIRef GetCachedBoundShaderState_Threadsafe(
 	FGeometryShaderRHIParamRef GeometryShader = NULL
 	);
 
-#else
-extern RHI_API FCachedBoundShaderStateLink* GetCachedBoundShaderState(
+class RHI_API FCachedBoundShaderStateLink_Threadsafe : public FCachedBoundShaderStateLink
+{
+public:
+	/** Adds the bound shader state to the cache. */
+	FCachedBoundShaderStateLink_Threadsafe(
+		FVertexDeclarationRHIParamRef VertexDeclaration,
+		FVertexShaderRHIParamRef VertexShader,
+		FPixelShaderRHIParamRef PixelShader,
+		FBoundShaderStateRHIParamRef InBoundShaderState
+		)
+		: FCachedBoundShaderStateLink(VertexDeclaration, VertexShader, PixelShader, InBoundShaderState, false)
+	{
+	}
+
+	/** Adds the bound shader state to the cache. */
+	FCachedBoundShaderStateLink_Threadsafe(
+		FVertexDeclarationRHIParamRef VertexDeclaration,
+		FVertexShaderRHIParamRef VertexShader,
+		FPixelShaderRHIParamRef PixelShader,
+		FHullShaderRHIParamRef HullShader,
+		FDomainShaderRHIParamRef DomainShader,
+		FGeometryShaderRHIParamRef GeometryShader,
+		FBoundShaderStateRHIParamRef InBoundShaderState
+		)
+		: FCachedBoundShaderStateLink(VertexDeclaration, VertexShader, PixelShader, HullShader, DomainShader, GeometryShader, InBoundShaderState, false)
+	{
+	}
+
+	void AddToCache();
+	void RemoveFromCache();
+};
+
+/**
+ * Searches for a cached bound shader state. Threadsafe version.
+ * @return If a bound shader state matching the parameters is cached, it is returned; otherwise NULL is returned.
+ */
+extern RHI_API FBoundShaderStateRHIRef GetCachedBoundShaderState_Threadsafe(
 	FVertexDeclarationRHIParamRef VertexDeclaration,
 	FVertexShaderRHIParamRef VertexShader,
 	FPixelShaderRHIParamRef PixelShader,
@@ -172,4 +204,4 @@ extern RHI_API FCachedBoundShaderStateLink* GetCachedBoundShaderState(
 	FDomainShaderRHIParamRef DomainShader = NULL,
 	FGeometryShaderRHIParamRef GeometryShader = NULL
 	);
-#endif
+
diff --git a/Engine/Source/Runtime/RHI/Public/DynamicRHI.h b/Engine/Source/Runtime/RHI/Public/DynamicRHI.h
--- a/Engine/Source/Runtime/RHI/Public/DynamicRHI.h
+++ b/Engine/Source/Runtime/RHI/Public/DynamicRHI.h
@@ -30,15 +30,11 @@ public:
 	/** Shutdown the RHI; handle shutdown and resource destruction before the RHI's actual destructor is called (so that all resources of the RHI are still available for shutdown). */
 	virtual void Shutdown() = 0;
 
-#if PLATFORM_RHI_USES_CONTEXT_OBJECT
 	#define DEFINE_RHIMETHOD(Type,Name,ParameterTypesAndNames,ParameterNames,ReturnStatement,NullImplementation) virtual Type RHI##Name ParameterTypesAndNames = 0
 	#define DEFINE_RHIMETHOD_CMDLIST(Type,Name,ParameterTypesAndNames,ParameterNames,ReturnStatement,NullImplementation)
 	#define DEFINE_RHIMETHOD_GLOBAL(Type,Name,ParameterTypesAndNames,ParameterNames,ReturnStatement,NullImplementation) virtual Type RHI##Name ParameterTypesAndNames = 0
 	#define DEFINE_RHIMETHOD_GLOBALFLUSH(Type,Name,ParameterTypesAndNames,ParameterNames,ReturnStatement,NullImplementation) virtual Type RHI##Name ParameterTypesAndNames = 0
 	#define DEFINE_RHIMETHOD_GLOBALTHREADSAFE(Type,Name,ParameterTypesAndNames,ParameterNames,ReturnStatement,NullImplementation) virtual Type RHI##Name ParameterTypesAndNames = 0
-#else
-	#define DEFINE_RHIMETHOD(Type,Name,ParameterTypesAndNames,ParameterNames,ReturnStatement,NullImplementation) virtual Type RHI##Name ParameterTypesAndNames = 0
-#endif
 	#include "RHIMethods.h"
 	#undef DEFINE_RHIMETHOD
 	#undef DEFINE_RHIMETHOD_CMDLIST
@@ -50,13 +46,7 @@ public:
 /** A global pointer to the dynamically bound RHI implementation. */
 extern RHI_API FDynamicRHI* GDynamicRHI;
 
-// Implement the statically bound RHI methods to simply call the dynamic RHI.
-#if PLATFORM_RHI_USES_CONTEXT_OBJECT
-	#define DEFINE_RHIMETHOD_CMDLIST(Type,Name,ParameterTypesAndNames,ParameterNames,ReturnStatement,NullImplementation)
-#else
-	#define DEFINE_RHIMETHOD_CMDLIST(Type,Name,ParameterTypesAndNames,ParameterNames,ReturnStatement,NullImplementation) \
-		extern RHI_API Type Name##_Internal ParameterTypesAndNames;
-#endif
+#define DEFINE_RHIMETHOD_CMDLIST(Type,Name,ParameterTypesAndNames,ParameterNames,ReturnStatement,NullImplementation)
 #define DEFINE_RHIMETHOD_GLOBAL(Type,Name,ParameterTypesAndNames,ParameterNames,ReturnStatement,NullImplementation) \
 	extern RHI_API Type Name##_Internal ParameterTypesAndNames;
 #define DEFINE_RHIMETHOD_GLOBALTHREADSAFE(Type,Name,ParameterTypesAndNames,ParameterNames,ReturnStatement,NullImplementation) \
diff --git a/Engine/Source/Runtime/RHI/Public/RHI.h b/Engine/Source/Runtime/RHI/Public/RHI.h
--- a/Engine/Source/Runtime/RHI/Public/RHI.h
+++ b/Engine/Source/Runtime/RHI/Public/RHI.h
@@ -194,6 +194,22 @@ extern RHI_API bool GRHISupportsFirstInstance;
 /** Whether or not the engine should set the BackBuffer as a render target early in the frame. */
 extern RHI_API bool GRHIRequiresEarlyBackBufferRenderTarget;
 
+/** Whether or not the RHI supports an RHI thread.
+Requirements for RHI thread
+* Microresources (those in RHIStaticStates.h) need to be able to be created by any thread at any time and be able to work with a radically simplified rhi resource lifecycle. CreateSamplerState, CreateRasterizerState, CreateDepthStencilState, CreateBlendState
+* CreateUniformBuffer needs to be threadsafe
+* GetRenderQueryResult should be threadsafe, but this isn't required. If it isn't threadsafe, then you need to flush yourself in the RHI
+* GetViewportBackBuffer and AdvanceFrameForGetViewportBackBuffer need to be threadsafe and need to support the fact that the render thread has a different concept of "current backbuffer" than the RHI thread. Without an RHIThread this is moot due to the next two items.
+* AdvanceFrameForGetViewportBackBuffer needs be added as an RHI method and this needs to work with GetViewportBackBuffer to give the render thread the right back buffer even though many commands relating to the beginning and end of the frame are queued.
+* BeginDrawingViewport, and 5 or so other frame advance methods are queued with an RHIThread. Without an RHIThread, these just flush internally.
+***/
+extern RHI_API bool GRHISupportsRHIThread;
+
+/** Whether or not the RHI supports parallel RHIThread executes / translates
+Requirements:
+* RHICreateBoundShaderState is threadsafe and GetCachedBoundShaderState must not be used. GetCachedBoundShaderState_Threadsafe has a slightly different protocol.
+***/
+extern RHI_API bool GRHISupportsParallelRHIExecute;
 
 /** Called once per frame only from within an RHI. */
 extern RHI_API void RHIPrivateBeginFrame();
diff --git a/Engine/Source/Runtime/RHI/Public/RHICommandList.h b/Engine/Source/Runtime/RHI/Public/RHICommandList.h
--- a/Engine/Source/Runtime/RHI/Public/RHICommandList.h
+++ b/Engine/Source/Runtime/RHI/Public/RHICommandList.h
@@ -133,9 +133,7 @@ public:
 
 	void SetContext(IRHICommandContext* InContext)
 	{
-#if PLATFORM_RHI_USES_CONTEXT_OBJECT
 		check(InContext);
-#endif
 		Context = InContext;
 	}
 
@@ -1080,19 +1078,6 @@ struct FRHICommandEndRenderQuery : public FRHICommand<FRHICommandEndRenderQuery>
 	RHI_API void Execute(FRHICommandListBase& CmdList);
 };
 
-#if !PLATFORM_SUPPORTS_RHI_THREAD
-struct FRHICommandResetRenderQuery : public FRHICommand<FRHICommandResetRenderQuery>
-{
-	FRenderQueryRHIParamRef RenderQuery;
-
-	FORCEINLINE_DEBUGGABLE FRHICommandResetRenderQuery(FRenderQueryRHIParamRef InRenderQuery)
-		: RenderQuery(InRenderQuery)
-	{
-	}
-	RHI_API void Execute(FRHICommandListBase& CmdList);
-};
-#endif
-
 struct FRHICommandBeginOcclusionQueryBatch : public FRHICommand<FRHICommandBeginOcclusionQueryBatch>
 {
 	FORCEINLINE_DEBUGGABLE FRHICommandBeginOcclusionQueryBatch()
@@ -1156,7 +1141,6 @@ struct FRHICommandEndFrame : public FRHICommand<FRHICommandEndFrame>
 	RHI_API void Execute(FRHICommandListBase& CmdList);
 };
 
-#if PLATFORM_SUPPORTS_RHI_THREAD
 struct FRHICommandBeginDrawingViewport : public FRHICommand<FRHICommandBeginDrawingViewport>
 {
 	FViewportRHIParamRef Viewport;
@@ -1184,7 +1168,6 @@ struct FRHICommandEndDrawingViewport : public FRHICommand<FRHICommandEndDrawingV
 	}
 	RHI_API void Execute(FRHICommandListBase& CmdList);
 };
-#endif
 
 struct FRHICommandPushEvent : public FRHICommand<FRHICommandPushEvent>
 {
@@ -1213,11 +1196,7 @@ struct FRHICommandDebugBreak : public FRHICommand<FRHICommandDebugBreak>
 	}
 };
 
-#if PLATFORM_RHI_USES_CONTEXT_OBJECT
 #define CMD_CONTEXT(Method) GetContext().RHI##Method
-#else
-#define CMD_CONTEXT(Method) Method##_Internal
-#endif
 
 
 class RHI_API FRHICommandList : public FRHICommandListBase
@@ -1275,11 +1254,7 @@ public:
 		FLocalBoundShaderState Result;
 		if (Bypass())
 		{
-#if HAS_THREADSAFE_CreateBoundShaderState
 			Result.BypassBSS = RHICreateBoundShaderState(VertexDeclarationRHI, VertexShaderRHI, HullShaderRHI, DomainShaderRHI, PixelShaderRHI, GeometryShaderRHI);
-#else
-			Result.BypassBSS = CreateBoundShaderState_Internal(VertexDeclarationRHI, VertexShaderRHI, HullShaderRHI, DomainShaderRHI, PixelShaderRHI, GeometryShaderRHI);
-#endif
 		}
 		else
 		{
@@ -1310,11 +1285,7 @@ public:
 		FLocalUniformBuffer Result;
 		if (Bypass())
 		{
-#if PLATFORM_SUPPORTS_RHI_THREAD
 			Result.BypassUniform = RHICreateUniformBuffer(Contents, Layout, UniformBuffer_SingleFrame);
-#else
-			Result.BypassUniform = CreateUniformBuffer_Internal(Contents, Layout, UniformBuffer_SingleFrame);
-#endif
 		}
 		else
 		{
@@ -1831,17 +1802,6 @@ public:
 		}
 		new (AllocCommand<FRHICommandEndRenderQuery>()) FRHICommandEndRenderQuery(RenderQuery);
 	}
-#if !PLATFORM_SUPPORTS_RHI_THREAD
-	FORCEINLINE_DEBUGGABLE void ResetRenderQuery(FRenderQueryRHIParamRef RenderQuery)
-	{
-		if (Bypass())
-		{
-			CMD_CONTEXT(ResetRenderQuery)(RenderQuery);
-			return;
-		}
-		new (AllocCommand<FRHICommandResetRenderQuery>()) FRHICommandResetRenderQuery(RenderQuery);
-	}
-#endif
 
 	FORCEINLINE_DEBUGGABLE void BeginOcclusionQueryBatch()
 	{
@@ -1863,30 +1823,13 @@ public:
 		new (AllocCommand<FRHICommandEndOcclusionQueryBatch>()) FRHICommandEndOcclusionQueryBatch();
 	}
 
-#if PLATFORM_SUPPORTS_RHI_THREAD
-	FORCEINLINE_DEBUGGABLE void BeginScene()
-	{
-		if (Bypass())
-		{
-			CMD_CONTEXT(BeginScene)();
-			return;
-		}
-		new (AllocCommand<FRHICommandBeginScene>()) FRHICommandBeginScene();
-	}
-	FORCEINLINE_DEBUGGABLE void EndScene()
-	{
-		if (Bypass())
-		{
-			CMD_CONTEXT(EndScene)();
-			return;
-		}
-		new (AllocCommand<FRHICommandEndScene>()) FRHICommandEndScene();
-	}
+	// These 6 are special in that they must be called on the immediate command list and they force a flush only when we are not doing RHI thread
+	void BeginScene();
+	void EndScene();
 	void BeginDrawingViewport(FViewportRHIParamRef Viewport, FTextureRHIParamRef RenderTargetRHI);
 	void EndDrawingViewport(FViewportRHIParamRef Viewport, bool bPresent, bool bLockToVsync);
 	void BeginFrame();
 	void EndFrame();
-#endif
 
 	FORCEINLINE_DEBUGGABLE void PushEvent(const TCHAR* Name)
 	{
@@ -2007,7 +1950,6 @@ public:
 
 // typedef to mark the recursive use of commandlists in the RHI implementations
 
-#if PLATFORM_RHI_USES_CONTEXT_OBJECT
 class RHI_API FRHICommandList_RecursiveHazardous : public FRHICommandList
 {
 	FRHICommandList_RecursiveHazardous()
@@ -2020,9 +1962,6 @@ public:
 		SetContext(Context);
 	}
 };
-#else
-typedef FRHICommandList FRHICommandList_RecursiveHazardous;
-#endif
 
 
 class RHI_API FRHICommandListExecutor
diff --git a/Engine/Source/Runtime/RHI/Public/RHICommandListCommandExecutes.inl b/Engine/Source/Runtime/RHI/Public/RHICommandListCommandExecutes.inl
--- a/Engine/Source/Runtime/RHI/Public/RHICommandListCommandExecutes.inl
+++ b/Engine/Source/Runtime/RHI/Public/RHICommandListCommandExecutes.inl
@@ -5,11 +5,7 @@
 =============================================================================*/
 
 #if !defined(INTERNAL_DECORATOR)
-	#if PLATFORM_RHI_USES_CONTEXT_OBJECT
-		#define INTERNAL_DECORATOR(Method) CmdList.GetContext().RHI##Method
-	#else
-		#define INTERNAL_DECORATOR(Method) Method##_Internal
-	#endif
+	#define INTERNAL_DECORATOR(Method) CmdList.GetContext().RHI##Method
 #endif
 
 // set this one to get a stat for each RHI command 
@@ -302,11 +298,7 @@ void FRHICommandBuildLocalBoundShaderState::Execute(FRHICommandListBase& CmdList
 	if (WorkArea.ComputedBSS->UseCount)
 	{
 		WorkArea.ComputedBSS->BSS = 	
-#if HAS_THREADSAFE_CreateBoundShaderState
 			RHICreateBoundShaderState(WorkArea.Args.VertexDeclarationRHI, WorkArea.Args.VertexShaderRHI, WorkArea.Args.HullShaderRHI, WorkArea.Args.DomainShaderRHI, WorkArea.Args.PixelShaderRHI, WorkArea.Args.GeometryShaderRHI);
-#else
-			CreateBoundShaderState_Internal(WorkArea.Args.VertexDeclarationRHI, WorkArea.Args.VertexShaderRHI, WorkArea.Args.HullShaderRHI, WorkArea.Args.DomainShaderRHI, WorkArea.Args.PixelShaderRHI, WorkArea.Args.GeometryShaderRHI);
-#endif
 	}
 }
 
@@ -329,11 +321,7 @@ void FRHICommandBuildLocalUniformBuffer::Execute(FRHICommandListBase& CmdList)
 	check(!IsValidRef(WorkArea.ComputedUniformBuffer->UniformBuffer) && WorkArea.Layout && WorkArea.Contents); // should not already have been created
 	if (WorkArea.ComputedUniformBuffer->UseCount)
 	{
-#if PLATFORM_SUPPORTS_RHI_THREAD
 		WorkArea.ComputedUniformBuffer->UniformBuffer = RHICreateUniformBuffer(WorkArea.Contents, *WorkArea.Layout, UniformBuffer_SingleFrame); 
-#else
-		WorkArea.ComputedUniformBuffer->UniformBuffer = CreateUniformBuffer_Internal(WorkArea.Contents, *WorkArea.Layout, UniformBuffer_SingleFrame);
-#endif
 	}
 	WorkArea.Layout = nullptr;
 	WorkArea.Contents = nullptr;
@@ -369,14 +357,6 @@ void FRHICommandEndRenderQuery::Execute(FRHICommandListBase& CmdList)
 	INTERNAL_DECORATOR(EndRenderQuery)(RenderQuery);
 }
 
-#if !PLATFORM_SUPPORTS_RHI_THREAD
-void FRHICommandResetRenderQuery::Execute(FRHICommandListBase& CmdList)
-{
-	RHISTAT(ResetRenderQuery);
-	INTERNAL_DECORATOR(ResetRenderQuery)(RenderQuery);
-}
-#endif
-
 void FRHICommandBeginOcclusionQueryBatch::Execute(FRHICommandListBase& CmdList)
 {
 	RHISTAT(BeginOcclusionQueryBatch);
@@ -392,21 +372,13 @@ void FRHICommandEndOcclusionQueryBatch::Execute(FRHICommandListBase& CmdList)
 void FRHICommandBeginScene::Execute(FRHICommandListBase& CmdList)
 {
 	RHISTAT(BeginScene);
-#if PLATFORM_SUPPORTS_RHI_THREAD
 	INTERNAL_DECORATOR(BeginScene)();
-#else
-	check(0);
-#endif
 }
 
 void FRHICommandEndScene::Execute(FRHICommandListBase& CmdList)
 {
 	RHISTAT(EndScene);
-#if PLATFORM_SUPPORTS_RHI_THREAD
 	INTERNAL_DECORATOR(EndScene)();
-#else
-	check(0);
-#endif
 }
 
 void FRHICommandUpdateVertexBuffer::Execute(FRHICommandListBase& CmdList)
@@ -420,24 +392,15 @@ void FRHICommandUpdateVertexBuffer::Execute(FRHICommandListBase& CmdList)
 void FRHICommandBeginFrame::Execute(FRHICommandListBase& CmdList)
 {
 	RHISTAT(BeginFrame);
-#if PLATFORM_SUPPORTS_RHI_THREAD
 	INTERNAL_DECORATOR(BeginFrame)();
-#else
-	check(0);
-#endif
 }
 
 void FRHICommandEndFrame::Execute(FRHICommandListBase& CmdList)
 {
 	RHISTAT(EndFrame);
-#if PLATFORM_SUPPORTS_RHI_THREAD
 	INTERNAL_DECORATOR(EndFrame)();
-#else
-	check(0);
-#endif
 }
 
-#if PLATFORM_SUPPORTS_RHI_THREAD
 void FRHICommandBeginDrawingViewport::Execute(FRHICommandListBase& CmdList)
 {
 	RHISTAT(BeginDrawingViewport);
@@ -449,7 +412,6 @@ void FRHICommandEndDrawingViewport::Execute(FRHICommandListBase& CmdList)
 	RHISTAT(EndDrawingViewport);
 	INTERNAL_DECORATOR(EndDrawingViewport)(Viewport, bPresent, bLockToVsync);
 }
-#endif
 
 void FRHICommandPushEvent::Execute(FRHICommandListBase& CmdList)
 {
diff --git a/Engine/Source/Runtime/RHI/Public/RHIMethods.h b/Engine/Source/Runtime/RHI/Public/RHIMethods.h
--- a/Engine/Source/Runtime/RHI/Public/RHIMethods.h
+++ b/Engine/Source/Runtime/RHI/Public/RHIMethods.h
@@ -470,6 +470,7 @@ DEFINE_RHIMETHOD_GLOBAL_1(
 
 /**
  * Creates a bound shader state instance which encapsulates a decl, vertex shader, hull shader, domain shader and pixel shader
+ * CAUTION: Even though this is marked as threadsafe, it is only valid to call from the render thread or the RHI thread. Thus is need not be threadsafe on platforms that do not support or aren't using an RHIThread
  * @param VertexDeclaration - existing vertex decl
  * @param VertexShader - existing vertex shader
  * @param HullShader - existing hull shader
@@ -478,37 +479,16 @@ DEFINE_RHIMETHOD_GLOBAL_1(
  * @param PixelShader - existing pixel shader
  */
 
-#if !defined(HAS_THREADSAFE_CreateBoundShaderState)
-	#define HAS_THREADSAFE_CreateBoundShaderState PLATFORM_SUPPORTS_PARALLEL_RHI_EXECUTE
-#endif
-
-#if HAS_THREADSAFE_CreateBoundShaderState != PLATFORM_SUPPORTS_PARALLEL_RHI_EXECUTE
-	#error "HAS_THREADSAFE_CreateBoundShaderState changed meaning"
-#endif
-
-#if HAS_THREADSAFE_CreateBoundShaderState
-	DEFINE_RHIMETHOD_GLOBALTHREADSAFE_6(
-		FBoundShaderStateRHIRef,CreateBoundShaderState,
-		FVertexDeclarationRHIParamRef,VertexDeclaration,
-		FVertexShaderRHIParamRef,VertexShader,
-		FHullShaderRHIParamRef,HullShader,
-		FDomainShaderRHIParamRef,DomainShader,
-		FPixelShaderRHIParamRef,PixelShader,
-		FGeometryShaderRHIParamRef,GeometryShader,
-		return,return new FRHIBoundShaderState();
-	);
-#else
-	DEFINE_RHIMETHOD_GLOBAL_6(
-		FBoundShaderStateRHIRef,CreateBoundShaderState,
-		FVertexDeclarationRHIParamRef,VertexDeclaration,
-		FVertexShaderRHIParamRef,VertexShader,
-		FHullShaderRHIParamRef,HullShader,
-		FDomainShaderRHIParamRef,DomainShader,
-		FPixelShaderRHIParamRef,PixelShader,
-		FGeometryShaderRHIParamRef,GeometryShader,
-		return,return new FRHIBoundShaderState();
-		);
-#endif
+DEFINE_RHIMETHOD_GLOBALTHREADSAFE_6(
+	FBoundShaderStateRHIRef,CreateBoundShaderState,
+	FVertexDeclarationRHIParamRef,VertexDeclaration,
+	FVertexShaderRHIParamRef,VertexShader,
+	FHullShaderRHIParamRef,HullShader,
+	FDomainShaderRHIParamRef,DomainShader,
+	FPixelShaderRHIParamRef,PixelShader,
+	FGeometryShaderRHIParamRef,GeometryShader,
+	return,return new FRHIBoundShaderState();
+);
 
 /**
  *Sets the current compute shader.  Mostly for compliance with platforms
@@ -558,28 +538,18 @@ DEFINE_RHIMETHOD_CMDLIST_2(
 
 /**
  * Creates a uniform buffer.  The contents of the uniform buffer are provided in a parameter, and are immutable.
+ * CAUTION: Even though this is marked as threadsafe, it is only valid to call from the render thread or the RHI thread. Thus is need not be threadsafe on platforms that do not support or aren't using an RHIThread
  * @param Contents - A pointer to a memory block of size NumBytes that is copied into the new uniform buffer.
  * @param NumBytes - The number of bytes the uniform buffer should contain.
  * @return The new uniform buffer.
  */
-#if PLATFORM_SUPPORTS_RHI_THREAD
-	DEFINE_RHIMETHOD_GLOBALTHREADSAFE_3(
-		FUniformBufferRHIRef,CreateUniformBuffer,
-		const void*,Contents,
-		const FRHIUniformBufferLayout&,Layout,
-		EUniformBufferUsage,Usage,
-		return,return new FRHIUniformBuffer(Layout);
-	);
-#else
-	DEFINE_RHIMETHOD_GLOBAL_3(
-		FUniformBufferRHIRef,CreateUniformBuffer,
-		const void*,Contents,
-		const FRHIUniformBufferLayout&,Layout,
-		EUniformBufferUsage,Usage,
-		return,return new FRHIUniformBuffer(Layout);
-	);
-#endif
-
+DEFINE_RHIMETHOD_GLOBALTHREADSAFE_3(
+	FUniformBufferRHIRef,CreateUniformBuffer,
+	const void*,Contents,
+	const FRHIUniformBufferLayout&,Layout,
+	EUniformBufferUsage,Usage,
+	return,return new FRHIUniformBuffer(Layout);
+);
 
 DEFINE_RHIMETHOD_GLOBAL_4(
 	FIndexBufferRHIRef,CreateIndexBuffer,
@@ -1292,14 +1262,6 @@ DEFINE_RHIMETHOD_GLOBAL_1(
 	return,return new FRHIRenderQuery();
 	);
 
-#if !PLATFORM_SUPPORTS_RHI_THREAD
-DEFINE_RHIMETHOD_CMDLIST_1(
-	void,ResetRenderQuery,
-	FRenderQueryRHIParamRef,RenderQuery,
-	,
-	);
-#endif
-
 DEFINE_RHIMETHOD_CMDLIST_1(
 	void,BeginRenderQuery,
 	FRenderQueryRHIParamRef,RenderQuery,
@@ -1310,23 +1272,15 @@ DEFINE_RHIMETHOD_CMDLIST_1(
 	FRenderQueryRHIParamRef,RenderQuery,
 	,
 	);
-#if PLATFORM_HAS_THREADSAFE_RHIGetRenderQueryResult
-	DEFINE_RHIMETHOD_GLOBALTHREADSAFE_3(
-		bool,GetRenderQueryResult,
-		FRenderQueryRHIParamRef,RenderQuery,
-		uint64&,OutResult,
-		bool,bWait,
-		return,return true;
-	);
-#else
-	DEFINE_RHIMETHOD_3(
-		bool,GetRenderQueryResult,
-		FRenderQueryRHIParamRef,RenderQuery,
-		uint64&,OutResult,
-		bool,bWait,
-		return,return true;
-	);
-#endif
+
+// CAUTION: Even though this is marked as threadsafe, it is only valid to call from the render thread. It is need not be threadsafe on platforms that do not support or aren't using an RHIThread
+DEFINE_RHIMETHOD_GLOBALTHREADSAFE_3(
+	bool,GetRenderQueryResult,
+	FRenderQueryRHIParamRef,RenderQuery,
+	uint64&,OutResult,
+	bool,bWait,
+	return,return true;
+);
 
 DEFINE_RHIMETHOD_CMDLIST_0(
 	void,BeginOcclusionQueryBatch,
@@ -1337,102 +1291,65 @@ DEFINE_RHIMETHOD_CMDLIST_0(
 	,
 	);
 
-#if PLATFORM_SUPPORTS_RHI_THREAD
-	DEFINE_RHIMETHOD_CMDLIST_2(
-		void,BeginDrawingViewport,
-		FViewportRHIParamRef,Viewport,
-		FTextureRHIParamRef,RenderTargetRHI,
-		,
-		);
-	DEFINE_RHIMETHOD_CMDLIST_3(
-		void,EndDrawingViewport,
-		FViewportRHIParamRef,Viewport,
-		bool,bPresent,
-		bool,bLockToVsync,
-		,
-		);
-	DEFINE_RHIMETHOD_GLOBALTHREADSAFE_1(
-		FTexture2DRHIRef,GetViewportBackBuffer,
-		FViewportRHIParamRef,Viewport,
-		return,return new FRHITexture2D(1,1,1,1,PF_B8G8R8A8,TexCreate_RenderTargetable);
-	);
-	DEFINE_RHIMETHOD_CMDLIST_0(
-		void,BeginFrame,
-		return,return;
-	);
-
-	DEFINE_RHIMETHOD_CMDLIST_0(
-		void,EndFrame,
-		return,return;
+// This method is queued with an RHIThread, otherwise it will flush after it is queued; without an RHI thread there is no benefit to queuing this frame advance commands
+DEFINE_RHIMETHOD_CMDLIST_2(
+	void,BeginDrawingViewport,
+	FViewportRHIParamRef,Viewport,
+	FTextureRHIParamRef,RenderTargetRHI,
+	,
 	);
-	/**
-	 * Signals the beginning of scene rendering. The RHI makes certain caching assumptions between
-	 * calls to BeginScene/EndScene. Currently the only restriction is that you can't update texture
-	 * references.
-	 */
-	DEFINE_RHIMETHOD_CMDLIST_0(
-		void,BeginScene,
-		return,return;
-		);
+// This method is queued with an RHIThread, otherwise it will flush after it is queued; without an RHI thread there is no benefit to queuing this frame advance commands
+DEFINE_RHIMETHOD_CMDLIST_3(
+	void,EndDrawingViewport,
+	FViewportRHIParamRef,Viewport,
+	bool,bPresent,
+	bool,bLockToVsync,
+	,
+	);
+// With RHI thread, this is the current backbuffer from the perspecitve of the render thread.
+DEFINE_RHIMETHOD_GLOBALTHREADSAFE_1(
+	FTexture2DRHIRef,GetViewportBackBuffer,
+	FViewportRHIParamRef,Viewport,
+	return,return new FRHITexture2D(1,1,1,1,PF_B8G8R8A8,TexCreate_RenderTargetable);
+);
+// This method is queued with an RHIThread, otherwise it will flush after it is queued; without an RHI thread there is no benefit to queuing this frame advance commands
+DEFINE_RHIMETHOD_CMDLIST_0(
+	void,BeginFrame,
+	return,return;
+);
 
-	/**
-	 * Signals the end of scene rendering. See RHIBeginScene.
-	 */
-	DEFINE_RHIMETHOD_CMDLIST_0(
-		void,EndScene,
-		return,return;
-		);
-
-	DEFINE_RHIMETHOD_GLOBALTHREADSAFE_0(
-		void,AdvanceFrameForGetViewportBackBuffer,
-		return,return;
-		);
+// This method is queued with an RHIThread, otherwise it will flush after it is queued; without an RHI thread there is no benefit to queuing this frame advance commands
+DEFINE_RHIMETHOD_CMDLIST_0(
+	void,EndFrame,
+	return,return;
+);
+/**
+	* Signals the beginning of scene rendering. The RHI makes certain caching assumptions between
+	* calls to BeginScene/EndScene. Currently the only restriction is that you can't update texture
+	* references.
+	*/
+// This method is queued with an RHIThread, otherwise it will flush after it is queued; without an RHI thread there is no benefit to queuing this frame advance commands
+DEFINE_RHIMETHOD_CMDLIST_0(
+	void,BeginScene,
+	return,return;
+	);
 
-#else
-	DEFINE_RHIMETHOD_2(
-		void,BeginDrawingViewport,
-		FViewportRHIParamRef,Viewport,
-		FTextureRHIParamRef,RenderTargetRHI,
-		,
-		);
-	DEFINE_RHIMETHOD_3(
-		void,EndDrawingViewport,
-		FViewportRHIParamRef,Viewport,
-		bool,bPresent,
-		bool,bLockToVsync,
-		,
-		);
-	DEFINE_RHIMETHOD_GLOBALFLUSH_1(
-		FTexture2DRHIRef,GetViewportBackBuffer,
-		FViewportRHIParamRef,Viewport,
-		return,return new FRHITexture2D(1,1,1,1,PF_B8G8R8A8,TexCreate_RenderTargetable);
-	);
-	DEFINE_RHIMETHOD_0(
-		void,BeginFrame,
-		return,return;
-	);
-	DEFINE_RHIMETHOD_0(
-		void,EndFrame,
-		return,return;
+/**
+	* Signals the end of scene rendering. See RHIBeginScene.
+	*/
+// This method is queued with an RHIThread, otherwise it will flush after it is queued; without an RHI thread there is no benefit to queuing this frame advance commands
+DEFINE_RHIMETHOD_CMDLIST_0(
+	void,EndScene,
+	return,return;
+	);
+
+// Only relevant with an RHI thread, this advances the backbuffer for the purpose of GetViewportBackBuffer
+DEFINE_RHIMETHOD_GLOBALTHREADSAFE_0(
+	void,AdvanceFrameForGetViewportBackBuffer,
+	return,return;
 	);
-	/**
-	 * Signals the beginning of scene rendering. The RHI makes certain caching assumptions between
-	 * calls to BeginScene/EndScene. Currently the only restriction is that you can't update texture
-	 * references.
-	 */
-	DEFINE_RHIMETHOD_0(
-		void,BeginScene,
-		return,return;
-		);
 
-	/**
-	 * Signals the end of scene rendering. See RHIBeginScene.
-	 */
-	DEFINE_RHIMETHOD_0(
-		void,EndScene,
-		return,return;
-		);
-#endif
+
 
 /**
  * Determine if currently drawing the viewport
@@ -2186,19 +2103,14 @@ DEFINE_RHIMETHOD_CMDLIST_0(
 	,
 	);
 
-#if PLATFORM_RHI_USES_CONTEXT_OBJECT
-	DEFINE_RHIMETHOD_GLOBALTHREADSAFE_0(
-		class IRHICommandContext*,GetDefaultContext,
-		return,return this
+DEFINE_RHIMETHOD_GLOBALTHREADSAFE_0(
+	class IRHICommandContext*,GetDefaultContext,
+	return,return this
 	);
 
-	#if PLATFORM_SUPPORTS_PARALLEL_RHI_EXECUTE
-			DEFINE_RHIMETHOD_GLOBALTHREADSAFE_1(
-				class IRHICommandContextContainer*,GetCommandContextContainer,
-				uint32, ClearTag,
-				return,return nullptr
-				);
-	#endif
+DEFINE_RHIMETHOD_GLOBALTHREADSAFE_0(
+	class IRHICommandContextContainer*,GetCommandContextContainer,
+	return,return nullptr
+	);
 
-#endif
 
diff --git a/Engine/Source/Runtime/RHI/Public/RHIStaticStates.h b/Engine/Source/Runtime/RHI/Public/RHIStaticStates.h
--- a/Engine/Source/Runtime/RHI/Public/RHIStaticStates.h
+++ b/Engine/Source/Runtime/RHI/Public/RHIStaticStates.h
@@ -64,27 +64,30 @@ public:
 
 		if (!StaticResource)
 		{
-#if PLATFORM_SUPPORTS_RHI_THREAD
-			FStaticStateResource* NewStaticResource = new FStaticStateResource();
-			FStaticStateResource* ValueWas = (FStaticStateResource*)FPlatformAtomics::InterlockedCompareExchangePointer((void**)&StaticResource, NewStaticResource, nullptr);
-			if (ValueWas)
+			if (GIsRHIInitialized && GRHISupportsRHIThread)
 			{
-				// we made a redundant one...leak it
-			}
-#else
-			if (!IsInRenderingThread())
-			{
-				check(IsInParallelRenderingThread());
+				FStaticStateResource* NewStaticResource = new FStaticStateResource();
+				FStaticStateResource* ValueWas = (FStaticStateResource*)FPlatformAtomics::InterlockedCompareExchangePointer((void**)&StaticResource, NewStaticResource, nullptr);
+				if (ValueWas)
 				{
-					FScopedEvent Event;
-					TGraphTask<FInitStaticResourceRenderThreadTask>::CreateTask().ConstructAndDispatchWhenReady(&GetRHI_WithNoReturnValue, Event);
+					// we made a redundant one...leak it
 				}
 			}
 			else
 			{
-				StaticResource = new FStaticStateResource();
+				if (!IsInRenderingThread())
+				{
+					check(IsInParallelRenderingThread());
+					{
+						FScopedEvent Event;
+						TGraphTask<FInitStaticResourceRenderThreadTask>::CreateTask().ConstructAndDispatchWhenReady(&GetRHI_WithNoReturnValue, Event);
+					}
+				}
+				else
+				{
+					StaticResource = new FStaticStateResource();
+				}
 			}
-#endif
 		}
 		return StaticResource->StateRHI;
 	};
@@ -96,47 +99,36 @@ private:
 	{
 	public:
 		RHIRefType StateRHI;
-#if PLATFORM_SUPPORTS_RHI_THREAD
 		FStaticStateResource()
 		{
-			StateRHI = InitializerType::CreateRHI();
-		}
-
-		// FRenderResource interface.
-		virtual void InitRHI() override
-		{
-			check(0);
-		}
-		virtual void ReleaseRHI() override
-		{
-			check(0);
-		}
+			if (GIsRHIInitialized && GRHISupportsRHIThread)
+			{
+				StateRHI = InitializerType::CreateRHI();
+			}
+			else
+			{
+				InitResource();
 
-		~FStaticStateResource()
-		{
-			check(0);
-		}
-#else
-		FStaticStateResource()
-		{
-			InitResource();
+			}
 		}
 
 		// FRenderResource interface.
 		virtual void InitRHI() override
 		{
+			check(!GIsRHIInitialized || !GRHISupportsRHIThread);
 			StateRHI = InitializerType::CreateRHI();
 		}
 		virtual void ReleaseRHI() override
 		{
+			check(!GIsRHIInitialized || !GRHISupportsRHIThread);
 			StateRHI.SafeRelease();
 		}
 
 		~FStaticStateResource()
 		{
+			check(!GIsRHIInitialized || !GRHISupportsRHIThread);
 			ReleaseResource();
 		}
-#endif
 	};
 };
 
diff --git a/Engine/Source/Runtime/RenderCore/Private/RenderingThread.cpp b/Engine/Source/Runtime/RenderCore/Private/RenderingThread.cpp
--- a/Engine/Source/Runtime/RenderCore/Private/RenderingThread.cpp
+++ b/Engine/Source/Runtime/RenderCore/Private/RenderingThread.cpp
@@ -216,7 +216,6 @@ uint32 GRenderThreadTime = 0;
 
 
 
-#if PLATFORM_SUPPORTS_RHI_THREAD
 /** The RHI thread runnable object. */
 class FRHIThread : public FRunnable
 {
@@ -251,8 +250,6 @@ public:
 	}
 };
 
-#endif
-
 /** The rendering thread main loop */
 void RenderingThreadMain( FEvent* TaskGraphBoundSyncEvent )
 {
@@ -510,7 +507,6 @@ void StartRenderingThread()
 	check(!GIsThreadedRendering && GUseThreadedRendering);
 
 	check(!GRHIThread)
-#if PLATFORM_SUPPORTS_RHI_THREAD
 	if (GUseRHIThread)
 	{
 		if (!FTaskGraphInterface::Get().IsThreadProcessingTasks(ENamedThreads::RHIThread))
@@ -526,7 +522,6 @@ void StartRenderingThread()
 		check(GRHIThread);
 		GRHICommandList.LatchBypass();
 	}
-#endif
 
 	// Turn on the threaded rendering flag.
 	GIsThreadedRendering = true;
@@ -589,7 +584,6 @@ void StopRenderingThread()
 		// The rendering thread may have already been stopped during the call to GFlushStreamingFunc or FlushRenderingCommands.
 		if ( GIsThreadedRendering )
 		{
-#if PLATFORM_SUPPORTS_RHI_THREAD
 			if (GRHIThread)
 			{
 				DECLARE_CYCLE_STAT(TEXT("Wait For RHIThread Finish"), STAT_WaitForRHIThreadFinish, STATGROUP_TaskGraphTasks);
@@ -598,8 +592,6 @@ void StopRenderingThread()
 				FTaskGraphInterface::Get().WaitUntilTaskCompletes(FlushTask, ENamedThreads::GameThread_Local);
 				GRHIThread = nullptr;
 			}
-#endif
-
 
 			check( GRenderingThread );
 			check(!GIsRenderingThreadSuspended);
diff --git a/Engine/Source/Runtime/RenderCore/Public/RenderResource.h b/Engine/Source/Runtime/RenderCore/Public/RenderResource.h
--- a/Engine/Source/Runtime/RenderCore/Public/RenderResource.h
+++ b/Engine/Source/Runtime/RenderCore/Public/RenderResource.h
@@ -606,9 +606,6 @@ private:
  * A list of the most recently used bound shader states.
  * This is used to keep bound shader states that have been used recently from being freed, as they're likely to be used again soon.
  */
-#if !defined(HAS_THREADSAFE_CreateBoundShaderState)
-#error "HAS_THREADSAFE_CreateBoundShaderState must be defined"
-#endif
 
 template<uint32 Size>
 class TBoundShaderStateHistory : public FRenderResource
@@ -623,16 +620,21 @@ public:
 	/** Adds a bound shader state to the history. */
 	void Add(FBoundShaderStateRHIParamRef BoundShaderState)
 	{
-#if HAS_THREADSAFE_CreateBoundShaderState
-		FScopeLock Lock(&BoundShaderStateHistoryLock);
-#endif
+		if (GRHISupportsParallelRHIExecute)
+		{
+			BoundShaderStateHistoryLock.Lock();
+		}
 		BoundShaderStates[NextBoundShaderStateIndex] = BoundShaderState;
 		NextBoundShaderStateIndex = (NextBoundShaderStateIndex + 1) % Size;
+		if (GRHISupportsParallelRHIExecute)
+		{
+			BoundShaderStateHistoryLock.Unlock();
+		}
 	}
 
 	FBoundShaderStateRHIParamRef GetLast()
 	{
-		check(!HAS_THREADSAFE_CreateBoundShaderState);
+		check(!GRHISupportsParallelRHIExecute);
 		// % doesn't work as we want on negative numbers, so handle the wraparound manually
 		uint32 LastIndex = NextBoundShaderStateIndex == 0 ? Size - 1 : NextBoundShaderStateIndex - 1;
 		return BoundShaderStates[LastIndex];
@@ -641,17 +643,23 @@ public:
 	// FRenderResource interface.
 	virtual void ReleaseRHI()
 	{
+		if (GRHISupportsParallelRHIExecute)
+		{
+			BoundShaderStateHistoryLock.Lock();
+		}
 		for(uint32 Index = 0;Index < Size;Index++)
 		{
 			BoundShaderStates[Index].SafeRelease();
 		}
+		if (GRHISupportsParallelRHIExecute)
+		{
+			BoundShaderStateHistoryLock.Unlock();
+		}
 	}
 
 private:
 
 	FBoundShaderStateRHIRef BoundShaderStates[Size];
 	uint32 NextBoundShaderStateIndex;
-#if HAS_THREADSAFE_CreateBoundShaderState
 	FCriticalSection BoundShaderStateHistoryLock;
-#endif
 };
diff --git a/Engine/Source/Runtime/Renderer/Private/SceneOcclusion.cpp b/Engine/Source/Runtime/Renderer/Private/SceneOcclusion.cpp
--- a/Engine/Source/Runtime/Renderer/Private/SceneOcclusion.cpp
+++ b/Engine/Source/Runtime/Renderer/Private/SceneOcclusion.cpp
@@ -79,11 +79,6 @@ void FRenderQueryPool::ReleaseQuery(FRHICommandListImmediate& RHICmdList, FRende
 		{
 			// Return it to the pool.
 			Queries.Add( Query );
-
-#if !PLATFORM_SUPPORTS_RHI_THREAD
-			// Tell RHI we don't need the result anymore.
-			RHICmdList.ResetRenderQuery(Query);
-#endif
 		}
 
 		// De-ref without deleting.
diff --git a/Engine/Source/Runtime/Renderer/Private/SceneRendering.cpp b/Engine/Source/Runtime/Renderer/Private/SceneRendering.cpp
--- a/Engine/Source/Runtime/Renderer/Private/SceneRendering.cpp
+++ b/Engine/Source/Runtime/Renderer/Private/SceneRendering.cpp
@@ -127,7 +127,7 @@ FParallelCommandListSet::FParallelCommandListSet(const FViewInfo& InView, FRHICo
 	, ParentCmdList(InParentCmdList)
 	, OutDirtyIfIgnored(false)
 	, OutDirty(InOutDirty ? *InOutDirty : OutDirtyIfIgnored)
-	, bParallelExecute(bInParallelExecute)
+	, bParallelExecute(GRHISupportsParallelRHIExecute && bInParallelExecute)
 {
 	Width = CVarRHICmdWidth.GetValueOnRenderThread();
 	CommandLists.Reserve(Width * 8);
@@ -138,14 +138,13 @@ FParallelCommandListSet::FParallelCommandListSet(const FViewInfo& InView, FRHICo
 void FParallelCommandListSet::Dispatch()
 {
 	check(CommandLists.Num() == Events.Num());
-#if PLATFORM_SUPPORTS_PARALLEL_RHI_EXECUTE
 	if (bParallelExecute && CommandLists.Num())
 	{
+		check(GRHISupportsParallelRHIExecute);
 		ParentCmdList.QueueParallelAsyncCommandListSubmit(&Events[0], &CommandLists[0], CommandLists.Num());
 		SetStateOnCommandList(ParentCmdList);
 	}
 	else
-#endif
 	{
 		for (int32 Index = 0; Index < CommandLists.Num(); Index++)
 		{
diff --git a/Engine/Source/Runtime/Windows/D3D11RHI/Private/D3D11Commands.cpp b/Engine/Source/Runtime/Windows/D3D11RHI/Private/D3D11Commands.cpp
--- a/Engine/Source/Runtime/Windows/D3D11RHI/Private/D3D11Commands.cpp
+++ b/Engine/Source/Runtime/Windows/D3D11RHI/Private/D3D11Commands.cpp
@@ -55,68 +55,6 @@ DECLARE_ISBOUNDSHADER(ComputeShader)
 	#include <GPUPerfAPI/Gpa.h>
 #endif
 
-#if PLATFORM_SUPPORTS_RHI_THREAD
-void FD3D11DynamicRHI::SetupRecursiveResources()
-{
-	FRHICommandList_RecursiveHazardous RHICmdList;
-	extern int32 GCreateShadersOnLoad;
-	TGuardValue<int32> Guard(GCreateShadersOnLoad, 1);
-	auto ShaderMap = GetGlobalShaderMap(GMaxRHIFeatureLevel);
-	TShaderMapRef<TOneColorVS<true> > VertexShader(ShaderMap);
-	GD3D11Vector4VertexDeclaration.InitRHI();
-
-	for (int32 NumBuffers = 1; NumBuffers <= MaxSimultaneousRenderTargets; NumBuffers++)
-	{
-		FOneColorPS* PixelShader = NULL;
-
-		// Set the shader to write to the appropriate number of render targets
-		// On AMD PC hardware, outputting to a color index in the shader without a matching render target set has a significant performance hit
-		if (NumBuffers <= 1)
-		{
-			TShaderMapRef<TOneColorPixelShaderMRT<1> > MRTPixelShader(ShaderMap);
-			PixelShader = *MRTPixelShader;
-		}
-		else if (NumBuffers == 2)
-		{
-			TShaderMapRef<TOneColorPixelShaderMRT<2> > MRTPixelShader(ShaderMap);
-			PixelShader = *MRTPixelShader;
-		}
-		else if (NumBuffers== 3)
-		{
-			TShaderMapRef<TOneColorPixelShaderMRT<3> > MRTPixelShader(ShaderMap);
-			PixelShader = *MRTPixelShader;
-		}
-		else if (NumBuffers == 4)
-		{
-			TShaderMapRef<TOneColorPixelShaderMRT<4> > MRTPixelShader(ShaderMap);
-			PixelShader = *MRTPixelShader;
-		}
-		else if (NumBuffers == 5)
-		{
-			TShaderMapRef<TOneColorPixelShaderMRT<5> > MRTPixelShader(ShaderMap);
-			PixelShader = *MRTPixelShader;
-		}
-		else if (NumBuffers == 6)
-		{
-			TShaderMapRef<TOneColorPixelShaderMRT<6> > MRTPixelShader(ShaderMap);
-			PixelShader = *MRTPixelShader;
-		}
-		else if (NumBuffers == 7)
-		{
-			TShaderMapRef<TOneColorPixelShaderMRT<7> > MRTPixelShader(ShaderMap);
-			PixelShader = *MRTPixelShader;
-		}
-		else if (NumBuffers == 8)
-		{
-			TShaderMapRef<TOneColorPixelShaderMRT<8> > MRTPixelShader(ShaderMap);
-			PixelShader = *MRTPixelShader;
-		}
-
-		SetGlobalBoundShaderState(RHICmdList, GMaxRHIFeatureLevel, GD3D11ClearMRTBoundShaderState[NumBuffers - 1], GD3D11Vector4VertexDeclaration.VertexDeclarationRHI, *VertexShader, PixelShader);
-	}
-}
-#endif
-
 void FD3D11DynamicRHI::RHIGpuTimeBegin(uint32 Hash, bool bCompute)
 {
 	#if WITH_GPA
@@ -1819,7 +1757,7 @@ void FD3D11DynamicRHI::RHIClearMRT(bool bClearColor,int32 NumClearColors,const F
 		}
 
 		{
-			FRHICommandList_RecursiveHazardous RHICmdList;
+			FRHICommandList_RecursiveHazardous RHICmdList(this);
 			SetGlobalBoundShaderState(RHICmdList, GMaxRHIFeatureLevel, GD3D11ClearMRTBoundShaderState[FMath::Max(BoundRenderTargets.GetNumActiveTargets() - 1, 0)], GD3D11Vector4VertexDeclaration.VertexDeclarationRHI, *VertexShader, PixelShader);
 			FLinearColor ShaderClearColors[MaxSimultaneousRenderTargets];
 			FMemory::Memzero(ShaderClearColors);
@@ -1992,3 +1930,14 @@ void FD3D11DynamicRHI::RHIEnableDepthBoundsTest(bool bEnable,float MinDepth,floa
 	}
 #endif
 }
+
+IRHICommandContext* FD3D11DynamicRHI::RHIGetDefaultContext()
+{
+	return this;
+}
+
+IRHICommandContextContainer* FD3D11DynamicRHI::RHIGetCommandContextContainer()
+{
+	return nullptr;
+}
+
diff --git a/Engine/Source/Runtime/Windows/D3D11RHI/Private/D3D11Query.cpp b/Engine/Source/Runtime/Windows/D3D11RHI/Private/D3D11Query.cpp
--- a/Engine/Source/Runtime/Windows/D3D11RHI/Private/D3D11Query.cpp
+++ b/Engine/Source/Runtime/Windows/D3D11RHI/Private/D3D11Query.cpp
@@ -37,17 +37,9 @@ FRenderQueryRHIRef FD3D11DynamicRHI::RHICreateRenderQuery(ERenderQueryType Query
 	return new FD3D11OcclusionQuery(Query, QueryType);
 }
 
-#if !PLATFORM_SUPPORTS_RHI_THREAD
-void FD3D11DynamicRHI::RHIResetRenderQuery(FRenderQueryRHIParamRef QueryRHI)
-{
-	DYNAMIC_CAST_D3D11RESOURCE(OcclusionQuery,Query);
-
-	Query->bResultIsCached = false;
-}
-#endif
-
 bool FD3D11DynamicRHI::RHIGetRenderQueryResult(FRenderQueryRHIParamRef QueryRHI,uint64& OutResult,bool bWait)
 {
+	check(IsInRenderingThread());
 	DYNAMIC_CAST_D3D11RESOURCE(OcclusionQuery,Query);
 
 	bool bSuccess = true;
@@ -81,9 +73,7 @@ void FD3D11DynamicRHI::RHIBeginRenderQuery(FRenderQueryRHIParamRef QueryRHI)
 
 	if(Query->QueryType == RQT_Occlusion)
 	{
-#if PLATFORM_SUPPORTS_RHI_THREAD
 		Query->bResultIsCached = false;
-#endif
 		Direct3DDeviceIMContext->Begin(Query->Resource);
 	}
 	else
@@ -96,11 +86,7 @@ void FD3D11DynamicRHI::RHIBeginRenderQuery(FRenderQueryRHIParamRef QueryRHI)
 void FD3D11DynamicRHI::RHIEndRenderQuery(FRenderQueryRHIParamRef QueryRHI)
 {
 	DYNAMIC_CAST_D3D11RESOURCE(OcclusionQuery,Query);
-#if PLATFORM_SUPPORTS_RHI_THREAD
 	Query->bResultIsCached = false; // for occlusion queries, this is redundant with the one in begin
-#elif PLATFORM_HAS_THREADSAFE_RHIGetRenderQueryResult
-	#error "PLATFORM_HAS_THREADSAFE_RHIGetRenderQueryResult requires PLATFORM_SUPPORTS_RHI_THREAD"
-#endif
 	Direct3DDeviceIMContext->End(Query->Resource);
 
 	//@todo - d3d debug spews warnings about OQ's that are being issued but not polled, need to investigate
diff --git a/Engine/Source/Runtime/Windows/D3D11RHI/Private/D3D11RHI.cpp b/Engine/Source/Runtime/Windows/D3D11RHI/Private/D3D11RHI.cpp
--- a/Engine/Source/Runtime/Windows/D3D11RHI/Private/D3D11RHI.cpp
+++ b/Engine/Source/Runtime/Windows/D3D11RHI/Private/D3D11RHI.cpp
@@ -288,7 +288,7 @@ void FD3D11DynamicRHI::IssueLongGPUTask()
 		{
 			FD3D11Viewport* Viewport = Viewports[LargestViewportIndex];
 
-			FRHICommandList_RecursiveHazardous RHICmdList;
+			FRHICommandList_RecursiveHazardous RHICmdList(this);
 
 			SetRenderTarget(RHICmdList, Viewport->GetBackBuffer(), FTextureRHIRef());
 			RHICmdList.SetBlendState(TStaticBlendState<CW_RGBA, BO_Add, BF_One, BF_One>::GetRHI(), FLinearColor::Black);
diff --git a/Engine/Source/Runtime/Windows/D3D11RHI/Private/D3D11RHIPrivate.h b/Engine/Source/Runtime/Windows/D3D11RHI/Private/D3D11RHIPrivate.h
--- a/Engine/Source/Runtime/Windows/D3D11RHI/Private/D3D11RHIPrivate.h
+++ b/Engine/Source/Runtime/Windows/D3D11RHI/Private/D3D11RHIPrivate.h
@@ -165,20 +165,12 @@ public:
 		Timing(InRHI, 1)
 	{
 		// Initialize Buffered timestamp queries 
-#if PLATFORM_SUPPORTS_RHI_THREAD
-		Timing.InitDynamicRHI();
-#else
 		Timing.InitResource(); // can't do this from the RHI thread
-#endif
 	}
 
 	virtual ~FD3D11EventNode()
 	{
-#if PLATFORM_SUPPORTS_RHI_THREAD
-		Timing.ReleaseDynamicRHI();
-#else
 		Timing.ReleaseResource();  // can't do this from the RHI thread
-#endif
 	}
 
 	/** 
@@ -211,26 +203,14 @@ public:
 		RootEventTiming(InRHI, 1),
 		DisjointQuery(InRHI)
 	{
-#if PLATFORM_SUPPORTS_RHI_THREAD
-		RootEventTiming.InitDynamicRHI();
-		DisjointQuery.InitDynamicRHI();
-#else
-		// can't do these on the RHI thread
 		RootEventTiming.InitResource();
 		DisjointQuery.InitResource();
-#endif
 	}
 
 	~FD3D11EventNodeFrame()
 	{
-#if PLATFORM_SUPPORTS_RHI_THREAD
-		RootEventTiming.ReleaseDynamicRHI();
-		DisjointQuery.ReleaseDynamicRHI();
-#else
-		// can't do these on the RHI thread
 		RootEventTiming.ReleaseResource();
 		DisjointQuery.ReleaseResource();
-#endif
 	}
 
 	/** Start this frame of per tracking */
@@ -259,22 +239,12 @@ public:
 	FD3D11GPUProfile(FD3D11DynamicRHI* D3D11RHI)
 	: DisjointQuery(D3D11RHI)
 	{
-#if PLATFORM_SUPPORTS_RHI_THREAD
-		DisjointQuery.InitDynamicRHI();
-#else
-		// can't do these on the RHI thread
-		DisjointQuery.InitResource();
-#endif
+		  DisjointQuery.InitResource();
 	}
 
 	~FD3D11GPUProfile()
 	{
-#if PLATFORM_SUPPORTS_RHI_THREAD
-		DisjointQuery.ReleaseDynamicRHI();
-#else
-		// can't do these on the RHI thread
-		DisjointQuery.ReleaseResource();
-#endif
+		  DisjointQuery.ReleaseResource();
 	}
 
 	FD3D11DisjointTimeStampQuery DisjointQuery;
@@ -333,7 +303,7 @@ struct FD3DGPUProfiler : public FGPUProfiler
 };
 
 /** The interface which is implemented by the dynamically bound RHI. */
-class FD3D11DynamicRHI : public FDynamicRHI
+class FD3D11DynamicRHI : public FDynamicRHI, public IRHICommandContext
 {
 public:
 
@@ -645,10 +615,6 @@ protected:
 
 	void ReadSurfaceDataMSAARaw(FRHICommandList_RecursiveHazardous& RHICmdList, FTextureRHIParamRef TextureRHI, FIntRect Rect, TArray<uint8>& OutData, FReadSurfaceDataFlags InFlags);
 
-#if PLATFORM_SUPPORTS_RHI_THREAD
-	void SetupRecursiveResources();
-#endif
-
 	friend struct FD3DGPUProfiler;
 };
 
diff --git a/Engine/Source/Runtime/Windows/D3D11RHI/Private/D3D11RenderTarget.cpp b/Engine/Source/Runtime/Windows/D3D11RHI/Private/D3D11RenderTarget.cpp
--- a/Engine/Source/Runtime/Windows/D3D11RHI/Private/D3D11RenderTarget.cpp
+++ b/Engine/Source/Runtime/Windows/D3D11RHI/Private/D3D11RenderTarget.cpp
@@ -194,7 +194,7 @@ void FD3D11DynamicRHI::RHICopyToResolveTarget(FTextureRHIParamRef SourceTextureR
 		return;
 	}
 
-	FRHICommandList_RecursiveHazardous RHICmdList;
+	FRHICommandList_RecursiveHazardous RHICmdList(this);
 
 
 	FD3D11Texture2D* SourceTexture2D = static_cast<FD3D11Texture2D*>(SourceTextureRHI->GetTexture2D());
@@ -928,7 +928,7 @@ void FD3D11DynamicRHI::RHIReadSurfaceData(FTextureRHIParamRef TextureRHI,FIntRec
 	}
 	else
 	{
-		FRHICommandList_RecursiveHazardous RHICmdList;
+		FRHICommandList_RecursiveHazardous RHICmdList(this);
 		ReadSurfaceDataMSAARaw(RHICmdList, TextureRHI, InRect, OutDataRaw, InFlags);
 	}
 
diff --git a/Engine/Source/Runtime/Windows/D3D11RHI/Private/D3D11Shaders.cpp b/Engine/Source/Runtime/Windows/D3D11RHI/Private/D3D11Shaders.cpp
--- a/Engine/Source/Runtime/Windows/D3D11RHI/Private/D3D11Shaders.cpp
+++ b/Engine/Source/Runtime/Windows/D3D11RHI/Private/D3D11Shaders.cpp
@@ -264,6 +264,8 @@ FBoundShaderStateRHIRef FD3D11DynamicRHI::RHICreateBoundShaderState(
 	FGeometryShaderRHIParamRef GeometryShaderRHI
 	)
 {
+	check(IsInRenderingThread());
+
 	SCOPE_CYCLE_COUNTER(STAT_D3D11CreateBoundShaderStateTime);
 
 	checkf(GIsRHIInitialized && Direct3DDeviceIMContext,(TEXT("Bound shader state RHI resource was created without initializing Direct3D first")));
diff --git a/Engine/Source/Runtime/Windows/D3D11RHI/Private/D3D11State.cpp b/Engine/Source/Runtime/Windows/D3D11RHI/Private/D3D11State.cpp
--- a/Engine/Source/Runtime/Windows/D3D11RHI/Private/D3D11State.cpp
+++ b/Engine/Source/Runtime/Windows/D3D11RHI/Private/D3D11State.cpp
@@ -8,20 +8,11 @@
 
 // Cache of Sampler States; we store pointers to both as we don't want the TMap to be artificially
 // modifying ref counts if not needed; so we manage that ourselves
-#if PLATFORM_SUPPORTS_RHI_THREAD
-FCriticalSection GSamplerStateCacheLock;
-#define LOCK_GSamplerStateCache (1)
-#else
-#define LOCK_GSamplerStateCache (0)
-#endif
 static TMap<ID3D11SamplerState*, FD3D11SamplerState*> GSamplerStateCache;
 
 
 void EmptyD3DSamplerStateCache()
 {
-#if LOCK_GSamplerStateCache
-	FScopeLock Lock(&GSamplerStateCacheLock);
-#endif
 	for (auto Iter = GSamplerStateCache.CreateIterator(); Iter; ++Iter )
 	{
 		auto* State = Iter.Value();
@@ -182,11 +173,6 @@ FSamplerStateRHIRef FD3D11DynamicRHI::RHICreateSamplerState(const FSamplerStateI
 	SamplerDesc.BorderColor[3] = LinearBorderColor.A;
 	SamplerDesc.ComparisonFunc = TranslateSamplerCompareFunction(Initializer.SamplerComparisonFunction);
 
-#if LOCK_GSamplerStateCache
-	QUICK_SCOPE_CYCLE_COUNTER(FD3D11DynamicRHI_RHICreateSamplerState_LockAndCreate);
-	FScopeLock Lock(&GSamplerStateCacheLock);
-#endif
-
 	// D3D11 will return the same pointer if the particular state description was already created
 	TRefCountPtr<ID3D11SamplerState> SamplerStateHandle;
 	VERIFYD3D11RESULT(Direct3DDevice->CreateSamplerState(&SamplerDesc, SamplerStateHandle.GetInitReference()));
diff --git a/Engine/Source/Runtime/Windows/D3D11RHI/Private/D3D11UniformBuffer.cpp b/Engine/Source/Runtime/Windows/D3D11RHI/Private/D3D11UniformBuffer.cpp
--- a/Engine/Source/Runtime/Windows/D3D11RHI/Private/D3D11UniformBuffer.cpp
+++ b/Engine/Source/Runtime/Windows/D3D11RHI/Private/D3D11UniformBuffer.cpp
@@ -98,7 +98,7 @@ static bool IsPoolingEnabled()
 
 FUniformBufferRHIRef FD3D11DynamicRHI::RHICreateUniformBuffer(const void* Contents,const FRHIUniformBufferLayout& Layout,EUniformBufferUsage Usage)
 {
-	check(IsInRenderingThread() || IsInRHIThread());
+	check(IsInRenderingThread());
 
 	FD3D11UniformBuffer* NewUniformBuffer = nullptr;
 	const uint32 NumBytes = Layout.ConstantBufferSize;
diff --git a/Engine/Source/Runtime/Windows/D3D11RHI/Private/D3D11Viewport.cpp b/Engine/Source/Runtime/Windows/D3D11RHI/Private/D3D11Viewport.cpp
--- a/Engine/Source/Runtime/Windows/D3D11RHI/Private/D3D11Viewport.cpp
+++ b/Engine/Source/Runtime/Windows/D3D11RHI/Private/D3D11Viewport.cpp
@@ -603,11 +603,9 @@ bool FD3D11DynamicRHI::RHIIsDrawingViewport()
 	return DrawingViewport != NULL;
 }
 
-#if PLATFORM_SUPPORTS_RHI_THREAD
 void FD3D11DynamicRHI::RHIAdvanceFrameForGetViewportBackBuffer()
 {
 }
-#endif
 
 FTexture2DRHIRef FD3D11DynamicRHI::RHIGetViewportBackBuffer(FViewportRHIParamRef ViewportRHI)
 {
