now SetRenderTarget requires to specify what DepthAndStencil access you want to do e.g. FExclusiveDepthStencil::DepthRead_StencilWrite
options:   Nop / Read / Write

added ensure() and check() to find bad usage patterns

this fixes DBuffer decal rendering, allows optimizations in the RHI (near hardware APIs like DX12 or consoles)

diff --git a/Engine/Source/Runtime/RHI/Public/RHIDefinitions.h b/Engine/Source/Runtime/RHI/Public/RHIDefinitions.h
--- a/Engine/Source/Runtime/RHI/Public/RHIDefinitions.h
+++ b/Engine/Source/Runtime/RHI/Public/RHIDefinitions.h
@@ -564,19 +564,14 @@ enum class ERenderTargetStoreAction
 enum class ESimpleRenderTargetMode
 {
 	// These will all store out color and depth
-	EExistingColorAndDepth,						// Color = Existing, Depth = Existing
-	EExistingColorAndReadOnlyDepth,					// Color = Existing, Depth = Existing. depth cannot be written.
+	EExistingColorAndDepth,							// Color = Existing, Depth = Existing
 	EUninitializedColorAndDepth,					// Color = ????, Depth = ????
 	EUninitializedColorExistingDepth,				// Color = ????, Depth = Existing
-	EUninitializedColorExistingReadOnlyDepth,		// Color = ????, Depth = Existing. depth cannot be written.
 	EUninitializedColorClearDepth,					// Color = ????, Depth = Default
 	EClearToDefault,								// Default Color = (0,0,0,0), Default Depth = 0.0f
 	EClearColorToBlack,								// Color = (0,0,0,0), Depth = Existing
-	EClearColorToBlackReadOnlyDepth,				// Color = (0,0,0,0), Depth = Existing. depth cannot be written
 	EClearColorToBlackWithFullAlpha,				// Color = (0,0,0,1), Depth = Existing
-	EClearColorToBlackWithFullAlphaReadOnlyDepth,	// Color = (0,0,0,1), Depth = Existing. depth cannot be written
 	EClearColorToWhite,								// Color = (1,1,1,1), Depth = Existing
-	EClearColorToWhiteReadOnlyDepth,				// Color = (1,1,1,1), Depth = Existing. depth cannot be written
 	EClearDepthToOne,								// Color = Existing, Depth = 1.0
 	EExistingContents_NoDepthStore,					// Load existing contents, but don't store depth out.  depth can be written.
 
diff --git a/Engine/Source/Runtime/RHI/Public/RHIResources.h b/Engine/Source/Runtime/RHI/Public/RHIResources.h
--- a/Engine/Source/Runtime/RHI/Public/RHIResources.h
+++ b/Engine/Source/Runtime/RHI/Public/RHIResources.h
@@ -655,6 +655,142 @@ public:
 	}
 };
 
+class FExclusiveDepthStencil
+{
+public:
+	enum Type
+	{
+		// don't use those directly, use the combined versions below
+		// 4 bits are used for depth and 4 for stencil to make the hex value readable and non overlapping
+		DepthNop =		0x00,
+		DepthRead =		0x01,
+		DepthWrite =	0x02,
+		DepthMask =		0x0f,
+		StencilNop =	0x00,
+		StencilRead =	0x10,
+		StencilWrite =	0x20,
+		StencilMask =	0xf0,
+
+		// use those:
+		DepthNop_StencilNop = DepthNop + StencilNop,
+		DepthRead_StencilNop = DepthRead + StencilNop,
+		DepthWrite_StencilNop = DepthWrite + StencilNop,
+		DepthNop_StencilRead = DepthNop + StencilRead,
+		DepthRead_StencilRead = DepthRead + StencilRead,
+		DepthWrite_StencilRead = DepthWrite + StencilRead,
+		DepthNop_StencilWrite = DepthNop + StencilWrite,
+		DepthRead_StencilWrite = DepthRead + StencilWrite,
+		DepthWrite_StencilWrite = DepthWrite + StencilWrite,
+	};
+
+private:
+	Type Value;
+
+public:
+	// constructor
+	FExclusiveDepthStencil(Type InValue = DepthNop_StencilNop)
+		: Value(InValue)
+	{
+	}
+
+	inline bool IsUsingDepthStencil() const
+	{
+		return Value != DepthNop_StencilNop;
+	}
+	inline bool IsDepthWrite() const
+	{
+		return ExtractDepth() == DepthWrite;
+	}
+	inline bool IsStencilWrite() const
+	{
+		return ExtractStencil() == StencilWrite;
+	}
+	inline void SetDepthWrite()
+	{
+		Value = (Type)(ExtractStencil() | DepthWrite);
+	}
+	inline void SetStencilWrite()
+	{
+		Value = (Type)(ExtractDepth() | StencilWrite);
+	}
+	inline void SetDepthStencilWrite(bool bDepth, bool bStencil)
+	{
+		Value = DepthNop_StencilNop;
+
+		if (bDepth)
+		{
+			SetDepthWrite();
+		}
+		if (bStencil)
+		{
+			SetStencilWrite();
+		}
+	}
+	bool operator==(const FExclusiveDepthStencil& rhs) const
+	{
+		return Value == rhs.Value;
+	}
+	inline bool IsValid(FExclusiveDepthStencil& Current) const
+	{
+		Type Depth = ExtractDepth();
+
+		if (Depth != DepthNop && Depth != Current.ExtractDepth())
+		{
+			return false;
+		}
+
+		Type Stencil = ExtractStencil();
+
+		if (Stencil != StencilNop && Stencil != Current.ExtractStencil())
+		{
+			return false;
+		}
+
+		return true;
+	}
+
+	uint32 GetIndex() const
+	{
+		// Note: The array to index has views created in that specific order.
+
+		// we don't care about the Nop versions so less views are needed
+		// we combine Nop and Write
+		switch (Value)
+		{
+			case DepthWrite_StencilNop:
+			case DepthNop_StencilWrite:
+			case DepthWrite_StencilWrite:
+			case DepthNop_StencilNop:
+				return 0; // old DSAT_Writable
+		
+			case DepthRead_StencilNop:
+			case DepthRead_StencilWrite:
+				return 1; // old DSAT_ReadOnlyDepth
+
+			case DepthNop_StencilRead:
+			case DepthWrite_StencilRead:
+				return 2; // old DSAT_ReadOnlyStencil
+
+			case DepthRead_StencilRead:
+				return 3; // old DSAT_ReadOnlyDepthAndStencil
+		}
+		// should never happen
+		check(0);
+		return -1;
+	}
+	static const uint32 MaxIndex = 4;
+
+private:
+	inline Type ExtractDepth() const
+	{
+		return (Type)(Value & DepthMask);
+	}
+	inline Type ExtractStencil() const
+	{
+		return (Type)(Value & StencilMask);
+	}
+};
+
 class FRHIDepthRenderTargetView
 {
 public:
@@ -663,8 +799,16 @@ public:
 	ERenderTargetLoadAction		DepthLoadAction;
 	ERenderTargetStoreAction	DepthStoreAction;
 	ERenderTargetLoadAction		StencilLoadAction;
+
+private:
 	ERenderTargetStoreAction	StencilStoreAction;
-	bool						bReadOnly; //target will not be written to.
+	FExclusiveDepthStencil		DepthStencilAccess;
+public:
+
+	// accessor to prevent write access to StencilStoreAction
+	ERenderTargetStoreAction GetStencilStoreAction() const { return StencilStoreAction; }
+	// accessor to prevent write access to DepthStencilAccess
+	FExclusiveDepthStencil GetDepthStencilAccess() const { return DepthStencilAccess; }
 
 	FRHIDepthRenderTargetView() :
 		Texture(nullptr),
@@ -672,8 +816,10 @@ public:
 		DepthStoreAction(ERenderTargetStoreAction::EStore),
 		StencilLoadAction(ERenderTargetLoadAction::EClear),
 		StencilStoreAction(ERenderTargetStoreAction::EStore),
-		bReadOnly(false)
-	{}
+		DepthStencilAccess(FExclusiveDepthStencil::DepthWrite_StencilWrite)
+	{
+		Validate();
+	}
 
 	FRHIDepthRenderTargetView(FTextureRHIParamRef InTexture) :
 		Texture(InTexture),
@@ -681,8 +827,10 @@ public:
 		DepthStoreAction(ERenderTargetStoreAction::EStore),
 		StencilLoadAction(ERenderTargetLoadAction::EClear),
 		StencilStoreAction(ERenderTargetStoreAction::EStore),
-		bReadOnly(false)
-	{}
+		DepthStencilAccess(FExclusiveDepthStencil::DepthWrite_StencilWrite)
+	{
+		Validate();
+	}
 
 	FRHIDepthRenderTargetView(FTextureRHIParamRef InTexture, ERenderTargetLoadAction InLoadAction, ERenderTargetStoreAction InStoreAction) :
 		Texture(InTexture),
@@ -690,18 +838,20 @@ public:
 		DepthStoreAction(InStoreAction),
 		StencilLoadAction(InLoadAction),
 		StencilStoreAction(InStoreAction),
-		bReadOnly(false)
-	{}
+		DepthStencilAccess(FExclusiveDepthStencil::DepthWrite_StencilWrite)
+	{
+		Validate();
+	}
 
-	FRHIDepthRenderTargetView(FTextureRHIParamRef InTexture, ERenderTargetLoadAction InLoadAction, ERenderTargetStoreAction InStoreAction, bool bInReadOnly) :
+	FRHIDepthRenderTargetView(FTextureRHIParamRef InTexture, ERenderTargetLoadAction InLoadAction, ERenderTargetStoreAction InStoreAction, FExclusiveDepthStencil InDepthStencilAccess) :
 		Texture(InTexture),
 		DepthLoadAction(InLoadAction),
 		DepthStoreAction(InStoreAction),
 		StencilLoadAction(InLoadAction),
 		StencilStoreAction(InStoreAction),
-		bReadOnly(bInReadOnly)
+		DepthStencilAccess(InDepthStencilAccess)
 	{
-		ensure(!bReadOnly || InStoreAction == ERenderTargetStoreAction::ENoAction);
+		Validate();
 	}
 
 	FRHIDepthRenderTargetView(FTextureRHIParamRef InTexture, ERenderTargetLoadAction InDepthLoadAction, ERenderTargetStoreAction InDepthStoreAction, ERenderTargetLoadAction InStencilLoadAction, ERenderTargetStoreAction InStencilStoreAction) :
@@ -710,18 +860,26 @@ public:
 		DepthStoreAction(InDepthStoreAction),
 		StencilLoadAction(InStencilLoadAction),
 		StencilStoreAction(InStencilStoreAction),
-		bReadOnly(false)
-	{}
+		DepthStencilAccess(FExclusiveDepthStencil::DepthWrite_StencilWrite)
+	{
+		Validate();
+	}
 
-	FRHIDepthRenderTargetView(FTextureRHIParamRef InTexture, ERenderTargetLoadAction InDepthLoadAction, ERenderTargetStoreAction InDepthStoreAction, ERenderTargetLoadAction InStencilLoadAction, ERenderTargetStoreAction InStencilStoreAction, bool bInReadOnly) :
+	FRHIDepthRenderTargetView(FTextureRHIParamRef InTexture, ERenderTargetLoadAction InDepthLoadAction, ERenderTargetStoreAction InDepthStoreAction, ERenderTargetLoadAction InStencilLoadAction, ERenderTargetStoreAction InStencilStoreAction, FExclusiveDepthStencil InDepthStencilAccess) :
 		Texture(InTexture),
 		DepthLoadAction(InDepthLoadAction),
 		DepthStoreAction(InDepthStoreAction),
 		StencilLoadAction(InStencilLoadAction),
 		StencilStoreAction(InStencilStoreAction),
-		bReadOnly(bInReadOnly)
+		DepthStencilAccess(InDepthStencilAccess)
+	{
+		Validate();
+	}
+
+	void Validate() const
 	{
-		ensure(!bReadOnly || InDepthStoreAction == ERenderTargetStoreAction::ENoAction);
+		// Missed optimization, resolve is not needed for read only ops
+		ensure(!(DepthStencilAccess == FExclusiveDepthStencil::DepthRead_StencilRead && StencilStoreAction != ERenderTargetStoreAction::ENoAction));
 	}
 
 	bool operator==(const FRHIDepthRenderTargetView& Other)
@@ -732,7 +890,7 @@ public:
 			DepthStoreAction == Other.DepthStoreAction &&
 			StencilLoadAction == Other.StencilLoadAction &&
 			StencilStoreAction == Other.StencilStoreAction &&
-			bReadOnly == Other.bReadOnly;
+			DepthStencilAccess == Other.DepthStencilAccess;
 	}
 };
 
diff --git a/Engine/Source/Runtime/RHI/Public/RHIUtilities.h b/Engine/Source/Runtime/RHI/Public/RHIUtilities.h
--- a/Engine/Source/Runtime/RHI/Public/RHIUtilities.h
+++ b/Engine/Source/Runtime/RHI/Public/RHIUtilities.h
@@ -123,14 +123,13 @@ struct FRWBufferByteAddress
  * Convert the ESimpleRenderTargetMode into usable values 
  * @todo: Can we easily put this into a .cpp somewhere?
  */
-inline void DecodeRenderTargetMode(ESimpleRenderTargetMode Mode, ERenderTargetLoadAction& ColorLoadAction, ERenderTargetStoreAction& ColorStoreAction, ERenderTargetLoadAction& DepthLoadAction, ERenderTargetStoreAction& DepthStoreAction, FLinearColor& ClearColor, float& ClearDepth, bool& bReadOnlyDepth)
+inline void DecodeRenderTargetMode(ESimpleRenderTargetMode Mode, ERenderTargetLoadAction& ColorLoadAction, ERenderTargetStoreAction& ColorStoreAction, ERenderTargetLoadAction& DepthLoadAction, ERenderTargetStoreAction& DepthStoreAction, FLinearColor& ClearColor, float& ClearDepth)
 {
 	// set defaults
 	ColorStoreAction = ERenderTargetStoreAction::EStore;
 	DepthStoreAction = ERenderTargetStoreAction::EStore;
 	ClearColor = FLinearColor(0, 0, 0, 0);
 	ClearDepth = 0.0f;
-	bReadOnlyDepth = false;
 
 	switch (Mode)
 	{
@@ -138,12 +137,6 @@ inline void DecodeRenderTargetMode(ESimpleRenderTargetMode Mode, ERenderTargetLo
 		ColorLoadAction = ERenderTargetLoadAction::ELoad;
 		DepthLoadAction = ERenderTargetLoadAction::ELoad;
 		break;
-	case ESimpleRenderTargetMode::EExistingColorAndReadOnlyDepth:
-		ColorLoadAction = ERenderTargetLoadAction::ELoad;
-		DepthLoadAction = ERenderTargetLoadAction::ELoad;
-		DepthStoreAction = ERenderTargetStoreAction::ENoAction;
-		bReadOnlyDepth = true;
-		break;
 	case ESimpleRenderTargetMode::EUninitializedColorAndDepth:
 		ColorLoadAction = ERenderTargetLoadAction::ENoAction;
 		DepthLoadAction = ERenderTargetLoadAction::ENoAction;
@@ -152,12 +145,6 @@ inline void DecodeRenderTargetMode(ESimpleRenderTargetMode Mode, ERenderTargetLo
 		ColorLoadAction = ERenderTargetLoadAction::ENoAction;
 		DepthLoadAction = ERenderTargetLoadAction::ELoad;
 		break;
-	case ESimpleRenderTargetMode::EUninitializedColorExistingReadOnlyDepth:
-		ColorLoadAction = ERenderTargetLoadAction::ENoAction;
-		DepthLoadAction = ERenderTargetLoadAction::ELoad;
-		DepthStoreAction = ERenderTargetStoreAction::ENoAction;
-		bReadOnlyDepth = true;
-		break;
 	case ESimpleRenderTargetMode::EUninitializedColorClearDepth:
 		ColorLoadAction = ERenderTargetLoadAction::ENoAction;
 		DepthLoadAction = ERenderTargetLoadAction::EClear;
@@ -171,37 +158,16 @@ inline void DecodeRenderTargetMode(ESimpleRenderTargetMode Mode, ERenderTargetLo
 		DepthLoadAction = ERenderTargetLoadAction::ELoad;
 		ClearColor = FLinearColor::Black;
 		break;
-	case ESimpleRenderTargetMode::EClearColorToBlackReadOnlyDepth:
-		ColorLoadAction = ERenderTargetLoadAction::EClear;
-		DepthLoadAction = ERenderTargetLoadAction::ELoad;
-		ClearColor = FLinearColor::Black;
-		DepthStoreAction = ERenderTargetStoreAction::ENoAction;
-		bReadOnlyDepth = true;
-		break;
 	case ESimpleRenderTargetMode::EClearColorToBlackWithFullAlpha:
 		ColorLoadAction = ERenderTargetLoadAction::EClear;
 		DepthLoadAction = ERenderTargetLoadAction::ELoad;
 		ClearColor = FLinearColor(0, 0, 0, 1);
 		break;
-	case ESimpleRenderTargetMode::EClearColorToBlackWithFullAlphaReadOnlyDepth:
-		ColorLoadAction = ERenderTargetLoadAction::EClear;
-		DepthLoadAction = ERenderTargetLoadAction::ELoad;
-		ClearColor = FLinearColor(0, 0, 0, 1);
-		DepthStoreAction = ERenderTargetStoreAction::ENoAction;
-		bReadOnlyDepth = true;
-		break;
 	case ESimpleRenderTargetMode::EClearColorToWhite:
 		ColorLoadAction = ERenderTargetLoadAction::EClear;
 		DepthLoadAction = ERenderTargetLoadAction::ELoad;
 		ClearColor = FLinearColor::White;
 		break;
-	case ESimpleRenderTargetMode::EClearColorToWhiteReadOnlyDepth:
-		ColorLoadAction = ERenderTargetLoadAction::EClear;
-		DepthLoadAction = ERenderTargetLoadAction::ELoad;
-		ClearColor = FLinearColor::White;
-		DepthStoreAction = ERenderTargetStoreAction::ENoAction;
-		bReadOnlyDepth = true;
-		break;
 	case ESimpleRenderTargetMode::EClearDepthToOne:
 		ColorLoadAction = ERenderTargetLoadAction::ELoad;
 		DepthLoadAction = ERenderTargetLoadAction::EClear;
@@ -226,18 +192,17 @@ inline void SetRenderTarget(FRHICommandList& RHICmdList, FTextureRHIParamRef New
 }
 
 /** Helper for the common case of using a single color and depth render target. */
-inline void SetRenderTarget(FRHICommandList& RHICmdList, FTextureRHIParamRef NewRenderTarget, FTextureRHIParamRef NewDepthStencilTarget, ESimpleRenderTargetMode Mode)
+inline void SetRenderTarget(FRHICommandList& RHICmdList, FTextureRHIParamRef NewRenderTarget, FTextureRHIParamRef NewDepthStencilTarget, ESimpleRenderTargetMode Mode, FExclusiveDepthStencil DepthStencilAccess = FExclusiveDepthStencil::DepthWrite_StencilWrite)
 {
 	ERenderTargetLoadAction ColorLoadAction, DepthLoadAction;
 	ERenderTargetStoreAction ColorStoreAction, DepthStoreAction;
 	FLinearColor ClearColor;
 	float ClearDepth;
-	bool bReadOnlyDepth;
-	DecodeRenderTargetMode(Mode, ColorLoadAction, ColorStoreAction, DepthLoadAction, DepthStoreAction, ClearColor, ClearDepth, bReadOnlyDepth);
+	DecodeRenderTargetMode(Mode, ColorLoadAction, ColorStoreAction, DepthLoadAction, DepthStoreAction, ClearColor, ClearDepth);
 
 	// now make the FRHISetRenderTargetsInfo that encapsulates all of the info
 	FRHIRenderTargetView ColorView(NewRenderTarget, 0, -1, ColorLoadAction, ColorStoreAction);
-	FRHISetRenderTargetsInfo Info(1, &ColorView, FRHIDepthRenderTargetView(NewDepthStencilTarget, DepthLoadAction, DepthStoreAction, bReadOnlyDepth));
+	FRHISetRenderTargetsInfo Info(1, &ColorView, FRHIDepthRenderTargetView(NewDepthStencilTarget, DepthLoadAction, DepthStoreAction, DepthStencilAccess));
 	Info.ClearColors[0] = ClearColor;
 	Info.DepthClearValue = ClearDepth;
 	RHICmdList.SetRenderTargetsAndClear(Info);
@@ -286,15 +251,15 @@ inline void SetRenderTargets(
 	uint32 NewNumSimultaneousRenderTargets,
 	const FTextureRHIParamRef* NewRenderTargetsRHI,
 	FTextureRHIParamRef NewDepthStencilTargetRHI,
-	ESimpleRenderTargetMode Mode
+	ESimpleRenderTargetMode Mode,
+	FExclusiveDepthStencil DepthStencilAccess
 	)
 {
 	ERenderTargetLoadAction ColorLoadAction, DepthLoadAction;
 	ERenderTargetStoreAction ColorStoreAction, DepthStoreAction;
 	FLinearColor ClearColor;
 	float ClearDepth;
-	bool bReadOnlyDepth;
-	DecodeRenderTargetMode(Mode, ColorLoadAction, ColorStoreAction, DepthLoadAction, DepthStoreAction, ClearColor, ClearDepth, bReadOnlyDepth);
+	DecodeRenderTargetMode(Mode, ColorLoadAction, ColorStoreAction, DepthLoadAction, DepthStoreAction, ClearColor, ClearDepth);
 
 	FRHIRenderTargetView RTVs[MaxSimultaneousRenderTargets];
 
@@ -303,7 +268,7 @@ inline void SetRenderTargets(
 		RTVs[Index] = FRHIRenderTargetView(NewRenderTargetsRHI[Index], 0, -1, ColorLoadAction, ColorStoreAction);
 	}
 
-	FRHIDepthRenderTargetView DepthRTV(NewDepthStencilTargetRHI, DepthLoadAction, DepthStoreAction, bReadOnlyDepth);
+	FRHIDepthRenderTargetView DepthRTV(NewDepthStencilTargetRHI, DepthLoadAction, DepthStoreAction, DepthStencilAccess);
 	RHICmdList.SetRenderTargets(NewNumSimultaneousRenderTargets, RTVs, &DepthRTV, 0, nullptr);
 }
 
diff --git a/Engine/Source/Runtime/Renderer/Private/AtmosphereRendering.cpp b/Engine/Source/Runtime/Renderer/Private/AtmosphereRendering.cpp
--- a/Engine/Source/Runtime/Renderer/Private/AtmosphereRendering.cpp
+++ b/Engine/Source/Runtime/Renderer/Private/AtmosphereRendering.cpp
@@ -414,7 +414,7 @@ void FDeferredShadingSceneRenderer::RenderAtmosphere(FRHICommandListImmediate& R
 		};
 
 		
-		GSceneRenderTargets.BeginRenderingSceneColor(RHICmdList, ESimpleRenderTargetMode::EUninitializedColorExistingReadOnlyDepth);
+		GSceneRenderTargets.BeginRenderingSceneColor(RHICmdList, ESimpleRenderTargetMode::EUninitializedColorExistingDepth, FExclusiveDepthStencil::DepthRead_StencilWrite);
 		for(int32 ViewIndex = 0;ViewIndex < Views.Num();ViewIndex++)
 		{
 			const FViewInfo& View = Views[ViewIndex];
@@ -1614,8 +1614,6 @@ bool ShouldRenderAtmosphere(const FSceneViewFamily& Family)
 	const FEngineShowFlags EngineShowFlags = Family.EngineShowFlags;
 
 	return GSupportsVolumeTextureRendering
-		&& EngineShowFlags.Atmosphere
-		&& !EngineShowFlags.ShaderComplexity
-		&& !EngineShowFlags.StationaryLightOverlap 
-		&& !EngineShowFlags.LightMapDensity;
+		&& EngineShowFlags.AtmosphericFog
+		&& EngineShowFlags.Fog;
 }
diff --git a/Engine/Source/Runtime/Renderer/Private/CompositionLighting/PostProcessDeferredDecals.cpp b/Engine/Source/Runtime/Renderer/Private/CompositionLighting/PostProcessDeferredDecals.cpp
--- a/Engine/Source/Runtime/Renderer/Private/CompositionLighting/PostProcessDeferredDecals.cpp
+++ b/Engine/Source/Runtime/Renderer/Private/CompositionLighting/PostProcessDeferredDecals.cpp
@@ -374,7 +374,7 @@ void StencilDecalMask(FRHICommandList& RHICmdList, const FViewInfo& View)
 	SCOPED_DRAW_EVENT(RHICmdList, StencilDecalMask);
 	RHICmdList.SetRasterizerState(TStaticRasterizerState<FM_Solid, CM_None>::GetRHI());
 	RHICmdList.SetBlendState(TStaticBlendState<CW_NONE>::GetRHI());
-	SetRenderTarget(RHICmdList, NULL, GSceneRenderTargets.GetSceneDepthSurface(), ESimpleRenderTargetMode::EUninitializedColorExistingReadOnlyDepth);
+	SetRenderTarget(RHICmdList, NULL, GSceneRenderTargets.GetSceneDepthSurface(), ESimpleRenderTargetMode::EUninitializedColorExistingDepth, FExclusiveDepthStencil::DepthRead_StencilWrite);
 	RHICmdList.SetViewport(View.ViewRect.Min.X, View.ViewRect.Min.Y, 0.0f, View.ViewRect.Max.X, View.ViewRect.Max.Y, 1.0f);
 
 	// Write 1 to highest bit of stencil to areas that should not receive decals
@@ -948,18 +948,18 @@ void FRCPassPostProcessDeferredDecals::Process(FRenderingCompositePassContext& C
 							TargetsToResolve[GBufferBIndex] = GSceneRenderTargets.GBufferB->GetRenderTargetItem().TargetableTexture;
 							TargetsToResolve[GBufferCIndex] = GSceneRenderTargets.GBufferC->GetRenderTargetItem().TargetableTexture;
 							
-							SetRenderTargets(RHICmdList, 4, TargetsToResolve, GSceneRenderTargets.GetSceneDepthSurface(), ESimpleRenderTargetMode::EExistingColorAndReadOnlyDepth);
+							SetRenderTargets(RHICmdList, 4, TargetsToResolve, GSceneRenderTargets.GetSceneDepthSurface(), ESimpleRenderTargetMode::EExistingColorAndDepth, FExclusiveDepthStencil::DepthRead_StencilWrite);
 						}
 						break;
 
 					case RTM_GBufferNormal:
 						TargetsToResolve[GBufferAIndex] = GSceneRenderTargets.GBufferA->GetRenderTargetItem().TargetableTexture;
-						SetRenderTarget(RHICmdList, TargetsToResolve[GBufferAIndex], GSceneRenderTargets.GetSceneDepthSurface(), ESimpleRenderTargetMode::EExistingColorAndReadOnlyDepth);
+						SetRenderTarget(RHICmdList, TargetsToResolve[GBufferAIndex], GSceneRenderTargets.GetSceneDepthSurface(), ESimpleRenderTargetMode::EExistingColorAndDepth, FExclusiveDepthStencil::DepthRead_StencilWrite);
 						break;
 					
 					case RTM_SceneColor:
 						TargetsToResolve[SceneColorIndex] = GSceneRenderTargets.GetSceneColor()->GetRenderTargetItem().TargetableTexture;
-						SetRenderTarget(RHICmdList, TargetsToResolve[SceneColorIndex], GSceneRenderTargets.GetSceneDepthSurface(), ESimpleRenderTargetMode::EExistingColorAndReadOnlyDepth);
+						SetRenderTarget(RHICmdList, TargetsToResolve[SceneColorIndex], GSceneRenderTargets.GetSceneDepthSurface(), ESimpleRenderTargetMode::EExistingColorAndDepth, FExclusiveDepthStencil::DepthRead_StencilWrite);
 						break;
 
 					case RTM_DBuffer:
@@ -967,7 +967,7 @@ void FRCPassPostProcessDeferredDecals::Process(FRenderingCompositePassContext& C
 							TargetsToResolve[DBufferAIndex] = GSceneRenderTargets.DBufferA->GetRenderTargetItem().TargetableTexture;
 							TargetsToResolve[DBufferBIndex] = GSceneRenderTargets.DBufferB->GetRenderTargetItem().TargetableTexture;
 							TargetsToResolve[DBufferCIndex] = GSceneRenderTargets.DBufferC->GetRenderTargetItem().TargetableTexture;
-							SetRenderTargets(RHICmdList, 3, &TargetsToResolve[DBufferAIndex], GSceneRenderTargets.GetSceneDepthSurface(), ESimpleRenderTargetMode::EExistingColorAndReadOnlyDepth);
+							SetRenderTargets(RHICmdList, 3, &TargetsToResolve[DBufferAIndex], GSceneRenderTargets.GetSceneDepthSurface(), ESimpleRenderTargetMode::EExistingColorAndDepth, FExclusiveDepthStencil::DepthRead_StencilWrite);
 						}
 						break;
 
diff --git a/Engine/Source/Runtime/Renderer/Private/CompositionLighting/PostProcessLpvIndirect.cpp b/Engine/Source/Runtime/Renderer/Private/CompositionLighting/PostProcessLpvIndirect.cpp
--- a/Engine/Source/Runtime/Renderer/Private/CompositionLighting/PostProcessLpvIndirect.cpp
+++ b/Engine/Source/Runtime/Renderer/Private/CompositionLighting/PostProcessLpvIndirect.cpp
@@ -162,7 +162,7 @@ void FRCPassPostProcessLpvIndirect::Process(FRenderingCompositePassContext& Cont
 	};
 
 	// Set the view family's render target/viewport.
-	SetRenderTargets(Context.RHICmdList, 1, RenderTargets, GSceneRenderTargets.GetSceneDepthSurface(), ESimpleRenderTargetMode::EExistingColorAndReadOnlyDepth);
+	SetRenderTargets(Context.RHICmdList, 1, RenderTargets, GSceneRenderTargets.GetSceneDepthSurface(), ESimpleRenderTargetMode::EExistingColorAndDepth, FExclusiveDepthStencil::DepthRead_StencilWrite);
 
 	Context.SetViewportAndCallRHI(View.ViewRect);
 
diff --git a/Engine/Source/Runtime/Renderer/Private/DistortionRendering.cpp b/Engine/Source/Runtime/Renderer/Private/DistortionRendering.cpp
--- a/Engine/Source/Runtime/Renderer/Private/DistortionRendering.cpp
+++ b/Engine/Source/Runtime/Renderer/Private/DistortionRendering.cpp
@@ -786,7 +786,7 @@ void FSceneRenderer::RenderDistortion(FRHICommandListImmediate& RHICmdList)
 		// DistortionRT==0 should never happen but better we don't crash
 		if(DistortionRT)
 		{
-			SetRenderTarget(RHICmdList, DistortionRT->GetRenderTargetItem().TargetableTexture, GSceneRenderTargets.GetSceneDepthSurface(), ESimpleRenderTargetMode::EUninitializedColorExistingReadOnlyDepth);
+			SetRenderTarget(RHICmdList, DistortionRT->GetRenderTargetItem().TargetableTexture, GSceneRenderTargets.GetSceneDepthSurface(), ESimpleRenderTargetMode::EUninitializedColorExistingDepth, FExclusiveDepthStencil::DepthRead_StencilWrite);
 
 			for(int32 ViewIndex = 0;ViewIndex < Views.Num();ViewIndex++)
 			{
diff --git a/Engine/Source/Runtime/Renderer/Private/FogRendering.cpp b/Engine/Source/Runtime/Renderer/Private/FogRendering.cpp
--- a/Engine/Source/Runtime/Renderer/Private/FogRendering.cpp
+++ b/Engine/Source/Runtime/Renderer/Private/FogRendering.cpp
@@ -308,7 +308,7 @@ bool FDeferredShadingSceneRenderer::RenderFog(FRHICommandListImmediate& RHICmdLi
 			0, 2, 3
 		};
 
-		GSceneRenderTargets.BeginRenderingSceneColor(RHICmdList, ESimpleRenderTargetMode::EExistingColorAndReadOnlyDepth);
+		GSceneRenderTargets.BeginRenderingSceneColor(RHICmdList, ESimpleRenderTargetMode::EExistingColorAndDepth, FExclusiveDepthStencil::DepthRead_StencilWrite);
 		for(int32 ViewIndex = 0;ViewIndex < Views.Num();ViewIndex++)
 		{
 			const FViewInfo& View = Views[ViewIndex];
diff --git a/Engine/Source/Runtime/Renderer/Private/LightRendering.cpp b/Engine/Source/Runtime/Renderer/Private/LightRendering.cpp
--- a/Engine/Source/Runtime/Renderer/Private/LightRendering.cpp
+++ b/Engine/Source/Runtime/Renderer/Private/LightRendering.cpp
@@ -425,7 +425,7 @@ void FDeferredShadingSceneRenderer::RenderLights(FRHICommandListImmediate& RHICm
 			}
 			else if (SimpleLights.InstanceData.Num() > 0)
 			{
-				GSceneRenderTargets.BeginRenderingSceneColor(RHICmdList, ESimpleRenderTargetMode::EExistingColorAndReadOnlyDepth);
+				GSceneRenderTargets.BeginRenderingSceneColor(RHICmdList, ESimpleRenderTargetMode::EExistingColorAndDepth, FExclusiveDepthStencil::DepthRead_StencilWrite);
 				RenderSimpleLightsStandardDeferred(RHICmdList, SimpleLights);
 			}
 
@@ -433,7 +433,7 @@ void FDeferredShadingSceneRenderer::RenderLights(FRHICommandListImmediate& RHICm
 				SCOPED_DRAW_EVENT(RHICmdList, StandardDeferredLighting);
 
 				// make sure we don't clear the depth
-				GSceneRenderTargets.BeginRenderingSceneColor(RHICmdList, ESimpleRenderTargetMode::EExistingColorAndReadOnlyDepth);
+				GSceneRenderTargets.BeginRenderingSceneColor(RHICmdList, ESimpleRenderTargetMode::EExistingColorAndDepth, FExclusiveDepthStencil::DepthRead_StencilWrite);
 
 				// Draw non-shadowed non-light function lights without changing render targets between them
 				for (int32 LightIndex = StandardDeferredStart; LightIndex < AttenuationLightStart; LightIndex++)
@@ -555,7 +555,7 @@ void FDeferredShadingSceneRenderer::RenderLights(FRHICommandListImmediate& RHICm
 			}
 
 			GSceneRenderTargets.SetLightAttenuationMode(bUsedLightAttenuation);
-			GSceneRenderTargets.BeginRenderingSceneColor(RHICmdList, ESimpleRenderTargetMode::EExistingColorAndReadOnlyDepth);
+			GSceneRenderTargets.BeginRenderingSceneColor(RHICmdList, ESimpleRenderTargetMode::EExistingColorAndDepth, FExclusiveDepthStencil::DepthRead_StencilWrite);
 
 			// Render the light to the scene color buffer, conditionally using the attenuation buffer or a 1x1 white texture as input 
 			if(bDirectLighting)
@@ -642,7 +642,7 @@ void FDeferredShadingSceneRenderer::RenderStationaryLightOverlap(FRHICommandList
 {
 	if (Scene->bIsEditorScene)
 	{
-		GSceneRenderTargets.BeginRenderingSceneColor(RHICmdList, ESimpleRenderTargetMode::EUninitializedColorExistingReadOnlyDepth);
+		GSceneRenderTargets.BeginRenderingSceneColor(RHICmdList, ESimpleRenderTargetMode::EUninitializedColorExistingDepth, FExclusiveDepthStencil::DepthRead_StencilWrite);
 
 		// Clear to discard base pass values in scene color since we didn't skip that, to have valid scene depths
 		RHICmdList.Clear(true, FLinearColor::Black, false, (float)ERHIZBuffer::FarPlane, false, 0, FIntRect());
diff --git a/Engine/Source/Runtime/Renderer/Private/LightShaftRendering.cpp b/Engine/Source/Runtime/Renderer/Private/LightShaftRendering.cpp
--- a/Engine/Source/Runtime/Renderer/Private/LightShaftRendering.cpp
+++ b/Engine/Source/Runtime/Renderer/Private/LightShaftRendering.cpp
@@ -832,7 +832,7 @@ IMPLEMENT_SHADER_TYPE(,FApplyLightShaftsPixelShader,TEXT("LightShaftShader"),TEX
 
 void ApplyLightShaftBloom(FRHICommandListImmediate& RHICmdList, const FViewInfo& View, const FLightSceneInfo* const LightSceneInfo, TRefCountPtr<IPooledRenderTarget>& LightShaftsSource)
 {
-	GSceneRenderTargets.BeginRenderingSceneColor(RHICmdList, ESimpleRenderTargetMode::EUninitializedColorExistingReadOnlyDepth);
+	GSceneRenderTargets.BeginRenderingSceneColor(RHICmdList, ESimpleRenderTargetMode::EUninitializedColorExistingDepth, FExclusiveDepthStencil::DepthRead_StencilWrite);
 
 	RHICmdList.SetViewport(View.ViewRect.Min.X, View.ViewRect.Min.Y, 0.0f, View.ViewRect.Max.X, View.ViewRect.Max.Y, 1.0f);
 	RHICmdList.SetBlendState(TStaticBlendState<CW_RGB, BO_Add, BF_One, BF_One>::GetRHI());
diff --git a/Engine/Source/Runtime/Renderer/Private/PostProcess/PostProcessTemporalAA.cpp b/Engine/Source/Runtime/Renderer/Private/PostProcess/PostProcessTemporalAA.cpp
--- a/Engine/Source/Runtime/Renderer/Private/PostProcess/PostProcessTemporalAA.cpp
+++ b/Engine/Source/Runtime/Renderer/Private/PostProcess/PostProcessTemporalAA.cpp
@@ -547,7 +547,7 @@ void FRCPassPostProcessTemporalAA::Process(FRenderingCompositePassContext& Conte
 	const FSceneRenderTargetItem& DestRenderTarget = PassOutputs[0].RequestSurface(Context);
 
 	//Context.SetRenderTarget(RHICmdList, DestRenderTarget.TargetableTexture, FTextureRHIRef());
-	SetRenderTarget(Context.RHICmdList, DestRenderTarget.TargetableTexture, GSceneRenderTargets.GetSceneDepthTexture(), ESimpleRenderTargetMode::EUninitializedColorExistingReadOnlyDepth);
+	SetRenderTarget(Context.RHICmdList, DestRenderTarget.TargetableTexture, GSceneRenderTargets.GetSceneDepthTexture(), ESimpleRenderTargetMode::EUninitializedColorExistingDepth, FExclusiveDepthStencil::DepthRead_StencilWrite);
 
 	// is optimized away if possible (RT size=view size, )
 	Context.RHICmdList.Clear(true, FLinearColor::Black, false, (float)ERHIZBuffer::FarPlane, false, 0, SrcRect);
diff --git a/Engine/Source/Runtime/Renderer/Private/PostProcess/SceneRenderTargets.cpp b/Engine/Source/Runtime/Renderer/Private/PostProcess/SceneRenderTargets.cpp
--- a/Engine/Source/Runtime/Renderer/Private/PostProcess/SceneRenderTargets.cpp
+++ b/Engine/Source/Runtime/Renderer/Private/PostProcess/SceneRenderTargets.cpp
@@ -299,13 +299,13 @@ void FSceneRenderTargets::Allocate(const FSceneViewFamily& ViewFamily)
 	AllocateRenderTargets();
 }
 
-void FSceneRenderTargets::BeginRenderingSceneColor(FRHICommandList& RHICmdList, ESimpleRenderTargetMode RenderTargetMode/*=EUninitializedColorExistingDepth*/)
+void FSceneRenderTargets::BeginRenderingSceneColor(FRHICommandList& RHICmdList, ESimpleRenderTargetMode RenderTargetMode/*=EUninitializedColorExistingDepth*/, FExclusiveDepthStencil DepthStencilAccess)
 {
 	SCOPED_DRAW_EVENT(RHICmdList, BeginRenderingSceneColor);
 
 	AllocSceneColor();
 	
-	SetRenderTarget(RHICmdList, GetSceneColorSurface(), GetSceneDepthSurface(), RenderTargetMode);
+	SetRenderTarget(RHICmdList, GetSceneColorSurface(), GetSceneDepthSurface(), RenderTargetMode, DepthStencilAccess);
 } 
 
 int32 FSceneRenderTargets::GetGBufferRenderTargets(ERenderTargetLoadAction ColorLoadAction, FRHIRenderTargetView OutRenderTargets[MaxSimultaneousRenderTargets], int32& OutVelocityRTIndex)
@@ -1041,7 +1041,7 @@ void FSceneRenderTargets::BeginRenderingLightAttenuation(FRHICommandList& RHICmd
 	GRenderTargetPool.VisualizeTexture.SetCheckPoint(RHICmdList, GSceneRenderTargets.GetLightAttenuation());
 
 	// Set the light attenuation surface as the render target, and the scene depth buffer as the depth-stencil surface.
-	SetRenderTarget(RHICmdList, GetLightAttenuationSurface(), GetSceneDepthSurface(), bClearToWhite ? ESimpleRenderTargetMode::EClearColorToWhiteReadOnlyDepth : ESimpleRenderTargetMode::EExistingColorAndReadOnlyDepth);
+	SetRenderTarget(RHICmdList, GetLightAttenuationSurface(), GetSceneDepthSurface(), bClearToWhite ? ESimpleRenderTargetMode::EClearColorToWhite : ESimpleRenderTargetMode::EExistingColorAndDepth, FExclusiveDepthStencil::DepthRead_StencilWrite);
 }
 
 void FSceneRenderTargets::FinishRenderingLightAttenuation(FRHICommandList& RHICmdList)
@@ -1057,7 +1057,7 @@ void FSceneRenderTargets::FinishRenderingLightAttenuation(FRHICommandList& RHICm
 void FSceneRenderTargets::BeginRenderingTranslucency(FRHICommandList& RHICmdList, const FViewInfo& View)
 {
 	// Use the scene color buffer.
-	GSceneRenderTargets.BeginRenderingSceneColor(RHICmdList, ESimpleRenderTargetMode::EExistingColorAndReadOnlyDepth);
+	GSceneRenderTargets.BeginRenderingSceneColor(RHICmdList, ESimpleRenderTargetMode::EExistingColorAndDepth, FExclusiveDepthStencil::DepthRead_StencilWrite);
 	
 	// viewport to match view size
 	RHICmdList.SetViewport(View.ViewRect.Min.X, View.ViewRect.Min.Y, 0.0f, View.ViewRect.Max.X, View.ViewRect.Max.Y, 1.0f);
@@ -1087,7 +1087,7 @@ bool FSceneRenderTargets::BeginRenderingSeparateTranslucency(FRHICommandList& RH
 
 		// clear the render target the first time, re-use afterwards
 		SetRenderTarget(RHICmdList, SeparateTranslucency->GetRenderTargetItem().TargetableTexture, GetSceneDepthSurface(), 
-			bFirstTimeThisFrame ? ESimpleRenderTargetMode::EClearColorToBlackWithFullAlphaReadOnlyDepth : ESimpleRenderTargetMode::EExistingColorAndReadOnlyDepth);
+			bFirstTimeThisFrame ? ESimpleRenderTargetMode::EClearColorToBlackWithFullAlpha : ESimpleRenderTargetMode::EExistingColorAndDepth, FExclusiveDepthStencil::DepthRead_StencilWrite);
 
 		if (!bFirstTimeThisFrame)
 		{
diff --git a/Engine/Source/Runtime/Renderer/Private/PostProcess/SceneRenderTargets.h b/Engine/Source/Runtime/Renderer/Private/PostProcess/SceneRenderTargets.h
--- a/Engine/Source/Runtime/Renderer/Private/PostProcess/SceneRenderTargets.h
+++ b/Engine/Source/Runtime/Renderer/Private/PostProcess/SceneRenderTargets.h
@@ -170,7 +170,7 @@ public:
 	/**
 	 * Sets the scene color target and restores its contents if necessary
 	 */
-	void BeginRenderingSceneColor(FRHICommandList& RHICmdList, ESimpleRenderTargetMode RenderTargetMode=ESimpleRenderTargetMode::EUninitializedColorExistingDepth);
+	void BeginRenderingSceneColor(FRHICommandList& RHICmdList, ESimpleRenderTargetMode RenderTargetMode = ESimpleRenderTargetMode::EUninitializedColorExistingDepth, FExclusiveDepthStencil DepthStencilAccess = FExclusiveDepthStencil::DepthWrite_StencilWrite);
 	
 	/**
 	 * Called when finished rendering to the scene color surface
diff --git a/Engine/Source/Runtime/Renderer/Private/ReflectionEnvironment.cpp b/Engine/Source/Runtime/Renderer/Private/ReflectionEnvironment.cpp
--- a/Engine/Source/Runtime/Renderer/Private/ReflectionEnvironment.cpp
+++ b/Engine/Source/Runtime/Renderer/Private/ReflectionEnvironment.cpp
@@ -655,7 +655,7 @@ IMPLEMENT_SHADER_TYPE(template<>,TStandardDeferredReflectionPS<false>,TEXT("Refl
 
 void FDeferredShadingSceneRenderer::RenderReflectionCaptureSpecularBounceForAllViews(FRHICommandListImmediate& RHICmdList)
 {
-	GSceneRenderTargets.BeginRenderingSceneColor(RHICmdList, ESimpleRenderTargetMode::EUninitializedColorExistingReadOnlyDepth);
+	GSceneRenderTargets.BeginRenderingSceneColor(RHICmdList, ESimpleRenderTargetMode::EUninitializedColorExistingDepth, FExclusiveDepthStencil::DepthRead_StencilWrite);
 	RHICmdList.SetRasterizerState(TStaticRasterizerState< FM_Solid, CM_None >::GetRHI());
 	RHICmdList.SetDepthStencilState(TStaticDepthStencilState< false, CF_Always >::GetRHI());
 	RHICmdList.SetBlendState(TStaticBlendState< CW_RGB, BO_Add, BF_One, BF_One >::GetRHI());
@@ -923,7 +923,7 @@ void FDeferredShadingSceneRenderer::RenderStandardDeferredImageBasedReflections(
 			// Apply reflections to screen
 			SCOPED_DRAW_EVENT(RHICmdList, ReflectionApply);
 
-			GSceneRenderTargets.BeginRenderingSceneColor(RHICmdList, ESimpleRenderTargetMode::EUninitializedColorExistingReadOnlyDepth);
+			GSceneRenderTargets.BeginRenderingSceneColor(RHICmdList, ESimpleRenderTargetMode::EUninitializedColorExistingDepth, FExclusiveDepthStencil::DepthRead_StencilWrite);
 
 			RHICmdList.SetViewport(View.ViewRect.Min.X, View.ViewRect.Min.Y, 0.0f, View.ViewRect.Max.X, View.ViewRect.Max.Y, 1.0f);
 			RHICmdList.SetRasterizerState(TStaticRasterizerState<FM_Solid, CM_None>::GetRHI());
diff --git a/Engine/Source/Runtime/Renderer/Private/SceneHitProxyRendering.cpp b/Engine/Source/Runtime/Renderer/Private/SceneHitProxyRendering.cpp
--- a/Engine/Source/Runtime/Renderer/Private/SceneHitProxyRendering.cpp
+++ b/Engine/Source/Runtime/Renderer/Private/SceneHitProxyRendering.cpp
@@ -350,7 +350,7 @@ TRefCountPtr<IPooledRenderTarget> InitHitProxyRender(FRHICommandListImmediate& R
 		return HitProxyRT;
 	}
 
-	SetRenderTarget(RHICmdList, HitProxyRT->GetRenderTargetItem().TargetableTexture, GSceneRenderTargets.GetSceneDepthSurface(), ESimpleRenderTargetMode::EExistingColorAndReadOnlyDepth);
+	SetRenderTarget(RHICmdList, HitProxyRT->GetRenderTargetItem().TargetableTexture, GSceneRenderTargets.GetSceneDepthSurface(), ESimpleRenderTargetMode::EExistingColorAndDepth, FExclusiveDepthStencil::DepthWrite_StencilWrite);
 
 	// Clear color for each view.
 	auto& Views = SceneRenderer->Views;
diff --git a/Engine/Source/Runtime/Renderer/Private/SceneOcclusion.cpp b/Engine/Source/Runtime/Renderer/Private/SceneOcclusion.cpp
--- a/Engine/Source/Runtime/Renderer/Private/SceneOcclusion.cpp
+++ b/Engine/Source/Runtime/Renderer/Private/SceneOcclusion.cpp
@@ -1025,11 +1025,11 @@ void FDeferredShadingSceneRenderer::BeginOcclusionTests(FRHICommandListImmediate
 
 	if (bUseDownsampledDepth)
 	{
-		SetRenderTarget(RHICmdList, NULL, GSceneRenderTargets.GetSmallDepthSurface(), ESimpleRenderTargetMode::EExistingColorAndReadOnlyDepth);
+		SetRenderTarget(RHICmdList, NULL, GSceneRenderTargets.GetSmallDepthSurface(), ESimpleRenderTargetMode::EExistingColorAndDepth, FExclusiveDepthStencil::DepthRead_StencilWrite);
 	}
 	else
 	{
-		SetRenderTarget(RHICmdList, NULL, GSceneRenderTargets.GetSceneDepthSurface(), ESimpleRenderTargetMode::EExistingColorAndReadOnlyDepth);
+		SetRenderTarget(RHICmdList, NULL, GSceneRenderTargets.GetSceneDepthSurface(), ESimpleRenderTargetMode::EExistingColorAndDepth, FExclusiveDepthStencil::DepthRead_StencilWrite);
 	}
 
 	if (bRenderQueries)
@@ -1192,6 +1192,6 @@ void FDeferredShadingSceneRenderer::BeginOcclusionTests(FRHICommandListImmediate
 	if (bUseDownsampledDepth)
 	{
 		// Restore default render target
-		GSceneRenderTargets.BeginRenderingSceneColor(RHICmdList, ESimpleRenderTargetMode::EUninitializedColorExistingReadOnlyDepth);
+		GSceneRenderTargets.BeginRenderingSceneColor(RHICmdList, ESimpleRenderTargetMode::EUninitializedColorExistingDepth, FExclusiveDepthStencil::DepthRead_StencilWrite);
 	}
 }
diff --git a/Engine/Source/Runtime/Renderer/Private/VelocityRendering.cpp b/Engine/Source/Runtime/Renderer/Private/VelocityRendering.cpp
--- a/Engine/Source/Runtime/Renderer/Private/VelocityRendering.cpp
+++ b/Engine/Source/Runtime/Renderer/Private/VelocityRendering.cpp
@@ -569,7 +569,7 @@ static void BeginVelocityRendering(FRHICommandList& RHICmdList, TRefCountPtr<IPo
 	{
 		// now make the FRHISetRenderTargetsInfo that encapsulates all of the info
 		FRHIRenderTargetView ColorView(VelocityTexture, 0, -1, ERenderTargetLoadAction::EClear, ERenderTargetStoreAction::EStore);
-		FRHIDepthRenderTargetView DepthView(DepthTexture, ERenderTargetLoadAction::ELoad, ERenderTargetStoreAction::ENoAction, true);
+		FRHIDepthRenderTargetView DepthView(DepthTexture, ERenderTargetLoadAction::ELoad, ERenderTargetStoreAction::ENoAction, FExclusiveDepthStencil::DepthRead_StencilWrite);
 
 		FRHISetRenderTargetsInfo Info(1, &ColorView, DepthView);
 		Info.ClearColors[0] = VelocityClearColor;
@@ -579,7 +579,7 @@ static void BeginVelocityRendering(FRHICommandList& RHICmdList, TRefCountPtr<IPo
 	}
 	else
 	{
-		SetRenderTarget(RHICmdList, VelocityTexture, DepthTexture, ESimpleRenderTargetMode::EExistingColorAndReadOnlyDepth);
+		SetRenderTarget(RHICmdList, VelocityTexture, DepthTexture, ESimpleRenderTargetMode::EExistingColorAndDepth, FExclusiveDepthStencil::DepthRead_StencilWrite);
 
 		// some platforms need the clear color when rendertargets transition to SRVs.  We propagate here to allow parallel rendering to always
 		// have the proper mapping when the RT is transitioned.
diff --git a/Engine/Source/Runtime/Windows/D3D11RHI/Private/D3D11Commands.cpp b/Engine/Source/Runtime/Windows/D3D11RHI/Private/D3D11Commands.cpp
--- a/Engine/Source/Runtime/Windows/D3D11RHI/Private/D3D11Commands.cpp
+++ b/Engine/Source/Runtime/Windows/D3D11RHI/Private/D3D11Commands.cpp
@@ -638,20 +638,30 @@ void FD3D11DynamicRHI::RHISetShaderParameter(FComputeShaderRHIParamRef ComputeSh
 	CSConstantBuffers[BufferIndex]->UpdateConstant((const uint8*)NewValue,BaseIndex,NumBytes);
 }
 
-void FD3D11DynamicRHI::RHISetDepthStencilState(FDepthStencilStateRHIParamRef NewStateRHI,uint32 StencilRef)
+void FD3D11DynamicRHI::ValidateExclusiveDepthStencilAccess(FExclusiveDepthStencil RequestedAccess) const
 {
-	DYNAMIC_CAST_D3D11RESOURCE(DepthStencilState,NewState);
+	const bool bSrcDepthWrite = RequestedAccess.IsDepthWrite();
+	const bool bSrcStencilWrite = RequestedAccess.IsStencilWrite();
 
-	if (CurrentDepthTexture && NewState->AccessType != CurrentDSVAccessType)
+	if (bSrcDepthWrite || bSrcStencilWrite)
 	{
-		CurrentDSVAccessType = NewState->AccessType;
-		CurrentDepthStencilTarget = CurrentDepthTexture->GetDepthStencilView(CurrentDSVAccessType);
+		// New Rule: You have to call SetRenderTarget[s]() before
+		ensure(CurrentDepthTexture);
 
-		// Unbind any shader views of the depth stencil target that are bound.
-		ConditionalClearShaderResource(CurrentDepthTexture);
+		const bool bDstDepthWrite = CurrentDSVAccessType.IsDepthWrite();
+		const bool bDstStencilWrite = CurrentDSVAccessType.IsStencilWrite();
 
-		CommitRenderTargetsAndUAVs();
+		// requested access is not possible, fix SetRenderTarget EExclusiveDepthStencil or request a different one
+		check(!bSrcDepthWrite || bDstDepthWrite);
+		check(!bSrcStencilWrite || bDstStencilWrite);
 	}
+}
+
+void FD3D11DynamicRHI::RHISetDepthStencilState(FDepthStencilStateRHIParamRef NewStateRHI,uint32 StencilRef)
+{
+	DYNAMIC_CAST_D3D11RESOURCE(DepthStencilState,NewState);
+
+	ValidateExclusiveDepthStencilAccess(NewState->AccessType);
 
 	StateCache.SetDepthStencilState(NewState->Resource, StencilRef);
 }
@@ -774,7 +784,7 @@ void FD3D11DynamicRHI::RHISetRenderTargets(
 
 #if CHECK_SRV_TRANSITIONS
 	// if the depth buffer is writable then it counts as unresolved.
-	if (NewDepthStencilTargetRHI && !NewDepthStencilTargetRHI->bReadOnly && NewDepthStencilTarget)
+	if (NewDepthStencilTargetRHI && NewDepthStencilTargetRHI->GetDepthStencilAccess() == FExclusiveDepthStencil::DepthWrite_StencilWrite && NewDepthStencilTarget)
 	{		
 		UnresolvedTargets.Add(NewDepthStencilTarget->GetResource(), FUnresolvedRTInfo(NewDepthStencilTargetRHI->Texture->GetName(), 0, 1, -1, 1));
 	}
@@ -788,8 +798,9 @@ void FD3D11DynamicRHI::RHISetRenderTargets(
 	ID3D11DepthStencilView* DepthStencilView = NULL;
 	if(NewDepthStencilTarget)
 	{
+		CurrentDSVAccessType = NewDepthStencilTargetRHI->GetDepthStencilAccess();
 		DepthStencilView = NewDepthStencilTarget->GetDepthStencilView(CurrentDSVAccessType);
-		
+
 		// Unbind any shader views of the depth stencil target that are bound.
 		ConditionalClearShaderResource(NewDepthStencilTarget);
 	}
@@ -1552,24 +1563,18 @@ void FD3D11DynamicRHI::RHIClearMRT(bool bClearColor,int32 NumClearColors,const F
 	// Must specify enough clear colors for all active RTs
 	check(!bClearColor || NumClearColors >= BoundRenderTargets.GetNumActiveTargets());
 
-	ID3D11DepthStencilView* DepthStencilView = BoundRenderTargets.GetDepthStencilView();
-
-	// If we're clearing depth or stencil and we have a readonly depth/stencil view bound, we need to use a writable depth/stencil view instead
+	// If we're clearing depth or stencil and we have a readonly depth/stencil view bound, we need to use a writable depth/stencil view
 	if (CurrentDepthTexture)
 	{
-		// Create an access type mask by setting the readonly bits according to the bClearDepth/bClearStencil bools.
-		EDepthStencilAccessType AccessMask = (EDepthStencilAccessType)((bClearDepth ? DSAT_Writable : DSAT_ReadOnlyDepth) | (bClearStencil ? DSAT_Writable : DSAT_ReadOnlyStencil));
-
-		// Apply the mask to the current access state
-		EDepthStencilAccessType RequiredAccess = (EDepthStencilAccessType)(CurrentDSVAccessType & AccessMask);
+		FExclusiveDepthStencil RequestedAccess;
+		
+		RequestedAccess.SetDepthStencilWrite(bClearDepth, bClearStencil);
 
-		if (RequiredAccess != CurrentDSVAccessType)
-		{
-			// A different access type is required than is currently set. Replace the DSV used for the clear.
-			DepthStencilView = CurrentDepthTexture->GetDepthStencilView(RequiredAccess);
-		}
+		ensure(RequestedAccess.IsValid(CurrentDSVAccessType));
 	}
 
+	ID3D11DepthStencilView* DepthStencilView = BoundRenderTargets.GetDepthStencilView();
+
 	// Determine if we're trying to clear a subrect of the screen
 	bool UseDrawClear = false;
 	uint32 NumViews = 1;
@@ -1680,21 +1685,13 @@ void FD3D11DynamicRHI::RHIClearMRT(bool bClearColor,int32 NumClearColors,const F
 					>::GetRHI()
 			:     TStaticDepthStencilState<false, CF_Always>::GetRHI();
 
-		bool bChangedDepthStencilTarget = false;
-		if (CurrentDepthTexture)
+		// Create an access type mask by setting the readonly bits according to the bClearDepth/bClearStencil bools.
 		{
-			// Create an access type mask by setting the readonly bits according to the bClearDepth/bClearStencil bools.
-			EDepthStencilAccessType AccessMask = (EDepthStencilAccessType)((bClearDepth ? DSAT_Writable : DSAT_ReadOnlyDepth) | (bClearStencil ? DSAT_Writable : DSAT_ReadOnlyStencil));
-
-			// Apply the mask to the current access state
-			EDepthStencilAccessType RequiredAccess = (EDepthStencilAccessType)(CurrentDSVAccessType & AccessMask);
+			FExclusiveDepthStencil RequestedAccess;
+			
+			RequestedAccess.SetDepthStencilWrite(bClearDepth, bClearStencil);
 
-			if (RequiredAccess != CurrentDSVAccessType)
-			{
-				bChangedDepthStencilTarget = true;
-				CurrentDepthStencilTarget = CurrentDepthTexture->GetDepthStencilView(RequiredAccess);
-				CommitRenderTargetsAndUAVs();
-			} 
+			ValidateExclusiveDepthStencilAccess(RequestedAccess);
 		}
 
 		DYNAMIC_CAST_D3D11RESOURCE(BlendState,BlendState);
@@ -1822,14 +1819,6 @@ void FD3D11DynamicRHI::RHIClearMRT(bool bClearColor,int32 NumClearColors,const F
 			// Implicit flush. Always call flush when using a command list in RHI implementations before doing anything else. This is super hazardous.
 		}
 
-		if (bChangedDepthStencilTarget)
-		{
-			// Restore the DST
-			CurrentDepthStencilTarget = CurrentDepthTexture->GetDepthStencilView(CurrentDSVAccessType);
-
-			CommitRenderTargetsAndUAVs();
-		}
-
 		// Restore the original device state
 		OriginalResourceState.RestoreDeviceState(StateCache); 
 	}
diff --git a/Engine/Source/Runtime/Windows/D3D11RHI/Private/D3D11Device.cpp b/Engine/Source/Runtime/Windows/D3D11RHI/Private/D3D11Device.cpp
--- a/Engine/Source/Runtime/Windows/D3D11RHI/Private/D3D11Device.cpp
+++ b/Engine/Source/Runtime/Windows/D3D11RHI/Private/D3D11Device.cpp
@@ -43,7 +43,7 @@ FD3D11DynamicRHI::FD3D11DynamicRHI(IDXGIFactory1* InDXGIFactory1,D3D_FEATURE_LEV
 	NumUAVs(0),
 	SceneFrameCounter(0),
 	ResourceTableFrameCounter(INDEX_NONE),
-	CurrentDSVAccessType(DSAT_Writable),
+	CurrentDSVAccessType(FExclusiveDepthStencil::DepthWrite_StencilWrite),
 	bDiscardSharedConstants(false),
 	bUsingTessellation(false),
 	PendingNumVertices(0),
diff --git a/Engine/Source/Runtime/Windows/D3D11RHI/Private/D3D11RHIPrivate.h b/Engine/Source/Runtime/Windows/D3D11RHI/Private/D3D11RHIPrivate.h
--- a/Engine/Source/Runtime/Windows/D3D11RHI/Private/D3D11RHIPrivate.h
+++ b/Engine/Source/Runtime/Windows/D3D11RHI/Private/D3D11RHIPrivate.h
@@ -428,7 +428,7 @@ protected:
 	uint16 DirtyUniformBuffers[SF_NumFrequencies];
 
 	/** Tracks the current depth stencil access type. */
-	EDepthStencilAccessType CurrentDSVAccessType;
+	FExclusiveDepthStencil CurrentDSVAccessType;
 
 	/** When a new shader is set, we discard all old constants set for the previous shader. */
 	bool bDiscardSharedConstants;
@@ -517,6 +517,8 @@ protected:
 	void CommitGraphicsResourceTables();
 	void CommitComputeResourceTables(FD3D11ComputeShader* ComputeShader);
 
+	void ValidateExclusiveDepthStencilAccess(FExclusiveDepthStencil Src) const;
+
 	/** 
 	 * Gets the best supported MSAA settings from the provided MSAA count to check against. 
 	 * 
diff --git a/Engine/Source/Runtime/Windows/D3D11RHI/Private/D3D11RenderTarget.cpp b/Engine/Source/Runtime/Windows/D3D11RHI/Private/D3D11RenderTarget.cpp
--- a/Engine/Source/Runtime/Windows/D3D11RHI/Private/D3D11RenderTarget.cpp
+++ b/Engine/Source/Runtime/Windows/D3D11RHI/Private/D3D11RenderTarget.cpp
@@ -251,7 +251,7 @@ void FD3D11DynamicRHI::RHICopyToResolveTarget(FTextureRHIParamRef SourceTextureR
 			GPUProfilingData.RegisterGPUWork();
 		
 			if(FeatureLevel == D3D_FEATURE_LEVEL_11_0 
-			&& DestTexture2D->GetDepthStencilView(DSAT_Writable)
+			&& DestTexture2D->GetDepthStencilView(FExclusiveDepthStencil::DepthWrite_StencilWrite)
 			&& SourceTextureRHI->IsMultisampled()
 			&& !DestTextureRHI->IsMultisampled())
 			{
@@ -264,7 +264,7 @@ void FD3D11DynamicRHI::RHICopyToResolveTarget(FTextureRHIParamRef SourceTextureR
 					SourceTexture2D,
 					DestTexture2D,
 					DestTexture2D->GetRenderTargetView(0, -1),
-					DestTexture2D->GetDepthStencilView(DSAT_Writable),
+					DestTexture2D->GetDepthStencilView(FExclusiveDepthStencil::DepthWrite_StencilWrite),
 					ResolveTargetDesc,
 					GetDefaultRect(ResolveParams.Rect,DestTexture2D->GetSizeX(),DestTexture2D->GetSizeY()),
 					GetDefaultRect(ResolveParams.Rect,DestTexture2D->GetSizeX(),DestTexture2D->GetSizeY()),
@@ -273,7 +273,7 @@ void FD3D11DynamicRHI::RHICopyToResolveTarget(FTextureRHIParamRef SourceTextureR
 					);
 			}
 			else if(FeatureLevel == D3D_FEATURE_LEVEL_10_0 
-					&& DestTexture2D->GetDepthStencilView(DSAT_Writable))
+				&& DestTexture2D->GetDepthStencilView(FExclusiveDepthStencil::DepthWrite_StencilWrite))
 			{
 				D3D11_TEXTURE2D_DESC ResolveTargetDesc;
 
@@ -284,7 +284,7 @@ void FD3D11DynamicRHI::RHICopyToResolveTarget(FTextureRHIParamRef SourceTextureR
 					SourceTexture2D,
 					DestTexture2D,
 					NULL,
-					DestTexture2D->GetDepthStencilView(DSAT_Writable),
+					DestTexture2D->GetDepthStencilView(FExclusiveDepthStencil::DepthWrite_StencilWrite),
 					ResolveTargetDesc,
 					GetDefaultRect(ResolveParams.Rect,DestTexture2D->GetSizeX(),DestTexture2D->GetSizeY()),
 					GetDefaultRect(ResolveParams.Rect,DestTexture2D->GetSizeX(),DestTexture2D->GetSizeY()),
diff --git a/Engine/Source/Runtime/Windows/D3D11RHI/Private/D3D11State.cpp b/Engine/Source/Runtime/Windows/D3D11RHI/Private/D3D11State.cpp
--- a/Engine/Source/Runtime/Windows/D3D11RHI/Private/D3D11State.cpp
+++ b/Engine/Source/Runtime/Windows/D3D11RHI/Private/D3D11State.cpp
@@ -251,7 +251,7 @@ FDepthStencilStateRHIRef FD3D11DynamicRHI::RHICreateDepthStencilState(const FDep
 		&& Initializer.BackFacePassStencilOp == SO_Keep;
 
 	const bool bMayWriteStencil = Initializer.StencilWriteMask != 0 && !bStencilOpIsKeep;
-	DepthStencilState->AccessType = (EDepthStencilAccessType)((Initializer.bEnableDepthWrite ? DSAT_Writable : DSAT_ReadOnlyDepth) | (bMayWriteStencil ? DSAT_Writable : DSAT_ReadOnlyStencil));
+	DepthStencilState->AccessType.SetDepthStencilWrite(Initializer.bEnableDepthWrite, bMayWriteStencil);
 
 	VERIFYD3D11RESULT(Direct3DDevice->CreateDepthStencilState(&DepthStencilDesc,DepthStencilState->Resource.GetInitReference()));
 	return DepthStencilState;
diff --git a/Engine/Source/Runtime/Windows/D3D11RHI/Private/D3D11Texture.cpp b/Engine/Source/Runtime/Windows/D3D11RHI/Private/D3D11Texture.cpp
--- a/Engine/Source/Runtime/Windows/D3D11RHI/Private/D3D11Texture.cpp
+++ b/Engine/Source/Runtime/Windows/D3D11RHI/Private/D3D11Texture.cpp
@@ -626,7 +626,7 @@ TD3D11Texture2D<BaseResourceType>* FD3D11DynamicRHI::CreateD3D11Texture2D(uint32
 	TRefCountPtr<ID3D11Texture2D> TextureResource;
 	TRefCountPtr<ID3D11ShaderResourceView> ShaderResourceView;
 	TArray<TRefCountPtr<ID3D11RenderTargetView> > RenderTargetViews;
-	TRefCountPtr<ID3D11DepthStencilView> DepthStencilViews[DSAT_Count];
+	TRefCountPtr<ID3D11DepthStencilView> DepthStencilViews[FExclusiveDepthStencil::MaxIndex];
 	
 #if PLATFORM_SUPPORTS_VIRTUAL_TEXTURES
 	// Turn off pooling when we are using virtual textures or the texture is offline processed as we control when the memory is released
@@ -760,17 +760,17 @@ TD3D11Texture2D<BaseResourceType>* FD3D11DynamicRHI::CreateD3D11Texture2D(uint32
 				DSVDesc.Texture2D.MipSlice = 0;
 			}
 
-			for(uint32 AccessType = 0; AccessType < DSAT_Count; ++AccessType)
+			for (uint32 AccessType = 0; AccessType < FExclusiveDepthStencil::MaxIndex; ++AccessType)
 			{
 				// Create a read-only access views for the texture.
 				// Read-only DSVs are not supported in Feature Level 10 so 
 				// a dummy DSV is created in order reduce logic complexity at a higher-level.
 				if(Direct3DDevice->GetFeatureLevel() == D3D_FEATURE_LEVEL_11_0)
 				{
-					DSVDesc.Flags  = (AccessType & DSAT_ReadOnlyDepth) ? D3D11_DSV_READ_ONLY_DEPTH : 0;
+					DSVDesc.Flags = (AccessType & FExclusiveDepthStencil::DepthRead_StencilWrite) ? D3D11_DSV_READ_ONLY_DEPTH : 0;
 					if(HasStencilBits(DSVDesc.Format))
 					{
-						DSVDesc.Flags |= (AccessType & DSAT_ReadOnlyStencil) ? D3D11_DSV_READ_ONLY_STENCIL : 0;
+						DSVDesc.Flags |= (AccessType & FExclusiveDepthStencil::DepthWrite_StencilRead) ? D3D11_DSV_READ_ONLY_STENCIL : 0;
 					}
 				}
 				VERIFYD3D11RESULT(Direct3DDevice->CreateDepthStencilView(TextureResource,&DSVDesc,DepthStencilViews[AccessType].GetInitReference()));
diff --git a/Engine/Source/Runtime/Windows/D3D11RHI/Private/D3D11Viewport.cpp b/Engine/Source/Runtime/Windows/D3D11RHI/Private/D3D11Viewport.cpp
--- a/Engine/Source/Runtime/Windows/D3D11RHI/Private/D3D11Viewport.cpp
+++ b/Engine/Source/Runtime/Windows/D3D11RHI/Private/D3D11Viewport.cpp
@@ -530,7 +530,7 @@ void FD3D11DynamicRHI::RHIEndDrawingViewport(FViewportRHIParamRef ViewportRHI,bo
 	// Clear references the device might have to resources.
 	CurrentDepthTexture = NULL;
 	CurrentDepthStencilTarget = NULL;
-	CurrentDSVAccessType = DSAT_Writable;
+	CurrentDSVAccessType = FExclusiveDepthStencil::DepthWrite_StencilWrite;
 	CurrentRenderTargets[0] = NULL;
 	for(uint32 RenderTargetIndex = 1;RenderTargetIndex < D3D11_SIMULTANEOUS_RENDER_TARGET_COUNT;++RenderTargetIndex)
 	{
diff --git a/Engine/Source/Runtime/Windows/D3D11RHI/Public/D3D11Resources.h b/Engine/Source/Runtime/Windows/D3D11RHI/Public/D3D11Resources.h
--- a/Engine/Source/Runtime/Windows/D3D11RHI/Public/D3D11Resources.h
+++ b/Engine/Source/Runtime/Windows/D3D11RHI/Public/D3D11Resources.h
@@ -192,7 +192,7 @@ public:
 		// Set the DSVs for all the access type combinations
 		if ( InDepthStencilViews != nullptr )
 		{
-			for ( uint32 Index=0; Index<DSAT_Count; Index++ )
+			for (uint32 Index = 0; Index < FExclusiveDepthStencil::MaxIndex; Index++)
 			{
 				DepthStencilViews[Index] = InDepthStencilViews[Index];
 				// New Monolithic Graphics drivers have optional "fast calls" replacing various D3d functions
@@ -245,9 +245,9 @@ public:
 		}
 		return 0;
 	}
-	ID3D11DepthStencilView* GetDepthStencilView(EDepthStencilAccessType AccessType) const 
+	ID3D11DepthStencilView* GetDepthStencilView(FExclusiveDepthStencil AccessType) const
 	{ 
-		return DepthStencilViews[AccessType]; 
+		return DepthStencilViews[AccessType.GetIndex()]; 
 	}
 
 	// New Monolithic Graphics drivers have optional "fast calls" replacing various D3d functions
@@ -282,7 +282,7 @@ protected:
 	int32 RTVArraySize;
 
 	/** A depth-stencil targetable view of the texture. */
-	TRefCountPtr<ID3D11DepthStencilView> DepthStencilViews[ DSAT_Count ];
+	TRefCountPtr<ID3D11DepthStencilView> DepthStencilViews[FExclusiveDepthStencil::MaxIndex];
 
 	/** Number of Depth Stencil Views - used for fast call tracking. */
 	uint32	NumDepthStencilViews;
diff --git a/Engine/Source/Runtime/Windows/D3D11RHI/Public/D3D11State.h b/Engine/Source/Runtime/Windows/D3D11RHI/Public/D3D11State.h
--- a/Engine/Source/Runtime/Windows/D3D11RHI/Public/D3D11State.h
+++ b/Engine/Source/Runtime/Windows/D3D11RHI/Public/D3D11State.h
@@ -20,16 +20,6 @@ public:
 	TRefCountPtr<ID3D11RasterizerState> Resource;
 };
 
-enum EDepthStencilAccessType
-{
-	DSAT_Writable = 0,
-	DSAT_ReadOnlyDepth = 1,
-	DSAT_ReadOnlyStencil = 2,
-	DSAT_ReadOnlyDepthAndStencil = 3,
-
-	DSAT_Count = 4
-};
-
 class FD3D11DepthStencilState : public FRHIDepthStencilState
 {
 public:
@@ -37,7 +27,7 @@ public:
 	TRefCountPtr<ID3D11DepthStencilState> Resource;
 
 	/* Describes the read/write state of the separate depth and stencil components of the DSV. */
-	EDepthStencilAccessType AccessType;
+	FExclusiveDepthStencil AccessType;
 };
 
 class FD3D11BlendState : public FRHIBlendState
