UE4 - Remove some padding by reordering some members (and reduced sizes)

diff --git a/Engine/Source/Runtime/Engine/Public/SceneManagement.h b/Engine/Source/Runtime/Engine/Public/SceneManagement.h
--- a/Engine/Source/Runtime/Engine/Public/SceneManagement.h
+++ b/Engine/Source/Runtime/Engine/Public/SceneManagement.h
@@ -436,8 +436,8 @@ public:
 
 	/** Default constructor. */
 	FShadowMapInteraction() :
-		Type(SMIT_None),
-		ShadowTexture(NULL)
+		ShadowTexture(nullptr),
+		Type(SMIT_None)
 	{
 		for (int Channel = 0; Channel < ARRAY_COUNT(bChannelValid); Channel++)
 		{
@@ -473,12 +473,11 @@ public:
 	}
 
 private:
-
-	EShadowMapInteractionType Type;
 	UShadowMapTexture2D* ShadowTexture;
 	FVector2D CoordinateScale;
 	FVector2D CoordinateBias;
 	bool bChannelValid[4];
+	EShadowMapInteractionType Type;
 };
 
 /**
diff --git a/Engine/Source/Runtime/RHI/Public/RHICommandList.h b/Engine/Source/Runtime/RHI/Public/RHICommandList.h
--- a/Engine/Source/Runtime/RHI/Public/RHICommandList.h
+++ b/Engine/Source/Runtime/RHI/Public/RHICommandList.h
@@ -686,12 +686,12 @@ struct FRHICommandFlushComputeShaderCache : public FRHICommand<FRHICommandFlushC
 
 struct FRHICommandDrawPrimitiveIndirect : public FRHICommand<FRHICommandDrawPrimitiveIndirect>
 {
-	uint32 PrimitiveType;
 	FVertexBufferRHIParamRef ArgumentBuffer;
+	uint32 PrimitiveType;
 	uint32 ArgumentOffset;
 	FORCEINLINE_DEBUGGABLE FRHICommandDrawPrimitiveIndirect(uint32 InPrimitiveType, FVertexBufferRHIParamRef InArgumentBuffer, uint32 InArgumentOffset)
-		: PrimitiveType(InPrimitiveType)
-		, ArgumentBuffer(InArgumentBuffer)
+		: ArgumentBuffer(InArgumentBuffer)
+		, PrimitiveType(InPrimitiveType)
 		, ArgumentOffset(InArgumentOffset)
 	{
 	}
@@ -719,15 +719,15 @@ struct FRHICommandDrawIndexedIndirect : public FRHICommand<FRHICommandDrawIndexe
 
 struct FRHICommandDrawIndexedPrimitiveIndirect : public FRHICommand<FRHICommandDrawIndexedPrimitiveIndirect>
 {
-	uint32 PrimitiveType;
 	FIndexBufferRHIParamRef IndexBuffer;
 	FVertexBufferRHIParamRef ArgumentsBuffer;
+	uint32 PrimitiveType;
 	uint32 ArgumentOffset;
 
 	FORCEINLINE_DEBUGGABLE FRHICommandDrawIndexedPrimitiveIndirect(uint32 InPrimitiveType, FIndexBufferRHIParamRef InIndexBuffer, FVertexBufferRHIParamRef InArgumentsBuffer, uint32 InArgumentOffset)
-		: PrimitiveType(InPrimitiveType)
-		, IndexBuffer(InIndexBuffer)
+		: IndexBuffer(InIndexBuffer)
 		, ArgumentsBuffer(InArgumentsBuffer)
+		, PrimitiveType(InPrimitiveType)
 		, ArgumentOffset(InArgumentOffset)
 	{
 	}
@@ -767,16 +767,16 @@ struct FRHICommandClearUAV : public FRHICommand<FRHICommandClearUAV>
 
 struct FRHICommandCopyToResolveTarget : public FRHICommand<FRHICommandCopyToResolveTarget>
 {
+	FResolveParams ResolveParams;
 	FTextureRHIParamRef SourceTexture;
 	FTextureRHIParamRef DestTexture;
 	bool bKeepOriginalSurface;
-	FResolveParams ResolveParams;
 
 	FORCEINLINE_DEBUGGABLE FRHICommandCopyToResolveTarget(FTextureRHIParamRef InSourceTexture, FTextureRHIParamRef InDestTexture, bool InbKeepOriginalSurface, const FResolveParams& InResolveParams)
-		: SourceTexture(InSourceTexture)
+		: ResolveParams(InResolveParams)
+		, SourceTexture(InSourceTexture)
 		, DestTexture(InDestTexture)
 		, bKeepOriginalSurface(InbKeepOriginalSurface)
-		, ResolveParams(InResolveParams)
 	{
 	}
 	RHI_API void Execute(FRHICommandListBase& CmdList);
diff --git a/Engine/Source/Runtime/Renderer/Private/SceneCore.h b/Engine/Source/Runtime/Renderer/Private/SceneCore.h
--- a/Engine/Source/Runtime/Renderer/Private/SceneCore.h
+++ b/Engine/Source/Runtime/Renderer/Private/SceneCore.h
@@ -53,15 +53,27 @@ public:
 	void operator delete(void *RawMemory);
 
 private:
-	/** The index into Scene->Lights of the light which affects the primitive. */
-	int32 LightId;
-
 	/** The light which affects the primitive. */
 	FLightSceneInfo* LightSceneInfo;
 
 	/** The primitive which is affected by the light. */
 	FPrimitiveSceneInfo* PrimitiveSceneInfo;
 
+	/** A pointer to the NextPrimitive member of the previous interaction in the light's interaction list. */
+	FLightPrimitiveInteraction** PrevPrimitiveLink;
+
+	/** The next interaction in the light's interaction list. */
+	FLightPrimitiveInteraction* NextPrimitive;
+
+	/** A pointer to the NextLight member of the previous interaction in the primitive's interaction list. */
+	FLightPrimitiveInteraction** PrevLightLink;
+
+	/** The next interaction in the primitive's interaction list. */
+	FLightPrimitiveInteraction* NextLight;
+
+	/** The index into Scene->Lights of the light which affects the primitive. */
+	int32 LightId;
+
 	/** True if the primitive casts a shadow from the light. */
 	uint32 bCastShadow : 1;
 
@@ -86,19 +98,6 @@ private:
 	/** True if the primitive only shadows itself. */
 	uint32 bSelfShadowOnly : 1;
 
-	/** A pointer to the NextPrimitive member of the previous interaction in the light's interaction list. */
-	FLightPrimitiveInteraction** PrevPrimitiveLink;
-
-	/** The next interaction in the light's interaction list. */
-	FLightPrimitiveInteraction* NextPrimitive;
-
-	/** A pointer to the NextLight member of the previous interaction in the primitive's interaction list. */
-	FLightPrimitiveInteraction** PrevLightLink;
-
-	/** The next interaction in the primitive's interaction list. */
-	FLightPrimitiveInteraction* NextLight;
-
-
 	/** Initialization constructor. */
 	FLightPrimitiveInteraction(FLightSceneInfo* InLightSceneInfo,FPrimitiveSceneInfo* InPrimitiveSceneInfo,
 		bool bIsDynamic,bool bInLightMapped,bool bInIsShadowMapped, bool bInHasTranslucentObjectShadow, bool bInHasInsetObjectShadow);
diff --git a/Engine/Source/Runtime/Renderer/Private/TranslucentRendering.cpp b/Engine/Source/Runtime/Renderer/Private/TranslucentRendering.cpp
--- a/Engine/Source/Runtime/Renderer/Private/TranslucentRendering.cpp
+++ b/Engine/Source/Runtime/Renderer/Private/TranslucentRendering.cpp
@@ -219,9 +219,9 @@ class FDrawTranslucentMeshAction
 public:
 
 	const FViewInfo& View;
-	bool bBackFace;
-	FHitProxyId HitProxyId;
 	const FProjectedShadowInfo* TranslucentSelfShadow;
+	FHitProxyId HitProxyId;
+	bool bBackFace;
 	bool bUseTranslucentSelfShadowing;
 
 	/** Initialization constructor. */
@@ -233,9 +233,9 @@ public:
 		bool bInUseTranslucentSelfShadowing
 		):
 		View(InView),
-		bBackFace(bInBackFace),
-		HitProxyId(InHitProxyId),
 		TranslucentSelfShadow(InTranslucentSelfShadow),
+		HitProxyId(InHitProxyId),
+		bBackFace(bInBackFace),
 		bUseTranslucentSelfShadowing(bInUseTranslucentSelfShadowing)
 	{}
 
