PVS-Studio detected that this value was not being cleared because the assignment occurred after a return statement.

diff --git a/Engine/Source/Runtime/RHI/Public/RHICommandList.h b/Engine/Source/Runtime/RHI/Public/RHICommandList.h
--- a/Engine/Source/Runtime/RHI/Public/RHICommandList.h
+++ b/Engine/Source/Runtime/RHI/Public/RHICommandList.h
@@ -1934,9 +1934,8 @@ public:
 		{ \
 			QUICK_SCOPE_CYCLE_COUNTER(STAT_RHIMETHOD_##Name##_Flush); \
 			ImmediateFlush(EImmediateFlushType::FlushRHIThread); \
-			bFlushedGlobal = true;\
+			TGuardValue<bool> GuardIsFlushedGlobal( bFlushedGlobal, true ); \
 			ReturnStatement Name##_Internal ParameterNames; \
-			bFlushedGlobal = false;\
 		}
 	#include "RHIMethods.h"
 	#undef DEFINE_RHIMETHOD
