UE4 - Mobile Preview for feature level ES3.1/Metal (experimental)

diff --git a/Engine/Shaders/Common.usf b/Engine/Shaders/Common.usf
--- a/Engine/Shaders/Common.usf
+++ b/Engine/Shaders/Common.usf
@@ -130,8 +130,12 @@
 	#define SM4_PROFILE 1
 	#define FCOLOR_COMPONENT_SWIZZLE .bgra  
 #elif ES2_PROFILE || ES3_1_PROFILE
-	//@todo ES3_1 GL
-	#define FEATURE_LEVEL FEATURE_LEVEL_ES2
+	#if ES3_1_PROFILE
+		#define FEATURE_LEVEL FEATURE_LEVEL_ES3_1
+	#else
+		//@todo ES3_1 GL
+		#define FEATURE_LEVEL FEATURE_LEVEL_ES2
+	#endif
 
 	#if COMPILER_GLSL_ES2 
 		// Swizzle as we only support GL_BGRA on non-ES2 platforms that have that extension
@@ -1069,7 +1073,7 @@ struct FWriteToSliceGeometryOutput
 
 // Helper macro to globally ignore requests for non-offset world positions in materials when lower than shader model 4. We do this
 // because we are using an extra interpolator for this second world position, and in < SM4 there may not be enough
-#define USE_WORLD_POSITION_EXCLUDING_SHADER_OFFSETS	(NEEDS_WORLD_POSITION_EXCLUDING_SHADER_OFFSETS && FEATURE_LEVEL >= FEATURE_LEVEL_SM4)
+#define USE_WORLD_POSITION_EXCLUDING_SHADER_OFFSETS	(NEEDS_WORLD_POSITION_EXCLUDING_SHADER_OFFSETS && FEATURE_LEVEL >= FEATURE_LEVEL_ES3_1)
 
 /** Used for calculating vertex positions and UVs when drawing with DrawRectangle */
 void DrawRectangle( in float4 InPosition, in float2 InTexCoord, out float4 OutPosition, out float2 OutTexCoord)
diff --git a/Engine/Source/Developer/ShaderFormatOpenGL/Private/OpenGLShaderCompiler.cpp b/Engine/Source/Developer/ShaderFormatOpenGL/Private/OpenGLShaderCompiler.cpp
--- a/Engine/Source/Developer/ShaderFormatOpenGL/Private/OpenGLShaderCompiler.cpp
+++ b/Engine/Source/Developer/ShaderFormatOpenGL/Private/OpenGLShaderCompiler.cpp
@@ -69,7 +69,7 @@ static FORCEINLINE bool IsES2Platform(GLSLVersion Version)
 
 static FORCEINLINE bool IsPCES2Platform(GLSLVersion Version)
 {
-	return (Version == GLSL_150_ES2 || Version == GLSL_150_ES2_NOUB);
+	return (Version == GLSL_150_ES2 || Version == GLSL_150_ES2_NOUB || Version == GLSL_150_ES3_1);
 }
 
 // This function should match OpenGLShaderPlatformSeparable
@@ -78,7 +78,7 @@ static FORCEINLINE bool SupportsSeparateShaderObjects(GLSLVersion Version)
 	// Only desktop shader platforms can use separable shaders for now,
 	// the generated code relies on macros supplied at runtime to determine whether
 	// shaders may be separable and/or linked.
-	return Version == GLSL_150 || Version == GLSL_150_MAC || Version == GLSL_150_ES2 || Version == GLSL_150_ES2_NOUB|| Version == GLSL_430;
+	return Version == GLSL_150 || Version == GLSL_150_MAC || Version == GLSL_150_ES2 || Version == GLSL_150_ES2_NOUB || Version == GLSL_150_ES3_1 || Version == GLSL_430;
 }
 
 /*------------------------------------------------------------------------------
@@ -1190,6 +1190,7 @@ static void OpenGLVersionFromGLSLVersion(GLSLVersion InVersion, int& OutMajorVer
 			break;
 		case GLSL_150_ES2:
 		case GLSL_150_ES2_NOUB:
+		case GLSL_150_ES3_1:
 			OutMajorVersion = 3;
 			OutMinorVersion = 2;
 			break;
@@ -1519,6 +1520,7 @@ static FString CreateCrossCompilerBatchFile( const FString& ShaderFile, const FS
 	{
 		case GLSL_150:
 		case GLSL_150_ES2:
+		case GLSL_150_ES3_1:
 			VersionSwitch = TEXT(" -gl3 -separateshaders");
 			break;
 
@@ -1624,6 +1626,13 @@ void CompileShader_Windows_OGL(const FShaderCompilerInput& Input,FShaderCompiler
 			AdditionalDefines.SetDefine(TEXT("row_major"), TEXT(""));
 			break;
 
+		case GLSL_150_ES3_1:
+			AdditionalDefines.SetDefine(TEXT("COMPILER_GLSL"), 1);
+			AdditionalDefines.SetDefine(TEXT("ES3_1_PROFILE"), 1);
+			HlslCompilerTarget = HCT_FeatureLevelSM4;
+			AdditionalDefines.SetDefine(TEXT("row_major"), TEXT(""));
+			break;
+
 		default:
 			check(0);
 	}
diff --git a/Engine/Source/Developer/ShaderFormatOpenGL/Private/ShaderFormatOpenGL.cpp b/Engine/Source/Developer/ShaderFormatOpenGL/Private/ShaderFormatOpenGL.cpp
--- a/Engine/Source/Developer/ShaderFormatOpenGL/Private/ShaderFormatOpenGL.cpp
+++ b/Engine/Source/Developer/ShaderFormatOpenGL/Private/ShaderFormatOpenGL.cpp
@@ -15,6 +15,7 @@ static FName NAME_GLSL_ES2(TEXT("GLSL_ES2"));
 static FName NAME_GLSL_ES2_WEBGL(TEXT("GLSL_ES2_WEBGL"));
 static FName NAME_GLSL_150_ES2(TEXT("GLSL_150_ES2"));
 static FName NAME_GLSL_150_ES2_NOUB(TEXT("GLSL_150_ES2_NOUB"));
+static FName NAME_GLSL_150_ES3_1(TEXT("GLSL_150_ES3_1"));
 static FName NAME_GLSL_ES2_IOS(TEXT("GLSL_ES2_IOS"));
 static FName NAME_GLSL_310_ES_EXT(TEXT("GLSL_310_ES_EXT"));
  
@@ -29,6 +30,7 @@ class FShaderFormatGLSL : public IShaderFormat
 		UE_SHADER_GLSL_ES2_VER = 14,
 		UE_SHADER_GLSL_150ES2_VER = 17,
 		UE_SHADER_GLSL_150ES2NOUB_VER = 1,
+		UE_SHADER_GLSL_150ES3_1_VER = 1,
 		UE_SHADER_GLSL_ES2_VER_WEBGL  = 15,
 		UE_SHADER_GLSL_ES2_IOS_VER  = 5,
 		UE_SHADER_GLSL_310_ES_EXT_VER = 3,
@@ -42,6 +44,7 @@ class FShaderFormatGLSL : public IShaderFormat
 				Format == NAME_GLSL_ES2 || 
 				Format == NAME_GLSL_150_ES2 ||
 				Format == NAME_GLSL_150_ES2_NOUB ||
+				Format == NAME_GLSL_150_ES3_1 ||
                 Format == NAME_GLSL_ES2_WEBGL ||
 				Format == NAME_GLSL_ES2_IOS ||
 				Format == NAME_GLSL_310_ES_EXT
@@ -73,6 +76,10 @@ public:
 		{
 			GLSLVersion = UE_SHADER_GLSL_150ES2_VER;
 		}
+		else if (Format == NAME_GLSL_150_ES3_1)
+		{
+			GLSLVersion = UE_SHADER_GLSL_150ES3_1_VER;
+		}
 		else if (Format == NAME_GLSL_150_ES2_NOUB)
 		{
 			GLSLVersion = UE_SHADER_GLSL_150ES2NOUB_VER;
@@ -105,6 +112,7 @@ public:
 		OutFormats.Add(NAME_GLSL_ES2);
 		OutFormats.Add(NAME_GLSL_ES2_WEBGL);
 		OutFormats.Add(NAME_GLSL_150_ES2);
+		OutFormats.Add(NAME_GLSL_150_ES3_1);
 		OutFormats.Add(NAME_GLSL_ES2_IOS);
 		OutFormats.Add(NAME_GLSL_310_ES_EXT);
 		OutFormats.Add(NAME_GLSL_150_ES2_NOUB);
@@ -149,7 +157,11 @@ public:
 		}
 		else if (Format == NAME_GLSL_ES2_IOS )
 		{
-			CompileShader_Windows_OGL(Input, Output, WorkingDirectory, GLSL_ES2_IOS);
+			CompileShader_Windows_OGL(Input, Output, WorkingDirectory, GLSL_150_ES3_1);
+		}
+		else if (Format == NAME_GLSL_150_ES3_1)
+		{
+			CompileShader_Windows_OGL(Input, Output, WorkingDirectory, GLSL_150_ES3_1);
 		}
 		else if (Format == NAME_GLSL_150_ES2_NOUB)
 		{
diff --git a/Engine/Source/Developer/ShaderFormatOpenGL/Private/ShaderFormatOpenGL.h b/Engine/Source/Developer/ShaderFormatOpenGL/Private/ShaderFormatOpenGL.h
--- a/Engine/Source/Developer/ShaderFormatOpenGL/Private/ShaderFormatOpenGL.h
+++ b/Engine/Source/Developer/ShaderFormatOpenGL/Private/ShaderFormatOpenGL.h
@@ -10,6 +10,7 @@ enum GLSLVersion
 	GLSL_ES2_WEBGL,
 	GLSL_150_ES2,	// ES2 Emulation
 	GLSL_150_ES2_NOUB,	// ES2 Emulation with NoUBs
+	GLSL_150_ES3_1,	// ES3.1 Emulation
 	GLSL_ES2_IOS,
 	GLSL_150_MAC, // Apple only
 	GLSL_310_ES_EXT,
diff --git a/Engine/Source/Developer/Windows/ShaderFormatD3D/Private/D3D11ShaderCompiler.cpp b/Engine/Source/Developer/Windows/ShaderFormatD3D/Private/D3D11ShaderCompiler.cpp
--- a/Engine/Source/Developer/Windows/ShaderFormatD3D/Private/D3D11ShaderCompiler.cpp
+++ b/Engine/Source/Developer/Windows/ShaderFormatD3D/Private/D3D11ShaderCompiler.cpp
@@ -95,7 +95,7 @@ static const TCHAR* GetShaderProfileName(FShaderTarget Target)
 			return TEXT("cs_5_0");
 		}
 	}
-	else if(Target.Platform == SP_PCD3D_SM4 || Target.Platform == SP_PCD3D_ES2)
+	else if(Target.Platform == SP_PCD3D_SM4 || Target.Platform == SP_PCD3D_ES2 || Target.Platform == SP_PCD3D_ES3_1)
 	{
 		checkSlow(Target.Frequency == SF_Vertex ||
 			Target.Frequency == SF_Pixel ||
@@ -783,4 +783,13 @@ void CompileShader_Windows_ES2(const FShaderCompilerInput& Input,FShaderCompiler
 	FShaderCompilerDefinitions AdditionalDefines;
 	AdditionalDefines.SetDefine(TEXT("ES2_PROFILE"), 1);
 	CompileD3D11Shader(Input, Output, AdditionalDefines, WorkingDirectory);
-}
\ No newline at end of file
+}
+
+void CompileShader_Windows_ES3_1(const FShaderCompilerInput& Input, FShaderCompilerOutput& Output, const FString& WorkingDirectory)
+{
+	check(Input.Target.Platform == SP_PCD3D_ES3_1);
+
+	FShaderCompilerDefinitions AdditionalDefines;
+	AdditionalDefines.SetDefine(TEXT("ES3_1_PROFILE"), 1);
+	CompileD3D11Shader(Input, Output, AdditionalDefines, WorkingDirectory);
+}
diff --git a/Engine/Source/Developer/Windows/ShaderFormatD3D/Private/ShaderFormatD3D.cpp b/Engine/Source/Developer/Windows/ShaderFormatD3D/Private/ShaderFormatD3D.cpp
--- a/Engine/Source/Developer/Windows/ShaderFormatD3D/Private/ShaderFormatD3D.cpp
+++ b/Engine/Source/Developer/Windows/ShaderFormatD3D/Private/ShaderFormatD3D.cpp
@@ -8,6 +8,7 @@
 
 static FName NAME_PCD3D_SM5(TEXT("PCD3D_SM5"));
 static FName NAME_PCD3D_SM4(TEXT("PCD3D_SM4"));
+static FName NAME_PCD3D_ES3_1(TEXT("PCD3D_ES3_1"));
 static FName NAME_PCD3D_ES2(TEXT("PCD3D_ES2"));
 
 class FShaderFormatD3D : public IShaderFormat
@@ -17,12 +18,13 @@ class FShaderFormatD3D : public IShaderFormat
 		/** Version for shader format, this becomes part of the DDC key. */
 		UE_SHADER_PCD3D_SM5_VER = 2,
 		UE_SHADER_PCD3D_SM4_VER = 2,
-		UE_SHADER_PCD3D_ES2_VER = 2
+		UE_SHADER_PCD3D_ES2_VER = 2,
+		UE_SHADER_PCD3D_ES3_1_VER = 1,
 	};
 
 	void CheckFormat(FName Format) const
 	{
-		check(Format == NAME_PCD3D_SM5 ||  Format == NAME_PCD3D_SM4 ||  Format == NAME_PCD3D_ES2);
+		check(Format == NAME_PCD3D_SM5 ||  Format == NAME_PCD3D_SM4 ||  Format == NAME_PCD3D_ES2 || Format == NAME_PCD3D_ES3_1);
 	}
 
 public:
@@ -37,7 +39,11 @@ public:
 		{
 			return UE_SHADER_PCD3D_SM4_VER;
 		}
-		else if (Format == NAME_PCD3D_ES2) 
+		else if (Format == NAME_PCD3D_ES3_1) 
+		{
+			return UE_SHADER_PCD3D_ES3_1_VER;
+		}
+		else if (Format == NAME_PCD3D_ES2)
 		{
 			return UE_SHADER_PCD3D_ES2_VER;
 		}
@@ -48,6 +54,7 @@ public:
 	{
 		OutFormats.Add(NAME_PCD3D_SM5);
 		OutFormats.Add(NAME_PCD3D_SM4);
+		OutFormats.Add(NAME_PCD3D_ES3_1);
 		OutFormats.Add(NAME_PCD3D_ES2);
 	}
 	virtual void CompileShader(FName Format, const struct FShaderCompilerInput& Input, struct FShaderCompilerOutput& Output,const FString& WorkingDirectory) const
@@ -65,6 +72,10 @@ public:
 		{
 			CompileShader_Windows_ES2(Input, Output, WorkingDirectory);
 		}
+		else if (Format == NAME_PCD3D_ES3_1)
+		{
+			CompileShader_Windows_ES3_1(Input, Output, WorkingDirectory);
+		}
 		else
 		{
 			check(0);
diff --git a/Engine/Source/Developer/Windows/ShaderFormatD3D/Private/ShaderFormatD3D.h b/Engine/Source/Developer/Windows/ShaderFormatD3D/Private/ShaderFormatD3D.h
--- a/Engine/Source/Developer/Windows/ShaderFormatD3D/Private/ShaderFormatD3D.h
+++ b/Engine/Source/Developer/Windows/ShaderFormatD3D/Private/ShaderFormatD3D.h
@@ -4,4 +4,5 @@
 
 void CompileShader_Windows_SM5(const struct FShaderCompilerInput& Input,struct FShaderCompilerOutput& Output,const class FString& WorkingDirectory);
 void CompileShader_Windows_SM4(const struct FShaderCompilerInput& Input,struct FShaderCompilerOutput& Output,const class FString& WorkingDirectory);
-void CompileShader_Windows_ES2(const struct FShaderCompilerInput& Input,struct FShaderCompilerOutput& Output,const class FString& WorkingDirectory);
\ No newline at end of file
+void CompileShader_Windows_ES2(const struct FShaderCompilerInput& Input,struct FShaderCompilerOutput& Output,const class FString& WorkingDirectory);
+void CompileShader_Windows_ES3_1(const struct FShaderCompilerInput& Input,struct FShaderCompilerOutput& Output,const class FString& WorkingDirectory);
diff --git a/Engine/Source/Editor/LevelEditor/Private/LevelEditorActions.cpp b/Engine/Source/Editor/LevelEditor/Private/LevelEditorActions.cpp
--- a/Engine/Source/Editor/LevelEditor/Private/LevelEditorActions.cpp
+++ b/Engine/Source/Editor/LevelEditor/Private/LevelEditorActions.cpp
@@ -3077,7 +3077,7 @@ void FLevelEditorCommands::RegisterCommands()
 	static const FText FeatureLevelLabels[ERHIFeatureLevel::Num] = 
 	{
 		NSLOCTEXT("LevelEditorCommands", "FeatureLevelPreviewType_ES2", "Mobile / HTML5"),
-		NSLOCTEXT("LevelEditorCommands", "FeatureLevelPreviewType_ES3", "High-End Mobile"),
+		NSLOCTEXT("LevelEditorCommands", "FeatureLevelPreviewType_ES3_1", "High-End Mobile / Metal"),
 		NSLOCTEXT("LevelEditorCommands", "FeatureLevelPreviewType_SM4", "Shader Model 4"),
 		NSLOCTEXT("LevelEditorCommands", "FeatureLevelPreviewType_SM5", "Shader Model 5"),
 	};
@@ -3085,7 +3085,7 @@ void FLevelEditorCommands::RegisterCommands()
 	static const FText FeatureLevelToolTips[ERHIFeatureLevel::Num] = 
 	{
 		NSLOCTEXT("LevelEditorCommands", "FeatureLevelPreviewTooltip_ES2", "OpenGLES 2"),
-		NSLOCTEXT("LevelEditorCommands", "FeatureLevelPreviewTooltip_ES3", "OpenGLES 3"),
+		NSLOCTEXT("LevelEditorCommands", "FeatureLevelPreviewTooltip_ES3", "OpenGLES 3.1, Metal"),
 		NSLOCTEXT("LevelEditorCommands", "FeatureLevelPreviewTooltip_SM4", "DirectX 10, OpenGL 3.3+"),
 		NSLOCTEXT("LevelEditorCommands", "FeatureLevelPreviewTooltip_SM5", "DirectX 11, OpenGL 4.3+, PS4, XB1"),
 	};
diff --git a/Engine/Source/Editor/LevelEditor/Private/LevelEditorToolBar.cpp b/Engine/Source/Editor/LevelEditor/Private/LevelEditorToolBar.cpp
--- a/Engine/Source/Editor/LevelEditor/Private/LevelEditorToolBar.cpp
+++ b/Engine/Source/Editor/LevelEditor/Private/LevelEditorToolBar.cpp
@@ -1610,7 +1610,7 @@ static void MakeShaderModelPreviewMenu(FMenuBuilder& MenuBuilder)
 	{
 		for (int32 i = GMaxRHIFeatureLevel; i >= 0; --i)
 		{
-			if (i != ERHIFeatureLevel::ES3_1)
+			if (i != ERHIFeatureLevel::ES3_1 || GetDefault<UEditorExperimentalSettings>()->bFeatureLevelES31Preview)
 			{
 				MenuBuilder.AddMenuEntry(FLevelEditorCommands::Get().FeatureLevelPreview[i]);
 			}
diff --git a/Engine/Source/Editor/UnrealEd/Classes/Settings/EditorExperimentalSettings.h b/Engine/Source/Editor/UnrealEd/Classes/Settings/EditorExperimentalSettings.h
--- a/Engine/Source/Editor/UnrealEd/Classes/Settings/EditorExperimentalSettings.h
+++ b/Engine/Source/Editor/UnrealEd/Classes/Settings/EditorExperimentalSettings.h
@@ -93,6 +93,10 @@ public:
 	UPROPERTY(EditAnywhere, config, Category = AI, meta = (DisplayName = "Environment Querying System"))
 	bool bEQSEditor;
 
+	/** Enables Environment Queries editor */
+	UPROPERTY(EditAnywhere, config, Category = Rendering, meta = (DisplayName = "Enable Metal/High-end mobile rendering preview"))
+	bool bFeatureLevelES31Preview;
+
 	/**
 	 * Returns an event delegate that is executed when a setting has changed.
 	 *
diff --git a/Engine/Source/Runtime/NullDrv/Private/NullRHI.cpp b/Engine/Source/Runtime/NullDrv/Private/NullRHI.cpp
--- a/Engine/Source/Runtime/NullDrv/Private/NullRHI.cpp
+++ b/Engine/Source/Runtime/NullDrv/Private/NullRHI.cpp
@@ -19,7 +19,7 @@ void FNullDynamicRHI::Init()
 {
 #if PLATFORM_WINDOWS
 	GShaderPlatformForFeatureLevel[ERHIFeatureLevel::ES2] = SP_PCD3D_ES2;
-	GShaderPlatformForFeatureLevel[ERHIFeatureLevel::ES3_1] = SP_NumPlatforms;
+	GShaderPlatformForFeatureLevel[ERHIFeatureLevel::ES3_1] = SP_PCD3D_ES3_1;
 	GShaderPlatformForFeatureLevel[ERHIFeatureLevel::SM4] = SP_PCD3D_SM4;
 	GShaderPlatformForFeatureLevel[ERHIFeatureLevel::SM5] = SP_PCD3D_SM5;
 #else
diff --git a/Engine/Source/Runtime/OpenGLDrv/Private/OpenGLDevice.cpp b/Engine/Source/Runtime/OpenGLDrv/Private/OpenGLDevice.cpp
--- a/Engine/Source/Runtime/OpenGLDrv/Private/OpenGLDevice.cpp
+++ b/Engine/Source/Runtime/OpenGLDrv/Private/OpenGLDevice.cpp
@@ -700,7 +700,7 @@ static void InitRHICapabilitiesForGL()
 	GHardwareHiddenSurfaceRemoval = FOpenGL::HasHardwareHiddenSurfaceRemoval();
 
 	GShaderPlatformForFeatureLevel[ERHIFeatureLevel::ES2] = (GMaxRHIFeatureLevel == ERHIFeatureLevel::ES2) ? GMaxRHIShaderPlatform : SP_OPENGL_PCES2;
-	GShaderPlatformForFeatureLevel[ERHIFeatureLevel::ES3_1] = SP_NumPlatforms;
+	GShaderPlatformForFeatureLevel[ERHIFeatureLevel::ES3_1] = (GMaxRHIFeatureLevel == ERHIFeatureLevel::ES3_1) ? GMaxRHIShaderPlatform : SP_OPENGL_PCES3_1;
 	GShaderPlatformForFeatureLevel[ERHIFeatureLevel::SM4] = PLATFORM_MAC ? SP_OPENGL_SM4_MAC : SP_OPENGL_SM4;
 	GShaderPlatformForFeatureLevel[ERHIFeatureLevel::SM5] = OPENGL_ES31 ? SP_OPENGL_ES31_EXT : SP_OPENGL_SM5;
 
diff --git a/Engine/Source/Runtime/OpenGLDrv/Public/OpenGL3.h b/Engine/Source/Runtime/OpenGLDrv/Public/OpenGL3.h
--- a/Engine/Source/Runtime/OpenGLDrv/Public/OpenGL3.h
+++ b/Engine/Source/Runtime/OpenGLDrv/Public/OpenGL3.h
@@ -719,10 +719,16 @@ struct FOpenGL3 : public FOpenGLBase
 
 	static FORCEINLINE ERHIFeatureLevel::Type GetFeatureLevel()
 	{
-		if (FParse::Param(FCommandLine::Get(), TEXT("FeatureLevelES2")) && !GIsEditor)
+		static bool bForceES2 = FParse::Param(FCommandLine::Get(), TEXT("FeatureLevelES2"));
+		static bool bForceES3_1 = FParse::Param(FCommandLine::Get(), TEXT("FeatureLevelES31")) || FParse::Param(FCommandLine::Get(), TEXT("FeatureLevelES3_1"));
+		if (bForceES2 && !GIsEditor)
 		{
 			return ERHIFeatureLevel::ES2;
 		}
+		else if (bForceES3_1 && !GIsEditor)
+		{
+			return ERHIFeatureLevel::ES3_1;
+		}
 		// Shader platform & RHI feature level
 		switch(GetMajorVersion())
 		{
@@ -745,6 +751,12 @@ struct FOpenGL3 : public FOpenGLBase
 			return SP_OPENGL_PCES2;
 		}
 
+		static bool bForceFeatureLevelES3_1 = (FParse::Param(FCommandLine::Get(), TEXT("FeatureLevelES31")) || FParse::Param(FCommandLine::Get(), TEXT("FeatureLevelES3_1"))) && !GIsEditor;
+		if (bForceFeatureLevelES3_1)
+		{
+			return SP_OPENGL_PCES3_1;
+		}
+
 		// Shader platform
 		switch(GetMajorVersion())
 		{
diff --git a/Engine/Source/Runtime/RHI/Private/RHI.cpp b/Engine/Source/Runtime/RHI/Private/RHI.cpp
--- a/Engine/Source/Runtime/RHI/Private/RHI.cpp
+++ b/Engine/Source/Runtime/RHI/Private/RHI.cpp
@@ -270,6 +270,7 @@ RHI_API void GetFeatureLevelName(ERHIFeatureLevel::Type InFeatureLevel, FName& O
 
 static FName NAME_PCD3D_SM5(TEXT("PCD3D_SM5"));
 static FName NAME_PCD3D_SM4(TEXT("PCD3D_SM4"));
+static FName NAME_PCD3D_ES3_1(TEXT("PCD3D_ES3_1"));
 static FName NAME_PCD3D_ES2(TEXT("PCD3D_ES2"));
 static FName NAME_GLSL_150(TEXT("GLSL_150"));
 static FName NAME_GLSL_150_MAC(TEXT("GLSL_150_MAC"));
@@ -278,6 +279,7 @@ static FName NAME_SF_XBOXONE(TEXT("SF_XBOXONE"));
 static FName NAME_GLSL_430(TEXT("GLSL_430"));
 static FName NAME_OPENGL_150_ES2(TEXT("GLSL_150_ES2"));
 static FName NAME_OPENGL_150_ES2_NOUB(TEXT("GLSL_150_ES2_NOUB"));
+static FName NAME_OPENGL_150_ES3_1(TEXT("GLSL_150_ES3_1"));
 static FName NAME_OPENGL_ES2(TEXT("GLSL_ES2"));
 static FName NAME_OPENGL_ES2_WEBGL(TEXT("GLSL_ES2_WEBGL"));
 static FName NAME_OPENGL_ES2_IOS(TEXT("GLSL_ES2_IOS"));
@@ -293,6 +295,8 @@ FName LegacyShaderPlatformToShaderFormat(EShaderPlatform Platform)
 		return NAME_PCD3D_SM5;
 	case SP_PCD3D_SM4:
 		return NAME_PCD3D_SM4;
+	case SP_PCD3D_ES3_1:
+		return NAME_PCD3D_ES3_1;
 	case SP_PCD3D_ES2:
 		return NAME_PCD3D_ES2;
 	case SP_OPENGL_SM4:
@@ -310,6 +314,8 @@ FName LegacyShaderPlatformToShaderFormat(EShaderPlatform Platform)
 		static auto* CVar = IConsoleManager::Get().FindTConsoleVariableDataInt(TEXT("OpenGL.UseEmulatedUBs"));
 		return (CVar && CVar->GetValueOnAnyThread() != 0) ? NAME_OPENGL_150_ES2_NOUB : NAME_OPENGL_150_ES2;
 	}
+	case SP_OPENGL_PCES3_1:
+		return NAME_OPENGL_150_ES3_1;
 	case SP_OPENGL_ES2:
 		return NAME_OPENGL_ES2;
 	case SP_OPENGL_ES2_WEBGL:
@@ -333,6 +339,7 @@ EShaderPlatform ShaderFormatToLegacyShaderPlatform(FName ShaderFormat)
 {
 	if (ShaderFormat == NAME_PCD3D_SM5)			return SP_PCD3D_SM5;	
 	if (ShaderFormat == NAME_PCD3D_SM4)			return SP_PCD3D_SM4;
+	if (ShaderFormat == NAME_PCD3D_ES3_1)		return SP_PCD3D_ES3_1;
 	if (ShaderFormat == NAME_PCD3D_ES2)			return SP_PCD3D_ES2;
 	if (ShaderFormat == NAME_GLSL_150)			return SP_OPENGL_SM4;
 	if (ShaderFormat == NAME_GLSL_150_MAC)		return SP_OPENGL_SM4_MAC;
@@ -341,6 +348,7 @@ EShaderPlatform ShaderFormatToLegacyShaderPlatform(FName ShaderFormat)
 	if (ShaderFormat == NAME_GLSL_430)			return SP_OPENGL_SM5;
 	if (ShaderFormat == NAME_OPENGL_150_ES2 || ShaderFormat == NAME_OPENGL_150_ES2_NOUB)
 												return SP_OPENGL_PCES2;
+	if (ShaderFormat == NAME_OPENGL_150_ES3_1)	return SP_OPENGL_PCES3_1;
 	if (ShaderFormat == NAME_OPENGL_ES2)			return SP_OPENGL_ES2;
 	if (ShaderFormat == NAME_OPENGL_ES2_WEBGL)	return SP_OPENGL_ES2_WEBGL;
 	if (ShaderFormat == NAME_OPENGL_ES2_IOS)		return SP_OPENGL_ES2_IOS;
diff --git a/Engine/Source/Runtime/RHI/Public/RHIDefinitions.h b/Engine/Source/Runtime/RHI/Public/RHIDefinitions.h
--- a/Engine/Source/Runtime/RHI/Public/RHIDefinitions.h
+++ b/Engine/Source/Runtime/RHI/Public/RHIDefinitions.h
@@ -15,7 +15,7 @@ enum EShaderFrequency
 	SF_Vertex			= 0,
 	SF_Hull				= 1,
 	SF_Domain			= 2,
-	SF_Pixel			= 3,
+	SF_Pixel				= 3,
 	SF_Geometry			= 4,
 	SF_Compute			= 5,
 
@@ -44,8 +44,12 @@ enum EShaderPlatform
 	SP_OPENGL_SM4_MAC	= 12,
 	SP_METAL_MRT			= 13,
 	SP_OPENGL_ES31_EXT	= 14,
+	/** Used when running in Feature Level ES3_1 in D3D11. */
+	SP_PCD3D_ES3_1		= 15,
+	/** Used when running in Feature Level ES3_1 in OpenGL. */
+	SP_OPENGL_PCES3_1	= 16,
 
-	SP_NumPlatforms		= 15,
+	SP_NumPlatforms		= 17,
 	SP_NumBits			= 5,
 };
 static_assert(SP_NumPlatforms <= (1 << SP_NumBits), "SP_NumPlatforms will not fit on SP_NumBits");
@@ -580,7 +584,7 @@ enum class ESimpleRenderTargetMode
 
 inline bool IsPCPlatform(const EShaderPlatform Platform)
 {
-	return Platform == SP_PCD3D_SM5 || Platform == SP_PCD3D_SM4 || Platform == SP_PCD3D_ES2 || Platform ==  SP_OPENGL_SM4 || Platform == SP_OPENGL_SM4_MAC || Platform == SP_OPENGL_SM5 || Platform == SP_OPENGL_PCES2;
+	return Platform == SP_PCD3D_SM5 || Platform == SP_PCD3D_SM4 || Platform == SP_PCD3D_ES2 || Platform == SP_PCD3D_ES3_1 || Platform ==  SP_OPENGL_SM4 || Platform == SP_OPENGL_SM4_MAC || Platform == SP_OPENGL_SM5 || Platform == SP_OPENGL_PCES2 || Platform == SP_OPENGL_PCES3_1;
 }
 
 /** Whether the shader platform corresponds to the ES2 feature level. */
@@ -592,7 +596,7 @@ inline bool IsES2Platform(const EShaderPlatform Platform)
 /** Whether the shader platform corresponds to the ES2 feature level. */
 inline bool IsMobilePlatform(const EShaderPlatform Platform)
 {
-	return IsES2Platform(Platform) || Platform == SP_METAL;
+	return IsES2Platform(Platform) || Platform == SP_METAL || Platform == SP_PCD3D_ES3_1 || Platform == SP_OPENGL_PCES3_1;
 }
 
 inline bool IsOpenGLPlatform(const EShaderPlatform Platform)
@@ -628,6 +632,8 @@ inline ERHIFeatureLevel::Type GetMaxSupportedFeatureLevel(EShaderPlatform InShad
 	case SP_OPENGL_ES2_IOS:
 		return ERHIFeatureLevel::ES2;
 	case SP_METAL:
+	case SP_PCD3D_ES3_1:
+	case SP_OPENGL_PCES3_1:
 		return ERHIFeatureLevel::ES3_1;
 	default:
 		check(0);
@@ -646,8 +652,12 @@ inline bool IsFeatureLevelSupported(EShaderPlatform InShaderPlatform, ERHIFeatur
 		return InFeatureLevel <= ERHIFeatureLevel::SM4;
 	case SP_PCD3D_ES2:
 		return InFeatureLevel <= ERHIFeatureLevel::ES2;
+	case SP_PCD3D_ES3_1:
+		return InFeatureLevel <= ERHIFeatureLevel::ES3_1;
 	case SP_OPENGL_PCES2:
 		return InFeatureLevel <= ERHIFeatureLevel::ES2;
+	case SP_OPENGL_PCES3_1:
+		return InFeatureLevel <= ERHIFeatureLevel::ES3_1;
 	case SP_OPENGL_ES2:
 		return InFeatureLevel <= ERHIFeatureLevel::ES2;
 	case SP_OPENGL_ES2_WEBGL: 
diff --git a/Engine/Source/Runtime/Renderer/Private/ShadowRendering.cpp b/Engine/Source/Runtime/Renderer/Private/ShadowRendering.cpp
--- a/Engine/Source/Runtime/Renderer/Private/ShadowRendering.cpp
+++ b/Engine/Source/Runtime/Renderer/Private/ShadowRendering.cpp
@@ -1018,17 +1018,18 @@ void FShadowDepthDrawingPolicyFactory::AddStaticMesh(FScene* Scene,FStaticMesh*
 						true,
 						false,
 						false,
-						Scene->GetFeatureLevel(),
+						FeatureLevel,
 						StaticMesh->VertexFactory,
 						MaterialRenderProxy,
 						bTwoSided,
 						StaticMesh->ReverseCulling),
-					Scene->GetFeatureLevel()
+					FeatureLevel
 					);
 			}
 		}
 		if ( bLitOpaque )
 		{
+			const auto FeatureLevel = Scene->GetFeatureLevel();
 			OverrideWithDefaultMaterialForShadowDepth(MaterialRenderProxy, Material, false, Scene->GetFeatureLevel()); 
 
 			// Add the static mesh to the shadow's subject draw list.
@@ -1040,12 +1041,12 @@ void FShadowDepthDrawingPolicyFactory::AddStaticMesh(FScene* Scene,FStaticMesh*
 					true,
 					false,
 					false,
-					Scene->GetFeatureLevel(),
+					FeatureLevel,
 					StaticMesh->VertexFactory,
 					MaterialRenderProxy,
 					bTwoSided,
 					StaticMesh->ReverseCulling),
-				Scene->GetFeatureLevel()
+				FeatureLevel
 				);
 		}
 	}
diff --git a/Engine/Source/Runtime/Windows/D3D11RHI/Private/D3D11Device.cpp b/Engine/Source/Runtime/Windows/D3D11RHI/Private/D3D11Device.cpp
--- a/Engine/Source/Runtime/Windows/D3D11RHI/Private/D3D11Device.cpp
+++ b/Engine/Source/Runtime/Windows/D3D11RHI/Private/D3D11Device.cpp
@@ -85,6 +85,11 @@ FD3D11DynamicRHI::FD3D11DynamicRHI(IDXGIFactory1* InDXGIFactory1,D3D_FEATURE_LEV
 		GMaxRHIFeatureLevel = ERHIFeatureLevel::ES2;
 		GMaxRHIShaderPlatform = SP_PCD3D_ES2;
 	}
+	else if ((FParse::Param(FCommandLine::Get(), TEXT("FeatureLevelES31")) || FParse::Param(FCommandLine::Get(), TEXT("FeatureLevelES3_1"))) && !GIsEditor)
+	{
+		GMaxRHIFeatureLevel = ERHIFeatureLevel::ES3_1;
+		GMaxRHIShaderPlatform = SP_PCD3D_ES3_1;
+	}
 	else if(FeatureLevel == D3D_FEATURE_LEVEL_11_0)
 	{
 		GMaxRHIFeatureLevel = ERHIFeatureLevel::SM5;
diff --git a/Engine/Source/Runtime/Windows/D3D11RHI/Private/Windows/WindowsD3D11Device.cpp b/Engine/Source/Runtime/Windows/D3D11RHI/Private/Windows/WindowsD3D11Device.cpp
--- a/Engine/Source/Runtime/Windows/D3D11RHI/Private/Windows/WindowsD3D11Device.cpp
+++ b/Engine/Source/Runtime/Windows/D3D11RHI/Private/Windows/WindowsD3D11Device.cpp
@@ -567,7 +567,7 @@ void FD3D11DynamicRHI::InitD3DDevice()
 			&& (DeviceFlags & D3D11_CREATE_DEVICE_SINGLETHREADED) == 0;
 
 		GShaderPlatformForFeatureLevel[ERHIFeatureLevel::ES2] = SP_PCD3D_ES2;
-		GShaderPlatformForFeatureLevel[ERHIFeatureLevel::ES3_1] = SP_NumPlatforms;
+		GShaderPlatformForFeatureLevel[ERHIFeatureLevel::ES3_1] = SP_PCD3D_ES3_1;
 		GShaderPlatformForFeatureLevel[ERHIFeatureLevel::SM4] = SP_PCD3D_SM4;
 		GShaderPlatformForFeatureLevel[ERHIFeatureLevel::SM5] = SP_PCD3D_SM5;
 
