added better help to cvar

diff --git a/Engine/Source/Runtime/RHI/Private/RHICommandList.cpp b/Engine/Source/Runtime/RHI/Private/RHICommandList.cpp
--- a/Engine/Source/Runtime/RHI/Private/RHICommandList.cpp
+++ b/Engine/Source/Runtime/RHI/Private/RHICommandList.cpp
@@ -21,7 +21,8 @@ static TAutoConsoleVariable<int32> CVarRHICmdBypass(
 	TEXT("r.RHICmdBypass"),
 	FRHICommandListExecutor::DefaultBypass,
 	TEXT("Whether to bypass the rhi command list and send the rhi commands immediately.\n")
-	TEXT("0: Disable, 1: Enable"));
+	TEXT("0: Disable (convenient for debugging low level graphics API calls)\n")
+	TEXT("1: Enable (required for the multithreded renderer)"));
 
 static TAutoConsoleVariable<int32> CVarRHICmdUseParallelAlgorithms(
 	TEXT("r.RHICmdUseParallelAlgorithms"),
