UE4 - remove some macros relating to the rhi resources

diff --git a/Engine/Source/Runtime/RHI/Public/RHIDefinitions.h b/Engine/Source/Runtime/RHI/Public/RHIDefinitions.h
--- a/Engine/Source/Runtime/RHI/Public/RHIDefinitions.h
+++ b/Engine/Source/Runtime/RHI/Public/RHIDefinitions.h
@@ -435,42 +435,37 @@ enum EBufferUsageFlags
 	BUF_AnyDynamic      = (BUF_Dynamic|BUF_Volatile),
 };
 
-#define ENUM_RHI_RESOURCE_TYPES(EnumerationMacro) \
-	EnumerationMacro(SamplerState,None) \
-	EnumerationMacro(RasterizerState,None) \
-	EnumerationMacro(DepthStencilState,None) \
-	EnumerationMacro(BlendState,None) \
-	EnumerationMacro(VertexDeclaration,None) \
-	EnumerationMacro(VertexShader,None) \
-	EnumerationMacro(HullShader,None) \
-	EnumerationMacro(DomainShader,None) \
-	EnumerationMacro(PixelShader,None) \
-	EnumerationMacro(GeometryShader,None) \
-	EnumerationMacro(ComputeShader,None) \
-	EnumerationMacro(BoundShaderState,None) \
-	EnumerationMacro(UniformBuffer,None) \
-	EnumerationMacro(IndexBuffer,None) \
-	EnumerationMacro(VertexBuffer,None) \
-	EnumerationMacro(StructuredBuffer,None) \
-	EnumerationMacro(Texture,None) \
-	EnumerationMacro(Texture2D,Texture) \
-	EnumerationMacro(Texture2DArray,Texture) \
-	EnumerationMacro(Texture3D,Texture) \
-	EnumerationMacro(TextureCube,Texture) \
-	EnumerationMacro(TextureReference,Texture) \
-	EnumerationMacro(RenderQuery,None) \
-	EnumerationMacro(Viewport,None) \
-	EnumerationMacro(UnorderedAccessView,None) \
-	EnumerationMacro(ShaderResourceView,None)
-
 /** An enumeration of the different RHI reference types. */
 enum ERHIResourceType
 {
 	RRT_None,
 
-#define DECLARE_RESOURCETYPE_ENUM(Type,ParentType) RRT_##Type,
-	ENUM_RHI_RESOURCE_TYPES(DECLARE_RESOURCETYPE_ENUM)
-#undef DECLARE_RESOURCETYPE_ENUM
+	RRT_SamplerState,
+	RRT_RasterizerState,
+	RRT_DepthStencilState,
+	RRT_BlendState,
+	RRT_VertexDeclaration,
+	RRT_VertexShader,
+	RRT_HullShader,
+	RRT_DomainShader,
+	RRT_PixelShader,
+	RRT_GeometryShader,
+	RRT_ComputeShader,
+	RRT_BoundShaderState,
+	RRT_UniformBuffer,
+	RRT_IndexBuffer,
+	RRT_VertexBuffer,
+	RRT_StructuredBuffer,
+	RRT_Texture,
+	RRT_Texture2D,
+	RRT_Texture2DArray,
+	RRT_Texture3D,
+	RRT_TextureCube,
+	RRT_TextureReference,
+	RRT_RenderQuery,
+	RRT_Viewport,
+	RRT_UnorderedAccessView,
+	RRT_ShaderResourceView,
 
 	RRT_Num
 };
diff --git a/Engine/Source/Runtime/RHI/Public/RHIResources.h b/Engine/Source/Runtime/RHI/Public/RHIResources.h
--- a/Engine/Source/Runtime/RHI/Public/RHIResources.h
+++ b/Engine/Source/Runtime/RHI/Public/RHIResources.h
@@ -586,11 +586,88 @@ public:
 class FRHIUnorderedAccessView : public FRHIResource {};
 class FRHIShaderResourceView : public FRHIResource {};
 
-// Declare RHI resource reference types.
-#define DEFINE_RHI_REFERENCE_TYPE(Type,ParentType) \
-	typedef FRHI##Type*              F##Type##RHIParamRef; \
-	typedef TRefCountPtr<FRHI##Type> F##Type##RHIRef;
-ENUM_RHI_RESOURCE_TYPES(DEFINE_RHI_REFERENCE_TYPE);
+
+
+typedef FRHISamplerState*              FSamplerStateRHIParamRef;
+typedef TRefCountPtr<FRHISamplerState> FSamplerStateRHIRef;
+
+typedef FRHIRasterizerState*              FRasterizerStateRHIParamRef;
+typedef TRefCountPtr<FRHIRasterizerState> FRasterizerStateRHIRef;
+
+typedef FRHIDepthStencilState*              FDepthStencilStateRHIParamRef;
+typedef TRefCountPtr<FRHIDepthStencilState> FDepthStencilStateRHIRef;
+
+typedef FRHIBlendState*              FBlendStateRHIParamRef;
+typedef TRefCountPtr<FRHIBlendState> FBlendStateRHIRef;
+
+typedef FRHIVertexDeclaration*              FVertexDeclarationRHIParamRef;
+typedef TRefCountPtr<FRHIVertexDeclaration> FVertexDeclarationRHIRef;
+
+typedef FRHIVertexShader*              FVertexShaderRHIParamRef;
+typedef TRefCountPtr<FRHIVertexShader> FVertexShaderRHIRef;
+
+typedef FRHIHullShader*              FHullShaderRHIParamRef;
+typedef TRefCountPtr<FRHIHullShader> FHullShaderRHIRef;
+
+typedef FRHIDomainShader*              FDomainShaderRHIParamRef;
+typedef TRefCountPtr<FRHIDomainShader> FDomainShaderRHIRef;
+
+typedef FRHIPixelShader*              FPixelShaderRHIParamRef;
+typedef TRefCountPtr<FRHIPixelShader> FPixelShaderRHIRef;
+
+typedef FRHIGeometryShader*              FGeometryShaderRHIParamRef;
+typedef TRefCountPtr<FRHIGeometryShader> FGeometryShaderRHIRef;
+
+typedef FRHIComputeShader*              FComputeShaderRHIParamRef;
+typedef TRefCountPtr<FRHIComputeShader> FComputeShaderRHIRef;
+
+typedef FRHIBoundShaderState*              FBoundShaderStateRHIParamRef;
+typedef TRefCountPtr<FRHIBoundShaderState> FBoundShaderStateRHIRef;
+
+typedef FRHIUniformBuffer*              FUniformBufferRHIParamRef;
+typedef TRefCountPtr<FRHIUniformBuffer> FUniformBufferRHIRef;
+
+typedef FRHIIndexBuffer*              FIndexBufferRHIParamRef;
+typedef TRefCountPtr<FRHIIndexBuffer> FIndexBufferRHIRef;
+
+typedef FRHIVertexBuffer*              FVertexBufferRHIParamRef;
+typedef TRefCountPtr<FRHIVertexBuffer> FVertexBufferRHIRef;
+
+typedef FRHIStructuredBuffer*              FStructuredBufferRHIParamRef;
+typedef TRefCountPtr<FRHIStructuredBuffer> FStructuredBufferRHIRef;
+
+typedef FRHITexture*              FTextureRHIParamRef;
+typedef TRefCountPtr<FRHITexture> FTextureRHIRef;
+
+typedef FRHITexture2D*              FTexture2DRHIParamRef;
+typedef TRefCountPtr<FRHITexture2D> FTexture2DRHIRef;
+
+typedef FRHITexture2DArray*              FTexture2DArrayRHIParamRef;
+typedef TRefCountPtr<FRHITexture2DArray> FTexture2DArrayRHIRef;
+
+typedef FRHITexture3D*              FTexture3DRHIParamRef;
+typedef TRefCountPtr<FRHITexture3D> FTexture3DRHIRef;
+
+typedef FRHITextureCube*              FTextureCubeRHIParamRef;
+typedef TRefCountPtr<FRHITextureCube> FTextureCubeRHIRef;
+
+typedef FRHITextureReference*              FTextureReferenceRHIParamRef;
+typedef TRefCountPtr<FRHITextureReference> FTextureReferenceRHIRef;
+
+typedef FRHIRenderQuery*              FRenderQueryRHIParamRef;
+typedef TRefCountPtr<FRHIRenderQuery> FRenderQueryRHIRef;
+
+typedef FRHIViewport*              FViewportRHIParamRef;
+typedef TRefCountPtr<FRHIViewport> FViewportRHIRef;
+
+typedef FRHIUnorderedAccessView*              FUnorderedAccessViewRHIParamRef;
+typedef TRefCountPtr<FRHIUnorderedAccessView> FUnorderedAccessViewRHIRef;
+
+typedef FRHIShaderResourceView*              FShaderResourceViewRHIParamRef;
+typedef TRefCountPtr<FRHIShaderResourceView> FShaderResourceViewRHIRef;
+
+
+
 
 class FRHIRenderTargetView
 {
@@ -974,7 +1051,10 @@ protected:
 	// Weak reference, don't create a circular dependency that would prevent the viewport from being destroyed.
 	FRHIViewport* ViewportRHI;
 };
-DEFINE_RHI_REFERENCE_TYPE(CustomPresent, Resource)
+
+
+typedef FRHICustomPresent*              FCustomPresentRHIParamRef;
+typedef TRefCountPtr<FRHICustomPresent> FCustomPresentRHIRef;
 
 // Template magic to convert an FRHI*Shader to its enum
 template<typename TRHIShader> struct TRHIShaderToEnum {};
