Merge RHISubmitHint

diff --git a/Engine/Source/Runtime/OpenGLDrv/Private/OpenGLCommands.cpp b/Engine/Source/Runtime/OpenGLDrv/Private/OpenGLCommands.cpp
--- a/Engine/Source/Runtime/OpenGLDrv/Private/OpenGLCommands.cpp
+++ b/Engine/Source/Runtime/OpenGLDrv/Private/OpenGLCommands.cpp
@@ -3543,6 +3543,9 @@ void FOpenGLDynamicRHI::RHIEnableDepthBoundsTest(bool bEnable,float MinDepth,flo
 	}
 }
 
+void FOpenGLDynamicRHI::RHISubmitCommandsHint()
+{
+}
 IRHICommandContext* FOpenGLDynamicRHI::RHIGetDefaultContext()
 {
 	return this;
diff --git a/Engine/Source/Runtime/RHI/Public/RHICommandList.h b/Engine/Source/Runtime/RHI/Public/RHICommandList.h
--- a/Engine/Source/Runtime/RHI/Public/RHICommandList.h
+++ b/Engine/Source/Runtime/RHI/Public/RHICommandList.h
@@ -1094,6 +1094,14 @@ struct FRHICommandEndOcclusionQueryBatch : public FRHICommand<FRHICommandEndOccl
 	RHI_API void Execute(FRHICommandListBase& CmdList);
 };
 
+struct FRHICommandSubmitCommandsHint : public FRHICommand<FRHICommandSubmitCommandsHint>
+{
+	FORCEINLINE_DEBUGGABLE FRHICommandSubmitCommandsHint()
+	{
+	}
+	RHI_API void Execute(FRHICommandListBase& CmdList);
+};
+
 struct FRHICommandBeginScene : public FRHICommand<FRHICommandBeginScene>
 {
 	FORCEINLINE_DEBUGGABLE FRHICommandBeginScene()
@@ -1823,7 +1831,17 @@ public:
 		new (AllocCommand<FRHICommandEndOcclusionQueryBatch>()) FRHICommandEndOcclusionQueryBatch();
 	}
 
-	// These 6 are special in that they must be called on the immediate command list and they force a flush only when we are not doing RHI thread
+	FORCEINLINE_DEBUGGABLE void SubmitCommandsHint()
+	{
+		if (Bypass())
+		{
+			CMD_CONTEXT(SubmitCommandsHint)();
+			return;
+		}
+		new (AllocCommand<FRHICommandSubmitCommandsHint>()) FRHICommandSubmitCommandsHint();
+	}	
+
+// These 6 are special in that they must be called on the immediate command list and they force a flush only when we are not doing RHI thread
 	void BeginScene();
 	void EndScene();
 	void BeginDrawingViewport(FViewportRHIParamRef Viewport, FTextureRHIParamRef RenderTargetRHI);
diff --git a/Engine/Source/Runtime/RHI/Public/RHICommandListCommandExecutes.inl b/Engine/Source/Runtime/RHI/Public/RHICommandListCommandExecutes.inl
--- a/Engine/Source/Runtime/RHI/Public/RHICommandListCommandExecutes.inl
+++ b/Engine/Source/Runtime/RHI/Public/RHICommandListCommandExecutes.inl
@@ -369,6 +369,12 @@ void FRHICommandEndOcclusionQueryBatch::Execute(FRHICommandListBase& CmdList)
 	INTERNAL_DECORATOR(EndOcclusionQueryBatch)();
 }
 
+void FRHICommandSubmitCommandsHint::Execute(FRHICommandListBase& CmdList)
+{
+	RHISTAT(SubmitCommandsHint);
+	INTERNAL_DECORATOR(SubmitCommandsHint)();
+}
+
 void FRHICommandBeginScene::Execute(FRHICommandListBase& CmdList)
 {
 	RHISTAT(BeginScene);
diff --git a/Engine/Source/Runtime/RHI/Public/RHIMethods.h b/Engine/Source/Runtime/RHI/Public/RHIMethods.h
--- a/Engine/Source/Runtime/RHI/Public/RHIMethods.h
+++ b/Engine/Source/Runtime/RHI/Public/RHIMethods.h
@@ -1292,6 +1292,11 @@ DEFINE_RHIMETHOD_CMDLIST_0(
 	,
 	);
 
+DEFINE_RHIMETHOD_CMDLIST_0(
+	void,SubmitCommandsHint,
+	,
+	);
+
 // This method is queued with an RHIThread, otherwise it will flush after it is queued; without an RHI thread there is no benefit to queuing this frame advance commands
 DEFINE_RHIMETHOD_CMDLIST_2(
 	void,BeginDrawingViewport,
diff --git a/Engine/Source/Runtime/Renderer/Private/DeferredShadingRenderer.cpp b/Engine/Source/Runtime/Renderer/Private/DeferredShadingRenderer.cpp
--- a/Engine/Source/Runtime/Renderer/Private/DeferredShadingRenderer.cpp
+++ b/Engine/Source/Runtime/Renderer/Private/DeferredShadingRenderer.cpp
@@ -667,6 +667,11 @@ void FDeferredShadingSceneRenderer::RenderOcclusion(FRHICommandListImmediate& RH
 		// This is done after the downsampled depth buffer is created so that it can be used for issuing queries
 		BeginOcclusionTests(RHICmdList, bRenderQueries, bRenderHZB);
 
+
+		// Hint to the RHI to submit commands up to this point to the GPU if possible.  Can help avoid CPU stalls next frame waiting
+		// for these query results on some platforms.
+		RHICmdList.SubmitCommandsHint();
+
 		if (bRenderQueries && GRHIThread)
 		{
 			QUICK_SCOPE_CYCLE_COUNTER(STAT_OcclusionSubmittedFence_Dispatch);
diff --git a/Engine/Source/Runtime/Windows/D3D11RHI/Private/D3D11Commands.cpp b/Engine/Source/Runtime/Windows/D3D11RHI/Private/D3D11Commands.cpp
--- a/Engine/Source/Runtime/Windows/D3D11RHI/Private/D3D11Commands.cpp
+++ b/Engine/Source/Runtime/Windows/D3D11RHI/Private/D3D11Commands.cpp
@@ -1925,6 +1925,10 @@ void FD3D11DynamicRHI::RHIEnableDepthBoundsTest(bool bEnable,float MinDepth,floa
 #endif
 }
 
+void FD3D11DynamicRHI::RHISubmitCommandsHint()
+{
+
+}
 IRHICommandContext* FD3D11DynamicRHI::RHIGetDefaultContext()
 {
 	return this;
