UE4 - remove rhimethods.h and all of the macro stuff relating to that

diff --git a/Engine/Source/Runtime/Core/Public/Android/AndroidPlatform.h b/Engine/Source/Runtime/Core/Public/Android/AndroidPlatform.h
--- a/Engine/Source/Runtime/Core/Public/Android/AndroidPlatform.h
+++ b/Engine/Source/Runtime/Core/Public/Android/AndroidPlatform.h
@@ -38,7 +38,6 @@ typedef FAndroidTypes FPlatformTypes;
 #define PLATFORM_USE_PTHREADS						1
 #define PLATFORM_MAX_FILEPATH_LENGTH				MAX_PATH
 #define PLATFORM_SUPPORTS_TEXTURE_STREAMING			0
-#define PLATFORM_USES_DYNAMIC_RHI					1
 #define PLATFORM_REQUIRES_FILESERVER				1
 #define PLATFORM_TCHAR_IS_4_BYTES					1
 #define PLATFORM_HAS_NO_EPROCLIM					1
diff --git a/Engine/Source/Runtime/Core/Public/HAL/Platform.h b/Engine/Source/Runtime/Core/Public/HAL/Platform.h
--- a/Engine/Source/Runtime/Core/Public/HAL/Platform.h
+++ b/Engine/Source/Runtime/Core/Public/HAL/Platform.h
@@ -196,10 +196,7 @@
 	#define PLATFORM_SUPPORTS_TEXTURE_STREAMING	1
 #endif
 #ifndef PLATFORM_SUPPORTS_VIRTUAL_TEXTURES
-#define PLATFORM_SUPPORTS_VIRTUAL_TEXTURES		0
-#endif
-#ifndef PLATFORM_USES_DYNAMIC_RHI
-	#define PLATFORM_USES_DYNAMIC_RHI			PLATFORM_DESKTOP
+	#define PLATFORM_SUPPORTS_VIRTUAL_TEXTURES		0
 #endif
 #ifndef PLATFORM_REQUIRES_FILESERVER
 	#define PLATFORM_REQUIRES_FILESERVER		0
diff --git a/Engine/Source/Runtime/Core/Public/HTML5/HTML5Platform.h b/Engine/Source/Runtime/Core/Public/HTML5/HTML5Platform.h
--- a/Engine/Source/Runtime/Core/Public/HTML5/HTML5Platform.h
+++ b/Engine/Source/Runtime/Core/Public/HTML5/HTML5Platform.h
@@ -46,7 +46,6 @@ typedef FHTML5Types FPlatformTypes;
 #endif
 #define PLATFORM_USE_PTHREADS						0
 #define PLATFORM_SUPPORTS_TEXTURE_STREAMING			1
-#define PLATFORM_USES_DYNAMIC_RHI					1
 #define PLATFORM_REQUIRES_FILESERVER				1
 #define PLATFORM_SUPPORTS_TBB						0
 #define PLATFORM_ENABLE_VECTORINTRINSICS			0
diff --git a/Engine/Source/Runtime/Core/Public/IOS/IOSPlatform.h b/Engine/Source/Runtime/Core/Public/IOS/IOSPlatform.h
--- a/Engine/Source/Runtime/Core/Public/IOS/IOSPlatform.h
+++ b/Engine/Source/Runtime/Core/Public/IOS/IOSPlatform.h
@@ -28,7 +28,6 @@ typedef FIOSPlatformTypes FPlatformTypes;
 #endif
 
 // Base defines, defaults are commented out
-#define PLATFORM_USES_DYNAMIC_RHI						1
 #define PLATFORM_LITTLE_ENDIAN							1
 #define PLATFORM_SUPPORTS_PRAGMA_PACK					1
 #define PLATFORM_COMPILER_DISTINGUISHES_INT_AND_LONG	1
diff --git a/Engine/Source/Runtime/Core/Public/Linux/LinuxPlatform.h b/Engine/Source/Runtime/Core/Public/Linux/LinuxPlatform.h
--- a/Engine/Source/Runtime/Core/Public/Linux/LinuxPlatform.h
+++ b/Engine/Source/Runtime/Core/Public/Linux/LinuxPlatform.h
@@ -43,8 +43,6 @@ typedef FLinuxPlatformTypes FPlatformTypes;
 #define PLATFORM_HAS_BSD_SOCKET_FEATURE_IOCTL		1
 #define PLATFORM_HAS_BSD_IPV6_SOCKETS				1
 
-#define PLATFORM_USES_DYNAMIC_RHI					1
-
 // only enable vectorintrinsics on x86(-64) for now
 #if defined(_M_IX86) || defined(__i386__) || defined(_M_X64) || defined(__x86_64__) || defined (__amd64__) 
 	#define PLATFORM_ENABLE_VECTORINTRINSICS		1
diff --git a/Engine/Source/Runtime/Core/Public/WinRT/WinRTARMPlatform.h b/Engine/Source/Runtime/Core/Public/WinRT/WinRTARMPlatform.h
--- a/Engine/Source/Runtime/Core/Public/WinRT/WinRTARMPlatform.h
+++ b/Engine/Source/Runtime/Core/Public/WinRT/WinRTARMPlatform.h
@@ -40,7 +40,6 @@ typedef FWinRTTypes FPlatformTypes;
 #define PLATFORM_MAX_FILEPATH_LENGTH						MAX_PATH
 [REDACTED_EMAIL]: Texture streaming support...
 #define PLATFORM_SUPPORTS_TEXTURE_STREAMING					0
-#define PLATFORM_USES_DYNAMIC_RHI							1
 #define PLATFORM_REQUIRES_FILESERVER						1
 #define PLATFORM_SUPPORTS_MULTITHREADED_GC					0
 #define PLATFORM_HAS_BSD_SOCKET_FEATURE_WINSOCKETS			1
diff --git a/Engine/Source/Runtime/Core/Public/WinRT/WinRTPlatform.h b/Engine/Source/Runtime/Core/Public/WinRT/WinRTPlatform.h
--- a/Engine/Source/Runtime/Core/Public/WinRT/WinRTPlatform.h
+++ b/Engine/Source/Runtime/Core/Public/WinRT/WinRTPlatform.h
@@ -45,7 +45,6 @@ typedef FWinRTTypes FPlatformTypes;
 #define PLATFORM_MAX_FILEPATH_LENGTH						MAX_PATH
 [REDACTED_EMAIL]: Texture streaming support...
 #define PLATFORM_SUPPORTS_TEXTURE_STREAMING					0
-#define PLATFORM_USES_DYNAMIC_RHI							1
 #define PLATFORM_REQUIRES_FILESERVER						1
 #define PLATFORM_SUPPORTS_MULTITHREADED_GC					0
 #define PLATFORM_HAS_BSD_SOCKET_FEATURE_WINSOCKETS			1
diff --git a/Engine/Source/Runtime/EmptyRHI/Public/EmptyRHI.h b/Engine/Source/Runtime/EmptyRHI/Public/EmptyRHI.h
--- a/Engine/Source/Runtime/EmptyRHI/Public/EmptyRHI.h
+++ b/Engine/Source/Runtime/EmptyRHI/Public/EmptyRHI.h
@@ -35,9 +35,179 @@ public:
 		return static_cast<TEmptyResourceTraits<TRHIType>::TConcreteType*>(Resource);
 	}
 
-	#define DEFINE_RHIMETHOD(Type,Name,ParameterTypesAndNames,ParameterNames,ReturnStatement,NullImplementation) virtual Type RHI##Name ParameterTypesAndNames
-	#include "RHIMethods.h"
-	#undef DEFINE_RHIMETHOD
+	virtual void RHIGpuTimeBegin(uint32 Hash,bool bCompute) final override;
+	virtual void RHIGpuTimeEnd(uint32 Hash,bool bCompute) final override;
+	virtual FSamplerStateRHIRef RHICreateSamplerState(const FSamplerStateInitializerRHI& Initializer) final override;
+	virtual FRasterizerStateRHIRef RHICreateRasterizerState(const FRasterizerStateInitializerRHI& Initializer) final override;
+	virtual FDepthStencilStateRHIRef RHICreateDepthStencilState(const FDepthStencilStateInitializerRHI& Initializer) final override;
+	virtual FBlendStateRHIRef RHICreateBlendState(const FBlendStateInitializerRHI& Initializer) final override;
+	virtual FVertexDeclarationRHIRef RHICreateVertexDeclaration(const FVertexDeclarationElementList& Elements) final override;
+	virtual FPixelShaderRHIRef RHICreatePixelShader(const TArray<uint8>& Code) final override;
+	virtual FVertexShaderRHIRef RHICreateVertexShader(const TArray<uint8>& Code) final override;
+	virtual FHullShaderRHIRef RHICreateHullShader(const TArray<uint8>& Code) final override;
+	virtual FDomainShaderRHIRef RHICreateDomainShader(const TArray<uint8>& Code) final override;
+	virtual FGeometryShaderRHIRef RHICreateGeometryShader(const TArray<uint8>& Code) final override;
+	virtual FGeometryShaderRHIRef RHICreateGeometryShaderWithStreamOutput(const TArray<uint8>& Code, const FStreamOutElementList& ElementList, uint32 NumStrides, const uint32* Strides, int32 RasterizedStream) final override;
+	virtual FComputeShaderRHIRef RHICreateComputeShader(const TArray<uint8>& Code) final override;
+	virtual FBoundShaderStateRHIRef RHICreateBoundShaderState(FVertexDeclarationRHIParamRef VertexDeclaration, FVertexShaderRHIParamRef VertexShader, FHullShaderRHIParamRef HullShader, FDomainShaderRHIParamRef DomainShader, FPixelShaderRHIParamRef PixelShader, FGeometryShaderRHIParamRef GeometryShader) final override;
+	virtual FUniformBufferRHIRef RHICreateUniformBuffer(const void* Contents, const FRHIUniformBufferLayout& Layout, EUniformBufferUsage Usage) final override;
+	virtual FIndexBufferRHIRef RHICreateIndexBuffer(uint32 Stride, uint32 Size, uint32 InUsage, FRHIResourceCreateInfo& CreateInfo) final override;
+	virtual void* RHILockIndexBuffer(FIndexBufferRHIParamRef IndexBuffer, uint32 Offset, uint32 Size, EResourceLockMode LockMode) final override;
+	virtual void RHIUnlockIndexBuffer(FIndexBufferRHIParamRef IndexBuffer) final override;
+	virtual FVertexBufferRHIRef RHICreateVertexBuffer(uint32 Size, uint32 InUsage, FRHIResourceCreateInfo& CreateInfo) final override;
+	virtual void* RHILockVertexBuffer(FVertexBufferRHIParamRef VertexBuffer, uint32 Offset, uint32 SizeRHI, EResourceLockMode LockMode) final override;
+	virtual void RHIUnlockVertexBuffer(FVertexBufferRHIParamRef VertexBuffer) final override;
+	virtual void RHICopyVertexBuffer(FVertexBufferRHIParamRef SourceBuffer,FVertexBufferRHIParamRef DestBuffer) final override;
+	virtual FStructuredBufferRHIRef RHICreateStructuredBuffer(uint32 Stride, uint32 Size, uint32 InUsage, FRHIResourceCreateInfo& CreateInfo) final override;
+	virtual void* RHILockStructuredBuffer(FStructuredBufferRHIParamRef StructuredBuffer, uint32 Offset, uint32 SizeRHI, EResourceLockMode LockMode) final override;
+	virtual void RHIUnlockStructuredBuffer(FStructuredBufferRHIParamRef StructuredBuffer) final override;
+	virtual FUnorderedAccessViewRHIRef RHICreateUnorderedAccessView(FStructuredBufferRHIParamRef StructuredBuffer, bool bUseUAVCounter, bool bAppendBuffer) final override;
+	virtual FUnorderedAccessViewRHIRef RHICreateUnorderedAccessView(FTextureRHIParamRef Texture) final override;
+	virtual FUnorderedAccessViewRHIRef RHICreateUnorderedAccessView(FVertexBufferRHIParamRef VertexBuffer, uint8 Format) final override;
+	virtual FShaderResourceViewRHIRef RHICreateShaderResourceView(FStructuredBufferRHIParamRef StructuredBuffer) final override;
+	virtual FShaderResourceViewRHIRef RHICreateShaderResourceView(FVertexBufferRHIParamRef VertexBuffer, uint32 Stride, uint8 Format) final override;
+	virtual uint64 RHICalcTexture2DPlatformSize(uint32 SizeX, uint32 SizeY, uint8 Format, uint32 NumMips, uint32 NumSamples, uint32 Flags, uint32& OutAlign) final override;
+	virtual uint64 RHICalcTexture3DPlatformSize(uint32 SizeX, uint32 SizeY, uint32 SizeZ, uint8 Format, uint32 NumMips, uint32 Flags, uint32& OutAlign) final override;
+	virtual uint64 RHICalcTextureCubePlatformSize(uint32 Size, uint8 Format, uint32 NumMips, uint32 Flags, uint32& OutAlign) final override;
+	virtual void RHIGetTextureMemoryStats(FTextureMemoryStats& OutStats) final override;
+	virtual bool RHIGetTextureMemoryVisualizeData(FColor* TextureData,int32 SizeX,int32 SizeY,int32 Pitch,int32 PixelSize) final override;
+	virtual FTextureReferenceRHIRef RHICreateTextureReference(FLastRenderTimeContainer* LastRenderTime) final override;
+	virtual void RHIUpdateTextureReference(FTextureReferenceRHIParamRef TextureRef, FTextureRHIParamRef NewTexture) final override;
+	virtual FTexture2DRHIRef RHICreateTexture2D(uint32 SizeX, uint32 SizeY, uint8 Format, uint32 NumMips, uint32 NumSamples, uint32 Flags, FRHIResourceCreateInfo& CreateInfo) final override;
+	virtual FTexture2DRHIRef RHIAsyncCreateTexture2D(uint32 SizeX, uint32 SizeY, uint8 Format, uint32 NumMips, uint32 Flags, void** InitialMipData, uint32 NumInitialMips) final override;
+	virtual void RHICopySharedMips(FTexture2DRHIParamRef DestTexture2D, FTexture2DRHIParamRef SrcTexture2D) final override;
+	virtual FTexture2DArrayRHIRef RHICreateTexture2DArray(uint32 SizeX, uint32 SizeY, uint32 SizeZ, uint8 Format, uint32 NumMips, uint32 Flags, FRHIResourceCreateInfo& CreateInfo) final override;
+	virtual FTexture3DRHIRef RHICreateTexture3D(uint32 SizeX, uint32 SizeY, uint32 SizeZ, uint8 Format, uint32 NumMips, uint32 Flags, FRHIResourceCreateInfo& CreateInfo) final override;
+	virtual void RHIGetResourceInfo(FTextureRHIParamRef Ref, FRHIResourceInfo& OutInfo) final override;
+	virtual FShaderResourceViewRHIRef RHICreateShaderResourceView(FTexture2DRHIParamRef Texture2DRHI, uint8 MipLevel) final override;
+	virtual FShaderResourceViewRHIRef RHICreateShaderResourceView(FTexture2DRHIParamRef Texture2DRHI, uint8 MipLevel, uint8 NumMipLevels, uint8 Format) final override;
+	virtual void RHIGenerateMips(FTextureRHIParamRef Texture) final override;
+	virtual uint32 RHIComputeMemorySize(FTextureRHIParamRef TextureRHI) final override;
+	virtual FTexture2DRHIRef RHIAsyncReallocateTexture2D(FTexture2DRHIParamRef Texture2D, int32 NewMipCount, int32 NewSizeX, int32 NewSizeY, FThreadSafeCounter* RequestStatus) final override;
+	virtual ETextureReallocationStatus RHIFinalizeAsyncReallocateTexture2D(FTexture2DRHIParamRef Texture2D, bool bBlockUntilCompleted) final override;
+	virtual ETextureReallocationStatus RHICancelAsyncReallocateTexture2D(FTexture2DRHIParamRef Texture2D, bool bBlockUntilCompleted) final override;
+	virtual void* RHILockTexture2D(FTexture2DRHIParamRef Texture, uint32 MipIndex, EResourceLockMode LockMode, uint32& DestStride, bool bLockWithinMiptail) final override;
+	virtual void RHIUnlockTexture2D(FTexture2DRHIParamRef Texture, uint32 MipIndex, bool bLockWithinMiptail) final override;
+	virtual void* RHILockTexture2DArray(FTexture2DArrayRHIParamRef Texture, uint32 TextureIndex, uint32 MipIndex, EResourceLockMode LockMode, uint32& DestStride, bool bLockWithinMiptail) final override;
+	virtual void RHIUnlockTexture2DArray(FTexture2DArrayRHIParamRef Texture, uint32 TextureIndex, uint32 MipIndex, bool bLockWithinMiptail) final override;
+	virtual void RHIUpdateTexture2D(FTexture2DRHIParamRef Texture, uint32 MipIndex, const struct FUpdateTextureRegion2D& UpdateRegion, uint32 SourcePitch, const uint8* SourceData) final override;
+	virtual void RHIUpdateTexture3D(FTexture3DRHIParamRef Texture, uint32 MipIndex, const struct FUpdateTextureRegion3D& UpdateRegion, uint32 SourceRowPitch, uint32 SourceDepthPitch, const uint8* SourceData) final override;
+	virtual FTextureCubeRHIRef RHICreateTextureCube(uint32 Size, uint8 Format, uint32 NumMips, uint32 Flags, FRHIResourceCreateInfo& CreateInfo) final override;
+	virtual FTextureCubeRHIRef RHICreateTextureCubeArray(uint32 Size, uint32 ArraySize, uint8 Format, uint32 NumMips, uint32 Flags, FRHIResourceCreateInfo& CreateInfo) final override;
+	virtual void* RHILockTextureCubeFace(FTextureCubeRHIParamRef Texture, uint32 FaceIndex, uint32 ArrayIndex, uint32 MipIndex, EResourceLockMode LockMode, uint32& DestStride, bool bLockWithinMiptail) final override;
+	virtual void RHIUnlockTextureCubeFace(FTextureCubeRHIParamRef Texture, uint32 FaceIndex, uint32 ArrayIndex, uint32 MipIndex, bool bLockWithinMiptail) final override;
+	virtual void RHIBindDebugLabelName(FTextureRHIParamRef Texture, const TCHAR* Name) final override;
+	virtual void RHIReadSurfaceData(FTextureRHIParamRef Texture,FIntRect Rect,TArray<FColor>& OutData,FReadSurfaceDataFlags InFlags) final override;
+	virtual void RHIMapStagingSurface(FTextureRHIParamRef Texture,void*& OutData,int32& OutWidth,int32& OutHeight) final override;
+	virtual void RHIUnmapStagingSurface(FTextureRHIParamRef Texture) final override;
+	virtual void RHIReadSurfaceFloatData(FTextureRHIParamRef Texture,FIntRect Rect,TArray<FFloat16Color>& OutData,ECubeFace CubeFace,int32 ArrayIndex,int32 MipIndex) final override;
+	virtual void RHIRead3DSurfaceFloatData(FTextureRHIParamRef Texture,FIntRect Rect,FIntPoint ZMinMax,TArray<FFloat16Color>& OutData) final override;
+	virtual FRenderQueryRHIRef RHICreateRenderQuery(ERenderQueryType QueryType) final override;
+	virtual bool RHIGetRenderQueryResult(FRenderQueryRHIParamRef RenderQuery, uint64& OutResult, bool bWait) final override;
+	virtual FTexture2DRHIRef RHIGetViewportBackBuffer(FViewportRHIParamRef Viewport) final override;
+	virtual void RHIAdvanceFrameForGetViewportBackBuffer() final override;
+	virtual bool RHIIsDrawingViewport() final override;
+	virtual void RHIAcquireThreadOwnership() final override;
+	virtual void RHIReleaseThreadOwnership() final override;
+	virtual void RHIFlushResources() final override;
+	virtual uint32 RHIGetGPUFrameCycles() final override;
+	virtual FViewportRHIRef RHICreateViewport(void* WindowHandle, uint32 SizeX, uint32 SizeY, bool bIsFullscreen, EPixelFormat PreferredPixelFormat) final override;
+	virtual void RHIResizeViewport(FViewportRHIParamRef Viewport, uint32 SizeX, uint32 SizeY, bool bIsFullscreen) final override;
+	virtual void RHITick(float DeltaTime) final override;
+	virtual void RHISetStreamOutTargets(uint32 NumTargets,const FVertexBufferRHIParamRef* VertexBuffers,const uint32* Offsets) final override;
+	virtual void RHIDiscardRenderTargets(bool Depth,bool Stencil,uint32 ColorBitMask) final override;
+	virtual void RHIBlockUntilGPUIdle() final override;
+	virtual void RHISuspendRendering() final override;
+	virtual void RHIResumeRendering() final override;
+	virtual bool RHIIsRenderingSuspended() final override;
+	virtual bool RHIGetAvailableResolutions(FScreenResolutionArray& Resolutions, bool bIgnoreRefreshRate) final override;
+	virtual void RHIGetSupportedResolution(uint32& Width, uint32& Height) final override;
+	virtual void RHIVirtualTextureSetFirstMipInMemory(FTexture2DRHIParamRef Texture, uint32 FirstMip) final override;
+	virtual void RHIVirtualTextureSetFirstMipVisible(FTexture2DRHIParamRef Texture, uint32 FirstMip) final override;
+	virtual void RHIExecuteCommandList(FRHICommandList* CmdList) final override;
+	virtual void* RHIGetNativeDevice() final override;
+	virtual class IRHICommandContext* RHIGetDefaultContext() final override;
+	virtual class IRHICommandContextContainer* RHIGetCommandContextContainer() final override;
+
+
+
+	virtual void RHISetComputeShader(FComputeShaderRHIParamRef ComputeShader) final override;
+	virtual void RHIDispatchComputeShader(uint32 ThreadGroupCountX, uint32 ThreadGroupCountY, uint32 ThreadGroupCountZ) final override;
+	virtual void RHIDispatchIndirectComputeShader(FVertexBufferRHIParamRef ArgumentBuffer, uint32 ArgumentOffset) final override;
+	virtual void RHIAutomaticCacheFlushAfterComputeShader(bool bEnable) final override;
+	virtual void RHIFlushComputeShaderCache() final override;
+	virtual void RHISetMultipleViewports(uint32 Count, const FViewportBounds* Data) final override;
+	virtual void RHIClearUAV(FUnorderedAccessViewRHIParamRef UnorderedAccessViewRHI, const uint32* Values) final override;
+	virtual void RHICopyToResolveTarget(FTextureRHIParamRef SourceTexture, FTextureRHIParamRef DestTexture, bool bKeepOriginalSurface, const FResolveParams& ResolveParams) final override;
+	virtual void RHIBeginRenderQuery(FRenderQueryRHIParamRef RenderQuery) final override;
+	virtual void RHIEndRenderQuery(FRenderQueryRHIParamRef RenderQuery) final override;
+	virtual void RHIBeginOcclusionQueryBatch() final override;
+	virtual void RHIEndOcclusionQueryBatch() final override;
+	virtual void RHISubmitCommandsHint() final override;
+	virtual void RHIBeginDrawingViewport(FViewportRHIParamRef Viewport, FTextureRHIParamRef RenderTargetRHI) final override;
+	virtual void RHIEndDrawingViewport(FViewportRHIParamRef Viewport, bool bPresent, bool bLockToVsync) final override;
+	virtual void RHIBeginFrame() final override;
+	virtual void RHIEndFrame() final override;
+	virtual void RHIBeginScene() final override;
+	virtual void RHIEndScene() final override;
+	virtual void RHISetStreamSource(uint32 StreamIndex, FVertexBufferRHIParamRef VertexBuffer, uint32 Stride, uint32 Offset) final override;
+	virtual void RHISetRasterizerState(FRasterizerStateRHIParamRef NewState) final override;
+	virtual void RHISetViewport(uint32 MinX, uint32 MinY, float MinZ, uint32 MaxX, uint32 MaxY, float MaxZ) final override;
+	virtual void RHISetScissorRect(bool bEnable, uint32 MinX, uint32 MinY, uint32 MaxX, uint32 MaxY) final override;
+	virtual void RHISetBoundShaderState(FBoundShaderStateRHIParamRef BoundShaderState) final override;
+	virtual void RHISetShaderTexture(FVertexShaderRHIParamRef VertexShader, uint32 TextureIndex, FTextureRHIParamRef NewTexture) final override;
+	virtual void RHISetShaderTexture(FHullShaderRHIParamRef HullShader, uint32 TextureIndex, FTextureRHIParamRef NewTexture) final override;
+	virtual void RHISetShaderTexture(FDomainShaderRHIParamRef DomainShader, uint32 TextureIndex, FTextureRHIParamRef NewTexture) final override;
+	virtual void RHISetShaderTexture(FGeometryShaderRHIParamRef GeometryShader, uint32 TextureIndex, FTextureRHIParamRef NewTexture) final override;
+	virtual void RHISetShaderTexture(FPixelShaderRHIParamRef PixelShader, uint32 TextureIndex, FTextureRHIParamRef NewTexture) final override;
+	virtual void RHISetShaderTexture(FComputeShaderRHIParamRef PixelShader, uint32 TextureIndex, FTextureRHIParamRef NewTexture) final override;
+	virtual void RHISetShaderSampler(FComputeShaderRHIParamRef ComputeShader, uint32 SamplerIndex, FSamplerStateRHIParamRef NewState) final override;
+	virtual void RHISetShaderSampler(FVertexShaderRHIParamRef VertexShader, uint32 SamplerIndex, FSamplerStateRHIParamRef NewState) final override;
+	virtual void RHISetShaderSampler(FGeometryShaderRHIParamRef GeometryShader, uint32 SamplerIndex, FSamplerStateRHIParamRef NewState) final override;
+	virtual void RHISetShaderSampler(FDomainShaderRHIParamRef DomainShader, uint32 SamplerIndex, FSamplerStateRHIParamRef NewState) final override;
+	virtual void RHISetShaderSampler(FHullShaderRHIParamRef HullShader, uint32 SamplerIndex, FSamplerStateRHIParamRef NewState) final override;
+	virtual void RHISetShaderSampler(FPixelShaderRHIParamRef PixelShader, uint32 SamplerIndex, FSamplerStateRHIParamRef NewState) final override;
+	virtual void RHISetUAVParameter(FComputeShaderRHIParamRef ComputeShader, uint32 UAVIndex, FUnorderedAccessViewRHIParamRef UAV) final override;
+	virtual void RHISetUAVParameter(FComputeShaderRHIParamRef ComputeShader, uint32 UAVIndex, FUnorderedAccessViewRHIParamRef UAV, uint32 InitialCount) final override;
+	virtual void RHISetShaderResourceViewParameter(FPixelShaderRHIParamRef PixelShader, uint32 SamplerIndex, FShaderResourceViewRHIParamRef SRV) final override;
+	virtual void RHISetShaderResourceViewParameter(FVertexShaderRHIParamRef VertexShader, uint32 SamplerIndex, FShaderResourceViewRHIParamRef SRV) final override;
+	virtual void RHISetShaderResourceViewParameter(FComputeShaderRHIParamRef ComputeShader, uint32 SamplerIndex, FShaderResourceViewRHIParamRef SRV) final override;
+	virtual void RHISetShaderResourceViewParameter(FHullShaderRHIParamRef HullShader, uint32 SamplerIndex, FShaderResourceViewRHIParamRef SRV) final override;
+	virtual void RHISetShaderResourceViewParameter(FDomainShaderRHIParamRef DomainShader, uint32 SamplerIndex, FShaderResourceViewRHIParamRef SRV) final override;
+	virtual void RHISetShaderResourceViewParameter(FGeometryShaderRHIParamRef GeometryShader, uint32 SamplerIndex, FShaderResourceViewRHIParamRef SRV) final override;
+	virtual void RHISetShaderUniformBuffer(FVertexShaderRHIParamRef VertexShader, uint32 BufferIndex, FUniformBufferRHIParamRef Buffer) final override;
+	virtual void RHISetShaderUniformBuffer(FHullShaderRHIParamRef HullShader, uint32 BufferIndex, FUniformBufferRHIParamRef Buffer) final override;
+	virtual void RHISetShaderUniformBuffer(FDomainShaderRHIParamRef DomainShader, uint32 BufferIndex, FUniformBufferRHIParamRef Buffer) final override;
+	virtual void RHISetShaderUniformBuffer(FGeometryShaderRHIParamRef GeometryShader, uint32 BufferIndex, FUniformBufferRHIParamRef Buffer) final override;
+	virtual void RHISetShaderUniformBuffer(FPixelShaderRHIParamRef PixelShader, uint32 BufferIndex, FUniformBufferRHIParamRef Buffer) final override;
+	virtual void RHISetShaderUniformBuffer(FComputeShaderRHIParamRef ComputeShader, uint32 BufferIndex, FUniformBufferRHIParamRef Buffer) final override;
+	virtual void RHISetShaderParameter(FVertexShaderRHIParamRef VertexShader, uint32 BufferIndex, uint32 BaseIndex, uint32 NumBytes, const void* NewValue) final override;
+	virtual void RHISetShaderParameter(FPixelShaderRHIParamRef PixelShader, uint32 BufferIndex, uint32 BaseIndex, uint32 NumBytes, const void* NewValue) final override;
+	virtual void RHISetShaderParameter(FHullShaderRHIParamRef HullShader, uint32 BufferIndex, uint32 BaseIndex, uint32 NumBytes, const void* NewValue) final override;
+	virtual void RHISetShaderParameter(FDomainShaderRHIParamRef DomainShader, uint32 BufferIndex, uint32 BaseIndex, uint32 NumBytes, const void* NewValue) final override;
+	virtual void RHISetShaderParameter(FGeometryShaderRHIParamRef GeometryShader, uint32 BufferIndex, uint32 BaseIndex, uint32 NumBytes, const void* NewValue) final override;
+	virtual void RHISetShaderParameter(FComputeShaderRHIParamRef ComputeShader, uint32 BufferIndex, uint32 BaseIndex, uint32 NumBytes, const void* NewValue) final override;
+	virtual void RHISetDepthStencilState(FDepthStencilStateRHIParamRef NewState, uint32 StencilRef) final override;
+	virtual void RHISetBlendState(FBlendStateRHIParamRef NewState, const FLinearColor& BlendFactor) final override;
+	virtual void RHISetRenderTargets(uint32 NumSimultaneousRenderTargets, const FRHIRenderTargetView* NewRenderTargets, const FRHIDepthRenderTargetView* NewDepthStencilTarget, uint32 NumUAVs, const FUnorderedAccessViewRHIParamRef* UAVs) final override;
+	virtual void RHISetRenderTargetsAndClear(const FRHISetRenderTargetsInfo& RenderTargetsInfo) final override;
+	virtual void RHIBindClearMRTValues(bool bClearColor, int32 NumClearColors, const FLinearColor* ColorArray, bool bClearDepth, float Depth, bool bClearStencil, uint32 Stencil) final override;
+	virtual void RHIDrawPrimitive(uint32 PrimitiveType, uint32 BaseVertexIndex, uint32 NumPrimitives, uint32 NumInstances) final override;
+	virtual void RHIDrawPrimitiveIndirect(uint32 PrimitiveType, FVertexBufferRHIParamRef ArgumentBuffer, uint32 ArgumentOffset) final override;
+	virtual void RHIDrawIndexedIndirect(FIndexBufferRHIParamRef IndexBufferRHI, uint32 PrimitiveType, FStructuredBufferRHIParamRef ArgumentsBufferRHI, int32 DrawArgumentsIndex, uint32 NumInstances) final override;
+	virtual void RHIDrawIndexedPrimitive(FIndexBufferRHIParamRef IndexBuffer, uint32 PrimitiveType, int32 BaseVertexIndex, uint32 FirstInstance, uint32 NumVertices, uint32 StartIndex, uint32 NumPrimitives, uint32 NumInstances) final override;
+	virtual void RHIDrawIndexedPrimitiveIndirect(uint32 PrimitiveType, FIndexBufferRHIParamRef IndexBuffer, FVertexBufferRHIParamRef ArgumentBuffer, uint32 ArgumentOffset) final override;
+	virtual void RHIBeginDrawPrimitiveUP(uint32 PrimitiveType, uint32 NumPrimitives, uint32 NumVertices, uint32 VertexDataStride, void*& OutVertexData) final override;
+	virtual void RHIEndDrawPrimitiveUP() final override;
+	virtual void RHIBeginDrawIndexedPrimitiveUP(uint32 PrimitiveType, uint32 NumPrimitives, uint32 NumVertices, uint32 VertexDataStride, void*& OutVertexData, uint32 MinVertexIndex, uint32 NumIndices, uint32 IndexDataStride, void*& OutIndexData) final override;
+	virtual void RHIEndDrawIndexedPrimitiveUP() final override;
+	virtual void RHIClear(bool bClearColor, const FLinearColor& Color, bool bClearDepth, float Depth, bool bClearStencil, uint32 Stencil, FIntRect ExcludeRect) final override;
+	virtual void RHIClearMRT(bool bClearColor, int32 NumClearColors, const FLinearColor* ColorArray, bool bClearDepth, float Depth, bool bClearStencil, uint32 Stencil, FIntRect ExcludeRect) final override;
+	virtual void RHIEnableDepthBoundsTest(bool bEnable, float MinDepth, float MaxDepth) final override;
+	virtual void RHIPushEvent(const TCHAR* Name) final override;
+	virtual void RHIPopEvent() final override;
+
+
 };
 
 /** Implements the Empty module as a dynamic RHI providing module. */
diff --git a/Engine/Source/Runtime/IOS/MetalRHI/Public/MetalRHI.h b/Engine/Source/Runtime/IOS/MetalRHI/Public/MetalRHI.h
--- a/Engine/Source/Runtime/IOS/MetalRHI/Public/MetalRHI.h
+++ b/Engine/Source/Runtime/IOS/MetalRHI/Public/MetalRHI.h
@@ -38,10 +38,177 @@ public:
 		return static_cast<typename TMetalResourceTraits<TRHIType>::TConcreteType*>(Resource);
 	}
 
-	// The RHI methods are defined as virtual functions in URenderHardwareInterface.
-	#define DEFINE_RHIMETHOD(Type,Name,ParameterTypesAndNames,ParameterNames,ReturnStatement,NullImplementation) virtual Type RHI##Name ParameterTypesAndNames
-	#include "RHIMethods.h"
-	#undef DEFINE_RHIMETHOD
+	virtual void RHIGpuTimeBegin(uint32 Hash,bool bCompute) final override;
+	virtual void RHIGpuTimeEnd(uint32 Hash,bool bCompute) final override;
+	virtual FSamplerStateRHIRef RHICreateSamplerState(const FSamplerStateInitializerRHI& Initializer) final override;
+	virtual FRasterizerStateRHIRef RHICreateRasterizerState(const FRasterizerStateInitializerRHI& Initializer) final override;
+	virtual FDepthStencilStateRHIRef RHICreateDepthStencilState(const FDepthStencilStateInitializerRHI& Initializer) final override;
+	virtual FBlendStateRHIRef RHICreateBlendState(const FBlendStateInitializerRHI& Initializer) final override;
+	virtual FVertexDeclarationRHIRef RHICreateVertexDeclaration(const FVertexDeclarationElementList& Elements) final override;
+	virtual FPixelShaderRHIRef RHICreatePixelShader(const TArray<uint8>& Code) final override;
+	virtual FVertexShaderRHIRef RHICreateVertexShader(const TArray<uint8>& Code) final override;
+	virtual FHullShaderRHIRef RHICreateHullShader(const TArray<uint8>& Code) final override;
+	virtual FDomainShaderRHIRef RHICreateDomainShader(const TArray<uint8>& Code) final override;
+	virtual FGeometryShaderRHIRef RHICreateGeometryShader(const TArray<uint8>& Code) final override;
+	virtual FGeometryShaderRHIRef RHICreateGeometryShaderWithStreamOutput(const TArray<uint8>& Code, const FStreamOutElementList& ElementList, uint32 NumStrides, const uint32* Strides, int32 RasterizedStream) final override;
+	virtual FComputeShaderRHIRef RHICreateComputeShader(const TArray<uint8>& Code) final override;
+	virtual FBoundShaderStateRHIRef RHICreateBoundShaderState(FVertexDeclarationRHIParamRef VertexDeclaration, FVertexShaderRHIParamRef VertexShader, FHullShaderRHIParamRef HullShader, FDomainShaderRHIParamRef DomainShader, FPixelShaderRHIParamRef PixelShader, FGeometryShaderRHIParamRef GeometryShader) final override;
+	virtual FUniformBufferRHIRef RHICreateUniformBuffer(const void* Contents, const FRHIUniformBufferLayout& Layout, EUniformBufferUsage Usage) final override;
+	virtual FIndexBufferRHIRef RHICreateIndexBuffer(uint32 Stride, uint32 Size, uint32 InUsage, FRHIResourceCreateInfo& CreateInfo) final override;
+	virtual void* RHILockIndexBuffer(FIndexBufferRHIParamRef IndexBuffer, uint32 Offset, uint32 Size, EResourceLockMode LockMode) final override;
+	virtual void RHIUnlockIndexBuffer(FIndexBufferRHIParamRef IndexBuffer) final override;
+	virtual FVertexBufferRHIRef RHICreateVertexBuffer(uint32 Size, uint32 InUsage, FRHIResourceCreateInfo& CreateInfo) final override;
+	virtual void* RHILockVertexBuffer(FVertexBufferRHIParamRef VertexBuffer, uint32 Offset, uint32 SizeRHI, EResourceLockMode LockMode) final override;
+	virtual void RHIUnlockVertexBuffer(FVertexBufferRHIParamRef VertexBuffer) final override;
+	virtual void RHICopyVertexBuffer(FVertexBufferRHIParamRef SourceBuffer,FVertexBufferRHIParamRef DestBuffer) final override;
+	virtual FStructuredBufferRHIRef RHICreateStructuredBuffer(uint32 Stride, uint32 Size, uint32 InUsage, FRHIResourceCreateInfo& CreateInfo) final override;
+	virtual void* RHILockStructuredBuffer(FStructuredBufferRHIParamRef StructuredBuffer, uint32 Offset, uint32 SizeRHI, EResourceLockMode LockMode) final override;
+	virtual void RHIUnlockStructuredBuffer(FStructuredBufferRHIParamRef StructuredBuffer) final override;
+	virtual FUnorderedAccessViewRHIRef RHICreateUnorderedAccessView(FStructuredBufferRHIParamRef StructuredBuffer, bool bUseUAVCounter, bool bAppendBuffer) final override;
+	virtual FUnorderedAccessViewRHIRef RHICreateUnorderedAccessView(FTextureRHIParamRef Texture) final override;
+	virtual FUnorderedAccessViewRHIRef RHICreateUnorderedAccessView(FVertexBufferRHIParamRef VertexBuffer, uint8 Format) final override;
+	virtual FShaderResourceViewRHIRef RHICreateShaderResourceView(FStructuredBufferRHIParamRef StructuredBuffer) final override;
+	virtual FShaderResourceViewRHIRef RHICreateShaderResourceView(FVertexBufferRHIParamRef VertexBuffer, uint32 Stride, uint8 Format) final override;
+	virtual uint64 RHICalcTexture2DPlatformSize(uint32 SizeX, uint32 SizeY, uint8 Format, uint32 NumMips, uint32 NumSamples, uint32 Flags, uint32& OutAlign) final override;
+	virtual uint64 RHICalcTexture3DPlatformSize(uint32 SizeX, uint32 SizeY, uint32 SizeZ, uint8 Format, uint32 NumMips, uint32 Flags, uint32& OutAlign) final override;
+	virtual uint64 RHICalcTextureCubePlatformSize(uint32 Size, uint8 Format, uint32 NumMips, uint32 Flags, uint32& OutAlign) final override;
+	virtual void RHIGetTextureMemoryStats(FTextureMemoryStats& OutStats) final override;
+	virtual bool RHIGetTextureMemoryVisualizeData(FColor* TextureData,int32 SizeX,int32 SizeY,int32 Pitch,int32 PixelSize) final override;
+	virtual FTextureReferenceRHIRef RHICreateTextureReference(FLastRenderTimeContainer* LastRenderTime) final override;
+	virtual void RHIUpdateTextureReference(FTextureReferenceRHIParamRef TextureRef, FTextureRHIParamRef NewTexture) final override;
+	virtual FTexture2DRHIRef RHICreateTexture2D(uint32 SizeX, uint32 SizeY, uint8 Format, uint32 NumMips, uint32 NumSamples, uint32 Flags, FRHIResourceCreateInfo& CreateInfo) final override;
+	virtual FTexture2DRHIRef RHIAsyncCreateTexture2D(uint32 SizeX, uint32 SizeY, uint8 Format, uint32 NumMips, uint32 Flags, void** InitialMipData, uint32 NumInitialMips) final override;
+	virtual void RHICopySharedMips(FTexture2DRHIParamRef DestTexture2D, FTexture2DRHIParamRef SrcTexture2D) final override;
+	virtual FTexture2DArrayRHIRef RHICreateTexture2DArray(uint32 SizeX, uint32 SizeY, uint32 SizeZ, uint8 Format, uint32 NumMips, uint32 Flags, FRHIResourceCreateInfo& CreateInfo) final override;
+	virtual FTexture3DRHIRef RHICreateTexture3D(uint32 SizeX, uint32 SizeY, uint32 SizeZ, uint8 Format, uint32 NumMips, uint32 Flags, FRHIResourceCreateInfo& CreateInfo) final override;
+	virtual void RHIGetResourceInfo(FTextureRHIParamRef Ref, FRHIResourceInfo& OutInfo) final override;
+	virtual FShaderResourceViewRHIRef RHICreateShaderResourceView(FTexture2DRHIParamRef Texture2DRHI, uint8 MipLevel) final override;
+	virtual FShaderResourceViewRHIRef RHICreateShaderResourceView(FTexture2DRHIParamRef Texture2DRHI, uint8 MipLevel, uint8 NumMipLevels, uint8 Format) final override;
+	virtual void RHIGenerateMips(FTextureRHIParamRef Texture) final override;
+	virtual uint32 RHIComputeMemorySize(FTextureRHIParamRef TextureRHI) final override;
+	virtual FTexture2DRHIRef RHIAsyncReallocateTexture2D(FTexture2DRHIParamRef Texture2D, int32 NewMipCount, int32 NewSizeX, int32 NewSizeY, FThreadSafeCounter* RequestStatus) final override;
+	virtual ETextureReallocationStatus RHIFinalizeAsyncReallocateTexture2D(FTexture2DRHIParamRef Texture2D, bool bBlockUntilCompleted) final override;
+	virtual ETextureReallocationStatus RHICancelAsyncReallocateTexture2D(FTexture2DRHIParamRef Texture2D, bool bBlockUntilCompleted) final override;
+	virtual void* RHILockTexture2D(FTexture2DRHIParamRef Texture, uint32 MipIndex, EResourceLockMode LockMode, uint32& DestStride, bool bLockWithinMiptail) final override;
+	virtual void RHIUnlockTexture2D(FTexture2DRHIParamRef Texture, uint32 MipIndex, bool bLockWithinMiptail) final override;
+	virtual void* RHILockTexture2DArray(FTexture2DArrayRHIParamRef Texture, uint32 TextureIndex, uint32 MipIndex, EResourceLockMode LockMode, uint32& DestStride, bool bLockWithinMiptail) final override;
+	virtual void RHIUnlockTexture2DArray(FTexture2DArrayRHIParamRef Texture, uint32 TextureIndex, uint32 MipIndex, bool bLockWithinMiptail) final override;
+	virtual void RHIUpdateTexture2D(FTexture2DRHIParamRef Texture, uint32 MipIndex, const struct FUpdateTextureRegion2D& UpdateRegion, uint32 SourcePitch, const uint8* SourceData) final override;
+	virtual void RHIUpdateTexture3D(FTexture3DRHIParamRef Texture, uint32 MipIndex, const struct FUpdateTextureRegion3D& UpdateRegion, uint32 SourceRowPitch, uint32 SourceDepthPitch, const uint8* SourceData) final override;
+	virtual FTextureCubeRHIRef RHICreateTextureCube(uint32 Size, uint8 Format, uint32 NumMips, uint32 Flags, FRHIResourceCreateInfo& CreateInfo) final override;
+	virtual FTextureCubeRHIRef RHICreateTextureCubeArray(uint32 Size, uint32 ArraySize, uint8 Format, uint32 NumMips, uint32 Flags, FRHIResourceCreateInfo& CreateInfo) final override;
+	virtual void* RHILockTextureCubeFace(FTextureCubeRHIParamRef Texture, uint32 FaceIndex, uint32 ArrayIndex, uint32 MipIndex, EResourceLockMode LockMode, uint32& DestStride, bool bLockWithinMiptail) final override;
+	virtual void RHIUnlockTextureCubeFace(FTextureCubeRHIParamRef Texture, uint32 FaceIndex, uint32 ArrayIndex, uint32 MipIndex, bool bLockWithinMiptail) final override;
+	virtual void RHIBindDebugLabelName(FTextureRHIParamRef Texture, const TCHAR* Name) final override;
+	virtual void RHIReadSurfaceData(FTextureRHIParamRef Texture,FIntRect Rect,TArray<FColor>& OutData,FReadSurfaceDataFlags InFlags) final override;
+	virtual void RHIMapStagingSurface(FTextureRHIParamRef Texture,void*& OutData,int32& OutWidth,int32& OutHeight) final override;
+	virtual void RHIUnmapStagingSurface(FTextureRHIParamRef Texture) final override;
+	virtual void RHIReadSurfaceFloatData(FTextureRHIParamRef Texture,FIntRect Rect,TArray<FFloat16Color>& OutData,ECubeFace CubeFace,int32 ArrayIndex,int32 MipIndex) final override;
+	virtual void RHIRead3DSurfaceFloatData(FTextureRHIParamRef Texture,FIntRect Rect,FIntPoint ZMinMax,TArray<FFloat16Color>& OutData) final override;
+	virtual FRenderQueryRHIRef RHICreateRenderQuery(ERenderQueryType QueryType) final override;
+	virtual bool RHIGetRenderQueryResult(FRenderQueryRHIParamRef RenderQuery, uint64& OutResult, bool bWait) final override;
+	virtual FTexture2DRHIRef RHIGetViewportBackBuffer(FViewportRHIParamRef Viewport) final override;
+	virtual void RHIAdvanceFrameForGetViewportBackBuffer() final override;
+	virtual bool RHIIsDrawingViewport() final override;
+	virtual void RHIAcquireThreadOwnership() final override;
+	virtual void RHIReleaseThreadOwnership() final override;
+	virtual void RHIFlushResources() final override;
+	virtual uint32 RHIGetGPUFrameCycles() final override;
+	virtual FViewportRHIRef RHICreateViewport(void* WindowHandle, uint32 SizeX, uint32 SizeY, bool bIsFullscreen, EPixelFormat PreferredPixelFormat) final override;
+	virtual void RHIResizeViewport(FViewportRHIParamRef Viewport, uint32 SizeX, uint32 SizeY, bool bIsFullscreen) final override;
+	virtual void RHITick(float DeltaTime) final override;
+	virtual void RHISetStreamOutTargets(uint32 NumTargets,const FVertexBufferRHIParamRef* VertexBuffers,const uint32* Offsets) final override;
+	virtual void RHIDiscardRenderTargets(bool Depth,bool Stencil,uint32 ColorBitMask) final override;
+	virtual void RHIBlockUntilGPUIdle() final override;
+	virtual void RHISuspendRendering() final override;
+	virtual void RHIResumeRendering() final override;
+	virtual bool RHIIsRenderingSuspended() final override;
+	virtual bool RHIGetAvailableResolutions(FScreenResolutionArray& Resolutions, bool bIgnoreRefreshRate) final override;
+	virtual void RHIGetSupportedResolution(uint32& Width, uint32& Height) final override;
+	virtual void RHIVirtualTextureSetFirstMipInMemory(FTexture2DRHIParamRef Texture, uint32 FirstMip) final override;
+	virtual void RHIVirtualTextureSetFirstMipVisible(FTexture2DRHIParamRef Texture, uint32 FirstMip) final override;
+	virtual void RHIExecuteCommandList(FRHICommandList* CmdList) final override;
+	virtual void* RHIGetNativeDevice() final override;
+	virtual class IRHICommandContext* RHIGetDefaultContext() final override;
+	virtual class IRHICommandContextContainer* RHIGetCommandContextContainer() final override;
+
+
+
+	virtual void RHISetComputeShader(FComputeShaderRHIParamRef ComputeShader) final override;
+	virtual void RHIDispatchComputeShader(uint32 ThreadGroupCountX, uint32 ThreadGroupCountY, uint32 ThreadGroupCountZ) final override;
+	virtual void RHIDispatchIndirectComputeShader(FVertexBufferRHIParamRef ArgumentBuffer, uint32 ArgumentOffset) final override;
+	virtual void RHIAutomaticCacheFlushAfterComputeShader(bool bEnable) final override;
+	virtual void RHIFlushComputeShaderCache() final override;
+	virtual void RHISetMultipleViewports(uint32 Count, const FViewportBounds* Data) final override;
+	virtual void RHIClearUAV(FUnorderedAccessViewRHIParamRef UnorderedAccessViewRHI, const uint32* Values) final override;
+	virtual void RHICopyToResolveTarget(FTextureRHIParamRef SourceTexture, FTextureRHIParamRef DestTexture, bool bKeepOriginalSurface, const FResolveParams& ResolveParams) final override;
+	virtual void RHIBeginRenderQuery(FRenderQueryRHIParamRef RenderQuery) final override;
+	virtual void RHIEndRenderQuery(FRenderQueryRHIParamRef RenderQuery) final override;
+	virtual void RHIBeginOcclusionQueryBatch() final override;
+	virtual void RHIEndOcclusionQueryBatch() final override;
+	virtual void RHISubmitCommandsHint() final override;
+	virtual void RHIBeginDrawingViewport(FViewportRHIParamRef Viewport, FTextureRHIParamRef RenderTargetRHI) final override;
+	virtual void RHIEndDrawingViewport(FViewportRHIParamRef Viewport, bool bPresent, bool bLockToVsync) final override;
+	virtual void RHIBeginFrame() final override;
+	virtual void RHIEndFrame() final override;
+	virtual void RHIBeginScene() final override;
+	virtual void RHIEndScene() final override;
+	virtual void RHISetStreamSource(uint32 StreamIndex, FVertexBufferRHIParamRef VertexBuffer, uint32 Stride, uint32 Offset) final override;
+	virtual void RHISetRasterizerState(FRasterizerStateRHIParamRef NewState) final override;
+	virtual void RHISetViewport(uint32 MinX, uint32 MinY, float MinZ, uint32 MaxX, uint32 MaxY, float MaxZ) final override;
+	virtual void RHISetScissorRect(bool bEnable, uint32 MinX, uint32 MinY, uint32 MaxX, uint32 MaxY) final override;
+	virtual void RHISetBoundShaderState(FBoundShaderStateRHIParamRef BoundShaderState) final override;
+	virtual void RHISetShaderTexture(FVertexShaderRHIParamRef VertexShader, uint32 TextureIndex, FTextureRHIParamRef NewTexture) final override;
+	virtual void RHISetShaderTexture(FHullShaderRHIParamRef HullShader, uint32 TextureIndex, FTextureRHIParamRef NewTexture) final override;
+	virtual void RHISetShaderTexture(FDomainShaderRHIParamRef DomainShader, uint32 TextureIndex, FTextureRHIParamRef NewTexture) final override;
+	virtual void RHISetShaderTexture(FGeometryShaderRHIParamRef GeometryShader, uint32 TextureIndex, FTextureRHIParamRef NewTexture) final override;
+	virtual void RHISetShaderTexture(FPixelShaderRHIParamRef PixelShader, uint32 TextureIndex, FTextureRHIParamRef NewTexture) final override;
+	virtual void RHISetShaderTexture(FComputeShaderRHIParamRef PixelShader, uint32 TextureIndex, FTextureRHIParamRef NewTexture) final override;
+	virtual void RHISetShaderSampler(FComputeShaderRHIParamRef ComputeShader, uint32 SamplerIndex, FSamplerStateRHIParamRef NewState) final override;
+	virtual void RHISetShaderSampler(FVertexShaderRHIParamRef VertexShader, uint32 SamplerIndex, FSamplerStateRHIParamRef NewState) final override;
+	virtual void RHISetShaderSampler(FGeometryShaderRHIParamRef GeometryShader, uint32 SamplerIndex, FSamplerStateRHIParamRef NewState) final override;
+	virtual void RHISetShaderSampler(FDomainShaderRHIParamRef DomainShader, uint32 SamplerIndex, FSamplerStateRHIParamRef NewState) final override;
+	virtual void RHISetShaderSampler(FHullShaderRHIParamRef HullShader, uint32 SamplerIndex, FSamplerStateRHIParamRef NewState) final override;
+	virtual void RHISetShaderSampler(FPixelShaderRHIParamRef PixelShader, uint32 SamplerIndex, FSamplerStateRHIParamRef NewState) final override;
+	virtual void RHISetUAVParameter(FComputeShaderRHIParamRef ComputeShader, uint32 UAVIndex, FUnorderedAccessViewRHIParamRef UAV) final override;
+	virtual void RHISetUAVParameter(FComputeShaderRHIParamRef ComputeShader, uint32 UAVIndex, FUnorderedAccessViewRHIParamRef UAV, uint32 InitialCount) final override;
+	virtual void RHISetShaderResourceViewParameter(FPixelShaderRHIParamRef PixelShader, uint32 SamplerIndex, FShaderResourceViewRHIParamRef SRV) final override;
+	virtual void RHISetShaderResourceViewParameter(FVertexShaderRHIParamRef VertexShader, uint32 SamplerIndex, FShaderResourceViewRHIParamRef SRV) final override;
+	virtual void RHISetShaderResourceViewParameter(FComputeShaderRHIParamRef ComputeShader, uint32 SamplerIndex, FShaderResourceViewRHIParamRef SRV) final override;
+	virtual void RHISetShaderResourceViewParameter(FHullShaderRHIParamRef HullShader, uint32 SamplerIndex, FShaderResourceViewRHIParamRef SRV) final override;
+	virtual void RHISetShaderResourceViewParameter(FDomainShaderRHIParamRef DomainShader, uint32 SamplerIndex, FShaderResourceViewRHIParamRef SRV) final override;
+	virtual void RHISetShaderResourceViewParameter(FGeometryShaderRHIParamRef GeometryShader, uint32 SamplerIndex, FShaderResourceViewRHIParamRef SRV) final override;
+	virtual void RHISetShaderUniformBuffer(FVertexShaderRHIParamRef VertexShader, uint32 BufferIndex, FUniformBufferRHIParamRef Buffer) final override;
+	virtual void RHISetShaderUniformBuffer(FHullShaderRHIParamRef HullShader, uint32 BufferIndex, FUniformBufferRHIParamRef Buffer) final override;
+	virtual void RHISetShaderUniformBuffer(FDomainShaderRHIParamRef DomainShader, uint32 BufferIndex, FUniformBufferRHIParamRef Buffer) final override;
+	virtual void RHISetShaderUniformBuffer(FGeometryShaderRHIParamRef GeometryShader, uint32 BufferIndex, FUniformBufferRHIParamRef Buffer) final override;
+	virtual void RHISetShaderUniformBuffer(FPixelShaderRHIParamRef PixelShader, uint32 BufferIndex, FUniformBufferRHIParamRef Buffer) final override;
+	virtual void RHISetShaderUniformBuffer(FComputeShaderRHIParamRef ComputeShader, uint32 BufferIndex, FUniformBufferRHIParamRef Buffer) final override;
+	virtual void RHISetShaderParameter(FVertexShaderRHIParamRef VertexShader, uint32 BufferIndex, uint32 BaseIndex, uint32 NumBytes, const void* NewValue) final override;
+	virtual void RHISetShaderParameter(FPixelShaderRHIParamRef PixelShader, uint32 BufferIndex, uint32 BaseIndex, uint32 NumBytes, const void* NewValue) final override;
+	virtual void RHISetShaderParameter(FHullShaderRHIParamRef HullShader, uint32 BufferIndex, uint32 BaseIndex, uint32 NumBytes, const void* NewValue) final override;
+	virtual void RHISetShaderParameter(FDomainShaderRHIParamRef DomainShader, uint32 BufferIndex, uint32 BaseIndex, uint32 NumBytes, const void* NewValue) final override;
+	virtual void RHISetShaderParameter(FGeometryShaderRHIParamRef GeometryShader, uint32 BufferIndex, uint32 BaseIndex, uint32 NumBytes, const void* NewValue) final override;
+	virtual void RHISetShaderParameter(FComputeShaderRHIParamRef ComputeShader, uint32 BufferIndex, uint32 BaseIndex, uint32 NumBytes, const void* NewValue) final override;
+	virtual void RHISetDepthStencilState(FDepthStencilStateRHIParamRef NewState, uint32 StencilRef) final override;
+	virtual void RHISetBlendState(FBlendStateRHIParamRef NewState, const FLinearColor& BlendFactor) final override;
+	virtual void RHISetRenderTargets(uint32 NumSimultaneousRenderTargets, const FRHIRenderTargetView* NewRenderTargets, const FRHIDepthRenderTargetView* NewDepthStencilTarget, uint32 NumUAVs, const FUnorderedAccessViewRHIParamRef* UAVs) final override;
+	virtual void RHISetRenderTargetsAndClear(const FRHISetRenderTargetsInfo& RenderTargetsInfo) final override;
+	virtual void RHIBindClearMRTValues(bool bClearColor, int32 NumClearColors, const FLinearColor* ColorArray, bool bClearDepth, float Depth, bool bClearStencil, uint32 Stencil) final override;
+	virtual void RHIDrawPrimitive(uint32 PrimitiveType, uint32 BaseVertexIndex, uint32 NumPrimitives, uint32 NumInstances) final override;
+	virtual void RHIDrawPrimitiveIndirect(uint32 PrimitiveType, FVertexBufferRHIParamRef ArgumentBuffer, uint32 ArgumentOffset) final override;
+	virtual void RHIDrawIndexedIndirect(FIndexBufferRHIParamRef IndexBufferRHI, uint32 PrimitiveType, FStructuredBufferRHIParamRef ArgumentsBufferRHI, int32 DrawArgumentsIndex, uint32 NumInstances) final override;
+	virtual void RHIDrawIndexedPrimitive(FIndexBufferRHIParamRef IndexBuffer, uint32 PrimitiveType, int32 BaseVertexIndex, uint32 FirstInstance, uint32 NumVertices, uint32 StartIndex, uint32 NumPrimitives, uint32 NumInstances) final override;
+	virtual void RHIDrawIndexedPrimitiveIndirect(uint32 PrimitiveType, FIndexBufferRHIParamRef IndexBuffer, FVertexBufferRHIParamRef ArgumentBuffer, uint32 ArgumentOffset) final override;
+	virtual void RHIBeginDrawPrimitiveUP(uint32 PrimitiveType, uint32 NumPrimitives, uint32 NumVertices, uint32 VertexDataStride, void*& OutVertexData) final override;
+	virtual void RHIEndDrawPrimitiveUP() final override;
+	virtual void RHIBeginDrawIndexedPrimitiveUP(uint32 PrimitiveType, uint32 NumPrimitives, uint32 NumVertices, uint32 VertexDataStride, void*& OutVertexData, uint32 MinVertexIndex, uint32 NumIndices, uint32 IndexDataStride, void*& OutIndexData) final override;
+	virtual void RHIEndDrawIndexedPrimitiveUP() final override;
+	virtual void RHIClear(bool bClearColor, const FLinearColor& Color, bool bClearDepth, float Depth, bool bClearStencil, uint32 Stencil, FIntRect ExcludeRect) final override;
+	virtual void RHIClearMRT(bool bClearColor, int32 NumClearColors, const FLinearColor* ColorArray, bool bClearDepth, float Depth, bool bClearStencil, uint32 Stencil, FIntRect ExcludeRect) final override;
+	virtual void RHIEnableDepthBoundsTest(bool bEnable, float MinDepth, float MaxDepth) final override;
+	virtual void RHIPushEvent(const TCHAR* Name) final override;
+	virtual void RHIPopEvent() final override;
 
 protected:
     TGlobalResource<TBoundShaderStateHistory<10000>> BoundShaderStateHistory;
diff --git a/Engine/Source/Runtime/NullDrv/Private/NullRHI.cpp b/Engine/Source/Runtime/NullDrv/Private/NullRHI.cpp
--- a/Engine/Source/Runtime/NullDrv/Private/NullRHI.cpp
+++ b/Engine/Source/Runtime/NullDrv/Private/NullRHI.cpp
@@ -3,8 +3,6 @@
 #include "NullDrvPrivate.h"
 #include "RHI.h"
 
-#if USE_DYNAMIC_RHI
-
 #include "NullRHI.h"
 #include "RenderResource.h"
 
@@ -73,10 +71,7 @@ This is used to get the same behavior as the old *_OR_IMMEDIATE present modes. *
 uint32 GPresentImmediateThreshold = 100;
 
 
-#else //USE_DYNAMIC_RHI
 
 
 // Suppress linker warning "warning LNK4221: no public symbols found; archive member will be inaccessible"
 int32 NullRHILinkerHelper;
-
-#endif //USE_DYNAMIC_RHI
diff --git a/Engine/Source/Runtime/NullDrv/Public/NullRHI.h b/Engine/Source/Runtime/NullDrv/Public/NullRHI.h
--- a/Engine/Source/Runtime/NullDrv/Public/NullRHI.h
+++ b/Engine/Source/Runtime/NullDrv/Public/NullRHI.h
@@ -14,11 +14,828 @@ public:
 	virtual void Init();
 	virtual void Shutdown();
 
-	// Implement the dynamic RHI interface using the null implementations defined in RHIMethods.h
-	#define DEFINE_RHIMETHOD(Type,Name,ParameterTypesAndNames,ParameterNames,ReturnStatement,NullImplementation) \
-	virtual Type RHI##Name ParameterTypesAndNames{ NullImplementation; }
-	#include "RHIMethods.h"
-	#undef DEFINE_RHIMETHOD
+	virtual void RHIGpuTimeBegin(uint32 Hash,bool bCompute) final override
+	{
+
+	}
+	virtual void RHIGpuTimeEnd(uint32 Hash,bool bCompute) final override
+	{
+
+	}
+	virtual FSamplerStateRHIRef RHICreateSamplerState(const FSamplerStateInitializerRHI& Initializer) final override
+	{ 
+		return new FRHISamplerState(); 
+	}
+	virtual FRasterizerStateRHIRef RHICreateRasterizerState(const FRasterizerStateInitializerRHI& Initializer) final override
+	{ 
+		return new FRHIRasterizerState(); 
+	}
+	virtual FDepthStencilStateRHIRef RHICreateDepthStencilState(const FDepthStencilStateInitializerRHI& Initializer) final override
+	{ 
+		return new FRHIDepthStencilState(); 
+	}
+	virtual FBlendStateRHIRef RHICreateBlendState(const FBlendStateInitializerRHI& Initializer) final override
+	{ 
+		return new FRHIBlendState(); 
+	}
+	virtual FVertexDeclarationRHIRef RHICreateVertexDeclaration(const FVertexDeclarationElementList& Elements) final override
+	{ 
+		return new FRHIVertexDeclaration(); 
+	}
+
+	virtual FPixelShaderRHIRef RHICreatePixelShader(const TArray<uint8>& Code) final override
+	{ 
+		return new FRHIPixelShader(); 
+	}
+
+	virtual FVertexShaderRHIRef RHICreateVertexShader(const TArray<uint8>& Code) final override
+	{ 
+		return new FRHIVertexShader(); 
+	}
+
+	virtual FHullShaderRHIRef RHICreateHullShader(const TArray<uint8>& Code) final override
+	{ 
+		return new FRHIHullShader(); 
+	}
+
+	virtual FDomainShaderRHIRef RHICreateDomainShader(const TArray<uint8>& Code) final override
+	{ 
+		return new FRHIDomainShader(); 
+	}
+
+	virtual FGeometryShaderRHIRef RHICreateGeometryShader(const TArray<uint8>& Code) final override
+	{ 
+		return new FRHIGeometryShader(); 
+	}
+
+
+	virtual FGeometryShaderRHIRef RHICreateGeometryShaderWithStreamOutput(const TArray<uint8>& Code, const FStreamOutElementList& ElementList, uint32 NumStrides, const uint32* Strides, int32 RasterizedStream) final override
+	{ 
+		return new FRHIGeometryShader(); 
+	}
+
+	virtual FComputeShaderRHIRef RHICreateComputeShader(const TArray<uint8>& Code) final override
+	{ 
+		return new FRHIComputeShader(); 
+	}
+
+
+	virtual FBoundShaderStateRHIRef RHICreateBoundShaderState(FVertexDeclarationRHIParamRef VertexDeclaration, FVertexShaderRHIParamRef VertexShader, FHullShaderRHIParamRef HullShader, FDomainShaderRHIParamRef DomainShader, FPixelShaderRHIParamRef PixelShader, FGeometryShaderRHIParamRef GeometryShader) final override
+	{ 
+		return new FRHIBoundShaderState(); 
+	}
+
+	virtual void RHISetComputeShader(FComputeShaderRHIParamRef ComputeShader) final override
+	{
+
+	}
+
+	virtual void RHIDispatchComputeShader(uint32 ThreadGroupCountX, uint32 ThreadGroupCountY, uint32 ThreadGroupCountZ) final override
+	{
+
+	}
+
+	virtual void RHIDispatchIndirectComputeShader(FVertexBufferRHIParamRef ArgumentBuffer, uint32 ArgumentOffset) final override
+	{
+
+	}
+
+	virtual void RHIAutomaticCacheFlushAfterComputeShader(bool bEnable) final override
+	{
+
+	}
+
+	virtual void RHIFlushComputeShaderCache() final override
+	{
+
+	}
+
+	virtual void RHISetMultipleViewports(uint32 Count, const FViewportBounds* Data) final override
+	{
+
+	}
+
+	virtual FUniformBufferRHIRef RHICreateUniformBuffer(const void* Contents, const FRHIUniformBufferLayout& Layout, EUniformBufferUsage Usage) final override
+	{ 
+		return new FRHIUniformBuffer(Layout); 
+	}
+
+	virtual FIndexBufferRHIRef RHICreateIndexBuffer(uint32 Stride, uint32 Size, uint32 InUsage, FRHIResourceCreateInfo& CreateInfo) final override
+	{ 
+		if(CreateInfo.ResourceArray) 
+		{ 
+			CreateInfo.ResourceArray->Discard(); 
+		} 
+		return new FRHIIndexBuffer(Stride,Size,InUsage); 
+	}
+
+	virtual void* RHILockIndexBuffer(FIndexBufferRHIParamRef IndexBuffer, uint32 Offset, uint32 Size, EResourceLockMode LockMode) final override
+	{ 
+		return GetStaticBuffer(); 
+	}
+	virtual void RHIUnlockIndexBuffer(FIndexBufferRHIParamRef IndexBuffer) final override
+	{
+
+	}
+
+	virtual FVertexBufferRHIRef RHICreateVertexBuffer(uint32 Size, uint32 InUsage, FRHIResourceCreateInfo& CreateInfo) final override
+	{ 
+		if(CreateInfo.ResourceArray) 
+		{ 
+			CreateInfo.ResourceArray->Discard(); 
+		} 
+		return new FRHIVertexBuffer(Size,InUsage); 
+	}
+
+	virtual void* RHILockVertexBuffer(FVertexBufferRHIParamRef VertexBuffer, uint32 Offset, uint32 SizeRHI, EResourceLockMode LockMode) final override
+	{ 
+		return GetStaticBuffer(); 
+	}
+	virtual void RHIUnlockVertexBuffer(FVertexBufferRHIParamRef VertexBuffer) final override
+	{
+
+	}
+
+
+	virtual void RHICopyVertexBuffer(FVertexBufferRHIParamRef SourceBuffer,FVertexBufferRHIParamRef DestBuffer) final override
+	{
+
+	}
+
+	virtual FStructuredBufferRHIRef RHICreateStructuredBuffer(uint32 Stride, uint32 Size, uint32 InUsage, FRHIResourceCreateInfo& CreateInfo) final override
+	{ 
+		if(CreateInfo.ResourceArray) 
+		{ 
+			CreateInfo.ResourceArray->Discard(); 
+		} 
+		return new FRHIStructuredBuffer(Stride,Size,InUsage); 
+	}
+
+	virtual void* RHILockStructuredBuffer(FStructuredBufferRHIParamRef StructuredBuffer, uint32 Offset, uint32 SizeRHI, EResourceLockMode LockMode) final override
+	{ 
+		return GetStaticBuffer(); 
+	}
+	virtual void RHIUnlockStructuredBuffer(FStructuredBufferRHIParamRef StructuredBuffer) final override
+	{
+
+	}
+
+
+	virtual FUnorderedAccessViewRHIRef RHICreateUnorderedAccessView(FStructuredBufferRHIParamRef StructuredBuffer, bool bUseUAVCounter, bool bAppendBuffer) final override
+	{ 
+		return new FRHIUnorderedAccessView(); 
+	}
+
+
+	virtual FUnorderedAccessViewRHIRef RHICreateUnorderedAccessView(FTextureRHIParamRef Texture) final override
+	{ 
+		return new FRHIUnorderedAccessView(); 
+	}
+
+
+	virtual FUnorderedAccessViewRHIRef RHICreateUnorderedAccessView(FVertexBufferRHIParamRef VertexBuffer, uint8 Format) final override
+	{ 
+		return new FRHIUnorderedAccessView(); 
+	}
+
+
+	virtual FShaderResourceViewRHIRef RHICreateShaderResourceView(FStructuredBufferRHIParamRef StructuredBuffer) final override
+	{ 
+		return new FRHIShaderResourceView(); 
+	}
+
+
+	virtual FShaderResourceViewRHIRef RHICreateShaderResourceView(FVertexBufferRHIParamRef VertexBuffer, uint32 Stride, uint8 Format) final override
+	{ 
+		return new FRHIShaderResourceView(); 
+	}
+
+
+	virtual void RHIClearUAV(FUnorderedAccessViewRHIParamRef UnorderedAccessViewRHI, const uint32* Values) final override
+	{
+
+	}
+
+	virtual uint64 RHICalcTexture2DPlatformSize(uint32 SizeX, uint32 SizeY, uint8 Format, uint32 NumMips, uint32 NumSamples, uint32 Flags, uint32& OutAlign) final override
+	{ 
+		OutAlign = 0; 
+		return 0; 
+	}
+
+
+	virtual uint64 RHICalcTexture3DPlatformSize(uint32 SizeX, uint32 SizeY, uint32 SizeZ, uint8 Format, uint32 NumMips, uint32 Flags, uint32& OutAlign) final override
+	{ 
+		OutAlign = 0; 
+		return 0; 
+	}
+
+	virtual uint64 RHICalcTextureCubePlatformSize(uint32 Size, uint8 Format, uint32 NumMips, uint32 Flags, uint32& OutAlign) final override
+	{ 
+		OutAlign = 0; 
+		return 0; 
+	}
+
+	virtual void RHIGetTextureMemoryStats(FTextureMemoryStats& OutStats) final override
+	{
+
+	}
+
+	virtual bool RHIGetTextureMemoryVisualizeData(FColor* TextureData,int32 SizeX,int32 SizeY,int32 Pitch,int32 PixelSize) final override
+	{ 
+		return false; 
+	}
+
+	virtual FTextureReferenceRHIRef RHICreateTextureReference(FLastRenderTimeContainer* LastRenderTime) final override
+	{ 
+		return new FRHITextureReferenceNullImpl(); 
+	}
+
+	virtual void RHIUpdateTextureReference(FTextureReferenceRHIParamRef TextureRef, FTextureRHIParamRef NewTexture) final override
+	{ 
+		if(TextureRef) 
+		{ 
+			((FRHITextureReferenceNullImpl*)TextureRef)->SetReferencedTexture(NewTexture); 
+		} 
+	}
+
+
+	virtual FTexture2DRHIRef RHICreateTexture2D(uint32 SizeX, uint32 SizeY, uint8 Format, uint32 NumMips, uint32 NumSamples, uint32 Flags, FRHIResourceCreateInfo& CreateInfo) final override
+	{ 
+		return new FRHITexture2D(SizeX,SizeY,NumMips,NumSamples,(EPixelFormat)Format,Flags); 
+	}
+
+	virtual FTexture2DRHIRef RHIAsyncCreateTexture2D(uint32 SizeX, uint32 SizeY, uint8 Format, uint32 NumMips, uint32 Flags, void** InitialMipData, uint32 NumInitialMips) final override
+	{ 
+		return FTexture2DRHIRef(); 
+	}
+
+	virtual void RHICopySharedMips(FTexture2DRHIParamRef DestTexture2D, FTexture2DRHIParamRef SrcTexture2D) final override
+	{
+	}
+	virtual FTexture2DArrayRHIRef RHICreateTexture2DArray(uint32 SizeX, uint32 SizeY, uint32 SizeZ, uint8 Format, uint32 NumMips, uint32 Flags, FRHIResourceCreateInfo& CreateInfo) final override
+	{ 
+		return new FRHITexture2DArray(SizeX,SizeY,SizeZ,NumMips,(EPixelFormat)Format,Flags); 
+	}
+
+	virtual FTexture3DRHIRef RHICreateTexture3D(uint32 SizeX, uint32 SizeY, uint32 SizeZ, uint8 Format, uint32 NumMips, uint32 Flags, FRHIResourceCreateInfo& CreateInfo) final override
+	{ 
+		return new FRHITexture3D(SizeX,SizeY,SizeZ,NumMips,(EPixelFormat)Format,Flags); 
+	}
+	virtual void RHIGetResourceInfo(FTextureRHIParamRef Ref, FRHIResourceInfo& OutInfo) final override
+	{
+	}
+
+	virtual FShaderResourceViewRHIRef RHICreateShaderResourceView(FTexture2DRHIParamRef Texture2DRHI, uint8 MipLevel) final override
+	{ 
+		return new FRHIShaderResourceView(); 
+	}
+
+	virtual FShaderResourceViewRHIRef RHICreateShaderResourceView(FTexture2DRHIParamRef Texture2DRHI, uint8 MipLevel, uint8 NumMipLevels, uint8 Format) final override
+	{ 
+		return new FRHIShaderResourceView(); 
+	}
+	virtual void RHIGenerateMips(FTextureRHIParamRef Texture) final override
+	{
+
+	}
+	virtual uint32 RHIComputeMemorySize(FTextureRHIParamRef TextureRHI) final override
+	{ 
+		return 0; 
+	}
+	virtual FTexture2DRHIRef RHIAsyncReallocateTexture2D(FTexture2DRHIParamRef Texture2D, int32 NewMipCount, int32 NewSizeX, int32 NewSizeY, FThreadSafeCounter* RequestStatus) final override
+	{ 
+		return new FRHITexture2D(NewSizeX,NewSizeY,NewMipCount,1,Texture2D->GetFormat(),Texture2D->GetFlags()); 
+	}
+	virtual ETextureReallocationStatus RHIFinalizeAsyncReallocateTexture2D(FTexture2DRHIParamRef Texture2D, bool bBlockUntilCompleted) final override
+	{ 
+		return TexRealloc_Succeeded; 
+	}
+	virtual ETextureReallocationStatus RHICancelAsyncReallocateTexture2D(FTexture2DRHIParamRef Texture2D, bool bBlockUntilCompleted) final override
+	{ 
+		return TexRealloc_Succeeded; 
+	}
+	virtual void* RHILockTexture2D(FTexture2DRHIParamRef Texture, uint32 MipIndex, EResourceLockMode LockMode, uint32& DestStride, bool bLockWithinMiptail) final override
+	{ 
+		DestStride = 0; 
+		return GetStaticBuffer(); 
+	}
+	virtual void RHIUnlockTexture2D(FTexture2DRHIParamRef Texture, uint32 MipIndex, bool bLockWithinMiptail) final override
+	{
+
+	}
+	virtual void* RHILockTexture2DArray(FTexture2DArrayRHIParamRef Texture, uint32 TextureIndex, uint32 MipIndex, EResourceLockMode LockMode, uint32& DestStride, bool bLockWithinMiptail) final override
+	{ 
+		DestStride = 0; 
+		return GetStaticBuffer(); 
+	}
+	virtual void RHIUnlockTexture2DArray(FTexture2DArrayRHIParamRef Texture, uint32 TextureIndex, uint32 MipIndex, bool bLockWithinMiptail) final override
+	{
+
+	}
+	virtual void RHIUpdateTexture2D(FTexture2DRHIParamRef Texture, uint32 MipIndex, const struct FUpdateTextureRegion2D& UpdateRegion, uint32 SourcePitch, const uint8* SourceData) final override
+	{
+
+	}
+	virtual void RHIUpdateTexture3D(FTexture3DRHIParamRef Texture, uint32 MipIndex, const struct FUpdateTextureRegion3D& UpdateRegion, uint32 SourceRowPitch, uint32 SourceDepthPitch, const uint8* SourceData) final override
+	{
+
+	}
+	virtual FTextureCubeRHIRef RHICreateTextureCube(uint32 Size, uint8 Format, uint32 NumMips, uint32 Flags, FRHIResourceCreateInfo& CreateInfo) final override
+	{ 
+		return new FRHITextureCube(Size,NumMips,(EPixelFormat)Format,Flags); 
+	}
+	virtual FTextureCubeRHIRef RHICreateTextureCubeArray(uint32 Size, uint32 ArraySize, uint8 Format, uint32 NumMips, uint32 Flags, FRHIResourceCreateInfo& CreateInfo) final override
+	{ 
+		return new FRHITextureCube(Size,NumMips,(EPixelFormat)Format,Flags); 
+	}
+	virtual void* RHILockTextureCubeFace(FTextureCubeRHIParamRef Texture, uint32 FaceIndex, uint32 ArrayIndex, uint32 MipIndex, EResourceLockMode LockMode, uint32& DestStride, bool bLockWithinMiptail) final override
+	{ 
+		DestStride = 0; 
+		return GetStaticBuffer(); 
+	}
+	virtual void RHIUnlockTextureCubeFace(FTextureCubeRHIParamRef Texture, uint32 FaceIndex, uint32 ArrayIndex, uint32 MipIndex, bool bLockWithinMiptail) final override
+	{
+	}
+	virtual void RHICopyToResolveTarget(FTextureRHIParamRef SourceTexture, FTextureRHIParamRef DestTexture, bool bKeepOriginalSurface, const FResolveParams& ResolveParams) final override
+	{
+
+	}
+
+	virtual void RHIBindDebugLabelName(FTextureRHIParamRef Texture, const TCHAR* Name) final override
+	{
+
+	}
+
+	virtual void RHIReadSurfaceData(FTextureRHIParamRef Texture,FIntRect Rect,TArray<FColor>& OutData,FReadSurfaceDataFlags InFlags) final override
+	{ 
+		OutData.AddZeroed(Rect.Width() * Rect.Height()); 
+	}
+
+
+	virtual void RHIMapStagingSurface(FTextureRHIParamRef Texture,void*& OutData,int32& OutWidth,int32& OutHeight) final override
+	{
+
+	}
+
+
+	virtual void RHIUnmapStagingSurface(FTextureRHIParamRef Texture) final override
+	{
+
+	}
+
+	virtual void RHIReadSurfaceFloatData(FTextureRHIParamRef Texture,FIntRect Rect,TArray<FFloat16Color>& OutData,ECubeFace CubeFace,int32 ArrayIndex,int32 MipIndex) final override
+	{
+
+	}
+
+	virtual void RHIRead3DSurfaceFloatData(FTextureRHIParamRef Texture,FIntRect Rect,FIntPoint ZMinMax,TArray<FFloat16Color>& OutData) final override
+	{
+
+	}
+
+
+
+	virtual FRenderQueryRHIRef RHICreateRenderQuery(ERenderQueryType QueryType) final override
+	{ 
+		return new FRHIRenderQuery(); 
+	}
+
+	virtual void RHIBeginRenderQuery(FRenderQueryRHIParamRef RenderQuery) final override
+	{
+
+	}
+	virtual void RHIEndRenderQuery(FRenderQueryRHIParamRef RenderQuery) final override
+	{
+
+	}
+
+
+	virtual bool RHIGetRenderQueryResult(FRenderQueryRHIParamRef RenderQuery, uint64& OutResult, bool bWait) final override
+	{ 
+		return true; 
+	}
+
+	virtual void RHIBeginOcclusionQueryBatch() final override
+	{
+
+	}
+	virtual void RHIEndOcclusionQueryBatch() final override
+	{
+
+	}
+	virtual void RHISubmitCommandsHint() final override
+	{
+
+	}
+
+
+	virtual void RHIBeginDrawingViewport(FViewportRHIParamRef Viewport, FTextureRHIParamRef RenderTargetRHI) final override
+	{
+
+	}
+
+	virtual void RHIEndDrawingViewport(FViewportRHIParamRef Viewport, bool bPresent, bool bLockToVsync) final override
+	{
+
+	}
+
+	virtual FTexture2DRHIRef RHIGetViewportBackBuffer(FViewportRHIParamRef Viewport) final override
+	{ 
+		return new FRHITexture2D(1,1,1,1,PF_B8G8R8A8,TexCreate_RenderTargetable); 
+	}
+
+	virtual void RHIBeginFrame() final override
+	{
+
+	}
+
+
+	virtual void RHIEndFrame() final override
+	{
+
+	}
+	virtual void RHIBeginScene() final override
+	{
+
+	}
+	virtual void RHIEndScene() final override
+	{
+
+	}
+	virtual void RHIAdvanceFrameForGetViewportBackBuffer() final override
+	{
+
+	}
+	virtual bool RHIIsDrawingViewport() final override
+	{ 
+		return false; 
+	}
+	virtual void RHIAcquireThreadOwnership() final override
+	{
+
+	}
+	virtual void RHIReleaseThreadOwnership() final override
+	{
+
+	}
+
+
+	virtual void RHIFlushResources() final override
+	{
+
+	}
+
+	virtual uint32 RHIGetGPUFrameCycles() final override
+	{ 
+		return 0; 
+	}
+
+	virtual FViewportRHIRef RHICreateViewport(void* WindowHandle, uint32 SizeX, uint32 SizeY, bool bIsFullscreen, EPixelFormat PreferredPixelFormat) final override
+	{ 
+		return new FRHIViewport(); 
+	}
+	virtual void RHIResizeViewport(FViewportRHIParamRef Viewport, uint32 SizeX, uint32 SizeY, bool bIsFullscreen) final override
+	{
+
+	}
+	virtual void RHITick(float DeltaTime) final override
+	{
+
+	}
+
+	virtual void RHISetStreamSource(uint32 StreamIndex, FVertexBufferRHIParamRef VertexBuffer, uint32 Stride, uint32 Offset) final override
+	{
+
+	}
+	virtual void RHISetStreamOutTargets(uint32 NumTargets,const FVertexBufferRHIParamRef* VertexBuffers,const uint32* Offsets) final override
+	{
+
+	}
+	virtual void RHISetRasterizerState(FRasterizerStateRHIParamRef NewState) final override
+	{
+
+	}
+
+	virtual void RHISetViewport(uint32 MinX, uint32 MinY, float MinZ, uint32 MaxX, uint32 MaxY, float MaxZ) final override
+	{
+
+	}
+
+	virtual void RHISetScissorRect(bool bEnable, uint32 MinX, uint32 MinY, uint32 MaxX, uint32 MaxY) final override
+	{
+	}
+
+	virtual void RHISetBoundShaderState(FBoundShaderStateRHIParamRef BoundShaderState) final override
+	{
+
+	}
+
+
+	virtual void RHISetShaderTexture(FVertexShaderRHIParamRef VertexShader, uint32 TextureIndex, FTextureRHIParamRef NewTexture) final override
+	{
+
+	}
+
+
+	virtual void RHISetShaderTexture(FHullShaderRHIParamRef HullShader, uint32 TextureIndex, FTextureRHIParamRef NewTexture) final override
+	{
+
+	}
+
+
+	virtual void RHISetShaderTexture(FDomainShaderRHIParamRef DomainShader, uint32 TextureIndex, FTextureRHIParamRef NewTexture) final override
+	{
+
+	}
+
+
+	virtual void RHISetShaderTexture(FGeometryShaderRHIParamRef GeometryShader, uint32 TextureIndex, FTextureRHIParamRef NewTexture) final override
+	{
+
+	}
+
+
+	virtual void RHISetShaderTexture(FPixelShaderRHIParamRef PixelShader, uint32 TextureIndex, FTextureRHIParamRef NewTexture) final override
+	{
+
+	}
+
+
+	virtual void RHISetShaderTexture(FComputeShaderRHIParamRef PixelShader, uint32 TextureIndex, FTextureRHIParamRef NewTexture) final override
+	{
+
+	}
+
+	virtual void RHISetShaderSampler(FComputeShaderRHIParamRef ComputeShader, uint32 SamplerIndex, FSamplerStateRHIParamRef NewState) final override
+	{
+
+	}
+
+	virtual void RHISetShaderSampler(FVertexShaderRHIParamRef VertexShader, uint32 SamplerIndex, FSamplerStateRHIParamRef NewState) final override
+	{
+
+	}
+
+	virtual void RHISetShaderSampler(FGeometryShaderRHIParamRef GeometryShader, uint32 SamplerIndex, FSamplerStateRHIParamRef NewState) final override
+	{
+
+	}
+	virtual void RHISetShaderSampler(FDomainShaderRHIParamRef DomainShader, uint32 SamplerIndex, FSamplerStateRHIParamRef NewState) final override
+	{
+
+	}
+	virtual void RHISetShaderSampler(FHullShaderRHIParamRef HullShader, uint32 SamplerIndex, FSamplerStateRHIParamRef NewState) final override
+	{
+
+	}
+
+	virtual void RHISetShaderSampler(FPixelShaderRHIParamRef PixelShader, uint32 SamplerIndex, FSamplerStateRHIParamRef NewState) final override
+	{
+
+	}
+
+
+	virtual void RHISetUAVParameter(FComputeShaderRHIParamRef ComputeShader, uint32 UAVIndex, FUnorderedAccessViewRHIParamRef UAV) final override
+	{
+
+	}
+
+
+	virtual void RHISetUAVParameter(FComputeShaderRHIParamRef ComputeShader, uint32 UAVIndex, FUnorderedAccessViewRHIParamRef UAV, uint32 InitialCount) final override
+	{
+
+	}
+
+
+	virtual void RHISetShaderResourceViewParameter(FPixelShaderRHIParamRef PixelShader, uint32 SamplerIndex, FShaderResourceViewRHIParamRef SRV) final override
+	{
+
+	}
+
+
+	virtual void RHISetShaderResourceViewParameter(FVertexShaderRHIParamRef VertexShader, uint32 SamplerIndex, FShaderResourceViewRHIParamRef SRV) final override
+	{
+
+	}
+
+
+	virtual void RHISetShaderResourceViewParameter(FComputeShaderRHIParamRef ComputeShader, uint32 SamplerIndex, FShaderResourceViewRHIParamRef SRV) final override
+	{
+
+	}
+
+
+	virtual void RHISetShaderResourceViewParameter(FHullShaderRHIParamRef HullShader, uint32 SamplerIndex, FShaderResourceViewRHIParamRef SRV) final override
+	{
+
+	}
+
+
+	virtual void RHISetShaderResourceViewParameter(FDomainShaderRHIParamRef DomainShader, uint32 SamplerIndex, FShaderResourceViewRHIParamRef SRV) final override
+	{
+
+	}
+	virtual void RHISetShaderResourceViewParameter(FGeometryShaderRHIParamRef GeometryShader, uint32 SamplerIndex, FShaderResourceViewRHIParamRef SRV) final override
+	{
+
+	}
+
+	virtual void RHISetShaderUniformBuffer(FVertexShaderRHIParamRef VertexShader, uint32 BufferIndex, FUniformBufferRHIParamRef Buffer) final override
+	{
+
+	}
+
+	virtual void RHISetShaderUniformBuffer(FHullShaderRHIParamRef HullShader, uint32 BufferIndex, FUniformBufferRHIParamRef Buffer) final override
+	{
+
+	}
+
+	virtual void RHISetShaderUniformBuffer(FDomainShaderRHIParamRef DomainShader, uint32 BufferIndex, FUniformBufferRHIParamRef Buffer) final override
+	{
+
+	}
+
+	virtual void RHISetShaderUniformBuffer(FGeometryShaderRHIParamRef GeometryShader, uint32 BufferIndex, FUniformBufferRHIParamRef Buffer) final override
+	{
+
+	}
+
+	virtual void RHISetShaderUniformBuffer(FPixelShaderRHIParamRef PixelShader, uint32 BufferIndex, FUniformBufferRHIParamRef Buffer) final override
+	{
+
+	}
+
+	virtual void RHISetShaderUniformBuffer(FComputeShaderRHIParamRef ComputeShader, uint32 BufferIndex, FUniformBufferRHIParamRef Buffer) final override
+	{
+
+	}
+
+	virtual void RHISetShaderParameter(FVertexShaderRHIParamRef VertexShader, uint32 BufferIndex, uint32 BaseIndex, uint32 NumBytes, const void* NewValue) final override
+	{
+
+	}
+	virtual void RHISetShaderParameter(FPixelShaderRHIParamRef PixelShader, uint32 BufferIndex, uint32 BaseIndex, uint32 NumBytes, const void* NewValue) final override
+	{
+
+	}
+
+	virtual void RHISetShaderParameter(FHullShaderRHIParamRef HullShader, uint32 BufferIndex, uint32 BaseIndex, uint32 NumBytes, const void* NewValue) final override
+	{
+
+	}
+	virtual void RHISetShaderParameter(FDomainShaderRHIParamRef DomainShader, uint32 BufferIndex, uint32 BaseIndex, uint32 NumBytes, const void* NewValue) final override
+	{
+
+	}
+	virtual void RHISetShaderParameter(FGeometryShaderRHIParamRef GeometryShader, uint32 BufferIndex, uint32 BaseIndex, uint32 NumBytes, const void* NewValue) final override
+	{
+
+	}
+	virtual void RHISetShaderParameter(FComputeShaderRHIParamRef ComputeShader, uint32 BufferIndex, uint32 BaseIndex, uint32 NumBytes, const void* NewValue) final override
+	{
+
+	}
+
+
+	virtual void RHISetDepthStencilState(FDepthStencilStateRHIParamRef NewState, uint32 StencilRef) final override
+	{
+
+	}
+
+	virtual void RHISetBlendState(FBlendStateRHIParamRef NewState, const FLinearColor& BlendFactor) final override
+	{
+
+	}
+	virtual void RHISetRenderTargets(uint32 NumSimultaneousRenderTargets, const FRHIRenderTargetView* NewRenderTargets, const FRHIDepthRenderTargetView* NewDepthStencilTarget, uint32 NumUAVs, const FUnorderedAccessViewRHIParamRef* UAVs) final override
+	{
+
+	}
+
+	virtual void RHIDiscardRenderTargets(bool Depth,bool Stencil,uint32 ColorBitMask) final override
+	{
+
+	}
+	virtual void RHISetRenderTargetsAndClear(const FRHISetRenderTargetsInfo& RenderTargetsInfo) final override
+	{
+
+	}
+
+
+
+
+	virtual void RHIBindClearMRTValues(bool bClearColor, int32 NumClearColors, const FLinearColor* ColorArray, bool bClearDepth, float Depth, bool bClearStencil, uint32 Stencil) final override
+	{
+
+	}
+
+
+	virtual void RHIDrawPrimitive(uint32 PrimitiveType, uint32 BaseVertexIndex, uint32 NumPrimitives, uint32 NumInstances) final override
+	{
+
+	}
+	virtual void RHIDrawPrimitiveIndirect(uint32 PrimitiveType, FVertexBufferRHIParamRef ArgumentBuffer, uint32 ArgumentOffset) final override
+	{
+
+	}
+
+	virtual void RHIDrawIndexedIndirect(FIndexBufferRHIParamRef IndexBufferRHI, uint32 PrimitiveType, FStructuredBufferRHIParamRef ArgumentsBufferRHI, int32 DrawArgumentsIndex, uint32 NumInstances) final override
+	{
+
+	}
+
+
+	virtual void RHIDrawIndexedPrimitive(FIndexBufferRHIParamRef IndexBuffer, uint32 PrimitiveType, int32 BaseVertexIndex, uint32 FirstInstance, uint32 NumVertices, uint32 StartIndex, uint32 NumPrimitives, uint32 NumInstances) final override
+	{
+
+	}
+	virtual void RHIDrawIndexedPrimitiveIndirect(uint32 PrimitiveType, FIndexBufferRHIParamRef IndexBuffer, FVertexBufferRHIParamRef ArgumentBuffer, uint32 ArgumentOffset) final override
+	{
+
+	}
+
+	virtual void RHIBeginDrawPrimitiveUP(uint32 PrimitiveType, uint32 NumPrimitives, uint32 NumVertices, uint32 VertexDataStride, void*& OutVertexData) final override
+	{ 
+		OutVertexData = GetStaticBuffer(); 
+	}
+	virtual void RHIEndDrawPrimitiveUP()
+	{
+
+	}
+	virtual void RHIBeginDrawIndexedPrimitiveUP(uint32 PrimitiveType, uint32 NumPrimitives, uint32 NumVertices, uint32 VertexDataStride, void*& OutVertexData, uint32 MinVertexIndex, uint32 NumIndices, uint32 IndexDataStride, void*& OutIndexData) final override
+	{ 
+		OutVertexData = GetStaticBuffer(); 
+		OutIndexData = GetStaticBuffer(); 
+	}
+	virtual void RHIEndDrawIndexedPrimitiveUP() final override
+	{
+
+	}
+	virtual void RHIClear(bool bClearColor, const FLinearColor& Color, bool bClearDepth, float Depth, bool bClearStencil, uint32 Stencil, FIntRect ExcludeRect) final override
+	{
+
+	}
+	virtual void RHIClearMRT(bool bClearColor, int32 NumClearColors, const FLinearColor* ColorArray, bool bClearDepth, float Depth, bool bClearStencil, uint32 Stencil, FIntRect ExcludeRect) final override
+	{
+
+	}
+	virtual void RHIBlockUntilGPUIdle() final override
+	{
+
+	}
+	virtual void RHISuspendRendering() final override
+	{
+
+	}
+	virtual void RHIResumeRendering() final override
+	{
+
+	}
+	virtual bool RHIIsRenderingSuspended() final override
+	{ 
+		return false; 
+	}
+	virtual bool RHIGetAvailableResolutions(FScreenResolutionArray& Resolutions, bool bIgnoreRefreshRate) final override
+	{ 
+		return false; 
+	}
+	virtual void RHIGetSupportedResolution(uint32& Width, uint32& Height) final override
+	{
+
+	}
+	virtual void RHIVirtualTextureSetFirstMipInMemory(FTexture2DRHIParamRef Texture, uint32 FirstMip) final override
+	{
+
+	}
+	virtual void RHIVirtualTextureSetFirstMipVisible(FTexture2DRHIParamRef Texture, uint32 FirstMip) final override
+	{
+
+	}
+	virtual void RHIExecuteCommandList(FRHICommandList* CmdList) final override
+	{
+
+	}
+	virtual void RHIEnableDepthBoundsTest(bool bEnable, float MinDepth, float MaxDepth) final override
+	{
+
+	}
+	virtual void* RHIGetNativeDevice() final override
+	{ 
+		return 0; 
+	}
+	virtual void RHIPushEvent(const TCHAR* Name) final override
+	{
+	}
+	virtual void RHIPopEvent()
+	{
+	}
+	virtual class IRHICommandContext* RHIGetDefaultContext() final override
+	{ 
+		return this; 
+	}
+	virtual class IRHICommandContextContainer* RHIGetCommandContextContainer() final override
+	{ 
+		return nullptr; 
+	}
 
 private:
 
diff --git a/Engine/Source/Runtime/OpenGLDrv/Private/OpenGLCommands.cpp b/Engine/Source/Runtime/OpenGLDrv/Private/OpenGLCommands.cpp
--- a/Engine/Source/Runtime/OpenGLDrv/Private/OpenGLCommands.cpp
+++ b/Engine/Source/Runtime/OpenGLDrv/Private/OpenGLCommands.cpp
@@ -263,9 +263,6 @@ static FORCEINLINE GLint ModifyFilterByMips(GLint Filter, bool bHasMips)
 	return Filter;
 }
 
-// Ignore functions from RHIMethods.h when parsing documentation; Doxygen's preprocessor can't parse the declaration, so spews warnings for the definitions.
-#if !UE_BUILD_DOCS
-
 void FOpenGLDynamicRHI::RHIGpuTimeBegin(uint32 Hash, bool bCompute)
 {
 	return;
@@ -3557,4 +3554,3 @@ IRHICommandContextContainer* FOpenGLDynamicRHI::RHIGetCommandContextContainer()
 }
 
 
-#endif
\ No newline at end of file
diff --git a/Engine/Source/Runtime/OpenGLDrv/Private/OpenGLDevice.cpp b/Engine/Source/Runtime/OpenGLDrv/Private/OpenGLDevice.cpp
--- a/Engine/Source/Runtime/OpenGLDrv/Private/OpenGLDevice.cpp
+++ b/Engine/Source/Runtime/OpenGLDrv/Private/OpenGLDevice.cpp
@@ -132,9 +132,6 @@ FOpenGLContextState& FOpenGLDynamicRHI::GetContextStateForCurrentContext()
 	}
 }
 
-// Ignore functions from RHIMethods.h when parsing documentation; Doxygen's preprocessor can't parse the declaration, so spews warnings for the definitions.
-#if !UE_BUILD_DOCS
-
 void FOpenGLDynamicRHI::RHIBeginFrame()
 {
 	RHIPrivateBeginFrame();
@@ -177,7 +174,7 @@ void FOpenGLDynamicRHI::RHIEndScene()
 	ResourceTableFrameCounter = INDEX_NONE;
 }
 
-#endif
+
 
 bool GDisableOpenGLDebugOutput = false;
 
@@ -1356,9 +1353,6 @@ void FOpenGLDynamicRHI::UnregisterQuery( FOpenGLRenderQuery* Query )
 	Queries.RemoveSingleSwap(Query);
 }
 
-// Ignore functions from RHIMethods.h when parsing documentation; Doxygen's preprocessor can't parse the declaration, so spews warnings for the definitions.
-#if !UE_BUILD_DOCS
-
 void FOpenGLDynamicRHI::RHIAutomaticCacheFlushAfterComputeShader(bool bEnable)
 {
 	// Nothing to do here...
@@ -1369,7 +1363,7 @@ void FOpenGLDynamicRHI::RHIFlushComputeShaderCache()
 	// Nothing to do here...
 }
 
-#endif
+
 
 void* FOpenGLDynamicRHI::RHIGetNativeDevice()
 {
diff --git a/Engine/Source/Runtime/OpenGLDrv/Private/OpenGLDrv.cpp b/Engine/Source/Runtime/OpenGLDrv/Private/OpenGLDrv.cpp
--- a/Engine/Source/Runtime/OpenGLDrv/Private/OpenGLDrv.cpp
+++ b/Engine/Source/Runtime/OpenGLDrv/Private/OpenGLDrv.cpp
@@ -19,8 +19,6 @@ IMPLEMENT_MODULE(FOpenGLDynamicRHIModule, OpenGLDrv);
 /** OpenGL Logging. */
 DEFINE_LOG_CATEGORY(LogOpenGL);
 
-// Ignore functions from RHIMethods.h when parsing documentation; Doxygen's preprocessor can't parse the declaration, so spews warnings for the definitions.
-#if !UE_BUILD_DOCS
 
 void FOpenGLDynamicRHI::RHIPushEvent(const TCHAR* Name)
 {
@@ -51,7 +49,7 @@ void FOpenGLGPUProfiler::PopEvent()
 
 }
 
-#endif
+
 
 void FOpenGLGPUProfiler::BeginFrame(FOpenGLDynamicRHI* InRHI)
 {
diff --git a/Engine/Source/Runtime/OpenGLDrv/Private/OpenGLQuery.cpp b/Engine/Source/Runtime/OpenGLDrv/Private/OpenGLQuery.cpp
--- a/Engine/Source/Runtime/OpenGLDrv/Private/OpenGLQuery.cpp
+++ b/Engine/Source/Runtime/OpenGLDrv/Private/OpenGLQuery.cpp
@@ -7,9 +7,6 @@
 #include "OpenGLDrvPrivate.h"
 #include "RenderCore.h"
 
-// Ignore functions from RHIMethods.h when parsing documentation; Doxygen's preprocessor can't parse the declaration, so spews warnings for the definitions.
-#if !UE_BUILD_DOCS
-
 void FOpenGLDynamicRHI::RHIBeginOcclusionQueryBatch()
 {
 }
@@ -204,7 +201,7 @@ bool FOpenGLDynamicRHI::RHIGetRenderQueryResult(FRenderQueryRHIParamRef QueryRHI
 	return bSuccess;
 }
 
-#endif
+
 
 extern void OnQueryCreation( FOpenGLRenderQuery* Query );
 extern void OnQueryDeletion( FOpenGLRenderQuery* Query );
diff --git a/Engine/Source/Runtime/OpenGLDrv/Private/OpenGLRenderTarget.cpp b/Engine/Source/Runtime/OpenGLDrv/Private/OpenGLRenderTarget.cpp
--- a/Engine/Source/Runtime/OpenGLDrv/Private/OpenGLRenderTarget.cpp
+++ b/Engine/Source/Runtime/OpenGLDrv/Private/OpenGLRenderTarget.cpp
@@ -297,9 +297,6 @@ void FOpenGLDynamicRHI::PurgeFramebufferFromCaches( GLuint Framebuffer )
 	}
 }
 
-// Ignore functions from RHIMethods.h when parsing documentation; Doxygen's preprocessor can't parse the declaration, so spews warnings for the definitions.
-#if !UE_BUILD_DOCS
-
 void FOpenGLDynamicRHI::RHICopyToResolveTarget(FTextureRHIParamRef SourceTextureRHI, FTextureRHIParamRef DestTextureRHI, bool bKeepOriginalSurface, const FResolveParams& ResolveParams)
 {
 	if (!SourceTextureRHI || !DestTextureRHI)
@@ -436,7 +433,7 @@ void FOpenGLDynamicRHI::RHICopyToResolveTarget(FTextureRHIParamRef SourceTexture
 	}
 }
 
-#endif
+
 
 void FOpenGLDynamicRHI::ReadSurfaceDataRaw(FOpenGLContextState& ContextState, FTextureRHIParamRef TextureRHI,FIntRect Rect,TArray<uint8>& OutData, FReadSurfaceDataFlags InFlags)
 {
@@ -702,9 +699,6 @@ void FOpenGLDynamicRHI::ReadSurfaceDataRaw(FOpenGLContextState& ContextState, FT
 	ContextState.Framebuffer = (GLuint)-1;
 }
 
-// Ignore functions from RHIMethods.h when parsing documentation; Doxygen's preprocessor can't parse the declaration, so spews warnings for the definitions.
-#if !UE_BUILD_DOCS
-
 void FOpenGLDynamicRHI::RHIReadSurfaceData(FTextureRHIParamRef TextureRHI,FIntRect Rect,TArray<FColor>& OutData, FReadSurfaceDataFlags InFlags)
 {
 	TArray<uint8> Temp;
@@ -872,7 +866,7 @@ void FOpenGLDynamicRHI::RHIRead3DSurfaceFloatData(FTextureRHIParamRef TextureRHI
 	ContextState.Framebuffer = (GLuint)-1;
 }
 
-#endif
+
 
 void FOpenGLDynamicRHI::BindPendingFramebuffer( FOpenGLContextState& ContextState )
 {
diff --git a/Engine/Source/Runtime/OpenGLDrv/Private/OpenGLTexture.cpp b/Engine/Source/Runtime/OpenGLDrv/Private/OpenGLTexture.cpp
--- a/Engine/Source/Runtime/OpenGLDrv/Private/OpenGLTexture.cpp
+++ b/Engine/Source/Runtime/OpenGLDrv/Private/OpenGLTexture.cpp
@@ -211,8 +211,6 @@ void FOpenGLDynamicRHI::RHIGetTextureMemoryStats(FTextureMemoryStats& OutStats)
 	OutStats.PendingMemoryAdjustment = 0;
 }
 
-// Ignore functions from RHIMethods.h when parsing documentation; Doxygen's preprocessor can't parse the declaration, so spews warnings for the definitions.
-#if !UE_BUILD_DOCS
 
 /**
  * Fills a texture with to visualize the texture pool memory.
@@ -230,7 +228,7 @@ bool FOpenGLDynamicRHI::RHIGetTextureMemoryVisualizeData( FColor* /*TextureData*
 	return false;
 }
 
-#endif
+
 
 FRHITexture* FOpenGLDynamicRHI::CreateOpenGLTexture(uint32 SizeX,uint32 SizeY,bool bCubeTexture, bool bArrayTexture, uint8 Format,uint32 NumMips,uint32 NumSamples, uint32 ArraySize, uint32 Flags, FResourceBulkDataInterface* BulkData)
 {
@@ -1777,8 +1775,6 @@ FShaderResourceViewRHIRef FOpenGLDynamicRHI::RHICreateShaderResourceView(FTextur
 	return View;
 }
 
-// Ignore functions from RHIMethods.h when parsing documentation; Doxygen's preprocessor can't parse the declaration, so spews warnings for the definitions.
-#if !UE_BUILD_DOCS
 
 /** Generates mip maps for the surface. */
 void FOpenGLDynamicRHI::RHIGenerateMips(FTextureRHIParamRef SurfaceRHI)
@@ -1804,7 +1800,7 @@ void FOpenGLDynamicRHI::RHIGenerateMips(FTextureRHIParamRef SurfaceRHI)
 	}
 }
 
-#endif
+
 
 /**
  * Computes the size in memory required by a given texture.
diff --git a/Engine/Source/Runtime/OpenGLDrv/Private/OpenGLUAV.cpp b/Engine/Source/Runtime/OpenGLDrv/Private/OpenGLUAV.cpp
--- a/Engine/Source/Runtime/OpenGLDrv/Private/OpenGLUAV.cpp
+++ b/Engine/Source/Runtime/OpenGLDrv/Private/OpenGLUAV.cpp
@@ -125,13 +125,10 @@ FShaderResourceViewRHIRef FOpenGLDynamicRHI::RHICreateShaderResourceView(FStruct
 	return new FOpenGLShaderResourceView(this,0,GL_TEXTURE_BUFFER);
 }
 
-// Ignore functions from RHIMethods.h when parsing documentation; Doxygen's preprocessor can't parse the declaration, so spews warnings for the definitions.
-#if !UE_BUILD_DOCS
-
 void FOpenGLDynamicRHI::RHIClearUAV(FUnorderedAccessViewRHIParamRef UnorderedAccessViewRHI, const uint32* Values)
 {
 	UE_LOG(LogRHI, Fatal,TEXT("OpenGL RHI doesn't support RHIClearUAV."));
 	GPUProfilingData.RegisterGPUWork(1);
 }
 
-#endif
+
diff --git a/Engine/Source/Runtime/OpenGLDrv/Private/OpenGLVertexBuffer.cpp b/Engine/Source/Runtime/OpenGLDrv/Private/OpenGLVertexBuffer.cpp
--- a/Engine/Source/Runtime/OpenGLDrv/Private/OpenGLVertexBuffer.cpp
+++ b/Engine/Source/Runtime/OpenGLDrv/Private/OpenGLVertexBuffer.cpp
@@ -202,9 +202,6 @@ void FOpenGLDynamicRHI::RHIUnlockVertexBuffer(FVertexBufferRHIParamRef VertexBuf
 	}
 }
 
-// Ignore functions from RHIMethods.h when parsing documentation; Doxygen's preprocessor can't parse the declaration, so spews warnings for the definitions.
-#if !UE_BUILD_DOCS
-
 void FOpenGLDynamicRHI::RHICopyVertexBuffer(FVertexBufferRHIParamRef SourceBufferRHI,FVertexBufferRHIParamRef DestBufferRHI)
 {
 	VERIFY_GL_SCOPE();
@@ -220,4 +217,3 @@ void FOpenGLDynamicRHI::RHICopyVertexBuffer(FVertexBufferRHIParamRef SourceBuffe
 	glBindBuffer(GL_COPY_WRITE_BUFFER,0);
 }
 
-#endif
diff --git a/Engine/Source/Runtime/OpenGLDrv/Private/OpenGLViewport.cpp b/Engine/Source/Runtime/OpenGLDrv/Private/OpenGLViewport.cpp
--- a/Engine/Source/Runtime/OpenGLDrv/Private/OpenGLViewport.cpp
+++ b/Engine/Source/Runtime/OpenGLDrv/Private/OpenGLViewport.cpp
@@ -92,9 +92,6 @@ void FOpenGLDynamicRHI::RHITick( float DeltaTime )
  *	Viewport functions.
  *=============================================================================*/
 
-// Ignore functions from RHIMethods.h when parsing documentation; Doxygen's preprocessor can't parse the declaration, so spews warnings for the definitions.
-#if !UE_BUILD_DOCS
-
 void FOpenGLDynamicRHI::RHIBeginDrawingViewport(FViewportRHIParamRef ViewportRHI, FTextureRHIParamRef RenderTarget)
 {
 	VERIFY_GL_SCOPE();
@@ -208,7 +205,7 @@ bool FOpenGLDynamicRHI::RHIIsDrawingViewport()
 	return DrawingViewport != NULL;
 }
 
-#endif
+
 
 FTexture2DRHIRef FOpenGLDynamicRHI::RHIGetViewportBackBuffer(FViewportRHIParamRef ViewportRHI)
 {
diff --git a/Engine/Source/Runtime/OpenGLDrv/Public/OpenGLDrv.h b/Engine/Source/Runtime/OpenGLDrv/Public/OpenGLDrv.h
--- a/Engine/Source/Runtime/OpenGLDrv/Public/OpenGLDrv.h
+++ b/Engine/Source/Runtime/OpenGLDrv/Public/OpenGLDrv.h
@@ -307,9 +307,177 @@ public:
 		return static_cast<typename TOpenGLResourceTraits<TRHIType>::TConcreteType*>(Resource);
 	}
 
-	#define DEFINE_RHIMETHOD(Type,Name,ParameterTypesAndNames,ParameterNames,ReturnStatement,NullImplementation) virtual Type RHI##Name ParameterTypesAndNames
-	#include "RHIMethods.h"
-	#undef DEFINE_RHIMETHOD
+	virtual void RHIGpuTimeBegin(uint32 Hash,bool bCompute) final override;
+	virtual void RHIGpuTimeEnd(uint32 Hash,bool bCompute) final override;
+	virtual FSamplerStateRHIRef RHICreateSamplerState(const FSamplerStateInitializerRHI& Initializer) final override;
+	virtual FRasterizerStateRHIRef RHICreateRasterizerState(const FRasterizerStateInitializerRHI& Initializer) final override;
+	virtual FDepthStencilStateRHIRef RHICreateDepthStencilState(const FDepthStencilStateInitializerRHI& Initializer) final override;
+	virtual FBlendStateRHIRef RHICreateBlendState(const FBlendStateInitializerRHI& Initializer) final override;
+	virtual FVertexDeclarationRHIRef RHICreateVertexDeclaration(const FVertexDeclarationElementList& Elements) final override;
+	virtual FPixelShaderRHIRef RHICreatePixelShader(const TArray<uint8>& Code) final override;
+	virtual FVertexShaderRHIRef RHICreateVertexShader(const TArray<uint8>& Code) final override;
+	virtual FHullShaderRHIRef RHICreateHullShader(const TArray<uint8>& Code) final override;
+	virtual FDomainShaderRHIRef RHICreateDomainShader(const TArray<uint8>& Code) final override;
+	virtual FGeometryShaderRHIRef RHICreateGeometryShader(const TArray<uint8>& Code) final override;
+	virtual FGeometryShaderRHIRef RHICreateGeometryShaderWithStreamOutput(const TArray<uint8>& Code, const FStreamOutElementList& ElementList, uint32 NumStrides, const uint32* Strides, int32 RasterizedStream) final override;
+	virtual FComputeShaderRHIRef RHICreateComputeShader(const TArray<uint8>& Code) final override;
+	virtual FBoundShaderStateRHIRef RHICreateBoundShaderState(FVertexDeclarationRHIParamRef VertexDeclaration, FVertexShaderRHIParamRef VertexShader, FHullShaderRHIParamRef HullShader, FDomainShaderRHIParamRef DomainShader, FPixelShaderRHIParamRef PixelShader, FGeometryShaderRHIParamRef GeometryShader) final override;
+	virtual FUniformBufferRHIRef RHICreateUniformBuffer(const void* Contents, const FRHIUniformBufferLayout& Layout, EUniformBufferUsage Usage) final override;
+	virtual FIndexBufferRHIRef RHICreateIndexBuffer(uint32 Stride, uint32 Size, uint32 InUsage, FRHIResourceCreateInfo& CreateInfo) final override;
+	virtual void* RHILockIndexBuffer(FIndexBufferRHIParamRef IndexBuffer, uint32 Offset, uint32 Size, EResourceLockMode LockMode) final override;
+	virtual void RHIUnlockIndexBuffer(FIndexBufferRHIParamRef IndexBuffer) final override;
+	virtual FVertexBufferRHIRef RHICreateVertexBuffer(uint32 Size, uint32 InUsage, FRHIResourceCreateInfo& CreateInfo) final override;
+	virtual void* RHILockVertexBuffer(FVertexBufferRHIParamRef VertexBuffer, uint32 Offset, uint32 SizeRHI, EResourceLockMode LockMode) final override;
+	virtual void RHIUnlockVertexBuffer(FVertexBufferRHIParamRef VertexBuffer) final override;
+	virtual void RHICopyVertexBuffer(FVertexBufferRHIParamRef SourceBuffer,FVertexBufferRHIParamRef DestBuffer) final override;
+	virtual FStructuredBufferRHIRef RHICreateStructuredBuffer(uint32 Stride, uint32 Size, uint32 InUsage, FRHIResourceCreateInfo& CreateInfo) final override;
+	virtual void* RHILockStructuredBuffer(FStructuredBufferRHIParamRef StructuredBuffer, uint32 Offset, uint32 SizeRHI, EResourceLockMode LockMode) final override;
+	virtual void RHIUnlockStructuredBuffer(FStructuredBufferRHIParamRef StructuredBuffer) final override;
+	virtual FUnorderedAccessViewRHIRef RHICreateUnorderedAccessView(FStructuredBufferRHIParamRef StructuredBuffer, bool bUseUAVCounter, bool bAppendBuffer) final override;
+	virtual FUnorderedAccessViewRHIRef RHICreateUnorderedAccessView(FTextureRHIParamRef Texture) final override;
+	virtual FUnorderedAccessViewRHIRef RHICreateUnorderedAccessView(FVertexBufferRHIParamRef VertexBuffer, uint8 Format) final override;
+	virtual FShaderResourceViewRHIRef RHICreateShaderResourceView(FStructuredBufferRHIParamRef StructuredBuffer) final override;
+	virtual FShaderResourceViewRHIRef RHICreateShaderResourceView(FVertexBufferRHIParamRef VertexBuffer, uint32 Stride, uint8 Format) final override;
+	virtual uint64 RHICalcTexture2DPlatformSize(uint32 SizeX, uint32 SizeY, uint8 Format, uint32 NumMips, uint32 NumSamples, uint32 Flags, uint32& OutAlign) final override;
+	virtual uint64 RHICalcTexture3DPlatformSize(uint32 SizeX, uint32 SizeY, uint32 SizeZ, uint8 Format, uint32 NumMips, uint32 Flags, uint32& OutAlign) final override;
+	virtual uint64 RHICalcTextureCubePlatformSize(uint32 Size, uint8 Format, uint32 NumMips, uint32 Flags, uint32& OutAlign) final override;
+	virtual void RHIGetTextureMemoryStats(FTextureMemoryStats& OutStats) final override;
+	virtual bool RHIGetTextureMemoryVisualizeData(FColor* TextureData,int32 SizeX,int32 SizeY,int32 Pitch,int32 PixelSize) final override;
+	virtual FTextureReferenceRHIRef RHICreateTextureReference(FLastRenderTimeContainer* LastRenderTime) final override;
+	virtual void RHIUpdateTextureReference(FTextureReferenceRHIParamRef TextureRef, FTextureRHIParamRef NewTexture) final override;
+	virtual FTexture2DRHIRef RHICreateTexture2D(uint32 SizeX, uint32 SizeY, uint8 Format, uint32 NumMips, uint32 NumSamples, uint32 Flags, FRHIResourceCreateInfo& CreateInfo) final override;
+	virtual FTexture2DRHIRef RHIAsyncCreateTexture2D(uint32 SizeX, uint32 SizeY, uint8 Format, uint32 NumMips, uint32 Flags, void** InitialMipData, uint32 NumInitialMips) final override;
+	virtual void RHICopySharedMips(FTexture2DRHIParamRef DestTexture2D, FTexture2DRHIParamRef SrcTexture2D) final override;
+	virtual FTexture2DArrayRHIRef RHICreateTexture2DArray(uint32 SizeX, uint32 SizeY, uint32 SizeZ, uint8 Format, uint32 NumMips, uint32 Flags, FRHIResourceCreateInfo& CreateInfo) final override;
+	virtual FTexture3DRHIRef RHICreateTexture3D(uint32 SizeX, uint32 SizeY, uint32 SizeZ, uint8 Format, uint32 NumMips, uint32 Flags, FRHIResourceCreateInfo& CreateInfo) final override;
+	virtual void RHIGetResourceInfo(FTextureRHIParamRef Ref, FRHIResourceInfo& OutInfo) final override;
+	virtual FShaderResourceViewRHIRef RHICreateShaderResourceView(FTexture2DRHIParamRef Texture2DRHI, uint8 MipLevel) final override;
+	virtual FShaderResourceViewRHIRef RHICreateShaderResourceView(FTexture2DRHIParamRef Texture2DRHI, uint8 MipLevel, uint8 NumMipLevels, uint8 Format) final override;
+	virtual void RHIGenerateMips(FTextureRHIParamRef Texture) final override;
+	virtual uint32 RHIComputeMemorySize(FTextureRHIParamRef TextureRHI) final override;
+	virtual FTexture2DRHIRef RHIAsyncReallocateTexture2D(FTexture2DRHIParamRef Texture2D, int32 NewMipCount, int32 NewSizeX, int32 NewSizeY, FThreadSafeCounter* RequestStatus) final override;
+	virtual ETextureReallocationStatus RHIFinalizeAsyncReallocateTexture2D(FTexture2DRHIParamRef Texture2D, bool bBlockUntilCompleted) final override;
+	virtual ETextureReallocationStatus RHICancelAsyncReallocateTexture2D(FTexture2DRHIParamRef Texture2D, bool bBlockUntilCompleted) final override;
+	virtual void* RHILockTexture2D(FTexture2DRHIParamRef Texture, uint32 MipIndex, EResourceLockMode LockMode, uint32& DestStride, bool bLockWithinMiptail) final override;
+	virtual void RHIUnlockTexture2D(FTexture2DRHIParamRef Texture, uint32 MipIndex, bool bLockWithinMiptail) final override;
+	virtual void* RHILockTexture2DArray(FTexture2DArrayRHIParamRef Texture, uint32 TextureIndex, uint32 MipIndex, EResourceLockMode LockMode, uint32& DestStride, bool bLockWithinMiptail) final override;
+	virtual void RHIUnlockTexture2DArray(FTexture2DArrayRHIParamRef Texture, uint32 TextureIndex, uint32 MipIndex, bool bLockWithinMiptail) final override;
+	virtual void RHIUpdateTexture2D(FTexture2DRHIParamRef Texture, uint32 MipIndex, const struct FUpdateTextureRegion2D& UpdateRegion, uint32 SourcePitch, const uint8* SourceData) final override;
+	virtual void RHIUpdateTexture3D(FTexture3DRHIParamRef Texture, uint32 MipIndex, const struct FUpdateTextureRegion3D& UpdateRegion, uint32 SourceRowPitch, uint32 SourceDepthPitch, const uint8* SourceData) final override;
+	virtual FTextureCubeRHIRef RHICreateTextureCube(uint32 Size, uint8 Format, uint32 NumMips, uint32 Flags, FRHIResourceCreateInfo& CreateInfo) final override;
+	virtual FTextureCubeRHIRef RHICreateTextureCubeArray(uint32 Size, uint32 ArraySize, uint8 Format, uint32 NumMips, uint32 Flags, FRHIResourceCreateInfo& CreateInfo) final override;
+	virtual void* RHILockTextureCubeFace(FTextureCubeRHIParamRef Texture, uint32 FaceIndex, uint32 ArrayIndex, uint32 MipIndex, EResourceLockMode LockMode, uint32& DestStride, bool bLockWithinMiptail) final override;
+	virtual void RHIUnlockTextureCubeFace(FTextureCubeRHIParamRef Texture, uint32 FaceIndex, uint32 ArrayIndex, uint32 MipIndex, bool bLockWithinMiptail) final override;
+	virtual void RHIBindDebugLabelName(FTextureRHIParamRef Texture, const TCHAR* Name) final override;
+	virtual void RHIReadSurfaceData(FTextureRHIParamRef Texture,FIntRect Rect,TArray<FColor>& OutData,FReadSurfaceDataFlags InFlags) final override;
+	virtual void RHIMapStagingSurface(FTextureRHIParamRef Texture,void*& OutData,int32& OutWidth,int32& OutHeight) final override;
+	virtual void RHIUnmapStagingSurface(FTextureRHIParamRef Texture) final override;
+	virtual void RHIReadSurfaceFloatData(FTextureRHIParamRef Texture,FIntRect Rect,TArray<FFloat16Color>& OutData,ECubeFace CubeFace,int32 ArrayIndex,int32 MipIndex) final override;
+	virtual void RHIRead3DSurfaceFloatData(FTextureRHIParamRef Texture,FIntRect Rect,FIntPoint ZMinMax,TArray<FFloat16Color>& OutData) final override;
+	virtual FRenderQueryRHIRef RHICreateRenderQuery(ERenderQueryType QueryType) final override;
+	virtual bool RHIGetRenderQueryResult(FRenderQueryRHIParamRef RenderQuery, uint64& OutResult, bool bWait) final override;
+	virtual FTexture2DRHIRef RHIGetViewportBackBuffer(FViewportRHIParamRef Viewport) final override;
+	virtual void RHIAdvanceFrameForGetViewportBackBuffer() final override;
+	virtual bool RHIIsDrawingViewport() final override;
+	virtual void RHIAcquireThreadOwnership() final override;
+	virtual void RHIReleaseThreadOwnership() final override;
+	virtual void RHIFlushResources() final override;
+	virtual uint32 RHIGetGPUFrameCycles() final override;
+	virtual FViewportRHIRef RHICreateViewport(void* WindowHandle, uint32 SizeX, uint32 SizeY, bool bIsFullscreen, EPixelFormat PreferredPixelFormat) final override;
+	virtual void RHIResizeViewport(FViewportRHIParamRef Viewport, uint32 SizeX, uint32 SizeY, bool bIsFullscreen) final override;
+	virtual void RHITick(float DeltaTime) final override;
+	virtual void RHISetStreamOutTargets(uint32 NumTargets,const FVertexBufferRHIParamRef* VertexBuffers,const uint32* Offsets) final override;
+	virtual void RHIDiscardRenderTargets(bool Depth,bool Stencil,uint32 ColorBitMask) final override;
+	virtual void RHIBlockUntilGPUIdle() final override;
+	virtual void RHISuspendRendering() final override;
+	virtual void RHIResumeRendering() final override;
+	virtual bool RHIIsRenderingSuspended() final override;
+	virtual bool RHIGetAvailableResolutions(FScreenResolutionArray& Resolutions, bool bIgnoreRefreshRate) final override;
+	virtual void RHIGetSupportedResolution(uint32& Width, uint32& Height) final override;
+	virtual void RHIVirtualTextureSetFirstMipInMemory(FTexture2DRHIParamRef Texture, uint32 FirstMip) final override;
+	virtual void RHIVirtualTextureSetFirstMipVisible(FTexture2DRHIParamRef Texture, uint32 FirstMip) final override;
+	virtual void RHIExecuteCommandList(FRHICommandList* CmdList) final override;
+	virtual void* RHIGetNativeDevice() final override;
+	virtual class IRHICommandContext* RHIGetDefaultContext() final override;
+	virtual class IRHICommandContextContainer* RHIGetCommandContextContainer() final override;
+
+
+
+	virtual void RHISetComputeShader(FComputeShaderRHIParamRef ComputeShader) final override;
+	virtual void RHIDispatchComputeShader(uint32 ThreadGroupCountX, uint32 ThreadGroupCountY, uint32 ThreadGroupCountZ) final override;
+	virtual void RHIDispatchIndirectComputeShader(FVertexBufferRHIParamRef ArgumentBuffer, uint32 ArgumentOffset) final override;
+	virtual void RHIAutomaticCacheFlushAfterComputeShader(bool bEnable) final override;
+	virtual void RHIFlushComputeShaderCache() final override;
+	virtual void RHISetMultipleViewports(uint32 Count, const FViewportBounds* Data) final override;
+	virtual void RHIClearUAV(FUnorderedAccessViewRHIParamRef UnorderedAccessViewRHI, const uint32* Values) final override;
+	virtual void RHICopyToResolveTarget(FTextureRHIParamRef SourceTexture, FTextureRHIParamRef DestTexture, bool bKeepOriginalSurface, const FResolveParams& ResolveParams) final override;
+	virtual void RHIBeginRenderQuery(FRenderQueryRHIParamRef RenderQuery) final override;
+	virtual void RHIEndRenderQuery(FRenderQueryRHIParamRef RenderQuery) final override;
+	virtual void RHIBeginOcclusionQueryBatch() final override;
+	virtual void RHIEndOcclusionQueryBatch() final override;
+	virtual void RHISubmitCommandsHint() final override;
+	virtual void RHIBeginDrawingViewport(FViewportRHIParamRef Viewport, FTextureRHIParamRef RenderTargetRHI) final override;
+	virtual void RHIEndDrawingViewport(FViewportRHIParamRef Viewport, bool bPresent, bool bLockToVsync) final override;
+	virtual void RHIBeginFrame() final override;
+	virtual void RHIEndFrame() final override;
+	virtual void RHIBeginScene() final override;
+	virtual void RHIEndScene() final override;
+	virtual void RHISetStreamSource(uint32 StreamIndex, FVertexBufferRHIParamRef VertexBuffer, uint32 Stride, uint32 Offset) final override;
+	virtual void RHISetRasterizerState(FRasterizerStateRHIParamRef NewState) final override;
+	virtual void RHISetViewport(uint32 MinX, uint32 MinY, float MinZ, uint32 MaxX, uint32 MaxY, float MaxZ) final override;
+	virtual void RHISetScissorRect(bool bEnable, uint32 MinX, uint32 MinY, uint32 MaxX, uint32 MaxY) final override;
+	virtual void RHISetBoundShaderState(FBoundShaderStateRHIParamRef BoundShaderState) final override;
+	virtual void RHISetShaderTexture(FVertexShaderRHIParamRef VertexShader, uint32 TextureIndex, FTextureRHIParamRef NewTexture) final override;
+	virtual void RHISetShaderTexture(FHullShaderRHIParamRef HullShader, uint32 TextureIndex, FTextureRHIParamRef NewTexture) final override;
+	virtual void RHISetShaderTexture(FDomainShaderRHIParamRef DomainShader, uint32 TextureIndex, FTextureRHIParamRef NewTexture) final override;
+	virtual void RHISetShaderTexture(FGeometryShaderRHIParamRef GeometryShader, uint32 TextureIndex, FTextureRHIParamRef NewTexture) final override;
+	virtual void RHISetShaderTexture(FPixelShaderRHIParamRef PixelShader, uint32 TextureIndex, FTextureRHIParamRef NewTexture) final override;
+	virtual void RHISetShaderTexture(FComputeShaderRHIParamRef PixelShader, uint32 TextureIndex, FTextureRHIParamRef NewTexture) final override;
+	virtual void RHISetShaderSampler(FComputeShaderRHIParamRef ComputeShader, uint32 SamplerIndex, FSamplerStateRHIParamRef NewState) final override;
+	virtual void RHISetShaderSampler(FVertexShaderRHIParamRef VertexShader, uint32 SamplerIndex, FSamplerStateRHIParamRef NewState) final override;
+	virtual void RHISetShaderSampler(FGeometryShaderRHIParamRef GeometryShader, uint32 SamplerIndex, FSamplerStateRHIParamRef NewState) final override;
+	virtual void RHISetShaderSampler(FDomainShaderRHIParamRef DomainShader, uint32 SamplerIndex, FSamplerStateRHIParamRef NewState) final override;
+	virtual void RHISetShaderSampler(FHullShaderRHIParamRef HullShader, uint32 SamplerIndex, FSamplerStateRHIParamRef NewState) final override;
+	virtual void RHISetShaderSampler(FPixelShaderRHIParamRef PixelShader, uint32 SamplerIndex, FSamplerStateRHIParamRef NewState) final override;
+	virtual void RHISetUAVParameter(FComputeShaderRHIParamRef ComputeShader, uint32 UAVIndex, FUnorderedAccessViewRHIParamRef UAV) final override;
+	virtual void RHISetUAVParameter(FComputeShaderRHIParamRef ComputeShader, uint32 UAVIndex, FUnorderedAccessViewRHIParamRef UAV, uint32 InitialCount) final override;
+	virtual void RHISetShaderResourceViewParameter(FPixelShaderRHIParamRef PixelShader, uint32 SamplerIndex, FShaderResourceViewRHIParamRef SRV) final override;
+	virtual void RHISetShaderResourceViewParameter(FVertexShaderRHIParamRef VertexShader, uint32 SamplerIndex, FShaderResourceViewRHIParamRef SRV) final override;
+	virtual void RHISetShaderResourceViewParameter(FComputeShaderRHIParamRef ComputeShader, uint32 SamplerIndex, FShaderResourceViewRHIParamRef SRV) final override;
+	virtual void RHISetShaderResourceViewParameter(FHullShaderRHIParamRef HullShader, uint32 SamplerIndex, FShaderResourceViewRHIParamRef SRV) final override;
+	virtual void RHISetShaderResourceViewParameter(FDomainShaderRHIParamRef DomainShader, uint32 SamplerIndex, FShaderResourceViewRHIParamRef SRV) final override;
+	virtual void RHISetShaderResourceViewParameter(FGeometryShaderRHIParamRef GeometryShader, uint32 SamplerIndex, FShaderResourceViewRHIParamRef SRV) final override;
+	virtual void RHISetShaderUniformBuffer(FVertexShaderRHIParamRef VertexShader, uint32 BufferIndex, FUniformBufferRHIParamRef Buffer) final override;
+	virtual void RHISetShaderUniformBuffer(FHullShaderRHIParamRef HullShader, uint32 BufferIndex, FUniformBufferRHIParamRef Buffer) final override;
+	virtual void RHISetShaderUniformBuffer(FDomainShaderRHIParamRef DomainShader, uint32 BufferIndex, FUniformBufferRHIParamRef Buffer) final override;
+	virtual void RHISetShaderUniformBuffer(FGeometryShaderRHIParamRef GeometryShader, uint32 BufferIndex, FUniformBufferRHIParamRef Buffer) final override;
+	virtual void RHISetShaderUniformBuffer(FPixelShaderRHIParamRef PixelShader, uint32 BufferIndex, FUniformBufferRHIParamRef Buffer) final override;
+	virtual void RHISetShaderUniformBuffer(FComputeShaderRHIParamRef ComputeShader, uint32 BufferIndex, FUniformBufferRHIParamRef Buffer) final override;
+	virtual void RHISetShaderParameter(FVertexShaderRHIParamRef VertexShader, uint32 BufferIndex, uint32 BaseIndex, uint32 NumBytes, const void* NewValue) final override;
+	virtual void RHISetShaderParameter(FPixelShaderRHIParamRef PixelShader, uint32 BufferIndex, uint32 BaseIndex, uint32 NumBytes, const void* NewValue) final override;
+	virtual void RHISetShaderParameter(FHullShaderRHIParamRef HullShader, uint32 BufferIndex, uint32 BaseIndex, uint32 NumBytes, const void* NewValue) final override;
+	virtual void RHISetShaderParameter(FDomainShaderRHIParamRef DomainShader, uint32 BufferIndex, uint32 BaseIndex, uint32 NumBytes, const void* NewValue) final override;
+	virtual void RHISetShaderParameter(FGeometryShaderRHIParamRef GeometryShader, uint32 BufferIndex, uint32 BaseIndex, uint32 NumBytes, const void* NewValue) final override;
+	virtual void RHISetShaderParameter(FComputeShaderRHIParamRef ComputeShader, uint32 BufferIndex, uint32 BaseIndex, uint32 NumBytes, const void* NewValue) final override;
+	virtual void RHISetDepthStencilState(FDepthStencilStateRHIParamRef NewState, uint32 StencilRef) final override;
+	virtual void RHISetBlendState(FBlendStateRHIParamRef NewState, const FLinearColor& BlendFactor) final override;
+	virtual void RHISetRenderTargets(uint32 NumSimultaneousRenderTargets, const FRHIRenderTargetView* NewRenderTargets, const FRHIDepthRenderTargetView* NewDepthStencilTarget, uint32 NumUAVs, const FUnorderedAccessViewRHIParamRef* UAVs) final override;
+	virtual void RHISetRenderTargetsAndClear(const FRHISetRenderTargetsInfo& RenderTargetsInfo) final override;
+	virtual void RHIBindClearMRTValues(bool bClearColor, int32 NumClearColors, const FLinearColor* ColorArray, bool bClearDepth, float Depth, bool bClearStencil, uint32 Stencil) final override;
+	virtual void RHIDrawPrimitive(uint32 PrimitiveType, uint32 BaseVertexIndex, uint32 NumPrimitives, uint32 NumInstances) final override;
+	virtual void RHIDrawPrimitiveIndirect(uint32 PrimitiveType, FVertexBufferRHIParamRef ArgumentBuffer, uint32 ArgumentOffset) final override;
+	virtual void RHIDrawIndexedIndirect(FIndexBufferRHIParamRef IndexBufferRHI, uint32 PrimitiveType, FStructuredBufferRHIParamRef ArgumentsBufferRHI, int32 DrawArgumentsIndex, uint32 NumInstances) final override;
+	virtual void RHIDrawIndexedPrimitive(FIndexBufferRHIParamRef IndexBuffer, uint32 PrimitiveType, int32 BaseVertexIndex, uint32 FirstInstance, uint32 NumVertices, uint32 StartIndex, uint32 NumPrimitives, uint32 NumInstances) final override;
+	virtual void RHIDrawIndexedPrimitiveIndirect(uint32 PrimitiveType, FIndexBufferRHIParamRef IndexBuffer, FVertexBufferRHIParamRef ArgumentBuffer, uint32 ArgumentOffset) final override;
+	virtual void RHIBeginDrawPrimitiveUP(uint32 PrimitiveType, uint32 NumPrimitives, uint32 NumVertices, uint32 VertexDataStride, void*& OutVertexData) final override;
+	virtual void RHIEndDrawPrimitiveUP() final override;
+	virtual void RHIBeginDrawIndexedPrimitiveUP(uint32 PrimitiveType, uint32 NumPrimitives, uint32 NumVertices, uint32 VertexDataStride, void*& OutVertexData, uint32 MinVertexIndex, uint32 NumIndices, uint32 IndexDataStride, void*& OutIndexData) final override;
+	virtual void RHIEndDrawIndexedPrimitiveUP() final override;
+	virtual void RHIClear(bool bClearColor, const FLinearColor& Color, bool bClearDepth, float Depth, bool bClearStencil, uint32 Stencil, FIntRect ExcludeRect) final override;
+	virtual void RHIClearMRT(bool bClearColor, int32 NumClearColors, const FLinearColor* ColorArray, bool bClearDepth, float Depth, bool bClearStencil, uint32 Stencil, FIntRect ExcludeRect) final override;
+	virtual void RHIEnableDepthBoundsTest(bool bEnable, float MinDepth, float MaxDepth) final override;
+	virtual void RHIPushEvent(const TCHAR* Name) final override;
+	virtual void RHIPopEvent() final override;
 
 	void Cleanup();
 
diff --git a/Engine/Source/Runtime/RHI/Private/Android/AndroidDynamicRHI.cpp b/Engine/Source/Runtime/RHI/Private/Android/AndroidDynamicRHI.cpp
--- a/Engine/Source/Runtime/RHI/Private/Android/AndroidDynamicRHI.cpp
+++ b/Engine/Source/Runtime/RHI/Private/Android/AndroidDynamicRHI.cpp
@@ -3,8 +3,6 @@
 #include "RHI.h"
 #include "ModuleManager.h"
 
-#if USE_DYNAMIC_RHI
-
 FDynamicRHI* PlatformCreateDynamicRHI()
 {
 	FDynamicRHI* DynamicRHI = NULL;
@@ -29,5 +27,3 @@ FDynamicRHI* PlatformCreateDynamicRHI()
 
 	return DynamicRHI;
 }
-
-#endif // USE_DYNAMIC_RHI
diff --git a/Engine/Source/Runtime/RHI/Private/Apple/AppleDynamicRHI.cpp b/Engine/Source/Runtime/RHI/Private/Apple/AppleDynamicRHI.cpp
--- a/Engine/Source/Runtime/RHI/Private/Apple/AppleDynamicRHI.cpp
+++ b/Engine/Source/Runtime/RHI/Private/Apple/AppleDynamicRHI.cpp
@@ -3,8 +3,6 @@
 #include "RHI.h"
 #include "ModuleManager.h"
 
-#if USE_DYNAMIC_RHI
-
 FDynamicRHI* PlatformCreateDynamicRHI()
 {
 	FDynamicRHI* DynamicRHI = NULL;
@@ -26,5 +24,3 @@ FDynamicRHI* PlatformCreateDynamicRHI()
 	DynamicRHI = DynamicRHIModule->CreateRHI();
 	return DynamicRHI;
 }
-
-#endif // USE_DYNAMIC_RHI
diff --git a/Engine/Source/Runtime/RHI/Private/DynamicRHI.cpp b/Engine/Source/Runtime/RHI/Private/DynamicRHI.cpp
--- a/Engine/Source/Runtime/RHI/Private/DynamicRHI.cpp
+++ b/Engine/Source/Runtime/RHI/Private/DynamicRHI.cpp
@@ -11,8 +11,6 @@
 	#define PLATFORM_ALLOW_NULL_RHI		0
 #endif
 
-#if USE_DYNAMIC_RHI
-
 // Globals.
 FDynamicRHI* GDynamicRHI = NULL;
 
@@ -76,44 +74,4 @@ void RHIExit()
 	}
 }
 
-#define DEFINE_RHIMETHOD_CMDLIST(Type,Name,ParameterTypesAndNames,ParameterNames,ReturnStatement,NullImplementation)
-
-#define DEFINE_RHIMETHOD_GLOBAL(Type,Name,ParameterTypesAndNames,ParameterNames,ReturnStatement,NullImplementation) \
-	RHI_API Type Name##_Internal ParameterTypesAndNames \
-	{ \
-		check(GDynamicRHI); \
-		ReturnStatement GDynamicRHI->RHI##Name ParameterNames; \
-	}
-#define DEFINE_RHIMETHOD_GLOBALTHREADSAFE(Type,Name,ParameterTypesAndNames,ParameterNames,ReturnStatement,NullImplementation) \
-	RHI_API Type RHI##Name ParameterTypesAndNames \
-	{ \
-		check(GDynamicRHI); \
-		ReturnStatement GDynamicRHI->RHI##Name ParameterNames; \
-	}
-#define DEFINE_RHIMETHOD_GLOBALFLUSH(Type,Name,ParameterTypesAndNames,ParameterNames,ReturnStatement,NullImplementation) \
-	RHI_API Type Name##_Internal ParameterTypesAndNames \
-	{ \
-		check(GDynamicRHI); \
-		ReturnStatement GDynamicRHI->RHI##Name ParameterNames; \
-	}
-#define DEFINE_RHIMETHOD(Type,Name,ParameterTypesAndNames,ParameterNames,ReturnStatement,NullImplementation) \
-	RHI_API Type Name##_Internal ParameterTypesAndNames \
-	{ \
-		check(GDynamicRHI); \
-		ReturnStatement GDynamicRHI->RHI##Name ParameterNames; \
-	}
-#include "RHIMethods.h"
-#undef DEFINE_RHIMETHOD
-#undef DEFINE_RHIMETHOD_CMDLIST
-#undef DEFINE_RHIMETHOD_GLOBAL
-#undef DEFINE_RHIMETHOD_GLOBALFLUSH
-#undef DEFINE_RHIMETHOD_GLOBALTHREADSAFE
-
-
-#else
-
-// Suppress linker warning "warning LNK4221: no public symbols found; archive member will be inaccessible"
-int32 DynamicRHILinkerHelper;
-
-#endif // USE_DYNAMIC_RHI
 
diff --git a/Engine/Source/Runtime/RHI/Private/HTML5/HTML5DynamicRHI.cpp b/Engine/Source/Runtime/RHI/Private/HTML5/HTML5DynamicRHI.cpp
--- a/Engine/Source/Runtime/RHI/Private/HTML5/HTML5DynamicRHI.cpp
+++ b/Engine/Source/Runtime/RHI/Private/HTML5/HTML5DynamicRHI.cpp
@@ -3,8 +3,6 @@
 #include "RHI.h"
 #include "ModuleManager.h"
 
-#if USE_DYNAMIC_RHI
-
 FDynamicRHI* PlatformCreateDynamicRHI()
 {
 	FDynamicRHI* DynamicRHI = NULL;
@@ -15,5 +13,3 @@ FDynamicRHI* PlatformCreateDynamicRHI()
 	DynamicRHI = DynamicRHIModule->CreateRHI();
 	return DynamicRHI;
 }
-
-#endif // USE_DYNAMIC_RHI
diff --git a/Engine/Source/Runtime/RHI/Private/Linux/LinuxDynamicRHI.cpp b/Engine/Source/Runtime/RHI/Private/Linux/LinuxDynamicRHI.cpp
--- a/Engine/Source/Runtime/RHI/Private/Linux/LinuxDynamicRHI.cpp
+++ b/Engine/Source/Runtime/RHI/Private/Linux/LinuxDynamicRHI.cpp
@@ -3,8 +3,6 @@
 #include "RHI.h"
 #include "ModuleManager.h"
 
-#if USE_DYNAMIC_RHI
-
 FDynamicRHI* PlatformCreateDynamicRHI()
 {
 	FDynamicRHI* DynamicRHI = NULL;
@@ -26,5 +24,3 @@ FDynamicRHI* PlatformCreateDynamicRHI()
 
 	return DynamicRHI;
 }
-
-#endif // USE_DYNAMIC_RHI
diff --git a/Engine/Source/Runtime/RHI/Private/RHICommandList.cpp b/Engine/Source/Runtime/RHI/Private/RHICommandList.cpp
--- a/Engine/Source/Runtime/RHI/Private/RHICommandList.cpp
+++ b/Engine/Source/Runtime/RHI/Private/RHICommandList.cpp
@@ -533,7 +533,6 @@ void FRHICommandListBase::Reset()
 	NumCommands = 0;
 	Root = nullptr;
 	CommandLink = &Root;
-	static_assert(USE_DYNAMIC_RHI, "static RHIs are not supported.");
 	Context = GDynamicRHI ? RHIGetDefaultContext() : nullptr;
 #if USE_RHICOMMAND_STATE_REDUCTION
 	StateCache = nullptr;
diff --git a/Engine/Source/Runtime/RHI/Private/WinRT/WinRTDynamicRHI.cpp b/Engine/Source/Runtime/RHI/Private/WinRT/WinRTDynamicRHI.cpp
--- a/Engine/Source/Runtime/RHI/Private/WinRT/WinRTDynamicRHI.cpp
+++ b/Engine/Source/Runtime/RHI/Private/WinRT/WinRTDynamicRHI.cpp
@@ -3,8 +3,6 @@
 #include "RHI.h"
 #include "ModuleManager.h"
 
-#if USE_DYNAMIC_RHI
-
 FDynamicRHI* PlatformCreateDynamicRHI()
 {
 	FDynamicRHI* DynamicRHI = NULL;
@@ -25,5 +23,3 @@ FDynamicRHI* PlatformCreateDynamicRHI()
 
 	return DynamicRHI;
 }
-
-#endif // USE_DYNAMIC_RHI
diff --git a/Engine/Source/Runtime/RHI/Private/Windows/WindowsDynamicRHI.cpp b/Engine/Source/Runtime/RHI/Private/Windows/WindowsDynamicRHI.cpp
--- a/Engine/Source/Runtime/RHI/Private/Windows/WindowsDynamicRHI.cpp
+++ b/Engine/Source/Runtime/RHI/Private/Windows/WindowsDynamicRHI.cpp
@@ -3,8 +3,6 @@
 #include "RHI.h"
 #include "ModuleManager.h"
 
-#if USE_DYNAMIC_RHI
-
 FDynamicRHI* PlatformCreateDynamicRHI()
 {
 	FDynamicRHI* DynamicRHI = NULL;
@@ -57,5 +55,3 @@ FDynamicRHI* PlatformCreateDynamicRHI()
 
 	return DynamicRHI;
 }
-
-#endif // USE_DYNAMIC_RHI
diff --git a/Engine/Source/Runtime/RHI/Public/DynamicRHI.h b/Engine/Source/Runtime/RHI/Public/DynamicRHI.h
--- a/Engine/Source/Runtime/RHI/Public/DynamicRHI.h
+++ b/Engine/Source/Runtime/RHI/Public/DynamicRHI.h
@@ -4,19 +4,14 @@
 	DynamicRHI.h: Dynamically bound Render Hardware Interface definitions.
 =============================================================================*/
 
-#ifndef __DYNAMICRHI_H__
-#define __DYNAMICRHI_H__
+#pragma once
 
 #include "ModuleInterface.h"
 
-#if USE_DYNAMIC_RHI
-
 //
 // Statically bound RHI resource reference type definitions for the dynamically bound RHI.
 //
 
-#undef DEFINE_DYNAMICRHI_REFERENCE_TYPE
-
 /** The interface which is implemented by the dynamically bound RHI. */
 class FDynamicRHI
 {
@@ -27,46 +22,1039 @@ public:
 
 	/** Initializes the RHI; separate from IDynamicRHIModule::CreateRHI so that GDynamicRHI is set when it is called. */
 	virtual void Init() = 0;
+
 	/** Shutdown the RHI; handle shutdown and resource destruction before the RHI's actual destructor is called (so that all resources of the RHI are still available for shutdown). */
 	virtual void Shutdown() = 0;
 
-	#define DEFINE_RHIMETHOD(Type,Name,ParameterTypesAndNames,ParameterNames,ReturnStatement,NullImplementation) virtual Type RHI##Name ParameterTypesAndNames = 0
-	#define DEFINE_RHIMETHOD_CMDLIST(Type,Name,ParameterTypesAndNames,ParameterNames,ReturnStatement,NullImplementation)
-	#define DEFINE_RHIMETHOD_GLOBAL(Type,Name,ParameterTypesAndNames,ParameterNames,ReturnStatement,NullImplementation) virtual Type RHI##Name ParameterTypesAndNames = 0
-	#define DEFINE_RHIMETHOD_GLOBALFLUSH(Type,Name,ParameterTypesAndNames,ParameterNames,ReturnStatement,NullImplementation) virtual Type RHI##Name ParameterTypesAndNames = 0
-	#define DEFINE_RHIMETHOD_GLOBALTHREADSAFE(Type,Name,ParameterTypesAndNames,ParameterNames,ReturnStatement,NullImplementation) virtual Type RHI##Name ParameterTypesAndNames = 0
-	#include "RHIMethods.h"
-	#undef DEFINE_RHIMETHOD
-	#undef DEFINE_RHIMETHOD_CMDLIST
-	#undef DEFINE_RHIMETHOD_GLOBAL
-	#undef DEFINE_RHIMETHOD_GLOBALFLUSH
-	#undef DEFINE_RHIMETHOD_GLOBALTHREADSAFE
+
+	/////// RHI Methods
+
+	//
+	// Interface for run-time gpu perf counters.
+	//
+	// FlushType: Flush Immediate
+	virtual void RHIGpuTimeBegin(uint32 Hash,bool bCompute) = 0;
+
+	// FlushType: Flush Immediate
+	virtual void RHIGpuTimeEnd(uint32 Hash,bool bCompute) = 0;
+
+	// FlushType: Thread safe
+	virtual FSamplerStateRHIRef RHICreateSamplerState(const FSamplerStateInitializerRHI& Initializer) = 0;
+
+	// FlushType: Thread safe
+	virtual FRasterizerStateRHIRef RHICreateRasterizerState(const FRasterizerStateInitializerRHI& Initializer) = 0;
+
+	// FlushType: Thread safe
+	virtual FDepthStencilStateRHIRef RHICreateDepthStencilState(const FDepthStencilStateInitializerRHI& Initializer) = 0;
+
+	// FlushType: Thread safe
+	virtual FBlendStateRHIRef RHICreateBlendState(const FBlendStateInitializerRHI& Initializer) = 0;
+
+	// FlushType: Wait RHI Thread
+	virtual FVertexDeclarationRHIRef RHICreateVertexDeclaration(const FVertexDeclarationElementList& Elements) = 0;
+
+	// FlushType: Wait RHI Thread
+	virtual FPixelShaderRHIRef RHICreatePixelShader(const TArray<uint8>& Code) = 0;
+
+	// FlushType: Wait RHI Thread
+	virtual FVertexShaderRHIRef RHICreateVertexShader(const TArray<uint8>& Code) = 0;
+
+	// FlushType: Wait RHI Thread
+	virtual FHullShaderRHIRef RHICreateHullShader(const TArray<uint8>& Code) = 0;
+
+	// FlushType: Wait RHI Thread
+	virtual FDomainShaderRHIRef RHICreateDomainShader(const TArray<uint8>& Code) = 0;
+
+	// FlushType: Wait RHI Thread
+	virtual FGeometryShaderRHIRef RHICreateGeometryShader(const TArray<uint8>& Code) = 0;
+
+	/** Creates a geometry shader with stream output ability, defined by ElementList. */
+	// FlushType: Wait RHI Thread
+	virtual FGeometryShaderRHIRef RHICreateGeometryShaderWithStreamOutput(const TArray<uint8>& Code, const FStreamOutElementList& ElementList, uint32 NumStrides, const uint32* Strides, int32 RasterizedStream) = 0;
+
+	// FlushType: Wait RHI Thread
+	virtual FComputeShaderRHIRef RHICreateComputeShader(const TArray<uint8>& Code) = 0;
+
+	/**
+	 * Creates a bound shader state instance which encapsulates a decl, vertex shader, hull shader, domain shader and pixel shader
+	 * CAUTION: Even though this is marked as threadsafe, it is only valid to call from the render thread or the RHI thread. It need not be threadsafe unless the RHI support parallel translation.
+	 * CAUTION: Platforms that support RHIThread but don't actually have a threadsafe implementation must flush internally with FRHICommandListExecutor::GetImmediateCommandList().ImmediateFlush(EImmediateFlushType::WaitForRHIThread); when the call is from the render thread
+	 * @param VertexDeclaration - existing vertex decl
+	 * @param VertexShader - existing vertex shader
+	 * @param HullShader - existing hull shader
+	 * @param DomainShader - existing domain shader
+	 * @param GeometryShader - existing geometry shader
+	 * @param PixelShader - existing pixel shader
+	 */
+	// FlushType: Thread safe, but varies depending on the RHI
+	virtual FBoundShaderStateRHIRef RHICreateBoundShaderState(FVertexDeclarationRHIParamRef VertexDeclaration, FVertexShaderRHIParamRef VertexShader, FHullShaderRHIParamRef HullShader, FDomainShaderRHIParamRef DomainShader, FPixelShaderRHIParamRef PixelShader, FGeometryShaderRHIParamRef GeometryShader) = 0;
+
+	/**
+	 * Creates a uniform buffer.  The contents of the uniform buffer are provided in a parameter, and are immutable.
+	 * CAUTION: Even though this is marked as threadsafe, it is only valid to call from the render thread or the RHI thread. Thus is need not be threadsafe on platforms that do not support or aren't using an RHIThread
+	 * @param Contents - A pointer to a memory block of size NumBytes that is copied into the new uniform buffer.
+	 * @param NumBytes - The number of bytes the uniform buffer should contain.
+	 * @return The new uniform buffer.
+	 */
+	// FlushType: Thread safe, but varies depending on the RHI
+	virtual FUniformBufferRHIRef RHICreateUniformBuffer(const void* Contents, const FRHIUniformBufferLayout& Layout, EUniformBufferUsage Usage) = 0;
+
+	// FlushType: Wait RHI Thread
+	virtual FIndexBufferRHIRef RHICreateIndexBuffer(uint32 Stride, uint32 Size, uint32 InUsage, FRHIResourceCreateInfo& CreateInfo) = 0;
+
+	// FlushType: Flush RHI Thread
+	virtual void* RHILockIndexBuffer(FIndexBufferRHIParamRef IndexBuffer, uint32 Offset, uint32 Size, EResourceLockMode LockMode) = 0;
+
+	// FlushType: Flush RHI Thread
+	virtual void RHIUnlockIndexBuffer(FIndexBufferRHIParamRef IndexBuffer) = 0;
+
+	/**
+	 * @param ResourceArray - An optional pointer to a resource array containing the resource's data.
+	 */
+	// FlushType: Wait RHI Thread
+	virtual FVertexBufferRHIRef RHICreateVertexBuffer(uint32 Size, uint32 InUsage, FRHIResourceCreateInfo& CreateInfo) = 0;
+
+	// FlushType: Flush RHI Thread
+	virtual void* RHILockVertexBuffer(FVertexBufferRHIParamRef VertexBuffer, uint32 Offset, uint32 SizeRHI, EResourceLockMode LockMode) = 0;
+
+	// FlushType: Flush RHI Thread
+	virtual void RHIUnlockVertexBuffer(FVertexBufferRHIParamRef VertexBuffer) = 0;
+
+	/** Copies the contents of one vertex buffer to another vertex buffer.  They must have identical sizes. */
+	// FlushType: Flush Immediate (seems dangerous)
+	virtual void RHICopyVertexBuffer(FVertexBufferRHIParamRef SourceBuffer,FVertexBufferRHIParamRef DestBuffer) = 0;
+
+	/**
+	 * @param ResourceArray - An optional pointer to a resource array containing the resource's data.
+	 */
+	// FlushType: Wait RHI Thread
+	virtual FStructuredBufferRHIRef RHICreateStructuredBuffer(uint32 Stride, uint32 Size, uint32 InUsage, FRHIResourceCreateInfo& CreateInfo) = 0;
+
+	// FlushType: Flush RHI Thread
+	virtual void* RHILockStructuredBuffer(FStructuredBufferRHIParamRef StructuredBuffer, uint32 Offset, uint32 SizeRHI, EResourceLockMode LockMode) = 0;
+
+	// FlushType: Flush RHI Thread
+	virtual void RHIUnlockStructuredBuffer(FStructuredBufferRHIParamRef StructuredBuffer) = 0;
+
+	/** Creates an unordered access view of the given structured buffer. */
+	// FlushType: Wait RHI Thread
+	virtual FUnorderedAccessViewRHIRef RHICreateUnorderedAccessView(FStructuredBufferRHIParamRef StructuredBuffer, bool bUseUAVCounter, bool bAppendBuffer) = 0;
+
+	/** Creates an unordered access view of the given texture. */
+	// FlushType: Wait RHI Thread
+	virtual FUnorderedAccessViewRHIRef RHICreateUnorderedAccessView(FTextureRHIParamRef Texture) = 0;
+
+	/** Creates an unordered access view of the given vertex buffer. */
+	// FlushType: Wait RHI Thread
+	virtual FUnorderedAccessViewRHIRef RHICreateUnorderedAccessView(FVertexBufferRHIParamRef VertexBuffer, uint8 Format) = 0;
+
+	/** Creates a shader resource view of the given structured buffer. */
+	// FlushType: Wait RHI Thread
+	virtual FShaderResourceViewRHIRef RHICreateShaderResourceView(FStructuredBufferRHIParamRef StructuredBuffer) = 0;
+
+	/** Creates a shader resource view of the given vertex buffer. */
+	// FlushType: Wait RHI Thread
+	virtual FShaderResourceViewRHIRef RHICreateShaderResourceView(FVertexBufferRHIParamRef VertexBuffer, uint32 Stride, uint8 Format) = 0;
+
+	/**
+	 * Computes the total size of a 2D texture with the specified parameters.
+	 *
+	 * @param SizeX - width of the texture to compute
+	 * @param SizeY - height of the texture to compute
+	 * @param Format - EPixelFormat texture format
+	 * @param NumMips - number of mips to compute or 0 for full mip pyramid
+	 * @param NumSamples - number of MSAA samples, usually 1
+	 * @param Flags - ETextureCreateFlags creation flags
+	 * @param OutAlign - Alignment required for this texture.  Output parameter.
+	 */
+	// FlushType: Thread safe
+	virtual uint64 RHICalcTexture2DPlatformSize(uint32 SizeX, uint32 SizeY, uint8 Format, uint32 NumMips, uint32 NumSamples, uint32 Flags, uint32& OutAlign) = 0;
+
+	/**
+	* Computes the total size of a 3D texture with the specified parameters.
+	* @param SizeX - width of the texture to create
+	* @param SizeY - height of the texture to create
+	* @param SizeZ - depth of the texture to create
+	* @param Format - EPixelFormat texture format
+	* @param NumMips - number of mips to generate or 0 for full mip pyramid
+	* @param Flags - ETextureCreateFlags creation flags
+	* @param OutAlign - Alignment required for this texture.  Output parameter.
+	*/
+	// FlushType: Thread safe
+	virtual uint64 RHICalcTexture3DPlatformSize(uint32 SizeX, uint32 SizeY, uint32 SizeZ, uint8 Format, uint32 NumMips, uint32 Flags, uint32& OutAlign) = 0;
+
+	/**
+	* Computes the total size of a cube texture with the specified parameters.
+	* @param Size - width/height of the texture to create
+	* @param Format - EPixelFormat texture format
+	* @param NumMips - number of mips to generate or 0 for full mip pyramid
+	* @param Flags - ETextureCreateFlags creation flags
+	* @param OutAlign - Alignment required for this texture.  Output parameter.
+	*/
+	// FlushType: Thread safe
+	virtual uint64 RHICalcTextureCubePlatformSize(uint32 Size, uint8 Format, uint32 NumMips, uint32 Flags, uint32& OutAlign) = 0;
+
+	/**
+	 * Retrieves texture memory stats.
+	 * safe to call on the main thread
+	 */
+	// FlushType: Thread safe
+	virtual void RHIGetTextureMemoryStats(FTextureMemoryStats& OutStats) = 0;
+
+	/**
+	 * Fills a texture with to visualize the texture pool memory.
+	 *
+	 * @param	TextureData		Start address
+	 * @param	SizeX			Number of pixels along X
+	 * @param	SizeY			Number of pixels along Y
+	 * @param	Pitch			Number of bytes between each row
+	 * @param	PixelSize		Number of bytes each pixel represents
+	 *
+	 * @return true if successful, false otherwise
+	 */
+	// FlushType: Flush Immediate
+	virtual bool RHIGetTextureMemoryVisualizeData(FColor* TextureData,int32 SizeX,int32 SizeY,int32 Pitch,int32 PixelSize) = 0;
+
+	// FlushType: Wait RHI Thread
+	virtual FTextureReferenceRHIRef RHICreateTextureReference(FLastRenderTimeContainer* LastRenderTime) = 0;
+
+	// FlushType: Flush RHI Thread
+	virtual void RHIUpdateTextureReference(FTextureReferenceRHIParamRef TextureRef, FTextureRHIParamRef NewTexture) = 0;
+
+	/**
+	* Creates a 2D RHI texture resource
+	* @param SizeX - width of the texture to create
+	* @param SizeY - height of the texture to create
+	* @param Format - EPixelFormat texture format
+	* @param NumMips - number of mips to generate or 0 for full mip pyramid
+	* @param NumSamples - number of MSAA samples, usually 1
+	* @param Flags - ETextureCreateFlags creation flags
+	*/
+	// FlushType: Wait RHI Thread
+	virtual FTexture2DRHIRef RHICreateTexture2D(uint32 SizeX, uint32 SizeY, uint8 Format, uint32 NumMips, uint32 NumSamples, uint32 Flags, FRHIResourceCreateInfo& CreateInfo) = 0;
+
+	/**
+	 * Thread-safe function that can be used to create a texture outside of the
+	 * rendering thread. This function can ONLY be called if GRHISupportsAsyncTextureCreation
+	 * is true.
+	 * @param SizeX - width of the texture to create
+	 * @param SizeY - height of the texture to create
+	 * @param Format - EPixelFormat texture format
+	 * @param NumMips - number of mips to generate or 0 for full mip pyramid
+	 * @param Flags - ETextureCreateFlags creation flags
+	 * @param InitialMipData - pointers to mip data with which to create the texture
+	 * @param NumInitialMips - how many mips are provided in InitialMipData
+	 * @returns a reference to a 2D texture resource
+	 */
+	// FlushType: Wait RHI Thread
+	virtual FTexture2DRHIRef RHIAsyncCreateTexture2D(uint32 SizeX, uint32 SizeY, uint8 Format, uint32 NumMips, uint32 Flags, void** InitialMipData, uint32 NumInitialMips) = 0;
+
+	/**
+	 * Copies shared mip levels from one texture to another. The textures must have
+	 * full mip chains, share the same format, and have the same aspect ratio. This
+	 * copy will not cause synchronization with the GPU.
+	 * @param DestTexture2D - destination texture
+	 * @param SrcTexture2D - source texture
+	 */
+	// FlushType: Flush RHI Thread
+	virtual void RHICopySharedMips(FTexture2DRHIParamRef DestTexture2D, FTexture2DRHIParamRef SrcTexture2D) = 0;
+
+	/**
+	* Creates a Array RHI texture resource
+	* @param SizeX - width of the texture to create
+	* @param SizeY - height of the texture to create
+	* @param SizeZ - depth of the texture to create
+	* @param Format - EPixelFormat texture format
+	* @param NumMips - number of mips to generate or 0 for full mip pyramid
+	* @param Flags - ETextureCreateFlags creation flags
+	*/
+	// FlushType: Wait RHI Thread
+	virtual FTexture2DArrayRHIRef RHICreateTexture2DArray(uint32 SizeX, uint32 SizeY, uint32 SizeZ, uint8 Format, uint32 NumMips, uint32 Flags, FRHIResourceCreateInfo& CreateInfo) = 0;
+
+	/**
+	* Creates a 3d RHI texture resource
+	* @param SizeX - width of the texture to create
+	* @param SizeY - height of the texture to create
+	* @param SizeZ - depth of the texture to create
+	* @param Format - EPixelFormat texture format
+	* @param NumMips - number of mips to generate or 0 for full mip pyramid
+	* @param Flags - ETextureCreateFlags creation flags
+	*/
+	// FlushType: Wait RHI Thread
+	virtual FTexture3DRHIRef RHICreateTexture3D(uint32 SizeX, uint32 SizeY, uint32 SizeZ, uint8 Format, uint32 NumMips, uint32 Flags, FRHIResourceCreateInfo& CreateInfo) = 0;
+
+	/**
+	 * @param Ref may be 0
+	 */
+	// FlushType: Thread safe
+	virtual void RHIGetResourceInfo(FTextureRHIParamRef Ref, FRHIResourceInfo& OutInfo) = 0;
+
+	/**
+	* Creates a shader resource view for a 2d texture, viewing only a single
+	* mip level. Useful when rendering to one mip while sampling from another.
+	*/
+	// FlushType: Wait RHI Thread
+	virtual FShaderResourceViewRHIRef RHICreateShaderResourceView(FTexture2DRHIParamRef Texture2DRHI, uint8 MipLevel) = 0;
+
+	/**FRHIResourceInfo
+	* Creates a shader resource view for a 2d texture, with a different
+	* format from the original.  Useful when sampling stencil.
+	*/
+	// FlushType: Wait RHI Thread
+	virtual FShaderResourceViewRHIRef RHICreateShaderResourceView(FTexture2DRHIParamRef Texture2DRHI, uint8 MipLevel, uint8 NumMipLevels, uint8 Format) = 0;
+
+	/**
+	* Generates mip maps for a texture.
+	*/
+	// FlushType: Flush Immediate
+	virtual void RHIGenerateMips(FTextureRHIParamRef Texture) = 0;
+
+	/**
+	 * Computes the size in memory required by a given texture.
+	 *
+	 * @param	TextureRHI		- Texture we want to know the size of, 0 is safely ignored
+	 * @return					- Size in Bytes
+	 */
+	// FlushType: Thread safe
+	virtual uint32 RHIComputeMemorySize(FTextureRHIParamRef TextureRHI) = 0;
+
+	/**
+	 * Starts an asynchronous texture reallocation. It may complete immediately if the reallocation
+	 * could be performed without any reshuffling of texture memory, or if there isn't enough memory.
+	 * The specified status counter will be decremented by 1 when the reallocation is complete (success or failure).
+	 *
+	 * Returns a new reference to the texture, which will represent the new mip count when the reallocation is complete.
+	 * RHIFinalizeAsyncReallocateTexture2D() must be called to complete the reallocation.
+	 *
+	 * @param Texture2D		- Texture to reallocate
+	 * @param NewMipCount	- New number of mip-levels
+	 * @param NewSizeX		- New width, in pixels
+	 * @param NewSizeY		- New height, in pixels
+	 * @param RequestStatus	- Will be decremented by 1 when the reallocation is complete (success or failure).
+	 * @return				- New reference to the texture, or an invalid reference upon failure
+	 */
+	// FlushType: Flush RHI Thread
+	virtual FTexture2DRHIRef RHIAsyncReallocateTexture2D(FTexture2DRHIParamRef Texture2D, int32 NewMipCount, int32 NewSizeX, int32 NewSizeY, FThreadSafeCounter* RequestStatus) = 0;
+
+	/**
+	 * Finalizes an async reallocation request.
+	 * If bBlockUntilCompleted is false, it will only poll the status and finalize if the reallocation has completed.
+	 *
+	 * @param Texture2D				- Texture to finalize the reallocation for
+	 * @param bBlockUntilCompleted	- Whether the function should block until the reallocation has completed
+	 * @return						- Current reallocation status:
+	 *	TexRealloc_Succeeded	Reallocation succeeded
+	 *	TexRealloc_Failed		Reallocation failed
+	 *	TexRealloc_InProgress	Reallocation is still in progress, try again later
+	 */
+	// FlushType: Wait RHI Thread
+	virtual ETextureReallocationStatus RHIFinalizeAsyncReallocateTexture2D(FTexture2DRHIParamRef Texture2D, bool bBlockUntilCompleted) = 0;
+
+	/**
+	 * Cancels an async reallocation for the specified texture.
+	 * This should be called for the new texture, not the original.
+	 *
+	 * @param Texture				Texture to cancel
+	 * @param bBlockUntilCompleted	If true, blocks until the cancellation is fully completed
+	 * @return						Reallocation status
+	 */
+	// FlushType: Wait RHI Thread
+	virtual ETextureReallocationStatus RHICancelAsyncReallocateTexture2D(FTexture2DRHIParamRef Texture2D, bool bBlockUntilCompleted) = 0;
+
+	/**
+	* Locks an RHI texture's mip-map for read/write operations on the CPU
+	* @param Texture - the RHI texture resource to lock, must not be 0
+	* @param MipIndex - index of the mip level to lock
+	* @param LockMode - Whether to lock the texture read-only instead of write-only
+	* @param DestStride - output to retrieve the textures row stride (pitch)
+	* @param bLockWithinMiptail - for platforms that support packed miptails allow locking of individual mip levels within the miptail
+	* @return pointer to the CPU accessible resource data
+	*/
+	// FlushType: Flush RHI Thread
+	virtual void* RHILockTexture2D(FTexture2DRHIParamRef Texture, uint32 MipIndex, EResourceLockMode LockMode, uint32& DestStride, bool bLockWithinMiptail) = 0;
+
+	/**
+	* Unlocks a previously locked RHI texture resource
+	* @param Texture - the RHI texture resource to unlock, must not be 0
+	* @param MipIndex - index of the mip level to unlock
+	* @param bLockWithinMiptail - for platforms that support packed miptails allow locking of individual mip levels within the miptail
+	*/
+	// FlushType: Flush RHI Thread
+	virtual void RHIUnlockTexture2D(FTexture2DRHIParamRef Texture, uint32 MipIndex, bool bLockWithinMiptail) = 0;
+
+	/**
+	* Locks an RHI texture's mip-map for read/write operations on the CPU
+	* @param Texture - the RHI texture resource to lock
+	* @param MipIndex - index of the mip level to lock
+	* @param LockMode - Whether to lock the texture read-only instead of write-only
+	* @param DestStride - output to retrieve the textures row stride (pitch)
+	* @param bLockWithinMiptail - for platforms that support packed miptails allow locking of individual mip levels within the miptail
+	* @return pointer to the CPU accessible resource data
+	*/
+	// FlushType: Flush RHI Thread
+	virtual void* RHILockTexture2DArray(FTexture2DArrayRHIParamRef Texture, uint32 TextureIndex, uint32 MipIndex, EResourceLockMode LockMode, uint32& DestStride, bool bLockWithinMiptail) = 0;
+
+	/**
+	* Unlocks a previously locked RHI texture resource
+	* @param Texture - the RHI texture resource to unlock
+	* @param MipIndex - index of the mip level to unlock
+	* @param bLockWithinMiptail - for platforms that support packed miptails allow locking of individual mip levels within the miptail
+	*/
+	// FlushType: Flush RHI Thread
+	virtual void RHIUnlockTexture2DArray(FTexture2DArrayRHIParamRef Texture, uint32 TextureIndex, uint32 MipIndex, bool bLockWithinMiptail) = 0;
+
+	/**
+	* Updates a region of a 2D texture from system memory
+	* @param Texture - the RHI texture resource to update
+	* @param MipIndex - mip level index to be modified
+	* @param UpdateRegion - The rectangle to copy source image data from
+	* @param SourcePitch - size in bytes of each row of the source image
+	* @param SourceData - source image data, starting at the upper left corner of the source rectangle (in same pixel format as texture)
+	*/
+	// FlushType: Flush RHI Thread
+	virtual void RHIUpdateTexture2D(FTexture2DRHIParamRef Texture, uint32 MipIndex, const struct FUpdateTextureRegion2D& UpdateRegion, uint32 SourcePitch, const uint8* SourceData) = 0;
+
+	/**
+	* Updates a region of a 3D texture from system memory
+	* @param Texture - the RHI texture resource to update
+	* @param MipIndex - mip level index to be modified
+	* @param UpdateRegion - The rectangle to copy source image data from
+	* @param SourceRowPitch - size in bytes of each row of the source image, usually Bpp * SizeX
+	* @param SourceDepthPitch - size in bytes of each depth slice of the source image, usually Bpp * SizeX * SizeY
+	* @param SourceData - source image data, starting at the upper left corner of the source rectangle (in same pixel format as texture)
+	*/
+	// FlushType: Flush RHI Thread
+	virtual void RHIUpdateTexture3D(FTexture3DRHIParamRef Texture, uint32 MipIndex, const struct FUpdateTextureRegion3D& UpdateRegion, uint32 SourceRowPitch, uint32 SourceDepthPitch, const uint8* SourceData) = 0;
+
+	/**
+	* Creates a Cube RHI texture resource
+	* @param Size - width/height of the texture to create
+	* @param Format - EPixelFormat texture format
+	* @param NumMips - number of mips to generate or 0 for full mip pyramid
+	* @param Flags - ETextureCreateFlags creation flags
+	*/
+	// FlushType: Wait RHI Thread
+	virtual FTextureCubeRHIRef RHICreateTextureCube(uint32 Size, uint8 Format, uint32 NumMips, uint32 Flags, FRHIResourceCreateInfo& CreateInfo) = 0;
+
+	/**
+	* Creates a Cube Array RHI texture resource
+	* @param Size - width/height of the texture to create
+	* @param ArraySize - number of array elements of the texture to create
+	* @param Format - EPixelFormat texture format
+	* @param NumMips - number of mips to generate or 0 for full mip pyramid
+	* @param Flags - ETextureCreateFlags creation flags
+	*/
+	// FlushType: Wait RHI Thread
+	virtual FTextureCubeRHIRef RHICreateTextureCubeArray(uint32 Size, uint32 ArraySize, uint8 Format, uint32 NumMips, uint32 Flags, FRHIResourceCreateInfo& CreateInfo) = 0;
+
+	/**
+	* Locks an RHI texture's mip-map for read/write operations on the CPU
+	* @param Texture - the RHI texture resource to lock
+	* @param MipIndex - index of the mip level to lock
+	* @param LockMode - Whether to lock the texture read-only instead of write-only.
+	* @param DestStride - output to retrieve the textures row stride (pitch)
+	* @param bLockWithinMiptail - for platforms that support packed miptails allow locking of individual mip levels within the miptail
+	* @return pointer to the CPU accessible resource data
+	*/
+	// FlushType: Flush RHI Thread
+	virtual void* RHILockTextureCubeFace(FTextureCubeRHIParamRef Texture, uint32 FaceIndex, uint32 ArrayIndex, uint32 MipIndex, EResourceLockMode LockMode, uint32& DestStride, bool bLockWithinMiptail) = 0;
+
+	/**
+	* Unlocks a previously locked RHI texture resource
+	* @param Texture - the RHI texture resource to unlock
+	* @param MipIndex - index of the mip level to unlock
+	* @param bLockWithinMiptail - for platforms that support packed miptails allow locking of individual mip levels within the miptail
+	*/
+	// FlushType: Flush RHI Thread
+	virtual void RHIUnlockTextureCubeFace(FTextureCubeRHIParamRef Texture, uint32 FaceIndex, uint32 ArrayIndex, uint32 MipIndex, bool bLockWithinMiptail) = 0;
+
+	// FlushType: Thread safe
+	virtual void RHIBindDebugLabelName(FTextureRHIParamRef Texture, const TCHAR* Name) = 0;
+
+	/**
+	 * Reads the contents of a texture to an output buffer (non MSAA and MSAA) and returns it as a FColor array.
+	 * If the format or texture type is unsupported the OutData array will be size 0
+	 */
+	// FlushType: Flush Immediate (seems wrong)
+	virtual void RHIReadSurfaceData(FTextureRHIParamRef Texture,FIntRect Rect,TArray<FColor>& OutData,FReadSurfaceDataFlags InFlags) = 0;
+
+	/** Watch out for OutData to be 0 (can happen on DXGI_ERROR_DEVICE_REMOVED), don't call RHIUnmapStagingSurface in that case. */
+	// FlushType: Flush Immediate (seems wrong)
+	virtual void RHIMapStagingSurface(FTextureRHIParamRef Texture,void*& OutData,int32& OutWidth,int32& OutHeight) = 0;
+
+	/** call after a succesful RHIMapStagingSurface() call */
+	// FlushType: Flush Immediate (seems wrong)
+	virtual void RHIUnmapStagingSurface(FTextureRHIParamRef Texture) = 0;
+
+	// FlushType: Flush Immediate (seems wrong)
+	virtual void RHIReadSurfaceFloatData(FTextureRHIParamRef Texture,FIntRect Rect,TArray<FFloat16Color>& OutData,ECubeFace CubeFace,int32 ArrayIndex,int32 MipIndex) = 0;
+
+	// FlushType: Flush Immediate (seems wrong)
+	virtual void RHIRead3DSurfaceFloatData(FTextureRHIParamRef Texture,FIntRect Rect,FIntPoint ZMinMax,TArray<FFloat16Color>& OutData) = 0;
+
+	// FlushType: Wait RHI Thread
+	virtual FRenderQueryRHIRef RHICreateRenderQuery(ERenderQueryType QueryType) = 0;
+
+	// CAUTION: Even though this is marked as threadsafe, it is only valid to call from the render thread. It is need not be threadsafe on platforms that do not support or aren't using an RHIThread
+	// FlushType: Thread safe, but varies by RHI
+	virtual bool RHIGetRenderQueryResult(FRenderQueryRHIParamRef RenderQuery, uint64& OutResult, bool bWait) = 0;
+
+	// With RHI thread, this is the current backbuffer from the perspective of the render thread.
+	// FlushType: Thread safe
+	virtual FTexture2DRHIRef RHIGetViewportBackBuffer(FViewportRHIParamRef Viewport) = 0;
+
+	// Only relevant with an RHI thread, this advances the backbuffer for the purpose of GetViewportBackBuffer
+	// FlushType: Thread safe
+	virtual void RHIAdvanceFrameForGetViewportBackBuffer() = 0;
+
+	/**
+	 * Determine if currently drawing the viewport
+	 * REMOVE THIS If you are updating the RHI for all platforms
+	 *
+	 * @return true if currently within a BeginDrawingViewport/EndDrawingViewport block
+	 */
+	// FlushType: Wait RHI Thread
+	virtual bool RHIIsDrawingViewport() = 0;
+
+	/*
+	 * Acquires or releases ownership of the platform-specific rendering context for the calling thread
+	 */
+	// FlushType: Flush RHI Thread
+	virtual void RHIAcquireThreadOwnership() = 0;
+
+	// FlushType: Flush RHI Thread
+	virtual void RHIReleaseThreadOwnership() = 0;
+
+	// Flush driver resources. Typically called when switching contexts/threads
+	// FlushType: Flush RHI Thread
+	virtual void RHIFlushResources() = 0;
+
+	/*
+	 * Returns the total GPU time taken to render the last frame. Same metric as FPlatformTime::Cycles().
+	 */
+	// FlushType: Thread safe
+	virtual uint32 RHIGetGPUFrameCycles() = 0;
+
+	//  must be called from the main thread.
+	// FlushType: Thread safe
+	virtual FViewportRHIRef RHICreateViewport(void* WindowHandle, uint32 SizeX, uint32 SizeY, bool bIsFullscreen, EPixelFormat PreferredPixelFormat) = 0;
+
+	//  must be called from the main thread.
+	// FlushType: Thread safe
+	virtual void RHIResizeViewport(FViewportRHIParamRef Viewport, uint32 SizeX, uint32 SizeY, bool bIsFullscreen) = 0;
+
+	//  must be called from the main thread.
+	// FlushType: Thread safe
+	virtual void RHITick(float DeltaTime) = 0;
+
+	/** Sets stream output targets, for use with a geometry shader created with RHICreateGeometryShaderWithStreamOutput. */
+	//@todo this should be a CMDLIST method
+	// FlushType: Flush Immediate (seems wrong)
+	virtual void RHISetStreamOutTargets(uint32 NumTargets,const FVertexBufferRHIParamRef* VertexBuffers,const uint32* Offsets) = 0;
+
+	// FlushType: Flush Immediate (seems wrong)
+	virtual void RHIDiscardRenderTargets(bool Depth,bool Stencil,uint32 ColorBitMask) = 0;
+
+	// Blocks the CPU until the GPU catches up and goes idle.
+	// FlushType: Flush Immediate (seems wrong)
+	virtual void RHIBlockUntilGPUIdle() = 0;
+
+	// Operations to suspend title rendering and yield control to the system
+	// FlushType: Thread safe
+	virtual void RHISuspendRendering() = 0;
+
+	// FlushType: Thread safe
+	virtual void RHIResumeRendering() = 0;
+
+	// FlushType: Flush Immediate
+	virtual bool RHIIsRenderingSuspended() = 0;
+
+	/**
+	 *	Retrieve available screen resolutions.
+	 *
+	 *	@param	Resolutions			TArray<FScreenResolutionRHI> parameter that will be filled in.
+	 *	@param	bIgnoreRefreshRate	If true, ignore refresh rates.
+	 *
+	 *	@return	bool				true if successfully filled the array
+	 */
+	// FlushType: Thread safe
+	virtual bool RHIGetAvailableResolutions(FScreenResolutionArray& Resolutions, bool bIgnoreRefreshRate) = 0;
+
+	/**
+	 * Returns a supported screen resolution that most closely matches input.
+	 * @param Width - Input: Desired resolution width in pixels. Output: A width that the platform supports.
+	 * @param Height - Input: Desired resolution height in pixels. Output: A height that the platform supports.
+	 */
+	// FlushType: Thread safe
+	virtual void RHIGetSupportedResolution(uint32& Width, uint32& Height) = 0;
+
+	/**
+	 * Function that is used to allocate / free space used for virtual texture mip levels.
+	 * Make sure you also update the visible mip levels.
+	 * @param Texture - the texture to update, must have been created with TexCreate_Virtual
+	 * @param FirstMip - the first mip that should be in memory
+	 */
+	// FlushType: Wait RHI Thread
+	virtual void RHIVirtualTextureSetFirstMipInMemory(FTexture2DRHIParamRef Texture, uint32 FirstMip) = 0;
+
+	/**
+	 * Function that can be used to update which is the first visible mip to the GPU.
+	 * @param Texture - the texture to update, must have been created with TexCreate_Virtual
+	 * @param FirstMip - the first mip that should be visible to the GPU
+	 */
+	// FlushType: Wait RHI Thread
+	virtual void RHIVirtualTextureSetFirstMipVisible(FTexture2DRHIParamRef Texture, uint32 FirstMip) = 0;
+
+	// FlushType: Wait RHI Thread
+	virtual void RHIExecuteCommandList(FRHICommandList* CmdList) = 0;
+
+	/**
+	 * Provides access to the native device. Generally this should be avoided but is useful for third party plugins.
+	 */
+	// FlushType: Flush RHI Thread
+	virtual void* RHIGetNativeDevice() = 0;
+
+	// FlushType: Thread safe
+	virtual class IRHICommandContext* RHIGetDefaultContext() = 0;
+
+	// FlushType: Thread safe
+	virtual class IRHICommandContextContainer* RHIGetCommandContextContainer() = 0;
+
 };
 
 /** A global pointer to the dynamically bound RHI implementation. */
 extern RHI_API FDynamicRHI* GDynamicRHI;
 
-#define DEFINE_RHIMETHOD_CMDLIST(Type,Name,ParameterTypesAndNames,ParameterNames,ReturnStatement,NullImplementation)
-#define DEFINE_RHIMETHOD_GLOBAL(Type,Name,ParameterTypesAndNames,ParameterNames,ReturnStatement,NullImplementation) \
-	extern RHI_API Type Name##_Internal ParameterTypesAndNames;
-#define DEFINE_RHIMETHOD_GLOBALTHREADSAFE(Type,Name,ParameterTypesAndNames,ParameterNames,ReturnStatement,NullImplementation) \
-	extern RHI_API Type RHI##Name ParameterTypesAndNames;
-#define DEFINE_RHIMETHOD_GLOBALFLUSH(Type,Name,ParameterTypesAndNames,ParameterNames,ReturnStatement,NullImplementation) \
-	extern RHI_API Type Name##_Internal ParameterTypesAndNames;
-#define DEFINE_RHIMETHOD(Type,Name,ParameterTypesAndNames,ParameterNames,ReturnStatement,NullImplementation) \
-	extern RHI_API Type Name##_Internal ParameterTypesAndNames;
-#include "RHIMethods.h"
-#undef DEFINE_RHIMETHOD
-#undef DEFINE_RHIMETHOD_CMDLIST
-#undef DEFINE_RHIMETHOD_GLOBAL
-#undef DEFINE_RHIMETHOD_GLOBALFLUSH
-#undef DEFINE_RHIMETHOD_GLOBALTHREADSAFE
+
+/** The interface RHI command context. Sometimes the RHI handles these. On platforms that can processes command lists in parallel, it is a separate object. */
+class IRHICommandContext
+{
+public:
+	virtual ~IRHICommandContext()
+	{
+	}
+
+	/////// RHI Context Methods
+
+
+	/**
+	 *Sets the current compute shader.  Mostly for compliance with platforms
+	 *that require shader setting before resource binding.
+	 */
+	virtual void RHISetComputeShader(FComputeShaderRHIParamRef ComputeShader) = 0;
+
+	virtual void RHIDispatchComputeShader(uint32 ThreadGroupCountX, uint32 ThreadGroupCountY, uint32 ThreadGroupCountZ) = 0;
+
+	virtual void RHIDispatchIndirectComputeShader(FVertexBufferRHIParamRef ArgumentBuffer, uint32 ArgumentOffset) = 0;
+
+	virtual void RHIAutomaticCacheFlushAfterComputeShader(bool bEnable) = 0;
+
+	virtual void RHIFlushComputeShaderCache() = 0;
+
+	// Useful when used with geometry shader (emit polygons to different viewports), otherwise SetViewPort() is simpler
+	// @param Count >0
+	// @param Data must not be 0
+	virtual void RHISetMultipleViewports(uint32 Count, const FViewportBounds* Data) = 0;
+
+	/** Clears a UAV to the multi-component value provided. */
+	virtual void RHIClearUAV(FUnorderedAccessViewRHIParamRef UnorderedAccessViewRHI, const uint32* Values) = 0;
+
+	/**
+	* Resolves from one texture to another.
+	* @param SourceTexture - texture to resolve from, 0 is silenty ignored
+	* @param DestTexture - texture to resolve to, 0 is silenty ignored
+	* @param bKeepOriginalSurface - true if the original surface will still be used after this function so must remain valid
+	* @param ResolveParams - optional resolve params
+	*/
+	virtual void RHICopyToResolveTarget(FTextureRHIParamRef SourceTexture, FTextureRHIParamRef DestTexture, bool bKeepOriginalSurface, const FResolveParams& ResolveParams) = 0;
+
+	virtual void RHIBeginRenderQuery(FRenderQueryRHIParamRef RenderQuery) = 0;
+
+	virtual void RHIEndRenderQuery(FRenderQueryRHIParamRef RenderQuery) = 0;
+
+	virtual void RHIBeginOcclusionQueryBatch() = 0;
+
+	virtual void RHIEndOcclusionQueryBatch() = 0;
+
+	virtual void RHISubmitCommandsHint() = 0;
+
+	// This method is queued with an RHIThread, otherwise it will flush after it is queued; without an RHI thread there is no benefit to queuing this frame advance commands
+	virtual void RHIBeginDrawingViewport(FViewportRHIParamRef Viewport, FTextureRHIParamRef RenderTargetRHI) = 0;
+
+	// This method is queued with an RHIThread, otherwise it will flush after it is queued; without an RHI thread there is no benefit to queuing this frame advance commands
+	virtual void RHIEndDrawingViewport(FViewportRHIParamRef Viewport, bool bPresent, bool bLockToVsync) = 0;
+
+	// This method is queued with an RHIThread, otherwise it will flush after it is queued; without an RHI thread there is no benefit to queuing this frame advance commands
+	virtual void RHIBeginFrame() = 0;
+
+	// This method is queued with an RHIThread, otherwise it will flush after it is queued; without an RHI thread there is no benefit to queuing this frame advance commands
+	virtual void RHIEndFrame() = 0;
+
+	/**
+		* Signals the beginning of scene rendering. The RHI makes certain caching assumptions between
+		* calls to BeginScene/EndScene. Currently the only restriction is that you can't update texture
+		* references.
+		*/
+	// This method is queued with an RHIThread, otherwise it will flush after it is queued; without an RHI thread there is no benefit to queuing this frame advance commands
+	virtual void RHIBeginScene() = 0;
+
+	/**
+		* Signals the end of scene rendering. See RHIBeginScene.
+		*/
+	// This method is queued with an RHIThread, otherwise it will flush after it is queued; without an RHI thread there is no benefit to queuing this frame advance commands
+	virtual void RHIEndScene() = 0;
+
+	virtual void RHISetStreamSource(uint32 StreamIndex, FVertexBufferRHIParamRef VertexBuffer, uint32 Stride, uint32 Offset) = 0;
+
+	virtual void RHISetRasterizerState(FRasterizerStateRHIParamRef NewState) = 0;
+
+	// @param MinX including like Win32 RECT
+	// @param MinY including like Win32 RECT
+	// @param MaxX excluding like Win32 RECT
+	// @param MaxY excluding like Win32 RECT
+	virtual void RHISetViewport(uint32 MinX, uint32 MinY, float MinZ, uint32 MaxX, uint32 MaxY, float MaxZ) = 0;
+
+	// @param MinX including like Win32 RECT
+	// @param MinY including like Win32 RECT
+	// @param MaxX excluding like Win32 RECT
+	// @param MaxY excluding like Win32 RECT
+	virtual void RHISetScissorRect(bool bEnable, uint32 MinX, uint32 MinY, uint32 MaxX, uint32 MaxY) = 0;
+
+	/**
+	 * Set bound shader state. This will set the vertex decl/shader, and pixel shader
+	 * @param BoundShaderState - state resource
+	 */
+	virtual void RHISetBoundShaderState(FBoundShaderStateRHIParamRef BoundShaderState) = 0;
+
+	/** Set the shader resource view of a surface.  This is used for binding TextureMS parameter types that need a multi sampled view. */
+	virtual void RHISetShaderTexture(FVertexShaderRHIParamRef VertexShader, uint32 TextureIndex, FTextureRHIParamRef NewTexture) = 0;
+
+	/** Set the shader resource view of a surface.  This is used for binding TextureMS parameter types that need a multi sampled view. */
+	virtual void RHISetShaderTexture(FHullShaderRHIParamRef HullShader, uint32 TextureIndex, FTextureRHIParamRef NewTexture) = 0;
+
+	/** Set the shader resource view of a surface.  This is used for binding TextureMS parameter types that need a multi sampled view. */
+	virtual void RHISetShaderTexture(FDomainShaderRHIParamRef DomainShader, uint32 TextureIndex, FTextureRHIParamRef NewTexture) = 0;
+
+	/** Set the shader resource view of a surface.  This is used for binding TextureMS parameter types that need a multi sampled view. */
+	virtual void RHISetShaderTexture(FGeometryShaderRHIParamRef GeometryShader, uint32 TextureIndex, FTextureRHIParamRef NewTexture) = 0;
+
+	/** Set the shader resource view of a surface.  This is used for binding TextureMS parameter types that need a multi sampled view. */
+	virtual void RHISetShaderTexture(FPixelShaderRHIParamRef PixelShader, uint32 TextureIndex, FTextureRHIParamRef NewTexture) = 0;
+
+	/** Set the shader resource view of a surface.  This is used for binding TextureMS parameter types that need a multi sampled view. */
+	virtual void RHISetShaderTexture(FComputeShaderRHIParamRef PixelShader, uint32 TextureIndex, FTextureRHIParamRef NewTexture) = 0;
+
+	/**
+	 * Sets sampler state.
+	 * @param GeometryShader	The geometry shader to set the sampler for.
+	 * @param SamplerIndex		The index of the sampler.
+	 * @param NewState			The new sampler state.
+	 */
+	virtual void RHISetShaderSampler(FComputeShaderRHIParamRef ComputeShader, uint32 SamplerIndex, FSamplerStateRHIParamRef NewState) = 0;
+
+	/**
+	 * Sets sampler state.
+	 * @param GeometryShader	The geometry shader to set the sampler for.
+	 * @param SamplerIndex		The index of the sampler.
+	 * @param NewState			The new sampler state.
+	 */
+	virtual void RHISetShaderSampler(FVertexShaderRHIParamRef VertexShader, uint32 SamplerIndex, FSamplerStateRHIParamRef NewState) = 0;
+
+	/**
+	 * Sets sampler state.
+	 * @param GeometryShader	The geometry shader to set the sampler for.
+	 * @param SamplerIndex		The index of the sampler.
+	 * @param NewState			The new sampler state.
+	 */
+	virtual void RHISetShaderSampler(FGeometryShaderRHIParamRef GeometryShader, uint32 SamplerIndex, FSamplerStateRHIParamRef NewState) = 0;
+
+	/**
+	 * Sets sampler state.
+	 * @param GeometryShader	The geometry shader to set the sampler for.
+	 * @param SamplerIndex		The index of the sampler.
+	 * @param NewState			The new sampler state.
+	 */
+	virtual void RHISetShaderSampler(FDomainShaderRHIParamRef DomainShader, uint32 SamplerIndex, FSamplerStateRHIParamRef NewState) = 0;
+
+	/**
+	 * Sets sampler state.
+	 * @param GeometryShader	The geometry shader to set the sampler for.
+	 * @param SamplerIndex		The index of the sampler.
+	 * @param NewState			The new sampler state.
+	 */
+	virtual void RHISetShaderSampler(FHullShaderRHIParamRef HullShader, uint32 SamplerIndex, FSamplerStateRHIParamRef NewState) = 0;
+
+	/**
+	 * Sets sampler state.
+	 * @param GeometryShader	The geometry shader to set the sampler for.
+	 * @param SamplerIndex		The index of the sampler.
+	 * @param NewState			The new sampler state.
+	 */
+	virtual void RHISetShaderSampler(FPixelShaderRHIParamRef PixelShader, uint32 SamplerIndex, FSamplerStateRHIParamRef NewState) = 0;
+
+	/** Sets a compute shader UAV parameter. */
+	virtual void RHISetUAVParameter(FComputeShaderRHIParamRef ComputeShader, uint32 UAVIndex, FUnorderedAccessViewRHIParamRef UAV) = 0;
+
+	/** Sets a compute shader UAV parameter and initial count */
+	virtual void RHISetUAVParameter(FComputeShaderRHIParamRef ComputeShader, uint32 UAVIndex, FUnorderedAccessViewRHIParamRef UAV, uint32 InitialCount) = 0;
+
+	virtual void RHISetShaderResourceViewParameter(FPixelShaderRHIParamRef PixelShader, uint32 SamplerIndex, FShaderResourceViewRHIParamRef SRV) = 0;
+
+	virtual void RHISetShaderResourceViewParameter(FVertexShaderRHIParamRef VertexShader, uint32 SamplerIndex, FShaderResourceViewRHIParamRef SRV) = 0;
+
+	virtual void RHISetShaderResourceViewParameter(FComputeShaderRHIParamRef ComputeShader, uint32 SamplerIndex, FShaderResourceViewRHIParamRef SRV) = 0;
+
+	virtual void RHISetShaderResourceViewParameter(FHullShaderRHIParamRef HullShader, uint32 SamplerIndex, FShaderResourceViewRHIParamRef SRV) = 0;
+
+	virtual void RHISetShaderResourceViewParameter(FDomainShaderRHIParamRef DomainShader, uint32 SamplerIndex, FShaderResourceViewRHIParamRef SRV) = 0;
+
+	virtual void RHISetShaderResourceViewParameter(FGeometryShaderRHIParamRef GeometryShader, uint32 SamplerIndex, FShaderResourceViewRHIParamRef SRV) = 0;
+
+	virtual void RHISetShaderUniformBuffer(FVertexShaderRHIParamRef VertexShader, uint32 BufferIndex, FUniformBufferRHIParamRef Buffer) = 0;
+
+	virtual void RHISetShaderUniformBuffer(FHullShaderRHIParamRef HullShader, uint32 BufferIndex, FUniformBufferRHIParamRef Buffer) = 0;
+
+	virtual void RHISetShaderUniformBuffer(FDomainShaderRHIParamRef DomainShader, uint32 BufferIndex, FUniformBufferRHIParamRef Buffer) = 0;
+
+	virtual void RHISetShaderUniformBuffer(FGeometryShaderRHIParamRef GeometryShader, uint32 BufferIndex, FUniformBufferRHIParamRef Buffer) = 0;
+
+	virtual void RHISetShaderUniformBuffer(FPixelShaderRHIParamRef PixelShader, uint32 BufferIndex, FUniformBufferRHIParamRef Buffer) = 0;
+
+	virtual void RHISetShaderUniformBuffer(FComputeShaderRHIParamRef ComputeShader, uint32 BufferIndex, FUniformBufferRHIParamRef Buffer) = 0;
+
+	virtual void RHISetShaderParameter(FVertexShaderRHIParamRef VertexShader, uint32 BufferIndex, uint32 BaseIndex, uint32 NumBytes, const void* NewValue) = 0;
+
+	virtual void RHISetShaderParameter(FPixelShaderRHIParamRef PixelShader, uint32 BufferIndex, uint32 BaseIndex, uint32 NumBytes, const void* NewValue) = 0;
+
+	virtual void RHISetShaderParameter(FHullShaderRHIParamRef HullShader, uint32 BufferIndex, uint32 BaseIndex, uint32 NumBytes, const void* NewValue) = 0;
+
+	virtual void RHISetShaderParameter(FDomainShaderRHIParamRef DomainShader, uint32 BufferIndex, uint32 BaseIndex, uint32 NumBytes, const void* NewValue) = 0;
+
+	virtual void RHISetShaderParameter(FGeometryShaderRHIParamRef GeometryShader, uint32 BufferIndex, uint32 BaseIndex, uint32 NumBytes, const void* NewValue) = 0;
+
+	virtual void RHISetShaderParameter(FComputeShaderRHIParamRef ComputeShader, uint32 BufferIndex, uint32 BaseIndex, uint32 NumBytes, const void* NewValue) = 0;
+
+	virtual void RHISetDepthStencilState(FDepthStencilStateRHIParamRef NewState, uint32 StencilRef) = 0;
+
+	// Allows to set the blend state, parameter can be created with RHICreateBlendState()
+	virtual void RHISetBlendState(FBlendStateRHIParamRef NewState, const FLinearColor& BlendFactor) = 0;
+
+	virtual void RHISetRenderTargets(uint32 NumSimultaneousRenderTargets, const FRHIRenderTargetView* NewRenderTargets, const FRHIDepthRenderTargetView* NewDepthStencilTarget, uint32 NumUAVs, const FUnorderedAccessViewRHIParamRef* UAVs) = 0;
+
+	virtual void RHISetRenderTargetsAndClear(const FRHISetRenderTargetsInfo& RenderTargetsInfo) = 0;
+
+	// Bind a clear values to currently bound rendertargets.  This is used by platforms which
+	// need the color when transitioning a target that supports hardware clears from a rendertarget to a shader resource.
+	// The explicit bind is needed to support parallel rendering.
+	virtual void RHIBindClearMRTValues(bool bClearColor, int32 NumClearColors, const FLinearColor* ColorArray, bool bClearDepth, float Depth, bool bClearStencil, uint32 Stencil) = 0;
+
+	virtual void RHIDrawPrimitive(uint32 PrimitiveType, uint32 BaseVertexIndex, uint32 NumPrimitives, uint32 NumInstances) = 0;
+
+	virtual void RHIDrawPrimitiveIndirect(uint32 PrimitiveType, FVertexBufferRHIParamRef ArgumentBuffer, uint32 ArgumentOffset) = 0;
+
+	virtual void RHIDrawIndexedIndirect(FIndexBufferRHIParamRef IndexBufferRHI, uint32 PrimitiveType, FStructuredBufferRHIParamRef ArgumentsBufferRHI, int32 DrawArgumentsIndex, uint32 NumInstances) = 0;
+
+	// @param NumPrimitives need to be >0 
+	virtual void RHIDrawIndexedPrimitive(FIndexBufferRHIParamRef IndexBuffer, uint32 PrimitiveType, int32 BaseVertexIndex, uint32 FirstInstance, uint32 NumVertices, uint32 StartIndex, uint32 NumPrimitives, uint32 NumInstances) = 0;
+
+	virtual void RHIDrawIndexedPrimitiveIndirect(uint32 PrimitiveType, FIndexBufferRHIParamRef IndexBuffer, FVertexBufferRHIParamRef ArgumentBuffer, uint32 ArgumentOffset) = 0;
+
+	/**
+	 * Preallocate memory or get a direct command stream pointer to fill up for immediate rendering . This avoids memcpys below in DrawPrimitiveUP
+	 * @param PrimitiveType The type (triangles, lineloop, etc) of primitive to draw
+	 * @param NumPrimitives The number of primitives in the VertexData buffer
+	 * @param NumVertices The number of vertices to be written
+	 * @param VertexDataStride Size of each vertex 
+	 * @param OutVertexData Reference to the allocated vertex memory
+	 */
+	virtual void RHIBeginDrawPrimitiveUP(uint32 PrimitiveType, uint32 NumPrimitives, uint32 NumVertices, uint32 VertexDataStride, void*& OutVertexData) = 0;
+
+	/**
+	 * Draw a primitive using the vertex data populated since RHIBeginDrawPrimitiveUP and clean up any memory as needed
+	 */
+	virtual void RHIEndDrawPrimitiveUP() = 0;
+
+	/**
+	 * Preallocate memory or get a direct command stream pointer to fill up for immediate rendering . This avoids memcpys below in DrawIndexedPrimitiveUP
+	 * @param PrimitiveType The type (triangles, lineloop, etc) of primitive to draw
+	 * @param NumPrimitives The number of primitives in the VertexData buffer
+	 * @param NumVertices The number of vertices to be written
+	 * @param VertexDataStride Size of each vertex
+	 * @param OutVertexData Reference to the allocated vertex memory
+	 * @param MinVertexIndex The lowest vertex index used by the index buffer
+	 * @param NumIndices Number of indices to be written
+	 * @param IndexDataStride Size of each index (either 2 or 4 bytes)
+	 * @param OutIndexData Reference to the allocated index memory
+	 */
+	virtual void RHIBeginDrawIndexedPrimitiveUP(uint32 PrimitiveType, uint32 NumPrimitives, uint32 NumVertices, uint32 VertexDataStride, void*& OutVertexData, uint32 MinVertexIndex, uint32 NumIndices, uint32 IndexDataStride, void*& OutIndexData) = 0;
+
+	/**
+	 * Draw a primitive using the vertex and index data populated since RHIBeginDrawIndexedPrimitiveUP and clean up any memory as needed
+	 */
+	virtual void RHIEndDrawIndexedPrimitiveUP() = 0;
+
+	/*
+	 * This method clears all MRT's, but to only one color value
+	 * @param ExcludeRect within the viewport in pixels, is only a hint to optimize - if a fast clear can be done this is preferred
+	 */
+	virtual void RHIClear(bool bClearColor, const FLinearColor& Color, bool bClearDepth, float Depth, bool bClearStencil, uint32 Stencil, FIntRect ExcludeRect) = 0;
+
+	/*
+	 * This method clears all MRT's to potentially different color values
+	 * @param ExcludeRect within the viewport in pixels, is only a hint to optimize - if a fast clear can be done this is preferred
+	 */
+	virtual void RHIClearMRT(bool bClearColor, int32 NumClearColors, const FLinearColor* ColorArray, bool bClearDepth, float Depth, bool bClearStencil, uint32 Stencil, FIntRect ExcludeRect) = 0;
+
+	/**
+	 * Enabled/Disables Depth Bounds Testing with the given min/max depth.
+	 * @param bEnable	Enable(non-zero)/disable(zero) the depth bounds test
+	 * @param MinDepth	The minimum depth for depth bounds test
+	 * @param MaxDepth	The maximum depth for depth bounds test.
+	 *					The valid values for fMinDepth and fMaxDepth are such that 0 <= fMinDepth <= fMaxDepth <= 1
+	 */
+	virtual void RHIEnableDepthBoundsTest(bool bEnable, float MinDepth, float MaxDepth) = 0;
+
+	virtual void RHIPushEvent(const TCHAR* Name) = 0;
+
+	virtual void RHIPopEvent() = 0;
+
+};
+
+
+
+FORCEINLINE FSamplerStateRHIRef RHICreateSamplerState(const FSamplerStateInitializerRHI& Initializer)
+{
+	return GDynamicRHI->RHICreateSamplerState(Initializer);
+}
+
+FORCEINLINE FRasterizerStateRHIRef RHICreateRasterizerState(const FRasterizerStateInitializerRHI& Initializer)
+{
+	return GDynamicRHI->RHICreateRasterizerState(Initializer);
+}
+
+FORCEINLINE FDepthStencilStateRHIRef RHICreateDepthStencilState(const FDepthStencilStateInitializerRHI& Initializer)
+{
+	return GDynamicRHI->RHICreateDepthStencilState(Initializer);
+}
+
+FORCEINLINE FBlendStateRHIRef RHICreateBlendState(const FBlendStateInitializerRHI& Initializer)
+{
+	return GDynamicRHI->RHICreateBlendState(Initializer);
+}
+
+FORCEINLINE FBoundShaderStateRHIRef RHICreateBoundShaderState(FVertexDeclarationRHIParamRef VertexDeclaration, FVertexShaderRHIParamRef VertexShader, FHullShaderRHIParamRef HullShader, FDomainShaderRHIParamRef DomainShader, FPixelShaderRHIParamRef PixelShader, FGeometryShaderRHIParamRef GeometryShader)
+{
+	return GDynamicRHI->RHICreateBoundShaderState(VertexDeclaration, VertexShader, HullShader, DomainShader, PixelShader, GeometryShader);
+}
+
+FORCEINLINE FUniformBufferRHIRef RHICreateUniformBuffer(const void* Contents, const FRHIUniformBufferLayout& Layout, EUniformBufferUsage Usage)
+{
+	return GDynamicRHI->RHICreateUniformBuffer(Contents, Layout, Usage);
+}
+
+FORCEINLINE uint64 RHICalcTexture2DPlatformSize(uint32 SizeX, uint32 SizeY, uint8 Format, uint32 NumMips, uint32 NumSamples, uint32 Flags, uint32& OutAlign)
+{
+	return GDynamicRHI->RHICalcTexture2DPlatformSize(SizeX, SizeY, Format, NumMips, NumSamples, Flags, OutAlign);
+}
+
+FORCEINLINE uint64 RHICalcTexture3DPlatformSize(uint32 SizeX, uint32 SizeY, uint32 SizeZ, uint8 Format, uint32 NumMips, uint32 Flags, uint32& OutAlign)
+{
+	return GDynamicRHI->RHICalcTexture3DPlatformSize(SizeX, SizeY, SizeZ, Format, NumMips, Flags, OutAlign);
+}
+
+FORCEINLINE uint64 RHICalcTextureCubePlatformSize(uint32 Size, uint8 Format, uint32 NumMips, uint32 Flags, uint32& OutAlign)
+{
+	return GDynamicRHI->RHICalcTextureCubePlatformSize(Size, Format, NumMips, Flags, OutAlign);
+}
+
+FORCEINLINE void RHIGetTextureMemoryStats(FTextureMemoryStats& OutStats)
+{
+	 GDynamicRHI->RHIGetTextureMemoryStats(OutStats);
+}
+
+FORCEINLINE void RHIGetResourceInfo(FTextureRHIParamRef Ref, FRHIResourceInfo& OutInfo)
+{
+	return GDynamicRHI->RHIGetResourceInfo(Ref, OutInfo);
+}
+
+FORCEINLINE uint32 RHIComputeMemorySize(FTextureRHIParamRef TextureRHI)
+{
+	return GDynamicRHI->RHIComputeMemorySize(TextureRHI);
+}
+
+FORCEINLINE void RHIBindDebugLabelName(FTextureRHIParamRef Texture, const TCHAR* Name)
+{
+	 GDynamicRHI->RHIBindDebugLabelName(Texture, Name);
+}
+
+FORCEINLINE bool RHIGetRenderQueryResult(FRenderQueryRHIParamRef RenderQuery, uint64& OutResult, bool bWait)
+{
+	return GDynamicRHI->RHIGetRenderQueryResult(RenderQuery, OutResult, bWait);
+}
+
+FORCEINLINE FTexture2DRHIRef RHIGetViewportBackBuffer(FViewportRHIParamRef Viewport)
+{
+	return GDynamicRHI->RHIGetViewportBackBuffer(Viewport);
+}
+
+FORCEINLINE void RHIAdvanceFrameForGetViewportBackBuffer()
+{
+	return GDynamicRHI->RHIAdvanceFrameForGetViewportBackBuffer();
+}
+
+FORCEINLINE uint32 RHIGetGPUFrameCycles()
+{
+	return GDynamicRHI->RHIGetGPUFrameCycles();
+}
+
+FORCEINLINE FViewportRHIRef RHICreateViewport(void* WindowHandle, uint32 SizeX, uint32 SizeY, bool bIsFullscreen, EPixelFormat PreferredPixelFormat)
+{
+	return GDynamicRHI->RHICreateViewport(WindowHandle, SizeX, SizeY, bIsFullscreen, PreferredPixelFormat);
+}
+
+FORCEINLINE void RHIResizeViewport(FViewportRHIParamRef Viewport, uint32 SizeX, uint32 SizeY, bool bIsFullscreen)
+{
+	 GDynamicRHI->RHIResizeViewport(Viewport, SizeX, SizeY, bIsFullscreen);
+}
+
+FORCEINLINE void RHITick(float DeltaTime)
+{
+	 GDynamicRHI->RHITick(DeltaTime);
+}
+
+FORCEINLINE void RHISuspendRendering()
+{
+	 GDynamicRHI->RHISuspendRendering();
+}
+
+FORCEINLINE void RHIResumeRendering()
+{
+	 GDynamicRHI->RHIResumeRendering();
+}
+
+FORCEINLINE bool RHIGetAvailableResolutions(FScreenResolutionArray& Resolutions, bool bIgnoreRefreshRate)
+{
+	return GDynamicRHI->RHIGetAvailableResolutions(Resolutions, bIgnoreRefreshRate);
+}
+
+FORCEINLINE void RHIGetSupportedResolution(uint32& Width, uint32& Height)
+{
+	 GDynamicRHI->RHIGetSupportedResolution(Width, Height);
+}
+
+FORCEINLINE class IRHICommandContext* RHIGetDefaultContext()
+{
+	return GDynamicRHI->RHIGetDefaultContext();
+}
+
+FORCEINLINE class IRHICommandContextContainer* RHIGetCommandContextContainer()
+{
+	return GDynamicRHI->RHIGetCommandContextContainer();
+}
 
 
-#endif	//#if USE_DYNAMIC_RHI
 
 
-#if PLATFORM_USES_DYNAMIC_RHI
 
 /**
  * Defragment the texture pool.
@@ -85,7 +1073,6 @@ inline bool appIsPoolTexture( FTextureRHIParamRef TextureRHI ) { return false; }
  */
 inline void appDumpTextureMemoryStats(const TCHAR* /*Message*/) {}
 
-#endif // PLATFORM_USES_DYNAMIC_RHI
 
 /** Defines the interface of a module implementing a dynamic RHI. */
 class IDynamicRHIModule : public IModuleInterface
@@ -99,13 +1086,9 @@ public:
 	virtual FDynamicRHI* CreateRHI() = 0;
 };
 
-#if USE_DYNAMIC_RHI
 /**
  *	Each platform that utilizes dynamic RHIs should implement this function
  *	Called to create the instance of the dynamic RHI.
  */
 FDynamicRHI* PlatformCreateDynamicRHI();
 
-#endif	//USE_DYNAMIC_RHI
-
-#endif
diff --git a/Engine/Source/Runtime/RHI/Public/RHI.h b/Engine/Source/Runtime/RHI/Public/RHI.h
--- a/Engine/Source/Runtime/RHI/Public/RHI.h
+++ b/Engine/Source/Runtime/RHI/Public/RHI.h
@@ -1024,42 +1024,8 @@ DECLARE_MEMORY_STAT_POOL_EXTERN(TEXT("Pixel buffer memory"),STAT_PixelBufferMemo
 
 // RHI base resource types.
 #include "RHIResources.h"
+#include "DynamicRHI.h"
 
-//
-// Platform-specific RHI types.
-//
-#if PLATFORM_USES_DYNAMIC_RHI
-	// Use dynamically bound RHIs on PCs and when using the null RHI.
-	#define USE_DYNAMIC_RHI 1
-	#include "DynamicRHI.h"
-#else
-	// Fall back to the null RHI
-	#define USE_DYNAMIC_RHI 1
-	#include "DynamicRHI.h"
-#endif
-
-#if !USE_DYNAMIC_RHI
-	// Define the statically bound RHI methods with the RHI name prefix.
-#define DEFINE_RHIMETHOD_CMDLIST(Type,Name,ParameterTypesAndNames,ParameterNames,ReturnStatement,NullImplementation) \
-	extern Type Name##_Internal ParameterTypesAndNames
-#define DEFINE_RHIMETHOD_GLOBAL(Type,Name,ParameterTypesAndNames,ParameterNames,ReturnStatement,NullImplementation) \
-	extern Type Name##_Internal ParameterTypesAndNames
-#define DEFINE_RHIMETHOD_GLOBALTHREADSAFE(Type,Name,ParameterTypesAndNames,ParameterNames,ReturnStatement,NullImplementation) \
-	extern Type RHI##Name ParameterTypesAndNames
-#define DEFINE_RHIMETHOD_GLOBALFLUSH(Type,Name,ParameterTypesAndNames,ParameterNames,ReturnStatement,NullImplementation) \
-	extern Type Name##_Internal ParameterTypesAndNames
-#define DEFINE_RHIMETHOD(Type,Name,ParameterTypesAndNames,ParameterNames,ReturnStatement,NullImplementation) \
-	extern Type Name##_Internal ParameterTypesAndNames
-
-#include "RHIMethods.h"
-#undef DEFINE_RHIMETHOD
-#undef DEFINE_RHIMETHOD_CMDLIST
-#undef DEFINE_RHIMETHOD_GLOBAL
-#undef DEFINE_RHIMETHOD_GLOBALFLUSH
-#undef DEFINE_RHIMETHOD_GLOBALTHREADSAFE
-
-
-#endif
 
 /** Initializes the RHI. */
 extern RHI_API void RHIInit(bool bHasEditorToken);
diff --git a/Engine/Source/Runtime/RHI/Public/RHICommandList.h b/Engine/Source/Runtime/RHI/Public/RHICommandList.h
--- a/Engine/Source/Runtime/RHI/Public/RHICommandList.h
+++ b/Engine/Source/Runtime/RHI/Public/RHICommandList.h
@@ -16,22 +16,12 @@ extern RHI_API TAutoConsoleVariable<int32> CVarRHICmdWidth;
 class FRHICommandListBase;
 
 
-class IRHICommandContext
-{
-public:
-	virtual ~IRHICommandContext() {}
-
-#define DEFINE_RHIMETHOD(Type,Name,ParameterTypesAndNames,ParameterNames,ReturnStatement,NullImplementation)
-#define DEFINE_RHIMETHOD_CMDLIST(Type,Name,ParameterTypesAndNames,ParameterNames,ReturnStatement,NullImplementation) virtual Type RHI##Name ParameterTypesAndNames = 0
-#include "RHIMethods.h"
-#undef DEFINE_RHIMETHOD
-#undef DEFINE_RHIMETHOD_CMDLIST
-};
-
 class IRHICommandContextContainer
 {
 public:
-	virtual ~IRHICommandContextContainer() {}
+	virtual ~IRHICommandContextContainer()
+{
+	}
 
 	virtual IRHICommandContext* GetContext()
 	{
@@ -1783,9 +1773,9 @@ public:
 	{
 		if (Bypass())
 		{
-			void* Data = LockVertexBuffer_Internal(VertexBuffer, 0, BufferSize, RLM_WriteOnly);
+			void* Data = GDynamicRHI->RHILockVertexBuffer(VertexBuffer, 0, BufferSize, RLM_WriteOnly);
 			FMemory::Memcpy(Data, Buffer, BufferSize);
-			UnlockVertexBuffer_Internal(VertexBuffer);
+			GDynamicRHI->RHIUnlockVertexBuffer(VertexBuffer);
 			return;
 		}
 		void* Data = Alloc(BufferSize, 16);
@@ -1924,43 +1914,711 @@ public:
 		return bFlushedGlobal;
 	};
 
-	#define DEFINE_RHIMETHOD_CMDLIST(Type,Name,ParameterTypesAndNames,ParameterNames,ReturnStatement,NullImplementation)
-	#define DEFINE_RHIMETHOD(Type, Name, ParameterTypesAndNames, ParameterNames, ReturnStatement, NullImplementation) \
-		FORCEINLINE_DEBUGGABLE Type Name ParameterTypesAndNames \
-		{ \
-			QUICK_SCOPE_CYCLE_COUNTER(STAT_RHIMETHOD_##Name##_Flush); \
-			ImmediateFlush(EImmediateFlushType::FlushRHIThread); \
-			ReturnStatement Name##_Internal ParameterNames; \
-		}
-	#define DEFINE_RHIMETHOD_GLOBAL(Type, Name, ParameterTypesAndNames, ParameterNames, ReturnStatement, NullImplementation) \
-		FORCEINLINE_DEBUGGABLE Type Name ParameterTypesAndNames \
-		{ \
-			if (GRHIThread) \
-			{ \
-				QUICK_SCOPE_CYCLE_COUNTER(STAT_RHIMETHOD_##Name##_WaitRHI); \
-				ImmediateFlush(EImmediateFlushType::WaitForRHIThread); \
-			} \
-			ReturnStatement Name##_Internal ParameterNames; \
-		}
-	#define DEFINE_RHIMETHOD_GLOBALTHREADSAFE(Type, Name, ParameterTypesAndNames, ParameterNames, ReturnStatement, NullImplementation) \
-		FORCEINLINE_DEBUGGABLE Type Name ParameterTypesAndNames \
-		{ \
-			ReturnStatement RHI##Name ParameterNames; \
-		}
-	#define DEFINE_RHIMETHOD_GLOBALFLUSH(Type, Name, ParameterTypesAndNames, ParameterNames, ReturnStatement, NullImplementation) \
-		FORCEINLINE_DEBUGGABLE Type Name ParameterTypesAndNames \
-		{ \
-			QUICK_SCOPE_CYCLE_COUNTER(STAT_RHIMETHOD_##Name##_Flush); \
-			ImmediateFlush(EImmediateFlushType::FlushRHIThread); \
-			TGuardValue<bool> GuardIsFlushedGlobal( bFlushedGlobal, true ); \
-			ReturnStatement Name##_Internal ParameterNames; \
-		}
-	#include "RHIMethods.h"
-	#undef DEFINE_RHIMETHOD
-	#undef DEFINE_RHIMETHOD_CMDLIST
-	#undef DEFINE_RHIMETHOD_GLOBAL
-	#undef DEFINE_RHIMETHOD_GLOBALFLUSH
-	#undef DEFINE_RHIMETHOD_GLOBALTHREADSAFE
+
+
+	FORCEINLINE void GpuTimeBegin(uint32 Hash,bool bCompute)
+	{
+		QUICK_SCOPE_CYCLE_COUNTER(STAT_RHIMETHOD_GpuTimeBegin_Flush);
+		ImmediateFlush(EImmediateFlushType::FlushRHIThread); 
+		return GDynamicRHI->RHIGpuTimeBegin(Hash,bCompute);
+	}
+	
+	FORCEINLINE void GpuTimeEnd(uint32 Hash,bool bCompute)
+	{
+		QUICK_SCOPE_CYCLE_COUNTER(STAT_RHIMETHOD_GpuTimeEnd_Flush);
+		ImmediateFlush(EImmediateFlushType::FlushRHIThread); 
+		return GDynamicRHI->RHIGpuTimeEnd(Hash,bCompute);
+	}
+	
+	FORCEINLINE FSamplerStateRHIRef CreateSamplerState(const FSamplerStateInitializerRHI& Initializer)
+	{
+		return RHICreateSamplerState(Initializer);
+	}
+	
+	FORCEINLINE FRasterizerStateRHIRef CreateRasterizerState(const FRasterizerStateInitializerRHI& Initializer)
+	{
+		return RHICreateRasterizerState(Initializer);
+	}
+	
+	FORCEINLINE FDepthStencilStateRHIRef CreateDepthStencilState(const FDepthStencilStateInitializerRHI& Initializer)
+	{
+		return RHICreateDepthStencilState(Initializer);
+	}
+	
+	FORCEINLINE FBlendStateRHIRef CreateBlendState(const FBlendStateInitializerRHI& Initializer)
+	{
+		return RHICreateBlendState(Initializer);
+	}
+	
+	FORCEINLINE FVertexDeclarationRHIRef CreateVertexDeclaration(const FVertexDeclarationElementList& Elements)
+	{
+		if(GRHIThread)
+		{
+			QUICK_SCOPE_CYCLE_COUNTER(STAT_RHIMETHOD_CreateVertexDeclaration_WaitRHI);
+			ImmediateFlush(EImmediateFlushType::WaitForRHIThread);
+		}
+		return GDynamicRHI->RHICreateVertexDeclaration(Elements);
+	}
+	
+	FORCEINLINE FPixelShaderRHIRef CreatePixelShader(const TArray<uint8>& Code)
+	{
+		if(GRHIThread)
+		{
+			QUICK_SCOPE_CYCLE_COUNTER(STAT_RHIMETHOD_CreatePixelShader_WaitRHI);
+			ImmediateFlush(EImmediateFlushType::WaitForRHIThread);
+		}
+		return GDynamicRHI->RHICreatePixelShader(Code);
+	}
+	
+	FORCEINLINE FVertexShaderRHIRef CreateVertexShader(const TArray<uint8>& Code)
+	{
+		if(GRHIThread)
+		{
+			QUICK_SCOPE_CYCLE_COUNTER(STAT_RHIMETHOD_CreateVertexShader_WaitRHI);
+			ImmediateFlush(EImmediateFlushType::WaitForRHIThread);
+		}
+		return GDynamicRHI->RHICreateVertexShader(Code);
+	}
+	
+	FORCEINLINE FHullShaderRHIRef CreateHullShader(const TArray<uint8>& Code)
+	{
+		if(GRHIThread)
+		{
+			QUICK_SCOPE_CYCLE_COUNTER(STAT_RHIMETHOD_CreateHullShader_WaitRHI);
+			ImmediateFlush(EImmediateFlushType::WaitForRHIThread);
+		}
+		return GDynamicRHI->RHICreateHullShader(Code);
+	}
+	
+	FORCEINLINE FDomainShaderRHIRef CreateDomainShader(const TArray<uint8>& Code)
+	{
+		if(GRHIThread)
+		{
+			QUICK_SCOPE_CYCLE_COUNTER(STAT_RHIMETHOD_CreateDomainShader_WaitRHI);
+			ImmediateFlush(EImmediateFlushType::WaitForRHIThread);
+		}
+		return GDynamicRHI->RHICreateDomainShader(Code);
+	}
+	
+	FORCEINLINE FGeometryShaderRHIRef CreateGeometryShader(const TArray<uint8>& Code)
+	{
+		if(GRHIThread)
+		{
+			QUICK_SCOPE_CYCLE_COUNTER(STAT_RHIMETHOD_CreateGeometryShader_WaitRHI);
+			ImmediateFlush(EImmediateFlushType::WaitForRHIThread);
+		}
+		return GDynamicRHI->RHICreateGeometryShader(Code);
+	}
+	
+	FORCEINLINE FGeometryShaderRHIRef CreateGeometryShaderWithStreamOutput(const TArray<uint8>& Code, const FStreamOutElementList& ElementList, uint32 NumStrides, const uint32* Strides, int32 RasterizedStream)
+	{
+		if(GRHIThread)
+		{
+			QUICK_SCOPE_CYCLE_COUNTER(STAT_RHIMETHOD_CreateGeometryShaderWithStreamOutput_WaitRHI);
+			ImmediateFlush(EImmediateFlushType::WaitForRHIThread);
+		}
+		return GDynamicRHI->RHICreateGeometryShaderWithStreamOutput(Code, ElementList, NumStrides, Strides, RasterizedStream);
+	}
+	
+	FORCEINLINE FComputeShaderRHIRef CreateComputeShader(const TArray<uint8>& Code)
+	{
+		if(GRHIThread)
+		{
+			QUICK_SCOPE_CYCLE_COUNTER(STAT_RHIMETHOD_CreateComputeShader_WaitRHI);
+			ImmediateFlush(EImmediateFlushType::WaitForRHIThread);
+		}
+		return GDynamicRHI->RHICreateComputeShader(Code);
+	}
+	
+	FORCEINLINE FBoundShaderStateRHIRef CreateBoundShaderState(FVertexDeclarationRHIParamRef VertexDeclaration, FVertexShaderRHIParamRef VertexShader, FHullShaderRHIParamRef HullShader, FDomainShaderRHIParamRef DomainShader, FPixelShaderRHIParamRef PixelShader, FGeometryShaderRHIParamRef GeometryShader)
+	{
+		return RHICreateBoundShaderState(VertexDeclaration, VertexShader, HullShader, DomainShader, PixelShader, GeometryShader);
+	}
+	
+	FORCEINLINE FUniformBufferRHIRef CreateUniformBuffer(const void* Contents, const FRHIUniformBufferLayout& Layout, EUniformBufferUsage Usage)
+	{
+		return RHICreateUniformBuffer(Contents, Layout, Usage);
+	}
+	
+	FORCEINLINE FIndexBufferRHIRef CreateIndexBuffer(uint32 Stride, uint32 Size, uint32 InUsage, FRHIResourceCreateInfo& CreateInfo)
+	{
+		if(GRHIThread)
+		{
+			QUICK_SCOPE_CYCLE_COUNTER(STAT_RHIMETHOD_CreateIndexBuffer_WaitRHI);
+			ImmediateFlush(EImmediateFlushType::WaitForRHIThread);
+		}
+		return GDynamicRHI->RHICreateIndexBuffer(Stride, Size, InUsage, CreateInfo);
+	}
+	
+	FORCEINLINE void* LockIndexBuffer(FIndexBufferRHIParamRef IndexBuffer, uint32 Offset, uint32 Size, EResourceLockMode LockMode)
+	{
+		QUICK_SCOPE_CYCLE_COUNTER(STAT_RHIMETHOD_LockIndexBuffer_Flush);
+		ImmediateFlush(EImmediateFlushType::FlushRHIThread); 
+		TGuardValue<bool> GuardIsFlushedGlobal( bFlushedGlobal, true ); 
+		return GDynamicRHI->RHILockIndexBuffer(IndexBuffer, Offset, Size, LockMode);
+	}
+	
+	FORCEINLINE void UnlockIndexBuffer(FIndexBufferRHIParamRef IndexBuffer)
+	{
+		QUICK_SCOPE_CYCLE_COUNTER(STAT_RHIMETHOD_UnlockIndexBuffer_Flush);
+		ImmediateFlush(EImmediateFlushType::FlushRHIThread); 
+		TGuardValue<bool> GuardIsFlushedGlobal( bFlushedGlobal, true );  
+		GDynamicRHI->RHIUnlockIndexBuffer(IndexBuffer);
+	}
+	
+	FORCEINLINE FVertexBufferRHIRef CreateVertexBuffer(uint32 Size, uint32 InUsage, FRHIResourceCreateInfo& CreateInfo)
+	{
+		if(GRHIThread)
+		{
+			QUICK_SCOPE_CYCLE_COUNTER(STAT_RHIMETHOD_CreateVertexBuffer_WaitRHI);
+			ImmediateFlush(EImmediateFlushType::WaitForRHIThread);
+		}
+		return GDynamicRHI->RHICreateVertexBuffer(Size, InUsage, CreateInfo);
+	}
+	
+	FORCEINLINE void* LockVertexBuffer(FVertexBufferRHIParamRef VertexBuffer, uint32 Offset, uint32 SizeRHI, EResourceLockMode LockMode)
+	{
+		QUICK_SCOPE_CYCLE_COUNTER(STAT_RHIMETHOD_LockVertexBuffer_Flush);
+		ImmediateFlush(EImmediateFlushType::FlushRHIThread); 
+		TGuardValue<bool> GuardIsFlushedGlobal( bFlushedGlobal, true ); 
+		return GDynamicRHI->RHILockVertexBuffer(VertexBuffer, Offset, SizeRHI, LockMode);
+	}
+	
+	FORCEINLINE void UnlockVertexBuffer(FVertexBufferRHIParamRef VertexBuffer)
+	{
+		QUICK_SCOPE_CYCLE_COUNTER(STAT_RHIMETHOD_UnlockVertexBuffer_Flush);
+		ImmediateFlush(EImmediateFlushType::FlushRHIThread); 
+		TGuardValue<bool> GuardIsFlushedGlobal( bFlushedGlobal, true );  
+		GDynamicRHI->RHIUnlockVertexBuffer(VertexBuffer);
+	}
+	
+	FORCEINLINE void CopyVertexBuffer(FVertexBufferRHIParamRef SourceBuffer,FVertexBufferRHIParamRef DestBuffer)
+	{
+		QUICK_SCOPE_CYCLE_COUNTER(STAT_RHIMETHOD_CopyVertexBuffer_Flush);
+		ImmediateFlush(EImmediateFlushType::FlushRHIThread);  
+		GDynamicRHI->RHICopyVertexBuffer(SourceBuffer,DestBuffer);
+	}
+	
+	FORCEINLINE FStructuredBufferRHIRef CreateStructuredBuffer(uint32 Stride, uint32 Size, uint32 InUsage, FRHIResourceCreateInfo& CreateInfo)
+	{
+		if(GRHIThread)
+		{
+			QUICK_SCOPE_CYCLE_COUNTER(STAT_RHIMETHOD_CreateStructuredBuffer_WaitRHI);
+			ImmediateFlush(EImmediateFlushType::WaitForRHIThread);
+		}
+		return GDynamicRHI->RHICreateStructuredBuffer(Stride, Size, InUsage, CreateInfo);
+	}
+	
+	FORCEINLINE void* LockStructuredBuffer(FStructuredBufferRHIParamRef StructuredBuffer, uint32 Offset, uint32 SizeRHI, EResourceLockMode LockMode)
+	{
+		QUICK_SCOPE_CYCLE_COUNTER(STAT_RHIMETHOD_LockStructuredBuffer_Flush);
+		ImmediateFlush(EImmediateFlushType::FlushRHIThread); 
+		TGuardValue<bool> GuardIsFlushedGlobal( bFlushedGlobal, true ); 
+		return GDynamicRHI->RHILockStructuredBuffer(StructuredBuffer, Offset, SizeRHI, LockMode);
+	}
+	
+	FORCEINLINE void UnlockStructuredBuffer(FStructuredBufferRHIParamRef StructuredBuffer)
+	{
+		QUICK_SCOPE_CYCLE_COUNTER(STAT_RHIMETHOD_UnlockStructuredBuffer_Flush);
+		ImmediateFlush(EImmediateFlushType::FlushRHIThread); 
+		TGuardValue<bool> GuardIsFlushedGlobal( bFlushedGlobal, true );  
+		GDynamicRHI->RHIUnlockStructuredBuffer(StructuredBuffer);
+	}
+	
+	FORCEINLINE FUnorderedAccessViewRHIRef CreateUnorderedAccessView(FStructuredBufferRHIParamRef StructuredBuffer, bool bUseUAVCounter, bool bAppendBuffer)
+	{
+		if(GRHIThread)
+		{
+			QUICK_SCOPE_CYCLE_COUNTER(STAT_RHIMETHOD_CreateUnorderedAccessView_WaitRHI);
+			ImmediateFlush(EImmediateFlushType::WaitForRHIThread);
+		}
+		return GDynamicRHI->RHICreateUnorderedAccessView(StructuredBuffer, bUseUAVCounter, bAppendBuffer);
+	}
+	
+	FORCEINLINE FUnorderedAccessViewRHIRef CreateUnorderedAccessView(FTextureRHIParamRef Texture)
+	{
+		if(GRHIThread)
+		{
+			QUICK_SCOPE_CYCLE_COUNTER(STAT_RHIMETHOD_CreateUnorderedAccessView_WaitRHI);
+			ImmediateFlush(EImmediateFlushType::WaitForRHIThread);
+		}
+		return GDynamicRHI->RHICreateUnorderedAccessView(Texture);
+	}
+	
+	FORCEINLINE FUnorderedAccessViewRHIRef CreateUnorderedAccessView(FVertexBufferRHIParamRef VertexBuffer, uint8 Format)
+	{
+		if(GRHIThread)
+		{
+			QUICK_SCOPE_CYCLE_COUNTER(STAT_RHIMETHOD_CreateUnorderedAccessView_WaitRHI);
+			ImmediateFlush(EImmediateFlushType::WaitForRHIThread);
+		}
+		return GDynamicRHI->RHICreateUnorderedAccessView(VertexBuffer, Format);
+	}
+	
+	FORCEINLINE FShaderResourceViewRHIRef CreateShaderResourceView(FStructuredBufferRHIParamRef StructuredBuffer)
+	{
+		if(GRHIThread)
+		{
+			QUICK_SCOPE_CYCLE_COUNTER(STAT_RHIMETHOD_CreateShaderResourceView_WaitRHI);
+			ImmediateFlush(EImmediateFlushType::WaitForRHIThread);
+		}
+		return GDynamicRHI->RHICreateShaderResourceView(StructuredBuffer);
+	}
+	
+	FORCEINLINE FShaderResourceViewRHIRef CreateShaderResourceView(FVertexBufferRHIParamRef VertexBuffer, uint32 Stride, uint8 Format)
+	{
+		if(GRHIThread)
+		{
+			QUICK_SCOPE_CYCLE_COUNTER(STAT_RHIMETHOD_CreateShaderResourceView_WaitRHI);
+			ImmediateFlush(EImmediateFlushType::WaitForRHIThread);
+		}
+		return GDynamicRHI->RHICreateShaderResourceView(VertexBuffer, Stride, Format);
+	}
+	
+	FORCEINLINE uint64 CalcTexture2DPlatformSize(uint32 SizeX, uint32 SizeY, uint8 Format, uint32 NumMips, uint32 NumSamples, uint32 Flags, uint32& OutAlign)
+	{
+		return RHICalcTexture2DPlatformSize(SizeX, SizeY, Format, NumMips, NumSamples, Flags, OutAlign);
+	}
+	
+	FORCEINLINE uint64 CalcTexture3DPlatformSize(uint32 SizeX, uint32 SizeY, uint32 SizeZ, uint8 Format, uint32 NumMips, uint32 Flags, uint32& OutAlign)
+	{
+		return RHICalcTexture3DPlatformSize(SizeX, SizeY, SizeZ, Format, NumMips, Flags, OutAlign);
+	}
+	
+	FORCEINLINE uint64 CalcTextureCubePlatformSize(uint32 Size, uint8 Format, uint32 NumMips, uint32 Flags, uint32& OutAlign)
+	{
+		return RHICalcTextureCubePlatformSize(Size, Format, NumMips, Flags, OutAlign);
+	}
+	
+	FORCEINLINE void GetTextureMemoryStats(FTextureMemoryStats& OutStats)
+	{
+		RHIGetTextureMemoryStats(OutStats);
+	}
+	
+	FORCEINLINE bool GetTextureMemoryVisualizeData(FColor* TextureData,int32 SizeX,int32 SizeY,int32 Pitch,int32 PixelSize)
+	{
+		QUICK_SCOPE_CYCLE_COUNTER(STAT_RHIMETHOD_GetTextureMemoryVisualizeData_Flush);
+		ImmediateFlush(EImmediateFlushType::FlushRHIThread); 
+		return GDynamicRHI->RHIGetTextureMemoryVisualizeData(TextureData,SizeX,SizeY,Pitch,PixelSize);
+	}
+	
+	FORCEINLINE FTextureReferenceRHIRef CreateTextureReference(FLastRenderTimeContainer* LastRenderTime)
+	{
+		if(GRHIThread)
+		{
+			QUICK_SCOPE_CYCLE_COUNTER(STAT_RHIMETHOD_CreateTextureReference_WaitRHI);
+			ImmediateFlush(EImmediateFlushType::WaitForRHIThread);
+		}
+		return GDynamicRHI->RHICreateTextureReference(LastRenderTime);
+	}
+	
+	FORCEINLINE void UpdateTextureReference(FTextureReferenceRHIParamRef TextureRef, FTextureRHIParamRef NewTexture)
+	{
+		QUICK_SCOPE_CYCLE_COUNTER(STAT_RHIMETHOD_UpdateTextureReference_Flush);
+		ImmediateFlush(EImmediateFlushType::FlushRHIThread); 
+		TGuardValue<bool> GuardIsFlushedGlobal( bFlushedGlobal, true );  
+		GDynamicRHI->RHIUpdateTextureReference(TextureRef, NewTexture);
+	}
+	
+	FORCEINLINE FTexture2DRHIRef CreateTexture2D(uint32 SizeX, uint32 SizeY, uint8 Format, uint32 NumMips, uint32 NumSamples, uint32 Flags, FRHIResourceCreateInfo& CreateInfo)
+	{
+		if(GRHIThread)
+		{
+			QUICK_SCOPE_CYCLE_COUNTER(STAT_RHIMETHOD_CreateTexture2D_WaitRHI);
+			ImmediateFlush(EImmediateFlushType::WaitForRHIThread);
+		}
+		return GDynamicRHI->RHICreateTexture2D(SizeX, SizeY, Format, NumMips, NumSamples, Flags, CreateInfo);
+	}
+	
+	FORCEINLINE FTexture2DRHIRef AsyncCreateTexture2D(uint32 SizeX, uint32 SizeY, uint8 Format, uint32 NumMips, uint32 Flags, void** InitialMipData, uint32 NumInitialMips)
+	{
+		if(GRHIThread)
+		{
+			QUICK_SCOPE_CYCLE_COUNTER(STAT_RHIMETHOD_AsyncCreateTexture2D_WaitRHI);
+			ImmediateFlush(EImmediateFlushType::WaitForRHIThread);
+		}
+		return GDynamicRHI->RHIAsyncCreateTexture2D(SizeX, SizeY, Format, NumMips, Flags, InitialMipData, NumInitialMips);
+	}
+	
+	FORCEINLINE void CopySharedMips(FTexture2DRHIParamRef DestTexture2D, FTexture2DRHIParamRef SrcTexture2D)
+	{
+		QUICK_SCOPE_CYCLE_COUNTER(STAT_RHIMETHOD_CopySharedMips_Flush);
+		ImmediateFlush(EImmediateFlushType::FlushRHIThread); 
+		TGuardValue<bool> GuardIsFlushedGlobal( bFlushedGlobal, true ); 
+		return GDynamicRHI->RHICopySharedMips(DestTexture2D, SrcTexture2D);
+	}
+	
+	FORCEINLINE FTexture2DArrayRHIRef CreateTexture2DArray(uint32 SizeX, uint32 SizeY, uint32 SizeZ, uint8 Format, uint32 NumMips, uint32 Flags, FRHIResourceCreateInfo& CreateInfo)
+	{
+		if(GRHIThread)
+		{
+			QUICK_SCOPE_CYCLE_COUNTER(STAT_RHIMETHOD_CreateTexture2DArray_WaitRHI);
+			ImmediateFlush(EImmediateFlushType::WaitForRHIThread);
+		}
+		return GDynamicRHI->RHICreateTexture2DArray(SizeX, SizeY, SizeZ, Format, NumMips, Flags, CreateInfo);
+	}
+	
+	FORCEINLINE FTexture3DRHIRef CreateTexture3D(uint32 SizeX, uint32 SizeY, uint32 SizeZ, uint8 Format, uint32 NumMips, uint32 Flags, FRHIResourceCreateInfo& CreateInfo)
+	{
+		if(GRHIThread)
+		{
+			QUICK_SCOPE_CYCLE_COUNTER(STAT_RHIMETHOD_CreateTexture3D_WaitRHI);
+			ImmediateFlush(EImmediateFlushType::WaitForRHIThread);
+		}
+		return GDynamicRHI->RHICreateTexture3D(SizeX, SizeY, SizeZ, Format, NumMips, Flags, CreateInfo);
+	}
+	
+	FORCEINLINE void GetResourceInfo(FTextureRHIParamRef Ref, FRHIResourceInfo& OutInfo)
+	{
+		return RHIGetResourceInfo(Ref, OutInfo);
+	}
+	
+	FORCEINLINE FShaderResourceViewRHIRef CreateShaderResourceView(FTexture2DRHIParamRef Texture2DRHI, uint8 MipLevel)
+	{
+		if(GRHIThread)
+		{
+			QUICK_SCOPE_CYCLE_COUNTER(STAT_RHIMETHOD_CreateShaderResourceView_WaitRHI);
+			ImmediateFlush(EImmediateFlushType::WaitForRHIThread);
+		}
+		return GDynamicRHI->RHICreateShaderResourceView(Texture2DRHI, MipLevel);
+	}
+	
+	FORCEINLINE FShaderResourceViewRHIRef CreateShaderResourceView(FTexture2DRHIParamRef Texture2DRHI, uint8 MipLevel, uint8 NumMipLevels, uint8 Format)
+	{
+		if(GRHIThread)
+		{
+			QUICK_SCOPE_CYCLE_COUNTER(STAT_RHIMETHOD_CreateShaderResourceView_WaitRHI);
+			ImmediateFlush(EImmediateFlushType::WaitForRHIThread);
+		}
+		return GDynamicRHI->RHICreateShaderResourceView(Texture2DRHI, MipLevel, NumMipLevels, Format);
+	}
+	
+	FORCEINLINE void GenerateMips(FTextureRHIParamRef Texture)
+	{
+		QUICK_SCOPE_CYCLE_COUNTER(STAT_RHIMETHOD_GenerateMips_Flush);
+		ImmediateFlush(EImmediateFlushType::FlushRHIThread); return GDynamicRHI->RHIGenerateMips(Texture);
+	}
+	
+	FORCEINLINE uint32 ComputeMemorySize(FTextureRHIParamRef TextureRHI)
+	{
+		return RHIComputeMemorySize(TextureRHI);
+	}
+	
+	FORCEINLINE FTexture2DRHIRef AsyncReallocateTexture2D(FTexture2DRHIParamRef Texture2D, int32 NewMipCount, int32 NewSizeX, int32 NewSizeY, FThreadSafeCounter* RequestStatus)
+	{
+		QUICK_SCOPE_CYCLE_COUNTER(STAT_RHIMETHOD_AsyncReallocateTexture2D_Flush);
+		ImmediateFlush(EImmediateFlushType::FlushRHIThread); TGuardValue<bool> GuardIsFlushedGlobal( bFlushedGlobal, true ); 
+		return GDynamicRHI->RHIAsyncReallocateTexture2D(Texture2D, NewMipCount, NewSizeX, NewSizeY, RequestStatus);
+	}
+	
+	FORCEINLINE ETextureReallocationStatus FinalizeAsyncReallocateTexture2D(FTexture2DRHIParamRef Texture2D, bool bBlockUntilCompleted)
+	{
+		if(GRHIThread)
+		{
+			QUICK_SCOPE_CYCLE_COUNTER(STAT_RHIMETHOD_FinalizeAsyncReallocateTexture2D_WaitRHI);
+			ImmediateFlush(EImmediateFlushType::WaitForRHIThread);
+		}
+		return GDynamicRHI->RHIFinalizeAsyncReallocateTexture2D(Texture2D, bBlockUntilCompleted);
+	}
+	
+	FORCEINLINE ETextureReallocationStatus CancelAsyncReallocateTexture2D(FTexture2DRHIParamRef Texture2D, bool bBlockUntilCompleted)
+	{
+		if(GRHIThread)
+		{
+			QUICK_SCOPE_CYCLE_COUNTER(STAT_RHIMETHOD_CancelAsyncReallocateTexture2D_WaitRHI);
+			ImmediateFlush(EImmediateFlushType::WaitForRHIThread);
+		}
+		return GDynamicRHI->RHICancelAsyncReallocateTexture2D(Texture2D, bBlockUntilCompleted);
+	}
+	
+	FORCEINLINE void* LockTexture2D(FTexture2DRHIParamRef Texture, uint32 MipIndex, EResourceLockMode LockMode, uint32& DestStride, bool bLockWithinMiptail)
+	{
+		QUICK_SCOPE_CYCLE_COUNTER(STAT_RHIMETHOD_LockTexture2D_Flush);
+		ImmediateFlush(EImmediateFlushType::FlushRHIThread); TGuardValue<bool> GuardIsFlushedGlobal( bFlushedGlobal, true ); 
+		return GDynamicRHI->RHILockTexture2D(Texture, MipIndex, LockMode, DestStride, bLockWithinMiptail);
+	}
+	
+	FORCEINLINE void UnlockTexture2D(FTexture2DRHIParamRef Texture, uint32 MipIndex, bool bLockWithinMiptail)
+	{
+		QUICK_SCOPE_CYCLE_COUNTER(STAT_RHIMETHOD_UnlockTexture2D_Flush);
+		ImmediateFlush(EImmediateFlushType::FlushRHIThread); TGuardValue<bool> GuardIsFlushedGlobal( bFlushedGlobal, true );  
+		GDynamicRHI->RHIUnlockTexture2D(Texture, MipIndex, bLockWithinMiptail);
+	}
+	
+	FORCEINLINE void* LockTexture2DArray(FTexture2DArrayRHIParamRef Texture, uint32 TextureIndex, uint32 MipIndex, EResourceLockMode LockMode, uint32& DestStride, bool bLockWithinMiptail)
+	{
+		QUICK_SCOPE_CYCLE_COUNTER(STAT_RHIMETHOD_LockTexture2DArray_Flush);
+		ImmediateFlush(EImmediateFlushType::FlushRHIThread); TGuardValue<bool> GuardIsFlushedGlobal( bFlushedGlobal, true ); 
+		return GDynamicRHI->RHILockTexture2DArray(Texture, TextureIndex, MipIndex, LockMode, DestStride, bLockWithinMiptail);
+	}
+	
+	FORCEINLINE void UnlockTexture2DArray(FTexture2DArrayRHIParamRef Texture, uint32 TextureIndex, uint32 MipIndex, bool bLockWithinMiptail)
+	{
+		QUICK_SCOPE_CYCLE_COUNTER(STAT_RHIMETHOD_UnlockTexture2DArray_Flush);
+		ImmediateFlush(EImmediateFlushType::FlushRHIThread); TGuardValue<bool> GuardIsFlushedGlobal( bFlushedGlobal, true );  
+		GDynamicRHI->RHIUnlockTexture2DArray(Texture, TextureIndex, MipIndex, bLockWithinMiptail);
+	}
+	
+	FORCEINLINE void UpdateTexture2D(FTexture2DRHIParamRef Texture, uint32 MipIndex, const struct FUpdateTextureRegion2D& UpdateRegion, uint32 SourcePitch, const uint8* SourceData)
+	{
+		QUICK_SCOPE_CYCLE_COUNTER(STAT_RHIMETHOD_UpdateTexture2D_Flush);
+		ImmediateFlush(EImmediateFlushType::FlushRHIThread); TGuardValue<bool> GuardIsFlushedGlobal( bFlushedGlobal, true );  
+		GDynamicRHI->RHIUpdateTexture2D(Texture, MipIndex, UpdateRegion, SourcePitch, SourceData);
+	}
+	
+	FORCEINLINE void UpdateTexture3D(FTexture3DRHIParamRef Texture, uint32 MipIndex, const struct FUpdateTextureRegion3D& UpdateRegion, uint32 SourceRowPitch, uint32 SourceDepthPitch, const uint8* SourceData)
+	{
+		QUICK_SCOPE_CYCLE_COUNTER(STAT_RHIMETHOD_UpdateTexture3D_Flush);
+		ImmediateFlush(EImmediateFlushType::FlushRHIThread); TGuardValue<bool> GuardIsFlushedGlobal( bFlushedGlobal, true );  
+		GDynamicRHI->RHIUpdateTexture3D(Texture, MipIndex, UpdateRegion, SourceRowPitch, SourceDepthPitch, SourceData);
+	}
+	
+	FORCEINLINE FTextureCubeRHIRef CreateTextureCube(uint32 Size, uint8 Format, uint32 NumMips, uint32 Flags, FRHIResourceCreateInfo& CreateInfo)
+	{
+		if(GRHIThread)
+		{
+			QUICK_SCOPE_CYCLE_COUNTER(STAT_RHIMETHOD_CreateTextureCube_WaitRHI);
+			ImmediateFlush(EImmediateFlushType::WaitForRHIThread);
+		}
+		return GDynamicRHI->RHICreateTextureCube(Size, Format, NumMips, Flags, CreateInfo);
+	}
+	
+	FORCEINLINE FTextureCubeRHIRef CreateTextureCubeArray(uint32 Size, uint32 ArraySize, uint8 Format, uint32 NumMips, uint32 Flags, FRHIResourceCreateInfo& CreateInfo)
+	{
+		if(GRHIThread)
+		{
+			QUICK_SCOPE_CYCLE_COUNTER(STAT_RHIMETHOD_CreateTextureCubeArray_WaitRHI);
+			ImmediateFlush(EImmediateFlushType::WaitForRHIThread);
+		}
+		return GDynamicRHI->RHICreateTextureCubeArray(Size, ArraySize, Format, NumMips, Flags, CreateInfo);
+	}
+	
+	FORCEINLINE void* LockTextureCubeFace(FTextureCubeRHIParamRef Texture, uint32 FaceIndex, uint32 ArrayIndex, uint32 MipIndex, EResourceLockMode LockMode, uint32& DestStride, bool bLockWithinMiptail)
+	{
+		QUICK_SCOPE_CYCLE_COUNTER(STAT_RHIMETHOD_LockTextureCubeFace_Flush);
+		ImmediateFlush(EImmediateFlushType::FlushRHIThread); TGuardValue<bool> GuardIsFlushedGlobal( bFlushedGlobal, true ); 
+		return GDynamicRHI->RHILockTextureCubeFace(Texture, FaceIndex, ArrayIndex, MipIndex, LockMode, DestStride, bLockWithinMiptail);
+	}
+	
+	FORCEINLINE void UnlockTextureCubeFace(FTextureCubeRHIParamRef Texture, uint32 FaceIndex, uint32 ArrayIndex, uint32 MipIndex, bool bLockWithinMiptail)
+	{
+		QUICK_SCOPE_CYCLE_COUNTER(STAT_RHIMETHOD_UnlockTextureCubeFace_Flush);
+		ImmediateFlush(EImmediateFlushType::FlushRHIThread); TGuardValue<bool> GuardIsFlushedGlobal( bFlushedGlobal, true );  
+		GDynamicRHI->RHIUnlockTextureCubeFace(Texture, FaceIndex, ArrayIndex, MipIndex, bLockWithinMiptail);
+	}
+	
+	FORCEINLINE void BindDebugLabelName(FTextureRHIParamRef Texture, const TCHAR* Name)
+	{
+		RHIBindDebugLabelName(Texture, Name);
+	}
+	
+	FORCEINLINE void ReadSurfaceData(FTextureRHIParamRef Texture,FIntRect Rect,TArray<FColor>& OutData,FReadSurfaceDataFlags InFlags)
+	{
+		QUICK_SCOPE_CYCLE_COUNTER(STAT_RHIMETHOD_ReadSurfaceData_Flush);
+		ImmediateFlush(EImmediateFlushType::FlushRHIThread);  
+		GDynamicRHI->RHIReadSurfaceData(Texture,Rect,OutData,InFlags);
+	}
+	
+	FORCEINLINE void MapStagingSurface(FTextureRHIParamRef Texture,void*& OutData,int32& OutWidth,int32& OutHeight)
+	{
+		QUICK_SCOPE_CYCLE_COUNTER(STAT_RHIMETHOD_MapStagingSurface_Flush);
+		ImmediateFlush(EImmediateFlushType::FlushRHIThread);  
+		GDynamicRHI->RHIMapStagingSurface(Texture,OutData,OutWidth,OutHeight);
+	}
+	
+	FORCEINLINE void UnmapStagingSurface(FTextureRHIParamRef Texture)
+	{
+		QUICK_SCOPE_CYCLE_COUNTER(STAT_RHIMETHOD_UnmapStagingSurface_Flush);
+		ImmediateFlush(EImmediateFlushType::FlushRHIThread);  
+		GDynamicRHI->RHIUnmapStagingSurface(Texture);
+	}
+	
+	FORCEINLINE void ReadSurfaceFloatData(FTextureRHIParamRef Texture,FIntRect Rect,TArray<FFloat16Color>& OutData,ECubeFace CubeFace,int32 ArrayIndex,int32 MipIndex)
+	{
+		QUICK_SCOPE_CYCLE_COUNTER(STAT_RHIMETHOD_ReadSurfaceFloatData_Flush);
+		ImmediateFlush(EImmediateFlushType::FlushRHIThread);  
+		GDynamicRHI->RHIReadSurfaceFloatData(Texture,Rect,OutData,CubeFace,ArrayIndex,MipIndex);
+	}
+	
+	FORCEINLINE void Read3DSurfaceFloatData(FTextureRHIParamRef Texture,FIntRect Rect,FIntPoint ZMinMax,TArray<FFloat16Color>& OutData)
+	{
+		QUICK_SCOPE_CYCLE_COUNTER(STAT_RHIMETHOD_Read3DSurfaceFloatData_Flush);
+		ImmediateFlush(EImmediateFlushType::FlushRHIThread);  
+		GDynamicRHI->RHIRead3DSurfaceFloatData(Texture,Rect,ZMinMax,OutData);
+	}
+	
+	FORCEINLINE FRenderQueryRHIRef CreateRenderQuery(ERenderQueryType QueryType)
+	{
+		if(GRHIThread)
+		{
+			QUICK_SCOPE_CYCLE_COUNTER(STAT_RHIMETHOD_CreateRenderQuery_WaitRHI);
+			ImmediateFlush(EImmediateFlushType::WaitForRHIThread);
+		}
+		return GDynamicRHI->RHICreateRenderQuery(QueryType);
+	}
+	
+	FORCEINLINE bool GetRenderQueryResult(FRenderQueryRHIParamRef RenderQuery, uint64& OutResult, bool bWait)
+	{
+		return RHIGetRenderQueryResult(RenderQuery, OutResult, bWait);
+	}
+	
+	FORCEINLINE FTexture2DRHIRef GetViewportBackBuffer(FViewportRHIParamRef Viewport)
+	{
+		return RHIGetViewportBackBuffer(Viewport);
+	}
+	
+	FORCEINLINE void AdvanceFrameForGetViewportBackBuffer()
+	{
+		return RHIAdvanceFrameForGetViewportBackBuffer();
+	}
+	
+	FORCEINLINE bool IsDrawingViewport()
+	{
+		QUICK_SCOPE_CYCLE_COUNTER(STAT_RHIMETHOD_IsDrawingViewport_Flush);
+		ImmediateFlush(EImmediateFlushType::FlushRHIThread); 
+		return GDynamicRHI->RHIIsDrawingViewport();
+	}
+	
+	FORCEINLINE void AcquireThreadOwnership()
+	{
+		QUICK_SCOPE_CYCLE_COUNTER(STAT_RHIMETHOD_AcquireThreadOwnership_Flush);
+		ImmediateFlush(EImmediateFlushType::FlushRHIThread); 
+		TGuardValue<bool> GuardIsFlushedGlobal( bFlushedGlobal, true ); 
+		return GDynamicRHI->RHIAcquireThreadOwnership();
+	}
+	
+	FORCEINLINE void ReleaseThreadOwnership()
+	{
+		QUICK_SCOPE_CYCLE_COUNTER(STAT_RHIMETHOD_ReleaseThreadOwnership_Flush);
+		ImmediateFlush(EImmediateFlushType::FlushRHIThread); 
+		TGuardValue<bool> GuardIsFlushedGlobal( bFlushedGlobal, true ); 
+		return GDynamicRHI->RHIReleaseThreadOwnership();
+	}
+	
+	FORCEINLINE void FlushResources()
+	{
+		QUICK_SCOPE_CYCLE_COUNTER(STAT_RHIMETHOD_FlushResources_Flush);
+		ImmediateFlush(EImmediateFlushType::FlushRHIThread); 
+		TGuardValue<bool> GuardIsFlushedGlobal( bFlushedGlobal, true ); 
+		return GDynamicRHI->RHIFlushResources();
+	}
+	
+	FORCEINLINE uint32 GetGPUFrameCycles()
+	{
+		return RHIGetGPUFrameCycles();
+	}
+	
+	FORCEINLINE FViewportRHIRef CreateViewport(void* WindowHandle, uint32 SizeX, uint32 SizeY, bool bIsFullscreen, EPixelFormat PreferredPixelFormat)
+	{
+		return RHICreateViewport(WindowHandle, SizeX, SizeY, bIsFullscreen, PreferredPixelFormat);
+	}
+	
+	FORCEINLINE void ResizeViewport(FViewportRHIParamRef Viewport, uint32 SizeX, uint32 SizeY, bool bIsFullscreen)
+	{
+		RHIResizeViewport(Viewport, SizeX, SizeY, bIsFullscreen);
+	}
+	
+	FORCEINLINE void Tick(float DeltaTime)
+	{
+		RHITick(DeltaTime);
+	}
+	
+	FORCEINLINE void SetStreamOutTargets(uint32 NumTargets,const FVertexBufferRHIParamRef* VertexBuffers,const uint32* Offsets)
+	{
+		QUICK_SCOPE_CYCLE_COUNTER(STAT_RHIMETHOD_SetStreamOutTargets_Flush);
+		ImmediateFlush(EImmediateFlushType::FlushRHIThread);  
+		GDynamicRHI->RHISetStreamOutTargets(NumTargets,VertexBuffers,Offsets);
+	}
+	
+	FORCEINLINE void DiscardRenderTargets(bool Depth,bool Stencil,uint32 ColorBitMask)
+	{
+		QUICK_SCOPE_CYCLE_COUNTER(STAT_RHIMETHOD_DiscardRenderTargets_Flush);
+		ImmediateFlush(EImmediateFlushType::FlushRHIThread);  
+		GDynamicRHI->RHIDiscardRenderTargets(Depth,Stencil,ColorBitMask);
+	}
+	
+	FORCEINLINE void BlockUntilGPUIdle()
+	{
+		QUICK_SCOPE_CYCLE_COUNTER(STAT_RHIMETHOD_BlockUntilGPUIdle_Flush);
+		ImmediateFlush(EImmediateFlushType::FlushRHIThread);  
+		GDynamicRHI->RHIBlockUntilGPUIdle();
+	}
+	
+	FORCEINLINE void SuspendRendering()
+	{
+		RHISuspendRendering();
+	}
+	
+	FORCEINLINE void ResumeRendering()
+	{
+		RHIResumeRendering();
+	}
+	
+	FORCEINLINE bool IsRenderingSuspended()
+	{
+		QUICK_SCOPE_CYCLE_COUNTER(STAT_RHIMETHOD_IsRenderingSuspended_Flush);
+		ImmediateFlush(EImmediateFlushType::FlushRHIThread); 
+		return GDynamicRHI->RHIIsRenderingSuspended();
+	}
+	
+	FORCEINLINE bool GetAvailableResolutions(FScreenResolutionArray& Resolutions, bool bIgnoreRefreshRate)
+	{
+		return RHIGetAvailableResolutions(Resolutions, bIgnoreRefreshRate);
+	}
+	
+	FORCEINLINE void GetSupportedResolution(uint32& Width, uint32& Height)
+	{
+		RHIGetSupportedResolution(Width, Height);
+	}
+	
+	FORCEINLINE void VirtualTextureSetFirstMipInMemory(FTexture2DRHIParamRef Texture, uint32 FirstMip)
+	{
+		if(GRHIThread)
+		{
+			QUICK_SCOPE_CYCLE_COUNTER(STAT_RHIMETHOD_VirtualTextureSetFirstMipInMemory_WaitRHI);
+			ImmediateFlush(EImmediateFlushType::WaitForRHIThread);
+		}
+		GDynamicRHI->RHIVirtualTextureSetFirstMipInMemory(Texture, FirstMip);
+	}
+	
+	FORCEINLINE void VirtualTextureSetFirstMipVisible(FTexture2DRHIParamRef Texture, uint32 FirstMip)
+	{
+		if(GRHIThread)
+		{
+			QUICK_SCOPE_CYCLE_COUNTER(STAT_RHIMETHOD_VirtualTextureSetFirstMipVisible_WaitRHI);
+			ImmediateFlush(EImmediateFlushType::WaitForRHIThread);
+		}
+		GDynamicRHI->RHIVirtualTextureSetFirstMipVisible(Texture, FirstMip);
+	}
+	
+	FORCEINLINE void ExecuteCommandList(FRHICommandList* CmdList)
+	{
+		if(GRHIThread)
+		{
+			QUICK_SCOPE_CYCLE_COUNTER(STAT_RHIMETHOD_ExecuteCommandList_WaitRHI);
+			ImmediateFlush(EImmediateFlushType::WaitForRHIThread);
+		}
+		GDynamicRHI->RHIExecuteCommandList(CmdList);
+	}
+	
+	FORCEINLINE void* GetNativeDevice()
+	{
+		QUICK_SCOPE_CYCLE_COUNTER(STAT_RHIMETHOD_GetNativeDevice_Flush);
+		ImmediateFlush(EImmediateFlushType::FlushRHIThread); 
+		TGuardValue<bool> GuardIsFlushedGlobal( bFlushedGlobal, true ); 
+		return GDynamicRHI->RHIGetNativeDevice();
+	}
+	
+	FORCEINLINE class IRHICommandContext* GetDefaultContext()
+	{
+		return RHIGetDefaultContext();
+	}
+	
+	FORCEINLINE class IRHICommandContextContainer* GetCommandContextContainer()
+	{
+		return RHIGetCommandContextContainer();
+	}
+	
 
 };
 
@@ -2072,27 +2730,268 @@ struct FScopedCommandListWaitForTasks
 	}
 };
 
-#define DEFINE_RHIMETHOD_CMDLIST(Type,Name,ParameterTypesAndNames,ParameterNames,ReturnStatement,NullImplementation)
-#define DEFINE_RHIMETHOD(Type, Name, ParameterTypesAndNames, ParameterNames, ReturnStatement, NullImplementation)
-#define DEFINE_RHIMETHOD_GLOBALTHREADSAFE(Type, Name, ParameterTypesAndNames, ParameterNames, ReturnStatement, NullImplementation)
 
-#define DEFINE_RHIMETHOD_GLOBAL(Type, Name, ParameterTypesAndNames, ParameterNames, ReturnStatement, NullImplementation) \
-	FORCEINLINE_DEBUGGABLE Type RHI##Name ParameterTypesAndNames \
-	{ \
-		ReturnStatement FRHICommandListExecutor::GetImmediateCommandList().Name ParameterNames; \
-	}
+FORCEINLINE FVertexDeclarationRHIRef RHICreateVertexDeclaration(const FVertexDeclarationElementList& Elements)
+{
+	return FRHICommandListExecutor::GetImmediateCommandList().CreateVertexDeclaration(Elements);
+}
+
+FORCEINLINE FPixelShaderRHIRef RHICreatePixelShader(const TArray<uint8>& Code)
+{
+	return FRHICommandListExecutor::GetImmediateCommandList().CreatePixelShader(Code);
+}
+
+FORCEINLINE FVertexShaderRHIRef RHICreateVertexShader(const TArray<uint8>& Code)
+{
+	return FRHICommandListExecutor::GetImmediateCommandList().CreateVertexShader(Code);
+}
+
+FORCEINLINE FHullShaderRHIRef RHICreateHullShader(const TArray<uint8>& Code)
+{
+	return FRHICommandListExecutor::GetImmediateCommandList().CreateHullShader(Code);
+}
+
+FORCEINLINE FDomainShaderRHIRef RHICreateDomainShader(const TArray<uint8>& Code)
+{
+	return FRHICommandListExecutor::GetImmediateCommandList().CreateDomainShader(Code);
+}
+
+FORCEINLINE FGeometryShaderRHIRef RHICreateGeometryShader(const TArray<uint8>& Code)
+{
+	return FRHICommandListExecutor::GetImmediateCommandList().CreateGeometryShader(Code);
+}
+
+FORCEINLINE FGeometryShaderRHIRef RHICreateGeometryShaderWithStreamOutput(const TArray<uint8>& Code, const FStreamOutElementList& ElementList, uint32 NumStrides, const uint32* Strides, int32 RasterizedStream)
+{
+	return FRHICommandListExecutor::GetImmediateCommandList().CreateGeometryShaderWithStreamOutput(Code, ElementList, NumStrides, Strides, RasterizedStream);
+}
+
+FORCEINLINE FComputeShaderRHIRef RHICreateComputeShader(const TArray<uint8>& Code)
+{
+	return FRHICommandListExecutor::GetImmediateCommandList().CreateComputeShader(Code);
+}
+
+FORCEINLINE FIndexBufferRHIRef RHICreateIndexBuffer(uint32 Stride, uint32 Size, uint32 InUsage, FRHIResourceCreateInfo& CreateInfo)
+{
+	return FRHICommandListExecutor::GetImmediateCommandList().CreateIndexBuffer(Stride, Size, InUsage, CreateInfo);
+}
+
+FORCEINLINE void* RHILockIndexBuffer(FIndexBufferRHIParamRef IndexBuffer, uint32 Offset, uint32 Size, EResourceLockMode LockMode)
+{
+	return FRHICommandListExecutor::GetImmediateCommandList().LockIndexBuffer(IndexBuffer, Offset, Size, LockMode);
+}
+
+FORCEINLINE void RHIUnlockIndexBuffer(FIndexBufferRHIParamRef IndexBuffer)
+{
+	 FRHICommandListExecutor::GetImmediateCommandList().UnlockIndexBuffer(IndexBuffer);
+}
+
+FORCEINLINE FVertexBufferRHIRef RHICreateVertexBuffer(uint32 Size, uint32 InUsage, FRHIResourceCreateInfo& CreateInfo)
+{
+	return FRHICommandListExecutor::GetImmediateCommandList().CreateVertexBuffer(Size, InUsage, CreateInfo);
+}
+
+FORCEINLINE void* RHILockVertexBuffer(FVertexBufferRHIParamRef VertexBuffer, uint32 Offset, uint32 SizeRHI, EResourceLockMode LockMode)
+{
+	return FRHICommandListExecutor::GetImmediateCommandList().LockVertexBuffer(VertexBuffer, Offset, SizeRHI, LockMode);
+}
+
+FORCEINLINE void RHIUnlockVertexBuffer(FVertexBufferRHIParamRef VertexBuffer)
+{
+	 FRHICommandListExecutor::GetImmediateCommandList().UnlockVertexBuffer(VertexBuffer);
+}
+
+FORCEINLINE FStructuredBufferRHIRef RHICreateStructuredBuffer(uint32 Stride, uint32 Size, uint32 InUsage, FRHIResourceCreateInfo& CreateInfo)
+{
+	return FRHICommandListExecutor::GetImmediateCommandList().CreateStructuredBuffer(Stride, Size, InUsage, CreateInfo);
+}
+
+FORCEINLINE void* RHILockStructuredBuffer(FStructuredBufferRHIParamRef StructuredBuffer, uint32 Offset, uint32 SizeRHI, EResourceLockMode LockMode)
+{
+	return FRHICommandListExecutor::GetImmediateCommandList().LockStructuredBuffer(StructuredBuffer, Offset, SizeRHI, LockMode);
+}
+
+FORCEINLINE void RHIUnlockStructuredBuffer(FStructuredBufferRHIParamRef StructuredBuffer)
+{
+	 FRHICommandListExecutor::GetImmediateCommandList().UnlockStructuredBuffer(StructuredBuffer);
+}
+
+FORCEINLINE FUnorderedAccessViewRHIRef RHICreateUnorderedAccessView(FStructuredBufferRHIParamRef StructuredBuffer, bool bUseUAVCounter, bool bAppendBuffer)
+{
+	return FRHICommandListExecutor::GetImmediateCommandList().CreateUnorderedAccessView(StructuredBuffer, bUseUAVCounter, bAppendBuffer);
+}
+
+FORCEINLINE FUnorderedAccessViewRHIRef RHICreateUnorderedAccessView(FTextureRHIParamRef Texture)
+{
+	return FRHICommandListExecutor::GetImmediateCommandList().CreateUnorderedAccessView(Texture);
+}
+
+FORCEINLINE FUnorderedAccessViewRHIRef RHICreateUnorderedAccessView(FVertexBufferRHIParamRef VertexBuffer, uint8 Format)
+{
+	return FRHICommandListExecutor::GetImmediateCommandList().CreateUnorderedAccessView(VertexBuffer, Format);
+}
+
+FORCEINLINE FShaderResourceViewRHIRef RHICreateShaderResourceView(FStructuredBufferRHIParamRef StructuredBuffer)
+{
+	return FRHICommandListExecutor::GetImmediateCommandList().CreateShaderResourceView(StructuredBuffer);
+}
+
+FORCEINLINE FShaderResourceViewRHIRef RHICreateShaderResourceView(FVertexBufferRHIParamRef VertexBuffer, uint32 Stride, uint8 Format)
+{
+	return FRHICommandListExecutor::GetImmediateCommandList().CreateShaderResourceView(VertexBuffer, Stride, Format);
+}
+
+FORCEINLINE FTextureReferenceRHIRef RHICreateTextureReference(FLastRenderTimeContainer* LastRenderTime)
+{
+	return FRHICommandListExecutor::GetImmediateCommandList().CreateTextureReference(LastRenderTime);
+}
+
+FORCEINLINE void RHIUpdateTextureReference(FTextureReferenceRHIParamRef TextureRef, FTextureRHIParamRef NewTexture)
+{
+	 FRHICommandListExecutor::GetImmediateCommandList().UpdateTextureReference(TextureRef, NewTexture);
+}
+
+FORCEINLINE FTexture2DRHIRef RHICreateTexture2D(uint32 SizeX, uint32 SizeY, uint8 Format, uint32 NumMips, uint32 NumSamples, uint32 Flags, FRHIResourceCreateInfo& CreateInfo)
+{
+	return FRHICommandListExecutor::GetImmediateCommandList().CreateTexture2D(SizeX, SizeY, Format, NumMips, NumSamples, Flags, CreateInfo);
+}
+
+FORCEINLINE FTexture2DRHIRef RHIAsyncCreateTexture2D(uint32 SizeX, uint32 SizeY, uint8 Format, uint32 NumMips, uint32 Flags, void** InitialMipData, uint32 NumInitialMips)
+{
+	return FRHICommandListExecutor::GetImmediateCommandList().AsyncCreateTexture2D(SizeX, SizeY, Format, NumMips, Flags, InitialMipData, NumInitialMips);
+}
+
+FORCEINLINE void RHICopySharedMips(FTexture2DRHIParamRef DestTexture2D, FTexture2DRHIParamRef SrcTexture2D)
+{
+	return FRHICommandListExecutor::GetImmediateCommandList().CopySharedMips(DestTexture2D, SrcTexture2D);
+}
+
+FORCEINLINE FTexture2DArrayRHIRef RHICreateTexture2DArray(uint32 SizeX, uint32 SizeY, uint32 SizeZ, uint8 Format, uint32 NumMips, uint32 Flags, FRHIResourceCreateInfo& CreateInfo)
+{
+	return FRHICommandListExecutor::GetImmediateCommandList().CreateTexture2DArray(SizeX, SizeY, SizeZ, Format, NumMips, Flags, CreateInfo);
+}
+
+FORCEINLINE FTexture3DRHIRef RHICreateTexture3D(uint32 SizeX, uint32 SizeY, uint32 SizeZ, uint8 Format, uint32 NumMips, uint32 Flags, FRHIResourceCreateInfo& CreateInfo)
+{
+	return FRHICommandListExecutor::GetImmediateCommandList().CreateTexture3D(SizeX, SizeY, SizeZ, Format, NumMips, Flags, CreateInfo);
+}
+
+FORCEINLINE FShaderResourceViewRHIRef RHICreateShaderResourceView(FTexture2DRHIParamRef Texture2DRHI, uint8 MipLevel)
+{
+	return FRHICommandListExecutor::GetImmediateCommandList().CreateShaderResourceView(Texture2DRHI, MipLevel);
+}
+
+FORCEINLINE FShaderResourceViewRHIRef RHICreateShaderResourceView(FTexture2DRHIParamRef Texture2DRHI, uint8 MipLevel, uint8 NumMipLevels, uint8 Format)
+{
+	return FRHICommandListExecutor::GetImmediateCommandList().CreateShaderResourceView(Texture2DRHI, MipLevel, NumMipLevels, Format);
+}
+
+FORCEINLINE FTexture2DRHIRef RHIAsyncReallocateTexture2D(FTexture2DRHIParamRef Texture2D, int32 NewMipCount, int32 NewSizeX, int32 NewSizeY, FThreadSafeCounter* RequestStatus)
+{
+	return FRHICommandListExecutor::GetImmediateCommandList().AsyncReallocateTexture2D(Texture2D, NewMipCount, NewSizeX, NewSizeY, RequestStatus);
+}
+
+FORCEINLINE ETextureReallocationStatus RHIFinalizeAsyncReallocateTexture2D(FTexture2DRHIParamRef Texture2D, bool bBlockUntilCompleted)
+{
+	return FRHICommandListExecutor::GetImmediateCommandList().FinalizeAsyncReallocateTexture2D(Texture2D, bBlockUntilCompleted);
+}
+
+FORCEINLINE ETextureReallocationStatus RHICancelAsyncReallocateTexture2D(FTexture2DRHIParamRef Texture2D, bool bBlockUntilCompleted)
+{
+	return FRHICommandListExecutor::GetImmediateCommandList().CancelAsyncReallocateTexture2D(Texture2D, bBlockUntilCompleted);
+}
+
+FORCEINLINE void* RHILockTexture2D(FTexture2DRHIParamRef Texture, uint32 MipIndex, EResourceLockMode LockMode, uint32& DestStride, bool bLockWithinMiptail)
+{
+	return FRHICommandListExecutor::GetImmediateCommandList().LockTexture2D(Texture, MipIndex, LockMode, DestStride, bLockWithinMiptail);
+}
+
+FORCEINLINE void RHIUnlockTexture2D(FTexture2DRHIParamRef Texture, uint32 MipIndex, bool bLockWithinMiptail)
+{
+	 FRHICommandListExecutor::GetImmediateCommandList().UnlockTexture2D(Texture, MipIndex, bLockWithinMiptail);
+}
+
+FORCEINLINE void* RHILockTexture2DArray(FTexture2DArrayRHIParamRef Texture, uint32 TextureIndex, uint32 MipIndex, EResourceLockMode LockMode, uint32& DestStride, bool bLockWithinMiptail)
+{
+	return FRHICommandListExecutor::GetImmediateCommandList().LockTexture2DArray(Texture, TextureIndex, MipIndex, LockMode, DestStride, bLockWithinMiptail);
+}
+
+FORCEINLINE void RHIUnlockTexture2DArray(FTexture2DArrayRHIParamRef Texture, uint32 TextureIndex, uint32 MipIndex, bool bLockWithinMiptail)
+{
+	 FRHICommandListExecutor::GetImmediateCommandList().UnlockTexture2DArray(Texture, TextureIndex, MipIndex, bLockWithinMiptail);
+}
+
+FORCEINLINE void RHIUpdateTexture2D(FTexture2DRHIParamRef Texture, uint32 MipIndex, const struct FUpdateTextureRegion2D& UpdateRegion, uint32 SourcePitch, const uint8* SourceData)
+{
+	 FRHICommandListExecutor::GetImmediateCommandList().UpdateTexture2D(Texture, MipIndex, UpdateRegion, SourcePitch, SourceData);
+}
+
+FORCEINLINE void RHIUpdateTexture3D(FTexture3DRHIParamRef Texture, uint32 MipIndex, const struct FUpdateTextureRegion3D& UpdateRegion, uint32 SourceRowPitch, uint32 SourceDepthPitch, const uint8* SourceData)
+{
+	 FRHICommandListExecutor::GetImmediateCommandList().UpdateTexture3D(Texture, MipIndex, UpdateRegion, SourceRowPitch, SourceDepthPitch, SourceData);
+}
+
+FORCEINLINE FTextureCubeRHIRef RHICreateTextureCube(uint32 Size, uint8 Format, uint32 NumMips, uint32 Flags, FRHIResourceCreateInfo& CreateInfo)
+{
+	return FRHICommandListExecutor::GetImmediateCommandList().CreateTextureCube(Size, Format, NumMips, Flags, CreateInfo);
+}
+
+FORCEINLINE FTextureCubeRHIRef RHICreateTextureCubeArray(uint32 Size, uint32 ArraySize, uint8 Format, uint32 NumMips, uint32 Flags, FRHIResourceCreateInfo& CreateInfo)
+{
+	return FRHICommandListExecutor::GetImmediateCommandList().CreateTextureCubeArray(Size, ArraySize, Format, NumMips, Flags, CreateInfo);
+}
+
+FORCEINLINE void* RHILockTextureCubeFace(FTextureCubeRHIParamRef Texture, uint32 FaceIndex, uint32 ArrayIndex, uint32 MipIndex, EResourceLockMode LockMode, uint32& DestStride, bool bLockWithinMiptail)
+{
+	return FRHICommandListExecutor::GetImmediateCommandList().LockTextureCubeFace(Texture, FaceIndex, ArrayIndex, MipIndex, LockMode, DestStride, bLockWithinMiptail);
+}
+
+FORCEINLINE void RHIUnlockTextureCubeFace(FTextureCubeRHIParamRef Texture, uint32 FaceIndex, uint32 ArrayIndex, uint32 MipIndex, bool bLockWithinMiptail)
+{
+	 FRHICommandListExecutor::GetImmediateCommandList().UnlockTextureCubeFace(Texture, FaceIndex, ArrayIndex, MipIndex, bLockWithinMiptail);
+}
+
+FORCEINLINE FRenderQueryRHIRef RHICreateRenderQuery(ERenderQueryType QueryType)
+{
+	return FRHICommandListExecutor::GetImmediateCommandList().CreateRenderQuery(QueryType);
+}
+
+FORCEINLINE void RHIAcquireThreadOwnership()
+{
+	return FRHICommandListExecutor::GetImmediateCommandList().AcquireThreadOwnership();
+}
+
+FORCEINLINE void RHIReleaseThreadOwnership()
+{
+	return FRHICommandListExecutor::GetImmediateCommandList().ReleaseThreadOwnership();
+}
+
+FORCEINLINE void RHIFlushResources()
+{
+	return FRHICommandListExecutor::GetImmediateCommandList().FlushResources();
+}
+
+FORCEINLINE void RHIVirtualTextureSetFirstMipInMemory(FTexture2DRHIParamRef Texture, uint32 FirstMip)
+{
+	 FRHICommandListExecutor::GetImmediateCommandList().VirtualTextureSetFirstMipInMemory(Texture, FirstMip);
+}
+
+FORCEINLINE void RHIVirtualTextureSetFirstMipVisible(FTexture2DRHIParamRef Texture, uint32 FirstMip)
+{
+	 FRHICommandListExecutor::GetImmediateCommandList().VirtualTextureSetFirstMipVisible(Texture, FirstMip);
+}
+
+FORCEINLINE void RHIExecuteCommandList(FRHICommandList* CmdList)
+{
+	 FRHICommandListExecutor::GetImmediateCommandList().ExecuteCommandList(CmdList);
+}
+
+FORCEINLINE void* RHIGetNativeDevice()
+{
+	return FRHICommandListExecutor::GetImmediateCommandList().GetNativeDevice();
+}
+
 
-#define DEFINE_RHIMETHOD_GLOBALFLUSH(Type, Name, ParameterTypesAndNames, ParameterNames, ReturnStatement, NullImplementation) \
-	FORCEINLINE_DEBUGGABLE Type RHI##Name ParameterTypesAndNames \
-	{ \
-		ReturnStatement FRHICommandListExecutor::GetImmediateCommandList().Name ParameterNames; \
-	}
-#include "RHIMethods.h"
-#undef DEFINE_RHIMETHOD
-#undef DEFINE_RHIMETHOD_CMDLIST
-#undef DEFINE_RHIMETHOD_GLOBAL
-#undef DEFINE_RHIMETHOD_GLOBALFLUSH
-#undef DEFINE_RHIMETHOD_GLOBALTHREADSAFE
 
 
 #include "RHICommandList.inl"
diff --git a/Engine/Source/Runtime/RHI/Public/RHICommandListCommandExecutes.inl b/Engine/Source/Runtime/RHI/Public/RHICommandListCommandExecutes.inl
--- a/Engine/Source/Runtime/RHI/Public/RHICommandListCommandExecutes.inl
+++ b/Engine/Source/Runtime/RHI/Public/RHICommandListCommandExecutes.inl
@@ -390,9 +390,9 @@ void FRHICommandEndScene::Execute(FRHICommandListBase& CmdList)
 void FRHICommandUpdateVertexBuffer::Execute(FRHICommandListBase& CmdList)
 {
 	RHISTAT(UpdateVertexBuffer);
-	void* Data = LockVertexBuffer_Internal(VertexBuffer, 0, BufferSize, RLM_WriteOnly);
+	void* Data = GDynamicRHI->RHILockVertexBuffer(VertexBuffer, 0, BufferSize, RLM_WriteOnly);
 	FMemory::Memcpy(Data, Buffer, BufferSize);
-	UnlockVertexBuffer_Internal(VertexBuffer);
+	GDynamicRHI->RHIUnlockVertexBuffer(VertexBuffer);
 }
 
 void FRHICommandBeginFrame::Execute(FRHICommandListBase& CmdList)
diff --git a/Engine/Source/Runtime/RHI/Public/RHIMethods.h b/Engine/Source/Runtime/RHI/Public/RHIMethods.h
deleted file mode 100644
diff --git a/Engine/Source/Runtime/Windows/D3D11RHI/Private/D3D11RHIPrivate.h b/Engine/Source/Runtime/Windows/D3D11RHI/Private/D3D11RHIPrivate.h
--- a/Engine/Source/Runtime/Windows/D3D11RHI/Private/D3D11RHIPrivate.h
+++ b/Engine/Source/Runtime/Windows/D3D11RHI/Private/D3D11RHIPrivate.h
@@ -304,9 +304,177 @@ public:
 	 */
 	bool GetQueryData(ID3D11Query* Query,void* Data,SIZE_T DataSize,bool bWait, ERenderQueryType QueryType);
 
-	#define DEFINE_RHIMETHOD(Type,Name,ParameterTypesAndNames,ParameterNames,ReturnStatement,NullImplementation) virtual Type RHI##Name ParameterTypesAndNames
-	#include "RHIMethods.h"
-	#undef DEFINE_RHIMETHOD
+	virtual void RHIGpuTimeBegin(uint32 Hash,bool bCompute) final override;
+	virtual void RHIGpuTimeEnd(uint32 Hash,bool bCompute) final override;
+	virtual FSamplerStateRHIRef RHICreateSamplerState(const FSamplerStateInitializerRHI& Initializer) final override;
+	virtual FRasterizerStateRHIRef RHICreateRasterizerState(const FRasterizerStateInitializerRHI& Initializer) final override;
+	virtual FDepthStencilStateRHIRef RHICreateDepthStencilState(const FDepthStencilStateInitializerRHI& Initializer) final override;
+	virtual FBlendStateRHIRef RHICreateBlendState(const FBlendStateInitializerRHI& Initializer) final override;
+	virtual FVertexDeclarationRHIRef RHICreateVertexDeclaration(const FVertexDeclarationElementList& Elements) final override;
+	virtual FPixelShaderRHIRef RHICreatePixelShader(const TArray<uint8>& Code) final override;
+	virtual FVertexShaderRHIRef RHICreateVertexShader(const TArray<uint8>& Code) final override;
+	virtual FHullShaderRHIRef RHICreateHullShader(const TArray<uint8>& Code) final override;
+	virtual FDomainShaderRHIRef RHICreateDomainShader(const TArray<uint8>& Code) final override;
+	virtual FGeometryShaderRHIRef RHICreateGeometryShader(const TArray<uint8>& Code) final override;
+	virtual FGeometryShaderRHIRef RHICreateGeometryShaderWithStreamOutput(const TArray<uint8>& Code, const FStreamOutElementList& ElementList, uint32 NumStrides, const uint32* Strides, int32 RasterizedStream) final override;
+	virtual FComputeShaderRHIRef RHICreateComputeShader(const TArray<uint8>& Code) final override;
+	virtual FBoundShaderStateRHIRef RHICreateBoundShaderState(FVertexDeclarationRHIParamRef VertexDeclaration, FVertexShaderRHIParamRef VertexShader, FHullShaderRHIParamRef HullShader, FDomainShaderRHIParamRef DomainShader, FPixelShaderRHIParamRef PixelShader, FGeometryShaderRHIParamRef GeometryShader) final override;
+	virtual FUniformBufferRHIRef RHICreateUniformBuffer(const void* Contents, const FRHIUniformBufferLayout& Layout, EUniformBufferUsage Usage) final override;
+	virtual FIndexBufferRHIRef RHICreateIndexBuffer(uint32 Stride, uint32 Size, uint32 InUsage, FRHIResourceCreateInfo& CreateInfo) final override;
+	virtual void* RHILockIndexBuffer(FIndexBufferRHIParamRef IndexBuffer, uint32 Offset, uint32 Size, EResourceLockMode LockMode) final override;
+	virtual void RHIUnlockIndexBuffer(FIndexBufferRHIParamRef IndexBuffer) final override;
+	virtual FVertexBufferRHIRef RHICreateVertexBuffer(uint32 Size, uint32 InUsage, FRHIResourceCreateInfo& CreateInfo) final override;
+	virtual void* RHILockVertexBuffer(FVertexBufferRHIParamRef VertexBuffer, uint32 Offset, uint32 SizeRHI, EResourceLockMode LockMode) final override;
+	virtual void RHIUnlockVertexBuffer(FVertexBufferRHIParamRef VertexBuffer) final override;
+	virtual void RHICopyVertexBuffer(FVertexBufferRHIParamRef SourceBuffer,FVertexBufferRHIParamRef DestBuffer) final override;
+	virtual FStructuredBufferRHIRef RHICreateStructuredBuffer(uint32 Stride, uint32 Size, uint32 InUsage, FRHIResourceCreateInfo& CreateInfo) final override;
+	virtual void* RHILockStructuredBuffer(FStructuredBufferRHIParamRef StructuredBuffer, uint32 Offset, uint32 SizeRHI, EResourceLockMode LockMode) final override;
+	virtual void RHIUnlockStructuredBuffer(FStructuredBufferRHIParamRef StructuredBuffer) final override;
+	virtual FUnorderedAccessViewRHIRef RHICreateUnorderedAccessView(FStructuredBufferRHIParamRef StructuredBuffer, bool bUseUAVCounter, bool bAppendBuffer) final override;
+	virtual FUnorderedAccessViewRHIRef RHICreateUnorderedAccessView(FTextureRHIParamRef Texture) final override;
+	virtual FUnorderedAccessViewRHIRef RHICreateUnorderedAccessView(FVertexBufferRHIParamRef VertexBuffer, uint8 Format) final override;
+	virtual FShaderResourceViewRHIRef RHICreateShaderResourceView(FStructuredBufferRHIParamRef StructuredBuffer) final override;
+	virtual FShaderResourceViewRHIRef RHICreateShaderResourceView(FVertexBufferRHIParamRef VertexBuffer, uint32 Stride, uint8 Format) final override;
+	virtual uint64 RHICalcTexture2DPlatformSize(uint32 SizeX, uint32 SizeY, uint8 Format, uint32 NumMips, uint32 NumSamples, uint32 Flags, uint32& OutAlign) final override;
+	virtual uint64 RHICalcTexture3DPlatformSize(uint32 SizeX, uint32 SizeY, uint32 SizeZ, uint8 Format, uint32 NumMips, uint32 Flags, uint32& OutAlign) final override;
+	virtual uint64 RHICalcTextureCubePlatformSize(uint32 Size, uint8 Format, uint32 NumMips, uint32 Flags, uint32& OutAlign) final override;
+	virtual void RHIGetTextureMemoryStats(FTextureMemoryStats& OutStats) final override;
+	virtual bool RHIGetTextureMemoryVisualizeData(FColor* TextureData,int32 SizeX,int32 SizeY,int32 Pitch,int32 PixelSize) final override;
+	virtual FTextureReferenceRHIRef RHICreateTextureReference(FLastRenderTimeContainer* LastRenderTime) final override;
+	virtual void RHIUpdateTextureReference(FTextureReferenceRHIParamRef TextureRef, FTextureRHIParamRef NewTexture) final override;
+	virtual FTexture2DRHIRef RHICreateTexture2D(uint32 SizeX, uint32 SizeY, uint8 Format, uint32 NumMips, uint32 NumSamples, uint32 Flags, FRHIResourceCreateInfo& CreateInfo) final override;
+	virtual FTexture2DRHIRef RHIAsyncCreateTexture2D(uint32 SizeX, uint32 SizeY, uint8 Format, uint32 NumMips, uint32 Flags, void** InitialMipData, uint32 NumInitialMips) final override;
+	virtual void RHICopySharedMips(FTexture2DRHIParamRef DestTexture2D, FTexture2DRHIParamRef SrcTexture2D) final override;
+	virtual FTexture2DArrayRHIRef RHICreateTexture2DArray(uint32 SizeX, uint32 SizeY, uint32 SizeZ, uint8 Format, uint32 NumMips, uint32 Flags, FRHIResourceCreateInfo& CreateInfo) final override;
+	virtual FTexture3DRHIRef RHICreateTexture3D(uint32 SizeX, uint32 SizeY, uint32 SizeZ, uint8 Format, uint32 NumMips, uint32 Flags, FRHIResourceCreateInfo& CreateInfo) final override;
+	virtual void RHIGetResourceInfo(FTextureRHIParamRef Ref, FRHIResourceInfo& OutInfo) final override;
+	virtual FShaderResourceViewRHIRef RHICreateShaderResourceView(FTexture2DRHIParamRef Texture2DRHI, uint8 MipLevel) final override;
+	virtual FShaderResourceViewRHIRef RHICreateShaderResourceView(FTexture2DRHIParamRef Texture2DRHI, uint8 MipLevel, uint8 NumMipLevels, uint8 Format) final override;
+	virtual void RHIGenerateMips(FTextureRHIParamRef Texture) final override;
+	virtual uint32 RHIComputeMemorySize(FTextureRHIParamRef TextureRHI) final override;
+	virtual FTexture2DRHIRef RHIAsyncReallocateTexture2D(FTexture2DRHIParamRef Texture2D, int32 NewMipCount, int32 NewSizeX, int32 NewSizeY, FThreadSafeCounter* RequestStatus) final override;
+	virtual ETextureReallocationStatus RHIFinalizeAsyncReallocateTexture2D(FTexture2DRHIParamRef Texture2D, bool bBlockUntilCompleted) final override;
+	virtual ETextureReallocationStatus RHICancelAsyncReallocateTexture2D(FTexture2DRHIParamRef Texture2D, bool bBlockUntilCompleted) final override;
+	virtual void* RHILockTexture2D(FTexture2DRHIParamRef Texture, uint32 MipIndex, EResourceLockMode LockMode, uint32& DestStride, bool bLockWithinMiptail) final override;
+	virtual void RHIUnlockTexture2D(FTexture2DRHIParamRef Texture, uint32 MipIndex, bool bLockWithinMiptail) final override;
+	virtual void* RHILockTexture2DArray(FTexture2DArrayRHIParamRef Texture, uint32 TextureIndex, uint32 MipIndex, EResourceLockMode LockMode, uint32& DestStride, bool bLockWithinMiptail) final override;
+	virtual void RHIUnlockTexture2DArray(FTexture2DArrayRHIParamRef Texture, uint32 TextureIndex, uint32 MipIndex, bool bLockWithinMiptail) final override;
+	virtual void RHIUpdateTexture2D(FTexture2DRHIParamRef Texture, uint32 MipIndex, const struct FUpdateTextureRegion2D& UpdateRegion, uint32 SourcePitch, const uint8* SourceData) final override;
+	virtual void RHIUpdateTexture3D(FTexture3DRHIParamRef Texture, uint32 MipIndex, const struct FUpdateTextureRegion3D& UpdateRegion, uint32 SourceRowPitch, uint32 SourceDepthPitch, const uint8* SourceData) final override;
+	virtual FTextureCubeRHIRef RHICreateTextureCube(uint32 Size, uint8 Format, uint32 NumMips, uint32 Flags, FRHIResourceCreateInfo& CreateInfo) final override;
+	virtual FTextureCubeRHIRef RHICreateTextureCubeArray(uint32 Size, uint32 ArraySize, uint8 Format, uint32 NumMips, uint32 Flags, FRHIResourceCreateInfo& CreateInfo) final override;
+	virtual void* RHILockTextureCubeFace(FTextureCubeRHIParamRef Texture, uint32 FaceIndex, uint32 ArrayIndex, uint32 MipIndex, EResourceLockMode LockMode, uint32& DestStride, bool bLockWithinMiptail) final override;
+	virtual void RHIUnlockTextureCubeFace(FTextureCubeRHIParamRef Texture, uint32 FaceIndex, uint32 ArrayIndex, uint32 MipIndex, bool bLockWithinMiptail) final override;
+	virtual void RHIBindDebugLabelName(FTextureRHIParamRef Texture, const TCHAR* Name) final override;
+	virtual void RHIReadSurfaceData(FTextureRHIParamRef Texture,FIntRect Rect,TArray<FColor>& OutData,FReadSurfaceDataFlags InFlags) final override;
+	virtual void RHIMapStagingSurface(FTextureRHIParamRef Texture,void*& OutData,int32& OutWidth,int32& OutHeight) final override;
+	virtual void RHIUnmapStagingSurface(FTextureRHIParamRef Texture) final override;
+	virtual void RHIReadSurfaceFloatData(FTextureRHIParamRef Texture,FIntRect Rect,TArray<FFloat16Color>& OutData,ECubeFace CubeFace,int32 ArrayIndex,int32 MipIndex) final override;
+	virtual void RHIRead3DSurfaceFloatData(FTextureRHIParamRef Texture,FIntRect Rect,FIntPoint ZMinMax,TArray<FFloat16Color>& OutData) final override;
+	virtual FRenderQueryRHIRef RHICreateRenderQuery(ERenderQueryType QueryType) final override;
+	virtual bool RHIGetRenderQueryResult(FRenderQueryRHIParamRef RenderQuery, uint64& OutResult, bool bWait) final override;
+	virtual FTexture2DRHIRef RHIGetViewportBackBuffer(FViewportRHIParamRef Viewport) final override;
+	virtual void RHIAdvanceFrameForGetViewportBackBuffer() final override;
+	virtual bool RHIIsDrawingViewport() final override;
+	virtual void RHIAcquireThreadOwnership() final override;
+	virtual void RHIReleaseThreadOwnership() final override;
+	virtual void RHIFlushResources() final override;
+	virtual uint32 RHIGetGPUFrameCycles() final override;
+	virtual FViewportRHIRef RHICreateViewport(void* WindowHandle, uint32 SizeX, uint32 SizeY, bool bIsFullscreen, EPixelFormat PreferredPixelFormat) final override;
+	virtual void RHIResizeViewport(FViewportRHIParamRef Viewport, uint32 SizeX, uint32 SizeY, bool bIsFullscreen) final override;
+	virtual void RHITick(float DeltaTime) final override;
+	virtual void RHISetStreamOutTargets(uint32 NumTargets,const FVertexBufferRHIParamRef* VertexBuffers,const uint32* Offsets) final override;
+	virtual void RHIDiscardRenderTargets(bool Depth,bool Stencil,uint32 ColorBitMask) final override;
+	virtual void RHIBlockUntilGPUIdle() final override;
+	virtual void RHISuspendRendering() final override;
+	virtual void RHIResumeRendering() final override;
+	virtual bool RHIIsRenderingSuspended() final override;
+	virtual bool RHIGetAvailableResolutions(FScreenResolutionArray& Resolutions, bool bIgnoreRefreshRate) final override;
+	virtual void RHIGetSupportedResolution(uint32& Width, uint32& Height) final override;
+	virtual void RHIVirtualTextureSetFirstMipInMemory(FTexture2DRHIParamRef Texture, uint32 FirstMip) final override;
+	virtual void RHIVirtualTextureSetFirstMipVisible(FTexture2DRHIParamRef Texture, uint32 FirstMip) final override;
+	virtual void RHIExecuteCommandList(FRHICommandList* CmdList) final override;
+	virtual void* RHIGetNativeDevice() final override;
+	virtual class IRHICommandContext* RHIGetDefaultContext() final override;
+	virtual class IRHICommandContextContainer* RHIGetCommandContextContainer() final override;
+
+
+
+	virtual void RHISetComputeShader(FComputeShaderRHIParamRef ComputeShader) final override;
+	virtual void RHIDispatchComputeShader(uint32 ThreadGroupCountX, uint32 ThreadGroupCountY, uint32 ThreadGroupCountZ) final override;
+	virtual void RHIDispatchIndirectComputeShader(FVertexBufferRHIParamRef ArgumentBuffer, uint32 ArgumentOffset) final override;
+	virtual void RHIAutomaticCacheFlushAfterComputeShader(bool bEnable) final override;
+	virtual void RHIFlushComputeShaderCache() final override;
+	virtual void RHISetMultipleViewports(uint32 Count, const FViewportBounds* Data) final override;
+	virtual void RHIClearUAV(FUnorderedAccessViewRHIParamRef UnorderedAccessViewRHI, const uint32* Values) final override;
+	virtual void RHICopyToResolveTarget(FTextureRHIParamRef SourceTexture, FTextureRHIParamRef DestTexture, bool bKeepOriginalSurface, const FResolveParams& ResolveParams) final override;
+	virtual void RHIBeginRenderQuery(FRenderQueryRHIParamRef RenderQuery) final override;
+	virtual void RHIEndRenderQuery(FRenderQueryRHIParamRef RenderQuery) final override;
+	virtual void RHIBeginOcclusionQueryBatch() final override;
+	virtual void RHIEndOcclusionQueryBatch() final override;
+	virtual void RHISubmitCommandsHint() final override;
+	virtual void RHIBeginDrawingViewport(FViewportRHIParamRef Viewport, FTextureRHIParamRef RenderTargetRHI) final override;
+	virtual void RHIEndDrawingViewport(FViewportRHIParamRef Viewport, bool bPresent, bool bLockToVsync) final override;
+	virtual void RHIBeginFrame() final override;
+	virtual void RHIEndFrame() final override;
+	virtual void RHIBeginScene() final override;
+	virtual void RHIEndScene() final override;
+	virtual void RHISetStreamSource(uint32 StreamIndex, FVertexBufferRHIParamRef VertexBuffer, uint32 Stride, uint32 Offset) final override;
+	virtual void RHISetRasterizerState(FRasterizerStateRHIParamRef NewState) final override;
+	virtual void RHISetViewport(uint32 MinX, uint32 MinY, float MinZ, uint32 MaxX, uint32 MaxY, float MaxZ) final override;
+	virtual void RHISetScissorRect(bool bEnable, uint32 MinX, uint32 MinY, uint32 MaxX, uint32 MaxY) final override;
+	virtual void RHISetBoundShaderState(FBoundShaderStateRHIParamRef BoundShaderState) final override;
+	virtual void RHISetShaderTexture(FVertexShaderRHIParamRef VertexShader, uint32 TextureIndex, FTextureRHIParamRef NewTexture) final override;
+	virtual void RHISetShaderTexture(FHullShaderRHIParamRef HullShader, uint32 TextureIndex, FTextureRHIParamRef NewTexture) final override;
+	virtual void RHISetShaderTexture(FDomainShaderRHIParamRef DomainShader, uint32 TextureIndex, FTextureRHIParamRef NewTexture) final override;
+	virtual void RHISetShaderTexture(FGeometryShaderRHIParamRef GeometryShader, uint32 TextureIndex, FTextureRHIParamRef NewTexture) final override;
+	virtual void RHISetShaderTexture(FPixelShaderRHIParamRef PixelShader, uint32 TextureIndex, FTextureRHIParamRef NewTexture) final override;
+	virtual void RHISetShaderTexture(FComputeShaderRHIParamRef PixelShader, uint32 TextureIndex, FTextureRHIParamRef NewTexture) final override;
+	virtual void RHISetShaderSampler(FComputeShaderRHIParamRef ComputeShader, uint32 SamplerIndex, FSamplerStateRHIParamRef NewState) final override;
+	virtual void RHISetShaderSampler(FVertexShaderRHIParamRef VertexShader, uint32 SamplerIndex, FSamplerStateRHIParamRef NewState) final override;
+	virtual void RHISetShaderSampler(FGeometryShaderRHIParamRef GeometryShader, uint32 SamplerIndex, FSamplerStateRHIParamRef NewState) final override;
+	virtual void RHISetShaderSampler(FDomainShaderRHIParamRef DomainShader, uint32 SamplerIndex, FSamplerStateRHIParamRef NewState) final override;
+	virtual void RHISetShaderSampler(FHullShaderRHIParamRef HullShader, uint32 SamplerIndex, FSamplerStateRHIParamRef NewState) final override;
+	virtual void RHISetShaderSampler(FPixelShaderRHIParamRef PixelShader, uint32 SamplerIndex, FSamplerStateRHIParamRef NewState) final override;
+	virtual void RHISetUAVParameter(FComputeShaderRHIParamRef ComputeShader, uint32 UAVIndex, FUnorderedAccessViewRHIParamRef UAV) final override;
+	virtual void RHISetUAVParameter(FComputeShaderRHIParamRef ComputeShader, uint32 UAVIndex, FUnorderedAccessViewRHIParamRef UAV, uint32 InitialCount) final override;
+	virtual void RHISetShaderResourceViewParameter(FPixelShaderRHIParamRef PixelShader, uint32 SamplerIndex, FShaderResourceViewRHIParamRef SRV) final override;
+	virtual void RHISetShaderResourceViewParameter(FVertexShaderRHIParamRef VertexShader, uint32 SamplerIndex, FShaderResourceViewRHIParamRef SRV) final override;
+	virtual void RHISetShaderResourceViewParameter(FComputeShaderRHIParamRef ComputeShader, uint32 SamplerIndex, FShaderResourceViewRHIParamRef SRV) final override;
+	virtual void RHISetShaderResourceViewParameter(FHullShaderRHIParamRef HullShader, uint32 SamplerIndex, FShaderResourceViewRHIParamRef SRV) final override;
+	virtual void RHISetShaderResourceViewParameter(FDomainShaderRHIParamRef DomainShader, uint32 SamplerIndex, FShaderResourceViewRHIParamRef SRV) final override;
+	virtual void RHISetShaderResourceViewParameter(FGeometryShaderRHIParamRef GeometryShader, uint32 SamplerIndex, FShaderResourceViewRHIParamRef SRV) final override;
+	virtual void RHISetShaderUniformBuffer(FVertexShaderRHIParamRef VertexShader, uint32 BufferIndex, FUniformBufferRHIParamRef Buffer) final override;
+	virtual void RHISetShaderUniformBuffer(FHullShaderRHIParamRef HullShader, uint32 BufferIndex, FUniformBufferRHIParamRef Buffer) final override;
+	virtual void RHISetShaderUniformBuffer(FDomainShaderRHIParamRef DomainShader, uint32 BufferIndex, FUniformBufferRHIParamRef Buffer) final override;
+	virtual void RHISetShaderUniformBuffer(FGeometryShaderRHIParamRef GeometryShader, uint32 BufferIndex, FUniformBufferRHIParamRef Buffer) final override;
+	virtual void RHISetShaderUniformBuffer(FPixelShaderRHIParamRef PixelShader, uint32 BufferIndex, FUniformBufferRHIParamRef Buffer) final override;
+	virtual void RHISetShaderUniformBuffer(FComputeShaderRHIParamRef ComputeShader, uint32 BufferIndex, FUniformBufferRHIParamRef Buffer) final override;
+	virtual void RHISetShaderParameter(FVertexShaderRHIParamRef VertexShader, uint32 BufferIndex, uint32 BaseIndex, uint32 NumBytes, const void* NewValue) final override;
+	virtual void RHISetShaderParameter(FPixelShaderRHIParamRef PixelShader, uint32 BufferIndex, uint32 BaseIndex, uint32 NumBytes, const void* NewValue) final override;
+	virtual void RHISetShaderParameter(FHullShaderRHIParamRef HullShader, uint32 BufferIndex, uint32 BaseIndex, uint32 NumBytes, const void* NewValue) final override;
+	virtual void RHISetShaderParameter(FDomainShaderRHIParamRef DomainShader, uint32 BufferIndex, uint32 BaseIndex, uint32 NumBytes, const void* NewValue) final override;
+	virtual void RHISetShaderParameter(FGeometryShaderRHIParamRef GeometryShader, uint32 BufferIndex, uint32 BaseIndex, uint32 NumBytes, const void* NewValue) final override;
+	virtual void RHISetShaderParameter(FComputeShaderRHIParamRef ComputeShader, uint32 BufferIndex, uint32 BaseIndex, uint32 NumBytes, const void* NewValue) final override;
+	virtual void RHISetDepthStencilState(FDepthStencilStateRHIParamRef NewState, uint32 StencilRef) final override;
+	virtual void RHISetBlendState(FBlendStateRHIParamRef NewState, const FLinearColor& BlendFactor) final override;
+	virtual void RHISetRenderTargets(uint32 NumSimultaneousRenderTargets, const FRHIRenderTargetView* NewRenderTargets, const FRHIDepthRenderTargetView* NewDepthStencilTarget, uint32 NumUAVs, const FUnorderedAccessViewRHIParamRef* UAVs) final override;
+	virtual void RHISetRenderTargetsAndClear(const FRHISetRenderTargetsInfo& RenderTargetsInfo) final override;
+	virtual void RHIBindClearMRTValues(bool bClearColor, int32 NumClearColors, const FLinearColor* ColorArray, bool bClearDepth, float Depth, bool bClearStencil, uint32 Stencil) final override;
+	virtual void RHIDrawPrimitive(uint32 PrimitiveType, uint32 BaseVertexIndex, uint32 NumPrimitives, uint32 NumInstances) final override;
+	virtual void RHIDrawPrimitiveIndirect(uint32 PrimitiveType, FVertexBufferRHIParamRef ArgumentBuffer, uint32 ArgumentOffset) final override;
+	virtual void RHIDrawIndexedIndirect(FIndexBufferRHIParamRef IndexBufferRHI, uint32 PrimitiveType, FStructuredBufferRHIParamRef ArgumentsBufferRHI, int32 DrawArgumentsIndex, uint32 NumInstances) final override;
+	virtual void RHIDrawIndexedPrimitive(FIndexBufferRHIParamRef IndexBuffer, uint32 PrimitiveType, int32 BaseVertexIndex, uint32 FirstInstance, uint32 NumVertices, uint32 StartIndex, uint32 NumPrimitives, uint32 NumInstances) final override;
+	virtual void RHIDrawIndexedPrimitiveIndirect(uint32 PrimitiveType, FIndexBufferRHIParamRef IndexBuffer, FVertexBufferRHIParamRef ArgumentBuffer, uint32 ArgumentOffset) final override;
+	virtual void RHIBeginDrawPrimitiveUP(uint32 PrimitiveType, uint32 NumPrimitives, uint32 NumVertices, uint32 VertexDataStride, void*& OutVertexData) final override;
+	virtual void RHIEndDrawPrimitiveUP() final override;
+	virtual void RHIBeginDrawIndexedPrimitiveUP(uint32 PrimitiveType, uint32 NumPrimitives, uint32 NumVertices, uint32 VertexDataStride, void*& OutVertexData, uint32 MinVertexIndex, uint32 NumIndices, uint32 IndexDataStride, void*& OutIndexData) final override;
+	virtual void RHIEndDrawIndexedPrimitiveUP() final override;
+	virtual void RHIClear(bool bClearColor, const FLinearColor& Color, bool bClearDepth, float Depth, bool bClearStencil, uint32 Stencil, FIntRect ExcludeRect) final override;
+	virtual void RHIClearMRT(bool bClearColor, int32 NumClearColors, const FLinearColor* ColorArray, bool bClearDepth, float Depth, bool bClearStencil, uint32 Stencil, FIntRect ExcludeRect) final override;
+	virtual void RHIEnableDepthBoundsTest(bool bEnable, float MinDepth, float MaxDepth) final override;
+	virtual void RHIPushEvent(const TCHAR* Name) final override;
+	virtual void RHIPopEvent() final override;
 	
 	// Accessors.
 	ID3D11Device* GetDevice() const
