Merge fix for bFlushedGlobal

diff --git a/Engine/Source/Runtime/RHI/Public/RHICommandList.h b/Engine/Source/Runtime/RHI/Public/RHICommandList.h
--- a/Engine/Source/Runtime/RHI/Public/RHICommandList.h
+++ b/Engine/Source/Runtime/RHI/Public/RHICommandList.h
@@ -2054,7 +2054,8 @@ public:
 	{
 		QUICK_SCOPE_CYCLE_COUNTER(STAT_RHIMETHOD_LockIndexBuffer_Flush);
 		ImmediateFlush(EImmediateFlushType::FlushRHIThread); 
-		TGuardValue<bool> GuardIsFlushedGlobal( bFlushedGlobal, true ); 
+		bFlushedGlobal = true;
+		TGuardValue<bool> GuardIsFlushedGlobal( bFlushedGlobal, false ); 
 		return GDynamicRHI->RHILockIndexBuffer(IndexBuffer, Offset, Size, LockMode);
 	}
 	
@@ -2062,7 +2063,8 @@ public:
 	{
 		QUICK_SCOPE_CYCLE_COUNTER(STAT_RHIMETHOD_UnlockIndexBuffer_Flush);
 		ImmediateFlush(EImmediateFlushType::FlushRHIThread); 
-		TGuardValue<bool> GuardIsFlushedGlobal( bFlushedGlobal, true );  
+		bFlushedGlobal = true;
+		TGuardValue<bool> GuardIsFlushedGlobal( bFlushedGlobal, false );  
 		GDynamicRHI->RHIUnlockIndexBuffer(IndexBuffer);
 	}
 	
@@ -2080,7 +2082,8 @@ public:
 	{
 		QUICK_SCOPE_CYCLE_COUNTER(STAT_RHIMETHOD_LockVertexBuffer_Flush);
 		ImmediateFlush(EImmediateFlushType::FlushRHIThread); 
-		TGuardValue<bool> GuardIsFlushedGlobal( bFlushedGlobal, true ); 
+		bFlushedGlobal = true;
+		TGuardValue<bool> GuardIsFlushedGlobal( bFlushedGlobal, false ); 
 		return GDynamicRHI->RHILockVertexBuffer(VertexBuffer, Offset, SizeRHI, LockMode);
 	}
 	
@@ -2088,7 +2091,8 @@ public:
 	{
 		QUICK_SCOPE_CYCLE_COUNTER(STAT_RHIMETHOD_UnlockVertexBuffer_Flush);
 		ImmediateFlush(EImmediateFlushType::FlushRHIThread); 
-		TGuardValue<bool> GuardIsFlushedGlobal( bFlushedGlobal, true );  
+		bFlushedGlobal = true;
+		TGuardValue<bool> GuardIsFlushedGlobal( bFlushedGlobal, false );  
 		GDynamicRHI->RHIUnlockVertexBuffer(VertexBuffer);
 	}
 	
