added AyncCompute feature for XboxOne from Lionhead

diff --git a/Engine/Source/Programs/UnrealBuildTool/Configuration/BuildConfiguration.cs b/Engine/Source/Programs/UnrealBuildTool/Configuration/BuildConfiguration.cs
--- a/Engine/Source/Programs/UnrealBuildTool/Configuration/BuildConfiguration.cs
+++ b/Engine/Source/Programs/UnrealBuildTool/Configuration/BuildConfiguration.cs
@@ -43,6 +43,11 @@ namespace UnrealBuildTool
 		[XmlConfig]
 		public static bool bUseFastMonoCalls;
 
+        /// Async Compute context support. Requires Mono and Fastcalls
+        /// </summary>
+        [XmlConfig]
+        public static bool bUseAsyncComputeContext;
+
         /// <summary>
 		/// An approximate number of bytes of C++ code to target for inclusion in a single unified C++ file.
 		/// </summary>
@@ -532,7 +537,8 @@ namespace UnrealBuildTool
 			//IMPORTANT THIS IS THE MAIN SWITCH FOR MONO FAST CALLS
 			//  if this is set to true, then fast calls will be on by default on Dingo, and if false it will be off by default on Dingo.
 			//  This can be overridden by -fastmonocalls  or -nofastmonocalls in the NMAKE params.
-			bUseFastMonoCalls = false;
+            bUseFastMonoCalls = true;
+            bUseAsyncComputeContext = bUseFastMonoCalls;
 
 			// By default we use the Release C++ Runtime (CRT), even when compiling Debug builds.  This is because the Debug C++
 			// Runtime isn't very useful when debugging Unreal Engine projects, and linking against the Debug CRT libraries forces
diff --git a/Engine/Source/Runtime/NullDrv/Public/NullRHI.h b/Engine/Source/Runtime/NullDrv/Public/NullRHI.h
--- a/Engine/Source/Runtime/NullDrv/Public/NullRHI.h
+++ b/Engine/Source/Runtime/NullDrv/Public/NullRHI.h
@@ -836,6 +836,16 @@ public:
 	{ 
 		return nullptr; 
 	}
+	
+	virtual void RHIBeginAsyncComputeJob_DrawThread(EAsyncComputePriority Priority) override
+	{
+	}
+	virtual void RHIEndAsyncComputeJob_DrawThread(uint32 FenceIndex) override
+	{
+	}
+	virtual void RHIGraphicsWaitOnAsyncComputeJob(uint32 FenceIndex) override
+	{
+	}
 
 private:
 
diff --git a/Engine/Source/Runtime/OpenGLDrv/Private/OpenGLCommands.cpp b/Engine/Source/Runtime/OpenGLDrv/Private/OpenGLCommands.cpp
--- a/Engine/Source/Runtime/OpenGLDrv/Private/OpenGLCommands.cpp
+++ b/Engine/Source/Runtime/OpenGLDrv/Private/OpenGLCommands.cpp
@@ -3513,6 +3513,20 @@ void FOpenGLDynamicRHI::RHIDispatchIndirectComputeShader(FVertexBufferRHIParamRe
 	}
 }
 
+void FOpenGLDynamicRHI::RHIBeginAsyncComputeJob_DrawThread(EAsyncComputePriority Priority) 
+{
+	UE_LOG(LogRHI, Fatal,TEXT("%s not implemented yet"),ANSI_TO_TCHAR(__FUNCTION__));
+}
+
+void FOpenGLDynamicRHI::RHIEndAsyncComputeJob_DrawThread(uint32 FenceIndex)
+{
+	UE_LOG(LogRHI, Fatal,TEXT("%s not implemented yet"),ANSI_TO_TCHAR(__FUNCTION__));
+}
+
+void FOpenGLDynamicRHI::RHIGraphicsWaitOnAsyncComputeJob(uint32 FenceIndex)
+{
+	UE_LOG(LogRHI, Fatal,TEXT("%s not implemented yet"),ANSI_TO_TCHAR(__FUNCTION__));
+}
 
 void FOpenGLDynamicRHI::RHISetMultipleViewports(uint32 Count, const FViewportBounds* Data)
 {
diff --git a/Engine/Source/Runtime/OpenGLDrv/Public/OpenGLDrv.h b/Engine/Source/Runtime/OpenGLDrv/Public/OpenGLDrv.h
--- a/Engine/Source/Runtime/OpenGLDrv/Public/OpenGLDrv.h
+++ b/Engine/Source/Runtime/OpenGLDrv/Public/OpenGLDrv.h
@@ -478,6 +478,9 @@ public:
 	virtual void RHIEnableDepthBoundsTest(bool bEnable, float MinDepth, float MaxDepth) final override;
 	virtual void RHIPushEvent(const TCHAR* Name) final override;
 	virtual void RHIPopEvent() final override;
+	virtual void RHIBeginAsyncComputeJob_DrawThread(EAsyncComputePriority Priority) override;
+	virtual void RHIEndAsyncComputeJob_DrawThread(uint32 FenceIndex) override;
+	virtual void RHIGraphicsWaitOnAsyncComputeJob(uint32 FenceIndex) override;
 
 	void Cleanup();
 
diff --git a/Engine/Source/Runtime/RHI/Private/RHICommandList.cpp b/Engine/Source/Runtime/RHI/Private/RHICommandList.cpp
--- a/Engine/Source/Runtime/RHI/Private/RHICommandList.cpp
+++ b/Engine/Source/Runtime/RHI/Private/RHICommandList.cpp
@@ -57,6 +57,8 @@ static FGraphEventArray WaitOutstandingTasks;
 static FGraphEventRef RHIThreadTask;
 static FGraphEventRef RenderThreadSublistDispatchTask;
 static FGraphEventRef RenderThreadExecutionTask;
+// Used by AsyncCompute
+RHI_API FRHICommandListFenceAllocator GRHIFenceAllocator;
 
 DECLARE_CYCLE_STAT(TEXT("RHI Thread Execute"), STAT_RHIThreadExecute, STATGROUP_RHICMDLIST);
 
diff --git a/Engine/Source/Runtime/RHI/Public/DynamicRHI.h b/Engine/Source/Runtime/RHI/Public/DynamicRHI.h
--- a/Engine/Source/Runtime/RHI/Public/DynamicRHI.h
+++ b/Engine/Source/Runtime/RHI/Public/DynamicRHI.h
@@ -918,6 +918,14 @@ public:
 
 	virtual void RHIPopEvent() = 0;
 
+	/** Start AsyncCompute command stream recording (no effect if not supported) */
+	virtual void RHIBeginAsyncComputeJob_DrawThread(EAsyncComputePriority Priority) = 0;
+
+	/** End AsyncCompute command stream recording (no effect if not supported) */
+	virtual void RHIEndAsyncComputeJob_DrawThread(uint32 FenceIndex) = 0;
+
+	/** Wait for AsyncCompute command stream to finish (no effect if not supported) */
+	virtual void RHIGraphicsWaitOnAsyncComputeJob(uint32 FenceIndex) = 0;
 };
 
 
diff --git a/Engine/Source/Runtime/RHI/Public/RHI.h b/Engine/Source/Runtime/RHI/Public/RHI.h
--- a/Engine/Source/Runtime/RHI/Public/RHI.h
+++ b/Engine/Source/Runtime/RHI/Public/RHI.h
@@ -11,6 +11,13 @@
 #include "RHIDefinitions.h"
 #include "StaticArray.h"
 
+#define INVALID_FENCE_ID (0xffffffffffffffffull)
+
+inline const bool IsValidFenceID( const uint64 FenceID )
+{
+	return ( ( FenceID & 0x8000000000000000ull ) == 0 );
+}
+
 /** Uniform buffer structs must be aligned to 16-byte boundaries. */
 #define UNIFORM_BUFFER_STRUCT_ALIGNMENT 16
 
@@ -205,6 +212,15 @@ Requirements for RHI thread
 ***/
 extern RHI_API bool GRHISupportsRHIThread;
 
+inline bool IsAsyncComputeEnabled()
+{
+#if USE_ASYNC_COMPUTE_CONTEXT
+	extern int32 GAllowAsyncComputeJobs;
+	return GAllowAsyncComputeJobs != 0;
+#endif
+	return false;
+}
+
 /** Whether or not the RHI supports parallel RHIThread executes / translates
 Requirements:
 * RHICreateBoundShaderState is threadsafe and GetCachedBoundShaderState must not be used. GetCachedBoundShaderState_Threadsafe has a slightly different protocol.
diff --git a/Engine/Source/Runtime/RHI/Public/RHICommandList.h b/Engine/Source/Runtime/RHI/Public/RHICommandList.h
--- a/Engine/Source/Runtime/RHI/Public/RHICommandList.h
+++ b/Engine/Source/Runtime/RHI/Public/RHICommandList.h
@@ -52,6 +52,46 @@ struct FRHICommandBase
 	}
 };
 
+// Thread-safe allocator for GPU fences used in deferred command list execution
+// Fences are stored in a ringbuffer
+class RHI_API FRHICommandListFenceAllocator
+{
+public:
+	static const int MAX_FENCE_INDICES		= 4096;
+	FRHICommandListFenceAllocator()
+	{
+		CurrentFenceIndex = 0;
+		for ( int i=0; i<MAX_FENCE_INDICES; i++)
+		{
+			FenceIDs[i] = 0xffffffffffffffffull;
+			FenceFrameNumber[i] = 0xffffffff;
+		}
+	}
+
+	uint32 AllocFenceIndex()
+	{
+		check(IsInRenderingThread());
+		uint32 FenceIndex = ( FPlatformAtomics::InterlockedIncrement(&CurrentFenceIndex)-1 ) % MAX_FENCE_INDICES;
+		check(FenceFrameNumber[FenceIndex] != GFrameNumberRenderThread);
+		FenceFrameNumber[FenceIndex] = GFrameNumberRenderThread;
+
+		return FenceIndex;
+	}
+
+	volatile uint64& GetFenceID( int32 FenceIndex )
+	{
+		check( FenceIndex < MAX_FENCE_INDICES );
+		return FenceIDs[ FenceIndex ];
+	}
+
+private:
+	volatile int32 CurrentFenceIndex;
+	uint64 FenceIDs[MAX_FENCE_INDICES];
+	uint32 FenceFrameNumber[MAX_FENCE_INDICES];
+};
+
+extern RHI_API FRHICommandListFenceAllocator GRHIFenceAllocator;
+
 class RHI_API FRHICommandListBase : public FNoncopyable
 {
 public:
@@ -1196,6 +1236,28 @@ struct FRHICommandDebugBreak : public FRHICommand<FRHICommandDebugBreak>
 
 #define CMD_CONTEXT(Method) GetContext().RHI##Method
 
+struct FRHIBeginAsyncComputeJob_DrawThread : public FRHICommand<FRHIBeginAsyncComputeJob_DrawThread>
+{
+	EAsyncComputePriority Priority; 
+	FORCEINLINE_DEBUGGABLE FRHIBeginAsyncComputeJob_DrawThread(EAsyncComputePriority InPriority) : Priority(InPriority) {}
+	RHI_API void Execute(FRHICommandListBase& CmdList);
+};
+
+struct FRHIEndAsyncComputeJob_DrawThread : public FRHICommand<FRHIEndAsyncComputeJob_DrawThread>
+{
+	uint32 FenceIndex;
+
+	FORCEINLINE_DEBUGGABLE FRHIEndAsyncComputeJob_DrawThread(uint32 InFenceIndex) : FenceIndex(InFenceIndex) {}
+	RHI_API void Execute(FRHICommandListBase& CmdList);
+};
+
+struct FRHIGraphicsWaitOnAsyncComputeJob : public FRHICommand<FRHIGraphicsWaitOnAsyncComputeJob>
+{
+	uint32 FenceIndex;
+
+	FORCEINLINE_DEBUGGABLE FRHIGraphicsWaitOnAsyncComputeJob(uint32 InFenceIndex) : FenceIndex(InFenceIndex) {}
+	RHI_API void Execute(FRHICommandListBase& CmdList);
+};
 
 class RHI_API FRHICommandList : public FRHICommandListBase
 {
@@ -1876,6 +1938,39 @@ public:
 		new (AllocCommand<FRHICommandDebugBreak>()) FRHICommandDebugBreak();
 #endif
 	}
+
+	FORCEINLINE_DEBUGGABLE void BeginAsyncComputeJob_DrawThread( EAsyncComputePriority Priority )
+	{
+		if (Bypass())
+		{
+			CMD_CONTEXT(BeginAsyncComputeJob_DrawThread)(Priority);
+			return;
+		}
+		new (AllocCommand<FRHIBeginAsyncComputeJob_DrawThread>()) FRHIBeginAsyncComputeJob_DrawThread(Priority);
+	}
+
+	FORCEINLINE_DEBUGGABLE uint32 EndAsyncComputeJob_DrawThread()
+	{
+		uint32 FenceIndex = GRHIFenceAllocator.AllocFenceIndex();
+
+		if (Bypass())
+		{
+			CMD_CONTEXT(EndAsyncComputeJob_DrawThread)(FenceIndex);
+			return FenceIndex;
+		}
+		new (AllocCommand<FRHIEndAsyncComputeJob_DrawThread>()) FRHIEndAsyncComputeJob_DrawThread(FenceIndex);
+		return FenceIndex;
+	}
+
+	FORCEINLINE_DEBUGGABLE void GraphicsWaitOnAsyncComputeJob( uint32 FenceIndex )
+	{
+		if (Bypass())
+		{
+			CMD_CONTEXT(GraphicsWaitOnAsyncComputeJob)(FenceIndex);
+			return;
+		}
+		new (AllocCommand<FRHIGraphicsWaitOnAsyncComputeJob>()) FRHIGraphicsWaitOnAsyncComputeJob(FenceIndex);
+	}
 };
 
 namespace EImmediateFlushType
diff --git a/Engine/Source/Runtime/RHI/Public/RHICommandListCommandExecutes.inl b/Engine/Source/Runtime/RHI/Public/RHICommandListCommandExecutes.inl
--- a/Engine/Source/Runtime/RHI/Public/RHICommandListCommandExecutes.inl
+++ b/Engine/Source/Runtime/RHI/Public/RHICommandListCommandExecutes.inl
@@ -291,6 +291,24 @@ void FRHICommandClearMRT::Execute(FRHICommandListBase& CmdList)
 	INTERNAL_DECORATOR(ClearMRT)(bClearColor, NumClearColors, ColorArray, bClearDepth, Depth, bClearStencil, Stencil, ExcludeRect);
 }
 
+void FRHIBeginAsyncComputeJob_DrawThread::Execute(FRHICommandListBase& CmdList)
+{
+	RHISTAT(BeginAsyncComputeJob_DrawThread);
+	INTERNAL_DECORATOR(BeginAsyncComputeJob_DrawThread)(Priority);
+}
+
+void FRHIEndAsyncComputeJob_DrawThread::Execute(FRHICommandListBase& CmdList)
+{
+	RHISTAT(EndAsyncComputeJob_DrawThread);
+	INTERNAL_DECORATOR(EndAsyncComputeJob_DrawThread)(FenceIndex);
+}
+
+void FRHIGraphicsWaitOnAsyncComputeJob::Execute(FRHICommandListBase& CmdList)
+{
+	RHISTAT(GraphicsWaitOnAsyncComputeJob);
+	INTERNAL_DECORATOR(GraphicsWaitOnAsyncComputeJob)(FenceIndex);
+}
+
 void FRHICommandBuildLocalBoundShaderState::Execute(FRHICommandListBase& CmdList)
 {
 	RHISTAT(BuildLocalBoundShaderState);
diff --git a/Engine/Source/Runtime/RHI/Public/RHIDefinitions.h b/Engine/Source/Runtime/RHI/Public/RHIDefinitions.h
--- a/Engine/Source/Runtime/RHI/Public/RHIDefinitions.h
+++ b/Engine/Source/Runtime/RHI/Public/RHIDefinitions.h
@@ -20,6 +20,15 @@ enum EShaderFrequency
 	SF_Compute			= 5,
 
 	SF_NumFrequencies	= 6,
+
+#if USE_ASYNC_COMPUTE_CONTEXT 
+	SF_AsyncCompute		= 6, // "Special" frequency to distinguish between Device Contexts, 
+							 // TODO: maybe have a separate enum for this
+	SF_NumFrequenciesIncludingAsync,
+#else
+	SF_NumFrequenciesIncludingAsync = SF_NumFrequencies,
+#endif
+
 	SF_NumBits			= 3,
 };
 static_assert(SF_NumFrequencies <= (1 << SF_NumBits), "SF_NumFrequencies will not fit on SF_NumBits");
@@ -533,6 +542,11 @@ enum ETextureCreateFlags
 	TexCreate_NoFastClear = 1 << 25,
 };
 
+enum EAsyncComputePriority
+{
+	AsyncComputePriority_Default = 0,
+	AsyncComputePriority_High,
+};
 /**
  * Async texture reallocation status, returned by RHIGetReallocateTexture2DStatus().
  */
diff --git a/Engine/Source/Runtime/Windows/D3D11RHI/Private/D3D11Commands.cpp b/Engine/Source/Runtime/Windows/D3D11RHI/Private/D3D11Commands.cpp
--- a/Engine/Source/Runtime/Windows/D3D11RHI/Private/D3D11Commands.cpp
+++ b/Engine/Source/Runtime/Windows/D3D11RHI/Private/D3D11Commands.cpp
@@ -60,6 +60,18 @@ DECLARE_ISBOUNDSHADER(ComputeShader)
 	#include <GPUPerfAPI/Gpa.h>
 #endif
 
+void FD3D11DynamicRHI::RHIBeginAsyncComputeJob_DrawThread(EAsyncComputePriority Priority) 
+{  
+}
+
+void FD3D11DynamicRHI::RHIEndAsyncComputeJob_DrawThread(uint32 FenceIndex)
+{ 
+}
+
+void FD3D11DynamicRHI::RHIGraphicsWaitOnAsyncComputeJob( uint32 FenceIndex )
+{ 
+}
+
 void FD3D11DynamicRHI::RHIGpuTimeBegin(uint32 Hash, bool bCompute)
 {
 	#if WITH_GPA
diff --git a/Engine/Source/Runtime/Windows/D3D11RHI/Private/D3D11RHIPrivate.h b/Engine/Source/Runtime/Windows/D3D11RHI/Private/D3D11RHIPrivate.h
--- a/Engine/Source/Runtime/Windows/D3D11RHI/Private/D3D11RHIPrivate.h
+++ b/Engine/Source/Runtime/Windows/D3D11RHI/Private/D3D11RHIPrivate.h
@@ -475,7 +475,10 @@ public:
 	virtual void RHIEnableDepthBoundsTest(bool bEnable, float MinDepth, float MaxDepth) final override;
 	virtual void RHIPushEvent(const TCHAR* Name) final override;
 	virtual void RHIPopEvent() final override;
-	
+	virtual void RHIBeginAsyncComputeJob_DrawThread(EAsyncComputePriority Priority);
+	virtual void RHIEndAsyncComputeJob_DrawThread(uint32 FenceIndex);
+	virtual void RHIGraphicsWaitOnAsyncComputeJob(uint32 FenceIndex);
+
 	// Accessors.
 	ID3D11Device* GetDevice() const
 	{
