Restore original GuardValue code.

diff --git a/Engine/Source/Runtime/RHI/Public/RHICommandList.h b/Engine/Source/Runtime/RHI/Public/RHICommandList.h
--- a/Engine/Source/Runtime/RHI/Public/RHICommandList.h
+++ b/Engine/Source/Runtime/RHI/Public/RHICommandList.h
@@ -2149,8 +2149,7 @@ public:
 	{
 		QUICK_SCOPE_CYCLE_COUNTER(STAT_RHIMETHOD_LockIndexBuffer_Flush);
 		ImmediateFlush(EImmediateFlushType::FlushRHIThread); 
-		bFlushedGlobal = true;
-		TGuardValue<bool> GuardIsFlushedGlobal( bFlushedGlobal, false ); 
+		TGuardValue<bool> GuardIsFlushedGlobal( bFlushedGlobal, true ); 
 		return GDynamicRHI->RHILockIndexBuffer(IndexBuffer, Offset, Size, LockMode);
 	}
 	
@@ -2158,8 +2157,7 @@ public:
 	{
 		QUICK_SCOPE_CYCLE_COUNTER(STAT_RHIMETHOD_UnlockIndexBuffer_Flush);
 		ImmediateFlush(EImmediateFlushType::FlushRHIThread); 
-		bFlushedGlobal = true;
-		TGuardValue<bool> GuardIsFlushedGlobal( bFlushedGlobal, false );  
+		TGuardValue<bool> GuardIsFlushedGlobal( bFlushedGlobal, true );  
 		GDynamicRHI->RHIUnlockIndexBuffer(IndexBuffer);
 	}
 	
@@ -2177,8 +2175,7 @@ public:
 	{
 		QUICK_SCOPE_CYCLE_COUNTER(STAT_RHIMETHOD_LockVertexBuffer_Flush);
 		ImmediateFlush(EImmediateFlushType::FlushRHIThread); 
-		bFlushedGlobal = true;
-		TGuardValue<bool> GuardIsFlushedGlobal( bFlushedGlobal, false ); 
+		TGuardValue<bool> GuardIsFlushedGlobal( bFlushedGlobal, true ); 
 		return GDynamicRHI->RHILockVertexBuffer(VertexBuffer, Offset, SizeRHI, LockMode);
 	}
 	
@@ -2186,8 +2183,7 @@ public:
 	{
 		QUICK_SCOPE_CYCLE_COUNTER(STAT_RHIMETHOD_UnlockVertexBuffer_Flush);
 		ImmediateFlush(EImmediateFlushType::FlushRHIThread); 
-		bFlushedGlobal = true;
-		TGuardValue<bool> GuardIsFlushedGlobal( bFlushedGlobal, false );  
+		TGuardValue<bool> GuardIsFlushedGlobal( bFlushedGlobal, true );  
 		GDynamicRHI->RHIUnlockVertexBuffer(VertexBuffer);
 	}
 	
