UE4 - Large upgrade to parallel rendering. No longer block per pass. Several RHI cmdlist passthrough functions added and implemented on the consoles

diff --git a/Engine/Source/Runtime/Core/Public/Misc/MemStack.h b/Engine/Source/Runtime/Core/Public/Misc/MemStack.h
--- a/Engine/Source/Runtime/Core/Public/Misc/MemStack.h
+++ b/Engine/Source/Runtime/Core/Public/Misc/MemStack.h
@@ -136,6 +136,10 @@ public:
 		check(!NumMarks && !MinMarksToAlloc);
 		FreeChunks(nullptr);
 	}
+	FORCEINLINE int32 GetNumMarks()
+	{
+		return NumMarks;
+	}
 	/** @return the number of bytes allocated for this FMemStack that are currently in use. */
 	int32 GetByteCount() const;
 
diff --git a/Engine/Source/Runtime/Engine/Classes/Engine/World.h b/Engine/Source/Runtime/Engine/Classes/Engine/World.h
--- a/Engine/Source/Runtime/Engine/Classes/Engine/World.h
+++ b/Engine/Source/Runtime/Engine/Classes/Engine/World.h
@@ -2596,9 +2596,8 @@ public:
 
 	/**
 	 * Send all render updates to the rendering thread.
-	 * @param OutCompletion - all async updates are added to this array, they typically need to be completed before GC or anything else
 	 */
-	void SendAllEndOfFrameUpdates(FGraphEventArray* OutCompletion = NULL);
+	void SendAllEndOfFrameUpdates();
 
 
 	/** Do per frame tick behaviors related to the network driver */
diff --git a/Engine/Source/Runtime/Engine/Private/GPUSkinVertexFactory.cpp b/Engine/Source/Runtime/Engine/Private/GPUSkinVertexFactory.cpp
--- a/Engine/Source/Runtime/Engine/Private/GPUSkinVertexFactory.cpp
+++ b/Engine/Source/Runtime/Engine/Private/GPUSkinVertexFactory.cpp
@@ -144,20 +144,12 @@ void FGPUBaseSkinVertexFactory::ShaderDataType::UpdateBoneData(ERHIFeatureLevel:
 		}
 		if(NumBones)
 		{
-			if (GRHIThread)
-			{
-				check(VectorArraySize == NumBones * sizeof(BoneMatrices[0]));
-				GRHICommandList.GetImmediateCommandList().UpdateVertexBuffer(BoneBuffer.VertexBufferRHI, BoneMatrices.GetData(), NumBones * sizeof(BoneMatrices[0]));
-			}
-			else
-			{
 				float* Data = (float*)RHILockVertexBuffer(BoneBuffer.VertexBufferRHI, 0, VectorArraySize, RLM_WriteOnly);
 				checkSlow(Data);
 				FMemory::Memcpy(Data, BoneMatrices.GetData(), NumBones * sizeof(BoneMatrices[0]));
 				RHIUnlockVertexBuffer(BoneBuffer.VertexBufferRHI);
 			}
 		}
-	}
 	else
 	{
 		check(NumBones * BoneMatrices.GetTypeSize() <= sizeof(GBoneUniformStruct));
@@ -173,31 +165,25 @@ void FGPUBaseSkinVertexFactory::ShaderDataType::UpdateBoneData(ERHIFeatureLevel:
 
 FBoneDataVertexBuffer::FBoneDataVertexBuffer()
 	: SizeX(80 * 1024)		// todo: we will replace this fixed size using FGlobalDynamicVertexBuffer
-	, AllocSize(0)
-	, AllocBlock(nullptr)
 {
 }
 
 float* FBoneDataVertexBuffer::LockData()
 {
-	check(IsInRenderingThread() && GetSizeX() && IsValidRef(BoneBuffer));
-	check(!AllocSize && !AllocBlock);
-	AllocSize = ComputeMemorySize();
-	AllocBlock = FMemory::Malloc(AllocSize + 15, 16);
-	float* Data = (float*)AllocBlock;
-	check(AllocSize && AllocBlock);
-	check(Data);
+	checkSlow(IsInRenderingThread());
+	checkSlow(GetSizeX());
+	checkSlow(IsValidRef(BoneBuffer));
+
+	float* Data = (float*)RHILockVertexBuffer(BoneBuffer.VertexBufferRHI, 0, ComputeMemorySize(), RLM_WriteOnly);
+	checkSlow(Data);
+
 	return Data;
 }
 
 void FBoneDataVertexBuffer::UnlockData(uint32 SizeInBytes)
 {
-	check(IsValidRef(BoneBuffer));
-	check(IsInRenderingThread() && AllocBlock && AllocSize);
-	FRHICommandListExecutor::GetImmediateCommandList().UpdateVertexBuffer(BoneBuffer.VertexBufferRHI, AllocBlock, SizeInBytes ? SizeInBytes : 4);
-	FMemory::Free(AllocBlock);
-	AllocBlock = nullptr;
-	AllocSize = 0;
+	checkSlow(IsValidRef(BoneBuffer));
+	RHIUnlockVertexBuffer(BoneBuffer.VertexBufferRHI);
 }
 
 uint32 FBoneDataVertexBuffer::GetSizeX() const
@@ -421,9 +407,11 @@ public:
 			}
 
 			bool bLocalPerBoneMotionBlur = false;
-			
-			if (FeatureLevel >= ERHIFeatureLevel::SM4 && PreviousBoneMatrices.IsBound())
+
+
+			if (FeatureLevel >= ERHIFeatureLevel::SM4 && GPrevPerBoneMotionBlur.IsVelocityPass(RHICmdList) && PreviousBoneMatrices.IsBound())
 			{
+				check(GPrevPerBoneMotionBlur.IsAppendStarted());
 				const bool bWorldIsPaused = View.Family->bWorldIsPaused;
 
 				// we are in the velocity rendering pass
@@ -484,24 +472,24 @@ public:
 					else
 					{
 #if !(UE_BUILD_SHIPPING || UE_BUILD_TEST)
-						if (CVarVelocityTest.GetValueOnRenderThread())
-						{
-							FBoneSkinning Tab[10];
+					if (CVarVelocityTest.GetValueOnRenderThread())
+					{
+						FBoneSkinning Tab[10];
 
 							FMemory::Memset(Tab, 0, sizeof(FBoneSkinning) * 10);
 							// We add some black gap between the elements, to ensure we pick the right data for this frame
 							// wwe create different gap size on alternating frames.
 							GPrevPerBoneMotionBlur.AppendData(Tab, (FrameNumber % 2) + 4);
-						}
+					}
 #endif // if !(UE_BUILD_SHIPPING || UE_BUILD_TEST)
 
-						// copy the bone data and tell the instance where it can pick it up next frame
-						// append data to a buffer we bind next frame to read old matrix data for motion blur
-						uint32 OldBoneDataStartIndex = GPrevPerBoneMotionBlur.AppendData(ShaderData.BoneMatrices.GetData(), ShaderData.BoneMatrices.Num());
-						GPUVertexFactory->SetOldBoneDataStartIndex(FrameNumber, OldBoneDataStartIndex);
-					}
+					// copy the bone data and tell the instance where it can pick it up next frame
+					// append data to a buffer we bind next frame to read old matrix data for motion blur
+					uint32 OldBoneDataStartIndex = GPrevPerBoneMotionBlur.AppendData(ShaderData.BoneMatrices.GetData(), ShaderData.BoneMatrices.Num());
+					GPUVertexFactory->SetOldBoneDataStartIndex(FrameNumber, OldBoneDataStartIndex);
 				}
 			}
+			}
 			SetShaderValue(RHICmdList, Shader->GetVertexShader(), PerBoneMotionBlur, bLocalPerBoneMotionBlur);
 		}
 	}
@@ -770,17 +758,17 @@ void FGPUBaseSkinAPEXClothVertexFactory::ClothShaderType::UpdateClothSimulData(c
 
 		if(NumSimulVerts)
 		{
-			float* Data = (float*)RHILockVertexBuffer(ClothSimulPositionBuffer.VertexBufferRHI, 0, VectorArraySize, RLM_WriteOnly);
-			checkSlow(Data);
-			FMemory::Memcpy(Data, InSimulPositions.GetData(), NumSimulVerts * sizeof(FVector4));
-			RHIUnlockVertexBuffer(ClothSimulPositionBuffer.VertexBufferRHI);
-
-			Data = (float*)RHILockVertexBuffer(ClothSimulNormalBuffer.VertexBufferRHI, 0, VectorArraySize, RLM_WriteOnly);
-			checkSlow(Data);
-			FMemory::Memcpy(Data, InSimulNormals.GetData(), NumSimulVerts * sizeof(FVector4));
-			RHIUnlockVertexBuffer(ClothSimulNormalBuffer.VertexBufferRHI);
+				float* Data = (float*)RHILockVertexBuffer(ClothSimulPositionBuffer.VertexBufferRHI, 0, VectorArraySize, RLM_WriteOnly);
+				checkSlow(Data);
+				FMemory::Memcpy(Data, InSimulPositions.GetData(), NumSimulVerts * sizeof(FVector4));
+				RHIUnlockVertexBuffer(ClothSimulPositionBuffer.VertexBufferRHI);
+
+				Data = (float*)RHILockVertexBuffer(ClothSimulNormalBuffer.VertexBufferRHI, 0, VectorArraySize, RLM_WriteOnly);
+				checkSlow(Data);
+				FMemory::Memcpy(Data, InSimulNormals.GetData(), NumSimulVerts * sizeof(FVector4));
+				RHIUnlockVertexBuffer(ClothSimulNormalBuffer.VertexBufferRHI);
+			}
 		}
-	}
 	else
 	{
 		UpdateClothUniformBuffer(InSimulPositions, InSimulNormals);
diff --git a/Engine/Source/Runtime/Engine/Private/LevelTick.cpp b/Engine/Source/Runtime/Engine/Private/LevelTick.cpp
--- a/Engine/Source/Runtime/Engine/Private/LevelTick.cpp
+++ b/Engine/Source/Runtime/Engine/Private/LevelTick.cpp
@@ -25,6 +25,7 @@
 #if !UE_BUILD_SHIPPING
 #include "STaskGraph.h"
 #endif
+#include "ParallelFor.h"
 
 // this will log out all of the objects that were ticked in the FDetailedTickStats struct so you can isolate what is expensive
 #define LOG_DETAILED_DUMPSTATS 0
@@ -756,7 +757,18 @@ void UWorld::UpdateActorComponentEndOfFrameUpdateState(UActorComponent* Componen
 void UWorld::MarkActorComponentForNeededEndOfFrameUpdate(UActorComponent* Component, bool bForceGameThread)
 {
 	check(!bPostTickComponentUpdate); // can't call this while we are doing the updates
-	if (Component->GetMarkedForEndOfFrameUpdateState() != EComponentMarkedForEndOfFrameUpdateState::MarkedForGameThread)
+
+	uint32 CurrentState = Component->GetMarkedForEndOfFrameUpdateState();
+
+	// force game thread can be turned on later, but we are not concerned about that, those are only cvars and constants; if those are changed during a frame, they won't fully kick in till next frame.
+	if (CurrentState == EComponentMarkedForEndOfFrameUpdateState::Marked && bForceGameThread)
+	{
+		verify(ComponentsThatNeedEndOfFrameUpdate.RemoveSwap(Component) == 1);
+		CurrentState = EComponentMarkedForEndOfFrameUpdateState::Unmarked;
+	}
+	// it is totally ok if it is currently marked for the gamethread but now they are not forcing game thread. It will run on the game thread this frame.
+
+	if (CurrentState == EComponentMarkedForEndOfFrameUpdateState::Unmarked)
 	{
 		if (!bForceGameThread)
 		{
@@ -767,15 +779,10 @@ void UWorld::MarkActorComponentForNeededEndOfFrameUpdate(UActorComponent* Compon
 
 		if (bForceGameThread)
 		{
-			// In this case we need to remove it from the other list
-			if (Component->GetMarkedForEndOfFrameUpdateState() == EComponentMarkedForEndOfFrameUpdateState::Marked)
-			{
-				ComponentsThatNeedEndOfFrameUpdate.RemoveSwap(Component);
-			}
 			ComponentsThatNeedEndOfFrameUpdate_OnGameThread.Add(Component);
 			FMarkComponentEndOfFrameUpdateState::Set(Component, EComponentMarkedForEndOfFrameUpdateState::MarkedForGameThread);
 		}
-		else if (Component->GetMarkedForEndOfFrameUpdateState() != EComponentMarkedForEndOfFrameUpdateState::Marked)
+		else
 		{
 			ComponentsThatNeedEndOfFrameUpdate.Add(Component);
 			FMarkComponentEndOfFrameUpdateState::Set(Component, EComponentMarkedForEndOfFrameUpdateState::Marked);
@@ -786,18 +793,11 @@ void UWorld::MarkActorComponentForNeededEndOfFrameUpdate(UActorComponent* Compon
 /**
 	* Send all render updates to the rendering thread.
 	*/
-void UWorld::SendAllEndOfFrameUpdates(FGraphEventArray* OutCompletion)
+void UWorld::SendAllEndOfFrameUpdates()
 {
 	SCOPE_CYCLE_COUNTER(STAT_PostTickComponentUpdate);
 	// update all dirty components. 
-	bPostTickComponentUpdate = true;
-
-	if (!OutCompletion)
-	{
-		// this is a viewer or something, just do everything on the gamethread
-		ComponentsThatNeedEndOfFrameUpdate_OnGameThread.Append(ComponentsThatNeedEndOfFrameUpdate);
-		ComponentsThatNeedEndOfFrameUpdate.Empty(ComponentsThatNeedEndOfFrameUpdate.Num());
-	}
+	TGuardValue<bool> GuardIsFlushedGlobal( bPostTickComponentUpdate, true ); 
 
 	// Game thread updates need to happen before we go wide on the other threads.
 	// These updates are things that have said that they are NOT SAFE to run concurrently.
@@ -806,100 +806,40 @@ void UWorld::SendAllEndOfFrameUpdates(FGraphEventArray* OutCompletion)
 		UActorComponent* Component = It->Get();
 		if (Component)
 		{
-			checkSlow(Component->GetMarkedForEndOfFrameUpdateState() != EComponentMarkedForEndOfFrameUpdateState::Unmarked);
 			if ( !Component->IsPendingKill() && Component->IsRegistered() && !Component->IsTemplate())
 			{
 				FScopeCycleCounterUObject ComponentScope(Component);
 				FScopeCycleCounterUObject AdditionalScope(STATS ? Component->AdditionalStatObject() : NULL);
 				Component->DoDeferredRenderUpdates_Concurrent();
 			}
+			check(Component->GetMarkedForEndOfFrameUpdateState() == EComponentMarkedForEndOfFrameUpdateState::MarkedForGameThread);
 			FMarkComponentEndOfFrameUpdateState::Set(Component, EComponentMarkedForEndOfFrameUpdateState::Unmarked);
 		}
 	}
+	ComponentsThatNeedEndOfFrameUpdate_OnGameThread.Reset();
 
 	if (ComponentsThatNeedEndOfFrameUpdate.Num())
 	{
-		check(OutCompletion);
-		enum
-		{
-			NUM_COMPONENTS_PER_TASK = 20
-		};
-		typedef TArray<UActorComponent*, TInlineAllocator<NUM_COMPONENTS_PER_TASK> > FTaskArray;
-		/** Helper class define the task of calling DoDeferredRenderUpdates_Concurrent on an array of components **/
-		class FDoRenderthreadUpdatesTask
-		{
-			/** Array of components to process, owned by the task **/
-			TScopedPointer<FTaskArray>	Components;
-		public:
-			FDoRenderthreadUpdatesTask(FTaskArray* InComponents)
-				: Components(InComponents)
-			{
-			}
-			FORCEINLINE TStatId GetStatId() const
-			{
-				RETURN_QUICK_DECLARE_CYCLE_STAT(DoRenderthreadUpdatesTask, STATGROUP_TaskGraphTasks);
-			}
-			static ENamedThreads::Type GetDesiredThread()
-			{
-				return ENamedThreads::AnyThread;
-			}
-			static ESubsequentsMode::Type GetSubsequentsMode() 
-			{ 
-				return ESubsequentsMode::TrackSubsequents; 
-			}
-			void DoTask(ENamedThreads::Type CurrentThread, const FGraphEventRef& MyCompletionGraphEvent)
+		TArray<TWeakObjectPtr<class UActorComponent> >& LocalComponentsThatNeedEndOfFrameUpdate = ComponentsThatNeedEndOfFrameUpdate;
+		ParallelFor(ComponentsThatNeedEndOfFrameUpdate.Num(), 
+			[&LocalComponentsThatNeedEndOfFrameUpdate](int32 Index)
 			{
-				for (FTaskArray::TIterator It(*Components); It; ++It)
-				{
-					UActorComponent* Component = *It;
-					if (!Component->IsPendingKill() && Component->IsRegistered() && !Component->IsTemplate())
-					{
-						FScopeCycleCounterUObject ComponentScope(Component);
-						FScopeCycleCounterUObject AdditionalScope(STATS ? Component->AdditionalStatObject() : NULL);
-						Component->DoDeferredRenderUpdates_Concurrent();
-					}
-				}
-			}
-		};
-
-		{
-			//@todo optimization, this loop could be done on another thread
-			// First get the async transform and render data updates underway
-			FTaskArray* Array = NULL;
-			for (TArray<TWeakObjectPtr<UActorComponent> >::TIterator It(ComponentsThatNeedEndOfFrameUpdate); It; ++It)
-			{
-				UActorComponent* NextComponent = It->Get();
+				UActorComponent* NextComponent = LocalComponentsThatNeedEndOfFrameUpdate[Index].Get();
 				if (NextComponent)
 				{
-					checkSlow(NextComponent->GetMarkedForEndOfFrameUpdateState() == EComponentMarkedForEndOfFrameUpdateState::Marked);
 					if (!NextComponent->IsPendingKill() && NextComponent->IsRegistered() && !NextComponent->IsTemplate())
 					{
-						if (!Array)
-						{
-							Array = new FTaskArray;
-						}
-						Array->Add(NextComponent);
-						if (Array->Num() == NUM_COMPONENTS_PER_TASK)
-						{
-							new (*OutCompletion) FGraphEventRef(TGraphTask<FDoRenderthreadUpdatesTask>::CreateTask(NULL, ENamedThreads::GameThread).ConstructAndDispatchWhenReady(Array)); 
-							Array = NULL; // Array belongs to the task
-						}
+						FScopeCycleCounterUObject ComponentScope(NextComponent);
+						FScopeCycleCounterUObject AdditionalScope(STATS ? NextComponent->AdditionalStatObject() : NULL);
+						NextComponent->DoDeferredRenderUpdates_Concurrent();
 					}
+					check(NextComponent->GetMarkedForEndOfFrameUpdateState() == EComponentMarkedForEndOfFrameUpdateState::Marked);
 					FMarkComponentEndOfFrameUpdateState::Set(NextComponent, EComponentMarkedForEndOfFrameUpdateState::Unmarked);
 				}
 			}
-			if (Array) // partial array if we had one
-			{
-				new (*OutCompletion) FGraphEventRef(TGraphTask<FDoRenderthreadUpdatesTask>::CreateTask(NULL, ENamedThreads::GameThread).ConstructAndDispatchWhenReady(Array)); 
-				Array = NULL; // Array belongs to the task
-			}
-		}
-		ComponentsThatNeedEndOfFrameUpdate.Empty(ComponentsThatNeedEndOfFrameUpdate.Num());
+		);
+		ComponentsThatNeedEndOfFrameUpdate.Reset();
 	}
-
-	ComponentsThatNeedEndOfFrameUpdate_OnGameThread.Empty(ComponentsThatNeedEndOfFrameUpdate_OnGameThread.Num());
-
-	bPostTickComponentUpdate = false;
 }
 
 
@@ -1318,8 +1258,7 @@ void UWorld::Tick( ELevelTick TickType, float DeltaSeconds )
 		BroadcastTickFlush(RealDeltaSeconds); // note: undilated time is being used here
 	}
 
-	FGraphEventArray FrameUpdateCompletions;
-	SendAllEndOfFrameUpdates(&FrameUpdateCompletions);
+	SendAllEndOfFrameUpdates();
      // PostTick all net drivers
 	{
 		SCOPE_CYCLE_COUNTER(STAT_NetBroadcastPostTickTime);
@@ -1341,12 +1280,6 @@ void UWorld::Tick( ELevelTick TickType, float DeltaSeconds )
 		FXSystem->Tick(DeltaSeconds);
 	}
 	
-	if (FrameUpdateCompletions.Num())
-	{
-		SCOPE_CYCLE_COUNTER(STAT_PostTickComponentUpdateWait);
-		FTaskGraphInterface::Get().WaitUntilTasksComplete(FrameUpdateCompletions, ENamedThreads::GameThread_Local);
-	}
-
 	// Finish up.
 	if(bDebugFrameStepExecution)
 	{
diff --git a/Engine/Source/Runtime/Engine/Private/Particles/ParticleGpuSimulation.cpp b/Engine/Source/Runtime/Engine/Private/Particles/ParticleGpuSimulation.cpp
--- a/Engine/Source/Runtime/Engine/Private/Particles/ParticleGpuSimulation.cpp
+++ b/Engine/Source/Runtime/Engine/Private/Particles/ParticleGpuSimulation.cpp
@@ -1197,8 +1197,8 @@ static void BuildTileVertexBuffer( FParticleBufferParamRef TileOffsetsRef, const
 		TileOffset[Index].X = 100.0f;
 		TileOffset[Index].Y = 100.0f;
 	}
-	RHIUnlockVertexBuffer( TileOffsetsRef );
-}
+		RHIUnlockVertexBuffer( TileOffsetsRef );
+	}
 
 /**
  * Builds a vertex buffer containing the offsets for a set of tiles.
@@ -1593,17 +1593,9 @@ void InjectNewParticles(FRHICommandList& RHICmdList, ERHIFeatureLevel::Type Feat
 		// Copy new particles in to the vertex buffer.
 		const int32 ParticlesThisDrawCall = FMath::Min<int32>( ParticleCount, MaxParticlesPerDrawCall );
 		const void* Src = NewParticles.GetData() + FirstParticle;
-		if (GRHIThread)
-		{
-			RHICmdList.UpdateVertexBuffer(ScratchVertexBufferRHI, Src, ParticlesThisDrawCall * sizeof(FNewParticle));
-		}
-		else
-		{
-			void* Dest = RHILockVertexBuffer( ScratchVertexBufferRHI, 0, ParticlesThisDrawCall * sizeof(FNewParticle), RLM_WriteOnly );
-			FMemory::Memcpy( Dest, Src, ParticlesThisDrawCall * sizeof(FNewParticle) );
-			RHIUnlockVertexBuffer( ScratchVertexBufferRHI );
-
-		}
+		void* Dest = RHILockVertexBuffer( ScratchVertexBufferRHI, 0, ParticlesThisDrawCall * sizeof(FNewParticle), RLM_WriteOnly );
+		FMemory::Memcpy( Dest, Src, ParticlesThisDrawCall * sizeof(FNewParticle) );
+		RHIUnlockVertexBuffer( ScratchVertexBufferRHI );
 		ParticleCount -= ParticlesThisDrawCall;
 		FirstParticle += ParticlesThisDrawCall;
 
diff --git a/Engine/Source/Runtime/Engine/Private/SceneView.cpp b/Engine/Source/Runtime/Engine/Private/SceneView.cpp
--- a/Engine/Source/Runtime/Engine/Private/SceneView.cpp
+++ b/Engine/Source/Runtime/Engine/Private/SceneView.cpp
@@ -282,7 +282,6 @@ FSceneView::FSceneView(const FSceneViewInitOptions& InitOptions)
 	, bOriginOffsetThisFrame(InitOptions.bOriginOffsetThisFrame)
 	, CursorPos(InitOptions.CursorPos)
 	, bIsGameView(false)
-	, bForceShowMaterials(false)
 	, bIsViewInfo(false)
 	, bIsSceneCapture(false)
 	, bIsReflectionCapture(false)
diff --git a/Engine/Source/Runtime/Engine/Private/SkeletalRenderGPUSkin.cpp b/Engine/Source/Runtime/Engine/Private/SkeletalRenderGPUSkin.cpp
--- a/Engine/Source/Runtime/Engine/Private/SkeletalRenderGPUSkin.cpp
+++ b/Engine/Source/Runtime/Engine/Private/SkeletalRenderGPUSkin.cpp
@@ -425,13 +425,13 @@ void FSkeletalMeshObjectGPUSkin::FSkeletalMeshObjectLOD::UpdateMorphVertexBuffer
 			}
 		} // ApplyDelta
 
-		// Lock the real buffer.
-		FMorphGPUSkinVertex* ActualBuffer = (FMorphGPUSkinVertex*)RHILockVertexBuffer(MorphVertexBuffer.VertexBufferRHI, 0, Size, RLM_WriteOnly);
-		FMemory::Memcpy(ActualBuffer, Buffer, Size);
-		FMemory::Free(Buffer);
+			// Lock the real buffer.
+			FMorphGPUSkinVertex* ActualBuffer = (FMorphGPUSkinVertex*)RHILockVertexBuffer(MorphVertexBuffer.VertexBufferRHI, 0, Size, RLM_WriteOnly);
+			FMemory::Memcpy(ActualBuffer, Buffer, Size);
+			FMemory::Free(Buffer);
 
-		// Unlock the buffer.
-		RHIUnlockVertexBuffer(MorphVertexBuffer.VertexBufferRHI);
+			// Unlock the buffer.
+			RHIUnlockVertexBuffer(MorphVertexBuffer.VertexBufferRHI);
 		// set update flag
 		MorphVertexBuffer.bHasBeenUpdated = true;
 	}
@@ -1088,7 +1088,14 @@ FPreviousPerBoneMotionBlur
 -----------------------------------------------------------------------------*/
 
 FPreviousPerBoneMotionBlur::FPreviousPerBoneMotionBlur()
-	: LockedData(0), LockedTexelCount(0), bWarningBufferSizeExceeded(false)
+	:LockedData(0), LockedTexelCount(0),
+	IsVelocityFunc(
+		[](FRHICommandList& RHICmdList)->bool
+		{
+			return false;
+		}
+	)
+	, bWarningBufferSizeExceeded(false)
 {
 }
 
@@ -1131,27 +1138,28 @@ void FPreviousPerBoneMotionBlur::InitIfNeeded()
 		InitResources();
 	}
 }
-void FPreviousPerBoneMotionBlur::StartAppend(bool bWorldIsPaused)
+void FPreviousPerBoneMotionBlur::StartAppend(FRHICommandListImmediate& RHICmdList, bool bWorldIsPaused)
 {
 	check(!LockedData);
 	check(IsInRenderingThread());
+	check(IsVelocityFunc(RHICmdList));
 
 	if(!bWorldIsPaused)
-	{		
+	{
 		InitIfNeeded();
 
 		AdvanceBufferLocation();
 
-		FBoneDataVertexBuffer& WriteTexture = PerChunkBoneMatricesTexture[GetWriteBufferIndex()];
+	FBoneDataVertexBuffer& WriteTexture = PerChunkBoneMatricesTexture[GetWriteBufferIndex()];
 
-		if(WriteTexture.IsValid())
-		{
-			LockedData = WriteTexture.LockData();
-			check(LockedTexelPosition.GetValue() == 0); //otherwise it wasn't unlocked or it was unlocked before it was done being filled by async stuff
-			LockedTexelPosition.Set(0);
-			LockedTexelCount = WriteTexture.GetSizeX();
-		}
+	if(WriteTexture.IsValid())
+	{
+		LockedData = WriteTexture.LockData();
+		check(LockedTexelPosition.GetValue() == 0); //otherwise it wasn't unlocked or it was unlocked before it was done being filled by async stuff
+		LockedTexelPosition.Set(0);
+		LockedTexelCount = WriteTexture.GetSizeX();
 	}
+}
 
 	if (CVarMotionBlurDebug.GetValueOnRenderThread())
 	{
@@ -1186,8 +1194,17 @@ uint32 FPreviousPerBoneMotionBlur::AppendData(FBoneSkinning *DataStart, uint32 B
 	}
 }
 
-void FPreviousPerBoneMotionBlur::EndAppend()
+void FPreviousPerBoneMotionBlur::EndAppendFence(FRHICommandListImmediate& RHICmdList)
+{
+	check(IsVelocityFunc(RHICmdList));
+	EndAppendRenderThreadTaskFence = FRHICommandListImmediate::RenderThreadTaskFence();
+}
+
+void FPreviousPerBoneMotionBlur::EndAppend(FRHICommandListImmediate& RHICmdList)
 {
+	check(!IsVelocityFunc(RHICmdList));
+	QUICK_SCOPE_CYCLE_COUNTER(STAT_FPreviousPerBoneMotionBlur_EndAppend);
+
 	if(CVarMotionBlurDebug.GetValueOnRenderThread())
 	{
 		UE_LOG(LogEngine, Log, TEXT(" "));
@@ -1196,34 +1213,43 @@ void FPreviousPerBoneMotionBlur::EndAppend()
 	if(IsAppendStarted())
 	{
 		check(IsInRenderingThread());
+		FRHICommandListImmediate::WaitOnRenderThreadTaskFence(EndAppendRenderThreadTaskFence);
+		EndAppendRenderThreadTaskFence = nullptr;
+
 		LockedTexelCount = 0;
 		LockedData = 0;
 
 		// we use float4
 		PerChunkBoneMatricesTexture[GetWriteBufferIndex()].UnlockData(LockedTexelPosition.GetValue() * 4 * sizeof(float));
 		LockedTexelPosition.Set(0);
-	}
 
-	{
-		static int LogSpawmPrevent = 0;
-
-		if(bWarningBufferSizeExceeded)
 		{
-			bWarningBufferSizeExceeded = false;
+			static int LogSpawmPrevent = 0;
 
-			if((LogSpawmPrevent % 16) == 0)
+			if(bWarningBufferSizeExceeded)
 			{
-				UE_LOG(LogSkeletalGPUSkinMesh, Warning, TEXT("Exceeded buffer for per bone motionblur for skinned mesh velocity rendering. Artifacts can occur. Change Content, increase buffer size or change to use FGlobalDynamicVertexBuffer."));
+				bWarningBufferSizeExceeded = false;
+
+				if((LogSpawmPrevent % 16) == 0)
+				{
+					UE_LOG(LogSkeletalGPUSkinMesh, Warning, TEXT("Exceeded buffer for per bone motionblur for skinned mesh velocity rendering. Artifacts can occur. Change Content, increase buffer size or change to use FGlobalDynamicVertexBuffer."));
+				}
+				++LogSpawmPrevent;
+			}
+			else
+			{
+				LogSpawmPrevent = 0;
 			}
-			++LogSpawmPrevent;
-		}
-		else
-		{
-			LogSpawmPrevent = 0;
 		}
 	}
 }
 
+void FPreviousPerBoneMotionBlur::SetVelocityPassCallback(const TFunction<bool(FRHICommandList& RHICmdList)>& InIsVelocityFunc)
+{
+	IsVelocityFunc = InIsVelocityFunc;
+}
+
+
 FString FPreviousPerBoneMotionBlur::GetDebugString() const
 {
 	return FString::Printf(TEXT("BufferIndex=%d Pos=%d"), GPrevPerBoneMotionBlur.GetWriteBufferIndex(), GPrevPerBoneMotionBlur.LockedTexelPosition.GetValue());
diff --git a/Engine/Source/Runtime/Engine/Private/SkeletalRenderGPUSkin.h b/Engine/Source/Runtime/Engine/Private/SkeletalRenderGPUSkin.h
--- a/Engine/Source/Runtime/Engine/Private/SkeletalRenderGPUSkin.h
+++ b/Engine/Source/Runtime/Engine/Private/SkeletalRenderGPUSkin.h
@@ -188,7 +188,9 @@ public:
 	bool IsAppendStarted() const;
 
 	/** needed before AppendData() ccan be called */
-	ENGINE_API void StartAppend(bool bWorldIsPaused);
+	ENGINE_API void StartAppend(FRHICommandListImmediate& RHICmdList, bool bWorldIsPaused);
+
+	ENGINE_API void SetVelocityPassCallback(const TFunction<bool(FRHICommandList& RHICmdList)>& InIsVelocityFunc);
 
 	/**
 	 * use between LockData() and UnlockData()
@@ -199,7 +201,10 @@ public:
 	uint32 AppendData(FBoneSkinning *DataStart, uint32 BoneCount);
 
 	/** only call if StartAppend(), if the append wasn't started it silently ignores the call */
-	ENGINE_API void EndAppend();
+	ENGINE_API void EndAppendFence(FRHICommandListImmediate& RHICmdList);
+
+	/** only call if StartAppend(), if the append wasn't started it silently ignores the call */
+	ENGINE_API void EndAppend(FRHICommandListImmediate& RHICmdList);
 
 	/**
 	 * @param Index 0 .. PER_BONE_BUFFER_COUNT - 1, usually from GetReadBufferIndex()
@@ -228,6 +233,11 @@ public:
 		return 1;
 	}
 
+	bool IsVelocityPass(FRHICommandList& RHICmdList)
+	{
+		return IsVelocityFunc(RHICmdList);
+	}
+
 
 private:
 
@@ -240,6 +250,12 @@ private:
 	/** only valid if LockedData != 0 */
 	uint32 LockedTexelCount;
 
+	/** Fence for render thread tasks that use this. */
+	FGraphEventRef EndAppendRenderThreadTaskFence;
+
+	/** Callback if we should save bone data or not. */
+	TFunction<bool(FRHICommandList& RHICmdList)> IsVelocityFunc;
+
 	bool bWarningBufferSizeExceeded;
 
 
diff --git a/Engine/Source/Runtime/Engine/Private/StaticBoundShaderState.cpp b/Engine/Source/Runtime/Engine/Private/StaticBoundShaderState.cpp
--- a/Engine/Source/Runtime/Engine/Private/StaticBoundShaderState.cpp
+++ b/Engine/Source/Runtime/Engine/Private/StaticBoundShaderState.cpp
@@ -235,6 +235,7 @@ void SetGlobalBoundShaderState(
 	// We need to do this on the render thread
 
 	FRHICommandList* CmdList = new FRHICommandList;
+	CmdList->CopyRenderThreadContexts(RHICmdList);
 	FGraphEventRef RenderThreadCompletionEvent = TGraphTask<FSetGlobalBoundShaderStateRenderThreadTask>::CreateTask().ConstructAndDispatchWhenReady(*CmdList, GlobalBoundShaderState, FeatureLevel);
 	RHICmdList.QueueRenderThreadCommandListSubmit(RenderThreadCompletionEvent, CmdList);
 }
diff --git a/Engine/Source/Runtime/Engine/Public/GPUSkinVertexFactory.h b/Engine/Source/Runtime/Engine/Public/GPUSkinVertexFactory.h
--- a/Engine/Source/Runtime/Engine/Public/GPUSkinVertexFactory.h
+++ b/Engine/Source/Runtime/Engine/Public/GPUSkinVertexFactory.h
@@ -256,12 +256,6 @@ private: // -------------------------------------------------------
 	/** Buffer size in texels */
 	uint32 SizeX;
 
-	/** Buffer size, as allocated */
-	uint32 AllocSize;
-
-	/** Allocated Buffer */
-	void *AllocBlock;
-
 	/** @return in bytes */
 	uint32 ComputeMemorySize();
 };
diff --git a/Engine/Source/Runtime/Engine/Public/SceneView.h b/Engine/Source/Runtime/Engine/Public/SceneView.h
--- a/Engine/Source/Runtime/Engine/Public/SceneView.h
+++ b/Engine/Source/Runtime/Engine/Public/SceneView.h
@@ -503,9 +503,6 @@ public:
 	/** True if this scene was created from a game world. */
 	bool bIsGameView;
 
-	/** Whether this view should override the view family's EngineShowFlags.Materials setting and force it on. */
-	bool bForceShowMaterials;
-
 	/** For sanity checking casts that are assumed to be safe. */
 	bool bIsViewInfo;
 
diff --git a/Engine/Source/Runtime/RHI/Private/RHICommandList.cpp b/Engine/Source/Runtime/RHI/Private/RHICommandList.cpp
--- a/Engine/Source/Runtime/RHI/Private/RHICommandList.cpp
+++ b/Engine/Source/Runtime/RHI/Private/RHICommandList.cpp
@@ -32,23 +32,58 @@ static TAutoConsoleVariable<int32> CVarRHICmdUseParallelAlgorithms(
 TAutoConsoleVariable<int32> CVarRHICmdWidth(
 	TEXT("r.RHICmdWidth"), 
 	8,
-	TEXT("Number of threads."));
+	TEXT("Controls the task granularity of a great number of things in the parallel renderer."));
 
 static TAutoConsoleVariable<int32> CVarRHICmdUseDeferredContexts(
 	TEXT("r.RHICmdUseDeferredContexts"),
 	1,
 	TEXT("True to use deferred contexts to parallelize command list execution. Only available on some RHIs."));
 
+TAutoConsoleVariable<int32> CVarRHICmdFlushRenderThreadTasks(
+	TEXT("r.RHICmdFlushRenderThreadTasks"),
+	0,
+	TEXT("If true, then we flush the render thread tasks every pass. For issue diagnosis. This is a master switch for more granular cvars."));
+
+TAutoConsoleVariable<int32> CVarRHICmdFlushUpdateTextureReference(
+	TEXT("r.RHICmdFlushUpdateTextureReference"),
+	0,
+	TEXT("If true, then we flush the rhi thread when we do RHIUpdateTextureReference, otherwise this is deferred. For issue diagnosis."));
+
 static TAutoConsoleVariable<int32> CVarRHICmdFlushOnQueueParallelSubmit(
 	TEXT("r.RHICmdFlushOnQueueParallelSubmit"),
 	0,
-	TEXT("Wait for completion of parallel commandlists immediately after submitting.  Debug feature. Only available on some RHIs."));
+	TEXT("Wait for completion of parallel commandlists immediately after submitting. For issue diagnosis. Only available on some RHIs."));
 
-static TAutoConsoleVariable<int32> CVarRHIMergeSmallDeferredContexts(
-	TEXT("r.RHIMergeSmallDeferredContexts"),
+static TAutoConsoleVariable<int32> CVarRHICmdMergeSmallDeferredContexts(
+	TEXT("r.RHICmdMergeSmallDeferredContexts"),
 	1,
 	TEXT("When it can be determined, merge small parallel translate tasks based on r.RHICmdMinDrawsPerParallelCmdList."));
 
+static TAutoConsoleVariable<int32> CVarRHICmdBufferWriteLocks(
+	TEXT("r.RHICmdBufferWriteLocks"),
+	1,
+	TEXT("Only relevant with an RHI thread. Debugging option to diagnose problems with buffered locks."));
+
+static TAutoConsoleVariable<int32> CVarRHICmdAsyncRHIThreadDispatch(
+	TEXT("r.RHICmdAsyncRHIThreadDispatch"),
+	0,
+	TEXT("Experiemental option to do RHI dispatches async. Currently deadlocks. This causes a greater 'wait for rhi thread' penalty, but keeps data flowing to the RHI thread faster."));
+
+static TAutoConsoleVariable<int32> CVarRHICmdCollectRHIThreadStatsFromHighLevel(
+	TEXT("r.RHICmdCollectRHIThreadStatsFromHighLevel"),
+	1,
+	TEXT("This pushes stats on the RHI thread executes so you can determine which high level pass they came from. This has an adverse effect on framerate. This is on by default."));
+
+static TAutoConsoleVariable<int32> CVarRHICmdUseThread(
+	TEXT("r.RHICmdUseThread"),
+	1,
+	TEXT("Uses the RHI thread. For issue diagnosis."));
+
+static TAutoConsoleVariable<int32> CVarRHICmdForceRHIFlush(
+	TEXT("r.RHICmdForceRHIFlush"),
+	0,
+	TEXT("Force a flush for every task sent to the RHI thread. For issue diagnosis."));
+
 
 RHI_API FRHICommandListExecutor GRHICommandList;
 
@@ -56,7 +91,8 @@ static FGraphEventArray AllOutstandingTasks;
 static FGraphEventArray WaitOutstandingTasks;
 static FGraphEventRef RHIThreadTask;
 static FGraphEventRef RenderThreadSublistDispatchTask;
-static FGraphEventRef RenderThreadExecutionTask;
+static FGraphEventRef RHIThreadBufferLockFence;
+
 // Used by AsyncCompute
 RHI_API FRHICommandListFenceAllocator GRHIFenceAllocator;
 
@@ -77,12 +113,40 @@ struct FRHICommandStat : public FRHICommand<FRHICommandStat>
 	}
 };
 
-bool FRHICommandListImmediate::bFlushedGlobal = false;
 void FRHICommandListImmediate::SetCurrentStat(TStatId Stat)
 {
 	new (AllocCommand<FRHICommandStat>()) FRHICommandStat(Stat);
 }
 
+DECLARE_CYCLE_STAT(TEXT("FNullGraphTask.RenderThreadTaskFence"), STAT_RenderThreadTaskFence, STATGROUP_TaskGraphTasks);
+DECLARE_CYCLE_STAT(TEXT("Render thread task fence wait"), STAT_RenderThreadTaskFenceWait, STATGROUP_TaskGraphTasks);
+FGraphEventRef FRHICommandListImmediate::RenderThreadTaskFence()
+{
+	FGraphEventRef Result;
+	check(IsInRenderingThread());
+	//@todo optimize, if there is only one outstanding, then return that instead
+	if (WaitOutstandingTasks.Num())
+	{
+		Result = TGraphTask<FNullGraphTask>::CreateTask(&WaitOutstandingTasks, ENamedThreads::RenderThread).ConstructAndDispatchWhenReady(GET_STATID(STAT_RenderThreadTaskFence), ENamedThreads::RenderThread_Local);
+	}
+	return Result;
+}
+void FRHICommandListImmediate::WaitOnRenderThreadTaskFence(FGraphEventRef& Fence)
+{
+	if (Fence.GetReference() && !Fence->IsComplete())
+	{
+		SCOPE_CYCLE_COUNTER(STAT_RenderThreadTaskFenceWait);
+		check(IsInRenderingThread() && !FTaskGraphInterface::Get().IsThreadProcessingTasks(ENamedThreads::RenderThread_Local));
+		FTaskGraphInterface::Get().WaitUntilTaskCompletes(Fence, ENamedThreads::RenderThread_Local);
+	}
+}
+
+bool FRHICommandListImmediate::AnyRenderThreadTasksOutstanding()
+{
+	return !!WaitOutstandingTasks.Num();
+}
+
+
 FRHICommandBase* GCurrentCommand = nullptr;
 
 void FRHICommandListExecutor::ExecuteInner_DoExecute(FRHICommandListBase& CmdList)
@@ -90,20 +154,25 @@ void FRHICommandListExecutor::ExecuteInner_DoExecute(FRHICommandListBase& CmdLis
 	CmdList.bExecuting = true;
 	check(CmdList.Context);
 
-	bool bDoStats = false;
-	//STAT(bDoStats = FThreadStats::IsCollectingData() && (IsInRenderingThread() || IsInRHIThread());)
 	FRHICommandListIterator Iter(CmdList);
+#if STATS
+	bool bDoStats =  CVarRHICmdCollectRHIThreadStatsFromHighLevel.GetValueOnRenderThread() > 0 && FThreadStats::IsCollectingData() && (IsInRenderingThread() || IsInRHIThread());
 	if (bDoStats)
 	{
 		while (Iter.HasCommandsLeft())
 		{
+			TStatIdData const* Stat = GCurrentExecuteStat.GetRawPointer();
+			FScopeCycleCounter Scope(GCurrentExecuteStat);
+			while (Iter.HasCommandsLeft() && Stat == GCurrentExecuteStat.GetRawPointer())
+			{
 			FRHICommandBase* Cmd = Iter.NextCommand();
 			//FPlatformMisc::Prefetch(Cmd->Next);
-			FScopeCycleCounter Scope(GCurrentExecuteStat);
 			Cmd->CallExecuteAndDestruct(CmdList);
 		}
 	}
+	}
 	else
+#endif
 	{
 		while (Iter.HasCommandsLeft())
 		{
@@ -121,19 +190,11 @@ void FRHICommandListExecutor::ExecuteInner_DoExecute(FRHICommandListBase& CmdLis
 class FExecuteRHIThreadTask
 {
 	FRHICommandListBase* RHICmdList;
-	FGraphEventRef ExecuteFenceTask;
 
 public:
 
 	FExecuteRHIThreadTask(FRHICommandListBase* InRHICmdList)
 		: RHICmdList(InRHICmdList)
-		, ExecuteFenceTask(nullptr)
-	{
-	}
-
-	FExecuteRHIThreadTask(FRHICommandListBase* InRHICmdList, FGraphEventRef InExecuteFenceTask)
-		: RHICmdList(InRHICmdList)
-		, ExecuteFenceTask(InExecuteFenceTask)
 	{
 	}
 
@@ -155,62 +216,19 @@ public:
 		SCOPE_CYCLE_COUNTER(STAT_RHIThreadExecute);
 		FRHICommandListExecutor::ExecuteInner_DoExecute(*RHICmdList);
 		delete RHICmdList;
-
-		if (ExecuteFenceTask)
-		{			
-			ExecuteFenceTask->DispatchSubsequents();
-			ExecuteFenceTask = nullptr;
-		}
-	}
-};
-
-class FExecuteFenceRHITask
-{	
-	FGraphEventRef ExecuteFenceTask;
-
-public:
-
-	FExecuteFenceRHITask(FGraphEventRef InExecuteFenceTask)		
-		: ExecuteFenceTask(InExecuteFenceTask)
-	{
-	}
-
-	FORCEINLINE TStatId GetStatId() const
-	{
-		RETURN_QUICK_DECLARE_CYCLE_STAT(FExecuteFenceRHITask, STATGROUP_TaskGraphTasks);
-	}
-
-	ENamedThreads::Type GetDesiredThread()
-	{
-		check(GRHIThread); // this should never be used on a platform that doesn't support the RHI thread
-		return ENamedThreads::RHIThread;
-	}
-
-	static ESubsequentsMode::Type GetSubsequentsMode() { return ESubsequentsMode::TrackSubsequents; }
-
-	void DoTask(ENamedThreads::Type CurrentThread, const FGraphEventRef& MyCompletionGraphEvent)
-	{
-		ExecuteFenceTask->DispatchSubsequents();
-		ExecuteFenceTask = nullptr;		
 	}
 };
 
 class FDispatchRHIThreadTask
 {
 	FRHICommandListBase* RHICmdList;
-	FGraphEventRef ExecuteFenceTask;
+	bool bAnyThread;
 
 public:
 
-	FDispatchRHIThreadTask(FRHICommandListBase* InRHICmdList)
+	FDispatchRHIThreadTask(FRHICommandListBase* InRHICmdList, bool bInAnyThread)
 		: RHICmdList(InRHICmdList)
-		, ExecuteFenceTask(nullptr)
-	{		
-	}
-
-	FDispatchRHIThreadTask(FRHICommandListBase* InRHICmdList, FGraphEventRef InExecuteFenceTask)
-		: RHICmdList(InRHICmdList)
-		, ExecuteFenceTask(InExecuteFenceTask)
+		, bAnyThread(bInAnyThread)
 	{		
 	}
 
@@ -222,72 +240,73 @@ public:
 	ENamedThreads::Type GetDesiredThread()
 	{
 		check(GRHIThread); // this should never be used on a platform that doesn't support the RHI thread
-		return ENamedThreads::RenderThread_Local;
+		return bAnyThread ? ENamedThreads::AnyThread : ENamedThreads::RenderThread_Local;
 	}
 
 	static ESubsequentsMode::Type GetSubsequentsMode() { return ESubsequentsMode::TrackSubsequents; }
 
 	void DoTask(ENamedThreads::Type CurrentThread, const FGraphEventRef& MyCompletionGraphEvent)
 	{
-		check(IsInRenderingThread());
+		check(bAnyThread || IsInRenderingThread());
 		FGraphEventArray Prereq;
 		if (RHIThreadTask.GetReference())
 		{
 			Prereq.Add(RHIThreadTask);
 		}
-		RHIThreadTask = TGraphTask<FExecuteRHIThreadTask>::CreateTask(&Prereq, ENamedThreads::RenderThread).ConstructAndDispatchWhenReady(RHICmdList, ExecuteFenceTask);
-		ExecuteFenceTask = nullptr;
+		RHIThreadTask = TGraphTask<FExecuteRHIThreadTask>::CreateTask(&Prereq, ENamedThreads::RenderThread).ConstructAndDispatchWhenReady(RHICmdList);
 	}
 };
 
-static TAutoConsoleVariable<int32> CVarRHICmdUseThread(
-	TEXT("r.RHICmdUseThread"),
-	1,
-	TEXT("Uses the RHI thread. Just for testing, should be removed. (GilG)\n")
-	TEXT("0: Disable, 1: Enable"),
-	ECVF_Cheat);
-static TAutoConsoleVariable<int32> CVarRHICmdForceRHIFlush(
-	TEXT("r.RHICmdForceRHIFlush"),
-	0,
-	TEXT("Force a flush for every task sent to the RHI thread. Just for testing, should be removed. (GilG)\n")
-	TEXT("0: Disable, 1: Enable"),
-	ECVF_Cheat);
+
 
 void FRHICommandListExecutor::ExecuteInner(FRHICommandListBase& CmdList)
 {
 	check(CmdList.RTTasks.Num() == 0 && CmdList.HasCommands()); 
 
 	bool bIsInRenderingThread = IsInRenderingThread();
+	bool bIsInGameThread = IsInGameThread();
 	if (GRHIThread)
 	{
+		bool bAsyncSubmit = false;
 		if (bIsInRenderingThread)
 		{
-			if (RHIThreadTask.GetReference() && RHIThreadTask->IsComplete())
+			if (!bIsInGameThread && !FTaskGraphInterface::Get().IsThreadProcessingTasks(ENamedThreads::RenderThread_Local))
 			{
-				RHIThreadTask = nullptr;
+				QUICK_SCOPE_CYCLE_COUNTER(STAT_FRHICommandListExecutor_ExecuteInner_DoTasksBeforeDispatch);
+				// move anything down the pipe that needs to go
+				FTaskGraphInterface::Get().ProcessThreadUntilIdle(ENamedThreads::RenderThread_Local);
 			}
+			bAsyncSubmit = CVarRHICmdAsyncRHIThreadDispatch.GetValueOnRenderThread() > 0;
 			if (RenderThreadSublistDispatchTask.GetReference() && RenderThreadSublistDispatchTask->IsComplete())
 			{
 				RenderThreadSublistDispatchTask = nullptr;
+				if (bAsyncSubmit && RHIThreadTask.GetReference() && RHIThreadTask->IsComplete())
+			{
+				RHIThreadTask = nullptr;
+			}
 			}
-			if (RenderThreadExecutionTask.GetReference() && RenderThreadExecutionTask->IsComplete())
+			if (!bAsyncSubmit && RHIThreadTask.GetReference() && RHIThreadTask->IsComplete())
 			{
-				RenderThreadExecutionTask = nullptr;
+				RHIThreadTask = nullptr;
 			}
 		}
-		if (CVarRHICmdUseThread.GetValueOnRenderThread() > 0 && bIsInRenderingThread && !IsInGameThread())
+		if (CVarRHICmdUseThread.GetValueOnRenderThread() > 0 && bIsInRenderingThread && !bIsInGameThread)
+		{
+			FRHICommandList* SwapCmdList;
 		{
-			FRHICommandList* SwapCmdList = new FRHICommandList;
+				QUICK_SCOPE_CYCLE_COUNTER(STAT_FRHICommandListExecutor_SwapCmdLists);
+				SwapCmdList = new FRHICommandList;
 
 			// Super scary stuff here, but we just want the swap command list to inherit everything and leave the immediate command list wiped.
 			// we should make command lists virtual and transfer ownership rather than this devious approach
 			static_assert(sizeof(FRHICommandList) == sizeof(FRHICommandListImmediate), "We are memswapping FRHICommandList and FRHICommandListImmediate; they need to be swappable.");
 			SwapCmdList->ExchangeCmdList(CmdList);
+			}
+			QUICK_SCOPE_CYCLE_COUNTER(STAT_FRHICommandListExecutor_SubmitTasks);
 
 			//if we use a FDispatchRHIThreadTask, we must have it pass an event along to the FExecuteRHIThreadTask it will spawn so that fences can know which event to wait on for execution completion
 			//before the dispatch completes.
 			//if we use a FExecuteRHIThreadTask directly we pass the same event just to keep things consistent.
-			RenderThreadExecutionTask = FGraphEvent::CreateGraphEvent();
 			if (AllOutstandingTasks.Num() || RenderThreadSublistDispatchTask.GetReference())
 			{
 				FGraphEventArray Prereq(AllOutstandingTasks);
@@ -296,16 +315,17 @@ void FRHICommandListExecutor::ExecuteInner(FRHICommandListBase& CmdList)
 				{
 					Prereq.Add(RenderThreadSublistDispatchTask);
 				}
-				RenderThreadSublistDispatchTask = TGraphTask<FDispatchRHIThreadTask>::CreateTask(&Prereq, ENamedThreads::RenderThread).ConstructAndDispatchWhenReady(SwapCmdList, RenderThreadExecutionTask);
+				RenderThreadSublistDispatchTask = TGraphTask<FDispatchRHIThreadTask>::CreateTask(&Prereq, ENamedThreads::RenderThread).ConstructAndDispatchWhenReady(SwapCmdList, bAsyncSubmit);
 			}
 			else
 			{
+				check(!RenderThreadSublistDispatchTask.GetReference()); // if we are doing submits, there better not be any of these in flight since then the RHIThreadTask would get out of order.
 				FGraphEventArray Prereq;
 				if (RHIThreadTask.GetReference())
 				{
 					Prereq.Add(RHIThreadTask);
 				}
-				RHIThreadTask = TGraphTask<FExecuteRHIThreadTask>::CreateTask(&Prereq, ENamedThreads::RenderThread).ConstructAndDispatchWhenReady(SwapCmdList, RenderThreadExecutionTask);
+				RHIThreadTask = TGraphTask<FExecuteRHIThreadTask>::CreateTask(&Prereq, ENamedThreads::RenderThread).ConstructAndDispatchWhenReady(SwapCmdList);
 			}
 			if (CVarRHICmdForceRHIFlush.GetValueOnRenderThread() > 0 )
 			{
@@ -469,6 +489,20 @@ void FRHICommandListExecutor::CheckNoOutstandingCmdLists()
 bool FRHICommandListExecutor::IsRHIThreadActive()
 {
 	checkSlow(IsInRenderingThread());
+	bool bAsyncSubmit = CVarRHICmdAsyncRHIThreadDispatch.GetValueOnRenderThread() > 0;
+	if (bAsyncSubmit)
+	{
+		if (RenderThreadSublistDispatchTask.GetReference() && RenderThreadSublistDispatchTask->IsComplete())
+		{
+			RenderThreadSublistDispatchTask = nullptr;
+		}
+		if (RenderThreadSublistDispatchTask.GetReference())
+		{
+			return true; // it might become active at any time
+		}
+		// otherwise we can safely look at RHIThreadTask
+	}
+
 	if (RHIThreadTask.GetReference() && RHIThreadTask->IsComplete())
 	{
 		RHIThreadTask = nullptr;
@@ -476,30 +510,47 @@ bool FRHICommandListExecutor::IsRHIThreadActive()
 	return !!RHIThreadTask.GetReference();
 }
 
-DECLARE_CYCLE_STAT(TEXT("FNullGraphTask.RHIThreadFence"), STAT_RHIThreadFence, STATGROUP_TaskGraphTasks);
+bool FRHICommandListExecutor::IsRHIThreadCompletelyFlushed()
+{
+	if (IsRHIThreadActive() || GetImmediateCommandList().HasCommands())
+	{
+		return false;
+	}
+	if (RenderThreadSublistDispatchTask.GetReference() && RenderThreadSublistDispatchTask->IsComplete())
+	{
+		RenderThreadSublistDispatchTask = nullptr;
+	}
+	return !RenderThreadSublistDispatchTask;
+}
 
-FGraphEventRef FRHICommandListExecutor::RHIThreadFence()
+struct FRHICommandRHIThreadFence : public FRHICommand<FRHICommandRHIThreadFence>
 {
-	check(IsInRenderingThread() && GRHIThread);
+	FGraphEventRef Fence;
+	FORCEINLINE_DEBUGGABLE FRHICommandRHIThreadFence()
+		: Fence(FGraphEvent::CreateGraphEvent())
+{
+	}
+	void Execute(FRHICommandListBase& CmdList)
 	{
-		QUICK_SCOPE_CYCLE_COUNTER(STAT_RHIThreadFence_Dispatch);
-		FRHICommandListExecutor::GetImmediateCommandList().ImmediateFlush(EImmediateFlushType::DispatchToRHIThread);
+		Fence->DispatchSubsequents();
+		Fence = nullptr;
 	}
+};
+
 
-	// every task dispatched or executed through the flush creates a new RenderThreadExecutionTask that represents the 'execution' of the last dispatched
-	// set of RHI commands.
-	if (RenderThreadExecutionTask.GetReference() && RenderThreadExecutionTask->IsComplete())
+FGraphEventRef FRHICommandListImmediate::RHIThreadFence()
 	{
-		RenderThreadExecutionTask = nullptr;
+	check(IsInRenderingThread() && GRHIThread);
+	FRHICommandRHIThreadFence* Cmd = new (AllocCommand<FRHICommandRHIThreadFence>()) FRHICommandRHIThreadFence();
+	return Cmd->Fence;
 	}		
 	
-	return RenderThreadExecutionTask;
-}
 void FRHICommandListExecutor::WaitOnRHIThreadFence(FGraphEventRef& Fence)
 {
-	check(IsInRenderingThread() && GRHIThread);
+	check(IsInRenderingThread());
 	if (Fence.GetReference() && !Fence->IsComplete())
 	{
+		check(GRHIThread);
 		QUICK_SCOPE_CYCLE_COUNTER(STAT_WaitOnRHIThreadFence_Wait);
 		if (FTaskGraphInterface::Get().IsThreadProcessingTasks(ENamedThreads::RenderThread_Local))
 		{
@@ -511,6 +562,8 @@ void FRHICommandListExecutor::WaitOnRHIThreadFence(FGraphEventRef& Fence)
 }
 
 
+
+
 FRHICommandListBase::FRHICommandListBase()
 	: MemManager(0)
 {
@@ -542,6 +595,10 @@ void FRHICommandListBase::Reset()
 	StateCache = nullptr;
 #endif
 	UID = GRHICommandList.UIDCounter.Increment();
+	for (int32 Index = 0; ERenderThreadContext(Index) < ERenderThreadContext::Num; Index++)
+	{
+		RenderThreadContexts[Index] = nullptr;
+	}
 }
 
 
@@ -681,6 +738,12 @@ void FRHICommandListBase::QueueParallelAsyncCommandListSubmit(FGraphEventRef* An
 	if (GRHIThread)
 	{
 		FRHICommandListExecutor::GetImmediateCommandList().ImmediateFlush(EImmediateFlushType::DispatchToRHIThread); // we should start on the stuff before this async list
+
+		// as good a place as any to clear this
+		if (RHIThreadBufferLockFence.GetReference() && RHIThreadBufferLockFence->IsComplete())
+		{
+			RHIThreadBufferLockFence = nullptr;
+		}
 	}
 #if !UE_BUILD_SHIPPING
 	// do a flush before hand so we can tell if it was this parallel set that broke something, or what came before.
@@ -699,7 +762,7 @@ void FRHICommandListBase::QueueParallelAsyncCommandListSubmit(FGraphEventRef* An
 		}
 		if (ContextContainer)
 		{
-			bool bMerge = !!CVarRHIMergeSmallDeferredContexts.GetValueOnRenderThread();
+			bool bMerge = !!CVarRHICmdMergeSmallDeferredContexts.GetValueOnRenderThread();
 
 			int32 EffectiveThreads = 0;
 			int32 Start = 0;
@@ -764,6 +827,12 @@ void FRHICommandListBase::QueueParallelAsyncCommandListSubmit(FGraphEventRef* An
 				}
 				UE_CLOG(bSpewMerge, LogTemp, Display, TEXT("Parallel translate %d->%d    %dKB mem   %d draws (-1 = unknown)"), Start, Last, FMath::DivideAndRoundUp(TotalMem, 1024), DrawCnt);
 
+				// this is used to ensure that any old buffer locks are completed before we start any parallel translates
+				if (RHIThreadBufferLockFence.GetReference())
+				{
+					Prereq.Add(RHIThreadBufferLockFence);
+				}
+
 				FGraphEventRef TranslateCompletionEvent = TGraphTask<FParallelTranslateCommandList>::CreateTask(&Prereq, ENamedThreads::RenderThread).ConstructAndDispatchWhenReady(&RHICmdLists[0], 1 + Last - Start, ContextContainer);
 
 				AllOutstandingTasks.Add(TranslateCompletionEvent);
@@ -846,9 +915,9 @@ struct FRHICommandWaitForAndSubmitRTSubList : public FRHICommand<FRHICommandWait
 	{
 		INC_DWORD_STAT_BY(STAT_RTChainLinkCount, 1);
 		{
-			SCOPE_CYCLE_COUNTER(STAT_RTChainWait);
-			if (!EventToWaitFor->IsComplete())
+			if (EventToWaitFor.GetReference() && !EventToWaitFor->IsComplete())
 			{
+			SCOPE_CYCLE_COUNTER(STAT_RTChainWait);
 				check(!GRHIThread || !IsInRHIThread()); // things should not be dispatched if they can't complete without further waits
 				if (IsInRenderingThread())
 				{
@@ -875,8 +944,12 @@ struct FRHICommandWaitForAndSubmitRTSubList : public FRHICommand<FRHICommandWait
 
 void FRHICommandListBase::QueueRenderThreadCommandListSubmit(FGraphEventRef& RenderThreadCompletionEvent, class FRHICommandList* CmdList)
 {
-	check(!IsInRenderingThread() && !IsImmediate() && !IsInRHIThread());
+	check(!IsInRHIThread());
+	if (RenderThreadCompletionEvent.GetReference())
+	{
+		check(!IsInRenderingThread() && !IsImmediate());
 	RTTasks.Add(RenderThreadCompletionEvent);
+	}
 	new (AllocCommand<FRHICommandWaitForAndSubmitRTSubList>()) FRHICommandWaitForAndSubmitRTSubList(RenderThreadCompletionEvent, CmdList);
 }
 
@@ -1018,7 +1091,7 @@ void FRHICommandListBase::WaitForTasks(bool bKnownToBeComplete)
 		{
 			if (!WaitOutstandingTasks[Index]->IsComplete())
 			{
-				check(!bKnownToBeComplete);
+				ensure(!bKnownToBeComplete);
 				bAny = true;
 				break;
 			}
@@ -1026,12 +1099,39 @@ void FRHICommandListBase::WaitForTasks(bool bKnownToBeComplete)
 		if (bAny)
 		{
 			SCOPE_CYCLE_COUNTER(STAT_ExplicitWait);
+			check(!FTaskGraphInterface::Get().IsThreadProcessingTasks(ENamedThreads::RenderThread_Local));
 			FTaskGraphInterface::Get().WaitUntilTasksComplete(WaitOutstandingTasks, ENamedThreads::RenderThread_Local);
 		}
 		WaitOutstandingTasks.Reset();
 	}
 }
 
+FScopedCommandListWaitForTasks::~FScopedCommandListWaitForTasks()
+{
+	check(IsInRenderingThread());
+	if (bWaitForTasks)
+	{
+		if (GRHIThread)
+		{
+#if 0
+			{
+				QUICK_SCOPE_CYCLE_COUNTER(STAT_FScopedCommandListWaitForTasks_Dispatch);
+				RHICmdList.ImmediateFlush(EImmediateFlushType::DispatchToRHIThread);
+			}
+#endif
+			{
+				QUICK_SCOPE_CYCLE_COUNTER(STAT_FScopedCommandListWaitForTasks_WaitAsync);
+				RHICmdList.ImmediateFlush(EImmediateFlushType::WaitForOutstandingTasksOnly);
+			}
+		}
+		else
+		{
+			QUICK_SCOPE_CYCLE_COUNTER(STAT_FScopedCommandListWaitForTasks_Flush);
+			RHICmdList.ImmediateFlush(EImmediateFlushType::FlushRHIThread);
+		}
+	}
+}
+
 DECLARE_CYCLE_STAT(TEXT("Explicit wait for dispatch"), STAT_ExplicitWaitDispatch, STATGROUP_RHICMDLIST);
 void FRHICommandListBase::WaitForDispatch()
 {
@@ -1058,10 +1158,42 @@ void FRHICommandListBase::WaitForDispatch()
 }
 
 DECLARE_CYCLE_STAT(TEXT("Explicit wait for RHI thread"), STAT_ExplicitWaitRHIThread, STATGROUP_RHICMDLIST);
+DECLARE_CYCLE_STAT(TEXT("Explicit wait for RHI thread async dispatch"), STAT_ExplicitWaitRHIThread_Dispatch, STATGROUP_RHICMDLIST);
 DECLARE_CYCLE_STAT(TEXT("Deep spin for stray resource init."), STAT_SpinWaitRHIThread, STATGROUP_RHICMDLIST);
 void FRHICommandListBase::WaitForRHIThreadTasks()
 {
 	check(IsImmediate() && IsInRenderingThread());
+	bool bAsyncSubmit = CVarRHICmdAsyncRHIThreadDispatch.GetValueOnRenderThread() > 0;
+	if (bAsyncSubmit)
+	{
+		if (RenderThreadSublistDispatchTask.GetReference() && RenderThreadSublistDispatchTask->IsComplete())
+		{
+			RenderThreadSublistDispatchTask = nullptr;
+		}
+		while (RenderThreadSublistDispatchTask.GetReference())
+		{
+			SCOPE_CYCLE_COUNTER(STAT_ExplicitWaitRHIThread_Dispatch);
+			if (FTaskGraphInterface::Get().IsThreadProcessingTasks(ENamedThreads::RenderThread_Local))
+			{
+				// we have to spin here because all task threads might be stalled, meaning the fire event anythread task might not be hit.
+				// todo, add a third queue
+				SCOPE_CYCLE_COUNTER(STAT_SpinWaitRHIThread);
+				while (!RenderThreadSublistDispatchTask->IsComplete())
+				{
+					FPlatformProcess::Sleep(0);
+				}
+			}
+			else
+			{
+				FTaskGraphInterface::Get().WaitUntilTaskCompletes(RenderThreadSublistDispatchTask, ENamedThreads::RenderThread_Local);
+			}
+			if (RenderThreadSublistDispatchTask.GetReference() && RenderThreadSublistDispatchTask->IsComplete())
+			{
+				RenderThreadSublistDispatchTask = nullptr;
+			}
+		}
+		// now we can safely look at RHIThreadTask
+	}
 	if (RHIThreadTask.GetReference() && RHIThreadTask->IsComplete())
 	{
 		RHIThreadTask = nullptr;
@@ -1124,3 +1256,275 @@ void FRHICommandList::operator delete(void *RawMemory)
 	//FMemory::Free(RawMemory);
 }	
 
+///////// Pass through functions that allow RHIs to optimize certain calls.
+
+
+FVertexBufferRHIRef FDynamicRHI::CreateVertexBuffer_RenderThread(class FRHICommandListImmediate& RHICmdList, uint32 Size, uint32 InUsage, FRHIResourceCreateInfo& CreateInfo)
+{
+	check(IsInRenderingThread());
+	if(GRHIThread)
+	{
+		QUICK_SCOPE_CYCLE_COUNTER(STAT_RHIMETHOD_CreateVertexBuffer_WaitRHI);
+		RHICmdList.ImmediateFlush(EImmediateFlushType::WaitForRHIThread);
+	}
+	return GDynamicRHI->RHICreateVertexBuffer(Size, InUsage, CreateInfo);
+}
+FIndexBufferRHIRef FDynamicRHI::CreateIndexBuffer_RenderThread(class FRHICommandListImmediate& RHICmdList, uint32 Stride, uint32 Size, uint32 InUsage, FRHIResourceCreateInfo& CreateInfo)
+{
+	check(IsInRenderingThread());
+	if(GRHIThread)
+	{
+		QUICK_SCOPE_CYCLE_COUNTER(STAT_RHIMETHOD_CreateIndexBuffer_WaitRHI);
+		RHICmdList.ImmediateFlush(EImmediateFlushType::WaitForRHIThread);
+	}
+	return GDynamicRHI->RHICreateIndexBuffer(Stride, Size, InUsage, CreateInfo);
+}
+
+FShaderResourceViewRHIRef FDynamicRHI::CreateShaderResourceView_RenderThread(class FRHICommandListImmediate& RHICmdList, FVertexBufferRHIParamRef VertexBuffer, uint32 Stride, uint8 Format)
+{
+	check(IsInRenderingThread());
+	if(GRHIThread)
+	{
+		QUICK_SCOPE_CYCLE_COUNTER(STAT_RHIMETHOD_CreateShaderResourceView_WaitRHI);
+		RHICmdList.ImmediateFlush(EImmediateFlushType::WaitForRHIThread);
+	}
+	return GDynamicRHI->RHICreateShaderResourceView(VertexBuffer, Stride, Format);
+}
+
+struct FRHICommandUpdateVertexBuffer : public FRHICommand<FRHICommandUpdateVertexBuffer>
+{
+	FVertexBufferRHIParamRef VertexBuffer;
+	void* Buffer;
+	uint32 BufferSize;
+	uint32 Offset;
+
+	FORCEINLINE_DEBUGGABLE FRHICommandUpdateVertexBuffer(FVertexBufferRHIParamRef InVertexBuffer, void* InBuffer, uint32 InOffset, uint32 InBufferSize)
+		: VertexBuffer(InVertexBuffer)
+		, Buffer(InBuffer)
+		, BufferSize(InBufferSize)
+		, Offset(InOffset)
+	{
+	}
+	void Execute(FRHICommandListBase& CmdList)
+	{
+		QUICK_SCOPE_CYCLE_COUNTER(STAT_FRHICommandUpdateVertexBuffer_Execute);
+		void* Data = GDynamicRHI->RHILockVertexBuffer(VertexBuffer, Offset, BufferSize, RLM_WriteOnly);
+		FMemory::Memcpy(Data, Buffer, BufferSize);
+		FMemory::Free(Buffer);
+		GDynamicRHI->RHIUnlockVertexBuffer(VertexBuffer);
+	}
+};
+
+struct FRHICommandUpdateIndexBuffer : public FRHICommand<FRHICommandUpdateIndexBuffer>
+{
+	FIndexBufferRHIParamRef IndexBuffer;
+	void* Buffer;
+	uint32 BufferSize;
+	uint32 Offset;
+
+	FORCEINLINE_DEBUGGABLE FRHICommandUpdateIndexBuffer(FIndexBufferRHIParamRef InIndexBuffer, void* InBuffer, uint32 InOffset, uint32 InBufferSize)
+		: IndexBuffer(InIndexBuffer)
+		, Buffer(InBuffer)
+		, BufferSize(InBufferSize)
+		, Offset(InOffset)
+	{
+	}
+	void Execute(FRHICommandListBase& CmdList)
+	{
+		QUICK_SCOPE_CYCLE_COUNTER(STAT_FRHICommandUpdateIndexBuffer_Execute);
+		void* Data = GDynamicRHI->RHILockIndexBuffer(IndexBuffer, Offset, BufferSize, RLM_WriteOnly);
+		FMemory::Memcpy(Data, Buffer, BufferSize);
+		FMemory::Free(Buffer);
+		GDynamicRHI->RHIUnlockIndexBuffer(IndexBuffer);
+	}
+};
+
+static struct FLockTracker
+{
+	struct FLockParams
+	{
+		void* Buffer;
+		uint32 BufferSize;
+		uint32 Offset;
+		EResourceLockMode LockMode;
+
+		FORCEINLINE_DEBUGGABLE FLockParams(void* InBuffer, uint32 InOffset, uint32 InBufferSize, EResourceLockMode InLockMode)
+			: Buffer(InBuffer)
+			, BufferSize(InBufferSize)
+			, Offset(InOffset)
+			, LockMode(InLockMode)
+		{
+		}
+	};
+	TMap<void *, FLockParams> OutstandingLocks;
+	uint32 TotalMemoryOutstanding;
+
+	FLockTracker()
+	{
+		TotalMemoryOutstanding = 0;
+	}
+
+	void Lock(void* RHIBuffer, void* Buffer, uint32 Offset, uint32 SizeRHI, EResourceLockMode LockMode)
+	{
+		check(!OutstandingLocks.Contains(RHIBuffer));
+		OutstandingLocks.Add(RHIBuffer, FLockParams(Buffer, Offset, SizeRHI, LockMode));
+		TotalMemoryOutstanding += SizeRHI;
+	}
+	FLockParams Unlock(void* RHIBuffer)
+	{
+		return OutstandingLocks.FindAndRemoveChecked(RHIBuffer);
+	}
+} GLockTracker;
+
+
+void* FDynamicRHI::LockVertexBuffer_RenderThread(class FRHICommandListImmediate& RHICmdList, FVertexBufferRHIParamRef VertexBuffer, uint32 Offset, uint32 SizeRHI, EResourceLockMode LockMode)
+{
+	QUICK_SCOPE_CYCLE_COUNTER(STAT_FDynamicRHI_LockVertexBuffer_RenderThread);
+	check(IsInRenderingThread());
+	bool bBuffer = CVarRHICmdBufferWriteLocks.GetValueOnRenderThread() > 0;
+	void* Result;
+	if (!bBuffer || LockMode != RLM_WriteOnly || RHICmdList.Bypass() || !GRHIThread)
+	{
+		RHICmdList.ImmediateFlush(EImmediateFlushType::FlushRHIThread); 
+		Result = GDynamicRHI->RHILockVertexBuffer(VertexBuffer, Offset, SizeRHI, LockMode);
+	}
+	else
+	{
+		Result = FMemory::Malloc(SizeRHI, 16);
+	}
+	check(Result);
+	GLockTracker.Lock(VertexBuffer, Result, Offset, SizeRHI, LockMode);
+	return Result;
+}
+
+void FDynamicRHI::UnlockVertexBuffer_RenderThread(class FRHICommandListImmediate& RHICmdList, FVertexBufferRHIParamRef VertexBuffer)
+{
+	QUICK_SCOPE_CYCLE_COUNTER(STAT_FDynamicRHI_UnlockVertexBuffer_RenderThread);
+	check(IsInRenderingThread());
+	bool bBuffer = CVarRHICmdBufferWriteLocks.GetValueOnRenderThread() > 0;
+	FLockTracker::FLockParams Params = GLockTracker.Unlock(VertexBuffer);
+	if (!bBuffer || Params.LockMode != RLM_WriteOnly || RHICmdList.Bypass() || !GRHIThread)
+	{
+		QUICK_SCOPE_CYCLE_COUNTER(STAT_RHIMETHOD_UnlockVertexBuffer_Flush);
+		RHICmdList.ImmediateFlush(EImmediateFlushType::FlushRHIThread); 
+		GDynamicRHI->RHIUnlockVertexBuffer(VertexBuffer);
+		GLockTracker.TotalMemoryOutstanding = 0;
+	}	
+	else
+	{
+		new (RHICmdList.AllocCommand<FRHICommandUpdateVertexBuffer>()) FRHICommandUpdateVertexBuffer(VertexBuffer, Params.Buffer, Params.Offset, Params.BufferSize);
+		RHIThreadBufferLockFence = RHICmdList.RHIThreadFence();
+		if (GLockTracker.TotalMemoryOutstanding > 256 * 1024)
+		{
+			// we could be loading a level or something, lets get this stuff going
+			RHICmdList.ImmediateFlush(EImmediateFlushType::DispatchToRHIThread); 
+			GLockTracker.TotalMemoryOutstanding = 0;
+		}
+	}
+}
+
+void* FDynamicRHI::LockIndexBuffer_RenderThread(class FRHICommandListImmediate& RHICmdList, FIndexBufferRHIParamRef IndexBuffer, uint32 Offset, uint32 SizeRHI, EResourceLockMode LockMode)
+{
+	QUICK_SCOPE_CYCLE_COUNTER(STAT_FDynamicRHI_LockIndexBuffer_RenderThread);
+	check(IsInRenderingThread());
+	bool bBuffer = CVarRHICmdBufferWriteLocks.GetValueOnRenderThread() > 0;
+	void* Result;
+	if (!bBuffer || LockMode != RLM_WriteOnly || RHICmdList.Bypass() || !GRHIThread)
+	{
+		RHICmdList.ImmediateFlush(EImmediateFlushType::FlushRHIThread); 
+		Result = GDynamicRHI->RHILockIndexBuffer(IndexBuffer, Offset, SizeRHI, LockMode);
+	}
+	else
+	{
+		Result = FMemory::Malloc(SizeRHI, 16);
+	}
+	check(Result);
+	GLockTracker.Lock(IndexBuffer, Result, Offset, SizeRHI, LockMode);
+	return Result;
+}
+
+void FDynamicRHI::UnlockIndexBuffer_RenderThread(class FRHICommandListImmediate& RHICmdList, FIndexBufferRHIParamRef IndexBuffer)
+{
+	QUICK_SCOPE_CYCLE_COUNTER(STAT_FDynamicRHI_UnlockIndexBuffer_RenderThread);
+	check(IsInRenderingThread());
+	bool bBuffer = CVarRHICmdBufferWriteLocks.GetValueOnRenderThread() > 0;
+	FLockTracker::FLockParams Params = GLockTracker.Unlock(IndexBuffer);
+	if (!bBuffer || Params.LockMode != RLM_WriteOnly || RHICmdList.Bypass() || !GRHIThread)
+	{
+		QUICK_SCOPE_CYCLE_COUNTER(STAT_RHIMETHOD_UnlockIndexBuffer_Flush);
+		RHICmdList.ImmediateFlush(EImmediateFlushType::FlushRHIThread); 
+		GDynamicRHI->RHIUnlockIndexBuffer(IndexBuffer);
+		GLockTracker.TotalMemoryOutstanding = 0;
+	}	
+	else
+	{
+		new (RHICmdList.AllocCommand<FRHICommandUpdateIndexBuffer>()) FRHICommandUpdateIndexBuffer(IndexBuffer, Params.Buffer, Params.Offset, Params.BufferSize);
+		RHIThreadBufferLockFence = RHICmdList.RHIThreadFence();
+		if (GLockTracker.TotalMemoryOutstanding > 256 * 1024)
+		{
+			// we could be loading a level or something, lets get this stuff going
+			RHICmdList.ImmediateFlush(EImmediateFlushType::DispatchToRHIThread); 
+			GLockTracker.TotalMemoryOutstanding = 0;
+		}
+	}
+}
+
+FTexture2DRHIRef FDynamicRHI::AsyncReallocateTexture2D_RenderThread(class FRHICommandListImmediate& RHICmdList, FTexture2DRHIParamRef Texture2D, int32 NewMipCount, int32 NewSizeX, int32 NewSizeY, FThreadSafeCounter* RequestStatus)
+{
+	QUICK_SCOPE_CYCLE_COUNTER(STAT_RHIMETHOD_AsyncReallocateTexture2D_Flush);
+	RHICmdList.ImmediateFlush(EImmediateFlushType::FlushRHIThread);  
+	return GDynamicRHI->RHIAsyncReallocateTexture2D(Texture2D, NewMipCount, NewSizeX, NewSizeY, RequestStatus);
+}
+
+ETextureReallocationStatus FDynamicRHI::FinalizeAsyncReallocateTexture2D_RenderThread(class FRHICommandListImmediate& RHICmdList, FTexture2DRHIParamRef Texture2D, bool bBlockUntilCompleted)
+{
+	if(GRHIThread)
+	{
+		QUICK_SCOPE_CYCLE_COUNTER(STAT_RHIMETHOD_FinalizeAsyncReallocateTexture2D_WaitRHI);
+		RHICmdList.ImmediateFlush(EImmediateFlushType::WaitForRHIThread);
+	}
+	return GDynamicRHI->RHIFinalizeAsyncReallocateTexture2D(Texture2D, bBlockUntilCompleted);
+}
+
+ETextureReallocationStatus FDynamicRHI::CancelAsyncReallocateTexture2D_RenderThread(class FRHICommandListImmediate& RHICmdList, FTexture2DRHIParamRef Texture2D, bool bBlockUntilCompleted)
+{
+	if(GRHIThread)
+	{
+		QUICK_SCOPE_CYCLE_COUNTER(STAT_RHIMETHOD_CancelAsyncReallocateTexture2D_WaitRHI);
+		RHICmdList.ImmediateFlush(EImmediateFlushType::WaitForRHIThread);
+	}
+	return GDynamicRHI->RHICancelAsyncReallocateTexture2D(Texture2D, bBlockUntilCompleted);
+}
+
+FVertexDeclarationRHIRef FDynamicRHI::CreateVertexDeclaration_RenderThread(class FRHICommandListImmediate& RHICmdList, const FVertexDeclarationElementList& Elements)
+{
+	if(GRHIThread)
+	{
+		QUICK_SCOPE_CYCLE_COUNTER(STAT_RHIMETHOD_CreateVertexDeclaration_WaitRHI);
+		RHICmdList.ImmediateFlush(EImmediateFlushType::WaitForRHIThread);
+	}
+	return GDynamicRHI->RHICreateVertexDeclaration(Elements);
+}
+
+
+void FRHICommandListImmediate::UpdateTextureReference(FTextureReferenceRHIParamRef TextureRef, FTextureRHIParamRef NewTexture)
+{
+	if (Bypass() || !GRHIThread || CVarRHICmdFlushUpdateTextureReference.GetValueOnRenderThread() > 0)
+	{
+		{
+			QUICK_SCOPE_CYCLE_COUNTER(STAT_RHIMETHOD_UpdateTextureReference_FlushRHI);
+			ImmediateFlush(EImmediateFlushType::FlushRHIThread);
+		}
+		CMD_CONTEXT(UpdateTextureReference)(TextureRef, NewTexture);
+		return;
+	}
+	new (AllocCommand<FRHICommandUpdateTextureReference>()) FRHICommandUpdateTextureReference(TextureRef, NewTexture);
+	RHIThreadBufferLockFence = RHIThreadFence();
+	if (GetUsedMemory() > 256 * 1024)
+	{
+		// we could be loading a level or something, lets get this stuff going
+		ImmediateFlush(EImmediateFlushType::DispatchToRHIThread); 
+	}
+}
+
+
diff --git a/Engine/Source/Runtime/RHI/Public/DynamicRHI.h b/Engine/Source/Runtime/RHI/Public/DynamicRHI.h
--- a/Engine/Source/Runtime/RHI/Public/DynamicRHI.h
+++ b/Engine/Source/Runtime/RHI/Public/DynamicRHI.h
@@ -13,7 +13,7 @@
 //
 
 /** The interface which is implemented by the dynamically bound RHI. */
-class FDynamicRHI
+class RHI_API FDynamicRHI
 {
 public:
 
@@ -218,9 +218,6 @@ public:
 	// FlushType: Wait RHI Thread
 	virtual FTextureReferenceRHIRef RHICreateTextureReference(FLastRenderTimeContainer* LastRenderTime) = 0;
 
-	// FlushType: Flush RHI Thread
-	virtual void RHIUpdateTextureReference(FTextureReferenceRHIParamRef TextureRef, FTextureRHIParamRef NewTexture) = 0;
-
 	/**
 	* Creates a 2D RHI texture resource
 	* @param SizeX - width of the texture to create
@@ -620,6 +617,19 @@ public:
 	// FlushType: Thread safe
 	virtual class IRHICommandContextContainer* RHIGetCommandContextContainer() = 0;
 
+	///////// Pass through functions that allow RHIs to optimize certain calls.
+
+	virtual FVertexBufferRHIRef CreateVertexBuffer_RenderThread(class FRHICommandListImmediate& RHICmdList, uint32 Size, uint32 InUsage, FRHIResourceCreateInfo& CreateInfo);
+	virtual FShaderResourceViewRHIRef CreateShaderResourceView_RenderThread(class FRHICommandListImmediate& RHICmdList, FVertexBufferRHIParamRef VertexBuffer, uint32 Stride, uint8 Format);
+	virtual void* LockVertexBuffer_RenderThread(class FRHICommandListImmediate& RHICmdList, FVertexBufferRHIParamRef VertexBuffer, uint32 Offset, uint32 SizeRHI, EResourceLockMode LockMode);
+	virtual void UnlockVertexBuffer_RenderThread(class FRHICommandListImmediate& RHICmdList, FVertexBufferRHIParamRef VertexBuffer);
+	virtual FTexture2DRHIRef AsyncReallocateTexture2D_RenderThread(class FRHICommandListImmediate& RHICmdList, FTexture2DRHIParamRef Texture2D, int32 NewMipCount, int32 NewSizeX, int32 NewSizeY, FThreadSafeCounter* RequestStatus);
+	virtual ETextureReallocationStatus FinalizeAsyncReallocateTexture2D_RenderThread(class FRHICommandListImmediate& RHICmdList, FTexture2DRHIParamRef Texture2D, bool bBlockUntilCompleted);
+	virtual ETextureReallocationStatus CancelAsyncReallocateTexture2D_RenderThread(class FRHICommandListImmediate& RHICmdList, FTexture2DRHIParamRef Texture2D, bool bBlockUntilCompleted);
+	virtual FIndexBufferRHIRef CreateIndexBuffer_RenderThread(class FRHICommandListImmediate& RHICmdList, uint32 Stride, uint32 Size, uint32 InUsage, FRHIResourceCreateInfo& CreateInfo);
+	virtual void* LockIndexBuffer_RenderThread(class FRHICommandListImmediate& RHICmdList, FIndexBufferRHIParamRef IndexBuffer, uint32 Offset, uint32 SizeRHI, EResourceLockMode LockMode);
+	virtual void UnlockIndexBuffer_RenderThread(class FRHICommandListImmediate& RHICmdList, FIndexBufferRHIParamRef IndexBuffer);
+	virtual FVertexDeclarationRHIRef CreateVertexDeclaration_RenderThread(class FRHICommandListImmediate& RHICmdList, const FVertexDeclarationElementList& Elements);
 };
 
 /** A global pointer to the dynamically bound RHI implementation. */
@@ -918,6 +928,8 @@ public:
 
 	virtual void RHIPopEvent() = 0;
 
+	virtual void RHIUpdateTextureReference(FTextureReferenceRHIParamRef TextureRef, FTextureRHIParamRef NewTexture) = 0;
+
 	/** Start AsyncCompute command stream recording (no effect if not supported) */
 	virtual void RHIBeginAsyncComputeJob_DrawThread(EAsyncComputePriority Priority) = 0;
 
@@ -926,6 +938,7 @@ public:
 
 	/** Wait for AsyncCompute command stream to finish (no effect if not supported) */
 	virtual void RHIGraphicsWaitOnAsyncComputeJob(uint32 FenceIndex) = 0;
+
 };
 
 
diff --git a/Engine/Source/Runtime/RHI/Public/RHICommandList.h b/Engine/Source/Runtime/RHI/Public/RHICommandList.h
--- a/Engine/Source/Runtime/RHI/Public/RHICommandList.h
+++ b/Engine/Source/Runtime/RHI/Public/RHICommandList.h
@@ -10,9 +10,18 @@
 
 DECLARE_STATS_GROUP(TEXT("RHICmdList"), STATGROUP_RHICMDLIST, STATCAT_Advanced);
 
-#define USE_RHICOMMAND_STATE_REDUCTION (0)
+// set this one to get a stat for each RHI command 
+#define RHI_STATS 0
+
+#if RHI_STATS
+DECLARE_STATS_GROUP(TEXT("RHICommands"),STATGROUP_RHI_COMMANDS, STATCAT_Advanced);
+#define RHISTAT(Method)	DECLARE_SCOPE_CYCLE_COUNTER(TEXT(#Method), STAT_RHI##Method, STATGROUP_RHI_COMMANDS)
+#else
+#define RHISTAT(Method)
+#endif
 
 extern RHI_API TAutoConsoleVariable<int32> CVarRHICmdWidth;
+extern RHI_API TAutoConsoleVariable<int32> CVarRHICmdFlushRenderThreadTasks;
 class FRHICommandListBase;
 
 
@@ -208,11 +217,6 @@ private:
 	IRHICommandContext* Context;
 
 	FMemStackBase MemManager; 
-#if USE_RHICOMMAND_STATE_REDUCTION
-protected:
-	struct FRHICommandListStateCache* StateCache;
-private:
-#endif
 	FGraphEventArray RTTasks;
 
 	void Reset();
@@ -221,50 +225,34 @@ private:
 	friend class FRHICommandListIterator;
 
 public:
-	FDrawUpData DrawUPData; 
-};
-
-#if USE_RHICOMMAND_STATE_REDUCTION
-struct FRHICommandListStateCache
-{
-	// warning, no destructor is ever called for this struct
-
-	FRHICommandListStateCache()
-	{
-		FlushShaderState();
-		FlushSamplerState();
-	}
-	enum 
+	enum class ERenderThreadContext
 	{
-		MAX_SAMPLERS_PER_SHADER_STAGE = 32, // it is okish if this is too small, those buffers simply won't get state reduction
-		MAX_UNIFORM_BUFFERS_PER_SHADER_STAGE = 14, // it is okish if this is too small, those buffers simply won't get state reduction
+		SceneRenderTargets,
+		Num
 	};
-	FSamplerStateRHIParamRef Samplers[SF_NumFrequencies][MAX_SAMPLERS_PER_SHADER_STAGE];
-	FUniformBufferRHIParamRef BoundUniformBuffers[SF_NumFrequencies][MAX_UNIFORM_BUFFERS_PER_SHADER_STAGE];
-
-	FORCEINLINE void FlushSamplerState()
+	void *RenderThreadContexts[(int32)ERenderThreadContext::Num];
+private:
+public:
+	void CopyRenderThreadContexts(const FRHICommandListBase& ParentCommandList)
 	{
-		for (int32 Index = 0; Index < SF_NumFrequencies; Index++)
+		for (int32 Index = 0; ERenderThreadContext(Index) < ERenderThreadContext::Num; Index++)
 		{
-			for (int32 IndexInner = 0; IndexInner < MAX_SAMPLERS_PER_SHADER_STAGE; IndexInner++)
-			{
-				Samplers[Index][IndexInner] = nullptr;
-			}
+			RenderThreadContexts[Index] = ParentCommandList.RenderThreadContexts[Index];
 		}
 	}
-
-	FORCEINLINE void FlushShaderState()
+	void SetRenderThreadContext(void* Context, ERenderThreadContext Slot)
 	{
-		for (int32 Index = 0; Index < SF_NumFrequencies; Index++)
-		{
-			for (int32 IndexInner = 0; IndexInner < MAX_UNIFORM_BUFFERS_PER_SHADER_STAGE; IndexInner++)
-			{
-				BoundUniformBuffers[Index][IndexInner] = nullptr;
-			}
-		}
+		RenderThreadContexts[int32(Slot)] = Context;
+	}
+	FORCEINLINE void* GetRenderThreadContext(ERenderThreadContext Slot)
+	{
+		return RenderThreadContexts[int32(Slot)];
 	}
+
+
+	FDrawUpData DrawUPData; 
 };
-#endif
+
 
 
 template<typename TCmd>
@@ -1148,21 +1136,6 @@ struct FRHICommandEndScene : public FRHICommand<FRHICommandEndScene>
 	RHI_API void Execute(FRHICommandListBase& CmdList);
 };
 
-struct FRHICommandUpdateVertexBuffer : public FRHICommand<FRHICommandUpdateVertexBuffer>
-{
-	FVertexBufferRHIParamRef VertexBuffer;
-	void const* Buffer;
-	int32 BufferSize;
-
-	FORCEINLINE_DEBUGGABLE FRHICommandUpdateVertexBuffer(FVertexBufferRHIParamRef InVertexBuffer, void const* InBuffer, int32 InBufferSize)
-		: VertexBuffer(InVertexBuffer)
-		, Buffer(InBuffer)
-		, BufferSize(InBufferSize)
-	{
-	}
-	RHI_API void Execute(FRHICommandListBase& CmdList);
-};
-
 struct FRHICommandBeginFrame : public FRHICommand<FRHICommandBeginFrame>
 {
 	FORCEINLINE_DEBUGGABLE FRHICommandBeginFrame()
@@ -1234,6 +1207,19 @@ struct FRHICommandDebugBreak : public FRHICommand<FRHICommandDebugBreak>
 	}
 };
 
+struct FRHICommandUpdateTextureReference : public FRHICommand<FRHICommandUpdateTextureReference>
+{
+	FTextureReferenceRHIParamRef TextureRef;
+	FTextureRHIParamRef NewTexture;
+	FORCEINLINE_DEBUGGABLE FRHICommandUpdateTextureReference(FTextureReferenceRHIParamRef InTextureRef, FTextureRHIParamRef InNewTexture)
+		: TextureRef(InTextureRef)
+		, NewTexture(InNewTexture)
+	{
+	}
+	RHI_API void Execute(FRHICommandListBase& CmdList);
+};
+
+
 #define CMD_CONTEXT(Method) GetContext().RHI##Method
 
 struct FRHIBeginAsyncComputeJob_DrawThread : public FRHICommand<FRHIBeginAsyncComputeJob_DrawThread>
@@ -1267,17 +1253,6 @@ public:
 	void* operator new(size_t Size);
 	void operator delete(void *RawMemory);
 
-#if USE_RHICOMMAND_STATE_REDUCTION
-	FORCEINLINE_DEBUGGABLE bool AllocStateCache()
-	{
-		if (!StateCache)
-		{
-			StateCache = new (Alloc<FRHICommandListStateCache>()) FRHICommandListStateCache();
-			return false;
-		}
-		return true;
-	}
-#endif
 	FORCEINLINE_DEBUGGABLE FLocalBoundShaderState BuildLocalBoundShaderState(const FBoundShaderStateInput& BoundShaderStateInput)
 	{
 		return BuildLocalBoundShaderState(
@@ -1331,12 +1306,6 @@ public:
 			CMD_CONTEXT(SetBoundShaderState)(LocalBoundShaderState.BypassBSS);
 			return;
 		}
-#if USE_RHICOMMAND_STATE_REDUCTION
-		if (StateCache)
-		{
-			StateCache->FlushShaderState();
-		}
-#endif
 		new (AllocCommand<FRHICommandSetLocalBoundShaderState>()) FRHICommandSetLocalBoundShaderState(this, LocalBoundShaderState);
 	}
 
@@ -1363,13 +1332,6 @@ public:
 			CMD_CONTEXT(SetShaderUniformBuffer)(Shader, BaseIndex, UniformBuffer.BypassUniform);
 			return;
 		}
-#if USE_RHICOMMAND_STATE_REDUCTION
-		if (StateCache)
-		{
-			//local uniform buffers are rare, so we will just flush
-			StateCache->BoundUniformBuffers[TRHIShaderToEnum<TShaderRHIParamRef>::ShaderFrequency][BaseIndex] = nullptr;
-		}
-#endif
 		new (AllocCommand<FRHICommandSetLocalUniformBuffer<TShaderRHIParamRef> >()) FRHICommandSetLocalUniformBuffer<TShaderRHIParamRef>(this, Shader, BaseIndex, UniformBuffer);
 	}
 
@@ -1381,16 +1343,6 @@ public:
 			CMD_CONTEXT(SetShaderUniformBuffer)(Shader, BaseIndex, UniformBuffer);
 			return;
 		}
-#if USE_RHICOMMAND_STATE_REDUCTION
-		if (BaseIndex < FRHICommandListStateCache::MAX_UNIFORM_BUFFERS_PER_SHADER_STAGE)
-		{
-			if (AllocStateCache() && StateCache->BoundUniformBuffers[TRHIShaderToEnum<TShaderRHI*>::ShaderFrequency][BaseIndex] == UniformBuffer)
-			{
-				return;
-			}
-			StateCache->BoundUniformBuffers[TRHIShaderToEnum<TShaderRHI*>::ShaderFrequency][BaseIndex] = UniformBuffer;
-		}
-#endif
 		new (AllocCommand<FRHICommandSetShaderUniformBuffer<TShaderRHI*> >()) FRHICommandSetShaderUniformBuffer<TShaderRHI*>(Shader, BaseIndex, UniformBuffer);
 	}
 	template <typename TShaderRHI>
@@ -1447,16 +1399,6 @@ public:
 			CMD_CONTEXT(SetShaderSampler)(Shader, SamplerIndex, State);
 			return;
 		}
-#if USE_RHICOMMAND_STATE_REDUCTION
-		if (SamplerIndex < FRHICommandListStateCache::MAX_SAMPLERS_PER_SHADER_STAGE)
-		{
-			if (AllocStateCache() && StateCache->Samplers[TRHIShaderToEnum<TShaderRHIParamRef>::ShaderFrequency][SamplerIndex] == State)
-			{
-				return;
-			}
-			StateCache->Samplers[TRHIShaderToEnum<TShaderRHIParamRef>::ShaderFrequency][SamplerIndex] = State;
-		}
-#endif
 		new (AllocCommand<FRHICommandSetShaderSampler<TShaderRHIParamRef> >()) FRHICommandSetShaderSampler<TShaderRHIParamRef>(Shader, SamplerIndex, State);
 	}
 
@@ -1487,12 +1429,6 @@ public:
 			CMD_CONTEXT(SetBoundShaderState)(BoundShaderState);
 			return;
 		}
-#if USE_RHICOMMAND_STATE_REDUCTION
-		if (StateCache)
-		{
-			StateCache->FlushShaderState();
-		}
-#endif
 		new (AllocCommand<FRHICommandSetBoundShaderState>()) FRHICommandSetBoundShaderState(BoundShaderState);
 	}
 
@@ -1702,12 +1638,6 @@ public:
 			CMD_CONTEXT(SetComputeShader)(ComputeShader);
 			return;
 		}
-#if USE_RHICOMMAND_STATE_REDUCTION
-		if (StateCache)
-		{
-			StateCache->FlushShaderState();
-		}
-#endif
 		new (AllocCommand<FRHICommandSetComputeShader>()) FRHICommandSetComputeShader(ComputeShader);
 	}
 
@@ -1831,19 +1761,6 @@ public:
 		new (AllocCommand<FRHICommandClearMRT>()) FRHICommandClearMRT(bClearColor, NumClearColors, ClearColorArray, bClearDepth, Depth, bClearStencil, Stencil, ExcludeRect);
 	}
 
-	FORCEINLINE_DEBUGGABLE void UpdateVertexBuffer(FVertexBufferRHIParamRef VertexBuffer, void const* Buffer, int32 BufferSize)
-	{
-		if (Bypass())
-		{
-			void* Data = GDynamicRHI->RHILockVertexBuffer(VertexBuffer, 0, BufferSize, RLM_WriteOnly);
-			FMemory::Memcpy(Data, Buffer, BufferSize);
-			GDynamicRHI->RHIUnlockVertexBuffer(VertexBuffer);
-			return;
-		}
-		void* Data = Alloc(BufferSize, 16);
-		FMemory::Memcpy(Data, Buffer, BufferSize);
-		new (AllocCommand<FRHICommandUpdateVertexBuffer>()) FRHICommandUpdateVertexBuffer(VertexBuffer, Data, BufferSize);
-	}
 	FORCEINLINE_DEBUGGABLE void BeginRenderQuery(FRenderQueryRHIParamRef RenderQuery)
 	{
 		if (Bypass())
@@ -1998,18 +1915,14 @@ class RHI_API FRHICommandListImmediate : public FRHICommandList
 	}
 public:
 
-	static bool bFlushedGlobal;	
-	
 	inline void ImmediateFlush(EImmediateFlushType::Type FlushType);
 
 	void SetCurrentStat(TStatId Stat);
 
-	static bool IsFullyFlushed()
-	{
-		return bFlushedGlobal;
-	};
-
-
+	static FGraphEventRef RenderThreadTaskFence();
+	static void WaitOnRenderThreadTaskFence(FGraphEventRef& Fence);
+	static bool AnyRenderThreadTasksOutstanding();
+	FGraphEventRef RHIThreadFence();
 
 	FORCEINLINE void GpuTimeBegin(uint32 Hash,bool bCompute)
 	{
@@ -2047,12 +1960,7 @@ public:
 	
 	FORCEINLINE FVertexDeclarationRHIRef CreateVertexDeclaration(const FVertexDeclarationElementList& Elements)
 	{
-		if(GRHIThread)
-		{
-			QUICK_SCOPE_CYCLE_COUNTER(STAT_RHIMETHOD_CreateVertexDeclaration_WaitRHI);
-			ImmediateFlush(EImmediateFlushType::WaitForRHIThread);
-		}
-		return GDynamicRHI->RHICreateVertexDeclaration(Elements);
+		return GDynamicRHI->CreateVertexDeclaration_RenderThread(*this, Elements);
 	}
 	
 	FORCEINLINE FPixelShaderRHIRef CreatePixelShader(const TArray<uint8>& Code)
@@ -2137,54 +2045,32 @@ public:
 	
 	FORCEINLINE FIndexBufferRHIRef CreateIndexBuffer(uint32 Stride, uint32 Size, uint32 InUsage, FRHIResourceCreateInfo& CreateInfo)
 	{
-		if(GRHIThread)
-		{
-			QUICK_SCOPE_CYCLE_COUNTER(STAT_RHIMETHOD_CreateIndexBuffer_WaitRHI);
-			ImmediateFlush(EImmediateFlushType::WaitForRHIThread);
-		}
-		return GDynamicRHI->RHICreateIndexBuffer(Stride, Size, InUsage, CreateInfo);
+		return GDynamicRHI->CreateIndexBuffer_RenderThread(*this, Stride, Size, InUsage, CreateInfo);
 	}
 	
-	FORCEINLINE void* LockIndexBuffer(FIndexBufferRHIParamRef IndexBuffer, uint32 Offset, uint32 Size, EResourceLockMode LockMode)
+	FORCEINLINE void* LockIndexBuffer(FIndexBufferRHIParamRef IndexBuffer, uint32 Offset, uint32 SizeRHI, EResourceLockMode LockMode)
 	{
-		QUICK_SCOPE_CYCLE_COUNTER(STAT_RHIMETHOD_LockIndexBuffer_Flush);
-		ImmediateFlush(EImmediateFlushType::FlushRHIThread); 
-		TGuardValue<bool> GuardIsFlushedGlobal( bFlushedGlobal, true ); 
-		return GDynamicRHI->RHILockIndexBuffer(IndexBuffer, Offset, Size, LockMode);
+		return GDynamicRHI->LockIndexBuffer_RenderThread(*this, IndexBuffer, Offset, SizeRHI, LockMode);
 	}
 	
 	FORCEINLINE void UnlockIndexBuffer(FIndexBufferRHIParamRef IndexBuffer)
 	{
-		QUICK_SCOPE_CYCLE_COUNTER(STAT_RHIMETHOD_UnlockIndexBuffer_Flush);
-		ImmediateFlush(EImmediateFlushType::FlushRHIThread); 
-		TGuardValue<bool> GuardIsFlushedGlobal( bFlushedGlobal, true );  
-		GDynamicRHI->RHIUnlockIndexBuffer(IndexBuffer);
+		GDynamicRHI->UnlockIndexBuffer_RenderThread(*this, IndexBuffer);
 	}
 	
 	FORCEINLINE FVertexBufferRHIRef CreateVertexBuffer(uint32 Size, uint32 InUsage, FRHIResourceCreateInfo& CreateInfo)
 	{
-		if(GRHIThread)
-		{
-			QUICK_SCOPE_CYCLE_COUNTER(STAT_RHIMETHOD_CreateVertexBuffer_WaitRHI);
-			ImmediateFlush(EImmediateFlushType::WaitForRHIThread);
-		}
-		return GDynamicRHI->RHICreateVertexBuffer(Size, InUsage, CreateInfo);
+		return GDynamicRHI->CreateVertexBuffer_RenderThread(*this, Size, InUsage, CreateInfo);
 	}
 	
 	FORCEINLINE void* LockVertexBuffer(FVertexBufferRHIParamRef VertexBuffer, uint32 Offset, uint32 SizeRHI, EResourceLockMode LockMode)
 	{
-		QUICK_SCOPE_CYCLE_COUNTER(STAT_RHIMETHOD_LockVertexBuffer_Flush);
-		ImmediateFlush(EImmediateFlushType::FlushRHIThread); 
-		TGuardValue<bool> GuardIsFlushedGlobal( bFlushedGlobal, true ); 
-		return GDynamicRHI->RHILockVertexBuffer(VertexBuffer, Offset, SizeRHI, LockMode);
+		return GDynamicRHI->LockVertexBuffer_RenderThread(*this, VertexBuffer, Offset, SizeRHI, LockMode);
 	}
 	
 	FORCEINLINE void UnlockVertexBuffer(FVertexBufferRHIParamRef VertexBuffer)
 	{
-		QUICK_SCOPE_CYCLE_COUNTER(STAT_RHIMETHOD_UnlockVertexBuffer_Flush);
-		ImmediateFlush(EImmediateFlushType::FlushRHIThread); 
-		TGuardValue<bool> GuardIsFlushedGlobal( bFlushedGlobal, true );  
-		GDynamicRHI->RHIUnlockVertexBuffer(VertexBuffer);
+		GDynamicRHI->UnlockVertexBuffer_RenderThread(*this, VertexBuffer);
 	}
 	
 	FORCEINLINE void CopyVertexBuffer(FVertexBufferRHIParamRef SourceBuffer,FVertexBufferRHIParamRef DestBuffer)
@@ -2208,7 +2094,7 @@ public:
 	{
 		QUICK_SCOPE_CYCLE_COUNTER(STAT_RHIMETHOD_LockStructuredBuffer_Flush);
 		ImmediateFlush(EImmediateFlushType::FlushRHIThread); 
-		TGuardValue<bool> GuardIsFlushedGlobal( bFlushedGlobal, true ); 
+		 
 		return GDynamicRHI->RHILockStructuredBuffer(StructuredBuffer, Offset, SizeRHI, LockMode);
 	}
 	
@@ -2216,7 +2102,7 @@ public:
 	{
 		QUICK_SCOPE_CYCLE_COUNTER(STAT_RHIMETHOD_UnlockStructuredBuffer_Flush);
 		ImmediateFlush(EImmediateFlushType::FlushRHIThread); 
-		TGuardValue<bool> GuardIsFlushedGlobal( bFlushedGlobal, true );  
+		  
 		GDynamicRHI->RHIUnlockStructuredBuffer(StructuredBuffer);
 	}
 	
@@ -2262,12 +2148,7 @@ public:
 	
 	FORCEINLINE FShaderResourceViewRHIRef CreateShaderResourceView(FVertexBufferRHIParamRef VertexBuffer, uint32 Stride, uint8 Format)
 	{
-		if(GRHIThread)
-		{
-			QUICK_SCOPE_CYCLE_COUNTER(STAT_RHIMETHOD_CreateShaderResourceView_WaitRHI);
-			ImmediateFlush(EImmediateFlushType::WaitForRHIThread);
-		}
-		return GDynamicRHI->RHICreateShaderResourceView(VertexBuffer, Stride, Format);
+		return GDynamicRHI->CreateShaderResourceView_RenderThread(*this, VertexBuffer, Stride, Format);
 	}
 	
 	FORCEINLINE uint64 CalcTexture2DPlatformSize(uint32 SizeX, uint32 SizeY, uint8 Format, uint32 NumMips, uint32 NumSamples, uint32 Flags, uint32& OutAlign)
@@ -2307,14 +2188,6 @@ public:
 		return GDynamicRHI->RHICreateTextureReference(LastRenderTime);
 	}
 	
-	FORCEINLINE void UpdateTextureReference(FTextureReferenceRHIParamRef TextureRef, FTextureRHIParamRef NewTexture)
-	{
-		QUICK_SCOPE_CYCLE_COUNTER(STAT_RHIMETHOD_UpdateTextureReference_Flush);
-		ImmediateFlush(EImmediateFlushType::FlushRHIThread); 
-		TGuardValue<bool> GuardIsFlushedGlobal( bFlushedGlobal, true );  
-		GDynamicRHI->RHIUpdateTextureReference(TextureRef, NewTexture);
-	}
-	
 	FORCEINLINE FTexture2DRHIRef CreateTexture2D(uint32 SizeX, uint32 SizeY, uint8 Format, uint32 NumMips, uint32 NumSamples, uint32 Flags, FRHIResourceCreateInfo& CreateInfo)
 	{
 		if(GRHIThread)
@@ -2339,7 +2212,7 @@ public:
 	{
 		QUICK_SCOPE_CYCLE_COUNTER(STAT_RHIMETHOD_CopySharedMips_Flush);
 		ImmediateFlush(EImmediateFlushType::FlushRHIThread); 
-		TGuardValue<bool> GuardIsFlushedGlobal( bFlushedGlobal, true ); 
+		 
 		return GDynamicRHI->RHICopySharedMips(DestTexture2D, SrcTexture2D);
 	}
 	
@@ -2401,70 +2274,58 @@ public:
 	
 	FORCEINLINE FTexture2DRHIRef AsyncReallocateTexture2D(FTexture2DRHIParamRef Texture2D, int32 NewMipCount, int32 NewSizeX, int32 NewSizeY, FThreadSafeCounter* RequestStatus)
 	{
-		QUICK_SCOPE_CYCLE_COUNTER(STAT_RHIMETHOD_AsyncReallocateTexture2D_Flush);
-		ImmediateFlush(EImmediateFlushType::FlushRHIThread); TGuardValue<bool> GuardIsFlushedGlobal( bFlushedGlobal, true ); 
-		return GDynamicRHI->RHIAsyncReallocateTexture2D(Texture2D, NewMipCount, NewSizeX, NewSizeY, RequestStatus);
+		return GDynamicRHI->AsyncReallocateTexture2D_RenderThread(*this, Texture2D, NewMipCount, NewSizeX, NewSizeY, RequestStatus);
 	}
 	
 	FORCEINLINE ETextureReallocationStatus FinalizeAsyncReallocateTexture2D(FTexture2DRHIParamRef Texture2D, bool bBlockUntilCompleted)
 	{
-		if(GRHIThread)
-		{
-			QUICK_SCOPE_CYCLE_COUNTER(STAT_RHIMETHOD_FinalizeAsyncReallocateTexture2D_WaitRHI);
-			ImmediateFlush(EImmediateFlushType::WaitForRHIThread);
-		}
-		return GDynamicRHI->RHIFinalizeAsyncReallocateTexture2D(Texture2D, bBlockUntilCompleted);
+		return GDynamicRHI->FinalizeAsyncReallocateTexture2D_RenderThread(*this, Texture2D, bBlockUntilCompleted);
 	}
 	
 	FORCEINLINE ETextureReallocationStatus CancelAsyncReallocateTexture2D(FTexture2DRHIParamRef Texture2D, bool bBlockUntilCompleted)
 	{
-		if(GRHIThread)
-		{
-			QUICK_SCOPE_CYCLE_COUNTER(STAT_RHIMETHOD_CancelAsyncReallocateTexture2D_WaitRHI);
-			ImmediateFlush(EImmediateFlushType::WaitForRHIThread);
-		}
-		return GDynamicRHI->RHICancelAsyncReallocateTexture2D(Texture2D, bBlockUntilCompleted);
+		return GDynamicRHI->CancelAsyncReallocateTexture2D_RenderThread(*this, Texture2D, bBlockUntilCompleted);
 	}
 	
 	FORCEINLINE void* LockTexture2D(FTexture2DRHIParamRef Texture, uint32 MipIndex, EResourceLockMode LockMode, uint32& DestStride, bool bLockWithinMiptail)
 	{
 		QUICK_SCOPE_CYCLE_COUNTER(STAT_RHIMETHOD_LockTexture2D_Flush);
-		ImmediateFlush(EImmediateFlushType::FlushRHIThread); TGuardValue<bool> GuardIsFlushedGlobal( bFlushedGlobal, true ); 
+		ImmediateFlush(EImmediateFlushType::FlushRHIThread);  
 		return GDynamicRHI->RHILockTexture2D(Texture, MipIndex, LockMode, DestStride, bLockWithinMiptail);
 	}
 	
 	FORCEINLINE void UnlockTexture2D(FTexture2DRHIParamRef Texture, uint32 MipIndex, bool bLockWithinMiptail)
 	{
 		QUICK_SCOPE_CYCLE_COUNTER(STAT_RHIMETHOD_UnlockTexture2D_Flush);
-		ImmediateFlush(EImmediateFlushType::FlushRHIThread); TGuardValue<bool> GuardIsFlushedGlobal( bFlushedGlobal, true );  
+		ImmediateFlush(EImmediateFlushType::FlushRHIThread);   
 		GDynamicRHI->RHIUnlockTexture2D(Texture, MipIndex, bLockWithinMiptail);
 	}
 	
 	FORCEINLINE void* LockTexture2DArray(FTexture2DArrayRHIParamRef Texture, uint32 TextureIndex, uint32 MipIndex, EResourceLockMode LockMode, uint32& DestStride, bool bLockWithinMiptail)
 	{
 		QUICK_SCOPE_CYCLE_COUNTER(STAT_RHIMETHOD_LockTexture2DArray_Flush);
-		ImmediateFlush(EImmediateFlushType::FlushRHIThread); TGuardValue<bool> GuardIsFlushedGlobal( bFlushedGlobal, true ); 
+		ImmediateFlush(EImmediateFlushType::FlushRHIThread);  
 		return GDynamicRHI->RHILockTexture2DArray(Texture, TextureIndex, MipIndex, LockMode, DestStride, bLockWithinMiptail);
 	}
 	
 	FORCEINLINE void UnlockTexture2DArray(FTexture2DArrayRHIParamRef Texture, uint32 TextureIndex, uint32 MipIndex, bool bLockWithinMiptail)
 	{
 		QUICK_SCOPE_CYCLE_COUNTER(STAT_RHIMETHOD_UnlockTexture2DArray_Flush);
-		ImmediateFlush(EImmediateFlushType::FlushRHIThread); TGuardValue<bool> GuardIsFlushedGlobal( bFlushedGlobal, true );  
+		ImmediateFlush(EImmediateFlushType::FlushRHIThread);   
 		GDynamicRHI->RHIUnlockTexture2DArray(Texture, TextureIndex, MipIndex, bLockWithinMiptail);
 	}
 	
 	FORCEINLINE void UpdateTexture2D(FTexture2DRHIParamRef Texture, uint32 MipIndex, const struct FUpdateTextureRegion2D& UpdateRegion, uint32 SourcePitch, const uint8* SourceData)
 	{
 		QUICK_SCOPE_CYCLE_COUNTER(STAT_RHIMETHOD_UpdateTexture2D_Flush);
-		ImmediateFlush(EImmediateFlushType::FlushRHIThread); TGuardValue<bool> GuardIsFlushedGlobal( bFlushedGlobal, true );  
+		ImmediateFlush(EImmediateFlushType::FlushRHIThread);   
 		GDynamicRHI->RHIUpdateTexture2D(Texture, MipIndex, UpdateRegion, SourcePitch, SourceData);
 	}
 	
 	FORCEINLINE void UpdateTexture3D(FTexture3DRHIParamRef Texture, uint32 MipIndex, const struct FUpdateTextureRegion3D& UpdateRegion, uint32 SourceRowPitch, uint32 SourceDepthPitch, const uint8* SourceData)
 	{
 		QUICK_SCOPE_CYCLE_COUNTER(STAT_RHIMETHOD_UpdateTexture3D_Flush);
-		ImmediateFlush(EImmediateFlushType::FlushRHIThread); TGuardValue<bool> GuardIsFlushedGlobal( bFlushedGlobal, true );  
+		ImmediateFlush(EImmediateFlushType::FlushRHIThread);   
 		GDynamicRHI->RHIUpdateTexture3D(Texture, MipIndex, UpdateRegion, SourceRowPitch, SourceDepthPitch, SourceData);
 	}
 	
@@ -2491,14 +2352,14 @@ public:
 	FORCEINLINE void* LockTextureCubeFace(FTextureCubeRHIParamRef Texture, uint32 FaceIndex, uint32 ArrayIndex, uint32 MipIndex, EResourceLockMode LockMode, uint32& DestStride, bool bLockWithinMiptail)
 	{
 		QUICK_SCOPE_CYCLE_COUNTER(STAT_RHIMETHOD_LockTextureCubeFace_Flush);
-		ImmediateFlush(EImmediateFlushType::FlushRHIThread); TGuardValue<bool> GuardIsFlushedGlobal( bFlushedGlobal, true ); 
+		ImmediateFlush(EImmediateFlushType::FlushRHIThread);  
 		return GDynamicRHI->RHILockTextureCubeFace(Texture, FaceIndex, ArrayIndex, MipIndex, LockMode, DestStride, bLockWithinMiptail);
 	}
 	
 	FORCEINLINE void UnlockTextureCubeFace(FTextureCubeRHIParamRef Texture, uint32 FaceIndex, uint32 ArrayIndex, uint32 MipIndex, bool bLockWithinMiptail)
 	{
 		QUICK_SCOPE_CYCLE_COUNTER(STAT_RHIMETHOD_UnlockTextureCubeFace_Flush);
-		ImmediateFlush(EImmediateFlushType::FlushRHIThread); TGuardValue<bool> GuardIsFlushedGlobal( bFlushedGlobal, true );  
+		ImmediateFlush(EImmediateFlushType::FlushRHIThread);   
 		GDynamicRHI->RHIUnlockTextureCubeFace(Texture, FaceIndex, ArrayIndex, MipIndex, bLockWithinMiptail);
 	}
 	
@@ -2578,7 +2439,7 @@ public:
 	{
 		QUICK_SCOPE_CYCLE_COUNTER(STAT_RHIMETHOD_AcquireThreadOwnership_Flush);
 		ImmediateFlush(EImmediateFlushType::FlushRHIThread); 
-		TGuardValue<bool> GuardIsFlushedGlobal( bFlushedGlobal, true ); 
+		 
 		return GDynamicRHI->RHIAcquireThreadOwnership();
 	}
 	
@@ -2586,7 +2447,7 @@ public:
 	{
 		QUICK_SCOPE_CYCLE_COUNTER(STAT_RHIMETHOD_ReleaseThreadOwnership_Flush);
 		ImmediateFlush(EImmediateFlushType::FlushRHIThread); 
-		TGuardValue<bool> GuardIsFlushedGlobal( bFlushedGlobal, true ); 
+		 
 		return GDynamicRHI->RHIReleaseThreadOwnership();
 	}
 	
@@ -2594,7 +2455,7 @@ public:
 	{
 		QUICK_SCOPE_CYCLE_COUNTER(STAT_RHIMETHOD_FlushResources_Flush);
 		ImmediateFlush(EImmediateFlushType::FlushRHIThread); 
-		TGuardValue<bool> GuardIsFlushedGlobal( bFlushedGlobal, true ); 
+		 
 		return GDynamicRHI->RHIFlushResources();
 	}
 	
@@ -2700,7 +2561,7 @@ public:
 	{
 		QUICK_SCOPE_CYCLE_COUNTER(STAT_RHIMETHOD_GetNativeDevice_Flush);
 		ImmediateFlush(EImmediateFlushType::FlushRHIThread); 
-		TGuardValue<bool> GuardIsFlushedGlobal( bFlushedGlobal, true ); 
+		 
 		return GDynamicRHI->RHIGetNativeDevice();
 	}
 	
@@ -2713,6 +2574,7 @@ public:
 	{
 		return RHIGetCommandContextContainer();
 	}
+	void UpdateTextureReference(FTextureReferenceRHIParamRef TextureRef, FTextureRHIParamRef NewTexture);
 	
 
 };
@@ -2752,7 +2614,6 @@ public:
 	void ExecuteList(FRHICommandListImmediate& CmdList);
 	void LatchBypass();
 
-	static FGraphEventRef RHIThreadFence();
 	static void WaitOnRHIThreadFence(FGraphEventRef& Fence);
 
 	FORCEINLINE_DEBUGGABLE bool Bypass()
@@ -2773,6 +2634,7 @@ public:
 	}
 	static void CheckNoOutstandingCmdLists();
 	static bool IsRHIThreadActive();
+	static bool IsRHIThreadCompletelyFlushed();
 
 private:
 
@@ -2798,31 +2660,14 @@ FORCEINLINE_DEBUGGABLE FRHICommandListImmediate& FRHICommandListExecutor::GetImm
 struct FScopedCommandListWaitForTasks
 {
 	FRHICommandListImmediate& RHICmdList;
+	bool bWaitForTasks;
 
-	FScopedCommandListWaitForTasks(FRHICommandListImmediate& InRHICmdList = FRHICommandListExecutor::GetImmediateCommandList())
+	FScopedCommandListWaitForTasks(bool InbWaitForTasks, FRHICommandListImmediate& InRHICmdList = FRHICommandListExecutor::GetImmediateCommandList())
 		: RHICmdList(InRHICmdList)
+		, bWaitForTasks(InbWaitForTasks)
 	{
 	}
-	~FScopedCommandListWaitForTasks()
-	{
-		check(IsInRenderingThread());
-		if (GRHIThread)
-		{
-			{
-				QUICK_SCOPE_CYCLE_COUNTER(STAT_FScopedCommandListWaitForTasks_Dispatch);
-				RHICmdList.ImmediateFlush(EImmediateFlushType::DispatchToRHIThread);
-			}
-			{
-				QUICK_SCOPE_CYCLE_COUNTER(STAT_FScopedCommandListWaitForTasks_WaitAsync);
-				RHICmdList.ImmediateFlush(EImmediateFlushType::WaitForOutstandingTasksOnly);
-			}
-		}
-		else
-		{
-			QUICK_SCOPE_CYCLE_COUNTER(STAT_FScopedCommandListWaitForTasks_Flush);
-			RHICmdList.ImmediateFlush(EImmediateFlushType::FlushRHIThread);
-		}
-	}
+	RHI_API ~FScopedCommandListWaitForTasks();
 };
 
 
diff --git a/Engine/Source/Runtime/RHI/Public/RHICommandListCommandExecutes.inl b/Engine/Source/Runtime/RHI/Public/RHICommandListCommandExecutes.inl
--- a/Engine/Source/Runtime/RHI/Public/RHICommandListCommandExecutes.inl
+++ b/Engine/Source/Runtime/RHI/Public/RHICommandListCommandExecutes.inl
@@ -8,16 +8,6 @@
 	#define INTERNAL_DECORATOR(Method) CmdList.GetContext().RHI##Method
 #endif
 
-// set this one to get a stat for each RHI command 
-#define RHI_STATS 0
-
-#if RHI_STATS
-	DECLARE_STATS_GROUP(TEXT("RHICommands"),STATGROUP_RHI_COMMANDS, STATCAT_Advanced);
-	#define RHISTAT(Method)	DECLARE_SCOPE_CYCLE_COUNTER(TEXT(#Method), STAT_RHI##Method, STATGROUP_RHI_COMMANDS)
-#else
-	#define RHISTAT(Method)
-#endif
-
 void FRHICommandSetRasterizerState::Execute(FRHICommandListBase& CmdList)
 {
 	RHISTAT(SetRasterizerState);
@@ -393,6 +383,12 @@ void FRHICommandSubmitCommandsHint::Execute(FRHICommandListBase& CmdList)
 	INTERNAL_DECORATOR(SubmitCommandsHint)();
 }
 
+void FRHICommandUpdateTextureReference::Execute(FRHICommandListBase& CmdList)
+{
+	RHISTAT(UpdateTextureReference);
+	INTERNAL_DECORATOR(UpdateTextureReference)(TextureRef, NewTexture);
+}
+
 void FRHICommandBeginScene::Execute(FRHICommandListBase& CmdList)
 {
 	RHISTAT(BeginScene);
@@ -405,14 +401,6 @@ void FRHICommandEndScene::Execute(FRHICommandListBase& CmdList)
 	INTERNAL_DECORATOR(EndScene)();
 }
 
-void FRHICommandUpdateVertexBuffer::Execute(FRHICommandListBase& CmdList)
-{
-	RHISTAT(UpdateVertexBuffer);
-	void* Data = GDynamicRHI->RHILockVertexBuffer(VertexBuffer, 0, BufferSize, RLM_WriteOnly);
-	FMemory::Memcpy(Data, Buffer, BufferSize);
-	GDynamicRHI->RHIUnlockVertexBuffer(VertexBuffer);
-}
-
 void FRHICommandBeginFrame::Execute(FRHICommandListBase& CmdList)
 {
 	RHISTAT(BeginFrame);
diff --git a/Engine/Source/Runtime/RenderCore/Private/RenderResource.cpp b/Engine/Source/Runtime/RenderCore/Private/RenderResource.cpp
--- a/Engine/Source/Runtime/RenderCore/Private/RenderResource.cpp
+++ b/Engine/Source/Runtime/RenderCore/Private/RenderResource.cpp
@@ -220,8 +220,6 @@ public:
 	enum { ALIGNMENT = (1 << 16) }; // 64KB
 	/** Pointer to the vertex buffer mapped in main memory. */
 	uint8* MappedBuffer;
-	/** Pointer to a buffer to map and copy later. */
-	uint8* DeferredBuffer;
 	/** Size of the vertex buffer in bytes. */
 	uint32 BufferSize;
 	/** Number of bytes currently allocated from the buffer. */
@@ -230,42 +228,18 @@ public:
 	/** Default constructor. */
 	explicit FDynamicVertexBuffer(uint32 InMinBufferSize)
 		: MappedBuffer(NULL)
-		, DeferredBuffer(NULL)
 		, BufferSize(FMath::Max<uint32>(Align(InMinBufferSize,ALIGNMENT),ALIGNMENT))
 		, AllocatedByteCount(0)
 	{
 	}
 
-	/**
-	 * Instead of locking, allocates a deferred memory block
-	 */
-	void AllocDeferred()
-	{
-		check(MappedBuffer == NULL && DeferredBuffer == NULL);
-		check(AllocatedByteCount == 0);
-		check(IsValidRef(VertexBufferRHI));
-		DeferredBuffer = (uint8*)FMemory::Malloc(BufferSize, 16);
-	}
-
-	/**
-	 * Lock, copy, free, unlock the deferred into place
-	 */
-	void CopyDeferred()
-	{
-		Lock();
-		FMemory::Memcpy(MappedBuffer, DeferredBuffer, AllocatedByteCount);
-		FMemory::Free(DeferredBuffer);
-		DeferredBuffer = nullptr;
-		Unlock();
-	}
-
 	/**
 	 * Locks the vertex buffer so it may be written to.
 	 */
 	void Lock()
 	{
 		check(MappedBuffer == NULL);
-		check(AllocatedByteCount == 0 || DeferredBuffer != NULL);
+		check(AllocatedByteCount == 0);
 		check(IsValidRef(VertexBufferRHI));
 		MappedBuffer = (uint8*)RHILockVertexBuffer(VertexBufferRHI, 0, BufferSize, RLM_WriteOnly);
 	}
@@ -275,7 +249,7 @@ public:
 	 */
 	void Unlock()
 	{
-		check(MappedBuffer != NULL && DeferredBuffer == NULL);
+		check(MappedBuffer != NULL);
 		check(IsValidRef(VertexBufferRHI));
 		RHIUnlockVertexBuffer(VertexBufferRHI);
 		MappedBuffer = NULL;
@@ -289,7 +263,6 @@ public:
 		FRHIResourceCreateInfo CreateInfo;
 		VertexBufferRHI = RHICreateVertexBuffer(BufferSize, BUF_Volatile, CreateInfo);
 		MappedBuffer = NULL;
-		DeferredBuffer = nullptr;
 		AllocatedByteCount = 0;
 	}
 
@@ -297,7 +270,6 @@ public:
 	{
 		FVertexBuffer::ReleaseRHI();
 		MappedBuffer = NULL;
-		DeferredBuffer = nullptr;
 		AllocatedByteCount = 0;
 	}
 
@@ -345,21 +317,19 @@ FGlobalDynamicVertexBuffer::~FGlobalDynamicVertexBuffer()
 	Pool = NULL;
 }
 
-FGlobalDynamicVertexBuffer::FAllocation FGlobalDynamicVertexBuffer::Allocate(uint32 SizeInBytes, bool bDeferLock)
+FGlobalDynamicVertexBuffer::FAllocation FGlobalDynamicVertexBuffer::Allocate(uint32 SizeInBytes)
 {
 	FAllocation Allocation;
 
 	FDynamicVertexBuffer* VertexBuffer = Pool->CurrentVertexBuffer;
-	if (VertexBuffer == NULL || VertexBuffer->AllocatedByteCount + SizeInBytes > VertexBuffer->BufferSize
-		|| (bDeferLock && VertexBuffer->MappedBuffer) || (!bDeferLock && VertexBuffer->DeferredBuffer)) // we won't mix these styles of allocation
+	if (VertexBuffer == NULL || VertexBuffer->AllocatedByteCount + SizeInBytes > VertexBuffer->BufferSize)
 	{
 		// Find a buffer in the pool big enough to service the request.
 		VertexBuffer = NULL;
 		for (int32 BufferIndex = 0, NumBuffers = Pool->VertexBuffers.Num(); BufferIndex < NumBuffers; ++BufferIndex)
 		{
 			FDynamicVertexBuffer& VertexBufferToCheck = Pool->VertexBuffers[BufferIndex];
-			bool bMismatch = (bDeferLock && VertexBufferToCheck.MappedBuffer) || (!bDeferLock && VertexBufferToCheck.DeferredBuffer);
-			if (VertexBufferToCheck.AllocatedByteCount + SizeInBytes <= VertexBufferToCheck.BufferSize && !bMismatch)
+			if (VertexBufferToCheck.AllocatedByteCount + SizeInBytes <= VertexBufferToCheck.BufferSize)
 			{
 				VertexBuffer = &VertexBufferToCheck;
 				break;
@@ -374,16 +344,10 @@ FGlobalDynamicVertexBuffer::FAllocation FGlobalDynamicVertexBuffer::Allocate(uin
 		}
 
 		// Lock the buffer if needed.
-		if (!bDeferLock && VertexBuffer->MappedBuffer == NULL)
+		if (VertexBuffer->MappedBuffer == NULL)
 		{
-			check(!VertexBuffer->DeferredBuffer);
 			VertexBuffer->Lock();
 		}
-		else if (bDeferLock && VertexBuffer->DeferredBuffer == NULL)
-		{
-			check(!VertexBuffer->MappedBuffer);
-			VertexBuffer->AllocDeferred();
-		}
 
 		// Remember this buffer, we'll try to allocate out of it in the future.
 		Pool->CurrentVertexBuffer = VertexBuffer;
@@ -391,9 +355,7 @@ FGlobalDynamicVertexBuffer::FAllocation FGlobalDynamicVertexBuffer::Allocate(uin
 
 	check(VertexBuffer != NULL);
 	checkf(VertexBuffer->AllocatedByteCount + SizeInBytes <= VertexBuffer->BufferSize, TEXT("Global vertex buffer allocation failed: BufferSize=%d AllocatedByteCount=%d SizeInBytes=%d"), VertexBuffer->BufferSize, VertexBuffer->AllocatedByteCount, SizeInBytes);
-	uint8* Base = bDeferLock ? VertexBuffer->DeferredBuffer : VertexBuffer->MappedBuffer;
-	check(Base);
-	Allocation.Buffer = Base + VertexBuffer->AllocatedByteCount;
+	Allocation.Buffer = VertexBuffer->MappedBuffer + VertexBuffer->AllocatedByteCount;
 	Allocation.VertexBuffer = VertexBuffer;
 	Allocation.VertexOffset = VertexBuffer->AllocatedByteCount;
 	VertexBuffer->AllocatedByteCount += SizeInBytes;
@@ -410,10 +372,6 @@ void FGlobalDynamicVertexBuffer::Commit()
 		{
 			VertexBuffer.Unlock();
 		}
-		else if (VertexBuffer.DeferredBuffer != NULL)
-		{
-			VertexBuffer.CopyDeferred();
-		}
 	}
 	Pool->CurrentVertexBuffer = NULL;
 }
diff --git a/Engine/Source/Runtime/RenderCore/Public/RenderResource.h b/Engine/Source/Runtime/RenderCore/Public/RenderResource.h
--- a/Engine/Source/Runtime/RenderCore/Public/RenderResource.h
+++ b/Engine/Source/Runtime/RenderCore/Public/RenderResource.h
@@ -508,7 +508,7 @@ public:
 	 * @param SizeInBytes - The amount of memory to allocate in bytes.
 	 * @returns An FAllocation with information regarding the allocated memory.
 	 */
-	FAllocation Allocate(uint32 SizeInBytes, bool bDeferLock = false);
+	FAllocation Allocate(uint32 SizeInBytes);
 
 	/**
 	 * Commits allocated memory to the GPU.
diff --git a/Engine/Source/Runtime/Renderer/Private/AtmosphereRendering.cpp b/Engine/Source/Runtime/Renderer/Private/AtmosphereRendering.cpp
--- a/Engine/Source/Runtime/Renderer/Private/AtmosphereRendering.cpp
+++ b/Engine/Source/Runtime/Renderer/Private/AtmosphereRendering.cpp
@@ -402,8 +402,9 @@ void FDeferredShadingSceneRenderer::RenderAtmosphere(FRHICommandListImmediate& R
 			0, 2, 3
 		};
 
+		FSceneRenderTargets& SceneContext = FSceneRenderTargets::Get(RHICmdList);
 		
-		GSceneRenderTargets.BeginRenderingSceneColor(RHICmdList, ESimpleRenderTargetMode::EUninitializedColorExistingDepth, FExclusiveDepthStencil::DepthRead_StencilWrite);
+		SceneContext.BeginRenderingSceneColor(RHICmdList, ESimpleRenderTargetMode::EUninitializedColorExistingDepth, FExclusiveDepthStencil::DepthRead_StencilWrite);
 		for(int32 ViewIndex = 0;ViewIndex < Views.Num();ViewIndex++)
 		{
 			const FViewInfo& View = Views[ViewIndex];
@@ -435,7 +436,7 @@ void FDeferredShadingSceneRenderer::RenderAtmosphere(FRHICommandListImmediate& R
 		}
 
 		//no need to resolve since we used alpha blending
-		GSceneRenderTargets.FinishRenderingSceneColor(RHICmdList, false);
+		SceneContext.FinishRenderingSceneColor(RHICmdList, false);
 	}
 }
 
