fixed Clang warnings

diff --git a/Engine/Source/Runtime/RHI/Public/RHICommandList.h b/Engine/Source/Runtime/RHI/Public/RHICommandList.h
--- a/Engine/Source/Runtime/RHI/Public/RHICommandList.h
+++ b/Engine/Source/Runtime/RHI/Public/RHICommandList.h
@@ -240,9 +240,9 @@ public:
 			RenderThreadContexts[Index] = ParentCommandList.RenderThreadContexts[Index];
 		}
 	}
-	void SetRenderThreadContext(void* Context, ERenderThreadContext Slot)
+	void SetRenderThreadContext(void* InContext, ERenderThreadContext Slot)
 	{
-		RenderThreadContexts[int32(Slot)] = Context;
+		RenderThreadContexts[int32(Slot)] = InContext;
 	}
 	FORCEINLINE void* GetRenderThreadContext(ERenderThreadContext Slot)
 	{
