Stats - Misc tweak

diff --git a/Engine/Source/Runtime/RHI/Private/RHI.cpp b/Engine/Source/Runtime/RHI/Private/RHI.cpp
--- a/Engine/Source/Runtime/RHI/Private/RHI.cpp
+++ b/Engine/Source/Runtime/RHI/Private/RHI.cpp
@@ -37,7 +37,7 @@ static void DumpRHIMemory(FOutputDevice& OutputDevice)
 	TArray<FStatMessage> Stats;
 	GetPermanentStats(Stats);
 
-	FName NAME_STATGROUP_RHI("STATGROUP_RHI");
+	FName NAME_STATGROUP_RHI(FStatGroup_STATGROUP_RHI::GetGroupName());
 	OutputDevice.Logf(TEXT("RHI resource memory (not tracked by our allocator)"));
 	int64 TotalMemory = 0;
 	for (int32 Index = 0; Index < Stats.Num(); Index++)
