UE4 - Misc parallel rendering tweaks and options. Balance deferred contexts after the fact, cvar can control occlusion latency, framework for GDME tasks. These are all disabled by default.

diff --git a/Engine/Source/Runtime/Core/Private/HAL/ConsoleManager.cpp b/Engine/Source/Runtime/Core/Private/HAL/ConsoleManager.cpp
--- a/Engine/Source/Runtime/Core/Private/HAL/ConsoleManager.cpp
+++ b/Engine/Source/Runtime/Core/Private/HAL/ConsoleManager.cpp
@@ -1782,7 +1782,7 @@ static TAutoConsoleVariable<int32> CVarNumBufferedOcclusionQueries(
 	1,
 	TEXT("Number of frames to buffer occlusion queries (including the current renderthread frame).\n")
 	TEXT("More frames reduces the chance of stalling the CPU waiting for results, but increases out of date query artifacts."),
-	ECVF_ReadOnly);
+	ECVF_RenderThreadSafe);
 
 static TAutoConsoleVariable<int32> CVarDistField(
 	TEXT("r.GenerateMeshDistanceFields"),
diff --git a/Engine/Source/Runtime/Engine/Private/Particles/ParticleSystemRender.cpp b/Engine/Source/Runtime/Engine/Private/Particles/ParticleSystemRender.cpp
--- a/Engine/Source/Runtime/Engine/Private/Particles/ParticleSystemRender.cpp
+++ b/Engine/Source/Runtime/Engine/Private/Particles/ParticleSystemRender.cpp
@@ -1012,20 +1012,105 @@ void FDynamicSpriteEmitterData::GetDynamicMeshElementsEmitter(const FParticleSys
 						)
 					{
 						ParticleOrder = GParticleOrderPool.GetParticleOrderData(ParticleCount);
-						SortSpriteParticles(SourceData->SortMode, SourceData->bUseLocalSpace, SourceData->ActiveParticleCount, 
-							SourceData->ParticleData, SourceData->ParticleStride, SourceData->ParticleIndices,
-							View, Proxy->GetLocalToWorld(), ParticleOrder);
 					}
 				}
 
-				// Fill vertex buffers.
-				if(bInstanced)
+				if (Collector.ShouldUseTasks())
 				{
-					GetVertexAndIndexData(Allocation.Buffer, DynamicParameterAllocation.Buffer, NULL, ParticleOrder, View->ViewMatrices.ViewOrigin, Proxy->GetLocalToWorld());
+					class FHandleParticleSortAndFill
+					{
+						const FDynamicSpriteEmitterData* Target;
+						const FSceneView* View;
+						FParticleOrder* ParticleOrder;
+						FMatrix LocalToWorld;
+						FGlobalDynamicVertexBuffer::FAllocation Allocation;
+						FGlobalDynamicVertexBuffer::FAllocation DynamicParameterAllocation;
+						bool bInstanced;
+					public:
+
+						FHandleParticleSortAndFill(						
+							const FDynamicSpriteEmitterData* InTarget,
+							const FSceneView* InView,
+							FParticleOrder* InParticleOrder,
+							const FMatrix& LocalToWorld,
+							const FGlobalDynamicVertexBuffer::FAllocation& InAllocation,
+							const FGlobalDynamicVertexBuffer::FAllocation& InDynamicParameterAllocation,
+							bool bInInstanced)
+							: Target(InTarget)
+							, View(InView)
+							, ParticleOrder(InParticleOrder)
+							, LocalToWorld(LocalToWorld)
+							, Allocation(InAllocation)
+							, DynamicParameterAllocation(InDynamicParameterAllocation)
+							, bInstanced(bInInstanced)
+						{
+						}
+
+						FORCEINLINE TStatId GetStatId() const
+						{
+							RETURN_QUICK_DECLARE_CYCLE_STAT(FHandleParticleSortAndFill, STATGROUP_TaskGraphTasks);
+						}
+
+						static FORCEINLINE ENamedThreads::Type GetDesiredThread()
+						{
+							return ENamedThreads::AnyThread;
+						}
+
+						static ESubsequentsMode::Type GetSubsequentsMode() { return ESubsequentsMode::TrackSubsequents; }
+
+						void DoTask(ENamedThreads::Type CurrentThread, const FGraphEventRef& MyCompletionGraphEvent)
+						{
+							const FDynamicSpriteEmitterReplayDataBase* SourceData = Target->GetSourceData();
+							if (ParticleOrder)
+							{
+								Target->SortSpriteParticles(SourceData->SortMode, SourceData->bUseLocalSpace, SourceData->ActiveParticleCount, 
+									SourceData->ParticleData, SourceData->ParticleStride, SourceData->ParticleIndices,
+									View, LocalToWorld, ParticleOrder);
+							}
+							// Fill vertex buffers.
+							if(bInstanced)
+							{
+								Target->GetVertexAndIndexData(Allocation.Buffer, DynamicParameterAllocation.Buffer, NULL, ParticleOrder, View->ViewMatrices.ViewOrigin, LocalToWorld);
+							}
+							else
+							{
+								Target->GetVertexAndIndexDataNonInstanced(Allocation.Buffer, DynamicParameterAllocation.Buffer, NULL, ParticleOrder, View->ViewMatrices.ViewOrigin, LocalToWorld);
+							}
+						}
+					};
+
+					check(IsInRenderingThread());
+
+					Collector.AddTask(
+						TGraphTask<FHandleParticleSortAndFill>::CreateTask(nullptr, ENamedThreads::RenderThread).ConstructAndDispatchWhenReady(
+							this,
+							View,
+							ParticleOrder,
+							Proxy->GetLocalToWorld(),
+							Allocation,
+							DynamicParameterAllocation,
+							bInstanced
+						)
+					);
+
 				}
 				else
 				{
-					GetVertexAndIndexDataNonInstanced(Allocation.Buffer, DynamicParameterAllocation.Buffer, NULL, ParticleOrder, View->ViewMatrices.ViewOrigin, Proxy->GetLocalToWorld());
+					if (ParticleOrder)
+					{
+						SortSpriteParticles(SourceData->SortMode, SourceData->bUseLocalSpace, SourceData->ActiveParticleCount, 
+							SourceData->ParticleData, SourceData->ParticleStride, SourceData->ParticleIndices,
+							View, Proxy->GetLocalToWorld(), ParticleOrder);
+					}
+					// Fill vertex buffers.
+					if(bInstanced)
+					{
+						GetVertexAndIndexData(Allocation.Buffer, DynamicParameterAllocation.Buffer, NULL, ParticleOrder, View->ViewMatrices.ViewOrigin, Proxy->GetLocalToWorld());
+					}
+					else
+					{
+						GetVertexAndIndexDataNonInstanced(Allocation.Buffer, DynamicParameterAllocation.Buffer, NULL, ParticleOrder, View->ViewMatrices.ViewOrigin, Proxy->GetLocalToWorld());
+					}
 				}
 
 				// Create per-view uniform buffer.
diff --git a/Engine/Source/Runtime/Engine/Private/SceneManagement.cpp b/Engine/Source/Runtime/Engine/Private/SceneManagement.cpp
--- a/Engine/Source/Runtime/Engine/Private/SceneManagement.cpp
+++ b/Engine/Source/Runtime/Engine/Private/SceneManagement.cpp
@@ -181,6 +181,33 @@ FMeshBatchAndRelevance::FMeshBatchAndRelevance(const FMeshBatch& InMesh, const F
 	bRenderInMainPass = PrimitiveSceneProxy->ShouldRenderInMainPass();
 }
 
+static TAutoConsoleVariable<int32> CVarUseParallelGetDynamicMeshElementsTasks(
+	TEXT("r.UseParallelGetDynamicMeshElementsTasks"),
+	0,
+	TEXT("If > 0, and if FApp::ShouldUseThreadingForPerformance(), then parts of GetDynamicMeshElements will be done in parallel."));
+
+FMeshElementCollector::FMeshElementCollector() :
+	PrimitiveSceneProxy(NULL),
+	FeatureLevel(ERHIFeatureLevel::Num),
+	bUseAsyncTasks(FApp::ShouldUseThreadingForPerformance() && CVarUseParallelGetDynamicMeshElementsTasks.GetValueOnRenderThread() > 0)
+{	
+}
+
+
+void FMeshElementCollector::WaitForTasks()
+{
+	check(IsInRenderingThread());
+	check(!TasksToWaitFor.Num() || bUseAsyncTasks);
+
+	if (TasksToWaitFor.Num())
+	{
+		QUICK_SCOPE_CYCLE_COUNTER(STAT_FMeshElementCollector_WaitForTasks);
+		FTaskGraphInterface::Get().WaitUntilTasksComplete(TasksToWaitFor, ENamedThreads::RenderThread_Local);
+		TasksToWaitFor.Empty();
+	}
+}
+
+
 void FMeshElementCollector::AddMesh(int32 ViewIndex, FMeshBatch& MeshBatch)
 {
 	checkSlow(MeshBatch.GetNumPrimitives() > 0);
diff --git a/Engine/Source/Runtime/Engine/Public/SceneManagement.h b/Engine/Source/Runtime/Engine/Public/SceneManagement.h
--- a/Engine/Source/Runtime/Engine/Public/SceneManagement.h
+++ b/Engine/Source/Runtime/Engine/Public/SceneManagement.h
@@ -1355,15 +1355,25 @@ public:
 		return *OneFrameResource;
 	}
 
+	FORCEINLINE bool ShouldUseTasks() const
+	{
+		return bUseAsyncTasks;
+	}
+
+	FORCEINLINE void AddTask(const FGraphEventRef& Task)
+	{
+		TasksToWaitFor.Add(Task);
+	}
+
+	ENGINE_API void WaitForTasks();
+
 private:
 
-	FMeshElementCollector() :
-		PrimitiveSceneProxy(NULL),
-		FeatureLevel(ERHIFeatureLevel::Num)
-	{}
+	ENGINE_API FMeshElementCollector();
 
 	~FMeshElementCollector()
 	{
+		check(!TasksToWaitFor.Num()); // We should have blocked on this already
 		for (int32 ProxyIndex = 0; ProxyIndex < TemporaryProxies.Num(); ProxyIndex++)
 		{
 			delete TemporaryProxies[ProxyIndex];
@@ -1432,6 +1442,13 @@ private:
 
 	ERHIFeatureLevel::Type FeatureLevel;
 
+	/** This is related to some cvars and FApp stuff and if true means calling code should use async tasks. */
+	const bool bUseAsyncTasks;
+
+	/** Tasks to wait for at the end of gathering dynamic mesh elements. */
+	FGraphEventArray TasksToWaitFor;
+
+
 	friend class FSceneRenderer;
 	friend class FProjectedShadowInfo;
 	friend class FUniformMeshConverter;
diff --git a/Engine/Source/Runtime/RHI/Private/RHICommandList.cpp b/Engine/Source/Runtime/RHI/Private/RHICommandList.cpp
--- a/Engine/Source/Runtime/RHI/Private/RHICommandList.cpp
+++ b/Engine/Source/Runtime/RHI/Private/RHICommandList.cpp
@@ -84,6 +84,21 @@ static TAutoConsoleVariable<int32> CVarRHICmdForceRHIFlush(
 	0,
 	TEXT("Force a flush for every task sent to the RHI thread. For issue diagnosis."));
 
+static TAutoConsoleVariable<int32> CVarRHICmdBalanceTranslatesAfterTasks(
+	TEXT("r.RHICmdBalanceTranslatesAfterTasks"),
+	0,
+	TEXT("Experimental option to balance the parallel translates after the render tasks are complete. This minimizes the number of deferred contexts, but adds latency to starting the translates. r.RHICmdBalanceParallelLists should be off with this option."));
+
+static TAutoConsoleVariable<int32> CVarRHICmdMinCmdlistForParallelTranslate(
+	TEXT("r.RHICmdMinCmdlistForParallelTranslate"),
+	2,
+	TEXT("If there are fewer than this number of parallel translates, they just run on the RHI thread and immediate context. Only relevant if r.RHICmdBalanceTranslatesAfterTasks is on."));
+
+static TAutoConsoleVariable<int32> CVarRHICmdMinCmdlistSizeForParallelTranslate(
+	TEXT("r.RHICmdMinCmdlistSizeForParallelTranslate"),
+	32,
+	TEXT("In kilobytes. Cmdlists are merged into one parallel translate until we have at least this much memory to process. For a given pass, we won't do more translates than we have task threads. Only relevant if r.RHICmdBalanceTranslatesAfterTasks is on."));
+
 
 RHI_API FRHICommandListExecutor GRHICommandList;
 
@@ -702,8 +717,8 @@ DECLARE_CYCLE_STAT(TEXT("Async Chain Execute"), STAT_ChainExecute, STATGROUP_RHI
 struct FRHICommandWaitForAndSubmitSubList : public FRHICommand<FRHICommandWaitForAndSubmitSubList>
 {
 	FGraphEventRef EventToWaitFor;
-	FRHICommandList* RHICmdList;
-	FORCEINLINE_DEBUGGABLE FRHICommandWaitForAndSubmitSubList(FGraphEventRef& InEventToWaitFor, FRHICommandList* InRHICmdList)
+	FRHICommandListBase* RHICmdList;
+	FORCEINLINE_DEBUGGABLE FRHICommandWaitForAndSubmitSubList(FGraphEventRef& InEventToWaitFor, FRHICommandListBase* InRHICmdList)
 		: EventToWaitFor(InEventToWaitFor)
 		, RHICmdList(InRHICmdList)
 	{
@@ -732,6 +747,108 @@ struct FRHICommandWaitForAndSubmitSubList : public FRHICommand<FRHICommandWaitFo
 	}
 };
 
+
+DECLARE_CYCLE_STAT(TEXT("Parallel Setup Translate"), STAT_ParallelSetupTranslate, STATGROUP_RHICMDLIST);
+
+class FParallelTranslateSetupCommandList
+{
+	FRHICommandList* RHICmdList;
+	FRHICommandListBase** RHICmdLists;
+	int32 NumCommandLists;
+	int32 MinSize;
+	int32 MinCount;
+public:
+
+	FParallelTranslateSetupCommandList(FRHICommandList* InRHICmdList, FRHICommandListBase** InRHICmdLists, int32 InNumCommandLists)
+		: RHICmdList(InRHICmdList)
+		, RHICmdLists(InRHICmdLists)
+		, NumCommandLists(InNumCommandLists)
+	{
+		check(RHICmdList && RHICmdLists && NumCommandLists);
+		MinSize = CVarRHICmdMinCmdlistSizeForParallelTranslate.GetValueOnRenderThread() * 1024;
+		MinCount = CVarRHICmdMinCmdlistForParallelTranslate.GetValueOnRenderThread();
+	}
+
+	FORCEINLINE TStatId GetStatId() const
+	{
+		RETURN_QUICK_DECLARE_CYCLE_STAT(FParallelTranslateSetupCommandList, STATGROUP_TaskGraphTasks);
+	}
+
+	ENamedThreads::Type GetDesiredThread()
+	{
+		return ENamedThreads::AnyThread;
+	}
+
+	static ESubsequentsMode::Type GetSubsequentsMode() { return ESubsequentsMode::TrackSubsequents; }
+
+	void DoTask(ENamedThreads::Type CurrentThread, const FGraphEventRef& MyCompletionGraphEvent)
+	{
+		SCOPE_CYCLE_COUNTER(STAT_ParallelSetupTranslate);
+
+		TArray<int32, TInlineAllocator<64> > Sizes;
+		Sizes.Reserve(NumCommandLists);
+		for (int32 Index = 0; Index < NumCommandLists; Index++)
+		{
+			Sizes.Add(RHICmdLists[Index]->GetUsedMemory());
+		}
+
+		int32 EffectiveThreads = 0;
+		int32 Start = 0;
+		// this is pretty silly but we need to know the number of jobs in advance, so we run the merge logic twice
+		while (Start < NumCommandLists)
+		{
+			int32 Last = Start;
+			int32 DrawCnt = Sizes[Start];
+
+			while (Last < NumCommandLists - 1 && DrawCnt + Sizes[Last + 1] <= MinSize)
+			{
+				Last++;
+				DrawCnt += Sizes[Last];
+			}
+			check(Last >= Start);
+			Start = Last + 1;
+			EffectiveThreads++;
+		} 
+
+		if (EffectiveThreads < MinCount)
+		{
+			FGraphEventRef Nothing;
+			for (int32 Index = 0; Index < NumCommandLists; Index++)
+			{
+				FRHICommandListBase* CmdList = RHICmdLists[Index];
+				new (RHICmdList->AllocCommand<FRHICommandWaitForAndSubmitSubList>()) FRHICommandWaitForAndSubmitSubList(Nothing, CmdList);
+			}
+		}
+		else
+		{
+			Start = 0;
+			int32 ThreadIndex = 0;
+
+			while (Start < NumCommandLists)
+			{
+				int32 Last = Start;
+				int32 DrawCnt = Sizes[Start];
+
+				while (Last < NumCommandLists - 1 && DrawCnt + Sizes[Last + 1] <= MinSize)
+				{
+					Last++;
+					DrawCnt += Sizes[Last];
+				}
+				check(Last >= Start);
+
+				IRHICommandContextContainer* ContextContainer =  RHIGetCommandContextContainer();
+				check(ContextContainer);
+
+				FGraphEventRef TranslateCompletionEvent = TGraphTask<FParallelTranslateCommandList>::CreateTask(nullptr, ENamedThreads::RenderThread).ConstructAndDispatchWhenReady(&RHICmdLists[Start], 1 + Last - Start, ContextContainer);
+				MyCompletionGraphEvent->DontCompleteUntil(TranslateCompletionEvent);
+				new (RHICmdList->AllocCommand<FRHICommandWaitForAndSubmitSubListParallel>()) FRHICommandWaitForAndSubmitSubListParallel(TranslateCompletionEvent, ContextContainer, EffectiveThreads, ThreadIndex++);
+				Start = Last + 1;
+			}
+			check(EffectiveThreads == ThreadIndex);
+		}
+	}
+};
+
 void FRHICommandListBase::QueueParallelAsyncCommandListSubmit(FGraphEventRef* AnyThreadCompletionEvents, FRHICommandList** CmdLists, int32* NumDrawsIfKnown, int32 Num, int32 MinDrawsPerTranslate, bool bSpewMerge)
 {
 	check(IsInRenderingThread() && IsImmediate() && Num);
@@ -755,106 +872,139 @@ void FRHICommandListBase::QueueParallelAsyncCommandListSubmit(FGraphEventRef* An
 
 	if (Num && GRHIThread)
 	{
-		IRHICommandContextContainer* ContextContainer = nullptr;
-		if (GRHISupportsParallelRHIExecute && CVarRHICmdUseDeferredContexts.GetValueOnAnyThread() >= 1)
+		if (CVarRHICmdBalanceTranslatesAfterTasks.GetValueOnRenderThread() > 0 && GRHISupportsParallelRHIExecute && CVarRHICmdUseDeferredContexts.GetValueOnAnyThread() > 0)
 		{
-			ContextContainer = RHIGetCommandContextContainer();
+			FGraphEventArray Prereq;
+			FRHICommandListBase** RHICmdLists = (FRHICommandListBase**)Alloc(sizeof(FRHICommandListBase*) * Num, ALIGNOF(FRHICommandListBase*));
+			for (int32 Index = 0; Index < Num; Index++)
+			{
+				FGraphEventRef& AnyThreadCompletionEvent = AnyThreadCompletionEvents[Index];
+				FRHICommandList* CmdList = CmdLists[Index];
+				RHICmdLists[Index] = CmdList;
+				if (AnyThreadCompletionEvent.GetReference())
+				{
+					Prereq.Add(AnyThreadCompletionEvent);
+					WaitOutstandingTasks.Add(AnyThreadCompletionEvent);
+				}
+			}
+			// this is used to ensure that any old buffer locks are completed before we start any parallel translates
+			if (RHIThreadBufferLockFence.GetReference())
+			{
+				Prereq.Add(RHIThreadBufferLockFence);
+			}
+			FRHICommandList* CmdList = new FRHICommandList;
+			CmdList->CopyRenderThreadContexts(*this);
+			FGraphEventRef TranslateSetupCompletionEvent = TGraphTask<FParallelTranslateSetupCommandList>::CreateTask(&Prereq, ENamedThreads::RenderThread).ConstructAndDispatchWhenReady(CmdList, &RHICmdLists[0], Num);
+			QueueCommandListSubmit(CmdList);
+			AllOutstandingTasks.Add(TranslateSetupCompletionEvent);
+			if (GRHIThread)
+			{
+				FRHICommandListExecutor::GetImmediateCommandList().ImmediateFlush(EImmediateFlushType::DispatchToRHIThread); // we don't want stuff after the async cmd list to be bundled with it
+			}
 		}
-		if (ContextContainer)
+		else
 		{
-			bool bMerge = !!CVarRHICmdMergeSmallDeferredContexts.GetValueOnRenderThread();
-
-			int32 EffectiveThreads = 0;
-			int32 Start = 0;
-			// this is pretty silly but we need to know the number of jobs in advance, so we run the merge logic twice
-			while (Start < Num)
+			IRHICommandContextContainer* ContextContainer = nullptr;
+			if (GRHISupportsParallelRHIExecute && CVarRHICmdUseDeferredContexts.GetValueOnAnyThread() > 0)
 			{
-				int32 Last = Start;
-				int32 DrawCnt = NumDrawsIfKnown[Start];
+				ContextContainer = RHIGetCommandContextContainer();
+			}
+			if (ContextContainer)
+			{
+				bool bMerge = !!CVarRHICmdMergeSmallDeferredContexts.GetValueOnRenderThread();
 
-				if (bMerge && DrawCnt >= 0)
+				int32 EffectiveThreads = 0;
+				int32 Start = 0;
+				// this is pretty silly but we need to know the number of jobs in advance, so we run the merge logic twice
+				while (Start < Num)
 				{
-					while (Last < Num - 1 && NumDrawsIfKnown[Last + 1] >= 0 && DrawCnt + NumDrawsIfKnown[Last + 1] <= MinDrawsPerTranslate)
+					int32 Last = Start;
+					int32 DrawCnt = NumDrawsIfKnown[Start];
+
+					if (bMerge && DrawCnt >= 0)
 					{
-						Last++;
-						DrawCnt += NumDrawsIfKnown[Last];
+						while (Last < Num - 1 && NumDrawsIfKnown[Last + 1] >= 0 && DrawCnt + NumDrawsIfKnown[Last + 1] <= MinDrawsPerTranslate)
+						{
+							Last++;
+							DrawCnt += NumDrawsIfKnown[Last];
+						}
 					}
-				}
-				check(Last >= Start);
-				Start = Last + 1;
-				EffectiveThreads++;
-			} 
-			Start = 0;
-			int32 ThreadIndex = 0;
-			   
-			while (Start < Num)       
-			{    
-				int32 Last = Start;
-				int32 DrawCnt = NumDrawsIfKnown[Start];
-				int32 TotalMem = bSpewMerge ? CmdLists[Start]->GetUsedMemory() : 0;   // the memory is only accurate if we are spewing because otherwise it isn't done yet!
-
-				if (bMerge && DrawCnt >= 0)
-				{
-					while (Last < Num - 1 && NumDrawsIfKnown[Last + 1] >= 0 && DrawCnt + NumDrawsIfKnown[Last + 1] <= MinDrawsPerTranslate)
+					check(Last >= Start);
+					Start = Last + 1;
+					EffectiveThreads++;
+				} 
+				Start = 0;
+				int32 ThreadIndex = 0;
+
+				while (Start < Num)       
+				{    
+					int32 Last = Start;
+					int32 DrawCnt = NumDrawsIfKnown[Start];
+					int32 TotalMem = bSpewMerge ? CmdLists[Start]->GetUsedMemory() : 0;   // the memory is only accurate if we are spewing because otherwise it isn't done yet!
+
+					if (bMerge && DrawCnt >= 0)
 					{
-						Last++;
-						DrawCnt += NumDrawsIfKnown[Last];
-						TotalMem += bSpewMerge ? CmdLists[Start]->GetUsedMemory() : 0;   // the memory is only accurate if we are spewing because otherwise it isn't done yet!
+						while (Last < Num - 1 && NumDrawsIfKnown[Last + 1] >= 0 && DrawCnt + NumDrawsIfKnown[Last + 1] <= MinDrawsPerTranslate)
+						{
+							Last++;
+							DrawCnt += NumDrawsIfKnown[Last];
+							TotalMem += bSpewMerge ? CmdLists[Start]->GetUsedMemory() : 0;   // the memory is only accurate if we are spewing because otherwise it isn't done yet!
+						}
 					}
-				}
 
-				check(Last >= Start);
+					check(Last >= Start);
 
-				if (!ContextContainer) 
-				{
-					ContextContainer = RHIGetCommandContextContainer();
-				}
-				check(ContextContainer);
+					if (!ContextContainer) 
+					{
+						ContextContainer = RHIGetCommandContextContainer();
+					}
+					check(ContextContainer);
 
-				FGraphEventArray Prereq;
-				FRHICommandListBase** RHICmdLists = (FRHICommandListBase**)Alloc(sizeof(FRHICommandListBase*) * (1 + Last - Start), ALIGNOF(FRHICommandListBase*));
-				for (int32 Index = Start; Index <= Last; Index++)
-				{
-					FGraphEventRef& AnyThreadCompletionEvent = AnyThreadCompletionEvents[Index];
-					FRHICommandList* CmdList = CmdLists[Index];
-					RHICmdLists[Index - Start] = CmdList;
-					if (AnyThreadCompletionEvent.GetReference())
+					FGraphEventArray Prereq;
+					FRHICommandListBase** RHICmdLists = (FRHICommandListBase**)Alloc(sizeof(FRHICommandListBase*) * (1 + Last - Start), ALIGNOF(FRHICommandListBase*));
+					for (int32 Index = Start; Index <= Last; Index++)
 					{
-						Prereq.Add(AnyThreadCompletionEvent);
-						AllOutstandingTasks.Add(AnyThreadCompletionEvent);
-						WaitOutstandingTasks.Add(AnyThreadCompletionEvent);
+						FGraphEventRef& AnyThreadCompletionEvent = AnyThreadCompletionEvents[Index];
+						FRHICommandList* CmdList = CmdLists[Index];
+						RHICmdLists[Index - Start] = CmdList;
+						if (AnyThreadCompletionEvent.GetReference())
+						{
+							Prereq.Add(AnyThreadCompletionEvent);
+							AllOutstandingTasks.Add(AnyThreadCompletionEvent);
+							WaitOutstandingTasks.Add(AnyThreadCompletionEvent);
+						}
 					}
-				}
-				UE_CLOG(bSpewMerge, LogTemp, Display, TEXT("Parallel translate %d->%d    %dKB mem   %d draws (-1 = unknown)"), Start, Last, FMath::DivideAndRoundUp(TotalMem, 1024), DrawCnt);
+					UE_CLOG(bSpewMerge, LogTemp, Display, TEXT("Parallel translate %d->%d    %dKB mem   %d draws (-1 = unknown)"), Start, Last, FMath::DivideAndRoundUp(TotalMem, 1024), DrawCnt);
 
-				// this is used to ensure that any old buffer locks are completed before we start any parallel translates
-				if (RHIThreadBufferLockFence.GetReference())
-				{
-					Prereq.Add(RHIThreadBufferLockFence);
-				}
+					// this is used to ensure that any old buffer locks are completed before we start any parallel translates
+					if (RHIThreadBufferLockFence.GetReference())
+					{
+						Prereq.Add(RHIThreadBufferLockFence);
+					}
 
-				FGraphEventRef TranslateCompletionEvent = TGraphTask<FParallelTranslateCommandList>::CreateTask(&Prereq, ENamedThreads::RenderThread).ConstructAndDispatchWhenReady(&RHICmdLists[0], 1 + Last - Start, ContextContainer);
+					FGraphEventRef TranslateCompletionEvent = TGraphTask<FParallelTranslateCommandList>::CreateTask(&Prereq, ENamedThreads::RenderThread).ConstructAndDispatchWhenReady(&RHICmdLists[0], 1 + Last - Start, ContextContainer);
 
-				AllOutstandingTasks.Add(TranslateCompletionEvent);
-				new (AllocCommand<FRHICommandWaitForAndSubmitSubListParallel>()) FRHICommandWaitForAndSubmitSubListParallel(TranslateCompletionEvent, ContextContainer, EffectiveThreads, ThreadIndex++);
-				if (GRHIThread)
-				{
-					FRHICommandListExecutor::GetImmediateCommandList().ImmediateFlush(EImmediateFlushType::DispatchToRHIThread); // we don't want stuff after the async cmd list to be bundled with it
-				}
+					AllOutstandingTasks.Add(TranslateCompletionEvent);
+					new (AllocCommand<FRHICommandWaitForAndSubmitSubListParallel>()) FRHICommandWaitForAndSubmitSubListParallel(TranslateCompletionEvent, ContextContainer, EffectiveThreads, ThreadIndex++);
+					if (GRHIThread)
+					{
+						FRHICommandListExecutor::GetImmediateCommandList().ImmediateFlush(EImmediateFlushType::DispatchToRHIThread); // we don't want stuff after the async cmd list to be bundled with it
+					}
 
-				ContextContainer = nullptr;
-				Start = Last + 1;
+					ContextContainer = nullptr;
+					Start = Last + 1;
+				}
+				check(EffectiveThreads == ThreadIndex);
 			}
-			check(EffectiveThreads == ThreadIndex);
+		}
 
 #if !UE_BUILD_SHIPPING
-			if (CVarRHICmdFlushOnQueueParallelSubmit.GetValueOnRenderThread())
-			{
-				FRHICommandListExecutor::GetImmediateCommandList().ImmediateFlush(EImmediateFlushType::FlushRHIThread);
-			}
-#endif
-			return;
+		if (CVarRHICmdFlushOnQueueParallelSubmit.GetValueOnRenderThread())
+		{
+			FRHICommandListExecutor::GetImmediateCommandList().ImmediateFlush(EImmediateFlushType::FlushRHIThread);
 		}
+#endif
+		return;
 	}
 	for (int32 Index = 0; Index < Num; Index++)
 	{
@@ -948,7 +1098,7 @@ void FRHICommandListBase::QueueRenderThreadCommandListSubmit(FGraphEventRef& Ren
 	if (RenderThreadCompletionEvent.GetReference())
 	{
 		check(!IsInRenderingThread() && !IsImmediate());
-	RTTasks.Add(RenderThreadCompletionEvent);
+		RTTasks.Add(RenderThreadCompletionEvent);
 	}
 	new (AllocCommand<FRHICommandWaitForAndSubmitRTSubList>()) FRHICommandWaitForAndSubmitRTSubList(RenderThreadCompletionEvent, CmdList);
 }
diff --git a/Engine/Source/Runtime/Renderer/Private/ConvertToUniformMesh.cpp b/Engine/Source/Runtime/Renderer/Private/ConvertToUniformMesh.cpp
--- a/Engine/Source/Runtime/Renderer/Private/ConvertToUniformMesh.cpp
+++ b/Engine/Source/Runtime/Renderer/Private/ConvertToUniformMesh.cpp
@@ -328,6 +328,8 @@ int32 FUniformMeshConverter::Convert(
 	PrimitiveSceneInfo->Proxy->GetDynamicMeshElements(ViewsArray, *(View.Family), 0x1, Renderer.MeshCollector);
 	View.bRenderFirstInstanceOnly = false;
 
+	Renderer.MeshCollector.WaitForTasks();
+
 	int32 NumTriangles = 0;
 
 	FPrimitiveViewRelevance ViewRelevance = PrimitiveSceneInfo->Proxy->GetViewRelevance(&View);
diff --git a/Engine/Source/Runtime/Renderer/Private/DeferredShadingRenderer.cpp b/Engine/Source/Runtime/Renderer/Private/DeferredShadingRenderer.cpp
--- a/Engine/Source/Runtime/Renderer/Private/DeferredShadingRenderer.cpp
+++ b/Engine/Source/Runtime/Renderer/Private/DeferredShadingRenderer.cpp
@@ -604,12 +604,17 @@ DECLARE_CYCLE_STAT(TEXT("AfterBasePass"), STAT_CLM_AfterBasePass, STATGROUP_Comm
 DECLARE_CYCLE_STAT(TEXT("Lighting"), STAT_CLM_Lighting, STATGROUP_CommandListMarkers);
 DECLARE_CYCLE_STAT(TEXT("AfterLighting"), STAT_CLM_AfterLighting, STATGROUP_CommandListMarkers);
 DECLARE_CYCLE_STAT(TEXT("Translucency"), STAT_CLM_Translucency, STATGROUP_CommandListMarkers);
+DECLARE_CYCLE_STAT(TEXT("RenderDistortion"), STAT_CLM_RenderDistortion, STATGROUP_CommandListMarkers);
 DECLARE_CYCLE_STAT(TEXT("AfterTranslucency"), STAT_CLM_AfterTranslucency, STATGROUP_CommandListMarkers);
+DECLARE_CYCLE_STAT(TEXT("RenderDistanceFieldLighting"), STAT_CLM_RenderDistanceFieldLighting, STATGROUP_CommandListMarkers);
+DECLARE_CYCLE_STAT(TEXT("LightShaftBloom"), STAT_CLM_LightShaftBloom, STATGROUP_CommandListMarkers);
+DECLARE_CYCLE_STAT(TEXT("PostProcessing"), STAT_CLM_PostProcessing, STATGROUP_CommandListMarkers);
 DECLARE_CYCLE_STAT(TEXT("Velocity"), STAT_CLM_Velocity, STATGROUP_CommandListMarkers);
 DECLARE_CYCLE_STAT(TEXT("AfterVelocity"), STAT_CLM_AfterVelocity, STATGROUP_CommandListMarkers);
+DECLARE_CYCLE_STAT(TEXT("RenderFinish"), STAT_CLM_RenderFinish, STATGROUP_CommandListMarkers);
 DECLARE_CYCLE_STAT(TEXT("AfterFrame"), STAT_CLM_AfterFrame, STATGROUP_CommandListMarkers);
 
-FGraphEventRef FDeferredShadingSceneRenderer::OcclusionSubmittedFence;
+FGraphEventRef FDeferredShadingSceneRenderer::OcclusionSubmittedFence[2];
 
 /**
  * Returns true if the depth Prepass needs to run
@@ -682,7 +687,12 @@ void FDeferredShadingSceneRenderer::RenderOcclusion(FRHICommandListImmediate& RH
 		if (bRenderQueries && GRHIThread)
 		{
 			QUICK_SCOPE_CYCLE_COUNTER(STAT_OcclusionSubmittedFence_Dispatch);
-			OcclusionSubmittedFence = RHICmdList.RHIThreadFence();
+			int32 NumFrames = FOcclusionQueryHelpers::GetNumBufferedFrames();
+			for (int32 Dest = 1; Dest < NumFrames; Dest++)
+			{
+				OcclusionSubmittedFence[Dest] = OcclusionSubmittedFence[Dest - 1];
+			}
+			OcclusionSubmittedFence[0] = RHICmdList.RHIThreadFence();
 		}
 	}
 }
@@ -701,8 +711,9 @@ void FDeferredShadingSceneRenderer::Render(FRHICommandListImmediate& RHICmdList)
 	if (GRHIThread)
 	{
 		QUICK_SCOPE_CYCLE_COUNTER(STAT_OcclusionSubmittedFence_Wait);
-		FRHICommandListExecutor::WaitOnRHIThreadFence(OcclusionSubmittedFence);
-		OcclusionSubmittedFence = nullptr;
+		int32 BlockFrame = FOcclusionQueryHelpers::GetNumBufferedFrames() - 1;
+		FRHICommandListExecutor::WaitOnRHIThreadFence(OcclusionSubmittedFence[BlockFrame]);
+		OcclusionSubmittedFence[BlockFrame] = nullptr;
 	}
 
 	if(!ViewFamily.EngineShowFlags.Rendering)
@@ -1164,7 +1175,6 @@ void FDeferredShadingSceneRenderer::Render(FRHICommandListImmediate& RHICmdList)
 
 		RHICmdList.SetCurrentStat(GET_STATID(STAT_CLM_Translucency));
 		RenderTranslucency(RHICmdList);
-		RHICmdList.SetCurrentStat(GET_STATID(STAT_CLM_AfterTranslucency));
 		ServiceLocalQueue();
 
 		if (ViewFamily.EngineShowFlags.Refraction)
@@ -1172,14 +1182,17 @@ void FDeferredShadingSceneRenderer::Render(FRHICommandListImmediate& RHICmdList)
 			// To apply refraction effect by distorting the scene color.
 			// After non separate translucency as that is considered at scene depth anyway
 			// It allows skybox translucency (set to non separate translucency) to be refracted.
+			RHICmdList.SetCurrentStat(GET_STATID(STAT_CLM_RenderDistortion));
 			RenderDistortion(RHICmdList);
 			ServiceLocalQueue();
 		}
+		RHICmdList.SetCurrentStat(GET_STATID(STAT_CLM_AfterTranslucency));
 	}
 
 	if (ViewFamily.EngineShowFlags.LightShafts)
 	{
 		QUICK_SCOPE_CYCLE_COUNTER(STAT_FDeferredShadingSceneRenderer_RenderLightShaftBloom);
+		RHICmdList.SetCurrentStat(GET_STATID(STAT_CLM_LightShaftBloom));
 		RenderLightShaftBloom(RHICmdList);
 		ServiceLocalQueue();
 	}
@@ -1189,6 +1202,7 @@ void FDeferredShadingSceneRenderer::Render(FRHICommandListImmediate& RHICmdList)
 		// Use the skylight's max distance if there is one, to be consistent with DFAO shadowing on the skylight
 		const float OcclusionMaxDistance = Scene->SkyLight && !Scene->SkyLight->bWantsStaticShadowing ? Scene->SkyLight->OcclusionMaxDistance : GDefaultDFAOMaxOcclusionDistance;
 		TRefCountPtr<IPooledRenderTarget> DummyOutput;
+		RHICmdList.SetCurrentStat(GET_STATID(STAT_CLM_RenderDistanceFieldLighting));
 		RenderDistanceFieldLighting(RHICmdList, FDistanceFieldAOParameters(OcclusionMaxDistance), VelocityRT, DummyOutput, DummyOutput, ViewFamily.EngineShowFlags.VisualizeDistanceFieldAO, ViewFamily.EngineShowFlags.VisualizeDistanceFieldGI);
 		ServiceLocalQueue();
 	}
@@ -1210,6 +1224,7 @@ void FDeferredShadingSceneRenderer::Render(FRHICommandListImmediate& RHICmdList)
 		SCOPED_DRAW_EVENT(RHICmdList, PostProcessing);
 		SCOPE_CYCLE_COUNTER(STAT_FinishRenderViewTargetTime);
 
+		RHICmdList.SetCurrentStat(GET_STATID(STAT_CLM_PostProcessing));
 		for(int32 ViewIndex = 0; ViewIndex < Views.Num(); ViewIndex++)
 		{
 			SCOPED_CONDITIONAL_DRAW_EVENTF(RHICmdList, EventView, Views.Num() > 1, TEXT("View%d"), ViewIndex);
@@ -1238,6 +1253,7 @@ void FDeferredShadingSceneRenderer::Render(FRHICommandListImmediate& RHICmdList)
 
 	{
 		QUICK_SCOPE_CYCLE_COUNTER(STAT_FDeferredShadingSceneRenderer_RenderFinish);
+		RHICmdList.SetCurrentStat(GET_STATID(STAT_CLM_RenderFinish));
 		RenderFinish(RHICmdList);
 		RHICmdList.SetCurrentStat(GET_STATID(STAT_CLM_AfterFrame));
 	}
diff --git a/Engine/Source/Runtime/Renderer/Private/DeferredShadingRenderer.h b/Engine/Source/Runtime/Renderer/Private/DeferredShadingRenderer.h
--- a/Engine/Source/Runtime/Renderer/Private/DeferredShadingRenderer.h
+++ b/Engine/Source/Runtime/Renderer/Private/DeferredShadingRenderer.h
@@ -119,7 +119,8 @@ public:
 
 private:
 
-	static FGraphEventRef OcclusionSubmittedFence;
+	// fences to make sure the rhi thread has digested the occlusion query renders before we attempt to read them back async
+	static FGraphEventRef OcclusionSubmittedFence[FOcclusionQueryHelpers::MaxBufferedOcclusionFrames];
 
 	/** Creates a per object projected shadow for the given interaction. */
 	void CreatePerObjectProjectedShadow(
diff --git a/Engine/Source/Runtime/Renderer/Private/RendererScene.cpp b/Engine/Source/Runtime/Renderer/Private/RendererScene.cpp
--- a/Engine/Source/Runtime/Renderer/Private/RendererScene.cpp
+++ b/Engine/Source/Runtime/Renderer/Private/RendererScene.cpp
@@ -100,11 +100,8 @@ FSceneViewState::FSceneViewState()
 	bIntializedGlobalDistanceFieldOrigins = false;
 	GlobalDistanceFieldUpdateIndex = 0;
 
-#if BUFFERED_OCCLUSION_QUERIES
-	NumBufferedFrames = FOcclusionQueryHelpers::GetNumBufferedFrames();
-	ShadowOcclusionQueryMaps.Empty(NumBufferedFrames);
-	ShadowOcclusionQueryMaps.AddZeroed(NumBufferedFrames);	
-#endif
+	ShadowOcclusionQueryMaps.Empty(FOcclusionQueryHelpers::MaxBufferedOcclusionFrames);
+	ShadowOcclusionQueryMaps.AddZeroed(FOcclusionQueryHelpers::MaxBufferedOcclusionFrames);	
 }
 
 void DestroyRenderResource(FRenderResource* RenderResource)
diff --git a/Engine/Source/Runtime/Renderer/Private/SceneOcclusion.cpp b/Engine/Source/Runtime/Renderer/Private/SceneOcclusion.cpp
--- a/Engine/Source/Runtime/Renderer/Private/SceneOcclusion.cpp
+++ b/Engine/Source/Runtime/Renderer/Private/SceneOcclusion.cpp
@@ -46,6 +46,19 @@ static FAutoConsoleVariableRef CVarShowRelevantPrecomputedVisibilityCells(
 /** Random table for occlusion **/
 FOcclusionRandomStream GOcclusionRandomStream;
 
+int32 FOcclusionQueryHelpers::GetNumBufferedFrames()
+{
+#if WITH_SLI
+	// If we're running with SLI, assume throughput is more important than latency, and buffer an extra frame
+	check(GNumActiveGPUsForRendering <= (int32)FOcclusionQueryHelpers::MaxBufferedOcclusionFrames);
+	return FMath::Min<int32>(GNumActiveGPUsForRendering, (int32)FOcclusionQueryHelpers::MaxBufferedOcclusionFrames);
+#else
+	static const auto NumBufferedQueriesVar = IConsoleManager::Get().FindTConsoleVariableDataInt(TEXT("r.NumBufferedOcclusionQueries"));
+	return FMath::Clamp<int32>(NumBufferedQueriesVar->GetValueOnAnyThread(), 1, (int32)FOcclusionQueryHelpers::MaxBufferedOcclusionFrames);
+#endif
+}
+
+
 // default, non-instanced shader implementation
 IMPLEMENT_SHADER_TYPE(,FOcclusionQueryVS,TEXT("OcclusionQueryVertexShader"),TEXT("Main"),SF_Vertex);
 
@@ -203,6 +216,8 @@ void FSceneViewState::TrimOcclusionHistory(FRHICommandListImmediate& RHICmdList,
 	// Only trim every few frames, since stale entries won't cause problems
 	if (FrameNumber % 6 == 0)
 	{
+		int32 NumBufferedFrames = FOcclusionQueryHelpers::GetNumBufferedFrames();
+
 		for(TSet<FPrimitiveOcclusionHistory,FPrimitiveOcclusionHistoryKeyFuncs>::TIterator PrimitiveIt(PrimitiveOcclusionHistorySet);
 			PrimitiveIt;
 			++PrimitiveIt
@@ -211,7 +226,7 @@ void FSceneViewState::TrimOcclusionHistory(FRHICommandListImmediate& RHICmdList,
 			// If the primitive has an old pending occlusion query, release it.
 			if(PrimitiveIt->LastConsideredTime < MinQueryTime)
 			{
-				PrimitiveIt->ReleaseQueries(RHICmdList, OcclusionQueryPool);
+				PrimitiveIt->ReleaseQueries(RHICmdList, OcclusionQueryPool, NumBufferedFrames);
 			}
 
 			// If the primitive hasn't been considered for visibility recently, remove its history from the set.
@@ -223,16 +238,14 @@ void FSceneViewState::TrimOcclusionHistory(FRHICommandListImmediate& RHICmdList,
 	}
 }
 
-bool FSceneViewState::IsShadowOccluded(FRHICommandListImmediate& RHICmdList, FPrimitiveComponentId PrimitiveId, const ULightComponent* Light, int32 InShadowSplitIndex, bool bTranslucentShadow) const
+bool FSceneViewState::IsShadowOccluded(FRHICommandListImmediate& RHICmdList, FPrimitiveComponentId PrimitiveId, const ULightComponent* Light, int32 InShadowSplitIndex, bool bTranslucentShadow, int32 NumBufferedFrames) const
 {
 	// Find the shadow's occlusion query from the previous frame.
 	const FSceneViewState::FProjectedShadowKey Key(PrimitiveId, Light, InShadowSplitIndex, bTranslucentShadow);
 
-#if BUFFERED_OCCLUSION_QUERIES
 	// Get the oldest occlusion query	
 	const uint32 QueryIndex = FOcclusionQueryHelpers::GetQueryLookupIndex(PendingPrevFrameNumber, NumBufferedFrames);
 	const FSceneViewState::ShadowKeyOcclusionQueryMap& ShadowOcclusionQueryMap = ShadowOcclusionQueryMaps[QueryIndex];	
-#endif
 
 
 	const FRenderQueryRHIRef* Query = ShadowOcclusionQueryMap.Find(Key);
@@ -274,15 +287,11 @@ void FSceneViewState::Destroy()
 
 SIZE_T FSceneViewState::GetSizeBytes() const
 {
-#if BUFFERED_OCCLUSION_QUERIES
 	uint32 ShadowOcclusionQuerySize = ShadowOcclusionQueryMaps.GetAllocatedSize();
 	for (int32 i = 0; i < ShadowOcclusionQueryMaps.Num(); ++i)
 	{
 		ShadowOcclusionQuerySize += ShadowOcclusionQueryMaps[i].GetAllocatedSize();
 	}
-#else
-	uint32 ShadowOcclusionQuerySize = ShadowOcclusionQueryMap.GetAllocatedSize();
-#endif
 
 	return sizeof(*this) 
 		+ ShadowOcclusionQuerySize
@@ -401,16 +410,12 @@ FRenderQueryRHIParamRef FOcclusionQueryBatcher::BatchPrimitive(const FVector& Bo
 	return CurrentBatchOcclusionQuery->Query;
 }
 
-static void IssueProjectedShadowOcclusionQuery(FRHICommandListImmediate& RHICmdList, FViewInfo& View, const FProjectedShadowInfo& ProjectedShadowInfo, FOcclusionQueryVS* VertexShader)
+static void IssueProjectedShadowOcclusionQuery(FRHICommandListImmediate& RHICmdList, FViewInfo& View, const FProjectedShadowInfo& ProjectedShadowInfo, FOcclusionQueryVS* VertexShader, int32 NumBufferedFrames)
 {	
 	FSceneViewState* ViewState = (FSceneViewState*)View.State;
 
-#if BUFFERED_OCCLUSION_QUERIES		
-	const uint32 QueryIndex = FOcclusionQueryHelpers::GetQueryIssueIndex(ViewState->PendingPrevFrameNumber, ViewState->NumBufferedFrames);
+	const uint32 QueryIndex = FOcclusionQueryHelpers::GetQueryIssueIndex(ViewState->PendingPrevFrameNumber, NumBufferedFrames);
 	FSceneViewState::ShadowKeyOcclusionQueryMap& ShadowOcclusionQueryMap = ViewState->ShadowOcclusionQueryMaps[QueryIndex];
-#else
-	FSceneViewState::ShadowKeyOcclusionQueryMap& ShadowOcclusionQueryMap = ViewState->ShadowOcclusionQueryMap;
-#endif
 
 	// The shadow transforms and view transforms are relative to different origins, so the world coordinates need to
 	// be translated.
@@ -1072,6 +1077,7 @@ void BuildHZB( FRHICommandListImmediate& RHICmdList, FViewInfo& View )
 void FDeferredShadingSceneRenderer::BeginOcclusionTests(FRHICommandListImmediate& RHICmdList, bool bRenderQueries, bool bRenderHZB)
 {
 	SCOPE_CYCLE_COUNTER(STAT_BeginOcclusionTestsTime);
+	int32 NumBufferedFrames = FOcclusionQueryHelpers::GetNumBufferedFrames();
 	FSceneRenderTargets& SceneContext = FSceneRenderTargets::Get(RHICmdList);
 	const bool bUseDownsampledDepth = IsValidRef(SceneContext.GetSmallDepthSurface()) && SceneContext.UseDownsizedOcclusionQueries();	
 
@@ -1088,6 +1094,7 @@ void FDeferredShadingSceneRenderer::BeginOcclusionTests(FRHICommandListImmediate
 	{
 		RHICmdList.BeginOcclusionQueryBatch();
 
+
 		// Perform occlusion queries for each view
 		for (int32 ViewIndex = 0; ViewIndex < Views.Num(); ViewIndex++)
 		{
@@ -1125,12 +1132,8 @@ void FDeferredShadingSceneRenderer::BeginOcclusionTests(FRHICommandListImmediate
 				VertexShader->SetParameters(RHICmdList, View);
 
 				// Issue this frame's occlusion queries (occlusion queries from last frame may still be in flight)
-#if BUFFERED_OCCLUSION_QUERIES				
-				const uint32 QueryIndex = FOcclusionQueryHelpers::GetQueryIssueIndex(ViewState->PendingPrevFrameNumber, ViewState->NumBufferedFrames);
+				const uint32 QueryIndex = FOcclusionQueryHelpers::GetQueryIssueIndex(ViewState->PendingPrevFrameNumber, NumBufferedFrames);
 				FSceneViewState::ShadowKeyOcclusionQueryMap& ShadowOcclusionQueryMap = ViewState->ShadowOcclusionQueryMaps[QueryIndex];
-#else
-				FSceneViewState::ShadowKeyOcclusionQueryMap& ShadowOcclusionQueryMap = ViewState->ShadowOcclusionQueryMap;
-#endif			
 
 				// Clear primitives which haven't been visible recently out of the occlusion history, and reset old pending occlusion queries.
 				ViewState->TrimOcclusionHistory(RHICmdList, ViewFamily.CurrentRealTime - GEngine->PrimitiveProbablyVisibleTime, ViewFamily.CurrentRealTime, ViewState->OcclusionFrameCounter);
@@ -1189,7 +1192,7 @@ void FDeferredShadingSceneRenderer::BeginOcclusionTests(FRHICommandListImmediate
 							// Don't query if any subjects are visible because the shadow frustum will be definitely unoccluded
 							else if (!ProjectedShadowInfo.IsWholeSceneDirectionalShadow() && !ProjectedShadowInfo.bPreShadow && !ProjectedShadowInfo.SubjectsVisible(View))
 							{
-								IssueProjectedShadowOcclusionQuery(RHICmdList, View, ProjectedShadowInfo, *VertexShader);
+								IssueProjectedShadowOcclusionQuery(RHICmdList, View, ProjectedShadowInfo, *VertexShader, NumBufferedFrames);
 							}
 						}
 
@@ -1197,7 +1200,7 @@ void FDeferredShadingSceneRenderer::BeginOcclusionTests(FRHICommandListImmediate
 						for (int32 ShadowIndex = 0; ShadowIndex < VisibleLightInfo.OccludedPerObjectShadows.Num(); ShadowIndex++)
 						{
 							const FProjectedShadowInfo& ProjectedShadowInfo = *VisibleLightInfo.OccludedPerObjectShadows[ShadowIndex];
-							IssueProjectedShadowOcclusionQuery(RHICmdList, View, ProjectedShadowInfo, *VertexShader);
+							IssueProjectedShadowOcclusionQuery(RHICmdList, View, ProjectedShadowInfo, *VertexShader, NumBufferedFrames);
 						}
 					}
 				}
diff --git a/Engine/Source/Runtime/Renderer/Private/ScenePrivate.h b/Engine/Source/Runtime/Renderer/Private/ScenePrivate.h
--- a/Engine/Source/Runtime/Renderer/Private/ScenePrivate.h
+++ b/Engine/Source/Runtime/Renderer/Private/ScenePrivate.h
@@ -40,6 +40,41 @@ bool IsMobileHDR();
 /** True if the mobile renderer is emulating HDR in a 32bpp render target. */
 bool IsMobileHDR32bpp();
 
+class FOcclusionQueryHelpers
+{
+public:
+
+	enum
+	{
+		MaxBufferedOcclusionFrames = 2
+	};
+
+	// get the system-wide number of frames of buffered occlusion queries.
+	static int32 GetNumBufferedFrames();
+
+	// get the index of the oldest query based on the current frame and number of buffered frames.
+	static uint32 GetQueryLookupIndex(int32 CurrentFrame, int32 NumBufferedFrames)
+	{
+		// queries are currently always requested earlier in the frame than they are issued.
+		// thus we can always overwrite the oldest query with the current one as we never need them
+		// to coexist.  This saves us a buffer entry.
+		const uint32 QueryIndex = CurrentFrame % NumBufferedFrames;
+		return QueryIndex;
+	}
+
+	// get the index of the query to overwrite for new queries.
+	static uint32 GetQueryIssueIndex(int32 CurrentFrame, int32 NumBufferedFrames)
+	{
+		// queries are currently always requested earlier in the frame than they are issued.
+		// thus we can always overwrite the oldest query with the current one as we never need them
+		// to coexist.  This saves us a buffer entry.
+		const uint32 QueryIndex = CurrentFrame % NumBufferedFrames;
+		return QueryIndex;
+	}
+};
+
+
+
 // Dependencies.
 #include "StaticBoundShaderState.h"
 #include "BatchedElements.h"
@@ -71,55 +106,9 @@ bool IsMobileHDR32bpp();
 #include "AtmosphereRendering.h"
 #include "GlobalDistanceFieldParameters.h"
 
-#if WITH_SLI || PLATFORM_SHOULD_BUFFER_QUERIES
-#define BUFFERED_OCCLUSION_QUERIES 1
-#endif
-
 /** Factor by which to grow occlusion tests **/
 #define OCCLUSION_SLOP (1.0f)
 
-class FOcclusionQueryHelpers
-{
-public:
-
-	// get the system-wide number of frames of buffered occlusion queries.
-	static int32 GetNumBufferedFrames()
-	{
-		int32 NumBufferedFrames = 1;
-
-#if BUFFERED_OCCLUSION_QUERIES		
-#	if WITH_SLI
-		// If we're running with SLI, assume throughput is more important than latency, and buffer an extra frame
-		NumBufferedFrames = GNumActiveGPUsForRendering == 1 ? 1 : GNumActiveGPUsForRendering;
-#	else
-		static const auto NumBufferedQueriesVar = IConsoleManager::Get().FindTConsoleVariableDataInt(TEXT("r.NumBufferedOcclusionQueries"));
-		NumBufferedFrames = NumBufferedQueriesVar->GetValueOnAnyThread();
-#	endif		
-#endif
-		return NumBufferedFrames;
-	}
-
-	// get the index of the oldest query based on the current frame and number of buffered frames.
-	static uint32 GetQueryLookupIndex(int32 CurrentFrame, int32 NumBufferedFrames)
-	{
-		// queries are currently always requested earlier in the frame than they are issued.
-		// thus we can always overwrite the oldest query with the current one as we never need them
-		// to coexist.  This saves us a buffer entry.
-		const uint32 QueryIndex = CurrentFrame % NumBufferedFrames;
-		return QueryIndex;
-	}
-
-	// get the index of the query to overwrite for new queries.
-	static uint32 GetQueryIssueIndex(int32 CurrentFrame, int32 NumBufferedFrames)
-	{
-		// queries are currently always requested earlier in the frame than they are issued.
-		// thus we can always overwrite the oldest query with the current one as we never need them
-		// to coexist.  This saves us a buffer entry.
-		const uint32 QueryIndex = CurrentFrame % NumBufferedFrames;
-		return QueryIndex;
-	}
-};
-
 /** Holds information about a single primitive's occlusion. */
 class FPrimitiveOcclusionHistory
 {
@@ -127,13 +116,8 @@ public:
 	/** The primitive the occlusion information is about. */
 	FPrimitiveComponentId PrimitiveId;
 
-#if BUFFERED_OCCLUSION_QUERIES
 	/** The occlusion query which contains the primitive's pending occlusion results. */
-	TArray<FRenderQueryRHIRef, TInlineAllocator<1> > PendingOcclusionQuery;
-#else
-	/** The occlusion query which contains the primitive's pending occlusion results. */
-	FRenderQueryRHIRef PendingOcclusionQuery;
-#endif
+	TArray<FRenderQueryRHIRef, TInlineAllocator<FOcclusionQueryHelpers::MaxBufferedOcclusionFrames> > PendingOcclusionQuery;
 
 	uint32 HZBTestIndex;
 	uint32 HZBTestFrameNumber;
@@ -153,12 +137,6 @@ public:
 	/** whether or not this primitive was grouped the last time it was queried */
 	bool bGroupedQuery;
 
-	/** 
-	 * Number of frames to buffer the occlusion queries. 
-	 * Larger numbers allow better SLI scaling but introduce latency in the results.
-	 */
-	int32 NumBufferedFrames;
-
 	/** 
 	 * For things that have subqueries (folaige), this is the non-zero
 	 */
@@ -175,11 +153,8 @@ public:
 		, bGroupedQuery(false)
 		, CustomIndex(SubQuery)
 	{
-#if BUFFERED_OCCLUSION_QUERIES
-		NumBufferedFrames = FOcclusionQueryHelpers::GetNumBufferedFrames();
-		PendingOcclusionQuery.Empty(NumBufferedFrames);
-		PendingOcclusionQuery.AddZeroed(NumBufferedFrames);
-#endif
+		PendingOcclusionQuery.Empty(FOcclusionQueryHelpers::MaxBufferedOcclusionFrames);
+		PendingOcclusionQuery.AddZeroed(FOcclusionQueryHelpers::MaxBufferedOcclusionFrames);
 	}
 
 	/** Destructor. Note that the query should have been released already. */
@@ -189,38 +164,26 @@ public:
 	}
 
 	template<class TOcclusionQueryPool> // here we use a template just to allow this to be inlined without sorting out the header order
-	FORCEINLINE void ReleaseQueries(FRHICommandListImmediate& RHICmdList, TOcclusionQueryPool& Pool)
+	FORCEINLINE void ReleaseQueries(FRHICommandListImmediate& RHICmdList, TOcclusionQueryPool& Pool, int32 NumBufferedFrames)
 	{
-#if BUFFERED_OCCLUSION_QUERIES
 		for (int32 QueryIndex = 0; QueryIndex < NumBufferedFrames; QueryIndex++)
 		{
 			Pool.ReleaseQuery(RHICmdList, PendingOcclusionQuery[QueryIndex]);
 		}
-#else
-		Pool.ReleaseQuery(RHICmdList, PendingOcclusionQuery);
-#endif
 	}
 
-	FORCEINLINE FRenderQueryRHIRef& GetPastQuery(uint32 FrameNumber)
+	FORCEINLINE FRenderQueryRHIRef& GetPastQuery(uint32 FrameNumber, int32 NumBufferedFrames)
 	{
-#if BUFFERED_OCCLUSION_QUERIES
 		// Get the oldest occlusion query
 		const uint32 QueryIndex = FOcclusionQueryHelpers::GetQueryLookupIndex(FrameNumber, NumBufferedFrames);
 		return PendingOcclusionQuery[QueryIndex];
-#else
-		return PendingOcclusionQuery;
-#endif
 	}
 
-	FORCEINLINE void SetCurrentQuery(uint32 FrameNumber, FRenderQueryRHIParamRef NewQuery)
+	FORCEINLINE void SetCurrentQuery(uint32 FrameNumber, FRenderQueryRHIParamRef NewQuery, int32 NumBufferedFrames)
 	{
-#if BUFFERED_OCCLUSION_QUERIES
 		// Get the current occlusion query
 		const uint32 QueryIndex = FOcclusionQueryHelpers::GetQueryIssueIndex(FrameNumber, NumBufferedFrames);
 		PendingOcclusionQuery[QueryIndex] = NewQuery;
-#else
-		PendingOcclusionQuery = NewQuery;
-#endif
 	}
 };
 
@@ -445,14 +408,9 @@ public:
 		bool bTranslucentShadow;
 	};
 
-	int32 NumBufferedFrames;
 	uint32 UniqueID;
 	typedef TMap<FSceneViewState::FProjectedShadowKey, FRenderQueryRHIRef> ShadowKeyOcclusionQueryMap;
-#if BUFFERED_OCCLUSION_QUERIES
-	TArray<ShadowKeyOcclusionQueryMap> ShadowOcclusionQueryMaps;
-#else
-	ShadowKeyOcclusionQueryMap ShadowOcclusionQueryMap;
-#endif
+	TArray<ShadowKeyOcclusionQueryMap, TInlineAllocator<FOcclusionQueryHelpers::MaxBufferedOcclusionFrames> > ShadowOcclusionQueryMaps;
 
 	/** The view's occlusion query pool. */
 	FRenderQueryPool OcclusionQueryPool;
@@ -677,7 +635,7 @@ public:
 	 * @param Primitive - The shadow subject.
 	 * @param Light - The shadow source.
 	 */
-	bool IsShadowOccluded(FRHICommandListImmediate& RHICmdList, FPrimitiveComponentId PrimitiveId, const ULightComponent* Light, int32 SplitIndex, bool bTranslucentShadow) const;
+	bool IsShadowOccluded(FRHICommandListImmediate& RHICmdList, FPrimitiveComponentId PrimitiveId, const ULightComponent* Light, int32 SplitIndex, bool bTranslucentShadow, int32 NumBufferedFrames) const;
 
 	TRefCountPtr<IPooledRenderTarget>& GetEyeAdaptation()
 	{
@@ -709,14 +667,10 @@ public:
 
 	virtual void ReleaseDynamicRHI() override
 	{
-#if BUFFERED_OCCLUSION_QUERIES
 		for (int i = 0; i < ShadowOcclusionQueryMaps.Num(); ++i)
 		{
 			ShadowOcclusionQueryMaps[i].Reset();
 		}
-#else
-		ShadowOcclusionQueryMap.Reset();
-#endif
 		PrimitiveOcclusionHistorySet.Empty();
 		PrimitiveFadingStates.Empty();
 		OcclusionQueryPool.Release();
diff --git a/Engine/Source/Runtime/Renderer/Private/SceneRendering.cpp b/Engine/Source/Runtime/Renderer/Private/SceneRendering.cpp
--- a/Engine/Source/Runtime/Renderer/Private/SceneRendering.cpp
+++ b/Engine/Source/Runtime/Renderer/Private/SceneRendering.cpp
@@ -1433,6 +1433,9 @@ void FSceneRenderer::WaitForTasksClearSnapshotsAndDeleteSceneRenderer(FRHIComman
 	}
 	FViewInfo::DestroyAllSnapshots(); // this destroys viewinfo snapshots
 	FSceneRenderTargets::Get(RHICmdList).DestroyAllSnapshots(); // this will destroy the render target snapshots
+	static const IConsoleVariable* AsyncDispatch	= IConsoleManager::Get().FindConsoleVariable(TEXT("r.RHICmdAsyncRHIThreadDispatch"));
+
+	if (AsyncDispatch->GetInt() == 0)
 	{
 		QUICK_SCOPE_CYCLE_COUNTER(STAT_DeleteSceneRenderer_Dispatch);
 		RHICmdList.ImmediateFlush(EImmediateFlushType::WaitForDispatchToRHIThread); // we want to make sure this all gets to the rhi thread this frame and doesn't hang around
diff --git a/Engine/Source/Runtime/Renderer/Private/SceneVisibility.cpp b/Engine/Source/Runtime/Renderer/Private/SceneVisibility.cpp
--- a/Engine/Source/Runtime/Renderer/Private/SceneVisibility.cpp
+++ b/Engine/Source/Runtime/Renderer/Private/SceneVisibility.cpp
@@ -346,6 +346,7 @@ static void UpdatePrimitiveFading(const FScene* Scene, FViewInfo& View)
 static int32 OcclusionCull(FRHICommandListImmediate& RHICmdList, const FScene* Scene, FViewInfo& View)
 {
 	SCOPE_CYCLE_COUNTER(STAT_OcclusionCull);
+	int32 NumBufferedFrames = FOcclusionQueryHelpers::GetNumBufferedFrames();
 
 	// INITVIEWS_TODO: This could be more efficient if broken up in to separate concerns:
 	// - What is occluded?
@@ -481,7 +482,7 @@ static int32 OcclusionCull(FRHICommandListImmediate& RHICmdList, const FScene* S
 							{
 								// Read the occlusion query results.
 								uint64 NumSamples = 0;
-								FRenderQueryRHIRef& PastQuery = PrimitiveOcclusionHistory->GetPastQuery(ViewState->OcclusionFrameCounter);
+								FRenderQueryRHIRef& PastQuery = PrimitiveOcclusionHistory->GetPastQuery(ViewState->OcclusionFrameCounter, NumBufferedFrames);
 								if (IsValidRef(PastQuery))
 								{
 									// NOTE: RHIGetOcclusionQueryResult should never fail when using a blocking call, rendering artifacts may show up.
@@ -545,7 +546,7 @@ static int32 OcclusionCull(FRHICommandListImmediate& RHICmdList, const FScene* S
 
 						if (bClearQueries)
 						{
-							ViewState->OcclusionQueryPool.ReleaseQuery(RHICmdList, PrimitiveOcclusionHistory->GetPastQuery(ViewState->OcclusionFrameCounter));
+							ViewState->OcclusionQueryPool.ReleaseQuery(RHICmdList, PrimitiveOcclusionHistory->GetPastQuery(ViewState->OcclusionFrameCounter, NumBufferedFrames));
 						}
 					}
 
@@ -619,7 +620,8 @@ static int32 OcclusionCull(FRHICommandListImmediate& RHICmdList, const FScene* S
 									PrimitiveOcclusionHistory->SetCurrentQuery(ViewState->OcclusionFrameCounter, 
 										bGroupedQuery ? 
 										View.GroupedOcclusionQueries.BatchPrimitive(OcclusionBounds.Origin + View.ViewMatrices.PreViewTranslation,OcclusionBounds.BoxExtent) :
-										View.IndividualOcclusionQueries.BatchPrimitive(OcclusionBounds.Origin + View.ViewMatrices.PreViewTranslation,OcclusionBounds.BoxExtent)
+										View.IndividualOcclusionQueries.BatchPrimitive(OcclusionBounds.Origin + View.ViewMatrices.PreViewTranslation,OcclusionBounds.BoxExtent),
+										NumBufferedFrames
 										);	
 								}
 								PrimitiveOcclusionHistory->bGroupedQuery = bGroupedQuery;
@@ -1271,6 +1273,7 @@ void FSceneRenderer::GatherDynamicMeshElements(
 			}
 		}
 	}
+	MeshCollector.WaitForTasks();
 }
 
 static void MarkAllPrimitivesForReflectionProxyUpdate(FScene* Scene)
diff --git a/Engine/Source/Runtime/Renderer/Private/ShadowSetup.cpp b/Engine/Source/Runtime/Renderer/Private/ShadowSetup.cpp
--- a/Engine/Source/Runtime/Renderer/Private/ShadowSetup.cpp
+++ b/Engine/Source/Runtime/Renderer/Private/ShadowSetup.cpp
@@ -989,6 +989,8 @@ void FProjectedShadowInfo::GatherDynamicMeshElements(FSceneRenderer& Renderer, F
 		    FoundView->ViewMatrices.GetDynamicMeshElementsShadowCullFrustum = nullptr;
     
 		    FoundView->ViewMatrices.ViewMatrix = OriginalViewMatrix;
+
+			Renderer.MeshCollector.WaitForTasks();
 	    }
     }
 }
@@ -1276,6 +1278,7 @@ void FDeferredShadingSceneRenderer::CreatePerObjectProjectedShadow(
 	bool bOpaqueRelevance = false;
 	bool bTranslucentRelevance = false;
 	bool bTranslucentShadowIsVisibleThisFrame = false;
+	int32 NumBufferedFrames = FOcclusionQueryHelpers::GetNumBufferedFrames();
 
 	for (int32 ViewIndex = 0; ViewIndex < Views.Num(); ViewIndex++)
 	{
@@ -1298,14 +1301,14 @@ void FDeferredShadingSceneRenderer::CreatePerObjectProjectedShadow(
 			!bCreateOpaqueObjectShadow ||
 			(
 				!View.bIgnoreExistingQueries &&	View.State &&
-				((FSceneViewState*)View.State)->IsShadowOccluded(RHICmdList, PrimitiveSceneInfo->PrimitiveComponentId, LightSceneInfo->Proxy->GetLightComponent(), INDEX_NONE, false)
+				((FSceneViewState*)View.State)->IsShadowOccluded(RHICmdList, PrimitiveSceneInfo->PrimitiveComponentId, LightSceneInfo->Proxy->GetLightComponent(), INDEX_NONE, false, NumBufferedFrames)
 			);
 
 		const bool bTranslucentShadowIsOccluded = 
 			!bCreateTranslucentObjectShadow ||
 			(
 				!View.bIgnoreExistingQueries && View.State &&
-				((FSceneViewState*)View.State)->IsShadowOccluded(RHICmdList, PrimitiveSceneInfo->PrimitiveComponentId, LightSceneInfo->Proxy->GetLightComponent(), INDEX_NONE, true)
+				((FSceneViewState*)View.State)->IsShadowOccluded(RHICmdList, PrimitiveSceneInfo->PrimitiveComponentId, LightSceneInfo->Proxy->GetLightComponent(), INDEX_NONE, true, NumBufferedFrames)
 			);
 
 		const bool bSubjectIsVisibleInThisView = View.PrimitiveVisibilityMap[PrimitiveSceneInfo->GetIndex()];
@@ -1760,6 +1763,8 @@ void FDeferredShadingSceneRenderer::CreateWholeSceneProjectedShadow(FLightSceneI
 void FSceneRenderer::InitProjectedShadowVisibility(FRHICommandListImmediate& RHICmdList)
 {
 	SCOPE_CYCLE_COUNTER(STAT_InitProjectedShadowVisibility);
+	int32 NumBufferedFrames = FOcclusionQueryHelpers::GetNumBufferedFrames();
+
 	// Initialize the views' ProjectedShadowVisibilityMaps and remove shadows without subjects.
 	for(TSparseArray<FLightSceneInfoCompact>::TConstIterator LightIt(Scene->Lights);LightIt;++LightIt)
 	{
@@ -1833,7 +1838,8 @@ void FSceneRenderer::InitProjectedShadowVisibility(FRHICommandListImmediate& RHI
 								FPrimitiveComponentId(),
 							ProjectedShadowInfo.GetLightSceneInfo().Proxy->GetLightComponent(),
 							ProjectedShadowInfo.CascadeSettings.ShadowSplitIndex,
-							ProjectedShadowInfo.bTranslucentShadow
+							ProjectedShadowInfo.bTranslucentShadow,
+							NumBufferedFrames
 							);
 
 					// The shadow is visible if it is view relevant and unoccluded.
