UE4 - remove cruft

diff --git a/Engine/Source/Runtime/RHI/Private/RHICommandList.cpp b/Engine/Source/Runtime/RHI/Private/RHICommandList.cpp
--- a/Engine/Source/Runtime/RHI/Private/RHICommandList.cpp
+++ b/Engine/Source/Runtime/RHI/Private/RHICommandList.cpp
@@ -606,9 +606,6 @@ void FRHICommandListBase::Reset()
 	Root = nullptr;
 	CommandLink = &Root;
 	Context = GDynamicRHI ? RHIGetDefaultContext() : nullptr;
-#if USE_RHICOMMAND_STATE_REDUCTION
-	StateCache = nullptr;
-#endif
 	UID = GRHICommandList.UIDCounter.Increment();
 	for (int32 Index = 0; ERenderThreadContext(Index) < ERenderThreadContext::Num; Index++)
 	{
