nicer GPU profile printout

diff --git a/Engine/Source/Runtime/RHI/Private/GPUProfiler.cpp b/Engine/Source/Runtime/RHI/Private/GPUProfiler.cpp
--- a/Engine/Source/Runtime/RHI/Private/GPUProfiler.cpp
+++ b/Engine/Source/Runtime/RHI/Private/GPUProfiler.cpp
@@ -241,6 +241,8 @@ void FGPUProfilerEventNodeFrame::DumpEventTree()
 			// It's very unlikely to cause a problem as the cvar is only changes by the user.
 			FString WildcardString = CVar->GetString(); 
 
+			FGPUProfilerEventNodeStats Sum;
+
 			const float ThresholdInMS = 5.0f;
 
 			if(WildcardString == FString(TEXT("*")))
@@ -276,6 +278,7 @@ void FGPUProfilerEventNodeFrame::DumpEventTree()
 				if (bDump)
 				{
 					UE_LOG(LogRHI, Warning, TEXT("   %.2fms   %s   Events %u   Draws %u"), NodeStats.TimingResult, *It.Key(), NodeStats.NumEvents, NodeStats.NumDraws);
+					Sum += NodeStats;
 				}
 				else
 				{
@@ -283,7 +286,8 @@ void FGPUProfilerEventNodeFrame::DumpEventTree()
 				}
 			}
 
-			UE_LOG(LogRHI, Warning, TEXT("   %u buckets not shown"), NumNotShown);
+			UE_LOG(LogRHI, Warning, TEXT("   Total %.2fms   Events %u   Draws %u,    %u buckets not shown"), 
+				Sum.TimingResult, Sum.NumEvents, Sum.NumDraws, NumNotShown);
 		}
 
 #if !UE_BUILD_SHIPPING
diff --git a/Engine/Source/Runtime/RHI/Public/GPUProfiler.h b/Engine/Source/Runtime/RHI/Public/GPUProfiler.h
--- a/Engine/Source/Runtime/RHI/Public/GPUProfiler.h
+++ b/Engine/Source/Runtime/RHI/Public/GPUProfiler.h
@@ -33,6 +33,17 @@ public:
 
 	/** Inclusive number of other perf events that this is the parent of. */
 	uint32 NumEvents;
+
+	const FGPUProfilerEventNodeStats operator+=(const FGPUProfilerEventNodeStats& rhs)
+	{
+		NumDraws += rhs.NumDraws;
+		NumPrimitives += rhs.NumPrimitives;
+		NumVertices += rhs.NumVertices;
+		TimingResult += rhs.TimingResult;
+		NumEvents += rhs.NumEvents;
+
+		return *this;
+	}
 };
 
 /** Stats for a single perf event node. */
