UE4 - Vulkan defines

diff --git a/Engine/Source/Runtime/RHI/Private/RHI.cpp b/Engine/Source/Runtime/RHI/Private/RHI.cpp
--- a/Engine/Source/Runtime/RHI/Private/RHI.cpp
+++ b/Engine/Source/Runtime/RHI/Private/RHI.cpp
@@ -286,6 +286,7 @@ static FName NAME_OPENGL_ES2_IOS(TEXT("GLSL_ES2_IOS"));
 static FName NAME_SF_METAL(TEXT("SF_METAL"));
 static FName NAME_SF_METAL_MRT(TEXT("SF_METAL_MRT"));
 static FName NAME_GLSL_310_ES_EXT(TEXT("GLSL_310_ES_EXT"));
+static FName NAME_PC_VULKAN_ES2(TEXT("PC_VULKAN_ES2"));
 
 FName LegacyShaderPlatformToShaderFormat(EShaderPlatform Platform)
 {
@@ -328,6 +329,8 @@ FName LegacyShaderPlatformToShaderFormat(EShaderPlatform Platform)
 		return NAME_SF_METAL_MRT;
 	case SP_OPENGL_ES31_EXT:
 		return NAME_GLSL_310_ES_EXT;
+	case SP_VULKAN_ES2:
+		return NAME_PC_VULKAN_ES2;
 
 	default:
 		check(0);
@@ -355,6 +358,7 @@ EShaderPlatform ShaderFormatToLegacyShaderPlatform(FName ShaderFormat)
 	if (ShaderFormat == NAME_SF_METAL)			return SP_METAL;
 	if (ShaderFormat == NAME_SF_METAL_MRT)		return SP_METAL_MRT;
 	if (ShaderFormat == NAME_GLSL_310_ES_EXT)	return SP_OPENGL_ES31_EXT;
+	if (ShaderFormat == NAME_PC_VULKAN_ES2)		return SP_VULKAN_ES2;
 	return SP_NumPlatforms;
 }
 
diff --git a/Engine/Source/Runtime/RHI/Public/RHIDefinitions.h b/Engine/Source/Runtime/RHI/Public/RHIDefinitions.h
--- a/Engine/Source/Runtime/RHI/Public/RHIDefinitions.h
+++ b/Engine/Source/Runtime/RHI/Public/RHIDefinitions.h
@@ -57,8 +57,9 @@ enum EShaderPlatform
 	SP_PCD3D_ES3_1		= 15,
 	/** Used when running in Feature Level ES3_1 in OpenGL. */
 	SP_OPENGL_PCES3_1	= 16,
+	SP_VULKAN_ES2		= 17,
 
-	SP_NumPlatforms		= 17,
+	SP_NumPlatforms		= 18,
 	SP_NumBits			= 5,
 };
 static_assert(SP_NumPlatforms <= (1 << SP_NumBits), "SP_NumPlatforms will not fit on SP_NumBits");
@@ -605,7 +606,7 @@ inline bool IsPCPlatform(const EShaderPlatform Platform)
 /** Whether the shader platform corresponds to the ES2 feature level. */
 inline bool IsES2Platform(const EShaderPlatform Platform)
 {
-	return Platform == SP_PCD3D_ES2 || Platform == SP_OPENGL_PCES2 || Platform == SP_OPENGL_ES2 || Platform == SP_OPENGL_ES2_WEBGL || Platform == SP_OPENGL_ES2_IOS || Platform == SP_METAL; 
+	return Platform == SP_PCD3D_ES2 || Platform == SP_OPENGL_PCES2 || Platform == SP_OPENGL_ES2 || Platform == SP_OPENGL_ES2_WEBGL || Platform == SP_OPENGL_ES2_IOS || Platform == SP_METAL || Platform == SP_VULKAN_ES2; 
 }
 
 /** Whether the shader platform corresponds to the ES2 feature level. */
@@ -617,7 +618,7 @@ inline bool IsMobilePlatform(const EShaderPlatform Platform)
 inline bool IsOpenGLPlatform(const EShaderPlatform Platform)
 {
 	return Platform == SP_OPENGL_SM4 || Platform == SP_OPENGL_SM4_MAC || Platform == SP_OPENGL_SM5 || Platform == SP_OPENGL_PCES2
-		|| Platform == SP_OPENGL_ES2 || Platform == SP_OPENGL_ES2_WEBGL || Platform == SP_OPENGL_ES2_IOS || Platform == SP_OPENGL_ES31_EXT;
+		|| Platform == SP_OPENGL_ES2 || Platform == SP_OPENGL_ES2_WEBGL || Platform == SP_OPENGL_ES2_IOS || Platform == SP_OPENGL_ES31_EXT || Platform == SP_VULKAN_ES2;
 }
 
 inline bool IsConsolePlatform(const EShaderPlatform Platform)
@@ -645,6 +646,7 @@ inline ERHIFeatureLevel::Type GetMaxSupportedFeatureLevel(EShaderPlatform InShad
 	case SP_OPENGL_ES2:
 	case SP_OPENGL_ES2_WEBGL:
 	case SP_OPENGL_ES2_IOS:
+	case SP_VULKAN_ES2:
 		return ERHIFeatureLevel::ES2;
 	case SP_METAL:
 	case SP_PCD3D_ES3_1:
@@ -674,6 +676,7 @@ inline bool IsFeatureLevelSupported(EShaderPlatform InShaderPlatform, ERHIFeatur
 	case SP_OPENGL_PCES3_1:
 		return InFeatureLevel <= ERHIFeatureLevel::ES3_1;
 	case SP_OPENGL_ES2:
+	case SP_VULKAN_ES2:
 		return InFeatureLevel <= ERHIFeatureLevel::ES2;
 	case SP_OPENGL_ES2_WEBGL: 
 		return InFeatureLevel <= ERHIFeatureLevel::ES2;
