UE4 - vk module

diff --git a/Engine/Source/Programs/UnrealBuildTool/Windows/UEBuildWindows.cs b/Engine/Source/Programs/UnrealBuildTool/Windows/UEBuildWindows.cs
--- a/Engine/Source/Programs/UnrealBuildTool/Windows/UEBuildWindows.cs
+++ b/Engine/Source/Programs/UnrealBuildTool/Windows/UEBuildWindows.cs
@@ -487,7 +487,16 @@ namespace UnrealBuildTool
                 {
                     InModule.AddDynamicallyLoadedModule("ShaderFormatD3D");
                     InModule.AddDynamicallyLoadedModule("ShaderFormatOpenGL");
-                }
+
+//#todo-rco: Remove when public
+					{
+						string VulkanSDKPath = Environment.GetEnvironmentVariable("VulkanSDK");
+						if (!String.IsNullOrEmpty(VulkanSDKPath))
+						{
+							InModule.AddDynamicallyLoadedModule("VulkanShaderFormat");
+						}
+					}
+				}
 
                 if (InModule.ToString() == "D3D11RHI")
                 {
diff --git a/Engine/Source/Runtime/RHI/Private/Windows/WindowsDynamicRHI.cpp b/Engine/Source/Runtime/RHI/Private/Windows/WindowsDynamicRHI.cpp
--- a/Engine/Source/Runtime/RHI/Private/Windows/WindowsDynamicRHI.cpp
+++ b/Engine/Source/Runtime/RHI/Private/Windows/WindowsDynamicRHI.cpp
@@ -10,12 +10,12 @@ FDynamicRHI* PlatformCreateDynamicRHI()
 	// command line overrides
 	const bool bForceD3D11 = FParse::Param(FCommandLine::Get(),TEXT("d3d11")) || FParse::Param(FCommandLine::Get(),TEXT("sm5")) || FParse::Param(FCommandLine::Get(),TEXT("dx11"));
 	const bool bForceD3D10 = FParse::Param(FCommandLine::Get(),TEXT("d3d10")) || FParse::Param(FCommandLine::Get(),TEXT("sm4")) || FParse::Param(FCommandLine::Get(),TEXT("dx10"));
-
 	const bool bForceOpenGL = FWindowsPlatformMisc::VerifyWindowsVersion(6, 0) == false || FParse::Param(FCommandLine::Get(), TEXT("opengl")) || FParse::Param(FCommandLine::Get(), TEXT("opengl3")) || FParse::Param(FCommandLine::Get(), TEXT("opengl4"));
+	const bool bForceVulkan = FParse::Param(FCommandLine::Get(), TEXT("vulkan"));
 
-	if (((bForceD3D11 ? 1 : 0) + (bForceD3D10 ? 1 : 0) + (bForceOpenGL ? 1 : 0)) > 1)
+	if (((bForceD3D11 ? 1 : 0) + (bForceD3D10 ? 1 : 0) + (bForceOpenGL ? 1 : 0) + (bForceVulkan ? 1 : 0)) > 1)
 	{
-		UE_LOG(LogRHI, Fatal,TEXT("-d3d11, -d3d10, and -opengl[3|4] mutually exclusive options, but more than one was specified on the command-line."));
+		UE_LOG(LogRHI, Fatal,TEXT("-d3d11, -d3d10, -vulkan and -opengl[3|4] mutually exclusive options, but more than one was specified on the command-line."));
 	}
 
 	// Load the dynamic RHI module.
@@ -31,6 +31,16 @@ FDynamicRHI* PlatformCreateDynamicRHI()
 			DynamicRHIModule = NULL;
 		}
 	}
+	else if (bForceVulkan)
+	{
+		DynamicRHIModule = &FModuleManager::LoadModuleChecked<IDynamicRHIModule>(TEXT("VulkanRHI"));
+		if (!DynamicRHIModule->IsSupported())
+		{
+			FMessageDialog::Open(EAppMsgType::Ok, NSLOCTEXT("WindowsDynamicRHI", "RequiredVulkan", "Vulkan Driver is required to run the engine."));
+			FPlatformMisc::RequestExit(1);
+			DynamicRHIModule = NULL;
+		}
+	}
 	else
 	{
 		DynamicRHIModule = &FModuleManager::LoadModuleChecked<IDynamicRHIModule>(TEXT("D3D11RHI"));
diff --git a/Engine/Source/Runtime/RHI/RHI.Build.cs b/Engine/Source/Runtime/RHI/RHI.Build.cs
--- a/Engine/Source/Runtime/RHI/RHI.Build.cs
+++ b/Engine/Source/Runtime/RHI/RHI.Build.cs
@@ -1,6 +1,7 @@
 // Copyright 1998-2015 Epic Games, Inc. All Rights Reserved.
 
 using UnrealBuildTool;
+using System;
 
 public class RHI : ModuleRules
 {
@@ -15,6 +16,12 @@ public class RHI : ModuleRules
 			if ((Target.Platform == UnrealTargetPlatform.Win32) || (Target.Platform == UnrealTargetPlatform.Win64))
 			{
 				DynamicallyLoadedModuleNames.Add("D3D11RHI");
+//#todo-rco: Remove when public
+				string VulkanSDKPath = Environment.GetEnvironmentVariable("VulkanSDK");
+				if (!String.IsNullOrEmpty(VulkanSDKPath))
+				{
+					DynamicallyLoadedModuleNames.Add("VulkanRHI");
+				}
 			}
 
 			if ((Target.Platform == UnrealTargetPlatform.Win32) ||
