Merging added SP_METAL_SM5 (only in RHI.cpp and RHIDefinitions.h)

diff --git a/Engine/Source/Runtime/RHI/Private/RHI.cpp b/Engine/Source/Runtime/RHI/Private/RHI.cpp
--- a/Engine/Source/Runtime/RHI/Private/RHI.cpp
+++ b/Engine/Source/Runtime/RHI/Private/RHI.cpp
@@ -286,6 +286,7 @@ static FName NAME_OPENGL_ES2_IOS(TEXT("GLSL_ES2_IOS"));
 static FName NAME_SF_METAL(TEXT("SF_METAL"));
 static FName NAME_SF_METAL_MRT(TEXT("SF_METAL_MRT"));
 static FName NAME_GLSL_310_ES_EXT(TEXT("GLSL_310_ES_EXT"));
+static FName NAME_SF_METAL_SM5(TEXT("SF_METAL_SM5"));
 static FName NAME_PC_VULKAN_ES2(TEXT("PC_VULKAN_ES2"));
 
 FName LegacyShaderPlatformToShaderFormat(EShaderPlatform Platform)
@@ -327,6 +328,8 @@ FName LegacyShaderPlatformToShaderFormat(EShaderPlatform Platform)
 		return NAME_SF_METAL;
 	case SP_METAL_MRT:
 		return NAME_SF_METAL_MRT;
+	case SP_METAL_SM5:
+		return NAME_SF_METAL_SM5;
 	case SP_OPENGL_ES31_EXT:
 		return NAME_GLSL_310_ES_EXT;
 	case SP_VULKAN_ES2:
@@ -358,6 +361,7 @@ EShaderPlatform ShaderFormatToLegacyShaderPlatform(FName ShaderFormat)
 	if (ShaderFormat == NAME_SF_METAL)			return SP_METAL;
 	if (ShaderFormat == NAME_SF_METAL_MRT)		return SP_METAL_MRT;
 	if (ShaderFormat == NAME_GLSL_310_ES_EXT)	return SP_OPENGL_ES31_EXT;
+	if (ShaderFormat == NAME_SF_METAL_SM5)		return SP_METAL_SM5;
 	if (ShaderFormat == NAME_PC_VULKAN_ES2)		return SP_VULKAN_ES2;
 	return SP_NumPlatforms;
 }
diff --git a/Engine/Source/Runtime/RHI/Public/RHIDefinitions.h b/Engine/Source/Runtime/RHI/Public/RHIDefinitions.h
--- a/Engine/Source/Runtime/RHI/Public/RHIDefinitions.h
+++ b/Engine/Source/Runtime/RHI/Public/RHIDefinitions.h
@@ -57,9 +57,10 @@ enum EShaderPlatform
 	SP_PCD3D_ES3_1		= 15,
 	/** Used when running in Feature Level ES3_1 in OpenGL. */
 	SP_OPENGL_PCES3_1	= 16,
-	SP_VULKAN_ES2		= 17,
+	SP_METAL_SM5		= 17,
+	SP_VULKAN_ES2		= 18,
 
-	SP_NumPlatforms		= 18,
+	SP_NumPlatforms		= 19,
 	SP_NumBits			= 5,
 };
 static_assert(SP_NumPlatforms <= (1 << SP_NumBits), "SP_NumPlatforms will not fit on SP_NumBits");
@@ -600,7 +601,7 @@ enum class ESimpleRenderTargetMode
 
 inline bool IsPCPlatform(const EShaderPlatform Platform)
 {
-	return Platform == SP_PCD3D_SM5 || Platform == SP_PCD3D_SM4 || Platform == SP_PCD3D_ES2 || Platform == SP_PCD3D_ES3_1 || Platform ==  SP_OPENGL_SM4 || Platform == SP_OPENGL_SM4_MAC || Platform == SP_OPENGL_SM5 || Platform == SP_OPENGL_PCES2 || Platform == SP_OPENGL_PCES3_1;
+	return Platform == SP_PCD3D_SM5 || Platform == SP_PCD3D_SM4 || Platform == SP_PCD3D_ES2 || Platform == SP_PCD3D_ES3_1 || Platform ==  SP_OPENGL_SM4 || Platform == SP_OPENGL_SM4_MAC || Platform == SP_OPENGL_SM5 || Platform == SP_OPENGL_PCES2 || Platform == SP_OPENGL_PCES3_1 || (PLATFORM_MAC && (Platform == SP_METAL || Platform == SP_METAL_MRT || Platform == SP_METAL_SM5));
 }
 
 /** Whether the shader platform corresponds to the ES2 feature level. */
@@ -635,6 +636,7 @@ inline ERHIFeatureLevel::Type GetMaxSupportedFeatureLevel(EShaderPlatform InShad
 	case SP_PS4:
 	case SP_XBOXONE:
 	case SP_OPENGL_ES31_EXT:
+	case SP_METAL_SM5:
 		return ERHIFeatureLevel::SM5;
 	case SP_PCD3D_SM4:
 	case SP_OPENGL_SM4:
@@ -698,6 +700,8 @@ inline bool IsFeatureLevelSupported(EShaderPlatform InShaderPlatform, ERHIFeatur
 		return InFeatureLevel <= ERHIFeatureLevel::SM4;
 	case SP_OPENGL_ES31_EXT:
 		return InFeatureLevel <= ERHIFeatureLevel::SM5;
+	case SP_METAL_SM5:
+		return InFeatureLevel <= ERHIFeatureLevel::SM5;
 	default:
 		return false;
 	}	
@@ -731,23 +735,23 @@ inline bool RHISupportsInstancing(ERHIFeatureLevel::Type FeatureLevel)
 inline bool RHISupportsSeparateMSAAAndResolveTextures(const EShaderPlatform Platform)
 {
 	// Metal needs to handle MSAA and resolve textures internally (unless RHICreateTexture2D was changed to take an optional resolve target)
-	return Platform != SP_METAL && Platform != SP_METAL_MRT;
+	return Platform != SP_METAL && Platform != SP_METAL_MRT && Platform != SP_METAL_SM5;
 }
 
 inline bool RHISupportsComputeShaders(const EShaderPlatform Platform)
 {
 	//@todo-rco: Add Metal support
-	return IsFeatureLevelSupported(Platform, ERHIFeatureLevel::SM5) || (Platform == SP_METAL_MRT);
+	return IsFeatureLevelSupported(Platform, ERHIFeatureLevel::SM5);
 }
 
 inline bool RHISupportsGeometryShaders(const EShaderPlatform Platform)
 {
-	return IsFeatureLevelSupported(Platform, ERHIFeatureLevel::SM4) && Platform != SP_METAL_MRT;
+	return IsFeatureLevelSupported(Platform, ERHIFeatureLevel::SM4) && Platform != SP_METAL_MRT && Platform != SP_METAL_SM5;
 }
 
 inline bool RHIHasTiledGPU(const EShaderPlatform Platform)
 {
-	return Platform == SP_METAL_MRT|| Platform == SP_METAL || Platform == SP_OPENGL_ES2_IOS || Platform == SP_OPENGL_ES2;
+	return (PLATFORM_IOS && Platform == SP_METAL_MRT) || Platform == SP_METAL || Platform == SP_OPENGL_ES2_IOS || Platform == SP_OPENGL_ES2;
 }
 
 inline uint32 GetFeatureLevelMaxTextureSamplers(ERHIFeatureLevel::Type FeatureLevel)
