Update comments about BUF_Dynamic and BUF_Volatile

diff --git a/Engine/Source/Runtime/RHI/Public/RHIDefinitions.h b/Engine/Source/Runtime/RHI/Public/RHIDefinitions.h
--- a/Engine/Source/Runtime/RHI/Public/RHIDefinitions.h
+++ b/Engine/Source/Runtime/RHI/Public/RHIDefinitions.h
@@ -411,8 +411,8 @@ enum EBufferUsageFlags
 {
 	// Mutually exclusive write-frequency flags
 	BUF_Static            = 0x0001, // The buffer will be written to once.
-	BUF_Dynamic           = 0x0002, // The buffer will be written to occasionally.
-	BUF_Volatile          = 0x0004, // The buffer will be written to frequently.
+	BUF_Dynamic           = 0x0002, // The buffer will be written to occasionally.  The data lifetime is until the next update, or the buffer is destroyed.
+	BUF_Volatile          = 0x0004, // The buffer's data will have a lifetime of one frame.  It MUST be written to each frame, or a new one created each frame.
 
 	// Mutually exclusive bind flags.
 	BUF_UnorderedAccess   = 0x0008, // Allows an unordered access view to be created for the buffer.
