Vertex and Index buffer write locks on PS4 no longer require memcopies.
RHi now provides CreateAndLock functions for Vertex and Index buffers which are more efficient on certain platforms.
Use r.PS4StandardWriteLocks=1 to enable the old default functionality for testing if necessary.

diff --git a/Engine/Plugins/Runtime/CustomMeshComponent/Source/CustomMeshComponent/Private/CustomMeshComponent.cpp b/Engine/Plugins/Runtime/CustomMeshComponent/Source/CustomMeshComponent/Private/CustomMeshComponent.cpp
--- a/Engine/Plugins/Runtime/CustomMeshComponent/Source/CustomMeshComponent/Private/CustomMeshComponent.cpp
+++ b/Engine/Plugins/Runtime/CustomMeshComponent/Source/CustomMeshComponent/Private/CustomMeshComponent.cpp
@@ -17,10 +17,10 @@ public:
 	virtual void InitRHI() override
 	{
 		FRHIResourceCreateInfo CreateInfo;
-		VertexBufferRHI = RHICreateVertexBuffer(Vertices.Num() * sizeof(FDynamicMeshVertex),BUF_Static,CreateInfo);
+		void* VertexBufferData = nullptr;
+		VertexBufferRHI = RHICreateAndLockVertexBuffer(Vertices.Num() * sizeof(FDynamicMeshVertex), BUF_Static, CreateInfo, VertexBufferData);
 
-		// Copy the vertex data into the vertex buffer.
-		void* VertexBufferData = RHILockVertexBuffer(VertexBufferRHI,0,Vertices.Num() * sizeof(FDynamicMeshVertex), RLM_WriteOnly);
+		// Copy the vertex data into the vertex buffer.		
 		FMemory::Memcpy(VertexBufferData,Vertices.GetData(),Vertices.Num() * sizeof(FDynamicMeshVertex));
 		RHIUnlockVertexBuffer(VertexBufferRHI);
 	}
@@ -36,10 +36,10 @@ public:
 	virtual void InitRHI() override
 	{
 		FRHIResourceCreateInfo CreateInfo;
-		IndexBufferRHI = RHICreateIndexBuffer(sizeof(int32),Indices.Num() * sizeof(int32),BUF_Static, CreateInfo);
+		void* Buffer = nullptr;
+		IndexBufferRHI = RHICreateAndLockIndexBuffer(sizeof(int32),Indices.Num() * sizeof(int32),BUF_Static, CreateInfo, Buffer);
 
-		// Write the indices to the index buffer.
-		void* Buffer = RHILockIndexBuffer(IndexBufferRHI,0,Indices.Num() * sizeof(int32),RLM_WriteOnly);
+		// Write the indices to the index buffer.		
 		FMemory::Memcpy(Buffer,Indices.GetData(),Indices.Num() * sizeof(int32));
 		RHIUnlockIndexBuffer(IndexBufferRHI);
 	}
diff --git a/Engine/Source/Runtime/Engine/Private/Components/TextRenderComponent.cpp b/Engine/Source/Runtime/Engine/Private/Components/TextRenderComponent.cpp
--- a/Engine/Source/Runtime/Engine/Private/Components/TextRenderComponent.cpp
+++ b/Engine/Source/Runtime/Engine/Private/Components/TextRenderComponent.cpp
@@ -115,10 +115,10 @@ public:
 	virtual void InitRHI() override
 	{
 		FRHIResourceCreateInfo CreateInfo;
-		VertexBufferRHI = RHICreateVertexBuffer(Vertices.Num() * sizeof(FDynamicMeshVertex),BUF_Static,CreateInfo);
+		void* VertexBufferData = nullptr;
+		VertexBufferRHI = RHICreateAndLockVertexBuffer(Vertices.Num() * sizeof(FDynamicMeshVertex),BUF_Static,CreateInfo, VertexBufferData);
 
-		// Copy the vertex data into the vertex buffer.
-		void* VertexBufferData = RHILockVertexBuffer(VertexBufferRHI,0,Vertices.Num() * sizeof(FDynamicMeshVertex), RLM_WriteOnly);
+		// Copy the vertex data into the vertex buffer.		
 		FMemory::Memcpy(VertexBufferData,Vertices.GetData(),Vertices.Num() * sizeof(FDynamicMeshVertex));
 		RHIUnlockVertexBuffer(VertexBufferRHI);
 	}
@@ -133,10 +133,10 @@ public:
 	void InitRHI()
 	{
 		FRHIResourceCreateInfo CreateInfo;
-		IndexBufferRHI = RHICreateIndexBuffer(sizeof(uint16), Indices.Num() * sizeof(uint16), BUF_Static, CreateInfo);
+		void* Buffer = nullptr;
+		IndexBufferRHI = RHICreateAndLockIndexBuffer(sizeof(uint16), Indices.Num() * sizeof(uint16), BUF_Static, CreateInfo, Buffer);
 
-		// Copy the index data into the index buffer.
-		void* Buffer = RHILockIndexBuffer(IndexBufferRHI, 0, Indices.Num() * sizeof(uint16), RLM_WriteOnly);
+		// Copy the index data into the index buffer.		
 		FMemory::Memcpy(Buffer, Indices.GetData(), Indices.Num() * sizeof(uint16));
 		RHIUnlockIndexBuffer(IndexBufferRHI);
 	}
diff --git a/Engine/Source/Runtime/Engine/Private/NavMeshRenderingHelpers.h b/Engine/Source/Runtime/Engine/Private/NavMeshRenderingHelpers.h
--- a/Engine/Source/Runtime/Engine/Private/NavMeshRenderingHelpers.h
+++ b/Engine/Source/Runtime/Engine/Private/NavMeshRenderingHelpers.h
@@ -181,10 +181,10 @@ public:
 		if (Indices.Num() > 0)
 		{
 			FRHIResourceCreateInfo CreateInfo;
-			IndexBufferRHI = RHICreateIndexBuffer(sizeof(int32), Indices.Num() * sizeof(int32), BUF_Static, CreateInfo);
+			void* Buffer = nullptr;
+			IndexBufferRHI = RHICreateAndLockIndexBuffer(sizeof(int32), Indices.Num() * sizeof(int32), BUF_Static, CreateInfo, Buffer);
 
-			// Write the indices to the index buffer.
-			void* Buffer = RHILockIndexBuffer(IndexBufferRHI, 0, Indices.Num() * sizeof(int32), RLM_WriteOnly);
+			// Write the indices to the index buffer.			
 			FMemory::Memcpy(Buffer, Indices.GetData(), Indices.Num() * sizeof(int32));
 			RHIUnlockIndexBuffer(IndexBufferRHI);
 		}
@@ -202,10 +202,10 @@ public:
 		if (Vertices.Num() > 0)
 		{
 			FRHIResourceCreateInfo CreateInfo;
-			VertexBufferRHI = RHICreateVertexBuffer(Vertices.Num() * sizeof(FDynamicMeshVertex), BUF_Static, CreateInfo);
+			void* VertexBufferData = nullptr;
+			VertexBufferRHI = RHICreateAndLockVertexBuffer(Vertices.Num() * sizeof(FDynamicMeshVertex), BUF_Static, CreateInfo, VertexBufferData);
 
-			// Copy the vertex data into the vertex buffer.
-			void* VertexBufferData = RHILockVertexBuffer(VertexBufferRHI, 0, Vertices.Num() * sizeof(FDynamicMeshVertex), RLM_WriteOnly);
+			// Copy the vertex data into the vertex buffer.			
 			FMemory::Memcpy(VertexBufferData, Vertices.GetData(), Vertices.Num() * sizeof(FDynamicMeshVertex));
 			RHIUnlockVertexBuffer(VertexBufferRHI);
 		}
diff --git a/Engine/Source/Runtime/Engine/Private/Particles/ParticleResources.cpp b/Engine/Source/Runtime/Engine/Private/Particles/ParticleResources.cpp
--- a/Engine/Source/Runtime/Engine/Private/Particles/ParticleResources.cpp
+++ b/Engine/Source/Runtime/Engine/Private/Particles/ParticleResources.cpp
@@ -17,8 +17,9 @@ void FParticleTexCoordVertexBuffer::InitRHI()
 {
 	const uint32 Size = sizeof(FVector2D) * 4 * MAX_PARTICLES_PER_INSTANCE;
 	FRHIResourceCreateInfo CreateInfo;
-	VertexBufferRHI = RHICreateVertexBuffer( Size, BUF_Static, CreateInfo);
-	FVector2D* Vertices = (FVector2D*)RHILockVertexBuffer( VertexBufferRHI, 0, Size, RLM_WriteOnly );
+	void* BufferData = nullptr;
+	VertexBufferRHI = RHICreateAndLockVertexBuffer(Size, BUF_Static, CreateInfo, BufferData);
+	FVector2D* Vertices = (FVector2D*)BufferData;
 	for (uint32 SpriteIndex = 0; SpriteIndex < MAX_PARTICLES_PER_INSTANCE; ++SpriteIndex)
 	{
 		Vertices[SpriteIndex*4 + 0] = FVector2D(0.0f, 0.0f);
@@ -44,8 +45,9 @@ void FParticleIndexBuffer::InitRHI()
 	const uint32 Size = sizeof(uint16) * 6 * MaxParticles;
 	const uint32 Stride = sizeof(uint16);
 	FRHIResourceCreateInfo CreateInfo;
-	IndexBufferRHI = RHICreateIndexBuffer( Stride, Size, BUF_Static, CreateInfo);
-	uint16* Indices = (uint16*)RHILockIndexBuffer( IndexBufferRHI, 0, Size, RLM_WriteOnly );
+	void* Buffer = nullptr;
+	IndexBufferRHI = RHICreateAndLockIndexBuffer( Stride, Size, BUF_Static, CreateInfo, Buffer);
+	uint16* Indices = (uint16*)Buffer;
 	for (uint32 SpriteIndex = 0; SpriteIndex < MaxParticles; ++SpriteIndex)
 	{
 		Indices[SpriteIndex*6 + 0] = SpriteIndex*4 + 0;
diff --git a/Engine/Source/Runtime/Engine/Private/PhysicsEngine/PhysDrawing.cpp b/Engine/Source/Runtime/Engine/Private/PhysicsEngine/PhysDrawing.cpp
--- a/Engine/Source/Runtime/Engine/Private/PhysicsEngine/PhysDrawing.cpp
+++ b/Engine/Source/Runtime/Engine/Private/PhysicsEngine/PhysDrawing.cpp
@@ -493,10 +493,10 @@ void FKConvexElem::AddCachedSolidConvexGeom(TArray<FDynamicMeshVertex>& VertexBu
 void FConvexCollisionVertexBuffer::InitRHI()
 {
 	FRHIResourceCreateInfo CreateInfo;
-	VertexBufferRHI = RHICreateVertexBuffer(Vertices.Num() * sizeof(FDynamicMeshVertex),BUF_Static, CreateInfo);
+	void* VertexBufferData = nullptr;
+	VertexBufferRHI = RHICreateAndLockVertexBuffer(Vertices.Num() * sizeof(FDynamicMeshVertex), BUF_Static, CreateInfo, VertexBufferData);
 
-	// Copy the vertex data into the vertex buffer.
-	void* VertexBufferData = RHILockVertexBuffer(VertexBufferRHI,0,Vertices.Num() * sizeof(FDynamicMeshVertex), RLM_WriteOnly);
+	// Copy the vertex data into the vertex buffer.	
 	FMemory::Memcpy(VertexBufferData,Vertices.GetData(),Vertices.Num() * sizeof(FDynamicMeshVertex));
 	RHIUnlockVertexBuffer(VertexBufferRHI);
 }
@@ -504,10 +504,10 @@ void FConvexCollisionVertexBuffer::InitRHI()
 void FConvexCollisionIndexBuffer::InitRHI()
 {
 	FRHIResourceCreateInfo CreateInfo;
-	IndexBufferRHI = RHICreateIndexBuffer(sizeof(int32),Indices.Num() * sizeof(int32),BUF_Static,CreateInfo);
+	void* Buffer = nullptr;
+	IndexBufferRHI = RHICreateAndLockIndexBuffer(sizeof(int32),Indices.Num() * sizeof(int32),BUF_Static,CreateInfo, Buffer);
 
-	// Write the indices to the index buffer.
-	void* Buffer = RHILockIndexBuffer(IndexBufferRHI,0,Indices.Num() * sizeof(int32),RLM_WriteOnly);
+	// Write the indices to the index buffer.	
 	FMemory::Memcpy(Buffer,Indices.GetData(),Indices.Num() * sizeof(int32));
 	RHIUnlockIndexBuffer(IndexBufferRHI);
 }
diff --git a/Engine/Source/Runtime/Engine/Private/RawIndexBuffer.cpp b/Engine/Source/Runtime/Engine/Private/RawIndexBuffer.cpp
--- a/Engine/Source/Runtime/Engine/Private/RawIndexBuffer.cpp
+++ b/Engine/Source/Runtime/Engine/Private/RawIndexBuffer.cpp
@@ -40,10 +40,10 @@ void FRawIndexBuffer::InitRHI()
 	{
 		// Create the index buffer.
 		FRHIResourceCreateInfo CreateInfo;
-		IndexBufferRHI = RHICreateIndexBuffer(sizeof(uint16),Size,BUF_Static,CreateInfo);
+		void* Buffer = nullptr;
+		IndexBufferRHI = RHICreateAndLockIndexBuffer(sizeof(uint16),Size,BUF_Static,CreateInfo, Buffer);
 
-		// Initialize the buffer.
-		void* Buffer = RHILockIndexBuffer(IndexBufferRHI, 0, Size, RLM_WriteOnly);
+		// Initialize the buffer.		
 		FMemory::Memcpy(Buffer,Indices.GetData(),Size);
 		RHIUnlockIndexBuffer(IndexBufferRHI);
 	}
@@ -103,10 +103,10 @@ void FRawIndexBuffer16or32::InitRHI()
 	{
 		// Create the index buffer.
 		FRHIResourceCreateInfo CreateInfo;
-		IndexBufferRHI = RHICreateIndexBuffer(IndexStride,Size,BUF_Static,CreateInfo);
+		void* Buffer = nullptr;
+		IndexBufferRHI = RHICreateAndLockIndexBuffer(IndexStride,Size,BUF_Static,CreateInfo, Buffer);
 		
-		// Initialize the buffer.
-		void* Buffer = RHILockIndexBuffer(IndexBufferRHI, 0, Size, RLM_WriteOnly);
+		// Initialize the buffer.		
 		if (b32Bit)
 		{
 			FMemory::Memcpy(Buffer, Indices.GetData(), Size);
diff --git a/Engine/Source/Runtime/Engine/Private/SkeletalRenderCPUSkin.cpp b/Engine/Source/Runtime/Engine/Private/SkeletalRenderCPUSkin.cpp
--- a/Engine/Source/Runtime/Engine/Private/SkeletalRenderCPUSkin.cpp
+++ b/Engine/Source/Runtime/Engine/Private/SkeletalRenderCPUSkin.cpp
@@ -76,10 +76,8 @@ void FFinalSkinVertexBuffer::InitVertexData(FStaticLODModel& LodModel)
 	uint32 Size = LodModel.NumVertices * sizeof(FFinalSkinVertex);
 
 	FRHIResourceCreateInfo CreateInfo;
-	VertexBufferRHI = RHICreateVertexBuffer(Size,BUF_Dynamic, CreateInfo);
-
-	// Lock the buffer.
-	void* Buffer = RHILockVertexBuffer(VertexBufferRHI,0,Size,RLM_WriteOnly);
+	void* Buffer = nullptr;
+	VertexBufferRHI = RHICreateAndLockVertexBuffer(Size,BUF_Dynamic, CreateInfo, Buffer);	
 
 	// Initialize the vertex data
 	// All chunks are combined into one (rigid first, soft next)
diff --git a/Engine/Source/Runtime/Engine/Private/SkeletalRenderGPUSkin.cpp b/Engine/Source/Runtime/Engine/Private/SkeletalRenderGPUSkin.cpp
--- a/Engine/Source/Runtime/Engine/Private/SkeletalRenderGPUSkin.cpp
+++ b/Engine/Source/Runtime/Engine/Private/SkeletalRenderGPUSkin.cpp
@@ -62,10 +62,11 @@ void FMorphVertexBuffer::InitDynamicRHI()
 	// Create the buffer rendering resource
 	uint32 Size = LodModel.NumVertices * sizeof(FMorphGPUSkinVertex);
 	FRHIResourceCreateInfo CreateInfo;
-	VertexBufferRHI = RHICreateVertexBuffer(Size,BUF_Dynamic,CreateInfo);
+	void* BufferData = nullptr;
+	VertexBufferRHI = RHICreateAndLockVertexBuffer(Size, BUF_Dynamic, CreateInfo, BufferData);
 
 	// Lock the buffer.
-	FMorphGPUSkinVertex* Buffer = (FMorphGPUSkinVertex*) RHILockVertexBuffer(VertexBufferRHI,0,Size,RLM_WriteOnly);
+	FMorphGPUSkinVertex* Buffer = (FMorphGPUSkinVertex*)BufferData;
 
 	// zero all deltas (NOTE: DeltaTangentZ is FPackedNormal, so we can't just FMemory::Memzero)
 	for (uint32 VertIndex=0; VertIndex < LodModel.NumVertices; ++VertIndex)
diff --git a/Engine/Source/Runtime/Engine/Private/TileRendering.cpp b/Engine/Source/Runtime/Engine/Private/TileRendering.cpp
--- a/Engine/Source/Runtime/Engine/Private/TileRendering.cpp
+++ b/Engine/Source/Runtime/Engine/Private/TileRendering.cpp
@@ -56,9 +56,9 @@ public:
 		uint32 Size = 4 * sizeof(FMaterialTileVertex);
 		// create vertex buffer
 		FRHIResourceCreateInfo CreateInfo;
-		VertexBufferRHI = RHICreateVertexBuffer(Size,BUF_Static,CreateInfo);
-		// lock it
-		void* Buffer = RHILockVertexBuffer(VertexBufferRHI,0,Size,RLM_WriteOnly);
+		void* Buffer = nullptr;
+		VertexBufferRHI = RHICreateAndLockVertexBuffer(Size,BUF_Static,CreateInfo, Buffer);
+				
 		// first vertex element
 		FMaterialTileVertex* DestVertex = (FMaterialTileVertex*)Buffer;
 
diff --git a/Engine/Source/Runtime/Engine/Private/TriangleRendering.cpp b/Engine/Source/Runtime/Engine/Private/TriangleRendering.cpp
--- a/Engine/Source/Runtime/Engine/Private/TriangleRendering.cpp
+++ b/Engine/Source/Runtime/Engine/Private/TriangleRendering.cpp
@@ -56,9 +56,9 @@ public:
 		uint32 Size = 3 * sizeof(FMaterialTriangleVertex);
 		// create vertex buffer
 		FRHIResourceCreateInfo CreateInfo;
-		VertexBufferRHI = RHICreateVertexBuffer(Size,BUF_Static,CreateInfo);
-		// lock it
-		void* Buffer = RHILockVertexBuffer(VertexBufferRHI,0,Size,RLM_WriteOnly);
+		void* Buffer = nullptr;
+		VertexBufferRHI = RHICreateAndLockVertexBuffer(Size,BUF_Static,CreateInfo, Buffer);
+
 		// first vertex element
 		FMaterialTriangleVertex* DestVertex = (FMaterialTriangleVertex*)Buffer;
 		// fill out the verts
diff --git a/Engine/Source/Runtime/Engine/Private/VectorFieldVisualization.cpp b/Engine/Source/Runtime/Engine/Private/VectorFieldVisualization.cpp
--- a/Engine/Source/Runtime/Engine/Private/VectorFieldVisualization.cpp
+++ b/Engine/Source/Runtime/Engine/Private/VectorFieldVisualization.cpp
@@ -86,8 +86,9 @@ public:
 	virtual void InitRHI() override
 	{
 		FRHIResourceCreateInfo CreateInfo;
-		VertexBufferRHI = RHICreateVertexBuffer(sizeof(FVector4) * 2, BUF_Static, CreateInfo);
-		FVector4* DummyContents = (FVector4*)RHILockVertexBuffer(VertexBufferRHI,0,sizeof(FVector4)*2,RLM_WriteOnly);
+		void* BufferData = nullptr;
+		VertexBufferRHI = RHICreateAndLockVertexBuffer(sizeof(FVector4) * 2, BUF_Static, CreateInfo, BufferData);
+		FVector4* DummyContents = (FVector4*)BufferData;
 		DummyContents[0] = FVector4(0.0f, 0.0f, 0.0f, 0.0f);
 		DummyContents[1] = FVector4(1.0f, 1.0f, 1.0f, 1.0f);
 		RHIUnlockVertexBuffer(VertexBufferRHI);
diff --git a/Engine/Source/Runtime/Engine/Public/ParticleHelper.h b/Engine/Source/Runtime/Engine/Public/ParticleHelper.h
--- a/Engine/Source/Runtime/Engine/Public/ParticleHelper.h
+++ b/Engine/Source/Runtime/Engine/Public/ParticleHelper.h
@@ -2579,8 +2579,9 @@ public:
 	{
 		// create a static vertex buffer
 		FRHIResourceCreateInfo CreateInfo;
-		VertexBufferRHI = RHICreateVertexBuffer(sizeof(FParticleVertexDynamicParameter), BUF_Static|BUF_ZeroStride, CreateInfo);
-		FParticleVertexDynamicParameter* Vertices = (FParticleVertexDynamicParameter*)RHILockVertexBuffer(VertexBufferRHI, 0, sizeof(FParticleVertexDynamicParameter), RLM_WriteOnly);
+		void* BufferData = nullptr;
+		VertexBufferRHI = RHICreateAndLockVertexBuffer(sizeof(FParticleVertexDynamicParameter), BUF_Static | BUF_ZeroStride, CreateInfo, BufferData);
+		FParticleVertexDynamicParameter* Vertices = (FParticleVertexDynamicParameter*)BufferData;
 		Vertices[0].DynamicValue[0] = Vertices[0].DynamicValue[1] = Vertices[0].DynamicValue[2] = Vertices[0].DynamicValue[3] = 1.0f;
 		RHIUnlockVertexBuffer(VertexBufferRHI);
 	}
diff --git a/Engine/Source/Runtime/GeometryCache/Private/GeometryCacheSceneProxy.cpp b/Engine/Source/Runtime/GeometryCache/Private/GeometryCacheSceneProxy.cpp
--- a/Engine/Source/Runtime/GeometryCache/Private/GeometryCacheSceneProxy.cpp
+++ b/Engine/Source/Runtime/GeometryCache/Private/GeometryCacheSceneProxy.cpp
@@ -270,10 +270,10 @@ void FGeomCacheVertexFactory::Init(const FGeomCacheVertexBuffer* VertexBuffer)
 void FGeomCacheIndexBuffer::InitRHI()
 {
 	FRHIResourceCreateInfo CreateInfo;
-	IndexBufferRHI = RHICreateIndexBuffer(sizeof(int32), Indices.Num() * sizeof(int32), BUF_Static, CreateInfo);
+	void* Buffer = nullptr;
+	IndexBufferRHI = RHICreateAndLockIndexBuffer(sizeof(int32), Indices.Num() * sizeof(int32), BUF_Static, CreateInfo);
 
-	// Write the indices to the index buffer.
-	void* Buffer = RHILockIndexBuffer(IndexBufferRHI, 0, Indices.Num() * sizeof(int32), RLM_WriteOnly);
+	// Write the indices to the index buffer.	
 	FMemory::Memcpy(Buffer, Indices.GetData(), Indices.Num() * sizeof(int32));
 	RHIUnlockIndexBuffer(IndexBufferRHI);
 }
diff --git a/Engine/Source/Runtime/Landscape/Private/LandscapeRender.cpp b/Engine/Source/Runtime/Landscape/Private/LandscapeRender.cpp
--- a/Engine/Source/Runtime/Landscape/Private/LandscapeRender.cpp
+++ b/Engine/Source/Runtime/Landscape/Private/LandscapeRender.cpp
@@ -1694,8 +1694,9 @@ void FLandscapeVertexBuffer::InitRHI()
 {
 	// create a static vertex buffer
 	FRHIResourceCreateInfo CreateInfo;
-	VertexBufferRHI = RHICreateVertexBuffer(NumVertices * sizeof(FLandscapeVertex), BUF_Static, CreateInfo);
-	FLandscapeVertex* Vertex = (FLandscapeVertex*)RHILockVertexBuffer(VertexBufferRHI, 0, NumVertices * sizeof(FLandscapeVertex), RLM_WriteOnly);
+	void* BufferData = nullptr;
+	VertexBufferRHI = RHICreateAndLockVertexBuffer(NumVertices * sizeof(FLandscapeVertex), BUF_Static, CreateInfo, BufferData);
+	FLandscapeVertex* Vertex = (FLandscapeVertex*)BufferData;
 	int32 VertexIndex = 0;
 	for (int32 SubY = 0; SubY < NumSubsections; SubY++)
 	{
diff --git a/Engine/Source/Runtime/Landscape/Private/LandscapeRenderMobile.cpp b/Engine/Source/Runtime/Landscape/Private/LandscapeRenderMobile.cpp
--- a/Engine/Source/Runtime/Landscape/Private/LandscapeRenderMobile.cpp
+++ b/Engine/Source/Runtime/Landscape/Private/LandscapeRenderMobile.cpp
@@ -206,8 +206,9 @@ void FLandscapeVertexBufferMobile::InitRHI()
 {
 	// create a static vertex buffer
 	FRHIResourceCreateInfo CreateInfo;
-	VertexBufferRHI = RHICreateVertexBuffer(DataSize, BUF_Static, CreateInfo);
-	void* VertexData = RHILockVertexBuffer(VertexBufferRHI, 0, DataSize, RLM_WriteOnly);
+	void* VertexData = nullptr;
+	VertexBufferRHI = RHICreateAndLockVertexBuffer(DataSize, BUF_Static, CreateInfo, VertexData);
+	
 	// Copy stored platform data
 	FMemory::Memcpy(VertexData, (uint8*)Data, DataSize);
 	RHIUnlockVertexBuffer(VertexBufferRHI);
diff --git a/Engine/Source/Runtime/ProceduralMeshComponent/Private/ProceduralMeshComponent.cpp b/Engine/Source/Runtime/ProceduralMeshComponent/Private/ProceduralMeshComponent.cpp
--- a/Engine/Source/Runtime/ProceduralMeshComponent/Private/ProceduralMeshComponent.cpp
+++ b/Engine/Source/Runtime/ProceduralMeshComponent/Private/ProceduralMeshComponent.cpp
@@ -55,10 +55,10 @@ public:
 	virtual void InitRHI() override
 	{
 		FRHIResourceCreateInfo CreateInfo;
-		IndexBufferRHI = RHICreateIndexBuffer(sizeof(int32), Indices.Num() * sizeof(int32), BUF_Static, CreateInfo);
+		void* Buffer = nullptr;
+		IndexBufferRHI = RHICreateAndLockIndexBuffer(sizeof(int32), Indices.Num() * sizeof(int32), BUF_Static, CreateInfo, Buffer);
 
-		// Write the indices to the index buffer.
-		void* Buffer = RHILockIndexBuffer(IndexBufferRHI, 0, Indices.Num() * sizeof(int32), RLM_WriteOnly);
+		// Write the indices to the index buffer.		
 		FMemory::Memcpy(Buffer, Indices.GetData(), Indices.Num() * sizeof(int32));
 		RHIUnlockIndexBuffer(IndexBufferRHI);
 	}
diff --git a/Engine/Source/Runtime/RHI/Private/RHICommandList.cpp b/Engine/Source/Runtime/RHI/Private/RHICommandList.cpp
--- a/Engine/Source/Runtime/RHI/Private/RHICommandList.cpp
+++ b/Engine/Source/Runtime/RHI/Private/RHICommandList.cpp
@@ -1413,6 +1413,21 @@ void FRHICommandList::operator delete(void *RawMemory)
 
 ///////// Pass through functions that allow RHIs to optimize certain calls.
 
+FVertexBufferRHIRef FDynamicRHI::CreateAndLockVertexBuffer_RenderThread(class FRHICommandListImmediate& RHICmdList, uint32 Size, uint32 InUsage, FRHIResourceCreateInfo& CreateInfo, void*& OutDataBuffer)
+{
+	FVertexBufferRHIRef VertexBuffer = CreateVertexBuffer_RenderThread(RHICmdList, Size, InUsage, CreateInfo);
+	OutDataBuffer = LockVertexBuffer_RenderThread(RHICmdList, VertexBuffer, 0, Size, RLM_WriteOnly);
+
+	return VertexBuffer;
+}
+
+FIndexBufferRHIRef FDynamicRHI::CreateAndLockIndexBuffer_RenderThread(class FRHICommandListImmediate& RHICmdList, uint32 Stride, uint32 Size, uint32 InUsage, FRHIResourceCreateInfo& CreateInfo, void*& OutDataBuffer)
+{
+	FIndexBufferRHIRef IndexBuffer = CreateIndexBuffer_RenderThread(RHICmdList, Stride, Size, InUsage, CreateInfo);
+	OutDataBuffer = LockIndexBuffer_RenderThread(RHICmdList, IndexBuffer, 0, Size, RLM_WriteOnly);
+	return IndexBuffer;
+}
+
 
 FVertexBufferRHIRef FDynamicRHI::CreateVertexBuffer_RenderThread(class FRHICommandListImmediate& RHICmdList, uint32 Size, uint32 InUsage, FRHIResourceCreateInfo& CreateInfo)
 {
diff --git a/Engine/Source/Runtime/RHI/Public/DynamicRHI.h b/Engine/Source/Runtime/RHI/Public/DynamicRHI.h
--- a/Engine/Source/Runtime/RHI/Public/DynamicRHI.h
+++ b/Engine/Source/Runtime/RHI/Public/DynamicRHI.h
@@ -118,7 +118,7 @@ public:
 	virtual void* RHILockVertexBuffer(FVertexBufferRHIParamRef VertexBuffer, uint32 Offset, uint32 SizeRHI, EResourceLockMode LockMode) = 0;
 
 	// FlushType: Flush RHI Thread
-	virtual void RHIUnlockVertexBuffer(FVertexBufferRHIParamRef VertexBuffer) = 0;
+	virtual void RHIUnlockVertexBuffer(FVertexBufferRHIParamRef VertexBuffer) = 0;	
 
 	/** Copies the contents of one vertex buffer to another vertex buffer.  They must have identical sizes. */
 	// FlushType: Flush Immediate (seems dangerous)
@@ -618,6 +618,8 @@ public:
 	virtual class IRHICommandContextContainer* RHIGetCommandContextContainer() = 0;
 
 	///////// Pass through functions that allow RHIs to optimize certain calls.
+	virtual FVertexBufferRHIRef CreateAndLockVertexBuffer_RenderThread(class FRHICommandListImmediate& RHICmdList, uint32 Size, uint32 InUsage, FRHIResourceCreateInfo& CreateInfo, void*& OutDataBuffer);
+	virtual FIndexBufferRHIRef CreateAndLockIndexBuffer_RenderThread(class FRHICommandListImmediate& RHICmdList, uint32 Stride, uint32 Size, uint32 InUsage, FRHIResourceCreateInfo& CreateInfo, void*& OutDataBuffer);
 
 	virtual FVertexBufferRHIRef CreateVertexBuffer_RenderThread(class FRHICommandListImmediate& RHICmdList, uint32 Size, uint32 InUsage, FRHIResourceCreateInfo& CreateInfo);
 	virtual FShaderResourceViewRHIRef CreateShaderResourceView_RenderThread(class FRHICommandListImmediate& RHICmdList, FVertexBufferRHIParamRef VertexBuffer, uint32 Stride, uint8 Format);
diff --git a/Engine/Source/Runtime/RHI/Public/RHICommandList.h b/Engine/Source/Runtime/RHI/Public/RHICommandList.h
--- a/Engine/Source/Runtime/RHI/Public/RHICommandList.h
+++ b/Engine/Source/Runtime/RHI/Public/RHICommandList.h
@@ -2042,6 +2042,11 @@ public:
 	{
 		return RHICreateUniformBuffer(Contents, Layout, Usage);
 	}
+
+	FORCEINLINE FIndexBufferRHIRef CreateAndLockIndexBuffer(uint32 Stride, uint32 Size, uint32 InUsage, FRHIResourceCreateInfo& CreateInfo, void*& OutDataBuffer)
+	{
+		return GDynamicRHI->CreateAndLockIndexBuffer_RenderThread(*this, Stride, Size, InUsage, CreateInfo, OutDataBuffer);
+	}
 	
 	FORCEINLINE FIndexBufferRHIRef CreateIndexBuffer(uint32 Stride, uint32 Size, uint32 InUsage, FRHIResourceCreateInfo& CreateInfo)
 	{
@@ -2058,6 +2063,11 @@ public:
 		GDynamicRHI->UnlockIndexBuffer_RenderThread(*this, IndexBuffer);
 	}
 	
+	FORCEINLINE FVertexBufferRHIRef CreateAndLockVertexBuffer(uint32 Size, uint32 InUsage, FRHIResourceCreateInfo& CreateInfo, void*& OutDataBuffer)
+	{
+		return GDynamicRHI->CreateAndLockVertexBuffer_RenderThread(*this, Size, InUsage, CreateInfo, OutDataBuffer);
+	}
+
 	FORCEINLINE FVertexBufferRHIRef CreateVertexBuffer(uint32 Size, uint32 InUsage, FRHIResourceCreateInfo& CreateInfo)
 	{
 		return GDynamicRHI->CreateVertexBuffer_RenderThread(*this, Size, InUsage, CreateInfo);
@@ -2711,6 +2721,11 @@ FORCEINLINE FComputeShaderRHIRef RHICreateComputeShader(const TArray<uint8>& Cod
 	return FRHICommandListExecutor::GetImmediateCommandList().CreateComputeShader(Code);
 }
 
+FORCEINLINE FIndexBufferRHIRef RHICreateAndLockIndexBuffer(uint32 Stride, uint32 Size, uint32 InUsage, FRHIResourceCreateInfo& CreateInfo, void*& OutDataBuffer)
+{
+	return FRHICommandListExecutor::GetImmediateCommandList().CreateAndLockIndexBuffer(Stride, Size, InUsage, CreateInfo, OutDataBuffer);
+}
+
 FORCEINLINE FIndexBufferRHIRef RHICreateIndexBuffer(uint32 Stride, uint32 Size, uint32 InUsage, FRHIResourceCreateInfo& CreateInfo)
 {
 	return FRHICommandListExecutor::GetImmediateCommandList().CreateIndexBuffer(Stride, Size, InUsage, CreateInfo);
@@ -2726,6 +2741,11 @@ FORCEINLINE void RHIUnlockIndexBuffer(FIndexBufferRHIParamRef IndexBuffer)
 	 FRHICommandListExecutor::GetImmediateCommandList().UnlockIndexBuffer(IndexBuffer);
 }
 
+FORCEINLINE FVertexBufferRHIRef RHICreateAndLockVertexBuffer(uint32 Size, uint32 InUsage, FRHIResourceCreateInfo& CreateInfo, void*& OutDataBuffer)
+{
+	return FRHICommandListExecutor::GetImmediateCommandList().CreateAndLockVertexBuffer(Size, InUsage, CreateInfo, OutDataBuffer);
+}
+
 FORCEINLINE FVertexBufferRHIRef RHICreateVertexBuffer(uint32 Size, uint32 InUsage, FRHIResourceCreateInfo& CreateInfo)
 {
 	return FRHICommandListExecutor::GetImmediateCommandList().CreateVertexBuffer(Size, InUsage, CreateInfo);
diff --git a/Engine/Source/Runtime/RenderCore/Public/RenderResource.h b/Engine/Source/Runtime/RenderCore/Public/RenderResource.h
--- a/Engine/Source/Runtime/RenderCore/Public/RenderResource.h
+++ b/Engine/Source/Runtime/RenderCore/Public/RenderResource.h
@@ -437,8 +437,10 @@ public:
 	{
 		// create a static vertex buffer
 		FRHIResourceCreateInfo CreateInfo;
-		VertexBufferRHI = RHICreateVertexBuffer(sizeof(uint32), BUF_Static | BUF_ZeroStride, CreateInfo);
-		uint32* Vertices = (uint32*)RHILockVertexBuffer(VertexBufferRHI, 0, sizeof(uint32), RLM_WriteOnly);
+		
+		void* LockedData = nullptr;
+		VertexBufferRHI = RHICreateAndLockVertexBuffer(sizeof(uint32), BUF_Static | BUF_ZeroStride, CreateInfo, LockedData);
+		uint32* Vertices = (uint32*)LockedData;
 		Vertices[0] = FColor(255, 255, 255, 255).DWColor();
 		RHIUnlockVertexBuffer(VertexBufferRHI);
 	}
diff --git a/Engine/Source/Runtime/Renderer/Private/DistanceFieldSurfaceCacheLighting.cpp b/Engine/Source/Runtime/Renderer/Private/DistanceFieldSurfaceCacheLighting.cpp
--- a/Engine/Source/Runtime/Renderer/Private/DistanceFieldSurfaceCacheLighting.cpp
+++ b/Engine/Source/Runtime/Renderer/Private/DistanceFieldSurfaceCacheLighting.cpp
@@ -292,8 +292,8 @@ public:
 		// Used as a non-indexed triangle list, so 3 vertices per triangle
 		const uint32 Size = 3 * NumSections * sizeof(FScreenVertex);
 		FRHIResourceCreateInfo CreateInfo;
-		VertexBufferRHI = RHICreateVertexBuffer(Size, BUF_Static, CreateInfo);
-		void* Buffer = RHILockVertexBuffer(VertexBufferRHI, 0, Size, RLM_WriteOnly);
+		void* Buffer = nullptr;
+		VertexBufferRHI = RHICreateAndLockVertexBuffer(Size, BUF_Static, CreateInfo, Buffer);
 		FScreenVertex* DestVertex = (FScreenVertex*)Buffer;
 
 		const float RadiansPerRingSegment = PI / (float)NumSections;
diff --git a/Engine/Source/Runtime/Renderer/Private/HeightfieldLighting.cpp b/Engine/Source/Runtime/Renderer/Private/HeightfieldLighting.cpp
--- a/Engine/Source/Runtime/Renderer/Private/HeightfieldLighting.cpp
+++ b/Engine/Source/Runtime/Renderer/Private/HeightfieldLighting.cpp
@@ -378,8 +378,9 @@ public:
 	{
 		const uint32 Size = 6 * sizeof(FScreenVertex);
 		FRHIResourceCreateInfo CreateInfo;
-		VertexBufferRHI = RHICreateVertexBuffer(Size, BUF_Static, CreateInfo);
-		void* Buffer = RHILockVertexBuffer(VertexBufferRHI, 0, Size, RLM_WriteOnly);
+
+		void* Buffer = nullptr;
+		VertexBufferRHI = RHICreateAndLockVertexBuffer(Size, BUF_Static, CreateInfo, Buffer);		
 		FScreenVertex* DestVertex = (FScreenVertex*)Buffer;
 
 		DestVertex[0].Position = FVector2D(1,  1);
diff --git a/Engine/Source/Runtime/Renderer/Private/PostProcess/PostProcessBokehDOF.cpp b/Engine/Source/Runtime/Renderer/Private/PostProcess/PostProcessBokehDOF.cpp
--- a/Engine/Source/Runtime/Renderer/Private/PostProcess/PostProcessBokehDOF.cpp
+++ b/Engine/Source/Runtime/Renderer/Private/PostProcess/PostProcessBokehDOF.cpp
@@ -35,8 +35,9 @@ public:
 		const uint32 Size = sizeof(uint16) * 6 * 8;
 		const uint32 Stride = sizeof(uint16);
 		FRHIResourceCreateInfo CreateInfo;
-		IndexBufferRHI = RHICreateIndexBuffer( Stride, Size, BUF_Static, CreateInfo);
-		uint16* Indices = (uint16*)RHILockIndexBuffer( IndexBufferRHI, 0, Size, RLM_WriteOnly );
+		void* Buffer = nullptr;
+		IndexBufferRHI = RHICreateAndLockIndexBuffer(Stride, Size, BUF_Static, CreateInfo, Buffer);
+		uint16* Indices = (uint16*)Buffer;
 		if(DOFIndexStyle == BIS_Fast)
 		{
 			for (uint32 SpriteIndex = 0; SpriteIndex < 8; ++SpriteIndex)
diff --git a/Engine/Source/Runtime/Renderer/Private/PostProcess/PostProcessMotionBlur.cpp b/Engine/Source/Runtime/Renderer/Private/PostProcess/PostProcessMotionBlur.cpp
--- a/Engine/Source/Runtime/Renderer/Private/PostProcess/PostProcessMotionBlur.cpp
+++ b/Engine/Source/Runtime/Renderer/Private/PostProcess/PostProcessMotionBlur.cpp
@@ -875,8 +875,9 @@ public:
 		const uint32 Size = sizeof(uint16) * 6 * 8;
 		const uint32 Stride = sizeof(uint16);
 		FRHIResourceCreateInfo CreateInfo;
-		IndexBufferRHI = RHICreateIndexBuffer( Stride, Size, BUF_Static, CreateInfo );
-		uint16* Indices = (uint16*)RHILockIndexBuffer( IndexBufferRHI, 0, Size, RLM_WriteOnly );
+		void* Buffer = nullptr;
+		IndexBufferRHI = RHICreateAndLockIndexBuffer( Stride, Size, BUF_Static, CreateInfo, Buffer );
+		uint16* Indices = (uint16*)Buffer;
 		for (uint32 SpriteIndex = 0; SpriteIndex < 8; ++SpriteIndex)
 		{
 #if PLATFORM_MAC // Avoid a driver bug on OSX/NV cards that causes driver to generate an unwound index buffer
diff --git a/Engine/Source/Runtime/Renderer/Private/SceneOcclusion.cpp b/Engine/Source/Runtime/Renderer/Private/SceneOcclusion.cpp
--- a/Engine/Source/Runtime/Renderer/Private/SceneOcclusion.cpp
+++ b/Engine/Source/Runtime/Renderer/Private/SceneOcclusion.cpp
@@ -310,8 +310,11 @@ public:
 		const uint32 SizeInBytes = MaxBatchedPrimitives * NUM_CUBE_VERTICES * Stride;
 
 		FRHIResourceCreateInfo CreateInfo;
-		IndexBufferRHI = RHICreateIndexBuffer(Stride,SizeInBytes,BUF_Static,CreateInfo);
-		uint16* RESTRICT Indices = (uint16*)RHILockIndexBuffer(IndexBufferRHI,0,SizeInBytes,RLM_WriteOnly);
+
+		void* BufferData;
+		IndexBufferRHI = RHICreateAndLockIndexBuffer(Stride, SizeInBytes, BUF_Static, CreateInfo, BufferData);
+		uint16* RESTRICT Indices = (uint16*)BufferData;
+
 		for(uint32 PrimitiveIndex = 0;PrimitiveIndex < MaxBatchedPrimitives;PrimitiveIndex++)
 		{
 			for(int32 Index = 0;Index < NUM_CUBE_VERTICES;Index++)
diff --git a/Engine/Source/Runtime/Renderer/Private/TranslucentLighting.cpp b/Engine/Source/Runtime/Renderer/Private/TranslucentLighting.cpp
--- a/Engine/Source/Runtime/Renderer/Private/TranslucentLighting.cpp
+++ b/Engine/Source/Runtime/Renderer/Private/TranslucentLighting.cpp
@@ -1038,8 +1038,8 @@ public:
 		// Used as a non-indexed triangle strip, so 4 vertices per quad
 		const uint32 Size = 4 * sizeof(FScreenVertex);
 		FRHIResourceCreateInfo CreateInfo;
-		VertexBufferRHI = RHICreateVertexBuffer(Size, BUF_Static, CreateInfo);
-		void* Buffer = RHILockVertexBuffer(VertexBufferRHI, 0, Size, RLM_WriteOnly);
+		void* Buffer = nullptr;
+		VertexBufferRHI = RHICreateAndLockVertexBuffer(Size, BUF_Static, CreateInfo, Buffer);		
 		FScreenVertex* DestVertex = (FScreenVertex*)Buffer;
 
 		// Setup a full - render target quad
diff --git a/Engine/Source/Runtime/ShaderCore/Private/ShaderCache.cpp b/Engine/Source/Runtime/ShaderCore/Private/ShaderCache.cpp
--- a/Engine/Source/Runtime/ShaderCore/Private/ShaderCache.cpp
+++ b/Engine/Source/Runtime/ShaderCore/Private/ShaderCache.cpp
@@ -801,8 +801,8 @@ void FShaderCache::InternalPreDrawShaders(FRHICommandList& RHICmdList)
 			FRHIResourceCreateInfo Info;
 			uint32 Stride = sizeof(uint16);
 			uint32 Size = sizeof(uint16) * 3;
-			IndexBufferUInt16 = RHICreateIndexBuffer(Stride, Size, BUF_Static, Info);
-			void* Data = RHILockIndexBuffer(IndexBufferUInt16, 0, Size, RLM_WriteOnly);
+			void* Data = nullptr;
+			IndexBufferUInt16 = RHICreateAndLockIndexBuffer(Stride, Size, BUF_Static, Info, Data);			
 			if ( Data )
 			{
 				FMemory::Memzero(Data, Size);
@@ -814,8 +814,8 @@ void FShaderCache::InternalPreDrawShaders(FRHICommandList& RHICmdList)
 			FRHIResourceCreateInfo Info;
 			uint32 Stride = sizeof(uint32);
 			uint32 Size = sizeof(uint32) * 3;
-			IndexBufferUInt32 = RHICreateIndexBuffer(Stride, Size, BUF_Static, Info);
-			void* Data = RHILockIndexBuffer(IndexBufferUInt32, 0, Size, RLM_WriteOnly);
+			void* Data = nullptr;
+			IndexBufferUInt32 = RHICreateAndLockIndexBuffer(Stride, Size, BUF_Static, Info, Data);			
 			if ( Data )
 			{
 				FMemory::Memzero(Data, Size);
