UE4 - fix clang warning

diff --git a/Engine/Source/Runtime/RHI/Private/RHICommandList.cpp b/Engine/Source/Runtime/RHI/Private/RHICommandList.cpp
--- a/Engine/Source/Runtime/RHI/Private/RHICommandList.cpp
+++ b/Engine/Source/Runtime/RHI/Private/RHICommandList.cpp
@@ -1411,7 +1411,7 @@ void FRHICommandList::operator delete(void *RawMemory)
 void* FRHICommandListBase::operator new(size_t Size)
 {
 	check(0); // you shouldn't be creating these
-	return nullptr;
+	return FMemory::Malloc(Size);
 }
 
 void FRHICommandListBase::operator delete(void *RawMemory)
