Fix ensure() from performance warning on depthstore action.

diff --git a/Engine/Source/Runtime/RHI/Public/RHIResources.h b/Engine/Source/Runtime/RHI/Public/RHIResources.h
--- a/Engine/Source/Runtime/RHI/Public/RHIResources.h
+++ b/Engine/Source/Runtime/RHI/Public/RHIResources.h
@@ -997,9 +997,9 @@ public:
 	}
 
 	void Validate() const
-	{
-		// Missed optimization, resolve is not needed for read only ops
-		ensure(!(DepthStencilAccess == FExclusiveDepthStencil::DepthRead_StencilRead && StencilStoreAction != ERenderTargetStoreAction::ENoAction));
+	{		
+		ensureMsgf(DepthStencilAccess.IsDepthWrite() || DepthStoreAction == ERenderTargetStoreAction::ENoAction, TEXT("Depth is read-only, but we are performing a store.  This is a waste on mobile.  If depth can't change, we don't need to store it out again"));
+		ensureMsgf(DepthStencilAccess.IsStencilWrite() || StencilStoreAction == ERenderTargetStoreAction::ENoAction, TEXT("Stencil is read-only, but we are performing a store.  This is a waste on mobile.  If stencil can't change, we don't need to store it out again"));
 	}
 
 	bool operator==(const FRHIDepthRenderTargetView& Other)
diff --git a/Engine/Source/Runtime/RHI/Public/RHIUtilities.h b/Engine/Source/Runtime/RHI/Public/RHIUtilities.h
--- a/Engine/Source/Runtime/RHI/Public/RHIUtilities.h
+++ b/Engine/Source/Runtime/RHI/Public/RHIUtilities.h
@@ -123,7 +123,7 @@ struct FRWBufferByteAddress
  * Convert the ESimpleRenderTargetMode into usable values 
  * @todo: Can we easily put this into a .cpp somewhere?
  */
-inline void DecodeRenderTargetMode(ESimpleRenderTargetMode Mode, ERenderTargetLoadAction& ColorLoadAction, ERenderTargetStoreAction& ColorStoreAction, ERenderTargetLoadAction& DepthLoadAction, ERenderTargetStoreAction& DepthStoreAction)
+inline void DecodeRenderTargetMode(ESimpleRenderTargetMode Mode, ERenderTargetLoadAction& ColorLoadAction, ERenderTargetStoreAction& ColorStoreAction, ERenderTargetLoadAction& DepthLoadAction, ERenderTargetStoreAction& DepthStoreAction, FExclusiveDepthStencil DepthStencilUsage)
 {
 	// set defaults
 	ColorStoreAction = ERenderTargetStoreAction::EStore;
@@ -163,6 +163,12 @@ inline void DecodeRenderTargetMode(ESimpleRenderTargetMode Mode, ERenderTargetLo
 	default:
 		UE_LOG(LogRHI, Fatal, TEXT("Using a ESimpleRenderTargetMode that wasn't decoded in DecodeRenderTargetMode [value = %d]"), (int32)Mode);
 	}
+
+	//if we aren't writing to depth, there's no reason to store it back out again.  Should save some bandwidth on mobile platforms.
+	if (!DepthStencilUsage.IsDepthWrite())
+	{
+		DepthStoreAction = ERenderTargetStoreAction::ENoAction;
+	}
 }
 
 /** Helper for the common case of using a single color and depth render target. */
@@ -178,7 +184,7 @@ inline void SetRenderTarget(FRHICommandList& RHICmdList, FTextureRHIParamRef New
 {
 	ERenderTargetLoadAction ColorLoadAction, DepthLoadAction;
 	ERenderTargetStoreAction ColorStoreAction, DepthStoreAction;	
-	DecodeRenderTargetMode(Mode, ColorLoadAction, ColorStoreAction, DepthLoadAction, DepthStoreAction);
+	DecodeRenderTargetMode(Mode, ColorLoadAction, ColorStoreAction, DepthLoadAction, DepthStoreAction, DepthStencilAccess);
 
 	// now make the FRHISetRenderTargetsInfo that encapsulates all of the info
 	FRHIRenderTargetView ColorView(NewRenderTarget, 0, -1, ColorLoadAction, ColorStoreAction);
@@ -235,7 +241,7 @@ inline void SetRenderTargets(
 {
 	ERenderTargetLoadAction ColorLoadAction, DepthLoadAction;
 	ERenderTargetStoreAction ColorStoreAction, DepthStoreAction;	
-	DecodeRenderTargetMode(Mode, ColorLoadAction, ColorStoreAction, DepthLoadAction, DepthStoreAction);
+	DecodeRenderTargetMode(Mode, ColorLoadAction, ColorStoreAction, DepthLoadAction, DepthStoreAction, DepthStencilAccess);
 
 	FRHIRenderTargetView RTVs[MaxSimultaneousRenderTargets];
 
