Re-enable parallel rendering

diff --git a/Engine/Source/Runtime/RHI/Private/RHICommandList.cpp b/Engine/Source/Runtime/RHI/Private/RHICommandList.cpp
--- a/Engine/Source/Runtime/RHI/Private/RHICommandList.cpp
+++ b/Engine/Source/Runtime/RHI/Private/RHICommandList.cpp
@@ -462,7 +462,6 @@ void FRHICommandListExecutor::LatchBypass()
 	{
 		GRHICommandList.GetImmediateCommandList().ImmediateFlush(EImmediateFlushType::DispatchToRHIThread);		
 
-#if 0
 		static bool bOnce = false;
 		if (!bOnce)
 		{
@@ -473,11 +472,10 @@ void FRHICommandListExecutor::LatchBypass()
 				BypassVar->Set(0, ECVF_SetByCommandline);
 			}
 		}
-#endif
 
 		check((GRHICommandList.OutstandingCmdListCount.GetValue() == 1 && !GRHICommandList.GetImmediateCommandList().HasCommands()));
 		bool NewBypass = (CVarRHICmdBypass.GetValueOnAnyThread() >= 1);
-		NewBypass = true;
+
 		if (NewBypass && !bLatchedBypass)
 		{
 			FRHIResource::FlushPendingDeletes();
diff --git a/Engine/Source/Runtime/RHI/Public/RHICommandList.h b/Engine/Source/Runtime/RHI/Public/RHICommandList.h
--- a/Engine/Source/Runtime/RHI/Public/RHICommandList.h
+++ b/Engine/Source/Runtime/RHI/Public/RHICommandList.h
@@ -2629,7 +2629,7 @@ class RHI_API FRHICommandListExecutor
 public:
 	enum
 	{
-		DefaultBypass = 1 //PLATFORM_RHITHREAD_DEFAULT_BYPASS
+		DefaultBypass = PLATFORM_RHITHREAD_DEFAULT_BYPASS
 	};
 	FRHICommandListExecutor()
 		: bLatchedBypass(!!DefaultBypass)
