UE4 - D3D12 - Fix for AsyncTexture2D

diff --git a/Engine/Source/Runtime/RHI/Public/RHICommandList.h b/Engine/Source/Runtime/RHI/Public/RHICommandList.h
--- a/Engine/Source/Runtime/RHI/Public/RHICommandList.h
+++ b/Engine/Source/Runtime/RHI/Public/RHICommandList.h
@@ -2198,11 +2198,6 @@ public:
 	
 	FORCEINLINE FTexture2DRHIRef AsyncCreateTexture2D(uint32 SizeX, uint32 SizeY, uint8 Format, uint32 NumMips, uint32 Flags, void** InitialMipData, uint32 NumInitialMips)
 	{
-		if(GRHIThread)
-		{
-			QUICK_SCOPE_CYCLE_COUNTER(STAT_RHIMETHOD_AsyncCreateTexture2D_WaitRHI);
-			ImmediateFlush(EImmediateFlushType::WaitForRHIThread);
-		}
 		return GDynamicRHI->RHIAsyncCreateTexture2D(SizeX, SizeY, Format, NumMips, Flags, InitialMipData, NumInitialMips);
 	}
 	
