UE4 - Fixed editor running parallel renderer for map builds.

diff --git a/Engine/Source/Runtime/RHI/Private/RHICommandList.cpp b/Engine/Source/Runtime/RHI/Private/RHICommandList.cpp
--- a/Engine/Source/Runtime/RHI/Private/RHICommandList.cpp
+++ b/Engine/Source/Runtime/RHI/Private/RHICommandList.cpp
@@ -1381,7 +1381,7 @@ void FRHICommandListBase::WaitForRHIThreadTasks()
 DECLARE_CYCLE_STAT(TEXT("RTTask completion join"), STAT_HandleRTThreadTaskCompletion_Join, STATGROUP_RHICMDLIST);
 void FRHICommandListBase::HandleRTThreadTaskCompletion(const FGraphEventRef& MyCompletionGraphEvent)
 {
-	check(!IsImmediate() && !IsInRenderingThread() && !IsInRHIThread())
+	check(!IsImmediate() && !IsInActualRenderingThread() && !IsInRHIThread())
 	for (int32 Index = 0; Index < RTTasks.Num(); Index++)
 	{
 		if (!RTTasks[Index]->IsComplete())
