Fix copy/paste bug in FClearValueBinding compare.

diff --git a/Engine/Source/Runtime/RHI/Public/RHI.h b/Engine/Source/Runtime/RHI/Public/RHI.h
--- a/Engine/Source/Runtime/RHI/Public/RHI.h
+++ b/Engine/Source/Runtime/RHI/Public/RHI.h
@@ -829,8 +829,8 @@ struct FClearValueBinding
 			if (ColorBinding == EClearBinding::EDepthStencilBound)
 			{
 				return
-					Value.DSValue.Depth == Value.DSValue.Depth &&
-					Value.DSValue.Stencil == Value.DSValue.Stencil;
+					Value.DSValue.Depth == Other.Value.DSValue.Depth &&
+					Value.DSValue.Stencil == Other.Value.DSValue.Stencil;
 			}
 			return true;
 		}
