Rename shader platform SP_OPENGL_ES2 to SP_OPENGL_ES2_ANDROID

diff --git a/Engine/Source/Developer/HTML5/HTML5TargetPlatform/Private/HTML5TargetPlatform.cpp b/Engine/Source/Developer/HTML5/HTML5TargetPlatform/Private/HTML5TargetPlatform.cpp
--- a/Engine/Source/Developer/HTML5/HTML5TargetPlatform/Private/HTML5TargetPlatform.cpp
+++ b/Engine/Source/Developer/HTML5/HTML5TargetPlatform/Private/HTML5TargetPlatform.cpp
@@ -89,11 +89,11 @@ bool FHTML5TargetPlatform::IsRunningPlatform( ) const
 
 #if WITH_ENGINE
 
-static FName NAME_OPENGL_ES2_WEBGL(TEXT("GLSL_ES2_WEBGL"));
+static FName NAME_GLSL_ES2_WEBGL(TEXT("GLSL_ES2_WEBGL"));
 
 void FHTML5TargetPlatform::GetAllPossibleShaderFormats( TArray<FName>& OutFormats ) const
 {
-	OutFormats.AddUnique(NAME_OPENGL_ES2_WEBGL);
+	OutFormats.AddUnique(NAME_GLSL_ES2_WEBGL);
 }
 
 
diff --git a/Engine/Source/Developer/iOS/IOSTargetPlatform/Private/IOSTargetPlatform.cpp b/Engine/Source/Developer/iOS/IOSTargetPlatform/Private/IOSTargetPlatform.cpp
--- a/Engine/Source/Developer/iOS/IOSTargetPlatform/Private/IOSTargetPlatform.cpp
+++ b/Engine/Source/Developer/iOS/IOSTargetPlatform/Private/IOSTargetPlatform.cpp
@@ -361,13 +361,13 @@ bool FIOSTargetPlatform::SupportsFeature( ETargetPlatformFeatures Feature ) cons
 
 void FIOSTargetPlatform::GetAllPossibleShaderFormats( TArray<FName>& OutFormats ) const
 {
-	static FName NAME_OPENGL_ES2_IOS(TEXT("GLSL_ES2_IOS"));
+	static FName NAME_GLSL_ES2_IOS(TEXT("GLSL_ES2_IOS"));
 	static FName NAME_SF_METAL(TEXT("SF_METAL"));
 	static FName NAME_SF_METAL_MRT(TEXT("SF_METAL_MRT"));
 
 	if (SupportsES2())
 	{
-		OutFormats.AddUnique(NAME_OPENGL_ES2_IOS);
+		OutFormats.AddUnique(NAME_GLSL_ES2_IOS);
 	}
 
 	if (SupportsMetal())
diff --git a/Engine/Source/Runtime/Engine/Public/SimpleElementShaders.h b/Engine/Source/Runtime/Engine/Public/SimpleElementShaders.h
--- a/Engine/Source/Runtime/Engine/Public/SimpleElementShaders.h
+++ b/Engine/Source/Runtime/Engine/Public/SimpleElementShaders.h
@@ -346,7 +346,7 @@ public:
 
 	static bool ShouldCache(EShaderPlatform Platform)
 	{
-		return ((IsES2Platform(Platform) && IsPCPlatform(Platform)) || Platform == SP_OPENGL_ES2) && TSimpleElementBase::ShouldCache(Platform);
+		return ((IsES2Platform(Platform) && IsPCPlatform(Platform)) || Platform == SP_OPENGL_ES2_ANDROID) && TSimpleElementBase::ShouldCache(Platform);
 	}
 };
 
diff --git a/Engine/Source/Runtime/NullDrv/Private/NullRHI.cpp b/Engine/Source/Runtime/NullDrv/Private/NullRHI.cpp
--- a/Engine/Source/Runtime/NullDrv/Private/NullRHI.cpp
+++ b/Engine/Source/Runtime/NullDrv/Private/NullRHI.cpp
@@ -21,7 +21,7 @@ void FNullDynamicRHI::Init()
 	GShaderPlatformForFeatureLevel[ERHIFeatureLevel::SM4] = SP_PCD3D_SM4;
 	GShaderPlatformForFeatureLevel[ERHIFeatureLevel::SM5] = SP_PCD3D_SM5;
 #else
-	GShaderPlatformForFeatureLevel[ERHIFeatureLevel::ES2] = SP_OPENGL_ES2;
+	GShaderPlatformForFeatureLevel[ERHIFeatureLevel::ES2] = SP_OPENGL_PCES2;
 	GShaderPlatformForFeatureLevel[ERHIFeatureLevel::ES3_1] = SP_NumPlatforms;
 	GShaderPlatformForFeatureLevel[ERHIFeatureLevel::SM4] = SP_OPENGL_SM4;
 	GShaderPlatformForFeatureLevel[ERHIFeatureLevel::SM5] = SP_OPENGL_SM5;
diff --git a/Engine/Source/Runtime/OpenGLDrv/Private/Android/AndroidES31OpenGL.h b/Engine/Source/Runtime/OpenGLDrv/Private/Android/AndroidES31OpenGL.h
--- a/Engine/Source/Runtime/OpenGLDrv/Private/Android/AndroidES31OpenGL.h
+++ b/Engine/Source/Runtime/OpenGLDrv/Private/Android/AndroidES31OpenGL.h
@@ -254,6 +254,11 @@ ENUM_GL_ENTRYPOINTS_ALL(DECLARE_GL_ENTRYPOINTS);
 
 struct FAndroidES31OpenGL : public FOpenGLES31
 {
+	static FORCEINLINE EShaderPlatform GetShaderPlatform()
+	{
+		return bES2Fallback ? SP_OPENGL_ES2_ANDROID : SP_OPENGL_ES31_EXT;
+	}
+
 	static FORCEINLINE void InitDebugContext()
 	{
 		bDebugContext = glIsEnabled( GL_DEBUG_OUTPUT_KHR) != GL_FALSE;
diff --git a/Engine/Source/Runtime/OpenGLDrv/Private/Android/AndroidOpenGL.h b/Engine/Source/Runtime/OpenGLDrv/Private/Android/AndroidOpenGL.h
--- a/Engine/Source/Runtime/OpenGLDrv/Private/Android/AndroidOpenGL.h
+++ b/Engine/Source/Runtime/OpenGLDrv/Private/Android/AndroidOpenGL.h
@@ -126,6 +126,11 @@ extern "C"
 
 struct FAndroidOpenGL : public FOpenGLES2
 {
+	static FORCEINLINE EShaderPlatform GetShaderPlatform()
+	{
+		return SP_OPENGL_ES2_ANDROID;
+	}
+
 	// Optional:
 	static FORCEINLINE void QueryTimestampCounter(GLuint QueryID)
 	{
diff --git a/Engine/Source/Runtime/OpenGLDrv/Private/OpenGLDrvPrivate.h b/Engine/Source/Runtime/OpenGLDrv/Private/OpenGLDrvPrivate.h
--- a/Engine/Source/Runtime/OpenGLDrv/Private/OpenGLDrvPrivate.h
+++ b/Engine/Source/Runtime/OpenGLDrv/Private/OpenGLDrvPrivate.h
@@ -510,7 +510,7 @@ inline bool OpenGLShaderPlatformNeedsBindLocation(const EShaderPlatform InShader
 		case SP_OPENGL_SM4:
 		case SP_OPENGL_SM4_MAC:
 		case SP_OPENGL_PCES2:
-		case SP_OPENGL_ES2:
+		case SP_OPENGL_ES2_ANDROID:
 		case SP_OPENGL_ES2_WEBGL:
 		case SP_OPENGL_ES2_IOS:
 			return true;
@@ -533,7 +533,7 @@ inline bool OpenGLShaderPlatformSeparable(const EShaderPlatform InShaderPlatform
 			return true;
 
 		case SP_OPENGL_ES31_EXT:
-		case SP_OPENGL_ES2:
+		case SP_OPENGL_ES2_ANDROID:
 		case SP_OPENGL_ES2_WEBGL:
 		case SP_OPENGL_ES2_IOS:
 			return false;
diff --git a/Engine/Source/Runtime/OpenGLDrv/Private/Windows/OpenGLWindows.h b/Engine/Source/Runtime/OpenGLDrv/Private/Windows/OpenGLWindows.h
--- a/Engine/Source/Runtime/OpenGLDrv/Private/Windows/OpenGLWindows.h
+++ b/Engine/Source/Runtime/OpenGLDrv/Private/Windows/OpenGLWindows.h
@@ -577,6 +577,11 @@ struct FWindowsOpenGL : public FOpenGL4
 
 struct FWindowsOpenGL : public FOpenGLES31
 {
+	static FORCEINLINE EShaderPlatform GetShaderPlatform()
+	{
+		return bES2Fallback ? SP_OPENGL_PCES2 : SP_OPENGL_PCES3_1;
+	}
+
 	static FORCEINLINE void InitDebugContext()
 	{
 		bDebugContext = glIsEnabled( GL_DEBUG_OUTPUT) != GL_FALSE;
diff --git a/Engine/Source/Runtime/OpenGLDrv/Public/OpenGLES2.h b/Engine/Source/Runtime/OpenGLDrv/Public/OpenGLES2.h
--- a/Engine/Source/Runtime/OpenGLDrv/Public/OpenGLES2.h
+++ b/Engine/Source/Runtime/OpenGLDrv/Public/OpenGLES2.h
@@ -296,11 +296,6 @@ struct FOpenGLES2 : public FOpenGLBase
 		return ERHIFeatureLevel::ES2;
 	}
 
-	static FORCEINLINE EShaderPlatform GetShaderPlatform()
-	{
-		return SP_OPENGL_ES2;
-	}
-
 	static FORCEINLINE FString GetAdapterName()
 	{
 		return (TCHAR*)ANSI_TO_TCHAR((const ANSICHAR*)glGetString(GL_RENDERER));
diff --git a/Engine/Source/Runtime/OpenGLDrv/Public/OpenGLES31.h b/Engine/Source/Runtime/OpenGLDrv/Public/OpenGLES31.h
--- a/Engine/Source/Runtime/OpenGLDrv/Public/OpenGLES31.h
+++ b/Engine/Source/Runtime/OpenGLDrv/Public/OpenGLES31.h
@@ -804,12 +804,6 @@ struct FOpenGLES31 : public FOpenGLBase
 		return bES2Fallback ? ERHIFeatureLevel::ES2 : ERHIFeatureLevel::SM5;
 	}
 
-	static FORCEINLINE EShaderPlatform GetShaderPlatform()
-	{
-		// Should this support a commandline forced ES2?
-		return bES2Fallback ? SP_OPENGL_ES2 : SP_OPENGL_ES31_EXT;
-	}
-
 	static FORCEINLINE FString GetAdapterName()
 	{
 		return ANSI_TO_TCHAR((const ANSICHAR*)glGetString(GL_RENDERER));
diff --git a/Engine/Source/Runtime/RHI/Private/RHI.cpp b/Engine/Source/Runtime/RHI/Private/RHI.cpp
--- a/Engine/Source/Runtime/RHI/Private/RHI.cpp
+++ b/Engine/Source/Runtime/RHI/Private/RHI.cpp
@@ -303,12 +303,12 @@ static FName NAME_GLSL_150_MAC(TEXT("GLSL_150_MAC"));
 static FName NAME_SF_PS4(TEXT("SF_PS4"));
 static FName NAME_SF_XBOXONE(TEXT("SF_XBOXONE"));
 static FName NAME_GLSL_430(TEXT("GLSL_430"));
-static FName NAME_OPENGL_150_ES2(TEXT("GLSL_150_ES2"));
-static FName NAME_OPENGL_150_ES2_NOUB(TEXT("GLSL_150_ES2_NOUB"));
-static FName NAME_OPENGL_150_ES3_1(TEXT("GLSL_150_ES31"));
-static FName NAME_OPENGL_ES2(TEXT("GLSL_ES2"));
-static FName NAME_OPENGL_ES2_WEBGL(TEXT("GLSL_ES2_WEBGL"));
-static FName NAME_OPENGL_ES2_IOS(TEXT("GLSL_ES2_IOS"));
+static FName NAME_GLSL_150_ES2(TEXT("GLSL_150_ES2"));
+static FName NAME_GLSL_150_ES2_NOUB(TEXT("GLSL_150_ES2_NOUB"));
+static FName NAME_GLSL_150_ES31(TEXT("GLSL_150_ES31"));
+static FName NAME_GLSL_ES2(TEXT("GLSL_ES2"));
+static FName NAME_GLSL_ES2_WEBGL(TEXT("GLSL_ES2_WEBGL"));
+static FName NAME_GLSL_ES2_IOS(TEXT("GLSL_ES2_IOS"));
 static FName NAME_SF_METAL(TEXT("SF_METAL"));
 static FName NAME_SF_METAL_MRT(TEXT("SF_METAL_MRT"));
 static FName NAME_GLSL_310_ES_EXT(TEXT("GLSL_310_ES_EXT"));
@@ -341,16 +341,16 @@ FName LegacyShaderPlatformToShaderFormat(EShaderPlatform Platform)
 	case SP_OPENGL_PCES2:
 	{
 		static auto* CVar = IConsoleManager::Get().FindTConsoleVariableDataInt(TEXT("OpenGL.UseEmulatedUBs"));
-		return (CVar && CVar->GetValueOnAnyThread() != 0) ? NAME_OPENGL_150_ES2_NOUB : NAME_OPENGL_150_ES2;
+		return (CVar && CVar->GetValueOnAnyThread() != 0) ? NAME_GLSL_150_ES2_NOUB : NAME_GLSL_150_ES2;
 	}
 	case SP_OPENGL_PCES3_1:
-		return NAME_OPENGL_150_ES3_1;
-	case SP_OPENGL_ES2:
-		return NAME_OPENGL_ES2;
+		return NAME_GLSL_150_ES31;
+	case SP_OPENGL_ES2_ANDROID:
+		return NAME_GLSL_ES2;
 	case SP_OPENGL_ES2_WEBGL:
-		return NAME_OPENGL_ES2_WEBGL;
+		return NAME_GLSL_ES2_WEBGL;
 	case SP_OPENGL_ES2_IOS:
-		return NAME_OPENGL_ES2_IOS;
+		return NAME_GLSL_ES2_IOS;
 	case SP_METAL:
 		return NAME_SF_METAL;
 	case SP_METAL_MRT:
@@ -381,12 +381,12 @@ EShaderPlatform ShaderFormatToLegacyShaderPlatform(FName ShaderFormat)
 	if (ShaderFormat == NAME_SF_PS4)				return SP_PS4;
 	if (ShaderFormat == NAME_SF_XBOXONE)			return SP_XBOXONE;
 	if (ShaderFormat == NAME_GLSL_430)			return SP_OPENGL_SM5;
-	if (ShaderFormat == NAME_OPENGL_150_ES2 || ShaderFormat == NAME_OPENGL_150_ES2_NOUB)
+	if (ShaderFormat == NAME_GLSL_150_ES2 || ShaderFormat == NAME_GLSL_150_ES2_NOUB)
 												return SP_OPENGL_PCES2;
-	if (ShaderFormat == NAME_OPENGL_150_ES3_1)	return SP_OPENGL_PCES3_1;
-	if (ShaderFormat == NAME_OPENGL_ES2)			return SP_OPENGL_ES2;
-	if (ShaderFormat == NAME_OPENGL_ES2_WEBGL)	return SP_OPENGL_ES2_WEBGL;
-	if (ShaderFormat == NAME_OPENGL_ES2_IOS)		return SP_OPENGL_ES2_IOS;
+	if (ShaderFormat == NAME_GLSL_150_ES31)		return SP_OPENGL_PCES3_1;
+	if (ShaderFormat == NAME_GLSL_ES2)			return SP_OPENGL_ES2_ANDROID;
+	if (ShaderFormat == NAME_GLSL_ES2_WEBGL)	return SP_OPENGL_ES2_WEBGL;
+	if (ShaderFormat == NAME_GLSL_ES2_IOS)		return SP_OPENGL_ES2_IOS;
 	if (ShaderFormat == NAME_SF_METAL)			return SP_METAL;
 	if (ShaderFormat == NAME_SF_METAL_MRT)		return SP_METAL_MRT;
 	if (ShaderFormat == NAME_GLSL_310_ES_EXT)	return SP_OPENGL_ES31_EXT;
diff --git a/Engine/Source/Runtime/RHI/Public/RHIDefinitions.h b/Engine/Source/Runtime/RHI/Public/RHIDefinitions.h
--- a/Engine/Source/Runtime/RHI/Public/RHIDefinitions.h
+++ b/Engine/Source/Runtime/RHI/Public/RHIDefinitions.h
@@ -46,7 +46,7 @@ enum EShaderPlatform
 	SP_OPENGL_SM5		= 6,
 	/** Used when running in Feature Level ES2 in D3D11. */
 	SP_PCD3D_ES2		= 7,
-	SP_OPENGL_ES2		= 8,
+	SP_OPENGL_ES2_ANDROID = 8,
 	SP_OPENGL_ES2_WEBGL = 9, 
 	SP_OPENGL_ES2_IOS	= 10,
 	SP_METAL			= 11,
@@ -605,7 +605,7 @@ inline bool IsPCPlatform(const EShaderPlatform Platform)
 /** Whether the shader platform corresponds to the ES2 feature level. */
 inline bool IsES2Platform(const EShaderPlatform Platform)
 {
-	return Platform == SP_PCD3D_ES2 || Platform == SP_OPENGL_PCES2 || Platform == SP_OPENGL_ES2 || Platform == SP_OPENGL_ES2_WEBGL || Platform == SP_OPENGL_ES2_IOS || Platform == SP_METAL || Platform == SP_VULKAN_ES2; 
+	return Platform == SP_PCD3D_ES2 || Platform == SP_OPENGL_PCES2 || Platform == SP_OPENGL_ES2_ANDROID || Platform == SP_OPENGL_ES2_WEBGL || Platform == SP_OPENGL_ES2_IOS || Platform == SP_METAL || Platform == SP_VULKAN_ES2;
 }
 
 /** Whether the shader platform corresponds to the ES2 feature level. */
@@ -617,7 +617,7 @@ inline bool IsMobilePlatform(const EShaderPlatform Platform)
 inline bool IsOpenGLPlatform(const EShaderPlatform Platform)
 {
 	return Platform == SP_OPENGL_SM4 || Platform == SP_OPENGL_SM4_MAC || Platform == SP_OPENGL_SM5 || Platform == SP_OPENGL_PCES2
-		|| Platform == SP_OPENGL_ES2 || Platform == SP_OPENGL_ES2_WEBGL || Platform == SP_OPENGL_ES2_IOS || Platform == SP_OPENGL_ES31_EXT || Platform == SP_VULKAN_ES2;
+		|| Platform == SP_OPENGL_ES2_ANDROID || Platform == SP_OPENGL_ES2_WEBGL || Platform == SP_OPENGL_ES2_IOS || Platform == SP_OPENGL_ES31_EXT || Platform == SP_VULKAN_ES2;
 }
 
 inline bool IsMetalPlatform(const EShaderPlatform Platform)
@@ -649,7 +649,7 @@ inline ERHIFeatureLevel::Type GetMaxSupportedFeatureLevel(EShaderPlatform InShad
 		return ERHIFeatureLevel::SM4;
 	case SP_PCD3D_ES2:
 	case SP_OPENGL_PCES2:
-	case SP_OPENGL_ES2:
+	case SP_OPENGL_ES2_ANDROID:
 	case SP_OPENGL_ES2_WEBGL:
 	case SP_OPENGL_ES2_IOS:
 	case SP_VULKAN_ES2:
@@ -681,7 +681,7 @@ inline bool IsFeatureLevelSupported(EShaderPlatform InShaderPlatform, ERHIFeatur
 		return InFeatureLevel <= ERHIFeatureLevel::ES2;
 	case SP_OPENGL_PCES3_1:
 		return InFeatureLevel <= ERHIFeatureLevel::ES3_1;
-	case SP_OPENGL_ES2:
+	case SP_OPENGL_ES2_ANDROID:
 	case SP_VULKAN_ES2:
 		return InFeatureLevel <= ERHIFeatureLevel::ES2;
 	case SP_OPENGL_ES2_WEBGL: 
@@ -762,7 +762,7 @@ inline bool RHISupportsGeometryShaders(const EShaderPlatform Platform)
 
 inline bool RHIHasTiledGPU(const EShaderPlatform Platform)
 {
-	return (Platform == SP_METAL_MRT) || Platform == SP_METAL || Platform == SP_OPENGL_ES2_IOS || Platform == SP_OPENGL_ES2;
+	return (Platform == SP_METAL_MRT) || Platform == SP_METAL || Platform == SP_OPENGL_ES2_IOS || Platform == SP_OPENGL_ES2_ANDROID;
 }
 
 inline uint32 GetFeatureLevelMaxTextureSamplers(ERHIFeatureLevel::Type FeatureLevel)
diff --git a/Engine/Source/Runtime/ShaderCore/Private/ShaderCache.cpp b/Engine/Source/Runtime/ShaderCore/Private/ShaderCache.cpp
--- a/Engine/Source/Runtime/ShaderCore/Private/ShaderCache.cpp
+++ b/Engine/Source/Runtime/ShaderCore/Private/ShaderCache.cpp
@@ -94,7 +94,7 @@ static bool ShaderPlatformCanPrebindBoundShaderState(EShaderPlatform Platform)
 		case SP_OPENGL_SM4:
 		case SP_OPENGL_PCES2:
 		case SP_OPENGL_SM5:
-		case SP_OPENGL_ES2:
+		case SP_OPENGL_ES2_ANDROID:
 		case SP_OPENGL_ES2_WEBGL:
 		case SP_OPENGL_ES2_IOS:
 		case SP_OPENGL_ES31_EXT:
