Fix asserts on PS4

diff --git a/Engine/Source/Runtime/RHI/Public/RHICommandList.h b/Engine/Source/Runtime/RHI/Public/RHICommandList.h
--- a/Engine/Source/Runtime/RHI/Public/RHICommandList.h
+++ b/Engine/Source/Runtime/RHI/Public/RHICommandList.h
@@ -1356,7 +1356,7 @@ public:
 		}
 		else
 		{
-			check(Contents && ContentsSize && (&Layout == nullptr));
+			check(Contents && ContentsSize && (&Layout != nullptr));
 			auto* Cmd = new (AllocCommand<FRHICommandBuildLocalUniformBuffer>()) FRHICommandBuildLocalUniformBuffer(this, Contents, ContentsSize, Layout);
 			Result.WorkArea = &Cmd->WorkArea;
 		}
