Attempt to fix XB1 compile

diff --git a/Engine/Source/Runtime/RHI/Public/GPUProfiler.h b/Engine/Source/Runtime/RHI/Public/GPUProfiler.h
--- a/Engine/Source/Runtime/RHI/Public/GPUProfiler.h
+++ b/Engine/Source/Runtime/RHI/Public/GPUProfiler.h
@@ -235,23 +235,14 @@ struct RHI_API FGPUProfiler
 class FWindowedGPUTimer
 {
 public:
-	FWindowedGPUTimer(FRHICommandListImmediate& RHICmdList)
-		:FWindowedGPUTimer(10, 2, RHICmdList)
+	explicit FWindowedGPUTimer(FRHICommandListImmediate& RHICmdList)
 	{
+		PrivateInit(10,2,RHICmdList);
 	}
 
 	FWindowedGPUTimer(uint32 BufferSize, uint32 FramesBehind, FRHICommandListImmediate& RHICmdList)
-		: QueriesFailed(0)
-		, QueriesFinished(0)
 	{
-		StartQueries.AddZeroed(BufferSize);
-		EndQueries.AddZeroed(BufferSize);
-		for (uint32 i = 0; i < BufferSize; i++)
-		{
-			StartQueries[i] = RHICmdList.CreateRenderQuery(RQT_AbsoluteTime);;
-			EndQueries[i] = RHICmdList.CreateRenderQuery(RQT_AbsoluteTime);;
-		}
-		WindowSize = BufferSize - FramesBehind;
+		PrivateInit(BufferSize,FramesBehind,RHICmdList);
 	}
 
 	void Begin(FRHICommandListImmediate& RHICmdList)
@@ -283,6 +274,20 @@ public:
 	float GetElapsedAverage(FRHICommandListImmediate& RHICmdList, float &OutAvgTimeInSeconds);
 
 private:
+	void PrivateInit(uint32 BufferSize, uint32 FramesBehind, FRHICommandListImmediate& RHICmdList)
+	{
+		QueriesFailed = 0;
+		QueriesFinished = 0;
+		StartQueries.AddZeroed(BufferSize);
+		EndQueries.AddZeroed(BufferSize);
+		for (uint32 i = 0; i < BufferSize; i++)
+		{
+			StartQueries[i] = RHICmdList.CreateRenderQuery(RQT_AbsoluteTime);;
+			EndQueries[i] = RHICmdList.CreateRenderQuery(RQT_AbsoluteTime);;
+		}
+		WindowSize = BufferSize - FramesBehind;
+	}
+
 	int32 QueriesFailed;
 	uint32 WindowSize;
 	int32 QueriesFinished;
