Added comment to AppleDynamicRHI.cpp

diff --git a/Engine/Source/Runtime/RHI/Private/Apple/AppleDynamicRHI.cpp b/Engine/Source/Runtime/RHI/Private/Apple/AppleDynamicRHI.cpp
--- a/Engine/Source/Runtime/RHI/Private/Apple/AppleDynamicRHI.cpp
+++ b/Engine/Source/Runtime/RHI/Private/Apple/AppleDynamicRHI.cpp
@@ -9,6 +9,7 @@ FDynamicRHI* PlatformCreateDynamicRHI()
 	IDynamicRHIModule* DynamicRHIModule = NULL;
 
 	// Load the dynamic RHI module.
+	// Use Metal by default on OS versions and devices that support it
 	if (FPlatformMisc::HasPlatformFeature(TEXT("Metal")))
 	{
 		DynamicRHIModule = &FModuleManager::LoadModuleChecked<IDynamicRHIModule>(TEXT("MetalRHI"));
