Fix for compile error on WinXP due to use of 64-bit atomics in common code.

A couple of different options for fixing this, I chose the version with the least amount of code in an #ifdef.

diff --git a/Engine/Source/Runtime/RHI/Public/GPUDefragAllocator.h b/Engine/Source/Runtime/RHI/Public/GPUDefragAllocator.h
--- a/Engine/Source/Runtime/RHI/Public/GPUDefragAllocator.h
+++ b/Engine/Source/Runtime/RHI/Public/GPUDefragAllocator.h
@@ -981,11 +981,18 @@ protected:
 	/** Cumulative time spent in allocator.							*/
 	double			TimeSpentInAllocator;
 	/** Allocated memory in uint8s.									*/
-	volatile int64	AllocatedMemorySize;
+#if WINVER < 0x0600
+	// Interlock...64 functions are only available from Vista onwards
+	typedef int32 memsize_t;
+#else
+	typedef int64 memsize_t;
+#endif
+	volatile memsize_t	AllocatedMemorySize;
 	/** Available memory in uint8s.									*/
-	volatile int64	AvailableMemorySize;
+	volatile memsize_t	AvailableMemorySize;
 	/** Adjustment to allocated memory, pending all reallocations.	*/
-	volatile int64	PendingMemoryAdjustment;
+	volatile memsize_t	PendingMemoryAdjustment;
+
 	/** Mapping from Pointer to chunk for fast removal.				*/
 	TMap<void*, FMemoryChunk*>	PointerToChunkMap;
 
