Copying //UE4/Orion-Staging to //UE4/Main (originated from //Orion/Dev-General @ 2831630)

diff --git a/Engine/Source/Runtime/Online/OnlineSubsystem/Private/OnlineError.cpp b/Engine/Source/Runtime/Online/OnlineSubsystem/Private/OnlineError.cpp
new file mode 100644
--- /dev/null
+++ b/Engine/Source/Runtime/Online/OnlineSubsystem/Private/OnlineError.cpp
@@ -0,0 +1,77 @@
+// Copyright 1998-2016 Epic Games, Inc. All Rights Reserved.
+
+#include "OnlineSubsystemPrivatePCH.h"
+#include "OnlineError.h"
+#include "OnlineSubsystemTypes.h"
+
+//static
+const FString FOnlineError::GenericErrorCode = TEXT("GenericError");
+
+FOnlineError::FOnlineError()
+	: bSucceeded(false)
+	, HttpResult(0)
+{
+}
+
+FOnlineError::FOnlineError(bool bSucceededIn)
+	: bSucceeded(bSucceededIn)
+	, HttpResult(0)
+{
+}
+
+FOnlineError::FOnlineError(const FString& ErrorCodeIn)
+	: bSucceeded(false)
+	, HttpResult(0)
+{
+	SetFromErrorCode(ErrorCodeIn);
+}
+
+void FOnlineError::SetFromErrorCode(const FString& ErrorCodeIn)
+{
+	ErrorCode = ErrorCodeIn;
+	ErrorRaw = ErrorCodeIn;
+}
+
+FOnlineError::FOnlineError(const FText& ErrorMessageIn)
+	: bSucceeded(false)
+	, HttpResult(0)
+{
+	SetFromErrorMessage(ErrorMessageIn);
+}
+
+void FOnlineError::SetFromErrorMessage(const FText& ErrorMessageIn)
+{
+	ErrorMessage = ErrorMessageIn;
+	ErrorCode = FTextInspector::GetKey(ErrorMessageIn).Get(GenericErrorCode);
+	ErrorRaw = ErrorMessageIn.ToString();
+}
+
+EOnlineServerConnectionStatus::Type FOnlineError::GetConnectionStatusFromHttpResult() const
+{
+	if (bSucceeded)
+	{
+		return EOnlineServerConnectionStatus::Connected;
+	}
+
+	switch (HttpResult)
+	{
+	case 0:
+		// no response means we couldn't even connect
+		return EOnlineServerConnectionStatus::ConnectionDropped;
+	case 401:
+		// No auth at all
+		return EOnlineServerConnectionStatus::InvalidUser;
+	case 403:
+		// Auth failure
+		return EOnlineServerConnectionStatus::NotAuthorized;
+	case 503:
+		// service is not avail
+		return EOnlineServerConnectionStatus::ServersTooBusy;
+	case 505:
+		// update to supported version
+		return EOnlineServerConnectionStatus::UpdateRequired;
+	default:
+		// other bad responses (ELB, etc)
+		return EOnlineServerConnectionStatus::ServiceUnavailable;
+	}
+}
diff --git a/Engine/Source/Runtime/Online/OnlineSubsystem/Private/OnlineKeyValuePair.cpp b/Engine/Source/Runtime/Online/OnlineSubsystem/Private/OnlineKeyValuePair.cpp
--- a/Engine/Source/Runtime/Online/OnlineSubsystem/Private/OnlineKeyValuePair.cpp
+++ b/Engine/Source/Runtime/Online/OnlineSubsystem/Private/OnlineKeyValuePair.cpp
@@ -753,6 +753,10 @@ bool FVariantData::operator==(const FVariantData& Other) const
 				return (Value.AsBlob.BlobSize == Other.Value.AsBlob.BlobSize) && 
 						 (FMemory::Memcmp(Value.AsBlob.BlobData, Other.Value.AsBlob.BlobData, Value.AsBlob.BlobSize) == 0);
 			}
+		case EOnlineKeyValuePairDataType::Bool:
+			{
+				return Value.AsBool == Other.Value.AsBool;
+			}
 		}
 	}
 	return false;
diff --git a/Engine/Source/Runtime/Online/OnlineSubsystem/Private/OnlinePartyInterface.cpp b/Engine/Source/Runtime/Online/OnlineSubsystem/Private/OnlinePartyInterface.cpp
new file mode 100644
--- /dev/null
+++ b/Engine/Source/Runtime/Online/OnlineSubsystem/Private/OnlinePartyInterface.cpp
@@ -0,0 +1,96 @@
+// Copyright 1998-2016 Epic Games, Inc. All Rights Reserved.
+
+#include "OnlineSubsystemPrivatePCH.h"
+#include "OnlinePartyInterface.h"
+
+DEFINE_LOG_CATEGORY(LogOnlineParty);
+
+void FOnlinePartyData::ToJsonFull(FString& JsonString) const
+{
+	JsonString.Empty();
+
+	// iterate over key/val attrs and convert each entry to a json string
+	TSharedRef<FJsonObject> JsonObject(new FJsonObject());
+	TArray<TSharedPtr<FJsonValue> > JsonProperties;
+	for (auto Iterator : KeyValAttrs)
+	{
+		const FString& PropertyName = Iterator.Key;
+		const FVariantData& PropertyValue = Iterator.Value;
+
+		TSharedRef<FJsonObject> PropertyJson = PropertyValue.ToJson();
+		PropertyJson->SetStringField(TEXT("Name"), *PropertyName);
+		JsonProperties.Add(MakeShareable(new FJsonValueObject(PropertyJson)));
+	}
+	JsonObject->SetNumberField(TEXT("Rev"), RevisionCount);
+	JsonObject->SetArrayField(TEXT("Attrs"), JsonProperties);
+
+	auto JsonWriter = TJsonWriterFactory<TCHAR, TCondensedJsonPrintPolicy<TCHAR> >::Create(&JsonString);
+	FJsonSerializer::Serialize(JsonObject, JsonWriter);
+	JsonWriter->Close();
+}
+
+void FOnlinePartyData::ToJsonDirty(FString& JsonString) const
+{
+	JsonString.Empty();
+
+	// iterate over key/val attrs and convert each entry to a json string
+	TSharedRef<FJsonObject> JsonObject(new FJsonObject());
+	TArray<TSharedPtr<FJsonValue> > JsonProperties;
+	for (auto& PropertyName : DirtyKeys)
+	{
+		const FVariantData* PropertyValue = KeyValAttrs.Find(PropertyName);
+		ensure(PropertyValue);
+
+		TSharedRef<FJsonObject> PropertyJson = PropertyValue->ToJson();
+		PropertyJson->SetStringField(TEXT("Name"), *PropertyName);
+		JsonProperties.Add(MakeShareable(new FJsonValueObject(PropertyJson)));
+	}
+	JsonObject->SetNumberField(TEXT("Rev"), RevisionCount);
+	JsonObject->SetArrayField(TEXT("Attrs"), JsonProperties);
+
+	auto JsonWriter = TJsonWriterFactory<TCHAR, TCondensedJsonPrintPolicy<TCHAR> >::Create(&JsonString);
+	FJsonSerializer::Serialize(JsonObject, JsonWriter);
+	JsonWriter->Close();
+}
+
+void FOnlinePartyData::FromJson(const FString& JsonString)
+{
+	// json string to key/val attrs
+	TSharedPtr<FJsonObject> JsonObject;
+	TSharedRef<TJsonReader<> > JsonReader = TJsonReaderFactory<>::Create(JsonString);
+	if (FJsonSerializer::Deserialize(JsonReader, JsonObject) &&
+		JsonObject.IsValid())
+	{
+		if (JsonObject->HasTypedField<EJson::Array>(TEXT("Attrs")))
+		{
+			const TArray<TSharedPtr<FJsonValue> >& JsonProperties = JsonObject->GetArrayField(TEXT("Attrs"));
+			for (auto JsonPropertyValue : JsonProperties)
+			{
+				TSharedPtr<FJsonObject> JsonPropertyObject = JsonPropertyValue->AsObject();
+				if (JsonPropertyObject.IsValid())
+				{
+					FString PropertyName;
+					if (JsonPropertyObject->TryGetStringField(TEXT("Name"), PropertyName) &&
+						!PropertyName.IsEmpty())
+					{
+						FVariantData PropertyData;
+						if (PropertyData.FromJson(JsonPropertyObject.ToSharedRef()))
+						{
+							KeyValAttrs.Add(PropertyName, PropertyData);
+						}
+					}
+				}
+			}
+		}
+
+		if (JsonObject->HasTypedField<EJson::Number>(TEXT("Rev")))
+		{
+			int32 NewRevisionCount = JsonObject->GetIntegerField(TEXT("Rev"));
+			if ((RevisionCount != 0) && (NewRevisionCount != RevisionCount) && (NewRevisionCount != (RevisionCount + 1)))
+			{
+				UE_LOG_ONLINEPARTY(Warning, TEXT("Unexpected revision received.  Current %d, new %d"), RevisionCount, NewRevisionCount);
+			}
+			RevisionCount = NewRevisionCount;
+		}
+	}
+}
diff --git a/Engine/Source/Runtime/Online/OnlineSubsystem/Private/OnlineSubsystem.cpp b/Engine/Source/Runtime/Online/OnlineSubsystem/Private/OnlineSubsystem.cpp
--- a/Engine/Source/Runtime/Online/OnlineSubsystem/Private/OnlineSubsystem.cpp
+++ b/Engine/Source/Runtime/Online/OnlineSubsystem/Private/OnlineSubsystem.cpp
@@ -10,7 +10,6 @@
 
 DEFINE_LOG_CATEGORY(LogOnline);
 DEFINE_LOG_CATEGORY(LogOnlineGame);
-DEFINE_LOG_CATEGORY(LogOnlineParty);
 DEFINE_LOG_CATEGORY(LogOnlineChat);
 
 #if STATS
diff --git a/Engine/Source/Runtime/Online/OnlineSubsystem/Public/Interfaces/IMessageSanitizerInterface.h b/Engine/Source/Runtime/Online/OnlineSubsystem/Public/Interfaces/IMessageSanitizerInterface.h
--- a/Engine/Source/Runtime/Online/OnlineSubsystem/Public/Interfaces/IMessageSanitizerInterface.h
+++ b/Engine/Source/Runtime/Online/OnlineSubsystem/Public/Interfaces/IMessageSanitizerInterface.h
@@ -2,11 +2,12 @@
 
 #pragma once
 
-DECLARE_DELEGATE_OneParam(FOnMessageProcessed, const FString& /*SanitizedMessage*/);
+DECLARE_DELEGATE_TwoParams(FOnMessageProcessed, bool /*Process success*/, const FString& /*SanitizedMessage*/);
+DECLARE_DELEGATE_TwoParams(FOnMessageArrayProcessed, bool /*Process success*/, const TArray<FString>& /*SanitizedMessages*/);
 
 struct FSanitizeMessage
 {
-	FSanitizeMessage(FString InRawMessage, FOnMessageProcessed InProcessCompleteDelegate)
+	FSanitizeMessage(const FString& InRawMessage, FOnMessageProcessed InProcessCompleteDelegate)
 		: RawMessage(InRawMessage)
 		, CompletionDelegate(InProcessCompleteDelegate)
 	{}
@@ -15,6 +16,16 @@ struct FSanitizeMessage
 	FOnMessageProcessed CompletionDelegate;
 };
 
+struct FMultiPartMessage
+{
+	TArray<FString> AlreadyProcessedMessages;
+	TArray<int32> AlreadyProcessedIndex;
+	TArray<FString> SanitizedMessages;
+	FString MessageToSanatize;
+	bool bCompleted;
+	int32 PartNumber;
+};
+
 class IMessageSanitizer
 {
 protected:
@@ -22,7 +33,8 @@ protected:
 
 public:
 	virtual ~IMessageSanitizer() {};
-	virtual void SanitizeMessage(FString Message, FOnMessageProcessed CompletionDelegate) = 0;
+	virtual void SanitizeMessage(const FString& Message, FOnMessageProcessed CompletionDelegate) = 0;
+	virtual void SanitizeMessageArray(const TArray<FString>& MessageArray, FOnMessageArrayProcessed CompletionDelegate) = 0;
 };
 
 typedef TSharedPtr<IMessageSanitizer, ESPMode::ThreadSafe> IMessageSanitizerPtr;
diff --git a/Engine/Source/Runtime/Online/OnlineSubsystem/Public/Interfaces/OnlinePartyInterface.h b/Engine/Source/Runtime/Online/OnlineSubsystem/Public/Interfaces/OnlinePartyInterface.h
--- a/Engine/Source/Runtime/Online/OnlineSubsystem/Public/Interfaces/OnlinePartyInterface.h
+++ b/Engine/Source/Runtime/Online/OnlineSubsystem/Public/Interfaces/OnlinePartyInterface.h
@@ -8,6 +8,14 @@
 #define F_PREFIX(TypeToPrefix) F##TypeToPrefix
 #define PARTY_DECLARE_DELEGATETYPE(Type) typedef F##Type::FDelegate F##Type##Delegate
 
+ONLINESUBSYSTEM_API DECLARE_LOG_CATEGORY_EXTERN(LogOnlineParty, Display, All);
+
+#define ONLINEPARTY_LOG_PREFIX TEXT("OSS::Party: ")
+#define UE_LOG_ONLINEPARTY(Verbosity, Format, ...) \
+{ \
+	UE_LOG(LogOnlineParty, Verbosity, TEXT("%s%s"), ONLINEPARTY_LOG_PREFIX, *FString::Printf(Format, ##__VA_ARGS__)); \
+}
+
 
 /**
  * Party member user info returned by IOnlineParty interface
@@ -20,10 +28,15 @@ public:
 /**
  * Data associated with the entire party
  */
-class FOnlinePartyData : public TSharedFromThis<FOnlinePartyData>
+class ONLINESUBSYSTEM_API FOnlinePartyData : public TSharedFromThis<FOnlinePartyData>
 {
 public:
-	FOnlinePartyData() {}
+	FOnlinePartyData() :
+		TotalBytes(0),
+		TotalEffectiveBytes(0),
+		TotalPackets(0),
+		RevisionCount(0)
+		{}
 	~FOnlinePartyData() {}
 
 	/**
@@ -56,14 +69,93 @@ public:
 	void SetAttribute(const FString& AttrName, const FVariantData& AttrValue)
 	{
 		FVariantData& NewAttrValue = KeyValAttrs.FindOrAdd(AttrName);
-		NewAttrValue = AttrValue;
+		if (NewAttrValue != AttrValue)
+		{
+			NewAttrValue = AttrValue;
+			DirtyKeys.Add(AttrName);
+		}
+	}
+
+	/**
+	 * Clear the attributes map
+	 */
+	void ClearAttributes()
+	{
+		KeyValAttrs.Empty();
+		DirtyKeys.Empty();
+	}
+
+	/** 
+	 * Clear the dirty keys set, called after successfully sending an update of the dirty elements
+	 */
+	void ClearDirty()
+	{
+		DirtyKeys.Empty();
+	}
+
+	/** 
+	 * Increment the stat tracking variables on packet sent
+	 * 
+	 * @param PacketSize - size of the packet generated
+	 * @param NumRecipients - number of recipients the packet was sent to
+	 * @param bIncrementRevision - this packet was a dirty packet update so we should increment the revision
+	 */
+	void OnPacketSent(int32 PacketSize, int32 NumRecipients, bool bIncrementRevision) const
+	{
+		TotalPackets++;
+		TotalBytes += PacketSize;
+		TotalEffectiveBytes += PacketSize * NumRecipients;
+		if (bIncrementRevision)
+		{
+			++RevisionCount;
+		}
 	}
 
-	void ToJson(FString& JsonString) const;
+	/** 
+	 * Generate a JSON packet containing all key-value attributes
+	 * 
+	 * @param JsonString - [out] string containing the resulting JSON output
+	 * 
+	 */
+	void ToJsonFull(FString& JsonString) const;
+	
+	/** 
+	 * Generate a JSON packet containing only the dirty key-value attributes for a delta update
+	 *
+	 * @param JsonString - [out] string containing the resulting JSON output
+	 *
+	 */
+	void ToJsonDirty(FString& JsonString) const;
+
+	/** 
+	 * Update attributes from a JSON packet
+	 *
+	 * @param JsonString - string containing the JSON packet
+	 *
+	 */
 	void FromJson(const FString& JsonString);
 
+	/** Accessor functions for KeyValAttrs map */
+	FOnlineKeyValuePairs<FString, FVariantData>& GetKeyValAttrs() { return KeyValAttrs; }
+	const FOnlineKeyValuePairs<FString, FVariantData>& GetKeyValAttrs() const { return KeyValAttrs; }
+
+	/** Stat tracking variables */
+	/** Total number of bytes generated by calls to ToJsonFull and ToJsonDirty */
+	mutable int32 TotalBytes;
+	/** Total number of bytes generated by calls to ToJsonFull and ToJsonDirty, multiplied by the number of recipients the packet was sent to */
+	mutable int32 TotalEffectiveBytes;
+	/** Total number of packets generated by calls to ToJsonFull and ToJsonDirty */
+	mutable int32 TotalPackets;
+
+	/** Id representing number of updates sent, useful for determining if a client has missed an update */
+	mutable int32 RevisionCount;
+
+private:
 	/** map of key/val attributes that represents the data */
 	FOnlineKeyValuePairs<FString, FVariantData>  KeyValAttrs;
+
+	/** set of which fields are dirty and need to transmitted */
+	TSet<FString> DirtyKeys;
 };
 
 /**
@@ -1788,75 +1880,18 @@ inline FString ToDebugString(const IOnlinePartyJoinInfo& JoinInfo)
 inline FString ToDebugString(const FOnlinePartyData& PartyData)
 {
 	FString Result;
-	bool bAddComma = false;
-	for (auto Iterator : PartyData.KeyValAttrs)
-	{
-		if (bAddComma)
-		{
-			Result += TEXT(", ");
-		}
-		else
-		{
-			bAddComma = true;
-		}
-		Result += FString::Printf(TEXT("[%s=%s]"), *(Iterator.Key), *(Iterator.Value.ToString()));
-	}
-	return Result;
-}
+	
+	int32 TotalBytesPerSec = PartyData.TotalPackets ? (PartyData.TotalBytes / PartyData.TotalPackets) : 0;
+	int32 TotalEffectiveBytesPerSec = PartyData.TotalPackets ? (PartyData.TotalEffectiveBytes / PartyData.TotalPackets) : 0;
 
-inline void FOnlinePartyData::ToJson(FString& JsonString) const
-{
-	JsonString.Empty();
+	Result += FString::Printf(TEXT("%dB [%d B/pkt], %dB [%d B/pkt], Rev: %d"), 
+		PartyData.TotalBytes, TotalBytesPerSec,
+		PartyData.TotalEffectiveBytes, TotalEffectiveBytesPerSec, 
+		PartyData.RevisionCount);
 
-	// iterate over key/val attrs and convert each entry to a json string
-	TSharedRef<FJsonObject> JsonObject(new FJsonObject());
-	TArray<TSharedPtr<FJsonValue> > JsonProperties;
-	for (auto Iterator : KeyValAttrs)
+	for (auto Iterator : PartyData.GetKeyValAttrs())
 	{
-		const FString& PropertyName = Iterator.Key;
-		const FVariantData& PropertyValue = Iterator.Value;
-
-		TSharedRef<FJsonObject> PropertyJson = PropertyValue.ToJson();
-		PropertyJson->SetStringField(TEXT("Name"), *PropertyName);
-		JsonProperties.Add(MakeShareable(new FJsonValueObject(PropertyJson)));
-	}
-	JsonObject->SetArrayField(TEXT("Attrs"), JsonProperties);
-
-	auto JsonWriter = TJsonWriterFactory<TCHAR, TCondensedJsonPrintPolicy<TCHAR> >::Create(&JsonString);
-	FJsonSerializer::Serialize(JsonObject, JsonWriter);
-	JsonWriter->Close();
-}
-
-inline void FOnlinePartyData::FromJson(const FString& JsonString)
-{
-	KeyValAttrs.Empty();
-
-	// json string to key/val attrs
-	TSharedPtr<FJsonObject> JsonObject;
-	TSharedRef<TJsonReader<> > JsonReader = TJsonReaderFactory<>::Create(JsonString);
-	if (FJsonSerializer::Deserialize(JsonReader, JsonObject) &&
-		JsonObject.IsValid())
-	{
-		if (JsonObject->HasTypedField<EJson::Array>(TEXT("Attrs")))
-		{
-			const TArray<TSharedPtr<FJsonValue> >& JsonProperties = JsonObject->GetArrayField(TEXT("Attrs"));
-			for (auto JsonPropertyValue : JsonProperties)
-			{
-				TSharedPtr<FJsonObject> JsonPropertyObject = JsonPropertyValue->AsObject();
-				if (JsonPropertyObject.IsValid())
-				{
-					FString PropertyName;
-					if (JsonPropertyObject->TryGetStringField(TEXT("Name"), PropertyName) &&
-						!PropertyName.IsEmpty())
-					{
-						FVariantData PropertyData;
-						if (PropertyData.FromJson(JsonPropertyObject.ToSharedRef()))
-						{
-							KeyValAttrs.Add(PropertyName, PropertyData);
-						}
-					}
-				}
-			}
-		}
+		Result += FString::Printf(TEXT(",[%s=%s]"), *(Iterator.Key), *(Iterator.Value.ToString()));
 	}
+	return Result;
 }
diff --git a/Engine/Source/Runtime/Online/OnlineSubsystem/Public/Interfaces/OnlinePurchaseInterface.h b/Engine/Source/Runtime/Online/OnlineSubsystem/Public/Interfaces/OnlinePurchaseInterface.h
--- a/Engine/Source/Runtime/Online/OnlineSubsystem/Public/Interfaces/OnlinePurchaseInterface.h
+++ b/Engine/Source/Runtime/Online/OnlineSubsystem/Public/Interfaces/OnlinePurchaseInterface.h
@@ -4,6 +4,7 @@
 #include "OnlineDelegateMacros.h"
 #include "OnlineStoreInterfaceV2.h"
 #include "OnlineEntitlementsInterface.h"
+#include "OnlineError.h"
 
 /**
  * Info needed for checkout
@@ -69,6 +70,11 @@ namespace EPurchaseTransactionState
 class FPurchaseReceipt
 {
 public:
+	FPurchaseReceipt()
+	: TransactionState(EPurchaseTransactionState::NotStarted)
+	{
+	}
+
 	struct FLineItemInfo
 	{
 		/** The platform identifier of this purchase type */
@@ -79,6 +85,8 @@ public:
 
 		/** platform-specific opaque validation info (required to verify UniqueId belongs to this account) */
 		FString ValidationInfo;
+
+		inline bool IsRedeemable() const { return !ValidationInfo.IsEmpty(); }
 	};
 
 	/**
@@ -140,17 +148,17 @@ public:
 /**
  * Delegate called when checkout process completes
  */
-DECLARE_DELEGATE_FourParams(FOnPurchaseCheckoutComplete, const FUniqueNetId& /*UserId*/, bool /*bWasSuccessful*/, const TSharedRef<FPurchaseReceipt>& /*Receipt*/, const FString& /*Error*/);
+DECLARE_DELEGATE_TwoParams(FOnPurchaseCheckoutComplete, const FOnlineError& /*Result*/, const TSharedRef<FPurchaseReceipt>& /*Receipt*/);
 
 /**
  * Delegate called when code redemption process completes
  */
-DECLARE_DELEGATE_FourParams(FOnPurchaseRedeemCodeComplete, const FUniqueNetId& /*UserId*/, bool /*bWasSuccessful*/, const TSharedRef<FPurchaseReceipt>& /*Receipt*/, const FString& /*Error*/);
+DECLARE_DELEGATE_TwoParams(FOnPurchaseRedeemCodeComplete, const FOnlineError& /*Result*/, const TSharedRef<FPurchaseReceipt>& /*Receipt*/);
 
 /**
  * Delegate called when query receipt process completes
  */
-DECLARE_DELEGATE_ThreeParams(FOnQueryReceiptsComplete, const FUniqueNetId& /*UserId*/, bool /*bWasSuccessful*/, const FString& /*Error*/);
+DECLARE_DELEGATE_OneParam(FOnQueryReceiptsComplete, const FOnlineError& /*Result*/);
 
 /**
  *	IOnlinePurchase - Interface for IAP (In App Purchases) services
@@ -174,32 +182,26 @@ public:
 	 *
 	 * @param UserId user initiating the request
 	 * @param CheckoutRequest info needed for the checkout request
-	 * @param Delegate completion callback
-	 *
-	 * @return true if user can make a purchase
+	 * @param Delegate completion callback (guaranteed to be called)
 	 */
-	virtual bool Checkout(const FUniqueNetId& UserId, const FPurchaseCheckoutRequest& CheckoutRequest, const FOnPurchaseCheckoutComplete& Delegate = FOnPurchaseCheckoutComplete()) = 0;
+	virtual void Checkout(const FUniqueNetId& UserId, const FPurchaseCheckoutRequest& CheckoutRequest, const FOnPurchaseCheckoutComplete& Delegate) = 0;
 
 	/**
 	 * Initiate the checkout process for obtaining offers via code redemption
 	 *
 	 * @param UserId user initiating the request
 	 * @param RedeemCodeRequest info needed for the redeem request
-	 * @param Delegate completion callback
-	 *
-	 * @return true if user can make a purchase
+	 * @param Delegate completion callback (guaranteed to be called)
 	 */
-	virtual bool RedeemCode(const FUniqueNetId& UserId, const FRedeemCodeRequest& RedeemCodeRequest, const FOnPurchaseRedeemCodeComplete& Delegate = FOnPurchaseRedeemCodeComplete()) = 0;
+	virtual void RedeemCode(const FUniqueNetId& UserId, const FRedeemCodeRequest& RedeemCodeRequest, const FOnPurchaseRedeemCodeComplete& Delegate) = 0;
 
 	/**
 	 * Query for all of the user's receipts from prior purchases
 	 *
 	 * @param UserId user initiating the request
-	 * @param Delegate completion callback
-	 *
-	 * @return true if user can make a purchase
+	 * @param Delegate completion callback (guaranteed to be called)
 	 */
-	virtual bool QueryReceipts(const FUniqueNetId& UserId, const FOnQueryReceiptsComplete& Delegate = FOnQueryReceiptsComplete()) = 0;
+	virtual void QueryReceipts(const FUniqueNetId& UserId, const FOnQueryReceiptsComplete& Delegate) = 0;
 
 	/**
 	 * Get list of cached receipts for user (includes transactions currently being processed)
diff --git a/Engine/Source/Runtime/Online/OnlineSubsystem/Public/Interfaces/OnlineStoreInterfaceV2.h b/Engine/Source/Runtime/Online/OnlineSubsystem/Public/Interfaces/OnlineStoreInterfaceV2.h
--- a/Engine/Source/Runtime/Online/OnlineSubsystem/Public/Interfaces/OnlineStoreInterfaceV2.h
+++ b/Engine/Source/Runtime/Online/OnlineSubsystem/Public/Interfaces/OnlineStoreInterfaceV2.h
@@ -43,7 +43,7 @@ public:
 	/** @return FText suitable for localized display */
 	virtual FText GetDisplayPrice() const
 	{
-		return FText::AsCurrency(NumericPrice / 100.0, CurrencyCode);
+		return FText::AsCurrencyBase(NumericPrice, CurrencyCode);
 	}
 
 	/** @return True if offer can be purchased */
diff --git a/Engine/Source/Runtime/Online/OnlineSubsystem/Public/OnlineAsyncTaskManager.h b/Engine/Source/Runtime/Online/OnlineSubsystem/Public/OnlineAsyncTaskManager.h
--- a/Engine/Source/Runtime/Online/OnlineSubsystem/Public/OnlineAsyncTaskManager.h
+++ b/Engine/Source/Runtime/Online/OnlineSubsystem/Public/OnlineAsyncTaskManager.h
@@ -104,6 +104,11 @@ public:
 	{
 	}
 
+	/**
+	 * Initialize the task - called on game thread when queued
+	 */
+	virtual void Initialize() {}
+
 	/**
 	 * Check the state of the async task
 	 * @return true if complete, false otherwise
@@ -123,7 +128,7 @@ public:
 	virtual void Tick()
 	{
 		// assert that we're not on the game thread
-		check(!IsInGameThread());
+		check(!IsInGameThread() || !FPlatformProcess::SupportsMultithreading());
 	}
 };
 
@@ -224,7 +229,13 @@ protected:
 	TArray<FOnlineAsyncTask*> InQueue;
 	/** Critical section for thread safe operation of the event in queue */
 	FCriticalSection InQueueLock;
-	
+
+	/** blah */
+	FOnlineAsyncTask* ActiveTask;
+
+	/** blah */
+	FCriticalSection ActiveTaskLock;
+
 	/** This queue is for tasks that are safe to run in parallel with one another */
 	TArray<FOnlineAsyncTask*> ParallelTasks;
 	/** Critical section for thread safe operation of the list */
@@ -300,11 +311,6 @@ public:
 	 */
 	void AddToInQueue(FOnlineAsyncTask* NewTask);
 
-	/**
-	 *	Remove the current async task from the queue
-	 */
-	void PopFromInQueue();
-
 	/**
 	 * Add completed online async tasks that need processing onto the queue
 	 * @param CompletedItem - some finished request of the online services
diff --git a/Engine/Source/Runtime/Online/OnlineSubsystem/Public/OnlineError.h b/Engine/Source/Runtime/Online/OnlineSubsystem/Public/OnlineError.h
new file mode 100644
--- /dev/null
+++ b/Engine/Source/Runtime/Online/OnlineSubsystem/Public/OnlineError.h
@@ -0,0 +1,42 @@
+// Copyright 1998-2016 Epic Games, Inc. All Rights Reserved.
+
+#pragma once
+
+namespace EOnlineServerConnectionStatus { enum Type; }
+typedef TSharedPtr<class IHttpResponse, ESPMode::ThreadSafe> FHttpResponsePtr;
+
+/** Generic Error response for OSS calls */
+struct ONLINESUBSYSTEM_API FOnlineError
+{
+public:
+	FOnlineError();
+	explicit FOnlineError(bool bSucceeded);
+	explicit FOnlineError(const FString& ErrorCode);
+	explicit FOnlineError(const FText& ErrorMessage);
+
+	/** Same as the Ctors but can be called any time */
+	void SetFromErrorCode(const FString& ErrorCode);
+	void SetFromErrorMessage(const FText& ErrorMessage);
+
+	/** Converts the HttpResult into a EOnlineServerConnectionStatus */
+	EOnlineServerConnectionStatus::Type GetConnectionStatusFromHttpResult() const;
+
+	/** Code useful when all you have is raw error info from old APIs */
+	static const FString GenericErrorCode;
+
+public:
+	/** Did the request succeed fully. If this is true the rest of the struct probably doesn't matter */
+	bool bSucceeded;
+
+	/** The HTTP response code. Will be 0 if a connection error occurred or if HTTP was not used */
+	int32 HttpResult;
+
+	/** The raw unparsed error message from server. Used for pass-through error processing by other systems. */
+	FString ErrorRaw;
+
+	/** Intended to be interpreted by code. */
+	FString ErrorCode;
+
+	/** Suitable for display to end user. Guaranteed to be in the current locale (or empty) */
+	FText ErrorMessage;
+};
diff --git a/Engine/Source/Runtime/Online/OnlineSubsystem/Public/OnlineJsonSerializer.h b/Engine/Source/Runtime/Online/OnlineSubsystem/Public/OnlineJsonSerializer.h
--- a/Engine/Source/Runtime/Online/OnlineSubsystem/Public/OnlineJsonSerializer.h
+++ b/Engine/Source/Runtime/Online/OnlineSubsystem/Public/OnlineJsonSerializer.h
@@ -601,11 +601,11 @@ struct FOnlineJsonSerializable
 	 *
 	 * @return the corresponding json string
 	 */
-	inline const FString ToJson() const
+	inline const FString ToJson(bool bPrettyPrint = true) const
 	{
 		// Strip away const, because we use a single method that can read/write which requires non-const semantics
 		// Otherwise, we'd have to have 2 separate macros for declaring const to json and non-const from json
-		return ((FOnlineJsonSerializable*)this)->ToJson();
+		return ((FOnlineJsonSerializable*)this)->ToJson(bPrettyPrint);
 	}
 	/**
 	 * Serializes this object to its JSON string form
diff --git a/Engine/Source/Runtime/Online/OnlineSubsystem/Public/OnlineSubsystem.h b/Engine/Source/Runtime/Online/OnlineSubsystem/Public/OnlineSubsystem.h
--- a/Engine/Source/Runtime/Online/OnlineSubsystem/Public/OnlineSubsystem.h
+++ b/Engine/Source/Runtime/Online/OnlineSubsystem/Public/OnlineSubsystem.h
@@ -13,7 +13,6 @@
 
 ONLINESUBSYSTEM_API DECLARE_LOG_CATEGORY_EXTERN(LogOnline, Display, All);
 ONLINESUBSYSTEM_API DECLARE_LOG_CATEGORY_EXTERN(LogOnlineGame, Display, All);
-ONLINESUBSYSTEM_API DECLARE_LOG_CATEGORY_EXTERN(LogOnlineParty, Display, All);
 ONLINESUBSYSTEM_API DECLARE_LOG_CATEGORY_EXTERN(LogOnlineChat, Display, All);
 
 /** Online subsystem stats */
diff --git a/Engine/Source/Runtime/RHI/Private/DynamicRHI.cpp b/Engine/Source/Runtime/RHI/Private/DynamicRHI.cpp
--- a/Engine/Source/Runtime/RHI/Private/DynamicRHI.cpp
+++ b/Engine/Source/Runtime/RHI/Private/DynamicRHI.cpp
@@ -14,6 +14,18 @@
 // Globals.
 FDynamicRHI* GDynamicRHI = NULL;
 
+static TAutoConsoleVariable<int32> CVarDetectAndWarnOfBadDrivers(
+	TEXT("r.DetectAndWarnOfBadDrivers"),
+	1,
+	TEXT("On engine startup we can check the current GPU driver and warn the user about issues and suggest a specific version\n")
+	TEXT("The test is fast so this should not cost any performance.\n")
+	TEXT(" 0: off\n")
+	TEXT(" 1: a message on startup might appear (default)\n")
+	TEXT(" 2: Test by pretending the driver has issues"),
+	ECVF_RenderThreadSafe
+	);
+
+
 void InitNullRHI()
 {
 	// Use the null RHI if it was specified on the command line, or if a commandlet is running.
@@ -33,6 +45,120 @@ void InitNullRHI()
 	GRHISupportsTextureStreaming = false;
 }
 
+// e.g. 36143 for the NVIDIA driver 361.43
+// no floating point to avoid precision issues
+// @return -1 if the version is unknown
+static int32 ExtractGPUDriverVersionInt()
+{
+	int32 Ret = -1;
+
+	// we use the internal version, not the user version to avoid problem where the name was altered 
+	const FString& InternalVersion = GRHIAdapterInternalDriverVersion;
+
+	if(IsRHIDeviceNVIDIA())
+	{
+		if(InternalVersion.Left(6) == TEXT("9.18.1"))		// e.g. "9.18.136.143"
+		{
+			// e.g. 36.143
+			FString RightPart = InternalVersion.RightChop(6);
+			// e.g. 36143
+			FString CompactNumber = RightPart.Replace(TEXT("."), TEXT(""));
+
+			if(FCString::IsNumeric(*CompactNumber))
+			{
+				Ret = FCString::Atoi(*CompactNumber);
+			}
+		}
+	}
+	else
+	{
+		// examples for AMD: "13.12" "15.101.1007" "13.351"
+		// todo: need to decide on what part we extract, maybe make the return argument multiple ints
+	}
+	return Ret;
+}
+
+void RHIDetectAndWarnOfBadDrivers()
+{
+	int32 CVarValue = CVarDetectAndWarnOfBadDrivers.GetValueOnGameThread();
+
+	if(!GIsRHIInitialized || !CVarValue || GRHIVendorId == 0)
+	{
+		return;
+	}
+
+	// for testing
+	if(CVarValue == 2)
+	{
+		GRHIVendorId = 0x10DE;
+		GRHIAdapterName = TEXT("Generic NVIDIA");
+		GRHIAdapterUserDriverVersion = TEXT("361.43");
+		GRHIAdapterInternalDriverVersion = TEXT("9.18.136.143");
+		GRHIAdapterDriverDate = TEXT("01-01-1900");
+	}
+
+	int32 VersionInt = ExtractGPUDriverVersionInt();
+
+	if(VersionInt < 0)
+	{
+		// no valid version
+		return;
+	}
+
+	// can be made data driven:
+	if(IsRHIDeviceNVIDIA())
+	{
+		// todo: make this data driven
+		FString SuggestedVersion = TEXT("359.06");
+
+		// todo: make this data driven
+		if(VersionInt == 36143	// 361.43 "UE-25096 Viewport flashes black and white when moving in the scene on latest Nvidia drivers", at that time that was the latest driver
+		|| VersionInt <= 34709)	// 347.09 "old NVIDIA driver version we've seen driver crashes with Paragon content", we might have to adjust that further
+		{
+			FFormatNamedArguments Args;
+			Args.Add(TEXT("AdapterName"), FText::FromString(GRHIAdapterName));
+			Args.Add(TEXT("UserDriverVersion"), FText::FromString(GRHIAdapterUserDriverVersion));
+			Args.Add(TEXT("InternalDriverVersion"), FText::FromString(GRHIAdapterInternalDriverVersion));
+			Args.Add(TEXT("DriverDate"), FText::FromString(GRHIAdapterDriverDate));
+			Args.Add(TEXT("SuggestedVersion"), FText::FromString(SuggestedVersion));
+
+			// explicitly compact
+			FText LocalizedMsg = FText::Format(NSLOCTEXT("MessageDialog", "VideoCardDriverIssueReport",
+				"Suggested version: {SuggestedVersion} (newer drivers might be good as well)\nCurrent version: {UserDriverVersion} ({InternalDriverVersion}) {DriverDate}\nName: {AdapterName}"),
+				Args);
+
+			// less compact and more readable, maybe for later
+			/*
+				"Current driver:\n"
+				"  Name:\t{AdapterName}\n"
+				"  Version:\t{UserDriverVersion} (internal {InternalDriverVersion})\n"
+				"  Date:\t{DriverDate}\n"
+				"\n"
+				"Suggested driver (newer drivers might be good as well):\n"
+				"  Version: {SuggestedVersion}\n"
+				"\n"
+				"Please update your drivers to avoid problems (performance, stalls, crashes, reboot)."
+			*/
+
+			FPlatformMisc::MessageBoxExt(EAppMsgType::Ok, 
+				*LocalizedMsg.ToString(),
+				*NSLOCTEXT("MessageDialog", "TitleVideoCardDriverIssue", "WARNING: Video card driver with known issues").ToString());
+		}
+	}
+	else if(IsRHIDeviceAMD())
+	{
+		// so far we don't have any data on issues
+	}
+	else if(IsRHIDeviceIntel())
+	{
+		// so far we don't have any data on issues
+	}
+	else
+	{
+		// so far we don't have any data on issues
+	}
+}
+
 void RHIInit(bool bHasEditorToken)
 {
 	if(!GDynamicRHI)
@@ -63,6 +189,8 @@ void RHIInit(bool bHasEditorToken)
 
 		check(GDynamicRHI);
 	}
+
+	RHIDetectAndWarnOfBadDrivers();
 }
 
 void RHIPostInit()
diff --git a/Engine/Source/Runtime/RHI/Private/GPUDefragAllocator.cpp b/Engine/Source/Runtime/RHI/Private/GPUDefragAllocator.cpp
--- a/Engine/Source/Runtime/RHI/Private/GPUDefragAllocator.cpp
+++ b/Engine/Source/Runtime/RHI/Private/GPUDefragAllocator.cpp
@@ -111,7 +111,7 @@ FBestFitAllocator implementation.
 * @param	bAllowFailure	Whether to allow allocation failure or not
 * @return	Pointer to allocated memory
 */
-void* FGPUDefragAllocator::Allocate(int64 AllocationSize, int32 Alignment, bool bAllowFailure)
+void* FGPUDefragAllocator::Allocate(int64 AllocationSize, int32 Alignment, TStatId InStat, bool bAllowFailure)
 {
 	SCOPE_SECONDS_COUNTER(TimeSpentInAllocator);
 	FScopeLock Lock(&SynchronizationObject);
@@ -188,6 +188,10 @@ void* FGPUDefragAllocator::Allocate(int64 AllocationSize, int32 Alignment, bool
 #endif
 
 	FMemoryChunk* AllocatedChunk = AllocateChunk(BestChunk, AllocationSize, false);
+	//todo: fix stats
+	//ensureMsgf(AllocatedChunk->Stat.IsNone(), TEXT("FreeChunk already has a stat."));
+	AllocatedChunk->Stat = InStat;
+
 	check(IsAligned(AllocatedChunk->Base, Alignment));
 	return AllocatedChunk->Base;
 }
@@ -247,6 +251,7 @@ void FGPUDefragAllocator::FreeChunk(FMemoryChunk* Chunk, bool bMaintainSortedFre
 {
 	// Remove the entry
 	PointerToChunkMap.Remove(Chunk->Base);
+	Chunk->Stat = TStatId();
 
 	// Update usage stats in a thread safe way.
 	FPlatformAtomics::InterlockedAdd(&AllocatedMemorySize, -Chunk->Size);
@@ -649,7 +654,7 @@ void FGPUDefragAllocator::SortFreeList(int32& NumFreeChunks, int64& LargestFreeC
 */
 FGPUDefragAllocator::FMemoryChunk* FGPUDefragAllocator::FindAdjacent(FMemoryChunk* UsedChunk, bool bAnyChunkType)
 {
-	if (UsedChunk && !UsedChunk->IsRelocating() && (bAnyChunkType || UsedChunk->HasReallocationRequest()))
+	if (UsedChunk && !UsedChunk->IsRelocating() && (bAnyChunkType))
 	{
 		FMemoryChunk* FreeChunkLeft = UsedChunk->PreviousChunk;
 		FMemoryChunk* FreeChunkRight = UsedChunk->NextChunk;
@@ -674,50 +679,6 @@ FGPUDefragAllocator::FMemoryChunk* FGPUDefragAllocator::FindAdjacent(FMemoryChun
 	return nullptr;
 }
 
-/**
-* Searches for a reallocation request that would fit within the specified free chunk.
-* Prefers allocation requests over reallocation requests.
-*
-* @param FreeChunk		Free chunk we're trying to fill up
-* @return				First reallocating chunk that could fit, or nullptr
-*/
-FGPUDefragAllocator::FRequestNode* FGPUDefragAllocator::FindAnyReallocation(FMemoryChunk* FreeChunk)
-{
-	FRequestNode* BestRequest = nullptr;
-	int64 BestFit = MAX_int64;
-	bool bFoundAllocation = false;
-
-	for (FRequestList::TIterator It(ReallocationRequests.GetHead()); It; ++It)
-	{
-		FRequestNode* RequestNode = It.GetNode();
-		FAsyncReallocationRequest* Request = RequestNode->GetValue();
-		FMemoryChunk* CurrentChunk = Request->MemoryChunk;
-		int64 CurrentFit = FreeChunk->Size - Request->GetNewSize();
-
-		// Have we started getting into reallocations, but we've already found a suitable allocation?
-		if (bFoundAllocation && Request->IsReallocation())
-		{
-			// Use the allocation first.
-			break;
-		}
-
-		// Better fit than previously?
-		if (CurrentFit >= 0 && CurrentFit < BestFit && (Request->IsAllocation() || CanRelocate(CurrentChunk)))
-		{
-			BestFit = CurrentFit;
-			BestRequest = RequestNode;
-			bFoundAllocation = Request->IsAllocation();
-
-			// Perfect fit?
-			if (CurrentFit == 0)
-			{
-				break;
-			}
-		}
-	}
-	return BestRequest;
-}
-
 /**
 * Searches for an allocated chunk that would fit within the specified free chunk.
 * The allocated chunk must be adjacent to a free chunk and have a larger
@@ -856,14 +817,13 @@ void FGPUDefragAllocator::CheckForErrors(bool bCheckSortedFreeList)
 */
 FGPUDefragAllocator::FMemoryChunk* FGPUDefragAllocator::RelocateIntoFreeChunk(FRelocationStats& Stats, FMemoryChunk* FreeChunk, FMemoryChunk* SourceChunk)
 {
-	FRequestNode* ReallocationRequestNode = SourceChunk->ReallocationRequestNode;
-
 	// Save off important data from 'SourceChunk', since it will get modified by the call to LinkFreeChunk().
 	uint8* OldBaseAddress = SourceChunk->Base;
 	void* UserPayload = SourceChunk->UserPayload;
 	int32 OldSize = SourceChunk->Size;
 	const int32 NewSize = SourceChunk->GetFinalSize();
 	const int32 UsedSize = FMath::Min(NewSize, OldSize);
+	TStatId	SourceStat = SourceChunk->Stat;
 
 	// Are we relocating into adjacent free chunk?
 	bool bAdjacentRelocation = SourceChunk->PreviousChunk == FreeChunk || SourceChunk->NextChunk == FreeChunk;
@@ -901,26 +861,8 @@ FGPUDefragAllocator::FMemoryChunk* FGPUDefragAllocator::RelocateIntoFreeChunk(FR
 	// Update our book-keeping.
 	PointerToChunkMap.Remove(OldBaseAddress);
 	PointerToChunkMap.Add(NewBaseAddress, DestinationChunk);
-
-	// Move the reallocation request into the InProgress list.
-	if (ReallocationRequestNode)
-	{
-		FAsyncReallocationRequest* ReallocationRequest = ReallocationRequestNode->GetValue();
-		check(SourceChunk == ReallocationRequest->MemoryChunk);
-
-		// Create a new node in the 'InProgress' list.
-		ReallocationRequestsInProgress.AddHead(ReallocationRequest);
-		FRequestNode* NewNode = ReallocationRequestsInProgress.GetHead();
-
-		// Swap chunk and node pointers.
-		SourceChunk->ReallocationRequestNode = nullptr;
-		DestinationChunk->ReallocationRequestNode = NewNode;
-		ReallocationRequest->MemoryChunk = DestinationChunk;
-		ReallocationRequest->NewAddress = NewBaseAddress;
-
-		// Remove from current request list.
-		ReallocationRequests.RemoveNode(ReallocationRequestNode);
-	}
+	SourceChunk->Stat = TStatId();
+	DestinationChunk->Stat = SourceStat;	
 
 	// Is there free space left over at the end of DestinationChunk?
 	FMemoryChunk* NextFreeChunk;
@@ -963,51 +905,6 @@ FGPUDefragAllocator::FMemoryChunk* FGPUDefragAllocator::RelocateIntoFreeChunk(FR
 	return NextFreeChunk;
 }
 
-/**
-* Allocates memory from the specified free chunk, servicing an async allocation request.
-*
-* @param Stats			[out] Stats
-* @param FreeChunk		Chunk to allocate memory from
-* @param RequestNode	List node to the allocation request
-* @return				Next Free chunk to try to fill up
-*/
-FGPUDefragAllocator::FMemoryChunk* FGPUDefragAllocator::AllocateIntoFreeChunk(FRelocationStats& Stats, FMemoryChunk* FreeChunk, FRequestNode* RequestNode)
-{
-	FMemoryChunk* NextFreeChunk = FreeChunk->NextFreeChunk;
-	FAsyncReallocationRequest* Request = RequestNode->GetValue();
-	check(Request->IsAllocation());
-
-	// Note: AllocateChunk() may split 'FreeChunk' and return a new chunk.
-	FMemoryChunk* AllocatedChunk = AllocateChunk(FreeChunk, Request->GetNewSize(), true);
-
-	// Create a new node in the 'InProgress' list.
-	ReallocationRequestsInProgress.AddHead(Request);
-	FRequestNode* NewNode = ReallocationRequestsInProgress.GetHead();
-
-	// Setup chunk and node variables.
-	AllocatedChunk->ReallocationRequestNode = NewNode;
-	Request->MemoryChunk = AllocatedChunk;
-	Request->NewAddress = AllocatedChunk->Base;
-
-	// Mark the chunk as "in use" during the current sync step.
-	// Note: This sync index will propagate if these chunks are involved in any merge/split in the future.
-	AllocatedChunk->SetSyncIndex(GetCurrentSyncIndex(), AllocatedChunk->Size);
-
-	// Remove from current request list.
-	ReallocationRequests.RemoveNode(RequestNode);
-
-	FPlatformAtomics::InterlockedAdd(&PendingMemoryAdjustment, -Request->GetNewSize());
-
-	if (AllocatedChunk->NextChunk && AllocatedChunk->NextChunk->bIsAvailable)
-	{
-		return AllocatedChunk->NextChunk;
-	}
-	else
-	{
-		return NextFreeChunk;
-	}
-}
-
 /**
 * Blocks the calling thread until all relocations and reallocations that were initiated by Tick() have completed.
 *
@@ -1021,46 +918,7 @@ bool FGPUDefragAllocator::FinishAllRelocations()
 	if (bWasAnyRelocationsInProgress)
 	{		
 		BlockOnFence();
-	}
-
-	// All reallocation requests have now completed.
-	int32 TotalMemoryAdjustment = 0;
-	for (FRequestList::TIterator It(ReallocationRequestsInProgress.GetHead()); It; ++It)
-	{
-		FAsyncReallocationRequest* FinishedRequest = *It;
-
-		FMemoryChunk* FinishedChunk = FinishedRequest->MemoryChunk;
-		if (!FinishedRequest->IsCanceled())
-		{
-			// Mark it complete.
-			FinishedRequest->MarkCompleted();
-
-			if (FinishedRequest->IsReallocation())
-			{
-				NumFinishedAsyncReallocations++;
-			}
-			else
-			{
-				NumFinishedAsyncAllocations++;
-			}
-
-			if (!bBenchmarkMode)
-			{
-				PlatformNotifyReallocationFinished(FinishedRequest, FinishedChunk->UserPayload);
-			}
-			TotalMemoryAdjustment += FinishedRequest->NewSize - FinishedRequest->OldSize;
-			FinishedRequest->MemoryChunk = nullptr;
-		}
-		else
-		{
-			// Delete the request, since it's our copy. The user has deleted the original already.
-			delete FinishedRequest;
-		}
-		FinishedChunk->ReallocationRequestNode = nullptr;
-	}
-	ReallocationRequestsInProgress.Empty();
-
-	check(ReallocationRequests.Num() > 0 || PendingMemoryAdjustment == 0);
+	}	
 
 	// Take the opportunity to free all chunks that couldn't be freed immediately before.
 	for (TDoubleLinkedList<FMemoryChunk*>::TIterator It(PendingFreeChunks.GetHead()); It; ++It)
@@ -1127,116 +985,6 @@ void FGPUDefragAllocator::BlockOnSyncIndex(uint32 SyncIndex)
 	}
 }
 
-/**
-* Cancels the specified reallocation request.
-* Note that the allocator doesn't keep track of requests after it's been completed,
-* so the user must provide the current base address. This may not match any of the
-* addresses in the (old) request since the memory may have been relocated since then.
-*
-* @param Request				Request to cancel. Must be a valid request.
-* @param CurrentBaseAddress	Current baseaddress used by the allocation.
-*/
-void FGPUDefragAllocator::CancelAsyncReallocation(FAsyncReallocationRequest* Request, const void* CurrentBaseAddress)
-{
-	check(Request && !Request->IsCanceled());
-	NumCanceledAsyncRequests++;
-
-	int32 MemoryAdjustment = Request->NewSize - Request->OldSize;
-
-	// We cannot undo shrinking reallocations. The memory will be lost and gone forever.
-	check(MemoryAdjustment > 0);
-
-	// Mark it canceled.
-	Request->bIsCanceled = true;
-
-	// Make sure it's marked 'completed'.
-	bool bHasStarted = Request->HasStarted();
-	bool bHasCompleted = Request->HasCompleted();
-	if (!bHasCompleted)
-	{
-		// This will also prevent it from further relocation until it's gone completely from the system.
-		Request->MarkCompleted();
-	}
-
-	// Has it not started yet?
-	if (!bHasStarted)
-	{
-		if (Request->IsReallocation())
-		{
-			// If it hasn't started yet, just remove the request and flag it completed.
-			FMemoryChunk* MatchingChunk = Request->MemoryChunk;
-			check(MatchingChunk && CurrentBaseAddress == nullptr);
-			FRequestNode* RequestNode = MatchingChunk->ReallocationRequestNode;
-			check(RequestNode);
-
-			// Remove it from our end of the system.
-			ReallocationRequests.RemoveNode(RequestNode);
-			MatchingChunk->ReallocationRequestNode = nullptr;
-		}
-		else
-		{
-			FRequestNode* RequestNode = ReallocationRequests.FindNode(Request);
-			check(RequestNode);
-			ReallocationRequests.RemoveNode(RequestNode);
-		}
-		FPlatformAtomics::InterlockedAdd(&PendingMemoryAdjustment, -MemoryAdjustment);
-	}
-	else
-	{
-		// Is it still in progress?
-		if (!bHasCompleted)
-		{
-			FMemoryChunk* MatchingChunk = Request->MemoryChunk;
-			check(MatchingChunk && CurrentBaseAddress == nullptr);
-
-			// Make a copy of the request, since the current Request will be deleted by the user.
-			FAsyncReallocationRequest* RequestCopy = new FAsyncReallocationRequest(*Request);
-			FRequestNode* RequestNode = MatchingChunk->ReallocationRequestNode;
-			MatchingChunk->ReallocationRequestNode = nullptr;
-
-			FMemoryChunk* NewChunk;
-			if (Request->IsReallocation())
-			{
-				// Undo the 'Grow' by immediately shrinking it back to the old size (adjusting baseaddress).
-				// Allow the GPU to finish relocating into the used portion. (Any new relocations will be pipelined.)
-				NewChunk = Shrink(MatchingChunk, MemoryAdjustment);
-			}
-			else
-			{
-				// Undo the allocation by freeing it.
-				FreeChunk(MatchingChunk, false);
-				NewChunk = MatchingChunk;
-			}
-
-			// Fixup all request pointers.
-			RequestNode->GetValue() = RequestCopy;
-			NewChunk->ReallocationRequestNode = RequestNode;
-			RequestCopy->MemoryChunk = NewChunk;
-		}
-		// Is it already completed?
-		else
-		{
-			// When it's completed already, the allocator doesn't keep track of it anymore in any lists.
-			if (Request->IsReallocation())
-			{
-				FMemoryChunk* MatchingChunk = PointerToChunkMap.FindRef(CurrentBaseAddress);
-				check(MatchingChunk && MatchingChunk->ReallocationRequestNode == nullptr);
-
-				// Undo the 'Grow' by immediately shrinking it back to the old size (adjusting baseaddress).
-				FMemoryChunk* NewChunk = Shrink(MatchingChunk, MemoryAdjustment);
-			}
-			else
-			{
-				FMemoryChunk* MatchingChunk = PointerToChunkMap.FindRef(Request->GetNewBaseAddress());
-				check(MatchingChunk && MatchingChunk->ReallocationRequestNode == nullptr);
-
-				// Undo the allocation by freeing it.
-				FreeChunk(MatchingChunk, false);
-			}
-		}
-	}
-}
-
 static int32 GGPUDefragEnableTimeLimits = 1;
 static FAutoConsoleVariableRef CVarGPUDefragEnableTimeLimits(
 	TEXT("r.GPUDefrag.EnableTimeLimits"),
@@ -1305,21 +1053,15 @@ void FGPUDefragAllocator::PartialDefragmentation(FRelocationStats& Stats, double
 				BestChunk = RelocateAllowed(FreeChunk, FindAdjacent(FreeChunk->NextChunk, false));
 
 				if (!BestChunk)
-				{
-					// 3. Merge with any async request (best-fitting)
-					BestRequestNode = FindAnyReallocation(FreeChunk);
+				{					
+					// 4. Merge with a used chunk adjacent to hole (to make that hole larger).
+					BestChunk = RelocateAllowed(FreeChunk, FindAdjacentToHole(FreeChunk));
 
-					if (!BestRequestNode)
+					if (!BestChunk)
 					{
-						// 4. Merge with a used chunk adjacent to hole (to make that hole larger).
-						BestChunk = RelocateAllowed(FreeChunk, FindAdjacentToHole(FreeChunk));
-
-						if (!BestChunk)
-						{
-							// 5. Merge with chunk from the end of the pool (well-fitting)
-							BestChunk = FindAny(FreeChunk);
-						}
-					}
+						// 5. Merge with chunk from the end of the pool (well-fitting)
+						BestChunk = FindAny(FreeChunk);
+					}					
 				}
 			}
 		}
@@ -1327,23 +1069,11 @@ void FGPUDefragAllocator::PartialDefragmentation(FRelocationStats& Stats, double
 		if (BestChunk)
 		{
 			FreeChunk = RelocateIntoFreeChunk(Stats, FreeChunk, BestChunk);			
-		}
-		else if (BestRequestNode)
-		{
-			FAsyncReallocationRequest* Request = BestRequestNode->GetValue();
-			if (Request->IsReallocation())
-			{				
-				FreeChunk = RelocateIntoFreeChunk(Stats, FreeChunk, Request->MemoryChunk);				
-			}
-			else
-			{
-				FreeChunk = AllocateIntoFreeChunk(Stats, FreeChunk, BestRequestNode);
-			}
-		}
+		}		
 		else
 		{
 			// Did the free chunk fail to defrag?
-			if (FreeChunk->DefragCounter == 0 && (FreeChunk->NextFreeChunk || ReallocationRequests.Num() > 0))
+			if (FreeChunk->DefragCounter == 0 && (FreeChunk->NextFreeChunk))
 			{
 				// Don't try it again for a while.
 				if (FreeChunk->Size < DEFRAG_SMALL_CHUNK_SIZE)
@@ -1690,18 +1420,7 @@ FGPUDefragAllocator::EMemoryElementType FGPUDefragAllocator::GetChunkType(FMemor
 	else if (Chunk->bIsAvailable)
 	{
 		ChunkType = MET_Free;			// Free (dark grey)
-	}
-	else if (Chunk->HasReallocationRequest())
-	{
-		if (Chunk->ReallocationRequestNode->GetValue()->HasCompleted())
-		{
-			ChunkType = MET_Resized;	// Has been resized but not finalized yet (dark green)
-		}
-		else
-		{
-			ChunkType = MET_Resizing;	// Allocated but wants a resize (green)
-		}
-	}
+	}	
 	else if (CanRelocate(Chunk) == false)
 	{
 		ChunkType = MET_Locked;			// Allocated but can't me relocated at this time (locked) (red)
diff --git a/Engine/Source/Runtime/RHI/Private/RHI.cpp b/Engine/Source/Runtime/RHI/Private/RHI.cpp
--- a/Engine/Source/Runtime/RHI/Private/RHI.cpp
+++ b/Engine/Source/Runtime/RHI/Private/RHI.cpp
@@ -229,6 +229,7 @@ FString GRHIAdapterInternalDriverVersion;
 FString GRHIAdapterUserDriverVersion;
 FString GRHIAdapterDriverDate;
 uint32 GRHIVendorId = 0;
+uint32 GRHIDeviceId = 0;
 bool GSupportsRenderDepthTargetableShaderResources = true;
 bool GSupportsRenderTargetFormat_PF_G8 = true;
 bool GSupportsRenderTargetFormat_PF_FloatRGBA = true;
diff --git a/Engine/Source/Runtime/RHI/Private/RHICommandList.cpp b/Engine/Source/Runtime/RHI/Private/RHICommandList.cpp
--- a/Engine/Source/Runtime/RHI/Private/RHICommandList.cpp
+++ b/Engine/Source/Runtime/RHI/Private/RHICommandList.cpp
@@ -480,7 +480,7 @@ void FRHICommandListExecutor::ExecuteList(FRHICommandListImmediate& CmdList)
 
 void FRHICommandListExecutor::LatchBypass()
 {
-#if !UE_BUILD_SHIPPING
+#if CAN_TOGGLE_COMMAND_LIST_BYPASS
 	if (GRHIThread)
 	{
 		if (bLatchedBypass)
@@ -522,7 +522,7 @@ void FRHICommandListExecutor::LatchBypass()
 	if (!bLatchedBypass)
 	{
 		bLatchedUseParallelAlgorithms = FApp::ShouldUseThreadingForPerformance() 
-#if !UE_BUILD_SHIPPING
+#if CAN_TOGGLE_COMMAND_LIST_BYPASS
 			&& (CVarRHICmdUseParallelAlgorithms.GetValueOnAnyThread() >= 1)
 #endif
 			;
@@ -1894,7 +1894,7 @@ void FDynamicRHI::UpdateTexture2D_RenderThread(class FRHICommandListImmediate& R
 
 void* FDynamicRHI::LockTexture2D_RenderThread(class FRHICommandListImmediate& RHICmdList, FTexture2DRHIParamRef Texture, uint32 MipIndex, EResourceLockMode LockMode, uint32& DestStride, bool bLockWithinMiptail, bool bNeedsDefaultRHIFlush)
 {
-	if (bNeedsDefaultRHIFlush)
+	if (bNeedsDefaultRHIFlush) 
 	{
 		QUICK_SCOPE_CYCLE_COUNTER(STAT_RHIMETHOD_LockTexture2D_Flush);
 		RHICmdList.ImmediateFlush(EImmediateFlushType::FlushRHIThread);
diff --git a/Engine/Source/Runtime/RHI/Public/GPUDefragAllocator.h b/Engine/Source/Runtime/RHI/Public/GPUDefragAllocator.h
--- a/Engine/Source/Runtime/RHI/Public/GPUDefragAllocator.h
+++ b/Engine/Source/Runtime/RHI/Public/GPUDefragAllocator.h
@@ -134,7 +134,7 @@ public:
 		* @param	ChunkToInsertAfter		Chunk to insert this after.
 		* @param	FirstFreeChunk			Reference to first free chunk Pointer.
 		*/
-		FMemoryChunk(uint8* InBase, int32 InSize, FGPUDefragAllocator& InBestFitAllocator, FMemoryChunk*& ChunkToInsertAfter, bool bSortedFreeList)
+		FMemoryChunk(uint8* InBase, int32 InSize, FGPUDefragAllocator& InBestFitAllocator, FMemoryChunk*& ChunkToInsertAfter, TStatId InStat, bool bSortedFreeList)
 			: Base(InBase)
 			, Size(InSize)
 			, bIsAvailable(false)
@@ -143,8 +143,8 @@ public:
 			, BestFitAllocator(InBestFitAllocator)
 			, SyncIndex(0)
 			, SyncSize(0)
-			, UserPayload(0)
-			, ReallocationRequestNode(nullptr)
+			, UserPayload(0)			
+			, Stat(InStat)
 		{
 			Link(ChunkToInsertAfter);
 			// This is going to change bIsAvailable.
@@ -253,15 +253,7 @@ public:
 
 			PreviousFreeChunk = nullptr;
 			NextFreeChunk = nullptr;
-		}
-
-		/**
-		*	Returns true if the Chunk has an associated reallocation request.
-		*/
-		bool	HasReallocationRequest() const
-		{
-			return ReallocationRequestNode != nullptr;
-		}
+		}		
 
 		/**
 		*	Returns true if the Chunk is being asynchronously relocated due to reallocation or defrag.
@@ -352,13 +344,10 @@ public:
 		/** Number of uint8s covered by the SyncIndex (starting from the beginning of the chunk). */
 		int64					SyncSize;
 		/** User payload, e.g. platform-specific texture Pointer. Only chunks with payload can be relocated. */
-		void*					UserPayload;
+		void*					UserPayload;		
 
-		/**
-		* Reallocation request for this chunk, or nullptr.
-		* Point32s straight to a node in FBestFitAllocator::ReallocationRequests or FBestFitAllocator::ReallocationRequestsInProgress.
-		*/
-		FRequestNode*			ReallocationRequestNode;
+		//stat associated with this allocation
+		TStatId Stat;
 	};
 
 	/** Constructor, zero initializing all member variables */
@@ -409,7 +398,7 @@ public:
 		// Update stats in a thread safe way.
 		FPlatformAtomics::InterlockedExchange(&AvailableMemorySize, MemorySize);
 		// Allocate initial chunk.
-		FirstChunk = new FMemoryChunk(MemoryBase, MemorySize, *this, FirstChunk, /*FirstChunk, FirstFreeChunk,*/ false);
+		FirstChunk = new FMemoryChunk(MemoryBase, MemorySize, *this, FirstChunk, TStatId(), false);
 		LastChunk = FirstChunk;
 	}
 
@@ -453,7 +442,7 @@ public:
 	* @param	bAllowFailure	Whether to allow allocation failure or not
 	* @return	Pointer to allocated memory
 	*/
-	virtual void*	Allocate(int64 AllocationSize, int32 Alignment, bool bAllowFailure);
+	virtual void*	Allocate(int64 AllocationSize, int32 Alignment, TStatId InStat, bool bAllowFailure);	
 
 	/**
 	* Frees allocation associated with the specified Pointer.
@@ -837,7 +826,10 @@ protected:
 		// Calculate size of second chunk...
 		int32 SecondSize = BaseChunk->Size - FirstSize;
 		// ... and create it.
-		FMemoryChunk* NewFreeChunk = new FMemoryChunk(BaseChunk->Base + FirstSize, SecondSize, *this, BaseChunk, /*FirstChunk, FirstFreeChunk,*/ bSortedFreeList);
+
+		//todo: fix stats
+		//ensureMsgf(BaseChunk->Stat.IsNone(), TEXT("Free chunk has stat"));
+		FMemoryChunk* NewFreeChunk = new FMemoryChunk(BaseChunk->Base + FirstSize, SecondSize, *this, BaseChunk, TStatId(), bSortedFreeList);
 
 		// Keep the original sync index for the new chunk if necessary.
 		if (BaseChunk->IsRelocating() && BaseChunk->SyncSize > FirstSize)
@@ -910,15 +902,7 @@ protected:
 	* @param bAnyChunkType	If false, only succeeds if 'UsedChunk' has a reallocation request and fits
 	* @return				Returns 'UsedChunk' if it fits the criteria, otherwise nullptr
 	*/
-	FMemoryChunk*	FindAdjacent(FMemoryChunk* UsedChunk, bool bAnyChunkType);
-
-	/**
-	* Searches for a reallocation request that would fit within the specified free chunk.
-	*
-	* @param FreeChunk		Free chunk we're trying to fill up
-	* @return				First request that could fit, or nullptr
-	*/
-	FRequestNode*	FindAnyReallocation(FMemoryChunk* FreeChunk);
+	FMemoryChunk*	FindAdjacent(FMemoryChunk* UsedChunk, bool bAnyChunkType);	
 
 	/**
 	* Searches for an allocated chunk that would fit within the specified free chunk.
@@ -952,16 +936,6 @@ protected:
 	*/
 	FMemoryChunk*	RelocateIntoFreeChunk(FRelocationStats& Stats, FMemoryChunk* FreeChunk, FMemoryChunk* UsedChunk);
 
-	/**
-	* Allocates memory from the specified free chunk, servicing an async allocation request.
-	*
-	* @param Stats			[out] Stats
-	* @param FreeChunk		Chunk to allocate memory from
-	* @param RequestNode	List node to the allocation request
-	* @return				Next Free chunk to try to fill up
-	*/
-	FMemoryChunk*	AllocateIntoFreeChunk(FRelocationStats& Stats, FMemoryChunk* FreeChunk, FRequestNode* RequestNode);
-
 	FMemoryChunk* RelocateAllowed(FMemoryChunk* FreeChunk, FMemoryChunk* UsedChunk);	
 
 	FCriticalSection	SynchronizationObject;
@@ -1009,12 +983,6 @@ protected:
 	/** Platform-specific (GPU) fence, used for synchronizing the Sync Index. */
 	uint64			PlatformSyncFence;
 
-	/** Asynchronous reallocations requested by user.				*/
-	FRequestList	ReallocationRequests;
-
-	/** Asynchronous reallocations currently being processed.		*/
-	FRequestList	ReallocationRequestsInProgress;
-
 	/** Chunks that couldn't be freed immediately because they were being relocated. */
 	TDoubleLinkedList<FMemoryChunk*>	PendingFreeChunks;
 
@@ -1195,7 +1163,7 @@ private:
 */
 FORCEINLINE int32 FGPUDefragAllocator::FMemoryChunk::GetFinalSize() const
 {
-	return ReallocationRequestNode ? ReallocationRequestNode->GetValue()->GetNewSize() : Size;
+	return Size;
 }
 
 /**
@@ -1207,14 +1175,6 @@ FORCEINLINE int32 FGPUDefragAllocator::FMemoryChunk::GetFinalSize() const
 */
 FORCEINLINE bool FGPUDefragAllocator::CanRelocate(const FMemoryChunk* Chunk) const
 {
-	// During reallocation, the new texture keeps the request around until it's been processed higher up.
-	// Can't relocate a request that has started until it's fully gone from the system.
-	const FRequestNode* RequestNode = Chunk->ReallocationRequestNode;
-	if (RequestNode && RequestNode->GetValue()->HasStarted())
-	{
-		return false;
-	}
-
 	if (Chunk->IsLocked())
 	{
 		return false;
diff --git a/Engine/Source/Runtime/RHI/Public/RHI.h b/Engine/Source/Runtime/RHI/Public/RHI.h
--- a/Engine/Source/Runtime/RHI/Public/RHI.h
+++ b/Engine/Source/Runtime/RHI/Public/RHI.h
@@ -76,6 +76,7 @@ extern RHI_API FString GRHIAdapterName;
 extern RHI_API FString GRHIAdapterInternalDriverVersion;
 extern RHI_API FString GRHIAdapterUserDriverVersion;
 extern RHI_API FString GRHIAdapterDriverDate;
+extern RHI_API uint32 GRHIDeviceId;
 
 // 0 means not defined yet, use functions like IsRHIDeviceAMD() to access
 extern RHI_API uint32 GRHIVendorId;
diff --git a/Engine/Source/Runtime/RHI/Public/RHICommandList.h b/Engine/Source/Runtime/RHI/Public/RHICommandList.h
--- a/Engine/Source/Runtime/RHI/Public/RHICommandList.h
+++ b/Engine/Source/Runtime/RHI/Public/RHICommandList.h
@@ -2948,6 +2948,9 @@ public:
 };
 
 
+// This controls if the cmd list bypass can be toggled at runtime. It is quite expensive to have these branches in there.
+#define CAN_TOGGLE_COMMAND_LIST_BYPASS (!UE_BUILD_SHIPPING && !UE_BUILD_TEST)
+
 class RHI_API FRHICommandListExecutor
 {
 public:
@@ -2971,7 +2974,7 @@ public:
 
 	FORCEINLINE_DEBUGGABLE bool Bypass()
 	{
-#if !UE_BUILD_SHIPPING || PLATFORM_CAN_TOGGLE_RHITHREAD_IN_SHIPPING
+#if CAN_TOGGLE_COMMAND_LIST_BYPASS
 		return bLatchedBypass;
 #else
 		return !!DefaultBypass;
@@ -2979,7 +2982,7 @@ public:
 	}
 	FORCEINLINE_DEBUGGABLE bool UseParallelAlgorithms()
 	{
-#if !UE_BUILD_SHIPPING
+#if CAN_TOGGLE_COMMAND_LIST_BYPASS
 		return bLatchedUseParallelAlgorithms;
 #else
 		return  FApp::ShouldUseThreadingForPerformance() && !Bypass();
diff --git a/Engine/Source/Runtime/Renderer/Private/PostProcess/PostProcessTemporalAA.cpp b/Engine/Source/Runtime/Renderer/Private/PostProcess/PostProcessTemporalAA.cpp
--- a/Engine/Source/Runtime/Renderer/Private/PostProcess/PostProcessTemporalAA.cpp
+++ b/Engine/Source/Runtime/Renderer/Private/PostProcess/PostProcessTemporalAA.cpp
@@ -36,7 +36,7 @@ static void TemporalRandom(FVector2D* RESTRICT const Constant, uint32 FrameNumbe
 
 static TAutoConsoleVariable<float> CVarTemporalAASharpness(
 	TEXT("r.TemporalAASharpness"),
-	0.0f,
+	1.0f,
 	TEXT("Sharpness of temporal AA (0.0 = smoother, 1.0 = sharper)."),
 	ECVF_Scalability | ECVF_RenderThreadSafe);
 
diff --git a/Engine/Source/Runtime/Renderer/Private/PostProcess/PostProcessTonemap.cpp b/Engine/Source/Runtime/Renderer/Private/PostProcess/PostProcessTonemap.cpp
--- a/Engine/Source/Runtime/Renderer/Private/PostProcess/PostProcessTonemap.cpp
+++ b/Engine/Source/Runtime/Renderer/Private/PostProcess/PostProcessTonemap.cpp
@@ -8,6 +8,7 @@
 #include "ScenePrivate.h"
 #include "SceneFilterRendering.h"
 #include "PostProcessEyeAdaptation.h"
+#include "PostProcessUpscale.h"
 #include "PostProcessTonemap.h"
 #include "PostProcessing.h"
 #include "PostProcessCombineLUTs.h"
@@ -1077,29 +1078,14 @@ IMPLEMENT_SHADER_TYPE(template<>, TPostProcessTonemapVS<true>, TEXT("PostProcess
 IMPLEMENT_SHADER_TYPE(template<>, TPostProcessTonemapVS<false>, TEXT("PostProcessTonemap"), TEXT("MainVS"), SF_Vertex);
 
 
-FRCPassPostProcessTonemap::FRCPassPostProcessTonemap(const FViewInfo& View, bool bInDoGammaOnly, bool bDoScreenPercentageInTonemapper, bool bInDoEyeAdaptation )
+FRCPassPostProcessTonemap::FRCPassPostProcessTonemap(const FViewInfo& InView, bool bInDoGammaOnly, bool bInDoEyeAdaptation)
 	: bDoGammaOnly(bInDoGammaOnly)
+	, bDoScreenPercentageInTonemapper(false)
 	, bDoEyeAdaptation(bInDoEyeAdaptation)
-	, RenderTargetExtend(0, 0)
+	, View(InView)	
 {
-	uint32 ConfigBitmask = TonemapperGenerateBitmaskPC(&View, bDoGammaOnly);
+	uint32 ConfigBitmask = TonemapperGenerateBitmaskPC(&InView, bDoGammaOnly);
 	ConfigIndexPC = TonemapperFindLeastExpensive(TonemapperConfBitmaskPC, sizeof(TonemapperConfBitmaskPC)/4, TonemapperCostTab, ConfigBitmask);;
-
-	if(bDoScreenPercentageInTonemapper)
-	{
-		// Texture could be bigger than viewport
-		if (View.Family->RenderTarget->GetRenderTargetTexture())
-		{
-			RenderTargetExtend.X = View.Family->RenderTarget->GetRenderTargetTexture()->GetSizeX();
-			RenderTargetExtend.Y = View.Family->RenderTarget->GetRenderTargetTexture()->GetSizeY();
-		}
-		else
-		{
-			RenderTargetExtend = View.Family->RenderTarget->GetSizeXY();
-		}
-	}
-	
-	check(bDoScreenPercentageInTonemapper == ShouldDoScreenPercentageInTonemapper());
 }
 
 namespace PostProcessTonemapUtil
@@ -1145,14 +1131,14 @@ void FRCPassPostProcessTonemap::Process(FRenderingCompositePassContext& Context)
 		// input is not hooked up correctly
 		return;
 	}
-	
-	const FSceneView& View = Context.View;
+
 	const FSceneViewFamily& ViewFamily = *(View.Family);
 	FIntRect SrcRect = View.ViewRect;
-	FIntRect DestRect = ShouldDoScreenPercentageInTonemapper() ? View.UnscaledViewRect : View.ViewRect;
+	FIntRect DestRect = bDoScreenPercentageInTonemapper ? View.UnscaledViewRect : View.ViewRect;
 	FIntPoint SrcSize = InputDesc->Extent;
-
-	SCOPED_DRAW_EVENTF(Context.RHICmdList, PostProcessTonemap, TEXT("Tonemapper#%d%s %dx%d"), ConfigIndexPC, bDoGammaOnly ? TEXT(" GammaOnly") : TEXT(""), DestRect.Width(), DestRect.Height());
+	
+	SCOPED_DRAW_EVENTF(Context.RHICmdList, PostProcessTonemap, TEXT("Tonemapper#%d GammaOnly=%d ScreenPercentage=%d  %dx%d"),
+		ConfigIndexPC, bDoGammaOnly, bDoScreenPercentageInTonemapper, DestRect.Width(), DestRect.Height());
 
 	const FSceneRenderTargetItem& DestRenderTarget = PassOutputs[0].RequestSurface(Context);
 	
@@ -1170,7 +1156,7 @@ void FRCPassPostProcessTonemap::Process(FRenderingCompositePassContext& Context)
 		Context.RHICmdList.Clear(true, FLinearColor::Black, false, 1.0f, false, 0, DestRect);
 	}
 
-	Context.SetViewportAndCallRHI(DestRect.Min.X, DestRect.Min.Y, 0.0f, DestRect.Max.X + 1, DestRect.Max.Y + 1, 1.0f );
+	Context.SetViewportAndCallRHI(DestRect, 0.0f, 1.0f );
 
 	// set the state
 	Context.RHICmdList.SetBlendState(TStaticBlendState<>::GetRHI());
@@ -1231,11 +1217,6 @@ void FRCPassPostProcessTonemap::Process(FRenderingCompositePassContext& Context)
 		// This becomes even more important with some limited VRam (XBoxOne).
 		SceneContext.SetSceneColor(0);
 	}
-
-	if(ShouldDoScreenPercentageInTonemapper())
-	{
-		Context.View.ViewRect = View.UnscaledViewRect;
-	}
 }
 
 FPooledRenderTargetDesc FRCPassPostProcessTonemap::ComputeOutputDesc(EPassOutputId InPassOutputId) const
@@ -1247,11 +1228,6 @@ FPooledRenderTargetDesc FRCPassPostProcessTonemap::ComputeOutputDesc(EPassOutput
 	Ret.Format = PF_B8G8R8A8;
 	Ret.DebugName = TEXT("Tonemap");
 
-	if(ShouldDoScreenPercentageInTonemapper())
-	{
-		Ret.Extent = RenderTargetExtend;
-	}
-
 	return Ret;
 }
 
diff --git a/Engine/Source/Runtime/Renderer/Private/PostProcess/PostProcessTonemap.h b/Engine/Source/Runtime/Renderer/Private/PostProcess/PostProcessTonemap.h
--- a/Engine/Source/Runtime/Renderer/Private/PostProcess/PostProcessTonemap.h
+++ b/Engine/Source/Runtime/Renderer/Private/PostProcess/PostProcessTonemap.h
@@ -41,7 +41,7 @@ class FRCPassPostProcessTonemap : public TRenderingCompositePassBase<4, 1>
 {
 public:
 	// constructor
-	FRCPassPostProcessTonemap(const FViewInfo& View, bool bInDoGammaOnly, bool bDoScreenPercentageInTonemapper, bool bDoEyeAdaptation);
+	FRCPassPostProcessTonemap(const FViewInfo& InView, bool bInDoGammaOnly, bool bDoEyeAdaptation);
 
 	// interface FRenderingCompositePass ---------
 
@@ -49,17 +49,13 @@ public:
 	virtual void Release() override { delete this; }
 	virtual FPooledRenderTargetDesc ComputeOutputDesc(EPassOutputId InPassOutputId) const override;
 
-private:
-	// set in constructor
 	bool bDoGammaOnly;
+	bool bDoScreenPercentageInTonemapper;
+private:
 	bool bDoEyeAdaptation;
 	uint32 ConfigIndexPC;
 
-	// used if ShouldDoScreenPercentageInTonemapper() is true, otherwise (0,0)
-	FIntPoint RenderTargetExtend;
-
-	// @return true: we did the upscale in this pass and don't need to run another one after it
-	bool ShouldDoScreenPercentageInTonemapper() const { return RenderTargetExtend != FIntPoint::ZeroValue; }
+	const FViewInfo& View;
 
 	void SetShader(const FRenderingCompositePassContext& Context);
 };
diff --git a/Engine/Source/Runtime/Renderer/Private/PostProcess/PostProcessUpscale.cpp b/Engine/Source/Runtime/Renderer/Private/PostProcess/PostProcessUpscale.cpp
--- a/Engine/Source/Runtime/Renderer/Private/PostProcess/PostProcessUpscale.cpp
+++ b/Engine/Source/Runtime/Renderer/Private/PostProcess/PostProcessUpscale.cpp
@@ -20,8 +20,46 @@ static TAutoConsoleVariable<float> CVarUpscaleSoftness(
 	TEXT(" 0..1 (0.3 is good for ScreenPercentage 90"),
 	ECVF_Scalability | ECVF_RenderThreadSafe);
 
+static TAutoConsoleVariable<float> CVarUpscalePaniniD(
+	TEXT("r.Upscale.Panini.D"),
+	0,
+	TEXT("Allow and configure to apply a panini distortion to the rendered image. Values between 0 and 1 allow to fade the effect (lerp).\n")
+	TEXT("Implementation from research paper \"Pannini: A New Projection for Rendering Wide Angle Perspective Images\"\n")
+	TEXT(" 0: off (default)\n")
+	TEXT(">0: enabled (requires an extra post processing pass if upsampling wasn't used - see r.ScreenPercentage)\n")
+	TEXT(" 1: Panini cylindrical stereographic projection"),
+	ECVF_RenderThreadSafe);
+
+static TAutoConsoleVariable<float> CVarUpscalePaniniS(
+	TEXT("r.Upscale.Panini.S"),
+	0,
+	TEXT("Panini projection's hard vertical compression factor.\n")
+	TEXT(" 0: no vertical compression factor (default)\n")
+	TEXT(" 1: Hard vertical compression"),
+	ECVF_RenderThreadSafe);
+
+static TAutoConsoleVariable<float> CVarUpscalePaniniScreenFit(
+	TEXT("r.Upscale.Panini.ScreenFit"),
+	1.0f,
+	TEXT("Panini projection screen fit effect factor (lerp).\n")
+	TEXT(" 0: fit vertically\n")
+	TEXT(" 1: fit horizontally (default)"),
+	ECVF_RenderThreadSafe);
+
 const FRCPassPostProcessUpscale::PaniniParams FRCPassPostProcessUpscale::PaniniParams::Default;
 
+FRCPassPostProcessUpscale::PaniniParams::PaniniParams(const FViewInfo& View)
+{
+	*this = PaniniParams();
+
+	if (View.IsPerspectiveProjection() && !GEngine->StereoRenderingDevice.IsValid())
+	{
+		D = FMath::Max(CVarUpscalePaniniD.GetValueOnRenderThread(), 0.0f);
+		S = CVarUpscalePaniniS.GetValueOnRenderThread();
+		ScreenFit = FMath::Max(CVarUpscalePaniniScreenFit.GetValueOnRenderThread(), 0.0f);
+	}
+}
+
 static FVector2D PaniniProjection(FVector2D OM, float d, float s)
 {
 	float PaniniDirectionXZInvLength = 1.0f / FMath::Sqrt(1.0f + OM.X * OM.X);
diff --git a/Engine/Source/Runtime/Renderer/Private/PostProcess/PostProcessUpscale.h b/Engine/Source/Runtime/Renderer/Private/PostProcess/PostProcessUpscale.h
--- a/Engine/Source/Runtime/Renderer/Private/PostProcess/PostProcessUpscale.h
+++ b/Engine/Source/Runtime/Renderer/Private/PostProcess/PostProcessUpscale.h
@@ -27,12 +27,20 @@ public:
 		// Panini screen fit factor (lerp between vertical and horizontal) 
 		float ScreenFit;
 
+		// constructor off
 		PaniniParams()
 			: D(0.0f)
 			, S(0.0f)
 			, ScreenFit(1.0f)
 		{}
 
+		PaniniParams(const FViewInfo& View);
+
+		bool IsEnabled() const
+		{
+			return D > 0.01f;
+		}
+
 		static const PaniniParams Default;
 	};
 
@@ -46,9 +54,11 @@ public:
 	FRCPassPostProcessUpscale(uint32 InUpscaleQuality, const PaniniParams& InPaniniConfig = PaniniParams::Default);
 
 	// interface FRenderingCompositePass ---------
+
 	virtual void Process(FRenderingCompositePassContext& Context) override;
 	virtual void Release() override { delete this; }
 	FPooledRenderTargetDesc ComputeOutputDesc(EPassOutputId InPassOutputId) const override;
+
 private:
 	// @param InCylinderDistortion 0=none..1=full in percent, must be in that range
 	template <uint32 Quality, uint32 bTesselatedQuad> static FShader* SetShader(const FRenderingCompositePassContext& Context, const PaniniParams& PaniniConfig);
diff --git a/Engine/Source/Runtime/Renderer/Private/PostProcess/PostProcessing.cpp b/Engine/Source/Runtime/Renderer/Private/PostProcess/PostProcessing.cpp
--- a/Engine/Source/Runtime/Renderer/Private/PostProcess/PostProcessing.cpp
+++ b/Engine/Source/Runtime/Renderer/Private/PostProcess/PostProcessing.cpp
@@ -91,32 +91,6 @@ static TAutoConsoleVariable<int32> CVarRenderTargetSwitchWorkaround(
 	TEXT("We want this enabled (1) on all 32 bit iOS devices (implemented through DeviceProfiles)."),
 	ECVF_RenderThreadSafe);
 
-static TAutoConsoleVariable<float> CVarUpscalePaniniD(
-	TEXT("r.Upscale.Panini.D"),
-	0,
-	TEXT("Allow and configure to apply a panini distortion to the rendered image. Values between 0 and 1 allow to fade the effect (lerp).\n")
-	TEXT("Implementation from research paper \"Pannini: A New Projection for Rendering Wide Angle Perspective Images\"\n")
-	TEXT(" 0: off (default)\n")
-	TEXT(">0: enabled (requires an extra post processing pass if upsampling wasn't used - see r.ScreenPercentage)\n")
-	TEXT(" 1: Panini cylindrical stereographic projection"),
-	ECVF_RenderThreadSafe);
-
-static TAutoConsoleVariable<float> CVarUpscalePaniniS(
-	TEXT("r.Upscale.Panini.S"),
-	0,
-	TEXT("Panini projection's hard vertical compression factor.\n")
-	TEXT(" 0: no vertical compression factor (default)\n")
-	TEXT(" 1: Hard vertical compression"),
-	ECVF_RenderThreadSafe);
-
-static TAutoConsoleVariable<float> CVarUpscalePaniniScreenFit(
-	TEXT("r.Upscale.Panini.ScreenFit"),
-	1.0f,
-	TEXT("Panini projection screen fit effect factor (lerp).\n")
-	TEXT(" 0: fit vertically\n")
-	TEXT(" 1: fit horizontally (default)"),
-	ECVF_RenderThreadSafe);
-
 static TAutoConsoleVariable<int32> CVarUpscaleQuality(
 	TEXT("r.Upscale.Quality"),
 	3,
@@ -161,13 +135,12 @@ static TAutoConsoleVariable<int32> CVarTonemapperScreenPercentage(
 	TEXT("Experimental ScreenPercentage upscale integrated into tonemapper pass (if certain conditions apply e.g. no FXAA)\n")
 	TEXT(" if enabled both features are done in one pass (faster, affects post process passes after the tonemapper including material post process e.g. sharpen)\n")
 	TEXT(" 0: off, the features run in separate passes (default)\n")
-	TEXT(" 1: enabled if nothing prevents it which is safe but it might do the two pass method\n")
-	TEXT(" 2: enabled even if some other feature would prevent it (good for testing, likely rendering artifacts)"),
+	TEXT(" 1: enabled if nothing prevents it which is safe but it might do the two pass method"),
 	ECVF_Scalability | ECVF_RenderThreadSafe);
 
 static TAutoConsoleVariable<int32> CVarMotionBlurNew(
 	TEXT("r.MotionBlurNew"),
-	0,
+	1,
 	TEXT(""),
 	ECVF_RenderThreadSafe
 );
@@ -337,8 +310,7 @@ static FRenderingCompositeOutputRef RenderBloom(
 	return RenderGaussianBlur(Context, TEXT("BloomBlurX"), TEXT("BloomBlurY"), PreviousBloom, Size, Tint, Additive,CrossBloom);
 }
 
-// @return bDoScreenPercentageInTonemapper
-static bool AddTonemapper(
+static FRCPassPostProcessTonemap* AddTonemapper(
 	FPostprocessContext& Context,
 	const FRenderingCompositeOutputRef& BloomOutputCombined,
 	const FRenderingCompositeOutputRef& EyeAdaptation,
@@ -346,45 +318,9 @@ static bool AddTonemapper(
 {
 	const FEngineShowFlags& EngineShowFlags = Context.View.Family->EngineShowFlags;
 
-	int32 TonemapperScreenPercentage = CVarTonemapperScreenPercentage.GetValueOnRenderThread();
-
-	bool bDoScreenPercentageInTonemapper = TonemapperScreenPercentage != 0;
-
-	// only check if the feature is requested
-	if(bDoScreenPercentageInTonemapper)
-#if !(UE_BUILD_SHIPPING || UE_BUILD_TEST)
-	// useful only during development
-	if(TonemapperScreenPercentage != 2)
-#endif
-	{
-		// some features run after the tone mapper and cannot cope with this optimization
-
-		const bool bOverviewModeEnabled = EngineShowFlags.VisualizeBuffer && (Context.View.CurrentBufferVisualizationMode == NAME_None);
-
-		if (EngineShowFlags.VisualizeDOF ||
-			EngineShowFlags.VisualizeSSS ||
-			EngineShowFlags.HighResScreenshotMask ||
-			Context.View.FinalPostProcessSettings.AntiAliasingMethod == AAM_FXAA ||
-			bOverviewModeEnabled)
-		{
-			bDoScreenPercentageInTonemapper = false;
-		}
-
-		// todo: check for Panini projection
-
-		if(HasPostProcessMaterial(Context, BL_AfterTonemapping))
-		{
-			// Post process materials after the tone mapper should be avoided -
-			// they would run in 2D resolution (e.g. 1080p) which is often larger than 3D resolution
-			// Some effects don't work on the data generated by Tonemapper.ScreenPercentage e.g. sharpen
-			// and fixing the texture size differences depending in which pass we run are extensive.
-			bDoScreenPercentageInTonemapper = false;
-		}
-	}
-
 	FRenderingCompositePass* CombinedLUT = Context.Graph.RegisterPass(new(FMemStack::Get()) FRCPassPostProcessCombineLUTs(Context.View.GetShaderPlatform()));
 	const bool bDoEyeAdaptation = IsAutoExposureMethodSupported(Context.View.GetFeatureLevel(), EyeAdapationMethodId);
-	FRCPassPostProcessTonemap* PostProcessTonemap =	Context.Graph.RegisterPass(new(FMemStack::Get()) FRCPassPostProcessTonemap(Context.View, false, bDoScreenPercentageInTonemapper, bDoEyeAdaptation));
+	FRCPassPostProcessTonemap* PostProcessTonemap =	Context.Graph.RegisterPass(new(FMemStack::Get()) FRCPassPostProcessTonemap(Context.View, false, bDoEyeAdaptation));
 
 	PostProcessTonemap->SetInput(ePId_Input0, Context.FinalOutput);
 	PostProcessTonemap->SetInput(ePId_Input1, BloomOutputCombined);
@@ -393,7 +329,7 @@ static bool AddTonemapper(
 
 	Context.FinalOutput = FRenderingCompositeOutputRef(PostProcessTonemap);
 
-	return bDoScreenPercentageInTonemapper;
+	return PostProcessTonemap;
 }
 
 #if WITH_EDITOR
@@ -412,7 +348,7 @@ static void AddSelectionOutline(FPostprocessContext& Context)
 
 static void AddGammaOnlyTonemapper(FPostprocessContext& Context)
 {
-	FRenderingCompositePass* PostProcessTonemap = Context.Graph.RegisterPass(new(FMemStack::Get()) FRCPassPostProcessTonemap(Context.View, true, false, false/*eye*/));
+	FRenderingCompositePass* PostProcessTonemap = Context.Graph.RegisterPass(new(FMemStack::Get()) FRCPassPostProcessTonemap(Context.View, true, false/*eye*/));
 
 	PostProcessTonemap->SetInput(ePId_Input0, Context.FinalOutput);
 
@@ -1224,13 +1160,26 @@ void FPostProcessing::Process(FRHICommandListImmediate& RHICmdList, FViewInfo& V
 		FRenderingCompositeOutputRef BloomOutputCombined;
 		// not always valid
 		FRenderingCompositePass* VelocityFlattenPass = 0;
-
 		// in the following code some feature might set this to false
 		bool bAllowTonemapper = FeatureLevel >= ERHIFeatureLevel::SM4;
-
+		//
+		bool bStereoRenderingAndHMD = View.Family->EngineShowFlags.StereoRendering && View.Family->EngineShowFlags.HMDDistortion;
+		//
+		FRCPassPostProcessUpscale::PaniniParams PaniniConfig(View);
+		//
 		EStereoscopicPass StereoPass = View.StereoPass;
-
+		//
 		FSceneViewState* ViewState = (FSceneViewState*)Context.View.State;
+		//
+		bool bDoScreenPercentage;
+		{
+			//
+			bool bHMDWantsUpscale = bStereoRenderingAndHMD && GEngine->HMDDevice->NeedsUpscalePostProcessPass();
+			// Do not use upscale if SeparateRenderTarget is in use! (stereo rendering wants to control this)
+			bool bAllowScreenPercentage = bHMDWantsUpscale || !View.Family->EngineShowFlags.StereoRendering || (!View.Family->EngineShowFlags.HMDDistortion && !View.Family->bUseSeparateRenderTarget);
+			// is Upscale from a lower resolution needed and allowed
+			bDoScreenPercentage = bAllowScreenPercentage && (View.UnscaledViewRect != View.ViewRect);
+		}
 
 		{
 			if (FSceneRenderTargets::Get(RHICmdList).SeparateTranslucencyRT)
@@ -1256,10 +1205,9 @@ void FPostProcessing::Process(FRHICommandListImmediate& RHICmdList, FViewInfo& V
 		{
 			bAllowTonemapper = false;
 		}
+	
+		FRCPassPostProcessTonemap* Tonemapper = 0;
 
-		// if ScreenPercentage (render 3d in different resolution and upscale later) is already done
-		bool bScreenPercentageIsDone = false;
-		
 		// add the passes we want to add to the graph (commenting a line means the pass is not inserted into the graph) ---------
 
 		if (AllowFullPostProcessing(View, FeatureLevel))
@@ -1728,7 +1676,7 @@ void FPostProcessing::Process(FRHICommandListImmediate& RHICmdList, FViewInfo& V
 				}
 				else
 				{
-					bScreenPercentageIsDone = AddTonemapper(Context, BloomOutputCombined, AutoExposure.EyeAdaptation, AutoExposure.MethodId);
+					Tonemapper = AddTonemapper(Context, BloomOutputCombined, AutoExposure.EyeAdaptation, AutoExposure.MethodId);
 				}
 			}
 
@@ -1845,8 +1793,6 @@ void FPostProcessing::Process(FRHICommandListImmediate& RHICmdList, FViewInfo& V
 
 		AddGBufferVisualizationOverview(Context, SeparateTranslucency, HDRColor);
 
-		bool bStereoRenderingAndHMD = View.Family->EngineShowFlags.StereoRendering && View.Family->EngineShowFlags.HMDDistortion;
-		bool bHMDWantsUpscale = false;
 		if (bStereoRenderingAndHMD)
 		{
 			FRenderingCompositePass* Node = NULL;
@@ -1865,8 +1811,6 @@ void FPostProcessing::Process(FRHICommandListImmediate& RHICmdList, FViewInfo& V
 #endif
 			}
 
-			bHMDWantsUpscale = GEngine->HMDDevice->NeedsUpscalePostProcessPass();
-			
 			if(Node)
 			{
 				Node->SetInput(ePId_Input0, FRenderingCompositeOutputRef(Context.FinalOutput));
@@ -1895,21 +1839,25 @@ void FPostProcessing::Process(FRHICommandListImmediate& RHICmdList, FViewInfo& V
 
 		AddHighResScreenshotMask(Context, SeparateTranslucency);
 
-		if(!bScreenPercentageIsDone)
+		if(bDoScreenPercentage)
 		{
-			// 0=none..1=full
-			FRCPassPostProcessUpscale::PaniniParams PaniniConfig;
-
-			if (View.IsPerspectiveProjection() && !GEngine->StereoRenderingDevice.IsValid())
+			// Check if we can save the Upscale pass and do it in the Tonemapper to save performance
+			if(Tonemapper && !PaniniConfig.IsEnabled() && !Tonemapper->bDoGammaOnly)
 			{
-				PaniniConfig.D = FMath::Max(CVarUpscalePaniniD.GetValueOnRenderThread(), 0.0f);
-				PaniniConfig.S = CVarUpscalePaniniS.GetValueOnRenderThread();
-				PaniniConfig.ScreenFit = FMath::Max(CVarUpscalePaniniScreenFit.GetValueOnRenderThread(), 0.0f);
+				int32 TonemapperScreenPercentage = CVarTonemapperScreenPercentage.GetValueOnRenderThread();
+
+				if(TonemapperScreenPercentage != 0)
+				{
+					if(Context.FinalOutput.GetPass() == Tonemapper)
+					{
+						Tonemapper->bDoScreenPercentageInTonemapper = true;
+						// the following pass is no longer needed
+						bDoScreenPercentage = false;
+					}
+				}
 			}
 
-		// Do not use upscale if SeparateRenderTarget is in use!
-			if ((PaniniConfig.D > 0.01f || View.UnscaledViewRect != View.ViewRect) &&
-				(bHMDWantsUpscale || !View.Family->EngineShowFlags.StereoRendering || (!View.Family->EngineShowFlags.HMDDistortion && !View.Family->bUseSeparateRenderTarget)))
+			if (PaniniConfig.IsEnabled() || bDoScreenPercentage)
 			{
 				int32 UpscaleQuality = CVarUpscaleQuality.GetValueOnRenderThread();
 				UpscaleQuality = FMath::Clamp(UpscaleQuality, 0, 3);
@@ -1918,9 +1866,6 @@ void FPostProcessing::Process(FRHICommandListImmediate& RHICmdList, FViewInfo& V
 				Node->SetInput(ePId_Input1, FRenderingCompositeOutputRef(Context.FinalOutput)); // Point sampling.
 				Context.FinalOutput = FRenderingCompositeOutputRef(Node);
 			}
-
-			// not needed but more correct
-			bScreenPercentageIsDone = true;
 		}
 
 		// After the graph is built but before the graph is processed.
diff --git a/Engine/Source/Runtime/Slate/Private/Framework/Application/SlateApplication.cpp b/Engine/Source/Runtime/Slate/Private/Framework/Application/SlateApplication.cpp
--- a/Engine/Source/Runtime/Slate/Private/Framework/Application/SlateApplication.cpp
+++ b/Engine/Source/Runtime/Slate/Private/Framework/Application/SlateApplication.cpp
@@ -1312,6 +1312,8 @@ extern SLATECORE_API int32 bFoldTick;
 
 void FSlateApplication::Tick()
 {
+	FScopeLock SlateTickAccess(&SlateTickCriticalSection);
+
 	FPlatformMisc::BeginNamedEvent(FColor::Magenta, "Slate::Tick");
 
 	{
diff --git a/Engine/Source/Runtime/Slate/Private/Widgets/Input/SInputKeySelector.cpp b/Engine/Source/Runtime/Slate/Private/Widgets/Input/SInputKeySelector.cpp
new file mode 100644
--- /dev/null
+++ b/Engine/Source/Runtime/Slate/Private/Widgets/Input/SInputKeySelector.cpp
@@ -0,0 +1,204 @@
+// Copyright 1998-2015 Epic Games, Inc. All Rights Reserved.
+
+#include "SlatePrivatePCH.h"
+#include "SInputKeySelector.h"
+
+
+void SInputKeySelector::Construct( const FArguments& InArgs )
+{
+	SelectedKey = InArgs._SelectedKey;
+	OnKeySelected = InArgs._OnKeySelected;
+	Font = InArgs._Font;
+	ColorAndOpacity = InArgs._ColorAndOpacity;
+	OnIsSelectingKeyChanged = InArgs._OnIsSelectingKeyChanged;
+	bAllowModifierKeys = InArgs._AllowModifierKeys;
+
+	DefaultFont = FCoreStyle::Get().GetWidgetStyle<FTextBlockStyle>( "NormalText" ).Font;
+
+	bIsSelectingKey = false;
+
+	ChildSlot
+	[
+		SAssignNew(Button, SButton)
+		.ButtonStyle(InArgs._ButtonStyle)
+		.OnClicked(this, &SInputKeySelector::OnClicked)
+		[
+			SNew(STextBlock)
+			.Text(this, &SInputKeySelector::GetSelectedKeyText)
+			.Font(this, &SInputKeySelector::GetFont)
+			.Margin(Margin)
+			.ColorAndOpacity(this, &SInputKeySelector::GetSlateColorAndOpacity)
+			.Justification(ETextJustify::Center)
+		]
+	];
+}
+
+FText SInputKeySelector::GetSelectedKeyText() const
+{
+	if ( bIsSelectingKey )
+	{
+		return KeySelectionText;
+	}
+	else if ( SelectedKey.IsSet() )
+	{
+		// If the key in the chord is a modifier key, print it's display name directly since the FInputChord
+		// displays these as empty text.
+		return SelectedKey.Get().Key.IsModifierKey()
+			? SelectedKey.Get().Key.GetDisplayName()
+			: SelectedKey.Get().GetInputText();
+	}
+	return FText();
+}
+
+FInputChord SInputKeySelector::GetSelectedKey() const
+{
+	return SelectedKey.IsSet() ? SelectedKey.Get() : EKeys::Invalid;
+}
+
+void SInputKeySelector::SetSelectedKey( TAttribute<FInputChord> InSelectedKey )
+{
+	if ( SelectedKey.IdenticalTo(InSelectedKey) == false)
+	{
+		SelectedKey = InSelectedKey;
+		OnKeySelected.ExecuteIfBound( SelectedKey.IsSet() ? SelectedKey.Get() : FInputChord( EKeys::Invalid ) );
+	}
+}
+
+FSlateFontInfo SInputKeySelector::GetFont() const
+{
+	return Font.IsSet() ? Font.Get() : DefaultFont;
+}
+
+void SInputKeySelector::SetFont( TAttribute<FSlateFontInfo> InFont )
+{
+	Font = InFont;
+}
+
+FMargin SInputKeySelector::GetMargin() const
+{
+	return Margin.Get();
+}
+
+void SInputKeySelector::SetMargin( TAttribute<FMargin> InMargin )
+{
+	Margin = InMargin;
+}
+
+FSlateColor SInputKeySelector::GetSlateColorAndOpacity() const
+{
+	return FSlateColor(GetColorAndOpacity());
+}
+
+void SInputKeySelector::SetColorAndOpacity( TAttribute<FLinearColor> InColorAndOpacity )
+{
+	ColorAndOpacity = InColorAndOpacity;
+}
+
+void SInputKeySelector::SetButtonStyle(const FButtonStyle* ButtonStyle )
+{
+	Button->SetButtonStyle(ButtonStyle);
+}
+
+void SInputKeySelector::SetKeySelectionText( FText InKeySelectionText )
+{
+	KeySelectionText = InKeySelectionText;
+}
+
+void SInputKeySelector::SetAllowModifierKeys( bool bInAllowModifierKeys )
+{
+	bAllowModifierKeys = bInAllowModifierKeys;
+}
+
+bool SInputKeySelector::GetIsSelectingKey() const
+{
+	return bIsSelectingKey;
+}
+
+FReply SInputKeySelector::OnClicked()
+{
+	if ( bIsSelectingKey == false )
+	{
+		SetIsSelectingKey(true);
+		return FReply::Handled()
+			.SetUserFocus(SharedThis(this), EFocusCause::SetDirectly);
+	}
+	return FReply::Unhandled();
+}
+
+void SInputKeySelector::SelectKey( FKey Key, bool bShiftDown, bool bControllDown, bool bAltDown, bool bCommandDown )
+{
+	FInputChord NewSelectedKey = bAllowModifierKeys 
+		? FInputChord( Key, bShiftDown, bControllDown, bAltDown, bCommandDown ) 
+		: FInputChord( Key );
+	if ( SelectedKey.IsBound() == false )
+	{
+		SelectedKey.Set( NewSelectedKey );
+	}
+	OnKeySelected.ExecuteIfBound( NewSelectedKey );
+}
+
+void SInputKeySelector::SetIsSelectingKey( bool bInIsSelectingKey )
+{
+	if ( bIsSelectingKey != bInIsSelectingKey )
+	{
+		bIsSelectingKey = bInIsSelectingKey;
+		OnIsSelectingKeyChanged.ExecuteIfBound();
+	}
+}
+
+FReply SInputKeySelector::OnPreviewKeyDown( const FGeometry& MyGeometry, const FKeyEvent& InKeyEvent )
+{
+	// TODO: Add an argument to allow gamepad key selection.
+	if ( bIsSelectingKey && InKeyEvent.GetKey().IsGamepadKey() == false)
+	{
+		// While selecting keys handle all key downs to prevent contained controls from
+		// interfering with key selection.
+		return FReply::Handled();
+	}
+	return SWidget::OnPreviewKeyDown( MyGeometry, InKeyEvent );
+}
+
+FReply SInputKeySelector::OnKeyUp( const FGeometry& MyGeometry, const FKeyEvent& InKeyEvent )
+{
+	// TODO: Add an argument to allow gamepad key selection.
+	FKey KeyUp = InKeyEvent.GetKey();
+	EModifierKey::Type ModifierKey = EModifierKey::FromBools(
+		InKeyEvent.IsControlDown() && KeyUp != EKeys::LeftControl && KeyUp != EKeys::RightControl,
+		InKeyEvent.IsAltDown() && KeyUp != EKeys::LeftAlt && KeyUp != EKeys::RightAlt,
+		InKeyEvent.IsShiftDown() && KeyUp != EKeys::LeftShift && KeyUp != EKeys::RightShift,
+		InKeyEvent.IsCommandDown() && KeyUp != EKeys::LeftCommand && KeyUp != EKeys::RightCommand );
+
+	// Don't allow chords consisting of just modifier keys.
+	if ( bIsSelectingKey && KeyUp.IsGamepadKey() == false && ( KeyUp.IsModifierKey() == false || ModifierKey == EModifierKey::None ) )
+	{
+		SetIsSelectingKey( false );
+		SelectKey(
+			KeyUp,
+			ModifierKey == EModifierKey::Shift,
+			ModifierKey == EModifierKey::Control,
+			ModifierKey == EModifierKey::Alt,
+			ModifierKey == EModifierKey::Command);
+		return FReply::Handled();
+	}
+	return SWidget::OnPreviewKeyDown( MyGeometry, InKeyEvent );
+}
+
+FReply SInputKeySelector::OnPreviewMouseButtonDown( const FGeometry& MyGeometry, const FPointerEvent& MouseEvent )
+{
+	if ( bIsSelectingKey )
+	{
+		SetIsSelectingKey(false);
+		// TODO: Add options for enabling mouse modifiers.
+		SelectKey(MouseEvent.GetEffectingButton(), false, false, false, false);
+		return FReply::Handled();
+	}
+	return SWidget::OnPreviewMouseButtonDown( MyGeometry, MouseEvent );
+}
+
+void SInputKeySelector::OnFocusLost( const FFocusEvent& InFocusEvent )
+{
+	if ( bIsSelectingKey )
+	{
+		SetIsSelectingKey(false);
+	}
+}
\ No newline at end of file
diff --git a/Engine/Source/Runtime/Slate/Public/Framework/Application/SlateApplication.h b/Engine/Source/Runtime/Slate/Public/Framework/Application/SlateApplication.h
--- a/Engine/Source/Runtime/Slate/Public/Framework/Application/SlateApplication.h
+++ b/Engine/Source/Runtime/Slate/Public/Framework/Application/SlateApplication.h
@@ -1679,4 +1679,6 @@ private:
 
 	/** Delegate for post slate Tick */
 	FSlateTickEvent PostTickEvent;
+
+	FCriticalSection SlateTickCriticalSection;
 };
diff --git a/Engine/Source/Runtime/Slate/Public/Framework/Commands/InputChord.h b/Engine/Source/Runtime/Slate/Public/Framework/Commands/InputChord.h
--- a/Engine/Source/Runtime/Slate/Public/Framework/Commands/InputChord.h
+++ b/Engine/Source/Runtime/Slate/Public/Framework/Commands/InputChord.h
@@ -8,7 +8,7 @@
 typedef struct FInputChord FInputGesture;
 
 /** An Input Chord is a key and the modifier keys that are to be held with it. */
-USTRUCT()
+USTRUCT(BlueprintType)
 struct SLATE_API FInputChord
 {
 	GENERATED_USTRUCT_BODY()
diff --git a/Engine/Source/Runtime/Slate/Public/Widgets/Input/SInputKeySelector.h b/Engine/Source/Runtime/Slate/Public/Widgets/Input/SInputKeySelector.h
new file mode 100644
--- /dev/null
+++ b/Engine/Source/Runtime/Slate/Public/Widgets/Input/SInputKeySelector.h
@@ -0,0 +1,143 @@
+// Copyright 1998-2015 Epic Games, Inc. All Rights Reserved.
+
+#pragma once
+
+/**
+ * A widget for selecting keys or input chords.
+ */
+class SLATE_API SInputKeySelector : public SCompoundWidget
+{
+public:
+	DECLARE_DELEGATE_OneParam( FOnKeySelected, FInputChord )
+	DECLARE_DELEGATE( FOnIsSelectingKeyChanged )
+
+	SLATE_BEGIN_ARGS( SInputKeySelector )
+		: _SelectedKey( FInputChord(EKeys::Invalid) )
+		, _ButtonStyle( &FCoreStyle::Get().GetWidgetStyle<FButtonStyle>( "Button" ) )
+		, _KeySelectionText( NSLOCTEXT("DefaultKeySelectionText", "InputKeySelector", "...") )
+		, _AllowModifierKeys( true )
+	{}
+		/** The currently selected key */
+		SLATE_ATTRIBUTE( FInputChord, SelectedKey )
+
+		/** The font used to display the currently selected key. */
+		SLATE_ATTRIBUTE( FSlateFontInfo, Font )
+
+		/** The margin around the selected key text. */
+		SLATE_ATTRIBUTE( FMargin, Margin )
+
+		/** The color of the selected key text. */
+		SLATE_ATTRIBUTE( FLinearColor, ColorAndOpacity )
+
+		/** The style of the button used to enable key selection. */
+		SLATE_STYLE_ARGUMENT( FButtonStyle, ButtonStyle )
+
+		/** The text to display while selecting a new key. */
+		SLATE_ARGUMENT( FText, KeySelectionText )
+
+		/** When true modifier keys are captured in the selected key chord, otherwise they are ignored. */
+		SLATE_ARGUMENT( bool, AllowModifierKeys )
+
+		/** Occurs whenever a new key is selected. */
+		SLATE_EVENT( FOnKeySelected, OnKeySelected )
+
+		/** Occurs whenever key selection mode starts and stops. */
+		SLATE_EVENT( FOnIsSelectingKeyChanged, OnIsSelectingKeyChanged )
+
+	SLATE_END_ARGS()
+
+	void Construct( const FArguments& InArgs );
+
+public:
+
+	/** Gets the currently selected key chord. */
+	FInputChord GetSelectedKey() const;
+
+	/** Sets the currently selected key chord. */
+	void SetSelectedKey( TAttribute<FInputChord> InSelectedKey );
+
+	/** Sets the font which is used to display the currently selected key. */
+	void SetFont( TAttribute<FSlateFontInfo> InFont );
+
+	/** Sets the margin around the text used to display the currently selected key */
+	void SetMargin( TAttribute<FMargin> InMargin );
+
+	/** Sets the color of the text used to display the currently selected key. */
+	void SetColorAndOpacity( TAttribute<FLinearColor> InColorAndOpacity );
+
+	/** Sets the style of the button which is used enter key selection mode. */
+	void SetButtonStyle( const FButtonStyle* ButtonStyle );
+
+	/** Sets the text which is displayed when selecting a key. */
+	void SetKeySelectionText( FText InKeySelectionText );
+
+	/** When true modifier keys are captured in the selected key chord, otherwise they are ignored. */
+	void SetAllowModifierKeys( bool bInAllowModifierKeys );
+
+	/** Returns true whenever key selection mode is active, otherwise returns false. */
+	bool GetIsSelectingKey() const;
+
+public:
+
+	virtual FReply OnPreviewKeyDown( const FGeometry& MyGeometry, const FKeyEvent& InKeyEvent ) override;
+	virtual FReply OnKeyUp( const FGeometry& MyGeometry, const FKeyEvent& InKeyEvent ) override;
+	virtual FReply OnPreviewMouseButtonDown( const FGeometry& MyGeometry, const FPointerEvent& MouseEvent ) override;
+
+	virtual void OnFocusLost( const FFocusEvent& InFocusEvent ) override;
+	virtual bool SupportsKeyboardFocus() const { return true; }
+
+private:
+
+	/** Gets the display text for the currently selected key. */
+	FText GetSelectedKeyText() const;
+
+	/** Gets the font used to display the currently selected key. */
+	FSlateFontInfo GetFont() const;
+
+	/**  Gets the margin around the text used to display the currently selected key. */
+	FMargin GetMargin() const;
+
+	/** Gets the color of the text used to display the currently selected key. */
+	FSlateColor GetSlateColorAndOpacity() const;
+
+	/** Handles the OnClicked event from the button which enables key selection mode. */
+	FReply OnClicked();
+
+	/** Sets the currently selected key and invokes the associated events. */
+	void SelectKey( FKey Key, bool bShiftDown, bool bControllDown, bool bAltDown, bool bCommandDown );
+
+	/** Sets bIsSelectingKey and invokes the associated events. */
+	void SetIsSelectingKey(bool bInIsSelectingKey);
+
+private:
+
+	/** True when key selection mode is active. */
+	bool bIsSelectingKey;
+
+	/** The currently selected key chord. */
+	TAttribute<FInputChord> SelectedKey;
+
+	/** The font for the text used to display the currently selected key. */
+	TAttribute<FSlateFontInfo> Font;
+
+	/** The margin around the text used to display the currently selected key. */
+	TAttribute<FMargin> Margin;
+
+	/** The text to display when selecting keys. */
+	FText KeySelectionText;
+
+	/** When true modifier keys are recorded on the selected key chord, otherwise they are ignored. */
+	bool bAllowModifierKeys;
+
+	/** The default font used for the text which displays the currently selected key. */
+	FSlateFontInfo DefaultFont;
+
+	/** Delegate which is run any time a new key is selected. */
+	FOnKeySelected OnKeySelected;
+
+	/** Delegate which is run when key selection mode starts and stops. */
+	FOnIsSelectingKeyChanged OnIsSelectingKeyChanged;
+
+	/** The button which starts key key selection mode. */
+	TSharedPtr<SButton> Button;
+};
\ No newline at end of file
diff --git a/Engine/Source/Runtime/UMG/Private/Components/InputKeySelector.cpp b/Engine/Source/Runtime/UMG/Private/Components/InputKeySelector.cpp
new file mode 100644
--- /dev/null
+++ b/Engine/Source/Runtime/UMG/Private/Components/InputKeySelector.cpp
@@ -0,0 +1,91 @@
+// Copyright 1998-2015 Epic Games, Inc. All Rights Reserved.
+
+#include "UMGPrivatePCH.h"
+
+UInputKeySelector::UInputKeySelector( const FObjectInitializer& ObjectInitializer )
+{
+	SInputKeySelector::FArguments InputKeySelectorDefaults;
+	SelectedKey = InputKeySelectorDefaults._SelectedKey.Get();
+	ButtonStyle = InputKeySelectorDefaults._ButtonStyle;
+	bAllowModifierKeys = InputKeySelectorDefaults._AllowModifierKeys;
+}
+
+void UInputKeySelector::SetSelectedKey( FInputChord InSelectedKey )
+{
+	if ( MyInputKeySelector.IsValid() )
+	{
+		MyInputKeySelector->SetSelectedKey( InSelectedKey );
+	}
+	SelectedKey = InSelectedKey;
+}
+
+void UInputKeySelector::SetKeySelectionText( FText InKeySelectionText )
+{
+	if ( MyInputKeySelector.IsValid() )
+	{
+		MyInputKeySelector->SetKeySelectionText( InKeySelectionText );
+	}
+	KeySelectionText = InKeySelectionText;
+}
+
+void UInputKeySelector::SetAllowModifierKeys( bool bInAllowModifierKeys )
+{
+	if ( MyInputKeySelector.IsValid() )
+	{
+		MyInputKeySelector->SetAllowModifierKeys( bInAllowModifierKeys );
+	}
+	bAllowModifierKeys = bInAllowModifierKeys;
+}
+
+bool UInputKeySelector::GetIsSelectingKey() const
+{
+	return MyInputKeySelector.IsValid() ? MyInputKeySelector->GetIsSelectingKey() : false;
+}
+
+void UInputKeySelector::SetButtonStyle( const FButtonStyle* InButtonStyle )
+{
+	if ( MyInputKeySelector.IsValid() )
+	{
+		MyInputKeySelector->SetButtonStyle( ButtonStyle );
+	}
+	ButtonStyle = InButtonStyle;
+}
+
+void UInputKeySelector::SynchronizeProperties()
+{
+	Super::SynchronizeProperties();
+
+	MyInputKeySelector->SetSelectedKey( SelectedKey );
+	MyInputKeySelector->SetFont( Font );
+	MyInputKeySelector->SetMargin( Margin );
+	MyInputKeySelector->SetColorAndOpacity( ColorAndOpacity );
+	MyInputKeySelector->SetButtonStyle( ButtonStyle );
+	MyInputKeySelector->SetKeySelectionText( KeySelectionText );
+	MyInputKeySelector->SetAllowModifierKeys( bAllowModifierKeys );
+}
+
+TSharedRef<SWidget> UInputKeySelector::RebuildWidget()
+{
+	MyInputKeySelector = SNew(SInputKeySelector)
+		.SelectedKey( SelectedKey)
+		.Font(Font)
+		.Margin(Margin)
+		.ColorAndOpacity(ColorAndOpacity)
+		.ButtonStyle(ButtonStyle)
+		.KeySelectionText(KeySelectionText)
+		.AllowModifierKeys(bAllowModifierKeys)
+		.OnKeySelected( BIND_UOBJECT_DELEGATE( SInputKeySelector::FOnKeySelected, HandleKeySelected ) )
+		.OnIsSelectingKeyChanged( BIND_UOBJECT_DELEGATE( SInputKeySelector::FOnIsSelectingKeyChanged, HandleIsSelectingKeyChanged ) );
+	return MyInputKeySelector.ToSharedRef();
+}
+
+void UInputKeySelector::HandleKeySelected( FInputChord InSelectedKey )
+{
+	SelectedKey = InSelectedKey;
+	OnKeySelected.Broadcast(SelectedKey);
+}
+
+void UInputKeySelector::HandleIsSelectingKeyChanged()
+{
+	OnIsSelectingKeyChanged.Broadcast();
+}
\ No newline at end of file
diff --git a/Engine/Source/Runtime/UMG/Public/Components/InputKeySelector.h b/Engine/Source/Runtime/UMG/Public/Components/InputKeySelector.h
new file mode 100644
--- /dev/null
+++ b/Engine/Source/Runtime/UMG/Public/Components/InputKeySelector.h
@@ -0,0 +1,92 @@
+// Copyright 1998-2015 Epic Games, Inc. All Rights Reserved.
+
+#pragma once
+
+#include "SInputKeySelector.h"
+#include "InputKeySelector.generated.h"
+
+/** A widget for selecting a single key or a single key with a modifier. */
+UCLASS()
+class UMG_API UInputKeySelector : public UWidget
+{
+	GENERATED_UCLASS_BODY()
+
+public:
+	DECLARE_DYNAMIC_MULTICAST_DELEGATE_OneParam( FOnKeySelected, FInputChord, SelectedKey );
+	DECLARE_DYNAMIC_MULTICAST_DELEGATE( FOnIsSelectingKeyChanged );
+
+public:
+	/** The currently selected key chord. */
+	UPROPERTY( EditAnywhere, BlueprintReadOnly, Category = "Key Selection" )
+	FInputChord SelectedKey;
+
+	/** The font used to display the currently selected key. */
+	UPROPERTY( EditAnywhere, BlueprintReadOnly, Category = Appearance )
+	FSlateFontInfo Font;
+
+	/** The amount of blank space around the text used to display the currently selected key. */
+	UPROPERTY( EditAnywhere, BlueprintReadOnly, Category = Appearance )
+	FMargin Margin;
+
+	/** The color of the text used to display the currently selected key. */
+	UPROPERTY( EditAnywhere, BlueprintReadOnly, Category = Appearance )
+	FLinearColor ColorAndOpacity;
+
+	/** Sets the text which is displayed while selecting keys. */
+	UPROPERTY( EditAnywhere, BlueprintReadOnly, Category = Appearance )
+	FText KeySelectionText;
+
+	/** When true modifier keys such as control and alt are allowed in the */
+	/** input chord representing the selected key, if false modifier keys are ignored. */
+	UPROPERTY( EditAnywhere, BlueprintReadOnly, Category = "Key Selection" )
+	bool bAllowModifierKeys;
+
+public:
+	/** Called whenever a new key is selected by the user. */
+	UPROPERTY( BlueprintAssignable, Category = "Widget Event" )
+	FOnKeySelected OnKeySelected;
+
+	/** Called whenever the key selection mode starts or stops. */
+	UPROPERTY( BlueprintAssignable, Category = "Widget Event" )
+	FOnIsSelectingKeyChanged OnIsSelectingKeyChanged;
+
+	/** Sets the currently selected key. */
+	UFUNCTION( BlueprintCallable, Category = "Widget" )
+	void SetSelectedKey(FInputChord InSelectedKey);
+
+	/** Sets the text which is displayed while selecting keys. */
+	UFUNCTION( BlueprintCallable, Category = "Widget" )
+	void SetKeySelectionText( FText InKeySelectionText );
+
+	/** Sets whether or not modifier keys are allowed in the selected key. */
+	UFUNCTION(BlueprintCallable, Category = "Widget" )
+	void SetAllowModifierKeys(bool bInAllowModifierKeys );
+
+	/** Returns true if the widget is currently selecting a key, otherwise returns false. */
+	UFUNCTION( BlueprintCallable, Category = "Widget" )
+	bool GetIsSelectingKey() const;
+
+	/** Sets the style of the button used to start key selection mode. */
+	void SetButtonStyle(const FButtonStyle* ButtonStyle);
+
+	//~ Begin UWidget Interface
+	virtual void SynchronizeProperties() override;
+	//~ End UWidget Interface
+
+protected:
+	//~ Begin UWidget Interface
+	virtual TSharedRef<SWidget> RebuildWidget() override;
+	//~ End UWidget Interface
+
+private:
+	void HandleKeySelected( FInputChord SelectedKey );
+	void HandleIsSelectingKeyChanged();
+
+private:
+
+	/** The style for the button used to start key selection mode. */
+	const FButtonStyle* ButtonStyle;
+
+	/** The input key selector widget managed by this object. */
+	TSharedPtr<SInputKeySelector> MyInputKeySelector;
+};
\ No newline at end of file
diff --git a/Engine/Source/Runtime/UMG/Public/UMG.h b/Engine/Source/Runtime/UMG/Public/UMG.h
--- a/Engine/Source/Runtime/UMG/Public/UMG.h
+++ b/Engine/Source/Runtime/UMG/Public/UMG.h
@@ -121,6 +121,8 @@
 
 #include "NativeWidgetHost.h"
 
+#include "InputKeySelector.h"
+
 // Slate
 #include "SObjectWidget.h"
 
diff --git a/Engine/Source/Runtime/Windows/D3D11RHI/Private/D3D11Commands.cpp b/Engine/Source/Runtime/Windows/D3D11RHI/Private/D3D11Commands.cpp
--- a/Engine/Source/Runtime/Windows/D3D11RHI/Private/D3D11Commands.cpp
+++ b/Engine/Source/Runtime/Windows/D3D11RHI/Private/D3D11Commands.cpp
@@ -715,8 +715,19 @@ void FD3D11DynamicRHI::ValidateExclusiveDepthStencilAccess(FExclusiveDepthStenci
 		const bool bDstStencilWrite = CurrentDSVAccessType.IsStencilWrite();
 
 		// requested access is not possible, fix SetRenderTarget EExclusiveDepthStencil or request a different one
-		check(!bSrcDepthWrite || bDstDepthWrite);
-		check(!bSrcStencilWrite || bDstStencilWrite);
+		ensureMsgf(
+			!bSrcDepthWrite || bDstDepthWrite, 
+			TEXT("Expected: SrcDepthWrite := false or DstDepthWrite := true. Actual: SrcDepthWrite := %s or DstDepthWrite := %s"),
+			(bSrcDepthWrite) ? TEXT("true") : TEXT("false"),
+			(bDstDepthWrite) ? TEXT("true") : TEXT("false")
+			);
+
+		ensureMsgf(
+			!bSrcStencilWrite || bDstStencilWrite,
+			TEXT("Expected: SrcStencilWrite := false or DstStencilWrite := true. Actual: SrcStencilWrite := %s or DstStencilWrite := %s"),
+			(bSrcStencilWrite) ? TEXT("true") : TEXT("false"),
+			(bDstStencilWrite) ? TEXT("true") : TEXT("false")
+			);
 	}
 }
 
