Branch snapshot for CL 2835500

diff --git a/Engine/Source/Runtime/Online/OnlineSubsystem/Private/OnlineAsyncTaskManager.cpp b/Engine/Source/Runtime/Online/OnlineSubsystem/Private/OnlineAsyncTaskManager.cpp
--- a/Engine/Source/Runtime/Online/OnlineSubsystem/Private/OnlineAsyncTaskManager.cpp
+++ b/Engine/Source/Runtime/Online/OnlineSubsystem/Private/OnlineAsyncTaskManager.cpp
@@ -9,7 +9,6 @@ int32 FOnlineAsyncTaskManager::InvocationCount = 0;
 #define POLLING_INTERVAL_MS 50
 
 FOnlineAsyncTaskManager::FOnlineAsyncTaskManager() :
-	ActiveTask(nullptr),
 	WorkEvent(nullptr),
 	PollingInterval(POLLING_INTERVAL_MS),
 	bRequestingExit(0),
@@ -39,6 +38,8 @@ uint32 FOnlineAsyncTaskManager::Run(void)
 
 	do 
 	{
+		int32 InitialQueueSize = 0;
+		int32 CurrentQueueSize = 0;
 		FOnlineAsyncTask* Task = NULL;
 
 		// Wait for a trigger event to start work
@@ -46,6 +47,11 @@ uint32 FOnlineAsyncTaskManager::Run(void)
 		if (!bRequestingExit)
 		{
 			SCOPE_CYCLE_COUNTER(STAT_Online_Async);
+			{
+				FScopeLock LockInQueue(&InQueueLock);
+				InitialQueueSize = InQueue.Num();
+				SET_DWORD_STAT(STAT_Online_AsyncTasks, InitialQueueSize);
+			}
 			
 			double TimeElapsed = 0;
 
@@ -91,38 +97,50 @@ uint32 FOnlineAsyncTaskManager::Run(void)
 			}
 
 			// Now process the serial "In" queue
-			{
-				FScopeLock LockActiveTask(&ActiveTaskLock);
-				Task = ActiveTask;
-			}
-			if (Task)
-			{
-				Task->Tick();
-	
-				if (Task->IsDone())
+			do
+			{	
+				Task = NULL;
+
+				// Grab the current task from the queue
 				{
-					if (Task->WasSuccessful())
+					FScopeLock LockInQueue(&InQueueLock);
+					CurrentQueueSize = InQueue.Num();
+					if (CurrentQueueSize > 0)
 					{
-						UE_LOG(LogOnline, Verbose, TEXT("Async task '%s' succeeded in %f seconds"),
-							*Task->ToString(),
-							Task->GetElapsedTime());
+						Task = InQueue[0];
 					}
-					else
+				}
+
+				if (Task)
+				{
+					Task->Tick();
+	
+					if (Task->IsDone())
 					{
-						UE_LOG(LogOnline, Warning, TEXT("Async task '%s' failed in %f seconds"),
-							*Task->ToString(),
-							Task->GetElapsedTime());
-					}
+						if (Task->WasSuccessful())
+						{
+							UE_LOG(LogOnline, Verbose, TEXT("Async task '%s' succeeded in %f seconds"),
+								*Task->ToString(),
+								Task->GetElapsedTime());
+						}
+						else
+						{
+							UE_LOG(LogOnline, Warning, TEXT("Async task '%s' failed in %f seconds"),
+								*Task->ToString(),
+								Task->GetElapsedTime());
+						}
 						
-					// Task is done, add to the outgoing queue
-					AddToOutQueue(Task);
-					
+						// Task is done, remove from the incoming queue and add to the outgoing queue
+						PopFromInQueue();
+						AddToOutQueue(Task);
+					}
+					else
 					{
-						FScopeLock LockActiveTask(&ActiveTaskLock);
-						ActiveTask = nullptr;
+						Task = NULL;
 					}
 				}
 			}
+			while (Task != NULL);
 		}
 	} 
 	while (!bRequestingExit);
@@ -149,6 +167,16 @@ void FOnlineAsyncTaskManager::AddToInQueue(FOnlineAsyncTask* NewTask)
 {
 	FScopeLock Lock(&InQueueLock);
 	InQueue.Add(NewTask);
+	WorkEvent->Trigger();
+}
+
+void FOnlineAsyncTaskManager::PopFromInQueue()
+{
+	// assert if not game thread
+	check(FPlatformTLS::GetCurrentThreadId() == OnlineThreadId || !FPlatformProcess::SupportsMultithreading());
+
+	FScopeLock Lock(&InQueueLock);
+	InQueue.RemoveAt(0);
 }
 
 void FOnlineAsyncTaskManager::AddToOutQueue(FOnlineAsyncItem* CompletedItem)
@@ -159,8 +187,6 @@ void FOnlineAsyncTaskManager::AddToOutQueue(FOnlineAsyncItem* CompletedItem)
 
 void FOnlineAsyncTaskManager::AddToParallelTasks(FOnlineAsyncTask* NewTask)
 {
-	NewTask->Initialize();
-
 	FScopeLock LockParallelTasks(&ParallelTasksLock);
 
 	ParallelTasks.Add( NewTask );
@@ -204,93 +230,62 @@ void FOnlineAsyncTaskManager::GameTick()
 		}
 	}
 	while (Item != NULL);
-
-	int32 QueueSize = 0;
-	bool bHasActiveTask = false;
-	{
-		{
-			FScopeLock LockInQueue(&InQueueLock);
-			QueueSize = InQueue.Num();
-		}
-		{
-			FScopeLock LockActiveTask(&ActiveTaskLock);
-			if (ActiveTask != nullptr)
-			{
-				++QueueSize;
-				bHasActiveTask = true;
-			}
-		}
-
-		if (!bHasActiveTask && QueueSize > 0)
-		{
-			// Grab the current task from the queue
-			FOnlineAsyncTask* Task = nullptr;
-			{
-				FScopeLock LockInQueue(&InQueueLock);
-				Task = InQueue[0];
-				InQueue.RemoveAt(0);
-			}
-			{
-				FScopeLock LockActiveTask(&ActiveTaskLock);
-				ActiveTask = Task;
-				ActiveTask->Initialize();
-			}
-			WorkEvent->Trigger();
-		}
-	}
-
-	SET_DWORD_STAT(STAT_Online_AsyncTasks, QueueSize);
 }
 
 void FOnlineAsyncTaskManager::Tick()
 {
+	// parallel Q.
+	FOnlineAsyncTask* Task = NULL;
+
 	// Tick Online services ( possibly callbacks ). 
 	OnlineTick();
 
-	{
-		// Tick all the parallel tasks - Tick unrelated tasks together. 
+	// Tick all the parallel tasks - Tick unrelated tasks together. 
 
-		// Create a copy of existing tasks. 
-		TArray<FOnlineAsyncTask*> CopyParallelTasks = ParallelTasks;
+	// Create a copy of existing tasks. 
+	TArray<FOnlineAsyncTask*> CopyParallelTasks = ParallelTasks;
 
-		// parallel Q.
-		FOnlineAsyncTask* Task = NULL;
+	// Iterate. 
+	for (auto it = CopyParallelTasks.CreateIterator(); it; ++it)
+	{
+		Task = *it;
+		Task->Tick();
 
-		// Iterate. 
-		for (auto it = CopyParallelTasks.CreateIterator(); it; ++it)
+		if (Task->IsDone())
 		{
-			Task = *it;
-			Task->Tick();
-
-			if (Task->IsDone())
-			{
-				UE_LOG(LogOnline, Log, TEXT("Async parallel Task '%s' completed in %f seconds with %d (Parallel)"),
-					*Task->ToString(),
-					Task->GetElapsedTime(),
-					Task->WasSuccessful());
-
-				// Task is done, fixup the original parallel task queue. 
-				RemoveFromParallelTasks(Task);
-				AddToOutQueue(Task);
-			}
+			UE_LOG(LogOnline, Log, TEXT("Async parallel Task '%s' completed in %f seconds with %d (Parallel)"),
+				*Task->ToString(),
+				Task->GetElapsedTime(),
+				Task->WasSuccessful());
+
+			// Task is done, fixup the original parallel task queue. 
+			RemoveFromParallelTasks(Task);
+			AddToOutQueue(Task);
 		}
 	}
 
 	// Serial Q.
-	if (ActiveTask == nullptr)
+	if (InQueue.Num() == 0)
+		return;
+
+	// Pick up the first element in the queue ( "current task" ). 
+	Task = InQueue[0];
+
+	// Needed ?
+	if (Task == NULL)
 		return;
 
-	ActiveTask->Tick();
+	Task->Tick();
 
-	if (ActiveTask->IsDone())
+	if (Task->IsDone())
 	{
 		UE_LOG(LogOnline, Log, TEXT("Async serial task '%s' completed in %f seconds with %d"),
-			*ActiveTask->ToString(),
-			ActiveTask->GetElapsedTime(),
-			ActiveTask->WasSuccessful());
+			*Task->ToString(),
+			Task->GetElapsedTime(),
+			Task->WasSuccessful());
 
-		// Task is done, add to the outgoing queue
-		AddToOutQueue(ActiveTask);
-		ActiveTask = nullptr;
+		// Task is done, remove from the incoming queue and add to the outgoing queue
+		PopFromInQueue();
+		AddToOutQueue(Task);
 	}
 }
diff --git a/Engine/Source/Runtime/Online/OnlineSubsystem/Private/OnlineKeyValuePair.cpp b/Engine/Source/Runtime/Online/OnlineSubsystem/Private/OnlineKeyValuePair.cpp
--- a/Engine/Source/Runtime/Online/OnlineSubsystem/Private/OnlineKeyValuePair.cpp
+++ b/Engine/Source/Runtime/Online/OnlineSubsystem/Private/OnlineKeyValuePair.cpp
@@ -753,10 +753,6 @@ bool FVariantData::operator==(const FVariantData& Other) const
 				return (Value.AsBlob.BlobSize == Other.Value.AsBlob.BlobSize) && 
 						 (FMemory::Memcmp(Value.AsBlob.BlobData, Other.Value.AsBlob.BlobData, Value.AsBlob.BlobSize) == 0);
 			}
-		case EOnlineKeyValuePairDataType::Bool:
-			{
-				return Value.AsBool == Other.Value.AsBool;
-			}
 		}
 	}
 	return false;
diff --git a/Engine/Source/Runtime/Online/OnlineSubsystem/Public/Interfaces/IMessageSanitizerInterface.h b/Engine/Source/Runtime/Online/OnlineSubsystem/Public/Interfaces/IMessageSanitizerInterface.h
--- a/Engine/Source/Runtime/Online/OnlineSubsystem/Public/Interfaces/IMessageSanitizerInterface.h
+++ b/Engine/Source/Runtime/Online/OnlineSubsystem/Public/Interfaces/IMessageSanitizerInterface.h
@@ -2,12 +2,11 @@
 
 #pragma once
 
-DECLARE_DELEGATE_TwoParams(FOnMessageProcessed, bool /*Process success*/, const FString& /*SanitizedMessage*/);
-DECLARE_DELEGATE_TwoParams(FOnMessageArrayProcessed, bool /*Process success*/, const TArray<FString>& /*SanitizedMessages*/);
+DECLARE_DELEGATE_OneParam(FOnMessageProcessed, const FString& /*SanitizedMessage*/);
 
 struct FSanitizeMessage
 {
-	FSanitizeMessage(const FString& InRawMessage, FOnMessageProcessed InProcessCompleteDelegate)
+	FSanitizeMessage(FString InRawMessage, FOnMessageProcessed InProcessCompleteDelegate)
 		: RawMessage(InRawMessage)
 		, CompletionDelegate(InProcessCompleteDelegate)
 	{}
@@ -16,16 +15,6 @@ struct FSanitizeMessage
 	FOnMessageProcessed CompletionDelegate;
 };
 
-struct FMultiPartMessage
-{
-	TArray<FString> AlreadyProcessedMessages;
-	TArray<int32> AlreadyProcessedIndex;
-	TArray<FString> SanitizedMessages;
-	FString MessageToSanatize;
-	bool bCompleted;
-	int32 PartNumber;
-};
-
 class IMessageSanitizer
 {
 protected:
@@ -33,8 +22,7 @@ protected:
 
 public:
 	virtual ~IMessageSanitizer() {};
-	virtual void SanitizeMessage(const FString& Message, FOnMessageProcessed CompletionDelegate) = 0;
-	virtual void SanitizeMessageArray(const TArray<FString>& MessageArray, FOnMessageArrayProcessed CompletionDelegate) = 0;
+	virtual void SanitizeMessage(FString Message, FOnMessageProcessed CompletionDelegate) = 0;
 };
 
 typedef TSharedPtr<IMessageSanitizer, ESPMode::ThreadSafe> IMessageSanitizerPtr;
diff --git a/Engine/Source/Runtime/Online/OnlineSubsystem/Public/Interfaces/OnlinePartyInterface.h b/Engine/Source/Runtime/Online/OnlineSubsystem/Public/Interfaces/OnlinePartyInterface.h
--- a/Engine/Source/Runtime/Online/OnlineSubsystem/Public/Interfaces/OnlinePartyInterface.h
+++ b/Engine/Source/Runtime/Online/OnlineSubsystem/Public/Interfaces/OnlinePartyInterface.h
@@ -8,14 +8,6 @@
 #define F_PREFIX(TypeToPrefix) F##TypeToPrefix
 #define PARTY_DECLARE_DELEGATETYPE(Type) typedef F##Type::FDelegate F##Type##Delegate
 
-ONLINESUBSYSTEM_API DECLARE_LOG_CATEGORY_EXTERN(LogOnlineParty, Display, All);
-
-#define ONLINEPARTY_LOG_PREFIX TEXT("OSS::Party: ")
-#define UE_LOG_ONLINEPARTY(Verbosity, Format, ...) \
-{ \
-	UE_LOG(LogOnlineParty, Verbosity, TEXT("%s%s"), ONLINEPARTY_LOG_PREFIX, *FString::Printf(Format, ##__VA_ARGS__)); \
-}
-
 
 /**
  * Party member user info returned by IOnlineParty interface
@@ -28,15 +20,10 @@ public:
 /**
  * Data associated with the entire party
  */
-class ONLINESUBSYSTEM_API FOnlinePartyData : public TSharedFromThis<FOnlinePartyData>
+class FOnlinePartyData : public TSharedFromThis<FOnlinePartyData>
 {
 public:
-	FOnlinePartyData() :
-		TotalBytes(0),
-		TotalEffectiveBytes(0),
-		TotalPackets(0),
-		RevisionCount(0)
-		{}
+	FOnlinePartyData() {}
 	~FOnlinePartyData() {}
 
 	/**
@@ -69,93 +56,14 @@ public:
 	void SetAttribute(const FString& AttrName, const FVariantData& AttrValue)
 	{
 		FVariantData& NewAttrValue = KeyValAttrs.FindOrAdd(AttrName);
-		if (NewAttrValue != AttrValue)
-		{
-			NewAttrValue = AttrValue;
-			DirtyKeys.Add(AttrName);
-		}
-	}
-
-	/**
-	 * Clear the attributes map
-	 */
-	void ClearAttributes()
-	{
-		KeyValAttrs.Empty();
-		DirtyKeys.Empty();
-	}
-
-	/** 
-	 * Clear the dirty keys set, called after successfully sending an update of the dirty elements
-	 */
-	void ClearDirty()
-	{
-		DirtyKeys.Empty();
-	}
-
-	/** 
-	 * Increment the stat tracking variables on packet sent
-	 * 
-	 * @param PacketSize - size of the packet generated
-	 * @param NumRecipients - number of recipients the packet was sent to
-	 * @param bIncrementRevision - this packet was a dirty packet update so we should increment the revision
-	 */
-	void OnPacketSent(int32 PacketSize, int32 NumRecipients, bool bIncrementRevision) const
-	{
-		TotalPackets++;
-		TotalBytes += PacketSize;
-		TotalEffectiveBytes += PacketSize * NumRecipients;
-		if (bIncrementRevision)
-		{
-			++RevisionCount;
-		}
+		NewAttrValue = AttrValue;
 	}
 
-	/** 
-	 * Generate a JSON packet containing all key-value attributes
-	 * 
-	 * @param JsonString - [out] string containing the resulting JSON output
-	 * 
-	 */
-	void ToJsonFull(FString& JsonString) const;
-	
-	/** 
-	 * Generate a JSON packet containing only the dirty key-value attributes for a delta update
-	 *
-	 * @param JsonString - [out] string containing the resulting JSON output
-	 *
-	 */
-	void ToJsonDirty(FString& JsonString) const;
-
-	/** 
-	 * Update attributes from a JSON packet
-	 *
-	 * @param JsonString - string containing the JSON packet
-	 *
-	 */
+	void ToJson(FString& JsonString) const;
 	void FromJson(const FString& JsonString);
 
-	/** Accessor functions for KeyValAttrs map */
-	FOnlineKeyValuePairs<FString, FVariantData>& GetKeyValAttrs() { return KeyValAttrs; }
-	const FOnlineKeyValuePairs<FString, FVariantData>& GetKeyValAttrs() const { return KeyValAttrs; }
-
-	/** Stat tracking variables */
-	/** Total number of bytes generated by calls to ToJsonFull and ToJsonDirty */
-	mutable int32 TotalBytes;
-	/** Total number of bytes generated by calls to ToJsonFull and ToJsonDirty, multiplied by the number of recipients the packet was sent to */
-	mutable int32 TotalEffectiveBytes;
-	/** Total number of packets generated by calls to ToJsonFull and ToJsonDirty */
-	mutable int32 TotalPackets;
-
-	/** Id representing number of updates sent, useful for determining if a client has missed an update */
-	mutable int32 RevisionCount;
-
-private:
 	/** map of key/val attributes that represents the data */
 	FOnlineKeyValuePairs<FString, FVariantData>  KeyValAttrs;
-
-	/** set of which fields are dirty and need to transmitted */
-	TSet<FString> DirtyKeys;
 };
 
 /**
@@ -1880,18 +1788,75 @@ inline FString ToDebugString(const IOnlinePartyJoinInfo& JoinInfo)
 inline FString ToDebugString(const FOnlinePartyData& PartyData)
 {
 	FString Result;
-	
-	int32 TotalBytesPerSec = PartyData.TotalPackets ? (PartyData.TotalBytes / PartyData.TotalPackets) : 0;
-	int32 TotalEffectiveBytesPerSec = PartyData.TotalPackets ? (PartyData.TotalEffectiveBytes / PartyData.TotalPackets) : 0;
+	bool bAddComma = false;
+	for (auto Iterator : PartyData.KeyValAttrs)
+	{
+		if (bAddComma)
+		{
+			Result += TEXT(", ");
+		}
+		else
+		{
+			bAddComma = true;
+		}
+		Result += FString::Printf(TEXT("[%s=%s]"), *(Iterator.Key), *(Iterator.Value.ToString()));
+	}
+	return Result;
+}
 
-	Result += FString::Printf(TEXT("%dB [%d B/pkt], %dB [%d B/pkt], Rev: %d"), 
-		PartyData.TotalBytes, TotalBytesPerSec,
-		PartyData.TotalEffectiveBytes, TotalEffectiveBytesPerSec, 
-		PartyData.RevisionCount);
+inline void FOnlinePartyData::ToJson(FString& JsonString) const
+{
+	JsonString.Empty();
 
-	for (auto Iterator : PartyData.GetKeyValAttrs())
+	// iterate over key/val attrs and convert each entry to a json string
+	TSharedRef<FJsonObject> JsonObject(new FJsonObject());
+	TArray<TSharedPtr<FJsonValue> > JsonProperties;
+	for (auto Iterator : KeyValAttrs)
 	{
-		Result += FString::Printf(TEXT(",[%s=%s]"), *(Iterator.Key), *(Iterator.Value.ToString()));
+		const FString& PropertyName = Iterator.Key;
+		const FVariantData& PropertyValue = Iterator.Value;
+
+		TSharedRef<FJsonObject> PropertyJson = PropertyValue.ToJson();
+		PropertyJson->SetStringField(TEXT("Name"), *PropertyName);
+		JsonProperties.Add(MakeShareable(new FJsonValueObject(PropertyJson)));
+	}
+	JsonObject->SetArrayField(TEXT("Attrs"), JsonProperties);
+
+	auto JsonWriter = TJsonWriterFactory<TCHAR, TCondensedJsonPrintPolicy<TCHAR> >::Create(&JsonString);
+	FJsonSerializer::Serialize(JsonObject, JsonWriter);
+	JsonWriter->Close();
+}
+
+inline void FOnlinePartyData::FromJson(const FString& JsonString)
+{
+	KeyValAttrs.Empty();
+
+	// json string to key/val attrs
+	TSharedPtr<FJsonObject> JsonObject;
+	TSharedRef<TJsonReader<> > JsonReader = TJsonReaderFactory<>::Create(JsonString);
+	if (FJsonSerializer::Deserialize(JsonReader, JsonObject) &&
+		JsonObject.IsValid())
+	{
+		if (JsonObject->HasTypedField<EJson::Array>(TEXT("Attrs")))
+		{
+			const TArray<TSharedPtr<FJsonValue> >& JsonProperties = JsonObject->GetArrayField(TEXT("Attrs"));
+			for (auto JsonPropertyValue : JsonProperties)
+			{
+				TSharedPtr<FJsonObject> JsonPropertyObject = JsonPropertyValue->AsObject();
+				if (JsonPropertyObject.IsValid())
+				{
+					FString PropertyName;
+					if (JsonPropertyObject->TryGetStringField(TEXT("Name"), PropertyName) &&
+						!PropertyName.IsEmpty())
+					{
+						FVariantData PropertyData;
+						if (PropertyData.FromJson(JsonPropertyObject.ToSharedRef()))
+						{
+							KeyValAttrs.Add(PropertyName, PropertyData);
+						}
+					}
+				}
+			}
+		}
 	}
-	return Result;
 }
diff --git a/Engine/Source/Runtime/Online/OnlineSubsystem/Public/Interfaces/OnlinePurchaseInterface.h b/Engine/Source/Runtime/Online/OnlineSubsystem/Public/Interfaces/OnlinePurchaseInterface.h
--- a/Engine/Source/Runtime/Online/OnlineSubsystem/Public/Interfaces/OnlinePurchaseInterface.h
+++ b/Engine/Source/Runtime/Online/OnlineSubsystem/Public/Interfaces/OnlinePurchaseInterface.h
@@ -4,7 +4,6 @@
 #include "OnlineDelegateMacros.h"
 #include "OnlineStoreInterfaceV2.h"
 #include "OnlineEntitlementsInterface.h"
-#include "OnlineError.h"
 
 /**
  * Info needed for checkout
@@ -70,11 +69,6 @@ namespace EPurchaseTransactionState
 class FPurchaseReceipt
 {
 public:
-	FPurchaseReceipt()
-	: TransactionState(EPurchaseTransactionState::NotStarted)
-	{
-	}
-
 	struct FLineItemInfo
 	{
 		/** The platform identifier of this purchase type */
@@ -85,8 +79,6 @@ public:
 
 		/** platform-specific opaque validation info (required to verify UniqueId belongs to this account) */
 		FString ValidationInfo;
-
-		inline bool IsRedeemable() const { return !ValidationInfo.IsEmpty(); }
 	};
 
 	/**
@@ -148,17 +140,17 @@ public:
 /**
  * Delegate called when checkout process completes
  */
-DECLARE_DELEGATE_TwoParams(FOnPurchaseCheckoutComplete, const FOnlineError& /*Result*/, const TSharedRef<FPurchaseReceipt>& /*Receipt*/);
+DECLARE_DELEGATE_FourParams(FOnPurchaseCheckoutComplete, const FUniqueNetId& /*UserId*/, bool /*bWasSuccessful*/, const TSharedRef<FPurchaseReceipt>& /*Receipt*/, const FString& /*Error*/);
 
 /**
  * Delegate called when code redemption process completes
  */
-DECLARE_DELEGATE_TwoParams(FOnPurchaseRedeemCodeComplete, const FOnlineError& /*Result*/, const TSharedRef<FPurchaseReceipt>& /*Receipt*/);
+DECLARE_DELEGATE_FourParams(FOnPurchaseRedeemCodeComplete, const FUniqueNetId& /*UserId*/, bool /*bWasSuccessful*/, const TSharedRef<FPurchaseReceipt>& /*Receipt*/, const FString& /*Error*/);
 
 /**
  * Delegate called when query receipt process completes
  */
-DECLARE_DELEGATE_OneParam(FOnQueryReceiptsComplete, const FOnlineError& /*Result*/);
+DECLARE_DELEGATE_ThreeParams(FOnQueryReceiptsComplete, const FUniqueNetId& /*UserId*/, bool /*bWasSuccessful*/, const FString& /*Error*/);
 
 /**
  *	IOnlinePurchase - Interface for IAP (In App Purchases) services
@@ -182,26 +174,32 @@ public:
 	 *
 	 * @param UserId user initiating the request
 	 * @param CheckoutRequest info needed for the checkout request
-	 * @param Delegate completion callback (guaranteed to be called)
+	 * @param Delegate completion callback
+	 *
+	 * @return true if user can make a purchase
 	 */
-	virtual void Checkout(const FUniqueNetId& UserId, const FPurchaseCheckoutRequest& CheckoutRequest, const FOnPurchaseCheckoutComplete& Delegate) = 0;
+	virtual bool Checkout(const FUniqueNetId& UserId, const FPurchaseCheckoutRequest& CheckoutRequest, const FOnPurchaseCheckoutComplete& Delegate = FOnPurchaseCheckoutComplete()) = 0;
 
 	/**
 	 * Initiate the checkout process for obtaining offers via code redemption
 	 *
 	 * @param UserId user initiating the request
 	 * @param RedeemCodeRequest info needed for the redeem request
-	 * @param Delegate completion callback (guaranteed to be called)
+	 * @param Delegate completion callback
+	 *
+	 * @return true if user can make a purchase
 	 */
-	virtual void RedeemCode(const FUniqueNetId& UserId, const FRedeemCodeRequest& RedeemCodeRequest, const FOnPurchaseRedeemCodeComplete& Delegate) = 0;
+	virtual bool RedeemCode(const FUniqueNetId& UserId, const FRedeemCodeRequest& RedeemCodeRequest, const FOnPurchaseRedeemCodeComplete& Delegate = FOnPurchaseRedeemCodeComplete()) = 0;
 
 	/**
 	 * Query for all of the user's receipts from prior purchases
 	 *
 	 * @param UserId user initiating the request
-	 * @param Delegate completion callback (guaranteed to be called)
+	 * @param Delegate completion callback
+	 *
+	 * @return true if user can make a purchase
 	 */
-	virtual void QueryReceipts(const FUniqueNetId& UserId, const FOnQueryReceiptsComplete& Delegate) = 0;
+	virtual bool QueryReceipts(const FUniqueNetId& UserId, const FOnQueryReceiptsComplete& Delegate = FOnQueryReceiptsComplete()) = 0;
 
 	/**
 	 * Get list of cached receipts for user (includes transactions currently being processed)
diff --git a/Engine/Source/Runtime/Online/OnlineSubsystem/Public/Interfaces/OnlineStoreInterfaceV2.h b/Engine/Source/Runtime/Online/OnlineSubsystem/Public/Interfaces/OnlineStoreInterfaceV2.h
--- a/Engine/Source/Runtime/Online/OnlineSubsystem/Public/Interfaces/OnlineStoreInterfaceV2.h
+++ b/Engine/Source/Runtime/Online/OnlineSubsystem/Public/Interfaces/OnlineStoreInterfaceV2.h
@@ -43,7 +43,7 @@ public:
 	/** @return FText suitable for localized display */
 	virtual FText GetDisplayPrice() const
 	{
-		return FText::AsCurrencyBase(NumericPrice, CurrencyCode);
+		return FText::AsCurrency(NumericPrice / 100.0, CurrencyCode);
 	}
 
 	/** @return True if offer can be purchased */
diff --git a/Engine/Source/Runtime/Online/OnlineSubsystem/Public/OnlineAsyncTaskManager.h b/Engine/Source/Runtime/Online/OnlineSubsystem/Public/OnlineAsyncTaskManager.h
--- a/Engine/Source/Runtime/Online/OnlineSubsystem/Public/OnlineAsyncTaskManager.h
+++ b/Engine/Source/Runtime/Online/OnlineSubsystem/Public/OnlineAsyncTaskManager.h
@@ -104,11 +104,6 @@ public:
 	{
 	}
 
-	/**
-	 * Initialize the task - called on game thread when queued
-	 */
-	virtual void Initialize() {}
-
 	/**
 	 * Check the state of the async task
 	 * @return true if complete, false otherwise
@@ -128,7 +123,7 @@ public:
 	virtual void Tick()
 	{
 		// assert that we're not on the game thread
-		check(!IsInGameThread() || !FPlatformProcess::SupportsMultithreading());
+		check(!IsInGameThread());
 	}
 };
 
@@ -229,13 +224,7 @@ protected:
 	TArray<FOnlineAsyncTask*> InQueue;
 	/** Critical section for thread safe operation of the event in queue */
 	FCriticalSection InQueueLock;
-
-	/** blah */
-	FOnlineAsyncTask* ActiveTask;
-
-	/** blah */
-	FCriticalSection ActiveTaskLock;
-
+	
 	/** This queue is for tasks that are safe to run in parallel with one another */
 	TArray<FOnlineAsyncTask*> ParallelTasks;
 	/** Critical section for thread safe operation of the list */
@@ -311,6 +300,11 @@ public:
 	 */
 	void AddToInQueue(FOnlineAsyncTask* NewTask);
 
+	/**
+	 *	Remove the current async task from the queue
+	 */
+	void PopFromInQueue();
+
 	/**
 	 * Add completed online async tasks that need processing onto the queue
 	 * @param CompletedItem - some finished request of the online services
diff --git a/Engine/Source/Runtime/Online/OnlineSubsystem/Public/OnlineJsonSerializer.h b/Engine/Source/Runtime/Online/OnlineSubsystem/Public/OnlineJsonSerializer.h
--- a/Engine/Source/Runtime/Online/OnlineSubsystem/Public/OnlineJsonSerializer.h
+++ b/Engine/Source/Runtime/Online/OnlineSubsystem/Public/OnlineJsonSerializer.h
@@ -601,11 +601,11 @@ struct FOnlineJsonSerializable
 	 *
 	 * @return the corresponding json string
 	 */
-	inline const FString ToJson(bool bPrettyPrint = true) const
+	inline const FString ToJson() const
 	{
 		// Strip away const, because we use a single method that can read/write which requires non-const semantics
 		// Otherwise, we'd have to have 2 separate macros for declaring const to json and non-const from json
-		return ((FOnlineJsonSerializable*)this)->ToJson(bPrettyPrint);
+		return ((FOnlineJsonSerializable*)this)->ToJson();
 	}
 	/**
 	 * Serializes this object to its JSON string form
diff --git a/Engine/Source/Runtime/RHI/Private/DynamicRHI.cpp b/Engine/Source/Runtime/RHI/Private/DynamicRHI.cpp
--- a/Engine/Source/Runtime/RHI/Private/DynamicRHI.cpp
+++ b/Engine/Source/Runtime/RHI/Private/DynamicRHI.cpp
@@ -14,18 +14,6 @@
 // Globals.
 FDynamicRHI* GDynamicRHI = NULL;
 
-static TAutoConsoleVariable<int32> CVarDetectAndWarnOfBadDrivers(
-	TEXT("r.DetectAndWarnOfBadDrivers"),
-	1,
-	TEXT("On engine startup we can check the current GPU driver and warn the user about issues and suggest a specific version\n")
-	TEXT("The test is fast so this should not cost any performance.\n")
-	TEXT(" 0: off\n")
-	TEXT(" 1: a message on startup might appear (default)\n")
-	TEXT(" 2: Test by pretending the driver has issues"),
-	ECVF_RenderThreadSafe
-	);
-
-
 void InitNullRHI()
 {
 	// Use the null RHI if it was specified on the command line, or if a commandlet is running.
@@ -45,120 +33,6 @@ void InitNullRHI()
 	GRHISupportsTextureStreaming = false;
 }
 
-// e.g. 36143 for the NVIDIA driver 361.43
-// no floating point to avoid precision issues
-// @return -1 if the version is unknown
-static int32 ExtractGPUDriverVersionInt()
-{
-	int32 Ret = -1;
-
-	// we use the internal version, not the user version to avoid problem where the name was altered 
-	const FString& InternalVersion = GRHIAdapterInternalDriverVersion;
-
-	if(IsRHIDeviceNVIDIA())
-	{
-		if(InternalVersion.Left(6) == TEXT("9.18.1"))		// e.g. "9.18.136.143"
-		{
-			// e.g. 36.143
-			FString RightPart = InternalVersion.RightChop(6);
-			// e.g. 36143
-			FString CompactNumber = RightPart.Replace(TEXT("."), TEXT(""));
-
-			if(FCString::IsNumeric(*CompactNumber))
-			{
-				Ret = FCString::Atoi(*CompactNumber);
-			}
-		}
-	}
-	else
-	{
-		// examples for AMD: "13.12" "15.101.1007" "13.351"
-		// todo: need to decide on what part we extract, maybe make the return argument multiple ints
-	}
-	return Ret;
-}
-
-void RHIDetectAndWarnOfBadDrivers()
-{
-	int32 CVarValue = CVarDetectAndWarnOfBadDrivers.GetValueOnGameThread();
-
-	if(!GIsRHIInitialized || !CVarValue || GRHIVendorId == 0)
-	{
-		return;
-	}
-
-	// for testing
-	if(CVarValue == 2)
-	{
-		GRHIVendorId = 0x10DE;
-		GRHIAdapterName = TEXT("Generic NVIDIA");
-		GRHIAdapterUserDriverVersion = TEXT("361.43");
-		GRHIAdapterInternalDriverVersion = TEXT("9.18.136.143");
-		GRHIAdapterDriverDate = TEXT("01-01-1900");
-	}
-
-	int32 VersionInt = ExtractGPUDriverVersionInt();
-
-	if(VersionInt < 0)
-	{
-		// no valid version
-		return;
-	}
-
-	// can be made data driven:
-	if(IsRHIDeviceNVIDIA())
-	{
-		// todo: make this data driven
-		FString SuggestedVersion = TEXT("359.06");
-
-		// todo: make this data driven
-		if(VersionInt == 36143	// 361.43 "UE-25096 Viewport flashes black and white when moving in the scene on latest Nvidia drivers", at that time that was the latest driver
-		|| VersionInt <= 34709)	// 347.09 "old NVIDIA driver version we've seen driver crashes with Paragon content", we might have to adjust that further
-		{
-			FFormatNamedArguments Args;
-			Args.Add(TEXT("AdapterName"), FText::FromString(GRHIAdapterName));
-			Args.Add(TEXT("UserDriverVersion"), FText::FromString(GRHIAdapterUserDriverVersion));
-			Args.Add(TEXT("InternalDriverVersion"), FText::FromString(GRHIAdapterInternalDriverVersion));
-			Args.Add(TEXT("DriverDate"), FText::FromString(GRHIAdapterDriverDate));
-			Args.Add(TEXT("SuggestedVersion"), FText::FromString(SuggestedVersion));
-
-			// explicitly compact
-			FText LocalizedMsg = FText::Format(NSLOCTEXT("MessageDialog", "VideoCardDriverIssueReport",
-				"Suggested version: {SuggestedVersion} (newer drivers might be good as well)\nCurrent version: {UserDriverVersion} ({InternalDriverVersion}) {DriverDate}\nName: {AdapterName}"),
-				Args);
-
-			// less compact and more readable, maybe for later
-			/*
-				"Current driver:\n"
-				"  Name:\t{AdapterName}\n"
-				"  Version:\t{UserDriverVersion} (internal {InternalDriverVersion})\n"
-				"  Date:\t{DriverDate}\n"
-				"\n"
-				"Suggested driver (newer drivers might be good as well):\n"
-				"  Version: {SuggestedVersion}\n"
-				"\n"
-				"Please update your drivers to avoid problems (performance, stalls, crashes, reboot)."
-			*/
-
-			FPlatformMisc::MessageBoxExt(EAppMsgType::Ok, 
-				*LocalizedMsg.ToString(),
-				*NSLOCTEXT("MessageDialog", "TitleVideoCardDriverIssue", "WARNING: Video card driver with known issues").ToString());
-		}
-	}
-	else if(IsRHIDeviceAMD())
-	{
-		// so far we don't have any data on issues
-	}
-	else if(IsRHIDeviceIntel())
-	{
-		// so far we don't have any data on issues
-	}
-	else
-	{
-		// so far we don't have any data on issues
-	}
-}
-
 void RHIInit(bool bHasEditorToken)
 {
 	if(!GDynamicRHI)
@@ -189,8 +63,6 @@ void RHIInit(bool bHasEditorToken)
 
 		check(GDynamicRHI);
 	}
-
-	RHIDetectAndWarnOfBadDrivers();
 }
 
 void RHIPostInit()
diff --git a/Engine/Source/Runtime/RHI/Private/GPUDefragAllocator.cpp b/Engine/Source/Runtime/RHI/Private/GPUDefragAllocator.cpp
--- a/Engine/Source/Runtime/RHI/Private/GPUDefragAllocator.cpp
+++ b/Engine/Source/Runtime/RHI/Private/GPUDefragAllocator.cpp
@@ -111,7 +111,7 @@ FBestFitAllocator implementation.
 * @param	bAllowFailure	Whether to allow allocation failure or not
 * @return	Pointer to allocated memory
 */
-void* FGPUDefragAllocator::Allocate(int64 AllocationSize, int32 Alignment, TStatId InStat, bool bAllowFailure)
+void* FGPUDefragAllocator::Allocate(int64 AllocationSize, int32 Alignment, bool bAllowFailure)
 {
 	SCOPE_SECONDS_COUNTER(TimeSpentInAllocator);
 	FScopeLock Lock(&SynchronizationObject);
@@ -188,10 +188,6 @@ void* FGPUDefragAllocator::Allocate(int64 AllocationSize, int32 Alignment, TStat
 #endif
 
 	FMemoryChunk* AllocatedChunk = AllocateChunk(BestChunk, AllocationSize, false);
-	//todo: fix stats
-	//ensureMsgf(AllocatedChunk->Stat.IsNone(), TEXT("FreeChunk already has a stat."));
-	AllocatedChunk->Stat = InStat;
-
 	check(IsAligned(AllocatedChunk->Base, Alignment));
 	return AllocatedChunk->Base;
 }
@@ -251,7 +247,6 @@ void FGPUDefragAllocator::FreeChunk(FMemoryChunk* Chunk, bool bMaintainSortedFre
 {
 	// Remove the entry
 	PointerToChunkMap.Remove(Chunk->Base);
-	Chunk->Stat = TStatId();
 
 	// Update usage stats in a thread safe way.
 	FPlatformAtomics::InterlockedAdd(&AllocatedMemorySize, -Chunk->Size);
@@ -654,7 +649,7 @@ void FGPUDefragAllocator::SortFreeList(int32& NumFreeChunks, int64& LargestFreeC
 */
 FGPUDefragAllocator::FMemoryChunk* FGPUDefragAllocator::FindAdjacent(FMemoryChunk* UsedChunk, bool bAnyChunkType)
 {
-	if (UsedChunk && !UsedChunk->IsRelocating() && (bAnyChunkType))
+	if (UsedChunk && !UsedChunk->IsRelocating() && (bAnyChunkType || UsedChunk->HasReallocationRequest()))
 	{
 		FMemoryChunk* FreeChunkLeft = UsedChunk->PreviousChunk;
 		FMemoryChunk* FreeChunkRight = UsedChunk->NextChunk;
@@ -679,6 +674,50 @@ FGPUDefragAllocator::FMemoryChunk* FGPUDefragAllocator::FindAdjacent(FMemoryChun
 	return nullptr;
 }
 
+/**
+* Searches for a reallocation request that would fit within the specified free chunk.
+* Prefers allocation requests over reallocation requests.
+*
+* @param FreeChunk		Free chunk we're trying to fill up
+* @return				First reallocating chunk that could fit, or nullptr
+*/
+FGPUDefragAllocator::FRequestNode* FGPUDefragAllocator::FindAnyReallocation(FMemoryChunk* FreeChunk)
+{
+	FRequestNode* BestRequest = nullptr;
+	int64 BestFit = MAX_int64;
+	bool bFoundAllocation = false;
+
+	for (FRequestList::TIterator It(ReallocationRequests.GetHead()); It; ++It)
+	{
+		FRequestNode* RequestNode = It.GetNode();
+		FAsyncReallocationRequest* Request = RequestNode->GetValue();
+		FMemoryChunk* CurrentChunk = Request->MemoryChunk;
+		int64 CurrentFit = FreeChunk->Size - Request->GetNewSize();
+
+		// Have we started getting into reallocations, but we've already found a suitable allocation?
+		if (bFoundAllocation && Request->IsReallocation())
+		{
+			// Use the allocation first.
+			break;
+		}
+
+		// Better fit than previously?
+		if (CurrentFit >= 0 && CurrentFit < BestFit && (Request->IsAllocation() || CanRelocate(CurrentChunk)))
+		{
+			BestFit = CurrentFit;
+			BestRequest = RequestNode;
+			bFoundAllocation = Request->IsAllocation();
+
+			// Perfect fit?
+			if (CurrentFit == 0)
+			{
+				break;
+			}
+		}
+	}
+	return BestRequest;
+}
+
 /**
 * Searches for an allocated chunk that would fit within the specified free chunk.
 * The allocated chunk must be adjacent to a free chunk and have a larger
@@ -817,13 +856,14 @@ void FGPUDefragAllocator::CheckForErrors(bool bCheckSortedFreeList)
 */
 FGPUDefragAllocator::FMemoryChunk* FGPUDefragAllocator::RelocateIntoFreeChunk(FRelocationStats& Stats, FMemoryChunk* FreeChunk, FMemoryChunk* SourceChunk)
 {
+	FRequestNode* ReallocationRequestNode = SourceChunk->ReallocationRequestNode;
+
 	// Save off important data from 'SourceChunk', since it will get modified by the call to LinkFreeChunk().
 	uint8* OldBaseAddress = SourceChunk->Base;
 	void* UserPayload = SourceChunk->UserPayload;
 	int32 OldSize = SourceChunk->Size;
 	const int32 NewSize = SourceChunk->GetFinalSize();
 	const int32 UsedSize = FMath::Min(NewSize, OldSize);
-	TStatId	SourceStat = SourceChunk->Stat;
 
 	// Are we relocating into adjacent free chunk?
 	bool bAdjacentRelocation = SourceChunk->PreviousChunk == FreeChunk || SourceChunk->NextChunk == FreeChunk;
@@ -861,8 +901,26 @@ FGPUDefragAllocator::FMemoryChunk* FGPUDefragAllocator::RelocateIntoFreeChunk(FR
 	// Update our book-keeping.
 	PointerToChunkMap.Remove(OldBaseAddress);
 	PointerToChunkMap.Add(NewBaseAddress, DestinationChunk);
-	SourceChunk->Stat = TStatId();
-	DestinationChunk->Stat = SourceStat;	
+
+	// Move the reallocation request into the InProgress list.
+	if (ReallocationRequestNode)
+	{
+		FAsyncReallocationRequest* ReallocationRequest = ReallocationRequestNode->GetValue();
+		check(SourceChunk == ReallocationRequest->MemoryChunk);
+
+		// Create a new node in the 'InProgress' list.
+		ReallocationRequestsInProgress.AddHead(ReallocationRequest);
+		FRequestNode* NewNode = ReallocationRequestsInProgress.GetHead();
+
+		// Swap chunk and node pointers.
+		SourceChunk->ReallocationRequestNode = nullptr;
+		DestinationChunk->ReallocationRequestNode = NewNode;
+		ReallocationRequest->MemoryChunk = DestinationChunk;
+		ReallocationRequest->NewAddress = NewBaseAddress;
+
+		// Remove from current request list.
+		ReallocationRequests.RemoveNode(ReallocationRequestNode);
+	}
 
 	// Is there free space left over at the end of DestinationChunk?
 	FMemoryChunk* NextFreeChunk;
@@ -905,6 +963,51 @@ FGPUDefragAllocator::FMemoryChunk* FGPUDefragAllocator::RelocateIntoFreeChunk(FR
 	return NextFreeChunk;
 }
 
+/**
+* Allocates memory from the specified free chunk, servicing an async allocation request.
+*
+* @param Stats			[out] Stats
+* @param FreeChunk		Chunk to allocate memory from
+* @param RequestNode	List node to the allocation request
+* @return				Next Free chunk to try to fill up
+*/
+FGPUDefragAllocator::FMemoryChunk* FGPUDefragAllocator::AllocateIntoFreeChunk(FRelocationStats& Stats, FMemoryChunk* FreeChunk, FRequestNode* RequestNode)
+{
+	FMemoryChunk* NextFreeChunk = FreeChunk->NextFreeChunk;
+	FAsyncReallocationRequest* Request = RequestNode->GetValue();
+	check(Request->IsAllocation());
+
+	// Note: AllocateChunk() may split 'FreeChunk' and return a new chunk.
+	FMemoryChunk* AllocatedChunk = AllocateChunk(FreeChunk, Request->GetNewSize(), true);
+
+	// Create a new node in the 'InProgress' list.
+	ReallocationRequestsInProgress.AddHead(Request);
+	FRequestNode* NewNode = ReallocationRequestsInProgress.GetHead();
+
+	// Setup chunk and node variables.
+	AllocatedChunk->ReallocationRequestNode = NewNode;
+	Request->MemoryChunk = AllocatedChunk;
+	Request->NewAddress = AllocatedChunk->Base;
+
+	// Mark the chunk as "in use" during the current sync step.
+	// Note: This sync index will propagate if these chunks are involved in any merge/split in the future.
+	AllocatedChunk->SetSyncIndex(GetCurrentSyncIndex(), AllocatedChunk->Size);
+
+	// Remove from current request list.
+	ReallocationRequests.RemoveNode(RequestNode);
+
+	FPlatformAtomics::InterlockedAdd(&PendingMemoryAdjustment, -Request->GetNewSize());
+
+	if (AllocatedChunk->NextChunk && AllocatedChunk->NextChunk->bIsAvailable)
+	{
+		return AllocatedChunk->NextChunk;
+	}
+	else
+	{
+		return NextFreeChunk;
+	}
+}
+
 /**
 * Blocks the calling thread until all relocations and reallocations that were initiated by Tick() have completed.
 *
@@ -918,7 +1021,46 @@ bool FGPUDefragAllocator::FinishAllRelocations()
 	if (bWasAnyRelocationsInProgress)
 	{		
 		BlockOnFence();
-	}	
+	}
+
+	// All reallocation requests have now completed.
+	int32 TotalMemoryAdjustment = 0;
+	for (FRequestList::TIterator It(ReallocationRequestsInProgress.GetHead()); It; ++It)
+	{
+		FAsyncReallocationRequest* FinishedRequest = *It;
+
+		FMemoryChunk* FinishedChunk = FinishedRequest->MemoryChunk;
+		if (!FinishedRequest->IsCanceled())
+		{
+			// Mark it complete.
+			FinishedRequest->MarkCompleted();
+
+			if (FinishedRequest->IsReallocation())
+			{
+				NumFinishedAsyncReallocations++;
+			}
+			else
+			{
+				NumFinishedAsyncAllocations++;
+			}
+
+			if (!bBenchmarkMode)
+			{
+				PlatformNotifyReallocationFinished(FinishedRequest, FinishedChunk->UserPayload);
+			}
+			TotalMemoryAdjustment += FinishedRequest->NewSize - FinishedRequest->OldSize;
+			FinishedRequest->MemoryChunk = nullptr;
+		}
+		else
+		{
+			// Delete the request, since it's our copy. The user has deleted the original already.
+			delete FinishedRequest;
+		}
+		FinishedChunk->ReallocationRequestNode = nullptr;
+	}
+	ReallocationRequestsInProgress.Empty();
+
+	check(ReallocationRequests.Num() > 0 || PendingMemoryAdjustment == 0);
 
 	// Take the opportunity to free all chunks that couldn't be freed immediately before.
 	for (TDoubleLinkedList<FMemoryChunk*>::TIterator It(PendingFreeChunks.GetHead()); It; ++It)
@@ -985,6 +1127,116 @@ void FGPUDefragAllocator::BlockOnSyncIndex(uint32 SyncIndex)
 	}
 }
 
+/**
+* Cancels the specified reallocation request.
+* Note that the allocator doesn't keep track of requests after it's been completed,
+* so the user must provide the current base address. This may not match any of the
+* addresses in the (old) request since the memory may have been relocated since then.
+*
+* @param Request				Request to cancel. Must be a valid request.
+* @param CurrentBaseAddress	Current baseaddress used by the allocation.
+*/
+void FGPUDefragAllocator::CancelAsyncReallocation(FAsyncReallocationRequest* Request, const void* CurrentBaseAddress)
+{
+	check(Request && !Request->IsCanceled());
+	NumCanceledAsyncRequests++;
+
+	int32 MemoryAdjustment = Request->NewSize - Request->OldSize;
+
+	// We cannot undo shrinking reallocations. The memory will be lost and gone forever.
+	check(MemoryAdjustment > 0);
+
+	// Mark it canceled.
+	Request->bIsCanceled = true;
+
+	// Make sure it's marked 'completed'.
+	bool bHasStarted = Request->HasStarted();
+	bool bHasCompleted = Request->HasCompleted();
+	if (!bHasCompleted)
+	{
+		// This will also prevent it from further relocation until it's gone completely from the system.
+		Request->MarkCompleted();
+	}
+
+	// Has it not started yet?
+	if (!bHasStarted)
+	{
+		if (Request->IsReallocation())
+		{
+			// If it hasn't started yet, just remove the request and flag it completed.
+			FMemoryChunk* MatchingChunk = Request->MemoryChunk;
+			check(MatchingChunk && CurrentBaseAddress == nullptr);
+			FRequestNode* RequestNode = MatchingChunk->ReallocationRequestNode;
+			check(RequestNode);
+
+			// Remove it from our end of the system.
+			ReallocationRequests.RemoveNode(RequestNode);
+			MatchingChunk->ReallocationRequestNode = nullptr;
+		}
+		else
+		{
+			FRequestNode* RequestNode = ReallocationRequests.FindNode(Request);
+			check(RequestNode);
+			ReallocationRequests.RemoveNode(RequestNode);
+		}
+		FPlatformAtomics::InterlockedAdd(&PendingMemoryAdjustment, -MemoryAdjustment);
+	}
+	else
+	{
+		// Is it still in progress?
+		if (!bHasCompleted)
+		{
+			FMemoryChunk* MatchingChunk = Request->MemoryChunk;
+			check(MatchingChunk && CurrentBaseAddress == nullptr);
+
+			// Make a copy of the request, since the current Request will be deleted by the user.
+			FAsyncReallocationRequest* RequestCopy = new FAsyncReallocationRequest(*Request);
+			FRequestNode* RequestNode = MatchingChunk->ReallocationRequestNode;
+			MatchingChunk->ReallocationRequestNode = nullptr;
+
+			FMemoryChunk* NewChunk;
+			if (Request->IsReallocation())
+			{
+				// Undo the 'Grow' by immediately shrinking it back to the old size (adjusting baseaddress).
+				// Allow the GPU to finish relocating into the used portion. (Any new relocations will be pipelined.)
+				NewChunk = Shrink(MatchingChunk, MemoryAdjustment);
+			}
+			else
+			{
+				// Undo the allocation by freeing it.
+				FreeChunk(MatchingChunk, false);
+				NewChunk = MatchingChunk;
+			}
+
+			// Fixup all request pointers.
+			RequestNode->GetValue() = RequestCopy;
+			NewChunk->ReallocationRequestNode = RequestNode;
+			RequestCopy->MemoryChunk = NewChunk;
+		}
+		// Is it already completed?
+		else
+		{
+			// When it's completed already, the allocator doesn't keep track of it anymore in any lists.
+			if (Request->IsReallocation())
+			{
+				FMemoryChunk* MatchingChunk = PointerToChunkMap.FindRef(CurrentBaseAddress);
+				check(MatchingChunk && MatchingChunk->ReallocationRequestNode == nullptr);
+
+				// Undo the 'Grow' by immediately shrinking it back to the old size (adjusting baseaddress).
+				FMemoryChunk* NewChunk = Shrink(MatchingChunk, MemoryAdjustment);
+			}
+			else
+			{
+				FMemoryChunk* MatchingChunk = PointerToChunkMap.FindRef(Request->GetNewBaseAddress());
+				check(MatchingChunk && MatchingChunk->ReallocationRequestNode == nullptr);
+
+				// Undo the allocation by freeing it.
+				FreeChunk(MatchingChunk, false);
+			}
+		}
+	}
+}
+
 static int32 GGPUDefragEnableTimeLimits = 1;
 static FAutoConsoleVariableRef CVarGPUDefragEnableTimeLimits(
 	TEXT("r.GPUDefrag.EnableTimeLimits"),
@@ -1053,15 +1305,21 @@ void FGPUDefragAllocator::PartialDefragmentation(FRelocationStats& Stats, double
 				BestChunk = RelocateAllowed(FreeChunk, FindAdjacent(FreeChunk->NextChunk, false));
 
 				if (!BestChunk)
-				{					
-					// 4. Merge with a used chunk adjacent to hole (to make that hole larger).
-					BestChunk = RelocateAllowed(FreeChunk, FindAdjacentToHole(FreeChunk));
+				{
+					// 3. Merge with any async request (best-fitting)
+					BestRequestNode = FindAnyReallocation(FreeChunk);
 
-					if (!BestChunk)
+					if (!BestRequestNode)
 					{
-						// 5. Merge with chunk from the end of the pool (well-fitting)
-						BestChunk = FindAny(FreeChunk);
-					}					
+						// 4. Merge with a used chunk adjacent to hole (to make that hole larger).
+						BestChunk = RelocateAllowed(FreeChunk, FindAdjacentToHole(FreeChunk));
+
+						if (!BestChunk)
+						{
+							// 5. Merge with chunk from the end of the pool (well-fitting)
+							BestChunk = FindAny(FreeChunk);
+						}
+					}
 				}
 			}
 		}
@@ -1069,11 +1327,23 @@ void FGPUDefragAllocator::PartialDefragmentation(FRelocationStats& Stats, double
 		if (BestChunk)
 		{
 			FreeChunk = RelocateIntoFreeChunk(Stats, FreeChunk, BestChunk);			
-		}		
+		}
+		else if (BestRequestNode)
+		{
+			FAsyncReallocationRequest* Request = BestRequestNode->GetValue();
+			if (Request->IsReallocation())
+			{				
+				FreeChunk = RelocateIntoFreeChunk(Stats, FreeChunk, Request->MemoryChunk);				
+			}
+			else
+			{
+				FreeChunk = AllocateIntoFreeChunk(Stats, FreeChunk, BestRequestNode);
+			}
+		}
 		else
 		{
 			// Did the free chunk fail to defrag?
-			if (FreeChunk->DefragCounter == 0 && (FreeChunk->NextFreeChunk))
+			if (FreeChunk->DefragCounter == 0 && (FreeChunk->NextFreeChunk || ReallocationRequests.Num() > 0))
 			{
 				// Don't try it again for a while.
 				if (FreeChunk->Size < DEFRAG_SMALL_CHUNK_SIZE)
@@ -1420,7 +1690,18 @@ FGPUDefragAllocator::EMemoryElementType FGPUDefragAllocator::GetChunkType(FMemor
 	else if (Chunk->bIsAvailable)
 	{
 		ChunkType = MET_Free;			// Free (dark grey)
-	}	
+	}
+	else if (Chunk->HasReallocationRequest())
+	{
+		if (Chunk->ReallocationRequestNode->GetValue()->HasCompleted())
+		{
+			ChunkType = MET_Resized;	// Has been resized but not finalized yet (dark green)
+		}
+		else
+		{
+			ChunkType = MET_Resizing;	// Allocated but wants a resize (green)
+		}
+	}
 	else if (CanRelocate(Chunk) == false)
 	{
 		ChunkType = MET_Locked;			// Allocated but can't me relocated at this time (locked) (red)
diff --git a/Engine/Source/Runtime/RHI/Private/RHICommandList.cpp b/Engine/Source/Runtime/RHI/Private/RHICommandList.cpp
--- a/Engine/Source/Runtime/RHI/Private/RHICommandList.cpp
+++ b/Engine/Source/Runtime/RHI/Private/RHICommandList.cpp
@@ -480,7 +480,7 @@ void FRHICommandListExecutor::ExecuteList(FRHICommandListImmediate& CmdList)
 
 void FRHICommandListExecutor::LatchBypass()
 {
-#if CAN_TOGGLE_COMMAND_LIST_BYPASS
+#if !UE_BUILD_SHIPPING
 	if (GRHIThread)
 	{
 		if (bLatchedBypass)
@@ -522,7 +522,7 @@ void FRHICommandListExecutor::LatchBypass()
 	if (!bLatchedBypass)
 	{
 		bLatchedUseParallelAlgorithms = FApp::ShouldUseThreadingForPerformance() 
-#if CAN_TOGGLE_COMMAND_LIST_BYPASS
+#if !UE_BUILD_SHIPPING
 			&& (CVarRHICmdUseParallelAlgorithms.GetValueOnAnyThread() >= 1)
 #endif
 			;
@@ -1894,7 +1894,7 @@ void FDynamicRHI::UpdateTexture2D_RenderThread(class FRHICommandListImmediate& R
 
 void* FDynamicRHI::LockTexture2D_RenderThread(class FRHICommandListImmediate& RHICmdList, FTexture2DRHIParamRef Texture, uint32 MipIndex, EResourceLockMode LockMode, uint32& DestStride, bool bLockWithinMiptail, bool bNeedsDefaultRHIFlush)
 {
-	if (bNeedsDefaultRHIFlush) 
+	if (bNeedsDefaultRHIFlush)
 	{
 		QUICK_SCOPE_CYCLE_COUNTER(STAT_RHIMETHOD_LockTexture2D_Flush);
 		RHICmdList.ImmediateFlush(EImmediateFlushType::FlushRHIThread);
diff --git a/Engine/Source/Runtime/RHI/Public/GPUDefragAllocator.h b/Engine/Source/Runtime/RHI/Public/GPUDefragAllocator.h
--- a/Engine/Source/Runtime/RHI/Public/GPUDefragAllocator.h
+++ b/Engine/Source/Runtime/RHI/Public/GPUDefragAllocator.h
@@ -134,7 +134,7 @@ public:
 		* @param	ChunkToInsertAfter		Chunk to insert this after.
 		* @param	FirstFreeChunk			Reference to first free chunk Pointer.
 		*/
-		FMemoryChunk(uint8* InBase, int32 InSize, FGPUDefragAllocator& InBestFitAllocator, FMemoryChunk*& ChunkToInsertAfter, TStatId InStat, bool bSortedFreeList)
+		FMemoryChunk(uint8* InBase, int32 InSize, FGPUDefragAllocator& InBestFitAllocator, FMemoryChunk*& ChunkToInsertAfter, bool bSortedFreeList)
 			: Base(InBase)
 			, Size(InSize)
 			, bIsAvailable(false)
@@ -143,8 +143,8 @@ public:
 			, BestFitAllocator(InBestFitAllocator)
 			, SyncIndex(0)
 			, SyncSize(0)
-			, UserPayload(0)			
-			, Stat(InStat)
+			, UserPayload(0)
+			, ReallocationRequestNode(nullptr)
 		{
 			Link(ChunkToInsertAfter);
 			// This is going to change bIsAvailable.
@@ -253,7 +253,15 @@ public:
 
 			PreviousFreeChunk = nullptr;
 			NextFreeChunk = nullptr;
-		}		
+		}
+
+		/**
+		*	Returns true if the Chunk has an associated reallocation request.
+		*/
+		bool	HasReallocationRequest() const
+		{
+			return ReallocationRequestNode != nullptr;
+		}
 
 		/**
 		*	Returns true if the Chunk is being asynchronously relocated due to reallocation or defrag.
@@ -344,10 +352,13 @@ public:
 		/** Number of uint8s covered by the SyncIndex (starting from the beginning of the chunk). */
 		int64					SyncSize;
 		/** User payload, e.g. platform-specific texture Pointer. Only chunks with payload can be relocated. */
-		void*					UserPayload;		
+		void*					UserPayload;
 
-		//stat associated with this allocation
-		TStatId Stat;
+		/**
+		* Reallocation request for this chunk, or nullptr.
+		* Point32s straight to a node in FBestFitAllocator::ReallocationRequests or FBestFitAllocator::ReallocationRequestsInProgress.
+		*/
+		FRequestNode*			ReallocationRequestNode;
 	};
 
 	/** Constructor, zero initializing all member variables */
@@ -398,7 +409,7 @@ public:
 		// Update stats in a thread safe way.
 		FPlatformAtomics::InterlockedExchange(&AvailableMemorySize, MemorySize);
 		// Allocate initial chunk.
-		FirstChunk = new FMemoryChunk(MemoryBase, MemorySize, *this, FirstChunk, TStatId(), false);
+		FirstChunk = new FMemoryChunk(MemoryBase, MemorySize, *this, FirstChunk, /*FirstChunk, FirstFreeChunk,*/ false);
 		LastChunk = FirstChunk;
 	}
 
@@ -442,7 +453,7 @@ public:
 	* @param	bAllowFailure	Whether to allow allocation failure or not
 	* @return	Pointer to allocated memory
 	*/
-	virtual void*	Allocate(int64 AllocationSize, int32 Alignment, TStatId InStat, bool bAllowFailure);	
+	virtual void*	Allocate(int64 AllocationSize, int32 Alignment, bool bAllowFailure);
 
 	/**
 	* Frees allocation associated with the specified Pointer.
@@ -826,10 +837,7 @@ protected:
 		// Calculate size of second chunk...
 		int32 SecondSize = BaseChunk->Size - FirstSize;
 		// ... and create it.
-
-		//todo: fix stats
-		//ensureMsgf(BaseChunk->Stat.IsNone(), TEXT("Free chunk has stat"));
-		FMemoryChunk* NewFreeChunk = new FMemoryChunk(BaseChunk->Base + FirstSize, SecondSize, *this, BaseChunk, TStatId(), bSortedFreeList);
+		FMemoryChunk* NewFreeChunk = new FMemoryChunk(BaseChunk->Base + FirstSize, SecondSize, *this, BaseChunk, /*FirstChunk, FirstFreeChunk,*/ bSortedFreeList);
 
 		// Keep the original sync index for the new chunk if necessary.
 		if (BaseChunk->IsRelocating() && BaseChunk->SyncSize > FirstSize)
@@ -902,7 +910,15 @@ protected:
 	* @param bAnyChunkType	If false, only succeeds if 'UsedChunk' has a reallocation request and fits
 	* @return				Returns 'UsedChunk' if it fits the criteria, otherwise nullptr
 	*/
-	FMemoryChunk*	FindAdjacent(FMemoryChunk* UsedChunk, bool bAnyChunkType);	
+	FMemoryChunk*	FindAdjacent(FMemoryChunk* UsedChunk, bool bAnyChunkType);
+
+	/**
+	* Searches for a reallocation request that would fit within the specified free chunk.
+	*
+	* @param FreeChunk		Free chunk we're trying to fill up
+	* @return				First request that could fit, or nullptr
+	*/
+	FRequestNode*	FindAnyReallocation(FMemoryChunk* FreeChunk);
 
 	/**
 	* Searches for an allocated chunk that would fit within the specified free chunk.
@@ -936,6 +952,16 @@ protected:
 	*/
 	FMemoryChunk*	RelocateIntoFreeChunk(FRelocationStats& Stats, FMemoryChunk* FreeChunk, FMemoryChunk* UsedChunk);
 
+	/**
+	* Allocates memory from the specified free chunk, servicing an async allocation request.
+	*
+	* @param Stats			[out] Stats
+	* @param FreeChunk		Chunk to allocate memory from
+	* @param RequestNode	List node to the allocation request
+	* @return				Next Free chunk to try to fill up
+	*/
+	FMemoryChunk*	AllocateIntoFreeChunk(FRelocationStats& Stats, FMemoryChunk* FreeChunk, FRequestNode* RequestNode);
+
 	FMemoryChunk* RelocateAllowed(FMemoryChunk* FreeChunk, FMemoryChunk* UsedChunk);	
 
 	FCriticalSection	SynchronizationObject;
@@ -983,6 +1009,12 @@ protected:
 	/** Platform-specific (GPU) fence, used for synchronizing the Sync Index. */
 	uint64			PlatformSyncFence;
 
+	/** Asynchronous reallocations requested by user.				*/
+	FRequestList	ReallocationRequests;
+
+	/** Asynchronous reallocations currently being processed.		*/
+	FRequestList	ReallocationRequestsInProgress;
+
 	/** Chunks that couldn't be freed immediately because they were being relocated. */
 	TDoubleLinkedList<FMemoryChunk*>	PendingFreeChunks;
 
@@ -1163,7 +1195,7 @@ private:
 */
 FORCEINLINE int32 FGPUDefragAllocator::FMemoryChunk::GetFinalSize() const
 {
-	return Size;
+	return ReallocationRequestNode ? ReallocationRequestNode->GetValue()->GetNewSize() : Size;
 }
 
 /**
@@ -1175,6 +1207,14 @@ FORCEINLINE int32 FGPUDefragAllocator::FMemoryChunk::GetFinalSize() const
 */
 FORCEINLINE bool FGPUDefragAllocator::CanRelocate(const FMemoryChunk* Chunk) const
 {
+	// During reallocation, the new texture keeps the request around until it's been processed higher up.
+	// Can't relocate a request that has started until it's fully gone from the system.
+	const FRequestNode* RequestNode = Chunk->ReallocationRequestNode;
+	if (RequestNode && RequestNode->GetValue()->HasStarted())
+	{
+		return false;
+	}
+
 	if (Chunk->IsLocked())
 	{
 		return false;
diff --git a/Engine/Source/Runtime/Renderer/Private/PostProcess/PostProcessTemporalAA.cpp b/Engine/Source/Runtime/Renderer/Private/PostProcess/PostProcessTemporalAA.cpp
--- a/Engine/Source/Runtime/Renderer/Private/PostProcess/PostProcessTemporalAA.cpp
+++ b/Engine/Source/Runtime/Renderer/Private/PostProcess/PostProcessTemporalAA.cpp
@@ -36,7 +36,7 @@ static void TemporalRandom(FVector2D* RESTRICT const Constant, uint32 FrameNumbe
 
 static TAutoConsoleVariable<float> CVarTemporalAASharpness(
 	TEXT("r.TemporalAASharpness"),
-	1.0f,
+	0.0f,
 	TEXT("Sharpness of temporal AA (0.0 = smoother, 1.0 = sharper)."),
 	ECVF_Scalability | ECVF_RenderThreadSafe);
 
diff --git a/Engine/Source/Runtime/Renderer/Private/PostProcess/PostProcessTonemap.cpp b/Engine/Source/Runtime/Renderer/Private/PostProcess/PostProcessTonemap.cpp
--- a/Engine/Source/Runtime/Renderer/Private/PostProcess/PostProcessTonemap.cpp
+++ b/Engine/Source/Runtime/Renderer/Private/PostProcess/PostProcessTonemap.cpp
@@ -8,7 +8,6 @@
 #include "ScenePrivate.h"
 #include "SceneFilterRendering.h"
 #include "PostProcessEyeAdaptation.h"
-#include "PostProcessUpscale.h"
 #include "PostProcessTonemap.h"
 #include "PostProcessing.h"
 #include "PostProcessCombineLUTs.h"
@@ -1078,14 +1077,29 @@ IMPLEMENT_SHADER_TYPE(template<>, TPostProcessTonemapVS<true>, TEXT("PostProcess
 IMPLEMENT_SHADER_TYPE(template<>, TPostProcessTonemapVS<false>, TEXT("PostProcessTonemap"), TEXT("MainVS"), SF_Vertex);
 
 
-FRCPassPostProcessTonemap::FRCPassPostProcessTonemap(const FViewInfo& InView, bool bInDoGammaOnly, bool bInDoEyeAdaptation)
+FRCPassPostProcessTonemap::FRCPassPostProcessTonemap(const FViewInfo& View, bool bInDoGammaOnly, bool bDoScreenPercentageInTonemapper, bool bInDoEyeAdaptation )
 	: bDoGammaOnly(bInDoGammaOnly)
-	, bDoScreenPercentageInTonemapper(false)
 	, bDoEyeAdaptation(bInDoEyeAdaptation)
-	, View(InView)	
+	, RenderTargetExtend(0, 0)
 {
-	uint32 ConfigBitmask = TonemapperGenerateBitmaskPC(&InView, bDoGammaOnly);
+	uint32 ConfigBitmask = TonemapperGenerateBitmaskPC(&View, bDoGammaOnly);
 	ConfigIndexPC = TonemapperFindLeastExpensive(TonemapperConfBitmaskPC, sizeof(TonemapperConfBitmaskPC)/4, TonemapperCostTab, ConfigBitmask);;
+
+	if(bDoScreenPercentageInTonemapper)
+	{
+		// Texture could be bigger than viewport
+		if (View.Family->RenderTarget->GetRenderTargetTexture())
+		{
+			RenderTargetExtend.X = View.Family->RenderTarget->GetRenderTargetTexture()->GetSizeX();
+			RenderTargetExtend.Y = View.Family->RenderTarget->GetRenderTargetTexture()->GetSizeY();
+		}
+		else
+		{
+			RenderTargetExtend = View.Family->RenderTarget->GetSizeXY();
+		}
+	}
+	
+	check(bDoScreenPercentageInTonemapper == ShouldDoScreenPercentageInTonemapper());
 }
 
 namespace PostProcessTonemapUtil
@@ -1131,14 +1145,14 @@ void FRCPassPostProcessTonemap::Process(FRenderingCompositePassContext& Context)
 		// input is not hooked up correctly
 		return;
 	}
-
+	
+	const FSceneView& View = Context.View;
 	const FSceneViewFamily& ViewFamily = *(View.Family);
 	FIntRect SrcRect = View.ViewRect;
-	FIntRect DestRect = bDoScreenPercentageInTonemapper ? View.UnscaledViewRect : View.ViewRect;
+	FIntRect DestRect = ShouldDoScreenPercentageInTonemapper() ? View.UnscaledViewRect : View.ViewRect;
 	FIntPoint SrcSize = InputDesc->Extent;
-	
-	SCOPED_DRAW_EVENTF(Context.RHICmdList, PostProcessTonemap, TEXT("Tonemapper#%d GammaOnly=%d ScreenPercentage=%d  %dx%d"),
-		ConfigIndexPC, bDoGammaOnly, bDoScreenPercentageInTonemapper, DestRect.Width(), DestRect.Height());
+
+	SCOPED_DRAW_EVENTF(Context.RHICmdList, PostProcessTonemap, TEXT("Tonemapper#%d%s %dx%d"), ConfigIndexPC, bDoGammaOnly ? TEXT(" GammaOnly") : TEXT(""), DestRect.Width(), DestRect.Height());
 
 	const FSceneRenderTargetItem& DestRenderTarget = PassOutputs[0].RequestSurface(Context);
 	
@@ -1156,7 +1170,7 @@ void FRCPassPostProcessTonemap::Process(FRenderingCompositePassContext& Context)
 		Context.RHICmdList.Clear(true, FLinearColor::Black, false, 1.0f, false, 0, DestRect);
 	}
 
-	Context.SetViewportAndCallRHI(DestRect, 0.0f, 1.0f );
+	Context.SetViewportAndCallRHI(DestRect.Min.X, DestRect.Min.Y, 0.0f, DestRect.Max.X + 1, DestRect.Max.Y + 1, 1.0f );
 
 	// set the state
 	Context.RHICmdList.SetBlendState(TStaticBlendState<>::GetRHI());
@@ -1217,6 +1231,11 @@ void FRCPassPostProcessTonemap::Process(FRenderingCompositePassContext& Context)
 		// This becomes even more important with some limited VRam (XBoxOne).
 		SceneContext.SetSceneColor(0);
 	}
+
+	if(ShouldDoScreenPercentageInTonemapper())
+	{
+		Context.View.ViewRect = View.UnscaledViewRect;
+	}
 }
 
 FPooledRenderTargetDesc FRCPassPostProcessTonemap::ComputeOutputDesc(EPassOutputId InPassOutputId) const
@@ -1228,6 +1247,11 @@ FPooledRenderTargetDesc FRCPassPostProcessTonemap::ComputeOutputDesc(EPassOutput
 	Ret.Format = PF_B8G8R8A8;
 	Ret.DebugName = TEXT("Tonemap");
 
+	if(ShouldDoScreenPercentageInTonemapper())
+	{
+		Ret.Extent = RenderTargetExtend;
+	}
+
 	return Ret;
 }
 
diff --git a/Engine/Source/Runtime/Renderer/Private/PostProcess/PostProcessTonemap.h b/Engine/Source/Runtime/Renderer/Private/PostProcess/PostProcessTonemap.h
--- a/Engine/Source/Runtime/Renderer/Private/PostProcess/PostProcessTonemap.h
+++ b/Engine/Source/Runtime/Renderer/Private/PostProcess/PostProcessTonemap.h
@@ -41,7 +41,7 @@ class FRCPassPostProcessTonemap : public TRenderingCompositePassBase<4, 1>
 {
 public:
 	// constructor
-	FRCPassPostProcessTonemap(const FViewInfo& InView, bool bInDoGammaOnly, bool bDoEyeAdaptation);
+	FRCPassPostProcessTonemap(const FViewInfo& View, bool bInDoGammaOnly, bool bDoScreenPercentageInTonemapper, bool bDoEyeAdaptation);
 
 	// interface FRenderingCompositePass ---------
 
@@ -49,13 +49,17 @@ public:
 	virtual void Release() override { delete this; }
 	virtual FPooledRenderTargetDesc ComputeOutputDesc(EPassOutputId InPassOutputId) const override;
 
-	bool bDoGammaOnly;
-	bool bDoScreenPercentageInTonemapper;
 private:
+	// set in constructor
+	bool bDoGammaOnly;
 	bool bDoEyeAdaptation;
 	uint32 ConfigIndexPC;
 
-	const FViewInfo& View;
+	// used if ShouldDoScreenPercentageInTonemapper() is true, otherwise (0,0)
+	FIntPoint RenderTargetExtend;
+
+	// @return true: we did the upscale in this pass and don't need to run another one after it
+	bool ShouldDoScreenPercentageInTonemapper() const { return RenderTargetExtend != FIntPoint::ZeroValue; }
 
 	void SetShader(const FRenderingCompositePassContext& Context);
 };
diff --git a/Engine/Source/Runtime/Renderer/Private/PostProcess/PostProcessUpscale.h b/Engine/Source/Runtime/Renderer/Private/PostProcess/PostProcessUpscale.h
--- a/Engine/Source/Runtime/Renderer/Private/PostProcess/PostProcessUpscale.h
+++ b/Engine/Source/Runtime/Renderer/Private/PostProcess/PostProcessUpscale.h
@@ -27,20 +27,12 @@ public:
 		// Panini screen fit factor (lerp between vertical and horizontal) 
 		float ScreenFit;
 
-		// constructor off
 		PaniniParams()
 			: D(0.0f)
 			, S(0.0f)
 			, ScreenFit(1.0f)
 		{}
 
-		PaniniParams(const FViewInfo& View);
-
-		bool IsEnabled() const
-		{
-			return D > 0.01f;
-		}
-
 		static const PaniniParams Default;
 	};
 
@@ -54,11 +46,9 @@ public:
 	FRCPassPostProcessUpscale(uint32 InUpscaleQuality, const PaniniParams& InPaniniConfig = PaniniParams::Default);
 
 	// interface FRenderingCompositePass ---------
-
 	virtual void Process(FRenderingCompositePassContext& Context) override;
 	virtual void Release() override { delete this; }
 	FPooledRenderTargetDesc ComputeOutputDesc(EPassOutputId InPassOutputId) const override;
-
 private:
 	// @param InCylinderDistortion 0=none..1=full in percent, must be in that range
 	template <uint32 Quality, uint32 bTesselatedQuad> static FShader* SetShader(const FRenderingCompositePassContext& Context, const PaniniParams& PaniniConfig);
diff --git a/Engine/Source/Runtime/Renderer/Private/PostProcess/PostProcessing.cpp b/Engine/Source/Runtime/Renderer/Private/PostProcess/PostProcessing.cpp
--- a/Engine/Source/Runtime/Renderer/Private/PostProcess/PostProcessing.cpp
+++ b/Engine/Source/Runtime/Renderer/Private/PostProcess/PostProcessing.cpp
@@ -91,6 +91,32 @@ static TAutoConsoleVariable<int32> CVarRenderTargetSwitchWorkaround(
 	TEXT("We want this enabled (1) on all 32 bit iOS devices (implemented through DeviceProfiles)."),
 	ECVF_RenderThreadSafe);
 
+static TAutoConsoleVariable<float> CVarUpscalePaniniD(
+	TEXT("r.Upscale.Panini.D"),
+	0,
+	TEXT("Allow and configure to apply a panini distortion to the rendered image. Values between 0 and 1 allow to fade the effect (lerp).\n")
+	TEXT("Implementation from research paper \"Pannini: A New Projection for Rendering Wide Angle Perspective Images\"\n")
+	TEXT(" 0: off (default)\n")
+	TEXT(">0: enabled (requires an extra post processing pass if upsampling wasn't used - see r.ScreenPercentage)\n")
+	TEXT(" 1: Panini cylindrical stereographic projection"),
+	ECVF_RenderThreadSafe);
+
+static TAutoConsoleVariable<float> CVarUpscalePaniniS(
+	TEXT("r.Upscale.Panini.S"),
+	0,
+	TEXT("Panini projection's hard vertical compression factor.\n")
+	TEXT(" 0: no vertical compression factor (default)\n")
+	TEXT(" 1: Hard vertical compression"),
+	ECVF_RenderThreadSafe);
+
+static TAutoConsoleVariable<float> CVarUpscalePaniniScreenFit(
+	TEXT("r.Upscale.Panini.ScreenFit"),
+	1.0f,
+	TEXT("Panini projection screen fit effect factor (lerp).\n")
+	TEXT(" 0: fit vertically\n")
+	TEXT(" 1: fit horizontally (default)"),
+	ECVF_RenderThreadSafe);
+
 static TAutoConsoleVariable<int32> CVarUpscaleQuality(
 	TEXT("r.Upscale.Quality"),
 	3,
@@ -135,12 +161,13 @@ static TAutoConsoleVariable<int32> CVarTonemapperScreenPercentage(
 	TEXT("Experimental ScreenPercentage upscale integrated into tonemapper pass (if certain conditions apply e.g. no FXAA)\n")
 	TEXT(" if enabled both features are done in one pass (faster, affects post process passes after the tonemapper including material post process e.g. sharpen)\n")
 	TEXT(" 0: off, the features run in separate passes (default)\n")
-	TEXT(" 1: enabled if nothing prevents it which is safe but it might do the two pass method"),
+	TEXT(" 1: enabled if nothing prevents it which is safe but it might do the two pass method\n")
+	TEXT(" 2: enabled even if some other feature would prevent it (good for testing, likely rendering artifacts)"),
 	ECVF_Scalability | ECVF_RenderThreadSafe);
 
 static TAutoConsoleVariable<int32> CVarMotionBlurNew(
 	TEXT("r.MotionBlurNew"),
-	1,
+	0,
 	TEXT(""),
 	ECVF_RenderThreadSafe
 );
@@ -310,7 +337,8 @@ static FRenderingCompositeOutputRef RenderBloom(
 	return RenderGaussianBlur(Context, TEXT("BloomBlurX"), TEXT("BloomBlurY"), PreviousBloom, Size, Tint, Additive,CrossBloom);
 }
 
-static FRCPassPostProcessTonemap* AddTonemapper(
+// @return bDoScreenPercentageInTonemapper
+static bool AddTonemapper(
 	FPostprocessContext& Context,
 	const FRenderingCompositeOutputRef& BloomOutputCombined,
 	const FRenderingCompositeOutputRef& EyeAdaptation,
@@ -318,9 +346,45 @@ static FRCPassPostProcessTonemap* AddTonemapper(
 {
 	const FEngineShowFlags& EngineShowFlags = Context.View.Family->EngineShowFlags;
 
+	int32 TonemapperScreenPercentage = CVarTonemapperScreenPercentage.GetValueOnRenderThread();
+
+	bool bDoScreenPercentageInTonemapper = TonemapperScreenPercentage != 0;
+
+	// only check if the feature is requested
+	if(bDoScreenPercentageInTonemapper)
+#if !(UE_BUILD_SHIPPING || UE_BUILD_TEST)
+	// useful only during development
+	if(TonemapperScreenPercentage != 2)
+#endif
+	{
+		// some features run after the tone mapper and cannot cope with this optimization
+
+		const bool bOverviewModeEnabled = EngineShowFlags.VisualizeBuffer && (Context.View.CurrentBufferVisualizationMode == NAME_None);
+
+		if (EngineShowFlags.VisualizeDOF ||
+			EngineShowFlags.VisualizeSSS ||
+			EngineShowFlags.HighResScreenshotMask ||
+			Context.View.FinalPostProcessSettings.AntiAliasingMethod == AAM_FXAA ||
+			bOverviewModeEnabled)
+		{
+			bDoScreenPercentageInTonemapper = false;
+		}
+
+		// todo: check for Panini projection
+
+		if(HasPostProcessMaterial(Context, BL_AfterTonemapping))
+		{
+			// Post process materials after the tone mapper should be avoided -
+			// they would run in 2D resolution (e.g. 1080p) which is often larger than 3D resolution
+			// Some effects don't work on the data generated by Tonemapper.ScreenPercentage e.g. sharpen
+			// and fixing the texture size differences depending in which pass we run are extensive.
+			bDoScreenPercentageInTonemapper = false;
+		}
+	}
+
 	FRenderingCompositePass* CombinedLUT = Context.Graph.RegisterPass(new(FMemStack::Get()) FRCPassPostProcessCombineLUTs(Context.View.GetShaderPlatform()));
 	const bool bDoEyeAdaptation = IsAutoExposureMethodSupported(Context.View.GetFeatureLevel(), EyeAdapationMethodId);
-	FRCPassPostProcessTonemap* PostProcessTonemap =	Context.Graph.RegisterPass(new(FMemStack::Get()) FRCPassPostProcessTonemap(Context.View, false, bDoEyeAdaptation));
+	FRCPassPostProcessTonemap* PostProcessTonemap =	Context.Graph.RegisterPass(new(FMemStack::Get()) FRCPassPostProcessTonemap(Context.View, false, bDoScreenPercentageInTonemapper, bDoEyeAdaptation));
 
 	PostProcessTonemap->SetInput(ePId_Input0, Context.FinalOutput);
 	PostProcessTonemap->SetInput(ePId_Input1, BloomOutputCombined);
@@ -329,7 +393,7 @@ static FRCPassPostProcessTonemap* AddTonemapper(
 
 	Context.FinalOutput = FRenderingCompositeOutputRef(PostProcessTonemap);
 
-	return PostProcessTonemap;
+	return bDoScreenPercentageInTonemapper;
 }
 
 #if WITH_EDITOR
@@ -348,7 +412,7 @@ static void AddSelectionOutline(FPostprocessContext& Context)
 
 static void AddGammaOnlyTonemapper(FPostprocessContext& Context)
 {
-	FRenderingCompositePass* PostProcessTonemap = Context.Graph.RegisterPass(new(FMemStack::Get()) FRCPassPostProcessTonemap(Context.View, true, false/*eye*/));
+	FRenderingCompositePass* PostProcessTonemap = Context.Graph.RegisterPass(new(FMemStack::Get()) FRCPassPostProcessTonemap(Context.View, true, false, false/*eye*/));
 
 	PostProcessTonemap->SetInput(ePId_Input0, Context.FinalOutput);
 
@@ -1160,26 +1224,13 @@ void FPostProcessing::Process(FRHICommandListImmediate& RHICmdList, FViewInfo& V
 		FRenderingCompositeOutputRef BloomOutputCombined;
 		// not always valid
 		FRenderingCompositePass* VelocityFlattenPass = 0;
+
 		// in the following code some feature might set this to false
 		bool bAllowTonemapper = FeatureLevel >= ERHIFeatureLevel::SM4;
-		//
-		bool bStereoRenderingAndHMD = View.Family->EngineShowFlags.StereoRendering && View.Family->EngineShowFlags.HMDDistortion;
-		//
-		FRCPassPostProcessUpscale::PaniniParams PaniniConfig(View);
-		//
+
 		EStereoscopicPass StereoPass = View.StereoPass;
-		//
+
 		FSceneViewState* ViewState = (FSceneViewState*)Context.View.State;
-		//
-		bool bDoScreenPercentage;
-		{
-			//
-			bool bHMDWantsUpscale = bStereoRenderingAndHMD && GEngine->HMDDevice->NeedsUpscalePostProcessPass();
-			// Do not use upscale if SeparateRenderTarget is in use! (stereo rendering wants to control this)
-			bool bAllowScreenPercentage = bHMDWantsUpscale || !View.Family->EngineShowFlags.StereoRendering || (!View.Family->EngineShowFlags.HMDDistortion && !View.Family->bUseSeparateRenderTarget);
-			// is Upscale from a lower resolution needed and allowed
-			bDoScreenPercentage = bAllowScreenPercentage && (View.UnscaledViewRect != View.ViewRect);
-		}
 
 		{
 			if (FSceneRenderTargets::Get(RHICmdList).SeparateTranslucencyRT)
@@ -1205,9 +1256,10 @@ void FPostProcessing::Process(FRHICommandListImmediate& RHICmdList, FViewInfo& V
 		{
 			bAllowTonemapper = false;
 		}
-	
-		FRCPassPostProcessTonemap* Tonemapper = 0;
 
+		// if ScreenPercentage (render 3d in different resolution and upscale later) is already done
+		bool bScreenPercentageIsDone = false;
+		
 		// add the passes we want to add to the graph (commenting a line means the pass is not inserted into the graph) ---------
 
 		if (AllowFullPostProcessing(View, FeatureLevel))
@@ -1676,7 +1728,7 @@ void FPostProcessing::Process(FRHICommandListImmediate& RHICmdList, FViewInfo& V
 				}
 				else
 				{
-					Tonemapper = AddTonemapper(Context, BloomOutputCombined, AutoExposure.EyeAdaptation, AutoExposure.MethodId);
+					bScreenPercentageIsDone = AddTonemapper(Context, BloomOutputCombined, AutoExposure.EyeAdaptation, AutoExposure.MethodId);
 				}
 			}
 
@@ -1793,6 +1845,8 @@ void FPostProcessing::Process(FRHICommandListImmediate& RHICmdList, FViewInfo& V
 
 		AddGBufferVisualizationOverview(Context, SeparateTranslucency, HDRColor);
 
+		bool bStereoRenderingAndHMD = View.Family->EngineShowFlags.StereoRendering && View.Family->EngineShowFlags.HMDDistortion;
+		bool bHMDWantsUpscale = false;
 		if (bStereoRenderingAndHMD)
 		{
 			FRenderingCompositePass* Node = NULL;
@@ -1811,6 +1865,8 @@ void FPostProcessing::Process(FRHICommandListImmediate& RHICmdList, FViewInfo& V
 #endif
 			}
 
+			bHMDWantsUpscale = GEngine->HMDDevice->NeedsUpscalePostProcessPass();
+			
 			if(Node)
 			{
 				Node->SetInput(ePId_Input0, FRenderingCompositeOutputRef(Context.FinalOutput));
@@ -1839,25 +1895,21 @@ void FPostProcessing::Process(FRHICommandListImmediate& RHICmdList, FViewInfo& V
 
 		AddHighResScreenshotMask(Context, SeparateTranslucency);
 
-		if(bDoScreenPercentage)
+		if(!bScreenPercentageIsDone)
 		{
-			// Check if we can save the Upscale pass and do it in the Tonemapper to save performance
-			if(Tonemapper && !PaniniConfig.IsEnabled() && !Tonemapper->bDoGammaOnly)
-			{
-				int32 TonemapperScreenPercentage = CVarTonemapperScreenPercentage.GetValueOnRenderThread();
+			// 0=none..1=full
+			FRCPassPostProcessUpscale::PaniniParams PaniniConfig;
 
-				if(TonemapperScreenPercentage != 0)
-				{
-					if(Context.FinalOutput.GetPass() == Tonemapper)
-					{
-						Tonemapper->bDoScreenPercentageInTonemapper = true;
-						// the following pass is no longer needed
-						bDoScreenPercentage = false;
-					}
-				}
+			if (View.IsPerspectiveProjection() && !GEngine->StereoRenderingDevice.IsValid())
+			{
+				PaniniConfig.D = FMath::Max(CVarUpscalePaniniD.GetValueOnRenderThread(), 0.0f);
+				PaniniConfig.S = CVarUpscalePaniniS.GetValueOnRenderThread();
+				PaniniConfig.ScreenFit = FMath::Max(CVarUpscalePaniniScreenFit.GetValueOnRenderThread(), 0.0f);
 			}
 
-			if (PaniniConfig.IsEnabled() || bDoScreenPercentage)
+		// Do not use upscale if SeparateRenderTarget is in use!
+			if ((PaniniConfig.D > 0.01f || View.UnscaledViewRect != View.ViewRect) &&
+				(bHMDWantsUpscale || !View.Family->EngineShowFlags.StereoRendering || (!View.Family->EngineShowFlags.HMDDistortion && !View.Family->bUseSeparateRenderTarget)))
 			{
 				int32 UpscaleQuality = CVarUpscaleQuality.GetValueOnRenderThread();
 				UpscaleQuality = FMath::Clamp(UpscaleQuality, 0, 3);
@@ -1866,6 +1918,9 @@ void FPostProcessing::Process(FRHICommandListImmediate& RHICmdList, FViewInfo& V
 				Node->SetInput(ePId_Input1, FRenderingCompositeOutputRef(Context.FinalOutput)); // Point sampling.
 				Context.FinalOutput = FRenderingCompositeOutputRef(Node);
 			}
+
+			// not needed but more correct
+			bScreenPercentageIsDone = true;
 		}
 
 		// After the graph is built but before the graph is processed.
diff --git a/Engine/Source/Runtime/Slate/Private/Framework/Application/SlateApplication.cpp b/Engine/Source/Runtime/Slate/Private/Framework/Application/SlateApplication.cpp
--- a/Engine/Source/Runtime/Slate/Private/Framework/Application/SlateApplication.cpp
+++ b/Engine/Source/Runtime/Slate/Private/Framework/Application/SlateApplication.cpp
@@ -1312,8 +1312,6 @@ extern SLATECORE_API int32 bFoldTick;
 
 void FSlateApplication::Tick()
 {
-	FScopeLock SlateTickAccess(&SlateTickCriticalSection);
-
 	FPlatformMisc::BeginNamedEvent(FColor::Magenta, "Slate::Tick");
 
 	{
diff --git a/Engine/Source/Runtime/Slate/Public/Framework/Application/SlateApplication.h b/Engine/Source/Runtime/Slate/Public/Framework/Application/SlateApplication.h
--- a/Engine/Source/Runtime/Slate/Public/Framework/Application/SlateApplication.h
+++ b/Engine/Source/Runtime/Slate/Public/Framework/Application/SlateApplication.h
@@ -1679,6 +1679,4 @@ private:
 
 	/** Delegate for post slate Tick */
 	FSlateTickEvent PostTickEvent;
-
-	FCriticalSection SlateTickCriticalSection;
 };
diff --git a/Engine/Source/Runtime/Slate/Public/Framework/Commands/InputChord.h b/Engine/Source/Runtime/Slate/Public/Framework/Commands/InputChord.h
--- a/Engine/Source/Runtime/Slate/Public/Framework/Commands/InputChord.h
+++ b/Engine/Source/Runtime/Slate/Public/Framework/Commands/InputChord.h
@@ -8,7 +8,7 @@
 typedef struct FInputChord FInputGesture;
 
 /** An Input Chord is a key and the modifier keys that are to be held with it. */
-USTRUCT(BlueprintType)
+USTRUCT()
 struct SLATE_API FInputChord
 {
 	GENERATED_USTRUCT_BODY()
diff --git a/Engine/Source/Runtime/UMG/Public/UMG.h b/Engine/Source/Runtime/UMG/Public/UMG.h
--- a/Engine/Source/Runtime/UMG/Public/UMG.h
+++ b/Engine/Source/Runtime/UMG/Public/UMG.h
@@ -121,8 +121,6 @@
 
 #include "NativeWidgetHost.h"
 
-#include "InputKeySelector.h"
-
 // Slate
 #include "SObjectWidget.h"
 
diff --git a/Engine/Source/Runtime/Windows/D3D11RHI/Private/D3D11Commands.cpp b/Engine/Source/Runtime/Windows/D3D11RHI/Private/D3D11Commands.cpp
--- a/Engine/Source/Runtime/Windows/D3D11RHI/Private/D3D11Commands.cpp
+++ b/Engine/Source/Runtime/Windows/D3D11RHI/Private/D3D11Commands.cpp
@@ -715,19 +715,8 @@ void FD3D11DynamicRHI::ValidateExclusiveDepthStencilAccess(FExclusiveDepthStenci
 		const bool bDstStencilWrite = CurrentDSVAccessType.IsStencilWrite();
 
 		// requested access is not possible, fix SetRenderTarget EExclusiveDepthStencil or request a different one
-		ensureMsgf(
-			!bSrcDepthWrite || bDstDepthWrite, 
-			TEXT("Expected: SrcDepthWrite := false or DstDepthWrite := true. Actual: SrcDepthWrite := %s or DstDepthWrite := %s"),
-			(bSrcDepthWrite) ? TEXT("true") : TEXT("false"),
-			(bDstDepthWrite) ? TEXT("true") : TEXT("false")
-			);
-
-		ensureMsgf(
-			!bSrcStencilWrite || bDstStencilWrite,
-			TEXT("Expected: SrcStencilWrite := false or DstStencilWrite := true. Actual: SrcStencilWrite := %s or DstStencilWrite := %s"),
-			(bSrcStencilWrite) ? TEXT("true") : TEXT("false"),
-			(bDstStencilWrite) ? TEXT("true") : TEXT("false")
-			);
+		check(!bSrcDepthWrite || bDstDepthWrite);
+		check(!bSrcStencilWrite || bDstStencilWrite);
 	}
 }
 
diff --git a/Engine/Source/Runtime/Windows/D3D11RHI/Private/D3D11Shaders.cpp b/Engine/Source/Runtime/Windows/D3D11RHI/Private/D3D11Shaders.cpp
--- a/Engine/Source/Runtime/Windows/D3D11RHI/Private/D3D11Shaders.cpp
+++ b/Engine/Source/Runtime/Windows/D3D11RHI/Private/D3D11Shaders.cpp
@@ -43,8 +43,9 @@ FVertexShaderRHIRef FD3D11DynamicRHI::RHICreateVertexShader(const TArray<uint8>&
 	const uint8* CodePtr = Code.GetData() + Offset;
 	const size_t CodeSize = ShaderCode.GetActualShaderCodeSize() - Offset;
 
+	VERIFYD3D11RESULT( Direct3DDevice->CreateVertexShader( (void*)CodePtr, CodeSize, NULL, Shader->Resource.GetInitReference() ) );
+
 	ReadShaderOptionalData(ShaderCode, *Shader);
-	VERIFYD3D11SHADERRESULT(Direct3DDevice->CreateVertexShader( (void*)CodePtr, CodeSize, NULL, Shader->Resource.GetInitReference() ), Shader, Direct3DDevice);
 	
 	// TEMP
 	Shader->Code = Code;
@@ -65,9 +66,9 @@ FGeometryShaderRHIRef FD3D11DynamicRHI::RHICreateGeometryShader(const TArray<uin
 	const uint8* CodePtr = Code.GetData() + Offset;
 	const size_t CodeSize = ShaderCode.GetActualShaderCodeSize() - Offset;
 
+	VERIFYD3D11RESULT( Direct3DDevice->CreateGeometryShader( (void*)CodePtr, CodeSize, NULL, Shader->Resource.GetInitReference() ) );
+	
 	ReadShaderOptionalData(ShaderCode, *Shader);
-	VERIFYD3D11SHADERRESULT(Direct3DDevice->CreateGeometryShader( (void*)CodePtr, CodeSize, NULL, Shader->Resource.GetInitReference() ), Shader, Direct3DDevice);
-
 
 	return Shader;
 }
@@ -102,19 +103,17 @@ FGeometryShaderRHIRef FD3D11DynamicRHI::RHICreateGeometryShaderWithStreamOutput(
 		StreamOutEntries[EntryIndex].OutputSlot = ElementList[EntryIndex].OutputSlot;
 	}
 
-	VERIFYD3D11SHADERRESULT( Direct3DDevice->CreateGeometryShaderWithStreamOutput(
-			(void*)CodePtr,
-			CodeSize,
-			StreamOutEntries,
-			ElementList.Num(),
-			Strides,
-			NumStrides,
-			D3DRasterizedStream,
-			NULL,
-			Shader->Resource.GetInitReference()
-			),
-		Shader,
-		Direct3DDevice);
+	VERIFYD3D11RESULT( Direct3DDevice->CreateGeometryShaderWithStreamOutput(
+		(void*)CodePtr,
+		CodeSize,
+		StreamOutEntries,
+		ElementList.Num(),
+		Strides,
+		NumStrides,
+		D3DRasterizedStream,
+		NULL,
+		Shader->Resource.GetInitReference()
+	) );
 	
 	auto PackedResourceCounts = ShaderCode.FindOptionalData<FShaderCodePackedResourceCounts>();
 	check(PackedResourceCounts);
@@ -138,8 +137,9 @@ FHullShaderRHIRef FD3D11DynamicRHI::RHICreateHullShader(const TArray<uint8>& Cod
 	const uint8* CodePtr = Code.GetData() + Offset;
 	const size_t CodeSize = ShaderCode.GetActualShaderCodeSize() - Offset;
 
+	VERIFYD3D11RESULT( Direct3DDevice->CreateHullShader( (void*)CodePtr, CodeSize, NULL, Shader->Resource.GetInitReference() ) );
+	
 	ReadShaderOptionalData(ShaderCode, *Shader);
-	VERIFYD3D11SHADERRESULT( Direct3DDevice->CreateHullShader( (void*)CodePtr, CodeSize, NULL, Shader->Resource.GetInitReference() ), Shader, Direct3DDevice);
 
 	return Shader;
 }
@@ -156,8 +156,9 @@ FDomainShaderRHIRef FD3D11DynamicRHI::RHICreateDomainShader(const TArray<uint8>&
 	const uint8* CodePtr = Code.GetData() + Offset;
 	const size_t CodeSize = ShaderCode.GetActualShaderCodeSize() - Offset;
 
+	VERIFYD3D11RESULT( Direct3DDevice->CreateDomainShader( (void*)CodePtr, CodeSize, NULL, Shader->Resource.GetInitReference() ) );
+	
 	ReadShaderOptionalData(ShaderCode, *Shader);
-	VERIFYD3D11SHADERRESULT( Direct3DDevice->CreateDomainShader( (void*)CodePtr, CodeSize, NULL, Shader->Resource.GetInitReference() ), Shader, Direct3DDevice);
 
 	return Shader;
 }
@@ -174,8 +175,9 @@ FPixelShaderRHIRef FD3D11DynamicRHI::RHICreatePixelShader(const TArray<uint8>& C
 	const uint8* CodePtr = Code.GetData() + Offset;
 	const size_t CodeSize = ShaderCode.GetActualShaderCodeSize() - Offset;
 
+	VERIFYD3D11RESULT( Direct3DDevice->CreatePixelShader( (void*)CodePtr, CodeSize, NULL, Shader->Resource.GetInitReference() ) );
+	
 	ReadShaderOptionalData(ShaderCode, *Shader);
-	VERIFYD3D11SHADERRESULT(Direct3DDevice->CreatePixelShader( (void*)CodePtr, CodeSize, NULL, Shader->Resource.GetInitReference() ), Shader, Direct3DDevice);
 
 	return Shader;
 }
@@ -192,8 +194,9 @@ FComputeShaderRHIRef FD3D11DynamicRHI::RHICreateComputeShader(const TArray<uint8
 	const uint8* CodePtr = Code.GetData() + Offset;
 	const size_t CodeSize = ShaderCode.GetActualShaderCodeSize() - Offset;
 
+	VERIFYD3D11RESULT( Direct3DDevice->CreateComputeShader( (void*)CodePtr, CodeSize, NULL, Shader->Resource.GetInitReference() ) );
+	
 	ReadShaderOptionalData(ShaderCode, *Shader);
-	VERIFYD3D11SHADERRESULT( Direct3DDevice->CreateComputeShader( (void*)CodePtr, CodeSize, NULL, Shader->Resource.GetInitReference() ), Shader, Direct3DDevice);
 
 	return Shader;
 }
diff --git a/Engine/Source/Runtime/Windows/D3D11RHI/Private/D3D11Util.cpp b/Engine/Source/Runtime/Windows/D3D11RHI/Private/D3D11Util.cpp
--- a/Engine/Source/Runtime/Windows/D3D11RHI/Private/D3D11Util.cpp
+++ b/Engine/Source/Runtime/Windows/D3D11RHI/Private/D3D11Util.cpp
@@ -200,28 +200,6 @@ void VerifyD3D11Result(HRESULT D3DResult,const ANSICHAR* Code,const ANSICHAR* Fi
 	UE_LOG(LogD3D11RHI, Fatal,TEXT("%s failed \n at %s:%u \n with error %s"),ANSI_TO_TCHAR(Code),ANSI_TO_TCHAR(Filename),Line,*ErrorString);
 }
 
-void VerifyD3D11ShaderResult(FRHIShader* Shader, HRESULT D3DResult, const ANSICHAR* Code, const ANSICHAR* Filename, uint32 Line, ID3D11Device* Device)
-{
-	check(FAILED(D3DResult));
-
-	const FString& ErrorString = GetD3D11ErrorString(D3DResult, Device);
-
-#if !(UE_BUILD_SHIPPING || UE_BUILD_TEST)
-	if (Shader->ShaderName.Len())
-	{
-		UE_LOG(LogD3D11RHI, Error, TEXT("%s failed trying to create shader %s\n at %s:%u \n with error %s"), ANSI_TO_TCHAR(Code), *Shader->ShaderName, ANSI_TO_TCHAR(Filename), Line, *ErrorString);
-		TerminateOnDeviceRemoved(D3DResult);
-		TerminateOnOutOfMemory(D3DResult, false);
-
-		UE_LOG(LogD3D11RHI, Fatal, TEXT("%s failed trying to create shader %s \n at %s:%u \n with error %s"), ANSI_TO_TCHAR(Code), *Shader->ShaderName, ANSI_TO_TCHAR(Filename), Line, *ErrorString);
-	}
-	else
-#endif
-	{
-		VerifyD3D11Result(D3DResult, Code, Filename, Line, Device);
-	}
-}
-
 void VerifyD3D11CreateTextureResult(HRESULT D3DResult,const ANSICHAR* Code,const ANSICHAR* Filename,uint32 Line,uint32 SizeX,uint32 SizeY,uint32 SizeZ,uint8 Format,uint32 NumMips,uint32 Flags)
 {
 	check(FAILED(D3DResult));
diff --git a/Engine/Source/Runtime/Windows/D3D11RHI/Private/Windows/WindowsD3D11Device.cpp b/Engine/Source/Runtime/Windows/D3D11RHI/Private/Windows/WindowsD3D11Device.cpp
--- a/Engine/Source/Runtime/Windows/D3D11RHI/Private/Windows/WindowsD3D11Device.cpp
+++ b/Engine/Source/Runtime/Windows/D3D11RHI/Private/Windows/WindowsD3D11Device.cpp
@@ -448,12 +448,12 @@ void FD3D11DynamicRHI::InitD3DDevice()
 
 		bDeviceRemoved = false;
 
+		// Cleanup the D3D device.
+		CleanupD3DDevice();
+
 		// We currently don't support removed devices because FTexture2DResource can't recreate its RHI resources from scratch.
 		// We would also need to recreate the viewport swap chains from scratch.
 		UE_LOG(LogD3D11RHI, Fatal, TEXT("The Direct3D 11 device that was being used has been removed (Error: %d '%s').  Please restart the game."), hRes, Reason);
-
-		// UE4 no longer supports clean-up and recovery.
-		//CleanupD3DDevice();
 	}
 
 	// If we don't have a device yet, either because this is the first viewport, or the old device was removed, create a device.
@@ -500,7 +500,6 @@ void FD3D11DynamicRHI::InitD3DDevice()
 
 					GRHIAdapterName = AdapterDesc.Description;
 					GRHIVendorId = AdapterDesc.VendorId;
-					GRHIDeviceId = AdapterDesc.DeviceId;
 					
 					// get driver version (todo: share with other RHIs)
 					{
diff --git a/Engine/Source/Runtime/Windows/D3D11RHI/Public/D3D11Util.h b/Engine/Source/Runtime/Windows/D3D11RHI/Public/D3D11Util.h
--- a/Engine/Source/Runtime/Windows/D3D11RHI/Public/D3D11Util.h
+++ b/Engine/Source/Runtime/Windows/D3D11RHI/Public/D3D11Util.h
@@ -29,24 +29,13 @@
 
 /**
  * Checks that the given result isn't a failure.  If it is, the application exits with an appropriate error message.
- * @param	Result - The result code to check.
+ * @param	Result - The result code to check
  * @param	Code - The code which yielded the result.
  * @param	Filename - The filename of the source file containing Code.
  * @param	Line - The line number of Code within Filename.
  */
 extern D3D11RHI_API void VerifyD3D11Result(HRESULT Result,const ANSICHAR* Code,const ANSICHAR* Filename,uint32 Line, ID3D11Device* Device);
 
-/**
- * Checks that the given result isn't a failure.  If it is, the application exits with an appropriate error message.
- * @param	Shader - The shader we are trying to create.
- * @param	Result - The result code to check.
- * @param	Code - The code which yielded the result.
- * @param	Filename - The filename of the source file containing Code.
- * @param	Line - The line number of Code within Filename.
- * @param	Device - The D3D device used to create the shadr.
- */
-extern D3D11RHI_API void VerifyD3D11ShaderResult(class FRHIShader* Shader, HRESULT Result, const ANSICHAR* Code, const ANSICHAR* Filename, uint32 Line, ID3D11Device* Device);
-
 /**
 * Checks that the given result isn't a failure.  If it is, the application exits with an appropriate error message.
 * @param	Result - The result code to check
@@ -62,7 +51,6 @@ extern D3D11RHI_API void VerifyD3D11CreateTextureResult(HRESULT D3DResult,const
  */
 #define VERIFYD3D11RESULT_EX(x, Device)	{HRESULT hr = x; if (FAILED(hr)) { VerifyD3D11Result(hr,#x,__FILE__,__LINE__, Device); }}
 #define VERIFYD3D11RESULT(x)			{HRESULT hr = x; if (FAILED(hr)) { VerifyD3D11Result(hr,#x,__FILE__,__LINE__, 0); }}
-#define VERIFYD3D11SHADERRESULT(Result, Shader, Device) {HRESULT hr = (Result); if (FAILED(hr)) { VerifyD3D11ShaderResult(Shader, hr, #Result,__FILE__,__LINE__, Device); }}
 #define VERIFYD3D11CREATETEXTURERESULT(x,SizeX,SizeY,SizeZ,Format,NumMips,Flags) {HRESULT hr = x; if (FAILED(hr)) { VerifyD3D11CreateTextureResult(hr,#x,__FILE__,__LINE__,SizeX,SizeY,SizeZ,Format,NumMips,Flags); }}
 
 /**
diff --git a/Engine/Source/Runtime/Windows/D3D12RHI/Private/Windows/WindowsD3D12Device.cpp b/Engine/Source/Runtime/Windows/D3D12RHI/Private/Windows/WindowsD3D12Device.cpp
--- a/Engine/Source/Runtime/Windows/D3D12RHI/Private/Windows/WindowsD3D12Device.cpp
+++ b/Engine/Source/Runtime/Windows/D3D12RHI/Private/Windows/WindowsD3D12Device.cpp
@@ -402,7 +402,6 @@ void FD3D12DynamicRHI::PerRHISetup(FD3D12Device* MainDevice)
     // TODO MS: We need to do this after CreateDevice due to a bug in QueryVideoMemoryInfo
     GRHIAdapterName = AdapterDesc->Description;
     GRHIVendorId = AdapterDesc->VendorId;
-	GRHIDeviceId = AdapterDesc->DeviceId;
 
     // Issue: 32bit windows doesn't report 64bit value, we take what we get.
     FD3D12GlobalStats::GDedicatedVideoMemory = int64(AdapterDesc->DedicatedVideoMemory);
