Copying //UE4/Orion-Staging to //UE4/Main (Source: //Orion/Dev-General @ 2845681)

diff --git a/Engine/Source/Runtime/PerfCounters/Private/PerfCounters.cpp b/Engine/Source/Runtime/PerfCounters/Private/PerfCounters.cpp
--- a/Engine/Source/Runtime/PerfCounters/Private/PerfCounters.cpp
+++ b/Engine/Source/Runtime/PerfCounters/Private/PerfCounters.cpp
@@ -11,6 +11,7 @@
 
 FPerfCounters::FPerfCounters(const FString& InUniqueInstanceId)
 : UniqueInstanceId(InUniqueInstanceId)
+, InternalCountersUpdateInterval(60)
 , Socket(nullptr)
 {
 }
@@ -30,6 +31,13 @@ FPerfCounters::~FPerfCounters()
 
 bool FPerfCounters::Initialize()
 {
+	float ConfigInternalCountersUpdateInterval = 60.0;
+	if (GConfig->GetFloat(TEXT("PerfCounters"), TEXT("InternalCountersUpdateInterval"), ConfigInternalCountersUpdateInterval, GEngineIni))
+	{
+		InternalCountersUpdateInterval = ConfigInternalCountersUpdateInterval;
+	}
+	LastTimeInternalCountersUpdated = FPlatformTime::Seconds() - InternalCountersUpdateInterval * FMath::FRand();	// randomize between servers
+
 	// get the requested port from the command line (if specified)
 	int32 StatsPort = -1;
 	FParse::Value(FCommandLine::Get(), TEXT("statsPort="), StatsPort);
@@ -231,11 +239,9 @@ bool FPerfCounters::Tick(float DeltaTime)
 		}
 	}
 
-	// set some internal perf stats ([RCL] FIXME 2015-12-08: move to a better place and make configurable)
-	const float kInternalStatsUpdateInterval = 60.0;	// 60 seconds
-	static float LastTimeStatsUpdated = FPlatformTime::Seconds() - kInternalStatsUpdateInterval * FMath::FRand();
+	// set some internal perf stats ([RCL] FIXME 2015-12-08: move to a better place)
 	float CurrentTime = FPlatformTime::Seconds();
-	if (CurrentTime - LastTimeStatsUpdated > kInternalStatsUpdateInterval)
+	if (CurrentTime - LastTimeInternalCountersUpdated > InternalCountersUpdateInterval)
 	{
 		FPlatformMemoryStats Stats = FPlatformMemory::GetStats();
 		Set(TEXT("AvailablePhysicalMemoryMB"), static_cast<uint64>(Stats.AvailablePhysical / (1024 * 1024)));
@@ -248,7 +254,7 @@ bool FPerfCounters::Tick(float DeltaTime)
 			Set(TEXT("FreeSpaceOnLogFileDiskInMB"), static_cast<uint64>(FreeBytesOnLogDrive / (1024 * 1024)));
 		}
 
-		LastTimeStatsUpdated = CurrentTime;
+		LastTimeInternalCountersUpdated = CurrentTime;
 	}
 
 	// keep ticking
diff --git a/Engine/Source/Runtime/PerfCounters/Private/PerfCounters.h b/Engine/Source/Runtime/PerfCounters/Private/PerfCounters.h
--- a/Engine/Source/Runtime/PerfCounters/Private/PerfCounters.h
+++ b/Engine/Source/Runtime/PerfCounters/Private/PerfCounters.h
@@ -102,6 +102,12 @@ private:
 	/** Unique name of this instance */
 	FString UniqueInstanceId;
 
+	/** Interval (in seconds) to update internal counters, added by the module itself, a config value */
+	float InternalCountersUpdateInterval;
+
+	/** Last time internal counters were updated */
+	float LastTimeInternalCountersUpdated;
+
 	/** Map of all known performance counters */
 	TMap<FString, FJsonVariant>  PerfCounterMap;
 
diff --git a/Engine/Source/Runtime/RHI/Private/DynamicRHI.cpp b/Engine/Source/Runtime/RHI/Private/DynamicRHI.cpp
--- a/Engine/Source/Runtime/RHI/Private/DynamicRHI.cpp
+++ b/Engine/Source/Runtime/RHI/Private/DynamicRHI.cpp
@@ -48,26 +48,24 @@ void InitNullRHI()
 // e.g. 36143 for the NVIDIA driver 361.43
 // no floating point to avoid precision issues
 // @return -1 if the version is unknown
-static int32 ExtractGPUDriverVersionInt()
+static int32 ExtractGPUDriverVersionInt(const FGPUDriverInfo& DriverInfo)
 {
 	int32 Ret = -1;
 
 	// we use the internal version, not the user version to avoid problem where the name was altered 
-	const FString& InternalVersion = GRHIAdapterInternalDriverVersion;
+	const FString& InternalVersion = DriverInfo.InternalDriverVersion;
 
-	if(IsRHIDeviceNVIDIA())
+	if(DriverInfo.IsNVIDIA())
 	{
-		if(InternalVersion.Left(6) == TEXT("9.18.1"))		// e.g. "9.18.136.143"
-		{
-			// e.g. 36.143
-			FString RightPart = InternalVersion.RightChop(6);
-			// e.g. 36143
-			FString CompactNumber = RightPart.Replace(TEXT("."), TEXT(""));
+		// we don't care about the windows version so we don't look at the front part of the driver version
+		// e.g. 36.143
+		FString RightPart = InternalVersion.Right(6);
+		// e.g. 36143
+		FString CompactNumber = RightPart.Replace(TEXT("."), TEXT(""));
 
-			if(FCString::IsNumeric(*CompactNumber))
-			{
-				Ret = FCString::Atoi(*CompactNumber);
-			}
+		if(FCString::IsNumeric(*CompactNumber))
+		{
+			Ret = FCString::Atoi(*CompactNumber);
 		}
 	}
 	else
@@ -87,17 +85,27 @@ void RHIDetectAndWarnOfBadDrivers()
 		return;
 	}
 
+	FGPUDriverInfo DriverInfo;
+
+	// later we should make the globals use the struct directly
+	DriverInfo.VendorId = GRHIVendorId;
+	DriverInfo.DeviceDescription = GRHIAdapterName;
+	DriverInfo.ProviderName = TEXT("Unknown");
+	DriverInfo.InternalDriverVersion = GRHIAdapterInternalDriverVersion;
+	DriverInfo.UserDriverVersion = GRHIAdapterUserDriverVersion;
+	DriverInfo.DriverDate = GRHIAdapterDriverDate;
+
 	// for testing
 	if(CVarValue == 2)
 	{
-		GRHIVendorId = 0x10DE;
-		GRHIAdapterName = TEXT("Generic NVIDIA");
-		GRHIAdapterUserDriverVersion = TEXT("361.43");
-		GRHIAdapterInternalDriverVersion = TEXT("9.18.136.143");
-		GRHIAdapterDriverDate = TEXT("01-01-1900");
+		DriverInfo.VendorId = 0x10DE;
+		DriverInfo.DeviceDescription = TEXT("Test NVIDIA");
+		DriverInfo.UserDriverVersion = TEXT("361.43");
+		DriverInfo.InternalDriverVersion = TEXT("9.18.136.143");
+		DriverInfo.DriverDate = TEXT("01-01-1900");
 	}
 
-	int32 VersionInt = ExtractGPUDriverVersionInt();
+	int32 VersionInt = ExtractGPUDriverVersionInt(DriverInfo);
 
 	if(VersionInt < 0)
 	{
@@ -106,7 +114,7 @@ void RHIDetectAndWarnOfBadDrivers()
 	}
 
 	// can be made data driven:
-	if(IsRHIDeviceNVIDIA())
+	if(DriverInfo.IsNVIDIA())
 	{
 		// todo: make this data driven
 		FString SuggestedVersion = TEXT("359.06");
@@ -116,40 +124,26 @@ void RHIDetectAndWarnOfBadDrivers()
 		|| VersionInt <= 34709)	// 347.09 "old NVIDIA driver version we've seen driver crashes with Paragon content", we might have to adjust that further
 		{
 			FFormatNamedArguments Args;
-			Args.Add(TEXT("AdapterName"), FText::FromString(GRHIAdapterName));
-			Args.Add(TEXT("UserDriverVersion"), FText::FromString(GRHIAdapterUserDriverVersion));
-			Args.Add(TEXT("InternalDriverVersion"), FText::FromString(GRHIAdapterInternalDriverVersion));
-			Args.Add(TEXT("DriverDate"), FText::FromString(GRHIAdapterDriverDate));
+			Args.Add(TEXT("AdapterName"), FText::FromString(DriverInfo.DeviceDescription));
+			Args.Add(TEXT("UserDriverVersion"), FText::FromString(DriverInfo.UserDriverVersion));
+			Args.Add(TEXT("InternalDriverVersion"), FText::FromString(DriverInfo.InternalDriverVersion));
+			Args.Add(TEXT("DriverDate"), FText::FromString(DriverInfo.DriverDate));
 			Args.Add(TEXT("SuggestedVersion"), FText::FromString(SuggestedVersion));
 
-			// explicitly compact
 			FText LocalizedMsg = FText::Format(NSLOCTEXT("MessageDialog", "VideoCardDriverIssueReport",
-				"Suggested version: {SuggestedVersion} (newer drivers might be good as well)\nCurrent version: {UserDriverVersion} ({InternalDriverVersion}) {DriverDate}\nName: {AdapterName}"),
+				"Your current video card driver version might cause problems.\n(reduced performance, artifacts, stalls, crashes, restarts)\n\nYour current driver:\n  Name:\t{AdapterName}\n  Version:\t{UserDriverVersion} (internal {InternalDriverVersion})\n  Date:\t{DriverDate}\n\nSuggested driver:\n  Version:\t{SuggestedVersion}\n\n"),
 				Args);
 
-			// less compact and more readable, maybe for later
-			/*
-				"Current driver:\n"
-				"  Name:\t{AdapterName}\n"
-				"  Version:\t{UserDriverVersion} (internal {InternalDriverVersion})\n"
-				"  Date:\t{DriverDate}\n"
-				"\n"
-				"Suggested driver (newer drivers might be good as well):\n"
-				"  Version: {SuggestedVersion}\n"
-				"\n"
-				"Please update your drivers to avoid problems (performance, stalls, crashes, reboot)."
-			*/
-
 			FPlatformMisc::MessageBoxExt(EAppMsgType::Ok, 
 				*LocalizedMsg.ToString(),
-				*NSLOCTEXT("MessageDialog", "TitleVideoCardDriverIssue", "WARNING: Video card driver with known issues").ToString());
+				*NSLOCTEXT("MessageDialog", "TitleVideoCardDriverIssue", "WARNING: Video card driver").ToString());
 		}
 	}
-	else if(IsRHIDeviceAMD())
+	else if(DriverInfo.IsAMD())
 	{
 		// so far we don't have any data on issues
 	}
-	else if(IsRHIDeviceIntel())
+	else if(DriverInfo.IsIntel())
 	{
 		// so far we don't have any data on issues
 	}
@@ -165,7 +159,7 @@ void RHIInit(bool bHasEditorToken)
 	{
 		GRHICommandList.LatchBypass(); // read commandline for bypass flag
 
-		if (USE_NULL_RHI || FParse::Param(FCommandLine::Get(),TEXT("nullrhi")) || !FApp::CanEverRender())
+		if (FApp::ShouldUseNullRHI())
 		{
 			InitNullRHI();
 		}
diff --git a/Engine/Source/Runtime/RHI/Private/GPUDefragAllocator.cpp b/Engine/Source/Runtime/RHI/Private/GPUDefragAllocator.cpp
--- a/Engine/Source/Runtime/RHI/Private/GPUDefragAllocator.cpp
+++ b/Engine/Source/Runtime/RHI/Private/GPUDefragAllocator.cpp
@@ -52,6 +52,7 @@ void FGPUDefragAllocator::FMemoryChunk::LinkFree(bool bMaintainSortOrder, FMemor
 	bIsAvailable = true;
 	DefragCounter = 0;
 	UserPayload = 0;
+	bTail = false;
 
 	if (bMaintainSortOrder == false)
 	{
@@ -205,15 +206,23 @@ void* FGPUDefragAllocator::Allocate(int64 AllocationSize, int32 Alignment, TStat
 * @param bAsync			If true, allows allocating from relocating chunks and maintains the free-list sort order.
 * @return					The memory chunk that was allocated (the original chunk could've been split).
 */
-FGPUDefragAllocator::FMemoryChunk* FGPUDefragAllocator::AllocateChunk(FMemoryChunk* FreeChunk, int32 AllocationSize, bool bAsync)
+FGPUDefragAllocator::FMemoryChunk* FGPUDefragAllocator::AllocateChunk(FMemoryChunk* FreeChunk, int32 AllocationSize, bool bAsync, bool bDoValidation)
 {
 	check(FreeChunk);
 	check(FreeChunk->GetAvailableSize() >= AllocationSize);
+	check(!FreeChunk->IsLocked());
+
+#if VALIDATE_SYNC_SIZE
+	printf("allocate before splits: 0x%p, %i, %i, %i, SyncSize: %i\n", FreeChunk->Base, (int32)FreeChunk->Size, (int32)AllocationSize, (int32)GetCurrentSyncIndex(), (int32)FreeChunk->SyncSize);
+#endif
 
 	// If this is an immediate allocation (i.e. the CPU will start accessing the memory right away)
 	// and the beginning of the chunk is currently being relocated by the GPU, split that part off and allocate from the rest.
 	if (!bAsync && FreeChunk->IsRelocating() && FreeChunk->SyncSize > 0 && FreeChunk->SyncSize < FreeChunk->Size)
 	{
+#if VALIDATE_SYNC_SIZE
+		printf("splitting:\n");
+#endif
 		Split(FreeChunk, FreeChunk->SyncSize, bAsync);
 		FreeChunk = FreeChunk->NextChunk;
 	}
@@ -224,6 +233,9 @@ FGPUDefragAllocator::FMemoryChunk* FGPUDefragAllocator::AllocateChunk(FMemoryChu
 	// Split chunk to avoid waste.
 	if (FreeChunk->Size > AllocationSize)
 	{
+#if VALIDATE_SYNC_SIZE
+		printf("splitting again:\n");
+#endif
 		Split(FreeChunk, AllocationSize, bAsync);
 	}
 
@@ -235,6 +247,14 @@ FGPUDefragAllocator::FMemoryChunk* FGPUDefragAllocator::AllocateChunk(FMemoryChu
 	FPlatformAtomics::InterlockedAdd(&AllocatedMemorySize, FreeChunk->Size);
 	FPlatformAtomics::InterlockedAdd(&AvailableMemorySize, -FreeChunk->Size);
 
+#if VALIDATE_SYNC_SIZE
+	printf("allocate: 0x%p, %i, %i\n", FreeChunk->Base, (int32)AllocationSize, (int32)GetCurrentSyncIndex());
+	if (bDoValidation)
+	{
+		ValidateRelocations(FreeChunk->Base, AllocationSize);
+	}
+#endif
+
 	// Keep track of mapping and return pointer.
 	PointerToChunkMap.Add(FreeChunk->Base, FreeChunk);
 	return FreeChunk;
@@ -249,6 +269,10 @@ FGPUDefragAllocator::FMemoryChunk* FGPUDefragAllocator::AllocateChunk(FMemoryChu
 */
 void FGPUDefragAllocator::FreeChunk(FMemoryChunk* Chunk, bool bMaintainSortedFreelist)
 {
+#if VALIDATE_SYNC_SIZE
+	printf("FreeChunk: 0x%p, %i, %i, %i\n", Chunk->Base, (int32)Chunk->Size, (int32)Chunk->SyncSize, (int32)Chunk->SyncIndex);
+#endif
+
 	// Remove the entry
 	PointerToChunkMap.Remove(Chunk->Base);
 	Chunk->Stat = TStatId();
@@ -306,7 +330,7 @@ void FGPUDefragAllocator::Lock(const void* Pointer)
 {
 	FScopeLock Lock(&SynchronizationObject);
 	FMemoryChunk* MatchingChunk = PointerToChunkMap.FindRef(Pointer);
-	check(MatchingChunk);
+	checkf(MatchingChunk, TEXT("Couldn't find chunk for address: 0x%p"), Pointer);
 
 	// Is this chunk is currently being relocated asynchronously (by the GPU)?
 	if (MatchingChunk->IsRelocating())
@@ -315,6 +339,7 @@ void FGPUDefragAllocator::Lock(const void* Pointer)
 		FinishAllRelocations();
 	}
 
+	check(MatchingChunk->LockCount >= 0);
 	++MatchingChunk->LockCount;
 	++NumLockedChunks;
 }
@@ -329,7 +354,7 @@ void FGPUDefragAllocator::Unlock(const void* Pointer)
 	FScopeLock Lock(&SynchronizationObject);
 	FMemoryChunk* MatchingChunk = PointerToChunkMap.FindRef(Pointer);
 	check(MatchingChunk && MatchingChunk->IsRelocating() == false);
-
+	check(MatchingChunk->LockCount > 0);
 	--MatchingChunk->LockCount;
 	--NumLockedChunks;
 }
@@ -504,6 +529,9 @@ FGPUDefragAllocator::FMemoryChunk* FGPUDefragAllocator::Shrink(FMemoryChunk* Chu
 
 		// Make the original chunk "free memory".
 		NewFreeChunk = Chunk->PreviousChunk;
+#if VALIDATE_SYNC_SIZE
+		printf("shrink free chunk\n");
+#endif
 		LinkFreeChunk(NewFreeChunk, false);
 	}
 
@@ -613,6 +641,9 @@ void FGPUDefragAllocator::SortFreeList(int32& NumFreeChunks, int64& LargestFreeC
 		FMemoryChunk* ChunkToSort = FirstFreeChunk->NextFreeChunk;
 		while (ChunkToSort)
 		{
+			//chunk to sort will be properly sorted.  Thus the next correct chunk to sort is the next one we haven't sorted.
+			//this ensures we get the proper chunk count also.
+			FMemoryChunk* NextChunkToSort = ChunkToSort->NextFreeChunk;
 			LargestFreeChunk = FMath::Max(LargestFreeChunk, ChunkToSort->Size);
 
 			// Out of order?
@@ -638,7 +669,7 @@ void FGPUDefragAllocator::SortFreeList(int32& NumFreeChunks, int64& LargestFreeC
 				}
 			}
 			LastSortedChunk = ChunkToSort;
-			ChunkToSort = ChunkToSort->NextFreeChunk;
+			ChunkToSort = NextChunkToSort;
 			NumFreeChunks++;
 		}
 	}
@@ -806,6 +837,28 @@ void FGPUDefragAllocator::CheckForErrors(bool bCheckSortedFreeList)
 	check(TotalFreeMem == AvailableMemorySize);
 }
 
+#if VALIDATE_SYNC_SIZE
+void FGPUDefragAllocator::ValidateRelocations(uint8* UsedBaseAddress, uint64 Size)
+{
+
+	FScopeLock SyncLock(&SynchronizationObject);
+
+	for (int32 i = 0; i < Relocations.Num(); ++i)
+	{
+		const FRelocationEntry& Relocation = Relocations[i];
+
+		bool bBeforeOrig = (UsedBaseAddress + Size) <= Relocation.OldBase;
+		bool bAfterOrig = (UsedBaseAddress) >= (Relocation.OldBase + Relocation.Size);
+
+		bool bBeforeNew = (UsedBaseAddress + Size) <= Relocation.NewBase;
+		bool bAfterNew = (UsedBaseAddress) >= (Relocation.NewBase + Relocation.Size);
+
+		bool bAfterSync = Relocation.SyncIndex <= CompletedSyncIndex;
+		checkf(((bBeforeOrig || bAfterOrig) && (bBeforeNew || bAfterNew)) || bAfterSync, TEXT("Corruption Hazard, Allocation not protected by sync size."));
+	}
+}
+#endif
+
 /**
 * Initiates an async relocation of an allocated chunk into a free chunk.
 * Takes potential reallocation request into account.
@@ -817,21 +870,32 @@ void FGPUDefragAllocator::CheckForErrors(bool bCheckSortedFreeList)
 */
 FGPUDefragAllocator::FMemoryChunk* FGPUDefragAllocator::RelocateIntoFreeChunk(FRelocationStats& Stats, FMemoryChunk* FreeChunk, FMemoryChunk* SourceChunk)
 {
-	// Save off important data from 'SourceChunk', since it will get modified by the call to LinkFreeChunk().
-	uint8* OldBaseAddress = SourceChunk->Base;
+	check(!FreeChunk->IsLocked());
+	check(!SourceChunk->IsLocked());
+	check(FreeChunk->bIsAvailable);
+	check(!SourceChunk->bIsAvailable);
+
+	// Save off important data from 'SourceChunk', since it will get modified by the call to LinkFreeChunk().	
 	void* UserPayload = SourceChunk->UserPayload;
-	int32 OldSize = SourceChunk->Size;
-	const int32 NewSize = SourceChunk->GetFinalSize();
-	const int32 UsedSize = FMath::Min(NewSize, OldSize);
+	const int64 OldSize = SourceChunk->Size;
+	const uint8* SourceOldBase = SourceChunk->Base;
+	const uint8* DestNewBase = FreeChunk->Base;
+	const int64 NewSize = SourceChunk->GetFinalSize();
+	const int64 UsedSize = FMath::Min(NewSize, OldSize);
 	TStatId	SourceStat = SourceChunk->Stat;
 
 	// Are we relocating into adjacent free chunk?
-	bool bAdjacentRelocation = SourceChunk->PreviousChunk == FreeChunk || SourceChunk->NextChunk == FreeChunk;
+	const bool bFreeChunkPreviousAdjacent = SourceChunk->PreviousChunk == FreeChunk;
+	const bool bFreeChunkNextAdjacent = SourceChunk->NextChunk == FreeChunk;
+	const bool bAdjacentRelocation = bFreeChunkPreviousAdjacent || bFreeChunkNextAdjacent;
 
 	// Enable for debugging:
 	// CheckForErrors( true );
 
 	// Merge adjacent free chunks into SourceChunk to make a single free chunk.
+#if VALIDATE_SYNC_SIZE
+	printf("relocate link free chunk\n");
+#endif
 	LinkFreeChunk(SourceChunk, true);
 
 	FMemoryChunk* DestinationChunk;
@@ -847,19 +911,21 @@ FGPUDefragAllocator::FMemoryChunk* FGPUDefragAllocator::RelocateIntoFreeChunk(FR
 	FreeChunk = nullptr;
 
 	// Leave room for new mips to stream in.
-	int32 DestinationOffset = FMath::Max(NewSize - OldSize, 0);
+	int64 DestinationOffset = FMath::Max((NewSize - OldSize), (int64)0);
 
 	// Relocate the memory if needed
+	bool bRelocated = false;
 	uint8* NewBaseAddress = DestinationChunk->Base;
-	if (OldBaseAddress != (NewBaseAddress + DestinationOffset))
+	if (SourceOldBase != (NewBaseAddress + DestinationOffset))
 	{
-		Relocate(Stats, DestinationChunk, DestinationOffset, OldBaseAddress, UsedSize, UserPayload);
+		Relocate(Stats, DestinationChunk, DestinationOffset, SourceOldBase, UsedSize, UserPayload);
+		bRelocated = true;
 	}
 	// Make sure the destination chunk keeps the UserPayload, no matter what. :)
 	DestinationChunk->UserPayload = UserPayload;
 
 	// Update our book-keeping.
-	PointerToChunkMap.Remove(OldBaseAddress);
+	PointerToChunkMap.Remove(SourceOldBase);
 	PointerToChunkMap.Add(NewBaseAddress, DestinationChunk);
 	SourceChunk->Stat = TStatId();
 	DestinationChunk->Stat = SourceStat;	
@@ -871,6 +937,65 @@ FGPUDefragAllocator::FMemoryChunk* FGPUDefragAllocator::RelocateIntoFreeChunk(FR
 		// Split the DestinationChunk into a used chunk and a free chunk.
 		Split(DestinationChunk, NewSize, true);
 		NextFreeChunk = DestinationChunk->NextChunk;
+		NextFreeChunk->bTail = bAdjacentRelocation;
+
+		if (bAdjacentRelocation)
+		{
+			//if the free chunk was on the right
+			if (bFreeChunkNextAdjacent)
+			{
+				ensureMsgf(!bRelocated, TEXT("free chunk was on the right, why did we relocate?"));
+			}
+			else
+			{
+				//if the free chunk was adjacent on the left then we need to make sure the tail where the source chunk used to reside is not allocated and written to before the GPU
+				//finishes moving it to its new location.
+				//Original Layout                   |xxxxxxxxxxxxx|uuuu|xxxx -> 
+				//Layout after move                 |uuuu|xxxxxxxxxxxxxxxxx| ->
+
+				//Layout after move                 |uuuu|ssssssssstttt| ->
+				//									|uuuu|sssssssss|tttt|
+
+
+				//Layout after tail safety split    |uuuu|xxxxxxxx|tttt|xxxx
+
+				//									|xxxx|uuuuuuuuuuuuu| ->
+				//									|uuuuuuuuuuuuu|xxxx| ->											
+
+				//to compute the relocatable size properly we need to account for any righthand coalescing that may have made the chunk bigger.
+				//right hand coalesce will be added into the tail.
+				int64 LeftShiftSize = (SourceOldBase - SourceChunk->Base);
+				int64 RightAddSize = (NextFreeChunk->Size - LeftShiftSize);
+				int64 RelocatableSize = FMath::Max((int64)((NextFreeChunk->Size - NewSize - RightAddSize)), (int64)0);				
+				if ((RelocatableSize > 0) && NextFreeChunk->GetAvailableSize())
+				{
+					Split(NextFreeChunk, RelocatableSize, true);
+					FMemoryChunk* InFlightTailChunk = NextFreeChunk->NextFreeChunk;					
+
+					//if the tail has right coalesced memory afterwards, split the right coalesce back off into its own chunk for re-use.
+					if (InFlightTailChunk->Size > OldSize)
+					{
+						Split(InFlightTailChunk, OldSize, true);
+					}
+					//defragger code assumes free chunks are not adjacent (they would have been coalesced)
+					//so mark this as 'used' so it can't be allocated while the tail is being moved by the GPU and put it on the deferred free list
+					//hack to avoid internal checks on the hack allocate
+					InFlightTailChunk->SetSyncIndex(CompletedSyncIndex, 0);
+					AllocateChunk(InFlightTailChunk, InFlightTailChunk->Size, false, false);
+					InFlightTailChunk->SetSyncIndex(GetCurrentSyncIndex(), InFlightTailChunk->Size);
+					PendingFreeChunks.AddTail(InFlightTailChunk);
+				}
+				else
+				{
+					//in this case there was no tail and we need to protect the existing sync size, or just the old size.
+					//maxing with the existing sync size keeps any sync-size from the right-coalesce from going away.
+					//should be a case like this:
+					// |xxx|uuuuuuuuuuuuu|fffffff|
+					// |uuuuuuuuuuuuu|sssffffffff|
+					NextFreeChunk->SetSyncIndex(GetCurrentSyncIndex(), FMath::Max(NextFreeChunk->SyncSize, OldSize));
+				}
+			}
+		}
 	}
 	else
 	{
@@ -880,10 +1005,35 @@ FGPUDefragAllocator::FMemoryChunk* FGPUDefragAllocator::RelocateIntoFreeChunk(FR
 	}
 	DestinationChunk->UnlinkFree();
 
+	//if source chunk did not coalesce with another free chunk when it was 'freed' or if it coalesced with a freechunk on its right side then we only need
+	//to protect up to the OldSize of the SourceChunk.
+	int64 SourceChunkSyncSize = OldSize;
+
+	const bool bSourceCoalescedOnLeft = SourceChunk->Base < SourceOldBase;
+	if (bSourceCoalescedOnLeft)
+	{
+		//if source chunk coalesced with a free chunk on the left, then protecting with OldSize does not suffice
+		//Orig. uuuu is being relocated by the GPU    |xx|uuuu|sssxxxxxxxxxxxx ->
+		//                                            |xx|rrrr| ->
+		//We must protect up to the full relocation.  |pp|pppp|ssssxxxxxxxxxxxxxx		
+		SourceChunkSyncSize += (SourceOldBase - SourceChunk->Base);		
+	}
+
+	//if the source chunk coalesced with a right chunk that HAD a sync size, then the coalense function would have
+	//set an appropriate syncsize to cover up to the right side's required protection.  However we are about to override that so it needs to be taken into account here.		
+	SourceChunkSyncSize = FMath::Max(SourceChunkSyncSize, SourceChunk->SyncSize);
+
 	// Mark both chunks as "in use" during the current sync step.
-	// Note: This sync index will propagate if these chunks are involved in any merge/split in the future.
-	SourceChunk->SetSyncIndex(GetCurrentSyncIndex(), OldSize);
-	DestinationChunk->SetSyncIndex(GetCurrentSyncIndex(), NewSize);
+	// Note: This sync index will propagate if these chunks are involved in any merge/split in the future.	
+	SourceChunk->SetSyncIndex(GetCurrentSyncIndex(), SourceChunkSyncSize);
+
+	//if the destination is already relocation a larger size than we are moving, we still need to protect the larger size so allocations don't split off an end that's still moving.
+	int64 DestinationSyncSize = NewSize;
+	if (DestinationChunk->IsRelocating())
+	{
+		DestinationSyncSize = FMath::Max(DestinationChunk->SyncSize, NewSize);
+	}
+	DestinationChunk->SetSyncIndex(GetCurrentSyncIndex(), DestinationSyncSize);
 
 	if (NewSize != OldSize)
 	{
@@ -893,11 +1043,16 @@ FGPUDefragAllocator::FMemoryChunk* FGPUDefragAllocator::RelocateIntoFreeChunk(FR
 		FPlatformAtomics::InterlockedAdd(&PendingMemoryAdjustment, -MemoryAdjustment);
 	}
 
+#if VALIDATE_SYNC_SIZE
+	Relocations.Add(FRelocationEntry(SourceOldBase, DestNewBase, OldSize, GetCurrentSyncIndex()));
+	printf("bFreeAdjOnLeft %i, bFreeAdjOnRight %i\n", bFreeChunkPreviousAdjacent, bFreeChunkNextAdjacent);
+#endif
+
 	// Enable for debugging:
 	// CheckForErrors( true );
 
 	// Did we free up a chunk to the left of NextFreeChunk?
-	if (!bAdjacentRelocation && OldBaseAddress < NewBaseAddress)
+	if (!bAdjacentRelocation && SourceOldBase < NewBaseAddress)
 	{
 		// Use that one for the next defrag step!
 		return SourceChunk;
@@ -1193,6 +1348,7 @@ void FGPUDefragAllocator::FullDefragmentation(FRelocationStats& Stats)
 
 int32 FGPUDefragAllocator::Tick(FRelocationStats& Stats, bool bPanicDefrag)
 {
+	FScopeLock SyncLock(&SynchronizationObject);
 	SET_CYCLE_COUNTER(STAT_TexturePool_Blocked, BlockedCycles);
 	double StartTime = FPlatformTime::Seconds();
 	BlockedCycles = 0;
@@ -1201,6 +1357,10 @@ int32 FGPUDefragAllocator::Tick(FRelocationStats& Stats, bool bPanicDefrag)
 	// There may still be chunks being flagged as 'IsRelocating' due to immediate shrinks between calls.
 	FinishAllRelocations();
 
+#if VALIDATE_SYNC_SIZE	
+	Relocations.Empty();
+#endif
+
 	// Sort the Free chunks.
 	SortFreeList(Stats.NumHoles, Stats.LargestHoleSize);
 
@@ -1232,6 +1392,24 @@ int32 FGPUDefragAllocator::Tick(FRelocationStats& Stats, bool bPanicDefrag)
 	CurrentLargestHole = Stats.LargestHoleSize;
 	CurrentNumHoles = Stats.NumHoles;
 
+#if GPU_DEFRAG_SANITYCHECK
+	FMemoryChunk* TestChunk = FirstChunk;
+	while (TestChunk)
+	{		
+		FMemoryChunk* TestAgainstChunk = FirstChunk->NextChunk;
+		while (TestAgainstChunk)
+		{			
+			bool bBefore = (TestChunk->Base + TestChunk->Size) <= TestAgainstChunk->Base;
+			bool bAfter = TestChunk->Base >= (TestAgainstChunk->Base + TestAgainstChunk->Size);
+
+			checkf(bBefore || bAfter, TEXT("%i, %i"), (int32)bBefore, (int32)bAfter);
+			check(TestChunk->Base != TestAgainstChunk->Base);
+			TestAgainstChunk = TestAgainstChunk->NextChunk;
+		}
+		TestChunk = TestChunk->NextChunk;
+	}	
+#endif
+
 	return Stats.NumBytesRelocated;
 }
 
@@ -1497,6 +1675,7 @@ void FGPUDefragAllocator::DefragmentMemory(FRelocationStats& Stats)
 void FGPUDefragAllocator::Coalesce(FMemoryChunk* FreedChunk)
 {
 	check(FreedChunk);
+	check(!FreedChunk->IsLocked())
 
 	uint32 LatestSyncIndex = 0;
 	int32 LatestSyncSize = 0;
@@ -1507,11 +1686,12 @@ void FGPUDefragAllocator::Coalesce(FMemoryChunk* FreedChunk)
 	FMemoryChunk* LeftChunk = FreedChunk->PreviousChunk;
 	if (LeftChunk && LeftChunk->bIsAvailable)
 	{
+		check(!LeftChunk->IsLocked());
 		LeftSize = LeftChunk->Size;
 
 		// Update relocation data for the left chunk.
 		if (LeftChunk->IsRelocating())
-		{
+		{			
 			LatestSyncIndex = LeftChunk->SyncIndex;
 			LatestSyncSize = LeftChunk->SyncSize;
 		}
@@ -1531,6 +1711,7 @@ void FGPUDefragAllocator::Coalesce(FMemoryChunk* FreedChunk)
 	FMemoryChunk* RightChunk = FreedChunk->NextChunk;
 	if (RightChunk && RightChunk->bIsAvailable)
 	{
+		check(!RightChunk->IsLocked());
 		RightSize = RightChunk->Size;
 
 		// Update relocation data for the right chunk.
@@ -1544,10 +1725,18 @@ void FGPUDefragAllocator::Coalesce(FMemoryChunk* FreedChunk)
 		delete RightChunk;
 	}
 
+#if VALIDATE_SYNC_SIZE
+	printf("FreeChunk Before Coalesce: 0x%p, %i, %i, %i\n", FreedChunk->Base, (int32)FreedChunk->Size, (int32)FreedChunk->SyncSize, (int32)FreedChunk->SyncIndex);
+#endif
+
 	// Merge.
 	FreedChunk->Base -= LeftSize;
 	FreedChunk->Size += LeftSize + RightSize;
 	FreedChunk->SetSyncIndex(LatestSyncIndex, LatestSyncSize);
+
+#if VALIDATE_SYNC_SIZE
+	printf("FreeChunk AFter Coalesce: 0x%p, %i, %i, %i\n", FreedChunk->Base, (int32)FreedChunk->Size, (int32)FreedChunk->SyncSize, (int32)FreedChunk->SyncIndex);
+#endif
 }
 
 /**
diff --git a/Engine/Source/Runtime/RHI/Public/GPUDefragAllocator.h b/Engine/Source/Runtime/RHI/Public/GPUDefragAllocator.h
--- a/Engine/Source/Runtime/RHI/Public/GPUDefragAllocator.h
+++ b/Engine/Source/Runtime/RHI/Public/GPUDefragAllocator.h
@@ -6,6 +6,7 @@
 
 #define LOG_EVERY_ALLOCATION			0
 #define DUMP_ALLOC_FREQUENCY			0 // 100
+#define VALIDATE_SYNC_SIZE				0
 
 
 /*-----------------------------------------------------------------------------
@@ -145,6 +146,7 @@ public:
 			, SyncSize(0)
 			, UserPayload(0)			
 			, Stat(InStat)
+			, bTail(false)
 		{
 			Link(ChunkToInsertAfter);
 			// This is going to change bIsAvailable.
@@ -281,7 +283,7 @@ public:
 		/**
 		*	Returns the current size (in uint8s), or the final size if it has a reallocating request.
 		*/
-		int32		GetFinalSize() const;
+		int64		GetFinalSize() const;
 
 		/**
 		* Sets the relocation sync index.
@@ -289,7 +291,7 @@ public:
 		* @param InSyncSize	Number of uint8s that require GPU synchronization (starting from the beginning of the chunk)
 		*/
 		void	SetSyncIndex(uint32 InSyncIndex, int32 InSyncSize)
-		{
+		{			
 			SyncIndex = InSyncIndex;
 			SyncSize = InSyncSize;
 		}
@@ -348,6 +350,7 @@ public:
 
 		//stat associated with this allocation
 		TStatId Stat;
+		bool bTail;
 	};
 
 	/** Constructor, zero initializing all member variables */
@@ -643,6 +646,28 @@ public:
 
 protected:
 
+#if VALIDATE_SYNC_SIZE
+	struct FRelocationEntry
+	{
+		FRelocationEntry(const uint8* InOldBase, const uint8* InNewBase, uint64 InSize, uint64 InSyncIndex)
+		: OldBase(InOldBase)
+		, NewBase(InNewBase)
+		, Size(InSize)
+		, SyncIndex(InSyncIndex)
+		{
+			printf("Relocation: 0x%p to 0x%p, %i, %i\n", OldBase, NewBase, (int32)Size, (int32)SyncIndex);
+		}
+
+		const uint8* OldBase;
+		const uint8* NewBase;
+		uint64 Size;
+		uint64 SyncIndex;
+	};
+
+	void ValidateRelocations(uint8* UsedBaseAddress, uint64 Size);
+	TArray<FRelocationEntry> Relocations;
+#endif
+
 	/**
 	* Copy memory from one location to another. If it returns false, the defragmentation
 	* process will assume the memory is not relocatable and keep it in place.
@@ -837,6 +862,7 @@ protected:
 			int32 SecondSyncSize = BaseChunk->SyncSize - FirstSize;
 			NewFreeChunk->SetSyncIndex(BaseChunk->SyncIndex, SecondSyncSize);
 		}
+		BaseChunk->SetSyncIndex(BaseChunk->SyncIndex, FMath::Min((int64)FirstSize, BaseChunk->SyncSize));
 
 		// Resize base chunk.
 		BaseChunk->Size = FirstSize;
@@ -851,7 +877,7 @@ protected:
 	* @param bAsync			If true, allows allocating from relocating chunks and maint32ains the free-list sort order.
 	* @return					The memory chunk that was allocated (the original chunk could've been split).
 	*/
-	FMemoryChunk*	AllocateChunk(FMemoryChunk* FreeChunk, int32 AllocationSize, bool bAsync);
+	FMemoryChunk*	AllocateChunk(FMemoryChunk* FreeChunk, int32 AllocationSize, bool bAsync, bool bDoValidation = true);
 
 	/**
 	* Marks the specified chunk as 'free' and updates tracking variables.
@@ -867,7 +893,7 @@ protected:
 	* Maint32ains the free-list order if bSortedFreeList is true.
 	*
 	* @param Chunk				Chunk to mark as available.
-	* @param bSortedFreeList	If true, maint32ains the free-list sort order
+	* @param bSortedFreeList	If true, maintains the free-list sort order
 	*/
 	void			LinkFreeChunk(FMemoryChunk* Chunk, bool bSortedFreeList)
 	{
@@ -974,9 +1000,9 @@ protected:
 	FSettings		Settings;
 
 	/** Ever-increasing index to synchronize all relocations initiated by Tick(). */
-	uint32			CurrentSyncIndex;
+	uint64			CurrentSyncIndex;
 	/** Sync index that has been completed, so far.					*/
-	uint32			CompletedSyncIndex;
+	uint64			CompletedSyncIndex;
 
 	/** Number of async relocations that are currently in progress.	*/
 	int32				NumRelocationsInProgress;
@@ -1015,6 +1041,7 @@ protected:
 	friend FScopedGPUDefragLock;
 };
 
+//FScopedGPUDefragLock can't cover any scope that will add dcb commands or we might deadlock with a master reserve failure.
 class FScopedGPUDefragLock
 {
 public:
@@ -1161,7 +1188,7 @@ private:
 /**
 *	Returns the current size (in uint8s), or the final size if it has a reallocating request.
 */
-FORCEINLINE int32 FGPUDefragAllocator::FMemoryChunk::GetFinalSize() const
+FORCEINLINE int64 FGPUDefragAllocator::FMemoryChunk::GetFinalSize() const
 {
 	return Size;
 }
diff --git a/Engine/Source/Runtime/Renderer/Private/CompositionLighting/CompositionLighting.cpp b/Engine/Source/Runtime/Renderer/Private/CompositionLighting/CompositionLighting.cpp
--- a/Engine/Source/Runtime/Renderer/Private/CompositionLighting/CompositionLighting.cpp
+++ b/Engine/Source/Runtime/Renderer/Private/CompositionLighting/CompositionLighting.cpp
@@ -400,6 +400,7 @@ void FCompositionLighting::ProcessAfterLighting(FRHICommandListImmediate& RHICmd
 	GRenderTargetPool.VisualizeTexture.SetCheckPoint(RHICmdList, SceneContext.ReflectiveShadowMapNormal);
 	GRenderTargetPool.VisualizeTexture.SetCheckPoint(RHICmdList, SceneContext.ReflectiveShadowMapDepth);
 
+#if !(UE_BUILD_SHIPPING || UE_BUILD_TEST)
 	{
 		extern int32 GShaderModelDebug;
 		if(GShaderModelDebug)
@@ -414,6 +415,7 @@ void FCompositionLighting::ProcessAfterLighting(FRHICommandListImmediate& RHICmd
 			}
 		}
 	}
+#endif
 
 	{
 		FMemMark Mark(FMemStack::Get());
diff --git a/Engine/Source/Runtime/Renderer/Private/CompositionLighting/PostProcessAmbientOcclusion.cpp b/Engine/Source/Runtime/Renderer/Private/CompositionLighting/PostProcessAmbientOcclusion.cpp
--- a/Engine/Source/Runtime/Renderer/Private/CompositionLighting/PostProcessAmbientOcclusion.cpp
+++ b/Engine/Source/Runtime/Renderer/Private/CompositionLighting/PostProcessAmbientOcclusion.cpp
@@ -20,13 +20,14 @@ static TAutoConsoleVariable<int32> CVarAmbientOcclusionCompute(
 	TEXT("r.AmbientOcclusion.Compute"),
 	0,
 	TEXT("If SSAO should use ComputeShader (not available on all platforms) or PixelShader.\n")
-	TEXT("The Compute Shader version is WIP, not optimized, requires hardware support (not mobile/DX10/OpenGL3),\n")
+	TEXT("The [Async] Compute Shader version is WIP, not optimized, requires hardware support (not mobile/DX10/OpenGL3),\n")
 	TEXT("does not use normals which allows it to run right after EarlyZPass (better performance when used with AyncCompute)\n")
+	TEXT("AyncCompute is currently only functional on PS4.\n")
 	TEXT(" 0: PixelShader (default)\n")
 	TEXT(" 1: (WIP) Use ComputeShader if possible, otherwise fall back to '0'\n")
 	TEXT(" 2: (WIP) Use AsyncCompute if efficient, otherwise fall back to '1'\n")
 	TEXT(" 3: (WIP) Use AsyncCompute if possible, otherwise fall back to '1'"),
-	ECVF_RenderThreadSafe);
+	ECVF_Scalability | ECVF_RenderThreadSafe);
 
 static TAutoConsoleVariable<float> CVarAmbientOcclusionMaxQuality(
 	TEXT("r.AmbientOcclusionMaxQuality"),
diff --git a/Engine/Source/Runtime/Renderer/Private/CompositionLighting/PostProcessLpvIndirect.cpp b/Engine/Source/Runtime/Renderer/Private/CompositionLighting/PostProcessLpvIndirect.cpp
--- a/Engine/Source/Runtime/Renderer/Private/CompositionLighting/PostProcessLpvIndirect.cpp
+++ b/Engine/Source/Runtime/Renderer/Private/CompositionLighting/PostProcessLpvIndirect.cpp
@@ -353,7 +353,7 @@ void FRCPassPostProcessLpvIndirect::DoDirectionalOcclusionPass(FRenderingComposi
 
 	SCOPED_DRAW_EVENT(Context.RHICmdList, PostProcessLpvDirectionalOcclusion );
 	const FSceneRenderTargetItem& DestDirectionalOcclusionRenderTarget = SceneContext.DirectionalOcclusion->GetRenderTargetItem();
-	FViewInfo& View = Context.View;
+	const FViewInfo& View = Context.View;
 	FSceneViewState* ViewState = (FSceneViewState*)View.State;
 
 	if(!ViewState)
diff --git a/Engine/Source/Runtime/Renderer/Private/DeferredShadingRenderer.cpp b/Engine/Source/Runtime/Renderer/Private/DeferredShadingRenderer.cpp
--- a/Engine/Source/Runtime/Renderer/Private/DeferredShadingRenderer.cpp
+++ b/Engine/Source/Runtime/Renderer/Private/DeferredShadingRenderer.cpp
@@ -1571,7 +1571,7 @@ bool FDeferredShadingSceneRenderer::RenderPrePassView(FRHICommandList& RHICmdLis
 		if (EarlyZPassMode >= DDM_AllOccluders)
 		{
 			// Draw opaque occluders with masked materials
-			SCOPED_DRAW_EVENT(RHICmdList, Opaque);
+			SCOPED_DRAW_EVENT(RHICmdList, Masked);
 			bDirty |= Scene->MaskedDepthDrawList.DrawVisible(RHICmdList, View, View.StaticMeshOccluderMap, View.StaticMeshBatchVisibility);
 		}
 	}
@@ -1589,12 +1589,16 @@ bool FDeferredShadingSceneRenderer::RenderPrePassView(FRHICommandList& RHICmdLis
 
 		if (EarlyZPassMode >= DDM_AllOccluders)
 		{
-			SCOPED_DRAW_EVENT(RHICmdList, Opaque);
+			SCOPED_DRAW_EVENT(RHICmdList, Masked);
 			bDirty |= Scene->MaskedDepthDrawList.DrawVisibleInstancedStereo(RHICmdList, StereoView);
 		}
 	}
+	
+	{
+		SCOPED_DRAW_EVENT(RHICmdList, Dynamic);
+		bDirty |= RenderPrePassViewDynamic(RHICmdList, View);
+	}
 
-	bDirty |= RenderPrePassViewDynamic(RHICmdList, View);
 	return bDirty;
 }
 
diff --git a/Engine/Source/Runtime/Renderer/Private/DepthRendering.cpp b/Engine/Source/Runtime/Renderer/Private/DepthRendering.cpp
--- a/Engine/Source/Runtime/Renderer/Private/DepthRendering.cpp
+++ b/Engine/Source/Runtime/Renderer/Private/DepthRendering.cpp
@@ -428,7 +428,7 @@ void FDepthDrawingPolicyFactory::AddStaticMesh(FScene* Scene,FStaticMesh* Static
 	const EBlendMode BlendMode = Material->GetBlendMode();
 	const auto FeatureLevel = Scene->GetFeatureLevel();
 
-	if (!Material->WritesEveryPixel())
+	if (!Material->WritesEveryPixel() || Material->HasPixelDepthOffsetConnected())
 	{
 		// only draw if required
 		Scene->MaskedDepthDrawList.AddMesh(
diff --git a/Engine/Source/Runtime/Renderer/Private/PostProcess/PostProcessMobile.cpp b/Engine/Source/Runtime/Renderer/Private/PostProcess/PostProcessMobile.cpp
--- a/Engine/Source/Runtime/Renderer/Private/PostProcess/PostProcessMobile.cpp
+++ b/Engine/Source/Runtime/Renderer/Private/PostProcess/PostProcessMobile.cpp
@@ -2781,24 +2781,27 @@ void FRCPassPostProcessAaES2::Process(FRenderingCompositePassContext& Context)
 
 	Context.RHICmdList.CopyToResolveTarget(DestRenderTarget.TargetableTexture, DestRenderTarget.ShaderResourceTexture, false, FResolveParams());
 
-	auto& View = Context.View;
-	if (FSceneRenderer::ShouldCompositeEditorPrimitives(View))
+	if (FSceneRenderer::ShouldCompositeEditorPrimitives(Context.View))
 	{
+		FRHICommandListExecutor::GetImmediateCommandList().ImmediateFlush(EImmediateFlushType::WaitForOutstandingTasksOnly);
+		// because of the flush it's ok to remove the const, this is not ideal as the flush can cost performance
+		FViewInfo& NonConstView = (FViewInfo&)Context.View;
+
 		// Remove jitter (ensures editor prims are stable.)
-		View.ViewMatrices.RemoveTemporalJitter();
+		NonConstView.ViewMatrices.RemoveTemporalJitter();
 
 		// Compute the view projection matrix and its inverse.
-		View.ViewProjectionMatrix = View.ViewMatrices.ViewMatrix * View.ViewMatrices.ProjMatrix;
-		View.InvViewProjectionMatrix = View.ViewMatrices.GetInvProjMatrix() * View.InvViewMatrix;
+		NonConstView.ViewProjectionMatrix = NonConstView.ViewMatrices.ViewMatrix * NonConstView.ViewMatrices.ProjMatrix;
+		NonConstView.InvViewProjectionMatrix = NonConstView.ViewMatrices.GetInvProjMatrix() * NonConstView.InvViewMatrix;
 
 		/** The view transform, starting from world-space points translated by -ViewOrigin. */
-		FMatrix TranslatedViewMatrix = FTranslationMatrix(-View.ViewMatrices.PreViewTranslation) * View.ViewMatrices.ViewMatrix;
+		FMatrix TranslatedViewMatrix = FTranslationMatrix(-NonConstView.ViewMatrices.PreViewTranslation) * NonConstView.ViewMatrices.ViewMatrix;
 
 		// Compute a transform from view origin centered world-space to clip space.
-		View.ViewMatrices.TranslatedViewProjectionMatrix = TranslatedViewMatrix * View.ViewMatrices.ProjMatrix;
-		View.ViewMatrices.InvTranslatedViewProjectionMatrix = View.ViewMatrices.TranslatedViewProjectionMatrix.Inverse();
+		NonConstView.ViewMatrices.TranslatedViewProjectionMatrix = TranslatedViewMatrix * NonConstView.ViewMatrices.ProjMatrix;
+		NonConstView.ViewMatrices.InvTranslatedViewProjectionMatrix = NonConstView.ViewMatrices.TranslatedViewProjectionMatrix.Inverse();
 
-		View.InitRHIResources(nullptr);
+		NonConstView.InitRHIResources(nullptr);
 	}
 }
 
diff --git a/Engine/Source/Runtime/Renderer/Private/PostProcess/PostProcessTemporalAA.cpp b/Engine/Source/Runtime/Renderer/Private/PostProcess/PostProcessTemporalAA.cpp
--- a/Engine/Source/Runtime/Renderer/Private/PostProcess/PostProcessTemporalAA.cpp
+++ b/Engine/Source/Runtime/Renderer/Private/PostProcess/PostProcessTemporalAA.cpp
@@ -551,7 +551,7 @@ void FRCPassPostProcessTemporalAA::Process(FRenderingCompositePassContext& Conte
 	}
 	FSceneRenderTargets& SceneContext = FSceneRenderTargets::Get(Context.RHICmdList);
 
-	FViewInfo& View = Context.View;
+	const FViewInfo& View = Context.View;
 	FSceneViewState* ViewState = Context.ViewState;
 
 	FIntPoint TexSize = InputDesc->Extent;
@@ -739,22 +739,25 @@ void FRCPassPostProcessTemporalAA::Process(FRenderingCompositePassContext& Conte
 		ViewState->TemporalAAHistoryRT = PassOutputs[0].PooledRenderTarget;
 	}
 
-#if WITH_EDITOR
-	FRHICommandListExecutor::GetImmediateCommandList().ImmediateFlush(EImmediateFlushType::WaitForOutstandingTasksOnly);
+	if (FSceneRenderer::ShouldCompositeEditorPrimitives(Context.View))
+	{
+		FRHICommandListExecutor::GetImmediateCommandList().ImmediateFlush(EImmediateFlushType::WaitForOutstandingTasksOnly);
+		// because of the flush it's ok to remove the const, this is not ideal as the flush can cost performance
+		FViewInfo& NonConstView = (FViewInfo&)Context.View;
 
-	// Remove jitter
-	View.ViewMatrices.RemoveTemporalJitter();
+		// Remove jitter
+		NonConstView.ViewMatrices.RemoveTemporalJitter();
 
-	// Compute the view projection matrix and its inverse.
-	View.ViewProjectionMatrix = View.ViewMatrices.ViewMatrix * View.ViewMatrices.ProjMatrix;
-	View.InvViewProjectionMatrix = View.ViewMatrices.GetInvProjMatrix() * View.InvViewMatrix;
+		// Compute the view projection matrix and its inverse.
+		NonConstView.ViewProjectionMatrix = NonConstView.ViewMatrices.ViewMatrix * NonConstView.ViewMatrices.ProjMatrix;
+		NonConstView.InvViewProjectionMatrix = NonConstView.ViewMatrices.GetInvProjMatrix() * NonConstView.InvViewMatrix;
 
-	// Compute a transform from view origin centered world-space to clip space.
-	View.ViewMatrices.TranslatedViewProjectionMatrix = View.ViewMatrices.TranslatedViewMatrix * View.ViewMatrices.ProjMatrix;
-	View.ViewMatrices.InvTranslatedViewProjectionMatrix = View.ViewMatrices.GetInvProjMatrix() * View.ViewMatrices.TranslatedViewMatrix.Inverse();
+		// Compute a transform from view origin centered world-space to clip space.
+		NonConstView.ViewMatrices.TranslatedViewProjectionMatrix = NonConstView.ViewMatrices.TranslatedViewMatrix * NonConstView.ViewMatrices.ProjMatrix;
+		NonConstView.ViewMatrices.InvTranslatedViewProjectionMatrix = NonConstView.ViewMatrices.GetInvProjMatrix() * NonConstView.ViewMatrices.TranslatedViewMatrix.Inverse();
 
-	View.InitRHIResources(nullptr);
-#endif
+		NonConstView.InitRHIResources(nullptr);
+	}
 }
 
 FPooledRenderTargetDesc FRCPassPostProcessTemporalAA::ComputeOutputDesc(EPassOutputId InPassOutputId) const
diff --git a/Engine/Source/Runtime/Renderer/Private/PostProcess/SceneRenderTargets.h b/Engine/Source/Runtime/Renderer/Private/PostProcess/SceneRenderTargets.h
--- a/Engine/Source/Runtime/Renderer/Private/PostProcess/SceneRenderTargets.h
+++ b/Engine/Source/Runtime/Renderer/Private/PostProcess/SceneRenderTargets.h
@@ -322,7 +322,7 @@ public:
 		if (!SeparateTranslucencyDepthRT || SeparateTranslucencyDepthRT->GetDesc().Extent != Size)
 		{
 			// Create the SeparateTranslucency depth render target 
-			FPooledRenderTargetDesc Desc(FPooledRenderTargetDesc::Create2DDesc(Size, PF_DepthStencil, FClearValueBinding::None, TexCreate_None, TexCreate_DepthStencilTargetable, true));
+			FPooledRenderTargetDesc Desc(FPooledRenderTargetDesc::Create2DDesc(Size, PF_DepthStencil, FClearValueBinding::None, TexCreate_None, TexCreate_DepthStencilTargetable, false));
 			GRenderTargetPool.FindFreeElement(RHICmdList, Desc, SeparateTranslucencyDepthRT, TEXT("SeparateTranslucencyDepth"));
 		}
 		return SeparateTranslucencyDepthRT;
diff --git a/Engine/Source/Runtime/Renderer/Private/SceneRendering.cpp b/Engine/Source/Runtime/Renderer/Private/SceneRendering.cpp
--- a/Engine/Source/Runtime/Renderer/Private/SceneRendering.cpp
+++ b/Engine/Source/Runtime/Renderer/Private/SceneRendering.cpp
@@ -351,6 +351,7 @@ void FViewInfo::Init()
 	bDisableDistanceBasedFadeTransitions = false;	
 	ShadingModelMaskInView = 0;
 
+#if !(UE_BUILD_SHIPPING || UE_BUILD_TEST)
 	{
 		extern int32 GShaderModelDebug;
 		if(GShaderModelDebug)
@@ -358,6 +359,7 @@ void FViewInfo::Init()
 			UE_LOG(LogRenderer, Log, TEXT("r.ShaderModelDebug: FViewInfo::Init %x"), ShadingModelMaskInView);
 		}
 	}
+#endif
 
 	NumVisibleStaticMeshElements = 0;
 	PrecomputedVisibilityData = 0;
@@ -1224,7 +1226,7 @@ IPooledRenderTarget* FViewInfo::GetLastEyeAdaptationRT(FRHICommandList& RHICmdLi
 	return result;
 }
 
-void FViewInfo::SwapEyeAdaptationRTs()
+void FViewInfo::SwapEyeAdaptationRTs() const
 {
 	FSceneViewState* EffectiveViewState = GetEffectiveViewState();
 	if (EffectiveViewState)
@@ -1244,7 +1246,7 @@ bool FViewInfo::HasValidEyeAdaptation() const
 	return false;
 }
 
-void FViewInfo::SetValidEyeAdaptation()
+void FViewInfo::SetValidEyeAdaptation() const
 {
 	FSceneViewState* EffectiveViewState = GetEffectiveViewState();	
 
diff --git a/Engine/Source/Runtime/Renderer/Private/SceneRendering.h b/Engine/Source/Runtime/Renderer/Private/SceneRendering.h
--- a/Engine/Source/Runtime/Renderer/Private/SceneRendering.h
+++ b/Engine/Source/Runtime/Renderer/Private/SceneRendering.h
@@ -625,13 +625,13 @@ public:
 	IPooledRenderTarget* GetLastEyeAdaptationRT(FRHICommandList& RHICmdList) const;
 
 	/**Swap the order of the two eye adaptation targets in the double buffer system */
-	void SwapEyeAdaptationRTs();
+	void SwapEyeAdaptationRTs() const;
 
 	/** Tells if the eyeadaptation texture exists without attempting to allocate it. */
 	bool HasValidEyeAdaptation() const;
 
 	/** Informs sceneinfo that eyedaptation has queued commands to compute it at least once */
-	void SetValidEyeAdaptation();
+	void SetValidEyeAdaptation() const;
 
 	/** Create acceleration data structure and information to do forward lighting with dynamic branching. */
 	void CreateLightGrid();
diff --git a/Engine/Source/Runtime/Renderer/Private/ShadowRendering.cpp b/Engine/Source/Runtime/Renderer/Private/ShadowRendering.cpp
--- a/Engine/Source/Runtime/Renderer/Private/ShadowRendering.cpp
+++ b/Engine/Source/Runtime/Renderer/Private/ShadowRendering.cpp
@@ -232,7 +232,7 @@ public:
 		{
 			return (Material->IsSpecialEngineMaterial()
 						// Masked and WPO materials need their shaders but cannot be used with a position only stream.
-					|| ((!Material->WritesEveryPixel() || Material->MaterialMayModifyMeshPosition()) && !bUsePositionOnlyStream))
+					|| ((!Material->WritesEveryPixel(true) || Material->MaterialMayModifyMeshPosition()) && !bUsePositionOnlyStream))
 					// Only compile one pass point light shaders for feature levels >= SM4
 					&& (ShaderMode != VertexShadowDepth_OnePassPointLight || IsFeatureLevelSupported(Platform, ERHIFeatureLevel::SM4))
 					// Only compile position-only shaders for vertex factories that support it.
@@ -588,7 +588,7 @@ public:
 		{
 			return (Material->IsSpecialEngineMaterial()
 				// Only compile for masked or lit translucent materials
-				|| !Material->WritesEveryPixel())
+				|| !Material->WritesEveryPixel(true))
 				&& ShaderMode == PixelShadowDepth_NonPerspectiveCorrect
 				// Don't render ShadowDepth for translucent unlit materials
 				&& (!IsTranslucentBlendMode(Material->GetBlendMode()) && Material->GetShadingModel() != MSM_Unlit)
@@ -609,7 +609,7 @@ public:
 			//Note: This logic needs to stay in sync with OverrideWithDefaultMaterialForShadowDepth!
 			return (Material->IsSpecialEngineMaterial()
 					// Only compile for masked or lit translucent materials
-					|| !Material->WritesEveryPixel()
+					|| !Material->WritesEveryPixel(true)
 					|| (Material->MaterialMayModifyMeshPosition() && Material->IsUsedWithInstancedStaticMeshes())
 					// Perspective correct rendering needs a pixel shader and WPO materials can't be overridden with default material.
 					|| (ShaderMode == PixelShadowDepth_PerspectiveCorrect && Material->MaterialMayModifyMeshPosition()))
@@ -758,9 +758,9 @@ void OverrideWithDefaultMaterialForShadowDepth(
 	ERHIFeatureLevel::Type InFeatureLevel)
 {
 	// Override with the default material when possible.
-	if (InOutMaterialResource->WritesEveryPixel() &&						// Don't override masked materials.
+	if (InOutMaterialResource->WritesEveryPixel(true) &&						// Don't override masked materials.
 		!InOutMaterialResource->MaterialModifiesMeshPosition_RenderThread() &&	// Don't override materials using world position offset.
-		!bReflectiveShadowmap)										// Don't override when rendering reflective shadow maps.
+		!bReflectiveShadowmap)													// Don't override when rendering reflective shadow maps.
 	{
 		const FMaterialRenderProxy* DefaultProxy = UMaterial::GetDefaultMaterial(MD_Surface)->GetRenderProxy(false);
 		const FMaterial* DefaultMaterialResource = DefaultProxy->GetMaterial(InFeatureLevel);
@@ -837,7 +837,7 @@ FShadowDepthDrawingPolicy<bRenderingReflectiveShadowMaps>::FShadowDepthDrawingPo
 
 	bUsePositionOnlyVS = !bRenderingReflectiveShadowMaps
 		&& VertexFactory->SupportsPositionOnlyStream()
-		&& MaterialResource->WritesEveryPixel()
+		&& MaterialResource->WritesEveryPixel(true)
 		&& !MaterialResource->MaterialModifiesMeshPosition_RenderThread();
 
 	// Vertex related shaders
@@ -915,7 +915,7 @@ FShadowDepthDrawingPolicy<bRenderingReflectiveShadowMaps>::FShadowDepthDrawingPo
 	}
 
 	// Pixel shaders
-	if (MaterialResource->WritesEveryPixel() && !bUsePerspectiveCorrectShadowDepths && !bRenderingReflectiveShadowMaps && VertexFactory->SupportsNullPixelShader())
+	if (MaterialResource->WritesEveryPixel(true) && !bUsePerspectiveCorrectShadowDepths && !bRenderingReflectiveShadowMaps && VertexFactory->SupportsNullPixelShader())
 	{
 		// No pixel shader necessary.
 		PixelShader = NULL;
diff --git a/Engine/Source/Runtime/Renderer/Private/ShadowSetup.cpp b/Engine/Source/Runtime/Renderer/Private/ShadowSetup.cpp
--- a/Engine/Source/Runtime/Renderer/Private/ShadowSetup.cpp
+++ b/Engine/Source/Runtime/Renderer/Private/ShadowSetup.cpp
@@ -1369,14 +1369,28 @@ void FSceneRenderer::CreatePerObjectProjectedShadow(
 	// Compute the composite bounds of this group of shadow primitives.
 	FBoxSphereBounds OriginalBounds = ShadowGroupPrimitives[0]->Proxy->GetBounds();
 
+	if (!ensureMsgf(OriginalBounds.ContainsNaN() == false, TEXT("OriginalBound contains NaN : %s"), *OriginalBounds.ToString()))
+	{
+		// fix up OriginalBounds. This is going to cause flickers
+		OriginalBounds = FBoxSphereBounds(FVector::ZeroVector, FVector(1.f), 1.f);
+	}
+
 	for (int32 ChildIndex = 1; ChildIndex < ShadowGroupPrimitives.Num(); ChildIndex++)
 	{
 		const FPrimitiveSceneInfo* ShadowChild = ShadowGroupPrimitives[ChildIndex];
 		if (ShadowChild->Proxy->CastsDynamicShadow())
 		{
-			OriginalBounds = OriginalBounds + ShadowChild->Proxy->GetBounds();
+			FBoxSphereBounds ChildBound = ShadowChild->Proxy->GetBounds();
+			OriginalBounds = OriginalBounds + ChildBound;
+
+			if (!ensureMsgf(OriginalBounds.ContainsNaN() == false, TEXT("Child %s contains NaN : %s"), *ShadowChild->Proxy->GetOwnerName().ToString(), *ChildBound.ToString()))
+			{
+				// fix up OriginalBounds. This is going to cause flickers
+				OriginalBounds = FBoxSphereBounds(FVector::ZeroVector, FVector(1.f), 1.f);
+			}
 		}
 	}
+
 	FSceneRenderTargets& SceneContext = FSceneRenderTargets::Get(RHICmdList);
 	
 	// Shadowing constants.
@@ -1407,7 +1421,7 @@ void FSceneRenderer::CreatePerObjectProjectedShadow(
 			OriginalBounds.SphereRadius /
 			FMath::Max(ShadowViewDistFromBounds, 1.0f);
 		// Early catch for invalid CalculateShadowFadeAlpha()
-		checkf(ScreenRadius >= 0.0f, TEXT("View.ShadowViewMatrices.ScreenScale %f, OriginalBounds.SphereRadius %f, ShadowViewDistFromBounds %f"), View.ShadowViewMatrices.ScreenScale, OriginalBounds.SphereRadius, ShadowViewDistFromBounds);
+		ensureMsgf(ScreenRadius >= 0.0f, TEXT("View.ShadowViewMatrices.ScreenScale %f, OriginalBounds.SphereRadius %f, ShadowViewDistFromBounds %f"), View.ShadowViewMatrices.ScreenScale, OriginalBounds.SphereRadius, ShadowViewDistFromBounds);
 
 		const float ScreenPercent = FMath::Max(
 			1.0f / 2.0f * View.ShadowViewMatrices.ProjectionScale.X,
diff --git a/Engine/Source/Runtime/Slate/Private/Framework/Text/TextLayout.cpp b/Engine/Source/Runtime/Slate/Private/Framework/Text/TextLayout.cpp
--- a/Engine/Source/Runtime/Slate/Private/Framework/Text/TextLayout.cpp
+++ b/Engine/Source/Runtime/Slate/Private/Framework/Text/TextLayout.cpp
@@ -50,6 +50,7 @@ FTextLayout::FBreakCandidate FTextLayout::CreateBreakCandidate( int32& OutRunInd
 
 	// We need to consider the Runs when detecting and measuring the text lengths of Lines because
 	// the font style used makes a difference.
+	const int32 FirstRunIndexChecked = OutRunIndex;
 	for (; OutRunIndex < Line.Runs.Num(); OutRunIndex++)
 	{
 		FRunModel& Run = Line.Runs[ OutRunIndex ];
@@ -124,7 +125,18 @@ FTextLayout::FBreakCandidate FTextLayout::CreateBreakCandidate( int32& OutRunInd
 			break;
 		}
 	}
-	check( SuccessfullyMeasuredSlice == true );
+	if (!SuccessfullyMeasuredSlice)
+	{
+		FString RunDebugData;
+		for (int32 RunIndex = 0; RunIndex < Line.Runs.Num(); ++RunIndex)
+		{
+			const FRunModel& Run = Line.Runs[RunIndex];
+			const FTextRange RunRange = Run.GetTextRange();
+			RunDebugData.Append(FString::Printf(TEXT("\t\t[%d] - Range: {%d, %d}\n"), RunIndex, RunRange.BeginIndex, RunRange.EndIndex));
+		}
+
+		checkf( SuccessfullyMeasuredSlice, TEXT("Failed to measure a slice of text!\n\tStart Index: %d\n\tEnd Index: %d\n\tStart Run Index: %d\n\tLine Runs:\n%s"), PreviousBreak, CurrentBreak, FirstRunIndexChecked, *RunDebugData );
+	}
 
 	BreakSize.Y = BreakSizeWithoutTrailingWhitespace.Y = MaxAboveBaseline + MaxBelowBaseline;
 
diff --git a/Engine/Source/Runtime/SlateRHIRenderer/Private/SlateRHIRenderer.cpp b/Engine/Source/Runtime/SlateRHIRenderer/Private/SlateRHIRenderer.cpp
--- a/Engine/Source/Runtime/SlateRHIRenderer/Private/SlateRHIRenderer.cpp
+++ b/Engine/Source/Runtime/SlateRHIRenderer/Private/SlateRHIRenderer.cpp
@@ -1236,18 +1236,10 @@ void FSlateRHIRenderer::RequestResize( const TSharedPtr<SWindow>& Window, uint32
 
 	FViewportInfo* ViewInfo = WindowToViewportInfo.FindRef( Window.Get() );
 
-	if( ViewInfo )
+	if (ViewInfo)
 	{
-		if (ViewInfo->DesiredWidth != NewWidth || ViewInfo->DesiredHeight != NewHeight)
-		{
-			ViewInfo->DesiredWidth = NewWidth;
-			ViewInfo->DesiredHeight = NewHeight;
-
-			if (Window->GetWindowMode() == EWindowMode::Windowed)
-			{
-				FSystemResolution::RequestResolutionChange(NewWidth, NewHeight, Window->GetWindowMode());
-			}
-		}
+		ViewInfo->DesiredWidth = NewWidth;
+		ViewInfo->DesiredHeight = NewHeight;
 	}
 }
 
diff --git a/Engine/Source/Runtime/Windows/D3D11RHI/Private/D3D11Viewport.cpp b/Engine/Source/Runtime/Windows/D3D11RHI/Private/D3D11Viewport.cpp
--- a/Engine/Source/Runtime/Windows/D3D11RHI/Private/D3D11Viewport.cpp
+++ b/Engine/Source/Runtime/Windows/D3D11RHI/Private/D3D11Viewport.cpp
@@ -183,7 +183,9 @@ DXGI_MODE_DESC FD3D11Viewport::SetupDXGI_MODE_DESC() const
 void FD3D11Viewport::Resize(uint32 InSizeX,uint32 InSizeY,bool bInIsFullscreen)
 {
 	// Unbind any dangling references to resources
+	D3DRHI->RHISetRenderTargets(0, nullptr, nullptr, 0, nullptr);
 	D3DRHI->ClearState();
+	D3DRHI->GetDeviceContext()->Flush(); // Potential perf hit
 
 	if (IsValidRef(CustomPresent))
 	{
diff --git a/Engine/Source/Runtime/Windows/D3D11RHI/Private/Windows/WindowsD3D11Device.cpp b/Engine/Source/Runtime/Windows/D3D11RHI/Private/Windows/WindowsD3D11Device.cpp
--- a/Engine/Source/Runtime/Windows/D3D11RHI/Private/Windows/WindowsD3D11Device.cpp
+++ b/Engine/Source/Runtime/Windows/D3D11RHI/Private/Windows/WindowsD3D11Device.cpp
@@ -22,7 +22,7 @@ static TAutoConsoleVariable<int32> CVarGraphicsAdapter(
 	TEXT("User request to pick a specific graphics adapter (e.g. when using a integrated graphics card with a descrete one)\n")
 	TEXT("At the moment this only works on Direct3D 11.\n")
 	TEXT(" -2: Take the first one that fulfills the criteria\n")
-	TEXT(" -1: Favour non integrated because there are usually faster\n")
+	TEXT(" -1: Favour non integrated because there are usually faster (default)\n")
 	TEXT("  0: Adpater #0\n")
 	TEXT("  1: Adpater #1, ..."),
 	ECVF_RenderThreadSafe);
@@ -501,7 +501,10 @@ void FD3D11DynamicRHI::InitD3DDevice()
 					GRHIAdapterName = AdapterDesc.Description;
 					GRHIVendorId = AdapterDesc.VendorId;
 					GRHIDeviceId = AdapterDesc.DeviceId;
-					
+
+					UE_LOG(LogD3D11RHI, Log, TEXT("    GPU DeviceId: 0x%x (for the marketing name, search the web for \"GPU Device Id\")"), 
+						AdapterDesc.DeviceId);
+
 					// get driver version (todo: share with other RHIs)
 					{
 						FPlatformMisc::GetGPUDriverInfo(GRHIAdapterName, GRHIAdapterInternalDriverVersion, GRHIAdapterUserDriverVersion, GRHIAdapterDriverDate);
@@ -791,22 +794,44 @@ bool FD3D11DynamicRHI::RHIGetAvailableResolutions(FScreenResolutionArray& Resolu
 
 		// TODO: GetDisplayModeList is a terribly SLOW call.  It can take up to a second per invocation.
 		//  We might want to work around some DXGI badness here.
-		DXGI_FORMAT Format = DXGI_FORMAT_R8G8B8A8_UNORM;
-		uint32 NumModes = 0;
-		HResult = Output->GetDisplayModeList(Format, 0, &NumModes, NULL);
-		if(HResult == DXGI_ERROR_NOT_FOUND)
-		{
-			continue;
-		}
-		else if(HResult == DXGI_ERROR_NOT_CURRENTLY_AVAILABLE)
+		const DXGI_FORMAT DisplayFormats[] = {DXGI_FORMAT_R8G8B8A8_UNORM, DXGI_FORMAT_B8G8R8A8_UNORM};
+		DXGI_FORMAT Format = DisplayFormats[0];
+		uint32 NumModes = 0;	
+
+		for (DXGI_FORMAT CurrentFormat : DisplayFormats)
 		{
-			UE_LOG(LogD3D11RHI, Warning,
-				TEXT("RHIGetAvailableResolutions does not return results when running under remote desktop.")
-				);
-			return false;
+			HResult = Output->GetDisplayModeList(CurrentFormat, 0, &NumModes, NULL);
+
+			if(FAILED(HResult))
+			{
+				if (HResult == DXGI_ERROR_NOT_FOUND)
+				{
+					continue;
+				}
+				else if (HResult == DXGI_ERROR_MORE_DATA)
+				{
+					UE_LOG(LogD3D11RHI, Warning, TEXT("RHIGetAvailableResolutions failed trying to return too much data."));
+					continue;
+				}
+				else if (HResult == DXGI_ERROR_NOT_CURRENTLY_AVAILABLE)
+				{
+					UE_LOG(LogD3D11RHI, Warning, TEXT("RHIGetAvailableResolutions does not return results when running under remote desktop."));
+					return false;
+				}
+				else
+				{
+					UE_LOG(LogD3D11RHI, Warning, TEXT("RHIGetAvailableResolutions failed with unknown error (0x%x)."), HResult);
+					return false;
+				}
+			}
+			else if(NumModes)
+			{
+				Format = CurrentFormat;
+				break;
+			}
 		}
 
-		checkf(NumModes > 0, TEXT("No display modes found for the standard format DXGI_FORMAT_R8G8B8A8_UNORM!"));
+		checkf(NumModes > 0, TEXT("No display modes found for DXGI_FORMAT_R8G8B8A8_UNORM or DXGI_FORMAT_B8G8R8A8_UNORM formats!"));
 
 		DXGI_MODE_DESC* ModeList = new DXGI_MODE_DESC[ NumModes ];
 		VERIFYD3D11RESULT(Output->GetDisplayModeList(Format, 0, &NumModes, ModeList));
diff --git a/Engine/Source/Runtime/Windows/XAudio2/Private/XAudio2Buffer.cpp b/Engine/Source/Runtime/Windows/XAudio2/Private/XAudio2Buffer.cpp
--- a/Engine/Source/Runtime/Windows/XAudio2/Private/XAudio2Buffer.cpp
+++ b/Engine/Source/Runtime/Windows/XAudio2/Private/XAudio2Buffer.cpp
@@ -349,6 +349,9 @@ void FXAudio2SoundBuffer::Seek( const float SeekTime )
  */
 FXAudio2SoundBuffer* FXAudio2SoundBuffer::CreateQueuedBuffer( FXAudio2Device* XAudio2Device, USoundWave* Wave )
 {
+	check(XAudio2Device);
+	check(Wave);
+
 	// Check to see if thread has finished decompressing on the other thread
 	if( Wave->AudioDecompressor != nullptr )
 	{
diff --git a/Engine/Source/Runtime/Windows/XAudio2/Private/XAudio2Device.cpp b/Engine/Source/Runtime/Windows/XAudio2/Private/XAudio2Device.cpp
--- a/Engine/Source/Runtime/Windows/XAudio2/Private/XAudio2Device.cpp
+++ b/Engine/Source/Runtime/Windows/XAudio2/Private/XAudio2Device.cpp
@@ -270,6 +270,8 @@ bool FXAudio2Device::HasCompressedAudioInfoClass(USoundWave* SoundWave)
 
 class ICompressedAudioInfo* FXAudio2Device::CreateCompressedAudioInfo(USoundWave* SoundWave)
 {
+	check(SoundWave);
+
 	if (SoundWave->IsStreaming())
 	{
 		return new FOpusAudioInfo();
@@ -278,8 +280,13 @@ class ICompressedAudioInfo* FXAudio2Device::CreateCompressedAudioInfo(USoundWave
 #if WITH_OGGVORBIS
 	if (SoundWave->CompressionName.IsNone() || SoundWave->CompressionName == TEXT("OGG"))
 	{
-
-		return new FVorbisAudioInfo();
+		ICompressedAudioInfo* CompressedInfo = new FVorbisAudioInfo();
+		if (!CompressedInfo)
+		{
+			UE_LOG(LogAudio, Error, TEXT("Failed to create new FVorbisAudioInfo for SoundWave %s: out of memory."), *SoundWave->GetName());
+			return NULL;
+		}
+		return CompressedInfo;
 	}
 	else
 	{
@@ -296,40 +303,9 @@ class ICompressedAudioInfo* FXAudio2Device::CreateCompressedAudioInfo(USoundWave
  */
 bool FXAudio2Device::ValidateAPICall( const TCHAR* Function, uint32 ErrorCode )
 {
-	if( ErrorCode != S_OK )
-	{
-		switch( ErrorCode )
-		{
-		case XAUDIO2_E_INVALID_CALL:
-			UE_LOG(LogAudio, Warning, TEXT( "%s error: Invalid Call" ), Function );
-			break;
-
-		case XAUDIO2_E_XMA_DECODER_ERROR:
-			UE_LOG(LogAudio, Warning, TEXT( "%s error: XMA Decoder Error" ), Function );
-			break;
-
-		case XAUDIO2_E_XAPO_CREATION_FAILED:
-			UE_LOG(LogAudio, Warning, TEXT( "%s error: XAPO Creation Failed" ), Function );
-			break;
-
-		case XAUDIO2_E_DEVICE_INVALIDATED:
-			UE_LOG(LogAudio, Warning, TEXT( "%s error: Device Invalidated" ), Function );
-			break;
-
-		default:
-			UE_LOG(LogAudio, Warning, TEXT( "%s error: Unhandled error code %d" ), Function, ErrorCode );
-			break;
-		};
-
-		return( false );
-	}
-
-	return( true );
+	return DeviceProperties->Validate(Function, ErrorCode);
 }
 
-
-
-
 /** 
  * Derives the output matrix to use based on the channel mask and the number of channels
  */
diff --git a/Engine/Source/Runtime/Windows/XAudio2/Private/XAudio2Support.h b/Engine/Source/Runtime/Windows/XAudio2/Private/XAudio2Support.h
--- a/Engine/Source/Runtime/Windows/XAudio2/Private/XAudio2Support.h
+++ b/Engine/Source/Runtime/Windows/XAudio2/Private/XAudio2Support.h
@@ -588,22 +588,31 @@ struct FXAudioDeviceProperties
 	/** The array of voice pools. Each pool is according to the sound format (and max effect chain channels) */
 	TArray<FSourceVoicePoolEntry*> VoicePool;
 
+	/** Number of non-free active voices */
+	int32 NumActiveVoices;
+
 	FXAudioDeviceProperties()
 		: XAudio2(nullptr)
 		, MasteringVoice(nullptr)
 		, XAudio2Dll(nullptr)
+		, NumActiveVoices(0)
 	{
 	}
 	
 	~FXAudioDeviceProperties()
 	{
+		// Make sure we've free'd all of our active voices at this point!
+		check(NumActiveVoices == 0);
+
 		// Destroy all the xaudio2 voices allocated in our pools
 		for (int32 i = 0; i < VoicePool.Num(); ++i)
 		{
 			for (int32 j = 0; j < VoicePool[i]->FreeVoices.Num(); ++j)
 			{
-				IXAudio2SourceVoice* Voice = VoicePool[i]->FreeVoices[j];
-				Voice->DestroyVoice();
+				IXAudio2SourceVoice** Voice = &VoicePool[i]->FreeVoices[j];
+				check(*Voice != nullptr);
+				(*Voice)->DestroyVoice();
+				*Voice = nullptr;
 			}
 		}
 
@@ -618,14 +627,15 @@ struct FXAudioDeviceProperties
 		if (MasteringVoice)
 		{
 			MasteringVoice->DestroyVoice();
-			MasteringVoice = NULL;
+			MasteringVoice = nullptr;
 		}
 
 		if (XAudio2)
 		{
 			// Force the hardware to release all references
-			XAudio2->Release();
-			XAudio2 = NULL;
+			Validate(TEXT("~FXAudioDeviceProperties: XAudio2->Release()"),
+					 XAudio2->Release());
+			XAudio2 = nullptr;
 		}
 
 #if PLATFORM_WINDOWS && PLATFORM_64BITS
@@ -639,6 +649,40 @@ struct FXAudioDeviceProperties
 #endif
 	}
 
+	bool Validate(const TCHAR* Function, uint32 ErrorCode) const
+	{
+		if (ErrorCode != S_OK)
+		{
+			switch (ErrorCode)
+			{
+				case XAUDIO2_E_INVALID_CALL:
+				UE_LOG(LogAudio, Error, TEXT("%s error: Invalid Call"), Function);
+				break;
+
+				case XAUDIO2_E_XMA_DECODER_ERROR:
+				UE_LOG(LogAudio, Error, TEXT("%s error: XMA Decoder Error"), Function);
+				break;
+
+				case XAUDIO2_E_XAPO_CREATION_FAILED:
+				UE_LOG(LogAudio, Error, TEXT("%s error: XAPO Creation Failed"), Function);
+				break;
+
+				case XAUDIO2_E_DEVICE_INVALIDATED:
+				UE_LOG(LogAudio, Error, TEXT("%s error: Device Invalidated"), Function);
+				break;
+
+				default:
+				UE_LOG(LogAudio, Error, TEXT("%s error: Unhandled error code %d"), Function, ErrorCode);
+				break;
+			};
+
+			return false;
+		}
+
+		return true;
+	}
+
+
 	void GetAudioDeviceList(TArray<FString>& OutAudioDeviceList) const
 	{
 		OutAudioDeviceList.Reset();
@@ -647,12 +691,18 @@ struct FXAudioDeviceProperties
 		uint32 NumDevices = 0;
 		if (XAudio2)
 		{
-			XAudio2->GetDeviceCount(&NumDevices);
-			for (uint32 i = 0; i < NumDevices; ++i)
+			if (Validate( TEXT("GetAudioDeviceList: XAudio2->GetDeviceCount"),
+				XAudio2->GetDeviceCount(&NumDevices)))
 			{
-				XAUDIO2_DEVICE_DETAILS Details;
-				XAudio2->GetDeviceDetails(i, &Details);
-				OutAudioDeviceList.Add(FString(Details.DisplayName));
+				for (uint32 i = 0; i < NumDevices; ++i)
+				{
+					XAUDIO2_DEVICE_DETAILS Details;
+					if (Validate(TEXT("GetAudioDeviceList: XAudio2->GetDeviceDetails"),
+						XAudio2->GetDeviceDetails(i, &Details)))
+					{
+						OutAudioDeviceList.Add(FString(Details.DisplayName));
+					}
+				}
 			}
 		}
 #endif
@@ -661,6 +711,8 @@ struct FXAudioDeviceProperties
 	/** Returns either a new IXAudio2SourceVoice or a recycled IXAudio2SourceVoice according to the sound format and max channel count in the voice's effect chain*/
 	void GetFreeSourceVoice(IXAudio2SourceVoice** Voice, const FPCMBufferInfo& BufferInfo, const XAUDIO2_EFFECT_CHAIN* EffectChain = nullptr, int32 MaxEffectChainChannels = 0)
 	{
+		bool bSuccess = false;
+
 		// First find the pool for the given format
 		FSourceVoicePoolEntry* VoicePoolEntry = nullptr;
 		for (int32 i = 0; i < VoicePool.Num(); ++i)
@@ -668,6 +720,8 @@ struct FXAudioDeviceProperties
 			if (VoicePool[i]->Format == BufferInfo.PCMFormat && VoicePool[i]->MaxEffectChainChannels == MaxEffectChainChannels)
 			{
 				VoicePoolEntry = VoicePool[i];
+				check(VoicePoolEntry);
+				bSuccess = true;
 				break;
 			}
 		}
@@ -677,32 +731,50 @@ struct FXAudioDeviceProperties
 		if (VoicePoolEntry && VoicePoolEntry->FreeVoices.Num() > 0)
 		{
 			*Voice = VoicePoolEntry->FreeVoices.Pop(false);
-			(*Voice)->SetEffectChain(EffectChain);
+			check(*Voice);
+
+			bSuccess = Validate(TEXT("GetFreeSourceVoice, XAudio2->CreateSourceVoice"),
+								(*Voice)->SetEffectChain(EffectChain));
 		}
 		else
 		{
 			// Create a brand new source voice with this format.
-			XAudio2->CreateSourceVoice(Voice, &BufferInfo.PCMFormat, XAUDIO2_VOICE_USEFILTER, MAX_PITCH, &SourceCallback, nullptr, EffectChain);
+			check(XAudio2 != nullptr);
+			bSuccess = Validate(TEXT("GetFreeSourceVoice, XAudio2->CreateSourceVoice"),
+								XAudio2->CreateSourceVoice(Voice, &BufferInfo.PCMFormat, XAUDIO2_VOICE_USEFILTER, MAX_PITCH, &SourceCallback, nullptr, EffectChain));
+		}
+
+		if (bSuccess)
+		{
+			// Track the number of source voices out in the world
+			++NumActiveVoices;
+		}
+		else
+		{
+			// If something failed, make sure we null the voice ptr output
+			*Voice = nullptr;
 		}
 	}
 
 	/** Releases the voice into a pool of free voices according to the voice format and the max effect chain channels */
 	void ReleaseSourceVoice(IXAudio2SourceVoice* Voice, const FPCMBufferInfo& BufferInfo, const int32 MaxEffectChainChannels)
 	{
+		check(Voice != nullptr);
+
 		// Make sure the voice is stopped
-		Voice->Stop();
+		Validate(TEXT("ReleaseSourceVoice, Voice->Stop()"), Voice->Stop());
 
 		// And make sure there's no audio remaining the voice so when it's re-used it's fresh.
-		Voice->FlushSourceBuffers();
+		Validate(TEXT("ReleaseSourceVoice, Voice->FlushSourceBuffers()"), Voice->FlushSourceBuffers());
 
 #if XAUDIO2_SUPPORTS_SENDLIST
 		// Clear out the send effects (OutputVoices). When the voice gets reused, the old internal state might be invalid 
 		// when the new send effects are applied to the voice.
-		Voice->SetOutputVoices(nullptr);
+		Validate(TEXT("ReleaseSourceVoice, Voice->SetOutputVoices(nullptr)"), Voice->SetOutputVoices(nullptr));
 #endif
 
 		// Release the effect chain
-		Voice->SetEffectChain(nullptr);
+		Validate(TEXT("ReleaseSourceVoice, Voice->SetEffectChain(nullptr);"), Voice->SetEffectChain(nullptr));
 
 		// See if there is an existing pool for this source voice
 		FSourceVoicePoolEntry* VoicePoolEntry = nullptr;
@@ -724,11 +796,21 @@ struct FXAudioDeviceProperties
 		{
 			// Otherwise We need to make a new voice pool entry with this format and max effect chain channels
 			VoicePoolEntry = new FSourceVoicePoolEntry();
-			VoicePoolEntry->Format = BufferInfo.PCMFormat;
-			VoicePoolEntry->FreeVoices.Add(Voice);
-			VoicePoolEntry->MaxEffectChainChannels = MaxEffectChainChannels;
-			VoicePool.Add(VoicePoolEntry);
+			if (VoicePoolEntry)
+			{
+				VoicePoolEntry->Format = BufferInfo.PCMFormat;
+				VoicePoolEntry->FreeVoices.Add(Voice);
+				VoicePoolEntry->MaxEffectChainChannels = MaxEffectChainChannels;
+				VoicePool.Add(VoicePoolEntry);
+			}
+			else
+			{
+				// If we failed to create a new voice pool entry, then destroy the voice
+				Voice->DestroyVoice();
+			}
 		}
+
+		--NumActiveVoices;
 	}
 };
 
