Copying //UE4/Orion-Staging to //UE4/Main (//Orion/Dev-General @ 2855324)

diff --git a/Engine/Source/Runtime/RHI/Private/GPUDefragAllocator.cpp b/Engine/Source/Runtime/RHI/Private/GPUDefragAllocator.cpp
--- a/Engine/Source/Runtime/RHI/Private/GPUDefragAllocator.cpp
+++ b/Engine/Source/Runtime/RHI/Private/GPUDefragAllocator.cpp
@@ -37,6 +37,16 @@ DECLARE_DWORD_ACCUMULATOR_STAT(TEXT("Total Async Cancels"), STAT_TexturePool_Tot
 #define DEFRAG_CHUNK_COUNTER_MIN		(20)
 #define DEFRAG_CHUNK_COUNTER_MAX		(80)
 
+#if VALIDATE_SYNC_SIZE
+static int32 GGPUDefragValidateSyncSize = 0;
+static FAutoConsoleVariableRef CVarGPUDefragValidateSyncSize(
+	TEXT("r.GPUDefrag.ValidateSyncSize"),
+	GGPUDefragValidateSyncSize,
+	TEXT("Enables validation and debug printout to detect unsafe allocations from the GPU defrag allocator.\n"),
+	ECVF_Default
+	);
+#endif
+
 
 /*-----------------------------------------------------------------------------
 FBestFitAllocator::FMemoryChunk implementation.
@@ -46,8 +56,9 @@ FBestFitAllocator::FMemoryChunk implementation.
 * Inserts this chunk at the head of the free chunk list.
 * If bMaintainSortOrder is true, insert-sort this chunk into the free chunk list.
 */
-void FGPUDefragAllocator::FMemoryChunk::LinkFree(bool bMaintainSortOrder, FMemoryChunk* FirstFreeChunkToSearch)
+void FGPUDefragAllocator::FMemoryChunk::LinkFree(FMemoryChunk* FirstFreeChunkToSearch)
 {
+	const bool bMaintainSortOrder = true;
 	check(!bIsAvailable);
 	bIsAvailable = true;
 	DefragCounter = 0;
@@ -118,7 +129,7 @@ void* FGPUDefragAllocator::Allocate(int64 AllocationSize, int32 Alignment, TStat
 	FScopeLock Lock(&SynchronizationObject);
 	check(FirstChunk);
 	check(Alignment <= AllocationAlignment);
-
+	const int64 OrigSize = AllocationSize;
 	// Make sure everything is appropriately aligned.
 	AllocationSize = Align(AllocationSize, AllocationAlignment);
 
@@ -189,6 +200,9 @@ void* FGPUDefragAllocator::Allocate(int64 AllocationSize, int32 Alignment, TStat
 #endif
 
 	FMemoryChunk* AllocatedChunk = AllocateChunk(BestChunk, AllocationSize, false);
+
+	AllocatedChunk->OrigSize = OrigSize;
+	FPlatformAtomics::InterlockedAdd(&PaddingWasteSize, AllocatedChunk->Size - OrigSize);
 	//todo: fix stats
 	//ensureMsgf(AllocatedChunk->Stat.IsNone(), TEXT("FreeChunk already has a stat."));
 	AllocatedChunk->Stat = InStat;
@@ -206,14 +220,17 @@ void* FGPUDefragAllocator::Allocate(int64 AllocationSize, int32 Alignment, TStat
 * @param bAsync			If true, allows allocating from relocating chunks and maintains the free-list sort order.
 * @return					The memory chunk that was allocated (the original chunk could've been split).
 */
-FGPUDefragAllocator::FMemoryChunk* FGPUDefragAllocator::AllocateChunk(FMemoryChunk* FreeChunk, int32 AllocationSize, bool bAsync, bool bDoValidation)
+FGPUDefragAllocator::FMemoryChunk* FGPUDefragAllocator::AllocateChunk(FMemoryChunk* FreeChunk, int64 AllocationSize, bool bAsync, bool bDoValidation)
 {
 	check(FreeChunk);
 	check(FreeChunk->GetAvailableSize() >= AllocationSize);
 	check(!FreeChunk->IsLocked());
 
 #if VALIDATE_SYNC_SIZE
-	printf("allocate before splits: 0x%p, %i, %i, %i, SyncSize: %i\n", FreeChunk->Base, (int32)FreeChunk->Size, (int32)AllocationSize, (int32)GetCurrentSyncIndex(), (int32)FreeChunk->SyncSize);
+	if (GGPUDefragValidateSyncSize)
+	{
+		printf("allocate before splits: 0x%p, %i, %i, %i, SyncSize: %i\n", FreeChunk->Base, (int32)FreeChunk->Size, (int32)AllocationSize, (int32)GetCurrentSyncIndex(), (int32)FreeChunk->SyncSize);
+	}
 #endif
 
 	// If this is an immediate allocation (i.e. the CPU will start accessing the memory right away)
@@ -221,9 +238,12 @@ FGPUDefragAllocator::FMemoryChunk* FGPUDefragAllocator::AllocateChunk(FMemoryChu
 	if (!bAsync && FreeChunk->IsRelocating() && FreeChunk->SyncSize > 0 && FreeChunk->SyncSize < FreeChunk->Size)
 	{
 #if VALIDATE_SYNC_SIZE
-		printf("splitting:\n");
+		if (GGPUDefragValidateSyncSize)
+		{
+			printf("splitting:\n");
+		}
 #endif
-		Split(FreeChunk, FreeChunk->SyncSize, bAsync);
+		Split(FreeChunk, FreeChunk->SyncSize);
 		FreeChunk = FreeChunk->NextChunk;
 	}
 
@@ -234,9 +254,12 @@ FGPUDefragAllocator::FMemoryChunk* FGPUDefragAllocator::AllocateChunk(FMemoryChu
 	if (FreeChunk->Size > AllocationSize)
 	{
 #if VALIDATE_SYNC_SIZE
-		printf("splitting again:\n");
+		if (GGPUDefragValidateSyncSize)
+		{
+			printf("splitting again:\n");
+		}
 #endif
-		Split(FreeChunk, AllocationSize, bAsync);
+		Split(FreeChunk, AllocationSize);
 	}
 
 	// Ensure that everything's in range.
@@ -248,10 +271,13 @@ FGPUDefragAllocator::FMemoryChunk* FGPUDefragAllocator::AllocateChunk(FMemoryChu
 	FPlatformAtomics::InterlockedAdd(&AvailableMemorySize, -FreeChunk->Size);
 
 #if VALIDATE_SYNC_SIZE
-	printf("allocate: 0x%p, %i, %i\n", FreeChunk->Base, (int32)AllocationSize, (int32)GetCurrentSyncIndex());
-	if (bDoValidation)
+	if (GGPUDefragValidateSyncSize)
 	{
-		ValidateRelocations(FreeChunk->Base, AllocationSize);
+		printf("allocate: 0x%p, %i, %i\n", FreeChunk->Base, (int32)AllocationSize, (int32)GetCurrentSyncIndex());
+		if (bDoValidation)
+		{
+			ValidateRelocations(FreeChunk->Base, AllocationSize);
+		}
 	}
 #endif
 
@@ -267,10 +293,13 @@ FGPUDefragAllocator::FMemoryChunk* FGPUDefragAllocator::AllocateChunk(FMemoryChu
 * @param Chunk						Chunk to free
 * @param bMaintainSortedFreelist	If true, maintains the free-list sort order
 */
-void FGPUDefragAllocator::FreeChunk(FMemoryChunk* Chunk, bool bMaintainSortedFreelist)
+void FGPUDefragAllocator::FreeChunk(FMemoryChunk* Chunk)
 {
 #if VALIDATE_SYNC_SIZE
-	printf("FreeChunk: 0x%p, %i, %i, %i\n", Chunk->Base, (int32)Chunk->Size, (int32)Chunk->SyncSize, (int32)Chunk->SyncIndex);
+	if (GGPUDefragValidateSyncSize)
+	{
+		printf("FreeChunk: 0x%p, %i, %i, %i\n", Chunk->Base, (int32)Chunk->Size, (int32)Chunk->SyncSize, (int32)Chunk->SyncIndex);
+	}
 #endif
 
 	// Remove the entry
@@ -282,7 +311,7 @@ void FGPUDefragAllocator::FreeChunk(FMemoryChunk* Chunk, bool bMaintainSortedFre
 	FPlatformAtomics::InterlockedAdd(&AvailableMemorySize, Chunk->Size);
 
 	// Free the chunk.
-	LinkFreeChunk(Chunk, bMaintainSortedFreelist);
+	LinkFreeChunk(Chunk);
 }
 
 #if !(UE_BUILD_SHIPPING || UE_BUILD_TEST)
@@ -309,6 +338,9 @@ void FGPUDefragAllocator::Free(void* Pointer)
 	check(MatchingChunk);
 	check(MatchingChunk->Base == Pointer);
 
+	int64 PaddingWaste = MatchingChunk->Size - MatchingChunk->OrigSize;
+	FPlatformAtomics::InterlockedAdd(&PaddingWasteSize, -PaddingWaste);
+
 	// Is this chunk is currently being relocated asynchronously (by the GPU)?
 	if (MatchingChunk->IsRelocating())
 	{
@@ -317,7 +349,7 @@ void FGPUDefragAllocator::Free(void* Pointer)
 	else
 	{
 		// Free the chunk.
-		FreeChunk(MatchingChunk, false);
+		FreeChunk(MatchingChunk);
 	}
 }
 
@@ -400,7 +432,7 @@ void* FGPUDefragAllocator::GetUserPayload(const void* Pointer)
 * @param	Pointer		Pointer to check.
 * @return				Number of bytes allocated
 */
-int32 FGPUDefragAllocator::GetAllocatedSize(void* Pointer)
+int64 FGPUDefragAllocator::GetAllocatedSize(void* Pointer)
 {
 	FScopeLock Lock(&SynchronizationObject);
 	FMemoryChunk* MatchingChunk = PointerToChunkMap.FindRef(Pointer);
@@ -432,7 +464,7 @@ void* FGPUDefragAllocator::Reallocate(void* OldBaseAddress, int64 NewSize)
 
 	int64 AlignedNewSize = Align(NewSize, AllocationAlignment);
 	FMemoryChunk* NewChunk = nullptr;
-	int32 MemoryAdjustment = FMath::Abs(AlignedNewSize - MatchingChunk->Size);
+	int64 MemoryAdjustment = FMath::Abs(AlignedNewSize - MatchingChunk->Size);
 
 	// Are we growing the allocation?
 	if (AlignedNewSize > MatchingChunk->Size)
@@ -453,7 +485,7 @@ void* FGPUDefragAllocator::Reallocate(void* OldBaseAddress, int64 NewSize)
 * @param GrowAmount	Number of bytes to grow by
 * @return				nullptr if it failed, otherwise the new grown chunk
 */
-FGPUDefragAllocator::FMemoryChunk* FGPUDefragAllocator::Grow(FMemoryChunk* Chunk, int32 GrowAmount)
+FGPUDefragAllocator::FMemoryChunk* FGPUDefragAllocator::Grow(FMemoryChunk* Chunk, int64 GrowAmount)
 {
 	// Is there enough free memory immediately before this chunk?
 	FMemoryChunk* PrevChunk = Chunk->PreviousChunk;
@@ -493,7 +525,7 @@ FGPUDefragAllocator::FMemoryChunk* FGPUDefragAllocator::Grow(FMemoryChunk* Chunk
 * @param ShrinkAmount	Number of bytes to shrink by
 * @return				The new shrunken chunk
 */
-FGPUDefragAllocator::FMemoryChunk* FGPUDefragAllocator::Shrink(FMemoryChunk* Chunk, int32 ShrinkAmount)
+FGPUDefragAllocator::FMemoryChunk* FGPUDefragAllocator::Shrink(FMemoryChunk* Chunk, int64 ShrinkAmount)
 {
 	// We're shrinking the allocation.
 	check(ShrinkAmount <= Chunk->Size);
@@ -508,20 +540,20 @@ FGPUDefragAllocator::FMemoryChunk* FGPUDefragAllocator::Shrink(FMemoryChunk* Chu
 		Chunk->Size -= ShrinkAmount;
 
 		// Grow the previous chunk.
-		int32 OriginalPrevSize = NewFreeChunk->Size;
+		int64 OriginalPrevSize = NewFreeChunk->Size;
 		NewFreeChunk->Size += ShrinkAmount;
 
 		// If the previous chunk was "in use", split it and insert a 2nd free chunk.
 		if (!NewFreeChunk->bIsAvailable)
 		{
-			Split(NewFreeChunk, OriginalPrevSize, false);
+			Split(NewFreeChunk, OriginalPrevSize);
 			NewFreeChunk = NewFreeChunk->NextChunk;
 		}
 	}
 	else
 	{
 		// This was the first chunk, split it.
-		Split(Chunk, ShrinkAmount, false);
+		Split(Chunk, ShrinkAmount);
 
 		// We're going to use the new chunk. Mark it as "used memory".
 		Chunk = Chunk->NextChunk;
@@ -530,9 +562,12 @@ FGPUDefragAllocator::FMemoryChunk* FGPUDefragAllocator::Shrink(FMemoryChunk* Chu
 		// Make the original chunk "free memory".
 		NewFreeChunk = Chunk->PreviousChunk;
 #if VALIDATE_SYNC_SIZE
-		printf("shrink free chunk\n");
+		if (GGPUDefragValidateSyncSize)
+		{
+			printf("shrink free chunk\n");
+		}
 #endif
-		LinkFreeChunk(NewFreeChunk, false);
+		LinkFreeChunk(NewFreeChunk);
 	}
 
 	// Mark the newly freed memory as "being relocated" and require GPU sync.
@@ -675,22 +710,14 @@ void FGPUDefragAllocator::SortFreeList(int32& NumFreeChunks, int64& LargestFreeC
 	}
 }
 
-/**
-* Defrag helper function. Checks if the specified allocation fits within
-* the adjacent free chunk(s), accounting for potential reallocation request.
-*
-* @param UsedChunk		Allocated chunk to check for a fit
-* @param bAnyChunkType	If false, only succeeds if 'UsedChunk' has a reallocation request and fits
-* @return				Returns 'UsedChunk' if it fits the criteria, otherwise nullptr
-*/
 FGPUDefragAllocator::FMemoryChunk* FGPUDefragAllocator::FindAdjacent(FMemoryChunk* UsedChunk, bool bAnyChunkType)
 {
-	if (UsedChunk && !UsedChunk->IsRelocating() && (bAnyChunkType))
+	if (UsedChunk && (bAnyChunkType))
 	{
 		FMemoryChunk* FreeChunkLeft = UsedChunk->PreviousChunk;
 		FMemoryChunk* FreeChunkRight = UsedChunk->NextChunk;
 
-		int32 AvailableSize = UsedChunk->Size;
+		int64 AvailableSize = UsedChunk->Size;
 		if (FreeChunkLeft && FreeChunkLeft->bIsAvailable)
 		{
 			AvailableSize += FreeChunkLeft->Size;
@@ -701,7 +728,7 @@ FGPUDefragAllocator::FMemoryChunk* FGPUDefragAllocator::FindAdjacent(FMemoryChun
 		}
 
 		// Does it fit?
-		int32 FinalSize = UsedChunk->GetFinalSize();
+		int64 FinalSize = UsedChunk->GetFinalSize();
 		if (FinalSize <= AvailableSize && CanRelocate(UsedChunk))
 		{
 			return UsedChunk;
@@ -743,7 +770,7 @@ FGPUDefragAllocator::FMemoryChunk* FGPUDefragAllocator::FindAdjacentToHole(FMemo
 		{
 			return Chunk->PreviousChunk;
 		}
-		Chunk = Chunk->NextFreeChunk;
+		Chunk = Chunk->PreviousFreeChunk;
 	}
 	return nullptr;
 }
@@ -766,7 +793,7 @@ FGPUDefragAllocator::FMemoryChunk* FGPUDefragAllocator::FindAny(FMemoryChunk* Fr
 	{
 		if (!CurrentChunk->bIsAvailable)
 		{
-			int32 CurrentFit = FreeChunk->Size - CurrentChunk->GetFinalSize();
+			int64 CurrentFit = FreeChunk->Size - CurrentChunk->GetFinalSize();
 
 			// Better fit than previously?
 			if (CurrentFit >= 0 && CurrentFit < BestFit && CanRelocate(CurrentChunk))
@@ -805,7 +832,7 @@ void FGPUDefragAllocator::CheckForErrors(bool bCheckSortedFreeList)
 	if (bCheckSortedFreeList)
 	{
 		FMemoryChunk* Chunk = FirstFreeChunk;
-		int32 TotalFreeMem = Chunk->Size;
+		int64 TotalFreeMem = Chunk->Size;
 		while (Chunk->NextFreeChunk)
 		{
 			check(Chunk->bIsAvailable);
@@ -818,8 +845,8 @@ void FGPUDefragAllocator::CheckForErrors(bool bCheckSortedFreeList)
 		check(TotalFreeMem == AvailableMemorySize);
 	}
 
-	int32 TotalUsedMem = 0;
-	int32 TotalFreeMem = 0;
+	int64 TotalUsedMem = 0;
+	int64 TotalFreeMem = 0;
 	FMemoryChunk* Chunk = FirstChunk;
 	while (Chunk)
 	{
@@ -840,21 +867,23 @@ void FGPUDefragAllocator::CheckForErrors(bool bCheckSortedFreeList)
 #if VALIDATE_SYNC_SIZE
 void FGPUDefragAllocator::ValidateRelocations(uint8* UsedBaseAddress, uint64 Size)
 {
-
-	FScopeLock SyncLock(&SynchronizationObject);
-
-	for (int32 i = 0; i < Relocations.Num(); ++i)
+	if (GGPUDefragValidateSyncSize)
 	{
-		const FRelocationEntry& Relocation = Relocations[i];
+		FScopeLock SyncLock(&SynchronizationObject);
 
-		bool bBeforeOrig = (UsedBaseAddress + Size) <= Relocation.OldBase;
-		bool bAfterOrig = (UsedBaseAddress) >= (Relocation.OldBase + Relocation.Size);
+		for (int32 i = 0; i < Relocations.Num(); ++i)
+		{
+			const FRelocationEntry& Relocation = Relocations[i];
+
+			bool bBeforeOrig = (UsedBaseAddress + Size) <= Relocation.OldBase;
+			bool bAfterOrig = (UsedBaseAddress) >= (Relocation.OldBase + Relocation.Size);
 
-		bool bBeforeNew = (UsedBaseAddress + Size) <= Relocation.NewBase;
-		bool bAfterNew = (UsedBaseAddress) >= (Relocation.NewBase + Relocation.Size);
+			bool bBeforeNew = (UsedBaseAddress + Size) <= Relocation.NewBase;
+			bool bAfterNew = (UsedBaseAddress) >= (Relocation.NewBase + Relocation.Size);
 
-		bool bAfterSync = Relocation.SyncIndex <= CompletedSyncIndex;
-		checkf(((bBeforeOrig || bAfterOrig) && (bBeforeNew || bAfterNew)) || bAfterSync, TEXT("Corruption Hazard, Allocation not protected by sync size."));
+			bool bAfterSync = Relocation.SyncIndex <= CompletedSyncIndex;
+			checkf(((bBeforeOrig || bAfterOrig) && (bBeforeNew || bAfterNew)) || bAfterSync, TEXT("Corruption Hazard, Allocation not protected by sync size."));
+		}
 	}
 }
 #endif
@@ -877,6 +906,7 @@ FGPUDefragAllocator::FMemoryChunk* FGPUDefragAllocator::RelocateIntoFreeChunk(FR
 
 	// Save off important data from 'SourceChunk', since it will get modified by the call to LinkFreeChunk().	
 	void* UserPayload = SourceChunk->UserPayload;
+	const int64 OrigSize = SourceChunk->OrigSize;
 	const int64 OldSize = SourceChunk->Size;
 	const uint8* SourceOldBase = SourceChunk->Base;
 	const uint8* DestNewBase = FreeChunk->Base;
@@ -894,9 +924,12 @@ FGPUDefragAllocator::FMemoryChunk* FGPUDefragAllocator::RelocateIntoFreeChunk(FR
 
 	// Merge adjacent free chunks into SourceChunk to make a single free chunk.
 #if VALIDATE_SYNC_SIZE
-	printf("relocate link free chunk\n");
+	if (GGPUDefragValidateSyncSize)
+	{
+		printf("relocate link free chunk\n");
+	}
 #endif
-	LinkFreeChunk(SourceChunk, true);
+	LinkFreeChunk(SourceChunk);
 
 	FMemoryChunk* DestinationChunk;
 	if (bAdjacentRelocation)
@@ -924,6 +957,9 @@ FGPUDefragAllocator::FMemoryChunk* FGPUDefragAllocator::RelocateIntoFreeChunk(FR
 	// Make sure the destination chunk keeps the UserPayload, no matter what. :)
 	DestinationChunk->UserPayload = UserPayload;
 
+	SourceChunk->OrigSize = 0;
+	DestinationChunk->OrigSize = OrigSize;
+
 	// Update our book-keeping.
 	PointerToChunkMap.Remove(SourceOldBase);
 	PointerToChunkMap.Add(NewBaseAddress, DestinationChunk);
@@ -935,7 +971,7 @@ FGPUDefragAllocator::FMemoryChunk* FGPUDefragAllocator::RelocateIntoFreeChunk(FR
 	if (DestinationChunk->Size > NewSize)
 	{
 		// Split the DestinationChunk into a used chunk and a free chunk.
-		Split(DestinationChunk, NewSize, true);
+		Split(DestinationChunk, NewSize);
 		NextFreeChunk = DestinationChunk->NextChunk;
 		NextFreeChunk->bTail = bAdjacentRelocation;
 
@@ -969,13 +1005,13 @@ FGPUDefragAllocator::FMemoryChunk* FGPUDefragAllocator::RelocateIntoFreeChunk(FR
 				int64 RelocatableSize = FMath::Max((int64)((NextFreeChunk->Size - NewSize - RightAddSize)), (int64)0);				
 				if ((RelocatableSize > 0) && NextFreeChunk->GetAvailableSize())
 				{
-					Split(NextFreeChunk, RelocatableSize, true);
+					Split(NextFreeChunk, RelocatableSize);
 					FMemoryChunk* InFlightTailChunk = NextFreeChunk->NextFreeChunk;					
 
 					//if the tail has right coalesced memory afterwards, split the right coalesce back off into its own chunk for re-use.
 					if (InFlightTailChunk->Size > OldSize)
 					{
-						Split(InFlightTailChunk, OldSize, true);
+						Split(InFlightTailChunk, OldSize);
 					}
 					//defragger code assumes free chunks are not adjacent (they would have been coalesced)
 					//so mark this as 'used' so it can't be allocated while the tail is being moved by the GPU and put it on the deferred free list
@@ -1037,15 +1073,18 @@ FGPUDefragAllocator::FMemoryChunk* FGPUDefragAllocator::RelocateIntoFreeChunk(FR
 
 	if (NewSize != OldSize)
 	{
-		int32 MemoryAdjustment = NewSize - OldSize;
+		int64 MemoryAdjustment = NewSize - OldSize;
 		FPlatformAtomics::InterlockedAdd(&AllocatedMemorySize, +MemoryAdjustment);
 		FPlatformAtomics::InterlockedAdd(&AvailableMemorySize, -MemoryAdjustment);
 		FPlatformAtomics::InterlockedAdd(&PendingMemoryAdjustment, -MemoryAdjustment);
 	}
 
 #if VALIDATE_SYNC_SIZE
-	Relocations.Add(FRelocationEntry(SourceOldBase, DestNewBase, OldSize, GetCurrentSyncIndex()));
-	printf("bFreeAdjOnLeft %i, bFreeAdjOnRight %i\n", bFreeChunkPreviousAdjacent, bFreeChunkNextAdjacent);
+	if (GGPUDefragValidateSyncSize)
+	{
+		Relocations.Add(FRelocationEntry(SourceOldBase, DestNewBase, OldSize, GetCurrentSyncIndex()));
+		printf("bFreeAdjOnLeft %i, bFreeAdjOnRight %i\n", bFreeChunkPreviousAdjacent, bFreeChunkNextAdjacent);
+	}
 #endif
 
 	// Enable for debugging:
@@ -1079,7 +1118,7 @@ bool FGPUDefragAllocator::FinishAllRelocations()
 	for (TDoubleLinkedList<FMemoryChunk*>::TIterator It(PendingFreeChunks.GetHead()); It; ++It)
 	{
 		FMemoryChunk* Chunk = *It;
-		FreeChunk(Chunk, false);
+		FreeChunk(Chunk);
 	}
 	PendingFreeChunks.Empty();
 
@@ -1186,7 +1225,7 @@ FGPUDefragAllocator::FMemoryChunk* FGPUDefragAllocator::RelocateAllowed(FMemoryC
 	return AllowedChunk;
 }
 
-void FGPUDefragAllocator::PartialDefragmentation(FRelocationStats& Stats, double StartTime)
+void FGPUDefragAllocator::PartialDefragmentationFast(FRelocationStats& Stats, double StartTime)
 {
 	FMemoryChunk* FreeChunk = FirstFreeChunk;
 	while (FreeChunk /*&& ReallocationRequests.Num() > 0*/ && Stats.NumBytesRelocated < Settings.MaxDefragRelocations && Stats.NumRelocations < GGPUDefragMaxRelocations)
@@ -1200,24 +1239,16 @@ void FGPUDefragAllocator::PartialDefragmentation(FRelocationStats& Stats, double
 		}
 		else
 		{
-			// 1. Merge with Left, if it has a Reallocation request and fits
-			BestChunk = RelocateAllowed(FreeChunk, FindAdjacent(FreeChunk->PreviousChunk, false));
-			if (!BestChunk)
-			{
-				// 2. Merge with Right, if it has an async request and fits
-				BestChunk = RelocateAllowed(FreeChunk, FindAdjacent(FreeChunk->NextChunk, false));
+			//Not much point merging to the used chunk on the left.  We should have already done any merges on the left with the ordered free chunk walk.
+			//BestChunk = RelocateAllowed(FreeChunk, FindAdjacent(FreeChunk->PreviousChunk, true));
+			
+			// 1. Merge with Right,
+			BestChunk = RelocateAllowed(FreeChunk, FindAdjacent(FreeChunk->NextChunk, true));
 
-				if (!BestChunk)
-				{					
-					// 4. Merge with a used chunk adjacent to hole (to make that hole larger).
-					BestChunk = RelocateAllowed(FreeChunk, FindAdjacentToHole(FreeChunk));
-
-					if (!BestChunk)
-					{
-						// 5. Merge with chunk from the end of the pool (well-fitting)
-						BestChunk = FindAny(FreeChunk);
-					}					
-				}
+			if (!BestChunk)
+			{					
+				// 2. Merge with a used chunk adjacent to hole (to make that hole larger).
+				BestChunk = RelocateAllowed(FreeChunk, FindAdjacentToHole(FreeChunk));				
 			}
 		}
 
@@ -1253,6 +1284,64 @@ void FGPUDefragAllocator::PartialDefragmentation(FRelocationStats& Stats, double
 	}
 }
 
+void FGPUDefragAllocator::PartialDefragmentationSlow(FRelocationStats& Stats, double StartTime)
+{
+	if (Stats.NumBytesRelocated > 0)
+	{
+		return;
+	}
+
+	FMemoryChunk* FreeChunk = FirstFreeChunk;
+	while (FreeChunk /*&& ReallocationRequests.Num() > 0*/ && Stats.NumBytesRelocated < Settings.MaxDefragRelocations && Stats.NumRelocations < GGPUDefragMaxRelocations)
+	{
+		FRequestNode* BestRequestNode = nullptr;
+		FMemoryChunk* BestChunk = nullptr;
+
+		if (FreeChunk->DefragCounter)
+		{
+			FreeChunk->DefragCounter--;
+		}
+		else
+		{
+			if (!BestChunk)
+			{
+				//1. Merge with chunk from the end of the pool (well-fitting)
+				BestChunk = FindAny(FreeChunk);
+			}			
+		}
+
+		if (BestChunk)
+		{
+			FreeChunk = RelocateIntoFreeChunk(Stats, FreeChunk, BestChunk);
+		}
+		else
+		{
+			// Did the free chunk fail to defrag?
+			if (FreeChunk->DefragCounter == 0 && (FreeChunk->NextFreeChunk))
+			{
+				// Don't try it again for a while.
+				if (FreeChunk->Size < DEFRAG_SMALL_CHUNK_SIZE)
+				{
+					FreeChunk->DefragCounter = DEFRAG_SMALL_CHUNK_COUNTER_MIN + FMath::RandHelper(DEFRAG_SMALL_CHUNK_COUNTER_MAX - DEFRAG_SMALL_CHUNK_COUNTER_MIN);
+				}
+				else
+				{
+					FreeChunk->DefragCounter = DEFRAG_CHUNK_COUNTER_MIN + FMath::RandHelper(DEFRAG_CHUNK_COUNTER_MAX - DEFRAG_CHUNK_COUNTER_MIN);
+				}
+			}
+
+			FreeChunk = FreeChunk->NextFreeChunk;
+		}
+
+		// Limit time spent
+		double TimeSpent = FPlatformTime::Seconds() - StartTime;
+		if ((GGPUDefragEnableTimeLimits != 0) && (TimeSpent > PARTIALDEFRAG_TIMELIMIT))
+		{
+			break;
+		}
+	}
+}
+
 /**
 * Performs a partial defrag by shifting down memory to fill holes, in a brute-force manner.
 * Takes consideration to async reallocations, but processes the all memory in order.
@@ -1357,8 +1446,11 @@ int32 FGPUDefragAllocator::Tick(FRelocationStats& Stats, bool bPanicDefrag)
 	// There may still be chunks being flagged as 'IsRelocating' due to immediate shrinks between calls.
 	FinishAllRelocations();
 
-#if VALIDATE_SYNC_SIZE	
-	Relocations.Empty();
+#if VALIDATE_SYNC_SIZE
+	if (GGPUDefragValidateSyncSize)
+	{
+		Relocations.Empty();
+	}
 #endif
 
 	// Sort the Free chunks.
@@ -1369,10 +1461,12 @@ int32 FGPUDefragAllocator::Tick(FRelocationStats& Stats, bool bPanicDefrag)
 		if (!bPanicDefrag)
 		{
 			// Smart defrag
-			PartialDefragmentation(Stats, StartTime);			
+			PartialDefragmentationFast(Stats, StartTime);
+
+			PartialDefragmentationSlow(Stats, StartTime);
 
 			// Brute-force defrag
-			PartialDefragmentationDownshift(Stats, StartTime);			
+			//PartialDefragmentationDownshift(Stats, StartTime);			
 		}
 		else
 		{
@@ -1419,17 +1513,17 @@ int32 FGPUDefragAllocator::Tick(FRelocationStats& Stats, bool bPanicDefrag)
 void FGPUDefragAllocator::DumpAllocs(FOutputDevice& Ar/*=*GLog*/)
 {
 	// Memory usage stats.
-	int32				UsedSize = 0;
-	int32				FreeSize = 0;
-	int32				NumUsedChunks = 0;
-	int32				NumFreeChunks = 0;
+	int64				UsedSize = 0;
+	int64				FreeSize = 0;
+	int64				NumUsedChunks = 0;
+	int64				NumFreeChunks = 0;
 
 	// Fragmentation and allocation size visualization.
-	int32				NumBlocks = MemorySize / AllocationAlignment;
-	int32				Dimension = 1 + NumBlocks / FMath::TruncToInt(FMath::Sqrt(NumBlocks));
+	int64				NumBlocks = MemorySize / AllocationAlignment;
+	int64				Dimension = 1 + NumBlocks / FMath::TruncToInt(FMath::Sqrt(NumBlocks));
 	TArray<FColor>	AllocationVisualization;
 	AllocationVisualization.AddZeroed(Dimension * Dimension);
-	int32				VisIndex = 0;
+	int64				VisIndex = 0;
 
 	// Traverse linked list and gather allocation information.
 	FMemoryChunk*	CurrentChunk = FirstChunk;
@@ -1655,8 +1749,8 @@ void FGPUDefragAllocator::DefragmentMemory(FRelocationStats& Stats)
 #if !(UE_BUILD_SHIPPING || UE_BUILD_TEST)
 	int32 NumHolesBefore = Stats.NumHoles;
 	int32 NumHolesAfter = 0;
-	int32 LargestHoleBefore = Stats.LargestHoleSize;
-	int32 LargestHoleAfter = GetLargestAvailableAllocation(&NumHolesAfter);
+	int64 LargestHoleBefore = Stats.LargestHoleSize;
+	int64 LargestHoleAfter = GetLargestAvailableAllocation(&NumHolesAfter);
 	double EndTime = FPlatformTime::Seconds();
 	double TotalDuration = EndTime - StartTime;
 	double GPUDuration = EndTime - MidTime;
@@ -1678,9 +1772,9 @@ void FGPUDefragAllocator::Coalesce(FMemoryChunk* FreedChunk)
 	check(!FreedChunk->IsLocked())
 
 	uint32 LatestSyncIndex = 0;
-	int32 LatestSyncSize = 0;
-	int32 LeftSize = 0;
-	int32 RightSize = 0;
+	int64 LatestSyncSize = 0;
+	int64 LeftSize = 0;
+	int64 RightSize = 0;
 
 	// Check if the previous chunk is available.
 	FMemoryChunk* LeftChunk = FreedChunk->PreviousChunk;
@@ -1726,7 +1820,10 @@ void FGPUDefragAllocator::Coalesce(FMemoryChunk* FreedChunk)
 	}
 
 #if VALIDATE_SYNC_SIZE
-	printf("FreeChunk Before Coalesce: 0x%p, %i, %i, %i\n", FreedChunk->Base, (int32)FreedChunk->Size, (int32)FreedChunk->SyncSize, (int32)FreedChunk->SyncIndex);
+	if (GGPUDefragValidateSyncSize)
+	{
+		printf("FreeChunk Before Coalesce: 0x%p, %i, %i, %i\n", FreedChunk->Base, (int32)FreedChunk->Size, (int32)FreedChunk->SyncSize, (int32)FreedChunk->SyncIndex);
+	}
 #endif
 
 	// Merge.
@@ -1735,7 +1832,10 @@ void FGPUDefragAllocator::Coalesce(FMemoryChunk* FreedChunk)
 	FreedChunk->SetSyncIndex(LatestSyncIndex, LatestSyncSize);
 
 #if VALIDATE_SYNC_SIZE
-	printf("FreeChunk AFter Coalesce: 0x%p, %i, %i, %i\n", FreedChunk->Base, (int32)FreedChunk->Size, (int32)FreedChunk->SyncSize, (int32)FreedChunk->SyncIndex);
+	if (GGPUDefragValidateSyncSize)
+	{
+		printf("FreeChunk AFter Coalesce: 0x%p, %i, %i, %i\n", FreedChunk->Base, (int32)FreedChunk->Size, (int32)FreedChunk->SyncSize, (int32)FreedChunk->SyncIndex);
+	}
 #endif
 }
 
@@ -1755,8 +1855,8 @@ void FGPUDefragAllocator::Benchmark(int32 MinChunkSize, int32 MaxChunkSize, floa
 //untested
 #if 0
 #if !(UE_BUILD_SHIPPING ||UE_BUILD_TEST)
-	int32 OriginalAllocatedMemorySize = AllocatedMemorySize;
-	int32 OriginalAvailableMemorySize = AvailableMemorySize;
+	int64 OriginalAllocatedMemorySize = AllocatedMemorySize;
+	int64 OriginalAvailableMemorySize = AvailableMemorySize;
 	FRandomStream Rand(0x1ee7c0de);
 	TArray<void*> Chunks;
 	TArray<FMemoryChunk*> OriginalChunks;
diff --git a/Engine/Source/Runtime/RHI/Public/GPUDefragAllocator.h b/Engine/Source/Runtime/RHI/Public/GPUDefragAllocator.h
--- a/Engine/Source/Runtime/RHI/Public/GPUDefragAllocator.h
+++ b/Engine/Source/Runtime/RHI/Public/GPUDefragAllocator.h
@@ -6,7 +6,7 @@
 
 #define LOG_EVERY_ALLOCATION			0
 #define DUMP_ALLOC_FREQUENCY			0 // 100
-#define VALIDATE_SYNC_SIZE				0
+#define VALIDATE_SYNC_SIZE				(!(UE_BUILD_TEST || UE_BUILD_SHIPPING))
 
 
 /*-----------------------------------------------------------------------------
@@ -69,9 +69,9 @@ public:
 		, Type(MET_Allocated)
 		{
 		}
-		FMemoryLayoutElement(int32 InSize, EMemoryElementType int32ype)
+		FMemoryLayoutElement(int32 InSize, EMemoryElementType InType)
 			: Size(InSize)
-			, Type(int32ype)
+			, Type(InType)
 		{
 		}
 		int32				Size;
@@ -80,9 +80,9 @@ public:
 		friend FArchive& operator<<(FArchive& Ar, FMemoryLayoutElement& Element)
 		{
 			Ar << Element.Size;
-			uint32 Typeuint32 = Element.Type;
-			Ar << Typeuint32;
-			Element.Type = EMemoryElementType(Typeuint32);
+			uint32 ElementType = Element.Type;
+			Ar << ElementType;
+			Element.Type = EMemoryElementType(ElementType);
 			return Ar;
 		}
 	};
@@ -135,9 +135,10 @@ public:
 		* @param	ChunkToInsertAfter		Chunk to insert this after.
 		* @param	FirstFreeChunk			Reference to first free chunk Pointer.
 		*/
-		FMemoryChunk(uint8* InBase, int32 InSize, FGPUDefragAllocator& InBestFitAllocator, FMemoryChunk*& ChunkToInsertAfter, TStatId InStat, bool bSortedFreeList)
+		FMemoryChunk(uint8* InBase, int64 InSize, FGPUDefragAllocator& InBestFitAllocator, FMemoryChunk*& ChunkToInsertAfter, TStatId InStat)
 			: Base(InBase)
 			, Size(InSize)
+			, OrigSize(0)
 			, bIsAvailable(false)
 			, LockCount(0)
 			, DefragCounter(0)
@@ -150,7 +151,7 @@ public:
 		{
 			Link(ChunkToInsertAfter);
 			// This is going to change bIsAvailable.
-			LinkFree(bSortedFreeList, ChunkToInsertAfter);
+			LinkFree(ChunkToInsertAfter);
 		}
 
 		/**
@@ -202,7 +203,7 @@ public:
 		/**
 		* Inserts this chunk at the head of the free chunk list.
 		*/
-		void	LinkFree(bool bMaint32ainSortOrder, FMemoryChunk* FirstFreeChunkToSearch);
+		void	LinkFree(FMemoryChunk* FirstFreeChunkToSearch);
 
 		/**
 		* Removes itself for linked list.
@@ -290,7 +291,7 @@ public:
 		* @param InSyncIndex	GPU synchronization identifier that can be compared with BestFitAllocator::CompletedSyncIndex
 		* @param InSyncSize	Number of uint8s that require GPU synchronization (starting from the beginning of the chunk)
 		*/
-		void	SetSyncIndex(uint32 InSyncIndex, int32 InSyncSize)
+		void	SetSyncIndex(uint32 InSyncIndex, int64 InSyncSize)
 		{			
 			SyncIndex = InSyncIndex;
 			SyncSize = InSyncSize;
@@ -322,6 +323,7 @@ public:
 		uint8*					Base;
 		/** Size of chunk.								*/
 		int64					Size;
+		int64					OrigSize;
 		/** Whether the chunk is available.				*/
 		bool					bIsAvailable;
 		/** Whether the chunk has been locked.			*/		
@@ -362,6 +364,7 @@ public:
 		, LastChunk(nullptr)
 		, FirstFreeChunk(nullptr)
 		, TimeSpentInAllocator(0.0)
+		, PaddingWasteSize(0)
 		, AllocatedMemorySize(0)
 		, AvailableMemorySize(0)
 		, PendingMemoryAdjustment(0)
@@ -401,7 +404,7 @@ public:
 		// Update stats in a thread safe way.
 		FPlatformAtomics::InterlockedExchange(&AvailableMemorySize, MemorySize);
 		// Allocate initial chunk.
-		FirstChunk = new FMemoryChunk(MemoryBase, MemorySize, *this, FirstChunk, TStatId(), false);
+		FirstChunk = new FMemoryChunk(MemoryBase, MemorySize, *this, FirstChunk, TStatId());
 		LastChunk = FirstChunk;
 	}
 
@@ -490,7 +493,7 @@ public:
 	* @param Pointer		Pointer to check.
 	* @return				Number of uint8s allocated
 	*/
-	int32		GetAllocatedSize(void* Pointer);
+	int64		GetAllocatedSize(void* Pointer);
 
 	/**
 	* Tries to reallocate texture memory in-place (without relocating),
@@ -531,11 +534,12 @@ public:
 	* @param	OutAvailableMemorySize		[out] Size of available memory
 	* @param	OutPendingMemoryAdjustment	[out] Size of pending allocation change (due to async reallocation)
 	*/
-	virtual void	GetMemoryStats(int64& OutAllocatedMemorySize, int64& OutAvailableMemorySize, int64& OutPendingMemoryAdjustment)
+	virtual void	GetMemoryStats(int64& OutAllocatedMemorySize, int64& OutAvailableMemorySize, int64& OutPendingMemoryAdjustment, int64& OutPaddingWasteSize)
 	{
 		OutAllocatedMemorySize = AllocatedMemorySize;
 		OutAvailableMemorySize = AvailableMemorySize;
 		OutPendingMemoryAdjustment = PendingMemoryAdjustment;
+		OutPaddingWasteSize = PaddingWasteSize;
 	}
 
 	int64 GetTotalSize() const
@@ -644,6 +648,11 @@ public:
 		return !(UPTRINT(Ptr) & (Alignment - 1));
 	}
 
+	int32 GetAllocationAlignment() const
+	{
+		return AllocationAlignment;
+	}
+
 protected:
 
 #if VALIDATE_SYNC_SIZE
@@ -678,7 +687,7 @@ protected:
 	* @param Size			Number of uint8s to copy
 	* @param UserPayload	User payload for this allocation
 	*/
-	virtual void	PlatformRelocate(void* Dest, const void* Source, int32 Size, void* UserPayload) = 0;
+	virtual void	PlatformRelocate(void* Dest, const void* Source, int64 Size, void* UserPayload) = 0;
 
 	/**
 	* Inserts a fence to synchronize relocations.
@@ -726,7 +735,7 @@ protected:
 	* @param Size			Number of uint8s to copy
 	* @param UserPayload	User payload for the allocation
 	*/
-	void Relocate(FRelocationStats& Stats, FMemoryChunk* Dest, int32 DestOffset, const void* Source, int32 Size, void* UserPayload)
+	void Relocate(FRelocationStats& Stats, FMemoryChunk* Dest, int64 DestOffset, const void* Source, int64 Size, void* UserPayload)
 	{
 		uint8* DestAddr = Dest->Base + DestOffset;
 		int64 MemDistance = (int64)(Dest)-(int64)(Source);
@@ -737,7 +746,7 @@ protected:
 		{
 			PlatformRelocate(DestAddr, Source, Size, UserPayload);
 		}
-		int32 RelocateSize = bOverlappedMove ? (Size * Settings.MaxDefragRelocations) : Size;
+		int64 RelocateSize = bOverlappedMove ? (Size * Settings.MaxDefragRelocations) : Size;
 		Dest->UserPayload = UserPayload;
 		Stats.NumBytesRelocated += RelocateSize;
 		Stats.NumRelocations++;		
@@ -752,12 +761,21 @@ protected:
 	}
 
 	/**
-	* Performs a partial defrag while trying to process any pending async reallocation requests.
+	* Performs a partial defrag doing fast checks only. Adjacency and freelist walk.
 	*
 	* @param Stats			[out] Stats
 	* @param StartTime		Start time, used for limiting the Tick() time
 	*/
-	void			PartialDefragmentation(FRelocationStats& Stats, double StartTime);
+	void			PartialDefragmentationFast(FRelocationStats& Stats, double StartTime);
+
+	/**
+	* Performs a partial defrag doing slow all chunk search to find used chunks to move that are surrounded by other used chunks
+	* That a freechunk walk won't find.
+	*
+	* @param Stats			[out] Stats
+	* @param StartTime		Start time, used for limiting the Tick() time
+	*/
+	void			PartialDefragmentationSlow(FRelocationStats& Stats, double StartTime);
 
 	/**
 	* Performs a partial defrag by shifting down memory to fill holes, in a brute-force manner.
@@ -782,7 +800,7 @@ protected:
 	* @param GrowAmount	Number of uint8s to grow by
 	* @return				nullptr if it failed, otherwise the new grown chunk
 	*/
-	FMemoryChunk*	Grow(FMemoryChunk* Chunk, int32 GrowAmount);
+	FMemoryChunk*	Grow(FMemoryChunk* Chunk, int64 GrowAmount);
 
 	/**
 	* Immediately shrinks a memory chunk by moving the base address, without relocating any memory.
@@ -792,7 +810,7 @@ protected:
 	* @param ShrinkAmount	Number of uint8s to shrink by
 	* @return				The new shrunken chunk
 	*/
-	FMemoryChunk*	Shrink(FMemoryChunk* Chunk, int32 ShrinkAmount);
+	FMemoryChunk*	Shrink(FMemoryChunk* Chunk, int64 ShrinkAmount);
 
 	/**
 	* Checks the int32ernal state for errors. (Slow)
@@ -835,7 +853,7 @@ protected:
 	* @param FirstSize			New size of first chunk
 	* @param bSortedFreeList	If true, maint32ains the free-list order
 	*/
-	void			Split(FMemoryChunk* BaseChunk, int32 FirstSize, bool bSortedFreeList)
+	void			Split(FMemoryChunk* BaseChunk, int64 FirstSize)
 	{
 		check(BaseChunk);
 		check(FirstSize < BaseChunk->Size);
@@ -849,17 +867,17 @@ protected:
 		// 		check( !BaseChunk->PreviousChunk || !BaseChunk->PreviousChunk->bIsAvailable || !BaseChunk->bIsAvailable );
 
 		// Calculate size of second chunk...
-		int32 SecondSize = BaseChunk->Size - FirstSize;
+		int64 SecondSize = BaseChunk->Size - FirstSize;
 		// ... and create it.
 
 		//todo: fix stats
 		//ensureMsgf(BaseChunk->Stat.IsNone(), TEXT("Free chunk has stat"));
-		FMemoryChunk* NewFreeChunk = new FMemoryChunk(BaseChunk->Base + FirstSize, SecondSize, *this, BaseChunk, TStatId(), bSortedFreeList);
+		FMemoryChunk* NewFreeChunk = new FMemoryChunk(BaseChunk->Base + FirstSize, SecondSize, *this, BaseChunk, TStatId());
 
 		// Keep the original sync index for the new chunk if necessary.
 		if (BaseChunk->IsRelocating() && BaseChunk->SyncSize > FirstSize)
 		{
-			int32 SecondSyncSize = BaseChunk->SyncSize - FirstSize;
+			int64 SecondSyncSize = BaseChunk->SyncSize - FirstSize;
 			NewFreeChunk->SetSyncIndex(BaseChunk->SyncIndex, SecondSyncSize);
 		}
 		BaseChunk->SetSyncIndex(BaseChunk->SyncIndex, FMath::Min((int64)FirstSize, BaseChunk->SyncSize));
@@ -877,16 +895,15 @@ protected:
 	* @param bAsync			If true, allows allocating from relocating chunks and maint32ains the free-list sort order.
 	* @return					The memory chunk that was allocated (the original chunk could've been split).
 	*/
-	FMemoryChunk*	AllocateChunk(FMemoryChunk* FreeChunk, int32 AllocationSize, bool bAsync, bool bDoValidation = true);
+	FMemoryChunk*	AllocateChunk(FMemoryChunk* FreeChunk, int64 AllocationSize, bool bAsync, bool bDoValidation = true);
 
 	/**
 	* Marks the specified chunk as 'free' and updates tracking variables.
 	* Calls LinkFreeChunk() to coalesce adjacent free memory.
 	*
-	* @param Chunk						Chunk to free
-	* @param bMaint32ainSortedFreelist	If true, maint32ains the free-list sort order
+	* @param Chunk						Chunk to free	
 	*/
-	void			FreeChunk(FMemoryChunk* Chunk, bool bMaint32ainSortedFreelist);
+	void			FreeChunk(FMemoryChunk* Chunk);
 
 	/**
 	* Frees the passed in chunk and coalesces adjacent free chunks into 'Chunk' if possible.
@@ -895,11 +912,11 @@ protected:
 	* @param Chunk				Chunk to mark as available.
 	* @param bSortedFreeList	If true, maintains the free-list sort order
 	*/
-	void			LinkFreeChunk(FMemoryChunk* Chunk, bool bSortedFreeList)
+	void			LinkFreeChunk(FMemoryChunk* Chunk)
 	{
 		check(Chunk);
 		// Mark chunk as available.
-		Chunk->LinkFree(bSortedFreeList, nullptr);
+		Chunk->LinkFree(nullptr);
 		// Kick of merge pass.
 		Coalesce(Chunk);
 	}
@@ -922,7 +939,7 @@ protected:
 
 	/**
 	* Defrag helper function. Checks if the specified allocation fits within
-	* the adjacent free chunk(s), accounting for potential reallocation request.
+	* the adjacent free chunk(s).
 	*
 	* @param UsedChunk		Allocated chunk to check for a fit
 	* @param bAnyChunkType	If false, only succeeds if 'UsedChunk' has a reallocation request and fits
@@ -981,12 +998,15 @@ protected:
 	/** Cumulative time spent in allocator.							*/
 	double			TimeSpentInAllocator;
 	/** Allocated memory in uint8s.									*/
-#if WINVER < 0x0600
+#if PLATFORM_WINDOWS && (WINVER < 0x0600)
 	// Interlock...64 functions are only available from Vista onwards
 	typedef int32 memsize_t;
 #else
 	typedef int64 memsize_t;
 #endif
+
+	volatile memsize_t	PaddingWasteSize;
+
 	volatile memsize_t	AllocatedMemorySize;
 	/** Available memory in uint8s.									*/
 	volatile memsize_t	AvailableMemorySize;
diff --git a/Engine/Source/Runtime/Renderer/Private/ClearQuad.cpp b/Engine/Source/Runtime/Renderer/Private/ClearQuad.cpp
--- a/Engine/Source/Runtime/Renderer/Private/ClearQuad.cpp
+++ b/Engine/Source/Runtime/Renderer/Private/ClearQuad.cpp
@@ -5,8 +5,7 @@
 
 FGlobalBoundShaderState GClearMRTBoundShaderState[8];
 
-// TODO support ExcludeRect
-void DrawClearQuadMRT(FRHICommandList& RHICmdList, ERHIFeatureLevel::Type FeatureLevel, bool bClearColor, int32 NumClearColors, const FLinearColor* ClearColorArray, bool bClearDepth, float Depth, bool bClearStencil, uint32 Stencil)
+static void ClearQuadSetup( FRHICommandList& RHICmdList, ERHIFeatureLevel::Type FeatureLevel, bool bClearColor, int32 NumClearColors, const FLinearColor* ClearColorArray, bool bClearDepth, float Depth, bool bClearStencil, uint32 Stencil )
 {
 	// Set new states
 	FBlendStateRHIParamRef BlendStateRHI;
@@ -100,51 +99,67 @@ void DrawClearQuadMRT(FRHICommandList& RHICmdList, ERHIFeatureLevel::Type Featur
 
 	SetGlobalBoundShaderState(RHICmdList, FeatureLevel, GClearMRTBoundShaderState[FMath::Max(NumClearColors - 1, 0)], GetVertexDeclarationFVector4(), *VertexShader, PixelShader);	
 	PixelShader->SetColors(RHICmdList, ClearColorArray, NumClearColors);
-		
+}
+
+void DrawClearQuadMRT( FRHICommandList& RHICmdList, ERHIFeatureLevel::Type FeatureLevel, bool bClearColor, int32 NumClearColors, const FLinearColor* ClearColorArray, bool bClearDepth, float Depth, bool bClearStencil, uint32 Stencil )
+{
+	ClearQuadSetup( RHICmdList, FeatureLevel, bClearColor, NumClearColors, ClearColorArray, bClearDepth, Depth, bClearStencil, Stencil );
+
+	// without a hole
+	FVector4 Vertices[4];
+	Vertices[0].Set( -1.0f,  1.0f, Depth, 1.0f );
+	Vertices[1].Set(  1.0f,  1.0f, Depth, 1.0f );
+	Vertices[2].Set( -1.0f, -1.0f, Depth, 1.0f );
+	Vertices[3].Set(  1.0f, -1.0f, Depth, 1.0f );
+	DrawPrimitiveUP(RHICmdList, PT_TriangleStrip, 2, Vertices, sizeof(Vertices[0]));
+}
+
+void DrawClearQuadMRT( FRHICommandList& RHICmdList, ERHIFeatureLevel::Type FeatureLevel, bool bClearColor, int32 NumClearColors, const FLinearColor* ClearColorArray, bool bClearDepth, float Depth, bool bClearStencil, uint32 Stencil, FIntPoint ViewSize, FIntRect ExcludeRect )
+{
+	ClearQuadSetup( RHICmdList, FeatureLevel, bClearColor, NumClearColors, ClearColorArray, bClearDepth, Depth, bClearStencil, Stencil );
+
+	// Draw a fullscreen quad
+	if(ExcludeRect.Width() > 0 && ExcludeRect.Height() > 0)
 	{
-		// Draw a fullscreen quad
-		/*if(ExcludeRect.Width() > 0 && ExcludeRect.Height() > 0)
-		{
-			// with a hole in it (optimization in case the hardware has non constant clear performance)
-			FVector4 OuterVertices[4];
-			OuterVertices[0].Set( -1.0f,  1.0f, Depth, 1.0f );
-			OuterVertices[1].Set(  1.0f,  1.0f, Depth, 1.0f );
-			OuterVertices[2].Set(  1.0f, -1.0f, Depth, 1.0f );
-			OuterVertices[3].Set( -1.0f, -1.0f, Depth, 1.0f );
-
-			float InvViewWidth = 1.0f / Viewport.Width;
-			float InvViewHeight = 1.0f / Viewport.Height;
-			FVector4 FractionRect = FVector4(ExcludeRect.Min.X * InvViewWidth, ExcludeRect.Min.Y * InvViewHeight, (ExcludeRect.Max.X - 1) * InvViewWidth, (ExcludeRect.Max.Y - 1) * InvViewHeight);
-
-			FVector4 InnerVertices[4];
-			InnerVertices[0].Set( FMath::Lerp(-1.0f,  1.0f, FractionRect.X), FMath::Lerp(1.0f, -1.0f, FractionRect.Y), Depth, 1.0f );
-			InnerVertices[1].Set( FMath::Lerp(-1.0f,  1.0f, FractionRect.Z), FMath::Lerp(1.0f, -1.0f, FractionRect.Y), Depth, 1.0f );
-			InnerVertices[2].Set( FMath::Lerp(-1.0f,  1.0f, FractionRect.Z), FMath::Lerp(1.0f, -1.0f, FractionRect.W), Depth, 1.0f );
-			InnerVertices[3].Set( FMath::Lerp(-1.0f,  1.0f, FractionRect.X), FMath::Lerp(1.0f, -1.0f, FractionRect.W), Depth, 1.0f );
+		// with a hole in it
+		FVector4 OuterVertices[4];
+		OuterVertices[0].Set( -1.0f,  1.0f, Depth, 1.0f );
+		OuterVertices[1].Set(  1.0f,  1.0f, Depth, 1.0f );
+		OuterVertices[2].Set(  1.0f, -1.0f, Depth, 1.0f );
+		OuterVertices[3].Set( -1.0f, -1.0f, Depth, 1.0f );
+
+		float InvViewWidth = 1.0f / ViewSize.X;
+		float InvViewHeight = 1.0f / ViewSize.Y;
+		FVector4 FractionRect = FVector4(ExcludeRect.Min.X * InvViewWidth, ExcludeRect.Min.Y * InvViewHeight, (ExcludeRect.Max.X - 1) * InvViewWidth, (ExcludeRect.Max.Y - 1) * InvViewHeight);
+
+		FVector4 InnerVertices[4];
+		InnerVertices[0].Set( FMath::Lerp(-1.0f,  1.0f, FractionRect.X), FMath::Lerp(1.0f, -1.0f, FractionRect.Y), Depth, 1.0f );
+		InnerVertices[1].Set( FMath::Lerp(-1.0f,  1.0f, FractionRect.Z), FMath::Lerp(1.0f, -1.0f, FractionRect.Y), Depth, 1.0f );
+		InnerVertices[2].Set( FMath::Lerp(-1.0f,  1.0f, FractionRect.Z), FMath::Lerp(1.0f, -1.0f, FractionRect.W), Depth, 1.0f );
+		InnerVertices[3].Set( FMath::Lerp(-1.0f,  1.0f, FractionRect.X), FMath::Lerp(1.0f, -1.0f, FractionRect.W), Depth, 1.0f );
 				
-			FVector4 Vertices[10];
-			Vertices[0] = OuterVertices[0];
-			Vertices[1] = InnerVertices[0];
-			Vertices[2] = OuterVertices[1];
-			Vertices[3] = InnerVertices[1];
-			Vertices[4] = OuterVertices[2];
-			Vertices[5] = InnerVertices[2];
-			Vertices[6] = OuterVertices[3];
-			Vertices[7] = InnerVertices[3];
-			Vertices[8] = OuterVertices[0];
-			Vertices[9] = InnerVertices[0];
-
-			DrawPrimitiveUP(RHICmdList, PT_TriangleStrip, 8, Vertices, sizeof(Vertices[0]) );
-		}
-		else*/
-		{
-			// without a hole
-			FVector4 Vertices[4];
-			Vertices[0].Set( -1.0f,  1.0f, Depth, 1.0f );
-			Vertices[1].Set(  1.0f,  1.0f, Depth, 1.0f );
-			Vertices[2].Set( -1.0f, -1.0f, Depth, 1.0f );
-			Vertices[3].Set(  1.0f, -1.0f, Depth, 1.0f );
-			DrawPrimitiveUP(RHICmdList, PT_TriangleStrip, 2, Vertices, sizeof(Vertices[0]));
-		}
+		FVector4 Vertices[10];
+		Vertices[0] = OuterVertices[0];
+		Vertices[1] = InnerVertices[0];
+		Vertices[2] = OuterVertices[1];
+		Vertices[3] = InnerVertices[1];
+		Vertices[4] = OuterVertices[2];
+		Vertices[5] = InnerVertices[2];
+		Vertices[6] = OuterVertices[3];
+		Vertices[7] = InnerVertices[3];
+		Vertices[8] = OuterVertices[0];
+		Vertices[9] = InnerVertices[0];
+
+		DrawPrimitiveUP(RHICmdList, PT_TriangleStrip, 8, Vertices, sizeof(Vertices[0]) );
+	}
+	else
+	{
+		// without a hole
+		FVector4 Vertices[4];
+		Vertices[0].Set( -1.0f,  1.0f, Depth, 1.0f );
+		Vertices[1].Set(  1.0f,  1.0f, Depth, 1.0f );
+		Vertices[2].Set( -1.0f, -1.0f, Depth, 1.0f );
+		Vertices[3].Set(  1.0f, -1.0f, Depth, 1.0f );
+		DrawPrimitiveUP(RHICmdList, PT_TriangleStrip, 2, Vertices, sizeof(Vertices[0]));
 	}
 }
diff --git a/Engine/Source/Runtime/Renderer/Private/ClearQuad.h b/Engine/Source/Runtime/Renderer/Private/ClearQuad.h
--- a/Engine/Source/Runtime/Renderer/Private/ClearQuad.h
+++ b/Engine/Source/Runtime/Renderer/Private/ClearQuad.h
@@ -2,10 +2,15 @@
 
 #pragma once
 
-// TODO support ExcludeRect
-void DrawClearQuadMRT(FRHICommandList& RHICmdList, ERHIFeatureLevel::Type FeatureLevel, bool bClearColor, int32 NumClearColors, const FLinearColor* ClearColorArray, bool bClearDepth, float Depth, bool bClearStencil, uint32 Stencil);
+void DrawClearQuadMRT( FRHICommandList& RHICmdList, ERHIFeatureLevel::Type FeatureLevel, bool bClearColor, int32 NumClearColors, const FLinearColor* ClearColorArray, bool bClearDepth, float Depth, bool bClearStencil, uint32 Stencil );
+void DrawClearQuadMRT( FRHICommandList& RHICmdList, ERHIFeatureLevel::Type FeatureLevel, bool bClearColor, int32 NumClearColors, const FLinearColor* ClearColorArray, bool bClearDepth, float Depth, bool bClearStencil, uint32 Stencil, FIntPoint ViewSize, FIntRect ExcludeRect );
 
 inline void DrawClearQuad(FRHICommandList& RHICmdList, ERHIFeatureLevel::Type FeatureLevel, bool bClearColor, const FLinearColor& Color, bool bClearDepth, float Depth, bool bClearStencil, uint32 Stencil)
 {
 	DrawClearQuadMRT(RHICmdList, FeatureLevel, bClearColor, 1, &Color, bClearDepth, Depth, bClearStencil, Stencil);
+}
+
+inline void DrawClearQuad(FRHICommandList& RHICmdList, ERHIFeatureLevel::Type FeatureLevel, bool bClearColor, const FLinearColor& Color, bool bClearDepth, float Depth, bool bClearStencil, uint32 Stencil, FIntPoint ViewSize, FIntRect ExcludeRect)
+{
+	DrawClearQuadMRT(RHICmdList, FeatureLevel, bClearColor, 1, &Color, bClearDepth, Depth, bClearStencil, Stencil, ViewSize, ExcludeRect);
 }
\ No newline at end of file
diff --git a/Engine/Source/Runtime/Renderer/Private/DeferredShadingRenderer.cpp b/Engine/Source/Runtime/Renderer/Private/DeferredShadingRenderer.cpp
--- a/Engine/Source/Runtime/Renderer/Private/DeferredShadingRenderer.cpp
+++ b/Engine/Source/Runtime/Renderer/Private/DeferredShadingRenderer.cpp
@@ -768,29 +768,47 @@ void FDeferredShadingSceneRenderer::RenderOcclusion(FRHICommandListImmediate& RH
 {		
 	if (bRenderQueries || bRenderHZB)
 	{
+		FSceneRenderTargets& SceneContext = FSceneRenderTargets::Get(RHICmdList);
+
 		{
 			// Update the quarter-sized depth buffer with the current contents of the scene depth texture.
 			// This needs to happen before occlusion tests, which makes use of the small depth buffer.
 			SCOPE_CYCLE_COUNTER(STAT_FDeferredShadingSceneRenderer_UpdateDownsampledDepthSurface);
 			UpdateDownsampledDepthSurface(RHICmdList);
 		}
-
 		
+		// Issue occlusion queries
+		// This is done after the downsampled depth buffer is created so that it can be used for issuing queries
+		BeginOcclusionTests(RHICmdList, bRenderQueries);
+
 		if (bRenderHZB)
-		{			
+		{
+			RHICmdList.TransitionResource( EResourceTransitionAccess::EReadable, SceneContext.GetSceneDepthSurface() );
+
 			static const auto ICVarAO		= IConsoleManager::Get().FindTConsoleVariableDataInt(TEXT("r.AmbientOcclusionLevels"));
 			static const auto ICVarHZBOcc	= IConsoleManager::Get().FindConsoleVariable(TEXT("r.HZBOcclusion"));
 			bool bSSAO						= ICVarAO->GetValueOnRenderThread() != 0;
-			bool bHzbOcclusion				= ICVarHZBOcc->GetInt() != 0;
+			bool bHZBOcclusion				= ICVarHZBOcc->GetInt() != 0;
 
 			for (int32 ViewIndex = 0; ViewIndex < Views.Num(); ViewIndex++)
 			{
-				const uint32 bSSR = DoScreenSpaceReflections(Views[ViewIndex]);
+				FViewInfo& View = Views[ViewIndex];
+				FSceneViewState* ViewState = (FSceneViewState*)View.State;
+				
+				const uint32 bSSR = DoScreenSpaceReflections( View );
 				
-				if (bSSAO || bHzbOcclusion || bSSR)
+				if (bSSAO || bHZBOcclusion || bSSR)
 				{
 					BuildHZB(RHICmdList, Views[ViewIndex]);
 				}
+
+				if (bHZBOcclusion && ViewState && ViewState->HZBOcclusionTests.GetNum() != 0)
+				{
+					check(ViewState->HZBOcclusionTests.IsValidFrame(ViewState->OcclusionFrameCounter));
+
+					SCOPED_DRAW_EVENT(RHICmdList, HZB);
+					ViewState->HZBOcclusionTests.Submit(RHICmdList, View);
+				}
 			}
 
 			//async ssao only requires HZB and depth as inputs so get started ASAP
@@ -800,11 +818,6 @@ void FDeferredShadingSceneRenderer::RenderOcclusion(FRHICommandListImmediate& RH
 			}
 		}
 
-		// Issue occlusion queries
-		// This is done after the downsampled depth buffer is created so that it can be used for issuing queries
-		BeginOcclusionTests(RHICmdList, bRenderQueries, bRenderHZB);
-
-
 		// Hint to the RHI to submit commands up to this point to the GPU if possible.  Can help avoid CPU stalls next frame waiting
 		// for these query results on some platforms.
 		RHICmdList.SubmitCommandsHint();
@@ -1043,10 +1056,6 @@ void FDeferredShadingSceneRenderer::Render(FRHICommandListImmediate& RHICmdList)
 	bool bOcclusionBeforeBasePass = (CVarOcclusionQueryLocation.GetValueOnRenderThread() == 1) && bNeedsPrePass;	
 	bool bHZBBeforeBasePass = bOcclusionBeforeBasePass && (EarlyZPassMode == EDepthDrawingMode::DDM_AllOccluders);
 
-	if (bHZBBeforeBasePass && GCompositionLighting.CanProcessAsyncSSAO(Views))
-	{
-		SceneContext.ResolveSceneDepthTexture(RHICmdList);
-	}
 	RenderOcclusion(RHICmdList, bOcclusionBeforeBasePass, bHZBBeforeBasePass);
 	ServiceLocalQueue();
 
@@ -2131,6 +2140,8 @@ void FDeferredShadingSceneRenderer::UpdateDownsampledDepthSurface(FRHICommandLis
 	FSceneRenderTargets& SceneContext = FSceneRenderTargets::Get(RHICmdList);
 	if (SceneContext.UseDownsizedOcclusionQueries() && (FeatureLevel >= ERHIFeatureLevel::SM4))
 	{
+		RHICmdList.TransitionResource( EResourceTransitionAccess::EReadable, SceneContext.GetSceneDepthSurface() );
+
 		for (int32 ViewIndex = 0; ViewIndex < Views.Num(); ViewIndex++)
 		{
 			const FViewInfo& View = Views[ViewIndex];
diff --git a/Engine/Source/Runtime/Renderer/Private/DeferredShadingRenderer.h b/Engine/Source/Runtime/Renderer/Private/DeferredShadingRenderer.h
--- a/Engine/Source/Runtime/Renderer/Private/DeferredShadingRenderer.h
+++ b/Engine/Source/Runtime/Renderer/Private/DeferredShadingRenderer.h
@@ -181,7 +181,7 @@ private:
 	bool RenderPrePassHMD(FRHICommandListImmediate& RHICmdList);
 
 	/** Issues occlusion queries. */
-	void BeginOcclusionTests(FRHICommandListImmediate& RHICmdList, bool bRenderQueries, bool bRenderHZB);
+	void BeginOcclusionTests(FRHICommandListImmediate& RHICmdList, bool bRenderQueries);
 
 	/** Renders the scene's fogging. */
 	bool RenderFog(FRHICommandListImmediate& RHICmdList, const FLightShaftsOutput& LightShaftsOutput);
diff --git a/Engine/Source/Runtime/Renderer/Private/DistortionRendering.cpp b/Engine/Source/Runtime/Renderer/Private/DistortionRendering.cpp
--- a/Engine/Source/Runtime/Renderer/Private/DistortionRendering.cpp
+++ b/Engine/Source/Runtime/Renderer/Private/DistortionRendering.cpp
@@ -850,7 +850,7 @@ void FSceneRenderer::RenderDistortion(FRHICommandListImmediate& RHICmdList)
 
 		// Create a texture to store the resolved light attenuation values, and a render-targetable surface to hold the unresolved light attenuation values.
 		{
-			FPooledRenderTargetDesc Desc(FPooledRenderTargetDesc::Create2DDesc(SceneContext.GetBufferSizeXY(), PF_B8G8R8A8, FClearValueBinding::None, TexCreate_None, TexCreate_RenderTargetable, false));
+			FPooledRenderTargetDesc Desc(FPooledRenderTargetDesc::Create2DDesc(SceneContext.GetBufferSizeXY(), PF_B8G8R8A8, FClearValueBinding::Transparent, TexCreate_None, TexCreate_RenderTargetable, false));
 			Desc.Flags |= TexCreate_FastVRAM;
 			GRenderTargetPool.FindFreeElement(RHICmdList, Desc, DistortionRT, TEXT("Distortion"));
 
@@ -864,7 +864,11 @@ void FSceneRenderer::RenderDistortion(FRHICommandListImmediate& RHICmdList)
 		// DistortionRT==0 should never happen but better we don't crash
 		if(DistortionRT)
 		{
-			SetRenderTarget(RHICmdList, DistortionRT->GetRenderTargetItem().TargetableTexture, SceneContext.GetSceneDepthSurface(), ESimpleRenderTargetMode::EUninitializedColorExistingDepth, FExclusiveDepthStencil::DepthRead_StencilWrite);
+			FRHIRenderTargetView ColorView( DistortionRT->GetRenderTargetItem().TargetableTexture, 0, -1, ERenderTargetLoadAction::EClear, ERenderTargetStoreAction::EStore );
+			FRHIDepthRenderTargetView DepthView( SceneContext.GetSceneDepthSurface(), ERenderTargetLoadAction::ELoad, ERenderTargetStoreAction::ENoAction, FExclusiveDepthStencil::DepthRead_StencilWrite );
+			FRHISetRenderTargetsInfo Info( 1, &ColorView, DepthView );		
+
+			RHICmdList.SetRenderTargetsAndClear(Info);
 
 			for(int32 ViewIndex = 0;ViewIndex < Views.Num();ViewIndex++)
 			{
@@ -874,9 +878,6 @@ void FSceneRenderer::RenderDistortion(FRHICommandListImmediate& RHICmdList)
 				// viewport to match view size
 				RHICmdList.SetViewport(View.ViewRect.Min.X, View.ViewRect.Min.Y, 0.0f, View.ViewRect.Max.X, View.ViewRect.Max.Y, 1.0f);
 
-				// clear offsets to 0
-				RHICmdList.Clear(true, FLinearColor(0, 0, 0, 0), false, (float)ERHIZBuffer::FarPlane, false, 0, FIntRect());
-
 				// test against depth and write stencil mask
 				RHICmdList.SetDepthStencilState(TStaticDepthStencilState<false, CF_DepthNearOrEqual, true, CF_Always, SO_Keep, SO_Keep, SO_Replace, false, CF_Always, SO_Keep, SO_Keep, SO_Keep, StencilMaskBit, StencilMaskBit>::GetRHI(), StencilMaskBit);
 
@@ -889,8 +890,8 @@ void FSceneRenderer::RenderDistortion(FRHICommandListImmediate& RHICmdList)
 
 			if (bDirty)
 			{
-				// resolve using the current ResolveParams 
-				RHICmdList.CopyToResolveTarget(DistortionRT->GetRenderTargetItem().TargetableTexture, DistortionRT->GetRenderTargetItem().ShaderResourceTexture, false, FResolveParams());
+				// Ideally we skip the EliminateFastClear since we don't need pixels with no stencil set to be cleared
+				RHICmdList.TransitionResource( EResourceTransitionAccess::EReadable, DistortionRT->GetRenderTargetItem().TargetableTexture );
 				// to be able to observe results with VisualizeTexture
 				GRenderTargetPool.VisualizeTexture.SetCheckPoint(RHICmdList, DistortionRT);
 			}
diff --git a/Engine/Source/Runtime/Renderer/Private/PostProcess/PostProcessDownsample.cpp b/Engine/Source/Runtime/Renderer/Private/PostProcess/PostProcessDownsample.cpp
--- a/Engine/Source/Runtime/Renderer/Private/PostProcess/PostProcessDownsample.cpp
+++ b/Engine/Source/Runtime/Renderer/Private/PostProcess/PostProcessDownsample.cpp
@@ -198,6 +198,18 @@ void FRCPassPostProcessDownsample::Process(FRenderingCompositePassContext& Conte
 
 	Context.SetViewportAndCallRHI(0, 0, 0.0f, DestSize.X, DestSize.Y, 1.0f );
 
+	// check if we have to clear the whole surface.
+	// Otherwise perform the clear when the dest rectangle has been computed.
+	auto FeatureLevel = Context.View.GetFeatureLevel();
+	if (FeatureLevel == ERHIFeatureLevel::ES2 || FeatureLevel == ERHIFeatureLevel::ES3_1)
+	{
+		Context.RHICmdList.Clear(true, FLinearColor(0, 0, 0, 0), false, 1.0f, false, 0, FIntRect());
+	}
+	else
+	{
+		DrawClearQuad(Context.RHICmdList, Context.GetFeatureLevel(), true, FLinearColor(0, 0, 0, 0), false, 1.0f, false, 0, DestSize, DestRect);
+	}
+
 	// set the state
 	Context.RHICmdList.SetBlendState(TStaticBlendState<>::GetRHI());
 	Context.RHICmdList.SetRasterizerState(TStaticRasterizerState<>::GetRHI());
@@ -224,26 +236,10 @@ void FRCPassPostProcessDownsample::Process(FRenderingCompositePassContext& Conte
 			SetShader<1>(Context, InputDesc);
 			InflateSize = 2;
 		}
-	}
-
-	bool bHasCleared = false;
-
-	// check if we have to clear the whole surface.
-	// Otherwise perform the clear when the dest rectangle has been computed.
-	auto FeatureLevel = Context.View.GetFeatureLevel();
-	if (FeatureLevel == ERHIFeatureLevel::ES2 || FeatureLevel == ERHIFeatureLevel::ES3_1)
-	{
-		Context.RHICmdList.Clear(true, FLinearColor(0, 0, 0, 0), false, 1.0f, false, 0, FIntRect());
-		bHasCleared = true;
-	}
+	}	
 
 	TShaderMapRef<FPostProcessDownsampleVS> VertexShader(Context.GetShaderMap());
 
-	if (!bHasCleared)
-	{
-		Context.RHICmdList.Clear(true, FLinearColor(0, 0, 0, 0), false, 1.0f, false, 0, DestRect);
-	}
-
 	DrawPostProcessPass(
 		Context.RHICmdList,
 		DestRect.Min.X, DestRect.Min.Y,
diff --git a/Engine/Source/Runtime/Renderer/Private/PostProcess/PostProcessWeightedSampleSum.cpp b/Engine/Source/Runtime/Renderer/Private/PostProcess/PostProcessWeightedSampleSum.cpp
--- a/Engine/Source/Runtime/Renderer/Private/PostProcess/PostProcessWeightedSampleSum.cpp
+++ b/Engine/Source/Runtime/Renderer/Private/PostProcess/PostProcessWeightedSampleSum.cpp
@@ -560,6 +560,21 @@ void FRCPassPostProcessWeightedSampleSum::Process(FRenderingCompositePassContext
 
 	Context.SetViewportAndCallRHI(0, 0, 0.0f, DestSize.X, DestSize.Y, 1.0f);
 
+	bool bRequiresClear = true;
+	// check if we have to clear the whole surface.
+	// Otherwise perform the clear when the dest rectangle has been computed.
+	if (FeatureLevel == ERHIFeatureLevel::ES2 || FeatureLevel == ERHIFeatureLevel::ES3_1)
+	{
+		Context.RHICmdList.Clear(true, FLinearColor(0, 0, 0, 0), false, 1.0f, false, 0, FIntRect());
+		bRequiresClear = false;
+	}
+
+	FIntRect SrcRect = FIntRect::DivideAndRoundUp(View.ViewRect, SrcScaleFactor);
+	if (bRequiresClear)
+	{
+		DrawClear(Context.RHICmdList, FeatureLevel, bDoFastBlur, SrcRect, DestRect, DestSize);
+	}
+
 	// set the state
 	Context.RHICmdList.SetBlendState(TStaticBlendState<>::GetRHI());
 	Context.RHICmdList.SetRasterizerState(TStaticRasterizerState<>::GetRHI());
@@ -620,20 +635,9 @@ void FRCPassPostProcessWeightedSampleSum::Process(FRenderingCompositePassContext
 		BlurWeights,
 		NumSamples,
 		&VertexShader
-		);
-
-	bool bRequiresClear = true;
-	// check if we have to clear the whole surface.
-	// Otherwise perform the clear when the dest rectangle has been computed.
-	if (FeatureLevel == ERHIFeatureLevel::ES2 || FeatureLevel == ERHIFeatureLevel::ES3_1)
-	{
-		Context.RHICmdList.Clear(true, FLinearColor(0, 0, 0, 0), false, 1.0f, false, 0, FIntRect());
-		bRequiresClear = false;
-	}
-
-	FIntRect SrcRect =  FIntRect::DivideAndRoundUp(View.ViewRect, SrcScaleFactor);
+		);	
 
-	DrawQuad(Context.RHICmdList, bDoFastBlur, SrcRect, DestRect, bRequiresClear, DestSize, SrcSize, VertexShader);
+	DrawQuad(Context.RHICmdList, FeatureLevel, bDoFastBlur, SrcRect, DestRect, DestSize, SrcSize, VertexShader);
 
 	Context.RHICmdList.CopyToResolveTarget(DestRenderTarget.TargetableTexture, DestRenderTarget.ShaderResourceTexture, false, FResolveParams());
 }
@@ -721,26 +725,36 @@ bool FRCPassPostProcessWeightedSampleSum::DoFastBlur() const
 	return bRet;
 }
 
-void FRCPassPostProcessWeightedSampleSum::DrawQuad(FRHICommandListImmediate& RHICmdList, bool bDoFastBlur, FIntRect SrcRect, FIntRect DestRect, bool bRequiresClear, FIntPoint DestSize, FIntPoint SrcSize, FShader* VertexShader) const
+void FRCPassPostProcessWeightedSampleSum::AdjustRectsForFastBlur(FIntRect& SrcRect, FIntRect& DestRect) const
+{	
+	if (FilterShape == EFS_Horiz)
+	{
+		SrcRect.Min.X = DestRect.Min.X * 2;
+		SrcRect.Max.X = DestRect.Max.X * 2;
+	}
+	else
+	{
+		DestRect.Min.X = SrcRect.Min.X * 2;
+		DestRect.Max.X = SrcRect.Max.X * 2;
+	}	
+}
+
+void FRCPassPostProcessWeightedSampleSum::DrawClear(FRHICommandListImmediate& RHICmdList, ERHIFeatureLevel::Type FeatureLevel, bool bDoFastBlur, FIntRect SrcRect, FIntRect DestRect, FIntPoint DestSize) const
 {
 	if (bDoFastBlur)
 	{
-		if (FilterShape == EFS_Horiz)
-		{
-			SrcRect.Min.X = DestRect.Min.X * 2;
-			SrcRect.Max.X = DestRect.Max.X * 2;
-		}
-		else
-		{
-			DestRect.Min.X = SrcRect.Min.X * 2;
-			DestRect.Max.X = SrcRect.Max.X * 2;
-		}
+		AdjustRectsForFastBlur(SrcRect, DestRect);
 	}
 
-	if (bRequiresClear)
+	DrawClearQuad(RHICmdList, FeatureLevel, true, FLinearColor(0, 0, 0, 0), false, 1.0f, false, 0, DestSize, DestRect);	
+}
+
+void FRCPassPostProcessWeightedSampleSum::DrawQuad(FRHICommandListImmediate& RHICmdList, ERHIFeatureLevel::Type FeatureLevel, bool bDoFastBlur, FIntRect SrcRect, FIntRect DestRect, FIntPoint DestSize, FIntPoint SrcSize, FShader* VertexShader) const
+{	
+	if (bDoFastBlur)
 	{
-		RHICmdList.Clear(true, FLinearColor(0, 0, 0, 0), false, 1.0f, false, 0, DestRect);
-	}
+		AdjustRectsForFastBlur(SrcRect, DestRect);
+	}	
 
 	// Draw a quad mapping scene color to the view's render target
 	DrawRectangle(
diff --git a/Engine/Source/Runtime/Renderer/Private/PostProcess/PostProcessWeightedSampleSum.h b/Engine/Source/Runtime/Renderer/Private/PostProcess/PostProcessWeightedSampleSum.h
--- a/Engine/Source/Runtime/Renderer/Private/PostProcess/PostProcessWeightedSampleSum.h
+++ b/Engine/Source/Runtime/Renderer/Private/PostProcess/PostProcessWeightedSampleSum.h
@@ -125,7 +125,9 @@ public:
 	void SetCrossCenterWeight(float InCrossCenterWeight) { check(InCrossCenterWeight >= 0.0f); CrossCenterWeight = InCrossCenterWeight; }
 
 private:
-	void DrawQuad(FRHICommandListImmediate& RHICmdList, bool bDoFastBlur, FIntRect SrcRect, FIntRect DestRect, bool bRequiresClear, FIntPoint DestSize, FIntPoint SrcSize, FShader* VertexShader) const;
+	void AdjustRectsForFastBlur(FIntRect& SrcRect, FIntRect& DestRect) const;
+	void DrawClear(FRHICommandListImmediate& RHICmdList, ERHIFeatureLevel::Type FeatureLevel, bool bDoFastBlur, FIntRect SrcRect, FIntRect DestRect, FIntPoint DestSize) const;
+	void DrawQuad(FRHICommandListImmediate& RHICmdList, ERHIFeatureLevel::Type FeatureLevel, bool bDoFastBlur, FIntRect SrcRect, FIntRect DestRect, FIntPoint DestSize, FIntPoint SrcSize, FShader* VertexShader) const;
 	static uint32 GetMaxNumSamples(ERHIFeatureLevel::Type InFeatureLevel);
 
 	// e.g. EFS_Horiz or EFS_Vert
diff --git a/Engine/Source/Runtime/Renderer/Private/SceneOcclusion.cpp b/Engine/Source/Runtime/Renderer/Private/SceneOcclusion.cpp
--- a/Engine/Source/Runtime/Renderer/Private/SceneOcclusion.cpp
+++ b/Engine/Source/Runtime/Renderer/Private/SceneOcclusion.cpp
@@ -1078,7 +1078,7 @@ void BuildHZB( FRHICommandListImmediate& RHICmdList, FViewInfo& View )
 	GRenderTargetPool.VisualizeTexture.SetCheckPoint( RHICmdList, View.HZB );
 }
 
-void FDeferredShadingSceneRenderer::BeginOcclusionTests(FRHICommandListImmediate& RHICmdList, bool bRenderQueries, bool bRenderHZB)
+void FDeferredShadingSceneRenderer::BeginOcclusionTests(FRHICommandListImmediate& RHICmdList, bool bRenderQueries)
 {
 	SCOPE_CYCLE_COUNTER(STAT_BeginOcclusionTestsTime);
 	int32 NumBufferedFrames = FOcclusionQueryHelpers::GetNumBufferedFrames();
@@ -1231,23 +1231,6 @@ void FDeferredShadingSceneRenderer::BeginOcclusionTests(FRHICommandListImmediate
 		RHICmdList.EndOcclusionQueryBatch();
 	}
 
-	if (bRenderHZB)
-	{
-		for (int32 ViewIndex = 0; ViewIndex < Views.Num(); ViewIndex++)
-		{
-			FViewInfo& View = Views[ViewIndex];
-			FSceneViewState* ViewState = (FSceneViewState*)View.State;
-
-			if (ViewState && ViewState->HZBOcclusionTests.GetNum() != 0)
-			{
-				check(ViewState->HZBOcclusionTests.IsValidFrame(ViewState->OcclusionFrameCounter));
-
-				SCOPED_DRAW_EVENT(RHICmdList, HZB);
-				ViewState->HZBOcclusionTests.Submit(RHICmdList, View);
-			}
-		}
-	}
-
 	if (bUseDownsampledDepth)
 	{
 		// Restore default render target
diff --git a/Engine/Source/Runtime/Renderer/Private/SceneRendering.cpp b/Engine/Source/Runtime/Renderer/Private/SceneRendering.cpp
--- a/Engine/Source/Runtime/Renderer/Private/SceneRendering.cpp
+++ b/Engine/Source/Runtime/Renderer/Private/SceneRendering.cpp
@@ -489,6 +489,7 @@ void FViewInfo::CreateUniformBuffer(
 
 	check(Family);
 	check(!DirectionalLightShadowInfo || DirectionalLightShadowInfo->Num() > 0);
+	checkfSlow(ViewRect.Area() > 0, TEXT("Invalid-size ViewRect passed to CreateUniformBuffer [%d * %d]."), ViewRect.Width(), ViewRect.Height());
 
 	// Calculate the vector used by shaders to convert clip space coordinates to texture space.
 	const FIntPoint BufferSize = SceneContext.GetBufferSizeXY();
diff --git a/Engine/Source/Runtime/Renderer/Private/ShadowSetup.cpp b/Engine/Source/Runtime/Renderer/Private/ShadowSetup.cpp
--- a/Engine/Source/Runtime/Renderer/Private/ShadowSetup.cpp
+++ b/Engine/Source/Runtime/Renderer/Private/ShadowSetup.cpp
@@ -494,20 +494,28 @@ bool FProjectedShadowInfo::SetupPerObjectProjection(
 
 		ResolutionY = FMath::Min<uint32>(FMath::TruncToInt(InResolutionX / AspectRatio), MaxShadowResolutionY);
 
-		// Store the view matrix
-		// Reorder the vectors to match the main view, since ShadowViewMatrix will be used to override the main view's view matrix during shadow depth rendering
-		ShadowViewMatrix = Initializer.WorldToLight * 
-			FMatrix(
-				FPlane(0,	0,	1,	0),
-				FPlane(1,	0,	0,	0),
-				FPlane(0,	1,	0,	0),
-				FPlane(0,	0,	0,	1));
-
-		GetViewFrustumBounds(CasterFrustum,SubjectAndReceiverMatrix,true);
-
-		InvReceiverMatrix = ReceiverMatrix.InverseFast();
-		GetViewFrustumBounds(ReceiverFrustum,ReceiverMatrix,true);
-		UpdateShaderDepthBias();
+		if (ResolutionX == 0 || ResolutionY == 0)
+		{
+			checkfSlow(0, TEXT("Tried to create shadow with invalid resolution [%d * %d]."), ResolutionX, ResolutionY);
+			bRet = false;
+		}
+		else
+		{
+			// Store the view matrix
+			// Reorder the vectors to match the main view, since ShadowViewMatrix will be used to override the main view's view matrix during shadow depth rendering
+			ShadowViewMatrix = Initializer.WorldToLight *
+				FMatrix(
+				FPlane(0, 0, 1, 0),
+				FPlane(1, 0, 0, 0),
+				FPlane(0, 1, 0, 0),
+				FPlane(0, 0, 0, 1));
+
+			GetViewFrustumBounds(CasterFrustum, SubjectAndReceiverMatrix, true);
+
+			InvReceiverMatrix = ReceiverMatrix.InverseFast();
+			GetViewFrustumBounds(ReceiverFrustum, ReceiverMatrix, true);
+			UpdateShaderDepthBias();
+		}
 	}
 
 	return bRet;
diff --git a/Engine/Source/Runtime/Slate/Private/Framework/Application/AnalogCursor.cpp b/Engine/Source/Runtime/Slate/Private/Framework/Application/AnalogCursor.cpp
--- a/Engine/Source/Runtime/Slate/Private/Framework/Application/AnalogCursor.cpp
+++ b/Engine/Source/Runtime/Slate/Private/Framework/Application/AnalogCursor.cpp
@@ -172,7 +172,7 @@ bool FAnalogCursor::HandleKeyUpEvent(FSlateApplication& SlateApp, const FKeyEven
 			SlateApp.GetPlatformApplication()->GetModifierKeys()
 			);
 
-		SlateApp.ProcessMouseButtonUpEvent(MouseEvent);
+		return SlateApp.ProcessMouseButtonUpEvent(MouseEvent);
 	}
 	
 	return false;
diff --git a/Engine/Source/Runtime/Slate/Private/Framework/Application/SlateApplication.cpp b/Engine/Source/Runtime/Slate/Private/Framework/Application/SlateApplication.cpp
--- a/Engine/Source/Runtime/Slate/Private/Framework/Application/SlateApplication.cpp
+++ b/Engine/Source/Runtime/Slate/Private/Framework/Application/SlateApplication.cpp
@@ -4497,12 +4497,14 @@ FReply FSlateApplication::RoutePointerDownEvent(FWidgetPath& WidgetsUnderPointer
 }
 
 
-void FSlateApplication::RoutePointerUpEvent(FWidgetPath& WidgetsUnderPointer, FPointerEvent& PointerEvent)
+bool FSlateApplication::RoutePointerUpEvent(FWidgetPath& WidgetsUnderPointer, FPointerEvent& PointerEvent)
 {
 #if PLATFORM_MAC
 	NSWindow* ActiveNativeWindow = [NSApp keyWindow];
 	TSharedPtr<SWindow> TopLevelWindow;
 #endif
+	FReply Reply = FReply::Unhandled();
+
 	if (MouseCaptor.HasCaptureForPointerIndex(PointerEvent.GetUserIndex(), PointerEvent.GetPointerIndex()))
 	{
 		//FWidgetPath MouseCaptorPath = MouseCaptor.ToWidgetPath(PointerEvent.GetPointerIndex());
@@ -4512,7 +4514,7 @@ void FSlateApplication::RoutePointerUpEvent(FWidgetPath& WidgetsUnderPointer, FP
 			// Switch worlds widgets in the current path
 			FScopedSwitchWorldHack SwitchWorld( MouseCaptorPath );
 
-			FReply Reply =
+			Reply =
 				FEventRouter::Route<FReply>( this, FEventRouter::FToLeafmostPolicy(MouseCaptorPath), PointerEvent, [this]( const FArrangedWidget& TargetWidget, const FPointerEvent& Event )
 				{
 					FReply TempReply = FReply::Unhandled();
@@ -4546,7 +4548,7 @@ void FSlateApplication::RoutePointerUpEvent(FWidgetPath& WidgetsUnderPointer, FP
 		TSharedPtr< FDragDropOperation > LocalDragDropContent = DragDropContent;
 		DragDropContent.Reset();
 
-		FReply Reply = FEventRouter::Route<FReply>( this, FEventRouter::FBubblePolicy(LocalWidgetsUnderCursor), PointerEvent, [&](const FArrangedWidget& CurWidget, const FPointerEvent& Event)
+		Reply = FEventRouter::Route<FReply>( this, FEventRouter::FBubblePolicy(LocalWidgetsUnderCursor), PointerEvent, [&](const FArrangedWidget& CurWidget, const FPointerEvent& Event)
 		{
 			FReply TempReply = FReply::Unhandled();
 			if (Event.IsTouchEvent())
@@ -4591,6 +4593,8 @@ void FSlateApplication::RoutePointerUpEvent(FWidgetPath& WidgetsUnderPointer, FP
 		TopLevelWindow->BringToFront(true);
 	}
 #endif
+
+	return Reply.IsEventHandled();
 }
 
 bool FSlateApplication::RoutePointerMoveEvent(const FWidgetPath& WidgetsUnderPointer, FPointerEvent& PointerEvent, bool bIsSynthetic)
@@ -4926,8 +4930,7 @@ bool FSlateApplication::ProcessMouseButtonUpEvent( FPointerEvent& MouseEvent )
 
 	// An empty widget path is passed in.  As an optimization, one will be generated only if a captured mouse event isnt routed
 	FWidgetPath EmptyPath;
-	RoutePointerUpEvent( EmptyPath, MouseEvent );
-	
+	const bool bHandled = RoutePointerUpEvent( EmptyPath, MouseEvent );
 
 	// If in responsive mode throttle, leave it on mouse up.
 	if( MouseButtonDownResponsivnessThrottle.IsValid() )
@@ -4941,7 +4944,7 @@ bool FSlateApplication::ProcessMouseButtonUpEvent( FPointerEvent& MouseEvent )
 		PlatformApplication->SetCapture( nullptr );
 	}
 
-	return true;
+	return bHandled;
 }
 
 bool FSlateApplication::OnMouseWheel( const float Delta )
diff --git a/Engine/Source/Runtime/Slate/Private/Framework/Text/TextEditHelper.cpp b/Engine/Source/Runtime/Slate/Private/Framework/Text/TextEditHelper.cpp
--- a/Engine/Source/Runtime/Slate/Private/Framework/Text/TextEditHelper.cpp
+++ b/Engine/Source/Runtime/Slate/Private/Framework/Text/TextEditHelper.cpp
@@ -323,6 +323,18 @@ FReply FTextEditHelper::OnMouseButtonDown( const FGeometry& MyGeometry, const FP
 				// should reset the selection range to the caret's position.
 				TextEditor->SetWasFocusedByLastMouseDown( true );
 			}
+			else
+			{
+				// On platforms using a virtual keyboard open the virtual keyboard again 
+				FSlateApplication& SlateApplication = FSlateApplication::Get();
+				if (FPlatformMisc::GetRequiresVirtualKeyboard())
+				{
+					if (!TextEditor->GetIsReadOnly())
+					{
+						SlateApplication.ShowVirtualKeyboard(true, InMouseEvent.GetUserIndex(), TextEditor->GetVirtualKeyboardEntry());
+					}
+				}
+			}
 
 			if( InMouseEvent.GetEffectingButton() == EKeys::LeftMouseButton )
 			{
diff --git a/Engine/Source/Runtime/Slate/Private/Widgets/Input/SEditableText.cpp b/Engine/Source/Runtime/Slate/Private/Widgets/Input/SEditableText.cpp
--- a/Engine/Source/Runtime/Slate/Private/Widgets/Input/SEditableText.cpp
+++ b/Engine/Source/Runtime/Slate/Private/Widgets/Input/SEditableText.cpp
@@ -178,9 +178,6 @@ void SEditableText::SetTextFromVirtualKeyboard(const FText& InNewText, ESetTextT
 		// fired on the game thread in Tick.
 		bTextChangedByVirtualKeyboard = true;
 	}
-
-	// Clear focus so that virtual keyboard can be reopened again my clicking on the text box
-	FSlateApplication::Get().ClearKeyboardFocus(EFocusCause::Cleared);
 }
 
 void SEditableText::RestoreOriginalText()
@@ -1164,6 +1161,11 @@ TSharedRef< SWidget > SEditableText::GetWidget()
 	return SharedThis( this );
 }
 
+TSharedRef<IVirtualKeyboardEntry> SEditableText::GetVirtualKeyboardEntry()
+{
+	return SharedThis(this);
+}
+
 FVector2D SEditableText::ComputeDesiredSize(float LayoutScaleMultiplier) const
 {
 	const TSharedRef< FSlateFontMeasure > FontMeasureService = FSlateApplication::Get().GetRenderer()->GetFontMeasureService();
diff --git a/Engine/Source/Runtime/Slate/Private/Widgets/Layout/SSafeZone.cpp b/Engine/Source/Runtime/Slate/Private/Widgets/Layout/SSafeZone.cpp
--- a/Engine/Source/Runtime/Slate/Private/Widgets/Layout/SSafeZone.cpp
+++ b/Engine/Source/Runtime/Slate/Private/Widgets/Layout/SSafeZone.cpp
@@ -1,29 +1,22 @@
 // Copyright 1998-2016 Epic Games, Inc. All Rights Reserved.
 
 #include "SlatePrivatePCH.h"
-#include "SSafeZone.h"
 
+#include "SSafeZone.h"
 
-void SSafeZone::Construct(const FArguments& InArgs)
+void SSafeZone::Construct( const FArguments& InArgs )
 {
 	SBox::Construct(SBox::FArguments()
 		.HAlign(InArgs._HAlign)
 		.VAlign(InArgs._VAlign)
-		.Padding(this, &SSafeZone::GetSafeZonePadding)
 		[
 			InArgs._Content.Widget
 		]
 	);
 
-	SetTitleSafe( InArgs._IsTitleSafe );
-
 	Padding = InArgs._Padding;
-}
 
-FMargin SSafeZone::GetSafeZonePadding() const
-{
-	// return either the TitleSafe or the ActionSafe size, added to the user padding
-	return Padding.Get() + SafeMargin;
+	SetTitleSafe( InArgs._IsTitleSafe );
 }
 
 void SSafeZone::SetTitleSafe( bool bIsTitleSafe )
@@ -40,3 +33,22 @@ void SSafeZone::SetTitleSafe( bool bIsTitleSafe )
 		SafeMargin = FMargin( Metrics.ActionSafePaddingSize.X, Metrics.ActionSafePaddingSize.Y );
 	}
 }
+
+void SSafeZone::OnArrangeChildren( const FGeometry& AllottedGeometry, FArrangedChildren& ArrangedChildren ) const
+{
+	const EVisibility& MyCurrentVisibility = this->GetVisibility();
+	if ( ArrangedChildren.Accepts( MyCurrentVisibility ) )
+	{
+		const FMargin SlotPadding               = Padding.Get() + ( SafeMargin * ( 1.f / AllottedGeometry.Scale ) );
+		AlignmentArrangeResult XAlignmentResult = AlignChild<Orient_Horizontal>( AllottedGeometry.Size.X, ChildSlot, SlotPadding );
+		AlignmentArrangeResult YAlignmentResult = AlignChild<Orient_Vertical>( AllottedGeometry.Size.Y, ChildSlot, SlotPadding );
+
+		ArrangedChildren.AddWidget(
+			AllottedGeometry.MakeChild(
+			ChildSlot.GetWidget(),
+			FVector2D( XAlignmentResult.Offset, YAlignmentResult.Offset ),
+			FVector2D( XAlignmentResult.Size, YAlignmentResult.Size )
+			)
+		);
+	}
+}
\ No newline at end of file
diff --git a/Engine/Source/Runtime/Slate/Private/Widgets/Text/SMultiLineEditableText.cpp b/Engine/Source/Runtime/Slate/Private/Widgets/Text/SMultiLineEditableText.cpp
--- a/Engine/Source/Runtime/Slate/Private/Widgets/Text/SMultiLineEditableText.cpp
+++ b/Engine/Source/Runtime/Slate/Private/Widgets/Text/SMultiLineEditableText.cpp
@@ -537,9 +537,6 @@ void SMultiLineEditableText::SetTextFromVirtualKeyboard(const FText& InNewText,
 		VirtualKeyboardTextCommitType = CommitType;
 		bTextCommittedByVirtualKeyboard = true;
 	}
-
-	// Clear focus so that virtual keyboard can be reopened again my clicking on the text box
-	FSlateApplication::Get().ClearKeyboardFocus(EFocusCause::Cleared);
 }
 
 bool SMultiLineEditableText::SetEditableText(const FText& TextToSet, const bool bForce)
@@ -2430,6 +2427,11 @@ TSharedRef< SWidget > SMultiLineEditableText::GetWidget()
 	return SharedThis( this );
 }
 
+TSharedRef< IVirtualKeyboardEntry > SMultiLineEditableText::GetVirtualKeyboardEntry()
+{
+	return SharedThis(this);
+}
+
 void SMultiLineEditableText::SummonContextMenu(const FVector2D& InLocation, TSharedPtr<SWindow> ParentWindow, const FWidgetPath& EventPath)
 {
 	if (!AllowContextMenu.Get())
diff --git a/Engine/Source/Runtime/Slate/Public/Framework/Application/SlateApplication.h b/Engine/Source/Runtime/Slate/Public/Framework/Application/SlateApplication.h
--- a/Engine/Source/Runtime/Slate/Public/Framework/Application/SlateApplication.h
+++ b/Engine/Source/Runtime/Slate/Public/Framework/Application/SlateApplication.h
@@ -1186,8 +1186,10 @@ public:
 	 *
 	 * @param WidgetsUnderPointer	The path of widgets the event is routed to.
 	 * @param PointerEvent		The event data that is is routed to the widget path
+	 *
+	 * @return True if handled, false otherwise
 	 */
-	void RoutePointerUpEvent(FWidgetPath& WidgetsUnderPointer, FPointerEvent& PointerEvent);
+	bool RoutePointerUpEvent(FWidgetPath& WidgetsUnderPointer, FPointerEvent& PointerEvent);
 
 	/**
 	 * Directly routes a pointer move event to the widgets in the specified widget path
diff --git a/Engine/Source/Runtime/Slate/Public/Widgets/Input/SEditableText.h b/Engine/Source/Runtime/Slate/Public/Widgets/Input/SEditableText.h
--- a/Engine/Source/Runtime/Slate/Public/Widgets/Input/SEditableText.h
+++ b/Engine/Source/Runtime/Slate/Public/Widgets/Input/SEditableText.h
@@ -374,6 +374,7 @@ public:
 	virtual void Undo() override;
 	virtual void Redo() override;
 	virtual TSharedRef< SWidget > GetWidget() override;
+	virtual TSharedRef<IVirtualKeyboardEntry> GetVirtualKeyboardEntry() override;
 	virtual void SummonContextMenu(const FVector2D& InLocation, TSharedPtr<SWindow> ParentWindow, const FWidgetPath& EventPath) override;
 	virtual void LoadText() override;
 	//~ End ITextEditorWidget Interface
diff --git a/Engine/Source/Runtime/Slate/Public/Widgets/Layout/SSafeZone.h b/Engine/Source/Runtime/Slate/Public/Widgets/Layout/SSafeZone.h
--- a/Engine/Source/Runtime/Slate/Public/Widgets/Layout/SSafeZone.h
+++ b/Engine/Source/Runtime/Slate/Public/Widgets/Layout/SSafeZone.h
@@ -27,7 +27,6 @@ class SLATE_API SSafeZone : public SBox
 		, _Padding( 0.0f )
 		, _Content()
 		, _IsTitleSafe( false )
-		 
 		{}
 
 		/** Horizontal alignment of content in the area allotted to the SBox by its parent */
@@ -47,12 +46,14 @@ class SLATE_API SSafeZone : public SBox
 
 	SLATE_END_ARGS()
 
-	void Construct( const FArguments& InArgs );
+public:
 
-	FMargin GetSafeZonePadding() const;
+	void Construct( const FArguments& InArgs );
 
 	void SetTitleSafe( bool bIsTitleSafe );
 
+	virtual void OnArrangeChildren( const FGeometry& AllottedGeometry, FArrangedChildren& ArrangedChildren ) const override;
+
 private:
 	/** Cached values from the args */
 	TAttribute<FMargin> Padding;
diff --git a/Engine/Source/Runtime/Slate/Public/Widgets/Text/ITextEditorWidget.h b/Engine/Source/Runtime/Slate/Public/Widgets/Text/ITextEditorWidget.h
--- a/Engine/Source/Runtime/Slate/Public/Widgets/Text/ITextEditorWidget.h
+++ b/Engine/Source/Runtime/Slate/Public/Widgets/Text/ITextEditorWidget.h
@@ -315,6 +315,8 @@ class SLATE_API ITextEditorWidget
 
 	virtual TSharedRef< SWidget > GetWidget() = 0;
 
+	virtual TSharedRef<IVirtualKeyboardEntry> GetVirtualKeyboardEntry() = 0;
+
 	DEPRECATED(4.9, "SummonContextMenu() without a FWidgetPath parameter is deprecated. Use the new version of SummonContextMenu().")
 	virtual void SummonContextMenu(const FVector2D& InLocation, TSharedPtr<SWindow> ParentWindow = TSharedPtr<SWindow>())
 	{
diff --git a/Engine/Source/Runtime/Slate/Public/Widgets/Text/SMultiLineEditableText.h b/Engine/Source/Runtime/Slate/Public/Widgets/Text/SMultiLineEditableText.h
--- a/Engine/Source/Runtime/Slate/Public/Widgets/Text/SMultiLineEditableText.h
+++ b/Engine/Source/Runtime/Slate/Public/Widgets/Text/SMultiLineEditableText.h
@@ -507,6 +507,7 @@ private:
 	virtual void Undo() override;
 	virtual void Redo() override;
 	virtual TSharedRef< SWidget > GetWidget() override;
+	virtual TSharedRef<IVirtualKeyboardEntry> GetVirtualKeyboardEntry() override;
 	virtual void SummonContextMenu(const FVector2D& InLocation, TSharedPtr<SWindow> ParentWindow, const FWidgetPath& EventPath) override;
 	virtual void LoadText() override;
 	//~ End ITextEditorWidget Interface
diff --git a/Engine/Source/Runtime/SlateCore/Private/Widgets/SWindow.cpp b/Engine/Source/Runtime/SlateCore/Private/Widgets/SWindow.cpp
--- a/Engine/Source/Runtime/SlateCore/Private/Widgets/SWindow.cpp
+++ b/Engine/Source/Runtime/SlateCore/Private/Widgets/SWindow.cpp
@@ -1363,7 +1363,7 @@ bool SWindow::OnIsActiveChanged( const FWindowActivateEvent& ActivateEvent )
 	}
 	else
 	{
-		if (SupportsKeyboardFocus())
+		if (SupportsKeyboardFocus() && ActivateEvent.GetActivationType() == FWindowActivateEvent::EA_Activate)
 		{
 			TArray< TSharedRef<SWindow> > JustThisWindow;
 			JustThisWindow.Add( SharedThis(this) );
diff --git a/Engine/Source/Runtime/SlateRHIRenderer/Private/SlateMaterialResource.cpp b/Engine/Source/Runtime/SlateRHIRenderer/Private/SlateMaterialResource.cpp
--- a/Engine/Source/Runtime/SlateRHIRenderer/Private/SlateMaterialResource.cpp
+++ b/Engine/Source/Runtime/SlateRHIRenderer/Private/SlateMaterialResource.cpp
@@ -6,7 +6,6 @@
 
 FSlateMaterialResource::FSlateMaterialResource(const UMaterialInterface& InMaterial, const FVector2D& InImageSize, FSlateShaderResource* InTextureMask )
 	: MaterialObject( &InMaterial )
-	, RenderProxy( InMaterial.GetRenderProxy(false) )
 	, SlateProxy( new FSlateShaderResourceProxy )
 	, TextureMaskResource( InTextureMask )
 	, Width(FMath::RoundToInt(InImageSize.X))
@@ -27,7 +26,6 @@ FSlateMaterialResource::~FSlateMaterialResource()
 void FSlateMaterialResource::UpdateMaterial(const UMaterialInterface& InMaterialResource, const FVector2D& InImageSize, FSlateShaderResource* InTextureMask )
 {
 	MaterialObject = &InMaterialResource;
-	RenderProxy = InMaterialResource.GetRenderProxy(false);
 	if( !SlateProxy )
 	{
 		SlateProxy = new FSlateShaderResourceProxy;
@@ -42,3 +40,15 @@ void FSlateMaterialResource::UpdateMaterial(const UMaterialInterface& InMaterial
 	Height = FMath::RoundToInt(InImageSize.Y);
 }
 
+void FSlateMaterialResource::ResetMaterial()
+{
+	MaterialObject = nullptr;
+	TextureMaskResource = nullptr;
+	if (SlateProxy)
+	{
+		delete SlateProxy;
+	}
+	SlateProxy = nullptr;
+	Width = 0;
+	Height = 0;
+}
\ No newline at end of file
diff --git a/Engine/Source/Runtime/SlateRHIRenderer/Private/SlateMaterialResource.h b/Engine/Source/Runtime/SlateRHIRenderer/Private/SlateMaterialResource.h
--- a/Engine/Source/Runtime/SlateRHIRenderer/Private/SlateMaterialResource.h
+++ b/Engine/Source/Runtime/SlateRHIRenderer/Private/SlateMaterialResource.h
@@ -16,9 +16,7 @@ public:
 	virtual ESlateShaderResource::Type GetType() const override { return ESlateShaderResource::Material; }
 
 	void UpdateMaterial(const UMaterialInterface& InMaterialResource, const FVector2D& InImageSize, FSlateShaderResource* InTextureMask );
-
-	/** Updates the rendering resource */
-	void UpdateRenderResource(FMaterialRenderProxy* Proxy) { RenderProxy = Proxy; }
+	void ResetMaterial();
 
 	/** @return The material render proxy */
 	FMaterialRenderProxy* GetRenderProxy() const { return MaterialObject->GetRenderProxy(false, false); }
@@ -29,8 +27,6 @@ public:
 	FSlateShaderResource* GetTextureMaskResource() const { return TextureMaskResource; }
 public:
 	const class UMaterialInterface* MaterialObject;
-	/** Rendering thread version of the material object */
-	FMaterialRenderProxy* RenderProxy;
 	/** Slate proxy used for batching the material */
 	FSlateShaderResourceProxy* SlateProxy;
 
diff --git a/Engine/Source/Runtime/SlateRHIRenderer/Private/SlateRHIRenderingPolicy.cpp b/Engine/Source/Runtime/SlateRHIRenderer/Private/SlateRHIRenderingPolicy.cpp
--- a/Engine/Source/Runtime/SlateRHIRenderer/Private/SlateRHIRenderingPolicy.cpp
+++ b/Engine/Source/Runtime/SlateRHIRenderer/Private/SlateRHIRenderingPolicy.cpp
@@ -500,81 +500,84 @@ void FSlateRHIRenderingPolicy::DrawElements(FRHICommandListImmediate& RHICmdList
 				}
 
 				FSlateMaterialResource* MaterialShaderResource = (FSlateMaterialResource*)ShaderResource;
-				FMaterialRenderProxy* MaterialRenderProxy = MaterialShaderResource->GetRenderProxy();
+				if( MaterialShaderResource->GetMaterialObject() != nullptr )
+				{
+					FMaterialRenderProxy* MaterialRenderProxy = MaterialShaderResource->GetRenderProxy();
 
-				const FMaterial* Material = MaterialRenderProxy->GetMaterial(SceneView->GetFeatureLevel());
+					const FMaterial* Material = MaterialRenderProxy->GetMaterial(SceneView->GetFeatureLevel());
 
-				FSlateMaterialShaderPS* PixelShader = GetMaterialPixelShader( Material, ShaderType, DrawEffects );
+					FSlateMaterialShaderPS* PixelShader = GetMaterialPixelShader(Material, ShaderType, DrawEffects);
 
-				const bool bUseInstancing = RenderBatch.InstanceCount > 0 && RenderBatch.InstanceData != nullptr;
-				FSlateMaterialShaderVS* VertexShader = GetMaterialVertexShader( Material, bUseInstancing );
+					const bool bUseInstancing = RenderBatch.InstanceCount > 0 && RenderBatch.InstanceData != nullptr;
+					FSlateMaterialShaderVS* VertexShader = GetMaterialVertexShader(Material, bUseInstancing);
 
-				if( VertexShader && PixelShader )
-				{
-					RHICmdList.SetLocalBoundShaderState(RHICmdList.BuildLocalBoundShaderState(
-						bUseInstancing ? GSlateInstancedVertexDeclaration.VertexDeclarationRHI : GSlateVertexDeclaration.VertexDeclarationRHI,
-						VertexShader->GetVertexShader(),
-						nullptr,
-						nullptr,
-						PixelShader->GetPixelShader(),
-						FGeometryShaderRHIRef()));
-
-					VertexShader->SetViewProjection(RHICmdList, DynamicOffset * ViewProjectionMatrix);
-					VertexShader->SetVerticalAxisMultiplier(RHICmdList, bAllowSwitchVerticalAxis && RHINeedsToSwitchVerticalAxis(GShaderPlatformForFeatureLevel[GMaxRHIFeatureLevel]) ? -1.0f : 1.0f );
-					VertexShader->SetMaterialShaderParameters(RHICmdList, *SceneView, MaterialRenderProxy, Material);
-
-					PixelShader->SetParameters(RHICmdList, *SceneView, MaterialRenderProxy, Material, DisplayGamma, ShaderParams.PixelParams);
-					PixelShader->SetDisplayGamma(RHICmdList, (DrawFlags & ESlateBatchDrawFlag::NoGamma) ? 1.0f : DisplayGamma);
-					
-					FSlateShaderResource* MaskResource = MaterialShaderResource->GetTextureMaskResource();
-					if( MaskResource )
+					if (VertexShader && PixelShader)
 					{
-						RHICmdList.SetBlendState(TStaticBlendState<CW_RGBA, BO_Add, BF_SourceAlpha, BF_InverseSourceAlpha, BO_Add, BF_InverseDestAlpha, BF_One>::GetRHI());
+						RHICmdList.SetLocalBoundShaderState(RHICmdList.BuildLocalBoundShaderState(
+							bUseInstancing ? GSlateInstancedVertexDeclaration.VertexDeclarationRHI : GSlateVertexDeclaration.VertexDeclarationRHI,
+							VertexShader->GetVertexShader(),
+							nullptr,
+							nullptr,
+							PixelShader->GetPixelShader(),
+							FGeometryShaderRHIRef()));
+
+						VertexShader->SetViewProjection(RHICmdList, DynamicOffset * ViewProjectionMatrix);
+						VertexShader->SetVerticalAxisMultiplier(RHICmdList, bAllowSwitchVerticalAxis && RHINeedsToSwitchVerticalAxis(GShaderPlatformForFeatureLevel[GMaxRHIFeatureLevel]) ? -1.0f : 1.0f);
+						VertexShader->SetMaterialShaderParameters(RHICmdList, *SceneView, MaterialRenderProxy, Material);
+
+						PixelShader->SetParameters(RHICmdList, *SceneView, MaterialRenderProxy, Material, DisplayGamma, ShaderParams.PixelParams);
+						PixelShader->SetDisplayGamma(RHICmdList, (DrawFlags & ESlateBatchDrawFlag::NoGamma) ? 1.0f : DisplayGamma);
+
+						FSlateShaderResource* MaskResource = MaterialShaderResource->GetTextureMaskResource();
+						if (MaskResource)
+						{
+							RHICmdList.SetBlendState(TStaticBlendState<CW_RGBA, BO_Add, BF_SourceAlpha, BF_InverseSourceAlpha, BO_Add, BF_InverseDestAlpha, BF_One>::GetRHI());
 
-						FTexture2DRHIRef TextureRHI;
-						TextureRHI = ((TSlateTexture<FTexture2DRHIRef>*)MaskResource)->GetTypedResource();
+							FTexture2DRHIRef TextureRHI;
+							TextureRHI = ((TSlateTexture<FTexture2DRHIRef>*)MaskResource)->GetTypedResource();
 
-						PixelShader->SetAdditionalTexture(RHICmdList, TextureRHI, BilinearClamp);
-					}
+							PixelShader->SetAdditionalTexture(RHICmdList, TextureRHI, BilinearClamp);
+						}
 
 #if SLATE_PRE_MULTIPLY
-					if ( RenderBatch.DrawFlags & ESlateBatchDrawFlag::AlphaCompositing )
-					{
-						RHICmdList.SetBlendState(TStaticBlendState<CW_RGBA, BO_Add, BF_One, BF_InverseSourceAlpha, BO_Add, BF_One, BF_InverseSourceAlpha>::GetRHI());
-					}
+						if (RenderBatch.DrawFlags & ESlateBatchDrawFlag::AlphaCompositing)
+						{
+							RHICmdList.SetBlendState(TStaticBlendState<CW_RGBA, BO_Add, BF_One, BF_InverseSourceAlpha, BO_Add, BF_One, BF_InverseSourceAlpha>::GetRHI());
+						}
 #endif
 
-					uint32 PrimitiveCount = RenderBatch.DrawPrimitiveType == ESlateDrawPrimitive::LineList ? RenderBatch.NumIndices / 2 : RenderBatch.NumIndices / 3;
+						uint32 PrimitiveCount = RenderBatch.DrawPrimitiveType == ESlateDrawPrimitive::LineList ? RenderBatch.NumIndices / 2 : RenderBatch.NumIndices / 3;
 
-					// for RHIs that can't handle VertexOffset, we need to offset the stream source each time
-					if (!GRHISupportsBaseVertexIndex)
-					{
-						if( bUseInstancing )
+						// for RHIs that can't handle VertexOffset, we need to offset the stream source each time
+						if (!GRHISupportsBaseVertexIndex)
 						{
-							FSlateUpdatableInstanceBuffer* InstanceBuffer = (FSlateUpdatableInstanceBuffer*)RenderBatch.InstanceData;
-							InstanceBuffer->BindStreamSource(RHICmdList, 1, RenderBatch.InstanceOffset);
+							if (bUseInstancing)
+							{
+								FSlateUpdatableInstanceBuffer* InstanceBuffer = (FSlateUpdatableInstanceBuffer*)RenderBatch.InstanceData;
+								InstanceBuffer->BindStreamSource(RHICmdList, 1, RenderBatch.InstanceOffset);
+							}
+							else
+							{
+								RHICmdList.SetStreamSource(1, nullptr, 0, 0);
+							}
+
+							RHICmdList.SetStreamSource(0, VertexBuffer->VertexBufferRHI, sizeof(FSlateVertex), RenderBatch.VertexOffset * sizeof(FSlateVertex));
+							RHICmdList.DrawIndexedPrimitive(IndexBuffer->IndexBufferRHI, GetRHIPrimitiveType(RenderBatch.DrawPrimitiveType), 0, 0, RenderBatch.NumVertices, RenderBatch.IndexOffset, PrimitiveCount, RenderBatch.InstanceCount);
 						}
 						else
 						{
-							RHICmdList.SetStreamSource(1, nullptr, 0, 0);
+							if (bUseInstancing)
+							{
+								FSlateUpdatableInstanceBuffer* InstanceBuffer = (FSlateUpdatableInstanceBuffer*)RenderBatch.InstanceData;
+								InstanceBuffer->BindStreamSource(RHICmdList, 1, RenderBatch.InstanceOffset);
+							}
+							else
+							{
+								RHICmdList.SetStreamSource(1, nullptr, 0, 0);
+							}
+
+							RHICmdList.DrawIndexedPrimitive(IndexBuffer->IndexBufferRHI, GetRHIPrimitiveType(RenderBatch.DrawPrimitiveType), RenderBatch.VertexOffset, 0, RenderBatch.NumVertices, RenderBatch.IndexOffset, PrimitiveCount, RenderBatch.InstanceCount);
 						}
-						
-						RHICmdList.SetStreamSource(0, VertexBuffer->VertexBufferRHI, sizeof(FSlateVertex), RenderBatch.VertexOffset * sizeof(FSlateVertex));
-						RHICmdList.DrawIndexedPrimitive(IndexBuffer->IndexBufferRHI, GetRHIPrimitiveType(RenderBatch.DrawPrimitiveType), 0, 0, RenderBatch.NumVertices, RenderBatch.IndexOffset, PrimitiveCount, RenderBatch.InstanceCount);
-					}
-					else
-					{
-						if( bUseInstancing )
-						{
-							FSlateUpdatableInstanceBuffer* InstanceBuffer = (FSlateUpdatableInstanceBuffer*)RenderBatch.InstanceData;
-							InstanceBuffer->BindStreamSource(RHICmdList, 1, RenderBatch.InstanceOffset);
-						}
-						else
-						{
-							RHICmdList.SetStreamSource(1, nullptr, 0, 0);
-						}
-
-						RHICmdList.DrawIndexedPrimitive(IndexBuffer->IndexBufferRHI, GetRHIPrimitiveType(RenderBatch.DrawPrimitiveType), RenderBatch.VertexOffset, 0, RenderBatch.NumVertices, RenderBatch.IndexOffset, PrimitiveCount, RenderBatch.InstanceCount);
 					}
 				}
 			}
diff --git a/Engine/Source/Runtime/SlateRHIRenderer/Private/SlateRHIResourceManager.cpp b/Engine/Source/Runtime/SlateRHIRenderer/Private/SlateRHIResourceManager.cpp
--- a/Engine/Source/Runtime/SlateRHIRenderer/Private/SlateRHIResourceManager.cpp
+++ b/Engine/Source/Runtime/SlateRHIRenderer/Private/SlateRHIResourceManager.cpp
@@ -687,8 +687,6 @@ FSlateMaterialResource* FSlateRHIResourceManager::GetMaterialResource(const UObj
 	}
 	else
 	{
-		// Keep the resource up to date
-		MaterialResource->UpdateRenderResource(Material->GetRenderProxy(false, false));
 		MaterialResource->SlateProxy->ActualSize = ImageSize.IntPoint();
 	}
 
@@ -748,6 +746,7 @@ void FSlateRHIResourceManager::ReleaseDynamicResource( const FSlateBrush& InBrus
 
 				if (MaterialResource.IsValid())
 				{
+					MaterialResource->ResetMaterial();
 					MaterialResourceFreeList.Add( MaterialResource );
 				}
 			}
diff --git a/Engine/Source/Runtime/VectorVM/Private/VectorVM.cpp b/Engine/Source/Runtime/VectorVM/Private/VectorVM.cpp
--- a/Engine/Source/Runtime/VectorVM/Private/VectorVM.cpp
+++ b/Engine/Source/Runtime/VectorVM/Private/VectorVM.cpp
@@ -80,7 +80,12 @@ UNiagaraSparseVolumeDataObject::UNiagaraSparseVolumeDataObject(const FObjectInit
 	Size = 64;
 	NumBuckets = Size*Size*Size;
 	//Data.AddZeroed(NumBuckets);
-	Data.Init(FVector4(0.1f, 0.1f, 0.1f, 0.1f), NumBuckets);
+
+	// avoid allocation in CDO
+	if (!HasAnyFlags(RF_ClassDefaultObject))
+	{
+		Data.Init(FVector4(0.1f, 0.1f, 0.1f, 0.1f), NumBuckets);
+	}
 }
 
 //////////////////////////////////////////////////////////////////////////
diff --git a/Engine/Source/Runtime/Windows/D3D11RHI/Private/D3D11Device.cpp b/Engine/Source/Runtime/Windows/D3D11RHI/Private/D3D11Device.cpp
--- a/Engine/Source/Runtime/Windows/D3D11RHI/Private/D3D11Device.cpp
+++ b/Engine/Source/Runtime/Windows/D3D11RHI/Private/D3D11Device.cpp
@@ -34,7 +34,7 @@ TAutoConsoleVariable<int32> CVarD3D11ZeroBufferSizeInMB(
 	ECVF_ReadOnly
 	);
 
-FD3D11DynamicRHI::FD3D11DynamicRHI(IDXGIFactory1* InDXGIFactory1,D3D_FEATURE_LEVEL InFeatureLevel, int32 InChosenAdapter) :
+FD3D11DynamicRHI::FD3D11DynamicRHI(IDXGIFactory1* InDXGIFactory1,D3D_FEATURE_LEVEL InFeatureLevel, int32 InChosenAdapter, const DXGI_ADAPTER_DESC& InChosenDescription) :
 	DXGIFactory1(InDXGIFactory1),
 	bDeviceRemoved(false),
 	FeatureLevel(InFeatureLevel),
@@ -55,7 +55,8 @@ FD3D11DynamicRHI::FD3D11DynamicRHI(IDXGIFactory1* InDXGIFactory1,D3D_FEATURE_LEV
 	PendingNumIndices(0),
 	PendingIndexDataStride(0),
 	GPUProfilingData(this),
-	ChosenAdapter(InChosenAdapter)
+	ChosenAdapter(InChosenAdapter),
+	ChosenDescription(InChosenDescription)
 {
 	// This should be called once at the start 
 	check(ChosenAdapter >= 0);
diff --git a/Engine/Source/Runtime/Windows/D3D11RHI/Private/D3D11RHIPrivate.h b/Engine/Source/Runtime/Windows/D3D11RHI/Private/D3D11RHIPrivate.h
--- a/Engine/Source/Runtime/Windows/D3D11RHI/Private/D3D11RHIPrivate.h
+++ b/Engine/Source/Runtime/Windows/D3D11RHI/Private/D3D11RHIPrivate.h
@@ -288,7 +288,7 @@ public:
 	TMap<FD3D11LockedKey,FD3D11LockedData> OutstandingLocks;
 
 	/** Initialization constructor. */
-	FD3D11DynamicRHI(IDXGIFactory1* InDXGIFactory1,D3D_FEATURE_LEVEL InFeatureLevel,int32 InChosenAdapter);
+	FD3D11DynamicRHI(IDXGIFactory1* InDXGIFactory1,D3D_FEATURE_LEVEL InFeatureLevel,int32 InChosenAdapter, const DXGI_ADAPTER_DESC& ChosenDescription);
 
 	/** Destructor */
 	virtual ~FD3D11DynamicRHI();
@@ -690,6 +690,8 @@ protected:
 	FD3DGPUProfiler GPUProfilingData;
 	// >= 0, was computed before, unless hardware was changed during engine init it should be the same
 	int32 ChosenAdapter;
+	// we don't use AdapterDesc.Description as there is a bug with Optimus where it can report the wrong name
+	DXGI_ADAPTER_DESC ChosenDescription;
 
 	template<typename BaseResourceType>
 	TD3D11Texture2D<BaseResourceType>* CreateD3D11Texture2D(uint32 SizeX,uint32 SizeY,uint32 SizeZ,bool bTextureArray,bool CubeTexture,uint8 Format,
@@ -813,6 +815,8 @@ public:
 
 private:
 	FD3D11Adapter ChosenAdapter;
+	// we don't use GetDesc().Description as there is a bug with Optimus where it can report the wrong name
+	DXGI_ADAPTER_DESC ChosenDescription;
 
 	// set MaxSupportedFeatureLevel and ChosenAdapter
 	void FindAdapter();
