Copying //UE4/Orion-Staging to //UE4/Main (Origin //Orion/Dev-General @ 2870388)

diff --git a/Engine/Source/Runtime/RHI/Private/GPUDefragAllocator.cpp b/Engine/Source/Runtime/RHI/Private/GPUDefragAllocator.cpp
--- a/Engine/Source/Runtime/RHI/Private/GPUDefragAllocator.cpp
+++ b/Engine/Source/Runtime/RHI/Private/GPUDefragAllocator.cpp
@@ -37,19 +37,30 @@ DECLARE_DWORD_ACCUMULATOR_STAT(TEXT("Total Async Cancels"), STAT_TexturePool_Tot
 #define DEFRAG_CHUNK_COUNTER_MIN		(20)
 #define DEFRAG_CHUNK_COUNTER_MAX		(80)
 
-#if VALIDATE_SYNC_SIZE
-static int32 GGPUDefragValidateSyncSize = 0;
-static FAutoConsoleVariableRef CVarGPUDefragValidateSyncSize(
-	TEXT("r.GPUDefrag.ValidateSyncSize"),
-	GGPUDefragValidateSyncSize,
-	TEXT("Enables validation and debug printout to detect unsafe allocations from the GPU defrag allocator.\n"),
+#if TRACK_RELOCATIONS
+int32 FGPUDefragAllocator::GGPUDefragDumpRelocationsToTTY = 0;
+FAutoConsoleVariableRef FGPUDefragAllocator::CVarGPUDefragDumpRelocationsToTTY(
+	TEXT("r.GPUDefrag.DumpRelocationsToTTY"),
+	GGPUDefragDumpRelocationsToTTY,
+	TEXT("Dumps logging information for every relocation.\n"),
 	ECVF_Default
 	);
-#endif
 
+FGPUDefragAllocator::FRelocationEntry::FRelocationEntry(const uint8* InOldBase, const uint8* InNewBase, uint64 InSize, uint64 InSyncIndex)
+: OldBase(InOldBase)
+, NewBase(InNewBase)
+, Size(InSize)
+, SyncIndex(InSyncIndex)
+{
+	if (GGPUDefragDumpRelocationsToTTY)
+	{
+		printf("Relocation: 0x%p to 0x%p, %i, %i\n", OldBase, NewBase, (int32)Size, (int32)SyncIndex);
+	}
+}
+#endif
 
 /*-----------------------------------------------------------------------------
-FBestFitAllocator::FMemoryChunk implementation.
+FGPUDefragAlllocator::FMemoryChunk implementation.
 -----------------------------------------------------------------------------*/
 
 /**
@@ -226,8 +237,8 @@ FGPUDefragAllocator::FMemoryChunk* FGPUDefragAllocator::AllocateChunk(FMemoryChu
 	check(FreeChunk->GetAvailableSize() >= AllocationSize);
 	check(!FreeChunk->IsLocked());
 
-#if VALIDATE_SYNC_SIZE
-	if (GGPUDefragValidateSyncSize)
+#if TRACK_RELOCATIONS
+	if (GGPUDefragDumpRelocationsToTTY)
 	{
 		printf("allocate before splits: 0x%p, %i, %i, %i, SyncSize: %i\n", FreeChunk->Base, (int32)FreeChunk->Size, (int32)AllocationSize, (int32)GetCurrentSyncIndex(), (int32)FreeChunk->SyncSize);
 	}
@@ -237,8 +248,8 @@ FGPUDefragAllocator::FMemoryChunk* FGPUDefragAllocator::AllocateChunk(FMemoryChu
 	// and the beginning of the chunk is currently being relocated by the GPU, split that part off and allocate from the rest.
 	if (!bAsync && FreeChunk->IsRelocating() && FreeChunk->SyncSize > 0 && FreeChunk->SyncSize < FreeChunk->Size)
 	{
-#if VALIDATE_SYNC_SIZE
-		if (GGPUDefragValidateSyncSize)
+#if TRACK_RELOCATIONS
+		if (GGPUDefragDumpRelocationsToTTY)
 		{
 			printf("splitting:\n");
 		}
@@ -253,8 +264,8 @@ FGPUDefragAllocator::FMemoryChunk* FGPUDefragAllocator::AllocateChunk(FMemoryChu
 	// Split chunk to avoid waste.
 	if (FreeChunk->Size > AllocationSize)
 	{
-#if VALIDATE_SYNC_SIZE
-		if (GGPUDefragValidateSyncSize)
+#if TRACK_RELOCATIONS
+		if (GGPUDefragDumpRelocationsToTTY)
 		{
 			printf("splitting again:\n");
 		}
@@ -270,17 +281,24 @@ FGPUDefragAllocator::FMemoryChunk* FGPUDefragAllocator::AllocateChunk(FMemoryChu
 	FPlatformAtomics::InterlockedAdd(&AllocatedMemorySize, FreeChunk->Size);
 	FPlatformAtomics::InterlockedAdd(&AvailableMemorySize, -FreeChunk->Size);
 
-#if VALIDATE_SYNC_SIZE
-	if (GGPUDefragValidateSyncSize)
+#if TRACK_RELOCATIONS
+	if (GGPUDefragDumpRelocationsToTTY)
 	{
 		printf("allocate: 0x%p, %i, %i\n", FreeChunk->Base, (int32)AllocationSize, (int32)GetCurrentSyncIndex());
-		if (bDoValidation)
-		{
-			ValidateRelocations(FreeChunk->Base, AllocationSize);
-		}
 	}
 #endif
 
+#if VALIDATE_SYNC_SIZE
+	if (bDoValidation)
+	{
+		ValidateRelocations(FreeChunk->Base, AllocationSize);
+	}	
+#endif
+
+#if VALIDATE_MEMORY_PROTECTION
+	PlatformSetStandardMemoryPrivileges(FMemProtectTracker(FreeChunk->Base, nullptr, FreeChunk->Size, 0));
+#endif
+
 	// Keep track of mapping and return pointer.
 	PointerToChunkMap.Add(FreeChunk->Base, FreeChunk);
 	return FreeChunk;
@@ -295,8 +313,8 @@ FGPUDefragAllocator::FMemoryChunk* FGPUDefragAllocator::AllocateChunk(FMemoryChu
 */
 void FGPUDefragAllocator::FreeChunk(FMemoryChunk* Chunk)
 {
-#if VALIDATE_SYNC_SIZE
-	if (GGPUDefragValidateSyncSize)
+#if TRACK_RELOCATIONS
+	if (GGPUDefragDumpRelocationsToTTY)
 	{
 		printf("FreeChunk: 0x%p, %i, %i, %i\n", Chunk->Base, (int32)Chunk->Size, (int32)Chunk->SyncSize, (int32)Chunk->SyncIndex);
 	}
@@ -310,6 +328,10 @@ void FGPUDefragAllocator::FreeChunk(FMemoryChunk* Chunk)
 	FPlatformAtomics::InterlockedAdd(&AllocatedMemorySize, -Chunk->Size);
 	FPlatformAtomics::InterlockedAdd(&AvailableMemorySize, Chunk->Size);
 
+#if VALIDATE_MEMORY_PROTECTION
+	PlatformSetNoMemoryPrivileges(FMemProtectTracker(Chunk->Base, nullptr, Chunk->Size, 0));
+#endif
+
 	// Free the chunk.
 	LinkFreeChunk(Chunk);
 }
@@ -337,6 +359,7 @@ void FGPUDefragAllocator::Free(void* Pointer)
 	FMemoryChunk* MatchingChunk = PointerToChunkMap.FindRef(Pointer);
 	check(MatchingChunk);
 	check(MatchingChunk->Base == Pointer);
+	checkf(MatchingChunk->LockCount == 0, TEXT("Chunk with base address: 0x%p is being freed with %i outstanding locks.  This is a data corruption hazard."), Pointer, MatchingChunk->LockCount);
 
 	int64 PaddingWaste = MatchingChunk->Size - MatchingChunk->OrigSize;
 	FPlatformAtomics::InterlockedAdd(&PaddingWasteSize, -PaddingWaste);
@@ -386,7 +409,7 @@ void FGPUDefragAllocator::Unlock(const void* Pointer)
 	FScopeLock Lock(&SynchronizationObject);
 	FMemoryChunk* MatchingChunk = PointerToChunkMap.FindRef(Pointer);
 	check(MatchingChunk && MatchingChunk->IsRelocating() == false);
-	check(MatchingChunk->LockCount > 0);
+	checkf(MatchingChunk->LockCount > 0, TEXT("Chunk: 0x%p has bad lockcount: %i"), Pointer, MatchingChunk->LockCount);
 	--MatchingChunk->LockCount;
 	--NumLockedChunks;
 }
@@ -406,6 +429,10 @@ void FGPUDefragAllocator::SetUserPayload(const void* Pointer, void* UserPayload)
 	{
 		MatchingChunk->UserPayload = UserPayload;
 	}
+
+#if VALIDATE_MEMORY_PROTECTION
+	PlatformSetStaticMemoryPrivileges(FMemProtectTracker(MatchingChunk->Base, UserPayload, MatchingChunk->Size, 0));
+#endif
 }
 
 /**
@@ -561,8 +588,8 @@ FGPUDefragAllocator::FMemoryChunk* FGPUDefragAllocator::Shrink(FMemoryChunk* Chu
 
 		// Make the original chunk "free memory".
 		NewFreeChunk = Chunk->PreviousChunk;
-#if VALIDATE_SYNC_SIZE
-		if (GGPUDefragValidateSyncSize)
+#if TRACK_RELOCATIONS
+		if (GGPUDefragDumpRelocationsToTTY)
 		{
 			printf("shrink free chunk\n");
 		}
@@ -866,25 +893,21 @@ void FGPUDefragAllocator::CheckForErrors(bool bCheckSortedFreeList)
 
 #if VALIDATE_SYNC_SIZE
 void FGPUDefragAllocator::ValidateRelocations(uint8* UsedBaseAddress, uint64 Size)
-{
-	if (GGPUDefragValidateSyncSize)
+{	
+	FScopeLock SyncLock(&SynchronizationObject);
+	for (int32 i = 0; i < Relocations.Num(); ++i)
 	{
-		FScopeLock SyncLock(&SynchronizationObject);
-
-		for (int32 i = 0; i < Relocations.Num(); ++i)
-		{
-			const FRelocationEntry& Relocation = Relocations[i];
+		const FRelocationEntry& Relocation = Relocations[i];
 
-			bool bBeforeOrig = (UsedBaseAddress + Size) <= Relocation.OldBase;
-			bool bAfterOrig = (UsedBaseAddress) >= (Relocation.OldBase + Relocation.Size);
+		bool bBeforeOrig = (UsedBaseAddress + Size) <= Relocation.OldBase;
+		bool bAfterOrig = (UsedBaseAddress) >= (Relocation.OldBase + Relocation.Size);
 
-			bool bBeforeNew = (UsedBaseAddress + Size) <= Relocation.NewBase;
-			bool bAfterNew = (UsedBaseAddress) >= (Relocation.NewBase + Relocation.Size);
+		bool bBeforeNew = (UsedBaseAddress + Size) <= Relocation.NewBase;
+		bool bAfterNew = (UsedBaseAddress) >= (Relocation.NewBase + Relocation.Size);
 
-			bool bAfterSync = Relocation.SyncIndex <= CompletedSyncIndex;
-			checkf(((bBeforeOrig || bAfterOrig) && (bBeforeNew || bAfterNew)) || bAfterSync, TEXT("Corruption Hazard, Allocation not protected by sync size."));
-		}
-	}
+		bool bAfterSync = Relocation.SyncIndex <= CompletedSyncIndex;
+		checkf(((bBeforeOrig || bAfterOrig) && (bBeforeNew || bAfterNew)) || bAfterSync, TEXT("Corruption Hazard, Allocation not protected by sync size."));
+	}	
 }
 #endif
 
@@ -923,8 +946,8 @@ FGPUDefragAllocator::FMemoryChunk* FGPUDefragAllocator::RelocateIntoFreeChunk(FR
 	// CheckForErrors( true );
 
 	// Merge adjacent free chunks into SourceChunk to make a single free chunk.
-#if VALIDATE_SYNC_SIZE
-	if (GGPUDefragValidateSyncSize)
+#if TRACK_RELOCATIONS
+	if (GGPUDefragDumpRelocationsToTTY)
 	{
 		printf("relocate link free chunk\n");
 	}
@@ -1019,6 +1042,7 @@ FGPUDefragAllocator::FMemoryChunk* FGPUDefragAllocator::RelocateIntoFreeChunk(FR
 					InFlightTailChunk->SetSyncIndex(CompletedSyncIndex, 0);
 					AllocateChunk(InFlightTailChunk, InFlightTailChunk->Size, false, false);
 					InFlightTailChunk->SetSyncIndex(GetCurrentSyncIndex(), InFlightTailChunk->Size);
+					++InFlightTailChunk->LockCount;
 					PendingFreeChunks.AddTail(InFlightTailChunk);
 				}
 				else
@@ -1079,10 +1103,20 @@ FGPUDefragAllocator::FMemoryChunk* FGPUDefragAllocator::RelocateIntoFreeChunk(FR
 		FPlatformAtomics::InterlockedAdd(&PendingMemoryAdjustment, -MemoryAdjustment);
 	}
 
-#if VALIDATE_SYNC_SIZE
-	if (GGPUDefragValidateSyncSize)
+#if VALIDATE_MOVES
+	for (const auto& Relocation : Relocations)
+	{
+		bool bBeforeReloc = (DestNewBase + OldSize) <= Relocation.NewBase;
+		bool bAfterReloc = (DestNewBase) >= (Relocation.NewBase + Relocation.Size);
+				
+		checkf((bBeforeReloc || bAfterReloc), TEXT("Corruption Hazard, Destinations overlap within same frame."));
+	}
+#endif
+#if TRACK_RELOCATIONS
+	Relocations.Add(FRelocationEntry(SourceOldBase, DestNewBase, OldSize, GetCurrentSyncIndex()));
+
+	if (GGPUDefragDumpRelocationsToTTY)
 	{
-		Relocations.Add(FRelocationEntry(SourceOldBase, DestNewBase, OldSize, GetCurrentSyncIndex()));
 		printf("bFreeAdjOnLeft %i, bFreeAdjOnRight %i\n", bFreeChunkPreviousAdjacent, bFreeChunkNextAdjacent);
 	}
 #endif
@@ -1118,6 +1152,7 @@ bool FGPUDefragAllocator::FinishAllRelocations()
 	for (TDoubleLinkedList<FMemoryChunk*>::TIterator It(PendingFreeChunks.GetHead()); It; ++It)
 	{
 		FMemoryChunk* Chunk = *It;
+		--Chunk->LockCount;
 		FreeChunk(Chunk);
 	}
 	PendingFreeChunks.Empty();
@@ -1140,7 +1175,7 @@ void FGPUDefragAllocator::InsertFence()
 }
 
 #if VALIDATE_MEMORY_PROTECTION
-void FGPUDefragAllocator::AllowCPUAccessToBlocks()
+void FGPUDefragAllocator::SetStaticMemoryPrivileges()
 {
 	FScopedDurationTimer Timer(TimeInMemProtect);
 	const int32 Count = BlocksToUnProtect.Num();
@@ -1149,7 +1184,7 @@ void FGPUDefragAllocator::AllowCPUAccessToBlocks()
 		const FMemProtectTracker& Block = BlocksToUnProtect[i];
 		if (Block.SyncIndex >= CompletedSyncIndex)
 		{
-			PlatformAllowCPUAccessToBlock(Block);
+			PlatformSetStaticMemoryPrivileges(Block);
 			BlocksToUnProtect.RemoveAtSwap(i, 1, false);
 		}
 	}	
@@ -1173,7 +1208,7 @@ void FGPUDefragAllocator::BlockOnFence()
 		BlockedCycles += FPlatformTime::Cycles() - StartTime;
 
 #if VALIDATE_MEMORY_PROTECTION
-		AllowCPUAccessToBlocks();
+		SetStaticMemoryPrivileges();
 #endif
 	}
 }
@@ -1472,11 +1507,8 @@ int32 FGPUDefragAllocator::Tick(FRelocationStats& Stats, bool bPanicDefrag)
 	// There may still be chunks being flagged as 'IsRelocating' due to immediate shrinks between calls.
 	FinishAllRelocations();
 
-#if VALIDATE_SYNC_SIZE
-	if (GGPUDefragValidateSyncSize)
-	{
-		Relocations.Empty();
-	}
+#if VALIDATE_SYNC_SIZE || VALIDATE_MOVES
+	Relocations.Empty();	
 #endif
 
 	// Sort the Free chunks.
@@ -1503,7 +1535,7 @@ int32 FGPUDefragAllocator::Tick(FRelocationStats& Stats, bool bPanicDefrag)
 #if VALIDATE_MEMORY_PROTECTION
 	{
 		FScopedDurationTimer MemProtectTimer(TimeInMemProtect);
-		PlatformRemoveCPUAccessTo(BlocksToProtect);
+		PlatformSetRelocationMemoryPrivileges(BlocksToProtect);
 		BlocksToUnProtect.Append(BlocksToProtect);
 		BlocksToProtect.Reset();
 	}
@@ -1853,8 +1885,8 @@ void FGPUDefragAllocator::Coalesce(FMemoryChunk* FreedChunk)
 		delete RightChunk;
 	}
 
-#if VALIDATE_SYNC_SIZE
-	if (GGPUDefragValidateSyncSize)
+#if TRACK_RELOCATIONS
+	if (GGPUDefragDumpRelocationsToTTY)
 	{
 		printf("FreeChunk Before Coalesce: 0x%p, %i, %i, %i\n", FreedChunk->Base, (int32)FreedChunk->Size, (int32)FreedChunk->SyncSize, (int32)FreedChunk->SyncIndex);
 	}
@@ -1865,8 +1897,8 @@ void FGPUDefragAllocator::Coalesce(FMemoryChunk* FreedChunk)
 	FreedChunk->Size += LeftSize + RightSize;
 	FreedChunk->SetSyncIndex(LatestSyncIndex, LatestSyncSize);
 
-#if VALIDATE_SYNC_SIZE
-	if (GGPUDefragValidateSyncSize)
+#if TRACK_RELOCATIONS
+	if (GGPUDefragDumpRelocationsToTTY)
 	{
 		printf("FreeChunk AFter Coalesce: 0x%p, %i, %i, %i\n", FreedChunk->Base, (int32)FreedChunk->Size, (int32)FreedChunk->SyncSize, (int32)FreedChunk->SyncIndex);
 	}
diff --git a/Engine/Source/Runtime/RHI/Public/GPUDefragAllocator.h b/Engine/Source/Runtime/RHI/Public/GPUDefragAllocator.h
--- a/Engine/Source/Runtime/RHI/Public/GPUDefragAllocator.h
+++ b/Engine/Source/Runtime/RHI/Public/GPUDefragAllocator.h
@@ -6,7 +6,10 @@
 
 #define LOG_EVERY_ALLOCATION			0
 #define DUMP_ALLOC_FREQUENCY			0 // 100
-#define VALIDATE_SYNC_SIZE				(!(UE_BUILD_TEST || UE_BUILD_SHIPPING))
+
+#define VALIDATE_SYNC_SIZE				1 //validates that the CPU is returning memory in an allocate that's part of an active GPU move.
+#define VALIDATE_MOVES					1 //validates that GPU moves in a given frame do not overlap destination areas.
+#define TRACK_RELOCATIONS (VALIDATE_SYNC_SIZE || VALIDATE_MOVES)
 #define VALIDATE_MEMORY_PROTECTION		1 //(!(UE_BUILD_TEST || UE_BUILD_SHIPPING))
 
 /*-----------------------------------------------------------------------------
@@ -655,17 +658,10 @@ public:
 
 protected:
 
-#if VALIDATE_SYNC_SIZE
+#if TRACK_RELOCATIONS
 	struct FRelocationEntry
 	{
-		FRelocationEntry(const uint8* InOldBase, const uint8* InNewBase, uint64 InSize, uint64 InSyncIndex)
-		: OldBase(InOldBase)
-		, NewBase(InNewBase)
-		, Size(InSize)
-		, SyncIndex(InSyncIndex)
-		{
-			printf("Relocation: 0x%p to 0x%p, %i, %i\n", OldBase, NewBase, (int32)Size, (int32)SyncIndex);
-		}
+		FRelocationEntry(const uint8* InOldBase, const uint8* InNewBase, uint64 InSize, uint64 InSyncIndex);
 
 		const uint8* OldBase;
 		const uint8* NewBase;
@@ -674,20 +670,23 @@ protected:
 	};
 
 	void ValidateRelocations(uint8* UsedBaseAddress, uint64 Size);
+	
 	TArray<FRelocationEntry> Relocations;
 #endif
 
 #if VALIDATE_MEMORY_PROTECTION
 	struct FMemProtectTracker
 	{
-		FMemProtectTracker(const void* InMemory, uint64 InBlockSize, uint32 InSyncIndex)
+		FMemProtectTracker(const void* InMemory, const void* InUserPayload, uint64 InBlockSize, uint32 InSyncIndex)
 		: Memory(InMemory)
+		, UserPayload(InUserPayload)
 		, BlockSize(InBlockSize)
 		, SyncIndex(InSyncIndex)
 		{
 
 		}
 		const void* Memory;
+		const void* UserPayload;
 		uint64 BlockSize;
 		uint32 SyncIndex;
 	};
@@ -695,9 +694,24 @@ protected:
 	TArray<FMemProtectTracker> BlocksToProtect;
 	TArray<FMemProtectTracker> BlocksToUnProtect;
 
-	void AllowCPUAccessToBlocks();
-	virtual void	PlatformAllowCPUAccessToBlock(const FMemProtectTracker& BlocksToAllow) {};
-	virtual void	PlatformRemoveCPUAccessTo(const TArray<FMemProtectTracker>& BlocksToRemove) {};
+	static int32 GGPUDefragDumpRelocationsToTTY;
+	static FAutoConsoleVariableRef CVarGPUDefragDumpRelocationsToTTY;
+
+	/** Sets Static memory privileges on blocks that have completed relocations. */
+	void SetStaticMemoryPrivileges();
+
+	/** Removes all CPU and GPU read/write privileges from the given Block.  Used for Free memory that is not part of a relocation. */
+	virtual void	PlatformSetNoMemoryPrivileges(const FMemProtectTracker& Block) {};
+
+	/** Allows all CPU and GPU read/write privileges on the given Block.  Used to reset privileges on allocation */
+	virtual void	PlatformSetStandardMemoryPrivileges(const FMemProtectTracker& Block) {};
+	
+	/** Allows all CPU R/W, GPU Read, and GPUWrite Platform implementation can determine that it should do so. Used when UserPayload is set.  Used during defraggable registration, and after waiting on the appropriate fence. */
+	virtual void	PlatformSetStaticMemoryPrivileges(const FMemProtectTracker& BlocksToAllow) {};
+
+	/** Allows Only GPU R/W.  CPU access is a bug.  This is used when this memory is part of an active relocation */
+	virtual void	PlatformSetRelocationMemoryPrivileges(const FMemProtectTracker& BlocksToAllow) {};
+	virtual void	PlatformSetRelocationMemoryPrivileges(const TArray<FMemProtectTracker>& BlocksToRemove) {};
 #endif
 
 
@@ -771,8 +785,8 @@ protected:
 		if (!bBenchmarkMode)
 		{
 #if VALIDATE_MEMORY_PROTECTION
-			BlocksToProtect.Emplace(DestAddr, Size, CurrentSyncIndex);
-			BlocksToProtect.Emplace(Source, Size, CurrentSyncIndex);
+			BlocksToProtect.Emplace(DestAddr, UserPayload, Size, CurrentSyncIndex);
+			BlocksToProtect.Emplace(Source, UserPayload, Size, CurrentSyncIndex);
 #endif
 			PlatformRelocate(DestAddr, Source, Size, UserPayload);
 		}
diff --git a/Engine/Source/Runtime/Renderer/Private/BasePassRendering.cpp b/Engine/Source/Runtime/Renderer/Private/BasePassRendering.cpp
--- a/Engine/Source/Runtime/Renderer/Private/BasePassRendering.cpp
+++ b/Engine/Source/Runtime/Renderer/Private/BasePassRendering.cpp
@@ -10,7 +10,7 @@
 // Changing this causes a full shader recompile
 static TAutoConsoleVariable<int32> CVarSelectiveBasePassOutputs(
 	TEXT("r.SelectiveBasePassOutputs"),
-	0,
+	1,
 	TEXT("Enables shaders to only export to relevant rendertargets.\n") \
 	TEXT(" 0: Export in all rendertargets.\n") \
 	TEXT(" 1: Export only into relevant rendertarget.\n"),
diff --git a/Engine/Source/Runtime/SlateCore/Private/Widgets/SWindow.cpp b/Engine/Source/Runtime/SlateCore/Private/Widgets/SWindow.cpp
--- a/Engine/Source/Runtime/SlateCore/Private/Widgets/SWindow.cpp
+++ b/Engine/Source/Runtime/SlateCore/Private/Widgets/SWindow.cpp
@@ -1119,6 +1119,11 @@ void SWindow::DestroyWindowImmediately()
 void SWindow::NotifyWindowBeingDestroyed()
 {
 	OnWindowClosed.ExecuteIfBound( SharedThis( this ) );
+
+	// Logging to track down window shutdown issues with movie loading threads. Too spammy in editor builds with all the windows
+#if !WITH_EDITOR
+	UE_LOG(LogSlate, Log, TEXT("Window '%s' being destroyed"), *GetTitle().ToString() );
+#endif
 }
 
 /** Make the window visible */
diff --git a/Engine/Source/Runtime/UMG/Private/Animation/UMGSequencePlayer.cpp b/Engine/Source/Runtime/UMG/Private/Animation/UMGSequencePlayer.cpp
--- a/Engine/Source/Runtime/UMG/Private/Animation/UMGSequencePlayer.cpp
+++ b/Engine/Source/Runtime/UMG/Private/Animation/UMGSequencePlayer.cpp
@@ -39,6 +39,10 @@ void UUMGSequencePlayer::InitSequencePlayer( const UWidgetAnimation& InAnimation
 			TArray<UObject*>& Objects = GuidToRuntimeObjectMap.FindOrAdd(Binding.AnimationGuid);
 			Objects.Add(FoundObject);
 		}
+		else
+		{
+			UE_LOG(LogUMG, Warning, TEXT("Failed to find runtime objects for %s animation, WidgetName: %s, SlotName: %s"), *InAnimation.GetPathName(), *Binding.WidgetName.ToString(), *Binding.SlotWidgetName.ToString() );
+		}
 	}
 }
 
@@ -185,10 +189,6 @@ void UUMGSequencePlayer::GetRuntimeObjects(TSharedRef<FMovieSceneSequenceInstanc
 	{
 		OutObjects.Append(*FoundObjects);
 	}
-	else
-	{
-		UE_LOG( LogUMG, Warning, TEXT("Failed to find runtime objects for %s animation"), Animation ? *Animation->GetPathName() : TEXT("(none)") );
-	}
 }
 
 EMovieScenePlayerStatus::Type UUMGSequencePlayer::GetPlaybackStatus() const
diff --git a/Engine/Source/Runtime/UMG/Private/Slate/WidgetRenderer.cpp b/Engine/Source/Runtime/UMG/Private/Slate/WidgetRenderer.cpp
--- a/Engine/Source/Runtime/UMG/Private/Slate/WidgetRenderer.cpp
+++ b/Engine/Source/Runtime/UMG/Private/Slate/WidgetRenderer.cpp
@@ -39,21 +39,36 @@ ISlate3DRenderer* FWidgetRenderer::GetSlateRenderer()
 
 UTextureRenderTarget2D* FWidgetRenderer::DrawWidget(const TSharedRef<SWidget>& Widget, FVector2D DrawSize)
 {
-	UTextureRenderTarget2D* RenderTarget = FWidgetRenderer::CreateTargetFor(DrawSize, TF_Bilinear, bUseGammaSpace);
+	if ( !IsRunningDedicatedServer() )
+	{
+		UTextureRenderTarget2D* RenderTarget = FWidgetRenderer::CreateTargetFor(DrawSize, TF_Bilinear, bUseGammaSpace);
+
+		DrawWidget(RenderTarget, Widget, DrawSize, 0);
 
-	DrawWidget(RenderTarget, Widget, DrawSize, 0);
+		return RenderTarget;
+	}
 
-	return RenderTarget;
+	return nullptr;
 }
 
 UTextureRenderTarget2D* FWidgetRenderer::CreateTargetFor(FVector2D DrawSize, TextureFilter InFilter, bool bUseGammaCorrection)
 {
-	UTextureRenderTarget2D* RenderTarget = NewObject<UTextureRenderTarget2D>();
-	RenderTarget->Filter = InFilter;
-	RenderTarget->ClearColor = FLinearColor::Transparent;
-	const bool bIsLinearSpace = !bUseGammaCorrection;
-	RenderTarget->InitCustomFormat(DrawSize.X, DrawSize.Y, PF_B8G8R8A8, bIsLinearSpace);
-	return RenderTarget;
+	if ( !IsRunningDedicatedServer() )
+	{
+		const bool bIsLinearSpace = !bUseGammaCorrection;
+
+		UTextureRenderTarget2D* RenderTarget = NewObject<UTextureRenderTarget2D>();
+		RenderTarget->Filter = InFilter;
+		RenderTarget->ClearColor = FLinearColor::Transparent;
+		RenderTarget->SRGB = bIsLinearSpace;
+		RenderTarget->TargetGamma = 1;
+		RenderTarget->InitCustomFormat(DrawSize.X, DrawSize.Y, PF_B8G8R8A8, bIsLinearSpace);
+		RenderTarget->UpdateResourceImmediate(true);
+
+		return RenderTarget;
+	}
+
+	return nullptr;
 }
 
 void FWidgetRenderer::DrawWidget(UTextureRenderTarget2D* RenderTarget, const TSharedRef<SWidget>& Widget, FVector2D DrawSize, float DeltaTime)
diff --git a/Engine/Source/Runtime/UMG/Private/UserWidget.cpp b/Engine/Source/Runtime/UMG/Private/UserWidget.cpp
--- a/Engine/Source/Runtime/UMG/Private/UserWidget.cpp
+++ b/Engine/Source/Runtime/UMG/Private/UserWidget.cpp
@@ -900,17 +900,7 @@ void UUserWidget::ListenForInputAction( FName ActionName, TEnumAsByte< EInputEve
 {
 	if ( !InputComponent )
 	{
-		if ( APlayerController* Controller = GetOwningPlayer() )
-		{
-			InputComponent = NewObject< UInputComponent >( this, NAME_None, RF_Transient );
-			InputComponent->bBlockInput = bStopAction;
-			InputComponent->Priority = Priority;
-			Controller->PushInputComponent( InputComponent );
-		}
-		else
-		{
-			FMessageLog("PIE").Info(FText::Format(LOCTEXT("NoInputListeningWithoutPlayerController", "Unable to listen to input actions without a player controller in {0}."), FText::FromName(GetClass()->GetFName())));
-		}
+		InitializeInputComponent();
 	}
 
 	if ( InputComponent )
@@ -1002,6 +992,21 @@ void UUserWidget::OnInputAction( FOnInputAction Callback )
 	}
 }
 
+void UUserWidget::InitializeInputComponent()
+{
+	if ( APlayerController* Controller = GetOwningPlayer() )
+	{
+		InputComponent = NewObject< UInputComponent >( this, NAME_None, RF_Transient );
+		InputComponent->bBlockInput = bStopAction;
+		InputComponent->Priority = Priority;
+		Controller->PushInputComponent( InputComponent );
+	}
+	else
+	{
+		FMessageLog("PIE").Info(FText::Format(LOCTEXT("NoInputListeningWithoutPlayerController", "Unable to listen to input actions without a player controller in {0}."), FText::FromName(GetClass()->GetFName())));
+	}
+}
+
 void UUserWidget::NativePaint( FPaintContext& InContext ) const
 {
 	if ( bCanEverPaint )
diff --git a/Engine/Source/Runtime/UMG/Public/Blueprint/UserWidget.h b/Engine/Source/Runtime/UMG/Public/Blueprint/UserWidget.h
--- a/Engine/Source/Runtime/UMG/Public/Blueprint/UserWidget.h
+++ b/Engine/Source/Runtime/UMG/Public/Blueprint/UserWidget.h
@@ -961,6 +961,8 @@ protected:
 
 	void OnInputAction( FOnInputAction Callback );
 
+	virtual void InitializeInputComponent();
+
 	UPROPERTY( transient )
 	class UInputComponent* InputComponent;
 
