Copying //UE4/Release-Staging-4.12 to //UE4/Dev-Main (Source: //UE4/Release-4.12 @ 2955635)

diff --git a/Engine/Source/Runtime/Engine/Private/World.cpp b/Engine/Source/Runtime/Engine/Private/World.cpp
--- a/Engine/Source/Runtime/Engine/Private/World.cpp
+++ b/Engine/Source/Runtime/Engine/Private/World.cpp
@@ -4869,6 +4869,8 @@ UWorld* FSeamlessTravelHandler::Tick()
 				}
 			} 
 
+			TArray<AActor*> ActuallyKeptActors;
+
 			// rename dynamic actors in the old world's PersistentLevel that we want to keep into the new world
 			for (FActorIterator It(CurrentWorld); It; ++It)
 			{
@@ -4886,21 +4888,7 @@ UWorld* FSeamlessTravelHandler::Tick()
 				if (bSameLevel && (bShouldKeep || (TheActor->Role < ROLE_Authority && !bDormant && !TheActor->IsNetStartupActor())) && !TheActor->IsA(ALevelScriptActor::StaticClass()))
 				{
 					It.ClearCurrent(); //@warning: invalidates *It until next iteration
-					KeepAnnotation.Clear(TheActor);
-					TheActor->Rename(NULL, LoadedWorld->PersistentLevel);
-					// if it's a Controller or a Pawn, add it to the appropriate list in the new world's WorldSettings
-					if (Cast<AController>(TheActor))
-					{
-						LoadedWorld->AddController(static_cast<AController*>(TheActor));
-					}
-					else if (Cast<APawn>(TheActor))
-					{
-						LoadedWorld->AddPawn(static_cast<APawn*>(TheActor));
-					}
-					// add to new world's actor list and remove from old
-					LoadedWorld->PersistentLevel->Actors.Add(TheActor);
-
-					TheActor->bActorSeamlessTraveled = true;
+					ActuallyKeptActors.Add(TheActor);
 				}
 				else
 				{
@@ -4921,6 +4909,28 @@ UWorld* FSeamlessTravelHandler::Tick()
 				}
 			}
 
+			// Second pass to rename and move actors that need to transition into the new world
+			// This is done after cleaning up actors that aren't transitioning in case those actors depend on these
+			// actors being in the same world.
+			for (AActor* const TheActor : ActuallyKeptActors)
+			{
+				KeepAnnotation.Clear(TheActor);
+				TheActor->Rename(NULL, LoadedWorld->PersistentLevel);
+				// if it's a Controller or a Pawn, add it to the appropriate list in the new world's WorldSettings
+				if (Cast<AController>(TheActor))
+				{
+					LoadedWorld->AddController(static_cast<AController*>(TheActor));
+				}
+				else if (Cast<APawn>(TheActor))
+				{
+					LoadedWorld->AddPawn(static_cast<APawn*>(TheActor));
+				}
+				// add to new world's actor list and remove from old
+				LoadedWorld->PersistentLevel->Actors.Add(TheActor);
+
+				TheActor->bActorSeamlessTraveled = true;
+			}
+
 			CopyWorldData(); // This copies the net driver too (LoadedWorld now has whatever NetDriver was previously held by CurrentWorld)
 
 			// only consider session ended if we're making the final switch so that HUD, etc. UI elements stay around until the end
diff --git a/Engine/Source/Runtime/Engine/Public/Animation/AnimInstanceProxy.h b/Engine/Source/Runtime/Engine/Public/Animation/AnimInstanceProxy.h
--- a/Engine/Source/Runtime/Engine/Public/Animation/AnimInstanceProxy.h
+++ b/Engine/Source/Runtime/Engine/Public/Animation/AnimInstanceProxy.h
@@ -171,13 +171,13 @@ public:
 		return LODLevel;
 	}
 
-	/** Get the current skeleton we are using */
+	/** Get the current skeleton we are using. Note that this will return nullptr outside of pre/post update */
 	USkeleton* GetSkeleton() 
 	{ 
 		return Skeleton; 
 	}
 
-	/** Get the current skeletal mesh component we are running on */
+	/** Get the current skeletal mesh component we are running on. Note that this will return nullptr outside of pre/post update */
 	USkeletalMeshComponent* GetSkelMeshComponent() 
 	{ 
 		return SkeletalMeshComponent; 
@@ -284,6 +284,9 @@ protected:
 	/** Updates the anim graph */
 	virtual void UpdateAnimationNode(float DeltaSeconds);
 
+	/** Called on the game thread pre-evaluate. */
+	virtual void PreEvaluateAnimation(UAnimInstance* InAnimInstance);
+
 	/** 
 	 * Evaluate override point 
 	 * @return true if this function is implemented, false otherwise.
@@ -294,6 +297,16 @@ protected:
 	/** Called after update so we can copy any data we need */
 	virtual void PostUpdate(UAnimInstance* InAnimInstance) const;
 
+	/** Copy any UObjects we might be using. Called Pre-update and pre-evaluate. */
+	virtual void InitializeObjects(UAnimInstance* InAnimInstance);
+
+	/** 
+	 * Clear any UObjects we might be using. Called at the end of the post-evaluate phase.
+	 * This is to ensure that objects are not used by anything apart from animation nodes.
+	 * Please make sure to call the base implementation if this is overridden.
+	 */
+	virtual void ClearObjects();
+
 	/** Calls Update(), updates the anim graph, ticks asset players */
 	void UpdateAnimation();
 
@@ -471,12 +484,10 @@ private:
 	/** Our anim blueprint generated class */
 	IAnimClassInterface* AnimClassInterface;
 
-	/** Skeleton we are using, only used for comparison purposes */
-	UPROPERTY(transient)
+	/** Skeleton we are using, only used for comparison purposes. Note that this will be nullptr outside of pre/post update */
 	USkeleton* Skeleton;
 
-	/** Skeletal mesh component we are attached to */
-	UPROPERTY(transient)
+	/** Skeletal mesh component we are attached to. Note that this will be nullptr outside of pre/post update */
 	USkeletalMeshComponent* SkeletalMeshComponent;
 
 	/** The last time passed into PreUpdate() */
diff --git a/Engine/Source/Runtime/Engine/Public/Animation/AnimSingleNodeInstanceProxy.h b/Engine/Source/Runtime/Engine/Public/Animation/AnimSingleNodeInstanceProxy.h
--- a/Engine/Source/Runtime/Engine/Public/Animation/AnimSingleNodeInstanceProxy.h
+++ b/Engine/Source/Runtime/Engine/Public/Animation/AnimSingleNodeInstanceProxy.h
@@ -42,6 +42,8 @@ public:
 	virtual bool Evaluate(FPoseContext& Output) override;
 	virtual void UpdateAnimationNode(float DeltaSeconds) override;
 	virtual void PostUpdate(UAnimInstance* InAnimInstance) const override;
+	virtual void InitializeObjects(UAnimInstance* InAnimInstance) override;
+	virtual void ClearObjects() override;
 
 	void SetPlaying(bool bIsPlaying)
 	{
@@ -75,10 +77,6 @@ public:
 
 	virtual void SetAnimationAsset(UAnimationAsset* NewAsset, USkeletalMeshComponent* MeshComponent, bool bIsLooping, float InPlayRate);
 
-	UAnimationAsset* GetCurrentAsset() { return CurrentAsset; }
-
-	UVertexAnimation* GetCurrentVertexAnimation() { return CurrentVertexAnim; }
-
 	void UpdateBlendspaceSamples(FVector InBlendInput);
 
 	void SetCurrentTime(float InCurrentTime)
@@ -109,10 +107,15 @@ public:
 
 	void SetReverse(bool bInReverse);
 
-	float GetLength();
-
 	void SetBlendSpaceInput(const FVector& InBlendInput);
 
+#if WITH_EDITOR
+	bool CanProcessAdditiveAnimations() const
+	{
+		return bCanProcessAdditiveAnimations;
+	}
+#endif
+
 private:
 	void InternalBlendSpaceEvaluatePose(class UBlendSpaceBase* BlendSpace, TArray<FBlendSampleData>& BlendSampleDataCache, FPoseContext& OutContext);
 
@@ -123,15 +126,13 @@ protected:
 	bool bCanProcessAdditiveAnimations;
 #endif
 
-private:
-	/** Current Asset being played **/
-	UPROPERTY(transient)
+	/** Current Asset being played. Note that this will be nullptr outside of pre/post update **/
 	UAnimationAsset* CurrentAsset;
 
-	/** Current vertex anim being played **/
-	UPROPERTY(transient)
+	/** Current vertex anim being played. Note that this will be nullptr outside of pre/post update **/
 	UVertexAnimation* CurrentVertexAnim;
 
+private:
 	/** Random cached values to play each asset **/
 	FVector BlendSpaceInput;
 
diff --git a/Engine/Source/Runtime/Engine/Public/SceneInterface.h b/Engine/Source/Runtime/Engine/Public/SceneInterface.h
--- a/Engine/Source/Runtime/Engine/Public/SceneInterface.h
+++ b/Engine/Source/Runtime/Engine/Public/SceneInterface.h
@@ -353,7 +353,7 @@ public:
 	 * Initialize the pixel inspector buffers.
 	 * @return True if implemented false otherwise.
 	 */
-	virtual bool InitializePixelInspector(FRenderTarget* BufferFinalColor, FRenderTarget* BufferDepth, FRenderTarget* BufferHDR, FRenderTarget* BufferA, FRenderTarget* BufferBCDE, int32 BufferIndex)
+	virtual bool InitializePixelInspector(FRenderTarget* BufferFinalColor, FRenderTarget* BufferSceneColor, FRenderTarget* BufferDepth, FRenderTarget* BufferHDR, FRenderTarget* BufferA, FRenderTarget* BufferBCDE, int32 BufferIndex)
 	{
 		return false;
 	}
diff --git a/Engine/Source/Runtime/EngineSettings/Classes/ConsoleSettings.h b/Engine/Source/Runtime/EngineSettings/Classes/ConsoleSettings.h
--- a/Engine/Source/Runtime/EngineSettings/Classes/ConsoleSettings.h
+++ b/Engine/Source/Runtime/EngineSettings/Classes/ConsoleSettings.h
@@ -79,6 +79,10 @@ class ENGINESETTINGS_API UConsoleSettings
 	UPROPERTY(config, EditAnywhere, Category=Colors, meta=(UIMin="0", UIMax="100", ClampMin="0", ClampMax="100"))
 	float BackgroundOpacityPercentage;
 
+	/** Whether we legacy bottom-to-top ordering or regular top-to-bottom ordering */
+	UPROPERTY(config, EditAnywhere, Category = AutoComplete)
+	bool bOrderTopToBottom;
+
 	/** The color used for text input. */
 	UPROPERTY(config, EditAnywhere, Category=Colors)
 	FColor InputColor;
diff --git a/Engine/Source/Runtime/GameplayAbilities/Private/Abilities/GameplayAbilityTargetActor.cpp b/Engine/Source/Runtime/GameplayAbilities/Private/Abilities/GameplayAbilityTargetActor.cpp
--- a/Engine/Source/Runtime/GameplayAbilities/Private/Abilities/GameplayAbilityTargetActor.cpp
+++ b/Engine/Source/Runtime/GameplayAbilities/Private/Abilities/GameplayAbilityTargetActor.cpp
@@ -135,11 +135,6 @@ bool AGameplayAbilityTargetActor::IsNetRelevantFor(const AActor* RealViewer, con
 	return Super::IsNetRelevantFor(RealViewer, ViewTarget, SrcLocation);
 }
 
-bool AGameplayAbilityTargetActor::GetReplicates() const
-{
-	return bReplicates;
-}
-
 bool AGameplayAbilityTargetActor::OnReplicatedTargetDataReceived(FGameplayAbilityTargetDataHandle& Data) const
 {
 	return true;
diff --git a/Engine/Source/Runtime/GameplayAbilities/Private/Abilities/GameplayAbilityWorldReticle.cpp b/Engine/Source/Runtime/GameplayAbilities/Private/Abilities/GameplayAbilityWorldReticle.cpp
--- a/Engine/Source/Runtime/GameplayAbilities/Private/Abilities/GameplayAbilityWorldReticle.cpp
+++ b/Engine/Source/Runtime/GameplayAbilities/Private/Abilities/GameplayAbilityWorldReticle.cpp
@@ -49,11 +49,6 @@ bool AGameplayAbilityWorldReticle::IsNetRelevantFor(const AActor* RealViewer, co
 	return Super::IsNetRelevantFor(RealViewer, ViewTarget, SrcLocation);
 }
 
-bool AGameplayAbilityWorldReticle::GetReplicates() const
-{
-	return bReplicates;
-}
-
 void AGameplayAbilityWorldReticle::SetIsTargetValid(bool bNewValue)
 {
 	if (bIsTargetValid != bNewValue)
diff --git a/Engine/Source/Runtime/GameplayAbilities/Private/Abilities/Tasks/AbilityTask_VisualizeTargeting.cpp b/Engine/Source/Runtime/GameplayAbilities/Private/Abilities/Tasks/AbilityTask_VisualizeTargeting.cpp
--- a/Engine/Source/Runtime/GameplayAbilities/Private/Abilities/Tasks/AbilityTask_VisualizeTargeting.cpp
+++ b/Engine/Source/Runtime/GameplayAbilities/Private/Abilities/Tasks/AbilityTask_VisualizeTargeting.cpp
@@ -122,7 +122,7 @@ bool UAbilityTask_VisualizeTargeting::ShouldSpawnTargetActor() const
 
 	const AGameplayAbilityTargetActor* CDO = CastChecked<AGameplayAbilityTargetActor>(TargetClass->GetDefaultObject());
 
-	const bool bReplicates = CDO->GetReplicates();
+	const bool bReplicates = CDO->GetIsReplicated();
 	const bool bIsLocallyControlled = Ability->GetCurrentActorInfo()->IsLocallyControlled();
 
 	return (bReplicates || bIsLocallyControlled);
diff --git a/Engine/Source/Runtime/GameplayAbilities/Private/Abilities/Tasks/AbilityTask_WaitTargetData.cpp b/Engine/Source/Runtime/GameplayAbilities/Private/Abilities/Tasks/AbilityTask_WaitTargetData.cpp
--- a/Engine/Source/Runtime/GameplayAbilities/Private/Abilities/Tasks/AbilityTask_WaitTargetData.cpp
+++ b/Engine/Source/Runtime/GameplayAbilities/Private/Abilities/Tasks/AbilityTask_WaitTargetData.cpp
@@ -122,7 +122,7 @@ bool UAbilityTask_WaitTargetData::ShouldSpawnTargetActor() const
 
 	const AGameplayAbilityTargetActor* CDO = CastChecked<AGameplayAbilityTargetActor>(TargetClass->GetDefaultObject());
 
-	const bool bReplicates = CDO->GetReplicates();
+	const bool bReplicates = CDO->GetIsReplicated();
 	const bool bIsLocallyControlled = Ability->GetCurrentActorInfo()->IsLocallyControlled();
 	const bool bShouldProduceTargetDataOnServer = CDO->ShouldProduceTargetDataOnServer;
 
diff --git a/Engine/Source/Runtime/GameplayAbilities/Public/Abilities/GameplayAbilityTargetActor.h b/Engine/Source/Runtime/GameplayAbilities/Public/Abilities/GameplayAbilityTargetActor.h
--- a/Engine/Source/Runtime/GameplayAbilities/Public/Abilities/GameplayAbilityTargetActor.h
+++ b/Engine/Source/Runtime/GameplayAbilities/Public/Abilities/GameplayAbilityTargetActor.h
@@ -59,7 +59,8 @@ public:
 	virtual bool OnReplicatedTargetDataReceived(FGameplayAbilityTargetDataHandle& Data) const;
 
 	/** Accessor for checking, before instantiating, if this TargetActor will replicate. */
-	bool GetReplicates() const;
+	DEPRECATED(4.12, "Call AActor::GetIsReplicated instead")
+	bool GetReplicates() const { return GetIsReplicated(); }
 
 	// ------------------------------
 	
diff --git a/Engine/Source/Runtime/GameplayAbilities/Public/Abilities/GameplayAbilityWorldReticle.h b/Engine/Source/Runtime/GameplayAbilities/Public/Abilities/GameplayAbilityWorldReticle.h
--- a/Engine/Source/Runtime/GameplayAbilities/Public/Abilities/GameplayAbilityWorldReticle.h
+++ b/Engine/Source/Runtime/GameplayAbilities/Public/Abilities/GameplayAbilityWorldReticle.h
@@ -31,7 +31,8 @@ public:
 	virtual void Tick(float DeltaSeconds) override;
 
 	/** Accessor for checking, before instantiating, if this WorldReticle will replicate. */
-	bool GetReplicates() const;
+	DEPRECATED(4.12, "Call AActor::GetIsReplicated instead")
+	bool GetReplicates() const { return GetIsReplicated(); }
 
 	// ------------------------------
 
diff --git a/Engine/Source/Runtime/Landscape/Classes/Materials/MaterialExpressionLandscapeLayerBlend.h b/Engine/Source/Runtime/Landscape/Classes/Materials/MaterialExpressionLandscapeLayerBlend.h
--- a/Engine/Source/Runtime/Landscape/Classes/Materials/MaterialExpressionLandscapeLayerBlend.h
+++ b/Engine/Source/Runtime/Landscape/Classes/Materials/MaterialExpressionLandscapeLayerBlend.h
@@ -11,7 +11,6 @@ enum ELandscapeLayerBlendType
 	LB_WeightBlend,
 	LB_AlphaBlend,
 	LB_HeightBlend,
-	LB_MAX,
 };
 
 USTRUCT()
diff --git a/Engine/Source/Runtime/Landscape/Private/LandscapeEdit.cpp b/Engine/Source/Runtime/Landscape/Private/LandscapeEdit.cpp
--- a/Engine/Source/Runtime/Landscape/Private/LandscapeEdit.cpp
+++ b/Engine/Source/Runtime/Landscape/Private/LandscapeEdit.cpp
@@ -79,6 +79,11 @@ void ULandscapeComponent::UpdateCachedBounds()
 			CachedLocalBox += CDI.GetLocalVertex(x, y);
 		}
 	}
+	if (CachedLocalBox.GetExtent().Z == 0)
+	{
+		// expand bounds to avoid flickering issues with zero-size bounds
+		CachedLocalBox.ExpandBy(FVector(0, 0, 1));
+	}
 
 	// Update collision component bounds
 	ULandscapeHeightfieldCollisionComponent* HFCollisionComponent = CollisionComponent.Get();
diff --git a/Engine/Source/Runtime/Landscape/Private/LandscapeSplines.cpp b/Engine/Source/Runtime/Landscape/Private/LandscapeSplines.cpp
--- a/Engine/Source/Runtime/Landscape/Private/LandscapeSplines.cpp
+++ b/Engine/Source/Runtime/Landscape/Private/LandscapeSplines.cpp
@@ -125,7 +125,7 @@ public:
 	{
 		// Slight Depth Bias so that the splines show up when they exactly match the target surface
 		// e.g. someone playing with splines on a newly-created perfectly-flat landscape
-		static const float DepthBias = -0.0001;
+		static const float DepthBias = 0.0001;
 
 		const FMatrix& LocalToWorld = GetLocalToWorld();
 
@@ -1581,8 +1581,15 @@ void ULandscapeSplineControlPoint::UpdateSplinePoints(bool bUpdateCollision, boo
 		Points.Emplace(StartLocation, LeftPos, RightPos, FalloffLeftPos, FalloffRightPos, 1.0f);
 	}
 
-	// Update Bounds
+	// Update bounds
 	Bounds = FBox(0);
+
+	// Sprite bounds
+	float SpriteScale = FMath::Clamp<float>(Width != 0 ? Width / 2 : SideFalloff / 4, 10, 1000);
+	Bounds += Location + FVector(0, 0, 0.75f * SpriteScale);
+	Bounds = Bounds.ExpandBy(SpriteScale);
+
+	// Points bounds
 	for (const FLandscapeSplineInterpPoint& Point : Points)
 	{
 		Bounds += Point.FalloffLeft;
diff --git a/Engine/Source/Runtime/Landscape/Private/Materials/MaterialExpressionLandscapeLayerBlend.cpp b/Engine/Source/Runtime/Landscape/Private/Materials/MaterialExpressionLandscapeLayerBlend.cpp
--- a/Engine/Source/Runtime/Landscape/Private/Materials/MaterialExpressionLandscapeLayerBlend.cpp
+++ b/Engine/Source/Runtime/Landscape/Private/Materials/MaterialExpressionLandscapeLayerBlend.cpp
@@ -6,6 +6,9 @@
 #include "EdGraph/EdGraphNode.h"
 #include "Engine/Engine.h"
 #include "EngineGlobals.h"
+#if WITH_EDITOR
+#include "MaterialGraph/MaterialGraphNode.h"
+#endif
 
 #define LOCTEXT_NAMESPACE "Landscape"
 
@@ -311,9 +314,9 @@ void UMaterialExpressionLandscapeLayerBlend::PostEditChangeProperty(FPropertyCha
 		const FName PropertyName = PropertyChangedEvent.MemberProperty->GetFName();
 		if (PropertyName == GET_MEMBER_NAME_CHECKED(UMaterialExpressionLandscapeLayerBlend, Layers))
 		{
-			if (GraphNode)
+			if (UMaterialGraphNode* MatGraphNode = Cast<UMaterialGraphNode>(GraphNode))
 			{
-				GraphNode->ReconstructNode();
+				MatGraphNode->RecreateAndLinkNode();
 			}
 		}
 	}
diff --git a/Engine/Source/Runtime/Launch/Private/LaunchEngineLoop.cpp b/Engine/Source/Runtime/Launch/Private/LaunchEngineLoop.cpp
--- a/Engine/Source/Runtime/Launch/Private/LaunchEngineLoop.cpp
+++ b/Engine/Source/Runtime/Launch/Private/LaunchEngineLoop.cpp
@@ -1442,8 +1442,10 @@ int32 FEngineLoop::PreInit( const TCHAR* CmdLine )
 	SlowTask.EnterProgressFrame(10);
 
 #if USE_LOCALIZED_PACKAGE_CACHE
-	// this loads a UObject, so we do this now after objects are brought up (otherwise, it can't load properties from config)
-	FPackageLocalizationManager::Get().InitializeFromDefaultCache();
+	FPackageLocalizationManager::Get().InitializeFromLazyCallback([](FPackageLocalizationManager& InPackageLocalizationManager)
+	{
+		InPackageLocalizationManager.InitializeFromCache(MakeShareable(new FEnginePackageLocalizationCache()));
+	});
 #endif	// USE_LOCALIZED_PACKAGE_CACHE
 
 	// Initialize the RHI.
@@ -1495,6 +1497,12 @@ int32 FEngineLoop::PreInit( const TCHAR* CmdLine )
 		// Make sure all UObject classes are registered and default properties have been initialized
 		ProcessNewlyLoadedUObjects();
 
+#if USE_LOCALIZED_PACKAGE_CACHE
+		// CoreUObject is definitely available now, so make sure the package localization cache is available
+		// This may have already been initialized from the CDO creation from ProcessNewlyLoadedUObjects
+		FPackageLocalizationManager::Get().PerformLazyInitialization();
+#endif	// USE_LOCALIZED_PACKAGE_CACHE
+
 		// Default materials may have been loaded due to dependencies when loading
 		// classes and class default objects. If not, do so now.
 		UMaterialInterface::InitDefaultMaterials();
@@ -1704,7 +1712,13 @@ int32 FEngineLoop::PreInit( const TCHAR* CmdLine )
 #if WITH_EDITOR
 				if ( GIsEditor )
 				{
-					UClass* EditorEngineClass = StaticLoadClass( UEditorEngine::StaticClass(), nullptr, TEXT("engine-ini:/Script/Engine.Engine.EditorEngine"), nullptr, LOAD_None, nullptr );
+					FString EditorEngineClassName;
+					GConfig->GetString(TEXT("/Script/Engine.Engine"), TEXT("EditorEngine"), EditorEngineClassName, GEngineIni);
+					UClass* EditorEngineClass = StaticLoadClass( UEditorEngine::StaticClass(), nullptr, *EditorEngineClassName);
+					if (EditorEngineClass == nullptr)
+					{
+						UE_LOG(LogInit, Fatal, TEXT("Failed to load Editor Engine class '%s'."), *EditorEngineClassName);
+					}
 
 					GEngine = GEditor = NewObject<UEditorEngine>(GetTransientPackage(), EditorEngineClass);
 
@@ -1717,7 +1731,15 @@ int32 FEngineLoop::PreInit( const TCHAR* CmdLine )
 				else
 #endif
 				{
-					UClass* EngineClass = StaticLoadClass( UEngine::StaticClass(), nullptr, TEXT("engine-ini:/Script/Engine.Engine.GameEngine"), nullptr, LOAD_None, nullptr );
+					FString GameEngineClassName;
+					GConfig->GetString(TEXT("/Script/Engine.Engine"), TEXT("GameEngine"), GameEngineClassName, GEngineIni);
+
+					UClass* EngineClass = StaticLoadClass( UEngine::StaticClass(), nullptr, *GameEngineClassName);
+
+					if (EngineClass == nullptr)
+					{
+						UE_LOG(LogInit, Fatal, TEXT("Failed to load Engine class '%s'."), *GameEngineClassName);
+					}
 
 					// must do this here so that the engine object that we create on the next line receives the correct property values
 					GEngine = NewObject<UEngine>(GetTransientPackage(), EngineClass);
@@ -2217,14 +2239,26 @@ int32 FEngineLoop::Init()
 	if( !GIsEditor )
 	{
 		// We're the game.
-		EngineClass = StaticLoadClass( UGameEngine::StaticClass(), nullptr, TEXT("engine-ini:/Script/Engine.Engine.GameEngine"), nullptr, LOAD_None, nullptr );
+		FString GameEngineClassName;
+		GConfig->GetString(TEXT("/Script/Engine.Engine"), TEXT("GameEngine"), GameEngineClassName, GEngineIni);
+		EngineClass = StaticLoadClass( UGameEngine::StaticClass(), nullptr, *GameEngineClassName);
+		if (EngineClass == nullptr)
+		{
+			UE_LOG(LogInit, Fatal, TEXT("Failed to load UnrealEd Engine class '%s'."), *GameEngineClassName);
+		}
 		GEngine = NewObject<UEngine>(GetTransientPackage(), EngineClass);
 	}
 	else
 	{
 #if WITH_EDITOR
 		// We're UnrealEd.
-		EngineClass = StaticLoadClass( UUnrealEdEngine::StaticClass(), nullptr, TEXT("engine-ini:/Script/Engine.Engine.UnrealEdEngine"), nullptr, LOAD_None, nullptr );
+		FString UnrealEdEngineClassName;
+		GConfig->GetString(TEXT("/Script/Engine.Engine"), TEXT("UnrealEdEngine"), UnrealEdEngineClassName, GEngineIni);
+		EngineClass = StaticLoadClass(UUnrealEdEngine::StaticClass(), nullptr, *UnrealEdEngineClassName);
+		if (EngineClass == nullptr)
+		{
+			UE_LOG(LogInit, Fatal, TEXT("Failed to load UnrealEd Engine class '%s'."), *UnrealEdEngineClassName);
+		}
 		GEngine = GEditor = GUnrealEd = NewObject<UUnrealEdEngine>(GetTransientPackage(), EngineClass);
 #else
 		check(0);
diff --git a/Engine/Source/Runtime/LevelSequence/Private/LevelSequenceObjectReference.cpp b/Engine/Source/Runtime/LevelSequence/Private/LevelSequenceObjectReference.cpp
--- a/Engine/Source/Runtime/LevelSequence/Private/LevelSequenceObjectReference.cpp
+++ b/Engine/Source/Runtime/LevelSequence/Private/LevelSequenceObjectReference.cpp
@@ -30,7 +30,7 @@ FLevelSequenceObjectReference::FLevelSequenceObjectReference(const FUniqueObject
 
 UObject* FLevelSequenceObjectReference::Resolve(UObject* InContext) const
 {
-	if (ObjectId.IsValid())
+	if (ObjectId.IsValid() && InContext != nullptr)
 	{
 	#if WITH_EDITOR
 		int32 PIEInstanceID = InContext->GetOutermost()->PIEInstanceID;
diff --git a/Engine/Source/Runtime/LevelSequence/Private/LevelSequenceSpawnRegister.cpp b/Engine/Source/Runtime/LevelSequence/Private/LevelSequenceSpawnRegister.cpp
--- a/Engine/Source/Runtime/LevelSequence/Private/LevelSequenceSpawnRegister.cpp
+++ b/Engine/Source/Runtime/LevelSequence/Private/LevelSequenceSpawnRegister.cpp
@@ -55,7 +55,9 @@ UObject* FLevelSequenceSpawnRegister::SpawnObject(const FGuid& BindingId, FMovie
 	{
 		SpawnInfo.Name = ActorName;
 		SpawnInfo.ObjectFlags = ObjectFlags;
-		SpawnInfo.Template = ObjectTemplate;
+		SpawnInfo.SpawnCollisionHandlingOverride = ESpawnActorCollisionHandlingMethod::AlwaysSpawn;
+		// @todo: Spawning with a non-CDO template is fraught with issues
+		//SpawnInfo.Template = ObjectTemplate;
 	}
 
 	FTransform SpawnTransform;
@@ -81,50 +83,15 @@ UObject* FLevelSequenceSpawnRegister::SpawnObject(const FGuid& BindingId, FMovie
 		WorldContext = GWorld;
 	}
 
-	// @todo: Remove when instanced components work correctly on spawned templates
-	TArray<UActorComponent*> TemplateInstanceComponents = ObjectTemplate->GetInstanceComponents();
-	ObjectTemplate->ClearInstanceComponents(false);
-
 	AActor* SpawnedActor = WorldContext->SpawnActorAbsolute(ObjectTemplate->GetClass(), SpawnTransform, SpawnInfo);
 	if (!SpawnedActor)
 	{
 		return nullptr;
 	}
-
-	// First, duplicate instance components
-	for (UActorComponent* TemplateComponent : TemplateInstanceComponents)
-	{
-		ObjectTemplate->AddInstanceComponent(TemplateComponent);
-
-		UActorComponent* NewComponent = DuplicateObject<UActorComponent>(TemplateComponent, SpawnedActor, TemplateComponent->GetFName());
-		SpawnedActor->AddInstanceComponent(NewComponent);
-	}
-
-	// Second, attach the components. We do this as a different pass, since instance components may be attached to each other
-	for (UActorComponent* ActorComp : SpawnedActor->GetInstanceComponents())
-	{
-		USceneComponent* SceneComp = Cast<USceneComponent>(ActorComp);
-		if (!SceneComp)
-		{
-			continue;
-		}
-		
-		USceneComponent* NewParent = SpawnedActor->GetRootComponent();
-
-		USceneComponent* OldParent = SceneComp->GetAttachParent();
-		if (OldParent)
-		{
-			FString PathName = OldParent->GetPathName(OldParent->GetOwner());
-
-			NewParent = FindObject<USceneComponent>(SpawnedActor, *PathName);
-		}
-
-		if (ensure(NewParent))
-		{
-			SceneComp->AttachToComponent(NewParent, FAttachmentTransformRules::KeepRelativeTransform, SceneComp->GetAttachSocketName());
-		}
-	}
-
+	
+	UEngine::FCopyPropertiesForUnrelatedObjectsParams CopyParams;
+	SpawnedActor->UnregisterAllComponents();
+	UEngine::CopyPropertiesForUnrelatedObjects(ObjectTemplate, SpawnedActor, CopyParams);
 	SpawnedActor->RegisterAllComponents();
 
 	// tag this actor so we know it was spawned by sequencer
diff --git a/Engine/Source/Runtime/LevelSequence/Public/LevelSequenceActor.h b/Engine/Source/Runtime/LevelSequence/Public/LevelSequenceActor.h
--- a/Engine/Source/Runtime/LevelSequence/Public/LevelSequenceActor.h
+++ b/Engine/Source/Runtime/LevelSequence/Public/LevelSequenceActor.h
@@ -9,7 +9,7 @@
 /**
  * Actor responsible for controlling a specific level sequence in the world.
  */
-UCLASS(Experimental, hideCategories=(Rendering, Physics, LOD, Activation))
+UCLASS(hideCategories=(Rendering, Physics, LOD, Activation))
 class LEVELSEQUENCE_API ALevelSequenceActor
 	: public AActor
 {
diff --git a/Engine/Source/Runtime/MovieScene/Public/IMovieSceneTrackInstance.h b/Engine/Source/Runtime/MovieScene/Public/IMovieSceneTrackInstance.h
--- a/Engine/Source/Runtime/MovieScene/Public/IMovieSceneTrackInstance.h
+++ b/Engine/Source/Runtime/MovieScene/Public/IMovieSceneTrackInstance.h
@@ -20,6 +20,7 @@ struct EMovieSceneUpdateData
 	float LastPosition;
 	bool bPreroll;
 	bool bJumpCut;
+	bool bLooped;
 	/** Indicates that this update was caused by the owning movie scene stopping playback due
 	    to the active sub-scene being deactivated. */
 	bool bSubSceneDeactivate;
@@ -32,6 +33,7 @@ struct EMovieSceneUpdateData
 		LastPosition = 0.0f;
 		bPreroll = false;
 		bJumpCut = false;
+		bLooped = false;
 		bSubSceneDeactivate = false;
 		UpdatePass = MSUP_PreUpdate;
 	}
@@ -41,6 +43,7 @@ struct EMovieSceneUpdateData
 		LastPosition = InLastPosition;
 		bPreroll = false;
 		bJumpCut = false;
+		bLooped = false;
 		bSubSceneDeactivate = false;
 		UpdatePass = MSUP_PreUpdate;
 	}
diff --git a/Engine/Source/Runtime/MovieSceneCapture/Public/Protocols/ImageSequenceProtocol.h b/Engine/Source/Runtime/MovieSceneCapture/Public/Protocols/ImageSequenceProtocol.h
--- a/Engine/Source/Runtime/MovieSceneCapture/Public/Protocols/ImageSequenceProtocol.h
+++ b/Engine/Source/Runtime/MovieSceneCapture/Public/Protocols/ImageSequenceProtocol.h
@@ -27,7 +27,7 @@ UCLASS(config=EditorPerProjectUserSettings, DisplayName="Image Encoding")
 class MOVIESCENECAPTURE_API UImageCaptureSettings : public UFrameGrabberProtocolSettings
 {
 public:
-	UImageCaptureSettings(const FObjectInitializer& Init) : UFrameGrabberProtocolSettings(Init), CompressionQuality(75) {}
+	UImageCaptureSettings(const FObjectInitializer& Init) : UFrameGrabberProtocolSettings(Init), CompressionQuality(100) {}
 
 	GENERATED_BODY()
 
diff --git a/Engine/Source/Runtime/MovieSceneTracks/MovieSceneTracks.Build.cs b/Engine/Source/Runtime/MovieSceneTracks/MovieSceneTracks.Build.cs
--- a/Engine/Source/Runtime/MovieSceneTracks/MovieSceneTracks.Build.cs
+++ b/Engine/Source/Runtime/MovieSceneTracks/MovieSceneTracks.Build.cs
@@ -28,10 +28,5 @@ public class MovieSceneTracks : ModuleRules
 				"SlateCore",
 			}
 		);
-
-		if (UEBuildConfiguration.bBuildEditor == true)
-		{
-			PrivateDependencyModuleNames.Add("UnrealEd");
-		}
 	}
 }
diff --git a/Engine/Source/Runtime/MovieSceneTracks/Private/Sections/MovieScene3DTransformSection.cpp b/Engine/Source/Runtime/MovieSceneTracks/Private/Sections/MovieScene3DTransformSection.cpp
--- a/Engine/Source/Runtime/MovieSceneTracks/Private/Sections/MovieScene3DTransformSection.cpp
+++ b/Engine/Source/Runtime/MovieSceneTracks/Private/Sections/MovieScene3DTransformSection.cpp
@@ -435,7 +435,9 @@ bool UMovieScene3DTransformSection::HasKeys(const FTransformKey& TransformKey) c
 void UMovieScene3DTransformSection::AddKey(float Time, const FTransformKey& TransformKey, EMovieSceneKeyInterpolation KeyInterpolation)
 {
 	FRichCurve* KeyCurve = GetCurveForChannelAndAxis(TransformKey.Channel, TransformKey.Axis, Translation, Rotation, Scale);
-	AddKeyToCurve(*KeyCurve, Time, TransformKey.Value, KeyInterpolation);
+	
+	bool bUnwindRotation = TransformKey.Channel == EKey3DTransformChannel::Rotation;
+	AddKeyToCurve(*KeyCurve, Time, TransformKey.Value, KeyInterpolation, bUnwindRotation);
 }
 
 
diff --git a/Engine/Source/Runtime/MovieSceneTracks/Private/TrackInstances/MovieSceneActorReferenceTrackInstance.cpp b/Engine/Source/Runtime/MovieSceneTracks/Private/TrackInstances/MovieSceneActorReferenceTrackInstance.cpp
--- a/Engine/Source/Runtime/MovieSceneTracks/Private/TrackInstances/MovieSceneActorReferenceTrackInstance.cpp
+++ b/Engine/Source/Runtime/MovieSceneTracks/Private/TrackInstances/MovieSceneActorReferenceTrackInstance.cpp
@@ -17,7 +17,8 @@ void FMovieSceneActorReferenceTrackInstance::SaveState(const TArray<TWeakObjectP
 	for (auto ObjectPtr : RuntimeObjects)
 	{
 		UObject* Object = ObjectPtr.Get();
-		if (InitActorReferenceMap.Find(Object) == nullptr)
+
+		if (!InitActorReferenceMap.Contains(Object))
 		{
 			AActor* ActorReferenceValue = PropertyBindings->GetCurrentValue<AActor*>(Object);
 			InitActorReferenceMap.Add(FObjectKey(Object), ActorReferenceValue);
@@ -37,10 +38,13 @@ void FMovieSceneActorReferenceTrackInstance::RestoreState(const TArray<TWeakObje
 			continue;
 		}
 
-		AActor** ActorReferenceValue = InitActorReferenceMap.Find(FObjectKey(Object));
-		if (ActorReferenceValue != nullptr)
+		if (InitActorReferenceMap.Contains(FObjectKey(Object)))
 		{
-			PropertyBindings->CallFunction<AActor*>(Object, ActorReferenceValue);
+			TWeakObjectPtr<AActor> ActorReferenceValue = *InitActorReferenceMap.Find(FObjectKey(Object));
+			if (ActorReferenceValue.IsValid())
+			{
+				PropertyBindings->CallFunction<AActor*>(Object, ActorReferenceValue.Get());
+			}
 		}
 	}
 
@@ -54,14 +58,17 @@ void FMovieSceneActorReferenceTrackInstance::Update(EMovieSceneUpdateData& Updat
 	FGuid ActorReferenceGuid;
 	if( ActorReferenceTrack->Eval( UpdateData.Position, UpdateData.LastPosition, ActorReferenceGuid ) )
 	{
-		AActor* CurrentActor;
-		AActor** CachedActorPtr = GuidToActorCache.Find( ActorReferenceGuid );
-
-		if (CachedActorPtr != nullptr)
+		AActor* CurrentActor = nullptr;
+		if (GuidToActorCache.Contains(ActorReferenceGuid))
 		{
-			CurrentActor = *CachedActorPtr;
+			TWeakObjectPtr<AActor> CachedActor = *GuidToActorCache.Find(ActorReferenceGuid);
+			if (CachedActor.IsValid())
+			{
+				CurrentActor = CachedActor.Get();
+			}
 		}
-		else
+
+		if (CurrentActor == nullptr)
 		{
 			TArray<TWeakObjectPtr<UObject>> RuntimeObjectsForGuid;
 			Player.GetRuntimeObjects(SequenceInstance.AsShared(), ActorReferenceGuid, RuntimeObjectsForGuid);
@@ -73,15 +80,21 @@ void FMovieSceneActorReferenceTrackInstance::Update(EMovieSceneUpdateData& Updat
 				if (ActorForGuid != nullptr)
 				{
 					CurrentActor = ActorForGuid;
-					GuidToActorCache.Add(ActorReferenceGuid, CurrentActor);
+					if (CurrentActor != nullptr)
+					{
+						GuidToActorCache.Add(ActorReferenceGuid, CurrentActor);
+					}
 				}
 			}
 		}
 
-		for (auto ObjectPtr : RuntimeObjects)
+		if (CurrentActor != nullptr)
 		{
-			UObject* Object = ObjectPtr.Get();
-			PropertyBindings->CallFunction<AActor*>(Object, &CurrentActor);
+			for (auto ObjectPtr : RuntimeObjects)
+			{
+				UObject* Object = ObjectPtr.Get();
+				PropertyBindings->CallFunction<AActor*>(Object, &CurrentActor);
+			}
 		}
 	}
 }
diff --git a/Engine/Source/Runtime/MovieSceneTracks/Private/TrackInstances/MovieSceneAudioTrackInstance.cpp b/Engine/Source/Runtime/MovieSceneTracks/Private/TrackInstances/MovieSceneAudioTrackInstance.cpp
--- a/Engine/Source/Runtime/MovieSceneTracks/Private/TrackInstances/MovieSceneAudioTrackInstance.cpp
+++ b/Engine/Source/Runtime/MovieSceneTracks/Private/TrackInstances/MovieSceneAudioTrackInstance.cpp
@@ -38,7 +38,8 @@ void FMovieSceneAudioTrackInstance::Update(EMovieSceneUpdateData& UpdateData, co
 
 	if (Player.GetPlaybackStatus() == EMovieScenePlayerStatus::Playing)
 	{
-		if (UpdateData.Position > UpdateData.LastPosition)
+		float LastPosition = UpdateData.bLooped ? UpdateData.Position : UpdateData.LastPosition;
+		if (UpdateData.Position >= LastPosition)
 		{
 			TMap<int32, TArray<UMovieSceneAudioSection*> > AudioSectionsBySectionIndex;
 			for (int32 i = 0; i < AudioSections.Num(); ++i)
@@ -68,7 +69,7 @@ void FMovieSceneAudioTrackInstance::Update(EMovieSceneUpdateData& UpdateData, co
 							UMovieSceneAudioSection* AudioSection = MovieSceneAudioSections[i];
 							if (AudioSection->IsTimeWithinAudioRange(UpdateData.Position))
 							{
-								if (!AudioSection->IsTimeWithinAudioRange(UpdateData.LastPosition) || !Component->IsPlaying())
+								if (!AudioSection->IsTimeWithinAudioRange(LastPosition) || !Component->IsPlaying())
 								{
 									PlaySound(AudioSection, Component, UpdateData.Position);
 								}
diff --git a/Engine/Source/Runtime/MovieSceneTracks/Private/TrackInstances/MovieSceneLevelVisibilityTrackInstance.cpp b/Engine/Source/Runtime/MovieSceneTracks/Private/TrackInstances/MovieSceneLevelVisibilityTrackInstance.cpp
--- a/Engine/Source/Runtime/MovieSceneTracks/Private/TrackInstances/MovieSceneLevelVisibilityTrackInstance.cpp
+++ b/Engine/Source/Runtime/MovieSceneTracks/Private/TrackInstances/MovieSceneLevelVisibilityTrackInstance.cpp
@@ -3,9 +3,6 @@
 #include "MovieSceneTracksPrivatePCH.h"
 #include "MovieSceneLevelVisibilityTrack.h"
 #include "MovieSceneLevelVisibilityTrackInstance.h"
-#if WITH_EDITOR
-#include "EditorLevelUtils.h"
-#endif
 
 
 FMovieSceneLevelVisibilityTrackInstance::FMovieSceneLevelVisibilityTrackInstance( UMovieSceneLevelVisibilityTrack& InLevelVisibilityTrack )
@@ -89,10 +86,32 @@ bool GetLevelVisibility( ULevelStreaming* Level )
 void SetLevelVisibility( ULevelStreaming* Level, bool bVisible )
 {
 #if WITH_EDITOR
-	UE_LOG( LogMovieScene, Log, TEXT( "Setting Visibility: Level=%s Visible=%s" ), *Level->GetWorldAssetPackageName(), bVisible ? TEXT( "true" ) : TEXT( "false" ) );
 	if ( GIsEditor && Level->GetWorld()->IsPlayInEditor() == false )
 	{
-		EditorLevelUtils::SetLevelVisibility( Level->GetLoadedLevel(), bVisible, false );
+		Level->bShouldBeVisibleInEditor = bVisible;
+		Level->GetWorld()->FlushLevelStreaming();
+
+		// Iterate over the level's actors
+		TTransArray<AActor*>& Actors = Level->GetLoadedLevel()->Actors;
+		for ( int32 ActorIndex = 0; ActorIndex < Actors.Num(); ++ActorIndex )
+		{
+			AActor* Actor = Actors[ActorIndex];
+			if ( Actor )
+			{
+				if (Actor->bHiddenEdLevel == bVisible )
+				{
+					Actor->bHiddenEdLevel = !bVisible;
+					if ( bVisible )
+					{
+						Actor->ReregisterAllComponents();
+					}
+					else
+					{
+						Actor->UnregisterAllComponents();
+					}
+				}
+			}
+		}
 	}
 	else
 #endif
diff --git a/Engine/Source/Runtime/MovieSceneTracks/Private/TrackInstances/MovieSceneParticleTrackInstance.cpp b/Engine/Source/Runtime/MovieSceneTracks/Private/TrackInstances/MovieSceneParticleTrackInstance.cpp
--- a/Engine/Source/Runtime/MovieSceneTracks/Private/TrackInstances/MovieSceneParticleTrackInstance.cpp
+++ b/Engine/Source/Runtime/MovieSceneTracks/Private/TrackInstances/MovieSceneParticleTrackInstance.cpp
@@ -29,7 +29,8 @@ void FMovieSceneParticleTrackInstance::Update(EMovieSceneUpdateData& UpdateData,
 	// @todo Sequencer We need something analagous to Matinee 1's particle replay tracks
 	// What we have here is simple toggling/triggering
 
-	if (UpdateData.Position > UpdateData.LastPosition && Player.GetPlaybackStatus() == EMovieScenePlayerStatus::Playing)
+	float LastPosition = UpdateData.bLooped ? UpdateData.Position : UpdateData.LastPosition;
+	if (UpdateData.Position >= LastPosition && Player.GetPlaybackStatus() == EMovieScenePlayerStatus::Playing)
 	{
 		const TArray<UMovieSceneSection*> Sections = ParticleTrack->GetAllParticleSections();
 		EParticleKey::Type ParticleKey = EParticleKey::Deactivate;
@@ -45,7 +46,7 @@ void FMovieSceneParticleTrackInstance::Update(EMovieSceneUpdateData& UpdateData,
 				if ( ParticleKeyCurve.IsKeyHandleValid( PreviousHandle ) )
 				{
 					FIntegralKey& PreviousKey = ParticleKeyCurve.GetKey( PreviousHandle );
-					if ( PreviousKey.Time >= UpdateData.LastPosition )
+					if ( PreviousKey.Time >= LastPosition )
 					{
 						ParticleKey = (EParticleKey::Type)PreviousKey.Value;
 						bKeyFound = true;
diff --git a/Engine/Source/Runtime/MovieSceneTracks/Private/TrackInstances/MovieSceneSkeletalAnimationTrackInstance.cpp b/Engine/Source/Runtime/MovieSceneTracks/Private/TrackInstances/MovieSceneSkeletalAnimationTrackInstance.cpp
--- a/Engine/Source/Runtime/MovieSceneTracks/Private/TrackInstances/MovieSceneSkeletalAnimationTrackInstance.cpp
+++ b/Engine/Source/Runtime/MovieSceneTracks/Private/TrackInstances/MovieSceneSkeletalAnimationTrackInstance.cpp
@@ -321,6 +321,6 @@ void FMovieSceneSkeletalAnimationTrackInstance::PreviewSetAnimPosition(USkeletal
 bool FMovieSceneSkeletalAnimationTrackInstance::ShouldUsePreviewPlayback(class IMovieScenePlayer& Player, UObject* RuntimeObject) const
 {
 	// we also use PreviewSetAnimPosition in PIE when not playing, as we can preview in PIE
-	bool bIsNotInPIEOrNotPlaying = (RuntimeObject && !RuntimeObject->GetWorld()->HasBegunPlay()) || Player.GetPlaybackStatus() != EMovieScenePlayerStatus::Playing;
+	bool bIsNotInPIEOrNotPlaying = (RuntimeObject && RuntimeObject->GetWorld() && !RuntimeObject->GetWorld()->HasBegunPlay()) || Player.GetPlaybackStatus() != EMovieScenePlayerStatus::Playing;
 	return GIsEditor && bIsNotInPIEOrNotPlaying;
 }
diff --git a/Engine/Source/Runtime/MovieSceneTracks/Private/TrackInstances/MovieSceneSubTrackInstance.cpp b/Engine/Source/Runtime/MovieSceneTracks/Private/TrackInstances/MovieSceneSubTrackInstance.cpp
--- a/Engine/Source/Runtime/MovieSceneTracks/Private/TrackInstances/MovieSceneSubTrackInstance.cpp
+++ b/Engine/Source/Runtime/MovieSceneTracks/Private/TrackInstances/MovieSceneSubTrackInstance.cpp
@@ -192,9 +192,12 @@ void FMovieSceneSubTrackInstance::Update(EMovieSceneUpdateData& UpdateData, cons
 
 	for ( TWeakObjectPtr<UMovieSceneSubSection> LastSubSection : LastTraversedSections )
 	{
-		if ( TraversedSections.Contains( LastSubSection.Get() ) == false )
+		if (LastSubSection.IsValid())
 		{
-			UpdateSection( UpdateData, Player, LastSubSection.Get(), true );
+			if ( TraversedSections.Contains( LastSubSection.Get() ) == false )
+			{
+				UpdateSection( UpdateData, Player, LastSubSection.Get(), true );
+			}
 		}
 	}
 	LastTraversedSections.Empty();
diff --git a/Engine/Source/Runtime/MovieSceneTracks/Public/TrackInstances/MovieSceneActorReferenceTrackInstance.h b/Engine/Source/Runtime/MovieSceneTracks/Public/TrackInstances/MovieSceneActorReferenceTrackInstance.h
--- a/Engine/Source/Runtime/MovieSceneTracks/Public/TrackInstances/MovieSceneActorReferenceTrackInstance.h
+++ b/Engine/Source/Runtime/MovieSceneTracks/Public/TrackInstances/MovieSceneActorReferenceTrackInstance.h
@@ -36,7 +36,8 @@ private:
 	TSharedPtr<FTrackInstancePropertyBindings> PropertyBindings;
 
 	/** Map from object to initial state */
-	TMap<FObjectKey, AActor*> InitActorReferenceMap;
+	TMap<FObjectKey, TWeakObjectPtr<AActor>> InitActorReferenceMap;
 
-	TMap<FGuid, AActor*> GuidToActorCache;
+	/** Map from guid to cached actor */
+	TMap<FGuid, TWeakObjectPtr<AActor>> GuidToActorCache;
 };
diff --git a/Engine/Source/Runtime/OpenGLDrv/Private/Android/AndroidOpenGL.cpp b/Engine/Source/Runtime/OpenGLDrv/Private/Android/AndroidOpenGL.cpp
--- a/Engine/Source/Runtime/OpenGLDrv/Private/Android/AndroidOpenGL.cpp
+++ b/Engine/Source/Runtime/OpenGLDrv/Private/Android/AndroidOpenGL.cpp
@@ -71,14 +71,6 @@ PFNGLCLEARBUFFERUIVPROC					glClearBufferuiv = NULL;
 PFNGLDRAWBUFFERSPROC					glDrawBuffers = NULL;
 PFNGLTEXBUFFEREXTPROC					glTexBufferEXT = NULL;
 
-static TAutoConsoleVariable<int32> CVarAndroidDisableTextureFormatBGRA8888(
-	TEXT("android.DisableTextureFormatBGRA8888"),
-	0,
-	TEXT("Whether to disable usage of GL_EXT_texture_format_BGRA8888 extension.\n")
-	TEXT(" 0: Enable when extension is available (default)\n")
-	TEXT(" 1: Always disabled"),
-	ECVF_ReadOnly);
-
 struct FPlatformOpenGLDevice
 {
 
@@ -346,6 +338,7 @@ bool FAndroidOpenGL::bUseES30ShadingLanguage = false;
 bool FAndroidOpenGL::bES30Support = false;
 bool FAndroidOpenGL::bES31Support = false;
 bool FAndroidOpenGL::bSupportsInstancing = false;
+bool FAndroidOpenGL::bHasHardwareHiddenSurfaceRemoval = false;
 
 void FAndroidOpenGL::ProcessExtensions(const FString& ExtensionsString)
 {
@@ -405,6 +398,13 @@ void FAndroidOpenGL::ProcessExtensions(const FString& ExtensionsString)
 		bRequiresTexture2DPrecisionHack = true;
 	}
 
+	const bool bIsPoverVRBased = RendererString.Contains(TEXT("PowerVR"));
+	if (bIsPoverVRBased)
+	{
+		bHasHardwareHiddenSurfaceRemoval = true;
+		UE_LOG(LogRHI, Log, TEXT("Enabling support for Hidden Surface Removal on PowerVR"));
+	}
+	
 	const bool bIsAdrenoBased = RendererString.Contains(TEXT("Adreno"));
 	if (bIsAdrenoBased)
 	{
@@ -495,10 +495,22 @@ void FAndroidOpenGL::ProcessExtensions(const FString& ExtensionsString)
 		bSupportsInstancing = false;
 	}
 
-	if (bSupportsBGRA8888 && CVarAndroidDisableTextureFormatBGRA8888.GetValueOnAnyThread() == 1)
+	if (bSupportsBGRA8888)
 	{
-		UE_LOG(LogRHI, Warning, TEXT("Disabling support for GL_EXT_texture_format_BGRA8888"));
-		bSupportsBGRA8888 = false;
+		// Check whether device supports BGRA as color attachment
+		GLuint FrameBuffer;
+		glGenFramebuffers(1, &FrameBuffer);
+		glBindFramebuffer(GL_DRAW_FRAMEBUFFER, FrameBuffer);
+		GLuint BGRA8888Texture;
+		glGenTextures(1, &BGRA8888Texture);
+		glBindTexture(GL_TEXTURE_2D, BGRA8888Texture);
+		glTexImage2D(GL_TEXTURE_2D, 0, GL_BGRA_EXT, 256, 256, 0, GL_BGRA_EXT, GL_UNSIGNED_BYTE, NULL);
+		glFramebufferTexture2D(GL_DRAW_FRAMEBUFFER, GL_COLOR_ATTACHMENT0, GL_TEXTURE_2D, BGRA8888Texture, 0);
+
+		bSupportsBGRA8888RenderTarget = (glCheckFramebufferStatus(GL_DRAW_FRAMEBUFFER) == GL_FRAMEBUFFER_COMPLETE);
+
+		glDeleteTextures(1, &BGRA8888Texture);
+		glDeleteFramebuffers(1, &FrameBuffer);
 	}
 }
 
diff --git a/Engine/Source/Runtime/OpenGLDrv/Private/Android/AndroidOpenGL.h b/Engine/Source/Runtime/OpenGLDrv/Private/Android/AndroidOpenGL.h
--- a/Engine/Source/Runtime/OpenGLDrv/Private/Android/AndroidOpenGL.h
+++ b/Engine/Source/Runtime/OpenGLDrv/Private/Android/AndroidOpenGL.h
@@ -147,6 +147,8 @@ struct FAndroidOpenGL : public FOpenGLES2
 		return SP_OPENGL_ES2_ANDROID;
 	}
 
+	static FORCEINLINE bool HasHardwareHiddenSurfaceRemoval() { return bHasHardwareHiddenSurfaceRemoval; };
+
 	// Optional:
 	static FORCEINLINE void QueryTimestampCounter(GLuint QueryID)
 	{
@@ -386,6 +388,9 @@ struct FAndroidOpenGL : public FOpenGLES2
 
 	// whether device supports hardware instancing
 	static bool bSupportsInstancing;
+
+	/** Whether device supports Hidden Surface Removal */
+	static bool bHasHardwareHiddenSurfaceRemoval;
 };
 
 typedef FAndroidOpenGL FOpenGL;
diff --git a/Engine/Source/Runtime/OpenGLDrv/Private/HTML5/HTML5OpenGL.cpp b/Engine/Source/Runtime/OpenGLDrv/Private/HTML5/HTML5OpenGL.cpp
--- a/Engine/Source/Runtime/OpenGLDrv/Private/HTML5/HTML5OpenGL.cpp
+++ b/Engine/Source/Runtime/OpenGLDrv/Private/HTML5/HTML5OpenGL.cpp
@@ -50,6 +50,7 @@ void FHTML5OpenGL::ProcessExtensions( const FString& ExtensionsString )
 	bSupportsRGBA8 = false;
 	// This is not color-renderable in WebGL/ANGLE (ANGLE exposes this)
 	bSupportsBGRA8888 = false;
+	bSupportsBGRA8888RenderTarget = false;
 	// ANGLE/WEBGL_depth_texture is sort of like OES_depth_texture, you just can't upload bulk data to it (via Tex*Image2D); that should be OK?
 	bSupportsDepthTexture = 
         ExtensionsString.Contains(TEXT("WEBGL_depth_texture")) ||    // Catch "WEBGL_depth_texture", "MOZ_WEBGL_depth_texture" and "WEBKIT_WEBGL_depth_texture".
diff --git a/Engine/Source/Runtime/OpenGLDrv/Private/OpenGLCommands.cpp b/Engine/Source/Runtime/OpenGLDrv/Private/OpenGLCommands.cpp
--- a/Engine/Source/Runtime/OpenGLDrv/Private/OpenGLCommands.cpp
+++ b/Engine/Source/Runtime/OpenGLDrv/Private/OpenGLCommands.cpp
@@ -2806,7 +2806,7 @@ void FOpenGLDynamicRHI::RHIDrawIndexedPrimitive(FIndexBufferRHIParamRef IndexBuf
 	VerifyProgramPipeline();
 #endif
 	
-	// @todo Workaround for radr://15076670 "Incorrect gl_VertexID in GLSL for glDrawElementsInstanced without vertex streams on Nvidia” Alternative fix that avoids exposing the messy details to the Renderer, keeping it here in the RHI.
+	// @todo Workaround for radr://15076670 "Incorrect gl_VertexID in GLSL for glDrawElementsInstanced without vertex streams on Nvidia" Alternative fix that avoids exposing the messy details to the Renderer, keeping it here in the RHI.
 	// This workaround has performance and correctness implications - it is only needed on Mac + OpenGL + Nvidia and will
 	// break AMD drivers entirely as it is technically an abuse of the OpenGL specification. Consequently it is deliberately
 	// compiled out for other platforms. Apple have closed the bug claiming the NV behaviour is permitted by the GL spec.
@@ -2855,7 +2855,7 @@ void FOpenGLDynamicRHI::RHIDrawIndexedPrimitive(FIndexBufferRHIParamRef IndexBuf
 		REPORT_GL_DRAW_RANGE_ELEMENTS_EVENT_FOR_FRAME_DUMP(DrawMode, MinIndex, MinIndex + NumVertices, NumElements, IndexType, (void *)StartIndex);
 	}
 	
-	// @todo Workaround for radr://15076670 "Incorrect gl_VertexID in GLSL for glDrawElementsInstanced without vertex streams on Nvidia”
+	// @todo Workaround for radr://15076670 "Incorrect gl_VertexID in GLSL for glDrawElementsInstanced without vertex streams on Nvidia"
 #if PLATFORM_MAC
 	if(bAttributeLessDraw)
 	{
diff --git a/Engine/Source/Runtime/OpenGLDrv/Private/OpenGLDevice.cpp b/Engine/Source/Runtime/OpenGLDrv/Private/OpenGLDevice.cpp
--- a/Engine/Source/Runtime/OpenGLDrv/Private/OpenGLDevice.cpp
+++ b/Engine/Source/Runtime/OpenGLDrv/Private/OpenGLDevice.cpp
@@ -480,8 +480,17 @@ void InitDebugContext()
 #endif
 }
 
-TAutoConsoleVariable<FString> CVarOpenGLStripExtensions(TEXT("r.OpenGL.StripExtensions"), TEXT(""), TEXT("List of comma separated OpenGL extensions to remove from driver reported extensions string"));
-TAutoConsoleVariable<FString> CVarOpenGLAddExtensions(TEXT("r.OpenGL.AddExtensions"), TEXT(""), TEXT("List of comma separated OpenGL extensions to driver reported extensions string"));
+TAutoConsoleVariable<FString> CVarOpenGLStripExtensions(
+	TEXT("r.OpenGL.StripExtensions"), 
+	TEXT(""), 
+	TEXT("List of comma separated OpenGL extensions to strip from a driver reported extensions string"),
+	ECVF_ReadOnly);
+
+TAutoConsoleVariable<FString> CVarOpenGLAddExtensions(
+	TEXT("r.OpenGL.AddExtensions"), 
+	TEXT(""), 
+	TEXT("List of comma separated OpenGL extensions to add to a driver reported extensions string"),
+	ECVF_ReadOnly);
 
 void ApplyExtensionsOverrides(FString& ExtensionsString)
 {
diff --git a/Engine/Source/Runtime/OpenGLDrv/Private/OpenGLES2.cpp b/Engine/Source/Runtime/OpenGLDrv/Private/OpenGLES2.cpp
--- a/Engine/Source/Runtime/OpenGLDrv/Private/OpenGLES2.cpp
+++ b/Engine/Source/Runtime/OpenGLDrv/Private/OpenGLES2.cpp
@@ -40,6 +40,9 @@ bool FOpenGLES2::bSupportsRGBA8 = false;
 /** GL_APPLE_texture_format_BGRA8888 */
 bool FOpenGLES2::bSupportsBGRA8888 = false;
 
+/** Whether BGRA supported as color attachment */
+bool FOpenGLES2::bSupportsBGRA8888RenderTarget = false;
+
 /** GL_EXT_discard_framebuffer */
 bool FOpenGLES2::bSupportsDiscardFrameBuffer = false;
 
@@ -205,6 +208,7 @@ void FOpenGLES2::ProcessExtensions( const FString& ExtensionsString )
 	bTimerQueryCanBeDisjoint = !ExtensionsString.Contains(TEXT("GL_NV_timer_query"));
 	bSupportsRGBA8 = ExtensionsString.Contains(TEXT("GL_OES_rgb8_rgba8"));
 	bSupportsBGRA8888 = ExtensionsString.Contains(TEXT("GL_APPLE_texture_format_BGRA8888")) || ExtensionsString.Contains(TEXT("GL_IMG_texture_format_BGRA8888")) || ExtensionsString.Contains(TEXT("GL_EXT_texture_format_BGRA8888"));
+	bSupportsBGRA8888RenderTarget = bSupportsBGRA8888;
 	bSupportsVertexHalfFloat = ExtensionsString.Contains(TEXT("GL_OES_vertex_half_float"));
 	bSupportsTextureFloat = ExtensionsString.Contains(TEXT("GL_OES_texture_float"));
 	bSupportsTextureHalfFloat = ExtensionsString.Contains(TEXT("GL_OES_texture_half_float"));
diff --git a/Engine/Source/Runtime/OpenGLDrv/Private/OpenGLES31.cpp b/Engine/Source/Runtime/OpenGLDrv/Private/OpenGLES31.cpp
--- a/Engine/Source/Runtime/OpenGLDrv/Private/OpenGLES31.cpp
+++ b/Engine/Source/Runtime/OpenGLDrv/Private/OpenGLES31.cpp
@@ -49,6 +49,9 @@ bool FOpenGLES31::bSupportsRGBA8 = false;
 /** GL_APPLE_texture_format_BGRA8888 */
 bool FOpenGLES31::bSupportsBGRA8888 = false;
 
+/** Whether BGRA supported as color attachment */
+bool FOpenGLES31::bSupportsBGRA8888RenderTarget = false;
+
 /** GL_EXT_discard_framebuffer */
 bool FOpenGLES31::bSupportsDiscardFrameBuffer = false;
 
@@ -345,6 +348,24 @@ void FOpenGLES31::ProcessExtensions( const FString& ExtensionsString )
 		glDeleteFramebuffers(1, &FrameBuffer);
 	}
 
+	if (bSupportsBGRA8888)
+	{
+		// Check whether device supports BGRA as color attachment
+		GLuint FrameBuffer;
+		glGenFramebuffers(1, &FrameBuffer);
+		glBindFramebuffer(GL_DRAW_FRAMEBUFFER, FrameBuffer);
+		GLuint BGRA8888Texture;
+		glGenTextures(1, &BGRA8888Texture);
+		glBindTexture(GL_TEXTURE_2D, BGRA8888Texture);
+		glTexImage2D(GL_TEXTURE_2D, 0, GL_BGRA_EXT, 256, 256, 0, GL_BGRA_EXT, GL_UNSIGNED_BYTE, NULL);
+		glFramebufferTexture2D(GL_DRAW_FRAMEBUFFER, GL_COLOR_ATTACHMENT0, GL_TEXTURE_2D, BGRA8888Texture, 0);
+
+		bSupportsBGRA8888RenderTarget = (glCheckFramebufferStatus(GL_DRAW_FRAMEBUFFER) == GL_FRAMEBUFFER_COMPLETE);
+
+		glDeleteTextures(1, &BGRA8888Texture);
+		glDeleteFramebuffers(1, &FrameBuffer);
+	}
+	
 	bSupportsCopyImage = ExtensionsString.Contains(TEXT("GL_EXT_copy_image"));
 }
 
diff --git a/Engine/Source/Runtime/OpenGLDrv/Private/OpenGLTexture.cpp b/Engine/Source/Runtime/OpenGLDrv/Private/OpenGLTexture.cpp
--- a/Engine/Source/Runtime/OpenGLDrv/Private/OpenGLTexture.cpp
+++ b/Engine/Source/Runtime/OpenGLDrv/Private/OpenGLTexture.cpp
@@ -256,8 +256,14 @@ FRHITexture* FOpenGLDynamicRHI::CreateOpenGLTexture(uint32 SizeX, uint32 SizeY,
 	check( bArrayTexture != (ArraySize == 1));
 #endif
 
-	bool bNoSRGBSupport = (GMaxRHIFeatureLevel <= ERHIFeatureLevel::ES2);
-	
+	bool bNoSRGBSupport = (GMaxRHIFeatureLevel <= ERHIFeatureLevel::ES3_1);
+
+	if ((Flags & TexCreate_RenderTargetable) && Format == PF_B8G8R8A8 && !FOpenGL::SupportsBGRA8888RenderTarget())
+	{
+		// Some android devices does not support BGRA as a color attachment
+		Format = PF_R8G8B8A8;
+	}
+		
 	if (bNoSRGBSupport)
 	{
 		// Remove sRGB read flag when not supported
@@ -1468,7 +1474,7 @@ FTexture2DArrayRHIRef FOpenGLDynamicRHI::RHICreateTexture2DArray(uint32 SizeX,ui
 		NumMips = FindMaxMipmapLevel(SizeX, SizeY);
 	}
 
-	if (GMaxRHIFeatureLevel <= ERHIFeatureLevel::ES2)
+	if (GMaxRHIFeatureLevel <= ERHIFeatureLevel::ES3_1)
 	{
 		// Remove sRGB read flag when not supported
 		Flags &= ~TexCreate_SRGB;
@@ -1579,7 +1585,7 @@ FTexture3DRHIRef FOpenGLDynamicRHI::RHICreateTexture3D(uint32 SizeX,uint32 SizeY
 		NumMips = FindMaxMipmapLevel(SizeX, SizeY, SizeZ);
 	}
 
-	if (GMaxRHIFeatureLevel <= ERHIFeatureLevel::ES2)
+	if (GMaxRHIFeatureLevel <= ERHIFeatureLevel::ES3_1)
 	{
 		// Remove sRGB read flag when not supported
 		Flags &= ~TexCreate_SRGB;
diff --git a/Engine/Source/Runtime/OpenGLDrv/Public/OpenGL.h b/Engine/Source/Runtime/OpenGLDrv/Public/OpenGL.h
--- a/Engine/Source/Runtime/OpenGLDrv/Public/OpenGL.h
+++ b/Engine/Source/Runtime/OpenGLDrv/Public/OpenGL.h
@@ -116,6 +116,7 @@ public:
 	static FORCEINLINE bool SupportsShaderDepthStencilFetch()			{ return false; }
 	static FORCEINLINE bool SupportsVertexArrayBGRA()					{ return true; }
 	static FORCEINLINE bool SupportsBGRA8888()							{ return true; }
+	static FORCEINLINE bool SupportsBGRA8888RenderTarget()				{ return true; }
 	static FORCEINLINE bool SupportsSRGB()								{ return true; }
 	static FORCEINLINE bool SupportsRGBA8()								{ return true; }
 	static FORCEINLINE bool SupportsDXT()								{ return true; }
diff --git a/Engine/Source/Runtime/OpenGLDrv/Public/OpenGLES2.h b/Engine/Source/Runtime/OpenGLDrv/Public/OpenGLES2.h
--- a/Engine/Source/Runtime/OpenGLDrv/Public/OpenGLES2.h
+++ b/Engine/Source/Runtime/OpenGLDrv/Public/OpenGLES2.h
@@ -117,6 +117,7 @@ struct FOpenGLES2 : public FOpenGLBase
 	static FORCEINLINE bool SupportsMultisampledRenderToTexture()		{ return bSupportsMultisampledRenderToTexture; }
 	static FORCEINLINE bool SupportsVertexArrayBGRA()					{ return false; }
 	static FORCEINLINE bool SupportsBGRA8888()							{ return bSupportsBGRA8888; }
+	static FORCEINLINE bool SupportsBGRA8888RenderTarget()				{ return bSupportsBGRA8888RenderTarget; }
 	static FORCEINLINE bool SupportsSRGB()								{ return bSupportsSGRB; }
 	static FORCEINLINE bool SupportsRGBA8()								{ return bSupportsRGBA8; }
 	static FORCEINLINE bool SupportsDXT()								{ return bSupportsDXT; }
@@ -409,6 +410,9 @@ protected:
 	/** GL_APPLE_texture_format_BGRA8888 */
 	static bool bSupportsBGRA8888;
 
+	/** Whether BGRA supported as color attachment */
+	static bool bSupportsBGRA8888RenderTarget;
+
 	/** GL_OES_vertex_half_float */
 	static bool bSupportsVertexHalfFloat;
 
diff --git a/Engine/Source/Runtime/OpenGLDrv/Public/OpenGLES31.h b/Engine/Source/Runtime/OpenGLDrv/Public/OpenGLES31.h
--- a/Engine/Source/Runtime/OpenGLDrv/Public/OpenGLES31.h
+++ b/Engine/Source/Runtime/OpenGLDrv/Public/OpenGLES31.h
@@ -76,6 +76,7 @@ struct FOpenGLES31 : public FOpenGLBase
 	static FORCEINLINE bool SupportsMultisampledRenderToTexture()		{ return bSupportsMultisampledRenderToTexture; }
 	static FORCEINLINE bool SupportsVertexArrayBGRA()					{ return false; }
 	static FORCEINLINE bool SupportsBGRA8888()							{ return bSupportsBGRA8888; }
+	static FORCEINLINE bool SupportsBGRA8888RenderTarget()				{ return bSupportsBGRA8888RenderTarget; }
 	static FORCEINLINE bool SupportsSRGB()								{ return bSupportsSGRB || !bES2Fallback; }
 	static FORCEINLINE bool SupportsRGBA8()								{ return bSupportsRGBA8; }
 	static FORCEINLINE bool SupportsDXT()								{ return bSupportsDXT; }
@@ -862,6 +863,9 @@ protected:
 	/** GL_APPLE_texture_format_BGRA8888 */
 	static bool bSupportsBGRA8888;
 
+	/** Whether BGRA supported as color attachment */
+	static bool bSupportsBGRA8888RenderTarget;
+
 	/** GL_OES_vertex_half_float */
 	static bool bSupportsVertexHalfFloat;
 
diff --git a/Engine/Source/Runtime/RHI/Private/RHI.cpp b/Engine/Source/Runtime/RHI/Private/RHI.cpp
--- a/Engine/Source/Runtime/RHI/Private/RHI.cpp
+++ b/Engine/Source/Runtime/RHI/Private/RHI.cpp
@@ -384,6 +384,7 @@ static FName NAME_VULKAN_SM4(TEXT("SF_VULKAN_SM4"));
 static FName NAME_VULKAN_SM5(TEXT("SF_VULKAN_SM5"));
 static FName NAME_SF_METAL_SM4(TEXT("SF_METAL_SM4"));
 static FName NAME_SF_METAL_MACES3_1(TEXT("SF_METAL_MACES3_1"));
+static FName NAME_SF_METAL_MACES2(TEXT("SF_METAL_MACES2"));
 
 FName LegacyShaderPlatformToShaderFormat(EShaderPlatform Platform)
 {
@@ -430,6 +431,8 @@ FName LegacyShaderPlatformToShaderFormat(EShaderPlatform Platform)
 		return NAME_SF_METAL_SM5;
 	case SP_METAL_MACES3_1:
 		return NAME_SF_METAL_MACES3_1;
+	case SP_METAL_MACES2:
+		return NAME_SF_METAL_MACES2;
 	case SP_OPENGL_ES31_EXT:
 		return NAME_GLSL_310_ES_EXT;
 	case SP_VULKAN_SM4:
@@ -474,10 +477,11 @@ EShaderPlatform ShaderFormatToLegacyShaderPlatform(FName ShaderFormat)
 	if (ShaderFormat == NAME_VULKAN_SM4)			return SP_VULKAN_SM4;
 	if (ShaderFormat == NAME_VULKAN_SM5)			return SP_VULKAN_SM5;
 	if (ShaderFormat == NAME_VULKAN_ES3_1_ANDROID)	return SP_VULKAN_ES3_1_ANDROID;
-	if (ShaderFormat == NAME_VULKAN_ES3_1)			return SP_VULKAN_ES3_1_ANDROID;
-	if (ShaderFormat == NAME_VULKAN_ES3_1_UB)		return SP_VULKAN_ES3_1_ANDROID;
+	if (ShaderFormat == NAME_VULKAN_ES3_1)			return SP_VULKAN_PCES3_1;
+	if (ShaderFormat == NAME_VULKAN_ES3_1_UB)		return SP_VULKAN_PCES3_1;
 	if (ShaderFormat == NAME_SF_METAL_SM4)		return SP_METAL_SM4;
 	if (ShaderFormat == NAME_SF_METAL_MACES3_1)	return SP_METAL_MACES3_1;
+	if (ShaderFormat == NAME_SF_METAL_MACES2)	return SP_METAL_MACES2;
 	return SP_NumPlatforms;
 }
 
diff --git a/Engine/Source/Runtime/RHI/Public/RHIDefinitions.h b/Engine/Source/Runtime/RHI/Public/RHIDefinitions.h
--- a/Engine/Source/Runtime/RHI/Public/RHIDefinitions.h
+++ b/Engine/Source/Runtime/RHI/Public/RHIDefinitions.h
@@ -11,53 +11,54 @@
 
 enum EShaderFrequency
 {
-	SF_Vertex = 0,
-	SF_Hull = 1,
-	SF_Domain = 2,
-	SF_Pixel = 3,
-	SF_Geometry = 4,
-	SF_Compute = 5,
+	SF_Vertex			= 0,
+	SF_Hull				= 1,
+	SF_Domain			= 2,
+	SF_Pixel			= 3,
+	SF_Geometry			= 4,
+	SF_Compute			= 5,
 
-	SF_NumFrequencies = 6,
+	SF_NumFrequencies	= 6,
 
-	SF_NumBits = 3,
+	SF_NumBits			= 3,
 };
 static_assert(SF_NumFrequencies <= (1 << SF_NumBits), "SF_NumFrequencies will not fit on SF_NumBits");
 
 /** @warning: update *LegacyShaderPlatform* when the below changes */
 enum EShaderPlatform
 {
-	SP_PCD3D_SM5 = 0,
-	SP_OPENGL_SM4 = 1,
-	SP_PS4 = 2,
+	SP_PCD3D_SM5		= 0,
+	SP_OPENGL_SM4		= 1,
+	SP_PS4				= 2,
 	/** Used when running in Feature Level ES2 in OpenGL. */
-	SP_OPENGL_PCES2 = 3,
-	SP_XBOXONE = 4,
-	SP_PCD3D_SM4 = 5,
-	SP_OPENGL_SM5 = 6,
+	SP_OPENGL_PCES2		= 3,
+	SP_XBOXONE			= 4,
+	SP_PCD3D_SM4		= 5,
+	SP_OPENGL_SM5		= 6,
 	/** Used when running in Feature Level ES2 in D3D11. */
-	SP_PCD3D_ES2 = 7,
+	SP_PCD3D_ES2		= 7,
 	SP_OPENGL_ES2_ANDROID = 8,
-	SP_OPENGL_ES2_WEBGL = 9,
-	SP_OPENGL_ES2_IOS = 10,
-	SP_METAL = 11,
-	SP_OPENGL_SM4_MAC = 12,
-	SP_METAL_MRT = 13,
-	SP_OPENGL_ES31_EXT = 14,
+	SP_OPENGL_ES2_WEBGL = 9, 
+	SP_OPENGL_ES2_IOS	= 10,
+	SP_METAL			= 11,
+	SP_OPENGL_SM4_MAC	= 12,
+	SP_METAL_MRT		= 13,
+	SP_OPENGL_ES31_EXT	= 14,
 	/** Used when running in Feature Level ES3_1 in D3D11. */
-	SP_PCD3D_ES3_1 = 15,
+	SP_PCD3D_ES3_1		= 15,
 	/** Used when running in Feature Level ES3_1 in OpenGL. */
-	SP_OPENGL_PCES3_1 = 16,
-	SP_METAL_SM5 = 17,
-	SP_VULKAN_PCES3_1 = 18,
-	SP_METAL_SM4 = 19,
-	SP_VULKAN_SM4 = 20,
-	SP_VULKAN_SM5 = 21,
+	SP_OPENGL_PCES3_1	= 16,
+	SP_METAL_SM5		= 17,
+	SP_VULKAN_PCES3_1	= 18,
+	SP_METAL_SM4		= 19,
+	SP_VULKAN_SM4		= 20,
+	SP_VULKAN_SM5		= 21,
 	SP_VULKAN_ES3_1_ANDROID = 22,
-	SP_METAL_MACES3_1 = 23,
+	SP_METAL_MACES3_1 	= 23,
+	SP_METAL_MACES2		= 24,
 
-	SP_NumPlatforms = 24,
-	SP_NumBits = 5,
+	SP_NumPlatforms		= 25,
+	SP_NumBits			= 5,
 };
 static_assert(SP_NumPlatforms <= (1 << SP_NumBits), "SP_NumPlatforms will not fit on SP_NumBits");
 
@@ -159,16 +160,16 @@ enum ERasterizerCullMode
 
 enum EColorWriteMask
 {
-	CW_RED = 0x01,
+	CW_RED   = 0x01,
 	CW_GREEN = 0x02,
-	CW_BLUE = 0x04,
+	CW_BLUE  = 0x04,
 	CW_ALPHA = 0x08,
 
-	CW_NONE = 0,
-	CW_RGB = CW_RED | CW_GREEN | CW_BLUE,
-	CW_RGBA = CW_RED | CW_GREEN | CW_BLUE | CW_ALPHA,
-	CW_RG = CW_RED | CW_GREEN,
-	CW_BA = CW_BLUE | CW_ALPHA,
+	CW_NONE  = 0,
+	CW_RGB   = CW_RED | CW_GREEN | CW_BLUE,
+	CW_RGBA  = CW_RED | CW_GREEN | CW_BLUE | CW_ALPHA,
+	CW_RG    = CW_RED | CW_GREEN,
+	CW_BA    = CW_BLUE | CW_ALPHA,
 };
 
 enum ECompareFunction
@@ -183,10 +184,10 @@ enum ECompareFunction
 	CF_Always,
 
 	// Utility enumerations
-	CF_DepthNearOrEqual = (((int32)ERHIZBuffer::IsInverted != 0) ? CF_GreaterEqual : CF_LessEqual),
-	CF_DepthNear = (((int32)ERHIZBuffer::IsInverted != 0) ? CF_Greater : CF_Less),
-	CF_DepthFartherOrEqual = (((int32)ERHIZBuffer::IsInverted != 0) ? CF_LessEqual : CF_GreaterEqual),
-	CF_DepthFarther = (((int32)ERHIZBuffer::IsInverted != 0) ? CF_Less : CF_Greater),
+	CF_DepthNearOrEqual		= (((int32)ERHIZBuffer::IsInverted != 0) ? CF_GreaterEqual : CF_LessEqual),
+	CF_DepthNear			= (((int32)ERHIZBuffer::IsInverted != 0) ? CF_Greater : CF_Less),
+	CF_DepthFartherOrEqual	= (((int32)ERHIZBuffer::IsInverted != 0) ? CF_LessEqual : CF_GreaterEqual),
+	CF_DepthFarther			= (((int32)ERHIZBuffer::IsInverted != 0) ? CF_Less : CF_Greater),
 };
 
 enum EStencilOp
@@ -320,17 +321,17 @@ public:
 private:
 	enum EResourceTableDefinitions
 	{
-		RTD_NumBits_UniformBufferIndex = 8,
-		RTD_NumBits_ResourceIndex = 16,
-		RTD_NumBits_BindIndex = 8,
+		RTD_NumBits_UniformBufferIndex	= 8,
+		RTD_NumBits_ResourceIndex		= 16,
+		RTD_NumBits_BindIndex			= 8,
 
-		RTD_Mask_UniformBufferIndex = (1 << RTD_NumBits_UniformBufferIndex) - 1,
-		RTD_Mask_ResourceIndex = (1 << RTD_NumBits_ResourceIndex) - 1,
-		RTD_Mask_BindIndex = (1 << RTD_NumBits_BindIndex) - 1,
+		RTD_Mask_UniformBufferIndex		= (1 << RTD_NumBits_UniformBufferIndex) - 1,
+		RTD_Mask_ResourceIndex			= (1 << RTD_NumBits_ResourceIndex) - 1,
+		RTD_Mask_BindIndex				= (1 << RTD_NumBits_BindIndex) - 1,
 
-		RTD_Shift_BindIndex = 0,
-		RTD_Shift_ResourceIndex = RTD_Shift_BindIndex + RTD_NumBits_BindIndex,
-		RTD_Shift_UniformBufferIndex = RTD_Shift_ResourceIndex + RTD_NumBits_ResourceIndex,
+		RTD_Shift_BindIndex				= 0,
+		RTD_Shift_ResourceIndex			= RTD_Shift_BindIndex + RTD_NumBits_BindIndex,
+		RTD_Shift_UniformBufferIndex	= RTD_Shift_ResourceIndex + RTD_NumBits_ResourceIndex,
 	};
 	static_assert(RTD_NumBits_UniformBufferIndex + RTD_NumBits_ResourceIndex + RTD_NumBits_BindIndex <= sizeof(uint32)* 8, "RTD_* values must fit in 32 bits");
 };
@@ -407,41 +408,41 @@ static_assert(PT_Num <= (1 << PT_NumBits), "PT_NumBits is too small");
 enum EBufferUsageFlags
 {
 	// Mutually exclusive write-frequency flags
-	BUF_Static = 0x0001, // The buffer will be written to once.
+	BUF_Static            = 0x0001, // The buffer will be written to once.
 	BUF_Dynamic           = 0x0002, // The buffer will be written to occasionally, GPU read only, CPU write only.  The data lifetime is until the next update, or the buffer is destroyed.
-	BUF_Volatile = 0x0004, // The buffer's data will have a lifetime of one frame.  It MUST be written to each frame, or a new one created each frame.
+	BUF_Volatile          = 0x0004, // The buffer's data will have a lifetime of one frame.  It MUST be written to each frame, or a new one created each frame.
 
 	// Mutually exclusive bind flags.
-	BUF_UnorderedAccess = 0x0008, // Allows an unordered access view to be created for the buffer.
+	BUF_UnorderedAccess   = 0x0008, // Allows an unordered access view to be created for the buffer.
 
 	/** Create a byte address buffer, which is basically a structured buffer with a uint32 type. */
 	BUF_ByteAddressBuffer = 0x0020,
 	/** Create a structured buffer with an atomic UAV counter. */
-	BUF_UAVCounter = 0x0040,
+	BUF_UAVCounter        = 0x0040,
 	/** Create a buffer that can be bound as a stream output target. */
-	BUF_StreamOutput = 0x0080,
+	BUF_StreamOutput      = 0x0080,
 	/** Create a buffer which contains the arguments used by DispatchIndirect or DrawIndirect. */
-	BUF_DrawIndirect = 0x0100,
-	/**
-	* Create a buffer that can be bound as a shader resource.
-	* This is only needed for buffer types which wouldn't ordinarily be used as a shader resource, like a vertex buffer.
-	*/
-	BUF_ShaderResource = 0x0200,
+	BUF_DrawIndirect      = 0x0100,
+	/** 
+	 * Create a buffer that can be bound as a shader resource. 
+	 * This is only needed for buffer types which wouldn't ordinarily be used as a shader resource, like a vertex buffer.
+	 */
+	BUF_ShaderResource    = 0x0200,
 
 	/**
-	* Request that this buffer is directly CPU accessible
-	* (@todo josh: this is probably temporary and will go away in a few months)
-	*/
+	 * Request that this buffer is directly CPU accessible
+	 * (@todo josh: this is probably temporary and will go away in a few months)
+	 */
 	BUF_KeepCPUAccessible = 0x0400,
 
 	/**
-	* Provide information that this buffer will contain only one vertex, which should be delivered to every primitive drawn.
-	* This is necessary for OpenGL implementations, which need to handle this case very differently (and can't handle GL_HALF_FLOAT in such vertices at all).
-	*/
-	BUF_ZeroStride = 0x0800,
+	 * Provide information that this buffer will contain only one vertex, which should be delivered to every primitive drawn.
+	 * This is necessary for OpenGL implementations, which need to handle this case very differently (and can't handle GL_HALF_FLOAT in such vertices at all).
+	 */
+	BUF_ZeroStride        = 0x0800,
 
 	/** Buffer should go in fast vram (hint only) */
-	BUF_FastVRAM = 0x1000,
+	BUF_FastVRAM          = 0x1000,
 
 	// Helper bit-masks
 	BUF_AnyDynamic = (BUF_Dynamic | BUF_Volatile),
@@ -485,7 +486,7 @@ enum ERHIResourceType
 /** Flags used for texture creation */
 enum ETextureCreateFlags
 {
-	TexCreate_None = 0,
+	TexCreate_None					= 0,
 
 	// Texture can be used as a render target
 	TexCreate_RenderTargetable		= 1<<0,
@@ -541,6 +542,8 @@ enum ETextureCreateFlags
 	TexCreate_NoFastClear = 1 << 25,
 	// Texture is a depth stencil resolve target
 	TexCreate_DepthStencilResolveTarget = 1 << 26,
+	// RenderTarget will create with delta color compression
+	TexCreate_DeltaColorCompression = 1 << 27,
 };
 
 enum EAsyncComputePriority
@@ -612,13 +615,13 @@ inline bool IsPCPlatform(const EShaderPlatform Platform)
 	return Platform == SP_PCD3D_SM5 || Platform == SP_PCD3D_SM4 || Platform == SP_PCD3D_ES2 || Platform == SP_PCD3D_ES3_1 ||
 		Platform == SP_OPENGL_SM4 || Platform == SP_OPENGL_SM4_MAC || Platform == SP_OPENGL_SM5 || Platform == SP_OPENGL_PCES2 || Platform == SP_OPENGL_PCES3_1 ||
 		Platform == SP_METAL_SM4 || Platform == SP_METAL_SM5 ||
-		Platform == SP_VULKAN_PCES3_1 || Platform == SP_VULKAN_SM4 || Platform == SP_VULKAN_SM5 || Platform == SP_METAL_MACES3_1;
+		Platform == SP_VULKAN_PCES3_1 || Platform == SP_VULKAN_SM4 || Platform == SP_VULKAN_SM5 || Platform == SP_METAL_MACES3_1 || Platform == SP_METAL_MACES2;
 }
 
 /** Whether the shader platform corresponds to the ES2 feature level. */
 inline bool IsES2Platform(const EShaderPlatform Platform)
 {
-	return Platform == SP_PCD3D_ES2 || Platform == SP_OPENGL_PCES2 || Platform == SP_OPENGL_ES2_ANDROID || Platform == SP_OPENGL_ES2_WEBGL || Platform == SP_OPENGL_ES2_IOS || Platform == SP_METAL;
+	return Platform == SP_PCD3D_ES2 || Platform == SP_OPENGL_PCES2 || Platform == SP_OPENGL_ES2_ANDROID || Platform == SP_OPENGL_ES2_WEBGL || Platform == SP_OPENGL_ES2_IOS || Platform == SP_METAL_MACES2;
 }
 
 /** Whether the shader platform corresponds to the ES2/ES3.1 feature level. */
@@ -635,7 +638,7 @@ inline bool IsOpenGLPlatform(const EShaderPlatform Platform)
 
 inline bool IsMetalPlatform(const EShaderPlatform Platform)
 {
-	return Platform == SP_METAL || Platform == SP_METAL_MRT || Platform == SP_METAL_SM4 || Platform == SP_METAL_SM5 || Platform == SP_METAL_MACES3_1;
+	return Platform == SP_METAL || Platform == SP_METAL_MRT || Platform == SP_METAL_SM4 || Platform == SP_METAL_SM5 || Platform == SP_METAL_MACES3_1 || Platform == SP_METAL_MACES2;
 }
 
 inline bool IsConsolePlatform(const EShaderPlatform Platform)
@@ -677,6 +680,7 @@ inline ERHIFeatureLevel::Type GetMaxSupportedFeatureLevel(EShaderPlatform InShad
 	case SP_OPENGL_ES2_ANDROID:
 	case SP_OPENGL_ES2_WEBGL:
 	case SP_OPENGL_ES2_IOS:
+	case SP_METAL_MACES2:
 		return ERHIFeatureLevel::ES2;
 	case SP_METAL:
 	case SP_METAL_MACES3_1:
@@ -741,6 +745,8 @@ inline bool IsFeatureLevelSupported(EShaderPlatform InShaderPlatform, ERHIFeatur
 		return InFeatureLevel <= ERHIFeatureLevel::SM5;
 	case SP_METAL_MACES3_1:
 		return InFeatureLevel <= ERHIFeatureLevel::ES3_1;
+	case SP_METAL_MACES2:
+		return InFeatureLevel <= ERHIFeatureLevel::ES2;
 	default:
 		return false;
 	}
diff --git a/Engine/Source/Runtime/Renderer/Private/BasePassRendering.cpp b/Engine/Source/Runtime/Renderer/Private/BasePassRendering.cpp
--- a/Engine/Source/Runtime/Renderer/Private/BasePassRendering.cpp
+++ b/Engine/Source/Runtime/Renderer/Private/BasePassRendering.cpp
@@ -143,10 +143,23 @@ void FBasePassReflectionParameters::SetMesh(FRHICommandList& RHICmdList, FPixelS
 	// Note: GBlackCubeArrayTexture has an alpha of 0, which is needed to represent invalid data so the sky cubemap can still be applied
 	FTextureRHIParamRef CubeArrayTexture = FeatureLevel >= ERHIFeatureLevel::SM5 ? GBlackCubeArrayTexture->TextureRHI : GBlackTextureCube->TextureRHI;
 	int32 ArrayIndex = 0;
+	const FReflectionCaptureProxy* ReflectionProxy = PrimitiveSceneInfo ? PrimitiveSceneInfo->CachedReflectionCaptureProxy : nullptr;
 
-	if (PrimitiveSceneInfo && PrimitiveSceneInfo->CachedReflectionCaptureProxy)
+	FMatrix BoxTransformVal = FMatrix::Identity;
+	FVector4 PositionAndRadius = FVector::ZeroVector;
+	FVector4 BoxScalesVal = FVector::ZeroVector;
+	FVector CaptureOffsetVal = FVector::ZeroVector;
+	EReflectionCaptureShape::Type CaptureShape = EReflectionCaptureShape::Box;
+	
+
+	if (PrimitiveSceneInfo && ReflectionProxy)
 	{
-		PrimitiveSceneInfo->Scene->GetCaptureParameters(PrimitiveSceneInfo->CachedReflectionCaptureProxy, CubeArrayTexture, ArrayIndex);
+		PrimitiveSceneInfo->Scene->GetCaptureParameters(ReflectionProxy, CubeArrayTexture, ArrayIndex);
+		PositionAndRadius = FVector4(ReflectionProxy->Position, ReflectionProxy->InfluenceRadius);
+		CaptureShape = ReflectionProxy->Shape;
+		BoxTransformVal = ReflectionProxy->BoxTransform;
+		BoxScalesVal = ReflectionProxy->BoxScales;
+		CaptureOffsetVal = ReflectionProxy->CaptureOffset;
 	}
 
 	SetTextureParameter(
@@ -158,6 +171,12 @@ void FBasePassReflectionParameters::SetMesh(FRHICommandList& RHICmdList, FPixelS
 		CubeArrayTexture);
 
 	SetShaderValue(RHICmdList, PixelShaderRHI, CubemapArrayIndex, ArrayIndex);
+	SetShaderValue(RHICmdList, PixelShaderRHI, ReflectionPositionAndRadius, PositionAndRadius);
+	SetShaderValue(RHICmdList, PixelShaderRHI, ReflectionShape, (float)CaptureShape);
+	SetShaderValue(RHICmdList, PixelShaderRHI, BoxTransform, BoxTransformVal);
+	SetShaderValue(RHICmdList, PixelShaderRHI, BoxScales, BoxScalesVal);
+	SetShaderValue(RHICmdList, PixelShaderRHI, CaptureOffset, CaptureOffsetVal);
+	
 }
 
 void FTranslucentLightingParameters::Set(FRHICommandList& RHICmdList, FPixelShaderRHIParamRef PixelShaderRHI, const FViewInfo* View)
diff --git a/Engine/Source/Runtime/Renderer/Private/BasePassRendering.h b/Engine/Source/Runtime/Renderer/Private/BasePassRendering.h
--- a/Engine/Source/Runtime/Renderer/Private/BasePassRendering.h
+++ b/Engine/Source/Runtime/Renderer/Private/BasePassRendering.h
@@ -400,6 +400,11 @@ public:
 		ReflectionCubemap.Bind(ParameterMap, TEXT("ReflectionCubemap"));
 		ReflectionCubemapSampler.Bind(ParameterMap, TEXT("ReflectionCubemapSampler"));
 		CubemapArrayIndex.Bind(ParameterMap, TEXT("CubemapArrayIndex"));
+		ReflectionPositionAndRadius.Bind(ParameterMap, TEXT("ReflectionPositionAndRadius"));
+		ReflectionShape.Bind(ParameterMap, TEXT("ReflectionShape"));
+		BoxTransform.Bind(ParameterMap, TEXT("BoxTransform"));
+		BoxScales.Bind(ParameterMap, TEXT("BoxScales"));		
+		CaptureOffset.Bind(ParameterMap, TEXT("CaptureOffset"));
 		SkyLightReflectionParameters.Bind(ParameterMap);
 	}
 
@@ -414,6 +419,11 @@ public:
 		Ar << P.ReflectionCubemap;
 		Ar << P.ReflectionCubemapSampler;
 		Ar << P.CubemapArrayIndex;
+		Ar << P.ReflectionPositionAndRadius;
+		Ar << P.ReflectionShape;
+		Ar << P.BoxTransform;
+		Ar << P.BoxScales;
+		Ar << P.CaptureOffset;
 		Ar << P.SkyLightReflectionParameters;
 		return Ar;
 	}
@@ -423,7 +433,14 @@ private:
 	FPlanarReflectionParameters PlanarReflectionParameters;
 	FShaderResourceParameter ReflectionCubemap;
 	FShaderResourceParameter ReflectionCubemapSampler;
+
 	FShaderParameter CubemapArrayIndex;
+	FShaderParameter ReflectionPositionAndRadius;	
+	FShaderParameter ReflectionShape;
+	FShaderParameter BoxTransform;
+	FShaderParameter BoxScales;
+	FShaderParameter CaptureOffset;
+
 	FSkyLightReflectionParameters SkyLightReflectionParameters;
 };
 
diff --git a/Engine/Source/Runtime/Renderer/Private/DeferredShadingRenderer.cpp b/Engine/Source/Runtime/Renderer/Private/DeferredShadingRenderer.cpp
--- a/Engine/Source/Runtime/Renderer/Private/DeferredShadingRenderer.cpp
+++ b/Engine/Source/Runtime/Renderer/Private/DeferredShadingRenderer.cpp
@@ -66,6 +66,14 @@ static TAutoConsoleVariable<int32> CVarVisualizeTexturePool(
 	ECVF_Cheat | ECVF_RenderThreadSafe);
 #endif
 
+static TAutoConsoleVariable<int32> CVarClearCoatNormal(
+	TEXT("r.ClearCoatNormal"),
+	0,
+	TEXT("0 to disable clear coat normal.\n")
+	TEXT(" 0: off\n")
+	TEXT(" 1: on"),
+	ECVF_ReadOnly);
+
 static TAutoConsoleVariable<int32> CVarRHICmdFlushRenderThreadTasksBasePass(
 	TEXT("r.RHICmdFlushRenderThreadTasksBasePass"),
 	0,
diff --git a/Engine/Source/Runtime/Renderer/Private/PostProcess/PostProcessBufferInspector.cpp b/Engine/Source/Runtime/Renderer/Private/PostProcess/PostProcessBufferInspector.cpp
--- a/Engine/Source/Runtime/Renderer/Private/PostProcess/PostProcessBufferInspector.cpp
+++ b/Engine/Source/Runtime/Renderer/Private/PostProcess/PostProcessBufferInspector.cpp
@@ -152,6 +152,27 @@ void FRCPassPostProcessBufferInspector::Process(FRenderingCompositePassContext&
 				}
 			}
 
+			//////////////////////////////////////////////////////////////////////////
+			// SCENE COLOR
+			const FTexture2DRHIRef &DestinationBufferSceneColor = Scene->PixelInspectorData.RenderTargetBufferSceneColor[PixelInspectorRequest->BufferIndex]->GetRenderTargetTexture();
+			if (DestinationBufferSceneColor.IsValid())
+			{
+				const FRenderingCompositeOutputRef* OutputRef0 = GetInput(ePId_Input2);
+				if (OutputRef0)
+				{
+					FRenderingCompositeOutput* Input = OutputRef0->GetOutput();
+					if (Input && Input->PooledRenderTarget.IsValid() && Input->PooledRenderTarget->GetRenderTargetItem().ShaderResourceTexture.IsValid())
+					{
+						FTexture2DRHIRef SourceBufferSceneColor = (FRHITexture2D*)(Input->PooledRenderTarget->GetRenderTargetItem().ShaderResourceTexture.GetReference());
+						if (DestinationBufferSceneColor->GetFormat() == SourceBufferSceneColor->GetFormat())
+						{
+							FBox2D DestinationBox(FVector2D(0.0f, 0.0f), FVector2D(1.0f, 1.0f));
+							RHICmdList.CopySubTextureRegion(SourceBufferSceneColor, DestinationBufferSceneColor, SourceBox, DestinationBox);
+						}
+					}
+				}
+			}
+
 			//////////////////////////////////////////////////////////////////////////
 			// HDR
 			const FTexture2DRHIRef &DestinationBufferHDR = Scene->PixelInspectorData.RenderTargetBufferHDR[PixelInspectorRequest->BufferIndex]->GetRenderTargetTexture();
diff --git a/Engine/Source/Runtime/Renderer/Private/PostProcess/PostProcessBufferInspector.h b/Engine/Source/Runtime/Renderer/Private/PostProcess/PostProcessBufferInspector.h
--- a/Engine/Source/Runtime/Renderer/Private/PostProcess/PostProcessBufferInspector.h
+++ b/Engine/Source/Runtime/Renderer/Private/PostProcess/PostProcessBufferInspector.h
@@ -11,7 +11,7 @@
 // derives from TRenderingCompositePassBase<InputCount, OutputCount>
 // ePId_Input0: SceneColor
 // ePId_Input1: SeparateTranslucency
-class FRCPassPostProcessBufferInspector : public TRenderingCompositePassBase<2, 1>
+class FRCPassPostProcessBufferInspector : public TRenderingCompositePassBase<3, 1>
 {
 public:
 
diff --git a/Engine/Source/Runtime/Renderer/Private/PostProcess/PostProcessing.cpp b/Engine/Source/Runtime/Renderer/Private/PostProcess/PostProcessing.cpp
--- a/Engine/Source/Runtime/Renderer/Private/PostProcess/PostProcessing.cpp
+++ b/Engine/Source/Runtime/Renderer/Private/PostProcess/PostProcessing.cpp
@@ -1836,6 +1836,7 @@ void FPostProcessing::Process(FRHICommandListImmediate& RHICmdList, FViewInfo& V
 			FRenderingCompositePass* Node = Context.Graph.RegisterPass(new(FMemStack::Get()) FRCPassPostProcessBufferInspector(RHICmdList));
 			Node->SetInput(ePId_Input0, Context.FinalOutput);
 			Node->SetInput(ePId_Input1, HDRColor);
+			Node->SetInput(ePId_Input2, Context.SceneColor);
 			Context.FinalOutput = FRenderingCompositeOutputRef(Node);
 		}
 #endif //WITH_EDITOR
diff --git a/Engine/Source/Runtime/Renderer/Private/PostProcess/SceneRenderTargets.cpp b/Engine/Source/Runtime/Renderer/Private/PostProcess/SceneRenderTargets.cpp
--- a/Engine/Source/Runtime/Renderer/Private/PostProcess/SceneRenderTargets.cpp
+++ b/Engine/Source/Runtime/Renderer/Private/PostProcess/SceneRenderTargets.cpp
@@ -669,7 +669,7 @@ void FSceneRenderTargets::AllocSceneColor(FRHICommandList& RHICmdList)
 
 	// Create the scene color.
 	{
-		FPooledRenderTargetDesc Desc(FPooledRenderTargetDesc::Create2DDesc(BufferSize, SceneColorBufferFormat, FClearValueBinding::Black, TexCreate_None, TexCreate_RenderTargetable, false));
+		FPooledRenderTargetDesc Desc(FPooledRenderTargetDesc::Create2DDesc(BufferSize, SceneColorBufferFormat, FClearValueBinding::Black, TexCreate_None, TexCreate_RenderTargetable | TexCreate_DeltaColorCompression, false));
 
 		Desc.Flags |= TexCreate_FastVRAM;
 
@@ -780,7 +780,7 @@ void FSceneRenderTargets::AllocGBufferTargets(FRHICommandList& RHICmdList)
 			NormalGBufferFormat = PF_FloatRGBA;
 		}
 
-		FPooledRenderTargetDesc Desc(FPooledRenderTargetDesc::Create2DDesc(BufferSize, NormalGBufferFormat, FClearValueBinding::Transparent, TexCreate_None, TexCreate_RenderTargetable, false));
+		FPooledRenderTargetDesc Desc(FPooledRenderTargetDesc::Create2DDesc(BufferSize, NormalGBufferFormat, FClearValueBinding::Transparent, TexCreate_None, TexCreate_RenderTargetable | TexCreate_DeltaColorCompression, false));
 		GRenderTargetPool.FindFreeElement(RHICmdList, Desc, GBufferA, TEXT("GBufferA"));
 	}
 
@@ -788,26 +788,26 @@ void FSceneRenderTargets::AllocGBufferTargets(FRHICommandList& RHICmdList)
 	{
 		const EPixelFormat SpecularGBufferFormat = bHighPrecisionGBuffers ? PF_FloatRGBA : PF_B8G8R8A8;
 
-		FPooledRenderTargetDesc Desc(FPooledRenderTargetDesc::Create2DDesc(BufferSize, SpecularGBufferFormat, FClearValueBinding::Transparent, TexCreate_None, TexCreate_RenderTargetable, false));
+		FPooledRenderTargetDesc Desc(FPooledRenderTargetDesc::Create2DDesc(BufferSize, SpecularGBufferFormat, FClearValueBinding::Transparent, TexCreate_None, TexCreate_RenderTargetable | TexCreate_DeltaColorCompression, false));
 		GRenderTargetPool.FindFreeElement(RHICmdList, Desc, GBufferB, TEXT("GBufferB"));
 	}
 
 	// Create the diffuse color g-buffer.
 	{
 		const EPixelFormat DiffuseGBufferFormat = bHighPrecisionGBuffers ? PF_FloatRGBA : PF_B8G8R8A8;
-		FPooledRenderTargetDesc Desc(FPooledRenderTargetDesc::Create2DDesc(BufferSize, DiffuseGBufferFormat, FClearValueBinding::Transparent, TexCreate_SRGB, TexCreate_RenderTargetable, false));
+		FPooledRenderTargetDesc Desc(FPooledRenderTargetDesc::Create2DDesc(BufferSize, DiffuseGBufferFormat, FClearValueBinding::Transparent, TexCreate_SRGB, TexCreate_RenderTargetable | TexCreate_DeltaColorCompression, false));
 		GRenderTargetPool.FindFreeElement(RHICmdList, Desc, GBufferC, TEXT("GBufferC"));
 	}
 
 	// Create the mask g-buffer (e.g. SSAO, subsurface scattering, wet surface mask, skylight mask, ...).
 	{
-		FPooledRenderTargetDesc Desc(FPooledRenderTargetDesc::Create2DDesc(BufferSize, PF_B8G8R8A8, FClearValueBinding(FLinearColor(0, 1, 1, 1)), TexCreate_None, TexCreate_RenderTargetable, false));
+		FPooledRenderTargetDesc Desc(FPooledRenderTargetDesc::Create2DDesc(BufferSize, PF_B8G8R8A8, FClearValueBinding(FLinearColor(0, 1, 1, 1)), TexCreate_None, TexCreate_RenderTargetable | TexCreate_DeltaColorCompression, false));
 		GRenderTargetPool.FindFreeElement(RHICmdList, Desc, GBufferD, TEXT("GBufferD"));
 	}
 
 	if (bAllowStaticLighting)
 	{
-		FPooledRenderTargetDesc Desc(FPooledRenderTargetDesc::Create2DDesc(BufferSize, PF_B8G8R8A8, FClearValueBinding(FLinearColor(1, 1, 1, 1)), TexCreate_None, TexCreate_RenderTargetable, false));
+		FPooledRenderTargetDesc Desc(FPooledRenderTargetDesc::Create2DDesc(BufferSize, PF_B8G8R8A8, FClearValueBinding(FLinearColor(1, 1, 1, 1)), TexCreate_None, TexCreate_RenderTargetable | TexCreate_DeltaColorCompression, false));
 		GRenderTargetPool.FindFreeElement(RHICmdList, Desc, GBufferE, TEXT("GBufferE"));
 	}
 
diff --git a/Engine/Source/Runtime/Renderer/Private/RendererScene.cpp b/Engine/Source/Runtime/Renderer/Private/RendererScene.cpp
--- a/Engine/Source/Runtime/Renderer/Private/RendererScene.cpp
+++ b/Engine/Source/Runtime/Renderer/Private/RendererScene.cpp
@@ -163,16 +163,18 @@ FPixelInspectorData::FPixelInspectorData()
 	{
 		RenderTargetBufferFinalColor[i] = nullptr;
 		RenderTargetBufferDepth[i] = nullptr;
+		RenderTargetBufferSceneColor[i] = nullptr;
 		RenderTargetBufferHDR[i] = nullptr;
 		RenderTargetBufferA[i] = nullptr;
 		RenderTargetBufferBCDE[i] = nullptr;
 	}
 }
 
-void FPixelInspectorData::InitializeBuffers(FRenderTarget* BufferFinalColor, FRenderTarget* BufferDepth, FRenderTarget* BufferHDR, FRenderTarget* BufferA, FRenderTarget* BufferBCDE, int32 BufferIndex)
+void FPixelInspectorData::InitializeBuffers(FRenderTarget* BufferFinalColor, FRenderTarget* BufferSceneColor, FRenderTarget* BufferDepth, FRenderTarget* BufferHDR, FRenderTarget* BufferA, FRenderTarget* BufferBCDE, int32 BufferIndex)
 {
 	RenderTargetBufferFinalColor[BufferIndex] = BufferFinalColor;
 	RenderTargetBufferDepth[BufferIndex] = BufferDepth;
+	RenderTargetBufferSceneColor[BufferIndex] = BufferSceneColor;
 	RenderTargetBufferHDR[BufferIndex] = BufferHDR;
 	RenderTargetBufferA[BufferIndex] = BufferA;
 	RenderTargetBufferBCDE[BufferIndex] = BufferBCDE;
@@ -202,6 +204,12 @@ void FPixelInspectorData::InitializeBuffers(FRenderTarget* BufferFinalColor, FRe
 		check(BufferSize.X == 1 && BufferSize.Y == 1);
 	}
 
+	if (RenderTargetBufferSceneColor[BufferIndex] != nullptr)
+	{
+		BufferSize = RenderTargetBufferSceneColor[BufferIndex]->GetSizeXY();
+		check(BufferSize.X == 1 && BufferSize.Y == 1);
+	}
+
 	if (RenderTargetBufferHDR[BufferIndex] != nullptr)
 	{
 		BufferSize = RenderTargetBufferHDR[BufferIndex]->GetSizeXY();
@@ -2506,10 +2514,10 @@ void FScene::OnLevelAddedToWorld_RenderThread(FName InLevelName)
 }
 
 #if WITH_EDITOR
-bool FScene::InitializePixelInspector(FRenderTarget* BufferFinalColor, FRenderTarget* BufferDepth, FRenderTarget* BufferHDR, FRenderTarget* BufferA, FRenderTarget* BufferBCDE, int32 BufferIndex)
+bool FScene::InitializePixelInspector(FRenderTarget* BufferFinalColor, FRenderTarget* BufferSceneColor, FRenderTarget* BufferDepth, FRenderTarget* BufferHDR, FRenderTarget* BufferA, FRenderTarget* BufferBCDE, int32 BufferIndex)
 {
 	//Initialize the buffers
-	PixelInspectorData.InitializeBuffers(BufferFinalColor, BufferDepth, BufferHDR, BufferA, BufferBCDE, BufferIndex);
+	PixelInspectorData.InitializeBuffers(BufferFinalColor, BufferSceneColor, BufferDepth, BufferHDR, BufferA, BufferBCDE, BufferIndex);
 	//return true when the interface is implemented
 	return true;
 }
diff --git a/Engine/Source/Runtime/Renderer/Private/ScenePrivate.h b/Engine/Source/Runtime/Renderer/Private/ScenePrivate.h
--- a/Engine/Source/Runtime/Renderer/Private/ScenePrivate.h
+++ b/Engine/Source/Runtime/Renderer/Private/ScenePrivate.h
@@ -1632,7 +1632,7 @@ typedef TMap<FMaterial*, FMaterialShaderMap*> FMaterialsToUpdateMap;
 	public:
 		FPixelInspectorData();
 
-		void InitializeBuffers(FRenderTarget* BufferFinalColor, FRenderTarget* BufferDepth, FRenderTarget* BufferHDR, FRenderTarget* BufferA, FRenderTarget* BufferBCDE, int32 bufferIndex);
+		void InitializeBuffers(FRenderTarget* BufferFinalColor, FRenderTarget* BufferSceneColor, FRenderTarget* BufferDepth, FRenderTarget* BufferHDR, FRenderTarget* BufferA, FRenderTarget* BufferBCDE, int32 bufferIndex);
 
 		bool AddPixelInspectorRequest(FPixelInspectorRequest *PixelInspectorRequest);
 
@@ -1642,6 +1642,7 @@ typedef TMap<FMaterial*, FMaterialShaderMap*> FMaterialsToUpdateMap;
 		FRenderTarget* RenderTargetBufferDepth[2];
 		FRenderTarget* RenderTargetBufferFinalColor[2];
 		FRenderTarget* RenderTargetBufferHDR[2];
+		FRenderTarget* RenderTargetBufferSceneColor[2];
 		FRenderTarget* RenderTargetBufferA[2];
 		FRenderTarget* RenderTargetBufferBCDE[2];
 	};
@@ -2026,7 +2027,7 @@ public:
 	}
 
 #if WITH_EDITOR
-	virtual bool InitializePixelInspector(FRenderTarget* BufferFinalColor, FRenderTarget* BufferDepth, FRenderTarget* BufferHDR, FRenderTarget* BufferA, FRenderTarget* BufferBCDE, int32 BufferIndex) override;
+	virtual bool InitializePixelInspector(FRenderTarget* BufferFinalColor, FRenderTarget* BufferSceneColor, FRenderTarget* BufferDepth, FRenderTarget* BufferHDR, FRenderTarget* BufferA, FRenderTarget* BufferBCDE, int32 BufferIndex) override;
 
 	virtual bool AddPixelInspectorRequest(FPixelInspectorRequest *PixelInspectorRequest) override;
 #endif //WITH_EDITOR
diff --git a/Engine/Source/Runtime/Renderer/Private/VelocityRendering.cpp b/Engine/Source/Runtime/Renderer/Private/VelocityRendering.cpp
--- a/Engine/Source/Runtime/Renderer/Private/VelocityRendering.cpp
+++ b/Engine/Source/Runtime/Renderer/Private/VelocityRendering.cpp
@@ -828,7 +828,7 @@ FPooledRenderTargetDesc FVelocityRendering::GetRenderTargetDesc()
 {
 	const FIntPoint BufferSize = FSceneRenderTargets::Get_FrameConstantsOnly().GetBufferSizeXY();
 	const FIntPoint VelocityBufferSize = BufferSize;		// full resolution so we can reuse the existing full res z buffer
-	return FPooledRenderTargetDesc(FPooledRenderTargetDesc::Create2DDesc(VelocityBufferSize, PF_G16R16, FClearValueBinding::Transparent, TexCreate_None, TexCreate_RenderTargetable, false));
+	return FPooledRenderTargetDesc(FPooledRenderTargetDesc::Create2DDesc(VelocityBufferSize, PF_G16R16, FClearValueBinding::Transparent, TexCreate_None, TexCreate_RenderTargetable | TexCreate_DeltaColorCompression, false));
 }
 
 bool FVelocityRendering::OutputsToGBuffer()
diff --git a/Engine/Source/Runtime/ShaderCore/Private/Shader.cpp b/Engine/Source/Runtime/ShaderCore/Private/Shader.cpp
--- a/Engine/Source/Runtime/ShaderCore/Private/Shader.cpp
+++ b/Engine/Source/Runtime/ShaderCore/Private/Shader.cpp
@@ -1577,6 +1577,11 @@ void ShaderMapAppendKeyString(EShaderPlatform Platform, FString& KeyString)
 		KeyString += (CVar && CVar->GetValueOnAnyThread() != 0) ? TEXT("_DXTN") : TEXT("_BC5N");
 	}
 
+	{
+		static const auto CVar = IConsoleManager::Get().FindTConsoleVariableDataInt(TEXT("r.ClearCoatNormal"));
+		KeyString += (CVar && CVar->GetValueOnAnyThread() != 0) ? TEXT("_CCBN") : TEXT("_NoCCBN");
+	}
+
 	{
 		static const auto CVar = IConsoleManager::Get().FindTConsoleVariableDataInt(TEXT("r.CompileShadersForDevelopment"));
 		KeyString += (CVar && CVar->GetValueOnAnyThread() != 0) ? TEXT("_DEV") : TEXT("_NoDEV");
diff --git a/Engine/Source/Runtime/ShaderCore/Private/ShaderCache.cpp b/Engine/Source/Runtime/ShaderCore/Private/ShaderCache.cpp
--- a/Engine/Source/Runtime/ShaderCore/Private/ShaderCache.cpp
+++ b/Engine/Source/Runtime/ShaderCore/Private/ShaderCache.cpp
@@ -149,6 +149,7 @@ static bool ShaderPlatformCanPrebindBoundShaderState(EShaderPlatform Platform)
 		case SP_METAL_MRT:
 		case SP_METAL_SM5:
 		case SP_METAL_MACES3_1:
+		case SP_METAL_MACES2:
 		{
 			return true;
 		}
diff --git a/Engine/Source/Runtime/Slate/Private/Framework/Application/MenuStack.cpp b/Engine/Source/Runtime/Slate/Private/Framework/Application/MenuStack.cpp
--- a/Engine/Source/Runtime/Slate/Private/Framework/Application/MenuStack.cpp
+++ b/Engine/Source/Runtime/Slate/Private/Framework/Application/MenuStack.cpp
@@ -213,9 +213,11 @@ TSharedRef<IMenu> FMenuStack::Push(const FWidgetPath& InOwnerPath, const TShared
 		ActiveMethod = InMethod.IsSet() ? FPopupMethodReply::UseMethod(InMethod.GetValue()) : QueryPopupMethod(InOwnerPath);
 
 		// The host window is determined when a new root menu is pushed
+		// This must be set prior to PushInternal below, as it will be referenced if the menu being created is a new root menu.
 		SetHostPath(InOwnerPath);
 	}
 
+	TGuardValue<bool> Guard(bHostWindowGuard, true);
 	return PushInternal(ParentMenu, InContent, Anchor, TransitionEffect, bFocusImmediately, ActiveMethod.GetShouldThrottle(), bIsCollapsedByParent, bEnablePerPixelTransparency);
 }
 
@@ -585,6 +587,11 @@ void FMenuStack::DismissInternal(int32 FirstStackIndexToRemove)
 
 void FMenuStack::SetHostPath(const FWidgetPath& InOwnerPath)
 {
+	if (bHostWindowGuard)
+	{
+		return;
+	}
+
 	if ( HostPopupLayer.IsValid() )
 	{
 		if ( !InOwnerPath.ContainsWidget(HostPopupLayer->GetHost()) )
diff --git a/Engine/Source/Runtime/Slate/Private/Framework/Application/SlateApplication.cpp b/Engine/Source/Runtime/Slate/Private/Framework/Application/SlateApplication.cpp
--- a/Engine/Source/Runtime/Slate/Private/Framework/Application/SlateApplication.cpp
+++ b/Engine/Source/Runtime/Slate/Private/Framework/Application/SlateApplication.cpp
@@ -1830,9 +1830,23 @@ void FSlateApplication::AddModalWindow( TSharedRef<SWindow> InSlateWindow, const
 			PlatformApplication->Cursor->Show( true );
 		}
 
+		//Throttle loop data
+		float LastLoopTime = (float)FPlatformTime::Seconds();
+		const float MinThrottlePeriod = (1.0f / 60.0f); //Throttle the loop to a maximum of 60Hz
+
 		// Tick slate from here in the event that we should not return until the modal window is closed.
 		while( InSlateWindow == GetActiveModalWindow() )
 		{
+			//Throttle the loop
+			const float CurrentLoopTime = FPlatformTime::Seconds();
+			const float SleepTime = MinThrottlePeriod - (CurrentLoopTime-LastLoopTime);
+			LastLoopTime = CurrentLoopTime;
+			if (SleepTime > 0.0f)
+			{
+				// Sleep a bit to not eat up all CPU time
+				FPlatformProcess::Sleep(SleepTime);
+			}
+
 			FPlatformMisc::BeginNamedEvent(FColor::Magenta, "Slate::Tick");
 
 			{
diff --git a/Engine/Source/Runtime/Slate/Private/Widgets/Input/SEditableText.cpp b/Engine/Source/Runtime/Slate/Private/Widgets/Input/SEditableText.cpp
--- a/Engine/Source/Runtime/Slate/Private/Widgets/Input/SEditableText.cpp
+++ b/Engine/Source/Runtime/Slate/Private/Widgets/Input/SEditableText.cpp
@@ -152,16 +152,7 @@ bool SEditableText::SupportsKeyboardFocus() const
 
 FReply SEditableText::OnFocusReceived( const FGeometry& MyGeometry, const FFocusEvent& InFocusEvent )
 {
-	if (EditableTextLayout->HandleFocusReceived(InFocusEvent))
-	{
-		if (InFocusEvent.GetCause() != EFocusCause::Mouse && InFocusEvent.GetCause() != EFocusCause::OtherWidgetLostFocus && bIsCaretMovedWhenGainFocus.Get())
-		{
-			EditableTextLayout->GoTo(ETextLocation::EndOfDocument);
-		}
-
-		return SWidget::OnFocusReceived(MyGeometry, InFocusEvent);
-	}
-
+	EditableTextLayout->HandleFocusReceived(InFocusEvent);
 	return FReply::Handled();
 }
 
@@ -390,6 +381,11 @@ bool SEditableText::IsMultiLineTextEdit() const
 	return false;
 }
 
+bool SEditableText::ShouldJumpCursorToEndWhenFocused() const
+{
+	return bIsCaretMovedWhenGainFocus.Get(false);
+}
+
 bool SEditableText::ShouldSelectAllTextWhenFocused() const
 {
 	return bSelectAllTextWhenFocused.Get(false);
diff --git a/Engine/Source/Runtime/Slate/Private/Widgets/Input/SSlider.cpp b/Engine/Source/Runtime/Slate/Private/Widgets/Input/SSlider.cpp
--- a/Engine/Source/Runtime/Slate/Private/Widgets/Input/SSlider.cpp
+++ b/Engine/Source/Runtime/Slate/Private/Widgets/Input/SSlider.cpp
@@ -80,8 +80,8 @@ int32 SSlider::OnPaint( const FPaintArgs& Args, const FGeometry& AllottedGeometr
 	const ESlateDrawEffect::Type DrawEffects = bEnabled ? ESlateDrawEffect::None : ESlateDrawEffect::DisabledEffect;
 
 	// draw slider bar
-	auto BarTopLeft = FVector2D(SliderStartPoint.X, SliderStartPoint.Y - 1);
-	auto BarSize = FVector2D(SliderEndPoint.X - SliderStartPoint.X, 2);
+	auto BarTopLeft = FVector2D(SliderStartPoint.X, SliderStartPoint.Y - Style->BarThickness * 0.5f);
+	auto BarSize = FVector2D(SliderEndPoint.X - SliderStartPoint.X, Style->BarThickness);
 	FSlateDrawElement::MakeBox(
 		OutDrawElements,
 		LayerId,
@@ -117,12 +117,14 @@ FVector2D SSlider::ComputeDesiredSize( float ) const
 		return SSliderDesiredSize;
 	}
 
+	const float Thickness = FMath::Max(Style->BarThickness, Style->NormalThumbImage.ImageSize.Y);
+
 	if (Orientation == Orient_Vertical)
 	{
-		return FVector2D(Style->NormalThumbImage.ImageSize.Y, SSliderDesiredSize.Y);
+		return FVector2D(Thickness, SSliderDesiredSize.Y);
 	}
 
-	return FVector2D(SSliderDesiredSize.X, Style->NormalThumbImage.ImageSize.Y);
+	return FVector2D(SSliderDesiredSize.X, Thickness);
 }
 
 
diff --git a/Engine/Source/Runtime/Slate/Private/Widgets/Text/SMultiLineEditableText.cpp b/Engine/Source/Runtime/Slate/Private/Widgets/Text/SMultiLineEditableText.cpp
--- a/Engine/Source/Runtime/Slate/Private/Widgets/Text/SMultiLineEditableText.cpp
+++ b/Engine/Source/Runtime/Slate/Private/Widgets/Text/SMultiLineEditableText.cpp
@@ -198,6 +198,11 @@ bool SMultiLineEditableText::IsMultiLineTextEdit() const
 	return true;
 }
 
+bool SMultiLineEditableText::ShouldJumpCursorToEndWhenFocused() const
+{
+	return false;
+}
+
 bool SMultiLineEditableText::ShouldSelectAllTextWhenFocused() const
 {
 	return bSelectAllTextWhenFocused.Get(false);
@@ -326,12 +331,7 @@ float SMultiLineEditableText::UpdateAndClampVerticalScrollBar(const float InView
 
 FReply SMultiLineEditableText::OnFocusReceived( const FGeometry& MyGeometry, const FFocusEvent& InFocusEvent )
 {
-	// Skip the focus received code if it's due to the context menu closing
-	if (EditableTextLayout->HandleFocusReceived(InFocusEvent))
-	{
-		return SWidget::OnFocusReceived(MyGeometry, InFocusEvent);
-	}
-
+	EditableTextLayout->HandleFocusReceived(InFocusEvent);
 	return FReply::Handled();
 }
 
diff --git a/Engine/Source/Runtime/Slate/Private/Widgets/Text/SlateEditableTextLayout.cpp b/Engine/Source/Runtime/Slate/Private/Widgets/Text/SlateEditableTextLayout.cpp
--- a/Engine/Source/Runtime/Slate/Private/Widgets/Text/SlateEditableTextLayout.cpp
+++ b/Engine/Source/Runtime/Slate/Private/Widgets/Text/SlateEditableTextLayout.cpp
@@ -193,16 +193,21 @@ FSlateEditableTextLayout::~FSlateEditableTextLayout()
 
 void FSlateEditableTextLayout::SetText(const TAttribute<FText>& InText)
 {
+	const FText PreviousText = BoundText.Get(FText::GetEmpty());
 	BoundText = InText;
-
-	const FText& TextToSet = BoundText.Get(FText::GetEmpty());
+	const FText NewText = BoundText.Get(FText::GetEmpty());
 
 	// We need to force an update if the text doesn't match the editable text, as the editable 
 	// text may not match the current bound text since it may have been changed by the user
 	const FText EditableText = GetEditableText();
-	const bool bForceRefresh = !EditableText.ToString().Equals(TextToSet.ToString(), ESearchCase::CaseSensitive);
+	const bool bForceRefresh = !EditableText.ToString().Equals(NewText.ToString(), ESearchCase::CaseSensitive);
+
+	// Only emit the "text changed" event if the text has actually been changed
+	const bool bHasTextChanged = OwnerWidget->GetSlateWidget()->HasKeyboardFocus() 
+		? !NewText.ToString().Equals(EditableText.ToString(), ESearchCase::CaseSensitive)
+		: !NewText.ToString().Equals(PreviousText.ToString(), ESearchCase::CaseSensitive);
 
-	if (RefreshImpl(&TextToSet, bForceRefresh))
+	if (RefreshImpl(&NewText, bForceRefresh))
 	{
 		// Make sure we move the cursor to the end of the new text if we had keyboard focus
 		if (OwnerWidget->GetSlateWidget()->HasKeyboardFocus())
@@ -211,7 +216,10 @@ void FSlateEditableTextLayout::SetText(const TAttribute<FText>& InText)
 		}
 
 		// Let outsiders know that the text content has been changed
-		OwnerWidget->OnTextChanged(TextToSet);
+		if (bHasTextChanged)
+		{
+			OwnerWidget->OnTextChanged(NewText);
+		}
 	}
 }
 
@@ -540,6 +548,12 @@ bool FSlateEditableTextLayout::HandleFocusReceived(const FFocusEvent& InFocusEve
 	// Store undo state to use for escape key reverts
 	MakeUndoState(OriginalText);
 
+	// Jump to the end of the document?
+	if (InFocusEvent.GetCause() != EFocusCause::Mouse && InFocusEvent.GetCause() != EFocusCause::OtherWidgetLostFocus && OwnerWidget->ShouldJumpCursorToEndWhenFocused())
+	{
+		GoTo(ETextLocation::EndOfDocument);
+	}
+
 	// Select All Text
 	if (OwnerWidget->ShouldSelectAllTextWhenFocused())
 	{
@@ -554,7 +568,7 @@ bool FSlateEditableTextLayout::HandleFocusReceived(const FFocusEvent& InFocusEve
 	// of making sure that gets scrolled into view
 	PositionToScrollIntoView.Reset();
 
-	return false;
+	return true;
 }
 
 bool FSlateEditableTextLayout::HandleFocusLost(const FFocusEvent& InFocusEvent)
@@ -609,6 +623,10 @@ bool FSlateEditableTextLayout::HandleFocusLost(const FFocusEvent& InFocusEvent)
 
 	OwnerWidget->OnTextCommitted(EditedText, TextAction);
 
+	// Reload underlying value now it is committed  (commit may alter the value) 
+	// so it can be re-displayed in the edit box
+	LoadText();
+
 	UpdateCursorHighlight();
 
 	// UpdateCursorHighlight always tries to scroll to the cursor, but we don't want that to happen when we 
diff --git a/Engine/Source/Runtime/Slate/Private/Widgets/Views/STableViewBase.cpp b/Engine/Source/Runtime/Slate/Private/Widgets/Views/STableViewBase.cpp
--- a/Engine/Source/Runtime/Slate/Private/Widgets/Views/STableViewBase.cpp
+++ b/Engine/Source/Runtime/Slate/Private/Widgets/Views/STableViewBase.cpp
@@ -688,6 +688,19 @@ void STableViewBase::SetScrollOffset( const float InScrollOffset )
 	}
 }
 
+void STableViewBase::AddScrollOffset(const float InScrollOffsetDelta, bool RefreshList)
+{
+	if (FMath::IsNearlyEqual(InScrollOffsetDelta, 0.0f) == false)
+	{
+		ScrollOffset += InScrollOffsetDelta;
+		if (RefreshList)
+		{
+			OnTableViewScrolled.ExecuteIfBound(ScrollOffset);
+			RequestListRefresh();
+		}
+	}
+}
+
 void STableViewBase::InsertWidget( const TSharedRef<ITableRow> & WidgetToInset )
 {
 	ItemsPanel->AddSlot(0)
diff --git a/Engine/Source/Runtime/Slate/Public/Framework/Application/MenuStack.h b/Engine/Source/Runtime/Slate/Public/Framework/Application/MenuStack.h
--- a/Engine/Source/Runtime/Slate/Public/Framework/Application/MenuStack.h
+++ b/Engine/Source/Runtime/Slate/Public/Framework/Application/MenuStack.h
@@ -58,6 +58,11 @@ struct FPopupTransitionEffect
 class FMenuStack
 {
 public:
+	/** Constructor */
+	FMenuStack()
+		: bHostWindowGuard(false)
+	{}
+
 	/**
 	 * [Deprecated] Pushes a new menu onto the stack.  Automatically removes windows from the stack that are no longer valid
 	 * Invalid windows are those that are not parents of the window menu being pushed and any windows on the same stack level
@@ -408,4 +413,7 @@ private:
 
 	/** Temporary ptr to a new menu created during the menu creation process. Nulled before the Push() call returns. Stops it collapsing the stack when it gets focus. */
 	TSharedPtr<class FMenuBase> PendingNewMenu;
+
+	/** Guard to prevent the HostWindow and HostWindowPopupPanel being set reentrantly */
+	bool bHostWindowGuard;
 };
diff --git a/Engine/Source/Runtime/Slate/Public/Widgets/Input/SEditableText.h b/Engine/Source/Runtime/Slate/Public/Widgets/Input/SEditableText.h
--- a/Engine/Source/Runtime/Slate/Public/Widgets/Input/SEditableText.h
+++ b/Engine/Source/Runtime/Slate/Public/Widgets/Input/SEditableText.h
@@ -284,6 +284,7 @@ public:
 
 protected:
 	//~ Begin ISlateEditableTextWidget Interface
+	virtual bool ShouldJumpCursorToEndWhenFocused() const override;
 	virtual bool ShouldSelectAllTextWhenFocused() const override;
 	virtual bool ShouldClearTextSelectionOnFocusLoss() const override;
 	virtual bool ShouldRevertTextOnEscape() const override;
