Copying //UE4/Dev-Rendering to //UE4/Dev-Main (Source: //UE4/Dev-Rendering @ 2981742)

diff --git a/Engine/Source/Runtime/Engine/Public/StaticMeshResources.h b/Engine/Source/Runtime/Engine/Public/StaticMeshResources.h
--- a/Engine/Source/Runtime/Engine/Public/StaticMeshResources.h
+++ b/Engine/Source/Runtime/Engine/Public/StaticMeshResources.h
@@ -108,66 +108,133 @@ private:
 	TMap<FName,FStaticMeshLODGroup> Groups;
 };
 
-/** 
- * All information about a static-mesh vertex with a variable number of texture coordinates.
- * Position information is stored separately to reduce vertex fetch bandwidth in passes that only need position. (z prepass)
- */
-struct FStaticMeshFullVertex
+template<typename TangentTypeT>
+struct TStaticMeshVertexTangentDatum
 {
-	uint32 RawTangentX;
-	uint32 TangentZ;
+	TangentTypeT TangentX;
+	TangentTypeT TangentZ;
 
-	FStaticMeshFullVertex()
-		: RawTangentX(0)
-		, TangentZ(0)
+	FORCEINLINE FVector GetTangentX() const
 	{
+		return TangentX;
 	}
 
-	/**
-	* Serializer
-	*
-	* @param Ar - archive to serialize with
-	*/
-	void Serialize(FArchive& Ar)
+	FORCEINLINE FVector4 GetTangentZ() const
+	{
+		return TangentZ;
+	}
+
+	FORCEINLINE FVector GetTangentY() const
 	{
-		Ar << RawTangentX;
-		Ar << TangentZ;
+		FVector  TanX = GetTangentX();
+		FVector4 TanZ = GetTangentZ();
+
+		return (FVector(TanZ) ^ TanX) * TanZ.W;
+	}
+
+	FORCEINLINE void SetTangents(FVector X, FVector Y, FVector Z)
+	{
+		TangentX = X;
+		TangentZ = FVector4(Z, GetBasisDeterminantSign(X, Y, Z));
 	}
 };
 
-/** 
-* 16 bit UV version of static mesh vertex
-*/
-template<uint32 NumTexCoords>
-struct TStaticMeshFullVertexFloat16UVs : public FStaticMeshFullVertex
+template<typename UVTypeT, uint32 NumTexCoords>
+struct TStaticMeshVertexUVsDatum
 {
-	FVector2DHalf UVs[NumTexCoords];
+	UVTypeT UVs[NumTexCoords];
 
-	/**
-	* Serializer
-	*
-	* @param Ar - archive to serialize with
-	* @param V - vertex to serialize
-	* @return archive that was used
-	*/
-	friend FArchive& operator<<(FArchive& Ar,TStaticMeshFullVertexFloat16UVs& Vertex)
+	FORCEINLINE FVector2D GetUV(uint32 TexCoordIndex) const
 	{
-		Vertex.Serialize(Ar);
-		for(uint32 UVIndex = 0;UVIndex < NumTexCoords;UVIndex++)
-		{
-			Ar << Vertex.UVs[UVIndex];
-		}
-		return Ar;
+		check(TexCoordIndex < NumTexCoords);
+
+		return UVs[TexCoordIndex];
 	}
+
+	FORCEINLINE void SetUV(uint32 TexCoordIndex, FVector2D UV)
+	{
+		check(TexCoordIndex < NumTexCoords);
+
+		UVs[TexCoordIndex] = UV;
+	}
+};
+
+enum class EStaticMeshVertexTangentBasisType
+{
+	Default,
+	HighPrecision,
 };
 
-/** 
-* 32 bit UV version of static mesh vertex
+enum class EStaticMeshVertexUVType
+{
+	Default,
+	HighPrecision,
+};
+
+template<EStaticMeshVertexTangentBasisType TangentBasisType>
+struct TStaticMeshVertexTangentTypeSelector
+{
+};
+
+template<>
+struct TStaticMeshVertexTangentTypeSelector<EStaticMeshVertexTangentBasisType::Default>
+{
+	typedef FPackedNormal TangentTypeT;
+	static const EVertexElementType VertexElementType = VET_PackedNormal;
+};
+
+template<>
+struct TStaticMeshVertexTangentTypeSelector<EStaticMeshVertexTangentBasisType::HighPrecision>
+{
+	typedef FPackedRGBA16N TangentTypeT;
+	static const EVertexElementType VertexElementType = VET_UShort4N;
+};
+
+template<EStaticMeshVertexUVType UVType>
+struct TStaticMeshVertexUVsTypeSelector
+{
+};
+
+template<>
+struct TStaticMeshVertexUVsTypeSelector<EStaticMeshVertexUVType::Default>
+{
+	typedef FVector2DHalf UVsTypeT;
+};
+
+template<>
+struct TStaticMeshVertexUVsTypeSelector<EStaticMeshVertexUVType::HighPrecision>
+{
+	typedef FVector2D UVsTypeT;
+};
+
+#define ConstExprStaticMeshVertexTypeID(bUseHighPrecisionTangentBasis, bUseFullPrecisionUVs, TexCoordsCount) \
+	(((((bUseHighPrecisionTangentBasis) ? 1 : 0) * 2) + ((bUseFullPrecisionUVs) ? 1 : 0)) * MAX_STATIC_TEXCOORDS) + ((TexCoordsCount) - 1)
+
+FORCEINLINE uint32 ComputeStaticMeshVertexTypeID(bool bUseHighPrecisionTangentBasis, bool bUseFullPrecisionUVs, uint32 TexCoordsCount)
+{
+	return ConstExprStaticMeshVertexTypeID(bUseHighPrecisionTangentBasis, bUseFullPrecisionUVs, TexCoordsCount);
+}
+
+/**
+* All information about a static-mesh vertex with a variable number of texture coordinates.
+* Position information is stored separately to reduce vertex fetch bandwidth in passes that only need position. (z prepass)
 */
-template<uint32 NumTexCoords>
-struct TStaticMeshFullVertexFloat32UVs : public FStaticMeshFullVertex
+template<EStaticMeshVertexTangentBasisType TangentBasisTypeT, EStaticMeshVertexUVType UVTypeT, uint32 NumTexCoordsT>
+struct TStaticMeshFullVertex : 
+	public TStaticMeshVertexTangentDatum<typename TStaticMeshVertexTangentTypeSelector<TangentBasisTypeT>::TangentTypeT>,
+	public TStaticMeshVertexUVsDatum<typename TStaticMeshVertexUVsTypeSelector<UVTypeT>::UVsTypeT, NumTexCoordsT>
 {
-	FVector2D UVs[NumTexCoords];
+	static_assert(NumTexCoordsT > 0, "Must have at least 1 texcoord.");
+
+	static const EStaticMeshVertexTangentBasisType TangentBasisType = TangentBasisTypeT;
+	static const EStaticMeshVertexUVType UVType = UVTypeT;
+	static const uint32 NumTexCoords = NumTexCoordsT;
+
+	static const uint32 VertexTypeID = ConstExprStaticMeshVertexTypeID(
+		TangentBasisTypeT == EStaticMeshVertexTangentBasisType::HighPrecision, 
+		UVTypeT == EStaticMeshVertexUVType::HighPrecision,
+		NumTexCoordsT
+		);
 
 	/**
 	* Serializer
@@ -176,17 +243,48 @@ struct TStaticMeshFullVertexFloat32UVs : public FStaticMeshFullVertex
 	* @param V - vertex to serialize
 	* @return archive that was used
 	*/
-	friend FArchive& operator<<(FArchive& Ar,TStaticMeshFullVertexFloat32UVs& Vertex)
+	friend FArchive& operator<<(FArchive& Ar, TStaticMeshFullVertex& Vertex)
 	{
-		Vertex.Serialize(Ar);
-		for(uint32 UVIndex = 0;UVIndex < NumTexCoords;UVIndex++)
+		Ar << Vertex.TangentX;
+		Ar << Vertex.TangentZ;
+
+		for (uint32 UVIndex = 0; UVIndex < NumTexCoordsT; UVIndex++)
 		{
 			Ar << Vertex.UVs[UVIndex];
 		}
+
 		return Ar;
 	}
 };
 
+#define APPLY_TO_STATIC_MESH_VERTEX(TangentBasisType, UVType, UVCount, ...) \
+	{ \
+		typedef TStaticMeshFullVertex<TangentBasisType, UVType, UVCount> VertexType; \
+		case VertexType::VertexTypeID: { __VA_ARGS__ } break; \
+	}
+
+#define SELECT_STATIC_MESH_VERTEX_TYPE_WITH_TEX_COORDS(TangentBasisType, UVType, ...) \
+	APPLY_TO_STATIC_MESH_VERTEX(TangentBasisType, UVType, 1, __VA_ARGS__); \
+	APPLY_TO_STATIC_MESH_VERTEX(TangentBasisType, UVType, 2, __VA_ARGS__); \
+	APPLY_TO_STATIC_MESH_VERTEX(TangentBasisType, UVType, 3, __VA_ARGS__); \
+	APPLY_TO_STATIC_MESH_VERTEX(TangentBasisType, UVType, 4, __VA_ARGS__); \
+	APPLY_TO_STATIC_MESH_VERTEX(TangentBasisType, UVType, 5, __VA_ARGS__); \
+	APPLY_TO_STATIC_MESH_VERTEX(TangentBasisType, UVType, 6, __VA_ARGS__); \
+	APPLY_TO_STATIC_MESH_VERTEX(TangentBasisType, UVType, 7, __VA_ARGS__); \
+	APPLY_TO_STATIC_MESH_VERTEX(TangentBasisType, UVType, 8, __VA_ARGS__);
+
+#define SELECT_STATIC_MESH_VERTEX_TYPE(bIsHighPrecisionTangentBais, bIsHigPrecisionUVs, NumTexCoords, ...) \
+	{ \
+		uint32 VertexTypeID = ComputeStaticMeshVertexTypeID(bIsHighPrecisionTangentBais, bIsHigPrecisionUVs, NumTexCoords); \
+		switch (VertexTypeID) \
+		{ \
+			SELECT_STATIC_MESH_VERTEX_TYPE_WITH_TEX_COORDS(EStaticMeshVertexTangentBasisType::Default,		  EStaticMeshVertexUVType::Default,		  __VA_ARGS__); \
+			SELECT_STATIC_MESH_VERTEX_TYPE_WITH_TEX_COORDS(EStaticMeshVertexTangentBasisType::Default,		  EStaticMeshVertexUVType::HighPrecision, __VA_ARGS__); \
+			SELECT_STATIC_MESH_VERTEX_TYPE_WITH_TEX_COORDS(EStaticMeshVertexTangentBasisType::HighPrecision,  EStaticMeshVertexUVType::Default,		  __VA_ARGS__); \
+			SELECT_STATIC_MESH_VERTEX_TYPE_WITH_TEX_COORDS(EStaticMeshVertexTangentBasisType::HighPrecision,  EStaticMeshVertexUVType::HighPrecision, __VA_ARGS__); \
+		}; \
+	}
+
 /**
  * A set of static mesh triangles which are rendered with the same material.
  */
@@ -395,25 +493,33 @@ public:
 	* Specialized assignment operator, only used when importing LOD's. 
 	*/
 	ENGINE_API void operator=(const FStaticMeshVertexBuffer &Other);
-	
+
 	FORCEINLINE FVector4 VertexTangentX(uint32 VertexIndex) const
 	{
 		checkSlow(VertexIndex < GetNumVertices());
 
-		const uint32* RawTangentXPtr = &((FStaticMeshFullVertex*)(Data + VertexIndex * Stride))->RawTangentX;
-
-		FVector4 result = FVector::ZeroVector;
 		if (GetUseHighPrecisionTangentBasis())
 		{
-			result = *reinterpret_cast<const FPackedRGB10A2*>(RawTangentXPtr);
+			return reinterpret_cast<const TStaticMeshFullVertex<EStaticMeshVertexTangentBasisType::HighPrecision, EStaticMeshVertexUVType::Default, 1>*>(Data + VertexIndex * Stride)->GetTangentX();
 		}
 		else
 		{
-			auto* PackedNormalX = reinterpret_cast<const FPackedNormal*>(RawTangentXPtr);
-			result = FVector4(*PackedNormalX, PackedNormalX->Vector.W ? 1.0f : -1.0f);
+			return reinterpret_cast<const TStaticMeshFullVertex<EStaticMeshVertexTangentBasisType::Default, EStaticMeshVertexUVType::Default, 1>*>(Data + VertexIndex * Stride)->GetTangentX();
 		}
+	}
+
+	FORCEINLINE FVector VertexTangentZ(uint32 VertexIndex) const
+	{
+		checkSlow(VertexIndex < GetNumVertices());
 
-		return result;
+		if (GetUseHighPrecisionTangentBasis())
+		{
+			return reinterpret_cast<const TStaticMeshFullVertex<EStaticMeshVertexTangentBasisType::HighPrecision, EStaticMeshVertexUVType::Default, 1>*>(Data + VertexIndex * Stride)->GetTangentZ();
+		}
+		else
+		{
+			return reinterpret_cast<const TStaticMeshFullVertex<EStaticMeshVertexTangentBasisType::Default, EStaticMeshVertexUVType::Default, 1>*>(Data + VertexIndex * Stride)->GetTangentZ();
+		}
 	}
 
 	/**
@@ -423,34 +529,31 @@ public:
 	* @return binormal (TangentY) vector
 	*/
 	FORCEINLINE FVector VertexTangentY(uint32 VertexIndex) const
-	{
-		FVector4 TangentX = VertexTangentX(VertexIndex);
-		FVector  TangentZ = VertexTangentZ(VertexIndex);
-
-		return (TangentZ ^ FVector(TangentX)) * TangentX.W;
-	}
-
-	FORCEINLINE FVector VertexTangentZ(uint32 VertexIndex) const
 	{
 		checkSlow(VertexIndex < GetNumVertices());
 
-		return OctDecodeUnit(((FStaticMeshFullVertex*)(Data + VertexIndex * Stride))->TangentZ);
+		if (GetUseHighPrecisionTangentBasis())
+		{
+			return reinterpret_cast<const TStaticMeshFullVertex<EStaticMeshVertexTangentBasisType::HighPrecision, EStaticMeshVertexUVType::Default, 1>*>(Data + VertexIndex * Stride)->GetTangentY();
+		}
+		else
+		{
+			return reinterpret_cast<const TStaticMeshFullVertex<EStaticMeshVertexTangentBasisType::Default, EStaticMeshVertexUVType::Default, 1>*>(Data + VertexIndex * Stride)->GetTangentY();
+		}
 	}
 
 	FORCEINLINE void SetVertexTangents(uint32 VertexIndex, FVector X, FVector Y, FVector Z)
 	{
-		uint32* RawTangentXPtr = &((FStaticMeshFullVertex*)(Data + VertexIndex * Stride))->RawTangentX;
+		checkSlow(VertexIndex < GetNumVertices());
 
 		if (GetUseHighPrecisionTangentBasis())
 		{
-			*reinterpret_cast<FPackedRGB10A2*>(RawTangentXPtr) = FVector4(X, GetBasisDeterminantSign(X, Y, Z));
+			return reinterpret_cast<TStaticMeshFullVertex<EStaticMeshVertexTangentBasisType::HighPrecision, EStaticMeshVertexUVType::Default, 1>*>(Data + VertexIndex * Stride)->SetTangents(X, Y, Z);
 		}
 		else
 		{
-			*reinterpret_cast<FPackedNormal*>(RawTangentXPtr) = FVector4(X, GetBasisDeterminantSign(X, Y, Z));
+			return reinterpret_cast<TStaticMeshFullVertex<EStaticMeshVertexTangentBasisType::Default, EStaticMeshVertexUVType::Default, 1>*>(Data + VertexIndex * Stride)->SetTangents(X, Y, Z);
 		}
-
-		((FStaticMeshFullVertex*)(Data + VertexIndex * Stride))->TangentZ = OctEncodeUnit(Z);
 	}
 
 	/**
@@ -460,21 +563,37 @@ public:
 	* @param UVIndex - [0,MAX_STATIC_TEXCOORDS] value to index into UVs array
 	* @param Vec2D - UV values to set
 	*/
-	FORCEINLINE void SetVertexUV(uint32 VertexIndex,uint32 UVIndex,const FVector2D& Vec2D)
+	FORCEINLINE void SetVertexUV(uint32 VertexIndex, uint32 UVIndex, const FVector2D& Vec2D)
 	{
 		checkSlow(VertexIndex < GetNumVertices());
-		if( !bUseFullPrecisionUVs )
+		checkSlow(UVIndex < GetNumTexCoords());
+		
+		if (GetUseHighPrecisionTangentBasis())
 		{
-			((TStaticMeshFullVertexFloat16UVs<MAX_STATIC_TEXCOORDS>*)(Data + VertexIndex * Stride))->UVs[UVIndex] = Vec2D;
+			if (GetUseFullPrecisionUVs())
+			{
+				reinterpret_cast<TStaticMeshFullVertex<EStaticMeshVertexTangentBasisType::HighPrecision, EStaticMeshVertexUVType::HighPrecision, MAX_STATIC_TEXCOORDS>*>(Data + VertexIndex * Stride)->SetUV(UVIndex, Vec2D);
+			}
+			else
+			{
+				reinterpret_cast<TStaticMeshFullVertex<EStaticMeshVertexTangentBasisType::HighPrecision, EStaticMeshVertexUVType::Default, MAX_STATIC_TEXCOORDS>*>(Data + VertexIndex * Stride)->SetUV(UVIndex, Vec2D);
+			}
 		}
 		else
 		{
-			((TStaticMeshFullVertexFloat32UVs<MAX_STATIC_TEXCOORDS>*)(Data + VertexIndex * Stride))->UVs[UVIndex] = Vec2D;
-		}		
+			if (GetUseFullPrecisionUVs())
+			{
+				reinterpret_cast<TStaticMeshFullVertex<EStaticMeshVertexTangentBasisType::Default, EStaticMeshVertexUVType::HighPrecision, MAX_STATIC_TEXCOORDS>*>(Data + VertexIndex * Stride)->SetUV(UVIndex, Vec2D);
+			}
+			else
+			{
+				reinterpret_cast<TStaticMeshFullVertex<EStaticMeshVertexTangentBasisType::Default, EStaticMeshVertexUVType::Default, MAX_STATIC_TEXCOORDS>*>(Data + VertexIndex * Stride)->SetUV(UVIndex, Vec2D);
+			}
+		}
 	}
 
 	/**
-	* Fet the vertex UV values at the given index in the vertex buffer
+	* Set the vertex UV values at the given index in the vertex buffer
 	*
 	* @param VertexIndex - index into the vertex buffer
 	* @param UVIndex - [0,MAX_STATIC_TEXCOORDS] value to index into UVs array
@@ -483,14 +602,30 @@ public:
 	FORCEINLINE FVector2D GetVertexUV(uint32 VertexIndex,uint32 UVIndex) const
 	{
 		checkSlow(VertexIndex < GetNumVertices());
-		if( !bUseFullPrecisionUVs )
+		checkSlow(UVIndex < GetNumTexCoords());
+
+		if (GetUseHighPrecisionTangentBasis())
 		{
-			return ((TStaticMeshFullVertexFloat16UVs<MAX_STATIC_TEXCOORDS>*)(Data + VertexIndex * Stride))->UVs[UVIndex];
+			if (GetUseFullPrecisionUVs())
+			{
+				return reinterpret_cast<const TStaticMeshFullVertex<EStaticMeshVertexTangentBasisType::HighPrecision, EStaticMeshVertexUVType::HighPrecision, MAX_STATIC_TEXCOORDS>*>(Data + VertexIndex * Stride)->GetUV(UVIndex);
+			}
+			else
+			{
+				return reinterpret_cast<const TStaticMeshFullVertex<EStaticMeshVertexTangentBasisType::HighPrecision, EStaticMeshVertexUVType::Default, MAX_STATIC_TEXCOORDS>*>(Data + VertexIndex * Stride)->GetUV(UVIndex);
+			}
 		}
 		else
 		{
-			return ((TStaticMeshFullVertexFloat32UVs<MAX_STATIC_TEXCOORDS>*)(Data + VertexIndex * Stride))->UVs[UVIndex];
-		}		
+			if (GetUseFullPrecisionUVs())
+			{
+				return reinterpret_cast<const TStaticMeshFullVertex<EStaticMeshVertexTangentBasisType::Default, EStaticMeshVertexUVType::HighPrecision, MAX_STATIC_TEXCOORDS>*>(Data + VertexIndex * Stride)->GetUV(UVIndex);
+			}
+			else
+			{
+				return reinterpret_cast<const TStaticMeshFullVertex<EStaticMeshVertexTangentBasisType::Default, EStaticMeshVertexUVType::Default, MAX_STATIC_TEXCOORDS>*>(Data + VertexIndex * Stride)->GetUV(UVIndex);
+			}
+		}
 	}
 
 	// Other accessors.
@@ -536,63 +671,15 @@ public:
 	}
 
 	/**
-	* Convert the existing data in this mesh from 16 bit to 32 bit UVs.
-	* Without rebuilding the mesh (loss of precision)
+	* Convert the existing data in this mesh without rebuilding the mesh (loss of precision)
 	*/
-	template<int32 NumTexCoords>
-	void ConvertToFullPrecisionUVs();
-
-	/**
-	* Convert the existing data(in place) in this mesh from RGB10A2 to RGBA8 tangent basis encoding.
-	* Without rebuilding the mesh (loss of precision)
-	*/
-	template<int32 NumTexCoordsT>
-	void ConvertToLowPrecisionTangentBasis();
+	template<typename SrcVertexTypeT, typename DstVertexTypeT>
+	void ConvertVertexFormat();
 
 	// FRenderResource interface.
 	virtual void InitRHI() override;
 	virtual FString GetFriendlyName() const override { return TEXT("Static-mesh vertices"); }
 
-	static uint32 OctEncodeUnit(FVector V)
-	{
-		float Proj = 1.0f / FVector::DotProduct(V.GetAbs(), FVector(1.0f));
-
-		FVector2D P(V.X, V.Y);
-		P *= Proj;
-		
-		FVector2D Encoded = (V.Z <= 0.0f) ? (FVector2D::UnitVector - FVector2D(P.Y, P.X).GetAbs()) * P.GetSignVector() : P;
-		
-		Encoded = (Encoded.ClampAxes(-1.0f, 1.0f) * 0.5f + 0.5f) * UINT16_MAX;
-		Encoded = FVector2D(FMath::RoundToInt(Encoded.X), FMath::RoundToInt(Encoded.Y));
-
-		uint32 Packed = (uint32(Encoded.X) << 00) | (uint32(Encoded.Y) << 16);
-
-		return Packed;
-	}
-
-	static FVector OctDecodeUnit(uint32 P)
-	{
-		FVector2D E(
-			float((P >> 00) & UINT16_MAX), 
-			float((P >> 16) & UINT16_MAX)
-			);
-		E = (E / float(UINT16_MAX)) * 2.0f - 1.0f;
-
-		FVector V(E.X, E.Y, 1.0f - FMath::Abs(E.X) - FMath::Abs(E.Y));
-
-		if (V.Z < 0.0f)
-		{
-			FVector2D T = (FVector2D::UnitVector - FVector2D(V.Y, V.X).GetAbs()) * E.GetSignVector();
-
-			V.X = T.X;
-			V.Y = T.Y;
-		}
-
-		V.Normalize();
-
-		return V;
-	}
-
 private:
 
 	/** The vertex data storage type */
@@ -917,7 +1004,7 @@ public:
 	virtual void GetLCIs(FLCIArray& LCIs) override;
 
 #if WITH_EDITORONLY_DATA
-	virtual const FStreamingSectionBuildInfo* GetStreamingSectionData(float& OutDistanceMultiplier, int32 LODIndex, int32 ElementIndex) const override;
+	virtual const FStreamingSectionBuildInfo* GetStreamingSectionData(float& OutComponentExtraScale, float& OutMeshExtraScale, int32 LODIndex, int32 ElementIndex) const override;
 #endif
 
 #if !(UE_BUILD_SHIPPING || UE_BUILD_TEST)
diff --git a/Engine/Source/Runtime/Engine/Public/UnrealEngine.h b/Engine/Source/Runtime/Engine/Public/UnrealEngine.h
--- a/Engine/Source/Runtime/Engine/Public/UnrealEngine.h
+++ b/Engine/Source/Runtime/Engine/Public/UnrealEngine.h
@@ -410,8 +410,6 @@ struct FCachedSystemScalabilityCVars
 	int32 MaxShadowResolution;
 	float ViewDistanceScale;
 	float ViewDistanceScaleSquared;
-	// 0:off, 1:on, -1:unknown
-	int32 SimpleDynamicLighting;
 
 	FCachedSystemScalabilityCVars();
 
diff --git a/Engine/Source/Runtime/Landscape/Private/LandscapeRender.cpp b/Engine/Source/Runtime/Landscape/Private/LandscapeRender.cpp
--- a/Engine/Source/Runtime/Landscape/Private/LandscapeRender.cpp
+++ b/Engine/Source/Runtime/Landscape/Private/LandscapeRender.cpp
@@ -2572,10 +2572,13 @@ void ULandscapeComponent::GetStreamingTextureInfo(FStreamingTextureLevelContext&
 		// TODO: Take into account which UVIndex is being used.
 		for (int32 TextureIndex = 0; TextureIndex < Textures.Num(); TextureIndex++)
 		{
+			UTexture2D* Texture2D = Cast<UTexture2D>(Textures[TextureIndex]);
+			if (!Texture2D) continue;
+
 			FStreamingTexturePrimitiveInfo& StreamingTexture = *new(OutStreamingTextures)FStreamingTexturePrimitiveInfo;
 			StreamingTexture.Bounds = BoundingSphere;
 			StreamingTexture.TexelFactor = TexelFactor;
-			StreamingTexture.Texture = Textures[TextureIndex];
+			StreamingTexture.Texture = Texture2D;
 		}
 
 		UMaterial* Material = MaterialInstance->GetMaterial();
@@ -2681,7 +2684,7 @@ void ULandscapeComponent::GetStreamingTextureInfo(FStreamingTextureLevelContext&
 		StreamingHeightmap.Bounds = BoundingSphere;
 
 		float HeightmapTexelFactor = TexelFactor * (static_cast<float>(HeightmapTexture->GetSizeY()) / (ComponentSizeQuads + 1));
-		StreamingHeightmap.TexelFactor = ForcedLOD >= 0 ? -13 + ForcedLOD : HeightmapTexelFactor; // Minus Value indicate ForcedLOD, 13 for 8k texture
+		StreamingHeightmap.TexelFactor = ForcedLOD >= 0 ? -(1 << (13 - ForcedLOD)) : HeightmapTexelFactor; // Minus Value indicate forced resolution (Mip 13 for 8k texture)
 		StreamingHeightmap.Texture = HeightmapTexture;
 	}
 
diff --git a/Engine/Source/Runtime/RHI/Private/RHICommandList.cpp b/Engine/Source/Runtime/RHI/Private/RHICommandList.cpp
--- a/Engine/Source/Runtime/RHI/Private/RHICommandList.cpp
+++ b/Engine/Source/Runtime/RHI/Private/RHICommandList.cpp
@@ -1288,7 +1288,7 @@ void FRHICommandList::BeginScene()
 	check(IsImmediate() && IsInRenderingThread());
 	if (Bypass())
 	{
-		CMD_CONTEXT(BeginScene)();
+		CMD_CONTEXT(RHIBeginScene)();
 		return;
 	}
 	new (AllocCommand<FRHICommandBeginScene>()) FRHICommandBeginScene();
@@ -1304,7 +1304,7 @@ void FRHICommandList::EndScene()
 	check(IsImmediate() && IsInRenderingThread());
 	if (Bypass())
 	{
-		CMD_CONTEXT(EndScene)();
+		CMD_CONTEXT(RHIEndScene)();
 		return;
 	}
 	new (AllocCommand<FRHICommandEndScene>()) FRHICommandEndScene();
@@ -1322,7 +1322,7 @@ void FRHICommandList::BeginDrawingViewport(FViewportRHIParamRef Viewport, FTextu
 	check(IsImmediate() && IsInRenderingThread());
 	if (Bypass())
 	{
-		CMD_CONTEXT(BeginDrawingViewport)(Viewport, RenderTargetRHI);
+		CMD_CONTEXT(RHIBeginDrawingViewport)(Viewport, RenderTargetRHI);
 		return;
 	}
 	new (AllocCommand<FRHICommandBeginDrawingViewport>()) FRHICommandBeginDrawingViewport(Viewport, RenderTargetRHI);
@@ -1339,7 +1339,7 @@ void FRHICommandList::EndDrawingViewport(FViewportRHIParamRef Viewport, bool bPr
 	check(IsImmediate() && IsInRenderingThread());
 	if (Bypass())
 	{
-		CMD_CONTEXT(EndDrawingViewport)(Viewport, bPresent, bLockToVsync);
+		CMD_CONTEXT(RHIEndDrawingViewport)(Viewport, bPresent, bLockToVsync);
 	}
 	else
 	{
@@ -1358,7 +1358,7 @@ void FRHICommandList::BeginFrame()
 	check(IsImmediate() && IsInRenderingThread());
 	if (Bypass())
 	{
-		CMD_CONTEXT(BeginFrame)();
+		CMD_CONTEXT(RHIBeginFrame)();
 		return;
 	}
 	new (AllocCommand<FRHICommandBeginFrame>()) FRHICommandBeginFrame();
@@ -1375,7 +1375,7 @@ void FRHICommandList::EndFrame()
 	check(IsImmediate() && IsInRenderingThread());
 	if (Bypass())
 	{
-		CMD_CONTEXT(EndFrame)();
+		CMD_CONTEXT(RHIEndFrame)();
 		return;
 	}
 	new (AllocCommand<FRHICommandEndFrame>()) FRHICommandEndFrame();
@@ -2075,7 +2075,7 @@ void FRHICommandListImmediate::UpdateTextureReference(FTextureReferenceRHIParamR
 			QUICK_SCOPE_CYCLE_COUNTER(STAT_RHIMETHOD_UpdateTextureReference_FlushRHI);
 			ImmediateFlush(EImmediateFlushType::FlushRHIThread);
 		}
-		CMD_CONTEXT(UpdateTextureReference)(TextureRef, NewTexture);
+		CMD_CONTEXT(RHIUpdateTextureReference)(TextureRef, NewTexture);
 		return;
 	}
 	new (AllocCommand<FRHICommandUpdateTextureReference>()) FRHICommandUpdateTextureReference(TextureRef, NewTexture);
diff --git a/Engine/Source/Runtime/RHI/Public/RHICommandList.h b/Engine/Source/Runtime/RHI/Public/RHICommandList.h
--- a/Engine/Source/Runtime/RHI/Public/RHICommandList.h
+++ b/Engine/Source/Runtime/RHI/Public/RHICommandList.h
@@ -1330,8 +1330,8 @@ struct FRHICommandUpdateTextureReference : public FRHICommand<FRHICommandUpdateT
 };
 
 
-#define CMD_CONTEXT(Method) GetContext().RHI##Method
-#define COMPUTE_CONTEXT(Method) GetComputeContext().RHI##Method
+#define CMD_CONTEXT(Method) GetContext().Method
+#define COMPUTE_CONTEXT(Method) GetComputeContext().Method
 
 template<> void FRHICommandSetShaderParameter<FComputeShaderRHIParamRef, ECmdList::ECompute>::Execute(FRHICommandListBase& CmdList);
 template<> void FRHICommandSetShaderUniformBuffer<FComputeShaderRHIParamRef, ECmdList::ECompute>::Execute(FRHICommandListBase& CmdList);
@@ -1397,7 +1397,7 @@ public:
 	{
 		if (Bypass())
 		{
-			CMD_CONTEXT(SetBoundShaderState)(LocalBoundShaderState.BypassBSS);
+			CMD_CONTEXT(RHISetBoundShaderState)(LocalBoundShaderState.BypassBSS);
 			return;
 		}
 		new (AllocCommand<FRHICommandSetLocalBoundShaderState>()) FRHICommandSetLocalBoundShaderState(this, LocalBoundShaderState);
@@ -1424,7 +1424,7 @@ public:
 	{
 		if (Bypass())
 		{
-			CMD_CONTEXT(SetShaderUniformBuffer)(Shader, BaseIndex, UniformBuffer.BypassUniform);
+			CMD_CONTEXT(RHISetShaderUniformBuffer)(Shader, BaseIndex, UniformBuffer.BypassUniform);
 			return;
 		}
 		new (AllocCommand<FRHICommandSetLocalUniformBuffer<TShaderRHIParamRef> >()) FRHICommandSetLocalUniformBuffer<TShaderRHIParamRef>(this, Shader, BaseIndex, UniformBuffer);
@@ -1441,7 +1441,7 @@ public:
 	{
 		if (Bypass())
 		{
-			CMD_CONTEXT(SetShaderUniformBuffer)(Shader, BaseIndex, UniformBuffer);
+			CMD_CONTEXT(RHISetShaderUniformBuffer)(Shader, BaseIndex, UniformBuffer);
 			return;
 		}
 		new (AllocCommand<FRHICommandSetShaderUniformBuffer<TShaderRHI*, ECmdList::EGfx> >()) FRHICommandSetShaderUniformBuffer<TShaderRHI*, ECmdList::EGfx>(Shader, BaseIndex, UniformBuffer);
@@ -1457,7 +1457,7 @@ public:
 	{
 		if (Bypass())
 		{
-			CMD_CONTEXT(SetShaderParameter)(Shader, BufferIndex, BaseIndex, NumBytes, NewValue);
+			CMD_CONTEXT(RHISetShaderParameter)(Shader, BufferIndex, BaseIndex, NumBytes, NewValue);
 			return;
 		}
 		void* UseValue = Alloc(NumBytes, 16);
@@ -1475,7 +1475,7 @@ public:
 	{
 		if (Bypass())
 		{
-			CMD_CONTEXT(SetShaderTexture)(Shader, TextureIndex, Texture);
+			CMD_CONTEXT(RHISetShaderTexture)(Shader, TextureIndex, Texture);
 			return;
 		}
 		new (AllocCommand<FRHICommandSetShaderTexture<TShaderRHIParamRef, ECmdList::EGfx> >()) FRHICommandSetShaderTexture<TShaderRHIParamRef, ECmdList::EGfx>(Shader, TextureIndex, Texture);
@@ -1492,7 +1492,7 @@ public:
 	{
 		if (Bypass())
 		{
-			CMD_CONTEXT(SetShaderResourceViewParameter)(Shader, SamplerIndex, SRV);
+			CMD_CONTEXT(RHISetShaderResourceViewParameter)(Shader, SamplerIndex, SRV);
 			return;
 		}
 		new (AllocCommand<FRHICommandSetShaderResourceViewParameter<TShaderRHIParamRef, ECmdList::EGfx> >()) FRHICommandSetShaderResourceViewParameter<TShaderRHIParamRef, ECmdList::EGfx>(Shader, SamplerIndex, SRV);
@@ -1509,7 +1509,7 @@ public:
 	{
 		if (Bypass())
 		{
-			CMD_CONTEXT(SetShaderSampler)(Shader, SamplerIndex, State);
+			CMD_CONTEXT(RHISetShaderSampler)(Shader, SamplerIndex, State);
 			return;
 		}
 		new (AllocCommand<FRHICommandSetShaderSampler<TShaderRHIParamRef, ECmdList::EGfx> >()) FRHICommandSetShaderSampler<TShaderRHIParamRef, ECmdList::EGfx>(Shader, SamplerIndex, State);
@@ -1525,7 +1525,7 @@ public:
 	{
 		if (Bypass())
 		{
-			CMD_CONTEXT(SetUAVParameter)(Shader, UAVIndex, UAV);
+			CMD_CONTEXT(RHISetUAVParameter)(Shader, UAVIndex, UAV);
 			return;
 		}
 		new (AllocCommand<FRHICommandSetUAVParameter<FComputeShaderRHIParamRef, ECmdList::EGfx> >()) FRHICommandSetUAVParameter<FComputeShaderRHIParamRef, ECmdList::EGfx>(Shader, UAVIndex, UAV);
@@ -1540,7 +1540,7 @@ public:
 	{
 		if (Bypass())
 		{
-			CMD_CONTEXT(SetUAVParameter)(Shader, UAVIndex, UAV, InitialCount);
+			CMD_CONTEXT(RHISetUAVParameter)(Shader, UAVIndex, UAV, InitialCount);
 			return;
 		}
 		new (AllocCommand<FRHICommandSetUAVParameter_IntialCount<FComputeShaderRHIParamRef, ECmdList::EGfx> >()) FRHICommandSetUAVParameter_IntialCount<FComputeShaderRHIParamRef, ECmdList::EGfx>(Shader, UAVIndex, UAV, InitialCount);
@@ -1555,7 +1555,7 @@ public:
 	{
 		if (Bypass())
 		{
-			CMD_CONTEXT(SetBoundShaderState)(BoundShaderState);
+			CMD_CONTEXT(RHISetBoundShaderState)(BoundShaderState);
 			return;
 		}
 		new (AllocCommand<FRHICommandSetBoundShaderState>()) FRHICommandSetBoundShaderState(BoundShaderState);
@@ -1565,7 +1565,7 @@ public:
 	{
 		if (Bypass())
 		{
-			CMD_CONTEXT(SetRasterizerState)(State);
+			CMD_CONTEXT(RHISetRasterizerState)(State);
 			return;
 		}
 		if (StateCacheEnabled && CachedRasterizerState && CachedRasterizerState->State == State)
@@ -1579,7 +1579,7 @@ public:
 	{
 		if (Bypass())
 		{
-			CMD_CONTEXT(SetBlendState)(State, BlendFactor);
+			CMD_CONTEXT(RHISetBlendState)(State, BlendFactor);
 			return;
 		}
 		new (AllocCommand<FRHICommandSetBlendState>()) FRHICommandSetBlendState(State, BlendFactor);
@@ -1589,7 +1589,7 @@ public:
 	{
 		if (Bypass())
 		{
-			CMD_CONTEXT(DrawPrimitive)(PrimitiveType, BaseVertexIndex, NumPrimitives, NumInstances);
+			CMD_CONTEXT(RHIDrawPrimitive)(PrimitiveType, BaseVertexIndex, NumPrimitives, NumInstances);
 			return;
 		}
 		new (AllocCommand<FRHICommandDrawPrimitive>()) FRHICommandDrawPrimitive(PrimitiveType, BaseVertexIndex, NumPrimitives, NumInstances);
@@ -1599,7 +1599,7 @@ public:
 	{
 		if (Bypass())
 		{
-			CMD_CONTEXT(DrawIndexedPrimitive)(IndexBuffer, PrimitiveType, BaseVertexIndex, FirstInstance, NumVertices, StartIndex, NumPrimitives, NumInstances);
+			CMD_CONTEXT(RHIDrawIndexedPrimitive)(IndexBuffer, PrimitiveType, BaseVertexIndex, FirstInstance, NumVertices, StartIndex, NumPrimitives, NumInstances);
 			return;
 		}
 		new (AllocCommand<FRHICommandDrawIndexedPrimitive>()) FRHICommandDrawIndexedPrimitive(IndexBuffer, PrimitiveType, BaseVertexIndex, FirstInstance, NumVertices, StartIndex, NumPrimitives, NumInstances);
@@ -1609,7 +1609,7 @@ public:
 	{
 		if (Bypass())
 		{
-			CMD_CONTEXT(SetStreamSource)(StreamIndex, VertexBuffer, Stride, Offset);
+			CMD_CONTEXT(RHISetStreamSource)(StreamIndex, VertexBuffer, Stride, Offset);
 			return;
 		}
 		new (AllocCommand<FRHICommandSetStreamSource>()) FRHICommandSetStreamSource(StreamIndex, VertexBuffer, Stride, Offset);
@@ -1619,7 +1619,7 @@ public:
 	{
 		if (Bypass())
 		{
-			CMD_CONTEXT(SetDepthStencilState)(NewStateRHI, StencilRef);
+			CMD_CONTEXT(RHISetDepthStencilState)(NewStateRHI, StencilRef);
 			return;
 		}
 		if (StateCacheEnabled && CachedDepthStencilState && CachedDepthStencilState->State == NewStateRHI && CachedDepthStencilState->StencilRef == StencilRef)
@@ -1633,7 +1633,7 @@ public:
 	{
 		if (Bypass())
 		{
-			CMD_CONTEXT(SetViewport)(MinX, MinY, MinZ, MaxX, MaxY, MaxZ);
+			CMD_CONTEXT(RHISetViewport)(MinX, MinY, MinZ, MaxX, MaxY, MaxZ);
 			return;
 		}
 		new (AllocCommand<FRHICommandSetViewport>()) FRHICommandSetViewport(MinX, MinY, MinZ, MaxX, MaxY, MaxZ);
@@ -1643,7 +1643,7 @@ public:
 	{
 		if (Bypass())
 		{
-			CMD_CONTEXT(SetScissorRect)(bEnable, MinX, MinY, MaxX, MaxY);
+			CMD_CONTEXT(RHISetScissorRect)(bEnable, MinX, MinY, MaxX, MaxY);
 			return;
 		}
 		new (AllocCommand<FRHICommandSetScissorRect>()) FRHICommandSetScissorRect(bEnable, MinX, MinY, MaxX, MaxY);
@@ -1659,7 +1659,7 @@ public:
 	{
 		if (Bypass())
 		{
-			CMD_CONTEXT(SetRenderTargets)(
+			CMD_CONTEXT(RHISetRenderTargets)(
 				NewNumSimultaneousRenderTargets,
 				NewRenderTargetsRHI,
 				NewDepthStencilTargetRHI,
@@ -1679,7 +1679,7 @@ public:
 	{
 		if (Bypass())
 		{
-			CMD_CONTEXT(SetRenderTargetsAndClear)(RenderTargetsInfo);
+			CMD_CONTEXT(RHISetRenderTargetsAndClear)(RenderTargetsInfo);
 			return;
 		}
 		new (AllocCommand<FRHICommandSetRenderTargetsAndClear>()) FRHICommandSetRenderTargetsAndClear(RenderTargetsInfo);
@@ -1689,7 +1689,7 @@ public:
 	{
 		if (Bypass())
 		{
-			CMD_CONTEXT(BindClearMRTValues)(bClearColor, bClearDepth, bClearStencil);
+			CMD_CONTEXT(RHIBindClearMRTValues)(bClearColor, bClearDepth, bClearStencil);
 			return;
 		}
 		new (AllocCommand<FRHICommandBindClearMRTValues>()) FRHICommandBindClearMRTValues(bClearColor, bClearDepth, bClearStencil);
@@ -1699,7 +1699,7 @@ public:
 	{
 		if (Bypass())
 		{
-			CMD_CONTEXT(BeginDrawPrimitiveUP)(PrimitiveType, NumPrimitives, NumVertices, VertexDataStride, OutVertexData);
+			CMD_CONTEXT(RHIBeginDrawPrimitiveUP)(PrimitiveType, NumPrimitives, NumVertices, VertexDataStride, OutVertexData);
 			return;
 		}
 		check(!DrawUPData.OutVertexData && NumVertices * VertexDataStride > 0);
@@ -1717,7 +1717,7 @@ public:
 	{
 		if (Bypass())
 		{
-			CMD_CONTEXT(EndDrawPrimitiveUP)();
+			CMD_CONTEXT(RHIEndDrawPrimitiveUP)();
 			return;
 		}
 		check(DrawUPData.OutVertexData && DrawUPData.NumVertices);
@@ -1731,7 +1731,7 @@ public:
 	{
 		if (Bypass())
 		{
-			CMD_CONTEXT(BeginDrawIndexedPrimitiveUP)(PrimitiveType, NumPrimitives, NumVertices, VertexDataStride, OutVertexData, MinVertexIndex, NumIndices, IndexDataStride, OutIndexData);
+			CMD_CONTEXT(RHIBeginDrawIndexedPrimitiveUP)(PrimitiveType, NumPrimitives, NumVertices, VertexDataStride, OutVertexData, MinVertexIndex, NumIndices, IndexDataStride, OutIndexData);
 			return;
 		}
 		check(!DrawUPData.OutVertexData && !DrawUPData.OutIndexData && NumVertices * VertexDataStride > 0 && NumIndices * IndexDataStride > 0);
@@ -1755,7 +1755,7 @@ public:
 	{
 		if (Bypass())
 		{
-			CMD_CONTEXT(EndDrawIndexedPrimitiveUP)();
+			CMD_CONTEXT(RHIEndDrawIndexedPrimitiveUP)();
 			return;
 		}
 		check(DrawUPData.OutVertexData && DrawUPData.OutIndexData && DrawUPData.NumIndices && DrawUPData.NumVertices);
@@ -1772,7 +1772,7 @@ public:
 	{
 		if (Bypass())
 		{
-			CMD_CONTEXT(SetComputeShader)(ComputeShader);
+			CMD_CONTEXT(RHISetComputeShader)(ComputeShader);
 			return;
 		}
 		new (AllocCommand<FRHICommandSetComputeShader<ECmdList::EGfx>>()) FRHICommandSetComputeShader<ECmdList::EGfx>(ComputeShader);
@@ -1782,7 +1782,7 @@ public:
 	{
 		if (Bypass())
 		{
-			CMD_CONTEXT(DispatchComputeShader)(ThreadGroupCountX, ThreadGroupCountY, ThreadGroupCountZ);
+			CMD_CONTEXT(RHIDispatchComputeShader)(ThreadGroupCountX, ThreadGroupCountY, ThreadGroupCountZ);
 			return;
 		}
 		new (AllocCommand<FRHICommandDispatchComputeShader<ECmdList::EGfx>>()) FRHICommandDispatchComputeShader<ECmdList::EGfx>(ThreadGroupCountX, ThreadGroupCountY, ThreadGroupCountZ);
@@ -1792,7 +1792,7 @@ public:
 	{
 		if (Bypass())
 		{
-			CMD_CONTEXT(DispatchIndirectComputeShader)(ArgumentBuffer, ArgumentOffset);
+			CMD_CONTEXT(RHIDispatchIndirectComputeShader)(ArgumentBuffer, ArgumentOffset);
 			return;
 		}
 		new (AllocCommand<FRHICommandDispatchIndirectComputeShader<ECmdList::EGfx>>()) FRHICommandDispatchIndirectComputeShader<ECmdList::EGfx>(ArgumentBuffer, ArgumentOffset);
@@ -1802,7 +1802,7 @@ public:
 	{
 		if (Bypass())
 		{
-			CMD_CONTEXT(AutomaticCacheFlushAfterComputeShader)(bEnable);
+			CMD_CONTEXT(RHIAutomaticCacheFlushAfterComputeShader)(bEnable);
 			return;
 		}
 		new (AllocCommand<FRHICommandAutomaticCacheFlushAfterComputeShader>()) FRHICommandAutomaticCacheFlushAfterComputeShader(bEnable);
@@ -1812,7 +1812,7 @@ public:
 	{
 		if (Bypass())
 		{
-			CMD_CONTEXT(FlushComputeShaderCache)();
+			CMD_CONTEXT(RHIFlushComputeShaderCache)();
 			return;
 		}
 		new (AllocCommand<FRHICommandFlushComputeShaderCache>()) FRHICommandFlushComputeShaderCache();
@@ -1822,7 +1822,7 @@ public:
 	{
 		if (Bypass())
 		{
-			CMD_CONTEXT(DrawPrimitiveIndirect)(PrimitiveType, ArgumentBuffer, ArgumentOffset);
+			CMD_CONTEXT(RHIDrawPrimitiveIndirect)(PrimitiveType, ArgumentBuffer, ArgumentOffset);
 			return;
 		}
 		new (AllocCommand<FRHICommandDrawPrimitiveIndirect>()) FRHICommandDrawPrimitiveIndirect(PrimitiveType, ArgumentBuffer, ArgumentOffset);
@@ -1832,7 +1832,7 @@ public:
 	{
 		if (Bypass())
 		{
-			CMD_CONTEXT(DrawIndexedIndirect)(IndexBufferRHI, PrimitiveType, ArgumentsBufferRHI, DrawArgumentsIndex, NumInstances);
+			CMD_CONTEXT(RHIDrawIndexedIndirect)(IndexBufferRHI, PrimitiveType, ArgumentsBufferRHI, DrawArgumentsIndex, NumInstances);
 			return;
 		}
 		new (AllocCommand<FRHICommandDrawIndexedIndirect>()) FRHICommandDrawIndexedIndirect(IndexBufferRHI, PrimitiveType, ArgumentsBufferRHI, DrawArgumentsIndex, NumInstances);
@@ -1842,7 +1842,7 @@ public:
 	{
 		if (Bypass())
 		{
-			CMD_CONTEXT(DrawIndexedPrimitiveIndirect)(PrimitiveType, IndexBuffer, ArgumentsBuffer, ArgumentOffset);
+			CMD_CONTEXT(RHIDrawIndexedPrimitiveIndirect)(PrimitiveType, IndexBuffer, ArgumentsBuffer, ArgumentOffset);
 			return;
 		}
 		new (AllocCommand<FRHICommandDrawIndexedPrimitiveIndirect>()) FRHICommandDrawIndexedPrimitiveIndirect(PrimitiveType, IndexBuffer, ArgumentsBuffer, ArgumentOffset);
@@ -1852,7 +1852,7 @@ public:
 	{
 		if (Bypass())
 		{
-			CMD_CONTEXT(EnableDepthBoundsTest)(bEnable, MinDepth, MaxDepth);
+			CMD_CONTEXT(RHIEnableDepthBoundsTest)(bEnable, MinDepth, MaxDepth);
 			return;
 		}
 		new (AllocCommand<FRHICommandEnableDepthBoundsTest>()) FRHICommandEnableDepthBoundsTest(bEnable, MinDepth, MaxDepth);
@@ -1862,7 +1862,7 @@ public:
 	{
 		if (Bypass())
 		{
-			CMD_CONTEXT(ClearUAV)(UnorderedAccessViewRHI, Values);
+			CMD_CONTEXT(RHIClearUAV)(UnorderedAccessViewRHI, Values);
 			return;
 		}
 		new (AllocCommand<FRHICommandClearUAV>()) FRHICommandClearUAV(UnorderedAccessViewRHI, Values);
@@ -1872,7 +1872,7 @@ public:
 	{
 		if (Bypass())
 		{
-			CMD_CONTEXT(CopyToResolveTarget)(SourceTextureRHI, DestTextureRHI, bKeepOriginalSurface, ResolveParams);
+			CMD_CONTEXT(RHICopyToResolveTarget)(SourceTextureRHI, DestTextureRHI, bKeepOriginalSurface, ResolveParams);
 			return;
 		}
 		new (AllocCommand<FRHICommandCopyToResolveTarget>()) FRHICommandCopyToResolveTarget(SourceTextureRHI, DestTextureRHI, bKeepOriginalSurface, ResolveParams);
@@ -1882,7 +1882,7 @@ public:
 	{
 		if (Bypass())
 		{
-			CMD_CONTEXT(Clear)(bClearColor, Color, bClearDepth, Depth, bClearStencil, Stencil, ExcludeRect);
+			CMD_CONTEXT(RHIClear)(bClearColor, Color, bClearDepth, Depth, bClearStencil, Stencil, ExcludeRect);
 			return;
 		}
 		new (AllocCommand<FRHICommandClear>()) FRHICommandClear(bClearColor, Color, bClearDepth, Depth, bClearStencil, Stencil, ExcludeRect);
@@ -1892,7 +1892,7 @@ public:
 	{
 		if (Bypass())
 		{
-			CMD_CONTEXT(ClearMRT)(bClearColor, NumClearColors, ClearColorArray, bClearDepth, Depth, bClearStencil, Stencil, ExcludeRect);
+			CMD_CONTEXT(RHIClearMRT)(bClearColor, NumClearColors, ClearColorArray, bClearDepth, Depth, bClearStencil, Stencil, ExcludeRect);
 			return;
 		}
 		new (AllocCommand<FRHICommandClearMRT>()) FRHICommandClearMRT(bClearColor, NumClearColors, ClearColorArray, bClearDepth, Depth, bClearStencil, Stencil, ExcludeRect);
@@ -1902,7 +1902,7 @@ public:
 	{
 		if (Bypass())
 		{
-			CMD_CONTEXT(BeginRenderQuery)(RenderQuery);
+			CMD_CONTEXT(RHIBeginRenderQuery)(RenderQuery);
 			return;
 		}
 		new (AllocCommand<FRHICommandBeginRenderQuery>()) FRHICommandBeginRenderQuery(RenderQuery);
@@ -1911,7 +1911,7 @@ public:
 	{
 		if (Bypass())
 		{
-			CMD_CONTEXT(EndRenderQuery)(RenderQuery);
+			CMD_CONTEXT(RHIEndRenderQuery)(RenderQuery);
 			return;
 		}
 		new (AllocCommand<FRHICommandEndRenderQuery>()) FRHICommandEndRenderQuery(RenderQuery);
@@ -1921,7 +1921,7 @@ public:
 	{
 		if (Bypass())
 		{
-			CMD_CONTEXT(BeginOcclusionQueryBatch)();
+			CMD_CONTEXT(RHIBeginOcclusionQueryBatch)();
 			return;
 		}
 		new (AllocCommand<FRHICommandBeginOcclusionQueryBatch>()) FRHICommandBeginOcclusionQueryBatch();
@@ -1931,7 +1931,7 @@ public:
 	{
 		if (Bypass())
 		{
-			CMD_CONTEXT(EndOcclusionQueryBatch)();
+			CMD_CONTEXT(RHIEndOcclusionQueryBatch)();
 			return;
 		}
 		new (AllocCommand<FRHICommandEndOcclusionQueryBatch>()) FRHICommandEndOcclusionQueryBatch();
@@ -1941,7 +1941,7 @@ public:
 	{
 		if (Bypass())
 		{
-			CMD_CONTEXT(SubmitCommandsHint)();
+			CMD_CONTEXT(RHISubmitCommandsHint)();
 			return;
 		}
 		new (AllocCommand<FRHICommandSubmitCommandsHint<ECmdList::EGfx>>()) FRHICommandSubmitCommandsHint<ECmdList::EGfx>();
@@ -1952,7 +1952,7 @@ public:
 		FTextureRHIParamRef Texture = InTexture;
 		if (Bypass())
 		{
-			CMD_CONTEXT(TransitionResources)(TransitionType, &Texture, 1);
+			CMD_CONTEXT(RHITransitionResources)(TransitionType, &Texture, 1);
 			return;
 		}
 		new (AllocCommand<FRHICommandTransitionTextures>()) FRHICommandTransitionTextures(TransitionType, &Texture, 1);
@@ -1962,7 +1962,7 @@ public:
 	{
 		if (Bypass())
 		{
-			CMD_CONTEXT(TransitionResources)(TransitionType, InTextures, NumTextures);
+			CMD_CONTEXT(RHITransitionResources)(TransitionType, InTextures, NumTextures);
 			return;
 		}
 		new (AllocCommand<FRHICommandTransitionTextures>()) FRHICommandTransitionTextures(TransitionType, InTextures, NumTextures);
@@ -1972,7 +1972,7 @@ public:
 	{
 		if (Bypass())
 		{
-			CMD_CONTEXT(TransitionResources)(TransitionType, &InTextures[0], InTextures.Num());
+			CMD_CONTEXT(RHITransitionResources)(TransitionType, &InTextures[0], InTextures.Num());
 			return;
 		}
 		new (AllocCommand<FRHICommandTransitionTexturesArray>()) FRHICommandTransitionTexturesArray(TransitionType, InTextures);
@@ -1983,7 +1983,7 @@ public:
 		FUnorderedAccessViewRHIParamRef UAV = InUAV;
 		if (Bypass())
 		{
-			CMD_CONTEXT(TransitionResources)(TransitionType, TransitionPipeline, &UAV, 1, WriteFence);
+			CMD_CONTEXT(RHITransitionResources)(TransitionType, TransitionPipeline, &UAV, 1, WriteFence);
 			return;
 		}
 		new (AllocCommand<FRHICommandTransitionUAVs<ECmdList::EGfx>>()) FRHICommandTransitionUAVs<ECmdList::EGfx>(TransitionType, TransitionPipeline, &UAV, 1, WriteFence);
@@ -1998,7 +1998,7 @@ public:
 	{
 		if (Bypass())
 		{
-			CMD_CONTEXT(TransitionResources)(TransitionType, TransitionPipeline, InUAVs, NumUAVs, WriteFence);
+			CMD_CONTEXT(RHITransitionResources)(TransitionType, TransitionPipeline, InUAVs, NumUAVs, WriteFence);
 			return;
 		}
 		new (AllocCommand<FRHICommandTransitionUAVs<ECmdList::EGfx>>()) FRHICommandTransitionUAVs<ECmdList::EGfx>(TransitionType, TransitionPipeline, InUAVs, NumUAVs, WriteFence);
@@ -2013,7 +2013,7 @@ public:
 	{
 		if (Bypass())
 		{
-			CMD_CONTEXT(WaitComputeFence)(WaitFence);
+			CMD_CONTEXT(RHIWaitComputeFence)(WaitFence);
 			return;
 		}
 		new (AllocCommand<FRHICommandWaitComputeFence<ECmdList::EGfx>>()) FRHICommandWaitComputeFence<ECmdList::EGfx>(WaitFence);
@@ -2033,7 +2033,7 @@ public:
 	{
 		if (Bypass())
 		{
-			CMD_CONTEXT(PushEvent)(Name, Color);
+			CMD_CONTEXT(RHIPushEvent)(Name, Color);
 			return;
 		}
 		int32 Len = FCString::Strlen(Name) + 1;
@@ -2046,7 +2046,7 @@ public:
 	{
 		if (Bypass())
 		{
-			CMD_CONTEXT(PopEvent)();
+			CMD_CONTEXT(RHIPopEvent)();
 			return;
 		}
 		new (AllocCommand<FRHICommandPopEvent<ECmdList::EGfx>>()) FRHICommandPopEvent<ECmdList::EGfx>();
@@ -2080,7 +2080,7 @@ public:
 	{
 		if (Bypass())
 		{
-			COMPUTE_CONTEXT(SetShaderUniformBuffer)(Shader, BaseIndex, UniformBuffer);
+			COMPUTE_CONTEXT(RHISetShaderUniformBuffer)(Shader, BaseIndex, UniformBuffer);
 			return;
 		}
 		new (AllocCommand<FRHICommandSetShaderUniformBuffer<FComputeShaderRHIParamRef, ECmdList::ECompute> >()) FRHICommandSetShaderUniformBuffer<FComputeShaderRHIParamRef, ECmdList::ECompute>(Shader, BaseIndex, UniformBuffer);		
@@ -2095,7 +2095,7 @@ public:
 	{
 		if (Bypass())
 		{
-			COMPUTE_CONTEXT(SetShaderParameter)(Shader, BufferIndex, BaseIndex, NumBytes, NewValue);
+			COMPUTE_CONTEXT(RHISetShaderParameter)(Shader, BufferIndex, BaseIndex, NumBytes, NewValue);
 			return;
 		}
 		void* UseValue = Alloc(NumBytes, 16);
@@ -2112,7 +2112,7 @@ public:
 	{
 		if (Bypass())
 		{
-			COMPUTE_CONTEXT(SetShaderTexture)(Shader, TextureIndex, Texture);
+			COMPUTE_CONTEXT(RHISetShaderTexture)(Shader, TextureIndex, Texture);
 			return;
 		}
 		new (AllocCommand<FRHICommandSetShaderTexture<FComputeShaderRHIParamRef, ECmdList::ECompute> >()) FRHICommandSetShaderTexture<FComputeShaderRHIParamRef, ECmdList::ECompute>(Shader, TextureIndex, Texture);
@@ -2122,7 +2122,7 @@ public:
 	{
 		if (Bypass())
 		{
-			COMPUTE_CONTEXT(SetShaderResourceViewParameter)(Shader, SamplerIndex, SRV);
+			COMPUTE_CONTEXT(RHISetShaderResourceViewParameter)(Shader, SamplerIndex, SRV);
 			return;
 		}
 		new (AllocCommand<FRHICommandSetShaderResourceViewParameter<FComputeShaderRHIParamRef, ECmdList::ECompute> >()) FRHICommandSetShaderResourceViewParameter<FComputeShaderRHIParamRef, ECmdList::ECompute>(Shader, SamplerIndex, SRV);
@@ -2132,7 +2132,7 @@ public:
 	{
 		if (Bypass())
 		{
-			COMPUTE_CONTEXT(SetShaderSampler)(Shader, SamplerIndex, State);
+			COMPUTE_CONTEXT(RHISetShaderSampler)(Shader, SamplerIndex, State);
 			return;
 		}
 		new (AllocCommand<FRHICommandSetShaderSampler<FComputeShaderRHIParamRef, ECmdList::ECompute> >()) FRHICommandSetShaderSampler<FComputeShaderRHIParamRef, ECmdList::ECompute>(Shader, SamplerIndex, State);
@@ -2142,7 +2142,7 @@ public:
 	{
 		if (Bypass())
 		{
-			COMPUTE_CONTEXT(SetUAVParameter)(Shader, UAVIndex, UAV);
+			COMPUTE_CONTEXT(RHISetUAVParameter)(Shader, UAVIndex, UAV);
 			return;
 		}
 		new (AllocCommand<FRHICommandSetUAVParameter<FComputeShaderRHIParamRef, ECmdList::ECompute> >()) FRHICommandSetUAVParameter<FComputeShaderRHIParamRef, ECmdList::ECompute>(Shader, UAVIndex, UAV);
@@ -2152,7 +2152,7 @@ public:
 	{
 		if (Bypass())
 		{
-			COMPUTE_CONTEXT(SetUAVParameter)(Shader, UAVIndex, UAV, InitialCount);
+			COMPUTE_CONTEXT(RHISetUAVParameter)(Shader, UAVIndex, UAV, InitialCount);
 			return;
 		}
 		new (AllocCommand<FRHICommandSetUAVParameter_IntialCount<FComputeShaderRHIParamRef, ECmdList::ECompute> >()) FRHICommandSetUAVParameter_IntialCount<FComputeShaderRHIParamRef, ECmdList::ECompute>(Shader, UAVIndex, UAV, InitialCount);
@@ -2162,7 +2162,7 @@ public:
 	{
 		if (Bypass())
 		{
-			COMPUTE_CONTEXT(SetComputeShader)(ComputeShader);
+			COMPUTE_CONTEXT(RHISetComputeShader)(ComputeShader);
 			return;
 		}
 		new (AllocCommand<FRHICommandSetComputeShader<ECmdList::ECompute> >()) FRHICommandSetComputeShader<ECmdList::ECompute>(ComputeShader);
@@ -2172,7 +2172,7 @@ public:
 	{
 		if (Bypass())
 		{
-			COMPUTE_CONTEXT(SetAsyncComputeBudget)(Budget);
+			COMPUTE_CONTEXT(RHISetAsyncComputeBudget)(Budget);
 			return;
 		}
 		new (AllocCommand<FRHICommandSetAsyncComputeBudget<ECmdList::ECompute>>()) FRHICommandSetAsyncComputeBudget<ECmdList::ECompute>(Budget);
@@ -2182,7 +2182,7 @@ public:
 	{
 		if (Bypass())
 		{
-			COMPUTE_CONTEXT(DispatchComputeShader)(ThreadGroupCountX, ThreadGroupCountY, ThreadGroupCountZ);
+			COMPUTE_CONTEXT(RHIDispatchComputeShader)(ThreadGroupCountX, ThreadGroupCountY, ThreadGroupCountZ);
 			return;
 		}
 		new (AllocCommand<FRHICommandDispatchComputeShader<ECmdList::ECompute> >()) FRHICommandDispatchComputeShader<ECmdList::ECompute>(ThreadGroupCountX, ThreadGroupCountY, ThreadGroupCountZ);
@@ -2192,7 +2192,7 @@ public:
 	{
 		if (Bypass())
 		{
-			COMPUTE_CONTEXT(DispatchIndirectComputeShader)(ArgumentBuffer, ArgumentOffset);
+			COMPUTE_CONTEXT(RHIDispatchIndirectComputeShader)(ArgumentBuffer, ArgumentOffset);
 			return;
 		}
 		new (AllocCommand<FRHICommandDispatchIndirectComputeShader<ECmdList::ECompute> >()) FRHICommandDispatchIndirectComputeShader<ECmdList::ECompute>(ArgumentBuffer, ArgumentOffset);
@@ -2203,7 +2203,7 @@ public:
 		FUnorderedAccessViewRHIParamRef UAV = InUAV;
 		if (Bypass())
 		{
-			COMPUTE_CONTEXT(TransitionResources)(TransitionType, TransitionPipeline, &UAV, 1, WriteFence);
+			COMPUTE_CONTEXT(RHITransitionResources)(TransitionType, TransitionPipeline, &UAV, 1, WriteFence);
 			return;
 		}
 		new (AllocCommand<FRHICommandTransitionUAVs<ECmdList::ECompute> >()) FRHICommandTransitionUAVs<ECmdList::ECompute>(TransitionType, TransitionPipeline, &UAV, 1, WriteFence);
@@ -2218,7 +2218,7 @@ public:
 	{
 		if (Bypass())
 		{
-			COMPUTE_CONTEXT(TransitionResources)(TransitionType, TransitionPipeline, InUAVs, NumUAVs, WriteFence);
+			COMPUTE_CONTEXT(RHITransitionResources)(TransitionType, TransitionPipeline, InUAVs, NumUAVs, WriteFence);
 			return;
 		}
 		new (AllocCommand<FRHICommandTransitionUAVs<ECmdList::ECompute> >()) FRHICommandTransitionUAVs<ECmdList::ECompute>(TransitionType, TransitionPipeline, InUAVs, NumUAVs, WriteFence);
@@ -2233,7 +2233,7 @@ public:
 	{
 		if (Bypass())
 		{
-			COMPUTE_CONTEXT(PushEvent)(Name, Color);
+			COMPUTE_CONTEXT(RHIPushEvent)(Name, Color);
 			return;
 		}
 		int32 Len = FCString::Strlen(Name) + 1;
@@ -2246,7 +2246,7 @@ public:
 	{
 		if (Bypass())
 		{
-			COMPUTE_CONTEXT(PopEvent)();
+			COMPUTE_CONTEXT(RHIPopEvent)();
 			return;
 		}
 		new (AllocCommand<FRHICommandPopEvent<ECmdList::ECompute> >()) FRHICommandPopEvent<ECmdList::ECompute>();
@@ -2271,7 +2271,7 @@ public:
 	{
 		if (Bypass())
 		{
-			COMPUTE_CONTEXT(SubmitCommandsHint)();
+			COMPUTE_CONTEXT(RHISubmitCommandsHint)();
 			return;
 		}
 		new (AllocCommand<FRHICommandSubmitCommandsHint<ECmdList::ECompute>>()) FRHICommandSubmitCommandsHint<ECmdList::ECompute>();
@@ -2281,7 +2281,7 @@ public:
 	{
 		if (Bypass())
 		{
-			COMPUTE_CONTEXT(WaitComputeFence)(WaitFence);
+			COMPUTE_CONTEXT(RHIWaitComputeFence)(WaitFence);
 			return;
 		}
 		new (AllocCommand<FRHICommandWaitComputeFence<ECmdList::ECompute>>()) FRHICommandWaitComputeFence<ECmdList::ECompute>(WaitFence);
diff --git a/Engine/Source/Runtime/RHI/Public/RHICommandListCommandExecutes.inl b/Engine/Source/Runtime/RHI/Public/RHICommandListCommandExecutes.inl
--- a/Engine/Source/Runtime/RHI/Public/RHICommandListCommandExecutes.inl
+++ b/Engine/Source/Runtime/RHI/Public/RHICommandListCommandExecutes.inl
@@ -5,17 +5,17 @@
 =============================================================================*/
 
 #if !defined(INTERNAL_DECORATOR)
-	#define INTERNAL_DECORATOR(Method) CmdList.GetContext().RHI##Method
+	#define INTERNAL_DECORATOR(Method) CmdList.GetContext().Method
 #endif
 
 //for functions where the signatures do not match between gfx and compute commandlists
 #if !defined(INTERNAL_DECORATOR_COMPUTE)
-#define INTERNAL_DECORATOR_COMPUTE(Method) CmdList.GetComputeContext().RHI##Method
+#define INTERNAL_DECORATOR_COMPUTE(Method) CmdList.GetComputeContext().Method
 #endif
 
 //for functions where the signatures match between gfx and compute commandlists
 #if !defined(INTERNAL_DECORATOR_CONTEXT_PARAM1)
-#define INTERNAL_DECORATOR_CONTEXT(Method) IRHIComputeContext& Context = (CmdListType == ECmdList::EGfx) ? CmdList.GetContext() : CmdList.GetComputeContext(); Context.RHI##Method
+#define INTERNAL_DECORATOR_CONTEXT(Method) IRHIComputeContext& Context = (CmdListType == ECmdList::EGfx) ? CmdList.GetContext() : CmdList.GetComputeContext(); Context.Method
 #endif
 
 
@@ -23,20 +23,20 @@
 void FRHICommandSetRasterizerState::Execute(FRHICommandListBase& CmdList)
 {
 	RHISTAT(SetRasterizerState);
-	INTERNAL_DECORATOR(SetRasterizerState)(State);
+	INTERNAL_DECORATOR(RHISetRasterizerState)(State);
 }
 
 void FRHICommandSetDepthStencilState::Execute(FRHICommandListBase& CmdList)
 {
 	RHISTAT(SetDepthStencilState);
-	INTERNAL_DECORATOR(SetDepthStencilState)(State, StencilRef);
+	INTERNAL_DECORATOR(RHISetDepthStencilState)(State, StencilRef);
 }
 
 template <typename TShaderRHIParamRef, ECmdList CmdListType>
 void FRHICommandSetShaderParameter<TShaderRHIParamRef, CmdListType>::Execute(FRHICommandListBase& CmdList)
 {
 	RHISTAT(SetShaderParameter);
-	INTERNAL_DECORATOR(SetShaderParameter)(Shader, BufferIndex, BaseIndex, NumBytes, NewValue); 
+	INTERNAL_DECORATOR(RHISetShaderParameter)(Shader, BufferIndex, BaseIndex, NumBytes, NewValue); 
 }
 template struct FRHICommandSetShaderParameter<FVertexShaderRHIParamRef, ECmdList::EGfx>;
 template struct FRHICommandSetShaderParameter<FHullShaderRHIParamRef, ECmdList::EGfx>;
@@ -47,14 +47,14 @@ template struct FRHICommandSetShaderParameter<FComputeShaderRHIParamRef, ECmdLis
 template<> void FRHICommandSetShaderParameter<FComputeShaderRHIParamRef, ECmdList::ECompute>::Execute(FRHICommandListBase& CmdList)
 {
 	RHISTAT(SetShaderParameter);
-	INTERNAL_DECORATOR_COMPUTE(SetShaderParameter)(Shader, BufferIndex, BaseIndex, NumBytes, NewValue);
+	INTERNAL_DECORATOR_COMPUTE(RHISetShaderParameter)(Shader, BufferIndex, BaseIndex, NumBytes, NewValue);
 }
 
 template <typename TShaderRHIParamRef, ECmdList CmdListType>
 void FRHICommandSetShaderUniformBuffer<TShaderRHIParamRef, CmdListType>::Execute(FRHICommandListBase& CmdList)
 {
 	RHISTAT(SetShaderUniformBuffer);
-	INTERNAL_DECORATOR(SetShaderUniformBuffer)(Shader, BaseIndex, UniformBuffer);
+	INTERNAL_DECORATOR(RHISetShaderUniformBuffer)(Shader, BaseIndex, UniformBuffer);
 }
 template struct FRHICommandSetShaderUniformBuffer<FVertexShaderRHIParamRef, ECmdList::EGfx>;
 template struct FRHICommandSetShaderUniformBuffer<FHullShaderRHIParamRef, ECmdList::EGfx>;
@@ -65,14 +65,14 @@ template struct FRHICommandSetShaderUniformBuffer<FComputeShaderRHIParamRef, ECm
 template<> void FRHICommandSetShaderUniformBuffer<FComputeShaderRHIParamRef, ECmdList::ECompute>::Execute(FRHICommandListBase& CmdList)
 {
 	RHISTAT(SetShaderUniformBuffer);
-	INTERNAL_DECORATOR_COMPUTE(SetShaderUniformBuffer)(Shader, BaseIndex, UniformBuffer);
+	INTERNAL_DECORATOR_COMPUTE(RHISetShaderUniformBuffer)(Shader, BaseIndex, UniformBuffer);
 }
 
 template <typename TShaderRHIParamRef, ECmdList CmdListType>
 void FRHICommandSetShaderTexture<TShaderRHIParamRef, CmdListType>::Execute(FRHICommandListBase& CmdList)
 {
 	RHISTAT(SetShaderTexture);
-	INTERNAL_DECORATOR(SetShaderTexture)(Shader, TextureIndex, Texture);
+	INTERNAL_DECORATOR(RHISetShaderTexture)(Shader, TextureIndex, Texture);
 }
 template struct FRHICommandSetShaderTexture<FVertexShaderRHIParamRef, ECmdList::EGfx>;
 template struct FRHICommandSetShaderTexture<FHullShaderRHIParamRef, ECmdList::EGfx>;
@@ -83,14 +83,14 @@ template struct FRHICommandSetShaderTexture<FComputeShaderRHIParamRef, ECmdList:
 template<> void FRHICommandSetShaderTexture<FComputeShaderRHIParamRef, ECmdList::ECompute>::Execute(FRHICommandListBase& CmdList)
 {
 	RHISTAT(SetShaderTexture);
-	INTERNAL_DECORATOR_COMPUTE(SetShaderTexture)(Shader, TextureIndex, Texture);
+	INTERNAL_DECORATOR_COMPUTE(RHISetShaderTexture)(Shader, TextureIndex, Texture);
 }
 
 template <typename TShaderRHIParamRef, ECmdList CmdListType>
 void FRHICommandSetShaderResourceViewParameter<TShaderRHIParamRef, CmdListType>::Execute(FRHICommandListBase& CmdList)
 {
 	RHISTAT(SetShaderResourceViewParameter);
-	INTERNAL_DECORATOR(SetShaderResourceViewParameter)(Shader, SamplerIndex, SRV);
+	INTERNAL_DECORATOR(RHISetShaderResourceViewParameter)(Shader, SamplerIndex, SRV);
 }
 template struct FRHICommandSetShaderResourceViewParameter<FVertexShaderRHIParamRef, ECmdList::EGfx>;
 template struct FRHICommandSetShaderResourceViewParameter<FHullShaderRHIParamRef, ECmdList::EGfx>;
@@ -101,14 +101,14 @@ template struct FRHICommandSetShaderResourceViewParameter<FComputeShaderRHIParam
 template<> void FRHICommandSetShaderResourceViewParameter<FComputeShaderRHIParamRef, ECmdList::ECompute>::Execute(FRHICommandListBase& CmdList)
 {
 	RHISTAT(SetShaderResourceViewParameter);
-	INTERNAL_DECORATOR_COMPUTE(SetShaderResourceViewParameter)(Shader, SamplerIndex, SRV);
+	INTERNAL_DECORATOR_COMPUTE(RHISetShaderResourceViewParameter)(Shader, SamplerIndex, SRV);
 }
 
 template <typename TShaderRHIParamRef, ECmdList CmdListType>
 void FRHICommandSetUAVParameter<TShaderRHIParamRef, CmdListType>::Execute(FRHICommandListBase& CmdList)
 {
 	RHISTAT(SetUAVParameter);
-	INTERNAL_DECORATOR_CONTEXT(SetUAVParameter)(Shader, UAVIndex, UAV);
+	INTERNAL_DECORATOR_CONTEXT(RHISetUAVParameter)(Shader, UAVIndex, UAV);
 }
 template struct FRHICommandSetUAVParameter<FComputeShaderRHIParamRef, ECmdList::EGfx>;
 template struct FRHICommandSetUAVParameter<FComputeShaderRHIParamRef, ECmdList::ECompute>;
@@ -117,7 +117,7 @@ template <typename TShaderRHIParamRef, ECmdList CmdListType>
 void FRHICommandSetUAVParameter_IntialCount<TShaderRHIParamRef, CmdListType>::Execute(FRHICommandListBase& CmdList)
 {
 	RHISTAT(SetUAVParameter);
-	INTERNAL_DECORATOR_CONTEXT(SetUAVParameter)(Shader, UAVIndex, UAV, InitialCount);
+	INTERNAL_DECORATOR_CONTEXT(RHISetUAVParameter)(Shader, UAVIndex, UAV, InitialCount);
 }
 template struct FRHICommandSetUAVParameter_IntialCount<FComputeShaderRHIParamRef, ECmdList::EGfx>;
 template struct FRHICommandSetUAVParameter_IntialCount<FComputeShaderRHIParamRef, ECmdList::ECompute>;
@@ -126,7 +126,7 @@ template <typename TShaderRHIParamRef, ECmdList CmdListType>
 void FRHICommandSetShaderSampler<TShaderRHIParamRef, CmdListType>::Execute(FRHICommandListBase& CmdList)
 {
 	RHISTAT(SetShaderSampler);
-	INTERNAL_DECORATOR(SetShaderSampler)(Shader, SamplerIndex, Sampler);
+	INTERNAL_DECORATOR(RHISetShaderSampler)(Shader, SamplerIndex, Sampler);
 }
 template struct FRHICommandSetShaderSampler<FVertexShaderRHIParamRef, ECmdList::EGfx>;
 template struct FRHICommandSetShaderSampler<FHullShaderRHIParamRef, ECmdList::EGfx>;
@@ -137,55 +137,55 @@ template struct FRHICommandSetShaderSampler<FComputeShaderRHIParamRef, ECmdList:
 template<> void FRHICommandSetShaderSampler<FComputeShaderRHIParamRef, ECmdList::ECompute>::Execute(FRHICommandListBase& CmdList)
 {
 	RHISTAT(SetShaderSampler);
-	INTERNAL_DECORATOR_COMPUTE(SetShaderSampler)(Shader, SamplerIndex, Sampler);
+	INTERNAL_DECORATOR_COMPUTE(RHISetShaderSampler)(Shader, SamplerIndex, Sampler);
 }
 
 void FRHICommandDrawPrimitive::Execute(FRHICommandListBase& CmdList)
 {
 	RHISTAT(DrawPrimitive);
-	INTERNAL_DECORATOR(DrawPrimitive)(PrimitiveType, BaseVertexIndex, NumPrimitives, NumInstances);
+	INTERNAL_DECORATOR(RHIDrawPrimitive)(PrimitiveType, BaseVertexIndex, NumPrimitives, NumInstances);
 }
 
 void FRHICommandDrawIndexedPrimitive::Execute(FRHICommandListBase& CmdList)
 {
 	RHISTAT(DrawIndexedPrimitive);
-	INTERNAL_DECORATOR(DrawIndexedPrimitive)(IndexBuffer, PrimitiveType, BaseVertexIndex, FirstInstance, NumVertices, StartIndex, NumPrimitives, NumInstances);
+	INTERNAL_DECORATOR(RHIDrawIndexedPrimitive)(IndexBuffer, PrimitiveType, BaseVertexIndex, FirstInstance, NumVertices, StartIndex, NumPrimitives, NumInstances);
 }
 
 void FRHICommandSetBoundShaderState::Execute(FRHICommandListBase& CmdList)
 {
 	RHISTAT(SetBoundShaderState);
-	INTERNAL_DECORATOR(SetBoundShaderState)(BoundShaderState);
+	INTERNAL_DECORATOR(RHISetBoundShaderState)(BoundShaderState);
 }
 
 void FRHICommandSetBlendState::Execute(FRHICommandListBase& CmdList)
 {
 	RHISTAT(SetBlendState);
-	INTERNAL_DECORATOR(SetBlendState)(State, BlendFactor);
+	INTERNAL_DECORATOR(RHISetBlendState)(State, BlendFactor);
 }
 
 void FRHICommandSetStreamSource::Execute(FRHICommandListBase& CmdList)
 {
 	RHISTAT(SetStreamSource);
-	INTERNAL_DECORATOR(SetStreamSource)(StreamIndex, VertexBuffer, Stride, Offset);
+	INTERNAL_DECORATOR(RHISetStreamSource)(StreamIndex, VertexBuffer, Stride, Offset);
 }
 
 void FRHICommandSetViewport::Execute(FRHICommandListBase& CmdList)
 {
 	RHISTAT(SetViewport);
-	INTERNAL_DECORATOR(SetViewport)(MinX, MinY, MinZ, MaxX, MaxY, MaxZ);
+	INTERNAL_DECORATOR(RHISetViewport)(MinX, MinY, MinZ, MaxX, MaxY, MaxZ);
 }
 
 void FRHICommandSetScissorRect::Execute(FRHICommandListBase& CmdList)
 {
 	RHISTAT(SetScissorRect);
-	INTERNAL_DECORATOR(SetScissorRect)(bEnable, MinX, MinY, MaxX, MaxY);
+	INTERNAL_DECORATOR(RHISetScissorRect)(bEnable, MinX, MinY, MaxX, MaxY);
 }
 
 void FRHICommandSetRenderTargets::Execute(FRHICommandListBase& CmdList)
 {
 	RHISTAT(SetRenderTargets);
-	INTERNAL_DECORATOR(SetRenderTargets)(
+	INTERNAL_DECORATOR(RHISetRenderTargets)(
 		NewNumSimultaneousRenderTargets,
 		NewRenderTargetsRHI,
 		&NewDepthStencilTarget,
@@ -196,13 +196,13 @@ void FRHICommandSetRenderTargets::Execute(FRHICommandListBase& CmdList)
 void FRHICommandSetRenderTargetsAndClear::Execute(FRHICommandListBase& CmdList)
 {
 	RHISTAT(SetRenderTargetsAndClear);
-	INTERNAL_DECORATOR(SetRenderTargetsAndClear)(RenderTargetsInfo);
+	INTERNAL_DECORATOR(RHISetRenderTargetsAndClear)(RenderTargetsInfo);
 }
 
 void FRHICommandBindClearMRTValues::Execute(FRHICommandListBase& CmdList)
 {
 	RHISTAT(BindClearMRTValues);
-	INTERNAL_DECORATOR(BindClearMRTValues)(
+	INTERNAL_DECORATOR(RHIBindClearMRTValues)(
 		bClearColor,
 		bClearDepth,
 		bClearStencil		
@@ -213,9 +213,9 @@ void FRHICommandEndDrawPrimitiveUP::Execute(FRHICommandListBase& CmdList)
 {
 	RHISTAT(EndDrawPrimitiveUP);
 	void* Buffer = NULL;
-	INTERNAL_DECORATOR(BeginDrawPrimitiveUP)(PrimitiveType, NumPrimitives, NumVertices, VertexDataStride, Buffer);
+	INTERNAL_DECORATOR(RHIBeginDrawPrimitiveUP)(PrimitiveType, NumPrimitives, NumVertices, VertexDataStride, Buffer);
 	FMemory::Memcpy(Buffer, OutVertexData, NumVertices * VertexDataStride);
-	INTERNAL_DECORATOR(EndDrawPrimitiveUP)();
+	INTERNAL_DECORATOR(RHIEndDrawPrimitiveUP)();
 }
 
 void FRHICommandEndDrawIndexedPrimitiveUP::Execute(FRHICommandListBase& CmdList)
@@ -223,7 +223,7 @@ void FRHICommandEndDrawIndexedPrimitiveUP::Execute(FRHICommandListBase& CmdList)
 	RHISTAT(EndDrawIndexedPrimitiveUP);
 	void* VertexBuffer = nullptr;
 	void* IndexBuffer = nullptr;
-	INTERNAL_DECORATOR(BeginDrawIndexedPrimitiveUP)(
+	INTERNAL_DECORATOR(RHIBeginDrawIndexedPrimitiveUP)(
 		PrimitiveType,
 		NumPrimitives,
 		NumVertices,
@@ -235,14 +235,14 @@ void FRHICommandEndDrawIndexedPrimitiveUP::Execute(FRHICommandListBase& CmdList)
 		IndexBuffer);
 	FMemory::Memcpy(VertexBuffer, OutVertexData, NumVertices * VertexDataStride);
 	FMemory::Memcpy(IndexBuffer, OutIndexData, NumIndices * IndexDataStride);
-	INTERNAL_DECORATOR(EndDrawIndexedPrimitiveUP)();
+	INTERNAL_DECORATOR(RHIEndDrawIndexedPrimitiveUP)();
 }
 
 template<ECmdList CmdListType>
 void FRHICommandSetComputeShader<CmdListType>::Execute(FRHICommandListBase& CmdList)
 {
 	RHISTAT(SetComputeShader);
-	INTERNAL_DECORATOR_CONTEXT(SetComputeShader)(ComputeShader);
+	INTERNAL_DECORATOR_CONTEXT(RHISetComputeShader)(ComputeShader);
 }
 template struct FRHICommandSetComputeShader<ECmdList::EGfx>;
 template struct FRHICommandSetComputeShader<ECmdList::ECompute>;
@@ -251,7 +251,7 @@ template<ECmdList CmdListType>
 void FRHICommandDispatchComputeShader<CmdListType>::Execute(FRHICommandListBase& CmdList)
 {
 	RHISTAT(DispatchComputeShader);
-	INTERNAL_DECORATOR_CONTEXT(DispatchComputeShader)(ThreadGroupCountX, ThreadGroupCountY, ThreadGroupCountZ);
+	INTERNAL_DECORATOR_CONTEXT(RHIDispatchComputeShader)(ThreadGroupCountX, ThreadGroupCountY, ThreadGroupCountZ);
 }
 template struct FRHICommandDispatchComputeShader<ECmdList::EGfx>;
 template struct FRHICommandDispatchComputeShader<ECmdList::ECompute>;
@@ -260,7 +260,7 @@ template<ECmdList CmdListType>
 void FRHICommandDispatchIndirectComputeShader<CmdListType>::Execute(FRHICommandListBase& CmdList)
 {
 	RHISTAT(DispatchIndirectComputeShader);
-	INTERNAL_DECORATOR_CONTEXT(DispatchIndirectComputeShader)(ArgumentBuffer, ArgumentOffset);
+	INTERNAL_DECORATOR_CONTEXT(RHIDispatchIndirectComputeShader)(ArgumentBuffer, ArgumentOffset);
 }
 template struct FRHICommandDispatchIndirectComputeShader<ECmdList::EGfx>;
 template struct FRHICommandDispatchIndirectComputeShader<ECmdList::ECompute>;
@@ -268,68 +268,68 @@ template struct FRHICommandDispatchIndirectComputeShader<ECmdList::ECompute>;
 void FRHICommandAutomaticCacheFlushAfterComputeShader::Execute(FRHICommandListBase& CmdList)
 {
 	RHISTAT(AutomaticCacheFlushAfterComputeShader);
-	INTERNAL_DECORATOR(AutomaticCacheFlushAfterComputeShader)(bEnable);
+	INTERNAL_DECORATOR(RHIAutomaticCacheFlushAfterComputeShader)(bEnable);
 }
 
 void FRHICommandFlushComputeShaderCache::Execute(FRHICommandListBase& CmdList)
 {
 	RHISTAT(FlushComputeShaderCache);
-	INTERNAL_DECORATOR(FlushComputeShaderCache)();
+	INTERNAL_DECORATOR(RHIFlushComputeShaderCache)();
 }
 
 void FRHICommandDrawPrimitiveIndirect::Execute(FRHICommandListBase& CmdList)
 {
 	RHISTAT(DrawPrimitiveIndirect);
-	INTERNAL_DECORATOR(DrawPrimitiveIndirect)(PrimitiveType, ArgumentBuffer, ArgumentOffset);
+	INTERNAL_DECORATOR(RHIDrawPrimitiveIndirect)(PrimitiveType, ArgumentBuffer, ArgumentOffset);
 }
 
 void FRHICommandDrawIndexedIndirect::Execute(FRHICommandListBase& CmdList)
 {
 	RHISTAT(DrawIndexedIndirect);
-	INTERNAL_DECORATOR(DrawIndexedIndirect)(IndexBufferRHI, PrimitiveType, ArgumentsBufferRHI, DrawArgumentsIndex, NumInstances);
+	INTERNAL_DECORATOR(RHIDrawIndexedIndirect)(IndexBufferRHI, PrimitiveType, ArgumentsBufferRHI, DrawArgumentsIndex, NumInstances);
 }
 
 void FRHICommandDrawIndexedPrimitiveIndirect::Execute(FRHICommandListBase& CmdList)
 {
 	RHISTAT(DrawIndexedPrimitiveIndirect);
-	INTERNAL_DECORATOR(DrawIndexedPrimitiveIndirect)(PrimitiveType, IndexBuffer, ArgumentsBuffer, ArgumentOffset);
+	INTERNAL_DECORATOR(RHIDrawIndexedPrimitiveIndirect)(PrimitiveType, IndexBuffer, ArgumentsBuffer, ArgumentOffset);
 }
 
 void FRHICommandEnableDepthBoundsTest::Execute(FRHICommandListBase& CmdList)
 {
 	RHISTAT(EnableDepthBoundsTest);
-	INTERNAL_DECORATOR(EnableDepthBoundsTest)(bEnable, MinDepth, MaxDepth);
+	INTERNAL_DECORATOR(RHIEnableDepthBoundsTest)(bEnable, MinDepth, MaxDepth);
 }
 
 void FRHICommandClearUAV::Execute(FRHICommandListBase& CmdList)
 {
 	RHISTAT(ClearUAV);
-	INTERNAL_DECORATOR(ClearUAV)(UnorderedAccessViewRHI, Values);
+	INTERNAL_DECORATOR(RHIClearUAV)(UnorderedAccessViewRHI, Values);
 }
 
 void FRHICommandCopyToResolveTarget::Execute(FRHICommandListBase& CmdList)
 {
 	RHISTAT(CopyToResolveTarget);
-	INTERNAL_DECORATOR(CopyToResolveTarget)(SourceTexture, DestTexture, bKeepOriginalSurface, ResolveParams);
+	INTERNAL_DECORATOR(RHICopyToResolveTarget)(SourceTexture, DestTexture, bKeepOriginalSurface, ResolveParams);
 }
 
 void FRHICommandTransitionTextures::Execute(FRHICommandListBase& CmdList)
 {
 	RHISTAT(TransitionTextures);
-	INTERNAL_DECORATOR(TransitionResources)(TransitionType, &Textures[0], NumTextures);
+	INTERNAL_DECORATOR(RHITransitionResources)(TransitionType, &Textures[0], NumTextures);
 }
 
 void FRHICommandTransitionTexturesArray::Execute(FRHICommandListBase& CmdList)
 {
 	RHISTAT(TransitionTextures);
-	INTERNAL_DECORATOR(TransitionResources)(TransitionType, &Textures[0], Textures.Num());
+	INTERNAL_DECORATOR(RHITransitionResources)(TransitionType, &Textures[0], Textures.Num());
 }
 
 template<ECmdList CmdListType>
 void FRHICommandTransitionUAVs<CmdListType>::Execute(FRHICommandListBase& CmdList)
 {
 	RHISTAT(TransitionUAVs);
-	INTERNAL_DECORATOR_CONTEXT(TransitionResources)(TransitionType, TransitionPipeline, UAVs, NumUAVs, WriteFence);
+	INTERNAL_DECORATOR_CONTEXT(RHITransitionResources)(TransitionType, TransitionPipeline, UAVs, NumUAVs, WriteFence);
 }
 template struct FRHICommandTransitionUAVs<ECmdList::EGfx>;
 template struct FRHICommandTransitionUAVs<ECmdList::ECompute>;
@@ -338,7 +338,7 @@ template<ECmdList CmdListType>
 void FRHICommandSetAsyncComputeBudget<CmdListType>::Execute(FRHICommandListBase& CmdList)
 {
 	RHISTAT(SetAsyncComputeBudget);
-	INTERNAL_DECORATOR_CONTEXT(SetAsyncComputeBudget)(Budget);
+	INTERNAL_DECORATOR_CONTEXT(RHISetAsyncComputeBudget)(Budget);
 }
 template struct FRHICommandSetAsyncComputeBudget<ECmdList::EGfx>;
 template struct FRHICommandSetAsyncComputeBudget<ECmdList::ECompute>;
@@ -347,7 +347,7 @@ template<ECmdList CmdListType>
 void FRHICommandWaitComputeFence<CmdListType>::Execute(FRHICommandListBase& CmdList)
 {
 	RHISTAT(WaitComputeFence);
-	INTERNAL_DECORATOR_CONTEXT(WaitComputeFence)(WaitFence);
+	INTERNAL_DECORATOR_CONTEXT(RHIWaitComputeFence)(WaitFence);
 }
 template struct FRHICommandWaitComputeFence<ECmdList::EGfx>;
 template struct FRHICommandWaitComputeFence<ECmdList::ECompute>;
@@ -355,13 +355,13 @@ template struct FRHICommandWaitComputeFence<ECmdList::ECompute>;
 void FRHICommandClear::Execute(FRHICommandListBase& CmdList)
 {
 	RHISTAT(Clear);
-	INTERNAL_DECORATOR(Clear)(bClearColor, Color, bClearDepth, Depth, bClearStencil, Stencil, ExcludeRect);
+	INTERNAL_DECORATOR(RHIClear)(bClearColor, Color, bClearDepth, Depth, bClearStencil, Stencil, ExcludeRect);
 }
 
 void FRHICommandClearMRT::Execute(FRHICommandListBase& CmdList)
 {
 	RHISTAT(ClearMRT);
-	INTERNAL_DECORATOR(ClearMRT)(bClearColor, NumClearColors, ColorArray, bClearDepth, Depth, bClearStencil, Stencil, ExcludeRect);
+	INTERNAL_DECORATOR(RHIClearMRT)(bClearColor, NumClearColors, ColorArray, bClearDepth, Depth, bClearStencil, Stencil, ExcludeRect);
 }
 
 void FRHICommandBuildLocalBoundShaderState::Execute(FRHICommandListBase& CmdList)
@@ -380,7 +380,7 @@ void FRHICommandSetLocalBoundShaderState::Execute(FRHICommandListBase& CmdList)
 	RHISTAT(SetLocalBoundShaderState);
 	check(LocalBoundShaderState.WorkArea->ComputedBSS->UseCount > 0 && IsValidRef(LocalBoundShaderState.WorkArea->ComputedBSS->BSS)); // this should have been created and should have uses outstanding
 
-	INTERNAL_DECORATOR(SetBoundShaderState)(LocalBoundShaderState.WorkArea->ComputedBSS->BSS);
+	INTERNAL_DECORATOR(RHISetBoundShaderState)(LocalBoundShaderState.WorkArea->ComputedBSS->BSS);
 
 	if (--LocalBoundShaderState.WorkArea->ComputedBSS->UseCount == 0)
 	{
@@ -407,7 +407,7 @@ void FRHICommandSetLocalUniformBuffer<TShaderRHIParamRef>::Execute(FRHICommandLi
 {
 	RHISTAT(SetLocalUniformBuffer);
 	check(LocalUniformBuffer.WorkArea->ComputedUniformBuffer->UseCount > 0 && IsValidRef(LocalUniformBuffer.WorkArea->ComputedUniformBuffer->UniformBuffer)); // this should have been created and should have uses outstanding
-	INTERNAL_DECORATOR(SetShaderUniformBuffer)(Shader, BaseIndex, LocalUniformBuffer.WorkArea->ComputedUniformBuffer->UniformBuffer);
+	INTERNAL_DECORATOR(RHISetShaderUniformBuffer)(Shader, BaseIndex, LocalUniformBuffer.WorkArea->ComputedUniformBuffer->UniformBuffer);
 	if (--LocalUniformBuffer.WorkArea->ComputedUniformBuffer->UseCount == 0)
 	{
 		LocalUniformBuffer.WorkArea->ComputedUniformBuffer->~FComputedUniformBuffer();
@@ -423,32 +423,32 @@ template struct FRHICommandSetLocalUniformBuffer<FComputeShaderRHIParamRef>;
 void FRHICommandBeginRenderQuery::Execute(FRHICommandListBase& CmdList)
 {
 	RHISTAT(BeginRenderQuery);
-	INTERNAL_DECORATOR(BeginRenderQuery)(RenderQuery);
+	INTERNAL_DECORATOR(RHIBeginRenderQuery)(RenderQuery);
 }
 
 void FRHICommandEndRenderQuery::Execute(FRHICommandListBase& CmdList)
 {
 	RHISTAT(EndRenderQuery);
-	INTERNAL_DECORATOR(EndRenderQuery)(RenderQuery);
+	INTERNAL_DECORATOR(RHIEndRenderQuery)(RenderQuery);
 }
 
 void FRHICommandBeginOcclusionQueryBatch::Execute(FRHICommandListBase& CmdList)
 {
 	RHISTAT(BeginOcclusionQueryBatch);
-	INTERNAL_DECORATOR(BeginOcclusionQueryBatch)();
+	INTERNAL_DECORATOR(RHIBeginOcclusionQueryBatch)();
 }
 
 void FRHICommandEndOcclusionQueryBatch::Execute(FRHICommandListBase& CmdList)
 {
 	RHISTAT(EndOcclusionQueryBatch);
-	INTERNAL_DECORATOR(EndOcclusionQueryBatch)();
+	INTERNAL_DECORATOR(RHIEndOcclusionQueryBatch)();
 }
 
 template<ECmdList CmdListType>
 void FRHICommandSubmitCommandsHint<CmdListType>::Execute(FRHICommandListBase& CmdList)
 {
 	RHISTAT(SubmitCommandsHint);
-	INTERNAL_DECORATOR_CONTEXT(SubmitCommandsHint)();
+	INTERNAL_DECORATOR_CONTEXT(RHISubmitCommandsHint)();
 }
 template struct FRHICommandSubmitCommandsHint<ECmdList::EGfx>;
 template struct FRHICommandSubmitCommandsHint<ECmdList::ECompute>;
@@ -456,50 +456,50 @@ template struct FRHICommandSubmitCommandsHint<ECmdList::ECompute>;
 void FRHICommandUpdateTextureReference::Execute(FRHICommandListBase& CmdList)
 {
 	RHISTAT(UpdateTextureReference);
-	INTERNAL_DECORATOR(UpdateTextureReference)(TextureRef, NewTexture);
+	INTERNAL_DECORATOR(RHIUpdateTextureReference)(TextureRef, NewTexture);
 }
 
 void FRHICommandBeginScene::Execute(FRHICommandListBase& CmdList)
 {
 	RHISTAT(BeginScene);
-	INTERNAL_DECORATOR(BeginScene)();
+	INTERNAL_DECORATOR(RHIBeginScene)();
 }
 
 void FRHICommandEndScene::Execute(FRHICommandListBase& CmdList)
 {
 	RHISTAT(EndScene);
-	INTERNAL_DECORATOR(EndScene)();
+	INTERNAL_DECORATOR(RHIEndScene)();
 }
 
 void FRHICommandBeginFrame::Execute(FRHICommandListBase& CmdList)
 {
 	RHISTAT(BeginFrame);
-	INTERNAL_DECORATOR(BeginFrame)();
+	INTERNAL_DECORATOR(RHIBeginFrame)();
 }
 
 void FRHICommandEndFrame::Execute(FRHICommandListBase& CmdList)
 {
 	RHISTAT(EndFrame);
-	INTERNAL_DECORATOR(EndFrame)();
+	INTERNAL_DECORATOR(RHIEndFrame)();
 }
 
 void FRHICommandBeginDrawingViewport::Execute(FRHICommandListBase& CmdList)
 {
 	RHISTAT(BeginDrawingViewport);
-	INTERNAL_DECORATOR(BeginDrawingViewport)(Viewport, RenderTargetRHI);
+	INTERNAL_DECORATOR(RHIBeginDrawingViewport)(Viewport, RenderTargetRHI);
 }
 
 void FRHICommandEndDrawingViewport::Execute(FRHICommandListBase& CmdList)
 {
 	RHISTAT(EndDrawingViewport);
-	INTERNAL_DECORATOR(EndDrawingViewport)(Viewport, bPresent, bLockToVsync);
+	INTERNAL_DECORATOR(RHIEndDrawingViewport)(Viewport, bPresent, bLockToVsync);
 }
 
 template<ECmdList CmdListType>
 void FRHICommandPushEvent<CmdListType>::Execute(FRHICommandListBase& CmdList)
 {
 	RHISTAT(PushEvent);
-	INTERNAL_DECORATOR_CONTEXT(PushEvent)(Name, Color);
+	INTERNAL_DECORATOR_CONTEXT(RHIPushEvent)(Name, Color);
 }
 template struct FRHICommandPushEvent<ECmdList::EGfx>;
 template struct FRHICommandPushEvent<ECmdList::ECompute>;
@@ -508,7 +508,7 @@ template<ECmdList CmdListType>
 void FRHICommandPopEvent<CmdListType>::Execute(FRHICommandListBase& CmdList)
 {
 	RHISTAT(PopEvent);
-	INTERNAL_DECORATOR_CONTEXT(PopEvent)();
+	INTERNAL_DECORATOR_CONTEXT(RHIPopEvent)();
 }
 template struct FRHICommandPopEvent<ECmdList::EGfx>;
 template struct FRHICommandPopEvent<ECmdList::ECompute>;
diff --git a/Engine/Source/Runtime/RenderCore/Private/RenderUtils.cpp b/Engine/Source/Runtime/RenderCore/Private/RenderUtils.cpp
--- a/Engine/Source/Runtime/RenderCore/Private/RenderUtils.cpp
+++ b/Engine/Source/Runtime/RenderCore/Private/RenderUtils.cpp
@@ -37,12 +37,21 @@ FArchive& operator<<(FArchive& Ar,FPackedNormal& N)
 	return Ar;
 }
 
-FArchive& operator<<(FArchive& Ar, FPackedRGB10A2& N)
+FArchive& operator<<(FArchive& Ar, FPackedRGB10A2N& N)
 {
 	Ar << N.Vector.Packed;
 	return Ar;
 }
 
+FArchive& operator<<(FArchive& Ar, FPackedRGBA16N& N)
+{
+	Ar << N.X;
+	Ar << N.Y;
+	Ar << N.Z;
+	Ar << N.W;
+	return Ar;
+}
+
 //
 //	Pixel format information.
 //
@@ -798,10 +807,17 @@ RENDERCORE_API FVertexDeclarationRHIRef& GetVertexDeclarationFVector3()
 	return GVector3VertexDeclaration.VertexDeclarationRHI;
 }
 
-RENDERCORE_API bool IsSimpleDynamicLightingEnabled()
+RENDERCORE_API bool PlatformSupportsSimpleForwardShading(EShaderPlatform Platform)
+{
+	static const auto SupportSimpleForwardShadingCVar = IConsoleManager::Get().FindTConsoleVariableDataInt(TEXT("r.SupportSimpleForwardShading"));
+	// Scalability feature only needed / used on PC
+	return IsPCPlatform(Platform) && SupportSimpleForwardShadingCVar->GetValueOnAnyThread() != 0;
+}
+
+RENDERCORE_API bool IsSimpleForwardShadingEnabled(EShaderPlatform Platform)
 {
-	static const auto CVar = IConsoleManager::Get().FindTConsoleVariableDataInt(TEXT("r.SimpleDynamicLighting"));
-	return (CVar->GetValueOnAnyThread() != 0);
+	static const auto CVar = IConsoleManager::Get().FindTConsoleVariableDataInt(TEXT("r.SimpleForwardShading"));
+	return CVar->GetValueOnAnyThread() != 0 && PlatformSupportsSimpleForwardShading(Platform);
 }
 
 class FUnitCubeVertexBuffer : public FVertexBuffer
diff --git a/Engine/Source/Runtime/RenderCore/Public/PackedNormal.h b/Engine/Source/Runtime/RenderCore/Public/PackedNormal.h
--- a/Engine/Source/Runtime/RenderCore/Public/PackedNormal.h
+++ b/Engine/Source/Runtime/RenderCore/Public/PackedNormal.h
@@ -37,6 +37,7 @@ struct FPackedNormal
 	void operator=(const FVector& InVector);
 	void operator=(const FVector4& InVector);
 	operator FVector() const;
+	operator FVector4() const;
 	VectorRegister GetVectorRegister() const;
 
 	// Set functions.
@@ -81,11 +82,8 @@ FORCEINLINE void FPackedNormal::operator=(const FVector4& InVector)
 
 FORCEINLINE bool FPackedNormal::operator==(const FPackedNormal& B) const
 {
-	if(Vector.Packed != B.Vector.Packed)
-		return 0;
-
-	FVector	V1 = *this,
-			V2 = B;
+	FVector	V1 = *this;
+	FVector	V2 = B;
 
 	if(FMath::Abs(V1.X - V2.X) > THRESH_NORMALS_ARE_SAME * 4.0f)
 		return 0;
@@ -101,22 +99,7 @@ FORCEINLINE bool FPackedNormal::operator==(const FPackedNormal& B) const
 
 FORCEINLINE bool FPackedNormal::operator!=(const FPackedNormal& B) const
 {
-	if(Vector.Packed == B.Vector.Packed)
-		return 0;
-
-	FVector	V1 = *this,
-			V2 = B;
-
-	if(FMath::Abs(V1.X - V2.X) > THRESH_NORMALS_ARE_SAME * 4.0f)
-		return 1;
-
-	if(FMath::Abs(V1.Y - V2.Y) > THRESH_NORMALS_ARE_SAME * 4.0f)
-		return 1;
-
-	if(FMath::Abs(V1.Z - V2.Z) > THRESH_NORMALS_ARE_SAME * 4.0f)
-		return 1;
-
-	return 0;
+	return !(*this == B);
 }
 
 FORCEINLINE FPackedNormal::operator FVector() const
@@ -128,6 +111,15 @@ FORCEINLINE FPackedNormal::operator FVector() const
 	return UnpackedVector;
 }
 
+FORCEINLINE FPackedNormal::operator FVector4() const
+{
+	VectorRegister VectorToUnpack = GetVectorRegister();
+	// Write to FVector4 and return it.
+	FVector4 UnpackedVector;
+	VectorStore(VectorToUnpack, &UnpackedVector);
+	return UnpackedVector;
+}
+
 FORCEINLINE VectorRegister FPackedNormal::GetVectorRegister() const
 {
 	// Rescale [0..255] range to [-1..1]
@@ -139,7 +131,7 @@ FORCEINLINE VectorRegister FPackedNormal::GetVectorRegister() const
 }
 
 /** A vector, quantized and packed into 32-bits. */
-struct FPackedRGB10A2
+struct FPackedRGB10A2N
 {
 	union
 	{
@@ -157,16 +149,20 @@ struct FPackedRGB10A2
 			uint32 X : 10;
 #endif
 		};
-		uint32		Packed;
-	}				Vector;
+
+		struct  
+		{
+			uint32 Packed;
+		};
+	} Vector;
 
 	// Constructors.
 
-	FPackedRGB10A2() { Vector.Packed = 0; }
-	FPackedRGB10A2(uint32 InPacked) { Vector.Packed = InPacked; }
-	FPackedRGB10A2(const FVector& InVector) { *this = InVector; }
-	FPackedRGB10A2(const FVector4& InVector) { *this = InVector; }
-	FPackedRGB10A2(uint32 InX, uint32 InY, uint32 InZ, uint32 InW) { Vector.X = InX; Vector.Y = InY; Vector.Z = InZ; Vector.W = InW; }
+	FPackedRGB10A2N() { Vector.Packed = 0; }
+	FPackedRGB10A2N(uint32 InPacked) { Vector.Packed = InPacked; }
+	FPackedRGB10A2N(const FVector& InVector) { *this = InVector; }
+	FPackedRGB10A2N(const FVector4& InVector) { *this = InVector; }
+	FPackedRGB10A2N(uint32 InX, uint32 InY, uint32 InZ, uint32 InW) { Vector.X = InX; Vector.Y = InY; Vector.Z = InZ; Vector.W = InW; }
 
 	// Conversion operators.
 
@@ -183,22 +179,22 @@ struct FPackedRGB10A2
 
 	// Equality operator.
 
-	bool operator==(const FPackedRGB10A2& B) const;
-	bool operator!=(const FPackedRGB10A2& B) const;
+	bool operator==(const FPackedRGB10A2N& B) const;
+	bool operator!=(const FPackedRGB10A2N& B) const;
 
 	// Serializer.
 
-	friend RENDERCORE_API FArchive& operator<<(FArchive& Ar, FPackedRGB10A2& N);
+	friend RENDERCORE_API FArchive& operator<<(FArchive& Ar, FPackedRGB10A2N& N);
 
 	FString ToString() const
 	{
 		return FString::Printf(TEXT("X=%d Y=%d Z=%d W=%d"), Vector.X, Vector.Y, Vector.Z, Vector.W);
 	}
 
-	static RENDERCORE_API FPackedRGB10A2 ZeroVector;
+	static RENDERCORE_API FPackedRGB10A2N ZeroVector;
 };
 
-FORCEINLINE void FPackedRGB10A2::operator=(const FVector& InVector)
+FORCEINLINE void FPackedRGB10A2N::operator=(const FVector& InVector)
 {
 	Vector.X = FMath::Clamp(FMath::TruncToInt(InVector.X * 511.5f + 511.5f), 0, 1023);
 	Vector.Y = FMath::Clamp(FMath::TruncToInt(InVector.Y * 511.5f + 511.5f), 0, 1023);
@@ -206,7 +202,7 @@ FORCEINLINE void FPackedRGB10A2::operator=(const FVector& InVector)
 	Vector.W = 3;
 }
 
-FORCEINLINE void FPackedRGB10A2::operator=(const FVector4& InVector)
+FORCEINLINE void FPackedRGB10A2N::operator=(const FVector4& InVector)
 {
 	Vector.X = FMath::Clamp(FMath::TruncToInt(InVector.X * 511.5f + 511.5f), 0, 1023);
 	Vector.Y = FMath::Clamp(FMath::TruncToInt(InVector.Y * 511.5f + 511.5f), 0, 1023);
@@ -214,13 +210,10 @@ FORCEINLINE void FPackedRGB10A2::operator=(const FVector4& InVector)
 	Vector.W = FMath::Clamp(FMath::TruncToInt(InVector.W * 1.5f   + 1.5f),   0, 3);
 }
 
-FORCEINLINE bool FPackedRGB10A2::operator==(const FPackedRGB10A2& B) const
+FORCEINLINE bool FPackedRGB10A2N::operator==(const FPackedRGB10A2N& B) const
 {
-	if (Vector.Packed != B.Vector.Packed)
-		return 0;
-
-	FVector	V1 = *this,
-			V2 = B;
+	FVector	V1 = *this;
+	FVector V2 = B;
 
 	if (FMath::Abs(V1.X - V2.X) > THRESH_NORMALS_ARE_SAME * 4.0f)
 		return 0;
@@ -234,27 +227,125 @@ FORCEINLINE bool FPackedRGB10A2::operator==(const FPackedRGB10A2& B) const
 	return 1;
 }
 
-FORCEINLINE bool FPackedRGB10A2::operator!=(const FPackedRGB10A2& B) const
+FORCEINLINE bool FPackedRGB10A2N::operator!=(const FPackedRGB10A2N& B) const
 {
-	if (Vector.Packed == B.Vector.Packed)
-		return 0;
+	return !(*this == B);
+}
+
+FORCEINLINE FPackedRGB10A2N::operator FVector() const
+{
+	VectorRegister VectorToUnpack = GetVectorRegister();
+	// Write to FVector and return it.
+	FVector UnpackedVector;
+	VectorStoreFloat3(VectorToUnpack, &UnpackedVector);
+	return UnpackedVector;
+}
+
+FORCEINLINE FPackedRGB10A2N::operator FVector4() const
+{
+	VectorRegister VectorToUnpack = GetVectorRegister();
+	// Write to FVector4 and return it.
+	FVector4 UnpackedVector;
+	VectorStore(VectorToUnpack, &UnpackedVector);
+	return UnpackedVector;
+}
+
+FORCEINLINE VectorRegister FPackedRGB10A2N::GetVectorRegister() const
+{
+	VectorRegister VectorToUnpack = VectorLoadURGB10A2N((void*)this);
+	VectorToUnpack = VectorMultiplyAdd(VectorToUnpack, MakeVectorRegister(2.0f, 2.0f, 2.0f, 2.0f), MakeVectorRegister(-1.0f, -1.0f, -1.0f, -1.0f));
+	VectorResetFloatRegisters();
+	// Return unpacked vector register.
+	return VectorToUnpack;
+}
+
+/** A vector, quantized and packed into 32-bits. */
+struct FPackedRGBA16N
+{
+	struct
+	{
+		uint16 X;
+		uint16 Y;
+		uint16 Z;
+		uint16 W;
+	};
+
+	// Constructors.
+
+	FPackedRGBA16N() { X = 0; Y = 0; Z = 0; W = 0; }
+	FPackedRGBA16N(const FVector& InVector) { *this = InVector; }
+	FPackedRGBA16N(const FVector4& InVector) { *this = InVector; }
+	FPackedRGBA16N(uint16 InX, uint16 InY, uint16 InZ, uint16 InW) { X = InX; Y = InY; Z = InZ; W = InW; }
+
+	// Conversion operators.
+
+	void operator=(const FVector& InVector);
+	void operator=(const FVector4& InVector);
+	operator FVector() const;
+	operator FVector4() const;
+
+	VectorRegister GetVectorRegister() const;
+
+	// Set functions.
+	void Set(const FVector& InVector) { *this = InVector; }
+	void Set(const FVector4& InVector) { *this = InVector; }
+
+	// Equality operator.
 
-	FVector	V1 = *this,
-			V2 = B;
+	bool operator==(const FPackedRGBA16N& B) const;
+	bool operator!=(const FPackedRGBA16N& B) const;
+
+	// Serializer.
+
+	friend RENDERCORE_API FArchive& operator<<(FArchive& Ar, FPackedRGBA16N& N);
+
+	FString ToString() const
+	{
+		return FString::Printf(TEXT("X=%d Y=%d Z=%d W=%d"), X, Y, Z, W);
+	}
+
+	static RENDERCORE_API FPackedRGBA16N ZeroVector;
+};
+
+FORCEINLINE void FPackedRGBA16N::operator=(const FVector& InVector)
+{
+	X = FMath::Clamp(FMath::TruncToInt(InVector.X * 32767.5f + 32767.5f), 0, 65535);
+	Y = FMath::Clamp(FMath::TruncToInt(InVector.Y * 32767.5f + 32767.5f), 0, 65535);
+	Z = FMath::Clamp(FMath::TruncToInt(InVector.Z * 32767.5f + 32767.5f), 0, 65535);
+	W = 65535;
+}
+
+FORCEINLINE void FPackedRGBA16N::operator=(const FVector4& InVector)
+{
+	X = FMath::Clamp(FMath::TruncToInt(InVector.X * 32767.5f + 32767.5f), 0, 65535);
+	Y = FMath::Clamp(FMath::TruncToInt(InVector.Y * 32767.5f + 32767.5f), 0, 65535);
+	Z = FMath::Clamp(FMath::TruncToInt(InVector.Z * 32767.5f + 32767.5f), 0, 65535);
+	W = FMath::Clamp(FMath::TruncToInt(InVector.W * 32767.5f + 32767.5f), 0, 65535);
+}
+
+FORCEINLINE bool FPackedRGBA16N::operator==(const FPackedRGBA16N& B) const
+{
+	FVector	V1 = *this;
+	FVector	V2 = B;
 
 	if (FMath::Abs(V1.X - V2.X) > THRESH_NORMALS_ARE_SAME * 4.0f)
-		return 1;
+		return 0;
 
 	if (FMath::Abs(V1.Y - V2.Y) > THRESH_NORMALS_ARE_SAME * 4.0f)
-		return 1;
+		return 0;
 
 	if (FMath::Abs(V1.Z - V2.Z) > THRESH_NORMALS_ARE_SAME * 4.0f)
-		return 1;
+		return 0;
+
+	return 1;
+}
 
-	return 0;
+FORCEINLINE bool FPackedRGBA16N::operator!=(const FPackedRGBA16N& B) const
+{
+	return !(*this == B);
 }
 
-FORCEINLINE FPackedRGB10A2::operator FVector() const
+FORCEINLINE FPackedRGBA16N::operator FVector() const
 {
 	VectorRegister VectorToUnpack = GetVectorRegister();
 	// Write to FVector and return it.
@@ -263,7 +354,7 @@ FORCEINLINE FPackedRGB10A2::operator FVector() const
 	return UnpackedVector;
 }
 
-FORCEINLINE FPackedRGB10A2::operator FVector4() const
+FORCEINLINE FPackedRGBA16N::operator FVector4() const
 {
 	VectorRegister VectorToUnpack = GetVectorRegister();
 	// Write to FVector4 and return it.
@@ -272,9 +363,9 @@ FORCEINLINE FPackedRGB10A2::operator FVector4() const
 	return UnpackedVector;
 }
 
-FORCEINLINE VectorRegister FPackedRGB10A2::GetVectorRegister() const
+FORCEINLINE VectorRegister FPackedRGBA16N::GetVectorRegister() const
 {
-	VectorRegister VectorToUnpack = VectorLoadRGB10A2((void*)this);
+	VectorRegister VectorToUnpack = VectorLoadURGBA16N((void*)this);
 	VectorToUnpack = VectorMultiplyAdd(VectorToUnpack, MakeVectorRegister(2.0f, 2.0f, 2.0f, 2.0f), MakeVectorRegister(-1.0f, -1.0f, -1.0f, -1.0f));
 	VectorResetFloatRegisters();
 	// Return unpacked vector register.
diff --git a/Engine/Source/Runtime/RenderCore/Public/RenderUtils.h b/Engine/Source/Runtime/RenderCore/Public/RenderUtils.h
--- a/Engine/Source/Runtime/RenderCore/Public/RenderUtils.h
+++ b/Engine/Source/Runtime/RenderCore/Public/RenderUtils.h
@@ -309,7 +309,9 @@ RENDERCORE_API FVertexDeclarationRHIRef& GetVertexDeclarationFVector4();
 
 RENDERCORE_API FVertexDeclarationRHIRef& GetVertexDeclarationFVector3();
 
-RENDERCORE_API bool IsSimpleDynamicLightingEnabled();
+RENDERCORE_API bool PlatformSupportsSimpleForwardShading(EShaderPlatform Platform);
+
+RENDERCORE_API bool IsSimpleForwardShadingEnabled(EShaderPlatform Platform);
 
 /** Unit cube vertex buffer (VertexDeclarationFVector4) */
 RENDERCORE_API FVertexBufferRHIRef& GetUnitCubeVertexBuffer();
diff --git a/Engine/Source/Runtime/Renderer/Private/BasePassRendering.cpp b/Engine/Source/Runtime/Renderer/Private/BasePassRendering.cpp
--- a/Engine/Source/Runtime/Renderer/Private/BasePassRendering.cpp
+++ b/Engine/Source/Runtime/Renderer/Private/BasePassRendering.cpp
@@ -65,7 +65,11 @@ IMPLEMENT_BASEPASS_LIGHTMAPPED_SHADER_TYPE( FSelfShadowedCachedPointIndirectLigh
 IMPLEMENT_BASEPASS_LIGHTMAPPED_SHADER_TYPE( TUniformLightMapPolicy<LMP_NO_LIGHTMAP>, FNoLightMapPolicy );
 IMPLEMENT_BASEPASS_LIGHTMAPPED_SHADER_TYPE( TUniformLightMapPolicy<LMP_CACHED_VOLUME_INDIRECT_LIGHTING>, FCachedVolumeIndirectLightingPolicy );
 IMPLEMENT_BASEPASS_LIGHTMAPPED_SHADER_TYPE( TUniformLightMapPolicy<LMP_CACHED_POINT_INDIRECT_LIGHTING>, FCachedPointIndirectLightingPolicy );
-IMPLEMENT_BASEPASS_LIGHTMAPPED_SHADER_TYPE( TUniformLightMapPolicy<LMP_SIMPLE_DYNAMIC_LIGHTING>, FSimpleDynamicLightingPolicy );
+IMPLEMENT_BASEPASS_LIGHTMAPPED_SHADER_TYPE( TUniformLightMapPolicy<LMP_SIMPLE_NO_LIGHTMAP>, FSimpleNoLightmapLightingPolicy );
+IMPLEMENT_BASEPASS_LIGHTMAPPED_SHADER_TYPE( TUniformLightMapPolicy<LMP_SIMPLE_LIGHTMAP_ONLY_LIGHTING>, FSimpleLightmapOnlyLightingPolicy );
+IMPLEMENT_BASEPASS_LIGHTMAPPED_SHADER_TYPE( TUniformLightMapPolicy<LMP_SIMPLE_DIRECTIONAL_LIGHT_LIGHTING>, FSimpleDirectionalLightLightingPolicy );
+IMPLEMENT_BASEPASS_LIGHTMAPPED_SHADER_TYPE( TUniformLightMapPolicy<LMP_SIMPLE_STATIONARY_PRECOMPUTED_SHADOW_LIGHTING>, FSimpleStationaryLightPrecomputedShadowsLightingPolicy );
+IMPLEMENT_BASEPASS_LIGHTMAPPED_SHADER_TYPE( TUniformLightMapPolicy<LMP_SIMPLE_STATIONARY_SINGLESAMPLE_SHADOW_LIGHTING>, FSimpleStationaryLightSingleSampleShadowsLightingPolicy );
 IMPLEMENT_BASEPASS_LIGHTMAPPED_SHADER_TYPE( TUniformLightMapPolicy<LMP_LQ_LIGHTMAP>, TLightMapPolicyLQ );
 IMPLEMENT_BASEPASS_LIGHTMAPPED_SHADER_TYPE( TUniformLightMapPolicy<LMP_HQ_LIGHTMAP>, TLightMapPolicyHQ );
 IMPLEMENT_BASEPASS_LIGHTMAPPED_SHADER_TYPE( TUniformLightMapPolicy<LMP_DISTANCE_FIELD_SHADOWS_AND_HQ_LIGHTMAP>, TDistanceFieldShadowsAndLightMapPolicyHQ  );
@@ -296,24 +300,30 @@ public:
 			TStaticMeshDrawList<TBasePassDrawingPolicy<LightMapPolicyType> >& DrawList =
 				Scene->GetBasePassDrawList<LightMapPolicyType>(DrawType);
 
+			const bool bRenderSkylight = Parameters.ShadingModel != MSM_Unlit
+				&& Scene->SkyLight 
+				&& !Scene->SkyLight->bHasStaticLighting 
+				// The deferred shading renderer does movable skylight diffuse in a later deferred pass, not in the base pass
+				&& (Scene->SkyLight->bWantsStaticShadowing || IsSimpleForwardShadingEnabled(Scene->GetShaderPlatform()));
+
 			// Add the static mesh to the draw list.
 			DrawList.AddMesh(
 				StaticMesh,
 				typename TBasePassDrawingPolicy<LightMapPolicyType>::ElementDataType(LightMapElementData),
 				TBasePassDrawingPolicy<LightMapPolicyType>(
-				StaticMesh->VertexFactory,
-				StaticMesh->MaterialRenderProxy,
-				*Parameters.Material,
-				Parameters.FeatureLevel,
-				LightMapPolicy,
-				Parameters.BlendMode,
-				Parameters.TextureMode,
-				Parameters.ShadingModel != MSM_Unlit && Scene->SkyLight && Scene->SkyLight->bWantsStaticShadowing && !Scene->SkyLight->bHasStaticLighting,
-				IsTranslucentBlendMode(Parameters.BlendMode) && Scene->HasAtmosphericFog(),
-				/* DebugViewShaderMode = */ DVSM_None,
-				/* bInAllowGlobalFog = */ false,
-				/* bInEnableEditorPrimitiveDepthTest = */ false,
-				/* bInEnableReceiveDecalOutput = */ true
+					StaticMesh->VertexFactory,
+					StaticMesh->MaterialRenderProxy,
+					*Parameters.Material,
+					Parameters.FeatureLevel,
+					LightMapPolicy,
+					Parameters.BlendMode,
+					Parameters.TextureMode,
+					bRenderSkylight,
+					IsTranslucentBlendMode(Parameters.BlendMode) && Scene->HasAtmosphericFog(),
+					/* DebugViewShaderMode = */ DVSM_None,
+					/* bInAllowGlobalFog = */ false,
+					/* bInEnableEditorPrimitiveDepthTest = */ false,
+					/* bInEnableReceiveDecalOutput = */ true
 				),
 				Scene->GetFeatureLevel()
 				);
@@ -412,6 +422,11 @@ public:
 #endif
 		const FScene* Scene = Parameters.PrimitiveSceneProxy ? Parameters.PrimitiveSceneProxy->GetPrimitiveSceneInfo()->Scene : NULL;
 
+		const bool bRenderSkylight = Scene && Scene->SkyLight 
+			&& !Scene->SkyLight->bHasStaticLighting 
+			// The deferred shading renderer does movable skylight diffuse in a later deferred pass, not in the base pass
+			&& (Scene->SkyLight->bWantsStaticShadowing || IsSimpleForwardShadingEnabled(Scene->GetShaderPlatform()))
+			&& bIsLitMaterial;
 
 		TBasePassDrawingPolicy<LightMapPolicyType> DrawingPolicy(
 			Parameters.Mesh.VertexFactory,
@@ -421,7 +436,7 @@ public:
 			LightMapPolicy,
 			Parameters.BlendMode,
 			Parameters.TextureMode,
-			Scene && Scene->SkyLight && !Scene->SkyLight->bHasStaticLighting && Scene->SkyLight->bWantsStaticShadowing && bIsLitMaterial,
+			bRenderSkylight,
 			IsTranslucentBlendMode(Parameters.BlendMode) && (Scene && Scene->HasAtmosphericFog()) && View.Family->EngineShowFlags.AtmosphericFog,
 			View.Family->GetDebugViewShaderMode(),
 			false,
@@ -616,8 +631,20 @@ void GetBasePassShaders<FUniformLightMapPolicy>(
 	case LMP_CACHED_POINT_INDIRECT_LIGHTING:
 		GetUniformBasePassShaders<LMP_CACHED_POINT_INDIRECT_LIGHTING>(Material, VertexFactoryType, bNeedsHSDS, bEnableAtmosphericFog, bEnableSkyLight, HullShader, DomainShader, VertexShader, PixelShader);
 		break;
-	case LMP_SIMPLE_DYNAMIC_LIGHTING:
-		GetUniformBasePassShaders<LMP_SIMPLE_DYNAMIC_LIGHTING>(Material, VertexFactoryType, bNeedsHSDS, bEnableAtmosphericFog, bEnableSkyLight, HullShader, DomainShader, VertexShader, PixelShader);
+	case LMP_SIMPLE_DIRECTIONAL_LIGHT_LIGHTING:
+		GetUniformBasePassShaders<LMP_SIMPLE_DIRECTIONAL_LIGHT_LIGHTING>(Material, VertexFactoryType, bNeedsHSDS, bEnableAtmosphericFog, bEnableSkyLight, HullShader, DomainShader, VertexShader, PixelShader);
+		break;
+	case LMP_SIMPLE_NO_LIGHTMAP:
+		GetUniformBasePassShaders<LMP_SIMPLE_NO_LIGHTMAP>(Material, VertexFactoryType, bNeedsHSDS, bEnableAtmosphericFog, bEnableSkyLight, HullShader, DomainShader, VertexShader, PixelShader);
+		break;
+	case LMP_SIMPLE_LIGHTMAP_ONLY_LIGHTING:
+		GetUniformBasePassShaders<LMP_SIMPLE_LIGHTMAP_ONLY_LIGHTING>(Material, VertexFactoryType, bNeedsHSDS, bEnableAtmosphericFog, bEnableSkyLight, HullShader, DomainShader, VertexShader, PixelShader);
+		break;
+	case LMP_SIMPLE_STATIONARY_PRECOMPUTED_SHADOW_LIGHTING:
+		GetUniformBasePassShaders<LMP_SIMPLE_STATIONARY_PRECOMPUTED_SHADOW_LIGHTING>(Material, VertexFactoryType, bNeedsHSDS, bEnableAtmosphericFog, bEnableSkyLight, HullShader, DomainShader, VertexShader, PixelShader);
+		break;
+	case LMP_SIMPLE_STATIONARY_SINGLESAMPLE_SHADOW_LIGHTING:
+		GetUniformBasePassShaders<LMP_SIMPLE_STATIONARY_SINGLESAMPLE_SHADOW_LIGHTING>(Material, VertexFactoryType, bNeedsHSDS, bEnableAtmosphericFog, bEnableSkyLight, HullShader, DomainShader, VertexShader, PixelShader);
 		break;
 	case LMP_LQ_LIGHTMAP:
 		GetUniformBasePassShaders<LMP_LQ_LIGHTMAP>(Material, VertexFactoryType, bNeedsHSDS, bEnableAtmosphericFog, bEnableSkyLight, HullShader, DomainShader, VertexShader, PixelShader);
diff --git a/Engine/Source/Runtime/Renderer/Private/BasePassRendering.h b/Engine/Source/Runtime/Renderer/Private/BasePassRendering.h
--- a/Engine/Source/Runtime/Renderer/Private/BasePassRendering.h
+++ b/Engine/Source/Runtime/Renderer/Private/BasePassRendering.h
@@ -252,7 +252,7 @@ public:
 		bool bShouldCache = Super::ShouldCache(Platform, Material, VertexFactoryType);
 		return bShouldCache 
 			&& (IsFeatureLevelSupported(Platform, ERHIFeatureLevel::SM4))
-			&& (!bEnableAtmosphericFog || IsFeatureLevelSupported(Platform, ERHIFeatureLevel::SM4));
+			&& (!bEnableAtmosphericFog || IsTranslucentBlendMode(Material->GetBlendMode()));
 	}
 
 	static void ModifyCompilationEnvironment(EShaderPlatform Platform, const FMaterial* Material, FShaderCompilerEnvironment& OutEnvironment)
@@ -1140,6 +1140,67 @@ public:
 	}
 };
 
+template<typename ProcessActionType>
+void ProcessBasePassMeshForSimpleForwardShading(
+	FRHICommandList& RHICmdList,
+	const FProcessBasePassMeshParameters& Parameters,
+	const ProcessActionType& Action,
+	const FLightMapInteraction& LightMapInteraction,
+	bool bIsLitMaterial,
+	bool bAllowStaticLighting
+	)
+{
+	if (bAllowStaticLighting && LightMapInteraction.GetType() == LMIT_Texture)
+	{
+		const FShadowMapInteraction ShadowMapInteraction = (Parameters.Mesh.LCI && bIsLitMaterial) 
+			? Parameters.Mesh.LCI->GetShadowMapInteraction() 
+			: FShadowMapInteraction();
+
+		if (ShadowMapInteraction.GetType() == SMIT_Texture)
+		{
+			Action.template Process< FUniformLightMapPolicy >(RHICmdList, Parameters, FUniformLightMapPolicy(LMP_SIMPLE_STATIONARY_PRECOMPUTED_SHADOW_LIGHTING), Parameters.Mesh.LCI);
+		}
+		else
+		{
+			Action.template Process< FUniformLightMapPolicy >(RHICmdList, Parameters, FUniformLightMapPolicy(LMP_SIMPLE_LIGHTMAP_ONLY_LIGHTING), Parameters.Mesh.LCI);
+		}
+	}
+	else if (bIsLitMaterial
+		&& IsIndirectLightingCacheAllowed(Parameters.FeatureLevel)
+		&& Action.AllowIndirectLightingCache()
+		&& Parameters.PrimitiveSceneProxy)
+	{
+		const FIndirectLightingCacheAllocation* IndirectLightingCacheAllocation = Parameters.PrimitiveSceneProxy->GetPrimitiveSceneInfo()->IndirectLightingCacheAllocation;
+		const bool bPrimitiveIsMovable = Parameters.PrimitiveSceneProxy->IsMovable();
+		const bool bPrimitiveUsesILC = Parameters.PrimitiveSceneProxy->GetIndirectLightingCacheQuality() != ILCQ_Off;								
+
+		// Use the indirect lighting cache shaders if the object has a cache allocation
+		// This happens for objects with unbuilt lighting
+		if (bPrimitiveUsesILC &&
+			((IndirectLightingCacheAllocation && IndirectLightingCacheAllocation->IsValid())
+			// Use the indirect lighting cache shaders if the object is movable, it may not have a cache allocation yet because that is done in InitViews
+			// And movable objects are sometimes rendered in the static draw lists
+			|| bPrimitiveIsMovable))
+		{
+			// Use a lightmap policy that supports reading indirect lighting from a single SH sample
+			Action.template Process< FUniformLightMapPolicy >(RHICmdList, Parameters, FUniformLightMapPolicy(LMP_SIMPLE_STATIONARY_SINGLESAMPLE_SHADOW_LIGHTING), Parameters.Mesh.LCI);
+		}
+		else
+		{
+			Action.template Process< FUniformLightMapPolicy >(RHICmdList, Parameters, FUniformLightMapPolicy(LMP_SIMPLE_NO_LIGHTMAP), Parameters.Mesh.LCI);
+		}
+	}
+	else if (bIsLitMaterial)
+	{
+		// Always choosing shaders to support dynamic directional even if one is not present
+		Action.template Process< FUniformLightMapPolicy >(RHICmdList, Parameters, FUniformLightMapPolicy(LMP_SIMPLE_DIRECTIONAL_LIGHT_LIGHTING), Parameters.Mesh.LCI);
+	}
+	else
+	{
+		Action.template Process< FUniformLightMapPolicy >(RHICmdList, Parameters, FUniformLightMapPolicy(LMP_SIMPLE_NO_LIGHTMAP), Parameters.Mesh.LCI);
+	}
+}
+
 /** Processes a base pass mesh using an unknown light map policy, and unknown fog density policy. */
 template<typename ProcessActionType>
 void ProcessBasePassMesh(
@@ -1182,41 +1243,44 @@ void ProcessBasePassMesh(
 				: FLightMapInteraction();
 
 			// force LQ lightmaps based on system settings
-			const bool bAllowHighQualityLightMaps = AllowHighQualityLightmaps(Parameters.FeatureLevel) && LightMapInteraction.AllowsHighQualityLightmaps();
+			const bool bPlatformAllowsHighQualityLightMaps = AllowHighQualityLightmaps(Parameters.FeatureLevel);
+			const bool bAllowHighQualityLightMaps = bPlatformAllowsHighQualityLightMaps && LightMapInteraction.AllowsHighQualityLightmaps();
 
-			switch(LightMapInteraction.GetType())
+			if (IsSimpleForwardShadingEnabled(GetFeatureLevelShaderPlatform(Parameters.FeatureLevel)))
 			{
-				case LMIT_Texture: 
-					if( bAllowHighQualityLightMaps ) 
-					{ 
-						const FShadowMapInteraction ShadowMapInteraction = (bAllowStaticLighting && Parameters.Mesh.LCI && bIsLitMaterial) 
-							? Parameters.Mesh.LCI->GetShadowMapInteraction() 
-							: FShadowMapInteraction();
-
-						if (ShadowMapInteraction.GetType() == SMIT_Texture)
-						{
-							Action.template Process< FUniformLightMapPolicy >(RHICmdList, Parameters, FUniformLightMapPolicy(LMP_DISTANCE_FIELD_SHADOWS_AND_HQ_LIGHTMAP), Parameters.Mesh.LCI);
-						}
-						else
-						{
-							Action.template Process< FUniformLightMapPolicy >(RHICmdList, Parameters, FUniformLightMapPolicy(LMP_HQ_LIGHTMAP), Parameters.Mesh.LCI);
-						}
-					} 
-					else 
-					{ 
-						Action.template Process< FUniformLightMapPolicy >(RHICmdList, Parameters, FUniformLightMapPolicy(LMP_LQ_LIGHTMAP), Parameters.Mesh.LCI);
-					} 
-					break;
-				default:
-					{
-						// Use simple dynamic lighting if enabled, which just renders an unshadowed directional light and a skylight
-						if (bIsLitMaterial)
-						{
-							if (IsSimpleDynamicLightingEnabled())
+				// Only compiling simple lighting shaders for HQ lightmaps to save on permutations
+				check(bPlatformAllowsHighQualityLightMaps);
+				ProcessBasePassMeshForSimpleForwardShading(RHICmdList, Parameters, Action, LightMapInteraction, bIsLitMaterial, bAllowStaticLighting);
+			}
+			else
+			{
+				switch(LightMapInteraction.GetType())
+				{
+					case LMIT_Texture: 
+						if( bAllowHighQualityLightMaps ) 
+						{ 
+							const FShadowMapInteraction ShadowMapInteraction = (bAllowStaticLighting && Parameters.Mesh.LCI && bIsLitMaterial) 
+								? Parameters.Mesh.LCI->GetShadowMapInteraction() 
+								: FShadowMapInteraction();
+
+							if (ShadowMapInteraction.GetType() == SMIT_Texture)
+							{
+								Action.template Process< FUniformLightMapPolicy >(RHICmdList, Parameters, FUniformLightMapPolicy(LMP_DISTANCE_FIELD_SHADOWS_AND_HQ_LIGHTMAP), Parameters.Mesh.LCI);
+							}
+							else
 							{
-								Action.template Process< FUniformLightMapPolicy >(RHICmdList, Parameters, FUniformLightMapPolicy(LMP_SIMPLE_DYNAMIC_LIGHTING), Parameters.Mesh.LCI);
+								Action.template Process< FUniformLightMapPolicy >(RHICmdList, Parameters, FUniformLightMapPolicy(LMP_HQ_LIGHTMAP), Parameters.Mesh.LCI);
 							}
-							else if (IsIndirectLightingCacheAllowed(Parameters.FeatureLevel)
+						} 
+						else 
+						{ 
+							Action.template Process< FUniformLightMapPolicy >(RHICmdList, Parameters, FUniformLightMapPolicy(LMP_LQ_LIGHTMAP), Parameters.Mesh.LCI);
+						} 
+						break;
+					default:
+						{
+							if (bIsLitMaterial
+								&& IsIndirectLightingCacheAllowed(Parameters.FeatureLevel)
 								&& Action.AllowIndirectLightingCache()
 								&& Parameters.PrimitiveSceneProxy)
 							{
@@ -1257,13 +1321,9 @@ void ProcessBasePassMesh(
 								Action.template Process< FUniformLightMapPolicy >(RHICmdList, Parameters, FUniformLightMapPolicy(LMP_NO_LIGHTMAP), Parameters.Mesh.LCI);
 							}
 						}
-						else
-						{
-							Action.template Process< FUniformLightMapPolicy >(RHICmdList, Parameters, FUniformLightMapPolicy(LMP_NO_LIGHTMAP), Parameters.Mesh.LCI);
-						}
-					}
-					break;
-			};
+						break;
+				};
+			}
 		}
 	}
 }
