Fix for crash in GPU profiler. This was caused by the RHIThread getting too far behind the renderthread. This change adds a fence wait on the renderthread in RHIEndDrawingViewport to ensure that the renderthread is never more than a frame ahead.

diff --git a/Engine/Source/Runtime/RHI/Private/RHICommandList.cpp b/Engine/Source/Runtime/RHI/Private/RHICommandList.cpp
--- a/Engine/Source/Runtime/RHI/Private/RHICommandList.cpp
+++ b/Engine/Source/Runtime/RHI/Private/RHICommandList.cpp
@@ -115,6 +115,9 @@ static FGraphEventRef RHIThreadTask;
 static FGraphEventRef RenderThreadSublistDispatchTask;
 static FGraphEventRef RHIThreadBufferLockFence;
 
+static FGraphEventRef GRHIThreadEndDrawingViewportFences[2];
+static uint32 GRHIThreadEndDrawingViewportFenceIndex = 0;
+
 // Used by AsyncCompute
 RHI_API FRHICommandListFenceAllocator GRHIFenceAllocator;
 
@@ -1344,12 +1347,28 @@ void FRHICommandList::EndDrawingViewport(FViewportRHIParamRef Viewport, bool bPr
 	else
 	{
 		new (AllocCommand<FRHICommandEndDrawingViewport>()) FRHICommandEndDrawingViewport(Viewport, bPresent, bLockToVsync);
+
+		if ( GRHIThread )
+		{
+			// Insert a fence to prevent the renderthread getting more than a frame ahead of the RHIThread
+			GRHIThreadEndDrawingViewportFences[GRHIThreadEndDrawingViewportFenceIndex] = static_cast<FRHICommandListImmediate*>(this)->RHIThreadFence();
+		}
 		// if we aren't running an RHIThread, there is no good reason to buffer this frame advance stuff and that complicates state management, so flush everything out now
 		{
 			QUICK_SCOPE_CYCLE_COUNTER(STAT_EndDrawingViewport_Dispatch);
 			FRHICommandListExecutor::GetImmediateCommandList().ImmediateFlush(EImmediateFlushType::DispatchToRHIThread);
 		}
 	}
+
+	if ( GRHIThread )
+	{
+		// Wait on the previous frame's RHI thread fence (we never want the rendering thread to get more than a frame ahead)
+		uint32 PreviousFrameFenceIndex = 1 - GRHIThreadEndDrawingViewportFenceIndex;
+		FRHICommandListExecutor::WaitOnRHIThreadFence( GRHIThreadEndDrawingViewportFences[PreviousFrameFenceIndex] );
+		GRHIThreadEndDrawingViewportFences[PreviousFrameFenceIndex] = nullptr;
+		GRHIThreadEndDrawingViewportFenceIndex = PreviousFrameFenceIndex;
+	}
+
 	RHIAdvanceFrameForGetViewportBackBuffer();
 }
 
