Improved solution for #jira UE-37216. The previous change prevented the crash, but didn't fix the issue of RHI resources getting accumulated while the loading screen was up without being freed until it was done.

diff --git a/Engine/Source/Runtime/Engine/Private/GameViewportClient.cpp b/Engine/Source/Runtime/Engine/Private/GameViewportClient.cpp
--- a/Engine/Source/Runtime/Engine/Private/GameViewportClient.cpp
+++ b/Engine/Source/Runtime/Engine/Private/GameViewportClient.cpp
@@ -36,6 +36,7 @@
 #include "MovieSceneCaptureSettings.h"
 #include "ActorEditorUtils.h"
 #include "ComponentRecreateRenderStateContext.h"
+#include "RenderingThread.h"
 
 #define LOCTEXT_NAMESPACE "GameViewport"
 
@@ -1122,6 +1123,14 @@ void UGameViewportClient::Draw(FViewport* InViewport, FCanvas* SceneCanvas)
 	{
 		GetRendererModule().BeginRenderingViewFamily(SceneCanvas,&ViewFamily);
 	}
+	else
+	{
+		// Make sure RHI resources get flushed if we're not using a renderer
+		ENQUEUE_UNIQUE_RENDER_COMMAND( UGameViewportClient_FlushRHIResources,
+		{ 
+			FRHICommandListExecutor::GetImmediateCommandList().ImmediateFlush(EImmediateFlushType::FlushRHIThreadFlushResources);
+		});
+	}
 
 	// Clear areas of the rendertarget (backbuffer) that aren't drawn over by the views.
 	if (!bBufferCleared)
diff --git a/Engine/Source/Runtime/RHI/Private/RHICommandList.cpp b/Engine/Source/Runtime/RHI/Private/RHICommandList.cpp
--- a/Engine/Source/Runtime/RHI/Private/RHICommandList.cpp
+++ b/Engine/Source/Runtime/RHI/Private/RHICommandList.cpp
@@ -1364,7 +1364,9 @@ void FRHICommandList::EndDrawingViewport(FViewportRHIParamRef Viewport, bool bPr
 	{
 		// Wait on the previous frame's RHI thread fence (we never want the rendering thread to get more than a frame ahead)
 		uint32 PreviousFrameFenceIndex = 1 - GRHIThreadEndDrawingViewportFenceIndex;
-		FRHICommandListExecutor::WaitOnRHIThreadFence( GRHIThreadEndDrawingViewportFences[PreviousFrameFenceIndex] );
+		FGraphEventRef& LastFrameFence = GRHIThreadEndDrawingViewportFences[PreviousFrameFenceIndex];
+		ensureMsgf(!LastFrameFence.GetReference() || LastFrameFence->IsComplete(), TEXT("The renderthread had to wait for the RHI thread to catch up. We are likely missing a call to FRHICommandListExecutor::GetImmediateCommandList().ImmediateFlush(EImmediateFlushType::FlushRHIThreadFlushResources)"));
+		FRHICommandListExecutor::WaitOnRHIThreadFence(LastFrameFence);
 		GRHIThreadEndDrawingViewportFences[PreviousFrameFenceIndex] = nullptr;
 		GRHIThreadEndDrawingViewportFenceIndex = PreviousFrameFenceIndex;
 	}
