UE4 - Warning fix

diff --git a/Engine/Source/Runtime/RHI/Private/PipelineStateCache.cpp b/Engine/Source/Runtime/RHI/Private/PipelineStateCache.cpp
--- a/Engine/Source/Runtime/RHI/Private/PipelineStateCache.cpp
+++ b/Engine/Source/Runtime/RHI/Private/PipelineStateCache.cpp
@@ -4,8 +4,6 @@
 PipelineStateCache.cpp: Pipeline state cache implementation.
 =============================================================================*/
 
-#pragma once
-
 #include "PipelineStateCache.h"
 #include "Misc/ScopeLock.h"
 
