[UE-43106] - Fixed stomped cached PSO state when CmdList is split which manifested itself on MacOS

diff --git a/Engine/Source/Runtime/RHI/Private/RHICommandList.cpp b/Engine/Source/Runtime/RHI/Private/RHICommandList.cpp
--- a/Engine/Source/Runtime/RHI/Private/RHICommandList.cpp
+++ b/Engine/Source/Runtime/RHI/Private/RHICommandList.cpp
@@ -205,6 +205,7 @@ void FRHIAsyncComputeCommandListImmediate::ImmediateDispatch(FRHIAsyncComputeCom
 			static_assert(sizeof(FRHICommandList) == sizeof(FRHIAsyncComputeCommandListImmediate), "We are memswapping FRHICommandList and FRHICommandListImmediate; they need to be swappable.");
 			check(RHIComputeCmdList.IsImmediateAsyncCompute());
 			SwapCmdList->ExchangeCmdList(RHIComputeCmdList);
+			RHIComputeCmdList.PSOContext = SwapCmdList->PSOContext;
 
 			//queue the execution of this async commandlist amongst other commands in the immediate gfx list.
 			//this guarantees resource update commands made on the gfx commandlist will be executed before the async compute.
@@ -383,6 +384,7 @@ void FRHICommandListExecutor::ExecuteInner(FRHICommandListBase& CmdList)
 				// we should make command lists virtual and transfer ownership rather than this devious approach
 				static_assert(sizeof(FRHICommandList) == sizeof(FRHICommandListImmediate), "We are memswapping FRHICommandList and FRHICommandListImmediate; they need to be swappable.");
 				SwapCmdList->ExchangeCmdList(CmdList);
+				CmdList.PSOContext = SwapCmdList->PSOContext;
 			}
 			QUICK_SCOPE_CYCLE_COUNTER(STAT_FRHICommandListExecutor_SubmitTasks);
 
@@ -703,7 +705,6 @@ void FRHICommandListExecutor::WaitOnRHIThreadFence(FGraphEventRef& Fence)
 
 FRHICommandListBase::FRHICommandListBase()
 	: MemManager(0)
-	, CachedNumSimultanousRenderTargets(0)
 {
 	GRHICommandList.OutstandingCmdListCount.Increment();
 	Reset();
diff --git a/Engine/Source/Runtime/RHI/Public/RHICommandList.h b/Engine/Source/Runtime/RHI/Public/RHICommandList.h
--- a/Engine/Source/Runtime/RHI/Public/RHICommandList.h
+++ b/Engine/Source/Runtime/RHI/Public/RHICommandList.h
@@ -307,11 +307,15 @@ public:
 	void *RenderThreadContexts[(int32)ERenderThreadContext::Num];
 	static int32 StateCacheEnabled;
 
-	uint32 CachedNumSimultanousRenderTargets;
-	TStaticArray<FRHIRenderTargetView, MaxSimultaneousRenderTargets> CachedRenderTargets;
-	FRHIDepthRenderTargetView CachedDepthStencilTarget;
-
 protected:
+	//the values of this struct must be copied when the commandlist is split 
+	struct FPSOContext
+	{
+		uint32 CachedNumSimultanousRenderTargets = 0;
+		TStaticArray<FRHIRenderTargetView, MaxSimultaneousRenderTargets> CachedRenderTargets;
+		FRHIDepthRenderTargetView CachedDepthStencilTarget;
+	} PSOContext;
+
 	struct FRHICommandSetRasterizerState* CachedRasterizerState;
 	struct FRHICommandSetDepthStencilState* CachedDepthStencilState;
 
@@ -321,14 +325,14 @@ protected:
 		const FRHIDepthRenderTargetView* NewDepthStencilTargetRHI
 		)
 	{
-		CachedNumSimultanousRenderTargets = NewNumSimultaneousRenderTargets;
+		PSOContext.CachedNumSimultanousRenderTargets = NewNumSimultaneousRenderTargets;
 
-		for (uint32 RTIdx = 0; RTIdx < CachedNumSimultanousRenderTargets; ++RTIdx)
+		for (uint32 RTIdx = 0; RTIdx < PSOContext.CachedNumSimultanousRenderTargets; ++RTIdx)
 		{
-			CachedRenderTargets[RTIdx] = NewRenderTargetsRHI[RTIdx];
+			PSOContext.CachedRenderTargets[RTIdx] = NewRenderTargetsRHI[RTIdx];
 		}
 
-		CachedDepthStencilTarget = (NewDepthStencilTargetRHI) ? *NewDepthStencilTargetRHI : FRHIDepthRenderTargetView();
+		PSOContext.CachedDepthStencilTarget = (NewDepthStencilTargetRHI) ? *NewDepthStencilTargetRHI : FRHIDepthRenderTargetView();
 	}
 
 public:
@@ -1976,47 +1980,47 @@ public:
 		FGraphicsPipelineStateInitializer& GraphicsPSOInit
 		)
 	{
-		GraphicsPSOInit.RenderTargetsEnabled = CachedNumSimultanousRenderTargets;
+		GraphicsPSOInit.RenderTargetsEnabled = PSOContext.CachedNumSimultanousRenderTargets;
 
 		for (uint32 i = 0; i < GraphicsPSOInit.RenderTargetsEnabled; ++i)
 		{
-			if (CachedRenderTargets[i].Texture)
+			if (PSOContext.CachedRenderTargets[i].Texture)
 			{
-				GraphicsPSOInit.RenderTargetFormats[i] = CachedRenderTargets[i].Texture->GetFormat();
-				GraphicsPSOInit.RenderTargetFlags[i] = CachedRenderTargets[i].Texture->GetFlags();
+				GraphicsPSOInit.RenderTargetFormats[i] = PSOContext.CachedRenderTargets[i].Texture->GetFormat();
+				GraphicsPSOInit.RenderTargetFlags[i] = PSOContext.CachedRenderTargets[i].Texture->GetFlags();
 			}
 			else
 			{
 				GraphicsPSOInit.RenderTargetFormats[i] = PF_Unknown;
 			}
 
-			GraphicsPSOInit.RenderTargetLoadActions[i] = CachedRenderTargets[i].LoadAction;
-			GraphicsPSOInit.RenderTargetStoreActions[i] = CachedRenderTargets[i].StoreAction;
+			GraphicsPSOInit.RenderTargetLoadActions[i] = PSOContext.CachedRenderTargets[i].LoadAction;
+			GraphicsPSOInit.RenderTargetStoreActions[i] = PSOContext.CachedRenderTargets[i].StoreAction;
 
 			if (GraphicsPSOInit.RenderTargetFormats[i] != PF_Unknown)
 			{
-				GraphicsPSOInit.NumSamples = CachedRenderTargets[i].Texture->GetNumSamples();
+				GraphicsPSOInit.NumSamples = PSOContext.CachedRenderTargets[i].Texture->GetNumSamples();
 			}
 		}
 
-		if (CachedDepthStencilTarget.Texture)
+		if (PSOContext.CachedDepthStencilTarget.Texture)
 		{
-			GraphicsPSOInit.DepthStencilTargetFormat = CachedDepthStencilTarget.Texture->GetFormat();
-			GraphicsPSOInit.DepthStencilTargetFlag = CachedDepthStencilTarget.Texture->GetFlags();
+			GraphicsPSOInit.DepthStencilTargetFormat = PSOContext.CachedDepthStencilTarget.Texture->GetFormat();
+			GraphicsPSOInit.DepthStencilTargetFlag = PSOContext.CachedDepthStencilTarget.Texture->GetFlags();
 		}
 		else
 		{
 			GraphicsPSOInit.DepthStencilTargetFormat = PF_Unknown;
 		}
 
-		GraphicsPSOInit.DepthTargetLoadAction = CachedDepthStencilTarget.DepthLoadAction;
-		GraphicsPSOInit.DepthTargetStoreAction = CachedDepthStencilTarget.DepthStoreAction;
-		GraphicsPSOInit.StencilTargetLoadAction = CachedDepthStencilTarget.StencilLoadAction;
-		GraphicsPSOInit.StencilTargetStoreAction = CachedDepthStencilTarget.GetStencilStoreAction();
+		GraphicsPSOInit.DepthTargetLoadAction = PSOContext.CachedDepthStencilTarget.DepthLoadAction;
+		GraphicsPSOInit.DepthTargetStoreAction = PSOContext.CachedDepthStencilTarget.DepthStoreAction;
+		GraphicsPSOInit.StencilTargetLoadAction = PSOContext.CachedDepthStencilTarget.StencilLoadAction;
+		GraphicsPSOInit.StencilTargetStoreAction = PSOContext.CachedDepthStencilTarget.GetStencilStoreAction();
 
 		if (GraphicsPSOInit.DepthStencilTargetFormat != PF_Unknown)
 		{
-			GraphicsPSOInit.NumSamples = CachedDepthStencilTarget.Texture->GetNumSamples();
+			GraphicsPSOInit.NumSamples = PSOContext.CachedDepthStencilTarget.Texture->GetNumSamples();
 		}
 	}
 
