Copying //UE4/Dev-Rendering to //UE4/Dev-Main (Source: //UE4/Dev-Rendering @ 3461187)

diff --git a/Engine/Source/Runtime/RHI/Private/RHICommandList.cpp b/Engine/Source/Runtime/RHI/Private/RHICommandList.cpp
--- a/Engine/Source/Runtime/RHI/Private/RHICommandList.cpp
+++ b/Engine/Source/Runtime/RHI/Private/RHICommandList.cpp
@@ -8,6 +8,7 @@
 #include "Misc/App.h"
 #include "Async/TaskGraphInterfaces.h"
 #include "RHI.h"
+#include "ScopeLock.h"
 
 DECLARE_CYCLE_STAT(TEXT("Nonimmed. Command List Execute"), STAT_NonImmedCmdListExecuteTime, STATGROUP_RHICMDLIST);
 DECLARE_DWORD_COUNTER_STAT(TEXT("Nonimmed. Command List memory"), STAT_NonImmedCmdListMemory, STATGROUP_RHICMDLIST);
@@ -103,6 +104,13 @@ static TAutoConsoleVariable<int32> CVarRHICmdMinCmdlistSizeForParallelTranslate(
 	32,
 	TEXT("In kilobytes. Cmdlists are merged into one parallel translate until we have at least this much memory to process. For a given pass, we won't do more translates than we have task threads. Only relevant if r.RHICmdBalanceTranslatesAfterTasks is on."));
 
+
+bool GUseRHIThread_InternalUseOnly = false;
+bool GUseRHITaskThreads_InternalUseOnly = false;
+bool GIsRunningRHIInSeparateThread_InternalUseOnly = false;
+bool GIsRunningRHIInDedicatedThread_InternalUseOnly = false;
+bool GIsRunningRHIInTaskThread_InternalUseOnly = false;
+
 RHI_API bool GEnableAsyncCompute = true;
 RHI_API FRHICommandListExecutor GRHICommandList;
 
@@ -257,6 +265,16 @@ void FRHICommandListExecutor::ExecuteInner_DoExecute(FRHICommandListBase& CmdLis
 }
 
 
+static FAutoConsoleTaskPriority CPrio_RHIThreadOnTaskThreads(
+	TEXT("TaskGraph.TaskPriorities.RHIThreadOnTaskThreads"),
+	TEXT("Task and thread priority for when we are running 'RHI thread' tasks on any thread."),
+	ENamedThreads::NormalThreadPriority,
+	ENamedThreads::NormalTaskPriority
+	);
+
+
+static FCriticalSection GRHIThreadOnTasksCritical;
+
 
 class FExecuteRHIThreadTask
 {
@@ -276,8 +294,8 @@ public:
 
 	ENamedThreads::Type GetDesiredThread()
 	{
-		check(GRHIThread); // this should never be used on a platform that doesn't support the RHI thread
-		return ENamedThreads::RHIThread;
+		check(IsRunningRHIInSeparateThread()); // this should never be used on a platform that doesn't support the RHI thread
+		return IsRunningRHIInDedicatedThread() ? ENamedThreads::RHIThread : CPrio_RHIThreadOnTaskThreads.Get();
 	}
 
 	static ESubsequentsMode::Type GetSubsequentsMode() { return ESubsequentsMode::TrackSubsequents; }
@@ -285,8 +303,19 @@ public:
 	void DoTask(ENamedThreads::Type CurrentThread, const FGraphEventRef& MyCompletionGraphEvent)
 	{
 		SCOPE_CYCLE_COUNTER(STAT_RHIThreadExecute);
-		FRHICommandListExecutor::ExecuteInner_DoExecute(*RHICmdList);
-		delete RHICmdList;
+		if (IsRunningRHIInTaskThread())
+		{
+			GRHIThreadId = FPlatformTLS::GetCurrentThreadId();
+		}
+		{
+			FScopeLock Lock(&GRHIThreadOnTasksCritical);
+			FRHICommandListExecutor::ExecuteInner_DoExecute(*RHICmdList);
+			delete RHICmdList;
+		}
+		if (IsRunningRHIInTaskThread())
+		{
+			GRHIThreadId = 0;
+		}
 	}
 };
 
@@ -312,8 +341,8 @@ public:
 	{
 		// If we are using async dispatch, this task is somewhat redundant, but it does allow things to wait for dispatch without waiting for execution. 
 		// since in that case we will be queuing an rhithread task from an rhithread task, the overhead is minor.
-		check(GRHIThread); // this should never be used on a platform that doesn't support the RHI thread
-		return bRHIThread ? ENamedThreads::RHIThread : ENamedThreads::RenderThread_Local;
+		check(IsRunningRHIInSeparateThread()); // this should never be used on a platform that doesn't support the RHI thread
+		return bRHIThread ? (IsRunningRHIInDedicatedThread() ? ENamedThreads::RHIThread : CPrio_RHIThreadOnTaskThreads.Get()) : ENamedThreads::RenderThread_Local;
 	}
 
 	static ESubsequentsMode::Type GetSubsequentsMode() { return ESubsequentsMode::TrackSubsequents; }
@@ -330,15 +359,13 @@ public:
 	}
 };
 
-
-
 void FRHICommandListExecutor::ExecuteInner(FRHICommandListBase& CmdList)
 {
 	check(CmdList.HasCommands()); 
 
 	bool bIsInRenderingThread = IsInRenderingThread();
 	bool bIsInGameThread = IsInGameThread();
-	if (GRHIThread)
+	if (IsRunningRHIInSeparateThread())
 	{
 		bool bAsyncSubmit = false;
 		if (bIsInRenderingThread)
@@ -354,9 +381,9 @@ void FRHICommandListExecutor::ExecuteInner(FRHICommandListBase& CmdList)
 			{
 				RenderThreadSublistDispatchTask = nullptr;
 				if (bAsyncSubmit && RHIThreadTask.GetReference() && RHIThreadTask->IsComplete())
-			{
-				RHIThreadTask = nullptr;
-			}
+				{
+					RHIThreadTask = nullptr;
+				}
 			}
 			if (!bAsyncSubmit && RHIThreadTask.GetReference() && RHIThreadTask->IsComplete())
 			{
@@ -521,7 +548,7 @@ void FRHICommandListExecutor::ExecuteList(FRHICommandListImmediate& CmdList)
 void FRHICommandListExecutor::LatchBypass()
 {
 #if CAN_TOGGLE_COMMAND_LIST_BYPASS
-	if (GRHIThread)
+	if (IsRunningRHIInSeparateThread())
 	{
 		if (bLatchedBypass)
 		{
@@ -629,7 +656,7 @@ struct FRHICommandRHIThreadFence : public FRHICommand<FRHICommandRHIThreadFence>
 	{
 		check(IsInRHIThread());
 		static TArray<FBaseGraphTask*> NewTasks;
-		Fence->DispatchSubsequents(NewTasks, ENamedThreads::RHIThread);
+		Fence->DispatchSubsequents(NewTasks, IsRunningRHIInDedicatedThread() ? ENamedThreads::RHIThread : ENamedThreads::AnyThread);
 		Fence = nullptr;
 	}
 };
@@ -637,7 +664,7 @@ struct FRHICommandRHIThreadFence : public FRHICommand<FRHICommandRHIThreadFence>
 
 FGraphEventRef FRHICommandListImmediate::RHIThreadFence(bool bSetLockFence)
 {
-	check(IsInRenderingThread() && GRHIThread);
+	check(IsInRenderingThread() && IsRunningRHIInSeparateThread());
 	FRHICommandRHIThreadFence* Cmd = new (AllocCommand<FRHICommandRHIThreadFence>()) FRHICommandRHIThreadFence();
 	if (bSetLockFence)
 	{
@@ -681,7 +708,7 @@ void FRHICommandListExecutor::WaitOnRHIThreadFence(FGraphEventRef& Fence)
 			QUICK_SCOPE_CYCLE_COUNTER(STAT_WaitOnRHIThreadFence_Dispatch);
 			GetImmediateCommandList().ImmediateFlush(EImmediateFlushType::DispatchToRHIThread); // necessary to prevent deadlock
 		}
-		check(GRHIThread);
+		check(IsRunningRHIInSeparateThread());
 		QUICK_SCOPE_CYCLE_COUNTER(STAT_WaitOnRHIThreadFence_Wait);
 		if (FTaskGraphInterface::Get().IsThreadProcessingTasks(ENamedThreads::RenderThread_Local))
 		{
@@ -834,7 +861,7 @@ struct FRHICommandWaitForAndSubmitSubListParallel : public FRHICommand<FRHIComma
 			}
 			else if (IsInRHIThread())
 			{
-				FTaskGraphInterface::Get().WaitUntilTaskCompletes(TranslateCompletionEvent, ENamedThreads::RHIThread);
+				FTaskGraphInterface::Get().WaitUntilTaskCompletes(TranslateCompletionEvent, IsRunningRHIInDedicatedThread() ? ENamedThreads::RHIThread : ENamedThreads::AnyThread);
 			}
 			else
 			{
@@ -868,7 +895,7 @@ struct FRHICommandWaitForAndSubmitSubList : public FRHICommand<FRHICommandWaitFo
 	void Execute(FRHICommandListBase& CmdList)
 	{
 		INC_DWORD_STAT_BY(STAT_ChainLinkCount, 1);
-		if (EventToWaitFor.GetReference() && !EventToWaitFor->IsComplete() && !(!GRHIThread || !IsInRHIThread()))
+		if (EventToWaitFor.GetReference() && !EventToWaitFor->IsComplete() && !(!IsRunningRHIInSeparateThread() || !IsInRHIThread()))
 		{
 			GEventToWaitFor = EventToWaitFor.GetReference();
 			FPlatformMisc::DebugBreak();
@@ -876,7 +903,7 @@ struct FRHICommandWaitForAndSubmitSubList : public FRHICommand<FRHICommandWaitFo
 		}
 		if (EventToWaitFor.GetReference() && !EventToWaitFor->IsComplete())
 		{
-			check(!GRHIThread || !IsInRHIThread()); // things should not be dispatched if they can't complete without further waits
+			check(!IsRunningRHIInSeparateThread() || !IsInRHIThread()); // things should not be dispatched if they can't complete without further waits
 			SCOPE_CYCLE_COUNTER(STAT_ChainWait);
 			if (IsInRenderingThread())
 			{
@@ -1012,7 +1039,7 @@ void FRHICommandListBase::QueueParallelAsyncCommandListSubmit(FGraphEventRef* An
 	QUICK_SCOPE_CYCLE_COUNTER(STAT_FRHICommandListBase_QueueParallelAsyncCommandListSubmit);
 	check(IsInRenderingThread() && IsImmediate() && Num);
 
-	if (GRHIThread)
+	if (IsRunningRHIInSeparateThread())
 	{
 		FRHICommandListExecutor::GetImmediateCommandList().ImmediateFlush(EImmediateFlushType::DispatchToRHIThread); // we should start on the stuff before this async list
 
@@ -1030,14 +1057,14 @@ void FRHICommandListBase::QueueParallelAsyncCommandListSubmit(FGraphEventRef* An
 	}
 #endif
 
-	if (Num && GRHIThread)
+	if (Num && IsRunningRHIInSeparateThread())
 	{
 		static const auto ICVarRHICmdBalanceParallelLists = IConsoleManager::Get().FindTConsoleVariableDataInt(TEXT("r.RHICmdBalanceParallelLists"));
 
 		if (ICVarRHICmdBalanceParallelLists->GetValueOnRenderThread() == 0 && CVarRHICmdBalanceTranslatesAfterTasks.GetValueOnRenderThread() > 0 && GRHISupportsParallelRHIExecute && CVarRHICmdUseDeferredContexts.GetValueOnAnyThread() > 0)
 		{
 			FGraphEventArray Prereq;
-			FRHICommandListBase** RHICmdLists = (FRHICommandListBase**)Alloc(sizeof(FRHICommandListBase*) * Num, ALIGNOF(FRHICommandListBase*));
+			FRHICommandListBase** RHICmdLists = (FRHICommandListBase**)Alloc(sizeof(FRHICommandListBase*) * Num, alignof(FRHICommandListBase*));
 			for (int32 Index = 0; Index < Num; Index++)
 			{
 				FGraphEventRef& AnyThreadCompletionEvent = AnyThreadCompletionEvents[Index];
@@ -1059,7 +1086,7 @@ void FRHICommandListBase::QueueParallelAsyncCommandListSubmit(FGraphEventRef* An
 			FGraphEventRef TranslateSetupCompletionEvent = TGraphTask<FParallelTranslateSetupCommandList>::CreateTask(&Prereq, ENamedThreads::RenderThread).ConstructAndDispatchWhenReady(CmdList, &RHICmdLists[0], Num, bIsPrepass);
 			QueueCommandListSubmit(CmdList);
 			AllOutstandingTasks.Add(TranslateSetupCompletionEvent);
-			if (GRHIThread)
+			if (IsRunningRHIInSeparateThread())
 			{
 				FRHICommandListExecutor::GetImmediateCommandList().ImmediateFlush(EImmediateFlushType::DispatchToRHIThread); // we don't want stuff after the async cmd list to be bundled with it
 			}
@@ -1127,7 +1154,7 @@ void FRHICommandListBase::QueueParallelAsyncCommandListSubmit(FGraphEventRef* An
 				check(ContextContainer);
 
 				FGraphEventArray Prereq;
-				FRHICommandListBase** RHICmdLists = (FRHICommandListBase**)Alloc(sizeof(FRHICommandListBase*) * (1 + Last - Start), ALIGNOF(FRHICommandListBase*));
+				FRHICommandListBase** RHICmdLists = (FRHICommandListBase**)Alloc(sizeof(FRHICommandListBase*) * (1 + Last - Start), alignof(FRHICommandListBase*));
 				for (int32 Index = Start; Index <= Last; Index++)
 				{
 					FGraphEventRef& AnyThreadCompletionEvent = AnyThreadCompletionEvents[Index];
@@ -1152,7 +1179,7 @@ void FRHICommandListBase::QueueParallelAsyncCommandListSubmit(FGraphEventRef* An
 
 				AllOutstandingTasks.Add(TranslateCompletionEvent);
 				new (AllocCommand<FRHICommandWaitForAndSubmitSubListParallel>()) FRHICommandWaitForAndSubmitSubListParallel(TranslateCompletionEvent, ContextContainer, EffectiveThreads, ThreadIndex++);
-				if (GRHIThread)
+				if (IsRunningRHIInSeparateThread())
 				{
 					FRHICommandListExecutor::GetImmediateCommandList().ImmediateFlush(EImmediateFlushType::DispatchToRHIThread); // we don't want stuff after the async cmd list to be bundled with it
 				}
@@ -1176,7 +1203,7 @@ void FRHICommandListBase::QueueParallelAsyncCommandListSubmit(FGraphEventRef* An
 		FRHICommandList* CmdList = CmdLists[Index];
 		if (AnyThreadCompletionEvent.GetReference())
 		{
-			if (GRHIThread)
+			if (IsRunningRHIInSeparateThread())
 			{
 				AllOutstandingTasks.Add(AnyThreadCompletionEvent);
 			}
@@ -1184,7 +1211,7 @@ void FRHICommandListBase::QueueParallelAsyncCommandListSubmit(FGraphEventRef* An
 		}
 		new (AllocCommand<FRHICommandWaitForAndSubmitSubList>()) FRHICommandWaitForAndSubmitSubList(AnyThreadCompletionEvent, CmdList);
 	}
-	if (GRHIThread)
+	if (IsRunningRHIInSeparateThread())
 	{
 		FRHICommandListExecutor::GetImmediateCommandList().ImmediateFlush(EImmediateFlushType::DispatchToRHIThread); // we don't want stuff after the async cmd list to be bundled with it
 	}
@@ -1194,20 +1221,20 @@ void FRHICommandListBase::QueueAsyncCommandListSubmit(FGraphEventRef& AnyThreadC
 {
 	check(IsInRenderingThread() && IsImmediate());
 
-	if (GRHIThread)
+	if (IsRunningRHIInSeparateThread())
 	{
 		FRHICommandListExecutor::GetImmediateCommandList().ImmediateFlush(EImmediateFlushType::DispatchToRHIThread); // we should start on the stuff before this async list
 	}
 	if (AnyThreadCompletionEvent.GetReference())
 	{
-		if (GRHIThread)
+		if (IsRunningRHIInSeparateThread())
 		{
 			AllOutstandingTasks.Add(AnyThreadCompletionEvent);
 		}
 		WaitOutstandingTasks.Add(AnyThreadCompletionEvent);
 	}
 	new (AllocCommand<FRHICommandWaitForAndSubmitSubList>()) FRHICommandWaitForAndSubmitSubList(AnyThreadCompletionEvent, CmdList);
-	if (GRHIThread)
+	if (IsRunningRHIInSeparateThread())
 	{
 		FRHICommandListExecutor::GetImmediateCommandList().ImmediateFlush(EImmediateFlushType::DispatchToRHIThread); // we don't want stuff after the async cmd list to be bundled with it
 	}
@@ -1233,7 +1260,7 @@ struct FRHICommandWaitForAndSubmitRTSubList : public FRHICommand<FRHICommandWait
 			if (EventToWaitFor.GetReference() && !EventToWaitFor->IsComplete())
 			{
 			SCOPE_CYCLE_COUNTER(STAT_RTChainWait);
-				check(!GRHIThread || !IsInRHIThread()); // things should not be dispatched if they can't complete without further waits
+				check(!IsRunningRHIInSeparateThread() || !IsInRHIThread()); // things should not be dispatched if they can't complete without further waits
 				if (IsInRenderingThread())
 				{
 					if (FTaskGraphInterface::Get().IsThreadProcessingTasks(ENamedThreads::RenderThread_Local))
@@ -1309,7 +1336,7 @@ void FRHICommandList::BeginScene()
 		return;
 	}
 	new (AllocCommand<FRHICommandBeginScene>()) FRHICommandBeginScene();
-	if (!GRHIThread)
+	if (!IsRunningRHIInSeparateThread())
 	{
 		// if we aren't running an RHIThread, there is no good reason to buffer this frame advance stuff and that complicates state management, so flush everything out now
 		QUICK_SCOPE_CYCLE_COUNTER(BeginScene_Flush);
@@ -1325,7 +1352,7 @@ void FRHICommandList::EndScene()
 		return;
 	}
 	new (AllocCommand<FRHICommandEndScene>()) FRHICommandEndScene();
-	if (!GRHIThread)
+	if (!IsRunningRHIInSeparateThread())
 	{
 		// if we aren't running an RHIThread, there is no good reason to buffer this frame advance stuff and that complicates state management, so flush everything out now
 		QUICK_SCOPE_CYCLE_COUNTER(EndScene_Flush);
@@ -1343,7 +1370,7 @@ void FRHICommandList::BeginDrawingViewport(FViewportRHIParamRef Viewport, FTextu
 		return;
 	}
 	new (AllocCommand<FRHICommandBeginDrawingViewport>()) FRHICommandBeginDrawingViewport(Viewport, RenderTargetRHI);
-	if (!GRHIThread)
+	if (!IsRunningRHIInSeparateThread())
 	{
 		// if we aren't running an RHIThread, there is no good reason to buffer this frame advance stuff and that complicates state management, so flush everything out now
 		QUICK_SCOPE_CYCLE_COUNTER(BeginDrawingViewport_Flush);
@@ -1362,7 +1389,7 @@ void FRHICommandList::EndDrawingViewport(FViewportRHIParamRef Viewport, bool bPr
 	{
 		new (AllocCommand<FRHICommandEndDrawingViewport>()) FRHICommandEndDrawingViewport(Viewport, bPresent, bLockToVsync);
 
-		if ( GRHIThread )
+		if (IsRunningRHIInSeparateThread())
 		{
 			// Insert a fence to prevent the renderthread getting more than a frame ahead of the RHIThread
 			GRHIThreadEndDrawingViewportFences[GRHIThreadEndDrawingViewportFenceIndex] = static_cast<FRHICommandListImmediate*>(this)->RHIThreadFence();
@@ -1374,7 +1401,7 @@ void FRHICommandList::EndDrawingViewport(FViewportRHIParamRef Viewport, bool bPr
 		}
 	}
 
-	if ( GRHIThread )
+	if (IsRunningRHIInSeparateThread())
 	{
 		// Wait on the previous frame's RHI thread fence (we never want the rendering thread to get more than a frame ahead)
 		uint32 PreviousFrameFenceIndex = 1 - GRHIThreadEndDrawingViewportFenceIndex;
@@ -1396,7 +1423,7 @@ void FRHICommandList::BeginFrame()
 		return;
 	}
 	new (AllocCommand<FRHICommandBeginFrame>()) FRHICommandBeginFrame();
-	if (!GRHIThread)
+	if (!IsRunningRHIInSeparateThread())
 	{
 		// if we aren't running an RHIThread, there is no good reason to buffer this frame advance stuff and that complicates state management, so flush everything out now
 		QUICK_SCOPE_CYCLE_COUNTER(BeginFrame_Flush);
@@ -1413,7 +1440,7 @@ void FRHICommandList::EndFrame()
 		return;
 	}
 	new (AllocCommand<FRHICommandEndFrame>()) FRHICommandEndFrame();
-	if (!GRHIThread)
+	if (!IsRunningRHIInSeparateThread())
 	{
 		// if we aren't running an RHIThread, there is no good reason to buffer this frame advance stuff and that complicates state management, so flush everything out now
 		QUICK_SCOPE_CYCLE_COUNTER(EndFrame_Flush);
@@ -1453,7 +1480,7 @@ FScopedCommandListWaitForTasks::~FScopedCommandListWaitForTasks()
 	check(IsInRenderingThread());
 	if (bWaitForTasks)
 	{
-		if (GRHIThread)
+		if (IsRunningRHIInSeparateThread())
 		{
 #if 0
 			{
@@ -1502,42 +1529,8 @@ void FRHICommandListBase::WaitForDispatch()
 DECLARE_CYCLE_STAT(TEXT("Explicit wait for RHI thread"), STAT_ExplicitWaitRHIThread, STATGROUP_RHICMDLIST);
 DECLARE_CYCLE_STAT(TEXT("Explicit wait for RHI thread async dispatch"), STAT_ExplicitWaitRHIThread_Dispatch, STATGROUP_RHICMDLIST);
 DECLARE_CYCLE_STAT(TEXT("Deep spin for stray resource init"), STAT_SpinWaitRHIThread, STATGROUP_RHICMDLIST);
-DECLARE_CYCLE_STAT(TEXT("Spin RHIThread unstall"), STAT_SpinWaitRHIThreadUnstall, STATGROUP_RHICMDLIST);
 DECLARE_CYCLE_STAT(TEXT("Spin RHIThread wait for stall"), STAT_SpinWaitRHIThreadStall, STATGROUP_RHICMDLIST);
 
-static FGraphEventRef GRHIThreadStallTask;
-static FEvent* GRHIThreadStallEvent = nullptr;
-
-class FStallRHIThreadTask
-{
-	FEvent* RHIThreadStallEvent;
-public:
-
-	FStallRHIThreadTask(FEvent* InRHIThreadStallEvent)
-		: RHIThreadStallEvent(InRHIThreadStallEvent)
-	{
-	}
-
-	FORCEINLINE TStatId GetStatId() const
-	{
-		RETURN_QUICK_DECLARE_CYCLE_STAT(STAT_STallRHIThread, STATGROUP_TaskGraphTasks);
-	}
-
-	ENamedThreads::Type GetDesiredThread()
-	{
-		check(GRHIThread); // this should never be used on a platform that doesn't support the RHI thread
-		return ENamedThreads::RHIThread;
-	}
-
-	static ESubsequentsMode::Type GetSubsequentsMode() { return ESubsequentsMode::TrackSubsequents; }
-
-	void DoTask(ENamedThreads::Type CurrentThread, const FGraphEventRef& MyCompletionGraphEvent)
-	{
-		GRHIThreadStallTask = MyCompletionGraphEvent;
-		FPlatformMisc::MemoryBarrier();
-		RHIThreadStallEvent->Wait();
-	}
-};
 
 int32 StallCount = 0;
 bool FRHICommandListImmediate::IsStalled()
@@ -1547,9 +1540,7 @@ bool FRHICommandListImmediate::IsStalled()
 
 bool FRHICommandListImmediate::StallRHIThread()
 {
-	FPlatformAtomics::InterlockedIncrement(&StallCount);
-
-	check(IsInRenderingThread() && GRHIThread && !GRHIThreadStallTask.GetReference());
+	check(IsInRenderingThread() && IsRunningRHIInSeparateThread());
 	bool bAsyncSubmit = CVarRHICmdAsyncRHIThreadDispatch.GetValueOnRenderThread() > 0;
 	if (bAsyncSubmit)
 	{
@@ -1568,15 +1559,10 @@ bool FRHICommandListImmediate::StallRHIThread()
 				return false;
 			}
 		}
-		if (!GRHIThreadStallEvent)
-		{
-			GRHIThreadStallEvent = FPlatformProcess::GetSynchEventFromPool();
-		}
-		TGraphTask<FStallRHIThreadTask>::CreateTask(nullptr, ENamedThreads::RenderThread).ConstructAndDispatchWhenReady(GRHIThreadStallEvent);
-		SCOPE_CYCLE_COUNTER(STAT_SpinWaitRHIThreadStall);
-		while (!GRHIThreadStallTask.GetReference())
+		FPlatformAtomics::InterlockedIncrement(&StallCount);
 		{
-			FPlatformProcess::SleepNoStats(0);
+			SCOPE_CYCLE_COUNTER(STAT_SpinWaitRHIThreadStall);
+			GRHIThreadOnTasksCritical.Lock();
 		}
 		return true;
 	}
@@ -1589,18 +1575,11 @@ bool FRHICommandListImmediate::StallRHIThread()
 
 void FRHICommandListImmediate::UnStallRHIThread()
 {
-	check(IsInRenderingThread() && GRHIThread && GRHIThreadStallTask.GetReference() && !GRHIThreadStallTask->IsComplete() && GRHIThreadStallEvent);
-	GRHIThreadStallEvent->Trigger();
-	SCOPE_CYCLE_COUNTER(STAT_SpinWaitRHIThreadUnstall);
-	while (!GRHIThreadStallTask->IsComplete())
-	{
-		FPlatformProcess::SleepNoStats(0);
-	}
-	GRHIThreadStallTask = nullptr;
+	check(IsInRenderingThread() && IsRunningRHIInSeparateThread());
+	GRHIThreadOnTasksCritical.Unlock();
 	FPlatformAtomics::InterlockedDecrement(&StallCount);
 }
 
-
 void FRHICommandListBase::WaitForRHIThreadTasks()
 {
 	check(IsImmediate() && IsInRenderingThread());
@@ -1858,7 +1837,7 @@ void* FDynamicRHI::LockVertexBuffer_RenderThread(class FRHICommandListImmediate&
 	check(IsInRenderingThread());
 	bool bBuffer = CVarRHICmdBufferWriteLocks.GetValueOnRenderThread() > 0;
 	void* Result;
-	if (!bBuffer || LockMode != RLM_WriteOnly || RHICmdList.Bypass() || !GRHIThread)
+	if (!bBuffer || LockMode != RLM_WriteOnly || RHICmdList.Bypass() || !IsRunningRHIInSeparateThread())
 	{
 		QUICK_SCOPE_CYCLE_COUNTER(STAT_RHIMETHOD_LockVertexBuffer_Flush);
 		RHICmdList.ImmediateFlush(EImmediateFlushType::FlushRHIThread); 
@@ -1880,7 +1859,7 @@ void FDynamicRHI::UnlockVertexBuffer_RenderThread(class FRHICommandListImmediate
 	check(IsInRenderingThread());
 	bool bBuffer = CVarRHICmdBufferWriteLocks.GetValueOnRenderThread() > 0;
 	FLockTracker::FLockParams Params = GLockTracker.Unlock(VertexBuffer);
-	if (!bBuffer || Params.LockMode != RLM_WriteOnly || RHICmdList.Bypass() || !GRHIThread)
+	if (!bBuffer || Params.LockMode != RLM_WriteOnly || RHICmdList.Bypass() || !IsRunningRHIInSeparateThread())
 	{
 		QUICK_SCOPE_CYCLE_COUNTER(STAT_RHIMETHOD_UnlockVertexBuffer_Flush);
 		RHICmdList.ImmediateFlush(EImmediateFlushType::FlushRHIThread); 
@@ -1907,7 +1886,7 @@ void* FDynamicRHI::LockIndexBuffer_RenderThread(class FRHICommandListImmediate&
 	check(IsInRenderingThread());
 	bool bBuffer = CVarRHICmdBufferWriteLocks.GetValueOnRenderThread() > 0;
 	void* Result;
-	if (!bBuffer || LockMode != RLM_WriteOnly || RHICmdList.Bypass() || !GRHIThread)
+	if (!bBuffer || LockMode != RLM_WriteOnly || RHICmdList.Bypass() || !IsRunningRHIInSeparateThread())
 	{
 		QUICK_SCOPE_CYCLE_COUNTER(STAT_RHIMETHOD_LockIndexBuffer_Flush);
 		RHICmdList.ImmediateFlush(EImmediateFlushType::FlushRHIThread); 
@@ -1929,7 +1908,7 @@ void FDynamicRHI::UnlockIndexBuffer_RenderThread(class FRHICommandListImmediate&
 	check(IsInRenderingThread());
 	bool bBuffer = CVarRHICmdBufferWriteLocks.GetValueOnRenderThread() > 0;
 	FLockTracker::FLockParams Params = GLockTracker.Unlock(IndexBuffer);
-	if (!bBuffer || Params.LockMode != RLM_WriteOnly || RHICmdList.Bypass() || !GRHIThread)
+	if (!bBuffer || Params.LockMode != RLM_WriteOnly || RHICmdList.Bypass() || !IsRunningRHIInSeparateThread())
 	{
 		QUICK_SCOPE_CYCLE_COUNTER(STAT_RHIMETHOD_UnlockIndexBuffer_Flush);
 		RHICmdList.ImmediateFlush(EImmediateFlushType::FlushRHIThread); 
@@ -2208,7 +2187,7 @@ FRenderQueryRHIRef FDynamicRHI::RHICreateRenderQuery_RenderThread(class FRHIComm
 
 void FRHICommandListImmediate::UpdateTextureReference(FTextureReferenceRHIParamRef TextureRef, FTextureRHIParamRef NewTexture)
 {
-	if (Bypass() || !GRHIThread || CVarRHICmdFlushUpdateTextureReference.GetValueOnRenderThread() > 0)
+	if (Bypass() || !IsRunningRHIInSeparateThread() || CVarRHICmdFlushUpdateTextureReference.GetValueOnRenderThread() > 0)
 	{
 		{
 			QUICK_SCOPE_CYCLE_COUNTER(STAT_RHIMETHOD_UpdateTextureReference_FlushRHI);
diff --git a/Engine/Source/Runtime/RHI/Public/RHI.h b/Engine/Source/Runtime/RHI/Public/RHI.h
--- a/Engine/Source/Runtime/RHI/Public/RHI.h
+++ b/Engine/Source/Runtime/RHI/Public/RHI.h
@@ -296,6 +296,8 @@ Requirements for RHI thread
 * BeginDrawingViewport, and 5 or so other frame advance methods are queued with an RHIThread. Without an RHIThread, these just flush internally.
 ***/
 extern RHI_API bool GRHISupportsRHIThread;
+/* as above, but we run the commands on arbitrary task threads */
+extern RHI_API bool GRHISupportsRHIOnTaskThread;
 
 /** Whether or not the RHI supports parallel RHIThread executes / translates
 Requirements:
@@ -935,7 +937,7 @@ struct FClearValueBinding
 	static RHI_API const FClearValueBinding DepthNear;
 	static RHI_API const FClearValueBinding DepthFar;	
 	static RHI_API const FClearValueBinding Green;
-	static RHI_API const FClearValueBinding MidGray;
+	static RHI_API const FClearValueBinding DefaultNormal8Bit;
 };
 
 struct FRHIResourceCreateInfo
diff --git a/Engine/Source/Runtime/RHI/Public/RHICommandList.h b/Engine/Source/Runtime/RHI/Public/RHICommandList.h
--- a/Engine/Source/Runtime/RHI/Public/RHICommandList.h
+++ b/Engine/Source/Runtime/RHI/Public/RHICommandList.h
@@ -8,7 +8,6 @@
 #include "CoreTypes.h"
 #include "Misc/AssertionMacros.h"
 #include "HAL/UnrealMemory.h"
-#include "Templates/AlignOf.h"
 #include "Templates/UnrealTemplate.h"
 #include "Math/Color.h"
 #include "Math/IntPoint.h"
@@ -57,6 +56,7 @@ class FGraphicsPipelineState;
 
 DECLARE_STATS_GROUP(TEXT("RHICmdList"), STATGROUP_RHICMDLIST, STATCAT_Advanced);
 
+
 // set this one to get a stat for each RHI command 
 #define RHI_STATS 0
 
@@ -66,7 +66,40 @@ DECLARE_STATS_GROUP(TEXT("RHICommands"),STATGROUP_RHI_COMMANDS, STATCAT_Advanced
 #else
 #define RHISTAT(Method)
 #endif
-/** Thread used for RHI */
+
+extern RHI_API bool GUseRHIThread_InternalUseOnly;
+extern RHI_API bool GUseRHITaskThreads_InternalUseOnly;
+extern RHI_API bool GIsRunningRHIInSeparateThread_InternalUseOnly;
+extern RHI_API bool GIsRunningRHIInDedicatedThread_InternalUseOnly;
+extern RHI_API bool GIsRunningRHIInTaskThread_InternalUseOnly;
+
+/**
+* Whether the RHI commands are being run in a thread other than the render thread
+*/
+bool FORCEINLINE IsRunningRHIInSeparateThread()
+{
+	return GIsRunningRHIInSeparateThread_InternalUseOnly;
+}
+
+/**
+* Whether the RHI commands are being run on a dedicated thread other than the render thread
+*/
+bool FORCEINLINE IsRunningRHIInDedicatedThread()
+{
+	return GIsRunningRHIInDedicatedThread_InternalUseOnly;
+}
+
+/**
+* Whether the RHI commands are being run on a dedicated thread other than the render thread
+*/
+bool FORCEINLINE IsRunningRHIInTaskThread()
+{
+	return GIsRunningRHIInTaskThread_InternalUseOnly;
+}
+
+
+
+
 extern RHI_API bool GEnableAsyncCompute;
 extern RHI_API TAutoConsoleVariable<int32> CVarRHICmdWidth;
 extern RHI_API TAutoConsoleVariable<int32> CVarRHICmdFlushRenderThreadTasks;
@@ -187,7 +220,7 @@ public:
 	template <typename T>
 	FORCEINLINE_DEBUGGABLE void* Alloc()
 	{
-		return Alloc(sizeof(T), ALIGNOF(T));
+		return Alloc(sizeof(T), alignof(T));
 	}
 
 	template <typename TCmd>
diff --git a/Engine/Source/Runtime/RHI/Public/RHICommandList.inl b/Engine/Source/Runtime/RHI/Public/RHICommandList.inl
--- a/Engine/Source/Runtime/RHI/Public/RHICommandList.inl
+++ b/Engine/Source/Runtime/RHI/Public/RHICommandList.inl
@@ -41,7 +41,7 @@ FORCEINLINE_DEBUGGABLE bool FRHICommandListBase::Bypass()
 FORCEINLINE_DEBUGGABLE FScopedRHIThreadStaller::FScopedRHIThreadStaller(class FRHICommandListImmediate& InImmed)
 	: Immed(nullptr)
 {
-	if (GRHIThread)
+	if (IsRunningRHIInSeparateThread())
 	{
 		check(IsInRenderingThread());
 		if (InImmed.StallRHIThread())
@@ -93,7 +93,7 @@ FORCEINLINE_DEBUGGABLE void FRHICommandListImmediate::ImmediateFlush(EImmediateF
 				GRHICommandList.ExecuteList(*this);
 			}
 			WaitForDispatch();
-			if (GRHIThread)
+			if (IsRunningRHIInSeparateThread())
 			{
 				WaitForRHIThreadTasks();
 			}
diff --git a/Engine/Source/Runtime/RHI/Public/RHIDefinitions.h b/Engine/Source/Runtime/RHI/Public/RHIDefinitions.h
--- a/Engine/Source/Runtime/RHI/Public/RHIDefinitions.h
+++ b/Engine/Source/Runtime/RHI/Public/RHIDefinitions.h
@@ -754,6 +754,11 @@ inline bool IsConsolePlatform(const EShaderPlatform Platform)
 	return Platform == SP_PS4 || Platform == SP_XBOXONE_D3D12 || Platform == SP_XBOXONE_D3D11;
 }
 
+inline bool IsSwitchPlatform(const EShaderPlatform Platform)
+{
+	return Platform == SP_SWITCH || Platform == SP_SWITCH_FORWARD;
+}
+
 inline bool IsVulkanPlatform(const EShaderPlatform Platform)
 {
 	return Platform == SP_VULKAN_SM5 || Platform == SP_VULKAN_SM4 || Platform == SP_VULKAN_PCES3_1 || Platform == SP_VULKAN_ES3_1_ANDROID;
@@ -788,6 +793,11 @@ inline bool IsD3DPlatform(const EShaderPlatform Platform, bool bIncludeXboxOne)
 	return false;
 }
 
+inline bool IsHlslccShaderPlatform(const EShaderPlatform Platform)
+{
+	return IsMetalPlatform(Platform) || IsVulkanPlatform(Platform) || IsSwitchPlatform(Platform) || IsOpenGLPlatform(Platform);
+}
+
 inline ERHIFeatureLevel::Type GetMaxSupportedFeatureLevel(EShaderPlatform InShaderPlatform)
 {
 	switch (InShaderPlatform)
@@ -807,7 +817,7 @@ inline ERHIFeatureLevel::Type GetMaxSupportedFeatureLevel(EShaderPlatform InShad
 	case SP_OPENGL_SM4:
 	case SP_OPENGL_SM4_MAC:
 	case SP_METAL_MRT:
-    case SP_METAL_MRT_MAC:
+	case SP_METAL_MRT_MAC:
 	case SP_METAL_SM4:
 		return ERHIFeatureLevel::SM4;
 	case SP_PCD3D_ES2:
@@ -850,7 +860,7 @@ inline bool RHINeedsToSwitchVerticalAxis(EShaderPlatform Platform)
 
 	// ES2 & ES3.1 need to flip when rendering to an RT that will be post processed
 	return IsOpenGLPlatform(Platform) && IsMobilePlatform(Platform) && !IsPCPlatform(Platform) && Platform != SP_METAL && !IsVulkanPlatform(Platform)
-	       && Platform != SP_SWITCH && Platform != SP_SWITCH_FORWARD;
+		   && Platform != SP_SWITCH && Platform != SP_SWITCH_FORWARD;
 }
 
 inline bool RHISupportsSeparateMSAAAndResolveTextures(const EShaderPlatform Platform)
diff --git a/Engine/Source/Runtime/RenderCore/Private/RenderUtils.cpp b/Engine/Source/Runtime/RenderCore/Private/RenderUtils.cpp
--- a/Engine/Source/Runtime/RenderCore/Private/RenderUtils.cpp
+++ b/Engine/Source/Runtime/RenderCore/Private/RenderUtils.cpp
@@ -829,13 +829,15 @@ RENDERCORE_API bool IsSimpleForwardShadingEnabled(EShaderPlatform Platform)
 	return CVar->GetValueOnAnyThread() != 0 && PlatformSupportsSimpleForwardShading(Platform);
 }
 
-RENDERCORE_API bool IsForwardShadingEnabled(ERHIFeatureLevel::Type FeatureLevel)
-{
-	static const auto CVar = IConsoleManager::Get().FindTConsoleVariableDataInt(TEXT("r.ForwardShading"));
-	return CVar->GetValueOnAnyThread() != 0 
-		// Culling uses compute shader
-		&& FeatureLevel >= ERHIFeatureLevel::SM5;
-}
+RENDERCORE_API int32 bUseForwardShading = 0;
+static FAutoConsoleVariableRef CVarForwardShading(
+	TEXT("r.ForwardShading"),
+	bUseForwardShading,
+	TEXT("Whether to use forward shading on desktop platforms - requires Shader Model 5 hardware.\n")
+	TEXT("Forward shading has lower constant cost, but fewer features supported. 0:off, 1:on\n")
+	TEXT("This rendering path is a work in progress with many unimplemented features, notably only a single reflection capture is applied per object and no translucency dynamic shadow receiving."),
+	ECVF_RenderThreadSafe | ECVF_ReadOnly
+	); 
 
 class FUnitCubeVertexBuffer : public FVertexBuffer
 {
diff --git a/Engine/Source/Runtime/RenderCore/Private/RenderingThread.cpp b/Engine/Source/Runtime/RenderCore/Private/RenderingThread.cpp
--- a/Engine/Source/Runtime/RenderCore/Private/RenderingThread.cpp
+++ b/Engine/Source/Runtime/RenderCore/Private/RenderingThread.cpp
@@ -20,6 +20,7 @@
 #include "Stats/StatsData.h"
 #include "HAL/ThreadHeartBeat.h"
 #include "RenderResource.h"
+#include "ScopeLock.h"
 
 //
 // Globals
@@ -29,7 +30,6 @@
 
 RENDERCORE_API bool GIsThreadedRendering = false;
 RENDERCORE_API bool GUseThreadedRendering = false;
-RENDERCORE_API bool GUseRHIThread = false;
 
 #if !(UE_BUILD_SHIPPING || UE_BUILD_TEST)
 	RENDERCORE_API bool GMainThreadBlockedOnRenderThread = false;
@@ -413,7 +413,7 @@ public:
 		GRenderThreadId = FPlatformTLS::GetCurrentThreadId();
 
 		// Acquire rendering context ownership on the current thread, unless using an RHI thread, which will be the real owner
-		if (!GUseRHIThread)
+		if (!IsRunningRHIInSeparateThread())
 		{
 			bAcquiredThreadOwnership = true;
 			RHIAcquireThreadOwnership();
@@ -632,7 +632,7 @@ public:
 	**/
 	void DoTask(ENamedThreads::Type CurrentThread, const FGraphEventRef& MyCompletionGraphEvent)
 	{
-		// note that this task is the first task run on the thread, before GRHIThread is assigned, so we can't check IsInRHIThread()
+		// note that this task is the first task run on the thread, before GRHIThread_InternalUseOnly is assigned, so we can't check IsInRHIThread()
 
 		if (bAcquireOwnership)
 		{
@@ -655,8 +655,9 @@ void StartRenderingThread()
 	static uint32 ThreadCount = 0;
 	check(!GIsThreadedRendering && GUseThreadedRendering);
 
-	check(!GRHIThread)
-	if (GUseRHIThread)
+	check(!GRHIThread_InternalUseOnly && !GIsRunningRHIInSeparateThread_InternalUseOnly && !GIsRunningRHIInDedicatedThread_InternalUseOnly && !GIsRunningRHIInTaskThread_InternalUseOnly);
+
+	if (GUseRHIThread_InternalUseOnly)
 	{
 		FRHICommandListExecutor::GetImmediateCommandList().ImmediateFlush(EImmediateFlushType::DispatchToRHIThread);		
 		if (!FTaskGraphInterface::Get().IsThreadProcessingTasks(ENamedThreads::RHIThread))
@@ -668,9 +669,17 @@ void StartRenderingThread()
 		FGraphEventRef CompletionEvent = TGraphTask<FOwnershipOfRHIThreadTask>::CreateTask(NULL, ENamedThreads::GameThread).ConstructAndDispatchWhenReady(true, GET_STATID(STAT_WaitForRHIThread));
 		QUICK_SCOPE_CYCLE_COUNTER(STAT_StartRenderingThread);
 		FTaskGraphInterface::Get().WaitUntilTaskCompletes(CompletionEvent, ENamedThreads::GameThread_Local);
-		GRHIThread = FRHIThread::Get().Thread;
-		check(GRHIThread);
+		GRHIThread_InternalUseOnly = FRHIThread::Get().Thread;
+		check(GRHIThread_InternalUseOnly);
 		GRHICommandList.LatchBypass();
+		GIsRunningRHIInDedicatedThread_InternalUseOnly = true;
+		GIsRunningRHIInSeparateThread_InternalUseOnly = true;
+		GRHIThreadId = GRHIThread_InternalUseOnly->GetThreadID();
+	}
+	else if (GUseRHITaskThreads_InternalUseOnly)
+	{
+		GIsRunningRHIInSeparateThread_InternalUseOnly = true;
+		GIsRunningRHIInTaskThread_InternalUseOnly = true;
 	}
 
 	// Turn on the threaded rendering flag.
@@ -736,15 +745,21 @@ void StopRenderingThread()
 		// The rendering thread may have already been stopped during the call to GFlushStreamingFunc or FlushRenderingCommands.
 		if ( GIsThreadedRendering )
 		{
-			if (GRHIThread)
+			if (GRHIThread_InternalUseOnly)
 			{
 				DECLARE_CYCLE_STAT(TEXT("Wait For RHIThread Finish"), STAT_WaitForRHIThreadFinish, STATGROUP_TaskGraphTasks);
 				FGraphEventRef ReleaseTask = TGraphTask<FOwnershipOfRHIThreadTask>::CreateTask(NULL, ENamedThreads::GameThread).ConstructAndDispatchWhenReady(false, GET_STATID(STAT_WaitForRHIThreadFinish));
 				QUICK_SCOPE_CYCLE_COUNTER(STAT_StopRenderingThread_RHIThread);
 				FTaskGraphInterface::Get().WaitUntilTaskCompletes(ReleaseTask, ENamedThreads::GameThread_Local);
-				GRHIThread = nullptr;
+				GRHIThread_InternalUseOnly = nullptr;
+				GRHIThreadId = 0;
 			}
 
+			GIsRunningRHIInSeparateThread_InternalUseOnly = false;
+			GIsRunningRHIInDedicatedThread_InternalUseOnly = false;
+			GIsRunningRHIInTaskThread_InternalUseOnly = false;
+
+
 			check( GRenderingThread );
 			check(!GIsRenderingThreadSuspended);
 
@@ -787,7 +802,7 @@ void StopRenderingThread()
 		delete PendingCleanupObjects;
 	}
 
-	check(!GRHIThread);
+	check(!GRHIThread_InternalUseOnly);
 }
 
 void CheckRenderingThreadHealth()
@@ -1077,7 +1092,7 @@ void FlushRenderingCommands()
 
 void FlushPendingDeleteRHIResources_GameThread()
 {
-	if (!GRHIThread)
+	if (!IsRunningRHIInSeparateThread())
 	{
 		ENQUEUE_UNIQUE_RENDER_COMMAND(
 			FlushPendingDeleteRHIResources,
@@ -1089,7 +1104,7 @@ void FlushPendingDeleteRHIResources_GameThread()
 }
 void FlushPendingDeleteRHIResources_RenderThread()
 {
-	if (!GRHIThread)
+	if (!IsRunningRHIInSeparateThread())
 	{
 		FRHIResource::FlushPendingDeletes();
 	}
@@ -1169,47 +1184,73 @@ FPendingCleanupObjects* GetPendingCleanupObjects()
 	return new FPendingCleanupObjects;
 }
 
-void SetRHIThreadEnabled(bool bEnable)
+void SetRHIThreadEnabled(bool bEnableDedicatedThread, bool bEnableRHIOnTaskThreads)
 {
-	if (bEnable != GUseRHIThread)
+	if (bEnableDedicatedThread != GUseRHIThread_InternalUseOnly || bEnableRHIOnTaskThreads != GUseRHITaskThreads_InternalUseOnly)
 	{
-		if (GRHISupportsRHIThread)
+		if ((bEnableRHIOnTaskThreads || bEnableDedicatedThread) && !GIsThreadedRendering)
 		{
-			if (!GIsThreadedRendering)
+			check(!IsRunningRHIInSeparateThread());
+			UE_LOG(LogConsoleResponse, Display, TEXT("Can't switch to RHI thread mode when we are not running a multithreaded renderer."));
+		}
+		else
+		{
+			StopRenderingThread();
+			if (bEnableRHIOnTaskThreads)
 			{
-				check(!GRHIThread);
-				UE_LOG(LogRendererCore, Display, TEXT("Can't switch to RHI thread mode when we are not running a multithreaded renderer."));
+				GUseRHIThread_InternalUseOnly = false;
+				GUseRHITaskThreads_InternalUseOnly = true;
+			}
+			else if (bEnableDedicatedThread)
+			{
+				GUseRHIThread_InternalUseOnly = true;
+				GUseRHITaskThreads_InternalUseOnly = false;
 			}
 			else
 			{
-				StopRenderingThread();
-				GUseRHIThread = bEnable;
-				StartRenderingThread();
+				GUseRHIThread_InternalUseOnly = false;
+				GUseRHITaskThreads_InternalUseOnly = false;
 			}
-			UE_LOG(LogRendererCore, Display, TEXT("RHIThread is now %s."), GRHIThread ? TEXT("active") : TEXT("inactive"));
+			StartRenderingThread();
+		}
+	}
+	if (IsRunningRHIInSeparateThread())
+	{
+		if (IsRunningRHIInDedicatedThread())
+		{
+			UE_LOG(LogConsoleResponse, Display, TEXT("RHIThread is now running on a dedicated thread."));
 		}
 		else
 		{
-			UE_LOG(LogRendererCore, Display, TEXT("This RHI does not support the RHI thread."));
+			check(IsRunningRHIInTaskThread());
+			UE_LOG(LogConsoleResponse, Display, TEXT("RHIThread is now running on task threads."));
 		}
 	}
+	else
+	{
+		check(!IsRunningRHIInTaskThread() && !IsRunningRHIInDedicatedThread());
+		UE_LOG(LogConsoleResponse, Display, TEXT("RHIThread is disabled."));
+	}
+
 }
 
 static void HandleRHIThreadEnableChanged(const TArray<FString>& Args)
 {
 	if (Args.Num() > 0)
 	{
-		const bool bUseRHIThread = Args[0].ToBool();
-		SetRHIThreadEnabled(bUseRHIThread);
+		const int32 UseRHIThread = FCString::Atoi(*Args[0]);
+		SetRHIThreadEnabled(UseRHIThread == 1, UseRHIThread == 2);
 	}
 	else
 	{
-		UE_LOG(LogRendererCore, Display, TEXT("Usage: r.RHIThread.Enable 0/1; Currently %d"), (int32)GUseRHIThread);
+		UE_LOG(LogConsoleResponse, Display, TEXT("Usage: r.RHIThread.Enable 0=off,  1=dedicated thread,  2=task threads; Currently %d"), IsRunningRHIInSeparateThread() ? (IsRunningRHIInDedicatedThread() ? 1 : 2) : 0);
 	}
 }
 
 static FAutoConsoleCommand CVarRHIThreadEnable(
 	TEXT("r.RHIThread.Enable"),
-	TEXT("Enables/disabled the RHI Thread\n"),	
+	TEXT("Enables/disabled the RHI Thread and determine if the RHI work runs on a dedicated thread or not.\n"),	
 	FConsoleCommandWithArgsDelegate::CreateStatic(&HandleRHIThreadEnableChanged)
 	);
+
+
diff --git a/Engine/Source/Runtime/RenderCore/Public/RenderUtils.h b/Engine/Source/Runtime/RenderCore/Public/RenderUtils.h
--- a/Engine/Source/Runtime/RenderCore/Public/RenderUtils.h
+++ b/Engine/Source/Runtime/RenderCore/Public/RenderUtils.h
@@ -316,7 +316,13 @@ RENDERCORE_API bool PlatformSupportsSimpleForwardShading(EShaderPlatform Platfor
 
 RENDERCORE_API bool IsSimpleForwardShadingEnabled(EShaderPlatform Platform);
 
-RENDERCORE_API bool IsForwardShadingEnabled(ERHIFeatureLevel::Type FeatureLevel);
+inline bool IsForwardShadingEnabled(ERHIFeatureLevel::Type FeatureLevel)
+{
+	extern RENDERCORE_API int32 bUseForwardShading;
+	return bUseForwardShading
+		// Culling uses compute shader
+		&& FeatureLevel >= ERHIFeatureLevel::SM5;
+}
 
 inline bool IsAnyForwardShadingEnabled(EShaderPlatform Platform)
 {
diff --git a/Engine/Source/Runtime/RenderCore/Public/RendererInterface.h b/Engine/Source/Runtime/RenderCore/Public/RendererInterface.h
--- a/Engine/Source/Runtime/RenderCore/Public/RendererInterface.h
+++ b/Engine/Source/Runtime/RenderCore/Public/RendererInterface.h
@@ -237,6 +237,11 @@ public:
 			&& ((TargetableFlags & TexCreate_UAV) == 0 || GMaxRHIFeatureLevel == ERHIFeatureLevel::SM5);
 	}
 
+	FIntVector GetSize() const
+	{
+		return FIntVector(Extent.X, Extent.Y, Depth);
+	}
+
 	/** 
 	 * for debugging purpose
 	 * @return e.g. (2D 128x64 PF_R8)
diff --git a/Engine/Source/Runtime/RenderCore/Public/RenderingThread.h b/Engine/Source/Runtime/RenderCore/Public/RenderingThread.h
--- a/Engine/Source/Runtime/RenderCore/Public/RenderingThread.h
+++ b/Engine/Source/Runtime/RenderCore/Public/RenderingThread.h
@@ -27,12 +27,8 @@ extern RENDERCORE_API bool GIsThreadedRendering;
  * Currently set by command line parameter and by the ToggleRenderingThread console command.
  */
 extern RENDERCORE_API bool GUseThreadedRendering;
-/**
- * Whether the RHI thread should be created or not, requires rendering thread
- * Currently set by command line parameter and by the r.rhithread.enable cvar.
- */
-extern RENDERCORE_API bool GUseRHIThread;
-extern RENDERCORE_API void SetRHIThreadEnabled(bool bEnable);
+
+extern RENDERCORE_API void SetRHIThreadEnabled(bool bEnableDedicatedThread, bool bEnableRHIOnTaskThreads);
 
 #if (UE_BUILD_SHIPPING || UE_BUILD_TEST)
 	static FORCEINLINE void CheckNotBlockedOnRenderThread() {}
diff --git a/Engine/Source/Runtime/Renderer/Private/CapsuleShadowRendering.cpp b/Engine/Source/Runtime/Renderer/Private/CapsuleShadowRendering.cpp
--- a/Engine/Source/Runtime/Renderer/Private/CapsuleShadowRendering.cpp
+++ b/Engine/Source/Runtime/Renderer/Private/CapsuleShadowRendering.cpp
@@ -771,7 +771,7 @@ bool FDeferredShadingSceneRenderer::RenderCapsuleDirectShadows(
 				
 				AllocateCapsuleTileIntersectionCountsBuffer(GroupSize, View.ViewState);
 
-				ClearUAV(RHICmdList, GMaxRHIFeatureLevel, View.ViewState->CapsuleTileIntersectionCountsBuffer, 0);
+				ClearUAV(RHICmdList, View.ViewState->CapsuleTileIntersectionCountsBuffer, 0);
 
 				{
 					SCOPED_DRAW_EVENT(RHICmdList, TiledCapsuleShadowing);
@@ -945,7 +945,7 @@ void FDeferredShadingSceneRenderer::CreateIndirectCapsuleShadows()
 	}
 }
 
-void FDeferredShadingSceneRenderer::SetupIndirectCapsuleShadows(FRHICommandListImmediate& RHICmdList, const FViewInfo& View, bool bPrepareLightData, int32& NumCapsuleShapes, int32& NumMeshDistanceFieldCasters) const
+void FDeferredShadingSceneRenderer::SetupIndirectCapsuleShadows(FRHICommandListImmediate& RHICmdList, const FViewInfo& View, bool bPrepareLightData, int32& NumCapsuleShapes, int32& NumMeshesWithCapsules, int32& NumMeshDistanceFieldCasters) const
 {
 	const float CosFadeStartAngle = FMath::Cos(GCapsuleShadowFadeAngleFromVertical);
 	const FSkyLightSceneProxy* SkyLight = Scene ? Scene->SkyLight : NULL;
@@ -1054,6 +1054,8 @@ void FDeferredShadingSceneRenderer::SetupIndirectCapsuleShadows(FRHICommandListI
 			{
 				DistanceFieldCasterLightSourceData.Add(PackedLightDirection);
 			}
+
+			NumMeshesWithCapsules++;
 		}
 	}
 
@@ -1198,8 +1200,9 @@ void FDeferredShadingSceneRenderer::RenderIndirectCapsuleShadows(
 					SCOPED_GPU_STAT(RHICmdList, Stat_GPU_CapsuleShadows);
 		
 					int32 NumCapsuleShapes = 0;
+					int32 NumMeshesWithCapsules = 0;
 					int32 NumMeshDistanceFieldCasters = 0;
-					SetupIndirectCapsuleShadows(RHICmdList, View, true, NumCapsuleShapes, NumMeshDistanceFieldCasters);
+					SetupIndirectCapsuleShadows(RHICmdList, View, true, NumCapsuleShapes, NumMeshesWithCapsules, NumMeshDistanceFieldCasters);
 
 					if (NumCapsuleShapes > 0 || NumMeshDistanceFieldCasters > 0)
 					{
@@ -1213,10 +1216,10 @@ void FDeferredShadingSceneRenderer::RenderIndirectCapsuleShadows(
 				
 						AllocateCapsuleTileIntersectionCountsBuffer(GroupSize, View.ViewState);
 
-						ClearUAV(RHICmdList, GMaxRHIFeatureLevel, View.ViewState->CapsuleTileIntersectionCountsBuffer, 0);
+						ClearUAV(RHICmdList, View.ViewState->CapsuleTileIntersectionCountsBuffer, 0);
 
 						{
-							SCOPED_DRAW_EVENT(RHICmdList, TiledCapsuleShadowing);
+							SCOPED_DRAW_EVENTF(RHICmdList, TiledCapsuleShadowing, TEXT("TiledCapsuleShadowing %u capsules among %u meshes"), NumCapsuleShapes, NumMeshesWithCapsules);
 
 							FSceneRenderTargetItem& RayTracedShadowsRTI = RayTracedShadowsRT->GetRenderTargetItem();
 							{
@@ -1406,8 +1409,9 @@ void FDeferredShadingSceneRenderer::RenderCapsuleShadowsForMovableSkylight(FRHIC
 					SCOPED_GPU_STAT(RHICmdList, Stat_GPU_CapsuleShadows);
 		
 					int32 NumCapsuleShapes = 0;
+					int32 NumMeshesWithCapsules = 0;
 					int32 NumMeshDistanceFieldCasters = 0;
-					SetupIndirectCapsuleShadows(RHICmdList, View, true, NumCapsuleShapes, NumMeshDistanceFieldCasters);
+					SetupIndirectCapsuleShadows(RHICmdList, View, true, NumCapsuleShapes, NumMeshesWithCapsules, NumMeshDistanceFieldCasters);
 
 					if (NumCapsuleShapes > 0 || NumMeshDistanceFieldCasters > 0)
 					{
@@ -1420,7 +1424,7 @@ void FDeferredShadingSceneRenderer::RenderCapsuleShadowsForMovableSkylight(FRHIC
 							uint32 GroupSizeY = FMath::DivideAndRoundUp(ScissorRect.Size().Y / GAODownsampleFactor, GShadowShapeTileSize);
 
 							{
-								SCOPED_DRAW_EVENT(RHICmdList, TiledCapsuleShadowing);
+								SCOPED_DRAW_EVENTF(RHICmdList, TiledCapsuleShadowing, TEXT("TiledCapsuleShadowing %u capsules among %u meshes"), NumCapsuleShapes, NumMeshesWithCapsules);
 
 								FSceneRenderTargetItem& RayTracedShadowsRTI = NewBentNormal->GetRenderTargetItem();
 								{
diff --git a/Engine/Source/Runtime/Renderer/Private/CompositionLighting/PostProcessDeferredDecals.cpp b/Engine/Source/Runtime/Renderer/Private/CompositionLighting/PostProcessDeferredDecals.cpp
--- a/Engine/Source/Runtime/Renderer/Private/CompositionLighting/PostProcessDeferredDecals.cpp
+++ b/Engine/Source/Runtime/Renderer/Private/CompositionLighting/PostProcessDeferredDecals.cpp
@@ -857,7 +857,7 @@ void FRCPassPostProcessDeferredDecals::Process(FRenderingCompositePassContext& C
 
 				// we don't modify stencil but if out input was having stencil for us (after base pass - we need to clear)
 				// Clear stencil to 0, which is the assumed default by other passes
-				DrawClearQuad(RHICmdList, SMFeatureLevel, false, FLinearColor(), false, 0, true, 0, SceneContext.GetSceneDepthSurface()->GetSizeXY(), FIntRect());
+				DrawClearQuad(RHICmdList, false, FLinearColor(), false, 0, true, 0, SceneContext.GetSceneDepthSurface()->GetSizeXY(), FIntRect());
 			}
 
 			if (CurrentStage == DRS_BeforeBasePass)
@@ -916,6 +916,11 @@ void FDecalRenderTargetManager::ResolveTargets()
 		TargetsToResolve[FDecalRenderTargetManager::GBufferAIndex] = SceneContext.GBufferA->GetRenderTargetItem().TargetableTexture;
 	}
 
+	//those have been cleared or rendered to and need to be resolved
+	TargetsToResolve[FDecalRenderTargetManager::DBufferAIndex] = SceneContext.DBufferA ? SceneContext.DBufferA->GetRenderTargetItem().TargetableTexture : nullptr;
+	TargetsToResolve[FDecalRenderTargetManager::DBufferBIndex] = SceneContext.DBufferB ? SceneContext.DBufferB->GetRenderTargetItem().TargetableTexture : nullptr;
+	TargetsToResolve[FDecalRenderTargetManager::DBufferCIndex] = SceneContext.DBufferC ? SceneContext.DBufferC->GetRenderTargetItem().TargetableTexture : nullptr;
+
 	// resolve the targets we wrote to.
 	FResolveParams ResolveParams;
 	for (int32 i = 0; i < ResolveBufferMax; ++i)
diff --git a/Engine/Source/Runtime/Renderer/Private/CompositionLighting/PostProcessMeshDecals.cpp b/Engine/Source/Runtime/Renderer/Private/CompositionLighting/PostProcessMeshDecals.cpp
--- a/Engine/Source/Runtime/Renderer/Private/CompositionLighting/PostProcessMeshDecals.cpp
+++ b/Engine/Source/Runtime/Renderer/Private/CompositionLighting/PostProcessMeshDecals.cpp
@@ -576,7 +576,7 @@ static void RenderPrimitive(FRenderingCompositePassContext& Context, FDecalDrawi
 						View,
 						DrawContext,
 						StaticMesh,
-						StaticMesh.Elements.Num() == 1 ? 1 : View.StaticMeshBatchVisibility[StaticMesh.Id],
+						StaticMesh.bRequiresPerElementVisibility ? View.StaticMeshBatchVisibility[StaticMesh.BatchVisibilityId] : ((1ull << StaticMesh.Elements.Num()) - 1),
 						false,
 						DrawRenderStateLocal,
 						PrimitiveSceneInfo->Proxy,
diff --git a/Engine/Source/Runtime/Renderer/Private/CustomDepthRendering.cpp b/Engine/Source/Runtime/Renderer/Private/CustomDepthRendering.cpp
--- a/Engine/Source/Runtime/Renderer/Private/CustomDepthRendering.cpp
+++ b/Engine/Source/Runtime/Renderer/Private/CustomDepthRendering.cpp
@@ -95,7 +95,7 @@ bool FCustomDepthPrimSet::DrawPrims(FRHICommandListImmediate& RHICmdList, const
 								View,
 								Context,
 								StaticMesh,
-								StaticMesh.bRequiresPerElementVisibility ? View.StaticMeshBatchVisibility[StaticMesh.Id] : ((1ull << StaticMesh.Elements.Num()) - 1),
+								StaticMesh.bRequiresPerElementVisibility ? View.StaticMeshBatchVisibility[StaticMesh.BatchVisibilityId] : ((1ull << StaticMesh.Elements.Num()) - 1),
 								true,
 								DrawRenderStateLocal2,
 								PrimitiveSceneProxy,
diff --git a/Engine/Source/Runtime/Renderer/Private/DeferredShadingRenderer.cpp b/Engine/Source/Runtime/Renderer/Private/DeferredShadingRenderer.cpp
--- a/Engine/Source/Runtime/Renderer/Private/DeferredShadingRenderer.cpp
+++ b/Engine/Source/Runtime/Renderer/Private/DeferredShadingRenderer.cpp
@@ -427,17 +427,9 @@ static void SetAndClearViewGBuffer(FRHICommandListImmediate& RHICmdList, FViewIn
 	FSceneRenderTargets::Get(RHICmdList).BeginRenderingGBuffer(RHICmdList, ERenderTargetLoadAction::EClear, DepthLoadAction, DepthStencilAccess, View.Family->EngineShowFlags.ShaderComplexity, ClearColor);
 }
 
-static TAutoConsoleVariable<int32> CVarOcclusionQueryLocation(
-	TEXT("r.OcclusionQueryLocation"),
-	0,
-	TEXT("Controls when occlusion queries are rendered.  Rendering before the base pass may give worse occlusion (because not all occluders generally render in the earlyzpass).  ")
-	TEXT("However, it may reduce CPU waiting for query result stalls on some platforms and increase overall performance.")
-	TEXT("0: After BasePass.")
-	TEXT("1: After EarlyZPass, but before BasePass."));
-
-void FDeferredShadingSceneRenderer::RenderOcclusion(FRHICommandListImmediate& RHICmdList, bool bRenderQueries, bool bRenderHZB)
+void FDeferredShadingSceneRenderer::RenderOcclusion(FRHICommandListImmediate& RHICmdList, bool bRenderQueries)
 {		
-	if (bRenderQueries || bRenderHZB)
+	if (bRenderQueries)
 	{
 		FSceneRenderTargets& SceneContext = FSceneRenderTargets::Get(RHICmdList);
 
@@ -454,7 +446,6 @@ void FDeferredShadingSceneRenderer::RenderOcclusion(FRHICommandListImmediate& RH
 		// This is done after the downsampled depth buffer is created so that it can be used for issuing queries
 		BeginOcclusionTests(RHICmdList, bRenderQueries);
 
-		if (bRenderHZB)
 		{
 			RHICmdList.TransitionResource( EResourceTransitionAccess::EReadable, SceneContext.GetSceneDepthSurface() );
 
@@ -494,7 +485,7 @@ void FDeferredShadingSceneRenderer::RenderOcclusion(FRHICommandListImmediate& RH
 		// for these query results on some platforms.
 		RHICmdList.SubmitCommandsHint();
 
-		if (bRenderQueries && GRHIThread)
+		if (bRenderQueries && IsRunningRHIInSeparateThread())
 		{
 			SCOPE_CYCLE_COUNTER(STAT_OcclusionSubmittedFence_Dispatch);
 			int32 NumFrames = FOcclusionQueryHelpers::GetNumBufferedFrames();
@@ -541,7 +532,7 @@ void FDeferredShadingSceneRenderer::Render(FRHICommandListImmediate& RHICmdList)
 
 	bool bDBuffer = IsDBufferEnabled();	
 
-	if (GRHIThread)
+	if (IsRunningRHIInSeparateThread())
 	{
 		SCOPE_CYCLE_COUNTER(STAT_OcclusionSubmittedFence_Wait);
 		int32 BlockFrame = FOcclusionQueryHelpers::GetNumBufferedFrames() - 1;
@@ -621,7 +612,7 @@ void FDeferredShadingSceneRenderer::Render(FRHICommandListImmediate& RHICmdList)
 		}	
 	}
 
-	if (GRHIThread)
+	if (IsRunningRHIInSeparateThread())
 	{
 		// we will probably stall on occlusion queries, so might as well have the RHI thread and GPU work while we wait.
 		SCOPE_CYCLE_COUNTER(STAT_PostInitViews_FlushDel);
@@ -791,11 +782,9 @@ void FDeferredShadingSceneRenderer::Render(FRHICommandListImmediate& RHICmdList)
 	}
 	SceneContext.AllocDummyGBufferTargets(RHICmdList);
 
-	//occlusion can't run before basepass if there's no prepass to fill in some depth to occlude against.
-	bool bOcclusionBeforeBasePass = ((CVarOcclusionQueryLocation.GetValueOnRenderThread() == 1) && bNeedsPrePass) || IsForwardShadingEnabled(FeatureLevel);
-	bool bHZBBeforeBasePass = bOcclusionBeforeBasePass && (EarlyZPassMode == EDepthDrawingMode::DDM_AllOccluders || EarlyZPassMode == EDepthDrawingMode::DDM_AllOpaque);
+	const bool bOcclusionBeforeBasePass = (EarlyZPassMode == EDepthDrawingMode::DDM_AllOccluders) || (EarlyZPassMode == EDepthDrawingMode::DDM_AllOpaque);
 
-	RenderOcclusion(RHICmdList, bOcclusionBeforeBasePass, bHZBBeforeBasePass);
+	RenderOcclusion(RHICmdList, bOcclusionBeforeBasePass);
 	ServiceLocalQueue();
 
 	if (bOcclusionBeforeBasePass)
@@ -919,7 +908,7 @@ void FDeferredShadingSceneRenderer::Render(FRHICommandListImmediate& RHICmdList)
 	{
 		// clear out emissive and baked lighting (not too efficient but simple and only needed for this debug view)
 		SceneContext.BeginRenderingSceneColor(RHICmdList);
-		DrawClearQuad(RHICmdList, GMaxRHIFeatureLevel, FLinearColor(0, 0, 0, 0));
+		DrawClearQuad(RHICmdList, FLinearColor(0, 0, 0, 0));
 	}
 
 	SceneContext.DBufferA.SafeRelease();
@@ -945,10 +934,15 @@ void FDeferredShadingSceneRenderer::Render(FRHICommandListImmediate& RHICmdList)
 	SceneContext.ResolveSceneDepthToAuxiliaryTexture(RHICmdList);
 
 	bool bOcclusionAfterBasePass = bIsOcclusionTesting && !bOcclusionBeforeBasePass;
-	bool bHZBAfterBasePass = !bHZBBeforeBasePass;
-	RenderOcclusion(RHICmdList, bOcclusionAfterBasePass, bHZBAfterBasePass);
+	RenderOcclusion(RHICmdList, bOcclusionAfterBasePass);
 	ServiceLocalQueue();
 
+	if (bUseGBuffer)
+	{
+		SCOPE_CYCLE_COUNTER(STAT_FDeferredShadingSceneRenderer_Resolve_After_Basepass);
+		SceneContext.FinishRenderingGBuffer(RHICmdList);
+	}
+
 	if (!bOcclusionBeforeBasePass)
 	{
 		RenderShadowDepthMaps(RHICmdList);
@@ -956,12 +950,6 @@ void FDeferredShadingSceneRenderer::Render(FRHICommandListImmediate& RHICmdList)
 		ServiceLocalQueue();
 	}
 
-	if (bUseGBuffer)
-	{
-		SCOPE_CYCLE_COUNTER(STAT_FDeferredShadingSceneRenderer_Resolve_After_Basepass);
-		SceneContext.FinishRenderingGBuffer(RHICmdList);
-	}
-
 	if(GetCustomDepthPassLocation() == 1)
 	{
 		QUICK_SCOPE_CYCLE_COUNTER(STAT_FDeferredShadingSceneRenderer_CustomDepthPass1);
@@ -1013,6 +1001,10 @@ void FDeferredShadingSceneRenderer::Render(FRHICommandListImmediate& RHICmdList)
 		SCOPE_CYCLE_COUNTER(STAT_FDeferredShadingSceneRenderer_AfterBasePass);
 
 		GRenderTargetPool.AddPhaseEvent(TEXT("AfterBasePass"));
+		if (!IsForwardShadingEnabled(FeatureLevel))
+		{
+			SceneContext.ResolveSceneDepthTexture(RHICmdList, FResolveRect(0, 0, ViewFamily.FamilySizeX, ViewFamily.FamilySizeY));
+		}
 
 		for (int32 ViewIndex = 0; ViewIndex < Views.Num(); ViewIndex++)
 		{
@@ -1184,7 +1176,19 @@ void FDeferredShadingSceneRenderer::Render(FRHICommandListImmediate& RHICmdList)
 		SCOPE_CYCLE_COUNTER(STAT_TranslucencyDrawTime);
 
 		RHICmdList.SetCurrentStat(GET_STATID(STAT_CLM_Translucency));
-		RenderTranslucency(RHICmdList);
+
+		// For now there is only one resolve for all translucency passes. This can be changed by enabling the resolve in RenderTranslucency()
+		ConditionalResolveSceneColorForTranslucentMaterials(RHICmdList);
+		if (ViewFamily.AllowTranslucencyAfterDOF())
+		{
+			RenderTranslucency(RHICmdList, ETranslucencyPass::TPT_StandardTranslucency);
+			// Translucency after DOF is rendered now, but stored in the separate translucency RT for later use.
+			RenderTranslucency(RHICmdList, ETranslucencyPass::TPT_TranslucencyAfterDOF);
+		}
+		else // Otherwise render translucent primitives in a single bucket.
+		{
+			RenderTranslucency(RHICmdList, ETranslucencyPass::TPT_AllTranslucency);
+		}
 		ServiceLocalQueue();
 
 		static const auto DisableDistortionCVar = IConsoleManager::Get().FindTConsoleVariableDataInt(TEXT("r.DisableDistortion"));
@@ -1199,6 +1203,7 @@ void FDeferredShadingSceneRenderer::Render(FRHICommandListImmediate& RHICmdList)
 			RenderDistortion(RHICmdList);
 			ServiceLocalQueue();
 		}
+
 		RHICmdList.SetCurrentStat(GET_STATID(STAT_CLM_AfterTranslucency));
 	}
 
@@ -1266,7 +1271,7 @@ void FDeferredShadingSceneRenderer::Render(FRHICommandListImmediate& RHICmdList)
 		}
 
 		// End of frame, we don't need it anymore
-		FSceneRenderTargets::Get(RHICmdList).FreeSeparateTranslucencyDepth();
+		FSceneRenderTargets::Get(RHICmdList).FreeDownsampledTranslucencyDepth();
 
 		// we rendered to it during the frame, seems we haven't made use of it, because it should be released
 		check(!FSceneRenderTargets::Get(RHICmdList).SeparateTranslucencyRT);
diff --git a/Engine/Source/Runtime/Renderer/Private/DeferredShadingRenderer.h b/Engine/Source/Runtime/Renderer/Private/DeferredShadingRenderer.h
--- a/Engine/Source/Runtime/Renderer/Private/DeferredShadingRenderer.h
+++ b/Engine/Source/Runtime/Renderer/Private/DeferredShadingRenderer.h
@@ -112,7 +112,7 @@ public:
 	/** Finishes the view family rendering. */
 	void RenderFinish(FRHICommandListImmediate& RHICmdList);
 
-	void RenderOcclusion(FRHICommandListImmediate& RHICmdList, bool bRenderQueries, bool bRenderHZB);
+	void RenderOcclusion(FRHICommandListImmediate& RHICmdList, bool bRenderQueries);
 
 	/** Renders the view family. */
 	virtual void Render(FRHICommandListImmediate& RHICmdList) override;
@@ -251,15 +251,15 @@ private:
 	/** Issues a timestamp query for the end of the separate translucency pass. */
 	void EndTimingSeparateTranslucencyPass(FRHICommandListImmediate& RHICmdList, const FViewInfo& View);
 
-	/** 
-	 * Renders the scene's translucency, parallel version
-	 */
-	void RenderTranslucencyParallel(FRHICommandListImmediate& RHICmdList);
 
-	/**
-	* Renders the scene's translucency.
-	*/
-	void RenderTranslucency(FRHICommandListImmediate& RHICmdList);
+	/** Setup the downsampled view uniform buffer if it was not already built */
+	void SetupDownsampledTranslucencyViewUniformBuffer(FRHICommandListImmediate& RHICmdList, FViewInfo& View);
+
+	/** Resolve the scene color if any translucent material needs it. */
+	void ConditionalResolveSceneColorForTranslucentMaterials(FRHICommandListImmediate& RHICmdList);
+
+	/** Renders the scene's translucency. */
+	void RenderTranslucency(FRHICommandListImmediate& RHICmdList, ETranslucencyPass::Type TranslucencyPass);
 
 	/** Renders the scene's light shafts */
 	void RenderLightShaftOcclusion(FRHICommandListImmediate& RHICmdList, FLightShaftsOutput& Output);
@@ -304,7 +304,7 @@ private:
 		bool bProjectingForForwardShading) const;
 
 	/** Sets up ViewState buffers for rendering capsule shadows. */
-	void SetupIndirectCapsuleShadows(FRHICommandListImmediate& RHICmdList, const FViewInfo& View, bool bPrepareLightData, int32& NumCapsuleShapes, int32& NumMeshDistanceFieldCasters) const;
+	void SetupIndirectCapsuleShadows(FRHICommandListImmediate& RHICmdList, const FViewInfo& View, bool bPrepareLightData, int32& NumCapsuleShapes, int32& NumMeshesWithCapsules, int32& NumMeshDistanceFieldCasters) const;
 
 	/** Renders indirect shadows from capsules modulated onto scene color. */
 	void RenderIndirectCapsuleShadows(
@@ -385,8 +385,9 @@ private:
 	void RenderLocalLightsForVolumetricFog(
 		FRHICommandListImmediate& RHICmdList,
 		FViewInfo& View,
+		bool bUseTemporalReprojection,
+		const struct FVolumetricFogIntegrationParameterData& IntegrationData,
 		const FExponentialHeightFogSceneInfo& FogInfo,
-		IPooledRenderTarget* VBufferA,
 		FIntVector VolumetricFogGridSize,
 		FVector GridZParams,
 		const FPooledRenderTargetDesc& VolumeDesc,
@@ -404,11 +405,10 @@ private:
 	void VoxelizeFogVolumePrimitives(
 		FRHICommandListImmediate& RHICmdList,
 		const FViewInfo& View,
+		const FVolumetricFogIntegrationParameterData& IntegrationData,
 		FIntVector VolumetricFogGridSize,
 		FVector GridZParams,
-		float VolumetricFogDistance,
-		IPooledRenderTarget* VBufferA,
-		IPooledRenderTarget* VBufferB);
+		float VolumetricFogDistance);
 
 	void ComputeVolumetricFog(FRHICommandListImmediate& RHICmdList);
 
@@ -437,7 +437,8 @@ private:
 
 	void UpdateGlobalDistanceFieldObjectBuffers(FRHICommandListImmediate& RHICmdList);
 
-	void DrawAllTranslucencyPasses(FRHICommandListImmediate& RHICmdList, const FViewInfo& View, const FDrawingPolicyRenderState& DrawRenderState, ETranslucencyPass::Type TranslucenyPassType);
+	void RenderViewTranslucency(FRHICommandListImmediate& RHICmdList, const FViewInfo& View, const FDrawingPolicyRenderState& DrawRenderState, ETranslucencyPass::Type TranslucenyPass);
+	void RenderViewTranslucencyParallel(FRHICommandListImmediate& RHICmdList, const FViewInfo& View, const FDrawingPolicyRenderState& DrawRenderState, ETranslucencyPass::Type TranslucencyPass);
 
 	void CopySceneCaptureComponentToTarget(FRHICommandListImmediate& RHICmdList);
 
diff --git a/Engine/Source/Runtime/Renderer/Private/DepthRendering.cpp b/Engine/Source/Runtime/Renderer/Private/DepthRendering.cpp
--- a/Engine/Source/Runtime/Renderer/Private/DepthRendering.cpp
+++ b/Engine/Source/Runtime/Renderer/Private/DepthRendering.cpp
@@ -1310,7 +1310,7 @@ bool FDeferredShadingSceneRenderer::RenderPrePass(FRHICommandListImmediate& RHIC
 			}
 			RHICmdList.SetViewport(FullViewRect.Min.X, FullViewRect.Min.Y, 0, FullViewRect.Max.X, FullViewRect.Max.Y, 1);
 		}
-		DrawClearQuad(RHICmdList, GMaxRHIFeatureLevel, false, FLinearColor::Transparent, false, 0, true, 0);
+		DrawClearQuad(RHICmdList, false, FLinearColor::Transparent, false, 0, true, 0);
 	}
 
 	SceneContext.FinishRenderingPrePass(RHICmdList);
diff --git a/Engine/Source/Runtime/Renderer/Private/DistanceFieldAmbientOcclusion.cpp b/Engine/Source/Runtime/Renderer/Private/DistanceFieldAmbientOcclusion.cpp
--- a/Engine/Source/Runtime/Renderer/Private/DistanceFieldAmbientOcclusion.cpp
+++ b/Engine/Source/Runtime/Renderer/Private/DistanceFieldAmbientOcclusion.cpp
@@ -91,6 +91,14 @@ FAutoConsoleVariableRef CVarAOOverwriteSceneColor(
 	ECVF_RenderThreadSafe
 	);
 
+int32 GAOJitterConeDirections = 0;
+FAutoConsoleVariableRef CVarAOJitterConeDirections(
+	TEXT("r.AOJitterConeDirections"),
+	GAOJitterConeDirections,
+	TEXT(""),
+	ECVF_RenderThreadSafe
+	);
+
 int32 GMaxDistanceFieldObjectsPerCullTile = 512;
 FAutoConsoleVariableRef CVarMaxDistanceFieldObjectsPerCullTile(
 	TEXT("r.AOMaxObjectsPerCullTile"),
@@ -160,7 +168,21 @@ const FVector RelaxedSpacedVectors9[] =
 	FVector(0.032967, -0.435625, 0.899524)
 };
 
-void GetSpacedVectors(TArray<FVector, TInlineAllocator<9> >& OutVectors)
+float TemporalHalton2( int32 Index, int32 Base )
+{
+	float Result = 0.0f;
+	float InvBase = 1.0f / Base;
+	float Fraction = InvBase;
+	while( Index > 0 )
+	{
+		Result += ( Index % Base ) * Fraction;
+		Index /= Base;
+		Fraction *= InvBase;
+	}
+	return Result;
+}
+
+void GetSpacedVectors(uint32 FrameNumber, TArray<FVector, TInlineAllocator<9> >& OutVectors)
 {
 	OutVectors.Empty(ARRAY_COUNT(SpacedVectors9));
 
@@ -178,6 +200,22 @@ void GetSpacedVectors(TArray<FVector, TInlineAllocator<9> >& OutVectors)
 			OutVectors.Add(RelaxedSpacedVectors9[i]);
 		}
 	}
+
+	if (GAOJitterConeDirections)
+	{
+		float RandomAngle = TemporalHalton2(FrameNumber & 1023, 2) * 2 * PI;
+		float CosRandomAngle = FMath::Cos(RandomAngle);
+		float SinRandomAngle = FMath::Sin(RandomAngle);
+
+		for (int32 i = 0; i < OutVectors.Num(); i++)
+		{
+			FVector ConeDirection = OutVectors[i];
+			FVector2D ConeDirectionXY(ConeDirection.X, ConeDirection.Y);
+			ConeDirectionXY = FVector2D(FVector2D::DotProduct(ConeDirectionXY, FVector2D(CosRandomAngle, -SinRandomAngle)), FVector2D::DotProduct(ConeDirectionXY, FVector2D(SinRandomAngle, CosRandomAngle)));
+			OutVectors[i].X = ConeDirectionXY.X;
+			OutVectors[i].Y = ConeDirectionXY.Y;
+		}
+	}
 }
 
 // Cone half angle derived from each cone covering an equal solid angle
diff --git a/Engine/Source/Runtime/Renderer/Private/DistanceFieldAmbientOcclusion.h b/Engine/Source/Runtime/Renderer/Private/DistanceFieldAmbientOcclusion.h
--- a/Engine/Source/Runtime/Renderer/Private/DistanceFieldAmbientOcclusion.h
+++ b/Engine/Source/Runtime/Renderer/Private/DistanceFieldAmbientOcclusion.h
@@ -234,7 +234,7 @@ public:
 	}
 };
 
-extern void GetSpacedVectors(TArray<FVector, TInlineAllocator<9> >& OutVectors);
+extern void GetSpacedVectors(uint32 FrameNumber, TArray<FVector, TInlineAllocator<9> >& OutVectors);
 
 BEGIN_UNIFORM_BUFFER_STRUCT(FAOSampleData2,)
 	DECLARE_UNIFORM_BUFFER_STRUCT_MEMBER_ARRAY(FVector4,SampleDirections,[NumConeSampleDirections])
diff --git a/Engine/Source/Runtime/Renderer/Private/DistanceFieldGlobalIllumination.cpp b/Engine/Source/Runtime/Renderer/Private/DistanceFieldGlobalIllumination.cpp
--- a/Engine/Source/Runtime/Renderer/Private/DistanceFieldGlobalIllumination.cpp
+++ b/Engine/Source/Runtime/Renderer/Private/DistanceFieldGlobalIllumination.cpp
@@ -372,7 +372,7 @@ void PlaceVPLs(
 	GVPLResources.AllocateFor(GVPLGridDimension * GVPLGridDimension);
 
 	{
-		ClearUAV(RHICmdList, GMaxRHIFeatureLevel, GVPLResources.VPLParameterBuffer, 0);
+		ClearUAV(RHICmdList, GVPLResources.VPLParameterBuffer, 0);
 	}
 
 	const FLightSceneProxy* DirectionalLightProxy = NULL;
@@ -530,7 +530,7 @@ void PlaceVPLs(
 			{
 				GCulledVPLResources.AllocateFor(GVPLGridDimension * GVPLGridDimension);
 
-				ClearUAV(RHICmdList, GMaxRHIFeatureLevel, GCulledVPLResources.VPLParameterBuffer, 0);
+				ClearUAV(RHICmdList, GCulledVPLResources.VPLParameterBuffer, 0);
 
 				TShaderMapRef<FCullVPLsForViewCS> ComputeShader(GetGlobalShaderMap(Scene->GetFeatureLevel()));
 				RHICmdList.SetComputeShader(ComputeShader->GetComputeShader());
@@ -863,7 +863,7 @@ void UpdateVPLs(
 			}
 			else
 			{
-				ClearUAV(RHICmdList, GMaxRHIFeatureLevel, Scene->DistanceFieldSceneData.InstancedSurfelBuffers->VPLFlux, 0);
+				ClearUAV(RHICmdList, Scene->DistanceFieldSceneData.InstancedSurfelBuffers->VPLFlux, 0);
 			}
 		}
 		else
@@ -922,7 +922,7 @@ public:
 		FAOSampleData2 AOSampleData;
 
 		TArray<FVector, TInlineAllocator<9> > SampleDirections;
-		GetSpacedVectors(SampleDirections);
+		GetSpacedVectors(View.Family->FrameNumber, SampleDirections);
 
 		for (int32 SampleIndex = 0; SampleIndex < NumConeSampleDirections; SampleIndex++)
 		{
@@ -1181,8 +1181,8 @@ void ComputeIrradianceForScreenGrid(
 	const uint32 GroupSizeX = FMath::DivideAndRoundUp(View.ViewRect.Size().X / GAODownsampleFactor, GScreenGridIrradianceThreadGroupSizeX);
 	const uint32 GroupSizeY = FMath::DivideAndRoundUp(View.ViewRect.Size().Y / GAODownsampleFactor, GScreenGridIrradianceThreadGroupSizeX);
 
-	ClearUAV(RHICmdList, GMaxRHIFeatureLevel, ScreenGridResources.HeightfieldIrradiance, 0);
-	ClearUAV(RHICmdList, GMaxRHIFeatureLevel, ScreenGridResources.SurfelIrradiance, 0);
+	ClearUAV(RHICmdList, ScreenGridResources.HeightfieldIrradiance, 0);
+	ClearUAV(RHICmdList, ScreenGridResources.SurfelIrradiance, 0);
 
 	View.HeightfieldLightingViewInfo.
 		ComputeIrradianceForScreenGrid(View, RHICmdList, DistanceFieldNormal, ScreenGridResources, Parameters);
diff --git a/Engine/Source/Runtime/Renderer/Private/DistanceFieldLightingPost.cpp b/Engine/Source/Runtime/Renderer/Private/DistanceFieldLightingPost.cpp
--- a/Engine/Source/Runtime/Renderer/Private/DistanceFieldLightingPost.cpp
+++ b/Engine/Source/Runtime/Renderer/Private/DistanceFieldLightingPost.cpp
@@ -20,6 +20,14 @@ FAutoConsoleVariableRef CVarAOUseHistory(
 	ECVF_RenderThreadSafe
 	);
 
+int32 GAOClearHistory = 0;
+FAutoConsoleVariableRef CVarAOClearHistory(
+	TEXT("r.AOClearHistory"),
+	GAOClearHistory,
+	TEXT(""),
+	ECVF_RenderThreadSafe
+	);
+
 int32 GAOHistoryStabilityPass = 1;
 FAutoConsoleVariableRef CVarAOHistoryStabilityPass(
 	TEXT("r.AOHistoryStabilityPass"),
@@ -426,10 +434,15 @@ void UpdateHistory(
 	{
 		const bool bUseDistanceFieldGI = IsDistanceFieldGIAllowed(View);
 
+		FIntPoint BufferSize = GetBufferSizeForAO();
+
 		if (*BentNormalHistoryState 
 			&& !View.bCameraCut 
 			&& !View.bPrevTransformsReset 
-			&& (!bUseDistanceFieldGI || (IrradianceHistoryState && *IrradianceHistoryState)))
+			&& (!bUseDistanceFieldGI || (IrradianceHistoryState && *IrradianceHistoryState))
+			&& !GAOClearHistory
+			// If the scene render targets reallocate, toss the history so we don't read uninitialized data
+			&& (*BentNormalHistoryState)->GetDesc().Extent == BufferSize)
 		{
 			// Reuse a render target from the pool with a consistent name, for vis purposes
 			TRefCountPtr<IPooledRenderTarget> NewBentNormalHistory;
diff --git a/Engine/Source/Runtime/Renderer/Private/DistanceFieldObjectCulling.cpp b/Engine/Source/Runtime/Renderer/Private/DistanceFieldObjectCulling.cpp
--- a/Engine/Source/Runtime/Renderer/Private/DistanceFieldObjectCulling.cpp
+++ b/Engine/Source/Runtime/Renderer/Private/DistanceFieldObjectCulling.cpp
@@ -196,7 +196,7 @@ void CullObjectsToView(FRHICommandListImmediate& RHICmdList, FScene* Scene, cons
 	}
 
 	{
-		ClearUAV(RHICmdList, GMaxRHIFeatureLevel, CulledObjectBuffers.Buffers.ObjectIndirectArguments, 0);
+		ClearUAV(RHICmdList, CulledObjectBuffers.Buffers.ObjectIndirectArguments, 0);
 
 		TShaderMapRef<FCullObjectsForViewCS> ComputeShader(GetGlobalShaderMap(Scene->GetFeatureLevel()));
 		RHICmdList.SetComputeShader(ComputeShader->GetComputeShader());
@@ -630,7 +630,7 @@ FIntPoint BuildTileObjectLists(FRHICommandListImmediate& RHICmdList, FScene* Sce
 				SCOPED_DRAW_EVENT(RHICmdList, CountTileObjectIntersections);
 
 				// Start at 0 tiles per object
-				ClearUAV(RHICmdList, GMaxRHIFeatureLevel, TileIntersectionResources->NumCulledTilesArray, 0);
+				ClearUAV(RHICmdList, TileIntersectionResources->NumCulledTilesArray, 0);
 
 				// Rasterize object bounding shapes and intersect with screen tiles to compute how many tiles intersect each object
 				ScatterTilesToObjects<true>(RHICmdList, View, TileListGroupSize, Parameters);
@@ -639,7 +639,7 @@ FIntPoint BuildTileObjectLists(FRHICommandListImmediate& RHICmdList, FScene* Sce
 			{
 				SCOPED_DRAW_EVENT(RHICmdList, ComputeStartOffsets);
 				// Start at 0 threadgroups
-				ClearUAV(RHICmdList, GMaxRHIFeatureLevel, TileIntersectionResources->ObjectTilesIndirectArguments, 0);
+				ClearUAV(RHICmdList, TileIntersectionResources->ObjectTilesIndirectArguments, 0);
 
 				// Accumulate how many cone trace threadgroups we should dispatch, and also compute the start offset for each object's culled tile data
 				TShaderMapRef<FComputeCulledTilesStartOffsetCS> ComputeShader(View.ShaderMap);
@@ -657,7 +657,7 @@ FIntPoint BuildTileObjectLists(FRHICommandListImmediate& RHICmdList, FScene* Sce
 				SCOPED_DRAW_EVENT(RHICmdList, CullTilesToObjects);
 
 				// Start at 0 tiles per object
-				ClearUAV(RHICmdList, GMaxRHIFeatureLevel, TileIntersectionResources->NumCulledTilesArray, 0);
+				ClearUAV(RHICmdList, TileIntersectionResources->NumCulledTilesArray, 0);
 
 				// Rasterize object bounding shapes and intersect with screen tiles, and write out intersecting tile indices for the cone tracing pass
 				ScatterTilesToObjects<false>(RHICmdList, View, TileListGroupSize, Parameters);
diff --git a/Engine/Source/Runtime/Renderer/Private/DistanceFieldScreenGridLighting.cpp b/Engine/Source/Runtime/Renderer/Private/DistanceFieldScreenGridLighting.cpp
--- a/Engine/Source/Runtime/Renderer/Private/DistanceFieldScreenGridLighting.cpp
+++ b/Engine/Source/Runtime/Renderer/Private/DistanceFieldScreenGridLighting.cpp
@@ -141,7 +141,7 @@ public:
 		FAOSampleData2 AOSampleData;
 
 		TArray<FVector, TInlineAllocator<9> > SampleDirections;
-		GetSpacedVectors(SampleDirections);
+		GetSpacedVectors(View.Family->FrameNumber, SampleDirections);
 
 		for (int32 SampleIndex = 0; SampleIndex < NumConeSampleDirections; SampleIndex++)
 		{
@@ -288,7 +288,7 @@ public:
 		FAOSampleData2 AOSampleData;
 
 		TArray<FVector, TInlineAllocator<9> > SampleDirections;
-		GetSpacedVectors(SampleDirections);
+		GetSpacedVectors(View.Family->FrameNumber, SampleDirections);
 
 		for (int32 SampleIndex = 0; SampleIndex < NumConeSampleDirections; SampleIndex++)
 		{
@@ -424,7 +424,7 @@ public:
 		FAOSampleData2 AOSampleData;
 
 		TArray<FVector, TInlineAllocator<9> > SampleDirections;
-		GetSpacedVectors(SampleDirections);
+		GetSpacedVectors(View.Family->FrameNumber, SampleDirections);
 
 		for (int32 SampleIndex = 0; SampleIndex < NumConeSampleDirections; SampleIndex++)
 		{
@@ -779,7 +779,7 @@ void FDeferredShadingSceneRenderer::RenderDistanceFieldAOScreenGrid(
 		SCOPED_DRAW_EVENT(RHICmdList, ConeTraceGlobal);
 
 		float ConeVisibilityClearValue = 1.0f;
-		ClearUAV(RHICmdList, GMaxRHIFeatureLevel, ScreenGridResources->ScreenGridConeVisibility, *(uint32*)&ConeVisibilityClearValue);
+		ClearUAV(RHICmdList, ScreenGridResources->ScreenGridConeVisibility, *(uint32*)&ConeVisibilityClearValue);
 
 		const uint32 GroupSizeX = FMath::DivideAndRoundUp(View.ViewRect.Size().X / GAODownsampleFactor / GConeTraceDownsampleFactor, GConeTraceGlobalDFTileSize);
 		const uint32 GroupSizeY = FMath::DivideAndRoundUp(View.ViewRect.Size().Y / GAODownsampleFactor / GConeTraceDownsampleFactor, GConeTraceGlobalDFTileSize);
diff --git a/Engine/Source/Runtime/Renderer/Private/DistanceFieldShadowing.cpp b/Engine/Source/Runtime/Renderer/Private/DistanceFieldShadowing.cpp
--- a/Engine/Source/Runtime/Renderer/Private/DistanceFieldShadowing.cpp
+++ b/Engine/Source/Runtime/Renderer/Private/DistanceFieldShadowing.cpp
@@ -44,14 +44,6 @@ FAutoConsoleVariableRef CVarFullResolutionDFShadowing(
 	ECVF_Scalability | ECVF_RenderThreadSafe
 	);
 
-int32 GAsyncComputeDFShadowing = 0;
-FAutoConsoleVariableRef CVarAsyncComputeDFShadowing(
-	TEXT("r.DFShadowAsyncCompute"),
-	GAsyncComputeDFShadowing,
-	TEXT("Whether to use async compute for ray marching distance fields."),
-	ECVF_Scalability | ECVF_RenderThreadSafe
-	);
-
 int32 GShadowScatterTileCulling = 1;
 FAutoConsoleVariableRef CVarShadowScatterTileCulling(
 	TEXT("r.DFShadowScatterTileCulling"),
@@ -687,7 +679,7 @@ void CullDistanceFieldObjectsForLight(
 		}
 
 		{
-			ClearUAV(RHICmdList, GMaxRHIFeatureLevel, GShadowCulledObjectBuffers.Buffers.ObjectIndirectArguments, 0);
+			ClearUAV(RHICmdList, GShadowCulledObjectBuffers.Buffers.ObjectIndirectArguments, 0);
 
 			TShaderMapRef<FCullObjectsForShadowCS> ComputeShader(GetGlobalShaderMap(Scene->GetFeatureLevel()));
 			RHICmdList.SetComputeShader(ComputeShader->GetComputeShader());
@@ -920,43 +912,20 @@ void FProjectedShadowInfo::BeginRenderRayTracedDistanceFieldProjection(FRHIComma
 			SCOPED_DRAW_EVENT(RHICmdList, RayTraceShadows);
 			SetRenderTarget(RHICmdList, NULL, NULL);
 
-			if (GAsyncComputeDFShadowing)
-			{
-				FRHIAsyncComputeCommandListImmediate& RHICmdListComputeImmediate = FRHICommandListExecutor::GetImmediateAsyncComputeCommandList();
-				static const FName RTShadowBeginComputeName(TEXT("RTShadowComputeBegin"));
-				static const FName RTShadowComputeEndName(TEXT("RTShadowComputeEnd"));
-				FComputeFenceRHIRef BeginFence = RHICmdList.CreateComputeFence(RTShadowBeginComputeName);
-				RayTracedShadowsEndFence = RHICmdList.CreateComputeFence(RTShadowComputeEndName);
-
-				RHICmdList.TransitionResources(EResourceTransitionAccess::EReadable, EResourceTransitionPipeline::EGfxToCompute, nullptr, 0, BeginFence);
-				RHICmdListComputeImmediate.WaitComputeFence(BeginFence);
-
-				RayTraceShadows(RHICmdListComputeImmediate, View, this, TileIntersectionResources);
-
-				RHICmdListComputeImmediate.TransitionResources(EResourceTransitionAccess::EReadable, EResourceTransitionPipeline::EComputeToGfx, nullptr, 0, RayTracedShadowsEndFence);
-
-				FRHIAsyncComputeCommandListImmediate::ImmediateDispatch(RHICmdListComputeImmediate);
-			}
-			else
-			{
-				RayTraceShadows(RHICmdList, View, this, TileIntersectionResources);
-			}
+			RayTraceShadows(RHICmdList, View, this, TileIntersectionResources);
 		}
 	}
 }
 
 void FProjectedShadowInfo::RenderRayTracedDistanceFieldProjection(FRHICommandListImmediate& RHICmdList, const FViewInfo& View, bool bProjectingForForwardShading) 
 {
+	BeginRenderRayTracedDistanceFieldProjection(RHICmdList, View);
+
 	if (RayTracedShadowsRT)
 	{
 		QUICK_SCOPE_CYCLE_COUNTER(STAT_RenderRayTracedDistanceFieldShadows);
 		SCOPED_DRAW_EVENT(RHICmdList, RayTracedDistanceFieldShadow);
 
-		if (GAsyncComputeDFShadowing)
-		{
-			RHICmdList.WaitComputeFence(RayTracedShadowsEndFence);
-		}
-
 		FIntRect ScissorRect;
 
 		if (!LightSceneInfo->Proxy->GetScissorRect(ScissorRect, View))
diff --git a/Engine/Source/Runtime/Renderer/Private/DistortionRendering.cpp b/Engine/Source/Runtime/Renderer/Private/DistortionRendering.cpp
--- a/Engine/Source/Runtime/Renderer/Private/DistortionRendering.cpp
+++ b/Engine/Source/Runtime/Renderer/Private/DistortionRendering.cpp
@@ -906,7 +906,7 @@ bool FDistortionPrimSet::DrawAccumulatedOffsets(FRHICommandListImmediate& RHICmd
 							&View,
 							bInitializeOffsets,
 							StaticMesh,
-							StaticMesh.bRequiresPerElementVisibility ? View.StaticMeshBatchVisibility[StaticMesh.Id] : ((1ull << StaticMesh.Elements.Num()) - 1),
+							StaticMesh.bRequiresPerElementVisibility ? View.StaticMeshBatchVisibility[StaticMesh.BatchVisibilityId] : ((1ull << StaticMesh.Elements.Num()) - 1),
 							DrawRenderState,
 							PrimitiveSceneProxy,
 							StaticMesh.BatchHitProxyId
diff --git a/Engine/Source/Runtime/Renderer/Private/GammaCorrection.cpp b/Engine/Source/Runtime/Renderer/Private/GammaCorrection.cpp
--- a/Engine/Source/Runtime/Renderer/Private/GammaCorrection.cpp
+++ b/Engine/Source/Runtime/Renderer/Private/GammaCorrection.cpp
@@ -102,7 +102,7 @@ void FSceneRenderer::GammaCorrectToViewportRenderTarget(FRHICommandList& RHICmdL
 		else
 		{
 			SetRenderTarget(RHICmdList, ViewFamily.RenderTarget->GetRenderTargetTexture(), FTextureRHIRef());
-			DrawClearQuad(RHICmdList, GMaxRHIFeatureLevel, FLinearColor::Black);
+			DrawClearQuad(RHICmdList, FLinearColor::Black);
 		}
 		ViewFamily.bDeferClear = false;
 	}
diff --git a/Engine/Source/Runtime/Renderer/Private/GlobalDistanceField.cpp b/Engine/Source/Runtime/Renderer/Private/GlobalDistanceField.cpp
--- a/Engine/Source/Runtime/Renderer/Private/GlobalDistanceField.cpp
+++ b/Engine/Source/Runtime/Renderer/Private/GlobalDistanceField.cpp
@@ -375,9 +375,25 @@ private:
 
 IMPLEMENT_SHADER_TYPE(,FCullObjectsToGridCS,TEXT("GlobalDistanceField"),TEXT("CullObjectsToGridCS"),SF_Compute);
 
-const int32 CompositeTileSize = 4;
+enum EFlattenedDimension
+{
+	Flatten_XAxis = 0,
+	Flatten_YAxis = 1,
+	Flatten_ZAxis = 2,
+	Flatten_None
+};
+
+int32 GetCompositeTileSize(int32 Dimension, EFlattenedDimension FlattenedDimension)
+{
+	if (FlattenedDimension == Flatten_None)
+	{
+		return 4;
+	}
+
+	return Dimension == (int32)FlattenedDimension ? 1 : 8;
+}
 
-template<bool bUseParentDistanceField>
+template<bool bUseParentDistanceField, EFlattenedDimension FlattenedDimension>
 class TCompositeObjectDistanceFieldsCS : public FGlobalShader
 {
 	DECLARE_SHADER_TYPE(TCompositeObjectDistanceFieldsCS,Global)
@@ -391,7 +407,9 @@ public:
 	static void ModifyCompilationEnvironment(EShaderPlatform Platform, FShaderCompilerEnvironment& OutEnvironment)
 	{
 		FGlobalShader::ModifyCompilationEnvironment(Platform,OutEnvironment);
-		OutEnvironment.SetDefine(TEXT("COMPOSITE_THREADGROUP_SIZE"), CompositeTileSize);
+		OutEnvironment.SetDefine(TEXT("COMPOSITE_THREADGROUP_SIZEX"), GetCompositeTileSize(0, FlattenedDimension));
+		OutEnvironment.SetDefine(TEXT("COMPOSITE_THREADGROUP_SIZEY"), GetCompositeTileSize(1, FlattenedDimension));
+		OutEnvironment.SetDefine(TEXT("COMPOSITE_THREADGROUP_SIZEZ"), GetCompositeTileSize(2, FlattenedDimension));
 		OutEnvironment.SetDefine(TEXT("CULL_GRID_TILE_SIZE"), GCullGridTileSize);
 		OutEnvironment.SetDefine(TEXT("MAX_GRID_CULLED_DF_OBJECTS"), GMaxGridCulledObjects);
 		OutEnvironment.SetDefine(TEXT("USE_PARENT_DISTANCE_FIELD"), bUseParentDistanceField ? 1 : 0);
@@ -505,8 +523,19 @@ private:
 	FShaderParameter AOGlobalMaxSphereQueryRadius;
 };
 
-IMPLEMENT_SHADER_TYPE(template<>,TCompositeObjectDistanceFieldsCS<true>,TEXT("GlobalDistanceField"),TEXT("CompositeObjectDistanceFieldsCS"),SF_Compute);
-IMPLEMENT_SHADER_TYPE(template<>,TCompositeObjectDistanceFieldsCS<false>,TEXT("GlobalDistanceField"),TEXT("CompositeObjectDistanceFieldsCS"),SF_Compute);
+#define IMPLEMENT_GLOBALDF_COMPOSITE_CS_TYPE(bUseParentDistanceField, FlattenedDimension) \
+	typedef TCompositeObjectDistanceFieldsCS<bUseParentDistanceField, FlattenedDimension> TCompositeObjectDistanceFieldsCS##bUseParentDistanceField##FlattenedDimension; \
+	IMPLEMENT_SHADER_TYPE(template<>,TCompositeObjectDistanceFieldsCS##bUseParentDistanceField##FlattenedDimension,TEXT("GlobalDistanceField"),TEXT("CompositeObjectDistanceFieldsCS"),SF_Compute);
+
+IMPLEMENT_GLOBALDF_COMPOSITE_CS_TYPE(true, Flatten_None);
+IMPLEMENT_GLOBALDF_COMPOSITE_CS_TYPE(true, Flatten_XAxis);
+IMPLEMENT_GLOBALDF_COMPOSITE_CS_TYPE(true, Flatten_YAxis);
+IMPLEMENT_GLOBALDF_COMPOSITE_CS_TYPE(true, Flatten_ZAxis);
+
+IMPLEMENT_GLOBALDF_COMPOSITE_CS_TYPE(false, Flatten_None);
+IMPLEMENT_GLOBALDF_COMPOSITE_CS_TYPE(false, Flatten_XAxis);
+IMPLEMENT_GLOBALDF_COMPOSITE_CS_TYPE(false, Flatten_YAxis);
+IMPLEMENT_GLOBALDF_COMPOSITE_CS_TYPE(false, Flatten_ZAxis);
 
 const int32 HeightfieldCompositeTileSize = 8;
 
@@ -1285,7 +1314,7 @@ void UpdateGlobalDistanceFieldVolume(
 
 				for (int32 ClipmapIndex = 0; ClipmapIndex < Clipmaps.Num(); ClipmapIndex++)
 				{
-					SCOPED_DRAW_EVENTF(RHICmdList, Clipmap, TEXT("CacheType %u Clipmap %u"), CacheType, ClipmapIndex);
+					SCOPED_DRAW_EVENTF(RHICmdList, Clipmap, TEXT("CacheType %s Clipmap %u"), CacheType == GDF_MostlyStatic ? TEXT("MostlyStatic") : TEXT("Movable"), ClipmapIndex);
 
 					FGlobalDistanceFieldClipmap& Clipmap = Clipmaps[ClipmapIndex];
 
@@ -1300,7 +1329,7 @@ void UpdateGlobalDistanceFieldVolume(
 
 								// Cull the global objects to the volume being updated
 								{
-									ClearUAV(RHICmdList, GMaxRHIFeatureLevel, GGlobalDistanceFieldCulledObjectBuffers.Buffers.ObjectIndirectArguments, 0);
+									ClearUAV(RHICmdList, GGlobalDistanceFieldCulledObjectBuffers.Buffers.ObjectIndirectArguments, 0);
 
 									TShaderMapRef<FCullObjectsForVolumeCS> ComputeShader(View.ShaderMap);
 									RHICmdList.SetComputeShader(ComputeShader->GetComputeShader());
@@ -1330,28 +1359,98 @@ void UpdateGlobalDistanceFieldVolume(
 							{
 								SCOPED_DRAW_EVENTF(RHICmdList, TileCullAndComposite, TEXT("TileCullAndComposite %ux%ux%u"), UpdateRegion.CellsSize.X, UpdateRegion.CellsSize.Y, UpdateRegion.CellsSize.Z);
 
-								//@todo - match typical update sizes.  Camera movement creates narrow slabs.
-								const uint32 NumGroupsX = FMath::DivideAndRoundUp<int32>(UpdateRegion.CellsSize.X, CompositeTileSize);
-								const uint32 NumGroupsY = FMath::DivideAndRoundUp<int32>(UpdateRegion.CellsSize.Y, CompositeTileSize);
-								const uint32 NumGroupsZ = FMath::DivideAndRoundUp<int32>(UpdateRegion.CellsSize.Z, CompositeTileSize);
+								int32 MinDimension = 2;
+
+								if (UpdateRegion.CellsSize.X < UpdateRegion.CellsSize.Y && UpdateRegion.CellsSize.X < UpdateRegion.CellsSize.Z)
+								{
+									MinDimension = 0;
+								}
+								else if (UpdateRegion.CellsSize.Y < UpdateRegion.CellsSize.X && UpdateRegion.CellsSize.Y < UpdateRegion.CellsSize.Z)
+								{
+									MinDimension = 1;
+								}
+
+								int32 MinSize = UpdateRegion.CellsSize[MinDimension];
+								int32 MaxSize = FMath::Max(UpdateRegion.CellsSize.X, FMath::Max(UpdateRegion.CellsSize.Y, UpdateRegion.CellsSize.Z));
+								const EFlattenedDimension FlattenedDimension = MaxSize >= MinSize * 8 ? (EFlattenedDimension)MinDimension : Flatten_None;
+
+								const uint32 NumGroupsX = FMath::DivideAndRoundUp<int32>(UpdateRegion.CellsSize.X, GetCompositeTileSize(0, FlattenedDimension));
+								const uint32 NumGroupsY = FMath::DivideAndRoundUp<int32>(UpdateRegion.CellsSize.Y, GetCompositeTileSize(1, FlattenedDimension));
+								const uint32 NumGroupsZ = FMath::DivideAndRoundUp<int32>(UpdateRegion.CellsSize.Z, GetCompositeTileSize(2, FlattenedDimension));
 
 								IPooledRenderTarget* ParentDistanceField = GlobalDistanceFieldInfo.MostlyStaticClipmaps[ClipmapIndex].RenderTarget;
 
 								if (CacheType == GDF_Full && GAOGlobalDistanceFieldCacheMostlyStaticSeparately && ParentDistanceField)
 								{
-									TShaderMapRef<TCompositeObjectDistanceFieldsCS<true>> ComputeShader(View.ShaderMap);
-									RHICmdList.SetComputeShader(ComputeShader->GetComputeShader());
-									ComputeShader->SetParameters(RHICmdList, Scene, View, MaxOcclusionDistance, GlobalDistanceFieldInfo.ParameterData, Clipmap, ParentDistanceField, ClipmapIndex, UpdateRegion);
-									DispatchComputeShader(RHICmdList, *ComputeShader, NumGroupsX, NumGroupsY, NumGroupsZ);
-									ComputeShader->UnsetParameters(RHICmdList, Clipmap);
+									if (FlattenedDimension == Flatten_None)
+									{
+										TShaderMapRef<TCompositeObjectDistanceFieldsCS<true, Flatten_None>> ComputeShader(View.ShaderMap);
+										RHICmdList.SetComputeShader(ComputeShader->GetComputeShader());
+										ComputeShader->SetParameters(RHICmdList, Scene, View, MaxOcclusionDistance, GlobalDistanceFieldInfo.ParameterData, Clipmap, ParentDistanceField, ClipmapIndex, UpdateRegion);
+										DispatchComputeShader(RHICmdList, *ComputeShader, NumGroupsX, NumGroupsY, NumGroupsZ);
+										ComputeShader->UnsetParameters(RHICmdList, Clipmap);
+									}
+									else if (FlattenedDimension == Flatten_XAxis)
+									{
+										TShaderMapRef<TCompositeObjectDistanceFieldsCS<true, Flatten_XAxis>> ComputeShader(View.ShaderMap);
+										RHICmdList.SetComputeShader(ComputeShader->GetComputeShader());
+										ComputeShader->SetParameters(RHICmdList, Scene, View, MaxOcclusionDistance, GlobalDistanceFieldInfo.ParameterData, Clipmap, ParentDistanceField, ClipmapIndex, UpdateRegion);
+										DispatchComputeShader(RHICmdList, *ComputeShader, NumGroupsX, NumGroupsY, NumGroupsZ);
+										ComputeShader->UnsetParameters(RHICmdList, Clipmap);
+									}
+									else if (FlattenedDimension == Flatten_YAxis)
+									{
+										TShaderMapRef<TCompositeObjectDistanceFieldsCS<true, Flatten_YAxis>> ComputeShader(View.ShaderMap);
+										RHICmdList.SetComputeShader(ComputeShader->GetComputeShader());
+										ComputeShader->SetParameters(RHICmdList, Scene, View, MaxOcclusionDistance, GlobalDistanceFieldInfo.ParameterData, Clipmap, ParentDistanceField, ClipmapIndex, UpdateRegion);
+										DispatchComputeShader(RHICmdList, *ComputeShader, NumGroupsX, NumGroupsY, NumGroupsZ);
+										ComputeShader->UnsetParameters(RHICmdList, Clipmap);
+									}
+									else
+									{
+										check(FlattenedDimension == Flatten_ZAxis);
+										TShaderMapRef<TCompositeObjectDistanceFieldsCS<true, Flatten_ZAxis>> ComputeShader(View.ShaderMap);
+										RHICmdList.SetComputeShader(ComputeShader->GetComputeShader());
+										ComputeShader->SetParameters(RHICmdList, Scene, View, MaxOcclusionDistance, GlobalDistanceFieldInfo.ParameterData, Clipmap, ParentDistanceField, ClipmapIndex, UpdateRegion);
+										DispatchComputeShader(RHICmdList, *ComputeShader, NumGroupsX, NumGroupsY, NumGroupsZ);
+										ComputeShader->UnsetParameters(RHICmdList, Clipmap);
+									}
 								}
 								else
 								{
-									TShaderMapRef<TCompositeObjectDistanceFieldsCS<false>> ComputeShader(View.ShaderMap);
-									RHICmdList.SetComputeShader(ComputeShader->GetComputeShader());
-									ComputeShader->SetParameters(RHICmdList, Scene, View, MaxOcclusionDistance, GlobalDistanceFieldInfo.ParameterData, Clipmap, NULL, ClipmapIndex, UpdateRegion);
-									DispatchComputeShader(RHICmdList, *ComputeShader, NumGroupsX, NumGroupsY, NumGroupsZ);
-									ComputeShader->UnsetParameters(RHICmdList, Clipmap);
+									if (FlattenedDimension == Flatten_None)
+									{
+										TShaderMapRef<TCompositeObjectDistanceFieldsCS<false, Flatten_None>> ComputeShader(View.ShaderMap);
+										RHICmdList.SetComputeShader(ComputeShader->GetComputeShader());
+										ComputeShader->SetParameters(RHICmdList, Scene, View, MaxOcclusionDistance, GlobalDistanceFieldInfo.ParameterData, Clipmap, NULL, ClipmapIndex, UpdateRegion);
+										DispatchComputeShader(RHICmdList, *ComputeShader, NumGroupsX, NumGroupsY, NumGroupsZ);
+										ComputeShader->UnsetParameters(RHICmdList, Clipmap);
+									}
+									else if (FlattenedDimension == Flatten_XAxis)
+									{
+										TShaderMapRef<TCompositeObjectDistanceFieldsCS<false, Flatten_XAxis>> ComputeShader(View.ShaderMap);
+										RHICmdList.SetComputeShader(ComputeShader->GetComputeShader());
+										ComputeShader->SetParameters(RHICmdList, Scene, View, MaxOcclusionDistance, GlobalDistanceFieldInfo.ParameterData, Clipmap, NULL, ClipmapIndex, UpdateRegion);
+										DispatchComputeShader(RHICmdList, *ComputeShader, NumGroupsX, NumGroupsY, NumGroupsZ);
+										ComputeShader->UnsetParameters(RHICmdList, Clipmap);
+									}
+									else if (FlattenedDimension == Flatten_YAxis)
+									{
+										TShaderMapRef<TCompositeObjectDistanceFieldsCS<false, Flatten_YAxis>> ComputeShader(View.ShaderMap);
+										RHICmdList.SetComputeShader(ComputeShader->GetComputeShader());
+										ComputeShader->SetParameters(RHICmdList, Scene, View, MaxOcclusionDistance, GlobalDistanceFieldInfo.ParameterData, Clipmap, NULL, ClipmapIndex, UpdateRegion);
+										DispatchComputeShader(RHICmdList, *ComputeShader, NumGroupsX, NumGroupsY, NumGroupsZ);
+										ComputeShader->UnsetParameters(RHICmdList, Clipmap);
+									}
+									else
+									{
+										check(FlattenedDimension == Flatten_ZAxis);
+										TShaderMapRef<TCompositeObjectDistanceFieldsCS<false, Flatten_ZAxis>> ComputeShader(View.ShaderMap);
+										RHICmdList.SetComputeShader(ComputeShader->GetComputeShader());
+										ComputeShader->SetParameters(RHICmdList, Scene, View, MaxOcclusionDistance, GlobalDistanceFieldInfo.ParameterData, Clipmap, NULL, ClipmapIndex, UpdateRegion);
+										DispatchComputeShader(RHICmdList, *ComputeShader, NumGroupsX, NumGroupsY, NumGroupsZ);
+										ComputeShader->UnsetParameters(RHICmdList, Clipmap);
+									}
 								}
 							}
 						}
diff --git a/Engine/Source/Runtime/Renderer/Private/HeightfieldLighting.cpp b/Engine/Source/Runtime/Renderer/Private/HeightfieldLighting.cpp
--- a/Engine/Source/Runtime/Renderer/Private/HeightfieldLighting.cpp
+++ b/Engine/Source/Runtime/Renderer/Private/HeightfieldLighting.cpp
@@ -1277,7 +1277,7 @@ public:
 		FAOSampleData2 AOSampleData;
 
 		TArray<FVector, TInlineAllocator<9> > SampleDirections;
-		GetSpacedVectors(SampleDirections);
+		GetSpacedVectors(View.Family->FrameNumber, SampleDirections);
 
 		for (int32 SampleIndex = 0; SampleIndex < NumConeSampleDirections; SampleIndex++)
 		{
@@ -1433,7 +1433,7 @@ public:
 			FAOSampleData2 AOSampleData;
 
 			TArray<FVector, TInlineAllocator<9> > SampleDirections;
-			GetSpacedVectors(SampleDirections);
+			GetSpacedVectors(View.Family->FrameNumber, SampleDirections);
 
 			for (int32 SampleIndex = 0; SampleIndex < NumConeSampleDirections; SampleIndex++)
 			{
diff --git a/Engine/Source/Runtime/Renderer/Private/IndirectLightingCache.cpp b/Engine/Source/Runtime/Renderer/Private/IndirectLightingCache.cpp
--- a/Engine/Source/Runtime/Renderer/Private/IndirectLightingCache.cpp
+++ b/Engine/Source/Runtime/Renderer/Private/IndirectLightingCache.cpp
@@ -462,14 +462,13 @@ bool FIndirectLightingCache::IndirectLightingAllowed(FScene* Scene, FSceneRender
 	return bAnyViewAllowsIndirectLightingCache;
 }
 
-void FIndirectLightingCache::ProcessPrimitiveUpdate(FScene* Scene, FViewInfo& View, int32 PrimitiveIndex, bool bAllowUnbuiltPreview, TMap<FIntVector, FBlockUpdateInfo>& OutBlocksToUpdate, TArray<FIndirectLightingCacheAllocation*>& OutTransitionsOverTimeToUpdate)
+void FIndirectLightingCache::ProcessPrimitiveUpdate(FScene* Scene, FViewInfo& View, int32 PrimitiveIndex, bool bAllowUnbuiltPreview, bool bAllowVolumeSample, TMap<FIntVector, FBlockUpdateInfo>& OutBlocksToUpdate, TArray<FIndirectLightingCacheAllocation*>& OutTransitionsOverTimeToUpdate)
 {
 	FPrimitiveSceneInfo* PrimitiveSceneInfo = Scene->Primitives[PrimitiveIndex];
 	const bool bPrecomputedLightingBufferWasDirty = PrimitiveSceneInfo->NeedsPrecomputedLightingBufferUpdate();
 
-	const FPrimitiveViewRelevance& PrimitiveRelevance = View.PrimitiveViewRelevanceMap[PrimitiveIndex];
 	const TMap<FPrimitiveComponentId, FAttachmentGroupSceneInfo>& AttachmentGroups = Scene->AttachmentGroups;
-	UpdateCachePrimitive(AttachmentGroups, PrimitiveSceneInfo, bAllowUnbuiltPreview, PrimitiveRelevance.bOpaqueRelevance, OutBlocksToUpdate, OutTransitionsOverTimeToUpdate);
+	UpdateCachePrimitive(AttachmentGroups, PrimitiveSceneInfo, bAllowUnbuiltPreview, bAllowVolumeSample, OutBlocksToUpdate, OutTransitionsOverTimeToUpdate);
 
 	// If it was already dirty, then the primitive is already in one of the view dirty primitive list at this point.
 	// This also ensures that a primitive does not get added twice to the list, which could create an array reallocation.
@@ -531,7 +530,9 @@ void FIndirectLightingCache::UpdateCachePrimitivesInternal(FScene* Scene, FScene
 				for (FSceneSetBitIterator BitIt(View.PrimitiveVisibilityMap); BitIt; ++BitIt)
 				{
 					uint32 PrimitiveIndex = BitIt.GetIndex();
-					ProcessPrimitiveUpdate(Scene, View, PrimitiveIndex, bAllowUnbuiltPreview, OutBlocksToUpdate, OutTransitionsOverTimeToUpdate);
+					// FDrawTranslucentMeshAction::AllowIndirectLightingCacheVolumeTexture doesn't allow volume samples on translucency, so we only need to support one if the primitive has at least one opaque material
+					const bool bAllowVolumeSample = View.PrimitiveViewRelevanceMap[PrimitiveIndex].bOpaqueRelevance;
+					ProcessPrimitiveUpdate(Scene, View, PrimitiveIndex, bAllowUnbuiltPreview, bAllowVolumeSample, OutBlocksToUpdate, OutTransitionsOverTimeToUpdate);
 				}
 
 				// Any visible primitives with an indirect shadow need their ILC updated, since that determines the indirect shadow direction
@@ -541,7 +542,7 @@ void FIndirectLightingCache::UpdateCachePrimitivesInternal(FScene* Scene, FScene
 
 					if (!View.PrimitiveVisibilityMap[PrimitiveIndex])
 					{
-						ProcessPrimitiveUpdate(Scene, View, PrimitiveIndex, bAllowUnbuiltPreview, OutBlocksToUpdate, OutTransitionsOverTimeToUpdate);
+						ProcessPrimitiveUpdate(Scene, View, PrimitiveIndex, bAllowUnbuiltPreview, true, OutBlocksToUpdate, OutTransitionsOverTimeToUpdate);
 					}
 				}
 			}			
@@ -647,11 +648,10 @@ void FIndirectLightingCache::UpdateCachePrimitive(
 	const TMap<FPrimitiveComponentId, FAttachmentGroupSceneInfo>& AttachmentGroups,
 	FPrimitiveSceneInfo* PrimitiveSceneInfo,
 	bool bAllowUnbuiltPreview,
-	bool bOpaqueRelevance,
+	bool bAllowVolumeSample,
 	TMap<FIntVector, FBlockUpdateInfo>& BlocksToUpdate,
 	TArray<FIndirectLightingCacheAllocation*>& TransitionsOverTimeToUpdate)
 {
-	
 	FPrimitiveSceneProxy* PrimitiveSceneProxy = PrimitiveSceneInfo->Proxy;	
 	FIndirectLightingCacheAllocation** PrimitiveAllocationPtr = PrimitiveAllocations.Find(PrimitiveSceneInfo->PrimitiveComponentId);
 	FIndirectLightingCacheAllocation* PrimitiveAllocation = PrimitiveAllocationPtr != NULL ? *PrimitiveAllocationPtr : NULL;
@@ -686,7 +686,7 @@ void FIndirectLightingCache::UpdateCachePrimitive(
 		{
 			FIndirectLightingCacheAllocation* OriginalAllocation = PrimitiveAllocation;
 			const bool bUnbuiltPreview = bAllowUnbuiltPreview && !bIsMovable;
-			const bool bPointSample = PrimitiveSceneProxy->GetIndirectLightingCacheQuality() == ILCQ_Point || bUnbuiltPreview || !bOpaqueRelevance;
+			const bool bPointSample = PrimitiveSceneProxy->GetIndirectLightingCacheQuality() == ILCQ_Point || bUnbuiltPreview || !bAllowVolumeSample;
 			const int32 BlockSize = bPointSample ? 1 : GLightingCacheMovableObjectAllocationSize;
 
 			// Light with the cumulative bounds of the entire attachment group
