[MAIN] [^] Merging fix for PS4 texture pool behaviour (original CL 3633040) from //UE4/Dev-Console to //UE4/Dev-Main
 - New behaviour is gated on a CVar. Default behaviour is the old one.
 - Licensees need to adjust their texture pool budgets when enabling the new behaviour.

diff --git a/Engine/Source/Runtime/Engine/Private/Texture2D.cpp b/Engine/Source/Runtime/Engine/Private/Texture2D.cpp
--- a/Engine/Source/Runtime/Engine/Private/Texture2D.cpp
+++ b/Engine/Source/Runtime/Engine/Private/Texture2D.cpp
@@ -1186,7 +1186,7 @@ void FTexture2DResource::InitRHI()
 	uint32 SizeY = OwnerMips[CurrentFirstMip].SizeY;
 
 	// Create the RHI texture.
-	uint32 TexCreateFlags = (Owner->SRGB ? TexCreate_SRGB : 0) | TexCreate_OfflineProcessed;
+	uint32 TexCreateFlags = (Owner->SRGB ? TexCreate_SRGB : 0) | TexCreate_OfflineProcessed | TexCreate_Streamable;
 	// if no miptail is available then create the texture without a packed miptail
 	if( Owner->GetMipTailBaseIndex() == -1 )
 	{
diff --git a/Engine/Source/Runtime/RHI/Public/RHIDefinitions.h b/Engine/Source/Runtime/RHI/Public/RHIDefinitions.h
--- a/Engine/Source/Runtime/RHI/Public/RHIDefinitions.h
+++ b/Engine/Source/Runtime/RHI/Public/RHIDefinitions.h
@@ -640,6 +640,8 @@ enum ETextureCreateFlags
 	TexCreate_NoFastClear = 1 << 25,
 	// Texture is a depth stencil resolve target
 	TexCreate_DepthStencilResolveTarget = 1 << 26,
+	// Flag used to indicted this texture is a streamable 2D texture, and should be counted towards the texture streaming pool budget.
+	TexCreate_Streamable = 1 << 27,
 	// Render target will not FinalizeFastClear; Caches and meta data will be flushed, but clearing will be skipped (avoids potentially trashing metadata)
 	TexCreate_NoFastClearFinalize = 1 << 28,
 	// Hint to the driver that this resource is managed properly by the engine for Alternate-Frame-Rendering in mGPU usage.
