Copying Orion-Staging @ 3748653 (Orion/Dev-General @ 3722124) to //UE4/Main

diff --git a/Engine/Source/Runtime/Engine/Public/ParticleHelper.h b/Engine/Source/Runtime/Engine/Public/ParticleHelper.h
--- a/Engine/Source/Runtime/Engine/Public/ParticleHelper.h
+++ b/Engine/Source/Runtime/Engine/Public/ParticleHelper.h
@@ -233,26 +233,26 @@ DECLARE_DWORD_COUNTER_STAT_EXTERN(TEXT("Sprite Ptcls Spawned"),STAT_SpritePartic
 DECLARE_DWORD_COUNTER_STAT_EXTERN(TEXT("Sprite Ptcls Updated"),STAT_SpriteParticlesUpdated,STATGROUP_Particles, );
 DECLARE_DWORD_COUNTER_STAT_EXTERN(TEXT("Sprite Ptcls Killed"),STAT_SpriteParticlesKilled,STATGROUP_Particles, );
 DECLARE_DWORD_COUNTER_STAT_EXTERN(TEXT("Particle Draw Calls"),STAT_ParticleDrawCalls,STATGROUP_Particles, );
-DECLARE_CYCLE_STAT_EXTERN(TEXT("Sort Time"),STAT_SortingTime,STATGROUP_Particles, );
-DECLARE_CYCLE_STAT_EXTERN(TEXT("Sprite Render Time"),STAT_SpriteRenderingTime,STATGROUP_Particles, );
-DECLARE_CYCLE_STAT_EXTERN(TEXT("Sprite Tick Time"),STAT_SpriteTickTime,STATGROUP_Particles, );
-DECLARE_CYCLE_STAT_EXTERN(TEXT("Sprite Spawn Time"),STAT_SpriteSpawnTime,STATGROUP_Particles, );
-DECLARE_CYCLE_STAT_EXTERN(TEXT("Sprite Update Time"),STAT_SpriteUpdateTime,STATGROUP_Particles, );
-DECLARE_CYCLE_STAT_EXTERN(TEXT("PSys Comp Tick Time"),STAT_PSysCompTickTime,STATGROUP_Particles, );
-DECLARE_CYCLE_STAT_EXTERN(TEXT("Particle Collision Time"),STAT_ParticleCollisionTime,STATGROUP_Particles, );
-DECLARE_CYCLE_STAT_EXTERN(TEXT("Particle SkelMeshSurf Time"),STAT_ParticleSkelMeshSurfTime,STATGROUP_Particles, );
-DECLARE_CYCLE_STAT_EXTERN(TEXT("Particle Pool Time"),STAT_ParticlePoolTime,STATGROUP_Particles, );
-DECLARE_CYCLE_STAT_EXTERN(TEXT("Particle Compute Time"),STAT_ParticleComputeTickTime,STATGROUP_Particles, );
-DECLARE_CYCLE_STAT_EXTERN(TEXT("Particle Finalize Time"),STAT_ParticleFinalizeTickTime,STATGROUP_Particles, );
+DECLARE_CYCLE_STAT_EXTERN(TEXT("Sort Time RT"),STAT_SortingTime,STATGROUP_Particles, );
+DECLARE_CYCLE_STAT_EXTERN(TEXT("Sprite Render Time RT"),STAT_SpriteRenderingTime,STATGROUP_Particles, );
+DECLARE_CYCLE_STAT_EXTERN(TEXT("Sprite Tick Time GT"),STAT_SpriteTickTime,STATGROUP_Particles, );
+DECLARE_CYCLE_STAT_EXTERN(TEXT("Sprite Spawn Time GT"),STAT_SpriteSpawnTime,STATGROUP_Particles, );
+DECLARE_CYCLE_STAT_EXTERN(TEXT("Sprite Update Time GT"),STAT_SpriteUpdateTime,STATGROUP_Particles, );
+DECLARE_CYCLE_STAT_EXTERN(TEXT("PSys Comp Tick Time GT"),STAT_PSysCompTickTime,STATGROUP_Particles, );
+DECLARE_CYCLE_STAT_EXTERN(TEXT("Particle Collision Time GT"),STAT_ParticleCollisionTime,STATGROUP_Particles, );
+DECLARE_CYCLE_STAT_EXTERN(TEXT("Particle SkelMeshSurf Time GT"),STAT_ParticleSkelMeshSurfTime,STATGROUP_Particles, );
+DECLARE_CYCLE_STAT_EXTERN(TEXT("Particle Pool Time RT"),STAT_ParticlePoolTime,STATGROUP_Particles, );
+DECLARE_CYCLE_STAT_EXTERN(TEXT("Particle Compute Time GT,AT"),STAT_ParticleComputeTickTime,STATGROUP_Particles, );
+DECLARE_CYCLE_STAT_EXTERN(TEXT("Particle Finalize Time GT"),STAT_ParticleFinalizeTickTime,STATGROUP_Particles, );
 DECLARE_CYCLE_STAT_EXTERN(TEXT("Particle GT Stall Time"),STAT_GTSTallTime,STATGROUP_Particles, );
-DECLARE_CYCLE_STAT_EXTERN(TEXT("Particle Render Time"),STAT_ParticleRenderingTime,STATGROUP_Particles, );
-DECLARE_CYCLE_STAT_EXTERN(TEXT("Particle Packing Time"),STAT_ParticlePackingTime,STATGROUP_Particles, );
-DECLARE_CYCLE_STAT_EXTERN(TEXT("SetTemplate Time"),STAT_ParticleSetTemplateTime,STATGROUP_Particles, );
-DECLARE_CYCLE_STAT_EXTERN(TEXT("Initialize Time"),STAT_ParticleInitializeTime,STATGROUP_Particles, );
-DECLARE_CYCLE_STAT_EXTERN(TEXT("Activate Time"),STAT_ParticleActivateTime,STATGROUP_Particles, );
+DECLARE_CYCLE_STAT_EXTERN(TEXT("Particle Render Time RT"),STAT_ParticleRenderingTime,STATGROUP_Particles, );
+DECLARE_CYCLE_STAT_EXTERN(TEXT("Particle Packing Time RT"),STAT_ParticlePackingTime,STATGROUP_Particles, );
+DECLARE_CYCLE_STAT_EXTERN(TEXT("SetTemplate Time GT"),STAT_ParticleSetTemplateTime,STATGROUP_Particles, );
+DECLARE_CYCLE_STAT_EXTERN(TEXT("Initialize Time GT"),STAT_ParticleInitializeTime,STATGROUP_Particles, );
+DECLARE_CYCLE_STAT_EXTERN(TEXT("Activate Time GT"),STAT_ParticleActivateTime,STATGROUP_Particles, );
 DECLARE_CYCLE_STAT_EXTERN(TEXT("Async Work Time"),STAT_ParticleAsyncTime,STATGROUP_Particles, );           // regardless of if it is actually performed on other threads or not
 DECLARE_CYCLE_STAT_EXTERN(TEXT("Wait For ASync Time"),STAT_ParticleAsyncWaitTime,STATGROUP_Particles, );   // can be either performed on this thread or a true wait
-DECLARE_CYCLE_STAT_EXTERN(TEXT("Update Bounds Time"),STAT_ParticleUpdateBounds,STATGROUP_Particles, );
+DECLARE_CYCLE_STAT_EXTERN(TEXT("Update Bounds Time GT"),STAT_ParticleUpdateBounds,STATGROUP_Particles, );
 
 DECLARE_CYCLE_STAT_EXTERN(TEXT("Particle Memory Time"),STAT_ParticleMemTime,STATGROUP_ParticleMem, );
 DECLARE_MEMORY_STAT_EXTERN(TEXT("Ptcls Data GT Mem"),STAT_GTParticleData,STATGROUP_ParticleMem, );
@@ -316,12 +316,14 @@ DECLARE_CYCLE_STAT_EXTERN(TEXT("Cull Vector Fields"),STAT_GPUParticleVFCullTime,
 DECLARE_CYCLE_STAT_EXTERN(TEXT("Misc1"),STAT_GPUParticleMisc1,STATGROUP_GPUParticles, );
 DECLARE_CYCLE_STAT_EXTERN(TEXT("Misc2"),STAT_GPUParticleMisc2,STATGROUP_GPUParticles, );
 DECLARE_CYCLE_STAT_EXTERN(TEXT("Misc3"),STAT_GPUParticleMisc3,STATGROUP_GPUParticles, );
+DECLARE_CYCLE_STAT_EXTERN(TEXT("Injection Tick Time"), STAT_GPUParticlesInjectionTime, STATGROUP_GPUParticles, );
+DECLARE_CYCLE_STAT_EXTERN(TEXT("Execute Simulation Commands Time"), STAT_GPUParticlesSimulationCommands, STATGROUP_GPUParticles, );
 DECLARE_DWORD_COUNTER_STAT_EXTERN(TEXT("Single Iteration Emitters"), STAT_GPUSingleIterationEmitters, STATGROUP_GPUParticles, );
 DECLARE_DWORD_COUNTER_STAT_EXTERN(TEXT("Multi Iterations Emitters"), STAT_GPUMultiIterationsEmitters, STATGROUP_GPUParticles, );
 
 DECLARE_DWORD_COUNTER_STAT_EXTERN(TEXT("Mesh Particles"),STAT_MeshParticles,STATGROUP_Particles, );
-DECLARE_CYCLE_STAT_EXTERN(TEXT("Mesh Render Time"),STAT_MeshRenderingTime,STATGROUP_Particles, );
-DECLARE_CYCLE_STAT_EXTERN(TEXT("Mesh Tick Time"),STAT_MeshTickTime,STATGROUP_Particles, );
+DECLARE_CYCLE_STAT_EXTERN(TEXT("Mesh Render Time RT"),STAT_MeshRenderingTime,STATGROUP_Particles, );
+DECLARE_CYCLE_STAT_EXTERN(TEXT("Mesh Tick Time GT,AT"),STAT_MeshTickTime,STATGROUP_Particles, );
 
 /**
  * Per-particle data sent to the GPU.
@@ -520,11 +522,11 @@ DECLARE_DWORD_COUNTER_STAT_EXTERN(TEXT("Trail Ptcls Spawned"),STAT_TrailParticle
 DECLARE_DWORD_COUNTER_STAT_EXTERN(TEXT("Trail Tick Calls"),STAT_TrailParticlesTickCalls,STATGROUP_TrailParticles, );
 DECLARE_DWORD_COUNTER_STAT_EXTERN(TEXT("Trail Ptcls Killed"),STAT_TrailParticlesKilled,STATGROUP_Particles, );
 DECLARE_DWORD_COUNTER_STAT_EXTERN(TEXT("Trail Ptcl Tris"),STAT_TrailParticlesTrianglesRendered,STATGROUP_Particles, );
-DECLARE_CYCLE_STAT_EXTERN(TEXT("Trail FillVertex Time"),STAT_TrailFillVertexTime,STATGROUP_TrailParticles, );
-DECLARE_CYCLE_STAT_EXTERN(TEXT("Trail FillIndex Time"),STAT_TrailFillIndexTime,STATGROUP_TrailParticles, );
-DECLARE_CYCLE_STAT_EXTERN(TEXT("Trail Render Time"),STAT_TrailRenderingTime,STATGROUP_Particles, );
-DECLARE_CYCLE_STAT_EXTERN(TEXT("Trail Tick Time"),STAT_TrailTickTime,STATGROUP_Particles, );
-DECLARE_CYCLE_STAT_EXTERN(TEXT("AnimTrail Notify Time"),STAT_AnimTrailNotifyTime,STATGROUP_Particles, );
+DECLARE_CYCLE_STAT_EXTERN(TEXT("Trail FillVertex Time RT"),STAT_TrailFillVertexTime,STATGROUP_TrailParticles, );
+DECLARE_CYCLE_STAT_EXTERN(TEXT("Trail FillIndex Time RT"),STAT_TrailFillIndexTime,STATGROUP_TrailParticles, );
+DECLARE_CYCLE_STAT_EXTERN(TEXT("Trail Render Time RT"),STAT_TrailRenderingTime,STATGROUP_Particles, );
+DECLARE_CYCLE_STAT_EXTERN(TEXT("Trail Tick Time GT"),STAT_TrailTickTime,STATGROUP_Particles, );
+DECLARE_CYCLE_STAT_EXTERN(TEXT("AnimTrail Notify Time GT"),STAT_AnimTrailNotifyTime,STATGROUP_Particles, );
 
 /**
  * Beam particle stats
@@ -536,11 +538,11 @@ DECLARE_DWORD_COUNTER_STAT_EXTERN(TEXT("Beam Ptcl Update Calls"),STAT_BeamPartic
 DECLARE_DWORD_COUNTER_STAT_EXTERN(TEXT("Beam Ptcls Updated"),STAT_BeamParticlesUpdated,STATGROUP_BeamParticles, );
 DECLARE_DWORD_COUNTER_STAT_EXTERN(TEXT("Beam Ptcls Killed"),STAT_BeamParticlesKilled,STATGROUP_Particles, );
 DECLARE_DWORD_COUNTER_STAT_EXTERN(TEXT("Beam Ptcl Tris"),STAT_BeamParticlesTrianglesRendered,STATGROUP_Particles, );
-DECLARE_CYCLE_STAT_EXTERN(TEXT("Beam Spawn Time"),STAT_BeamSpawnTime,STATGROUP_Particles, );
-DECLARE_CYCLE_STAT_EXTERN(TEXT("Beam FillVertex Time"),STAT_BeamFillVertexTime,STATGROUP_BeamParticles, );
-DECLARE_CYCLE_STAT_EXTERN(TEXT("Beam FillIndex Time"),STAT_BeamFillIndexTime,STATGROUP_BeamParticles, );
-DECLARE_CYCLE_STAT_EXTERN(TEXT("Beam Render Time"),STAT_BeamRenderingTime,STATGROUP_Particles, );
-DECLARE_CYCLE_STAT_EXTERN(TEXT("Beam Tick Time"),STAT_BeamTickTime,STATGROUP_Particles, );
+DECLARE_CYCLE_STAT_EXTERN(TEXT("Beam Spawn Time GT"),STAT_BeamSpawnTime,STATGROUP_Particles, );
+DECLARE_CYCLE_STAT_EXTERN(TEXT("Beam FillVertex Time RT"),STAT_BeamFillVertexTime,STATGROUP_BeamParticles, );
+DECLARE_CYCLE_STAT_EXTERN(TEXT("Beam FillIndex Time RT"),STAT_BeamFillIndexTime,STATGROUP_BeamParticles, );
+DECLARE_CYCLE_STAT_EXTERN(TEXT("Beam Render Time RT"),STAT_BeamRenderingTime,STATGROUP_Particles, );
+DECLARE_CYCLE_STAT_EXTERN(TEXT("Beam Tick Time GT"),STAT_BeamTickTime,STATGROUP_Particles, );
 
 //
 //	Helper structures for payload data...
@@ -2468,18 +2470,41 @@ public:
 //	Scene Proxies
 //
 
-class FParticleSystemSceneProxy : public FPrimitiveSceneProxy
+class FParticleSystemSceneProxy final : public FPrimitiveSceneProxy
 {
 public:
+	SIZE_T GetTypeHash() const override;
+
 	/** Initialization constructor. */
-	FParticleSystemSceneProxy(const UParticleSystemComponent* Component, FParticleDynamicData* InDynamicData);
+	FParticleSystemSceneProxy(const UParticleSystemComponent* Component, FParticleDynamicData* InDynamicData, bool bCanBeOccluded);
 	virtual ~FParticleSystemSceneProxy();
 
 	// FPrimitiveSceneProxy interface.
 	virtual bool CanBeOccluded() const override
 	{
-		return false;
+		return  bCanBeOccluded ? !MaterialRelevance.bDisableDepthTest : false;
+	}
+
+	/**
+	*	Returns whether the proxy utilizes custom occlusion bounds or not
+	*
+	*	@return	bool		true if custom occlusion bounds are used, false if not;
+	*/
+	virtual bool HasCustomOcclusionBounds() const override
+	{
+		return bCanBeOccluded ? bHasCustomOcclusionBounds : FPrimitiveSceneProxy::HasCustomOcclusionBounds();
 	}
+
+	/**
+	*	Return the custom occlusion bounds for this scene proxy.
+	*
+	*	@return	FBoxSphereBounds		The custom occlusion bounds.
+	*/
+	virtual FBoxSphereBounds GetCustomOcclusionBounds() const override
+	{
+		return bCanBeOccluded ? OcclusionBounds.TransformBy(GetLocalToWorld()) : FPrimitiveSceneProxy::GetCustomOcclusionBounds();
+	}
+
 	virtual void GetDynamicMeshElements(const TArray<const FSceneView*>& Views, const FSceneViewFamily& ViewFamily, uint32 VisibilityMap, FMeshElementCollector& Collector) const override;
 	virtual FPrimitiveViewRelevance GetViewRelevance(const FSceneView* View) const override;
 	virtual void OnTransformChanged() override;
@@ -2662,44 +2687,9 @@ protected:
 	ERHIFeatureLevel::Type FeatureLevel;
 
 	friend struct FDynamicSpriteEmitterDataBase;
-};
-
-class FParticleSystemOcclusionSceneProxy : public FParticleSystemSceneProxy
-{
-public:
-	/** Initialization constructor. */
-	FParticleSystemOcclusionSceneProxy(const UParticleSystemComponent* Component, FParticleDynamicData* InDynamicData);
-	virtual ~FParticleSystemOcclusionSceneProxy();
-
-	// FPrimitiveSceneProxy interface.
-	/** @return true if the proxy requires occlusion queries */
-	virtual bool CanBeOccluded() const override
-	{
-		return !MaterialRelevance.bDisableDepthTest;
-	}
-	
-	/**
-	 *	Returns whether the proxy utilizes custom occlusion bounds or not
-	 *
-	 *	@return	bool		true if custom occlusion bounds are used, false if not;
-	 */
-	virtual bool HasCustomOcclusionBounds() const override
-	{
-		return bHasCustomOcclusionBounds;
-	}
-
-	/**
-	 *	Return the custom occlusion bounds for this scene proxy.
-	 *	
-	 *	@return	FBoxSphereBounds		The custom occlusion bounds.
-	 */
-	virtual FBoxSphereBounds GetCustomOcclusionBounds() const override
-	{
-		return OcclusionBounds.TransformBy(GetLocalToWorld());
-	}
 
 private:
-
+	uint32	bCanBeOccluded : 1;
 	uint32	bHasCustomOcclusionBounds : 1;
 
 	/** Bounds for occlusion rendering. */
diff --git a/Engine/Source/Runtime/Engine/Public/PrecomputedVolumetricLightmap.h b/Engine/Source/Runtime/Engine/Public/PrecomputedVolumetricLightmap.h
--- a/Engine/Source/Runtime/Engine/Public/PrecomputedVolumetricLightmap.h
+++ b/Engine/Source/Runtime/Engine/Public/PrecomputedVolumetricLightmap.h
@@ -40,7 +40,7 @@ public:
 
 	virtual void Discard() override
 	{
-		// Keep backing data in the editor for saving
+		//temporarily remove until the semantics around lighting scenarios are fixed.
 		if (!GIsEditor)
 		{
 			Data.Empty();
diff --git a/Engine/Source/Runtime/Engine/Public/PrimitiveSceneProxy.h b/Engine/Source/Runtime/Engine/Public/PrimitiveSceneProxy.h
--- a/Engine/Source/Runtime/Engine/Public/PrimitiveSceneProxy.h
+++ b/Engine/Source/Runtime/Engine/Public/PrimitiveSceneProxy.h
@@ -128,6 +128,9 @@ public:
 	/** Virtual destructor. */
 	ENGINE_API virtual ~FPrimitiveSceneProxy();
 
+	/** Return a type (or subtype) specific hash for sorting purposes */
+	ENGINE_API virtual SIZE_T GetTypeHash() const = 0;
+
 	/**
 	 * Updates selection for the primitive proxy. This simply sends a message to the rendering thread to call SetSelection_RenderThread.
 	 * This is called in the game thread as selection is toggled.
diff --git a/Engine/Source/Runtime/Engine/Public/StaticMeshResources.h b/Engine/Source/Runtime/Engine/Public/StaticMeshResources.h
--- a/Engine/Source/Runtime/Engine/Public/StaticMeshResources.h
+++ b/Engine/Source/Runtime/Engine/Public/StaticMeshResources.h
@@ -500,6 +500,7 @@ private:
 class ENGINE_API FStaticMeshSceneProxy : public FPrimitiveSceneProxy
 {
 public:
+	SIZE_T GetTypeHash() const override;
 
 	/** Initialization constructor. */
 	FStaticMeshSceneProxy(UStaticMeshComponent* Component, bool bForceLODsShareStaticLighting);
diff --git a/Engine/Source/Runtime/Foliage/Private/FoliageComponent.cpp b/Engine/Source/Runtime/Foliage/Private/FoliageComponent.cpp
--- a/Engine/Source/Runtime/Foliage/Private/FoliageComponent.cpp
+++ b/Engine/Source/Runtime/Foliage/Private/FoliageComponent.cpp
@@ -15,9 +15,14 @@
 #include "Engine/StaticMesh.h"
 
 /** Scene proxy class for UInteractiveFoliageComponent. */
-class FInteractiveFoliageSceneProxy : public FStaticMeshSceneProxy
+class FInteractiveFoliageSceneProxy final : public FStaticMeshSceneProxy
 {
 public:
+	SIZE_T GetTypeHash() const override
+	{
+		static size_t UniquePointer;
+		return reinterpret_cast<size_t>(&UniquePointer);
+	}
 
 	FInteractiveFoliageSceneProxy(UInteractiveFoliageComponent* InComponent) :
 		FStaticMeshSceneProxy(InComponent, false),
diff --git a/Engine/Source/Runtime/GameplayTags/Classes/GameplayTagsManager.h b/Engine/Source/Runtime/GameplayTags/Classes/GameplayTagsManager.h
--- a/Engine/Source/Runtime/GameplayTags/Classes/GameplayTagsManager.h
+++ b/Engine/Source/Runtime/GameplayTags/Classes/GameplayTagsManager.h
@@ -227,10 +227,11 @@ class GAMEPLAYTAGS_API UGameplayTagsManager : public UObject
 	 * This can only be called during engine initialization, the table needs to be locked down before replication
 	 *
 	 * @param TagName The Name of the tag to add
+	 * @param TagDevComment The developer comment clarifying the usage of the tag
 	 * 
 	 * @return Will return the corresponding FGameplayTag
 	 */
-	FGameplayTag AddNativeGameplayTag(FName TagName);
+	FGameplayTag AddNativeGameplayTag(FName TagName, const FString& TagDevComment = TEXT("(Native)"));
 
 	/** Call to flush the list of native tags, once called it is unsafe to add more */
 	void DoneAddingNativeTags();
@@ -378,6 +379,9 @@ class GAMEPLAYTAGS_API UGameplayTagsManager : public UObject
 	/** Handles redirectors for a single tag, will also error on invalid tag. This is only called for when individual tags are serialized on their own */
 	void RedirectSingleGameplayTag(FGameplayTag& Tag, UProperty* SerializingProperty) const;
 
+	/** Handles establishing a single tag from an imported tag name (accounts for redirects too). Called when tags are imported via text. */
+	bool ImportSingleGameplayTag(FGameplayTag& Tag, FName ImportedTagName) const;
+
 	/** Gets a tag name from net index and vice versa, used for replication efficiency */
 	FName GetTagNameFromNetIndex(FGameplayTagNetIndex Index) const;
 	FGameplayTagNetIndex GetNetIndexFromTag(const FGameplayTag &InTag) const;
@@ -396,6 +400,8 @@ class GAMEPLAYTAGS_API UGameplayTagsManager : public UObject
 
 	const TArray<TSharedPtr<FGameplayTagNode>>& GetNetworkGameplayTagNodeIndex() const { return NetworkGameplayTagNodeIndex; }
 
+	bool IsNativelyAddedTag(FGameplayTag Tag) const;
+
 #if WITH_EDITOR
 	/** Gets a Filtered copy of the GameplayRootTags Array based on the comma delimited filter string passed in */
 	void GetFilteredGameplayRootTags(const FString& InFilterString, TArray< TSharedPtr<FGameplayTagNode> >& OutTagArray) const;
diff --git a/Engine/Source/Runtime/GameplayTags/Classes/GameplayTagsSettings.h b/Engine/Source/Runtime/GameplayTags/Classes/GameplayTagsSettings.h
--- a/Engine/Source/Runtime/GameplayTags/Classes/GameplayTagsSettings.h
+++ b/Engine/Source/Runtime/GameplayTags/Classes/GameplayTagsSettings.h
@@ -25,6 +25,12 @@ struct GAMEPLAYTAGS_API FGameplayTagRedirect
 	{
 		return A.OldTagName == B.OldTagName && A.NewTagName == B.NewTagName;
 	}
+
+	// This enables lookups by old tag name via FindByKey
+	bool operator==(FName OtherOldTagName) const
+	{
+		return OldTagName == OtherOldTagName;
+	}
 };
 
 /** Category remapping. This allows base engine tag category meta data to remap to multiple project-specific categories. */
diff --git a/Engine/Source/Runtime/GameplayTags/Private/GameplayTagContainer.cpp b/Engine/Source/Runtime/GameplayTags/Private/GameplayTagContainer.cpp
--- a/Engine/Source/Runtime/GameplayTags/Private/GameplayTagContainer.cpp
+++ b/Engine/Source/Runtime/GameplayTags/Private/GameplayTagContainer.cpp
@@ -1319,16 +1319,7 @@ bool FGameplayTag::ImportTextItem(const TCHAR*& Buffer, int32 PortFlags, UObject
 		return false;
 	}
 
-	FName ImportedTagName = FName(*ImportedTag);
-	if (UGameplayTagsManager::Get().ValidateTagCreation(ImportedTagName))
-	{
-		// We found the tag. Assign it here.
-		TagName = ImportedTagName;
-		return true;
-	}
-
-	// Let normal ImportText try.
-	return false;
+	return UGameplayTagsManager::Get().ImportSingleGameplayTag(*this, FName(*ImportedTag));
 }
 
 void FGameplayTag::FromExportString(const FString& ExportString)
diff --git a/Engine/Source/Runtime/GameplayTags/Private/GameplayTagsManager.cpp b/Engine/Source/Runtime/GameplayTags/Private/GameplayTagsManager.cpp
--- a/Engine/Source/Runtime/GameplayTags/Private/GameplayTagsManager.cpp
+++ b/Engine/Source/Runtime/GameplayTags/Private/GameplayTagsManager.cpp
@@ -525,6 +525,25 @@ void UGameplayTagsManager::RedirectSingleGameplayTag(FGameplayTag& Tag, UPropert
 #endif
 }
 
+bool UGameplayTagsManager::ImportSingleGameplayTag(FGameplayTag& Tag, FName ImportedTagName) const
+{
+	if (const FGameplayTag* RedirectedTag = TagRedirects.Find(ImportedTagName))
+	{
+		Tag = *RedirectedTag;
+		return true;
+	}
+	else if (ValidateTagCreation(ImportedTagName))
+	{
+		// The tag name is valid
+		Tag.TagName = ImportedTagName;
+		return true;
+	}
+
+	// No valid tag established in this attempt
+	Tag.TagName = NAME_None;
+	return false;
+}
+
 void UGameplayTagsManager::InitializeManager()
 {
 	check(!SingletonManager);
@@ -612,6 +631,11 @@ void UGameplayTagsManager::DestroyGameplayTagTree()
 	}
 }
 
+bool UGameplayTagsManager::IsNativelyAddedTag(FGameplayTag Tag) const
+{
+	return NativeTagsToAdd.Contains(Tag.GetTagName());
+}
+
 int32 UGameplayTagsManager::InsertTagIntoNodeArray(FName Tag, TSharedPtr<FGameplayTagNode> ParentNode, TArray< TSharedPtr<FGameplayTagNode> >& NodeArray, FName SourceName, const FString& DevComment)
 {
 	int32 InsertionIdx = INDEX_NONE;
@@ -1165,7 +1189,7 @@ FGameplayTag UGameplayTagsManager::FindGameplayTagFromPartialString_Slow(FString
 	return FoundTag;
 }
 
-FGameplayTag UGameplayTagsManager::AddNativeGameplayTag(FName TagName)
+FGameplayTag UGameplayTagsManager::AddNativeGameplayTag(FName TagName, const FString& TagDevComment)
 {
 	if (TagName.IsNone())
 	{
@@ -1182,7 +1206,7 @@ FGameplayTag UGameplayTagsManager::AddNativeGameplayTag(FName TagName)
 			NativeTagsToAdd.Add(TagName);
 		}
 
-		AddTagTableRow(FGameplayTagTableRow(TagName), FGameplayTagSource::GetNativeName());
+		AddTagTableRow(FGameplayTagTableRow(TagName, TagDevComment), FGameplayTagSource::GetNativeName());
 
 		return NewTag;
 	}
@@ -1219,17 +1243,18 @@ FSimpleMulticastDelegate& UGameplayTagsManager::OnLastChanceToAddNativeTags()
 
 void UGameplayTagsManager::DoneAddingNativeTags()
 {
-	// Safe to call multiple times, only works the first time
-	if (!bDoneAddingNativeTags)
+	// Safe to call multiple times, only works the first time, must be called after the engine
+	// is initialized (DoneAddingNativeTags is bound to PostEngineInit to cover anything that's skipped).
+	if (GEngine && !bDoneAddingNativeTags)
 	{
 		UE_LOG(LogGameplayTags, Display, TEXT("UGameplayTagsManager::DoneAddingNativeTags. DelegateIsBound: %d"), (int32)OnLastChanceToAddNativeTags().IsBound());
 		OnLastChanceToAddNativeTags().Broadcast();
 		bDoneAddingNativeTags = true;
 
-		if (ShouldUseFastReplication())
-		{
-			ConstructNetIndex();
-		}
+		// We may add native tags that are needed for redirectors, so reconstruct the GameplayTag tree
+		DestroyGameplayTagTree();
+		ConstructGameplayTagTree();
+
 		OnDoneAddingNativeTagsDelegate().Broadcast();
 	}
 }
diff --git a/Engine/Source/Runtime/GeometryCache/Private/GeometryCacheSceneProxy.cpp b/Engine/Source/Runtime/GeometryCache/Private/GeometryCacheSceneProxy.cpp
--- a/Engine/Source/Runtime/GeometryCache/Private/GeometryCacheSceneProxy.cpp
+++ b/Engine/Source/Runtime/GeometryCache/Private/GeometryCacheSceneProxy.cpp
@@ -74,6 +74,12 @@ FGeometryCacheSceneProxy::~FGeometryCacheSceneProxy()
 	Sections.Empty();
 }
 
+SIZE_T FGeometryCacheSceneProxy::GetTypeHash() const
+{
+	static size_t UniquePointer;
+	return reinterpret_cast<size_t>(&UniquePointer);
+}
+
 void FGeometryCacheSceneProxy::GetDynamicMeshElements(const TArray<const FSceneView*>& Views, const FSceneViewFamily& ViewFamily, uint32 VisibilityMap, FMeshElementCollector& Collector) const
 {
 	SCOPE_CYCLE_COUNTER(STAT_GeometryCacheSceneProxy_GetMeshElements);
diff --git a/Engine/Source/Runtime/GeometryCache/Private/GeometryCacheSceneProxy.h b/Engine/Source/Runtime/GeometryCache/Private/GeometryCacheSceneProxy.h
--- a/Engine/Source/Runtime/GeometryCache/Private/GeometryCacheSceneProxy.h
+++ b/Engine/Source/Runtime/GeometryCache/Private/GeometryCacheSceneProxy.h
@@ -65,10 +65,11 @@ public:
 };
 
 /** Procedural mesh scene proxy */
-class GEOMETRYCACHE_API FGeometryCacheSceneProxy : public FPrimitiveSceneProxy
+class GEOMETRYCACHE_API FGeometryCacheSceneProxy final : public FPrimitiveSceneProxy
 {
 public:
-	
+	SIZE_T GetTypeHash() const override;
+
 	FGeometryCacheSceneProxy(class UGeometryCacheComponent* Component);
 
 	virtual ~FGeometryCacheSceneProxy();
diff --git a/Engine/Source/Runtime/Landscape/Private/LandscapeGizmoActor.cpp b/Engine/Source/Runtime/Landscape/Private/LandscapeGizmoActor.cpp
--- a/Engine/Source/Runtime/Landscape/Private/LandscapeGizmoActor.cpp
+++ b/Engine/Source/Runtime/Landscape/Private/LandscapeGizmoActor.cpp
@@ -125,9 +125,15 @@ public:
 };
 
 /** Represents a LandscapeGizmoRenderingComponent to the scene manager. */
-class FLandscapeGizmoRenderSceneProxy : public FPrimitiveSceneProxy
+class FLandscapeGizmoRenderSceneProxy final : public FPrimitiveSceneProxy
 {
 public:
+	SIZE_T GetTypeHash() const override
+	{
+		static size_t UniquePointer;
+		return reinterpret_cast<size_t>(&UniquePointer);
+	}
+
 	FMatrix MeshRT;
 	FVector XAxis, YAxis, Origin;
 	FVector FrustumVerts[8];
diff --git a/Engine/Source/Runtime/Landscape/Private/LandscapeRender.cpp b/Engine/Source/Runtime/Landscape/Private/LandscapeRender.cpp
--- a/Engine/Source/Runtime/Landscape/Private/LandscapeRender.cpp
+++ b/Engine/Source/Runtime/Landscape/Private/LandscapeRender.cpp
@@ -985,6 +985,12 @@ void FLandscapeComponentSceneProxy::GetLightRelevance(const FLightSceneProxy* Li
 	}
 }
 
+SIZE_T FLandscapeComponentSceneProxy::GetTypeHash() const
+{
+	static size_t UniquePointer;
+	return reinterpret_cast<size_t>(&UniquePointer);
+}
+
 FLightInteraction FLandscapeComponentSceneProxy::FLandscapeLCI::GetInteraction(const class FLightSceneProxy* LightSceneProxy) const
 {
 	// ask base class
@@ -3169,6 +3175,13 @@ FLandscapeMeshProxySceneProxy::FLandscapeMeshProxySceneProxy(UStaticMeshComponen
 	}
 }
 
+SIZE_T FLandscapeMeshProxySceneProxy::GetTypeHash() const
+{
+	static size_t UniquePointer;
+	return reinterpret_cast<size_t>(&UniquePointer);
+}
+
+
 void FLandscapeMeshProxySceneProxy::CreateRenderThreadResources()
 {
 	FStaticMeshSceneProxy::CreateRenderThreadResources();
diff --git a/Engine/Source/Runtime/Landscape/Private/LandscapeRenderMobile.cpp b/Engine/Source/Runtime/Landscape/Private/LandscapeRenderMobile.cpp
--- a/Engine/Source/Runtime/Landscape/Private/LandscapeRenderMobile.cpp
+++ b/Engine/Source/Runtime/Landscape/Private/LandscapeRenderMobile.cpp
@@ -274,6 +274,12 @@ FLandscapeComponentSceneProxyMobile::~FLandscapeComponentSceneProxyMobile()
 	}
 }
 
+SIZE_T FLandscapeComponentSceneProxyMobile::GetTypeHash() const
+{
+	static size_t UniquePointer;
+	return reinterpret_cast<size_t>(&UniquePointer);
+}
+
 void FLandscapeComponentSceneProxyMobile::CreateRenderThreadResources()
 {
 	auto FeatureLevel = GetScene().GetFeatureLevel();
diff --git a/Engine/Source/Runtime/Landscape/Private/LandscapeSplines.cpp b/Engine/Source/Runtime/Landscape/Private/LandscapeSplines.cpp
--- a/Engine/Source/Runtime/Landscape/Private/LandscapeSplines.cpp
+++ b/Engine/Source/Runtime/Landscape/Private/LandscapeSplines.cpp
@@ -51,9 +51,15 @@ IMPLEMENT_HIT_PROXY(HLandscapeSplineProxy_Tangent, HLandscapeSplineProxy);
 
 /** Represents a ULandscapeSplinesComponent to the scene manager. */
 #if WITH_EDITOR
-class FLandscapeSplinesSceneProxy : public FPrimitiveSceneProxy
+class FLandscapeSplinesSceneProxy final : public FPrimitiveSceneProxy
 {
 private:
+	SIZE_T GetTypeHash() const override
+	{
+		static size_t UniquePointer;
+		return reinterpret_cast<size_t>(&UniquePointer);
+	}
+
 	const FLinearColor	SplineColor;
 
 	const UTexture2D*	ControlPointSprite;
diff --git a/Engine/Source/Runtime/Landscape/Public/LandscapeRender.h b/Engine/Source/Runtime/Landscape/Public/LandscapeRender.h
--- a/Engine/Source/Runtime/Landscape/Public/LandscapeRender.h
+++ b/Engine/Source/Runtime/Landscape/Public/LandscapeRender.h
@@ -405,10 +405,12 @@ public:
 //
 // FLandscapeMeshProxySceneProxy
 //
-class FLandscapeMeshProxySceneProxy : public FStaticMeshSceneProxy
+class FLandscapeMeshProxySceneProxy final : public FStaticMeshSceneProxy
 {
 	TArray<FLandscapeNeighborInfo> ProxyNeighborInfos;
 public:
+	SIZE_T GetTypeHash() const override;
+
 	FLandscapeMeshProxySceneProxy(UStaticMeshComponent* InComponent, const FGuid& InGuid, const TArray<FIntPoint>& InProxyComponentBases, int8 InProxyLOD);
 	virtual ~FLandscapeMeshProxySceneProxy();
 	virtual void CreateRenderThreadResources() override;
@@ -423,6 +425,7 @@ class FLandscapeComponentSceneProxy : public FPrimitiveSceneProxy, public FLands
 {
 	friend class FLandscapeSharedBuffers;
 
+	SIZE_T GetTypeHash() const override;
 
 	class FLandscapeLCI final : public FLightCacheInterface
 	{
diff --git a/Engine/Source/Runtime/Landscape/Public/LandscapeRenderMobile.h b/Engine/Source/Runtime/Landscape/Public/LandscapeRenderMobile.h
--- a/Engine/Source/Runtime/Landscape/Public/LandscapeRenderMobile.h
+++ b/Engine/Source/Runtime/Landscape/Public/LandscapeRenderMobile.h
@@ -116,13 +116,15 @@ public:
 //
 // FLandscapeComponentSceneProxy
 //
-class FLandscapeComponentSceneProxyMobile : public FLandscapeComponentSceneProxy
+class FLandscapeComponentSceneProxyMobile final : public FLandscapeComponentSceneProxy
 {
 	TSharedPtr<FLandscapeMobileRenderData, ESPMode::ThreadSafe> MobileRenderData;
 
 	virtual ~FLandscapeComponentSceneProxyMobile();
 
 public:
+	SIZE_T GetTypeHash() const override;
+
 	FLandscapeComponentSceneProxyMobile(ULandscapeComponent* InComponent);
 
 	virtual void CreateRenderThreadResources() override;
diff --git a/Engine/Source/Runtime/Launch/Private/LaunchEngineLoop.cpp b/Engine/Source/Runtime/Launch/Private/LaunchEngineLoop.cpp
--- a/Engine/Source/Runtime/Launch/Private/LaunchEngineLoop.cpp
+++ b/Engine/Source/Runtime/Launch/Private/LaunchEngineLoop.cpp
@@ -931,18 +931,6 @@ int32 FEngineLoop::PreInit(const TCHAR* CmdLine)
 		return -1;
 	}
 
-	// Check for special instruction cpu support if needed
-	if (FPlatformMisc::NeedsNonoptionalCPUFeaturesCheck())
-	{
-		bool bHasNonoptionalCPUFeatures = FPlatformMisc::HasNonoptionalCPUFeatures();
-		// If it's not supported, we'll crash later so better to return an error
-		if (!bHasNonoptionalCPUFeatures)
-		{
-			FMessageDialog::Open(EAppMsgType::Ok, LOCTEXT("RequiresNonoptionalCPUFeatures", "Error: This application requires a CPU that supports the specific instruction set(s)"));
-			return -1;
-		}
-	}
-
 #if WITH_ENGINE
 	FCoreUObjectDelegates::PostGarbageCollectConditionalBeginDestroy.AddStatic(DeferredPhysResourceCleanup);
 #endif
diff --git a/Engine/Source/Runtime/MRMesh/Private/MRMeshComponent.cpp b/Engine/Source/Runtime/MRMesh/Private/MRMeshComponent.cpp
--- a/Engine/Source/Runtime/MRMesh/Private/MRMeshComponent.cpp
+++ b/Engine/Source/Runtime/MRMesh/Private/MRMeshComponent.cpp
@@ -187,9 +187,15 @@ static void InitVertexFactory(FLocalVertexFactory* VertexFactory, const FMRMeshP
 }
 
 
-class FMRMeshProxy : public FPrimitiveSceneProxy
+class FMRMeshProxy final : public FPrimitiveSceneProxy
 {
 public:
+	SIZE_T GetTypeHash() const override
+	{
+		static size_t UniquePointer;
+		return reinterpret_cast<size_t>(&UniquePointer);
+	}
+
 	FMRMeshProxy(const UMRMeshComponent* InComponent)
 	: FPrimitiveSceneProxy(InComponent)
 	, MaterialToUse((InComponent->Material!=nullptr) ? InComponent->Material : UMaterial::GetDefaultMaterial(MD_Surface) )
diff --git a/Engine/Source/Runtime/Online/HTTP/HTTP.Build.cs b/Engine/Source/Runtime/Online/HTTP/HTTP.Build.cs
--- a/Engine/Source/Runtime/Online/HTTP/HTTP.Build.cs
+++ b/Engine/Source/Runtime/Online/HTTP/HTTP.Build.cs
@@ -4,27 +4,28 @@ using UnrealBuildTool;
 
 public class HTTP : ModuleRules
 {
-	public HTTP(ReadOnlyTargetRules Target) : base(Target)
-	{
-		Definitions.Add("HTTP_PACKAGE=1");
+    public HTTP(ReadOnlyTargetRules Target) : base(Target)
+    {
+        Definitions.Add("HTTP_PACKAGE=1");
 
-		PrivateIncludePaths.AddRange(
-			new string[] {
+        PrivateIncludePaths.AddRange(
+            new string[] {
 				"Runtime/Online/HTTP/Private",
 			}
-			);
+            );
 
-		PrivateDependencyModuleNames.AddRange(
-			new string[] {
+        PrivateDependencyModuleNames.AddRange(
+			new string[] { 
 				"Core",
+				"Sockets"
 			}
 			);
 
 		bool bWithCurl = false;
 
-		if (Target.Platform == UnrealTargetPlatform.Win32 ||
-			Target.Platform == UnrealTargetPlatform.Win64)
-		{
+        if (Target.Platform == UnrealTargetPlatform.Win32 ||
+            Target.Platform == UnrealTargetPlatform.Win64)
+        {
 			AddEngineThirdPartyPrivateStaticDependencies(Target, "WinInet");
 			AddEngineThirdPartyPrivateStaticDependencies(Target, "WinHttp");
 			AddEngineThirdPartyPrivateStaticDependencies(Target, "libcurl");
@@ -33,7 +34,7 @@ public class HTTP : ModuleRules
 
 			PrivateDependencyModuleNames.Add("SSL");
 		}
-		else if (Target.Platform == UnrealTargetPlatform.Linux ||
+        else if (Target.Platform == UnrealTargetPlatform.Linux ||
 			Target.Platform == UnrealTargetPlatform.Android ||
 			Target.Platform == UnrealTargetPlatform.Switch)
 		{
@@ -58,13 +59,13 @@ public class HTTP : ModuleRules
 		}
 
 		if (Target.Platform == UnrealTargetPlatform.HTML5)
-		{
-			PrivateDependencyModuleNames.Add("HTML5JS");
-		}
+        {
+                PrivateDependencyModuleNames.Add("HTML5JS");
+            }
 
 		if (Target.Platform == UnrealTargetPlatform.IOS || Target.Platform == UnrealTargetPlatform.TVOS || Target.Platform == UnrealTargetPlatform.Mac)
 		{
 			PublicFrameworks.Add("Security");
 		}
-	}
+    }
 }
diff --git a/Engine/Source/Runtime/Online/HTTP/Private/Curl/CurlHttpManager.cpp b/Engine/Source/Runtime/Online/HTTP/Private/Curl/CurlHttpManager.cpp
--- a/Engine/Source/Runtime/Online/HTTP/Private/Curl/CurlHttpManager.cpp
+++ b/Engine/Source/Runtime/Online/HTTP/Private/Curl/CurlHttpManager.cpp
@@ -10,14 +10,15 @@
 #include "Misc/LocalTimestampDirectoryVisitor.h"
 #include "Curl/CurlHttpThread.h"
 #include "Curl/CurlHttp.h"
+#include "Http.h"
 #include "Modules/ModuleManager.h"
-
-#if WITH_SSL
-#include "Modules/ModuleManager.h"
-#endif
+#include "Misc/OutputDeviceRedirector.h"
 
 #if WITH_LIBCURL
 
+#include "SocketSubsystem.h"
+#include "IPAddress.h"
+
 #if PLATFORM_WINDOWS
 #include "SslModule.h"
 #endif
@@ -505,6 +506,37 @@ void FCurlHttpManager::InitCurl()
 		}
 	}
 
+	CurlRequestOptions.MaxHostConnections = FHttpModule::Get().GetHttpMaxConnectionsPerServer();
+	if (CurlRequestOptions.MaxHostConnections > 0)
+	{
+		const CURLMcode SetOptResult = curl_multi_setopt(GMultiHandle, CURLMOPT_MAX_HOST_CONNECTIONS, static_cast<long>(CurlRequestOptions.MaxHostConnections));
+		if (SetOptResult != CURLM_OK)
+		{
+			FUTF8ToTCHAR Converter(curl_multi_strerror(SetOptResult));
+			UE_LOG(LogHttp, Warning, TEXT("Failed to set max host connections options (%d), error %d ('%s')"),
+				CurlRequestOptions.MaxHostConnections, (int32)SetOptResult, Converter.Get());
+			CurlRequestOptions.MaxHostConnections = 0;
+		}
+	}
+	else
+	{
+		CurlRequestOptions.MaxHostConnections = 0;
+	}
+
+	TCHAR Home[256] = TEXT("");
+	if (FParse::Value(FCommandLine::Get(), TEXT("MULTIHOME="), Home, ARRAY_COUNT(Home)))
+	{
+		ISocketSubsystem* SocketSubsystem = ISocketSubsystem::Get(PLATFORM_SOCKETSUBSYSTEM);
+
+		if (SocketSubsystem)
+		{
+			// Check for the local addr to route requests through (respects MULTIHOME cmd line option)
+			bool bCanBindAll;
+			TSharedRef<FInternetAddr> HostAddr = SocketSubsystem->GetLocalHostAddr(*GLog, bCanBindAll);
+			CurlRequestOptions.LocalHostAddr = HostAddr->ToString(false);
+		}
+	}
+
 	// print for visibility
 	CurlRequestOptions.Log();
 }
@@ -523,7 +555,7 @@ void FCurlHttpManager::FCurlRequestOptions::Log()
 		);	
 	if (bUseHttpProxy)
 	{
-		UE_LOG(LogInit, Log, TEXT(" - HttpProxyAddress = '%s'"), *CurlRequestOptions.HttpProxyAddress);
+		UE_LOG(LogInit, Log, TEXT(" - HttpProxyAddress = '%s'"), *HttpProxyAddress);
 	}
 
 	UE_LOG(LogInit, Log, TEXT(" - bDontReuseConnections = %s  - Libcurl will %sreuse connections"),
@@ -535,6 +567,13 @@ void FCurlHttpManager::FCurlRequestOptions::Log()
 		(CertBundlePath != nullptr) ? *FString(CertBundlePath) : TEXT("nullptr"),
 		(CertBundlePath != nullptr) ? TEXT("set CURLOPT_CAINFO to it") : TEXT("use whatever was configured at build time.")
 		);
+
+	UE_LOG(LogInit, Log, TEXT(" - MaxHostConnections = %d  - Libcurl will %slimit the number of connections to a host"),
+		MaxHostConnections,
+		(MaxHostConnections == 0) ? TEXT("NOT ") : TEXT("")
+		);
+
+	UE_LOG(LogInit, Log, TEXT(" - LocalHostAddr = %s"), LocalHostAddr.IsEmpty() ? TEXT("Default") : *LocalHostAddr);
 }
 
 
@@ -556,6 +595,4 @@ FHttpThread* FCurlHttpManager::CreateHttpThread()
 	return new FCurlHttpThread();
 }
 
-
-
 #endif //WITH_LIBCURL
diff --git a/Engine/Source/Runtime/Online/HTTP/Private/Curl/CurlHttpManager.h b/Engine/Source/Runtime/Online/HTTP/Private/Curl/CurlHttpManager.h
--- a/Engine/Source/Runtime/Online/HTTP/Private/Curl/CurlHttpManager.h
+++ b/Engine/Source/Runtime/Online/HTTP/Private/Curl/CurlHttpManager.h
@@ -35,6 +35,7 @@ public:
 			,	bUseHttpProxy(false)
 			,	bDontReuseConnections(false)
 			,	CertBundlePath(nullptr)
+			,	MaxHostConnections(0)
 		{}
 
 		/** Prints out the options to the log */
@@ -54,6 +55,12 @@ public:
 
 		/** A path to certificate bundle */
 		const char * CertBundlePath;
+
+		/** The maximum number of connections to a particular host */
+		int32 MaxHostConnections;
+
+		/** Local address to use when making request, respects MULTIHOME command line option */
+		FString LocalHostAddr;
 	}
 	CurlRequestOptions;
 
diff --git a/Engine/Source/Runtime/Online/HTTP/Public/HttpModule.h b/Engine/Source/Runtime/Online/HTTP/Public/HttpModule.h
--- a/Engine/Source/Runtime/Online/HTTP/Public/HttpModule.h
+++ b/Engine/Source/Runtime/Online/HTTP/Public/HttpModule.h
@@ -100,17 +100,17 @@ public:
 	/**
 	 * @return timeout in seconds to send a request on the connection
 	 */
-	inline float GetHttpMaxConnectionsPerServer() const
+	inline float GetHttpSendTimeout() const
 	{
-		return HttpMaxConnectionsPerServer;
+		return HttpSendTimeout;
 	}
 
 	/**
 	 * @return max number of simultaneous connections to a specific server
 	 */
-	inline float GetHttpSendTimeout() const
+	inline int32 GetHttpMaxConnectionsPerServer() const
 	{
-		return HttpSendTimeout;
+		return HttpMaxConnectionsPerServer;
 	}
 
 	/**
diff --git a/Engine/Source/Runtime/PacketHandlers/PacketHandler/Private/PacketHandler.cpp b/Engine/Source/Runtime/PacketHandlers/PacketHandler/Private/PacketHandler.cpp
--- a/Engine/Source/Runtime/PacketHandlers/PacketHandler/Private/PacketHandler.cpp
+++ b/Engine/Source/Runtime/PacketHandlers/PacketHandler/Private/PacketHandler.cpp
@@ -77,7 +77,7 @@ void PacketHandler::Tick(float DeltaTime)
 	}
 }
 
-void PacketHandler::Initialize(Handler::Mode InMode, uint32 InMaxPacketBits, bool bConnectionlessOnly/*=false*/)
+void PacketHandler::Initialize(Handler::Mode InMode, uint32 InMaxPacketBits, bool bConnectionlessOnly/*=false*/, TSharedPtr<IAnalyticsProvider> InProvider/*=nullptr*/)
 {
 	Mode = InMode;
 	MaxPacketBits = InMaxPacketBits;
@@ -121,6 +121,8 @@ void PacketHandler::Initialize(Handler::Mode InMode, uint32 InMaxPacketBits, boo
 		ReliabilityComponent = MakeShareable(new ReliabilityHandlerComponent);
 		AddHandler(ReliabilityComponent, true);
 	}
+
+	Provider = InProvider;
 }
 
 void PacketHandler::InitializeComponents()
@@ -145,6 +147,7 @@ void PacketHandler::InitializeComponents()
 		if (!CurComponent.IsInitialized())
 		{
 			CurComponent.Initialize();
+			CurComponent.SetAnalyticsProvider(Provider);
 		}
 	}
 
diff --git a/Engine/Source/Runtime/PacketHandlers/PacketHandler/Public/PacketHandler.h b/Engine/Source/Runtime/PacketHandlers/PacketHandler/Public/PacketHandler.h
--- a/Engine/Source/Runtime/PacketHandlers/PacketHandler/Public/PacketHandler.h
+++ b/Engine/Source/Runtime/PacketHandlers/PacketHandler/Public/PacketHandler.h
@@ -195,8 +195,9 @@ public:
 	 * @param Mode					The mode the manager should be initialized in
 	 * @param InMaxPacketBits		The maximum supported packet size
 	 * @param bConnectionlessOnly	Whether or not this is a connectionless-only manager (ignores .ini components)
+	 * @param InProvider			The analytics provider
 	 */
-	void Initialize(Handler::Mode Mode, uint32 InMaxPacketBits, bool bConnectionlessOnly=false);
+	void Initialize(Handler::Mode Mode, uint32 InMaxPacketBits, bool bConnectionlessOnly=false, TSharedPtr<class IAnalyticsProvider> InProvider=nullptr);
 
 	/**
 	 * Used for external initialization of delegates
@@ -545,6 +546,9 @@ private:
 	/** Whether or not outgoing packets bypass the handler */
 	bool bRawSend;
 
+	/** The analytics provider */
+	TSharedPtr<class IAnalyticsProvider> Provider;
+	
 	/** Whether or not component handshaking has begun */
 	bool bBeganHandshaking;
 };
@@ -669,6 +673,13 @@ public:
 	 */
 	virtual int32 GetReservedPacketBits() PURE_VIRTUAL(Handler::Component::GetReservedPacketBits, return -1;);
 
+	/**
+	* Sets the analytics provider that can be used to send analytics events as needed.
+	*
+	* @param Provider The analytics provider
+	*/
+	virtual void SetAnalyticsProvider(TSharedPtr<class IAnalyticsProvider> Provider) {}
+
 protected:
 	/**
 	 * Sets the state of the handler
diff --git a/Engine/Source/Runtime/PakFile/Private/IPlatformFilePak.cpp b/Engine/Source/Runtime/PakFile/Private/IPlatformFilePak.cpp
--- a/Engine/Source/Runtime/PakFile/Private/IPlatformFilePak.cpp
+++ b/Engine/Source/Runtime/PakFile/Private/IPlatformFilePak.cpp
@@ -4250,7 +4250,8 @@ bool FPakPlatformFile::Initialize(IPlatformFile* Inner, const TCHAR* CmdLine)
 	DecryptionKey.Exponent.Parse(PakSigningKeyExponent);
 	DecryptionKey.Modulus.Parse(PakSigningKeyModulus);
 
-	bSigned = !DecryptionKey.Exponent.IsZero() && !DecryptionKey.Modulus.IsZero();
+	// signed if we have keys, and are not running with fileopenlog (currently results in a deadlock).
+	bSigned = !DecryptionKey.Exponent.IsZero() && !DecryptionKey.Modulus.IsZero() && !FParse::Param(FCommandLine::Get(), TEXT("fileopenlog"));;
 	
 	bool bMountPaks = true;
 	TArray<FString> PaksToLoad;
diff --git a/Engine/Source/Runtime/Portal/LauncherCheck/Private/LauncherCheckModule.cpp b/Engine/Source/Runtime/Portal/LauncherCheck/Private/LauncherCheckModule.cpp
--- a/Engine/Source/Runtime/Portal/LauncherCheck/Private/LauncherCheckModule.cpp
+++ b/Engine/Source/Runtime/Portal/LauncherCheck/Private/LauncherCheckModule.cpp
@@ -38,7 +38,7 @@ public:
 	*/
 	bool IsEnabled() const
 	{
-		return FParse::Param(FCommandLine::Get(), TEXT("NoEpicPortal")) == false;
+		return FParse::Param(FCommandLine::Get(), TEXT("NoEpicPortal")) == false && FParse::Param(FCommandLine::Get(), TEXT("q")) == false;
 	}
 
 public:
diff --git a/Engine/Source/Runtime/RHI/Private/DynamicRHI.cpp b/Engine/Source/Runtime/RHI/Private/DynamicRHI.cpp
--- a/Engine/Source/Runtime/RHI/Private/DynamicRHI.cpp
+++ b/Engine/Source/Runtime/RHI/Private/DynamicRHI.cpp
@@ -223,7 +223,7 @@ void RHIExit()
 	if ( !GUsingNullRHI && GDynamicRHI != NULL )
 	{
 		// Clean up all cached pipelines
-		ClearPipelineCache();
+		PipelineStateCache::Shutdown();
 
 		// Destruct the dynamic RHI.
 		GDynamicRHI->Shutdown();
diff --git a/Engine/Source/Runtime/RHI/Private/PipelineStateCache.cpp b/Engine/Source/Runtime/RHI/Private/PipelineStateCache.cpp
--- a/Engine/Source/Runtime/RHI/Private/PipelineStateCache.cpp
+++ b/Engine/Source/Runtime/RHI/Private/PipelineStateCache.cpp
@@ -6,17 +6,23 @@ PipelineStateCache.cpp: Pipeline state cache implementation.
 
 #include "PipelineStateCache.h"
 #include "Misc/ScopeRWLock.h"
+#include "Misc/ScopeLock.h"
+#include "CoreDelegates.h"
+#include "CoreGlobals.h"
+#include "TimeGuard.h"
+#include "DiscardableKeyValueCache.h"
 
-extern RHI_API FRHIComputePipelineState* ExecuteSetComputePipelineState(FComputePipelineState* ComputePipelineState);
-extern RHI_API FRHIGraphicsPipelineState* ExecuteSetGraphicsPipelineState(FGraphicsPipelineState* GraphicsPipelineState);
+// perform cache eviction each frame, used to stress the system and flush out bugs
+#define PSO_DO_CACHE_EVICT_EACH_FRAME 0
 
-static TAutoConsoleVariable<int32> GCVarAsyncPipelineCompile(
-	TEXT("r.AsyncPipelineCompile"),
-	1,
-	TEXT("0 to Create PSOs at the moment they are requested\n")\
-	TEXT("1 to Create Pipeline State Objects asynchronously(default)"),
-	ECVF_ReadOnly | ECVF_RenderThreadSafe
-);
+// Log event and info about cache eviction
+#define PSO_LOG_CACHE_EVICT 0
+
+// Stat tracking
+#define PSO_TRACK_CACHE_STATS 0
+
+
+#define PIPELINESTATECACHE_VERIFYTHREADSAFE (!UE_BUILD_SHIPPING && !UE_BUILD_TEST)
 
 static inline uint32 GetTypeHash(const FBoundShaderStateInput& Input)
 {
@@ -33,7 +39,7 @@ static inline uint32 GetTypeHash(const FGraphicsPipelineStateInitializer& Initia
 	//#todo-rco: Hash!
 	return (GetTypeHash(Initializer.BoundShaderState) | (Initializer.NumSamples << 28)) ^ ((uint32)Initializer.PrimitiveType << 24) ^ GetTypeHash(Initializer.BlendState)
 		^ Initializer.RenderTargetsEnabled ^ GetTypeHash(Initializer.RasterizerState) ^ GetTypeHash(Initializer.DepthStencilState);
-/*
+	/*
 	uint32							RenderTargetsEnabled;
 	TRenderTargetFormats			RenderTargetFormats;
 	TRenderTargetFlags				RenderTargetFlags;
@@ -43,23 +49,95 @@ static inline uint32 GetTypeHash(const FGraphicsPipelineStateInitializer& Initia
 	uint32							DepthStencilTargetFlag;
 	ERenderTargetLoadAction			DepthStencilTargetLoadAction;
 	ERenderTargetStoreAction		DepthStencilTargetStoreAction;
-*/
+	*/
 }
 
+static TAutoConsoleVariable<int32> GCVarAsyncPipelineCompile(
+	TEXT("r.AsyncPipelineCompile"),
+	1,
+	TEXT("0 to Create PSOs at the moment they are requested\n")\
+	TEXT("1 to Create Pipeline State Objects asynchronously(default)"),
+	ECVF_ReadOnly | ECVF_RenderThreadSafe
+);
+
+static TAutoConsoleVariable<int32> CVarPSOEvictionTime(
+	TEXT("r.pso.evictiontime"),
+	60,
+	TEXT("Time between checks to remove stale objects from the cache. 0 = no eviction (which may eventually OOM...)"),
+	ECVF_ReadOnly | ECVF_RenderThreadSafe
+);
+
+extern void DumpPipelineCacheStats();
 
-static FRWLock GComputeLock;
-static TMap<FRHIComputeShader*, FComputePipelineState*> GComputePipelines;
-static FRWLock GGraphicsLock;
-static TMap <FGraphicsPipelineStateInitializer, FGraphicsPipelineState*> GGraphicsPipelines;
+static FAutoConsoleCommand DumpPipelineCmd(
+	TEXT("r.DumpPipelineCache"),
+	TEXT("Dump current cache stats."),
+	FConsoleCommandDelegate::CreateStatic(DumpPipelineCacheStats)
+);
+
+void SetComputePipelineState(FRHICommandList& RHICmdList, FRHIComputeShader* ComputeShader)
+{
+	RHICmdList.SetComputePipelineState(PipelineStateCache::GetAndOrCreateComputePipelineState(RHICmdList, ComputeShader));
+}
 
+extern RHI_API FRHIComputePipelineState* ExecuteSetComputePipelineState(FComputePipelineState* ComputePipelineState);
+extern RHI_API FRHIGraphicsPipelineState* ExecuteSetGraphicsPipelineState(FGraphicsPipelineState* GraphicsPipelineState);
+
+/**
+ * Base class to hold pipeline state (and optionally stats)
+ */
 class FPipelineState
 {
 public:
-	virtual ~FPipelineState() { }
+
+	FPipelineState()
+	{
+		InitStats();
+	}
+
+	virtual ~FPipelineState() 
+	{
+	}
+
 	virtual bool IsCompute() const = 0;
+
 	FGraphEventRef CompletionEvent;
+
+#if PSO_TRACK_CACHE_STATS
+	void InitStats()
+	{
+		FirstUsedTime = LastUsedTime = FPlatformTime::Seconds();
+		FirstFrameUsed = LastFrameUsed = 0;
+		Hits = HitsAcrossFrames = 0;
+	}
+
+	void AddHit()
+	{
+		LastUsedTime = FPlatformTime::Seconds();
+		Hits++;
+
+		if (LastFrameUsed != GFrameCounter)
+		{
+			LastFrameUsed = GFrameCounter;
+			HitsAcrossFrames++;
+		}
+	}
+
+	double			FirstUsedTime;
+	double			LastUsedTime;
+	uint64			FirstFrameUsed;
+	uint64			LastFrameUsed;
+	int				Hits;
+	int				HitsAcrossFrames;
+
+#else
+	void InitStats() {}
+	void AddHit() {}
+#endif // PSO_TRACK_CACHE_STATS
+
 };
 
+/* State for compute  */
 class FComputePipelineState : public FPipelineState
 {
 public:
@@ -73,17 +151,16 @@ public:
 		return true;
 	}
 
-//protected:
+	//protected:
 	FRHIComputeShader* ComputeShader;
 	TRefCountPtr<FRHIComputePipelineState> RHIPipeline;
 };
 
-
+/* State for graphics */
 class FGraphicsPipelineState : public FPipelineState
 {
 public:
-	FGraphicsPipelineState(const FGraphicsPipelineStateInitializer& InInitializer)
-		: Initializer(InInitializer)
+	FGraphicsPipelineState() 
 	{
 	}
 
@@ -93,17 +170,295 @@ public:
 	}
 
 	//protected:
-	FGraphicsPipelineStateInitializer Initializer;
 	TRefCountPtr<FRHIGraphicsPipelineState> RHIPipeline;
+#if PIPELINESTATECACHE_VERIFYTHREADSAFE
+	FThreadSafeCounter InUseCount;
+#endif
+};
+
+void SetGraphicsPipelineState(FRHICommandList& RHICmdList, const FGraphicsPipelineStateInitializer& Initializer, EApplyRendertargetOption ApplyFlags)
+{
+	FGraphicsPipelineState* PipelineState = PipelineStateCache::GetAndOrCreateGraphicsPipelineState(RHICmdList, Initializer, ApplyFlags);
+#if PIPELINESTATECACHE_VERIFYTHREADSAFE
+	int32 Result = PipelineState->InUseCount.Increment();
+	check(Result >= 1);
+#endif
+	check(IsInRenderingThread() || IsInParallelRenderingThread());
+	RHICmdList.SetGraphicsPipelineState(PipelineState);
+}
+
+/* TSharedPipelineStateCache
+ * This is a cache of the * pipeline states
+ * there is a local thread cache which is consolidated with the global thread cache
+ * global thread cache is read only until the end of the frame when the local thread caches are consolidated
+ */
+template<class TMyKey,class TMyValue>
+class TSharedPipelineStateCache
+{
+private:
+
+	TMap<TMyKey, TMyValue>& GetLocalCache()
+	{
+		void* TLSValue = FPlatformTLS::GetTlsValue(TLSSlot);
+		if (TLSValue == nullptr)
+		{
+			FPipelineStateCacheType* PipelineStateCache = new FPipelineStateCacheType;
+			FPlatformTLS::SetTlsValue(TLSSlot, (void*)(PipelineStateCache) );
+
+			FScopeLock S(&AllThreadsLock);
+			AllThreadsPipelineStateCache.Add(PipelineStateCache);
+			return *PipelineStateCache;
+		}
+		return *((FPipelineStateCacheType*)TLSValue);
+	}
+
+#if PIPELINESTATECACHE_VERIFYTHREADSAFE
+	struct FScopeVerifyIncrement
+	{
+		volatile int32 &VerifyMutex;
+		FScopeVerifyIncrement(volatile int32& InVerifyMutex) : VerifyMutex(InVerifyMutex)
+		{
+			int32 Result = FPlatformAtomics::InterlockedIncrement(&VerifyMutex);
+			if (Result <= 0)
+			{
+				UE_LOG(LogRHI, Fatal, TEXT("Find was hit while Consolidate was running"));
+			}
+		}
+
+		~FScopeVerifyIncrement()
+		{
+			int32 Result = FPlatformAtomics::InterlockedDecrement(&VerifyMutex);
+			if (Result < 0)
+			{
+				UE_LOG(LogRHI, Fatal, TEXT("Find was hit while Consolidate was running"));
+			}
+		}
+	};
+
+	struct FScopeVerifyDecrement
+	{
+		volatile int32 &VerifyMutex;
+		FScopeVerifyDecrement(volatile int32& InVerifyMutex) : VerifyMutex(InVerifyMutex)
+		{
+			int32 Result = FPlatformAtomics::InterlockedDecrement(&VerifyMutex);
+			if (Result >= 0)
+			{
+				UE_LOG(LogRHI, Fatal, TEXT("Consolidate was hit while Get/SetPSO was running"));
+			}
+		}
+
+		~FScopeVerifyDecrement()
+		{
+			int32 Result = FPlatformAtomics::InterlockedIncrement(&VerifyMutex);
+			if (Result != 0)
+			{
+				UE_LOG(LogRHI, Fatal, TEXT("Consolidate was hit while Get/SetPSO was running"));
+			}
+		}
+	};
+#endif
+
+public:
+	typedef TMap<TMyKey,TMyValue> FPipelineStateCacheType;
+
+	TSharedPipelineStateCache()
+	{
+		CurrentMap = &Map1;
+		BackfillMap = &Map2;
+		DuplicateStateGenerated = 0;
+		TLSSlot = FPlatformTLS::AllocTlsSlot();
+	}
+
+	bool Find( const TMyKey& InKey, TMyValue& OutResult )
+	{
+#if PIPELINESTATECACHE_VERIFYTHREADSAFE
+		FScopeVerifyIncrement S(VerifyMutex);
+#endif
+		// safe because we only ever find when we don't add
+		TMyValue* Result = CurrentMap->Find(InKey);
+
+		if ( Result )
+		{
+			OutResult = *Result;
+			return true;
+		}
+
+		// check the local cahce which is safe because only this thread adds to it
+		TMap<TMyKey, TMyValue> &LocalCache = GetLocalCache();
+		// if it's not in the local cache then it will rebuild
+		Result = LocalCache.Find(InKey);
+		if (Result)
+		{
+			OutResult = *Result;
+			return true;
+		}
+
+		Result = BackfillMap->Find(InKey);
+
+		if ( Result )
+		{
+			LocalCache.Add(InKey, *Result);
+			OutResult = *Result;
+			return true;
+		}
+
+
+		return false;
+		
+		
+	}
+
+	bool Add(const TMyKey& InKey, const TMyValue& InValue)
+	{
+#if PIPELINESTATECACHE_VERIFYTHREADSAFE
+		FScopeVerifyIncrement S(VerifyMutex);
+#endif
+		// everything is added to the local cache then at end of frame we consoldate them all
+		TMap<TMyKey, TMyValue> &LocalCache = GetLocalCache();
+
+		check( LocalCache.Contains(InKey) == false );
+		LocalCache.Add(InKey, InValue);
+		return true;
+	}
+
+	void ConsolidateThreadedCaches()
+	{
+
+		SCOPE_TIME_GUARD_MS(TEXT("ConsolidatePipelineCache"), 0.1);
+		check(IsInRenderingThread());
+#if PIPELINESTATECACHE_VERIFYTHREADSAFE
+		FScopeVerifyDecrement S(VerifyMutex);
+#endif
+		
+		// consolidate all the local threads keys with the current thread
+		// No one is allowed to call GetLocalCache while this is running
+		// this is verified by the VerifyMutex.
+		for ( FPipelineStateCacheType* PipelineStateCache : AllThreadsPipelineStateCache)
+		{
+			for (const auto& PipelineStateCacheIterator : *PipelineStateCache)
+			{
+				const TMyKey& ThreadKey = PipelineStateCacheIterator.Key;
+				const TMyValue& ThreadValue = PipelineStateCacheIterator.Value;
+
+				// All events should be complete because we are running this code after the RHI Flush
+				check(!ThreadValue->CompletionEvent.IsValid() || ThreadValue->CompletionEvent->IsComplete());
+				
+				ThreadValue->CompletionEvent = nullptr;
+
+				BackfillMap->Remove(ThreadKey);
+
+				TMyValue* CurrentValue = CurrentMap->Find(ThreadKey);
+				if (CurrentValue) 
+				{
+					// if two threads get from the backfill map then we might just be dealing with one pipelinestate, in which case we have already added it to the currentmap and don't need to do anything else
+					if ( *CurrentValue != ThreadValue ) 
+					{
+						++DuplicateStateGenerated;
+						DeleteArray.Add(ThreadValue);
+					}
+				}
+				else
+				{
+					CurrentMap->Add(ThreadKey, ThreadValue);
+				}
+			}
+			PipelineStateCache->Empty();
+		}
+
+	}
+
+	void ProcessDelayedCleanup()
+	{
+		check(IsInRenderingThread());
+
+		for (TMyValue& OldPipelineState : DeleteArray)
+		{
+			//once in the delayed list this object should not be findable anymore, so the 0 should remain, making this safe
+#if PIPELINESTATECACHE_VERIFYTHREADSAFE
+			check(OldPipelineState->InUseCount.GetValue() == 0);
+#endif
+			delete OldPipelineState;
+		}
+		DeleteArray.Empty();
+	}
+
+
+	int32 DiscardAndSwap()
+	{
+		// the consolidate should always be run before the DiscardAndSwap.
+		// there should be no inuse pipeline states in the backfill map (because they should have been moved into the CurrentMap).
+		int32 Discarded = BackfillMap->Num();
+
+
+		for ( const auto& DiscardIterator : *BackfillMap )
+		{
+#if PIPELINESTATECACHE_VERIFYTHREADSAFE
+			check( DiscardIterator.Value->InUseCount.GetValue() == 0);
+#endif
+			delete DiscardIterator.Value;
+		}
+
+		BackfillMap->Empty();
+
+
+		if ( CurrentMap == &Map1 )
+		{
+			CurrentMap = &Map2;
+			BackfillMap = &Map1;
+		}
+		else
+		{
+			CurrentMap = &Map1;
+			BackfillMap = &Map2;
+		}
+		return Discarded;
+	}
+private:
+	uint32 TLSSlot;
+	FPipelineStateCacheType *CurrentMap;
+	FPipelineStateCacheType *BackfillMap;
+
+	FPipelineStateCacheType Map1;
+	FPipelineStateCacheType Map2;
+
+	/*TArray<TMyValue> *CurrentDeleteArray;
+	TArray<TMyValue> DeleteArray1;
+	TArray<TMyValue> DeleteArray2;*/
+
+	TArray<TMyValue> DeleteArray;
+
+	FCriticalSection AllThreadsLock;
+	TArray<FPipelineStateCacheType*> AllThreadsPipelineStateCache;
+
+	uint32 DuplicateStateGenerated;
+
+#if PIPELINESTATECACHE_VERIFYTHREADSAFE
+	volatile int32 VerifyMutex;
+#endif
+
 };
 
+// Typed caches for compute and graphics
+typedef TDiscardableKeyValueCache< FRHIComputeShader*, FComputePipelineState*> FComputePipelineCache;
+typedef TSharedPipelineStateCache<FGraphicsPipelineStateInitializer, FGraphicsPipelineState*> FGraphicsPipelineCache;
+
+// These are the actual caches for both pipelines
+FComputePipelineCache GComputePipelineCache;
+FGraphicsPipelineCache GGraphicsPipelineCache;
+
+/**
+ *  Compile task
+ */
 class FCompilePipelineStateTask
 {
 public:
 	FPipelineState* Pipeline;
+	FGraphicsPipelineStateInitializer Initializer;
 
-	FCompilePipelineStateTask(FPipelineState* InPipeline)
+	// InInitializer is only used for non-compute tasks, a default can just be used otherwise
+	FCompilePipelineStateTask(FPipelineState* InPipeline, const FGraphicsPipelineStateInitializer& InInitializer)
 		: Pipeline(InPipeline)
+		, Initializer(InInitializer)
 	{
 	}
 
@@ -119,7 +474,7 @@ public:
 		else
 		{
 			FGraphicsPipelineState* GfxPipeline = static_cast<FGraphicsPipelineState*>(Pipeline);
-			GfxPipeline->RHIPipeline = RHICreateGraphicsPipelineState(GfxPipeline->Initializer);
+			GfxPipeline->RHIPipeline = RHICreateGraphicsPipelineState(Initializer);
 		}
 	}
 
@@ -134,48 +489,145 @@ public:
 	}
 };
 
+/**
+* Called at the end of each frame during the RHI . Evicts all items left in the backfill cached based on time
+*/
+void PipelineStateCache::FlushResources()
+{
+	static bool PerformedOneTimeInit = false;
+
+	check(IsInRenderingThread());
+
+	GGraphicsPipelineCache.ConsolidateThreadedCaches();
+	GGraphicsPipelineCache.ProcessDelayedCleanup();
+	// Thread-safe one-time initialization of things we need to set up
+	if (PerformedOneTimeInit == false)
+	{
+		PerformedOneTimeInit = true;
+
+		// We don't trim, but we will dump out how much memory we're using..
+		FCoreDelegates::GetMemoryTrimDelegate().AddLambda([]()
+		{
+#if PSO_TRACK_CACHE_STATS
+			DumpPipelineCacheStats();
+#endif
+		});
+
+		PerformedOneTimeInit = true;
+
+	}
+
+	static double LastEvictionTime = FPlatformTime::Seconds();
+	double CurrentTime = FPlatformTime::Seconds();
+
+#if PSO_DO_CACHE_EVICT_EACH_FRAME
+	LastEvictionTime = 0;
+#endif
+	
+	// because it takes two cycles for an object to move from main->backfill->gone we check
+	// at half the desired eviction time
+	int32 EvictionPeriod = CVarPSOEvictionTime.GetValueOnAnyThread();
+
+	if (EvictionPeriod == 0 || CurrentTime - LastEvictionTime < EvictionPeriod)
+	{
+		return;
+	}
+
+	// This should be very fast, if not it's likely eviction time is too high and too 
+	// many items are building up.
+	SCOPE_TIME_GUARD_MS(TEXT("TrimPiplelineCache"), 0.1);
+
+#if PSO_TRACK_CACHE_STATS
+	DumpPipelineCacheStats();
+#endif
+
+	LastEvictionTime = CurrentTime;
+
+	int ReleasedComputeEntries = 0;
+	int ReleasedGraphicsEntries = 0;
+
+	ReleasedComputeEntries = GComputePipelineCache.Discard([](FComputePipelineState* CacheItem) {
+		delete CacheItem;
+	});
+
+	ReleasedGraphicsEntries = GGraphicsPipelineCache.DiscardAndSwap();
+
+#if PSO_TRACK_CACHE_STATS
+	UE_LOG(LogRHI, Log, TEXT("Cleared state cache in %.02f ms. %d ComputeEntries, %d Graphics entries")
+		, (FPlatformTime::Seconds() - CurrentTime) / 1000
+		, ReleasedComputeEntries, ReleasedGraphicsEntries);
+#endif // PSO_TRACK_CACHE_STATS
+
+}
+
 static bool IsAsyncCompilationAllowed(FRHICommandList& RHICmdList)
 {
 	return GCVarAsyncPipelineCompile.GetValueOnAnyThread() && !RHICmdList.Bypass() && IsRunningRHIInSeparateThread();
 }
 
-FComputePipelineState* GetAndOrCreateComputePipelineState(FRHICommandList& RHICmdList, FRHIComputeShader* ComputeShader)
+FComputePipelineState* PipelineStateCache::GetAndOrCreateComputePipelineState(FRHICommandList& RHICmdList, FRHIComputeShader* ComputeShader)
 {
 	SCOPE_CYCLE_COUNTER(STAT_GetOrCreatePSO);
 
+	bool DoAsyncCompile = IsAsyncCompilationAllowed(RHICmdList);
+
+	FComputePipelineState* OutCachedState = nullptr;
+
+	uint32 LockFlags = GComputePipelineCache.ApplyLock(0, FComputePipelineCache::LockFlags::ReadLock);
+
+	bool WasFound = GComputePipelineCache.Find(ComputeShader, OutCachedState, LockFlags | FComputePipelineCache::LockFlags::WriteLockOnAddFail, LockFlags);
+
+	if (WasFound == false)
 	{
-		// Should be hitting this case more often once the cache is hot
-		FRWScopeLock ScopeLock(GComputeLock, SLT_ReadOnly);
+		// create new graphics state
+		OutCachedState = new FComputePipelineState(ComputeShader);
 
-		FComputePipelineState** Found = GComputePipelines.Find(ComputeShader);
-		if (Found)
+		// create a compilation task, or just do it now...
+		if (DoAsyncCompile)
 		{
-			if (IsAsyncCompilationAllowed(RHICmdList))
-			{
-				FGraphEventRef& CompletionEvent = (*Found)->CompletionEvent;
-				RHICmdList.QueueAsyncPipelineStateCompile(CompletionEvent);
-			}
-			return *Found;
+			OutCachedState->CompletionEvent = TGraphTask<FCompilePipelineStateTask>::CreateTask().ConstructAndDispatchWhenReady(OutCachedState, FGraphicsPipelineStateInitializer());
+			RHICmdList.QueueAsyncPipelineStateCompile(OutCachedState->CompletionEvent);
 		}
-		
-		ScopeLock.RaiseLockToWrite();
-		Found = GComputePipelines.Find(ComputeShader);
-		if (Found)
+		else
 		{
-			if (IsAsyncCompilationAllowed(RHICmdList))
+			OutCachedState->RHIPipeline = RHICreateComputePipelineState(OutCachedState->ComputeShader);
+		}
+
+		GComputePipelineCache.Add(ComputeShader, OutCachedState, LockFlags);
+	}
+	else
+	{
+		if (DoAsyncCompile)
+		{
+			FGraphEventRef& CompletionEvent = OutCachedState->CompletionEvent;
+			if ( CompletionEvent.IsValid() && !CompletionEvent->IsComplete() )
 			{
-				FGraphEventRef& CompletionEvent = (*Found)->CompletionEvent;
 				RHICmdList.QueueAsyncPipelineStateCompile(CompletionEvent);
 			}
-			return *Found;
 		}
-		else
-		{
-			FComputePipelineState* PipelineState = new FComputePipelineState(ComputeShader);
 
-			if (IsAsyncCompilationAllowed(RHICmdList))
+#if PSO_TRACK_CACHE_STATS
+		OutCachedState->AddHit();
+#endif
+	}
+
+	GComputePipelineCache.Unlock(LockFlags);
+
+#if 0
+	bool DoAsyncCompile = IsAsyncCompilationAllowed(RHICmdList);
+
+
+	TSharedPtr<FComputePipelineState, ESPMode::Fast> OutCachedState;
+
+	// Find or add an entry for this initializer
+	bool WasFound = GComputePipelineCache.FindOrAdd(ComputeShader, OutCachedState, [&RHICmdList, &ComputeShader, &DoAsyncCompile] {
+			// create new graphics state
+			TSharedPtr<FComputePipelineState, ESPMode::Fast> PipelineState(new FComputePipelineState(ComputeShader));
+
+			// create a compilation task, or just do it now...
+			if (DoAsyncCompile)
 			{
-				PipelineState->CompletionEvent = TGraphTask<FCompilePipelineStateTask>::CreateTask().ConstructAndDispatchWhenReady(PipelineState);
+				PipelineState->CompletionEvent = TGraphTask<FCompilePipelineStateTask>::CreateTask().ConstructAndDispatchWhenReady(PipelineState.Get(), FGraphicsPipelineStateInitializer());
 				RHICmdList.QueueAsyncPipelineStateCompile(PipelineState->CompletionEvent);
 			}
 			else
@@ -183,26 +635,47 @@ FComputePipelineState* GetAndOrCreateComputePipelineState(FRHICommandList& RHICm
 				PipelineState->RHIPipeline = RHICreateComputePipelineState(PipelineState->ComputeShader);
 			}
 
-			GComputePipelines.Add(ComputeShader, PipelineState);
-			
+			// wrap it and return it
 			return PipelineState;
+		});
+
+	check(OutCachedState.IsValid());
+
+	// if we found an entry the block above wasn't executed
+	if (WasFound)
+	{
+		if (DoAsyncCompile)
+		{
+			FRWScopeLock ScopeLock(GComputePipelineCache.RWLock(), SLT_ReadOnly);
+			FGraphEventRef& CompletionEvent = OutCachedState->CompletionEvent;
+			if (CompletionEvent.IsValid() && !CompletionEvent->IsComplete())
+			{
+				RHICmdList.QueueAsyncPipelineStateCompile(CompletionEvent);
+			}
 		}
+#if PSO_TRACK_CACHE_STATS
+		OutCachedState->AddHit();
+#endif
 	}
-
-	return nullptr;
+#endif
+	// return the state pointer
+	return OutCachedState;
 }
 
 FRHIComputePipelineState* ExecuteSetComputePipelineState(FComputePipelineState* ComputePipelineState)
 {
 	ensure(ComputePipelineState->RHIPipeline);
-	FRWScopeLock ScopeLock(GComputeLock, SLT_Write);
+	FRWScopeLock ScopeLock(GComputePipelineCache.RWLock(), SLT_Write);
+
 	ComputePipelineState->CompletionEvent = nullptr;
 	return ComputePipelineState->RHIPipeline;
 }
 
-FGraphicsPipelineState* GetAndOrCreateGraphicsPipelineState(FRHICommandList& RHICmdList, const FGraphicsPipelineStateInitializer& OriginalInitializer, EApplyRendertargetOption ApplyFlags)
+FGraphicsPipelineState* PipelineStateCache::GetAndOrCreateGraphicsPipelineState(FRHICommandList& RHICmdList, const FGraphicsPipelineStateInitializer& OriginalInitializer, EApplyRendertargetOption ApplyFlags)
 {
+
 	SCOPE_CYCLE_COUNTER(STAT_GetOrCreatePSO);
+
 	FGraphicsPipelineStateInitializer NewInitializer;
 	const FGraphicsPipelineStateInitializer* Initializer = &OriginalInitializer;
 
@@ -213,25 +686,25 @@ FGraphicsPipelineState* GetAndOrCreateGraphicsPipelineState(FRHICommandList& RHI
 		RHICmdList.ApplyCachedRenderTargets(NewInitializer);
 		Initializer = &NewInitializer;
 	}
-	
-#if !UE_BUILD_SHIPPING && !UE_BUILD_TEST
-	else if(!!(ApplyFlags & EApplyRendertargetOption::CheckApply))
+
+#if !UE_BUILD_SHIPPING && !UE_BUILD_TEST 
+	else if (!!(ApplyFlags & EApplyRendertargetOption::CheckApply))
 	{
 		// Catch cases where the state does not match
 		NewInitializer = OriginalInitializer;
 		RHICmdList.ApplyCachedRenderTargets(NewInitializer);
 
 		int AnyFailed = 0;
-		AnyFailed |= (NewInitializer.RenderTargetsEnabled != OriginalInitializer.RenderTargetsEnabled)			<< 0;
+		AnyFailed |= (NewInitializer.RenderTargetsEnabled != OriginalInitializer.RenderTargetsEnabled) << 0;
 
 		if (AnyFailed == 0)
 		{
 			for (int i = 0; i < (int)NewInitializer.RenderTargetsEnabled; i++)
 			{
-				AnyFailed |= (NewInitializer.RenderTargetFormats[i] != OriginalInitializer.RenderTargetFormats[i])				<< 1;
-				AnyFailed |= (NewInitializer.RenderTargetFlags[i] != OriginalInitializer.RenderTargetFlags[i])					<< 2;
-				AnyFailed |= (NewInitializer.RenderTargetLoadActions[i] != OriginalInitializer.RenderTargetLoadActions[i])		<< 3;
-				AnyFailed |= (NewInitializer.RenderTargetStoreActions[i] != OriginalInitializer.RenderTargetStoreActions[i])	<< 4;
+				AnyFailed |= (NewInitializer.RenderTargetFormats[i] != OriginalInitializer.RenderTargetFormats[i]) << 1;
+				AnyFailed |= (NewInitializer.RenderTargetFlags[i] != OriginalInitializer.RenderTargetFlags[i]) << 2;
+				AnyFailed |= (NewInitializer.RenderTargetLoadActions[i] != OriginalInitializer.RenderTargetLoadActions[i]) << 3;
+				AnyFailed |= (NewInitializer.RenderTargetStoreActions[i] != OriginalInitializer.RenderTargetStoreActions[i]) << 4;
 
 				if (AnyFailed)
 				{
@@ -241,95 +714,164 @@ FGraphicsPipelineState* GetAndOrCreateGraphicsPipelineState(FRHICommandList& RHI
 			}
 		}
 
-		AnyFailed |= (NewInitializer.DepthStencilTargetFormat != OriginalInitializer.DepthStencilTargetFormat)	<< 5;
-		AnyFailed |= (NewInitializer.DepthStencilTargetFlag != OriginalInitializer.DepthStencilTargetFlag)		<< 6;
-		AnyFailed |= (NewInitializer.DepthTargetLoadAction != OriginalInitializer.DepthTargetLoadAction)		<< 7;
-		AnyFailed |= (NewInitializer.DepthTargetStoreAction != OriginalInitializer.DepthTargetStoreAction)		<< 8;
-		AnyFailed |= (NewInitializer.StencilTargetLoadAction != OriginalInitializer.StencilTargetLoadAction)	<< 9;
-		AnyFailed |= (NewInitializer.StencilTargetStoreAction != OriginalInitializer.StencilTargetStoreAction)	<< 10;
-
-		ensureMsgf(AnyFailed == 0, TEXT("GetAndOrCreateGraphicsPipelineState RenderTarget check failed with: %i !"), AnyFailed);
+		AnyFailed |= (NewInitializer.DepthStencilTargetFormat != OriginalInitializer.DepthStencilTargetFormat) << 5;
+		AnyFailed |= (NewInitializer.DepthStencilTargetFlag != OriginalInitializer.DepthStencilTargetFlag) << 6;
+		AnyFailed |= (NewInitializer.DepthTargetLoadAction != OriginalInitializer.DepthTargetLoadAction) << 7;
+		AnyFailed |= (NewInitializer.DepthTargetStoreAction != OriginalInitializer.DepthTargetStoreAction) << 8;
+		AnyFailed |= (NewInitializer.StencilTargetLoadAction != OriginalInitializer.StencilTargetLoadAction) << 9;
+		AnyFailed |= (NewInitializer.StencilTargetStoreAction != OriginalInitializer.StencilTargetStoreAction) << 10;
+
+		// Paragon - remove this ensure which is wrecking havoc
+		//ensureMsgf(AnyFailed == 0, TEXT("GetAndOrCreateGraphicsPipelineState RenderTarget check failed with: %i !"), AnyFailed);
+		static double LastTime = 0;
+		if (AnyFailed != 0 && (FPlatformTime::Seconds() - LastTime) >= 10.0f)
+		{
+			LastTime = FPlatformTime::Seconds();
+			UE_LOG(LogRHI, Error, TEXT("GetAndOrCreateGraphicsPipelineState RenderTarget check failed with: %i !"), AnyFailed);
+		}
 		Initializer = (AnyFailed != 0) ? &NewInitializer : &OriginalInitializer;
 	}
 #endif
 
+	bool DoAsyncCompile = IsAsyncCompilationAllowed(RHICmdList);
+
+	FGraphicsPipelineState* OutCachedState = nullptr;
+
+	bool WasFound = GGraphicsPipelineCache.Find(*Initializer, OutCachedState);
+
+	if (WasFound == false)
 	{
-		// Should be hitting this case more often once the cache is hot
-		FRWScopeLock ScopeLock(GGraphicsLock, SLT_ReadOnly);
+		// create new graphics state
+		OutCachedState = new FGraphicsPipelineState();
 
-		FGraphicsPipelineState** Found = GGraphicsPipelines.Find(*Initializer);
-		if (Found)
+		// create a compilation task, or just do it now...
+		if (DoAsyncCompile)
 		{
-			if (IsAsyncCompilationAllowed(RHICmdList))
-			{
-				FGraphEventRef& CompletionEvent = (*Found)->CompletionEvent;
-				RHICmdList.QueueAsyncPipelineStateCompile(CompletionEvent);
-			}
-			return *Found;
+			OutCachedState->CompletionEvent = TGraphTask<FCompilePipelineStateTask>::CreateTask().ConstructAndDispatchWhenReady(OutCachedState, *Initializer);
+			RHICmdList.QueueAsyncPipelineStateCompile(OutCachedState->CompletionEvent);
+		}
+		else
+		{
+			OutCachedState->RHIPipeline = RHICreateGraphicsPipelineState(*Initializer);
 		}
 
-		ScopeLock.RaiseLockToWrite();
-		Found = GGraphicsPipelines.Find(*Initializer);
-		if (Found)
+		// GGraphicsPipelineCache.Add(*Initializer, OutCachedState, LockFlags);
+		GGraphicsPipelineCache.Add(*Initializer, OutCachedState);
+	}
+	else
+	{
+		if (DoAsyncCompile)
 		{
-			if (IsAsyncCompilationAllowed(RHICmdList))
+			FGraphEventRef& CompletionEvent = OutCachedState->CompletionEvent;
+			if ( CompletionEvent.IsValid() && !CompletionEvent->IsComplete() )
 			{
-				FGraphEventRef& CompletionEvent = (*Found)->CompletionEvent;
 				RHICmdList.QueueAsyncPipelineStateCompile(CompletionEvent);
 			}
-			return *Found;
 		}
-		else
-		{
-			FGraphicsPipelineState* PipelineState = new FGraphicsPipelineState(*Initializer);
 
-			if (IsAsyncCompilationAllowed(RHICmdList))
-			{
-				PipelineState->CompletionEvent = TGraphTask<FCompilePipelineStateTask>::CreateTask().ConstructAndDispatchWhenReady(PipelineState);
-				RHICmdList.QueueAsyncPipelineStateCompile(PipelineState->CompletionEvent);
-			}
-			else
-			{
-				PipelineState->RHIPipeline = RHICreateGraphicsPipelineState(*Initializer);
-			}
-
-			GGraphicsPipelines.Add(*Initializer, PipelineState);
-			return PipelineState;
-		}
+#if PSO_TRACK_CACHE_STATS
+		OutCachedState->AddHit();
+#endif
 	}
 
-	return nullptr;
+	// GGraphicsPipelineCache.Unlock(LockFlags);
+
+	// return the state pointer
+	return OutCachedState;
 }
 
 FRHIGraphicsPipelineState* ExecuteSetGraphicsPipelineState(FGraphicsPipelineState* GraphicsPipelineState)
 {
-	ensure(GraphicsPipelineState->RHIPipeline);
-	FRWScopeLock ScopeLock(GGraphicsLock, SLT_Write);
-	GraphicsPipelineState->CompletionEvent = nullptr;
-	return GraphicsPipelineState->RHIPipeline;
-}
+/*	ensure(GraphicsPipelineState->RHIPipeline);
+	FRWScopeLock ScopeLock(GGraphicsPipelineCache.RWLock(), SLT_Write);*/
 
-void SetComputePipelineState(FRHICommandList& RHICmdList, FRHIComputeShader* ComputeShader)
-{
-	RHICmdList.SetComputePipelineState(GetAndOrCreateComputePipelineState(RHICmdList, ComputeShader));
+	FRHIGraphicsPipelineState* RHIPipeline = GraphicsPipelineState->RHIPipeline;
+
+#if PIPELINESTATECACHE_VERIFYTHREADSAFE
+	int32 Result = GraphicsPipelineState->InUseCount.Decrement();
+	check(Result >= 0);
+#endif
+	// GraphicsPipelineState->CompletionEvent = nullptr;
+	return RHIPipeline;
 }
 
-void SetGraphicsPipelineState(FRHICommandList& RHICmdList, const FGraphicsPipelineStateInitializer& Initializer, EApplyRendertargetOption ApplyFlags)
+void DumpPipelineCacheStats()
 {
-	RHICmdList.SetGraphicsPipelineState(GetAndOrCreateGraphicsPipelineState(RHICmdList, Initializer, ApplyFlags));
+#if PSO_TRACK_CACHE_STATS
+	double TotalTime = 0.0;
+	double MinTime = FLT_MAX;
+	double MaxTime = FLT_MIN;
+
+	int MinFrames = INT_MAX;
+	int MaxFrames = INT_MIN;
+	int TotalFrames = 0;
+
+	int NumUsedLastMin = 0;
+	int NumHits = 0;
+	int NumHitsAcrossFrames = 0;
+	int NumItemsMultipleFrameHits = 0;
+
+	int NumCachedItems = GGraphicsPipelineCache.Current().Num();
+
+	if (NumCachedItems == 0)
+	{
+		return;
+	}
+
+	for (auto GraphicsPipeLine : GGraphicsPipelineCache.Current())
+	{
+		TSharedPtr<FGraphicsPipelineState, ESPMode::Fast> State = GraphicsPipeLine.Value;
+
+		// calc timestats
+		double SinceUse = FPlatformTime::Seconds() - State->FirstUsedTime;
+
+		TotalTime += SinceUse;
+
+		if (SinceUse <= 30.0)
+		{
+			NumUsedLastMin++;
+		}
+
+		MinTime = FMath::Min(MinTime, SinceUse);
+		MaxTime = FMath::Max(MaxTime, SinceUse);
+
+		// calc frame stats
+		int FramesUsed = State->LastFrameUsed - State->FirstFrameUsed;
+		TotalFrames += FramesUsed;
+		MinFrames = FMath::Min(MinFrames, FramesUsed);
+		MaxFrames = FMath::Max(MaxFrames, FramesUsed);
+
+		NumHits += State->Hits;
+
+		if (State->HitsAcrossFrames > 0)
+		{
+			NumHitsAcrossFrames += State->HitsAcrossFrames;
+			NumItemsMultipleFrameHits++;
+		}
+	}
+
+	UE_LOG(LogRHI, Log, TEXT("Have %d GraphicsPipeline entries"), NumCachedItems);
+	UE_LOG(LogRHI, Log, TEXT("Secs Used: Min=%.02f, Max=%.02f, Avg=%.02f. %d used in last 30 secs"), MinTime, MaxTime, TotalTime / NumCachedItems, NumUsedLastMin);
+	UE_LOG(LogRHI, Log, TEXT("Frames Used: Min=%d, Max=%d, Avg=%d"), MinFrames, MaxFrames, TotalFrames / NumCachedItems);
+	UE_LOG(LogRHI, Log, TEXT("Hits: Avg=%d, Items with hits across frames=%d, Avg Hits across Frames=%d"), NumHits / NumCachedItems, NumItemsMultipleFrameHits, NumHitsAcrossFrames / NumCachedItems);
+
+	size_t TrackingMem = sizeof(FGraphicsPipelineStateInitializer) * GGraphicsPipelineCache.Num();
+	UE_LOG(LogRHI, Log, TEXT("Tracking Mem: %d kb"), TrackingMem / 1024);
+#else
+	UE_LOG(LogRHI, Error, TEXT("DEfine PSO_TRACK_CACHE_STATS for state and stats!"));
+#endif // PSO_VALIDATE_CACHE
 }
 
-void ClearPipelineCache()
+
+void PipelineStateCache::Shutdown()
 {
-	for (auto GraphicsPipeLine : GGraphicsPipelines)
+	// call discard twice to clear both the backing and main caches
+	for (int i = 0; i < 2; i++)
 	{
-		delete GraphicsPipeLine.Value;
- 	}
-	GGraphicsPipelines.Empty();
+		GComputePipelineCache.Discard([](FComputePipelineState* CacheItem) {
+			delete CacheItem;
+		});
 
-	for (auto ComputePipeline : GComputePipelines)
-	{
-		delete ComputePipeline.Value;
+		GGraphicsPipelineCache.DiscardAndSwap();
 	}
-	GComputePipelines.Empty();
-}
\ No newline at end of file
+}
diff --git a/Engine/Source/Runtime/RHI/Private/RHICommandList.cpp b/Engine/Source/Runtime/RHI/Private/RHICommandList.cpp
--- a/Engine/Source/Runtime/RHI/Private/RHICommandList.cpp
+++ b/Engine/Source/Runtime/RHI/Private/RHICommandList.cpp
@@ -9,6 +9,7 @@
 #include "Async/TaskGraphInterfaces.h"
 #include "RHI.h"
 #include "ScopeLock.h"
+#include "PipelineStateCache.h"
 
 DECLARE_CYCLE_STAT(TEXT("Nonimmed. Command List Execute"), STAT_NonImmedCmdListExecuteTime, STATGROUP_RHICMDLIST);
 DECLARE_DWORD_COUNTER_STAT(TEXT("Nonimmed. Command List memory"), STAT_NonImmedCmdListMemory, STATGROUP_RHICMDLIST);
@@ -1649,6 +1650,7 @@ void FRHICommandList::BeginFrame()
 		QUICK_SCOPE_CYCLE_COUNTER(BeginFrame_Flush);
 		FRHICommandListExecutor::GetImmediateCommandList().ImmediateFlush(EImmediateFlushType::FlushRHIThread);
 	}
+
 }
 
 void FRHICommandList::EndFrame()
@@ -1868,7 +1870,7 @@ void FRHICommandListBase::HandleRTThreadTaskCompletion(const FGraphEventRef& MyC
 	check(!IsImmediate() && !IsInRHIThread());
 	for (int32 Index = 0; Index < RTTasks.Num(); Index++)
 	{
-		if (!RTTasks[Index]->IsComplete())
+		if (RTTasks[Index].IsValid() && !RTTasks[Index]->IsComplete())
 		{
 			MyCompletionGraphEvent->SetGatherThreadForDontCompleteUntil(ENamedThreads::Type(CPrio_FParallelTranslateCommandList.Get() | ENamedThreads::HighTaskPriority));
 			MyCompletionGraphEvent->DontCompleteUntil(RTTasks[Index]);
diff --git a/Engine/Source/Runtime/RHI/Public/PipelineStateCache.h b/Engine/Source/Runtime/RHI/Public/PipelineStateCache.h
--- a/Engine/Source/Runtime/RHI/Public/PipelineStateCache.h
+++ b/Engine/Source/Runtime/RHI/Public/PipelineStateCache.h
@@ -13,7 +13,7 @@
 // Utility flags for modifying render target behavior on a PSO
 enum class EApplyRendertargetOption : int
 {
-	DoNothing  = 0,			// Just use the PSO from initializer's values, no checking and no modifying (faster)
+	DoNothing = 0,			// Just use the PSO from initializer's values, no checking and no modifying (faster)
 	ForceApply = 1 << 0,	// Always apply the Cmd List's Render Target formats into the PSO initializer
 	CheckApply = 1 << 1,	// Verify that the PSO's RT formats match the last Render Target formats set into the CmdList
 };
@@ -22,5 +22,16 @@ ENUM_CLASS_FLAGS(EApplyRendertargetOption);
 
 extern RHI_API void SetComputePipelineState(FRHICommandList& RHICmdList, FRHIComputeShader* ComputeShader);
 extern RHI_API void SetGraphicsPipelineState(FRHICommandList& RHICmdList, const FGraphicsPipelineStateInitializer& Initializer, EApplyRendertargetOption ApplyFlags = EApplyRendertargetOption::CheckApply);
-extern RHI_API FGraphicsPipelineState* GetAndOrCreateGraphicsPipelineState(FRHICommandList& RHICmdList, const FGraphicsPipelineStateInitializer& OriginalInitializer, EApplyRendertargetOption ApplyFlags);
-extern RHI_API void ClearPipelineCache();
\ No newline at end of file
+
+namespace PipelineStateCache
+{
+	extern RHI_API FComputePipelineState*	GetAndOrCreateComputePipelineState(FRHICommandList& RHICmdList, FRHIComputeShader* ComputeShader);
+
+	extern RHI_API FGraphicsPipelineState*	GetAndOrCreateGraphicsPipelineState(FRHICommandList& RHICmdList, const FGraphicsPipelineStateInitializer& OriginalInitializer, EApplyRendertargetOption ApplyFlags);
+
+	/* Evicts unused state entries based on r.pso.evictiontime time. Called in RHICommandList::BeginFrame */
+	extern RHI_API void FlushResources();
+
+	/* Clears all pipeline cached state. Called on shutdown, calling GetAndOrCreate after this will recreate state */
+	extern RHI_API void Shutdown();
+}
\ No newline at end of file
diff --git a/Engine/Source/Runtime/RHI/Public/RHICommandList.inl b/Engine/Source/Runtime/RHI/Public/RHICommandList.inl
--- a/Engine/Source/Runtime/RHI/Public/RHICommandList.inl
+++ b/Engine/Source/Runtime/RHI/Public/RHICommandList.inl
@@ -64,6 +64,11 @@ FORCEINLINE_DEBUGGABLE void FRHIRenderPassCommandList::ApplyCachedRenderTargets(
 	return GetParent().ApplyCachedRenderTargets(GraphicsPSOInit);
 }
 
+namespace PipelineStateCache
+{
+	/* Evicts unused state entries based on r.pso.evictiontime time. Called in RHICommandList::BeginFrame */
+	extern RHI_API void FlushResources();
+}
 
 FORCEINLINE_DEBUGGABLE void FRHICommandListImmediate::ImmediateFlush(EImmediateFlushType::Type FlushType)
 {
@@ -115,6 +120,7 @@ FORCEINLINE_DEBUGGABLE void FRHICommandListImmediate::ImmediateFlush(EImmediateF
 			WaitForDispatch();
 			WaitForRHIThreadTasks();
 			WaitForTasks(true); // these are already done, but this resets the outstanding array
+			PipelineStateCache::FlushResources();
 			FRHIResource::FlushPendingDeletes();
 		}
 		break;
diff --git a/Engine/Source/Runtime/RHI/Public/RHIDefinitions.h b/Engine/Source/Runtime/RHI/Public/RHIDefinitions.h
--- a/Engine/Source/Runtime/RHI/Public/RHIDefinitions.h
+++ b/Engine/Source/Runtime/RHI/Public/RHIDefinitions.h
@@ -924,6 +924,12 @@ inline bool RHISupportsMobileMultiView(const EShaderPlatform Platform)
 	return (Platform == EShaderPlatform::SP_OPENGL_ES3_1_ANDROID || Platform == EShaderPlatform::SP_OPENGL_ES2_ANDROID);
 }
 
+inline bool RHISupportsDrawIndirect(const EShaderPlatform Platform)
+{
+	return (Platform == EShaderPlatform::SP_METAL_SM5 || Platform == EShaderPlatform::SP_PCD3D_SM5 || Platform == EShaderPlatform::SP_VULKAN_SM5 || Platform == EShaderPlatform::SP_PS4);
+}
+
+
 // Return what the expected number of samplers will be supported by a feature level
 // Note that since the Feature Level is pretty orthogonal to the RHI/HW, this is not going to be perfect
 // If should only be used for a guess at the limit, the real limit will not be known until runtime
diff --git a/Engine/Source/Runtime/RenderCore/Private/RenderingThread.cpp b/Engine/Source/Runtime/RenderCore/Private/RenderingThread.cpp
--- a/Engine/Source/Runtime/RenderCore/Private/RenderingThread.cpp
+++ b/Engine/Source/Runtime/RenderCore/Private/RenderingThread.cpp
@@ -1051,6 +1051,11 @@ static void GameThreadWaitForTask(const FGraphEventRef& Task, bool bEmptyGameThr
 			const double StartTime = FPlatformTime::Seconds();
 			const double EndTime = StartTime + (GTimeoutForBlockOnRenderFence / 1000.0);
 
+			bool bRenderThreadEnsured = FDebug::IsEnsuring();
+
+			static bool bDisabled = FParse::Param(FCommandLine::Get(), TEXT("nothreadtimeout"));
+			static bool bGPUDebugging = FParse::Param(FCommandLine::Get(), TEXT("gpucrashdebugging"));
+
 			do
 			{
 				CheckRenderingThreadHealth();
@@ -1061,13 +1066,13 @@ static void GameThreadWaitForTask(const FGraphEventRef& Task, bool bEmptyGameThr
 				}
 				bDone = Event->Wait(WaitTime);
 
+				// track whether the thread ensured, if so don't do timeout checks
+				bRenderThreadEnsured = FDebug::IsEnsuring() || bRenderThreadEnsured;
+
 #if !WITH_EDITOR
 				// editor threads can block for quite a while... 
-				if (!bDone && !FPlatformMisc::IsDebuggerPresent())
-				{
-					static bool bDisabled = FParse::Param(FCommandLine::Get(), TEXT("nothreadtimeout"));
-					static bool bGPUDebugging = FParse::Param(FCommandLine::Get(), TEXT("gpucrashdebugging"));
-
+				if (!bDone && !bRenderThreadEnsured && !FPlatformMisc::IsDebuggerPresent())
+				{				
 					if (bGPUDebugging && FPlatformTime::Seconds() > 2.0f)
 					// Fatal timeout if we run out of time and this thread is being monitor for heartbeats
 					// (We could just let the heartbeat monitor error for us, but this leads to better diagnostics).
diff --git a/Engine/Source/Runtime/Renderer/Private/DrawingPolicy.cpp b/Engine/Source/Runtime/Renderer/Private/DrawingPolicy.cpp
--- a/Engine/Source/Runtime/Renderer/Private/DrawingPolicy.cpp
+++ b/Engine/Source/Runtime/Renderer/Private/DrawingPolicy.cpp
@@ -155,20 +155,32 @@ void FMeshDrawingPolicy::DrawMesh(FRHICommandList& RHICmdList, const FMeshBatch&
 		}
 		else
 		{
-			// Currently only supporting this path for instanced stereo.
-			const uint32 InstanceCount = ((bIsInstancedStereo && !BatchElement.bIsInstancedMesh) ? 2 : BatchElement.NumInstances);
-			SetInstanceParameters(RHICmdList, 0, InstanceCount);
-
-			RHICmdList.DrawIndexedPrimitive(
-				BatchElement.IndexBuffer->IndexBufferRHI,
-				Mesh.Type,
-				BatchElement.BaseVertexIndex,
-				0,
-				BatchElement.MaxVertexIndex - BatchElement.MinVertexIndex + 1,
-				BatchElement.FirstIndex,
-				BatchElement.NumPrimitives,
-				InstanceCount * GetInstanceFactor()
-			);
+			if (BatchElement.IndirectArgsBuffer)
+			{
+				RHICmdList.DrawIndexedPrimitiveIndirect(
+					Mesh.Type, 
+					BatchElement.IndexBuffer->IndexBufferRHI, 
+					BatchElement.IndirectArgsBuffer, 
+					0
+				);
+			}
+			else
+			{
+				// Currently only supporting this path for instanced stereo.
+				const uint32 InstanceCount = ((bIsInstancedStereo && !BatchElement.bIsInstancedMesh) ? 2 : BatchElement.NumInstances);
+				SetInstanceParameters(RHICmdList, 0, InstanceCount);
+
+				RHICmdList.DrawIndexedPrimitive(
+					BatchElement.IndexBuffer->IndexBufferRHI,
+					Mesh.Type,
+					BatchElement.BaseVertexIndex,
+					0,
+					BatchElement.MaxVertexIndex - BatchElement.MinVertexIndex + 1,
+					BatchElement.FirstIndex,
+					BatchElement.NumPrimitives,
+					InstanceCount * GetInstanceFactor()
+				);
+			}
 		}
 	}
 	else
diff --git a/Engine/Source/Runtime/Renderer/Private/IndirectLightingCache.cpp b/Engine/Source/Runtime/Renderer/Private/IndirectLightingCache.cpp
--- a/Engine/Source/Runtime/Renderer/Private/IndirectLightingCache.cpp
+++ b/Engine/Source/Runtime/Renderer/Private/IndirectLightingCache.cpp
@@ -476,7 +476,7 @@ void FIndirectLightingCache::ProcessPrimitiveUpdate(FScene* Scene, FViewInfo& Vi
 	if (!bPrecomputedLightingBufferWasDirty && PrimitiveSceneInfo->NeedsPrecomputedLightingBufferUpdate())
 	{
 		// Since the update can be executed on a threaded job (see GILCUpdatePrimTaskEnabled), no reallocation must happen here.
-		checkSlow(View.DirtyPrecomputedLightingBufferPrimitives.Num() < View.DirtyPrecomputedLightingBufferPrimitives.Max());
+		check(View.DirtyPrecomputedLightingBufferPrimitives.Num() < View.DirtyPrecomputedLightingBufferPrimitives.Max());
 		View.DirtyPrecomputedLightingBufferPrimitives.Push(PrimitiveSceneInfo);
 	}
 }
diff --git a/Engine/Source/Runtime/Renderer/Private/PostProcess/PostProcessTonemap.cpp b/Engine/Source/Runtime/Renderer/Private/PostProcess/PostProcessTonemap.cpp
--- a/Engine/Source/Runtime/Renderer/Private/PostProcess/PostProcessTonemap.cpp
+++ b/Engine/Source/Runtime/Renderer/Private/PostProcess/PostProcessTonemap.cpp
@@ -903,6 +903,7 @@ public:
 		ColorGradingLUT.Bind(ParameterMap, TEXT("ColorGradingLUT"));
 		ColorGradingLUTSampler.Bind(ParameterMap, TEXT("ColorGradingLUTSampler"));
 		InverseGamma.Bind(ParameterMap,TEXT("InverseGamma"));
+		ChromaticAberrationParams.Bind(ParameterMap, TEXT("ChromaticAberrationParams"));
 
 		ColorMatrixR_ColorCurveCd1.Bind(ParameterMap, TEXT("ColorMatrixR_ColorCurveCd1"));
 		ColorMatrixG_ColorCurveCd3Cm3.Bind(ParameterMap, TEXT("ColorMatrixG_ColorCurveCd3Cm3"));
@@ -1070,6 +1071,34 @@ public:
 			SetShaderValue(RHICmdList, ShaderRHI, InverseGamma, InvDisplayGammaValue);
 		}
 
+		{
+			// for scene color fringe
+			// from percent to fraction
+			float Offset = 0.f;
+			float StartOffset = 0.f;
+			float Multiplier = 1.f;
+
+			if (Context.View.FinalPostProcessSettings.ChromaticAberrationStartOffset < 1.f - KINDA_SMALL_NUMBER)
+			{
+				Offset = Context.View.FinalPostProcessSettings.SceneFringeIntensity * 0.01f;
+				StartOffset = Context.View.FinalPostProcessSettings.ChromaticAberrationStartOffset;
+				Multiplier = 1.f / (1.f - StartOffset);
+			}
+
+			// Wavelength of primaries in nm
+			const float PrimaryR = 611.3f;
+			const float PrimaryG = 549.1f;
+			const float PrimaryB = 464.3f;
+
+			// Simple lens chromatic aberration is roughly linear in wavelength
+			float ScaleR = 0.007f * (PrimaryR - PrimaryB);
+			float ScaleG = 0.007f * (PrimaryG - PrimaryB);
+			FVector4 Value(Offset * ScaleR * Multiplier, Offset * ScaleG * Multiplier, StartOffset, 0.f);
+
+			// we only get bigger to not leak in content from outside
+			SetShaderValue(Context.RHICmdList, ShaderRHI, ChromaticAberrationParams, Value);
+		}
+
 		{
 			FVector4 Constants[8];
 			FilmPostSetConstants(Constants, TonemapperConfBitmaskPC[ConfigIndex], &Context.View.FinalPostProcessSettings, false);
@@ -1092,6 +1121,7 @@ public:
 		Ar << P.ColorMatrixR_ColorCurveCd1 << P.ColorMatrixG_ColorCurveCd3Cm3 << P.ColorMatrixB_ColorCurveCm2 << P.ColorCurve_Cm0Cd0_Cd2_Ch0Cm1_Ch3 << P.ColorCurve_Ch1_Ch2 << P.ColorShadow_Luma << P.ColorShadow_Tint1 << P.ColorShadow_Tint2;
 		Ar << P.OverlayColor;
 		Ar << P.OutputDevice << P.OutputGamut << P.EncodeHDROutput;
+		Ar << P.ChromaticAberrationParams;
 
 		return Ar;
 	}
@@ -1106,6 +1136,7 @@ public:
 	FShaderResourceParameter ColorGradingLUT;
 	FShaderResourceParameter ColorGradingLUTSampler;
 	FShaderParameter InverseGamma;
+	FShaderParameter ChromaticAberrationParams;
 
 	FShaderParameter ColorMatrixR_ColorCurveCd1;
 	FShaderParameter ColorMatrixG_ColorCurveCd3Cm3;
@@ -1277,7 +1308,6 @@ public:
 	// VS params
 	FShaderResourceParameter EyeAdaptation;
 	FShaderParameter GrainRandomFull;
-	FShaderParameter FringeUVParams;
 	FShaderParameter DefaultEyeExposure;
 
 	// PS params
@@ -1296,7 +1326,6 @@ public:
 		// VS params
 		EyeAdaptation.Bind(Initializer.ParameterMap, TEXT("EyeAdaptation"));
 		GrainRandomFull.Bind(Initializer.ParameterMap, TEXT("GrainRandomFull"));
-		FringeUVParams.Bind(Initializer.ParameterMap, TEXT("FringeUVParams"));
 		DefaultEyeExposure.Bind(Initializer.ParameterMap, TEXT("DefaultEyeExposure"));
 	}
 
@@ -1337,26 +1366,6 @@ public:
 			SetShaderValue(RHICmdList, ShaderRHI, DefaultEyeExposure, FixedExposure);
 		}
 
-		{
-			// for scene color fringe
-			// from percent to fraction
-			float Offset = Context.View.FinalPostProcessSettings.SceneFringeIntensity * 0.01f;
-			//FVector4 Value(1.0f - Offset * 0.5f, 1.0f - Offset, 0.0f, 0.0f);
-
-			// Wavelength of primaries in nm
-			const float PrimaryR = 611.3f;
-			const float PrimaryG = 549.1f;
-			const float PrimaryB = 464.3f;
-
-			// Simple lens chromatic aberration is roughly linear in wavelength
-			float ScaleR = 0.007f * ( PrimaryR - PrimaryB );
-			float ScaleG = 0.007f * ( PrimaryG - PrimaryB );
-			FVector4 Value( 1.0f / ( 1.0f + Offset * ScaleG ), 1.0f / ( 1.0f + Offset * ScaleR ), 0.0f, 0.0f);
-
-			// we only get bigger to not leak in content from outside
-			SetShaderValue(RHICmdList, ShaderRHI, FringeUVParams, Value);
-		}
-
 		// PS params
 		{
 			// filtering can cost performance so we use point where possible, we don't want anisotropic sampling
@@ -1389,7 +1398,7 @@ public:
 		Ar << PostprocessParameter << OutComputeTex << TonemapComputeParams;
 
 		// VS params
-		Ar << GrainRandomFull << EyeAdaptation << FringeUVParams << DefaultEyeExposure;
+		Ar << GrainRandomFull << EyeAdaptation << DefaultEyeExposure;
 
 		// PS params
 		Ar << PostProcessTonemapShaderParameters;
diff --git a/Engine/Source/Runtime/Renderer/Private/PostProcess/PostProcessTonemap.h b/Engine/Source/Runtime/Renderer/Private/PostProcess/PostProcessTonemap.h
--- a/Engine/Source/Runtime/Renderer/Private/PostProcess/PostProcessTonemap.h
+++ b/Engine/Source/Runtime/Renderer/Private/PostProcess/PostProcessTonemap.h
@@ -126,7 +126,6 @@ public:
 	FPostProcessPassParameters PostprocessParameter;
 	FShaderResourceParameter EyeAdaptation;
 	FShaderParameter GrainRandomFull;
-	FShaderParameter FringeUVParams;
 	FShaderParameter DefaultEyeExposure;
 
 	/** Initialization constructor. */
@@ -136,7 +135,6 @@ public:
 		PostprocessParameter.Bind(Initializer.ParameterMap);
 		EyeAdaptation.Bind(Initializer.ParameterMap, TEXT("EyeAdaptation"));
 		GrainRandomFull.Bind(Initializer.ParameterMap, TEXT("GrainRandomFull"));
-		FringeUVParams.Bind(Initializer.ParameterMap, TEXT("FringeUVParams"));
 		DefaultEyeExposure.Bind(Initializer.ParameterMap, TEXT("DefaultEyeExposure"));
 	}
 
@@ -193,33 +191,13 @@ public:
 			// Load a default value 
 			SetShaderValue(Context.RHICmdList, ShaderRHI, DefaultEyeExposure, FixedExposure);
 		}
-
-		{
-			// for scene color fringe
-			// from percent to fraction
-			float Offset = Context.View.FinalPostProcessSettings.SceneFringeIntensity * 0.01f;
-			//FVector4 Value(1.0f - Offset * 0.5f, 1.0f - Offset, 0.0f, 0.0f);
-
-			// Wavelength of primaries in nm
-			const float PrimaryR = 611.3f;
-			const float PrimaryG = 549.1f;
-			const float PrimaryB = 464.3f;
-
-			// Simple lens chromatic aberration is roughly linear in wavelength
-			float ScaleR = 0.007f * ( PrimaryR - PrimaryB );
-			float ScaleG = 0.007f * ( PrimaryG - PrimaryB );
-			FVector4 Value( 1.0f / ( 1.0f + Offset * ScaleG ), 1.0f / ( 1.0f + Offset * ScaleR ), 0.0f, 0.0f);
-
-			// we only get bigger to not leak in content from outside
-			SetShaderValue(Context.RHICmdList, ShaderRHI, FringeUVParams, Value);
-		}
 	}
 	
 	// FShader interface.
 	virtual bool Serialize(FArchive& Ar) override
 	{
 		bool bShaderHasOutdatedParameters = FGlobalShader::Serialize(Ar);
-		Ar << PostprocessParameter << GrainRandomFull << EyeAdaptation << FringeUVParams << DefaultEyeExposure;
+		Ar << PostprocessParameter << GrainRandomFull << EyeAdaptation << DefaultEyeExposure;
 
 		return bShaderHasOutdatedParameters;
 	}
diff --git a/Engine/Source/Runtime/Renderer/Private/ReflectionEnvironment.cpp b/Engine/Source/Runtime/Renderer/Private/ReflectionEnvironment.cpp
--- a/Engine/Source/Runtime/Renderer/Private/ReflectionEnvironment.cpp
+++ b/Engine/Source/Runtime/Renderer/Private/ReflectionEnvironment.cpp
@@ -696,12 +696,11 @@ void GatherAndSortReflectionCaptures(const FViewInfo& View, const FScene* Scene,
 
 			NewSortEntry.CaptureIndex = -1;
 			NewSortEntry.CaptureOffsetAndAverageBrightness = FVector4(CurrentCapture->CaptureOffset, 1.0f);
-
 			if (Scene->GetFeatureLevel() >= ERHIFeatureLevel::SM5)
 			{
 				const FCaptureComponentSceneState* ComponentStatePtr = Scene->ReflectionSceneData.AllocatedReflectionCaptureState.Find(CurrentCapture->Component);
-				NewSortEntry.CaptureIndex = ComponentStatePtr ? ComponentStatePtr->CaptureIndex : -1;
-				check(NewSortEntry.CaptureIndex < MaxCubemaps);
+				NewSortEntry.CaptureIndex = ComponentStatePtr ? ComponentStatePtr->CaptureIndex : 0;
+				check(NewSortEntry.CaptureIndex < MaxCubemaps || NewSortEntry.CaptureIndex == 0);
 				NewSortEntry.CaptureOffsetAndAverageBrightness.W = ComponentStatePtr ? ComponentStatePtr->AverageBrightness : 1.0f;
 			}
 
