Misc Vulkan/Lumin fixes from Dev-VR

CL 4071730, 4077567, 4077947, 4078460, 4078467, 4081212, 4081315, 4081648, 4083015

diff --git a/Engine/Plugins/Lumin/MagicLeap/Source/MagicLeap/Private/MagicLeapCustomPresent.cpp b/Engine/Plugins/Lumin/MagicLeap/Source/MagicLeap/Private/MagicLeapCustomPresent.cpp
--- a/Engine/Plugins/Lumin/MagicLeap/Source/MagicLeap/Private/MagicLeapCustomPresent.cpp
+++ b/Engine/Plugins/Lumin/MagicLeap/Source/MagicLeap/Private/MagicLeapCustomPresent.cpp
@@ -461,6 +461,7 @@ FMagicLeapCustomPresentVulkan::FMagicLeapCustomPresentVulkan(FMagicLeapHMD* plug
 	: FMagicLeapCustomPresent(plugin)
 	, RenderTargetTexture(VK_NULL_HANDLE)
 	, RenderTargetTextureAllocation(VK_NULL_HANDLE)
+	, RenderTargetTextureAllocationOffset(0)
 	, RenderTargetTextureSRGB(VK_NULL_HANDLE)
 	, LastAliasedRenderTarget(VK_NULL_HANDLE)
 {
@@ -580,7 +581,7 @@ void FMagicLeapCustomPresentVulkan::FinishRendering()
 			if (RenderTargetTextureAllocation != VK_NULL_HANDLE && LastAliasedRenderTarget != RenderTargetTexture)
 			{
 				// TODO: If RenderTargetTextureSRGB is non-null, we're leaking the previous handle here. Also leaking on shutdown.
-				RenderTargetTextureSRGB = (VkImage)FMagicLeapHelperVulkan::AliasImageSRGB((uint64)RenderTargetTextureAllocation, vp_width * 2, vp_height);
+				RenderTargetTextureSRGB = (VkImage)FMagicLeapHelperVulkan::AliasImageSRGB((uint64)RenderTargetTextureAllocation, RenderTargetTextureAllocationOffset, vp_width * 2, vp_height);
 				check(RenderTargetTextureSRGB != VK_NULL_HANDLE);
 				LastAliasedRenderTarget = RenderTargetTexture;
 				UE_LOG(LogMagicLeap, Log, TEXT("Aliased render target for correct sRGB ouput."));
@@ -631,6 +632,7 @@ void FMagicLeapCustomPresentVulkan::UpdateViewport(const FViewport& Viewport, FR
 
 	RenderTargetTexture = (VkImage)RT->GetNativeResource();
 	RenderTargetTextureAllocation = reinterpret_cast<FVulkanTexture2D*>(RT->GetTexture2D())->Surface.GetAllocationHandle();
+	RenderTargetTextureAllocationOffset = reinterpret_cast<FVulkanTexture2D*>(RT->GetTexture2D())->Surface.GetAllocationOffset();
 
 	InViewportRHI->SetCustomPresent(this);
 	
diff --git a/Engine/Plugins/Lumin/MagicLeap/Source/MagicLeap/Private/MagicLeapCustomPresent.h b/Engine/Plugins/Lumin/MagicLeap/Source/MagicLeap/Private/MagicLeapCustomPresent.h
--- a/Engine/Plugins/Lumin/MagicLeap/Source/MagicLeap/Private/MagicLeapCustomPresent.h
+++ b/Engine/Plugins/Lumin/MagicLeap/Source/MagicLeap/Private/MagicLeapCustomPresent.h
@@ -232,6 +232,7 @@ public:
 protected:
 	VkImage RenderTargetTexture = VK_NULL_HANDLE;
 	VkDeviceMemory RenderTargetTextureAllocation = VK_NULL_HANDLE;
+	uint64 RenderTargetTextureAllocationOffset = 0;
 	VkImage RenderTargetTextureSRGB = VK_NULL_HANDLE;
 	VkImage LastAliasedRenderTarget = VK_NULL_HANDLE;
 };
diff --git a/Engine/Plugins/Lumin/MagicLeap/Source/MagicLeapHelperVulkan/Private/MagicLeapHelperVulkan.cpp b/Engine/Plugins/Lumin/MagicLeap/Source/MagicLeapHelperVulkan/Private/MagicLeapHelperVulkan.cpp
--- a/Engine/Plugins/Lumin/MagicLeap/Source/MagicLeapHelperVulkan/Private/MagicLeapHelperVulkan.cpp
+++ b/Engine/Plugins/Lumin/MagicLeap/Source/MagicLeapHelperVulkan/Private/MagicLeapHelperVulkan.cpp
@@ -147,7 +147,7 @@ void FMagicLeapHelperVulkan::SignalObjects(uint64 SignalObject0, uint64 SignalOb
 #endif
 }
 
-uint64 FMagicLeapHelperVulkan::AliasImageSRGB(const uint64 Allocation, const uint32 Width, const uint32 Height)
+uint64 FMagicLeapHelperVulkan::AliasImageSRGB(const uint64 Allocation, const uint64 AllocationOffset, const uint32 Width, const uint32 Height)
 {
 #if !PLATFORM_MAC
 	VkImageCreateInfo ImageCreateInfo;
@@ -175,15 +175,12 @@ uint64 FMagicLeapHelperVulkan::AliasImageSRGB(const uint64 Allocation, const uin
 	ImageCreateInfo.samples = VK_SAMPLE_COUNT_1_BIT;
 	ImageCreateInfo.initialLayout = VK_IMAGE_LAYOUT_PREINITIALIZED;
 
-	FVulkanDynamicRHI* const RHI = (FVulkanDynamicRHI*)GDynamicRHI;
-	
+	FVulkanDynamicRHI* RHI = (FVulkanDynamicRHI*)GDynamicRHI;
+	FVulkanDevice* Device = RHI->GetDevice();
 	VkImage Result = VK_NULL_HANDLE;
-	VERIFYVULKANRESULT(VulkanRHI::vkCreateImage(RHI->GetDevice()->GetInstanceHandle(), &ImageCreateInfo, nullptr, &Result));
-
-	// Assuming the offset is 0, we're not currently querying the surface for this.
-	const VkDeviceSize Offset = 0;
+	VERIFYVULKANRESULT(VulkanRHI::vkCreateImage(Device->GetInstanceHandle(), &ImageCreateInfo, nullptr, &Result));
 
-	VERIFYVULKANRESULT(VulkanRHI::vkBindImageMemory(RHI->GetDevice()->GetInstanceHandle(), Result, (VkDeviceMemory)Allocation, Offset));
+	VERIFYVULKANRESULT(VulkanRHI::vkBindImageMemory(Device->GetInstanceHandle(), Result, (VkDeviceMemory)Allocation, AllocationOffset));
 
 	return (uint64)Result;
 #else
diff --git a/Engine/Plugins/Lumin/MagicLeap/Source/MagicLeapHelperVulkan/Public/MagicLeapHelperVulkan.h b/Engine/Plugins/Lumin/MagicLeap/Source/MagicLeapHelperVulkan/Public/MagicLeapHelperVulkan.h
--- a/Engine/Plugins/Lumin/MagicLeap/Source/MagicLeapHelperVulkan/Public/MagicLeapHelperVulkan.h
+++ b/Engine/Plugins/Lumin/MagicLeap/Source/MagicLeapHelperVulkan/Public/MagicLeapHelperVulkan.h
@@ -43,5 +43,5 @@ public:
 	static void BlitImage(uint64 SrcName, int32 SrcLevel, int32 SrcX, int32 SrcY, int32 SrcZ, int32 SrcWidth, int32 SrcHeight, int32 SrcDepth, uint64 DstName, int32 DstLevel, int32 DstX, int32 DstY, int32 DstZ, int32 DstWidth, int32 DstHeight, int32 DstDepth);
 	static void SignalObjects(uint64 SignalObject0, uint64 SignalObject1);
 	static void TestClear(uint64 Dest);
-	static uint64 AliasImageSRGB(const uint64 Allocation, const uint32 Width, const uint32 Height);
+	static uint64 AliasImageSRGB(const uint64 Allocation, const uint64 AllocationOffset, const uint32 Width, const uint32 Height);
 };
diff --git a/Engine/Source/Developer/Lumin/LuminTargetPlatform/Private/LuminTargetPlatform.cpp b/Engine/Source/Developer/Lumin/LuminTargetPlatform/Private/LuminTargetPlatform.cpp
--- a/Engine/Source/Developer/Lumin/LuminTargetPlatform/Private/LuminTargetPlatform.cpp
+++ b/Engine/Source/Developer/Lumin/LuminTargetPlatform/Private/LuminTargetPlatform.cpp
@@ -32,7 +32,7 @@ FLuminTargetPlatform::FLuminTargetPlatform(bool bIsClient)
 	GetAllTargetedShaderFormats(TargetedShaderFormats);
 
 	// If we are targeting ES 2.0/3.1, we also must cook encoded HDR reflection captures
-	static FName NAME_VULKAN_ES31(TEXT("SF_VULKAN_ES31_ANDROID"));
+	static FName NAME_VULKAN_ES31(TEXT("SF_VULKAN_ES31_LUMIN"));
 	static FName NAME_GLSL_ES2(TEXT("GLSL_ES2"));
 	static FName NAME_GLSL_SM5(TEXT("GLSL_430"));
 	bRequiresEncodedHDRReflectionCaptures = TargetedShaderFormats.Contains(NAME_VULKAN_ES31)
@@ -89,6 +89,33 @@ static bool LuminSupportsVulkan(const FConfigFile& EngineSettings)
 	return bSupportsVulkan;
 }
 
+bool FLuminTargetPlatform::SupportsFeature(ETargetPlatformFeatures Feature) const
+{
+	switch (Feature)
+	{
+	case ETargetPlatformFeatures::Packaging:
+		return true;
+
+	case ETargetPlatformFeatures::LowQualityLightmaps:
+	case ETargetPlatformFeatures::MobileRendering:
+		return SupportsMobileRendering() || SupportsVulkan();
+
+	case ETargetPlatformFeatures::HighQualityLightmaps:
+	//#todo-rco: Enable when Vulkan supports it
+	//case ETargetPlatformFeatures::Tessellation:
+	case ETargetPlatformFeatures::DeferredRendering:
+		return SupportsDesktopRendering();
+
+	case ETargetPlatformFeatures::SoftwareOcclusion:
+		return SupportsSoftwareOcclusion();
+
+	default:
+		break;
+	}
+
+	return TTargetPlatformBase<FAndroidPlatformProperties>::SupportsFeature(Feature);
+}
+
 void FLuminTargetPlatform::InitializeDeviceDetection()
 {
 	DeviceDetection = FModuleManager::LoadModuleChecked<IAndroidDeviceDetectionModule>("AndroidDeviceDetection").GetAndroidDeviceDetection(TEXT("Lumin"));
@@ -111,15 +138,15 @@ void FLuminTargetPlatform::GetAllPossibleShaderFormats( TArray<FName>& OutFormat
 //	static FName NAME_GLSL_310_ES_EXT(TEXT("GLSL_310_ES_EXT"));
 //	static FName NAME_GLSL_SM4(TEXT("GLSL_150"));
 	static FName NAME_GLSL_SM5(TEXT("GLSL_430"));
-	static FName NAME_VULKAN_SM5(TEXT("SF_VULKAN_SM5"));
-	static FName NAME_VULKAN_ES31(TEXT("SF_VULKAN_ES31_ANDROID"));
+	static FName NAME_VULKAN_SM5_LUMIN(TEXT("SF_VULKAN_SM5_LUMIN"));
+	static FName NAME_VULKAN_ES31_LUMIN(TEXT("SF_VULKAN_ES31_LUMIN"));
 
 	if (SupportsMobileRendering())
 	{
 		OutFormats.AddUnique(NAME_GLSL_ES2);
 		if (LuminSupportsVulkan(EngineSettings))
 		{
-			OutFormats.AddUnique(NAME_VULKAN_ES31);
+			OutFormats.AddUnique(NAME_VULKAN_ES31_LUMIN);
 		}
 	}
 
@@ -131,7 +158,7 @@ void FLuminTargetPlatform::GetAllPossibleShaderFormats( TArray<FName>& OutFormat
 
 		if (LuminSupportsVulkan(EngineSettings))
 		{
-			OutFormats.AddUnique(NAME_VULKAN_SM5);
+			OutFormats.AddUnique(NAME_VULKAN_SM5_LUMIN);
 		}
 	}
 }
diff --git a/Engine/Source/Developer/Lumin/LuminTargetPlatform/Private/LuminTargetPlatform.h b/Engine/Source/Developer/Lumin/LuminTargetPlatform/Private/LuminTargetPlatform.h
--- a/Engine/Source/Developer/Lumin/LuminTargetPlatform/Private/LuminTargetPlatform.h
+++ b/Engine/Source/Developer/Lumin/LuminTargetPlatform/Private/LuminTargetPlatform.h
@@ -87,6 +87,8 @@ public:
 	virtual bool SupportsMobileRendering() const;// override;
 	virtual void InitializeDeviceDetection() override;
 
+	virtual bool SupportsFeature(ETargetPlatformFeatures Feature) const override;
+
 protected:
 	// Holds the Engine INI settings, for quick use.
 	FConfigFile EngineSettings;
diff --git a/Engine/Source/Developer/SharedSettingsWidgets/Private/ShaderFormatsPropertyDetails.cpp b/Engine/Source/Developer/SharedSettingsWidgets/Private/ShaderFormatsPropertyDetails.cpp
--- a/Engine/Source/Developer/SharedSettingsWidgets/Private/ShaderFormatsPropertyDetails.cpp
+++ b/Engine/Source/Developer/SharedSettingsWidgets/Private/ShaderFormatsPropertyDetails.cpp
@@ -87,10 +87,12 @@ static FText GetFriendlyNameFromRHIName(const FString& InRHIName)
 			FriendlyRHIName = LOCTEXT("VulkanSM4", "Vulkan (SM4)");
 			break;
 		case SP_VULKAN_SM5:
+		case SP_VULKAN_SM5_LUMIN:
 			FriendlyRHIName = LOCTEXT("VulkanSM5", "Vulkan (SM5)");
 			break;
 		case SP_VULKAN_PCES3_1:
 		case SP_VULKAN_ES3_1_ANDROID:
+		case SP_VULKAN_ES3_1_LUMIN:
 			FriendlyRHIName = LOCTEXT("VulkanES31", "Vulkan (ES 3.1)");
 			break;	
 		default:
diff --git a/Engine/Source/Developer/VulkanShaderFormat/Private/VulkanShaderFormat.cpp b/Engine/Source/Developer/VulkanShaderFormat/Private/VulkanShaderFormat.cpp
--- a/Engine/Source/Developer/VulkanShaderFormat/Private/VulkanShaderFormat.cpp
+++ b/Engine/Source/Developer/VulkanShaderFormat/Private/VulkanShaderFormat.cpp
@@ -11,11 +11,13 @@
 static FName NAME_VULKAN_ES3_1_ANDROID(TEXT("SF_VULKAN_ES31_ANDROID"));
 static FName NAME_VULKAN_ES3_1_ANDROID_NOUB(TEXT("SF_VULKAN_ES31_ANDROID_NOUB"));
 static FName NAME_VULKAN_ES3_1(TEXT("SF_VULKAN_ES31"));
+static FName NAME_VULKAN_ES3_1_LUMIN(TEXT("SF_VULKAN_ES31_LUMIN"));
 static FName NAME_VULKAN_ES3_1_NOUB(TEXT("SF_VULKAN_ES31_NOUB"));
 static FName NAME_VULKAN_SM4_NOUB(TEXT("SF_VULKAN_SM4_NOUB"));
 static FName NAME_VULKAN_SM4(TEXT("SF_VULKAN_SM4"));
 static FName NAME_VULKAN_SM5_NOUB(TEXT("SF_VULKAN_SM5_NOUB"));
 static FName NAME_VULKAN_SM5(TEXT("SF_VULKAN_SM5"));
+static FName NAME_VULKAN_SM5_LUMIN(TEXT("SF_VULKAN_SM5_LUMIN"));
 
 class FShaderFormatVulkan : public IShaderFormat
 {
@@ -31,11 +33,11 @@ class FShaderFormatVulkan : public IShaderFormat
 		{
 			return UE_SHADER_VULKAN_SM5_VER;
 		}
-		else if (Format == NAME_VULKAN_SM5 || Format == NAME_VULKAN_SM5_NOUB)
+		else if (Format == NAME_VULKAN_SM5 || Format == NAME_VULKAN_SM5_NOUB || Format == NAME_VULKAN_SM5_LUMIN)
 		{
 			return UE_SHADER_VULKAN_SM5_VER;
 		}
-		else if (Format == NAME_VULKAN_ES3_1_ANDROID || Format == NAME_VULKAN_ES3_1_ANDROID_NOUB || Format == NAME_VULKAN_ES3_1 || Format == NAME_VULKAN_ES3_1_NOUB)
+		else if (Format == NAME_VULKAN_ES3_1_ANDROID || Format == NAME_VULKAN_ES3_1_ANDROID_NOUB || Format == NAME_VULKAN_ES3_1 || Format == NAME_VULKAN_ES3_1_NOUB || Format == NAME_VULKAN_ES3_1_LUMIN)
 		{
 			return UE_SHADER_VULKAN_ES3_1_VER;
 		}
@@ -55,9 +57,11 @@ public:
 	{
 		OutFormats.Add(NAME_VULKAN_SM4);
 		OutFormats.Add(NAME_VULKAN_SM5);
+		OutFormats.Add(NAME_VULKAN_SM5_LUMIN);
 		OutFormats.Add(NAME_VULKAN_ES3_1_ANDROID);
 		OutFormats.Add(NAME_VULKAN_ES3_1_ANDROID_NOUB);
 		OutFormats.Add(NAME_VULKAN_ES3_1);
+		OutFormats.Add(NAME_VULKAN_ES3_1_LUMIN);
 		OutFormats.Add(NAME_VULKAN_ES3_1_NOUB);
 		OutFormats.Add(NAME_VULKAN_SM4_NOUB);
 		OutFormats.Add(NAME_VULKAN_SM5_NOUB);
@@ -66,7 +70,7 @@ public:
 	virtual void CompileShader(FName Format, const struct FShaderCompilerInput& Input, struct FShaderCompilerOutput& Output,const FString& WorkingDirectory) const
 	{
 		check(InternalGetVersion(Format) >= 0);
-		if (Format == NAME_VULKAN_ES3_1)
+		if (Format == NAME_VULKAN_ES3_1 || Format == NAME_VULKAN_ES3_1_LUMIN)
 		{
 			CompileShader_Windows_Vulkan(Input, Output, WorkingDirectory, EVulkanShaderVersion::ES3_1);
 		}
@@ -94,7 +98,7 @@ public:
 		{
 			CompileShader_Windows_Vulkan(Input, Output, WorkingDirectory, EVulkanShaderVersion::SM5_NOUB);
 		}
-		else if (Format == NAME_VULKAN_SM5)
+		else if (Format == NAME_VULKAN_SM5 || Format == NAME_VULKAN_SM5_LUMIN)
 		{
 			CompileShader_Windows_Vulkan(Input, Output, WorkingDirectory, EVulkanShaderVersion::SM5);
 		}
diff --git a/Engine/Source/Programs/ShaderCompileWorker/Private/ShaderCompileWorker.cpp b/Engine/Source/Programs/ShaderCompileWorker/Private/ShaderCompileWorker.cpp
--- a/Engine/Source/Programs/ShaderCompileWorker/Private/ShaderCompileWorker.cpp
+++ b/Engine/Source/Programs/ShaderCompileWorker/Private/ShaderCompileWorker.cpp
@@ -701,12 +701,14 @@ static FName NAME_GLSL_310_ES_EXT(TEXT("GLSL_310_ES_EXT"));
 static FName NAME_SF_METAL_SM5(TEXT("SF_METAL_SM5"));
 static FName NAME_VULKAN_ES3_1_ANDROID(TEXT("SF_VULKAN_ES31_ANDROID"));
 static FName NAME_VULKAN_ES3_1_ANDROID_NOUB(TEXT("SF_VULKAN_ES31_ANDROID_NOUB"));
+static FName NAME_VULKAN_ES3_1_LUMIN(TEXT("SF_VULKAN_ES31_LUMIN"));
 static FName NAME_VULKAN_ES3_1(TEXT("SF_VULKAN_ES31"));
 static FName NAME_VULKAN_ES3_1_NOUB(TEXT("SF_VULKAN_ES31_NOUB"));
 static FName NAME_VULKAN_SM4_NOUB(TEXT("SF_VULKAN_SM4_NOUB"));
 static FName NAME_VULKAN_SM4(TEXT("SF_VULKAN_SM4"));
 static FName NAME_VULKAN_SM5_NOUB(TEXT("SF_VULKAN_SM5_NOUB"));
 static FName NAME_VULKAN_SM5(TEXT("SF_VULKAN_SM5"));
+static FName NAME_VULKAN_SM5_LUMIN(TEXT("SF_VULKAN_SM5_LUMIN"));
 static FName NAME_SF_METAL_SM5_NOTESS(TEXT("SF_METAL_SM5_NOTESS"));
 static FName NAME_SF_METAL_MACES3_1(TEXT("SF_METAL_MACES3_1"));
 static FName NAME_GLSL_ES3_1_ANDROID(TEXT("GLSL_ES3_1_ANDROID"));
@@ -733,8 +735,10 @@ static EShaderPlatform FormatNameToEnum(FName ShaderFormat)
 	if (ShaderFormat == NAME_SF_METAL_SM5)		return SP_METAL_SM5;
 	if (ShaderFormat == NAME_VULKAN_SM4)			return SP_VULKAN_SM4;
 	if (ShaderFormat == NAME_VULKAN_SM5)			return SP_VULKAN_SM5;
+	if (ShaderFormat == NAME_VULKAN_SM5_LUMIN)		return SP_VULKAN_SM5_LUMIN;
 	if (ShaderFormat == NAME_VULKAN_ES3_1_ANDROID)	return SP_VULKAN_ES3_1_ANDROID;
 	if (ShaderFormat == NAME_VULKAN_ES3_1_ANDROID_NOUB)	return SP_VULKAN_ES3_1_ANDROID;
+	if (ShaderFormat == NAME_VULKAN_ES3_1_LUMIN)		return SP_VULKAN_ES3_1_LUMIN;
 	if (ShaderFormat == NAME_VULKAN_ES3_1)			return SP_VULKAN_PCES3_1;
 	if (ShaderFormat == NAME_VULKAN_ES3_1_NOUB)	return SP_VULKAN_PCES3_1;
 	if (ShaderFormat == NAME_VULKAN_SM4_NOUB)		return SP_VULKAN_SM4;
diff --git a/Engine/Source/Runtime/Engine/Private/GPUSkinCache.cpp b/Engine/Source/Runtime/Engine/Private/GPUSkinCache.cpp
--- a/Engine/Source/Runtime/Engine/Private/GPUSkinCache.cpp
+++ b/Engine/Source/Runtime/Engine/Private/GPUSkinCache.cpp
@@ -114,7 +114,7 @@ bool IsGPUSkinCacheAvailable()
 
 static inline bool DoesPlatformSupportGPUSkinCache(EShaderPlatform Platform)
 {
-	return Platform == SP_PCD3D_SM5 || Platform == SP_METAL_SM5 || Platform == SP_METAL_SM5_NOTESS || Platform == SP_METAL_MRT_MAC || Platform == SP_METAL_MRT || Platform == SP_VULKAN_SM5 || Platform == SP_OPENGL_SM5;
+	return Platform == SP_PCD3D_SM5 || Platform == SP_METAL_SM5 || Platform == SP_METAL_SM5_NOTESS || Platform == SP_METAL_MRT_MAC || Platform == SP_METAL_MRT || IsVulkanSM5Platform(Platform) || Platform == SP_OPENGL_SM5;
 }
 
 // We don't have it always enabled as it's not clear if this has a performance cost
diff --git a/Engine/Source/Runtime/Engine/Public/SceneManagement.h b/Engine/Source/Runtime/Engine/Public/SceneManagement.h
--- a/Engine/Source/Runtime/Engine/Public/SceneManagement.h
+++ b/Engine/Source/Runtime/Engine/Public/SceneManagement.h
@@ -298,7 +298,7 @@ static const int32 MAX_NUM_LIGHTMAP_COEF = 2;
 
 /** Compile out low quality lightmaps to save memory */
 // @todo-mobile: Need to fix this!
-#define ALLOW_LQ_LIGHTMAPS (PLATFORM_DESKTOP || PLATFORM_IOS || PLATFORM_ANDROID || PLATFORM_HTML5 || PLATFORM_SWITCH )
+#define ALLOW_LQ_LIGHTMAPS (PLATFORM_DESKTOP || PLATFORM_IOS || PLATFORM_ANDROID || PLATFORM_HTML5 || PLATFORM_SWITCH || PLATFORM_LUMINGL4)
 
 /** Compile out high quality lightmaps to save memory */
 #define ALLOW_HQ_LIGHTMAPS 1
@@ -838,7 +838,10 @@ public:
 inline bool DoesPlatformSupportDistanceFieldShadowing(EShaderPlatform Platform)
 {
 	// Hasn't been tested elsewhere yet
-	return Platform == SP_PCD3D_SM5 || Platform == SP_PS4 || (IsMetalPlatform(Platform) && GetMaxSupportedFeatureLevel(Platform) >= ERHIFeatureLevel::SM5 && RHIGetShaderLanguageVersion(Platform) >= 2) || Platform == SP_XBOXONE_D3D12 || Platform == SP_VULKAN_SM5;
+	return Platform == SP_PCD3D_SM5 || Platform == SP_PS4
+		|| (IsMetalPlatform(Platform) && GetMaxSupportedFeatureLevel(Platform) >= ERHIFeatureLevel::SM5 && RHIGetShaderLanguageVersion(Platform) >= 2)
+		|| Platform == SP_XBOXONE_D3D12
+		|| IsVulkanSM5Platform(Platform);
 }
 
 /** Represents a USkyLightComponent to the rendering thread. */
diff --git a/Engine/Source/Runtime/OpenGLDrv/Private/OpenGLDevice.cpp b/Engine/Source/Runtime/OpenGLDrv/Private/OpenGLDevice.cpp
--- a/Engine/Source/Runtime/OpenGLDrv/Private/OpenGLDevice.cpp
+++ b/Engine/Source/Runtime/OpenGLDrv/Private/OpenGLDevice.cpp
@@ -1220,7 +1220,7 @@ static bool VerifyCompiledShader(GLuint Shader, const ANSICHAR* GlslCode, bool I
 }
 #endif
 
-#if PLATFORM_ANDROID
+#if PLATFORM_ANDROID && !PLATFORM_LUMINGL4
 static void AttemptLinkProgramExecute(GLuint VertexShaderResource, GLuint PixelShaderResource)
 {
 	GLuint Program = glCreateProgram();
diff --git a/Engine/Source/Runtime/RHI/Private/RHI.cpp b/Engine/Source/Runtime/RHI/Private/RHI.cpp
--- a/Engine/Source/Runtime/RHI/Private/RHI.cpp
+++ b/Engine/Source/Runtime/RHI/Private/RHI.cpp
@@ -621,12 +621,14 @@ static FName NAME_GLSL_ES3_1_ANDROID(TEXT("GLSL_ES3_1_ANDROID"));
 static FName NAME_SF_METAL_SM5(TEXT("SF_METAL_SM5"));
 static FName NAME_VULKAN_ES3_1_ANDROID(TEXT("SF_VULKAN_ES31_ANDROID"));
 static FName NAME_VULKAN_ES3_1_ANDROID_NOUB(TEXT("SF_VULKAN_ES31_ANDROID_NOUB"));
+static FName NAME_VULKAN_ES3_1_LUMIN(TEXT("SF_VULKAN_ES31_LUMIN"));
 static FName NAME_VULKAN_ES3_1(TEXT("SF_VULKAN_ES31"));
 static FName NAME_VULKAN_ES3_1_NOUB(TEXT("SF_VULKAN_ES31_NOUB"));
 static FName NAME_VULKAN_SM4_NOUB(TEXT("SF_VULKAN_SM4_NOUB"));
 static FName NAME_VULKAN_SM4(TEXT("SF_VULKAN_SM4"));
 static FName NAME_VULKAN_SM5_NOUB(TEXT("SF_VULKAN_SM5_NOUB"));
 static FName NAME_VULKAN_SM5(TEXT("SF_VULKAN_SM5"));
+static FName NAME_VULKAN_SM5_LUMIN(TEXT("SF_VULKAN_SM5_LUMIN"));
 static FName NAME_SF_METAL_SM5_NOTESS(TEXT("SF_METAL_SM5_NOTESS"));
 static FName NAME_SF_METAL_MACES3_1(TEXT("SF_METAL_MACES3_1"));
 static FName NAME_SF_METAL_MACES2(TEXT("SF_METAL_MACES2"));
@@ -688,6 +690,10 @@ FName LegacyShaderPlatformToShaderFormat(EShaderPlatform Platform)
 		return (CVarUseVulkanRealUBs->GetInt() != 0) ? NAME_VULKAN_SM4 : NAME_VULKAN_SM4_NOUB;
 	case SP_VULKAN_SM5:
 		return (CVarUseVulkanRealUBs->GetInt() != 0) ? NAME_VULKAN_SM5 : NAME_VULKAN_SM5_NOUB;
+	case SP_VULKAN_SM5_LUMIN:
+		return NAME_VULKAN_SM5_LUMIN;
+	case SP_VULKAN_ES3_1_LUMIN:
+		return NAME_VULKAN_ES3_1_LUMIN;
 	case SP_VULKAN_PCES3_1:
 		return (CVarUseVulkanRealUBs->GetInt() != 0) ? NAME_VULKAN_ES3_1 : NAME_VULKAN_ES3_1_NOUB;
 	case SP_VULKAN_ES3_1_ANDROID:
@@ -726,8 +732,10 @@ EShaderPlatform ShaderFormatToLegacyShaderPlatform(FName ShaderFormat)
 	if (ShaderFormat == NAME_SF_METAL_SM5)				return SP_METAL_SM5;
 	if (ShaderFormat == NAME_VULKAN_SM4)				return SP_VULKAN_SM4;
 	if (ShaderFormat == NAME_VULKAN_SM5)				return SP_VULKAN_SM5;
+	if (ShaderFormat == NAME_VULKAN_SM5_LUMIN)			return SP_VULKAN_SM5_LUMIN;
 	if (ShaderFormat == NAME_VULKAN_ES3_1_ANDROID)		return SP_VULKAN_ES3_1_ANDROID;
 	if (ShaderFormat == NAME_VULKAN_ES3_1_ANDROID_NOUB)	return SP_VULKAN_ES3_1_ANDROID;
+	if (ShaderFormat == NAME_VULKAN_ES3_1_LUMIN)		return SP_VULKAN_ES3_1_LUMIN;
 	if (ShaderFormat == NAME_VULKAN_ES3_1)				return SP_VULKAN_PCES3_1;
 	if (ShaderFormat == NAME_VULKAN_ES3_1_NOUB)			return SP_VULKAN_PCES3_1;
 	if (ShaderFormat == NAME_VULKAN_SM4_NOUB)			return SP_VULKAN_SM4;
@@ -817,7 +825,7 @@ RHI_API bool RHISupportsTessellation(const EShaderPlatform Platform)
 {
 	if (IsFeatureLevelSupported(Platform, ERHIFeatureLevel::SM5) && !IsMetalPlatform(Platform))
 	{
-		return (Platform == SP_PCD3D_SM5) || (Platform == SP_XBOXONE_D3D12) || (Platform == SP_OPENGL_SM5) || (Platform == SP_OPENGL_ES31_EXT)/* || (Platform == SP_VULKAN_SM5)*/;
+		return (Platform == SP_PCD3D_SM5) || (Platform == SP_XBOXONE_D3D12) || (Platform == SP_OPENGL_SM5) || (Platform == SP_OPENGL_ES31_EXT)/* || (IsVulkanSM5Platform(Platform)*/;
 	}
 	// For Metal we can only support tessellation if we are willing to sacrifice backward compatibility with OS versions.
 	// As such it becomes an opt-in project setting.
diff --git a/Engine/Source/Runtime/RHI/Public/RHIDefinitions.h b/Engine/Source/Runtime/RHI/Public/RHIDefinitions.h
--- a/Engine/Source/Runtime/RHI/Public/RHIDefinitions.h
+++ b/Engine/Source/Runtime/RHI/Public/RHIDefinitions.h
@@ -29,41 +29,43 @@ static_assert(SF_NumFrequencies <= (1 << SF_NumBits), "SF_NumFrequencies will no
 /** @warning: update *LegacyShaderPlatform* when the below changes */
 enum EShaderPlatform
 {
-	SP_PCD3D_SM5		= 0,
-	SP_OPENGL_SM4		= 1,
-	SP_PS4				= 2,
+	SP_PCD3D_SM5 = 0,
+	SP_OPENGL_SM4 = 1,
+	SP_PS4 = 2,
 	/** Used when running in Feature Level ES2 in OpenGL. */
-	SP_OPENGL_PCES2		= 3,
-	SP_XBOXONE_D3D12    = 4,
-	SP_PCD3D_SM4		= 5,
-	SP_OPENGL_SM5		= 6,
+	SP_OPENGL_PCES2			= 3,
+	SP_XBOXONE_D3D12		= 4,
+	SP_PCD3D_SM4			= 5,
+	SP_OPENGL_SM5			= 6,
 	/** Used when running in Feature Level ES2 in D3D11. */
-	SP_PCD3D_ES2		= 7,
-	SP_OPENGL_ES2_ANDROID = 8,
-	SP_OPENGL_ES2_WEBGL = 9, 
-	SP_OPENGL_ES2_IOS	= 10,
-	SP_METAL			= 11,
-	SP_METAL_MRT		= 12,
-	SP_OPENGL_ES31_EXT	= 13,
+	SP_PCD3D_ES2			= 7,
+	SP_OPENGL_ES2_ANDROID	= 8,
+	SP_OPENGL_ES2_WEBGL		= 9, 
+	SP_OPENGL_ES2_IOS		= 10,
+	SP_METAL				= 11,
+	SP_METAL_MRT			= 12,
+	SP_OPENGL_ES31_EXT		= 13,
 	/** Used when running in Feature Level ES3_1 in D3D11. */
-	SP_PCD3D_ES3_1		= 14,
+	SP_PCD3D_ES3_1			= 14,
 	/** Used when running in Feature Level ES3_1 in OpenGL. */
-	SP_OPENGL_PCES3_1	= 15,
-	SP_METAL_SM5		= 16,
-	SP_VULKAN_PCES3_1	= 17,
-	SP_METAL_SM5_NOTESS	= 18,
-	SP_VULKAN_SM4		= 19,
-	SP_VULKAN_SM5		= 20,
-	SP_VULKAN_ES3_1_ANDROID = 21,
-	SP_METAL_MACES3_1 	= 22,
-	SP_METAL_MACES2		= 23,
+	SP_OPENGL_PCES3_1		= 15,
+	SP_METAL_SM5			= 16,
+	SP_VULKAN_PCES3_1		= 17,
+	SP_METAL_SM5_NOTESS		= 18,
+	SP_VULKAN_SM4			= 19,
+	SP_VULKAN_SM5			= 20,
+	SP_VULKAN_ES3_1_ANDROID	= 21,
+	SP_METAL_MACES3_1 		= 22,
+	SP_METAL_MACES2			= 23,
 	SP_OPENGL_ES3_1_ANDROID = 24,
 	SP_SWITCH				= 25,
 	SP_SWITCH_FORWARD		= 26,
-	SP_METAL_MRT_MAC	= 27,
+	SP_METAL_MRT_MAC		= 27,
+	SP_VULKAN_SM5_LUMIN		= 28,
+	SP_VULKAN_ES3_1_LUMIN	= 29,
 
-	SP_NumPlatforms		= 28,
-	SP_NumBits			= 5,
+	SP_NumPlatforms			= 30,
+	SP_NumBits				= 5,
 };
 static_assert(SP_NumPlatforms <= (1 << SP_NumBits), "SP_NumPlatforms will not fit on SP_NumBits");
 
@@ -808,8 +810,11 @@ inline bool IsMobileOpenGlPlatform(const EShaderPlatform Platform)
 inline bool IsMobilePlatform(const EShaderPlatform Platform)
 {
 	return IsES2Platform(Platform)
-		|| Platform == SP_METAL || Platform == SP_PCD3D_ES3_1 || Platform == SP_OPENGL_PCES3_1 || Platform == SP_VULKAN_ES3_1_ANDROID
-		|| Platform == SP_VULKAN_PCES3_1 || Platform == SP_METAL_MACES3_1 || Platform == SP_OPENGL_ES3_1_ANDROID || Platform == SP_SWITCH_FORWARD;
+		|| Platform == SP_METAL || Platform == SP_METAL_MACES3_1
+		|| Platform == SP_PCD3D_ES3_1
+		|| Platform == SP_OPENGL_PCES3_1 || Platform == SP_OPENGL_ES3_1_ANDROID
+		|| Platform == SP_VULKAN_ES3_1_ANDROID || Platform == SP_VULKAN_PCES3_1 || Platform == SP_VULKAN_ES3_1_LUMIN
+		|| Platform == SP_SWITCH_FORWARD;
 }
 
 inline bool IsOpenGLPlatform(const EShaderPlatform Platform)
@@ -841,7 +846,12 @@ inline bool IsPS4Platform(const EShaderPlatform Platform)
 
 inline bool IsVulkanPlatform(const EShaderPlatform Platform)
 {
-	return Platform == SP_VULKAN_SM5 || Platform == SP_VULKAN_SM4 || Platform == SP_VULKAN_PCES3_1 || Platform == SP_VULKAN_ES3_1_ANDROID;
+	return Platform == SP_VULKAN_SM5 || Platform == SP_VULKAN_SM5_LUMIN || Platform == SP_VULKAN_SM4 || Platform == SP_VULKAN_PCES3_1 || Platform == SP_VULKAN_ES3_1_ANDROID || Platform == SP_VULKAN_ES3_1_LUMIN;
+}
+
+inline bool IsVulkanSM5Platform(const EShaderPlatform Platform)
+{
+	return Platform == SP_VULKAN_SM5 || Platform == SP_VULKAN_SM5_LUMIN;
 }
 
 inline bool IsAndroidOpenGLESPlatform(const EShaderPlatform Platform)
@@ -851,7 +861,7 @@ inline bool IsAndroidOpenGLESPlatform(const EShaderPlatform Platform)
 
 inline bool IsVulkanMobilePlatform(const EShaderPlatform Platform)
 {
-	return Platform == SP_VULKAN_PCES3_1 || Platform == SP_VULKAN_ES3_1_ANDROID;
+	return Platform == SP_VULKAN_PCES3_1 || Platform == SP_VULKAN_ES3_1_ANDROID || Platform == SP_VULKAN_ES3_1_LUMIN;
 }
 
 inline bool IsD3DPlatform(const EShaderPlatform Platform, bool bIncludeXboxOne)
@@ -891,6 +901,7 @@ inline ERHIFeatureLevel::Type GetMaxSupportedFeatureLevel(EShaderPlatform InShad
 	case SP_METAL_MRT_MAC:
 	case SP_METAL_SM5_NOTESS:
 	case SP_VULKAN_SM5:
+	case SP_VULKAN_SM5_LUMIN:
 	case SP_SWITCH:
 		return ERHIFeatureLevel::SM5;
 	case SP_VULKAN_SM4:
@@ -910,6 +921,7 @@ inline ERHIFeatureLevel::Type GetMaxSupportedFeatureLevel(EShaderPlatform InShad
 	case SP_OPENGL_PCES3_1:
 	case SP_VULKAN_PCES3_1:
 	case SP_VULKAN_ES3_1_ANDROID:
+	case SP_VULKAN_ES3_1_LUMIN:
 	case SP_OPENGL_ES3_1_ANDROID:
 	case SP_SWITCH_FORWARD:
 		return ERHIFeatureLevel::ES3_1;
@@ -995,7 +1007,7 @@ inline bool RHISupportsComputeShaders(const EShaderPlatform Platform)
 
 inline bool RHISupportsGeometryShaders(const EShaderPlatform Platform)
 {
-	return IsFeatureLevelSupported(Platform, ERHIFeatureLevel::SM4) && !IsMetalPlatform(Platform) && Platform != SP_VULKAN_PCES3_1 && Platform != SP_VULKAN_ES3_1_ANDROID;
+	return IsFeatureLevelSupported(Platform, ERHIFeatureLevel::SM4) && !IsMetalPlatform(Platform) && !IsVulkanMobilePlatform(Platform);
 }
 
 inline bool RHISupportsShaderCompression(const EShaderPlatform Platform)
@@ -1021,7 +1033,7 @@ inline bool RHISupportsMobileMultiView(const EShaderPlatform Platform)
 
 inline bool RHISupportsDrawIndirect(const EShaderPlatform Platform)
 {
-	return (Platform == EShaderPlatform::SP_METAL_SM5 || Platform == EShaderPlatform::SP_PCD3D_SM5 || Platform == EShaderPlatform::SP_VULKAN_SM5 || Platform == EShaderPlatform::SP_PS4);
+	return (Platform == EShaderPlatform::SP_METAL_SM5 || Platform == EShaderPlatform::SP_PCD3D_SM5 || IsVulkanSM5Platform(Platform) || Platform == EShaderPlatform::SP_PS4);
 }
 
 inline bool RHISupportsNativeShaderLibraries(const EShaderPlatform Platform)
diff --git a/Engine/Source/Runtime/Renderer/Private/CapsuleShadowRendering.h b/Engine/Source/Runtime/Renderer/Private/CapsuleShadowRendering.h
--- a/Engine/Source/Runtime/Renderer/Private/CapsuleShadowRendering.h
+++ b/Engine/Source/Runtime/Renderer/Private/CapsuleShadowRendering.h
@@ -16,7 +16,10 @@ extern int32 GCapsuleIndirectShadows;
 inline bool DoesPlatformSupportCapsuleShadows(EShaderPlatform Platform)
 {
 	// Hasn't been tested elsewhere yet
-	return Platform == SP_PCD3D_SM5 || Platform == SP_PS4 || Platform == SP_XBOXONE_D3D12 || Platform == SP_METAL_SM5 || Platform == SP_METAL_SM5_NOTESS || Platform == SP_METAL_MRT_MAC || Platform == SP_METAL_MRT || Platform == SP_VULKAN_SM5 || Platform == SP_SWITCH;
+	return Platform == SP_PCD3D_SM5 || Platform == SP_PS4 || Platform == SP_XBOXONE_D3D12
+		|| Platform == SP_METAL_SM5 || Platform == SP_METAL_SM5_NOTESS || Platform == SP_METAL_MRT_MAC || Platform == SP_METAL_MRT
+		|| IsVulkanSM5Platform(Platform)
+		|| Platform == SP_SWITCH;
 }
 
 inline bool SupportsCapsuleShadows(ERHIFeatureLevel::Type FeatureLevel, EShaderPlatform ShaderPlatform)
diff --git a/Engine/Source/Runtime/Renderer/Private/DistanceFieldAmbientOcclusion.h b/Engine/Source/Runtime/Renderer/Private/DistanceFieldAmbientOcclusion.h
--- a/Engine/Source/Runtime/Renderer/Private/DistanceFieldAmbientOcclusion.h
+++ b/Engine/Source/Runtime/Renderer/Private/DistanceFieldAmbientOcclusion.h
@@ -28,7 +28,9 @@ extern const uint32 UpdateObjectsGroupSize;
 
 inline bool DoesPlatformSupportDistanceFieldAO(EShaderPlatform Platform)
 {
-	return Platform == SP_PCD3D_SM5 || Platform == SP_PS4 || Platform == SP_XBOXONE_D3D12 || (IsMetalPlatform(Platform) && GetMaxSupportedFeatureLevel(Platform) >= ERHIFeatureLevel::SM5 && RHIGetShaderLanguageVersion(Platform) >= 2) || Platform == SP_VULKAN_SM5;
+	return Platform == SP_PCD3D_SM5 || Platform == SP_PS4 || Platform == SP_XBOXONE_D3D12
+		|| (IsMetalPlatform(Platform) && GetMaxSupportedFeatureLevel(Platform) >= ERHIFeatureLevel::SM5 && RHIGetShaderLanguageVersion(Platform) >= 2)
+		|| IsVulkanSM5Platform(Platform);
 }
 
 extern FIntPoint GetBufferSizeForAO();
diff --git a/Engine/Source/Runtime/Renderer/Private/PostProcess/DiaphragmDOF.h b/Engine/Source/Runtime/Renderer/Private/PostProcess/DiaphragmDOF.h
--- a/Engine/Source/Runtime/Renderer/Private/PostProcess/DiaphragmDOF.h
+++ b/Engine/Source/Runtime/Renderer/Private/PostProcess/DiaphragmDOF.h
@@ -114,7 +114,7 @@ inline bool IsSupported(EShaderPlatform ShaderPlatform)
 	#endif
 
 	// Only compile diaphragm DOF on platform it has been tested to ensure this is not blocking anyone else.
-	return ShaderPlatform == SP_PCD3D_SM5 || ShaderPlatform == SP_XBOXONE_D3D12 || ShaderPlatform == SP_PS4 || ShaderPlatform == SP_VULKAN_SM5;
+	return ShaderPlatform == SP_PCD3D_SM5 || ShaderPlatform == SP_XBOXONE_D3D12 || ShaderPlatform == SP_PS4 || IsVulkanSM5Platform(ShaderPlatform);
 }
 
 
diff --git a/Engine/Source/Runtime/Renderer/Private/PostProcess/PostProcessTonemap.cpp b/Engine/Source/Runtime/Renderer/Private/PostProcess/PostProcessTonemap.cpp
--- a/Engine/Source/Runtime/Renderer/Private/PostProcess/PostProcessTonemap.cpp
+++ b/Engine/Source/Runtime/Renderer/Private/PostProcess/PostProcessTonemap.cpp
@@ -1246,7 +1246,7 @@ class FPostProcessTonemapPS_ES2 : public FGlobalShader
 
 		static const auto CVarMobileMSAA = IConsoleManager::Get().FindTConsoleVariableDataInt(TEXT("r.MobileMSAA"));
 		const EShaderPlatform ShaderPlatform = GShaderPlatformForFeatureLevel[View.GetFeatureLevel()];
-		if ((GSupportsShaderFramebufferFetch && (ShaderPlatform == SP_METAL || ShaderPlatform == SP_VULKAN_PCES3_1)) && (CVarMobileMSAA ? CVarMobileMSAA->GetValueOnAnyThread() > 1 : false))
+		if ((GSupportsShaderFramebufferFetch && (ShaderPlatform == SP_METAL || IsVulkanMobilePlatform(ShaderPlatform))) && (CVarMobileMSAA ? CVarMobileMSAA->GetValueOnAnyThread() > 1 : false))
 		{
 			MobilePermutationVector.Set<FTonemapperMsaaDim>(true);
 		}
diff --git a/Engine/Source/Runtime/Renderer/Private/ShadowRendering.h b/Engine/Source/Runtime/Renderer/Private/ShadowRendering.h
--- a/Engine/Source/Runtime/Renderer/Private/ShadowRendering.h
+++ b/Engine/Source/Runtime/Renderer/Private/ShadowRendering.h
@@ -1575,7 +1575,8 @@ public:
 
 	static bool ShouldCompilePermutation(const FGlobalShaderPermutationParameters& Parameters)
 	{
-		return TShadowProjectionPS<Quality, bUseFadePlane>::ShouldCompilePermutation(Parameters) && (Parameters.Platform == SP_PCD3D_SM5 || Parameters.Platform == SP_VULKAN_SM5 || Parameters.Platform == SP_METAL_SM5 || Parameters.Platform == SP_METAL_SM5_NOTESS);
+		return TShadowProjectionPS<Quality, bUseFadePlane>::ShouldCompilePermutation(Parameters)
+			&& (Parameters.Platform == SP_PCD3D_SM5 || IsVulkanSM5Platform(Parameters.Platform) || Parameters.Platform == SP_METAL_SM5 || Parameters.Platform == SP_METAL_SM5_NOTESS);
 	}
 
 	virtual void SetParameters(
@@ -1634,7 +1635,8 @@ public:
 
 	static bool ShouldCompilePermutation(const FGlobalShaderPermutationParameters& Parameters)
 	{
-		return IsFeatureLevelSupported(Parameters.Platform, ERHIFeatureLevel::SM5) && (Parameters.Platform == SP_PCD3D_SM5 || Parameters.Platform == SP_VULKAN_SM5 || Parameters.Platform == SP_METAL_SM5 || Parameters.Platform == SP_METAL_SM5_NOTESS);
+		return IsFeatureLevelSupported(Parameters.Platform, ERHIFeatureLevel::SM5)
+			&& (Parameters.Platform == SP_PCD3D_SM5 || IsVulkanSM5Platform(Parameters.Platform) || Parameters.Platform == SP_METAL_SM5 || Parameters.Platform == SP_METAL_SM5_NOTESS);
 	}
 
 	static void ModifyCompilationEnvironment(const FGlobalShaderPermutationParameters& Parameters, FShaderCompilerEnvironment& OutEnvironment)
diff --git a/Engine/Source/Runtime/Renderer/Private/VolumetricFog.h b/Engine/Source/Runtime/Renderer/Private/VolumetricFog.h
--- a/Engine/Source/Runtime/Renderer/Private/VolumetricFog.h
+++ b/Engine/Source/Runtime/Renderer/Private/VolumetricFog.h
@@ -12,14 +12,16 @@
 
 inline bool DoesPlatformSupportVolumetricFog(EShaderPlatform Platform)
 {
-	return Platform == SP_PCD3D_SM5 || Platform == SP_PS4 || Platform == SP_XBOXONE_D3D12 || Platform == SP_METAL_SM5
-		|| Platform == SP_METAL_SM5_NOTESS || Platform == SP_VULKAN_SM5;
+	return Platform == SP_PCD3D_SM5 || Platform == SP_PS4 || Platform == SP_XBOXONE_D3D12
+		|| Platform == SP_METAL_SM5 || Platform == SP_METAL_SM5_NOTESS
+		|| IsVulkanSM5Platform(Platform);
 }
 
 inline bool DoesPlatformSupportVolumetricFogVoxelization(EShaderPlatform Platform)
 {
-	return Platform == SP_PCD3D_SM5 || Platform == SP_PS4 || Platform == SP_XBOXONE_D3D12 || Platform == SP_METAL_SM5
-	 || Platform == SP_METAL_SM5_NOTESS || Platform == SP_VULKAN_SM5;
+	return Platform == SP_PCD3D_SM5 || Platform == SP_PS4 || Platform == SP_XBOXONE_D3D12
+		|| Platform == SP_METAL_SM5 || Platform == SP_METAL_SM5_NOTESS
+		|| IsVulkanSM5Platform(Platform);
 }
 
 extern bool ShouldRenderVolumetricFog(const FScene* Scene, const FSceneViewFamily& ViewFamily);
diff --git a/Engine/Source/Runtime/VulkanRHI/Private/Lumin/VulkanLuminPlatform.cpp b/Engine/Source/Runtime/VulkanRHI/Private/Lumin/VulkanLuminPlatform.cpp
--- a/Engine/Source/Runtime/VulkanRHI/Private/Lumin/VulkanLuminPlatform.cpp
+++ b/Engine/Source/Runtime/VulkanRHI/Private/Lumin/VulkanLuminPlatform.cpp
@@ -10,6 +10,7 @@ ENUM_VK_ENTRYPOINTS_ALL(DEFINE_VK_ENTRYPOINTS)
 
 #define CHECK_VK_ENTRYPOINTS(Type,Func) if (VulkanDynamicAPI::Func == NULL) { bFoundAllEntryPoints = false; UE_LOG(LogRHI, Warning, TEXT("Failed to find entry point for %s"), TEXT(#Func)); }
 
+static bool GFoundTegraGfxDebugger = false;
 
 void* FVulkanLuminPlatform::VulkanLib = nullptr;
 bool FVulkanLuminPlatform::bAttemptedLoad = false;
@@ -95,13 +96,39 @@ void FVulkanLuminPlatform::CreateSurface(void* WindowHandle, VkInstance Instance
 	//VERIFYVULKANRESULT(vkCreateLuminSurfaceKHR(Instance, &SurfaceCreateInfo, nullptr, OutSurface));
 }
 
+void FVulkanLuminPlatform::NotifyFoundInstanceLayersAndExtensions(const TArray<FString>& Layers, const TArray<FString>& Extensions)
+{
+	if (Extensions.Find(TEXT("VK_LAYER_NV_vgd")))
+	{
+		GFoundTegraGfxDebugger = true;
+	}
+}
 
+void FVulkanLuminPlatform::NotifyFoundDeviceLayersAndExtensions(VkPhysicalDevice PhysicalDevice, const TArray<FString>& Layers, const TArray<FString>& Extensions)
+{
+	if (Extensions.Find(TEXT("VK_LAYER_NV_vgd")))
+	{
+		GFoundTegraGfxDebugger = true;
+	}
+}
 
 void FVulkanLuminPlatform::GetInstanceExtensions(TArray<const ANSICHAR*>& OutExtensions)
 {
+	if (GFoundTegraGfxDebugger)
+	{
+		OutExtensions.Add("VK_LAYER_NV_vgd");
+	}
 }
 
 void FVulkanLuminPlatform::GetDeviceExtensions(TArray<const ANSICHAR*>& OutExtensions)
 {
+	if (GFoundTegraGfxDebugger)
+	{
+		OutExtensions.Add("VK_LAYER_NV_vgd");
+	}
 }
 
+bool FVulkanLuminPlatform::ForceEnableDebugMarkers()
+{
+	return GFoundTegraGfxDebugger;
+}
diff --git a/Engine/Source/Runtime/VulkanRHI/Private/Lumin/VulkanLuminPlatform.h b/Engine/Source/Runtime/VulkanRHI/Private/Lumin/VulkanLuminPlatform.h
--- a/Engine/Source/Runtime/VulkanRHI/Private/Lumin/VulkanLuminPlatform.h
+++ b/Engine/Source/Runtime/VulkanRHI/Private/Lumin/VulkanLuminPlatform.h
@@ -32,6 +32,9 @@ public:
 	static bool LoadVulkanInstanceFunctions(VkInstance inInstance);
 	static void FreeVulkanLibrary();
 
+	static void NotifyFoundInstanceLayersAndExtensions(const TArray<FString>& Layers, const TArray<FString>& Extensions);
+	static void NotifyFoundDeviceLayersAndExtensions(VkPhysicalDevice PhysicalDevice, const TArray<FString>& Layers, const TArray<FString>& Extensions);
+
 	static void GetInstanceExtensions(TArray<const ANSICHAR*>& OutExtensions);
 	static void GetDeviceExtensions(TArray<const ANSICHAR*>& OutExtensions);
 
@@ -44,7 +47,7 @@ public:
 	static bool SupportsStandardSwapchain() { return false; }
 	static EPixelFormat GetPixelFormatForNonDefaultSwapchain() { return PF_R8G8B8A8; }
 
-	static bool SupportsMarkersWithoutExtension() { return true; }
+	static bool ForceEnableDebugMarkers();
 
 protected:
 	static void* VulkanLib;
diff --git a/Engine/Source/Runtime/VulkanRHI/Private/VulkanDevice.cpp b/Engine/Source/Runtime/VulkanRHI/Private/VulkanDevice.cpp
--- a/Engine/Source/Runtime/VulkanRHI/Private/VulkanDevice.cpp
+++ b/Engine/Source/Runtime/VulkanRHI/Private/VulkanDevice.cpp
@@ -235,11 +235,17 @@ void FVulkanDevice::CreateDevice()
 	TransferQueue = new FVulkanQueue(this, TransferQueueFamilyIndex);
 
 #if VULKAN_ENABLE_DRAW_MARKERS
-	if (bDebugMarkersFound || FVulkanPlatform::SupportsMarkersWithoutExtension())
+	if (bDebugMarkersFound || FVulkanPlatform::ForceEnableDebugMarkers())
 	{
 		DebugMarkers.CmdBegin = (PFN_vkCmdDebugMarkerBeginEXT)(void*)VulkanRHI::vkGetDeviceProcAddr(Device, "vkCmdDebugMarkerBeginEXT");
 		DebugMarkers.CmdEnd = (PFN_vkCmdDebugMarkerEndEXT)(void*)VulkanRHI::vkGetDeviceProcAddr(Device, "vkCmdDebugMarkerEndEXT");
 		DebugMarkers.CmdSetObjectName = (PFN_vkDebugMarkerSetObjectNameEXT)(void*)VulkanRHI::vkGetDeviceProcAddr(Device, "vkDebugMarkerSetObjectNameEXT");
+
+		if (DebugMarkers.CmdBegin && DebugMarkers.CmdEnd && DebugMarkers.CmdSetObjectName)
+		{
+			bDebugMarkersFound = true;
+		}
+
 		if (!DebugMarkers.CmdBegin || !DebugMarkers.CmdEnd || !DebugMarkers.CmdSetObjectName)
 		{
 			UE_LOG(LogVulkanRHI, Warning, TEXT("Extension found, but entry points for vkCmdDebugMarker(Begin|End)EXT NOT found!"));
diff --git a/Engine/Source/Runtime/VulkanRHI/Private/VulkanGenericPlatform.h b/Engine/Source/Runtime/VulkanRHI/Private/VulkanGenericPlatform.h
--- a/Engine/Source/Runtime/VulkanRHI/Private/VulkanGenericPlatform.h
+++ b/Engine/Source/Runtime/VulkanRHI/Private/VulkanGenericPlatform.h
@@ -13,7 +13,11 @@ public:
 	static bool LoadVulkanInstanceFunctions(VkInstance inInstance) { return true; }
 	static void FreeVulkanLibrary() {}
 
-	// Array of extensions for the platform - required to implement!
+	// Called after querying all the available extensions and layers
+	static void NotifyFoundInstanceLayersAndExtensions(const TArray<FString>& Layers, const TArray<FString>& Extensions) {}
+	static void NotifyFoundDeviceLayersAndExtensions(VkPhysicalDevice PhysicalDevice, const TArray<FString>& Layers, const TArray<FString>& Extensions) {}
+
+	// Array of required extensions for the platform (Required!)
 	static void GetInstanceExtensions(TArray<const ANSICHAR*>& OutExtensions);
 	static void GetDeviceExtensions(TArray<const ANSICHAR*>& OutExtensions);
 
@@ -43,7 +47,7 @@ public:
 	// Some platforms have issues with the access flags for the Present layout
 	static bool RequiresPresentLayoutFix() { return false; }
 
-	static bool SupportsMarkersWithoutExtension() { return false; }
+	static bool ForceEnableDebugMarkers() { return false; }
 
 	static bool SupportsDeviceLocalHostVisibleWithNoPenalty() { return false; }
 
diff --git a/Engine/Source/Runtime/VulkanRHI/Private/VulkanLayers.cpp b/Engine/Source/Runtime/VulkanRHI/Private/VulkanLayers.cpp
--- a/Engine/Source/Runtime/VulkanRHI/Private/VulkanLayers.cpp
+++ b/Engine/Source/Runtime/VulkanRHI/Private/VulkanLayers.cpp
@@ -248,6 +248,8 @@ void FVulkanDynamicRHI::GetInstanceLayersAndExtensions(TArray<const ANSICHAR*>&
 		UE_LOG(LogVulkanRHI, Display, TEXT("- Found instance extension %s"), *Name);
 	}
 
+	FVulkanPlatform::NotifyFoundInstanceLayersAndExtensions(FoundUniqueLayers, FoundUniqueExtensions);
+
 	if (FParse::Param(FCommandLine::Get(), TEXT("vktrace")))
 	{
 		const char* VkTraceName = "VK_LAYER_LUNARG_vktrace";
@@ -417,6 +419,8 @@ void FVulkanDevice::GetDeviceExtensionsAndLayers(TArray<const ANSICHAR*>& OutDev
 		UE_LOG(LogVulkanRHI, Display, TEXT("- Found device extension %s"), *Name);
 	}
 
+	FVulkanPlatform::NotifyFoundDeviceLayersAndExtensions(Gpu, FoundUniqueLayers, FoundUniqueExtensions);
+
 	TArray<FString> UniqueUsedDeviceExtensions;
 	auto AddDeviceLayers = [&](const char* LayerName)
 	{
@@ -537,7 +541,7 @@ void FVulkanDevice::GetDeviceExtensionsAndLayers(TArray<const ANSICHAR*>& OutDev
 	}
 
 #if VULKAN_ENABLE_DRAW_MARKERS
-	if (!bOutDebugMarkers && ListContains(AvailableExtensions, VK_EXT_DEBUG_MARKER_EXTENSION_NAME) && (GRenderDocFound || FVulkanPlatform::SupportsMarkersWithoutExtension()))
+	if (!bOutDebugMarkers && ListContains(AvailableExtensions, VK_EXT_DEBUG_MARKER_EXTENSION_NAME) && (GRenderDocFound || FVulkanPlatform::ForceEnableDebugMarkers()))
 	{
 		OutDeviceExtensions.Add(VK_EXT_DEBUG_MARKER_EXTENSION_NAME);
 		bOutDebugMarkers = true;
diff --git a/Engine/Source/Runtime/VulkanRHI/Private/VulkanRHI.cpp b/Engine/Source/Runtime/VulkanRHI/Private/VulkanRHI.cpp
--- a/Engine/Source/Runtime/VulkanRHI/Private/VulkanRHI.cpp
+++ b/Engine/Source/Runtime/VulkanRHI/Private/VulkanRHI.cpp
@@ -61,7 +61,7 @@ FDynamicRHI* FVulkanDynamicRHIModule::CreateRHI(ERHIFeatureLevel::Type InRequest
 			FParse::Param(FCommandLine::Get(), TEXT("featureleveles31")) || FParse::Param(FCommandLine::Get(), TEXT("featureleveles2"))))
 	{
 		GMaxRHIFeatureLevel = ERHIFeatureLevel::ES3_1;
-		GMaxRHIShaderPlatform = PLATFORM_ANDROID ? SP_VULKAN_ES3_1_ANDROID : SP_VULKAN_PCES3_1;
+		GMaxRHIShaderPlatform = PLATFORM_LUMIN ? SP_VULKAN_ES3_1_LUMIN : (PLATFORM_ANDROID ? SP_VULKAN_ES3_1_ANDROID : SP_VULKAN_PCES3_1);
 	}
 	else if (InRequestedFeatureLevel == ERHIFeatureLevel::SM4)
 	{
@@ -71,7 +71,7 @@ FDynamicRHI* FVulkanDynamicRHIModule::CreateRHI(ERHIFeatureLevel::Type InRequest
 	else
 	{
 		GMaxRHIFeatureLevel = ERHIFeatureLevel::SM5;
-		GMaxRHIShaderPlatform = SP_VULKAN_SM5;
+		GMaxRHIShaderPlatform = (PLATFORM_LUMINGL4 || PLATFORM_LUMIN) ? SP_VULKAN_SM5_LUMIN : SP_VULKAN_SM5;
 	}
 
 	// VULKAN_USE_MSAA_RESOLVE_ATTACHMENTS=0 requires separate MSAA and resolve textures
diff --git a/Engine/Source/Runtime/VulkanRHI/Private/VulkanViewport.cpp b/Engine/Source/Runtime/VulkanRHI/Private/VulkanViewport.cpp
--- a/Engine/Source/Runtime/VulkanRHI/Private/VulkanViewport.cpp
+++ b/Engine/Source/Runtime/VulkanRHI/Private/VulkanViewport.cpp
@@ -179,7 +179,10 @@ void FVulkanViewport::AcquireBackBuffer(FRHICommandListBase& CmdList, FVulkanBac
 
 	// Submit here so we can add a dependency with the acquired semaphore
 	CmdBuffer->End();
-	CmdBuffer->AddWaitSemaphore(VK_PIPELINE_STAGE_BOTTOM_OF_PIPE_BIT, AcquiredSemaphore);
+	if (FVulkanPlatform::SupportsStandardSwapchain())
+	{
+		CmdBuffer->AddWaitSemaphore(VK_PIPELINE_STAGE_BOTTOM_OF_PIPE_BIT, AcquiredSemaphore);
+	}
 	Device->GetGraphicsQueue()->Submit(CmdBuffer);
 	CmdBufferManager->PrepareForNewActiveCommandBuffer();
 }
@@ -615,12 +618,12 @@ bool FVulkanViewport::Present(FVulkanCommandListContext* Context, FVulkanCmdBuff
 
 	CmdBuffer->End();
 
-	if (DelayAcquireBackBuffer())
-	{
-		CmdBuffer->AddWaitSemaphore(VK_PIPELINE_STAGE_BOTTOM_OF_PIPE_BIT, AcquiredSemaphore);
-	}
 	if (FVulkanPlatform::SupportsStandardSwapchain())
 	{
+		if (DelayAcquireBackBuffer())
+		{
+			CmdBuffer->AddWaitSemaphore(VK_PIPELINE_STAGE_BOTTOM_OF_PIPE_BIT, AcquiredSemaphore);
+		}
 		Queue->Submit(CmdBuffer, RenderingDoneSemaphores[AcquiredImageIndex]->GetHandle());
 	}
 	else
diff --git a/Engine/Source/Runtime/VulkanRHI/Public/VulkanResources.h b/Engine/Source/Runtime/VulkanRHI/Public/VulkanResources.h
--- a/Engine/Source/Runtime/VulkanRHI/Public/VulkanResources.h
+++ b/Engine/Source/Runtime/VulkanRHI/Public/VulkanResources.h
@@ -310,6 +310,18 @@ public:
 		}
 	}
 
+	inline uint64 GetAllocationOffset() const
+	{
+		if (ResourceAllocation.IsValid())
+		{
+			return ResourceAllocation->GetOffset();
+		}
+		else
+		{
+			return 0;
+		}
+	}
+
 	FVulkanDevice* Device;
 
 	VkImage Image;
