RL - D3d12 issues with resource allocation and de allocation when the RHI is not alive yet.
UDN https://udn.unrealengine.com/questions/434516/check-problems-in-shipping.html

diff --git a/Engine/Source/Runtime/D3D12RHI/Private/D3D12Adapter.cpp b/Engine/Source/Runtime/D3D12RHI/Private/D3D12Adapter.cpp
--- a/Engine/Source/Runtime/D3D12RHI/Private/D3D12Adapter.cpp
+++ b/Engine/Source/Runtime/D3D12RHI/Private/D3D12Adapter.cpp
@@ -497,8 +497,7 @@ void FD3D12Adapter::Cleanup()
 	if (FrameFence)
 	{
 		FrameFence->Destroy();
-		delete FrameFence;
-		FrameFence = nullptr;
+		FrameFence.SafeRelease();
 	}
 
 	PipelineStateCache.Close();
diff --git a/Engine/Source/Runtime/D3D12RHI/Private/D3D12Adapter.h b/Engine/Source/Runtime/D3D12RHI/Private/D3D12Adapter.h
--- a/Engine/Source/Runtime/D3D12RHI/Private/D3D12Adapter.h
+++ b/Engine/Source/Runtime/D3D12RHI/Private/D3D12Adapter.h
@@ -342,7 +342,7 @@ protected:
 	TRefCountPtr<IDXGIFactory2> DxgiFactory2;
 
 	/** A Fence whos value increases every frame*/
-	FD3D12ManualFence* FrameFence;
+	TRefCountPtr<FD3D12ManualFence> FrameFence;
 
 	FD3D12DeferredDeletionQueue DeferredDeletionQueue;
 
diff --git a/Engine/Source/Runtime/D3D12RHI/Private/D3D12DirectCommandListManager.cpp b/Engine/Source/Runtime/D3D12RHI/Private/D3D12DirectCommandListManager.cpp
--- a/Engine/Source/Runtime/D3D12RHI/Private/D3D12DirectCommandListManager.cpp
+++ b/Engine/Source/Runtime/D3D12RHI/Private/D3D12DirectCommandListManager.cpp
@@ -276,8 +276,7 @@ void FD3D12CommandListManager::Destroy()
 	if (CommandListFence)
 	{
 		CommandListFence->Destroy();
-		delete CommandListFence;
-		CommandListFence = nullptr;
+		CommandListFence.SafeRelease();
 	}
 }
 
diff --git a/Engine/Source/Runtime/D3D12RHI/Private/D3D12DirectCommandListManager.h b/Engine/Source/Runtime/D3D12RHI/Private/D3D12DirectCommandListManager.h
--- a/Engine/Source/Runtime/D3D12RHI/Private/D3D12DirectCommandListManager.h
+++ b/Engine/Source/Runtime/D3D12RHI/Private/D3D12DirectCommandListManager.h
@@ -214,7 +214,7 @@ protected:
 	FD3D12CommandAllocatorManager ResourceBarrierCommandAllocatorManager;
 	FD3D12CommandAllocator* ResourceBarrierCommandAllocator;
 
-	FD3D12Fence* CommandListFence;
+	TRefCountPtr<FD3D12Fence> CommandListFence;
 
 	D3D12_COMMAND_LIST_TYPE					CommandListType;
 	ED3D12CommandQueueType					QueueType;
diff --git a/Engine/Source/Runtime/D3D12RHI/Private/Windows/WindowsD3D12Device.cpp b/Engine/Source/Runtime/D3D12RHI/Private/Windows/WindowsD3D12Device.cpp
--- a/Engine/Source/Runtime/D3D12RHI/Private/Windows/WindowsD3D12Device.cpp
+++ b/Engine/Source/Runtime/D3D12RHI/Private/Windows/WindowsD3D12Device.cpp
@@ -617,6 +617,9 @@ void FD3D12DynamicRHI::Init()
 
 	GSupportsDepthBoundsTest = SupportsDepthBoundsTest(this);
 
+	GRHICommandList.GetImmediateCommandList().SetContext(GDynamicRHI->RHIGetDefaultContext());
+	GRHICommandList.GetImmediateAsyncComputeCommandList().SetComputeContext(GDynamicRHI->RHIGetDefaultAsyncComputeContext());
+
 	// Notify all initialized FRenderResources that there's a valid RHI device to create their RHI resources for now.
 	for (TLinkedList<FRenderResource*>::TIterator ResourceIt(FRenderResource::GetResourceList()); ResourceIt; ResourceIt.Next())
 	{
diff --git a/Engine/Source/Runtime/RHI/Private/RHICommandList.cpp b/Engine/Source/Runtime/RHI/Private/RHICommandList.cpp
--- a/Engine/Source/Runtime/RHI/Private/RHICommandList.cpp
+++ b/Engine/Source/Runtime/RHI/Private/RHICommandList.cpp
@@ -745,13 +745,19 @@ struct FRHICommandRHIThreadFence final : public FRHICommand<FRHICommandRHIThread
 
 FGraphEventRef FRHICommandListImmediate::RHIThreadFence(bool bSetLockFence)
 {
-	check(IsInRenderingThread() && IsRunningRHIInSeparateThread());
-	FRHICommandRHIThreadFence* Cmd = new (AllocCommand<FRHICommandRHIThreadFence>()) FRHICommandRHIThreadFence();
-	if (bSetLockFence)
+	check(IsInRenderingThread());
+
+	if (IsRunningRHIInSeparateThread())
 	{
-		RHIThreadBufferLockFence = Cmd->Fence;
+		FRHICommandRHIThreadFence* Cmd = new (AllocCommand<FRHICommandRHIThreadFence>()) FRHICommandRHIThreadFence();
+		if (bSetLockFence)
+		{
+			RHIThreadBufferLockFence = Cmd->Fence;
+		}
+		return Cmd->Fence;
 	}
-	return Cmd->Fence;
+
+	return nullptr;
 }		
 
 DECLARE_CYCLE_STAT(TEXT("Async Compute CmdList Execute"), STAT_AsyncComputeExecute, STATGROUP_RHICMDLIST);
