Fixed: GCurrentNumDrawCallsRHI incremented twice per-drawcall in development

diff --git a/Engine/Source/Runtime/RHI/Public/RHI.h b/Engine/Source/Runtime/RHI/Public/RHI.h
--- a/Engine/Source/Runtime/RHI/Public/RHI.h
+++ b/Engine/Source/Runtime/RHI/Public/RHI.h
@@ -1475,8 +1475,7 @@ DECLARE_DWORD_COUNTER_STAT_EXTERN(TEXT("Lines drawn"),STAT_RHILines,STATGROUP_RH
 		RHI_DRAW_CALL_INC(); \
 		INC_DWORD_STAT_BY(STAT_RHITriangles,(uint32)(PrimitiveType != PT_LineList ? (NumPrimitives) : 0)); \
 		INC_DWORD_STAT_BY(STAT_RHILines,(uint32)(PrimitiveType == PT_LineList ? (NumPrimitives) : 0)); \
-		FPlatformAtomics::InterlockedAdd(&GCurrentNumPrimitivesDrawnRHI, NumPrimitives); \
-		FPlatformAtomics::InterlockedIncrement(&GCurrentNumDrawCallsRHI);
+		FPlatformAtomics::InterlockedAdd(&GCurrentNumPrimitivesDrawnRHI, NumPrimitives);
 #else
 	#define RHI_DRAW_CALL_INC() \
 		FPlatformAtomics::InterlockedIncrement(&GCurrentNumDrawCallsRHI);
