Integrated multiple RHI optimizations from Erebus branch

Change 4114412 by [REDACTED_PERSON][REDACTED_HANDLE] on 2018/06/05 18:02:15

	RHIT optimization: cache FNVNUniformBuffer address, as it's accessed very often and never changes.

Change 4115221 by [REDACTED_PERSON][REDACTED_HANDLE] on 2018/06/06 07:45:43

	RHIT optimization: replaced FNVNTempBlockManager::Get() singleton pattern with a global to avoid static overhead in frequent API calls, leading to ~0.3ms saving.

Change 4115225 by [REDACTED_PERSON][REDACTED_HANDLE] on 2018/06/06 07:46:44

	RHIT optimization: added a non-reference-counted version of FBoundShaderStateKey that's used as TMap key, avoiding RC overhead during map lookups, leading to ~0.4ms saving on Erebus.

	Full RC FBoundShaderStateKey is stored as in FCachedBoundShaderStateLink, which ensures that object lifetimes are not affected by this change.

Change 4116317 by [REDACTED_PERSON][REDACTED_HANDLE] on 2018/06/06 15:01:00

	RHIT optimization: added bound shader state history buffer and vertex declaration cache (similar to other RHIs) to improve GBoundShaderStateCache hit rate. ~1ms saving on Erebus.

Change 4129554 by [REDACTED_PERSON][REDACTED_HANDLE] on 2018/06/14 07:26:54

	RHIT optimization: improved FNVNUniformBuffer::GetSize() and FNVNMultiBufferResource::GetPoolMemory() cache line utilization. Saves ~0.6ms on Erebus.

	- Cache LayoutConstantBufferSize inside FNVNUniformBuffer
	- Cache MemoryPool ptr and offset inside FNVNMultiBufferResource, taking advantage of the hot code path in GetPoolMemory()
	- Renamed GetPoolMemory() to GetPoolAndOffset()
	- Added few asserts to validate caching assumptions in dev builds

diff --git a/Engine/Source/Runtime/RHI/Private/BoundShaderStateCache.cpp b/Engine/Source/Runtime/RHI/Private/BoundShaderStateCache.cpp
--- a/Engine/Source/Runtime/RHI/Private/BoundShaderStateCache.cpp
+++ b/Engine/Source/Runtime/RHI/Private/BoundShaderStateCache.cpp
@@ -8,8 +8,8 @@
 #include "Misc/ScopeLock.h"
 
 
-typedef TMap<FBoundShaderStateKey,FCachedBoundShaderStateLink*> FBoundShaderStateCache;
-typedef TMap<FBoundShaderStateKey,FCachedBoundShaderStateLink_Threadsafe*> FBoundShaderStateCache_Threadsafe;
+typedef TMap<FBoundShaderStateLookupKey,FCachedBoundShaderStateLink*> FBoundShaderStateCache;
+typedef TMap<FBoundShaderStateLookupKey,FCachedBoundShaderStateLink_Threadsafe*> FBoundShaderStateCache_Threadsafe;
 
 static FBoundShaderStateCache GBoundShaderStateCache;
 static FBoundShaderStateCache_Threadsafe GBoundShaderStateCache_ThreadSafe;
@@ -86,7 +86,7 @@ FCachedBoundShaderStateLink* GetCachedBoundShaderState(
 {
 	// Find the existing bound shader state in the cache.
 	return GetBoundShaderStateCache().FindRef(
-		FBoundShaderStateKey(VertexDeclaration,VertexShader,PixelShader,HullShader,DomainShader,GeometryShader)
+		FBoundShaderStateLookupKey(VertexDeclaration,VertexShader,PixelShader,HullShader,DomainShader,GeometryShader)
 		);
 }
 
@@ -115,7 +115,7 @@ FBoundShaderStateRHIRef GetCachedBoundShaderState_Threadsafe(
 	FScopeLock Lock(&BoundShaderStateCacheLock);
 	// Find the existing bound shader state in the cache.
 	FCachedBoundShaderStateLink_Threadsafe* CachedBoundShaderStateLink = GetBoundShaderStateCache_Threadsafe().FindRef(
-		FBoundShaderStateKey(VertexDeclaration,VertexShader,PixelShader,HullShader,DomainShader,GeometryShader)
+		FBoundShaderStateLookupKey(VertexDeclaration,VertexShader,PixelShader,HullShader,DomainShader,GeometryShader)
 		);
 	if(CachedBoundShaderStateLink)
 	{
diff --git a/Engine/Source/Runtime/RHI/Public/BoundShaderStateCache.h b/Engine/Source/Runtime/RHI/Public/BoundShaderStateCache.h
--- a/Engine/Source/Runtime/RHI/Public/BoundShaderStateCache.h
+++ b/Engine/Source/Runtime/RHI/Public/BoundShaderStateCache.h
@@ -33,36 +33,6 @@ public:
 		, GeometryShader(InGeometryShader)
 	{}
 
-	/**
-	* Equality is based on decl, vertex shader and pixel shader 
-	* @param Other - instance to compare against
-	* @return true if equal
-	*/
-	friend bool operator ==(const FBoundShaderStateKey& A,const FBoundShaderStateKey& B)
-	{
-		return	A.VertexDeclaration == B.VertexDeclaration && 
-			A.VertexShader == B.VertexShader && 
-			A.PixelShader == B.PixelShader && 
-			A.HullShader == B.HullShader &&
-			A.DomainShader == B.DomainShader &&
-			A.GeometryShader == B.GeometryShader;
-	}
-
-	/**
-	* Get the hash for this type. 
-	* @param Key - struct to hash
-	* @return dword hash based on type
-	*/
-	friend uint32 GetTypeHash(const FBoundShaderStateKey &Key)
-	{
-		return GetTypeHash(Key.VertexDeclaration) ^
-			GetTypeHash(Key.VertexShader) ^
-			GetTypeHash(Key.PixelShader) ^
-			GetTypeHash(Key.HullShader) ^
-			GetTypeHash(Key.DomainShader) ^
-			GetTypeHash(Key.GeometryShader);
-	}
-
 	/**
 	 * Get the RHI shader for the given frequency.
 	 */
@@ -99,6 +69,82 @@ private:
 	FDomainShaderRHIRef DomainShader;
 	/** gs for this combination */
 	FGeometryShaderRHIRef GeometryShader;
+
+	friend class FBoundShaderStateLookupKey;
+};
+
+// Non-reference-counted version of shader state key.
+// This structure is used as the actual key type for TMap, which avoids reference counting overhead during lookup.
+// Note that FCachedBoundShaderStateLink contains a full-fat reference-counted FBoundShaderStateKey, ensuring
+// correct lifetime management.
+class FBoundShaderStateLookupKey
+{
+public:
+
+	// Note: implicit cast is allowed/expected for this constructor
+	FBoundShaderStateLookupKey(const FBoundShaderStateKey& Key)
+		: VertexDeclaration(Key.VertexDeclaration.GetReference())
+		, VertexShader(Key.VertexShader.GetReference())
+		, PixelShader(Key.PixelShader.GetReference())
+		, HullShader(Key.HullShader.GetReference())
+		, DomainShader(Key.DomainShader.GetReference())
+		, GeometryShader(Key.GeometryShader.GetReference())
+	{}
+
+	FBoundShaderStateLookupKey(
+		FRHIVertexDeclaration* InVertexDeclaration,
+		FRHIVertexShader* InVertexShader,
+		FRHIPixelShader* InPixelShader,
+		FRHIHullShader* InHullShader = NULL,
+		FRHIDomainShader* InDomainShader = NULL,
+		FRHIGeometryShader* InGeometryShader = NULL
+	)
+		: VertexDeclaration(InVertexDeclaration)
+		, VertexShader(InVertexShader)
+		, PixelShader(InPixelShader)
+		, HullShader(InHullShader)
+		, DomainShader(InDomainShader)
+		, GeometryShader(InGeometryShader)
+	{}
+
+	/**
+	* Equality is based on decl, vertex shader and pixel shader
+	* @param Other - instance to compare against
+	* @return true if equal
+	*/
+	friend bool operator == (const FBoundShaderStateLookupKey& A, const FBoundShaderStateLookupKey& B)
+	{
+		return	A.VertexDeclaration == B.VertexDeclaration &&
+			A.VertexShader == B.VertexShader &&
+			A.PixelShader == B.PixelShader &&
+			A.HullShader == B.HullShader &&
+			A.DomainShader == B.DomainShader &&
+			A.GeometryShader == B.GeometryShader;
+	}
+
+	/**
+	* Get the hash for this type.
+	* @param Key - struct to hash
+	* @return dword hash based on type
+	*/
+	friend uint32 GetTypeHash(const FBoundShaderStateLookupKey& Key)
+	{
+		return GetTypeHash(Key.VertexDeclaration) ^
+			GetTypeHash(Key.VertexShader) ^
+			GetTypeHash(Key.PixelShader) ^
+			GetTypeHash(Key.HullShader) ^
+			GetTypeHash(Key.DomainShader) ^
+			GetTypeHash(Key.GeometryShader);
+	}
+
+private:
+
+	const FRHIVertexDeclaration* VertexDeclaration = nullptr;
+	const FRHIVertexShader* VertexShader = nullptr;
+	const FRHIPixelShader* PixelShader = nullptr;
+	const FRHIHullShader* HullShader = nullptr;
+	const FRHIDomainShader* DomainShader = nullptr;
+	const FRHIGeometryShader* GeometryShader = nullptr;
 };
 
 /**
diff --git a/Engine/Source/Runtime/RHI/Public/RHIResources.h b/Engine/Source/Runtime/RHI/Public/RHIResources.h
--- a/Engine/Source/Runtime/RHI/Public/RHIResources.h
+++ b/Engine/Source/Runtime/RHI/Public/RHIResources.h
@@ -318,15 +318,22 @@ public:
 	/** Initialization constructor. */
 	FRHIUniformBuffer(const FRHIUniformBufferLayout& InLayout)
 	: Layout(&InLayout)
+	, LayoutConstantBufferSize(InLayout.ConstantBufferSize)
 	{}
 
 	/** @return The number of bytes in the uniform buffer. */
-	uint32 GetSize() const { return Layout->ConstantBufferSize; }
+	uint32 GetSize() const
+	{
+		check(LayoutConstantBufferSize == Layout->ConstantBufferSize);
+		return LayoutConstantBufferSize;
+	}
 	const FRHIUniformBufferLayout& GetLayout() const { return *Layout; }
 
 private:
 	/** Layout of the uniform buffer. */
 	const FRHIUniformBufferLayout* Layout;
+
+	uint32 LayoutConstantBufferSize;
 };
 
 class FRHIIndexBuffer : public FRHIResource
