[Switch] Re-enable RHISupportsManualVertexFetch.

diff --git a/Engine/Source/Runtime/RHI/Public/RHI.h b/Engine/Source/Runtime/RHI/Public/RHI.h
--- a/Engine/Source/Runtime/RHI/Public/RHI.h
+++ b/Engine/Source/Runtime/RHI/Public/RHI.h
@@ -159,10 +159,10 @@ inline bool RHISupports4ComponentUAVReadWrite(EShaderPlatform Platform)
 }
 
 /** Whether Manual Vertex Fetch is supported for the specified shader platform.
-	Shader Platform must not use the mobile renderer, and for Metal, the shader lanugage must be at least 2. */
+	Shader Platform must not use the mobile renderer, and for Metal, the shader language must be at least 2. */
 inline bool RHISupportsManualVertexFetch(EShaderPlatform InShaderPlatform)
 {
-	return !IsOpenGLPlatform(InShaderPlatform) && !IsMobilePlatform(InShaderPlatform) && (!IsMetalPlatform(InShaderPlatform) || RHIGetShaderLanguageVersion(InShaderPlatform) >= 2);
+	return (!IsOpenGLPlatform(InShaderPlatform) || IsSwitchPlatform(InShaderPlatform)) && !IsMobilePlatform(InShaderPlatform) && (!IsMetalPlatform(InShaderPlatform) || RHIGetShaderLanguageVersion(InShaderPlatform) >= 2);
 }
 
 // Wrapper for GRHI## global variables, allows values to be overridden for mobile preview modes.
