Remove memset that was reintroduced as part of merge.
#rnx

diff --git a/Engine/Source/Runtime/RHI/Public/RHIResources.h b/Engine/Source/Runtime/RHI/Public/RHIResources.h
--- a/Engine/Source/Runtime/RHI/Public/RHIResources.h
+++ b/Engine/Source/Runtime/RHI/Public/RHIResources.h
@@ -1713,6 +1713,9 @@ struct FImmutableSamplerState
 class FGraphicsMinimalPipelineStateInitializer
 {
 public:
+	// Can't use TEnumByte<EPixelFormat> as it changes the struct to be non trivially constructible, breaking memset
+	using TRenderTargetFormats		= TStaticArray<uint8/*EPixelFormat*/, MaxSimultaneousRenderTargets>;
+	using TRenderTargetFlags		= TStaticArray<uint32, MaxSimultaneousRenderTargets>;
 
 	FGraphicsMinimalPipelineStateInitializer()
 		: BlendState(nullptr)
@@ -1720,6 +1723,8 @@ public:
 		, DepthStencilState(nullptr)
 		, PrimitiveType(PT_Num)
 	{
+		static_assert(sizeof(EPixelFormat) != sizeof(uint8), "Change TRenderTargetFormats's uint8 to EPixelFormat");
+		static_assert(PF_MAX < MAX_uint8, "TRenderTargetFormats assumes EPixelFormat can fit in a uint8!");
 	}
 
 	FGraphicsMinimalPipelineStateInitializer(
