Project option to select default graphics api for a project.
Visual cue can be gotten by hovering over the project title

diff --git a/Engine/Source/Editor/LevelEditor/Private/LevelEditor.cpp b/Engine/Source/Editor/LevelEditor/Private/LevelEditor.cpp
--- a/Engine/Source/Editor/LevelEditor/Private/LevelEditor.cpp
+++ b/Engine/Source/Editor/LevelEditor/Private/LevelEditor.cpp
@@ -108,7 +108,9 @@ public:
 		TooltipArgs.Add(TEXT("Branch"), FText::FromString(FEngineVersion::Current().GetBranch()));
 		TooltipArgs.Add(TEXT("BuildConfiguration"), EBuildConfigurations::ToText(BuildConfig));
 		TooltipArgs.Add(TEXT("BuildDate"), FText::FromString(FApp::GetBuildDate()));
-		RightContentTooltip = FText::Format(NSLOCTEXT("UnrealEditor", "TitleBarRightContentTooltip", "Version: {Version}\nBranch: {Branch}\nBuild Configuration: {BuildConfiguration}\nBuild Date: {BuildDate}"), TooltipArgs);
+		TooltipArgs.Add(TEXT("GraphicsAPI"), FText::FromString(FApp::GetGraphicsAPI()));
+
+		RightContentTooltip = FText::Format(NSLOCTEXT("UnrealEditor", "TitleBarRightContentTooltip", "Version: {Version}\nBranch: {Branch}\nBuild Configuration: {BuildConfiguration}\nBuild Date: {BuildDate}\nGraphics API: {GraphicsAPI}"), TooltipArgs);
 
 		SetToolTipText(RightContentTooltip);
 
diff --git a/Engine/Source/Runtime/Core/Private/Misc/App.cpp b/Engine/Source/Runtime/Core/Private/Misc/App.cpp
--- a/Engine/Source/Runtime/Core/Private/Misc/App.cpp
+++ b/Engine/Source/Runtime/Core/Private/Misc/App.cpp
@@ -24,6 +24,7 @@ FGuid FApp::InstanceId = FGuid::NewGuid();
 FGuid FApp::SessionId = FGuid::NewGuid();
 FString FApp::SessionName = FString();
 FString FApp::SessionOwner = FString();
+FString FApp::GraphicsAPI = FString();
 TArray<FString> FApp::SessionUsers = TArray<FString>();
 bool FApp::Standalone = true;
 bool FApp::bIsBenchmarking = false;
@@ -98,6 +99,15 @@ FString FApp::GetBuildDate()
 	return FString(BuildSettings::GetBuildDate());
 }
 
+FString FApp::GetGraphicsAPI()
+{
+	return GraphicsAPI;
+}
+
+void FApp::SetGraphicsAPI(FString APIString)
+{
+	GraphicsAPI = APIString;
+}
 
 void FApp::InitializeSession()
 {
diff --git a/Engine/Source/Runtime/Core/Public/Misc/App.h b/Engine/Source/Runtime/Core/Public/Misc/App.h
--- a/Engine/Source/Runtime/Core/Public/Misc/App.h
+++ b/Engine/Source/Runtime/Core/Public/Misc/App.h
@@ -64,6 +64,21 @@ public:
 	 */
 	static FString GetBuildDate();
 
+	/**
+	 * Gets the name of the graphics API currently in use.
+	 *
+	 * @return name of Graphics API
+	 */
+	static FString GetGraphicsAPI();
+
+	/**
+	 * Sets the Graphics Api currently in use
+	 *
+	 * @return name of Graphics API
+	 */
+	static void SetGraphicsAPI(FString APIString);
+
+
 	/**
 	 * Gets the value of ENGINE_IS_PROMOTED_BUILD.
 	 */
@@ -710,6 +725,10 @@ private:
 	/** Holds the name of the user that launched session. */
 	static FString SessionOwner;
 
+	/** Holds the name the graphics api currently in use*/
+	static FString GraphicsAPI;
+
+
 	/** List of authorized session users. */
 	static TArray<FString> SessionUsers;
 
diff --git a/Engine/Source/Runtime/Engine/Classes/Engine/RendererSettings.h b/Engine/Source/Runtime/Engine/Classes/Engine/RendererSettings.h
--- a/Engine/Source/Runtime/Engine/Classes/Engine/RendererSettings.h
+++ b/Engine/Source/Runtime/Engine/Classes/Engine/RendererSettings.h
@@ -44,6 +44,23 @@ namespace ECompositingSampleCount
 }
 
 
+/**
+ * Enumerates different graphics apis
+ */
+UENUM()
+namespace EDefaultGraphicsApi
+{
+	enum Type
+	{
+		DefaultGraphicsApi_Default = 0 UMETA(DisplayName = "Default"),
+		DefaultGraphicsApi_DX11 = 1 UMETA(DisplayName = "DirectX 11"),
+		DefaultGraphicsApi_DX12 = 2 UMETA(DisplayName = "DirectX 12"),
+		DefaultGraphicsApi_Vulkan = 3 UMETA(DisplayName = "Vulkan"),
+		DefaultGraphicsApi_OpenGL = 4 UMETA(DisplayName = "OpenGL"),
+	};
+}
+
+
 /**
 * Enumerates available mobile MSAA sample counts.
 */
@@ -689,6 +706,11 @@ class ENGINE_API URendererSettings : public UDeveloperSettings
 		ConfigRestartRequired = true))
 		uint32 bSupportMaterialLayers : 1;
 
+	UPROPERTY(config, EditAnywhere, Category = API, meta = (
+		ConsoleVariable = "r.DefaultGraphicsApi", DisplayName = "Default Graphics Api",
+		ToolTip = "Graphics API to use by default"))
+		TEnumAsByte<EDefaultGraphicsApi::Type> DefaultGraphicsApi;
+
 public:
 
 	//~ Begin UObject Interface
diff --git a/Engine/Source/Runtime/RHI/Private/Windows/WindowsDynamicRHI.cpp b/Engine/Source/Runtime/RHI/Private/Windows/WindowsDynamicRHI.cpp
--- a/Engine/Source/Runtime/RHI/Private/Windows/WindowsDynamicRHI.cpp
+++ b/Engine/Source/Runtime/RHI/Private/Windows/WindowsDynamicRHI.cpp
@@ -6,6 +6,16 @@
 #include "Misc/ConfigCacheIni.h"
 #include "Misc/MessageDialog.h"
 
+static TAutoConsoleVariable<int32> CVarDefaultGraphicsApi(
+	TEXT("r.DefaultGraphicsApi"),
+	0,
+	TEXT("0: Default\n")
+	TEXT("1: DirectX 11\n")
+	TEXT("2: DirectX 12\n")
+	TEXT("3: Vulkan\n")
+	TEXT("4: OpenGL\n"), 
+	ECVF_ReadOnly
+);
 
 static const TCHAR* GLoadedRHIModuleName;
 
@@ -50,6 +60,22 @@ static IDynamicRHIModule* LoadDynamicRHIModule(ERHIFeatureLevel::Type& DesiredFe
 	bool bForceD3D11 = FParse::Param(FCommandLine::Get(), TEXT("d3d11")) || FParse::Param(FCommandLine::Get(), TEXT("dx11")) || (bForceSM5 && !bForceVulkan && !bForceOpenGL);
 	bool bForceD3D12 = FParse::Param(FCommandLine::Get(), TEXT("d3d12")) || FParse::Param(FCommandLine::Get(), TEXT("dx12"));
 	DesiredFeatureLevel = ERHIFeatureLevel::Num;
+	//Default graphics api is only used if no command line option is specified
+	if(!(bForceVulkan||bForceOpenGL||bForceD3D10||bForceD3D11||bForceD3D12))
+	{
+		int32 DefaultGraphicsApi = CVarDefaultGraphicsApi.GetValueOnAnyThread();
+		switch(DefaultGraphicsApi)
+		{
+		case 0:	break;
+		case 1:		bForceD3D11 = true; break;
+		case 2:		bForceD3D12 = true; break;
+		case 3:		bForceVulkan = true; break;
+		case 4:		bForceOpenGL= true; break;
+		}
+	}
+
+
+
 	int32 Sum = ((bForceD3D12 ? 1 : 0) + (bForceD3D11 ? 1 : 0) + (bForceD3D10 ? 1 : 0) + (bForceOpenGL ? 1 : 0) + (bForceVulkan ? 1 : 0));
 
 	if (bForceSM5 && bForceSM4)
@@ -111,6 +137,7 @@ static IDynamicRHIModule* LoadDynamicRHIModule(ERHIFeatureLevel::Type& DesiredFe
 			DynamicRHIModule = NULL;
 		}
 		LoadedRHIModuleName = SwitchRHIModuleName;
+		FApp::SetGraphicsAPI(TEXT("Switch"));
 	}
 	else
 #endif
@@ -127,6 +154,7 @@ static IDynamicRHIModule* LoadDynamicRHIModule(ERHIFeatureLevel::Type& DesiredFe
 			DynamicRHIModule = NULL;
 		}
 		LoadedRHIModuleName = OpenGLRHIModuleName;
+		FApp::SetGraphicsAPI(TEXT("OpenGL"));
 	}
 	else if (bForceVulkan)
 	{
@@ -139,6 +167,7 @@ static IDynamicRHIModule* LoadDynamicRHIModule(ERHIFeatureLevel::Type& DesiredFe
 			DynamicRHIModule = NULL;
 		}
 		LoadedRHIModuleName = VulkanRHIModuleName;
+		FApp::SetGraphicsAPI(TEXT("Vulkan"));
 	}
 	else if (bForceD3D12 || bPreferD3D12)
 	{
@@ -163,6 +192,7 @@ static IDynamicRHIModule* LoadDynamicRHIModule(ERHIFeatureLevel::Type& DesiredFe
 		{
 			FMessageDialog::Open(EAppMsgType::Ok, NSLOCTEXT("WindowsDynamicRHI", "UseExpressionEncoder", "Fraps has been known to crash D3D12. Please use Microsoft Expression Encoder instead for capturing."));
 		}
+		FApp::SetGraphicsAPI(TEXT("DirectX 12"));
 	}
 
 	// Fallback to D3D11RHI if nothing is selected
@@ -182,6 +212,7 @@ static IDynamicRHIModule* LoadDynamicRHIModule(ERHIFeatureLevel::Type& DesiredFe
 			FMessageDialog::Open(EAppMsgType::Ok, NSLOCTEXT("WindowsDynamicRHI", "UseExpressionEncoderDX11", "Fraps has been known to crash D3D11. Please use Microsoft Expression Encoder instead for capturing."));
 		}
 		LoadedRHIModuleName = D3D11RHIModuleName;
+		FApp::SetGraphicsAPI(TEXT("DirectX 11"));
 	}
 	return DynamicRHIModule;
 }
