move cvar to consolemanager.cpp. should fix mac build

diff --git a/Engine/Source/Runtime/Core/Private/HAL/ConsoleManager.cpp b/Engine/Source/Runtime/Core/Private/HAL/ConsoleManager.cpp
--- a/Engine/Source/Runtime/Core/Private/HAL/ConsoleManager.cpp
+++ b/Engine/Source/Runtime/Core/Private/HAL/ConsoleManager.cpp
@@ -2616,3 +2616,14 @@ static TAutoConsoleVariable<int32> CVarSkyLightingQuality(
 	TEXT("  1: on\n"),
 	ECVF_Scalability | ECVF_RenderThreadSafe
 );
+
+static TAutoConsoleVariable<int32> CVarDefaultGraphicsApi(
+	TEXT("r.DefaultGraphicsApi"),
+	0,
+	TEXT("0: Default\n")
+	TEXT("1: DirectX 11\n")
+	TEXT("2: DirectX 12\n")
+	TEXT("3: Vulkan\n")
+	TEXT("4: OpenGL\n"),
+	ECVF_ReadOnly
+);
diff --git a/Engine/Source/Runtime/RHI/Private/Windows/WindowsDynamicRHI.cpp b/Engine/Source/Runtime/RHI/Private/Windows/WindowsDynamicRHI.cpp
--- a/Engine/Source/Runtime/RHI/Private/Windows/WindowsDynamicRHI.cpp
+++ b/Engine/Source/Runtime/RHI/Private/Windows/WindowsDynamicRHI.cpp
@@ -6,17 +6,6 @@
 #include "Misc/ConfigCacheIni.h"
 #include "Misc/MessageDialog.h"
 
-static TAutoConsoleVariable<int32> CVarDefaultGraphicsApi(
-	TEXT("r.DefaultGraphicsApi"),
-	0,
-	TEXT("0: Default\n")
-	TEXT("1: DirectX 11\n")
-	TEXT("2: DirectX 12\n")
-	TEXT("3: Vulkan\n")
-	TEXT("4: OpenGL\n"), 
-	ECVF_ReadOnly
-);
-
 static const TCHAR* GLoadedRHIModuleName;
 
 static bool ShouldPreferD3D12()
@@ -63,7 +52,8 @@ static IDynamicRHIModule* LoadDynamicRHIModule(ERHIFeatureLevel::Type& DesiredFe
 	//Default graphics api is only used if no command line option is specified
 	if(!(bForceVulkan||bForceOpenGL||bForceD3D10||bForceD3D11||bForceD3D12))
 	{
-		int32 DefaultGraphicsApi = CVarDefaultGraphicsApi.GetValueOnAnyThread();
+		static const auto CVarDefaultGraphicsApi = IConsoleManager::Get().FindTConsoleVariableDataInt(TEXT("r.DefaultGraphicsApi"));
+		int32 DefaultGraphicsApi = CVarDefaultGraphicsApi->GetValueOnAnyThread();
 		switch(DefaultGraphicsApi)
 		{
 		case 0:	break;
