Address PSO compile failures by marking the PSO as invalid in the PipelineCache file meta data and don't return them as part of GetOrderedPSOHashes().

This change add an extra engine flag field in the meta data that we can now use for marking other features during runtime.  Any PSO's marked as invalid won't get processed as part of the ShaderPipelineCache commandlet processing, that is Build and Expand commands will skip these entries.

diff --git a/Engine/Source/Runtime/RHI/Private/PipelineFileCache.cpp b/Engine/Source/Runtime/RHI/Private/PipelineFileCache.cpp
--- a/Engine/Source/Runtime/RHI/Private/PipelineFileCache.cpp
+++ b/Engine/Source/Runtime/RHI/Private/PipelineFileCache.cpp
@@ -55,12 +55,13 @@ enum class EPipelineCacheFileFormatVersions : uint32
 	PSOUsageMask = 13,
 	PSOBindCount = 14,
 	EOFMarker = 15,
+	EngineFlags = 16,
 };
 
 const uint64 FPipelineCacheFileFormatMagic = 0x5049504543414348; // PIPECACH
 const uint64 FPipelineCacheTOCFileFormatMagic = 0x544F435354415254; // TOCSTART
 const uint64 FPipelineCacheEOFFileFormatMagic = 0x454F462D4D41524B; // EOF-MARK
-const uint32 FPipelineCacheFileFormatCurrentVersion = (uint32)EPipelineCacheFileFormatVersions::EOFMarker;
+const uint32 FPipelineCacheFileFormatCurrentVersion = (uint32)EPipelineCacheFileFormatVersions::EngineFlags;
 
 /**
   * PipelineFileCache API access
@@ -116,7 +117,7 @@ static TAutoConsoleVariable<int32> CVarLazyLoadShadersWhenPSOCacheIsPresent(
 FRWLock FPipelineFileCache::FileCacheLock;
 FPipelineCacheFile* FPipelineFileCache::FileCache = nullptr;
 TMap<uint32, FPSOUsageData> FPipelineFileCache::RunTimeToPSOUsage;
-TMap<uint32, uint64> FPipelineFileCache::NewPSOUsageMasks;
+TMap<uint32, FPSOUsageData> FPipelineFileCache::NewPSOUsage;
 TMap<uint32, FPipelineStateStats*> FPipelineFileCache::Stats;
 TSet<FPipelineCacheFileFormatPSO> FPipelineFileCache::NewPSOs;
 uint32 FPipelineFileCache::NumNewPSOs;
@@ -184,11 +185,17 @@ FArchive& operator<<( FArchive& Ar, FPipelineStateStats& Info )
 	return Ar;
 }
 
+/**
+  * PipelineFileCache MetaData Engine Flags
+  **/
+const uint16 FPipelineCacheFlagInvalidPSO = 1 << 0;
+
 struct FPipelineCacheFileFormatPSOMetaData
 {
 	FPipelineCacheFileFormatPSOMetaData()
 	: FileOffset(0)
 	, UsageMask(0)
+	, EngineFlags(0)
 	{
 	}
 	
@@ -202,6 +209,7 @@ struct FPipelineCacheFileFormatPSOMetaData
 	FPipelineStateStats Stats;
 	TSet<FSHAHash> Shaders;
 	uint64 UsageMask;
+	uint16 EngineFlags;
 	
 	friend FArchive& operator<<(FArchive& Ar, FPipelineCacheFileFormatPSOMetaData& Info)
 	{
@@ -224,6 +232,11 @@ struct FPipelineCacheFileFormatPSOMetaData
 			Ar << Info.UsageMask;
 		}
 		
+		if(Ar.GameNetVer() >= (uint32)EPipelineCacheFileFormatVersions::EngineFlags)
+		{
+			Ar << Info.EngineFlags;
+		}
+		
 		return Ar;
 	}
 };
@@ -1279,6 +1292,9 @@ public:
         {
             TOC = GameTOC;
         }
+		
+		uint32 InvalidEntryCount = 0;
+		
         for (auto const& Entry : TOC.MetaData)
         {
             FPipelineStateStats* Stat = FPipelineFileCache::Stats.FindRef(Entry.Key);
@@ -1289,6 +1305,17 @@ public:
                 Stat->TotalBindCount = -1;
                 FPipelineFileCache::Stats.Add(Entry.Key, Stat);
             }
+#if !UE_BUILD_SHIPPING
+			if((Entry.Value.EngineFlags & FPipelineCacheFlagInvalidPSO) != 0)
+			{
+				++InvalidEntryCount;
+			}
+#endif
+        }
+		
+        if(InvalidEntryCount > 0)
+        {
+        	UE_LOG(LogRHI, Warning, TEXT("Found %d / %d PSO entries marked as invalid."), InvalidEntryCount, TOC.MetaData.Num());
         }
 		
 		SET_MEMORY_STAT(STAT_FileCacheMemory, TOC.MetaData.GetAllocatedSize());
@@ -1296,17 +1323,18 @@ public:
 		return bGameFileOk || bUserFileOk;
 	}
 	
-	void MergeUsageMasksToMetaData(TMap<uint32, uint64>& NewPSOUsageMasks, TMap<uint32, FPipelineCacheFileFormatPSOMetaData>& MetaData, bool bRemoveUpdatedentries = false)
+	void MergePSOUsageToMetaData(TMap<uint32, FPSOUsageData>& NewPSOUsage, TMap<uint32, FPipelineCacheFileFormatPSOMetaData>& MetaData, bool bRemoveUpdatedentries = false)
 	{
-		// Update MetaData Usage Masks
-		for(auto It = NewPSOUsageMasks.CreateIterator(); It; ++It)
+		for(auto It = NewPSOUsage.CreateIterator(); It; ++It)
 		{
 			auto& MaskEntry = *It;
+			
 			//Don't use FindChecked as if new PSO was not bound - it might not be in the TOC.MetaData - they are not always added in every save mode - this is not an error
 			auto* PSOMetaData = MetaData.Find(MaskEntry.Key);
 			if(PSOMetaData != nullptr)
 			{
-				PSOMetaData->UsageMask |= MaskEntry.Value;
+				PSOMetaData->UsageMask |= MaskEntry.Value.UsageMask;
+				PSOMetaData->EngineFlags |= MaskEntry.Value.EngineFlags;
 				
 				if(bRemoveUpdatedentries)
 				{
@@ -1316,7 +1344,7 @@ public:
 		}
 	}
 	
-	bool SavePipelineFileCache(FString const& FilePath, FPipelineFileCache::SaveMode Mode, TMap<uint32, FPipelineStateStats*> const& Stats, TSet<FPipelineCacheFileFormatPSO>& NewEntries, FPipelineFileCache::PSOOrder Order, TMap<uint32, uint64>& NewPSOUsageMasks)
+	bool SavePipelineFileCache(FString const& FilePath, FPipelineFileCache::SaveMode Mode, TMap<uint32, FPipelineStateStats*> const& Stats, TSet<FPipelineCacheFileFormatPSO>& NewEntries, FPipelineFileCache::PSOOrder Order, TMap<uint32, FPSOUsageData>& NewPSOUsage)
 	{
 		QUICK_SCOPE_CYCLE_COUNTER(STAT_SavePipelineFileCache);
 		double StartTime = FPlatformTime::Seconds();
@@ -1331,7 +1359,7 @@ public:
 		bool bPerformWrite = true;
 		if (FPipelineFileCache::SaveMode::Incremental == Mode)
 		{
-			bPerformWrite = NewEntries.Num() || Order != TOC.SortedOrder || NewPSOUsageMasks.Num();
+			bPerformWrite = NewEntries.Num() || Order != TOC.SortedOrder || NewPSOUsage.Num();
 			bFileWriteSuccess = !bPerformWrite;
 		}
 		
@@ -1553,8 +1581,8 @@ public:
                                 TOC.SortedOrder = FPipelineFileCache::PSOOrder::Default;
                             }
 							
-							// Update TOC Metadata usage masks and clear relevant entries NewPSOUsageMasks as we are saving this file cache TOC
-							MergeUsageMasksToMetaData(NewPSOUsageMasks, TOC.MetaData, true);
+							// Update TOC Metadata usage and clear relevant entries in NewPSOUsage as we are saving this file cache TOC
+							MergePSOUsageToMetaData(NewPSOUsage, TOC.MetaData, true);
 							
                             Header.TableOffset = PSOOffset;
                             TOCOffset = PSOOffset;
@@ -1625,8 +1653,8 @@ public:
                                 PSOs.Add(Meta.Stats.PSOHash, Entry);
                             }
 							
-							// Update TOC Metadata usage masks - don't clear NewPSOUsageMasks as we are using a TempTOC
-							MergeUsageMasksToMetaData(NewPSOUsageMasks, TempTOC.MetaData);
+							// Update TOC Metadata usage masks - don't clear NewPSOUsage as we are using a TempTOC
+							MergePSOUsageToMetaData(NewPSOUsage, TempTOC.MetaData);
                             
                             for (auto& Pair : Stats)
                             {
@@ -1806,8 +1834,8 @@ public:
                                 PSOs.Add(Meta.Stats.PSOHash, Entry);
                             }
 							
-							// Update TOC Metadata usage masks and clear NewPSOUsageMasks as we are file TOC is getting updated
-							MergeUsageMasksToMetaData(NewPSOUsageMasks, TOC.MetaData, true);
+							// Update TOC Metadata usage and clear updated entries in NewPSOUsage as file TOC is getting updated
+							MergePSOUsageToMetaData(NewPSOUsage, TOC.MetaData, true);
 							
 							FPipelineCacheFileFormatTOC TempTOC = TOC;
 							// Update PSO usage stats for new and old
@@ -1987,17 +2015,18 @@ public:
 		return bFileWriteSuccess;
 	}
 	
-	bool IsPSOEntryCached(FPipelineCacheFileFormatPSO const& NewEntry, uint64* OutUsageMask = nullptr) const
+	bool IsPSOEntryCached(FPipelineCacheFileFormatPSO const& NewEntry, FPSOUsageData* EntryData = nullptr) const
 	{
 		uint32 PSOHash = GetTypeHash(NewEntry);
 		FPipelineCacheFileFormatPSOMetaData const * const Existing = TOC.MetaData.Find(PSOHash);
 		
-		if(Existing != nullptr && OutUsageMask != nullptr)
+		if(Existing != nullptr && EntryData != nullptr)
 		{
-			*OutUsageMask = Existing->UsageMask;
+			EntryData->UsageMask = Existing->UsageMask;
+			EntryData->EngineFlags = Existing->EngineFlags;
 		}
 		
-		return  Existing != nullptr;
+		return Existing != nullptr;
 	}
 	
 	bool IsBSSEquivalentPSOEntryCached(FPipelineCacheFileFormatPSO const& NewEntry) const
@@ -2073,7 +2102,8 @@ public:
 		
 		for (auto const& Hash : TOC.MetaData)
 		{
-			if(FPipelineFileCache::MaskComparisonFn(FPipelineFileCache::GameUsageMask, Hash.Value.UsageMask) &&
+			if( (Hash.Value.EngineFlags & FPipelineCacheFlagInvalidPSO) == 0 &&
+				FPipelineFileCache::MaskComparisonFn(FPipelineFileCache::GameUsageMask, Hash.Value.UsageMask) &&
 				Hash.Value.Stats.TotalBindCount >= MinBindCount &&
 				!AlreadyCompiledHashes.Contains(Hash.Key))
 			{
@@ -2112,6 +2142,14 @@ public:
 			FPipelineCacheFileFormatPSORead* Entry = *It;
 			FPipelineCacheFileFormatPSOMetaData const& Meta = TOC.MetaData.FindChecked(Entry->Hash);
 			
+			if((Meta.EngineFlags & FPipelineCacheFlagInvalidPSO) != 0)
+			{
+				// In reality we should not get to this case as GetOrderedPSOHashes() won't pass back PSOs that have this flag set
+				UE_LOG(LogRHI, Verbose, TEXT("Encountered a PSO entry %u marked invalid - ignoring"), Entry->Hash);
+				Entry->bValid = false;
+				continue;
+			}
+			
 			if (Meta.FileGuid == GameFileGuid)
 			{
                 FPipelineCacheFileFormatPSOMetaData const* GameMeta = GameTOC.MetaData.Find(Entry->Hash);
@@ -2313,7 +2351,7 @@ bool FPipelineFileCache::SavePipelineFileCache(FString const& Name, SaveMode Mod
 		{
 			FName PlatformName = FileCache->GetPlatformName();
 			FString Path = FPaths::ProjectSavedDir() / FString::Printf(TEXT("%s_%s.upipelinecache"), *Name, *PlatformName.ToString());
-			bOk = FileCache->SavePipelineFileCache(Path, Mode, Stats, NewPSOs, RequestedOrder, NewPSOUsageMasks);
+			bOk = FileCache->SavePipelineFileCache(Path, Mode, Stats, NewPSOs, RequestedOrder, NewPSOUsage);
 			
 			// If successful clear new PSO's as they should have been saved out
 			// Leave everything else in-tact (e.g stats) for subsequent in place save operations
@@ -2359,7 +2397,7 @@ void FPipelineFileCache::ClosePipelineFileCache()
 			
 			// Clear Runtime hashes otherwise we can't start adding newPSO's for a newly opened file
 			RunTimeToPSOUsage.Empty();
-			NewPSOUsageMasks.Empty();
+			NewPSOUsage.Empty();
 			NewPSOs.Empty();
             NumNewPSOs = 0;
 			
@@ -2369,12 +2407,12 @@ void FPipelineFileCache::ClosePipelineFileCache()
 	}
 }
 
-void FPipelineFileCache::EnsurePSOUsageMask(uint32 PSOHash, uint64 UsageMask)
+void FPipelineFileCache::RegisterPSOUsageDataUpdateForNextSave(FPSOUsageData& UsageData)
 {
-	uint64& Entry = NewPSOUsageMasks.FindOrAdd(PSOHash);
-	Entry |= UsageMask;
-	
-	UE_LOG(LogRHI, Warning, TEXT("Added UsageMask %llu for PSO: %u"), UsageMask, PSOHash);
+	FPSOUsageData& CurrentEntry = NewPSOUsage.FindOrAdd(UsageData.PSOHash);	
+	CurrentEntry.PSOHash = UsageData.PSOHash;
+	CurrentEntry.UsageMask |= UsageData.UsageMask;
+	CurrentEntry.EngineFlags |= UsageData.EngineFlags;
 }
 
 void FPipelineFileCache::CacheGraphicsPSO(uint32 RunTimeHash, FGraphicsPipelineStateInitializer const& Initializer)
@@ -2398,8 +2436,9 @@ void FPipelineFileCache::CacheGraphicsPSO(uint32 RunTimeHash, FGraphicsPipelineS
 					check(bOK);
 					
 					uint32 PSOHash = GetTypeHash(NewEntry);
-					uint64 ExistingPSOUsageMask = 0;
-					if (!FileCache->IsPSOEntryCached(NewEntry, &ExistingPSOUsageMask))
+					FPSOUsageData CurrentUsageData(PSOHash, 0, 0);
+					
+					if (!FileCache->IsPSOEntryCached(NewEntry, &CurrentUsageData))
 					{
 						bool bActuallyNewPSO = true;
 						if (IsOpenGLPlatform(GMaxRHIShaderPlatform)) // OpenGL is a BSS platform and so we don't report BSS matches as missing.
@@ -2431,19 +2470,20 @@ void FPipelineFileCache::CacheGraphicsPSO(uint32 RunTimeHash, FGraphicsPipelineS
 						}
 					}
 					
-					// Apply the existing file PSO Usage mask and current to our "fast" runtime check
-					RunTimeToPSOUsage.Add(RunTimeHash, FPSOUsageData(PSOHash, FPipelineFileCache::GameUsageMask | ExistingPSOUsageMask));
-					
 					// Only set if the file cache doesn't have this Mask for the PSO - avoid making more entries and unnessary file saves
-					if(!IsReferenceMaskSet(FPipelineFileCache::GameUsageMask, ExistingPSOUsageMask))
+					if(!IsReferenceMaskSet(FPipelineFileCache::GameUsageMask, CurrentUsageData.UsageMask))
 					{
-						EnsurePSOUsageMask(PSOHash, FPipelineFileCache::GameUsageMask);
+						CurrentUsageData.UsageMask |= FPipelineFileCache::GameUsageMask;
+						RegisterPSOUsageDataUpdateForNextSave(CurrentUsageData);
 					}
+					
+					// Apply the existing file PSO Usage mask and current to our "fast" runtime check
+					RunTimeToPSOUsage.Add(RunTimeHash, CurrentUsageData);
 				}
 				else if(!IsReferenceMaskSet(FPipelineFileCache::GameUsageMask, PSOUsage->UsageMask))
 				{
 					PSOUsage->UsageMask |= FPipelineFileCache::GameUsageMask;
-					EnsurePSOUsageMask(PSOUsage->PSOHash, FPipelineFileCache::GameUsageMask);
+					RegisterPSOUsageDataUpdateForNextSave(*PSOUsage);
 				}
 			}
 		}
@@ -2471,10 +2511,9 @@ void FPipelineFileCache::CacheComputePSO(uint32 RunTimeHash, FRHIComputeShader c
 					check(bOK);
 					
 					uint32 PSOHash = GetTypeHash(NewEntry);
-					RunTimeToPSOUsage.Add(RunTimeHash, FPSOUsageData(PSOHash, FPipelineFileCache::GameUsageMask));
+					FPSOUsageData CurrentUsageData(PSOHash, 0, 0);
 					
-					uint64 ExistingPSOUsageMask = 0;
-					if (!FileCache->IsPSOEntryCached(NewEntry, &ExistingPSOUsageMask))
+					if (!FileCache->IsPSOEntryCached(NewEntry, &CurrentUsageData))
 					{
 						CSV_EVENT(PSO, TEXT("Encountered new compute PSO"));
 						UE_LOG(LogRHI, Warning, TEXT("Encountered a new compute PSO: %u"), PSOHash);
@@ -2499,19 +2538,66 @@ void FPipelineFileCache::CacheComputePSO(uint32 RunTimeHash, FRHIComputeShader c
 						}
 					}
 					
-					// Apply the existing file PSO Usage mask and current to our "fast" runtime check
-					RunTimeToPSOUsage.Add(RunTimeHash, FPSOUsageData(PSOHash, FPipelineFileCache::GameUsageMask | ExistingPSOUsageMask));
-					
 					// Only set if the file cache doesn't have this Mask for the PSO - avoid making more entries and unnessary file saves
-					if(!IsReferenceMaskSet(FPipelineFileCache::GameUsageMask, ExistingPSOUsageMask))
+					if(!IsReferenceMaskSet(FPipelineFileCache::GameUsageMask, CurrentUsageData.UsageMask))
 					{
-						EnsurePSOUsageMask(PSOHash, FPipelineFileCache::GameUsageMask);
+						CurrentUsageData.UsageMask |= FPipelineFileCache::GameUsageMask;
+						RegisterPSOUsageDataUpdateForNextSave(CurrentUsageData);
 					}
+					
+					// Apply the existing file PSO Usage mask and current to our "fast" runtime check
+					RunTimeToPSOUsage.Add(RunTimeHash, CurrentUsageData);
 				}
 				else if(!IsReferenceMaskSet(FPipelineFileCache::GameUsageMask, PSOUsage->UsageMask))
 				{
 					PSOUsage->UsageMask |= FPipelineFileCache::GameUsageMask;
-					EnsurePSOUsageMask(PSOUsage->PSOHash, FPipelineFileCache::GameUsageMask);
+					RegisterPSOUsageDataUpdateForNextSave(*PSOUsage);
+				}
+			}
+		}
+	}
+}
+
+void FPipelineFileCache::RegisterPSOCompileFailure(uint32 RunTimeHash, FGraphicsPipelineStateInitializer const& Initializer)
+{
+	if(IsPipelineFileCacheEnabled() && (LogPSOtoFileCache() || ReportNewPSOs()) && Initializer.bFromPSOFileCache)
+	{
+		FRWScopeLock Lock(FileCacheLock, SLT_ReadOnly);
+		
+		if(FileCache)
+		{
+			FPSOUsageData* PSOUsage = RunTimeToPSOUsage.Find(RunTimeHash);
+			if(PSOUsage == nullptr || !IsReferenceMaskSet(FPipelineCacheFlagInvalidPSO, PSOUsage->EngineFlags))
+			{
+				Lock.ReleaseReadOnlyLockAndAcquireWriteLock_USE_WITH_CAUTION();
+				PSOUsage = RunTimeToPSOUsage.Find(RunTimeHash);
+				
+				if(PSOUsage == nullptr)
+				{
+					FPipelineCacheFileFormatPSO ShouldBeExistingEntry;
+					bool bOK = FPipelineCacheFileFormatPSO::Init(ShouldBeExistingEntry, Initializer);
+					check(bOK);
+					
+					uint32 PSOHash = GetTypeHash(ShouldBeExistingEntry);
+					FPSOUsageData CurrentUsageData(PSOHash, 0, 0);
+					
+					bool bCached = FileCache->IsPSOEntryCached(ShouldBeExistingEntry, &CurrentUsageData);
+					check(bCached);	//bFromPSOFileCache was set but not in the cache something has gone wrong
+					{
+						CurrentUsageData.EngineFlags |= FPipelineCacheFlagInvalidPSO;
+						
+						RegisterPSOUsageDataUpdateForNextSave(CurrentUsageData);
+						RunTimeToPSOUsage.Add(RunTimeHash, CurrentUsageData);
+						
+						UE_LOG(LogRHI, Warning, TEXT("Graphics PSO (%u) compile failure registering to File Cache"), PSOHash);
+					}
+				}
+				else if(!IsReferenceMaskSet(FPipelineCacheFlagInvalidPSO, PSOUsage->EngineFlags))
+				{
+					PSOUsage->EngineFlags |= FPipelineCacheFlagInvalidPSO;
+					RegisterPSOUsageDataUpdateForNextSave(*PSOUsage);
+					
+					UE_LOG(LogRHI, Warning, TEXT("Graphics PSO (%u) compile failure registering to File Cache"), PSOUsage->PSOHash);
 				}
 			}
 		}
@@ -2601,7 +2687,9 @@ struct FPipelineCacheFileData
 				{
 					for (auto& Entry : Data.TOC.MetaData)
 					{
-                        if (Entry.Value.FileGuid == Data.Header.Guid && Entry.Value.FileSize > sizeof(FPipelineCacheFileFormatPSO::DescriptorType))
+						if ( (Entry.Value.EngineFlags & FPipelineCacheFlagInvalidPSO) == 0 &&
+                        	 Entry.Value.FileGuid == Data.Header.Guid &&
+                        	 Entry.Value.FileSize > sizeof(FPipelineCacheFileFormatPSO::DescriptorType))
                         {
                             FPipelineCacheFileFormatPSO PSO;
                             FileAReader->Seek(Entry.Value.FileOffset);
@@ -2706,16 +2794,16 @@ uint32 FPipelineFileCache::NumPSOsLogged()
 	uint32 Result = 0;
 	if(IsPipelineFileCacheEnabled() && LogPSOtoFileCache())
 	{
-		// Only count PSOs that are both new and have at least one bind, otherwise we can ignore them.
+		// Only count PSOs that are both new and have at least one bind or have been marked invalid (compile failure) otherwise we can ignore them
 		FRWScopeLock Lock(FileCacheLock, SLT_ReadOnly);
 		
 		// We now need to know if the number of usage masks changes - this number should be as least the same as before but could be conceptually more if an existing PSO has an extra usage mask applied
-		if(NewPSOUsageMasks.Num() > 0)
+		if(NewPSOUsage.Num() > 0)
 		{
-			for(auto& MaskEntry : NewPSOUsageMasks)
+			for(auto& MaskEntry : NewPSOUsage)
 			{
 				FPipelineStateStats const* Stat = Stats.FindRef(MaskEntry.Key);
-				if (Stat && Stat->TotalBindCount > 0)
+				if (Stat && Stat->TotalBindCount > 0 || (MaskEntry.Value.EngineFlags & FPipelineCacheFlagInvalidPSO) != 0)
 				{
 					Result++;
 				}
@@ -2887,15 +2975,28 @@ bool FPipelineFileCache::MergePipelineFileCaches(FString const& PathA, FString c
 		uint32 MergeCount = 0;
 		for (auto const& Entry : A.TOC.MetaData)
 		{
+			// Don't merge PSOs that have the invalid bit set
+			if((Entry.Value.EngineFlags & FPipelineCacheFlagInvalidPSO) != 0)
+			{
+				continue;
+			}
+
 			Output.TOC.MetaData.Add(Entry.Key, Entry.Value);
 		}
 		for (auto const& Entry : B.TOC.MetaData)
 		{
+			// Don't merge PSOs that have the invalid bit set
+			if((Entry.Value.EngineFlags & FPipelineCacheFlagInvalidPSO) != 0)
+			{
+				continue;
+			}
+
 			// Make sure these usage masks for the same PSOHash find their way in
 			auto* ExistingMetaEntry = Output.TOC.MetaData.Find(Entry.Key);
 			if(ExistingMetaEntry != nullptr)
 			{
 				ExistingMetaEntry->UsageMask |=  Entry.Value.UsageMask;
+				ExistingMetaEntry->EngineFlags |= Entry.Value.EngineFlags;
 				++MergeCount;
 			}
 			else
diff --git a/Engine/Source/Runtime/RHI/Private/PipelineStateCache.cpp b/Engine/Source/Runtime/RHI/Private/PipelineStateCache.cpp
--- a/Engine/Source/Runtime/RHI/Private/PipelineStateCache.cpp
+++ b/Engine/Source/Runtime/RHI/Private/PipelineStateCache.cpp
@@ -131,7 +131,12 @@ static void HandlePipelineCreationFailure(const FGraphicsPipelineStateInitialize
 	UE_LOG(LogRHI, Error, TEXT("0x%x"), Init.DepthStencilTargetFormat);
 #endif
 	
-	if(!Init.bFromPSOFileCache)
+	if(Init.bFromPSOFileCache)
+	{
+		// Let the cache know so it hopefully won't give out this one again
+		FPipelineFileCache::RegisterPSOCompileFailure(GetTypeHash(Init), Init);
+	}
+	else
 	{
 		UE_LOG(LogRHI, Fatal, TEXT("Shader compilation failures are Fatal."));
 	}
diff --git a/Engine/Source/Runtime/RHI/Public/PipelineFileCache.h b/Engine/Source/Runtime/RHI/Public/PipelineFileCache.h
--- a/Engine/Source/Runtime/RHI/Public/PipelineFileCache.h
+++ b/Engine/Source/Runtime/RHI/Public/PipelineFileCache.h
@@ -240,9 +240,11 @@ typedef bool(*FPSOMaskComparisonFn)(uint64 ReferenceMask, uint64 PSOMask);
 
 struct FPSOUsageData
 {
-	FPSOUsageData(uint32 InPSOHash, uint64 InUsageMask): PSOHash(InPSOHash), UsageMask(InUsageMask) {}
+	FPSOUsageData(): PSOHash(0), UsageMask(0), EngineFlags(0) {}
+	FPSOUsageData(uint32 InPSOHash, uint64 InUsageMask, uint16 InEngineFlags): PSOHash(InPSOHash), UsageMask(InUsageMask), EngineFlags(InEngineFlags) {}
 	uint32 PSOHash;
 	uint64 UsageMask;
+	uint16 EngineFlags;
 };
 
 class RHI_API FPipelineFileCache
@@ -281,6 +283,12 @@ public:
 	static void CacheComputePSO(uint32 RunTimeHash, FRHIComputeShader const* Initializer);
 	static FPipelineStateStats* RegisterPSOStats(uint32 RunTimeHash);
 	
+	/*
+	 * This PSO has failed compile and is invalid - this cache should not return this invalid PSO from subsequent calls for PreCompile requests.
+	 * Note: Not implementated for Compute that has no flag to say it came from this cache - don't want to consume failures that didn't propagate from this cache.
+	 */
+	static void RegisterPSOCompileFailure(uint32 RunTimeHash, FGraphicsPipelineStateInitializer const& Initializer);
+	
 	/**
 	 * Event signature for being notified that a new PSO has been logged
 	 */
@@ -310,13 +318,13 @@ public:
 	
 private:
 	
-	static void EnsurePSOUsageMask(uint32 PSOHash, uint64 UsageMask);
+	static void RegisterPSOUsageDataUpdateForNextSave(FPSOUsageData& UsageData);
 	
 private:
 	static FRWLock FileCacheLock;
 	static class FPipelineCacheFile* FileCache;
-	static TMap<uint32, FPSOUsageData> RunTimeToPSOUsage;
-	static TMap<uint32, uint64> NewPSOUsageMasks;			// Note: this is not only for new PSO's but also for existing PSO's in the filecache when they have a new UsageMask
+	static TMap<uint32, FPSOUsageData> RunTimeToPSOUsage;		// Fast check structure - Not saved (External state cache runtime hash to seen usage data)
+	static TMap<uint32, FPSOUsageData> NewPSOUsage;				// For mask or engine updates - Merged + Saved (Our internal PSO hash to latest usage data) - temp working scratch, only holds updates since last "save" so is not the authority on state
 	static TMap<uint32, FPipelineStateStats*> Stats;
 	static TSet<FPipelineCacheFileFormatPSO> NewPSOs;
     static uint32 NumNewPSOs;
