Fix build

diff --git a/Engine/Source/Runtime/RHI/Private/PipelineFileCache.cpp b/Engine/Source/Runtime/RHI/Private/PipelineFileCache.cpp
--- a/Engine/Source/Runtime/RHI/Private/PipelineFileCache.cpp
+++ b/Engine/Source/Runtime/RHI/Private/PipelineFileCache.cpp
@@ -2803,7 +2803,7 @@ uint32 FPipelineFileCache::NumPSOsLogged()
 			for(auto& MaskEntry : NewPSOUsage)
 			{
 				FPipelineStateStats const* Stat = Stats.FindRef(MaskEntry.Key);
-				if (Stat && Stat->TotalBindCount > 0 || (MaskEntry.Value.EngineFlags & FPipelineCacheFlagInvalidPSO) != 0)
+				if ((Stat && Stat->TotalBindCount > 0) || (MaskEntry.Value.EngineFlags & FPipelineCacheFlagInvalidPSO) != 0)
 				{
 					Result++;
 				}
