UE4.22 - Start deprecating SetRenderTargets API
#rnx

diff --git a/Engine/Source/Runtime/RHI/Public/RHIUtilities.h b/Engine/Source/Runtime/RHI/Public/RHIUtilities.h
--- a/Engine/Source/Runtime/RHI/Public/RHIUtilities.h
+++ b/Engine/Source/Runtime/RHI/Public/RHIUtilities.h
@@ -451,6 +451,7 @@ inline void SetRenderTarget(FRHICommandList& RHICmdList, FTextureRHIParamRef New
 }
 
 /** Helper for the common case of using a single color and depth render target, with a mip index for the color target. */
+UE_DEPRECATED(4.22, "SetRenderTarget API is deprecated; please use RHIBegin/EndRenderPass instead.")
 inline void SetRenderTarget(FRHICommandList& RHICmdList, FTextureRHIParamRef NewRenderTarget, int32 MipIndex, FTextureRHIParamRef NewDepthStencilTarget, bool bWritableBarrier = false)
 {
 	FRHIRenderTargetView RTV(NewRenderTarget, ERenderTargetLoadAction::ELoad, MipIndex, -1);
@@ -465,6 +466,7 @@ inline void SetRenderTarget(FRHICommandList& RHICmdList, FTextureRHIParamRef New
 }
 
 /** Helper for the common case of using a single color and depth render target, with a mip index for the color target. */
+UE_DEPRECATED(4.22, "SetRenderTarget API is deprecated; please use RHIBegin/EndRenderPass instead.")
 inline void SetRenderTarget(FRHICommandList& RHICmdList, FTextureRHIParamRef NewRenderTarget, int32 MipIndex, int32 ArraySliceIndex, FTextureRHIParamRef NewDepthStencilTarget, bool bWritableBarrier = false)
 {
 	FRHIRenderTargetView RTV(NewRenderTarget, ERenderTargetLoadAction::ELoad, MipIndex, ArraySliceIndex);
@@ -479,6 +481,7 @@ inline void SetRenderTarget(FRHICommandList& RHICmdList, FTextureRHIParamRef New
 }
 
 /** Helper that converts FTextureRHIParamRef's into FRHIRenderTargetView's. */
+UE_DEPRECATED(4.22, "SetRenderTargets API is deprecated; please use RHIBegin/EndRenderPass instead.")
 inline void SetRenderTargets(
 	FRHICommandList& RHICmdList,
 	uint32 NewNumSimultaneousRenderTargets, 
