Move default graphics RHI to windows target settings

diff --git a/Engine/Source/Developer/Windows/WindowsTargetPlatform/Classes/WindowsTargetSettings.h b/Engine/Source/Developer/Windows/WindowsTargetPlatform/Classes/WindowsTargetSettings.h
--- a/Engine/Source/Developer/Windows/WindowsTargetPlatform/Classes/WindowsTargetSettings.h
+++ b/Engine/Source/Developer/Windows/WindowsTargetPlatform/Classes/WindowsTargetSettings.h
@@ -24,6 +24,17 @@ enum class ECompilerVersion : uint8
 	VisualStudio2017 = 2 UMETA(DisplayName = "Visual Studio 2017"),
 	VisualStudio2019 = 3 UMETA(DisplayName = "Visual Studio 2019"),
 };
+UENUM()
+enum class EDefaultGraphicsRHI : uint8
+{
+	DefaultGraphicsRHI_Default = 0 UMETA(DisplayName = "Default"),
+	DefaultGraphicsRHI_DX11 = 1 UMETA(DisplayName = "DirectX 11"),
+	DefaultGraphicsRHI_DX12 = 2 UMETA(DisplayName = "DirectX 12"),
+	DefaultGraphicsRHI_Vulkan = 3 UMETA(DisplayName = "Vulkan"),
+	DefaultGraphicsRHI_OpenGL = 4 UMETA(DisplayName = "OpenGL"),
+};
+
+
 
 /**
  * Implements the settings for the Windows target platform. The first instance of this class is initialized in
@@ -49,6 +60,13 @@ public:
 	UPROPERTY(EditAnywhere, config, Category=Rendering)
 	TArray<FString> TargetedRHIs;
 
+	/** 
+	 * Default Graphics RHI. Select which RHIto use. Make sure its also selected as a Targeted RHI
+	 * Requires Editor restart
+	 */
+	UPROPERTY(EditAnywhere, config, Category="Targeted RHIs", Meta = (DisplayName = "Default RHI"))
+	EDefaultGraphicsRHI DefaultGraphicsRHI;
+
 	/**
 	 * Determine the minimum supported 
 	 */
diff --git a/Engine/Source/Editor/LevelEditor/Private/LevelEditor.cpp b/Engine/Source/Editor/LevelEditor/Private/LevelEditor.cpp
--- a/Engine/Source/Editor/LevelEditor/Private/LevelEditor.cpp
+++ b/Engine/Source/Editor/LevelEditor/Private/LevelEditor.cpp
@@ -108,9 +108,9 @@ public:
 		TooltipArgs.Add(TEXT("Branch"), FText::FromString(FEngineVersion::Current().GetBranch()));
 		TooltipArgs.Add(TEXT("BuildConfiguration"), EBuildConfigurations::ToText(BuildConfig));
 		TooltipArgs.Add(TEXT("BuildDate"), FText::FromString(FApp::GetBuildDate()));
-		TooltipArgs.Add(TEXT("GraphicsAPI"), FText::FromString(FApp::GetGraphicsAPI()));
+		TooltipArgs.Add(TEXT("GraphicsRHI"), FText::FromString(FApp::GetGraphicsRHI()));
 
-		RightContentTooltip = FText::Format(NSLOCTEXT("UnrealEditor", "TitleBarRightContentTooltip", "Version: {Version}\nBranch: {Branch}\nBuild Configuration: {BuildConfiguration}\nBuild Date: {BuildDate}\nGraphics API: {GraphicsAPI}"), TooltipArgs);
+		RightContentTooltip = FText::Format(NSLOCTEXT("UnrealEditor", "TitleBarRightContentTooltip", "Version: {Version}\nBranch: {Branch}\nBuild Configuration: {BuildConfiguration}\nBuild Date: {BuildDate}\nGraphics RHI: {GraphicsRHI}"), TooltipArgs);
 
 		SetToolTipText(RightContentTooltip);
 
diff --git a/Engine/Source/Runtime/Core/Private/HAL/ConsoleManager.cpp b/Engine/Source/Runtime/Core/Private/HAL/ConsoleManager.cpp
--- a/Engine/Source/Runtime/Core/Private/HAL/ConsoleManager.cpp
+++ b/Engine/Source/Runtime/Core/Private/HAL/ConsoleManager.cpp
@@ -2616,14 +2616,3 @@ static TAutoConsoleVariable<int32> CVarSkyLightingQuality(
 	TEXT("  1: on\n"),
 	ECVF_Scalability | ECVF_RenderThreadSafe
 );
-
-static TAutoConsoleVariable<int32> CVarDefaultGraphicsApi(
-	TEXT("r.DefaultGraphicsApi"),
-	0,
-	TEXT("0: Default\n")
-	TEXT("1: DirectX 11\n")
-	TEXT("2: DirectX 12\n")
-	TEXT("3: Vulkan\n")
-	TEXT("4: OpenGL\n"),
-	ECVF_ReadOnly
-);
diff --git a/Engine/Source/Runtime/Core/Private/Misc/App.cpp b/Engine/Source/Runtime/Core/Private/Misc/App.cpp
--- a/Engine/Source/Runtime/Core/Private/Misc/App.cpp
+++ b/Engine/Source/Runtime/Core/Private/Misc/App.cpp
@@ -24,7 +24,7 @@ FGuid FApp::InstanceId = FGuid::NewGuid();
 FGuid FApp::SessionId = FGuid::NewGuid();
 FString FApp::SessionName = FString();
 FString FApp::SessionOwner = FString();
-FString FApp::GraphicsAPI = FString();
+FString FApp::GraphicsRHI = FString();
 TArray<FString> FApp::SessionUsers = TArray<FString>();
 bool FApp::Standalone = true;
 bool FApp::bIsBenchmarking = false;
@@ -99,14 +99,14 @@ FString FApp::GetBuildDate()
 	return FString(BuildSettings::GetBuildDate());
 }
 
-FString FApp::GetGraphicsAPI()
+FString FApp::GetGraphicsRHI()
 {
-	return GraphicsAPI;
+	return GraphicsRHI;
 }
 
-void FApp::SetGraphicsAPI(FString APIString)
+void FApp::SetGraphicsRHI(FString RHIString)
 {
-	GraphicsAPI = APIString;
+	GraphicsRHI = RHIString;
 }
 
 void FApp::InitializeSession()
diff --git a/Engine/Source/Runtime/Core/Public/Misc/App.h b/Engine/Source/Runtime/Core/Public/Misc/App.h
--- a/Engine/Source/Runtime/Core/Public/Misc/App.h
+++ b/Engine/Source/Runtime/Core/Public/Misc/App.h
@@ -65,18 +65,16 @@ public:
 	static FString GetBuildDate();
 
 	/**
-	 * Gets the name of the graphics API currently in use.
+	 * Gets the name of the graphics RHI currently in use.
 	 *
-	 * @return name of Graphics API
+	 * @return name of Graphics RHI
 	 */
-	static FString GetGraphicsAPI();
+	static FString GetGraphicsRHI();
 
 	/**
-	 * Sets the Graphics Api currently in use
-	 *
-	 * @return name of Graphics API
+	 * Sets the Graphics RHI currently in use
 	 */
-	static void SetGraphicsAPI(FString APIString);
+	static void SetGraphicsRHI(FString RHIString);
 
 
 	/**
@@ -725,8 +723,8 @@ private:
 	/** Holds the name of the user that launched session. */
 	static FString SessionOwner;
 
-	/** Holds the name the graphics api currently in use*/
-	static FString GraphicsAPI;
+	/** Holds the name the graphics RHI currently in use*/
+	static FString GraphicsRHI;
 
 
 	/** List of authorized session users. */
diff --git a/Engine/Source/Runtime/Engine/Classes/Engine/RendererSettings.h b/Engine/Source/Runtime/Engine/Classes/Engine/RendererSettings.h
--- a/Engine/Source/Runtime/Engine/Classes/Engine/RendererSettings.h
+++ b/Engine/Source/Runtime/Engine/Classes/Engine/RendererSettings.h
@@ -44,23 +44,6 @@ namespace ECompositingSampleCount
 }
 
 
-/**
- * Enumerates different graphics apis
- */
-UENUM()
-namespace EDefaultGraphicsApi
-{
-	enum Type
-	{
-		DefaultGraphicsApi_Default = 0 UMETA(DisplayName = "Default"),
-		DefaultGraphicsApi_DX11 = 1 UMETA(DisplayName = "DirectX 11"),
-		DefaultGraphicsApi_DX12 = 2 UMETA(DisplayName = "DirectX 12"),
-		DefaultGraphicsApi_Vulkan = 3 UMETA(DisplayName = "Vulkan"),
-		DefaultGraphicsApi_OpenGL = 4 UMETA(DisplayName = "OpenGL"),
-	};
-}
-
-
 /**
 * Enumerates available mobile MSAA sample counts.
 */
@@ -706,11 +689,6 @@ class ENGINE_API URendererSettings : public UDeveloperSettings
 		ConfigRestartRequired = true))
 		uint32 bSupportMaterialLayers : 1;
 
-	UPROPERTY(config, EditAnywhere, Category = API, meta = (
-		ConsoleVariable = "r.DefaultGraphicsApi", DisplayName = "Default Graphics Api",
-		ToolTip = "Graphics API to use by default"))
-		TEnumAsByte<EDefaultGraphicsApi::Type> DefaultGraphicsApi;
-
 public:
 
 	//~ Begin UObject Interface
diff --git a/Engine/Source/Runtime/RHI/Private/Windows/WindowsDynamicRHI.cpp b/Engine/Source/Runtime/RHI/Private/Windows/WindowsDynamicRHI.cpp
--- a/Engine/Source/Runtime/RHI/Private/Windows/WindowsDynamicRHI.cpp
+++ b/Engine/Source/Runtime/RHI/Private/Windows/WindowsDynamicRHI.cpp
@@ -49,18 +49,37 @@ static IDynamicRHIModule* LoadDynamicRHIModule(ERHIFeatureLevel::Type& DesiredFe
 	bool bForceD3D11 = FParse::Param(FCommandLine::Get(), TEXT("d3d11")) || FParse::Param(FCommandLine::Get(), TEXT("dx11")) || (bForceSM5 && !bForceVulkan && !bForceOpenGL);
 	bool bForceD3D12 = FParse::Param(FCommandLine::Get(), TEXT("d3d12")) || FParse::Param(FCommandLine::Get(), TEXT("dx12"));
 	DesiredFeatureLevel = ERHIFeatureLevel::Num;
-	//Default graphics api is only used if no command line option is specified
+	
 	if(!(bForceVulkan||bForceOpenGL||bForceD3D10||bForceD3D11||bForceD3D12))
 	{
-		static const auto CVarDefaultGraphicsApi = IConsoleManager::Get().FindTConsoleVariableDataInt(TEXT("r.DefaultGraphicsApi"));
-		int32 DefaultGraphicsApi = CVarDefaultGraphicsApi->GetValueOnAnyThread();
-		switch(DefaultGraphicsApi)
+		//Default graphics RHI is only used if no command line option is specified
+		FConfigFile EngineSettings;
+		FString PlatformNameString = FPlatformProperties::PlatformName();
+		const TCHAR* PlatformName = *PlatformNameString;
+		FConfigCacheIni::LoadLocalIniFile(EngineSettings, TEXT("Engine"), true, PlatformName);
+		FString DefaultGraphicsRHI;
+		if(EngineSettings.GetString(TEXT("/Script/WindowsTargetPlatform.WindowsTargetSettings"), TEXT("DefaultGraphicsRHI"), DefaultGraphicsRHI))
 		{
-		case 0:	break;
-		case 1:		bForceD3D11 = true; break;
-		case 2:		bForceD3D12 = true; break;
-		case 3:		bForceVulkan = true; break;
-		case 4:		bForceOpenGL= true; break;
+			static FString NAME_DX11(TEXT("DefaultGraphicsRHI_DX11"));
+			static FString NAME_DX12(TEXT("DefaultGraphicsRHI_DX12"));
+			static FString NAME_VULKAN(TEXT("DefaultGraphicsRHI_Vulkan"));
+			static FString NAME_OPENGL(TEXT("DefaultGraphicsRHI_OpenGL"));
+			if(DefaultGraphicsRHI == NAME_DX11)
+			{
+				bForceD3D11 = true;
+			}
+			else if (DefaultGraphicsRHI == NAME_DX12)
+			{
+				bForceD3D12 = true;
+			}
+			else if (DefaultGraphicsRHI == NAME_VULKAN)
+			{
+				bForceVulkan = true;
+			}
+			else if (DefaultGraphicsRHI == NAME_OPENGL)
+			{
+				bForceOpenGL = true;
+			}
 		}
 	}
 
@@ -127,7 +146,7 @@ static IDynamicRHIModule* LoadDynamicRHIModule(ERHIFeatureLevel::Type& DesiredFe
 			DynamicRHIModule = NULL;
 		}
 		LoadedRHIModuleName = SwitchRHIModuleName;
-		FApp::SetGraphicsAPI(TEXT("Switch"));
+		FApp::SetGraphicsRHI(TEXT("Switch"));
 	}
 	else
 #endif
@@ -144,7 +163,7 @@ static IDynamicRHIModule* LoadDynamicRHIModule(ERHIFeatureLevel::Type& DesiredFe
 			DynamicRHIModule = NULL;
 		}
 		LoadedRHIModuleName = OpenGLRHIModuleName;
-		FApp::SetGraphicsAPI(TEXT("OpenGL"));
+		FApp::SetGraphicsRHI(TEXT("OpenGL"));
 	}
 	else if (bForceVulkan)
 	{
@@ -157,7 +176,7 @@ static IDynamicRHIModule* LoadDynamicRHIModule(ERHIFeatureLevel::Type& DesiredFe
 			DynamicRHIModule = NULL;
 		}
 		LoadedRHIModuleName = VulkanRHIModuleName;
-		FApp::SetGraphicsAPI(TEXT("Vulkan"));
+		FApp::SetGraphicsRHI(TEXT("Vulkan"));
 	}
 	else if (bForceD3D12 || bPreferD3D12)
 	{
@@ -182,7 +201,7 @@ static IDynamicRHIModule* LoadDynamicRHIModule(ERHIFeatureLevel::Type& DesiredFe
 		{
 			FMessageDialog::Open(EAppMsgType::Ok, NSLOCTEXT("WindowsDynamicRHI", "UseExpressionEncoder", "Fraps has been known to crash D3D12. Please use Microsoft Expression Encoder instead for capturing."));
 		}
-		FApp::SetGraphicsAPI(TEXT("DirectX 12"));
+		FApp::SetGraphicsRHI(TEXT("DirectX 12"));
 	}
 
 	// Fallback to D3D11RHI if nothing is selected
@@ -202,7 +221,7 @@ static IDynamicRHIModule* LoadDynamicRHIModule(ERHIFeatureLevel::Type& DesiredFe
 			FMessageDialog::Open(EAppMsgType::Ok, NSLOCTEXT("WindowsDynamicRHI", "UseExpressionEncoderDX11", "Fraps has been known to crash D3D11. Please use Microsoft Expression Encoder instead for capturing."));
 		}
 		LoadedRHIModuleName = D3D11RHIModuleName;
-		FApp::SetGraphicsAPI(TEXT("DirectX 11"));
+		FApp::SetGraphicsRHI(TEXT("DirectX 11"));
 	}
 	return DynamicRHIModule;
 }
