Fix RHISetGPUCaptureOptions usage text.

diff --git a/Engine/Source/Runtime/RHI/Private/DynamicRHI.cpp b/Engine/Source/Runtime/RHI/Private/DynamicRHI.cpp
--- a/Engine/Source/Runtime/RHI/Private/DynamicRHI.cpp
+++ b/Engine/Source/Runtime/RHI/Private/DynamicRHI.cpp
@@ -265,7 +265,7 @@ static void BaseRHISetGPUCaptureOptions(const TArray<FString>& Args, UWorld* Wor
 	}
 	else
 	{
-		UE_LOG(LogRHI, Display, TEXT("Usage: r.PS4.EnableCaptureMode 0 or r.PS4.EnableCaptureMode 1"));
+		UE_LOG(LogRHI, Display, TEXT("Usage: r.RHISetGPUCaptureOptions 0 or r.RHISetGPUCaptureOptions 1"));
 	}
 }
 
