Changed the thread priority for FCompilePipelineStateTask on Mac to AnyBackgroundThreadNormalTask and reverted the default batch size for the fast mode to 50. Thanks to this we don't starve the game thread as much as before without sacrificing the load times.

diff --git a/Engine/Source/Runtime/RHI/Private/PipelineStateCache.cpp b/Engine/Source/Runtime/RHI/Private/PipelineStateCache.cpp
--- a/Engine/Source/Runtime/RHI/Private/PipelineStateCache.cpp
+++ b/Engine/Source/Runtime/RHI/Private/PipelineStateCache.cpp
@@ -582,11 +582,7 @@ FGraphicsPipelineCache GGraphicsPipelineCache;
 FAutoConsoleTaskPriority CPrio_FCompilePipelineStateTask(
 	TEXT("TaskGraph.TaskPriorities.CompilePipelineStateTask"),
 	TEXT("Task and thread priority for FCompilePipelineStateTask."),
-#if PLATFORM_MAC
-	ENamedThreads::NormalThreadPriority,	// On Mac, use normal thread priority to not freeze the loading screen for extended period of time
-#else
 	ENamedThreads::HighThreadPriority,		// if we have high priority task threads, then use them...
-#endif
 	ENamedThreads::NormalTaskPriority,		// .. at normal task priority
 	ENamedThreads::HighTaskPriority,		// if we don't have hi pri threads, then use normal priority threads at high task priority instead
 	EPowerSavingEligibility::NotEligible	// Not eligible for downgrade when power saving is requested.
@@ -741,7 +737,9 @@ public:
 
 	ENamedThreads::Type GetDesiredThread()
 	{
-		return CPrio_FCompilePipelineStateTask.Get();
+		// On Mac the compilation is handled using external processes, so engine threads have very little work to do
+		// and it's better to leave more CPU time to these extrenal processes and other engine threads.
+		return PLATFORM_MAC ? ENamedThreads::AnyBackgroundThreadNormalTask : CPrio_FCompilePipelineStateTask.Get();
 	}
 };
 
diff --git a/Engine/Source/Runtime/RenderCore/Private/ShaderPipelineCache.cpp b/Engine/Source/Runtime/RenderCore/Private/ShaderPipelineCache.cpp
--- a/Engine/Source/Runtime/RenderCore/Private/ShaderPipelineCache.cpp
+++ b/Engine/Source/Runtime/RenderCore/Private/ShaderPipelineCache.cpp
@@ -65,11 +65,7 @@ static TAutoConsoleVariable<int32> CVarPSOFileCacheBackgroundBatchSize(
 														  );
 static TAutoConsoleVariable<int32> CVarPSOFileCacheBatchSize(
 														   TEXT("r.ShaderPipelineCache.BatchSize"),
-#if PLATFORM_MAC
-														   16, // On Mac, where we have many more PSOs to preload due to different video settings 16 works better than 50
-#else
 														   50,
-#endif
 														   TEXT("Set the number of PipelineStateObjects to compile in a single batch operation when compiling takes priority. Defaults to a maximum of 50 per frame, due to async. file IO it is less in practice."),
 														   ECVF_Default | ECVF_RenderThreadSafe
 														   );
