Fix editor link error

diff --git a/Engine/Source/Runtime/RHI/Public/RHIResources.h b/Engine/Source/Runtime/RHI/Public/RHIResources.h
--- a/Engine/Source/Runtime/RHI/Public/RHIResources.h
+++ b/Engine/Source/Runtime/RHI/Public/RHIResources.h
@@ -2580,5 +2580,5 @@ struct FRHIRenderPassInfo
 	bool bIsMSAA = false;
 
 private:
-	FORCENOINLINE void OnVerifyNumUAVsFailed(int32 InNumUAVs);
+	RHI_API void OnVerifyNumUAVsFailed(int32 InNumUAVs);
 };
