fix UE-71894. device was being used from both rhi thread and renderthread.

diff --git a/Engine/Source/Runtime/RHI/Private/RHICommandList.cpp b/Engine/Source/Runtime/RHI/Private/RHICommandList.cpp
--- a/Engine/Source/Runtime/RHI/Private/RHICommandList.cpp
+++ b/Engine/Source/Runtime/RHI/Private/RHICommandList.cpp
@@ -2104,14 +2104,14 @@ void* FDynamicRHI::LockStagingBuffer_RenderThread(class FRHICommandListImmediate
 {
 	QUICK_SCOPE_CYCLE_COUNTER(STAT_FDynamicRHI_LockStagingBuffer_RenderThread);
 	check(IsInRenderingThread());
-
+	RHICmdList.ImmediateFlush(EImmediateFlushType::FlushRHIThread);
 	return GDynamicRHI->RHILockStagingBuffer(StagingBuffer, Offset, SizeRHI);
 }
 void FDynamicRHI::UnlockStagingBuffer_RenderThread(class FRHICommandListImmediate& RHICmdList, FStagingBufferRHIParamRef StagingBuffer)
 {
 	QUICK_SCOPE_CYCLE_COUNTER(STAT_FDynamicRHI_UnlockStagingBuffer_RenderThread);
 	check(IsInRenderingThread());
-	
+	RHICmdList.ImmediateFlush(EImmediateFlushType::FlushRHIThread);
 	GDynamicRHI->RHIUnlockStagingBuffer(StagingBuffer);
 }
 
