rhi display for mac/linux

diff --git a/Engine/Source/Runtime/RHI/Private/Apple/AppleDynamicRHI.cpp b/Engine/Source/Runtime/RHI/Private/Apple/AppleDynamicRHI.cpp
--- a/Engine/Source/Runtime/RHI/Private/Apple/AppleDynamicRHI.cpp
+++ b/Engine/Source/Runtime/RHI/Private/Apple/AppleDynamicRHI.cpp
@@ -124,6 +124,7 @@ FDynamicRHI* PlatformCreateDynamicRHI()
 			}
 #endif
 		}
+		FApp::SetGraphicsRHI(TEXT("Metal"));
 	}
 	else if (bForceOpenGL)
 	{
@@ -136,6 +137,7 @@ FDynamicRHI* PlatformCreateDynamicRHI()
 			RequestedFeatureLevel = ERHIFeatureLevel::ES2;
 #endif
 		}
+		FApp::SetGraphicsRHI(TEXT("OpenGL"));
 	}
 	else
 	{
diff --git a/Engine/Source/Runtime/RHI/Private/Linux/LinuxDynamicRHI.cpp b/Engine/Source/Runtime/RHI/Private/Linux/LinuxDynamicRHI.cpp
--- a/Engine/Source/Runtime/RHI/Private/Linux/LinuxDynamicRHI.cpp
+++ b/Engine/Source/Runtime/RHI/Private/Linux/LinuxDynamicRHI.cpp
@@ -54,6 +54,7 @@ FDynamicRHI* PlatformCreateDynamicRHI()
 			}
 			else
 			{
+				FApp::SetGraphicsRHI(TEXT("Vulkan"));
 				FPlatformApplicationMisc::UsingVulkan();
 
 				FName ShaderFormatName(*TargetedShaderFormats[SfIdx]);
@@ -72,6 +73,7 @@ FDynamicRHI* PlatformCreateDynamicRHI()
 			}
 			else
 			{
+				FApp::SetGraphicsRHI(TEXT("OpenGL"));
 				FPlatformApplicationMisc::UsingOpenGL();
 
 				FName ShaderFormatName(*TargetedShaderFormats[SfIdx]);
