Changes needed to turn on static mesh LOD streaming post HLR
- RHI APIs that support releasing or changing underlying resources without releasing RHI objects (only implemented for dx11, dx12, and vulkan currently)
- Fixed a crash when mesh particles are used with LOD streaming on
- Fixed a static mesh derived data atomicity problem

Tested on following platforms:
- Win64 editor: vertex color painting, changing some static mesh settings, PIE
- Win64 client: play in a map, LLM
- XB1 client: play in a map

diff --git a/Engine/Source/Runtime/Core/Public/GenericPlatform/GenericPlatformProperties.h b/Engine/Source/Runtime/Core/Public/GenericPlatform/GenericPlatformProperties.h
--- a/Engine/Source/Runtime/Core/Public/GenericPlatform/GenericPlatformProperties.h
+++ b/Engine/Source/Runtime/Core/Public/GenericPlatform/GenericPlatformProperties.h
@@ -241,7 +241,7 @@ struct FGenericPlatformProperties
 
 	static FORCEINLINE bool SupportsMeshLODStreaming()
 	{
-		return true;
+		return false;
 	}
 
 	static FORCEINLINE bool SupportsMemoryMappedFiles()
diff --git a/Engine/Source/Runtime/Core/Public/Windows/WindowsPlatformProperties.h b/Engine/Source/Runtime/Core/Public/Windows/WindowsPlatformProperties.h
--- a/Engine/Source/Runtime/Core/Public/Windows/WindowsPlatformProperties.h
+++ b/Engine/Source/Runtime/Core/Public/Windows/WindowsPlatformProperties.h
@@ -78,6 +78,11 @@ struct FWindowsPlatformProperties
 		return !IsServerOnly();
 	}
 
+	static FORCEINLINE bool SupportsMeshLODStreaming()
+	{
+		return !IsServerOnly();
+	}
+
 	static FORCEINLINE bool SupportsGrayscaleSRGB()
 	{
 		return false; // Requires expand from G8 to RGBA
diff --git a/Engine/Source/Runtime/D3D12RHI/Private/D3D12IndexBuffer.cpp b/Engine/Source/Runtime/D3D12RHI/Private/D3D12IndexBuffer.cpp
--- a/Engine/Source/Runtime/D3D12RHI/Private/D3D12IndexBuffer.cpp
+++ b/Engine/Source/Runtime/D3D12RHI/Private/D3D12IndexBuffer.cpp
@@ -31,7 +31,10 @@ D3D12_RESOURCE_DESC CreateIndexBufferResourceDesc(uint32 Size, uint32 InUsage)
 
 FD3D12IndexBuffer::~FD3D12IndexBuffer()
 {
-	UpdateBufferStats<FD3D12IndexBuffer>(&ResourceLocation, false);
+	if (ResourceLocation.IsValid())
+	{
+		UpdateBufferStats<FD3D12IndexBuffer>(&ResourceLocation, false);
+	}
 }
 
 void FD3D12IndexBuffer::Rename(FD3D12ResourceLocation& NewLocation)
@@ -59,8 +62,35 @@ void FD3D12IndexBuffer::RenameLDAChain(FD3D12ResourceLocation& NewLocation)
 	}
 }
 
+void FD3D12IndexBuffer::Swap(FD3D12IndexBuffer& Other)
+{
+	check(!LockedData.bLocked && !Other.LockedData.bLocked);
+	FRHIIndexBuffer::Swap(Other);
+	FD3D12BaseShaderResource::Swap(Other);
+	FD3D12TransientResource::Swap(Other);
+	FD3D12LinkedAdapterObject<FD3D12IndexBuffer>::Swap(Other);
+}
+
+void FD3D12IndexBuffer::ReleaseUnderlyingResource()
+{
+	check(!LockedData.bLocked && ResourceLocation.IsValid());
+	UpdateBufferStats<FD3D12IndexBuffer>(&ResourceLocation, false);
+	ResourceLocation.Clear();
+	FRHIIndexBuffer::ReleaseUnderlyingResource();
+	FD3D12IndexBuffer* NextIB = GetNextObject();
+	if (NextIB)
+	{
+		NextIB->ReleaseUnderlyingResource();
+	}
+}
+
 FIndexBufferRHIRef FD3D12DynamicRHI::RHICreateIndexBuffer(uint32 Stride, uint32 Size, uint32 InUsage, FRHIResourceCreateInfo& CreateInfo)
 {
+	if (CreateInfo.bWithoutNativeResource)
+	{
+		return new FD3D12IndexBuffer();
+	}
+
 	const D3D12_RESOURCE_DESC Desc = CreateIndexBufferResourceDesc(Size, InUsage);
 	const uint32 Alignment = 4;
 
@@ -84,8 +114,28 @@ void FD3D12DynamicRHI::RHIUnlockIndexBuffer(FIndexBufferRHIParamRef IndexBufferR
 	UnlockBuffer(nullptr, FD3D12DynamicRHI::ResourceCast(IndexBufferRHI));
 }
 
+void FD3D12DynamicRHI::RHITransferIndexBufferUnderlyingResource(FIndexBufferRHIParamRef DestIndexBuffer, FIndexBufferRHIParamRef SrcIndexBuffer)
+{
+	check(DestIndexBuffer);
+	FD3D12IndexBuffer* Dest = ResourceCast(DestIndexBuffer);
+	if (!SrcIndexBuffer)
+	{
+		Dest->ReleaseUnderlyingResource();
+	}
+	else
+	{
+		FD3D12IndexBuffer* Src = ResourceCast(SrcIndexBuffer);
+		Dest->Swap(*Src);
+	}
+}
+
 FIndexBufferRHIRef FD3D12DynamicRHI::CreateIndexBuffer_RenderThread(class FRHICommandListImmediate& RHICmdList, uint32 Stride, uint32 Size, uint32 InUsage, FRHIResourceCreateInfo& CreateInfo)
 {
+	if (CreateInfo.bWithoutNativeResource)
+	{
+		return new FD3D12IndexBuffer();
+	}
+
 	const D3D12_RESOURCE_DESC Desc = CreateIndexBufferResourceDesc(Size, InUsage);
 	const uint32 Alignment = 4;
 
diff --git a/Engine/Source/Runtime/D3D12RHI/Private/D3D12RHICommon.h b/Engine/Source/Runtime/D3D12RHI/Private/D3D12RHICommon.h
--- a/Engine/Source/Runtime/D3D12RHI/Private/D3D12RHICommon.h
+++ b/Engine/Source/Runtime/D3D12RHI/Private/D3D12RHICommon.h
@@ -153,6 +153,12 @@ public:
 		return NextNode.GetReference();
 	}
 
+	void Swap(FD3D12LinkedAdapterObject& Other)
+	{
+		check(bIsHeadLink && Other.bIsHeadLink);
+		NextNode.Swap(Other.NextNode);
+	}
+
 private:
 
 	TRefCountPtr<ObjectType> NextNode;
diff --git a/Engine/Source/Runtime/D3D12RHI/Private/D3D12RHIPrivate.h b/Engine/Source/Runtime/D3D12RHI/Private/D3D12RHIPrivate.h
--- a/Engine/Source/Runtime/D3D12RHI/Private/D3D12RHIPrivate.h
+++ b/Engine/Source/Runtime/D3D12RHI/Private/D3D12RHIPrivate.h
@@ -283,10 +283,12 @@ public:
 	virtual FIndexBufferRHIRef RHICreateIndexBuffer(uint32 Stride, uint32 Size, uint32 InUsage, FRHIResourceCreateInfo& CreateInfo) final override;
 	virtual void* RHILockIndexBuffer(FIndexBufferRHIParamRef IndexBuffer, uint32 Offset, uint32 Size, EResourceLockMode LockMode) final override;
 	virtual void RHIUnlockIndexBuffer(FIndexBufferRHIParamRef IndexBuffer) final override;
+	virtual void RHITransferIndexBufferUnderlyingResource(FIndexBufferRHIParamRef DestIndexBuffer, FIndexBufferRHIParamRef SrcIndexBuffer) final override;
 	virtual FVertexBufferRHIRef RHICreateVertexBuffer(uint32 Size, uint32 InUsage, FRHIResourceCreateInfo& CreateInfo) final override;
 	virtual void* RHILockVertexBuffer(FVertexBufferRHIParamRef VertexBuffer, uint32 Offset, uint32 SizeRHI, EResourceLockMode LockMode) final override;
 	virtual void RHIUnlockVertexBuffer(FVertexBufferRHIParamRef VertexBuffer) final override;
 	virtual void RHICopyVertexBuffer(FVertexBufferRHIParamRef SourceBuffer, FVertexBufferRHIParamRef DestBuffer) final override;
+	virtual void RHITransferVertexBufferUnderlyingResource(FVertexBufferRHIParamRef DestVertexBuffer, FVertexBufferRHIParamRef SrcVertexBuffer) final override;
 	virtual FStructuredBufferRHIRef RHICreateStructuredBuffer(uint32 Stride, uint32 Size, uint32 InUsage, FRHIResourceCreateInfo& CreateInfo) final override;
 	virtual void* RHILockStructuredBuffer(FStructuredBufferRHIParamRef StructuredBuffer, uint32 Offset, uint32 SizeRHI, EResourceLockMode LockMode) final override;
 	virtual void RHIUnlockStructuredBuffer(FStructuredBufferRHIParamRef StructuredBuffer) final override;
@@ -298,6 +300,7 @@ public:
 	virtual FShaderResourceViewRHIRef RHICreateShaderResourceView(FVertexBufferRHIParamRef VertexBuffer, uint32 Stride, uint8 Format) final override;
 	virtual FShaderResourceViewRHIRef RHICreateShaderResourceView(FIndexBufferRHIParamRef Buffer) final override;
 	virtual void RHIUpdateShaderResourceView(FShaderResourceViewRHIParamRef SRV, FVertexBufferRHIParamRef VertexBuffer, uint32 Stride, uint8 Format) final override;
+	virtual void RHIUpdateShaderResourceView(FShaderResourceViewRHIParamRef SRV, FIndexBufferRHIParamRef IndexBuffer) final override;
 	virtual uint64 RHICalcTexture2DPlatformSize(uint32 SizeX, uint32 SizeY, uint8 Format, uint32 NumMips, uint32 NumSamples, uint32 Flags, uint32& OutAlign) override;
 	virtual uint64 RHICalcTexture3DPlatformSize(uint32 SizeX, uint32 SizeY, uint32 SizeZ, uint8 Format, uint32 NumMips, uint32 Flags, uint32& OutAlign) final override;
 	virtual uint64 RHICalcTextureCubePlatformSize(uint32 Size, uint8 Format, uint32 NumMips, uint32 Flags, uint32& OutAlign) final override;
diff --git a/Engine/Source/Runtime/D3D12RHI/Private/D3D12Resources.cpp b/Engine/Source/Runtime/D3D12RHI/Private/D3D12Resources.cpp
--- a/Engine/Source/Runtime/D3D12RHI/Private/D3D12Resources.cpp
+++ b/Engine/Source/Runtime/D3D12RHI/Private/D3D12Resources.cpp
@@ -480,6 +480,10 @@ void FD3D12ResourceLocation::TransferOwnership(FD3D12ResourceLocation& Destinati
 	Source.InternalClear<false>();
 }
 
+void FD3D12ResourceLocation::Swap(FD3D12ResourceLocation& Other)
+{
+	::Swap(*this, Other);
+}
 
 void FD3D12ResourceLocation::Alias(FD3D12ResourceLocation & Destination, FD3D12ResourceLocation & Source)
 {
diff --git a/Engine/Source/Runtime/D3D12RHI/Private/D3D12VertexBuffer.cpp b/Engine/Source/Runtime/D3D12RHI/Private/D3D12VertexBuffer.cpp
--- a/Engine/Source/Runtime/D3D12RHI/Private/D3D12VertexBuffer.cpp
+++ b/Engine/Source/Runtime/D3D12RHI/Private/D3D12VertexBuffer.cpp
@@ -78,9 +78,36 @@ void FD3D12VertexBuffer::RenameLDAChain(FD3D12ResourceLocation& NewLocation)
 	}
 }
 
+void FD3D12VertexBuffer::Swap(FD3D12VertexBuffer& Other)
+{
+	check(!LockedData.bLocked && !Other.LockedData.bLocked);
+	FRHIVertexBuffer::Swap(Other);
+	FD3D12BaseShaderResource::Swap(Other);
+	FD3D12TransientResource::Swap(Other);
+	FD3D12LinkedAdapterObject<FD3D12VertexBuffer>::Swap(Other);
+	::Swap(DynamicSRV, Other.DynamicSRV);
+}
+
+void FD3D12VertexBuffer::ReleaseUnderlyingResource()
+{
+	check(!LockedData.bLocked && ResourceLocation.IsValid());
+	UpdateBufferStats<FD3D12VertexBuffer>(&ResourceLocation, false);
+	ResourceLocation.Clear();
+	FRHIVertexBuffer::ReleaseUnderlyingResource();
+	FD3D12VertexBuffer* NextVB = GetNextObject();
+	if (NextVB)
+	{
+		NextVB->ReleaseUnderlyingResource();
+	}
+}
 
 FVertexBufferRHIRef FD3D12DynamicRHI::RHICreateVertexBuffer(uint32 Size, uint32 InUsage, FRHIResourceCreateInfo& CreateInfo)
 {
+	if (CreateInfo.bWithoutNativeResource)
+	{
+		return new FD3D12VertexBuffer();
+	}
+
 	const D3D12_RESOURCE_DESC Desc = CreateVertexBufferResourceDesc(Size, InUsage);
 	const uint32 Alignment = 4;
 
@@ -106,6 +133,11 @@ void FD3D12DynamicRHI::RHIUnlockVertexBuffer(FVertexBufferRHIParamRef VertexBuff
 
 FVertexBufferRHIRef FD3D12DynamicRHI::CreateVertexBuffer_RenderThread(FRHICommandListImmediate& RHICmdList, uint32 Size, uint32 InUsage, FRHIResourceCreateInfo& CreateInfo)
 {	
+	if (CreateInfo.bWithoutNativeResource)
+	{
+		return new FD3D12VertexBuffer();
+	}
+
 	const D3D12_RESOURCE_DESC Desc = CreateVertexBufferResourceDesc(Size, InUsage);
 	const uint32 Alignment = 4;
 
@@ -171,6 +203,21 @@ void FD3D12DynamicRHI::RHICopyVertexBuffer(FVertexBufferRHIParamRef SourceBuffer
 	}
 }
 
+void FD3D12DynamicRHI::RHITransferVertexBufferUnderlyingResource(FVertexBufferRHIParamRef DestVertexBuffer, FVertexBufferRHIParamRef SrcVertexBuffer)
+{
+	check(DestVertexBuffer);
+	FD3D12VertexBuffer* Dest = ResourceCast(DestVertexBuffer);
+	if (!SrcVertexBuffer)
+	{
+		Dest->ReleaseUnderlyingResource();
+	}
+	else
+	{
+		FD3D12VertexBuffer* Src = ResourceCast(SrcVertexBuffer);
+		Dest->Swap(*Src);
+	}
+}
+
 #if D3D12_RHI_RAYTRACING
 void FD3D12CommandContext::RHICopyBufferRegion(FVertexBufferRHIParamRef DestBufferRHI, uint64 DstOffset, FVertexBufferRHIParamRef SourceBufferRHI, uint64 SrcOffset, uint64 NumBytes)
 {
diff --git a/Engine/Source/Runtime/D3D12RHI/Private/D3D12View.cpp b/Engine/Source/Runtime/D3D12RHI/Private/D3D12View.cpp
--- a/Engine/Source/Runtime/D3D12RHI/Private/D3D12View.cpp
+++ b/Engine/Source/Runtime/D3D12RHI/Private/D3D12View.cpp
@@ -6,6 +6,64 @@
 
 #include "D3D12RHIPrivate.h"
 
+static FORCEINLINE D3D12_SHADER_RESOURCE_VIEW_DESC GetVertexBufferSRVDesc(FD3D12VertexBuffer* VertexBuffer, uint32& CreationStride, uint8 Format)
+{
+	D3D12_SHADER_RESOURCE_VIEW_DESC SRVDesc = {};
+	SRVDesc.Shader4ComponentMapping = D3D12_DEFAULT_SHADER_4_COMPONENT_MAPPING;
+	SRVDesc.ViewDimension = D3D12_SRV_DIMENSION_BUFFER;
+	if (VertexBuffer->GetUsage() & BUF_ByteAddressBuffer)
+	{
+		SRVDesc.Format = DXGI_FORMAT_R32_TYPELESS;
+		SRVDesc.Buffer.Flags = D3D12_BUFFER_SRV_FLAG_RAW;
+		CreationStride = 4;
+	}
+	else
+	{
+		SRVDesc.Format = FindShaderResourceDXGIFormat((DXGI_FORMAT)GPixelFormats[Format].PlatformFormat, false);
+	}
+	SRVDesc.Buffer.FirstElement = VertexBuffer->ResourceLocation.GetOffsetFromBaseOfResource() / CreationStride;
+	SRVDesc.Buffer.NumElements = VertexBuffer->GetSize() / CreationStride;
+	SRVDesc.Buffer.StructureByteStride = 0;
+	return SRVDesc;
+}
+
+static FORCEINLINE D3D12_SHADER_RESOURCE_VIEW_DESC GetIndexBufferSRVDesc(FD3D12IndexBuffer* IndexBuffer)
+{
+	const uint32 Usage = IndexBuffer->GetUsage();
+	const uint32 Width = IndexBuffer->GetSize();
+	const uint32 CreationStride = IndexBuffer->GetStride();
+	const FD3D12ResourceLocation& Location = IndexBuffer->ResourceLocation;
+
+	D3D12_SHADER_RESOURCE_VIEW_DESC SRVDesc = {};
+	SRVDesc.Shader4ComponentMapping = D3D12_DEFAULT_SHADER_4_COMPONENT_MAPPING;
+	SRVDesc.ViewDimension = D3D12_SRV_DIMENSION_BUFFER;
+	if (Usage & BUF_ByteAddressBuffer)
+	{
+		check(CreationStride == 4u);
+		SRVDesc.Format = DXGI_FORMAT_R32_TYPELESS;
+		SRVDesc.Buffer.Flags = D3D12_BUFFER_SRV_FLAG_RAW;
+	}
+	else
+	{
+		check(CreationStride == 2u || CreationStride == 4u);
+		SRVDesc.Format = CreationStride == 2u ? DXGI_FORMAT_R16_UINT : DXGI_FORMAT_R32_UINT;
+	}
+	SRVDesc.Buffer.NumElements = Width / CreationStride;
+	SRVDesc.Buffer.StructureByteStride = 0;
+
+	if (Location.GetResource())
+	{
+		// Create a Shader Resource View
+		SRVDesc.Buffer.FirstElement = Location.GetOffsetFromBaseOfResource() / CreationStride;
+	}
+	else
+	{
+		// Null underlying D3D12 resource should only be the case for dynamic resources
+		check(Usage & BUF_AnyDynamic);
+	}
+	return SRVDesc;
+}
+
 template<typename TextureType>
 FD3D12ShaderResourceView* CreateSRV(TextureType* Texture, D3D12_SHADER_RESOURCE_VIEW_DESC& Desc)
 {
@@ -218,25 +276,7 @@ FShaderResourceViewRHIRef FD3D12DynamicRHI::RHICreateShaderResourceView(FVertexB
 		void Execute(FRHICommandListBase& RHICmdList)
 		{
 			uint32 CreationStride = Stride;
-			uint32 Width = VertexBuffer->GetSize();
-
-			D3D12_SHADER_RESOURCE_VIEW_DESC SRVDesc = {};
-			SRVDesc.Shader4ComponentMapping = D3D12_DEFAULT_SHADER_4_COMPONENT_MAPPING;
-			SRVDesc.ViewDimension = D3D12_SRV_DIMENSION_BUFFER;
-			if (VertexBuffer->GetUsage() & BUF_ByteAddressBuffer)
-			{
-				SRVDesc.Format = DXGI_FORMAT_R32_TYPELESS;
-				SRVDesc.Buffer.NumElements = Width / 4;
-				SRVDesc.Buffer.Flags = D3D12_BUFFER_SRV_FLAG_RAW;
-				CreationStride = 4;
-			}
-			else
-			{
-				SRVDesc.Format = FindShaderResourceDXGIFormat((DXGI_FORMAT)GPixelFormats[Format].PlatformFormat, false);
-				SRVDesc.Buffer.NumElements = Width / Stride;
-			}
-			SRVDesc.Buffer.StructureByteStride = 0;
-			SRVDesc.Buffer.FirstElement = VertexBuffer->ResourceLocation.GetOffsetFromBaseOfResource() / CreationStride;
+			D3D12_SHADER_RESOURCE_VIEW_DESC SRVDesc = GetVertexBufferSRVDesc(VertexBuffer, CreationStride, Format);
 
 			// Create a Shader Resource View
 			SRV->Initialize(SRVDesc, VertexBuffer->ResourceLocation, CreationStride);
@@ -278,6 +318,10 @@ FShaderResourceViewRHIRef FD3D12DynamicRHI::RHICreateShaderResourceView(FVertexB
 
 FShaderResourceViewRHIRef FD3D12DynamicRHI::RHICreateShaderResourceView(FIndexBufferRHIParamRef BufferRHI)
 {
+	if (!BufferRHI)
+	{
+		return new FD3D12ShaderResourceView(nullptr);
+	}
 	FD3D12IndexBuffer* IndexBuffer = FD3D12DynamicRHI::ResourceCast(BufferRHI);
 	return GetAdapter().CreateLinkedViews<FD3D12IndexBuffer, FD3D12ShaderResourceView>(IndexBuffer,
 		[](FD3D12IndexBuffer* IndexBuffer)
@@ -285,41 +329,8 @@ FShaderResourceViewRHIRef FD3D12DynamicRHI::RHICreateShaderResourceView(FIndexBu
 		check(IndexBuffer);
 
 		FD3D12ResourceLocation& Location = IndexBuffer->ResourceLocation;
-
-		const uint32 Width = IndexBuffer->GetSize();
-
-		FD3D12Resource* pResource = Location.GetResource();
-
-		uint32 CreationStride = IndexBuffer->GetStride();
-
-		D3D12_SHADER_RESOURCE_VIEW_DESC SRVDesc = {};
-		SRVDesc.Shader4ComponentMapping = D3D12_DEFAULT_SHADER_4_COMPONENT_MAPPING;
-		SRVDesc.ViewDimension = D3D12_SRV_DIMENSION_BUFFER;
-		if (IndexBuffer->GetUsage() & BUF_ByteAddressBuffer)
-		{
-			check(CreationStride == 4u);
-			SRVDesc.Format = DXGI_FORMAT_R32_TYPELESS;
-			SRVDesc.Buffer.NumElements = Width / 4;
-			SRVDesc.Buffer.Flags = D3D12_BUFFER_SRV_FLAG_RAW;
-		}
-		else
-		{
-			check(CreationStride == 2u || CreationStride == 4u);
-			SRVDesc.Format = CreationStride == 2u ? DXGI_FORMAT_R16_UINT : DXGI_FORMAT_R32_UINT;
-			SRVDesc.Buffer.NumElements = Width / CreationStride;
-		}
-		SRVDesc.Buffer.StructureByteStride = 0;
-
-		if (pResource)
-		{
-			// Create a Shader Resource View
-			SRVDesc.Buffer.FirstElement = Location.GetOffsetFromBaseOfResource() / CreationStride;
-		}
-		else
-		{
-			// Null underlying D3D12 resource should only be the case for dynamic resources
-			check(IndexBuffer->GetUsage() & BUF_AnyDynamic);
-		}
+		const uint32 CreationStride = IndexBuffer->GetStride();
+		D3D12_SHADER_RESOURCE_VIEW_DESC SRVDesc = GetIndexBufferSRVDesc(IndexBuffer);
 
 		FD3D12ShaderResourceView* ShaderResourceView = new FD3D12ShaderResourceView(IndexBuffer->GetParentDevice(), SRVDesc, Location, CreationStride);
 		return ShaderResourceView;
@@ -333,23 +344,7 @@ void FD3D12DynamicRHI::RHIUpdateShaderResourceView(FShaderResourceViewRHIParamRe
 	{
 		FD3D12VertexBuffer* VBD3D12 = ResourceCast(VertexBuffer);
 		FD3D12ShaderResourceView* SRVD3D12 = ResourceCast(SRV);
-		D3D12_SHADER_RESOURCE_VIEW_DESC SRVDesc = {};
-		SRVDesc.Shader4ComponentMapping = D3D12_DEFAULT_SHADER_4_COMPONENT_MAPPING;
-		SRVDesc.ViewDimension = D3D12_SRV_DIMENSION_BUFFER;
-		if (VBD3D12->GetUsage() & BUF_ByteAddressBuffer)
-		{
-			SRVDesc.Format = DXGI_FORMAT_R32_TYPELESS;
-			SRVDesc.Buffer.NumElements = VBD3D12->GetSize() / 4;
-			SRVDesc.Buffer.Flags = D3D12_BUFFER_SRV_FLAG_RAW;
-			Stride = 4;
-		}
-		else
-		{
-			SRVDesc.Format = FindShaderResourceDXGIFormat((DXGI_FORMAT)GPixelFormats[Format].PlatformFormat, false);
-			SRVDesc.Buffer.NumElements = VBD3D12->GetSize() / Stride;
-		}
-		SRVDesc.Buffer.StructureByteStride = 0;
-		SRVDesc.Buffer.FirstElement = VBD3D12->ResourceLocation.GetOffsetFromBaseOfResource() / Stride;
+		D3D12_SHADER_RESOURCE_VIEW_DESC SRVDesc = GetVertexBufferSRVDesc(VBD3D12, Stride, Format);
 
 		// Rename the SRV to view on the new vertex buffer
 		while (VBD3D12)
@@ -360,7 +355,32 @@ void FD3D12DynamicRHI::RHIUpdateShaderResourceView(FShaderResourceViewRHIParamRe
 			VBD3D12 = VBD3D12->GetNextObject();
 			if (VBD3D12 && !SRVD3D12->GetNextObject())
 			{
-				SRVD3D12->SetNextObject(new FD3D12ShaderResourceView(ParentDevice));
+				SRVD3D12->SetNextObject(new FD3D12ShaderResourceView(VBD3D12->GetParentDevice()));
+			}
+			SRVD3D12 = SRVD3D12->GetNextObject();
+		}
+	}
+}
+
+void FD3D12DynamicRHI::RHIUpdateShaderResourceView(FShaderResourceViewRHIParamRef SRV, FIndexBufferRHIParamRef IndexBuffer)
+{
+	check(SRV);
+	if (IndexBuffer)
+	{
+		FD3D12IndexBuffer* IBD3D12 = ResourceCast(IndexBuffer);
+		FD3D12ShaderResourceView* SRVD3D12 = ResourceCast(SRV);
+		D3D12_SHADER_RESOURCE_VIEW_DESC SRVDesc = GetIndexBufferSRVDesc(IBD3D12);
+		const uint32 Stride = IBD3D12->GetStride();
+
+		// Rename the SRV to view on the new index buffer
+		while (IBD3D12)
+		{
+			FD3D12Device* ParentDevice = IBD3D12->GetParentDevice();
+			SRVD3D12->Initialize(ParentDevice, SRVDesc, IBD3D12->ResourceLocation, Stride);
+			IBD3D12 = IBD3D12->GetNextObject();
+			if (IBD3D12 && !SRVD3D12->GetNextObject())
+			{
+				SRVD3D12->SetNextObject(new FD3D12ShaderResourceView(IBD3D12->GetParentDevice()));
 			}
 			SRVD3D12 = SRVD3D12->GetNextObject();
 		}
diff --git a/Engine/Source/Runtime/D3D12RHI/Public/D3D12Resources.h b/Engine/Source/Runtime/D3D12RHI/Public/D3D12Resources.h
--- a/Engine/Source/Runtime/D3D12RHI/Public/D3D12Resources.h
+++ b/Engine/Source/Runtime/D3D12RHI/Public/D3D12Resources.h
@@ -524,6 +524,8 @@ public:
 		return bTransient;
 	}
 
+	void Swap(FD3D12ResourceLocation& Other);
+
 private:
 
 	template<bool bReleaseResource>
@@ -657,6 +659,13 @@ class FD3D12BaseShaderResource : public FD3D12DeviceChild, public IRefCountedObj
 public:
 	FD3D12Resource* GetResource() const { return ResourceLocation.GetResource(); }
 
+	void Swap(FD3D12BaseShaderResource& Other)
+	{
+		::Swap(Parent, Other.Parent);
+		ResourceLocation.Swap(Other.ResourceLocation);
+		::Swap(BufferAlignment, Other.BufferAlignment);
+	}
+
 	FD3D12ResourceLocation ResourceLocation;
 	uint32 BufferAlignment;
 
@@ -711,6 +720,8 @@ public:
 class FD3D12TransientResource
 {
 	// Nothing special for fast ram
+public:
+	void Swap(FD3D12TransientResource&) {}
 };
 class FD3D12FastClearResource {};
 #endif
@@ -719,6 +730,10 @@ class FD3D12FastClearResource {};
 class FD3D12IndexBuffer : public FRHIIndexBuffer, public FD3D12BaseShaderResource, public FD3D12TransientResource, public FD3D12LinkedAdapterObject<FD3D12IndexBuffer>
 {
 public:
+	FD3D12IndexBuffer()
+		: FD3D12BaseShaderResource(nullptr)
+		, LockedData(nullptr)
+	{}
 
 	FD3D12IndexBuffer(FD3D12Device* InParent, uint32 InStride, uint32 InSize, uint32 InUsage)
 		: FRHIIndexBuffer(InStride, InSize, InUsage)
@@ -731,6 +746,10 @@ public:
 	void Rename(FD3D12ResourceLocation& NewLocation);
 	void RenameLDAChain(FD3D12ResourceLocation& NewLocation);
 
+	void Swap(FD3D12IndexBuffer& Other);
+
+	void ReleaseUnderlyingResource();
+
 	// IRefCountedObject interface.
 	virtual uint32 AddRef() const
 	{
@@ -791,6 +810,12 @@ public:
 	// Current SRV
 	FD3D12ShaderResourceView* DynamicSRV;
 
+	FD3D12VertexBuffer()
+		: FD3D12BaseShaderResource(nullptr)
+		, DynamicSRV(nullptr)
+		, LockedData(nullptr)
+	{}
+
 	FD3D12VertexBuffer(FD3D12Device* InParent, uint32 InStride, uint32 InSize, uint32 InUsage)
 		: FRHIVertexBuffer(InSize, InUsage)
 		, FD3D12BaseShaderResource(InParent)
@@ -810,6 +835,10 @@ public:
 		DynamicSRV = InSRV;
 	}
 
+	void Swap(FD3D12VertexBuffer& Other);
+
+	void ReleaseUnderlyingResource();
+
 	// IRefCountedObject interface.
 	virtual uint32 AddRef() const
 	{
diff --git a/Engine/Source/Runtime/EmptyRHI/Private/EmptyIndexBuffer.cpp b/Engine/Source/Runtime/EmptyRHI/Private/EmptyIndexBuffer.cpp
--- a/Engine/Source/Runtime/EmptyRHI/Private/EmptyIndexBuffer.cpp
+++ b/Engine/Source/Runtime/EmptyRHI/Private/EmptyIndexBuffer.cpp
@@ -25,6 +25,11 @@ void FEmptyIndexBuffer::Unlock()
 
 FIndexBufferRHIRef FEmptyDynamicRHI::RHICreateIndexBuffer(uint32 Stride, uint32 Size, uint32 InUsage, FRHIResourceCreateInfo& CreateInfo)
 {
+	if (CreateInfo.bCreateRHIObjectOnly)
+	{
+		return new FEmptyIndexBuffer();
+	}
+
 	// make the RHI object, which will allocate memory
 	FEmptyIndexBuffer* IndexBuffer = new FEmptyIndexBuffer(Stride, Size, InUsage);
 	
@@ -60,3 +65,8 @@ void FEmptyDynamicRHI::RHIUnlockIndexBuffer(FIndexBufferRHIParamRef IndexBufferR
 
 	IndexBuffer->Unlock();
 }
+
+void FEmptyDynamicRHI::RHITransferIndexBufferUnderlyingResource(FIndexBufferRHIParamRef DestIndexBuffer, FIndexBufferRHIParamRef SrcIndexBuffer)
+{
+
+}
diff --git a/Engine/Source/Runtime/EmptyRHI/Private/EmptyVertexBuffer.cpp b/Engine/Source/Runtime/EmptyRHI/Private/EmptyVertexBuffer.cpp
--- a/Engine/Source/Runtime/EmptyRHI/Private/EmptyVertexBuffer.cpp
+++ b/Engine/Source/Runtime/EmptyRHI/Private/EmptyVertexBuffer.cpp
@@ -24,6 +24,11 @@ void FEmptyVertexBuffer::Unlock()
 
 FVertexBufferRHIRef FEmptyDynamicRHI::RHICreateVertexBuffer(uint32 Size, uint32 InUsage, FRHIResourceCreateInfo& CreateInfo)
 {
+	if (CreateInfo.bCreateRHIObjectOnly)
+	{
+		return new FEmptyVertexBuffer();
+	}
+
 	// make the RHI object, which will allocate memory
 	FEmptyVertexBuffer* VertexBuffer = new FEmptyVertexBuffer(Size, InUsage);
 
@@ -64,4 +69,9 @@ void FEmptyDynamicRHI::RHIUnlockVertexBuffer(FVertexBufferRHIParamRef VertexBuff
 void FEmptyDynamicRHI::RHICopyVertexBuffer(FVertexBufferRHIParamRef SourceBufferRHI,FVertexBufferRHIParamRef DestBufferRHI)
 {
 
+}
+
+void FEmptyDynamicRHI::RHITransferVertexBufferUnderlyingResource(FVertexBufferRHIParamRef DestVertexBuffer, FVertexBufferRHIParamRef SrcVertexBuffer)
+{
+
 }
\ No newline at end of file
diff --git a/Engine/Source/Runtime/EmptyRHI/Public/EmptyRHI.h b/Engine/Source/Runtime/EmptyRHI/Public/EmptyRHI.h
--- a/Engine/Source/Runtime/EmptyRHI/Public/EmptyRHI.h
+++ b/Engine/Source/Runtime/EmptyRHI/Public/EmptyRHI.h
@@ -53,10 +53,12 @@ public:
 	virtual FIndexBufferRHIRef RHICreateIndexBuffer(uint32 Stride, uint32 Size, uint32 InUsage, FRHIResourceCreateInfo& CreateInfo) final override;
 	virtual void* RHILockIndexBuffer(FIndexBufferRHIParamRef IndexBuffer, uint32 Offset, uint32 Size, EResourceLockMode LockMode) final override;
 	virtual void RHIUnlockIndexBuffer(FIndexBufferRHIParamRef IndexBuffer) final override;
+	virtual void RHITransferIndexBufferUnderlyingResource(FIndexBufferRHIParamRef DestIndexBuffer, FIndexBufferRHIParamRef SrcIndexBuffer) final override;
 	virtual FVertexBufferRHIRef RHICreateVertexBuffer(uint32 Size, uint32 InUsage, FRHIResourceCreateInfo& CreateInfo) final override;
 	virtual void* RHILockVertexBuffer(FVertexBufferRHIParamRef VertexBuffer, uint32 Offset, uint32 SizeRHI, EResourceLockMode LockMode) final override;
 	virtual void RHIUnlockVertexBuffer(FVertexBufferRHIParamRef VertexBuffer) final override;
 	virtual void RHICopyVertexBuffer(FVertexBufferRHIParamRef SourceBuffer,FVertexBufferRHIParamRef DestBuffer) final override;
+	virtual void RHITransferVertexBufferUnderlyingResource(FVertexBufferRHIParamRef DestVertexBuffer, FVertexBufferRHIParamRef SrcVertexBuffer) final override;
 	virtual FStructuredBufferRHIRef RHICreateStructuredBuffer(uint32 Stride, uint32 Size, uint32 InUsage, FRHIResourceCreateInfo& CreateInfo) final override;
 	virtual void* RHILockStructuredBuffer(FStructuredBufferRHIParamRef StructuredBuffer, uint32 Offset, uint32 SizeRHI, EResourceLockMode LockMode) final override;
 	virtual void RHIUnlockStructuredBuffer(FStructuredBufferRHIParamRef StructuredBuffer) final override;
diff --git a/Engine/Source/Runtime/EmptyRHI/Public/EmptyResources.h b/Engine/Source/Runtime/EmptyRHI/Public/EmptyResources.h
--- a/Engine/Source/Runtime/EmptyRHI/Public/EmptyResources.h
+++ b/Engine/Source/Runtime/EmptyRHI/Public/EmptyResources.h
@@ -182,6 +182,7 @@ public:
 class FEmptyIndexBuffer : public FRHIIndexBuffer
 {
 public:
+	FEmptyIndexBuffer() = default;
 
 	/** Constructor */
 	FEmptyIndexBuffer(uint32 InStride, uint32 InSize, uint32 InUsage);
@@ -202,6 +203,7 @@ public:
 class FEmptyVertexBuffer : public FRHIVertexBuffer
 {
 public:
+	FEmptyVertexBuffer() = default;
 
 	/** Constructor */
 	FEmptyVertexBuffer(uint32 InSize, uint32 InUsage);
diff --git a/Engine/Source/Runtime/Engine/Private/Particles/ParticleSystemRender.cpp b/Engine/Source/Runtime/Engine/Private/Particles/ParticleSystemRender.cpp
--- a/Engine/Source/Runtime/Engine/Private/Particles/ParticleSystemRender.cpp
+++ b/Engine/Source/Runtime/Engine/Private/Particles/ParticleSystemRender.cpp
@@ -1534,7 +1534,8 @@ void FDynamicMeshEmitterData::Init( bool bInSelected,
 FParticleVertexFactoryBase* FDynamicMeshEmitterData::BuildVertexFactory(const FParticleSystemSceneProxy* InOwnerProxy)
 {
 	FParticleVertexFactoryBase* PoolVertexFactory = GParticleVertexFactoryPool.GetParticleVertexFactory(PVFT_Mesh, InOwnerProxy->GetScene().GetFeatureLevel(), this);
-	SetupVertexFactory((FMeshParticleVertexFactory*)PoolVertexFactory, StaticMesh->RenderData->LODResources[GetMeshLODIndexFromProxy(InOwnerProxy)]);
+	const uint32 LODIdx = GetMeshLODIndexFromProxy(InOwnerProxy);
+	SetupVertexFactory((FMeshParticleVertexFactory*)PoolVertexFactory, StaticMesh->RenderData->LODResources[LODIdx], LODIdx);
 	return PoolVertexFactory;
 }
 
@@ -1604,7 +1605,8 @@ FParticleVertexFactoryBase *FDynamicMeshEmitterData::CreateVertexFactory(ERHIFea
 	FMeshParticleVertexFactory *VertexFactory = ConstructMeshParticleVertexFactory(InFeatureLevel);
 
 	VertexFactory->SetParticleFactoryType(PVFT_Mesh);
-	SetupVertexFactory(VertexFactory, StaticMesh->RenderData->LODResources[GetMeshLODIndexFromProxy(InOwnerProxy)]);
+	const uint32 LODIdx = GetMeshLODIndexFromProxy(InOwnerProxy);
+	SetupVertexFactory(VertexFactory, StaticMesh->RenderData->LODResources[LODIdx], LODIdx);
 
 	const int32 InstanceVertexStride = GetDynamicVertexStride(InFeatureLevel);
 	const int32 DynamicParameterVertexStride = bUsesDynamicParameter ? GetDynamicParameterVertexStride() : 0;
@@ -1814,7 +1816,12 @@ void FDynamicMeshEmitterData::GetDynamicMeshElementsEmitter(const FParticleSyste
 			}
 			check(StaticMesh != nullptr);
 
-			const FStaticMeshLODResources& LODModel = StaticMesh->RenderData->LODResources[GetMeshLODIndexFromProxy(Proxy)];
+			const uint32 ChosenLODIdx = GetMeshLODIndexFromProxy(Proxy);
+			const FStaticMeshLODResources& LODModel = StaticMesh->RenderData->LODResources[ChosenLODIdx];
+			if (ChosenLODIdx != MeshVertexFactory->GetLODIdx())
+			{
+				SetupVertexFactory(MeshVertexFactory, LODModel, ChosenLODIdx);
+			}
 
 			const bool bIsWireframe = AllowDebugViewmodes() && View->Family->EngineShowFlags.Wireframe;
 
@@ -2576,7 +2583,7 @@ void FDynamicMeshEmitterData::GetInstanceData(void* InstanceData, void* DynamicP
 	}
 }
 
-void FDynamicMeshEmitterData::SetupVertexFactory( FMeshParticleVertexFactory* InVertexFactory, FStaticMeshLODResources& LODResources) const
+void FDynamicMeshEmitterData::SetupVertexFactory( FMeshParticleVertexFactory* InVertexFactory, const FStaticMeshLODResources& LODResources, uint32 LODIdx) const
 {
 		FMeshParticleVertexFactory::FDataType Data;
 
@@ -2635,6 +2642,7 @@ void FDynamicMeshEmitterData::SetupVertexFactory( FMeshParticleVertexFactory* In
 
 		Data.bInitialized = true;
 		InVertexFactory->SetData(Data);
+		InVertexFactory->SetLODIdx((uint8)LODIdx);
 }
 
 ///////////////////////////////////////////////////////////////////////////////
diff --git a/Engine/Source/Runtime/Engine/Private/PositionVertexBuffer.cpp b/Engine/Source/Runtime/Engine/Private/PositionVertexBuffer.cpp
--- a/Engine/Source/Runtime/Engine/Private/PositionVertexBuffer.cpp
+++ b/Engine/Source/Runtime/Engine/Private/PositionVertexBuffer.cpp
@@ -29,8 +29,7 @@ FPositionVertexBuffer::FPositionVertexBuffer():
 	VertexData(NULL),
 	Data(NULL),
 	Stride(0),
-	NumVertices(0),
-	bStreamed(false)
+	NumVertices(0)
 {}
 
 FPositionVertexBuffer::~FPositionVertexBuffer()
@@ -187,19 +186,19 @@ void FPositionVertexBuffer::operator=(const FPositionVertexBuffer &Other)
 template <bool bRenderThread>
 FVertexBufferRHIRef FPositionVertexBuffer::CreateRHIBuffer_Internal()
 {
-	check(VertexData);
-	FResourceArrayInterface* ResourceArray = VertexData->GetResourceArray();
-	if (ResourceArray->GetResourceDataSize())
+	if (NumVertices)
 	{
-		// Create the vertex buffer.
+		FResourceArrayInterface* RESTRICT ResourceArray = VertexData ? VertexData->GetResourceArray() : nullptr;
+		const uint32 SizeInBytes = ResourceArray ? ResourceArray->GetResourceDataSize() : 0;
 		FRHIResourceCreateInfo CreateInfo(ResourceArray);
+		CreateInfo.bWithoutNativeResource = !VertexData;
 		if (bRenderThread)
 		{
-			return RHICreateVertexBuffer(ResourceArray->GetResourceDataSize(), BUF_Static | BUF_ShaderResource, CreateInfo);
+			return RHICreateVertexBuffer(SizeInBytes, BUF_Static | BUF_ShaderResource, CreateInfo);
 		}
 		else
 		{
-			return RHIAsyncCreateVertexBuffer(ResourceArray->GetResourceDataSize(), BUF_Static | BUF_ShaderResource, CreateInfo);
+			return RHIAsyncCreateVertexBuffer(SizeInBytes, BUF_Static | BUF_ShaderResource, CreateInfo);
 		}
 	}
 	return nullptr;
@@ -217,14 +216,11 @@ FVertexBufferRHIRef FPositionVertexBuffer::CreateRHIBuffer_Async()
 
 void FPositionVertexBuffer::InitRHI()
 {
-	if (!bStreamed)
-	{
-		VertexBufferRHI = CreateRHIBuffer_RenderThread();
-	}
+	VertexBufferRHI = CreateRHIBuffer_RenderThread();
 	// we have decide to create the SRV based on GMaxRHIShaderPlatform because this is created once and shared between feature levels for editor preview.
-	if ((bStreamed || VertexBufferRHI) && (RHISupportsManualVertexFetch(GMaxRHIShaderPlatform) || IsGPUSkinCacheAvailable()))
+	if (VertexBufferRHI && (RHISupportsManualVertexFetch(GMaxRHIShaderPlatform) || IsGPUSkinCacheAvailable()))
 	{
-		PositionComponentSRV = RHICreateShaderResourceView(VertexBufferRHI, 4, PF_R32_FLOAT);
+		PositionComponentSRV = RHICreateShaderResourceView(VertexData ? VertexBufferRHI : nullptr, 4, PF_R32_FLOAT);
 	}
 }
 
diff --git a/Engine/Source/Runtime/Engine/Private/RawIndexBuffer.cpp b/Engine/Source/Runtime/Engine/Private/RawIndexBuffer.cpp
--- a/Engine/Source/Runtime/Engine/Private/RawIndexBuffer.cpp
+++ b/Engine/Source/Runtime/Engine/Private/RawIndexBuffer.cpp
@@ -142,7 +142,6 @@ FRawStaticIndexBuffer::FRawStaticIndexBuffer(bool InNeedsCPUAccess)
 	: IndexStorage(InNeedsCPUAccess)
 	, CachedNumIndices(-1)
 	, b32Bit(false)
-	, bStreamed(false)
 {
 }
 
@@ -288,10 +287,11 @@ FIndexBufferRHIRef FRawStaticIndexBuffer::CreateRHIBuffer_Internal()
 	const uint32 IndexStride = b32Bit ? sizeof(uint32) : sizeof(uint16);
 	const uint32 SizeInBytes = IndexStorage.Num();
 
-	if (SizeInBytes > 0)
+	if (GetNumIndices() > 0)
 	{
 		// Create the index buffer.
 		FRHIResourceCreateInfo CreateInfo(&IndexStorage);
+		CreateInfo.bWithoutNativeResource = !SizeInBytes;
 		if (bRenderThread)
 		{
 			return RHICreateIndexBuffer(IndexStride, SizeInBytes, BUF_Static, CreateInfo);
@@ -314,18 +314,9 @@ FIndexBufferRHIRef FRawStaticIndexBuffer::CreateRHIBuffer_Async()
 	return CreateRHIBuffer_Internal<false>();
 }
 
-void FRawStaticIndexBuffer::InitRHIForStreaming(FIndexBufferRHIParamRef IntermediateBuffer)
-{
-	check(!IndexBufferRHI);
-	IndexBufferRHI = IntermediateBuffer;
-}
-
 void FRawStaticIndexBuffer::InitRHI()
 {
-	if (!bStreamed)
-	{
-		IndexBufferRHI = CreateRHIBuffer_RenderThread();
-	}
+	IndexBufferRHI = CreateRHIBuffer_RenderThread();
 }
 
 void FRawStaticIndexBuffer::Serialize(FArchive& Ar, bool bNeedsCPUAccess)
@@ -376,6 +367,7 @@ void FRawStaticIndexBuffer16or32<INDEX_TYPE>::CacheOptimize()
 {
 #if WITH_EDITOR
 	CacheOptimizeIndexBuffer(Indices);
+	CachedNumIndices = Indices.Num();
 #endif
 }
 
diff --git a/Engine/Source/Runtime/Engine/Private/Rendering/ColorVertexBuffer.cpp b/Engine/Source/Runtime/Engine/Private/Rendering/ColorVertexBuffer.cpp
--- a/Engine/Source/Runtime/Engine/Private/Rendering/ColorVertexBuffer.cpp
+++ b/Engine/Source/Runtime/Engine/Private/Rendering/ColorVertexBuffer.cpp
@@ -27,8 +27,7 @@ FColorVertexBuffer::FColorVertexBuffer():
 	VertexData(NULL),
 	Data(NULL),
 	Stride(0),
-	NumVertices(0),
-	bStreamed(false)
+	NumVertices(0)
 {
 }
 
@@ -384,21 +383,19 @@ uint32 FColorVertexBuffer::GetAllocatedSize() const
 template <bool bRenderThread>
 FVertexBufferRHIRef FColorVertexBuffer::CreateRHIBuffer_Internal()
 {
-	if (VertexData != NULL)
+	if (NumVertices)
 	{
-		FResourceArrayInterface* ResourceArray = VertexData->GetResourceArray();
-		if (ResourceArray->GetResourceDataSize())
+		FResourceArrayInterface* RESTRICT ResourceArray = VertexData ? VertexData->GetResourceArray() : nullptr;
+		const uint32 SizeInBytes = ResourceArray ? ResourceArray->GetResourceDataSize() : 0;
+		FRHIResourceCreateInfo CreateInfo(ResourceArray);
+		CreateInfo.bWithoutNativeResource = !VertexData;
+		if (bRenderThread)
 		{
-			// Create the vertex buffer.
-			FRHIResourceCreateInfo CreateInfo(ResourceArray);
-			if (bRenderThread)
-			{
-				return RHICreateVertexBuffer(ResourceArray->GetResourceDataSize(), BUF_Static | BUF_ShaderResource, CreateInfo);
-			}
-			else
-			{
-				return RHIAsyncCreateVertexBuffer(ResourceArray->GetResourceDataSize(), BUF_Static | BUF_ShaderResource, CreateInfo);
-			}
+			return RHICreateVertexBuffer(SizeInBytes, BUF_Static | BUF_ShaderResource, CreateInfo);
+		}
+		else
+		{
+			return RHIAsyncCreateVertexBuffer(SizeInBytes, BUF_Static | BUF_ShaderResource, CreateInfo);
 		}
 	}
 	return nullptr;
@@ -416,13 +413,10 @@ FVertexBufferRHIRef FColorVertexBuffer::CreateRHIBuffer_Async()
 
 void FColorVertexBuffer::InitRHI()
 {
-	if (!bStreamed)
-	{
-		VertexBufferRHI = CreateRHIBuffer_RenderThread();
-	}
-	if (bStreamed || VertexBufferRHI)
+	VertexBufferRHI = CreateRHIBuffer_RenderThread();
+	if (VertexBufferRHI)
 	{
-		ColorComponentsSRV = RHICreateShaderResourceView(VertexBufferRHI, 4, PF_R8G8B8A8);
+		ColorComponentsSRV = RHICreateShaderResourceView(VertexData ? VertexBufferRHI : nullptr, 4, PF_R8G8B8A8);
 	}
 }
 
diff --git a/Engine/Source/Runtime/Engine/Private/Rendering/StaticMeshVertexBuffer.cpp b/Engine/Source/Runtime/Engine/Private/Rendering/StaticMeshVertexBuffer.cpp
--- a/Engine/Source/Runtime/Engine/Private/Rendering/StaticMeshVertexBuffer.cpp
+++ b/Engine/Source/Runtime/Engine/Private/Rendering/StaticMeshVertexBuffer.cpp
@@ -13,8 +13,7 @@ FStaticMeshVertexBuffer::FStaticMeshVertexBuffer() :
 	NumTexCoords(0),
 	NumVertices(0),
 	bUseFullPrecisionUVs(!GVertexElementTypeSupport.IsSupported(VET_Half2)),
-	bUseHighPrecisionTangentBasis(false),
-	bStreamed(false)
+	bUseHighPrecisionTangentBasis(false)
 {}
 
 FStaticMeshVertexBuffer::~FStaticMeshVertexBuffer()
@@ -261,19 +260,19 @@ void FStaticMeshVertexBuffer::operator=(const FStaticMeshVertexBuffer &Other)
 template <bool bRenderThread>
 FVertexBufferRHIRef FStaticMeshVertexBuffer::CreateTangentsRHIBuffer_Internal()
 {
-	check(TangentsData);
-	FResourceArrayInterface* ResourceArray = TangentsData->GetResourceArray();
-	if (ResourceArray->GetResourceDataSize())
+	if (GetNumVertices())
 	{
-		// Create the vertex buffer.
+		FResourceArrayInterface* RESTRICT ResourceArray = TangentsData ? TangentsData->GetResourceArray() : nullptr;
+		const uint32 SizeInBytes = ResourceArray ? ResourceArray->GetResourceDataSize() : 0;
 		FRHIResourceCreateInfo CreateInfo(ResourceArray);
+		CreateInfo.bWithoutNativeResource = !TangentsData;
 		if (bRenderThread)
 		{
-			return RHICreateVertexBuffer(ResourceArray->GetResourceDataSize(), BUF_Static | BUF_ShaderResource, CreateInfo);
+			return RHICreateVertexBuffer(SizeInBytes, BUF_Static | BUF_ShaderResource, CreateInfo);
 		}
 		else
 		{
-			return RHIAsyncCreateVertexBuffer(ResourceArray->GetResourceDataSize(), BUF_Static | BUF_ShaderResource, CreateInfo);
+			return RHIAsyncCreateVertexBuffer(SizeInBytes, BUF_Static | BUF_ShaderResource, CreateInfo);
 		}
 	}
 	return nullptr;
@@ -292,19 +291,19 @@ FVertexBufferRHIRef FStaticMeshVertexBuffer::CreateTangentsRHIBuffer_Async()
 template <bool bRenderThread>
 FVertexBufferRHIRef FStaticMeshVertexBuffer::CreateTexCoordRHIBuffer_Internal()
 {
-	check(TexcoordData);
-	FResourceArrayInterface* ResourceArray = TexcoordData->GetResourceArray();
-	if (ResourceArray->GetResourceDataSize())
+	if (GetNumTexCoords())
 	{
-		// Create the vertex buffer.
+		FResourceArrayInterface* RESTRICT ResourceArray = TexcoordData ? TexcoordData->GetResourceArray() : nullptr;
+		const uint32 SizeInBytes = ResourceArray ? ResourceArray->GetResourceDataSize() : 0;
 		FRHIResourceCreateInfo CreateInfo(ResourceArray);
+		CreateInfo.bWithoutNativeResource = !TexcoordData;
 		if (bRenderThread)
 		{
-			return RHICreateVertexBuffer(ResourceArray->GetResourceDataSize(), BUF_Static | BUF_ShaderResource, CreateInfo);
+			return RHICreateVertexBuffer(SizeInBytes, BUF_Static | BUF_ShaderResource, CreateInfo);
 		}
 		else
 		{
-			return RHIAsyncCreateVertexBuffer(ResourceArray->GetResourceDataSize(), BUF_Static | BUF_ShaderResource, CreateInfo);
+			return RHIAsyncCreateVertexBuffer(SizeInBytes, BUF_Static | BUF_ShaderResource, CreateInfo);
 		}
 	}
 	return nullptr;
@@ -322,22 +321,19 @@ FVertexBufferRHIRef FStaticMeshVertexBuffer::CreateTexCoordRHIBuffer_Async()
 
 void FStaticMeshVertexBuffer::InitRHI()
 {
-	if (!bStreamed)
-	{
-		TangentsVertexBuffer.VertexBufferRHI = CreateTangentsRHIBuffer_RenderThread();
-		TexCoordVertexBuffer.VertexBufferRHI = CreateTexCoordRHIBuffer_RenderThread();
-	}
-	if ((bStreamed || TangentsVertexBuffer.VertexBufferRHI) && (RHISupportsManualVertexFetch(GMaxRHIShaderPlatform) || IsGPUSkinCacheAvailable()))
+	TangentsVertexBuffer.VertexBufferRHI = CreateTangentsRHIBuffer_RenderThread();
+	TexCoordVertexBuffer.VertexBufferRHI = CreateTexCoordRHIBuffer_RenderThread();
+	if (TangentsVertexBuffer.VertexBufferRHI && (RHISupportsManualVertexFetch(GMaxRHIShaderPlatform) || IsGPUSkinCacheAvailable()))
 	{
 		TangentsSRV = RHICreateShaderResourceView(
-			TangentsVertexBuffer.VertexBufferRHI,
+			TangentsData ? TangentsVertexBuffer.VertexBufferRHI : nullptr,
 			GetUseHighPrecisionTangentBasis() ? 8 : 4,
 			GetUseHighPrecisionTangentBasis() ? PF_R16G16B16A16_SNORM : PF_R8G8B8A8_SNORM);
 	}
-	if ((bStreamed || TexCoordVertexBuffer.VertexBufferRHI) && RHISupportsManualVertexFetch(GMaxRHIShaderPlatform))
+	if (TexCoordVertexBuffer.VertexBufferRHI && RHISupportsManualVertexFetch(GMaxRHIShaderPlatform))
 	{
 		TextureCoordinatesSRV = RHICreateShaderResourceView(
-			TexCoordVertexBuffer.VertexBufferRHI,
+			TexcoordData ? TexCoordVertexBuffer.VertexBufferRHI : nullptr,
 			GetUseFullPrecisionUVs() ? 8 : 4,
 			GetUseFullPrecisionUVs() ? PF_G32R32F : PF_G16R16F);
 	}
diff --git a/Engine/Source/Runtime/Engine/Private/StaticMesh.cpp b/Engine/Source/Runtime/Engine/Private/StaticMesh.cpp
--- a/Engine/Source/Runtime/Engine/Private/StaticMesh.cpp
+++ b/Engine/Source/Runtime/Engine/Private/StaticMesh.cpp
@@ -363,11 +363,6 @@ void FStaticMeshLODResources::SerializeBuffers(FArchive& Ar, UStaticMesh* OwnerS
 		}
 		SerializedAdditionalIndexBuffers = AdditionalIndexBuffers;
 	}
-	
-	SerializedAdditionalIndexBuffers->ReversedIndexBuffer.SetIsStreamed(!bBuffersInlined);
-	SerializedAdditionalIndexBuffers->ReversedDepthOnlyIndexBuffer.SetIsStreamed(!bBuffersInlined);
-	SerializedAdditionalIndexBuffers->WireframeIndexBuffer.SetIsStreamed(!bBuffersInlined);
-	SerializedAdditionalIndexBuffers->AdjacencyIndexBuffer.SetIsStreamed(!bBuffersInlined);
 
 	if (bSerializeReversedIndexBuffer)
 	{
@@ -475,12 +470,6 @@ void FStaticMeshLODResources::SerializeAvailabilityInfo(FArchive& Ar)
 		SerializedAdditionalIndexBuffers = AdditionalIndexBuffers;
 	}
 
-	SerializedAdditionalIndexBuffers->ReversedIndexBuffer.SetIsStreamed(!bBuffersInlined);
-	SerializedAdditionalIndexBuffers->ReversedDepthOnlyIndexBuffer.SetIsStreamed(!bBuffersInlined);
-	SerializedAdditionalIndexBuffers->WireframeIndexBuffer.SetIsStreamed(!bBuffersInlined);
-	SerializedAdditionalIndexBuffers->AdjacencyIndexBuffer.SetIsStreamed(!bBuffersInlined);
-	
-
 	SerializedAdditionalIndexBuffers->ReversedIndexBuffer.SerializeMetaData(Ar);
 	if (!bHasReversedIndices)
 	{
@@ -530,12 +519,6 @@ void FStaticMeshLODResources::Serialize(FArchive& Ar, UObject* Owner, int32 Inde
 	Ar << bInlined;
 	bBuffersInlined = bInlined;
 
-	VertexBuffers.StaticMeshVertexBuffer.SetIsStreamed(!bBuffersInlined);
-	VertexBuffers.PositionVertexBuffer.SetIsStreamed(!bBuffersInlined);
-	VertexBuffers.ColorVertexBuffer.SetIsStreamed(!bBuffersInlined);
-	IndexBuffer.SetIsStreamed(!bBuffersInlined);
-	DepthOnlyIndexBuffer.SetIsStreamed(!bBuffersInlined);
-
 	if (!StripFlags.IsDataStrippedForServer() && !bIsLODCookedOut)
 	{
 		FStaticMeshBuffersSize TmpBuffersSize;
@@ -1876,7 +1859,7 @@ static void SerializeBuildSettingsForDDC(FArchive& Ar, FMeshBuildSettings& Build
 // differences, etc.) replace the version GUID below with a new one.
 // In case of merge conflicts with DDC versions, you MUST generate a new GUID
 // and set this new GUID as the version.                                       
-#define STATICMESH_DERIVEDDATA_VER TEXT("9891BF2FF72141F6921E24E347DB73C1")
+#define STATICMESH_DERIVEDDATA_VER TEXT("B5FB810437E4428D9CC6367AE010BEEC")
 
 static const FString& GetStaticMeshDerivedDataVersion()
 {
@@ -1993,7 +1976,8 @@ static FString BuildStaticMeshDerivedDataKeySuffix(UStaticMesh* Mesh, const FSta
 	const ITargetPlatform* RunningPlatform = GetTargetPlatformManagerRef().GetRunningTargetPlatform();
 	check(RunningPlatform);
 	const bool bAllowLODStreaming = RunningPlatform->SupportsFeature(ETargetPlatformFeatures::MeshLODStreaming) && LODGroup.IsLODStreamingSupported();
-	KeySuffix.AppendChar(bAllowLODStreaming ? TEXT('1') : TEXT('0'));
+	KeySuffix += bAllowLODStreaming ? TEXT("LS1") : TEXT("LS0");
+	KeySuffix += TEXT("MNS");
 	if (bAllowLODStreaming)
 	{
 		int32 MaxNumStreamedLODs = Mesh->NumStreamedLODs.GetValueForPlatformIdentifiers(
@@ -2174,11 +2158,8 @@ void FStaticMeshRenderData::Cache(UStaticMesh* Owner, const FStaticMeshLODSettin
 				}
 				// TODO: can we postpone the loading to streaming time?
 				LODResource.DerivedDataKey = BuildStaticMeshLODDerivedDataKey(KeySuffix, LODIdx);
-				TArray<uint8> LODDerivedData;
-				verify(GetDerivedDataCacheRef().GetSynchronous(*LODResource.DerivedDataKey, LODDerivedData));
-				FMemoryReader LODAr(LODDerivedData, true);
 				typename FStaticMeshLODResources::FStaticMeshBuffersSize DummyBuffersSize;
-				LODResource.SerializeBuffers(LODAr, Owner, 0, DummyBuffersSize);
+				LODResource.SerializeBuffers(Ar, Owner, 0, DummyBuffersSize);
 				typename FStaticMeshLODResources::FStaticMeshBuffersSize LODBuffersSize;
 				Ar << LODBuffersSize;
 				LODResource.BuffersSize = LODBuffersSize.CalcBuffersSize();
@@ -2223,14 +2204,12 @@ void FStaticMeshRenderData::Cache(UStaticMesh* Owner, const FStaticMeshLODSettin
 				{
 					break;
 				}
-				TArray<uint8> LODDerivedData;
-				FMemoryWriter LODAr(LODDerivedData, true);
 				typename FStaticMeshLODResources::FStaticMeshBuffersSize LODBuffersSize;
-				const uint8 LODStripFlags = FStaticMeshLODResources::GenerateClassStripFlags(LODAr, Owner, LODIdx);
-				LODResource.SerializeBuffers(LODAr, Owner, LODStripFlags, LODBuffersSize);
+				const uint8 LODStripFlags = FStaticMeshLODResources::GenerateClassStripFlags(Ar, Owner, LODIdx);
+				LODResource.SerializeBuffers(Ar, Owner, LODStripFlags, LODBuffersSize);
 				Ar << LODBuffersSize;
 				LODResource.DerivedDataKey = BuildStaticMeshLODDerivedDataKey(KeySuffix, LODIdx);
-				GetDerivedDataCacheRef().Put(*LODResource.DerivedDataKey, LODDerivedData);
+				// TODO: Save non-inlined LODs separately
 			}
 
 			GetDerivedDataCacheRef().Put(*DerivedDataKey, DerivedData);
@@ -4684,6 +4663,7 @@ int32 UStaticMesh::CalcCumulativeLODSize(int32 NumLODs) const
 bool UStaticMesh::GetMipDataFilename(const int32 MipIndex, FString& OutBulkDataFilename) const
 {
 #if !WITH_EDITOR
+	// TODO: this is slow. Should cache the name once per mesh
 	FString PackageName = GetOutermost()->FileName.ToString();
 	// Handle name redirection and localization
 	const FCoreRedirectObjectName RedirectedName =
diff --git a/Engine/Source/Runtime/Engine/Private/Streaming/LevelTextureManager.cpp b/Engine/Source/Runtime/Engine/Private/Streaming/LevelTextureManager.cpp
--- a/Engine/Source/Runtime/Engine/Private/Streaming/LevelTextureManager.cpp
+++ b/Engine/Source/Runtime/Engine/Private/Streaming/LevelTextureManager.cpp
@@ -73,7 +73,8 @@ float FLevelRenderAssetManager::GetWorldTime() const
 		if (World && !World->IsPaused())
 		{
 			// In the editor, we only return a time for the PIE world.
-			if (!GIsEditor || World->IsPlayInEditor())
+			// TODO: figure out why there are more than one PIE world
+			if (!GIsEditor || (World->IsPlayInEditor() && World->Scene->GetFrameNumber()))
 			{
 				return World->GetTimeSeconds();
 			}
diff --git a/Engine/Source/Runtime/Engine/Private/Streaming/StaticMeshUpdate.cpp b/Engine/Source/Runtime/Engine/Private/Streaming/StaticMeshUpdate.cpp
--- a/Engine/Source/Runtime/Engine/Private/Streaming/StaticMeshUpdate.cpp
+++ b/Engine/Source/Runtime/Engine/Private/Streaming/StaticMeshUpdate.cpp
@@ -11,6 +11,9 @@ StaticMeshUpdate.cpp: Helpers to stream in and out static mesh LODs.
 #include "HAL/PlatformFilemanager.h"
 #include "Serialization/MemoryReader.h"
 
+static constexpr uint32 GMaxNumResourceUpdatesPerLOD = 14;
+static constexpr uint32 GMaxNumResourceUpdatesPerBatch = (MAX_STATIC_MESH_LODS - 1) * GMaxNumResourceUpdatesPerLOD;
+
 FStaticMeshUpdateContext::FStaticMeshUpdateContext(UStaticMesh* InMesh, EThreadType InCurrentThread)
 	: Mesh(InMesh)
 	, CurrentThread(InCurrentThread)
@@ -100,22 +103,22 @@ void FStaticMeshStreamIn::FIntermediateBuffers::SafeRelease()
 	AdjacencyIndexBuffer.SafeRelease();
 }
 
-template <int32 MaxNumUpdates>
+template <uint32 MaxNumUpdates>
 void FStaticMeshStreamIn::FIntermediateBuffers::TransferBuffers(FStaticMeshLODResources& LODResource, TRHIResourceUpdateBatcher<MaxNumUpdates>& Batcher)
 {
 	FStaticMeshVertexBuffers& VBs = LODResource.VertexBuffers;
 	VBs.StaticMeshVertexBuffer.InitRHIForStreaming(TangentsVertexBuffer, TexCoordVertexBuffer, Batcher);
 	VBs.PositionVertexBuffer.InitRHIForStreaming(PositionVertexBuffer, Batcher);
 	VBs.ColorVertexBuffer.InitRHIForStreaming(ColorVertexBuffer, Batcher);
-	LODResource.IndexBuffer.InitRHIForStreaming(IndexBuffer);
-	LODResource.DepthOnlyIndexBuffer.InitRHIForStreaming(DepthOnlyIndexBuffer);
+	LODResource.IndexBuffer.InitRHIForStreaming(IndexBuffer, Batcher);
+	LODResource.DepthOnlyIndexBuffer.InitRHIForStreaming(DepthOnlyIndexBuffer, Batcher);
 
 	if (LODResource.AdditionalIndexBuffers)
 	{
-		LODResource.AdditionalIndexBuffers->ReversedIndexBuffer.InitRHIForStreaming(ReversedIndexBuffer);
-		LODResource.AdditionalIndexBuffers->ReversedDepthOnlyIndexBuffer.InitRHIForStreaming(ReversedDepthOnlyIndexBuffer);
-		LODResource.AdditionalIndexBuffers->WireframeIndexBuffer.InitRHIForStreaming(WireframeIndexBuffer);
-		LODResource.AdditionalIndexBuffers->AdjacencyIndexBuffer.InitRHIForStreaming(AdjacencyIndexBuffer);
+		LODResource.AdditionalIndexBuffers->ReversedIndexBuffer.InitRHIForStreaming(ReversedIndexBuffer, Batcher);
+		LODResource.AdditionalIndexBuffers->ReversedDepthOnlyIndexBuffer.InitRHIForStreaming(ReversedDepthOnlyIndexBuffer, Batcher);
+		LODResource.AdditionalIndexBuffers->WireframeIndexBuffer.InitRHIForStreaming(WireframeIndexBuffer, Batcher);
+		LODResource.AdditionalIndexBuffers->AdjacencyIndexBuffer.InitRHIForStreaming(AdjacencyIndexBuffer, Batcher);
 	}
 	SafeRelease();
 }
@@ -210,9 +213,7 @@ void FStaticMeshStreamIn::DoFinishUpdate(const FContext& Context)
 			&& PendingFirstMip < CurrentFirstLODIdx);
 		// Use a scope to flush the batcher before updating CurrentFirstLODIdx
 		{
-			constexpr int32 MaxNumUpdatesPerLOD = 4;
-			constexpr int32 MaxNumUpdates = (MAX_STATIC_MESH_LODS - 1) * MaxNumUpdatesPerLOD;
-			TRHIResourceUpdateBatcher<MaxNumUpdates> Batcher;
+			TRHIResourceUpdateBatcher<GMaxNumResourceUpdatesPerBatch> Batcher;
 
 			for (int32 LODIdx = PendingFirstMip; LODIdx < CurrentFirstLODIdx; ++LODIdx)
 			{
@@ -262,10 +263,8 @@ void FStaticMeshStreamOut::DoReleaseBuffers(const FContext& Context)
 
 		RenderData->CurrentFirstLODIdx = PendingFirstMip;
 		Mesh->SetCachedNumResidentLODs(static_cast<uint8>(RenderData->LODResources.Num() - PendingFirstMip));
-		
-		constexpr int32 MaxNumUpdatesPerLOD = 4;
-		constexpr int32 MaxNumUpdates = (MAX_STATIC_MESH_LODS - 1) * MaxNumUpdatesPerLOD;
-		TRHIResourceUpdateBatcher<MaxNumUpdates> Batcher;
+
+		TRHIResourceUpdateBatcher<GMaxNumResourceUpdatesPerBatch> Batcher;
 
 		for (int32 LODIdx = CurrentFirstLODIdx; LODIdx < PendingFirstMip; ++LODIdx)
 		{
@@ -276,15 +275,15 @@ void FStaticMeshStreamOut::DoReleaseBuffers(const FContext& Context)
 			VBs.PositionVertexBuffer.ReleaseRHIForStreaming(Batcher);
 			VBs.ColorVertexBuffer.ReleaseRHIForStreaming(Batcher);
 			// Index buffers don't need to update SRV so we can reuse ReleaseRHI
-			LODResource.IndexBuffer.ReleaseRHI();
-			LODResource.DepthOnlyIndexBuffer.ReleaseRHI();
+			LODResource.IndexBuffer.ReleaseRHIForStreaming(Batcher);
+			LODResource.DepthOnlyIndexBuffer.ReleaseRHIForStreaming(Batcher);
 
 			if (LODResource.AdditionalIndexBuffers)
 			{
-				LODResource.AdditionalIndexBuffers->ReversedIndexBuffer.ReleaseRHI();
-				LODResource.AdditionalIndexBuffers->ReversedDepthOnlyIndexBuffer.ReleaseRHI();
-				LODResource.AdditionalIndexBuffers->WireframeIndexBuffer.ReleaseRHI();
-				LODResource.AdditionalIndexBuffers->AdjacencyIndexBuffer.ReleaseRHI();
+				LODResource.AdditionalIndexBuffers->ReversedIndexBuffer.ReleaseRHIForStreaming(Batcher);
+				LODResource.AdditionalIndexBuffers->ReversedDepthOnlyIndexBuffer.ReleaseRHIForStreaming(Batcher);
+				LODResource.AdditionalIndexBuffers->WireframeIndexBuffer.ReleaseRHIForStreaming(Batcher);
+				LODResource.AdditionalIndexBuffers->AdjacencyIndexBuffer.ReleaseRHIForStreaming(Batcher);
 			}
 		}
 	}
diff --git a/Engine/Source/Runtime/Engine/Private/Streaming/StaticMeshUpdate.h b/Engine/Source/Runtime/Engine/Private/Streaming/StaticMeshUpdate.h
--- a/Engine/Source/Runtime/Engine/Private/Streaming/StaticMeshUpdate.h
+++ b/Engine/Source/Runtime/Engine/Private/Streaming/StaticMeshUpdate.h
@@ -97,7 +97,7 @@ protected:
 		void SafeRelease();
 
 		/** Transfer ownership of buffers to a LOD resource */
-		template <int32 MaxNumUpdates>
+		template <uint32 MaxNumUpdates>
 		void TransferBuffers(FStaticMeshLODResources& LODResource, TRHIResourceUpdateBatcher<MaxNumUpdates>& Batcher);
 
 		void CheckIsNull() const;
diff --git a/Engine/Source/Runtime/Engine/Public/MeshParticleVertexFactory.h b/Engine/Source/Runtime/Engine/Public/MeshParticleVertexFactory.h
--- a/Engine/Source/Runtime/Engine/Public/MeshParticleVertexFactory.h
+++ b/Engine/Source/Runtime/Engine/Public/MeshParticleVertexFactory.h
@@ -84,6 +84,7 @@ public:
 	/** Default constructor. */
 	FMeshParticleVertexFactory(EParticleVertexFactoryType InType, ERHIFeatureLevel::Type InFeatureLevel, int32 InDynamicVertexStride, int32 InDynamicParameterVertexStride)
 		: FParticleVertexFactoryBase(InType, InFeatureLevel)
+		, LODIdx(0xff)
 		, DynamicVertexStride(InDynamicVertexStride)
 		, DynamicParameterVertexStride(InDynamicParameterVertexStride)
 		, InstanceVerticesCPU(nullptr)
@@ -91,6 +92,7 @@ public:
 
 	FMeshParticleVertexFactory()
 		: FParticleVertexFactoryBase(PVFT_MAX, ERHIFeatureLevel::Num)
+		, LODIdx(0xff)
 		, DynamicVertexStride(-1)
 		, DynamicParameterVertexStride(-1)
 		, InstanceVerticesCPU(nullptr)
@@ -184,9 +186,21 @@ public:
 		return InstanceVerticesCPU;
 	}
 
+	void SetLODIdx(uint8 InLODIdx)
+	{
+		LODIdx = InLODIdx;
+	}
+
+	uint8 GetLODIdx() const
+	{
+		return LODIdx;
+	}
+
 protected:
 	FDataType Data;
 	
+	uint8 LODIdx;
+
 	/** Stride information for instanced mesh particles */
 	int32 DynamicVertexStride;
 	int32 DynamicParameterVertexStride;
diff --git a/Engine/Source/Runtime/Engine/Public/ParticleHelper.h b/Engine/Source/Runtime/Engine/Public/ParticleHelper.h
--- a/Engine/Source/Runtime/Engine/Public/ParticleHelper.h
+++ b/Engine/Source/Runtime/Engine/Public/ParticleHelper.h
@@ -1923,7 +1923,7 @@ struct FDynamicMeshEmitterData : public FDynamicSpriteEmitterDataBase
 	/**
 	 *	 Initialize this emitter's vertex factory with the vertex buffers from the mesh's rendering data.
 	 */
-	void SetupVertexFactory( FMeshParticleVertexFactory* InVertexFactory, FStaticMeshLODResources& LODResources) const;
+	void SetupVertexFactory( FMeshParticleVertexFactory* InVertexFactory, const FStaticMeshLODResources& LODResources, uint32 LODIdx) const;
 
 	/** Returns the source data for this particle system */
 	virtual const FDynamicEmitterReplayDataBase& GetSource() const override
diff --git a/Engine/Source/Runtime/Engine/Public/RawIndexBuffer.h b/Engine/Source/Runtime/Engine/Public/RawIndexBuffer.h
--- a/Engine/Source/Runtime/Engine/Public/RawIndexBuffer.h
+++ b/Engine/Source/Runtime/Engine/Public/RawIndexBuffer.h
@@ -244,11 +244,25 @@ public:
 	FIndexBufferRHIRef CreateRHIBuffer_RenderThread();
 	FIndexBufferRHIRef CreateRHIBuffer_Async();
 
-	/** Set whether this buffer is managed by the streamer. Must set before InitRHI is called */
-	void SetIsStreamed(bool bValue) { bStreamed = bValue; }
-
 	/** Take over ownership of IntermediateBuffer */
-	void InitRHIForStreaming(FIndexBufferRHIParamRef IntermediateBuffer);
+	template <uint32 MaxNumUpdates>
+	void InitRHIForStreaming(FIndexBufferRHIParamRef IntermediateBuffer, TRHIResourceUpdateBatcher<MaxNumUpdates>& Batcher)
+	{
+		if (IndexBufferRHI && IntermediateBuffer)
+		{
+			Batcher.QueueUpdateRequest(IndexBufferRHI, IntermediateBuffer);
+		}
+	}
+
+	/** Release any GPU resource owned by the RHI object */
+	template <uint32 MaxNumUpdates>
+	void ReleaseRHIForStreaming(TRHIResourceUpdateBatcher<MaxNumUpdates>& Batcher)
+	{
+		if (IndexBufferRHI)
+		{
+			Batcher.QueueUpdateRequest(IndexBufferRHI, nullptr);
+		}
+	}
 
 	/**
 	 * Serialization.
@@ -283,8 +297,6 @@ private:
 
 	/** 32bit or 16bit? */
 	bool b32Bit;
-
-	bool bStreamed;
 };
 
 /**
@@ -341,17 +353,17 @@ public:
 	virtual void InitRHI() override
 	{
 		uint32 Size = Indices.Num() * sizeof(INDEX_TYPE);
-		if(Indices.Num())
+		if (Indices.Num())
 		{
 			// Create the index buffer.
 			FRHIResourceCreateInfo CreateInfo(&Indices);
-			
+
 			extern ENGINE_API bool DoSkeletalMeshIndexBuffersNeedSRV();
 			bool bSRV = DoSkeletalMeshIndexBuffersNeedSRV();
 
 			EBufferUsageFlags Flags = BUF_Static;
 
-			if(bSRV)
+			if (bSRV)
 			{
 				// BUF_ShaderResource is needed for SkinCache RecomputeSkinTangents
 				Flags = (EBufferUsageFlags)(Flags | BUF_ShaderResource);
@@ -359,11 +371,11 @@ public:
 
 			IndexBufferRHI = RHICreateIndexBuffer(sizeof(INDEX_TYPE), Size, Flags, CreateInfo);
 
-			if(bSRV)
+			if (bSRV)
 			{
 				SRVValue = RHICreateShaderResourceView(IndexBufferRHI);
 			}
-		}    
+		}
 	}
 	
 	virtual void ReleaseRHI() override
@@ -380,7 +392,7 @@ public:
 	*/
 	virtual void Serialize( FArchive& Ar ) override
 	{
-			Indices.BulkSerialize( Ar );
+		Indices.BulkSerialize(Ar);
 	}
 
 	/**
diff --git a/Engine/Source/Runtime/Engine/Public/Rendering/ColorVertexBuffer.h b/Engine/Source/Runtime/Engine/Public/Rendering/ColorVertexBuffer.h
--- a/Engine/Source/Runtime/Engine/Public/Rendering/ColorVertexBuffer.h
+++ b/Engine/Source/Runtime/Engine/Public/Rendering/ColorVertexBuffer.h
@@ -138,30 +138,29 @@ public:
 	FVertexBufferRHIRef CreateRHIBuffer_RenderThread();
 	FVertexBufferRHIRef CreateRHIBuffer_Async();
 
-	/** Set whether this buffer is managed by the streamer. Must be set before InitRHI is called */
-	void SetIsStreamed(bool bValue) { bStreamed = bValue; }
-
 	/** Similar to Init/ReleaseRHI but only update existing SRV so references to the SRV stays valid */
-	template <int32 MaxNumUpdates>
+	template <uint32 MaxNumUpdates>
 	void InitRHIForStreaming(FVertexBufferRHIParamRef IntermediateBuffer, TRHIResourceUpdateBatcher<MaxNumUpdates>& Batcher)
 	{
-		check(!VertexBufferRHI);
-		if (IntermediateBuffer)
+		if (VertexBufferRHI && IntermediateBuffer)
 		{
 			check(ColorComponentsSRV);
-			VertexBufferRHI = IntermediateBuffer;
+			Batcher.QueueUpdateRequest(VertexBufferRHI, IntermediateBuffer);
 			Batcher.QueueUpdateRequest(ColorComponentsSRV, VertexBufferRHI, 4, PF_R8G8B8A8);
 		}
 	}
 
-	template <int32 MaxNumUpdates>
+	template <uint32 MaxNumUpdates>
 	void ReleaseRHIForStreaming(TRHIResourceUpdateBatcher<MaxNumUpdates>& Batcher)
 	{
+		if (VertexBufferRHI)
+		{
+			Batcher.QueueUpdateRequest(VertexBufferRHI, nullptr);
+		}
 		if (ColorComponentsSRV)
 		{
 			Batcher.QueueUpdateRequest(ColorComponentsSRV, nullptr, 0, 0);
 		}
-		FVertexBuffer::ReleaseRHI();
 	}
 
 	// FRenderResource interface.
@@ -200,8 +199,6 @@ private:
 
 	bool NeedsCPUAccess = true;
 
-	bool bStreamed;
-
 	/** Allocates the vertex data storage type. */
 	void AllocateData(bool bNeedsCPUAccess = true);
 
diff --git a/Engine/Source/Runtime/Engine/Public/Rendering/PositionVertexBuffer.h b/Engine/Source/Runtime/Engine/Public/Rendering/PositionVertexBuffer.h
--- a/Engine/Source/Runtime/Engine/Public/Rendering/PositionVertexBuffer.h
+++ b/Engine/Source/Runtime/Engine/Public/Rendering/PositionVertexBuffer.h
@@ -96,17 +96,14 @@ public:
 	FVertexBufferRHIRef CreateRHIBuffer_RenderThread();
 	FVertexBufferRHIRef CreateRHIBuffer_Async();
 
-	/** Set whether this buffer is managed by the streamer. Must be set before InitRHI is called */
-	void SetIsStreamed(bool bValue) { bStreamed = bValue; }
-
 	/** Similar to Init/ReleaseRHI but only update existing SRV so references to the SRV stays valid */
-	template <int32 MaxNumUpdates>
+	template <uint32 MaxNumUpdates>
 	void InitRHIForStreaming(FVertexBufferRHIParamRef IntermediateBuffer, TRHIResourceUpdateBatcher<MaxNumUpdates>& Batcher)
 	{
-		check(!VertexBufferRHI);
+		check(VertexBufferRHI);
 		if (IntermediateBuffer)
 		{
-			VertexBufferRHI = IntermediateBuffer;
+			Batcher.QueueUpdateRequest(VertexBufferRHI, IntermediateBuffer);
 			if (PositionComponentSRV)
 			{
 				Batcher.QueueUpdateRequest(PositionComponentSRV, VertexBufferRHI, 4, PF_R32_FLOAT);
@@ -114,14 +111,15 @@ public:
 		}
 	}
 
-	template <int32 MaxNumUpdates>
+	template <uint32 MaxNumUpdates>
 	void ReleaseRHIForStreaming(TRHIResourceUpdateBatcher<MaxNumUpdates>& Batcher)
 	{
+		check(VertexBufferRHI);
+		Batcher.QueueUpdateRequest(VertexBufferRHI, nullptr);
 		if (PositionComponentSRV)
 		{
 			Batcher.QueueUpdateRequest(PositionComponentSRV, nullptr, 0, 0);
 		}
-		FVertexBuffer::ReleaseRHI();
 	}
 
 	// FRenderResource interface.
@@ -154,8 +152,6 @@ private:
 
 	bool bNeedsCPUAccess = true;
 
-	bool bStreamed;
-
 	/** Allocates the vertex data storage type. */
 	void AllocateData(bool bInNeedsCPUAccess = true);
 
diff --git a/Engine/Source/Runtime/Engine/Public/Rendering/StaticMeshVertexBuffer.h b/Engine/Source/Runtime/Engine/Public/Rendering/StaticMeshVertexBuffer.h
--- a/Engine/Source/Runtime/Engine/Public/Rendering/StaticMeshVertexBuffer.h
+++ b/Engine/Source/Runtime/Engine/Public/Rendering/StaticMeshVertexBuffer.h
@@ -392,20 +392,17 @@ public:
 	FVertexBufferRHIRef CreateTexCoordRHIBuffer_RenderThread();
 	FVertexBufferRHIRef CreateTexCoordRHIBuffer_Async();
 
-	/** Set whether this buffer is managed by the streamer. Must be set before InitRHI is called */
-	void SetIsStreamed(bool bValue) { bStreamed = bValue; }
-
 	/** Similar to Init/ReleaseRHI but only update existing SRV so references to the SRV stays valid */
-	template <int32 MaxNumUpdates>
+	template <uint32 MaxNumUpdates>
 	void InitRHIForStreaming(
 		FVertexBufferRHIParamRef IntermediateTangentsBuffer,
 		FVertexBufferRHIParamRef IntermediateTexCoordBuffer,
 		TRHIResourceUpdateBatcher<MaxNumUpdates>& Batcher)
 	{
-		check(!TangentsVertexBuffer.VertexBufferRHI && !TexCoordVertexBuffer.VertexBufferRHI);
+		check(TangentsVertexBuffer.VertexBufferRHI && TexCoordVertexBuffer.VertexBufferRHI);
 		if (IntermediateTangentsBuffer)
 		{
-			TangentsVertexBuffer.VertexBufferRHI = IntermediateTangentsBuffer;
+			Batcher.QueueUpdateRequest(TangentsVertexBuffer.VertexBufferRHI, IntermediateTangentsBuffer);
 			if (TangentsSRV)
 			{
 				Batcher.QueueUpdateRequest(
@@ -417,7 +414,7 @@ public:
 		}
 		if (IntermediateTexCoordBuffer)
 		{
-			TexCoordVertexBuffer.VertexBufferRHI = IntermediateTexCoordBuffer;
+			Batcher.QueueUpdateRequest(TexCoordVertexBuffer.VertexBufferRHI, IntermediateTexCoordBuffer);
 			if (TextureCoordinatesSRV)
 			{
 				Batcher.QueueUpdateRequest(
@@ -429,9 +426,12 @@ public:
 		}
 	}
 
-	template<int32 MaxNumUpdates>
+	template<uint32 MaxNumUpdates>
 	void ReleaseRHIForStreaming(TRHIResourceUpdateBatcher<MaxNumUpdates>& Batcher)
 	{
+		check(TangentsVertexBuffer.VertexBufferRHI && TexCoordVertexBuffer.VertexBufferRHI);
+		Batcher.QueueUpdateRequest(TangentsVertexBuffer.VertexBufferRHI, nullptr);
+		Batcher.QueueUpdateRequest(TexCoordVertexBuffer.VertexBufferRHI, nullptr);
 		if (TangentsSRV)
 		{
 			Batcher.QueueUpdateRequest(TangentsSRV, nullptr, 0, 0);
@@ -440,8 +440,6 @@ public:
 		{
 			Batcher.QueueUpdateRequest(TextureCoordinatesSRV, nullptr, 0, 0);
 		}
-		TangentsVertexBuffer.ReleaseRHI();
-		TexCoordVertexBuffer.ReleaseRHI();
 	}
 
 	// FRenderResource interface.
@@ -525,8 +523,6 @@ private:
 
 	bool NeedsCPUAccess = true;
 
-	bool bStreamed;
-
 	/** Allocates the vertex data storage type. */
 	void AllocateData(bool bNeedsCPUAccess = true);
 
diff --git a/Engine/Source/Runtime/NullDrv/Public/NullRHI.h b/Engine/Source/Runtime/NullDrv/Public/NullRHI.h
--- a/Engine/Source/Runtime/NullDrv/Public/NullRHI.h
+++ b/Engine/Source/Runtime/NullDrv/Public/NullRHI.h
@@ -141,6 +141,11 @@ public:
 
 	}
 
+	virtual void RHITransferIndexBufferUnderlyingResource(FIndexBufferRHIParamRef DestIndexBuffer, FIndexBufferRHIParamRef SrcIndexBuffer) final override
+	{
+
+	}
+
 	virtual FVertexBufferRHIRef RHICreateVertexBuffer(uint32 Size, uint32 InUsage, FRHIResourceCreateInfo& CreateInfo) final override
 	{ 
 		if(CreateInfo.ResourceArray) 
@@ -159,6 +164,11 @@ public:
 
 	}
 
+	virtual void RHITransferVertexBufferUnderlyingResource(FVertexBufferRHIParamRef DestVertexBuffer, FVertexBufferRHIParamRef SrcVertexBuffer) final override
+	{
+
+	}
+
 
 	virtual void RHICopyVertexBuffer(FVertexBufferRHIParamRef SourceBuffer,FVertexBufferRHIParamRef DestBuffer) final override
 	{
diff --git a/Engine/Source/Runtime/RHI/Private/DynamicRHI.cpp b/Engine/Source/Runtime/RHI/Private/DynamicRHI.cpp
--- a/Engine/Source/Runtime/RHI/Private/DynamicRHI.cpp
+++ b/Engine/Source/Runtime/RHI/Private/DynamicRHI.cpp
@@ -331,3 +331,23 @@ void FDynamicRHI::EnableIdealGPUCaptureOptions(bool bEnabled)
 		RHICmdBypassVar->Set(bShouldRHICmdBypass ? 1 : 0, ECVF_SetByConsole);		
 	}	
 }
+
+void FDynamicRHI::RHITransferIndexBufferUnderlyingResource(FIndexBufferRHIParamRef DestIndexBuffer, FIndexBufferRHIParamRef SrcIndexBuffer)
+{
+	UE_LOG(LogRHI, Fatal, TEXT("RHITransferIndexBufferUnderlyingResource isn't implemented for the current RHI"));
+}
+
+void FDynamicRHI::RHITransferVertexBufferUnderlyingResource(FVertexBufferRHIParamRef DestVertexBuffer, FVertexBufferRHIParamRef SrcVertexBuffer)
+{
+	UE_LOG(LogRHI, Fatal, TEXT("RHITransferVertexBufferUnderlyingResource isn't implemented for the current RHI"));
+}
+
+void FDynamicRHI::RHIUpdateShaderResourceView(FShaderResourceViewRHIParamRef SRV, FVertexBufferRHIParamRef VertexBuffer, uint32 Stride, uint8 Format)
+{
+	UE_LOG(LogRHI, Fatal, TEXT("RHIUpdateShaderResourceView isn't implemented for the current RHI"));
+}
+
+void FDynamicRHI::RHIUpdateShaderResourceView(FShaderResourceViewRHIParamRef SRV, FIndexBufferRHIParamRef IndexBuffer)
+{
+	UE_LOG(LogRHI, Fatal, TEXT("RHIUpdateShaderResourceView isn't implemented for the current RHI"));
+}
diff --git a/Engine/Source/Runtime/RHI/Private/RHICommandList.cpp b/Engine/Source/Runtime/RHI/Private/RHICommandList.cpp
--- a/Engine/Source/Runtime/RHI/Private/RHICommandList.cpp
+++ b/Engine/Source/Runtime/RHI/Private/RHICommandList.cpp
@@ -2494,11 +2494,11 @@ void FRHICommandListImmediate::UpdateTextureReference(FTextureReferenceRHIParamR
 	}
 }
 
-void FRHICommandListImmediate::UpdateRHIResources(FRHIResourceUpdateInfo* UpdateInfos, int32 Num)
+void FRHICommandListImmediate::UpdateRHIResources(FRHIResourceUpdateInfo* UpdateInfos, int32 Num, bool bNeedReleaseRefs)
 {
 	if (this->Bypass())
 	{
-		FRHICommandUpdateRHIResources Cmd(UpdateInfos, Num);
+		FRHICommandUpdateRHIResources Cmd(UpdateInfos, Num, bNeedReleaseRefs);
 		Cmd.Execute(*this);
 	}
 	else
@@ -2506,7 +2506,7 @@ void FRHICommandListImmediate::UpdateRHIResources(FRHIResourceUpdateInfo* Update
 		const SIZE_T NumBytes = sizeof(FRHIResourceUpdateInfo) * Num;
 		FRHIResourceUpdateInfo* LocalUpdateInfos = reinterpret_cast<FRHIResourceUpdateInfo*>(this->Alloc(NumBytes, alignof(FRHIResourceUpdateInfo)));
 		FMemory::Memcpy(LocalUpdateInfos, UpdateInfos, NumBytes);
-		new (AllocCommand<FRHICommandUpdateRHIResources>()) FRHICommandUpdateRHIResources(LocalUpdateInfos, Num);
+		new (AllocCommand<FRHICommandUpdateRHIResources>()) FRHICommandUpdateRHIResources(LocalUpdateInfos, Num, bNeedReleaseRefs);
 		RHIThreadFence(true);
 		if (GetUsedMemory() > 256 * 1024)
 		{
diff --git a/Engine/Source/Runtime/RHI/Public/DynamicRHI.h b/Engine/Source/Runtime/RHI/Public/DynamicRHI.h
--- a/Engine/Source/Runtime/RHI/Public/DynamicRHI.h
+++ b/Engine/Source/Runtime/RHI/Public/DynamicRHI.h
@@ -175,8 +175,8 @@ public:
 	virtual FVertexDeclarationRHIRef RHICreateVertexDeclaration(const FVertexDeclarationElementList& Elements) = 0;
 
 	// FlushType: Wait RHI Thread
-	virtual FPixelShaderRHIRef RHICreatePixelShader(const TArray<uint8>& Code) = 0 ;
-	
+	virtual FPixelShaderRHIRef RHICreatePixelShader(const TArray<uint8>& Code) = 0;
+
 	// FlushType: Wait RHI Thread
 	virtual FPixelShaderRHIRef RHICreatePixelShader(FRHIShaderLibraryParamRef Library, FSHAHash Hash)
 	{
@@ -185,7 +185,7 @@ public:
 
 	// FlushType: Wait RHI Thread
 	virtual FVertexShaderRHIRef RHICreateVertexShader(const TArray<uint8>& Code) = 0;
-	
+
 	// FlushType: Wait RHI Thread
 	virtual FVertexShaderRHIRef RHICreateVertexShader(FRHIShaderLibraryParamRef Library, FSHAHash Hash)
 	{
@@ -194,7 +194,7 @@ public:
 
 	// FlushType: Wait RHI Thread
 	virtual FHullShaderRHIRef RHICreateHullShader(const TArray<uint8>& Code) = 0;
-	
+
 	// FlushType: Wait RHI Thread
 	virtual FHullShaderRHIRef RHICreateHullShader(FRHIShaderLibraryParamRef Library, FSHAHash Hash)
 	{
@@ -203,7 +203,7 @@ public:
 
 	// FlushType: Wait RHI Thread
 	virtual FDomainShaderRHIRef RHICreateDomainShader(const TArray<uint8>& Code) = 0;
-	
+
 	// FlushType: Wait RHI Thread
 	virtual FDomainShaderRHIRef RHICreateDomainShader(FRHIShaderLibraryParamRef Library, FSHAHash Hash)
 	{
@@ -212,7 +212,7 @@ public:
 
 	// FlushType: Wait RHI Thread
 	virtual FGeometryShaderRHIRef RHICreateGeometryShader(const TArray<uint8>& Code) = 0;
-	
+
 	// FlushType: Wait RHI Thread
 	virtual FGeometryShaderRHIRef RHICreateGeometryShader(FRHIShaderLibraryParamRef Library, FSHAHash Hash)
 	{
@@ -222,7 +222,7 @@ public:
 	/** Creates a geometry shader with stream output ability, defined by ElementList. */
 	// FlushType: Wait RHI Thread
 	virtual FGeometryShaderRHIRef RHICreateGeometryShaderWithStreamOutput(const TArray<uint8>& Code, const FStreamOutElementList& ElementList, uint32 NumStrides, const uint32* Strides, int32 RasterizedStream) = 0;
-	
+
 	/** Creates a geometry shader with stream output ability, defined by ElementList. */
 	// FlushType: Wait RHI Thread
 	virtual FGeometryShaderRHIRef RHICreateGeometryShaderWithStreamOutput(const FStreamOutElementList& ElementList, uint32 NumStrides, const uint32* Strides, int32 RasterizedStream, FRHIShaderLibraryParamRef Library, FSHAHash Hash)
@@ -235,26 +235,26 @@ public:
 
 	// FlushType: Wait RHI Thread
 	virtual FComputeShaderRHIRef RHICreateComputeShader(const TArray<uint8>& Code) = 0;
-	
+
 	// FlushType: Wait RHI Thread
 	virtual FComputeShaderRHIRef RHICreateComputeShader(FRHIShaderLibraryParamRef Library, FSHAHash Hash)
 	{
 		return nullptr;
 	}
-    
-    /**
-     * Attempts to open a shader library for the given shader platform & name within the provided directory.
-     * @param Platform The shader platform for shaders withing the library.
-     * @param FilePath The directory in which the library should exist.
-     * @param Name The name of the library, i.e. "Global" or "Unreal" without shader-platform or file-extension qualification.
-     * @return The new library if one exists and can be constructed, otherwise nil.
-     */
-    // FlushType: Must be Thread-Safe.
-    virtual FRHIShaderLibraryRef RHICreateShaderLibrary(EShaderPlatform Platform, FString const& FilePath, FString const& Name)
+
+	/**
+	 * Attempts to open a shader library for the given shader platform & name within the provided directory.
+	 * @param Platform The shader platform for shaders withing the library.
+	 * @param FilePath The directory in which the library should exist.
+	 * @param Name The name of the library, i.e. "Global" or "Unreal" without shader-platform or file-extension qualification.
+	 * @return The new library if one exists and can be constructed, otherwise nil.
+	 */
+	 // FlushType: Must be Thread-Safe.
+	virtual FRHIShaderLibraryRef RHICreateShaderLibrary(EShaderPlatform Platform, FString const& FilePath, FString const& Name)
 	{
 		return nullptr;
 	}
-	
+
 	/**
 	* Creates a compute fence.  Compute fences are named GPU fences which can be written to once before resetting.
 	* A command to write the fence must be enqueued before any commands to wait on them.  This is enforced on the CPU to avoid GPU hangs.
@@ -282,39 +282,39 @@ public:
 		return new FGenericRHIStagingBuffer();
 	}
 
-    /**
-     * Lock a staging buffer to read contents on the CPU that were written by the GPU, without having to stall.
-     * @discussion This function requires that you have issued an CopyToStagingBuffer invocation and verified that the FRHIGPUFence has been signaled before calling.
-     * @param StagingBuffer The buffer to lock.
-     * @param Offset The offset in the buffer to return.
-     * @param SizeRHI The length of the region in the buffer to lock.
-     * @returns A pointer to the data starting at 'Offset' and of length 'SizeRHI' from 'StagingBuffer', or nullptr when there is an error.
-     */
+	/**
+	 * Lock a staging buffer to read contents on the CPU that were written by the GPU, without having to stall.
+	 * @discussion This function requires that you have issued an CopyToStagingBuffer invocation and verified that the FRHIGPUFence has been signaled before calling.
+	 * @param StagingBuffer The buffer to lock.
+	 * @param Offset The offset in the buffer to return.
+	 * @param SizeRHI The length of the region in the buffer to lock.
+	 * @returns A pointer to the data starting at 'Offset' and of length 'SizeRHI' from 'StagingBuffer', or nullptr when there is an error.
+	 */
 	virtual void* RHILockStagingBuffer(FStagingBufferRHIParamRef StagingBuffer, uint32 Offset, uint32 SizeRHI);
-    
-    /**
-     * Unlock a staging buffer previously locked with RHILockStagingBuffer.
-     * @param StagingBuffer The buffer that was previously locked.
-     */
-    virtual void RHIUnlockStagingBuffer(FStagingBufferRHIParamRef StagingBuffer);
-
-    /**
-     * Lock a staging buffer to read contents on the CPU that were written by the GPU, without having to stall.
-     * @discussion This function requires that you have issued an CopyToStagingBuffer invocation and verified that the FRHIGPUFence has been signaled before calling.
-     * @param RHICmdList The command-list to execute on or synchronize with.
-     * @param StagingBuffer The buffer to lock.
-     * @param Offset The offset in the buffer to return.
-     * @param SizeRHI The length of the region in the buffer to lock.
-     * @returns A pointer to the data starting at 'Offset' and of length 'SizeRHI' from 'StagingBuffer', or nullptr when there is an error.
-     */
+
+	/**
+	 * Unlock a staging buffer previously locked with RHILockStagingBuffer.
+	 * @param StagingBuffer The buffer that was previously locked.
+	 */
+	virtual void RHIUnlockStagingBuffer(FStagingBufferRHIParamRef StagingBuffer);
+
+	/**
+	 * Lock a staging buffer to read contents on the CPU that were written by the GPU, without having to stall.
+	 * @discussion This function requires that you have issued an CopyToStagingBuffer invocation and verified that the FRHIGPUFence has been signaled before calling.
+	 * @param RHICmdList The command-list to execute on or synchronize with.
+	 * @param StagingBuffer The buffer to lock.
+	 * @param Offset The offset in the buffer to return.
+	 * @param SizeRHI The length of the region in the buffer to lock.
+	 * @returns A pointer to the data starting at 'Offset' and of length 'SizeRHI' from 'StagingBuffer', or nullptr when there is an error.
+	 */
 	virtual void* LockStagingBuffer_RenderThread(class FRHICommandListImmediate& RHICmdList, FStagingBufferRHIParamRef StagingBuffer, uint32 Offset, uint32 SizeRHI);
-    
-    /**
-     * Unlock a staging buffer previously locked with LockStagingBuffer_RenderThread.
-     * @param RHICmdList The command-list to execute on or synchronize with.
-     * @param StagingBuffer The buffer what was previously locked.
-     */
-    virtual void UnlockStagingBuffer_RenderThread(class FRHICommandListImmediate& RHICmdList, FStagingBufferRHIParamRef StagingBuffer);
+
+	/**
+	 * Unlock a staging buffer previously locked with LockStagingBuffer_RenderThread.
+	 * @param RHICmdList The command-list to execute on or synchronize with.
+	 * @param StagingBuffer The buffer what was previously locked.
+	 */
+	virtual void UnlockStagingBuffer_RenderThread(class FRHICommandListImmediate& RHICmdList, FStagingBufferRHIParamRef StagingBuffer);
 
 	/**
 	* Creates a bound shader state instance which encapsulates a decl, vertex shader, hull shader, domain shader and pixel shader
@@ -350,12 +350,12 @@ public:
 	{
 		return new FRHIComputePipelineStateFallback(ComputeShader);
 	}
-	
+
 	virtual FGraphicsPipelineStateRHIRef RHICreateGraphicsPipelineState(const FGraphicsPipelineStateInitializer& Initializer, FRHIPipelineBinaryLibraryParamRef PipelineBinary)
 	{
 		return RHICreateGraphicsPipelineState(Initializer);
 	}
-	
+
 	virtual TRefCountPtr<FRHIComputePipelineState> RHICreateComputePipelineState(FRHIComputeShader* ComputeShader, FRHIPipelineBinaryLibraryParamRef PipelineBinary)
 	{
 		return RHICreateComputePipelineState(ComputeShader);
@@ -382,6 +382,13 @@ public:
 	// FlushType: Flush RHI Thread
 	virtual void RHIUnlockIndexBuffer(FIndexBufferRHIParamRef IndexBuffer) = 0;
 
+	/**
+	* Transfer metadata and underlying resource from src to dest and release any resource owned by dest.
+	* @param DestIndexBuffer - the index buffer to update
+	* @param SrcIndexBuffer - don't use after call. If null, will release any resource owned by DestIndexBuffer
+	*/
+	virtual void RHITransferIndexBufferUnderlyingResource(FIndexBufferRHIParamRef DestIndexBuffer, FIndexBufferRHIParamRef SrcIndexBuffer);
+
 	/**
 	* @param ResourceArray - An optional pointer to a resource array containing the resource's data.
 	*/
@@ -398,6 +405,13 @@ public:
 	// FlushType: Flush Immediate (seems dangerous)
 	virtual void RHICopyVertexBuffer(FVertexBufferRHIParamRef SourceBuffer, FVertexBufferRHIParamRef DestBuffer) = 0;
 
+	/**
+	 * Transfer metadata and underlying resource from src to dest and release any resource owned by dest.
+	 * @param DestVertexBuffer - the vertex buffer to update
+	 * @param SrcVertexBuffer - don't use after call. If null, will release any resource owned by DestVertexBuffer
+	 */
+	virtual void RHITransferVertexBufferUnderlyingResource(FVertexBufferRHIParamRef DestVertexBuffer, FVertexBufferRHIParamRef SrcVertexBuffer);
+
 	/**
 	* @param ResourceArray - An optional pointer to a resource array containing the resource's data.
 	*/
@@ -440,7 +454,9 @@ public:
 
 	// Must be called on RHI thread timeline
 	// Make sure to call RHIThreadFence(true) afterwards so that parallel translation doesn't refer old resources
-	virtual void RHIUpdateShaderResourceView(FShaderResourceViewRHIParamRef SRV, FVertexBufferRHIParamRef VertexBuffer, uint32 Stride, uint8 Format) /*= 0;*/ {}
+	virtual void RHIUpdateShaderResourceView(FShaderResourceViewRHIParamRef SRV, FVertexBufferRHIParamRef VertexBuffer, uint32 Stride, uint8 Format);
+
+	virtual void RHIUpdateShaderResourceView(FShaderResourceViewRHIParamRef SRV, FIndexBufferRHIParamRef IndexBuffer);
 
 	/**
 	* Computes the total size of a 2D texture with the specified parameters.
diff --git a/Engine/Source/Runtime/RHI/Public/RHI.h b/Engine/Source/Runtime/RHI/Public/RHI.h
--- a/Engine/Source/Runtime/RHI/Public/RHI.h
+++ b/Engine/Source/Runtime/RHI/Public/RHI.h
@@ -1166,6 +1166,7 @@ struct FRHIResourceCreateInfo
 		: BulkData(nullptr)
 		, ResourceArray(nullptr)
 		, ClearValueBinding(FLinearColor::Transparent)
+		, bWithoutNativeResource(false)
 		, DebugName(nullptr)
 	{}
 
@@ -1174,6 +1175,7 @@ struct FRHIResourceCreateInfo
 		: BulkData(InBulkData)
 		, ResourceArray(nullptr)
 		, ClearValueBinding(FLinearColor::Transparent)
+		, bWithoutNativeResource(false)
 		, DebugName(nullptr)
 	{}
 
@@ -1182,6 +1184,7 @@ struct FRHIResourceCreateInfo
 		: BulkData(nullptr)
 		, ResourceArray(InResourceArray)
 		, ClearValueBinding(FLinearColor::Transparent)
+		, bWithoutNativeResource(false)
 		, DebugName(nullptr)
 	{}
 
@@ -1189,6 +1192,7 @@ struct FRHIResourceCreateInfo
 		: BulkData(nullptr)
 		, ResourceArray(nullptr)
 		, ClearValueBinding(InClearValueBinding)
+		, bWithoutNativeResource(false)
 		, DebugName(nullptr)
 	{
 	}
@@ -1197,6 +1201,7 @@ struct FRHIResourceCreateInfo
 		: BulkData(nullptr)
 		, ResourceArray(nullptr)
 		, ClearValueBinding(FLinearColor::Transparent)
+		, bWithoutNativeResource(false)
 		, DebugName(InDebugName)
 	{
 	}
@@ -1208,6 +1213,8 @@ struct FRHIResourceCreateInfo
 
 	// for binding clear colors to rendertargets.
 	FClearValueBinding ClearValueBinding;
+	// whether to create an RHI object with no underlying resource
+	bool bWithoutNativeResource;
 	const TCHAR* DebugName;
 };
 
diff --git a/Engine/Source/Runtime/RHI/Public/RHICommandList.h b/Engine/Source/Runtime/RHI/Public/RHICommandList.h
--- a/Engine/Source/Runtime/RHI/Public/RHICommandList.h
+++ b/Engine/Source/Runtime/RHI/Public/RHICommandList.h
@@ -1807,30 +1807,59 @@ struct FRHIShaderResourceViewUpdateInfo_VB
 	uint8 Format;
 };
 
+struct FRHIVertexBufferUpdateInfo
+{
+	FVertexBufferRHIParamRef DestBuffer;
+	FVertexBufferRHIParamRef SrcBuffer;
+};
+
+struct FRHIIndexBufferUpdateInfo
+{
+	FIndexBufferRHIParamRef DestBuffer;
+	FIndexBufferRHIParamRef SrcBuffer;
+};
+
 struct FRHIResourceUpdateInfo
 {
 	enum EUpdateType
 	{
+		/** Take over underlying resource from an intermediate vertex buffer */
+		UT_VertexBuffer,
+		/** Take over underlying resource from an intermediate index buffer */
+		UT_IndexBuffer,
+		/** Update an SRV to view on a different vertex buffer */
 		UT_VertexBufferSRV,
+		/** Update an SRV to view on a different index buffer */
+		UT_IndexBufferSRV,
+		/** Number of update types */
 		UT_Num
 	};
 
 	EUpdateType Type;
 	union
 	{
+		FRHIVertexBufferUpdateInfo VertexBuffer;
+		FRHIIndexBufferUpdateInfo IndexBuffer;
 		FRHIShaderResourceViewUpdateInfo_VB VertexBufferSRV;
 	};
+
+	void ReleaseRefs();
 };
 
 struct FRHICommandUpdateRHIResources final : public FRHICommand<FRHICommandUpdateRHIResources>
 {
 	FRHIResourceUpdateInfo* UpdateInfos;
 	int32 Num;
+	bool bNeedReleaseRefs;
 
-	FORCEINLINE_DEBUGGABLE FRHICommandUpdateRHIResources(FRHIResourceUpdateInfo* InUpdateInfos, int32 InNum)
+	FRHICommandUpdateRHIResources(FRHIResourceUpdateInfo* InUpdateInfos, int32 InNum, bool bInNeedReleaseRefs)
 		: UpdateInfos(InUpdateInfos)
 		, Num(InNum)
+		, bNeedReleaseRefs(bInNeedReleaseRefs)
 	{}
+
+	~FRHICommandUpdateRHIResources();
+
 	RHI_API void Execute(FRHICommandListBase& CmdList);
 };
 
@@ -4347,7 +4376,12 @@ public:
 		GDynamicRHI->RHIPollRenderQueryResults();
 	}
 
-	void UpdateRHIResources(FRHIResourceUpdateInfo* UpdateInfos, int32 Num);
+	/**
+	 * @param UpdateInfos - an array of update infos
+	 * @param Num - number of update infos
+	 * @param bNeedReleaseRefs - whether Release need to be called on RHI resources referenced by update infos
+	 */
+	void UpdateRHIResources(FRHIResourceUpdateInfo* UpdateInfos, int32 Num, bool bNeedReleaseRefs);
 };
 
  struct FScopedGPUMask
@@ -4686,9 +4720,9 @@ FORCEINLINE FShaderResourceViewRHIRef RHICreateShaderResourceView(FIndexBufferRH
 	return FRHICommandListExecutor::GetImmediateCommandList().CreateShaderResourceView(Buffer);
 }
 
-FORCEINLINE void RHIUpdateRHIResources(FRHIResourceUpdateInfo* UpdateInfos, int32 Num)
+FORCEINLINE void RHIUpdateRHIResources(FRHIResourceUpdateInfo* UpdateInfos, int32 Num, bool bNeedReleaseRefs)
 {
-	return FRHICommandListExecutor::GetImmediateCommandList().UpdateRHIResources(UpdateInfos, Num);
+	return FRHICommandListExecutor::GetImmediateCommandList().UpdateRHIResources(UpdateInfos, Num, bNeedReleaseRefs);
 }
 
 FORCEINLINE FTextureReferenceRHIRef RHICreateTextureReference(FLastRenderTimeContainer* LastRenderTime)
@@ -4926,11 +4960,11 @@ FORCEINLINE void RHIUnlockStagingBuffer(FStagingBufferRHIParamRef StagingBuffer)
 	 FRHICommandListExecutor::GetImmediateCommandList().UnlockStagingBuffer(StagingBuffer);
 }
 
-template <int32 MaxNumUpdates>
+template <uint32 MaxNumUpdates>
 struct TRHIResourceUpdateBatcher
 {
 	FRHIResourceUpdateInfo UpdateInfos[MaxNumUpdates];
-	int32 NumBatched;
+	uint32 NumBatched;
 
 	TRHIResourceUpdateBatcher()
 		: NumBatched(0)
@@ -4945,24 +4979,61 @@ struct TRHIResourceUpdateBatcher
 	{
 		if (NumBatched > 0)
 		{
-			RHIUpdateRHIResources(UpdateInfos, NumBatched);
+			RHIUpdateRHIResources(UpdateInfos, NumBatched, true);
 			NumBatched = 0;
 		}
 	}
 
+	void QueueUpdateRequest(FVertexBufferRHIParamRef DestVertexBuffer, FVertexBufferRHIParamRef SrcVertexBuffer)
+	{
+		FRHIResourceUpdateInfo& UpdateInfo = GetNextUpdateInfo();
+		UpdateInfo.Type = FRHIResourceUpdateInfo::UT_VertexBuffer;
+		UpdateInfo.VertexBuffer = { DestVertexBuffer, SrcVertexBuffer };
+		DestVertexBuffer->AddRef();
+		if (SrcVertexBuffer)
+		{
+			SrcVertexBuffer->AddRef();
+		}
+	}
+
+	void QueueUpdateRequest(FIndexBufferRHIParamRef DestIndexBuffer, FIndexBufferRHIParamRef SrcIndexBuffer)
+	{
+		FRHIResourceUpdateInfo& UpdateInfo = GetNextUpdateInfo();
+		UpdateInfo.Type = FRHIResourceUpdateInfo::UT_IndexBuffer;
+		UpdateInfo.IndexBuffer = { DestIndexBuffer, SrcIndexBuffer };
+		DestIndexBuffer->AddRef();
+		if (SrcIndexBuffer)
+		{
+			SrcIndexBuffer->AddRef();
+		}
+	}
+
 	void QueueUpdateRequest(FShaderResourceViewRHIParamRef SRV, FVertexBufferRHIParamRef VertexBuffer, uint32 Stride, uint8 Format)
 	{
-		check(NumBatched >= 0 && NumBatched <= MaxNumUpdates);
-		if (NumBatched == MaxNumUpdates)
+		FRHIResourceUpdateInfo& UpdateInfo = GetNextUpdateInfo();
+		UpdateInfo.Type = FRHIResourceUpdateInfo::UT_VertexBufferSRV;
+		UpdateInfo.VertexBufferSRV = { SRV, VertexBuffer, Stride, Format };
+		SRV->AddRef();
+		if (VertexBuffer)
 		{
-			Flush();
+			VertexBuffer->AddRef();
 		}
-		if (LIKELY(NumBatched >= 0 && NumBatched < MaxNumUpdates))
+	}
+
+	void QueueUpdateRequest(FShaderResourceViewRHIParamRef SRV, FIndexBufferRHIParamRef IndexBuffer)
+	{
+		// TODO
+	}
+
+private:
+	FRHIResourceUpdateInfo & GetNextUpdateInfo()
+	{
+		check(NumBatched <= MaxNumUpdates);
+		if (NumBatched >= MaxNumUpdates)
 		{
-			const int32 Idx = NumBatched++;
-			UpdateInfos[Idx].Type = FRHIResourceUpdateInfo::UT_VertexBufferSRV;
-			UpdateInfos[Idx].VertexBufferSRV = { SRV, VertexBuffer, Stride, Format };
+			Flush();
 		}
+		return UpdateInfos[NumBatched++];
 	}
 };
 
diff --git a/Engine/Source/Runtime/RHI/Public/RHICommandListCommandExecutes.inl b/Engine/Source/Runtime/RHI/Public/RHICommandListCommandExecutes.inl
--- a/Engine/Source/Runtime/RHI/Public/RHICommandListCommandExecutes.inl
+++ b/Engine/Source/Runtime/RHI/Public/RHICommandListCommandExecutes.inl
@@ -609,6 +609,48 @@ void FRHICommandUpdateTextureReference::Execute(FRHICommandListBase& CmdList)
 	INTERNAL_DECORATOR(RHIUpdateTextureReference)(TextureRef, NewTexture);
 }
 
+void FRHIResourceUpdateInfo::ReleaseRefs()
+{
+	switch (Type)
+	{
+	case UT_VertexBuffer:
+		VertexBuffer.DestBuffer->Release();
+		if (VertexBuffer.SrcBuffer)
+		{
+			VertexBuffer.SrcBuffer->Release();
+		}
+		break;
+	case UT_IndexBuffer:
+		IndexBuffer.DestBuffer->Release();
+		if (IndexBuffer.SrcBuffer)
+		{
+			IndexBuffer.SrcBuffer->Release();
+		}
+		break;
+	case UT_VertexBufferSRV:
+		VertexBufferSRV.SRV->Release();
+		if (VertexBufferSRV.VertexBuffer)
+		{
+			VertexBufferSRV.VertexBuffer->Release();
+		}
+		break;
+	default:
+		// Unrecognized type, do nothing
+		break;
+	}
+}
+
+FRHICommandUpdateRHIResources::~FRHICommandUpdateRHIResources()
+{
+	if (bNeedReleaseRefs)
+	{
+		for (int32 Idx = 0; Idx < Num; ++Idx)
+		{
+			UpdateInfos[Idx].ReleaseRefs();
+		}
+	}
+}
+
 void FRHICommandUpdateRHIResources::Execute(FRHICommandListBase& CmdList)
 {
 	RHISTAT(UpdateRHIResources);
@@ -617,6 +659,16 @@ void FRHICommandUpdateRHIResources::Execute(FRHICommandListBase& CmdList)
 		FRHIResourceUpdateInfo& Info = UpdateInfos[Idx];
 		switch (Info.Type)
 		{
+		case FRHIResourceUpdateInfo::UT_VertexBuffer:
+			GDynamicRHI->RHITransferVertexBufferUnderlyingResource(
+				Info.VertexBuffer.DestBuffer,
+				Info.VertexBuffer.SrcBuffer);
+			break;
+		case FRHIResourceUpdateInfo::UT_IndexBuffer:
+			GDynamicRHI->RHITransferIndexBufferUnderlyingResource(
+				Info.IndexBuffer.DestBuffer,
+				Info.IndexBuffer.SrcBuffer);
+			break;
 		case FRHIResourceUpdateInfo::UT_VertexBufferSRV:
 			GDynamicRHI->RHIUpdateShaderResourceView(
 				Info.VertexBufferSRV.SRV,
diff --git a/Engine/Source/Runtime/RHI/Public/RHIResources.h b/Engine/Source/Runtime/RHI/Public/RHIResources.h
--- a/Engine/Source/Runtime/RHI/Public/RHIResources.h
+++ b/Engine/Source/Runtime/RHI/Public/RHIResources.h
@@ -451,6 +451,25 @@ public:
 	/** @return The usage flags used to create the index buffer. */
 	uint32 GetUsage() const { return Usage; }
 
+protected:
+	FRHIIndexBuffer()
+		: Stride(0)
+		, Size(0)
+		, Usage(0)
+	{}
+
+	void Swap(FRHIIndexBuffer& Other)
+	{
+		::Swap(Stride, Other.Stride);
+		::Swap(Size, Other.Size);
+		::Swap(Usage, Other.Usage);
+	}
+
+	void ReleaseUnderlyingResource()
+	{
+		Stride = Size = Usage = 0;
+	}
+
 private:
 	uint32 Stride;
 	uint32 Size;
@@ -476,6 +495,24 @@ public:
 	/** @return The usage flags used to create the vertex buffer. e.g. BUF_UnorderedAccess */
 	uint32 GetUsage() const { return Usage; }
 
+protected:
+	FRHIVertexBuffer()
+		: Size(0)
+		, Usage(0)
+	{}
+
+	void Swap(FRHIVertexBuffer& Other)
+	{
+		::Swap(Size, Other.Size);
+		::Swap(Usage, Other.Usage);
+	}
+
+	void ReleaseUnderlyingResource()
+	{
+		Size = 0;
+		Usage = 0;
+	}
+
 private:
 	uint32 Size;
 	// e.g. BUF_UnorderedAccess
diff --git a/Engine/Source/Runtime/VulkanRHI/Private/VulkanIndexBuffer.cpp b/Engine/Source/Runtime/VulkanRHI/Private/VulkanIndexBuffer.cpp
--- a/Engine/Source/Runtime/VulkanRHI/Private/VulkanIndexBuffer.cpp
+++ b/Engine/Source/Runtime/VulkanRHI/Private/VulkanIndexBuffer.cpp
@@ -333,6 +333,11 @@ void FVulkanResourceMultiBuffer::Unlock(bool bFromRenderingThread)
 	}
 }
 
+void FVulkanResourceMultiBuffer::Swap(FVulkanResourceMultiBuffer& Other)
+{
+	::Swap(*this, Other);
+}
+
 
 FVulkanIndexBuffer::FVulkanIndexBuffer(FVulkanDevice* InDevice, uint32 InStride, uint32 InSize, uint32 InUsage, FRHIResourceCreateInfo& CreateInfo, class FRHICommandListImmediate* InRHICmdList)
 	: FRHIIndexBuffer(InStride, InSize, InUsage)
@@ -341,10 +346,21 @@ FVulkanIndexBuffer::FVulkanIndexBuffer(FVulkanDevice* InDevice, uint32 InStride,
 {
 }
 
+void FVulkanIndexBuffer::Swap(FVulkanIndexBuffer& Other)
+{
+	FRHIIndexBuffer::Swap(Other);
+	FVulkanResourceMultiBuffer::Swap(Other);
+	::Swap(IndexType, Other.IndexType);
+}
+
 
 FIndexBufferRHIRef FVulkanDynamicRHI::RHICreateIndexBuffer(uint32 Stride, uint32 Size, uint32 InUsage, FRHIResourceCreateInfo& CreateInfo)
 {
 	LLM_SCOPE_VULKAN(ELLMTagVulkan::VulkanIndexBuffers);
+	if (CreateInfo.bWithoutNativeResource)
+	{
+		return new FVulkanIndexBuffer(nullptr, 0, 0, 0, CreateInfo, nullptr);
+	}
 	return new FVulkanIndexBuffer(Device, Stride, Size, InUsage, CreateInfo, nullptr);
 }
 
@@ -373,3 +389,20 @@ void FVulkanDynamicRHI::UnlockIndexBuffer_RenderThread(FRHICommandListImmediate&
 	this->RHIUnlockIndexBuffer(IndexBufferRHI);
 }
 #endif
+
+void FVulkanDynamicRHI::RHITransferIndexBufferUnderlyingResource(FIndexBufferRHIParamRef DestIndexBuffer, FIndexBufferRHIParamRef SrcIndexBuffer)
+{
+	check(DestIndexBuffer);
+	FVulkanIndexBuffer* Dest = ResourceCast(DestIndexBuffer);
+	if (!SrcIndexBuffer)
+	{
+		FRHIResourceCreateInfo CreateInfo;
+		TRefCountPtr<FVulkanIndexBuffer> DeletionProxy = new FVulkanIndexBuffer(Dest->GetParent(), 0, 0, 0, CreateInfo, nullptr);
+		Dest->Swap(*DeletionProxy);
+	}
+	else
+	{
+		FVulkanIndexBuffer* Src = ResourceCast(SrcIndexBuffer);
+		Dest->Swap(*Src);
+	}
+}
diff --git a/Engine/Source/Runtime/VulkanRHI/Private/VulkanUAV.cpp b/Engine/Source/Runtime/VulkanRHI/Private/VulkanUAV.cpp
--- a/Engine/Source/Runtime/VulkanRHI/Private/VulkanUAV.cpp
+++ b/Engine/Source/Runtime/VulkanRHI/Private/VulkanUAV.cpp
@@ -307,6 +307,10 @@ FShaderResourceViewRHIRef FVulkanDynamicRHI::RHICreateShaderResourceView(FTextur
 
 FShaderResourceViewRHIRef FVulkanDynamicRHI::RHICreateShaderResourceView(FIndexBufferRHIParamRef IndexBufferRHI)
 {
+	if (!IndexBufferRHI)
+	{
+		return new FVulkanShaderResourceView(Device, nullptr, nullptr, 0, PF_R16_UINT);
+	}
 	FVulkanIndexBuffer* IndexBuffer = ResourceCast(IndexBufferRHI);
 	check(IndexBufferRHI->GetStride() == 2 || IndexBufferRHI->GetStride() == 4);
 	EPixelFormat Format = (IndexBufferRHI->GetStride() == 4) ? PF_R32_UINT : PF_R16_UINT;
@@ -329,6 +333,21 @@ void FVulkanDynamicRHI::RHIUpdateShaderResourceView(FShaderResourceViewRHIParamR
 	}
 }
 
+void FVulkanDynamicRHI::RHIUpdateShaderResourceView(FShaderResourceViewRHIParamRef SRV, FIndexBufferRHIParamRef IndexBuffer)
+{
+	FVulkanShaderResourceView* SRVVk = ResourceCast(SRV);
+	check(SRVVk && SRVVk->GetParent() == Device);
+	if (!IndexBuffer)
+	{
+		SRVVk->Clear();
+	}
+	else if (SRVVk->SourceRHIBuffer.GetReference() != IndexBuffer)
+	{
+		FVulkanIndexBuffer* IndexBufferVk = ResourceCast(IndexBuffer);
+		SRVVk->Rename(IndexBuffer, IndexBufferVk, IndexBufferVk->GetSize(), IndexBufferVk->GetStride() == 2u ? PF_R16_UINT : PF_R32_UINT);
+	}
+}
+
 void FVulkanCommandListContext::RHIClearTinyUAV(FUnorderedAccessViewRHIParamRef UnorderedAccessViewRHI, const uint32* Values)
 {
 	FVulkanUnorderedAccessView* UnorderedAccessView = ResourceCast(UnorderedAccessViewRHI);
diff --git a/Engine/Source/Runtime/VulkanRHI/Private/VulkanVertexBuffer.cpp b/Engine/Source/Runtime/VulkanRHI/Private/VulkanVertexBuffer.cpp
--- a/Engine/Source/Runtime/VulkanRHI/Private/VulkanVertexBuffer.cpp
+++ b/Engine/Source/Runtime/VulkanRHI/Private/VulkanVertexBuffer.cpp
@@ -13,9 +13,19 @@ FVulkanVertexBuffer::FVulkanVertexBuffer(FVulkanDevice* InDevice, uint32 InSize,
 {
 }
 
+void FVulkanVertexBuffer::Swap(FVulkanVertexBuffer& Other)
+{
+	FRHIVertexBuffer::Swap(Other);
+	FVulkanResourceMultiBuffer::Swap(Other);
+}
+
 FVertexBufferRHIRef FVulkanDynamicRHI::RHICreateVertexBuffer(uint32 Size, uint32 InUsage, FRHIResourceCreateInfo& CreateInfo)
 {
 	LLM_SCOPE_VULKAN(ELLMTagVulkan::VulkanVertexBuffers);
+	if (CreateInfo.bWithoutNativeResource)
+	{
+		return new FVulkanVertexBuffer(nullptr, 0, 0, CreateInfo, nullptr);
+	}
 	FVulkanVertexBuffer* VertexBuffer = new FVulkanVertexBuffer(Device, Size, InUsage, CreateInfo, nullptr);
 	return VertexBuffer;
 }
@@ -52,3 +62,20 @@ void FVulkanDynamicRHI::RHICopyVertexBuffer(FVertexBufferRHIParamRef SourceBuffe
 {
 	VULKAN_SIGNAL_UNIMPLEMENTED();
 }
+
+void FVulkanDynamicRHI::RHITransferVertexBufferUnderlyingResource(FVertexBufferRHIParamRef DestVertexBuffer, FVertexBufferRHIParamRef SrcVertexBuffer)
+{
+	check(DestVertexBuffer);
+	FVulkanVertexBuffer* Dest = ResourceCast(DestVertexBuffer);
+	if (!SrcVertexBuffer)
+	{
+		FRHIResourceCreateInfo CreateInfo;
+		TRefCountPtr<FVulkanVertexBuffer> DeletionProxy = new FVulkanVertexBuffer(Dest->GetParent(), 0, 0, CreateInfo, nullptr);
+		Dest->Swap(*DeletionProxy);
+	}
+	else
+	{
+		FVulkanVertexBuffer* Src = ResourceCast(SrcVertexBuffer);
+		Dest->Swap(*Src);
+	}
+}
diff --git a/Engine/Source/Runtime/VulkanRHI/Public/VulkanDynamicRHI.h b/Engine/Source/Runtime/VulkanRHI/Public/VulkanDynamicRHI.h
--- a/Engine/Source/Runtime/VulkanRHI/Public/VulkanDynamicRHI.h
+++ b/Engine/Source/Runtime/VulkanRHI/Public/VulkanDynamicRHI.h
@@ -62,10 +62,12 @@ public:
 	virtual FIndexBufferRHIRef RHICreateIndexBuffer(uint32 Stride, uint32 Size, uint32 InUsage, FRHIResourceCreateInfo& CreateInfo) final override;
 	virtual void* RHILockIndexBuffer(FIndexBufferRHIParamRef IndexBuffer, uint32 Offset, uint32 Size, EResourceLockMode LockMode) final override;
 	virtual void RHIUnlockIndexBuffer(FIndexBufferRHIParamRef IndexBuffer) final override;
+	virtual void RHITransferIndexBufferUnderlyingResource(FIndexBufferRHIParamRef DestIndexBuffer, FIndexBufferRHIParamRef SrcIndexBuffer) final override;
 	virtual FVertexBufferRHIRef RHICreateVertexBuffer(uint32 Size, uint32 InUsage, FRHIResourceCreateInfo& CreateInfo) final override;
 	virtual void* RHILockVertexBuffer(FVertexBufferRHIParamRef VertexBuffer, uint32 Offset, uint32 SizeRHI, EResourceLockMode LockMode) final override;
 	virtual void RHIUnlockVertexBuffer(FVertexBufferRHIParamRef VertexBuffer) final override;
 	virtual void RHICopyVertexBuffer(FVertexBufferRHIParamRef SourceBuffer, FVertexBufferRHIParamRef DestBuffer) final override;
+	virtual void RHITransferVertexBufferUnderlyingResource(FVertexBufferRHIParamRef DestVertexBuffer, FVertexBufferRHIParamRef SrcVertexBuffer) final override;
 	virtual FStructuredBufferRHIRef RHICreateStructuredBuffer(uint32 Stride, uint32 Size, uint32 InUsage, FRHIResourceCreateInfo& CreateInfo) final override;
 	virtual void* RHILockStructuredBuffer(FStructuredBufferRHIParamRef StructuredBuffer, uint32 Offset, uint32 SizeRHI, EResourceLockMode LockMode) final override;
 	virtual void RHIUnlockStructuredBuffer(FStructuredBufferRHIParamRef StructuredBuffer) final override;
@@ -77,6 +79,7 @@ public:
 	virtual FShaderResourceViewRHIRef RHICreateShaderResourceView(FVertexBufferRHIParamRef VertexBuffer, uint32 Stride, uint8 Format) final override;
 	virtual FShaderResourceViewRHIRef RHICreateShaderResourceView(FIndexBufferRHIParamRef Buffer) final override;
 	virtual void RHIUpdateShaderResourceView(FShaderResourceViewRHIParamRef SRV, FVertexBufferRHIParamRef VertexBuffer, uint32 Stride, uint8 Format) final override;
+	virtual void RHIUpdateShaderResourceView(FShaderResourceViewRHIParamRef SRV, FIndexBufferRHIParamRef IndexBuffer) final override;
 	virtual uint64 RHICalcTexture2DPlatformSize(uint32 SizeX, uint32 SizeY, uint8 Format, uint32 NumMips, uint32 NumSamples, uint32 Flags, uint32& OutAlign) final override;
 	virtual uint64 RHICalcTexture3DPlatformSize(uint32 SizeX, uint32 SizeY, uint32 SizeZ, uint8 Format, uint32 NumMips, uint32 Flags, uint32& OutAlign) final override;
 	virtual uint64 RHICalcTextureCubePlatformSize(uint32 Size, uint8 Format, uint32 NumMips, uint32 Flags, uint32& OutAlign) final override;
diff --git a/Engine/Source/Runtime/VulkanRHI/Public/VulkanResources.h b/Engine/Source/Runtime/VulkanRHI/Public/VulkanResources.h
--- a/Engine/Source/Runtime/VulkanRHI/Public/VulkanResources.h
+++ b/Engine/Source/Runtime/VulkanRHI/Public/VulkanResources.h
@@ -1154,6 +1154,8 @@ public:
 	void* Lock(bool bFromRenderingThread, EResourceLockMode LockMode, uint32 Size, uint32 Offset);
 	void Unlock(bool bFromRenderingThread);
 
+	void Swap(FVulkanResourceMultiBuffer& Other);
+
 protected:
 	uint32 UEUsage;
 	VkBufferUsageFlags BufferUsageFlags;
@@ -1191,6 +1193,8 @@ public:
 		return IndexType;
 	}
 
+	void Swap(FVulkanIndexBuffer& Other);
+
 private:
 	VkIndexType IndexType;
 };
@@ -1199,6 +1203,8 @@ class FVulkanVertexBuffer : public FRHIVertexBuffer, public FVulkanResourceMulti
 {
 public:
 	FVulkanVertexBuffer(FVulkanDevice* InDevice, uint32 InSize, uint32 InUsage, FRHIResourceCreateInfo& CreateInfo, class FRHICommandListImmediate* InRHICmdList);
+
+	void Swap(FVulkanVertexBuffer& Other);
 };
 
 class FVulkanUniformBuffer : public FRHIUniformBuffer
diff --git a/Engine/Source/Runtime/Windows/D3D11RHI/Private/D3D11IndexBuffer.cpp b/Engine/Source/Runtime/Windows/D3D11RHI/Private/D3D11IndexBuffer.cpp
--- a/Engine/Source/Runtime/Windows/D3D11RHI/Private/D3D11IndexBuffer.cpp
+++ b/Engine/Source/Runtime/Windows/D3D11RHI/Private/D3D11IndexBuffer.cpp
@@ -8,6 +8,11 @@
 
 FIndexBufferRHIRef FD3D11DynamicRHI::RHICreateIndexBuffer(uint32 Stride,uint32 Size,uint32 InUsage, FRHIResourceCreateInfo& CreateInfo)
 {
+	if (CreateInfo.bWithoutNativeResource)
+	{
+		return new FD3D11IndexBuffer();
+	}
+
 	// Explicitly check that the size is nonzero before allowing CreateIndexBuffer to opaquely fail.
 	check(Size > 0);
 
@@ -178,3 +183,18 @@ void FD3D11DynamicRHI::RHIUnlockIndexBuffer(FIndexBufferRHIParamRef IndexBufferR
 	// If the lock involved a staging resource, this releases it.
 	OutstandingLocks.Remove(LockedKey);
 }
+
+void FD3D11DynamicRHI::RHITransferIndexBufferUnderlyingResource(FIndexBufferRHIParamRef DestIndexBuffer, FIndexBufferRHIParamRef SrcIndexBuffer)
+{
+	check(DestIndexBuffer);
+	FD3D11IndexBuffer* Dest = ResourceCast(DestIndexBuffer);
+	if (!SrcIndexBuffer)
+	{
+		Dest->ReleaseUnderlyingResource();
+	}
+	else
+	{
+		FD3D11IndexBuffer* Src = ResourceCast(SrcIndexBuffer);
+		Dest->Swap(*Src);
+	}
+}
diff --git a/Engine/Source/Runtime/Windows/D3D11RHI/Private/D3D11RHIPrivate.h b/Engine/Source/Runtime/Windows/D3D11RHI/Private/D3D11RHIPrivate.h
--- a/Engine/Source/Runtime/Windows/D3D11RHI/Private/D3D11RHIPrivate.h
+++ b/Engine/Source/Runtime/Windows/D3D11RHI/Private/D3D11RHIPrivate.h
@@ -413,10 +413,12 @@ public:
 	virtual FIndexBufferRHIRef RHICreateIndexBuffer(uint32 Stride, uint32 Size, uint32 InUsage, FRHIResourceCreateInfo& CreateInfo) final override;
 	virtual void* RHILockIndexBuffer(FIndexBufferRHIParamRef IndexBuffer, uint32 Offset, uint32 Size, EResourceLockMode LockMode) final override;
 	virtual void RHIUnlockIndexBuffer(FIndexBufferRHIParamRef IndexBuffer) final override;
+	virtual void RHITransferIndexBufferUnderlyingResource(FIndexBufferRHIParamRef DestIndexBuffer, FIndexBufferRHIParamRef SrcIndexBuffer) final override;
 	virtual FVertexBufferRHIRef RHICreateVertexBuffer(uint32 Size, uint32 InUsage, FRHIResourceCreateInfo& CreateInfo) final override;
 	virtual void* RHILockVertexBuffer(FVertexBufferRHIParamRef VertexBuffer, uint32 Offset, uint32 SizeRHI, EResourceLockMode LockMode) final override;
 	virtual void RHIUnlockVertexBuffer(FVertexBufferRHIParamRef VertexBuffer) final override;
 	virtual void RHICopyVertexBuffer(FVertexBufferRHIParamRef SourceBuffer,FVertexBufferRHIParamRef DestBuffer) final override;
+	virtual void RHITransferVertexBufferUnderlyingResource(FVertexBufferRHIParamRef DestVertexBuffer, FVertexBufferRHIParamRef SrcVertexBuffer) final override;
 	virtual FStructuredBufferRHIRef RHICreateStructuredBuffer(uint32 Stride, uint32 Size, uint32 InUsage, FRHIResourceCreateInfo& CreateInfo) final override;
 	virtual void* RHILockStructuredBuffer(FStructuredBufferRHIParamRef StructuredBuffer, uint32 Offset, uint32 SizeRHI, EResourceLockMode LockMode) final override;
 	virtual void RHIUnlockStructuredBuffer(FStructuredBufferRHIParamRef StructuredBuffer) final override;
@@ -428,6 +430,7 @@ public:
 	virtual FShaderResourceViewRHIRef RHICreateShaderResourceView(FVertexBufferRHIParamRef VertexBuffer, uint32 Stride, uint8 Format) final override;
 	virtual FShaderResourceViewRHIRef RHICreateShaderResourceView(FIndexBufferRHIParamRef Buffer) final override;
 	virtual void RHIUpdateShaderResourceView(FShaderResourceViewRHIParamRef SRV, FVertexBufferRHIParamRef VertexBuffer, uint32 Stride, uint8 Format) final override;
+	virtual void RHIUpdateShaderResourceView(FShaderResourceViewRHIParamRef SRV, FIndexBufferRHIParamRef IndexBuffer) final override;
 	virtual uint64 RHICalcTexture2DPlatformSize(uint32 SizeX, uint32 SizeY, uint8 Format, uint32 NumMips, uint32 NumSamples, uint32 Flags, uint32& OutAlign) final override;
 	virtual uint64 RHICalcTexture3DPlatformSize(uint32 SizeX, uint32 SizeY, uint32 SizeZ, uint8 Format, uint32 NumMips, uint32 Flags, uint32& OutAlign) final override;
 	virtual uint64 RHICalcTextureCubePlatformSize(uint32 Size, uint8 Format, uint32 NumMips, uint32 Flags, uint32& OutAlign) final override;
diff --git a/Engine/Source/Runtime/Windows/D3D11RHI/Private/D3D11UAV.cpp b/Engine/Source/Runtime/Windows/D3D11RHI/Private/D3D11UAV.cpp
--- a/Engine/Source/Runtime/Windows/D3D11RHI/Private/D3D11UAV.cpp
+++ b/Engine/Source/Runtime/Windows/D3D11RHI/Private/D3D11UAV.cpp
@@ -228,7 +228,7 @@ FShaderResourceViewRHIRef FD3D11DynamicRHI::RHICreateShaderResourceView_RenderTh
 	return RHICreateShaderResourceView(StructuredBuffer);
 }
 
-static void CreateD3D11ShaderResourceViewOnVertexBuffer(ID3D11Device* Direct3DDevice, ID3D11Buffer* Buffer, uint32 Stride, uint8 Format, ID3D11ShaderResourceView** OutSRV)
+static void CreateD3D11ShaderResourceViewOnBuffer(ID3D11Device* Direct3DDevice, ID3D11Buffer* Buffer, uint32 Stride, uint8 Format, ID3D11ShaderResourceView** OutSRV)
 {
 	D3D11_BUFFER_DESC BufferDesc;
 	Buffer->GetDesc(&BufferDesc);
@@ -252,7 +252,7 @@ static void CreateD3D11ShaderResourceViewOnVertexBuffer(ID3D11Device* Direct3DDe
 		}
 		if (FAILED(hr))
 		{
-			UE_LOG(LogD3D11RHI, Error, TEXT("Failed to create shader resource view for vertex buffer: ByteWidth=%d NumElements=%d Format=%s"), BufferDesc.ByteWidth, BufferDesc.ByteWidth / Stride, GPixelFormats[Format].Name);
+			UE_LOG(LogD3D11RHI, Error, TEXT("Failed to create shader resource view for buffer: ByteWidth=%d NumElements=%d Format=%s"), BufferDesc.ByteWidth, BufferDesc.ByteWidth / Stride, GPixelFormats[Format].Name);
 			VerifyD3D11Result(hr, "Direct3DDevice->CreateShaderResourceView", __FILE__, __LINE__, Direct3DDevice);
 		}
 	}
@@ -269,7 +269,7 @@ FShaderResourceViewRHIRef FD3D11DynamicRHI::RHICreateShaderResourceView(FVertexB
 	check(VertexBuffer->Resource);
 
 	TRefCountPtr<ID3D11ShaderResourceView> ShaderResourceView;
-	CreateD3D11ShaderResourceViewOnVertexBuffer(Direct3DDevice, VertexBuffer->Resource, Stride, Format, ShaderResourceView.GetInitReference());
+	CreateD3D11ShaderResourceViewOnBuffer(Direct3DDevice, VertexBuffer->Resource, Stride, Format, ShaderResourceView.GetInitReference());
 
 	return new FD3D11ShaderResourceView(ShaderResourceView,VertexBuffer);
 }
@@ -297,52 +297,50 @@ void FD3D11DynamicRHI::RHIUpdateShaderResourceView(FShaderResourceViewRHIParamRe
 		check(VertexBuffer->Resource);
 
 		TRefCountPtr<ID3D11ShaderResourceView> ShaderResourceView;
-		CreateD3D11ShaderResourceViewOnVertexBuffer(Direct3DDevice, VertexBuffer->Resource, Stride, Format, ShaderResourceView.GetInitReference());
+		CreateD3D11ShaderResourceViewOnBuffer(Direct3DDevice, VertexBuffer->Resource, Stride, Format, ShaderResourceView.GetInitReference());
 
 		SRVD3D11->Rename(ShaderResourceView, VertexBuffer);
 	}
 }
 
+void FD3D11DynamicRHI::RHIUpdateShaderResourceView(FShaderResourceViewRHIParamRef SRV, FIndexBufferRHIParamRef IndexBufferRHI)
+{
+	check(SRV);
+	FD3D11ShaderResourceView* SRVD3D11 = ResourceCast(SRV);
+	if (!IndexBufferRHI)
+	{
+		SRVD3D11->Rename(nullptr, nullptr);
+	}
+	else
+	{
+		FD3D11IndexBuffer* IndexBuffer = ResourceCast(IndexBufferRHI);
+		check(IndexBuffer->Resource);
+
+		const uint32 Stride = IndexBuffer->GetStride();
+		const EPixelFormat Format = Stride == 2 ? PF_R16_UINT : PF_R32_UINT;
+		TRefCountPtr<ID3D11ShaderResourceView> ShaderResourceView;
+		CreateD3D11ShaderResourceViewOnBuffer(Direct3DDevice, IndexBuffer->Resource, Stride, Format, ShaderResourceView.GetInitReference());
+
+		SRVD3D11->Rename(ShaderResourceView, IndexBuffer);
+	}
+}
+
 FShaderResourceViewRHIRef FD3D11DynamicRHI::RHICreateShaderResourceView(FIndexBufferRHIParamRef BufferRHI)
 {
+	if (!BufferRHI)
+	{
+		return new FD3D11ShaderResourceView(nullptr, nullptr);
+	}
 	FD3D11IndexBuffer* Buffer = ResourceCast(BufferRHI);
 	check(Buffer);
 	check(Buffer->Resource);
 
 	// The stride in bytes of the index buffer; must be 2 or 4
-	uint32 Stride = BufferRHI->GetStride();
-
+	const uint32 Stride = BufferRHI->GetStride();
 	check(Stride == 2 || Stride == 4);
-
-	D3D11_BUFFER_DESC BufferDesc;
-	Buffer->Resource->GetDesc(&BufferDesc);
-
-	// Create a Shader Resource View
-	D3D11_SHADER_RESOURCE_VIEW_DESC SRVDesc;
-	FMemory::Memzero(SRVDesc);
-	SRVDesc.ViewDimension = D3D11_SRV_DIMENSION_BUFFER;
-	SRVDesc.Buffer.FirstElement = 0;
-
-	EPixelFormat Format = (Stride == 2) ? PF_R16_UINT : PF_R32_UINT;
-
-	SRVDesc.Format = FindShaderResourceDXGIFormat((DXGI_FORMAT)GPixelFormats[Format].PlatformFormat,false);
-	SRVDesc.Buffer.NumElements = BufferDesc.ByteWidth / Stride;
+	const EPixelFormat Format = (Stride == 2) ? PF_R16_UINT : PF_R32_UINT;
 	TRefCountPtr<ID3D11ShaderResourceView> ShaderResourceView;
-	
-	HRESULT hr = Direct3DDevice->CreateShaderResourceView(Buffer->Resource, &SRVDesc, (ID3D11ShaderResourceView**)ShaderResourceView.GetInitReference());
-	if (FAILED(hr))
-	{
-		if (hr == E_OUTOFMEMORY)
-		{
-			// There appears to be a driver bug that causes SRV creation to fail with an OOM error and then succeed on the next call.
-			hr = Direct3DDevice->CreateShaderResourceView(Buffer->Resource, &SRVDesc, (ID3D11ShaderResourceView**)ShaderResourceView.GetInitReference());
-		}
-		if (FAILED(hr))
-		{
-			UE_LOG(LogD3D11RHI,Error,TEXT("Failed to create shader resource view for index buffer: ByteWidth=%d NumElements=%d Format=%s"),BufferDesc.ByteWidth,BufferDesc.ByteWidth / Stride, GPixelFormats[Format].Name);
-			VerifyD3D11Result(hr,"Direct3DDevice->CreateShaderResourceView",__FILE__,__LINE__,Direct3DDevice);
-		}
-	}
+	CreateD3D11ShaderResourceViewOnBuffer(Direct3DDevice, Buffer->Resource, Stride, Format, ShaderResourceView.GetInitReference());
 
 	return new FD3D11ShaderResourceView(ShaderResourceView, Buffer);
 }
diff --git a/Engine/Source/Runtime/Windows/D3D11RHI/Private/D3D11VertexBuffer.cpp b/Engine/Source/Runtime/Windows/D3D11RHI/Private/D3D11VertexBuffer.cpp
--- a/Engine/Source/Runtime/Windows/D3D11RHI/Private/D3D11VertexBuffer.cpp
+++ b/Engine/Source/Runtime/Windows/D3D11RHI/Private/D3D11VertexBuffer.cpp
@@ -8,6 +8,11 @@
 
 FVertexBufferRHIRef FD3D11DynamicRHI::RHICreateVertexBuffer(uint32 Size,uint32 InUsage, FRHIResourceCreateInfo& CreateInfo)
 {
+	if (CreateInfo.bWithoutNativeResource)
+	{
+		return new FD3D11VertexBuffer();
+	}
+
 	// Explicitly check that the size is nonzero before allowing CreateVertexBuffer to opaquely fail.
 	check(Size > 0);
 
@@ -236,3 +241,18 @@ void FD3D11DynamicRHI::RHICopyVertexBuffer(FVertexBufferRHIParamRef SourceBuffer
 
 	GPUProfilingData.RegisterGPUWork(1);
 }
+
+void FD3D11DynamicRHI::RHITransferVertexBufferUnderlyingResource(FVertexBufferRHIParamRef DestVertexBuffer, FVertexBufferRHIParamRef SrcVertexBuffer)
+{
+	check(DestVertexBuffer);
+	FD3D11VertexBuffer* Dest = ResourceCast(DestVertexBuffer);
+	if (!SrcVertexBuffer)
+	{
+		Dest->ReleaseUnderlyingResource();
+	}
+	else
+	{
+		FD3D11VertexBuffer* Src = ResourceCast(SrcVertexBuffer);
+		Dest->Swap(*Src);
+	}
+}
diff --git a/Engine/Source/Runtime/Windows/D3D11RHI/Public/D3D11Resources.h b/Engine/Source/Runtime/Windows/D3D11RHI/Public/D3D11Resources.h
--- a/Engine/Source/Runtime/Windows/D3D11RHI/Public/D3D11Resources.h
+++ b/Engine/Source/Runtime/Windows/D3D11RHI/Public/D3D11Resources.h
@@ -694,6 +694,8 @@ public:
 	/** The index buffer resource */
 	TRefCountPtr<ID3D11Buffer> Resource;
 
+	FD3D11IndexBuffer() = default;
+
 	FD3D11IndexBuffer(ID3D11Buffer* InResource, uint32 InStride, uint32 InSize, uint32 InUsage)
 	: FRHIIndexBuffer(InStride,InSize,InUsage)
 	, Resource(InResource)
@@ -701,7 +703,26 @@ public:
 
 	virtual ~FD3D11IndexBuffer()
 	{
+		if (Resource)
+		{
+			UpdateBufferStats(Resource, false);
+		}
+	}
+
+	void Swap(FD3D11IndexBuffer& Other)
+	{
+		check(GetCurrentGPUAccess() == EResourceTransitionAccess::EReadable
+			&& GetCurrentGPUAccess() == Other.GetCurrentGPUAccess());
+		FRHIIndexBuffer::Swap(Other);
+		Resource.Swap(Other.Resource);
+	}
+
+	void ReleaseUnderlyingResource()
+	{
+		check(Resource);
 		UpdateBufferStats(Resource, false);
+		Resource = nullptr;
+		FRHIIndexBuffer::ReleaseUnderlyingResource();
 	}
 
 	// IRefCountedObject interface.
@@ -760,6 +781,8 @@ public:
 
 	TRefCountPtr<ID3D11Buffer> Resource;
 
+	FD3D11VertexBuffer() = default;
+
 	FD3D11VertexBuffer(ID3D11Buffer* InResource, uint32 InSize, uint32 InUsage)
 	: FRHIVertexBuffer(InSize,InUsage)
 	, Resource(InResource)
@@ -767,9 +790,28 @@ public:
 
 	virtual ~FD3D11VertexBuffer()
 	{
+		if (Resource)
+		{
+			UpdateBufferStats(Resource, false);
+		}
+	}
+
+	void Swap(FD3D11VertexBuffer& SrcBuffer)
+	{
+		check(GetCurrentGPUAccess() == EResourceTransitionAccess::EReadable
+			&& GetCurrentGPUAccess() == SrcBuffer.GetCurrentGPUAccess());
+		FRHIVertexBuffer::Swap(SrcBuffer);
+		Resource.Swap(SrcBuffer.Resource);
+	}
+
+	void ReleaseUnderlyingResource()
+	{
+		check(Resource);
 		UpdateBufferStats(Resource, false);
+		Resource = nullptr;
+		FRHIVertexBuffer::ReleaseUnderlyingResource();
 	}
-	
+
 	// IRefCountedObject interface.
 	virtual uint32 AddRef() const
 	{
