Try to fix static analysis warning

diff --git a/Engine/Source/Runtime/RHI/Public/RHICommandList.h b/Engine/Source/Runtime/RHI/Public/RHICommandList.h
--- a/Engine/Source/Runtime/RHI/Public/RHICommandList.h
+++ b/Engine/Source/Runtime/RHI/Public/RHICommandList.h
@@ -5033,7 +5033,14 @@ private:
 		{
 			Flush();
 		}
+#ifdef _MSC_VER
+#pragma warning(push)
+#pragma warning(disable : 6385) // Access is alawys in-bound due to the Flush above
+#endif
 		return UpdateInfos[NumBatched++];
+#ifdef _MSC_VER
+#pragma warning(pop)
+#endif
 	}
 };
 
