Disable Vulkan tessellation
#rnx

diff --git a/Engine/Shaders/Private/BasePassCommon.ush b/Engine/Shaders/Private/BasePassCommon.ush
--- a/Engine/Shaders/Private/BasePassCommon.ush
+++ b/Engine/Shaders/Private/BasePassCommon.ush
@@ -106,7 +106,7 @@ struct FSharedBasePassInterpolants
 #endif
 };
 
-#if PC_D3D || GL4_PROFILE || SWITCH_PROFILE || SWITCH_PROFILE_FORWARD || METAL_TESSELLATION_PROFILE || VULKAN_PROFILE_SM5
+#if PC_D3D || GL4_PROFILE || SWITCH_PROFILE || SWITCH_PROFILE_FORWARD || METAL_TESSELLATION_PROFILE
 
 /** Interpolants passed from the vertex shader to the pixel shader. */
 struct FBasePassInterpolantsVSToPS : FSharedBasePassInterpolants
diff --git a/Engine/Source/Runtime/RHI/Private/RHI.cpp b/Engine/Source/Runtime/RHI/Private/RHI.cpp
--- a/Engine/Source/Runtime/RHI/Private/RHI.cpp
+++ b/Engine/Source/Runtime/RHI/Private/RHI.cpp
@@ -796,7 +796,7 @@ RHI_API bool RHISupportsTessellation(const EShaderPlatform Platform)
 {
 	if (IsFeatureLevelSupported(Platform, ERHIFeatureLevel::SM5))
 	{
-		return (Platform == SP_PCD3D_SM5) || (Platform == SP_XBOXONE_D3D12) || (Platform == SP_OPENGL_SM5) || (Platform == SP_OPENGL_ES31_EXT) || (Platform == SP_METAL_SM5) || (IsVulkanSM5Platform(Platform));
+		return (Platform == SP_PCD3D_SM5) || (Platform == SP_XBOXONE_D3D12) || (Platform == SP_OPENGL_SM5) || (Platform == SP_OPENGL_ES31_EXT) || (Platform == SP_METAL_SM5) /*|| (IsVulkanSM5Platform(Platform))*/;
 	}
 	return false;
 }
