Fixed static analysis warning: Dereferencing NULL pointer 'DepthStencilRenderTarget.DepthStencilTarget'

diff --git a/Engine/Source/Runtime/RHI/Private/RHI.cpp b/Engine/Source/Runtime/RHI/Private/RHI.cpp
--- a/Engine/Source/Runtime/RHI/Private/RHI.cpp
+++ b/Engine/Source/Runtime/RHI/Private/RHI.cpp
@@ -1002,20 +1002,20 @@ void FRHIRenderPassInfo::Validate() const
 			// this check is incorrect for mobile, depth/stencil is intermediate and we don't want to store it to main memory
 			//ensure(StencilStore == ERenderTargetStoreAction::EStore);
 		}
+		
+		if (SubpassHint == ESubpassHint::DepthReadSubpass)
+		{
+			// for depth read sub-pass
+			// 1. render pass must have depth target
+			// 2. depth target must support InputAttachement
+			ensure((DepthStencilRenderTarget.DepthStencilTarget->GetFlags() & TexCreate_InputAttachmentRead) != 0);
+		}
 	}
 	else
 	{
 		ensure(DepthStencilRenderTarget.Action == EDepthStencilTargetActions::DontLoad_DontStore);
 		ensure(DepthStencilRenderTarget.ExclusiveDepthStencil == FExclusiveDepthStencil::DepthNop_StencilNop);
-	}
-
-	if (SubpassHint == ESubpassHint::DepthReadSubpass)
-	{
-		// for depth read sub-pass
-		// 1. render pass must have depth target
-		// 2. depth target must support InputAttachement
-		ensure(DepthStencilRenderTarget.DepthStencilTarget);
-		ensure((DepthStencilRenderTarget.DepthStencilTarget->GetFlags() & TexCreate_InputAttachmentRead) != 0);
+		ensure(SubpassHint != ESubpassHint::DepthReadSubpass);
 	}
 }
 #endif
