Fix incorrect size passed to FMemory::MemZero, causing FDataDrivenShaderPlatformInfo to not be fully initialized.

diff --git a/Engine/Source/Runtime/RHI/Public/RHIDefinitions.h b/Engine/Source/Runtime/RHI/Public/RHIDefinitions.h
--- a/Engine/Source/Runtime/RHI/Public/RHIDefinitions.h
+++ b/Engine/Source/Runtime/RHI/Public/RHIDefinitions.h
@@ -166,7 +166,7 @@ struct RHI_API FDataDrivenShaderPlatformInfo
 
 	FDataDrivenShaderPlatformInfo()
 	{
-		FMemory::Memzero(this, sizeof(this));
+		FMemory::Memzero(this, sizeof(*this));
 		MaxFeatureLevel = ERHIFeatureLevel::Num;
 	}
 
