UE4.23 - Deprecate FUniformBufferRHIParamRef, FIndexBufferRHIParamRef, FVertexBufferRHIParamRef, FStructuredBufferRHIParamRef, FTextureRHIParamRef, FTexture2DRHIParamRef, FTexture2DArrayRHIParamRef, F

diff --git a/Engine/Source/Runtime/RHI/Public/DynamicRHI.h b/Engine/Source/Runtime/RHI/Public/DynamicRHI.h
--- a/Engine/Source/Runtime/RHI/Public/DynamicRHI.h
+++ b/Engine/Source/Runtime/RHI/Public/DynamicRHI.h
@@ -30,7 +30,7 @@ struct FRHIGPUMask;
 /** Struct to hold common data between begin/end updatetexture3d */
 struct FUpdateTexture3DData
 {
-	FUpdateTexture3DData(FTexture3DRHIParamRef InTexture, uint32 InMipIndex, const struct FUpdateTextureRegion3D& InUpdateRegion, uint32 InSourceRowPitch, uint32 InSourceDepthPitch, uint8* InSourceData, uint32 InDataSizeBytes, uint32 InFrameNumber )
+	FUpdateTexture3DData(FRHITexture3D* InTexture, uint32 InMipIndex, const struct FUpdateTextureRegion3D& InUpdateRegion, uint32 InSourceRowPitch, uint32 InSourceDepthPitch, uint8* InSourceData, uint32 InDataSizeBytes, uint32 InFrameNumber )
 		: Texture(InTexture)
 		, MipIndex(InMipIndex)
 		, UpdateRegion(InUpdateRegion)
@@ -42,7 +42,7 @@ struct FUpdateTexture3DData
 	{
 	}
 
-	FTexture3DRHIParamRef Texture;
+	FRHITexture3D* Texture;
 	uint32 MipIndex;
 	FUpdateTextureRegion3D UpdateRegion;
 	uint32 RowPitch;
@@ -436,23 +436,23 @@ public:
 	// FlushType: Thread safe, but varies depending on the RHI
 	virtual FUniformBufferRHIRef RHICreateUniformBuffer(const void* Contents, const FRHIUniformBufferLayout& Layout, EUniformBufferUsage Usage, EUniformBufferValidation Validation) = 0;
 
-	virtual void RHIUpdateUniformBuffer(FUniformBufferRHIParamRef UniformBufferRHI, const void* Contents) = 0;
+	virtual void RHIUpdateUniformBuffer(FRHIUniformBuffer* UniformBufferRHI, const void* Contents) = 0;
 
 	// FlushType: Wait RHI Thread
 	virtual FIndexBufferRHIRef RHICreateIndexBuffer(uint32 Stride, uint32 Size, uint32 InUsage, FRHIResourceCreateInfo& CreateInfo) = 0;
 
 	// FlushType: Flush RHI Thread
-	virtual void* RHILockIndexBuffer(FIndexBufferRHIParamRef IndexBuffer, uint32 Offset, uint32 Size, EResourceLockMode LockMode) = 0;
+	virtual void* RHILockIndexBuffer(FRHIIndexBuffer* IndexBuffer, uint32 Offset, uint32 Size, EResourceLockMode LockMode) = 0;
 
 	// FlushType: Flush RHI Thread
-	virtual void RHIUnlockIndexBuffer(FIndexBufferRHIParamRef IndexBuffer) = 0;
+	virtual void RHIUnlockIndexBuffer(FRHIIndexBuffer* IndexBuffer) = 0;
 
 	/**
 	* Transfer metadata and underlying resource from src to dest and release any resource owned by dest.
 	* @param DestIndexBuffer - the index buffer to update
 	* @param SrcIndexBuffer - don't use after call. If null, will release any resource owned by DestIndexBuffer
 	*/
-	virtual void RHITransferIndexBufferUnderlyingResource(FIndexBufferRHIParamRef DestIndexBuffer, FIndexBufferRHIParamRef SrcIndexBuffer);
+	virtual void RHITransferIndexBufferUnderlyingResource(FRHIIndexBuffer* DestIndexBuffer, FRHIIndexBuffer* SrcIndexBuffer);
 
 	/**
 	* @param ResourceArray - An optional pointer to a resource array containing the resource's data.
@@ -461,21 +461,21 @@ public:
 	virtual FVertexBufferRHIRef RHICreateVertexBuffer(uint32 Size, uint32 InUsage, FRHIResourceCreateInfo& CreateInfo) = 0;
 
 	// FlushType: Flush RHI Thread
-	virtual void* RHILockVertexBuffer(FVertexBufferRHIParamRef VertexBuffer, uint32 Offset, uint32 SizeRHI, EResourceLockMode LockMode) = 0;
+	virtual void* RHILockVertexBuffer(FRHIVertexBuffer* VertexBuffer, uint32 Offset, uint32 SizeRHI, EResourceLockMode LockMode) = 0;
 
 	// FlushType: Flush RHI Thread
-	virtual void RHIUnlockVertexBuffer(FVertexBufferRHIParamRef VertexBuffer) = 0;
+	virtual void RHIUnlockVertexBuffer(FRHIVertexBuffer* VertexBuffer) = 0;
 
 	/** Copies the contents of one vertex buffer to another vertex buffer.  They must have identical sizes. */
 	// FlushType: Flush Immediate (seems dangerous)
-	virtual void RHICopyVertexBuffer(FVertexBufferRHIParamRef SourceBuffer, FVertexBufferRHIParamRef DestBuffer) = 0;
+	virtual void RHICopyVertexBuffer(FRHIVertexBuffer* SourceBuffer, FRHIVertexBuffer* DestBuffer) = 0;
 
 	/**
 	 * Transfer metadata and underlying resource from src to dest and release any resource owned by dest.
 	 * @param DestVertexBuffer - the vertex buffer to update
 	 * @param SrcVertexBuffer - don't use after call. If null, will release any resource owned by DestVertexBuffer
 	 */
-	virtual void RHITransferVertexBufferUnderlyingResource(FVertexBufferRHIParamRef DestVertexBuffer, FVertexBufferRHIParamRef SrcVertexBuffer);
+	virtual void RHITransferVertexBufferUnderlyingResource(FRHIVertexBuffer* DestVertexBuffer, FRHIVertexBuffer* SrcVertexBuffer);
 
 	/**
 	* @param ResourceArray - An optional pointer to a resource array containing the resource's data.
@@ -484,44 +484,44 @@ public:
 	virtual FStructuredBufferRHIRef RHICreateStructuredBuffer(uint32 Stride, uint32 Size, uint32 InUsage, FRHIResourceCreateInfo& CreateInfo) = 0;
 
 	// FlushType: Flush RHI Thread
-	virtual void* RHILockStructuredBuffer(FStructuredBufferRHIParamRef StructuredBuffer, uint32 Offset, uint32 SizeRHI, EResourceLockMode LockMode) = 0;
+	virtual void* RHILockStructuredBuffer(FRHIStructuredBuffer* StructuredBuffer, uint32 Offset, uint32 SizeRHI, EResourceLockMode LockMode) = 0;
 
 	// FlushType: Flush RHI Thread
-	virtual void RHIUnlockStructuredBuffer(FStructuredBufferRHIParamRef StructuredBuffer) = 0;
+	virtual void RHIUnlockStructuredBuffer(FRHIStructuredBuffer* StructuredBuffer) = 0;
 
 	/** Creates an unordered access view of the given structured buffer. */
 	// FlushType: Wait RHI Thread
-	virtual FUnorderedAccessViewRHIRef RHICreateUnorderedAccessView(FStructuredBufferRHIParamRef StructuredBuffer, bool bUseUAVCounter, bool bAppendBuffer) = 0;
+	virtual FUnorderedAccessViewRHIRef RHICreateUnorderedAccessView(FRHIStructuredBuffer* StructuredBuffer, bool bUseUAVCounter, bool bAppendBuffer) = 0;
 
 	/** Creates an unordered access view of the given texture. */
 	// FlushType: Wait RHI Thread
-	virtual FUnorderedAccessViewRHIRef RHICreateUnorderedAccessView(FTextureRHIParamRef Texture, uint32 MipLevel) = 0;
+	virtual FUnorderedAccessViewRHIRef RHICreateUnorderedAccessView(FRHITexture* Texture, uint32 MipLevel) = 0;
 
 	/** Creates an unordered access view of the given vertex buffer. */
 	// FlushType: Wait RHI Thread
-	virtual FUnorderedAccessViewRHIRef RHICreateUnorderedAccessView(FVertexBufferRHIParamRef VertexBuffer, uint8 Format) = 0;
+	virtual FUnorderedAccessViewRHIRef RHICreateUnorderedAccessView(FRHIVertexBuffer* VertexBuffer, uint8 Format) = 0;
 
 	/** Creates an unordered access view of the given index buffer. */
 	// FlushType: Wait RHI Thread
-	virtual FUnorderedAccessViewRHIRef RHICreateUnorderedAccessView(FIndexBufferRHIParamRef IndexBuffer, uint8 Format) = 0;
+	virtual FUnorderedAccessViewRHIRef RHICreateUnorderedAccessView(FRHIIndexBuffer* IndexBuffer, uint8 Format) = 0;
 
 	/** Creates a shader resource view of the given structured buffer. */
 	// FlushType: Wait RHI Thread
-	virtual FShaderResourceViewRHIRef RHICreateShaderResourceView(FStructuredBufferRHIParamRef StructuredBuffer) = 0;
+	virtual FShaderResourceViewRHIRef RHICreateShaderResourceView(FRHIStructuredBuffer* StructuredBuffer) = 0;
 
 	/** Creates a shader resource view of the given vertex buffer. */
 	// FlushType: Wait RHI Thread
-	virtual FShaderResourceViewRHIRef RHICreateShaderResourceView(FVertexBufferRHIParamRef VertexBuffer, uint32 Stride, uint8 Format) = 0;
+	virtual FShaderResourceViewRHIRef RHICreateShaderResourceView(FRHIVertexBuffer* VertexBuffer, uint32 Stride, uint8 Format) = 0;
 
 	/** Creates a shader resource view of the given index buffer. */
 	// FlushType: Wait RHI Thread
-	virtual FShaderResourceViewRHIRef RHICreateShaderResourceView(FIndexBufferRHIParamRef Buffer) = 0;
+	virtual FShaderResourceViewRHIRef RHICreateShaderResourceView(FRHIIndexBuffer* Buffer) = 0;
 
 	// Must be called on RHI thread timeline
 	// Make sure to call RHIThreadFence(true) afterwards so that parallel translation doesn't refer old resources
-	virtual void RHIUpdateShaderResourceView(FRHIShaderResourceView* SRV, FVertexBufferRHIParamRef VertexBuffer, uint32 Stride, uint8 Format);
+	virtual void RHIUpdateShaderResourceView(FRHIShaderResourceView* SRV, FRHIVertexBuffer* VertexBuffer, uint32 Stride, uint8 Format);
 
-	virtual void RHIUpdateShaderResourceView(FRHIShaderResourceView* SRV, FIndexBufferRHIParamRef IndexBuffer);
+	virtual void RHIUpdateShaderResourceView(FRHIShaderResourceView* SRV, FRHIIndexBuffer* IndexBuffer);
 
 	/**
 	* Computes the total size of a 2D texture with the specified parameters.
@@ -616,7 +616,7 @@ public:
 	* Creates an FStructuredBuffer for the RT write mask of a render target
 	* @param RenderTarget - the RT to create the buffer for
 	*/
-	virtual FStructuredBufferRHIRef RHICreateRTWriteMaskBuffer(FTexture2DRHIParamRef RenderTarget)
+	virtual FStructuredBufferRHIRef RHICreateRTWriteMaskBuffer(FRHITexture2D* RenderTarget)
 	{
 		return nullptr;
 	}
@@ -645,7 +645,7 @@ public:
 	* @param SrcTexture2D - source texture
 	*/
 	// FlushType: Flush RHI Thread
-	virtual void RHICopySharedMips(FTexture2DRHIParamRef DestTexture2D, FTexture2DRHIParamRef SrcTexture2D) = 0;
+	virtual void RHICopySharedMips(FRHITexture2D* DestTexture2D, FRHITexture2D* SrcTexture2D) = 0;
 
 	/**
 	* Synchronizes the content of a texture resource between two GPUs using a copy operation.
@@ -656,7 +656,7 @@ public:
 	* @param PullData - whether the source writes the data to the dest, or the dest reads the data from the source.
 	*/
 	// FlushType: Flush RHI Thread
-	virtual void RHITransferTexture(FTexture2DRHIParamRef Texture, FIntRect Rect, uint32 SrcGPUIndex, uint32 DestGPUIndex, bool PullData) { unimplemented(); };
+	virtual void RHITransferTexture(FRHITexture2D* Texture, FIntRect Rect, uint32 SrcGPUIndex, uint32 DestGPUIndex, bool PullData) { unimplemented(); };
 
 	/**
 	* Creates a Array RHI texture resource
@@ -703,13 +703,13 @@ public:
 	* @param Ref may be 0
 	*/
 	// FlushType: Thread safe
-	virtual void RHIGetResourceInfo(FTextureRHIParamRef Ref, FRHIResourceInfo& OutInfo) = 0;
+	virtual void RHIGetResourceInfo(FRHITexture* Ref, FRHIResourceInfo& OutInfo) = 0;
 
 	/**
 	* Creates a shader resource view for a texture
 	*/
 	// FlushType: Wait RHI Thread
-	virtual FShaderResourceViewRHIRef RHICreateShaderResourceView(FTextureRHIParamRef Texture2DRHI, const FRHITextureSRVCreateInfo& CreateInfo) = 0;
+	virtual FShaderResourceViewRHIRef RHICreateShaderResourceView(FRHITexture* Texture2DRHI, const FRHITextureSRVCreateInfo& CreateInfo) = 0;
 
 	/**
 	* Generates mip maps for a texture.
@@ -717,7 +717,7 @@ public:
 	// FlushType: Flush Immediate (NP: this should be queued on the command list for RHI thread execution, not flushed)
 
 	//UE_DEPRECATED(4.23, "This function is deprecated and will be removed in future releases. Renderer version implemented.")
-	virtual void RHIGenerateMips(FTextureRHIParamRef Texture) {}
+	virtual void RHIGenerateMips(FRHITexture* Texture) {}
 
 	/**
 	* Computes the size in memory required by a given texture.
@@ -726,7 +726,7 @@ public:
 	* @return					- Size in Bytes
 	*/
 	// FlushType: Thread safe
-	virtual uint32 RHIComputeMemorySize(FTextureRHIParamRef TextureRHI) = 0;
+	virtual uint32 RHIComputeMemorySize(FRHITexture* TextureRHI) = 0;
 
 	/**
 	* Starts an asynchronous texture reallocation. It may complete immediately if the reallocation
@@ -745,7 +745,7 @@ public:
 	*/
 	// FlushType: Flush RHI Thread
 	// NP: Note that no RHI currently implements this as an async call, we should simplify the API.
-	virtual FTexture2DRHIRef RHIAsyncReallocateTexture2D(FTexture2DRHIParamRef Texture2D, int32 NewMipCount, int32 NewSizeX, int32 NewSizeY, FThreadSafeCounter* RequestStatus) = 0;
+	virtual FTexture2DRHIRef RHIAsyncReallocateTexture2D(FRHITexture2D* Texture2D, int32 NewMipCount, int32 NewSizeX, int32 NewSizeY, FThreadSafeCounter* RequestStatus) = 0;
 
 	/**
 	* Finalizes an async reallocation request.
@@ -759,7 +759,7 @@ public:
 	*	TexRealloc_InProgress	Reallocation is still in progress, try again later
 	*/
 	// FlushType: Wait RHI Thread
-	virtual ETextureReallocationStatus RHIFinalizeAsyncReallocateTexture2D(FTexture2DRHIParamRef Texture2D, bool bBlockUntilCompleted) = 0;
+	virtual ETextureReallocationStatus RHIFinalizeAsyncReallocateTexture2D(FRHITexture2D* Texture2D, bool bBlockUntilCompleted) = 0;
 
 	/**
 	* Cancels an async reallocation for the specified texture.
@@ -770,7 +770,7 @@ public:
 	* @return						Reallocation status
 	*/
 	// FlushType: Wait RHI Thread
-	virtual ETextureReallocationStatus RHICancelAsyncReallocateTexture2D(FTexture2DRHIParamRef Texture2D, bool bBlockUntilCompleted) = 0;
+	virtual ETextureReallocationStatus RHICancelAsyncReallocateTexture2D(FRHITexture2D* Texture2D, bool bBlockUntilCompleted) = 0;
 
 	/**
 	* Locks an RHI texture's mip-map for read/write operations on the CPU
@@ -782,7 +782,7 @@ public:
 	* @return pointer to the CPU accessible resource data
 	*/
 	// FlushType: Flush RHI Thread
-	virtual void* RHILockTexture2D(FTexture2DRHIParamRef Texture, uint32 MipIndex, EResourceLockMode LockMode, uint32& DestStride, bool bLockWithinMiptail) = 0;
+	virtual void* RHILockTexture2D(FRHITexture2D* Texture, uint32 MipIndex, EResourceLockMode LockMode, uint32& DestStride, bool bLockWithinMiptail) = 0;
 
 	/**
 	* Unlocks a previously locked RHI texture resource
@@ -791,7 +791,7 @@ public:
 	* @param bLockWithinMiptail - for platforms that support packed miptails allow locking of individual mip levels within the miptail
 	*/
 	// FlushType: Flush RHI Thread
-	virtual void RHIUnlockTexture2D(FTexture2DRHIParamRef Texture, uint32 MipIndex, bool bLockWithinMiptail) = 0;
+	virtual void RHIUnlockTexture2D(FRHITexture2D* Texture, uint32 MipIndex, bool bLockWithinMiptail) = 0;
 
 	/**
 	* Locks an RHI texture's mip-map for read/write operations on the CPU
@@ -803,7 +803,7 @@ public:
 	* @return pointer to the CPU accessible resource data
 	*/
 	// FlushType: Flush RHI Thread
-	virtual void* RHILockTexture2DArray(FTexture2DArrayRHIParamRef Texture, uint32 TextureIndex, uint32 MipIndex, EResourceLockMode LockMode, uint32& DestStride, bool bLockWithinMiptail) = 0;
+	virtual void* RHILockTexture2DArray(FRHITexture2DArray* Texture, uint32 TextureIndex, uint32 MipIndex, EResourceLockMode LockMode, uint32& DestStride, bool bLockWithinMiptail) = 0;
 
 	/**
 	* Unlocks a previously locked RHI texture resource
@@ -812,7 +812,7 @@ public:
 	* @param bLockWithinMiptail - for platforms that support packed miptails allow locking of individual mip levels within the miptail
 	*/
 	// FlushType: Flush RHI Thread
-	virtual void RHIUnlockTexture2DArray(FTexture2DArrayRHIParamRef Texture, uint32 TextureIndex, uint32 MipIndex, bool bLockWithinMiptail) = 0;
+	virtual void RHIUnlockTexture2DArray(FRHITexture2DArray* Texture, uint32 TextureIndex, uint32 MipIndex, bool bLockWithinMiptail) = 0;
 
 	/**
 	* Updates a region of a 2D texture from system memory
@@ -823,7 +823,7 @@ public:
 	* @param SourceData - source image data, starting at the upper left corner of the source rectangle (in same pixel format as texture)
 	*/
 	// FlushType: Flush RHI Thread
-	virtual void RHIUpdateTexture2D(FTexture2DRHIParamRef Texture, uint32 MipIndex, const struct FUpdateTextureRegion2D& UpdateRegion, uint32 SourcePitch, const uint8* SourceData) = 0;
+	virtual void RHIUpdateTexture2D(FRHITexture2D* Texture, uint32 MipIndex, const struct FUpdateTextureRegion2D& UpdateRegion, uint32 SourcePitch, const uint8* SourceData) = 0;
 
 	/**
 	* Updates a region of a 2D texture from GPU memory provided by the given buffer (may not be implemented on every platform)
@@ -834,7 +834,7 @@ public:
 	* @param Buffer, BufferOffset - source image data, starting at the upper left corner of the source rectangle (in same pixel format as texture)
 	*/
 	// FlushType: Flush RHI Thread
-	virtual void RHIUpdateFromBufferTexture2D(FTexture2DRHIParamRef Texture, uint32 MipIndex, const struct FUpdateTextureRegion2D& UpdateRegion, uint32 SourcePitch, FStructuredBufferRHIParamRef Buffer, uint32 BufferOffset) {}
+	virtual void RHIUpdateFromBufferTexture2D(FRHITexture2D* Texture, uint32 MipIndex, const struct FUpdateTextureRegion2D& UpdateRegion, uint32 SourcePitch, FRHIStructuredBuffer* Buffer, uint32 BufferOffset) {}
 
 	/**
 	* Updates a region of a 3D texture from system memory
@@ -846,7 +846,7 @@ public:
 	* @param SourceData - source image data, starting at the upper left corner of the source rectangle (in same pixel format as texture)
 	*/
 	// FlushType: Flush RHI Thread
-	virtual void RHIUpdateTexture3D(FTexture3DRHIParamRef Texture, uint32 MipIndex, const struct FUpdateTextureRegion3D& UpdateRegion, uint32 SourceRowPitch, uint32 SourceDepthPitch, const uint8* SourceData) = 0;
+	virtual void RHIUpdateTexture3D(FRHITexture3D* Texture, uint32 MipIndex, const struct FUpdateTextureRegion3D& UpdateRegion, uint32 SourceRowPitch, uint32 SourceDepthPitch, const uint8* SourceData) = 0;
 
 	/**
 	* Creates a Cube RHI texture resource
@@ -879,7 +879,7 @@ public:
 	* @return pointer to the CPU accessible resource data
 	*/
 	// FlushType: Flush RHI Thread
-	virtual void* RHILockTextureCubeFace(FTextureCubeRHIParamRef Texture, uint32 FaceIndex, uint32 ArrayIndex, uint32 MipIndex, EResourceLockMode LockMode, uint32& DestStride, bool bLockWithinMiptail) = 0;
+	virtual void* RHILockTextureCubeFace(FRHITextureCube* Texture, uint32 FaceIndex, uint32 ArrayIndex, uint32 MipIndex, EResourceLockMode LockMode, uint32& DestStride, bool bLockWithinMiptail) = 0;
 
 	/**
 	* Unlocks a previously locked RHI texture resource
@@ -888,10 +888,10 @@ public:
 	* @param bLockWithinMiptail - for platforms that support packed miptails allow locking of individual mip levels within the miptail
 	*/
 	// FlushType: Flush RHI Thread
-	virtual void RHIUnlockTextureCubeFace(FTextureCubeRHIParamRef Texture, uint32 FaceIndex, uint32 ArrayIndex, uint32 MipIndex, bool bLockWithinMiptail) = 0;
+	virtual void RHIUnlockTextureCubeFace(FRHITextureCube* Texture, uint32 FaceIndex, uint32 ArrayIndex, uint32 MipIndex, bool bLockWithinMiptail) = 0;
 
 	// FlushType: Thread safe
-	virtual void RHIBindDebugLabelName(FTextureRHIParamRef Texture, const TCHAR* Name) = 0;
+	virtual void RHIBindDebugLabelName(FRHITexture* Texture, const TCHAR* Name) = 0;
 	virtual void RHIBindDebugLabelName(FRHIUnorderedAccessView* UnorderedAccessViewRHI, const TCHAR* Name) {}
 
 	/**
@@ -899,10 +899,10 @@ public:
 	* If the format or texture type is unsupported the OutData array will be size 0
 	*/
 	// FlushType: Flush Immediate (seems wrong)
-	virtual void RHIReadSurfaceData(FTextureRHIParamRef Texture, FIntRect Rect, TArray<FColor>& OutData, FReadSurfaceDataFlags InFlags) = 0;
+	virtual void RHIReadSurfaceData(FRHITexture* Texture, FIntRect Rect, TArray<FColor>& OutData, FReadSurfaceDataFlags InFlags) = 0;
 
 	// Default fallback; will not work for non-8-bit surfaces and it's extremely slow.
-	virtual void RHIReadSurfaceData(FTextureRHIParamRef Texture, FIntRect Rect, TArray<FLinearColor>& OutData, FReadSurfaceDataFlags InFlags)
+	virtual void RHIReadSurfaceData(FRHITexture* Texture, FIntRect Rect, TArray<FLinearColor>& OutData, FReadSurfaceDataFlags InFlags)
 	{
 		TArray<FColor> TempData;
 		RHIReadSurfaceData(Texture, Rect, TempData, InFlags);
@@ -915,17 +915,17 @@ public:
 
 	/** Watch out for OutData to be 0 (can happen on DXGI_ERROR_DEVICE_REMOVED), don't call RHIUnmapStagingSurface in that case. */
 	// FlushType: Flush Immediate (seems wrong)
-	virtual void RHIMapStagingSurface(FTextureRHIParamRef Texture, void*& OutData, int32& OutWidth, int32& OutHeight) = 0;
+	virtual void RHIMapStagingSurface(FRHITexture* Texture, void*& OutData, int32& OutWidth, int32& OutHeight) = 0;
 
 	/** call after a succesful RHIMapStagingSurface() call */
 	// FlushType: Flush Immediate (seems wrong)
-	virtual void RHIUnmapStagingSurface(FTextureRHIParamRef Texture) = 0;
+	virtual void RHIUnmapStagingSurface(FRHITexture* Texture) = 0;
 
 	// FlushType: Flush Immediate (seems wrong)
-	virtual void RHIReadSurfaceFloatData(FTextureRHIParamRef Texture, FIntRect Rect, TArray<FFloat16Color>& OutData, ECubeFace CubeFace, int32 ArrayIndex, int32 MipIndex) = 0;
+	virtual void RHIReadSurfaceFloatData(FRHITexture* Texture, FIntRect Rect, TArray<FFloat16Color>& OutData, ECubeFace CubeFace, int32 ArrayIndex, int32 MipIndex) = 0;
 
 	// FlushType: Flush Immediate (seems wrong)
-	virtual void RHIRead3DSurfaceFloatData(FTextureRHIParamRef Texture, FIntRect Rect, FIntPoint ZMinMax, TArray<FFloat16Color>& OutData) = 0;
+	virtual void RHIRead3DSurfaceFloatData(FRHITexture* Texture, FIntRect Rect, FIntPoint ZMinMax, TArray<FFloat16Color>& OutData) = 0;
 
 	// FlushType: Wait RHI Thread
 	virtual FRenderQueryRHIRef RHICreateRenderQuery(ERenderQueryType QueryType) = 0;
@@ -949,12 +949,12 @@ public:
 		return FUnorderedAccessViewRHIRef();
 	}
 	
-	virtual FTexture2DRHIRef RHIGetFMaskTexture(FTextureRHIParamRef SourceTextureRHI)
+	virtual FTexture2DRHIRef RHIGetFMaskTexture(FRHITexture* SourceTextureRHI)
 	{
 		return nullptr;
 	}
 
-	virtual void RHIAliasTextureResources(FTextureRHIParamRef DestTexture, FTextureRHIParamRef SrcTexture)
+	virtual void RHIAliasTextureResources(FRHITexture* DestTexture, FRHITexture* SrcTexture)
 	{
 		checkNoEntry();
 	}
@@ -1003,7 +1003,7 @@ public:
 	/** Sets stream output targets, for use with a geometry shader created with RHICreateGeometryShaderWithStreamOutput. */
 	//@todo this should be a CMDLIST method
 	// FlushType: Flush Immediate (seems wrong)
-	virtual void RHISetStreamOutTargets(uint32 NumTargets, const FVertexBufferRHIParamRef* VertexBuffers, const uint32* Offsets) = 0;
+	virtual void RHISetStreamOutTargets(uint32 NumTargets, FRHIVertexBuffer* const* VertexBuffers, const uint32* Offsets) = 0;
 
 	// Blocks the CPU until the GPU catches up and goes idle.
 	// FlushType: Flush Immediate (seems wrong)
@@ -1056,7 +1056,7 @@ public:
 	* @param FirstMip - the first mip that should be in memory
 	*/
 	// FlushType: Wait RHI Thread
-	virtual void RHIVirtualTextureSetFirstMipInMemory(FTexture2DRHIParamRef Texture, uint32 FirstMip) = 0;
+	virtual void RHIVirtualTextureSetFirstMipInMemory(FRHITexture2D* Texture, uint32 FirstMip) = 0;
 
 	/**
 	* Function that can be used to update which is the first visible mip to the GPU.
@@ -1064,7 +1064,7 @@ public:
 	* @param FirstMip - the first mip that should be visible to the GPU
 	*/
 	// FlushType: Wait RHI Thread
-	virtual void RHIVirtualTextureSetFirstMipVisible(FTexture2DRHIParamRef Texture, uint32 FirstMip) = 0;
+	virtual void RHIVirtualTextureSetFirstMipVisible(FRHITexture2D* Texture, uint32 FirstMip) = 0;
 
 	/**
 	* Called once per frame just before deferred deletion in FRHIResource::FlushPendingDeletes
@@ -1117,16 +1117,16 @@ public:
 
 	virtual FVertexBufferRHIRef CreateVertexBuffer_RenderThread(class FRHICommandListImmediate& RHICmdList, uint32 Size, uint32 InUsage, FRHIResourceCreateInfo& CreateInfo);
 	virtual FStructuredBufferRHIRef CreateStructuredBuffer_RenderThread(class FRHICommandListImmediate& RHICmdList, uint32 Stride, uint32 Size, uint32 InUsage, FRHIResourceCreateInfo& CreateInfo);
-	virtual FShaderResourceViewRHIRef CreateShaderResourceView_RenderThread(class FRHICommandListImmediate& RHICmdList, FVertexBufferRHIParamRef VertexBuffer, uint32 Stride, uint8 Format);
-	virtual FShaderResourceViewRHIRef CreateShaderResourceView_RenderThread(class FRHICommandListImmediate& RHICmdList, FIndexBufferRHIParamRef Buffer);
-	virtual void* LockVertexBuffer_RenderThread(class FRHICommandListImmediate& RHICmdList, FVertexBufferRHIParamRef VertexBuffer, uint32 Offset, uint32 SizeRHI, EResourceLockMode LockMode);
-	virtual void UnlockVertexBuffer_RenderThread(class FRHICommandListImmediate& RHICmdList, FVertexBufferRHIParamRef VertexBuffer);
-	virtual FTexture2DRHIRef AsyncReallocateTexture2D_RenderThread(class FRHICommandListImmediate& RHICmdList, FTexture2DRHIParamRef Texture2D, int32 NewMipCount, int32 NewSizeX, int32 NewSizeY, FThreadSafeCounter* RequestStatus);
-	virtual ETextureReallocationStatus FinalizeAsyncReallocateTexture2D_RenderThread(class FRHICommandListImmediate& RHICmdList, FTexture2DRHIParamRef Texture2D, bool bBlockUntilCompleted);
-	virtual ETextureReallocationStatus CancelAsyncReallocateTexture2D_RenderThread(class FRHICommandListImmediate& RHICmdList, FTexture2DRHIParamRef Texture2D, bool bBlockUntilCompleted);
+	virtual FShaderResourceViewRHIRef CreateShaderResourceView_RenderThread(class FRHICommandListImmediate& RHICmdList, FRHIVertexBuffer* VertexBuffer, uint32 Stride, uint8 Format);
+	virtual FShaderResourceViewRHIRef CreateShaderResourceView_RenderThread(class FRHICommandListImmediate& RHICmdList, FRHIIndexBuffer* Buffer);
+	virtual void* LockVertexBuffer_RenderThread(class FRHICommandListImmediate& RHICmdList, FRHIVertexBuffer* VertexBuffer, uint32 Offset, uint32 SizeRHI, EResourceLockMode LockMode);
+	virtual void UnlockVertexBuffer_RenderThread(class FRHICommandListImmediate& RHICmdList, FRHIVertexBuffer* VertexBuffer);
+	virtual FTexture2DRHIRef AsyncReallocateTexture2D_RenderThread(class FRHICommandListImmediate& RHICmdList, FRHITexture2D* Texture2D, int32 NewMipCount, int32 NewSizeX, int32 NewSizeY, FThreadSafeCounter* RequestStatus);
+	virtual ETextureReallocationStatus FinalizeAsyncReallocateTexture2D_RenderThread(class FRHICommandListImmediate& RHICmdList, FRHITexture2D* Texture2D, bool bBlockUntilCompleted);
+	virtual ETextureReallocationStatus CancelAsyncReallocateTexture2D_RenderThread(class FRHICommandListImmediate& RHICmdList, FRHITexture2D* Texture2D, bool bBlockUntilCompleted);
 	virtual FIndexBufferRHIRef CreateIndexBuffer_RenderThread(class FRHICommandListImmediate& RHICmdList, uint32 Stride, uint32 Size, uint32 InUsage, FRHIResourceCreateInfo& CreateInfo);
-	virtual void* LockIndexBuffer_RenderThread(class FRHICommandListImmediate& RHICmdList, FIndexBufferRHIParamRef IndexBuffer, uint32 Offset, uint32 SizeRHI, EResourceLockMode LockMode);
-	virtual void UnlockIndexBuffer_RenderThread(class FRHICommandListImmediate& RHICmdList, FIndexBufferRHIParamRef IndexBuffer);
+	virtual void* LockIndexBuffer_RenderThread(class FRHICommandListImmediate& RHICmdList, FRHIIndexBuffer* IndexBuffer, uint32 Offset, uint32 SizeRHI, EResourceLockMode LockMode);
+	virtual void UnlockIndexBuffer_RenderThread(class FRHICommandListImmediate& RHICmdList, FRHIIndexBuffer* IndexBuffer);
 	virtual FVertexShaderRHIRef CreateVertexShader_RenderThread(class FRHICommandListImmediate& RHICmdList, const TArray<uint8>& Code);
 	virtual FVertexShaderRHIRef CreateVertexShader_RenderThread(class FRHICommandListImmediate& RHICmdList, FRHIShaderLibrary* Library, FSHAHash Hash);
 	virtual FPixelShaderRHIRef CreatePixelShader_RenderThread(class FRHICommandListImmediate& RHICmdList, const TArray<uint8>& Code);
@@ -1141,17 +1141,17 @@ public:
 	virtual FHullShaderRHIRef CreateHullShader_RenderThread(class FRHICommandListImmediate& RHICmdList, FRHIShaderLibrary* Library, FSHAHash Hash);
 	virtual FDomainShaderRHIRef CreateDomainShader_RenderThread(class FRHICommandListImmediate& RHICmdList, const TArray<uint8>& Code);
 	virtual FDomainShaderRHIRef CreateDomainShader_RenderThread(class FRHICommandListImmediate& RHICmdList, FRHIShaderLibrary* Library, FSHAHash Hash);
-	virtual void* LockTexture2D_RenderThread(class FRHICommandListImmediate& RHICmdList, FTexture2DRHIParamRef Texture, uint32 MipIndex, EResourceLockMode LockMode, uint32& DestStride, bool bLockWithinMiptail, bool bNeedsDefaultRHIFlush = true);
-	virtual void UnlockTexture2D_RenderThread(class FRHICommandListImmediate& RHICmdList, FTexture2DRHIParamRef Texture, uint32 MipIndex, bool bLockWithinMiptail, bool bNeedsDefaultRHIFlush = true);
-	virtual void UpdateTexture2D_RenderThread(class FRHICommandListImmediate& RHICmdList, FTexture2DRHIParamRef Texture, uint32 MipIndex, const struct FUpdateTextureRegion2D& UpdateRegion, uint32 SourcePitch, const uint8* SourceData);
-	virtual void UpdateFromBufferTexture2D_RenderThread(class FRHICommandListImmediate& RHICmdList, FTexture2DRHIParamRef Texture, uint32 MipIndex, const struct FUpdateTextureRegion2D& UpdateRegion, uint32 SourcePitch, FStructuredBufferRHIParamRef Buffer, uint32 BufferOffset);
+	virtual void* LockTexture2D_RenderThread(class FRHICommandListImmediate& RHICmdList, FRHITexture2D* Texture, uint32 MipIndex, EResourceLockMode LockMode, uint32& DestStride, bool bLockWithinMiptail, bool bNeedsDefaultRHIFlush = true);
+	virtual void UnlockTexture2D_RenderThread(class FRHICommandListImmediate& RHICmdList, FRHITexture2D* Texture, uint32 MipIndex, bool bLockWithinMiptail, bool bNeedsDefaultRHIFlush = true);
+	virtual void UpdateTexture2D_RenderThread(class FRHICommandListImmediate& RHICmdList, FRHITexture2D* Texture, uint32 MipIndex, const struct FUpdateTextureRegion2D& UpdateRegion, uint32 SourcePitch, const uint8* SourceData);
+	virtual void UpdateFromBufferTexture2D_RenderThread(class FRHICommandListImmediate& RHICmdList, FRHITexture2D* Texture, uint32 MipIndex, const struct FUpdateTextureRegion2D& UpdateRegion, uint32 SourcePitch, FRHIStructuredBuffer* Buffer, uint32 BufferOffset);
 
-	virtual FUpdateTexture3DData BeginUpdateTexture3D_RenderThread(class FRHICommandListImmediate& RHICmdList, FTexture3DRHIParamRef Texture, uint32 MipIndex, const struct FUpdateTextureRegion3D& UpdateRegion);
+	virtual FUpdateTexture3DData BeginUpdateTexture3D_RenderThread(class FRHICommandListImmediate& RHICmdList, FRHITexture3D* Texture, uint32 MipIndex, const struct FUpdateTextureRegion3D& UpdateRegion);
 	virtual void EndUpdateTexture3D_RenderThread(class FRHICommandListImmediate& RHICmdList, FUpdateTexture3DData& UpdateData);
 
 	virtual void EndMultiUpdateTexture3D_RenderThread(class FRHICommandListImmediate& RHICmdList, TArray<FUpdateTexture3DData>& UpdateDataArray);
 
-	virtual void UpdateTexture3D_RenderThread(class FRHICommandListImmediate& RHICmdList, FTexture3DRHIParamRef Texture, uint32 MipIndex, const struct FUpdateTextureRegion3D& UpdateRegion, uint32 SourceRowPitch, uint32 SourceDepthPitch, const uint8* SourceData);
+	virtual void UpdateTexture3D_RenderThread(class FRHICommandListImmediate& RHICmdList, FRHITexture3D* Texture, uint32 MipIndex, const struct FUpdateTextureRegion3D& UpdateRegion, uint32 SourceRowPitch, uint32 SourceDepthPitch, const uint8* SourceData);
 
 	virtual FRHIShaderLibraryRef RHICreateShaderLibrary_RenderThread(class FRHICommandListImmediate& RHICmdList, EShaderPlatform Platform, FString FilePath, FString Name);
 	virtual FTextureReferenceRHIRef RHICreateTextureReference_RenderThread(class FRHICommandListImmediate& RHICmdList, FLastRenderTimeContainer* LastRenderTime);
@@ -1166,33 +1166,33 @@ public:
 	}
 
 	virtual FTexture3DRHIRef RHICreateTexture3D_RenderThread(class FRHICommandListImmediate& RHICmdList, uint32 SizeX, uint32 SizeY, uint32 SizeZ, uint8 Format, uint32 NumMips, uint32 Flags, FRHIResourceCreateInfo& CreateInfo);
-	virtual FUnorderedAccessViewRHIRef RHICreateUnorderedAccessView_RenderThread(class FRHICommandListImmediate& RHICmdList, FStructuredBufferRHIParamRef StructuredBuffer, bool bUseUAVCounter, bool bAppendBuffer);
-	virtual FUnorderedAccessViewRHIRef RHICreateUnorderedAccessView_RenderThread(class FRHICommandListImmediate& RHICmdList, FTextureRHIParamRef Texture, uint32 MipLevel);
-	virtual FUnorderedAccessViewRHIRef RHICreateUnorderedAccessView_RenderThread(class FRHICommandListImmediate& RHICmdList, FVertexBufferRHIParamRef VertexBuffer, uint8 Format);
-	virtual FUnorderedAccessViewRHIRef RHICreateUnorderedAccessView_RenderThread(class FRHICommandListImmediate& RHICmdList, FIndexBufferRHIParamRef IndexBuffer, uint8 Format);
-	virtual FShaderResourceViewRHIRef RHICreateShaderResourceView_RenderThread(class FRHICommandListImmediate& RHICmdList, FTextureRHIParamRef Texture, const FRHITextureSRVCreateInfo& CreateInfo);
-	virtual FShaderResourceViewRHIRef RHICreateShaderResourceView_RenderThread(class FRHICommandListImmediate& RHICmdList, FVertexBufferRHIParamRef VertexBuffer, uint32 Stride, uint8 Format);
-	virtual FShaderResourceViewRHIRef RHICreateShaderResourceView_RenderThread(class FRHICommandListImmediate& RHICmdList, FIndexBufferRHIParamRef Buffer);
-	virtual FShaderResourceViewRHIRef RHICreateShaderResourceView_RenderThread(class FRHICommandListImmediate& RHICmdList, FStructuredBufferRHIParamRef StructuredBuffer);
+	virtual FUnorderedAccessViewRHIRef RHICreateUnorderedAccessView_RenderThread(class FRHICommandListImmediate& RHICmdList, FRHIStructuredBuffer* StructuredBuffer, bool bUseUAVCounter, bool bAppendBuffer);
+	virtual FUnorderedAccessViewRHIRef RHICreateUnorderedAccessView_RenderThread(class FRHICommandListImmediate& RHICmdList, FRHITexture* Texture, uint32 MipLevel);
+	virtual FUnorderedAccessViewRHIRef RHICreateUnorderedAccessView_RenderThread(class FRHICommandListImmediate& RHICmdList, FRHIVertexBuffer* VertexBuffer, uint8 Format);
+	virtual FUnorderedAccessViewRHIRef RHICreateUnorderedAccessView_RenderThread(class FRHICommandListImmediate& RHICmdList, FRHIIndexBuffer* IndexBuffer, uint8 Format);
+	virtual FShaderResourceViewRHIRef RHICreateShaderResourceView_RenderThread(class FRHICommandListImmediate& RHICmdList, FRHITexture* Texture, const FRHITextureSRVCreateInfo& CreateInfo);
+	virtual FShaderResourceViewRHIRef RHICreateShaderResourceView_RenderThread(class FRHICommandListImmediate& RHICmdList, FRHIVertexBuffer* VertexBuffer, uint32 Stride, uint8 Format);
+	virtual FShaderResourceViewRHIRef RHICreateShaderResourceView_RenderThread(class FRHICommandListImmediate& RHICmdList, FRHIIndexBuffer* Buffer);
+	virtual FShaderResourceViewRHIRef RHICreateShaderResourceView_RenderThread(class FRHICommandListImmediate& RHICmdList, FRHIStructuredBuffer* StructuredBuffer);
 	virtual FTextureCubeRHIRef RHICreateTextureCube_RenderThread(class FRHICommandListImmediate& RHICmdList, uint32 Size, uint8 Format, uint32 NumMips, uint32 Flags, FRHIResourceCreateInfo& CreateInfo);
 	virtual FTextureCubeRHIRef RHICreateTextureCubeArray_RenderThread(class FRHICommandListImmediate& RHICmdList, uint32 Size, uint32 ArraySize, uint8 Format, uint32 NumMips, uint32 Flags, FRHIResourceCreateInfo& CreateInfo);
 	virtual FRenderQueryRHIRef RHICreateRenderQuery_RenderThread(class FRHICommandListImmediate& RHICmdList, ERenderQueryType QueryType);
 
 	
-	virtual void* RHILockTextureCubeFace_RenderThread(class FRHICommandListImmediate& RHICmdList, FTextureCubeRHIParamRef Texture, uint32 FaceIndex, uint32 ArrayIndex, uint32 MipIndex, EResourceLockMode LockMode, uint32& DestStride, bool bLockWithinMiptail);
-	virtual void RHIUnlockTextureCubeFace_RenderThread(class FRHICommandListImmediate& RHICmdList, FTextureCubeRHIParamRef Texture, uint32 FaceIndex, uint32 ArrayIndex, uint32 MipIndex, bool bLockWithinMiptail);
+	virtual void* RHILockTextureCubeFace_RenderThread(class FRHICommandListImmediate& RHICmdList, FRHITextureCube* Texture, uint32 FaceIndex, uint32 ArrayIndex, uint32 MipIndex, EResourceLockMode LockMode, uint32& DestStride, bool bLockWithinMiptail);
+	virtual void RHIUnlockTextureCubeFace_RenderThread(class FRHICommandListImmediate& RHICmdList, FRHITextureCube* Texture, uint32 FaceIndex, uint32 ArrayIndex, uint32 MipIndex, bool bLockWithinMiptail);
 
-	virtual void RHIAcquireTransientResource_RenderThread(FTextureRHIParamRef Texture) { }
-	virtual void RHIDiscardTransientResource_RenderThread(FTextureRHIParamRef Texture) { }
-	virtual void RHIAcquireTransientResource_RenderThread(FVertexBufferRHIParamRef Buffer) { }
-	virtual void RHIDiscardTransientResource_RenderThread(FVertexBufferRHIParamRef Buffer) { }
-	virtual void RHIAcquireTransientResource_RenderThread(FStructuredBufferRHIParamRef Buffer) { }
-	virtual void RHIDiscardTransientResource_RenderThread(FStructuredBufferRHIParamRef Buffer) { }
+	virtual void RHIAcquireTransientResource_RenderThread(FRHITexture* Texture) { }
+	virtual void RHIDiscardTransientResource_RenderThread(FRHITexture* Texture) { }
+	virtual void RHIAcquireTransientResource_RenderThread(FRHIVertexBuffer* Buffer) { }
+	virtual void RHIDiscardTransientResource_RenderThread(FRHIVertexBuffer* Buffer) { }
+	virtual void RHIAcquireTransientResource_RenderThread(FRHIStructuredBuffer* Buffer) { }
+	virtual void RHIDiscardTransientResource_RenderThread(FRHIStructuredBuffer* Buffer) { }
 
 
-	virtual void RHIMapStagingSurface_RenderThread(class FRHICommandListImmediate& RHICmdList, FTextureRHIParamRef Texture, void*& OutData, int32& OutWidth, int32& OutHeight);
-	virtual void RHIUnmapStagingSurface_RenderThread(class FRHICommandListImmediate& RHICmdList, FTextureRHIParamRef Texture);
-	virtual void RHIReadSurfaceFloatData_RenderThread(class FRHICommandListImmediate& RHICmdList, FTextureRHIParamRef Texture, FIntRect Rect, TArray<FFloat16Color>& OutData, ECubeFace CubeFace, int32 ArrayIndex, int32 MipIndex);
+	virtual void RHIMapStagingSurface_RenderThread(class FRHICommandListImmediate& RHICmdList, FRHITexture* Texture, void*& OutData, int32& OutWidth, int32& OutHeight);
+	virtual void RHIUnmapStagingSurface_RenderThread(class FRHICommandListImmediate& RHICmdList, FRHITexture* Texture);
+	virtual void RHIReadSurfaceFloatData_RenderThread(class FRHICommandListImmediate& RHICmdList, FRHITexture* Texture, FIntRect Rect, TArray<FFloat16Color>& OutData, ECubeFace CubeFace, int32 ArrayIndex, int32 MipIndex);
 
 	//Utilities
 	virtual void EnableIdealGPUCaptureOptions(bool bEnable);
@@ -1200,12 +1200,12 @@ public:
 	//checks if the GPU is still alive.
 	virtual bool CheckGpuHeartbeat() const { return true; }
 
-	virtual void VirtualTextureSetFirstMipInMemory_RenderThread(class FRHICommandListImmediate& RHICmdList, FTexture2DRHIParamRef Texture, uint32 FirstMip);
-	virtual void VirtualTextureSetFirstMipVisible_RenderThread(class FRHICommandListImmediate& RHICmdList, FTexture2DRHIParamRef Texture, uint32 FirstMip);
+	virtual void VirtualTextureSetFirstMipInMemory_RenderThread(class FRHICommandListImmediate& RHICmdList, FRHITexture2D* Texture, uint32 FirstMip);
+	virtual void VirtualTextureSetFirstMipVisible_RenderThread(class FRHICommandListImmediate& RHICmdList, FRHITexture2D* Texture, uint32 FirstMip);
 
 	/* Copy the source box pixels in the destination box texture, return true if implemented for the current platform*/
-	virtual void RHICopySubTextureRegion_RenderThread(class FRHICommandListImmediate& RHICmdList, FTexture2DRHIParamRef SourceTexture, FTexture2DRHIParamRef DestinationTexture, FBox2D SourceBox, FBox2D DestinationBox);
-	virtual void RHICopySubTextureRegion(FTexture2DRHIParamRef SourceTexture, FTexture2DRHIParamRef DestinationTexture, FBox2D SourceBox, FBox2D DestinationBox) { }
+	virtual void RHICopySubTextureRegion_RenderThread(class FRHICommandListImmediate& RHICmdList, FRHITexture2D* SourceTexture, FRHITexture2D* DestinationTexture, FBox2D SourceBox, FBox2D DestinationBox);
+	virtual void RHICopySubTextureRegion(FRHITexture2D* SourceTexture, FRHITexture2D* DestinationTexture, FBox2D SourceBox, FBox2D DestinationBox) { }
 
 	virtual FRHIFlipDetails RHIWaitForFlip(double TimeoutInSeconds) { return FRHIFlipDetails(); }
 	virtual void RHISignalFlipEvent() { }
@@ -1307,7 +1307,7 @@ FORCEINLINE FUniformBufferRHIRef RHICreateUniformBuffer(const void* Contents, co
 	return GDynamicRHI->RHICreateUniformBuffer(Contents, Layout, Usage, Validation);
 }
 
-FORCEINLINE void RHIUpdateUniformBuffer(FUniformBufferRHIParamRef UniformBufferRHI, const void* Contents)
+FORCEINLINE void RHIUpdateUniformBuffer(FRHIUniformBuffer* UniformBufferRHI, const void* Contents)
 {
 	return GDynamicRHI->RHIUpdateUniformBuffer(UniformBufferRHI, Contents);
 }
@@ -1332,17 +1332,17 @@ FORCEINLINE void RHIGetTextureMemoryStats(FTextureMemoryStats& OutStats)
 	GDynamicRHI->RHIGetTextureMemoryStats(OutStats);
 }
 
-FORCEINLINE void RHIGetResourceInfo(FTextureRHIParamRef Ref, FRHIResourceInfo& OutInfo)
+FORCEINLINE void RHIGetResourceInfo(FRHITexture* Ref, FRHIResourceInfo& OutInfo)
 {
 	return GDynamicRHI->RHIGetResourceInfo(Ref, OutInfo);
 }
 
-FORCEINLINE uint32 RHIComputeMemorySize(FTextureRHIParamRef TextureRHI)
+FORCEINLINE uint32 RHIComputeMemorySize(FRHITexture* TextureRHI)
 {
 	return GDynamicRHI->RHIComputeMemorySize(TextureRHI);
 }
 
-FORCEINLINE void RHIBindDebugLabelName(FTextureRHIParamRef Texture, const TCHAR* Name)
+FORCEINLINE void RHIBindDebugLabelName(FRHITexture* Texture, const TCHAR* Name)
 {
 	GDynamicRHI->RHIBindDebugLabelName(Texture, Name);
 }
@@ -1367,7 +1367,7 @@ FORCEINLINE FTexture2DRHIRef RHIGetViewportBackBuffer(FRHIViewport* Viewport)
 	return GDynamicRHI->RHIGetViewportBackBuffer(Viewport);
 }
 
-FORCEINLINE FTexture2DRHIRef RHIGetFMaskTexture(FTextureRHIParamRef SourceTextureRHI)
+FORCEINLINE FTexture2DRHIRef RHIGetFMaskTexture(FRHITexture* SourceTextureRHI)
 {
 	return GDynamicRHI->RHIGetFMaskTexture(SourceTextureRHI);
 }
@@ -1471,7 +1471,7 @@ inline void appDefragmentTexturePool() {}
 /**
 * Checks if the texture data is allocated within the texture pool or not.
 */
-inline bool appIsPoolTexture(FTextureRHIParamRef TextureRHI) { return false; }
+inline bool appIsPoolTexture(FRHITexture* TextureRHI) { return false; }
 
 /**
 * Log the current texture memory stats.
