UE4.23 - Deprecate FUniformBufferRHIParamRef, FIndexBufferRHIParamRef, FVertexBufferRHIParamRef, FStructuredBufferRHIParamRef, FTextureRHIParamRef, FTexture2DRHIParamRef, FTexture2DArrayRHIParamRef, F

diff --git a/Engine/Source/Runtime/RHI/Public/RHICommandList.h b/Engine/Source/Runtime/RHI/Public/RHICommandList.h
--- a/Engine/Source/Runtime/RHI/Public/RHICommandList.h
+++ b/Engine/Source/Runtime/RHI/Public/RHICommandList.h
@@ -112,9 +112,9 @@ extern RHI_API TAutoConsoleVariable<int32> CVarRHICmdFlushRenderThreadTasks;
 #if RHI_RAYTRACING
 struct FRayTracingShaderBindings
 {
-	FTextureRHIParamRef Textures[32] = {};
+	FRHITexture* Textures[32] = {};
 	FRHIShaderResourceView* SRVs[32] = {};
-	FUniformBufferRHIParamRef UniformBuffers[8] = {};
+	FRHIUniformBuffer* UniformBuffers[8] = {};
 	FRHISamplerState* Samplers[16] = {};
 	FRHIUnorderedAccessView* UAVs[8] = {};
 };
@@ -681,8 +681,8 @@ struct FRHICommand : public FRHICommandBase
 
 struct  FRHICommandBeginUpdateMultiFrameResource final : public FRHICommand<FRHICommandBeginUpdateMultiFrameResource>
 {
-	FTextureRHIParamRef Texture;
-	FORCEINLINE_DEBUGGABLE FRHICommandBeginUpdateMultiFrameResource(FTextureRHIParamRef InTexture)
+	FRHITexture* Texture;
+	FORCEINLINE_DEBUGGABLE FRHICommandBeginUpdateMultiFrameResource(FRHITexture* InTexture)
 		: Texture(InTexture)
 	{
 	}
@@ -691,8 +691,8 @@ struct  FRHICommandBeginUpdateMultiFrameResource final : public FRHICommand<FRHI
 
 struct  FRHICommandEndUpdateMultiFrameResource final : public FRHICommand<FRHICommandEndUpdateMultiFrameResource>
 {
-	FTextureRHIParamRef Texture;
-	FORCEINLINE_DEBUGGABLE FRHICommandEndUpdateMultiFrameResource(FTextureRHIParamRef InTexture)
+	FRHITexture* Texture;
+	FORCEINLINE_DEBUGGABLE FRHICommandEndUpdateMultiFrameResource(FRHITexture* InTexture)
 		: Texture(InTexture)
 	{
 	}
@@ -753,8 +753,8 @@ struct FRHICommandSetShaderUniformBuffer final : public FRHICommand<FRHICommandS
 {
 	TRHIShader* Shader;
 	uint32 BaseIndex;
-	FUniformBufferRHIParamRef UniformBuffer;
-	FORCEINLINE_DEBUGGABLE FRHICommandSetShaderUniformBuffer(TRHIShader* InShader, uint32 InBaseIndex, FUniformBufferRHIParamRef InUniformBuffer)
+	FRHIUniformBuffer* UniformBuffer;
+	FORCEINLINE_DEBUGGABLE FRHICommandSetShaderUniformBuffer(TRHIShader* InShader, uint32 InBaseIndex, FRHIUniformBuffer* InUniformBuffer)
 		: Shader(InShader)
 		, BaseIndex(InBaseIndex)
 		, UniformBuffer(InUniformBuffer)
@@ -768,8 +768,8 @@ struct FRHICommandSetShaderTexture final : public FRHICommand<FRHICommandSetShad
 {
 	TRHIShader* Shader;
 	uint32 TextureIndex;
-	FTextureRHIParamRef Texture;
-	FORCEINLINE_DEBUGGABLE FRHICommandSetShaderTexture(TRHIShader* InShader, uint32 InTextureIndex, FTextureRHIParamRef InTexture)
+	FRHITexture* Texture;
+	FORCEINLINE_DEBUGGABLE FRHICommandSetShaderTexture(TRHIShader* InShader, uint32 InTextureIndex, FRHITexture* InTexture)
 		: Shader(InShader)
 		, TextureIndex(InTextureIndex)
 		, Texture(InTexture)
@@ -856,14 +856,14 @@ struct FRHICommandDrawPrimitive final : public FRHICommand<FRHICommandDrawPrimit
 
 struct FRHICommandDrawIndexedPrimitive final : public FRHICommand<FRHICommandDrawIndexedPrimitive>
 {
-	FIndexBufferRHIParamRef IndexBuffer;
+	FRHIIndexBuffer* IndexBuffer;
 	int32 BaseVertexIndex;
 	uint32 FirstInstance;
 	uint32 NumVertices;
 	uint32 StartIndex;
 	uint32 NumPrimitives;
 	uint32 NumInstances;
-	FORCEINLINE_DEBUGGABLE FRHICommandDrawIndexedPrimitive(FIndexBufferRHIParamRef InIndexBuffer, int32 InBaseVertexIndex, uint32 InFirstInstance, uint32 InNumVertices, uint32 InStartIndex, uint32 InNumPrimitives, uint32 InNumInstances)
+	FORCEINLINE_DEBUGGABLE FRHICommandDrawIndexedPrimitive(FRHIIndexBuffer* InIndexBuffer, int32 InBaseVertexIndex, uint32 InFirstInstance, uint32 InNumVertices, uint32 InStartIndex, uint32 InNumPrimitives, uint32 InNumInstances)
 		: IndexBuffer(InIndexBuffer)
 		, BaseVertexIndex(InBaseVertexIndex)
 		, FirstInstance(InFirstInstance)
@@ -889,9 +889,9 @@ struct FRHICommandSetBlendFactor final : public FRHICommand<FRHICommandSetBlendF
 struct FRHICommandSetStreamSource final : public FRHICommand<FRHICommandSetStreamSource>
 {
 	uint32 StreamIndex;
-	FVertexBufferRHIParamRef VertexBuffer;
+	FRHIVertexBuffer* VertexBuffer;
 	uint32 Offset;
-	FORCEINLINE_DEBUGGABLE FRHICommandSetStreamSource(uint32 InStreamIndex, FVertexBufferRHIParamRef InVertexBuffer, uint32 InOffset)
+	FORCEINLINE_DEBUGGABLE FRHICommandSetStreamSource(uint32 InStreamIndex, FRHIVertexBuffer* InVertexBuffer, uint32 InOffset)
 		: StreamIndex(InStreamIndex)
 		, VertexBuffer(InVertexBuffer)
 		, Offset(InOffset)
@@ -1203,9 +1203,9 @@ struct FRHICommandDispatchComputeShader final : public FRHICommand<FRHICommandDi
 template<ECmdList CmdListType>
 struct FRHICommandDispatchIndirectComputeShader final : public FRHICommand<FRHICommandDispatchIndirectComputeShader<CmdListType>>
 {
-	FVertexBufferRHIParamRef ArgumentBuffer;
+	FRHIVertexBuffer* ArgumentBuffer;
 	uint32 ArgumentOffset;
-	FORCEINLINE_DEBUGGABLE FRHICommandDispatchIndirectComputeShader(FVertexBufferRHIParamRef InArgumentBuffer, uint32 InArgumentOffset)
+	FORCEINLINE_DEBUGGABLE FRHICommandDispatchIndirectComputeShader(FRHIVertexBuffer* InArgumentBuffer, uint32 InArgumentOffset)
 		: ArgumentBuffer(InArgumentBuffer)
 		, ArgumentOffset(InArgumentOffset)
 	{
@@ -1230,9 +1230,9 @@ struct FRHICommandFlushComputeShaderCache final : public FRHICommand<FRHICommand
 
 struct FRHICommandDrawPrimitiveIndirect final : public FRHICommand<FRHICommandDrawPrimitiveIndirect>
 {
-	FVertexBufferRHIParamRef ArgumentBuffer;
+	FRHIVertexBuffer* ArgumentBuffer;
 	uint32 ArgumentOffset;
-	FORCEINLINE_DEBUGGABLE FRHICommandDrawPrimitiveIndirect(FVertexBufferRHIParamRef InArgumentBuffer, uint32 InArgumentOffset)
+	FORCEINLINE_DEBUGGABLE FRHICommandDrawPrimitiveIndirect(FRHIVertexBuffer* InArgumentBuffer, uint32 InArgumentOffset)
 		: ArgumentBuffer(InArgumentBuffer)
 		, ArgumentOffset(InArgumentOffset)
 	{
@@ -1242,12 +1242,12 @@ struct FRHICommandDrawPrimitiveIndirect final : public FRHICommand<FRHICommandDr
 
 struct FRHICommandDrawIndexedIndirect final : public FRHICommand<FRHICommandDrawIndexedIndirect>
 {
-	FIndexBufferRHIParamRef IndexBufferRHI;
-	FStructuredBufferRHIParamRef ArgumentsBufferRHI;
+	FRHIIndexBuffer* IndexBufferRHI;
+	FRHIStructuredBuffer* ArgumentsBufferRHI;
 	uint32 DrawArgumentsIndex;
 	uint32 NumInstances;
 
-	FORCEINLINE_DEBUGGABLE FRHICommandDrawIndexedIndirect(FIndexBufferRHIParamRef InIndexBufferRHI, FStructuredBufferRHIParamRef InArgumentsBufferRHI, uint32 InDrawArgumentsIndex, uint32 InNumInstances)
+	FORCEINLINE_DEBUGGABLE FRHICommandDrawIndexedIndirect(FRHIIndexBuffer* InIndexBufferRHI, FRHIStructuredBuffer* InArgumentsBufferRHI, uint32 InDrawArgumentsIndex, uint32 InNumInstances)
 		: IndexBufferRHI(InIndexBufferRHI)
 		, ArgumentsBufferRHI(InArgumentsBufferRHI)
 		, DrawArgumentsIndex(InDrawArgumentsIndex)
@@ -1259,11 +1259,11 @@ struct FRHICommandDrawIndexedIndirect final : public FRHICommand<FRHICommandDraw
 
 struct FRHICommandDrawIndexedPrimitiveIndirect final : public FRHICommand<FRHICommandDrawIndexedPrimitiveIndirect>
 {
-	FIndexBufferRHIParamRef IndexBuffer;
-	FVertexBufferRHIParamRef ArgumentsBuffer;
+	FRHIIndexBuffer* IndexBuffer;
+	FRHIVertexBuffer* ArgumentsBuffer;
 	uint32 ArgumentOffset;
 
-	FORCEINLINE_DEBUGGABLE FRHICommandDrawIndexedPrimitiveIndirect(FIndexBufferRHIParamRef InIndexBuffer, FVertexBufferRHIParamRef InArgumentsBuffer, uint32 InArgumentOffset)
+	FORCEINLINE_DEBUGGABLE FRHICommandDrawIndexedPrimitiveIndirect(FRHIIndexBuffer* InIndexBuffer, FRHIVertexBuffer* InArgumentsBuffer, uint32 InArgumentOffset)
 		: IndexBuffer(InIndexBuffer)
 		, ArgumentsBuffer(InArgumentsBuffer)
 		, ArgumentOffset(InArgumentOffset)
@@ -1304,10 +1304,10 @@ struct FRHICommandClearTinyUAV final : public FRHICommand<FRHICommandClearTinyUA
 struct FRHICommandCopyToResolveTarget final : public FRHICommand<FRHICommandCopyToResolveTarget>
 {
 	FResolveParams ResolveParams;
-	FTextureRHIParamRef SourceTexture;
-	FTextureRHIParamRef DestTexture;
+	FRHITexture* SourceTexture;
+	FRHITexture* DestTexture;
 
-	FORCEINLINE_DEBUGGABLE FRHICommandCopyToResolveTarget(FTextureRHIParamRef InSourceTexture, FTextureRHIParamRef InDestTexture, const FResolveParams& InResolveParams)
+	FORCEINLINE_DEBUGGABLE FRHICommandCopyToResolveTarget(FRHITexture* InSourceTexture, FRHITexture* InDestTexture, const FResolveParams& InResolveParams)
 		: ResolveParams(InResolveParams)
 		, SourceTexture(InSourceTexture)
 		, DestTexture(InDestTexture)
@@ -1323,10 +1323,10 @@ struct FRHICommandCopyToResolveTarget final : public FRHICommand<FRHICommandCopy
 struct FRHICommandCopyTexture final : public FRHICommand<FRHICommandCopyTexture>
 {
 	FRHICopyTextureInfo CopyInfo;
-	FTextureRHIParamRef SourceTexture;
-	FTextureRHIParamRef DestTexture;
+	FRHITexture* SourceTexture;
+	FRHITexture* DestTexture;
 
-	FORCEINLINE_DEBUGGABLE FRHICommandCopyTexture(FTextureRHIParamRef InSourceTexture, FTextureRHIParamRef InDestTexture, const FRHICopyTextureInfo& InCopyInfo)
+	FORCEINLINE_DEBUGGABLE FRHICommandCopyTexture(FRHITexture* InSourceTexture, FRHITexture* InDestTexture, const FRHICopyTextureInfo& InCopyInfo)
 		: CopyInfo(InCopyInfo)
 		, SourceTexture(InSourceTexture)
 		, DestTexture(InDestTexture)
@@ -1342,9 +1342,9 @@ struct FRHICommandCopyTexture final : public FRHICommand<FRHICommandCopyTexture>
 struct FRHICommandTransitionTextures final : public FRHICommand<FRHICommandTransitionTextures>
 {
 	int32 NumTextures;
-	FTextureRHIParamRef* Textures; // Pointer to an array of textures, allocated inline with the command list
+	FRHITexture** Textures; // Pointer to an array of textures, allocated inline with the command list
 	EResourceTransitionAccess TransitionType;
-	FORCEINLINE_DEBUGGABLE FRHICommandTransitionTextures(EResourceTransitionAccess InTransitionType, FTextureRHIParamRef* InTextures, int32 InNumTextures)
+	FORCEINLINE_DEBUGGABLE FRHICommandTransitionTextures(EResourceTransitionAccess InTransitionType, FRHITexture** InTextures, int32 InNumTextures)
 		: NumTextures(InNumTextures)
 		, Textures(InTextures)
 		, TransitionType(InTransitionType)
@@ -1355,9 +1355,9 @@ struct FRHICommandTransitionTextures final : public FRHICommand<FRHICommandTrans
 
 struct FRHICommandTransitionTexturesArray final : public FRHICommand<FRHICommandTransitionTexturesArray>
 {	
-	TArray<FTextureRHIParamRef>& Textures;
+	TArray<FRHITexture*>& Textures;
 	EResourceTransitionAccess TransitionType;
-	FORCEINLINE_DEBUGGABLE FRHICommandTransitionTexturesArray(EResourceTransitionAccess InTransitionType, TArray<FTextureRHIParamRef>& InTextures)
+	FORCEINLINE_DEBUGGABLE FRHICommandTransitionTexturesArray(EResourceTransitionAccess InTransitionType, TArray<FRHITexture*>& InTextures)
 		: Textures(InTextures)
 		, TransitionType(InTransitionType)
 	{		
@@ -1412,12 +1412,12 @@ struct FRHICommandWaitComputeFence final : public FRHICommand<FRHICommandWaitCom
 template<ECmdList CmdListType>
 struct FRHICommandCopyToStagingBuffer final : public FRHICommand<FRHICommandCopyToStagingBuffer<CmdListType>>
 {
-	FVertexBufferRHIParamRef SourceBuffer;
+	FRHIVertexBuffer* SourceBuffer;
 	FRHIStagingBuffer* DestinationStagingBuffer;
 	uint32 Offset;
 	uint32 NumBytes;
 
-	FORCEINLINE_DEBUGGABLE FRHICommandCopyToStagingBuffer(FVertexBufferRHIParamRef InSourceBuffer, FRHIStagingBuffer* InDestinationStagingBuffer, uint32 InOffset, uint32 InNumBytes)
+	FORCEINLINE_DEBUGGABLE FRHICommandCopyToStagingBuffer(FRHIVertexBuffer* InSourceBuffer, FRHIStagingBuffer* InDestinationStagingBuffer, uint32 InOffset, uint32 InNumBytes)
 		: SourceBuffer(InSourceBuffer)
 		, DestinationStagingBuffer(InDestinationStagingBuffer)
 		, Offset(InOffset)
@@ -1441,11 +1441,11 @@ struct FRHICommandWriteGPUFence final : public FRHICommand<FRHICommandWriteGPUFe
 
 struct FRHICommandClearColorTexture final : public FRHICommand<FRHICommandClearColorTexture>
 {
-	FTextureRHIParamRef Texture;
+	FRHITexture* Texture;
 	FLinearColor Color;
 
 	FORCEINLINE_DEBUGGABLE FRHICommandClearColorTexture(
-		FTextureRHIParamRef InTexture,
+		FRHITexture* InTexture,
 		const FLinearColor& InColor
 		)
 		: Texture(InTexture)
@@ -1457,13 +1457,13 @@ struct FRHICommandClearColorTexture final : public FRHICommand<FRHICommandClearC
 
 struct FRHICommandClearDepthStencilTexture final : public FRHICommand<FRHICommandClearDepthStencilTexture>
 {
-	FTextureRHIParamRef Texture;
+	FRHITexture* Texture;
 	float Depth;
 	uint32 Stencil;
 	EClearDepthStencil ClearDepthStencil;
 
 	FORCEINLINE_DEBUGGABLE FRHICommandClearDepthStencilTexture(
-		FTextureRHIParamRef InTexture,
+		FRHITexture* InTexture,
 		EClearDepthStencil InClearDepthStencil,
 		float InDepth,
 		uint32 InStencil
@@ -1480,12 +1480,12 @@ struct FRHICommandClearDepthStencilTexture final : public FRHICommand<FRHIComman
 struct FRHICommandClearColorTextures final : public FRHICommand<FRHICommandClearColorTextures>
 {
 	FLinearColor ColorArray[MaxSimultaneousRenderTargets];
-	FTextureRHIParamRef Textures[MaxSimultaneousRenderTargets];
+	FRHITexture* Textures[MaxSimultaneousRenderTargets];
 	int32 NumClearColors;
 
 	FORCEINLINE_DEBUGGABLE FRHICommandClearColorTextures(
 		int32 InNumClearColors,
-		FTextureRHIParamRef* InTextures,
+		FRHITexture** InTextures,
 		const FLinearColor* InColorArray
 		)
 		: NumClearColors(InNumClearColors)
@@ -1668,9 +1668,9 @@ struct FRHICommandEndFrame final : public FRHICommand<FRHICommandEndFrame>
 struct FRHICommandBeginDrawingViewport final : public FRHICommand<FRHICommandBeginDrawingViewport>
 {
 	FRHIViewport* Viewport;
-	FTextureRHIParamRef RenderTargetRHI;
+	FRHITexture* RenderTargetRHI;
 
-	FORCEINLINE_DEBUGGABLE FRHICommandBeginDrawingViewport(FRHIViewport* InViewport, FTextureRHIParamRef InRenderTargetRHI)
+	FORCEINLINE_DEBUGGABLE FRHICommandBeginDrawingViewport(FRHIViewport* InViewport, FRHITexture* InRenderTargetRHI)
 		: Viewport(InViewport)
 		, RenderTargetRHI(InRenderTargetRHI)
 	{
@@ -1758,8 +1758,8 @@ struct FRHICommandDebugBreak final : public FRHICommand<FRHICommandDebugBreak>
 struct FRHICommandUpdateTextureReference final : public FRHICommand<FRHICommandUpdateTextureReference>
 {
 	FRHITextureReference* TextureRef;
-	FTextureRHIParamRef NewTexture;
-	FORCEINLINE_DEBUGGABLE FRHICommandUpdateTextureReference(FRHITextureReference* InTextureRef, FTextureRHIParamRef InNewTexture)
+	FRHITexture* NewTexture;
+	FORCEINLINE_DEBUGGABLE FRHICommandUpdateTextureReference(FRHITextureReference* InTextureRef, FRHITexture* InNewTexture)
 		: TextureRef(InTextureRef)
 		, NewTexture(InNewTexture)
 	{
@@ -1770,21 +1770,21 @@ struct FRHICommandUpdateTextureReference final : public FRHICommand<FRHICommandU
 struct FRHIShaderResourceViewUpdateInfo_VB
 {
 	FRHIShaderResourceView* SRV;
-	FVertexBufferRHIParamRef VertexBuffer;
+	FRHIVertexBuffer* VertexBuffer;
 	uint32 Stride;
 	uint8 Format;
 };
 
 struct FRHIVertexBufferUpdateInfo
 {
-	FVertexBufferRHIParamRef DestBuffer;
-	FVertexBufferRHIParamRef SrcBuffer;
+	FRHIVertexBuffer* DestBuffer;
+	FRHIVertexBuffer* SrcBuffer;
 };
 
 struct FRHIIndexBufferUpdateInfo
 {
-	FIndexBufferRHIParamRef DestBuffer;
-	FIndexBufferRHIParamRef SrcBuffer;
+	FRHIIndexBuffer* DestBuffer;
+	FRHIIndexBuffer* SrcBuffer;
 };
 
 struct FRHIResourceUpdateInfo
@@ -1834,13 +1834,13 @@ struct FRHICommandUpdateRHIResources final : public FRHICommand<FRHICommandUpdat
 #if RHI_RAYTRACING
 struct FRHICommandCopyBufferRegion final : public FRHICommand<FRHICommandCopyBufferRegion>
 {
-	FVertexBufferRHIParamRef DestBuffer;
+	FRHIVertexBuffer* DestBuffer;
 	uint64 DstOffset;
-	FVertexBufferRHIParamRef SourceBuffer;
+	FRHIVertexBuffer* SourceBuffer;
 	uint64 SrcOffset;
 	uint64 NumBytes;
 
-	explicit FRHICommandCopyBufferRegion(FVertexBufferRHIParamRef InDestBuffer, uint64 InDstOffset, FVertexBufferRHIParamRef InSourceBuffer, uint64 InSrcOffset, uint64 InNumBytes)
+	explicit FRHICommandCopyBufferRegion(FRHIVertexBuffer* InDestBuffer, uint64 InDstOffset, FRHIVertexBuffer* InSourceBuffer, uint64 InSrcOffset, uint64 InNumBytes)
 		: DestBuffer(InDestBuffer)
 		, DstOffset(InDstOffset)
 		, SourceBuffer(InSourceBuffer)
@@ -1978,7 +1978,7 @@ struct FRHICommandSetRayTracingBindings final : public FRHICommand<FRHICommandSe
 	FRayTracingPipelineState* Pipeline;
 	uint32 ShaderIndex;
 	uint32 NumUniformBuffers;
-	const FUniformBufferRHIParamRef* UniformBuffers; // Pointer to an array of uniform buffers, allocated inline within the command list
+	FRHIUniformBuffer* const* UniformBuffers; // Pointer to an array of uniform buffers, allocated inline within the command list
 	uint32 LooseParameterDataSize;
 	const void* LooseParameterData;
 	uint32 UserData;
@@ -1986,7 +1986,7 @@ struct FRHICommandSetRayTracingBindings final : public FRHICommand<FRHICommandSe
 
 	// Hit group bindings
 	FRHICommandSetRayTracingBindings(FRHIRayTracingScene* InScene, uint32 InInstanceIndex, uint32 InSegmentIndex, uint32 InShaderSlot,
-		FRayTracingPipelineState* InPipeline, uint32 InHitGroupIndex, uint32 InNumUniformBuffers, const FUniformBufferRHIParamRef* InUniformBuffers,
+		FRayTracingPipelineState* InPipeline, uint32 InHitGroupIndex, uint32 InNumUniformBuffers, FRHIUniformBuffer* const* InUniformBuffers,
 		uint32 InLooseParameterDataSize, const void* InLooseParameterData,
 		uint32 InUserData)
 		: Scene(InScene)
@@ -2007,7 +2007,7 @@ struct FRHICommandSetRayTracingBindings final : public FRHICommand<FRHICommandSe
 	// Callable shader bindings
 	FRHICommandSetRayTracingBindings(FRHIRayTracingScene* InScene, uint32 InShaderSlot,
 		FRayTracingPipelineState* InPipeline, uint32 InShaderIndex,
-		uint32 InNumUniformBuffers, const FUniformBufferRHIParamRef* InUniformBuffers,
+		uint32 InNumUniformBuffers, FRHIUniformBuffer* const* InUniformBuffers,
 		uint32 InUserData)
 		: Scene(InScene)
 		, InstanceIndex(0)
@@ -2052,7 +2052,7 @@ public:
 	void* operator new(size_t Size);
 	void operator delete(void *RawMemory);
 	
-	FORCEINLINE_DEBUGGABLE void BeginUpdateMultiFrameResource( FTextureRHIParamRef Texture)
+	FORCEINLINE_DEBUGGABLE void BeginUpdateMultiFrameResource(FRHITexture* Texture)
 	{
 		//check(IsOutsideRenderPass());
 		if (Bypass())
@@ -2063,7 +2063,7 @@ public:
 		ALLOC_COMMAND(FRHICommandBeginUpdateMultiFrameResource)(Texture);
 	}
 
-	FORCEINLINE_DEBUGGABLE void EndUpdateMultiFrameResource(FTextureRHIParamRef Texture)
+	FORCEINLINE_DEBUGGABLE void EndUpdateMultiFrameResource(FRHITexture* Texture)
 	{
 		//check(IsOutsideRenderPass());
 		if (Bypass())
@@ -2132,7 +2132,7 @@ public:
 	}
 
 	template <typename TShaderRHI>
-	FORCEINLINE_DEBUGGABLE void SetShaderUniformBuffer(TShaderRHI* Shader, uint32 BaseIndex, FUniformBufferRHIParamRef UniformBuffer)
+	FORCEINLINE_DEBUGGABLE void SetShaderUniformBuffer(TShaderRHI* Shader, uint32 BaseIndex, FRHIUniformBuffer* UniformBuffer)
 	{
 		//check(IsOutsideRenderPass());
 		if (Bypass())
@@ -2143,7 +2143,7 @@ public:
 		ALLOC_COMMAND(FRHICommandSetShaderUniformBuffer<TShaderRHI, ECmdList::EGfx>)(Shader, BaseIndex, UniformBuffer);
 	}
 	template <typename TShaderRHI>
-	FORCEINLINE void SetShaderUniformBuffer(TRefCountPtr<TShaderRHI>& Shader, uint32 BaseIndex, FUniformBufferRHIParamRef UniformBuffer)
+	FORCEINLINE void SetShaderUniformBuffer(TRefCountPtr<TShaderRHI>& Shader, uint32 BaseIndex, FRHIUniformBuffer* UniformBuffer)
 	{
 		SetShaderUniformBuffer(Shader.GetReference(), BaseIndex, UniformBuffer);
 	}
@@ -2168,7 +2168,7 @@ public:
 	}
 
 	template <typename TRHIShader>
-	FORCEINLINE_DEBUGGABLE void SetShaderTexture(TRHIShader* Shader, uint32 TextureIndex, FTextureRHIParamRef Texture)
+	FORCEINLINE_DEBUGGABLE void SetShaderTexture(TRHIShader* Shader, uint32 TextureIndex, FRHITexture* Texture)
 	{
 		//check(IsOutsideRenderPass());
 		if (Bypass())
@@ -2180,13 +2180,13 @@ public:
 	}
 
 	template <typename TShaderRHI>
-	FORCEINLINE_DEBUGGABLE void SetShaderTexture(TRefCountPtr<TShaderRHI>& Shader, uint32 TextureIndex, FTextureRHIParamRef Texture)
+	FORCEINLINE_DEBUGGABLE void SetShaderTexture(TRefCountPtr<TShaderRHI>& Shader, uint32 TextureIndex, FRHITexture* Texture)
 	{
 		SetShaderTexture(Shader.GetReference(), TextureIndex, Texture);
 	}
 
 	template <typename TShaderRHI>
-	FORCEINLINE_DEBUGGABLE void SetShaderTexture(const TShaderRHI* Shader, uint32 TextureIndex, FTextureRHIParamRef Texture)
+	FORCEINLINE_DEBUGGABLE void SetShaderTexture(const TShaderRHI* Shader, uint32 TextureIndex, FRHITexture* Texture)
 	{
 		SetShaderTexture((TShaderRHI*)Shader, TextureIndex, Texture);
 	}
@@ -2287,7 +2287,7 @@ public:
 		ALLOC_COMMAND(FRHICommandDrawPrimitive)(BaseVertexIndex, NumPrimitives, NumInstances);
 	}
 
-	FORCEINLINE_DEBUGGABLE void DrawIndexedPrimitive(FIndexBufferRHIParamRef IndexBuffer, int32 BaseVertexIndex, uint32 FirstInstance, uint32 NumVertices, uint32 StartIndex, uint32 NumPrimitives, uint32 NumInstances)
+	FORCEINLINE_DEBUGGABLE void DrawIndexedPrimitive(FRHIIndexBuffer* IndexBuffer, int32 BaseVertexIndex, uint32 FirstInstance, uint32 NumVertices, uint32 StartIndex, uint32 NumPrimitives, uint32 NumInstances)
 	{
 		if (!IndexBuffer)
 		{
@@ -2303,7 +2303,7 @@ public:
 		ALLOC_COMMAND(FRHICommandDrawIndexedPrimitive)(IndexBuffer, BaseVertexIndex, FirstInstance, NumVertices, StartIndex, NumPrimitives, NumInstances);
 	}
 
-	FORCEINLINE_DEBUGGABLE void SetStreamSource(uint32 StreamIndex, FVertexBufferRHIParamRef VertexBuffer, uint32 Offset)
+	FORCEINLINE_DEBUGGABLE void SetStreamSource(uint32 StreamIndex, FRHIVertexBuffer* VertexBuffer, uint32 Offset)
 	{
 		if (Bypass())
 		{
@@ -2520,7 +2520,7 @@ public:
 		ALLOC_COMMAND(FRHICommandDispatchComputeShader<ECmdList::EGfx>)(ThreadGroupCountX, ThreadGroupCountY, ThreadGroupCountZ);
 	}
 
-	FORCEINLINE_DEBUGGABLE void DispatchIndirectComputeShader(FVertexBufferRHIParamRef ArgumentBuffer, uint32 ArgumentOffset)
+	FORCEINLINE_DEBUGGABLE void DispatchIndirectComputeShader(FRHIVertexBuffer* ArgumentBuffer, uint32 ArgumentOffset)
 	{
 		if (Bypass())
 		{
@@ -2550,7 +2550,7 @@ public:
 		ALLOC_COMMAND(FRHICommandFlushComputeShaderCache)();
 	}
 
-	FORCEINLINE_DEBUGGABLE void DrawPrimitiveIndirect(FVertexBufferRHIParamRef ArgumentBuffer, uint32 ArgumentOffset)
+	FORCEINLINE_DEBUGGABLE void DrawPrimitiveIndirect(FRHIVertexBuffer* ArgumentBuffer, uint32 ArgumentOffset)
 	{
 		//check(IsOutsideRenderPass());
 		if (Bypass())
@@ -2561,7 +2561,7 @@ public:
 		ALLOC_COMMAND(FRHICommandDrawPrimitiveIndirect)(ArgumentBuffer, ArgumentOffset);
 	}
 
-	FORCEINLINE_DEBUGGABLE void DrawIndexedIndirect(FIndexBufferRHIParamRef IndexBufferRHI, FStructuredBufferRHIParamRef ArgumentsBufferRHI, uint32 DrawArgumentsIndex, uint32 NumInstances)
+	FORCEINLINE_DEBUGGABLE void DrawIndexedIndirect(FRHIIndexBuffer* IndexBufferRHI, FRHIStructuredBuffer* ArgumentsBufferRHI, uint32 DrawArgumentsIndex, uint32 NumInstances)
 	{
 		//check(IsOutsideRenderPass());
 		if (Bypass())
@@ -2572,7 +2572,7 @@ public:
 		ALLOC_COMMAND(FRHICommandDrawIndexedIndirect)(IndexBufferRHI, ArgumentsBufferRHI, DrawArgumentsIndex, NumInstances);
 	}
 
-	FORCEINLINE_DEBUGGABLE void DrawIndexedPrimitiveIndirect(FIndexBufferRHIParamRef IndexBuffer, FVertexBufferRHIParamRef ArgumentsBuffer, uint32 ArgumentOffset)
+	FORCEINLINE_DEBUGGABLE void DrawIndexedPrimitiveIndirect(FRHIIndexBuffer* IndexBuffer, FRHIVertexBuffer* ArgumentsBuffer, uint32 ArgumentOffset)
 	{
 		//check(IsOutsideRenderPass());
 		if (Bypass())
@@ -2594,7 +2594,7 @@ public:
 		ALLOC_COMMAND(FRHICommandSetDepthBounds)(MinDepth, MaxDepth);
 	}
 
-	FORCEINLINE_DEBUGGABLE void CopyToResolveTarget(FTextureRHIParamRef SourceTextureRHI, FTextureRHIParamRef DestTextureRHI, const FResolveParams& ResolveParams)
+	FORCEINLINE_DEBUGGABLE void CopyToResolveTarget(FRHITexture* SourceTextureRHI, FRHITexture* DestTextureRHI, const FResolveParams& ResolveParams)
 	{
 		//check(IsOutsideRenderPass());
 		if (Bypass())
@@ -2605,7 +2605,7 @@ public:
 		ALLOC_COMMAND(FRHICommandCopyToResolveTarget)(SourceTextureRHI, DestTextureRHI, ResolveParams);
 	}
 
-	FORCEINLINE_DEBUGGABLE void CopyTexture(FTextureRHIParamRef SourceTextureRHI, FTextureRHIParamRef DestTextureRHI, const FRHICopyTextureInfo& CopyInfo)
+	FORCEINLINE_DEBUGGABLE void CopyTexture(FRHITexture* SourceTextureRHI, FRHITexture* DestTextureRHI, const FRHICopyTextureInfo& CopyInfo)
 	{
 		check(IsOutsideRenderPass());
 		if (GRHISupportsCopyToTextureMultipleMips)
@@ -2690,9 +2690,9 @@ public:
 		ALLOC_COMMAND(FRHICommandPollOcclusionQueries)();
 	}
 
-	FORCEINLINE_DEBUGGABLE void TransitionResource(EResourceTransitionAccess TransitionType, FTextureRHIParamRef InTexture)
+	FORCEINLINE_DEBUGGABLE void TransitionResource(EResourceTransitionAccess TransitionType, FRHITexture* InTexture)
 	{
-		FTextureRHIParamRef Texture = InTexture;
+		FRHITexture* Texture = InTexture;
 		check(Texture == nullptr || Texture->IsCommitted());
 		if (Bypass())
 		{
@@ -2701,12 +2701,12 @@ public:
 		}
 
 		// Allocate space to hold the single texture pointer inline in the command list itself.
-		FTextureRHIParamRef* TextureArray = (FTextureRHIParamRef*)Alloc(sizeof(FTextureRHIParamRef), alignof(FTextureRHIParamRef));
+		FRHITexture** TextureArray = (FRHITexture**)Alloc(sizeof(FRHITexture*), alignof(FRHITexture*));
 		TextureArray[0] = Texture;
 		ALLOC_COMMAND(FRHICommandTransitionTextures)(TransitionType, TextureArray, 1);
 	}
 
-	FORCEINLINE_DEBUGGABLE void TransitionResources(EResourceTransitionAccess TransitionType, FTextureRHIParamRef* InTextures, int32 NumTextures)
+	FORCEINLINE_DEBUGGABLE void TransitionResources(EResourceTransitionAccess TransitionType, FRHITexture** InTextures, int32 NumTextures)
 	{
 		if (Bypass())
 		{
@@ -2715,7 +2715,7 @@ public:
 		}
 
 		// Allocate space to hold the list of textures inline in the command list itself.
-		FTextureRHIParamRef* InlineTextureArray = (FTextureRHIParamRef*)Alloc(sizeof(FTextureRHIParamRef) * NumTextures, alignof(FTextureRHIParamRef));
+		FRHITexture** InlineTextureArray = (FRHITexture**)Alloc(sizeof(FRHITexture*) * NumTextures, alignof(FRHITexture*));
 		for (int32 Index = 0; Index < NumTextures; ++Index)
 		{
 			InlineTextureArray[Index] = InTextures[Index];
@@ -2724,7 +2724,7 @@ public:
 		ALLOC_COMMAND(FRHICommandTransitionTextures)(TransitionType, InlineTextureArray, NumTextures);
 	}
 
-	FORCEINLINE_DEBUGGABLE void TransitionResourceArrayNoCopy(EResourceTransitionAccess TransitionType, TArray<FTextureRHIParamRef>& InTextures)
+	FORCEINLINE_DEBUGGABLE void TransitionResourceArrayNoCopy(EResourceTransitionAccess TransitionType, TArray<FRHITexture*>& InTextures)
 	{
 		if (Bypass())
 		{
@@ -2788,7 +2788,7 @@ public:
 		ALLOC_COMMAND(FRHICommandWaitComputeFence<ECmdList::EGfx>)(WaitFence);
 	}
 
-	FORCEINLINE_DEBUGGABLE void CopyToStagingBuffer(FVertexBufferRHIParamRef SourceBuffer, FRHIStagingBuffer* DestinationStagingBuffer, uint32 Offset, uint32 NumBytes)
+	FORCEINLINE_DEBUGGABLE void CopyToStagingBuffer(FRHIVertexBuffer* SourceBuffer, FRHIStagingBuffer* DestinationStagingBuffer, uint32 Offset, uint32 NumBytes)
 	{
 		if (Bypass())
 		{
@@ -2902,7 +2902,7 @@ public:
 	// These 6 are special in that they must be called on the immediate command list and they force a flush only when we are not doing RHI thread
 	void BeginScene();
 	void EndScene();
-	void BeginDrawingViewport(FRHIViewport* Viewport, FTextureRHIParamRef RenderTargetRHI);
+	void BeginDrawingViewport(FRHIViewport* Viewport, FRHITexture* RenderTargetRHI);
 	void EndDrawingViewport(FRHIViewport* Viewport, bool bPresent, bool bLockToVsync);
 	void BeginFrame();
 	void EndFrame();
@@ -2965,7 +2965,7 @@ public:
 
 #if RHI_RAYTRACING
 	// Ray tracing API
-	FORCEINLINE_DEBUGGABLE void CopyBufferRegion(FVertexBufferRHIParamRef DestBuffer, uint64 DstOffset, FVertexBufferRHIParamRef SourceBuffer, uint64 SrcOffset, uint64 NumBytes)
+	FORCEINLINE_DEBUGGABLE void CopyBufferRegion(FRHIVertexBuffer* DestBuffer, uint64 DstOffset, FRHIVertexBuffer* SourceBuffer, uint64 SrcOffset, uint64 NumBytes)
 	{
 		// No copy/DMA operation inside render passes
 		check(IsOutsideRenderPass());
@@ -3098,7 +3098,7 @@ public:
 	FORCEINLINE_DEBUGGABLE void SetRayTracingHitGroup(
 		FRHIRayTracingScene* Scene, uint32 InstanceIndex, uint32 SegmentIndex, uint32 ShaderSlot,
 		FRayTracingPipelineState* Pipeline, uint32 HitGroupIndex,
-		uint32 NumUniformBuffers, const FUniformBufferRHIParamRef* UniformBuffers,
+		uint32 NumUniformBuffers, FRHIUniformBuffer* const* UniformBuffers,
 		uint32 LooseParameterDataSize, const void* LooseParameterData,
 		uint32 UserData)
 	{
@@ -3112,10 +3112,10 @@ public:
 		}
 		else
 		{
-			FUniformBufferRHIParamRef* InlineUniformBuffers = nullptr;
+			FRHIUniformBuffer** InlineUniformBuffers = nullptr;
 			if (NumUniformBuffers)
 			{
-				InlineUniformBuffers = (FUniformBufferRHIParamRef*)Alloc(sizeof(FUniformBufferRHIParamRef) * NumUniformBuffers, alignof(FUniformBufferRHIParamRef));
+				InlineUniformBuffers = (FRHIUniformBuffer**)Alloc(sizeof(FRHIUniformBuffer*) * NumUniformBuffers, alignof(FRHIUniformBuffer*));
 				for (uint32 Index = 0; Index < NumUniformBuffers; ++Index)
 				{
 					InlineUniformBuffers[Index] = UniformBuffers[Index];
@@ -3139,7 +3139,7 @@ public:
 	FORCEINLINE_DEBUGGABLE void SetRayTracingCallableShader(
 		FRHIRayTracingScene* Scene, uint32 ShaderSlotInScene,
 		FRayTracingPipelineState* Pipeline, uint32 ShaderIndexInPipeline,
-		uint32 NumUniformBuffers, const FUniformBufferRHIParamRef* UniformBuffers,
+		uint32 NumUniformBuffers, FRHIUniformBuffer* const* UniformBuffers,
 		uint32 UserData)
 	{
 		if (Bypass())
@@ -3149,10 +3149,10 @@ public:
 		}
 		else
 		{
-			FUniformBufferRHIParamRef* InlineUniformBuffers = nullptr;
+			FRHIUniformBuffer** InlineUniformBuffers = nullptr;
 			if (NumUniformBuffers)
 			{
-				InlineUniformBuffers = (FUniformBufferRHIParamRef*)Alloc(sizeof(FUniformBufferRHIParamRef) * NumUniformBuffers, alignof(FUniformBufferRHIParamRef));
+				InlineUniformBuffers = (FRHIUniformBuffer**)Alloc(sizeof(FRHIUniformBuffer*) * NumUniformBuffers, alignof(FRHIUniformBuffer*));
 				for (uint32 Index = 0; Index < NumUniformBuffers; ++Index)
 				{
 					InlineUniformBuffers[Index] = UniformBuffers[Index];
@@ -3176,7 +3176,7 @@ public:
 	void* operator new(size_t Size);
 	void operator delete(void *RawMemory);
 
-	FORCEINLINE_DEBUGGABLE void SetShaderUniformBuffer(FRHIComputeShader* Shader, uint32 BaseIndex, FUniformBufferRHIParamRef UniformBuffer)
+	FORCEINLINE_DEBUGGABLE void SetShaderUniformBuffer(FRHIComputeShader* Shader, uint32 BaseIndex, FRHIUniformBuffer* UniformBuffer)
 	{
 		if (Bypass())
 		{
@@ -3186,7 +3186,7 @@ public:
 		ALLOC_COMMAND(FRHICommandSetShaderUniformBuffer<FRHIComputeShader, ECmdList::ECompute>)(Shader, BaseIndex, UniformBuffer);
 	}
 	
-	FORCEINLINE void SetShaderUniformBuffer(FComputeShaderRHIRef& Shader, uint32 BaseIndex, FUniformBufferRHIParamRef UniformBuffer)
+	FORCEINLINE void SetShaderUniformBuffer(FComputeShaderRHIRef& Shader, uint32 BaseIndex, FRHIUniformBuffer* UniformBuffer)
 	{
 		SetShaderUniformBuffer(Shader.GetReference(), BaseIndex, UniformBuffer);
 	}
@@ -3208,7 +3208,7 @@ public:
 		SetShaderParameter(Shader.GetReference(), BufferIndex, BaseIndex, NumBytes, NewValue);
 	}
 	
-	FORCEINLINE_DEBUGGABLE void SetShaderTexture(FRHIComputeShader* Shader, uint32 TextureIndex, FTextureRHIParamRef Texture)
+	FORCEINLINE_DEBUGGABLE void SetShaderTexture(FRHIComputeShader* Shader, uint32 TextureIndex, FRHITexture* Texture)
 	{
 		if (Bypass())
 		{
@@ -3308,7 +3308,7 @@ public:
 		ALLOC_COMMAND(FRHICommandDispatchComputeShader<ECmdList::ECompute>)(ThreadGroupCountX, ThreadGroupCountY, ThreadGroupCountZ);
 	}
 
-	FORCEINLINE_DEBUGGABLE void DispatchIndirectComputeShader(FVertexBufferRHIParamRef ArgumentBuffer, uint32 ArgumentOffset)
+	FORCEINLINE_DEBUGGABLE void DispatchIndirectComputeShader(FRHIVertexBuffer* ArgumentBuffer, uint32 ArgumentOffset)
 	{
 		if (Bypass())
 		{
@@ -3417,7 +3417,7 @@ public:
 		ALLOC_COMMAND(FRHICommandWaitComputeFence<ECmdList::ECompute>)(WaitFence);
 	}
 
-	FORCEINLINE_DEBUGGABLE void CopyToStagingBuffer(FVertexBufferRHIParamRef SourceBuffer, FRHIStagingBuffer* DestinationStagingBuffer, uint32 Offset, uint32 NumBytes)
+	FORCEINLINE_DEBUGGABLE void CopyToStagingBuffer(FRHIVertexBuffer* SourceBuffer, FRHIStagingBuffer* DestinationStagingBuffer, uint32 Offset, uint32 NumBytes)
 	{
 		if (Bypass())
 		{
@@ -3688,12 +3688,12 @@ public:
 		return GDynamicRHI->CreateIndexBuffer_RenderThread(*this, Stride, Size, InUsage, CreateInfo);
 	}
 	
-	FORCEINLINE void* LockIndexBuffer(FIndexBufferRHIParamRef IndexBuffer, uint32 Offset, uint32 SizeRHI, EResourceLockMode LockMode)
+	FORCEINLINE void* LockIndexBuffer(FRHIIndexBuffer* IndexBuffer, uint32 Offset, uint32 SizeRHI, EResourceLockMode LockMode)
 	{
 		return GDynamicRHI->LockIndexBuffer_RenderThread(*this, IndexBuffer, Offset, SizeRHI, LockMode);
 	}
 	
-	FORCEINLINE void UnlockIndexBuffer(FIndexBufferRHIParamRef IndexBuffer)
+	FORCEINLINE void UnlockIndexBuffer(FRHIIndexBuffer* IndexBuffer)
 	{
 		GDynamicRHI->UnlockIndexBuffer_RenderThread(*this, IndexBuffer);
 	}
@@ -3718,17 +3718,17 @@ public:
 		return GDynamicRHI->CreateVertexBuffer_RenderThread(*this, Size, InUsage, CreateInfo);
 	}
 	
-	FORCEINLINE void* LockVertexBuffer(FVertexBufferRHIParamRef VertexBuffer, uint32 Offset, uint32 SizeRHI, EResourceLockMode LockMode)
+	FORCEINLINE void* LockVertexBuffer(FRHIVertexBuffer* VertexBuffer, uint32 Offset, uint32 SizeRHI, EResourceLockMode LockMode)
 	{
 		return GDynamicRHI->LockVertexBuffer_RenderThread(*this, VertexBuffer, Offset, SizeRHI, LockMode);
 	}
 	
-	FORCEINLINE void UnlockVertexBuffer(FVertexBufferRHIParamRef VertexBuffer)
+	FORCEINLINE void UnlockVertexBuffer(FRHIVertexBuffer* VertexBuffer)
 	{
 		GDynamicRHI->UnlockVertexBuffer_RenderThread(*this, VertexBuffer);
 	}
 	
-	FORCEINLINE void CopyVertexBuffer(FVertexBufferRHIParamRef SourceBuffer,FVertexBufferRHIParamRef DestBuffer)
+	FORCEINLINE void CopyVertexBuffer(FRHIVertexBuffer* SourceBuffer, FRHIVertexBuffer* DestBuffer)
 	{
 		QUICK_SCOPE_CYCLE_COUNTER(STAT_RHIMETHOD_CopyVertexBuffer_Flush);
 		ImmediateFlush(EImmediateFlushType::FlushRHIThread);  
@@ -3741,7 +3741,7 @@ public:
 		return GDynamicRHI->CreateStructuredBuffer_RenderThread(*this, Stride, Size, InUsage, CreateInfo);
 	}
 	
-	FORCEINLINE void* LockStructuredBuffer(FStructuredBufferRHIParamRef StructuredBuffer, uint32 Offset, uint32 SizeRHI, EResourceLockMode LockMode)
+	FORCEINLINE void* LockStructuredBuffer(FRHIStructuredBuffer* StructuredBuffer, uint32 Offset, uint32 SizeRHI, EResourceLockMode LockMode)
 	{
 		LLM_SCOPE(ELLMTag::RHIMisc);
 		QUICK_SCOPE_CYCLE_COUNTER(STAT_RHIMETHOD_LockStructuredBuffer_Flush);
@@ -3750,7 +3750,7 @@ public:
 		return GDynamicRHI->RHILockStructuredBuffer(StructuredBuffer, Offset, SizeRHI, LockMode);
 	}
 	
-	FORCEINLINE void UnlockStructuredBuffer(FStructuredBufferRHIParamRef StructuredBuffer)
+	FORCEINLINE void UnlockStructuredBuffer(FRHIStructuredBuffer* StructuredBuffer)
 	{
 		LLM_SCOPE(ELLMTag::RHIMisc);
 		QUICK_SCOPE_CYCLE_COUNTER(STAT_RHIMETHOD_UnlockStructuredBuffer_Flush);
@@ -3759,43 +3759,43 @@ public:
 		GDynamicRHI->RHIUnlockStructuredBuffer(StructuredBuffer);
 	}
 	
-	FORCEINLINE FUnorderedAccessViewRHIRef CreateUnorderedAccessView(FStructuredBufferRHIParamRef StructuredBuffer, bool bUseUAVCounter, bool bAppendBuffer)
+	FORCEINLINE FUnorderedAccessViewRHIRef CreateUnorderedAccessView(FRHIStructuredBuffer* StructuredBuffer, bool bUseUAVCounter, bool bAppendBuffer)
 	{
 		LLM_SCOPE(ELLMTag::RHIMisc);
 		return GDynamicRHI->RHICreateUnorderedAccessView_RenderThread(*this, StructuredBuffer, bUseUAVCounter, bAppendBuffer);
 	}
 	
-	FORCEINLINE FUnorderedAccessViewRHIRef CreateUnorderedAccessView(FTextureRHIParamRef Texture, uint32 MipLevel)
+	FORCEINLINE FUnorderedAccessViewRHIRef CreateUnorderedAccessView(FRHITexture* Texture, uint32 MipLevel)
 	{
 		LLM_SCOPE(ELLMTag::RHIMisc);
 		return GDynamicRHI->RHICreateUnorderedAccessView_RenderThread(*this, Texture, MipLevel);
 	}
 	
-	FORCEINLINE FUnorderedAccessViewRHIRef CreateUnorderedAccessView(FVertexBufferRHIParamRef VertexBuffer, uint8 Format)
+	FORCEINLINE FUnorderedAccessViewRHIRef CreateUnorderedAccessView(FRHIVertexBuffer* VertexBuffer, uint8 Format)
 	{
 		LLM_SCOPE(ELLMTag::RHIMisc);
 		return GDynamicRHI->RHICreateUnorderedAccessView_RenderThread(*this, VertexBuffer, Format);
 	}
 
-	FORCEINLINE FUnorderedAccessViewRHIRef CreateUnorderedAccessView(FIndexBufferRHIParamRef IndexBuffer, uint8 Format)
+	FORCEINLINE FUnorderedAccessViewRHIRef CreateUnorderedAccessView(FRHIIndexBuffer* IndexBuffer, uint8 Format)
 	{
 		LLM_SCOPE(ELLMTag::RHIMisc);
 		return GDynamicRHI->RHICreateUnorderedAccessView_RenderThread(*this, IndexBuffer, Format);
 	}
 
-	FORCEINLINE FShaderResourceViewRHIRef CreateShaderResourceView(FStructuredBufferRHIParamRef StructuredBuffer)
+	FORCEINLINE FShaderResourceViewRHIRef CreateShaderResourceView(FRHIStructuredBuffer* StructuredBuffer)
 	{
 		LLM_SCOPE(ELLMTag::RHIMisc);
 		return GDynamicRHI->RHICreateShaderResourceView_RenderThread(*this, StructuredBuffer);
 	}
 	
-	FORCEINLINE FShaderResourceViewRHIRef CreateShaderResourceView(FVertexBufferRHIParamRef VertexBuffer, uint32 Stride, uint8 Format)
+	FORCEINLINE FShaderResourceViewRHIRef CreateShaderResourceView(FRHIVertexBuffer* VertexBuffer, uint32 Stride, uint8 Format)
 	{
 		LLM_SCOPE(ELLMTag::RHIMisc);
 		return GDynamicRHI->CreateShaderResourceView_RenderThread(*this, VertexBuffer, Stride, Format);
 	}
 	
-	FORCEINLINE FShaderResourceViewRHIRef CreateShaderResourceView(FIndexBufferRHIParamRef Buffer)
+	FORCEINLINE FShaderResourceViewRHIRef CreateShaderResourceView(FRHIIndexBuffer* Buffer)
 	{
 		LLM_SCOPE(ELLMTag::RHIMisc);
 		return GDynamicRHI->CreateShaderResourceView_RenderThread(*this, Buffer);
@@ -3857,7 +3857,7 @@ public:
 		return GDynamicRHI->RHIAsyncCreateTexture2D(SizeX, SizeY, Format, NumMips, Flags, InitialMipData, NumInitialMips);
 	}
 	
-	FORCEINLINE void CopySharedMips(FTexture2DRHIParamRef DestTexture2D, FTexture2DRHIParamRef SrcTexture2D)
+	FORCEINLINE void CopySharedMips(FRHITexture2D* DestTexture2D, FRHITexture2D* SrcTexture2D)
 	{
 		QUICK_SCOPE_CYCLE_COUNTER(STAT_RHIMETHOD_CopySharedMips_Flush);
 		DestTexture2D->AddRef();
@@ -3871,7 +3871,7 @@ public:
 		});
 	}
 
-	FORCEINLINE void TransferTexture(FTexture2DRHIParamRef Texture, FIntRect Rect, uint32 SrcGPUIndex, uint32 DestGPUIndex, bool PullData)
+	FORCEINLINE void TransferTexture(FRHITexture2D* Texture, FIntRect Rect, uint32 SrcGPUIndex, uint32 DestGPUIndex, bool PullData)
 	{
 		LLM_SCOPE(ELLMTag::Textures);
 		ImmediateFlush(EImmediateFlushType::FlushRHIThread);
@@ -3897,25 +3897,25 @@ public:
 		return GDynamicRHI->RHICreateTexture3D_RenderThread(*this, SizeX, SizeY, SizeZ, Format, NumMips, Flags, CreateInfo);
 	}
 	
-	FORCEINLINE void GetResourceInfo(FTextureRHIParamRef Ref, FRHIResourceInfo& OutInfo)
+	FORCEINLINE void GetResourceInfo(FRHITexture* Ref, FRHIResourceInfo& OutInfo)
 	{
 		return RHIGetResourceInfo(Ref, OutInfo);
 	}
 	
-	FORCEINLINE FShaderResourceViewRHIRef CreateShaderResourceView(FTextureRHIParamRef Texture, const FRHITextureSRVCreateInfo& CreateInfo)
+	FORCEINLINE FShaderResourceViewRHIRef CreateShaderResourceView(FRHITexture* Texture, const FRHITextureSRVCreateInfo& CreateInfo)
 	{
 		LLM_SCOPE(ELLMTag::RHIMisc);
 		return GDynamicRHI->RHICreateShaderResourceView_RenderThread(*this, Texture, CreateInfo);
 	}
 
-	FORCEINLINE FShaderResourceViewRHIRef CreateShaderResourceView(FTextureRHIParamRef Texture, uint8 MipLevel)
+	FORCEINLINE FShaderResourceViewRHIRef CreateShaderResourceView(FRHITexture* Texture, uint8 MipLevel)
 	{
 		LLM_SCOPE(ELLMTag::RHIMisc);
 		const FRHITextureSRVCreateInfo CreateInfo(MipLevel, 1, Texture->GetFormat());
 		return GDynamicRHI->RHICreateShaderResourceView_RenderThread(*this, Texture, CreateInfo);
 	}
 	
-	FORCEINLINE FShaderResourceViewRHIRef CreateShaderResourceView(FTextureRHIParamRef Texture, uint8 MipLevel, uint8 NumMipLevels, uint8 Format)
+	FORCEINLINE FShaderResourceViewRHIRef CreateShaderResourceView(FRHITexture* Texture, uint8 MipLevel, uint8 NumMipLevels, uint8 Format)
 	{
 		LLM_SCOPE(ELLMTag::RHIMisc);
 		const FRHITextureSRVCreateInfo CreateInfo(MipLevel, NumMipLevels, Format);
@@ -3923,46 +3923,46 @@ public:
 	}
 
 	//UE_DEPRECATED(4.23, "This function is deprecated and will be removed in future releases. Renderer version implemented.")
-	FORCEINLINE void GenerateMips(FTextureRHIParamRef Texture)
+	FORCEINLINE void GenerateMips(FRHITexture* Texture)
 	{
 		QUICK_SCOPE_CYCLE_COUNTER(STAT_RHIMETHOD_GenerateMips_Flush);
 		ImmediateFlush(EImmediateFlushType::FlushRHIThread); return GDynamicRHI->RHIGenerateMips(Texture);
 	}
 	
-	FORCEINLINE uint32 ComputeMemorySize(FTextureRHIParamRef TextureRHI)
+	FORCEINLINE uint32 ComputeMemorySize(FRHITexture* TextureRHI)
 	{
 		return RHIComputeMemorySize(TextureRHI);
 	}
 	
-	FORCEINLINE FTexture2DRHIRef AsyncReallocateTexture2D(FTexture2DRHIParamRef Texture2D, int32 NewMipCount, int32 NewSizeX, int32 NewSizeY, FThreadSafeCounter* RequestStatus)
+	FORCEINLINE FTexture2DRHIRef AsyncReallocateTexture2D(FRHITexture2D* Texture2D, int32 NewMipCount, int32 NewSizeX, int32 NewSizeY, FThreadSafeCounter* RequestStatus)
 	{
 		LLM_SCOPE(ELLMTag::Textures);
 		return GDynamicRHI->AsyncReallocateTexture2D_RenderThread(*this, Texture2D, NewMipCount, NewSizeX, NewSizeY, RequestStatus);
 	}
 	
-	FORCEINLINE ETextureReallocationStatus FinalizeAsyncReallocateTexture2D(FTexture2DRHIParamRef Texture2D, bool bBlockUntilCompleted)
+	FORCEINLINE ETextureReallocationStatus FinalizeAsyncReallocateTexture2D(FRHITexture2D* Texture2D, bool bBlockUntilCompleted)
 	{
 		LLM_SCOPE(ELLMTag::Textures);
 		return GDynamicRHI->FinalizeAsyncReallocateTexture2D_RenderThread(*this, Texture2D, bBlockUntilCompleted);
 	}
 	
-	FORCEINLINE ETextureReallocationStatus CancelAsyncReallocateTexture2D(FTexture2DRHIParamRef Texture2D, bool bBlockUntilCompleted)
+	FORCEINLINE ETextureReallocationStatus CancelAsyncReallocateTexture2D(FRHITexture2D* Texture2D, bool bBlockUntilCompleted)
 	{
 		return GDynamicRHI->CancelAsyncReallocateTexture2D_RenderThread(*this, Texture2D, bBlockUntilCompleted);
 	}
 	
-	FORCEINLINE void* LockTexture2D(FTexture2DRHIParamRef Texture, uint32 MipIndex, EResourceLockMode LockMode, uint32& DestStride, bool bLockWithinMiptail, bool bFlushRHIThread = true)
+	FORCEINLINE void* LockTexture2D(FRHITexture2D* Texture, uint32 MipIndex, EResourceLockMode LockMode, uint32& DestStride, bool bLockWithinMiptail, bool bFlushRHIThread = true)
 	{
 		LLM_SCOPE(ELLMTag::Textures);
 		return GDynamicRHI->LockTexture2D_RenderThread(*this, Texture, MipIndex, LockMode, DestStride, bLockWithinMiptail, bFlushRHIThread);
 	}
 	
-	FORCEINLINE void UnlockTexture2D(FTexture2DRHIParamRef Texture, uint32 MipIndex, bool bLockWithinMiptail, bool bFlushRHIThread = true)
+	FORCEINLINE void UnlockTexture2D(FRHITexture2D* Texture, uint32 MipIndex, bool bLockWithinMiptail, bool bFlushRHIThread = true)
 	{		
 		GDynamicRHI->UnlockTexture2D_RenderThread(*this, Texture, MipIndex, bLockWithinMiptail, bFlushRHIThread);
 	}
 	
-	FORCEINLINE void* LockTexture2DArray(FTexture2DArrayRHIParamRef Texture, uint32 TextureIndex, uint32 MipIndex, EResourceLockMode LockMode, uint32& DestStride, bool bLockWithinMiptail)
+	FORCEINLINE void* LockTexture2DArray(FRHITexture2DArray* Texture, uint32 TextureIndex, uint32 MipIndex, EResourceLockMode LockMode, uint32& DestStride, bool bLockWithinMiptail)
 	{
 		LLM_SCOPE(ELLMTag::Textures);
 		QUICK_SCOPE_CYCLE_COUNTER(STAT_RHIMETHOD_LockTexture2DArray_Flush);
@@ -3970,7 +3970,7 @@ public:
 		return GDynamicRHI->RHILockTexture2DArray(Texture, TextureIndex, MipIndex, LockMode, DestStride, bLockWithinMiptail);
 	}
 	
-	FORCEINLINE void UnlockTexture2DArray(FTexture2DArrayRHIParamRef Texture, uint32 TextureIndex, uint32 MipIndex, bool bLockWithinMiptail)
+	FORCEINLINE void UnlockTexture2DArray(FRHITexture2DArray* Texture, uint32 TextureIndex, uint32 MipIndex, bool bLockWithinMiptail)
 	{
 		LLM_SCOPE(ELLMTag::Textures);
 		QUICK_SCOPE_CYCLE_COUNTER(STAT_RHIMETHOD_UnlockTexture2DArray_Flush);
@@ -3978,7 +3978,7 @@ public:
 		GDynamicRHI->RHIUnlockTexture2DArray(Texture, TextureIndex, MipIndex, bLockWithinMiptail);
 	}
 	
-	FORCEINLINE void UpdateTexture2D(FTexture2DRHIParamRef Texture, uint32 MipIndex, const struct FUpdateTextureRegion2D& UpdateRegion, uint32 SourcePitch, const uint8* SourceData)
+	FORCEINLINE void UpdateTexture2D(FRHITexture2D* Texture, uint32 MipIndex, const struct FUpdateTextureRegion2D& UpdateRegion, uint32 SourcePitch, const uint8* SourceData)
 	{		
 		checkf(UpdateRegion.DestX + UpdateRegion.Width <= Texture->GetSizeX(), TEXT("UpdateTexture2D out of bounds on X. Texture: %s, %i, %i, %i"), *Texture->GetName().ToString(), UpdateRegion.DestX, UpdateRegion.Width, Texture->GetSizeX());
 		checkf(UpdateRegion.DestY + UpdateRegion.Height <= Texture->GetSizeY(), TEXT("UpdateTexture2D out of bounds on Y. Texture: %s, %i, %i, %i"), *Texture->GetName().ToString(), UpdateRegion.DestY, UpdateRegion.Height, Texture->GetSizeY());
@@ -3986,7 +3986,7 @@ public:
 		GDynamicRHI->UpdateTexture2D_RenderThread(*this, Texture, MipIndex, UpdateRegion, SourcePitch, SourceData);
 	}
 
-	FORCEINLINE void UpdateFromBufferTexture2D(FTexture2DRHIParamRef Texture, uint32 MipIndex, const struct FUpdateTextureRegion2D& UpdateRegion, uint32 SourcePitch, FStructuredBufferRHIParamRef Buffer, uint32 BufferOffset)
+	FORCEINLINE void UpdateFromBufferTexture2D(FRHITexture2D* Texture, uint32 MipIndex, const struct FUpdateTextureRegion2D& UpdateRegion, uint32 SourcePitch, FRHIStructuredBuffer* Buffer, uint32 BufferOffset)
 	{
 		checkf(UpdateRegion.DestX + UpdateRegion.Width <= Texture->GetSizeX(), TEXT("UpdateFromBufferTexture2D out of bounds on X. Texture: %s, %i, %i, %i"), *Texture->GetName().ToString(), UpdateRegion.DestX, UpdateRegion.Width, Texture->GetSizeX());
 		checkf(UpdateRegion.DestY + UpdateRegion.Height <= Texture->GetSizeY(), TEXT("UpdateFromBufferTexture2D out of bounds on Y. Texture: %s, %i, %i, %i"), *Texture->GetName().ToString(), UpdateRegion.DestY, UpdateRegion.Height, Texture->GetSizeY());
@@ -3994,7 +3994,7 @@ public:
 		GDynamicRHI->UpdateFromBufferTexture2D_RenderThread(*this, Texture, MipIndex, UpdateRegion, SourcePitch, Buffer, BufferOffset);
 	}
 
-	FORCEINLINE FUpdateTexture3DData BeginUpdateTexture3D(FTexture3DRHIParamRef Texture, uint32 MipIndex, const struct FUpdateTextureRegion3D& UpdateRegion)
+	FORCEINLINE FUpdateTexture3DData BeginUpdateTexture3D(FRHITexture3D* Texture, uint32 MipIndex, const struct FUpdateTextureRegion3D& UpdateRegion)
 	{
 		checkf(UpdateRegion.DestX + UpdateRegion.Width <= Texture->GetSizeX(), TEXT("UpdateTexture3D out of bounds on X. Texture: %s, %i, %i, %i"), *Texture->GetName().ToString(), UpdateRegion.DestX, UpdateRegion.Width, Texture->GetSizeX());
 		checkf(UpdateRegion.DestY + UpdateRegion.Height <= Texture->GetSizeY(), TEXT("UpdateTexture3D out of bounds on Y. Texture: %s, %i, %i, %i"), *Texture->GetName().ToString(), UpdateRegion.DestY, UpdateRegion.Height, Texture->GetSizeY());
@@ -4015,7 +4015,7 @@ public:
 		GDynamicRHI->EndMultiUpdateTexture3D_RenderThread(*this, UpdateDataArray);
 	}
 	
-	FORCEINLINE void UpdateTexture3D(FTexture3DRHIParamRef Texture, uint32 MipIndex, const struct FUpdateTextureRegion3D& UpdateRegion, uint32 SourceRowPitch, uint32 SourceDepthPitch, const uint8* SourceData)
+	FORCEINLINE void UpdateTexture3D(FRHITexture3D* Texture, uint32 MipIndex, const struct FUpdateTextureRegion3D& UpdateRegion, uint32 SourceRowPitch, uint32 SourceDepthPitch, const uint8* SourceData)
 	{
 		checkf(UpdateRegion.DestX + UpdateRegion.Width <= Texture->GetSizeX(), TEXT("UpdateTexture3D out of bounds on X. Texture: %s, %i, %i, %i"), *Texture->GetName().ToString(), UpdateRegion.DestX, UpdateRegion.Width, Texture->GetSizeX());
 		checkf(UpdateRegion.DestY + UpdateRegion.Height <= Texture->GetSizeY(), TEXT("UpdateTexture3D out of bounds on Y. Texture: %s, %i, %i, %i"), *Texture->GetName().ToString(), UpdateRegion.DestY, UpdateRegion.Height, Texture->GetSizeY());
@@ -4036,19 +4036,19 @@ public:
 		return GDynamicRHI->RHICreateTextureCubeArray_RenderThread(*this, Size, ArraySize, Format, NumMips, Flags, CreateInfo);
 	}
 	
-	FORCEINLINE void* LockTextureCubeFace(FTextureCubeRHIParamRef Texture, uint32 FaceIndex, uint32 ArrayIndex, uint32 MipIndex, EResourceLockMode LockMode, uint32& DestStride, bool bLockWithinMiptail)
+	FORCEINLINE void* LockTextureCubeFace(FRHITextureCube* Texture, uint32 FaceIndex, uint32 ArrayIndex, uint32 MipIndex, EResourceLockMode LockMode, uint32& DestStride, bool bLockWithinMiptail)
 	{
 		LLM_SCOPE(ELLMTag::Textures);
 		return GDynamicRHI->RHILockTextureCubeFace_RenderThread(*this, Texture, FaceIndex, ArrayIndex, MipIndex, LockMode, DestStride, bLockWithinMiptail);
 	}
 	
-	FORCEINLINE void UnlockTextureCubeFace(FTextureCubeRHIParamRef Texture, uint32 FaceIndex, uint32 ArrayIndex, uint32 MipIndex, bool bLockWithinMiptail)
+	FORCEINLINE void UnlockTextureCubeFace(FRHITextureCube* Texture, uint32 FaceIndex, uint32 ArrayIndex, uint32 MipIndex, bool bLockWithinMiptail)
 	{
 		LLM_SCOPE(ELLMTag::Textures);
 		GDynamicRHI->RHIUnlockTextureCubeFace_RenderThread(*this, Texture, FaceIndex, ArrayIndex, MipIndex, bLockWithinMiptail);
 	}
 	
-	FORCEINLINE void BindDebugLabelName(FTextureRHIParamRef Texture, const TCHAR* Name)
+	FORCEINLINE void BindDebugLabelName(FRHITexture* Texture, const TCHAR* Name)
 	{
 		RHIBindDebugLabelName(Texture, Name);
 	}
@@ -4058,37 +4058,37 @@ public:
 		RHIBindDebugLabelName(UnorderedAccessViewRHI, Name);
 	}
 
-	FORCEINLINE void ReadSurfaceData(FTextureRHIParamRef Texture,FIntRect Rect,TArray<FColor>& OutData,FReadSurfaceDataFlags InFlags)
+	FORCEINLINE void ReadSurfaceData(FRHITexture* Texture,FIntRect Rect,TArray<FColor>& OutData,FReadSurfaceDataFlags InFlags)
 	{
 		QUICK_SCOPE_CYCLE_COUNTER(STAT_RHIMETHOD_ReadSurfaceData_Flush);
 		ImmediateFlush(EImmediateFlushType::FlushRHIThread);  
 		GDynamicRHI->RHIReadSurfaceData(Texture,Rect,OutData,InFlags);
 	}
 
-	FORCEINLINE void ReadSurfaceData(FTextureRHIParamRef Texture, FIntRect Rect, TArray<FLinearColor>& OutData, FReadSurfaceDataFlags InFlags)
+	FORCEINLINE void ReadSurfaceData(FRHITexture* Texture, FIntRect Rect, TArray<FLinearColor>& OutData, FReadSurfaceDataFlags InFlags)
 	{
 		QUICK_SCOPE_CYCLE_COUNTER(STAT_RHIMETHOD_ReadSurfaceData_Flush);
 		ImmediateFlush(EImmediateFlushType::FlushRHIThread);
 		GDynamicRHI->RHIReadSurfaceData(Texture, Rect, OutData, InFlags);
 	}
 	
-	FORCEINLINE void MapStagingSurface(FTextureRHIParamRef Texture,void*& OutData,int32& OutWidth,int32& OutHeight)
+	FORCEINLINE void MapStagingSurface(FRHITexture* Texture,void*& OutData,int32& OutWidth,int32& OutHeight)
 	{
 		GDynamicRHI->RHIMapStagingSurface_RenderThread(*this, Texture,OutData,OutWidth,OutHeight);
 	}
 	
-	FORCEINLINE void UnmapStagingSurface(FTextureRHIParamRef Texture)
+	FORCEINLINE void UnmapStagingSurface(FRHITexture* Texture)
 	{
 		GDynamicRHI->RHIUnmapStagingSurface_RenderThread(*this, Texture);
 	}
 	
-	FORCEINLINE void ReadSurfaceFloatData(FTextureRHIParamRef Texture,FIntRect Rect,TArray<FFloat16Color>& OutData,ECubeFace CubeFace,int32 ArrayIndex,int32 MipIndex)
+	FORCEINLINE void ReadSurfaceFloatData(FRHITexture* Texture,FIntRect Rect,TArray<FFloat16Color>& OutData,ECubeFace CubeFace,int32 ArrayIndex,int32 MipIndex)
 	{
 		LLM_SCOPE(ELLMTag::Textures);
 		GDynamicRHI->RHIReadSurfaceFloatData_RenderThread(*this, Texture,Rect,OutData,CubeFace,ArrayIndex,MipIndex);
 	}
 
-	FORCEINLINE void Read3DSurfaceFloatData(FTextureRHIParamRef Texture,FIntRect Rect,FIntPoint ZMinMax,TArray<FFloat16Color>& OutData)
+	FORCEINLINE void Read3DSurfaceFloatData(FRHITexture* Texture,FIntRect Rect,FIntPoint ZMinMax,TArray<FFloat16Color>& OutData)
 	{
 		QUICK_SCOPE_CYCLE_COUNTER(STAT_RHIMETHOD_Read3DSurfaceFloatData_Flush);
 		ImmediateFlush(EImmediateFlushType::FlushRHIThread);  
@@ -4109,7 +4109,7 @@ public:
 	}
 
 
-	FORCEINLINE void AcquireTransientResource_RenderThread(FTextureRHIParamRef Texture)
+	FORCEINLINE void AcquireTransientResource_RenderThread(FRHITexture* Texture)
 	{
 		if (!Texture->IsCommitted() )
 		{
@@ -4121,7 +4121,7 @@ public:
 		}
 	}
 
-	FORCEINLINE void DiscardTransientResource_RenderThread(FTextureRHIParamRef Texture)
+	FORCEINLINE void DiscardTransientResource_RenderThread(FRHITexture* Texture)
 	{
 		if (Texture->IsCommitted())
 		{
@@ -4133,7 +4133,7 @@ public:
 		}
 	}
 
-	FORCEINLINE void AcquireTransientResource_RenderThread(FVertexBufferRHIParamRef Buffer)
+	FORCEINLINE void AcquireTransientResource_RenderThread(FRHIVertexBuffer* Buffer)
 	{
 		if (!Buffer->IsCommitted())
 		{
@@ -4145,7 +4145,7 @@ public:
 		}
 	}
 
-	FORCEINLINE void DiscardTransientResource_RenderThread(FVertexBufferRHIParamRef Buffer)
+	FORCEINLINE void DiscardTransientResource_RenderThread(FRHIVertexBuffer* Buffer)
 	{
 		if (Buffer->IsCommitted())
 		{
@@ -4157,7 +4157,7 @@ public:
 		}
 	}
 
-	FORCEINLINE void AcquireTransientResource_RenderThread(FStructuredBufferRHIParamRef Buffer)
+	FORCEINLINE void AcquireTransientResource_RenderThread(FRHIStructuredBuffer* Buffer)
 	{
 		if (!Buffer->IsCommitted())
 		{
@@ -4169,7 +4169,7 @@ public:
 		}
 	}
 
-	FORCEINLINE void DiscardTransientResource_RenderThread(FStructuredBufferRHIParamRef Buffer)
+	FORCEINLINE void DiscardTransientResource_RenderThread(FRHIStructuredBuffer* Buffer)
 	{
 		if (Buffer->IsCommitted())
 		{
@@ -4248,7 +4248,7 @@ public:
 		RHITick(DeltaTime);
 	}
 	
-	FORCEINLINE void SetStreamOutTargets(uint32 NumTargets,const FVertexBufferRHIParamRef* VertexBuffers,const uint32* Offsets)
+	FORCEINLINE void SetStreamOutTargets(uint32 NumTargets, FRHIVertexBuffer* const* VertexBuffers,const uint32* Offsets)
 	{
 		QUICK_SCOPE_CYCLE_COUNTER(STAT_RHIMETHOD_SetStreamOutTargets_Flush);
 		ImmediateFlush(EImmediateFlushType::FlushRHIThread);  
@@ -4306,18 +4306,18 @@ public:
 		RHIGetSupportedResolution(Width, Height);
 	}
 	
-	FORCEINLINE void VirtualTextureSetFirstMipInMemory(FTexture2DRHIParamRef Texture, uint32 FirstMip)
+	FORCEINLINE void VirtualTextureSetFirstMipInMemory(FRHITexture2D* Texture, uint32 FirstMip)
 	{
 		GDynamicRHI->VirtualTextureSetFirstMipInMemory_RenderThread(*this, Texture, FirstMip);
 	}
 	
-	FORCEINLINE void VirtualTextureSetFirstMipVisible(FTexture2DRHIParamRef Texture, uint32 FirstMip)
+	FORCEINLINE void VirtualTextureSetFirstMipVisible(FRHITexture2D* Texture, uint32 FirstMip)
 	{
 		GDynamicRHI->VirtualTextureSetFirstMipVisible_RenderThread(*this, Texture, FirstMip);
 	}
 
 	UE_DEPRECATED(4.23, "CopySubTextureRegion API is deprecated; please use CopyTexture instead.")
-	FORCEINLINE void CopySubTextureRegion(FTexture2DRHIParamRef SourceTexture, FTexture2DRHIParamRef DestinationTexture, FBox2D SourceBox, FBox2D DestinationBox)
+	FORCEINLINE void CopySubTextureRegion(FRHITexture2D* SourceTexture, FRHITexture2D* DestinationTexture, FBox2D SourceBox, FBox2D DestinationBox)
 	{
 		GDynamicRHI->RHICopySubTextureRegion_RenderThread(*this, SourceTexture, DestinationTexture, SourceBox, DestinationBox);
 	}
@@ -4345,7 +4345,7 @@ public:
 	{
 		return RHIGetCommandContextContainer(Index, Num, GetGPUMask());
 	}
-	void UpdateTextureReference(FRHITextureReference* TextureRef, FTextureRHIParamRef NewTexture);
+	void UpdateTextureReference(FRHITextureReference* TextureRef, FRHITexture* NewTexture);
 
 	FORCEINLINE void PollRenderQueryResults()
 	{
@@ -4617,12 +4617,12 @@ FORCEINLINE FIndexBufferRHIRef RHIAsyncCreateIndexBuffer(uint32 Stride, uint32 S
 	return GDynamicRHI->RHICreateIndexBuffer(Stride, Size, InUsage, CreateInfo);
 }
 
-FORCEINLINE void* RHILockIndexBuffer(FIndexBufferRHIParamRef IndexBuffer, uint32 Offset, uint32 Size, EResourceLockMode LockMode)
+FORCEINLINE void* RHILockIndexBuffer(FRHIIndexBuffer* IndexBuffer, uint32 Offset, uint32 Size, EResourceLockMode LockMode)
 {
 	return FRHICommandListExecutor::GetImmediateCommandList().LockIndexBuffer(IndexBuffer, Offset, Size, LockMode);
 }
 
-FORCEINLINE void RHIUnlockIndexBuffer(FIndexBufferRHIParamRef IndexBuffer)
+FORCEINLINE void RHIUnlockIndexBuffer(FRHIIndexBuffer* IndexBuffer)
 {
 	 FRHICommandListExecutor::GetImmediateCommandList().UnlockIndexBuffer(IndexBuffer);
 }
@@ -4642,12 +4642,12 @@ FORCEINLINE FVertexBufferRHIRef RHIAsyncCreateVertexBuffer(uint32 Size, uint32 I
 	return GDynamicRHI->RHICreateVertexBuffer(Size, InUsage, CreateInfo);
 }
 
-FORCEINLINE void* RHILockVertexBuffer(FVertexBufferRHIParamRef VertexBuffer, uint32 Offset, uint32 SizeRHI, EResourceLockMode LockMode)
+FORCEINLINE void* RHILockVertexBuffer(FRHIVertexBuffer* VertexBuffer, uint32 Offset, uint32 SizeRHI, EResourceLockMode LockMode)
 {
 	return FRHICommandListExecutor::GetImmediateCommandList().LockVertexBuffer(VertexBuffer, Offset, SizeRHI, LockMode);
 }
 
-FORCEINLINE void RHIUnlockVertexBuffer(FVertexBufferRHIParamRef VertexBuffer)
+FORCEINLINE void RHIUnlockVertexBuffer(FRHIVertexBuffer* VertexBuffer)
 {
 	 FRHICommandListExecutor::GetImmediateCommandList().UnlockVertexBuffer(VertexBuffer);
 }
@@ -4657,47 +4657,47 @@ FORCEINLINE FStructuredBufferRHIRef RHICreateStructuredBuffer(uint32 Stride, uin
 	return FRHICommandListExecutor::GetImmediateCommandList().CreateStructuredBuffer(Stride, Size, InUsage, CreateInfo);
 }
 
-FORCEINLINE void* RHILockStructuredBuffer(FStructuredBufferRHIParamRef StructuredBuffer, uint32 Offset, uint32 SizeRHI, EResourceLockMode LockMode)
+FORCEINLINE void* RHILockStructuredBuffer(FRHIStructuredBuffer* StructuredBuffer, uint32 Offset, uint32 SizeRHI, EResourceLockMode LockMode)
 {
 	return FRHICommandListExecutor::GetImmediateCommandList().LockStructuredBuffer(StructuredBuffer, Offset, SizeRHI, LockMode);
 }
 
-FORCEINLINE void RHIUnlockStructuredBuffer(FStructuredBufferRHIParamRef StructuredBuffer)
+FORCEINLINE void RHIUnlockStructuredBuffer(FRHIStructuredBuffer* StructuredBuffer)
 {
 	 FRHICommandListExecutor::GetImmediateCommandList().UnlockStructuredBuffer(StructuredBuffer);
 }
 
-FORCEINLINE FUnorderedAccessViewRHIRef RHICreateUnorderedAccessView(FStructuredBufferRHIParamRef StructuredBuffer, bool bUseUAVCounter, bool bAppendBuffer)
+FORCEINLINE FUnorderedAccessViewRHIRef RHICreateUnorderedAccessView(FRHIStructuredBuffer* StructuredBuffer, bool bUseUAVCounter, bool bAppendBuffer)
 {
 	return FRHICommandListExecutor::GetImmediateCommandList().CreateUnorderedAccessView(StructuredBuffer, bUseUAVCounter, bAppendBuffer);
 }
 
-FORCEINLINE FUnorderedAccessViewRHIRef RHICreateUnorderedAccessView(FTextureRHIParamRef Texture, uint32 MipLevel = 0)
+FORCEINLINE FUnorderedAccessViewRHIRef RHICreateUnorderedAccessView(FRHITexture* Texture, uint32 MipLevel = 0)
 {
 	return FRHICommandListExecutor::GetImmediateCommandList().CreateUnorderedAccessView(Texture, MipLevel);
 }
 
-FORCEINLINE FUnorderedAccessViewRHIRef RHICreateUnorderedAccessView(FVertexBufferRHIParamRef VertexBuffer, uint8 Format)
+FORCEINLINE FUnorderedAccessViewRHIRef RHICreateUnorderedAccessView(FRHIVertexBuffer* VertexBuffer, uint8 Format)
 {
 	return FRHICommandListExecutor::GetImmediateCommandList().CreateUnorderedAccessView(VertexBuffer, Format);
 }
 
-FORCEINLINE FUnorderedAccessViewRHIRef RHICreateUnorderedAccessView(FIndexBufferRHIParamRef IndexBuffer, uint8 Format)
+FORCEINLINE FUnorderedAccessViewRHIRef RHICreateUnorderedAccessView(FRHIIndexBuffer* IndexBuffer, uint8 Format)
 {
 	return FRHICommandListExecutor::GetImmediateCommandList().CreateUnorderedAccessView(IndexBuffer, Format);
 }
 
-FORCEINLINE FShaderResourceViewRHIRef RHICreateShaderResourceView(FStructuredBufferRHIParamRef StructuredBuffer)
+FORCEINLINE FShaderResourceViewRHIRef RHICreateShaderResourceView(FRHIStructuredBuffer* StructuredBuffer)
 {
 	return FRHICommandListExecutor::GetImmediateCommandList().CreateShaderResourceView(StructuredBuffer);
 }
 
-FORCEINLINE FShaderResourceViewRHIRef RHICreateShaderResourceView(FVertexBufferRHIParamRef VertexBuffer, uint32 Stride, uint8 Format)
+FORCEINLINE FShaderResourceViewRHIRef RHICreateShaderResourceView(FRHIVertexBuffer* VertexBuffer, uint32 Stride, uint8 Format)
 {
 	return FRHICommandListExecutor::GetImmediateCommandList().CreateShaderResourceView(VertexBuffer, Stride, Format);
 }
 
-FORCEINLINE FShaderResourceViewRHIRef RHICreateShaderResourceView(FIndexBufferRHIParamRef Buffer)
+FORCEINLINE FShaderResourceViewRHIRef RHICreateShaderResourceView(FRHIIndexBuffer* Buffer)
 {
 	return FRHICommandListExecutor::GetImmediateCommandList().CreateShaderResourceView(Buffer);
 }
@@ -4712,7 +4712,7 @@ FORCEINLINE FTextureReferenceRHIRef RHICreateTextureReference(FLastRenderTimeCon
 	return FRHICommandListExecutor::GetImmediateCommandList().CreateTextureReference(LastRenderTime);
 }
 
-FORCEINLINE void RHIUpdateTextureReference(FRHITextureReference* TextureRef, FTextureRHIParamRef NewTexture)
+FORCEINLINE void RHIUpdateTextureReference(FRHITextureReference* TextureRef, FRHITexture* NewTexture)
 {
 	 FRHICommandListExecutor::GetImmediateCommandList().UpdateTextureReference(TextureRef, NewTexture);
 }
@@ -4732,12 +4732,12 @@ FORCEINLINE FTexture2DRHIRef RHIAsyncCreateTexture2D(uint32 SizeX, uint32 SizeY,
 	return FRHICommandListExecutor::GetImmediateCommandList().AsyncCreateTexture2D(SizeX, SizeY, Format, NumMips, Flags, InitialMipData, NumInitialMips);
 }
 
-FORCEINLINE void RHICopySharedMips(FTexture2DRHIParamRef DestTexture2D, FTexture2DRHIParamRef SrcTexture2D)
+FORCEINLINE void RHICopySharedMips(FRHITexture2D* DestTexture2D, FRHITexture2D* SrcTexture2D)
 {
 	return FRHICommandListExecutor::GetImmediateCommandList().CopySharedMips(DestTexture2D, SrcTexture2D);
 }
 
-FORCEINLINE void RHITransferTexture(FTexture2DRHIParamRef Texture, FIntRect Rect, uint32 SrcGPUIndex, uint32 DestGPUIndex, bool PullData)
+FORCEINLINE void RHITransferTexture(FRHITexture2D* Texture, FIntRect Rect, uint32 SrcGPUIndex, uint32 DestGPUIndex, bool PullData)
 {
 	return FRHICommandListExecutor::GetImmediateCommandList().TransferTexture(Texture, Rect, SrcGPUIndex, DestGPUIndex, PullData);
 }
@@ -4758,62 +4758,62 @@ FORCEINLINE FTexture3DRHIRef RHICreateTexture3D(uint32 SizeX, uint32 SizeY, uint
 	return FRHICommandListExecutor::GetImmediateCommandList().CreateTexture3D(SizeX, SizeY, SizeZ, Format, NumMips, Flags, CreateInfo);
 }
 
-FORCEINLINE FShaderResourceViewRHIRef RHICreateShaderResourceView(FTextureRHIParamRef Texture, uint8 MipLevel)
+FORCEINLINE FShaderResourceViewRHIRef RHICreateShaderResourceView(FRHITexture* Texture, uint8 MipLevel)
 {
 	return FRHICommandListExecutor::GetImmediateCommandList().CreateShaderResourceView(Texture, MipLevel);
 }
 
-FORCEINLINE FShaderResourceViewRHIRef RHICreateShaderResourceView(FTextureRHIParamRef Texture, uint8 MipLevel, uint8 NumMipLevels, uint8 Format)
+FORCEINLINE FShaderResourceViewRHIRef RHICreateShaderResourceView(FRHITexture* Texture, uint8 MipLevel, uint8 NumMipLevels, uint8 Format)
 {
 	return FRHICommandListExecutor::GetImmediateCommandList().CreateShaderResourceView(Texture, MipLevel, NumMipLevels, Format);
 }
 
-FORCEINLINE FShaderResourceViewRHIRef RHICreateShaderResourceView(FTextureRHIParamRef Texture, const FRHITextureSRVCreateInfo& CreateInfo)
+FORCEINLINE FShaderResourceViewRHIRef RHICreateShaderResourceView(FRHITexture* Texture, const FRHITextureSRVCreateInfo& CreateInfo)
 {
 	return FRHICommandListExecutor::GetImmediateCommandList().CreateShaderResourceView(Texture, CreateInfo);
 }
 
-FORCEINLINE FTexture2DRHIRef RHIAsyncReallocateTexture2D(FTexture2DRHIParamRef Texture2D, int32 NewMipCount, int32 NewSizeX, int32 NewSizeY, FThreadSafeCounter* RequestStatus)
+FORCEINLINE FTexture2DRHIRef RHIAsyncReallocateTexture2D(FRHITexture2D* Texture2D, int32 NewMipCount, int32 NewSizeX, int32 NewSizeY, FThreadSafeCounter* RequestStatus)
 {
 	return FRHICommandListExecutor::GetImmediateCommandList().AsyncReallocateTexture2D(Texture2D, NewMipCount, NewSizeX, NewSizeY, RequestStatus);
 }
 
-FORCEINLINE ETextureReallocationStatus RHIFinalizeAsyncReallocateTexture2D(FTexture2DRHIParamRef Texture2D, bool bBlockUntilCompleted)
+FORCEINLINE ETextureReallocationStatus RHIFinalizeAsyncReallocateTexture2D(FRHITexture2D* Texture2D, bool bBlockUntilCompleted)
 {
 	return FRHICommandListExecutor::GetImmediateCommandList().FinalizeAsyncReallocateTexture2D(Texture2D, bBlockUntilCompleted);
 }
 
-FORCEINLINE ETextureReallocationStatus RHICancelAsyncReallocateTexture2D(FTexture2DRHIParamRef Texture2D, bool bBlockUntilCompleted)
+FORCEINLINE ETextureReallocationStatus RHICancelAsyncReallocateTexture2D(FRHITexture2D* Texture2D, bool bBlockUntilCompleted)
 {
 	return FRHICommandListExecutor::GetImmediateCommandList().CancelAsyncReallocateTexture2D(Texture2D, bBlockUntilCompleted);
 }
 
-FORCEINLINE void* RHILockTexture2D(FTexture2DRHIParamRef Texture, uint32 MipIndex, EResourceLockMode LockMode, uint32& DestStride, bool bLockWithinMiptail, bool bFlushRHIThread = true)
+FORCEINLINE void* RHILockTexture2D(FRHITexture2D* Texture, uint32 MipIndex, EResourceLockMode LockMode, uint32& DestStride, bool bLockWithinMiptail, bool bFlushRHIThread = true)
 {
 	return FRHICommandListExecutor::GetImmediateCommandList().LockTexture2D(Texture, MipIndex, LockMode, DestStride, bLockWithinMiptail, bFlushRHIThread);
 }
 
-FORCEINLINE void RHIUnlockTexture2D(FTexture2DRHIParamRef Texture, uint32 MipIndex, bool bLockWithinMiptail, bool bFlushRHIThread = true)
+FORCEINLINE void RHIUnlockTexture2D(FRHITexture2D* Texture, uint32 MipIndex, bool bLockWithinMiptail, bool bFlushRHIThread = true)
 {
 	 FRHICommandListExecutor::GetImmediateCommandList().UnlockTexture2D(Texture, MipIndex, bLockWithinMiptail, bFlushRHIThread);
 }
 
-FORCEINLINE void* RHILockTexture2DArray(FTexture2DArrayRHIParamRef Texture, uint32 TextureIndex, uint32 MipIndex, EResourceLockMode LockMode, uint32& DestStride, bool bLockWithinMiptail)
+FORCEINLINE void* RHILockTexture2DArray(FRHITexture2DArray* Texture, uint32 TextureIndex, uint32 MipIndex, EResourceLockMode LockMode, uint32& DestStride, bool bLockWithinMiptail)
 {
 	return FRHICommandListExecutor::GetImmediateCommandList().LockTexture2DArray(Texture, TextureIndex, MipIndex, LockMode, DestStride, bLockWithinMiptail);
 }
 
-FORCEINLINE void RHIUnlockTexture2DArray(FTexture2DArrayRHIParamRef Texture, uint32 TextureIndex, uint32 MipIndex, bool bLockWithinMiptail)
+FORCEINLINE void RHIUnlockTexture2DArray(FRHITexture2DArray* Texture, uint32 TextureIndex, uint32 MipIndex, bool bLockWithinMiptail)
 {
 	 FRHICommandListExecutor::GetImmediateCommandList().UnlockTexture2DArray(Texture, TextureIndex, MipIndex, bLockWithinMiptail);
 }
 
-FORCEINLINE void RHIUpdateTexture2D(FTexture2DRHIParamRef Texture, uint32 MipIndex, const struct FUpdateTextureRegion2D& UpdateRegion, uint32 SourcePitch, const uint8* SourceData)
+FORCEINLINE void RHIUpdateTexture2D(FRHITexture2D* Texture, uint32 MipIndex, const struct FUpdateTextureRegion2D& UpdateRegion, uint32 SourcePitch, const uint8* SourceData)
 {
 	 FRHICommandListExecutor::GetImmediateCommandList().UpdateTexture2D(Texture, MipIndex, UpdateRegion, SourcePitch, SourceData);
 }
 
-FORCEINLINE FUpdateTexture3DData RHIBeginUpdateTexture3D(FTexture3DRHIParamRef Texture, uint32 MipIndex, const struct FUpdateTextureRegion3D& UpdateRegion)
+FORCEINLINE FUpdateTexture3DData RHIBeginUpdateTexture3D(FRHITexture3D* Texture, uint32 MipIndex, const struct FUpdateTextureRegion3D& UpdateRegion)
 {
 	return FRHICommandListExecutor::GetImmediateCommandList().BeginUpdateTexture3D(Texture, MipIndex, UpdateRegion);
 }
@@ -4828,7 +4828,7 @@ FORCEINLINE void RHIEndMultiUpdateTexture3D(TArray<FUpdateTexture3DData>& Update
 	FRHICommandListExecutor::GetImmediateCommandList().EndMultiUpdateTexture3D(UpdateDataArray);
 }
 
-FORCEINLINE void RHIUpdateTexture3D(FTexture3DRHIParamRef Texture, uint32 MipIndex, const struct FUpdateTextureRegion3D& UpdateRegion, uint32 SourceRowPitch, uint32 SourceDepthPitch, const uint8* SourceData)
+FORCEINLINE void RHIUpdateTexture3D(FRHITexture3D* Texture, uint32 MipIndex, const struct FUpdateTextureRegion3D& UpdateRegion, uint32 SourceRowPitch, uint32 SourceDepthPitch, const uint8* SourceData)
 {
 	 FRHICommandListExecutor::GetImmediateCommandList().UpdateTexture3D(Texture, MipIndex, UpdateRegion, SourceRowPitch, SourceDepthPitch, SourceData);
 }
@@ -4843,12 +4843,12 @@ FORCEINLINE FTextureCubeRHIRef RHICreateTextureCubeArray(uint32 Size, uint32 Arr
 	return FRHICommandListExecutor::GetImmediateCommandList().CreateTextureCubeArray(Size, ArraySize, Format, NumMips, Flags, CreateInfo);
 }
 
-FORCEINLINE void* RHILockTextureCubeFace(FTextureCubeRHIParamRef Texture, uint32 FaceIndex, uint32 ArrayIndex, uint32 MipIndex, EResourceLockMode LockMode, uint32& DestStride, bool bLockWithinMiptail)
+FORCEINLINE void* RHILockTextureCubeFace(FRHITextureCube* Texture, uint32 FaceIndex, uint32 ArrayIndex, uint32 MipIndex, EResourceLockMode LockMode, uint32& DestStride, bool bLockWithinMiptail)
 {
 	return FRHICommandListExecutor::GetImmediateCommandList().LockTextureCubeFace(Texture, FaceIndex, ArrayIndex, MipIndex, LockMode, DestStride, bLockWithinMiptail);
 }
 
-FORCEINLINE void RHIUnlockTextureCubeFace(FTextureCubeRHIParamRef Texture, uint32 FaceIndex, uint32 ArrayIndex, uint32 MipIndex, bool bLockWithinMiptail)
+FORCEINLINE void RHIUnlockTextureCubeFace(FRHITextureCube* Texture, uint32 FaceIndex, uint32 ArrayIndex, uint32 MipIndex, bool bLockWithinMiptail)
 {
 	 FRHICommandListExecutor::GetImmediateCommandList().UnlockTextureCubeFace(Texture, FaceIndex, ArrayIndex, MipIndex, bLockWithinMiptail);
 }
@@ -4861,32 +4861,32 @@ FORCEINLINE FRenderQueryRHIRef RHICreateRenderQuery(ERenderQueryType QueryType)
 	PRAGMA_ENABLE_DEPRECATION_WARNINGS
 }
 
-FORCEINLINE void RHIAcquireTransientResource(FTextureRHIParamRef Resource)
+FORCEINLINE void RHIAcquireTransientResource(FRHITexture* Resource)
 {
 	FRHICommandListExecutor::GetImmediateCommandList().AcquireTransientResource_RenderThread(Resource);
 }
 
-FORCEINLINE void RHIDiscardTransientResource(FTextureRHIParamRef Resource)
+FORCEINLINE void RHIDiscardTransientResource(FRHITexture* Resource)
 {
 	FRHICommandListExecutor::GetImmediateCommandList().DiscardTransientResource_RenderThread(Resource);
 }
 
-FORCEINLINE void RHIAcquireTransientResource(FVertexBufferRHIParamRef Resource)
+FORCEINLINE void RHIAcquireTransientResource(FRHIVertexBuffer* Resource)
 {
 	FRHICommandListExecutor::GetImmediateCommandList().AcquireTransientResource_RenderThread(Resource);
 }
 
-FORCEINLINE void RHIDiscardTransientResource(FVertexBufferRHIParamRef Resource)
+FORCEINLINE void RHIDiscardTransientResource(FRHIVertexBuffer* Resource)
 {
 	FRHICommandListExecutor::GetImmediateCommandList().DiscardTransientResource_RenderThread(Resource);
 }
 
-FORCEINLINE void RHIAcquireTransientResource(FStructuredBufferRHIParamRef Resource)
+FORCEINLINE void RHIAcquireTransientResource(FRHIStructuredBuffer* Resource)
 {
 	FRHICommandListExecutor::GetImmediateCommandList().AcquireTransientResource_RenderThread(Resource);
 }
 
-FORCEINLINE void RHIDiscardTransientResource(FStructuredBufferRHIParamRef Resource)
+FORCEINLINE void RHIDiscardTransientResource(FRHIStructuredBuffer* Resource)
 {
 	FRHICommandListExecutor::GetImmediateCommandList().DiscardTransientResource_RenderThread(Resource);
 }
@@ -4906,12 +4906,12 @@ FORCEINLINE void RHIFlushResources()
 	return FRHICommandListExecutor::GetImmediateCommandList().FlushResources();
 }
 
-FORCEINLINE void RHIVirtualTextureSetFirstMipInMemory(FTexture2DRHIParamRef Texture, uint32 FirstMip)
+FORCEINLINE void RHIVirtualTextureSetFirstMipInMemory(FRHITexture2D* Texture, uint32 FirstMip)
 {
 	 FRHICommandListExecutor::GetImmediateCommandList().VirtualTextureSetFirstMipInMemory(Texture, FirstMip);
 }
 
-FORCEINLINE void RHIVirtualTextureSetFirstMipVisible(FTexture2DRHIParamRef Texture, uint32 FirstMip)
+FORCEINLINE void RHIVirtualTextureSetFirstMipVisible(FRHITexture2D* Texture, uint32 FirstMip)
 {
 	 FRHICommandListExecutor::GetImmediateCommandList().VirtualTextureSetFirstMipVisible(Texture, FirstMip);
 }
@@ -4965,7 +4965,7 @@ struct TRHIResourceUpdateBatcher
 		}
 	}
 
-	void QueueUpdateRequest(FVertexBufferRHIParamRef DestVertexBuffer, FVertexBufferRHIParamRef SrcVertexBuffer)
+	void QueueUpdateRequest(FRHIVertexBuffer* DestVertexBuffer, FRHIVertexBuffer* SrcVertexBuffer)
 	{
 		FRHIResourceUpdateInfo& UpdateInfo = GetNextUpdateInfo();
 		UpdateInfo.Type = FRHIResourceUpdateInfo::UT_VertexBuffer;
@@ -4977,7 +4977,7 @@ struct TRHIResourceUpdateBatcher
 		}
 	}
 
-	void QueueUpdateRequest(FIndexBufferRHIParamRef DestIndexBuffer, FIndexBufferRHIParamRef SrcIndexBuffer)
+	void QueueUpdateRequest(FRHIIndexBuffer* DestIndexBuffer, FRHIIndexBuffer* SrcIndexBuffer)
 	{
 		FRHIResourceUpdateInfo& UpdateInfo = GetNextUpdateInfo();
 		UpdateInfo.Type = FRHIResourceUpdateInfo::UT_IndexBuffer;
@@ -4989,7 +4989,7 @@ struct TRHIResourceUpdateBatcher
 		}
 	}
 
-	void QueueUpdateRequest(FRHIShaderResourceView* SRV, FVertexBufferRHIParamRef VertexBuffer, uint32 Stride, uint8 Format)
+	void QueueUpdateRequest(FRHIShaderResourceView* SRV, FRHIVertexBuffer* VertexBuffer, uint32 Stride, uint8 Format)
 	{
 		FRHIResourceUpdateInfo& UpdateInfo = GetNextUpdateInfo();
 		UpdateInfo.Type = FRHIResourceUpdateInfo::UT_VertexBufferSRV;
@@ -5001,7 +5001,7 @@ struct TRHIResourceUpdateBatcher
 		}
 	}
 
-	void QueueUpdateRequest(FRHIShaderResourceView* SRV, FIndexBufferRHIParamRef IndexBuffer)
+	void QueueUpdateRequest(FRHIShaderResourceView* SRV, FRHIIndexBuffer* IndexBuffer)
 	{
 		// TODO
 	}
diff --git a/Engine/Source/Runtime/RHI/Public/RHIContext.h b/Engine/Source/Runtime/RHI/Public/RHIContext.h
--- a/Engine/Source/Runtime/RHI/Public/RHIContext.h
+++ b/Engine/Source/Runtime/RHI/Public/RHIContext.h
@@ -53,7 +53,7 @@ public:
 
 	virtual void RHIDispatchComputeShader(uint32 ThreadGroupCountX, uint32 ThreadGroupCountY, uint32 ThreadGroupCountZ) = 0;
 
-	virtual void RHIDispatchIndirectComputeShader(FVertexBufferRHIParamRef ArgumentBuffer, uint32 ArgumentOffset) = 0;
+	virtual void RHIDispatchIndirectComputeShader(FRHIVertexBuffer* ArgumentBuffer, uint32 ArgumentOffset) = 0;
 
 	virtual void RHISetAsyncComputeBudget(EAsyncComputeBudget Budget) {}
 
@@ -71,7 +71,7 @@ public:
 	virtual void RHITransitionResources(EResourceTransitionAccess TransitionType, EResourceTransitionPipeline TransitionPipeline, FRHIUnorderedAccessView** InUAVs, int32 NumUAVs, FRHIComputeFence* WriteComputeFence) = 0;
 
 	/** Set the shader resource view of a surface.  This is used for binding TextureMS parameter types that need a multi sampled view. */
-	virtual void RHISetShaderTexture(FRHIComputeShader* PixelShader, uint32 TextureIndex, FTextureRHIParamRef NewTexture) = 0;
+	virtual void RHISetShaderTexture(FRHIComputeShader* PixelShader, uint32 TextureIndex, FRHITexture* NewTexture) = 0;
 
 	/**
 	* Sets sampler state.
@@ -100,7 +100,7 @@ public:
 
 	virtual void RHISetShaderResourceViewParameter(FRHIComputeShader* ComputeShader, uint32 SamplerIndex, FRHIShaderResourceView* SRV) = 0;
 
-	virtual void RHISetShaderUniformBuffer(FRHIComputeShader* ComputeShader, uint32 BufferIndex, FUniformBufferRHIParamRef Buffer) = 0;
+	virtual void RHISetShaderUniformBuffer(FRHIComputeShader* ComputeShader, uint32 BufferIndex, FRHIUniformBuffer* Buffer) = 0;
 
 	virtual void RHISetShaderParameter(FRHIComputeShader* ComputeShader, uint32 BufferIndex, uint32 BaseIndex, uint32 NumBytes, const void* NewValue) = 0;
 
@@ -126,7 +126,7 @@ public:
 	 * @param Offset The start of the data in 'SourceBuffer'
 	 * @param NumBytes The number of bytes to copy out of 'SourceBuffer'
 	 */
-	virtual void RHICopyToStagingBuffer(FVertexBufferRHIParamRef SourceBufferRHI, FRHIStagingBuffer* DestinationStagingBufferRHI, uint32 InOffset, uint32 InNumBytes)
+	virtual void RHICopyToStagingBuffer(FRHIVertexBuffer* SourceBufferRHI, FRHIStagingBuffer* DestinationStagingBufferRHI, uint32 InOffset, uint32 InNumBytes)
 	{
 		check(false);
 	}
@@ -149,9 +149,9 @@ struct FAccelerationStructureUpdateParams
 
 struct FCopyBufferRegionParams
 {
-	FVertexBufferRHIParamRef DestBuffer;
+	FRHIVertexBuffer* DestBuffer;
 	uint64 DstOffset;
-	FVertexBufferRHIParamRef SourceBuffer;
+	FRHIVertexBuffer* SourceBuffer;
 	uint64 SrcOffset;
 	uint64 NumBytes;
 };
@@ -177,7 +177,7 @@ public:
 
 	virtual void RHIDispatchComputeShader(uint32 ThreadGroupCountX, uint32 ThreadGroupCountY, uint32 ThreadGroupCountZ) = 0;
 
-	virtual void RHIDispatchIndirectComputeShader(FVertexBufferRHIParamRef ArgumentBuffer, uint32 ArgumentOffset) = 0;
+	virtual void RHIDispatchIndirectComputeShader(FRHIVertexBuffer* ArgumentBuffer, uint32 ArgumentOffset) = 0;
 
 	virtual void RHIAutomaticCacheFlushAfterComputeShader(bool bEnable) = 0;
 
@@ -198,7 +198,7 @@ public:
 	* @param ResolveParams - optional resolve params
 	* @param Fence - optional fence, will be set once copy is completed by GPU
 	*/
-	virtual void RHICopyToResolveTarget(FTextureRHIParamRef SourceTexture, FTextureRHIParamRef DestTexture, const FResolveParams& ResolveParams) = 0;
+	virtual void RHICopyToResolveTarget(FRHITexture* SourceTexture, FRHITexture* DestTexture, const FResolveParams& ResolveParams) = 0;
 
 	/**
 	* Explicitly transition a texture resource from readable -> writable by the GPU or vice versa.
@@ -209,7 +209,7 @@ public:
 	* @param InTextures - array of texture objects to transition
 	* @param NumTextures - number of textures to transition
 	*/
-	virtual void RHITransitionResources(EResourceTransitionAccess TransitionType, FTextureRHIParamRef* InTextures, int32 NumTextures)
+	virtual void RHITransitionResources(EResourceTransitionAccess TransitionType, FRHITexture** InTextures, int32 NumTextures)
 	{
 		if (TransitionType == EResourceTransitionAccess::EReadable)
 		{
@@ -259,7 +259,7 @@ public:
 	virtual void RHIDiscardRenderTargets(bool Depth, bool Stencil, uint32 ColorBitMask) {}
 
 	// This method is queued with an RHIThread, otherwise it will flush after it is queued; without an RHI thread there is no benefit to queuing this frame advance commands
-	virtual void RHIBeginDrawingViewport(FRHIViewport* Viewport, FTextureRHIParamRef RenderTargetRHI) = 0;
+	virtual void RHIBeginDrawingViewport(FRHIViewport* Viewport, FRHITexture* RenderTargetRHI) = 0;
 
 	// This method is queued with an RHIThread, otherwise it will flush after it is queued; without an RHI thread there is no benefit to queuing this frame advance commands
 	virtual void RHIEndDrawingViewport(FRHIViewport* Viewport, bool bPresent, bool bLockToVsync) = 0;
@@ -287,12 +287,12 @@ public:
 	/**
 	* Signals the beginning and ending of rendering to a resource to be used in the next frame on a multiGPU system
 	*/
-	virtual void RHIBeginUpdateMultiFrameResource(FTextureRHIParamRef Texture)
+	virtual void RHIBeginUpdateMultiFrameResource(FRHITexture* Texture)
 	{
 		/* empty default implementation */
 	}
 
-	virtual void RHIEndUpdateMultiFrameResource(FTextureRHIParamRef Texture)
+	virtual void RHIEndUpdateMultiFrameResource(FRHITexture* Texture)
 	{
 		/* empty default implementation */
 	}
@@ -307,7 +307,7 @@ public:
 		/* empty default implementation */
 	}
 
-	virtual void RHISetStreamSource(uint32 StreamIndex, FVertexBufferRHIParamRef VertexBuffer, uint32 Offset) = 0;
+	virtual void RHISetStreamSource(uint32 StreamIndex, FRHIVertexBuffer* VertexBuffer, uint32 Offset) = 0;
 
 	// @param MinX including like Win32 RECT
 	// @param MinY including like Win32 RECT
@@ -329,22 +329,22 @@ public:
 	virtual void RHISetGraphicsPipelineState(FRHIGraphicsPipelineState* GraphicsState) = 0;
 
 	/** Set the shader resource view of a surface.  This is used for binding TextureMS parameter types that need a multi sampled view. */
-	virtual void RHISetShaderTexture(FRHIVertexShader* VertexShader, uint32 TextureIndex, FTextureRHIParamRef NewTexture) = 0;
+	virtual void RHISetShaderTexture(FRHIVertexShader* VertexShader, uint32 TextureIndex, FRHITexture* NewTexture) = 0;
 
 	/** Set the shader resource view of a surface.  This is used for binding TextureMS parameter types that need a multi sampled view. */
-	virtual void RHISetShaderTexture(FRHIHullShader* HullShader, uint32 TextureIndex, FTextureRHIParamRef NewTexture) = 0;
+	virtual void RHISetShaderTexture(FRHIHullShader* HullShader, uint32 TextureIndex, FRHITexture* NewTexture) = 0;
 
 	/** Set the shader resource view of a surface.  This is used for binding TextureMS parameter types that need a multi sampled view. */
-	virtual void RHISetShaderTexture(FRHIDomainShader* DomainShader, uint32 TextureIndex, FTextureRHIParamRef NewTexture) = 0;
+	virtual void RHISetShaderTexture(FRHIDomainShader* DomainShader, uint32 TextureIndex, FRHITexture* NewTexture) = 0;
 
 	/** Set the shader resource view of a surface.  This is used for binding TextureMS parameter types that need a multi sampled view. */
-	virtual void RHISetShaderTexture(FRHIGeometryShader* GeometryShader, uint32 TextureIndex, FTextureRHIParamRef NewTexture) = 0;
+	virtual void RHISetShaderTexture(FRHIGeometryShader* GeometryShader, uint32 TextureIndex, FRHITexture* NewTexture) = 0;
 
 	/** Set the shader resource view of a surface.  This is used for binding TextureMS parameter types that need a multi sampled view. */
-	virtual void RHISetShaderTexture(FRHIPixelShader* PixelShader, uint32 TextureIndex, FTextureRHIParamRef NewTexture) = 0;
+	virtual void RHISetShaderTexture(FRHIPixelShader* PixelShader, uint32 TextureIndex, FRHITexture* NewTexture) = 0;
 
 	/** Set the shader resource view of a surface.  This is used for binding TextureMS parameter types that need a multi sampled view. */
-	virtual void RHISetShaderTexture(FRHIComputeShader* PixelShader, uint32 TextureIndex, FTextureRHIParamRef NewTexture) = 0;
+	virtual void RHISetShaderTexture(FRHIComputeShader* PixelShader, uint32 TextureIndex, FRHITexture* NewTexture) = 0;
 
 	/**
 	* Sets sampler state.
@@ -423,17 +423,17 @@ public:
 
 	virtual void RHISetShaderResourceViewParameter(FRHIGeometryShader* GeometryShader, uint32 SamplerIndex, FRHIShaderResourceView* SRV) = 0;
 
-	virtual void RHISetShaderUniformBuffer(FRHIVertexShader* VertexShader, uint32 BufferIndex, FUniformBufferRHIParamRef Buffer) = 0;
+	virtual void RHISetShaderUniformBuffer(FRHIVertexShader* VertexShader, uint32 BufferIndex, FRHIUniformBuffer* Buffer) = 0;
 
-	virtual void RHISetShaderUniformBuffer(FRHIHullShader* HullShader, uint32 BufferIndex, FUniformBufferRHIParamRef Buffer) = 0;
+	virtual void RHISetShaderUniformBuffer(FRHIHullShader* HullShader, uint32 BufferIndex, FRHIUniformBuffer* Buffer) = 0;
 
-	virtual void RHISetShaderUniformBuffer(FRHIDomainShader* DomainShader, uint32 BufferIndex, FUniformBufferRHIParamRef Buffer) = 0;
+	virtual void RHISetShaderUniformBuffer(FRHIDomainShader* DomainShader, uint32 BufferIndex, FRHIUniformBuffer* Buffer) = 0;
 
-	virtual void RHISetShaderUniformBuffer(FRHIGeometryShader* GeometryShader, uint32 BufferIndex, FUniformBufferRHIParamRef Buffer) = 0;
+	virtual void RHISetShaderUniformBuffer(FRHIGeometryShader* GeometryShader, uint32 BufferIndex, FRHIUniformBuffer* Buffer) = 0;
 
-	virtual void RHISetShaderUniformBuffer(FRHIPixelShader* PixelShader, uint32 BufferIndex, FUniformBufferRHIParamRef Buffer) = 0;
+	virtual void RHISetShaderUniformBuffer(FRHIPixelShader* PixelShader, uint32 BufferIndex, FRHIUniformBuffer* Buffer) = 0;
 
-	virtual void RHISetShaderUniformBuffer(FRHIComputeShader* ComputeShader, uint32 BufferIndex, FUniformBufferRHIParamRef Buffer) = 0;
+	virtual void RHISetShaderUniformBuffer(FRHIComputeShader* ComputeShader, uint32 BufferIndex, FRHIUniformBuffer* Buffer) = 0;
 
 	virtual void RHISetShaderParameter(FRHIVertexShader* VertexShader, uint32 BufferIndex, uint32 BaseIndex, uint32 NumBytes, const void* NewValue) = 0;
 
@@ -462,14 +462,14 @@ public:
 
 	virtual void RHIDrawPrimitive(uint32 BaseVertexIndex, uint32 NumPrimitives, uint32 NumInstances) = 0;
 
-	virtual void RHIDrawPrimitiveIndirect(FVertexBufferRHIParamRef ArgumentBuffer, uint32 ArgumentOffset) = 0;
+	virtual void RHIDrawPrimitiveIndirect(FRHIVertexBuffer* ArgumentBuffer, uint32 ArgumentOffset) = 0;
 
-	virtual void RHIDrawIndexedIndirect(FIndexBufferRHIParamRef IndexBufferRHI, FStructuredBufferRHIParamRef ArgumentsBufferRHI, int32 DrawArgumentsIndex, uint32 NumInstances) = 0;
+	virtual void RHIDrawIndexedIndirect(FRHIIndexBuffer* IndexBufferRHI, FRHIStructuredBuffer* ArgumentsBufferRHI, int32 DrawArgumentsIndex, uint32 NumInstances) = 0;
 
 	// @param NumPrimitives need to be >0 
-	virtual void RHIDrawIndexedPrimitive(FIndexBufferRHIParamRef IndexBuffer, int32 BaseVertexIndex, uint32 FirstInstance, uint32 NumVertices, uint32 StartIndex, uint32 NumPrimitives, uint32 NumInstances) = 0;
+	virtual void RHIDrawIndexedPrimitive(FRHIIndexBuffer* IndexBuffer, int32 BaseVertexIndex, uint32 FirstInstance, uint32 NumVertices, uint32 StartIndex, uint32 NumPrimitives, uint32 NumInstances) = 0;
 
-	virtual void RHIDrawIndexedPrimitiveIndirect(FIndexBufferRHIParamRef IndexBuffer, FVertexBufferRHIParamRef ArgumentBuffer, uint32 ArgumentOffset) = 0;
+	virtual void RHIDrawIndexedPrimitiveIndirect(FRHIIndexBuffer* IndexBuffer, FRHIVertexBuffer* ArgumentBuffer, uint32 ArgumentOffset) = 0;
 
 	/**
 	* Sets Depth Bounds range with the given min/max depth.
@@ -479,7 +479,7 @@ public:
 	*/
 	virtual void RHISetDepthBounds(float MinDepth, float MaxDepth) = 0;
 
-	virtual void RHIUpdateTextureReference(FRHITextureReference* TextureRef, FTextureRHIParamRef NewTexture) = 0;
+	virtual void RHIUpdateTextureReference(FRHITextureReference* TextureRef, FRHITexture* NewTexture) = 0;
 
 	virtual void RHIBeginRenderPass(const FRHIRenderPassInfo& InInfo, const TCHAR* InName)
 	{
@@ -546,7 +546,7 @@ public:
 	{
 	}
 
-	virtual void RHICopyTexture(FTextureRHIParamRef SourceTexture, FTextureRHIParamRef DestTexture, const FRHICopyTextureInfo& CopyInfo)
+	virtual void RHICopyTexture(FRHITexture* SourceTexture, FRHITexture* DestTexture, const FRHICopyTextureInfo& CopyInfo)
 	{
 		const bool bIsCube = SourceTexture->GetTextureCube() != nullptr;
 		const bool bAllCubeFaces = bIsCube && (CopyInfo.NumSlices % 6) == 0;
@@ -569,7 +569,7 @@ public:
 		}
 	}
 #if RHI_RAYTRACING
-	virtual void RHICopyBufferRegion(FVertexBufferRHIParamRef DestBuffer, uint64 DstOffset, FVertexBufferRHIParamRef SourceBuffer, uint64 SrcOffset, uint64 NumBytes)
+	virtual void RHICopyBufferRegion(FRHIVertexBuffer* DestBuffer, uint64 DstOffset, FRHIVertexBuffer* SourceBuffer, uint64 SrcOffset, uint64 NumBytes)
 	{
 		checkNoEntry();
 	}
@@ -626,7 +626,7 @@ public:
 	virtual void RHISetRayTracingHitGroup(
 		FRHIRayTracingScene* Scene, uint32 InstanceIndex, uint32 SegmentIndex, uint32 ShaderSlot,
 		FRHIRayTracingPipelineState* Pipeline, uint32 HitGroupIndex,
-		uint32 NumUniformBuffers, const FUniformBufferRHIParamRef* UniformBuffers,
+		uint32 NumUniformBuffers, FRHIUniformBuffer* const* UniformBuffers,
 		uint32 LooseParameterDataSize, const void* LooseParameterData,
 		uint32 UserData)
 	{
@@ -636,7 +636,7 @@ public:
 	virtual void RHISetRayTracingCallableShader(
 		FRHIRayTracingScene* Scene, uint32 ShaderSlotInScene,
 		FRHIRayTracingPipelineState* Pipeline, uint32 ShaderIndexInPipeline,
-		uint32 NumUniformBuffers, const FUniformBufferRHIParamRef* UniformBuffers,
+		uint32 NumUniformBuffers, FRHIUniformBuffer* const* UniformBuffers,
 		uint32 UserData)
 	{
 		checkNoEntry();
diff --git a/Engine/Source/Runtime/RHI/Public/RHIGPUReadback.h b/Engine/Source/Runtime/RHI/Public/RHIGPUReadback.h
--- a/Engine/Source/Runtime/RHI/Public/RHIGPUReadback.h
+++ b/Engine/Source/Runtime/RHI/Public/RHIGPUReadback.h
@@ -38,12 +38,12 @@ public:
 	 * @param SourceBuffer The buffer holding the source data.
 	 * @param NumBytes The number of bytes to copy. If 0, this will copy the entire buffer.
 	 */
-	virtual void EnqueueCopy(FRHICommandList& RHICmdList, FVertexBufferRHIParamRef SourceBuffer, uint32 NumBytes = 0) 
+	virtual void EnqueueCopy(FRHICommandList& RHICmdList, FRHIVertexBuffer* SourceBuffer, uint32 NumBytes = 0)
 	{
 		unimplemented();
 	}
 
-	virtual void EnqueueCopy(FRHICommandList& RHICmdList, FTextureRHIParamRef SourceTexture, FResolveRect Rect = FResolveRect())
+	virtual void EnqueueCopy(FRHICommandList& RHICmdList, FRHITexture* SourceTexture, FResolveRect Rect = FResolveRect())
 	{
 		unimplemented();
 	}
@@ -72,7 +72,7 @@ public:
 
 	FRHIGPUBufferReadback(FName RequestName);
 	 
-	void EnqueueCopy(FRHICommandList& RHICmdList, FVertexBufferRHIParamRef SourceBuffer, uint32 NumBytes = 0) override;
+	void EnqueueCopy(FRHICommandList& RHICmdList, FRHIVertexBuffer* SourceBuffer, uint32 NumBytes = 0) override;
 	void* Lock(uint32 NumBytes) override;
 	void Unlock() override;
 
@@ -89,7 +89,7 @@ public:
 
 	FRHIGPUTextureReadback(FName RequestName);
 
-	void EnqueueCopy(FRHICommandList& RHICmdList, FTextureRHIParamRef SourceTexture, FResolveRect Rect = FResolveRect()) override;
+	void EnqueueCopy(FRHICommandList& RHICmdList, FRHITexture* SourceTexture, FResolveRect Rect = FResolveRect()) override;
 	void* Lock(uint32 NumBytes) override;
 	void Unlock() override;
 
diff --git a/Engine/Source/Runtime/RHI/Public/RHIResources.h b/Engine/Source/Runtime/RHI/Public/RHIResources.h
--- a/Engine/Source/Runtime/RHI/Public/RHIResources.h
+++ b/Engine/Source/Runtime/RHI/Public/RHIResources.h
@@ -1146,30 +1146,39 @@ UE_DEPRECATED(4.23, "FBoundShaderStateRHIParamRef typedef is deprecated; please
 typedef FRHIBoundShaderState*              FBoundShaderStateRHIParamRef;
 typedef TRefCountPtr<FRHIBoundShaderState> FBoundShaderStateRHIRef;
 
+UE_DEPRECATED(4.23, "FUniformBufferRHIParamRef typedef is deprecated; please use FRHIUniformBuffer* directly instead.")
 typedef FRHIUniformBuffer*              FUniformBufferRHIParamRef;
 typedef TRefCountPtr<FRHIUniformBuffer> FUniformBufferRHIRef;
 
+UE_DEPRECATED(4.23, "FIndexBufferRHIParamRef typedef is deprecated; please use FRHIIndexBuffer* directly instead.")
 typedef FRHIIndexBuffer*              FIndexBufferRHIParamRef;
 typedef TRefCountPtr<FRHIIndexBuffer> FIndexBufferRHIRef;
 
+UE_DEPRECATED(4.23, "FVertexBufferRHIParamRef typedef is deprecated; please use FRHIVertexBuffer* directly instead.")
 typedef FRHIVertexBuffer*              FVertexBufferRHIParamRef;
 typedef TRefCountPtr<FRHIVertexBuffer> FVertexBufferRHIRef;
 
+UE_DEPRECATED(4.23, "FStructuredBufferRHIParamRef typedef is deprecated; please use FRHIStructuredBuffer* directly instead.")
 typedef FRHIStructuredBuffer*              FStructuredBufferRHIParamRef;
 typedef TRefCountPtr<FRHIStructuredBuffer> FStructuredBufferRHIRef;
 
+UE_DEPRECATED(4.23, "FTextureRHIParamRef typedef is deprecated; please use FRHITexture* directly instead.")
 typedef FRHITexture*              FTextureRHIParamRef;
 typedef TRefCountPtr<FRHITexture> FTextureRHIRef;
 
+UE_DEPRECATED(4.23, "FTexture2DRHIParamRef typedef is deprecated; please use FRHITexture2D* directly instead.")
 typedef FRHITexture2D*              FTexture2DRHIParamRef;
 typedef TRefCountPtr<FRHITexture2D> FTexture2DRHIRef;
 
+UE_DEPRECATED(4.23, "FTexture2DArrayRHIParamRef typedef is deprecated; please use FRHITexture2DArray* directly instead.")
 typedef FRHITexture2DArray*              FTexture2DArrayRHIParamRef;
 typedef TRefCountPtr<FRHITexture2DArray> FTexture2DArrayRHIRef;
 
+UE_DEPRECATED(4.23, "FTexture3DRHIParamRef typedef is deprecated; please use FRHITexture3D* directly instead.")
 typedef FRHITexture3D*              FTexture3DRHIParamRef;
 typedef TRefCountPtr<FRHITexture3D> FTexture3DRHIRef;
 
+UE_DEPRECATED(4.23, "FTextureCubeRHIParamRef typedef is deprecated; please use FRHITextureCube* directly instead.")
 typedef FRHITextureCube*              FTextureCubeRHIParamRef;
 typedef TRefCountPtr<FRHITextureCube> FTextureCubeRHIRef;
 
@@ -1273,7 +1282,7 @@ typedef TRefCountPtr<FRHIStagingBuffer>	FStagingBufferRHIRef;
 class FRHIRenderTargetView
 {
 public:
-	FTextureRHIParamRef Texture;
+	FRHITexture* Texture;
 	uint32 MipIndex;
 
 	/** Array slice or texture cube face.  Only valid if texture resource was created with TexCreate_TargetArraySlicesIndependently! */
@@ -1299,7 +1308,7 @@ public:
 	{}
 
 	//common case
-	explicit FRHIRenderTargetView(FTextureRHIParamRef InTexture, ERenderTargetLoadAction InLoadAction) :
+	explicit FRHIRenderTargetView(FRHITexture* InTexture, ERenderTargetLoadAction InLoadAction) :
 		Texture(InTexture),
 		MipIndex(0),
 		ArraySliceIndex(-1),
@@ -1308,7 +1317,7 @@ public:
 	{}
 
 	//common case
-	explicit FRHIRenderTargetView(FTextureRHIParamRef InTexture, ERenderTargetLoadAction InLoadAction, uint32 InMipIndex, uint32 InArraySliceIndex) :
+	explicit FRHIRenderTargetView(FRHITexture* InTexture, ERenderTargetLoadAction InLoadAction, uint32 InMipIndex, uint32 InArraySliceIndex) :
 		Texture(InTexture),
 		MipIndex(InMipIndex),
 		ArraySliceIndex(InArraySliceIndex),
@@ -1316,7 +1325,7 @@ public:
 		StoreAction(ERenderTargetStoreAction::EStore)
 	{}
 	
-	explicit FRHIRenderTargetView(FTextureRHIParamRef InTexture, uint32 InMipIndex, uint32 InArraySliceIndex, ERenderTargetLoadAction InLoadAction, ERenderTargetStoreAction InStoreAction) :
+	explicit FRHIRenderTargetView(FRHITexture* InTexture, uint32 InMipIndex, uint32 InArraySliceIndex, ERenderTargetLoadAction InLoadAction, ERenderTargetStoreAction InStoreAction) :
 		Texture(InTexture),
 		MipIndex(InMipIndex),
 		ArraySliceIndex(InArraySliceIndex),
@@ -1494,7 +1503,7 @@ private:
 class FRHIDepthRenderTargetView
 {
 public:
-	FTextureRHIParamRef Texture;
+	FRHITexture* Texture;
 
 	ERenderTargetLoadAction		DepthLoadAction;
 	ERenderTargetStoreAction	DepthStoreAction;
@@ -1522,7 +1531,7 @@ public:
 	}
 
 	//common case
-	explicit FRHIDepthRenderTargetView(FTextureRHIParamRef InTexture, ERenderTargetLoadAction InLoadAction, ERenderTargetStoreAction InStoreAction) :
+	explicit FRHIDepthRenderTargetView(FRHITexture* InTexture, ERenderTargetLoadAction InLoadAction, ERenderTargetStoreAction InStoreAction) :
 		Texture(InTexture),
 		DepthLoadAction(InLoadAction),
 		DepthStoreAction(InStoreAction),
@@ -1533,7 +1542,7 @@ public:
 		Validate();
 	}
 
-	explicit FRHIDepthRenderTargetView(FTextureRHIParamRef InTexture, ERenderTargetLoadAction InLoadAction, ERenderTargetStoreAction InStoreAction, FExclusiveDepthStencil InDepthStencilAccess) :
+	explicit FRHIDepthRenderTargetView(FRHITexture* InTexture, ERenderTargetLoadAction InLoadAction, ERenderTargetStoreAction InStoreAction, FExclusiveDepthStencil InDepthStencilAccess) :
 		Texture(InTexture),
 		DepthLoadAction(InLoadAction),
 		DepthStoreAction(InStoreAction),
@@ -1544,7 +1553,7 @@ public:
 		Validate();
 	}
 
-	explicit FRHIDepthRenderTargetView(FTextureRHIParamRef InTexture, ERenderTargetLoadAction InDepthLoadAction, ERenderTargetStoreAction InDepthStoreAction, ERenderTargetLoadAction InStencilLoadAction, ERenderTargetStoreAction InStencilStoreAction) :
+	explicit FRHIDepthRenderTargetView(FRHITexture* InTexture, ERenderTargetLoadAction InDepthLoadAction, ERenderTargetStoreAction InDepthStoreAction, ERenderTargetLoadAction InStencilLoadAction, ERenderTargetStoreAction InStencilStoreAction) :
 		Texture(InTexture),
 		DepthLoadAction(InDepthLoadAction),
 		DepthStoreAction(InDepthStoreAction),
@@ -1555,7 +1564,7 @@ public:
 		Validate();
 	}
 
-	explicit FRHIDepthRenderTargetView(FTextureRHIParamRef InTexture, ERenderTargetLoadAction InDepthLoadAction, ERenderTargetStoreAction InDepthStoreAction, ERenderTargetLoadAction InStencilLoadAction, ERenderTargetStoreAction InStencilStoreAction, FExclusiveDepthStencil InDepthStencilAccess) :
+	explicit FRHIDepthRenderTargetView(FRHITexture* InTexture, ERenderTargetLoadAction InDepthLoadAction, ERenderTargetStoreAction InDepthStoreAction, ERenderTargetLoadAction InStencilLoadAction, ERenderTargetStoreAction InStencilStoreAction, FExclusiveDepthStencil InDepthStencilAccess) :
 		Texture(InTexture),
 		DepthLoadAction(InDepthLoadAction),
 		DepthStoreAction(InDepthStoreAction),
diff --git a/Engine/Source/Runtime/RHI/Public/RHIUtilities.h b/Engine/Source/Runtime/RHI/Public/RHIUtilities.h
--- a/Engine/Source/Runtime/RHI/Public/RHIUtilities.h
+++ b/Engine/Source/Runtime/RHI/Public/RHIUtilities.h
@@ -370,10 +370,10 @@ inline void DecodeRenderTargetMode(ESimpleRenderTargetMode Mode, ERenderTargetLo
 }
 
 UE_DEPRECATED(4.22, "TransitionSetRenderTargetsHelper API is deprecated; please use RHITransitionResources directly instead.")
-inline void TransitionSetRenderTargetsHelper(FRHICommandList& RHICmdList, FTextureRHIParamRef NewRenderTarget, FTextureRHIParamRef NewDepthStencilTarget, FExclusiveDepthStencil DepthStencilAccess)
+inline void TransitionSetRenderTargetsHelper(FRHICommandList& RHICmdList, FRHITexture* NewRenderTarget, FRHITexture* NewDepthStencilTarget, FExclusiveDepthStencil DepthStencilAccess)
 {
 	int32 TransitionIndex = 0;
-	FTextureRHIParamRef Transitions[2];
+	FRHITexture* Transitions[2];
 	if (NewRenderTarget)
 	{
 		Transitions[TransitionIndex] = NewRenderTarget;
@@ -388,9 +388,9 @@ inline void TransitionSetRenderTargetsHelper(FRHICommandList& RHICmdList, FTextu
 }
 
 UE_DEPRECATED(4.22, "TransitionSetRenderTargetsHelper API is deprecated; please use RHITransitionResources directly instead.")
-inline void TransitionSetRenderTargetsHelper(FRHICommandList& RHICmdList, uint32 NumRenderTargets, const FTextureRHIParamRef* NewRenderTargetsRHI, const FTextureRHIParamRef NewDepthStencilTargetRHI, FExclusiveDepthStencil DepthStencilAccess)
+inline void TransitionSetRenderTargetsHelper(FRHICommandList& RHICmdList, uint32 NumRenderTargets, FRHITexture* const* NewRenderTargetsRHI, FRHITexture* NewDepthStencilTargetRHI, FExclusiveDepthStencil DepthStencilAccess)
 {
-	FTextureRHIParamRef Transitions[MaxSimultaneousRenderTargets + 1];
+	FRHITexture* Transitions[MaxSimultaneousRenderTargets + 1];
 	int32 TransitionIndex = 0;
 	for (uint32 Index = 0; Index < NumRenderTargets; Index++)
 	{
@@ -411,7 +411,7 @@ inline void TransitionSetRenderTargetsHelper(FRHICommandList& RHICmdList, uint32
 
 inline void TransitionRenderPassTargets(FRHICommandList& RHICmdList, const FRHIRenderPassInfo& RPInfo)
 {
-	FTextureRHIParamRef Transitions[MaxSimultaneousRenderTargets + 1];
+	FRHITexture* Transitions[MaxSimultaneousRenderTargets + 1];
 	int32 TransitionIndex = 0;
 	uint32 NumColorRenderTargets = RPInfo.GetNumColorRenderTargets();
 	for (uint32 Index = 0; Index < NumColorRenderTargets; Index++)
@@ -436,7 +436,7 @@ inline void TransitionRenderPassTargets(FRHICommandList& RHICmdList, const FRHIR
 
 /** Helper for the common case of using a single color and depth render target. */
 UE_DEPRECATED(4.22, "SetRenderTarget API is deprecated; please use RHIBegin/EndRenderPass instead.")
-inline void SetRenderTarget(FRHICommandList& RHICmdList, FTextureRHIParamRef NewRenderTarget, FTextureRHIParamRef NewDepthStencilTarget, bool bWritableBarrier = false)
+inline void SetRenderTarget(FRHICommandList& RHICmdList, FRHITexture* NewRenderTarget, FRHITexture* NewDepthStencilTarget, bool bWritableBarrier = false)
 {
 	FRHIRenderTargetView RTV(NewRenderTarget, ERenderTargetLoadAction::ELoad);
 	FRHIDepthRenderTargetView DepthRTV(NewDepthStencilTarget, ERenderTargetLoadAction::ELoad, ERenderTargetStoreAction::EStore);
@@ -451,13 +451,13 @@ inline void SetRenderTarget(FRHICommandList& RHICmdList, FTextureRHIParamRef New
 	PRAGMA_ENABLE_DEPRECATION_WARNINGS
 }
 
-/** Helper that converts FTextureRHIParamRef's into FRHIRenderTargetView's. */
+/** Helper that converts FRHITexture*'s into FRHIRenderTargetView's. */
 UE_DEPRECATED(4.22, "SetRenderTargets API is deprecated; please use RHIBegin/EndRenderPass instead.")
 inline void SetRenderTargets(
 	FRHICommandList& RHICmdList,
 	uint32 NewNumSimultaneousRenderTargets,
-	const FTextureRHIParamRef* NewRenderTargetsRHI,
-	FTextureRHIParamRef NewDepthStencilTargetRHI,
+	FRHITexture* const* NewRenderTargetsRHI,
+	FRHITexture* NewDepthStencilTargetRHI,
 	ESimpleRenderTargetMode Mode,
 	FExclusiveDepthStencil DepthStencilAccess,
 	bool bWritableBarrier = false
