Cherry-pick 6953158: Added automatic ray tracing pipeline state object cache eviction. By default only up to 50 full pipeline states are kept in the cache. This reduces the likelihood of editor crash when many ray tracing shaders are being continuously loaded (i.e. when recompiling shaders).

Cache size can be controlled using r.RayTracing.PSOCacheSize CVar.
Oldest pipelines are removed from the cache after 10 frames of not being needed, until the total cache size reaches the specified value.

Individual compiled ray tracing shaders are not affected by this, only full linked pipelines.

Large levels with many shaders, such as Infiltrator, are still not stable unless PSOCacheSize is set to a very small number (<5) using NVIDIA driver 430.86.

diff --git a/Engine/Source/Runtime/RHI/Private/PipelineStateCache.cpp b/Engine/Source/Runtime/RHI/Private/PipelineStateCache.cpp
--- a/Engine/Source/Runtime/RHI/Private/PipelineStateCache.cpp
+++ b/Engine/Source/Runtime/RHI/Private/PipelineStateCache.cpp
@@ -61,6 +61,15 @@ static TAutoConsoleVariable<int32> CVarPSOEvictionTime(
 	ECVF_ReadOnly | ECVF_RenderThreadSafe
 );
 
+#if RHI_RAYTRACING
+static TAutoConsoleVariable<int32> CVarRTPSOCacheSize(
+	TEXT("r.RayTracing.PSOCacheSize"),
+	50,
+	TEXT("Number of ray tracing pipelines to keep in the cache (default = 50). Set to 0 to disable eviction.\n"),
+	ECVF_ReadOnly | ECVF_RenderThreadSafe
+);
+#endif // RHI_RAYTRACING
+
 extern void DumpPipelineCacheStats();
 
 static FAutoConsoleCommand DumpPipelineCmd(
@@ -260,7 +269,22 @@ public:
 		return false;
 	}
 
+	inline void AddHit()
+	{
+		if (LastFrameHit != GFrameCounter)
+		{
+			LastFrameHit = GFrameCounter;
+			HitsAcrossFrames++;
+		}
+
+		FPipelineState::AddHit();
+	}
+
 	FRayTracingPipelineStateRHIRef RHIPipeline;
+
+	uint64 HitsAcrossFrames = 0;
+	uint64 LastFrameHit = 0;
+
 #if PIPELINESTATECACHE_VERIFYTHREADSAFE
 	FThreadSafeCounter InUseCount;
 #endif
@@ -609,6 +633,8 @@ public:
 		if (FoundState)
 		{
 			OutCachedState = *FoundState;
+			OutCachedState->AddHit();
+
 			return true;
 		}
 		else
@@ -632,10 +658,74 @@ public:
 		}
 	}
 
+	void Trim(int32 TargetNumEntries)
+	{
+		FScopeLock ScopeLock(&CriticalSection);
+		
+		if (Cache.Num() < TargetNumEntries)
+		{
+			return;
+		}
+
+		struct FEntry
+		{
+			FRayTracingPipelineStateInitializer Key;
+			uint64 LastFrameHit;
+			uint64 HitsAcrossFrames;
+			FRayTracingPipelineState* Pipeline;
+		};
+		TArray<FEntry, TMemStackAllocator<>> Entries;
+		Entries.Reserve(Cache.Num());
+
+		const uint64 CurrentFrame = GFrameCounter;
+		const uint32 NumLatencyFrames = 10;
+
+		// Find all pipelines that were not used in the last 10 frames
+
+		for (const auto& It : Cache)
+		{
+			if (It.Value->LastFrameHit + NumLatencyFrames <= CurrentFrame)
+			{
+				FEntry Entry;
+				Entry.Key = It.Key;
+				Entry.HitsAcrossFrames = It.Value->HitsAcrossFrames;
+				Entry.LastFrameHit = It.Value->LastFrameHit;
+				Entry.Pipeline = It.Value;
+				Entries.Add(Entry);
+			}
+		}
+
+		Entries.Sort([](const FEntry& A, const FEntry& B)
+		{
+			if (A.LastFrameHit == B.LastFrameHit)
+			{
+				return B.HitsAcrossFrames < A.HitsAcrossFrames;
+			}
+			else
+			{
+				return B.LastFrameHit < A.LastFrameHit;
+			}
+		});
+
+		// Remove least useful pipelines
+
+		while (Cache.Num() > TargetNumEntries && Entries.Num())
+		{
+			delete Entries.Last().Pipeline;
+			Cache.Remove(Entries.Last().Key);
+			Entries.Pop(false);
+		}
+
+		LastTrimFrame = CurrentFrame;
+	}
+
+	uint64 GetLastTrimFrame() const { return LastTrimFrame; }
+
 private:
 
 	mutable FCriticalSection CriticalSection;
 	TMap<FRayTracingPipelineStateInitializer, FRayTracingPipelineState*> Cache;
+	uint64 LastTrimFrame = 0;
 };
 
 FRayTracingPipelineCache GRayTracingPipelineCache;
@@ -1008,6 +1098,13 @@ FRayTracingPipelineState* PipelineStateCache::GetAndOrCreateRayTracingPipelineSt
 	{
 		// #dxr_todo UE-68235: RT PSO disk caching
 
+		// Remove old pipelines once per frame
+		const int32 TargetCacheSize = CVarRTPSOCacheSize.GetValueOnAnyThread();
+		if (TargetCacheSize > 0 && GRayTracingPipelineCache.GetLastTrimFrame() != GFrameCounter)
+		{
+			GRayTracingPipelineCache.Trim(TargetCacheSize);
+		}
+
 		OutCachedState = new FRayTracingPipelineState();
 
 		if (DoAsyncCompile)
@@ -1025,12 +1122,6 @@ FRayTracingPipelineState* PipelineStateCache::GetAndOrCreateRayTracingPipelineSt
 
 		GRayTracingPipelineCache.Add(Initializer, OutCachedState);
 	}
-	else
-	{
-#if PSO_TRACK_CACHE_STATS
-		OutCachedState->AddHit();
-#endif
-	}
 
 	return OutCachedState;
 }
