Fix build warning.

diff --git a/Engine/Source/Runtime/RHI/Private/PipelineFileCache.cpp b/Engine/Source/Runtime/RHI/Private/PipelineFileCache.cpp
--- a/Engine/Source/Runtime/RHI/Private/PipelineFileCache.cpp
+++ b/Engine/Source/Runtime/RHI/Private/PipelineFileCache.cpp
@@ -623,8 +623,8 @@ bool FPipelineCacheFileFormatPSO::Verify() const
 			// Not using tessellation - we shouldn't try to draw patches
 			return false;
 		}
-		else if( GraphicsDesc.HullShader != FSHAHash() && GraphicsDesc.DomainShader == FSHAHash() ||
-				 GraphicsDesc.DomainShader == FSHAHash() && GraphicsDesc.HullShader != FSHAHash() )
+		else if( (GraphicsDesc.HullShader != FSHAHash() && GraphicsDesc.DomainShader == FSHAHash()) ||
+				 (GraphicsDesc.DomainShader == FSHAHash() && GraphicsDesc.HullShader != FSHAHash()) )
 		{
 			// Hull without Domain or vice-versa
 			return false;
