UE4.23 - Remove some SetRenderTarget* deprecated functions so we can clean up RHIs

diff --git a/Engine/Plugins/Runtime/WindowsMixedReality/Source/WindowsMixedRealityHMD/Private/WindowsMixedRealityHMD.cpp b/Engine/Plugins/Runtime/WindowsMixedReality/Source/WindowsMixedRealityHMD/Private/WindowsMixedRealityHMD.cpp
--- a/Engine/Plugins/Runtime/WindowsMixedReality/Source/WindowsMixedRealityHMD/Private/WindowsMixedRealityHMD.cpp
+++ b/Engine/Plugins/Runtime/WindowsMixedReality/Source/WindowsMixedRealityHMD/Private/WindowsMixedRealityHMD.cpp
@@ -986,9 +986,8 @@ namespace WindowsMixedReality
 
 		const uint32 x = (ViewportWidth - width) * 0.5f;
 		const uint32 y = (ViewportHeight - height) * 0.5f;
-		PRAGMA_DISABLE_DEPRECATION_WARNINGS
-		SetRenderTarget(RHICmdList, BackBuffer, FTextureRHIRef());
-		PRAGMA_ENABLE_DEPRECATION_WARNINGS
+		FRHIRenderPassInfo RPInfo(BackBuffer, ERenderTargetActions::Load_Store);
+		RHICmdList.BeginRenderPass(RPInfo, TEXT("WindowsMixedRealityHMD"));
 		DrawClearQuad(RHICmdList, FLinearColor(0.0f, 0.0f, 0.0f, 1.0f));
 		RHICmdList.SetViewport(x, y, 0, width + x, height + y, 1.0f);
 
@@ -1022,6 +1021,8 @@ namespace WindowsMixedReality
 			FIntPoint(1, 1),
 			*VertexShader,
 			EDRF_Default);
+
+		RHICmdList.EndRenderPass();
 	}
 
 	// Create a BGRA backbuffer for rendering.
diff --git a/Engine/Source/Runtime/GameplayMediaEncoder/Private/Microsoft/Windows/AmdAmfVideoEncoder.cpp b/Engine/Source/Runtime/GameplayMediaEncoder/Private/Microsoft/Windows/AmdAmfVideoEncoder.cpp
--- a/Engine/Source/Runtime/GameplayMediaEncoder/Private/Microsoft/Windows/AmdAmfVideoEncoder.cpp
+++ b/Engine/Source/Runtime/GameplayMediaEncoder/Private/Microsoft/Windows/AmdAmfVideoEncoder.cpp
@@ -365,9 +365,8 @@ void FAmdAmfVideoEncoder::ResolveBackBuffer(const FTexture2DRHIRef& BackBuffer,
 	}
 	else // Texture format mismatch, use a shader to do the copy.
 	{
-		PRAGMA_DISABLE_DEPRECATION_WARNINGS
-		SetRenderTarget(RHICmdList, ResolvedBackBuffer, FTextureRHIRef());
-		PRAGMA_ENABLE_DEPRECATION_WARNINGS
+		FRHIRenderPassInfo RPInfo(ResolvedBackBuffer, ERenderTargetActions::Load_Store);
+		RHICmdList.BeginRenderPass(RPInfo, TEXT("AmdAmfVideoEncoder"));
 		RHICmdList.SetViewport(0, 0, 0.0f, ResolvedBackBuffer->GetSizeX(), ResolvedBackBuffer->GetSizeY(), 1.0f);
 
 		FGraphicsPipelineStateInitializer GraphicsPSOInit;
@@ -388,9 +387,13 @@ void FAmdAmfVideoEncoder::ResolveBackBuffer(const FTexture2DRHIRef& BackBuffer,
 		SetGraphicsPipelineState(RHICmdList, GraphicsPSOInit);
 
 		if (ResolvedBackBuffer->GetSizeX() != BackBuffer->GetSizeX() || ResolvedBackBuffer->GetSizeY() != BackBuffer->GetSizeY())
+		{
 			PixelShader->SetParameters(RHICmdList, TStaticSamplerState<SF_Bilinear>::GetRHI(), BackBuffer);
+		}
 		else
+		{
 			PixelShader->SetParameters(RHICmdList, TStaticSamplerState<SF_Point>::GetRHI(), BackBuffer);
+		}
 
 		RendererModule->DrawRectangle(
 			RHICmdList,
@@ -404,7 +407,7 @@ void FAmdAmfVideoEncoder::ResolveBackBuffer(const FTexture2DRHIRef& BackBuffer,
 			*VertexShader,
 			EDRF_Default);
 
-		UnbindRenderTargets(RHICmdList);
+		RHICmdList.EndRenderPass();
 	}
 }
 
diff --git a/Engine/Source/Runtime/GameplayMediaEncoder/Private/Microsoft/Windows/NvVideoEncoder.cpp b/Engine/Source/Runtime/GameplayMediaEncoder/Private/Microsoft/Windows/NvVideoEncoder.cpp
--- a/Engine/Source/Runtime/GameplayMediaEncoder/Private/Microsoft/Windows/NvVideoEncoder.cpp
+++ b/Engine/Source/Runtime/GameplayMediaEncoder/Private/Microsoft/Windows/NvVideoEncoder.cpp
@@ -334,9 +334,8 @@ void FNvVideoEncoder::CopyBackBuffer(const FTexture2DRHIRef& SrcBackBuffer, cons
 	}
 	else // Texture format mismatch, use a shader to do the copy.
 	{
-		PRAGMA_DISABLE_DEPRECATION_WARNINGS
-		SetRenderTarget(RHICmdList, DstFrame.ResolvedBackBuffer, FTextureRHIRef());
-		PRAGMA_ENABLE_DEPRECATION_WARNINGS
+		FRHIRenderPassInfo RPInfo(DstFrame.ResolvedBackBuffer, ERenderTargetActions::Load_Store);
+		RHICmdList.BeginRenderPass(RPInfo, TEXT("NvVideoEncoder"));
 		RHICmdList.SetViewport(0, 0, 0.0f, DstFrame.ResolvedBackBuffer->GetSizeX(), DstFrame.ResolvedBackBuffer->GetSizeY(), 1.0f);
 	
 		FGraphicsPipelineStateInitializer GraphicsPSOInit;
@@ -357,9 +356,13 @@ void FNvVideoEncoder::CopyBackBuffer(const FTexture2DRHIRef& SrcBackBuffer, cons
 		SetGraphicsPipelineState(RHICmdList, GraphicsPSOInit);
 
 		if (DstFrame.ResolvedBackBuffer->GetSizeX() != SrcBackBuffer->GetSizeX() || DstFrame.ResolvedBackBuffer->GetSizeY() != SrcBackBuffer->GetSizeY())
+		{
 			PixelShader->SetParameters(RHICmdList, TStaticSamplerState<SF_Bilinear>::GetRHI(), SrcBackBuffer);
+		}
 		else
+		{
 			PixelShader->SetParameters(RHICmdList, TStaticSamplerState<SF_Point>::GetRHI(), SrcBackBuffer);
+		}
 	
 		RendererModule->DrawRectangle(
 			RHICmdList,
@@ -372,6 +375,7 @@ void FNvVideoEncoder::CopyBackBuffer(const FTexture2DRHIRef& SrcBackBuffer, cons
 			FIntPoint(1, 1),						// Source texture size
 			*VertexShader,
 			EDRF_Default);
+		RHICmdList.EndRenderPass();
 	}
 }
 
diff --git a/Engine/Source/Runtime/GameplayMediaEncoder/Private/Microsoft/Windows/WmfVideoEncoder.cpp b/Engine/Source/Runtime/GameplayMediaEncoder/Private/Microsoft/Windows/WmfVideoEncoder.cpp
--- a/Engine/Source/Runtime/GameplayMediaEncoder/Private/Microsoft/Windows/WmfVideoEncoder.cpp
+++ b/Engine/Source/Runtime/GameplayMediaEncoder/Private/Microsoft/Windows/WmfVideoEncoder.cpp
@@ -480,9 +480,8 @@ void FWmfVideoEncoder::ResolveBackBuffer(const FTexture2DRHIRef& BackBuffer, con
 	}
 	else // Texture format mismatch, use a shader to do the copy.
 	{
-		PRAGMA_DISABLE_DEPRECATION_WARNINGS
-		SetRenderTarget(RHICmdList, ResolvedBackBuffer, FTextureRHIRef());
-		PRAGMA_ENABLE_DEPRECATION_WARNINGS
+		FRHIRenderPassInfo RPInfo(ResolvedBackBuffer, ERenderTargetActions::Load_Store);
+		RHICmdList.BeginRenderPass(RPInfo, TEXT("WmfVideoEncoder"));
 		RHICmdList.SetViewport(0, 0, 0.0f, ResolvedBackBuffer->GetSizeX(), ResolvedBackBuffer->GetSizeY(), 1.0f);
 
 		FGraphicsPipelineStateInitializer GraphicsPSOInit;
@@ -503,9 +502,13 @@ void FWmfVideoEncoder::ResolveBackBuffer(const FTexture2DRHIRef& BackBuffer, con
 		SetGraphicsPipelineState(RHICmdList, GraphicsPSOInit);
 
 		if (ResolvedBackBuffer->GetSizeX() != BackBuffer->GetSizeX() || ResolvedBackBuffer->GetSizeY() != BackBuffer->GetSizeY())
+		{
 			PixelShader->SetParameters(RHICmdList, TStaticSamplerState<SF_Bilinear>::GetRHI(), BackBuffer);
+		}
 		else
+		{
 			PixelShader->SetParameters(RHICmdList, TStaticSamplerState<SF_Point>::GetRHI(), BackBuffer);
+		}
 
 		RendererModule->DrawRectangle(
 			RHICmdList,
@@ -518,6 +521,7 @@ void FWmfVideoEncoder::ResolveBackBuffer(const FTexture2DRHIRef& BackBuffer, con
 			FIntPoint(1, 1),						// Source texture size
 			*VertexShader,
 			EDRF_Default);
+		RHICmdList.EndRenderPass();
 	}
 }
 
diff --git a/Engine/Source/Runtime/RHI/Public/RHICommandList.h b/Engine/Source/Runtime/RHI/Public/RHICommandList.h
--- a/Engine/Source/Runtime/RHI/Public/RHICommandList.h
+++ b/Engine/Source/Runtime/RHI/Public/RHICommandList.h
@@ -2445,24 +2445,6 @@ public:
 			UAVs);
 	}
 
-	UE_DEPRECATED(4.22, "SetRenderTargets API is deprecated; please use RHIBegin/EndRenderPass instead.")
-	FORCEINLINE_DEBUGGABLE void SetRenderTargetsAndClear(const FRHISetRenderTargetsInfo& RenderTargetsInfo)
-	{
-		check(IsOutsideRenderPass());
-		CacheActiveRenderTargets(
-			RenderTargetsInfo.NumColorRenderTargets, 
-			RenderTargetsInfo.ColorRenderTarget, 
-			&RenderTargetsInfo.DepthStencilRenderTarget
-			);
-
-		if (Bypass())
-		{
-			GetContext().RHISetRenderTargetsAndClear(RenderTargetsInfo);
-			return;
-		}
-		ALLOC_COMMAND(FRHICommandSetRenderTargetsAndClear)(RenderTargetsInfo);
-	}	
-
 	FORCEINLINE_DEBUGGABLE void BindClearMRTValues(bool bClearColor, bool bClearDepth, bool bClearStencil)
 	{
 		//check(IsOutsideRenderPass());
diff --git a/Engine/Source/Runtime/RHI/Public/RHIUtilities.h b/Engine/Source/Runtime/RHI/Public/RHIUtilities.h
--- a/Engine/Source/Runtime/RHI/Public/RHIUtilities.h
+++ b/Engine/Source/Runtime/RHI/Public/RHIUtilities.h
@@ -369,46 +369,6 @@ inline void DecodeRenderTargetMode(ESimpleRenderTargetMode Mode, ERenderTargetLo
 	}
 }
 
-UE_DEPRECATED(4.22, "TransitionSetRenderTargetsHelper API is deprecated; please use RHITransitionResources directly instead.")
-inline void TransitionSetRenderTargetsHelper(FRHICommandList& RHICmdList, FRHITexture* NewRenderTarget, FRHITexture* NewDepthStencilTarget, FExclusiveDepthStencil DepthStencilAccess)
-{
-	int32 TransitionIndex = 0;
-	FRHITexture* Transitions[2];
-	if (NewRenderTarget)
-	{
-		Transitions[TransitionIndex] = NewRenderTarget;
-		++TransitionIndex;
-	}
-	if (NewDepthStencilTarget && DepthStencilAccess.IsDepthWrite())
-	{
-		Transitions[TransitionIndex] = NewDepthStencilTarget;
-		++TransitionIndex;
-	}
-	RHICmdList.TransitionResources(EResourceTransitionAccess::EWritable, Transitions, TransitionIndex);
-}
-
-UE_DEPRECATED(4.22, "TransitionSetRenderTargetsHelper API is deprecated; please use RHITransitionResources directly instead.")
-inline void TransitionSetRenderTargetsHelper(FRHICommandList& RHICmdList, uint32 NumRenderTargets, FRHITexture* const* NewRenderTargetsRHI, FRHITexture* NewDepthStencilTargetRHI, FExclusiveDepthStencil DepthStencilAccess)
-{
-	FRHITexture* Transitions[MaxSimultaneousRenderTargets + 1];
-	int32 TransitionIndex = 0;
-	for (uint32 Index = 0; Index < NumRenderTargets; Index++)
-	{
-		if (NewRenderTargetsRHI[Index])
-		{
-			Transitions[TransitionIndex] = NewRenderTargetsRHI[Index];
-			++TransitionIndex;
-		}
-	}
-
-	if (NewDepthStencilTargetRHI && DepthStencilAccess.IsDepthWrite())
-	{
-		Transitions[TransitionIndex] = NewDepthStencilTargetRHI;
-		++TransitionIndex;
-	}	
-	RHICmdList.TransitionResources(EResourceTransitionAccess::EWritable, Transitions, TransitionIndex);
-}
-
 inline void TransitionRenderPassTargets(FRHICommandList& RHICmdList, const FRHIRenderPassInfo& RPInfo)
 {
 	FRHITexture* Transitions[MaxSimultaneousRenderTargets + 1];
@@ -434,58 +394,6 @@ inline void TransitionRenderPassTargets(FRHICommandList& RHICmdList, const FRHIR
 	RHICmdList.TransitionResources(EResourceTransitionAccess::EWritable, Transitions, TransitionIndex);
 }
 
-/** Helper for the common case of using a single color and depth render target. */
-UE_DEPRECATED(4.22, "SetRenderTarget API is deprecated; please use RHIBegin/EndRenderPass instead.")
-inline void SetRenderTarget(FRHICommandList& RHICmdList, FRHITexture* NewRenderTarget, FRHITexture* NewDepthStencilTarget, bool bWritableBarrier = false)
-{
-	FRHIRenderTargetView RTV(NewRenderTarget, ERenderTargetLoadAction::ELoad);
-	FRHIDepthRenderTargetView DepthRTV(NewDepthStencilTarget, ERenderTargetLoadAction::ELoad, ERenderTargetStoreAction::EStore);
-
-	//make these rendertargets safely writable
-	PRAGMA_DISABLE_DEPRECATION_WARNINGS
-	if (bWritableBarrier)
-	{
-		TransitionSetRenderTargetsHelper(RHICmdList, NewRenderTarget, NewDepthStencilTarget, FExclusiveDepthStencil::DepthWrite_StencilWrite);
-	}
-	RHICmdList.SetRenderTargets(1, &RTV, &DepthRTV, 0, NULL);
-	PRAGMA_ENABLE_DEPRECATION_WARNINGS
-}
-
-/** Helper that converts FRHITexture*'s into FRHIRenderTargetView's. */
-UE_DEPRECATED(4.22, "SetRenderTargets API is deprecated; please use RHIBegin/EndRenderPass instead.")
-inline void SetRenderTargets(
-	FRHICommandList& RHICmdList,
-	uint32 NewNumSimultaneousRenderTargets,
-	FRHITexture* const* NewRenderTargetsRHI,
-	FRHITexture* NewDepthStencilTargetRHI,
-	ESimpleRenderTargetMode Mode,
-	FExclusiveDepthStencil DepthStencilAccess,
-	bool bWritableBarrier = false
-	)
-{
-	ERenderTargetLoadAction ColorLoadAction, DepthLoadAction, StencilLoadAction;
-	ERenderTargetStoreAction ColorStoreAction, DepthStoreAction, StencilStoreAction;
-	DecodeRenderTargetMode(Mode, ColorLoadAction, ColorStoreAction, DepthLoadAction, DepthStoreAction, StencilLoadAction, StencilStoreAction, DepthStencilAccess);
-
-	FRHIRenderTargetView RTVs[MaxSimultaneousRenderTargets];
-		
-	for (uint32 Index = 0; Index < NewNumSimultaneousRenderTargets; Index++)
-	{
-		RTVs[Index] = FRHIRenderTargetView(NewRenderTargetsRHI[Index], 0, -1, ColorLoadAction, ColorStoreAction);	
-	}
-
-	//make these rendertargets safely writable
-	PRAGMA_DISABLE_DEPRECATION_WARNINGS
-	if (bWritableBarrier)
-	{
-		TransitionSetRenderTargetsHelper(RHICmdList, NewNumSimultaneousRenderTargets, NewRenderTargetsRHI, NewDepthStencilTargetRHI, DepthStencilAccess);
-	}
-
-	FRHIDepthRenderTargetView DepthRTV(NewDepthStencilTargetRHI, DepthLoadAction, DepthStoreAction, StencilLoadAction, StencilStoreAction, DepthStencilAccess);
-	RHICmdList.SetRenderTargets(NewNumSimultaneousRenderTargets, RTVs, &DepthRTV, 0, nullptr);
-	PRAGMA_ENABLE_DEPRECATION_WARNINGS
-}
-
 // Will soon be deprecated as well...
 inline void UnbindRenderTargets(FRHICommandList& RHICmdList)
 {
diff --git a/Engine/Source/Runtime/Renderer/Private/PostProcess/SceneRenderTargets.h b/Engine/Source/Runtime/Renderer/Private/PostProcess/SceneRenderTargets.h
--- a/Engine/Source/Runtime/Renderer/Private/PostProcess/SceneRenderTargets.h
+++ b/Engine/Source/Runtime/Renderer/Private/PostProcess/SceneRenderTargets.h
@@ -216,7 +216,6 @@ public:
 	static int32 GetQuadOverdrawUAVIndex(EShaderPlatform Platform, ERHIFeatureLevel::Type FeatureLevel);
 #endif
 
-	void SetQuadOverdrawUAV(FRHICommandList& RHICmdList, bool bBindQuadOverdrawBuffers, FRHISetRenderTargetsInfo& Info);
 	void SetQuadOverdrawUAV(FRHICommandList& RHICmdList, bool bBindQuadOverdrawBuffers, bool bClearQuadOverdrawBuffers, FRHIRenderPassInfo& Info);
 
 	void BindVirtualTextureFeedbackUAV(FRHIRenderPassInfo& RPInfo);
