Use buffered locks for structured buffers when running with RHIT, same path as vertex/index buffers

diff --git a/Engine/Source/Runtime/RHI/Private/RHICommandList.cpp b/Engine/Source/Runtime/RHI/Private/RHICommandList.cpp
--- a/Engine/Source/Runtime/RHI/Private/RHICommandList.cpp
+++ b/Engine/Source/Runtime/RHI/Private/RHICommandList.cpp
@@ -1966,6 +1966,30 @@ struct FRHICommandUpdateIndexBuffer final : public FRHICommand<FRHICommandUpdate
 	}
 };
 
+struct FRHICommandUpdateStructuredBuffer final : public FRHICommand<FRHICommandUpdateStructuredBuffer>
+{
+	FStructuredBufferRHIParamRef StructuredBuffer;
+	void* Buffer;
+	uint32 BufferSize;
+	uint32 Offset;
+
+	FORCEINLINE_DEBUGGABLE FRHICommandUpdateStructuredBuffer(FStructuredBufferRHIParamRef InStructuredBuffer, void* InBuffer, uint32 InOffset, uint32 InBufferSize)
+		: StructuredBuffer(InStructuredBuffer)
+		, Buffer(InBuffer)
+		, BufferSize(InBufferSize)
+		, Offset(InOffset)
+	{
+	}
+	void Execute(FRHICommandListBase& CmdList)
+	{
+		QUICK_SCOPE_CYCLE_COUNTER(STAT_FRHICommandUpdateStructuredBuffer_Execute);
+		void* Data = GDynamicRHI->RHILockStructuredBuffer(StructuredBuffer, Offset, BufferSize, RLM_WriteOnly);
+		FMemory::Memcpy(Data, Buffer, BufferSize);
+		FMemory::Free(Buffer);
+		GDynamicRHI->RHIUnlockStructuredBuffer(StructuredBuffer);
+	}
+};
+
 static FLockTracker GLockTracker;
 
 
@@ -2087,6 +2111,65 @@ void FDynamicRHI::UnlockIndexBuffer_RenderThread(class FRHICommandListImmediate&
 	}
 }
 
+void* FDynamicRHI::LockStructuredBuffer_RenderThread(class FRHICommandListImmediate& RHICmdList, FStructuredBufferRHIParamRef StructuredBuffer, uint32 Offset, uint32 SizeRHI, EResourceLockMode LockMode)
+{
+	QUICK_SCOPE_CYCLE_COUNTER(STAT_FDynamicRHI_LockStructuredBuffer_RenderThread);
+	check(IsInRenderingThread());
+	bool bBuffer = CVarRHICmdBufferWriteLocks.GetValueOnRenderThread() > 0;
+	void* Result;
+	if (!bBuffer || LockMode != RLM_WriteOnly || RHICmdList.Bypass() || !IsRunningRHIInSeparateThread())
+	{
+		{
+			QUICK_SCOPE_CYCLE_COUNTER(STAT_RHIMETHOD_LockStructuredBuffer_Flush);
+			RHICmdList.ImmediateFlush(EImmediateFlushType::FlushRHIThread);
+		}
+		{
+			QUICK_SCOPE_CYCLE_COUNTER(STAT_RHIMETHOD_LockStructuredBuffer_Lock);
+			Result = GDynamicRHI->RHILockStructuredBuffer(StructuredBuffer, Offset, SizeRHI, LockMode);
+		}
+	}
+	else
+	{
+		QUICK_SCOPE_CYCLE_COUNTER(STAT_RHIMETHOD_LockStructuredBuffer_Malloc);
+		Result = FMemory::Malloc(SizeRHI, 16);
+	}
+	check(Result);
+	GLockTracker.Lock(StructuredBuffer, Result, Offset, SizeRHI, LockMode);
+	return Result;
+}
+
+void FDynamicRHI::UnlockStructuredBuffer_RenderThread(class FRHICommandListImmediate& RHICmdList, FStructuredBufferRHIParamRef StructuredBuffer)
+{
+	QUICK_SCOPE_CYCLE_COUNTER(STAT_FDynamicRHI_UnlockStructuredBuffer_RenderThread);
+	check(IsInRenderingThread());
+	bool bBuffer = CVarRHICmdBufferWriteLocks.GetValueOnRenderThread() > 0;
+	FLockTracker::FLockParams Params = GLockTracker.Unlock(StructuredBuffer);
+	if (!bBuffer || Params.LockMode != RLM_WriteOnly || RHICmdList.Bypass() || !IsRunningRHIInSeparateThread())
+	{
+		{
+			QUICK_SCOPE_CYCLE_COUNTER(STAT_RHIMETHOD_UnlockStructuredBuffer_Flush);
+			RHICmdList.ImmediateFlush(EImmediateFlushType::FlushRHIThread);
+		}
+		{
+			QUICK_SCOPE_CYCLE_COUNTER(STAT_RHIMETHOD_UnlockStructuredBuffer_Unlock);
+			GDynamicRHI->RHIUnlockStructuredBuffer(StructuredBuffer);
+			GLockTracker.TotalMemoryOutstanding = 0;
+		}
+	}	
+	else
+	{
+		ALLOC_COMMAND_CL(RHICmdList, FRHICommandUpdateStructuredBuffer)(StructuredBuffer, Params.Buffer, Params.Offset, Params.BufferSize);
+		RHICmdList.RHIThreadFence(true);
+		if (GLockTracker.TotalMemoryOutstanding > 256 * 1024)
+		{
+			QUICK_SCOPE_CYCLE_COUNTER(STAT_RHIMETHOD_UnlockStructuredBuffer_FlushForMem);
+			// we could be loading a level or something, lets get this stuff going
+			RHICmdList.ImmediateFlush(EImmediateFlushType::DispatchToRHIThread); 
+			GLockTracker.TotalMemoryOutstanding = 0;
+		}
+	}
+}
+
 // @todo-mattc-staging Default implementation
 void* FDynamicRHI::RHILockStagingBuffer(FRHIStagingBuffer* StagingBuffer, uint32 Offset, uint32 SizeRHI)
 {
diff --git a/Engine/Source/Runtime/RHI/Public/DynamicRHI.h b/Engine/Source/Runtime/RHI/Public/DynamicRHI.h
--- a/Engine/Source/Runtime/RHI/Public/DynamicRHI.h
+++ b/Engine/Source/Runtime/RHI/Public/DynamicRHI.h
@@ -1145,6 +1145,8 @@ public:
 	virtual void UnlockTexture2D_RenderThread(class FRHICommandListImmediate& RHICmdList, FRHITexture2D* Texture, uint32 MipIndex, bool bLockWithinMiptail, bool bNeedsDefaultRHIFlush = true);
 	virtual void UpdateTexture2D_RenderThread(class FRHICommandListImmediate& RHICmdList, FRHITexture2D* Texture, uint32 MipIndex, const struct FUpdateTextureRegion2D& UpdateRegion, uint32 SourcePitch, const uint8* SourceData);
 	virtual void UpdateFromBufferTexture2D_RenderThread(class FRHICommandListImmediate& RHICmdList, FRHITexture2D* Texture, uint32 MipIndex, const struct FUpdateTextureRegion2D& UpdateRegion, uint32 SourcePitch, FRHIStructuredBuffer* Buffer, uint32 BufferOffset);
+	virtual void* LockStructuredBuffer_RenderThread(class FRHICommandListImmediate& RHICmdList, FRHIStructuredBuffer* StructuredBuffer, uint32 Offset, uint32 SizeRHI, EResourceLockMode LockMode);
+	virtual void UnlockStructuredBuffer_RenderThread(class FRHICommandListImmediate& RHICmdList, FRHIStructuredBuffer* StructuredBuffer);
 
 	virtual FUpdateTexture3DData BeginUpdateTexture3D_RenderThread(class FRHICommandListImmediate& RHICmdList, FRHITexture3D* Texture, uint32 MipIndex, const struct FUpdateTextureRegion3D& UpdateRegion);
 	virtual void EndUpdateTexture3D_RenderThread(class FRHICommandListImmediate& RHICmdList, FUpdateTexture3DData& UpdateData);
diff --git a/Engine/Source/Runtime/RHI/Public/RHICommandList.h b/Engine/Source/Runtime/RHI/Public/RHICommandList.h
--- a/Engine/Source/Runtime/RHI/Public/RHICommandList.h
+++ b/Engine/Source/Runtime/RHI/Public/RHICommandList.h
@@ -3726,19 +3726,13 @@ public:
 	FORCEINLINE void* LockStructuredBuffer(FRHIStructuredBuffer* StructuredBuffer, uint32 Offset, uint32 SizeRHI, EResourceLockMode LockMode)
 	{
 		LLM_SCOPE(ELLMTag::RHIMisc);
-		QUICK_SCOPE_CYCLE_COUNTER(STAT_RHIMETHOD_LockStructuredBuffer_Flush);
-		ImmediateFlush(EImmediateFlushType::FlushRHIThread); 
-		 
-		return GDynamicRHI->RHILockStructuredBuffer(StructuredBuffer, Offset, SizeRHI, LockMode);
+		return GDynamicRHI->LockStructuredBuffer_RenderThread(*this, StructuredBuffer, Offset, SizeRHI, LockMode);
 	}
 	
 	FORCEINLINE void UnlockStructuredBuffer(FRHIStructuredBuffer* StructuredBuffer)
 	{
 		LLM_SCOPE(ELLMTag::RHIMisc);
-		QUICK_SCOPE_CYCLE_COUNTER(STAT_RHIMETHOD_UnlockStructuredBuffer_Flush);
-		ImmediateFlush(EImmediateFlushType::FlushRHIThread); 
-		  
-		GDynamicRHI->RHIUnlockStructuredBuffer(StructuredBuffer);
+		GDynamicRHI->UnlockStructuredBuffer_RenderThread(*this, StructuredBuffer);
 	}
 	
 	FORCEINLINE FUnorderedAccessViewRHIRef CreateUnorderedAccessView(FRHIStructuredBuffer* StructuredBuffer, bool bUseUAVCounter, bool bAppendBuffer)
