UE4- Deprecation fix
#rnx

diff --git a/Engine/Source/Runtime/RHI/Private/RHICommandList.cpp b/Engine/Source/Runtime/RHI/Private/RHICommandList.cpp
--- a/Engine/Source/Runtime/RHI/Private/RHICommandList.cpp
+++ b/Engine/Source/Runtime/RHI/Private/RHICommandList.cpp
@@ -1968,12 +1968,12 @@ struct FRHICommandUpdateIndexBuffer final : public FRHICommand<FRHICommandUpdate
 
 struct FRHICommandUpdateStructuredBuffer final : public FRHICommand<FRHICommandUpdateStructuredBuffer>
 {
-	FStructuredBufferRHIParamRef StructuredBuffer;
+	FRHIStructuredBuffer* StructuredBuffer;
 	void* Buffer;
 	uint32 BufferSize;
 	uint32 Offset;
 
-	FORCEINLINE_DEBUGGABLE FRHICommandUpdateStructuredBuffer(FStructuredBufferRHIParamRef InStructuredBuffer, void* InBuffer, uint32 InOffset, uint32 InBufferSize)
+	FORCEINLINE_DEBUGGABLE FRHICommandUpdateStructuredBuffer(FRHIStructuredBuffer* InStructuredBuffer, void* InBuffer, uint32 InOffset, uint32 InBufferSize)
 		: StructuredBuffer(InStructuredBuffer)
 		, Buffer(InBuffer)
 		, BufferSize(InBufferSize)
@@ -2111,7 +2111,7 @@ void FDynamicRHI::UnlockIndexBuffer_RenderThread(class FRHICommandListImmediate&
 	}
 }
 
-void* FDynamicRHI::LockStructuredBuffer_RenderThread(class FRHICommandListImmediate& RHICmdList, FStructuredBufferRHIParamRef StructuredBuffer, uint32 Offset, uint32 SizeRHI, EResourceLockMode LockMode)
+void* FDynamicRHI::LockStructuredBuffer_RenderThread(class FRHICommandListImmediate& RHICmdList, FRHIStructuredBuffer* StructuredBuffer, uint32 Offset, uint32 SizeRHI, EResourceLockMode LockMode)
 {
 	QUICK_SCOPE_CYCLE_COUNTER(STAT_FDynamicRHI_LockStructuredBuffer_RenderThread);
 	check(IsInRenderingThread());
@@ -2138,7 +2138,7 @@ void* FDynamicRHI::LockStructuredBuffer_RenderThread(class FRHICommandListImmedi
 	return Result;
 }
 
-void FDynamicRHI::UnlockStructuredBuffer_RenderThread(class FRHICommandListImmediate& RHICmdList, FStructuredBufferRHIParamRef StructuredBuffer)
+void FDynamicRHI::UnlockStructuredBuffer_RenderThread(class FRHICommandListImmediate& RHICmdList, FRHIStructuredBuffer* StructuredBuffer)
 {
 	QUICK_SCOPE_CYCLE_COUNTER(STAT_FDynamicRHI_UnlockStructuredBuffer_RenderThread);
 	check(IsInRenderingThread());
