Duplicate CL #4399472: Handle all Metal formats in ShaderPlatformToPlatformName - I have not made this handle all formats or platforms so this function is *wrong* for a number of others and should *assert* on failure to identify a platform. This fixes incorrect shader compilation for Forward rendering on iOS.

diff --git a/Engine/Source/Runtime/RHI/Private/RHI.cpp b/Engine/Source/Runtime/RHI/Private/RHI.cpp
--- a/Engine/Source/Runtime/RHI/Private/RHI.cpp
+++ b/Engine/Source/Runtime/RHI/Private/RHI.cpp
@@ -686,8 +686,13 @@ FName ShaderPlatformToPlatformName(EShaderPlatform Platform)
 	case SP_VULKAN_ES3_1_ANDROID:
 		return NAME_PLATFORM_ANDROID;
 	case SP_METAL:
+	case SP_METAL_MRT:
 		return NAME_PLATFORM_IOS;
 	case SP_METAL_SM5:
+	case SP_METAL_SM5_NOTESS:
+	case SP_METAL_MACES3_1:
+	case SP_METAL_MACES2:
+	case SP_METAL_MRT_MAC:
 		return NAME_PLATFORM_MAC;
 	case SP_SWITCH:
 	case SP_SWITCH_FORWARD:
