DR - Deprecate PrimitiveType

diff --git a/Engine/Plugins/Compositing/LensDistortion/Source/LensDistortion/Private/LensDistortionRendering.cpp b/Engine/Plugins/Compositing/LensDistortion/Source/LensDistortion/Private/LensDistortionRendering.cpp
--- a/Engine/Plugins/Compositing/LensDistortion/Source/LensDistortion/Private/LensDistortionRendering.cpp
+++ b/Engine/Plugins/Compositing/LensDistortion/Source/LensDistortion/Private/LensDistortionRendering.cpp
@@ -220,7 +220,7 @@ static void DrawUVDisplacementToRenderTarget_RenderThread(
 
 	// Draw grid.
 	uint32 PrimitiveCount = kGridSubdivisionX * kGridSubdivisionY * 2;
-	RHICmdList.DrawPrimitive(PT_TriangleList, 0, PrimitiveCount, 1);
+	RHICmdList.DrawPrimitive(0, PrimitiveCount, 1);
 
 	// Resolve render target.
 	RHICmdList.CopyToResolveTarget(
diff --git a/Engine/Plugins/Compositing/OpenCVLensDistortion/Source/OpenCVLensDistortion/Private/LensDistortionDisplacementMapRendering.cpp b/Engine/Plugins/Compositing/OpenCVLensDistortion/Source/OpenCVLensDistortion/Private/LensDistortionDisplacementMapRendering.cpp
--- a/Engine/Plugins/Compositing/OpenCVLensDistortion/Source/OpenCVLensDistortion/Private/LensDistortionDisplacementMapRendering.cpp
+++ b/Engine/Plugins/Compositing/OpenCVLensDistortion/Source/OpenCVLensDistortion/Private/LensDistortionDisplacementMapRendering.cpp
@@ -153,7 +153,7 @@ static void DrawUVDisplacementToRenderTargetFromPreComputedDisplacementMap_Rende
 
 	// Draw grid.
 	const uint32 PrimitiveCount = kGridSubdivisionX * kGridSubdivisionY * 2;
-	RHICmdList.DrawPrimitive(PT_TriangleList, 0, PrimitiveCount, 1);
+	RHICmdList.DrawPrimitive(0, PrimitiveCount, 1);
 
 	// Resolve render target.
 	RHICmdList.CopyToResolveTarget(OutTextureRenderTargetResource->GetRenderTargetTexture(), OutTextureRenderTargetResource->TextureRHI, FResolveParams());
diff --git a/Engine/Plugins/Experimental/SimpleHMD/Source/SimpleHMD/Private/SimpleHMD.cpp b/Engine/Plugins/Experimental/SimpleHMD/Source/SimpleHMD/Private/SimpleHMD.cpp
--- a/Engine/Plugins/Experimental/SimpleHMD/Source/SimpleHMD/Private/SimpleHMD.cpp
+++ b/Engine/Plugins/Experimental/SimpleHMD/Source/SimpleHMD/Private/SimpleHMD.cpp
@@ -214,7 +214,7 @@ void FSimpleHMD::DrawDistortionMesh_RenderThread(struct FRenderingCompositePassC
 	RHIUnlockIndexBuffer(IndexBufferRHI);
 
 	RHICmdList.SetStreamSource(0, VertexBufferRHI, 0);
-	RHICmdList.DrawIndexedPrimitive(IndexBufferRHI, PT_TriangleList, 0, 0, NumVerts, 0, NumTris, 1);
+	RHICmdList.DrawIndexedPrimitive(IndexBufferRHI, 0, 0, NumVerts, 0, NumTris, 1);
 
 	IndexBufferRHI.SafeRelease();
 	VertexBufferRHI.SafeRelease();
diff --git a/Engine/Plugins/Runtime/GoogleVR/GoogleVRHMD/Source/GoogleVRHMD/Private/GoogleVRHMDRender.cpp b/Engine/Plugins/Runtime/GoogleVR/GoogleVRHMD/Source/GoogleVRHMD/Private/GoogleVRHMDRender.cpp
--- a/Engine/Plugins/Runtime/GoogleVR/GoogleVRHMD/Source/GoogleVRHMD/Private/GoogleVRHMDRender.cpp
+++ b/Engine/Plugins/Runtime/GoogleVR/GoogleVRHMD/Source/GoogleVRHMD/Private/GoogleVRHMDRender.cpp
@@ -166,13 +166,13 @@ void FGoogleVRHMD::DrawDistortionMesh_RenderThread(struct FRenderingCompositePas
 		{
 			RHICmdList.SetViewport(0, 0, 0.0f, XBound, TextureSize.Y, 1.0f);
 			RHICmdList.SetStreamSource(0, VertexBufferRHI, 0);
-			RHICmdList.DrawIndexedPrimitive(GTwoTrianglesIndexBuffer.IndexBufferRHI, PT_TriangleList, 0, 0, LocalNumVertsPerEye, 0, LocalNumTrisPerEye, 1);
+			RHICmdList.DrawIndexedPrimitive(GTwoTrianglesIndexBuffer.IndexBufferRHI, 0, 0, LocalNumVertsPerEye, 0, LocalNumTrisPerEye, 1);
 		}
 		else
 		{
 			RHICmdList.SetViewport(XBound, 0, 0.0f, TextureSize.X, TextureSize.Y, 1.0f);
 			RHICmdList.SetStreamSource(0, VertexBufferRHI, 0);
-			RHICmdList.DrawIndexedPrimitive(GTwoTrianglesIndexBuffer.IndexBufferRHI, PT_TriangleList, 0, 0, LocalNumVertsPerEye, 0, LocalNumTrisPerEye, 1);
+			RHICmdList.DrawIndexedPrimitive(GTwoTrianglesIndexBuffer.IndexBufferRHI, 0, 0, LocalNumVertsPerEye, 0, LocalNumTrisPerEye, 1);
 		}
 		VertexBufferRHI.SafeRelease();
 	}
diff --git a/Engine/Plugins/Runtime/Oculus/OculusVR/Source/OculusHMD/Private/OculusHMD.cpp b/Engine/Plugins/Runtime/Oculus/OculusVR/Source/OculusHMD/Private/OculusHMD.cpp
--- a/Engine/Plugins/Runtime/Oculus/OculusVR/Source/OculusHMD/Private/OculusHMD.cpp
+++ b/Engine/Plugins/Runtime/Oculus/OculusVR/Source/OculusHMD/Private/OculusHMD.cpp
@@ -942,7 +942,7 @@ namespace OculusHMD
 		check(Mesh.IsValid());
 
 		RHICmdList.SetStreamSource(0, Mesh.VertexBufferRHI, 0);
-		RHICmdList.DrawIndexedPrimitive(Mesh.IndexBufferRHI, PT_TriangleList, 0, 0, Mesh.NumVertices, 0, Mesh.NumTriangles, 1);
+		RHICmdList.DrawIndexedPrimitive(Mesh.IndexBufferRHI, 0, 0, Mesh.NumVertices, 0, Mesh.NumTriangles, 1);
 	}
 
 
diff --git a/Engine/Plugins/Runtime/Steam/SteamVR/Source/SteamVR/Private/SteamVRRender.cpp b/Engine/Plugins/Runtime/Steam/SteamVR/Source/SteamVR/Private/SteamVRRender.cpp
--- a/Engine/Plugins/Runtime/Steam/SteamVR/Source/SteamVR/Private/SteamVRRender.cpp
+++ b/Engine/Plugins/Runtime/Steam/SteamVR/Source/SteamVR/Private/SteamVRRender.cpp
@@ -55,7 +55,7 @@ static void DrawOcclusionMesh(FRHICommandList& RHICmdList, EStereoscopicPass Ste
 	check(Mesh.IsValid());
 
 	RHICmdList.SetStreamSource(0, Mesh.VertexBufferRHI, 0);
-	RHICmdList.DrawIndexedPrimitive(Mesh.IndexBufferRHI, PT_TriangleList, 0, 0, Mesh.NumVertices, 0, Mesh.NumTriangles, 1);
+	RHICmdList.DrawIndexedPrimitive(Mesh.IndexBufferRHI, 0, 0, Mesh.NumVertices, 0, Mesh.NumTriangles, 1);
 }
 
 void FSteamVRHMD::DrawHiddenAreaMesh_RenderThread(FRHICommandList& RHICmdList, EStereoscopicPass StereoPass) const
diff --git a/Engine/Plugins/Runtime/WindowsMixedReality/Source/WindowsMixedRealityHMD/Private/WindowsMixedRealityHMD.cpp b/Engine/Plugins/Runtime/WindowsMixedReality/Source/WindowsMixedRealityHMD/Private/WindowsMixedRealityHMD.cpp
--- a/Engine/Plugins/Runtime/WindowsMixedReality/Source/WindowsMixedRealityHMD/Private/WindowsMixedRealityHMD.cpp
+++ b/Engine/Plugins/Runtime/WindowsMixedReality/Source/WindowsMixedRealityHMD/Private/WindowsMixedRealityHMD.cpp
@@ -706,7 +706,7 @@ namespace WindowsMixedReality
 		check(Mesh.IsValid());
 
 		RHICmdList.SetStreamSource(0, Mesh.VertexBufferRHI, 0);
-		RHICmdList.DrawIndexedPrimitive(Mesh.IndexBufferRHI, PT_TriangleList, 0, 0, Mesh.NumVertices, 0, Mesh.NumTriangles, 1);
+		RHICmdList.DrawIndexedPrimitive(Mesh.IndexBufferRHI, 0, 0, Mesh.NumVertices, 0, Mesh.NumTriangles, 1);
 	}
 
 	bool FWindowsMixedRealityHMD::HasVisibleAreaMesh() const
@@ -726,7 +726,7 @@ namespace WindowsMixedReality
 		check(Mesh.IsValid());
 
 		RHICmdList.SetStreamSource(0, Mesh.VertexBufferRHI, 0);
-		RHICmdList.DrawIndexedPrimitive(Mesh.IndexBufferRHI, PT_TriangleList, 0, 0, Mesh.NumVertices, 0, Mesh.NumTriangles, 1);
+		RHICmdList.DrawIndexedPrimitive(Mesh.IndexBufferRHI, 0, 0, Mesh.NumVertices, 0, Mesh.NumTriangles, 1);
 	}
 
 	void FWindowsMixedRealityHMD::SetupViewFamily(FSceneViewFamily& InViewFamily)
diff --git a/Engine/Source/Runtime/Apple/MetalRHI/Private/MetalCommands.cpp b/Engine/Source/Runtime/Apple/MetalRHI/Private/MetalCommands.cpp
--- a/Engine/Source/Runtime/Apple/MetalRHI/Private/MetalCommands.cpp
+++ b/Engine/Source/Runtime/Apple/MetalRHI/Private/MetalCommands.cpp
@@ -662,7 +662,7 @@ void FMetalRHICommandContext::RHISetRenderTargetsAndClear(const FRHISetRenderTar
 }
 
 
-void FMetalRHICommandContext::RHIDrawPrimitive(uint32 PrimitiveType_DEPRECATED, uint32 BaseVertexIndex, uint32 NumPrimitives, uint32 NumInstances)
+void FMetalRHICommandContext::RHIDrawPrimitive(uint32 BaseVertexIndex, uint32 NumPrimitives, uint32 NumInstances)
 {
 	@autoreleasepool {
 	SCOPE_CYCLE_COUNTER(STAT_MetalDrawCallTime);
@@ -684,7 +684,7 @@ void FMetalRHICommandContext::RHIDrawPrimitive(uint32 PrimitiveType_DEPRECATED,
 	}
 }
 
-void FMetalRHICommandContext::RHIDrawPrimitiveIndirect(uint32 PrimitiveType_DEPRECATED, FVertexBufferRHIParamRef VertexBufferRHI, uint32 ArgumentOffset)
+void FMetalRHICommandContext::RHIDrawPrimitiveIndirect(FVertexBufferRHIParamRef VertexBufferRHI, uint32 ArgumentOffset)
 {
 	@autoreleasepool {
 #if PLATFORM_IOS
@@ -702,7 +702,7 @@ void FMetalRHICommandContext::RHIDrawPrimitiveIndirect(uint32 PrimitiveType_DEPR
 	}
 }
 
-void FMetalRHICommandContext::RHIDrawIndexedPrimitive(FIndexBufferRHIParamRef IndexBufferRHI, uint32 PrimitiveType_DEPRECATED, int32 BaseVertexIndex, uint32 FirstInstance,
+void FMetalRHICommandContext::RHIDrawIndexedPrimitive(FIndexBufferRHIParamRef IndexBufferRHI, int32 BaseVertexIndex, uint32 FirstInstance,
 	uint32 NumVertices, uint32 StartIndex, uint32 NumPrimitives, uint32 NumInstances)
 {
 	@autoreleasepool {
@@ -721,7 +721,7 @@ void FMetalRHICommandContext::RHIDrawIndexedPrimitive(FIndexBufferRHIParamRef In
 	}
 }
 
-void FMetalRHICommandContext::RHIDrawIndexedIndirect(FIndexBufferRHIParamRef IndexBufferRHI, uint32 PrimitiveType_DEPRECATED, FStructuredBufferRHIParamRef VertexBufferRHI, int32 DrawArgumentsIndex, uint32 NumInstances)
+void FMetalRHICommandContext::RHIDrawIndexedIndirect(FIndexBufferRHIParamRef IndexBufferRHI, FStructuredBufferRHIParamRef VertexBufferRHI, int32 DrawArgumentsIndex, uint32 NumInstances)
 {
 	@autoreleasepool {
 	if (GetMetalDeviceContext().SupportsFeature(EMetalFeaturesIndirectBuffer))
@@ -770,7 +770,7 @@ void FMetalRHICommandContext::RHIDrawIndexedPrimitiveIndirect(uint32 PrimitiveTy
 }
 
 
-void FMetalRHICommandContext::RHIBeginDrawPrimitiveUP( uint32 PrimitiveType_DEPRECATED, uint32 NumPrimitives, uint32 NumVertices, uint32 VertexDataStride, void*& OutVertexData)
+void FMetalRHICommandContext::RHIBeginDrawPrimitiveUP( uint32 NumPrimitives, uint32 NumVertices, uint32 VertexDataStride, void*& OutVertexData)
 {
 	@autoreleasepool {
 	SCOPE_CYCLE_COUNTER(STAT_MetalDrawCallTime);
@@ -816,7 +816,7 @@ void FMetalRHICommandContext::RHIEndDrawPrimitiveUP()
 	}
 }
 
-void FMetalRHICommandContext::RHIBeginDrawIndexedPrimitiveUP( uint32 PrimitiveType_DEPRECATED, uint32 NumPrimitives, uint32 NumVertices, uint32 VertexDataStride, void*& OutVertexData, uint32 MinVertexIndex, uint32 NumIndices, uint32 IndexDataStride, void*& OutIndexData)
+void FMetalRHICommandContext::RHIBeginDrawIndexedPrimitiveUP( uint32 NumPrimitives, uint32 NumVertices, uint32 VertexDataStride, void*& OutVertexData, uint32 MinVertexIndex, uint32 NumIndices, uint32 IndexDataStride, void*& OutIndexData)
 {
 	@autoreleasepool {
 	SCOPE_CYCLE_COUNTER(STAT_MetalDrawCallTime);
diff --git a/Engine/Source/Runtime/Apple/MetalRHI/Public/MetalRHIContext.h b/Engine/Source/Runtime/Apple/MetalRHI/Public/MetalRHIContext.h
--- a/Engine/Source/Runtime/Apple/MetalRHI/Public/MetalRHIContext.h
+++ b/Engine/Source/Runtime/Apple/MetalRHI/Public/MetalRHIContext.h
@@ -242,16 +242,16 @@ public:
 	
 	virtual void RHISetRenderTargetsAndClear(const FRHISetRenderTargetsInfo& RenderTargetsInfo) final override;
 	
-	virtual void RHIDrawPrimitive(uint32 PrimitiveType, uint32 BaseVertexIndex, uint32 NumPrimitives, uint32 NumInstances) final override;
+	virtual void RHIDrawPrimitive(uint32 BaseVertexIndex, uint32 NumPrimitives, uint32 NumInstances) final override;
 	
-	virtual void RHIDrawPrimitiveIndirect(uint32 PrimitiveType, FVertexBufferRHIParamRef ArgumentBuffer, uint32 ArgumentOffset) final override;
+	virtual void RHIDrawPrimitiveIndirect(FVertexBufferRHIParamRef ArgumentBuffer, uint32 ArgumentOffset) final override;
 	
-	virtual void RHIDrawIndexedIndirect(FIndexBufferRHIParamRef IndexBufferRHI, uint32 PrimitiveType, FStructuredBufferRHIParamRef ArgumentsBufferRHI, int32 DrawArgumentsIndex, uint32 NumInstances) final override;
+	virtual void RHIDrawIndexedIndirect(FIndexBufferRHIParamRef IndexBufferRHI, FStructuredBufferRHIParamRef ArgumentsBufferRHI, int32 DrawArgumentsIndex, uint32 NumInstances) final override;
 	
 	// @param NumPrimitives need to be >0
-	virtual void RHIDrawIndexedPrimitive(FIndexBufferRHIParamRef IndexBuffer, uint32 PrimitiveType, int32 BaseVertexIndex, uint32 FirstInstance, uint32 NumVertices, uint32 StartIndex, uint32 NumPrimitives, uint32 NumInstances) final override;
+	virtual void RHIDrawIndexedPrimitive(FIndexBufferRHIParamRef IndexBuffer, int32 BaseVertexIndex, uint32 FirstInstance, uint32 NumVertices, uint32 StartIndex, uint32 NumPrimitives, uint32 NumInstances) final override;
 	
-	virtual void RHIDrawIndexedPrimitiveIndirect(uint32 PrimitiveType, FIndexBufferRHIParamRef IndexBuffer, FVertexBufferRHIParamRef ArgumentBuffer, uint32 ArgumentOffset) final override;
+	virtual void RHIDrawIndexedPrimitiveIndirect(FIndexBufferRHIParamRef IndexBuffer, FVertexBufferRHIParamRef ArgumentBuffer, uint32 ArgumentOffset) final override;
 	
 	/**
 	 * Preallocate memory or get a direct command stream pointer to fill up for immediate rendering . This avoids memcpys below in DrawPrimitiveUP
@@ -261,7 +261,7 @@ public:
 	 * @param VertexDataStride Size of each vertex
 	 * @param OutVertexData Reference to the allocated vertex memory
 	 */
-	virtual void RHIBeginDrawPrimitiveUP(uint32 PrimitiveType, uint32 NumPrimitives, uint32 NumVertices, uint32 VertexDataStride, void*& OutVertexData) final override;
+	virtual void RHIBeginDrawPrimitiveUP(uint32 NumPrimitives, uint32 NumVertices, uint32 VertexDataStride, void*& OutVertexData) final override;
 	
 	/**
 	 * Draw a primitive using the vertex data populated since RHIBeginDrawPrimitiveUP and clean up any memory as needed
@@ -280,7 +280,7 @@ public:
 	 * @param IndexDataStride Size of each index (either 2 or 4 bytes)
 	 * @param OutIndexData Reference to the allocated index memory
 	 */
-	virtual void RHIBeginDrawIndexedPrimitiveUP(uint32 PrimitiveType, uint32 NumPrimitives, uint32 NumVertices, uint32 VertexDataStride, void*& OutVertexData, uint32 MinVertexIndex, uint32 NumIndices, uint32 IndexDataStride, void*& OutIndexData) final override;
+	virtual void RHIBeginDrawIndexedPrimitiveUP(uint32 NumPrimitives, uint32 NumVertices, uint32 VertexDataStride, void*& OutVertexData, uint32 MinVertexIndex, uint32 NumIndices, uint32 IndexDataStride, void*& OutIndexData) final override;
 	
 	/**
 	 * Draw a primitive using the vertex and index data populated since RHIBeginDrawIndexedPrimitiveUP and clean up any memory as needed
diff --git a/Engine/Source/Runtime/D3D12RHI/Private/D3D12CommandContext.cpp b/Engine/Source/Runtime/D3D12RHI/Private/D3D12CommandContext.cpp
--- a/Engine/Source/Runtime/D3D12RHI/Private/D3D12CommandContext.cpp
+++ b/Engine/Source/Runtime/D3D12RHI/Private/D3D12CommandContext.cpp
@@ -656,11 +656,10 @@ FD3D12CommandContextRedirector::FD3D12CommandContextRedirector(class FD3D12Adapt
 	FMemory::Memzero(PhysicalContexts, sizeof(PhysicalContexts[0]) * MAX_NUM_GPUS);
 }
 
-void FD3D12CommandContextRedirector::RHIBeginDrawPrimitiveUP(uint32 PrimitiveType, uint32 NumPrimitives, uint32 NumVertices, uint32 VertexDataStride, void*& OutVertexData)
+void FD3D12CommandContextRedirector::RHIBeginDrawPrimitiveUP(uint32 NumPrimitives, uint32 NumVertices, uint32 VertexDataStride, void*& OutVertexData)
 {
 	check(!PendingUP);
 
-	PendingUP.PrimitiveType = PrimitiveType;
 	PendingUP.NumPrimitives = NumPrimitives;
 	PendingUP.NumVertices = NumVertices;
 	PendingUP.VertexDataStride = VertexDataStride;
@@ -677,7 +676,7 @@ void FD3D12CommandContextRedirector::RHIEndDrawPrimitiveUP()
 			check(GPUContext);
 
 			void* GPUVertexData = nullptr;
-			GPUContext->RHIBeginDrawPrimitiveUP(PendingUP.PrimitiveType, PendingUP.NumPrimitives, PendingUP.NumVertices, PendingUP.VertexDataStride, GPUVertexData);
+			GPUContext->RHIBeginDrawPrimitiveUP(PendingUP.NumPrimitives, PendingUP.NumVertices, PendingUP.VertexDataStride, GPUVertexData);
 			if (GPUVertexData)
 			{
 				FMemory::Memcpy(GPUVertexData, PendingUP.VertexData, PendingUP.NumVertices * PendingUP.VertexDataStride);
@@ -691,11 +690,10 @@ void FD3D12CommandContextRedirector::RHIEndDrawPrimitiveUP()
 	PendingUP.Reset();
 }
 	
-void FD3D12CommandContextRedirector::RHIBeginDrawIndexedPrimitiveUP(uint32 PrimitiveType, uint32 NumPrimitives, uint32 NumVertices, uint32 VertexDataStride, void*& OutVertexData, uint32 MinVertexIndex, uint32 NumIndices, uint32 IndexDataStride, void*& OutIndexData)
+void FD3D12CommandContextRedirector::RHIBeginDrawIndexedPrimitiveUP(uint32 NumPrimitives, uint32 NumVertices, uint32 VertexDataStride, void*& OutVertexData, uint32 MinVertexIndex, uint32 NumIndices, uint32 IndexDataStride, void*& OutIndexData)
 {
 	check(!PendingUP);
 
-	PendingUP.PrimitiveType = PrimitiveType;
 	PendingUP.NumPrimitives = NumPrimitives;
 	PendingUP.NumVertices = NumVertices;
 	PendingUP.VertexDataStride = VertexDataStride;
@@ -718,7 +716,7 @@ void FD3D12CommandContextRedirector::RHIEndDrawIndexedPrimitiveUP()
 
 			void* GPUVertexData = nullptr;
 			void* GPUIndexData = nullptr;
-			GPUContext->RHIBeginDrawIndexedPrimitiveUP(PendingUP.PrimitiveType, PendingUP.NumPrimitives, PendingUP.NumVertices, PendingUP.VertexDataStride, GPUVertexData, PendingUP.MinVertexIndex, PendingUP.NumIndices, PendingUP.IndexDataStride, GPUIndexData);
+			GPUContext->RHIBeginDrawIndexedPrimitiveUP(PendingUP.NumPrimitives, PendingUP.NumVertices, PendingUP.VertexDataStride, GPUVertexData, PendingUP.MinVertexIndex, PendingUP.NumIndices, PendingUP.IndexDataStride, GPUIndexData);
 			if (GPUVertexData)
 			{
 				FMemory::Memcpy(GPUVertexData, PendingUP.VertexData, PendingUP.NumVertices * PendingUP.VertexDataStride);
diff --git a/Engine/Source/Runtime/D3D12RHI/Private/D3D12CommandContext.h b/Engine/Source/Runtime/D3D12RHI/Private/D3D12CommandContext.h
--- a/Engine/Source/Runtime/D3D12RHI/Private/D3D12CommandContext.h
+++ b/Engine/Source/Runtime/D3D12RHI/Private/D3D12CommandContext.h
@@ -59,7 +59,6 @@ protected:
 		FUserPrimitiveData() { FMemory::Memzero(*this); }
 		~FUserPrimitiveData() { check(!VertexData && !IndexData); }
 
-		uint32 PrimitiveType;
 		uint32 NumPrimitives;
 		uint32 NumVertices;
 		uint32 VertexDataStride;
@@ -305,14 +304,14 @@ public:
 	virtual void RHISetRenderTargets(uint32 NumSimultaneousRenderTargets, const FRHIRenderTargetView* NewRenderTargets, const FRHIDepthRenderTargetView* NewDepthStencilTarget, uint32 NumUAVs, const FUnorderedAccessViewRHIParamRef* UAVs) final override;
 	virtual void RHISetRenderTargetsAndClear(const FRHISetRenderTargetsInfo& RenderTargetsInfo) final override;
 	virtual void RHIBindClearMRTValues(bool bClearColor, bool bClearDepth, bool bClearStencil) final override;
-	virtual void RHIDrawPrimitive(uint32 PrimitiveType, uint32 BaseVertexIndex, uint32 NumPrimitives, uint32 NumInstances) final override;
-	virtual void RHIDrawPrimitiveIndirect(uint32 PrimitiveType, FVertexBufferRHIParamRef ArgumentBuffer, uint32 ArgumentOffset) final override;
-	virtual void RHIDrawIndexedIndirect(FIndexBufferRHIParamRef IndexBufferRHI, uint32 PrimitiveType, FStructuredBufferRHIParamRef ArgumentsBufferRHI, int32 DrawArgumentsIndex, uint32 NumInstances) final override;
-	virtual void RHIDrawIndexedPrimitive(FIndexBufferRHIParamRef IndexBuffer, uint32 PrimitiveType, int32 BaseVertexIndex, uint32 FirstInstance, uint32 NumVertices, uint32 StartIndex, uint32 NumPrimitives, uint32 NumInstances) final override;
-	virtual void RHIDrawIndexedPrimitiveIndirect(uint32 PrimitiveType, FIndexBufferRHIParamRef IndexBuffer, FVertexBufferRHIParamRef ArgumentBuffer, uint32 ArgumentOffset) final override;
-	virtual void RHIBeginDrawPrimitiveUP(uint32 PrimitiveType, uint32 NumPrimitives, uint32 NumVertices, uint32 VertexDataStride, void*& OutVertexData) final override;
+	virtual void RHIDrawPrimitive(uint32 BaseVertexIndex, uint32 NumPrimitives, uint32 NumInstances) final override;
+	virtual void RHIDrawPrimitiveIndirect(FVertexBufferRHIParamRef ArgumentBuffer, uint32 ArgumentOffset) final override;
+	virtual void RHIDrawIndexedIndirect(FIndexBufferRHIParamRef IndexBufferRHI, FStructuredBufferRHIParamRef ArgumentsBufferRHI, int32 DrawArgumentsIndex, uint32 NumInstances) final override;
+	virtual void RHIDrawIndexedPrimitive(FIndexBufferRHIParamRef IndexBuffer, int32 BaseVertexIndex, uint32 FirstInstance, uint32 NumVertices, uint32 StartIndex, uint32 NumPrimitives, uint32 NumInstances) final override;
+	virtual void RHIDrawIndexedPrimitiveIndirect(FIndexBufferRHIParamRef IndexBuffer, FVertexBufferRHIParamRef ArgumentBuffer, uint32 ArgumentOffset) final override;
+	virtual void RHIBeginDrawPrimitiveUP(uint32 NumPrimitives, uint32 NumVertices, uint32 VertexDataStride, void*& OutVertexData) final override;
 	virtual void RHIEndDrawPrimitiveUP() final override;
-	virtual void RHIBeginDrawIndexedPrimitiveUP(uint32 PrimitiveType, uint32 NumPrimitives, uint32 NumVertices, uint32 VertexDataStride, void*& OutVertexData, uint32 MinVertexIndex, uint32 NumIndices, uint32 IndexDataStride, void*& OutIndexData) final override;
+	virtual void RHIBeginDrawIndexedPrimitiveUP(uint32 NumPrimitives, uint32 NumVertices, uint32 VertexDataStride, void*& OutVertexData, uint32 MinVertexIndex, uint32 NumIndices, uint32 IndexDataStride, void*& OutIndexData) final override;
 	virtual void RHIEndDrawIndexedPrimitiveUP() final override;
 	virtual void RHIEnableDepthBoundsTest(bool bEnable) final override;
 	virtual void RHISetDepthBounds(float MinDepth, float MaxDepth) final override;
@@ -681,32 +680,32 @@ public:
 	{
 		ContextRedirect(RHIBindClearMRTValues(bClearColor, bClearDepth, bClearStencil));
 	}
-	FORCEINLINE virtual void RHIDrawPrimitive(uint32 PrimitiveType, uint32 BaseVertexIndex, uint32 NumPrimitives, uint32 NumInstances) final override
+	FORCEINLINE virtual void RHIDrawPrimitive(uint32 BaseVertexIndex, uint32 NumPrimitives, uint32 NumInstances) final override
 	{
-		ContextRedirect(RHIDrawPrimitive(PrimitiveType, BaseVertexIndex, NumPrimitives, NumInstances));
+		ContextRedirect(RHIDrawPrimitive(BaseVertexIndex, NumPrimitives, NumInstances));
 	}
-	FORCEINLINE virtual void RHIDrawPrimitiveIndirect(uint32 PrimitiveType, FVertexBufferRHIParamRef ArgumentBuffer, uint32 ArgumentOffset) final override
+	FORCEINLINE virtual void RHIDrawPrimitiveIndirect(FVertexBufferRHIParamRef ArgumentBuffer, uint32 ArgumentOffset) final override
 	{
-		ContextRedirect(RHIDrawPrimitiveIndirect(PrimitiveType, ArgumentBuffer, ArgumentOffset));
+		ContextRedirect(RHIDrawPrimitiveIndirect(ArgumentBuffer, ArgumentOffset));
 	}
-	FORCEINLINE virtual void RHIDrawIndexedIndirect(FIndexBufferRHIParamRef IndexBufferRHI, uint32 PrimitiveType, FStructuredBufferRHIParamRef ArgumentsBufferRHI, int32 DrawArgumentsIndex, uint32 NumInstances) final override
+	FORCEINLINE virtual void RHIDrawIndexedIndirect(FIndexBufferRHIParamRef IndexBufferRHI, FStructuredBufferRHIParamRef ArgumentsBufferRHI, int32 DrawArgumentsIndex, uint32 NumInstances) final override
 	{
-		ContextRedirect(RHIDrawIndexedIndirect(IndexBufferRHI, PrimitiveType, ArgumentsBufferRHI, DrawArgumentsIndex, NumInstances));
+		ContextRedirect(RHIDrawIndexedIndirect(IndexBufferRHI, ArgumentsBufferRHI, DrawArgumentsIndex, NumInstances));
 	}
-	FORCEINLINE virtual void RHIDrawIndexedPrimitive(FIndexBufferRHIParamRef IndexBuffer, uint32 PrimitiveType, int32 BaseVertexIndex, uint32 FirstInstance, uint32 NumVertices, uint32 StartIndex, uint32 NumPrimitives, uint32 NumInstances) final override
+	FORCEINLINE virtual void RHIDrawIndexedPrimitive(FIndexBufferRHIParamRef IndexBuffer, int32 BaseVertexIndex, uint32 FirstInstance, uint32 NumVertices, uint32 StartIndex, uint32 NumPrimitives, uint32 NumInstances) final override
 	{
-		ContextRedirect(RHIDrawIndexedPrimitive(IndexBuffer, PrimitiveType, BaseVertexIndex, FirstInstance, NumVertices, StartIndex, NumPrimitives, NumInstances));
+		ContextRedirect(RHIDrawIndexedPrimitive(IndexBuffer, BaseVertexIndex, FirstInstance, NumVertices, StartIndex, NumPrimitives, NumInstances));
 	}
-	FORCEINLINE virtual void RHIDrawIndexedPrimitiveIndirect(uint32 PrimitiveType, FIndexBufferRHIParamRef IndexBuffer, FVertexBufferRHIParamRef ArgumentBuffer, uint32 ArgumentOffset) final override
+	FORCEINLINE virtual void RHIDrawIndexedPrimitiveIndirect(FIndexBufferRHIParamRef IndexBuffer, FVertexBufferRHIParamRef ArgumentBuffer, uint32 ArgumentOffset) final override
 	{
-		ContextRedirect(RHIDrawIndexedPrimitiveIndirect(PrimitiveType, IndexBuffer, ArgumentBuffer, ArgumentOffset));
+		ContextRedirect(RHIDrawIndexedPrimitiveIndirect(IndexBuffer, ArgumentBuffer, ArgumentOffset));
 	}
 	
-	virtual void RHIBeginDrawPrimitiveUP(uint32 PrimitiveType, uint32 NumPrimitives, uint32 NumVertices, uint32 VertexDataStride, void*& OutVertexData) final override;
+	virtual void RHIBeginDrawPrimitiveUP(uint32 NumPrimitives, uint32 NumVertices, uint32 VertexDataStride, void*& OutVertexData) final override;
 	
 	virtual void RHIEndDrawPrimitiveUP() final override;
 	
-	virtual void RHIBeginDrawIndexedPrimitiveUP(uint32 PrimitiveType, uint32 NumPrimitives, uint32 NumVertices, uint32 VertexDataStride, void*& OutVertexData, uint32 MinVertexIndex, uint32 NumIndices, uint32 IndexDataStride, void*& OutIndexData) final override;
+	virtual void RHIBeginDrawIndexedPrimitiveUP(uint32 NumPrimitives, uint32 NumVertices, uint32 VertexDataStride, void*& OutVertexData, uint32 MinVertexIndex, uint32 NumIndices, uint32 IndexDataStride, void*& OutIndexData) final override;
 	
 	virtual void RHIEndDrawIndexedPrimitiveUP() final override;
 
diff --git a/Engine/Source/Runtime/D3D12RHI/Private/D3D12Commands.cpp b/Engine/Source/Runtime/D3D12RHI/Private/D3D12Commands.cpp
--- a/Engine/Source/Runtime/D3D12RHI/Private/D3D12Commands.cpp
+++ b/Engine/Source/Runtime/D3D12RHI/Private/D3D12Commands.cpp
@@ -436,7 +436,7 @@ void FD3D12CommandContext::RHISetBoundShaderState(FBoundShaderStateRHIParamRef B
 	unimplemented();
 }
 
-D3D_PRIMITIVE_TOPOLOGY GetD3D12PrimitiveType(uint32 PrimitiveType, bool bUsingTessellation)
+static inline D3D_PRIMITIVE_TOPOLOGY GetD3D12PrimitiveType(EPrimitiveType PrimitiveType, bool bUsingTessellation)
 {
 	if (bUsingTessellation)
 	{
@@ -1396,21 +1396,21 @@ void FD3D12CommandContext::CommitComputeResourceTables(FD3D12ComputeShader* InCo
 	SetResourcesFromTables(ComputeShader);
 }
 
-void FD3D12CommandContext::RHIDrawPrimitive(uint32 PrimitiveType, uint32 BaseVertexIndex, uint32 NumPrimitives, uint32 NumInstances)
+void FD3D12CommandContext::RHIDrawPrimitive(uint32 BaseVertexIndex, uint32 NumPrimitives, uint32 NumInstances)
 {
-	RHI_DRAW_CALL_STATS(PrimitiveType, NumInstances*NumPrimitives);
+	RHI_DRAW_CALL_STATS(StateCache.GetGraphicsPipelinePrimitiveType(), NumInstances*NumPrimitives);
 
 	CommitGraphicsResourceTables();
 	CommitNonComputeShaderConstants();
 
-	uint32 VertexCount = GetVertexCountForPrimitiveCount(NumPrimitives, PrimitiveType);
+	uint32 VertexCount = GetVertexCountForPrimitiveCount(NumPrimitives, StateCache.GetGraphicsPipelinePrimitiveType());
 
 	if (IsDefaultContext())
 	{
 		GetParentDevice()->RegisterGPUWork(NumPrimitives * NumInstances, VertexCount * NumInstances);
 	}
 
-	StateCache.SetPrimitiveTopology(GetD3D12PrimitiveType(PrimitiveType, bUsingTessellation));
+	StateCache.SetPrimitiveTopology(GetD3D12PrimitiveType(StateCache.GetGraphicsPipelinePrimitiveType(), bUsingTessellation));
 
 	StateCache.ApplyState();
 	numDraws++;
@@ -1422,7 +1422,7 @@ void FD3D12CommandContext::RHIDrawPrimitive(uint32 PrimitiveType, uint32 BaseVer
 
 }
 
-void FD3D12CommandContext::RHIDrawPrimitiveIndirect(uint32 PrimitiveType, FVertexBufferRHIParamRef ArgumentBufferRHI, uint32 ArgumentOffset)
+void FD3D12CommandContext::RHIDrawPrimitiveIndirect(FVertexBufferRHIParamRef ArgumentBufferRHI, uint32 ArgumentOffset)
 {
 	FD3D12VertexBuffer* ArgumentBuffer = RetrieveObject<FD3D12VertexBuffer>(ArgumentBufferRHI);
 
@@ -1436,7 +1436,7 @@ void FD3D12CommandContext::RHIDrawPrimitiveIndirect(uint32 PrimitiveType, FVerte
 	CommitGraphicsResourceTables();
 	CommitNonComputeShaderConstants();
 
-	StateCache.SetPrimitiveTopology(GetD3D12PrimitiveType(PrimitiveType, bUsingTessellation));
+	StateCache.SetPrimitiveTopology(GetD3D12PrimitiveType(StateCache.GetGraphicsPipelinePrimitiveType(), bUsingTessellation));
 
 	FD3D12ResourceLocation& Location = ArgumentBuffer->ResourceLocation;
 	FD3D12DynamicRHI::TransitionResource(CommandListHandle, Location.GetResource(), D3D12_RESOURCE_STATE_INDIRECT_ARGUMENT, D3D12_RESOURCE_BARRIER_ALL_SUBRESOURCES);
@@ -1462,7 +1462,7 @@ void FD3D12CommandContext::RHIDrawPrimitiveIndirect(uint32 PrimitiveType, FVerte
 
 }
 
-void FD3D12CommandContext::RHIDrawIndexedIndirect(FIndexBufferRHIParamRef IndexBufferRHI, uint32 PrimitiveType, FStructuredBufferRHIParamRef ArgumentsBufferRHI, int32 DrawArgumentsIndex, uint32 NumInstances)
+void FD3D12CommandContext::RHIDrawIndexedIndirect(FIndexBufferRHIParamRef IndexBufferRHI, FStructuredBufferRHIParamRef ArgumentsBufferRHI, int32 DrawArgumentsIndex, uint32 NumInstances)
 {
 	FD3D12IndexBuffer* IndexBuffer = RetrieveObject<FD3D12IndexBuffer>(IndexBufferRHI);
 	FD3D12StructuredBuffer* ArgumentsBuffer = RetrieveObject<FD3D12StructuredBuffer>(ArgumentsBufferRHI);
@@ -1482,7 +1482,7 @@ void FD3D12CommandContext::RHIDrawIndexedIndirect(FIndexBufferRHIParamRef IndexB
 	const DXGI_FORMAT Format = (IndexBuffer->GetStride() == sizeof(uint16) ? DXGI_FORMAT_R16_UINT : DXGI_FORMAT_R32_UINT);
 
 	StateCache.SetIndexBuffer(&IndexBuffer->ResourceLocation, Format, 0);
-	StateCache.SetPrimitiveTopology(GetD3D12PrimitiveType(PrimitiveType, bUsingTessellation));
+	StateCache.SetPrimitiveTopology(GetD3D12PrimitiveType(StateCache.GetGraphicsPipelinePrimitiveType(), bUsingTessellation));
 
 	FD3D12ResourceLocation& Location = ArgumentsBuffer->ResourceLocation;
 	FD3D12DynamicRHI::TransitionResource(CommandListHandle, Location.GetResource(), D3D12_RESOURCE_STATE_INDIRECT_ARGUMENT, D3D12_RESOURCE_BARRIER_ALL_SUBRESOURCES);
@@ -1507,14 +1507,14 @@ void FD3D12CommandContext::RHIDrawIndexedIndirect(FIndexBufferRHIParamRef IndexB
 	DEBUG_EXECUTE_COMMAND_LIST(this);
 }
 
-void FD3D12CommandContext::RHIDrawIndexedPrimitive(FIndexBufferRHIParamRef IndexBufferRHI, uint32 PrimitiveType, int32 BaseVertexIndex, uint32 FirstInstance, uint32 NumVertices, uint32 StartIndex, uint32 NumPrimitives, uint32 NumInstances)
+void FD3D12CommandContext::RHIDrawIndexedPrimitive(FIndexBufferRHIParamRef IndexBufferRHI, int32 BaseVertexIndex, uint32 FirstInstance, uint32 NumVertices, uint32 StartIndex, uint32 NumPrimitives, uint32 NumInstances)
 {
 	FD3D12IndexBuffer* IndexBuffer = RetrieveObject<FD3D12IndexBuffer>(IndexBufferRHI);
 
 	// called should make sure the input is valid, this avoid hidden bugs
 	ensure(NumPrimitives > 0);
 
-	RHI_DRAW_CALL_STATS(PrimitiveType, NumInstances*NumPrimitives);
+	RHI_DRAW_CALL_STATS(StateCache.GetGraphicsPipelinePrimitiveType(), NumInstances*NumPrimitives);
 
 	if (IsDefaultContext())
 	{
@@ -1528,15 +1528,15 @@ void FD3D12CommandContext::RHIDrawIndexedPrimitive(FIndexBufferRHIParamRef Index
 	uint32 SizeFormat = sizeof(DXGI_FORMAT);
 	const DXGI_FORMAT Format = (IndexBuffer->GetStride() == sizeof(uint16) ? DXGI_FORMAT_R16_UINT : DXGI_FORMAT_R32_UINT);
 
-	uint32 IndexCount = GetVertexCountForPrimitiveCount(NumPrimitives, PrimitiveType);
+	uint32 IndexCount = GetVertexCountForPrimitiveCount(NumPrimitives, StateCache.GetGraphicsPipelinePrimitiveType());
 
 	// Verify that we are not trying to read outside the index buffer range
 	// test is an optimized version of: StartIndex + IndexCount <= IndexBuffer->GetSize() / IndexBuffer->GetStride() 
 	checkf((StartIndex + IndexCount) * IndexBuffer->GetStride() <= IndexBuffer->GetSize(),
-		TEXT("Start %u, Count %u, Type %u, Buffer Size %u, Buffer stride %u"), StartIndex, IndexCount, PrimitiveType, IndexBuffer->GetSize(), IndexBuffer->GetStride());
+		TEXT("Start %u, Count %u, Type %u, Buffer Size %u, Buffer stride %u"), StartIndex, IndexCount, StateCache.GetGraphicsPipelinePrimitiveType(), IndexBuffer->GetSize(), IndexBuffer->GetStride());
 
 	StateCache.SetIndexBuffer(&IndexBuffer->ResourceLocation, Format, 0);
-	StateCache.SetPrimitiveTopology(GetD3D12PrimitiveType(PrimitiveType, bUsingTessellation));
+	StateCache.SetPrimitiveTopology(GetD3D12PrimitiveType(StateCache.GetGraphicsPipelinePrimitiveType(), bUsingTessellation));
 	StateCache.ApplyState();
 
 	numDraws++;
@@ -1548,7 +1548,7 @@ void FD3D12CommandContext::RHIDrawIndexedPrimitive(FIndexBufferRHIParamRef Index
 
 }
 
-void FD3D12CommandContext::RHIDrawIndexedPrimitiveIndirect(uint32 PrimitiveType, FIndexBufferRHIParamRef IndexBufferRHI, FVertexBufferRHIParamRef ArgumentBufferRHI, uint32 ArgumentOffset)
+void FD3D12CommandContext::RHIDrawIndexedPrimitiveIndirect(FIndexBufferRHIParamRef IndexBufferRHI, FVertexBufferRHIParamRef ArgumentBufferRHI, uint32 ArgumentOffset)
 {
 	FD3D12IndexBuffer* IndexBuffer = RetrieveObject<FD3D12IndexBuffer>(IndexBufferRHI);
 	FD3D12VertexBuffer* ArgumentBuffer = RetrieveObject<FD3D12VertexBuffer>(ArgumentBufferRHI);
@@ -1567,7 +1567,7 @@ void FD3D12CommandContext::RHIDrawIndexedPrimitiveIndirect(uint32 PrimitiveType,
 	const uint32 SizeFormat = sizeof(DXGI_FORMAT);
 	const DXGI_FORMAT Format = (IndexBuffer->GetStride() == sizeof(uint16) ? DXGI_FORMAT_R16_UINT : DXGI_FORMAT_R32_UINT);
 	StateCache.SetIndexBuffer(&IndexBuffer->ResourceLocation, Format, 0);
-	StateCache.SetPrimitiveTopology(GetD3D12PrimitiveType(PrimitiveType, bUsingTessellation));
+	StateCache.SetPrimitiveTopology(GetD3D12PrimitiveType(StateCache.GetGraphicsPipelinePrimitiveType(), bUsingTessellation));
 
 	FD3D12ResourceLocation& Location = ArgumentBuffer->ResourceLocation;
 	FD3D12DynamicRHI::TransitionResource(CommandListHandle, Location.GetResource(), D3D12_RESOURCE_STATE_INDIRECT_ARGUMENT, D3D12_RESOURCE_BARRIER_ALL_SUBRESOURCES);
@@ -1594,18 +1594,16 @@ void FD3D12CommandContext::RHIDrawIndexedPrimitiveIndirect(uint32 PrimitiveType,
 
 /**
 * Preallocate memory or get a direct command stream pointer to fill up for immediate rendering . This avoids memcpys below in DrawPrimitiveUP
-* @param PrimitiveType The type (triangles, lineloop, etc) of primitive to draw
 * @param NumPrimitives The number of primitives in the VertexData buffer
 * @param NumVertices The number of vertices to be written
 * @param VertexDataStride Size of each vertex
 * @param OutVertexData Reference to the allocated vertex memory
 */
-void FD3D12CommandContext::RHIBeginDrawPrimitiveUP(uint32 PrimitiveType, uint32 NumPrimitives, uint32 NumVertices, uint32 VertexDataStride, void*& OutVertexData)
+void FD3D12CommandContext::RHIBeginDrawPrimitiveUP(uint32 NumPrimitives, uint32 NumVertices, uint32 VertexDataStride, void*& OutVertexData)
 {
 	checkSlow(PendingUP.NumVertices == 0);
 
 	// Remember the parameters for the draw call.
-	PendingUP.PrimitiveType = PrimitiveType;
 	PendingUP.NumPrimitives = NumPrimitives;
 	PendingUP.NumVertices = NumVertices;
 	PendingUP.VertexDataStride = VertexDataStride;
@@ -1619,9 +1617,9 @@ void FD3D12CommandContext::RHIBeginDrawPrimitiveUP(uint32 PrimitiveType, uint32
 */
 void FD3D12CommandContext::RHIEndDrawPrimitiveUP()
 {
-	RHI_DRAW_CALL_STATS(PendingUP.PrimitiveType, PendingUP.NumPrimitives);
+	RHI_DRAW_CALL_STATS(StateCache.GetGraphicsPipelinePrimitiveType(), PendingUP.NumPrimitives);
 
-	checkSlow(!bUsingTessellation || PendingUP.PrimitiveType == PT_TriangleList);
+	checkSlow(!bUsingTessellation || StateCache.GetGraphicsPipelinePrimitiveType() == PT_TriangleList);
 
 	if (IsDefaultContext())
 	{
@@ -1636,7 +1634,7 @@ void FD3D12CommandContext::RHIEndDrawPrimitiveUP()
 	CommitGraphicsResourceTables();
 	CommitNonComputeShaderConstants();
 	StateCache.SetStreamSource(BufferLocation, 0, PendingUP.VertexDataStride, VBOffset);
-	StateCache.SetPrimitiveTopology(GetD3D12PrimitiveType(PendingUP.PrimitiveType, bUsingTessellation));
+	StateCache.SetPrimitiveTopology(GetD3D12PrimitiveType(StateCache.GetGraphicsPipelinePrimitiveType(), bUsingTessellation));
 	StateCache.ApplyState();
 	numDraws++;
 	CommandListHandle->DrawInstanced(PendingUP.NumVertices, 1, 0, 0);
@@ -1651,7 +1649,6 @@ void FD3D12CommandContext::RHIEndDrawPrimitiveUP()
 
 /**
 * Preallocate memory or get a direct command stream pointer to fill up for immediate rendering . This avoids memcpys below in DrawIndexedPrimitiveUP
-* @param PrimitiveType The type (triangles, lineloop, etc) of primitive to draw
 * @param NumPrimitives The number of primitives in the VertexData buffer
 * @param NumVertices The number of vertices to be written
 * @param VertexDataStride Size of each vertex
@@ -1661,12 +1658,11 @@ void FD3D12CommandContext::RHIEndDrawPrimitiveUP()
 * @param IndexDataStride Size of each index (either 2 or 4 bytes)
 * @param OutIndexData Reference to the allocated index memory
 */
-void FD3D12CommandContext::RHIBeginDrawIndexedPrimitiveUP(uint32 PrimitiveType, uint32 NumPrimitives, uint32 NumVertices, uint32 VertexDataStride, void*& OutVertexData, uint32 MinVertexIndex, uint32 NumIndices, uint32 IndexDataStride, void*& OutIndexData)
+void FD3D12CommandContext::RHIBeginDrawIndexedPrimitiveUP(uint32 NumPrimitives, uint32 NumVertices, uint32 VertexDataStride, void*& OutVertexData, uint32 MinVertexIndex, uint32 NumIndices, uint32 IndexDataStride, void*& OutIndexData)
 {
 	checkSlow((sizeof(uint16) == IndexDataStride) || (sizeof(uint32) == IndexDataStride));
 
 	// Store off information needed for the draw call.
-	PendingUP.PrimitiveType = PrimitiveType;
 	PendingUP.NumPrimitives = NumPrimitives;
 	PendingUP.MinVertexIndex = MinVertexIndex;
 	PendingUP.IndexDataStride = IndexDataStride;
@@ -1685,9 +1681,9 @@ void FD3D12CommandContext::RHIBeginDrawIndexedPrimitiveUP(uint32 PrimitiveType,
 void FD3D12CommandContext::RHIEndDrawIndexedPrimitiveUP()
 {
 	// tessellation only supports trilists
-	checkSlow(!bUsingTessellation || PendingUP.PrimitiveType == PT_TriangleList);
+	checkSlow(!bUsingTessellation || StateCache.GetGraphicsPipelinePrimitiveType() == PT_TriangleList);
 
-	RHI_DRAW_CALL_STATS(PendingUP.PrimitiveType, PendingUP.NumPrimitives);
+	RHI_DRAW_CALL_STATS(StateCache.GetGraphicsPipelinePrimitiveType(), PendingUP.NumPrimitives);
 
 	if (IsDefaultContext())
 	{
@@ -1704,7 +1700,7 @@ void FD3D12CommandContext::RHIEndDrawIndexedPrimitiveUP()
 	CommitNonComputeShaderConstants();
 	StateCache.SetStreamSource(VertexBufferLocation, 0, PendingUP.VertexDataStride, VBOffset);
 	StateCache.SetIndexBuffer(IndexBufferLocation, PendingUP.IndexDataStride == sizeof(uint16) ? DXGI_FORMAT_R16_UINT : DXGI_FORMAT_R32_UINT, 0);
-	StateCache.SetPrimitiveTopology(GetD3D12PrimitiveType(PendingUP.PrimitiveType, bUsingTessellation));
+	StateCache.SetPrimitiveTopology(GetD3D12PrimitiveType(StateCache.GetGraphicsPipelinePrimitiveType(), bUsingTessellation));
 	StateCache.ApplyState();
 
 	numDraws++;
diff --git a/Engine/Source/Runtime/D3D12RHI/Private/D3D12RenderTarget.cpp b/Engine/Source/Runtime/D3D12RHI/Private/D3D12RenderTarget.cpp
--- a/Engine/Source/Runtime/D3D12RHI/Private/D3D12RenderTarget.cpp
+++ b/Engine/Source/Runtime/D3D12RHI/Private/D3D12RenderTarget.cpp
@@ -209,7 +209,7 @@ void FD3D12CommandContext::ResolveTextureUsingShader(
 
 	RHIUnlockVertexBuffer(VertexBufferRHI);
 	RHICmdList.SetStreamSource(0, VertexBufferRHI, 0);
-	RHICmdList.DrawPrimitive(PT_TriangleStrip, 0, 2, 1);
+	RHICmdList.DrawPrimitive(0, 2, 1);
 
 	RHICmdList.Flush(); // always call flush when using a command list in RHI implementations before doing anything else. This is super hazardous.
 
diff --git a/Engine/Source/Runtime/D3D12RHI/Private/D3D12StateCachePrivate.h b/Engine/Source/Runtime/D3D12RHI/Private/D3D12StateCachePrivate.h
--- a/Engine/Source/Runtime/D3D12RHI/Private/D3D12StateCachePrivate.h
+++ b/Engine/Source/Runtime/D3D12RHI/Private/D3D12StateCachePrivate.h
@@ -320,6 +320,8 @@ protected:
 
 			float MinDepth;
 			float MaxDepth;
+
+			EPrimitiveType PrimitiveType = PT_Num;
 		} Graphics;
 
 		struct
@@ -443,6 +445,11 @@ public:
 		return PipelineState.Graphics.CurrentPipelineStateObject ? PipelineState.Graphics.CurrentPipelineStateObject->RootSignature : nullptr;
 	}
 
+	inline EPrimitiveType GetGraphicsPipelinePrimitiveType() const
+	{
+		return PipelineState.Graphics.PrimitiveType;
+	}
+
 	const FD3D12RootSignature* GetComputeRootSignature()
 	{
 		return PipelineState.Compute.CurrentPipelineStateObject ? PipelineState.Compute.CurrentPipelineStateObject->ComputeShader->pRootSignature : nullptr;
@@ -677,6 +684,7 @@ public:
 			// Save the PSO
 			PipelineState.Common.bNeedSetPSO = true;
 			PipelineState.Graphics.CurrentPipelineStateObject = GraphicsPipelineState;
+			PipelineState.Graphics.PrimitiveType = GraphicsPipelineState->PipelineStateInitializer.PrimitiveType;
 
 			// Set the PSO
 			InternalSetPipelineState<false>();
diff --git a/Engine/Source/Runtime/Engine/Private/BatchedElements.cpp b/Engine/Source/Runtime/Engine/Private/BatchedElements.cpp
--- a/Engine/Source/Runtime/Engine/Private/BatchedElements.cpp
+++ b/Engine/Source/Runtime/Engine/Private/BatchedElements.cpp
@@ -816,7 +816,7 @@ void FBatchedElements::DrawPointElements(FRHICommandList& RHICmdList, const FMat
 
 		RHIUnlockVertexBuffer(VertexBufferRHI);
 		RHICmdList.SetStreamSource(0, VertexBufferRHI, 0);
-		RHICmdList.DrawPrimitive(PT_TriangleList, 0, NumTris, 1);
+		RHICmdList.DrawPrimitive(0, NumTris, 1);
 	}
 }
 
@@ -913,7 +913,7 @@ bool FBatchedElements::Draw(FRHICommandList& RHICmdList, const FDrawingPolicyRen
 				while( MinVertex < TotalVerts )
 				{
 					int32 NumLinePrims = FMath::Min( MaxVerticesAllowed, TotalVerts - MinVertex ) / 2;
-					RHICmdList.DrawPrimitive(PT_LineList, MinVertex, NumLinePrims, 1);
+					RHICmdList.DrawPrimitive(MinVertex, NumLinePrims, 1);
 					MinVertex += NumLinePrims * 2;
 				}
 
@@ -1040,7 +1040,7 @@ bool FBatchedElements::Draw(FRHICommandList& RHICmdList, const FDrawingPolicyRen
 
 					RHIUnlockVertexBuffer(VertexBufferRHI);
 					RHICmdList.SetStreamSource(0, VertexBufferRHI, 0);
-					RHICmdList.DrawPrimitive(PT_TriangleList, 0, 8 * NumLinesThisBatch, 1);
+					RHICmdList.DrawPrimitive(0, 8 * NumLinesThisBatch, 1);
 				}
 
 				GraphicsPSOInit.RasterizerState = TStaticRasterizerState<FM_Solid, CM_None>::GetRHI();
@@ -1091,7 +1091,7 @@ bool FBatchedElements::Draw(FRHICommandList& RHICmdList, const FDrawingPolicyRen
 					RHICmdList.SetStencilRef(StencilRef);
 
 					int32 NumTris = MaxTri - MinTri;
-					RHICmdList.DrawPrimitive(PT_TriangleList, MinTri * 3, NumTris, 1);
+					RHICmdList.DrawPrimitive(MinTri * 3, NumTris, 1);
 					MinTri = MaxTri;
 				}
 				VertexBufferRHI.SafeRelease();
@@ -1186,7 +1186,7 @@ bool FBatchedElements::Draw(FRHICommandList& RHICmdList, const FDrawingPolicyRen
 					RHICmdList.SetStencilRef(StencilRef);
 					RHICmdList.SetStreamSource(0, VertexBufferRHI, 0);
 
-					RHICmdList.DrawPrimitive(PT_TriangleList, BaseVertex, PrimCount, 1);
+					RHICmdList.DrawPrimitive(BaseVertex, PrimCount, 1);
 				}
 
 				if (!isVeryLastDummyIndex)
@@ -1233,7 +1233,7 @@ bool FBatchedElements::Draw(FRHICommandList& RHICmdList, const FDrawingPolicyRen
 
 					// Draw the mesh.
 					RHICmdList.SetStreamSource(0, VertexBufferRHI, MeshElement.MinVertex * sizeof(FSimpleElementVertex));
-					RHICmdList.DrawIndexedPrimitive(IndexBufferRHI, PT_TriangleList, 0, 0, MeshElement.MaxVertex - MeshElement.MinVertex + 1, 0, MeshElement.Indices.Num() / 3, 1);
+					RHICmdList.DrawIndexedPrimitive(IndexBufferRHI, 0, 0, MeshElement.MaxVertex - MeshElement.MinVertex + 1, 0, MeshElement.Indices.Num() / 3, 1);
 
 					IndexBufferRHI.SafeRelease();
 				}
diff --git a/Engine/Source/Runtime/Engine/Private/Particles/ParticleCurveTexture.cpp b/Engine/Source/Runtime/Engine/Private/Particles/ParticleCurveTexture.cpp
--- a/Engine/Source/Runtime/Engine/Private/Particles/ParticleCurveTexture.cpp
+++ b/Engine/Source/Runtime/Engine/Private/Particles/ParticleCurveTexture.cpp
@@ -278,7 +278,6 @@ static void InjectCurves(
 		// Inject particles.
 		RHICmdList.DrawIndexedPrimitive(
 			GParticleIndexBuffer.IndexBufferRHI,
-			PT_TriangleList,
 			/*BaseVertexIndex=*/ 0,
 			/*MinIndex=*/ 0,
 			/*NumVertices=*/ 4,
diff --git a/Engine/Source/Runtime/Engine/Private/Particles/ParticleGpuSimulation.cpp b/Engine/Source/Runtime/Engine/Private/Particles/ParticleGpuSimulation.cpp
--- a/Engine/Source/Runtime/Engine/Private/Particles/ParticleGpuSimulation.cpp
+++ b/Engine/Source/Runtime/Engine/Private/Particles/ParticleGpuSimulation.cpp
@@ -1379,7 +1379,6 @@ static void DrawAlignedParticleTiles(FRHICommandList& RHICmdList, int32 TileCoun
 	// Draw tiles.
 	RHICmdList.DrawIndexedPrimitive(
 		GParticleIndexBuffer.IndexBufferRHI,
-		PT_TriangleList,
 		/*BaseVertexIndex=*/ 0,
 		/*MinIndex=*/ 0,
 		/*NumVertices=*/ 4,
@@ -1831,7 +1830,6 @@ void InjectNewParticles(FRHICommandList& RHICmdList, FGraphicsPipelineStateIniti
 		// Inject particles.
 		RHICmdList.DrawIndexedPrimitive(
 			GParticleIndexBuffer.IndexBufferRHI,
-			PT_TriangleList,
 			/*BaseVertexIndex=*/ 0,
 			/*MinIndex=*/ 0,
 			/*NumVertices=*/ 4,
@@ -2052,7 +2050,6 @@ static void VisualizeGPUSimulation(
 	// Draw.
 	RHICmdList.DrawIndexedPrimitive(
 		GParticleIndexBuffer.IndexBufferRHI,
-		PT_TriangleList,
 		/*BaseVertexIndex=*/ 0,
 		/*MinIndex=*/ 0,
 		/*NumVertices=*/ 4,
diff --git a/Engine/Source/Runtime/Engine/Private/VolumeRendering.cpp b/Engine/Source/Runtime/Engine/Private/VolumeRendering.cpp
--- a/Engine/Source/Runtime/Engine/Private/VolumeRendering.cpp
+++ b/Engine/Source/Runtime/Engine/Private/VolumeRendering.cpp
@@ -21,5 +21,5 @@ ENGINE_API void RasterizeToVolumeTexture(FRHICommandList& RHICmdList, FVolumeBou
 	RHICmdList.SetStreamSource(0, GVolumeRasterizeVertexBuffer.VertexBufferRHI, 0);
 	const int32 NumInstances = VolumeBounds.MaxZ - VolumeBounds.MinZ;
 	// Render a quad per slice affected by the given bounds
-	RHICmdList.DrawPrimitive(PT_TriangleStrip, 0, 2, NumInstances);
+	RHICmdList.DrawPrimitive(0, 2, NumInstances);
 }
diff --git a/Engine/Source/Runtime/MediaAssets/Private/Misc/MediaTextureResource.cpp b/Engine/Source/Runtime/MediaAssets/Private/Misc/MediaTextureResource.cpp
--- a/Engine/Source/Runtime/MediaAssets/Private/Misc/MediaTextureResource.cpp
+++ b/Engine/Source/Runtime/MediaAssets/Private/Misc/MediaTextureResource.cpp
@@ -622,7 +622,7 @@ void FMediaTextureResource::ConvertSample(const TSharedPtr<IMediaTextureSample,
 		// set viewport to RT size
 		CommandList.SetViewport(0, 0, 0.0f, OutputDim.X, OutputDim.Y, 1.0f);
 
-		CommandList.DrawPrimitive(PT_TriangleStrip, 0, 2, 1);
+		CommandList.DrawPrimitive(0, 2, 1);
 		CommandList.TransitionResource(EResourceTransitionAccess::EReadable, RenderTargetTextureRHI);
 	}
 
diff --git a/Engine/Source/Runtime/NullDrv/Public/NullRHI.h b/Engine/Source/Runtime/NullDrv/Public/NullRHI.h
--- a/Engine/Source/Runtime/NullDrv/Public/NullRHI.h
+++ b/Engine/Source/Runtime/NullDrv/Public/NullRHI.h
@@ -723,31 +723,31 @@ public:
 
 	}
 
-	virtual void RHIDrawPrimitive(uint32 PrimitiveType, uint32 BaseVertexIndex, uint32 NumPrimitives, uint32 NumInstances) final override
+	virtual void RHIDrawPrimitive(uint32 BaseVertexIndex, uint32 NumPrimitives, uint32 NumInstances) final override
 	{
 
 	}
-	virtual void RHIDrawPrimitiveIndirect(uint32 PrimitiveType, FVertexBufferRHIParamRef ArgumentBuffer, uint32 ArgumentOffset) final override
+	virtual void RHIDrawPrimitiveIndirect(FVertexBufferRHIParamRef ArgumentBuffer, uint32 ArgumentOffset) final override
 	{
 
 	}
 
-	virtual void RHIDrawIndexedIndirect(FIndexBufferRHIParamRef IndexBufferRHI, uint32 PrimitiveType, FStructuredBufferRHIParamRef ArgumentsBufferRHI, int32 DrawArgumentsIndex, uint32 NumInstances) final override
+	virtual void RHIDrawIndexedIndirect(FIndexBufferRHIParamRef IndexBufferRHI, FStructuredBufferRHIParamRef ArgumentsBufferRHI, int32 DrawArgumentsIndex, uint32 NumInstances) final override
 	{
 
 	}
 
 
-	virtual void RHIDrawIndexedPrimitive(FIndexBufferRHIParamRef IndexBuffer, uint32 PrimitiveType, int32 BaseVertexIndex, uint32 FirstInstance, uint32 NumVertices, uint32 StartIndex, uint32 NumPrimitives, uint32 NumInstances) final override
+	virtual void RHIDrawIndexedPrimitive(FIndexBufferRHIParamRef IndexBuffer, int32 BaseVertexIndex, uint32 FirstInstance, uint32 NumVertices, uint32 StartIndex, uint32 NumPrimitives, uint32 NumInstances) final override
 	{
 
 	}
-	virtual void RHIDrawIndexedPrimitiveIndirect(uint32 PrimitiveType, FIndexBufferRHIParamRef IndexBuffer, FVertexBufferRHIParamRef ArgumentBuffer, uint32 ArgumentOffset) final override
+	virtual void RHIDrawIndexedPrimitiveIndirect(FIndexBufferRHIParamRef IndexBuffer, FVertexBufferRHIParamRef ArgumentBuffer, uint32 ArgumentOffset) final override
 	{
 
 	}
 
-	virtual void RHIBeginDrawPrimitiveUP(uint32 PrimitiveType, uint32 NumPrimitives, uint32 NumVertices, uint32 VertexDataStride, void*& OutVertexData) final override
+	virtual void RHIBeginDrawPrimitiveUP(uint32 NumPrimitives, uint32 NumVertices, uint32 VertexDataStride, void*& OutVertexData) final override
 	{ 
 		OutVertexData = GetStaticBuffer(); 
 	}
@@ -755,7 +755,7 @@ public:
 	{
 
 	}
-	virtual void RHIBeginDrawIndexedPrimitiveUP(uint32 PrimitiveType, uint32 NumPrimitives, uint32 NumVertices, uint32 VertexDataStride, void*& OutVertexData, uint32 MinVertexIndex, uint32 NumIndices, uint32 IndexDataStride, void*& OutIndexData) final override
+	virtual void RHIBeginDrawIndexedPrimitiveUP(uint32 NumPrimitives, uint32 NumVertices, uint32 VertexDataStride, void*& OutVertexData, uint32 MinVertexIndex, uint32 NumIndices, uint32 IndexDataStride, void*& OutIndexData) final override
 	{ 
 		OutVertexData = GetStaticBuffer(); 
 		OutIndexData = GetStaticBuffer(); 
diff --git a/Engine/Source/Runtime/OpenGLDrv/Private/OpenGLCommands.cpp b/Engine/Source/Runtime/OpenGLDrv/Private/OpenGLCommands.cpp
--- a/Engine/Source/Runtime/OpenGLDrv/Private/OpenGLCommands.cpp
+++ b/Engine/Source/Runtime/OpenGLDrv/Private/OpenGLCommands.cpp
@@ -2628,7 +2628,7 @@ static void VerifyProgramPipeline()
 }
 #endif
 
-void FOpenGLDynamicRHI::RHIDrawPrimitive(uint32 PrimitiveType,uint32 BaseVertexIndex,uint32 NumPrimitives,uint32 NumInstances)
+void FOpenGLDynamicRHI::RHIDrawPrimitive(uint32 BaseVertexIndex,uint32 NumPrimitives,uint32 NumInstances)
 {
 	SCOPE_CYCLE_COUNTER_DETAILED(STAT_OpenGLDrawPrimitiveTime);
 	VERIFY_GL_SCOPE();
@@ -2681,7 +2681,7 @@ void FOpenGLDynamicRHI::RHIDrawPrimitive(uint32 PrimitiveType,uint32 BaseVertexI
 	}
 }
 
-void FOpenGLDynamicRHI::RHIDrawPrimitiveIndirect(uint32 PrimitiveType,FVertexBufferRHIParamRef ArgumentBufferRHI,uint32 ArgumentOffset)
+void FOpenGLDynamicRHI::RHIDrawPrimitiveIndirect(FVertexBufferRHIParamRef ArgumentBufferRHI,uint32 ArgumentOffset)
 {
 	if (FOpenGL::SupportsDrawIndirect())
 	{
@@ -2732,7 +2732,7 @@ void FOpenGLDynamicRHI::RHIDrawPrimitiveIndirect(uint32 PrimitiveType,FVertexBuf
 
 }
 
-void FOpenGLDynamicRHI::RHIDrawIndexedIndirect(FIndexBufferRHIParamRef IndexBufferRHI, uint32 PrimitiveType, FStructuredBufferRHIParamRef ArgumentsBufferRHI, int32 DrawArgumentsIndex, uint32 NumInstances)
+void FOpenGLDynamicRHI::RHIDrawIndexedIndirect(FIndexBufferRHIParamRef IndexBufferRHI,  FStructuredBufferRHIParamRef ArgumentsBufferRHI, int32 DrawArgumentsIndex, uint32 NumInstances)
 {
 	if (FOpenGL::SupportsDrawIndirect())
 	{
@@ -2791,7 +2791,7 @@ void FOpenGLDynamicRHI::RHIDrawIndexedIndirect(FIndexBufferRHIParamRef IndexBuff
 	}
 }
 
-void FOpenGLDynamicRHI::RHIDrawIndexedPrimitive(FIndexBufferRHIParamRef IndexBufferRHI,uint32 PrimitiveType,int32 BaseVertexIndex,uint32 FirstInstance,uint32 NumVertices,uint32 StartIndex,uint32 NumPrimitives,uint32 NumInstances)
+void FOpenGLDynamicRHI::RHIDrawIndexedPrimitive(FIndexBufferRHIParamRef IndexBufferRHI,int32 BaseVertexIndex,uint32 FirstInstance,uint32 NumVertices,uint32 StartIndex,uint32 NumPrimitives,uint32 NumInstances)
 {
 	SCOPE_CYCLE_COUNTER_DETAILED(STAT_OpenGLDrawPrimitiveTime);
 	VERIFY_GL_SCOPE();
@@ -2893,7 +2893,7 @@ void FOpenGLDynamicRHI::RHIDrawIndexedPrimitive(FIndexBufferRHIParamRef IndexBuf
 	}
 }
 
-void FOpenGLDynamicRHI::RHIDrawIndexedPrimitiveIndirect(uint32 PrimitiveType,FIndexBufferRHIParamRef IndexBufferRHI,FVertexBufferRHIParamRef ArgumentBufferRHI,uint32 ArgumentOffset)
+void FOpenGLDynamicRHI::RHIDrawIndexedPrimitiveIndirect(FIndexBufferRHIParamRef IndexBufferRHI,FVertexBufferRHIParamRef ArgumentBufferRHI,uint32 ArgumentOffset)
 {
 	if (FOpenGL::SupportsDrawIndirect())
 	{
@@ -2957,7 +2957,7 @@ void FOpenGLDynamicRHI::RHIDrawIndexedPrimitiveIndirect(uint32 PrimitiveType,FIn
  * @param VertexDataStride Size of each vertex 
  * @param OutVertexData Reference to the allocated vertex memory
  */
-void FOpenGLDynamicRHI::RHIBeginDrawPrimitiveUP( uint32 PrimitiveType, uint32 NumPrimitives, uint32 NumVertices, uint32 VertexDataStride, void*& OutVertexData)
+void FOpenGLDynamicRHI::RHIBeginDrawPrimitiveUP( uint32 NumPrimitives, uint32 NumVertices, uint32 VertexDataStride, void*& OutVertexData)
 {
 	SCOPE_CYCLE_COUNTER_DETAILED(STAT_OpenGLDrawPrimitiveUPTime);
 	VERIFY_GL_SCOPE();
@@ -3075,7 +3075,7 @@ void FOpenGLDynamicRHI::RHIEndDrawPrimitiveUP()
  * @param IndexDataStride Size of each index (either 2 or 4 bytes)
  * @param OutIndexData Reference to the allocated index memory
  */
-void FOpenGLDynamicRHI::RHIBeginDrawIndexedPrimitiveUP( uint32 PrimitiveType, uint32 NumPrimitives, uint32 NumVertices, uint32 VertexDataStride, void*& OutVertexData, uint32 MinVertexIndex, uint32 NumIndices, uint32 IndexDataStride, void*& OutIndexData)
+void FOpenGLDynamicRHI::RHIBeginDrawIndexedPrimitiveUP(uint32 NumPrimitives, uint32 NumVertices, uint32 VertexDataStride, void*& OutVertexData, uint32 MinVertexIndex, uint32 NumIndices, uint32 IndexDataStride, void*& OutIndexData)
 {
 	VERIFY_GL_SCOPE();
 	SCOPE_CYCLE_COUNTER_DETAILED(STAT_OpenGLDrawPrimitiveUPTime);
diff --git a/Engine/Source/Runtime/OpenGLDrv/Public/OpenGLDrv.h b/Engine/Source/Runtime/OpenGLDrv/Public/OpenGLDrv.h
--- a/Engine/Source/Runtime/OpenGLDrv/Public/OpenGLDrv.h
+++ b/Engine/Source/Runtime/OpenGLDrv/Public/OpenGLDrv.h
@@ -456,6 +456,13 @@ public:
 	virtual void RHISetRasterizerState(FRasterizerStateRHIParamRef NewState) final override;
 	virtual void RHISetViewport(uint32 MinX, uint32 MinY, float MinZ, uint32 MaxX, uint32 MaxY, float MaxZ) final override;
 	virtual void RHISetScissorRect(bool bEnable, uint32 MinX, uint32 MinY, uint32 MaxX, uint32 MaxY) final override;
+	virtual void RHISetGraphicsPipelineState(FGraphicsPipelineStateRHIParamRef GraphicsState) final override
+	{
+		FRHIGraphicsPipelineStateFallBack* FallbackGraphicsState = static_cast<FRHIGraphicsPipelineStateFallBack*>(GraphicsState);
+		IRHICommandContext::RHISetGraphicsPipelineState(GraphicsState);
+		// Store the PSO's primitive (after since IRHICommandContext::RHISetGraphicsPipelineState sets the BSS)
+		PrimitiveType = FallbackGraphicsState->Initializer.PrimitiveType;
+	}
 	virtual void RHISetBoundShaderState(FBoundShaderStateRHIParamRef BoundShaderState) final override;
 	virtual void RHISetShaderTexture(FVertexShaderRHIParamRef VertexShader, uint32 TextureIndex, FTextureRHIParamRef NewTexture) final override;
 	virtual void RHISetShaderTexture(FHullShaderRHIParamRef HullShader, uint32 TextureIndex, FTextureRHIParamRef NewTexture) final override;
@@ -498,14 +505,14 @@ public:
 	}
 	virtual void RHISetRenderTargets(uint32 NumSimultaneousRenderTargets, const FRHIRenderTargetView* NewRenderTargets, const FRHIDepthRenderTargetView* NewDepthStencilTarget, uint32 NumUAVs, const FUnorderedAccessViewRHIParamRef* UAVs) final override;
 	virtual void RHISetRenderTargetsAndClear(const FRHISetRenderTargetsInfo& RenderTargetsInfo) final override;
-	virtual void RHIDrawPrimitive(uint32 PrimitiveType, uint32 BaseVertexIndex, uint32 NumPrimitives, uint32 NumInstances) final override;
-	virtual void RHIDrawPrimitiveIndirect(uint32 PrimitiveType, FVertexBufferRHIParamRef ArgumentBuffer, uint32 ArgumentOffset) final override;
-	virtual void RHIDrawIndexedIndirect(FIndexBufferRHIParamRef IndexBufferRHI, uint32 PrimitiveType, FStructuredBufferRHIParamRef ArgumentsBufferRHI, int32 DrawArgumentsIndex, uint32 NumInstances) final override;
-	virtual void RHIDrawIndexedPrimitive(FIndexBufferRHIParamRef IndexBuffer, uint32 PrimitiveType, int32 BaseVertexIndex, uint32 FirstInstance, uint32 NumVertices, uint32 StartIndex, uint32 NumPrimitives, uint32 NumInstances) final override;
-	virtual void RHIDrawIndexedPrimitiveIndirect(uint32 PrimitiveType, FIndexBufferRHIParamRef IndexBuffer, FVertexBufferRHIParamRef ArgumentBuffer, uint32 ArgumentOffset) final override;
-	virtual void RHIBeginDrawPrimitiveUP(uint32 PrimitiveType, uint32 NumPrimitives, uint32 NumVertices, uint32 VertexDataStride, void*& OutVertexData) final override;
+	virtual void RHIDrawPrimitive(uint32 BaseVertexIndex, uint32 NumPrimitives, uint32 NumInstances) final override;
+	virtual void RHIDrawPrimitiveIndirect(FVertexBufferRHIParamRef ArgumentBuffer, uint32 ArgumentOffset) final override;
+	virtual void RHIDrawIndexedIndirect(FIndexBufferRHIParamRef IndexBufferRHI, FStructuredBufferRHIParamRef ArgumentsBufferRHI, int32 DrawArgumentsIndex, uint32 NumInstances) final override;
+	virtual void RHIDrawIndexedPrimitive(FIndexBufferRHIParamRef IndexBuffer, int32 BaseVertexIndex, uint32 FirstInstance, uint32 NumVertices, uint32 StartIndex, uint32 NumPrimitives, uint32 NumInstances) final override;
+	virtual void RHIDrawIndexedPrimitiveIndirect(FIndexBufferRHIParamRef IndexBuffer, FVertexBufferRHIParamRef ArgumentBuffer, uint32 ArgumentOffset) final override;
+	virtual void RHIBeginDrawPrimitiveUP(uint32 NumPrimitives, uint32 NumVertices, uint32 VertexDataStride, void*& OutVertexData) final override;
 	virtual void RHIEndDrawPrimitiveUP() final override;
-	virtual void RHIBeginDrawIndexedPrimitiveUP(uint32 PrimitiveType, uint32 NumPrimitives, uint32 NumVertices, uint32 VertexDataStride, void*& OutVertexData, uint32 MinVertexIndex, uint32 NumIndices, uint32 IndexDataStride, void*& OutIndexData) final override;
+	virtual void RHIBeginDrawIndexedPrimitiveUP(uint32 NumPrimitives, uint32 NumVertices, uint32 VertexDataStride, void*& OutVertexData, uint32 MinVertexIndex, uint32 NumIndices, uint32 IndexDataStride, void*& OutIndexData) final override;
 	virtual void RHIEndDrawIndexedPrimitiveUP() final override;
 	virtual void RHIEnableDepthBoundsTest(bool bEnable) final override;
 	virtual void RHISetDepthBounds(float MinDepth, float MaxDepth) final override;
@@ -1016,6 +1023,8 @@ private:
 
 	bool								bIsRenderingContextAcquired;
 
+	EPrimitiveType						PrimitiveType = PT_Num;
+
 	/** A history of the most recently used bound shader states, used to keep transient bound shader states from being recreated for each use. */
 	TGlobalResource< TBoundShaderStateHistory<10000> > BoundShaderStateHistory;
 
diff --git a/Engine/Source/Runtime/RHI/Public/RHICommandList.h b/Engine/Source/Runtime/RHI/Public/RHICommandList.h
--- a/Engine/Source/Runtime/RHI/Public/RHICommandList.h
+++ b/Engine/Source/Runtime/RHI/Public/RHICommandList.h
@@ -482,8 +482,6 @@ public:
 
 	struct FDrawUpData
 	{
-
-		uint32 PrimitiveType;
 		uint32 NumPrimitives;
 		uint32 NumVertices;
 		uint32 VertexDataStride;
@@ -494,8 +492,7 @@ public:
 		void* OutIndexData;
 
 		FDrawUpData()
-			: PrimitiveType(PT_Num)
-			, OutVertexData(nullptr)
+			: OutVertexData(nullptr)
 			, OutIndexData(nullptr)
 		{
 		}
@@ -820,13 +817,11 @@ struct FRHICommandSetShaderSampler final : public FRHICommand<FRHICommandSetShad
 
 struct FRHICommandDrawPrimitive final : public FRHICommand<FRHICommandDrawPrimitive>
 {
-	uint32 PrimitiveType;
 	uint32 BaseVertexIndex;
 	uint32 NumPrimitives;
 	uint32 NumInstances;
-	FORCEINLINE_DEBUGGABLE FRHICommandDrawPrimitive(uint32 InPrimitiveType, uint32 InBaseVertexIndex, uint32 InNumPrimitives, uint32 InNumInstances)
-		: PrimitiveType(InPrimitiveType)
-		, BaseVertexIndex(InBaseVertexIndex)
+	FORCEINLINE_DEBUGGABLE FRHICommandDrawPrimitive(uint32 InBaseVertexIndex, uint32 InNumPrimitives, uint32 InNumInstances)
+		: BaseVertexIndex(InBaseVertexIndex)
 		, NumPrimitives(InNumPrimitives)
 		, NumInstances(InNumInstances)
 	{
@@ -837,16 +832,14 @@ struct FRHICommandDrawPrimitive final : public FRHICommand<FRHICommandDrawPrimit
 struct FRHICommandDrawIndexedPrimitive final : public FRHICommand<FRHICommandDrawIndexedPrimitive>
 {
 	FIndexBufferRHIParamRef IndexBuffer;
-	uint32 PrimitiveType;
 	int32 BaseVertexIndex;
 	uint32 FirstInstance;
 	uint32 NumVertices;
 	uint32 StartIndex;
 	uint32 NumPrimitives;
 	uint32 NumInstances;
-	FORCEINLINE_DEBUGGABLE FRHICommandDrawIndexedPrimitive(FIndexBufferRHIParamRef InIndexBuffer, uint32 InPrimitiveType, int32 InBaseVertexIndex, uint32 InFirstInstance, uint32 InNumVertices, uint32 InStartIndex, uint32 InNumPrimitives, uint32 InNumInstances)
+	FORCEINLINE_DEBUGGABLE FRHICommandDrawIndexedPrimitive(FIndexBufferRHIParamRef InIndexBuffer, int32 InBaseVertexIndex, uint32 InFirstInstance, uint32 InNumVertices, uint32 InStartIndex, uint32 InNumPrimitives, uint32 InNumInstances)
 		: IndexBuffer(InIndexBuffer)
-		, PrimitiveType(InPrimitiveType)
 		, BaseVertexIndex(InBaseVertexIndex)
 		, FirstInstance(InFirstInstance)
 		, NumVertices(InNumVertices)
@@ -1150,15 +1143,13 @@ struct FRHICommandBindClearMRTValues final : public FRHICommand<FRHICommandBindC
 
 struct FRHICommandEndDrawPrimitiveUP final : public FRHICommand<FRHICommandEndDrawPrimitiveUP>
 {
-	uint32 PrimitiveType;
 	uint32 NumPrimitives;
 	uint32 NumVertices;
 	uint32 VertexDataStride;
 	void* OutVertexData;
 
-	FORCEINLINE_DEBUGGABLE FRHICommandEndDrawPrimitiveUP(uint32 InPrimitiveType, uint32 InNumPrimitives, uint32 InNumVertices, uint32 InVertexDataStride, void* InOutVertexData)
-		: PrimitiveType(InPrimitiveType)
-		, NumPrimitives(InNumPrimitives)
+	FORCEINLINE_DEBUGGABLE FRHICommandEndDrawPrimitiveUP(uint32 InNumPrimitives, uint32 InNumVertices, uint32 InVertexDataStride, void* InOutVertexData)
+		: NumPrimitives(InNumPrimitives)
 		, NumVertices(InNumVertices)
 		, VertexDataStride(InVertexDataStride)
 		, OutVertexData(InOutVertexData)
@@ -1170,7 +1161,6 @@ struct FRHICommandEndDrawPrimitiveUP final : public FRHICommand<FRHICommandEndDr
 
 struct FRHICommandEndDrawIndexedPrimitiveUP final : public FRHICommand<FRHICommandEndDrawIndexedPrimitiveUP>
 {
-	uint32 PrimitiveType;
 	uint32 NumPrimitives;
 	uint32 NumVertices;
 	uint32 VertexDataStride;
@@ -1180,9 +1170,8 @@ struct FRHICommandEndDrawIndexedPrimitiveUP final : public FRHICommand<FRHIComma
 	uint32 IndexDataStride;
 	void* OutIndexData;
 
-	FORCEINLINE_DEBUGGABLE FRHICommandEndDrawIndexedPrimitiveUP(uint32 InPrimitiveType, uint32 InNumPrimitives, uint32 InNumVertices, uint32 InVertexDataStride, void* InOutVertexData, uint32 InMinVertexIndex, uint32 InNumIndices, uint32 InIndexDataStride, void* InOutIndexData)
-		: PrimitiveType(InPrimitiveType)
-		, NumPrimitives(InNumPrimitives)
+	FORCEINLINE_DEBUGGABLE FRHICommandEndDrawIndexedPrimitiveUP(uint32 InNumPrimitives, uint32 InNumVertices, uint32 InVertexDataStride, void* InOutVertexData, uint32 InMinVertexIndex, uint32 InNumIndices, uint32 InIndexDataStride, void* InOutIndexData)
+		: NumPrimitives(InNumPrimitives)
 		, NumVertices(InNumVertices)
 		, VertexDataStride(InVertexDataStride)
 		, OutVertexData(InOutVertexData)
@@ -1273,11 +1262,9 @@ struct FRHICommandFlushComputeShaderCache final : public FRHICommand<FRHICommand
 struct FRHICommandDrawPrimitiveIndirect final : public FRHICommand<FRHICommandDrawPrimitiveIndirect>
 {
 	FVertexBufferRHIParamRef ArgumentBuffer;
-	uint32 PrimitiveType;
 	uint32 ArgumentOffset;
-	FORCEINLINE_DEBUGGABLE FRHICommandDrawPrimitiveIndirect(uint32 InPrimitiveType, FVertexBufferRHIParamRef InArgumentBuffer, uint32 InArgumentOffset)
+	FORCEINLINE_DEBUGGABLE FRHICommandDrawPrimitiveIndirect(FVertexBufferRHIParamRef InArgumentBuffer, uint32 InArgumentOffset)
 		: ArgumentBuffer(InArgumentBuffer)
-		, PrimitiveType(InPrimitiveType)
 		, ArgumentOffset(InArgumentOffset)
 	{
 	}
@@ -1287,14 +1274,12 @@ struct FRHICommandDrawPrimitiveIndirect final : public FRHICommand<FRHICommandDr
 struct FRHICommandDrawIndexedIndirect final : public FRHICommand<FRHICommandDrawIndexedIndirect>
 {
 	FIndexBufferRHIParamRef IndexBufferRHI;
-	uint32 PrimitiveType;
 	FStructuredBufferRHIParamRef ArgumentsBufferRHI;
 	uint32 DrawArgumentsIndex;
 	uint32 NumInstances;
 
-	FORCEINLINE_DEBUGGABLE FRHICommandDrawIndexedIndirect(FIndexBufferRHIParamRef InIndexBufferRHI, uint32 InPrimitiveType, FStructuredBufferRHIParamRef InArgumentsBufferRHI, uint32 InDrawArgumentsIndex, uint32 InNumInstances)
+	FORCEINLINE_DEBUGGABLE FRHICommandDrawIndexedIndirect(FIndexBufferRHIParamRef InIndexBufferRHI, FStructuredBufferRHIParamRef InArgumentsBufferRHI, uint32 InDrawArgumentsIndex, uint32 InNumInstances)
 		: IndexBufferRHI(InIndexBufferRHI)
-		, PrimitiveType(InPrimitiveType)
 		, ArgumentsBufferRHI(InArgumentsBufferRHI)
 		, DrawArgumentsIndex(InDrawArgumentsIndex)
 		, NumInstances(InNumInstances)
@@ -1307,13 +1292,11 @@ struct FRHICommandDrawIndexedPrimitiveIndirect final : public FRHICommand<FRHICo
 {
 	FIndexBufferRHIParamRef IndexBuffer;
 	FVertexBufferRHIParamRef ArgumentsBuffer;
-	uint32 PrimitiveType;
 	uint32 ArgumentOffset;
 
-	FORCEINLINE_DEBUGGABLE FRHICommandDrawIndexedPrimitiveIndirect(uint32 InPrimitiveType, FIndexBufferRHIParamRef InIndexBuffer, FVertexBufferRHIParamRef InArgumentsBuffer, uint32 InArgumentOffset)
+	FORCEINLINE_DEBUGGABLE FRHICommandDrawIndexedPrimitiveIndirect(FIndexBufferRHIParamRef InIndexBuffer, FVertexBufferRHIParamRef InArgumentsBuffer, uint32 InArgumentOffset)
 		: IndexBuffer(InIndexBuffer)
 		, ArgumentsBuffer(InArgumentsBuffer)
-		, PrimitiveType(InPrimitiveType)
 		, ArgumentOffset(InArgumentOffset)
 	{
 	}
@@ -2130,18 +2113,24 @@ public:
 		new (AllocCommand<FRHICommandSetBlendFactor>()) FRHICommandSetBlendFactor(BlendFactor);
 	}
 
-	FORCEINLINE_DEBUGGABLE void DrawPrimitive(uint32 PrimitiveType, uint32 BaseVertexIndex, uint32 NumPrimitives, uint32 NumInstances)
+	FORCEINLINE_DEBUGGABLE void DrawPrimitive(uint32 BaseVertexIndex, uint32 NumPrimitives, uint32 NumInstances)
 	{
 		//check(IsOutsideRenderPass());
 		if (Bypass())
 		{
-			CMD_CONTEXT(RHIDrawPrimitive)(PrimitiveType, BaseVertexIndex, NumPrimitives, NumInstances);
+			CMD_CONTEXT(RHIDrawPrimitive)(BaseVertexIndex, NumPrimitives, NumInstances);
 			return;
 		}
-		new (AllocCommand<FRHICommandDrawPrimitive>()) FRHICommandDrawPrimitive(PrimitiveType, BaseVertexIndex, NumPrimitives, NumInstances);
+		new (AllocCommand<FRHICommandDrawPrimitive>()) FRHICommandDrawPrimitive(BaseVertexIndex, NumPrimitives, NumInstances);
 	}
 
-	FORCEINLINE_DEBUGGABLE void DrawIndexedPrimitive(FIndexBufferRHIParamRef IndexBuffer, uint32 PrimitiveType, int32 BaseVertexIndex, uint32 FirstInstance, uint32 NumVertices, uint32 StartIndex, uint32 NumPrimitives, uint32 NumInstances)
+	DEPRECATED(4.22, "PrimitiveType is now a part of the Graphics Pipeline State.")
+	FORCEINLINE_DEBUGGABLE void DrawPrimitive(uint32 PrimitiveType, uint32 BaseVertexIndex, uint32 NumPrimitives, uint32 NumInstances)
+	{
+		DrawPrimitive(BaseVertexIndex, NumPrimitives, NumInstances);
+	}
+
+	FORCEINLINE_DEBUGGABLE void DrawIndexedPrimitive(FIndexBufferRHIParamRef IndexBuffer, int32 BaseVertexIndex, uint32 FirstInstance, uint32 NumVertices, uint32 StartIndex, uint32 NumPrimitives, uint32 NumInstances)
 	{
 		if (!IndexBuffer)
 		{
@@ -2151,10 +2140,16 @@ public:
 		//check(IsOutsideRenderPass());
 		if (Bypass())
 		{
-			CMD_CONTEXT(RHIDrawIndexedPrimitive)(IndexBuffer, PrimitiveType, BaseVertexIndex, FirstInstance, NumVertices, StartIndex, NumPrimitives, NumInstances);
+			CMD_CONTEXT(RHIDrawIndexedPrimitive)(IndexBuffer, BaseVertexIndex, FirstInstance, NumVertices, StartIndex, NumPrimitives, NumInstances);
 			return;
 		}
-		new (AllocCommand<FRHICommandDrawIndexedPrimitive>()) FRHICommandDrawIndexedPrimitive(IndexBuffer, PrimitiveType, BaseVertexIndex, FirstInstance, NumVertices, StartIndex, NumPrimitives, NumInstances);
+		new (AllocCommand<FRHICommandDrawIndexedPrimitive>()) FRHICommandDrawIndexedPrimitive(IndexBuffer, BaseVertexIndex, FirstInstance, NumVertices, StartIndex, NumPrimitives, NumInstances);
+	}
+
+	DEPRECATED(4.22, "PrimitiveType is now a part of the Graphics Pipeline State.")
+	FORCEINLINE_DEBUGGABLE void DrawIndexedPrimitive(FIndexBufferRHIParamRef IndexBuffer, uint32 PrimitiveType, int32 BaseVertexIndex, uint32 FirstInstance, uint32 NumVertices, uint32 StartIndex, uint32 NumPrimitives, uint32 NumInstances)
+	{
+		DrawIndexedPrimitive(IndexBuffer, BaseVertexIndex, FirstInstance, NumVertices, StartIndex, NumPrimitives, NumInstances);
 	}
 
 	FORCEINLINE_DEBUGGABLE void SetStreamSource(uint32 StreamIndex, FVertexBufferRHIParamRef VertexBuffer, uint32 Offset)
@@ -2325,7 +2320,7 @@ public:
 		//check(IsOutsideRenderPass());
 		if (Bypass())
 		{
-			CMD_CONTEXT(RHIBeginDrawPrimitiveUP)(PrimitiveType, NumPrimitives, NumVertices, VertexDataStride, OutVertexData);
+			CMD_CONTEXT(RHIBeginDrawPrimitiveUP)(NumPrimitives, NumVertices, VertexDataStride, OutVertexData);
 			return;
 		}
 		checkf(!DrawUPData.OutVertexData && NumVertices * VertexDataStride > 0,
@@ -2333,7 +2328,6 @@ public:
 
 		OutVertexData = Alloc(NumVertices * VertexDataStride, 16);
 
-		DrawUPData.PrimitiveType = PrimitiveType;
 		DrawUPData.NumPrimitives = NumPrimitives;
 		DrawUPData.NumVertices = NumVertices;
 		DrawUPData.VertexDataStride = VertexDataStride;
@@ -2350,7 +2344,7 @@ public:
 			return;
 		}
 		check(DrawUPData.OutVertexData && DrawUPData.NumVertices);
-		new (AllocCommand<FRHICommandEndDrawPrimitiveUP>()) FRHICommandEndDrawPrimitiveUP(DrawUPData.PrimitiveType, DrawUPData.NumPrimitives, DrawUPData.NumVertices, DrawUPData.VertexDataStride, DrawUPData.OutVertexData);
+		new (AllocCommand<FRHICommandEndDrawPrimitiveUP>()) FRHICommandEndDrawPrimitiveUP(DrawUPData.NumPrimitives, DrawUPData.NumVertices, DrawUPData.VertexDataStride, DrawUPData.OutVertexData);
 
 		DrawUPData.OutVertexData = nullptr;
 		DrawUPData.NumVertices = 0;
@@ -2362,7 +2356,7 @@ public:
 		//check(IsOutsideRenderPass());
 		if (Bypass())
 		{
-			CMD_CONTEXT(RHIBeginDrawIndexedPrimitiveUP)(PrimitiveType, NumPrimitives, NumVertices, VertexDataStride, OutVertexData, MinVertexIndex, NumIndices, IndexDataStride, OutIndexData);
+			CMD_CONTEXT(RHIBeginDrawIndexedPrimitiveUP)(NumPrimitives, NumVertices, VertexDataStride, OutVertexData, MinVertexIndex, NumIndices, IndexDataStride, OutIndexData);
 			return;
 		}
 		check(!DrawUPData.OutVertexData && !DrawUPData.OutIndexData && NumVertices * VertexDataStride > 0 && NumIndices * IndexDataStride > 0);
@@ -2370,7 +2364,6 @@ public:
 		OutVertexData = Alloc(NumVertices * VertexDataStride, 16);
 		OutIndexData = Alloc(NumIndices * IndexDataStride, 16);
 
-		DrawUPData.PrimitiveType = PrimitiveType;
 		DrawUPData.NumPrimitives = NumPrimitives;
 		DrawUPData.NumVertices = NumVertices;
 		DrawUPData.VertexDataStride = VertexDataStride;
@@ -2393,7 +2386,7 @@ public:
 		}
 		check(DrawUPData.OutVertexData && DrawUPData.OutIndexData && DrawUPData.NumIndices && DrawUPData.NumVertices);
 
-		new (AllocCommand<FRHICommandEndDrawIndexedPrimitiveUP>()) FRHICommandEndDrawIndexedPrimitiveUP(DrawUPData.PrimitiveType, DrawUPData.NumPrimitives, DrawUPData.NumVertices, DrawUPData.VertexDataStride, DrawUPData.OutVertexData, DrawUPData.MinVertexIndex, DrawUPData.NumIndices, DrawUPData.IndexDataStride, DrawUPData.OutIndexData);
+		new (AllocCommand<FRHICommandEndDrawIndexedPrimitiveUP>()) FRHICommandEndDrawIndexedPrimitiveUP(DrawUPData.NumPrimitives, DrawUPData.NumVertices, DrawUPData.VertexDataStride, DrawUPData.OutVertexData, DrawUPData.MinVertexIndex, DrawUPData.NumIndices, DrawUPData.IndexDataStride, DrawUPData.OutIndexData);
 
 		DrawUPData.OutVertexData = nullptr;
 		DrawUPData.OutIndexData = nullptr;
@@ -2477,37 +2470,56 @@ public:
 		new (AllocCommand<FRHICommandFlushComputeShaderCache>()) FRHICommandFlushComputeShaderCache();
 	}
 
-	FORCEINLINE_DEBUGGABLE void DrawPrimitiveIndirect(uint32 PrimitiveType, FVertexBufferRHIParamRef ArgumentBuffer, uint32 ArgumentOffset)
+	FORCEINLINE_DEBUGGABLE void DrawPrimitiveIndirect(FVertexBufferRHIParamRef ArgumentBuffer, uint32 ArgumentOffset)
 	{
 		//check(IsOutsideRenderPass());
 		if (Bypass())
 		{
-			CMD_CONTEXT(RHIDrawPrimitiveIndirect)(PrimitiveType, ArgumentBuffer, ArgumentOffset);
+			CMD_CONTEXT(RHIDrawPrimitiveIndirect)(ArgumentBuffer, ArgumentOffset);
 			return;
 		}
-		new (AllocCommand<FRHICommandDrawPrimitiveIndirect>()) FRHICommandDrawPrimitiveIndirect(PrimitiveType, ArgumentBuffer, ArgumentOffset);
+		new (AllocCommand<FRHICommandDrawPrimitiveIndirect>()) FRHICommandDrawPrimitiveIndirect(ArgumentBuffer, ArgumentOffset);
 	}
 
-	FORCEINLINE_DEBUGGABLE void DrawIndexedIndirect(FIndexBufferRHIParamRef IndexBufferRHI, uint32 PrimitiveType, FStructuredBufferRHIParamRef ArgumentsBufferRHI, uint32 DrawArgumentsIndex, uint32 NumInstances)
+	DEPRECATED(4.22, "PrimitiveType is now a part of the Graphics Pipeline State.")
+	FORCEINLINE_DEBUGGABLE void DrawPrimitiveIndirect(uint32 PrimitiveType, FVertexBufferRHIParamRef ArgumentBuffer, uint32 ArgumentOffset)
+	{
+		DrawPrimitiveIndirect(ArgumentBuffer, ArgumentOffset);
+	}
+
+
+	FORCEINLINE_DEBUGGABLE void DrawIndexedIndirect(FIndexBufferRHIParamRef IndexBufferRHI, FStructuredBufferRHIParamRef ArgumentsBufferRHI, uint32 DrawArgumentsIndex, uint32 NumInstances)
 	{
 		//check(IsOutsideRenderPass());
 		if (Bypass())
 		{
-			CMD_CONTEXT(RHIDrawIndexedIndirect)(IndexBufferRHI, PrimitiveType, ArgumentsBufferRHI, DrawArgumentsIndex, NumInstances);
+			CMD_CONTEXT(RHIDrawIndexedIndirect)(IndexBufferRHI, ArgumentsBufferRHI, DrawArgumentsIndex, NumInstances);
 			return;
 		}
-		new (AllocCommand<FRHICommandDrawIndexedIndirect>()) FRHICommandDrawIndexedIndirect(IndexBufferRHI, PrimitiveType, ArgumentsBufferRHI, DrawArgumentsIndex, NumInstances);
+		new (AllocCommand<FRHICommandDrawIndexedIndirect>()) FRHICommandDrawIndexedIndirect(IndexBufferRHI, ArgumentsBufferRHI, DrawArgumentsIndex, NumInstances);
 	}
 
-	FORCEINLINE_DEBUGGABLE void DrawIndexedPrimitiveIndirect(uint32 PrimitiveType, FIndexBufferRHIParamRef IndexBuffer, FVertexBufferRHIParamRef ArgumentsBuffer, uint32 ArgumentOffset)
+	DEPRECATED(4.22, "PrimitiveType is now a part of the Graphics Pipeline State.")
+	FORCEINLINE_DEBUGGABLE void DrawIndexedIndirect(FIndexBufferRHIParamRef IndexBufferRHI, uint32 PrimitiveType, FStructuredBufferRHIParamRef ArgumentsBufferRHI, uint32 DrawArgumentsIndex, uint32 NumInstances)
+	{
+		DrawIndexedIndirect(IndexBufferRHI, ArgumentsBufferRHI, DrawArgumentsIndex, NumInstances);
+	}
+
+	FORCEINLINE_DEBUGGABLE void DrawIndexedPrimitiveIndirect(FIndexBufferRHIParamRef IndexBuffer, FVertexBufferRHIParamRef ArgumentsBuffer, uint32 ArgumentOffset)
 	{
 		//check(IsOutsideRenderPass());
 		if (Bypass())
 		{
-			CMD_CONTEXT(RHIDrawIndexedPrimitiveIndirect)(PrimitiveType, IndexBuffer, ArgumentsBuffer, ArgumentOffset);
+			CMD_CONTEXT(RHIDrawIndexedPrimitiveIndirect)(IndexBuffer, ArgumentsBuffer, ArgumentOffset);
 			return;
 		}
-		new (AllocCommand<FRHICommandDrawIndexedPrimitiveIndirect>()) FRHICommandDrawIndexedPrimitiveIndirect(PrimitiveType, IndexBuffer, ArgumentsBuffer, ArgumentOffset);
+		new (AllocCommand<FRHICommandDrawIndexedPrimitiveIndirect>()) FRHICommandDrawIndexedPrimitiveIndirect(IndexBuffer, ArgumentsBuffer, ArgumentOffset);
+	}
+
+	DEPRECATED(4.22, "PrimitiveType is now a part of the Graphics Pipeline State.")
+	FORCEINLINE_DEBUGGABLE void DrawIndexedPrimitiveIndirect(uint32 PrimitiveType, FIndexBufferRHIParamRef IndexBuffer, FVertexBufferRHIParamRef ArgumentsBuffer, uint32 ArgumentOffset)
+	{
+		DrawIndexedPrimitiveIndirect(IndexBuffer, ArgumentsBuffer, ArgumentOffset);
 	}
 
 	FORCEINLINE_DEBUGGABLE void SetDepthBounds(float MinDepth, float MaxDepth)
diff --git a/Engine/Source/Runtime/RHI/Public/RHICommandListCommandExecutes.inl b/Engine/Source/Runtime/RHI/Public/RHICommandListCommandExecutes.inl
--- a/Engine/Source/Runtime/RHI/Public/RHICommandListCommandExecutes.inl
+++ b/Engine/Source/Runtime/RHI/Public/RHICommandListCommandExecutes.inl
@@ -223,13 +223,13 @@ template<> void FRHICommandSetShaderSampler<FComputeShaderRHIParamRef, ECmdList:
 void FRHICommandDrawPrimitive::Execute(FRHICommandListBase& CmdList)
 {
 	RHISTAT(DrawPrimitive);
-	INTERNAL_DECORATOR(RHIDrawPrimitive)(PrimitiveType, BaseVertexIndex, NumPrimitives, NumInstances);
+	INTERNAL_DECORATOR(RHIDrawPrimitive)(BaseVertexIndex, NumPrimitives, NumInstances);
 }
 
 void FRHICommandDrawIndexedPrimitive::Execute(FRHICommandListBase& CmdList)
 {
 	RHISTAT(DrawIndexedPrimitive);
-	INTERNAL_DECORATOR(RHIDrawIndexedPrimitive)(IndexBuffer, PrimitiveType, BaseVertexIndex, FirstInstance, NumVertices, StartIndex, NumPrimitives, NumInstances);
+	INTERNAL_DECORATOR(RHIDrawIndexedPrimitive)(IndexBuffer, BaseVertexIndex, FirstInstance, NumVertices, StartIndex, NumPrimitives, NumInstances);
 }
 
 void FRHICommandSetBoundShaderState::Execute(FRHICommandListBase& CmdList)
@@ -329,7 +329,7 @@ void FRHICommandEndDrawPrimitiveUP::Execute(FRHICommandListBase& CmdList)
 {
 	RHISTAT(EndDrawPrimitiveUP);
 	void* Buffer = NULL;
-	INTERNAL_DECORATOR(RHIBeginDrawPrimitiveUP)(PrimitiveType, NumPrimitives, NumVertices, VertexDataStride, Buffer);
+	INTERNAL_DECORATOR(RHIBeginDrawPrimitiveUP)(NumPrimitives, NumVertices, VertexDataStride, Buffer);
 	FMemory::Memcpy(Buffer, OutVertexData, NumVertices * VertexDataStride);
 	INTERNAL_DECORATOR(RHIEndDrawPrimitiveUP)();
 }
@@ -340,7 +340,6 @@ void FRHICommandEndDrawIndexedPrimitiveUP::Execute(FRHICommandListBase& CmdList)
 	void* VertexBuffer = nullptr;
 	void* IndexBuffer = nullptr;
 	INTERNAL_DECORATOR(RHIBeginDrawIndexedPrimitiveUP)(
-		PrimitiveType,
 		NumPrimitives,
 		NumVertices,
 		VertexDataStride,
@@ -415,19 +414,19 @@ void FRHICommandFlushComputeShaderCache::Execute(FRHICommandListBase& CmdList)
 void FRHICommandDrawPrimitiveIndirect::Execute(FRHICommandListBase& CmdList)
 {
 	RHISTAT(DrawPrimitiveIndirect);
-	INTERNAL_DECORATOR(RHIDrawPrimitiveIndirect)(PrimitiveType, ArgumentBuffer, ArgumentOffset);
+	INTERNAL_DECORATOR(RHIDrawPrimitiveIndirect)(ArgumentBuffer, ArgumentOffset);
 }
 
 void FRHICommandDrawIndexedIndirect::Execute(FRHICommandListBase& CmdList)
 {
 	RHISTAT(DrawIndexedIndirect);
-	INTERNAL_DECORATOR(RHIDrawIndexedIndirect)(IndexBufferRHI, PrimitiveType, ArgumentsBufferRHI, DrawArgumentsIndex, NumInstances);
+	INTERNAL_DECORATOR(RHIDrawIndexedIndirect)(IndexBufferRHI, ArgumentsBufferRHI, DrawArgumentsIndex, NumInstances);
 }
 
 void FRHICommandDrawIndexedPrimitiveIndirect::Execute(FRHICommandListBase& CmdList)
 {
 	RHISTAT(DrawIndexedPrimitiveIndirect);
-	INTERNAL_DECORATOR(RHIDrawIndexedPrimitiveIndirect)(PrimitiveType, IndexBuffer, ArgumentsBuffer, ArgumentOffset);
+	INTERNAL_DECORATOR(RHIDrawIndexedPrimitiveIndirect)(IndexBuffer, ArgumentsBuffer, ArgumentOffset);
 }
 
 void FRHICommandSetDepthBounds::Execute(FRHICommandListBase& CmdList)
diff --git a/Engine/Source/Runtime/RHI/Public/RHIContext.h b/Engine/Source/Runtime/RHI/Public/RHIContext.h
--- a/Engine/Source/Runtime/RHI/Public/RHIContext.h
+++ b/Engine/Source/Runtime/RHI/Public/RHIContext.h
@@ -499,26 +499,25 @@ public:
 	// The explicit bind is needed to support parallel rendering (propagate state between contexts).
 	virtual void RHIBindClearMRTValues(bool bClearColor, bool bClearDepth, bool bClearStencil) {}
 
-	virtual void RHIDrawPrimitive(uint32 PrimitiveType, uint32 BaseVertexIndex, uint32 NumPrimitives, uint32 NumInstances) = 0;
+	virtual void RHIDrawPrimitive(uint32 BaseVertexIndex, uint32 NumPrimitives, uint32 NumInstances) = 0;
 
-	virtual void RHIDrawPrimitiveIndirect(uint32 PrimitiveType, FVertexBufferRHIParamRef ArgumentBuffer, uint32 ArgumentOffset) = 0;
+	virtual void RHIDrawPrimitiveIndirect(FVertexBufferRHIParamRef ArgumentBuffer, uint32 ArgumentOffset) = 0;
 
-	virtual void RHIDrawIndexedIndirect(FIndexBufferRHIParamRef IndexBufferRHI, uint32 PrimitiveType, FStructuredBufferRHIParamRef ArgumentsBufferRHI, int32 DrawArgumentsIndex, uint32 NumInstances) = 0;
+	virtual void RHIDrawIndexedIndirect(FIndexBufferRHIParamRef IndexBufferRHI, FStructuredBufferRHIParamRef ArgumentsBufferRHI, int32 DrawArgumentsIndex, uint32 NumInstances) = 0;
 
 	// @param NumPrimitives need to be >0 
-	virtual void RHIDrawIndexedPrimitive(FIndexBufferRHIParamRef IndexBuffer, uint32 PrimitiveType, int32 BaseVertexIndex, uint32 FirstInstance, uint32 NumVertices, uint32 StartIndex, uint32 NumPrimitives, uint32 NumInstances) = 0;
+	virtual void RHIDrawIndexedPrimitive(FIndexBufferRHIParamRef IndexBuffer, int32 BaseVertexIndex, uint32 FirstInstance, uint32 NumVertices, uint32 StartIndex, uint32 NumPrimitives, uint32 NumInstances) = 0;
 
-	virtual void RHIDrawIndexedPrimitiveIndirect(uint32 PrimitiveType, FIndexBufferRHIParamRef IndexBuffer, FVertexBufferRHIParamRef ArgumentBuffer, uint32 ArgumentOffset) = 0;
+	virtual void RHIDrawIndexedPrimitiveIndirect(FIndexBufferRHIParamRef IndexBuffer, FVertexBufferRHIParamRef ArgumentBuffer, uint32 ArgumentOffset) = 0;
 
 	/**
 	* Preallocate memory or get a direct command stream pointer to fill up for immediate rendering . This avoids memcpys below in DrawPrimitiveUP
-	* @param PrimitiveType The type (triangles, lineloop, etc) of primitive to draw
 	* @param NumPrimitives The number of primitives in the VertexData buffer
 	* @param NumVertices The number of vertices to be written
 	* @param VertexDataStride Size of each vertex
 	* @param OutVertexData Reference to the allocated vertex memory
 	*/
-	virtual void RHIBeginDrawPrimitiveUP(uint32 PrimitiveType, uint32 NumPrimitives, uint32 NumVertices, uint32 VertexDataStride, void*& OutVertexData) = 0;
+	virtual void RHIBeginDrawPrimitiveUP(uint32 NumPrimitives, uint32 NumVertices, uint32 VertexDataStride, void*& OutVertexData) = 0;
 
 	/**
 	* Draw a primitive using the vertex data populated since RHIBeginDrawPrimitiveUP and clean up any memory as needed
@@ -527,7 +526,6 @@ public:
 
 	/**
 	* Preallocate memory or get a direct command stream pointer to fill up for immediate rendering . This avoids memcpys below in DrawIndexedPrimitiveUP
-	* @param PrimitiveType The type (triangles, lineloop, etc) of primitive to draw
 	* @param NumPrimitives The number of primitives in the VertexData buffer
 	* @param NumVertices The number of vertices to be written
 	* @param VertexDataStride Size of each vertex
@@ -537,7 +535,7 @@ public:
 	* @param IndexDataStride Size of each index (either 2 or 4 bytes)
 	* @param OutIndexData Reference to the allocated index memory
 	*/
-	virtual void RHIBeginDrawIndexedPrimitiveUP(uint32 PrimitiveType, uint32 NumPrimitives, uint32 NumVertices, uint32 VertexDataStride, void*& OutVertexData, uint32 MinVertexIndex, uint32 NumIndices, uint32 IndexDataStride, void*& OutIndexData) = 0;
+	virtual void RHIBeginDrawIndexedPrimitiveUP(uint32 NumPrimitives, uint32 NumVertices, uint32 VertexDataStride, void*& OutVertexData, uint32 MinVertexIndex, uint32 NumIndices, uint32 IndexDataStride, void*& OutIndexData) = 0;
 
 	/**
 	* Draw a primitive using the vertex and index data populated since RHIBeginDrawIndexedPrimitiveUP and clean up any memory as needed
diff --git a/Engine/Source/Runtime/RHI/Public/RHIUtilities.h b/Engine/Source/Runtime/RHI/Public/RHIUtilities.h
--- a/Engine/Source/Runtime/RHI/Public/RHIUtilities.h
+++ b/Engine/Source/Runtime/RHI/Public/RHIUtilities.h
@@ -816,7 +816,9 @@ inline void DrawPrimitiveUP(FRHICommandList& RHICmdList, uint32 PrimitiveType, u
 	RHIUnlockVertexBuffer(VertexBufferRHI);
 
 	RHICmdList.SetStreamSource(0, VertexBufferRHI, 0);
-	RHICmdList.DrawPrimitive(PrimitiveType, 0, NumPrimitives, 1);
+	PRAGMA_DISABLE_DEPRECATION_WARNINGS
+	RHICmdList.DrawPrimitive(0, NumPrimitives, 1);
+	PRAGMA_ENABLE_DEPRECATION_WARNINGS
 
 	VertexBufferRHI.SafeRelease();
 }
@@ -858,7 +860,7 @@ inline void DrawIndexedPrimitiveUP(
 	RHIUnlockIndexBuffer(IndexBufferRHI);
 
 	RHICmdList.SetStreamSource(0, VertexBufferRHI, 0);
-	RHICmdList.DrawIndexedPrimitive(IndexBufferRHI, PrimitiveType, MinVertexIndex, 0, NumVertices, 0, NumPrimitives, 1);
+	RHICmdList.DrawIndexedPrimitive(IndexBufferRHI, MinVertexIndex, 0, NumVertices, 0, NumPrimitives, 1);
 
 	IndexBufferRHI.SafeRelease();
 	VertexBufferRHI.SafeRelease();
diff --git a/Engine/Source/Runtime/Renderer/Private/AtmosphereRendering.cpp b/Engine/Source/Runtime/Renderer/Private/AtmosphereRendering.cpp
--- a/Engine/Source/Runtime/Renderer/Private/AtmosphereRendering.cpp
+++ b/Engine/Source/Runtime/Renderer/Private/AtmosphereRendering.cpp
@@ -420,7 +420,7 @@ void FDeferredShadingSceneRenderer::RenderAtmosphere(FRHICommandListImmediate& R
 
 			// Draw a quad covering the view.
 			RHICmdList.SetStreamSource(0, GScreenSpaceVertexBuffer.VertexBufferRHI, 0);
-			RHICmdList.DrawIndexedPrimitive(GTwoTrianglesIndexBuffer.IndexBufferRHI, PT_TriangleList, 0, 0, 4, 0, 2, 1);
+			RHICmdList.DrawIndexedPrimitive(GTwoTrianglesIndexBuffer.IndexBufferRHI, 0, 0, 4, 0, 2, 1);
 		}
 	}
 }
diff --git a/Engine/Source/Runtime/Renderer/Private/CapsuleShadowRendering.cpp b/Engine/Source/Runtime/Renderer/Private/CapsuleShadowRendering.cpp
--- a/Engine/Source/Runtime/Renderer/Private/CapsuleShadowRendering.cpp
+++ b/Engine/Source/Runtime/Renderer/Private/CapsuleShadowRendering.cpp
@@ -1004,7 +1004,7 @@ bool FDeferredShadingSceneRenderer::RenderCapsuleDirectShadows(
 					}
 
 					RHICmdList.SetStreamSource(0, GTileTexCoordVertexBuffer.VertexBufferRHI, 0);
-					RHICmdList.DrawIndexedPrimitive(GTileIndexBuffer.IndexBufferRHI, PT_TriangleList, 0, 0, 4, 0, 2 * NumTileQuadsInBuffer, FMath::DivideAndRoundUp(GroupSize.X * GroupSize.Y, NumTileQuadsInBuffer));
+					RHICmdList.DrawIndexedPrimitive(GTileIndexBuffer.IndexBufferRHI, 0, 0, 4, 0, 2 * NumTileQuadsInBuffer, FMath::DivideAndRoundUp(GroupSize.X * GroupSize.Y, NumTileQuadsInBuffer));
 				}
 			}
 		}
@@ -1499,7 +1499,7 @@ void FDeferredShadingSceneRenderer::RenderIndirectCapsuleShadows(
 							}
 
 							RHICmdList.SetStreamSource(0, GTileTexCoordVertexBuffer.VertexBufferRHI, 0);
-							RHICmdList.DrawIndexedPrimitive(GTileIndexBuffer.IndexBufferRHI, PT_TriangleList, 0, 0, 4, 0, 2 * NumTileQuadsInBuffer, FMath::DivideAndRoundUp(GroupSize.X * GroupSize.Y, NumTileQuadsInBuffer));
+							RHICmdList.DrawIndexedPrimitive(GTileIndexBuffer.IndexBufferRHI, 0, 0, 4, 0, 2 * NumTileQuadsInBuffer, FMath::DivideAndRoundUp(GroupSize.X * GroupSize.Y, NumTileQuadsInBuffer));
 						}
 					}
 				}
diff --git a/Engine/Source/Runtime/Renderer/Private/CompositionLighting/PostProcessDeferredDecals.cpp b/Engine/Source/Runtime/Renderer/Private/CompositionLighting/PostProcessDeferredDecals.cpp
--- a/Engine/Source/Runtime/Renderer/Private/CompositionLighting/PostProcessDeferredDecals.cpp
+++ b/Engine/Source/Runtime/Renderer/Private/CompositionLighting/PostProcessDeferredDecals.cpp
@@ -405,7 +405,7 @@ bool RenderPreStencil(FRenderingCompositePassContext& Context, const FMatrix& Co
 	Context.RHICmdList.SetStreamSource(0, GetUnitCubeVertexBuffer(), 0);
 
 	// Render decal mask
-	Context.RHICmdList.DrawIndexedPrimitive(GetUnitCubeIndexBuffer(), PT_TriangleList, 0, 0, 8, 0, ARRAY_COUNT(GCubeIndices) / 3, 1);
+	Context.RHICmdList.DrawIndexedPrimitive(GetUnitCubeIndexBuffer(), 0, 0, 8, 0, ARRAY_COUNT(GCubeIndices) / 3, 1);
 
 	return true;
 }
@@ -904,7 +904,7 @@ void FRCPassPostProcessDeferredDecals::Process(FRenderingCompositePassContext& C
 					FDecalRendering::SetShader(RHICmdList, GraphicsPSOInit, View, DecalData, CurrentStage, FrustumComponentToClip);
 					RHICmdList.SetStencilRef(StencilRef);
 
-					RHICmdList.DrawIndexedPrimitive(GetUnitCubeIndexBuffer(), PT_TriangleList, 0, 0, 8, 0, ARRAY_COUNT(GCubeIndices) / 3, 1);
+					RHICmdList.DrawIndexedPrimitive(GetUnitCubeIndexBuffer(), 0, 0, 8, 0, ARRAY_COUNT(GCubeIndices) / 3, 1);
 					RenderTargetManager.bGufferADirty |= (RenderTargetManager.TargetsToResolve[FDecalRenderTargetManager::GBufferAIndex] != nullptr);
 				}
 
diff --git a/Engine/Source/Runtime/Renderer/Private/DeferredShadingRenderer.cpp b/Engine/Source/Runtime/Renderer/Private/DeferredShadingRenderer.cpp
--- a/Engine/Source/Runtime/Renderer/Private/DeferredShadingRenderer.cpp
+++ b/Engine/Source/Runtime/Renderer/Private/DeferredShadingRenderer.cpp
@@ -327,7 +327,7 @@ void FDeferredShadingSceneRenderer::ClearGBufferAtMaxZ(FRHICommandList& RHICmdLi
 
 		RHICmdList.SetStreamSource(0, GClearVertexBuffer.VertexBufferRHI, 0);
 		// Render quad
-		RHICmdList.DrawPrimitive(PT_TriangleStrip, 0, 2, 1);
+		RHICmdList.DrawPrimitive(0, 2, 1);
 	}
 }
 
diff --git a/Engine/Source/Runtime/Renderer/Private/DistanceFieldObjectCulling.cpp b/Engine/Source/Runtime/Renderer/Private/DistanceFieldObjectCulling.cpp
--- a/Engine/Source/Runtime/Renderer/Private/DistanceFieldObjectCulling.cpp
+++ b/Engine/Source/Runtime/Renderer/Private/DistanceFieldObjectCulling.cpp
@@ -594,7 +594,6 @@ void ScatterTilesToObjects(FRHICommandListImmediate& RHICmdList, const FViewInfo
 	RHICmdList.SetStreamSource(0, StencilingGeometry::GLowPolyStencilSphereVertexBuffer.VertexBufferRHI, 0);
 
 	RHICmdList.DrawIndexedPrimitiveIndirect(
-		PT_TriangleList,
 		StencilingGeometry::GLowPolyStencilSphereIndexBuffer.IndexBufferRHI, 
 		GAOCulledObjectBuffers.Buffers.ObjectIndirectArguments.Buffer,
 		0);
diff --git a/Engine/Source/Runtime/Renderer/Private/DistanceFieldShadowing.cpp b/Engine/Source/Runtime/Renderer/Private/DistanceFieldShadowing.cpp
--- a/Engine/Source/Runtime/Renderer/Private/DistanceFieldShadowing.cpp
+++ b/Engine/Source/Runtime/Renderer/Private/DistanceFieldShadowing.cpp
@@ -722,7 +722,6 @@ void ScatterObjectsToShadowTiles(
 	RHICmdList.SetStreamSource(0, GetUnitCubeVertexBuffer(), 0);
 
 	RHICmdList.DrawIndexedPrimitiveIndirect(
-		PT_TriangleList,
 		GetUnitCubeIndexBuffer(),
 		GShadowCulledObjectBuffers.Buffers.ObjectIndirectArguments.Buffer,
 		0);
diff --git a/Engine/Source/Runtime/Renderer/Private/DrawingPolicy.cpp b/Engine/Source/Runtime/Renderer/Private/DrawingPolicy.cpp
--- a/Engine/Source/Runtime/Renderer/Private/DrawingPolicy.cpp
+++ b/Engine/Source/Runtime/Renderer/Private/DrawingPolicy.cpp
@@ -111,7 +111,6 @@ void FMeshDrawingPolicy::DrawMesh(FRHICommandList& RHICmdList, const FSceneView&
 
 						RHICmdList.DrawIndexedPrimitive(
 							BatchElement.IndexBuffer->IndexBufferRHI,
-							Mesh.Type,
 							BatchElement.BaseVertexIndex,
 							0,
 							BatchElement.MaxVertexIndex - BatchElement.MinVertexIndex + 1,
@@ -132,7 +131,6 @@ void FMeshDrawingPolicy::DrawMesh(FRHICommandList& RHICmdList, const FSceneView&
 
 						RHICmdList.DrawIndexedPrimitive(
 							BatchElement.IndexBuffer->IndexBufferRHI,
-							Mesh.Type,
 							BatchElement.BaseVertexIndex,
 							0,
 							BatchElement.MaxVertexIndex - BatchElement.MinVertexIndex + 1,
@@ -153,7 +151,6 @@ void FMeshDrawingPolicy::DrawMesh(FRHICommandList& RHICmdList, const FSceneView&
 
 					RHICmdList.DrawIndexedPrimitive(
 						BatchElement.IndexBuffer->IndexBufferRHI,
-						Mesh.Type,
 						BatchElement.BaseVertexIndex,
 						InstanceOffset,
 						BatchElement.MaxVertexIndex - BatchElement.MinVertexIndex + 1,
@@ -169,7 +166,6 @@ void FMeshDrawingPolicy::DrawMesh(FRHICommandList& RHICmdList, const FSceneView&
 			if (BatchElement.IndirectArgsBuffer)
 			{
 				RHICmdList.DrawIndexedPrimitiveIndirect(
-					Mesh.Type, 
 					BatchElement.IndexBuffer->IndexBufferRHI, 
 					BatchElement.IndirectArgsBuffer, 
 					0
@@ -183,7 +179,6 @@ void FMeshDrawingPolicy::DrawMesh(FRHICommandList& RHICmdList, const FSceneView&
 
 				RHICmdList.DrawIndexedPrimitive(
 					BatchElement.IndexBuffer->IndexBufferRHI,
-					Mesh.Type,
 					BatchElement.BaseVertexIndex,
 					0,
 					BatchElement.MaxVertexIndex - BatchElement.MinVertexIndex + 1,
@@ -199,7 +194,6 @@ void FMeshDrawingPolicy::DrawMesh(FRHICommandList& RHICmdList, const FSceneView&
 		SetInstanceParameters(RHICmdList, View, BatchElement.BaseVertexIndex + BatchElement.FirstIndex, 0, 1);
 
 		RHICmdList.DrawPrimitive(
-			Mesh.Type,
 			BatchElement.BaseVertexIndex + BatchElement.FirstIndex,
 			BatchElement.NumPrimitives,
 			BatchElement.NumInstances
diff --git a/Engine/Source/Runtime/Renderer/Private/FogRendering.cpp b/Engine/Source/Runtime/Renderer/Private/FogRendering.cpp
--- a/Engine/Source/Runtime/Renderer/Private/FogRendering.cpp
+++ b/Engine/Source/Runtime/Renderer/Private/FogRendering.cpp
@@ -421,7 +421,7 @@ void FDeferredShadingSceneRenderer::RenderViewFog(FRHICommandList& RHICmdList, c
 
 	// Draw a quad covering the view.
 	RHICmdList.SetStreamSource(0, GScreenSpaceVertexBuffer.VertexBufferRHI, 0);
-	RHICmdList.DrawIndexedPrimitive(GTwoTrianglesIndexBuffer.IndexBufferRHI, PT_TriangleList, 0, 0, 4, 0, 2, 1);
+	RHICmdList.DrawIndexedPrimitive(GTwoTrianglesIndexBuffer.IndexBufferRHI, 0, 0, 4, 0, 2, 1);
 }
 
 bool FDeferredShadingSceneRenderer::RenderFog(FRHICommandListImmediate& RHICmdList, const FLightShaftsOutput& LightShaftsOutput)
diff --git a/Engine/Source/Runtime/Renderer/Private/GPUBenchmark.cpp b/Engine/Source/Runtime/Renderer/Private/GPUBenchmark.cpp
--- a/Engine/Source/Runtime/Renderer/Private/GPUBenchmark.cpp
+++ b/Engine/Source/Runtime/Renderer/Private/GPUBenchmark.cpp
@@ -231,7 +231,7 @@ void RunBenchmarkShader(FRHICommandList& RHICmdList, FVertexBufferRHIParamRef Ve
 
 			RHICmdList.SetStreamSource(0, VertexThroughputBuffer, 0);
 
-			RHICmdList.DrawPrimitive(PT_TriangleList, 0, PrimitivesThisPass, 1);
+			RHICmdList.DrawPrimitive(0, PrimitivesThisPass, 1);
 
 			TotalNumVertices -= VerticesThisPass;
 		}
diff --git a/Engine/Source/Runtime/Renderer/Private/HeightfieldLighting.cpp b/Engine/Source/Runtime/Renderer/Private/HeightfieldLighting.cpp
--- a/Engine/Source/Runtime/Renderer/Private/HeightfieldLighting.cpp
+++ b/Engine/Source/Runtime/Renderer/Private/HeightfieldLighting.cpp
@@ -503,7 +503,7 @@ void FHeightfieldLightingViewInfo::SetupVisibleHeightfields(const FViewInfo& Vie
 							VertexShader->SetParameters(RHICmdList, View);
 							PixelShader->SetParameters(RHICmdList, View, It.Key().HeightAndNormal, It.Key().DiffuseColor);
 
-							RHICmdList.DrawPrimitive(PT_TriangleList, 0, 2, NumQuads);
+							RHICmdList.DrawPrimitive(0, 2, NumQuads);
 						}
 					}
 
@@ -867,7 +867,7 @@ void FHeightfieldLightingViewInfo::ComputeShadowMapShadowing(const FViewInfo& Vi
 					VertexShader->SetParameters(RHICmdList, View, HeightfieldDescriptions.Num());
 					PixelShader->SetParameters(RHICmdList, View, ProjectedShadowInfo, HeightfieldDescriptions.Num(), Atlas);
 
-					RHICmdList.DrawPrimitive(PT_TriangleList, 0, 2, HeightfieldDescriptions.Num());
+					RHICmdList.DrawPrimitive(0, 2, HeightfieldDescriptions.Num());
 				}
 			}
 		}
@@ -1032,7 +1032,7 @@ void FHeightfieldLightingViewInfo::ComputeRayTracedShadowing(
 					VertexShader->SetParameters(RHICmdList, View, HeightfieldDescriptions.Num());
 					PixelShader->SetParameters(RHICmdList, View, ProjectedShadowInfo, HeightfieldDescriptions.Num(), Atlas, TileIntersectionResources, CulledObjectBuffers);
 
-					RHICmdList.DrawPrimitive(PT_TriangleList, 0, 2, HeightfieldDescriptions.Num());
+					RHICmdList.DrawPrimitive(0, 2, HeightfieldDescriptions.Num());
 				}
 			}
 		}
@@ -1208,7 +1208,7 @@ void FHeightfieldLightingViewInfo::ComputeLighting(const FViewInfo& View, FRHICo
 					VertexShader->SetParameters(RHICmdList, View, HeightfieldDescriptions.Num());
 					PixelShader->SetParameters(RHICmdList, View, LightSceneInfo, MaterialProxy, HeightfieldDescriptions.Num(), Atlas, SkyLightIndirectScale);
 
-					RHICmdList.DrawPrimitive(PT_TriangleList, 0, 2, HeightfieldDescriptions.Num());
+					RHICmdList.DrawPrimitive(0, 2, HeightfieldDescriptions.Num());
 				}
 			}
 		}
diff --git a/Engine/Source/Runtime/Renderer/Private/LightPropagationVolumeVisualisation.cpp b/Engine/Source/Runtime/Renderer/Private/LightPropagationVolumeVisualisation.cpp
--- a/Engine/Source/Runtime/Renderer/Private/LightPropagationVolumeVisualisation.cpp
+++ b/Engine/Source/Runtime/Renderer/Private/LightPropagationVolumeVisualisation.cpp
@@ -235,7 +235,7 @@ void FLightPropagationVolume::Visualise(FRHICommandList& RHICmdList, const FView
 	PixelShader->SetParameters(RHICmdList, this, View);
 
 	RHICmdList.SetStreamSource(0, NULL, 0);
-	RHICmdList.DrawPrimitive(PrimType, 0, NumPrims, 32 * 3);
+	RHICmdList.DrawPrimitive(0, NumPrims, 32 * 3);
 
 	PixelShader->UnbindBuffers(RHICmdList);
 }
diff --git a/Engine/Source/Runtime/Renderer/Private/LightRendering.cpp b/Engine/Source/Runtime/Renderer/Private/LightRendering.cpp
--- a/Engine/Source/Runtime/Renderer/Private/LightRendering.cpp
+++ b/Engine/Source/Runtime/Renderer/Private/LightRendering.cpp
@@ -58,7 +58,7 @@ float GetLightFadeFactor(const FSceneView& View, const FLightSceneProxy* Proxy)
 void StencilingGeometry::DrawSphere(FRHICommandList& RHICmdList)
 {
 	RHICmdList.SetStreamSource(0, StencilingGeometry::GStencilSphereVertexBuffer.VertexBufferRHI, 0);
-	RHICmdList.DrawIndexedPrimitive(StencilingGeometry::GStencilSphereIndexBuffer.IndexBufferRHI, PT_TriangleList, 0, 0,
+	RHICmdList.DrawIndexedPrimitive(StencilingGeometry::GStencilSphereIndexBuffer.IndexBufferRHI, 0, 0,
 		StencilingGeometry::GStencilSphereVertexBuffer.GetVertexCount(), 0, 
 		StencilingGeometry::GStencilSphereIndexBuffer.GetIndexCount() / 3, 1);
 }
@@ -66,7 +66,7 @@ void StencilingGeometry::DrawSphere(FRHICommandList& RHICmdList)
 void StencilingGeometry::DrawVectorSphere(FRHICommandList& RHICmdList)
 {
 	RHICmdList.SetStreamSource(0, StencilingGeometry::GStencilSphereVectorBuffer.VertexBufferRHI, 0);
-	RHICmdList.DrawIndexedPrimitive(StencilingGeometry::GStencilSphereIndexBuffer.IndexBufferRHI, PT_TriangleList, 0, 0,
+	RHICmdList.DrawIndexedPrimitive(StencilingGeometry::GStencilSphereIndexBuffer.IndexBufferRHI, 0, 0,
 									StencilingGeometry::GStencilSphereVectorBuffer.GetVertexCount(), 0,
 									StencilingGeometry::GStencilSphereIndexBuffer.GetIndexCount() / 3, 1);
 }
@@ -76,7 +76,7 @@ void StencilingGeometry::DrawCone(FRHICommandList& RHICmdList)
 	// No Stream Source needed since it will generate vertices on the fly
 	RHICmdList.SetStreamSource(0, StencilingGeometry::GStencilConeVertexBuffer.VertexBufferRHI, 0);
 
-	RHICmdList.DrawIndexedPrimitive(StencilingGeometry::GStencilConeIndexBuffer.IndexBufferRHI, PT_TriangleList, 0, 0,
+	RHICmdList.DrawIndexedPrimitive(StencilingGeometry::GStencilConeIndexBuffer.IndexBufferRHI, 0, 0,
 		FStencilConeIndexBuffer::NumVerts, 0, StencilingGeometry::GStencilConeIndexBuffer.GetIndexCount() / 3, 1);
 }
 
diff --git a/Engine/Source/Runtime/Renderer/Private/MobileDecalRendering.cpp b/Engine/Source/Runtime/Renderer/Private/MobileDecalRendering.cpp
--- a/Engine/Source/Runtime/Renderer/Private/MobileDecalRendering.cpp
+++ b/Engine/Source/Runtime/Renderer/Private/MobileDecalRendering.cpp
@@ -147,7 +147,7 @@ void FMobileSceneRenderer::RenderDecals(FRHICommandListImmediate& RHICmdList)
 				// Set shader params
 				FDecalRendering::SetShader(RHICmdList, GraphicsPSOInit, View, DecalData, DRS_Mobile, FrustumComponentToClip);
 			
-				RHICmdList.DrawIndexedPrimitive(GetUnitCubeIndexBuffer(), PT_TriangleList, 0, 0, 8, 0, ARRAY_COUNT(GCubeIndices) / 3, 1);
+				RHICmdList.DrawIndexedPrimitive(GetUnitCubeIndexBuffer(), 0, 0, 8, 0, ARRAY_COUNT(GCubeIndices) / 3, 1);
 			}
 		}
 	}
diff --git a/Engine/Source/Runtime/Renderer/Private/PostProcess/DiaphragmDOFPasses.cpp b/Engine/Source/Runtime/Renderer/Private/PostProcess/DiaphragmDOFPasses.cpp
--- a/Engine/Source/Runtime/Renderer/Private/PostProcess/DiaphragmDOFPasses.cpp
+++ b/Engine/Source/Runtime/Renderer/Private/PostProcess/DiaphragmDOFPasses.cpp
@@ -1720,14 +1720,13 @@ void FRCPassDiaphragmDOFHybridScatter::Process(FRenderingCompositePassContext& C
 	if (GRHISupportsRectTopology)
 	{
 		Context.RHICmdList.DrawPrimitiveIndirect(
-			PrimitiveType,
 			DrawIndirectParametersBuffer->Buffer,
 			sizeof(FRHIDrawIndirectParameters) * DrawIndirectParametersOffset);
 	}
 	else
 	{
 		Context.RHICmdList.DrawIndexedPrimitiveIndirect(
-			PrimitiveType, GDiaphragmDOFGlobalResource.ScatterIndexBuffer.IndexBufferRHI,
+			GDiaphragmDOFGlobalResource.ScatterIndexBuffer.IndexBufferRHI,
 			DrawIndirectParametersBuffer->Buffer,
 			sizeof(FRHIDrawIndexedIndirectParameters) * DrawIndirectParametersOffset);
 	}
diff --git a/Engine/Source/Runtime/Renderer/Private/PostProcess/PostProcessBokehDOF.cpp b/Engine/Source/Runtime/Renderer/Private/PostProcess/PostProcessBokehDOF.cpp
--- a/Engine/Source/Runtime/Renderer/Private/PostProcess/PostProcessBokehDOF.cpp
+++ b/Engine/Source/Runtime/Renderer/Private/PostProcess/PostProcessBokehDOF.cpp
@@ -911,7 +911,7 @@ void FRCPassPostProcessBokehDOF::Process(FRenderingCompositePassContext& Context
 	}
 
 	Context.RHICmdList.SetStreamSource(0, NULL, 0);
-	Context.RHICmdList.DrawIndexedPrimitive(GBokehIndexBuffer.IndexBufferRHI, PT_TriangleList, 0, 0, 4 * GBokehDOFQuadsPerInstance, 0, 2 * GBokehDOFQuadsPerInstance, FMath::DivideAndRoundUp(TileCount.X * TileCount.Y, GBokehDOFQuadsPerInstance));
+	Context.RHICmdList.DrawIndexedPrimitive(GBokehIndexBuffer.IndexBufferRHI, 0, 0, 4 * GBokehDOFQuadsPerInstance, 0, 2 * GBokehDOFQuadsPerInstance, FMath::DivideAndRoundUp(TileCount.X * TileCount.Y, GBokehDOFQuadsPerInstance));
 
 	Context.RHICmdList.CopyToResolveTarget(DestRenderTarget.TargetableTexture, DestRenderTarget.ShaderResourceTexture, FResolveParams());
 }
diff --git a/Engine/Source/Runtime/Renderer/Private/PostProcess/PostProcessLensBlur.cpp b/Engine/Source/Runtime/Renderer/Private/PostProcess/PostProcessLensBlur.cpp
--- a/Engine/Source/Runtime/Renderer/Private/PostProcess/PostProcessLensBlur.cpp
+++ b/Engine/Source/Runtime/Renderer/Private/PostProcess/PostProcessLensBlur.cpp
@@ -224,7 +224,7 @@ void FRCPassPostProcessLensBlur::Process(FRenderingCompositePassContext& Context
 	// needs to be the same on shader side (faster on NVIDIA and AMD)
 	int32 QuadsPerInstance = 4;
 
-	Context.RHICmdList.DrawPrimitive(PT_TriangleList, 0, 2, FMath::DivideAndRoundUp(TileCount.X * TileCount.Y, QuadsPerInstance));
+	Context.RHICmdList.DrawPrimitive(0, 2, FMath::DivideAndRoundUp(TileCount.X * TileCount.Y, QuadsPerInstance));
 
 	Context.RHICmdList.CopyToResolveTarget(DestRenderTarget.TargetableTexture, DestRenderTarget.ShaderResourceTexture, FResolveParams());
 }
diff --git a/Engine/Source/Runtime/Renderer/Private/PostProcess/PostProcessMotionBlur.cpp b/Engine/Source/Runtime/Renderer/Private/PostProcess/PostProcessMotionBlur.cpp
--- a/Engine/Source/Runtime/Renderer/Private/PostProcess/PostProcessMotionBlur.cpp
+++ b/Engine/Source/Runtime/Renderer/Private/PostProcess/PostProcessMotionBlur.cpp
@@ -385,7 +385,7 @@ void FRCPassPostProcessVelocityScatter::Process(FRenderingCompositePassContext&
 		int32 QuadsPerInstance = 8;
 
 		Context.RHICmdList.SetStreamSource(0, NULL, 0);
-		Context.RHICmdList.DrawIndexedPrimitive(GScatterQuadIndexBuffer.IndexBufferRHI, PT_TriangleList, 0, 0, 32, 0, 2 * QuadsPerInstance, FMath::DivideAndRoundUp(TileCount.X * TileCount.Y, QuadsPerInstance));
+		Context.RHICmdList.DrawIndexedPrimitive(GScatterQuadIndexBuffer.IndexBufferRHI, 0, 0, 32, 0, 2 * QuadsPerInstance, FMath::DivideAndRoundUp(TileCount.X * TileCount.Y, QuadsPerInstance));
 	}
 
 	Context.RHICmdList.CopyToResolveTarget(DestRenderTarget.TargetableTexture, DestRenderTarget.ShaderResourceTexture, FResolveParams());
diff --git a/Engine/Source/Runtime/Renderer/Private/PostProcess/SceneFilterRendering.cpp b/Engine/Source/Runtime/Renderer/Private/PostProcess/SceneFilterRendering.cpp
--- a/Engine/Source/Runtime/Renderer/Private/PostProcess/SceneFilterRendering.cpp
+++ b/Engine/Source/Runtime/Renderer/Private/PostProcess/SceneFilterRendering.cpp
@@ -204,7 +204,6 @@ static inline void InternalDrawRectangle(
 
 		RHICmdList.DrawIndexedPrimitive(
 			GTesselatedScreenRectangleIndexBuffer.IndexBufferRHI,
-			PT_TriangleList,
 			/*BaseVertexIndex=*/ 0,
 			/*MinIndex=*/ 0,
 			/*NumVertices=*/ GTesselatedScreenRectangleIndexBuffer.NumVertices(),
@@ -225,7 +224,6 @@ static inline void InternalDrawRectangle(
 			// We offset into the index buffer when using the triangle optimization to access the correct vertices.
 			RHICmdList.DrawIndexedPrimitive(
 				GScreenRectangleIndexBuffer.IndexBufferRHI,
-				PT_TriangleList,
 				/*BaseVertexIndex=*/ 0,
 				/*MinIndex=*/ 0,
 				/*NumVertices=*/ 3,
@@ -240,7 +238,6 @@ static inline void InternalDrawRectangle(
 
 			RHICmdList.DrawIndexedPrimitive(
 				GScreenRectangleIndexBuffer.IndexBufferRHI,
-				PT_TriangleList,
 				/*BaseVertexIndex=*/ 0,
 				/*MinIndex=*/ 0,
 				/*NumVertices=*/ 4,
@@ -319,7 +316,7 @@ void DrawTransformedRectangle(
 
 	RHIUnlockVertexBuffer(VertexBufferRHI);
 	RHICmdList.SetStreamSource(0, VertexBufferRHI, 0);
-	RHICmdList.DrawIndexedPrimitive(GTwoTrianglesIndexBuffer.IndexBufferRHI, PT_TriangleList, 0, 0, 4, 0, 2, 1);
+	RHICmdList.DrawIndexedPrimitive(GTwoTrianglesIndexBuffer.IndexBufferRHI, 0, 0, 4, 0, 2, 1);
 	VertexBufferRHI.SafeRelease();
 }
 
diff --git a/Engine/Source/Runtime/Renderer/Private/PostProcess/SceneRenderTargets.cpp b/Engine/Source/Runtime/Renderer/Private/PostProcess/SceneRenderTargets.cpp
--- a/Engine/Source/Runtime/Renderer/Private/PostProcess/SceneRenderTargets.cpp
+++ b/Engine/Source/Runtime/Renderer/Private/PostProcess/SceneRenderTargets.cpp
@@ -1587,7 +1587,7 @@ void FSceneRenderTargets::ResolveDepthTexture(FRHICommandList& RHICmdList, const
 
 	RHIUnlockVertexBuffer(VertexBufferRHI);
 	RHICmdList.SetStreamSource(0, VertexBufferRHI, 0);
-	RHICmdList.DrawPrimitive(PT_TriangleStrip, 0, 2, 1);
+	RHICmdList.DrawPrimitive(0, 2, 1);
 }
 void FSceneRenderTargets::ResolveSceneDepthTexture(FRHICommandList& RHICmdList, const FResolveRect& ResolveRect)
 {
diff --git a/Engine/Source/Runtime/Renderer/Private/SceneOcclusion.cpp b/Engine/Source/Runtime/Renderer/Private/SceneOcclusion.cpp
--- a/Engine/Source/Runtime/Renderer/Private/SceneOcclusion.cpp
+++ b/Engine/Source/Runtime/Renderer/Private/SceneOcclusion.cpp
@@ -365,7 +365,6 @@ void FOcclusionQueryBatcher::Flush(FRHICommandList& RHICmdList)
 			RHICmdList.SetStreamSource(0, VertexBufferRHI, VertexBufferOffset);
 			RHICmdList.DrawIndexedPrimitive(
 				IndexBufferRHI,
-				PT_TriangleList,
 				/*BaseVertexIndex=*/ 0,
 				/*MinIndex=*/ 0,
 				/*NumVertices=*/ 8 * NumPrimitivesThisBatch,
@@ -531,7 +530,7 @@ static void ExecuteDirectionalLightShadowOcclusionQuery(FRHICommandList& RHICmdL
 {
 	RHICmdList.BeginRenderQuery(ShadowOcclusionQuery);
 
-	RHICmdList.DrawPrimitive(PT_TriangleList, BaseVertexIndex, 2, 1);
+	RHICmdList.DrawPrimitive(BaseVertexIndex, 2, 1);
 	BaseVertexIndex += 6;
 
 	RHICmdList.EndRenderQuery(ShadowOcclusionQuery);
@@ -572,7 +571,7 @@ static void ExecuteProjectedShadowOcclusionQuery(FRHICommandList& RHICmdList, ui
 	// Draw the primitive's bounding box, using the occlusion query.
 	RHICmdList.BeginRenderQuery(ShadowOcclusionQuery);
 
-	RHICmdList.DrawIndexedPrimitive(GCubeIndexBuffer.IndexBufferRHI, PT_TriangleList, BaseVertexIndex, 0, 8, 0, 12, 1);
+	RHICmdList.DrawIndexedPrimitive(GCubeIndexBuffer.IndexBufferRHI, BaseVertexIndex, 0, 8, 0, 12, 1);
 	BaseVertexIndex += 8;
 
 	RHICmdList.EndRenderQuery(ShadowOcclusionQuery);
@@ -648,7 +647,7 @@ static void ExecutePlanarReflectionOcclusionQuery(FRHICommandList& RHICmdList, u
 	// Draw the primitive's bounding box, using the occlusion query.
 	RHICmdList.BeginRenderQuery(OcclusionQuery);
 
-	RHICmdList.DrawIndexedPrimitive(GCubeIndexBuffer.IndexBufferRHI, PT_TriangleList, BaseVertexIndex, 0, 8, 0, 12, 1);
+	RHICmdList.DrawIndexedPrimitive(GCubeIndexBuffer.IndexBufferRHI, BaseVertexIndex, 0, 8, 0, 12, 1);
 
 	RHICmdList.EndRenderQuery(OcclusionQuery);
 }
diff --git a/Engine/Source/Runtime/Renderer/Private/SceneRendering.cpp b/Engine/Source/Runtime/Renderer/Private/SceneRendering.cpp
--- a/Engine/Source/Runtime/Renderer/Private/SceneRendering.cpp
+++ b/Engine/Source/Runtime/Renderer/Private/SceneRendering.cpp
@@ -3604,7 +3604,7 @@ void FSceneRenderer::ResolveSceneColor(FRHICommandList& RHICmdList)
 				}
 			}
 
-				RHICmdList.DrawPrimitive(PT_TriangleList, 0, 1, 1);
+				RHICmdList.DrawPrimitive(0, 1, 1);
 			}
 		}
 
diff --git a/Engine/Source/Runtime/Renderer/Private/ShadowRendering.cpp b/Engine/Source/Runtime/Renderer/Private/ShadowRendering.cpp
--- a/Engine/Source/Runtime/Renderer/Private/ShadowRendering.cpp
+++ b/Engine/Source/Runtime/Renderer/Private/ShadowRendering.cpp
@@ -728,7 +728,7 @@ void FProjectedShadowInfo::SetupProjectionStencilMask(
 
 		GraphicsPSOInit.BoundShaderState.VertexDeclarationRHI = GetVertexDeclarationFVector4();
 		GraphicsPSOInit.BoundShaderState.VertexShaderRHI = GETSAFERHISHADER_VERTEX(*VertexShaderNoTransform);
-		GraphicsPSOInit.PrimitiveType = PT_TriangleList;
+		GraphicsPSOInit.PrimitiveType = PT_TriangleStrip;
 
 		SetGraphicsPipelineState(RHICmdList, GraphicsPSOInit);
 
@@ -762,7 +762,7 @@ void FProjectedShadowInfo::SetupProjectionStencilMask(
 
 		RHIUnlockVertexBuffer(VertexBufferRHI);
 		RHICmdList.SetStreamSource(0, VertexBufferRHI, 0);
-		RHICmdList.DrawPrimitive(PT_TriangleStrip, 0, (CascadeSettings.ShadowSplitIndex > 0) ? 4 : 2, 1);
+		RHICmdList.DrawPrimitive(0, (CascadeSettings.ShadowSplitIndex > 0) ? 4 : 2, 1);
 	}
 	// Not a preshadow, mask the projection to any pixels inside the frustum.
 	else
@@ -820,7 +820,7 @@ void FProjectedShadowInfo::SetupProjectionStencilMask(
 
 		RHICmdList.SetStreamSource(0, VertexBufferRHI, 0);
 		// Draw the frustum using the stencil buffer to mask just the pixels which are inside the shadow frustum.
-		RHICmdList.DrawIndexedPrimitive(GCubeIndexBuffer.IndexBufferRHI, PT_TriangleList, 0, 0, 8, 0, 12, 1);
+		RHICmdList.DrawIndexedPrimitive(GCubeIndexBuffer.IndexBufferRHI, 0, 0, 8, 0, 12, 1);
 		VertexBufferRHI.SafeRelease();
 
 		// if rendering modulated shadows mask out subject mesh elements to prevent self shadowing.
@@ -997,7 +997,7 @@ void FProjectedShadowInfo::RenderProjection(FRHICommandListImmediate& RHICmdList
 	if (IsWholeSceneDirectionalShadow())
 	{
 		RHICmdList.SetStreamSource(0, GClearVertexBuffer.VertexBufferRHI, 0);
-		RHICmdList.DrawPrimitive(PT_TriangleStrip, 0, 2, 1);
+		RHICmdList.DrawPrimitive(0, 2, 1);
 	}
 	else
 	{
@@ -1009,7 +1009,7 @@ void FProjectedShadowInfo::RenderProjection(FRHICommandListImmediate& RHICmdList
 
 		RHICmdList.SetStreamSource(0, VertexBufferRHI, 0);
 		// Draw the frustum using the projection shader..
-		RHICmdList.DrawIndexedPrimitive(GCubeIndexBuffer.IndexBufferRHI, PT_TriangleList, 0, 0, 8, 0, 12, 1);
+		RHICmdList.DrawIndexedPrimitive(GCubeIndexBuffer.IndexBufferRHI, 0, 0, 8, 0, 12, 1);
 		VertexBufferRHI.SafeRelease();
 	}
 
diff --git a/Engine/Source/Runtime/Renderer/Private/VT/VirtualTextureSpace.cpp b/Engine/Source/Runtime/Renderer/Private/VT/VirtualTextureSpace.cpp
--- a/Engine/Source/Runtime/Renderer/Private/VT/VirtualTextureSpace.cpp
+++ b/Engine/Source/Runtime/Renderer/Private/VT/VirtualTextureSpace.cpp
@@ -362,7 +362,7 @@ void FVirtualTextureSpace::ApplyUpdates( FRHICommandList& RHICmdList )
 			uint32 QuadsPerInstance = 8;
 
 			RHICmdList.SetStreamSource( 0, NULL, 0 );
-			RHICmdList.DrawIndexedPrimitive( GQuadIndexBuffer.IndexBufferRHI, PT_TriangleList, 0, 0, 32, 0, 2 * QuadsPerInstance, FMath::DivideAndRoundUp( NumUpdates, QuadsPerInstance ) );
+			RHICmdList.DrawIndexedPrimitive( GQuadIndexBuffer.IndexBufferRHI, 0, 0, 32, 0, 2 * QuadsPerInstance, FMath::DivideAndRoundUp( NumUpdates, QuadsPerInstance ) );
 
 			ExpandedUpdates[ Mip ].Reset();
 		}
diff --git a/Engine/Source/Runtime/Renderer/Private/VisualizeVolumetricLightmap.cpp b/Engine/Source/Runtime/Renderer/Private/VisualizeVolumetricLightmap.cpp
--- a/Engine/Source/Runtime/Renderer/Private/VisualizeVolumetricLightmap.cpp
+++ b/Engine/Source/Runtime/Renderer/Private/VisualizeVolumetricLightmap.cpp
@@ -204,7 +204,7 @@ void FDeferredShadingSceneRenderer::VisualizeVolumetricLightmap(FRHICommandListI
 			int32 NumQuads = VolumetricLightmapData->IndirectionTextureDimensions.X * VolumetricLightmapData->IndirectionTextureDimensions.Y * VolumetricLightmapData->IndirectionTextureDimensions.Z * BrickSize * BrickSize * BrickSize;
 
 			RHICmdList.SetStreamSource(0, NULL, 0);
-			RHICmdList.DrawIndexedPrimitive(GVisualizeQuadIndexBuffer.IndexBufferRHI, PT_TriangleList, 0, 0, 4 * GQuadsPerVisualizeInstance, 0, 2 * GQuadsPerVisualizeInstance, FMath::DivideAndRoundUp(NumQuads, GQuadsPerVisualizeInstance));
+			RHICmdList.DrawIndexedPrimitive(GVisualizeQuadIndexBuffer.IndexBufferRHI, 0, 0, 4 * GQuadsPerVisualizeInstance, 0, 2 * GQuadsPerVisualizeInstance, FMath::DivideAndRoundUp(NumQuads, GQuadsPerVisualizeInstance));
 		}
 	}
 }
\ No newline at end of file
diff --git a/Engine/Source/Runtime/Renderer/Private/VolumetricFog.cpp b/Engine/Source/Runtime/Renderer/Private/VolumetricFog.cpp
--- a/Engine/Source/Runtime/Renderer/Private/VolumetricFog.cpp
+++ b/Engine/Source/Runtime/Renderer/Private/VolumetricFog.cpp
@@ -624,7 +624,7 @@ void FDeferredShadingSceneRenderer::RenderLocalLightsForVolumetricFog(
 				RHICmdList.SetStreamSource(0, GCircleRasterizeVertexBuffer.VertexBufferRHI, 0);
 				const int32 NumInstances = VolumeZBounds.Y - VolumeZBounds.X;
 				const int32 NumTriangles = FCircleRasterizeVertexBuffer::NumVertices - 2;
-				RHICmdList.DrawIndexedPrimitive(GCircleRasterizeIndexBuffer.IndexBufferRHI, PT_TriangleList, 0, 0, FCircleRasterizeVertexBuffer::NumVertices, 0, NumTriangles, NumInstances);
+				RHICmdList.DrawIndexedPrimitive(GCircleRasterizeIndexBuffer.IndexBufferRHI, 0, 0, FCircleRasterizeVertexBuffer::NumVertices, 0, NumTriangles, NumInstances);
 			}
 		}
 
diff --git a/Engine/Source/Runtime/Renderer/Private/WideCustomResolveShaders.cpp b/Engine/Source/Runtime/Renderer/Private/WideCustomResolveShaders.cpp
--- a/Engine/Source/Runtime/Renderer/Private/WideCustomResolveShaders.cpp
+++ b/Engine/Source/Runtime/Renderer/Private/WideCustomResolveShaders.cpp
@@ -63,7 +63,7 @@ static void ResolveColorWideInternal2(
 	
 	PixelShader->SetParameters(RHICmdList, SrcTexture, FMaskTexture, SrcOrigin);
 
-	RHICmdList.DrawPrimitive(PT_TriangleList, 0, 1, 1);
+	RHICmdList.DrawPrimitive(0, 1, 1);
 }
 
 template <unsigned MSAA, bool UseFMask>
diff --git a/Engine/Source/Runtime/ShaderCore/Private/ShaderCache.cpp b/Engine/Source/Runtime/ShaderCore/Private/ShaderCache.cpp
--- a/Engine/Source/Runtime/ShaderCore/Private/ShaderCache.cpp
+++ b/Engine/Source/Runtime/ShaderCore/Private/ShaderCache.cpp
@@ -3317,17 +3317,17 @@ void FShaderCache::InternalPreDrawShader(FRHICommandList& RHICmdList, FShaderCac
 			{
 				case 0:
 				{
-					RHICmdList.DrawPrimitive(PT_TriangleList, 0, 1, 1);
+					RHICmdList.DrawPrimitive(0, 1, 1);
 					break;
 				}
 				case 2:
 				{
-					RHICmdList.DrawIndexedPrimitive(IndexBufferUInt16, PT_TriangleList, 0, 0, 3, 0, 1, 1);
+					RHICmdList.DrawIndexedPrimitive(IndexBufferUInt16, 0, 0, 3, 0, 1, 1);
 					break;
 				}
 				case 4:
 				{
-					RHICmdList.DrawIndexedPrimitive(IndexBufferUInt32, PT_TriangleList, 0, 0, 3, 0, 1, 1);
+					RHICmdList.DrawIndexedPrimitive(IndexBufferUInt32, 0, 0, 3, 0, 1, 1);
 					break;
 				}
 				default:
diff --git a/Engine/Source/Runtime/SlateRHIRenderer/Private/SlateRHIRenderingPolicy.cpp b/Engine/Source/Runtime/SlateRHIRenderer/Private/SlateRHIRenderingPolicy.cpp
--- a/Engine/Source/Runtime/SlateRHIRenderer/Private/SlateRHIRenderingPolicy.cpp
+++ b/Engine/Source/Runtime/SlateRHIRenderer/Private/SlateRHIRenderingPolicy.cpp
@@ -464,7 +464,7 @@ static void UpdateScissorRect(
 
 						RHIUnlockVertexBuffer(VertexBufferRHI);
 						RHICmdList.SetStreamSource(0, VertexBufferRHI, 0);
-						RHICmdList.DrawPrimitive(PT_TriangleStrip, 0, 2, 1);
+						RHICmdList.DrawPrimitive(0, 2, 1);
 					}
 
 					// Now setup the pipeline to use SO_SaturatedIncrement, since we've established the initial
@@ -516,7 +516,7 @@ static void UpdateScissorRect(
 
 					RHIUnlockVertexBuffer(VertexBufferRHI);
 					RHICmdList.SetStreamSource(0, VertexBufferRHI, 0);
-					RHICmdList.DrawPrimitive(PT_TriangleStrip, 0, 2, 1);
+					RHICmdList.DrawPrimitive(0, 2, 1);
 				}
 
 				// Setup the stenciling state to be read only now, disable depth writes, and restore the color buffer
@@ -972,13 +972,13 @@ void FSlateRHIRenderingPolicy::DrawElements(
 					if (!GRHISupportsBaseVertexIndex && !bAbsoluteIndices)
 					{
 						RHICmdList.SetStreamSource(0, VertexBuffer->VertexBufferRHI, RenderBatch.VertexOffset * sizeof(FSlateVertex));
-						RHICmdList.DrawIndexedPrimitive(IndexBuffer->IndexBufferRHI, GetRHIPrimitiveType(RenderBatch.DrawPrimitiveType), 0, 0, RenderBatch.NumVertices, RenderBatch.IndexOffset, PrimitiveCount, RenderBatch.InstanceCount);
+						RHICmdList.DrawIndexedPrimitive(IndexBuffer->IndexBufferRHI, 0, 0, RenderBatch.NumVertices, RenderBatch.IndexOffset, PrimitiveCount, RenderBatch.InstanceCount);
 					}
 					else
 					{
 						uint32 VertexOffset = bAbsoluteIndices ? 0 : RenderBatch.VertexOffset;
 						RHICmdList.SetStreamSource(0, VertexBuffer->VertexBufferRHI, 0);
-						RHICmdList.DrawIndexedPrimitive(IndexBuffer->IndexBufferRHI, GetRHIPrimitiveType(RenderBatch.DrawPrimitiveType), VertexOffset, 0, RenderBatch.NumVertices, RenderBatch.IndexOffset, PrimitiveCount, RenderBatch.InstanceCount);
+						RHICmdList.DrawIndexedPrimitive(IndexBuffer->IndexBufferRHI, VertexOffset, 0, RenderBatch.NumVertices, RenderBatch.IndexOffset, PrimitiveCount, RenderBatch.InstanceCount);
 					}
 				}
 			}
@@ -1111,13 +1111,13 @@ void FSlateRHIRenderingPolicy::DrawElements(
 									if (!GRHISupportsBaseVertexIndex && !bAbsoluteIndices)
 									{
 										RHICmdList.SetStreamSource(0, VertexBuffer->VertexBufferRHI, RenderBatch.VertexOffset * sizeof(FSlateVertex));
-										RHICmdList.DrawIndexedPrimitive(IndexBuffer->IndexBufferRHI, GetRHIPrimitiveType(RenderBatch.DrawPrimitiveType), 0, 0, RenderBatch.NumVertices, RenderBatch.IndexOffset, PrimitiveCount, InstanceCount);
+										RHICmdList.DrawIndexedPrimitive(IndexBuffer->IndexBufferRHI, 0, 0, RenderBatch.NumVertices, RenderBatch.IndexOffset, PrimitiveCount, InstanceCount);
 									}
 									else
 									{
 										uint32 VertexOffset = bAbsoluteIndices ? 0 : RenderBatch.VertexOffset;
 										RHICmdList.SetStreamSource(0, VertexBuffer->VertexBufferRHI, 0);
-										RHICmdList.DrawIndexedPrimitive(IndexBuffer->IndexBufferRHI, GetRHIPrimitiveType(RenderBatch.DrawPrimitiveType), VertexOffset, 0, RenderBatch.NumVertices, RenderBatch.IndexOffset, PrimitiveCount, InstanceCount);
+										RHICmdList.DrawIndexedPrimitive(IndexBuffer->IndexBufferRHI, VertexOffset, 0, RenderBatch.NumVertices, RenderBatch.IndexOffset, PrimitiveCount, InstanceCount);
 									}
 								}
 							}
@@ -1129,13 +1129,13 @@ void FSlateRHIRenderingPolicy::DrawElements(
 								if (!GRHISupportsBaseVertexIndex && !bAbsoluteIndices)
 								{
 									RHICmdList.SetStreamSource(0, VertexBuffer->VertexBufferRHI, RenderBatch.VertexOffset * sizeof(FSlateVertex));
-									RHICmdList.DrawIndexedPrimitive(IndexBuffer->IndexBufferRHI, GetRHIPrimitiveType(RenderBatch.DrawPrimitiveType), 0, 0, RenderBatch.NumVertices, RenderBatch.IndexOffset, PrimitiveCount, 1);
+									RHICmdList.DrawIndexedPrimitive(IndexBuffer->IndexBufferRHI, 0, 0, RenderBatch.NumVertices, RenderBatch.IndexOffset, PrimitiveCount, 1);
 								}
 								else
 								{
 									uint32 VertexOffset = bAbsoluteIndices ? 0 : RenderBatch.VertexOffset;
 									RHICmdList.SetStreamSource(0, VertexBuffer->VertexBufferRHI, 0);
-									RHICmdList.DrawIndexedPrimitive(IndexBuffer->IndexBufferRHI, GetRHIPrimitiveType(RenderBatch.DrawPrimitiveType), VertexOffset, 0, RenderBatch.NumVertices, RenderBatch.IndexOffset, PrimitiveCount, 1);
+									RHICmdList.DrawIndexedPrimitive(IndexBuffer->IndexBufferRHI, VertexOffset, 0, RenderBatch.NumVertices, RenderBatch.IndexOffset, PrimitiveCount, 1);
 								}
 							}
 						}
diff --git a/Engine/Source/Runtime/UtilityShaders/Private/ClearQuad.cpp b/Engine/Source/Runtime/UtilityShaders/Private/ClearQuad.cpp
--- a/Engine/Source/Runtime/UtilityShaders/Private/ClearQuad.cpp
+++ b/Engine/Source/Runtime/UtilityShaders/Private/ClearQuad.cpp
@@ -260,7 +260,7 @@ void DrawClearQuadMRT(FRHICommandList& RHICmdList, bool bClearColor, int32 NumCl
 	ClearQuadSetup(RHICmdList, bClearColor, NumClearColors, ClearColorArray, bClearDepth, Depth, bClearStencil, Stencil);
 
 	RHICmdList.SetStreamSource(0, GClearVertexBuffer.VertexBufferRHI, 0);
-	RHICmdList.DrawPrimitive(PT_TriangleStrip, 0, 2, 1);
+	RHICmdList.DrawPrimitive(0, 2, 1);
 }
 
 void DrawClearQuadMRT(FRHICommandList& RHICmdList, bool bClearColor, int32 NumClearColors, const FLinearColor* ClearColorArray, bool bClearDepth, float Depth, bool bClearStencil, uint32 Stencil, FIntPoint ViewSize, FIntRect ExcludeRect)
@@ -312,7 +312,7 @@ void DrawClearQuadMRT(FRHICommandList& RHICmdList, bool bClearColor, int32 NumCl
 		RHIUnlockVertexBuffer(VertexBufferRHI);
 		RHICmdList.SetStreamSource(0, VertexBufferRHI, 0);
 
-		RHICmdList.DrawPrimitive(PT_TriangleStrip, 0, 8, 1);
+		RHICmdList.DrawPrimitive(0, 8, 1);
 
 		VertexBufferRHI.SafeRelease();
 	}
@@ -320,6 +320,6 @@ void DrawClearQuadMRT(FRHICommandList& RHICmdList, bool bClearColor, int32 NumCl
 	{
 		// without a hole
 		RHICmdList.SetStreamSource(0, GClearVertexBuffer.VertexBufferRHI, 0);
-		RHICmdList.DrawPrimitive(PT_TriangleStrip, 0, 2, 1);
+		RHICmdList.DrawPrimitive(0, 2, 1);
 	}
 }
diff --git a/Engine/Source/Runtime/UtilityShaders/Private/LongGPUTask.cpp b/Engine/Source/Runtime/UtilityShaders/Private/LongGPUTask.cpp
--- a/Engine/Source/Runtime/UtilityShaders/Private/LongGPUTask.cpp
+++ b/Engine/Source/Runtime/UtilityShaders/Private/LongGPUTask.cpp
@@ -72,7 +72,7 @@ void IssueScalableLongGPUTask(FRHICommandListImmediate& RHICmdList, int32 NumIte
 
 	for (int32 Iteration = 0; Iteration < NumIteration; Iteration++)
 	{
-		RHICmdList.DrawPrimitive(PT_TriangleStrip, 0, 2, 1);
+		RHICmdList.DrawPrimitive(0, 2, 1);
 	}
 }
 
diff --git a/Engine/Source/Runtime/VulkanRHI/Private/VulkanCommands.cpp b/Engine/Source/Runtime/VulkanRHI/Private/VulkanCommands.cpp
--- a/Engine/Source/Runtime/VulkanRHI/Private/VulkanCommands.cpp
+++ b/Engine/Source/Runtime/VulkanRHI/Private/VulkanCommands.cpp
@@ -771,7 +771,7 @@ void FVulkanCommandListContext::RHISetStencilRef(uint32 StencilRef)
 	PendingGfxState->SetStencilRef(StencilRef);
 }
 
-void FVulkanCommandListContext::RHIDrawPrimitive(uint32 PrimitiveType_DEPRECATED, uint32 BaseVertexIndex, uint32 NumPrimitives, uint32 NumInstances)
+void FVulkanCommandListContext::RHIDrawPrimitive(uint32 BaseVertexIndex, uint32 NumPrimitives, uint32 NumInstances)
 {
 #if VULKAN_ENABLE_AGGRESSIVE_STATS
 	SCOPE_CYCLE_COUNTER(STAT_VulkanDrawCallTime);
@@ -790,7 +790,7 @@ void FVulkanCommandListContext::RHIDrawPrimitive(uint32 PrimitiveType_DEPRECATED
 	}
 }
 
-void FVulkanCommandListContext::RHIDrawPrimitiveIndirect(uint32 PrimitiveType_DEPRECATED, FVertexBufferRHIParamRef ArgumentBufferRHI, uint32 ArgumentOffset)
+void FVulkanCommandListContext::RHIDrawPrimitiveIndirect(FVertexBufferRHIParamRef ArgumentBufferRHI, uint32 ArgumentOffset)
 {
 	static_assert(sizeof(FRHIDrawIndirectParameters) == sizeof(VkDrawIndirectCommand), "Draw indirect doesn't match!");
 
@@ -813,7 +813,7 @@ void FVulkanCommandListContext::RHIDrawPrimitiveIndirect(uint32 PrimitiveType_DE
 	}
 }
 
-void FVulkanCommandListContext::RHIDrawIndexedPrimitive(FIndexBufferRHIParamRef IndexBufferRHI, uint32 PrimitiveType_DEPRECATED, int32 BaseVertexIndex, uint32 FirstInstance,
+void FVulkanCommandListContext::RHIDrawIndexedPrimitive(FIndexBufferRHIParamRef IndexBufferRHI, int32 BaseVertexIndex, uint32 FirstInstance,
 	uint32 NumVertices, uint32 StartIndex, uint32 NumPrimitives, uint32 NumInstances)
 {
 #if VULKAN_ENABLE_AGGRESSIVE_STATS
@@ -838,7 +838,7 @@ void FVulkanCommandListContext::RHIDrawIndexedPrimitive(FIndexBufferRHIParamRef
 	}
 }
 
-void FVulkanCommandListContext::RHIDrawIndexedIndirect(FIndexBufferRHIParamRef IndexBufferRHI, uint32 PrimitiveType_DEPRECATED, FStructuredBufferRHIParamRef ArgumentsBufferRHI, int32 DrawArgumentsIndex, uint32 NumInstances)
+void FVulkanCommandListContext::RHIDrawIndexedIndirect(FIndexBufferRHIParamRef IndexBufferRHI, FStructuredBufferRHIParamRef ArgumentsBufferRHI, int32 DrawArgumentsIndex, uint32 NumInstances)
 {
 #if VULKAN_ENABLE_AGGRESSIVE_STATS
 	SCOPE_CYCLE_COUNTER(STAT_VulkanDrawCallTime);
@@ -860,7 +860,7 @@ void FVulkanCommandListContext::RHIDrawIndexedIndirect(FIndexBufferRHIParamRef I
 	}
 }
 
-void FVulkanCommandListContext::RHIDrawIndexedPrimitiveIndirect(uint32 PrimitiveType_DEPRECATED,FIndexBufferRHIParamRef IndexBufferRHI,FVertexBufferRHIParamRef ArgumentBufferRHI,uint32 ArgumentOffset)
+void FVulkanCommandListContext::RHIDrawIndexedPrimitiveIndirect(FIndexBufferRHIParamRef IndexBufferRHI,FVertexBufferRHIParamRef ArgumentBufferRHI,uint32 ArgumentOffset)
 {
 #if VULKAN_ENABLE_AGGRESSIVE_STATS
 	SCOPE_CYCLE_COUNTER(STAT_VulkanDrawCallTime);
@@ -883,7 +883,7 @@ void FVulkanCommandListContext::RHIDrawIndexedPrimitiveIndirect(uint32 Primitive
 	}
 }
 
-void FVulkanCommandListContext::RHIBeginDrawPrimitiveUP(uint32 PrimitiveType_DEPRECATED, uint32 NumPrimitives, uint32 NumVertices, uint32 VertexDataStride, void*& OutVertexData)
+void FVulkanCommandListContext::RHIBeginDrawPrimitiveUP(uint32 NumPrimitives, uint32 NumVertices, uint32 VertexDataStride, void*& OutVertexData)
 {
 #if VULKAN_ENABLE_AGGRESSIVE_STATS
 	SCOPE_CYCLE_COUNTER(STAT_VulkanUPPrepTime);
@@ -918,7 +918,7 @@ void FVulkanCommandListContext::RHIEndDrawPrimitiveUP()
 	}
 }
 
-void FVulkanCommandListContext::RHIBeginDrawIndexedPrimitiveUP( uint32 PrimitiveType_DEPRECATED, uint32 NumPrimitives, uint32 NumVertices, uint32 VertexDataStride,
+void FVulkanCommandListContext::RHIBeginDrawIndexedPrimitiveUP(uint32 NumPrimitives, uint32 NumVertices, uint32 VertexDataStride,
 	void*& OutVertexData, uint32 MinVertexIndex, uint32 NumIndices, uint32 IndexDataStride, void*& OutIndexData)
 {
 	LLM_SCOPE_VULKAN(ELLMTagVulkan::VulkanMisc);
diff --git a/Engine/Source/Runtime/VulkanRHI/Public/VulkanContext.h b/Engine/Source/Runtime/VulkanRHI/Public/VulkanContext.h
--- a/Engine/Source/Runtime/VulkanRHI/Public/VulkanContext.h
+++ b/Engine/Source/Runtime/VulkanRHI/Public/VulkanContext.h
@@ -212,14 +212,14 @@ public:
 	virtual void RHISetBlendState(FBlendStateRHIParamRef NewState, const FLinearColor& BlendFactor) final override;
 	virtual void RHISetRenderTargets(uint32 NumSimultaneousRenderTargets, const FRHIRenderTargetView* NewRenderTargets, const FRHIDepthRenderTargetView* NewDepthStencilTarget, uint32 NumUAVs, const FUnorderedAccessViewRHIParamRef* UAVs) final override;
 	virtual void RHISetRenderTargetsAndClear(const FRHISetRenderTargetsInfo& RenderTargetsInfo) final override;
-	virtual void RHIDrawPrimitive(uint32 PrimitiveType, uint32 BaseVertexIndex, uint32 NumPrimitives, uint32 NumInstances) final override;
-	virtual void RHIDrawPrimitiveIndirect(uint32 PrimitiveType, FVertexBufferRHIParamRef ArgumentBuffer, uint32 ArgumentOffset) final override;
-	virtual void RHIDrawIndexedIndirect(FIndexBufferRHIParamRef IndexBufferRHI, uint32 PrimitiveType, FStructuredBufferRHIParamRef ArgumentsBufferRHI, int32 DrawArgumentsIndex, uint32 NumInstances) final override;
-	virtual void RHIDrawIndexedPrimitive(FIndexBufferRHIParamRef IndexBuffer, uint32 PrimitiveType, int32 BaseVertexIndex, uint32 FirstInstance, uint32 NumVertices, uint32 StartIndex, uint32 NumPrimitives, uint32 NumInstances) final override;
-	virtual void RHIDrawIndexedPrimitiveIndirect(uint32 PrimitiveType, FIndexBufferRHIParamRef IndexBuffer, FVertexBufferRHIParamRef ArgumentBuffer, uint32 ArgumentOffset) final override;
-	virtual void RHIBeginDrawPrimitiveUP(uint32 PrimitiveType, uint32 NumPrimitives, uint32 NumVertices, uint32 VertexDataStride, void*& OutVertexData) final override;
+	virtual void RHIDrawPrimitive(uint32 BaseVertexIndex, uint32 NumPrimitives, uint32 NumInstances) final override;
+	virtual void RHIDrawPrimitiveIndirect(FVertexBufferRHIParamRef ArgumentBuffer, uint32 ArgumentOffset) final override;
+	virtual void RHIDrawIndexedIndirect(FIndexBufferRHIParamRef IndexBufferRHI, FStructuredBufferRHIParamRef ArgumentsBufferRHI, int32 DrawArgumentsIndex, uint32 NumInstances) final override;
+	virtual void RHIDrawIndexedPrimitive(FIndexBufferRHIParamRef IndexBuffer, int32 BaseVertexIndex, uint32 FirstInstance, uint32 NumVertices, uint32 StartIndex, uint32 NumPrimitives, uint32 NumInstances) final override;
+	virtual void RHIDrawIndexedPrimitiveIndirect(FIndexBufferRHIParamRef IndexBuffer, FVertexBufferRHIParamRef ArgumentBuffer, uint32 ArgumentOffset) final override;
+	virtual void RHIBeginDrawPrimitiveUP(uint32 NumPrimitives, uint32 NumVertices, uint32 VertexDataStride, void*& OutVertexData) final override;
 	virtual void RHIEndDrawPrimitiveUP() final override;
-	virtual void RHIBeginDrawIndexedPrimitiveUP(uint32 PrimitiveType, uint32 NumPrimitives, uint32 NumVertices, uint32 VertexDataStride, void*& OutVertexData, uint32 MinVertexIndex, uint32 NumIndices, uint32 IndexDataStride, void*& OutIndexData) final override;
+	virtual void RHIBeginDrawIndexedPrimitiveUP(uint32 NumPrimitives, uint32 NumVertices, uint32 VertexDataStride, void*& OutVertexData, uint32 MinVertexIndex, uint32 NumIndices, uint32 IndexDataStride, void*& OutIndexData) final override;
 	virtual void RHIEndDrawIndexedPrimitiveUP() final override;
 	virtual void RHIEnableDepthBoundsTest(bool bEnable) final override;
 	virtual void RHISetDepthBounds(float MinDepth, float MaxDepth) final override;
diff --git a/Engine/Source/Runtime/Windows/D3D11RHI/Private/D3D11Commands.cpp b/Engine/Source/Runtime/Windows/D3D11RHI/Private/D3D11Commands.cpp
--- a/Engine/Source/Runtime/Windows/D3D11RHI/Private/D3D11Commands.cpp
+++ b/Engine/Source/Runtime/Windows/D3D11RHI/Private/D3D11Commands.cpp
@@ -1275,10 +1275,6 @@ static D3D11_PRIMITIVE_TOPOLOGY GetD3D11PrimitiveType(EPrimitiveType PrimitiveTy
 	return D3D11_PRIMITIVE_TOPOLOGY_TRIANGLELIST;
 }
 
-static inline void VerifyPrimitiveType(uint32 DrawCallPrimitiveType, EPrimitiveType PSOPrimitiveType)
-{
-	ensureAlwaysMsgf((EPrimitiveType)DrawCallPrimitiveType == PSOPrimitiveType, TEXT("PSO was created using PrimitiveType %d, but the Draw call is using %d! PrimitiveTypes must match as deprecation is imminent..."), (uint32)PSOPrimitiveType, (uint32)DrawCallPrimitiveType);
-}
 
 void FD3D11DynamicRHI::CommitNonComputeShaderConstants()
 {
@@ -1582,7 +1578,7 @@ void FD3D11DynamicRHI::CommitComputeResourceTables(FD3D11ComputeShader* InComput
 	SetResourcesFromTables(ComputeShader);
 }
 
-void FD3D11DynamicRHI::RHIDrawPrimitive(uint32 PrimitiveType_DEPRECATED,uint32 BaseVertexIndex,uint32 NumPrimitives,uint32 NumInstances)
+void FD3D11DynamicRHI::RHIDrawPrimitive(uint32 BaseVertexIndex,uint32 NumPrimitives,uint32 NumInstances)
 {
 	RHI_DRAW_CALL_STATS(PrimitiveType,NumInstances*NumPrimitives);
 
@@ -1592,7 +1588,6 @@ void FD3D11DynamicRHI::RHIDrawPrimitive(uint32 PrimitiveType_DEPRECATED,uint32 B
 	uint32 VertexCount = GetVertexCountForPrimitiveCount(NumPrimitives,PrimitiveType);
 
 	GPUProfilingData.RegisterGPUWork(NumPrimitives * NumInstances, VertexCount * NumInstances);
-	VerifyPrimitiveType(PrimitiveType_DEPRECATED, PrimitiveType);
 	StateCache.SetPrimitiveTopology(GetD3D11PrimitiveType(PrimitiveType,bUsingTessellation));
 	if(NumInstances > 1)
 	{
@@ -1604,7 +1599,7 @@ void FD3D11DynamicRHI::RHIDrawPrimitive(uint32 PrimitiveType_DEPRECATED,uint32 B
 	}
 }
 
-void FD3D11DynamicRHI::RHIDrawPrimitiveIndirect(uint32 PrimitiveType_DEPRECATED,FVertexBufferRHIParamRef ArgumentBufferRHI,uint32 ArgumentOffset)
+void FD3D11DynamicRHI::RHIDrawPrimitiveIndirect(FVertexBufferRHIParamRef ArgumentBufferRHI,uint32 ArgumentOffset)
 {
 	FD3D11VertexBuffer* ArgumentBuffer = ResourceCast(ArgumentBufferRHI);
 
@@ -1615,12 +1610,11 @@ void FD3D11DynamicRHI::RHIDrawPrimitiveIndirect(uint32 PrimitiveType_DEPRECATED,
 	CommitGraphicsResourceTables();
 	CommitNonComputeShaderConstants();
 
-	VerifyPrimitiveType(PrimitiveType_DEPRECATED, PrimitiveType);
 	StateCache.SetPrimitiveTopology(GetD3D11PrimitiveType(PrimitiveType,bUsingTessellation));
 	Direct3DDeviceIMContext->DrawInstancedIndirect(ArgumentBuffer->Resource,ArgumentOffset);
 }
 
-void FD3D11DynamicRHI::RHIDrawIndexedIndirect(FIndexBufferRHIParamRef IndexBufferRHI, uint32 PrimitiveType_DEPRECATED, FStructuredBufferRHIParamRef ArgumentsBufferRHI, int32 DrawArgumentsIndex, uint32 NumInstances)
+void FD3D11DynamicRHI::RHIDrawIndexedIndirect(FIndexBufferRHIParamRef IndexBufferRHI, FStructuredBufferRHIParamRef ArgumentsBufferRHI, int32 DrawArgumentsIndex, uint32 NumInstances)
 {
 	FD3D11IndexBuffer* IndexBuffer = ResourceCast(IndexBufferRHI);
 	FD3D11StructuredBuffer* ArgumentsBuffer = ResourceCast(ArgumentsBufferRHI);
@@ -1637,7 +1631,6 @@ void FD3D11DynamicRHI::RHIDrawIndexedIndirect(FIndexBufferRHIParamRef IndexBuffe
 	const DXGI_FORMAT Format = (IndexBuffer->GetStride() == sizeof(uint16) ? DXGI_FORMAT_R16_UINT : DXGI_FORMAT_R32_UINT);
 
 	StateCache.SetIndexBuffer(IndexBuffer->Resource, Format, 0);
-	VerifyPrimitiveType(PrimitiveType_DEPRECATED, PrimitiveType);
 	StateCache.SetPrimitiveTopology(GetD3D11PrimitiveType(PrimitiveType,bUsingTessellation));
 
 	if(NumInstances > 1)
@@ -1650,7 +1643,7 @@ void FD3D11DynamicRHI::RHIDrawIndexedIndirect(FIndexBufferRHIParamRef IndexBuffe
 	}
 }
 
-void FD3D11DynamicRHI::RHIDrawIndexedPrimitive(FIndexBufferRHIParamRef IndexBufferRHI,uint32 PrimitiveType_DEPRECATED,int32 BaseVertexIndex,uint32 FirstInstance,uint32 NumVertices,uint32 StartIndex,uint32 NumPrimitives,uint32 NumInstances)
+void FD3D11DynamicRHI::RHIDrawIndexedPrimitive(FIndexBufferRHIParamRef IndexBufferRHI,int32 BaseVertexIndex,uint32 FirstInstance,uint32 NumVertices,uint32 StartIndex,uint32 NumPrimitives,uint32 NumInstances)
 {
 	FD3D11IndexBuffer* IndexBuffer = ResourceCast(IndexBufferRHI);
 
@@ -1676,7 +1669,6 @@ void FD3D11DynamicRHI::RHIDrawIndexedPrimitive(FIndexBufferRHIParamRef IndexBuff
 		TEXT("Start %u, Count %u, Type %u, Buffer Size %u, Buffer stride %u"), StartIndex, IndexCount, PrimitiveType, IndexBuffer->GetSize(), IndexBuffer->GetStride());
 
 	StateCache.SetIndexBuffer(IndexBuffer->Resource, Format, 0);
-	VerifyPrimitiveType(PrimitiveType_DEPRECATED, PrimitiveType);
 	StateCache.SetPrimitiveTopology(GetD3D11PrimitiveType(PrimitiveType,bUsingTessellation));
 
 	if (NumInstances > 1 || FirstInstance != 0)
@@ -1691,7 +1683,7 @@ void FD3D11DynamicRHI::RHIDrawIndexedPrimitive(FIndexBufferRHIParamRef IndexBuff
 	}
 }
 
-void FD3D11DynamicRHI::RHIDrawIndexedPrimitiveIndirect(uint32 PrimitiveType_DEPRECATED,FIndexBufferRHIParamRef IndexBufferRHI,FVertexBufferRHIParamRef ArgumentBufferRHI,uint32 ArgumentOffset)
+void FD3D11DynamicRHI::RHIDrawIndexedPrimitiveIndirect(FIndexBufferRHIParamRef IndexBufferRHI,FVertexBufferRHIParamRef ArgumentBufferRHI,uint32 ArgumentOffset)
 {
 	FD3D11IndexBuffer* IndexBuffer = ResourceCast(IndexBufferRHI);
 	FD3D11VertexBuffer* ArgumentBuffer = ResourceCast(ArgumentBufferRHI);
@@ -1707,7 +1699,6 @@ void FD3D11DynamicRHI::RHIDrawIndexedPrimitiveIndirect(uint32 PrimitiveType_DEPR
 	const uint32 SizeFormat = sizeof(DXGI_FORMAT);
 	const DXGI_FORMAT Format = (IndexBuffer->GetStride() == sizeof(uint16) ? DXGI_FORMAT_R16_UINT : DXGI_FORMAT_R32_UINT);
 	StateCache.SetIndexBuffer(IndexBuffer->Resource, Format, 0);
-	VerifyPrimitiveType(PrimitiveType_DEPRECATED, PrimitiveType);
 	StateCache.SetPrimitiveTopology(GetD3D11PrimitiveType(PrimitiveType,bUsingTessellation));
 	Direct3DDeviceIMContext->DrawIndexedInstancedIndirect(ArgumentBuffer->Resource,ArgumentOffset);
 }
@@ -1720,12 +1711,10 @@ void FD3D11DynamicRHI::RHIDrawIndexedPrimitiveIndirect(uint32 PrimitiveType_DEPR
  * @param VertexDataStride Size of each vertex 
  * @param OutVertexData Reference to the allocated vertex memory
  */
-void FD3D11DynamicRHI::RHIBeginDrawPrimitiveUP( uint32 PrimitiveType_DEPRECATED, uint32 NumPrimitives, uint32 NumVertices, uint32 VertexDataStride, void*& OutVertexData)
+void FD3D11DynamicRHI::RHIBeginDrawPrimitiveUP(uint32 NumPrimitives, uint32 NumVertices, uint32 VertexDataStride, void*& OutVertexData)
 {
 	checkSlow( PendingNumVertices == 0 );
 
-	VerifyPrimitiveType(PrimitiveType_DEPRECATED, PrimitiveType);
-
 	// Remember the parameters for the draw call.
 	PendingNumPrimitives = NumPrimitives;
 	PendingNumVertices = NumVertices;
@@ -1775,12 +1764,10 @@ void FD3D11DynamicRHI::RHIEndDrawPrimitiveUP()
  * @param IndexDataStride Size of each index (either 2 or 4 bytes)
  * @param OutIndexData Reference to the allocated index memory
  */
-void FD3D11DynamicRHI::RHIBeginDrawIndexedPrimitiveUP( uint32 PrimitiveType_DEPRECATED, uint32 NumPrimitives, uint32 NumVertices, uint32 VertexDataStride, void*& OutVertexData, uint32 MinVertexIndex, uint32 NumIndices, uint32 IndexDataStride, void*& OutIndexData)
+void FD3D11DynamicRHI::RHIBeginDrawIndexedPrimitiveUP(uint32 NumPrimitives, uint32 NumVertices, uint32 VertexDataStride, void*& OutVertexData, uint32 MinVertexIndex, uint32 NumIndices, uint32 IndexDataStride, void*& OutIndexData)
 {
 	checkSlow((sizeof(uint16) == IndexDataStride) || (sizeof(uint32) == IndexDataStride));
 
-	VerifyPrimitiveType(PrimitiveType_DEPRECATED, PrimitiveType);
-
 	// Store off information needed for the draw call.
 	PendingNumPrimitives = NumPrimitives;
 	PendingMinVertexIndex = MinVertexIndex;
diff --git a/Engine/Source/Runtime/Windows/D3D11RHI/Private/D3D11RHIPrivate.h b/Engine/Source/Runtime/Windows/D3D11RHI/Private/D3D11RHIPrivate.h
--- a/Engine/Source/Runtime/Windows/D3D11RHI/Private/D3D11RHIPrivate.h
+++ b/Engine/Source/Runtime/Windows/D3D11RHI/Private/D3D11RHIPrivate.h
@@ -551,14 +551,14 @@ public:
 	virtual void RHISetRenderTargets(uint32 NumSimultaneousRenderTargets, const FRHIRenderTargetView* NewRenderTargets, const FRHIDepthRenderTargetView* NewDepthStencilTarget, uint32 NumUAVs, const FUnorderedAccessViewRHIParamRef* UAVs) final override;
 	virtual void RHISetRenderTargetsAndClear(const FRHISetRenderTargetsInfo& RenderTargetsInfo) final override;
 	virtual void RHIBindClearMRTValues(bool bClearColor, bool bClearDepth, bool bClearStencil) final override;
-	virtual void RHIDrawPrimitive(uint32 PrimitiveType, uint32 BaseVertexIndex, uint32 NumPrimitives, uint32 NumInstances) final override;
-	virtual void RHIDrawPrimitiveIndirect(uint32 PrimitiveType, FVertexBufferRHIParamRef ArgumentBuffer, uint32 ArgumentOffset) final override;
-	virtual void RHIDrawIndexedIndirect(FIndexBufferRHIParamRef IndexBufferRHI, uint32 PrimitiveType, FStructuredBufferRHIParamRef ArgumentsBufferRHI, int32 DrawArgumentsIndex, uint32 NumInstances) final override;
-	virtual void RHIDrawIndexedPrimitive(FIndexBufferRHIParamRef IndexBuffer, uint32 PrimitiveType, int32 BaseVertexIndex, uint32 FirstInstance, uint32 NumVertices, uint32 StartIndex, uint32 NumPrimitives, uint32 NumInstances) final override;
-	virtual void RHIDrawIndexedPrimitiveIndirect(uint32 PrimitiveType, FIndexBufferRHIParamRef IndexBuffer, FVertexBufferRHIParamRef ArgumentBuffer, uint32 ArgumentOffset) final override;
-	virtual void RHIBeginDrawPrimitiveUP(uint32 PrimitiveType, uint32 NumPrimitives, uint32 NumVertices, uint32 VertexDataStride, void*& OutVertexData) final override;
+	virtual void RHIDrawPrimitive(uint32 BaseVertexIndex, uint32 NumPrimitives, uint32 NumInstances) final override;
+	virtual void RHIDrawPrimitiveIndirect(FVertexBufferRHIParamRef ArgumentBuffer, uint32 ArgumentOffset) final override;
+	virtual void RHIDrawIndexedIndirect(FIndexBufferRHIParamRef IndexBufferRHI, FStructuredBufferRHIParamRef ArgumentsBufferRHI, int32 DrawArgumentsIndex, uint32 NumInstances) final override;
+	virtual void RHIDrawIndexedPrimitive(FIndexBufferRHIParamRef IndexBuffer, int32 BaseVertexIndex, uint32 FirstInstance, uint32 NumVertices, uint32 StartIndex, uint32 NumPrimitives, uint32 NumInstances) final override;
+	virtual void RHIDrawIndexedPrimitiveIndirect(FIndexBufferRHIParamRef IndexBuffer, FVertexBufferRHIParamRef ArgumentBuffer, uint32 ArgumentOffset) final override;
+	virtual void RHIBeginDrawPrimitiveUP(uint32 NumPrimitives, uint32 NumVertices, uint32 VertexDataStride, void*& OutVertexData) final override;
 	virtual void RHIEndDrawPrimitiveUP() final override;
-	virtual void RHIBeginDrawIndexedPrimitiveUP(uint32 PrimitiveType, uint32 NumPrimitives, uint32 NumVertices, uint32 VertexDataStride, void*& OutVertexData, uint32 MinVertexIndex, uint32 NumIndices, uint32 IndexDataStride, void*& OutIndexData) final override;
+	virtual void RHIBeginDrawIndexedPrimitiveUP(uint32 NumPrimitives, uint32 NumVertices, uint32 VertexDataStride, void*& OutVertexData, uint32 MinVertexIndex, uint32 NumIndices, uint32 IndexDataStride, void*& OutIndexData) final override;
 	virtual void RHIEndDrawIndexedPrimitiveUP() final override;
 	virtual void RHIEnableDepthBoundsTest(bool bEnable) final override
 	{
diff --git a/Engine/Source/Runtime/Windows/D3D11RHI/Private/D3D11RenderTarget.cpp b/Engine/Source/Runtime/Windows/D3D11RHI/Private/D3D11RenderTarget.cpp
--- a/Engine/Source/Runtime/Windows/D3D11RHI/Private/D3D11RenderTarget.cpp
+++ b/Engine/Source/Runtime/Windows/D3D11RHI/Private/D3D11RenderTarget.cpp
@@ -190,7 +190,7 @@ void FD3D11DynamicRHI::ResolveTextureUsingShader(
 
 	RHIUnlockVertexBuffer(VertexBufferRHI);
 	RHICmdList.SetStreamSource(0, VertexBufferRHI, 0);
-	RHICmdList.DrawPrimitive(PT_TriangleStrip, 0, 2, 1);
+	RHICmdList.DrawPrimitive(0, 2, 1);
 
 	RHICmdList.Flush(); // always call flush when using a command list in RHI implementations before doing anything else. This is super hazardous.
 
