GPUFamily5 supports layered rendering: https://developer.apple.com/metal/Metal-Feature-Set-Tables.pdf

diff --git a/Engine/Source/Runtime/Apple/MetalRHI/Private/MetalCommandQueue.cpp b/Engine/Source/Runtime/Apple/MetalRHI/Private/MetalCommandQueue.cpp
--- a/Engine/Source/Runtime/Apple/MetalRHI/Private/MetalCommandQueue.cpp
+++ b/Engine/Source/Runtime/Apple/MetalRHI/Private/MetalCommandQueue.cpp
@@ -122,6 +122,10 @@ FMetalCommandQueue::FMetalCommandQueue(mtlpp::Device InDevice, uint32 const MaxN
 				{
 					Features |= EMetalFeaturesTextureBuffers;
 				}
+				if ([Device supportsFeatureSet:MTLFeatureSet_iOS_GPUFamily5_v1])
+				{
+					Features |= EMetalFeaturesLayeredRendering;
+				}
 			}
         }
 		
diff --git a/Engine/Source/Runtime/Apple/MetalRHI/Private/MetalRHI.cpp b/Engine/Source/Runtime/Apple/MetalRHI/Private/MetalRHI.cpp
--- a/Engine/Source/Runtime/Apple/MetalRHI/Private/MetalRHI.cpp
+++ b/Engine/Source/Runtime/Apple/MetalRHI/Private/MetalRHI.cpp
@@ -143,6 +143,9 @@ FMetalDynamicRHI::FMetalDynamicRHI(ERHIFeatureLevel::Type RequestedFeatureLevel)
 			GRHIDeviceId++;
 		}
 	}
+		
+	GSupportsVolumeTextureRendering = FMetalCommandQueue::SupportsFeature(EMetalFeaturesLayeredRendering);
+	bSupportsPointLights = GSupportsVolumeTextureRendering;
 #endif
 
     bool bProjectSupportsMRTs = false;
diff --git a/Engine/Source/Runtime/RHI/Public/RHI.h b/Engine/Source/Runtime/RHI/Public/RHI.h
--- a/Engine/Source/Runtime/RHI/Public/RHI.h
+++ b/Engine/Source/Runtime/RHI/Public/RHI.h
@@ -154,6 +154,11 @@ inline bool RHISupportsVolumeTextures(ERHIFeatureLevel::Type FeatureLevel)
 	return FeatureLevel >= ERHIFeatureLevel::SM4;
 }
 
+inline bool RHISupportsVertexShaderLayer(const EShaderPlatform Platform)
+{
+	return IsFeatureLevelSupported(Platform, ERHIFeatureLevel::SM4) && IsMetalPlatform(Platform) && (IsPCPlatform(Platform) || (Platform == SP_METAL_MRT && RHIGetShaderLanguageVersion(Platform) >= 4));
+}
+
 inline bool RHISupports4ComponentUAVReadWrite(EShaderPlatform Platform)
 {
 	// Must match usf PLATFORM_SUPPORTS_4COMPONENT_UAV_READ_WRITE
diff --git a/Engine/Source/Runtime/RHI/Public/RHIDefinitions.h b/Engine/Source/Runtime/RHI/Public/RHIDefinitions.h
--- a/Engine/Source/Runtime/RHI/Public/RHIDefinitions.h
+++ b/Engine/Source/Runtime/RHI/Public/RHIDefinitions.h
@@ -1020,11 +1020,6 @@ inline bool RHIHasTiledGPU(const EShaderPlatform Platform)
 	return Platform == SP_METAL || Platform == SP_OPENGL_ES2_IOS || Platform == SP_OPENGL_ES2_ANDROID || Platform == SP_OPENGL_ES3_1_ANDROID;
 }
 
-inline bool RHISupportsVertexShaderLayer(const EShaderPlatform Platform)
-{
-	return IsFeatureLevelSupported(Platform, ERHIFeatureLevel::SM4) && IsPCPlatform(Platform) && IsMetalPlatform(Platform);
-}
-
 inline bool RHISupportsMobileMultiView(const EShaderPlatform Platform)
 {
 	return (Platform == EShaderPlatform::SP_OPENGL_ES3_1_ANDROID || Platform == EShaderPlatform::SP_OPENGL_ES2_ANDROID);
