Copy 4462892
Windows - allow preferring d3d12 on launch. If preferred, will prefer d3d12 to d3d11 but has no effect when non-D3D RHI is chosen
Windows - Added logic to use d3d12 device profile on Windows when appropriate

diff --git a/Engine/Config/Windows/WindowsEngine.ini b/Engine/Config/Windows/WindowsEngine.ini
--- a/Engine/Config/Windows/WindowsEngine.ini
+++ b/Engine/Config/Windows/WindowsEngine.ini
@@ -20,4 +20,7 @@ ConfigClass=/Engine/DeviceProfiles/SetDeviceProfileWindows.SetDeviceProfileWindo
 r.setres=1280x720
 
 [PlatformCrypto]
-PlatformRequiresDataCrypto=True
\ No newline at end of file
+PlatformRequiresDataCrypto=True
+
+[DeviceProfileManager]
+DeviceProfileSelectionModule="WindowsDeviceProfileSelector"
\ No newline at end of file
diff --git a/Engine/Plugins/Runtime/WindowsDeviceProfileSelector/Source/WindowsDeviceProfileSelector/Private/WindowsDeviceProfileSelectorModule.cpp b/Engine/Plugins/Runtime/WindowsDeviceProfileSelector/Source/WindowsDeviceProfileSelector/Private/WindowsDeviceProfileSelectorModule.cpp
--- a/Engine/Plugins/Runtime/WindowsDeviceProfileSelector/Source/WindowsDeviceProfileSelector/Private/WindowsDeviceProfileSelectorModule.cpp
+++ b/Engine/Plugins/Runtime/WindowsDeviceProfileSelector/Source/WindowsDeviceProfileSelector/Private/WindowsDeviceProfileSelectorModule.cpp
@@ -2,6 +2,7 @@
 
 #include "WindowsDeviceProfileSelectorModule.h"
 #include "Modules/ModuleManager.h"
+#include "DynamicRHI.h"
 
 
 IMPLEMENT_MODULE(FWindowsDeviceProfileSelectorModule, WindowsDeviceProfileSelector);
@@ -18,7 +19,25 @@ void FWindowsDeviceProfileSelectorModule::ShutdownModule()
 
 const FString FWindowsDeviceProfileSelectorModule::GetRuntimeDeviceProfileName()
 {
+	// Windows, WindowsNoEditor, WindowsClient, or WindowsServer
 	FString ProfileName = FPlatformProperties::PlatformName();
+
+	if (FApp::CanEverRender())
+	{
+		FString DeviceProfileFileName;
+		FConfigCacheIni::LoadGlobalIniFile(DeviceProfileFileName, TEXT("DeviceProfiles"));
+
+		TArray<FString> AvailableProfiles;
+		GConfig->GetSectionNames(DeviceProfileFileName, AvailableProfiles);
+
+		FString TmpProfileName = ProfileName + TCHAR('_') + GetSelectedDynamicRHIModuleName(false);
+		if (AvailableProfiles.Contains(FString::Printf(TEXT("%s DeviceProfile"), *TmpProfileName)))
+		{
+			// Use RHI specific device profile if exist
+			ProfileName = MoveTemp(TmpProfileName);
+		}
+	}
+
 	UE_LOG(LogInit, Log, TEXT("Selected Device Profile: [%s]"), *ProfileName);
 	return ProfileName;
 }
diff --git a/Engine/Plugins/Runtime/WindowsDeviceProfileSelector/Source/WindowsDeviceProfileSelector/WindowsDeviceProfileSelector.Build.cs b/Engine/Plugins/Runtime/WindowsDeviceProfileSelector/Source/WindowsDeviceProfileSelector/WindowsDeviceProfileSelector.Build.cs
--- a/Engine/Plugins/Runtime/WindowsDeviceProfileSelector/Source/WindowsDeviceProfileSelector/WindowsDeviceProfileSelector.Build.cs
+++ b/Engine/Plugins/Runtime/WindowsDeviceProfileSelector/Source/WindowsDeviceProfileSelector/WindowsDeviceProfileSelector.Build.cs
@@ -29,6 +29,7 @@ namespace UnrealBuildTool.Rules
 				    "Core",
 				    "CoreUObject",
 				    "Engine",
+                    "RHI",
 				}
 				);
 
diff --git a/Engine/Source/Runtime/D3D12RHI/Private/Windows/WindowsD3D12PipelineState.cpp b/Engine/Source/Runtime/D3D12RHI/Private/Windows/WindowsD3D12PipelineState.cpp
--- a/Engine/Source/Runtime/D3D12RHI/Private/Windows/WindowsD3D12PipelineState.cpp
+++ b/Engine/Source/Runtime/D3D12RHI/Private/Windows/WindowsD3D12PipelineState.cpp
@@ -9,9 +9,10 @@
 #include "Misc/ScopeRWLock.h"
 #include "Stats/StatsMisc.h"
 
+// D3D12RHI PSO file cache doesn't work anymore. Use FPipelineFileCache instead
 static TAutoConsoleVariable<int32> CVarPipelineStateDiskCache(
 	TEXT("D3D12.PSO.DiskCache"),
-	1,
+	0,
 	TEXT("Enables a disk cache for Pipeline State Objects (PSOs).\n")
 	TEXT("PSO descs are cached to disk so subsequent runs can create PSOs at load-time instead of at run-time.\n")
 	TEXT("This cache contains data that is independent of hardware, driver, or machine that it was created on. It can be distributed with shipping content.\n")
@@ -21,7 +22,7 @@ static TAutoConsoleVariable<int32> CVarPipelineStateDiskCache(
 
 static TAutoConsoleVariable<int32> CVarDriverOptimizedPipelineStateDiskCache(
 	TEXT("D3D12.PSO.DriverOptimizedDiskCache"),
-	1,
+	0,
 	TEXT("Enables a disk cache for driver-optimized Pipeline State Objects (PSOs).\n")
 	TEXT("PSO descs are cached to disk so subsequent runs can create PSOs at load-time instead of at run-time.\n")
 	TEXT("This cache contains data specific to the hardware, driver, and machine that it was created on.\n")
diff --git a/Engine/Source/Runtime/RHI/Private/Windows/WindowsDynamicRHI.cpp b/Engine/Source/Runtime/RHI/Private/Windows/WindowsDynamicRHI.cpp
--- a/Engine/Source/Runtime/RHI/Private/Windows/WindowsDynamicRHI.cpp
+++ b/Engine/Source/Runtime/RHI/Private/Windows/WindowsDynamicRHI.cpp
@@ -6,34 +6,45 @@
 #include "Misc/ConfigCacheIni.h"
 #include "Misc/MessageDialog.h"
 
-FDynamicRHI* PlatformCreateDynamicRHI()
-{
-	FDynamicRHI* DynamicRHI = nullptr;
 
-#if UE_BUILD_DEBUG || UE_BUILD_DEVELOPMENT
-	if (!FPlatformMisc::IsDebuggerPresent())
+static const TCHAR* GLoadedRHIModuleName;
+
+static bool ShouldPreferD3D12()
+{
+	bool bPreferD3D12 = false;
+	if (GIsEditor)
 	{
-		if (FParse::Param(FCommandLine::Get(), TEXT("AttachDebugger")))
-		{
-			// Wait to attach debugger
-			do
-			{
-				FPlatformProcess::Sleep(0);
-			}
-			while (!FPlatformMisc::IsDebuggerPresent());
-		}
+		GConfig->GetBool(TEXT("D3DRHIPerference"), TEXT("bPreferD3D12InEditor"), bPreferD3D12, GEngineIni);
 	}
-#endif
+	else
+	{
+		GConfig->GetBool(TEXT("D3DRHIPerference"), TEXT("bPreferD3D12InGame"), bPreferD3D12, GEngineIni);
+	}
+
+	int32 MinNumCPUCores = 0;
+	GConfig->GetInt(TEXT("D3DRHIPerference"), TEXT("con.MinNumCPUCores"), MinNumCPUCores, GEngineIni);
+	const bool bHasEnoughCPUCores = FPlatformMisc::NumberOfCoresIncludingHyperthreads() >= MinNumCPUCores;
+
+	int32 MinPhysicalMemGB = 0;
+	GConfig->GetInt(TEXT("D3DRHIPerference"), TEXT("con.MinPhysicalMemGB"), MinPhysicalMemGB, GEngineIni);
+	const bool bHasEnoughMem = FPlatformMemory::GetConstants().TotalPhysical >= MinPhysicalMemGB * (1llu << 30);
+
+	return bPreferD3D12 && bHasEnoughCPUCores && bHasEnoughMem;
+}
 
+static IDynamicRHIModule* LoadDynamicRHIModule(ERHIFeatureLevel::Type& DesiredFeatureLevel, const TCHAR*& LoadedRHIModuleName)
+{
+	bool bPreferD3D12 = ShouldPreferD3D12();
+	
 	// command line overrides
-	bool bForceSM5  = FParse::Param(FCommandLine::Get(),TEXT("sm5"));
-	bool bForceSM4  = FParse::Param(FCommandLine::Get(), TEXT("sm4"));
+	bool bForceSM5 = FParse::Param(FCommandLine::Get(), TEXT("sm5"));
+	bool bForceSM4 = FParse::Param(FCommandLine::Get(), TEXT("sm4"));
 	bool bForceVulkan = FParse::Param(FCommandLine::Get(), TEXT("vulkan"));
 	bool bForceOpenGL = FWindowsPlatformMisc::VerifyWindowsVersion(6, 0) == false || FParse::Param(FCommandLine::Get(), TEXT("opengl")) || FParse::Param(FCommandLine::Get(), TEXT("opengl3")) || FParse::Param(FCommandLine::Get(), TEXT("opengl4"));
-	bool bForceD3D10 = FParse::Param(FCommandLine::Get(),TEXT("d3d10")) || FParse::Param(FCommandLine::Get(),TEXT("dx10")) || (bForceSM4 && !bForceVulkan && !bForceOpenGL);
-	bool bForceD3D11 = FParse::Param(FCommandLine::Get(),TEXT("d3d11")) || FParse::Param(FCommandLine::Get(),TEXT("dx11")) || (bForceSM5 && !bForceVulkan && !bForceOpenGL);
+	bool bForceD3D10 = FParse::Param(FCommandLine::Get(), TEXT("d3d10")) || FParse::Param(FCommandLine::Get(), TEXT("dx10")) || (bForceSM4 && !bForceVulkan && !bForceOpenGL);
+	bool bForceD3D11 = FParse::Param(FCommandLine::Get(), TEXT("d3d11")) || FParse::Param(FCommandLine::Get(), TEXT("dx11")) || (bForceSM5 && !bForceVulkan && !bForceOpenGL);
 	bool bForceD3D12 = FParse::Param(FCommandLine::Get(), TEXT("d3d12")) || FParse::Param(FCommandLine::Get(), TEXT("dx12"));
-	ERHIFeatureLevel::Type RequestedFeatureLevel = ERHIFeatureLevel::Num;
+	DesiredFeatureLevel = ERHIFeatureLevel::Num;
 	int32 Sum = ((bForceD3D12 ? 1 : 0) + (bForceD3D11 ? 1 : 0) + (bForceD3D10 ? 1 : 0) + (bForceOpenGL ? 1 : 0) + (bForceVulkan ? 1 : 0));
 
 	if (bForceSM5 && bForceSM4)
@@ -43,7 +54,7 @@ FDynamicRHI* PlatformCreateDynamicRHI()
 
 	if (Sum > 1)
 	{
-		UE_LOG(LogRHI, Fatal,TEXT("-d3d12, -d3d11, -d3d10, -vulkan, and -opengl[3|4] are mutually exclusive options, but more than one was specified on the command-line."));
+		UE_LOG(LogRHI, Fatal, TEXT("-d3d12, -d3d11, -d3d10, -vulkan, and -opengl[3|4] are mutually exclusive options, but more than one was specified on the command-line."));
 	}
 	else if (Sum == 0)
 	{
@@ -56,21 +67,22 @@ FDynamicRHI* PlatformCreateDynamicRHI()
 			FName ShaderFormatName(*TargetedShaderFormats[0]);
 			EShaderPlatform TargetedPlatform = ShaderFormatToLegacyShaderPlatform(ShaderFormatName);
 			bForceVulkan = IsVulkanPlatform(TargetedPlatform);
-			bForceD3D11 = IsD3DPlatform(TargetedPlatform, false);
+			bForceD3D11 = !bPreferD3D12 && IsD3DPlatform(TargetedPlatform, false);
 			bForceOpenGL = IsOpenGLPlatform(TargetedPlatform);
-			RequestedFeatureLevel = GetMaxSupportedFeatureLevel(TargetedPlatform);
+			DesiredFeatureLevel = GetMaxSupportedFeatureLevel(TargetedPlatform);
 		}
 	}
 	else
 	{
 		if (bForceSM5)
 		{
-			RequestedFeatureLevel = ERHIFeatureLevel::SM5;
+			DesiredFeatureLevel = ERHIFeatureLevel::SM5;
 		}
 
 		if (bForceSM4)
 		{
-			RequestedFeatureLevel = ERHIFeatureLevel::SM4;
+			DesiredFeatureLevel = ERHIFeatureLevel::SM4;
+			bPreferD3D12 = false;
 		}
 	}
 
@@ -82,23 +94,26 @@ FDynamicRHI* PlatformCreateDynamicRHI()
 	// Load the dynamic RHI module.
 	if (bForceSwitch)
 	{
-		#define A(x) #x
-		#define B(x) A(x)
-		#define SWITCH_RHI_STR B(SWITCHRHI)
-		DynamicRHIModule = &FModuleManager::LoadModuleChecked<IDynamicRHIModule>(TEXT( SWITCH_RHI_STR ));
+#define A(x) #x
+#define B(x) A(x)
+#define SWITCH_RHI_STR B(SWITCHRHI)
+		const TCHAR* SwitchRHIModuleName = TEXT(SWITCH_RHI_STR);
+		DynamicRHIModule = &FModuleManager::LoadModuleChecked<IDynamicRHIModule>(SwitchRHIModuleName);
 		if (!DynamicRHIModule->IsSupported())
 		{
 			FMessageDialog::Open(EAppMsgType::Ok, NSLOCTEXT("SwitchDynamicRHI", "UnsupportedRHI", "The chosen RHI is not supported"));
 			FPlatformMisc::RequestExit(1);
 			DynamicRHIModule = NULL;
 		}
+		LoadedRHIModuleName = SwitchRHIModuleName;
 	}
 	else
 #endif
 
 	if (bForceOpenGL)
 	{
-		DynamicRHIModule = &FModuleManager::LoadModuleChecked<IDynamicRHIModule>(TEXT("OpenGLDrv"));
+		const TCHAR* OpenGLRHIModuleName = TEXT("OpenGLDrv");
+		DynamicRHIModule = &FModuleManager::LoadModuleChecked<IDynamicRHIModule>(OpenGLRHIModuleName);
 
 		if (!DynamicRHIModule->IsSupported())
 		{
@@ -106,35 +121,50 @@ FDynamicRHI* PlatformCreateDynamicRHI()
 			FPlatformMisc::RequestExit(1);
 			DynamicRHIModule = NULL;
 		}
+		LoadedRHIModuleName = OpenGLRHIModuleName;
 	}
 	else if (bForceVulkan)
 	{
-		DynamicRHIModule = &FModuleManager::LoadModuleChecked<IDynamicRHIModule>(TEXT("VulkanRHI"));
+		const TCHAR* VulkanRHIModuleName = TEXT("VulkanRHI");
+		DynamicRHIModule = &FModuleManager::LoadModuleChecked<IDynamicRHIModule>(VulkanRHIModuleName);
 		if (!DynamicRHIModule->IsSupported())
 		{
 			FMessageDialog::Open(EAppMsgType::Ok, NSLOCTEXT("WindowsDynamicRHI", "RequiredVulkan", "Vulkan Driver is required to run the engine."));
 			FPlatformMisc::RequestExit(1);
 			DynamicRHIModule = NULL;
 		}
+		LoadedRHIModuleName = VulkanRHIModuleName;
 	}
-	else if (bForceD3D12)
+	else if (bForceD3D12 || bPreferD3D12)
 	{
-		DynamicRHIModule = &FModuleManager::LoadModuleChecked<IDynamicRHIModule>(TEXT("D3D12RHI"));
+		LoadedRHIModuleName = TEXT("D3D12RHI");
+		DynamicRHIModule = FModuleManager::LoadModulePtr<IDynamicRHIModule>(LoadedRHIModuleName);
 
-		if (!DynamicRHIModule->IsSupported())
+		if (!DynamicRHIModule || !DynamicRHIModule->IsSupported())
 		{
-			FMessageDialog::Open(EAppMsgType::Ok, NSLOCTEXT("WindowsDynamicRHI", "RequiredDX12", "DX12 is not supported on your system. Try running without the -dx12 or -d3d12 command line argument."));
-			FPlatformMisc::RequestExit(1);
+			if (bForceD3D12)
+			{
+				FMessageDialog::Open(EAppMsgType::Ok, NSLOCTEXT("WindowsDynamicRHI", "RequiredDX12", "DX12 is not supported on your system. Try running without the -dx12 or -d3d12 command line argument."));
+				FPlatformMisc::RequestExit(1);
+			}
+			if (DynamicRHIModule)
+			{
+				FModuleManager::Get().UnloadModule(LoadedRHIModuleName);
+			}
 			DynamicRHIModule = NULL;
+			LoadedRHIModuleName = nullptr;
 		}
 		else if (FPlatformProcess::IsApplicationRunning(TEXT("fraps.exe")))
 		{
 			FMessageDialog::Open(EAppMsgType::Ok, NSLOCTEXT("WindowsDynamicRHI", "UseExpressionEncoder", "Fraps has been known to crash D3D12. Please use Microsoft Expression Encoder instead for capturing."));
 		}
 	}
-	else
+
+	// Fallback to D3D11RHI if nothing is selected
+	if (!DynamicRHIModule)
 	{
-		DynamicRHIModule = &FModuleManager::LoadModuleChecked<IDynamicRHIModule>(TEXT("D3D11RHI"));
+		const TCHAR* D3D11RHIModuleName = TEXT("D3D11RHI");
+		DynamicRHIModule = &FModuleManager::LoadModuleChecked<IDynamicRHIModule>(D3D11RHIModuleName);
 
 		if (!DynamicRHIModule->IsSupported())
 		{
@@ -146,13 +176,63 @@ FDynamicRHI* PlatformCreateDynamicRHI()
 		{
 			FMessageDialog::Open(EAppMsgType::Ok, NSLOCTEXT("WindowsDynamicRHI", "UseExpressionEncoderDX11", "Fraps has been known to crash D3D11. Please use Microsoft Expression Encoder instead for capturing."));
 		}
+		LoadedRHIModuleName = D3D11RHIModuleName;
+	}
+	return DynamicRHIModule;
+}
+
+FDynamicRHI* PlatformCreateDynamicRHI()
+{
+	FDynamicRHI* DynamicRHI = nullptr;
+
+#if UE_BUILD_DEBUG || UE_BUILD_DEVELOPMENT
+	if (!FPlatformMisc::IsDebuggerPresent())
+	{
+		if (FParse::Param(FCommandLine::Get(), TEXT("AttachDebugger")))
+		{
+			// Wait to attach debugger
+			do
+			{
+				FPlatformProcess::Sleep(0);
+			}
+			while (!FPlatformMisc::IsDebuggerPresent());
+		}
 	}
+#endif
+
+	ERHIFeatureLevel::Type RequestedFeatureLevel;
+	const TCHAR* LoadedRHIModuleName;
+	IDynamicRHIModule* DynamicRHIModule = LoadDynamicRHIModule(RequestedFeatureLevel, LoadedRHIModuleName);
 
 	if (DynamicRHIModule)
 	{
 		// Create the dynamic RHI.
 		DynamicRHI = DynamicRHIModule->CreateRHI(RequestedFeatureLevel);
+		GLoadedRHIModuleName = LoadedRHIModuleName;
 	}
 
 	return DynamicRHI;
 }
+
+const TCHAR* GetSelectedDynamicRHIModuleName(bool bCleanup)
+{
+	check(FApp::CanEverRender());
+	if (GDynamicRHI)
+	{
+		check(!!GLoadedRHIModuleName);
+		return GLoadedRHIModuleName;
+	}
+	else
+	{
+		ERHIFeatureLevel::Type DesiredFeatureLevel;
+		const TCHAR* RHIModuleName;
+		IDynamicRHIModule* DynamicRHIModule = LoadDynamicRHIModule(DesiredFeatureLevel, RHIModuleName);
+		check(DynamicRHIModule);
+		check(RHIModuleName);
+		if (bCleanup)
+		{
+			FModuleManager::Get().UnloadModule(RHIModuleName);
+		}
+		return RHIModuleName;
+	}
+}
diff --git a/Engine/Source/Runtime/RHI/Public/DynamicRHI.h b/Engine/Source/Runtime/RHI/Public/DynamicRHI.h
--- a/Engine/Source/Runtime/RHI/Public/DynamicRHI.h
+++ b/Engine/Source/Runtime/RHI/Public/DynamicRHI.h
@@ -1271,3 +1271,7 @@ public:
 *	Called to create the instance of the dynamic RHI.
 */
 FDynamicRHI* PlatformCreateDynamicRHI();
+
+// Name of the RHI module that will be created when PlatformCreateDynamicRHI is called
+// NOTE: This function is very slow when called before RHIInit
+extern RHI_API const TCHAR* GetSelectedDynamicRHIModuleName(bool bCleanup = true);
\ No newline at end of file
