RHIUtilities and API updates for Renderpasses.

Deprecates SetRenderTargets and friends.
Adds UnbindRenderTarget and TransitionRenderPassTargets.

diff --git a/Engine/Source/Runtime/RHI/Public/RHIResources.h b/Engine/Source/Runtime/RHI/Public/RHIResources.h
--- a/Engine/Source/Runtime/RHI/Public/RHIResources.h
+++ b/Engine/Source/Runtime/RHI/Public/RHIResources.h
@@ -2039,6 +2039,11 @@ struct FRHIRenderPassInfo
 	// Some RHIs need to know if this render pass is going to be reading and writing to the same texture in the case of generating mip maps for partial resource transitions
 	bool bGeneratingMips = false;
 
+	//#RenderPasses
+	int32 UAVIndex = -1;
+	int32 NumUAVs = 0;
+	FUnorderedAccessViewRHIRef UAVs[MaxSimultaneousUAVs];
+
 	// Color, no depth, optional resolve, optional mip, optional array slice
 	explicit FRHIRenderPassInfo(FRHITexture* ColorRT, ERenderTargetActions ColorAction, FRHITexture* ResolveRT = nullptr, uint32 InMipIndex = 0, int32 InArraySlice = -1)
 	{
@@ -2051,6 +2056,7 @@ struct FRHIRenderPassInfo
 		DepthStencilRenderTarget.DepthStencilTarget = nullptr;
 		DepthStencilRenderTarget.Action = EDepthStencilTargetActions::DontLoad_DontStore;
 		DepthStencilRenderTarget.ExclusiveDepthStencil = FExclusiveDepthStencil::DepthNop_StencilNop;
+		DepthStencilRenderTarget.ResolveTarget = nullptr;
 		bIsMSAA = ColorRT->GetNumSamples() > 1;
 		FMemory::Memzero(&ColorRenderTargets[1], sizeof(FColorEntry) * (MaxSimultaneousRenderTargets - 1));
 	}
@@ -2071,6 +2077,7 @@ struct FRHIRenderPassInfo
 		DepthStencilRenderTarget.DepthStencilTarget = nullptr;
 		DepthStencilRenderTarget.Action = EDepthStencilTargetActions::DontLoad_DontStore;
 		DepthStencilRenderTarget.ExclusiveDepthStencil = FExclusiveDepthStencil::DepthNop_StencilNop;
+		DepthStencilRenderTarget.ResolveTarget = nullptr;
 		if (NumColorRTs < MaxSimultaneousRenderTargets)
 		{
 			FMemory::Memzero(&ColorRenderTargets[NumColorRTs], sizeof(FColorEntry) * (MaxSimultaneousRenderTargets - NumColorRTs));
@@ -2093,6 +2100,7 @@ struct FRHIRenderPassInfo
 		DepthStencilRenderTarget.DepthStencilTarget = nullptr;
 		DepthStencilRenderTarget.Action = EDepthStencilTargetActions::DontLoad_DontStore;
 		DepthStencilRenderTarget.ExclusiveDepthStencil = FExclusiveDepthStencil::DepthNop_StencilNop;
+		DepthStencilRenderTarget.ResolveTarget = nullptr;
 		if (NumColorRTs < MaxSimultaneousRenderTargets)
 		{
 			FMemory::Memzero(&ColorRenderTargets[NumColorRTs], sizeof(FColorEntry) * (MaxSimultaneousRenderTargets - NumColorRTs));
@@ -2212,6 +2220,16 @@ struct FRHIRenderPassInfo
 		FMemory::Memzero(&ColorRenderTargets[1], sizeof(FColorEntry) * (MaxSimultaneousRenderTargets - 1));
 	}
 
+	explicit FRHIRenderPassInfo(int32 InNumUAVs, FUnorderedAccessViewRHIParamRef* InUAVs)
+	{
+		FMemory::Memzero(*this);
+		NumUAVs = InNumUAVs;
+		for (int32 Index = 0; Index < InNumUAVs; Index++)
+		{
+			UAVs[Index] = InUAVs[Index];
+		}
+	}
+
 	inline int32 GetNumColorRenderTargets() const
 	{
 		int32 ColorIndex = 0;
@@ -2240,11 +2258,6 @@ struct FRHIRenderPassInfo
 	RHI_API void Validate() const;
 	RHI_API void ConvertToRenderTargetsInfo(FRHISetRenderTargetsInfo& OutRTInfo) const;
 
-	//#RenderPasses
-	int32 UAVIndex = -1;
-	int32 NumUAVs = 0;
-	FUnorderedAccessViewRHIRef UAVs[MaxSimultaneousUAVs];
-
 	FRHIRenderPassInfo& operator = (const FRHIRenderPassInfo& In)
 	{
 		FMemory::Memcpy(*this, In);
diff --git a/Engine/Source/Runtime/RHI/Public/RHIUtilities.h b/Engine/Source/Runtime/RHI/Public/RHIUtilities.h
--- a/Engine/Source/Runtime/RHI/Public/RHIUtilities.h
+++ b/Engine/Source/Runtime/RHI/Public/RHIUtilities.h
@@ -392,7 +392,33 @@ inline void TransitionSetRenderTargetsHelper(FRHICommandList& RHICmdList, uint32
 	RHICmdList.TransitionResources(EResourceTransitionAccess::EWritable, Transitions, TransitionIndex);
 }
 
+inline void TransitionRenderPassTargets(FRHICommandList& RHICmdList, const FRHIRenderPassInfo& RPInfo)
+{
+	FTextureRHIParamRef Transitions[MaxSimultaneousRenderTargets + 1];
+	int32 TransitionIndex = 0;
+	uint32 NumColorRenderTargets = RPInfo.GetNumColorRenderTargets();
+	for (uint32 Index = 0; Index < NumColorRenderTargets; Index++)
+	{
+		const FRHIRenderPassInfo::FColorEntry& ColorRenderTarget = RPInfo.ColorRenderTargets[Index];
+		if (ColorRenderTarget.RenderTarget != nullptr)
+		{
+			Transitions[TransitionIndex] = ColorRenderTarget.RenderTarget;
+			TransitionIndex++;
+		}
+	}
+
+	const FRHIRenderPassInfo::FDepthStencilEntry& DepthStencilTarget = RPInfo.DepthStencilRenderTarget;
+	if (DepthStencilTarget.DepthStencilTarget != nullptr && RPInfo.DepthStencilRenderTarget.ExclusiveDepthStencil.IsDepthWrite())
+	{
+		Transitions[TransitionIndex] = DepthStencilTarget.DepthStencilTarget;
+		TransitionIndex++;
+	}
+
+	RHICmdList.TransitionResources(EResourceTransitionAccess::EWritable, Transitions, TransitionIndex);
+}
+
 /** Helper for the common case of using a single color and depth render target. */
+DEPRECATED(4.22, "Utility SetRenderTarget (Single Color, DepthStencil) is deprecated. Convert new code to use Render Passes.");
 inline void SetRenderTarget(FRHICommandList& RHICmdList, FTextureRHIParamRef NewRenderTarget, FTextureRHIParamRef NewDepthStencilTarget, bool bWritableBarrier = false)
 {
 	FRHIRenderTargetView RTV(NewRenderTarget, ERenderTargetLoadAction::ELoad);
@@ -407,6 +433,7 @@ inline void SetRenderTarget(FRHICommandList& RHICmdList, FTextureRHIParamRef New
 }
 
 /** Helper for the common case of using a single color and depth render target. */
+DEPRECATED(4.22, "Utility SetRenderTarget (Single Color, DepthStencil, Simple Mode, DepthStencilAcess) is deprecated. Convert new code to use Render Passes.");
 inline void SetRenderTarget(FRHICommandList& RHICmdList, FTextureRHIParamRef NewRenderTarget, FTextureRHIParamRef NewDepthStencilTarget, ESimpleRenderTargetMode Mode, FExclusiveDepthStencil DepthStencilAccess = FExclusiveDepthStencil::DepthWrite_StencilWrite, bool bWritableBarrier = false)
 {
 	ERenderTargetLoadAction ColorLoadAction, DepthLoadAction, StencilLoadAction;
@@ -426,6 +453,7 @@ inline void SetRenderTarget(FRHICommandList& RHICmdList, FTextureRHIParamRef New
 }
 
 /** Helper for the common case of using a single color and depth render target, with a mip index for the color target. */
+DEPRECATED(4.22, "Utility SetRenderTarget (Single Color, Mip, DepthStencil) is deprecated. Convert new code to use Render Passes.");
 inline void SetRenderTarget(FRHICommandList& RHICmdList, FTextureRHIParamRef NewRenderTarget, int32 MipIndex, FTextureRHIParamRef NewDepthStencilTarget, bool bWritableBarrier = false)
 {
 	FRHIRenderTargetView RTV(NewRenderTarget, ERenderTargetLoadAction::ELoad, MipIndex, -1);
@@ -440,6 +468,7 @@ inline void SetRenderTarget(FRHICommandList& RHICmdList, FTextureRHIParamRef New
 }
 
 /** Helper for the common case of using a single color and depth render target, with a mip index for the color target. */
+DEPRECATED(4.22, "Utility SetRenderTarget (Single Color, Mip, Slice, DepthStencil) is deprecated. Convert new code to use Render Passes.");
 inline void SetRenderTarget(FRHICommandList& RHICmdList, FTextureRHIParamRef NewRenderTarget, int32 MipIndex, int32 ArraySliceIndex, FTextureRHIParamRef NewDepthStencilTarget, bool bWritableBarrier = false)
 {
 	FRHIRenderTargetView RTV(NewRenderTarget, ERenderTargetLoadAction::ELoad, MipIndex, ArraySliceIndex);
@@ -454,6 +483,7 @@ inline void SetRenderTarget(FRHICommandList& RHICmdList, FTextureRHIParamRef New
 }
 
 /** Helper that converts FTextureRHIParamRef's into FRHIRenderTargetView's. */
+DEPRECATED(4.22, "Utility SetRenderTargets (Colors, DepthStencil, UAVs) is deprecated. Convert new code to use Render Passes.");
 inline void SetRenderTargets(
 	FRHICommandList& RHICmdList,
 	uint32 NewNumSimultaneousRenderTargets, 
@@ -482,6 +512,7 @@ inline void SetRenderTargets(
 }
 
 /** Helper that converts FTextureRHIParamRef's into FRHIRenderTargetView's. */
+DEPRECATED(4.22, "Utility SetRenderTargets (Colors, DepthStencil, UAVs, Simple Mode, DepthStencil Access) is deprecated. Convert new code to use Render Passes.");
 inline void SetRenderTargets(
 	FRHICommandList& RHICmdList,
 	uint32 NewNumSimultaneousRenderTargets,
@@ -513,6 +544,14 @@ inline void SetRenderTargets(
 	RHICmdList.SetRenderTargets(NewNumSimultaneousRenderTargets, RTVs, &DepthRTV, 0, nullptr);
 }
 
+DEPRECATED(4.22, "Utility to unset render targets. Should not be used in new code - use Render Passes instead");
+inline void UnbindRenderTargets(FRHICommandList& RHICmdList)
+{
+	check(RHICmdList.IsOutsideRenderPass());
+
+	SetRenderTarget(RHICmdList, nullptr, nullptr);
+}
+
 /**
  * Creates 1 or 2 textures with the same dimensions/format.
  * If the RHI supports textures that can be used as both shader resources and render targets,
