Fix "missing declaration" compile error due to extraneous semi-colons.

diff --git a/Engine/Source/Runtime/RHI/Public/RHIUtilities.h b/Engine/Source/Runtime/RHI/Public/RHIUtilities.h
--- a/Engine/Source/Runtime/RHI/Public/RHIUtilities.h
+++ b/Engine/Source/Runtime/RHI/Public/RHIUtilities.h
@@ -418,7 +418,7 @@ inline void TransitionRenderPassTargets(FRHICommandList& RHICmdList, const FRHIR
 }
 
 /** Helper for the common case of using a single color and depth render target. */
-DEPRECATED(4.22, "Utility SetRenderTarget (Single Color, DepthStencil) is deprecated. Convert new code to use Render Passes.");
+DEPRECATED(4.22, "Utility SetRenderTarget (Single Color, DepthStencil) is deprecated. Convert new code to use Render Passes.")
 inline void SetRenderTarget(FRHICommandList& RHICmdList, FTextureRHIParamRef NewRenderTarget, FTextureRHIParamRef NewDepthStencilTarget, bool bWritableBarrier = false)
 {
 	FRHIRenderTargetView RTV(NewRenderTarget, ERenderTargetLoadAction::ELoad);
@@ -433,7 +433,7 @@ inline void SetRenderTarget(FRHICommandList& RHICmdList, FTextureRHIParamRef New
 }
 
 /** Helper for the common case of using a single color and depth render target. */
-DEPRECATED(4.22, "Utility SetRenderTarget (Single Color, DepthStencil, Simple Mode, DepthStencilAcess) is deprecated. Convert new code to use Render Passes.");
+DEPRECATED(4.22, "Utility SetRenderTarget (Single Color, DepthStencil, Simple Mode, DepthStencilAcess) is deprecated. Convert new code to use Render Passes.")
 inline void SetRenderTarget(FRHICommandList& RHICmdList, FTextureRHIParamRef NewRenderTarget, FTextureRHIParamRef NewDepthStencilTarget, ESimpleRenderTargetMode Mode, FExclusiveDepthStencil DepthStencilAccess = FExclusiveDepthStencil::DepthWrite_StencilWrite, bool bWritableBarrier = false)
 {
 	ERenderTargetLoadAction ColorLoadAction, DepthLoadAction, StencilLoadAction;
@@ -453,7 +453,7 @@ inline void SetRenderTarget(FRHICommandList& RHICmdList, FTextureRHIParamRef New
 }
 
 /** Helper for the common case of using a single color and depth render target, with a mip index for the color target. */
-DEPRECATED(4.22, "Utility SetRenderTarget (Single Color, Mip, DepthStencil) is deprecated. Convert new code to use Render Passes.");
+DEPRECATED(4.22, "Utility SetRenderTarget (Single Color, Mip, DepthStencil) is deprecated. Convert new code to use Render Passes.")
 inline void SetRenderTarget(FRHICommandList& RHICmdList, FTextureRHIParamRef NewRenderTarget, int32 MipIndex, FTextureRHIParamRef NewDepthStencilTarget, bool bWritableBarrier = false)
 {
 	FRHIRenderTargetView RTV(NewRenderTarget, ERenderTargetLoadAction::ELoad, MipIndex, -1);
@@ -468,7 +468,7 @@ inline void SetRenderTarget(FRHICommandList& RHICmdList, FTextureRHIParamRef New
 }
 
 /** Helper for the common case of using a single color and depth render target, with a mip index for the color target. */
-DEPRECATED(4.22, "Utility SetRenderTarget (Single Color, Mip, Slice, DepthStencil) is deprecated. Convert new code to use Render Passes.");
+DEPRECATED(4.22, "Utility SetRenderTarget (Single Color, Mip, Slice, DepthStencil) is deprecated. Convert new code to use Render Passes.")
 inline void SetRenderTarget(FRHICommandList& RHICmdList, FTextureRHIParamRef NewRenderTarget, int32 MipIndex, int32 ArraySliceIndex, FTextureRHIParamRef NewDepthStencilTarget, bool bWritableBarrier = false)
 {
 	FRHIRenderTargetView RTV(NewRenderTarget, ERenderTargetLoadAction::ELoad, MipIndex, ArraySliceIndex);
@@ -483,7 +483,7 @@ inline void SetRenderTarget(FRHICommandList& RHICmdList, FTextureRHIParamRef New
 }
 
 /** Helper that converts FTextureRHIParamRef's into FRHIRenderTargetView's. */
-DEPRECATED(4.22, "Utility SetRenderTargets (Colors, DepthStencil, UAVs) is deprecated. Convert new code to use Render Passes.");
+DEPRECATED(4.22, "Utility SetRenderTargets (Colors, DepthStencil, UAVs) is deprecated. Convert new code to use Render Passes.")
 inline void SetRenderTargets(
 	FRHICommandList& RHICmdList,
 	uint32 NewNumSimultaneousRenderTargets, 
@@ -512,7 +512,7 @@ inline void SetRenderTargets(
 }
 
 /** Helper that converts FTextureRHIParamRef's into FRHIRenderTargetView's. */
-DEPRECATED(4.22, "Utility SetRenderTargets (Colors, DepthStencil, UAVs, Simple Mode, DepthStencil Access) is deprecated. Convert new code to use Render Passes.");
+DEPRECATED(4.22, "Utility SetRenderTargets (Colors, DepthStencil, UAVs, Simple Mode, DepthStencil Access) is deprecated. Convert new code to use Render Passes.")
 inline void SetRenderTargets(
 	FRHICommandList& RHICmdList,
 	uint32 NewNumSimultaneousRenderTargets,
@@ -544,7 +544,7 @@ inline void SetRenderTargets(
 	RHICmdList.SetRenderTargets(NewNumSimultaneousRenderTargets, RTVs, &DepthRTV, 0, nullptr);
 }
 
-DEPRECATED(4.22, "Utility to unset render targets. Should not be used in new code - use Render Passes instead");
+DEPRECATED(4.22, "Utility to unset render targets. Should not be used in new code - use Render Passes instead")
 inline void UnbindRenderTargets(FRHICommandList& RHICmdList)
 {
 	check(RHICmdList.IsOutsideRenderPass());
