DR - Remove double deprecation warning
#rnx

diff --git a/Engine/Source/Runtime/RHI/Public/RHIUtilities.h b/Engine/Source/Runtime/RHI/Public/RHIUtilities.h
--- a/Engine/Source/Runtime/RHI/Public/RHIUtilities.h
+++ b/Engine/Source/Runtime/RHI/Public/RHIUtilities.h
@@ -548,8 +548,9 @@ DEPRECATED(4.22, "Utility to unset render targets. Should not be used in new cod
 inline void UnbindRenderTargets(FRHICommandList& RHICmdList)
 {
 	check(RHICmdList.IsOutsideRenderPass());
-
+PRAGMA_DISABLE_DEPRECATION_WARNINGS
 	SetRenderTarget(RHICmdList, nullptr, nullptr);
+PRAGMA_ENABLE_DEPRECATION_WARNINGS
 }
 
 /**
