Renderpass updates. Adding more validation.

diff --git a/Engine/Source/Runtime/Apple/MetalRHI/Private/MetalRHIContext.cpp b/Engine/Source/Runtime/Apple/MetalRHI/Private/MetalRHIContext.cpp
--- a/Engine/Source/Runtime/Apple/MetalRHI/Private/MetalRHIContext.cpp
+++ b/Engine/Source/Runtime/Apple/MetalRHI/Private/MetalRHIContext.cpp
@@ -139,9 +139,6 @@ FMetalRHIImmediateCommandContext::FMetalRHIImmediateCommandContext(class FMetalP
 
 void FMetalRHICommandContext::RHIBeginRenderPass(const FRHIRenderPassInfo& InInfo, const TCHAR* InName)
 {
-	// Fallback...
-	InInfo.Validate();
-	
 	if (InInfo.bGeneratingMips)
 	{
 		FRHITexture* Textures[MaxSimultaneousRenderTargets];
diff --git a/Engine/Source/Runtime/RHI/Private/RHI.cpp b/Engine/Source/Runtime/RHI/Private/RHI.cpp
--- a/Engine/Source/Runtime/RHI/Private/RHI.cpp
+++ b/Engine/Source/Runtime/RHI/Private/RHI.cpp
@@ -896,7 +896,7 @@ void FRHIRenderPassInfo::ConvertToRenderTargetsInfo(FRHISetRenderTargetsInfo& Ou
 	}
 }
 
-
+#if !(UE_BUILD_SHIPPING || UE_BUILD_TEST)
 void FRHIRenderPassInfo::Validate() const
 {
 	int32 NumSamples = -1;	// -1 means nothing found yet
@@ -921,6 +921,11 @@ void FRHIRenderPassInfo::Validate() const
 			ensure(Store != ERenderTargetStoreAction::EMultisampleResolve || Entry.RenderTarget->GetNumSamples() > 1);
 			// Don't resolve to null
 			ensure(Store != ERenderTargetStoreAction::EMultisampleResolve || Entry.ResolveTarget);
+
+			if (Entry.ResolveTarget)
+			{
+				ensure(Store == ERenderTargetStoreAction::EMultisampleResolve);
+			}
 		}
 		else
 		{
@@ -963,8 +968,8 @@ void FRHIRenderPassInfo::Validate() const
 		// Don't resolve to null
 		//ensure(DepthStencilRenderTarget.ResolveTarget || DepthStore != ERenderTargetStoreAction::EStore);
 		// Don't write to depth if read-only
-		ensure(DepthStencilRenderTarget.ExclusiveDepthStencil.IsDepthWrite() || DepthStore != ERenderTargetStoreAction::EStore);
-		ensure(DepthStencilRenderTarget.ExclusiveDepthStencil.IsStencilWrite() || StencilStore != ERenderTargetStoreAction::EStore);
+		ensure((DepthStencilRenderTarget.ExclusiveDepthStencil.IsDepthWrite() && DepthStore == ERenderTargetStoreAction::EStore) || DepthStore != ERenderTargetStoreAction::EStore);
+		ensure((DepthStencilRenderTarget.ExclusiveDepthStencil.IsStencilWrite() && StencilStore == ERenderTargetStoreAction::EStore) || StencilStore != ERenderTargetStoreAction::EStore);
 	}
 	else
 	{
@@ -972,6 +977,7 @@ void FRHIRenderPassInfo::Validate() const
 		ensure(DepthStencilRenderTarget.ExclusiveDepthStencil == FExclusiveDepthStencil::DepthNop_StencilNop);
 	}
 }
+#endif
 
 static FRHIPanicEvent RHIPanicEvent;
 FRHIPanicEvent& RHIGetPanicDelegate()
diff --git a/Engine/Source/Runtime/RHI/Public/RHICommandList.h b/Engine/Source/Runtime/RHI/Public/RHICommandList.h
--- a/Engine/Source/Runtime/RHI/Public/RHICommandList.h
+++ b/Engine/Source/Runtime/RHI/Public/RHICommandList.h
@@ -2736,6 +2736,8 @@ public:
 		check(!IsInsideRenderPass());
 		check(!IsInsideComputePass());
 
+		InInfo.Validate();
+
 		if (Bypass())
 		{
 			CMD_CONTEXT(RHIBeginRenderPass)(InInfo, Name);
diff --git a/Engine/Source/Runtime/RHI/Public/RHIContext.h b/Engine/Source/Runtime/RHI/Public/RHIContext.h
--- a/Engine/Source/Runtime/RHI/Public/RHIContext.h
+++ b/Engine/Source/Runtime/RHI/Public/RHIContext.h
@@ -566,9 +566,6 @@ public:
 
 	virtual void RHIBeginRenderPass(const FRHIRenderPassInfo& InInfo, const TCHAR* InName)
 	{
-		// Fallback...
-		InInfo.Validate();
-
 		if (InInfo.bGeneratingMips)
 		{
 			FRHITexture* Textures[MaxSimultaneousRenderTargets];
diff --git a/Engine/Source/Runtime/RHI/Public/RHIResources.h b/Engine/Source/Runtime/RHI/Public/RHIResources.h
--- a/Engine/Source/Runtime/RHI/Public/RHIResources.h
+++ b/Engine/Source/Runtime/RHI/Public/RHIResources.h
@@ -2255,7 +2255,11 @@ struct FRHIRenderPassInfo
 		return bIsMSAA;
 	}
 
+#if !(UE_BUILD_SHIPPING || UE_BUILD_TEST)
 	RHI_API void Validate() const;
+#else
+	RHI_API void Validate() const {}
+#endif
 	RHI_API void ConvertToRenderTargetsInfo(FRHISetRenderTargetsInfo& OutRTInfo) const;
 
 	FRHIRenderPassInfo& operator = (const FRHIRenderPassInfo& In)
