* Removing the GetOrCreatePSO profile marker as it doesn't tell us anything super useful and is one of the most frequently spammed markers.

diff --git a/Engine/Source/Runtime/RHI/Private/PipelineStateCache.cpp b/Engine/Source/Runtime/RHI/Private/PipelineStateCache.cpp
--- a/Engine/Source/Runtime/RHI/Private/PipelineStateCache.cpp
+++ b/Engine/Source/Runtime/RHI/Private/PipelineStateCache.cpp
@@ -745,9 +745,7 @@ static bool IsAsyncCompilationAllowed(FRHICommandList& RHICmdList)
 }
 
 FComputePipelineState* PipelineStateCache::GetAndOrCreateComputePipelineState(FRHICommandList& RHICmdList, FRHIComputeShader* ComputeShader)
-{
-	SCOPE_CYCLE_COUNTER(STAT_GetOrCreatePSO);
-	
+{	
 	bool DoAsyncCompile = IsAsyncCompilationAllowed(RHICmdList);
 
 	FComputePipelineState* OutCachedState = nullptr;
@@ -857,7 +855,6 @@ FRHIComputePipelineState* ExecuteSetComputePipelineState(FComputePipelineState*
 FGraphicsPipelineState* PipelineStateCache::GetAndOrCreateGraphicsPipelineState(FRHICommandList& RHICmdList, const FGraphicsPipelineStateInitializer& OriginalInitializer, EApplyRendertargetOption ApplyFlags)
 {
 	LLM_SCOPE(ELLMTag::PSO);
-	SCOPE_CYCLE_COUNTER(STAT_GetOrCreatePSO);
 
 	// Workaround until we have a better way for storing shaders in PSO cache 
 	{
diff --git a/Engine/Source/Runtime/RHI/Private/RHI.cpp b/Engine/Source/Runtime/RHI/Private/RHI.cpp
--- a/Engine/Source/Runtime/RHI/Private/RHI.cpp
+++ b/Engine/Source/Runtime/RHI/Private/RHI.cpp
@@ -33,7 +33,6 @@ DEFINE_STAT(STAT_IndexBufferMemory);
 DEFINE_STAT(STAT_VertexBufferMemory);
 DEFINE_STAT(STAT_StructuredBufferMemory);
 DEFINE_STAT(STAT_PixelBufferMemory);
-DEFINE_STAT(STAT_GetOrCreatePSO);
 
 static FAutoConsoleVariable CVarUseVulkanRealUBs(
 	TEXT("r.Vulkan.UseRealUBs"),
diff --git a/Engine/Source/Runtime/RHI/Public/RHI.h b/Engine/Source/Runtime/RHI/Public/RHI.h
--- a/Engine/Source/Runtime/RHI/Public/RHI.h
+++ b/Engine/Source/Runtime/RHI/Public/RHI.h
@@ -1535,7 +1535,6 @@ DECLARE_MEMORY_STAT_POOL_EXTERN(TEXT("Index buffer memory"),STAT_IndexBufferMemo
 DECLARE_MEMORY_STAT_POOL_EXTERN(TEXT("Vertex buffer memory"),STAT_VertexBufferMemory,STATGROUP_RHI,FPlatformMemory::MCR_GPU,RHI_API);
 DECLARE_MEMORY_STAT_POOL_EXTERN(TEXT("Structured buffer memory"),STAT_StructuredBufferMemory,STATGROUP_RHI,FPlatformMemory::MCR_GPU,RHI_API);
 DECLARE_MEMORY_STAT_POOL_EXTERN(TEXT("Pixel buffer memory"),STAT_PixelBufferMemory,STATGROUP_RHI,FPlatformMemory::MCR_GPU,RHI_API);
-DECLARE_CYCLE_STAT_EXTERN(TEXT("Get/Create PSO"), STAT_GetOrCreatePSO, STATGROUP_RHI, RHI_API);
 
 
 // RHI base resource types.
