Update copyright notice for RHIShaderPlatformDefinitions.inl

diff --git a/Engine/Source/Runtime/RHI/Public/RHIShaderPlatformDefinitions.inl b/Engine/Source/Runtime/RHI/Public/RHIShaderPlatformDefinitions.inl
--- a/Engine/Source/Runtime/RHI/Public/RHIShaderPlatformDefinitions.inl
+++ b/Engine/Source/Runtime/RHI/Public/RHIShaderPlatformDefinitions.inl
@@ -1,4 +1,4 @@
-// Copyright 1998-2018 Epic Games, Inc. All Rights Reserved.
+// Copyright 1998-2019 Epic Games, Inc. All Rights Reserved.
 
 /*=============================================================================
 	RHIShaderPlatformDefinitions.h: Localizable Friendly Names for Shader Platforms
