Fixed localization gather warnings

diff --git a/Engine/Plugins/Developer/Concert/ConcertFrontend/Source/ConcertFrontend/Private/Widgets/SConcertBrowser.cpp b/Engine/Plugins/Developer/Concert/ConcertFrontend/Source/ConcertFrontend/Private/Widgets/SConcertBrowser.cpp
--- a/Engine/Plugins/Developer/Concert/ConcertFrontend/Source/ConcertFrontend/Private/Widgets/SConcertBrowser.cpp
+++ b/Engine/Plugins/Developer/Concert/ConcertFrontend/Source/ConcertFrontend/Private/Widgets/SConcertBrowser.cpp
@@ -687,7 +687,7 @@ TSharedRef<ITableRow> SConcertBrowser::MakeSessionRowWidget(TSharedPtr<FConcertS
 		ActiveSessionDef.Style = EConcertUIStyle::Info;
 		ActiveSessionDef.Visibility = MakeAttributeSP(this, &SConcertBrowser::IsActiveSessionVisible, Item->SessionName);
 		ActiveSessionDef.Text = FEditorFontGlyphs::Info_Circle;
-		ActiveSessionDef.ToolTipText = LOCTEXT("ActiveSessionToolTip", "See the current active session.");
+		ActiveSessionDef.ToolTipText = LOCTEXT("ActiveSessionToolTip", "See the current active session");
 		ActiveSessionDef.OnClicked.BindSP(this, &SConcertBrowser::OnClickActiveSession);
 		
 		// Resume Session
diff --git a/Engine/Plugins/Developer/Concert/ConcertSync/ConcertSyncClient/Source/ConcertSyncClient/Private/ConcertClientWorkspace.cpp b/Engine/Plugins/Developer/Concert/ConcertSync/ConcertSyncClient/Source/ConcertSyncClient/Private/ConcertClientWorkspace.cpp
--- a/Engine/Plugins/Developer/Concert/ConcertSync/ConcertSyncClient/Source/ConcertSyncClient/Private/ConcertClientWorkspace.cpp
+++ b/Engine/Plugins/Developer/Concert/ConcertSync/ConcertSyncClient/Source/ConcertSyncClient/Private/ConcertClientWorkspace.cpp
@@ -736,7 +736,7 @@ void FConcertClientWorkspace::HandleWorkspaceSyncPackageEvent(const FConcertSess
 	if (InitialSyncSlowTask.IsValid())
 	{
 		InitialSyncSlowTask->TotalAmountOfWork = InitialSyncSlowTask->CompletedWork + Event.RemainingWork + 1;
-		InitialSyncSlowTask->EnterProgressFrame(FMath::Min<float>(Event.RemainingWork, 1.0f), FText::Format(LOCTEXT("SynchronizedTransactionFmt", "Synchronized Package {0}..."), FText::FromName(Event.Package.Info.PackageName)));
+		InitialSyncSlowTask->EnterProgressFrame(FMath::Min<float>(Event.RemainingWork, 1.0f), FText::Format(LOCTEXT("SynchronizedPackageFmt", "Synchronized Package {0}..."), FText::FromName(Event.Package.Info.PackageName)));
 	}
 
 	switch (Event.Package.Info.PackageUpdateType)
diff --git a/Engine/Plugins/Developer/Concert/ConcertSync/ConcertSyncCore/Source/ConcertSyncCore/Private/ConcertActivityEvents.cpp b/Engine/Plugins/Developer/Concert/ConcertSync/ConcertSyncCore/Source/ConcertSyncCore/Private/ConcertActivityEvents.cpp
--- a/Engine/Plugins/Developer/Concert/ConcertSync/ConcertSyncCore/Source/ConcertSyncCore/Private/ConcertActivityEvents.cpp
+++ b/Engine/Plugins/Developer/Concert/ConcertSync/ConcertSyncCore/Source/ConcertSyncCore/Private/ConcertActivityEvents.cpp
@@ -156,7 +156,7 @@ FText FConcertPackageAddedActivityEvent::ToDisplayText(bool bRichText /*=false*/
 	FFormatNamedArguments Arguments;
 	Arguments.Add(TEXT("UserName"), ActivityEventsUtil::ToRichTextBold(GetClientDisplayName(), bRichText));
 	Arguments.Add(TEXT("PackageName"), ActivityEventsUtil::ToRichTextBold(PackageName, bRichText));
-	return FText::Format(LOCTEXT("PackageUpdated", "{UserName} created the package {PackageName}."), Arguments);
+	return FText::Format(LOCTEXT("PackageAdded", "{UserName} created the package {PackageName}."), Arguments);
 }
 
 FConcertPackageDeletedActivityEvent::~FConcertPackageDeletedActivityEvent()
@@ -168,7 +168,7 @@ FText FConcertPackageDeletedActivityEvent::ToDisplayText(bool bRichText /*=false
 	FFormatNamedArguments Arguments;
 	Arguments.Add(TEXT("UserName"), ActivityEventsUtil::ToRichTextBold(GetClientDisplayName(), bRichText));
 	Arguments.Add(TEXT("PackageName"), ActivityEventsUtil::ToRichTextBold(PackageName, bRichText));
-	return FText::Format(LOCTEXT("PackageUpdated", "{UserName} deleted the package {PackageName}."), Arguments);
+	return FText::Format(LOCTEXT("PackageDeleted", "{UserName} deleted the package {PackageName}."), Arguments);
 }
 
 FConcertPackageRenamedActivityEvent::~FConcertPackageRenamedActivityEvent()
@@ -181,7 +181,7 @@ FText FConcertPackageRenamedActivityEvent::ToDisplayText(bool bRichText /*=false
 	Arguments.Add(TEXT("UserName"), ActivityEventsUtil::ToRichTextBold(GetClientDisplayName(), bRichText));
 	Arguments.Add(TEXT("PackageOldName"), ActivityEventsUtil::ToRichTextBold(PackageName, bRichText));
 	Arguments.Add(TEXT("PackageNewName"), ActivityEventsUtil::ToRichTextBold(NewPackageName, bRichText));
-	return FText::Format(LOCTEXT("PackageUpdated", "{UserName} moved the package {PackageOldName} to {PackageNewName}."), Arguments);
+	return FText::Format(LOCTEXT("PackageRenamed", "{UserName} moved the package {PackageOldName} to {PackageNewName}."), Arguments);
 }
 
 #undef LOCTEXT_NAMESPACE
diff --git a/Engine/Plugins/Developer/Concert/ConcertSync/ConcertSyncTest/Source/ConcertSyncTest/Private/ConcertDataStoreTests.cpp b/Engine/Plugins/Developer/Concert/ConcertSync/ConcertSyncTest/Source/ConcertSyncTest/Private/ConcertDataStoreTests.cpp
--- a/Engine/Plugins/Developer/Concert/ConcertSync/ConcertSyncTest/Source/ConcertSyncTest/Private/ConcertDataStoreTests.cpp
+++ b/Engine/Plugins/Developer/Concert/ConcertSync/ConcertSyncTest/Source/ConcertSyncTest/Private/ConcertDataStoreTests.cpp
@@ -1108,3 +1108,5 @@ bool FConcertDataStoreClientPrivateStore::RunTest(const FString& Parameters)
 
 	return true;
 };
+
+#undef LOCTEXT_NAMESPACE
diff --git a/Engine/Plugins/Editor/MeshEditor/Source/MeshEditor/MeshEditorMode.cpp b/Engine/Plugins/Editor/MeshEditor/Source/MeshEditor/MeshEditorMode.cpp
--- a/Engine/Plugins/Editor/MeshEditor/Source/MeshEditor/MeshEditorMode.cpp
+++ b/Engine/Plugins/Editor/MeshEditor/Source/MeshEditor/MeshEditorMode.cpp
@@ -459,7 +459,7 @@ FMeshEditorMode::FMeshEditorMode()
 	FMeshEditorSelectionModifiers::Register();
 
 	// Mesh fracture configuration settings
-	MeshFractureSettings = NewObject<UMeshFractureSettings>(GetTransientPackage(), *LOCTEXT("FractureSettingsName", "FractureSettings").ToString());
+	MeshFractureSettings = NewObject<UMeshFractureSettings>(GetTransientPackage(), TEXT("FractureSettings"));
 	MeshFractureSettings->AddToRoot();
 
 	// Register UI commands
diff --git a/Engine/Plugins/Experimental/BlastPlugin/Source/BlastAuthoring/Private/MeshFractureSettings.cpp b/Engine/Plugins/Experimental/BlastPlugin/Source/BlastAuthoring/Private/MeshFractureSettings.cpp
--- a/Engine/Plugins/Experimental/BlastPlugin/Source/BlastAuthoring/Private/MeshFractureSettings.cpp
+++ b/Engine/Plugins/Experimental/BlastPlugin/Source/BlastAuthoring/Private/MeshFractureSettings.cpp
@@ -7,7 +7,6 @@
 #include "GeometryCollection/GeometryCollectionComponent.h"
 #include "GeometryCollection/GeometryCollectionActor.h"
 
-#define LOCTEXT_NAMESPACE "MeshFractureSettings"
 
 float UMeshFractureSettings::ExplodedViewExpansion = 0.0f;
 
@@ -76,30 +75,29 @@ UBrickFractureSettings::UBrickFractureSettings() : SlicesX(3), SlicesY(3), Slice
 
 UMeshFractureSettings::UMeshFractureSettings()
 {
-	CommonSettings = NewObject<UCommonFractureSettings>(GetTransientPackage(), *LOCTEXT("SettingsName", "CommonSettings").ToString());
+	CommonSettings = NewObject<UCommonFractureSettings>(GetTransientPackage(), TEXT("CommonSettings"));
 	CommonSettings->AddToRoot();
 
-	UniformSettings = NewObject<UUniformFractureSettings>(GetTransientPackage(), *LOCTEXT("SettingsName", "UniformSettings").ToString());
+	UniformSettings = NewObject<UUniformFractureSettings>(GetTransientPackage(), TEXT("UniformSettings"));
 	UniformSettings->AddToRoot();
 
-	ClusterSettings = NewObject<UClusterFractureSettings>(GetTransientPackage(), *LOCTEXT("SettingsName", "ClusterSettings").ToString());
+	ClusterSettings = NewObject<UClusterFractureSettings>(GetTransientPackage(), TEXT("ClusterSettings"));
 	ClusterSettings->AddToRoot();
 
-	RadialSettings = NewObject<URadialFractureSettings>(GetTransientPackage(), *LOCTEXT("SettingsName", "RadialSettings").ToString());
+	RadialSettings = NewObject<URadialFractureSettings>(GetTransientPackage(), TEXT("RadialSettings"));
 	RadialSettings->AddToRoot();
 
-	SlicingSettings = NewObject<USlicingFractureSettings>(GetTransientPackage(), *LOCTEXT("SettingsName", "SlicingSettings").ToString());
+	SlicingSettings = NewObject<USlicingFractureSettings>(GetTransientPackage(), TEXT("SlicingSettings"));
 	SlicingSettings->AddToRoot();
 
-	PlaneCutSettings = NewObject<UPlaneCutFractureSettings>(GetTransientPackage(), *LOCTEXT("SettingsName", "PlaneCutSettings").ToString());
+	PlaneCutSettings = NewObject<UPlaneCutFractureSettings>(GetTransientPackage(), TEXT("PlaneCutSettings"));
 	PlaneCutSettings->AddToRoot();
 
-	CutoutSettings = NewObject<UCutoutFractureSettings>(GetTransientPackage(), *LOCTEXT("SettingsName", "CutoutSettings").ToString());
+	CutoutSettings = NewObject<UCutoutFractureSettings>(GetTransientPackage(), TEXT("CutoutSettings"));
 	CutoutSettings->AddToRoot();
 
-	BrickSettings = NewObject<UBrickFractureSettings>(GetTransientPackage(), *LOCTEXT("SettingsName", "BrickSettings").ToString());
+	BrickSettings = NewObject<UBrickFractureSettings>(GetTransientPackage(), TEXT("BrickSettings"));
 	BrickSettings->AddToRoot();
-
 }
 
 UMeshFractureSettings::~UMeshFractureSettings()
diff --git a/Engine/Plugins/Experimental/ControlRig/Source/ControlRigEditor/Private/EditMode/ControlRigEditMode.cpp b/Engine/Plugins/Experimental/ControlRig/Source/ControlRigEditor/Private/EditMode/ControlRigEditMode.cpp
--- a/Engine/Plugins/Experimental/ControlRig/Source/ControlRigEditor/Private/EditMode/ControlRigEditMode.cpp
+++ b/Engine/Plugins/Experimental/ControlRig/Source/ControlRigEditor/Private/EditMode/ControlRigEditMode.cpp
@@ -57,7 +57,7 @@ FControlRigEditMode::FControlRigEditMode()
 	, bSelectingByPath(false)
 	, PivotTransform(FTransform::Identity)
 {
-	Settings = NewObject<UControlRigEditModeSettings>(GetTransientPackage(), *LOCTEXT("SettingsName", "Settings").ToString());
+	Settings = NewObject<UControlRigEditModeSettings>(GetTransientPackage(), TEXT("Settings"));
 
 	OnControlsSelectedDelegate.AddRaw(this, &FControlRigEditMode::HandleSelectionChanged);
 
diff --git a/Engine/Plugins/Experimental/SkeletalReduction/Source/Private/SkeletalMeshReductionPlugin.cpp b/Engine/Plugins/Experimental/SkeletalReduction/Source/Private/SkeletalMeshReductionPlugin.cpp
--- a/Engine/Plugins/Experimental/SkeletalReduction/Source/Private/SkeletalMeshReductionPlugin.cpp
+++ b/Engine/Plugins/Experimental/SkeletalReduction/Source/Private/SkeletalMeshReductionPlugin.cpp
@@ -1547,3 +1547,5 @@ void FQuadricSkeletalMeshReduction::ReduceSkeletalMesh(USkeletalMesh& SkeletalMe
 
 	SkeletalMesh.CalculateRequiredBones(SkeletalMeshResource.LODModels[LODIndex], SkeletalMesh.RefSkeleton, &BonesToRemove);
 }
+
+#undef LOCTEXT_NAMESPACE
diff --git a/Engine/Plugins/Media/MediaIOFramework/Source/MediaIOCore/Private/MediaIOCoreCommonDisplayMode.cpp b/Engine/Plugins/Media/MediaIOFramework/Source/MediaIOCore/Private/MediaIOCoreCommonDisplayMode.cpp
--- a/Engine/Plugins/Media/MediaIOFramework/Source/MediaIOCore/Private/MediaIOCoreCommonDisplayMode.cpp
+++ b/Engine/Plugins/Media/MediaIOFramework/Source/MediaIOCore/Private/MediaIOCoreCommonDisplayMode.cpp
@@ -165,7 +165,7 @@ FText FMediaIOCommonDisplayModes::GetMediaIOCommonDisplayModeInfoName(int32 InWi
 	{
 	case EMediaIOStandardType::Interlaced: Standard = LOCTEXT("Interlaced_Short", "i"); break;
 	case EMediaIOStandardType::Progressive: Standard = LOCTEXT("Progressive_Short", "p"); break;
-	case EMediaIOStandardType::ProgressiveSegmentedFrame: Standard = LOCTEXT("Progressive_Short", "psf"); break;
+	case EMediaIOStandardType::ProgressiveSegmentedFrame: Standard = LOCTEXT("ProgressiveSegmentedFrame_Short", "psf"); break;
 	}
 
 	const FMediaIOCommonDisplayModeResolutionInfo* Resolution = Find(InWidth, InHeight);
diff --git a/Engine/Source/Developer/AssetTools/Private/AssetTypeActions/AssetTypeActions_AnimCurveCompressionSettings.cpp b/Engine/Source/Developer/AssetTools/Private/AssetTypeActions/AssetTypeActions_AnimCurveCompressionSettings.cpp
--- a/Engine/Source/Developer/AssetTools/Private/AssetTypeActions/AssetTypeActions_AnimCurveCompressionSettings.cpp
+++ b/Engine/Source/Developer/AssetTools/Private/AssetTypeActions/AssetTypeActions_AnimCurveCompressionSettings.cpp
@@ -108,3 +108,5 @@ void FAssetTypeActions_AnimCurveCompressionSettings::ExecuteCompression(TWeakObj
 		AnimSeq->RequestSyncAnimRecompression(false);
 	}
 }
+
+#undef LOCTEXT_NAMESPACE
diff --git a/Engine/Source/Editor/LevelEditor/Private/LevelEditorToolBar.cpp b/Engine/Source/Editor/LevelEditor/Private/LevelEditorToolBar.cpp
--- a/Engine/Source/Editor/LevelEditor/Private/LevelEditorToolBar.cpp
+++ b/Engine/Source/Editor/LevelEditor/Private/LevelEditorToolBar.cpp
@@ -57,10 +57,7 @@
 #include "LauncherPlatformModule.h"
 #include "Misc/ScopedSlowTask.h"
 #include "MaterialShaderQualitySettings.h"
-
-#define LOCTEXT_NAMESPACE "EditorViewport"
 #include "RHIShaderPlatformDefinitions.inl"
-#undef LOCTEXT_NAMESPACE
 
 namespace LevelEditorActionHelpers
 {
diff --git a/Engine/Source/Runtime/Engine/Private/ActorEditor.cpp b/Engine/Source/Runtime/Engine/Private/ActorEditor.cpp
--- a/Engine/Source/Runtime/Engine/Private/ActorEditor.cpp
+++ b/Engine/Source/Runtime/Engine/Private/ActorEditor.cpp
@@ -980,7 +980,7 @@ EDataValidationResult AActor::IsDataValid(TArray<FText>& ValidationErrors)
 	bool bSuccess = CheckDefaultSubobjects();
 	if (!bSuccess)
 	{
-		FText ErrorMsg = FText::Format(LOCTEXT("IsDataValid_Failed_CheckForErrors", "{0} failed CheckDefaultSubobjectsInternal()"), FText::FromString(GetName()));
+		FText ErrorMsg = FText::Format(LOCTEXT("IsDataValid_Failed_CheckDefaultSubobjectsInternal", "{0} failed CheckDefaultSubobjectsInternal()"), FText::FromString(GetName()));
 		ValidationErrors.Add(ErrorMsg);
 	}
 
diff --git a/Engine/Source/Runtime/Experimental/Chaos/Private/Chaos/PBDRigidsEvolutionGBF.cpp b/Engine/Source/Runtime/Experimental/Chaos/Private/Chaos/PBDRigidsEvolutionGBF.cpp
--- a/Engine/Source/Runtime/Experimental/Chaos/Private/Chaos/PBDRigidsEvolutionGBF.cpp
+++ b/Engine/Source/Runtime/Experimental/Chaos/Private/Chaos/PBDRigidsEvolutionGBF.cpp
@@ -289,3 +289,5 @@ void TPBDRigidsEvolutionGBF<T, d>::AdvanceOneTimeStep(const T Dt)
 }
 
 template class Chaos::TPBDRigidsEvolutionGBF<float, 3>;
+
+#undef LOCTEXT_NAMESPACE
diff --git a/Engine/Source/Runtime/Experimental/Chaos/Private/Chaos/PBDRigidsEvolutionPGS.cpp b/Engine/Source/Runtime/Experimental/Chaos/Private/Chaos/PBDRigidsEvolutionPGS.cpp
--- a/Engine/Source/Runtime/Experimental/Chaos/Private/Chaos/PBDRigidsEvolutionPGS.cpp
+++ b/Engine/Source/Runtime/Experimental/Chaos/Private/Chaos/PBDRigidsEvolutionPGS.cpp
@@ -153,3 +153,5 @@ void TPBDRigidsEvolutionPGS<T, d>::AdvanceOneTimeStep(const T Dt)
 }
 
 template class Chaos::TPBDRigidsEvolutionPGS<float, 3>;
+
+#undef LOCTEXT_NAMESPACE
diff --git a/Engine/Source/Runtime/RHI/Public/RHIShaderPlatformDefinitions.inl b/Engine/Source/Runtime/RHI/Public/RHIShaderPlatformDefinitions.inl
--- a/Engine/Source/Runtime/RHI/Public/RHIShaderPlatformDefinitions.inl
+++ b/Engine/Source/Runtime/RHI/Public/RHIShaderPlatformDefinitions.inl
@@ -16,7 +16,7 @@ static const FText& GetFriendlyShaderPlatformName(const EShaderPlatform InShader
 	case SP_VULKAN_SM5:
 	case SP_OPENGL_ES31_EXT:
 	{
-		static FText Description = LOCTEXT("Generic_SM5_loc", "SM5");
+		static const FText Description = NSLOCTEXT("FriendlyShaderPlatformNames", "Generic_SM5_loc", "SM5");
 		return Description;
 	}
 	break;
@@ -25,7 +25,7 @@ static const FText& GetFriendlyShaderPlatformName(const EShaderPlatform InShader
 	case SP_METAL_MRT_MAC:
 	case SP_METAL_MRT:
 	{
-		static FText Description = LOCTEXT("Metal_SM5_loc", "Metal SM5");
+		static const FText Description = NSLOCTEXT("FriendlyShaderPlatformNames", "Metal_SM5_loc", "Metal SM5");
 		return Description;
 	}
 	break;
@@ -34,35 +34,35 @@ static const FText& GetFriendlyShaderPlatformName(const EShaderPlatform InShader
 	case SP_OPENGL_SM4:
 	case SP_VULKAN_SM4:
 	{
-		static FText Description = LOCTEXT("Generic_SM4_loc", "SM4");
+		static const FText Description = NSLOCTEXT("FriendlyShaderPlatformNames", "Generic_SM4_loc", "SM4");
 		return Description;
 	}
 	break;
 
 	case SP_PS4:
 	{
-		static FText Description = LOCTEXT("Playstation_4_loc", "PlayStation 4");
+		static const FText Description = NSLOCTEXT("FriendlyShaderPlatformNames", "Playstation_4_loc", "PlayStation 4");
 		return Description;
 	}
 	break;
 
 	case SP_XBOXONE_D3D12:
 	{
-		static FText Description = LOCTEXT("XBox_One_loc", "XBox One");
+		static const FText Description = NSLOCTEXT("FriendlyShaderPlatformNames", "XBox_One_loc", "XBox One");
 		return Description;
 	}
 	break;
 
 	case SP_SWITCH:
 	{
-		static FText Description = LOCTEXT("Switch_loc", "Switch");
+		static const FText Description = NSLOCTEXT("FriendlyShaderPlatformNames", "Switch_loc", "Switch");
 		return Description;
 	}
 	break;
 
 	case SP_SWITCH_FORWARD:
 	{
-		static FText Description = LOCTEXT("Switch_Forward_loc", "Switch Forward");
+		static const FText Description = NSLOCTEXT("FriendlyShaderPlatformNames", "Switch_Forward_loc", "Switch Forward");
 		return Description;
 	}
 	break;
@@ -71,28 +71,28 @@ static const FText& GetFriendlyShaderPlatformName(const EShaderPlatform InShader
 	case SP_VULKAN_PCES3_1:
 	case SP_OPENGL_PCES3_1:
 	{
-		static FText Description = LOCTEXT("Generic_ES31_loc", "ES31");
+		static const FText Description = NSLOCTEXT("FriendlyShaderPlatformNames", "Generic_ES31_loc", "ES31");
 		return Description;
 	}
 	break;
 
 	case SP_VULKAN_ES3_1_ANDROID:
 	{
-		static FText Description = LOCTEXT("Android_Vulkan_ES31_loc", "Android Vulkan ES31");
+		static const FText Description = NSLOCTEXT("FriendlyShaderPlatformNames", "Android_Vulkan_ES31_loc", "Android Vulkan ES31");
 		return Description;
 	}
 	break;
 
 	case SP_OPENGL_ES3_1_ANDROID:
 	{
-		static FText Description = LOCTEXT("Android_ES31_loc", "Android ES31");
+		static const FText Description = NSLOCTEXT("FriendlyShaderPlatformNames", "Android_ES31_loc", "Android ES31");
 		return Description;
 	}
 	break;
 
 	case SP_VULKAN_ES3_1_LUMIN:
 	{
-		static FText Description = LOCTEXT("Lumin_Vulkan_ES31_loc", "Lumin Vulkan ES31");
+		static const FText Description = NSLOCTEXT("FriendlyShaderPlatformNames", "Lumin_Vulkan_ES31_loc", "Lumin Vulkan ES31");
 		return Description;
 	}
 	break;
@@ -100,14 +100,14 @@ static const FText& GetFriendlyShaderPlatformName(const EShaderPlatform InShader
 	case SP_METAL:
 	case SP_METAL_MACES3_1:
 	{
-		static FText Description = LOCTEXT("iOS_Metal_Mace_31_loc", "Metal ES31");
+		static const FText Description = NSLOCTEXT("FriendlyShaderPlatformNames", "iOS_Metal_Mace_31_loc", "Metal ES31");
 		return Description;
 	}
 	break;
 
 	case SP_VULKAN_SM5_LUMIN:
 	{
-		static FText Description = LOCTEXT("Lumin_Vulkan_SM5_loc", "Lumin Vulkan SM5");
+		static const FText Description = NSLOCTEXT("FriendlyShaderPlatformNames", "Lumin_Vulkan_SM5_loc", "Lumin Vulkan SM5");
 		return Description;
 	}
 	break;
@@ -115,35 +115,35 @@ static const FText& GetFriendlyShaderPlatformName(const EShaderPlatform InShader
 	case SP_PCD3D_ES2:
 	case SP_OPENGL_PCES2:
 	{
-		static FText Description = LOCTEXT("Generic_ES2_loc", "ES2");
+		static const FText Description = NSLOCTEXT("FriendlyShaderPlatformNames", "Generic_ES2_loc", "ES2");
 		return Description;
 	}
 	break;
 
 	case SP_OPENGL_ES2_ANDROID:
 	{
-		static FText Description = LOCTEXT("Android_ES2_loc", "Android ES2");
+		static const FText Description = NSLOCTEXT("FriendlyShaderPlatformNames", "Android_ES2_loc", "Android ES2");
 		return Description;
 	}
 	break;
 
 	case SP_OPENGL_ES2_WEBGL:
 	{
-		static FText Description = LOCTEXT("WebGL_ES2_loc", "WebGL ES2");
+		static const FText Description = NSLOCTEXT("FriendlyShaderPlatformNames", "WebGL_ES2_loc", "WebGL ES2");
 		return Description;
 	}
 	break;
 
 	case SP_OPENGL_ES2_IOS:
 	{
-		static FText Description = LOCTEXT("iOS_ES2_loc", "iOS ES2");
+		static const FText Description = NSLOCTEXT("FriendlyShaderPlatformNames", "iOS_ES2_loc", "iOS ES2");
 		return Description;
 	}
 	break;
 
 	case SP_METAL_MACES2:
 	{
-		static FText Description = LOCTEXT("iOS_Metal_Mace_2_loc", "iOS Metal ES2");
+		static const FText Description = NSLOCTEXT("FriendlyShaderPlatformNames", "iOS_Metal_Mace_2_loc", "iOS Metal ES2");
 		return Description;
 	}
 	break;
@@ -152,7 +152,5 @@ static const FText& GetFriendlyShaderPlatformName(const EShaderPlatform InShader
 		break;
 	};
 
-	static FText Description = LOCTEXT("", "");
-	return Description;
+	return FText::GetEmpty();
 }
-
diff --git a/Engine/Source/Runtime/SlateCore/Private/Debugging/ConsoleSlateDebugger.cpp b/Engine/Source/Runtime/SlateCore/Private/Debugging/ConsoleSlateDebugger.cpp
--- a/Engine/Source/Runtime/SlateCore/Private/Debugging/ConsoleSlateDebugger.cpp
+++ b/Engine/Source/Runtime/SlateCore/Private/Debugging/ConsoleSlateDebugger.cpp
@@ -38,7 +38,7 @@ FConsoleSlateDebugger::FConsoleSlateDebugger()
 		*LOCTEXT("CaptureStack", "Should we capture the stack when there are events?").ToString())
 	, SetInputFilterCommand(
 		TEXT("SlateDebugger.SetInputFilter"),
-		*LOCTEXT("StopDebugger", "Enable or Disable specific filters").ToString(),
+		*LOCTEXT("SetInputFilter", "Enable or Disable specific filters").ToString(),
 		FConsoleCommandWithArgsDelegate::CreateRaw(this, &FConsoleSlateDebugger::SetInputFilter))
 {
 	//EnabledInputEvents[(uint8)ESlateDebuggingInputEvent::MouseMove] = true;
diff --git a/Engine/Source/Runtime/UMG/Private/Components/ListViewBase.cpp b/Engine/Source/Runtime/UMG/Private/Components/ListViewBase.cpp
--- a/Engine/Source/Runtime/UMG/Private/Components/ListViewBase.cpp
+++ b/Engine/Source/Runtime/UMG/Private/Components/ListViewBase.cpp
@@ -87,7 +87,7 @@ TSharedRef<SWidget> UListViewBase::RebuildWidget()
 		}
 		else if (EntryWidgetBP->Status == BS_Error)
 		{
-			ErrorText = FText::Format(LOCTEXT("Error_NonBPEntryWidget", "EntryWidget BP [{0}] has not compiled successfully"), FText::FromString(EntryWidgetBP->GetName()));
+			ErrorText = FText::Format(LOCTEXT("Error_CompilationError", "EntryWidget BP [{0}] has not compiled successfully"), FText::FromString(EntryWidgetBP->GetName()));
 		}
 	}
 #endif
