added comment
#rnx

diff --git a/Engine/Source/Runtime/RHI/Public/RHIDefinitions.h b/Engine/Source/Runtime/RHI/Public/RHIDefinitions.h
--- a/Engine/Source/Runtime/RHI/Public/RHIDefinitions.h
+++ b/Engine/Source/Runtime/RHI/Public/RHIDefinitions.h
@@ -973,6 +973,7 @@ inline ERHIFeatureLevel::Type GetMaxSupportedFeatureLevel(EShaderPlatform InShad
 	}
 }
 
+/* Returns true if the shader platform Platform is used to simulate a mobile feature level on a PC platform. */
 inline bool IsSimulatedPlatform(EShaderPlatform Platform)
 {
 	switch (Platform)
