VR - Fix Lumin profile

diff --git a/Engine/Source/Runtime/RHI/Public/RHIDefinitions.h b/Engine/Source/Runtime/RHI/Public/RHIDefinitions.h
--- a/Engine/Source/Runtime/RHI/Public/RHIDefinitions.h
+++ b/Engine/Source/Runtime/RHI/Public/RHIDefinitions.h
@@ -810,8 +810,11 @@ inline bool IsMobileOpenGlPlatform(const EShaderPlatform Platform)
 inline bool IsMobilePlatform(const EShaderPlatform Platform)
 {
 	return IsES2Platform(Platform)
-		|| Platform == SP_METAL || Platform == SP_PCD3D_ES3_1 || Platform == SP_OPENGL_PCES3_1 || Platform == SP_VULKAN_ES3_1_ANDROID
-		|| Platform == SP_VULKAN_PCES3_1 || Platform == SP_METAL_MACES3_1 || Platform == SP_OPENGL_ES3_1_ANDROID || Platform == SP_SWITCH_FORWARD;
+		|| Platform == SP_METAL || Platform == SP_METAL_MACES3_1
+		|| Platform == SP_PCD3D_ES3_1
+		|| Platform == SP_OPENGL_PCES3_1 || Platform == SP_OPENGL_ES3_1_ANDROID
+		|| Platform == SP_VULKAN_ES3_1_ANDROID || Platform == SP_VULKAN_PCES3_1 || Platform == SP_VULKAN_ES3_1_LUMIN
+		|| Platform == SP_SWITCH_FORWARD;
 }
 
 inline bool IsOpenGLPlatform(const EShaderPlatform Platform)
