ReadSurfaceFloatData flushes rhi commandlist.
Fixes sky captures occasionally failing on some android devices.

diff --git a/Engine/Source/Runtime/RHI/Public/RHICommandList.h b/Engine/Source/Runtime/RHI/Public/RHICommandList.h
--- a/Engine/Source/Runtime/RHI/Public/RHICommandList.h
+++ b/Engine/Source/Runtime/RHI/Public/RHICommandList.h
@@ -3791,7 +3791,9 @@ public:
 	
 	FORCEINLINE void ReadSurfaceFloatData(FTextureRHIParamRef Texture,FIntRect Rect,TArray<FFloat16Color>& OutData,ECubeFace CubeFace,int32 ArrayIndex,int32 MipIndex)
 	{
+		QUICK_SCOPE_CYCLE_COUNTER(STAT_RHIMETHOD_ReadSurfaceFloatData_Flush);
 		LLM_SCOPE(ELLMTag::Textures);
+		ImmediateFlush(EImmediateFlushType::FlushRHIThread);
 		GDynamicRHI->RHIReadSurfaceFloatData_RenderThread(*this, Texture,Rect,OutData,CubeFace,ArrayIndex,MipIndex);
 	}
 
