Eliminate type mismatch possibility in command allocation (after being hit by one)

diff --git a/Engine/Source/Runtime/RHI/Private/RHICommandList.cpp b/Engine/Source/Runtime/RHI/Private/RHICommandList.cpp
--- a/Engine/Source/Runtime/RHI/Private/RHICommandList.cpp
+++ b/Engine/Source/Runtime/RHI/Private/RHICommandList.cpp
@@ -171,7 +171,7 @@ void FRHICommandListImmediate::SetCurrentStat(TStatId Stat)
 {
 	if (!Bypass())
 	{
-		new (AllocCommand<FRHICommandStat>()) FRHICommandStat(Stat);
+		ALLOC_COMMAND(FRHICommandStat)(Stat);
 	}
 }
 
@@ -753,7 +753,7 @@ FGraphEventRef FRHICommandListImmediate::RHIThreadFence(bool bSetLockFence)
 
 	if (IsRunningRHIInSeparateThread())
 	{
-		FRHICommandRHIThreadFence* Cmd = new (AllocCommand<FRHICommandRHIThreadFence>()) FRHICommandRHIThreadFence();
+		FRHICommandRHIThreadFence* Cmd = ALLOC_COMMAND(FRHICommandRHIThreadFence)();
 		if (bSetLockFence)
 		{
 			RHIThreadBufferLockFence = Cmd->Fence;
@@ -787,7 +787,7 @@ void FRHICommandListImmediate::QueueAsyncCompute(FRHIAsyncComputeCommandList& RH
 		delete &RHIComputeCmdList;
 		return;
 	}
-	new (AllocCommand<FRHIAsyncComputeSubmitList>()) FRHIAsyncComputeSubmitList(&RHIComputeCmdList);
+	ALLOC_COMMAND(FRHIAsyncComputeSubmitList)(&RHIComputeCmdList);
 }
 	
 void FRHICommandListExecutor::WaitOnRHIThreadFence(FGraphEventRef& Fence)
@@ -1107,7 +1107,7 @@ public:
 			for (int32 Index = 0; Index < NumCommandLists; Index++)
 			{
 				FRHICommandListBase* CmdList = RHICmdLists[Index];
-				new (RHICmdList->AllocCommand<FRHICommandWaitForAndSubmitSubList>()) FRHICommandWaitForAndSubmitSubList(Nothing, CmdList);
+				ALLOC_COMMAND_CL(*RHICmdList, FRHICommandWaitForAndSubmitSubList)(Nothing, CmdList);
 			}
 		}
 		else
@@ -1133,7 +1133,7 @@ public:
 				FGraphEventRef TranslateCompletionEvent = TGraphTask<FParallelTranslateCommandList>::CreateTask(nullptr, ENamedThreads::GetRenderThread()).ConstructAndDispatchWhenReady(&RHICmdLists[Start], 1 + Last - Start, ContextContainer, bIsPrepass);
 				MyCompletionGraphEvent->SetGatherThreadForDontCompleteUntil(ENamedThreads::Type(CPrio_FParallelTranslateCommandList.Get() | ENamedThreads::HighTaskPriority));
 				MyCompletionGraphEvent->DontCompleteUntil(TranslateCompletionEvent);
-				new (RHICmdList->AllocCommand<FRHICommandWaitForAndSubmitSubListParallel>()) FRHICommandWaitForAndSubmitSubListParallel(TranslateCompletionEvent, ContextContainer, EffectiveThreads, ThreadIndex++);
+				ALLOC_COMMAND_CL(*RHICmdList, FRHICommandWaitForAndSubmitSubListParallel)(TranslateCompletionEvent, ContextContainer, EffectiveThreads, ThreadIndex++);
 				Start = Last + 1;
 			}
 			check(EffectiveThreads == ThreadIndex);
@@ -1285,7 +1285,7 @@ void FRHICommandListBase::QueueParallelAsyncCommandListSubmit(FGraphEventRef* An
 				FGraphEventRef TranslateCompletionEvent = TGraphTask<FParallelTranslateCommandList>::CreateTask(&Prereq, ENamedThreads::GetRenderThread()).ConstructAndDispatchWhenReady(&RHICmdLists[0], 1 + Last - Start, ContextContainer, bIsPrepass);
 
 				AllOutstandingTasks.Add(TranslateCompletionEvent);
-				new (AllocCommand<FRHICommandWaitForAndSubmitSubListParallel>()) FRHICommandWaitForAndSubmitSubListParallel(TranslateCompletionEvent, ContextContainer, EffectiveThreads, ThreadIndex++);
+				ALLOC_COMMAND(FRHICommandWaitForAndSubmitSubListParallel)(TranslateCompletionEvent, ContextContainer, EffectiveThreads, ThreadIndex++);
 				if (IsRunningRHIInSeparateThread())
 				{
 					FRHICommandListExecutor::GetImmediateCommandList().ImmediateFlush(EImmediateFlushType::DispatchToRHIThread); // we don't want stuff after the async cmd list to be bundled with it
@@ -1316,7 +1316,7 @@ void FRHICommandListBase::QueueParallelAsyncCommandListSubmit(FGraphEventRef* An
 			}
 			WaitOutstandingTasks.Add(AnyThreadCompletionEvent);
 		}
-		new (AllocCommand<FRHICommandWaitForAndSubmitSubList>()) FRHICommandWaitForAndSubmitSubList(AnyThreadCompletionEvent, CmdList);
+		ALLOC_COMMAND(FRHICommandWaitForAndSubmitSubList)(AnyThreadCompletionEvent, CmdList);
 	}
 	if (IsRunningRHIInSeparateThread())
 	{
@@ -1340,7 +1340,7 @@ void FRHICommandListBase::QueueAsyncCommandListSubmit(FGraphEventRef& AnyThreadC
 		}
 		WaitOutstandingTasks.Add(AnyThreadCompletionEvent);
 	}
-	new (AllocCommand<FRHICommandWaitForAndSubmitSubList>()) FRHICommandWaitForAndSubmitSubList(AnyThreadCompletionEvent, CmdList);
+	ALLOC_COMMAND(FRHICommandWaitForAndSubmitSubList)(AnyThreadCompletionEvent, CmdList);
 	if (IsRunningRHIInSeparateThread())
 	{
 		FRHICommandListExecutor::GetImmediateCommandList().ImmediateFlush(EImmediateFlushType::DispatchToRHIThread); // we don't want stuff after the async cmd list to be bundled with it
@@ -1401,7 +1401,7 @@ void FRHICommandListBase::QueueRenderThreadCommandListSubmit(FGraphEventRef& Ren
 		check(!IsInActualRenderingThread() && !IsInGameThread() && !IsImmediate());
 		RTTasks.Add(RenderThreadCompletionEvent);
 	}
-	new (AllocCommand<FRHICommandWaitForAndSubmitRTSubList>()) FRHICommandWaitForAndSubmitRTSubList(RenderThreadCompletionEvent, CmdList);
+	ALLOC_COMMAND(FRHICommandWaitForAndSubmitRTSubList)(RenderThreadCompletionEvent, CmdList);
 }
 
 void FRHICommandListBase::QueueAsyncPipelineStateCompile(FGraphEventRef& AsyncCompileCompletionEvent)
@@ -1432,7 +1432,7 @@ struct FRHICommandSubmitSubList final : public FRHICommand<FRHICommandSubmitSubL
 
 void FRHICommandListBase::QueueCommandListSubmit(class FRHICommandList* CmdList)
 {
-	new (AllocCommand<FRHICommandSubmitSubList>()) FRHICommandSubmitSubList(CmdList);
+	ALLOC_COMMAND(FRHICommandSubmitSubList)(CmdList);
 }
 
 
@@ -1444,7 +1444,7 @@ void FRHICommandList::BeginScene()
 		CMD_CONTEXT(RHIBeginScene)();
 		return;
 	}
-	new (AllocCommand<FRHICommandBeginScene>()) FRHICommandBeginScene();
+	ALLOC_COMMAND(FRHICommandBeginScene)();
 	if (!IsRunningRHIInSeparateThread())
 	{
 		// if we aren't running an RHIThread, there is no good reason to buffer this frame advance stuff and that complicates state management, so flush everything out now
@@ -1460,7 +1460,7 @@ void FRHICommandList::EndScene()
 		CMD_CONTEXT(RHIEndScene)();
 		return;
 	}
-	new (AllocCommand<FRHICommandEndScene>()) FRHICommandEndScene();
+	ALLOC_COMMAND(FRHICommandEndScene)();
 	if (!IsRunningRHIInSeparateThread())
 	{
 		// if we aren't running an RHIThread, there is no good reason to buffer this frame advance stuff and that complicates state management, so flush everything out now
@@ -1478,7 +1478,7 @@ void FRHICommandList::BeginDrawingViewport(FViewportRHIParamRef Viewport, FTextu
 		CMD_CONTEXT(RHIBeginDrawingViewport)(Viewport, RenderTargetRHI);
 		return;
 	}
-	new (AllocCommand<FRHICommandBeginDrawingViewport>()) FRHICommandBeginDrawingViewport(Viewport, RenderTargetRHI);
+	ALLOC_COMMAND(FRHICommandBeginDrawingViewport)(Viewport, RenderTargetRHI);
 	if (!IsRunningRHIInSeparateThread())
 	{
 		// if we aren't running an RHIThread, there is no good reason to buffer this frame advance stuff and that complicates state management, so flush everything out now
@@ -1496,7 +1496,7 @@ void FRHICommandList::EndDrawingViewport(FViewportRHIParamRef Viewport, bool bPr
 	}
 	else
 	{
-		new (AllocCommand<FRHICommandEndDrawingViewport>()) FRHICommandEndDrawingViewport(Viewport, bPresent, bLockToVsync);
+		ALLOC_COMMAND(FRHICommandEndDrawingViewport)(Viewport, bPresent, bLockToVsync);
 
 		if (IsRunningRHIInSeparateThread())
 		{
@@ -1531,7 +1531,7 @@ void FRHICommandList::BeginFrame()
 		CMD_CONTEXT(RHIBeginFrame)();
 		return;
 	}
-	new (AllocCommand<FRHICommandBeginFrame>()) FRHICommandBeginFrame();
+	ALLOC_COMMAND(FRHICommandBeginFrame)();
 	if (!IsRunningRHIInSeparateThread())
 	{
 		// if we aren't running an RHIThread, there is no good reason to buffer this frame advance stuff and that complicates state management, so flush everything out now
@@ -1549,7 +1549,7 @@ void FRHICommandList::EndFrame()
 		CMD_CONTEXT(RHIEndFrame)();
 		return;
 	}
-	new (AllocCommand<FRHICommandEndFrame>()) FRHICommandEndFrame();
+	ALLOC_COMMAND(FRHICommandEndFrame)();
 	if (!IsRunningRHIInSeparateThread())
 	{
 		// if we aren't running an RHIThread, there is no good reason to buffer this frame advance stuff and that complicates state management, so flush everything out now
@@ -1998,7 +1998,7 @@ void FDynamicRHI::UnlockVertexBuffer_RenderThread(class FRHICommandListImmediate
 	}	
 	else
 	{
-		new (RHICmdList.AllocCommand<FRHICommandUpdateVertexBuffer>()) FRHICommandUpdateVertexBuffer(VertexBuffer, Params.Buffer, Params.Offset, Params.BufferSize);
+		ALLOC_COMMAND_CL(RHICmdList, FRHICommandUpdateVertexBuffer)(VertexBuffer, Params.Buffer, Params.Offset, Params.BufferSize);
 		RHICmdList.RHIThreadFence(true);
 		if (GLockTracker.TotalMemoryOutstanding > 256 * 1024)
 		{
@@ -2057,7 +2057,7 @@ void FDynamicRHI::UnlockIndexBuffer_RenderThread(class FRHICommandListImmediate&
 	}	
 	else
 	{
-		new (RHICmdList.AllocCommand<FRHICommandUpdateIndexBuffer>()) FRHICommandUpdateIndexBuffer(IndexBuffer, Params.Buffer, Params.Offset, Params.BufferSize);
+		ALLOC_COMMAND_CL(RHICmdList, FRHICommandUpdateIndexBuffer)(IndexBuffer, Params.Buffer, Params.Offset, Params.BufferSize);
 		RHICmdList.RHIThreadFence(true);
 		if (GLockTracker.TotalMemoryOutstanding > 256 * 1024)
 		{
@@ -2463,7 +2463,7 @@ void FRHICommandListImmediate::UpdateTextureReference(FTextureReferenceRHIParamR
 		CMD_CONTEXT(RHIUpdateTextureReference)(TextureRef, NewTexture);
 		return;
 	}
-	new (AllocCommand<FRHICommandUpdateTextureReference>()) FRHICommandUpdateTextureReference(TextureRef, NewTexture);
+	ALLOC_COMMAND(FRHICommandUpdateTextureReference)(TextureRef, NewTexture);
 	RHIThreadFence(true);
 	if (GetUsedMemory() > 256 * 1024)
 	{
diff --git a/Engine/Source/Runtime/RHI/Public/RHICommandList.h b/Engine/Source/Runtime/RHI/Public/RHICommandList.h
--- a/Engine/Source/Runtime/RHI/Public/RHICommandList.h
+++ b/Engine/Source/Runtime/RHI/Public/RHICommandList.h
@@ -401,16 +401,21 @@ public:
 		return NameCopy;
 	}
 
-	template <typename TCmd>
-	FORCEINLINE_DEBUGGABLE void* AllocCommand()
+	FORCEINLINE_DEBUGGABLE void* AllocCommand(int32 AllocSize, int32 Alignment)
 	{
 		checkSlow(!IsExecuting());
-		TCmd* Result = (TCmd*)Alloc<TCmd>();
+		FRHICommandBase* Result = (FRHICommandBase*) MemManager.Alloc(AllocSize, Alignment);
 		++NumCommands;
 		*CommandLink = Result;
 		CommandLink = &Result->Next;
 		return Result;
 	}
+	template <typename TCmd>
+	FORCEINLINE void* AllocCommand()
+	{
+		return AllocCommand(sizeof(TCmd), alignof(TCmd));
+	}
+
 	FORCEINLINE uint32 GetUID()
 	{
 		return UID;
@@ -1859,6 +1864,10 @@ struct FRHICommandUpdateTextureReference final : public FRHICommand<FRHICommandU
 #define CMD_CONTEXT(Method) GetContext().Method
 #define COMPUTE_CONTEXT(Method) GetComputeContext().Method
 
+// Using variadic macro because some types are fancy template<A,B> stuff, which gets broken off at the comma and interpreted as multiple arguments. 
+#define ALLOC_COMMAND(...) new ( AllocCommand(sizeof(__VA_ARGS__), alignof(__VA_ARGS__)) ) __VA_ARGS__
+#define ALLOC_COMMAND_CL(RHICmdList, ...) new ( (RHICmdList).AllocCommand(sizeof(__VA_ARGS__), alignof(__VA_ARGS__)) ) __VA_ARGS__
+
 template<> void FRHICommandSetShaderParameter<FComputeShaderRHIParamRef, ECmdList::ECompute>::Execute(FRHICommandListBase& CmdList);
 template<> void FRHICommandSetShaderUniformBuffer<FComputeShaderRHIParamRef, ECmdList::ECompute>::Execute(FRHICommandListBase& CmdList);
 template<> void FRHICommandSetShaderTexture<FComputeShaderRHIParamRef, ECmdList::ECompute>::Execute(FRHICommandListBase& CmdList);
@@ -1884,7 +1893,7 @@ public:
 			CMD_CONTEXT(RHIBeginUpdateMultiFrameResource)( Texture);
 			return;
 		}
-		new (AllocCommand<FRHICommandBeginUpdateMultiFrameResource>()) FRHICommandBeginUpdateMultiFrameResource( Texture);
+		ALLOC_COMMAND(FRHICommandBeginUpdateMultiFrameResource)(Texture);
 	}
 
 	FORCEINLINE_DEBUGGABLE void EndUpdateMultiFrameResource(FTextureRHIParamRef Texture)
@@ -1895,7 +1904,7 @@ public:
 			CMD_CONTEXT(RHIEndUpdateMultiFrameResource)(Texture);
 			return;
 		}
-		new (AllocCommand<FRHICommandEndUpdateMultiFrameResource>()) FRHICommandEndUpdateMultiFrameResource(Texture);
+		ALLOC_COMMAND(FRHICommandEndUpdateMultiFrameResource)(Texture);
 	}
 
 	FORCEINLINE_DEBUGGABLE void BeginUpdateMultiFrameResource(FUnorderedAccessViewRHIParamRef UAV)
@@ -1906,7 +1915,7 @@ public:
 			CMD_CONTEXT(RHIBeginUpdateMultiFrameResource)(UAV);
 			return;
 		}
-		new (AllocCommand<FRHICommandBeginUpdateMultiFrameUAV>()) FRHICommandBeginUpdateMultiFrameUAV(UAV);
+		ALLOC_COMMAND(FRHICommandBeginUpdateMultiFrameUAV)(UAV);
 	}
 
 	FORCEINLINE_DEBUGGABLE void EndUpdateMultiFrameResource(FUnorderedAccessViewRHIParamRef UAV)
@@ -1917,7 +1926,7 @@ public:
 			CMD_CONTEXT(RHIEndUpdateMultiFrameResource)(UAV);
 			return;
 		}
-		new (AllocCommand<FRHICommandEndUpdateMultiFrameUAV>()) FRHICommandEndUpdateMultiFrameUAV(UAV);
+		ALLOC_COMMAND(FRHICommandEndUpdateMultiFrameUAV)(UAV);
 	}
 
 	DEPRECATED(4.20, "BuildLocalGraphicsPipelineState is deprecated. Use PipelineStateCache::GetAndOrCreateGraphicsPipelineState() instead.")
@@ -1933,7 +1942,7 @@ public:
 		}
 		else
 		{
-			auto* Cmd = new (AllocCommand<FRHICommandBuildLocalGraphicsPipelineState>()) FRHICommandBuildLocalGraphicsPipelineState(this, Initializer);
+			auto* Cmd = ALLOC_COMMAND(FRHICommandBuildLocalGraphicsPipelineState)(this, Initializer);
 			Result.WorkArea = &Cmd->WorkArea;
 		}
 		return Result;
@@ -1948,7 +1957,7 @@ public:
 			CMD_CONTEXT(RHISetGraphicsPipelineState)(LocalGraphicsPipelineState.BypassGraphicsPipelineState);
 			return;
 		}
-		new (AllocCommand<FRHICommandSetLocalGraphicsPipelineState>()) FRHICommandSetLocalGraphicsPipelineState(this, LocalGraphicsPipelineState);
+		ALLOC_COMMAND(FRHICommandSetLocalGraphicsPipelineState)(this, LocalGraphicsPipelineState);
 	}
 	
 	FORCEINLINE_DEBUGGABLE FLocalUniformBuffer BuildLocalUniformBuffer(const void* Contents, uint32 ContentsSize, const FRHIUniformBufferLayout& Layout)
@@ -1962,7 +1971,7 @@ public:
 		else
 		{
 			check(Contents && ContentsSize && (&Layout != nullptr));
-			auto* Cmd = new (AllocCommand<FRHICommandBuildLocalUniformBuffer>()) FRHICommandBuildLocalUniformBuffer(this, Contents, ContentsSize, Layout);
+			auto* Cmd = ALLOC_COMMAND(FRHICommandBuildLocalUniformBuffer)(this, Contents, ContentsSize, Layout);
 			Result.WorkArea = &Cmd->WorkArea;
 		}
 		return Result;
@@ -1977,7 +1986,7 @@ public:
 			CMD_CONTEXT(RHISetShaderUniformBuffer)(Shader, BaseIndex, UniformBuffer.BypassUniform);
 			return;
 		}
-		new (AllocCommand<FRHICommandSetLocalUniformBuffer<TShaderRHIParamRef> >()) FRHICommandSetLocalUniformBuffer<TShaderRHIParamRef>(this, Shader, BaseIndex, UniformBuffer);
+		ALLOC_COMMAND(FRHICommandSetLocalUniformBuffer<TShaderRHIParamRef>)(this, Shader, BaseIndex, UniformBuffer);
 	}
 
 	template <typename TShaderRHI>
@@ -1995,7 +2004,7 @@ public:
 			CMD_CONTEXT(RHISetShaderUniformBuffer)(Shader, BaseIndex, UniformBuffer);
 			return;
 		}
-		new (AllocCommand<FRHICommandSetShaderUniformBuffer<TShaderRHI*, ECmdList::EGfx> >()) FRHICommandSetShaderUniformBuffer<TShaderRHI*, ECmdList::EGfx>(Shader, BaseIndex, UniformBuffer);
+		ALLOC_COMMAND(FRHICommandSetShaderUniformBuffer<TShaderRHI*, ECmdList::EGfx>)(Shader, BaseIndex, UniformBuffer);
 	}
 	template <typename TShaderRHI>
 	FORCEINLINE void SetShaderUniformBuffer(TRefCountPtr<TShaderRHI>& Shader, uint32 BaseIndex, FUniformBufferRHIParamRef UniformBuffer)
@@ -2014,7 +2023,7 @@ public:
 		}
 		void* UseValue = Alloc(NumBytes, 16);
 		FMemory::Memcpy(UseValue, NewValue, NumBytes);
-		new (AllocCommand<FRHICommandSetShaderParameter<TShaderRHI*, ECmdList::EGfx> >()) FRHICommandSetShaderParameter<TShaderRHI*, ECmdList::EGfx>(Shader, BufferIndex, BaseIndex, NumBytes, UseValue);
+		ALLOC_COMMAND(FRHICommandSetShaderParameter<TShaderRHI*, ECmdList::EGfx>)(Shader, BufferIndex, BaseIndex, NumBytes, UseValue);
 	}
 	template <typename TShaderRHI>
 	FORCEINLINE void SetShaderParameter(TRefCountPtr<TShaderRHI>& Shader, uint32 BufferIndex, uint32 BaseIndex, uint32 NumBytes, const void* NewValue)
@@ -2031,7 +2040,7 @@ public:
 			CMD_CONTEXT(RHISetShaderTexture)(Shader, TextureIndex, Texture);
 			return;
 		}
-		new (AllocCommand<FRHICommandSetShaderTexture<TShaderRHIParamRef, ECmdList::EGfx> >()) FRHICommandSetShaderTexture<TShaderRHIParamRef, ECmdList::EGfx>(Shader, TextureIndex, Texture);
+		ALLOC_COMMAND(FRHICommandSetShaderTexture<TShaderRHIParamRef, ECmdList::EGfx>)(Shader, TextureIndex, Texture);
 	}
 
 	template <typename TShaderRHI>
@@ -2049,7 +2058,7 @@ public:
 			CMD_CONTEXT(RHISetShaderResourceViewParameter)(Shader, SamplerIndex, SRV);
 			return;
 		}
-		new (AllocCommand<FRHICommandSetShaderResourceViewParameter<TShaderRHIParamRef, ECmdList::EGfx> >()) FRHICommandSetShaderResourceViewParameter<TShaderRHIParamRef, ECmdList::EGfx>(Shader, SamplerIndex, SRV);
+		ALLOC_COMMAND(FRHICommandSetShaderResourceViewParameter<TShaderRHIParamRef, ECmdList::EGfx>)(Shader, SamplerIndex, SRV);
 	}
 
 	template <typename TShaderRHI>
@@ -2075,7 +2084,7 @@ public:
 			CMD_CONTEXT(RHISetShaderSampler)(Shader, SamplerIndex, State);
 			return;
 		}
-		new (AllocCommand<FRHICommandSetShaderSampler<TShaderRHIParamRef, ECmdList::EGfx> >()) FRHICommandSetShaderSampler<TShaderRHIParamRef, ECmdList::EGfx>(Shader, SamplerIndex, State);
+		ALLOC_COMMAND(FRHICommandSetShaderSampler<TShaderRHIParamRef, ECmdList::EGfx>)(Shader, SamplerIndex, State);
 	}
 
 	template <typename TShaderRHI>
@@ -2091,7 +2100,7 @@ public:
 			CMD_CONTEXT(RHISetUAVParameter)(Shader, UAVIndex, UAV);
 			return;
 		}
-		new (AllocCommand<FRHICommandSetUAVParameter<FComputeShaderRHIParamRef, ECmdList::EGfx> >()) FRHICommandSetUAVParameter<FComputeShaderRHIParamRef, ECmdList::EGfx>(Shader, UAVIndex, UAV);
+		ALLOC_COMMAND(FRHICommandSetUAVParameter<FComputeShaderRHIParamRef, ECmdList::EGfx>)(Shader, UAVIndex, UAV);
 	}
 
 	FORCEINLINE_DEBUGGABLE void SetUAVParameter(TRefCountPtr<FRHIComputeShader>& Shader, uint32 UAVIndex, FUnorderedAccessViewRHIParamRef UAV)
@@ -2106,7 +2115,7 @@ public:
 			CMD_CONTEXT(RHISetUAVParameter)(Shader, UAVIndex, UAV, InitialCount);
 			return;
 		}
-		new (AllocCommand<FRHICommandSetUAVParameter_IntialCount<FComputeShaderRHIParamRef, ECmdList::EGfx> >()) FRHICommandSetUAVParameter_IntialCount<FComputeShaderRHIParamRef, ECmdList::EGfx>(Shader, UAVIndex, UAV, InitialCount);
+		ALLOC_COMMAND(FRHICommandSetUAVParameter_IntialCount<FComputeShaderRHIParamRef, ECmdList::EGfx>)(Shader, UAVIndex, UAV, InitialCount);
 	}
 
 	FORCEINLINE_DEBUGGABLE void SetUAVParameter(TRefCountPtr<FRHIComputeShader>& Shader, uint32 UAVIndex, FUnorderedAccessViewRHIParamRef UAV, uint32 InitialCount)
@@ -2122,7 +2131,7 @@ public:
 			CMD_CONTEXT(RHISetBlendFactor)(BlendFactor);
 			return;
 		}
-		new (AllocCommand<FRHICommandSetBlendFactor>()) FRHICommandSetBlendFactor(BlendFactor);
+		ALLOC_COMMAND(FRHICommandSetBlendFactor)(BlendFactor);
 	}
 
 	FORCEINLINE_DEBUGGABLE void DrawPrimitive(uint32 BaseVertexIndex, uint32 NumPrimitives, uint32 NumInstances)
@@ -2133,7 +2142,7 @@ public:
 			CMD_CONTEXT(RHIDrawPrimitive)(BaseVertexIndex, NumPrimitives, NumInstances);
 			return;
 		}
-		new (AllocCommand<FRHICommandDrawPrimitive>()) FRHICommandDrawPrimitive(BaseVertexIndex, NumPrimitives, NumInstances);
+		ALLOC_COMMAND(FRHICommandDrawPrimitive)(BaseVertexIndex, NumPrimitives, NumInstances);
 	}
 
 	DEPRECATED(4.22, "PrimitiveType is now a part of the Graphics Pipeline State.")
@@ -2155,7 +2164,7 @@ public:
 			CMD_CONTEXT(RHIDrawIndexedPrimitive)(IndexBuffer, BaseVertexIndex, FirstInstance, NumVertices, StartIndex, NumPrimitives, NumInstances);
 			return;
 		}
-		new (AllocCommand<FRHICommandDrawIndexedPrimitive>()) FRHICommandDrawIndexedPrimitive(IndexBuffer, BaseVertexIndex, FirstInstance, NumVertices, StartIndex, NumPrimitives, NumInstances);
+		ALLOC_COMMAND(FRHICommandDrawIndexedPrimitive)(IndexBuffer, BaseVertexIndex, FirstInstance, NumVertices, StartIndex, NumPrimitives, NumInstances);
 	}
 
 	DEPRECATED(4.22, "PrimitiveType is now a part of the Graphics Pipeline State.")
@@ -2171,7 +2180,7 @@ public:
 			CMD_CONTEXT(RHISetStreamSource)(StreamIndex, VertexBuffer, Offset);
 			return;
 		}
-		new (AllocCommand<FRHICommandSetStreamSource>()) FRHICommandSetStreamSource(StreamIndex, VertexBuffer, Offset);
+		ALLOC_COMMAND(FRHICommandSetStreamSource)(StreamIndex, VertexBuffer, Offset);
 	}
 
 	FORCEINLINE_DEBUGGABLE void SetStencilRef(uint32 StencilRef)
@@ -2183,7 +2192,7 @@ public:
 			return;
 		}
 
-		new(AllocCommand<FRHICommandSetStencilRef>()) FRHICommandSetStencilRef(StencilRef);
+		ALLOC_COMMAND(FRHICommandSetStencilRef)(StencilRef);
 	}
 
 	FORCEINLINE_DEBUGGABLE void SetViewport(uint32 MinX, uint32 MinY, float MinZ, uint32 MaxX, uint32 MaxY, float MaxZ)
@@ -2194,7 +2203,7 @@ public:
 			CMD_CONTEXT(RHISetViewport)(MinX, MinY, MinZ, MaxX, MaxY, MaxZ);
 			return;
 		}
-		new (AllocCommand<FRHICommandSetViewport>()) FRHICommandSetViewport(MinX, MinY, MinZ, MaxX, MaxY, MaxZ);
+		ALLOC_COMMAND(FRHICommandSetViewport)(MinX, MinY, MinZ, MaxX, MaxY, MaxZ);
 	}
 
 	FORCEINLINE_DEBUGGABLE void SetStereoViewport(uint32 LeftMinX, uint32 RightMinX, uint32 LeftMinY, uint32 RightMinY, float MinZ, uint32 LeftMaxX, uint32 RightMaxX, uint32 LeftMaxY, uint32 RightMaxY, float MaxZ)
@@ -2205,7 +2214,7 @@ public:
 			CMD_CONTEXT(RHISetStereoViewport)(LeftMinX, RightMinX, LeftMinY, RightMinY, MinZ, LeftMaxX, RightMaxX, LeftMaxY, RightMaxY, MaxZ);
 			return;
 		}
-		new (AllocCommand<FRHICommandSetStereoViewport>()) FRHICommandSetStereoViewport(LeftMinX, RightMinX, LeftMinY, RightMinY, MinZ, LeftMaxX, RightMaxX, LeftMaxY, RightMaxY, MaxZ);
+		ALLOC_COMMAND(FRHICommandSetStereoViewport)(LeftMinX, RightMinX, LeftMinY, RightMinY, MinZ, LeftMaxX, RightMaxX, LeftMaxY, RightMaxY, MaxZ);
 	}
 
 	FORCEINLINE_DEBUGGABLE void SetScissorRect(bool bEnable, uint32 MinX, uint32 MinY, uint32 MaxX, uint32 MaxY)
@@ -2216,7 +2225,7 @@ public:
 			CMD_CONTEXT(RHISetScissorRect)(bEnable, MinX, MinY, MaxX, MaxY);
 			return;
 		}
-		new (AllocCommand<FRHICommandSetScissorRect>()) FRHICommandSetScissorRect(bEnable, MinX, MinY, MaxX, MaxY);
+		ALLOC_COMMAND(FRHICommandSetScissorRect)(bEnable, MinX, MinY, MaxX, MaxY);
 	}
 
 	void ApplyCachedRenderTargets(
@@ -2290,7 +2299,7 @@ public:
 				UAVs);
 			return;
 		}
-		new (AllocCommand<FRHICommandSetRenderTargets>()) FRHICommandSetRenderTargets(
+		ALLOC_COMMAND(FRHICommandSetRenderTargets)(
 			NewNumSimultaneousRenderTargets,
 			NewRenderTargetsRHI,
 			NewDepthStencilTargetRHI,
@@ -2312,7 +2321,7 @@ public:
 			CMD_CONTEXT(RHISetRenderTargetsAndClear)(RenderTargetsInfo);
 			return;
 		}
-		new (AllocCommand<FRHICommandSetRenderTargetsAndClear>()) FRHICommandSetRenderTargetsAndClear(RenderTargetsInfo);
+		ALLOC_COMMAND(FRHICommandSetRenderTargetsAndClear)(RenderTargetsInfo);
 	}	
 
 	FORCEINLINE_DEBUGGABLE void BindClearMRTValues(bool bClearColor, bool bClearDepth, bool bClearStencil)
@@ -2323,7 +2332,7 @@ public:
 			CMD_CONTEXT(RHIBindClearMRTValues)(bClearColor, bClearDepth, bClearStencil);
 			return;
 		}
-		new (AllocCommand<FRHICommandBindClearMRTValues>()) FRHICommandBindClearMRTValues(bClearColor, bClearDepth, bClearStencil);
+		ALLOC_COMMAND(FRHICommandBindClearMRTValues)(bClearColor, bClearDepth, bClearStencil);
 	}	
 
 	DEPRECATED(4.21, "This function is deprecated and will be removed in future releases.")
@@ -2356,7 +2365,7 @@ public:
 			return;
 		}
 		check(DrawUPData.OutVertexData && DrawUPData.NumVertices);
-		new (AllocCommand<FRHICommandEndDrawPrimitiveUP>()) FRHICommandEndDrawPrimitiveUP(DrawUPData.NumPrimitives, DrawUPData.NumVertices, DrawUPData.VertexDataStride, DrawUPData.OutVertexData);
+		ALLOC_COMMAND(FRHICommandEndDrawPrimitiveUP)(DrawUPData.NumPrimitives, DrawUPData.NumVertices, DrawUPData.VertexDataStride, DrawUPData.OutVertexData);
 
 		DrawUPData.OutVertexData = nullptr;
 		DrawUPData.NumVertices = 0;
@@ -2398,7 +2407,7 @@ public:
 		}
 		check(DrawUPData.OutVertexData && DrawUPData.OutIndexData && DrawUPData.NumIndices && DrawUPData.NumVertices);
 
-		new (AllocCommand<FRHICommandEndDrawIndexedPrimitiveUP>()) FRHICommandEndDrawIndexedPrimitiveUP(DrawUPData.NumPrimitives, DrawUPData.NumVertices, DrawUPData.VertexDataStride, DrawUPData.OutVertexData, DrawUPData.MinVertexIndex, DrawUPData.NumIndices, DrawUPData.IndexDataStride, DrawUPData.OutIndexData);
+		ALLOC_COMMAND(FRHICommandEndDrawIndexedPrimitiveUP)(DrawUPData.NumPrimitives, DrawUPData.NumVertices, DrawUPData.VertexDataStride, DrawUPData.OutVertexData, DrawUPData.MinVertexIndex, DrawUPData.NumIndices, DrawUPData.IndexDataStride, DrawUPData.OutIndexData);
 
 		DrawUPData.OutVertexData = nullptr;
 		DrawUPData.OutIndexData = nullptr;
@@ -2414,7 +2423,7 @@ public:
 			CMD_CONTEXT(RHISetComputeShader)(ComputeShader);
 			return;
 		}
-		new (AllocCommand<FRHICommandSetComputeShader<ECmdList::EGfx>>()) FRHICommandSetComputeShader<ECmdList::EGfx>(ComputeShader);
+		ALLOC_COMMAND(FRHICommandSetComputeShader<ECmdList::EGfx>)(ComputeShader);
 	}
 
 	FORCEINLINE_DEBUGGABLE void SetComputePipelineState(class FComputePipelineState* ComputePipelineState)
@@ -2426,7 +2435,7 @@ public:
 			CMD_CONTEXT(RHISetComputePipelineState)(RHIComputePipelineState);
 			return;
 		}
-		new (AllocCommand<FRHICommandSetComputePipelineState<ECmdList::EGfx>>()) FRHICommandSetComputePipelineState<ECmdList::EGfx>(ComputePipelineState);
+		ALLOC_COMMAND(FRHICommandSetComputePipelineState<ECmdList::EGfx>)(ComputePipelineState);
 	}
 
 	FORCEINLINE_DEBUGGABLE void SetGraphicsPipelineState(class FGraphicsPipelineState* GraphicsPipelineState)
@@ -2439,7 +2448,7 @@ public:
 			CMD_CONTEXT(RHISetGraphicsPipelineState)(RHIGraphicsPipelineState);
 			return;
 		}
-		new (AllocCommand<FRHICommandSetGraphicsPipelineState>()) FRHICommandSetGraphicsPipelineState(GraphicsPipelineState);
+		ALLOC_COMMAND(FRHICommandSetGraphicsPipelineState)(GraphicsPipelineState);
 	}
 
 	FORCEINLINE_DEBUGGABLE void DispatchComputeShader(uint32 ThreadGroupCountX, uint32 ThreadGroupCountY, uint32 ThreadGroupCountZ)
@@ -2449,7 +2458,7 @@ public:
 			CMD_CONTEXT(RHIDispatchComputeShader)(ThreadGroupCountX, ThreadGroupCountY, ThreadGroupCountZ);
 			return;
 		}
-		new (AllocCommand<FRHICommandDispatchComputeShader<ECmdList::EGfx>>()) FRHICommandDispatchComputeShader<ECmdList::EGfx>(ThreadGroupCountX, ThreadGroupCountY, ThreadGroupCountZ);
+		ALLOC_COMMAND(FRHICommandDispatchComputeShader<ECmdList::EGfx>)(ThreadGroupCountX, ThreadGroupCountY, ThreadGroupCountZ);
 	}
 
 	FORCEINLINE_DEBUGGABLE void DispatchIndirectComputeShader(FVertexBufferRHIParamRef ArgumentBuffer, uint32 ArgumentOffset)
@@ -2459,7 +2468,7 @@ public:
 			CMD_CONTEXT(RHIDispatchIndirectComputeShader)(ArgumentBuffer, ArgumentOffset);
 			return;
 		}
-		new (AllocCommand<FRHICommandDispatchIndirectComputeShader<ECmdList::EGfx>>()) FRHICommandDispatchIndirectComputeShader<ECmdList::EGfx>(ArgumentBuffer, ArgumentOffset);
+		ALLOC_COMMAND(FRHICommandDispatchIndirectComputeShader<ECmdList::EGfx>)(ArgumentBuffer, ArgumentOffset);
 	}
 
 	FORCEINLINE_DEBUGGABLE void AutomaticCacheFlushAfterComputeShader(bool bEnable)
@@ -2469,7 +2478,7 @@ public:
 			CMD_CONTEXT(RHIAutomaticCacheFlushAfterComputeShader)(bEnable);
 			return;
 		}
-		new (AllocCommand<FRHICommandAutomaticCacheFlushAfterComputeShader>()) FRHICommandAutomaticCacheFlushAfterComputeShader(bEnable);
+		ALLOC_COMMAND(FRHICommandAutomaticCacheFlushAfterComputeShader)(bEnable);
 	}
 
 	FORCEINLINE_DEBUGGABLE void FlushComputeShaderCache()
@@ -2479,7 +2488,7 @@ public:
 			CMD_CONTEXT(RHIFlushComputeShaderCache)();
 			return;
 		}
-		new (AllocCommand<FRHICommandFlushComputeShaderCache>()) FRHICommandFlushComputeShaderCache();
+		ALLOC_COMMAND(FRHICommandFlushComputeShaderCache)();
 	}
 
 	FORCEINLINE_DEBUGGABLE void DrawPrimitiveIndirect(FVertexBufferRHIParamRef ArgumentBuffer, uint32 ArgumentOffset)
@@ -2490,7 +2499,7 @@ public:
 			CMD_CONTEXT(RHIDrawPrimitiveIndirect)(ArgumentBuffer, ArgumentOffset);
 			return;
 		}
-		new (AllocCommand<FRHICommandDrawPrimitiveIndirect>()) FRHICommandDrawPrimitiveIndirect(ArgumentBuffer, ArgumentOffset);
+		ALLOC_COMMAND(FRHICommandDrawPrimitiveIndirect)(ArgumentBuffer, ArgumentOffset);
 	}
 
 	DEPRECATED(4.22, "PrimitiveType is now a part of the Graphics Pipeline State.")
@@ -2508,7 +2517,7 @@ public:
 			CMD_CONTEXT(RHIDrawIndexedIndirect)(IndexBufferRHI, ArgumentsBufferRHI, DrawArgumentsIndex, NumInstances);
 			return;
 		}
-		new (AllocCommand<FRHICommandDrawIndexedIndirect>()) FRHICommandDrawIndexedIndirect(IndexBufferRHI, ArgumentsBufferRHI, DrawArgumentsIndex, NumInstances);
+		ALLOC_COMMAND(FRHICommandDrawIndexedIndirect)(IndexBufferRHI, ArgumentsBufferRHI, DrawArgumentsIndex, NumInstances);
 	}
 
 	DEPRECATED(4.22, "PrimitiveType is now a part of the Graphics Pipeline State.")
@@ -2525,7 +2534,7 @@ public:
 			CMD_CONTEXT(RHIDrawIndexedPrimitiveIndirect)(IndexBuffer, ArgumentsBuffer, ArgumentOffset);
 			return;
 		}
-		new (AllocCommand<FRHICommandDrawIndexedPrimitiveIndirect>()) FRHICommandDrawIndexedPrimitiveIndirect(IndexBuffer, ArgumentsBuffer, ArgumentOffset);
+		ALLOC_COMMAND(FRHICommandDrawIndexedPrimitiveIndirect)(IndexBuffer, ArgumentsBuffer, ArgumentOffset);
 	}
 
 	DEPRECATED(4.22, "PrimitiveType is now a part of the Graphics Pipeline State.")
@@ -2542,7 +2551,7 @@ public:
 			CMD_CONTEXT(RHISetDepthBounds)(MinDepth, MaxDepth);
 			return;
 		}
-		new (AllocCommand<FRHICommandSetDepthBounds>()) FRHICommandSetDepthBounds(MinDepth, MaxDepth);
+		ALLOC_COMMAND(FRHICommandSetDepthBounds)(MinDepth, MaxDepth);
 	}
 
 	FORCEINLINE_DEBUGGABLE void CopyToResolveTarget(FTextureRHIParamRef SourceTextureRHI, FTextureRHIParamRef DestTextureRHI, const FResolveParams& ResolveParams)
@@ -2553,7 +2562,7 @@ public:
 			CMD_CONTEXT(RHICopyToResolveTarget)(SourceTextureRHI, DestTextureRHI, ResolveParams);
 			return;
 		}
-		new (AllocCommand<FRHICommandCopyToResolveTarget>()) FRHICommandCopyToResolveTarget(SourceTextureRHI, DestTextureRHI, ResolveParams);
+		ALLOC_COMMAND(FRHICommandCopyToResolveTarget)(SourceTextureRHI, DestTextureRHI, ResolveParams);
 	}
 
 	DEPRECATED(4.20, "This signature of CopyToResolveTarget is deprecated. Please use the new one instead.")
@@ -2570,7 +2579,7 @@ public:
 			CMD_CONTEXT(RHICopyTexture)(SourceTextureRHI, DestTextureRHI, CopyInfo);
 			return;
 		}
-		new (AllocCommand<FRHICommandCopyTexture>()) FRHICommandCopyTexture(SourceTextureRHI, DestTextureRHI, CopyInfo);
+		ALLOC_COMMAND(FRHICommandCopyTexture)(SourceTextureRHI, DestTextureRHI, CopyInfo);
 	}
 
 	FORCEINLINE_DEBUGGABLE void ClearTinyUAV(FUnorderedAccessViewRHIParamRef UnorderedAccessViewRHI, const uint32(&Values)[4])
@@ -2581,7 +2590,7 @@ public:
 			CMD_CONTEXT(RHIClearTinyUAV)(UnorderedAccessViewRHI, Values);
 			return;
 		}
-		new (AllocCommand<FRHICommandClearTinyUAV>()) FRHICommandClearTinyUAV(UnorderedAccessViewRHI, Values);
+		ALLOC_COMMAND(FRHICommandClearTinyUAV)(UnorderedAccessViewRHI, Values);
 	}
 
 	FORCEINLINE_DEBUGGABLE void BeginRenderQuery(FRenderQueryRHIParamRef RenderQuery)
@@ -2591,7 +2600,7 @@ public:
 			CMD_CONTEXT(RHIBeginRenderQuery)(RenderQuery);
 			return;
 		}
-		new (AllocCommand<FRHICommandBeginRenderQuery>()) FRHICommandBeginRenderQuery(RenderQuery);
+		ALLOC_COMMAND(FRHICommandBeginRenderQuery)(RenderQuery);
 	}
 	FORCEINLINE_DEBUGGABLE void EndRenderQuery(FRenderQueryRHIParamRef RenderQuery)
 	{
@@ -2600,7 +2609,7 @@ public:
 			CMD_CONTEXT(RHIEndRenderQuery)(RenderQuery);
 			return;
 		}
-		new (AllocCommand<FRHICommandEndRenderQuery>()) FRHICommandEndRenderQuery(RenderQuery);
+		ALLOC_COMMAND(FRHICommandEndRenderQuery)(RenderQuery);
 	}
 
 	FORCEINLINE_DEBUGGABLE void SubmitCommandsHint()
@@ -2610,7 +2619,7 @@ public:
 			CMD_CONTEXT(RHISubmitCommandsHint)();
 			return;
 		}
-		new (AllocCommand<FRHICommandSubmitCommandsHint<ECmdList::EGfx>>()) FRHICommandSubmitCommandsHint<ECmdList::EGfx>();
+		ALLOC_COMMAND(FRHICommandSubmitCommandsHint<ECmdList::EGfx>)();
 	}
 
 	FORCEINLINE_DEBUGGABLE void PollOcclusionQueries()
@@ -2620,7 +2629,7 @@ public:
 			CMD_CONTEXT(RHIPollOcclusionQueries)();
 			return;
 		}
-		new (AllocCommand<FRHICommandPollOcclusionQueries>()) FRHICommandPollOcclusionQueries();
+		ALLOC_COMMAND(FRHICommandPollOcclusionQueries)();
 	}
 
 	FORCEINLINE_DEBUGGABLE void TransitionResource(EResourceTransitionAccess TransitionType, FTextureRHIParamRef InTexture)
@@ -2636,7 +2645,7 @@ public:
 		// Allocate space to hold the single texture pointer inline in the command list itself.
 		FTextureRHIParamRef* TextureArray = (FTextureRHIParamRef*)Alloc(sizeof(FTextureRHIParamRef), alignof(FTextureRHIParamRef));
 		TextureArray[0] = Texture;
-		new (AllocCommand<FRHICommandTransitionTextures>()) FRHICommandTransitionTextures(TransitionType, TextureArray, 1);
+		ALLOC_COMMAND(FRHICommandTransitionTextures)(TransitionType, TextureArray, 1);
 	}
 
 	FORCEINLINE_DEBUGGABLE void TransitionResources(EResourceTransitionAccess TransitionType, FTextureRHIParamRef* InTextures, int32 NumTextures)
@@ -2654,7 +2663,7 @@ public:
 			InlineTextureArray[Index] = InTextures[Index];
 		}
 
-		new (AllocCommand<FRHICommandTransitionTextures>()) FRHICommandTransitionTextures(TransitionType, InlineTextureArray, NumTextures);
+		ALLOC_COMMAND(FRHICommandTransitionTextures)(TransitionType, InlineTextureArray, NumTextures);
 	}
 
 	FORCEINLINE_DEBUGGABLE void TransitionResourceArrayNoCopy(EResourceTransitionAccess TransitionType, TArray<FTextureRHIParamRef>& InTextures)
@@ -2664,7 +2673,7 @@ public:
 			CMD_CONTEXT(RHITransitionResources)(TransitionType, &InTextures[0], InTextures.Num());
 			return;
 		}
-		new (AllocCommand<FRHICommandTransitionTexturesArray>()) FRHICommandTransitionTexturesArray(TransitionType, InTextures);
+		ALLOC_COMMAND(FRHICommandTransitionTexturesArray)(TransitionType, InTextures);
 	}
 
 	FORCEINLINE_DEBUGGABLE void TransitionResource(EResourceTransitionAccess TransitionType, EResourceTransitionPipeline TransitionPipeline, FUnorderedAccessViewRHIParamRef InUAV, FComputeFenceRHIParamRef WriteFence)
@@ -2680,7 +2689,7 @@ public:
 		// Allocate space to hold the single UAV pointer inline in the command list itself.
 		FUnorderedAccessViewRHIParamRef* UAVArray = (FUnorderedAccessViewRHIParamRef*)Alloc(sizeof(FUnorderedAccessViewRHIParamRef), alignof(FUnorderedAccessViewRHIParamRef));
 		UAVArray[0] = UAV;
-		new (AllocCommand<FRHICommandTransitionUAVs<ECmdList::EGfx>>()) FRHICommandTransitionUAVs<ECmdList::EGfx>(TransitionType, TransitionPipeline, UAVArray, 1, WriteFence);
+		ALLOC_COMMAND(FRHICommandTransitionUAVs<ECmdList::EGfx>)(TransitionType, TransitionPipeline, UAVArray, 1, WriteFence);
 	}
 
 	FORCEINLINE_DEBUGGABLE void TransitionResource(EResourceTransitionAccess TransitionType, EResourceTransitionPipeline TransitionPipeline, FUnorderedAccessViewRHIParamRef InUAV)
@@ -2703,7 +2712,7 @@ public:
 			UAVArray[Index] = InUAVs[Index];
 		}
 
-		new (AllocCommand<FRHICommandTransitionUAVs<ECmdList::EGfx>>()) FRHICommandTransitionUAVs<ECmdList::EGfx>(TransitionType, TransitionPipeline, UAVArray, NumUAVs, WriteFence);
+		ALLOC_COMMAND(FRHICommandTransitionUAVs<ECmdList::EGfx>)(TransitionType, TransitionPipeline, UAVArray, NumUAVs, WriteFence);
 	}
 
 	FORCEINLINE_DEBUGGABLE void TransitionResources(EResourceTransitionAccess TransitionType, EResourceTransitionPipeline TransitionPipeline, FUnorderedAccessViewRHIParamRef* InUAVs, int32 NumUAVs)
@@ -2718,7 +2727,7 @@ public:
 			CMD_CONTEXT(RHIWaitComputeFence)(WaitFence);
 			return;
 		}
-		new (AllocCommand<FRHICommandWaitComputeFence<ECmdList::EGfx>>()) FRHICommandWaitComputeFence<ECmdList::EGfx>(WaitFence);
+		ALLOC_COMMAND(FRHICommandWaitComputeFence<ECmdList::EGfx>)(WaitFence);
 	}
 
 	FORCEINLINE_DEBUGGABLE void EnqueueStagedRead(FStagingBufferRHIParamRef StagingBuffer, FGPUFenceRHIParamRef Fence, uint32 Offset, uint32 NumBytes)
@@ -2728,7 +2737,7 @@ public:
 			CMD_CONTEXT(RHIEnqueueStagedRead)(StagingBuffer, Fence, Offset, NumBytes);
 			return;
 		}
-		new (AllocCommand<FRHICommandEnqueueStagedRead<ECmdList::EGfx>>()) FRHICommandEnqueueStagedRead<ECmdList::EGfx>(StagingBuffer, Fence, Offset, NumBytes);
+		ALLOC_COMMAND(FRHICommandEnqueueStagedRead<ECmdList::EGfx>)(StagingBuffer, Fence, Offset, NumBytes);
 	}
 
 	FORCEINLINE_DEBUGGABLE void BeginRenderPass(const FRHIRenderPassInfo& InInfo, const TCHAR* Name)
@@ -2745,7 +2754,7 @@ public:
 		else
 		{
 			TCHAR* NameCopy  = AllocString(Name);
-			new (AllocCommand<FRHICommandBeginRenderPass>()) FRHICommandBeginRenderPass(InInfo, NameCopy);
+			ALLOC_COMMAND(FRHICommandBeginRenderPass)(InInfo, NameCopy);
 		}
 		Data.bInsideRenderPass = true;
 
@@ -2763,7 +2772,7 @@ public:
 		}
 		else
 		{
-			new (AllocCommand<FRHICommandEndRenderPass>()) FRHICommandEndRenderPass();
+			ALLOC_COMMAND(FRHICommandEndRenderPass)();
 		}
 		Data.bInsideRenderPass = false;
 	}
@@ -2780,7 +2789,7 @@ public:
 		else
 		{
 			TCHAR* NameCopy  = AllocString(Name);
-			new (AllocCommand<FRHICommandBeginComputePass>()) FRHICommandBeginComputePass(NameCopy);
+			ALLOC_COMMAND(FRHICommandBeginComputePass)(NameCopy);
 		}
 		Data.bInsideComputePass = true;
 
@@ -2797,7 +2806,7 @@ public:
 		}
 		else
 		{
-			new (AllocCommand<FRHICommandEndComputePass>()) FRHICommandEndComputePass();
+			ALLOC_COMMAND(FRHICommandEndComputePass)();
 		}
 		Data.bInsideComputePass = false;
 	}
@@ -2818,7 +2827,7 @@ public:
 			return;
 		}
 		TCHAR* NameCopy  = AllocString(Name);
-		new (AllocCommand<FRHICommandPushEvent<ECmdList::EGfx>>()) FRHICommandPushEvent<ECmdList::EGfx>(NameCopy, Color);
+		ALLOC_COMMAND(FRHICommandPushEvent<ECmdList::EGfx>)(NameCopy, Color);
 	}
 
 	FORCEINLINE_DEBUGGABLE void PopEvent()
@@ -2828,7 +2837,7 @@ public:
 			CMD_CONTEXT(RHIPopEvent)();
 			return;
 		}
-		new (AllocCommand<FRHICommandPopEvent<ECmdList::EGfx>>()) FRHICommandPopEvent<ECmdList::EGfx>();
+		ALLOC_COMMAND(FRHICommandPopEvent<ECmdList::EGfx>)();
 	}
 	
 	FORCEINLINE_DEBUGGABLE void RHIInvalidateCachedState()
@@ -2838,7 +2847,7 @@ public:
 			CMD_CONTEXT(RHIInvalidateCachedState)();
 			return;
 		}
-		new (AllocCommand<FRHICommandInvalidateCachedState>()) FRHICommandInvalidateCachedState();
+		ALLOC_COMMAND(FRHICommandInvalidateCachedState)();
 	}
 
 	FORCEINLINE void DiscardRenderTargets(bool Depth, bool Stencil, uint32 ColorBitMask)
@@ -2848,7 +2857,7 @@ public:
 			CMD_CONTEXT(RHIDiscardRenderTargets)(Depth, Stencil, ColorBitMask);
 			return;
 		}
-		new (AllocCommand<FRHICommandDiscardRenderTargets>()) FRHICommandDiscardRenderTargets(Depth, Stencil, ColorBitMask);
+		ALLOC_COMMAND(FRHICommandDiscardRenderTargets)(Depth, Stencil, ColorBitMask);
 	}
 	
 	FORCEINLINE_DEBUGGABLE void BreakPoint()
@@ -2862,7 +2871,7 @@ public:
 			}
 			return;
 		}
-		new (AllocCommand<FRHICommandDebugBreak>()) FRHICommandDebugBreak();
+		ALLOC_COMMAND(FRHICommandDebugBreak)();
 #endif
 	}
 };
@@ -2884,7 +2893,7 @@ public:
 			COMPUTE_CONTEXT(RHISetShaderUniformBuffer)(Shader, BaseIndex, UniformBuffer);
 			return;
 		}
-		new (AllocCommand<FRHICommandSetShaderUniformBuffer<FComputeShaderRHIParamRef, ECmdList::ECompute> >()) FRHICommandSetShaderUniformBuffer<FComputeShaderRHIParamRef, ECmdList::ECompute>(Shader, BaseIndex, UniformBuffer);		
+		ALLOC_COMMAND(FRHICommandSetShaderUniformBuffer<FComputeShaderRHIParamRef, ECmdList::ECompute>)(Shader, BaseIndex, UniformBuffer);		
 	}
 	
 	FORCEINLINE void SetShaderUniformBuffer(FComputeShaderRHIRef& Shader, uint32 BaseIndex, FUniformBufferRHIParamRef UniformBuffer)
@@ -2901,7 +2910,7 @@ public:
 		}
 		void* UseValue = Alloc(NumBytes, 16);
 		FMemory::Memcpy(UseValue, NewValue, NumBytes);
-		new (AllocCommand<FRHICommandSetShaderParameter<FComputeShaderRHIParamRef, ECmdList::ECompute> >()) FRHICommandSetShaderParameter<FComputeShaderRHIParamRef, ECmdList::ECompute>(Shader, BufferIndex, BaseIndex, NumBytes, UseValue);
+		ALLOC_COMMAND(FRHICommandSetShaderParameter<FComputeShaderRHIParamRef, ECmdList::ECompute>)(Shader, BufferIndex, BaseIndex, NumBytes, UseValue);
 	}
 	
 	FORCEINLINE void SetShaderParameter(FComputeShaderRHIRef& Shader, uint32 BufferIndex, uint32 BaseIndex, uint32 NumBytes, const void* NewValue)
@@ -2916,7 +2925,7 @@ public:
 			COMPUTE_CONTEXT(RHISetShaderTexture)(Shader, TextureIndex, Texture);
 			return;
 		}
-		new (AllocCommand<FRHICommandSetShaderTexture<FComputeShaderRHIParamRef, ECmdList::ECompute> >()) FRHICommandSetShaderTexture<FComputeShaderRHIParamRef, ECmdList::ECompute>(Shader, TextureIndex, Texture);
+		ALLOC_COMMAND(FRHICommandSetShaderTexture<FComputeShaderRHIParamRef, ECmdList::ECompute>)(Shader, TextureIndex, Texture);
 	}
 	
 	FORCEINLINE_DEBUGGABLE void SetShaderResourceViewParameter(FComputeShaderRHIParamRef Shader, uint32 SamplerIndex, FShaderResourceViewRHIParamRef SRV)
@@ -2926,7 +2935,7 @@ public:
 			COMPUTE_CONTEXT(RHISetShaderResourceViewParameter)(Shader, SamplerIndex, SRV);
 			return;
 		}
-		new (AllocCommand<FRHICommandSetShaderResourceViewParameter<FComputeShaderRHIParamRef, ECmdList::ECompute> >()) FRHICommandSetShaderResourceViewParameter<FComputeShaderRHIParamRef, ECmdList::ECompute>(Shader, SamplerIndex, SRV);
+		ALLOC_COMMAND(FRHICommandSetShaderResourceViewParameter<FComputeShaderRHIParamRef, ECmdList::ECompute>)(Shader, SamplerIndex, SRV);
 	}
 	
 	FORCEINLINE_DEBUGGABLE void SetShaderSampler(FComputeShaderRHIParamRef Shader, uint32 SamplerIndex, FSamplerStateRHIParamRef State)
@@ -2943,7 +2952,7 @@ public:
 			COMPUTE_CONTEXT(RHISetShaderSampler)(Shader, SamplerIndex, State);
 			return;
 		}
-		new (AllocCommand<FRHICommandSetShaderSampler<FComputeShaderRHIParamRef, ECmdList::ECompute> >()) FRHICommandSetShaderSampler<FComputeShaderRHIParamRef, ECmdList::ECompute>(Shader, SamplerIndex, State);
+		ALLOC_COMMAND(FRHICommandSetShaderSampler<FComputeShaderRHIParamRef, ECmdList::ECompute>)(Shader, SamplerIndex, State);
 	}
 
 	FORCEINLINE_DEBUGGABLE void SetUAVParameter(FComputeShaderRHIParamRef Shader, uint32 UAVIndex, FUnorderedAccessViewRHIParamRef UAV)
@@ -2953,7 +2962,7 @@ public:
 			COMPUTE_CONTEXT(RHISetUAVParameter)(Shader, UAVIndex, UAV);
 			return;
 		}
-		new (AllocCommand<FRHICommandSetUAVParameter<FComputeShaderRHIParamRef, ECmdList::ECompute> >()) FRHICommandSetUAVParameter<FComputeShaderRHIParamRef, ECmdList::ECompute>(Shader, UAVIndex, UAV);
+		ALLOC_COMMAND(FRHICommandSetUAVParameter<FComputeShaderRHIParamRef, ECmdList::ECompute>)(Shader, UAVIndex, UAV);
 	}
 
 	FORCEINLINE_DEBUGGABLE void SetUAVParameter(FComputeShaderRHIParamRef Shader, uint32 UAVIndex, FUnorderedAccessViewRHIParamRef UAV, uint32 InitialCount)
@@ -2963,7 +2972,7 @@ public:
 			COMPUTE_CONTEXT(RHISetUAVParameter)(Shader, UAVIndex, UAV, InitialCount);
 			return;
 		}
-		new (AllocCommand<FRHICommandSetUAVParameter_IntialCount<FComputeShaderRHIParamRef, ECmdList::ECompute> >()) FRHICommandSetUAVParameter_IntialCount<FComputeShaderRHIParamRef, ECmdList::ECompute>(Shader, UAVIndex, UAV, InitialCount);
+		ALLOC_COMMAND(FRHICommandSetUAVParameter_IntialCount<FComputeShaderRHIParamRef, ECmdList::ECompute>)(Shader, UAVIndex, UAV, InitialCount);
 	}
 	
 	FORCEINLINE_DEBUGGABLE void SetComputeShader(FComputeShaderRHIParamRef ComputeShader)
@@ -2974,7 +2983,7 @@ public:
 			COMPUTE_CONTEXT(RHISetComputeShader)(ComputeShader);
 			return;
 		}
-		new (AllocCommand<FRHICommandSetComputeShader<ECmdList::ECompute> >()) FRHICommandSetComputeShader<ECmdList::ECompute>(ComputeShader);
+		ALLOC_COMMAND(FRHICommandSetComputeShader<ECmdList::ECompute>)(ComputeShader);
 	}
 
 	FORCEINLINE_DEBUGGABLE void SetComputePipelineState(FComputePipelineState* ComputePipelineState)
@@ -2986,7 +2995,7 @@ public:
 			COMPUTE_CONTEXT(RHISetComputePipelineState)(RHIComputePipelineState);
 			return;
 		}
-		new (AllocCommand<FRHICommandSetComputePipelineState<ECmdList::ECompute> >()) FRHICommandSetComputePipelineState<ECmdList::ECompute>(ComputePipelineState);
+		ALLOC_COMMAND(FRHICommandSetComputePipelineState<ECmdList::ECompute>)(ComputePipelineState);
 	}
 
 	FORCEINLINE_DEBUGGABLE void SetAsyncComputeBudget(EAsyncComputeBudget Budget)
@@ -2996,7 +3005,7 @@ public:
 			COMPUTE_CONTEXT(RHISetAsyncComputeBudget)(Budget);
 			return;
 		}
-		new (AllocCommand<FRHICommandSetAsyncComputeBudget<ECmdList::ECompute>>()) FRHICommandSetAsyncComputeBudget<ECmdList::ECompute>(Budget);
+		ALLOC_COMMAND(FRHICommandSetAsyncComputeBudget<ECmdList::ECompute>)(Budget);
 	}
 
 	FORCEINLINE_DEBUGGABLE void DispatchComputeShader(uint32 ThreadGroupCountX, uint32 ThreadGroupCountY, uint32 ThreadGroupCountZ)
@@ -3006,7 +3015,7 @@ public:
 			COMPUTE_CONTEXT(RHIDispatchComputeShader)(ThreadGroupCountX, ThreadGroupCountY, ThreadGroupCountZ);
 			return;
 		}
-		new (AllocCommand<FRHICommandDispatchComputeShader<ECmdList::ECompute> >()) FRHICommandDispatchComputeShader<ECmdList::ECompute>(ThreadGroupCountX, ThreadGroupCountY, ThreadGroupCountZ);
+		ALLOC_COMMAND(FRHICommandDispatchComputeShader<ECmdList::ECompute>)(ThreadGroupCountX, ThreadGroupCountY, ThreadGroupCountZ);
 	}
 
 	FORCEINLINE_DEBUGGABLE void DispatchIndirectComputeShader(FVertexBufferRHIParamRef ArgumentBuffer, uint32 ArgumentOffset)
@@ -3016,7 +3025,7 @@ public:
 			COMPUTE_CONTEXT(RHIDispatchIndirectComputeShader)(ArgumentBuffer, ArgumentOffset);
 			return;
 		}
-		new (AllocCommand<FRHICommandDispatchIndirectComputeShader<ECmdList::ECompute> >()) FRHICommandDispatchIndirectComputeShader<ECmdList::ECompute>(ArgumentBuffer, ArgumentOffset);
+		ALLOC_COMMAND(FRHICommandDispatchIndirectComputeShader<ECmdList::ECompute>)(ArgumentBuffer, ArgumentOffset);
 	}
 
 	FORCEINLINE_DEBUGGABLE void TransitionResource(EResourceTransitionAccess TransitionType, EResourceTransitionPipeline TransitionPipeline, FUnorderedAccessViewRHIParamRef InUAV, FComputeFenceRHIParamRef WriteFence)
@@ -3031,7 +3040,7 @@ public:
 		// Allocate space to hold the single UAV pointer inline in the command list itself.
 		FUnorderedAccessViewRHIParamRef* UAVArray = (FUnorderedAccessViewRHIParamRef*)Alloc(sizeof(FUnorderedAccessViewRHIParamRef), alignof(FUnorderedAccessViewRHIParamRef));
 		UAVArray[0] = UAV;
-		new (AllocCommand<FRHICommandTransitionUAVs<ECmdList::ECompute> >()) FRHICommandTransitionUAVs<ECmdList::ECompute>(TransitionType, TransitionPipeline, UAVArray, 1, WriteFence);
+		ALLOC_COMMAND(FRHICommandTransitionUAVs<ECmdList::ECompute>)(TransitionType, TransitionPipeline, UAVArray, 1, WriteFence);
 	}
 
 	FORCEINLINE_DEBUGGABLE void TransitionResource(EResourceTransitionAccess TransitionType, EResourceTransitionPipeline TransitionPipeline, FUnorderedAccessViewRHIParamRef InUAV)
@@ -3054,7 +3063,7 @@ public:
 			UAVArray[Index] = InUAVs[Index];
 		}
 
-		new (AllocCommand<FRHICommandTransitionUAVs<ECmdList::ECompute> >()) FRHICommandTransitionUAVs<ECmdList::ECompute>(TransitionType, TransitionPipeline, UAVArray, NumUAVs, WriteFence);
+		ALLOC_COMMAND(FRHICommandTransitionUAVs<ECmdList::ECompute>)(TransitionType, TransitionPipeline, UAVArray, NumUAVs, WriteFence);
 	}
 
 	FORCEINLINE_DEBUGGABLE void TransitionResources(EResourceTransitionAccess TransitionType, EResourceTransitionPipeline TransitionPipeline, FUnorderedAccessViewRHIParamRef* InUAVs, int32 NumUAVs)
@@ -3070,7 +3079,7 @@ public:
 			return;
 		}
 		TCHAR* NameCopy  = AllocString(Name);
-		new (AllocCommand<FRHICommandPushEvent<ECmdList::ECompute> >()) FRHICommandPushEvent<ECmdList::ECompute>(NameCopy, Color);
+		ALLOC_COMMAND(FRHICommandPushEvent<ECmdList::ECompute>)(NameCopy, Color);
 	}
 
 	FORCEINLINE_DEBUGGABLE void PopEvent()
@@ -3080,7 +3089,7 @@ public:
 			COMPUTE_CONTEXT(RHIPopEvent)();
 			return;
 		}
-		new (AllocCommand<FRHICommandPopEvent<ECmdList::ECompute> >()) FRHICommandPopEvent<ECmdList::ECompute>();
+		ALLOC_COMMAND(FRHICommandPopEvent<ECmdList::ECompute>)();
 	}
 
 	FORCEINLINE_DEBUGGABLE void BreakPoint()
@@ -3094,7 +3103,7 @@ public:
 			}
 			return;
 		}
-		new (AllocCommand<FRHICommandDebugBreak>()) FRHICommandDebugBreak();
+		ALLOC_COMMAND(FRHICommandDebugBreak)();
 #endif
 	}
 
@@ -3105,7 +3114,7 @@ public:
 			COMPUTE_CONTEXT(RHISubmitCommandsHint)();
 			return;
 		}
-		new (AllocCommand<FRHICommandSubmitCommandsHint<ECmdList::ECompute>>()) FRHICommandSubmitCommandsHint<ECmdList::ECompute>();
+		ALLOC_COMMAND(FRHICommandSubmitCommandsHint<ECmdList::ECompute>)();
 	}
 
 	FORCEINLINE_DEBUGGABLE void WaitComputeFence(FComputeFenceRHIParamRef WaitFence)
@@ -3115,7 +3124,7 @@ public:
 			COMPUTE_CONTEXT(RHIWaitComputeFence)(WaitFence);
 			return;
 		}
-		new (AllocCommand<FRHICommandWaitComputeFence<ECmdList::ECompute>>()) FRHICommandWaitComputeFence<ECmdList::ECompute>(WaitFence);
+		ALLOC_COMMAND(FRHICommandWaitComputeFence<ECmdList::ECompute>)(WaitFence);
 	}
 
 	FORCEINLINE_DEBUGGABLE void EnqueueStagedRead(FStagingBufferRHIParamRef StagingBuffer, FGPUFenceRHIParamRef Fence, uint32 Offset, uint32 NumBytes)
@@ -3125,7 +3134,7 @@ public:
 			CMD_CONTEXT(RHIEnqueueStagedRead)(StagingBuffer, Fence, Offset, NumBytes);
 			return;
 		}
-		new (AllocCommand<FRHICommandEnqueueStagedRead<ECmdList::ECompute>>()) FRHICommandEnqueueStagedRead<ECmdList::ECompute>(StagingBuffer, Fence, Offset, NumBytes);
+		ALLOC_COMMAND(FRHICommandEnqueueStagedRead<ECmdList::ECompute>)(StagingBuffer, Fence, Offset, NumBytes);
 	}
 };
 
@@ -3209,7 +3218,7 @@ public:
 		}
 		else
 		{
-			new (AllocCommand<TRHILambdaCommand<LAMBDA>>()) TRHILambdaCommand<LAMBDA>(Forward<LAMBDA>(Lambda));
+			ALLOC_COMMAND(TRHILambdaCommand<LAMBDA>)(Forward<LAMBDA>(Lambda));
 			return true;
 		}
 	}
