Index buffer SRVs & UAVs work on Metal.

diff --git a/Engine/Source/Runtime/RHI/Private/RHI.cpp b/Engine/Source/Runtime/RHI/Private/RHI.cpp
--- a/Engine/Source/Runtime/RHI/Private/RHI.cpp
+++ b/Engine/Source/Runtime/RHI/Private/RHI.cpp
@@ -816,7 +816,7 @@ RHI_API bool RHISupportsPixelShaderUAVs(const EShaderPlatform Platform)
 
 RHI_API bool RHISupportsIndexBufferUAVs(const EShaderPlatform Platform)
 {
-	return Platform == SP_PCD3D_SM5 || IsVulkanPlatform(Platform) || Platform == SP_XBOXONE_D3D12 || Platform == SP_PS4;
+	return Platform == SP_PCD3D_SM5 || IsVulkanPlatform(Platform) || IsMetalSM5Platform(Platform) || Platform == SP_XBOXONE_D3D12 || Platform == SP_PS4;
 }
 
 static ERHIFeatureLevel::Type GRHIMobilePreviewFeatureLevel = ERHIFeatureLevel::Num;
