DR - Fix deprecation deprecation
#rnx

diff --git a/Engine/Source/Runtime/Core/Public/Math/Color.h b/Engine/Source/Runtime/Core/Public/Math/Color.h
--- a/Engine/Source/Runtime/Core/Public/Math/Color.h
+++ b/Engine/Source/Runtime/Core/Public/Math/Color.h
@@ -258,7 +258,7 @@ struct FLinearColor
 		return NewCopy;
 	}
 
-	DEPRECATED(4.22, "FGetHSV doesn't perform a valid HSV conversion, use MakeFromHSV8 instead")
+	UE_DEPRECATED(4.22, "FGetHSV doesn't perform a valid HSV conversion, use MakeFromHSV8 instead")
 	static CORE_API FLinearColor FGetHSV(uint8 H,uint8 S,uint8 V);
 
 	/**
diff --git a/Engine/Source/Runtime/RHI/Public/RHICommandList.h b/Engine/Source/Runtime/RHI/Public/RHICommandList.h
--- a/Engine/Source/Runtime/RHI/Public/RHICommandList.h
+++ b/Engine/Source/Runtime/RHI/Public/RHICommandList.h
@@ -2145,7 +2145,7 @@ public:
 		ALLOC_COMMAND(FRHICommandDrawPrimitive)(BaseVertexIndex, NumPrimitives, NumInstances);
 	}
 
-	DEPRECATED(4.22, "PrimitiveType is now a part of the Graphics Pipeline State.")
+	UE_DEPRECATED(4.22, "PrimitiveType is now a part of the Graphics Pipeline State.")
 	FORCEINLINE_DEBUGGABLE void DrawPrimitive(uint32 PrimitiveType, uint32 BaseVertexIndex, uint32 NumPrimitives, uint32 NumInstances)
 	{
 		DrawPrimitive(BaseVertexIndex, NumPrimitives, NumInstances);
@@ -2167,7 +2167,7 @@ public:
 		ALLOC_COMMAND(FRHICommandDrawIndexedPrimitive)(IndexBuffer, BaseVertexIndex, FirstInstance, NumVertices, StartIndex, NumPrimitives, NumInstances);
 	}
 
-	DEPRECATED(4.22, "PrimitiveType is now a part of the Graphics Pipeline State.")
+	UE_DEPRECATED(4.22, "PrimitiveType is now a part of the Graphics Pipeline State.")
 	FORCEINLINE_DEBUGGABLE void DrawIndexedPrimitive(FIndexBufferRHIParamRef IndexBuffer, uint32 PrimitiveType, int32 BaseVertexIndex, uint32 FirstInstance, uint32 NumVertices, uint32 StartIndex, uint32 NumPrimitives, uint32 NumInstances)
 	{
 		DrawIndexedPrimitive(IndexBuffer, BaseVertexIndex, FirstInstance, NumVertices, StartIndex, NumPrimitives, NumInstances);
@@ -2502,7 +2502,7 @@ public:
 		ALLOC_COMMAND(FRHICommandDrawPrimitiveIndirect)(ArgumentBuffer, ArgumentOffset);
 	}
 
-	DEPRECATED(4.22, "PrimitiveType is now a part of the Graphics Pipeline State.")
+	UE_DEPRECATED(4.22, "PrimitiveType is now a part of the Graphics Pipeline State.")
 	FORCEINLINE_DEBUGGABLE void DrawPrimitiveIndirect(uint32 PrimitiveType, FVertexBufferRHIParamRef ArgumentBuffer, uint32 ArgumentOffset)
 	{
 		DrawPrimitiveIndirect(ArgumentBuffer, ArgumentOffset);
@@ -2520,7 +2520,7 @@ public:
 		ALLOC_COMMAND(FRHICommandDrawIndexedIndirect)(IndexBufferRHI, ArgumentsBufferRHI, DrawArgumentsIndex, NumInstances);
 	}
 
-	DEPRECATED(4.22, "PrimitiveType is now a part of the Graphics Pipeline State.")
+	UE_DEPRECATED(4.22, "PrimitiveType is now a part of the Graphics Pipeline State.")
 	FORCEINLINE_DEBUGGABLE void DrawIndexedIndirect(FIndexBufferRHIParamRef IndexBufferRHI, uint32 PrimitiveType, FStructuredBufferRHIParamRef ArgumentsBufferRHI, uint32 DrawArgumentsIndex, uint32 NumInstances)
 	{
 		DrawIndexedIndirect(IndexBufferRHI, ArgumentsBufferRHI, DrawArgumentsIndex, NumInstances);
@@ -2537,7 +2537,7 @@ public:
 		ALLOC_COMMAND(FRHICommandDrawIndexedPrimitiveIndirect)(IndexBuffer, ArgumentsBuffer, ArgumentOffset);
 	}
 
-	DEPRECATED(4.22, "PrimitiveType is now a part of the Graphics Pipeline State.")
+	UE_DEPRECATED(4.22, "PrimitiveType is now a part of the Graphics Pipeline State.")
 	FORCEINLINE_DEBUGGABLE void DrawIndexedPrimitiveIndirect(uint32 PrimitiveType, FIndexBufferRHIParamRef IndexBuffer, FVertexBufferRHIParamRef ArgumentsBuffer, uint32 ArgumentOffset)
 	{
 		DrawIndexedPrimitiveIndirect(IndexBuffer, ArgumentsBuffer, ArgumentOffset);
