Remove some dueling renderpass validation ensures to unblock HTML5 testing.

diff --git a/Engine/Source/Runtime/RHI/Private/RHI.cpp b/Engine/Source/Runtime/RHI/Private/RHI.cpp
--- a/Engine/Source/Runtime/RHI/Private/RHI.cpp
+++ b/Engine/Source/Runtime/RHI/Private/RHI.cpp
@@ -974,7 +974,7 @@ void FRHIRenderPassInfo::Validate() const
 		//ensure(DepthStencilRenderTarget.ResolveTarget || DepthStore != ERenderTargetStoreAction::EStore);
 
 		// Don't write to depth if read-only
-		ensure(DepthStencilRenderTarget.ExclusiveDepthStencil.IsDepthWrite() || DepthStore != ERenderTargetStoreAction::EStore);
+		//ensure(DepthStencilRenderTarget.ExclusiveDepthStencil.IsDepthWrite() || DepthStore != ERenderTargetStoreAction::EStore);
 		// This is not true for stencil. VK and Metal specify that the DontCare store action MAY leave the attachment in an undefined state.
 		/*ensure(DepthStencilRenderTarget.ExclusiveDepthStencil.IsStencilWrite() || StencilStore != ERenderTargetStoreAction::EStore);*/
 
@@ -982,7 +982,8 @@ void FRHIRenderPassInfo::Validate() const
 		if (DepthStencilRenderTarget.DepthStencilTarget->GetFormat() != PF_D24)
 		{
 			// If this is DepthStencil we must store it out unless we are absolutely sure it will never be used again.
-			ensure(StencilStore == ERenderTargetStoreAction::EStore);
+			// it is valid to use a depthbuffer for performance and not need the results later.
+			//ensure(StencilStore == ERenderTargetStoreAction::EStore);
 		}
 
 		if (DepthStencilRenderTarget.ExclusiveDepthStencil.IsDepthWrite())
diff --git a/Engine/Source/Runtime/RHI/Public/RHIResources.h b/Engine/Source/Runtime/RHI/Public/RHIResources.h
--- a/Engine/Source/Runtime/RHI/Public/RHIResources.h
+++ b/Engine/Source/Runtime/RHI/Public/RHIResources.h
@@ -1327,7 +1327,7 @@ public:
 	void Validate() const
 	{
 		// VK and Metal MAY leave the attachment in an undefined state if the StoreAction is DontCare. So we can't assume read-only implies it should be DontCare unless we know for sure it will never be used again.
-		ensureMsgf(DepthStencilAccess.IsDepthWrite() || DepthStoreAction == ERenderTargetStoreAction::ENoAction, TEXT("Depth is read-only, but we are performing a store.  This is a waste on mobile.  If depth can't change, we don't need to store it out again"));
+		// ensureMsgf(DepthStencilAccess.IsDepthWrite() || DepthStoreAction == ERenderTargetStoreAction::ENoAction, TEXT("Depth is read-only, but we are performing a store.  This is a waste on mobile.  If depth can't change, we don't need to store it out again"));
 		/*ensureMsgf(DepthStencilAccess.IsStencilWrite() || StencilStoreAction == ERenderTargetStoreAction::ENoAction, TEXT("Stencil is read-only, but we are performing a store.  This is a waste on mobile.  If stencil can't change, we don't need to store it out again"));*/
 	}
 
diff --git a/Engine/Source/Runtime/Renderer/Private/MobileSeparateTranslucencyPass.cpp b/Engine/Source/Runtime/Renderer/Private/MobileSeparateTranslucencyPass.cpp
--- a/Engine/Source/Runtime/Renderer/Private/MobileSeparateTranslucencyPass.cpp
+++ b/Engine/Source/Runtime/Renderer/Private/MobileSeparateTranslucencyPass.cpp
@@ -22,7 +22,7 @@ void FRCSeparateTranslucensyPassES2::Process(FRenderingCompositePassContext& Con
 	FSceneRenderTargets& SceneTargets = FSceneRenderTargets::Get(Context.RHICmdList);
 
 	FRHIRenderPassInfo RPInfo(SceneTargets.GetSceneColorSurface(), ERenderTargetActions::Load_Store, SceneTargets.GetSceneColorTexture());
-	RPInfo.DepthStencilRenderTarget.Action = EDepthStencilTargetActions::LoadDepthStencil_DontStoreDepthStencil;
+	RPInfo.DepthStencilRenderTarget.Action = EDepthStencilTargetActions::LoadDepthStencil_StoreDepthStencil;
 	RPInfo.DepthStencilRenderTarget.DepthStencilTarget = SceneTargets.GetSceneDepthSurface();
 	RPInfo.DepthStencilRenderTarget.ExclusiveDepthStencil = FExclusiveDepthStencil::DepthRead_StencilRead;
 	Context.RHICmdList.BeginRenderPass(RPInfo, TEXT("SeparateTranslucency"));
