Make sure that we enforce the minimum Metal version to match the shader compiler.

diff --git a/Engine/Source/Runtime/Apple/MetalRHI/Private/MetalCommandQueue.cpp b/Engine/Source/Runtime/Apple/MetalRHI/Private/MetalCommandQueue.cpp
--- a/Engine/Source/Runtime/Apple/MetalRHI/Private/MetalCommandQueue.cpp
+++ b/Engine/Source/Runtime/Apple/MetalRHI/Private/MetalCommandQueue.cpp
@@ -31,16 +31,19 @@ FMetalCommandQueue::FMetalCommandQueue(mtlpp::Device InDevice, uint32 const MaxN
 {
 	int32 MaxShaderVersion = 0;
 #if PLATFORM_MAC
-	int32 DefaultMaxShaderVersion = 2;
+	int32 DefaultMaxShaderVersion = 3;
+	int32 MinShaderVersion = 3;
     const TCHAR* const Settings = TEXT("/Script/MacTargetPlatform.MacTargetSettings");
 #else
 	int32 DefaultMaxShaderVersion = 0;
+	int32 MinShaderVersion = 0;
     const TCHAR* const Settings = TEXT("/Script/IOSRuntimeSettings.IOSRuntimeSettings");
 #endif
     if(!GConfig->GetInt(Settings, TEXT("MaxShaderLanguageVersion"), MaxShaderVersion, GEngineIni))
     {
         MaxShaderVersion = DefaultMaxShaderVersion;
     }
+	MaxShaderVersion = FMath::Max(MinShaderVersion, MaxShaderVersion);
 	ValidateVersion(MaxShaderVersion);
 
 	if(MaxNumCommandBuffers == 0)
diff --git a/Engine/Source/Runtime/RHI/Private/RHI.cpp b/Engine/Source/Runtime/RHI/Private/RHI.cpp
--- a/Engine/Source/Runtime/RHI/Private/RHI.cpp
+++ b/Engine/Source/Runtime/RHI/Private/RHI.cpp
@@ -762,10 +762,12 @@ RHI_API uint32 RHIGetShaderLanguageVersion(const EShaderPlatform Platform)
 			if (MaxShaderVersion < 0)
 			{
 				MaxShaderVersion = 2;
+				int32 MinShaderVersion = 3;
 				if(!GConfig->GetInt(TEXT("/Script/MacTargetPlatform.MacTargetSettings"), TEXT("MaxShaderLanguageVersion"), MaxShaderVersion, GEngineIni))
 				{
-					MaxShaderVersion = 2;
+					MaxShaderVersion = 3;
 				}
+				MaxShaderVersion = FMath::Max(MinShaderVersion, MaxShaderVersion);
 			}
 			Version = (uint32)MaxShaderVersion;
 		}
@@ -775,10 +777,12 @@ RHI_API uint32 RHIGetShaderLanguageVersion(const EShaderPlatform Platform)
 			if (MaxShaderVersion < 0)
 			{
 				MaxShaderVersion = 0;
+				int32 MinShaderVersion = 0;
 				if(!GConfig->GetInt(TEXT("/Script/IOSRuntimeSettings.IOSRuntimeSettings"), TEXT("MaxShaderLanguageVersion"), MaxShaderVersion, GEngineIni))
 				{
 					MaxShaderVersion = 0;
 				}
+				MaxShaderVersion = FMath::Max(MinShaderVersion, MaxShaderVersion);
 			}
 			Version = (uint32)MaxShaderVersion;
 		}
