Fix random leak in PSO cache due to uninitialized variable in copy constructor

diff --git a/Engine/Source/Runtime/RHI/Private/PipelineStateCache.cpp b/Engine/Source/Runtime/RHI/Private/PipelineStateCache.cpp
--- a/Engine/Source/Runtime/RHI/Private/PipelineStateCache.cpp
+++ b/Engine/Source/Runtime/RHI/Private/PipelineStateCache.cpp
@@ -367,6 +367,7 @@ public:
 
 		check( LocalCache.Contains(InKey) == false );
 		LocalCache.Add(InKey, InValue);
+		checkfSlow(LocalCache.Contains(InKey), TEXT("PSO not found immediately after adding.  Likely cause is an uninitialized field in a constructor or copy constructor"));
 		return true;
 	}
 
diff --git a/Engine/Source/Runtime/RHI/Public/RHIResources.h b/Engine/Source/Runtime/RHI/Public/RHIResources.h
--- a/Engine/Source/Runtime/RHI/Public/RHIResources.h
+++ b/Engine/Source/Runtime/RHI/Public/RHIResources.h
@@ -1669,6 +1669,7 @@ public:
 		, RasterizerState(InMinimalState.RasterizerState)
 		, DepthStencilState(InMinimalState.DepthStencilState)
 		, ImmutableSamplerState(InMinimalState.ImmutableSamplerState)
+		, bDepthBounds(InMinimalState.bDepthBounds)
 		, PrimitiveType(InMinimalState.PrimitiveType)
 	{
 	}
