DR - Delete some deprecated methods
- Add support for RHICOMMAND_CALLSTACK

diff --git a/Engine/Source/Runtime/RHI/Public/RHICommandList.h b/Engine/Source/Runtime/RHI/Public/RHICommandList.h
--- a/Engine/Source/Runtime/RHI/Public/RHICommandList.h
+++ b/Engine/Source/Runtime/RHI/Public/RHICommandList.h
@@ -26,31 +26,29 @@
 #include "Async/TaskGraphInterfaces.h"
 #include "HAL/LowLevelMemTracker.h"
 
+
+// Set to 1 to capture the callstack for every RHI command. Cheap & memory efficient representation: Use the 
+// value in FRHICommand::StackFrames to get the pointer to the code (ie paste on a disassembly window)
+#define RHICOMMAND_CALLSTACK		0
+#if RHICOMMAND_CALLSTACK
+#include "HAL/PlatformStackwalk.h"
+#endif
+
 class FApp;
 class FBlendStateInitializerRHI;
 class FGraphicsPipelineStateInitializer;
 class FLastRenderTimeContainer;
-class FReadSurfaceDataFlags;
 class FRHICommandListBase;
 class FRHIComputeShader;
-class FRHIDepthRenderTargetView;
-class FRHIRenderTargetView;
-class FRHISetRenderTargetsInfo;
 class IRHICommandContext;
 class IRHIComputeContext;
-struct FBoundShaderStateInput;
 struct FDepthStencilStateInitializerRHI;
 struct FRasterizerStateInitializerRHI;
-struct FResolveParams;
 struct FRHIResourceCreateInfo;
 struct FRHIResourceInfo;
 struct FRHIUniformBufferLayout;
 struct FSamplerStateInitializerRHI;
 struct FTextureMemoryStats;
-enum class EAsyncComputeBudget;
-enum class EClearDepthStencil;
-enum class EResourceTransitionAccess;
-enum class EResourceTransitionPipeline;
 class FComputePipelineState;
 class FGraphicsPipelineState;
 
@@ -103,12 +101,10 @@ bool FORCEINLINE IsRunningRHIInTaskThread()
 }
 
 
-
-
 extern RHI_API bool GEnableAsyncCompute;
 extern RHI_API TAutoConsoleVariable<int32> CVarRHICmdWidth;
 extern RHI_API TAutoConsoleVariable<int32> CVarRHICmdFlushRenderThreadTasks;
-class FRHICommandListBase;
+
 
 struct RHI_API FLockTracker
 {
@@ -626,6 +622,15 @@ public:
 template<typename TCmd>
 struct FRHICommand : public FRHICommandBase
 {
+#if RHICOMMAND_CALLSTACK
+	uint64 StackFrames[16];
+
+	FRHICommand()
+	{
+		FPlatformStackWalk::CaptureStackBackTrace(StackFrames, 16);
+	}
+#endif
+
 	void ExecuteAndDestruct(FRHICommandListBase& CmdList, FRHICommandListDebugContext& Context) override final
 	{
 		TCmd *ThisCmd = static_cast<TCmd*>(this);
@@ -679,28 +684,6 @@ struct  FRHICommandEndUpdateMultiFrameUAV final : public FRHICommand<FRHICommand
 	RHI_API void Execute(FRHICommandListBase& CmdList);
 };
 
-struct FRHICommandSetRasterizerState final : public FRHICommand<FRHICommandSetRasterizerState>
-{
-	FRasterizerStateRHIParamRef State;
-	FORCEINLINE_DEBUGGABLE FRHICommandSetRasterizerState(FRasterizerStateRHIParamRef InState)
-		: State(InState)
-	{
-	}
-	RHI_API void Execute(FRHICommandListBase& CmdList);
-};
-
-struct FRHICommandSetDepthStencilState final : public FRHICommand<FRHICommandSetDepthStencilState>
-{
-	FDepthStencilStateRHIParamRef State;
-	uint32 StencilRef;
-	FORCEINLINE_DEBUGGABLE FRHICommandSetDepthStencilState(FDepthStencilStateRHIParamRef InState, uint32 InStencilRef)
-		: State(InState)
-		, StencilRef(InStencilRef)
-	{
-	}
-	RHI_API void Execute(FRHICommandListBase& CmdList);
-};
-
 struct FRHICommandSetStencilRef final : public FRHICommand<FRHICommandSetStencilRef>
 {
 	uint32 StencilRef;
@@ -868,18 +851,6 @@ struct FRHICommandSetBoundShaderState final : public FRHICommand<FRHICommandSetB
 	RHI_API void Execute(FRHICommandListBase& CmdList);
 };
 
-struct FRHICommandSetBlendState final : public FRHICommand<FRHICommandSetBlendState>
-{
-	FBlendStateRHIParamRef State;
-	FLinearColor BlendFactor;
-	FORCEINLINE_DEBUGGABLE FRHICommandSetBlendState(FBlendStateRHIParamRef InState, const FLinearColor& InBlendFactor)
-		: State(InState)
-		, BlendFactor(InBlendFactor)
-	{
-	}
-	RHI_API void Execute(FRHICommandListBase& CmdList);
-};
-
 struct FRHICommandSetBlendFactor final : public FRHICommand<FRHICommandSetBlendFactor>
 {
 	FLinearColor BlendFactor;
@@ -1536,74 +1507,6 @@ struct FComputedGraphicsPipelineState
 	}
 };
 
-struct FLocalGraphicsPipelineStateWorkArea
-{
-	FGraphicsPipelineStateInitializer Args;
-	FComputedGraphicsPipelineState* ComputedGraphicsPipelineState;
-#if DO_CHECK // the below variables are used in check(), which can be enabled in Shipping builds (see Build.h)
-	FRHICommandListBase* CheckCmdList;
-	int32 UID;
-#endif
-
-	FORCEINLINE_DEBUGGABLE FLocalGraphicsPipelineStateWorkArea(
-		FRHICommandListBase* InCheckCmdList,
-		const FGraphicsPipelineStateInitializer& Initializer
-		)
-		: Args(Initializer)
-#if DO_CHECK
-		, CheckCmdList(InCheckCmdList)
-		, UID(InCheckCmdList->GetUID())
-#endif
-	{
-		ComputedGraphicsPipelineState = new (InCheckCmdList->Alloc<FComputedGraphicsPipelineState>()) FComputedGraphicsPipelineState;
-	}
-};
-
-struct FLocalGraphicsPipelineState
-{
-	FLocalGraphicsPipelineStateWorkArea* WorkArea;
-	FGraphicsPipelineStateRHIRef BypassGraphicsPipelineState; // this is only used in the case of Bypass, should eventually be deleted
-
-	FLocalGraphicsPipelineState()
-		: WorkArea(nullptr)
-	{
-	}
-
-	FLocalGraphicsPipelineState(const FLocalGraphicsPipelineState& Other)
-		: WorkArea(Other.WorkArea)
-		, BypassGraphicsPipelineState(Other.BypassGraphicsPipelineState)
-	{
-	}
-};
-
-struct FRHICommandBuildLocalGraphicsPipelineState final : public FRHICommand<FRHICommandBuildLocalGraphicsPipelineState>
-{
-	FLocalGraphicsPipelineStateWorkArea WorkArea;
-
-	FORCEINLINE_DEBUGGABLE FRHICommandBuildLocalGraphicsPipelineState(
-		FRHICommandListBase* CheckCmdList,
-		const FGraphicsPipelineStateInitializer& Initializer
-		)
-		: WorkArea(CheckCmdList, Initializer)
-	{
-	}
-
-	RHI_API void Execute(FRHICommandListBase& CmdList);
-};
-
-struct FRHICommandSetLocalGraphicsPipelineState final : public FRHICommand<FRHICommandSetLocalGraphicsPipelineState>
-{
-	FLocalGraphicsPipelineState LocalGraphicsPipelineState;
-	FORCEINLINE_DEBUGGABLE FRHICommandSetLocalGraphicsPipelineState(FRHICommandListBase* CheckCmdList, FLocalGraphicsPipelineState& InLocalGraphicsPipelineState)
-		: LocalGraphicsPipelineState(InLocalGraphicsPipelineState)
-	{
-		check(CheckCmdList == LocalGraphicsPipelineState.WorkArea->CheckCmdList && CheckCmdList->GetUID() == LocalGraphicsPipelineState.WorkArea->UID); // this PSO was not built for this particular commandlist
-		LocalGraphicsPipelineState.WorkArea->ComputedGraphicsPipelineState->UseCount++;
-	}
-
-	RHI_API void Execute(FRHICommandListBase& CmdList);
-};
-
 struct FComputedUniformBuffer
 {
 	FUniformBufferRHIRef UniformBuffer;
@@ -1930,37 +1833,6 @@ public:
 		ALLOC_COMMAND(FRHICommandEndUpdateMultiFrameUAV)(UAV);
 	}
 
-	UE_DEPRECATED(4.20, "BuildLocalGraphicsPipelineState is deprecated. Use PipelineStateCache::GetAndOrCreateGraphicsPipelineState() instead.")
-	FORCEINLINE_DEBUGGABLE FLocalGraphicsPipelineState BuildLocalGraphicsPipelineState(
-		const FGraphicsPipelineStateInitializer& Initializer
-		)
-	{
-		//check(IsOutsideRenderPass());
-		FLocalGraphicsPipelineState Result;
-		if (Bypass())
-		{
-			Result.BypassGraphicsPipelineState = RHICreateGraphicsPipelineState(Initializer);
-		}
-		else
-		{
-			auto* Cmd = ALLOC_COMMAND(FRHICommandBuildLocalGraphicsPipelineState)(this, Initializer);
-			Result.WorkArea = &Cmd->WorkArea;
-		}
-		return Result;
-	}
-
-	UE_DEPRECATED(4.20, "SetLocalGraphicsPipelineState is deprecated. Use SetGraphicsPipelineState() instead.")
-	FORCEINLINE_DEBUGGABLE void SetLocalGraphicsPipelineState(FLocalGraphicsPipelineState LocalGraphicsPipelineState)
-	{
-		//check(IsOutsideRenderPass());
-		if (Bypass())
-		{
-			CMD_CONTEXT(RHISetGraphicsPipelineState)(LocalGraphicsPipelineState.BypassGraphicsPipelineState);
-			return;
-		}
-		ALLOC_COMMAND(FRHICommandSetLocalGraphicsPipelineState)(this, LocalGraphicsPipelineState);
-	}
-	
 	FORCEINLINE_DEBUGGABLE FLocalUniformBuffer BuildLocalUniformBuffer(const void* Contents, uint32 ContentsSize, const FRHIUniformBufferLayout& Layout)
 	{
 		//check(IsOutsideRenderPass());
@@ -2566,12 +2438,6 @@ public:
 		ALLOC_COMMAND(FRHICommandCopyToResolveTarget)(SourceTextureRHI, DestTextureRHI, ResolveParams);
 	}
 
-	UE_DEPRECATED(4.20, "This signature of CopyToResolveTarget is deprecated. Please use the new one instead.")
-	FORCEINLINE_DEBUGGABLE void CopyToResolveTarget(FTextureRHIParamRef SourceTextureRHI, FTextureRHIParamRef DestTextureRHI, bool bKeepOriginalSurface, const FResolveParams& ResolveParams)
-	{
-		CopyToResolveTarget(SourceTextureRHI, DestTextureRHI, ResolveParams);
-	}
-
 	FORCEINLINE_DEBUGGABLE void CopyTexture(FTextureRHIParamRef SourceTextureRHI, FTextureRHIParamRef DestTextureRHI, const FRHICopyTextureInfo& CopyInfo)
 	{
 		check(IsOutsideRenderPass());
@@ -4628,5 +4494,6 @@ FORCEINLINE void RHIUnlockStagingBuffer(FStagingBufferRHIParamRef StagingBuffer)
 	 FRHICommandListExecutor::GetImmediateCommandList().UnlockStagingBuffer(StagingBuffer);
 }
 
+#undef RHICOMMAND_CALLSTACK
 
 #include "RHICommandList.inl"
diff --git a/Engine/Source/Runtime/RHI/Public/RHICommandListCommandExecutes.inl b/Engine/Source/Runtime/RHI/Public/RHICommandListCommandExecutes.inl
--- a/Engine/Source/Runtime/RHI/Public/RHICommandListCommandExecutes.inl
+++ b/Engine/Source/Runtime/RHI/Public/RHICommandListCommandExecutes.inl
@@ -51,10 +51,7 @@ struct FRHICommandEndRenderQuery;
 struct FRHICommandEndScene;
 struct FRHICommandFlushComputeShaderCache;
 struct FRHICommandSetBlendFactor;
-struct FRHICommandSetBlendState;
 struct FRHICommandSetBoundShaderState;
-struct FRHICommandSetDepthStencilState;
-struct FRHICommandSetLocalBoundShaderState;
 struct FRHICommandSetLocalGraphicsPipelineState;
 struct FRHICommandSetRasterizerState;
 struct FRHICommandSetRenderTargets;
@@ -94,18 +91,6 @@ void FRHICommandEndUpdateMultiFrameUAV::Execute(FRHICommandListBase& CmdList)
 	INTERNAL_DECORATOR(RHIEndUpdateMultiFrameResource)(UAV);
 }
 
-void FRHICommandSetRasterizerState::Execute(FRHICommandListBase& CmdList)
-{
-	RHISTAT(SetRasterizerState);
-	INTERNAL_DECORATOR(RHISetRasterizerState)(State);
-}
-
-void FRHICommandSetDepthStencilState::Execute(FRHICommandListBase& CmdList)
-{
-	RHISTAT(SetDepthStencilState);
-	INTERNAL_DECORATOR(RHISetDepthStencilState)(State, StencilRef);
-}
-
 void FRHICommandSetStencilRef::Execute(FRHICommandListBase& CmdList)
 {
 	RHISTAT(SetStencilRef);
@@ -238,12 +223,6 @@ void FRHICommandSetBoundShaderState::Execute(FRHICommandListBase& CmdList)
 	INTERNAL_DECORATOR(RHISetBoundShaderState)(BoundShaderState);
 }
 
-void FRHICommandSetBlendState::Execute(FRHICommandListBase& CmdList)
-{
-	RHISTAT(SetBlendState);
-	INTERNAL_DECORATOR(RHISetBlendState)(State, BlendFactor);
-}
-
 void FRHICommandSetBlendFactor::Execute(FRHICommandListBase& CmdList)
 {
 	RHISTAT(SetBlendFactor);
@@ -501,31 +480,6 @@ void FRHICommandEnqueueStagedRead<CmdListType>::Execute(FRHICommandListBase& Cmd
 template struct FRHICommandEnqueueStagedRead<ECmdList::EGfx>;
 template struct FRHICommandEnqueueStagedRead<ECmdList::ECompute>;
 
-void FRHICommandBuildLocalGraphicsPipelineState::Execute(FRHICommandListBase& CmdList)
-{
-	LLM_SCOPE(ELLMTag::Shaders);
-	RHISTAT(BuildLocalGraphicsPipelineState);
-	check(!IsValidRef(WorkArea.ComputedGraphicsPipelineState->GraphicsPipelineState));
-	if (WorkArea.ComputedGraphicsPipelineState->UseCount)
-	{
-		WorkArea.ComputedGraphicsPipelineState->GraphicsPipelineState =
-			RHICreateGraphicsPipelineState(WorkArea.Args);
-	}
-}
-
-void FRHICommandSetLocalGraphicsPipelineState::Execute(FRHICommandListBase& CmdList)
-{
-	RHISTAT(SetLocalGraphicsPipelineState);
-	check(LocalGraphicsPipelineState.WorkArea->ComputedGraphicsPipelineState->UseCount > 0 && IsValidRef(LocalGraphicsPipelineState.WorkArea->ComputedGraphicsPipelineState->GraphicsPipelineState)); // this should have been created and should have uses outstanding
-
-	INTERNAL_DECORATOR(RHISetGraphicsPipelineState)(LocalGraphicsPipelineState.WorkArea->ComputedGraphicsPipelineState->GraphicsPipelineState);
-
-	if (--LocalGraphicsPipelineState.WorkArea->ComputedGraphicsPipelineState->UseCount == 0)
-	{
-		LocalGraphicsPipelineState.WorkArea->ComputedGraphicsPipelineState->~FComputedGraphicsPipelineState();
-	}
-}
-
 void FRHICommandBuildLocalUniformBuffer::Execute(FRHICommandListBase& CmdList)
 {
 	LLM_SCOPE(ELLMTag::Shaders);
