Removed CMD_CONTEXT and COMPUTE_CONTEXT macros, as they are no longer useful and only disrupt intellisense.

diff --git a/Engine/Source/Runtime/RHI/Private/RHICommandList.cpp b/Engine/Source/Runtime/RHI/Private/RHICommandList.cpp
--- a/Engine/Source/Runtime/RHI/Private/RHICommandList.cpp
+++ b/Engine/Source/Runtime/RHI/Private/RHICommandList.cpp
@@ -1441,7 +1441,7 @@ void FRHICommandList::BeginScene()
 	check(IsImmediate() && IsInRenderingThread());
 	if (Bypass())
 	{
-		CMD_CONTEXT(RHIBeginScene)();
+		GetContext().RHIBeginScene();
 		return;
 	}
 	ALLOC_COMMAND(FRHICommandBeginScene)();
@@ -1457,7 +1457,7 @@ void FRHICommandList::EndScene()
 	check(IsImmediate() && IsInRenderingThread());
 	if (Bypass())
 	{
-		CMD_CONTEXT(RHIEndScene)();
+		GetContext().RHIEndScene();
 		return;
 	}
 	ALLOC_COMMAND(FRHICommandEndScene)();
@@ -1475,7 +1475,7 @@ void FRHICommandList::BeginDrawingViewport(FViewportRHIParamRef Viewport, FTextu
 	check(IsImmediate() && IsInRenderingThread());
 	if (Bypass())
 	{
-		CMD_CONTEXT(RHIBeginDrawingViewport)(Viewport, RenderTargetRHI);
+		GetContext().RHIBeginDrawingViewport(Viewport, RenderTargetRHI);
 		return;
 	}
 	ALLOC_COMMAND(FRHICommandBeginDrawingViewport)(Viewport, RenderTargetRHI);
@@ -1492,7 +1492,7 @@ void FRHICommandList::EndDrawingViewport(FViewportRHIParamRef Viewport, bool bPr
 	check(IsImmediate() && IsInRenderingThread());
 	if (Bypass())
 	{
-		CMD_CONTEXT(RHIEndDrawingViewport)(Viewport, bPresent, bLockToVsync);
+		GetContext().RHIEndDrawingViewport(Viewport, bPresent, bLockToVsync);
 	}
 	else
 	{
@@ -1528,7 +1528,7 @@ void FRHICommandList::BeginFrame()
 	check(IsImmediate() && IsInRenderingThread());
 	if (Bypass())
 	{
-		CMD_CONTEXT(RHIBeginFrame)();
+		GetContext().RHIBeginFrame();
 		return;
 	}
 	ALLOC_COMMAND(FRHICommandBeginFrame)();
@@ -1546,7 +1546,7 @@ void FRHICommandList::EndFrame()
 	check(IsImmediate() && IsInRenderingThread());
 	if (Bypass())
 	{
-		CMD_CONTEXT(RHIEndFrame)();
+		GetContext().RHIEndFrame();
 		return;
 	}
 	ALLOC_COMMAND(FRHICommandEndFrame)();
@@ -2469,7 +2469,7 @@ void FRHICommandListImmediate::UpdateTextureReference(FTextureReferenceRHIParamR
 			QUICK_SCOPE_CYCLE_COUNTER(STAT_RHIMETHOD_UpdateTextureReference_FlushRHI);
 			ImmediateFlush(EImmediateFlushType::FlushRHIThread);
 		}
-		CMD_CONTEXT(RHIUpdateTextureReference)(TextureRef, NewTexture);
+		GetContext().RHIUpdateTextureReference(TextureRef, NewTexture);
 		return;
 	}
 	ALLOC_COMMAND(FRHICommandUpdateTextureReference)(TextureRef, NewTexture);
diff --git a/Engine/Source/Runtime/RHI/Public/RHICommandList.h b/Engine/Source/Runtime/RHI/Public/RHICommandList.h
--- a/Engine/Source/Runtime/RHI/Public/RHICommandList.h
+++ b/Engine/Source/Runtime/RHI/Public/RHICommandList.h
@@ -1934,9 +1934,6 @@ struct FRHICommandSetRayTracingHitGroup final : public FRHICommand<FRHICommandSe
 };
 #endif // RHI_RAYTRACING
 
-#define CMD_CONTEXT(Method) GetContext().Method
-#define COMPUTE_CONTEXT(Method) GetComputeContext().Method
-
 // Using variadic macro because some types are fancy template<A,B> stuff, which gets broken off at the comma and interpreted as multiple arguments. 
 #define ALLOC_COMMAND(...) new ( AllocCommand(sizeof(__VA_ARGS__), alignof(__VA_ARGS__)) ) __VA_ARGS__
 #define ALLOC_COMMAND_CL(RHICmdList, ...) new ( (RHICmdList).AllocCommand(sizeof(__VA_ARGS__), alignof(__VA_ARGS__)) ) __VA_ARGS__
@@ -1963,7 +1960,7 @@ public:
 		//check(IsOutsideRenderPass());
 		if (Bypass())
 		{
-			CMD_CONTEXT(RHIBeginUpdateMultiFrameResource)( Texture);
+			GetContext().RHIBeginUpdateMultiFrameResource( Texture);
 			return;
 		}
 		ALLOC_COMMAND(FRHICommandBeginUpdateMultiFrameResource)(Texture);
@@ -1974,7 +1971,7 @@ public:
 		//check(IsOutsideRenderPass());
 		if (Bypass())
 		{
-			CMD_CONTEXT(RHIEndUpdateMultiFrameResource)(Texture);
+			GetContext().RHIEndUpdateMultiFrameResource(Texture);
 			return;
 		}
 		ALLOC_COMMAND(FRHICommandEndUpdateMultiFrameResource)(Texture);
@@ -1985,7 +1982,7 @@ public:
 		//check(IsOutsideRenderPass());
 		if (Bypass())
 		{
-			CMD_CONTEXT(RHIBeginUpdateMultiFrameResource)(UAV);
+			GetContext().RHIBeginUpdateMultiFrameResource(UAV);
 			return;
 		}
 		ALLOC_COMMAND(FRHICommandBeginUpdateMultiFrameUAV)(UAV);
@@ -1996,7 +1993,7 @@ public:
 		//check(IsOutsideRenderPass());
 		if (Bypass())
 		{
-			CMD_CONTEXT(RHIEndUpdateMultiFrameResource)(UAV);
+			GetContext().RHIEndUpdateMultiFrameResource(UAV);
 			return;
 		}
 		ALLOC_COMMAND(FRHICommandEndUpdateMultiFrameUAV)(UAV);
@@ -2025,7 +2022,7 @@ public:
 		//check(IsOutsideRenderPass());
 		if (Bypass())
 		{
-			CMD_CONTEXT(RHISetShaderUniformBuffer)(Shader, BaseIndex, UniformBuffer.BypassUniform);
+			GetContext().RHISetShaderUniformBuffer(Shader, BaseIndex, UniformBuffer.BypassUniform);
 			return;
 		}
 		ALLOC_COMMAND(FRHICommandSetLocalUniformBuffer<TShaderRHIParamRef>)(this, Shader, BaseIndex, UniformBuffer);
@@ -2043,7 +2040,7 @@ public:
 		//check(IsOutsideRenderPass());
 		if (Bypass())
 		{
-			CMD_CONTEXT(RHISetShaderUniformBuffer)(Shader, BaseIndex, UniformBuffer);
+			GetContext().RHISetShaderUniformBuffer(Shader, BaseIndex, UniformBuffer);
 			return;
 		}
 		ALLOC_COMMAND(FRHICommandSetShaderUniformBuffer<TShaderRHI*, ECmdList::EGfx>)(Shader, BaseIndex, UniformBuffer);
@@ -2060,7 +2057,7 @@ public:
 		//check(IsOutsideRenderPass());
 		if (Bypass())
 		{
-			CMD_CONTEXT(RHISetShaderParameter)(Shader, BufferIndex, BaseIndex, NumBytes, NewValue);
+			GetContext().RHISetShaderParameter(Shader, BufferIndex, BaseIndex, NumBytes, NewValue);
 			return;
 		}
 		void* UseValue = Alloc(NumBytes, 16);
@@ -2079,7 +2076,7 @@ public:
 		//check(IsOutsideRenderPass());
 		if (Bypass())
 		{
-			CMD_CONTEXT(RHISetShaderTexture)(Shader, TextureIndex, Texture);
+			GetContext().RHISetShaderTexture(Shader, TextureIndex, Texture);
 			return;
 		}
 		ALLOC_COMMAND(FRHICommandSetShaderTexture<TShaderRHIParamRef, ECmdList::EGfx>)(Shader, TextureIndex, Texture);
@@ -2097,7 +2094,7 @@ public:
 		//check(IsOutsideRenderPass());
 		if (Bypass())
 		{
-			CMD_CONTEXT(RHISetShaderResourceViewParameter)(Shader, SamplerIndex, SRV);
+			GetContext().RHISetShaderResourceViewParameter(Shader, SamplerIndex, SRV);
 			return;
 		}
 		ALLOC_COMMAND(FRHICommandSetShaderResourceViewParameter<TShaderRHIParamRef, ECmdList::EGfx>)(Shader, SamplerIndex, SRV);
@@ -2123,7 +2120,7 @@ public:
 
 		if (Bypass())
 		{
-			CMD_CONTEXT(RHISetShaderSampler)(Shader, SamplerIndex, State);
+			GetContext().RHISetShaderSampler(Shader, SamplerIndex, State);
 			return;
 		}
 		ALLOC_COMMAND(FRHICommandSetShaderSampler<TShaderRHIParamRef, ECmdList::EGfx>)(Shader, SamplerIndex, State);
@@ -2139,7 +2136,7 @@ public:
 	{
 		if (Bypass())
 		{
-			CMD_CONTEXT(RHISetUAVParameter)(Shader, UAVIndex, UAV);
+			GetContext().RHISetUAVParameter(Shader, UAVIndex, UAV);
 			return;
 		}
 		ALLOC_COMMAND(FRHICommandSetUAVParameter<FComputeShaderRHIParamRef, ECmdList::EGfx>)(Shader, UAVIndex, UAV);
@@ -2154,7 +2151,7 @@ public:
 	{
 		if (Bypass())
 		{
-			CMD_CONTEXT(RHISetUAVParameter)(Shader, UAVIndex, UAV, InitialCount);
+			GetContext().RHISetUAVParameter(Shader, UAVIndex, UAV, InitialCount);
 			return;
 		}
 		ALLOC_COMMAND(FRHICommandSetUAVParameter_IntialCount<FComputeShaderRHIParamRef, ECmdList::EGfx>)(Shader, UAVIndex, UAV, InitialCount);
@@ -2170,7 +2167,7 @@ public:
 		//check(IsOutsideRenderPass());
 		if (Bypass())
 		{
-			CMD_CONTEXT(RHISetBlendFactor)(BlendFactor);
+			GetContext().RHISetBlendFactor(BlendFactor);
 			return;
 		}
 		ALLOC_COMMAND(FRHICommandSetBlendFactor)(BlendFactor);
@@ -2181,7 +2178,7 @@ public:
 		//check(IsOutsideRenderPass());
 		if (Bypass())
 		{
-			CMD_CONTEXT(RHIDrawPrimitive)(BaseVertexIndex, NumPrimitives, NumInstances);
+			GetContext().RHIDrawPrimitive(BaseVertexIndex, NumPrimitives, NumInstances);
 			return;
 		}
 		ALLOC_COMMAND(FRHICommandDrawPrimitive)(BaseVertexIndex, NumPrimitives, NumInstances);
@@ -2203,7 +2200,7 @@ public:
 		//check(IsOutsideRenderPass());
 		if (Bypass())
 		{
-			CMD_CONTEXT(RHIDrawIndexedPrimitive)(IndexBuffer, BaseVertexIndex, FirstInstance, NumVertices, StartIndex, NumPrimitives, NumInstances);
+			GetContext().RHIDrawIndexedPrimitive(IndexBuffer, BaseVertexIndex, FirstInstance, NumVertices, StartIndex, NumPrimitives, NumInstances);
 			return;
 		}
 		ALLOC_COMMAND(FRHICommandDrawIndexedPrimitive)(IndexBuffer, BaseVertexIndex, FirstInstance, NumVertices, StartIndex, NumPrimitives, NumInstances);
@@ -2219,7 +2216,7 @@ public:
 	{
 		if (Bypass())
 		{
-			CMD_CONTEXT(RHISetStreamSource)(StreamIndex, VertexBuffer, Offset);
+			GetContext().RHISetStreamSource(StreamIndex, VertexBuffer, Offset);
 			return;
 		}
 		ALLOC_COMMAND(FRHICommandSetStreamSource)(StreamIndex, VertexBuffer, Offset);
@@ -2230,7 +2227,7 @@ public:
 		//check(IsOutsideRenderPass());
 		if (Bypass())
 		{
-			CMD_CONTEXT(RHISetStencilRef)(StencilRef);
+			GetContext().RHISetStencilRef(StencilRef);
 			return;
 		}
 
@@ -2242,7 +2239,7 @@ public:
 		//check(IsOutsideRenderPass());
 		if (Bypass())
 		{
-			CMD_CONTEXT(RHISetViewport)(MinX, MinY, MinZ, MaxX, MaxY, MaxZ);
+			GetContext().RHISetViewport(MinX, MinY, MinZ, MaxX, MaxY, MaxZ);
 			return;
 		}
 		ALLOC_COMMAND(FRHICommandSetViewport)(MinX, MinY, MinZ, MaxX, MaxY, MaxZ);
@@ -2253,7 +2250,7 @@ public:
 		//check(IsOutsideRenderPass());
 		if (Bypass())
 		{
-			CMD_CONTEXT(RHISetStereoViewport)(LeftMinX, RightMinX, LeftMinY, RightMinY, MinZ, LeftMaxX, RightMaxX, LeftMaxY, RightMaxY, MaxZ);
+			GetContext().RHISetStereoViewport(LeftMinX, RightMinX, LeftMinY, RightMinY, MinZ, LeftMaxX, RightMaxX, LeftMaxY, RightMaxY, MaxZ);
 			return;
 		}
 		ALLOC_COMMAND(FRHICommandSetStereoViewport)(LeftMinX, RightMinX, LeftMinY, RightMinY, MinZ, LeftMaxX, RightMaxX, LeftMaxY, RightMaxY, MaxZ);
@@ -2264,7 +2261,7 @@ public:
 		//check(IsOutsideRenderPass());
 		if (Bypass())
 		{
-			CMD_CONTEXT(RHISetScissorRect)(bEnable, MinX, MinY, MaxX, MaxY);
+			GetContext().RHISetScissorRect(bEnable, MinX, MinY, MaxX, MaxY);
 			return;
 		}
 		ALLOC_COMMAND(FRHICommandSetScissorRect)(bEnable, MinX, MinY, MaxX, MaxY);
@@ -2333,7 +2330,7 @@ public:
 
 		if (Bypass())
 		{
-			CMD_CONTEXT(RHISetRenderTargets)(
+			GetContext().RHISetRenderTargets(
 				NewNumSimultaneousRenderTargets,
 				NewRenderTargetsRHI,
 				NewDepthStencilTargetRHI,
@@ -2360,7 +2357,7 @@ public:
 
 		if (Bypass())
 		{
-			CMD_CONTEXT(RHISetRenderTargetsAndClear)(RenderTargetsInfo);
+			GetContext().RHISetRenderTargetsAndClear(RenderTargetsInfo);
 			return;
 		}
 		ALLOC_COMMAND(FRHICommandSetRenderTargetsAndClear)(RenderTargetsInfo);
@@ -2371,7 +2368,7 @@ public:
 		//check(IsOutsideRenderPass());
 		if (Bypass())
 		{
-			CMD_CONTEXT(RHIBindClearMRTValues)(bClearColor, bClearDepth, bClearStencil);
+			GetContext().RHIBindClearMRTValues(bClearColor, bClearDepth, bClearStencil);
 			return;
 		}
 		ALLOC_COMMAND(FRHICommandBindClearMRTValues)(bClearColor, bClearDepth, bClearStencil);
@@ -2383,7 +2380,7 @@ public:
 		//check(IsOutsideRenderPass());
 		if (Bypass())
 		{
-			CMD_CONTEXT(RHIBeginDrawPrimitiveUP)(NumPrimitives, NumVertices, VertexDataStride, OutVertexData);
+			GetContext().RHIBeginDrawPrimitiveUP(NumPrimitives, NumVertices, VertexDataStride, OutVertexData);
 			return;
 		}
 		checkf(!DrawUPData.OutVertexData && NumVertices * VertexDataStride > 0,
@@ -2403,7 +2400,7 @@ public:
 		//check(IsOutsideRenderPass());
 		if (Bypass())
 		{
-			CMD_CONTEXT(RHIEndDrawPrimitiveUP)();
+			GetContext().RHIEndDrawPrimitiveUP();
 			return;
 		}
 		check(DrawUPData.OutVertexData && DrawUPData.NumVertices);
@@ -2419,7 +2416,7 @@ public:
 		//check(IsOutsideRenderPass());
 		if (Bypass())
 		{
-			CMD_CONTEXT(RHIBeginDrawIndexedPrimitiveUP)(NumPrimitives, NumVertices, VertexDataStride, OutVertexData, MinVertexIndex, NumIndices, IndexDataStride, OutIndexData);
+			GetContext().RHIBeginDrawIndexedPrimitiveUP(NumPrimitives, NumVertices, VertexDataStride, OutVertexData, MinVertexIndex, NumIndices, IndexDataStride, OutIndexData);
 			return;
 		}
 		check(!DrawUPData.OutVertexData && !DrawUPData.OutIndexData && NumVertices * VertexDataStride > 0 && NumIndices * IndexDataStride > 0);
@@ -2444,7 +2441,7 @@ public:
 		//check(IsOutsideRenderPass());
 		if (Bypass())
 		{
-			CMD_CONTEXT(RHIEndDrawIndexedPrimitiveUP)();
+			GetContext().RHIEndDrawIndexedPrimitiveUP();
 			return;
 		}
 		check(DrawUPData.OutVertexData && DrawUPData.OutIndexData && DrawUPData.NumIndices && DrawUPData.NumVertices);
@@ -2462,7 +2459,7 @@ public:
 		ComputeShader->UpdateStats();
 		if (Bypass())
 		{
-			CMD_CONTEXT(RHISetComputeShader)(ComputeShader);
+			GetContext().RHISetComputeShader(ComputeShader);
 			return;
 		}
 		ALLOC_COMMAND(FRHICommandSetComputeShader<ECmdList::EGfx>)(ComputeShader);
@@ -2474,7 +2471,7 @@ public:
 		{
 			extern RHI_API FRHIComputePipelineState* ExecuteSetComputePipelineState(class FComputePipelineState* ComputePipelineState);
 			FRHIComputePipelineState* RHIComputePipelineState = ExecuteSetComputePipelineState(ComputePipelineState);
-			CMD_CONTEXT(RHISetComputePipelineState)(RHIComputePipelineState);
+			GetContext().RHISetComputePipelineState(RHIComputePipelineState);
 			return;
 		}
 		ALLOC_COMMAND(FRHICommandSetComputePipelineState<ECmdList::EGfx>)(ComputePipelineState);
@@ -2487,7 +2484,7 @@ public:
 		{
 			extern RHI_API FRHIGraphicsPipelineState* ExecuteSetGraphicsPipelineState(class FGraphicsPipelineState* GraphicsPipelineState);
 			FRHIGraphicsPipelineState* RHIGraphicsPipelineState = ExecuteSetGraphicsPipelineState(GraphicsPipelineState);
-			CMD_CONTEXT(RHISetGraphicsPipelineState)(RHIGraphicsPipelineState);
+			GetContext().RHISetGraphicsPipelineState(RHIGraphicsPipelineState);
 			return;
 		}
 		ALLOC_COMMAND(FRHICommandSetGraphicsPipelineState)(GraphicsPipelineState);
@@ -2497,7 +2494,7 @@ public:
 	{
 		if (Bypass())
 		{
-			CMD_CONTEXT(RHIDispatchComputeShader)(ThreadGroupCountX, ThreadGroupCountY, ThreadGroupCountZ);
+			GetContext().RHIDispatchComputeShader(ThreadGroupCountX, ThreadGroupCountY, ThreadGroupCountZ);
 			return;
 		}
 		ALLOC_COMMAND(FRHICommandDispatchComputeShader<ECmdList::EGfx>)(ThreadGroupCountX, ThreadGroupCountY, ThreadGroupCountZ);
@@ -2507,7 +2504,7 @@ public:
 	{
 		if (Bypass())
 		{
-			CMD_CONTEXT(RHIDispatchIndirectComputeShader)(ArgumentBuffer, ArgumentOffset);
+			GetContext().RHIDispatchIndirectComputeShader(ArgumentBuffer, ArgumentOffset);
 			return;
 		}
 		ALLOC_COMMAND(FRHICommandDispatchIndirectComputeShader<ECmdList::EGfx>)(ArgumentBuffer, ArgumentOffset);
@@ -2517,7 +2514,7 @@ public:
 	{
 		if (Bypass())
 		{
-			CMD_CONTEXT(RHIAutomaticCacheFlushAfterComputeShader)(bEnable);
+			GetContext().RHIAutomaticCacheFlushAfterComputeShader(bEnable);
 			return;
 		}
 		ALLOC_COMMAND(FRHICommandAutomaticCacheFlushAfterComputeShader)(bEnable);
@@ -2527,7 +2524,7 @@ public:
 	{
 		if (Bypass())
 		{
-			CMD_CONTEXT(RHIFlushComputeShaderCache)();
+			GetContext().RHIFlushComputeShaderCache();
 			return;
 		}
 		ALLOC_COMMAND(FRHICommandFlushComputeShaderCache)();
@@ -2538,7 +2535,7 @@ public:
 		//check(IsOutsideRenderPass());
 		if (Bypass())
 		{
-			CMD_CONTEXT(RHIDrawPrimitiveIndirect)(ArgumentBuffer, ArgumentOffset);
+			GetContext().RHIDrawPrimitiveIndirect(ArgumentBuffer, ArgumentOffset);
 			return;
 		}
 		ALLOC_COMMAND(FRHICommandDrawPrimitiveIndirect)(ArgumentBuffer, ArgumentOffset);
@@ -2556,7 +2553,7 @@ public:
 		//check(IsOutsideRenderPass());
 		if (Bypass())
 		{
-			CMD_CONTEXT(RHIDrawIndexedIndirect)(IndexBufferRHI, ArgumentsBufferRHI, DrawArgumentsIndex, NumInstances);
+			GetContext().RHIDrawIndexedIndirect(IndexBufferRHI, ArgumentsBufferRHI, DrawArgumentsIndex, NumInstances);
 			return;
 		}
 		ALLOC_COMMAND(FRHICommandDrawIndexedIndirect)(IndexBufferRHI, ArgumentsBufferRHI, DrawArgumentsIndex, NumInstances);
@@ -2573,7 +2570,7 @@ public:
 		//check(IsOutsideRenderPass());
 		if (Bypass())
 		{
-			CMD_CONTEXT(RHIDrawIndexedPrimitiveIndirect)(IndexBuffer, ArgumentsBuffer, ArgumentOffset);
+			GetContext().RHIDrawIndexedPrimitiveIndirect(IndexBuffer, ArgumentsBuffer, ArgumentOffset);
 			return;
 		}
 		ALLOC_COMMAND(FRHICommandDrawIndexedPrimitiveIndirect)(IndexBuffer, ArgumentsBuffer, ArgumentOffset);
@@ -2590,7 +2587,7 @@ public:
 		//check(IsOutsideRenderPass());
 		if (Bypass())
 		{
-			CMD_CONTEXT(RHISetDepthBounds)(MinDepth, MaxDepth);
+			GetContext().RHISetDepthBounds(MinDepth, MaxDepth);
 			return;
 		}
 		ALLOC_COMMAND(FRHICommandSetDepthBounds)(MinDepth, MaxDepth);
@@ -2601,7 +2598,7 @@ public:
 		//check(IsOutsideRenderPass());
 		if (Bypass())
 		{
-			CMD_CONTEXT(RHICopyToResolveTarget)(SourceTextureRHI, DestTextureRHI, ResolveParams);
+			GetContext().RHICopyToResolveTarget(SourceTextureRHI, DestTextureRHI, ResolveParams);
 			return;
 		}
 		ALLOC_COMMAND(FRHICommandCopyToResolveTarget)(SourceTextureRHI, DestTextureRHI, ResolveParams);
@@ -2612,7 +2609,7 @@ public:
 		check(IsOutsideRenderPass());
 		if (Bypass())
 		{
-			CMD_CONTEXT(RHICopyTexture)(SourceTextureRHI, DestTextureRHI, CopyInfo);
+			GetContext().RHICopyTexture(SourceTextureRHI, DestTextureRHI, CopyInfo);
 			return;
 		}
 		ALLOC_COMMAND(FRHICommandCopyTexture)(SourceTextureRHI, DestTextureRHI, CopyInfo);
@@ -2623,7 +2620,7 @@ public:
 		//check(IsOutsideRenderPass());
 		if (Bypass())
 		{
-			CMD_CONTEXT(RHIClearTinyUAV)(UnorderedAccessViewRHI, Values);
+			GetContext().RHIClearTinyUAV(UnorderedAccessViewRHI, Values);
 			return;
 		}
 		ALLOC_COMMAND(FRHICommandClearTinyUAV)(UnorderedAccessViewRHI, Values);
@@ -2633,7 +2630,7 @@ public:
 	{
 		if (Bypass())
 		{
-			CMD_CONTEXT(RHIBeginRenderQuery)(RenderQuery);
+			GetContext().RHIBeginRenderQuery(RenderQuery);
 			return;
 		}
 		ALLOC_COMMAND(FRHICommandBeginRenderQuery)(RenderQuery);
@@ -2642,7 +2639,7 @@ public:
 	{
 		if (Bypass())
 		{
-			CMD_CONTEXT(RHIEndRenderQuery)(RenderQuery);
+			GetContext().RHIEndRenderQuery(RenderQuery);
 			return;
 		}
 		ALLOC_COMMAND(FRHICommandEndRenderQuery)(RenderQuery);
@@ -2652,7 +2649,7 @@ public:
 	{
 		if (Bypass())
 		{
-			CMD_CONTEXT(RHISubmitCommandsHint)();
+			GetContext().RHISubmitCommandsHint();
 			return;
 		}
 		ALLOC_COMMAND(FRHICommandSubmitCommandsHint<ECmdList::EGfx>)();
@@ -2662,7 +2659,7 @@ public:
 	{
 		if (Bypass())
 		{
-			CMD_CONTEXT(RHIPollOcclusionQueries)();
+			GetContext().RHIPollOcclusionQueries();
 			return;
 		}
 		ALLOC_COMMAND(FRHICommandPollOcclusionQueries)();
@@ -2674,7 +2671,7 @@ public:
 		check(Texture == nullptr || Texture->IsCommitted());
 		if (Bypass())
 		{
-			CMD_CONTEXT(RHITransitionResources)(TransitionType, &Texture, 1);
+			GetContext().RHITransitionResources(TransitionType, &Texture, 1);
 			return;
 		}
 
@@ -2688,7 +2685,7 @@ public:
 	{
 		if (Bypass())
 		{
-			CMD_CONTEXT(RHITransitionResources)(TransitionType, InTextures, NumTextures);
+			GetContext().RHITransitionResources(TransitionType, InTextures, NumTextures);
 			return;
 		}
 
@@ -2706,7 +2703,7 @@ public:
 	{
 		if (Bypass())
 		{
-			CMD_CONTEXT(RHITransitionResources)(TransitionType, &InTextures[0], InTextures.Num());
+			GetContext().RHITransitionResources(TransitionType, &InTextures[0], InTextures.Num());
 			return;
 		}
 		ALLOC_COMMAND(FRHICommandTransitionTexturesArray)(TransitionType, InTextures);
@@ -2718,7 +2715,7 @@ public:
 		FUnorderedAccessViewRHIParamRef UAV = InUAV;
 		if (Bypass())
 		{
-			CMD_CONTEXT(RHITransitionResources)(TransitionType, TransitionPipeline, &UAV, 1, WriteFence);
+			GetContext().RHITransitionResources(TransitionType, TransitionPipeline, &UAV, 1, WriteFence);
 			return;
 		}
 
@@ -2737,7 +2734,7 @@ public:
 	{
 		if (Bypass())
 		{
-			CMD_CONTEXT(RHITransitionResources)(TransitionType, TransitionPipeline, InUAVs, NumUAVs, WriteFence);
+			GetContext().RHITransitionResources(TransitionType, TransitionPipeline, InUAVs, NumUAVs, WriteFence);
 			return;
 		}
 
@@ -2760,7 +2757,7 @@ public:
 	{
 		if (Bypass())
 		{
-			CMD_CONTEXT(RHIWaitComputeFence)(WaitFence);
+			GetContext().RHIWaitComputeFence(WaitFence);
 			return;
 		}
 		ALLOC_COMMAND(FRHICommandWaitComputeFence<ECmdList::EGfx>)(WaitFence);
@@ -2770,7 +2767,7 @@ public:
 	{
 		if (Bypass())
 		{
-			CMD_CONTEXT(RHIEnqueueStagedRead)(StagingBuffer, Fence, Offset, NumBytes);
+			GetContext().RHIEnqueueStagedRead(StagingBuffer, Fence, Offset, NumBytes);
 			return;
 		}
 		ALLOC_COMMAND(FRHICommandEnqueueStagedRead<ECmdList::EGfx>)(StagingBuffer, Fence, Offset, NumBytes);
@@ -2785,7 +2782,7 @@ public:
 
 		if (Bypass())
 		{
-			CMD_CONTEXT(RHIBeginRenderPass)(InInfo, Name);
+			GetContext().RHIBeginRenderPass(InInfo, Name);
 		}
 		else
 		{
@@ -2804,7 +2801,7 @@ public:
 		check(!IsInsideComputePass());
 		if (Bypass())
 		{
-			CMD_CONTEXT(RHIEndRenderPass)();
+			GetContext().RHIEndRenderPass();
 		}
 		else
 		{
@@ -2820,7 +2817,7 @@ public:
 
 		if (Bypass())
 		{
-			CMD_CONTEXT(RHIBeginComputePass)(Name);
+			GetContext().RHIBeginComputePass(Name);
 		}
 		else
 		{
@@ -2838,7 +2835,7 @@ public:
 		check(!IsInsideRenderPass());
 		if (Bypass())
 		{
-			CMD_CONTEXT(RHIEndComputePass)();
+			GetContext().RHIEndComputePass();
 		}
 		else
 		{
@@ -2859,7 +2856,7 @@ public:
 	{
 		if (Bypass())
 		{
-			CMD_CONTEXT(RHIPushEvent)(Name, Color);
+			GetContext().RHIPushEvent(Name, Color);
 			return;
 		}
 		TCHAR* NameCopy  = AllocString(Name);
@@ -2870,7 +2867,7 @@ public:
 	{
 		if (Bypass())
 		{
-			CMD_CONTEXT(RHIPopEvent)();
+			GetContext().RHIPopEvent();
 			return;
 		}
 		ALLOC_COMMAND(FRHICommandPopEvent<ECmdList::EGfx>)();
@@ -2880,7 +2877,7 @@ public:
 	{
 		if (Bypass())
 		{
-			CMD_CONTEXT(RHIInvalidateCachedState)();
+			GetContext().RHIInvalidateCachedState();
 			return;
 		}
 		ALLOC_COMMAND(FRHICommandInvalidateCachedState)();
@@ -2890,7 +2887,7 @@ public:
 	{
 		if (Bypass())
 		{
-			CMD_CONTEXT(RHIDiscardRenderTargets)(Depth, Stencil, ColorBitMask);
+			GetContext().RHIDiscardRenderTargets(Depth, Stencil, ColorBitMask);
 			return;
 		}
 		ALLOC_COMMAND(FRHICommandDiscardRenderTargets)(Depth, Stencil, ColorBitMask);
@@ -2918,7 +2915,7 @@ public:
 	{
 		if (Bypass())
 		{
-			CMD_CONTEXT(RHIBuildAccelerationStructure)(Geometry);
+			GetContext().RHIBuildAccelerationStructure(Geometry);
 		}
 		else
 		{
@@ -2930,7 +2927,7 @@ public:
 	{
 		if (Bypass())
 		{
-			CMD_CONTEXT(RHIUpdateAccelerationStructures)(Params);
+			GetContext().RHIUpdateAccelerationStructures(Params);
 		}
 		else
 		{
@@ -2942,7 +2939,7 @@ public:
 	{
 		if (Bypass())
 		{
-			CMD_CONTEXT(RHIBuildAccelerationStructures)(Params);
+			GetContext().RHIBuildAccelerationStructures(Params);
 		}
 		else
 		{
@@ -2954,7 +2951,7 @@ public:
 	{
 		if (Bypass())
 		{
-			CMD_CONTEXT(RHIBuildAccelerationStructure)(Scene);
+			GetContext().RHIBuildAccelerationStructure(Scene);
 		}
 		else
 		{
@@ -2974,7 +2971,7 @@ public:
 	{
 		if (Bypass())
 		{
-			CMD_CONTEXT(RHIRayTraceOcclusion)(Scene, Rays, Output, NumRays);
+			GetContext().RHIRayTraceOcclusion(Scene, Rays, Output, NumRays);
 		}
 		else
 		{
@@ -2993,7 +2990,7 @@ public:
 	{
 		if (Bypass())
 		{
-			CMD_CONTEXT(RHIRayTraceIntersection)(Scene, Rays, Output, NumRays);
+			GetContext().RHIRayTraceIntersection(Scene, Rays, Output, NumRays);
 		}
 		else
 		{
@@ -3005,7 +3002,7 @@ public:
 	{
 		if (Bypass())
 		{
-			CMD_CONTEXT(RHIRayTraceDispatch)(Pipeline, GlobalResourceBindings, Width, Height);
+			GetContext().RHIRayTraceDispatch(Pipeline, GlobalResourceBindings, Width, Height);
 		}
 		else
 		{
@@ -3017,7 +3014,7 @@ public:
 	{
 		if (Bypass())
 		{
-			CMD_CONTEXT(RHIRayTraceDispatch)(Pipeline, Scene, GlobalResourceBindings, Width, Height);
+			GetContext().RHIRayTraceDispatch(Pipeline, Scene, GlobalResourceBindings, Width, Height);
 		}
 		else
 		{
@@ -3032,7 +3029,7 @@ public:
 	{
 		if (Bypass())
 		{
-			CMD_CONTEXT(RHISetRayTracingHitGroup)(Scene, InstanceIndex, SegmentIndex, Pipeline, HitGroupIndex, ResourceBindings);
+			GetContext().RHISetRayTracingHitGroup(Scene, InstanceIndex, SegmentIndex, Pipeline, HitGroupIndex, ResourceBindings);
 		}
 		else
 		{
@@ -3057,7 +3054,7 @@ public:
 	{
 		if (Bypass())
 		{
-			COMPUTE_CONTEXT(RHISetShaderUniformBuffer)(Shader, BaseIndex, UniformBuffer);
+			GetComputeContext().RHISetShaderUniformBuffer(Shader, BaseIndex, UniformBuffer);
 			return;
 		}
 		ALLOC_COMMAND(FRHICommandSetShaderUniformBuffer<FComputeShaderRHIParamRef, ECmdList::ECompute>)(Shader, BaseIndex, UniformBuffer);		
@@ -3072,7 +3069,7 @@ public:
 	{
 		if (Bypass())
 		{
-			COMPUTE_CONTEXT(RHISetShaderParameter)(Shader, BufferIndex, BaseIndex, NumBytes, NewValue);
+			GetComputeContext().RHISetShaderParameter(Shader, BufferIndex, BaseIndex, NumBytes, NewValue);
 			return;
 		}
 		void* UseValue = Alloc(NumBytes, 16);
@@ -3089,7 +3086,7 @@ public:
 	{
 		if (Bypass())
 		{
-			COMPUTE_CONTEXT(RHISetShaderTexture)(Shader, TextureIndex, Texture);
+			GetComputeContext().RHISetShaderTexture(Shader, TextureIndex, Texture);
 			return;
 		}
 		ALLOC_COMMAND(FRHICommandSetShaderTexture<FComputeShaderRHIParamRef, ECmdList::ECompute>)(Shader, TextureIndex, Texture);
@@ -3099,7 +3096,7 @@ public:
 	{
 		if (Bypass())
 		{
-			COMPUTE_CONTEXT(RHISetShaderResourceViewParameter)(Shader, SamplerIndex, SRV);
+			GetComputeContext().RHISetShaderResourceViewParameter(Shader, SamplerIndex, SRV);
 			return;
 		}
 		ALLOC_COMMAND(FRHICommandSetShaderResourceViewParameter<FComputeShaderRHIParamRef, ECmdList::ECompute>)(Shader, SamplerIndex, SRV);
@@ -3116,7 +3113,7 @@ public:
 		
 		if (Bypass())
 		{
-			COMPUTE_CONTEXT(RHISetShaderSampler)(Shader, SamplerIndex, State);
+			GetComputeContext().RHISetShaderSampler(Shader, SamplerIndex, State);
 			return;
 		}
 		ALLOC_COMMAND(FRHICommandSetShaderSampler<FComputeShaderRHIParamRef, ECmdList::ECompute>)(Shader, SamplerIndex, State);
@@ -3126,7 +3123,7 @@ public:
 	{
 		if (Bypass())
 		{
-			COMPUTE_CONTEXT(RHISetUAVParameter)(Shader, UAVIndex, UAV);
+			GetComputeContext().RHISetUAVParameter(Shader, UAVIndex, UAV);
 			return;
 		}
 		ALLOC_COMMAND(FRHICommandSetUAVParameter<FComputeShaderRHIParamRef, ECmdList::ECompute>)(Shader, UAVIndex, UAV);
@@ -3136,7 +3133,7 @@ public:
 	{
 		if (Bypass())
 		{
-			COMPUTE_CONTEXT(RHISetUAVParameter)(Shader, UAVIndex, UAV, InitialCount);
+			GetComputeContext().RHISetUAVParameter(Shader, UAVIndex, UAV, InitialCount);
 			return;
 		}
 		ALLOC_COMMAND(FRHICommandSetUAVParameter_IntialCount<FComputeShaderRHIParamRef, ECmdList::ECompute>)(Shader, UAVIndex, UAV, InitialCount);
@@ -3147,7 +3144,7 @@ public:
 		ComputeShader->UpdateStats();
 		if (Bypass())
 		{
-			COMPUTE_CONTEXT(RHISetComputeShader)(ComputeShader);
+			GetComputeContext().RHISetComputeShader(ComputeShader);
 			return;
 		}
 		ALLOC_COMMAND(FRHICommandSetComputeShader<ECmdList::ECompute>)(ComputeShader);
@@ -3159,7 +3156,7 @@ public:
 		{
 			extern FRHIComputePipelineState* ExecuteSetComputePipelineState(FComputePipelineState* ComputePipelineState);
 			FRHIComputePipelineState* RHIComputePipelineState = ExecuteSetComputePipelineState(ComputePipelineState);
-			COMPUTE_CONTEXT(RHISetComputePipelineState)(RHIComputePipelineState);
+			GetComputeContext().RHISetComputePipelineState(RHIComputePipelineState);
 			return;
 		}
 		ALLOC_COMMAND(FRHICommandSetComputePipelineState<ECmdList::ECompute>)(ComputePipelineState);
@@ -3169,7 +3166,7 @@ public:
 	{
 		if (Bypass())
 		{
-			COMPUTE_CONTEXT(RHISetAsyncComputeBudget)(Budget);
+			GetComputeContext().RHISetAsyncComputeBudget(Budget);
 			return;
 		}
 		ALLOC_COMMAND(FRHICommandSetAsyncComputeBudget<ECmdList::ECompute>)(Budget);
@@ -3179,7 +3176,7 @@ public:
 	{
 		if (Bypass())
 		{
-			COMPUTE_CONTEXT(RHIDispatchComputeShader)(ThreadGroupCountX, ThreadGroupCountY, ThreadGroupCountZ);
+			GetComputeContext().RHIDispatchComputeShader(ThreadGroupCountX, ThreadGroupCountY, ThreadGroupCountZ);
 			return;
 		}
 		ALLOC_COMMAND(FRHICommandDispatchComputeShader<ECmdList::ECompute>)(ThreadGroupCountX, ThreadGroupCountY, ThreadGroupCountZ);
@@ -3189,7 +3186,7 @@ public:
 	{
 		if (Bypass())
 		{
-			COMPUTE_CONTEXT(RHIDispatchIndirectComputeShader)(ArgumentBuffer, ArgumentOffset);
+			GetComputeContext().RHIDispatchIndirectComputeShader(ArgumentBuffer, ArgumentOffset);
 			return;
 		}
 		ALLOC_COMMAND(FRHICommandDispatchIndirectComputeShader<ECmdList::ECompute>)(ArgumentBuffer, ArgumentOffset);
@@ -3200,7 +3197,7 @@ public:
 		FUnorderedAccessViewRHIParamRef UAV = InUAV;
 		if (Bypass())
 		{
-			COMPUTE_CONTEXT(RHITransitionResources)(TransitionType, TransitionPipeline, &UAV, 1, WriteFence);
+			GetComputeContext().RHITransitionResources(TransitionType, TransitionPipeline, &UAV, 1, WriteFence);
 			return;
 		}
 
@@ -3219,7 +3216,7 @@ public:
 	{
 		if (Bypass())
 		{
-			COMPUTE_CONTEXT(RHITransitionResources)(TransitionType, TransitionPipeline, InUAVs, NumUAVs, WriteFence);
+			GetComputeContext().RHITransitionResources(TransitionType, TransitionPipeline, InUAVs, NumUAVs, WriteFence);
 			return;
 		}
 
@@ -3242,7 +3239,7 @@ public:
 	{
 		if (Bypass())
 		{
-			COMPUTE_CONTEXT(RHIPushEvent)(Name, Color);
+			GetComputeContext().RHIPushEvent(Name, Color);
 			return;
 		}
 		TCHAR* NameCopy  = AllocString(Name);
@@ -3253,7 +3250,7 @@ public:
 	{
 		if (Bypass())
 		{
-			COMPUTE_CONTEXT(RHIPopEvent)();
+			GetComputeContext().RHIPopEvent();
 			return;
 		}
 		ALLOC_COMMAND(FRHICommandPopEvent<ECmdList::ECompute>)();
@@ -3278,7 +3275,7 @@ public:
 	{
 		if (Bypass())
 		{
-			COMPUTE_CONTEXT(RHISubmitCommandsHint)();
+			GetComputeContext().RHISubmitCommandsHint();
 			return;
 		}
 		ALLOC_COMMAND(FRHICommandSubmitCommandsHint<ECmdList::ECompute>)();
@@ -3288,7 +3285,7 @@ public:
 	{
 		if (Bypass())
 		{
-			COMPUTE_CONTEXT(RHIWaitComputeFence)(WaitFence);
+			GetComputeContext().RHIWaitComputeFence(WaitFence);
 			return;
 		}
 		ALLOC_COMMAND(FRHICommandWaitComputeFence<ECmdList::ECompute>)(WaitFence);
@@ -3298,7 +3295,7 @@ public:
 	{
 		if (Bypass())
 		{
-			CMD_CONTEXT(RHIEnqueueStagedRead)(StagingBuffer, Fence, Offset, NumBytes);
+			GetContext().RHIEnqueueStagedRead(StagingBuffer, Fence, Offset, NumBytes);
 			return;
 		}
 		ALLOC_COMMAND(FRHICommandEnqueueStagedRead<ECmdList::ECompute>)(StagingBuffer, Fence, Offset, NumBytes);
