Removed constructors from FRayTracingGeometrySegment, FRayTracingGeometryInitializer and FRayTracingSceneInitializer in favor of non-static member initialization

diff --git a/Engine/Source/Runtime/Engine/Private/SkeletalRenderCPUSkin.cpp b/Engine/Source/Runtime/Engine/Private/SkeletalRenderCPUSkin.cpp
--- a/Engine/Source/Runtime/Engine/Private/SkeletalRenderCPUSkin.cpp
+++ b/Engine/Source/Runtime/Engine/Private/SkeletalRenderCPUSkin.cpp
@@ -369,7 +369,10 @@ void FSkeletalMeshObjectCPUSkin::FSkeletalMeshObjectLOD::InitResources(FSkelMesh
 				GeometrySections.Reserve(RenderSections->Num());
 				for (const FSkelMeshRenderSection& Section : *RenderSections)
 				{
-					GeometrySections.Add(FRayTracingGeometrySegment( (uint32)(Section.BaseIndex / 3), (uint32)(Section.NumTriangles) ));
+					FRayTracingGeometrySegment Segment;
+					Segment.FirstPrimitive = Section.BaseIndex / 3;
+					Segment.NumPrimitives = Section.NumTriangles;
+					GeometrySections.Add(Segment);
 				}
 				Initializer.Segments = GeometrySections;
 
diff --git a/Engine/Source/Runtime/Engine/Private/SkeletalRenderGPUSkin.cpp b/Engine/Source/Runtime/Engine/Private/SkeletalRenderGPUSkin.cpp
--- a/Engine/Source/Runtime/Engine/Private/SkeletalRenderGPUSkin.cpp
+++ b/Engine/Source/Runtime/Engine/Private/SkeletalRenderGPUSkin.cpp
@@ -367,7 +367,10 @@ void FSkeletalMeshObjectGPUSkin::UpdateDynamicData_RenderThread(FGPUSkinCache* G
 				GeometrySections.Reserve(LODModel.RenderSections.Num());
 				for (const FSkelMeshRenderSection& Section : LODModel.RenderSections)
 				{
-					GeometrySections.Add(FRayTracingGeometrySegment( (uint32)(Section.BaseIndex / 3), (uint32)(Section.NumTriangles) ));
+					FRayTracingGeometrySegment Segment;
+					Segment.FirstPrimitive = Section.BaseIndex / 3;
+					Segment.NumPrimitives = Section.NumTriangles;
+					GeometrySections.Add(Segment);
 				}
 				Initializer.Segments = GeometrySections;
 
diff --git a/Engine/Source/Runtime/Engine/Private/SkeletalRenderStatic.cpp b/Engine/Source/Runtime/Engine/Private/SkeletalRenderStatic.cpp
--- a/Engine/Source/Runtime/Engine/Private/SkeletalRenderStatic.cpp
+++ b/Engine/Source/Runtime/Engine/Private/SkeletalRenderStatic.cpp
@@ -158,7 +158,10 @@ void FSkeletalMeshObjectStatic::FSkeletalMeshObjectLOD::InitResources(FSkelMeshC
 				GeometrySections.Reserve(RenderSections->Num());
 				for (const FSkelMeshRenderSection& Section : *RenderSections)
 				{
-					GeometrySections.Add(FRayTracingGeometrySegment( (uint32)(Section.BaseIndex / 3), (uint32)Section.NumTriangles ));
+					FRayTracingGeometrySegment Segment;
+					Segment.FirstPrimitive = Section.BaseIndex / 3;
+					Segment.NumPrimitives = Section.NumTriangles;
+					GeometrySections.Add(Segment);
 				}
 				Initializer.Segments = GeometrySections;
 
diff --git a/Engine/Source/Runtime/Engine/Private/StaticMesh.cpp b/Engine/Source/Runtime/Engine/Private/StaticMesh.cpp
--- a/Engine/Source/Runtime/Engine/Private/StaticMesh.cpp
+++ b/Engine/Source/Runtime/Engine/Private/StaticMesh.cpp
@@ -706,7 +706,10 @@ void FStaticMeshLODResources::InitResources(UStaticMesh* Parent)
 				GeometrySections.Reserve(Sections.Num());
 				for (const FStaticMeshSection& Section : Sections)
 				{
-					GeometrySections.Add(FRayTracingGeometrySegment( (uint32)(Section.FirstIndex / 3), (uint32)(Section.NumTriangles)));
+					FRayTracingGeometrySegment Segment;
+					Segment.FirstPrimitive = Section.FirstIndex / 3;
+					Segment.NumPrimitives = Section.NumTriangles;
+					GeometrySections.Add(Segment);
 				}
 				Initializer.Segments = GeometrySections;
 				
diff --git a/Engine/Source/Runtime/RHI/Public/DynamicRHI.h b/Engine/Source/Runtime/RHI/Public/DynamicRHI.h
--- a/Engine/Source/Runtime/RHI/Public/DynamicRHI.h
+++ b/Engine/Source/Runtime/RHI/Public/DynamicRHI.h
@@ -85,20 +85,8 @@ struct FRayTracingGeometryInstance
 
 struct FRayTracingGeometrySegment
 {
-	FRayTracingGeometrySegment()
-		: FirstPrimitive(0)
-		, NumPrimitives(0)
-	{}
-
-	FRayTracingGeometrySegment(uint32 InFirstPrimitive, uint32 InNumPrimitives)
-		: FirstPrimitive(InFirstPrimitive)
-		, NumPrimitives(InNumPrimitives)
-	{}
-
-	uint32 FirstPrimitive;
-	uint32 NumPrimitives;
-
-	// #dxr_todo: what should be the sensible default values of these flags?
+	uint32 FirstPrimitive = 0;
+	uint32 NumPrimitives = 0;
 
 	// Indicates whether any-hit shader could be invoked when hitting this geometry segment.
 	// Setting this to `false` turns off any-hit shaders, making the section "opaque" and improving ray tracing performance.
@@ -111,19 +99,8 @@ struct FRayTracingGeometrySegment
 
 struct FRayTracingGeometryInitializer
 {
-	FRayTracingGeometryInitializer()
-		: VertexBufferByteOffset(0)
-		, VertexBufferStride(0)
-		, VertexBufferElementType(VET_Float3)
-		, BaseVertexIndex(0)
-		, PrimitiveType(PT_TriangleList)
-		, TotalPrimitiveCount(0)
-		, bFastBuild(false)
-		, bAllowUpdate(false)
-	{}
-
-	FVertexBufferRHIRef PositionVertexBuffer;
-	FIndexBufferRHIRef IndexBuffer;
+	FVertexBufferRHIRef PositionVertexBuffer = nullptr;
+	FIndexBufferRHIRef IndexBuffer = nullptr;
 	uint32 VertexBufferByteOffset = 0;
 	uint32 VertexBufferStride = 0;
 	EVertexElementType VertexBufferElementType = VET_Float3;
@@ -137,12 +114,8 @@ struct FRayTracingGeometryInitializer
 
 struct FRayTracingSceneInitializer
 {
-	FRayTracingSceneInitializer()
-		: bIsDynamic(false)
-	{}
-
 	TArrayView<FRayTracingGeometryInstance> Instances;
-	bool bIsDynamic;
+	bool bIsDynamic = false;
 };
 
 /** The interface which is implemented by the dynamically bound RHI. */
