DR - Rework IRHICommandContext so we can remove a bunch of unnecessary methods on modern RHIs

diff --git a/Engine/Source/Runtime/Apple/MetalRHI/Private/MetalCommands.cpp b/Engine/Source/Runtime/Apple/MetalRHI/Private/MetalCommands.cpp
--- a/Engine/Source/Runtime/Apple/MetalRHI/Private/MetalCommands.cpp
+++ b/Engine/Source/Runtime/Apple/MetalRHI/Private/MetalCommands.cpp
@@ -106,11 +106,6 @@ void FMetalDynamicRHI::RHISetStreamOutTargets(uint32 NumTargets, const FVertexBu
 
 }
 
-void FMetalRHICommandContext::RHISetRasterizerState(FRasterizerStateRHIParamRef NewStateRHI)
-{
-	NOT_SUPPORTED("RHISetRasterizerState");
-}
-
 void FMetalRHICommandContext::RHISetComputeShader(FComputeShaderRHIParamRef ComputeShaderRHI)
 {
 	@autoreleasepool {
@@ -234,11 +229,6 @@ void FMetalRHICommandContext::RHISetScissorRect(bool bEnable,uint32 MinX,uint32
 	}
 }
 
-void FMetalRHICommandContext::RHISetBoundShaderState( FBoundShaderStateRHIParamRef BoundShaderStateRHI)
-{
-	NOT_SUPPORTED("RHISetBoundShaderState");
-}
-
 void FMetalRHICommandContext::RHISetGraphicsPipelineState(FGraphicsPipelineStateRHIParamRef GraphicsState)
 {
 	@autoreleasepool {
@@ -611,21 +601,11 @@ void FMetalRHICommandContext::RHISetShaderUniformBuffer(FComputeShaderRHIParamRe
 }
 
 
-void FMetalRHICommandContext::RHISetDepthStencilState(FDepthStencilStateRHIParamRef NewStateRHI, uint32 StencilRef)
-{
-	NOT_SUPPORTED("RHISetDepthStencilState");
-}
-
 void FMetalRHICommandContext::RHISetStencilRef(uint32 StencilRef)
 {
 	Context->GetCurrentState().SetStencilRef(StencilRef);
 }
 
-void FMetalRHICommandContext::RHISetBlendState(FBlendStateRHIParamRef NewStateRHI, const FLinearColor& BlendFactor)
-{
-	NOT_SUPPORTED("RHISetBlendState");
-}
-
 void FMetalRHICommandContext::RHISetBlendFactor(const FLinearColor& BlendFactor)
 {
 	Context->GetCurrentState().SetBlendFactor(BlendFactor);
@@ -1018,11 +998,6 @@ void FMetalDynamicRHI::RHIExecuteCommandList(FRHICommandList* RHICmdList)
 	NOT_SUPPORTED("RHIExecuteCommandList");
 }
 
-void FMetalRHICommandContext::RHIEnableDepthBoundsTest(bool bEnable)
-{
-	// Nothing required here
-}
-
 void FMetalRHICommandContext::RHISetDepthBounds(float MinDepth, float MaxDepth)
 {
 	// Nothing required here
diff --git a/Engine/Source/Runtime/Apple/MetalRHI/Public/MetalRHIContext.h b/Engine/Source/Runtime/Apple/MetalRHI/Public/MetalRHIContext.h
--- a/Engine/Source/Runtime/Apple/MetalRHI/Public/MetalRHIContext.h
+++ b/Engine/Source/Runtime/Apple/MetalRHI/Public/MetalRHIContext.h
@@ -95,8 +95,6 @@ public:
 	
 	virtual void RHISetStreamSource(uint32 StreamIndex, FVertexBufferRHIParamRef VertexBuffer, uint32 Offset) final override;
 
-	virtual void RHISetRasterizerState(FRasterizerStateRHIParamRef NewState) final override;
-	
 	// @param MinX including like Win32 RECT
 	// @param MinY including like Win32 RECT
 	// @param MaxX excluding like Win32 RECT
@@ -111,12 +109,6 @@ public:
 	// @param MaxY excluding like Win32 RECT
 	virtual void RHISetScissorRect(bool bEnable, uint32 MinX, uint32 MinY, uint32 MaxX, uint32 MaxY) final override;
 	
-	/**
-	 * Set bound shader state. This will set the vertex decl/shader, and pixel shader
-	 * @param BoundShaderState - state resource
-	 */
-	virtual void RHISetBoundShaderState(FBoundShaderStateRHIParamRef BoundShaderState) final override;
-	
 	virtual void RHISetGraphicsPipelineState(FGraphicsPipelineStateRHIParamRef GraphicsState) final override;
 	
 	/** Set the shader resource view of a surface.  This is used for binding TextureMS parameter types that need a multi sampled view. */
@@ -227,13 +219,8 @@ public:
 	
 	virtual void RHISetShaderParameter(FComputeShaderRHIParamRef ComputeShader, uint32 BufferIndex, uint32 BaseIndex, uint32 NumBytes, const void* NewValue) final override;
 	
-	virtual void RHISetDepthStencilState(FDepthStencilStateRHIParamRef NewState, uint32 StencilRef) final override;
-
 	virtual void RHISetStencilRef(uint32 StencilRef) final override;
 
-	// Allows to set the blend state, parameter can be created with RHICreateBlendState()
-	virtual void RHISetBlendState(FBlendStateRHIParamRef NewState, const FLinearColor& BlendFactor) final override;
-	
 	virtual void RHISetBlendFactor(const FLinearColor& BlendFactor) final override;
 
 	virtual void RHISetRenderTargets(uint32 NumSimultaneousRenderTargets, const FRHIRenderTargetView* NewRenderTargets, const FRHIDepthRenderTargetView* NewDepthStencilTarget, uint32 NumUAVs, const FUnorderedAccessViewRHIParamRef* UAVs) final override;
@@ -284,12 +271,6 @@ public:
 	 * Draw a primitive using the vertex and index data populated since RHIBeginDrawIndexedPrimitiveUP and clean up any memory as needed
 	 */
 	virtual void RHIEndDrawIndexedPrimitiveUP() final override;
-	
-	/**
-	 * Enabled/Disables Depth Bounds Testing.
-	 * @param bEnable	Enable(non-zero)/disable(zero) the depth bounds test
-	 */
-	virtual void RHIEnableDepthBoundsTest(bool bEnable) final override;
 
 	/**
 	* Sets Depth Bounds Testing with the given min/max depth.
diff --git a/Engine/Source/Runtime/D3D12RHI/Private/D3D12CommandContext.h b/Engine/Source/Runtime/D3D12RHI/Private/D3D12CommandContext.h
--- a/Engine/Source/Runtime/D3D12RHI/Private/D3D12CommandContext.h
+++ b/Engine/Source/Runtime/D3D12RHI/Private/D3D12CommandContext.h
@@ -273,10 +273,8 @@ public:
 	virtual void RHIBeginScene() final override;
 	virtual void RHIEndScene() final override;
 	virtual void RHISetStreamSource(uint32 StreamIndex, FVertexBufferRHIParamRef VertexBuffer, uint32 Offset) final override;
-	virtual void RHISetRasterizerState(FRasterizerStateRHIParamRef NewState) final override;
 	virtual void RHISetViewport(uint32 MinX, uint32 MinY, float MinZ, uint32 MaxX, uint32 MaxY, float MaxZ) final override;
 	virtual void RHISetScissorRect(bool bEnable, uint32 MinX, uint32 MinY, uint32 MaxX, uint32 MaxY) final override;
-	virtual void RHISetBoundShaderState(FBoundShaderStateRHIParamRef BoundShaderState) final override;
 	virtual void RHISetGraphicsPipelineState(FGraphicsPipelineStateRHIParamRef GraphicsPipelineState) final override;
 	virtual void RHISetShaderTexture(FVertexShaderRHIParamRef VertexShader, uint32 TextureIndex, FTextureRHIParamRef NewTexture) final override;
 	virtual void RHISetShaderTexture(FHullShaderRHIParamRef HullShader, uint32 TextureIndex, FTextureRHIParamRef NewTexture) final override;
@@ -303,9 +301,7 @@ public:
 	virtual void RHISetShaderParameter(FHullShaderRHIParamRef HullShader, uint32 BufferIndex, uint32 BaseIndex, uint32 NumBytes, const void* NewValue) final override;
 	virtual void RHISetShaderParameter(FDomainShaderRHIParamRef DomainShader, uint32 BufferIndex, uint32 BaseIndex, uint32 NumBytes, const void* NewValue) final override;
 	virtual void RHISetShaderParameter(FGeometryShaderRHIParamRef GeometryShader, uint32 BufferIndex, uint32 BaseIndex, uint32 NumBytes, const void* NewValue) final override;
-	virtual void RHISetDepthStencilState(FDepthStencilStateRHIParamRef NewState, uint32 StencilRef) final override;
 	virtual void RHISetStencilRef(uint32 StencilRef) final override;
-	virtual void RHISetBlendState(FBlendStateRHIParamRef NewState, const FLinearColor& BlendFactor) final override;
 	virtual void RHISetBlendFactor(const FLinearColor& BlendFactor) final override;
 	virtual void RHISetRenderTargets(uint32 NumSimultaneousRenderTargets, const FRHIRenderTargetView* NewRenderTargets, const FRHIDepthRenderTargetView* NewDepthStencilTarget, uint32 NumUAVs, const FUnorderedAccessViewRHIParamRef* UAVs) final override;
 	virtual void RHISetRenderTargetsAndClear(const FRHISetRenderTargetsInfo& RenderTargetsInfo) final override;
@@ -319,7 +315,6 @@ public:
 	virtual void RHIEndDrawPrimitiveUP() final override;
 	virtual void RHIBeginDrawIndexedPrimitiveUP(uint32 NumPrimitives, uint32 NumVertices, uint32 VertexDataStride, void*& OutVertexData, uint32 MinVertexIndex, uint32 NumIndices, uint32 IndexDataStride, void*& OutIndexData) final override;
 	virtual void RHIEndDrawIndexedPrimitiveUP() final override;
-	virtual void RHIEnableDepthBoundsTest(bool bEnable) final override;
 	virtual void RHISetDepthBounds(float MinDepth, float MaxDepth) final override;
 	virtual void RHIUpdateTextureReference(FTextureReferenceRHIParamRef TextureRef, FTextureRHIParamRef NewTexture) final override;
 
@@ -567,10 +562,6 @@ public:
 	{
 		ContextRedirect(RHISetStreamSource(StreamIndex, VertexBuffer, Offset));
 	}
-	FORCEINLINE virtual void RHISetRasterizerState(FRasterizerStateRHIParamRef NewState) final override
-	{
-		ContextRedirect(RHISetRasterizerState(NewState));
-	}
 	FORCEINLINE virtual void RHISetViewport(uint32 MinX, uint32 MinY, float MinZ, uint32 MaxX, uint32 MaxY, float MaxZ) final override
 	{
 		ContextRedirect(RHISetViewport(MinX, MinY, MinZ, MaxX, MaxY, MaxZ));
@@ -579,10 +570,6 @@ public:
 	{
 		ContextRedirect(RHISetScissorRect(bEnable, MinX, MinY, MaxX, MaxY));
 	}
-	FORCEINLINE virtual void RHISetBoundShaderState(FBoundShaderStateRHIParamRef BoundShaderState) final override
-	{
-		ContextRedirect(RHISetBoundShaderState(BoundShaderState));
-	}
 	FORCEINLINE void RHISetGraphicsPipelineState(FGraphicsPipelineStateRHIParamRef GraphicsPipelineState) final override
 	{
 		ContextRedirect(RHISetGraphicsPipelineState(GraphicsPipelineState));
@@ -687,18 +674,10 @@ public:
 	{
 		ContextRedirect(RHISetShaderParameter(GeometryShader, BufferIndex, BaseIndex, NumBytes, NewValue));
 	}
-	FORCEINLINE virtual void RHISetDepthStencilState(FDepthStencilStateRHIParamRef NewState, uint32 StencilRef) final override
-	{
-		ContextRedirect(RHISetDepthStencilState(NewState, StencilRef));
-	}
 	FORCEINLINE virtual void RHISetStencilRef(uint32 StencilRef) final override
 	{
 		ContextRedirect(RHISetStencilRef(StencilRef));
 	}
-	FORCEINLINE virtual void RHISetBlendState(FBlendStateRHIParamRef NewState, const FLinearColor& BlendFactor) final override
-	{
-		ContextRedirect(RHISetBlendState(NewState, BlendFactor));
-	}
 	FORCEINLINE void RHISetBlendFactor(const FLinearColor& BlendFactor) final override
 	{
 		ContextRedirect(RHISetBlendFactor(BlendFactor));
@@ -744,10 +723,6 @@ public:
 	
 	virtual void RHIEndDrawIndexedPrimitiveUP() final override;
 
-	FORCEINLINE virtual void RHIEnableDepthBoundsTest(bool bEnable) final override
-	{
-		ContextRedirect(RHIEnableDepthBoundsTest(bEnable));
-	}
 	FORCEINLINE virtual void RHISetDepthBounds(float MinDepth, float MaxDepth) final override
 	{
 		ContextRedirect(RHISetDepthBounds(MinDepth, MaxDepth));
diff --git a/Engine/Source/Runtime/D3D12RHI/Private/D3D12Commands.cpp b/Engine/Source/Runtime/D3D12RHI/Private/D3D12Commands.cpp
--- a/Engine/Source/Runtime/D3D12RHI/Private/D3D12Commands.cpp
+++ b/Engine/Source/Runtime/D3D12RHI/Private/D3D12Commands.cpp
@@ -124,13 +124,6 @@ void FD3D12DynamicRHI::RHISetStreamOutTargets(uint32 NumTargets, const FVertexBu
 	CmdContext.StateCache.SetStreamOutTargets(NumTargets, D3DVertexBuffers, D3DOffsets);
 }
 
-// Rasterizer state.
-void FD3D12CommandContext::RHISetRasterizerState(FRasterizerStateRHIParamRef NewStateRHI)
-{
-	// This shouldn't be used. Instead use RHISetGraphicsPipelineState.
-	unimplemented();
-}
-
 void FD3D12CommandContext::RHISetComputeShader(FComputeShaderRHIParamRef ComputeShaderRHI)
 {
 #if D3D12_RHI_RAYTRACING
@@ -468,16 +461,6 @@ void FD3D12CommandContext::RHISetScissorRect(bool bEnable, uint32 MinX, uint32 M
 	}
 }
 
-/**
-* Set bound shader state. This will set the vertex decl/shader, and pixel shader
-* @param BoundShaderState - state resource
-*/
-void FD3D12CommandContext::RHISetBoundShaderState(FBoundShaderStateRHIParamRef BoundShaderStateRHI)
-{
-	// This shouldn't be used. Instead use RHISetGraphicsPipelineState.
-	unimplemented();
-}
-
 static inline D3D_PRIMITIVE_TOPOLOGY GetD3D12PrimitiveType(EPrimitiveType PrimitiveType, bool bUsingTessellation)
 {
 	if (bUsingTessellation)
@@ -571,7 +554,10 @@ void FD3D12CommandContext::RHISetGraphicsPipelineState(FGraphicsPipelineStateRHI
 	// @TODO : really should only discard the constants if the shader state has actually changed.
 	bDiscardSharedConstants = true;
 
-	RHIEnableDepthBoundsTest(GraphicsPipelineState->PipelineStateInitializer.bDepthBounds);
+	if (!GraphicsPipelineState->PipelineStateInitializer.bDepthBounds)
+	{
+		StateCache.SetDepthBounds(0.0f, 1.0f);
+	}
 
 	StateCache.SetGraphicsPipelineState(GraphicsPipelineState);
 	StateCache.SetStencilRef(0);
@@ -909,23 +895,11 @@ void FD3D12CommandContext::ValidateExclusiveDepthStencilAccess(FExclusiveDepthSt
 	}
 }
 
-void FD3D12CommandContext::RHISetDepthStencilState(FDepthStencilStateRHIParamRef NewStateRHI, uint32 StencilRef)
-{
-	// This shouldn't be used. Instead use RHISetGraphicsPipelineState.
-	unimplemented();
-}
-
 void FD3D12CommandContext::RHISetStencilRef(uint32 StencilRef)
 {
 	StateCache.SetStencilRef(StencilRef);
 }
 
-void FD3D12CommandContext::RHISetBlendState(FBlendStateRHIParamRef NewStateRHI, const FLinearColor& BlendFactor)
-{
-	// This shouldn't be used. Instead use RHISetGraphicsPipelineState.
-	unimplemented();
-}
-
 void FD3D12CommandContext::RHISetBlendFactor(const FLinearColor& BlendFactor)
 {
 	StateCache.SetBlendFactor((const float*)&BlendFactor);
@@ -1963,14 +1937,6 @@ void FD3D12DynamicRHI::RHIExecuteCommandList(FRHICommandList* CmdList)
 }
 
 
-void FD3D12CommandContext::RHIEnableDepthBoundsTest(bool bEnable)
-{
-	if (!bEnable)
-	{
-		StateCache.SetDepthBounds(0.0f, 1.0f);
-	}
-}
-
 void FD3D12CommandContext::RHISetDepthBounds(float MinDepth, float MaxDepth)
 {
 	StateCache.SetDepthBounds(MinDepth, MaxDepth);
diff --git a/Engine/Source/Runtime/NullDrv/Public/NullRHI.h b/Engine/Source/Runtime/NullDrv/Public/NullRHI.h
--- a/Engine/Source/Runtime/NullDrv/Public/NullRHI.h
+++ b/Engine/Source/Runtime/NullDrv/Public/NullRHI.h
@@ -9,7 +9,7 @@
 struct Rect;
 
 /** A null implementation of the dynamically bound RHI. */
-class FNullDynamicRHI : public FDynamicRHI , public IRHICommandContext
+class FNullDynamicRHI : public FDynamicRHI , public IRHICommandContextPSOFallback
 {
 public:
 
diff --git a/Engine/Source/Runtime/OpenGLDrv/Public/OpenGLDrv.h b/Engine/Source/Runtime/OpenGLDrv/Public/OpenGLDrv.h
--- a/Engine/Source/Runtime/OpenGLDrv/Public/OpenGLDrv.h
+++ b/Engine/Source/Runtime/OpenGLDrv/Public/OpenGLDrv.h
@@ -296,7 +296,7 @@ struct FOpenGLGPUProfiler : public FGPUProfiler
 
 
 /** The interface which is implemented by the dynamically bound RHI. */
-class OPENGLDRV_API FOpenGLDynamicRHI  final : public FDynamicRHI, public IRHICommandContext
+class OPENGLDRV_API FOpenGLDynamicRHI  final : public FDynamicRHI, public IRHICommandContextPSOFallback
 {
 public:
 
diff --git a/Engine/Source/Runtime/RHI/Public/RHICommandList.h b/Engine/Source/Runtime/RHI/Public/RHICommandList.h
--- a/Engine/Source/Runtime/RHI/Public/RHICommandList.h
+++ b/Engine/Source/Runtime/RHI/Public/RHICommandList.h
@@ -879,16 +879,6 @@ struct FRHICommandDrawIndexedPrimitive final : public FRHICommand<FRHICommandDra
 	RHI_API void Execute(FRHICommandListBase& CmdList);
 };
 
-struct FRHICommandSetBoundShaderState final : public FRHICommand<FRHICommandSetBoundShaderState>
-{
-	FBoundShaderStateRHIParamRef BoundShaderState;
-	FORCEINLINE_DEBUGGABLE FRHICommandSetBoundShaderState(FBoundShaderStateRHIParamRef InBoundShaderState)
-		: BoundShaderState(InBoundShaderState)
-	{
-	}
-	RHI_API void Execute(FRHICommandListBase& CmdList);
-};
-
 struct FRHICommandSetBlendFactor final : public FRHICommand<FRHICommandSetBlendFactor>
 {
 	FLinearColor BlendFactor;
diff --git a/Engine/Source/Runtime/RHI/Public/RHICommandListCommandExecutes.inl b/Engine/Source/Runtime/RHI/Public/RHICommandListCommandExecutes.inl
--- a/Engine/Source/Runtime/RHI/Public/RHICommandListCommandExecutes.inl
+++ b/Engine/Source/Runtime/RHI/Public/RHICommandListCommandExecutes.inl
@@ -223,12 +223,6 @@ void FRHICommandDrawIndexedPrimitive::Execute(FRHICommandListBase& CmdList)
 	INTERNAL_DECORATOR(RHIDrawIndexedPrimitive)(IndexBuffer, BaseVertexIndex, FirstInstance, NumVertices, StartIndex, NumPrimitives, NumInstances);
 }
 
-void FRHICommandSetBoundShaderState::Execute(FRHICommandListBase& CmdList)
-{
-	RHISTAT(SetBoundShaderState);
-	INTERNAL_DECORATOR(RHISetBoundShaderState)(BoundShaderState);
-}
-
 void FRHICommandSetBlendFactor::Execute(FRHICommandListBase& CmdList)
 {
 	RHISTAT(SetBlendFactor);
diff --git a/Engine/Source/Runtime/RHI/Public/RHIContext.h b/Engine/Source/Runtime/RHI/Public/RHIContext.h
--- a/Engine/Source/Runtime/RHI/Public/RHIContext.h
+++ b/Engine/Source/Runtime/RHI/Public/RHIContext.h
@@ -28,16 +28,6 @@
 enum class EResourceTransitionAccess;
 enum class EResourceTransitionPipeline;
 
-
-FORCEINLINE FBoundShaderStateRHIRef RHICreateBoundShaderState(
-	FVertexDeclarationRHIParamRef VertexDeclaration,
-	FVertexShaderRHIParamRef VertexShader,
-	FHullShaderRHIParamRef HullShader,
-	FDomainShaderRHIParamRef DomainShader,
-	FPixelShaderRHIParamRef PixelShader,
-	FGeometryShaderRHIParamRef GeometryShader
-);
-
 /** Context that is capable of doing Compute work.  Can be async or compute on the gfx pipe. */
 class IRHIComputeContext
 {
@@ -145,26 +135,6 @@
 	}
 };
 
-// These states are now set by the Pipeline State Object and are now deprecated
-class IRHIDeprecatedContext
-{
-public:
-	/**
-	* Set bound shader state. This will set the vertex decl/shader, and pixel shader
-	* @param BoundShaderState - state resource
-	*/
-	virtual void RHISetBoundShaderState(FBoundShaderStateRHIParamRef BoundShaderState) = 0;
-
-	virtual void RHISetDepthStencilState(FDepthStencilStateRHIParamRef NewState, uint32 StencilRef) = 0;
-
-	virtual void RHISetRasterizerState(FRasterizerStateRHIParamRef NewState) = 0;
-
-	// Allows to set the blend state, parameter can be created with RHICreateBlendState()
-	virtual void RHISetBlendState(FBlendStateRHIParamRef NewState, const FLinearColor& BlendFactor) = 0;
-
-	virtual void RHIEnableDepthBoundsTest(bool bEnable) = 0;
-};
-
 struct FAccelerationStructureUpdateParams
 {
 	FRayTracingGeometryRHIParamRef Geometry;
@@ -172,7 +142,7 @@
 };
 
 /** The interface RHI command context. Sometimes the RHI handles these. On platforms that can processes command lists in parallel, it is a separate object. */
-class IRHICommandContext : public IRHIComputeContext, public IRHIDeprecatedContext
+class IRHICommandContext : public IRHIComputeContext
 {
 public:
 	virtual ~IRHICommandContext()
@@ -352,36 +322,7 @@
 	// @param MaxY excluding like Win32 RECT
 	virtual void RHISetScissorRect(bool bEnable, uint32 MinX, uint32 MinY, uint32 MaxX, uint32 MaxY) = 0;
 
-	/**
-	* This will set most relevant pipeline state. Legacy APIs are expected to set corresponding disjoint state as well.
-	* @param GraphicsShaderState - the graphics pipeline state
-	* This implementation is only in place while we transition/refactor.
-	*/
-	virtual void RHISetGraphicsPipelineState(FGraphicsPipelineStateRHIParamRef GraphicsState)
-	{
-		FRHIGraphicsPipelineStateFallBack* FallbackGraphicsState = static_cast<FRHIGraphicsPipelineStateFallBack*>(GraphicsState);
+	virtual void RHISetGraphicsPipelineState(FGraphicsPipelineStateRHIParamRef GraphicsState) = 0;
 
-		auto& PsoInit = FallbackGraphicsState->Initializer;
-
-		RHISetBoundShaderState(
-			RHICreateBoundShaderState(
-				PsoInit.BoundShaderState.VertexDeclarationRHI,
-				PsoInit.BoundShaderState.VertexShaderRHI,
-				PsoInit.BoundShaderState.HullShaderRHI,
-				PsoInit.BoundShaderState.DomainShaderRHI,
-				PsoInit.BoundShaderState.PixelShaderRHI,
-				PsoInit.BoundShaderState.GeometryShaderRHI
-			).GetReference()
-		);
-
-		RHISetDepthStencilState(FallbackGraphicsState->Initializer.DepthStencilState, 0);
-		RHISetRasterizerState(FallbackGraphicsState->Initializer.RasterizerState);
-		RHISetBlendState(FallbackGraphicsState->Initializer.BlendState, FLinearColor(1.0f, 1.0f, 1.0f));
-		if (GSupportsDepthBoundsTest)
-		{
-			RHIEnableDepthBoundsTest(FallbackGraphicsState->Initializer.bDepthBounds);
-		}
-	}
-
 	/** Set the shader resource view of a surface.  This is used for binding TextureMS parameter types that need a multi sampled view. */
 	virtual void RHISetShaderTexture(FVertexShaderRHIParamRef VertexShader, uint32 TextureIndex, FTextureRHIParamRef NewTexture) = 0;
@@ -718,3 +659,63 @@
 	protected:
 		FRHIRenderPassInfo RenderPassInfo;
 };
+
+
+
+FORCEINLINE FBoundShaderStateRHIRef RHICreateBoundShaderState(
+	FVertexDeclarationRHIParamRef VertexDeclaration,
+	FVertexShaderRHIParamRef VertexShader,
+	FHullShaderRHIParamRef HullShader,
+	FDomainShaderRHIParamRef DomainShader,
+	FPixelShaderRHIParamRef PixelShader,
+	FGeometryShaderRHIParamRef GeometryShader
+);
+
+
+// Command Context for RHIs that do not support real Graphics Pipelines.
+class IRHICommandContextPSOFallback : public IRHICommandContext
+{
+public:
+	/**
+	* Set bound shader state. This will set the vertex decl/shader, and pixel shader
+	* @param BoundShaderState - state resource
+	*/
+	virtual void RHISetBoundShaderState(FBoundShaderStateRHIParamRef BoundShaderState) = 0;
+
+	virtual void RHISetDepthStencilState(FDepthStencilStateRHIParamRef NewState, uint32 StencilRef) = 0;
+
+	virtual void RHISetRasterizerState(FRasterizerStateRHIParamRef NewState) = 0;
+
+	virtual void RHISetBlendState(FBlendStateRHIParamRef NewState, const FLinearColor& BlendFactor) = 0;
+
+	virtual void RHIEnableDepthBoundsTest(bool bEnable) = 0;
+
+	/**
+	* This will set most relevant pipeline state. Legacy APIs are expected to set corresponding disjoint state as well.
+	* @param GraphicsShaderState - the graphics pipeline state
+	*/
+	virtual void RHISetGraphicsPipelineState(FGraphicsPipelineStateRHIParamRef GraphicsState) override
+	{
+		FRHIGraphicsPipelineStateFallBack* FallbackGraphicsState = static_cast<FRHIGraphicsPipelineStateFallBack*>(GraphicsState);
+		FGraphicsPipelineStateInitializer& PsoInit = FallbackGraphicsState->Initializer;
+
+		RHISetBoundShaderState(
+			RHICreateBoundShaderState(
+				PsoInit.BoundShaderState.VertexDeclarationRHI,
+				PsoInit.BoundShaderState.VertexShaderRHI,
+				PsoInit.BoundShaderState.HullShaderRHI,
+				PsoInit.BoundShaderState.DomainShaderRHI,
+				PsoInit.BoundShaderState.PixelShaderRHI,
+				PsoInit.BoundShaderState.GeometryShaderRHI
+			).GetReference()
+		);
+
+		RHISetDepthStencilState(FallbackGraphicsState->Initializer.DepthStencilState, 0);
+		RHISetRasterizerState(FallbackGraphicsState->Initializer.RasterizerState);
+		RHISetBlendState(FallbackGraphicsState->Initializer.BlendState, FLinearColor(1.0f, 1.0f, 1.0f));
+		if (GSupportsDepthBoundsTest)
+		{
+			RHIEnableDepthBoundsTest(FallbackGraphicsState->Initializer.bDepthBounds);
+		}
+	}
+};
diff --git a/Engine/Source/Runtime/VulkanRHI/Private/VulkanCommands.cpp b/Engine/Source/Runtime/VulkanRHI/Private/VulkanCommands.cpp
--- a/Engine/Source/Runtime/VulkanRHI/Private/VulkanCommands.cpp
+++ b/Engine/Source/Runtime/VulkanRHI/Private/VulkanCommands.cpp
@@ -62,11 +62,6 @@ void FVulkanDynamicRHI::RHISetStreamOutTargets(uint32 NumTargets, const FVertexB
 	VULKAN_SIGNAL_UNIMPLEMENTED();
 }
 
-void FVulkanCommandListContext::RHISetRasterizerState(FRasterizerStateRHIParamRef NewStateRHI)
-{
-	check(0);
-}
-
 void FVulkanCommandListContext::RHISetComputeShader(FComputeShaderRHIParamRef ComputeShaderRHI)
 {
 	FVulkanComputeShader* ComputeShader = ResourceCast(ComputeShaderRHI);
@@ -161,12 +156,6 @@ void FVulkanCommandListContext::RHIDispatchIndirectComputeShader(FVertexBufferRH
 	VulkanRHI::/*Debug*/HeavyWeightBarrier(CmdBuffer/*, 2*/);
 }
 
-void FVulkanCommandListContext::RHISetBoundShaderState(FBoundShaderStateRHIParamRef BoundShaderStateRHI)
-{
-	check(0);
-}
-
-
 void FVulkanCommandListContext::RHISetUAVParameter(FComputeShaderRHIParamRef ComputeShaderRHI, uint32 UAVIndex, FUnorderedAccessViewRHIParamRef UAVRHI)
 {
 	check(PendingComputeState->GetCurrentShader() == ResourceCast(ComputeShaderRHI));
@@ -750,16 +739,6 @@ void FVulkanCommandListContext::RHISetShaderUniformBuffer(FComputeShaderRHIParam
 	}
 }
 
-void FVulkanCommandListContext::RHISetDepthStencilState(FDepthStencilStateRHIParamRef NewStateRHI, uint32 StencilRef)
-{
-	check(0);
-}
-
-void FVulkanCommandListContext::RHISetBlendState(FBlendStateRHIParamRef NewStateRHI, const FLinearColor& BlendFactor)
-{
-	check(0);
-}
-
 void FVulkanCommandListContext::RHISetStencilRef(uint32 StencilRef)
 {
 	PendingGfxState->SetStencilRef(StencilRef);
@@ -1059,11 +1038,6 @@ void FVulkanDynamicRHI::RHIExecuteCommandList(FRHICommandList* CmdList)
 	VULKAN_SIGNAL_UNIMPLEMENTED();
 }
 
-void FVulkanCommandListContext::RHIEnableDepthBoundsTest(bool bEnable)
-{
-	check(0);
-}
-
 void FVulkanCommandListContext::RHISetDepthBounds(float MinDepth, float MaxDepth)
 {
 	FVulkanCmdBuffer* CmdBuffer = CommandBufferManager->GetActiveCmdBuffer();
diff --git a/Engine/Source/Runtime/VulkanRHI/Public/VulkanContext.h b/Engine/Source/Runtime/VulkanRHI/Public/VulkanContext.h
--- a/Engine/Source/Runtime/VulkanRHI/Public/VulkanContext.h
+++ b/Engine/Source/Runtime/VulkanRHI/Public/VulkanContext.h
@@ -170,10 +170,8 @@ public:
 	}
 
 	virtual void RHISetStreamSource(uint32 StreamIndex, FVertexBufferRHIParamRef VertexBuffer, uint32 Offset) final override;
-	virtual void RHISetRasterizerState(FRasterizerStateRHIParamRef NewState) final override;
 	virtual void RHISetViewport(uint32 MinX, uint32 MinY, float MinZ, uint32 MaxX, uint32 MaxY, float MaxZ) final override;
 	virtual void RHISetScissorRect(bool bEnable, uint32 MinX, uint32 MinY, uint32 MaxX, uint32 MaxY) final override;
-	virtual void RHISetBoundShaderState(FBoundShaderStateRHIParamRef BoundShaderState) final override;
 	virtual void RHISetGraphicsPipelineState(FGraphicsPipelineStateRHIParamRef GraphicsState) final override;
 	virtual void RHISetShaderTexture(FVertexShaderRHIParamRef VertexShader, uint32 TextureIndex, FTextureRHIParamRef NewTexture) final override;
 	virtual void RHISetShaderTexture(FHullShaderRHIParamRef HullShader, uint32 TextureIndex, FTextureRHIParamRef NewTexture) final override;
@@ -207,9 +205,7 @@ public:
 	virtual void RHISetShaderParameter(FGeometryShaderRHIParamRef GeometryShader, uint32 BufferIndex, uint32 BaseIndex, uint32 NumBytes, const void* NewValue) final override;
 	virtual void RHISetShaderParameter(FPixelShaderRHIParamRef PixelShader, uint32 BufferIndex, uint32 BaseIndex, uint32 NumBytes, const void* NewValue) final override;
 	virtual void RHISetShaderParameter(FComputeShaderRHIParamRef ComputeShader, uint32 BufferIndex, uint32 BaseIndex, uint32 NumBytes, const void* NewValue) final override;
-	virtual void RHISetDepthStencilState(FDepthStencilStateRHIParamRef NewState, uint32 StencilRef) final override;
 	virtual void RHISetStencilRef(uint32 StencilRef) final override;
-	virtual void RHISetBlendState(FBlendStateRHIParamRef NewState, const FLinearColor& BlendFactor) final override;
 	virtual void RHISetRenderTargets(uint32 NumSimultaneousRenderTargets, const FRHIRenderTargetView* NewRenderTargets, const FRHIDepthRenderTargetView* NewDepthStencilTarget, uint32 NumUAVs, const FUnorderedAccessViewRHIParamRef* UAVs) final override;
 	virtual void RHISetRenderTargetsAndClear(const FRHISetRenderTargetsInfo& RenderTargetsInfo) final override;
 	virtual void RHIDrawPrimitive(uint32 BaseVertexIndex, uint32 NumPrimitives, uint32 NumInstances) final override;
@@ -221,7 +217,6 @@ public:
 	virtual void RHIEndDrawPrimitiveUP() final override;
 	virtual void RHIBeginDrawIndexedPrimitiveUP(uint32 NumPrimitives, uint32 NumVertices, uint32 VertexDataStride, void*& OutVertexData, uint32 MinVertexIndex, uint32 NumIndices, uint32 IndexDataStride, void*& OutIndexData) final override;
 	virtual void RHIEndDrawIndexedPrimitiveUP() final override;
-	virtual void RHIEnableDepthBoundsTest(bool bEnable) final override;
 	virtual void RHISetDepthBounds(float MinDepth, float MaxDepth) final override;
 	virtual void RHIPushEvent(const TCHAR* Name, FColor Color) final override;
 	virtual void RHIPopEvent() final override;
diff --git a/Engine/Source/Runtime/Windows/D3D11RHI/Private/D3D11RHIPrivate.h b/Engine/Source/Runtime/Windows/D3D11RHI/Private/D3D11RHIPrivate.h
--- a/Engine/Source/Runtime/Windows/D3D11RHI/Private/D3D11RHIPrivate.h
+++ b/Engine/Source/Runtime/Windows/D3D11RHI/Private/D3D11RHIPrivate.h
@@ -331,7 +331,7 @@ private:
 struct AGSContext;
 
 /** The interface which is implemented by the dynamically bound RHI. */
-class D3D11RHI_API FD3D11DynamicRHI : public FDynamicRHI, public IRHICommandContext
+class D3D11RHI_API FD3D11DynamicRHI : public FDynamicRHI, public IRHICommandContextPSOFallback
 {
 public:
 
@@ -512,7 +512,7 @@ public:
 	virtual void RHISetGraphicsPipelineState(FGraphicsPipelineStateRHIParamRef GraphicsState) final override
 	{
 		FRHIGraphicsPipelineStateFallBack* FallbackGraphicsState = static_cast<FRHIGraphicsPipelineStateFallBack*>(GraphicsState);
-		IRHICommandContext::RHISetGraphicsPipelineState(GraphicsState);
+		IRHICommandContextPSOFallback::RHISetGraphicsPipelineState(GraphicsState);
 		// Store the PSO's primitive (after since IRHICommandContext::RHISetGraphicsPipelineState sets the BSS)
 		PrimitiveType = FallbackGraphicsState->Initializer.PrimitiveType;
 	}
