DR - Remove dummy function

diff --git a/Engine/Source/Runtime/Engine/Private/ShaderCompiler/ShaderCompiler.cpp b/Engine/Source/Runtime/Engine/Private/ShaderCompiler/ShaderCompiler.cpp
--- a/Engine/Source/Runtime/Engine/Private/ShaderCompiler/ShaderCompiler.cpp
+++ b/Engine/Source/Runtime/Engine/Private/ShaderCompiler/ShaderCompiler.cpp
@@ -3453,12 +3453,6 @@ public:
 			// Need to recache all cached mesh draw commands, as they store pointers to material uniform buffers which we just invalidated.
 			GetRendererModule().UpdateStaticDrawLists();
 		}
-
-		ENQUEUE_UNIQUE_RENDER_COMMAND(
-			FRecreateBoundShaderStates,
-		{
-			RHIRecreateRecursiveBoundShaderStates();
-		});
 	}
 
 private:
diff --git a/Engine/Source/Runtime/RHI/Public/DynamicRHI.h b/Engine/Source/Runtime/RHI/Public/DynamicRHI.h
--- a/Engine/Source/Runtime/RHI/Public/DynamicRHI.h
+++ b/Engine/Source/Runtime/RHI/Public/DynamicRHI.h
@@ -917,9 +917,6 @@ public:
 	virtual bool RHIEnqueueDecompress(uint8_t* SrcBuffer, uint8_t* DestBuffer, int CompressedSize, void* ErrorCodeBuffer) { return false; }
 	virtual bool RHIEnqueueCompress(uint8_t* SrcBuffer, uint8_t* DestBuffer, int UnCompressedSize, void* ErrorCodeBuffer) { return false; }
 
-	// FlushType: Flush Immediate
-	virtual void RHIRecreateRecursiveBoundShaderStates() {}
-
 	/**
 	*	Retrieve available screen resolutions.
 	*
diff --git a/Engine/Source/Runtime/RHI/Public/RHICommandList.h b/Engine/Source/Runtime/RHI/Public/RHICommandList.h
--- a/Engine/Source/Runtime/RHI/Public/RHICommandList.h
+++ b/Engine/Source/Runtime/RHI/Public/RHICommandList.h
@@ -4760,12 +4760,6 @@ FORCEINLINE void* RHIGetNativeDevice()
 	return FRHICommandListExecutor::GetImmediateCommandList().GetNativeDevice();
 }
 
-FORCEINLINE void RHIRecreateRecursiveBoundShaderStates()
-{
-	FRHICommandListExecutor::GetImmediateCommandList(). ImmediateFlush(EImmediateFlushType::FlushRHIThread);
-	GDynamicRHI->RHIRecreateRecursiveBoundShaderStates();
-}
-
 FORCEINLINE FRHIShaderLibraryRef RHICreateShaderLibrary(EShaderPlatform Platform, FString const& FilePath, FString const& Name)
 {
     return GDynamicRHI->RHICreateShaderLibrary(Platform, FilePath, Name);
