First round of compile and runtime fixes.

diff --git a/Engine/Source/Runtime/Engine/Public/MaterialShared.h b/Engine/Source/Runtime/Engine/Public/MaterialShared.h
--- a/Engine/Source/Runtime/Engine/Public/MaterialShared.h
+++ b/Engine/Source/Runtime/Engine/Public/MaterialShared.h
@@ -413,7 +413,7 @@ public:
 		bUsesGlobalDistanceField(false),
 		bUsesPixelDepthOffset(false),
 		bUsesSceneDepthLookup(false),
-		bUsesVelocitySceneTexture(false)
+		bUsesVelocitySceneTexture(false),
 		bUsesDistanceCullFade(false)
 	{}
 
diff --git a/Engine/Source/Runtime/Engine/Public/SkeletalRenderPublic.h b/Engine/Source/Runtime/Engine/Public/SkeletalRenderPublic.h
--- a/Engine/Source/Runtime/Engine/Public/SkeletalRenderPublic.h
+++ b/Engine/Source/Runtime/Engine/Public/SkeletalRenderPublic.h
@@ -192,8 +192,6 @@ public:
 	FORCEINLINE FSkeletalMeshRenderData& GetSkeletalMeshRenderData() const { return *SkeletalMeshRenderData; }
 
 #if RHI_RAYTRACING
-	bool bRequireRecreatingRayTracingGeometry;
-
 	/** Retrieve ray tracing geometry from the underlying mesh object */
 	virtual const FRayTracingGeometry* GetRayTracingGeometry() const { return nullptr; }
 #endif // RHI_RAYTRACING
@@ -233,6 +231,10 @@ public:
 	/** This is set to true when we have sent our Mesh data to the rendering thread at least once as it needs to have have a datastructure created there for each MeshObject **/
 	bool bHasBeenUpdatedAtLeastOnce;
 
+#if RHI_RAYTRACING
+	bool bRequireRecreatingRayTracingGeometry;
+#endif
+
 #if WITH_EDITORONLY_DATA
 	/** Index of the section to preview... If set to -1, all section will be rendered */
 	int32 SectionIndexPreview;
diff --git a/Engine/Source/Runtime/RHI/Private/PipelineStateCache.cpp b/Engine/Source/Runtime/RHI/Private/PipelineStateCache.cpp
--- a/Engine/Source/Runtime/RHI/Private/PipelineStateCache.cpp
+++ b/Engine/Source/Runtime/RHI/Private/PipelineStateCache.cpp
@@ -557,66 +557,9 @@ private:
 
 };
 
-class FGraphicsPipelineStateKey
-{
-public:
-	explicit FGraphicsPipelineStateKey(const FGraphicsPipelineStateInitializer& InInitializer)
-	{
-		FMemory::Memcpy(Initializer, InInitializer);
-		
-		if (InInitializer.BoundShaderState.VertexShaderRHI)
-		{
-			VertexShaderHash = InInitializer.BoundShaderState.VertexShaderRHI->GetHash();
-		}
-		if (InInitializer.BoundShaderState.PixelShaderRHI)
-		{
-			PixelShaderHash = InInitializer.BoundShaderState.PixelShaderRHI->GetHash();
-		}
-		if (InInitializer.BoundShaderState.GeometryShaderRHI)
-		{
-			GeometryShaderHash = InInitializer.BoundShaderState.GeometryShaderRHI->GetHash();
-		}
-		if (InInitializer.BoundShaderState.HullShaderRHI)
-		{
-			HullShaderHash = InInitializer.BoundShaderState.HullShaderRHI->GetHash();
-		}
-		if (InInitializer.BoundShaderState.DomainShaderRHI)
-		{
-			DomainShaderHash = InInitializer.BoundShaderState.DomainShaderRHI->GetHash();
-		}
-	}
-	
-	bool operator==(const FGraphicsPipelineStateKey& rhs) const
-	{
-		if (Initializer == rhs.Initializer && 
-			VertexShaderHash == rhs.VertexShaderHash &&
-			PixelShaderHash == rhs.PixelShaderHash &&
-			GeometryShaderHash == rhs.GeometryShaderHash &&
-			HullShaderHash == rhs.HullShaderHash &&
-			DomainShaderHash == rhs.DomainShaderHash)
-		{
-			return true;
-		}
-
-		return false;
-	}
-
-	FGraphicsPipelineStateInitializer Initializer;
-	FSHAHash VertexShaderHash;
-	FSHAHash PixelShaderHash;
-	FSHAHash GeometryShaderHash;
-	FSHAHash HullShaderHash;
-	FSHAHash DomainShaderHash;
-};
-
-static inline uint32 GetTypeHash(const FGraphicsPipelineStateKey& PipelineStateKey)
-{
-	return GetTypeHash(PipelineStateKey.Initializer);
-}
-
 // Typed caches for compute and graphics
 typedef TDiscardableKeyValueCache< FRHIComputeShader*, FComputePipelineState*> FComputePipelineCache;
-typedef TSharedPipelineStateCache<FGraphicsPipelineStateKey, FGraphicsPipelineState*> FGraphicsPipelineCache;
+typedef TSharedPipelineStateCache<FGraphicsPipelineStateInitializer, FGraphicsPipelineState*> FGraphicsPipelineCache;
 
 // These are the actual caches for both pipelines
 FComputePipelineCache GComputePipelineCache;
@@ -940,7 +883,6 @@ FComputePipelineState* PipelineStateCache::GetAndOrCreateComputePipelineState(FR
 FRHIRayTracingPipelineState* PipelineStateCache::GetAndOrCreateRayTracingPipelineState(const FRayTracingPipelineStateInitializer& Initializer)
 {
 	LLM_SCOPE(ELLMTag::PSO);
-	SCOPE_CYCLE_COUNTER(STAT_GetOrCreatePSO);
 
 	check(IsInRenderingThread() || IsInParallelRenderingThread());
 
@@ -1052,8 +994,7 @@ FGraphicsPipelineState* PipelineStateCache::GetAndOrCreateGraphicsPipelineState(
 
 	FGraphicsPipelineState* OutCachedState = nullptr;
 
-	FGraphicsPipelineStateKey PipelineStateKey(*Initializer);
-	bool bWasFound = GGraphicsPipelineCache.Find(PipelineStateKey, OutCachedState);
+	bool bWasFound = GGraphicsPipelineCache.Find(*Initializer, OutCachedState);
 
 	if (bWasFound == false)
 	{
@@ -1079,7 +1020,7 @@ FGraphicsPipelineState* PipelineStateCache::GetAndOrCreateGraphicsPipelineState(
 		}
 
 		// GGraphicsPipelineCache.Add(*Initializer, OutCachedState, LockFlags);
-		GGraphicsPipelineCache.Add(PipelineStateKey, OutCachedState);
+		GGraphicsPipelineCache.Add(*Initializer, OutCachedState);
 	}
 	else
 	{
diff --git a/Engine/Source/Runtime/RHI/Public/RHIResources.h b/Engine/Source/Runtime/RHI/Public/RHIResources.h
--- a/Engine/Source/Runtime/RHI/Public/RHIResources.h
+++ b/Engine/Source/Runtime/RHI/Public/RHIResources.h
@@ -1736,12 +1736,7 @@ public:
 			BoundShaderState.PixelShaderRHI != rhs.BoundShaderState.PixelShaderRHI ||
 			BoundShaderState.GeometryShaderRHI != rhs.BoundShaderState.GeometryShaderRHI ||
 			BoundShaderState.DomainShaderRHI != rhs.BoundShaderState.DomainShaderRHI ||
-			BoundShaderState.HullShaderRHI != rhs.BoundShaderState.HullShaderRHI ||
-			VertexShaderHash != rhs.VertexShaderHash ||
-			PixelShaderHash != rhs.PixelShaderHash ||
-			GeometryShaderHash != rhs.GeometryShaderHash ||
-			HullShaderHash != rhs.HullShaderHash ||
-			DomainShaderHash != rhs.DomainShaderHash ||
+			BoundShaderState.HullShaderRHI != rhs.BoundShaderState.HullShaderRHI ||			
 			BlendState != rhs.BlendState || 
 			RasterizerState != rhs.RasterizerState || 
 			DepthStencilState != rhs.DepthStencilState ||
@@ -1916,6 +1911,11 @@ public:
 	bool operator==(const FGraphicsPipelineStateInitializer& rhs) const
 	{
 		if (!FGraphicsMinimalPipelineStateInitializer::operator ==(rhs) ||
+			VertexShaderHash != rhs.VertexShaderHash ||
+			PixelShaderHash != rhs.PixelShaderHash ||
+			GeometryShaderHash != rhs.GeometryShaderHash ||
+			HullShaderHash != rhs.HullShaderHash ||
+			DomainShaderHash != rhs.DomainShaderHash ||
 			RenderTargetsEnabled != rhs.RenderTargetsEnabled ||
 			RenderTargetFormats != rhs.RenderTargetFormats || 
 			RenderTargetFlags != rhs.RenderTargetFlags || 
@@ -1953,15 +1953,6 @@ public:
 		return RenderTargetsEnabled;
 	}
 
-	// TODO: [PSO API] - As we migrate reuse existing API objects, but eventually we can move to the direct initializers. 
-	// When we do that work, move this to RHI.h as its more appropriate there, but here for now since dependent typdefs are here.
-	FBoundShaderStateInput			BoundShaderState;
-	FBlendStateRHIParamRef			BlendState;
-	FRasterizerStateRHIParamRef		RasterizerState;
-	FDepthStencilStateRHIParamRef	DepthStencilState;
-	FImmutableSamplerState			ImmutableSamplerState;
-	bool							bDepthBounds = false;
-	EPrimitiveType					PrimitiveType;
 	FSHAHash						VertexShaderHash;
 	FSHAHash						PixelShaderHash;
 	FSHAHash						GeometryShaderHash;
diff --git a/Engine/Source/Runtime/Renderer/Private/ScenePrivate.h b/Engine/Source/Runtime/Renderer/Private/ScenePrivate.h
--- a/Engine/Source/Runtime/Renderer/Private/ScenePrivate.h
+++ b/Engine/Source/Runtime/Renderer/Private/ScenePrivate.h
@@ -2982,11 +2982,5 @@ inline bool ShouldIncludeDomainInMeshPass(EMaterialDomain Domain)
 	return Domain != MD_Volume;
 }
 
-template<>
-TStaticMeshDrawList<TMobileBasePassDrawingPolicy<FUniformLightMapPolicy>>& FScene::GetMobileBasePassDrawList<FUniformLightMapPolicy>(EBasePassDrawListType DrawType);
-
-template<>
-TStaticMeshDrawList<TMobileBasePassDrawingPolicy<FUniformLightMapPolicy>>& FScene::GetMobileBasePassCSMDrawList<FUniformLightMapPolicy>(EBasePassDrawListType DrawType);
-
 #include "BasePassRendering.inl"
 
diff --git a/Engine/Source/Runtime/Renderer/Private/SceneRendering.h b/Engine/Source/Runtime/Renderer/Private/SceneRendering.h
--- a/Engine/Source/Runtime/Renderer/Private/SceneRendering.h
+++ b/Engine/Source/Runtime/Renderer/Private/SceneRendering.h
@@ -1818,7 +1818,8 @@ protected:
 
 	/** Whether GPU particle collisions simulation is allowed. */
 	bool IsGPUParticleCollisionEnabled(const FViewInfo& View);
-	void SortMobileBasePassAfterShadowInit(FExclusiveDepthStencil::Type BasePassDepthStencilAccess, FMeshDrawCommandsPerPassPerView& VisibleCommandsPerView);
+	void SortMobileBasePassAfterShadowInit(FExclusiveDepthStencil::Type BasePassDepthStencilAccess, FViewVisibleCommandsPerView& ViewCommandsPerView);
+	void SetupMobileBasePassAfterShadowInit(FExclusiveDepthStencil::Type BasePassDepthStencilAccess, FViewVisibleCommandsPerView& ViewCommandsPerView);
 
 	void UpdateOpaqueBasePassUniformBuffer(FRHICommandListImmediate& RHICmdList, const FViewInfo& View);
 	void UpdateTranslucentBasePassUniformBuffer(FRHICommandListImmediate& RHICmdList, const FViewInfo& View);
diff --git a/Engine/Source/Runtime/Renderer/Private/ShadowSetupMobile.cpp b/Engine/Source/Runtime/Renderer/Private/ShadowSetupMobile.cpp
--- a/Engine/Source/Runtime/Renderer/Private/ShadowSetupMobile.cpp
+++ b/Engine/Source/Runtime/Renderer/Private/ShadowSetupMobile.cpp
@@ -426,7 +426,7 @@ void FMobileSceneRenderer::BuildCSMVisibilityState(FLightSceneInfo* LightSceneIn
 
 			if (ViewFamily.EngineShowFlags.ShadowFrustums)
 			{
-				FViewElementPDI ShadowFrustumPDI(&View, NULL);
+				FViewElementPDI ShadowFrustumPDI(&View, nullptr, nullptr);
 				
 				const FMatrix ViewMatrix = View.ViewMatrices.GetViewMatrix();
 				const FMatrix ProjectionMatrix = View.ViewMatrices.GetProjectionMatrix();
