Accidentally disabled tessellation on Metal.

diff --git a/Engine/Source/Runtime/RHI/Private/RHI.cpp b/Engine/Source/Runtime/RHI/Private/RHI.cpp
--- a/Engine/Source/Runtime/RHI/Private/RHI.cpp
+++ b/Engine/Source/Runtime/RHI/Private/RHI.cpp
@@ -792,7 +792,7 @@ RHI_API uint32 RHIGetShaderLanguageVersion(const EShaderPlatform Platform)
 
 RHI_API bool RHISupportsTessellation(const EShaderPlatform Platform)
 {
-	if (IsFeatureLevelSupported(Platform, ERHIFeatureLevel::SM5) && !IsMetalPlatform(Platform))
+	if (IsFeatureLevelSupported(Platform, ERHIFeatureLevel::SM5))
 	{
 		return (Platform == SP_PCD3D_SM5) || (Platform == SP_XBOXONE_D3D12) || (Platform == SP_OPENGL_SM5) || (Platform == SP_OPENGL_ES31_EXT) || (Platform == SP_METAL_SM5) /* || (IsVulkanSM5Platform(Platform)*/;
 	}
