DR - Save 24 bytes in the PSO initializer
#rnx

diff --git a/Engine/Source/Runtime/RHI/Private/PipelineStateCache.cpp b/Engine/Source/Runtime/RHI/Private/PipelineStateCache.cpp
--- a/Engine/Source/Runtime/RHI/Private/PipelineStateCache.cpp
+++ b/Engine/Source/Runtime/RHI/Private/PipelineStateCache.cpp
@@ -115,7 +115,7 @@ static void HandlePipelineCreationFailure(const FGraphicsPipelineStateInitialize
 	for(int32 i = 0; i < Init.RenderTargetFormats.Num(); ++i)
 	{
 		//#todo-mattc GetPixelFormatString is not available in scw. Need to move it so we can print more info here.
-		UE_LOG(LogRHI, Error, TEXT("0x%x"), Init.RenderTargetFormats[i]);
+		UE_LOG(LogRHI, Error, TEXT("0x%x"), (uint32)Init.RenderTargetFormats[i]);
 	}
 	
 	UE_LOG(LogRHI, Error, TEXT("Depth Stencil Format:"));
diff --git a/Engine/Source/Runtime/RHI/Public/RHIResources.h b/Engine/Source/Runtime/RHI/Public/RHIResources.h
--- a/Engine/Source/Runtime/RHI/Public/RHIResources.h
+++ b/Engine/Source/Runtime/RHI/Public/RHIResources.h
@@ -1838,7 +1838,7 @@ public:
 class FGraphicsPipelineStateInitializer final : public FGraphicsMinimalPipelineStateInitializer
 {
 public:
-	using TRenderTargetFormats		= TStaticArray<EPixelFormat, MaxSimultaneousRenderTargets>;
+	using TRenderTargetFormats		= TStaticArray<TEnumAsByte<EPixelFormat>, MaxSimultaneousRenderTargets>;
 	using TRenderTargetFlags		= TStaticArray<uint32, MaxSimultaneousRenderTargets>;
 
 	FGraphicsPipelineStateInitializer()
@@ -1854,6 +1854,7 @@ public:
 		, NumSamples(0)
 		, Flags(0)
 	{
+		static_assert(PF_MAX <= 255, "EPixelFormat is assumed to be a byte; check usage of TEnumAsByte<EPixelFormat>");
 	}
 
 	FGraphicsPipelineStateInitializer(const FGraphicsMinimalPipelineStateInitializer& InMinimalState)
