DR - Merge fix
#rnx

diff --git a/Engine/Source/Runtime/RHI/Public/RHICommandList.h b/Engine/Source/Runtime/RHI/Public/RHICommandList.h
--- a/Engine/Source/Runtime/RHI/Public/RHICommandList.h
+++ b/Engine/Source/Runtime/RHI/Public/RHICommandList.h
@@ -2762,7 +2762,7 @@ public:
 	{
 		if (Bypass())
 		{
-			CMD_CONTEXT(RHICopyToStagingBuffer)(SourceBuffer, DestinationStagingBuffer, Offset, NumBytes, Fence);
+			GetContext().RHICopyToStagingBuffer(SourceBuffer, DestinationStagingBuffer, Offset, NumBytes, Fence);
 			return;
 		}
 		ALLOC_COMMAND(FRHICommandCopyToStagingBuffer<ECmdList::EGfx>)(SourceBuffer, DestinationStagingBuffer, Offset, NumBytes, Fence);
@@ -3300,7 +3300,7 @@ public:
 	{
 		if (Bypass())
 		{
-			CMD_CONTEXT(RHICopyToStagingBuffer)(SourceBuffer, DestinationStagingBuffer, Offset, NumBytes, Fence);
+			GetComputeContext().RHICopyToStagingBuffer(SourceBuffer, DestinationStagingBuffer, Offset, NumBytes, Fence);
 			return;
 		}
 		ALLOC_COMMAND(FRHICommandCopyToStagingBuffer<ECmdList::ECompute>)(SourceBuffer, DestinationStagingBuffer, Offset, NumBytes, Fence);
