#UE4 Remove pragma warning for now since it causes a compile error on some platforms

diff --git a/Engine/Source/Runtime/RHI/Public/RHICommandList.h b/Engine/Source/Runtime/RHI/Public/RHICommandList.h
--- a/Engine/Source/Runtime/RHI/Public/RHICommandList.h
+++ b/Engine/Source/Runtime/RHI/Public/RHICommandList.h
@@ -4890,8 +4890,6 @@ struct TRHIResourceUpdateBatcher
 
 	void QueueUpdateRequest(FShaderResourceViewRHIParamRef SRV, FVertexBufferRHIParamRef VertexBuffer, uint32 Stride, uint8 Format)
 	{
-#pragma warning( push )
-#pragma warning( disable : 6386 ) // Disable incorrect compile warning: C6386 buffer overrun while writing to 'this->UpdateInfos'
 		check(NumBatched <= MaxNumUpdates);
 		if (NumBatched == MaxNumUpdates)
 		{
@@ -4901,7 +4899,6 @@ struct TRHIResourceUpdateBatcher
 		UpdateInfos[Idx].Type = FRHIResourceUpdateInfo::UT_VertexBufferSRV;
 		UpdateInfos[Idx].VertexBufferSRV = { SRV, VertexBuffer, Stride, Format };
 	}
-#pragma warning( pop )
 };
 
 #undef RHICOMMAND_CALLSTACK
