Try to fix UGS error

diff --git a/Engine/Source/Runtime/RHI/Public/RHICommandList.h b/Engine/Source/Runtime/RHI/Public/RHICommandList.h
--- a/Engine/Source/Runtime/RHI/Public/RHICommandList.h
+++ b/Engine/Source/Runtime/RHI/Public/RHICommandList.h
@@ -4890,14 +4890,17 @@ struct TRHIResourceUpdateBatcher
 
 	void QueueUpdateRequest(FShaderResourceViewRHIParamRef SRV, FVertexBufferRHIParamRef VertexBuffer, uint32 Stride, uint8 Format)
 	{
-		check(NumBatched <= MaxNumUpdates);
+		check(NumBatched >= 0 && NumBatched <= MaxNumUpdates);
 		if (NumBatched == MaxNumUpdates)
 		{
 			Flush();
 		}
-		const int32 Idx = NumBatched++;
-		UpdateInfos[Idx].Type = FRHIResourceUpdateInfo::UT_VertexBufferSRV;
-		UpdateInfos[Idx].VertexBufferSRV = { SRV, VertexBuffer, Stride, Format };
+		if (LIKELY(NumBatched >= 0 && NumBatched < MaxNumUpdates))
+		{
+			const int32 Idx = NumBatched++;
+			UpdateInfos[Idx].Type = FRHIResourceUpdateInfo::UT_VertexBufferSRV;
+			UpdateInfos[Idx].VertexBufferSRV = { SRV, VertexBuffer, Stride, Format };
+		}
 	}
 };
 
