Fixed issue where CopyTexture would return after copying the first mip when using bypass mode.

diff --git a/Engine/Source/Runtime/RHI/Public/RHICommandList.h b/Engine/Source/Runtime/RHI/Public/RHICommandList.h
--- a/Engine/Source/Runtime/RHI/Public/RHICommandList.h
+++ b/Engine/Source/Runtime/RHI/Public/RHICommandList.h
@@ -2656,9 +2656,11 @@ public:
 				if (Bypass())
 				{
 					GetContext().RHICopyTexture(SourceTextureRHI, DestTextureRHI, PerMipInfo);
-					return;
 				}
-				ALLOC_COMMAND(FRHICommandCopyTexture)(SourceTextureRHI, DestTextureRHI, PerMipInfo);
+				else
+				{
+					ALLOC_COMMAND(FRHICommandCopyTexture)(SourceTextureRHI, DestTextureRHI, PerMipInfo);
+				}
 
 				++PerMipInfo.SourceMipIndex;
 				++PerMipInfo.DestMipIndex;
