F - Compile out geometry and tessellation stages on mobile platforms. Platforms can disable either or both for some memory savings.
#rnx

diff --git a/Engine/Source/Runtime/Core/Public/Android/AndroidPlatform.h b/Engine/Source/Runtime/Core/Public/Android/AndroidPlatform.h
--- a/Engine/Source/Runtime/Core/Public/Android/AndroidPlatform.h
+++ b/Engine/Source/Runtime/Core/Public/Android/AndroidPlatform.h
@@ -56,6 +56,8 @@ typedef FAndroidTypes FPlatformTypes;
 #define PLATFORM_UI_NEEDS_TOOLTIPS					0
 #define PLATFORM_UI_NEEDS_FOCUS_OUTLINES			0
 #define PLATFORM_SUPPORTS_EARLY_MOVIE_PLAYBACK		1 // movies will start before engine is initalized
+#define PLATFORM_SUPPORTS_GEOMETRY_SHADERS			0
+#define PLATFORM_SUPPORTS_TESSELLATION_SHADERS		0
 
 #define PLATFORM_CODE_SECTION(Name)					__attribute__((section(Name)))
 
diff --git a/Engine/Source/Runtime/Core/Public/HAL/Platform.h b/Engine/Source/Runtime/Core/Public/HAL/Platform.h
--- a/Engine/Source/Runtime/Core/Public/HAL/Platform.h
+++ b/Engine/Source/Runtime/Core/Public/HAL/Platform.h
@@ -337,6 +337,18 @@
 	#define PLATFORM_USES_ES2					0
 #endif
 
+#ifndef PLATFORM_SUPPORTS_GEOMETRY_SHADERS
+	#define PLATFORM_SUPPORTS_GEOMETRY_SHADERS		1
+#endif
+
+#ifndef PLATFORM_SUPPORTS_TESSELLATION_SHADERS
+	#define PLATFORM_SUPPORTS_TESSELLATION_SHADERS	1
+#endif
+
+#if PLATFORM_SUPPORTS_TESSELLATION_SHADERS && !PLATFORM_SUPPORTS_GEOMETRY_SHADERS
+	#error Geometry shader support is required by tessellation
+#endif
+
 #ifndef PLATFORM_BUILTIN_VERTEX_HALF_FLOAT
 	#define PLATFORM_BUILTIN_VERTEX_HALF_FLOAT	1
 #endif
diff --git a/Engine/Source/Runtime/D3D12RHI/Private/D3D12PipelineState.h b/Engine/Source/Runtime/D3D12RHI/Private/D3D12PipelineState.h
--- a/Engine/Source/Runtime/D3D12RHI/Private/D3D12PipelineState.h
+++ b/Engine/Source/Runtime/D3D12RHI/Private/D3D12PipelineState.h
@@ -320,9 +320,18 @@ struct FD3D12GraphicsPipelineState : public FRHIGraphicsPipelineState
 
 	FORCEINLINE class FD3D12VertexShader*   GetVertexShader() const { return (FD3D12VertexShader*)PipelineStateInitializer.BoundShaderState.VertexShaderRHI; }
 	FORCEINLINE class FD3D12PixelShader*    GetPixelShader() const { return (FD3D12PixelShader*)PipelineStateInitializer.BoundShaderState.PixelShaderRHI; }
+#if PLATFORM_SUPPORTS_TESSELLATION_SHADERS
 	FORCEINLINE class FD3D12HullShader*     GetHullShader() const { return (FD3D12HullShader*)PipelineStateInitializer.BoundShaderState.HullShaderRHI; }
 	FORCEINLINE class FD3D12DomainShader*   GetDomainShader() const { return (FD3D12DomainShader*)PipelineStateInitializer.BoundShaderState.DomainShaderRHI; }
+#else
+	FORCEINLINE class FD3D12HullShader*     GetHullShader() const { return nullptr; }
+	FORCEINLINE class FD3D12DomainShader*   GetDomainShader() const { return nullptr; }
+#endif
+#if PLATFORM_SUPPORTS_GEOMETRY_SHADERS
 	FORCEINLINE class FD3D12GeometryShader* GetGeometryShader() const { return (FD3D12GeometryShader*)PipelineStateInitializer.BoundShaderState.GeometryShaderRHI; }
+#else
+	FORCEINLINE class FD3D12GeometryShader* GetGeometryShader() const { return nullptr; }
+#endif
 };
 
 struct FD3D12ComputePipelineState : public FRHIComputePipelineState
diff --git a/Engine/Source/Runtime/D3D12RHI/Private/D3D12State.cpp b/Engine/Source/Runtime/D3D12RHI/Private/D3D12State.cpp
--- a/Engine/Source/Runtime/D3D12RHI/Private/D3D12State.cpp
+++ b/Engine/Source/Runtime/D3D12RHI/Private/D3D12State.cpp
@@ -468,10 +468,10 @@ FGraphicsPipelineStateRHIRef FD3D12DynamicRHI::RHICreateGraphicsPipelineState(co
 	FBoundShaderStateRHIRef const BSS = RHICreateBoundShaderState(
 		Initializer.BoundShaderState.VertexDeclarationRHI,
 		Initializer.BoundShaderState.VertexShaderRHI,
-		Initializer.BoundShaderState.HullShaderRHI,
-		Initializer.BoundShaderState.DomainShaderRHI,
+		TESSELLATION_SHADER(Initializer.BoundShaderState.HullShaderRHI),
+		TESSELLATION_SHADER(Initializer.BoundShaderState.DomainShaderRHI),
 		Initializer.BoundShaderState.PixelShaderRHI,
-		Initializer.BoundShaderState.GeometryShaderRHI);
+		GEOMETRY_SHADER(Initializer.BoundShaderState.GeometryShaderRHI));
 
 	// Next try to find the PSO based on the hash of its desc.
 	FD3D12BoundShaderState* const BoundShaderState = FD3D12DynamicRHI::ResourceCast(BSS.GetReference());
diff --git a/Engine/Source/Runtime/D3D12RHI/Private/D3D12StateCachePrivate.h b/Engine/Source/Runtime/D3D12RHI/Private/D3D12StateCachePrivate.h
--- a/Engine/Source/Runtime/D3D12RHI/Private/D3D12StateCachePrivate.h
+++ b/Engine/Source/Runtime/D3D12RHI/Private/D3D12StateCachePrivate.h
@@ -380,9 +380,13 @@ protected:
 	}
 	DECLARE_SHADER_TRAITS(Vertex);
 	DECLARE_SHADER_TRAITS(Pixel);
+#if PLATFORM_SUPPORTS_TESSELLATION_SHADERS
 	DECLARE_SHADER_TRAITS(Domain);
 	DECLARE_SHADER_TRAITS(Hull);
+#endif
+#if PLATFORM_SUPPORTS_GEOMETRY_SHADERS
 	DECLARE_SHADER_TRAITS(Geometry);
+#endif
 #undef DECLARE_SHADER_TRAITS
 
 	template <typename TShader> D3D12_STATE_CACHE_INLINE void SetShader(TShader* Shader)
@@ -661,17 +665,29 @@ public:
 
 	D3D12_STATE_CACHE_INLINE void GetHullShader(FD3D12HullShader** Shader)
 	{
+#if PLATFORM_SUPPORTS_TESSELLATION_SHADERS
 		GetShader(Shader);
+#else
+		*Shader = nullptr;
+#endif
 	}
 
 	D3D12_STATE_CACHE_INLINE void GetDomainShader(FD3D12DomainShader** Shader)
 	{
+#if PLATFORM_SUPPORTS_TESSELLATION_SHADERS
 		GetShader(Shader);
+#else
+		*Shader = nullptr;
+#endif
 	}
 
 	D3D12_STATE_CACHE_INLINE void GetGeometryShader(FD3D12GeometryShader** Shader)
 	{
+#if PLATFORM_SUPPORTS_GEOMETRY_SHADERS
 		GetShader(Shader);
+#else
+		*Shader = nullptr;
+#endif
 	}
 
 	D3D12_STATE_CACHE_INLINE void GetPixelShader(FD3D12PixelShader** Shader)
@@ -687,10 +703,13 @@ public:
 			SetStreamStrides(GraphicsPipelineState->StreamStrides);
 			SetShader(GraphicsPipelineState->GetVertexShader());
 			SetShader(GraphicsPipelineState->GetPixelShader());
+#if PLATFORM_SUPPORTS_TESSELLATION_SHADERS
 			SetShader(GraphicsPipelineState->GetDomainShader());
 			SetShader(GraphicsPipelineState->GetHullShader());
+#endif
+#if PLATFORM_SUPPORTS_GEOMETRY_SHADERS
 			SetShader(GraphicsPipelineState->GetGeometryShader());
-
+#endif
 			// See if we need to change the root signature
 			if (GetGraphicsRootSignature() != GraphicsPipelineState->RootSignature)
 			{
diff --git a/Engine/Source/Runtime/OpenGLDrv/Public/OpenGLDrv.h b/Engine/Source/Runtime/OpenGLDrv/Public/OpenGLDrv.h
--- a/Engine/Source/Runtime/OpenGLDrv/Public/OpenGLDrv.h
+++ b/Engine/Source/Runtime/OpenGLDrv/Public/OpenGLDrv.h
@@ -1007,10 +1007,10 @@ public:
 			RHICreateBoundShaderState_internal(
 				PsoInit.BoundShaderState.VertexDeclarationRHI,
 				PsoInit.BoundShaderState.VertexShaderRHI,
-				PsoInit.BoundShaderState.HullShaderRHI,
-				PsoInit.BoundShaderState.DomainShaderRHI,
+				TESSELLATION_SHADER(PsoInit.BoundShaderState.HullShaderRHI),
+				TESSELLATION_SHADER(PsoInit.BoundShaderState.DomainShaderRHI),
 				PsoInit.BoundShaderState.PixelShaderRHI,
-				PsoInit.BoundShaderState.GeometryShaderRHI,
+				GEOMETRY_SHADER(PsoInit.BoundShaderState.GeometryShaderRHI),
 				PsoInit.bFromPSOFileCache
 			).GetReference()
 		);
diff --git a/Engine/Source/Runtime/RHI/Private/PipelineFileCache.cpp b/Engine/Source/Runtime/RHI/Private/PipelineFileCache.cpp
--- a/Engine/Source/Runtime/RHI/Private/PipelineFileCache.cpp
+++ b/Engine/Source/Runtime/RHI/Private/PipelineFileCache.cpp
@@ -848,6 +848,7 @@ FPipelineCacheFileFormatPSO::FPipelineCacheFileFormatPSO()
 		PSO.GraphicsDesc.VertexShader = Init.BoundShaderState.VertexShaderRHI->GetHash();
 	}
 	
+#if PLATFORM_SUPPORTS_TESSELLATION_SHADERS
 	if (Init.BoundShaderState.HullShaderRHI)
 	{
 		PSO.GraphicsDesc.HullShader = Init.BoundShaderState.HullShaderRHI->GetHash();
@@ -857,16 +858,18 @@ FPipelineCacheFileFormatPSO::FPipelineCacheFileFormatPSO()
 	{
 		PSO.GraphicsDesc.DomainShader = Init.BoundShaderState.DomainShaderRHI->GetHash();
 	}
-	
+#endif
 	if (Init.BoundShaderState.PixelShaderRHI)
 	{
 		PSO.GraphicsDesc.FragmentShader = Init.BoundShaderState.PixelShaderRHI->GetHash();
 	}
 	
+#if PLATFORM_SUPPORTS_GEOMETRY_SHADERS
 	if (Init.BoundShaderState.GeometryShaderRHI)
 	{
 		PSO.GraphicsDesc.GeometryShader = Init.BoundShaderState.GeometryShaderRHI->GetHash();
 	}
+#endif
 	
 	check (Init.BlendState);
 	{
diff --git a/Engine/Source/Runtime/RHI/Private/PipelineStateCache.cpp b/Engine/Source/Runtime/RHI/Private/PipelineStateCache.cpp
--- a/Engine/Source/Runtime/RHI/Private/PipelineStateCache.cpp
+++ b/Engine/Source/Runtime/RHI/Private/PipelineStateCache.cpp
@@ -26,12 +26,17 @@ PipelineStateCache.cpp: Pipeline state cache implementation.
 
 static inline uint32 GetTypeHash(const FBoundShaderStateInput& Input)
 {
-	return GetTypeHash(Input.VertexDeclarationRHI) ^
-		GetTypeHash(Input.VertexShaderRHI) ^
-		GetTypeHash(Input.PixelShaderRHI) ^
-		GetTypeHash(Input.HullShaderRHI) ^
-		GetTypeHash(Input.DomainShaderRHI) ^
-		GetTypeHash(Input.GeometryShaderRHI);
+	return GetTypeHash(Input.VertexDeclarationRHI)
+		^ GetTypeHash(Input.VertexShaderRHI)
+		^ GetTypeHash(Input.PixelShaderRHI)
+#if PLATFORM_SUPPORTS_TESSELLATION_SHADERS
+		^ GetTypeHash(Input.HullShaderRHI)
+		^ GetTypeHash(Input.DomainShaderRHI)
+#endif
+#if PLATFORM_SUPPORTS_GEOMETRY_SHADERS
+		^ GetTypeHash(Input.GeometryShaderRHI)
+#endif
+		;
 }
 
 static inline uint32 GetTypeHash(const FGraphicsPipelineStateInitializer& Initializer)
@@ -94,6 +99,7 @@ static void HandlePipelineCreationFailure(const FGraphicsPipelineStateInitialize
 	{
 		UE_LOG(LogRHI, Error, TEXT("Vertex: %s"), *Init.BoundShaderState.VertexShaderRHI->ShaderName);
 	}
+#if PLATFORM_SUPPORTS_TESSELLATION_SHADERS
 	if(Init.BoundShaderState.HullShaderRHI)
 	{
 		UE_LOG(LogRHI, Error, TEXT("Hull: %s"), *Init.BoundShaderState.HullShaderRHI->ShaderName);
@@ -102,10 +108,13 @@ static void HandlePipelineCreationFailure(const FGraphicsPipelineStateInitialize
 	{
 		UE_LOG(LogRHI, Error, TEXT("Domain: %s"), *Init.BoundShaderState.DomainShaderRHI->ShaderName);
 	}
+#endif
+#if PLATFORM_SUPPORTS_GEOMETRY_SHADERS
 	if(Init.BoundShaderState.GeometryShaderRHI)
 	{
 		UE_LOG(LogRHI, Error, TEXT("Geometry: %s"), *Init.BoundShaderState.GeometryShaderRHI->ShaderName);
 	}
+#endif
 	if(Init.BoundShaderState.PixelShaderRHI)
 	{
 		UE_LOG(LogRHI, Error, TEXT("Pixel: %s"), *Init.BoundShaderState.PixelShaderRHI->ShaderName);
@@ -652,12 +661,16 @@ public:
 			Initializer.BoundShaderState.VertexShaderRHI->AddRef();
 		if (Initializer.BoundShaderState.PixelShaderRHI)
 			Initializer.BoundShaderState.PixelShaderRHI->AddRef();
+#if PLATFORM_SUPPORTS_GEOMETRY_SHADERS
 		if (Initializer.BoundShaderState.GeometryShaderRHI)
 			Initializer.BoundShaderState.GeometryShaderRHI->AddRef();
+#endif
+#if PLATFORM_SUPPORTS_TESSELLATION_SHADERS
 		if (Initializer.BoundShaderState.DomainShaderRHI)
 			Initializer.BoundShaderState.DomainShaderRHI->AddRef();
 		if (Initializer.BoundShaderState.HullShaderRHI)
 			Initializer.BoundShaderState.HullShaderRHI->AddRef();
+#endif
 		if (Initializer.BlendState)
 			Initializer.BlendState->AddRef();
 		if (Initializer.RasterizerState)
@@ -693,12 +706,16 @@ public:
 				Initializer.BoundShaderState.VertexShaderRHI->Release();
 			if (Initializer.BoundShaderState.PixelShaderRHI)
 				Initializer.BoundShaderState.PixelShaderRHI->Release();
+#if PLATFORM_SUPPORTS_GEOMETRY_SHADERS
 			if (Initializer.BoundShaderState.GeometryShaderRHI)
 				Initializer.BoundShaderState.GeometryShaderRHI->Release();
+#endif
+#if PLATFORM_SUPPORTS_TESSELLATION_SHADERS
 			if (Initializer.BoundShaderState.DomainShaderRHI)
 				Initializer.BoundShaderState.DomainShaderRHI->Release();
 			if (Initializer.BoundShaderState.HullShaderRHI)
 				Initializer.BoundShaderState.HullShaderRHI->Release();
+#endif
 			if (Initializer.BlendState)
 				Initializer.BlendState->Release();
 			if (Initializer.RasterizerState)
@@ -930,10 +947,14 @@ FGraphicsPipelineState* PipelineStateCache::GetAndOrCreateGraphicsPipelineState(
 	{
 		FGraphicsPipelineStateInitializer& HashableInitializer = const_cast<FGraphicsPipelineStateInitializer&>(OriginalInitializer);
 		HashableInitializer.VertexShaderHash = HashableInitializer.BoundShaderState.VertexShaderRHI ? HashableInitializer.BoundShaderState.VertexShaderRHI->GetHash() : FSHAHash();
+#if PLATFORM_SUPPORTS_TESSELLATION_SHADERS
 		HashableInitializer.HullShaderHash = HashableInitializer.BoundShaderState.HullShaderRHI ? HashableInitializer.BoundShaderState.HullShaderRHI->GetHash() : FSHAHash();
 		HashableInitializer.DomainShaderHash = HashableInitializer.BoundShaderState.DomainShaderRHI ? HashableInitializer.BoundShaderState.DomainShaderRHI->GetHash() : FSHAHash();
+#endif
 		HashableInitializer.PixelShaderHash = HashableInitializer.BoundShaderState.PixelShaderRHI ? HashableInitializer.BoundShaderState.PixelShaderRHI->GetHash() : FSHAHash();
+#if PLATFORM_SUPPORTS_GEOMETRY_SHADERS
 		HashableInitializer.GeometryShaderHash = HashableInitializer.BoundShaderState.GeometryShaderRHI ? HashableInitializer.BoundShaderState.GeometryShaderRHI->GetHash() : FSHAHash();
+#endif
 	}
 	
 	FGraphicsPipelineStateInitializer NewInitializer;
diff --git a/Engine/Source/Runtime/RHI/Public/RHIDefinitions.h b/Engine/Source/Runtime/RHI/Public/RHIDefinitions.h
--- a/Engine/Source/Runtime/RHI/Public/RHIDefinitions.h
+++ b/Engine/Source/Runtime/RHI/Public/RHIDefinitions.h
@@ -1212,3 +1212,15 @@ inline const TCHAR* GetShaderFrequencyString(EShaderFrequency Frequency, bool bI
 	String += Index;
 	return String;
 };
+
+#if PLATFORM_SUPPORTS_GEOMETRY_SHADERS
+	#define GEOMETRY_SHADER(GeometryShader)	(GeometryShader)
+#else
+	#define GEOMETRY_SHADER(GeometryShader)	nullptr
+#endif
+
+#if PLATFORM_SUPPORTS_TESSELLATION_SHADERS
+	#define TESSELLATION_SHADER(HullOrDomainShader)	(HullOrDomainShader)
+#else
+	#define TESSELLATION_SHADER(HullOrDomainShader)	nullptr
+#endif
diff --git a/Engine/Source/Runtime/RHI/Public/RHIResources.h b/Engine/Source/Runtime/RHI/Public/RHIResources.h
--- a/Engine/Source/Runtime/RHI/Public/RHIResources.h
+++ b/Engine/Source/Runtime/RHI/Public/RHIResources.h
@@ -1646,17 +1646,25 @@ struct FBoundShaderStateInput
 	(
 		FVertexDeclarationRHIParamRef InVertexDeclarationRHI
 		, FVertexShaderRHIParamRef InVertexShaderRHI
+#if PLATFORM_SUPPORTS_TESSELLATION_SHADERS
 		, FHullShaderRHIParamRef InHullShaderRHI
 		, FDomainShaderRHIParamRef InDomainShaderRHI
+#endif
 		, FPixelShaderRHIParamRef InPixelShaderRHI
+#if PLATFORM_SUPPORTS_GEOMETRY_SHADERS
 		, FGeometryShaderRHIParamRef InGeometryShaderRHI
+#endif
 	)
 		: VertexDeclarationRHI(InVertexDeclarationRHI)
 		, VertexShaderRHI(InVertexShaderRHI)
+#if PLATFORM_SUPPORTS_TESSELLATION_SHADERS
 		, HullShaderRHI(InHullShaderRHI)
 		, DomainShaderRHI(InDomainShaderRHI)
+#endif
 		, PixelShaderRHI(InPixelShaderRHI)
+#if PLATFORM_SUPPORTS_GEOMETRY_SHADERS
 		, GeometryShaderRHI(InGeometryShaderRHI)
+#endif
 	{
 	}
 
@@ -1747,9 +1755,13 @@ public:
 		if (BoundShaderState.VertexDeclarationRHI != rhs.BoundShaderState.VertexDeclarationRHI || 
 			BoundShaderState.VertexShaderRHI != rhs.BoundShaderState.VertexShaderRHI ||
 			BoundShaderState.PixelShaderRHI != rhs.BoundShaderState.PixelShaderRHI ||
+#if PLATFORM_SUPPORTS_GEOMETRY_SHADERS
 			BoundShaderState.GeometryShaderRHI != rhs.BoundShaderState.GeometryShaderRHI ||
+#endif
+#if PLATFORM_SUPPORTS_TESSELLATION_SHADERS
 			BoundShaderState.DomainShaderRHI != rhs.BoundShaderState.DomainShaderRHI ||
-			BoundShaderState.HullShaderRHI != rhs.BoundShaderState.HullShaderRHI ||			
+			BoundShaderState.HullShaderRHI != rhs.BoundShaderState.HullShaderRHI ||	
+#endif		
 			BlendState != rhs.BlendState || 
 			RasterizerState != rhs.RasterizerState || 
 			DepthStencilState != rhs.DepthStencilState ||
@@ -1794,9 +1806,13 @@ public:
 		COMPARE_FIELD_BEGIN(BoundShaderState.VertexDeclarationRHI)
 		COMPARE_FIELD(BoundShaderState.VertexShaderRHI)
 		COMPARE_FIELD(BoundShaderState.PixelShaderRHI)
+#if PLATFORM_SUPPORTS_GEOMETRY_SHADERS
 		COMPARE_FIELD(BoundShaderState.GeometryShaderRHI)
+#endif
+#if PLATFORM_SUPPORTS_TESSELLATION_SHADERS
 		COMPARE_FIELD(BoundShaderState.DomainShaderRHI)
 		COMPARE_FIELD(BoundShaderState.HullShaderRHI)
+#endif
 		COMPARE_FIELD(BlendState)
 		COMPARE_FIELD(RasterizerState)
 		COMPARE_FIELD(DepthStencilState)
@@ -1814,9 +1830,13 @@ public:
 		COMPARE_FIELD_BEGIN(BoundShaderState.VertexDeclarationRHI)
 		COMPARE_FIELD(BoundShaderState.VertexShaderRHI)
 		COMPARE_FIELD(BoundShaderState.PixelShaderRHI)
+#if PLATFORM_SUPPORTS_GEOMETRY_SHADERS
 		COMPARE_FIELD(BoundShaderState.GeometryShaderRHI)
+#endif
+#if PLATFORM_SUPPORTS_TESSELLATION_SHADERS
 		COMPARE_FIELD(BoundShaderState.DomainShaderRHI)
 		COMPARE_FIELD(BoundShaderState.HullShaderRHI)
+#endif
 		COMPARE_FIELD(BlendState)
 		COMPARE_FIELD(RasterizerState)
 		COMPARE_FIELD(DepthStencilState)
@@ -1969,9 +1989,13 @@ public:
 
 	FSHAHash						VertexShaderHash;
 	FSHAHash						PixelShaderHash;
+#if PLATFORM_SUPPORTS_GEOMETRY_SHADERS
 	FSHAHash						GeometryShaderHash;
+#endif
+#if PLATFORM_SUPPORTS_TESSELLATION_SHADERS
 	FSHAHash						HullShaderHash;
 	FSHAHash						DomainShaderHash;
+#endif
 	uint32							RenderTargetsEnabled;
 	TRenderTargetFormats			RenderTargetFormats;
 	TRenderTargetFlags				RenderTargetFlags;
diff --git a/Engine/Source/Runtime/RenderCore/Private/ShaderPipelineCache.cpp b/Engine/Source/Runtime/RenderCore/Private/ShaderPipelineCache.cpp
--- a/Engine/Source/Runtime/RenderCore/Private/ShaderPipelineCache.cpp
+++ b/Engine/Source/Runtime/RenderCore/Private/ShaderPipelineCache.cpp
@@ -550,6 +550,7 @@ bool FShaderPipelineCache::Precompile(FRHICommandListImmediate& RHICmdList, ESha
 			GraphicsInitializer.BoundShaderState.VertexShaderRHI = VertexShader;
 		}
 
+#if PLATFORM_SUPPORTS_TESSELLATION_SHADERS
 		FHullShaderRHIRef HullShader;
 		if (PSO.GraphicsDesc.HullShader != FSHAHash())
 		{
@@ -563,7 +564,7 @@ bool FShaderPipelineCache::Precompile(FRHICommandListImmediate& RHICmdList, ESha
 			DomainShader = FShaderCodeLibrary::CreateDomainShader(Platform, PSO.GraphicsDesc.DomainShader, DummyCode);
 			GraphicsInitializer.BoundShaderState.DomainShaderRHI = DomainShader;
 		}
-
+#endif
 		FPixelShaderRHIRef FragmentShader;
 		if (PSO.GraphicsDesc.FragmentShader != FSHAHash())
 		{
@@ -571,13 +572,14 @@ bool FShaderPipelineCache::Precompile(FRHICommandListImmediate& RHICmdList, ESha
 			GraphicsInitializer.BoundShaderState.PixelShaderRHI = FragmentShader;
 		}
 
+#if PLATFORM_SUPPORTS_GEOMETRY_SHADERS
 		FGeometryShaderRHIRef GeometryShader;
 		if (PSO.GraphicsDesc.GeometryShader != FSHAHash())
 		{
 			GeometryShader = FShaderCodeLibrary::CreateGeometryShader(Platform, PSO.GraphicsDesc.GeometryShader, DummyCode);
 			GraphicsInitializer.BoundShaderState.GeometryShaderRHI = GeometryShader;
 		}
-		
+#endif
 		auto BlendState = RHICmdList.CreateBlendState(PSO.GraphicsDesc.BlendState);
 		GraphicsInitializer.BlendState = BlendState;
 		
diff --git a/Engine/Source/Runtime/Renderer/Private/AtmosphereRendering.cpp b/Engine/Source/Runtime/Renderer/Private/AtmosphereRendering.cpp
--- a/Engine/Source/Runtime/Renderer/Private/AtmosphereRendering.cpp
+++ b/Engine/Source/Runtime/Renderer/Private/AtmosphereRendering.cpp
@@ -1182,7 +1182,9 @@ void FAtmosphericFogSceneInfo::RenderAtmosphereShaders(FRHICommandList& RHICmdLi
 
 					GraphicsPSOInit.BoundShaderState.VertexDeclarationRHI = GFilterVertexDeclaration.VertexDeclarationRHI;
 					GraphicsPSOInit.BoundShaderState.VertexShaderRHI = GETSAFERHISHADER_VERTEX(*VertexShader);
+#if PLATFORM_SUPPORTS_GEOMETRY_SHADERS
 					GraphicsPSOInit.BoundShaderState.GeometryShaderRHI = GETSAFERHISHADER_GEOMETRY(*GeometryShader);
+#endif
 					GraphicsPSOInit.BoundShaderState.PixelShaderRHI = GETSAFERHISHADER_PIXEL(*PixelShader);
 					GraphicsPSOInit.PrimitiveType = PT_TriangleList;
 
@@ -1239,7 +1241,9 @@ void FAtmosphericFogSceneInfo::RenderAtmosphereShaders(FRHICommandList& RHICmdLi
 
 					GraphicsPSOInit.BoundShaderState.VertexDeclarationRHI = GFilterVertexDeclaration.VertexDeclarationRHI;
 					GraphicsPSOInit.BoundShaderState.VertexShaderRHI = GETSAFERHISHADER_VERTEX(*VertexShader);
+#if PLATFORM_SUPPORTS_GEOMETRY_SHADERS
 					GraphicsPSOInit.BoundShaderState.GeometryShaderRHI = GETSAFERHISHADER_GEOMETRY(*GeometryShader);
+#endif
 					GraphicsPSOInit.BoundShaderState.PixelShaderRHI = GETSAFERHISHADER_PIXEL(*PixelShader);
 					GraphicsPSOInit.PrimitiveType = PT_TriangleList;
 
@@ -1281,7 +1285,9 @@ void FAtmosphericFogSceneInfo::RenderAtmosphereShaders(FRHICommandList& RHICmdLi
 
 					GraphicsPSOInit.BoundShaderState.VertexDeclarationRHI = GFilterVertexDeclaration.VertexDeclarationRHI;
 					GraphicsPSOInit.BoundShaderState.VertexShaderRHI = GETSAFERHISHADER_VERTEX(*VertexShader);
+#if PLATFORM_SUPPORTS_GEOMETRY_SHADERS
 					GraphicsPSOInit.BoundShaderState.GeometryShaderRHI = GETSAFERHISHADER_GEOMETRY(*GeometryShader);
+#endif
 					GraphicsPSOInit.BoundShaderState.PixelShaderRHI = GETSAFERHISHADER_PIXEL(*PixelShader);
 					GraphicsPSOInit.PrimitiveType = PT_TriangleList;
 
@@ -1350,7 +1356,9 @@ void FAtmosphericFogSceneInfo::RenderAtmosphereShaders(FRHICommandList& RHICmdLi
 
 					GraphicsPSOInit.BoundShaderState.VertexDeclarationRHI = GFilterVertexDeclaration.VertexDeclarationRHI;
 					GraphicsPSOInit.BoundShaderState.VertexShaderRHI = GETSAFERHISHADER_VERTEX(*VertexShader);
+#if PLATFORM_SUPPORTS_GEOMETRY_SHADERS
 					GraphicsPSOInit.BoundShaderState.GeometryShaderRHI = GETSAFERHISHADER_GEOMETRY(*GeometryShader);
+#endif
 					GraphicsPSOInit.BoundShaderState.PixelShaderRHI = GETSAFERHISHADER_PIXEL(*PixelShader);
 					GraphicsPSOInit.PrimitiveType = PT_TriangleList;
 
@@ -1426,7 +1434,9 @@ void FAtmosphericFogSceneInfo::RenderAtmosphereShaders(FRHICommandList& RHICmdLi
 
 					GraphicsPSOInit.BoundShaderState.VertexDeclarationRHI = GFilterVertexDeclaration.VertexDeclarationRHI;
 					GraphicsPSOInit.BoundShaderState.VertexShaderRHI = GETSAFERHISHADER_VERTEX(*VertexShader);
+#if PLATFORM_SUPPORTS_GEOMETRY_SHADERS
 					GraphicsPSOInit.BoundShaderState.GeometryShaderRHI = GETSAFERHISHADER_GEOMETRY(*GeometryShader);
+#endif
 					GraphicsPSOInit.BoundShaderState.PixelShaderRHI = GETSAFERHISHADER_PIXEL(*PixelShader);
 					GraphicsPSOInit.PrimitiveType = PT_TriangleList;
 
@@ -1471,7 +1481,9 @@ void FAtmosphericFogSceneInfo::RenderAtmosphereShaders(FRHICommandList& RHICmdLi
 
 					GraphicsPSOInit.BoundShaderState.VertexDeclarationRHI = GFilterVertexDeclaration.VertexDeclarationRHI;
 					GraphicsPSOInit.BoundShaderState.VertexShaderRHI = GETSAFERHISHADER_VERTEX(*VertexShader);
+#if PLATFORM_SUPPORTS_GEOMETRY_SHADERS
 					GraphicsPSOInit.BoundShaderState.GeometryShaderRHI = GETSAFERHISHADER_GEOMETRY(*GeometryShader);
+#endif
 					GraphicsPSOInit.BoundShaderState.PixelShaderRHI = GETSAFERHISHADER_PIXEL(*PixelShader);
 					GraphicsPSOInit.PrimitiveType = PT_TriangleList;
 
@@ -1513,7 +1525,9 @@ void FAtmosphericFogSceneInfo::RenderAtmosphereShaders(FRHICommandList& RHICmdLi
 
 					GraphicsPSOInit.BoundShaderState.VertexDeclarationRHI = GFilterVertexDeclaration.VertexDeclarationRHI;
 					GraphicsPSOInit.BoundShaderState.VertexShaderRHI = GETSAFERHISHADER_VERTEX(*VertexShader);
+#if PLATFORM_SUPPORTS_GEOMETRY_SHADERS
 					GraphicsPSOInit.BoundShaderState.GeometryShaderRHI = GETSAFERHISHADER_GEOMETRY(*GeometryShader);
+#endif
 					GraphicsPSOInit.BoundShaderState.PixelShaderRHI = GETSAFERHISHADER_PIXEL(*PixelShader);
 					GraphicsPSOInit.PrimitiveType = PT_TriangleList;
 
diff --git a/Engine/Source/Runtime/Renderer/Private/LightPropagationVolumeVisualisation.cpp b/Engine/Source/Runtime/Renderer/Private/LightPropagationVolumeVisualisation.cpp
--- a/Engine/Source/Runtime/Renderer/Private/LightPropagationVolumeVisualisation.cpp
+++ b/Engine/Source/Runtime/Renderer/Private/LightPropagationVolumeVisualisation.cpp
@@ -221,7 +221,9 @@ void FLightPropagationVolume::Visualise(FRHICommandList& RHICmdList, const FView
 	
 	GraphicsPSOInit.BoundShaderState.VertexDeclarationRHI = GSimpleElementVertexDeclaration.VertexDeclarationRHI;
 	GraphicsPSOInit.BoundShaderState.VertexShaderRHI = GETSAFERHISHADER_VERTEX(*VertexShader);
+#if PLATFORM_SUPPORTS_GEOMETRY_SHADERS
 	GraphicsPSOInit.BoundShaderState.GeometryShaderRHI = GETSAFERHISHADER_GEOMETRY(*GeometryShader);
+#endif
 	GraphicsPSOInit.BoundShaderState.PixelShaderRHI = GETSAFERHISHADER_PIXEL(*PixelShader);
 	GraphicsPSOInit.PrimitiveType = PrimType;
 
diff --git a/Engine/Source/Runtime/Renderer/Private/PostProcess/PostProcessCombineLUTs.cpp b/Engine/Source/Runtime/Renderer/Private/PostProcess/PostProcessCombineLUTs.cpp
--- a/Engine/Source/Runtime/Renderer/Private/PostProcess/PostProcessCombineLUTs.cpp
+++ b/Engine/Source/Runtime/Renderer/Private/PostProcess/PostProcessCombineLUTs.cpp
@@ -725,7 +725,9 @@ static void SetLUTBlenderShader(FRenderingCompositePassContext& Context, TRHICom
 		GraphicsPSOInit.PrimitiveType = PT_TriangleStrip;
 		GraphicsPSOInit.BoundShaderState.VertexDeclarationRHI = GScreenVertexDeclaration.VertexDeclarationRHI;
 		GraphicsPSOInit.BoundShaderState.VertexShaderRHI = GETSAFERHISHADER_VERTEX(*VertexShader);
+#if PLATFORM_SUPPORTS_GEOMETRY_SHADERS
 		GraphicsPSOInit.BoundShaderState.GeometryShaderRHI = GETSAFERHISHADER_GEOMETRY(*GeometryShader);
+#endif
 		GraphicsPSOInit.BoundShaderState.PixelShaderRHI = GETSAFERHISHADER_PIXEL(LocalPixelShader);
 		SetGraphicsPipelineState(RHICmdList, GraphicsPSOInit);
 
diff --git a/Engine/Source/Runtime/Renderer/Private/PostProcess/SceneRenderTargets.cpp b/Engine/Source/Runtime/Renderer/Private/PostProcess/SceneRenderTargets.cpp
--- a/Engine/Source/Runtime/Renderer/Private/PostProcess/SceneRenderTargets.cpp
+++ b/Engine/Source/Runtime/Renderer/Private/PostProcess/SceneRenderTargets.cpp
@@ -2429,7 +2429,9 @@ void FSceneRenderTargets::ClearVolumeTextures(FRHICommandList& RHICmdList, ERHIF
 
 	GraphicsPSOInit.BoundShaderState.VertexDeclarationRHI = GScreenVertexDeclaration.VertexDeclarationRHI;
 	GraphicsPSOInit.BoundShaderState.VertexShaderRHI = GETSAFERHISHADER_VERTEX(*VertexShader);
+#if PLATFORM_SUPPORTS_GEOMETRY_SHADERS
 	GraphicsPSOInit.BoundShaderState.GeometryShaderRHI = GETSAFERHISHADER_GEOMETRY(*GeometryShader);
+#endif
 	GraphicsPSOInit.BoundShaderState.PixelShaderRHI = GETSAFERHISHADER_PIXEL(*PixelShader);
 	GraphicsPSOInit.PrimitiveType = PT_TriangleStrip;
 
diff --git a/Engine/Source/Runtime/Renderer/Private/ShadowDepthRendering.cpp b/Engine/Source/Runtime/Renderer/Private/ShadowDepthRendering.cpp
--- a/Engine/Source/Runtime/Renderer/Private/ShadowDepthRendering.cpp
+++ b/Engine/Source/Runtime/Renderer/Private/ShadowDepthRendering.cpp
@@ -978,7 +978,9 @@ void FProjectedShadowInfo::CopyCachedShadowMap(FRHICommandList& RHICmdList, cons
 
 				GraphicsPSOInit.BoundShaderState.VertexDeclarationRHI = GFilterVertexDeclaration.VertexDeclarationRHI;
 				GraphicsPSOInit.BoundShaderState.VertexShaderRHI = GETSAFERHISHADER_VERTEX(*ScreenVertexShader);
+#if PLATFORM_SUPPORTS_GEOMETRY_SHADERS
 				GraphicsPSOInit.BoundShaderState.GeometryShaderRHI = GETSAFERHISHADER_GEOMETRY(*GeometryShader);
+#endif
 				GraphicsPSOInit.BoundShaderState.PixelShaderRHI = GETSAFERHISHADER_PIXEL(*PixelShader);
 				GraphicsPSOInit.PrimitiveType = PT_TriangleList;
 
diff --git a/Engine/Source/Runtime/Renderer/Private/TranslucentLighting.cpp b/Engine/Source/Runtime/Renderer/Private/TranslucentLighting.cpp
--- a/Engine/Source/Runtime/Renderer/Private/TranslucentLighting.cpp
+++ b/Engine/Source/Runtime/Renderer/Private/TranslucentLighting.cpp
@@ -1134,7 +1134,9 @@ void FDeferredShadingSceneRenderer::InjectAmbientCubemapTranslucentVolumeLightin
 
 				GraphicsPSOInit.BoundShaderState.VertexDeclarationRHI = GScreenVertexDeclaration.VertexDeclarationRHI;
 				GraphicsPSOInit.BoundShaderState.VertexShaderRHI = GETSAFERHISHADER_VERTEX(*VertexShader);
+#if PLATFORM_SUPPORTS_GEOMETRY_SHADERS
 				GraphicsPSOInit.BoundShaderState.GeometryShaderRHI = GETSAFERHISHADER_GEOMETRY(*GeometryShader);
+#endif
 				GraphicsPSOInit.BoundShaderState.PixelShaderRHI = GETSAFERHISHADER_PIXEL(*PixelShader);
 				GraphicsPSOInit.PrimitiveType = PT_TriangleStrip;
 									
@@ -1266,7 +1268,9 @@ void FDeferredShadingSceneRenderer::AccumulateTranslucentVolumeObjectShadowing(F
 
 					GraphicsPSOInit.BoundShaderState.VertexDeclarationRHI = GScreenVertexDeclaration.VertexDeclarationRHI;
 					GraphicsPSOInit.BoundShaderState.VertexShaderRHI = GETSAFERHISHADER_VERTEX(*VertexShader);
+#if PLATFORM_SUPPORTS_GEOMETRY_SHADERS
 					GraphicsPSOInit.BoundShaderState.GeometryShaderRHI = GETSAFERHISHADER_GEOMETRY(*GeometryShader);
+#endif
 					GraphicsPSOInit.BoundShaderState.PixelShaderRHI = GETSAFERHISHADER_PIXEL(*PixelShader);
 					GraphicsPSOInit.PrimitiveType = PT_TriangleStrip;
 
@@ -1356,7 +1360,9 @@ void SetInjectionShader(
 	
 	GraphicsPSOInit.BoundShaderState.VertexDeclarationRHI = GScreenVertexDeclaration.VertexDeclarationRHI;
 	GraphicsPSOInit.BoundShaderState.VertexShaderRHI = GETSAFERHISHADER_VERTEX(VertexShader);
+#if PLATFORM_SUPPORTS_GEOMETRY_SHADERS
 	GraphicsPSOInit.BoundShaderState.GeometryShaderRHI = GETSAFERHISHADER_GEOMETRY(GeometryShader);
+#endif
 	GraphicsPSOInit.BoundShaderState.PixelShaderRHI = GETSAFERHISHADER_PIXEL(PixelShader);
 	SetGraphicsPipelineState(RHICmdList, GraphicsPSOInit);
 
@@ -1718,7 +1724,9 @@ void FDeferredShadingSceneRenderer::InjectSimpleTranslucentVolumeLightingArray(F
 
 							GraphicsPSOInit.BoundShaderState.VertexDeclarationRHI = GScreenVertexDeclaration.VertexDeclarationRHI;
 							GraphicsPSOInit.BoundShaderState.VertexShaderRHI = GETSAFERHISHADER_VERTEX(*VertexShader);
+#if PLATFORM_SUPPORTS_GEOMETRY_SHADERS
 							GraphicsPSOInit.BoundShaderState.GeometryShaderRHI = GETSAFERHISHADER_GEOMETRY(*GeometryShader);
+#endif
 							GraphicsPSOInit.BoundShaderState.PixelShaderRHI = GETSAFERHISHADER_PIXEL(*PixelShader);
 
 							SetGraphicsPipelineState(RHICmdList, GraphicsPSOInit);
@@ -1821,7 +1829,9 @@ void FDeferredShadingSceneRenderer::FilterTranslucentVolumeLighting(FRHICommandL
 				
 					GraphicsPSOInit.BoundShaderState.VertexDeclarationRHI = GScreenVertexDeclaration.VertexDeclarationRHI;
 					GraphicsPSOInit.BoundShaderState.VertexShaderRHI = GETSAFERHISHADER_VERTEX(*VertexShader);
+#if PLATFORM_SUPPORTS_GEOMETRY_SHADERS
 					GraphicsPSOInit.BoundShaderState.GeometryShaderRHI = GETSAFERHISHADER_GEOMETRY(*GeometryShader);
+#endif
 					GraphicsPSOInit.BoundShaderState.PixelShaderRHI = GETSAFERHISHADER_PIXEL(*PixelShader);
 					GraphicsPSOInit.PrimitiveType = PT_TriangleStrip;
 
diff --git a/Engine/Source/Runtime/SlateRHIRenderer/Private/SlateRHIRenderer.cpp b/Engine/Source/Runtime/SlateRHIRenderer/Private/SlateRHIRenderer.cpp
--- a/Engine/Source/Runtime/SlateRHIRenderer/Private/SlateRHIRenderer.cpp
+++ b/Engine/Source/Runtime/SlateRHIRenderer/Private/SlateRHIRenderer.cpp
@@ -865,16 +865,20 @@ void FSlateRHIRenderer::DrawWindow_RenderThread(FRHICommandListImmediate& RHICmd
 
 						GraphicsPSOInit.BoundShaderState.VertexDeclarationRHI = GScreenVertexDeclaration.VertexDeclarationRHI;
 						GraphicsPSOInit.BoundShaderState.VertexShaderRHI = GETSAFERHISHADER_VERTEX(*VertexShader);
+#if PLATFORM_SUPPORTS_GEOMETRY_SHADERS
 						GraphicsPSOInit.BoundShaderState.GeometryShaderRHI = GETSAFERHISHADER_GEOMETRY(*GeometryShader);
+#endif
 						GraphicsPSOInit.BoundShaderState.PixelShaderRHI = GETSAFERHISHADER_PIXEL(*PixelShader);
 						GraphicsPSOInit.PrimitiveType = PT_TriangleStrip;
 						SetGraphicsPipelineState(RHICmdList, GraphicsPSOInit);
 
 						VertexShader->SetParameters(RHICmdList, VolumeBounds, FIntVector(VolumeBounds.MaxX - VolumeBounds.MinX));
+#if PLATFORM_SUPPORTS_GEOMETRY_SHADERS
 						if (GeometryShader.IsValid())
 						{
 							GeometryShader->SetParameters(RHICmdList, VolumeBounds.MinZ);
 						}
+#endif
 						PixelShader->SetParameters(RHICmdList);
 
 						RasterizeToVolumeTexture(RHICmdList, VolumeBounds);
diff --git a/Engine/Source/Runtime/VulkanRHI/Private/VulkanPipeline.cpp b/Engine/Source/Runtime/VulkanRHI/Private/VulkanPipeline.cpp
--- a/Engine/Source/Runtime/VulkanRHI/Private/VulkanPipeline.cpp
+++ b/Engine/Source/Runtime/VulkanRHI/Private/VulkanPipeline.cpp
@@ -1929,6 +1929,7 @@ void GetVulkanShaders(const FBoundShaderStateInput& BSI, FVulkanShader* OutShade
 		OutShaders[ShaderStage::Pixel] = ResourceCast(TShaderMapRef<FNULLPS>(GetGlobalShaderMap(GMaxRHIFeatureLevel))->GetPixelShader());
 	}
 
+#if PLATFORM_SUPPORTS_GEOMETRY_SHADERS
 	if (BSI.GeometryShaderRHI)
 	{
 #if VULKAN_SUPPORTS_GEOMETRY_SHADERS
@@ -1937,7 +1938,9 @@ void GetVulkanShaders(const FBoundShaderStateInput& BSI, FVulkanShader* OutShade
 		ensureMsgf(0, TEXT("Geometry not supported!"));
 #endif
 	}
+#endif
 
+#if PLATFORM_SUPPORTS_TESSELLATION_SHADERS
 	if (BSI.HullShaderRHI)
 	{
 		ensureMsgf(0, TEXT("Tessellation not supported yet!"));
@@ -1953,6 +1956,7 @@ void GetVulkanShaders(const FBoundShaderStateInput& BSI, FVulkanShader* OutShade
 		// Can't have Domain w/o Hull
 		check(BSI.DomainShaderRHI == nullptr);
 	}
+#endif
 }
 
 void GetVulkanShaders(FVulkanDevice* Device, const FVulkanRHIGraphicsPipelineState& GfxPipelineState, FVulkanShader* OutShaders[ShaderStage::NumStages])
diff --git a/Engine/Source/Runtime/VulkanRHI/Public/VulkanCommon.h b/Engine/Source/Runtime/VulkanRHI/Public/VulkanCommon.h
--- a/Engine/Source/Runtime/VulkanRHI/Public/VulkanCommon.h
+++ b/Engine/Source/Runtime/VulkanRHI/Public/VulkanCommon.h
@@ -9,7 +9,7 @@
 #include "RHIDefinitions.h"
 
 #ifndef VULKAN_SUPPORTS_GEOMETRY_SHADERS
-	#define VULKAN_SUPPORTS_GEOMETRY_SHADERS					(!(PLATFORM_ANDROID) || PLATFORM_LUMIN || PLATFORM_LUMINGL4)
+	#define VULKAN_SUPPORTS_GEOMETRY_SHADERS					(!(PLATFORM_ANDROID) || PLATFORM_LUMIN || PLATFORM_LUMINGL4) && PLATFORM_SUPPORTS_GEOMETRY_SHADERS
 #endif
 
 // This defines controls shader generation (so will cause a format rebuild)
@@ -25,11 +25,7 @@ namespace ShaderStage
 		Vertex			= 0,
 		Pixel			= 1,
 
-#if PLATFORM_ANDROID && !PLATFORM_LUMIN && !PLATFORM_LUMINGL4
-		NumStages		= 2,
-
-		MaxNumSets		= 4,
-#else
+#if VULKAN_SUPPORTS_GEOMETRY_SHADERS
 		// We don't support tessellation on desktop currently
 		//Hull			= 3,
 		//Domain		= 4,
@@ -38,6 +34,10 @@ namespace ShaderStage
 		NumStages		= 3,
 
 		MaxNumSets		= 8,
+#else
+		NumStages		= 2,
+
+		MaxNumSets		= 4,
 #endif
 
 		// Compute is its own pipeline, so it can all live as set 0
