UE4 - Android merge/compile fix
#rnx

diff --git a/Engine/Source/Runtime/RHI/Public/RHIResources.h b/Engine/Source/Runtime/RHI/Public/RHIResources.h
--- a/Engine/Source/Runtime/RHI/Public/RHIResources.h
+++ b/Engine/Source/Runtime/RHI/Public/RHIResources.h
@@ -1949,9 +1949,13 @@ public:
 		if (!FGraphicsMinimalPipelineStateInitializer::operator ==(rhs) ||
 			VertexShaderHash != rhs.VertexShaderHash ||
 			PixelShaderHash != rhs.PixelShaderHash ||
+#if PLATFORM_SUPPORTS_GEOMETRY_SHADERS
 			GeometryShaderHash != rhs.GeometryShaderHash ||
+#endif
+#if PLATFORM_SUPPORTS_TESSELLATION_SHADERS
 			HullShaderHash != rhs.HullShaderHash ||
 			DomainShaderHash != rhs.DomainShaderHash ||
+#endif
 			RenderTargetsEnabled != rhs.RenderTargetsEnabled ||
 			RenderTargetFormats != rhs.RenderTargetFormats || 
 			RenderTargetFlags != rhs.RenderTargetFlags || 
