UE4.22 - Vertex declaration cache
- RHICreateVertexDeclaration is no longer an immediate cmd list method
#rnx

diff --git a/Engine/Source/Runtime/Apple/MetalRHI/Private/MetalVertexDeclaration.cpp b/Engine/Source/Runtime/Apple/MetalRHI/Private/MetalVertexDeclaration.cpp
--- a/Engine/Source/Runtime/Apple/MetalRHI/Private/MetalVertexDeclaration.cpp
+++ b/Engine/Source/Runtime/Apple/MetalRHI/Private/MetalVertexDeclaration.cpp
@@ -163,11 +163,6 @@ FMetalVertexDeclaration::~FMetalVertexDeclaration()
 {
 }
 
-FVertexDeclarationRHIRef FMetalDynamicRHI::CreateVertexDeclaration_RenderThread(class FRHICommandListImmediate& RHICmdList, const FVertexDeclarationElementList& Elements)
-{
-	return GDynamicRHI->RHICreateVertexDeclaration(Elements);
-}
-
 FVertexDeclarationRHIRef FMetalDynamicRHI::RHICreateVertexDeclaration(const FVertexDeclarationElementList& Elements)
 {
 	@autoreleasepool {
diff --git a/Engine/Source/Runtime/Apple/MetalRHI/Public/MetalRHI.h b/Engine/Source/Runtime/Apple/MetalRHI/Public/MetalRHI.h
--- a/Engine/Source/Runtime/Apple/MetalRHI/Public/MetalRHI.h
+++ b/Engine/Source/Runtime/Apple/MetalRHI/Public/MetalRHI.h
@@ -166,8 +166,6 @@ public:
 	virtual FTextureCubeRHIRef RHICreateTextureCube_RenderThread(class FRHICommandListImmediate& RHICmdList, uint32 Size, uint8 Format, uint32 NumMips, uint32 Flags, FRHIResourceCreateInfo& CreateInfo) final override;
 	virtual FTextureCubeRHIRef RHICreateTextureCubeArray_RenderThread(class FRHICommandListImmediate& RHICmdList, uint32 Size, uint32 ArraySize, uint8 Format, uint32 NumMips, uint32 Flags, FRHIResourceCreateInfo& CreateInfo) final override;
 	
-	virtual FVertexDeclarationRHIRef CreateVertexDeclaration_RenderThread(class FRHICommandListImmediate& RHICmdList, const FVertexDeclarationElementList& Elements) final override;
-	
 	virtual FRenderQueryRHIRef RHICreateRenderQuery_RenderThread(class FRHICommandListImmediate& RHICmdList, ERenderQueryType QueryType) final override;
 	
 	virtual FUnorderedAccessViewRHIRef RHICreateUnorderedAccessView_RenderThread(class FRHICommandListImmediate& RHICmdList, FStructuredBufferRHIParamRef StructuredBuffer, bool bUseUAVCounter, bool bAppendBuffer) final override;
diff --git a/Engine/Source/Runtime/D3D12RHI/Private/D3D12RHIPrivate.h b/Engine/Source/Runtime/D3D12RHI/Private/D3D12RHIPrivate.h
--- a/Engine/Source/Runtime/D3D12RHI/Private/D3D12RHIPrivate.h
+++ b/Engine/Source/Runtime/D3D12RHI/Private/D3D12RHIPrivate.h
@@ -385,11 +385,6 @@ public:
 	// These will be un-commented as they are implemented.
 	//
 
-	virtual FVertexDeclarationRHIRef CreateVertexDeclaration_RenderThread(class FRHICommandListImmediate& RHICmdList, const FVertexDeclarationElementList& Elements) override final
-	{
-		return RHICreateVertexDeclaration(Elements);
-	}
-
 	virtual FVertexShaderRHIRef CreateVertexShader_RenderThread(class FRHICommandListImmediate& RHICmdList, const TArray<uint8>& Code) override final
 	{
 		return RHICreateVertexShader(Code);
diff --git a/Engine/Source/Runtime/OpenGLDrv/Public/OpenGLDrv.h b/Engine/Source/Runtime/OpenGLDrv/Public/OpenGLDrv.h
--- a/Engine/Source/Runtime/OpenGLDrv/Public/OpenGLDrv.h
+++ b/Engine/Source/Runtime/OpenGLDrv/Public/OpenGLDrv.h
@@ -797,12 +797,6 @@ public:
 		return this->RHICreateShaderResourceView(Buffer);
 	}
 
-	virtual FVertexDeclarationRHIRef CreateVertexDeclaration_RenderThread(class FRHICommandListImmediate& RHICmdList, const FVertexDeclarationElementList& Elements) final override
-	{
-		// threadsafe, doesn't really do anything
-		return this->RHICreateVertexDeclaration(Elements);
-	}
-
 	virtual FTextureReferenceRHIRef RHICreateTextureReference_RenderThread(class FRHICommandListImmediate& RHICmdList, FLastRenderTimeContainer* LastRenderTime) final override
 	{
 		// threadsafe, doesn't really do anything
diff --git a/Engine/Source/Runtime/RHI/Private/PipelineStateCache.cpp b/Engine/Source/Runtime/RHI/Private/PipelineStateCache.cpp
--- a/Engine/Source/Runtime/RHI/Private/PipelineStateCache.cpp
+++ b/Engine/Source/Runtime/RHI/Private/PipelineStateCache.cpp
@@ -1167,3 +1167,26 @@ void PipelineStateCache::Shutdown()
 	}
 	FPipelineFileCache::Shutdown();
 }
+
+/** Global cache of vertex declarations. Note we don't store TRefCountPtrs, instead we AddRef() manually. */
+static TMap<uint32, FRHIVertexDeclaration*> GVertexDeclarationCache;
+static FCriticalSection GVertexDeclarationLock;
+
+FRHIVertexDeclaration*	PipelineStateCache::GetOrCreateVertexDeclaration(const FVertexDeclarationElementList& Elements)
+{
+	// Actual locking/contention time should be close to unmeasurable
+	FScopeLock ScopeLock(&GVertexDeclarationLock);
+	uint32 Key = FCrc::MemCrc_DEPRECATED(Elements.GetData(), Elements.Num() * sizeof(FVertexElement));
+	FRHIVertexDeclaration** Found = GVertexDeclarationCache.Find(Key);
+	if (Found)
+	{
+		return *Found;
+	}
+
+	FVertexDeclarationRHIRef NewDeclaration = RHICreateVertexDeclaration(Elements);
+
+	// Add an extra reference so we don't have TRefCountPtr in the maps
+	NewDeclaration->AddRef();
+	GVertexDeclarationCache.Add(Key, NewDeclaration);
+	return NewDeclaration;
+}
diff --git a/Engine/Source/Runtime/RHI/Private/RHICommandList.cpp b/Engine/Source/Runtime/RHI/Private/RHICommandList.cpp
--- a/Engine/Source/Runtime/RHI/Private/RHICommandList.cpp
+++ b/Engine/Source/Runtime/RHI/Private/RHICommandList.cpp
@@ -2132,13 +2132,6 @@ ETextureReallocationStatus FDynamicRHI::CancelAsyncReallocateTexture2D_RenderThr
 	return GDynamicRHI->RHICancelAsyncReallocateTexture2D(Texture2D, bBlockUntilCompleted);
 }
 
-FVertexDeclarationRHIRef FDynamicRHI::CreateVertexDeclaration_RenderThread(class FRHICommandListImmediate& RHICmdList, const FVertexDeclarationElementList& Elements)
-{
-	CSV_SCOPED_TIMING_STAT(RHITStalls, CreateVertexDeclaration_RenderThread);
-	FScopedRHIThreadStaller StallRHIThread(RHICmdList);
-	return GDynamicRHI->RHICreateVertexDeclaration(Elements);
-}
-
 FVertexShaderRHIRef FDynamicRHI::CreateVertexShader_RenderThread(class FRHICommandListImmediate& RHICmdList, const TArray<uint8>& Code)
 {
 	CSV_SCOPED_TIMING_STAT(RHITStalls, CreateVertexShader_RenderThread);
diff --git a/Engine/Source/Runtime/RHI/Public/DynamicRHI.h b/Engine/Source/Runtime/RHI/Public/DynamicRHI.h
--- a/Engine/Source/Runtime/RHI/Public/DynamicRHI.h
+++ b/Engine/Source/Runtime/RHI/Public/DynamicRHI.h
@@ -1026,7 +1026,6 @@ public:
 	virtual FIndexBufferRHIRef CreateIndexBuffer_RenderThread(class FRHICommandListImmediate& RHICmdList, uint32 Stride, uint32 Size, uint32 InUsage, FRHIResourceCreateInfo& CreateInfo);
 	virtual void* LockIndexBuffer_RenderThread(class FRHICommandListImmediate& RHICmdList, FIndexBufferRHIParamRef IndexBuffer, uint32 Offset, uint32 SizeRHI, EResourceLockMode LockMode);
 	virtual void UnlockIndexBuffer_RenderThread(class FRHICommandListImmediate& RHICmdList, FIndexBufferRHIParamRef IndexBuffer);
-	virtual FVertexDeclarationRHIRef CreateVertexDeclaration_RenderThread(class FRHICommandListImmediate& RHICmdList, const FVertexDeclarationElementList& Elements);
 	virtual FVertexShaderRHIRef CreateVertexShader_RenderThread(class FRHICommandListImmediate& RHICmdList, const TArray<uint8>& Code);
 	virtual FVertexShaderRHIRef CreateVertexShader_RenderThread(class FRHICommandListImmediate& RHICmdList, FRHIShaderLibraryParamRef Library, FSHAHash Hash);
 	virtual FPixelShaderRHIRef CreatePixelShader_RenderThread(class FRHICommandListImmediate& RHICmdList, const TArray<uint8>& Code);
@@ -1172,11 +1171,18 @@ FORCEINLINE FBoundShaderStateRHIRef RHICreateBoundShaderState(FVertexDeclaration
 	return GDynamicRHI->RHICreateBoundShaderState(VertexDeclaration, VertexShader, HullShader, DomainShader, PixelShader, GeometryShader);
 }
 
+/** Before using this directly go through PipelineStateCache::GetAndOrCreateGraphicsPipelineState() */
 FORCEINLINE FGraphicsPipelineStateRHIRef RHICreateGraphicsPipelineState(const FGraphicsPipelineStateInitializer& Initializer)
 {
 	return GDynamicRHI->RHICreateGraphicsPipelineState(Initializer);
 }
 
+/** Before using this directly go through PipelineStateCache::GetOrCreateVertexDeclaration() */
+FORCEINLINE FVertexDeclarationRHIRef RHICreateVertexDeclaration(const FVertexDeclarationElementList& Elements)
+{
+	return GDynamicRHI->RHICreateVertexDeclaration(Elements);
+}
+
 FORCEINLINE TRefCountPtr<FRHIComputePipelineState> RHICreateComputePipelineState(FRHIComputeShader* ComputeShader)
 {
 	return GDynamicRHI->RHICreateComputePipelineState(ComputeShader);
diff --git a/Engine/Source/Runtime/RHI/Public/PipelineStateCache.h b/Engine/Source/Runtime/RHI/Public/PipelineStateCache.h
--- a/Engine/Source/Runtime/RHI/Public/PipelineStateCache.h
+++ b/Engine/Source/Runtime/RHI/Public/PipelineStateCache.h
@@ -33,6 +33,8 @@ namespace PipelineStateCache
 
 	extern RHI_API FGraphicsPipelineState*	GetAndOrCreateGraphicsPipelineState(FRHICommandList& RHICmdList, const FGraphicsPipelineStateInitializer& OriginalInitializer, EApplyRendertargetOption ApplyFlags);
 
+	extern RHI_API FRHIVertexDeclaration*	GetOrCreateVertexDeclaration(const FVertexDeclarationElementList& Elements);
+
 #if RHI_RAYTRACING
 	extern RHI_API FRHIRayTracingPipelineState* GetAndOrCreateRayTracingPipelineState(const FRayTracingPipelineStateInitializer& Initializer);
 #endif
diff --git a/Engine/Source/Runtime/RHI/Public/RHICommandList.h b/Engine/Source/Runtime/RHI/Public/RHICommandList.h
--- a/Engine/Source/Runtime/RHI/Public/RHICommandList.h
+++ b/Engine/Source/Runtime/RHI/Public/RHICommandList.h
@@ -3490,10 +3490,10 @@ public:
 		return RHICreateBlendState(Initializer);
 	}
 	
+	UE_DEPRECATED(4.22, "Use PipelineStateCache::GetOrCreateVertexDeclaration() instead.")
 	FORCEINLINE FVertexDeclarationRHIRef CreateVertexDeclaration(const FVertexDeclarationElementList& Elements)
 	{
-		LLM_SCOPE(ELLMTag::Shaders);
-		return GDynamicRHI->CreateVertexDeclaration_RenderThread(*this, Elements);
+		return GDynamicRHI->RHICreateVertexDeclaration(Elements);
 	}
 	
 	FORCEINLINE FPixelShaderRHIRef CreatePixelShader(const TArray<uint8>& Code)
@@ -4444,11 +4444,6 @@ struct FScopedCommandListWaitForTasks
 };
 
 
-FORCEINLINE FVertexDeclarationRHIRef RHICreateVertexDeclaration(const FVertexDeclarationElementList& Elements)
-{
-	return FRHICommandListExecutor::GetImmediateCommandList().CreateVertexDeclaration(Elements);
-}
-
 FORCEINLINE FPixelShaderRHIRef RHICreatePixelShader(const TArray<uint8>& Code)
 {
 	return FRHICommandListExecutor::GetImmediateCommandList().CreatePixelShader(Code);
diff --git a/Engine/Source/Runtime/VulkanRHI/Public/VulkanDynamicRHI.h b/Engine/Source/Runtime/VulkanRHI/Public/VulkanDynamicRHI.h
--- a/Engine/Source/Runtime/VulkanRHI/Public/VulkanDynamicRHI.h
+++ b/Engine/Source/Runtime/VulkanRHI/Public/VulkanDynamicRHI.h
@@ -187,10 +187,6 @@ public:
 	virtual void UnlockIndexBuffer_RenderThread(class FRHICommandListImmediate& RHICmdList, FIndexBufferRHIParamRef IndexBuffer) override final;
 #endif
 
-	virtual FVertexDeclarationRHIRef CreateVertexDeclaration_RenderThread(class FRHICommandListImmediate& RHICmdList, const FVertexDeclarationElementList& Elements) override final
-	{
-		return RHICreateVertexDeclaration(Elements);
-	}
 	virtual FVertexShaderRHIRef CreateVertexShader_RenderThread(class FRHICommandListImmediate& RHICmdList, const TArray<uint8>& Code) override final
 	{
 		return RHICreateVertexShader(Code);
