UE4.22 - Free vertex declaration cache
#rnx

diff --git a/Engine/Source/Runtime/RHI/Private/PipelineStateCache.cpp b/Engine/Source/Runtime/RHI/Private/PipelineStateCache.cpp
--- a/Engine/Source/Runtime/RHI/Private/PipelineStateCache.cpp
+++ b/Engine/Source/Runtime/RHI/Private/PipelineStateCache.cpp
@@ -1144,6 +1144,10 @@ void DumpPipelineCacheStats()
 #endif // PSO_VALIDATE_CACHE
 }
 
+/** Global cache of vertex declarations. Note we don't store TRefCountPtrs, instead we AddRef() manually. */
+static TMap<uint32, FRHIVertexDeclaration*> GVertexDeclarationCache;
+static FCriticalSection GVertexDeclarationLock;
+
 void PipelineStateCache::Shutdown()
 {
 	GGraphicsPipelineCache.WaitTasksComplete();
@@ -1166,11 +1170,13 @@ void PipelineStateCache::Shutdown()
 		GGraphicsPipelineCache.DiscardAndSwap();
 	}
 	FPipelineFileCache::Shutdown();
-}
 
-/** Global cache of vertex declarations. Note we don't store TRefCountPtrs, instead we AddRef() manually. */
-static TMap<uint32, FRHIVertexDeclaration*> GVertexDeclarationCache;
-static FCriticalSection GVertexDeclarationLock;
+	for (auto Pair : GVertexDeclarationCache)
+	{
+		Pair.Value->Release();
+	}
+	GVertexDeclarationCache.Empty();
+}
 
 FRHIVertexDeclaration*	PipelineStateCache::GetOrCreateVertexDeclaration(const FVertexDeclarationElementList& Elements)
 {
