UE4.22 - Revert to emulated uniform buffers on Vulkan for all platforms
#rnx

diff --git a/Engine/Config/BaseDeviceProfiles.ini b/Engine/Config/BaseDeviceProfiles.ini
--- a/Engine/Config/BaseDeviceProfiles.ini
+++ b/Engine/Config/BaseDeviceProfiles.ini
@@ -120,7 +120,7 @@ TextureLODGroups=(Group=TEXTUREGROUP_World,MinLODSize=1,MaxLODSize=8192,LODBias=
 [Windows DeviceProfile]
 DeviceType=Windows
 BaseProfileName=
-+CVars=r.Vulkan.UseRealUBs=1
++CVars=r.Vulkan.UseRealUBs=0		; Required until we fix HLR fall-out
 ; +CVars=r.Vulkan.EnableValidation=0
 
 [WindowsNoEditor DeviceProfile]
@@ -634,6 +634,7 @@ BaseProfileName=
 +CVars=r.LightShaftQuality=0
 +CVars=r.RefractionQuality=0
 +CVars=r.ShadowQuality=2
++CVars=r.Vulkan.UseRealUBs=0		; Required until we fix HLR fall-out
 
 [LuminClient DeviceProfile]
 DeviceType=Lumin
diff --git a/Engine/Source/Runtime/RHI/Private/RHI.cpp b/Engine/Source/Runtime/RHI/Private/RHI.cpp
--- a/Engine/Source/Runtime/RHI/Private/RHI.cpp
+++ b/Engine/Source/Runtime/RHI/Private/RHI.cpp
@@ -36,7 +36,7 @@ DEFINE_STAT(STAT_PixelBufferMemory);
 
 static FAutoConsoleVariable CVarUseVulkanRealUBs(
 	TEXT("r.Vulkan.UseRealUBs"),
-	1,
+	0,
 	TEXT("0: Emulate uniform buffers on Vulkan SM4/SM5 (debugging ONLY)\n")
 	TEXT("1: Use real uniform buffers [default]"),
 	ECVF_ReadOnly
diff --git a/Engine/Source/Runtime/RHI/Public/RHIShaderFormatDefinitions.inl b/Engine/Source/Runtime/RHI/Public/RHIShaderFormatDefinitions.inl
--- a/Engine/Source/Runtime/RHI/Public/RHIShaderFormatDefinitions.inl
+++ b/Engine/Source/Runtime/RHI/Public/RHIShaderFormatDefinitions.inl
@@ -127,31 +127,18 @@ static FName ShaderPlatformToShaderFormatName(EShaderPlatform Platform)
 
 	case SP_VULKAN_PCES3_1:
 	{
-		// If you modify this, make sure to update FWindowsPlatform::UseRealUBsOptimization()
-#if UE_BUILD_DEBUG || UE_BUILD_DEVELOPMENT
 		static auto* CVar = IConsoleManager::Get().FindTConsoleVariableDataInt(TEXT("r.Vulkan.UseRealUBs"));
 		return (CVar && CVar->GetValueOnAnyThread() == 0) ? NAME_VULKAN_ES3_1_NOUB : NAME_VULKAN_ES3_1;
-#else
-		return NAME_VULKAN_ES3_1;
-#endif
 	}
 	case SP_VULKAN_SM4:
 	{
-#if UE_BUILD_DEBUG || UE_BUILD_DEVELOPMENT
 		static auto* CVar = IConsoleManager::Get().FindTConsoleVariableDataInt(TEXT("r.Vulkan.UseRealUBs"));
 		return (CVar && CVar->GetValueOnAnyThread() == 0) ? NAME_VULKAN_SM4_NOUB : NAME_VULKAN_SM4;
-#else
-		return NAME_VULKAN_SM4;
-#endif
 	}
 	case SP_VULKAN_SM5:
 	{
-#if UE_BUILD_DEBUG || UE_BUILD_DEVELOPMENT
 		static auto* CVar = IConsoleManager::Get().FindTConsoleVariableDataInt(TEXT("r.Vulkan.UseRealUBs"));
 		return (CVar && CVar->GetValueOnAnyThread() == 0) ? NAME_VULKAN_SM5_NOUB : NAME_VULKAN_SM5;
-#else
-		return NAME_VULKAN_SM5;
-#endif
 	}
 	case SP_VULKAN_SM5_LUMIN:
 		return NAME_VULKAN_SM5_LUMIN;
diff --git a/Engine/Source/Runtime/VulkanRHI/Private/Linux/VulkanLinuxPlatform.h b/Engine/Source/Runtime/VulkanRHI/Private/Linux/VulkanLinuxPlatform.h
--- a/Engine/Source/Runtime/VulkanRHI/Private/Linux/VulkanLinuxPlatform.h
+++ b/Engine/Source/Runtime/VulkanRHI/Private/Linux/VulkanLinuxPlatform.h
@@ -49,12 +49,8 @@ public:
 	// Some platforms only support real or non-real UBs, so this function can optimize it out
 	static bool UseRealUBsOptimization(bool bCodeHeaderUseRealUBs)
 	{
-#if UE_BUILD_DEBUG || UE_BUILD_DEVELOPMENT
 		static auto* CVar = IConsoleManager::Get().FindTConsoleVariableDataInt(TEXT("r.Vulkan.UseRealUBs"));
 		return (CVar && CVar->GetValueOnAnyThread() == 0) ? false : bCodeHeaderUseRealUBs;
-#else
-		return GMaxRHIFeatureLevel >= ERHIFeatureLevel::ES3_1 ? bCodeHeaderUseRealUBs : false;
-#endif
 	}
 
 	static void WriteCrashMarker(const FOptionalVulkanDeviceExtensions& OptionalExtensions, VkCommandBuffer CmdBuffer, VkBuffer DestBuffer, const TArrayView<uint32>& Entries, bool bAdding);
diff --git a/Engine/Source/Runtime/VulkanRHI/Private/Windows/VulkanWindowsPlatform.h b/Engine/Source/Runtime/VulkanRHI/Private/Windows/VulkanWindowsPlatform.h
--- a/Engine/Source/Runtime/VulkanRHI/Private/Windows/VulkanWindowsPlatform.h
+++ b/Engine/Source/Runtime/VulkanRHI/Private/Windows/VulkanWindowsPlatform.h
@@ -67,12 +67,8 @@ public:
 	// Some platforms only support real or non-real UBs, so this function can optimize it out
 	static bool UseRealUBsOptimization(bool bCodeHeaderUseRealUBs)
 	{
-#if UE_BUILD_DEBUG || UE_BUILD_DEVELOPMENT
 		static auto* CVar = IConsoleManager::Get().FindTConsoleVariableDataInt(TEXT("r.Vulkan.UseRealUBs"));
 		return (CVar && CVar->GetValueOnAnyThread() == 0) ? false : bCodeHeaderUseRealUBs;
-#else
-		return GMaxRHIFeatureLevel >= ERHIFeatureLevel::ES3_1 ? bCodeHeaderUseRealUBs : false;
-#endif
 	}
 };
 
