UE4.22 - Deprecate  SetRenderTarget API

diff --git a/Engine/Plugins/Compositing/OpenColorIO/Source/OpenColorIO/Private/OpenColorIORendering.cpp b/Engine/Plugins/Compositing/OpenColorIO/Source/OpenColorIO/Private/OpenColorIORendering.cpp
--- a/Engine/Plugins/Compositing/OpenColorIO/Source/OpenColorIO/Private/OpenColorIORendering.cpp
+++ b/Engine/Plugins/Compositing/OpenColorIO/Source/OpenColorIO/Private/OpenColorIORendering.cpp
@@ -71,13 +71,8 @@ static void ProcessOCIOColorSpaceTransform_RenderThread(
 
 	SCOPED_DRAW_EVENT(InRHICmdList, ProcessOCIOColorSpaceTransform);
 
-	// Set render target.
-	SetRenderTarget(
-		InRHICmdList,
-		OutputSpaceColorResource->TextureRHI,
-		FTextureRHIRef(),
-		ESimpleRenderTargetMode::EUninitializedColorAndDepth,
-		FExclusiveDepthStencil::DepthNop_StencilNop);
+	FRHIRenderPassInfo RPInfo(OutputSpaceColorResource->TextureRHI, ERenderTargetActions::DontLoad_Store);
+	InRHICmdList.BeginRenderPass(RPInfo, TEXT("ProcessOCIOColorSpaceXfrm"));
 
 	FIntPoint Resolution(OutputSpaceColorResource->GetSizeX(), OutputSpaceColorResource->GetSizeY());
 
@@ -115,7 +110,7 @@ static void ProcessOCIOColorSpaceTransform_RenderThread(
 	InRHICmdList.DrawPrimitive(0, 2, 1);
 
 	// Resolve render target.
-	InRHICmdList.CopyToResolveTarget(OutputSpaceColorResource->TextureRHI, OutputSpaceColorResource->TextureRHI, FResolveParams());
+	InRHICmdList.EndRenderPass();
 }
 
 // static
diff --git a/Engine/Plugins/Media/MediaIOFramework/Source/MediaIOCore/Private/MediaCapture.cpp b/Engine/Plugins/Media/MediaIOFramework/Source/MediaIOCore/Private/MediaCapture.cpp
--- a/Engine/Plugins/Media/MediaIOFramework/Source/MediaIOCore/Private/MediaCapture.cpp
+++ b/Engine/Plugins/Media/MediaIOFramework/Source/MediaIOCore/Private/MediaCapture.cpp
@@ -667,7 +667,9 @@ void UMediaCapture::OnEndFrame_GameThread()
 					// convert the source with a draw call
 					FGraphicsPipelineStateInitializer GraphicsPSOInit;
 					FRHITexture* RenderTarget = DestRenderTarget.TargetableTexture.GetReference();
+					PRAGMA_DISABLE_DEPRECATION_WARNINGS
 					SetRenderTargets(RHICmdList, 1, &RenderTarget, nullptr, ESimpleRenderTargetMode::EExistingColorAndDepth, FExclusiveDepthStencil::DepthNop_StencilNop);
+					PRAGMA_ENABLE_DEPRECATION_WARNINGS
 
 					RHICmdList.ApplyCachedRenderTargets(GraphicsPSOInit);
 
diff --git a/Engine/Source/Runtime/RHI/Public/RHICommandList.h b/Engine/Source/Runtime/RHI/Public/RHICommandList.h
--- a/Engine/Source/Runtime/RHI/Public/RHICommandList.h
+++ b/Engine/Source/Runtime/RHI/Public/RHICommandList.h
@@ -2381,6 +2381,7 @@ public:
 		}
 	}
 
+	UE_DEPRECATED(4.22, "SetRenderTargets API is deprecated; please use RHIBegin/EndRenderPass instead.")
 	FORCEINLINE_DEBUGGABLE void SetRenderTargets(
 		uint32 NewNumSimultaneousRenderTargets,
 		const FRHIRenderTargetView* NewRenderTargetsRHI,
@@ -2414,6 +2415,7 @@ public:
 			UAVs);
 	}
 
+	UE_DEPRECATED(4.22, "SetRenderTargets API is deprecated; please use RHIBegin/EndRenderPass instead.")
 	FORCEINLINE_DEBUGGABLE void SetRenderTargetsAndClear(const FRHISetRenderTargetsInfo& RenderTargetsInfo)
 	{
 		check(IsOutsideRenderPass());
diff --git a/Engine/Source/Runtime/RHI/Public/RHIUtilities.h b/Engine/Source/Runtime/RHI/Public/RHIUtilities.h
--- a/Engine/Source/Runtime/RHI/Public/RHIUtilities.h
+++ b/Engine/Source/Runtime/RHI/Public/RHIUtilities.h
@@ -418,6 +418,7 @@ inline void TransitionRenderPassTargets(FRHICommandList& RHICmdList, const FRHIR
 }
 
 /** Helper for the common case of using a single color and depth render target. */
+UE_DEPRECATED(4.22, "SetRenderTargets API is deprecated; please use RHIBegin/EndRenderPass instead.")
 inline void SetRenderTarget(FRHICommandList& RHICmdList, FTextureRHIParamRef NewRenderTarget, FTextureRHIParamRef NewDepthStencilTarget, bool bWritableBarrier = false)
 {
 	FRHIRenderTargetView RTV(NewRenderTarget, ERenderTargetLoadAction::ELoad);
@@ -428,10 +429,13 @@ inline void SetRenderTarget(FRHICommandList& RHICmdList, FTextureRHIParamRef New
 	{
 		TransitionSetRenderTargetsHelper(RHICmdList, NewRenderTarget, NewDepthStencilTarget, FExclusiveDepthStencil::DepthWrite_StencilWrite);
 	}
+	PRAGMA_DISABLE_DEPRECATION_WARNINGS
 	RHICmdList.SetRenderTargets(1, &RTV, &DepthRTV, 0, NULL);
+	PRAGMA_ENABLE_DEPRECATION_WARNINGS
 }
 
 /** Helper for the common case of using a single color and depth render target. */
+UE_DEPRECATED(4.22, "SetRenderTargets API is deprecated; please use RHIBegin/EndRenderPass instead.")
 inline void SetRenderTarget(FRHICommandList& RHICmdList, FTextureRHIParamRef NewRenderTarget, FTextureRHIParamRef NewDepthStencilTarget, ESimpleRenderTargetMode Mode, FExclusiveDepthStencil DepthStencilAccess = FExclusiveDepthStencil::DepthWrite_StencilWrite, bool bWritableBarrier = false)
 {
 	ERenderTargetLoadAction ColorLoadAction, DepthLoadAction, StencilLoadAction;
@@ -447,7 +451,9 @@ inline void SetRenderTarget(FRHICommandList& RHICmdList, FTextureRHIParamRef New
 	// now make the FRHISetRenderTargetsInfo that encapsulates all of the info
 	FRHIRenderTargetView ColorView(NewRenderTarget, 0, -1, ColorLoadAction, ColorStoreAction);
 	FRHISetRenderTargetsInfo Info(1, &ColorView, FRHIDepthRenderTargetView(NewDepthStencilTarget, DepthLoadAction, DepthStoreAction, StencilLoadAction, StencilStoreAction, DepthStencilAccess));
+	PRAGMA_DISABLE_DEPRECATION_WARNINGS
 	RHICmdList.SetRenderTargetsAndClear(Info);
+	PRAGMA_ENABLE_DEPRECATION_WARNINGS
 }
 
 /** Helper for the common case of using a single color and depth render target, with a mip index for the color target. */
@@ -462,7 +468,9 @@ inline void SetRenderTarget(FRHICommandList& RHICmdList, FTextureRHIParamRef New
 	{
 		TransitionSetRenderTargetsHelper(RHICmdList, NewRenderTarget, NewDepthStencilTarget, FExclusiveDepthStencil::DepthWrite_StencilWrite);
 	}
+	PRAGMA_DISABLE_DEPRECATION_WARNINGS
 	RHICmdList.SetRenderTargets(1, &RTV, &DepthRTV, 0, nullptr);
+	PRAGMA_ENABLE_DEPRECATION_WARNINGS
 }
 
 /** Helper for the common case of using a single color and depth render target, with a mip index for the color target. */
@@ -477,7 +485,9 @@ inline void SetRenderTarget(FRHICommandList& RHICmdList, FTextureRHIParamRef New
 	{
 		TransitionSetRenderTargetsHelper(RHICmdList, NewRenderTarget, NewDepthStencilTarget, FExclusiveDepthStencil::DepthWrite_StencilWrite);
 	}
+	PRAGMA_DISABLE_DEPRECATION_WARNINGS
 	RHICmdList.SetRenderTargets(1, &RTV, &DepthRTV, 0, nullptr);
+	PRAGMA_ENABLE_DEPRECATION_WARNINGS
 }
 
 /** Helper that converts FTextureRHIParamRef's into FRHIRenderTargetView's. */
@@ -506,10 +516,13 @@ inline void SetRenderTargets(
 	}
 
 	FRHIDepthRenderTargetView DepthRTV(NewDepthStencilTargetRHI, ERenderTargetLoadAction::ELoad, ERenderTargetStoreAction::EStore);
+	PRAGMA_DISABLE_DEPRECATION_WARNINGS
 	RHICmdList.SetRenderTargets(NewNumSimultaneousRenderTargets, RTVs, &DepthRTV, NewNumUAVs, UAVs);
+	PRAGMA_ENABLE_DEPRECATION_WARNINGS
 }
 
 /** Helper that converts FTextureRHIParamRef's into FRHIRenderTargetView's. */
+UE_DEPRECATED(4.22, "SetRenderTargets API is deprecated; please use RHIBegin/EndRenderPass instead.")
 inline void SetRenderTargets(
 	FRHICommandList& RHICmdList,
 	uint32 NewNumSimultaneousRenderTargets,
@@ -538,9 +551,12 @@ inline void SetRenderTargets(
 	}
 
 	FRHIDepthRenderTargetView DepthRTV(NewDepthStencilTargetRHI, DepthLoadAction, DepthStoreAction, StencilLoadAction, StencilStoreAction, DepthStencilAccess);
+	PRAGMA_DISABLE_DEPRECATION_WARNINGS
 	RHICmdList.SetRenderTargets(NewNumSimultaneousRenderTargets, RTVs, &DepthRTV, 0, nullptr);
+	PRAGMA_ENABLE_DEPRECATION_WARNINGS
 }
 
+// Will soon be deprecated as well...
 inline void UnbindRenderTargets(FRHICommandList& RHICmdList)
 {
 	check(RHICmdList.IsOutsideRenderPass());
diff --git a/Engine/Source/Runtime/Renderer/Private/PostProcess/PostProcessMitchellNetravali.cpp b/Engine/Source/Runtime/Renderer/Private/PostProcess/PostProcessMitchellNetravali.cpp
--- a/Engine/Source/Runtime/Renderer/Private/PostProcess/PostProcessMitchellNetravali.cpp
+++ b/Engine/Source/Runtime/Renderer/Private/PostProcess/PostProcessMitchellNetravali.cpp
@@ -77,7 +77,7 @@ void FRCPassMitchellNetravaliDownsample::Process(FRenderingCompositePassContext&
 		DestRect.Width(), DestRect.Height());
 
 	// Common setup
-	SetRenderTarget(Context.RHICmdList, nullptr, nullptr);
+	UnbindRenderTargets(Context.RHICmdList);
 	Context.SetViewportAndCallRHI(DestRect, 0.0f, 1.0f);
 
 	FTextureRHIRef EyeAdaptationTex = GWhiteTexture->TextureRHI;
diff --git a/Engine/Source/Runtime/Renderer/Private/RayTracing/RayTracingRectLight.cpp b/Engine/Source/Runtime/Renderer/Private/RayTracing/RayTracingRectLight.cpp
--- a/Engine/Source/Runtime/Renderer/Private/RayTracing/RayTracingRectLight.cpp
+++ b/Engine/Source/Runtime/Renderer/Private/RayTracing/RayTracingRectLight.cpp
@@ -270,7 +270,9 @@ void FDeferredShadingSceneRenderer::VisualizeRectLightMipTree(
 		SceneContext.GetSceneColor()->GetRenderTargetItem().TargetableTexture,
 		RectLightMipTreeRT->GetRenderTargetItem().TargetableTexture
 	};
+	PRAGMA_DISABLE_DEPRECATION_WARNINGS
 	SetRenderTargets(RHICmdList, 2, RenderTargets, SceneContext.GetSceneDepthSurface(), ESimpleRenderTargetMode::EExistingColorAndDepth, FExclusiveDepthStencil::DepthRead_StencilNop);
+	PRAGMA_ENABLE_DEPRECATION_WARNINGS
 
 	// PSO definition
 	FGraphicsPipelineStateInitializer GraphicsPSOInit;
diff --git a/Engine/Source/Runtime/Renderer/Private/SceneRendering.cpp b/Engine/Source/Runtime/Renderer/Private/SceneRendering.cpp
--- a/Engine/Source/Runtime/Renderer/Private/SceneRendering.cpp
+++ b/Engine/Source/Runtime/Renderer/Private/SceneRendering.cpp
@@ -3251,7 +3251,9 @@ void FRendererModule::CreateAndInitSingleView(FRHICommandListImmediate& RHICmdLi
 	// Create and add the new view
 	FViewInfo* NewView = new FViewInfo(*ViewInitOptions);
 	ViewFamily->Views.Add(NewView);
+	PRAGMA_DISABLE_DEPRECATION_WARNINGS
 	SetRenderTarget(RHICmdList, ViewFamily->RenderTarget->GetRenderTargetTexture(), nullptr, ESimpleRenderTargetMode::EClearColorExistingDepth);
+	PRAGMA_ENABLE_DEPRECATION_WARNINGS
 	FViewInfo* View = (FViewInfo*)ViewFamily->Views[0];
 	View->ViewRect = View->UnscaledViewRect;
 	View->InitRHIResources();
@@ -3502,7 +3504,9 @@ static void DisplayInternals(FRHICommandListImmediate& RHICmdList, FViewInfo& In
 		FCanvas Canvas((FRenderTarget*)Family->RenderTarget, NULL, Family->CurrentRealTime, Family->CurrentWorldTime, Family->DeltaWorldTime, InView.GetFeatureLevel());
 		Canvas.SetRenderTargetRect(FIntRect(0, 0, Family->RenderTarget->GetSizeXY().X, Family->RenderTarget->GetSizeXY().Y));
 
+		PRAGMA_DISABLE_DEPRECATION_WARNINGS
 		SetRenderTarget(RHICmdList, Family->RenderTarget->GetRenderTargetTexture(), FTextureRHIRef());
+		PRAGMA_ENABLE_DEPRECATION_WARNINGS
 
 		// further down to not intersect with "LIGHTING NEEDS TO BE REBUILT"
 		FVector2D Pos(30, 140);
diff --git a/Engine/Source/Runtime/Renderer/Private/VisualizeTexturePresent.cpp b/Engine/Source/Runtime/Renderer/Private/VisualizeTexturePresent.cpp
--- a/Engine/Source/Runtime/Renderer/Private/VisualizeTexturePresent.cpp
+++ b/Engine/Source/Runtime/Renderer/Private/VisualizeTexturePresent.cpp
@@ -453,7 +453,9 @@ void FVisualizeTexturePresent::PresentContent(FRHICommandListImmediate& RHICmdLi
 	}
 
 	auto& RenderTarget = View.Family->RenderTarget->GetRenderTargetTexture();
+	PRAGMA_DISABLE_DEPRECATION_WARNINGS
 	SetRenderTarget(RHICmdList, RenderTarget, FTextureRHIRef(), true);
+	PRAGMA_ENABLE_DEPRECATION_WARNINGS
 	RHICmdList.SetViewport(DestRect.Min.X, DestRect.Min.Y, 0.0f, DestRect.Max.X, DestRect.Max.Y, 1.0f);
 
 	FGraphicsPipelineStateInitializer GraphicsPSOInit;
