[LLM] removing rhi tasks from rhi misc so they're captured by the tag that created them.

(undo CL 5312123)

diff --git a/Engine/Source/Runtime/RHI/Private/RHICommandList.cpp b/Engine/Source/Runtime/RHI/Private/RHICommandList.cpp
--- a/Engine/Source/Runtime/RHI/Private/RHICommandList.cpp
+++ b/Engine/Source/Runtime/RHI/Private/RHICommandList.cpp
@@ -425,8 +425,6 @@ public:
 
 	void DoTask(ENamedThreads::Type CurrentThread, const FGraphEventRef& MyCompletionGraphEvent)
 	{
-		LLM_SCOPE(ELLMTag::RHIMisc);
-
 		check(bRHIThread || IsInRenderingThread());
 		FGraphEventArray Prereq;
 		if (RHIThreadTask.GetReference())
