[LLM] Removed VertexData tag and now Skeletal, Static etc. vertex data is included in the respective tags. This makes the tag data closer to that reported by 'obj list'

diff --git a/Engine/Source/Runtime/Core/Private/HAL/LowLevelMemTracker.cpp b/Engine/Source/Runtime/Core/Private/HAL/LowLevelMemTracker.cpp
--- a/Engine/Source/Runtime/Core/Private/HAL/LowLevelMemTracker.cpp
+++ b/Engine/Source/Runtime/Core/Private/HAL/LowLevelMemTracker.cpp
@@ -110,7 +110,6 @@ DECLARE_LLM_MEMORY_STAT(TEXT("VideoRecording"), STAT_VideoRecordingLLM, STATGROU
 DECLARE_LLM_MEMORY_STAT(TEXT("Replays"), STAT_ReplaysLLM, STATGROUP_LLMFULL);
 DECLARE_LLM_MEMORY_STAT(TEXT("CsvProfiler"), STAT_CsvProfilerLLM, STATGROUP_LLMFULL);
 DECLARE_LLM_MEMORY_STAT(TEXT("MaterialInstance"), STAT_MaterialInstanceLLM, STATGROUP_LLMFULL);
-DECLARE_LLM_MEMORY_STAT(TEXT("VertexData"), STAT_VertexDataLLM, STATGROUP_LLMFULL);
 DECLARE_LLM_MEMORY_STAT(TEXT("SkeletalMesh"), STAT_SkeletalMeshLLM, STATGROUP_LLMFULL);
 DECLARE_LLM_MEMORY_STAT(TEXT("InstancedMesh"), STAT_InstancedMeshLLM, STATGROUP_LLMFULL);
 DECLARE_LLM_MEMORY_STAT(TEXT("Landscape"), STAT_LandscapeLLM, STATGROUP_LLMFULL);
diff --git a/Engine/Source/Runtime/Core/Public/HAL/LowLevelMemTracker.h b/Engine/Source/Runtime/Core/Public/HAL/LowLevelMemTracker.h
--- a/Engine/Source/Runtime/Core/Public/HAL/LowLevelMemTracker.h
+++ b/Engine/Source/Runtime/Core/Public/HAL/LowLevelMemTracker.h
@@ -235,7 +235,6 @@ enum class ELLMTagSet : uint8
 	macro(VideoRecording,						"VideoRecording",				GET_STATFNAME(STAT_VideoRecordingLLM),						GET_STATFNAME(STAT_EngineSummaryLLM),			-1)\
 	macro(Replays,								"Replays",						GET_STATFNAME(STAT_ReplaysLLM),								GET_STATFNAME(STAT_EngineSummaryLLM),			-1)\
 	macro(MaterialInstance,						"MaterialInstance",				GET_STATFNAME(STAT_MaterialInstanceLLM),					GET_STATFNAME(STAT_EngineSummaryLLM),			-1)\
-	macro(VertexData,							"VertexData",					GET_STATFNAME(STAT_VertexDataLLM),							GET_STATFNAME(STAT_EngineSummaryLLM),			-1)\
 	macro(SkeletalMesh,							"SkeletalMesh",					GET_STATFNAME(STAT_SkeletalMeshLLM),						GET_STATFNAME(STAT_EngineSummaryLLM),			ELLMTag::Meshes)\
 	macro(InstancedMesh,						"InstancedMesh",				GET_STATFNAME(STAT_InstancedMeshLLM),						GET_STATFNAME(STAT_EngineSummaryLLM),			ELLMTag::Meshes)\
 	macro(Landscape,							"Landscape",					GET_STATFNAME(STAT_LandscapeLLM),							GET_STATFNAME(STAT_EngineSummaryLLM),			ELLMTag::Meshes)\
diff --git a/Engine/Source/Runtime/Engine/Private/SkeletalMesh.cpp b/Engine/Source/Runtime/Engine/Private/SkeletalMesh.cpp
--- a/Engine/Source/Runtime/Engine/Private/SkeletalMesh.cpp
+++ b/Engine/Source/Runtime/Engine/Private/SkeletalMesh.cpp
@@ -613,6 +613,8 @@ bool USkeletalMesh::NeedCPUData(int32 LODIndex)const
 
 void USkeletalMesh::InitResources()
 {
+	LLM_SCOPE(ELLMTag::SkeletalMesh);
+
 	UpdateUVChannelData(false);
 
 	FSkeletalMeshRenderData* SkelMeshRenderData = GetResourceForRendering();
diff --git a/Engine/Source/Runtime/Engine/Private/StaticMesh.cpp b/Engine/Source/Runtime/Engine/Private/StaticMesh.cpp
--- a/Engine/Source/Runtime/Engine/Private/StaticMesh.cpp
+++ b/Engine/Source/Runtime/Engine/Private/StaticMesh.cpp
@@ -2362,6 +2362,8 @@ void UStaticMesh::PostInitProperties()
  */
 void UStaticMesh::InitResources()
 {
+	LLM_SCOPE(ELLMTag::StaticMesh);
+
 	bRenderingResourcesInitialized = true;
 
 	UpdateUVChannelData(false);
diff --git a/Engine/Source/Runtime/Landscape/Private/LandscapeRender.cpp b/Engine/Source/Runtime/Landscape/Private/LandscapeRender.cpp
--- a/Engine/Source/Runtime/Landscape/Private/LandscapeRender.cpp
+++ b/Engine/Source/Runtime/Landscape/Private/LandscapeRender.cpp
@@ -35,6 +35,7 @@ LandscapeRender.cpp: New terrain rendering
 #include "SceneView.h"
 #include "Runtime/Renderer/Private/SceneCore.h"
 #include "LandscapeProxy.h"
+#include "HAL/LowLevelMemTracker.h"
 #include "MeshMaterialShader.h"
 
 IMPLEMENT_GLOBAL_SHADER_PARAMETER_STRUCT(FLandscapeUniformShaderParameters, "LandscapeParameters");
@@ -850,6 +851,8 @@ FLandscapeComponentSceneProxy::FLandscapeComponentSceneProxy(ULandscapeComponent
 
 void FLandscapeComponentSceneProxy::CreateRenderThreadResources()
 {
+	LLM_SCOPE(ELLMTag::Landscape);
+
 	check(HeightmapTexture != nullptr);
 
 	if (IsComponentLevelVisible())
diff --git a/Engine/Source/Runtime/Landscape/Private/LandscapeRenderMobile.cpp b/Engine/Source/Runtime/Landscape/Private/LandscapeRenderMobile.cpp
--- a/Engine/Source/Runtime/Landscape/Private/LandscapeRenderMobile.cpp
+++ b/Engine/Source/Runtime/Landscape/Private/LandscapeRenderMobile.cpp
@@ -10,6 +10,7 @@ LandscapeRenderMobile.cpp: Landscape Rendering without using vertex texture fetc
 #include "Serialization/MemoryReader.h"
 #include "PrimitiveSceneInfo.h"
 #include "LandscapeLayerInfoObject.h"
+#include "HAL/LowLevelMemTracker.h"
 #include "MeshMaterialShader.h"
 
 void FLandscapeVertexFactoryMobile::InitRHI()
@@ -378,6 +379,8 @@ SIZE_T FLandscapeComponentSceneProxyMobile::GetTypeHash() const
 
 void FLandscapeComponentSceneProxyMobile::CreateRenderThreadResources()
 {
+	LLM_SCOPE(ELLMTag::Landscape);
+
 	if (IsComponentLevelVisible())
 	{
 		RegisterNeighbors();
diff --git a/Engine/Source/Runtime/RHI/Public/RHICommandList.h b/Engine/Source/Runtime/RHI/Public/RHICommandList.h
--- a/Engine/Source/Runtime/RHI/Public/RHICommandList.h
+++ b/Engine/Source/Runtime/RHI/Public/RHICommandList.h
@@ -3670,19 +3670,16 @@ public:
 	
 	FORCEINLINE FIndexBufferRHIRef CreateAndLockIndexBuffer(uint32 Stride, uint32 Size, uint32 InUsage, FRHIResourceCreateInfo& CreateInfo, void*& OutDataBuffer)
 	{
-		LLM_SCOPE(ELLMTag::VertexData);
 		return GDynamicRHI->CreateAndLockIndexBuffer_RenderThread(*this, Stride, Size, InUsage, CreateInfo, OutDataBuffer);
 	}
 	
 	FORCEINLINE FIndexBufferRHIRef CreateIndexBuffer(uint32 Stride, uint32 Size, uint32 InUsage, FRHIResourceCreateInfo& CreateInfo)
 	{
-		LLM_SCOPE(ELLMTag::VertexData);
 		return GDynamicRHI->CreateIndexBuffer_RenderThread(*this, Stride, Size, InUsage, CreateInfo);
 	}
 	
 	FORCEINLINE void* LockIndexBuffer(FIndexBufferRHIParamRef IndexBuffer, uint32 Offset, uint32 SizeRHI, EResourceLockMode LockMode)
 	{
-		LLM_SCOPE(ELLMTag::VertexData);
 		return GDynamicRHI->LockIndexBuffer_RenderThread(*this, IndexBuffer, Offset, SizeRHI, LockMode);
 	}
 	
@@ -3693,7 +3690,6 @@ public:
 	
 	FORCEINLINE void* LockStagingBuffer(FStagingBufferRHIParamRef StagingBuffer, uint32 Offset, uint32 SizeRHI)
 	{
-		LLM_SCOPE(ELLMTag::VertexData);
 		return GDynamicRHI->LockStagingBuffer_RenderThread(*this, StagingBuffer, Offset, SizeRHI);
 	}
 	
@@ -3704,19 +3700,16 @@ public:
 	
 	FORCEINLINE FVertexBufferRHIRef CreateAndLockVertexBuffer(uint32 Size, uint32 InUsage, FRHIResourceCreateInfo& CreateInfo, void*& OutDataBuffer)
 	{
-		LLM_SCOPE(ELLMTag::VertexData);
 		return GDynamicRHI->CreateAndLockVertexBuffer_RenderThread(*this, Size, InUsage, CreateInfo, OutDataBuffer);
 	}
 
 	FORCEINLINE FVertexBufferRHIRef CreateVertexBuffer(uint32 Size, uint32 InUsage, FRHIResourceCreateInfo& CreateInfo)
 	{
-		LLM_SCOPE(ELLMTag::VertexData);
 		return GDynamicRHI->CreateVertexBuffer_RenderThread(*this, Size, InUsage, CreateInfo);
 	}
 	
 	FORCEINLINE void* LockVertexBuffer(FVertexBufferRHIParamRef VertexBuffer, uint32 Offset, uint32 SizeRHI, EResourceLockMode LockMode)
 	{
-		LLM_SCOPE(ELLMTag::VertexData);
 		return GDynamicRHI->LockVertexBuffer_RenderThread(*this, VertexBuffer, Offset, SizeRHI, LockMode);
 	}
 	
