Lower the priority of PSO preload threads on Mac as with the highest it's blocking the game thread for too long.

diff --git a/Engine/Source/Runtime/RHI/Private/PipelineStateCache.cpp b/Engine/Source/Runtime/RHI/Private/PipelineStateCache.cpp
--- a/Engine/Source/Runtime/RHI/Private/PipelineStateCache.cpp
+++ b/Engine/Source/Runtime/RHI/Private/PipelineStateCache.cpp
@@ -582,7 +582,11 @@ FGraphicsPipelineCache GGraphicsPipelineCache;
 FAutoConsoleTaskPriority CPrio_FCompilePipelineStateTask(
 	TEXT("TaskGraph.TaskPriorities.CompilePipelineStateTask"),
 	TEXT("Task and thread priority for FCompilePipelineStateTask."),
+#if PLATFORM_MAC
+	ENamedThreads::NormalThreadPriority,	// On Mac, use normal thread priority to not freeze the loading screen for extended period of time
+#else
 	ENamedThreads::HighThreadPriority,		// if we have high priority task threads, then use them...
+#endif
 	ENamedThreads::NormalTaskPriority,		// .. at normal task priority
 	ENamedThreads::HighTaskPriority,		// if we don't have hi pri threads, then use normal priority threads at high task priority instead
 	EPowerSavingEligibility::NotEligible	// Not eligible for downgrade when power saving is requested.
