Fix deprecation warning

diff --git a/Engine/Source/Runtime/RHI/Public/RHICommandList.h b/Engine/Source/Runtime/RHI/Public/RHICommandList.h
--- a/Engine/Source/Runtime/RHI/Public/RHICommandList.h
+++ b/Engine/Source/Runtime/RHI/Public/RHICommandList.h
@@ -1765,8 +1765,8 @@ struct FRHIShaderResourceViewUpdateInfo_VB
 
 struct FRHIShaderResourceViewUpdateInfo_IB
 {
-	FShaderResourceViewRHIParamRef SRV;
-	FIndexBufferRHIParamRef IndexBuffer;
+	FRHIShaderResourceView* SRV;
+	FRHIIndexBuffer* IndexBuffer;
 };
 
 struct FRHIVertexBufferUpdateInfo
