Add RHI thread flush CSV stats. These use the RHITFlushes category, which is off by default. Use -csvcategories to enable. E.g: -csvcategories="rhitflushes,rhitstalls".

diff --git a/Engine/Source/Runtime/RHI/Private/RHICommandList.cpp b/Engine/Source/Runtime/RHI/Private/RHICommandList.cpp
--- a/Engine/Source/Runtime/RHI/Private/RHICommandList.cpp
+++ b/Engine/Source/Runtime/RHI/Private/RHICommandList.cpp
@@ -13,6 +13,7 @@
 #include "ProfilingDebugging/CsvProfiler.h"
 
 CSV_DEFINE_CATEGORY(RHITStalls,false);
+CSV_DEFINE_CATEGORY(RHITFlushes, false);
 
 DECLARE_CYCLE_STAT(TEXT("Nonimmed. Command List Execute"), STAT_NonImmedCmdListExecuteTime, STATGROUP_RHICMDLIST);
 DECLARE_DWORD_COUNTER_STAT(TEXT("Nonimmed. Command List memory"), STAT_NonImmedCmdListMemory, STATGROUP_RHICMDLIST);
@@ -1178,6 +1179,7 @@ void FRHICommandListBase::QueueParallelAsyncCommandListSubmit(FGraphEventRef* An
 	// do a flush before hand so we can tell if it was this parallel set that broke something, or what came before.
 	if (CVarRHICmdFlushOnQueueParallelSubmit.GetValueOnRenderThread())
 	{
+		CSV_SCOPED_TIMING_STAT(RHITFlushes, QueueParallelAsyncCommandListSubmit);
 		FRHICommandListExecutor::GetImmediateCommandList().ImmediateFlush(EImmediateFlushType::FlushRHIThread);
 	}
 #endif
@@ -1218,6 +1220,7 @@ void FRHICommandListBase::QueueParallelAsyncCommandListSubmit(FGraphEventRef* An
 #if !UE_BUILD_SHIPPING
 			if (CVarRHICmdFlushOnQueueParallelSubmit.GetValueOnRenderThread())
 			{
+				CSV_SCOPED_TIMING_STAT(RHITFlushes, QueueParallelAsyncCommandListSubmit);
 				FRHICommandListExecutor::GetImmediateCommandList().ImmediateFlush(EImmediateFlushType::FlushRHIThread);
 			}
 #endif
@@ -1316,6 +1319,7 @@ void FRHICommandListBase::QueueParallelAsyncCommandListSubmit(FGraphEventRef* An
 #if !UE_BUILD_SHIPPING
 			if (CVarRHICmdFlushOnQueueParallelSubmit.GetValueOnRenderThread())
 			{
+				CSV_SCOPED_TIMING_STAT(RHITFlushes, QueueParallelAsyncCommandListSubmit);
 				FRHICommandListExecutor::GetImmediateCommandList().ImmediateFlush(EImmediateFlushType::FlushRHIThread);
 			}
 #endif
@@ -1467,6 +1471,7 @@ void FRHICommandList::BeginScene()
 	{
 		// if we aren't running an RHIThread, there is no good reason to buffer this frame advance stuff and that complicates state management, so flush everything out now
 		QUICK_SCOPE_CYCLE_COUNTER(BeginScene_Flush);
+		CSV_SCOPED_TIMING_STAT(RHITFlushes, BeginScene);
 		FRHICommandListExecutor::GetImmediateCommandList().ImmediateFlush(EImmediateFlushType::FlushRHIThread);
 	}
 }
@@ -1483,6 +1488,7 @@ void FRHICommandList::EndScene()
 	{
 		// if we aren't running an RHIThread, there is no good reason to buffer this frame advance stuff and that complicates state management, so flush everything out now
 		QUICK_SCOPE_CYCLE_COUNTER(EndScene_Flush);
+		CSV_SCOPED_TIMING_STAT(RHITFlushes, EndScene);
 		FRHICommandListExecutor::GetImmediateCommandList().ImmediateFlush(EImmediateFlushType::FlushRHIThread);
 	}
 }
@@ -1501,6 +1507,7 @@ void FRHICommandList::BeginDrawingViewport(FRHIViewport* Viewport, FRHITexture*
 	{
 		// if we aren't running an RHIThread, there is no good reason to buffer this frame advance stuff and that complicates state management, so flush everything out now
 		QUICK_SCOPE_CYCLE_COUNTER(BeginDrawingViewport_Flush);
+		CSV_SCOPED_TIMING_STAT(RHITFlushes, BeginDrawingViewport);
 		FRHICommandListExecutor::GetImmediateCommandList().ImmediateFlush(EImmediateFlushType::FlushRHIThread);
 	}
 }
@@ -1554,6 +1561,7 @@ void FRHICommandList::BeginFrame()
 	{
 		// if we aren't running an RHIThread, there is no good reason to buffer this frame advance stuff and that complicates state management, so flush everything out now
 		QUICK_SCOPE_CYCLE_COUNTER(BeginFrame_Flush);
+		CSV_SCOPED_TIMING_STAT(RHITFlushes, BeginFrame);
 		FRHICommandListExecutor::GetImmediateCommandList().ImmediateFlush(EImmediateFlushType::FlushRHIThread);
 	}
 
@@ -1572,6 +1580,7 @@ void FRHICommandList::EndFrame()
 	{
 		// if we aren't running an RHIThread, there is no good reason to buffer this frame advance stuff and that complicates state management, so flush everything out now
 		QUICK_SCOPE_CYCLE_COUNTER(EndFrame_Flush);
+		CSV_SCOPED_TIMING_STAT(RHITFlushes, EndFrame);
 		FRHICommandListExecutor::GetImmediateCommandList().ImmediateFlush(EImmediateFlushType::FlushRHIThread);
 	}
 	else
@@ -1629,6 +1638,7 @@ FScopedCommandListWaitForTasks::~FScopedCommandListWaitForTasks()
 		else
 		{
 			QUICK_SCOPE_CYCLE_COUNTER(STAT_FScopedCommandListWaitForTasks_Flush);
+			CSV_SCOPED_TIMING_STAT(RHITFlushes, FScopedCommandListWaitForTasksDtor);
 			RHICmdList.ImmediateFlush(EImmediateFlushType::FlushRHIThread);
 		}
 	}
@@ -1662,12 +1672,14 @@ void FRHICommandListBase::WaitForDispatch()
 
 void FDynamicRHI::VirtualTextureSetFirstMipInMemory_RenderThread(class FRHICommandListImmediate& RHICmdList, FRHITexture2D* Texture, uint32 FirstMip)
 {
+	CSV_SCOPED_TIMING_STAT(RHITFlushes, VirtualTextureSetFirstMipInMemory_RenderThread);
 	RHICmdList.ImmediateFlush(EImmediateFlushType::FlushRHIThread);
 	GDynamicRHI->RHIVirtualTextureSetFirstMipInMemory(Texture, FirstMip);
 }
 
 void FDynamicRHI::VirtualTextureSetFirstMipVisible_RenderThread(class FRHICommandListImmediate& RHICmdList, FRHITexture2D* Texture, uint32 FirstMip)
 {
+	CSV_SCOPED_TIMING_STAT(RHITFlushes, VirtualTextureSetFirstMipVisible_RenderThread);
 	RHICmdList.ImmediateFlush(EImmediateFlushType::FlushRHIThread);
 	GDynamicRHI->RHIVirtualTextureSetFirstMipVisible(Texture, FirstMip);
 }
@@ -2003,6 +2015,7 @@ void* FDynamicRHI::LockVertexBuffer_RenderThread(class FRHICommandListImmediate&
 	{
 		{
 			QUICK_SCOPE_CYCLE_COUNTER(STAT_RHIMETHOD_LockVertexBuffer_Flush);
+			CSV_SCOPED_TIMING_STAT(RHITFlushes, LockVertexBuffer_RenderThread);
 			RHICmdList.ImmediateFlush(EImmediateFlushType::FlushRHIThread);
 		}
 		{
@@ -2030,6 +2043,7 @@ void FDynamicRHI::UnlockVertexBuffer_RenderThread(class FRHICommandListImmediate
 	{
 		{
 			QUICK_SCOPE_CYCLE_COUNTER(STAT_RHIMETHOD_UnlockVertexBuffer_Flush);
+			CSV_SCOPED_TIMING_STAT(RHITFlushes, UnlockVertexBuffer_RenderThread);
 			RHICmdList.ImmediateFlush(EImmediateFlushType::FlushRHIThread);
 		}
 		{
@@ -2062,6 +2076,7 @@ void* FDynamicRHI::LockIndexBuffer_RenderThread(class FRHICommandListImmediate&
 	{
 		{
 			QUICK_SCOPE_CYCLE_COUNTER(STAT_RHIMETHOD_LockIndexBuffer_Flush);
+			CSV_SCOPED_TIMING_STAT(RHITFlushes, LockIndexBuffer_RenderThread);
 			RHICmdList.ImmediateFlush(EImmediateFlushType::FlushRHIThread);
 		}
 		{
@@ -2089,6 +2104,7 @@ void FDynamicRHI::UnlockIndexBuffer_RenderThread(class FRHICommandListImmediate&
 	{
 		{
 			QUICK_SCOPE_CYCLE_COUNTER(STAT_RHIMETHOD_UnlockIndexBuffer_Flush);
+			CSV_SCOPED_TIMING_STAT(RHITFlushes, UnlockIndexBuffer_RenderThread);
 			RHICmdList.ImmediateFlush(EImmediateFlushType::FlushRHIThread);
 		}
 		{
@@ -2121,6 +2137,7 @@ void* FDynamicRHI::LockStructuredBuffer_RenderThread(class FRHICommandListImmedi
 	{
 		{
 			QUICK_SCOPE_CYCLE_COUNTER(STAT_RHIMETHOD_LockStructuredBuffer_Flush);
+			CSV_SCOPED_TIMING_STAT(RHITFlushes, LockStructuredBuffer_RenderThread);
 			RHICmdList.ImmediateFlush(EImmediateFlushType::FlushRHIThread);
 		}
 		{
@@ -2148,6 +2165,7 @@ void FDynamicRHI::UnlockStructuredBuffer_RenderThread(class FRHICommandListImmed
 	{
 		{
 			QUICK_SCOPE_CYCLE_COUNTER(STAT_RHIMETHOD_UnlockStructuredBuffer_Flush);
+			CSV_SCOPED_TIMING_STAT(RHITFlushes, UnlockStructuredBuffer_RenderThread);
 			RHICmdList.ImmediateFlush(EImmediateFlushType::FlushRHIThread);
 		}
 		{
@@ -2201,6 +2219,7 @@ void FDynamicRHI::UnlockStagingBuffer_RenderThread(class FRHICommandListImmediat
 FTexture2DRHIRef FDynamicRHI::AsyncReallocateTexture2D_RenderThread(class FRHICommandListImmediate& RHICmdList, FRHITexture2D* Texture2D, int32 NewMipCount, int32 NewSizeX, int32 NewSizeY, FThreadSafeCounter* RequestStatus)
 {
 	QUICK_SCOPE_CYCLE_COUNTER(STAT_RHIMETHOD_AsyncReallocateTexture2D_Flush);
+	CSV_SCOPED_TIMING_STAT(RHITFlushes, AsyncReallocateTexture2D_RenderThread);
 	RHICmdList.ImmediateFlush(EImmediateFlushType::FlushRHIThread);  
 	return GDynamicRHI->RHIAsyncReallocateTexture2D(Texture2D, NewMipCount, NewSizeX, NewSizeY, RequestStatus);
 }
@@ -2375,6 +2394,7 @@ void* FDynamicRHI::LockTexture2D_RenderThread(class FRHICommandListImmediate& RH
 	if (bNeedsDefaultRHIFlush) 
 	{
 		QUICK_SCOPE_CYCLE_COUNTER(STAT_RHIMETHOD_LockTexture2D_Flush);
+		CSV_SCOPED_TIMING_STAT(RHITFlushes, LockTexture2D_RenderThread);
 		RHICmdList.ImmediateFlush(EImmediateFlushType::FlushRHIThread);
 		return GDynamicRHI->RHILockTexture2D(Texture, MipIndex, LockMode, DestStride, bLockWithinMiptail);
 	}
@@ -2388,6 +2408,7 @@ void FDynamicRHI::UnlockTexture2D_RenderThread(class FRHICommandListImmediate& R
 	if (bNeedsDefaultRHIFlush)
 	{
 		QUICK_SCOPE_CYCLE_COUNTER(STAT_RHIMETHOD_UnlockTexture2D_Flush);
+		CSV_SCOPED_TIMING_STAT(RHITFlushes, UnlockTexture2D_RenderThread);
 		RHICmdList.ImmediateFlush(EImmediateFlushType::FlushRHIThread);
 		GDynamicRHI->RHIUnlockTexture2D(Texture, MipIndex, bLockWithinMiptail);
 		return;
@@ -2516,6 +2537,7 @@ FRenderQueryRHIRef FDynamicRHI::RHICreateRenderQuery_RenderThread(class FRHIComm
 void* FDynamicRHI::RHILockTextureCubeFace_RenderThread(class FRHICommandListImmediate& RHICmdList, FRHITextureCube* Texture, uint32 FaceIndex, uint32 ArrayIndex, uint32 MipIndex, EResourceLockMode LockMode, uint32& DestStride, bool bLockWithinMiptail)
 {
 	QUICK_SCOPE_CYCLE_COUNTER(STAT_RHIMETHOD_LockTextureCubeFace_Flush);
+	CSV_SCOPED_TIMING_STAT(RHITFlushes, RHILockTextureCubeFace_RenderThread);
 	RHICmdList.ImmediateFlush(EImmediateFlushType::FlushRHIThread);
 	return GDynamicRHI->RHILockTextureCubeFace(Texture, FaceIndex, ArrayIndex, MipIndex, LockMode, DestStride, bLockWithinMiptail);
 }
@@ -2523,6 +2545,7 @@ void* FDynamicRHI::RHILockTextureCubeFace_RenderThread(class FRHICommandListImme
 void FDynamicRHI::RHIUnlockTextureCubeFace_RenderThread(class FRHICommandListImmediate& RHICmdList, FRHITextureCube* Texture, uint32 FaceIndex, uint32 ArrayIndex, uint32 MipIndex, bool bLockWithinMiptail)
 {
 	QUICK_SCOPE_CYCLE_COUNTER(STAT_RHIMETHOD_UnlockTextureCubeFace_Flush);
+	CSV_SCOPED_TIMING_STAT(RHITFlushes, RHIUnlockTextureCubeFace_RenderThread);
 	RHICmdList.ImmediateFlush(EImmediateFlushType::FlushRHIThread);
 	GDynamicRHI->RHIUnlockTextureCubeFace(Texture, FaceIndex, ArrayIndex, MipIndex, bLockWithinMiptail);
 }
@@ -2549,6 +2572,7 @@ void FDynamicRHI::RHIUnmapStagingSurface_RenderThread(class FRHICommandListImmed
 void FDynamicRHI::RHIReadSurfaceFloatData_RenderThread(class FRHICommandListImmediate& RHICmdList, FRHITexture* Texture, FIntRect Rect, TArray<FFloat16Color>& OutData, ECubeFace CubeFace, int32 ArrayIndex, int32 MipIndex)
 {
 	QUICK_SCOPE_CYCLE_COUNTER(STAT_RHIMETHOD_ReadSurfaceFloatData_Flush);
+	CSV_SCOPED_TIMING_STAT(RHITFlushes, RHIReadSurfaceFloatData_RenderThread);
 	RHICmdList.ImmediateFlush(EImmediateFlushType::FlushRHIThread);
 	GDynamicRHI->RHIReadSurfaceFloatData(Texture, Rect, OutData, CubeFace, ArrayIndex, MipIndex);
 }
@@ -2560,6 +2584,7 @@ void FRHICommandListImmediate::UpdateTextureReference(FRHITextureReference* Text
 	{
 		{
 			QUICK_SCOPE_CYCLE_COUNTER(STAT_RHIMETHOD_UpdateTextureReference_FlushRHI);
+			CSV_SCOPED_TIMING_STAT(RHITFlushes, UpdateTextureReference);
 			ImmediateFlush(EImmediateFlushType::FlushRHIThread);
 		}
 		GetContext().RHIUpdateTextureReference(TextureRef, NewTexture);
diff --git a/Engine/Source/Runtime/RHI/Public/RHICommandList.h b/Engine/Source/Runtime/RHI/Public/RHICommandList.h
--- a/Engine/Source/Runtime/RHI/Public/RHICommandList.h
+++ b/Engine/Source/Runtime/RHI/Public/RHICommandList.h
@@ -25,7 +25,10 @@
 #include "HAL/IConsoleManager.h"
 #include "Async/TaskGraphInterfaces.h"
 #include "HAL/LowLevelMemTracker.h"
+#include "ProfilingDebugging/CsvProfiler.h"
 
+CSV_DECLARE_CATEGORY_EXTERN(RHITStalls);
+CSV_DECLARE_CATEGORY_EXTERN(RHITFlushes);
 
 // Set to 1 to capture the callstack for every RHI command. Cheap & memory efficient representation: Use the 
 // value in FRHICommand::StackFrames to get the pointer to the code (ie paste on a disassembly window)
diff --git a/Engine/Source/Runtime/RHI/Public/RHICommandList.inl b/Engine/Source/Runtime/RHI/Public/RHICommandList.inl
--- a/Engine/Source/Runtime/RHI/Public/RHICommandList.inl
+++ b/Engine/Source/Runtime/RHI/Public/RHICommandList.inl
@@ -92,6 +92,7 @@ FORCEINLINE_DEBUGGABLE void FRHICommandListImmediate::ImmediateFlush(EImmediateF
 		break;
 	case EImmediateFlushType::FlushRHIThread:
 		{
+			CSV_SCOPED_TIMING_STAT(RHITFlushes, FlushRHIThreadTotal);
 			if (HasCommands())
 			{
 				GRHICommandList.ExecuteList(*this);
@@ -107,6 +108,7 @@ FORCEINLINE_DEBUGGABLE void FRHICommandListImmediate::ImmediateFlush(EImmediateF
 	case EImmediateFlushType::FlushRHIThreadFlushResources:
 	case EImmediateFlushType::FlushRHIThreadFlushResourcesFlushDeferredDeletes:
 		{
+			CSV_SCOPED_TIMING_STAT(RHITFlushes, FlushRHIThreadFlushResourcesTotal);
 			if (HasCommands())
 			{
 				GRHICommandList.ExecuteList(*this);
