Remove OpenGL4 from the DefaultRHI list

diff --git a/Engine/Source/Developer/Windows/WindowsTargetPlatform/Classes/WindowsTargetSettings.h b/Engine/Source/Developer/Windows/WindowsTargetPlatform/Classes/WindowsTargetSettings.h
--- a/Engine/Source/Developer/Windows/WindowsTargetPlatform/Classes/WindowsTargetSettings.h
+++ b/Engine/Source/Developer/Windows/WindowsTargetPlatform/Classes/WindowsTargetSettings.h
@@ -32,7 +32,6 @@ enum class EDefaultGraphicsRHI : uint8
 	DefaultGraphicsRHI_DX11 = 1 UMETA(DisplayName = "DirectX 11"),
 	DefaultGraphicsRHI_DX12 = 2 UMETA(DisplayName = "DirectX 12"),
 	DefaultGraphicsRHI_Vulkan = 3 UMETA(DisplayName = "Vulkan"),
-	DefaultGraphicsRHI_OpenGL = 4 UMETA(DisplayName = "OpenGL"),
 };
 
 
diff --git a/Engine/Source/Runtime/RHI/Private/Windows/WindowsDynamicRHI.cpp b/Engine/Source/Runtime/RHI/Private/Windows/WindowsDynamicRHI.cpp
--- a/Engine/Source/Runtime/RHI/Private/Windows/WindowsDynamicRHI.cpp
+++ b/Engine/Source/Runtime/RHI/Private/Windows/WindowsDynamicRHI.cpp
@@ -63,7 +63,6 @@ static IDynamicRHIModule* LoadDynamicRHIModule(ERHIFeatureLevel::Type& DesiredFe
 			static FString NAME_DX11(TEXT("DefaultGraphicsRHI_DX11"));
 			static FString NAME_DX12(TEXT("DefaultGraphicsRHI_DX12"));
 			static FString NAME_VULKAN(TEXT("DefaultGraphicsRHI_Vulkan"));
-			static FString NAME_OPENGL(TEXT("DefaultGraphicsRHI_OpenGL"));
 			if(DefaultGraphicsRHI == NAME_DX11)
 			{
 				bForceD3D11 = true;
@@ -76,10 +75,6 @@ static IDynamicRHIModule* LoadDynamicRHIModule(ERHIFeatureLevel::Type& DesiredFe
 			{
 				bForceVulkan = true;
 			}
-			else if (DefaultGraphicsRHI == NAME_OPENGL)
-			{
-				bForceOpenGL = true;
-			}
 		}
 	}
 
