Disable geometry shaders on Lumin
This fixes a black screen when mobile rendering is disabled

diff --git a/Engine/Source/Runtime/RHI/Public/RHIDefinitions.h b/Engine/Source/Runtime/RHI/Public/RHIDefinitions.h
--- a/Engine/Source/Runtime/RHI/Public/RHIDefinitions.h
+++ b/Engine/Source/Runtime/RHI/Public/RHIDefinitions.h
@@ -1165,7 +1165,7 @@ inline bool RHISupportsComputeShaders(const EShaderPlatform Platform)
 
 inline bool RHISupportsGeometryShaders(const EShaderPlatform Platform)
 {
-	return IsFeatureLevelSupported(Platform, ERHIFeatureLevel::SM4) && !IsMetalPlatform(Platform) && !IsVulkanMobilePlatform(Platform);
+	return IsFeatureLevelSupported(Platform, ERHIFeatureLevel::SM4) && !IsMetalPlatform(Platform) && !IsVulkanMobilePlatform(Platform) && Platform != SP_VULKAN_SM5_LUMIN;
 }
 
 /** Return true if and only if the GPU support rendering to volume textures (2D Array, 3D) is guaranteed supported for a target platform.
diff --git a/Engine/Source/Runtime/VulkanRHI/Public/VulkanCommon.h b/Engine/Source/Runtime/VulkanRHI/Public/VulkanCommon.h
--- a/Engine/Source/Runtime/VulkanRHI/Public/VulkanCommon.h
+++ b/Engine/Source/Runtime/VulkanRHI/Public/VulkanCommon.h
@@ -9,7 +9,7 @@
 #include "RHIDefinitions.h"
 
 #ifndef VULKAN_SUPPORTS_GEOMETRY_SHADERS
-	#define VULKAN_SUPPORTS_GEOMETRY_SHADERS					(!(PLATFORM_ANDROID) || PLATFORM_LUMIN || PLATFORM_LUMINGL4) && PLATFORM_SUPPORTS_GEOMETRY_SHADERS
+	#define VULKAN_SUPPORTS_GEOMETRY_SHADERS					(!(PLATFORM_ANDROID || PLATFORM_LUMIN || PLATFORM_LUMINGL4)) && PLATFORM_SUPPORTS_GEOMETRY_SHADERS
 #endif
 
 // This defines controls shader generation (so will cause a format rebuild)
