Fixed an issue with ray tracing pipeline state cache that could result in pipelines created on current frame to be immediately evicted, resulting in use-after-free error.

diff --git a/Engine/Source/Runtime/RHI/Private/PipelineStateCache.cpp b/Engine/Source/Runtime/RHI/Private/PipelineStateCache.cpp
--- a/Engine/Source/Runtime/RHI/Private/PipelineStateCache.cpp
+++ b/Engine/Source/Runtime/RHI/Private/PipelineStateCache.cpp
@@ -647,6 +647,7 @@ public:
 	{
 		FScopeLock ScopeLock(&CriticalSection);
 		Cache.Add(Initializer, State);
+		State->AddHit();
 	}
 
 	void Shutdown()
