Re-enable geometry shaders on Lumin
Volumetric textures are still a problem so these are disabled instead

diff --git a/Engine/Source/Runtime/Core/Public/Lumin/LuminPlatform.h b/Engine/Source/Runtime/Core/Public/Lumin/LuminPlatform.h
--- a/Engine/Source/Runtime/Core/Public/Lumin/LuminPlatform.h
+++ b/Engine/Source/Runtime/Core/Public/Lumin/LuminPlatform.h
@@ -4,4 +4,6 @@
 
 #include "Android/AndroidPlatform.h"
 
+#undef PLATFORM_SUPPORTS_GEOMETRY_SHADERS
+#define PLATFORM_SUPPORTS_GEOMETRY_SHADERS			1
 #define PLATFORM_GLOBAL_LOG_CATEGORY_ALT LogLumin
diff --git a/Engine/Source/Runtime/RHI/Public/RHIDefinitions.h b/Engine/Source/Runtime/RHI/Public/RHIDefinitions.h
--- a/Engine/Source/Runtime/RHI/Public/RHIDefinitions.h
+++ b/Engine/Source/Runtime/RHI/Public/RHIDefinitions.h
@@ -1165,7 +1165,7 @@ inline bool RHISupportsComputeShaders(const EShaderPlatform Platform)
 
 inline bool RHISupportsGeometryShaders(const EShaderPlatform Platform)
 {
-	return IsFeatureLevelSupported(Platform, ERHIFeatureLevel::SM4) && !IsMetalPlatform(Platform) && !IsVulkanMobilePlatform(Platform) && Platform != SP_VULKAN_SM5_LUMIN;
+	return IsFeatureLevelSupported(Platform, ERHIFeatureLevel::SM4) && !IsMetalPlatform(Platform) && !IsVulkanMobilePlatform(Platform);
 }
 
 /** Return true if and only if the GPU support rendering to volume textures (2D Array, 3D) is guaranteed supported for a target platform.
@@ -1175,7 +1175,8 @@ inline bool RHIVolumeTextureRenderingSupportGuaranteed(const EShaderPlatform Pla
 {
 	return IsFeatureLevelSupported(Platform, ERHIFeatureLevel::SM4)
 		&& !IsMetalPlatform(Platform)		// EMetalFeaturesLayeredRendering supports needs to be checked at runtime
-		&& !IsOpenGLPlatform(Platform);		// Apparently, some OpenGL 3.3 cards support SM4 but can't render to volume textures 
+		&& !IsOpenGLPlatform(Platform)		// Apparently, some OpenGL 3.3 cards support SM4 but can't render to volume textures 
+		&& Platform != SP_VULKAN_SM5_LUMIN;
 }
 
 inline bool RHIHasTiledGPU(const EShaderPlatform Platform)
diff --git a/Engine/Source/Runtime/VulkanRHI/Public/VulkanCommon.h b/Engine/Source/Runtime/VulkanRHI/Public/VulkanCommon.h
--- a/Engine/Source/Runtime/VulkanRHI/Public/VulkanCommon.h
+++ b/Engine/Source/Runtime/VulkanRHI/Public/VulkanCommon.h
@@ -9,7 +9,7 @@
 #include "RHIDefinitions.h"
 
 #ifndef VULKAN_SUPPORTS_GEOMETRY_SHADERS
-	#define VULKAN_SUPPORTS_GEOMETRY_SHADERS					(!(PLATFORM_ANDROID || PLATFORM_LUMIN || PLATFORM_LUMINGL4)) && PLATFORM_SUPPORTS_GEOMETRY_SHADERS
+	#define VULKAN_SUPPORTS_GEOMETRY_SHADERS					(!PLATFORM_ANDROID || PLATFORM_LUMIN || PLATFORM_LUMINGL4) && PLATFORM_SUPPORTS_GEOMETRY_SHADERS
 #endif
 
 // This defines controls shader generation (so will cause a format rebuild)
