CIS fix

diff --git a/Engine/Source/Runtime/RHI/Private/PipelineFileCache.cpp b/Engine/Source/Runtime/RHI/Private/PipelineFileCache.cpp
--- a/Engine/Source/Runtime/RHI/Private/PipelineFileCache.cpp
+++ b/Engine/Source/Runtime/RHI/Private/PipelineFileCache.cpp
@@ -2484,10 +2484,10 @@ void FPipelineFileCache::ClearOSPipelineCache()
 {
 	if (CVarClearOSPSOFileCache.GetValueOnAnyThread() > 0)
 	{
-		SCOPED_AUTORELEASE_POOL;
-
 		// clear the PSO cache on IOS if the executable is newer
 #if PLATFORM_IOS
+		SCOPED_AUTORELEASE_POOL;
+
 		static FString ExecutablePath = FString([[NSBundle mainBundle] bundlePath]) + TEXT("/") + FPlatformProcess::ExecutableName();
 		struct stat FileInfo;
 		if(stat(TCHAR_TO_UTF8(*ExecutablePath), &FileInfo) != -1)
@@ -2516,6 +2516,8 @@ void FPipelineFileCache::ClearOSPipelineCache()
 #elif PLATFORM_MAC && (UE_BUILD_TEST || UE_BUILD_SHIPPING)
 		if (!FPlatformProcess::IsSandboxedApplication())
 		{
+			SCOPED_AUTORELEASE_POOL;
+
 			static FString ExecutablePath = FString([[NSBundle mainBundle] executablePath]);
 			struct stat FileInfo;
 			if (stat(TCHAR_TO_UTF8(*ExecutablePath), &FileInfo) != -1)
