Adding a clear color for future ARKit hotfix.

diff --git a/Engine/Source/Runtime/RHI/Private/RHI.cpp b/Engine/Source/Runtime/RHI/Private/RHI.cpp
--- a/Engine/Source/Runtime/RHI/Private/RHI.cpp
+++ b/Engine/Source/Runtime/RHI/Private/RHI.cpp
@@ -107,6 +107,7 @@ static FAutoConsoleCommandWithOutputDevice GDumpRHIMemoryCmd(
 //Static init order is undefined and you will likely end up with bad values on some platforms.
 const FClearValueBinding FClearValueBinding::None(EClearBinding::ENoneBound);
 const FClearValueBinding FClearValueBinding::Black(FLinearColor(0.0f, 0.0f, 0.0f, 1.0f));
+const FClearValueBinding FClearValueBinding::BlackMaxAlpha(FLinearColor(0.0f, 0.0f, 0.0f, FLT_MAX));
 const FClearValueBinding FClearValueBinding::White(FLinearColor(1.0f, 1.0f, 1.0f, 1.0f));
 const FClearValueBinding FClearValueBinding::Transparent(FLinearColor(0.0f, 0.0f, 0.0f, 0.0f));
 const FClearValueBinding FClearValueBinding::DepthOne(1.0f, 0);
diff --git a/Engine/Source/Runtime/RHI/Public/RHI.h b/Engine/Source/Runtime/RHI/Public/RHI.h
--- a/Engine/Source/Runtime/RHI/Public/RHI.h
+++ b/Engine/Source/Runtime/RHI/Public/RHI.h
@@ -1186,6 +1186,7 @@ struct FClearValueBinding
 	// common clear values
 	static RHI_API const FClearValueBinding None;
 	static RHI_API const FClearValueBinding Black;
+	static RHI_API const FClearValueBinding BlackMaxAlpha;
 	static RHI_API const FClearValueBinding White;
 	static RHI_API const FClearValueBinding Transparent;
 	static RHI_API const FClearValueBinding DepthOne;
