Give better warnings as we no longer fallback to opengl4

diff --git a/Engine/Source/Runtime/RHI/Private/Linux/LinuxDynamicRHI.cpp b/Engine/Source/Runtime/RHI/Private/Linux/LinuxDynamicRHI.cpp
--- a/Engine/Source/Runtime/RHI/Private/Linux/LinuxDynamicRHI.cpp
+++ b/Engine/Source/Runtime/RHI/Private/Linux/LinuxDynamicRHI.cpp
@@ -134,7 +134,18 @@ FDynamicRHI* PlatformCreateDynamicRHI()
 			}
 			else
 			{
-				FMessageDialog::Open(EAppMsgType::Ok, NSLOCTEXT("LinuxDynamicRHI", "NoTargetedRHI", "The project does not target Vulkan or OpenGL RHIs, check project settings or pass -nullrhi."));
+				if (bVulkanFailed)
+				{
+					FMessageDialog::Open(EAppMsgType::Ok, NSLOCTEXT("LinuxDynamicRHI", "NoVulkanDriver", "Failed to load Vulkan Driver which is required to run the engine.\nThe engine no longer fallbacks to OpenGL4 which has been deprecated."));
+				}
+				else if (bOpenGLFailed)
+				{
+					FMessageDialog::Open(EAppMsgType::Ok, NSLOCTEXT("LinuxDynamicRHI", "NoOpenGLDriver", "Failed to load OpenGL Driver which is required to run the engine.\nOpenGL4 has been deprecated and should use Vulkan."));
+				}
+				else
+				{
+					FMessageDialog::Open(EAppMsgType::Ok, NSLOCTEXT("LinuxDynamicRHI", "NoTargetedRHI", "The project does not target Vulkan or OpenGL RHIs, check project settings or pass -nullrhi."));
+				}
 			}
 
 			FPlatformMisc::RequestExitWithStatus(true, 1);
