User setting for using DX12 RHI

diff --git a/Engine/Source/Runtime/RHI/Private/Windows/WindowsDynamicRHI.cpp b/Engine/Source/Runtime/RHI/Private/Windows/WindowsDynamicRHI.cpp
--- a/Engine/Source/Runtime/RHI/Private/Windows/WindowsDynamicRHI.cpp
+++ b/Engine/Source/Runtime/RHI/Private/Windows/WindowsDynamicRHI.cpp
@@ -10,24 +10,32 @@ static const TCHAR* GLoadedRHIModuleName;
 
 static bool ShouldPreferD3D12()
 {
-	// Disabled until D3D12RHI is ready
+	if (!GIsEditor)
+	{
+		bool bPreferD3D12 = false;
+		if (GConfig->GetBool(TEXT("D3DRHIPreference"), TEXT("bPreferD3D12InGame"), bPreferD3D12, GGameUserSettingsIni))
+		{
+			return bPreferD3D12;
+		}
+	}
+
 #if 0
 	bool bPreferD3D12 = false;
 	if (GIsEditor)
 	{
-		GConfig->GetBool(TEXT("D3DRHIPerference"), TEXT("bPreferD3D12InEditor"), bPreferD3D12, GEngineIni);
+		GConfig->GetBool(TEXT("D3DRHIPreference"), TEXT("bPreferD3D12InEditor"), bPreferD3D12, GEngineIni);
 	}
 	else
 	{
-		GConfig->GetBool(TEXT("D3DRHIPerference"), TEXT("bPreferD3D12InGame"), bPreferD3D12, GEngineIni);
+		GConfig->GetBool(TEXT("D3DRHIPreference"), TEXT("bPreferD3D12InGame"), bPreferD3D12, GEngineIni);
 	}
 
 	int32 MinNumCPUCores = 0;
-	GConfig->GetInt(TEXT("D3DRHIPerference"), TEXT("con.MinNumCPUCores"), MinNumCPUCores, GEngineIni);
+	GConfig->GetInt(TEXT("D3DRHIPreference"), TEXT("con.MinNumCPUCores"), MinNumCPUCores, GEngineIni);
 	const bool bHasEnoughCPUCores = FPlatformMisc::NumberOfCoresIncludingHyperthreads() >= MinNumCPUCores;
 
 	int32 MinPhysicalMemGB = 0;
-	GConfig->GetInt(TEXT("D3DRHIPerference"), TEXT("con.MinPhysicalMemGB"), MinPhysicalMemGB, GEngineIni);
+	GConfig->GetInt(TEXT("D3DRHIPreference"), TEXT("con.MinPhysicalMemGB"), MinPhysicalMemGB, GEngineIni);
 	const bool bHasEnoughMem = FPlatformMemory::GetConstants().TotalPhysical >= MinPhysicalMemGB * (1llu << 30);
 
 	return bPreferD3D12 && bHasEnoughCPUCores && bHasEnoughMem;
