Enable -psocache by default for DX12

diff --git a/Engine/Source/Runtime/RHI/Private/Windows/WindowsDynamicRHI.cpp b/Engine/Source/Runtime/RHI/Private/Windows/WindowsDynamicRHI.cpp
--- a/Engine/Source/Runtime/RHI/Private/Windows/WindowsDynamicRHI.cpp
+++ b/Engine/Source/Runtime/RHI/Private/Windows/WindowsDynamicRHI.cpp
@@ -180,6 +180,12 @@ static IDynamicRHIModule* LoadDynamicRHIModule(ERHIFeatureLevel::Type& DesiredFe
 		LoadedRHIModuleName = TEXT("D3D12RHI");
 		DynamicRHIModule = FModuleManager::LoadModulePtr<IDynamicRHIModule>(LoadedRHIModuleName);
 
+#if !WITH_EDITOR
+		// Enable -psocache by default on DX12. Since RHI is selected at runtime we can't set this at compile time with PIPELINE_CACHE_DEFAULT_ENABLED.
+		auto CVarPSOFileCacheEnabled = IConsoleManager::Get().FindTConsoleVariableDataInt(TEXT("r.ShaderPipelineCache.Enabled"));
+		*CVarPSOFileCacheEnabled = 1;
+#endif
+
 		if (!DynamicRHIModule || !DynamicRHIModule->IsSupported())
 		{
 			if (bForceD3D12)
