Fix shadowing warning
#rnx

diff --git a/Engine/Source/Runtime/RHI/Private/Windows/WindowsDynamicRHI.cpp b/Engine/Source/Runtime/RHI/Private/Windows/WindowsDynamicRHI.cpp
--- a/Engine/Source/Runtime/RHI/Private/Windows/WindowsDynamicRHI.cpp
+++ b/Engine/Source/Runtime/RHI/Private/Windows/WindowsDynamicRHI.cpp
@@ -182,8 +182,8 @@ static IDynamicRHIModule* LoadDynamicRHIModule(ERHIFeatureLevel::Type& DesiredFe
 
 #if !WITH_EDITOR
 		// Enable -psocache by default on DX12. Since RHI is selected at runtime we can't set this at compile time with PIPELINE_CACHE_DEFAULT_ENABLED.
-		auto CVarPSOFileCacheEnabled = IConsoleManager::Get().FindTConsoleVariableDataInt(TEXT("r.ShaderPipelineCache.Enabled"));
-		*CVarPSOFileCacheEnabled = 1;
+		auto PSOFileCacheEnabledCVar = IConsoleManager::Get().FindTConsoleVariableDataInt(TEXT("r.ShaderPipelineCache.Enabled"));
+		*PSOFileCacheEnabledCVar = 1;
 #endif
 
 		if (!DynamicRHIModule || !DynamicRHIModule->IsSupported())
