Remove RHISupportsDrawIndirect. It returns incorrectly false on X1, and anyway we are just checking for SM5 in existing code when doing indirect dispatches

diff --git a/Engine/Source/Runtime/RHI/Public/RHIDefinitions.h b/Engine/Source/Runtime/RHI/Public/RHIDefinitions.h
--- a/Engine/Source/Runtime/RHI/Public/RHIDefinitions.h
+++ b/Engine/Source/Runtime/RHI/Public/RHIDefinitions.h
@@ -1196,12 +1196,6 @@ inline bool RHISupportsMobileMultiView(const EShaderPlatform Platform)
 		|| FDataDrivenShaderPlatformInfo::GetInfo(Platform).bSupportsMobileMultiView;
 }
 
-inline bool RHISupportsDrawIndirect(const EShaderPlatform Platform)
-{
-	return Platform == EShaderPlatform::SP_METAL_SM5 || Platform == EShaderPlatform::SP_PCD3D_SM5 || IsVulkanSM5Platform(Platform) || Platform == EShaderPlatform::SP_PS4
-		|| FDataDrivenShaderPlatformInfo::GetInfo(Platform).bSupportsDrawIndirect;
-}
-
 inline bool RHISupportsNativeShaderLibraries(const EShaderPlatform Platform)
 {
 	return IsMetalPlatform(Platform);
