RHI support for UAV with format alias
GRHISupportsUAVFormatAliasing specifies supported platforms
Only implemented on PS4 so far

diff --git a/Engine/Source/Runtime/RHI/Private/DynamicRHI.cpp b/Engine/Source/Runtime/RHI/Private/DynamicRHI.cpp
--- a/Engine/Source/Runtime/RHI/Private/DynamicRHI.cpp
+++ b/Engine/Source/Runtime/RHI/Private/DynamicRHI.cpp
@@ -345,6 +345,12 @@ void FDynamicRHI::RHITransferVertexBufferUnderlyingResource(FRHIVertexBuffer* De
 	UE_LOG(LogRHI, Fatal, TEXT("RHITransferVertexBufferUnderlyingResource isn't implemented for the current RHI"));
 }
 
+FUnorderedAccessViewRHIRef FDynamicRHI::RHICreateUnorderedAccessView(FRHITexture* Texture, uint32 MipLevel, uint8 Format)
+{
+	UE_LOG(LogRHI, Fatal, TEXT("RHICreateUnorderedAccessView with Format parameter isn't implemented for the current RHI"));
+	return RHICreateUnorderedAccessView(Texture, MipLevel);
+}
+
 void FDynamicRHI::RHIUpdateShaderResourceView(FRHIShaderResourceView* SRV, FRHIVertexBuffer* VertexBuffer, uint32 Stride, uint8 Format)
 {
 	UE_LOG(LogRHI, Fatal, TEXT("RHIUpdateShaderResourceView isn't implemented for the current RHI"));
diff --git a/Engine/Source/Runtime/RHI/Private/RHI.cpp b/Engine/Source/Runtime/RHI/Private/RHI.cpp
--- a/Engine/Source/Runtime/RHI/Private/RHI.cpp
+++ b/Engine/Source/Runtime/RHI/Private/RHI.cpp
@@ -523,6 +523,7 @@ bool GSupportsParallelOcclusionQueries = false;
 bool GSupportsRenderTargetWriteMask = false;
 bool GSupportsTransientResourceAliasing = false;
 bool GRHIRequiresRenderTargetForPixelShaderUAVs = false;
+bool GRHISupportsUAVFormatAliasing = false;
 
 bool GRHISupportsMSAADepthSampleAccess = false;
 bool GRHISupportsResolveCubemapFaces = false;
diff --git a/Engine/Source/Runtime/RHI/Private/RHICommandList.cpp b/Engine/Source/Runtime/RHI/Private/RHICommandList.cpp
--- a/Engine/Source/Runtime/RHI/Private/RHICommandList.cpp
+++ b/Engine/Source/Runtime/RHI/Private/RHICommandList.cpp
@@ -2446,6 +2446,13 @@ FUnorderedAccessViewRHIRef FDynamicRHI::RHICreateUnorderedAccessView_RenderThrea
 	return GDynamicRHI->RHICreateUnorderedAccessView(Texture, MipLevel);
 }
 
+FUnorderedAccessViewRHIRef FDynamicRHI::RHICreateUnorderedAccessView_RenderThread(class FRHICommandListImmediate& RHICmdList, FRHITexture* Texture, uint32 MipLevel, uint8 Format)
+{
+	CSV_SCOPED_TIMING_STAT(RHITStalls, RHICreateUnorderedAccessView_RenderThread);
+	FScopedRHIThreadStaller StallRHIThread(RHICmdList);
+	return GDynamicRHI->RHICreateUnorderedAccessView(Texture, MipLevel, Format);
+}
+
 FUnorderedAccessViewRHIRef FDynamicRHI::RHICreateUnorderedAccessView_RenderThread(class FRHICommandListImmediate& RHICmdList, FRHIVertexBuffer* VertexBuffer, uint8 Format)
 {
 	CSV_SCOPED_TIMING_STAT(RHITStalls, RHICreateUnorderedAccessView_RenderThread);
diff --git a/Engine/Source/Runtime/RHI/Public/DynamicRHI.h b/Engine/Source/Runtime/RHI/Public/DynamicRHI.h
--- a/Engine/Source/Runtime/RHI/Public/DynamicRHI.h
+++ b/Engine/Source/Runtime/RHI/Public/DynamicRHI.h
@@ -483,6 +483,10 @@ public:
 	// FlushType: Wait RHI Thread
 	virtual FUnorderedAccessViewRHIRef RHICreateUnorderedAccessView(FRHITexture* Texture, uint32 MipLevel) = 0;
 
+	/** Creates an unordered access view of the given texture. */
+	// FlushType: Wait RHI Thread
+	virtual FUnorderedAccessViewRHIRef RHICreateUnorderedAccessView(FRHITexture* Texture, uint32 MipLevel, uint8 Format);
+
 	/** Creates an unordered access view of the given vertex buffer. */
 	// FlushType: Wait RHI Thread
 	virtual FUnorderedAccessViewRHIRef RHICreateUnorderedAccessView(FRHIVertexBuffer* VertexBuffer, uint8 Format) = 0;
@@ -1151,6 +1155,7 @@ public:
 	virtual FTexture3DRHIRef RHICreateTexture3D_RenderThread(class FRHICommandListImmediate& RHICmdList, uint32 SizeX, uint32 SizeY, uint32 SizeZ, uint8 Format, uint32 NumMips, uint32 Flags, FRHIResourceCreateInfo& CreateInfo);
 	virtual FUnorderedAccessViewRHIRef RHICreateUnorderedAccessView_RenderThread(class FRHICommandListImmediate& RHICmdList, FRHIStructuredBuffer* StructuredBuffer, bool bUseUAVCounter, bool bAppendBuffer);
 	virtual FUnorderedAccessViewRHIRef RHICreateUnorderedAccessView_RenderThread(class FRHICommandListImmediate& RHICmdList, FRHITexture* Texture, uint32 MipLevel);
+	virtual FUnorderedAccessViewRHIRef RHICreateUnorderedAccessView_RenderThread(class FRHICommandListImmediate& RHICmdList, FRHITexture* Texture, uint32 MipLevel, uint8 Format);
 	virtual FUnorderedAccessViewRHIRef RHICreateUnorderedAccessView_RenderThread(class FRHICommandListImmediate& RHICmdList, FRHIVertexBuffer* VertexBuffer, uint8 Format);
 	virtual FUnorderedAccessViewRHIRef RHICreateUnorderedAccessView_RenderThread(class FRHICommandListImmediate& RHICmdList, FRHIIndexBuffer* IndexBuffer, uint8 Format);
 	virtual FShaderResourceViewRHIRef RHICreateShaderResourceView_RenderThread(class FRHICommandListImmediate& RHICmdList, FRHITexture* Texture, const FRHITextureSRVCreateInfo& CreateInfo);
diff --git a/Engine/Source/Runtime/RHI/Public/RHI.h b/Engine/Source/Runtime/RHI/Public/RHI.h
--- a/Engine/Source/Runtime/RHI/Public/RHI.h
+++ b/Engine/Source/Runtime/RHI/Public/RHI.h
@@ -367,6 +367,9 @@ extern RHI_API bool GSupportsTransientResourceAliasing;
 /** true if the RHI requires a valid RT bound during UAV scatter operation inside the pixel shader */
 extern RHI_API bool GRHIRequiresRenderTargetForPixelShaderUAVs;
 
+/** true if the RHI supports unordered access view format aliasing */
+extern RHI_API bool GRHISupportsUAVFormatAliasing;
+
 /** The minimum Z value in clip space for the RHI. */
 extern RHI_API float GMinClipZ;
 
diff --git a/Engine/Source/Runtime/RHI/Public/RHICommandList.h b/Engine/Source/Runtime/RHI/Public/RHICommandList.h
--- a/Engine/Source/Runtime/RHI/Public/RHICommandList.h
+++ b/Engine/Source/Runtime/RHI/Public/RHICommandList.h
@@ -3873,7 +3873,13 @@ public:
 		LLM_SCOPE(ELLMTag::RHIMisc);
 		return GDynamicRHI->RHICreateUnorderedAccessView_RenderThread(*this, Texture, MipLevel);
 	}
-	
+
+	FORCEINLINE FUnorderedAccessViewRHIRef CreateUnorderedAccessView(FRHITexture* Texture, uint32 MipLevel, uint8 Format)
+	{
+		LLM_SCOPE(ELLMTag::RHIMisc);
+		return GDynamicRHI->RHICreateUnorderedAccessView_RenderThread(*this, Texture, MipLevel, Format);
+	}
+
 	FORCEINLINE FUnorderedAccessViewRHIRef CreateUnorderedAccessView(FRHIVertexBuffer* VertexBuffer, uint8 Format)
 	{
 		LLM_SCOPE(ELLMTag::RHIMisc);
@@ -4781,6 +4787,11 @@ FORCEINLINE FUnorderedAccessViewRHIRef RHICreateUnorderedAccessView(FRHITexture*
 	return FRHICommandListExecutor::GetImmediateCommandList().CreateUnorderedAccessView(Texture, MipLevel);
 }
 
+FORCEINLINE FUnorderedAccessViewRHIRef RHICreateUnorderedAccessView(FRHITexture* Texture, uint32 MipLevel, uint8 Format)
+{
+	return FRHICommandListExecutor::GetImmediateCommandList().CreateUnorderedAccessView(Texture, MipLevel, Format);
+}
+
 FORCEINLINE FUnorderedAccessViewRHIRef RHICreateUnorderedAccessView(FRHIVertexBuffer* VertexBuffer, uint8 Format)
 {
 	return FRHICommandListExecutor::GetImmediateCommandList().CreateUnorderedAccessView(VertexBuffer, Format);
