Copy up of Unreal Insights from //Tasks/UE4/Dev-Core-Trace.

diff --git a/Engine/Source/Runtime/RHI/Public/GpuProfilerTrace.h b/Engine/Source/Runtime/RHI/Public/GpuProfilerTrace.h
new file mode 100644
--- /dev/null
+++ b/Engine/Source/Runtime/RHI/Public/GpuProfilerTrace.h
@@ -0,0 +1,81 @@
+#pragma once
+
+#include "CoreTypes.h"
+
+#if !IS_PROGRAM && !UE_BUILD_SHIPPING
+#define GPUPROFILERTRACE_ENABLED 0
+#else
+#define GPUPROFILERTRACE_ENABLED 0
+#endif
+
+#if GPUPROFILERTRACE_ENABLED
+
+struct FGpuProfilerTrace
+{
+	struct FEventType
+	{
+		RHI_API FEventType(const TCHAR* Name);
+	};
+
+	RHI_API static void BeginFrame();
+	RHI_API static void BeginEvent(const FEventType* EventType, uint32 FrameNumber, uint64 TimestampMicroseconds);
+	RHI_API static void EndEvent(uint64 TimestampMicroseconds);
+	RHI_API static void EndFrame();
+
+private:
+	enum
+	{
+		MaxEventBufferSize = 32768
+	};
+
+	struct FFrame
+	{
+		uint64 TimestampBase;
+		uint64 LastTimestamp;
+		uint32 RenderingFrameNumber;
+		uint16 EventBufferSize;
+		uint8 EventBuffer[MaxEventBufferSize];
+	};
+	static FFrame CurrentFrame;
+};
+
+#define TRACE_GPUPROFILER_DEFINE_EVENT_TYPE(Name) \
+	FGpuProfilerTrace::FEventType PREPROCESSOR_JOIN(__GGpuProfilerEventType, Name)(TEXT(#Name));
+
+#define TRACE_GPUPROFILER_DECLARE_EVENT_TYPE_EXTERN(Name) \
+	extern FGpuProfilerTrace::FEventType PREPROCESSOR_JOIN(__GGpuProfilerEventType, Name);
+
+#define TRACE_GPUPROFILER_EVENT_TYPE(Name) \
+	&PREPROCESSOR_JOIN(__GGpuProfilerEventType, Name)
+
+#define TRACE_GPUPROFILER_BEGIN_FRAME() \
+	FGpuProfilerTrace::BeginFrame();
+
+#define TRACE_GPUPROFILER_BEGIN_EVENT(EventType, FrameNumber, TimestampMicroseconds) \
+	FGpuProfilerTrace::BeginEvent(EventType, FrameNumber, TimestampMicroseconds);
+
+#define TRACE_GPUPROFILER_END_EVENT(TimestampMicroseconds) \
+	FGpuProfilerTrace::EndEvent(TimestampMicroseconds);
+
+#define TRACE_GPUPROFILER_END_FRAME() \
+	FGpuProfilerTrace::EndFrame();
+
+#else
+
+struct FGpuProfilerTrace
+{
+	struct FEventType
+	{
+		FEventType(const TCHAR* Name) {};
+	};
+};
+
+#define TRACE_GPUPROFILER_DEFINE_EVENT_TYPE(...)
+#define TRACE_GPUPROFILER_DECLARE_EVENT_TYPE_EXTERN(...)
+#define TRACE_GPUPROFILER_EVENT_TYPE(...) nullptr
+#define TRACE_GPUPROFILER_BEGIN_FRAME(...)
+#define TRACE_GPUPROFILER_BEGIN_EVENT(...)
+#define TRACE_GPUPROFILER_END_EVENT(...)
+#define TRACE_GPUPROFILER_END_FRAME(...)
+
+#endif
diff --git a/Engine/Source/Runtime/RHI/RHI.Build.cs b/Engine/Source/Runtime/RHI/RHI.Build.cs
--- a/Engine/Source/Runtime/RHI/RHI.Build.cs
+++ b/Engine/Source/Runtime/RHI/RHI.Build.cs
@@ -9,6 +9,7 @@ public class RHI : ModuleRules
 	{
 		PrivateDependencyModuleNames.Add("Core");
 		PrivateDependencyModuleNames.Add("ApplicationCore");
+		PrivateDependencyModuleNames.Add("Trace");
 
 		if (Target.bCompileAgainstEngine)
 		{
diff --git a/Engine/Source/Runtime/RenderCore/Private/RenderingThread.cpp b/Engine/Source/Runtime/RenderCore/Private/RenderingThread.cpp
--- a/Engine/Source/Runtime/RenderCore/Private/RenderingThread.cpp
+++ b/Engine/Source/Runtime/RenderCore/Private/RenderingThread.cpp
@@ -23,6 +23,7 @@
 #include "RenderResource.h"
 #include "Misc/ScopeLock.h"
 #include "HAL/LowLevelMemTracker.h"
+#include "ProfilingDebugging/MiscTrace.h"
 
 //
 // Globals
@@ -701,6 +702,7 @@ void StartRenderingThread()
 	GRenderingThreadRunnable = new FRenderingThread();
 
 	GRenderingThread = FRunnableThread::Create(GRenderingThreadRunnable, *BuildRenderingThreadName(ThreadCount), 0, FPlatformAffinity::GetRenderingThreadPriority(), FPlatformAffinity::GetRenderingThreadMask());
+	TRACE_SET_THREAD_GROUP(GRenderingThread->GetThreadID(), TraceThreadGroup_Render);
 
 	// Wait for render thread to have taskgraph bound before we dispatch any tasks for it.
 	((FRenderingThread*)GRenderingThreadRunnable)->TaskGraphBoundSyncEvent->Wait();
diff --git a/Engine/Source/Runtime/Trace/Trace.Build.cs b/Engine/Source/Runtime/Trace/Trace.Build.cs
new file mode 100644
--- /dev/null
+++ b/Engine/Source/Runtime/Trace/Trace.Build.cs
@@ -0,0 +1,19 @@
+// Copyright 1998-2019 Epic Games, Inc. All Rights Reserved.
+
+using UnrealBuildTool;
+using System;
+using System.IO;
+
+public class Trace : ModuleRules
+{
+	public Trace(ReadOnlyTargetRules Target) : base(Target)
+	{
+		bRequiresImplementModule = false;
+		PrivateIncludePathModuleNames.Add("Core");
+
+        if (Target.Platform == UnrealTargetPlatform.PS4)
+        {
+            PublicAdditionalLibraries.Add("SceNet_stub_weak");
+        }
+    }
+}
