Make forcing dx11 actually force it

diff --git a/Engine/Source/Runtime/RHI/Private/Windows/WindowsDynamicRHI.cpp b/Engine/Source/Runtime/RHI/Private/Windows/WindowsDynamicRHI.cpp
--- a/Engine/Source/Runtime/RHI/Private/Windows/WindowsDynamicRHI.cpp
+++ b/Engine/Source/Runtime/RHI/Private/Windows/WindowsDynamicRHI.cpp
@@ -174,7 +174,7 @@ static IDynamicRHIModule* LoadDynamicRHIModule(ERHIFeatureLevel::Type& DesiredFe
 		}
 		LoadedRHIModuleName = VulkanRHIModuleName;
 	}
-	else if (bForceD3D12 || bPreferD3D12)
+	else if (bForceD3D12 || (bPreferD3D12 && !bForceD3D11))
 	{
 		FApp::SetGraphicsRHI(TEXT("DirectX 12"));
 		LoadedRHIModuleName = TEXT("D3D12RHI");
