Add optional fence parameter to MapStagingSurface()
If fence is passed in we Poll() to find if we have passed the fence and if so we skip the RHI flush
For the benefit of non parallel APIs we still stall the RHI thread while processing the Map()

diff --git a/Engine/Source/Runtime/Apple/MetalRHI/Private/MetalRenderTarget.cpp b/Engine/Source/Runtime/Apple/MetalRHI/Private/MetalRenderTarget.cpp
--- a/Engine/Source/Runtime/Apple/MetalRHI/Private/MetalRenderTarget.cpp
+++ b/Engine/Source/Runtime/Apple/MetalRHI/Private/MetalRenderTarget.cpp
@@ -343,7 +343,7 @@ void FMetalDynamicRHI::RHIReadSurfaceData(FRHITexture* TextureRHI, FIntRect Rect
 	}
 }
 
-void FMetalDynamicRHI::RHIMapStagingSurface(FRHITexture* TextureRHI,void*& OutData,int32& OutWidth,int32& OutHeight)
+void FMetalDynamicRHI::RHIMapStagingSurface(FRHITexture* TextureRHI, FRHIGPUFence* FenceRHI, void*& OutData, int32& OutWidth, int32& OutHeight)
 {
 	@autoreleasepool {
     FMetalSurface* Surface = GetMetalSurfaceFromRHITexture(TextureRHI);
diff --git a/Engine/Source/Runtime/Apple/MetalRHI/Public/MetalRHI.h b/Engine/Source/Runtime/Apple/MetalRHI/Public/MetalRHI.h
--- a/Engine/Source/Runtime/Apple/MetalRHI/Public/MetalRHI.h
+++ b/Engine/Source/Runtime/Apple/MetalRHI/Public/MetalRHI.h
@@ -114,7 +114,7 @@ public:
 	virtual void RHIBindDebugLabelName(FRHITexture* Texture, const TCHAR* Name) final override;
 	virtual void RHIReadSurfaceData(FRHITexture* Texture,FIntRect Rect,TArray<FColor>& OutData,FReadSurfaceDataFlags InFlags) final override;
 	virtual void RHIReadSurfaceData(FRHITexture* TextureRHI, FIntRect InRect, TArray<FLinearColor>& OutData, FReadSurfaceDataFlags InFlags) final override;
-	virtual void RHIMapStagingSurface(FRHITexture* Texture,void*& OutData,int32& OutWidth,int32& OutHeight) final override;
+	virtual void RHIMapStagingSurface(FRHITexture* Texture, FRHIGPUFence* Fence, void*& OutData, int32& OutWidth, int32& OutHeight) final override;
 	virtual void RHIUnmapStagingSurface(FRHITexture* Texture) final override;
 	virtual void RHIReadSurfaceFloatData(FRHITexture* Texture,FIntRect Rect,TArray<FFloat16Color>& OutData,ECubeFace CubeFace,int32 ArrayIndex,int32 MipIndex) final override;
 	virtual void RHIRead3DSurfaceFloatData(FRHITexture* Texture,FIntRect Rect,FIntPoint ZMinMax,TArray<FFloat16Color>& OutData) final override;
diff --git a/Engine/Source/Runtime/D3D12RHI/Private/D3D12RHIPrivate.h b/Engine/Source/Runtime/D3D12RHI/Private/D3D12RHIPrivate.h
--- a/Engine/Source/Runtime/D3D12RHI/Private/D3D12RHIPrivate.h
+++ b/Engine/Source/Runtime/D3D12RHI/Private/D3D12RHIPrivate.h
@@ -333,7 +333,7 @@ public:
 	virtual void RHIBindDebugLabelName(FRHITexture* Texture, const TCHAR* Name) final override;
 	virtual void RHIReadSurfaceData(FRHITexture* Texture, FIntRect Rect, TArray<FColor>& OutData, FReadSurfaceDataFlags InFlags) final override;
 	virtual void RHIReadSurfaceData(FRHITexture* TextureRHI, FIntRect InRect, TArray<FLinearColor>& OutData, FReadSurfaceDataFlags InFlags) final override;
-	virtual void RHIMapStagingSurface(FRHITexture* Texture, void*& OutData, int32& OutWidth, int32& OutHeight) final override;
+	virtual void RHIMapStagingSurface(FRHITexture* Texture, FRHIGPUFence* Fence, void*& OutData, int32& OutWidth, int32& OutHeight) final override;
 	virtual void RHIUnmapStagingSurface(FRHITexture* Texture) final override;
 	virtual void RHIReadSurfaceFloatData(FRHITexture* Texture, FIntRect Rect, TArray<FFloat16Color>& OutData, ECubeFace CubeFace, int32 ArrayIndex, int32 MipIndex) final override;
 	virtual void RHIRead3DSurfaceFloatData(FRHITexture* Texture, FIntRect Rect, FIntPoint ZMinMax, TArray<FFloat16Color>& OutData) final override;
diff --git a/Engine/Source/Runtime/D3D12RHI/Private/D3D12RenderTarget.cpp b/Engine/Source/Runtime/D3D12RHI/Private/D3D12RenderTarget.cpp
--- a/Engine/Source/Runtime/D3D12RHI/Private/D3D12RenderTarget.cpp
+++ b/Engine/Source/Runtime/D3D12RHI/Private/D3D12RenderTarget.cpp
@@ -1407,7 +1407,7 @@ void FD3D12DynamicRHI::ReadSurfaceDataMSAARaw(FRHICommandList_RecursiveHazardous
 	}
 }
 
-void FD3D12DynamicRHI::RHIMapStagingSurface(FRHITexture* TextureRHI, void*& OutData, int32& OutWidth, int32& OutHeight)
+void FD3D12DynamicRHI::RHIMapStagingSurface(FRHITexture* TextureRHI, FRHIGPUFence* FenceRHI, void*& OutData, int32& OutWidth, int32& OutHeight)
 {
 	FD3D12Resource* Texture = GetD3D12TextureFromRHITexture(TextureRHI)->GetResource();
 
diff --git a/Engine/Source/Runtime/EmptyRHI/Private/EmptyRenderTarget.cpp b/Engine/Source/Runtime/EmptyRHI/Private/EmptyRenderTarget.cpp
--- a/Engine/Source/Runtime/EmptyRHI/Private/EmptyRenderTarget.cpp
+++ b/Engine/Source/Runtime/EmptyRHI/Private/EmptyRenderTarget.cpp
@@ -18,7 +18,7 @@ void FEmptyDynamicRHI::RHIReadSurfaceData(FRHITexture* TextureRHI, FIntRect Rect
 
 }
 
-void FEmptyDynamicRHI::RHIMapStagingSurface(FRHITexture* TextureRHI,void*& OutData,int32& OutWidth,int32& OutHeight)
+void FEmptyDynamicRHI::RHIMapStagingSurface(FRHITexture* TextureRHI, FRHIGPUFence* FenceRHI, void*& OutData, int32& OutWidth, int32& OutHeight)
 {
 
 }
diff --git a/Engine/Source/Runtime/EmptyRHI/Public/EmptyRHI.h b/Engine/Source/Runtime/EmptyRHI/Public/EmptyRHI.h
--- a/Engine/Source/Runtime/EmptyRHI/Public/EmptyRHI.h
+++ b/Engine/Source/Runtime/EmptyRHI/Public/EmptyRHI.h
@@ -103,7 +103,7 @@ public:
 	virtual void RHIUnlockTextureCubeFace(FRHITextureCube* Texture, uint32 FaceIndex, uint32 ArrayIndex, uint32 MipIndex, bool bLockWithinMiptail) final override;
 	virtual void RHIBindDebugLabelName(FRHITexture* Texture, const TCHAR* Name) final override;
 	virtual void RHIReadSurfaceData(FRHITexture* Texture,FIntRect Rect,TArray<FColor>& OutData,FReadSurfaceDataFlags InFlags) final override;
-	virtual void RHIMapStagingSurface(FRHITexture* Texture,void*& OutData,int32& OutWidth,int32& OutHeight) final override;
+	virtual void RHIMapStagingSurface(FRHITexture* Texture, FRHIGPUFence* Fence, void*& OutData, int32& OutWidth, int32& OutHeight) final override;
 	virtual void RHIUnmapStagingSurface(FRHITexture* Texture) final override;
 	virtual void RHIReadSurfaceFloatData(FRHITexture* Texture,FIntRect Rect,TArray<FFloat16Color>& OutData,ECubeFace CubeFace,int32 ArrayIndex,int32 MipIndex) final override;
 	virtual void RHIRead3DSurfaceFloatData(FRHITexture* Texture,FIntRect Rect,FIntPoint ZMinMax,TArray<FFloat16Color>& OutData) final override;
diff --git a/Engine/Source/Runtime/NullDrv/Public/NullRHI.h b/Engine/Source/Runtime/NullDrv/Public/NullRHI.h
--- a/Engine/Source/Runtime/NullDrv/Public/NullRHI.h
+++ b/Engine/Source/Runtime/NullDrv/Public/NullRHI.h
@@ -389,7 +389,7 @@ public:
 	}
 
 
-	virtual void RHIMapStagingSurface(FRHITexture* Texture,void*& OutData,int32& OutWidth,int32& OutHeight) final override
+	virtual void RHIMapStagingSurface(FRHITexture* Texture, FRHIGPUFence* Fence, void*& OutData, int32& OutWidth, int32& OutHeight) final override
 	{
 
 	}
diff --git a/Engine/Source/Runtime/OpenGLDrv/Private/OpenGLRenderTarget.cpp b/Engine/Source/Runtime/OpenGLDrv/Private/OpenGLRenderTarget.cpp
--- a/Engine/Source/Runtime/OpenGLDrv/Private/OpenGLRenderTarget.cpp
+++ b/Engine/Source/Runtime/OpenGLDrv/Private/OpenGLRenderTarget.cpp
@@ -936,7 +936,7 @@ void FOpenGLDynamicRHI::RHIReadSurfaceData(FRHITexture* TextureRHI, FIntRect Rec
 	GetContextStateForCurrentContext().Framebuffer = (GLuint)-1;
 }
 
-void FOpenGLDynamicRHI::RHIMapStagingSurface(FRHITexture* TextureRHI,void*& OutData,int32& OutWidth,int32& OutHeight)
+void FOpenGLDynamicRHI::RHIMapStagingSurface(FRHITexture* TextureRHI, FRHIGPUFence* FenceRHI, void*& OutData, int32& OutWidth, int32& OutHeight)
 {
 	FRHICommandListImmediate& RHICmdList = FRHICommandListExecutor::GetImmediateCommandList();
 
diff --git a/Engine/Source/Runtime/OpenGLDrv/Public/OpenGLDrv.h b/Engine/Source/Runtime/OpenGLDrv/Public/OpenGLDrv.h
--- a/Engine/Source/Runtime/OpenGLDrv/Public/OpenGLDrv.h
+++ b/Engine/Source/Runtime/OpenGLDrv/Public/OpenGLDrv.h
@@ -411,7 +411,7 @@ public:
 	virtual void RHIBindDebugLabelName(FRHITexture* Texture, const TCHAR* Name) final override;
 	virtual void RHIReadSurfaceData(FRHITexture* Texture,FIntRect Rect,TArray<FColor>& OutData,FReadSurfaceDataFlags InFlags) final override;
 	virtual void RHIReadSurfaceData(FRHITexture* Texture, FIntRect Rect, TArray<FLinearColor>& OutData, FReadSurfaceDataFlags InFlags) final override;
-	virtual void RHIMapStagingSurface(FRHITexture* Texture,void*& OutData,int32& OutWidth,int32& OutHeight) final override;
+	virtual void RHIMapStagingSurface(FRHITexture* Texture, FRHIGPUFence* Fence, void*& OutData, int32& OutWidth, int32& OutHeight) final override;
 	virtual void RHIUnmapStagingSurface(FRHITexture* Texture) final override;
 	virtual void RHIReadSurfaceFloatData(FRHITexture* Texture,FIntRect Rect,TArray<FFloat16Color>& OutData,ECubeFace CubeFace,int32 ArrayIndex,int32 MipIndex) final override;
 	virtual void RHIRead3DSurfaceFloatData(FRHITexture* Texture,FIntRect Rect,FIntPoint ZMinMax,TArray<FFloat16Color>& OutData) final override;
diff --git a/Engine/Source/Runtime/RHI/Private/RHICommandList.cpp b/Engine/Source/Runtime/RHI/Private/RHICommandList.cpp
--- a/Engine/Source/Runtime/RHI/Private/RHICommandList.cpp
+++ b/Engine/Source/Runtime/RHI/Private/RHICommandList.cpp
@@ -2532,21 +2532,22 @@ void FDynamicRHI::RHIUnlockTextureCubeFace_RenderThread(class FRHICommandListImm
 }
 
 
-void FDynamicRHI::RHIMapStagingSurface_RenderThread(class FRHICommandListImmediate& RHICmdList, FRHITexture* Texture, void*& OutData, int32& OutWidth, int32& OutHeight)
+void FDynamicRHI::RHIMapStagingSurface_RenderThread(class FRHICommandListImmediate& RHICmdList, FRHITexture* Texture, FRHIGPUFence* Fence, void*& OutData, int32& OutWidth, int32& OutHeight)
 {
+	if (Fence == nullptr || !Fence->Poll())
 	{
 		QUICK_SCOPE_CYCLE_COUNTER(STAT_RHIMETHOD_MapStagingSurface_Flush);
 		RHICmdList.ImmediateFlush(EImmediateFlushType::FlushRHIThread);
 	}
-	GDynamicRHI->RHIMapStagingSurface(Texture, OutData, OutWidth, OutHeight);
+	{
+		FScopedRHIThreadStaller StallRHIThread(RHICmdList);
+		GDynamicRHI->RHIMapStagingSurface(Texture, Fence, OutData, OutWidth, OutHeight);
+	}
 }
 
 void FDynamicRHI::RHIUnmapStagingSurface_RenderThread(class FRHICommandListImmediate& RHICmdList, FRHITexture* Texture)
 {
-	{
-		QUICK_SCOPE_CYCLE_COUNTER(STAT_RHIMETHOD_MapStagingSurface_Flush);
-		RHICmdList.ImmediateFlush(EImmediateFlushType::FlushRHIThread);
-	}
+	FScopedRHIThreadStaller StallRHIThread(RHICmdList);
 	GDynamicRHI->RHIUnmapStagingSurface(Texture);
 }
 
diff --git a/Engine/Source/Runtime/RHI/Public/DynamicRHI.h b/Engine/Source/Runtime/RHI/Public/DynamicRHI.h
--- a/Engine/Source/Runtime/RHI/Public/DynamicRHI.h
+++ b/Engine/Source/Runtime/RHI/Public/DynamicRHI.h
@@ -905,7 +905,7 @@ public:
 
 	/** Watch out for OutData to be 0 (can happen on DXGI_ERROR_DEVICE_REMOVED), don't call RHIUnmapStagingSurface in that case. */
 	// FlushType: Flush Immediate (seems wrong)
-	virtual void RHIMapStagingSurface(FRHITexture* Texture, void*& OutData, int32& OutWidth, int32& OutHeight) = 0;
+	virtual void RHIMapStagingSurface(FRHITexture* Texture, FRHIGPUFence* Fence, void*& OutData, int32& OutWidth, int32& OutHeight) = 0;
 
 	/** call after a succesful RHIMapStagingSurface() call */
 	// FlushType: Flush Immediate (seems wrong)
@@ -1189,7 +1189,7 @@ public:
 	virtual void RHIDiscardTransientResource_RenderThread(FRHIStructuredBuffer* Buffer) { }
 
 
-	virtual void RHIMapStagingSurface_RenderThread(class FRHICommandListImmediate& RHICmdList, FRHITexture* Texture, void*& OutData, int32& OutWidth, int32& OutHeight);
+	virtual void RHIMapStagingSurface_RenderThread(class FRHICommandListImmediate& RHICmdList, FRHITexture* Texture, FRHIGPUFence* Fence, void*& OutData, int32& OutWidth, int32& OutHeight);
 	virtual void RHIUnmapStagingSurface_RenderThread(class FRHICommandListImmediate& RHICmdList, FRHITexture* Texture);
 	virtual void RHIReadSurfaceFloatData_RenderThread(class FRHICommandListImmediate& RHICmdList, FRHITexture* Texture, FIntRect Rect, TArray<FFloat16Color>& OutData, ECubeFace CubeFace, int32 ArrayIndex, int32 MipIndex);
 
diff --git a/Engine/Source/Runtime/RHI/Public/RHICommandList.h b/Engine/Source/Runtime/RHI/Public/RHICommandList.h
--- a/Engine/Source/Runtime/RHI/Public/RHICommandList.h
+++ b/Engine/Source/Runtime/RHI/Public/RHICommandList.h
@@ -4175,11 +4175,16 @@ public:
 		GDynamicRHI->RHIReadSurfaceData(Texture, Rect, OutData, InFlags);
 	}
 	
-	FORCEINLINE void MapStagingSurface(FRHITexture* Texture,void*& OutData,int32& OutWidth,int32& OutHeight)
+	FORCEINLINE void MapStagingSurface(FRHITexture* Texture, void*& OutData, int32& OutWidth, int32& OutHeight)
 	{
-		GDynamicRHI->RHIMapStagingSurface_RenderThread(*this, Texture,OutData,OutWidth,OutHeight);
+		GDynamicRHI->RHIMapStagingSurface_RenderThread(*this, Texture, nullptr, OutData, OutWidth, OutHeight);
 	}
-	
+
+	FORCEINLINE void MapStagingSurface(FRHITexture* Texture, FRHIGPUFence* Fence, void*& OutData, int32& OutWidth, int32& OutHeight)
+	{
+		GDynamicRHI->RHIMapStagingSurface_RenderThread(*this, Texture, Fence, OutData, OutWidth, OutHeight);
+	}
+
 	FORCEINLINE void UnmapStagingSurface(FRHITexture* Texture)
 	{
 		GDynamicRHI->RHIUnmapStagingSurface_RenderThread(*this, Texture);
diff --git a/Engine/Source/Runtime/RHI/Public/RHIValidation.h b/Engine/Source/Runtime/RHI/Public/RHIValidation.h
--- a/Engine/Source/Runtime/RHI/Public/RHIValidation.h
+++ b/Engine/Source/Runtime/RHI/Public/RHIValidation.h
@@ -905,9 +905,9 @@ public:
 
 	/** Watch out for OutData to be 0 (can happen on DXGI_ERROR_DEVICE_REMOVED), don't call RHIUnmapStagingSurface in that case. */
 	// FlushType: Flush Immediate (seems wrong)
-	virtual void RHIMapStagingSurface(FRHITexture* Texture, void*& OutData, int32& OutWidth, int32& OutHeight) override final
+	virtual void RHIMapStagingSurface(FRHITexture* Texture, FRHIGPUFence* Fence, void*& OutData, int32& OutWidth, int32& OutHeight) final override
 	{
-		RHI->RHIMapStagingSurface(Texture, OutData, OutWidth, OutHeight);
+		RHI->RHIMapStagingSurface(Texture, Fence, OutData, OutWidth, OutHeight);
 	}
 
 	/** call after a succesful RHIMapStagingSurface() call */
diff --git a/Engine/Source/Runtime/Renderer/Private/VT/VirtualTextureFeedback.cpp b/Engine/Source/Runtime/Renderer/Private/VT/VirtualTextureFeedback.cpp
--- a/Engine/Source/Runtime/Renderer/Private/VT/VirtualTextureFeedback.cpp
+++ b/Engine/Source/Runtime/Renderer/Private/VT/VirtualTextureFeedback.cpp
@@ -116,7 +116,7 @@ bool FVirtualTextureFeedback::Map(FRHICommandListImmediate& RHICmdList, MapResul
 		OutResult.Rect = FeedbackEntryCPU.Rect;
 
 		int32 LockHeight = 0;
-		RHICmdList.MapStagingSurface(FeedbackEntryCPU.TextureCPU->GetRenderTargetItem().ShaderResourceTexture, *(void**)&OutResult.Buffer, OutResult.Pitch, LockHeight);
+		RHICmdList.MapStagingSurface(FeedbackEntryCPU.TextureCPU->GetRenderTargetItem().ShaderResourceTexture, FeedbackEntryCPU.GPUFenceRHI, *(void**)&OutResult.Buffer, OutResult.Pitch, LockHeight);
 
 		check(PendingTargetCount > 0u);
 		--PendingTargetCount;
diff --git a/Engine/Source/Runtime/VulkanRHI/Private/VulkanRenderTarget.cpp b/Engine/Source/Runtime/VulkanRHI/Private/VulkanRenderTarget.cpp
--- a/Engine/Source/Runtime/VulkanRHI/Private/VulkanRenderTarget.cpp
+++ b/Engine/Source/Runtime/VulkanRHI/Private/VulkanRenderTarget.cpp
@@ -926,7 +926,7 @@ void FVulkanDynamicRHI::RHIReadSurfaceData(FRHITexture* TextureRHI, FIntRect Rec
 	}
 }
 
-void FVulkanDynamicRHI::RHIMapStagingSurface(FRHITexture* TextureRHI,void*& OutData,int32& OutWidth,int32& OutHeight)
+void FVulkanDynamicRHI::RHIMapStagingSurface(FRHITexture* TextureRHI, FRHIGPUFence* FenceRHI, void*& OutData, int32& OutWidth, int32& OutHeight)
 {
 	FRHITexture2D* TextureRHI2D = TextureRHI->GetTexture2D();
 	check(TextureRHI2D);
@@ -958,9 +958,9 @@ void FVulkanDynamicRHI::RHIMapStagingSurface(FRHITexture* TextureRHI,void*& OutD
 }
 
 
-void FVulkanDynamicRHI::RHIMapStagingSurface_RenderThread(class FRHICommandListImmediate& RHICmdList, FRHITexture* Texture, void*& OutData, int32& OutWidth, int32& OutHeight)
+void FVulkanDynamicRHI::RHIMapStagingSurface_RenderThread(class FRHICommandListImmediate& RHICmdList, FRHITexture* Texture, FRHIGPUFence* Fence, void*& OutData, int32& OutWidth, int32& OutHeight)
 {
-	RHIMapStagingSurface(Texture, OutData, OutWidth, OutHeight);
+	RHIMapStagingSurface(Texture, Fence, OutData, OutWidth, OutHeight);
 }
 
 void FVulkanDynamicRHI::RHIUnmapStagingSurface_RenderThread(class FRHICommandListImmediate& RHICmdList, FRHITexture* Texture)
diff --git a/Engine/Source/Runtime/VulkanRHI/Public/VulkanDynamicRHI.h b/Engine/Source/Runtime/VulkanRHI/Public/VulkanDynamicRHI.h
--- a/Engine/Source/Runtime/VulkanRHI/Public/VulkanDynamicRHI.h
+++ b/Engine/Source/Runtime/VulkanRHI/Public/VulkanDynamicRHI.h
@@ -147,9 +147,9 @@ public:
 	virtual void RHIBindDebugLabelName(FRHIUnorderedAccessView* UnorderedAccessViewRHI, const TCHAR* Name) final override;
 	virtual void RHIReadSurfaceData(FRHITexture* Texture, FIntRect Rect, TArray<FColor>& OutData, FReadSurfaceDataFlags InFlags) final override;
 	virtual void RHIReadSurfaceData(FRHITexture* Texture, FIntRect Rect, TArray<FLinearColor>& OutData, FReadSurfaceDataFlags InFlags) final override;
-	virtual void RHIMapStagingSurface(FRHITexture* Texture, void*& OutData, int32& OutWidth, int32& OutHeight) final override;
+	virtual void RHIMapStagingSurface(FRHITexture* Texture, FRHIGPUFence* Fence, void*& OutData, int32& OutWidth, int32& OutHeight) final override;
 	virtual void RHIUnmapStagingSurface(FRHITexture* Texture) final override;
-	virtual void RHIMapStagingSurface_RenderThread(class FRHICommandListImmediate& RHICmdList, FRHITexture* Texture, void*& OutData, int32& OutWidth, int32& OutHeight);
+	virtual void RHIMapStagingSurface_RenderThread(class FRHICommandListImmediate& RHICmdList, FRHITexture* Texture, FRHIGPUFence* Fence, void*& OutData, int32& OutWidth, int32& OutHeight);
 	virtual void RHIUnmapStagingSurface_RenderThread(class FRHICommandListImmediate& RHICmdList, FRHITexture* Texture);
 
 
diff --git a/Engine/Source/Runtime/Windows/D3D11RHI/Private/D3D11RHIPrivate.h b/Engine/Source/Runtime/Windows/D3D11RHI/Private/D3D11RHIPrivate.h
--- a/Engine/Source/Runtime/Windows/D3D11RHI/Private/D3D11RHIPrivate.h
+++ b/Engine/Source/Runtime/Windows/D3D11RHI/Private/D3D11RHIPrivate.h
@@ -469,7 +469,7 @@ public:
 	virtual void RHIBindDebugLabelName(FRHIUnorderedAccessView* UnorderedAccessViewRHI, const TCHAR* Name) final override;
 	virtual void RHIReadSurfaceData(FRHITexture* Texture,FIntRect Rect,TArray<FColor>& OutData,FReadSurfaceDataFlags InFlags) final override;
 	virtual void RHIReadSurfaceData(FRHITexture* TextureRHI, FIntRect InRect, TArray<FLinearColor>& OutData, FReadSurfaceDataFlags InFlags) final override;
-	virtual void RHIMapStagingSurface(FRHITexture* Texture,void*& OutData,int32& OutWidth,int32& OutHeight) final override;
+	virtual void RHIMapStagingSurface(FRHITexture* Texture, FRHIGPUFence* Fence, void*& OutData, int32& OutWidth, int32& OutHeight) final override;
 	virtual void RHIUnmapStagingSurface(FRHITexture* Texture) final override;
 	virtual void RHIReadSurfaceFloatData(FRHITexture* Texture,FIntRect Rect,TArray<FFloat16Color>& OutData,ECubeFace CubeFace,int32 ArrayIndex,int32 MipIndex) final override;
 	virtual void RHIRead3DSurfaceFloatData(FRHITexture* Texture,FIntRect Rect,FIntPoint ZMinMax,TArray<FFloat16Color>& OutData) final override;
diff --git a/Engine/Source/Runtime/Windows/D3D11RHI/Private/D3D11RenderTarget.cpp b/Engine/Source/Runtime/Windows/D3D11RHI/Private/D3D11RenderTarget.cpp
--- a/Engine/Source/Runtime/Windows/D3D11RHI/Private/D3D11RenderTarget.cpp
+++ b/Engine/Source/Runtime/Windows/D3D11RHI/Private/D3D11RenderTarget.cpp
@@ -815,7 +815,7 @@ void FD3D11DynamicRHI::ReadSurfaceDataMSAARaw(FRHICommandList_RecursiveHazardous
 	}
 }
 
-void FD3D11DynamicRHI::RHIMapStagingSurface(FRHITexture* TextureRHI,void*& OutData,int32& OutWidth,int32& OutHeight)
+void FD3D11DynamicRHI::RHIMapStagingSurface(FRHITexture* TextureRHI, FRHIGPUFence* FenceRHI, void*& OutData, int32& OutWidth, int32& OutHeight)
 {
 	ID3D11Texture2D* Texture = (ID3D11Texture2D*)(GetD3D11TextureFromRHITexture(TextureRHI)->GetResource());
 	
