Small rewrite of the GPU readback framework to support texture readback using the fence to avoid stalling

diff --git a/Engine/Plugins/FX/Niagara/Source/Niagara/Classes/NiagaraScriptExecutionContext.h b/Engine/Plugins/FX/Niagara/Source/Niagara/Classes/NiagaraScriptExecutionContext.h
--- a/Engine/Plugins/FX/Niagara/Source/Niagara/Classes/NiagaraScriptExecutionContext.h
+++ b/Engine/Plugins/FX/Niagara/Source/Niagara/Classes/NiagaraScriptExecutionContext.h
@@ -267,7 +267,7 @@ public:
 	TArray< FNiagaraDataInterfaceGPUParamInfo >  DIParamInfo;
 #endif
 
-	FRHIGPUMemoryReadback *GPUDataReadback;
+	FRHIGPUBufferReadback *GPUDataReadback;
 	uint32 AccumulatedSpawnRate;
 	uint32 NumIndicesPerInstance;	// how many vtx indices per instance the renderer is going to have for its draw call
 
@@ -280,9 +280,9 @@ public:
 
 
 #if WITH_EDITORONLY_DATA
-	mutable FRHIGPUMemoryReadback *GPUDebugDataReadbackFloat;
-	mutable FRHIGPUMemoryReadback *GPUDebugDataReadbackInt;
-	mutable FRHIGPUMemoryReadback *GPUDebugDataReadbackCounts;
+	mutable FRHIGPUBufferReadback *GPUDebugDataReadbackFloat;
+	mutable FRHIGPUBufferReadback *GPUDebugDataReadbackInt;
+	mutable FRHIGPUBufferReadback *GPUDebugDataReadbackCounts;
 	mutable int32 GPUDebugDataCurrBufferIdx;
 	mutable uint32 GPUDebugDataFloatSize;
 	mutable uint32 GPUDebugDataIntSize;
diff --git a/Engine/Plugins/FX/Niagara/Source/Niagara/Private/NiagaraEmitterInstanceBatcher.cpp b/Engine/Plugins/FX/Niagara/Source/Niagara/Private/NiagaraEmitterInstanceBatcher.cpp
--- a/Engine/Plugins/FX/Niagara/Source/Niagara/Private/NiagaraEmitterInstanceBatcher.cpp
+++ b/Engine/Plugins/FX/Niagara/Source/Niagara/Private/NiagaraEmitterInstanceBatcher.cpp
@@ -404,7 +404,7 @@ void NiagaraEmitterInstanceBatcher::ResolveDatasetWrites(FRHICommandList &RHICmd
 	{
 		SCOPE_CYCLE_COUNTER(STAT_NiagaraAllocateGPUReadback_RT);
 
-		Context->GPUDataReadback = new FRHIGPUMemoryReadback(TEXT("Niagara GPU Emitter Readback"));
+		Context->GPUDataReadback = new FRHIGPUBufferReadback(TEXT("Niagara GPU Emitter Readback"));
 		INC_DWORD_STAT(STAT_NiagaraReadbackLatency);
 		Context->GPUDataReadback->EnqueueCopy(RHICmdList, DatasetIndexBufferWrite.Buffer);
 	}
@@ -718,19 +718,22 @@ void NiagaraEmitterInstanceBatcher::Run(const FNiagaraComputeExecutionContext *C
 
 			if (Context->MainDataSet->GetNumFloatComponents() > 0)
 			{
-				Context->GPUDebugDataReadbackFloat = new FRHIGPUMemoryReadback(TEXT("Niagara GPU Debug Info Float Emitter Readback"));
+				static const FName ReadbackFloatName(TEXT("Niagara GPU Debug Info Float Emitter Readback"));
+				Context->GPUDebugDataReadbackFloat = new FRHIGPUBufferReadback(ReadbackFloatName);
 				Context->GPUDebugDataReadbackFloat->EnqueueCopy(RHICmdList, Context->MainDataSet->GetDataByIndex(WriteBufferIdx).GetGPUBufferFloat()->Buffer);
 				Context->GPUDebugDataFloatSize = Context->MainDataSet->GetDataByIndex(WriteBufferIdx).GetGPUBufferFloat()->NumBytes;
 			}
 
 			if (Context->MainDataSet->GetNumInt32Components() > 0)
 			{
-				Context->GPUDebugDataReadbackInt = new FRHIGPUMemoryReadback(TEXT("Niagara GPU Debug Info Int Emitter Readback"));
+				static const FName ReadbackIntName(TEXT("Niagara GPU Debug Info Int Emitter Readback"));
+				Context->GPUDebugDataReadbackInt = new FRHIGPUBufferReadback(ReadbackIntName);
 				Context->GPUDebugDataReadbackInt->EnqueueCopy(RHICmdList, Context->MainDataSet->GetDataByIndex(WriteBufferIdx).GetGPUBufferInt()->Buffer);
 				Context->GPUDebugDataIntSize = Context->MainDataSet->GetDataByIndex(WriteBufferIdx).GetGPUBufferInt()->NumBytes;
 			}
 
-			Context->GPUDebugDataReadbackCounts = new FRHIGPUMemoryReadback(TEXT("Niagara GPU Emitter Readback"));
+			static const FName ReadbackCountsName(TEXT("Niagara GPU Emitter Readback"));
+			Context->GPUDebugDataReadbackCounts = new FRHIGPUBufferReadback(ReadbackCountsName);
 			Context->GPUDebugDataReadbackCounts->EnqueueCopy(RHICmdList, DatasetIndexBufferWrite.Buffer);
 		}
 	}
diff --git a/Engine/Source/Runtime/Apple/MetalRHI/Private/MetalUAV.cpp b/Engine/Source/Runtime/Apple/MetalRHI/Private/MetalUAV.cpp
--- a/Engine/Source/Runtime/Apple/MetalRHI/Private/MetalUAV.cpp
+++ b/Engine/Source/Runtime/Apple/MetalRHI/Private/MetalUAV.cpp
@@ -871,7 +871,7 @@ void FMetalGPUFence::WriteInternal(mtlpp::CommandBuffer& CmdBuffer)
 	check(Fence);
 }
 
-void FMetalRHICommandContext::RHICopyToStagingBuffer(FVertexBufferRHIParamRef SourceBufferRHI, FStagingBufferRHIParamRef DestinationStagingBufferRHI, uint32 Offset, uint32 NumBytes, FGPUFenceRHIParamRef FenceRHI)
+void FMetalRHICommandContext::RHICopyToStagingBuffer(FVertexBufferRHIParamRef SourceBufferRHI, FStagingBufferRHIParamRef DestinationStagingBufferRHI, uint32 Offset, uint32 NumBytes)
 {
 	@autoreleasepool {
 		check(DestinationStagingBufferRHI);
@@ -895,12 +895,15 @@ void FMetalRHICommandContext::RHICopyToStagingBuffer(FVertexBufferRHIParamRef So
 
 		// Inline copy from the actual buffer to the shadow
 		GetMetalDeviceContext().CopyFromBufferToBuffer(SourceBuffer->Buffer, Offset, ReadbackBuffer, 0, NumBytes);
+	}
+}
 
-		if (FenceRHI)
-		{
-			FMetalGPUFence* Fence = ResourceCast(FenceRHI);
-			Fence->WriteInternal(Context->GetCurrentCommandBuffer());
-		}
+void FMetalRHICommandContext::RHIWriteGPUFence(FGPUFenceRHIParamRef FenceRHI)
+{
+	@autoreleasepool {
+		check(FenceRHI);
+		FMetalGPUFence* Fence = ResourceCast(FenceRHI);
+		Fence->WriteInternal(Context->GetCurrentCommandBuffer());
 	}
 }
 
diff --git a/Engine/Source/Runtime/Apple/MetalRHI/Public/MetalRHIContext.h b/Engine/Source/Runtime/Apple/MetalRHI/Public/MetalRHIContext.h
--- a/Engine/Source/Runtime/Apple/MetalRHI/Public/MetalRHIContext.h
+++ b/Engine/Source/Runtime/Apple/MetalRHI/Public/MetalRHIContext.h
@@ -301,7 +301,8 @@ public:
 	 */
 	virtual void RHITransitionResources(EResourceTransitionAccess TransitionType, EResourceTransitionPipeline TransitionPipeline, FUnorderedAccessViewRHIParamRef* InUAVs, int32 NumUAVs, FComputeFenceRHIParamRef WriteComputeFence) final override;
 	
-	virtual void RHICopyToStagingBuffer(FVertexBufferRHIParamRef SourceBufferRHI, FStagingBufferRHIParamRef DestinationStagingBufferRHI, uint32 Offset, uint32 NumBytes, FGPUFenceRHIParamRef FenceRHI) final override;
+	virtual void RHICopyToStagingBuffer(FVertexBufferRHIParamRef SourceBufferRHI, FStagingBufferRHIParamRef DestinationStagingBufferRHI, uint32 Offset, uint32 NumBytes) final override;
+	virtual void RHIWriteGPUFence(FGPUFenceRHIParamRef FenceRHI) final override;
 
 	/**
 	 * Explicitly transition a texture resource from readable -> writable by the GPU or vice versa.
diff --git a/Engine/Source/Runtime/D3D12RHI/Private/D3D12CommandContext.h b/Engine/Source/Runtime/D3D12RHI/Private/D3D12CommandContext.h
--- a/Engine/Source/Runtime/D3D12RHI/Private/D3D12CommandContext.h
+++ b/Engine/Source/Runtime/D3D12RHI/Private/D3D12CommandContext.h
@@ -273,7 +273,8 @@ public:
 	virtual void RHIClearTinyUAV(FUnorderedAccessViewRHIParamRef UnorderedAccessViewRHI, const uint32* Values) final override;
 	virtual void RHICopyToResolveTarget(FTextureRHIParamRef SourceTexture, FTextureRHIParamRef DestTexture, const FResolveParams& ResolveParams) override;
 	virtual void RHITransitionResources(EResourceTransitionAccess TransitionType, FTextureRHIParamRef* InTextures, int32 NumTextures) final override;
-	virtual void RHICopyToStagingBuffer(FVertexBufferRHIParamRef SourceBuffer, FStagingBufferRHIParamRef DestinationStagingBuffer, uint32 Offset, uint32 NumBytes, FGPUFenceRHIParamRef Fence = nullptr) final override;
+	virtual void RHICopyToStagingBuffer(FVertexBufferRHIParamRef SourceBuffer, FStagingBufferRHIParamRef DestinationStagingBuffer, uint32 Offset, uint32 NumBytes) final override;
+	virtual void RHIWriteGPUFence(FGPUFenceRHIParamRef Fence) final override;
 	virtual void RHIBeginRenderQuery(FRenderQueryRHIParamRef RenderQuery) final override;
 	virtual void RHIEndRenderQuery(FRenderQueryRHIParamRef RenderQuery) final override;
 	void RHIBeginOcclusionQueryBatch(uint32 NumQueriesInBatch);
@@ -480,9 +481,13 @@ public:
 	// Special implementation that only signal the fence once.
 	virtual void RHITransitionResources(EResourceTransitionAccess TransitionType, EResourceTransitionPipeline TransitionPipeline, FUnorderedAccessViewRHIParamRef* InUAVs, int32 NumUAVs, FComputeFenceRHIParamRef WriteComputeFenceRHI) final override;
 
-	FORCEINLINE virtual void RHICopyToStagingBuffer(FVertexBufferRHIParamRef SourceBuffer, FStagingBufferRHIParamRef DestinationStagingBuffer, uint32 Offset, uint32 NumBytes, FGPUFenceRHIParamRef Fence = nullptr) final override
+	FORCEINLINE virtual void RHICopyToStagingBuffer(FVertexBufferRHIParamRef SourceBuffer, FStagingBufferRHIParamRef DestinationStagingBuffer, uint32 Offset, uint32 NumBytes) final override
 	{
-		ContextRedirect(RHICopyToStagingBuffer(SourceBuffer, DestinationStagingBuffer, Offset, NumBytes, Fence));
+		ContextRedirect(RHICopyToStagingBuffer(SourceBuffer, DestinationStagingBuffer, Offset, NumBytes));
+	}
+	FORCEINLINE virtual void RHIWriteGPUFence(FGPUFenceRHIParamRef Fence) final override
+	{
+		ContextRedirect(RHIWriteGPUFence(Fence));
 	}
 	FORCEINLINE virtual void RHISetShaderTexture(FComputeShaderRHIParamRef PixelShader, uint32 TextureIndex, FTextureRHIParamRef NewTexture) final override
 	{
diff --git a/Engine/Source/Runtime/D3D12RHI/Private/D3D12Commands.cpp b/Engine/Source/Runtime/D3D12RHI/Private/D3D12Commands.cpp
--- a/Engine/Source/Runtime/D3D12RHI/Private/D3D12Commands.cpp
+++ b/Engine/Source/Runtime/D3D12RHI/Private/D3D12Commands.cpp
@@ -393,7 +393,7 @@ void FD3D12CommandContext::RHITransitionResources(EResourceTransitionAccess Tran
 }
 
 
-void FD3D12CommandContext::RHICopyToStagingBuffer(FVertexBufferRHIParamRef SourceBufferRHI, FStagingBufferRHIParamRef StagingBufferRHI, uint32 Offset, uint32 NumBytes, FGPUFenceRHIParamRef FenceRHI)
+void FD3D12CommandContext::RHICopyToStagingBuffer(FVertexBufferRHIParamRef SourceBufferRHI, FStagingBufferRHIParamRef StagingBufferRHI, uint32 Offset, uint32 NumBytes)
 {
 	FD3D12StagingBuffer* StagingBuffer = FD3D12DynamicRHI::ResourceCast(StagingBufferRHI);
 	check(StagingBuffer);
@@ -432,14 +432,16 @@ void FD3D12CommandContext::RHICopyToStagingBuffer(FVertexBufferRHIParamRef Sourc
 		CommandListHandle.UpdateResidency(pDestResource);
 		CommandListHandle.UpdateResidency(pSourceResource);
 	}
+}
 
-	if (FenceRHI)
-	{
-		// @todo-mattc we don't want to flush here. That should be the caller's responsibility.
-		RHISubmitCommandsHint();
-		FD3D12GPUFence* Fence = FD3D12DynamicRHI::ResourceCast(FenceRHI);
-		Fence->WriteInternal(ED3D12CommandQueueType::Default);
-	}
+void FD3D12CommandContext::RHIWriteGPUFence(FGPUFenceRHIParamRef FenceRHI)
+{
+	check(FenceRHI)
+
+	// @todo-mattc we don't want to flush here. That should be the caller's responsibility.
+	RHISubmitCommandsHint();
+	FD3D12GPUFence* Fence = FD3D12DynamicRHI::ResourceCast(FenceRHI);
+	Fence->WriteInternal(ED3D12CommandQueueType::Default);
 }
 
 void FD3D12CommandContext::RHISetViewport(uint32 MinX, uint32 MinY, float MinZ, uint32 MaxX, uint32 MaxY, float MaxZ)
diff --git a/Engine/Source/Runtime/OpenGLDrv/Private/OpenGLCommands.cpp b/Engine/Source/Runtime/OpenGLDrv/Private/OpenGLCommands.cpp
--- a/Engine/Source/Runtime/OpenGLDrv/Private/OpenGLCommands.cpp
+++ b/Engine/Source/Runtime/OpenGLDrv/Private/OpenGLCommands.cpp
@@ -3593,13 +3593,12 @@ void FOpenGLDynamicRHI::RHIInvalidateCachedState()
 	SharedContextState.InitializeResources(FOpenGL::GetMaxCombinedTextureImageUnits(), OGL_MAX_COMPUTE_STAGE_UAV_UNITS);
 }
 
-void FOpenGLDynamicRHI::RHICopyToStagingBuffer(FVertexBufferRHIParamRef SourceBufferRHI, FStagingBufferRHIParamRef DestinationStagingBufferRHI, uint32 InOffset, uint32 InNumBytes, FGPUFenceRHIParamRef FenceRHI)
+void FOpenGLDynamicRHI::RHICopyToStagingBuffer(FVertexBufferRHIParamRef SourceBufferRHI, FStagingBufferRHIParamRef DestinationStagingBufferRHI, uint32 InOffset, uint32 InNumBytes)
 {
 #if OPENGL_GL3
 	VERIFY_GL_SCOPE();
 	FOpenGLVertexBuffer* SourceBuffer = ResourceCast(SourceBufferRHI);
 	FOpenGLStagingBuffer* DestinationBuffer = ResourceCast(DestinationStagingBufferRHI);
-	FOpenGLGPUFence* CopyFence = ResourceCast(FenceRHI);
 
 	check(DestinationBuffer->ShadowBuffer != 0);
 	glBindBuffer(GL_COPY_WRITE_BUFFER, DestinationBuffer->ShadowBuffer);
@@ -3614,13 +3613,19 @@ void FOpenGLDynamicRHI::RHICopyToStagingBuffer(FVertexBufferRHIParamRef SourceBu
 	FOpenGL::CopyBufferSubData(GL_COPY_READ_BUFFER, GL_COPY_WRITE_BUFFER, InOffset, 0, InNumBytes);
 	glBindBuffer(GL_COPY_READ_BUFFER, 0);
 	glBindBuffer(GL_COPY_WRITE_BUFFER, 0);
-
-	CopyFence->WriteInternal();
 #else
 	UE_LOG(LogRHI, Fatal, TEXT("Only OpenGL3 or higher supports CopyToStagingBuffer"));
 #endif
 }
 
+void FOpenGLDynamicRHI::RHIWriteGPUFence(FGPUFenceRHIParamRef FenceRHI)
+{
+	check(FenceRHI);
+	FOpenGLGPUFence* CopyFence = ResourceCast(FenceRHI);
+	CopyFence->WriteInternal();
+}
+
+
 #if PLATFORM_USES_FIXED_RHI_CLASS
 #define INTERNAL_DECORATOR(Method) ((FOpenGLDynamicRHI&)CmdList.GetContext()).FOpenGLDynamicRHI::Method
 #include "RHICommandListCommandExecutes.inl"
diff --git a/Engine/Source/Runtime/OpenGLDrv/Public/OpenGLDrv.h b/Engine/Source/Runtime/OpenGLDrv/Public/OpenGLDrv.h
--- a/Engine/Source/Runtime/OpenGLDrv/Public/OpenGLDrv.h
+++ b/Engine/Source/Runtime/OpenGLDrv/Public/OpenGLDrv.h
@@ -535,7 +535,8 @@ public:
 	virtual void RHIAliasTextureResources(FTextureRHIParamRef DestTexture, FTextureRHIParamRef SrcTexture);
 
 	// Inline copy
-	virtual void RHICopyToStagingBuffer(FVertexBufferRHIParamRef SourceBufferRHI, FStagingBufferRHIParamRef DestinationStagingBufferRHI, uint32 InOffset, uint32 InNumBytes, FGPUFenceRHIParamRef FenceRHI) final override;
+	virtual void RHICopyToStagingBuffer(FVertexBufferRHIParamRef SourceBufferRHI, FStagingBufferRHIParamRef DestinationStagingBufferRHI, uint32 InOffset, uint32 InNumBytes) final override;
+	virtual void RHIWriteGPUFence(FGPUFenceRHIParamRef FenceRHI) final override;
 	virtual void* RHILockStagingBuffer(FStagingBufferRHIParamRef StagingBuffer, uint32 Offset, uint32 SizeRHI) final override;
 	virtual void RHIUnlockStagingBuffer(FStagingBufferRHIParamRef StagingBuffer) final override;
 	virtual void* LockStagingBuffer_RenderThread(class FRHICommandListImmediate& RHICmdList, FStagingBufferRHIParamRef StagingBuffer, uint32 Offset, uint32 SizeRHI) final override;
diff --git a/Engine/Source/Runtime/RHI/Private/RHIGPUReadback.cpp b/Engine/Source/Runtime/RHI/Private/RHIGPUReadback.cpp
--- a/Engine/Source/Runtime/RHI/Private/RHIGPUReadback.cpp
+++ b/Engine/Source/Runtime/RHI/Private/RHIGPUReadback.cpp
@@ -7,6 +7,10 @@
 
 #include "RHIGPUReadback.h"
 
+///////////////////////////////////////////////////////////////////////////////
+//////////////////////     FGenericRHIGPUFence    /////////////////////////////
+///////////////////////////////////////////////////////////////////////////////
+
 FGenericRHIGPUFence::FGenericRHIGPUFence(FName InName)
 	: FRHIGPUFence(InName)
 	, InsertedFrameNumber(MAX_uint32)
@@ -32,6 +36,10 @@ bool FGenericRHIGPUFence::Poll() const
 	return false;
 }
 
+///////////////////////////////////////////////////////////////////////////////
+////////////////////     FGenericRHIStagingBuffer    //////////////////////////
+///////////////////////////////////////////////////////////////////////////////
+
 void* FGenericRHIStagingBuffer::Lock(uint32 InOffset, uint32 NumBytes)
 {
 	check(ShadowBuffer);
@@ -45,4 +53,95 @@ void FGenericRHIStagingBuffer::Unlock()
 	check(bIsLocked);
 	RHIUnlockVertexBuffer(ShadowBuffer);
 	bIsLocked = false;
-}
\ No newline at end of file
+}
+
+///////////////////////////////////////////////////////////////////////////////
+////////////////////     FRHIGPUBufferReadback    /////////////////////////////
+///////////////////////////////////////////////////////////////////////////////
+
+FRHIGPUBufferReadback::FRHIGPUBufferReadback(FName RequestName) : FRHIGPUMemoryReadback(RequestName)
+{
+	DestinationStagingBuffer = RHICreateStagingBuffer();
+}
+
+void FRHIGPUBufferReadback::EnqueueCopy(FRHICommandList& RHICmdList, FVertexBufferRHIParamRef SourceBuffer, uint32 NumBytes)
+{
+	Fence->Clear();
+	RHICmdList.CopyToStagingBuffer(SourceBuffer, DestinationStagingBuffer, 0, NumBytes ? NumBytes : SourceBuffer->GetSize());
+	RHICmdList.WriteGPUFence(Fence);
+}
+
+void* FRHIGPUBufferReadback::Lock(uint32 NumBytes)
+{
+	if (DestinationStagingBuffer)
+	{
+		ensure(Fence->Poll());
+		return RHILockStagingBuffer(DestinationStagingBuffer, 0, NumBytes);
+	}
+	else
+	{
+		return nullptr;
+	}
+}
+
+void FRHIGPUBufferReadback::Unlock()
+{
+	ensure(DestinationStagingBuffer);
+	RHIUnlockStagingBuffer(DestinationStagingBuffer);
+}
+
+///////////////////////////////////////////////////////////////////////////////
+////////////////////     FRHIGPUTextureReadback    ////////////////////////////
+///////////////////////////////////////////////////////////////////////////////
+
+FRHIGPUTextureReadback::FRHIGPUTextureReadback(FName RequestName) : FRHIGPUMemoryReadback(RequestName)
+{
+}
+
+void FRHIGPUTextureReadback::EnqueueCopy(FRHICommandList& RHICmdList, FTextureRHIParamRef SourceTexture, FResolveRect Rect)
+{
+	Fence->Clear();
+
+	if (SourceTexture)
+	{
+		// We only support 2d textures for now.
+		ensure(SourceTexture->GetTexture2D());
+
+		// Assume for now that every enqueue happens on a texture of the same format and size (when reused).
+		if (!DestinationStagingBuffer)
+		{
+			FIntVector TextureSize = SourceTexture->GetSizeXYZ();
+
+			FRHIResourceCreateInfo CreateInfo(*Fence->GetFName().ToString());
+			DestinationStagingBuffer = RHICreateTexture2D(TextureSize.X, TextureSize.Y, SourceTexture->GetFormat(), 1, 1, TexCreate_CPUReadback | TexCreate_HideInVisualizeTexture, CreateInfo);
+		}
+
+		// Transfer memory GPU -> CPU
+		RHICmdList.CopyToResolveTarget(SourceTexture, DestinationStagingBuffer, FResolveParams(Rect));
+		RHICmdList.WriteGPUFence(Fence);
+	}
+}
+
+void* FRHIGPUTextureReadback::Lock(uint32 NumBytes)
+{
+	if (DestinationStagingBuffer)
+	{
+		void* ResultsBuffer = nullptr;
+		int32 BufferWidth = 0, BufferHeight = 0;
+		FRHICommandListImmediate& RHICmdList = FRHICommandListExecutor::GetImmediateCommandList();
+		RHICmdList.MapStagingSurface(DestinationStagingBuffer, ResultsBuffer, BufferWidth, BufferHeight);
+		return ResultsBuffer;
+	}
+	else
+	{
+		return nullptr;
+	}
+}
+
+void FRHIGPUTextureReadback::Unlock()
+{
+	ensure(DestinationStagingBuffer);
+
+	FRHICommandListImmediate& RHICmdList = FRHICommandListExecutor::GetImmediateCommandList();
+	RHICmdList.UnmapStagingSurface(DestinationStagingBuffer);
+}
diff --git a/Engine/Source/Runtime/RHI/Public/RHICommandList.h b/Engine/Source/Runtime/RHI/Public/RHICommandList.h
--- a/Engine/Source/Runtime/RHI/Public/RHICommandList.h
+++ b/Engine/Source/Runtime/RHI/Public/RHICommandList.h
@@ -1456,14 +1456,12 @@ struct FRHICommandCopyToStagingBuffer final : public FRHICommand<FRHICommandCopy
 {
 	FVertexBufferRHIParamRef SourceBuffer;
 	FStagingBufferRHIParamRef DestinationStagingBuffer;
-	FGPUFenceRHIParamRef Fence;
 	uint32 Offset;
 	uint32 NumBytes;
 
-	FORCEINLINE_DEBUGGABLE FRHICommandCopyToStagingBuffer(FVertexBufferRHIParamRef InSourceBuffer, FStagingBufferRHIParamRef InDestinationStagingBuffer, uint32 InOffset, uint32 InNumBytes, FGPUFenceRHIParamRef InFence = nullptr)
+	FORCEINLINE_DEBUGGABLE FRHICommandCopyToStagingBuffer(FVertexBufferRHIParamRef InSourceBuffer, FStagingBufferRHIParamRef InDestinationStagingBuffer, uint32 InOffset, uint32 InNumBytes)
 		: SourceBuffer(InSourceBuffer)
 		, DestinationStagingBuffer(InDestinationStagingBuffer)
-		, Fence(InFence)
 		, Offset(InOffset)
 		, NumBytes(InNumBytes)
 	{
@@ -1471,6 +1469,18 @@ struct FRHICommandCopyToStagingBuffer final : public FRHICommand<FRHICommandCopy
 	RHI_API void Execute(FRHICommandListBase& CmdList);
 };
 
+template<ECmdList CmdListType>
+struct FRHICommandWriteGPUFence final : public FRHICommand<FRHICommandWriteGPUFence<CmdListType>>
+{
+	FGPUFenceRHIParamRef Fence;
+
+	FORCEINLINE_DEBUGGABLE FRHICommandWriteGPUFence(FGPUFenceRHIParamRef InFence)
+		: Fence(InFence)
+	{
+	}
+	RHI_API void Execute(FRHICommandListBase& CmdList);
+};
+
 struct FRHICommandClearColorTexture final : public FRHICommand<FRHICommandClearColorTexture>
 {
 	FTextureRHIParamRef Texture;
@@ -2857,14 +2867,24 @@ public:
 		ALLOC_COMMAND(FRHICommandWaitComputeFence<ECmdList::EGfx>)(WaitFence);
 	}
 
-	FORCEINLINE_DEBUGGABLE void CopyToStagingBuffer(FVertexBufferRHIParamRef SourceBuffer, FStagingBufferRHIParamRef DestinationStagingBuffer, uint32 Offset, uint32 NumBytes, FGPUFenceRHIParamRef Fence)
+	FORCEINLINE_DEBUGGABLE void CopyToStagingBuffer(FVertexBufferRHIParamRef SourceBuffer, FStagingBufferRHIParamRef DestinationStagingBuffer, uint32 Offset, uint32 NumBytes)
+	{
+		if (Bypass())
+		{
+			GetContext().RHICopyToStagingBuffer(SourceBuffer, DestinationStagingBuffer, Offset, NumBytes);
+			return;
+		}
+		ALLOC_COMMAND(FRHICommandCopyToStagingBuffer<ECmdList::EGfx>)(SourceBuffer, DestinationStagingBuffer, Offset, NumBytes);
+	}
+
+	FORCEINLINE_DEBUGGABLE void WriteGPUFence(FGPUFenceRHIParamRef Fence)
 	{
 		if (Bypass())
 		{
-			GetContext().RHICopyToStagingBuffer(SourceBuffer, DestinationStagingBuffer, Offset, NumBytes, Fence);
+			GetContext().RHIWriteGPUFence(Fence);
 			return;
 		}
-		ALLOC_COMMAND(FRHICommandCopyToStagingBuffer<ECmdList::EGfx>)(SourceBuffer, DestinationStagingBuffer, Offset, NumBytes, Fence);
+		ALLOC_COMMAND(FRHICommandWriteGPUFence<ECmdList::EGfx>)(Fence);
 	}
 
 	FORCEINLINE_DEBUGGABLE void BeginRenderPass(const FRHIRenderPassInfo& InInfo, const TCHAR* Name)
@@ -3413,14 +3433,24 @@ public:
 		ALLOC_COMMAND(FRHICommandWaitComputeFence<ECmdList::ECompute>)(WaitFence);
 	}
 
-	FORCEINLINE_DEBUGGABLE void CopyToStagingBuffer(FVertexBufferRHIParamRef SourceBuffer, FStagingBufferRHIParamRef DestinationStagingBuffer, uint32 Offset, uint32 NumBytes, FGPUFenceRHIParamRef Fence = nullptr)
+	FORCEINLINE_DEBUGGABLE void CopyToStagingBuffer(FVertexBufferRHIParamRef SourceBuffer, FStagingBufferRHIParamRef DestinationStagingBuffer, uint32 Offset, uint32 NumBytes)
+	{
+		if (Bypass())
+		{
+			GetComputeContext().RHICopyToStagingBuffer(SourceBuffer, DestinationStagingBuffer, Offset, NumBytes);
+			return;
+		}
+		ALLOC_COMMAND(FRHICommandCopyToStagingBuffer<ECmdList::ECompute>)(SourceBuffer, DestinationStagingBuffer, Offset, NumBytes);
+	}
+
+	FORCEINLINE_DEBUGGABLE void WriteGPUFence(FGPUFenceRHIParamRef Fence)
 	{
 		if (Bypass())
 		{
-			GetComputeContext().RHICopyToStagingBuffer(SourceBuffer, DestinationStagingBuffer, Offset, NumBytes, Fence);
+			GetComputeContext().RHIWriteGPUFence(Fence);
 			return;
 		}
-		ALLOC_COMMAND(FRHICommandCopyToStagingBuffer<ECmdList::ECompute>)(SourceBuffer, DestinationStagingBuffer, Offset, NumBytes, Fence);
+		ALLOC_COMMAND(FRHICommandWriteGPUFence<ECmdList::ECompute>)(Fence);
 	}
 };
 
diff --git a/Engine/Source/Runtime/RHI/Public/RHICommandListCommandExecutes.inl b/Engine/Source/Runtime/RHI/Public/RHICommandListCommandExecutes.inl
--- a/Engine/Source/Runtime/RHI/Public/RHICommandListCommandExecutes.inl
+++ b/Engine/Source/Runtime/RHI/Public/RHICommandListCommandExecutes.inl
@@ -475,11 +475,21 @@ template<ECmdList CmdListType>
 void FRHICommandCopyToStagingBuffer<CmdListType>::Execute(FRHICommandListBase& CmdList)
 {
 	RHISTAT(EnqueueStagedRead);
-	INTERNAL_DECORATOR_CONTEXT(RHICopyToStagingBuffer)(SourceBuffer, DestinationStagingBuffer, Offset, NumBytes, Fence);
+	INTERNAL_DECORATOR_CONTEXT(RHICopyToStagingBuffer)(SourceBuffer, DestinationStagingBuffer, Offset, NumBytes);
 }
 template struct FRHICommandCopyToStagingBuffer<ECmdList::EGfx>;
 template struct FRHICommandCopyToStagingBuffer<ECmdList::ECompute>;
 
+template<ECmdList CmdListType>
+void FRHICommandWriteGPUFence<CmdListType>::Execute(FRHICommandListBase& CmdList)
+{
+	RHISTAT(WriteGPUFence);
+	INTERNAL_DECORATOR_CONTEXT(RHIWriteGPUFence)(Fence);
+}
+template struct FRHICommandWriteGPUFence<ECmdList::EGfx>;
+template struct FRHICommandWriteGPUFence<ECmdList::ECompute>;
+
+
 void FRHICommandBuildLocalUniformBuffer::Execute(FRHICommandListBase& CmdList)
 {
 	LLM_SCOPE(ELLMTag::Shaders);
diff --git a/Engine/Source/Runtime/RHI/Public/RHIContext.h b/Engine/Source/Runtime/RHI/Public/RHIContext.h
--- a/Engine/Source/Runtime/RHI/Public/RHIContext.h
+++ b/Engine/Source/Runtime/RHI/Public/RHIContext.h
@@ -125,12 +125,21 @@ public:
 	 * @param DestinationStagingBuffer The the host-visible destination buffer
 	 * @param Offset The start of the data in 'SourceBuffer'
 	 * @param NumBytes The number of bytes to copy out of 'SourceBuffer'
-	 * @param Fence (optional) A GPU fence that will be inserted into the GPU timeline and which must then be tested on the CPU to know when the Copy was completed. Can be NULL, though that implies you will not have any guarantees as to when it is safe to read from 'DestinationStagingBuffer'
 	 */
-	virtual void RHICopyToStagingBuffer(FVertexBufferRHIParamRef SourceBufferRHI, FStagingBufferRHIParamRef DestinationStagingBufferRHI, uint32 InOffset, uint32 InNumBytes, FGPUFenceRHIParamRef FenceRHI = nullptr)
+	virtual void RHICopyToStagingBuffer(FVertexBufferRHIParamRef SourceBufferRHI, FStagingBufferRHIParamRef DestinationStagingBufferRHI, uint32 InOffset, uint32 InNumBytes)
 	{
 		check(false);
 	}
+
+	/**
+	 * Write the fence in the GPU timeline. The fence can then be tested on the CPU to know if the previous GPU commands are completed.
+	 * @param Fence 
+	 */
+	virtual void RHIWriteGPUFence(FGPUFenceRHIParamRef FenceRHI)
+	{
+		check(false);
+	}
+
 };
 
 struct FAccelerationStructureUpdateParams
diff --git a/Engine/Source/Runtime/RHI/Public/RHIGPUReadback.h b/Engine/Source/Runtime/RHI/Public/RHIGPUReadback.h
--- a/Engine/Source/Runtime/RHI/Public/RHIGPUReadback.h
+++ b/Engine/Source/Runtime/RHI/Public/RHIGPUReadback.h
@@ -11,42 +11,41 @@
 #include "CoreMinimal.h"
 #include "RHI.h"
 
-/* FRHIGPUMemoryReadback: Represents a memory readback request scheduled with CopyToStagingBuffer
-** Wraps a FRHIStagingBuffer with a FRHIGPUFence for synchronization.
-*/
-class FRHIGPUMemoryReadback
+/**
+ * FRHIGPUMemoryReadback: Represents a memory readback request scheduled with CopyToStagingBuffer
+ * Wraps a staging buffer with a FRHIGPUFence for synchronization.
+ */
+class RHI_API FRHIGPUMemoryReadback
 {
 public:
 
 	FRHIGPUMemoryReadback(FName RequestName)
 	{
-		DestinationStagingBuffer = RHICreateStagingBuffer();
 		Fence = RHICreateGPUFence(RequestName);
 	}
 
-	~FRHIGPUMemoryReadback()
+	virtual ~FRHIGPUMemoryReadback() {}
+
+	/** Indicates if the data is in place and ready to be read. */
+	FORCEINLINE bool IsReady()
 	{
+		return !Fence || Fence->Poll();
 	}
 
 	/**
-	 * Returns the CPU accessible pointer that backs this staging buffer.
+	 * Copy the current state of the resource to the readback data.
 	 * @param RHICmdList The command list to enqueue the copy request on.
 	 * @param SourceBuffer The buffer holding the source data.
 	 * @param NumBytes The number of bytes to copy. If 0, this will copy the entire buffer.
 	 */
-
-	void EnqueueCopy(FRHICommandList& RHICmdList, FVertexBufferRHIParamRef SourceBuffer, uint32 NumBytes = 0)
+	virtual void EnqueueCopy(FRHICommandList& RHICmdList, FVertexBufferRHIParamRef SourceBuffer, uint32 NumBytes = 0) 
 	{
-		Fence->Clear();
-		RHICmdList.CopyToStagingBuffer(SourceBuffer, DestinationStagingBuffer, 0, NumBytes ? NumBytes : SourceBuffer->GetSize(), Fence);
+		unimplemented();
 	}
 
-	/**
-	 * Indicates if the data is in place and ready to be read.
-	 */
-	bool IsReady()
+	virtual void EnqueueCopy(FRHICommandList& RHICmdList, FTextureRHIParamRef SourceTexture, FResolveRect Rect = FResolveRect())
 	{
-		return !Fence || Fence->Poll();
+		unimplemented();
 	}
 
 	/**
@@ -54,21 +53,47 @@ public:
 	 * @param NumBytes The maximum number of bytes the host will read from this pointer.
 	 * @returns A CPU accessible pointer to the backing buffer.
 	 */
-	void* Lock(uint32 NumBytes)
-	{
-		ensure(Fence->Poll());
-		return RHILockStagingBuffer(DestinationStagingBuffer, 0, NumBytes);
-	}
+	virtual void* Lock(uint32 NumBytes) = 0;
 
 	/**
 	 * Signals that the host is finished reading from the backing buffer.
 	 */
-	void Unlock()
-	{
-		RHIUnlockStagingBuffer(DestinationStagingBuffer);
-	}
+	virtual void Unlock() = 0;
+
+protected:
+
+	FGPUFenceRHIRef Fence;
+};
+
+/** Buffer readback implementation. */
+class RHI_API FRHIGPUBufferReadback final : public FRHIGPUMemoryReadback
+{
+public:
+
+	FRHIGPUBufferReadback(FName RequestName);
+	 
+	void EnqueueCopy(FRHICommandList& RHICmdList, FVertexBufferRHIParamRef SourceBuffer, uint32 NumBytes = 0) override;
+	void* Lock(uint32 NumBytes) override;
+	void Unlock() override;
 
 private:
+
 	FStagingBufferRHIRef DestinationStagingBuffer;
-	FGPUFenceRHIRef Fence;
+};
+
+
+/** Texture readback implementation. */
+class RHI_API FRHIGPUTextureReadback final : public FRHIGPUMemoryReadback
+{
+public:
+
+	FRHIGPUTextureReadback(FName RequestName);
+
+	void EnqueueCopy(FRHICommandList& RHICmdList, FTextureRHIParamRef SourceTexture, FResolveRect Rect = FResolveRect()) override;
+	void* Lock(uint32 NumBytes) override;
+	void Unlock() override;
+
+private:
+
+	FTextureRHIRef DestinationStagingBuffer;
 };
diff --git a/Engine/Source/Runtime/RHI/Public/RHIResources.h b/Engine/Source/Runtime/RHI/Public/RHIResources.h
--- a/Engine/Source/Runtime/RHI/Public/RHIResources.h
+++ b/Engine/Source/Runtime/RHI/Public/RHIResources.h
@@ -881,6 +881,8 @@ public:
 	 */
 	virtual bool Poll() const = 0;
 
+	const FName& GetFName() const { return FenceName; }
+
 protected:
 	FName FenceName;
 };
diff --git a/Engine/Source/Runtime/Renderer/Private/PathTracing/RayCounterCompute.cpp b/Engine/Source/Runtime/Renderer/Private/PathTracing/RayCounterCompute.cpp
--- a/Engine/Source/Runtime/Renderer/Private/PathTracing/RayCounterCompute.cpp
+++ b/Engine/Source/Runtime/Renderer/Private/PathTracing/RayCounterCompute.cpp
@@ -108,7 +108,7 @@ void FDeferredShadingSceneRenderer::ComputeRayCount(FRHICommandListImmediate& RH
 	int32 NumGroups = FMath::DivideAndRoundUp<int32>(ViewSize.Y, FRayCounterCS::GetGroupSize());
 	DispatchComputeShader(RHICmdList, *RayCounterComputeShader, NumGroups, 1, 1);
 
-	FRHIGPUMemoryReadback* RayCountGPUReadback = ViewState->RayCountGPUReadback;
+	FRHIGPUBufferReadback* RayCountGPUReadback = ViewState->RayCountGPUReadback;
 
 	// Read read count data from the GPU using a stage buffer to avoid stalls
 	if (!ViewState->bReadbackInitialized)
diff --git a/Engine/Source/Runtime/Renderer/Private/PostProcess/PostProcessEyeAdaptation.cpp b/Engine/Source/Runtime/Renderer/Private/PostProcess/PostProcessEyeAdaptation.cpp
--- a/Engine/Source/Runtime/Renderer/Private/PostProcess/PostProcessEyeAdaptation.cpp
+++ b/Engine/Source/Runtime/Renderer/Private/PostProcess/PostProcessEyeAdaptation.cpp
@@ -12,7 +12,7 @@
 #include "Math/UnrealMathUtility.h"
 #include "ScenePrivate.h"
 #include "Curves/CurveFloat.h"
-
+#include "RHIGPUReadback.h"
 
 RENDERCORE_API bool UsePreExposure(EShaderPlatform Platform);
 
@@ -828,15 +828,17 @@ FPooledRenderTargetDesc FRCPassPostProcessBasicEyeAdaptation::ComputeOutputDesc(
 	return Ret;
 }
 
+//
+FSceneViewState::FEyeAdaptationRTManager::~FEyeAdaptationRTManager()
+{
+}
+
 void FSceneViewState::FEyeAdaptationRTManager::SafeRelease()
 {
 	PooledRenderTarget[0].SafeRelease();
 	PooledRenderTarget[1].SafeRelease();
 
-	for (int32 i = 0; i < NUM_STAGING_BUFFERS; ++i)
-	{
-		StagingBuffers[i].SafeRelease();
-	}
+	ExposureTextureReadback = nullptr;
 }
 
 void FSceneViewState::FEyeAdaptationRTManager::SwapRTs(bool bInUpdateLastExposure)
@@ -846,33 +848,28 @@ void FSceneViewState::FEyeAdaptationRTManager::SwapRTs(bool bInUpdateLastExposur
 	FRHICommandListImmediate& RHICmdList = FRHICommandListExecutor::GetImmediateCommandList();
 	if (bInUpdateLastExposure && PooledRenderTarget[CurrentBuffer].IsValid())
 	{
-		if (!StagingBuffers[CurrentStagingBuffer].IsValid())
+		if (!ExposureTextureReadback)
 		{
-			FPooledRenderTargetDesc Desc(FPooledRenderTargetDesc::Create2DDesc(FIntPoint(1, 1), PF_A32B32G32R32F, FClearValueBinding::None, TexCreate_CPUReadback | TexCreate_HideInVisualizeTexture, TexCreate_None, false));
-			GRenderTargetPool.FindFreeElement(RHICmdList, Desc, StagingBuffers[CurrentStagingBuffer], TEXT("EyeAdaptationCPUReadBack"), true, ERenderTargetTransience::NonTransient);
+			static const FName ExposureValueName(TEXT("Scene view state exposure readback"));
+			ExposureTextureReadback.Reset(new FRHIGPUTextureReadback(ExposureValueName));
+			// Send the first request.
+			ExposureTextureReadback->EnqueueCopy(RHICmdList, PooledRenderTarget[CurrentBuffer]->GetRenderTargetItem().TargetableTexture);
 		}
-
-		// Transfer memory GPU -> CPU
-		RHICmdList.CopyToResolveTarget(PooledRenderTarget[CurrentBuffer]->GetRenderTargetItem().TargetableTexture, StagingBuffers[CurrentStagingBuffer]->GetRenderTargetItem().ShaderResourceTexture, FResolveParams());
-
-		// Take the oldest buffer and read it
-		const int32 NextStagingBuffer = (CurrentStagingBuffer + 1) % NUM_STAGING_BUFFERS;
-		if (StagingBuffers[NextStagingBuffer].IsValid())
+		else if (ExposureTextureReadback->IsReady())
 		{
-			const float* ResultsBuffer = nullptr;
-			int32 BufferWidth = 0, BufferHeight = 0;
-			RHICmdList.MapStagingSurface(StagingBuffers[NextStagingBuffer]->GetRenderTargetItem().ShaderResourceTexture, *(void**)&ResultsBuffer, BufferWidth, BufferHeight);
-
-			if (ResultsBuffer)
+			// Read the last request results.
+			FVector4* ReadbackData = (FVector4*)ExposureTextureReadback->Lock(sizeof(FVector4));
+			if (ReadbackData)
 			{
-				LastExposure = ResultsBuffer[0];
-				LastAverageSceneLuminance = ResultsBuffer[2];
+				LastExposure = ReadbackData->X;
+				LastAverageSceneLuminance = ReadbackData->Z;
+
+				ExposureTextureReadback->Unlock();
 			}
-			RHICmdList.UnmapStagingSurface(StagingBuffers[NextStagingBuffer]->GetRenderTargetItem().ShaderResourceTexture);
-		}
 
-		// Update indices
-		CurrentStagingBuffer = NextStagingBuffer;
+			// Send the request for next update.
+			ExposureTextureReadback->EnqueueCopy(RHICmdList, PooledRenderTarget[CurrentBuffer]->GetRenderTargetItem().TargetableTexture);
+		}
 	}
 
 	CurrentBuffer = 1 - CurrentBuffer;
diff --git a/Engine/Source/Runtime/Renderer/Private/RendererScene.cpp b/Engine/Source/Runtime/Renderer/Private/RendererScene.cpp
--- a/Engine/Source/Runtime/Renderer/Private/RendererScene.cpp
+++ b/Engine/Source/Runtime/Renderer/Private/RendererScene.cpp
@@ -55,8 +55,8 @@
 #include "GPUScene.h"
 #if RHI_RAYTRACING
 #include "RayTracingDynamicGeometryCollection.h"
-#include "RHIGPUReadback.h"
 #endif
+#include "RHIGPUReadback.h"
 
 // Enable this define to do slow checks for components being added to the wrong
 // world's scene, when using PIE. This can happen if a PIE component is reattached
@@ -212,7 +212,7 @@ FSceneViewState::FSceneViewState()
 			});
 	}
 	bReadbackInitialized = false;
-	RayCountGPUReadback = new FRHIGPUMemoryReadback(TEXT("Ray Count Readback"));
+	RayCountGPUReadback = new FRHIGPUBufferReadback(TEXT("Ray Count Readback"));
 #endif
 }
 
diff --git a/Engine/Source/Runtime/Renderer/Private/ScenePrivate.h b/Engine/Source/Runtime/Renderer/Private/ScenePrivate.h
--- a/Engine/Source/Runtime/Renderer/Private/ScenePrivate.h
+++ b/Engine/Source/Runtime/Renderer/Private/ScenePrivate.h
@@ -74,7 +74,8 @@ class UStaticMesh;
 class UStaticMeshComponent;
 class UTextureCube;
 class UWindDirectionalSourceComponent;
-class FRHIGPUMemoryReadback;
+class FRHIGPUBufferReadback;
+class FRHIGPUTextureReadback;
 
 /** Holds information about a single primitive's occlusion. */
 class FPrimitiveOcclusionHistory
@@ -763,7 +764,8 @@ private:
 	{
 	public:
 
-		FEyeAdaptationRTManager() :	CurrentBuffer(0), LastExposure(0.f), CurrentStagingBuffer(0) {}
+		// Allows forward declaration of FRHIGPUTextureReadback
+		~FEyeAdaptationRTManager();
 
 		void SafeRelease();
 
@@ -800,16 +802,14 @@ private:
 
 	private:
 
-		int32 CurrentBuffer;
+		int32 CurrentBuffer = 0;
 
-		float LastExposure;
+		float LastExposure = 0;
 		float LastAverageSceneLuminance = 0; // 0 means invalid. Used for Exposure Compensation Curve.
 
-		int32 CurrentStagingBuffer;
-		static const int32 NUM_STAGING_BUFFERS = 3;
-
 		TRefCountPtr<IPooledRenderTarget> PooledRenderTarget[2];
-		TRefCountPtr<IPooledRenderTarget> StagingBuffers[NUM_STAGING_BUFFERS];
+		TUniquePtr<FRHIGPUTextureReadback> ExposureTextureReadback;
+
 	} EyeAdaptationRTManager;
 
 	// eye adaptation is only valid after it has been computed, not on allocation of the RT
@@ -918,7 +918,7 @@ public:
 	FRWBuffer* TotalRayCountBuffer;
 
 	// Ray Count readback:
-	FRHIGPUMemoryReadback* RayCountGPUReadback;
+	FRHIGPUBufferReadback* RayCountGPUReadback;
 	bool bReadbackInitialized = false;
 
 	// IES light profiles
diff --git a/Engine/Source/Runtime/VulkanRHI/Private/VulkanCommands.cpp b/Engine/Source/Runtime/VulkanRHI/Private/VulkanCommands.cpp
--- a/Engine/Source/Runtime/VulkanRHI/Private/VulkanCommands.cpp
+++ b/Engine/Source/Runtime/VulkanRHI/Private/VulkanCommands.cpp
@@ -1185,14 +1185,10 @@ void FVulkanCommandListContext::PrepareParallelFromBase(const FVulkanCommandList
 	TransitionAndLayoutManager.TempCopy(BaseContext.TransitionAndLayoutManager);
 }
 
-void FVulkanCommandListContext::RHICopyToStagingBuffer(FVertexBufferRHIParamRef SourceBufferRHI, FStagingBufferRHIParamRef StagingBufferRHI, uint32 Offset, uint32 NumBytes, FGPUFenceRHIParamRef FenceRHI)
+void FVulkanCommandListContext::RHICopyToStagingBuffer(FVertexBufferRHIParamRef SourceBufferRHI, FStagingBufferRHIParamRef StagingBufferRHI, uint32 Offset, uint32 NumBytes)
 {
 	FVulkanCmdBuffer* CmdBuffer = CommandBufferManager->GetActiveCmdBuffer();
-
 	FVulkanVertexBuffer* VertexBuffer = ResourceCast(SourceBufferRHI);
-	FVulkanGPUFence* Fence = ResourceCast(FenceRHI);
-	Fence->CmdBuffer = CmdBuffer;
-	Fence->FenceSignaledCounter = CmdBuffer->GetFenceSignaledCounter();
 
 	ensure(CmdBuffer->IsOutsideRenderPass());
 	FVulkanStagingBuffer* StagingBuffer = ResourceCast(StagingBufferRHI);
@@ -1219,6 +1215,15 @@ void FVulkanCommandListContext::RHICopyToStagingBuffer(FVertexBufferRHIParamRef
 	VulkanRHI::vkCmdCopyBuffer(CmdBuffer->GetHandle(), VertexBuffer->GetHandle(), StagingBuffer->StagingBuffer->GetHandle(), 1, &Region);
 }
 
+void FVulkanCommandListContext::RHIWriteGPUFence(FGPUFenceRHIParamRef FenceRHI)
+{
+	FVulkanCmdBuffer* CmdBuffer = CommandBufferManager->GetActiveCmdBuffer();
+	FVulkanGPUFence* Fence = ResourceCast(FenceRHI);
+
+	Fence->CmdBuffer = CmdBuffer;
+	Fence->FenceSignaledCounter = CmdBuffer->GetFenceSignaledCounter();
+}
+
 
 FVulkanCommandContextContainer::FVulkanCommandContextContainer(FVulkanDevice* InDevice)
 	: VulkanRHI::FDeviceChild(InDevice)
diff --git a/Engine/Source/Runtime/VulkanRHI/Public/VulkanContext.h b/Engine/Source/Runtime/VulkanRHI/Public/VulkanContext.h
--- a/Engine/Source/Runtime/VulkanRHI/Public/VulkanContext.h
+++ b/Engine/Source/Runtime/VulkanRHI/Public/VulkanContext.h
@@ -235,7 +235,8 @@ public:
 	virtual void RHICopyTexture(FTextureRHIParamRef SourceTexture, FTextureRHIParamRef DestTexture, const FRHICopyTextureInfo& CopyInfo) final override;
 	virtual void RHITransitionResources(EResourceTransitionAccess TransitionType, FTextureRHIParamRef* InRenderTargets, int32 NumTextures) final override;
 	virtual void RHITransitionResources(EResourceTransitionAccess TransitionType, EResourceTransitionPipeline TransitionPipeline, FUnorderedAccessViewRHIParamRef* InUAVs, int32 NumUAVs, FComputeFenceRHIParamRef WriteComputeFence) final override;
-	virtual void RHICopyToStagingBuffer(FVertexBufferRHIParamRef SourceBuffer, FStagingBufferRHIParamRef DestinationStagingBuffer, uint32 Offset, uint32 NumBytes, FGPUFenceRHIParamRef Fence) final override;
+	virtual void RHICopyToStagingBuffer(FVertexBufferRHIParamRef SourceBuffer, FStagingBufferRHIParamRef DestinationStagingBuffer, uint32 Offset, uint32 NumBytes) final override;
+	virtual void RHIWriteGPUFence(FGPUFenceRHIParamRef Fence) final override;
 
 	// Render time measurement
 	virtual void RHIBeginRenderQuery(FRenderQueryRHIParamRef RenderQuery) final override;
diff --git a/Engine/Source/Runtime/Windows/D3D11RHI/Private/D3D11Commands.cpp b/Engine/Source/Runtime/Windows/D3D11RHI/Private/D3D11Commands.cpp
--- a/Engine/Source/Runtime/Windows/D3D11RHI/Private/D3D11Commands.cpp
+++ b/Engine/Source/Runtime/Windows/D3D11RHI/Private/D3D11Commands.cpp
@@ -2176,7 +2176,7 @@ FStagingBufferRHIRef FD3D11DynamicRHI::RHICreateStagingBuffer()
 	return new FD3D11StagingBuffer();
 }
 
-void FD3D11DynamicRHI::RHICopyToStagingBuffer(FVertexBufferRHIParamRef SourceBufferRHI, FStagingBufferRHIParamRef StagingBufferRHI, uint32 Offset, uint32 NumBytes, FGPUFenceRHIParamRef FenceRHI)
+void FD3D11DynamicRHI::RHICopyToStagingBuffer(FVertexBufferRHIParamRef SourceBufferRHI, FStagingBufferRHIParamRef StagingBufferRHI, uint32 Offset, uint32 NumBytes)
 {
 	FD3D11VertexBuffer* SourceBuffer = ResourceCast(SourceBufferRHI);
 	FD3D11StagingBuffer* StagingBuffer = ResourceCast(StagingBufferRHI);
@@ -2217,15 +2217,15 @@ void FD3D11DynamicRHI::RHICopyToStagingBuffer(FVertexBufferRHIParamRef SourceBuf
 			Direct3DDeviceIMContext->CopySubresourceRegion(StagingBuffer->StagedRead, 0, 0, 0, 0, SourceBuffer->Resource, 0, &SourceBox);
 		}
 	}
+}
 
+void FD3D11DynamicRHI::RHIWriteGPUFence(FGPUFenceRHIParamRef FenceRHI)
+{
 	// @todo-staging Implement real fences for D3D11
 	// D3D11 only has the generic fence for now.
 	FGenericRHIGPUFence* Fence = ResourceCast(FenceRHI);
-
-	if (Fence)
-	{
-		Fence->WriteInternal();
-	}
+	check(Fence);
+	Fence->WriteInternal();
 }
 
 void* FD3D11DynamicRHI::RHILockStagingBuffer(FStagingBufferRHIParamRef StagingBufferRHI, uint32 Offset, uint32 SizeRHI)
diff --git a/Engine/Source/Runtime/Windows/D3D11RHI/Private/D3D11RHIPrivate.h b/Engine/Source/Runtime/Windows/D3D11RHI/Private/D3D11RHIPrivate.h
--- a/Engine/Source/Runtime/Windows/D3D11RHI/Private/D3D11RHIPrivate.h
+++ b/Engine/Source/Runtime/Windows/D3D11RHI/Private/D3D11RHIPrivate.h
@@ -404,7 +404,8 @@ public:
 	virtual FGeometryShaderRHIRef RHICreateGeometryShaderWithStreamOutput(const TArray<uint8>& Code, const FStreamOutElementList& ElementList, uint32 NumStrides, const uint32* Strides, int32 RasterizedStream) final override;
 	virtual FComputeShaderRHIRef RHICreateComputeShader(const TArray<uint8>& Code) final override;
 	virtual FStagingBufferRHIRef RHICreateStagingBuffer() final override;
-	virtual void RHICopyToStagingBuffer(FVertexBufferRHIParamRef SourceBuffer, FStagingBufferRHIParamRef DestinationStagingBuffer, uint32 Offset, uint32 NumBytes, FGPUFenceRHIParamRef Fence = nullptr) final override;
+	virtual void RHICopyToStagingBuffer(FVertexBufferRHIParamRef SourceBuffer, FStagingBufferRHIParamRef DestinationStagingBuffer, uint32 Offset, uint32 NumBytes) final override;
+	virtual void RHIWriteGPUFence(FGPUFenceRHIParamRef Fence) final override;
 	virtual void* RHILockStagingBuffer(FStagingBufferRHIParamRef StagingBuffer, uint32 Offset, uint32 SizeRHI) final override;
     virtual void RHIUnlockStagingBuffer(FStagingBufferRHIParamRef StagingBuffer) final override;
 	virtual FBoundShaderStateRHIRef RHICreateBoundShaderState(FVertexDeclarationRHIParamRef VertexDeclaration, FVertexShaderRHIParamRef VertexShader, FHullShaderRHIParamRef HullShader, FDomainShaderRHIParamRef DomainShader, FPixelShaderRHIParamRef PixelShader, FGeometryShaderRHIParamRef GeometryShader) final override;
