Use 64 bit hash of the shader bytecode in FRayTracingPipelineStateInitializer to generate a more robust shader table signature.

This also allows persistent disk caching of ray tracing pipelines, since shader signatures are stable.

diff --git a/Engine/Source/Runtime/RHI/Public/RHIResources.h b/Engine/Source/Runtime/RHI/Public/RHIResources.h
--- a/Engine/Source/Runtime/RHI/Public/RHIResources.h
+++ b/Engine/Source/Runtime/RHI/Public/RHIResources.h
@@ -2088,20 +2088,22 @@ public:
 		HitGroupHash = ComputeShaderTableHash(HitGroupTable);
 	}
 
-	uint32 GetHitGroupHash() const { return HitGroupHash; }
-	uint32 GetRayGenHash()   const { return RayGenHash; }
-	uint32 GetRayMissHash()  const { return MissHash; }
+	uint64 GetHitGroupHash() const { return HitGroupHash; }
+	uint64 GetRayGenHash()   const { return RayGenHash; }
+	uint64 GetRayMissHash()  const { return MissHash; }
 
 private:
 
-	uint32 ComputeShaderTableHash(const TArrayView<const FRayTracingShaderRHIParamRef>& ShaderTable, uint32 InitialHash = 2085640061)
+	uint64 ComputeShaderTableHash(const TArrayView<const FRayTracingShaderRHIParamRef>& ShaderTable, uint64 InitialHash = 5699878132332235837ull)
 	{
-		uint32 CombinedHash = InitialHash;
+		uint64 CombinedHash = InitialHash;
 		for (FRayTracingShaderRHIParamRef ShaderRHI : ShaderTable)
 		{
-			// #dxr_todo: some sort of session-unique ID should be used instead of pointers to ensure that unique hash is
-			// produced if the same memory happens to be re-used for a different shader (i.e. delete followed by new).
-			CombinedHash = PointerHash(ShaderRHI, CombinedHash);
+			uint64 ShaderHash; // 64 bits from the shader SHA1
+			FMemory::Memcpy(&ShaderHash, ShaderRHI->GetHash().Hash, sizeof(ShaderHash));
+
+			// 64 bit hash combination as per boost::hash_combine_impl
+			CombinedHash ^= ShaderHash + 0x9e3779b9 + (CombinedHash << 6) + (CombinedHash >> 2);
 		}
 
 		return CombinedHash;
@@ -2111,9 +2113,9 @@ private:
 	TArrayView<const FRayTracingShaderRHIParamRef> MissTable;
 	TArrayView<const FRayTracingShaderRHIParamRef> HitGroupTable;
 
-	uint32 RayGenHash = 0;
-	uint32 MissHash = 0;
-	uint32 HitGroupHash = 0;
+	uint64 RayGenHash = 0;
+	uint64 MissHash = 0;
+	uint64 HitGroupHash = 0;
 };
 #endif // RHI_RAYTRACING
 
diff --git a/Engine/Source/Runtime/RenderCore/Private/Shader.cpp b/Engine/Source/Runtime/RenderCore/Private/Shader.cpp
--- a/Engine/Source/Runtime/RenderCore/Private/Shader.cpp
+++ b/Engine/Source/Runtime/RenderCore/Private/Shader.cpp
@@ -936,6 +936,8 @@ void FShaderResource::InitRHI()
 			{
 				RayTracingMaterialLibraryIndex = AddToRayTracingLibrary(RayTracingShader);
 			}
+
+			RayTracingShader->SetHash(OutputHash);
 		}
 	}
 #endif // RHI_RAYTRACING
