DR - Very initial RHI validation layer. When adding support to your rhi, please use "-RHIValidation" so it's a common command line to enable.

diff --git a/Engine/Source/Runtime/RHI/Private/ValidationRHI.cpp b/Engine/Source/Runtime/RHI/Private/ValidationRHI.cpp
new file mode 100644
--- /dev/null
+++ b/Engine/Source/Runtime/RHI/Private/ValidationRHI.cpp
@@ -0,0 +1,47 @@
+// Copyright 1998-2019 Epic Games, Inc. All Rights Reserved.
+
+/*=============================================================================
+	ValidationRHI.cpp: Public Valdation RHI definitions.
+=============================================================================*/
+
+#include "ValidationRHI.h"
+#include "RHI.h"
+#include "ValidationContext.h"
+
+
+FValidationRHI::FValidationRHI(FDynamicRHI* InRHI)
+	: RHI(InRHI)
+{
+	check(RHI);
+
+	Context = new FValidationContext(this);
+}
+
+IRHICommandContext* FValidationRHI::RHIGetDefaultContext()
+{
+	if (!Context->RHIContext)
+	{
+		Context->RHIContext = RHI->RHIGetDefaultContext();
+	}
+
+	return Context;
+}
+
+
+FValidationContext::FValidationContext(FValidationRHI* InRHI)
+	: RHIContext(nullptr)
+	, RHI(InRHI)
+{
+	State.Reset();
+}
+
+void FValidationContext::FState::Reset()
+{
+	bInsideBeginRenderPass = false;
+	bGfxPSOSet = false;
+	RenderPassName.Reset();
+	PreviousRenderPassName.Reset();
+	bInsideComputePass = false;
+	ComputePassName.Reset();
+	bComputeShaderSet = false;
+}
diff --git a/Engine/Source/Runtime/RHI/Public/DynamicRHI.h b/Engine/Source/Runtime/RHI/Public/DynamicRHI.h
--- a/Engine/Source/Runtime/RHI/Public/DynamicRHI.h
+++ b/Engine/Source/Runtime/RHI/Public/DynamicRHI.h
@@ -1156,6 +1156,7 @@ public:
 
 protected:
 	TArray<uint32> PixelFormatBlockBytes;
+	friend class FValidationRHI;
 };
 
 /** A global pointer to the dynamically bound RHI implementation. */
diff --git a/Engine/Source/Runtime/RHI/Public/ValidationContext.h b/Engine/Source/Runtime/RHI/Public/ValidationContext.h
new file mode 100644
--- /dev/null
+++ b/Engine/Source/Runtime/RHI/Public/ValidationContext.h
@@ -0,0 +1,708 @@
+// Copyright 1998-2019 Epic Games, Inc. All Rights Reserved.
+
+/*=============================================================================
+	ValidationContext.h: Public Validation Context definitions.
+=============================================================================*/
+
+#pragma once
+
+#include "RHIResources.h"
+#include "RHIContext.h"
+#include "DynamicRHI.h"
+
+
+class FValidationRHI;
+
+class FValidationContext : public IRHICommandContext
+{
+public:
+	FValidationContext(FValidationRHI* InRHI);
+
+	/**
+	* Compute queue will wait for the fence to be written before continuing.
+	*/
+	virtual void RHIWaitComputeFence(FComputeFenceRHIParamRef InFence) override
+	{
+		RHIContext->RHIWaitComputeFence(InFence);
+	}
+
+	/**
+	*Sets the current compute shader.  Mostly for compliance with platforms
+	*that require shader setting before resource binding.
+	*/
+	virtual void RHISetComputeShader(FComputeShaderRHIParamRef Shader) override
+	{
+		State.bComputeShaderSet = true;
+		State.bGfxPSOSet = false;
+		RHIContext->RHISetComputeShader(Shader);
+	}
+
+	virtual void RHIDispatchComputeShader(uint32 ThreadGroupCountX, uint32 ThreadGroupCountY, uint32 ThreadGroupCountZ) override
+	{
+		RHIContext->RHIDispatchComputeShader(ThreadGroupCountX, ThreadGroupCountY, ThreadGroupCountZ);
+	}
+
+	virtual void RHIDispatchIndirectComputeShader(FVertexBufferRHIParamRef ArgumentBuffer, uint32 ArgumentOffset) override
+	{
+		RHIContext->RHIDispatchIndirectComputeShader(ArgumentBuffer, ArgumentOffset);
+	}
+
+	virtual void RHISetAsyncComputeBudget(EAsyncComputeBudget Budget) override
+	{
+		RHIContext->RHISetAsyncComputeBudget(Budget);
+	}
+
+	virtual void RHIAutomaticCacheFlushAfterComputeShader(bool bEnable) override
+	{
+		RHIContext->RHIAutomaticCacheFlushAfterComputeShader(bEnable);
+	}
+
+	virtual void RHIFlushComputeShaderCache() override
+	{
+		RHIContext->RHIFlushComputeShaderCache();
+	}
+
+	// Useful when used with geometry shader (emit polygons to different viewports), otherwise SetViewPort() is simpler
+	// @param Count >0
+	// @param Data must not be 0
+	virtual void RHISetMultipleViewports(uint32 Count, const FViewportBounds* Data) override
+	{
+		RHIContext->RHISetMultipleViewports(Count, Data);
+	}
+
+	/** Clears a UAV to the multi-component value provided. */
+	virtual void RHIClearTinyUAV(FUnorderedAccessViewRHIParamRef UnorderedAccessViewRHI, const uint32* Values) override
+	{
+		RHIContext->RHIClearTinyUAV(UnorderedAccessViewRHI, Values);
+	}
+
+	/**
+	* Resolves from one texture to another.
+	* @param SourceTexture - texture to resolve from, 0 is silenty ignored
+	* @param DestTexture - texture to resolve to, 0 is silenty ignored
+	* @param ResolveParams - optional resolve params
+	*/
+	virtual void RHICopyToResolveTarget(FTextureRHIParamRef SourceTexture, FTextureRHIParamRef DestTexture, const FResolveParams& ResolveParams) override
+	{
+		RHIContext->RHICopyToResolveTarget(SourceTexture, DestTexture, ResolveParams);
+	}
+
+	/**
+	* Explicitly transition a texture resource from readable -> writable by the GPU or vice versa.
+	* We know rendertargets are only used as rendered targets on the Gfx pipeline, so these transitions are assumed to be implemented such
+	* Gfx->Gfx and Gfx->Compute pipeline transitions are both handled by this call by the RHI implementation.  Hence, no pipeline parameter on this call.
+	*
+	* @param TransitionType - direction of the transition
+	* @param InTextures - array of texture objects to transition
+	* @param NumTextures - number of textures to transition
+	*/
+	virtual void RHITransitionResources(EResourceTransitionAccess TransitionType, FTextureRHIParamRef* InTextures, int32 NumTextures) override
+	{
+		RHIContext->RHITransitionResources(TransitionType, InTextures, NumTextures);
+	}
+
+	/**
+	* Explicitly transition a UAV from readable -> writable by the GPU or vice versa.
+	* Also explicitly states which pipeline the UAV can be used on next.  For example, if a Compute job just wrote this UAV for a Pixel shader to read
+	* you would do EResourceTransitionAccess::Readable and EResourceTransitionPipeline::EComputeToGfx
+	*
+	* @param TransitionType - direction of the transition
+	* @param EResourceTransitionPipeline - How this UAV is transitioning between Gfx and Compute, if at all.
+	* @param InUAVs - array of UAV objects to transition
+	* @param NumUAVs - number of UAVs to transition
+	* @param WriteComputeFence - Optional ComputeFence to write as part of this transition
+	*/
+	virtual void RHITransitionResources(EResourceTransitionAccess TransitionType, EResourceTransitionPipeline TransitionPipeline, FUnorderedAccessViewRHIParamRef* InUAVs, int32 NumUAVs, FComputeFenceRHIParamRef WriteComputeFence) override
+	{
+		RHIContext->RHITransitionResources(TransitionType, TransitionPipeline, InUAVs, NumUAVs, WriteComputeFence);
+	}
+
+	virtual void RHIBeginRenderQuery(FRenderQueryRHIParamRef RenderQuery) override
+	{
+		RHIContext->RHIBeginRenderQuery(RenderQuery);
+	}
+
+	virtual void RHIEndRenderQuery(FRenderQueryRHIParamRef RenderQuery) override
+	{
+		RHIContext->RHIEndRenderQuery(RenderQuery);
+	}
+
+	virtual void RHISubmitCommandsHint() override
+	{
+		ensureMsgf(!State.bInsideBeginRenderPass, TEXT("Submitting inside a RenderPass is not efficient!"));
+		RHIContext->RHISubmitCommandsHint();
+	}
+
+	// Used for OpenGL to check and see if any occlusion queries can be read back on the RHI thread. If they aren't ready when we need them, then we end up stalling.
+	virtual void RHIPollOcclusionQueries() override
+	{
+		RHIContext->RHIPollOcclusionQueries();
+	}
+
+	// Not all RHIs need this (Mobile specific)
+	virtual void RHIDiscardRenderTargets(bool bDepth, bool bStencil, uint32 ColorBitMask) override
+	{
+		RHIContext->RHIDiscardRenderTargets(bDepth, bStencil, ColorBitMask);
+	}
+
+	// This method is queued with an RHIThread, otherwise it will flush after it is queued; without an RHI thread there is no benefit to queuing this frame advance commands
+	virtual void RHIBeginDrawingViewport(FViewportRHIParamRef Viewport, FTextureRHIParamRef RenderTargetRHI) override
+	{
+		RHIContext->RHIBeginDrawingViewport(Viewport, RenderTargetRHI);
+	}
+
+	// This method is queued with an RHIThread, otherwise it will flush after it is queued; without an RHI thread there is no benefit to queuing this frame advance commands
+	virtual void RHIEndDrawingViewport(FViewportRHIParamRef Viewport, bool bPresent, bool bLockToVsync) override
+	{
+		RHIContext->RHIEndDrawingViewport(Viewport, bPresent, bLockToVsync);
+	}
+
+	// This method is queued with an RHIThread, otherwise it will flush after it is queued; without an RHI thread there is no benefit to queuing this frame advance commands
+	virtual void RHIBeginFrame() override
+	{
+		State.Reset();
+		RHIContext->RHIBeginFrame();
+	}
+
+	// This method is queued with an RHIThread, otherwise it will flush after it is queued; without an RHI thread there is no benefit to queuing this frame advance commands
+	virtual void RHIEndFrame() override
+	{
+		RHIContext->RHIEndFrame();
+	}
+
+	/**
+	* Signals the beginning of scene rendering. The RHI makes certain caching assumptions between
+	* calls to BeginScene/EndScene. Currently the only restriction is that you can't update texture
+	* references.
+	*/
+	// This method is queued with an RHIThread, otherwise it will flush after it is queued; without an RHI thread there is no benefit to queuing this frame advance commands
+	virtual void RHIBeginScene() override
+	{
+		RHIContext->RHIBeginScene();
+	}
+
+	/**
+	* Signals the end of scene rendering. See RHIBeginScene.
+	*/
+	// This method is queued with an RHIThread, otherwise it will flush after it is queued; without an RHI thread there is no benefit to queuing this frame advance commands
+	virtual void RHIEndScene() override
+	{
+		RHIContext->RHIEndScene();
+	}
+
+	/**
+	* Signals the beginning and ending of rendering to a resource to be used in the next frame on a multiGPU system
+	*/
+	virtual void RHIBeginUpdateMultiFrameResource(FTextureRHIParamRef Texture) override
+	{
+		RHIContext->RHIBeginUpdateMultiFrameResource(Texture);
+	}
+
+	virtual void RHIEndUpdateMultiFrameResource(FTextureRHIParamRef Texture) override
+	{
+		RHIContext->RHIEndUpdateMultiFrameResource(Texture);
+	}
+
+	virtual void RHIBeginUpdateMultiFrameResource(FUnorderedAccessViewRHIParamRef UAV) override
+	{
+		RHIContext->RHIBeginUpdateMultiFrameResource(UAV);
+	}
+
+	virtual void RHIEndUpdateMultiFrameResource(FUnorderedAccessViewRHIParamRef UAV) override
+	{
+		RHIContext->RHIEndUpdateMultiFrameResource(UAV);
+	}
+
+	virtual void RHISetStreamSource(uint32 StreamIndex, FVertexBufferRHIParamRef VertexBuffer, uint32 Offset) override
+	{
+		//#todo-rco: Decide if this is needed or not...
+		//checkf(State.bGfxPSOSet, TEXT("A Graphics PSO has to be set to set-up the vertex streams!"));
+
+		checkf(State.bInsideBeginRenderPass, TEXT("A RenderPass has to be set to set-up the vertex streams!"));
+		RHIContext->RHISetStreamSource(StreamIndex, VertexBuffer, Offset);
+	}
+
+	// @param MinX including like Win32 RECT
+	// @param MinY including like Win32 RECT
+	// @param MaxX excluding like Win32 RECT
+	// @param MaxY excluding like Win32 RECT
+	virtual void RHISetViewport(uint32 MinX, uint32 MinY, float MinZ, uint32 MaxX, uint32 MaxY, float MaxZ) override
+	{
+		RHIContext->RHISetViewport(MinX, MinY, MinZ, MaxX, MaxY, MaxZ);
+	}
+
+	virtual void RHISetStereoViewport(uint32 LeftMinX, uint32 RightMinX, uint32 LeftMinY, uint32 RightMinY, float MinZ, uint32 LeftMaxX, uint32 RightMaxX, uint32 LeftMaxY, uint32 RightMaxY, float MaxZ) override
+	{
+		RHIContext->RHISetStereoViewport(LeftMinX, RightMinX, LeftMinY, RightMinY, MinZ, LeftMaxX, RightMaxX, LeftMaxY, RightMaxY, MaxZ);
+	}
+
+	// @param MinX including like Win32 RECT
+	// @param MinY including like Win32 RECT
+	// @param MaxX excluding like Win32 RECT
+	// @param MaxY excluding like Win32 RECT
+	virtual void RHISetScissorRect(bool bEnable, uint32 MinX, uint32 MinY, uint32 MaxX, uint32 MaxY) override
+	{
+		RHIContext->RHISetScissorRect(bEnable, MinX, MinY, MaxX, MaxY);
+	}
+
+	virtual void RHISetGraphicsPipelineState(FGraphicsPipelineStateRHIParamRef GraphicsState) override
+	{
+		checkf(State.bInsideBeginRenderPass, TEXT("Graphics PSOs can only be set inside a RenderPass!"));
+		State.bGfxPSOSet = true;
+		State.bComputeShaderSet = false;
+		RHIContext->RHISetGraphicsPipelineState(GraphicsState);
+	}
+
+	/** Set the shader resource view of a surface.  This is used for binding TextureMS parameter types that need a multi sampled view. */
+	virtual void RHISetShaderTexture(FVertexShaderRHIParamRef Shader, uint32 TextureIndex, FTextureRHIParamRef NewTexture) override
+	{
+		checkf(State.bGfxPSOSet, TEXT("A Graphics PSO has to be set to set resources into a shader!"));
+		RHIContext->RHISetShaderTexture(Shader, TextureIndex, NewTexture);
+	}
+
+	/** Set the shader resource view of a surface.  This is used for binding TextureMS parameter types that need a multi sampled view. */
+	virtual void RHISetShaderTexture(FHullShaderRHIParamRef Shader, uint32 TextureIndex, FTextureRHIParamRef NewTexture) override
+	{
+		checkf(State.bGfxPSOSet, TEXT("A Graphics PSO has to be set to set resources into a shader!"));
+		RHIContext->RHISetShaderTexture(Shader, TextureIndex, NewTexture);
+	}
+
+	/** Set the shader resource view of a surface.  This is used for binding TextureMS parameter types that need a multi sampled view. */
+	virtual void RHISetShaderTexture(FDomainShaderRHIParamRef Shader, uint32 TextureIndex, FTextureRHIParamRef NewTexture) override
+	{
+		checkf(State.bGfxPSOSet, TEXT("A Graphics PSO has to be set to set resources into a shader!"));
+		RHIContext->RHISetShaderTexture(Shader, TextureIndex, NewTexture);
+	}
+
+	/** Set the shader resource view of a surface.  This is used for binding TextureMS parameter types that need a multi sampled view. */
+	virtual void RHISetShaderTexture(FGeometryShaderRHIParamRef Shader, uint32 TextureIndex, FTextureRHIParamRef NewTexture) override
+	{
+		checkf(State.bGfxPSOSet, TEXT("A Graphics PSO has to be set to set resources into a shader!"));
+		RHIContext->RHISetShaderTexture(Shader, TextureIndex, NewTexture);
+	}
+
+	/** Set the shader resource view of a surface.  This is used for binding TextureMS parameter types that need a multi sampled view. */
+	virtual void RHISetShaderTexture(FPixelShaderRHIParamRef Shader, uint32 TextureIndex, FTextureRHIParamRef NewTexture) override
+	{
+		checkf(State.bGfxPSOSet, TEXT("A Graphics PSO has to be set to set resources into a shader!"));
+		RHIContext->RHISetShaderTexture(Shader, TextureIndex, NewTexture);
+	}
+
+	/** Set the shader resource view of a surface.  This is used for binding TextureMS parameter types that need a multi sampled view. */
+	virtual void RHISetShaderTexture(FComputeShaderRHIParamRef Shader, uint32 TextureIndex, FTextureRHIParamRef NewTexture) override
+	{
+		checkf(State.bComputeShaderSet, TEXT("A Compute shader has to be set to set resources into a shader!"));
+		RHIContext->RHISetShaderTexture(Shader, TextureIndex, NewTexture);
+	}
+
+	/**
+	* Sets sampler state.
+	* @param Shader	The geometry shader to set the sampler for.
+	* @param SamplerIndex		The index of the sampler.
+	* @param NewState			The new sampler state.
+	*/
+	virtual void RHISetShaderSampler(FComputeShaderRHIParamRef Shader, uint32 SamplerIndex, FSamplerStateRHIParamRef NewState) override
+	{
+		checkf(State.bComputeShaderSet, TEXT("A Compute shader has to be set to set resources into a shader!"));
+		RHIContext->RHISetShaderSampler(Shader, SamplerIndex, NewState);
+	}
+
+	/**
+	* Sets sampler state.
+	* @param Shader	The geometry shader to set the sampler for.
+	* @param SamplerIndex		The index of the sampler.
+	* @param NewState			The new sampler state.
+	*/
+	virtual void RHISetShaderSampler(FVertexShaderRHIParamRef Shader, uint32 SamplerIndex, FSamplerStateRHIParamRef NewState) override
+	{
+		checkf(State.bGfxPSOSet, TEXT("A Graphics PSO has to be set to set resources into a shader!"));
+		RHIContext->RHISetShaderSampler(Shader, SamplerIndex, NewState);
+	}
+
+
+	/**
+	* Sets sampler state.
+	* @param Shader	The geometry shader to set the sampler for.
+	* @param SamplerIndex		The index of the sampler.
+	* @param NewState			The new sampler state.
+	*/
+	virtual void RHISetShaderSampler(FGeometryShaderRHIParamRef Shader, uint32 SamplerIndex, FSamplerStateRHIParamRef NewState) override
+	{
+		checkf(State.bGfxPSOSet, TEXT("A Graphics PSO has to be set to set resources into a shader!"));
+		RHIContext->RHISetShaderSampler(Shader, SamplerIndex, NewState);
+	}
+
+	/**
+	* Sets sampler state.
+	* @param Shader	The geometry shader to set the sampler for.
+	* @param SamplerIndex		The index of the sampler.
+	* @param NewState			The new sampler state.
+	*/
+	virtual void RHISetShaderSampler(FDomainShaderRHIParamRef Shader, uint32 SamplerIndex, FSamplerStateRHIParamRef NewState) override
+	{
+		checkf(State.bGfxPSOSet, TEXT("A Graphics PSO has to be set to set resources into a shader!"));
+		RHIContext->RHISetShaderSampler(Shader, SamplerIndex, NewState);
+	}
+
+	/**
+	* Sets sampler state.
+	* @param Shader	The geometry shader to set the sampler for.
+	* @param SamplerIndex		The index of the sampler.
+	* @param NewState			The new sampler state.
+	*/
+	virtual void RHISetShaderSampler(FHullShaderRHIParamRef Shader, uint32 SamplerIndex, FSamplerStateRHIParamRef NewState) override
+	{
+		checkf(State.bGfxPSOSet, TEXT("A Graphics PSO has to be set to set resources into a shader!"));
+		RHIContext->RHISetShaderSampler(Shader, SamplerIndex, NewState);
+	}
+
+	/**
+	* Sets sampler state.
+	* @param Shader	The geometry shader to set the sampler for.
+	* @param SamplerIndex		The index of the sampler.
+	* @param NewState			The new sampler state.
+	*/
+	virtual void RHISetShaderSampler(FPixelShaderRHIParamRef Shader, uint32 SamplerIndex, FSamplerStateRHIParamRef NewState) override
+	{
+		checkf(State.bGfxPSOSet, TEXT("A Graphics PSO has to be set to set resources into a shader!"));
+		RHIContext->RHISetShaderSampler(Shader, SamplerIndex, NewState);
+	}
+
+	/**
+	* Sets a compute shader UAV parameter.
+	* @param Shader	The compute shader to set the UAV for.
+	* @param UAVIndex		The index of the UAVIndex.
+	* @param UAV			The new UAV.
+	*/
+	virtual void RHISetUAVParameter(FComputeShaderRHIParamRef Shader, uint32 UAVIndex, FUnorderedAccessViewRHIParamRef UAV) override
+	{
+		checkf(State.bComputeShaderSet, TEXT("A Compute shader has to be set to set resources into a shader!"));
+		RHIContext->RHISetUAVParameter(Shader, UAVIndex, UAV);
+	}
+
+	/**
+	* Sets a compute shader counted UAV parameter and initial count
+	* @param Shader	The compute shader to set the UAV for.
+	* @param UAVIndex		The index of the UAVIndex.
+	* @param UAV			The new UAV.
+	* @param InitialCount	The initial number of items in the UAV.
+	*/
+	virtual void RHISetUAVParameter(FComputeShaderRHIParamRef Shader, uint32 UAVIndex, FUnorderedAccessViewRHIParamRef UAV, uint32 InitialCount) override
+	{
+		checkf(State.bComputeShaderSet, TEXT("A Compute shader has to be set to set resources into a shader!"));
+		RHIContext->RHISetUAVParameter(Shader, UAVIndex, UAV, InitialCount);
+	}
+
+	virtual void RHISetShaderResourceViewParameter(FPixelShaderRHIParamRef Shader, uint32 SamplerIndex, FShaderResourceViewRHIParamRef SRV) override
+	{
+		checkf(State.bGfxPSOSet, TEXT("A Graphics PSO has to be set to set resources into a shader!"));
+		RHIContext->RHISetShaderResourceViewParameter(Shader, SamplerIndex, SRV);
+	}
+
+	virtual void RHISetShaderResourceViewParameter(FVertexShaderRHIParamRef Shader, uint32 SamplerIndex, FShaderResourceViewRHIParamRef SRV) override
+	{
+		checkf(State.bGfxPSOSet, TEXT("A Graphics PSO has to be set to set resources into a shader!"));
+		RHIContext->RHISetShaderResourceViewParameter(Shader, SamplerIndex, SRV);
+	}
+
+	virtual void RHISetShaderResourceViewParameter(FComputeShaderRHIParamRef Shader, uint32 SamplerIndex, FShaderResourceViewRHIParamRef SRV) override
+	{
+		checkf(State.bComputeShaderSet, TEXT("A Compute shader has to be set to set resources into a shader!"));
+		RHIContext->RHISetShaderResourceViewParameter(Shader, SamplerIndex, SRV);
+	}
+
+	virtual void RHISetShaderResourceViewParameter(FHullShaderRHIParamRef Shader, uint32 SamplerIndex, FShaderResourceViewRHIParamRef SRV) override
+	{
+		checkf(State.bGfxPSOSet, TEXT("A Graphics PSO has to be set to set resources into a shader!"));
+		RHIContext->RHISetShaderResourceViewParameter(Shader, SamplerIndex, SRV);
+	}
+
+	virtual void RHISetShaderResourceViewParameter(FDomainShaderRHIParamRef Shader, uint32 SamplerIndex, FShaderResourceViewRHIParamRef SRV) override
+	{
+		checkf(State.bGfxPSOSet, TEXT("A Graphics PSO has to be set to set resources into a shader!"));
+		RHIContext->RHISetShaderResourceViewParameter(Shader, SamplerIndex, SRV);
+	}
+
+	virtual void RHISetShaderResourceViewParameter(FGeometryShaderRHIParamRef Shader, uint32 SamplerIndex, FShaderResourceViewRHIParamRef SRV) override
+	{
+		checkf(State.bGfxPSOSet, TEXT("A Graphics PSO has to be set to set resources into a shader!"));
+		RHIContext->RHISetShaderResourceViewParameter(Shader, SamplerIndex, SRV);
+	}
+
+	virtual void RHISetShaderUniformBuffer(FVertexShaderRHIParamRef Shader, uint32 BufferIndex, FUniformBufferRHIParamRef Buffer) override
+	{
+		checkf(State.bGfxPSOSet, TEXT("A Graphics PSO has to be set to set resources into a shader!"));
+		RHIContext->RHISetShaderUniformBuffer(Shader, BufferIndex, Buffer);
+	}
+
+	virtual void RHISetShaderUniformBuffer(FHullShaderRHIParamRef Shader, uint32 BufferIndex, FUniformBufferRHIParamRef Buffer) override
+	{
+		checkf(State.bGfxPSOSet, TEXT("A Graphics PSO has to be set to set resources into a shader!"));
+		RHIContext->RHISetShaderUniformBuffer(Shader, BufferIndex, Buffer);
+	}
+
+	virtual void RHISetShaderUniformBuffer(FDomainShaderRHIParamRef Shader, uint32 BufferIndex, FUniformBufferRHIParamRef Buffer) override
+	{
+		checkf(State.bGfxPSOSet, TEXT("A Graphics PSO has to be set to set resources into a shader!"));
+		RHIContext->RHISetShaderUniformBuffer(Shader, BufferIndex, Buffer);
+	}
+
+	virtual void RHISetShaderUniformBuffer(FGeometryShaderRHIParamRef Shader, uint32 BufferIndex, FUniformBufferRHIParamRef Buffer) override
+	{
+		checkf(State.bGfxPSOSet, TEXT("A Graphics PSO has to be set to set resources into a shader!"));
+		RHIContext->RHISetShaderUniformBuffer(Shader, BufferIndex, Buffer);
+	}
+
+	virtual void RHISetShaderUniformBuffer(FPixelShaderRHIParamRef Shader, uint32 BufferIndex, FUniformBufferRHIParamRef Buffer) override
+	{
+		checkf(State.bGfxPSOSet, TEXT("A Graphics PSO has to be set to set resources into a shader!"));
+		RHIContext->RHISetShaderUniformBuffer(Shader, BufferIndex, Buffer);
+	}
+
+	virtual void RHISetShaderUniformBuffer(FComputeShaderRHIParamRef Shader, uint32 BufferIndex, FUniformBufferRHIParamRef Buffer) override
+	{
+		checkf(State.bComputeShaderSet, TEXT("A Compute shader has to be set to set resources into a shader!"));
+		RHIContext->RHISetShaderUniformBuffer(Shader, BufferIndex, Buffer);
+	}
+
+	virtual void RHISetShaderParameter(FVertexShaderRHIParamRef Shader, uint32 BufferIndex, uint32 BaseIndex, uint32 NumBytes, const void* NewValue) override
+	{
+		checkf(State.bGfxPSOSet, TEXT("A Graphics PSO has to be set to set resources into a shader!"));
+		RHIContext->RHISetShaderParameter(Shader, BufferIndex, BaseIndex, NumBytes, NewValue);
+	}
+
+	virtual void RHISetShaderParameter(FPixelShaderRHIParamRef Shader, uint32 BufferIndex, uint32 BaseIndex, uint32 NumBytes, const void* NewValue) override
+	{
+		checkf(State.bGfxPSOSet, TEXT("A Graphics PSO has to be set to set resources into a shader!"));
+		RHIContext->RHISetShaderParameter(Shader, BufferIndex, BaseIndex, NumBytes, NewValue);
+	}
+
+	virtual void RHISetShaderParameter(FHullShaderRHIParamRef Shader, uint32 BufferIndex, uint32 BaseIndex, uint32 NumBytes, const void* NewValue) override
+	{
+		checkf(State.bGfxPSOSet, TEXT("A Graphics PSO has to be set to set resources into a shader!"));
+		RHIContext->RHISetShaderParameter(Shader, BufferIndex, BaseIndex, NumBytes, NewValue);
+	}
+
+	virtual void RHISetShaderParameter(FDomainShaderRHIParamRef Shader, uint32 BufferIndex, uint32 BaseIndex, uint32 NumBytes, const void* NewValue) override
+	{
+		checkf(State.bGfxPSOSet, TEXT("A Graphics PSO has to be set to set resources into a shader!"));
+		RHIContext->RHISetShaderParameter(Shader, BufferIndex, BaseIndex, NumBytes, NewValue);
+	}
+
+	virtual void RHISetShaderParameter(FGeometryShaderRHIParamRef Shader, uint32 BufferIndex, uint32 BaseIndex, uint32 NumBytes, const void* NewValue) override
+	{
+		checkf(State.bGfxPSOSet, TEXT("A Graphics PSO has to be set to set resources into a shader!"));
+		RHIContext->RHISetShaderParameter(Shader, BufferIndex, BaseIndex, NumBytes, NewValue);
+	}
+
+	virtual void RHISetShaderParameter(FComputeShaderRHIParamRef Shader, uint32 BufferIndex, uint32 BaseIndex, uint32 NumBytes, const void* NewValue) override
+	{
+		checkf(State.bComputeShaderSet, TEXT("A Compute shader has to be set to set resources into a shader!"));
+		RHIContext->RHISetShaderParameter(Shader, BufferIndex, BaseIndex, NumBytes, NewValue);
+	}
+
+	virtual void RHISetStencilRef(uint32 StencilRef) override
+	{
+		checkf(State.bGfxPSOSet, TEXT("A Graphics PSO has to be set to change stencil ref!"));
+		RHIContext->RHISetStencilRef(StencilRef);
+	}
+
+	virtual void RHISetBlendFactor(const FLinearColor& BlendFactor) override
+	{
+		checkf(State.bGfxPSOSet, TEXT("A Graphics PSO has to be set to change blend factor!"));
+		RHIContext->RHISetBlendFactor(BlendFactor);
+	}
+
+	virtual void RHISetRenderTargets(uint32 NumSimultaneousRenderTargets, const FRHIRenderTargetView* NewRenderTargets, const FRHIDepthRenderTargetView* NewDepthStencilTarget, uint32 NumUAVs, const FUnorderedAccessViewRHIParamRef* UAVs) override
+	{
+		RHIContext->RHISetRenderTargets(NumSimultaneousRenderTargets, NewRenderTargets, NewDepthStencilTarget, NumUAVs, UAVs);
+	}
+
+	virtual void RHISetRenderTargetsAndClear(const FRHISetRenderTargetsInfo& RenderTargetsInfo) override
+	{
+		RHIContext->RHISetRenderTargetsAndClear(RenderTargetsInfo);
+	}
+
+	// Bind the clear state of the currently set rendertargets.  This is used by platforms which
+	// need the state of the target when finalizing a hardware clear or a resource transition to SRV
+	// The explicit bind is needed to support parallel rendering (propagate state between contexts).
+	virtual void RHIBindClearMRTValues(bool bClearColor, bool bClearDepth, bool bClearStencil) override
+	{
+		RHIContext->RHIBindClearMRTValues(bClearColor, bClearDepth, bClearStencil);
+	}
+
+	virtual void RHIDrawPrimitive(uint32 BaseVertexIndex, uint32 NumPrimitives, uint32 NumInstances) override
+	{
+		checkf(State.bGfxPSOSet, TEXT("A Graphics PSO has to be set to draw!"));
+		RHIContext->RHIDrawPrimitive(BaseVertexIndex, NumPrimitives, NumInstances);
+	}
+
+	virtual void RHIDrawPrimitiveIndirect(FVertexBufferRHIParamRef ArgumentBuffer, uint32 ArgumentOffset) override
+	{
+		checkf(State.bGfxPSOSet, TEXT("A Graphics PSO has to be set to draw!"));
+		RHIContext->RHIDrawPrimitiveIndirect(ArgumentBuffer, ArgumentOffset);
+	}
+
+	virtual void RHIDrawIndexedIndirect(FIndexBufferRHIParamRef IndexBufferRHI, FStructuredBufferRHIParamRef ArgumentsBufferRHI, int32 DrawArgumentsIndex, uint32 NumInstances) override
+	{
+		checkf(State.bGfxPSOSet, TEXT("A Graphics PSO has to be set to draw!"));
+		RHIContext->RHIDrawIndexedIndirect(IndexBufferRHI, ArgumentsBufferRHI, DrawArgumentsIndex, NumInstances);
+	}
+
+	// @param NumPrimitives need to be >0 
+	virtual void RHIDrawIndexedPrimitive(FIndexBufferRHIParamRef IndexBuffer, int32 BaseVertexIndex, uint32 FirstInstance, uint32 NumVertices, uint32 StartIndex, uint32 NumPrimitives, uint32 NumInstances) override
+	{
+		checkf(State.bGfxPSOSet, TEXT("A Graphics PSO has to be set to draw!"));
+		RHIContext->RHIDrawIndexedPrimitive(IndexBuffer, BaseVertexIndex, FirstInstance, NumVertices, StartIndex, NumPrimitives, NumInstances);
+	}
+
+	virtual void RHIDrawIndexedPrimitiveIndirect(FIndexBufferRHIParamRef IndexBuffer, FVertexBufferRHIParamRef ArgumentBuffer, uint32 ArgumentOffset) override
+	{
+		checkf(State.bGfxPSOSet, TEXT("A Graphics PSO has to be set to draw!"));
+		RHIContext->RHIDrawIndexedPrimitiveIndirect(IndexBuffer, ArgumentBuffer, ArgumentOffset);
+	}
+
+	/**
+	* Sets Depth Bounds range with the given min/max depth.
+	* @param MinDepth	The minimum depth for depth bounds test
+	* @param MaxDepth	The maximum depth for depth bounds test.
+	*					The valid values for fMinDepth and fMaxDepth are such that 0 <= fMinDepth <= fMaxDepth <= 1
+	*/
+	virtual void RHISetDepthBounds(float MinDepth, float MaxDepth) override
+	{
+		checkf(State.bGfxPSOSet, TEXT("A Graphics PSO has to be set to set resources into a shader!"));
+		RHIContext->RHISetDepthBounds(MinDepth, MaxDepth);
+	}
+
+	virtual void RHIPushEvent(const TCHAR* Name, FColor Color) override
+	{
+		RHIContext->RHIPushEvent(Name, Color);
+	}
+
+	virtual void RHIPopEvent() override
+	{
+		RHIContext->RHIPopEvent();
+	}
+
+	virtual void RHIUpdateTextureReference(FTextureReferenceRHIParamRef TextureRef, FTextureRHIParamRef NewTexture) override
+	{
+		RHIContext->RHIUpdateTextureReference(TextureRef, NewTexture);
+	}
+
+	virtual void RHIBeginRenderPass(const FRHIRenderPassInfo& InInfo, const TCHAR* InName) override
+	{
+		checkf(!State.bInsideBeginRenderPass, TEXT("Trying to begin RenderPass '%s', but already inside '%s'!"), *State.RenderPassName, InName);
+		checkf(InName, TEXT("RenderPass should have a name!"));
+		State.bInsideBeginRenderPass = true;
+		State.RenderPassInfo = InInfo;
+		State.RenderPassName = InName;
+		State.bGfxPSOSet = false;
+		RHIContext->RHIBeginRenderPass(InInfo, InName);
+	}
+
+	virtual void RHIEndRenderPass() override
+	{
+		checkf(State.bInsideBeginRenderPass, TEXT("Trying to end a RenderPass but not inside one!"));
+		RHIContext->RHIEndRenderPass();
+		State.bInsideBeginRenderPass = false;
+		State.PreviousRenderPassName = State.RenderPassName;
+		State.bGfxPSOSet = false;
+	}
+
+	virtual void RHIBeginComputePass(const TCHAR* InName) override
+	{
+		checkf(InName && *InName, TEXT("ComputePass should have a name!"));
+		checkf(!State.bInsideBeginRenderPass, TEXT("Can't begin a compute pass from inside RenderPass '%s'"), *State.RenderPassName);
+		checkf(!State.bInsideComputePass, TEXT("Can't begin a compute pass inside from inside ComputePass '%s'"), *State.ComputePassName);
+		State.bInsideComputePass = true;
+		State.ComputePassName = InName;
+		RHIContext->RHIBeginComputePass(InName);
+	}
+
+	virtual void RHIEndComputePass() override
+	{
+		checkf(State.bInsideComputePass, TEXT("Can't end a compute pass without a Begin!"), *State.ComputePassName);
+		RHIContext->RHIEndComputePass();
+		State.bInsideComputePass = false;
+	}
+
+	virtual void RHICopyTexture(FTextureRHIParamRef SourceTexture, FTextureRHIParamRef DestTexture, const FRHICopyTextureInfo& CopyInfo) override
+	{
+		ensureMsgf(!State.bInsideBeginRenderPass, TEXT("Copying inside a RenderPass is not efficient!"));
+		RHIContext->RHICopyTexture(SourceTexture, DestTexture, CopyInfo);
+	}
+
+	virtual void RHIBuildAccelerationStructure(FRayTracingGeometryRHIParamRef Geometry) override
+	{
+		RHIContext->RHIBuildAccelerationStructure(Geometry);
+	}
+
+	virtual void RHIUpdateAccelerationStructures(const TArrayView<const FAccelerationStructureUpdateParams> Params) override
+	{
+		RHIContext->RHIUpdateAccelerationStructures(Params);
+	}
+
+	virtual void RHIBuildAccelerationStructures(const TArrayView<const FAccelerationStructureUpdateParams> Params) override
+	{
+		RHIContext->RHIBuildAccelerationStructures(Params);
+	}
+
+	virtual void RHIBuildAccelerationStructure(FRayTracingSceneRHIParamRef Scene) override
+	{
+		RHIContext->RHIBuildAccelerationStructure(Scene);
+	}
+
+	virtual void RHIRayTraceOcclusion(FRayTracingSceneRHIParamRef Scene,
+		FShaderResourceViewRHIParamRef Rays,
+		FUnorderedAccessViewRHIParamRef Output,
+		uint32 NumRays) override
+	{
+		RHIContext->RHIRayTraceOcclusion(Scene, Rays, Output, NumRays);
+	}
+
+	virtual void RHIRayTraceIntersection(FRayTracingSceneRHIParamRef Scene,
+		FShaderResourceViewRHIParamRef Rays,
+		FUnorderedAccessViewRHIParamRef Output,
+		uint32 NumRays) override
+	{
+		RHIContext->RHIRayTraceIntersection(Scene, Rays, Output, NumRays);
+	}
+
+	virtual void RHIRayTraceDispatch(FRayTracingPipelineStateRHIParamRef RayTracingPipelineState, FRHIRayTracingShader* RayGenShader,
+		FRayTracingSceneRHIParamRef Scene,
+		const FRayTracingShaderBindings& GlobalResourceBindings,
+		uint32 Width, uint32 Height) override
+	{
+		RHIContext->RHIRayTraceDispatch(RayTracingPipelineState, RayGenShader, Scene, GlobalResourceBindings, Width, Height);
+	}
+
+	virtual void RHISetRayTracingHitGroup(
+		FRayTracingSceneRHIParamRef Scene, uint32 InstanceIndex, uint32 SegmentIndex, uint32 ShaderSlot,
+		FRayTracingPipelineStateRHIParamRef Pipeline, uint32 HitGroupIndex,
+		uint32 NumUniformBuffers, const FUniformBufferRHIParamRef* UniformBuffers,
+		uint32 UserData) override
+	{
+		RHIContext->RHISetRayTracingHitGroup(Scene, InstanceIndex, SegmentIndex, ShaderSlot, Pipeline, HitGroupIndex, NumUniformBuffers, UniformBuffers, UserData);
+	}
+
+	IRHICommandContext* RHIContext;
+	FValidationRHI*		RHI;
+
+protected:
+	struct FState
+	{
+		bool bInsideBeginRenderPass;
+		FRHIRenderPassInfo RenderPassInfo;
+		FString RenderPassName;
+		FString PreviousRenderPassName;
+		bool bInsideComputePass;
+		FString ComputePassName;
+		bool bGfxPSOSet;
+		bool bComputeShaderSet;
+
+		void Reset();
+	};
+	FState State;
+
+	friend class FValidationRHI;
+};
diff --git a/Engine/Source/Runtime/RHI/Public/ValidationRHI.h b/Engine/Source/Runtime/RHI/Public/ValidationRHI.h
new file mode 100644
--- /dev/null
+++ b/Engine/Source/Runtime/RHI/Public/ValidationRHI.h
@@ -0,0 +1,1649 @@
+// Copyright 1998-2019 Epic Games, Inc. All Rights Reserved.
+
+/*=============================================================================
+	ValidationRHI.h: Public Valdation RHI definitions.
+=============================================================================*/
+
+#pragma once 
+
+#include "RHI.h"
+#include "RHIResources.h"
+#include "DynamicRHI.h"
+
+class FValidationContext;
+
+class FValidationRHI : public FDynamicRHI
+{
+public:
+	RHI_API FValidationRHI(FDynamicRHI* InRHI);
+
+	virtual void Init() override
+	{
+		RHI->Init();
+		RHIName = RHI->GetName();
+		RHIName += TEXT("_Validation");
+	}
+
+	/** Called after the RHI is initialized; before the render thread is started. */
+	virtual void PostInit() override
+	{
+		// Need to copy this as each DynamicRHI has an instance
+		check(RHI->PixelFormatBlockBytes.Num() <= PixelFormatBlockBytes.Num());
+		RHI->PixelFormatBlockBytes = PixelFormatBlockBytes;
+		RHI->PostInit();
+	}
+
+	/** Shutdown the RHI; handle shutdown and resource destruction before the RHI's actual destructor is called (so that all resources of the RHI are still available for shutdown). */
+	virtual void Shutdown() override
+	{
+		RHI->Shutdown();
+	}
+
+	virtual const TCHAR* GetName() override
+	{
+		return *RHIName;
+	}
+
+	/////// RHI Methods
+
+	// FlushType: Thread safe
+	virtual FSamplerStateRHIRef RHICreateSamplerState(const FSamplerStateInitializerRHI& Initializer) override
+	{
+		return RHI->RHICreateSamplerState(Initializer);
+	}
+
+	// FlushType: Thread safe
+	virtual FRasterizerStateRHIRef RHICreateRasterizerState(const FRasterizerStateInitializerRHI& Initializer) override
+	{
+		return RHI->RHICreateRasterizerState(Initializer);
+	}
+
+	// FlushType: Thread safe
+	virtual FDepthStencilStateRHIRef RHICreateDepthStencilState(const FDepthStencilStateInitializerRHI& Initializer) override
+	{
+		return RHI->RHICreateDepthStencilState(Initializer);
+	}
+
+	// FlushType: Thread safe
+	virtual FBlendStateRHIRef RHICreateBlendState(const FBlendStateInitializerRHI& Initializer) override
+	{
+		return RHI->RHICreateBlendState(Initializer);
+	}
+
+	// FlushType: Wait RHI Thread
+	virtual FVertexDeclarationRHIRef RHICreateVertexDeclaration(const FVertexDeclarationElementList& Elements) override
+	{
+		return RHI->RHICreateVertexDeclaration(Elements);
+	}
+
+	// FlushType: Wait RHI Thread
+	virtual FPixelShaderRHIRef RHICreatePixelShader(const TArray<uint8>& Code) override
+	{
+		return RHI->RHICreatePixelShader(Code);
+	}
+
+	// FlushType: Wait RHI Thread
+	virtual FPixelShaderRHIRef RHICreatePixelShader(FRHIShaderLibrary* Library, FSHAHash Hash) override
+	{
+		return RHI->RHICreatePixelShader(Library, Hash);
+	}
+
+	// FlushType: Wait RHI Thread
+	virtual FVertexShaderRHIRef RHICreateVertexShader(const TArray<uint8>& Code) override
+	{
+		return RHI->RHICreateVertexShader(Code);
+	}
+
+	// FlushType: Wait RHI Thread
+	virtual FVertexShaderRHIRef RHICreateVertexShader(FRHIShaderLibrary* Library, FSHAHash Hash) override
+	{
+		return RHI->RHICreateVertexShader(Library, Hash);
+	}
+
+	// FlushType: Wait RHI Thread
+	virtual FHullShaderRHIRef RHICreateHullShader(const TArray<uint8>& Code) override
+	{
+		return RHI->RHICreateHullShader(Code);
+	}
+
+	// FlushType: Wait RHI Thread
+	virtual FHullShaderRHIRef RHICreateHullShader(FRHIShaderLibrary* Library, FSHAHash Hash) override
+	{
+		return RHI->RHICreateHullShader(Library, Hash);
+	}
+
+	// FlushType: Wait RHI Thread
+	virtual FDomainShaderRHIRef RHICreateDomainShader(const TArray<uint8>& Code) override
+	{
+		return RHI->RHICreateDomainShader(Code);
+	}
+
+	// FlushType: Wait RHI Thread
+	virtual FDomainShaderRHIRef RHICreateDomainShader(FRHIShaderLibrary* Library, FSHAHash Hash) override
+	{
+		return RHI->RHICreateDomainShader(Library, Hash);
+	}
+
+	// FlushType: Wait RHI Thread
+	virtual FGeometryShaderRHIRef RHICreateGeometryShader(const TArray<uint8>& Code) override
+	{
+		return RHI->RHICreateGeometryShader(Code);
+	}
+
+	// FlushType: Wait RHI Thread
+	virtual FGeometryShaderRHIRef RHICreateGeometryShader(FRHIShaderLibrary* Library, FSHAHash Hash) override
+	{
+		return RHI->RHICreateGeometryShader(Library, Hash);
+	}
+
+	/** Creates a geometry shader with stream output ability, defined by ElementList. */
+	// FlushType: Wait RHI Thread
+	virtual FGeometryShaderRHIRef RHICreateGeometryShaderWithStreamOutput(const TArray<uint8>& Code, const FStreamOutElementList& ElementList, uint32 NumStrides, const uint32* Strides, int32 RasterizedStream) override
+	{
+		return RHI->RHICreateGeometryShaderWithStreamOutput(Code, ElementList, NumStrides, Strides, RasterizedStream);
+	}
+
+	/** Creates a geometry shader with stream output ability, defined by ElementList. */
+	// FlushType: Wait RHI Thread
+	virtual FGeometryShaderRHIRef RHICreateGeometryShaderWithStreamOutput(const FStreamOutElementList& ElementList, uint32 NumStrides, const uint32* Strides, int32 RasterizedStream, FRHIShaderLibrary* Library, FSHAHash Hash) override
+	{
+		return RHI->RHICreateGeometryShaderWithStreamOutput(ElementList, NumStrides, Strides, RasterizedStream, Library, Hash);
+	}
+
+	// Some RHIs can have pending messages/logs for error tracking, or debug modes
+	virtual void FlushPendingLogs() override
+	{
+		RHI->FlushPendingLogs();
+	}
+
+	// FlushType: Wait RHI Thread
+	virtual FComputeShaderRHIRef RHICreateComputeShader(const TArray<uint8>& Code) override
+	{
+		return RHI->RHICreateComputeShader(Code);
+	}
+	// FlushType: Wait RHI Thread
+	virtual FComputeShaderRHIRef RHICreateComputeShader(FRHIShaderLibrary* Library, FSHAHash Hash) override
+	{
+		return RHI->RHICreateComputeShader(Library, Hash);
+	}
+
+	/**
+	 * Attempts to open a shader library for the given shader platform & name within the provided directory.
+	 * @param Platform The shader platform for shaders withing the library.
+	 * @param FilePath The directory in which the library should exist.
+	 * @param Name The name of the library, i.e. "Global" or "Unreal" without shader-platform or file-extension qualification.
+	 * @return The new library if one exists and can be constructed, otherwise nil.
+	 */
+	 // FlushType: Must be Thread-Safe.
+	virtual FRHIShaderLibraryRef RHICreateShaderLibrary(EShaderPlatform Platform, FString const& FilePath, FString const& Name) override
+	{
+		return RHI->RHICreateShaderLibrary(Platform, FilePath, Name);
+	}
+
+	/**
+	* Creates a compute fence.  Compute fences are named GPU fences which can be written to once before resetting.
+	* A command to write the fence must be enqueued before any commands to wait on them.  This is enforced on the CPU to avoid GPU hangs.
+	* @param Name - Friendly name for the Fence.  e.g. ReflectionEnvironmentComplete
+	* @return The new Fence.
+	*/
+	// FlushType: Thread safe, but varies depending on the RHI	
+	virtual FComputeFenceRHIRef RHICreateComputeFence(const FName& Name) override
+	{
+		return RHI->RHICreateComputeFence(Name);
+	}
+
+	virtual FGPUFenceRHIRef RHICreateGPUFence(const FName &Name) override
+	{
+		return RHI->RHICreateGPUFence(Name);
+	}
+
+	/**
+	* Creates a staging buffer, which is memory visible to the cpu without any locking.
+	* @return The new staging-buffer.
+	*/
+	// FlushType: Thread safe.	
+	virtual FStagingBufferRHIRef RHICreateStagingBuffer() override
+	{
+		return RHI->RHICreateStagingBuffer();
+	}
+
+	/**
+	 * Lock a staging buffer to read contents on the CPU that were written by the GPU, without having to stall.
+	 * @discussion This function requires that you have issued an CopyToStagingBuffer invocation and verified that the FRHIGPUFence has been signaled before calling.
+	 * @param StagingBuffer The buffer to lock.
+	 * @param Offset The offset in the buffer to return.
+	 * @param SizeRHI The length of the region in the buffer to lock.
+	 * @returns A pointer to the data starting at 'Offset' and of length 'SizeRHI' from 'StagingBuffer', or nullptr when there is an error.
+	 */
+	virtual void* RHILockStagingBuffer(FRHIStagingBuffer* StagingBuffer, uint32 Offset, uint32 SizeRHI) override
+	{
+		return RHI->RHILockStagingBuffer(StagingBuffer, Offset, SizeRHI);
+	}
+
+	/**
+	 * Unlock a staging buffer previously locked with RHILockStagingBuffer.
+	 * @param StagingBuffer The buffer that was previously locked.
+	 */
+	virtual void RHIUnlockStagingBuffer(FRHIStagingBuffer* StagingBuffer) override
+	{
+		RHI->RHIUnlockStagingBuffer(StagingBuffer);
+	}
+
+	/**
+	 * Lock a staging buffer to read contents on the CPU that were written by the GPU, without having to stall.
+	 * @discussion This function requires that you have issued an CopyToStagingBuffer invocation and verified that the FRHIGPUFence has been signaled before calling.
+	 * @param RHICmdList The command-list to execute on or synchronize with.
+	 * @param StagingBuffer The buffer to lock.
+	 * @param Offset The offset in the buffer to return.
+	 * @param SizeRHI The length of the region in the buffer to lock.
+	 * @returns A pointer to the data starting at 'Offset' and of length 'SizeRHI' from 'StagingBuffer', or nullptr when there is an error.
+	 */
+	virtual void* LockStagingBuffer_RenderThread(class FRHICommandListImmediate& RHICmdList, FRHIStagingBuffer* StagingBuffer, uint32 Offset, uint32 SizeRHI) override
+	{
+		return RHI->LockStagingBuffer_RenderThread(RHICmdList, StagingBuffer, Offset, SizeRHI);
+	}
+
+	/**
+	 * Unlock a staging buffer previously locked with LockStagingBuffer_RenderThread.
+	 * @param RHICmdList The command-list to execute on or synchronize with.
+	 * @param StagingBuffer The buffer what was previously locked.
+	 */
+	virtual void UnlockStagingBuffer_RenderThread(class FRHICommandListImmediate& RHICmdList, FRHIStagingBuffer* StagingBuffer) override
+	{
+		RHI->UnlockStagingBuffer_RenderThread(RHICmdList, StagingBuffer);
+	}
+
+	/**
+	* Creates a bound shader state instance which encapsulates a decl, vertex shader, hull shader, domain shader and pixel shader
+	* CAUTION: Even though this is marked as threadsafe, it is only valid to call from the render thread or the RHI thread. It need not be threadsafe unless the RHI support parallel translation.
+	* CAUTION: Platforms that support RHIThread but don't actually have a threadsafe implementation must flush internally with FScopedRHIThreadStaller StallRHIThread(FRHICommandListExecutor::GetImmediateCommandList()); when the call is from the render thread
+	* @param VertexDeclaration - existing vertex decl
+	* @param VertexShader - existing vertex shader
+	* @param HullShader - existing hull shader
+	* @param DomainShader - existing domain shader
+	* @param GeometryShader - existing geometry shader
+	* @param PixelShader - existing pixel shader
+	*/
+	// FlushType: Thread safe, but varies depending on the RHI
+	virtual FBoundShaderStateRHIRef RHICreateBoundShaderState(FRHIVertexDeclaration* VertexDeclaration, FRHIVertexShader* VertexShader, FRHIHullShader* HullShader, FRHIDomainShader* DomainShader, FRHIPixelShader* PixelShader, FRHIGeometryShader* GeometryShader) override
+	{
+		return RHI->RHICreateBoundShaderState(VertexDeclaration, VertexShader, HullShader, DomainShader, PixelShader, GeometryShader);
+	}
+
+	/**
+	* Creates a graphics pipeline state object (PSO) that represents a complete gpu pipeline for rendering.
+	* This function should be considered expensive to call at runtime and may cause hitches as pipelines are compiled.
+	* @param Initializer - Descriptor object defining all the information needed to create the PSO, as well as behavior hints to the RHI.
+	* @return FGraphicsPipelineStateRHIRef that can be bound for rendering; nullptr if the compilation fails.
+	* CAUTION: On certain RHI implementations (eg, ones that do not support runtime compilation) a compilation failure is a Fatal error and this function will not return.
+	* CAUTION: Even though this is marked as threadsafe, it is only valid to call from the render thread or the RHI thread. It need not be threadsafe unless the RHI support parallel translation.
+	* CAUTION: Platforms that support RHIThread but don't actually have a threadsafe implementation must flush internally with FScopedRHIThreadStaller StallRHIThread(FRHICommandListExecutor::GetImmediateCommandList()); when the call is from the render thread
+	*/
+	// FlushType: Thread safe
+	// TODO: [PSO API] Make pure virtual
+	virtual FGraphicsPipelineStateRHIRef RHICreateGraphicsPipelineState(const FGraphicsPipelineStateInitializer& Initializer) override
+	{
+		return RHI->RHICreateGraphicsPipelineState(Initializer);
+	}
+
+	virtual TRefCountPtr<FRHIComputePipelineState> RHICreateComputePipelineState(FRHIComputeShader* ComputeShader) override
+	{
+		return RHI->RHICreateComputePipelineState(ComputeShader);
+	}
+
+	virtual FGraphicsPipelineStateRHIRef RHICreateGraphicsPipelineState(const FGraphicsPipelineStateInitializer& Initializer, FRHIPipelineBinaryLibrary* PipelineBinary) override
+	{
+		return RHI->RHICreateGraphicsPipelineState(Initializer);
+	}
+
+	virtual TRefCountPtr<FRHIComputePipelineState> RHICreateComputePipelineState(FRHIComputeShader* ComputeShader, FRHIPipelineBinaryLibrary* PipelineBinary) override
+	{
+		return RHI->RHICreateComputePipelineState(ComputeShader, PipelineBinary);
+	}
+
+	/**
+	* Creates a uniform buffer.  The contents of the uniform buffer are provided in a parameter, and are immutable.
+	* CAUTION: Even though this is marked as threadsafe, it is only valid to call from the render thread or the RHI thread. Thus is need not be threadsafe on platforms that do not support or aren't using an RHIThread
+	* @param Contents - A pointer to a memory block of size NumBytes that is copied into the new uniform buffer.
+	* @param NumBytes - The number of bytes the uniform buffer should contain.
+	* @return The new uniform buffer.
+	*/
+	// FlushType: Thread safe, but varies depending on the RHI override
+	virtual FUniformBufferRHIRef RHICreateUniformBuffer(const void* Contents, const FRHIUniformBufferLayout& Layout, EUniformBufferUsage Usage, EUniformBufferValidation Validation) override
+	{
+		return RHI->RHICreateUniformBuffer(Contents, Layout, Usage, Validation);
+	}
+
+	virtual void RHIUpdateUniformBuffer(FRHIUniformBuffer* UniformBufferRHI, const void* Contents) override
+	{
+		RHI->RHIUpdateUniformBuffer(UniformBufferRHI, Contents);
+	}
+
+	// FlushType: Wait RHI Thread
+	virtual FIndexBufferRHIRef RHICreateIndexBuffer(uint32 Stride, uint32 Size, uint32 InUsage, FRHIResourceCreateInfo& CreateInfo) override
+	{
+		return RHI->RHICreateIndexBuffer(Stride, Size, InUsage, CreateInfo);
+	}
+
+	// FlushType: Flush RHI Thread
+	virtual void* RHILockIndexBuffer(FRHIIndexBuffer* IndexBuffer, uint32 Offset, uint32 Size, EResourceLockMode LockMode) override
+	{
+		return RHI->RHILockIndexBuffer(IndexBuffer, Offset, Size, LockMode);
+	}
+
+	// FlushType: Flush RHI Thread
+	virtual void RHIUnlockIndexBuffer(FRHIIndexBuffer* IndexBuffer) override
+	{
+		RHI->RHIUnlockIndexBuffer(IndexBuffer);
+	}
+
+	/**
+	* @param ResourceArray - An optional pointer to a resource array containing the resource's data.
+	*/
+	// FlushType: Wait RHI Thread
+	virtual FVertexBufferRHIRef RHICreateVertexBuffer(uint32 Size, uint32 InUsage, FRHIResourceCreateInfo& CreateInfo) override
+	{
+		return RHI->RHICreateVertexBuffer(Size, InUsage, CreateInfo);
+	}
+
+	// FlushType: Flush RHI Thread
+	virtual void* RHILockVertexBuffer(FRHIVertexBuffer* VertexBuffer, uint32 Offset, uint32 SizeRHI, EResourceLockMode LockMode) override
+	{
+		return RHI->RHILockVertexBuffer(VertexBuffer, Offset, SizeRHI, LockMode);
+	}
+
+	// FlushType: Flush RHI Thread
+	virtual void RHIUnlockVertexBuffer(FRHIVertexBuffer* VertexBuffer) override
+	{
+		RHI->RHIUnlockVertexBuffer(VertexBuffer);
+	}
+
+	/** Copies the contents of one vertex buffer to another vertex buffer.  They must have identical sizes. */
+	// FlushType: Flush Immediate (seems dangerous)
+	virtual void RHICopyVertexBuffer(FRHIVertexBuffer* SourceBuffer, FRHIVertexBuffer* DestBuffer) override
+	{
+		RHI->RHICopyVertexBuffer(SourceBuffer, DestBuffer);
+	}
+
+	/**
+	* @param ResourceArray - An optional pointer to a resource array containing the resource's data.
+	*/
+	// FlushType: Wait RHI Thread
+	virtual FStructuredBufferRHIRef RHICreateStructuredBuffer(uint32 Stride, uint32 Size, uint32 InUsage, FRHIResourceCreateInfo& CreateInfo) override
+	{
+		return RHI->RHICreateStructuredBuffer(Stride, Size, InUsage, CreateInfo);
+	}
+
+	// FlushType: Flush RHI Thread
+	virtual void* RHILockStructuredBuffer(FRHIStructuredBuffer* StructuredBuffer, uint32 Offset, uint32 SizeRHI, EResourceLockMode LockMode) override
+	{
+		return RHI->RHILockStructuredBuffer(StructuredBuffer, Offset, SizeRHI, LockMode);
+	}
+
+	// FlushType: Flush RHI Thread
+	virtual void RHIUnlockStructuredBuffer(FRHIStructuredBuffer* StructuredBuffer) override
+	{
+		RHI->RHIUnlockStructuredBuffer(StructuredBuffer);
+	}
+
+	/** Creates an unordered access view of the given structured buffer. */
+	// FlushType: Wait RHI Thread
+	virtual FUnorderedAccessViewRHIRef RHICreateUnorderedAccessView(FRHIStructuredBuffer* StructuredBuffer, bool bUseUAVCounter, bool bAppendBuffer) override
+	{
+		return RHI->RHICreateUnorderedAccessView(StructuredBuffer, bUseUAVCounter, bAppendBuffer);
+	}
+
+	/** Creates an unordered access view of the given texture. */
+	// FlushType: Wait RHI Thread
+	virtual FUnorderedAccessViewRHIRef RHICreateUnorderedAccessView(FRHITexture* Texture, uint32 MipLevel) override
+	{
+		return RHI->RHICreateUnorderedAccessView(Texture, MipLevel);
+	}
+
+	/** Creates an unordered access view of the given vertex buffer. */
+	// FlushType: Wait RHI Thread
+	virtual FUnorderedAccessViewRHIRef RHICreateUnorderedAccessView(FRHIVertexBuffer* VertexBuffer, uint8 Format) override
+	{
+		return RHI->RHICreateUnorderedAccessView(VertexBuffer, Format);
+	}
+
+	/** Creates an unordered access view of the given index buffer. */
+	// FlushType: Wait RHI Thread
+	virtual FUnorderedAccessViewRHIRef RHICreateUnorderedAccessView(FRHIIndexBuffer* IndexBuffer, uint8 Format) override
+	{
+		return RHI->RHICreateUnorderedAccessView(IndexBuffer, Format);
+	}
+
+	/** Creates a shader resource view of the given structured buffer. */
+	// FlushType: Wait RHI Thread
+	virtual FShaderResourceViewRHIRef RHICreateShaderResourceView(FRHIStructuredBuffer* StructuredBuffer) override
+	{
+		return RHI->RHICreateShaderResourceView(StructuredBuffer);
+	}
+
+	/** Creates a shader resource view of the given vertex buffer. */
+	// FlushType: Wait RHI Thread
+	virtual FShaderResourceViewRHIRef RHICreateShaderResourceView(FRHIVertexBuffer* VertexBuffer, uint32 Stride, uint8 Format) override
+	{
+		return RHI->RHICreateShaderResourceView(VertexBuffer, Stride, Format);
+	}
+
+	/** Creates a shader resource view of the given index buffer. */
+	// FlushType: Wait RHI Thread
+	virtual FShaderResourceViewRHIRef RHICreateShaderResourceView(FRHIIndexBuffer* Buffer) override
+	{
+		return RHI->RHICreateShaderResourceView(Buffer);
+	}
+
+	// Must be called on RHI thread timeline
+	// Make sure to call RHIThreadFence(true) afterwards so that parallel translation doesn't refer old resources
+	virtual void RHIUpdateShaderResourceView(FRHIShaderResourceView* SRV, FRHIVertexBuffer* VertexBuffer, uint32 Stride, uint8 Format) override
+	{
+		RHI->RHIUpdateShaderResourceView(SRV, VertexBuffer, Stride, Format);
+	}
+
+	/**
+	* Computes the total size of a 2D texture with the specified parameters.
+	*
+	* @param SizeX - width of the texture to compute
+	* @param SizeY - height of the texture to compute
+	* @param Format - EPixelFormat texture format
+	* @param NumMips - number of mips to compute or 0 for full mip pyramid
+	* @param NumSamples - number of MSAA samples, usually 1
+	* @param Flags - ETextureCreateFlags creation flags
+	* @param OutAlign - Alignment required for this texture.  Output parameter.
+	*/
+	// FlushType: Thread safe
+	virtual uint64 RHICalcTexture2DPlatformSize(uint32 SizeX, uint32 SizeY, uint8 Format, uint32 NumMips, uint32 NumSamples, uint32 Flags, uint32& OutAlign) override
+	{
+		return RHI->RHICalcTexture2DPlatformSize(SizeX, SizeY, Format, NumMips, NumSamples, Flags, OutAlign);
+	}
+
+	/**
+	* Computes the total size of a 3D texture with the specified parameters.
+	* @param SizeX - width of the texture to create
+	* @param SizeY - height of the texture to create
+	* @param SizeZ - depth of the texture to create
+	* @param Format - EPixelFormat texture format
+	* @param NumMips - number of mips to generate or 0 for full mip pyramid
+	* @param Flags - ETextureCreateFlags creation flags
+	* @param OutAlign - Alignment required for this texture.  Output parameter.
+	*/
+	// FlushType: Thread safe
+	virtual uint64 RHICalcTexture3DPlatformSize(uint32 SizeX, uint32 SizeY, uint32 SizeZ, uint8 Format, uint32 NumMips, uint32 Flags, uint32& OutAlign) override
+	{
+		return RHI->RHICalcTexture3DPlatformSize(SizeX, SizeY, SizeZ, Format, NumMips, Flags, OutAlign);
+	}
+
+	/**
+	* Computes the total size of a cube texture with the specified parameters.
+	* @param Size - width/height of the texture to create
+	* @param Format - EPixelFormat texture format
+	* @param NumMips - number of mips to generate or 0 for full mip pyramid
+	* @param Flags - ETextureCreateFlags creation flags
+	* @param OutAlign - Alignment required for this texture.  Output parameter.
+	*/
+	// FlushType: Thread safe
+	virtual uint64 RHICalcTextureCubePlatformSize(uint32 Size, uint8 Format, uint32 NumMips, uint32 Flags, uint32& OutAlign) override
+	{
+		return RHI->RHICalcTextureCubePlatformSize(Size, Format, NumMips, Flags, OutAlign);
+	}
+
+	/**
+	* Retrieves texture memory stats.
+	* safe to call on the main thread
+	*/
+	// FlushType: Thread safe
+	virtual void RHIGetTextureMemoryStats(FTextureMemoryStats& OutStats) override
+	{
+		RHI->RHIGetTextureMemoryStats(OutStats);
+	}
+
+	/**
+	* Fills a texture with to visualize the texture pool memory.
+	*
+	* @param	TextureData		Start address
+	* @param	SizeX			Number of pixels along X
+	* @param	SizeY			Number of pixels along Y
+	* @param	Pitch			Number of bytes between each row
+	* @param	PixelSize		Number of bytes each pixel represents
+	*
+	* @return true if successful, false otherwise
+	*/
+	// FlushType: Flush Immediate
+	virtual bool RHIGetTextureMemoryVisualizeData(FColor* TextureData, int32 SizeX, int32 SizeY, int32 Pitch, int32 PixelSize) override
+	{
+		return RHI->RHIGetTextureMemoryVisualizeData(TextureData, SizeX, SizeY, Pitch, PixelSize);
+	}
+
+	// FlushType: Wait RHI Thread
+	virtual FTextureReferenceRHIRef RHICreateTextureReference(FLastRenderTimeContainer* LastRenderTime) override
+	{
+		return RHI->RHICreateTextureReference(LastRenderTime);
+	}
+
+	/**
+	* Creates a 2D RHI texture resource
+	* @param SizeX - width of the texture to create
+	* @param SizeY - height of the texture to create
+	* @param Format - EPixelFormat texture format
+	* @param NumMips - number of mips to generate or 0 for full mip pyramid
+	* @param NumSamples - number of MSAA samples, usually 1
+	* @param Flags - ETextureCreateFlags creation flags
+	*/
+	// FlushType: Wait RHI Thread
+	virtual FTexture2DRHIRef RHICreateTexture2D(uint32 SizeX, uint32 SizeY, uint8 Format, uint32 NumMips, uint32 NumSamples, uint32 Flags, FRHIResourceCreateInfo& CreateInfo) override
+	{
+		return RHI->RHICreateTexture2D(SizeX, SizeY, Format, NumMips, NumSamples, Flags, CreateInfo);
+	}
+
+	/**
+	* Creates a 2D RHI texture external resource
+	* @param SizeX - width of the texture to create
+	* @param SizeY - height of the texture to create
+	* @param Format - EPixelFormat texture format
+	* @param NumMips - number of mips to generate or 0 for full mip pyramid
+	* @param NumSamples - number of MSAA samples, usually 1
+	* @param Flags - ETextureCreateFlags creation flags
+	*/
+	// FlushType: Wait RHI Thread
+	virtual FTexture2DRHIRef RHICreateTextureExternal2D(uint32 SizeX, uint32 SizeY, uint8 Format, uint32 NumMips, uint32 NumSamples, uint32 Flags, FRHIResourceCreateInfo& CreateInfo) override
+	{
+		return RHI->RHICreateTextureExternal2D(SizeX, SizeY, Format, NumMips, NumSamples, Flags, CreateInfo);
+	}
+
+	/**
+	* Creates an FStructuredBuffer for the RT write mask of a render target
+	* @param RenderTarget - the RT to create the buffer for
+	*/
+	virtual FStructuredBufferRHIRef RHICreateRTWriteMaskBuffer(FRHITexture2D* RenderTarget) override
+	{
+		return RHI->RHICreateRTWriteMaskBuffer(RenderTarget);
+	}
+
+	/**
+	* Thread-safe function that can be used to create a texture outside of the
+	* rendering thread. This function can ONLY be called if GRHISupportsAsyncTextureCreation
+	* is true.  Cannot create rendertargets with this method.
+	* @param SizeX - width of the texture to create
+	* @param SizeY - height of the texture to create
+	* @param Format - EPixelFormat texture format
+	* @param NumMips - number of mips to generate or 0 for full mip pyramid
+	* @param Flags - ETextureCreateFlags creation flags
+	* @param InitialMipData - pointers to mip data with which to create the texture
+	* @param NumInitialMips - how many mips are provided in InitialMipData
+	* @returns a reference to a 2D texture resource
+	*/
+	// FlushType: Thread safe
+	virtual FTexture2DRHIRef RHIAsyncCreateTexture2D(uint32 SizeX, uint32 SizeY, uint8 Format, uint32 NumMips, uint32 Flags, void** InitialMipData, uint32 NumInitialMips) override
+	{
+		return RHI->RHIAsyncCreateTexture2D(SizeX, SizeY, Format, NumMips, Flags, InitialMipData, NumInitialMips);
+	}
+
+	/**
+	* Copies shared mip levels from one texture to another. The textures must have
+	* full mip chains, share the same format, and have the same aspect ratio. This
+	* copy will not cause synchronization with the GPU.
+	* @param DestTexture2D - destination texture
+	* @param SrcTexture2D - source texture
+	*/
+	// FlushType: Flush RHI Thread
+	virtual void RHICopySharedMips(FRHITexture2D* DestTexture2D, FRHITexture2D* SrcTexture2D) override
+	{
+		RHI->RHICopySharedMips(DestTexture2D, SrcTexture2D);
+	}
+
+	/**
+	* Synchronizes the content of a texture resource between two GPUs using a copy operation.
+	* @param Texture - the texture to synchronize.
+	* @param Rect - the rectangle area to update.
+	* @param SrcGPUIndex - the index of the gpu which content will be red from
+	* @param DestGPUIndex - the index of the gpu which content will be updated.
+	* @param PullData - whether the source writes the data to the dest, or the dest reads the data from the source.
+	*/
+	// FlushType: Flush RHI Thread
+	virtual void RHITransferTexture(FRHITexture2D* Texture, FIntRect Rect, uint32 SrcGPUIndex, uint32 DestGPUIndex, bool bPullData) override
+	{
+		RHI->RHITransferTexture(Texture, Rect, SrcGPUIndex, DestGPUIndex, bPullData);
+	}
+
+	/**
+	* Creates a Array RHI texture resource
+	* @param SizeX - width of the texture to create
+	* @param SizeY - height of the texture to create
+	* @param SizeZ - depth of the texture to create
+	* @param Format - EPixelFormat texture format
+	* @param NumMips - number of mips to generate or 0 for full mip pyramid
+	* @param Flags - ETextureCreateFlags creation flags
+	*/
+	// FlushType: Wait RHI Thread
+	virtual FTexture2DArrayRHIRef RHICreateTexture2DArray(uint32 SizeX, uint32 SizeY, uint32 SizeZ, uint8 Format, uint32 NumMips, uint32 Flags, FRHIResourceCreateInfo& CreateInfo) override
+	{
+		return RHI->RHICreateTexture2DArray(SizeX, SizeY, SizeZ, Format, NumMips, Flags, CreateInfo);
+	}
+
+	/**
+	* Creates a 3d RHI texture resource
+	* @param SizeX - width of the texture to create
+	* @param SizeY - height of the texture to create
+	* @param SizeZ - depth of the texture to create
+	* @param Format - EPixelFormat texture format
+	* @param NumMips - number of mips to generate or 0 for full mip pyramid
+	* @param Flags - ETextureCreateFlags creation flags
+	*/
+	// FlushType: Wait RHI Thread
+	virtual FTexture3DRHIRef RHICreateTexture3D(uint32 SizeX, uint32 SizeY, uint32 SizeZ, uint8 Format, uint32 NumMips, uint32 Flags, FRHIResourceCreateInfo& CreateInfo) override
+	{
+		return RHI->RHICreateTexture3D(SizeX, SizeY, SizeZ, Format, NumMips, Flags, CreateInfo);
+	}
+
+	/**
+	* @param Ref may be 0
+	*/
+	// FlushType: Thread safe
+	virtual void RHIGetResourceInfo(FRHITexture* Ref, FRHIResourceInfo& OutInfo) override
+	{
+		return RHI->RHIGetResourceInfo(Ref, OutInfo);
+	}
+
+	/**
+	* Creates a shader resource view for a 2d texture, viewing only a single
+	* mip level. Useful when rendering to one mip while sampling from another.
+	*/
+	// FlushType: Wait RHI Thread
+	virtual FShaderResourceViewRHIRef RHICreateShaderResourceView(FRHITexture2D* Texture2DRHI, uint8 MipLevel) override
+	{
+		return RHI->RHICreateShaderResourceView(Texture2DRHI, MipLevel);
+	}
+
+	/**FRHIResourceInfo
+	* Creates a shader resource view for a 2d texture, with a different
+	* format from the original.  Useful when sampling stencil.
+	*/
+	// FlushType: Wait RHI Thread
+	virtual FShaderResourceViewRHIRef RHICreateShaderResourceView(FRHITexture2D* Texture2DRHI, uint8 MipLevel, uint8 NumMipLevels, uint8 Format) override
+	{
+		return RHI->RHICreateShaderResourceView(Texture2DRHI, MipLevel, NumMipLevels, Format);
+	}
+
+	/**
+	* Creates a shader resource view for a 3d texture, viewing only a single
+	* mip level.
+	*/
+	// FlushType: Wait RHI Thread
+	virtual FShaderResourceViewRHIRef RHICreateShaderResourceView(FRHITexture3D* Texture3DRHI, uint8 MipLevel) override
+	{
+		return RHI->RHICreateShaderResourceView(Texture3DRHI, MipLevel);
+	}
+
+	/**
+	* Creates a shader resource view for a 2d texture array, viewing only a single
+	* mip level.
+	*/
+	// FlushType: Wait RHI Thread
+	virtual FShaderResourceViewRHIRef RHICreateShaderResourceView(FRHITexture2DArray* Texture2DArrayRHI, uint8 MipLevel) override
+	{
+		return RHI->RHICreateShaderResourceView(Texture2DArrayRHI, MipLevel);
+	}
+
+	/**
+	* Creates a shader resource view for a cube texture, viewing only a single
+	* mip level.
+	*/
+	// FlushType: Wait RHI Thread
+	virtual FShaderResourceViewRHIRef RHICreateShaderResourceView(FRHITextureCube* TextureCubeRHI, uint8 MipLevel) override
+	{
+		return RHI->RHICreateShaderResourceView(TextureCubeRHI, MipLevel);
+	}
+
+	/**
+	* Generates mip maps for a texture.
+	*/
+	// FlushType: Flush Immediate (NP: this should be queued on the command list for RHI thread execution, not flushed)
+	virtual void RHIGenerateMips(FRHITexture* Texture) override
+	{
+		return RHI->RHIGenerateMips(Texture);
+	}
+
+	/**
+	* Computes the size in memory required by a given texture.
+	*
+	* @param	TextureRHI		- Texture we want to know the size of, 0 is safely ignored
+	* @return					- Size in Bytes
+	*/
+	// FlushType: Thread safe
+	virtual uint32 RHIComputeMemorySize(FRHITexture* TextureRHI) override
+	{
+		return RHI->RHIComputeMemorySize(TextureRHI);
+	}
+
+	/**
+	* Starts an asynchronous texture reallocation. It may complete immediately if the reallocation
+	* could be performed without any reshuffling of texture memory, or if there isn't enough memory.
+	* The specified status counter will be decremented by 1 when the reallocation is complete (success or failure).
+	*
+	* Returns a new reference to the texture, which will represent the new mip count when the reallocation is complete.
+	* RHIFinalizeAsyncReallocateTexture2D() must be called to complete the reallocation.
+	*
+	* @param Texture2D		- Texture to reallocate
+	* @param NewMipCount	- New number of mip-levels
+	* @param NewSizeX		- New width, in pixels
+	* @param NewSizeY		- New height, in pixels
+	* @param RequestStatus	- Will be decremented by 1 when the reallocation is complete (success or failure).
+	* @return				- New reference to the texture, or an invalid reference upon failure
+	*/
+	// FlushType: Flush RHI Thread
+	// NP: Note that no RHI currently implements this as an async call, we should simplify the API.
+	virtual FTexture2DRHIRef RHIAsyncReallocateTexture2D(FRHITexture2D* Texture2D, int32 NewMipCount, int32 NewSizeX, int32 NewSizeY, FThreadSafeCounter* RequestStatus) override
+	{
+		return RHI->RHIAsyncReallocateTexture2D(Texture2D, NewMipCount, NewSizeX, NewSizeY, RequestStatus);
+	}
+
+	/**
+	* Finalizes an async reallocation request.
+	* If bBlockUntilCompleted is false, it will only poll the status and finalize if the reallocation has completed.
+	*
+	* @param Texture2D				- Texture to finalize the reallocation for
+	* @param bBlockUntilCompleted	- Whether the function should block until the reallocation has completed
+	* @return						- Current reallocation status:
+	*	TexRealloc_Succeeded	Reallocation succeeded
+	*	TexRealloc_Failed		Reallocation failed
+	*	TexRealloc_InProgress	Reallocation is still in progress, try again later
+	*/
+	// FlushType: Wait RHI Thread
+	virtual ETextureReallocationStatus RHIFinalizeAsyncReallocateTexture2D(FRHITexture2D* Texture2D, bool bBlockUntilCompleted) override
+	{
+		return RHI->RHIFinalizeAsyncReallocateTexture2D(Texture2D, bBlockUntilCompleted);
+	}
+
+	/**
+	* Cancels an async reallocation for the specified texture.
+	* This should be called for the new texture, not the original.
+	*
+	* @param Texture				Texture to cancel
+	* @param bBlockUntilCompleted	If true, blocks until the cancellation is fully completed
+	* @return						Reallocation status
+	*/
+	// FlushType: Wait RHI Thread
+	virtual ETextureReallocationStatus RHICancelAsyncReallocateTexture2D(FRHITexture2D* Texture2D, bool bBlockUntilCompleted) override
+	{
+		return RHI->RHICancelAsyncReallocateTexture2D(Texture2D, bBlockUntilCompleted);
+	}
+
+	/**
+	* Locks an RHI texture's mip-map for read/write operations on the CPU
+	* @param Texture - the RHI texture resource to lock, must not be 0
+	* @param MipIndex - index of the mip level to lock
+	* @param LockMode - Whether to lock the texture read-only instead of write-only
+	* @param DestStride - output to retrieve the textures row stride (pitch)
+	* @param bLockWithinMiptail - for platforms that support packed miptails allow locking of individual mip levels within the miptail
+	* @return pointer to the CPU accessible resource data
+	*/
+	// FlushType: Flush RHI Thread
+	virtual void* RHILockTexture2D(FRHITexture2D* Texture, uint32 MipIndex, EResourceLockMode LockMode, uint32& DestStride, bool bLockWithinMiptail) override
+	{
+		return RHI->RHILockTexture2D(Texture, MipIndex, LockMode, DestStride, bLockWithinMiptail);
+	}
+
+	/**
+	* Unlocks a previously locked RHI texture resource
+	* @param Texture - the RHI texture resource to unlock, must not be 0
+	* @param MipIndex - index of the mip level to unlock
+	* @param bLockWithinMiptail - for platforms that support packed miptails allow locking of individual mip levels within the miptail
+	*/
+	// FlushType: Flush RHI Thread
+	virtual void RHIUnlockTexture2D(FRHITexture2D* Texture, uint32 MipIndex, bool bLockWithinMiptail) override
+	{
+		RHI->RHIUnlockTexture2D(Texture, MipIndex, bLockWithinMiptail);
+	}
+
+	/**
+	* Locks an RHI texture's mip-map for read/write operations on the CPU
+	* @param Texture - the RHI texture resource to lock
+	* @param MipIndex - index of the mip level to lock
+	* @param LockMode - Whether to lock the texture read-only instead of write-only
+	* @param DestStride - output to retrieve the textures row stride (pitch)
+	* @param bLockWithinMiptail - for platforms that support packed miptails allow locking of individual mip levels within the miptail
+	* @return pointer to the CPU accessible resource data
+	*/
+	// FlushType: Flush RHI Thread
+	virtual void* RHILockTexture2DArray(FRHITexture2DArray* Texture, uint32 TextureIndex, uint32 MipIndex, EResourceLockMode LockMode, uint32& DestStride, bool bLockWithinMiptail) override
+	{
+		return RHI->RHILockTexture2DArray(Texture, TextureIndex, MipIndex, LockMode, DestStride, bLockWithinMiptail);
+	}
+
+	/**
+	* Unlocks a previously locked RHI texture resource
+	* @param Texture - the RHI texture resource to unlock
+	* @param MipIndex - index of the mip level to unlock
+	* @param bLockWithinMiptail - for platforms that support packed miptails allow locking of individual mip levels within the miptail
+	*/
+	// FlushType: Flush RHI Thread
+	virtual void RHIUnlockTexture2DArray(FRHITexture2DArray* Texture, uint32 TextureIndex, uint32 MipIndex, bool bLockWithinMiptail) override
+	{
+		RHI->RHIUnlockTexture2DArray(Texture, TextureIndex, MipIndex, bLockWithinMiptail);
+	}
+
+	/**
+	* Updates a region of a 2D texture from system memory
+	* @param Texture - the RHI texture resource to update
+	* @param MipIndex - mip level index to be modified
+	* @param UpdateRegion - The rectangle to copy source image data from
+	* @param SourcePitch - size in bytes of each row of the source image
+	* @param SourceData - source image data, starting at the upper left corner of the source rectangle (in same pixel format as texture)
+	*/
+	// FlushType: Flush RHI Thread
+	virtual void RHIUpdateTexture2D(FRHITexture2D* Texture, uint32 MipIndex, const struct FUpdateTextureRegion2D& UpdateRegion, uint32 SourcePitch, const uint8* SourceData) override
+	{
+		RHI->RHIUpdateTexture2D(Texture, MipIndex, UpdateRegion, SourcePitch, SourceData);
+	}
+
+	/**
+	* Updates a region of a 3D texture from system memory
+	* @param Texture - the RHI texture resource to update
+	* @param MipIndex - mip level index to be modified
+	* @param UpdateRegion - The rectangle to copy source image data from
+	* @param SourceRowPitch - size in bytes of each row of the source image, usually Bpp * SizeX
+	* @param SourceDepthPitch - size in bytes of each depth slice of the source image, usually Bpp * SizeX * SizeY
+	* @param SourceData - source image data, starting at the upper left corner of the source rectangle (in same pixel format as texture)
+	*/
+	// FlushType: Flush RHI Thread
+	virtual void RHIUpdateTexture3D(FRHITexture3D* Texture, uint32 MipIndex, const struct FUpdateTextureRegion3D& UpdateRegion, uint32 SourceRowPitch, uint32 SourceDepthPitch, const uint8* SourceData) override
+	{
+		RHI->RHIUpdateTexture3D(Texture, MipIndex, UpdateRegion, SourceRowPitch, SourceDepthPitch, SourceData);
+	}
+
+	/**
+	* Creates a Cube RHI texture resource
+	* @param Size - width/height of the texture to create
+	* @param Format - EPixelFormat texture format
+	* @param NumMips - number of mips to generate or 0 for full mip pyramid
+	* @param Flags - ETextureCreateFlags creation flags
+	*/
+	// FlushType: Wait RHI Thread
+	virtual FTextureCubeRHIRef RHICreateTextureCube(uint32 Size, uint8 Format, uint32 NumMips, uint32 Flags, FRHIResourceCreateInfo& CreateInfo) override
+	{
+		return RHI->RHICreateTextureCube(Size, Format, NumMips, Flags, CreateInfo);
+	}
+
+	/**
+	* Creates a Cube Array RHI texture resource
+	* @param Size - width/height of the texture to create
+	* @param ArraySize - number of array elements of the texture to create
+	* @param Format - EPixelFormat texture format
+	* @param NumMips - number of mips to generate or 0 for full mip pyramid
+	* @param Flags - ETextureCreateFlags creation flags
+	*/
+	// FlushType: Wait RHI Thread
+	virtual FTextureCubeRHIRef RHICreateTextureCubeArray(uint32 Size, uint32 ArraySize, uint8 Format, uint32 NumMips, uint32 Flags, FRHIResourceCreateInfo& CreateInfo) override
+	{
+		return RHI->RHICreateTextureCubeArray(Size, ArraySize, Format, NumMips, Flags, CreateInfo);
+	}
+
+	/**
+	* Locks an RHI texture's mip-map for read/write operations on the CPU
+	* @param Texture - the RHI texture resource to lock
+	* @param MipIndex - index of the mip level to lock
+	* @param LockMode - Whether to lock the texture read-only instead of write-only.
+	* @param DestStride - output to retrieve the textures row stride (pitch)
+	* @param bLockWithinMiptail - for platforms that support packed miptails allow locking of individual mip levels within the miptail
+	* @return pointer to the CPU accessible resource data
+	*/
+	// FlushType: Flush RHI Thread
+	virtual void* RHILockTextureCubeFace(FRHITextureCube* Texture, uint32 FaceIndex, uint32 ArrayIndex, uint32 MipIndex, EResourceLockMode LockMode, uint32& DestStride, bool bLockWithinMiptail) override
+	{
+		return RHI->RHILockTextureCubeFace(Texture, FaceIndex, ArrayIndex, MipIndex, LockMode, DestStride, bLockWithinMiptail);
+	}
+
+	/**
+	* Unlocks a previously locked RHI texture resource
+	* @param Texture - the RHI texture resource to unlock
+	* @param MipIndex - index of the mip level to unlock
+	* @param bLockWithinMiptail - for platforms that support packed miptails allow locking of individual mip levels within the miptail
+	*/
+	// FlushType: Flush RHI Thread
+	virtual void RHIUnlockTextureCubeFace(FRHITextureCube* Texture, uint32 FaceIndex, uint32 ArrayIndex, uint32 MipIndex, bool bLockWithinMiptail) override
+	{
+		RHI->RHIUnlockTextureCubeFace(Texture, FaceIndex, ArrayIndex, MipIndex, bLockWithinMiptail);
+	}
+
+	// FlushType: Thread safe
+	virtual void RHIBindDebugLabelName(FRHITexture* Texture, const TCHAR* Name) override
+	{
+		RHI->RHIBindDebugLabelName(Texture, Name);
+	}
+
+	virtual void RHIBindDebugLabelName(FRHIUnorderedAccessView* UnorderedAccessViewRHI, const TCHAR* Name) override
+	{
+		RHI->RHIBindDebugLabelName(UnorderedAccessViewRHI, Name);
+	}
+
+	/**
+	* Reads the contents of a texture to an output buffer (non MSAA and MSAA) and returns it as a FColor array.
+	* If the format or texture type is unsupported the OutData array will be size 0
+	*/
+	// FlushType: Flush Immediate (seems wrong)
+	virtual void RHIReadSurfaceData(FRHITexture* Texture, FIntRect Rect, TArray<FColor>& OutData, FReadSurfaceDataFlags InFlags) override
+	{
+		RHI->RHIReadSurfaceData(Texture, Rect, OutData, InFlags);
+	}
+
+	// Default fallback; will not work for non-8-bit surfaces and it's extremely slow.
+	virtual void RHIReadSurfaceData(FRHITexture* Texture, FIntRect Rect, TArray<FLinearColor>& OutData, FReadSurfaceDataFlags InFlags) override
+	{
+		RHI->RHIReadSurfaceData(Texture, Rect, OutData, InFlags);
+	}
+
+	/** Watch out for OutData to be 0 (can happen on DXGI_ERROR_DEVICE_REMOVED), don't call RHIUnmapStagingSurface in that case. */
+	// FlushType: Flush Immediate (seems wrong)
+	virtual void RHIMapStagingSurface(FRHITexture* Texture, void*& OutData, int32& OutWidth, int32& OutHeight) override
+	{
+		RHI->RHIMapStagingSurface(Texture, OutData, OutWidth, OutHeight);
+	}
+
+	/** call after a succesful RHIMapStagingSurface() call */
+	// FlushType: Flush Immediate (seems wrong)
+	virtual void RHIUnmapStagingSurface(FRHITexture* Texture) override
+	{
+		RHI->RHIUnmapStagingSurface(Texture);
+	}
+
+	// FlushType: Flush Immediate (seems wrong)
+	virtual void RHIReadSurfaceFloatData(FRHITexture* Texture, FIntRect Rect, TArray<FFloat16Color>& OutData, ECubeFace CubeFace, int32 ArrayIndex, int32 MipIndex) override
+	{
+		RHI->RHIReadSurfaceFloatData(Texture, Rect, OutData, CubeFace, ArrayIndex, MipIndex);
+	}
+
+	// FlushType: Flush Immediate (seems wrong)
+	virtual void RHIRead3DSurfaceFloatData(FRHITexture* Texture, FIntRect Rect, FIntPoint ZMinMax, TArray<FFloat16Color>& OutData) override
+	{
+		RHI->RHIRead3DSurfaceFloatData(Texture, Rect, ZMinMax, OutData);
+	}
+
+	// FlushType: Wait RHI Thread
+	virtual FRenderQueryRHIRef RHICreateRenderQuery(ERenderQueryType QueryType) override
+	{
+		return RHI->RHICreateRenderQuery(QueryType);
+	}
+	// CAUTION: Even though this is marked as threadsafe, it is only valid to call from the render thread. It is need not be threadsafe on platforms that do not support or aren't using an RHIThread
+	// FlushType: Thread safe, but varies by RHI
+	virtual bool RHIGetRenderQueryResult(FRHIRenderQuery* RenderQuery, uint64& OutResult, bool bWait) override
+	{
+		return RHI->RHIGetRenderQueryResult(RenderQuery, OutResult, bWait);
+	}
+
+	// FlushType: Thread safe
+	virtual uint32 RHIGetViewportNextPresentGPUIndex(FRHIViewport* Viewport) override
+	{
+		return RHI->RHIGetViewportNextPresentGPUIndex(Viewport);
+	}
+
+	// With RHI thread, this is the current backbuffer from the perspective of the render thread.
+	// FlushType: Thread safe
+	virtual FTexture2DRHIRef RHIGetViewportBackBuffer(FRHIViewport* Viewport) override
+	{
+		return RHI->RHIGetViewportBackBuffer(Viewport);
+	}
+
+	virtual FUnorderedAccessViewRHIRef RHIGetViewportBackBufferUAV(FRHIViewport* ViewportRHI) override
+	{
+		return RHI->RHIGetViewportBackBufferUAV(ViewportRHI);
+	}
+
+	virtual FTexture2DRHIRef RHIGetFMaskTexture(FRHITexture* SourceTextureRHI) override
+	{
+		return RHI->RHIGetFMaskTexture(SourceTextureRHI);
+	}
+
+	// Only relevant with an RHI thread, this advances the backbuffer for the purpose of GetViewportBackBuffer
+	// FlushType: Thread safe
+	virtual void RHIAdvanceFrameForGetViewportBackBuffer(FRHIViewport* Viewport) override
+	{
+		RHI->RHIAdvanceFrameForGetViewportBackBuffer(Viewport);
+	}
+	/*
+	* Acquires or releases ownership of the platform-specific rendering context for the calling thread
+	*/
+	// FlushType: Flush RHI Thread
+	virtual void RHIAcquireThreadOwnership() override
+	{
+		RHI->RHIAcquireThreadOwnership();
+	}
+
+	// FlushType: Flush RHI Thread
+	virtual void RHIReleaseThreadOwnership() override
+	{
+		RHI->RHIReleaseThreadOwnership();
+	}
+
+	// Flush driver resources. Typically called when switching contexts/threads
+	// FlushType: Flush RHI Thread
+	virtual void RHIFlushResources() override
+	{
+		RHI->RHIFlushResources();
+	}
+
+	/*
+	* Returns the total GPU time taken to render the last frame. Same metric as FPlatformTime::Cycles().
+	*/
+	// FlushType: Thread safe
+	virtual uint32 RHIGetGPUFrameCycles() override
+	{
+		return RHI->RHIGetGPUFrameCycles();
+	}
+
+	//  must be called from the main thread.
+	// FlushType: Thread safe
+	virtual FViewportRHIRef RHICreateViewport(void* WindowHandle, uint32 SizeX, uint32 SizeY, bool bIsFullscreen, EPixelFormat PreferredPixelFormat) override
+	{
+		return RHI->RHICreateViewport(WindowHandle, SizeX, SizeY, bIsFullscreen, PreferredPixelFormat);
+	}
+
+	//  must be called from the main thread.
+	// FlushType: Thread safe
+	virtual void RHIResizeViewport(FRHIViewport* Viewport, uint32 SizeX, uint32 SizeY, bool bIsFullscreen) override
+	{
+		RHI->RHIResizeViewport(Viewport, SizeX, SizeY, bIsFullscreen);
+	}
+
+	virtual void RHIResizeViewport(FRHIViewport* Viewport, uint32 SizeX, uint32 SizeY, bool bIsFullscreen, EPixelFormat PreferredPixelFormat) override
+	{
+		// Default implementation for RHIs that cannot change formats on the fly
+		RHI->RHIResizeViewport(Viewport, SizeX, SizeY, bIsFullscreen, PreferredPixelFormat);
+	}
+
+	//  must be called from the main thread.
+	// FlushType: Thread safe
+	virtual void RHITick(float DeltaTime) override
+	{
+		RHI->RHITick(DeltaTime);
+	}
+
+	/** Sets stream output targets, for use with a geometry shader created with RHICreateGeometryShaderWithStreamOutput. */
+	//@todo this should be a CMDLIST method
+	// FlushType: Flush Immediate (seems wrong)
+	virtual void RHISetStreamOutTargets(uint32 NumTargets, const FVertexBufferRHIParamRef* VertexBuffers, const uint32* Offsets) override
+	{
+		RHI->RHISetStreamOutTargets(NumTargets, VertexBuffers, Offsets);
+	}
+
+	// Blocks the CPU until the GPU catches up and goes idle.
+	// FlushType: Flush Immediate (seems wrong)
+	virtual void RHIBlockUntilGPUIdle() override
+	{
+		RHI->RHIBlockUntilGPUIdle();
+	}
+
+	// Kicks the current frame and makes sure GPU is actively working on them
+	// FlushType: Flush Immediate (copied from RHIBlockUntilGPUIdle)
+	virtual void RHISubmitCommandsAndFlushGPU() override
+	{
+		RHI->RHISubmitCommandsAndFlushGPU();
+	}
+
+	// Tells the RHI we're about to suspend it
+	virtual void RHIBeginSuspendRendering() override
+	{
+		RHI->RHIBeginSuspendRendering();
+	}
+
+	// Operations to suspend title rendering and yield control to the system
+	// FlushType: Thread safe
+	virtual void RHISuspendRendering() override
+	{
+		RHI->RHISuspendRendering();
+	}
+
+	// FlushType: Thread safe
+	virtual void RHIResumeRendering() override
+	{
+		RHI->RHIResumeRendering();
+	}
+
+	// FlushType: Flush Immediate
+	virtual bool RHIIsRenderingSuspended() override
+	{
+		return RHI->RHIIsRenderingSuspended();
+	}
+
+	// FlushType: Flush Immediate
+	virtual bool RHIEnqueueDecompress(uint8_t* SrcBuffer, uint8_t* DestBuffer, int CompressedSize, void* ErrorCodeBuffer) override
+	{
+		return RHI->RHIEnqueueDecompress(SrcBuffer, DestBuffer, CompressedSize, ErrorCodeBuffer);
+	}
+
+	virtual bool RHIEnqueueCompress(uint8_t* SrcBuffer, uint8_t* DestBuffer, int UnCompressedSize, void* ErrorCodeBuffer) override
+	{
+		return RHI->RHIEnqueueCompress(SrcBuffer, DestBuffer, UnCompressedSize, ErrorCodeBuffer);
+	}
+
+	/**
+	*	Retrieve available screen resolutions.
+	*
+	*	@param	Resolutions			TArray<FScreenResolutionRHI> parameter that will be filled in.
+	*	@param	bIgnoreRefreshRate	If true, ignore refresh rates.
+	*
+	*	@return	bool				true if successfully filled the array
+	*/
+	// FlushType: Thread safe
+	virtual bool RHIGetAvailableResolutions(FScreenResolutionArray& Resolutions, bool bIgnoreRefreshRate) override
+	{
+		return RHI->RHIGetAvailableResolutions(Resolutions, bIgnoreRefreshRate);
+	}
+
+	/**
+	* Returns a supported screen resolution that most closely matches input.
+	* @param Width - Input: Desired resolution width in pixels. Output: A width that the platform supports.
+	* @param Height - Input: Desired resolution height in pixels. Output: A height that the platform supports.
+	*/
+	// FlushType: Thread safe
+	virtual void RHIGetSupportedResolution(uint32& Width, uint32& Height) override
+	{
+		RHI->RHIGetSupportedResolution(Width, Height);
+	}
+
+	/**
+	* Function that is used to allocate / free space used for virtual texture mip levels.
+	* Make sure you also update the visible mip levels.
+	* @param Texture - the texture to update, must have been created with TexCreate_Virtual
+	* @param FirstMip - the first mip that should be in memory
+	*/
+	// FlushType: Wait RHI Thread
+	virtual void RHIVirtualTextureSetFirstMipInMemory(FRHITexture2D* Texture, uint32 FirstMip) override
+	{
+		RHI->RHIVirtualTextureSetFirstMipInMemory(Texture, FirstMip);
+	}
+
+	/**
+	* Function that can be used to update which is the first visible mip to the GPU.
+	* @param Texture - the texture to update, must have been created with TexCreate_Virtual
+	* @param FirstMip - the first mip that should be visible to the GPU
+	*/
+	// FlushType: Wait RHI Thread
+	virtual void RHIVirtualTextureSetFirstMipVisible(FRHITexture2D* Texture, uint32 FirstMip) override
+	{
+		RHI->RHIVirtualTextureSetFirstMipVisible(Texture, FirstMip);
+	}
+
+	/**
+	* Called once per frame just before deferred deletion in FRHIResource::FlushPendingDeletes
+	*/
+	// FlushType: called from render thread when RHI thread is flushed 
+	virtual void RHIPerFrameRHIFlushComplete() override
+	{
+		RHI->RHIPerFrameRHIFlushComplete();
+	}
+
+	// FlushType: Wait RHI Thread
+	virtual void RHIExecuteCommandList(FRHICommandList* CmdList) override
+	{
+		return RHI->RHIExecuteCommandList(CmdList);
+	}
+
+	/**
+	* Provides access to the native device. Generally this should be avoided but is useful for third party plugins.
+	*/
+	// FlushType: Flush RHI Thread
+	virtual void* RHIGetNativeDevice() override
+	{
+		return RHI->RHIGetNativeDevice();
+	}
+
+	// FlushType: Thread safe
+	virtual IRHICommandContext* RHIGetDefaultContext() override;
+
+	// FlushType: Thread safe
+	virtual IRHIComputeContext* RHIGetDefaultAsyncComputeContext() override
+	{
+		return RHI->RHIGetDefaultAsyncComputeContext();
+	}
+
+	// FlushType: Thread safe
+	virtual class IRHICommandContextContainer* RHIGetCommandContextContainer(int32 Index, int32 Num) override
+	{
+		return RHI->RHIGetCommandContextContainer(Index, Num);
+	}
+
+	/** Returns a context for sending commands to the given GPU mask. Default implementation is only valid when not using multi-gpu. */
+	virtual IRHICommandContext* RHIGetDefaultContext(FRHIGPUMask GPUMask) override
+	{
+		return RHI->RHIGetDefaultContext(GPUMask);
+	}
+
+	virtual IRHIComputeContext* RHIGetDefaultAsyncComputeContext(FRHIGPUMask GPUMask) override
+	{
+		return RHI->RHIGetDefaultAsyncComputeContext(GPUMask);
+	}
+
+	virtual IRHICommandContextContainer* RHIGetCommandContextContainer(int32 Index, int32 Num, FRHIGPUMask GPUMask) override
+	{
+		return RHI->RHIGetCommandContextContainer(Index, Num, GPUMask);
+	}
+
+	///////// Pass through functions that allow RHIs to optimize certain calls.
+	virtual FVertexBufferRHIRef CreateAndLockVertexBuffer_RenderThread(class FRHICommandListImmediate& RHICmdList, uint32 Size, uint32 InUsage, FRHIResourceCreateInfo& CreateInfo, void*& OutDataBuffer) override
+	{
+		return RHI->CreateAndLockVertexBuffer_RenderThread(RHICmdList, Size, InUsage, CreateInfo, OutDataBuffer);
+	}
+
+	virtual FIndexBufferRHIRef CreateAndLockIndexBuffer_RenderThread(class FRHICommandListImmediate& RHICmdList, uint32 Stride, uint32 Size, uint32 InUsage, FRHIResourceCreateInfo& CreateInfo, void*& OutDataBuffer) override
+	{
+		return RHI->CreateAndLockIndexBuffer_RenderThread(RHICmdList, Stride, Size, InUsage, CreateInfo, OutDataBuffer);
+	}
+
+	virtual FVertexBufferRHIRef CreateVertexBuffer_RenderThread(class FRHICommandListImmediate& RHICmdList, uint32 Size, uint32 InUsage, FRHIResourceCreateInfo& CreateInfo) override
+	{
+		return RHI->CreateVertexBuffer_RenderThread(RHICmdList, Size, InUsage, CreateInfo);
+	}
+
+	virtual FStructuredBufferRHIRef CreateStructuredBuffer_RenderThread(class FRHICommandListImmediate& RHICmdList, uint32 Stride, uint32 Size, uint32 InUsage, FRHIResourceCreateInfo& CreateInfo) override
+	{
+		return RHI->CreateStructuredBuffer_RenderThread(RHICmdList, Stride, Size, InUsage, CreateInfo);
+	}
+
+	virtual FShaderResourceViewRHIRef CreateShaderResourceView_RenderThread(class FRHICommandListImmediate& RHICmdList, FRHIVertexBuffer* VertexBuffer, uint32 Stride, uint8 Format) override
+	{
+		return RHI->CreateShaderResourceView_RenderThread(RHICmdList, VertexBuffer, Stride, Format);
+	}
+
+	virtual FShaderResourceViewRHIRef CreateShaderResourceView_RenderThread(class FRHICommandListImmediate& RHICmdList, FRHIIndexBuffer* Buffer) override
+	{
+		return RHI->CreateShaderResourceView_RenderThread(RHICmdList, Buffer);
+	}
+
+	virtual void* LockVertexBuffer_RenderThread(class FRHICommandListImmediate& RHICmdList, FRHIVertexBuffer* VertexBuffer, uint32 Offset, uint32 SizeRHI, EResourceLockMode LockMode) override
+	{
+		return RHI->LockVertexBuffer_RenderThread(RHICmdList, VertexBuffer, Offset, SizeRHI, LockMode);
+	}
+
+	virtual void UnlockVertexBuffer_RenderThread(class FRHICommandListImmediate& RHICmdList, FRHIVertexBuffer* VertexBuffer) override
+	{
+		return RHI->UnlockVertexBuffer_RenderThread(RHICmdList, VertexBuffer);
+	}
+
+	virtual FTexture2DRHIRef AsyncReallocateTexture2D_RenderThread(class FRHICommandListImmediate& RHICmdList, FRHITexture2D* Texture2D, int32 NewMipCount, int32 NewSizeX, int32 NewSizeY, FThreadSafeCounter* RequestStatus) override
+	{
+		return RHI->AsyncReallocateTexture2D_RenderThread(RHICmdList, Texture2D, NewMipCount, NewSizeX, NewSizeY, RequestStatus);
+	}
+
+	virtual ETextureReallocationStatus FinalizeAsyncReallocateTexture2D_RenderThread(class FRHICommandListImmediate& RHICmdList, FRHITexture2D* Texture2D, bool bBlockUntilCompleted) override
+	{
+		return RHI->FinalizeAsyncReallocateTexture2D_RenderThread(RHICmdList, Texture2D, bBlockUntilCompleted);
+	}
+
+	virtual ETextureReallocationStatus CancelAsyncReallocateTexture2D_RenderThread(class FRHICommandListImmediate& RHICmdList, FRHITexture2D* Texture2D, bool bBlockUntilCompleted) override
+	{
+		return RHI->CancelAsyncReallocateTexture2D_RenderThread(RHICmdList, Texture2D, bBlockUntilCompleted);
+	}
+
+	virtual FIndexBufferRHIRef CreateIndexBuffer_RenderThread(class FRHICommandListImmediate& RHICmdList, uint32 Stride, uint32 Size, uint32 InUsage, FRHIResourceCreateInfo& CreateInfo) override
+	{
+		return RHI->CreateIndexBuffer_RenderThread(RHICmdList, Stride, Size, InUsage, CreateInfo);
+	}
+
+	virtual void* LockIndexBuffer_RenderThread(class FRHICommandListImmediate& RHICmdList, FRHIIndexBuffer* IndexBuffer, uint32 Offset, uint32 SizeRHI, EResourceLockMode LockMode) override
+	{
+		return RHI->LockIndexBuffer_RenderThread(RHICmdList, IndexBuffer, Offset, SizeRHI, LockMode);
+	}
+
+	virtual void UnlockIndexBuffer_RenderThread(class FRHICommandListImmediate& RHICmdList, FRHIIndexBuffer* IndexBuffer) override
+	{
+		RHI->UnlockIndexBuffer_RenderThread(RHICmdList, IndexBuffer);
+	}
+
+	virtual FVertexShaderRHIRef CreateVertexShader_RenderThread(class FRHICommandListImmediate& RHICmdList, const TArray<uint8>& Code) override
+	{
+		return RHI->CreateVertexShader_RenderThread(RHICmdList, Code);
+	}
+
+	virtual FVertexShaderRHIRef CreateVertexShader_RenderThread(class FRHICommandListImmediate& RHICmdList, FRHIShaderLibrary* Library, FSHAHash Hash) override
+	{
+		return RHI->CreateVertexShader_RenderThread(RHICmdList, Library, Hash);
+	}
+
+	virtual FPixelShaderRHIRef CreatePixelShader_RenderThread(class FRHICommandListImmediate& RHICmdList, const TArray<uint8>& Code) override
+	{
+		return RHI->CreatePixelShader_RenderThread(RHICmdList, Code);
+	}
+
+	virtual FPixelShaderRHIRef CreatePixelShader_RenderThread(class FRHICommandListImmediate& RHICmdList, FRHIShaderLibrary* Library, FSHAHash Hash) override
+	{
+		return RHI->CreatePixelShader_RenderThread(RHICmdList, Library, Hash);
+	}
+
+	virtual FGeometryShaderRHIRef CreateGeometryShader_RenderThread(class FRHICommandListImmediate& RHICmdList, const TArray<uint8>& Code) override
+	{
+		return RHI->CreateGeometryShader_RenderThread(RHICmdList, Code);
+	}
+
+	virtual FGeometryShaderRHIRef CreateGeometryShader_RenderThread(class FRHICommandListImmediate& RHICmdList, FRHIShaderLibrary* Library, FSHAHash Hash) override
+	{
+		return RHI->CreateGeometryShader_RenderThread(RHICmdList, Library, Hash);
+	}
+
+	virtual FGeometryShaderRHIRef CreateGeometryShaderWithStreamOutput_RenderThread(class FRHICommandListImmediate& RHICmdList, const TArray<uint8>& Code, const FStreamOutElementList& ElementList, uint32 NumStrides, const uint32* Strides, int32 RasterizedStream) override
+	{
+		return RHI->CreateGeometryShaderWithStreamOutput_RenderThread(RHICmdList, Code, ElementList, NumStrides, Strides, RasterizedStream);
+	}
+
+	virtual FGeometryShaderRHIRef CreateGeometryShaderWithStreamOutput_RenderThread(class FRHICommandListImmediate& RHICmdList, const FStreamOutElementList& ElementList, uint32 NumStrides, const uint32* Strides, int32 RasterizedStream, FRHIShaderLibrary* Library, FSHAHash Hash) override
+	{
+		return RHI->CreateGeometryShaderWithStreamOutput_RenderThread(RHICmdList, ElementList, NumStrides, Strides, RasterizedStream, Library, Hash);
+	}
+
+	virtual FComputeShaderRHIRef CreateComputeShader_RenderThread(class FRHICommandListImmediate& RHICmdList, const TArray<uint8>& Code) override
+	{
+		return RHI->CreateComputeShader_RenderThread(RHICmdList, Code);
+	}
+
+	virtual FComputeShaderRHIRef CreateComputeShader_RenderThread(class FRHICommandListImmediate& RHICmdList, FRHIShaderLibrary* Library, FSHAHash Hash) override
+	{
+		return RHI->CreateComputeShader_RenderThread(RHICmdList, Library, Hash);
+	}
+
+	virtual FHullShaderRHIRef CreateHullShader_RenderThread(class FRHICommandListImmediate& RHICmdList, const TArray<uint8>& Code) override
+	{
+		return RHI->CreateHullShader_RenderThread(RHICmdList, Code);
+	}
+
+	virtual FHullShaderRHIRef CreateHullShader_RenderThread(class FRHICommandListImmediate& RHICmdList, FRHIShaderLibrary* Library, FSHAHash Hash) override
+	{
+		return RHI->CreateHullShader_RenderThread(RHICmdList, Library, Hash);
+	}
+
+	virtual FDomainShaderRHIRef CreateDomainShader_RenderThread(class FRHICommandListImmediate& RHICmdList, const TArray<uint8>& Code) override
+	{
+		return RHI->CreateDomainShader_RenderThread(RHICmdList, Code);
+	}
+
+	virtual FDomainShaderRHIRef CreateDomainShader_RenderThread(class FRHICommandListImmediate& RHICmdList, FRHIShaderLibrary* Library, FSHAHash Hash) override
+	{
+		return RHI->CreateDomainShader_RenderThread(RHICmdList, Library, Hash);
+	}
+
+	virtual void* LockTexture2D_RenderThread(class FRHICommandListImmediate& RHICmdList, FRHITexture2D* Texture, uint32 MipIndex, EResourceLockMode LockMode, uint32& DestStride, bool bLockWithinMiptail, bool bNeedsDefaultRHIFlush = true) override
+	{
+		return RHI->LockTexture2D_RenderThread(RHICmdList, Texture, MipIndex, LockMode, DestStride, bLockWithinMiptail, bNeedsDefaultRHIFlush);
+	}
+
+	virtual void UnlockTexture2D_RenderThread(class FRHICommandListImmediate& RHICmdList, FRHITexture2D* Texture, uint32 MipIndex, bool bLockWithinMiptail, bool bNeedsDefaultRHIFlush = true) override
+	{
+		RHI->UnlockTexture2D_RenderThread(RHICmdList, Texture, MipIndex, bLockWithinMiptail, bNeedsDefaultRHIFlush);
+	}
+
+	virtual void UpdateTexture2D_RenderThread(class FRHICommandListImmediate& RHICmdList, FRHITexture2D* Texture, uint32 MipIndex, const struct FUpdateTextureRegion2D& UpdateRegion, uint32 SourcePitch, const uint8* SourceData) override
+	{
+		RHI->UpdateTexture2D_RenderThread(RHICmdList, Texture, MipIndex, UpdateRegion, SourcePitch, SourceData);
+	}
+
+	virtual FUpdateTexture3DData BeginUpdateTexture3D_RenderThread(class FRHICommandListImmediate& RHICmdList, FRHITexture3D* Texture, uint32 MipIndex, const struct FUpdateTextureRegion3D& UpdateRegion) override
+	{
+		return RHI->BeginUpdateTexture3D_RenderThread(RHICmdList, Texture, MipIndex, UpdateRegion);
+	}
+
+	virtual void EndUpdateTexture3D_RenderThread(class FRHICommandListImmediate& RHICmdList, FUpdateTexture3DData& UpdateData) override
+	{
+		RHI->EndUpdateTexture3D_RenderThread(RHICmdList, UpdateData);
+	}
+
+	virtual void EndMultiUpdateTexture3D_RenderThread(class FRHICommandListImmediate& RHICmdList, TArray<FUpdateTexture3DData>& UpdateDataArray) override
+	{
+		RHI->EndMultiUpdateTexture3D_RenderThread(RHICmdList, UpdateDataArray);
+	}
+
+	virtual void UpdateTexture3D_RenderThread(class FRHICommandListImmediate& RHICmdList, FRHITexture3D* Texture, uint32 MipIndex, const struct FUpdateTextureRegion3D& UpdateRegion, uint32 SourceRowPitch, uint32 SourceDepthPitch, const uint8* SourceData) override
+	{
+		RHI->UpdateTexture3D_RenderThread(RHICmdList, Texture, MipIndex, UpdateRegion, SourceRowPitch, SourceDepthPitch, SourceData);
+	}
+
+	virtual FRHIShaderLibraryRef RHICreateShaderLibrary_RenderThread(class FRHICommandListImmediate& RHICmdList, EShaderPlatform Platform, FString FilePath, FString Name) override
+	{
+		return RHI->RHICreateShaderLibrary_RenderThread(RHICmdList, Platform, FilePath, Name);
+	}
+
+	virtual FTextureReferenceRHIRef RHICreateTextureReference_RenderThread(class FRHICommandListImmediate& RHICmdList, FLastRenderTimeContainer* LastRenderTime) override
+	{
+		return RHI->RHICreateTextureReference_RenderThread(RHICmdList, LastRenderTime);
+	}
+
+	virtual FTexture2DRHIRef RHICreateTexture2D_RenderThread(class FRHICommandListImmediate& RHICmdList, uint32 SizeX, uint32 SizeY, uint8 Format, uint32 NumMips, uint32 NumSamples, uint32 Flags, FRHIResourceCreateInfo& CreateInfo) override
+	{
+		return RHI->RHICreateTexture2D_RenderThread(RHICmdList, SizeX, SizeY, Format, NumMips, NumSamples, Flags, CreateInfo);
+	}
+
+	virtual FTexture2DRHIRef RHICreateTextureExternal2D_RenderThread(class FRHICommandListImmediate& RHICmdList, uint32 SizeX, uint32 SizeY, uint8 Format, uint32 NumMips, uint32 NumSamples, uint32 Flags, FRHIResourceCreateInfo& CreateInfo) override
+	{
+		return RHI->RHICreateTextureExternal2D_RenderThread(RHICmdList, SizeX, SizeY, Format, NumMips, NumSamples, Flags, CreateInfo);
+	}
+
+	virtual FTexture2DArrayRHIRef RHICreateTexture2DArray_RenderThread(class FRHICommandListImmediate& RHICmdList, uint32 SizeX, uint32 SizeY, uint32 SizeZ, uint8 Format, uint32 NumMips, uint32 Flags, FRHIResourceCreateInfo& CreateInfo) override
+	{
+		return RHI->RHICreateTexture2DArray_RenderThread(RHICmdList, SizeX, SizeY, SizeZ, Format, NumMips, Flags, CreateInfo);
+	}
+
+	virtual FTexture3DRHIRef RHICreateTexture3D_RenderThread(class FRHICommandListImmediate& RHICmdList, uint32 SizeX, uint32 SizeY, uint32 SizeZ, uint8 Format, uint32 NumMips, uint32 Flags, FRHIResourceCreateInfo& CreateInfo) override
+	{
+		return RHI->RHICreateTexture3D_RenderThread(RHICmdList, SizeX, SizeY, SizeZ, Format, NumMips, Flags, CreateInfo);
+	}
+
+	virtual FUnorderedAccessViewRHIRef RHICreateUnorderedAccessView_RenderThread(class FRHICommandListImmediate& RHICmdList, FRHIStructuredBuffer* StructuredBuffer, bool bUseUAVCounter, bool bAppendBuffer) override
+	{
+		return RHI->RHICreateUnorderedAccessView_RenderThread(RHICmdList, StructuredBuffer, bUseUAVCounter, bAppendBuffer);
+	}
+
+	virtual FUnorderedAccessViewRHIRef RHICreateUnorderedAccessView_RenderThread(class FRHICommandListImmediate& RHICmdList, FRHITexture* Texture, uint32 MipLevel) override
+	{
+		return RHI->RHICreateUnorderedAccessView_RenderThread(RHICmdList, Texture, MipLevel);
+	}
+
+	virtual FUnorderedAccessViewRHIRef RHICreateUnorderedAccessView_RenderThread(class FRHICommandListImmediate& RHICmdList, FRHIVertexBuffer* VertexBuffer, uint8 Format) override
+	{
+		return RHI->RHICreateUnorderedAccessView_RenderThread(RHICmdList, VertexBuffer, Format);
+	}
+
+	virtual FUnorderedAccessViewRHIRef RHICreateUnorderedAccessView_RenderThread(class FRHICommandListImmediate& RHICmdList, FRHIIndexBuffer* IndexBuffer, uint8 Format) override
+	{
+		return RHI->RHICreateUnorderedAccessView_RenderThread(RHICmdList, IndexBuffer, Format);
+	}
+
+	virtual FShaderResourceViewRHIRef RHICreateShaderResourceView_RenderThread(class FRHICommandListImmediate& RHICmdList, FRHITexture2D* Texture2DRHI, uint8 MipLevel) override
+	{
+		return RHI->RHICreateShaderResourceView_RenderThread(RHICmdList, Texture2DRHI, MipLevel);
+	}
+
+	virtual FShaderResourceViewRHIRef RHICreateShaderResourceView_RenderThread(class FRHICommandListImmediate& RHICmdList, FRHITexture2D* Texture2DRHI, uint8 MipLevel, uint8 NumMipLevels, uint8 Format) override
+	{
+		return RHI->RHICreateShaderResourceView_RenderThread(RHICmdList, Texture2DRHI, MipLevel, NumMipLevels, Format);
+	}
+
+	virtual FShaderResourceViewRHIRef RHICreateShaderResourceView_RenderThread(class FRHICommandListImmediate& RHICmdList, FRHITexture3D* Texture3DRHI, uint8 MipLevel) override
+	{
+		return RHI->RHICreateShaderResourceView_RenderThread(RHICmdList, Texture3DRHI, MipLevel);
+	}
+
+	virtual FShaderResourceViewRHIRef RHICreateShaderResourceView_RenderThread(class FRHICommandListImmediate& RHICmdList, FRHITexture2DArray* Texture2DArrayRHI, uint8 MipLevel) override
+	{
+		return RHI->RHICreateShaderResourceView_RenderThread(RHICmdList, Texture2DArrayRHI, MipLevel);
+	}
+
+	virtual FShaderResourceViewRHIRef RHICreateShaderResourceView_RenderThread(class FRHICommandListImmediate& RHICmdList, FRHITextureCube* TextureCubeRHI, uint8 MipLevel) override
+	{
+		return RHI->RHICreateShaderResourceView_RenderThread(RHICmdList, TextureCubeRHI, MipLevel);
+	}
+
+	virtual FShaderResourceViewRHIRef RHICreateShaderResourceView_RenderThread(class FRHICommandListImmediate& RHICmdList, FRHIVertexBuffer* VertexBuffer, uint32 Stride, uint8 Format) override
+	{
+		return RHI->RHICreateShaderResourceView_RenderThread(RHICmdList, VertexBuffer, Stride, Format);
+	}
+
+	virtual FShaderResourceViewRHIRef RHICreateShaderResourceView_RenderThread(class FRHICommandListImmediate& RHICmdList, FRHIIndexBuffer* Buffer) override
+	{
+		return RHI->RHICreateShaderResourceView_RenderThread(RHICmdList, Buffer);
+	}
+
+	virtual FShaderResourceViewRHIRef RHICreateShaderResourceView_RenderThread(class FRHICommandListImmediate& RHICmdList, FRHIStructuredBuffer* StructuredBuffer) override
+	{
+		return RHI->RHICreateShaderResourceView_RenderThread(RHICmdList, StructuredBuffer);
+	}
+
+	virtual FTextureCubeRHIRef RHICreateTextureCube_RenderThread(class FRHICommandListImmediate& RHICmdList, uint32 Size, uint8 Format, uint32 NumMips, uint32 Flags, FRHIResourceCreateInfo& CreateInfo) override
+	{
+		return RHI->RHICreateTextureCube_RenderThread(RHICmdList, Size, Format, NumMips, Flags, CreateInfo);
+	}
+
+	virtual FTextureCubeRHIRef RHICreateTextureCubeArray_RenderThread(class FRHICommandListImmediate& RHICmdList, uint32 Size, uint32 ArraySize, uint8 Format, uint32 NumMips, uint32 Flags, FRHIResourceCreateInfo& CreateInfo) override
+	{
+		return RHI->RHICreateTextureCubeArray_RenderThread(RHICmdList, Size, ArraySize, Format, NumMips, Flags, CreateInfo);
+	}
+	
+	virtual FRenderQueryRHIRef RHICreateRenderQuery_RenderThread(class FRHICommandListImmediate& RHICmdList, ERenderQueryType QueryType) override
+	{
+		return RHI->RHICreateRenderQuery_RenderThread(RHICmdList, QueryType);
+	}
+
+
+	virtual void* RHILockTextureCubeFace_RenderThread(class FRHICommandListImmediate& RHICmdList, FRHITextureCube* Texture, uint32 FaceIndex, uint32 ArrayIndex, uint32 MipIndex, EResourceLockMode LockMode, uint32& DestStride, bool bLockWithinMiptail) override
+	{
+		return RHI->RHILockTextureCubeFace_RenderThread(RHICmdList, Texture, FaceIndex, ArrayIndex, MipIndex, LockMode, DestStride, bLockWithinMiptail);
+	}
+	virtual void RHIUnlockTextureCubeFace_RenderThread(class FRHICommandListImmediate& RHICmdList, FRHITextureCube* Texture, uint32 FaceIndex, uint32 ArrayIndex, uint32 MipIndex, bool bLockWithinMiptail) override
+	{
+		RHI->RHIUnlockTextureCubeFace_RenderThread(RHICmdList, Texture, FaceIndex, ArrayIndex, MipIndex, bLockWithinMiptail);
+	}
+
+	virtual void RHIAcquireTransientResource_RenderThread(FRHITexture* Texture) override
+	{
+		RHI->RHIAcquireTransientResource_RenderThread(Texture);
+	}
+
+	virtual void RHIDiscardTransientResource_RenderThread(FRHITexture* Texture) override
+	{
+		RHI->RHIDiscardTransientResource_RenderThread(Texture);
+	}
+
+	virtual void RHIAcquireTransientResource_RenderThread(FRHIVertexBuffer* Buffer) override
+	{
+		RHI->RHIAcquireTransientResource_RenderThread(Buffer);
+	}
+
+	virtual void RHIDiscardTransientResource_RenderThread(FRHIVertexBuffer* Buffer) override
+	{
+		RHI->RHIDiscardTransientResource_RenderThread(Buffer);
+	}
+
+	virtual void RHIAcquireTransientResource_RenderThread(FRHIStructuredBuffer* Buffer) override
+	{
+		RHI->RHIAcquireTransientResource_RenderThread(Buffer);
+	}
+
+	virtual void RHIDiscardTransientResource_RenderThread(FRHIStructuredBuffer* Buffer) override
+	{
+		RHI->RHIDiscardTransientResource_RenderThread(Buffer);
+	}
+
+	virtual void RHIReadSurfaceFloatData_RenderThread(class FRHICommandListImmediate& RHICmdList, FRHITexture* Texture, FIntRect Rect, TArray<FFloat16Color>& OutData, ECubeFace CubeFace, int32 ArrayIndex, int32 MipIndex) override
+	{
+		RHI->RHIReadSurfaceFloatData_RenderThread(RHICmdList, Texture, Rect, OutData, CubeFace, ArrayIndex, MipIndex);
+	}
+
+	//Utilities
+	virtual void EnableIdealGPUCaptureOptions(bool bEnable) override
+	{
+		RHI->EnableIdealGPUCaptureOptions(bEnable);
+	}
+
+	virtual void RHISetResourceAliasability_RenderThread(class FRHICommandListImmediate& RHICmdList, EResourceAliasability AliasMode, FRHITexture** InTextures, int32 NumTextures) override
+	{
+		RHI->RHISetResourceAliasability_RenderThread(RHICmdList, AliasMode, InTextures, NumTextures);
+	}
+
+	//checks if the GPU is still alive.
+	virtual bool CheckGpuHeartbeat() const override
+	{
+		return RHI->CheckGpuHeartbeat();
+	}
+
+	virtual void VirtualTextureSetFirstMipInMemory_RenderThread(class FRHICommandListImmediate& RHICmdList, FRHITexture2D* Texture, uint32 FirstMip) override
+	{
+		RHI->VirtualTextureSetFirstMipInMemory_RenderThread(RHICmdList, Texture, FirstMip);
+	}
+
+	virtual void VirtualTextureSetFirstMipVisible_RenderThread(class FRHICommandListImmediate& RHICmdList, FRHITexture2D* Texture, uint32 FirstMip) override
+	{
+		RHI->VirtualTextureSetFirstMipVisible_RenderThread(RHICmdList, Texture, FirstMip);
+	}
+
+	/* Copy the source box pixels in the destination box texture, return true if implemented for the current platform*/
+	virtual void RHICopySubTextureRegion_RenderThread(class FRHICommandListImmediate& RHICmdList, FRHITexture2D* SourceTexture, FRHITexture2D* DestinationTexture, FBox2D SourceBox, FBox2D DestinationBox) override
+	{
+		return RHI->RHICopySubTextureRegion_RenderThread(RHICmdList, SourceTexture, DestinationTexture, SourceBox, DestinationBox);
+	}
+
+	virtual void RHICopySubTextureRegion(FRHITexture2D* SourceTexture, FRHITexture2D* DestinationTexture, FBox2D SourceBox, FBox2D DestinationBox) override
+	{
+		RHI->RHICopySubTextureRegion(SourceTexture, DestinationTexture, SourceBox, DestinationBox);
+	}
+
+	virtual FRHIFlipDetails RHIWaitForFlip(double TimeoutInSeconds) override
+	{
+		return RHI->RHIWaitForFlip(TimeoutInSeconds);
+	}
+
+	virtual void RHISignalFlipEvent() override
+	{
+		RHI->RHISignalFlipEvent();
+	}
+
+	virtual void RHICalibrateTimers() override
+	{
+		RHI->RHICalibrateTimers();
+	}
+
+	virtual void RHIPollRenderQueryResults() override
+	{
+		RHI->RHIPollRenderQueryResults();
+	}
+
+	virtual bool RHIIsTypedUAVLoadSupported(EPixelFormat PixelFormat) override
+	{
+		return RHI->RHIIsTypedUAVLoadSupported(PixelFormat);
+	}
+
+	virtual uint16 RHIGetPlatformTextureMaxSampleCount() override
+	{
+		return RHI->RHIGetPlatformTextureMaxSampleCount();
+	};
+
+
+#if RHI_RAYTRACING
+	// #dxr_todo: reverted index buffer may be needed to support mirrored meshes
+	virtual FRayTracingGeometryRHIRef RHICreateRayTracingGeometry(const FRayTracingGeometryInitializer& Initializer) override
+	{
+		return RHI->RHICreateRayTracingGeometry(Initializer);
+	}
+
+	// This is a placeholder basic interface for the scene. We will likely need to revise this based on real use cases.
+	virtual FRayTracingSceneRHIRef RHICreateRayTracingScene(const FRayTracingSceneInitializer& Initializer) override
+	{
+		return RHI->RHICreateRayTracingScene(Initializer);
+	}
+
+	virtual FRayTracingShaderRHIRef RHICreateRayTracingShader(const TArray<uint8>& Code, EShaderFrequency ShaderFrequency) override
+	{
+		return RHI->RHICreateRayTracingShader(Code, ShaderFrequency);
+	}
+
+	virtual FRayTracingPipelineStateRHIRef RHICreateRayTracingPipelineState(const FRayTracingPipelineStateInitializer& Initializer) override
+	{
+		return RHI->RHICreateRayTracingPipelineState(Initializer);
+	}
+#endif // RHI_RAYTRACING
+
+
+//protected:
+	FDynamicRHI*		RHI;
+	FValidationContext* Context;
+
+private:
+	FString RHIName;
+};
