Fixed a few compile errors on platforms with WITH_MGPU=0

diff --git a/Engine/Source/Runtime/RHI/Public/ValidationRHI.h b/Engine/Source/Runtime/RHI/Public/ValidationRHI.h
--- a/Engine/Source/Runtime/RHI/Public/ValidationRHI.h
+++ b/Engine/Source/Runtime/RHI/Public/ValidationRHI.h
@@ -1204,6 +1204,7 @@ public:
 		return RHI->RHIGetCommandContextContainer(Index, Num);
 	}
 
+#if WITH_MGPU
 	/** Returns a context for sending commands to the given GPU mask. Default implementation is only valid when not using multi-gpu. */
 	virtual IRHICommandContext* RHIGetDefaultContext(FRHIGPUMask GPUMask) override
 	{
@@ -1219,6 +1220,7 @@ public:
 	{
 		return RHI->RHIGetCommandContextContainer(Index, Num, GPUMask);
 	}
+#endif // WITH_MGPU
 
 	///////// Pass through functions that allow RHIs to optimize certain calls.
 	virtual FVertexBufferRHIRef CreateAndLockVertexBuffer_RenderThread(class FRHICommandListImmediate& RHICmdList, uint32 Size, uint32 InUsage, FRHIResourceCreateInfo& CreateInfo, void*& OutDataBuffer) override
