DR - D3D11 now supports -RHIValidation
#rnx

diff --git a/Engine/Source/Runtime/RHI/Private/ValidationRHI.cpp b/Engine/Source/Runtime/RHI/Private/ValidationRHI.cpp
--- a/Engine/Source/Runtime/RHI/Private/ValidationRHI.cpp
+++ b/Engine/Source/Runtime/RHI/Private/ValidationRHI.cpp
@@ -5,9 +5,12 @@
 =============================================================================*/
 
 #include "ValidationRHI.h"
+
+#if ENABLE_RHI_VALIDATION
 #include "RHI.h"
 #include "ValidationContext.h"
 
+FValidationRHI* GValidationRHI = nullptr;
 
 FValidationRHI::FValidationRHI(FDynamicRHI* InRHI)
 	: RHI(InRHI)
@@ -15,6 +18,13 @@ FValidationRHI::FValidationRHI(FDynamicRHI* InRHI)
 	check(RHI);
 
 	Context = new FValidationContext(this);
+	AsyncComputeContext = new FValidationComputeContext(this);
+}
+
+FValidationRHI::~FValidationRHI()
+{
+	delete AsyncComputeContext;
+	delete Context;
 }
 
 IRHICommandContext* FValidationRHI::RHIGetDefaultContext()
@@ -27,6 +37,30 @@ IRHICommandContext* FValidationRHI::RHIGetDefaultContext()
 	return Context;
 }
 
+IRHIComputeContext* FValidationRHI::RHIGetDefaultAsyncComputeContext()
+{
+	if (!AsyncComputeContext->RHIContext)
+	{
+		AsyncComputeContext->RHIContext = RHI->RHIGetDefaultAsyncComputeContext();
+	}
+
+	return AsyncComputeContext;
+}
+
+
+FValidationComputeContext::FValidationComputeContext(FValidationRHI* InRHI)
+	: RHIContext(nullptr)
+	, RHI(InRHI)
+{
+	State.Reset();
+}
+
+void FValidationComputeContext::FState::Reset()
+{
+	ComputePassName.Reset();
+	bComputeShaderSet = false;
+}
+
 
 FValidationContext::FValidationContext(FValidationRHI* InRHI)
 	: RHIContext(nullptr)
@@ -45,3 +79,4 @@ void FValidationContext::FState::Reset()
 	ComputePassName.Reset();
 	bComputeShaderSet = false;
 }
+#endif	// ENABLE_RHI_VALIDATION
diff --git a/Engine/Source/Runtime/RHI/Public/ValidationContext.h b/Engine/Source/Runtime/RHI/Public/ValidationContext.h
--- a/Engine/Source/Runtime/RHI/Public/ValidationContext.h
+++ b/Engine/Source/Runtime/RHI/Public/ValidationContext.h
@@ -6,13 +6,163 @@
 
 #pragma once
 
-#include "RHIResources.h"
-#include "RHIContext.h"
-#include "DynamicRHI.h"
-
+#include "ValidationRHICommon.h"
 
+#if ENABLE_RHI_VALIDATION
 class FValidationRHI;
 
+
+class FValidationComputeContext : public IRHIComputeContext
+{
+public:
+	FValidationComputeContext(FValidationRHI* InRHI);
+
+	virtual ~FValidationComputeContext()
+	{
+	}
+
+	/**
+	* Compute queue will wait for the fence to be written before continuing.
+	*/
+	virtual void RHIWaitComputeFence(FComputeFenceRHIParamRef InFence) override final
+	{
+		RHIContext->RHIWaitComputeFence(InFence);
+	}
+
+	/**
+	*Sets the current compute shader.  Mostly for compliance with platforms
+	*that require shader setting before resource binding.
+	*/
+	virtual void RHISetComputeShader(FComputeShaderRHIParamRef Shader) override final
+	{
+		State.bComputeShaderSet = true;
+		RHIContext->RHISetComputeShader(Shader);
+	}
+
+	virtual void RHIDispatchComputeShader(uint32 ThreadGroupCountX, uint32 ThreadGroupCountY, uint32 ThreadGroupCountZ) override final
+	{
+		RHIContext->RHIDispatchComputeShader(ThreadGroupCountX, ThreadGroupCountY, ThreadGroupCountZ);
+	}
+
+	virtual void RHIDispatchIndirectComputeShader(FVertexBufferRHIParamRef ArgumentBuffer, uint32 ArgumentOffset) override final
+	{
+		RHIContext->RHIDispatchIndirectComputeShader(ArgumentBuffer, ArgumentOffset);
+	}
+
+	virtual void RHISetAsyncComputeBudget(EAsyncComputeBudget Budget) override final
+	{
+		RHIContext->RHISetAsyncComputeBudget(Budget);
+	}
+
+	/**
+	* Explicitly transition a UAV from readable -> writable by the GPU or vice versa.
+	* Also explicitly states which pipeline the UAV can be used on next.  For example, if a Compute job just wrote this UAV for a Pixel shader to read
+	* you would do EResourceTransitionAccess::Readable and EResourceTransitionPipeline::EComputeToGfx
+	*
+	* @param TransitionType - direction of the transition
+	* @param EResourceTransitionPipeline - How this UAV is transitioning between Gfx and Compute, if at all.
+	* @param InUAVs - array of UAV objects to transition
+	* @param NumUAVs - number of UAVs to transition
+	* @param WriteComputeFence - Optional ComputeFence to write as part of this transition
+	*/
+	virtual void RHITransitionResources(EResourceTransitionAccess TransitionType, EResourceTransitionPipeline TransitionPipeline, FUnorderedAccessViewRHIParamRef* InUAVs, int32 NumUAVs, FComputeFenceRHIParamRef WriteComputeFence) override final
+	{
+		RHIContext->RHITransitionResources(TransitionType, TransitionPipeline, InUAVs, NumUAVs, WriteComputeFence);
+	}
+
+	virtual void RHISubmitCommandsHint() override final
+	{
+		RHIContext->RHISubmitCommandsHint();
+	}
+
+	/** Set the shader resource view of a surface.  This is used for binding TextureMS parameter types that need a multi sampled view. */
+	virtual void RHISetShaderTexture(FComputeShaderRHIParamRef Shader, uint32 TextureIndex, FTextureRHIParamRef NewTexture) override final
+	{
+		checkf(State.bComputeShaderSet, TEXT("A Compute shader has to be set to set resources into a shader!"));
+		RHIContext->RHISetShaderTexture(Shader, TextureIndex, NewTexture);
+	}
+
+	/**
+	* Sets sampler state.
+	* @param Shader	The geometry shader to set the sampler for.
+	* @param SamplerIndex		The index of the sampler.
+	* @param NewState			The new sampler state.
+	*/
+	virtual void RHISetShaderSampler(FComputeShaderRHIParamRef Shader, uint32 SamplerIndex, FSamplerStateRHIParamRef NewState) override final
+	{
+		checkf(State.bComputeShaderSet, TEXT("A Compute shader has to be set to set resources into a shader!"));
+		RHIContext->RHISetShaderSampler(Shader, SamplerIndex, NewState);
+	}
+
+	/**
+	* Sets a compute shader UAV parameter.
+	* @param Shader	The compute shader to set the UAV for.
+	* @param UAVIndex		The index of the UAVIndex.
+	* @param UAV			The new UAV.
+	*/
+	virtual void RHISetUAVParameter(FComputeShaderRHIParamRef Shader, uint32 UAVIndex, FUnorderedAccessViewRHIParamRef UAV) override final
+	{
+		checkf(State.bComputeShaderSet, TEXT("A Compute shader has to be set to set resources into a shader!"));
+		RHIContext->RHISetUAVParameter(Shader, UAVIndex, UAV);
+	}
+
+	/**
+	* Sets a compute shader counted UAV parameter and initial count
+	* @param Shader	The compute shader to set the UAV for.
+	* @param UAVIndex		The index of the UAVIndex.
+	* @param UAV			The new UAV.
+	* @param InitialCount	The initial number of items in the UAV.
+	*/
+	virtual void RHISetUAVParameter(FComputeShaderRHIParamRef Shader, uint32 UAVIndex, FUnorderedAccessViewRHIParamRef UAV, uint32 InitialCount) override final
+	{
+		checkf(State.bComputeShaderSet, TEXT("A Compute shader has to be set to set resources into a shader!"));
+		RHIContext->RHISetUAVParameter(Shader, UAVIndex, UAV, InitialCount);
+	}
+
+	virtual void RHISetShaderResourceViewParameter(FComputeShaderRHIParamRef Shader, uint32 SamplerIndex, FShaderResourceViewRHIParamRef SRV) override final
+	{
+		checkf(State.bComputeShaderSet, TEXT("A Compute shader has to be set to set resources into a shader!"));
+		RHIContext->RHISetShaderResourceViewParameter(Shader, SamplerIndex, SRV);
+	}
+
+	virtual void RHISetShaderUniformBuffer(FComputeShaderRHIParamRef Shader, uint32 BufferIndex, FUniformBufferRHIParamRef Buffer) override final
+	{
+		checkf(State.bComputeShaderSet, TEXT("A Compute shader has to be set to set resources into a shader!"));
+		RHIContext->RHISetShaderUniformBuffer(Shader, BufferIndex, Buffer);
+	}
+
+	virtual void RHISetShaderParameter(FComputeShaderRHIParamRef Shader, uint32 BufferIndex, uint32 BaseIndex, uint32 NumBytes, const void* NewValue) override final
+	{
+		checkf(State.bComputeShaderSet, TEXT("A Compute shader has to be set to set resources into a shader!"));
+		RHIContext->RHISetShaderParameter(Shader, BufferIndex, BaseIndex, NumBytes, NewValue);
+	}
+
+	virtual void RHIPushEvent(const TCHAR* Name, FColor Color) override final
+	{
+		RHIContext->RHIPushEvent(Name, Color);
+	}
+
+	virtual void RHIPopEvent() override final
+	{
+		RHIContext->RHIPopEvent();
+	}
+
+	IRHIComputeContext* RHIContext;
+	FValidationRHI*		RHI;
+
+protected:
+	struct FState
+	{
+		FString ComputePassName;
+		bool bComputeShaderSet;
+
+		void Reset();
+	};
+	FState State;
+
+	friend class FValidationRHI;
+};
+
 class FValidationContext : public IRHICommandContext
 {
 public:
@@ -21,7 +171,7 @@ public:
 	/**
 	* Compute queue will wait for the fence to be written before continuing.
 	*/
-	virtual void RHIWaitComputeFence(FComputeFenceRHIParamRef InFence) override
+	virtual void RHIWaitComputeFence(FComputeFenceRHIParamRef InFence) override final
 	{
 		RHIContext->RHIWaitComputeFence(InFence);
 	}
@@ -30,34 +180,34 @@ public:
 	*Sets the current compute shader.  Mostly for compliance with platforms
 	*that require shader setting before resource binding.
 	*/
-	virtual void RHISetComputeShader(FComputeShaderRHIParamRef Shader) override
+	virtual void RHISetComputeShader(FComputeShaderRHIParamRef Shader) override final
 	{
 		State.bComputeShaderSet = true;
 		State.bGfxPSOSet = false;
 		RHIContext->RHISetComputeShader(Shader);
 	}
 
-	virtual void RHIDispatchComputeShader(uint32 ThreadGroupCountX, uint32 ThreadGroupCountY, uint32 ThreadGroupCountZ) override
+	virtual void RHIDispatchComputeShader(uint32 ThreadGroupCountX, uint32 ThreadGroupCountY, uint32 ThreadGroupCountZ) override final
 	{
 		RHIContext->RHIDispatchComputeShader(ThreadGroupCountX, ThreadGroupCountY, ThreadGroupCountZ);
 	}
 
-	virtual void RHIDispatchIndirectComputeShader(FVertexBufferRHIParamRef ArgumentBuffer, uint32 ArgumentOffset) override
+	virtual void RHIDispatchIndirectComputeShader(FVertexBufferRHIParamRef ArgumentBuffer, uint32 ArgumentOffset) override final
 	{
 		RHIContext->RHIDispatchIndirectComputeShader(ArgumentBuffer, ArgumentOffset);
 	}
 
-	virtual void RHISetAsyncComputeBudget(EAsyncComputeBudget Budget) override
+	virtual void RHISetAsyncComputeBudget(EAsyncComputeBudget Budget) override final
 	{
 		RHIContext->RHISetAsyncComputeBudget(Budget);
 	}
 
-	virtual void RHIAutomaticCacheFlushAfterComputeShader(bool bEnable) override
+	virtual void RHIAutomaticCacheFlushAfterComputeShader(bool bEnable) override final
 	{
 		RHIContext->RHIAutomaticCacheFlushAfterComputeShader(bEnable);
 	}
 
-	virtual void RHIFlushComputeShaderCache() override
+	virtual void RHIFlushComputeShaderCache() override final
 	{
 		RHIContext->RHIFlushComputeShaderCache();
 	}
@@ -65,13 +215,13 @@ public:
 	// Useful when used with geometry shader (emit polygons to different viewports), otherwise SetViewPort() is simpler
 	// @param Count >0
 	// @param Data must not be 0
-	virtual void RHISetMultipleViewports(uint32 Count, const FViewportBounds* Data) override
+	virtual void RHISetMultipleViewports(uint32 Count, const FViewportBounds* Data) override final
 	{
 		RHIContext->RHISetMultipleViewports(Count, Data);
 	}
 
 	/** Clears a UAV to the multi-component value provided. */
-	virtual void RHIClearTinyUAV(FUnorderedAccessViewRHIParamRef UnorderedAccessViewRHI, const uint32* Values) override
+	virtual void RHIClearTinyUAV(FUnorderedAccessViewRHIParamRef UnorderedAccessViewRHI, const uint32* Values) override final
 	{
 		RHIContext->RHIClearTinyUAV(UnorderedAccessViewRHI, Values);
 	}
@@ -82,7 +232,7 @@ public:
 	* @param DestTexture - texture to resolve to, 0 is silenty ignored
 	* @param ResolveParams - optional resolve params
 	*/
-	virtual void RHICopyToResolveTarget(FTextureRHIParamRef SourceTexture, FTextureRHIParamRef DestTexture, const FResolveParams& ResolveParams) override
+	virtual void RHICopyToResolveTarget(FTextureRHIParamRef SourceTexture, FTextureRHIParamRef DestTexture, const FResolveParams& ResolveParams) override final
 	{
 		RHIContext->RHICopyToResolveTarget(SourceTexture, DestTexture, ResolveParams);
 	}
@@ -96,7 +246,7 @@ public:
 	* @param InTextures - array of texture objects to transition
 	* @param NumTextures - number of textures to transition
 	*/
-	virtual void RHITransitionResources(EResourceTransitionAccess TransitionType, FTextureRHIParamRef* InTextures, int32 NumTextures) override
+	virtual void RHITransitionResources(EResourceTransitionAccess TransitionType, FTextureRHIParamRef* InTextures, int32 NumTextures) override final
 	{
 		RHIContext->RHITransitionResources(TransitionType, InTextures, NumTextures);
 	}
@@ -112,60 +262,60 @@ public:
 	* @param NumUAVs - number of UAVs to transition
 	* @param WriteComputeFence - Optional ComputeFence to write as part of this transition
 	*/
-	virtual void RHITransitionResources(EResourceTransitionAccess TransitionType, EResourceTransitionPipeline TransitionPipeline, FUnorderedAccessViewRHIParamRef* InUAVs, int32 NumUAVs, FComputeFenceRHIParamRef WriteComputeFence) override
+	virtual void RHITransitionResources(EResourceTransitionAccess TransitionType, EResourceTransitionPipeline TransitionPipeline, FUnorderedAccessViewRHIParamRef* InUAVs, int32 NumUAVs, FComputeFenceRHIParamRef WriteComputeFence) override final
 	{
 		RHIContext->RHITransitionResources(TransitionType, TransitionPipeline, InUAVs, NumUAVs, WriteComputeFence);
 	}
 
-	virtual void RHIBeginRenderQuery(FRenderQueryRHIParamRef RenderQuery) override
+	virtual void RHIBeginRenderQuery(FRenderQueryRHIParamRef RenderQuery) override final
 	{
 		RHIContext->RHIBeginRenderQuery(RenderQuery);
 	}
 
-	virtual void RHIEndRenderQuery(FRenderQueryRHIParamRef RenderQuery) override
+	virtual void RHIEndRenderQuery(FRenderQueryRHIParamRef RenderQuery) override final
 	{
 		RHIContext->RHIEndRenderQuery(RenderQuery);
 	}
 
-	virtual void RHISubmitCommandsHint() override
+	virtual void RHISubmitCommandsHint() override final
 	{
 		ensureMsgf(!State.bInsideBeginRenderPass, TEXT("Submitting inside a RenderPass is not efficient!"));
 		RHIContext->RHISubmitCommandsHint();
 	}
 
 	// Used for OpenGL to check and see if any occlusion queries can be read back on the RHI thread. If they aren't ready when we need them, then we end up stalling.
-	virtual void RHIPollOcclusionQueries() override
+	virtual void RHIPollOcclusionQueries() override final
 	{
 		RHIContext->RHIPollOcclusionQueries();
 	}
 
 	// Not all RHIs need this (Mobile specific)
-	virtual void RHIDiscardRenderTargets(bool bDepth, bool bStencil, uint32 ColorBitMask) override
+	virtual void RHIDiscardRenderTargets(bool bDepth, bool bStencil, uint32 ColorBitMask) override final
 	{
 		RHIContext->RHIDiscardRenderTargets(bDepth, bStencil, ColorBitMask);
 	}
 
 	// This method is queued with an RHIThread, otherwise it will flush after it is queued; without an RHI thread there is no benefit to queuing this frame advance commands
-	virtual void RHIBeginDrawingViewport(FViewportRHIParamRef Viewport, FTextureRHIParamRef RenderTargetRHI) override
+	virtual void RHIBeginDrawingViewport(FViewportRHIParamRef Viewport, FTextureRHIParamRef RenderTargetRHI) override final
 	{
 		RHIContext->RHIBeginDrawingViewport(Viewport, RenderTargetRHI);
 	}
 
 	// This method is queued with an RHIThread, otherwise it will flush after it is queued; without an RHI thread there is no benefit to queuing this frame advance commands
-	virtual void RHIEndDrawingViewport(FViewportRHIParamRef Viewport, bool bPresent, bool bLockToVsync) override
+	virtual void RHIEndDrawingViewport(FViewportRHIParamRef Viewport, bool bPresent, bool bLockToVsync) override final
 	{
 		RHIContext->RHIEndDrawingViewport(Viewport, bPresent, bLockToVsync);
 	}
 
 	// This method is queued with an RHIThread, otherwise it will flush after it is queued; without an RHI thread there is no benefit to queuing this frame advance commands
-	virtual void RHIBeginFrame() override
+	virtual void RHIBeginFrame() override final
 	{
 		State.Reset();
 		RHIContext->RHIBeginFrame();
 	}
 
 	// This method is queued with an RHIThread, otherwise it will flush after it is queued; without an RHI thread there is no benefit to queuing this frame advance commands
-	virtual void RHIEndFrame() override
+	virtual void RHIEndFrame() override final
 	{
 		RHIContext->RHIEndFrame();
 	}
@@ -176,7 +326,7 @@ public:
 	* references.
 	*/
 	// This method is queued with an RHIThread, otherwise it will flush after it is queued; without an RHI thread there is no benefit to queuing this frame advance commands
-	virtual void RHIBeginScene() override
+	virtual void RHIBeginScene() override final
 	{
 		RHIContext->RHIBeginScene();
 	}
@@ -185,7 +335,7 @@ public:
 	* Signals the end of scene rendering. See RHIBeginScene.
 	*/
 	// This method is queued with an RHIThread, otherwise it will flush after it is queued; without an RHI thread there is no benefit to queuing this frame advance commands
-	virtual void RHIEndScene() override
+	virtual void RHIEndScene() override final
 	{
 		RHIContext->RHIEndScene();
 	}
@@ -193,27 +343,27 @@ public:
 	/**
 	* Signals the beginning and ending of rendering to a resource to be used in the next frame on a multiGPU system
 	*/
-	virtual void RHIBeginUpdateMultiFrameResource(FTextureRHIParamRef Texture) override
+	virtual void RHIBeginUpdateMultiFrameResource(FTextureRHIParamRef Texture) override final
 	{
 		RHIContext->RHIBeginUpdateMultiFrameResource(Texture);
 	}
 
-	virtual void RHIEndUpdateMultiFrameResource(FTextureRHIParamRef Texture) override
+	virtual void RHIEndUpdateMultiFrameResource(FTextureRHIParamRef Texture) override final
 	{
 		RHIContext->RHIEndUpdateMultiFrameResource(Texture);
 	}
 
-	virtual void RHIBeginUpdateMultiFrameResource(FUnorderedAccessViewRHIParamRef UAV) override
+	virtual void RHIBeginUpdateMultiFrameResource(FUnorderedAccessViewRHIParamRef UAV) override final
 	{
 		RHIContext->RHIBeginUpdateMultiFrameResource(UAV);
 	}
 
-	virtual void RHIEndUpdateMultiFrameResource(FUnorderedAccessViewRHIParamRef UAV) override
+	virtual void RHIEndUpdateMultiFrameResource(FUnorderedAccessViewRHIParamRef UAV) override final
 	{
 		RHIContext->RHIEndUpdateMultiFrameResource(UAV);
 	}
 
-	virtual void RHISetStreamSource(uint32 StreamIndex, FVertexBufferRHIParamRef VertexBuffer, uint32 Offset) override
+	virtual void RHISetStreamSource(uint32 StreamIndex, FVertexBufferRHIParamRef VertexBuffer, uint32 Offset) override final
 	{
 		//#todo-rco: Decide if this is needed or not...
 		//checkf(State.bGfxPSOSet, TEXT("A Graphics PSO has to be set to set-up the vertex streams!"));
@@ -226,12 +376,12 @@ public:
 	// @param MinY including like Win32 RECT
 	// @param MaxX excluding like Win32 RECT
 	// @param MaxY excluding like Win32 RECT
-	virtual void RHISetViewport(uint32 MinX, uint32 MinY, float MinZ, uint32 MaxX, uint32 MaxY, float MaxZ) override
+	virtual void RHISetViewport(uint32 MinX, uint32 MinY, float MinZ, uint32 MaxX, uint32 MaxY, float MaxZ) override final
 	{
 		RHIContext->RHISetViewport(MinX, MinY, MinZ, MaxX, MaxY, MaxZ);
 	}
 
-	virtual void RHISetStereoViewport(uint32 LeftMinX, uint32 RightMinX, uint32 LeftMinY, uint32 RightMinY, float MinZ, uint32 LeftMaxX, uint32 RightMaxX, uint32 LeftMaxY, uint32 RightMaxY, float MaxZ) override
+	virtual void RHISetStereoViewport(uint32 LeftMinX, uint32 RightMinX, uint32 LeftMinY, uint32 RightMinY, float MinZ, uint32 LeftMaxX, uint32 RightMaxX, uint32 LeftMaxY, uint32 RightMaxY, float MaxZ) override final
 	{
 		RHIContext->RHISetStereoViewport(LeftMinX, RightMinX, LeftMinY, RightMinY, MinZ, LeftMaxX, RightMaxX, LeftMaxY, RightMaxY, MaxZ);
 	}
@@ -240,12 +390,12 @@ public:
 	// @param MinY including like Win32 RECT
 	// @param MaxX excluding like Win32 RECT
 	// @param MaxY excluding like Win32 RECT
-	virtual void RHISetScissorRect(bool bEnable, uint32 MinX, uint32 MinY, uint32 MaxX, uint32 MaxY) override
+	virtual void RHISetScissorRect(bool bEnable, uint32 MinX, uint32 MinY, uint32 MaxX, uint32 MaxY) override final
 	{
 		RHIContext->RHISetScissorRect(bEnable, MinX, MinY, MaxX, MaxY);
 	}
 
-	virtual void RHISetGraphicsPipelineState(FGraphicsPipelineStateRHIParamRef GraphicsState) override
+	virtual void RHISetGraphicsPipelineState(FGraphicsPipelineStateRHIParamRef GraphicsState) override final
 	{
 		checkf(State.bInsideBeginRenderPass, TEXT("Graphics PSOs can only be set inside a RenderPass!"));
 		State.bGfxPSOSet = true;
@@ -254,42 +404,42 @@ public:
 	}
 
 	/** Set the shader resource view of a surface.  This is used for binding TextureMS parameter types that need a multi sampled view. */
-	virtual void RHISetShaderTexture(FVertexShaderRHIParamRef Shader, uint32 TextureIndex, FTextureRHIParamRef NewTexture) override
+	virtual void RHISetShaderTexture(FVertexShaderRHIParamRef Shader, uint32 TextureIndex, FTextureRHIParamRef NewTexture) override final
 	{
 		checkf(State.bGfxPSOSet, TEXT("A Graphics PSO has to be set to set resources into a shader!"));
 		RHIContext->RHISetShaderTexture(Shader, TextureIndex, NewTexture);
 	}
 
 	/** Set the shader resource view of a surface.  This is used for binding TextureMS parameter types that need a multi sampled view. */
-	virtual void RHISetShaderTexture(FHullShaderRHIParamRef Shader, uint32 TextureIndex, FTextureRHIParamRef NewTexture) override
+	virtual void RHISetShaderTexture(FHullShaderRHIParamRef Shader, uint32 TextureIndex, FTextureRHIParamRef NewTexture) override final
 	{
 		checkf(State.bGfxPSOSet, TEXT("A Graphics PSO has to be set to set resources into a shader!"));
 		RHIContext->RHISetShaderTexture(Shader, TextureIndex, NewTexture);
 	}
 
 	/** Set the shader resource view of a surface.  This is used for binding TextureMS parameter types that need a multi sampled view. */
-	virtual void RHISetShaderTexture(FDomainShaderRHIParamRef Shader, uint32 TextureIndex, FTextureRHIParamRef NewTexture) override
+	virtual void RHISetShaderTexture(FDomainShaderRHIParamRef Shader, uint32 TextureIndex, FTextureRHIParamRef NewTexture) override final
 	{
 		checkf(State.bGfxPSOSet, TEXT("A Graphics PSO has to be set to set resources into a shader!"));
 		RHIContext->RHISetShaderTexture(Shader, TextureIndex, NewTexture);
 	}
 
 	/** Set the shader resource view of a surface.  This is used for binding TextureMS parameter types that need a multi sampled view. */
-	virtual void RHISetShaderTexture(FGeometryShaderRHIParamRef Shader, uint32 TextureIndex, FTextureRHIParamRef NewTexture) override
+	virtual void RHISetShaderTexture(FGeometryShaderRHIParamRef Shader, uint32 TextureIndex, FTextureRHIParamRef NewTexture) override final
 	{
 		checkf(State.bGfxPSOSet, TEXT("A Graphics PSO has to be set to set resources into a shader!"));
 		RHIContext->RHISetShaderTexture(Shader, TextureIndex, NewTexture);
 	}
 
 	/** Set the shader resource view of a surface.  This is used for binding TextureMS parameter types that need a multi sampled view. */
-	virtual void RHISetShaderTexture(FPixelShaderRHIParamRef Shader, uint32 TextureIndex, FTextureRHIParamRef NewTexture) override
+	virtual void RHISetShaderTexture(FPixelShaderRHIParamRef Shader, uint32 TextureIndex, FTextureRHIParamRef NewTexture) override final
 	{
 		checkf(State.bGfxPSOSet, TEXT("A Graphics PSO has to be set to set resources into a shader!"));
 		RHIContext->RHISetShaderTexture(Shader, TextureIndex, NewTexture);
 	}
 
 	/** Set the shader resource view of a surface.  This is used for binding TextureMS parameter types that need a multi sampled view. */
-	virtual void RHISetShaderTexture(FComputeShaderRHIParamRef Shader, uint32 TextureIndex, FTextureRHIParamRef NewTexture) override
+	virtual void RHISetShaderTexture(FComputeShaderRHIParamRef Shader, uint32 TextureIndex, FTextureRHIParamRef NewTexture) override final
 	{
 		checkf(State.bComputeShaderSet, TEXT("A Compute shader has to be set to set resources into a shader!"));
 		RHIContext->RHISetShaderTexture(Shader, TextureIndex, NewTexture);
@@ -301,7 +451,7 @@ public:
 	* @param SamplerIndex		The index of the sampler.
 	* @param NewState			The new sampler state.
 	*/
-	virtual void RHISetShaderSampler(FComputeShaderRHIParamRef Shader, uint32 SamplerIndex, FSamplerStateRHIParamRef NewState) override
+	virtual void RHISetShaderSampler(FComputeShaderRHIParamRef Shader, uint32 SamplerIndex, FSamplerStateRHIParamRef NewState) override final
 	{
 		checkf(State.bComputeShaderSet, TEXT("A Compute shader has to be set to set resources into a shader!"));
 		RHIContext->RHISetShaderSampler(Shader, SamplerIndex, NewState);
@@ -313,7 +463,7 @@ public:
 	* @param SamplerIndex		The index of the sampler.
 	* @param NewState			The new sampler state.
 	*/
-	virtual void RHISetShaderSampler(FVertexShaderRHIParamRef Shader, uint32 SamplerIndex, FSamplerStateRHIParamRef NewState) override
+	virtual void RHISetShaderSampler(FVertexShaderRHIParamRef Shader, uint32 SamplerIndex, FSamplerStateRHIParamRef NewState) override final
 	{
 		checkf(State.bGfxPSOSet, TEXT("A Graphics PSO has to be set to set resources into a shader!"));
 		RHIContext->RHISetShaderSampler(Shader, SamplerIndex, NewState);
@@ -326,7 +476,7 @@ public:
 	* @param SamplerIndex		The index of the sampler.
 	* @param NewState			The new sampler state.
 	*/
-	virtual void RHISetShaderSampler(FGeometryShaderRHIParamRef Shader, uint32 SamplerIndex, FSamplerStateRHIParamRef NewState) override
+	virtual void RHISetShaderSampler(FGeometryShaderRHIParamRef Shader, uint32 SamplerIndex, FSamplerStateRHIParamRef NewState) override final
 	{
 		checkf(State.bGfxPSOSet, TEXT("A Graphics PSO has to be set to set resources into a shader!"));
 		RHIContext->RHISetShaderSampler(Shader, SamplerIndex, NewState);
@@ -338,7 +488,7 @@ public:
 	* @param SamplerIndex		The index of the sampler.
 	* @param NewState			The new sampler state.
 	*/
-	virtual void RHISetShaderSampler(FDomainShaderRHIParamRef Shader, uint32 SamplerIndex, FSamplerStateRHIParamRef NewState) override
+	virtual void RHISetShaderSampler(FDomainShaderRHIParamRef Shader, uint32 SamplerIndex, FSamplerStateRHIParamRef NewState) override final
 	{
 		checkf(State.bGfxPSOSet, TEXT("A Graphics PSO has to be set to set resources into a shader!"));
 		RHIContext->RHISetShaderSampler(Shader, SamplerIndex, NewState);
@@ -350,7 +500,7 @@ public:
 	* @param SamplerIndex		The index of the sampler.
 	* @param NewState			The new sampler state.
 	*/
-	virtual void RHISetShaderSampler(FHullShaderRHIParamRef Shader, uint32 SamplerIndex, FSamplerStateRHIParamRef NewState) override
+	virtual void RHISetShaderSampler(FHullShaderRHIParamRef Shader, uint32 SamplerIndex, FSamplerStateRHIParamRef NewState) override final
 	{
 		checkf(State.bGfxPSOSet, TEXT("A Graphics PSO has to be set to set resources into a shader!"));
 		RHIContext->RHISetShaderSampler(Shader, SamplerIndex, NewState);
@@ -362,7 +512,7 @@ public:
 	* @param SamplerIndex		The index of the sampler.
 	* @param NewState			The new sampler state.
 	*/
-	virtual void RHISetShaderSampler(FPixelShaderRHIParamRef Shader, uint32 SamplerIndex, FSamplerStateRHIParamRef NewState) override
+	virtual void RHISetShaderSampler(FPixelShaderRHIParamRef Shader, uint32 SamplerIndex, FSamplerStateRHIParamRef NewState) override final
 	{
 		checkf(State.bGfxPSOSet, TEXT("A Graphics PSO has to be set to set resources into a shader!"));
 		RHIContext->RHISetShaderSampler(Shader, SamplerIndex, NewState);
@@ -374,7 +524,7 @@ public:
 	* @param UAVIndex		The index of the UAVIndex.
 	* @param UAV			The new UAV.
 	*/
-	virtual void RHISetUAVParameter(FComputeShaderRHIParamRef Shader, uint32 UAVIndex, FUnorderedAccessViewRHIParamRef UAV) override
+	virtual void RHISetUAVParameter(FComputeShaderRHIParamRef Shader, uint32 UAVIndex, FUnorderedAccessViewRHIParamRef UAV) override final
 	{
 		checkf(State.bComputeShaderSet, TEXT("A Compute shader has to be set to set resources into a shader!"));
 		RHIContext->RHISetUAVParameter(Shader, UAVIndex, UAV);
@@ -387,138 +537,138 @@ public:
 	* @param UAV			The new UAV.
 	* @param InitialCount	The initial number of items in the UAV.
 	*/
-	virtual void RHISetUAVParameter(FComputeShaderRHIParamRef Shader, uint32 UAVIndex, FUnorderedAccessViewRHIParamRef UAV, uint32 InitialCount) override
+	virtual void RHISetUAVParameter(FComputeShaderRHIParamRef Shader, uint32 UAVIndex, FUnorderedAccessViewRHIParamRef UAV, uint32 InitialCount) override final
 	{
 		checkf(State.bComputeShaderSet, TEXT("A Compute shader has to be set to set resources into a shader!"));
 		RHIContext->RHISetUAVParameter(Shader, UAVIndex, UAV, InitialCount);
 	}
 
-	virtual void RHISetShaderResourceViewParameter(FPixelShaderRHIParamRef Shader, uint32 SamplerIndex, FShaderResourceViewRHIParamRef SRV) override
+	virtual void RHISetShaderResourceViewParameter(FPixelShaderRHIParamRef Shader, uint32 SamplerIndex, FShaderResourceViewRHIParamRef SRV) override final
 	{
 		checkf(State.bGfxPSOSet, TEXT("A Graphics PSO has to be set to set resources into a shader!"));
 		RHIContext->RHISetShaderResourceViewParameter(Shader, SamplerIndex, SRV);
 	}
 
-	virtual void RHISetShaderResourceViewParameter(FVertexShaderRHIParamRef Shader, uint32 SamplerIndex, FShaderResourceViewRHIParamRef SRV) override
+	virtual void RHISetShaderResourceViewParameter(FVertexShaderRHIParamRef Shader, uint32 SamplerIndex, FShaderResourceViewRHIParamRef SRV) override final
 	{
 		checkf(State.bGfxPSOSet, TEXT("A Graphics PSO has to be set to set resources into a shader!"));
 		RHIContext->RHISetShaderResourceViewParameter(Shader, SamplerIndex, SRV);
 	}
 
-	virtual void RHISetShaderResourceViewParameter(FComputeShaderRHIParamRef Shader, uint32 SamplerIndex, FShaderResourceViewRHIParamRef SRV) override
+	virtual void RHISetShaderResourceViewParameter(FComputeShaderRHIParamRef Shader, uint32 SamplerIndex, FShaderResourceViewRHIParamRef SRV) override final
 	{
 		checkf(State.bComputeShaderSet, TEXT("A Compute shader has to be set to set resources into a shader!"));
 		RHIContext->RHISetShaderResourceViewParameter(Shader, SamplerIndex, SRV);
 	}
 
-	virtual void RHISetShaderResourceViewParameter(FHullShaderRHIParamRef Shader, uint32 SamplerIndex, FShaderResourceViewRHIParamRef SRV) override
+	virtual void RHISetShaderResourceViewParameter(FHullShaderRHIParamRef Shader, uint32 SamplerIndex, FShaderResourceViewRHIParamRef SRV) override final
 	{
 		checkf(State.bGfxPSOSet, TEXT("A Graphics PSO has to be set to set resources into a shader!"));
 		RHIContext->RHISetShaderResourceViewParameter(Shader, SamplerIndex, SRV);
 	}
 
-	virtual void RHISetShaderResourceViewParameter(FDomainShaderRHIParamRef Shader, uint32 SamplerIndex, FShaderResourceViewRHIParamRef SRV) override
+	virtual void RHISetShaderResourceViewParameter(FDomainShaderRHIParamRef Shader, uint32 SamplerIndex, FShaderResourceViewRHIParamRef SRV) override final
 	{
 		checkf(State.bGfxPSOSet, TEXT("A Graphics PSO has to be set to set resources into a shader!"));
 		RHIContext->RHISetShaderResourceViewParameter(Shader, SamplerIndex, SRV);
 	}
 
-	virtual void RHISetShaderResourceViewParameter(FGeometryShaderRHIParamRef Shader, uint32 SamplerIndex, FShaderResourceViewRHIParamRef SRV) override
+	virtual void RHISetShaderResourceViewParameter(FGeometryShaderRHIParamRef Shader, uint32 SamplerIndex, FShaderResourceViewRHIParamRef SRV) override final
 	{
 		checkf(State.bGfxPSOSet, TEXT("A Graphics PSO has to be set to set resources into a shader!"));
 		RHIContext->RHISetShaderResourceViewParameter(Shader, SamplerIndex, SRV);
 	}
 
-	virtual void RHISetShaderUniformBuffer(FVertexShaderRHIParamRef Shader, uint32 BufferIndex, FUniformBufferRHIParamRef Buffer) override
+	virtual void RHISetShaderUniformBuffer(FVertexShaderRHIParamRef Shader, uint32 BufferIndex, FUniformBufferRHIParamRef Buffer) override final
 	{
 		checkf(State.bGfxPSOSet, TEXT("A Graphics PSO has to be set to set resources into a shader!"));
 		RHIContext->RHISetShaderUniformBuffer(Shader, BufferIndex, Buffer);
 	}
 
-	virtual void RHISetShaderUniformBuffer(FHullShaderRHIParamRef Shader, uint32 BufferIndex, FUniformBufferRHIParamRef Buffer) override
+	virtual void RHISetShaderUniformBuffer(FHullShaderRHIParamRef Shader, uint32 BufferIndex, FUniformBufferRHIParamRef Buffer) override final
 	{
 		checkf(State.bGfxPSOSet, TEXT("A Graphics PSO has to be set to set resources into a shader!"));
 		RHIContext->RHISetShaderUniformBuffer(Shader, BufferIndex, Buffer);
 	}
 
-	virtual void RHISetShaderUniformBuffer(FDomainShaderRHIParamRef Shader, uint32 BufferIndex, FUniformBufferRHIParamRef Buffer) override
+	virtual void RHISetShaderUniformBuffer(FDomainShaderRHIParamRef Shader, uint32 BufferIndex, FUniformBufferRHIParamRef Buffer) override final
 	{
 		checkf(State.bGfxPSOSet, TEXT("A Graphics PSO has to be set to set resources into a shader!"));
 		RHIContext->RHISetShaderUniformBuffer(Shader, BufferIndex, Buffer);
 	}
 
-	virtual void RHISetShaderUniformBuffer(FGeometryShaderRHIParamRef Shader, uint32 BufferIndex, FUniformBufferRHIParamRef Buffer) override
+	virtual void RHISetShaderUniformBuffer(FGeometryShaderRHIParamRef Shader, uint32 BufferIndex, FUniformBufferRHIParamRef Buffer) override final
 	{
 		checkf(State.bGfxPSOSet, TEXT("A Graphics PSO has to be set to set resources into a shader!"));
 		RHIContext->RHISetShaderUniformBuffer(Shader, BufferIndex, Buffer);
 	}
 
-	virtual void RHISetShaderUniformBuffer(FPixelShaderRHIParamRef Shader, uint32 BufferIndex, FUniformBufferRHIParamRef Buffer) override
+	virtual void RHISetShaderUniformBuffer(FPixelShaderRHIParamRef Shader, uint32 BufferIndex, FUniformBufferRHIParamRef Buffer) override final
 	{
 		checkf(State.bGfxPSOSet, TEXT("A Graphics PSO has to be set to set resources into a shader!"));
 		RHIContext->RHISetShaderUniformBuffer(Shader, BufferIndex, Buffer);
 	}
 
-	virtual void RHISetShaderUniformBuffer(FComputeShaderRHIParamRef Shader, uint32 BufferIndex, FUniformBufferRHIParamRef Buffer) override
+	virtual void RHISetShaderUniformBuffer(FComputeShaderRHIParamRef Shader, uint32 BufferIndex, FUniformBufferRHIParamRef Buffer) override final
 	{
 		checkf(State.bComputeShaderSet, TEXT("A Compute shader has to be set to set resources into a shader!"));
 		RHIContext->RHISetShaderUniformBuffer(Shader, BufferIndex, Buffer);
 	}
 
-	virtual void RHISetShaderParameter(FVertexShaderRHIParamRef Shader, uint32 BufferIndex, uint32 BaseIndex, uint32 NumBytes, const void* NewValue) override
+	virtual void RHISetShaderParameter(FVertexShaderRHIParamRef Shader, uint32 BufferIndex, uint32 BaseIndex, uint32 NumBytes, const void* NewValue) override final
 	{
 		checkf(State.bGfxPSOSet, TEXT("A Graphics PSO has to be set to set resources into a shader!"));
 		RHIContext->RHISetShaderParameter(Shader, BufferIndex, BaseIndex, NumBytes, NewValue);
 	}
 
-	virtual void RHISetShaderParameter(FPixelShaderRHIParamRef Shader, uint32 BufferIndex, uint32 BaseIndex, uint32 NumBytes, const void* NewValue) override
+	virtual void RHISetShaderParameter(FPixelShaderRHIParamRef Shader, uint32 BufferIndex, uint32 BaseIndex, uint32 NumBytes, const void* NewValue) override final
 	{
 		checkf(State.bGfxPSOSet, TEXT("A Graphics PSO has to be set to set resources into a shader!"));
 		RHIContext->RHISetShaderParameter(Shader, BufferIndex, BaseIndex, NumBytes, NewValue);
 	}
 
-	virtual void RHISetShaderParameter(FHullShaderRHIParamRef Shader, uint32 BufferIndex, uint32 BaseIndex, uint32 NumBytes, const void* NewValue) override
+	virtual void RHISetShaderParameter(FHullShaderRHIParamRef Shader, uint32 BufferIndex, uint32 BaseIndex, uint32 NumBytes, const void* NewValue) override final
 	{
 		checkf(State.bGfxPSOSet, TEXT("A Graphics PSO has to be set to set resources into a shader!"));
 		RHIContext->RHISetShaderParameter(Shader, BufferIndex, BaseIndex, NumBytes, NewValue);
 	}
 
-	virtual void RHISetShaderParameter(FDomainShaderRHIParamRef Shader, uint32 BufferIndex, uint32 BaseIndex, uint32 NumBytes, const void* NewValue) override
+	virtual void RHISetShaderParameter(FDomainShaderRHIParamRef Shader, uint32 BufferIndex, uint32 BaseIndex, uint32 NumBytes, const void* NewValue) override final
 	{
 		checkf(State.bGfxPSOSet, TEXT("A Graphics PSO has to be set to set resources into a shader!"));
 		RHIContext->RHISetShaderParameter(Shader, BufferIndex, BaseIndex, NumBytes, NewValue);
 	}
 
-	virtual void RHISetShaderParameter(FGeometryShaderRHIParamRef Shader, uint32 BufferIndex, uint32 BaseIndex, uint32 NumBytes, const void* NewValue) override
+	virtual void RHISetShaderParameter(FGeometryShaderRHIParamRef Shader, uint32 BufferIndex, uint32 BaseIndex, uint32 NumBytes, const void* NewValue) override final
 	{
 		checkf(State.bGfxPSOSet, TEXT("A Graphics PSO has to be set to set resources into a shader!"));
 		RHIContext->RHISetShaderParameter(Shader, BufferIndex, BaseIndex, NumBytes, NewValue);
 	}
 
-	virtual void RHISetShaderParameter(FComputeShaderRHIParamRef Shader, uint32 BufferIndex, uint32 BaseIndex, uint32 NumBytes, const void* NewValue) override
+	virtual void RHISetShaderParameter(FComputeShaderRHIParamRef Shader, uint32 BufferIndex, uint32 BaseIndex, uint32 NumBytes, const void* NewValue) override final
 	{
 		checkf(State.bComputeShaderSet, TEXT("A Compute shader has to be set to set resources into a shader!"));
 		RHIContext->RHISetShaderParameter(Shader, BufferIndex, BaseIndex, NumBytes, NewValue);
 	}
 
-	virtual void RHISetStencilRef(uint32 StencilRef) override
+	virtual void RHISetStencilRef(uint32 StencilRef) override final
 	{
 		checkf(State.bGfxPSOSet, TEXT("A Graphics PSO has to be set to change stencil ref!"));
 		RHIContext->RHISetStencilRef(StencilRef);
 	}
 
-	virtual void RHISetBlendFactor(const FLinearColor& BlendFactor) override
+	virtual void RHISetBlendFactor(const FLinearColor& BlendFactor) override final
 	{
 		checkf(State.bGfxPSOSet, TEXT("A Graphics PSO has to be set to change blend factor!"));
 		RHIContext->RHISetBlendFactor(BlendFactor);
 	}
 
-	virtual void RHISetRenderTargets(uint32 NumSimultaneousRenderTargets, const FRHIRenderTargetView* NewRenderTargets, const FRHIDepthRenderTargetView* NewDepthStencilTarget, uint32 NumUAVs, const FUnorderedAccessViewRHIParamRef* UAVs) override
+	virtual void RHISetRenderTargets(uint32 NumSimultaneousRenderTargets, const FRHIRenderTargetView* NewRenderTargets, const FRHIDepthRenderTargetView* NewDepthStencilTarget, uint32 NumUAVs, const FUnorderedAccessViewRHIParamRef* UAVs) override final
 	{
 		RHIContext->RHISetRenderTargets(NumSimultaneousRenderTargets, NewRenderTargets, NewDepthStencilTarget, NumUAVs, UAVs);
 	}
 
-	virtual void RHISetRenderTargetsAndClear(const FRHISetRenderTargetsInfo& RenderTargetsInfo) override
+	virtual void RHISetRenderTargetsAndClear(const FRHISetRenderTargetsInfo& RenderTargetsInfo) override final
 	{
 		RHIContext->RHISetRenderTargetsAndClear(RenderTargetsInfo);
 	}
@@ -526,37 +676,37 @@ public:
 	// Bind the clear state of the currently set rendertargets.  This is used by platforms which
 	// need the state of the target when finalizing a hardware clear or a resource transition to SRV
 	// The explicit bind is needed to support parallel rendering (propagate state between contexts).
-	virtual void RHIBindClearMRTValues(bool bClearColor, bool bClearDepth, bool bClearStencil) override
+	virtual void RHIBindClearMRTValues(bool bClearColor, bool bClearDepth, bool bClearStencil) override final
 	{
 		RHIContext->RHIBindClearMRTValues(bClearColor, bClearDepth, bClearStencil);
 	}
 
-	virtual void RHIDrawPrimitive(uint32 BaseVertexIndex, uint32 NumPrimitives, uint32 NumInstances) override
+	virtual void RHIDrawPrimitive(uint32 BaseVertexIndex, uint32 NumPrimitives, uint32 NumInstances) override final
 	{
 		checkf(State.bGfxPSOSet, TEXT("A Graphics PSO has to be set to draw!"));
 		RHIContext->RHIDrawPrimitive(BaseVertexIndex, NumPrimitives, NumInstances);
 	}
 
-	virtual void RHIDrawPrimitiveIndirect(FVertexBufferRHIParamRef ArgumentBuffer, uint32 ArgumentOffset) override
+	virtual void RHIDrawPrimitiveIndirect(FVertexBufferRHIParamRef ArgumentBuffer, uint32 ArgumentOffset) override final
 	{
 		checkf(State.bGfxPSOSet, TEXT("A Graphics PSO has to be set to draw!"));
 		RHIContext->RHIDrawPrimitiveIndirect(ArgumentBuffer, ArgumentOffset);
 	}
 
-	virtual void RHIDrawIndexedIndirect(FIndexBufferRHIParamRef IndexBufferRHI, FStructuredBufferRHIParamRef ArgumentsBufferRHI, int32 DrawArgumentsIndex, uint32 NumInstances) override
+	virtual void RHIDrawIndexedIndirect(FIndexBufferRHIParamRef IndexBufferRHI, FStructuredBufferRHIParamRef ArgumentsBufferRHI, int32 DrawArgumentsIndex, uint32 NumInstances) override final
 	{
 		checkf(State.bGfxPSOSet, TEXT("A Graphics PSO has to be set to draw!"));
 		RHIContext->RHIDrawIndexedIndirect(IndexBufferRHI, ArgumentsBufferRHI, DrawArgumentsIndex, NumInstances);
 	}
 
 	// @param NumPrimitives need to be >0 
-	virtual void RHIDrawIndexedPrimitive(FIndexBufferRHIParamRef IndexBuffer, int32 BaseVertexIndex, uint32 FirstInstance, uint32 NumVertices, uint32 StartIndex, uint32 NumPrimitives, uint32 NumInstances) override
+	virtual void RHIDrawIndexedPrimitive(FIndexBufferRHIParamRef IndexBuffer, int32 BaseVertexIndex, uint32 FirstInstance, uint32 NumVertices, uint32 StartIndex, uint32 NumPrimitives, uint32 NumInstances) override final
 	{
 		checkf(State.bGfxPSOSet, TEXT("A Graphics PSO has to be set to draw!"));
 		RHIContext->RHIDrawIndexedPrimitive(IndexBuffer, BaseVertexIndex, FirstInstance, NumVertices, StartIndex, NumPrimitives, NumInstances);
 	}
 
-	virtual void RHIDrawIndexedPrimitiveIndirect(FIndexBufferRHIParamRef IndexBuffer, FVertexBufferRHIParamRef ArgumentBuffer, uint32 ArgumentOffset) override
+	virtual void RHIDrawIndexedPrimitiveIndirect(FIndexBufferRHIParamRef IndexBuffer, FVertexBufferRHIParamRef ArgumentBuffer, uint32 ArgumentOffset) override final
 	{
 		checkf(State.bGfxPSOSet, TEXT("A Graphics PSO has to be set to draw!"));
 		RHIContext->RHIDrawIndexedPrimitiveIndirect(IndexBuffer, ArgumentBuffer, ArgumentOffset);
@@ -568,28 +718,28 @@ public:
 	* @param MaxDepth	The maximum depth for depth bounds test.
 	*					The valid values for fMinDepth and fMaxDepth are such that 0 <= fMinDepth <= fMaxDepth <= 1
 	*/
-	virtual void RHISetDepthBounds(float MinDepth, float MaxDepth) override
+	virtual void RHISetDepthBounds(float MinDepth, float MaxDepth) override final
 	{
 		checkf(State.bGfxPSOSet, TEXT("A Graphics PSO has to be set to set resources into a shader!"));
 		RHIContext->RHISetDepthBounds(MinDepth, MaxDepth);
 	}
 
-	virtual void RHIPushEvent(const TCHAR* Name, FColor Color) override
+	virtual void RHIPushEvent(const TCHAR* Name, FColor Color) override final
 	{
 		RHIContext->RHIPushEvent(Name, Color);
 	}
 
-	virtual void RHIPopEvent() override
+	virtual void RHIPopEvent() override final
 	{
 		RHIContext->RHIPopEvent();
 	}
 
-	virtual void RHIUpdateTextureReference(FTextureReferenceRHIParamRef TextureRef, FTextureRHIParamRef NewTexture) override
+	virtual void RHIUpdateTextureReference(FTextureReferenceRHIParamRef TextureRef, FTextureRHIParamRef NewTexture) override final
 	{
 		RHIContext->RHIUpdateTextureReference(TextureRef, NewTexture);
 	}
 
-	virtual void RHIBeginRenderPass(const FRHIRenderPassInfo& InInfo, const TCHAR* InName) override
+	virtual void RHIBeginRenderPass(const FRHIRenderPassInfo& InInfo, const TCHAR* InName) override final
 	{
 		checkf(!State.bInsideBeginRenderPass, TEXT("Trying to begin RenderPass '%s', but already inside '%s'!"), *State.RenderPassName, InName);
 		checkf(InName, TEXT("RenderPass should have a name!"));
@@ -600,7 +750,7 @@ public:
 		RHIContext->RHIBeginRenderPass(InInfo, InName);
 	}
 
-	virtual void RHIEndRenderPass() override
+	virtual void RHIEndRenderPass() override final
 	{
 		checkf(State.bInsideBeginRenderPass, TEXT("Trying to end a RenderPass but not inside one!"));
 		RHIContext->RHIEndRenderPass();
@@ -609,7 +759,7 @@ public:
 		State.bGfxPSOSet = false;
 	}
 
-	virtual void RHIBeginComputePass(const TCHAR* InName) override
+	virtual void RHIBeginComputePass(const TCHAR* InName) override final
 	{
 		checkf(InName && *InName, TEXT("ComputePass should have a name!"));
 		checkf(!State.bInsideBeginRenderPass, TEXT("Can't begin a compute pass from inside RenderPass '%s'"), *State.RenderPassName);
@@ -619,35 +769,35 @@ public:
 		RHIContext->RHIBeginComputePass(InName);
 	}
 
-	virtual void RHIEndComputePass() override
+	virtual void RHIEndComputePass() override final
 	{
 		checkf(State.bInsideComputePass, TEXT("Can't end a compute pass without a Begin!"), *State.ComputePassName);
 		RHIContext->RHIEndComputePass();
 		State.bInsideComputePass = false;
 	}
 
-	virtual void RHICopyTexture(FTextureRHIParamRef SourceTexture, FTextureRHIParamRef DestTexture, const FRHICopyTextureInfo& CopyInfo) override
+	virtual void RHICopyTexture(FTextureRHIParamRef SourceTexture, FTextureRHIParamRef DestTexture, const FRHICopyTextureInfo& CopyInfo) override final
 	{
 		ensureMsgf(!State.bInsideBeginRenderPass, TEXT("Copying inside a RenderPass is not efficient!"));
 		RHIContext->RHICopyTexture(SourceTexture, DestTexture, CopyInfo);
 	}
 
-	virtual void RHIBuildAccelerationStructure(FRayTracingGeometryRHIParamRef Geometry) override
+	virtual void RHIBuildAccelerationStructure(FRayTracingGeometryRHIParamRef Geometry) override final
 	{
 		RHIContext->RHIBuildAccelerationStructure(Geometry);
 	}
 
-	virtual void RHIUpdateAccelerationStructures(const TArrayView<const FAccelerationStructureUpdateParams> Params) override
+	virtual void RHIUpdateAccelerationStructures(const TArrayView<const FAccelerationStructureUpdateParams> Params) override final
 	{
 		RHIContext->RHIUpdateAccelerationStructures(Params);
 	}
 
-	virtual void RHIBuildAccelerationStructures(const TArrayView<const FAccelerationStructureUpdateParams> Params) override
+	virtual void RHIBuildAccelerationStructures(const TArrayView<const FAccelerationStructureUpdateParams> Params) override final
 	{
 		RHIContext->RHIBuildAccelerationStructures(Params);
 	}
 
-	virtual void RHIBuildAccelerationStructure(FRayTracingSceneRHIParamRef Scene) override
+	virtual void RHIBuildAccelerationStructure(FRayTracingSceneRHIParamRef Scene) override final
 	{
 		RHIContext->RHIBuildAccelerationStructure(Scene);
 	}
@@ -655,7 +805,7 @@ public:
 	virtual void RHIRayTraceOcclusion(FRayTracingSceneRHIParamRef Scene,
 		FShaderResourceViewRHIParamRef Rays,
 		FUnorderedAccessViewRHIParamRef Output,
-		uint32 NumRays) override
+		uint32 NumRays) override final
 	{
 		RHIContext->RHIRayTraceOcclusion(Scene, Rays, Output, NumRays);
 	}
@@ -663,7 +813,7 @@ public:
 	virtual void RHIRayTraceIntersection(FRayTracingSceneRHIParamRef Scene,
 		FShaderResourceViewRHIParamRef Rays,
 		FUnorderedAccessViewRHIParamRef Output,
-		uint32 NumRays) override
+		uint32 NumRays) override final
 	{
 		RHIContext->RHIRayTraceIntersection(Scene, Rays, Output, NumRays);
 	}
@@ -671,7 +821,7 @@ public:
 	virtual void RHIRayTraceDispatch(FRayTracingPipelineStateRHIParamRef RayTracingPipelineState, FRHIRayTracingShader* RayGenShader,
 		FRayTracingSceneRHIParamRef Scene,
 		const FRayTracingShaderBindings& GlobalResourceBindings,
-		uint32 Width, uint32 Height) override
+		uint32 Width, uint32 Height) override final
 	{
 		RHIContext->RHIRayTraceDispatch(RayTracingPipelineState, RayGenShader, Scene, GlobalResourceBindings, Width, Height);
 	}
@@ -680,7 +830,7 @@ public:
 		FRayTracingSceneRHIParamRef Scene, uint32 InstanceIndex, uint32 SegmentIndex, uint32 ShaderSlot,
 		FRayTracingPipelineStateRHIParamRef Pipeline, uint32 HitGroupIndex,
 		uint32 NumUniformBuffers, const FUniformBufferRHIParamRef* UniformBuffers,
-		uint32 UserData) override
+		uint32 UserData) override final
 	{
 		RHIContext->RHISetRayTracingHitGroup(Scene, InstanceIndex, SegmentIndex, ShaderSlot, Pipeline, HitGroupIndex, NumUniformBuffers, UniformBuffers, UserData);
 	}
@@ -706,3 +856,4 @@ protected:
 
 	friend class FValidationRHI;
 };
+#endif	// ENABLE_RHI_VALIDATION
diff --git a/Engine/Source/Runtime/RHI/Public/ValidationRHI.h b/Engine/Source/Runtime/RHI/Public/ValidationRHI.h
--- a/Engine/Source/Runtime/RHI/Public/ValidationRHI.h
+++ b/Engine/Source/Runtime/RHI/Public/ValidationRHI.h
@@ -6,18 +6,19 @@
 
 #pragma once 
 
-#include "RHI.h"
-#include "RHIResources.h"
-#include "DynamicRHI.h"
+#include "ValidationRHICommon.h"
 
+#if ENABLE_RHI_VALIDATION
+class FValidationComputeContext;
 class FValidationContext;
 
 class FValidationRHI : public FDynamicRHI
 {
 public:
 	RHI_API FValidationRHI(FDynamicRHI* InRHI);
+	RHI_API virtual ~FValidationRHI();
 
-	virtual void Init() override
+	virtual void Init() override final
 	{
 		RHI->Init();
 		RHIName = RHI->GetName();
@@ -25,7 +26,7 @@ public:
 	}
 
 	/** Called after the RHI is initialized; before the render thread is started. */
-	virtual void PostInit() override
+	virtual void PostInit() override final
 	{
 		// Need to copy this as each DynamicRHI has an instance
 		check(RHI->PixelFormatBlockBytes.Num() <= PixelFormatBlockBytes.Num());
@@ -34,12 +35,12 @@ public:
 	}
 
 	/** Shutdown the RHI; handle shutdown and resource destruction before the RHI's actual destructor is called (so that all resources of the RHI are still available for shutdown). */
-	virtual void Shutdown() override
+	virtual void Shutdown() override final
 	{
 		RHI->Shutdown();
 	}
 
-	virtual const TCHAR* GetName() override
+	virtual const TCHAR* GetName() override final
 	{
 		return *RHIName;
 	}
@@ -47,122 +48,122 @@ public:
 	/////// RHI Methods
 
 	// FlushType: Thread safe
-	virtual FSamplerStateRHIRef RHICreateSamplerState(const FSamplerStateInitializerRHI& Initializer) override
+	virtual FSamplerStateRHIRef RHICreateSamplerState(const FSamplerStateInitializerRHI& Initializer) override final
 	{
 		return RHI->RHICreateSamplerState(Initializer);
 	}
 
 	// FlushType: Thread safe
-	virtual FRasterizerStateRHIRef RHICreateRasterizerState(const FRasterizerStateInitializerRHI& Initializer) override
+	virtual FRasterizerStateRHIRef RHICreateRasterizerState(const FRasterizerStateInitializerRHI& Initializer) override final
 	{
 		return RHI->RHICreateRasterizerState(Initializer);
 	}
 
 	// FlushType: Thread safe
-	virtual FDepthStencilStateRHIRef RHICreateDepthStencilState(const FDepthStencilStateInitializerRHI& Initializer) override
+	virtual FDepthStencilStateRHIRef RHICreateDepthStencilState(const FDepthStencilStateInitializerRHI& Initializer) override final
 	{
 		return RHI->RHICreateDepthStencilState(Initializer);
 	}
 
 	// FlushType: Thread safe
-	virtual FBlendStateRHIRef RHICreateBlendState(const FBlendStateInitializerRHI& Initializer) override
+	virtual FBlendStateRHIRef RHICreateBlendState(const FBlendStateInitializerRHI& Initializer) override final
 	{
 		return RHI->RHICreateBlendState(Initializer);
 	}
 
 	// FlushType: Wait RHI Thread
-	virtual FVertexDeclarationRHIRef RHICreateVertexDeclaration(const FVertexDeclarationElementList& Elements) override
+	virtual FVertexDeclarationRHIRef RHICreateVertexDeclaration(const FVertexDeclarationElementList& Elements) override final
 	{
 		return RHI->RHICreateVertexDeclaration(Elements);
 	}
 
 	// FlushType: Wait RHI Thread
-	virtual FPixelShaderRHIRef RHICreatePixelShader(const TArray<uint8>& Code) override
+	virtual FPixelShaderRHIRef RHICreatePixelShader(const TArray<uint8>& Code) override final
 	{
 		return RHI->RHICreatePixelShader(Code);
 	}
 
 	// FlushType: Wait RHI Thread
-	virtual FPixelShaderRHIRef RHICreatePixelShader(FRHIShaderLibrary* Library, FSHAHash Hash) override
+	virtual FPixelShaderRHIRef RHICreatePixelShader(FRHIShaderLibrary* Library, FSHAHash Hash) override final
 	{
 		return RHI->RHICreatePixelShader(Library, Hash);
 	}
 
 	// FlushType: Wait RHI Thread
-	virtual FVertexShaderRHIRef RHICreateVertexShader(const TArray<uint8>& Code) override
+	virtual FVertexShaderRHIRef RHICreateVertexShader(const TArray<uint8>& Code) override final
 	{
 		return RHI->RHICreateVertexShader(Code);
 	}
 
 	// FlushType: Wait RHI Thread
-	virtual FVertexShaderRHIRef RHICreateVertexShader(FRHIShaderLibrary* Library, FSHAHash Hash) override
+	virtual FVertexShaderRHIRef RHICreateVertexShader(FRHIShaderLibrary* Library, FSHAHash Hash) override final
 	{
 		return RHI->RHICreateVertexShader(Library, Hash);
 	}
 
 	// FlushType: Wait RHI Thread
-	virtual FHullShaderRHIRef RHICreateHullShader(const TArray<uint8>& Code) override
+	virtual FHullShaderRHIRef RHICreateHullShader(const TArray<uint8>& Code) override final
 	{
 		return RHI->RHICreateHullShader(Code);
 	}
 
 	// FlushType: Wait RHI Thread
-	virtual FHullShaderRHIRef RHICreateHullShader(FRHIShaderLibrary* Library, FSHAHash Hash) override
+	virtual FHullShaderRHIRef RHICreateHullShader(FRHIShaderLibrary* Library, FSHAHash Hash) override final
 	{
 		return RHI->RHICreateHullShader(Library, Hash);
 	}
 
 	// FlushType: Wait RHI Thread
-	virtual FDomainShaderRHIRef RHICreateDomainShader(const TArray<uint8>& Code) override
+	virtual FDomainShaderRHIRef RHICreateDomainShader(const TArray<uint8>& Code) override final
 	{
 		return RHI->RHICreateDomainShader(Code);
 	}
 
 	// FlushType: Wait RHI Thread
-	virtual FDomainShaderRHIRef RHICreateDomainShader(FRHIShaderLibrary* Library, FSHAHash Hash) override
+	virtual FDomainShaderRHIRef RHICreateDomainShader(FRHIShaderLibrary* Library, FSHAHash Hash) override final
 	{
 		return RHI->RHICreateDomainShader(Library, Hash);
 	}
 
 	// FlushType: Wait RHI Thread
-	virtual FGeometryShaderRHIRef RHICreateGeometryShader(const TArray<uint8>& Code) override
+	virtual FGeometryShaderRHIRef RHICreateGeometryShader(const TArray<uint8>& Code) override final
 	{
 		return RHI->RHICreateGeometryShader(Code);
 	}
 
 	// FlushType: Wait RHI Thread
-	virtual FGeometryShaderRHIRef RHICreateGeometryShader(FRHIShaderLibrary* Library, FSHAHash Hash) override
+	virtual FGeometryShaderRHIRef RHICreateGeometryShader(FRHIShaderLibrary* Library, FSHAHash Hash) override final
 	{
 		return RHI->RHICreateGeometryShader(Library, Hash);
 	}
 
 	/** Creates a geometry shader with stream output ability, defined by ElementList. */
 	// FlushType: Wait RHI Thread
-	virtual FGeometryShaderRHIRef RHICreateGeometryShaderWithStreamOutput(const TArray<uint8>& Code, const FStreamOutElementList& ElementList, uint32 NumStrides, const uint32* Strides, int32 RasterizedStream) override
+	virtual FGeometryShaderRHIRef RHICreateGeometryShaderWithStreamOutput(const TArray<uint8>& Code, const FStreamOutElementList& ElementList, uint32 NumStrides, const uint32* Strides, int32 RasterizedStream) override final
 	{
 		return RHI->RHICreateGeometryShaderWithStreamOutput(Code, ElementList, NumStrides, Strides, RasterizedStream);
 	}
 
 	/** Creates a geometry shader with stream output ability, defined by ElementList. */
 	// FlushType: Wait RHI Thread
-	virtual FGeometryShaderRHIRef RHICreateGeometryShaderWithStreamOutput(const FStreamOutElementList& ElementList, uint32 NumStrides, const uint32* Strides, int32 RasterizedStream, FRHIShaderLibrary* Library, FSHAHash Hash) override
+	virtual FGeometryShaderRHIRef RHICreateGeometryShaderWithStreamOutput(const FStreamOutElementList& ElementList, uint32 NumStrides, const uint32* Strides, int32 RasterizedStream, FRHIShaderLibrary* Library, FSHAHash Hash) override final
 	{
 		return RHI->RHICreateGeometryShaderWithStreamOutput(ElementList, NumStrides, Strides, RasterizedStream, Library, Hash);
 	}
 
 	// Some RHIs can have pending messages/logs for error tracking, or debug modes
-	virtual void FlushPendingLogs() override
+	virtual void FlushPendingLogs() override final
 	{
 		RHI->FlushPendingLogs();
 	}
 
 	// FlushType: Wait RHI Thread
-	virtual FComputeShaderRHIRef RHICreateComputeShader(const TArray<uint8>& Code) override
+	virtual FComputeShaderRHIRef RHICreateComputeShader(const TArray<uint8>& Code) override final
 	{
 		return RHI->RHICreateComputeShader(Code);
 	}
 	// FlushType: Wait RHI Thread
-	virtual FComputeShaderRHIRef RHICreateComputeShader(FRHIShaderLibrary* Library, FSHAHash Hash) override
+	virtual FComputeShaderRHIRef RHICreateComputeShader(FRHIShaderLibrary* Library, FSHAHash Hash) override final
 	{
 		return RHI->RHICreateComputeShader(Library, Hash);
 	}
@@ -175,7 +176,7 @@ public:
 	 * @return The new library if one exists and can be constructed, otherwise nil.
 	 */
 	 // FlushType: Must be Thread-Safe.
-	virtual FRHIShaderLibraryRef RHICreateShaderLibrary(EShaderPlatform Platform, FString const& FilePath, FString const& Name) override
+	virtual FRHIShaderLibraryRef RHICreateShaderLibrary(EShaderPlatform Platform, FString const& FilePath, FString const& Name) override final
 	{
 		return RHI->RHICreateShaderLibrary(Platform, FilePath, Name);
 	}
@@ -187,12 +188,12 @@ public:
 	* @return The new Fence.
 	*/
 	// FlushType: Thread safe, but varies depending on the RHI	
-	virtual FComputeFenceRHIRef RHICreateComputeFence(const FName& Name) override
+	virtual FComputeFenceRHIRef RHICreateComputeFence(const FName& Name) override final
 	{
 		return RHI->RHICreateComputeFence(Name);
 	}
 
-	virtual FGPUFenceRHIRef RHICreateGPUFence(const FName &Name) override
+	virtual FGPUFenceRHIRef RHICreateGPUFence(const FName &Name) override final
 	{
 		return RHI->RHICreateGPUFence(Name);
 	}
@@ -202,7 +203,7 @@ public:
 	* @return The new staging-buffer.
 	*/
 	// FlushType: Thread safe.	
-	virtual FStagingBufferRHIRef RHICreateStagingBuffer() override
+	virtual FStagingBufferRHIRef RHICreateStagingBuffer() override final
 	{
 		return RHI->RHICreateStagingBuffer();
 	}
@@ -215,7 +216,7 @@ public:
 	 * @param SizeRHI The length of the region in the buffer to lock.
 	 * @returns A pointer to the data starting at 'Offset' and of length 'SizeRHI' from 'StagingBuffer', or nullptr when there is an error.
 	 */
-	virtual void* RHILockStagingBuffer(FRHIStagingBuffer* StagingBuffer, uint32 Offset, uint32 SizeRHI) override
+	virtual void* RHILockStagingBuffer(FRHIStagingBuffer* StagingBuffer, uint32 Offset, uint32 SizeRHI) override final
 	{
 		return RHI->RHILockStagingBuffer(StagingBuffer, Offset, SizeRHI);
 	}
@@ -224,7 +225,7 @@ public:
 	 * Unlock a staging buffer previously locked with RHILockStagingBuffer.
 	 * @param StagingBuffer The buffer that was previously locked.
 	 */
-	virtual void RHIUnlockStagingBuffer(FRHIStagingBuffer* StagingBuffer) override
+	virtual void RHIUnlockStagingBuffer(FRHIStagingBuffer* StagingBuffer) override final
 	{
 		RHI->RHIUnlockStagingBuffer(StagingBuffer);
 	}
@@ -238,7 +239,7 @@ public:
 	 * @param SizeRHI The length of the region in the buffer to lock.
 	 * @returns A pointer to the data starting at 'Offset' and of length 'SizeRHI' from 'StagingBuffer', or nullptr when there is an error.
 	 */
-	virtual void* LockStagingBuffer_RenderThread(class FRHICommandListImmediate& RHICmdList, FRHIStagingBuffer* StagingBuffer, uint32 Offset, uint32 SizeRHI) override
+	virtual void* LockStagingBuffer_RenderThread(class FRHICommandListImmediate& RHICmdList, FRHIStagingBuffer* StagingBuffer, uint32 Offset, uint32 SizeRHI) override final
 	{
 		return RHI->LockStagingBuffer_RenderThread(RHICmdList, StagingBuffer, Offset, SizeRHI);
 	}
@@ -248,7 +249,7 @@ public:
 	 * @param RHICmdList The command-list to execute on or synchronize with.
 	 * @param StagingBuffer The buffer what was previously locked.
 	 */
-	virtual void UnlockStagingBuffer_RenderThread(class FRHICommandListImmediate& RHICmdList, FRHIStagingBuffer* StagingBuffer) override
+	virtual void UnlockStagingBuffer_RenderThread(class FRHICommandListImmediate& RHICmdList, FRHIStagingBuffer* StagingBuffer) override final
 	{
 		RHI->UnlockStagingBuffer_RenderThread(RHICmdList, StagingBuffer);
 	}
@@ -265,7 +266,7 @@ public:
 	* @param PixelShader - existing pixel shader
 	*/
 	// FlushType: Thread safe, but varies depending on the RHI
-	virtual FBoundShaderStateRHIRef RHICreateBoundShaderState(FRHIVertexDeclaration* VertexDeclaration, FRHIVertexShader* VertexShader, FRHIHullShader* HullShader, FRHIDomainShader* DomainShader, FRHIPixelShader* PixelShader, FRHIGeometryShader* GeometryShader) override
+	virtual FBoundShaderStateRHIRef RHICreateBoundShaderState(FRHIVertexDeclaration* VertexDeclaration, FRHIVertexShader* VertexShader, FRHIHullShader* HullShader, FRHIDomainShader* DomainShader, FRHIPixelShader* PixelShader, FRHIGeometryShader* GeometryShader) override final
 	{
 		return RHI->RHICreateBoundShaderState(VertexDeclaration, VertexShader, HullShader, DomainShader, PixelShader, GeometryShader);
 	}
@@ -281,22 +282,22 @@ public:
 	*/
 	// FlushType: Thread safe
 	// TODO: [PSO API] Make pure virtual
-	virtual FGraphicsPipelineStateRHIRef RHICreateGraphicsPipelineState(const FGraphicsPipelineStateInitializer& Initializer) override
+	virtual FGraphicsPipelineStateRHIRef RHICreateGraphicsPipelineState(const FGraphicsPipelineStateInitializer& Initializer) override final
 	{
 		return RHI->RHICreateGraphicsPipelineState(Initializer);
 	}
 
-	virtual TRefCountPtr<FRHIComputePipelineState> RHICreateComputePipelineState(FRHIComputeShader* ComputeShader) override
+	virtual TRefCountPtr<FRHIComputePipelineState> RHICreateComputePipelineState(FRHIComputeShader* ComputeShader) override final
 	{
 		return RHI->RHICreateComputePipelineState(ComputeShader);
 	}
 
-	virtual FGraphicsPipelineStateRHIRef RHICreateGraphicsPipelineState(const FGraphicsPipelineStateInitializer& Initializer, FRHIPipelineBinaryLibrary* PipelineBinary) override
+	virtual FGraphicsPipelineStateRHIRef RHICreateGraphicsPipelineState(const FGraphicsPipelineStateInitializer& Initializer, FRHIPipelineBinaryLibrary* PipelineBinary) override final
 	{
 		return RHI->RHICreateGraphicsPipelineState(Initializer);
 	}
 
-	virtual TRefCountPtr<FRHIComputePipelineState> RHICreateComputePipelineState(FRHIComputeShader* ComputeShader, FRHIPipelineBinaryLibrary* PipelineBinary) override
+	virtual TRefCountPtr<FRHIComputePipelineState> RHICreateComputePipelineState(FRHIComputeShader* ComputeShader, FRHIPipelineBinaryLibrary* PipelineBinary) override final
 	{
 		return RHI->RHICreateComputePipelineState(ComputeShader, PipelineBinary);
 	}
@@ -308,31 +309,31 @@ public:
 	* @param NumBytes - The number of bytes the uniform buffer should contain.
 	* @return The new uniform buffer.
 	*/
-	// FlushType: Thread safe, but varies depending on the RHI override
-	virtual FUniformBufferRHIRef RHICreateUniformBuffer(const void* Contents, const FRHIUniformBufferLayout& Layout, EUniformBufferUsage Usage, EUniformBufferValidation Validation) override
+	// FlushType: Thread safe, but varies depending on the RHI override final
+	virtual FUniformBufferRHIRef RHICreateUniformBuffer(const void* Contents, const FRHIUniformBufferLayout& Layout, EUniformBufferUsage Usage, EUniformBufferValidation Validation) override final
 	{
 		return RHI->RHICreateUniformBuffer(Contents, Layout, Usage, Validation);
 	}
 
-	virtual void RHIUpdateUniformBuffer(FRHIUniformBuffer* UniformBufferRHI, const void* Contents) override
+	virtual void RHIUpdateUniformBuffer(FRHIUniformBuffer* UniformBufferRHI, const void* Contents) override final
 	{
 		RHI->RHIUpdateUniformBuffer(UniformBufferRHI, Contents);
 	}
 
 	// FlushType: Wait RHI Thread
-	virtual FIndexBufferRHIRef RHICreateIndexBuffer(uint32 Stride, uint32 Size, uint32 InUsage, FRHIResourceCreateInfo& CreateInfo) override
+	virtual FIndexBufferRHIRef RHICreateIndexBuffer(uint32 Stride, uint32 Size, uint32 InUsage, FRHIResourceCreateInfo& CreateInfo) override final
 	{
 		return RHI->RHICreateIndexBuffer(Stride, Size, InUsage, CreateInfo);
 	}
 
 	// FlushType: Flush RHI Thread
-	virtual void* RHILockIndexBuffer(FRHIIndexBuffer* IndexBuffer, uint32 Offset, uint32 Size, EResourceLockMode LockMode) override
+	virtual void* RHILockIndexBuffer(FRHIIndexBuffer* IndexBuffer, uint32 Offset, uint32 Size, EResourceLockMode LockMode) override final
 	{
 		return RHI->RHILockIndexBuffer(IndexBuffer, Offset, Size, LockMode);
 	}
 
 	// FlushType: Flush RHI Thread
-	virtual void RHIUnlockIndexBuffer(FRHIIndexBuffer* IndexBuffer) override
+	virtual void RHIUnlockIndexBuffer(FRHIIndexBuffer* IndexBuffer) override final
 	{
 		RHI->RHIUnlockIndexBuffer(IndexBuffer);
 	}
@@ -341,26 +342,26 @@ public:
 	* @param ResourceArray - An optional pointer to a resource array containing the resource's data.
 	*/
 	// FlushType: Wait RHI Thread
-	virtual FVertexBufferRHIRef RHICreateVertexBuffer(uint32 Size, uint32 InUsage, FRHIResourceCreateInfo& CreateInfo) override
+	virtual FVertexBufferRHIRef RHICreateVertexBuffer(uint32 Size, uint32 InUsage, FRHIResourceCreateInfo& CreateInfo) override final
 	{
 		return RHI->RHICreateVertexBuffer(Size, InUsage, CreateInfo);
 	}
 
 	// FlushType: Flush RHI Thread
-	virtual void* RHILockVertexBuffer(FRHIVertexBuffer* VertexBuffer, uint32 Offset, uint32 SizeRHI, EResourceLockMode LockMode) override
+	virtual void* RHILockVertexBuffer(FRHIVertexBuffer* VertexBuffer, uint32 Offset, uint32 SizeRHI, EResourceLockMode LockMode) override final
 	{
 		return RHI->RHILockVertexBuffer(VertexBuffer, Offset, SizeRHI, LockMode);
 	}
 
 	// FlushType: Flush RHI Thread
-	virtual void RHIUnlockVertexBuffer(FRHIVertexBuffer* VertexBuffer) override
+	virtual void RHIUnlockVertexBuffer(FRHIVertexBuffer* VertexBuffer) override final
 	{
 		RHI->RHIUnlockVertexBuffer(VertexBuffer);
 	}
 
 	/** Copies the contents of one vertex buffer to another vertex buffer.  They must have identical sizes. */
 	// FlushType: Flush Immediate (seems dangerous)
-	virtual void RHICopyVertexBuffer(FRHIVertexBuffer* SourceBuffer, FRHIVertexBuffer* DestBuffer) override
+	virtual void RHICopyVertexBuffer(FRHIVertexBuffer* SourceBuffer, FRHIVertexBuffer* DestBuffer) override final
 	{
 		RHI->RHICopyVertexBuffer(SourceBuffer, DestBuffer);
 	}
@@ -369,75 +370,75 @@ public:
 	* @param ResourceArray - An optional pointer to a resource array containing the resource's data.
 	*/
 	// FlushType: Wait RHI Thread
-	virtual FStructuredBufferRHIRef RHICreateStructuredBuffer(uint32 Stride, uint32 Size, uint32 InUsage, FRHIResourceCreateInfo& CreateInfo) override
+	virtual FStructuredBufferRHIRef RHICreateStructuredBuffer(uint32 Stride, uint32 Size, uint32 InUsage, FRHIResourceCreateInfo& CreateInfo) override final
 	{
 		return RHI->RHICreateStructuredBuffer(Stride, Size, InUsage, CreateInfo);
 	}
 
 	// FlushType: Flush RHI Thread
-	virtual void* RHILockStructuredBuffer(FRHIStructuredBuffer* StructuredBuffer, uint32 Offset, uint32 SizeRHI, EResourceLockMode LockMode) override
+	virtual void* RHILockStructuredBuffer(FRHIStructuredBuffer* StructuredBuffer, uint32 Offset, uint32 SizeRHI, EResourceLockMode LockMode) override final
 	{
 		return RHI->RHILockStructuredBuffer(StructuredBuffer, Offset, SizeRHI, LockMode);
 	}
 
 	// FlushType: Flush RHI Thread
-	virtual void RHIUnlockStructuredBuffer(FRHIStructuredBuffer* StructuredBuffer) override
+	virtual void RHIUnlockStructuredBuffer(FRHIStructuredBuffer* StructuredBuffer) override final
 	{
 		RHI->RHIUnlockStructuredBuffer(StructuredBuffer);
 	}
 
 	/** Creates an unordered access view of the given structured buffer. */
 	// FlushType: Wait RHI Thread
-	virtual FUnorderedAccessViewRHIRef RHICreateUnorderedAccessView(FRHIStructuredBuffer* StructuredBuffer, bool bUseUAVCounter, bool bAppendBuffer) override
+	virtual FUnorderedAccessViewRHIRef RHICreateUnorderedAccessView(FRHIStructuredBuffer* StructuredBuffer, bool bUseUAVCounter, bool bAppendBuffer) override final
 	{
 		return RHI->RHICreateUnorderedAccessView(StructuredBuffer, bUseUAVCounter, bAppendBuffer);
 	}
 
 	/** Creates an unordered access view of the given texture. */
 	// FlushType: Wait RHI Thread
-	virtual FUnorderedAccessViewRHIRef RHICreateUnorderedAccessView(FRHITexture* Texture, uint32 MipLevel) override
+	virtual FUnorderedAccessViewRHIRef RHICreateUnorderedAccessView(FRHITexture* Texture, uint32 MipLevel) override final
 	{
 		return RHI->RHICreateUnorderedAccessView(Texture, MipLevel);
 	}
 
 	/** Creates an unordered access view of the given vertex buffer. */
 	// FlushType: Wait RHI Thread
-	virtual FUnorderedAccessViewRHIRef RHICreateUnorderedAccessView(FRHIVertexBuffer* VertexBuffer, uint8 Format) override
+	virtual FUnorderedAccessViewRHIRef RHICreateUnorderedAccessView(FRHIVertexBuffer* VertexBuffer, uint8 Format) override final
 	{
 		return RHI->RHICreateUnorderedAccessView(VertexBuffer, Format);
 	}
 
 	/** Creates an unordered access view of the given index buffer. */
 	// FlushType: Wait RHI Thread
-	virtual FUnorderedAccessViewRHIRef RHICreateUnorderedAccessView(FRHIIndexBuffer* IndexBuffer, uint8 Format) override
+	virtual FUnorderedAccessViewRHIRef RHICreateUnorderedAccessView(FRHIIndexBuffer* IndexBuffer, uint8 Format) override final
 	{
 		return RHI->RHICreateUnorderedAccessView(IndexBuffer, Format);
 	}
 
 	/** Creates a shader resource view of the given structured buffer. */
 	// FlushType: Wait RHI Thread
-	virtual FShaderResourceViewRHIRef RHICreateShaderResourceView(FRHIStructuredBuffer* StructuredBuffer) override
+	virtual FShaderResourceViewRHIRef RHICreateShaderResourceView(FRHIStructuredBuffer* StructuredBuffer) override final
 	{
 		return RHI->RHICreateShaderResourceView(StructuredBuffer);
 	}
 
 	/** Creates a shader resource view of the given vertex buffer. */
 	// FlushType: Wait RHI Thread
-	virtual FShaderResourceViewRHIRef RHICreateShaderResourceView(FRHIVertexBuffer* VertexBuffer, uint32 Stride, uint8 Format) override
+	virtual FShaderResourceViewRHIRef RHICreateShaderResourceView(FRHIVertexBuffer* VertexBuffer, uint32 Stride, uint8 Format) override final
 	{
 		return RHI->RHICreateShaderResourceView(VertexBuffer, Stride, Format);
 	}
 
 	/** Creates a shader resource view of the given index buffer. */
 	// FlushType: Wait RHI Thread
-	virtual FShaderResourceViewRHIRef RHICreateShaderResourceView(FRHIIndexBuffer* Buffer) override
+	virtual FShaderResourceViewRHIRef RHICreateShaderResourceView(FRHIIndexBuffer* Buffer) override final
 	{
 		return RHI->RHICreateShaderResourceView(Buffer);
 	}
 
 	// Must be called on RHI thread timeline
 	// Make sure to call RHIThreadFence(true) afterwards so that parallel translation doesn't refer old resources
-	virtual void RHIUpdateShaderResourceView(FRHIShaderResourceView* SRV, FRHIVertexBuffer* VertexBuffer, uint32 Stride, uint8 Format) override
+	virtual void RHIUpdateShaderResourceView(FRHIShaderResourceView* SRV, FRHIVertexBuffer* VertexBuffer, uint32 Stride, uint8 Format) override final
 	{
 		RHI->RHIUpdateShaderResourceView(SRV, VertexBuffer, Stride, Format);
 	}
@@ -454,7 +455,7 @@ public:
 	* @param OutAlign - Alignment required for this texture.  Output parameter.
 	*/
 	// FlushType: Thread safe
-	virtual uint64 RHICalcTexture2DPlatformSize(uint32 SizeX, uint32 SizeY, uint8 Format, uint32 NumMips, uint32 NumSamples, uint32 Flags, uint32& OutAlign) override
+	virtual uint64 RHICalcTexture2DPlatformSize(uint32 SizeX, uint32 SizeY, uint8 Format, uint32 NumMips, uint32 NumSamples, uint32 Flags, uint32& OutAlign) override final
 	{
 		return RHI->RHICalcTexture2DPlatformSize(SizeX, SizeY, Format, NumMips, NumSamples, Flags, OutAlign);
 	}
@@ -470,7 +471,7 @@ public:
 	* @param OutAlign - Alignment required for this texture.  Output parameter.
 	*/
 	// FlushType: Thread safe
-	virtual uint64 RHICalcTexture3DPlatformSize(uint32 SizeX, uint32 SizeY, uint32 SizeZ, uint8 Format, uint32 NumMips, uint32 Flags, uint32& OutAlign) override
+	virtual uint64 RHICalcTexture3DPlatformSize(uint32 SizeX, uint32 SizeY, uint32 SizeZ, uint8 Format, uint32 NumMips, uint32 Flags, uint32& OutAlign) override final
 	{
 		return RHI->RHICalcTexture3DPlatformSize(SizeX, SizeY, SizeZ, Format, NumMips, Flags, OutAlign);
 	}
@@ -484,7 +485,7 @@ public:
 	* @param OutAlign - Alignment required for this texture.  Output parameter.
 	*/
 	// FlushType: Thread safe
-	virtual uint64 RHICalcTextureCubePlatformSize(uint32 Size, uint8 Format, uint32 NumMips, uint32 Flags, uint32& OutAlign) override
+	virtual uint64 RHICalcTextureCubePlatformSize(uint32 Size, uint8 Format, uint32 NumMips, uint32 Flags, uint32& OutAlign) override final
 	{
 		return RHI->RHICalcTextureCubePlatformSize(Size, Format, NumMips, Flags, OutAlign);
 	}
@@ -494,7 +495,7 @@ public:
 	* safe to call on the main thread
 	*/
 	// FlushType: Thread safe
-	virtual void RHIGetTextureMemoryStats(FTextureMemoryStats& OutStats) override
+	virtual void RHIGetTextureMemoryStats(FTextureMemoryStats& OutStats) override final
 	{
 		RHI->RHIGetTextureMemoryStats(OutStats);
 	}
@@ -511,13 +512,13 @@ public:
 	* @return true if successful, false otherwise
 	*/
 	// FlushType: Flush Immediate
-	virtual bool RHIGetTextureMemoryVisualizeData(FColor* TextureData, int32 SizeX, int32 SizeY, int32 Pitch, int32 PixelSize) override
+	virtual bool RHIGetTextureMemoryVisualizeData(FColor* TextureData, int32 SizeX, int32 SizeY, int32 Pitch, int32 PixelSize) override final
 	{
 		return RHI->RHIGetTextureMemoryVisualizeData(TextureData, SizeX, SizeY, Pitch, PixelSize);
 	}
 
 	// FlushType: Wait RHI Thread
-	virtual FTextureReferenceRHIRef RHICreateTextureReference(FLastRenderTimeContainer* LastRenderTime) override
+	virtual FTextureReferenceRHIRef RHICreateTextureReference(FLastRenderTimeContainer* LastRenderTime) override final
 	{
 		return RHI->RHICreateTextureReference(LastRenderTime);
 	}
@@ -532,7 +533,7 @@ public:
 	* @param Flags - ETextureCreateFlags creation flags
 	*/
 	// FlushType: Wait RHI Thread
-	virtual FTexture2DRHIRef RHICreateTexture2D(uint32 SizeX, uint32 SizeY, uint8 Format, uint32 NumMips, uint32 NumSamples, uint32 Flags, FRHIResourceCreateInfo& CreateInfo) override
+	virtual FTexture2DRHIRef RHICreateTexture2D(uint32 SizeX, uint32 SizeY, uint8 Format, uint32 NumMips, uint32 NumSamples, uint32 Flags, FRHIResourceCreateInfo& CreateInfo) override final
 	{
 		return RHI->RHICreateTexture2D(SizeX, SizeY, Format, NumMips, NumSamples, Flags, CreateInfo);
 	}
@@ -547,7 +548,7 @@ public:
 	* @param Flags - ETextureCreateFlags creation flags
 	*/
 	// FlushType: Wait RHI Thread
-	virtual FTexture2DRHIRef RHICreateTextureExternal2D(uint32 SizeX, uint32 SizeY, uint8 Format, uint32 NumMips, uint32 NumSamples, uint32 Flags, FRHIResourceCreateInfo& CreateInfo) override
+	virtual FTexture2DRHIRef RHICreateTextureExternal2D(uint32 SizeX, uint32 SizeY, uint8 Format, uint32 NumMips, uint32 NumSamples, uint32 Flags, FRHIResourceCreateInfo& CreateInfo) override final
 	{
 		return RHI->RHICreateTextureExternal2D(SizeX, SizeY, Format, NumMips, NumSamples, Flags, CreateInfo);
 	}
@@ -556,7 +557,7 @@ public:
 	* Creates an FStructuredBuffer for the RT write mask of a render target
 	* @param RenderTarget - the RT to create the buffer for
 	*/
-	virtual FStructuredBufferRHIRef RHICreateRTWriteMaskBuffer(FRHITexture2D* RenderTarget) override
+	virtual FStructuredBufferRHIRef RHICreateRTWriteMaskBuffer(FRHITexture2D* RenderTarget) override final
 	{
 		return RHI->RHICreateRTWriteMaskBuffer(RenderTarget);
 	}
@@ -575,7 +576,7 @@ public:
 	* @returns a reference to a 2D texture resource
 	*/
 	// FlushType: Thread safe
-	virtual FTexture2DRHIRef RHIAsyncCreateTexture2D(uint32 SizeX, uint32 SizeY, uint8 Format, uint32 NumMips, uint32 Flags, void** InitialMipData, uint32 NumInitialMips) override
+	virtual FTexture2DRHIRef RHIAsyncCreateTexture2D(uint32 SizeX, uint32 SizeY, uint8 Format, uint32 NumMips, uint32 Flags, void** InitialMipData, uint32 NumInitialMips) override final
 	{
 		return RHI->RHIAsyncCreateTexture2D(SizeX, SizeY, Format, NumMips, Flags, InitialMipData, NumInitialMips);
 	}
@@ -588,7 +589,7 @@ public:
 	* @param SrcTexture2D - source texture
 	*/
 	// FlushType: Flush RHI Thread
-	virtual void RHICopySharedMips(FRHITexture2D* DestTexture2D, FRHITexture2D* SrcTexture2D) override
+	virtual void RHICopySharedMips(FRHITexture2D* DestTexture2D, FRHITexture2D* SrcTexture2D) override final
 	{
 		RHI->RHICopySharedMips(DestTexture2D, SrcTexture2D);
 	}
@@ -602,7 +603,7 @@ public:
 	* @param PullData - whether the source writes the data to the dest, or the dest reads the data from the source.
 	*/
 	// FlushType: Flush RHI Thread
-	virtual void RHITransferTexture(FRHITexture2D* Texture, FIntRect Rect, uint32 SrcGPUIndex, uint32 DestGPUIndex, bool bPullData) override
+	virtual void RHITransferTexture(FRHITexture2D* Texture, FIntRect Rect, uint32 SrcGPUIndex, uint32 DestGPUIndex, bool bPullData) override final
 	{
 		RHI->RHITransferTexture(Texture, Rect, SrcGPUIndex, DestGPUIndex, bPullData);
 	}
@@ -617,7 +618,7 @@ public:
 	* @param Flags - ETextureCreateFlags creation flags
 	*/
 	// FlushType: Wait RHI Thread
-	virtual FTexture2DArrayRHIRef RHICreateTexture2DArray(uint32 SizeX, uint32 SizeY, uint32 SizeZ, uint8 Format, uint32 NumMips, uint32 Flags, FRHIResourceCreateInfo& CreateInfo) override
+	virtual FTexture2DArrayRHIRef RHICreateTexture2DArray(uint32 SizeX, uint32 SizeY, uint32 SizeZ, uint8 Format, uint32 NumMips, uint32 Flags, FRHIResourceCreateInfo& CreateInfo) override final
 	{
 		return RHI->RHICreateTexture2DArray(SizeX, SizeY, SizeZ, Format, NumMips, Flags, CreateInfo);
 	}
@@ -632,7 +633,7 @@ public:
 	* @param Flags - ETextureCreateFlags creation flags
 	*/
 	// FlushType: Wait RHI Thread
-	virtual FTexture3DRHIRef RHICreateTexture3D(uint32 SizeX, uint32 SizeY, uint32 SizeZ, uint8 Format, uint32 NumMips, uint32 Flags, FRHIResourceCreateInfo& CreateInfo) override
+	virtual FTexture3DRHIRef RHICreateTexture3D(uint32 SizeX, uint32 SizeY, uint32 SizeZ, uint8 Format, uint32 NumMips, uint32 Flags, FRHIResourceCreateInfo& CreateInfo) override final
 	{
 		return RHI->RHICreateTexture3D(SizeX, SizeY, SizeZ, Format, NumMips, Flags, CreateInfo);
 	}
@@ -641,7 +642,7 @@ public:
 	* @param Ref may be 0
 	*/
 	// FlushType: Thread safe
-	virtual void RHIGetResourceInfo(FRHITexture* Ref, FRHIResourceInfo& OutInfo) override
+	virtual void RHIGetResourceInfo(FRHITexture* Ref, FRHIResourceInfo& OutInfo) override final
 	{
 		return RHI->RHIGetResourceInfo(Ref, OutInfo);
 	}
@@ -651,7 +652,7 @@ public:
 	* mip level. Useful when rendering to one mip while sampling from another.
 	*/
 	// FlushType: Wait RHI Thread
-	virtual FShaderResourceViewRHIRef RHICreateShaderResourceView(FRHITexture2D* Texture2DRHI, uint8 MipLevel) override
+	virtual FShaderResourceViewRHIRef RHICreateShaderResourceView(FRHITexture2D* Texture2DRHI, uint8 MipLevel) override final
 	{
 		return RHI->RHICreateShaderResourceView(Texture2DRHI, MipLevel);
 	}
@@ -661,7 +662,7 @@ public:
 	* format from the original.  Useful when sampling stencil.
 	*/
 	// FlushType: Wait RHI Thread
-	virtual FShaderResourceViewRHIRef RHICreateShaderResourceView(FRHITexture2D* Texture2DRHI, uint8 MipLevel, uint8 NumMipLevels, uint8 Format) override
+	virtual FShaderResourceViewRHIRef RHICreateShaderResourceView(FRHITexture2D* Texture2DRHI, uint8 MipLevel, uint8 NumMipLevels, uint8 Format) override final
 	{
 		return RHI->RHICreateShaderResourceView(Texture2DRHI, MipLevel, NumMipLevels, Format);
 	}
@@ -671,7 +672,7 @@ public:
 	* mip level.
 	*/
 	// FlushType: Wait RHI Thread
-	virtual FShaderResourceViewRHIRef RHICreateShaderResourceView(FRHITexture3D* Texture3DRHI, uint8 MipLevel) override
+	virtual FShaderResourceViewRHIRef RHICreateShaderResourceView(FRHITexture3D* Texture3DRHI, uint8 MipLevel) override final
 	{
 		return RHI->RHICreateShaderResourceView(Texture3DRHI, MipLevel);
 	}
@@ -681,7 +682,7 @@ public:
 	* mip level.
 	*/
 	// FlushType: Wait RHI Thread
-	virtual FShaderResourceViewRHIRef RHICreateShaderResourceView(FRHITexture2DArray* Texture2DArrayRHI, uint8 MipLevel) override
+	virtual FShaderResourceViewRHIRef RHICreateShaderResourceView(FRHITexture2DArray* Texture2DArrayRHI, uint8 MipLevel) override final
 	{
 		return RHI->RHICreateShaderResourceView(Texture2DArrayRHI, MipLevel);
 	}
@@ -691,7 +692,7 @@ public:
 	* mip level.
 	*/
 	// FlushType: Wait RHI Thread
-	virtual FShaderResourceViewRHIRef RHICreateShaderResourceView(FRHITextureCube* TextureCubeRHI, uint8 MipLevel) override
+	virtual FShaderResourceViewRHIRef RHICreateShaderResourceView(FRHITextureCube* TextureCubeRHI, uint8 MipLevel) override final
 	{
 		return RHI->RHICreateShaderResourceView(TextureCubeRHI, MipLevel);
 	}
@@ -700,7 +701,7 @@ public:
 	* Generates mip maps for a texture.
 	*/
 	// FlushType: Flush Immediate (NP: this should be queued on the command list for RHI thread execution, not flushed)
-	virtual void RHIGenerateMips(FRHITexture* Texture) override
+	virtual void RHIGenerateMips(FRHITexture* Texture) override final
 	{
 		return RHI->RHIGenerateMips(Texture);
 	}
@@ -712,7 +713,7 @@ public:
 	* @return					- Size in Bytes
 	*/
 	// FlushType: Thread safe
-	virtual uint32 RHIComputeMemorySize(FRHITexture* TextureRHI) override
+	virtual uint32 RHIComputeMemorySize(FRHITexture* TextureRHI) override final
 	{
 		return RHI->RHIComputeMemorySize(TextureRHI);
 	}
@@ -734,7 +735,7 @@ public:
 	*/
 	// FlushType: Flush RHI Thread
 	// NP: Note that no RHI currently implements this as an async call, we should simplify the API.
-	virtual FTexture2DRHIRef RHIAsyncReallocateTexture2D(FRHITexture2D* Texture2D, int32 NewMipCount, int32 NewSizeX, int32 NewSizeY, FThreadSafeCounter* RequestStatus) override
+	virtual FTexture2DRHIRef RHIAsyncReallocateTexture2D(FRHITexture2D* Texture2D, int32 NewMipCount, int32 NewSizeX, int32 NewSizeY, FThreadSafeCounter* RequestStatus) override final
 	{
 		return RHI->RHIAsyncReallocateTexture2D(Texture2D, NewMipCount, NewSizeX, NewSizeY, RequestStatus);
 	}
@@ -751,7 +752,7 @@ public:
 	*	TexRealloc_InProgress	Reallocation is still in progress, try again later
 	*/
 	// FlushType: Wait RHI Thread
-	virtual ETextureReallocationStatus RHIFinalizeAsyncReallocateTexture2D(FRHITexture2D* Texture2D, bool bBlockUntilCompleted) override
+	virtual ETextureReallocationStatus RHIFinalizeAsyncReallocateTexture2D(FRHITexture2D* Texture2D, bool bBlockUntilCompleted) override final
 	{
 		return RHI->RHIFinalizeAsyncReallocateTexture2D(Texture2D, bBlockUntilCompleted);
 	}
@@ -765,7 +766,7 @@ public:
 	* @return						Reallocation status
 	*/
 	// FlushType: Wait RHI Thread
-	virtual ETextureReallocationStatus RHICancelAsyncReallocateTexture2D(FRHITexture2D* Texture2D, bool bBlockUntilCompleted) override
+	virtual ETextureReallocationStatus RHICancelAsyncReallocateTexture2D(FRHITexture2D* Texture2D, bool bBlockUntilCompleted) override final
 	{
 		return RHI->RHICancelAsyncReallocateTexture2D(Texture2D, bBlockUntilCompleted);
 	}
@@ -780,7 +781,7 @@ public:
 	* @return pointer to the CPU accessible resource data
 	*/
 	// FlushType: Flush RHI Thread
-	virtual void* RHILockTexture2D(FRHITexture2D* Texture, uint32 MipIndex, EResourceLockMode LockMode, uint32& DestStride, bool bLockWithinMiptail) override
+	virtual void* RHILockTexture2D(FRHITexture2D* Texture, uint32 MipIndex, EResourceLockMode LockMode, uint32& DestStride, bool bLockWithinMiptail) override final
 	{
 		return RHI->RHILockTexture2D(Texture, MipIndex, LockMode, DestStride, bLockWithinMiptail);
 	}
@@ -792,7 +793,7 @@ public:
 	* @param bLockWithinMiptail - for platforms that support packed miptails allow locking of individual mip levels within the miptail
 	*/
 	// FlushType: Flush RHI Thread
-	virtual void RHIUnlockTexture2D(FRHITexture2D* Texture, uint32 MipIndex, bool bLockWithinMiptail) override
+	virtual void RHIUnlockTexture2D(FRHITexture2D* Texture, uint32 MipIndex, bool bLockWithinMiptail) override final
 	{
 		RHI->RHIUnlockTexture2D(Texture, MipIndex, bLockWithinMiptail);
 	}
@@ -807,7 +808,7 @@ public:
 	* @return pointer to the CPU accessible resource data
 	*/
 	// FlushType: Flush RHI Thread
-	virtual void* RHILockTexture2DArray(FRHITexture2DArray* Texture, uint32 TextureIndex, uint32 MipIndex, EResourceLockMode LockMode, uint32& DestStride, bool bLockWithinMiptail) override
+	virtual void* RHILockTexture2DArray(FRHITexture2DArray* Texture, uint32 TextureIndex, uint32 MipIndex, EResourceLockMode LockMode, uint32& DestStride, bool bLockWithinMiptail) override final
 	{
 		return RHI->RHILockTexture2DArray(Texture, TextureIndex, MipIndex, LockMode, DestStride, bLockWithinMiptail);
 	}
@@ -819,7 +820,7 @@ public:
 	* @param bLockWithinMiptail - for platforms that support packed miptails allow locking of individual mip levels within the miptail
 	*/
 	// FlushType: Flush RHI Thread
-	virtual void RHIUnlockTexture2DArray(FRHITexture2DArray* Texture, uint32 TextureIndex, uint32 MipIndex, bool bLockWithinMiptail) override
+	virtual void RHIUnlockTexture2DArray(FRHITexture2DArray* Texture, uint32 TextureIndex, uint32 MipIndex, bool bLockWithinMiptail) override final
 	{
 		RHI->RHIUnlockTexture2DArray(Texture, TextureIndex, MipIndex, bLockWithinMiptail);
 	}
@@ -833,7 +834,7 @@ public:
 	* @param SourceData - source image data, starting at the upper left corner of the source rectangle (in same pixel format as texture)
 	*/
 	// FlushType: Flush RHI Thread
-	virtual void RHIUpdateTexture2D(FRHITexture2D* Texture, uint32 MipIndex, const struct FUpdateTextureRegion2D& UpdateRegion, uint32 SourcePitch, const uint8* SourceData) override
+	virtual void RHIUpdateTexture2D(FRHITexture2D* Texture, uint32 MipIndex, const struct FUpdateTextureRegion2D& UpdateRegion, uint32 SourcePitch, const uint8* SourceData) override final
 	{
 		RHI->RHIUpdateTexture2D(Texture, MipIndex, UpdateRegion, SourcePitch, SourceData);
 	}
@@ -848,7 +849,7 @@ public:
 	* @param SourceData - source image data, starting at the upper left corner of the source rectangle (in same pixel format as texture)
 	*/
 	// FlushType: Flush RHI Thread
-	virtual void RHIUpdateTexture3D(FRHITexture3D* Texture, uint32 MipIndex, const struct FUpdateTextureRegion3D& UpdateRegion, uint32 SourceRowPitch, uint32 SourceDepthPitch, const uint8* SourceData) override
+	virtual void RHIUpdateTexture3D(FRHITexture3D* Texture, uint32 MipIndex, const struct FUpdateTextureRegion3D& UpdateRegion, uint32 SourceRowPitch, uint32 SourceDepthPitch, const uint8* SourceData) override final
 	{
 		RHI->RHIUpdateTexture3D(Texture, MipIndex, UpdateRegion, SourceRowPitch, SourceDepthPitch, SourceData);
 	}
@@ -861,7 +862,7 @@ public:
 	* @param Flags - ETextureCreateFlags creation flags
 	*/
 	// FlushType: Wait RHI Thread
-	virtual FTextureCubeRHIRef RHICreateTextureCube(uint32 Size, uint8 Format, uint32 NumMips, uint32 Flags, FRHIResourceCreateInfo& CreateInfo) override
+	virtual FTextureCubeRHIRef RHICreateTextureCube(uint32 Size, uint8 Format, uint32 NumMips, uint32 Flags, FRHIResourceCreateInfo& CreateInfo) override final
 	{
 		return RHI->RHICreateTextureCube(Size, Format, NumMips, Flags, CreateInfo);
 	}
@@ -875,7 +876,7 @@ public:
 	* @param Flags - ETextureCreateFlags creation flags
 	*/
 	// FlushType: Wait RHI Thread
-	virtual FTextureCubeRHIRef RHICreateTextureCubeArray(uint32 Size, uint32 ArraySize, uint8 Format, uint32 NumMips, uint32 Flags, FRHIResourceCreateInfo& CreateInfo) override
+	virtual FTextureCubeRHIRef RHICreateTextureCubeArray(uint32 Size, uint32 ArraySize, uint8 Format, uint32 NumMips, uint32 Flags, FRHIResourceCreateInfo& CreateInfo) override final
 	{
 		return RHI->RHICreateTextureCubeArray(Size, ArraySize, Format, NumMips, Flags, CreateInfo);
 	}
@@ -890,7 +891,7 @@ public:
 	* @return pointer to the CPU accessible resource data
 	*/
 	// FlushType: Flush RHI Thread
-	virtual void* RHILockTextureCubeFace(FRHITextureCube* Texture, uint32 FaceIndex, uint32 ArrayIndex, uint32 MipIndex, EResourceLockMode LockMode, uint32& DestStride, bool bLockWithinMiptail) override
+	virtual void* RHILockTextureCubeFace(FRHITextureCube* Texture, uint32 FaceIndex, uint32 ArrayIndex, uint32 MipIndex, EResourceLockMode LockMode, uint32& DestStride, bool bLockWithinMiptail) override final
 	{
 		return RHI->RHILockTextureCubeFace(Texture, FaceIndex, ArrayIndex, MipIndex, LockMode, DestStride, bLockWithinMiptail);
 	}
@@ -902,18 +903,18 @@ public:
 	* @param bLockWithinMiptail - for platforms that support packed miptails allow locking of individual mip levels within the miptail
 	*/
 	// FlushType: Flush RHI Thread
-	virtual void RHIUnlockTextureCubeFace(FRHITextureCube* Texture, uint32 FaceIndex, uint32 ArrayIndex, uint32 MipIndex, bool bLockWithinMiptail) override
+	virtual void RHIUnlockTextureCubeFace(FRHITextureCube* Texture, uint32 FaceIndex, uint32 ArrayIndex, uint32 MipIndex, bool bLockWithinMiptail) override final
 	{
 		RHI->RHIUnlockTextureCubeFace(Texture, FaceIndex, ArrayIndex, MipIndex, bLockWithinMiptail);
 	}
 
 	// FlushType: Thread safe
-	virtual void RHIBindDebugLabelName(FRHITexture* Texture, const TCHAR* Name) override
+	virtual void RHIBindDebugLabelName(FRHITexture* Texture, const TCHAR* Name) override final
 	{
 		RHI->RHIBindDebugLabelName(Texture, Name);
 	}
 
-	virtual void RHIBindDebugLabelName(FRHIUnorderedAccessView* UnorderedAccessViewRHI, const TCHAR* Name) override
+	virtual void RHIBindDebugLabelName(FRHIUnorderedAccessView* UnorderedAccessViewRHI, const TCHAR* Name) override final
 	{
 		RHI->RHIBindDebugLabelName(UnorderedAccessViewRHI, Name);
 	}
@@ -923,81 +924,81 @@ public:
 	* If the format or texture type is unsupported the OutData array will be size 0
 	*/
 	// FlushType: Flush Immediate (seems wrong)
-	virtual void RHIReadSurfaceData(FRHITexture* Texture, FIntRect Rect, TArray<FColor>& OutData, FReadSurfaceDataFlags InFlags) override
+	virtual void RHIReadSurfaceData(FRHITexture* Texture, FIntRect Rect, TArray<FColor>& OutData, FReadSurfaceDataFlags InFlags) override final
 	{
 		RHI->RHIReadSurfaceData(Texture, Rect, OutData, InFlags);
 	}
 
 	// Default fallback; will not work for non-8-bit surfaces and it's extremely slow.
-	virtual void RHIReadSurfaceData(FRHITexture* Texture, FIntRect Rect, TArray<FLinearColor>& OutData, FReadSurfaceDataFlags InFlags) override
+	virtual void RHIReadSurfaceData(FRHITexture* Texture, FIntRect Rect, TArray<FLinearColor>& OutData, FReadSurfaceDataFlags InFlags) override final
 	{
 		RHI->RHIReadSurfaceData(Texture, Rect, OutData, InFlags);
 	}
 
 	/** Watch out for OutData to be 0 (can happen on DXGI_ERROR_DEVICE_REMOVED), don't call RHIUnmapStagingSurface in that case. */
 	// FlushType: Flush Immediate (seems wrong)
-	virtual void RHIMapStagingSurface(FRHITexture* Texture, void*& OutData, int32& OutWidth, int32& OutHeight) override
+	virtual void RHIMapStagingSurface(FRHITexture* Texture, void*& OutData, int32& OutWidth, int32& OutHeight) override final
 	{
 		RHI->RHIMapStagingSurface(Texture, OutData, OutWidth, OutHeight);
 	}
 
 	/** call after a succesful RHIMapStagingSurface() call */
 	// FlushType: Flush Immediate (seems wrong)
-	virtual void RHIUnmapStagingSurface(FRHITexture* Texture) override
+	virtual void RHIUnmapStagingSurface(FRHITexture* Texture) override final
 	{
 		RHI->RHIUnmapStagingSurface(Texture);
 	}
 
 	// FlushType: Flush Immediate (seems wrong)
-	virtual void RHIReadSurfaceFloatData(FRHITexture* Texture, FIntRect Rect, TArray<FFloat16Color>& OutData, ECubeFace CubeFace, int32 ArrayIndex, int32 MipIndex) override
+	virtual void RHIReadSurfaceFloatData(FRHITexture* Texture, FIntRect Rect, TArray<FFloat16Color>& OutData, ECubeFace CubeFace, int32 ArrayIndex, int32 MipIndex) override final
 	{
 		RHI->RHIReadSurfaceFloatData(Texture, Rect, OutData, CubeFace, ArrayIndex, MipIndex);
 	}
 
 	// FlushType: Flush Immediate (seems wrong)
-	virtual void RHIRead3DSurfaceFloatData(FRHITexture* Texture, FIntRect Rect, FIntPoint ZMinMax, TArray<FFloat16Color>& OutData) override
+	virtual void RHIRead3DSurfaceFloatData(FRHITexture* Texture, FIntRect Rect, FIntPoint ZMinMax, TArray<FFloat16Color>& OutData) override final
 	{
 		RHI->RHIRead3DSurfaceFloatData(Texture, Rect, ZMinMax, OutData);
 	}
 
 	// FlushType: Wait RHI Thread
-	virtual FRenderQueryRHIRef RHICreateRenderQuery(ERenderQueryType QueryType) override
+	virtual FRenderQueryRHIRef RHICreateRenderQuery(ERenderQueryType QueryType) override final
 	{
 		return RHI->RHICreateRenderQuery(QueryType);
 	}
 	// CAUTION: Even though this is marked as threadsafe, it is only valid to call from the render thread. It is need not be threadsafe on platforms that do not support or aren't using an RHIThread
 	// FlushType: Thread safe, but varies by RHI
-	virtual bool RHIGetRenderQueryResult(FRHIRenderQuery* RenderQuery, uint64& OutResult, bool bWait) override
+	virtual bool RHIGetRenderQueryResult(FRHIRenderQuery* RenderQuery, uint64& OutResult, bool bWait) override final
 	{
 		return RHI->RHIGetRenderQueryResult(RenderQuery, OutResult, bWait);
 	}
 
 	// FlushType: Thread safe
-	virtual uint32 RHIGetViewportNextPresentGPUIndex(FRHIViewport* Viewport) override
+	virtual uint32 RHIGetViewportNextPresentGPUIndex(FRHIViewport* Viewport) override final
 	{
 		return RHI->RHIGetViewportNextPresentGPUIndex(Viewport);
 	}
 
 	// With RHI thread, this is the current backbuffer from the perspective of the render thread.
 	// FlushType: Thread safe
-	virtual FTexture2DRHIRef RHIGetViewportBackBuffer(FRHIViewport* Viewport) override
+	virtual FTexture2DRHIRef RHIGetViewportBackBuffer(FRHIViewport* Viewport) override final
 	{
 		return RHI->RHIGetViewportBackBuffer(Viewport);
 	}
 
-	virtual FUnorderedAccessViewRHIRef RHIGetViewportBackBufferUAV(FRHIViewport* ViewportRHI) override
+	virtual FUnorderedAccessViewRHIRef RHIGetViewportBackBufferUAV(FRHIViewport* ViewportRHI) override final
 	{
 		return RHI->RHIGetViewportBackBufferUAV(ViewportRHI);
 	}
 
-	virtual FTexture2DRHIRef RHIGetFMaskTexture(FRHITexture* SourceTextureRHI) override
+	virtual FTexture2DRHIRef RHIGetFMaskTexture(FRHITexture* SourceTextureRHI) override final
 	{
 		return RHI->RHIGetFMaskTexture(SourceTextureRHI);
 	}
 
 	// Only relevant with an RHI thread, this advances the backbuffer for the purpose of GetViewportBackBuffer
 	// FlushType: Thread safe
-	virtual void RHIAdvanceFrameForGetViewportBackBuffer(FRHIViewport* Viewport) override
+	virtual void RHIAdvanceFrameForGetViewportBackBuffer(FRHIViewport* Viewport) override final
 	{
 		RHI->RHIAdvanceFrameForGetViewportBackBuffer(Viewport);
 	}
@@ -1005,20 +1006,20 @@ public:
 	* Acquires or releases ownership of the platform-specific rendering context for the calling thread
 	*/
 	// FlushType: Flush RHI Thread
-	virtual void RHIAcquireThreadOwnership() override
+	virtual void RHIAcquireThreadOwnership() override final
 	{
 		RHI->RHIAcquireThreadOwnership();
 	}
 
 	// FlushType: Flush RHI Thread
-	virtual void RHIReleaseThreadOwnership() override
+	virtual void RHIReleaseThreadOwnership() override final
 	{
 		RHI->RHIReleaseThreadOwnership();
 	}
 
 	// Flush driver resources. Typically called when switching contexts/threads
 	// FlushType: Flush RHI Thread
-	virtual void RHIFlushResources() override
+	virtual void RHIFlushResources() override final
 	{
 		RHI->RHIFlushResources();
 	}
@@ -1027,26 +1028,26 @@ public:
 	* Returns the total GPU time taken to render the last frame. Same metric as FPlatformTime::Cycles().
 	*/
 	// FlushType: Thread safe
-	virtual uint32 RHIGetGPUFrameCycles() override
+	virtual uint32 RHIGetGPUFrameCycles() override final
 	{
 		return RHI->RHIGetGPUFrameCycles();
 	}
 
 	//  must be called from the main thread.
 	// FlushType: Thread safe
-	virtual FViewportRHIRef RHICreateViewport(void* WindowHandle, uint32 SizeX, uint32 SizeY, bool bIsFullscreen, EPixelFormat PreferredPixelFormat) override
+	virtual FViewportRHIRef RHICreateViewport(void* WindowHandle, uint32 SizeX, uint32 SizeY, bool bIsFullscreen, EPixelFormat PreferredPixelFormat) override final
 	{
 		return RHI->RHICreateViewport(WindowHandle, SizeX, SizeY, bIsFullscreen, PreferredPixelFormat);
 	}
 
 	//  must be called from the main thread.
 	// FlushType: Thread safe
-	virtual void RHIResizeViewport(FRHIViewport* Viewport, uint32 SizeX, uint32 SizeY, bool bIsFullscreen) override
+	virtual void RHIResizeViewport(FRHIViewport* Viewport, uint32 SizeX, uint32 SizeY, bool bIsFullscreen) override final
 	{
 		RHI->RHIResizeViewport(Viewport, SizeX, SizeY, bIsFullscreen);
 	}
 
-	virtual void RHIResizeViewport(FRHIViewport* Viewport, uint32 SizeX, uint32 SizeY, bool bIsFullscreen, EPixelFormat PreferredPixelFormat) override
+	virtual void RHIResizeViewport(FRHIViewport* Viewport, uint32 SizeX, uint32 SizeY, bool bIsFullscreen, EPixelFormat PreferredPixelFormat) override final
 	{
 		// Default implementation for RHIs that cannot change formats on the fly
 		RHI->RHIResizeViewport(Viewport, SizeX, SizeY, bIsFullscreen, PreferredPixelFormat);
@@ -1054,7 +1055,7 @@ public:
 
 	//  must be called from the main thread.
 	// FlushType: Thread safe
-	virtual void RHITick(float DeltaTime) override
+	virtual void RHITick(float DeltaTime) override final
 	{
 		RHI->RHITick(DeltaTime);
 	}
@@ -1062,57 +1063,57 @@ public:
 	/** Sets stream output targets, for use with a geometry shader created with RHICreateGeometryShaderWithStreamOutput. */
 	//@todo this should be a CMDLIST method
 	// FlushType: Flush Immediate (seems wrong)
-	virtual void RHISetStreamOutTargets(uint32 NumTargets, const FVertexBufferRHIParamRef* VertexBuffers, const uint32* Offsets) override
+	virtual void RHISetStreamOutTargets(uint32 NumTargets, const FVertexBufferRHIParamRef* VertexBuffers, const uint32* Offsets) override final
 	{
 		RHI->RHISetStreamOutTargets(NumTargets, VertexBuffers, Offsets);
 	}
 
 	// Blocks the CPU until the GPU catches up and goes idle.
 	// FlushType: Flush Immediate (seems wrong)
-	virtual void RHIBlockUntilGPUIdle() override
+	virtual void RHIBlockUntilGPUIdle() override final
 	{
 		RHI->RHIBlockUntilGPUIdle();
 	}
 
 	// Kicks the current frame and makes sure GPU is actively working on them
 	// FlushType: Flush Immediate (copied from RHIBlockUntilGPUIdle)
-	virtual void RHISubmitCommandsAndFlushGPU() override
+	virtual void RHISubmitCommandsAndFlushGPU() override final
 	{
 		RHI->RHISubmitCommandsAndFlushGPU();
 	}
 
 	// Tells the RHI we're about to suspend it
-	virtual void RHIBeginSuspendRendering() override
+	virtual void RHIBeginSuspendRendering() override final
 	{
 		RHI->RHIBeginSuspendRendering();
 	}
 
 	// Operations to suspend title rendering and yield control to the system
 	// FlushType: Thread safe
-	virtual void RHISuspendRendering() override
+	virtual void RHISuspendRendering() override final
 	{
 		RHI->RHISuspendRendering();
 	}
 
 	// FlushType: Thread safe
-	virtual void RHIResumeRendering() override
+	virtual void RHIResumeRendering() override final
 	{
 		RHI->RHIResumeRendering();
 	}
 
 	// FlushType: Flush Immediate
-	virtual bool RHIIsRenderingSuspended() override
+	virtual bool RHIIsRenderingSuspended() override final
 	{
 		return RHI->RHIIsRenderingSuspended();
 	}
 
 	// FlushType: Flush Immediate
-	virtual bool RHIEnqueueDecompress(uint8_t* SrcBuffer, uint8_t* DestBuffer, int CompressedSize, void* ErrorCodeBuffer) override
+	virtual bool RHIEnqueueDecompress(uint8_t* SrcBuffer, uint8_t* DestBuffer, int CompressedSize, void* ErrorCodeBuffer) override final
 	{
 		return RHI->RHIEnqueueDecompress(SrcBuffer, DestBuffer, CompressedSize, ErrorCodeBuffer);
 	}
 
-	virtual bool RHIEnqueueCompress(uint8_t* SrcBuffer, uint8_t* DestBuffer, int UnCompressedSize, void* ErrorCodeBuffer) override
+	virtual bool RHIEnqueueCompress(uint8_t* SrcBuffer, uint8_t* DestBuffer, int UnCompressedSize, void* ErrorCodeBuffer) override final
 	{
 		return RHI->RHIEnqueueCompress(SrcBuffer, DestBuffer, UnCompressedSize, ErrorCodeBuffer);
 	}
@@ -1126,7 +1127,7 @@ public:
 	*	@return	bool				true if successfully filled the array
 	*/
 	// FlushType: Thread safe
-	virtual bool RHIGetAvailableResolutions(FScreenResolutionArray& Resolutions, bool bIgnoreRefreshRate) override
+	virtual bool RHIGetAvailableResolutions(FScreenResolutionArray& Resolutions, bool bIgnoreRefreshRate) override final
 	{
 		return RHI->RHIGetAvailableResolutions(Resolutions, bIgnoreRefreshRate);
 	}
@@ -1137,7 +1138,7 @@ public:
 	* @param Height - Input: Desired resolution height in pixels. Output: A height that the platform supports.
 	*/
 	// FlushType: Thread safe
-	virtual void RHIGetSupportedResolution(uint32& Width, uint32& Height) override
+	virtual void RHIGetSupportedResolution(uint32& Width, uint32& Height) override final
 	{
 		RHI->RHIGetSupportedResolution(Width, Height);
 	}
@@ -1149,7 +1150,7 @@ public:
 	* @param FirstMip - the first mip that should be in memory
 	*/
 	// FlushType: Wait RHI Thread
-	virtual void RHIVirtualTextureSetFirstMipInMemory(FRHITexture2D* Texture, uint32 FirstMip) override
+	virtual void RHIVirtualTextureSetFirstMipInMemory(FRHITexture2D* Texture, uint32 FirstMip) override final
 	{
 		RHI->RHIVirtualTextureSetFirstMipInMemory(Texture, FirstMip);
 	}
@@ -1160,7 +1161,7 @@ public:
 	* @param FirstMip - the first mip that should be visible to the GPU
 	*/
 	// FlushType: Wait RHI Thread
-	virtual void RHIVirtualTextureSetFirstMipVisible(FRHITexture2D* Texture, uint32 FirstMip) override
+	virtual void RHIVirtualTextureSetFirstMipVisible(FRHITexture2D* Texture, uint32 FirstMip) override final
 	{
 		RHI->RHIVirtualTextureSetFirstMipVisible(Texture, FirstMip);
 	}
@@ -1169,13 +1170,13 @@ public:
 	* Called once per frame just before deferred deletion in FRHIResource::FlushPendingDeletes
 	*/
 	// FlushType: called from render thread when RHI thread is flushed 
-	virtual void RHIPerFrameRHIFlushComplete() override
+	virtual void RHIPerFrameRHIFlushComplete() override final
 	{
 		RHI->RHIPerFrameRHIFlushComplete();
 	}
 
 	// FlushType: Wait RHI Thread
-	virtual void RHIExecuteCommandList(FRHICommandList* CmdList) override
+	virtual void RHIExecuteCommandList(FRHICommandList* CmdList) override final
 	{
 		return RHI->RHIExecuteCommandList(CmdList);
 	}
@@ -1184,456 +1185,453 @@ public:
 	* Provides access to the native device. Generally this should be avoided but is useful for third party plugins.
 	*/
 	// FlushType: Flush RHI Thread
-	virtual void* RHIGetNativeDevice() override
+	virtual void* RHIGetNativeDevice() override final
 	{
 		return RHI->RHIGetNativeDevice();
 	}
 
 	// FlushType: Thread safe
-	virtual IRHICommandContext* RHIGetDefaultContext() override;
+	virtual IRHICommandContext* RHIGetDefaultContext() override final;
 
 	// FlushType: Thread safe
-	virtual IRHIComputeContext* RHIGetDefaultAsyncComputeContext() override
-	{
-		return RHI->RHIGetDefaultAsyncComputeContext();
-	}
+	virtual IRHIComputeContext* RHIGetDefaultAsyncComputeContext() override final;
 
 	// FlushType: Thread safe
-	virtual class IRHICommandContextContainer* RHIGetCommandContextContainer(int32 Index, int32 Num) override
+	virtual class IRHICommandContextContainer* RHIGetCommandContextContainer(int32 Index, int32 Num) override final
 	{
 		return RHI->RHIGetCommandContextContainer(Index, Num);
 	}
 
 #if WITH_MGPU
 	/** Returns a context for sending commands to the given GPU mask. Default implementation is only valid when not using multi-gpu. */
-	virtual IRHICommandContext* RHIGetDefaultContext(FRHIGPUMask GPUMask) override
+	virtual IRHICommandContext* RHIGetDefaultContext(FRHIGPUMask GPUMask) override final
 	{
 		return RHI->RHIGetDefaultContext(GPUMask);
 	}
 
-	virtual IRHIComputeContext* RHIGetDefaultAsyncComputeContext(FRHIGPUMask GPUMask) override
+	virtual IRHIComputeContext* RHIGetDefaultAsyncComputeContext(FRHIGPUMask GPUMask) override final
 	{
 		return RHI->RHIGetDefaultAsyncComputeContext(GPUMask);
 	}
 
-	virtual IRHICommandContextContainer* RHIGetCommandContextContainer(int32 Index, int32 Num, FRHIGPUMask GPUMask) override
+	virtual IRHICommandContextContainer* RHIGetCommandContextContainer(int32 Index, int32 Num, FRHIGPUMask GPUMask) override final
 	{
 		return RHI->RHIGetCommandContextContainer(Index, Num, GPUMask);
 	}
 #endif // WITH_MGPU
 
 	///////// Pass through functions that allow RHIs to optimize certain calls.
-	virtual FVertexBufferRHIRef CreateAndLockVertexBuffer_RenderThread(class FRHICommandListImmediate& RHICmdList, uint32 Size, uint32 InUsage, FRHIResourceCreateInfo& CreateInfo, void*& OutDataBuffer) override
+	virtual FVertexBufferRHIRef CreateAndLockVertexBuffer_RenderThread(class FRHICommandListImmediate& RHICmdList, uint32 Size, uint32 InUsage, FRHIResourceCreateInfo& CreateInfo, void*& OutDataBuffer) override final
 	{
 		return RHI->CreateAndLockVertexBuffer_RenderThread(RHICmdList, Size, InUsage, CreateInfo, OutDataBuffer);
 	}
 
-	virtual FIndexBufferRHIRef CreateAndLockIndexBuffer_RenderThread(class FRHICommandListImmediate& RHICmdList, uint32 Stride, uint32 Size, uint32 InUsage, FRHIResourceCreateInfo& CreateInfo, void*& OutDataBuffer) override
+	virtual FIndexBufferRHIRef CreateAndLockIndexBuffer_RenderThread(class FRHICommandListImmediate& RHICmdList, uint32 Stride, uint32 Size, uint32 InUsage, FRHIResourceCreateInfo& CreateInfo, void*& OutDataBuffer) override final
 	{
 		return RHI->CreateAndLockIndexBuffer_RenderThread(RHICmdList, Stride, Size, InUsage, CreateInfo, OutDataBuffer);
 	}
 
-	virtual FVertexBufferRHIRef CreateVertexBuffer_RenderThread(class FRHICommandListImmediate& RHICmdList, uint32 Size, uint32 InUsage, FRHIResourceCreateInfo& CreateInfo) override
+	virtual FVertexBufferRHIRef CreateVertexBuffer_RenderThread(class FRHICommandListImmediate& RHICmdList, uint32 Size, uint32 InUsage, FRHIResourceCreateInfo& CreateInfo) override final
 	{
 		return RHI->CreateVertexBuffer_RenderThread(RHICmdList, Size, InUsage, CreateInfo);
 	}
 
-	virtual FStructuredBufferRHIRef CreateStructuredBuffer_RenderThread(class FRHICommandListImmediate& RHICmdList, uint32 Stride, uint32 Size, uint32 InUsage, FRHIResourceCreateInfo& CreateInfo) override
+	virtual FStructuredBufferRHIRef CreateStructuredBuffer_RenderThread(class FRHICommandListImmediate& RHICmdList, uint32 Stride, uint32 Size, uint32 InUsage, FRHIResourceCreateInfo& CreateInfo) override final
 	{
 		return RHI->CreateStructuredBuffer_RenderThread(RHICmdList, Stride, Size, InUsage, CreateInfo);
 	}
 
-	virtual FShaderResourceViewRHIRef CreateShaderResourceView_RenderThread(class FRHICommandListImmediate& RHICmdList, FRHIVertexBuffer* VertexBuffer, uint32 Stride, uint8 Format) override
+	virtual FShaderResourceViewRHIRef CreateShaderResourceView_RenderThread(class FRHICommandListImmediate& RHICmdList, FRHIVertexBuffer* VertexBuffer, uint32 Stride, uint8 Format) override final
 	{
 		return RHI->CreateShaderResourceView_RenderThread(RHICmdList, VertexBuffer, Stride, Format);
 	}
 
-	virtual FShaderResourceViewRHIRef CreateShaderResourceView_RenderThread(class FRHICommandListImmediate& RHICmdList, FRHIIndexBuffer* Buffer) override
+	virtual FShaderResourceViewRHIRef CreateShaderResourceView_RenderThread(class FRHICommandListImmediate& RHICmdList, FRHIIndexBuffer* Buffer) override final
 	{
 		return RHI->CreateShaderResourceView_RenderThread(RHICmdList, Buffer);
 	}
 
-	virtual void* LockVertexBuffer_RenderThread(class FRHICommandListImmediate& RHICmdList, FRHIVertexBuffer* VertexBuffer, uint32 Offset, uint32 SizeRHI, EResourceLockMode LockMode) override
+	virtual void* LockVertexBuffer_RenderThread(class FRHICommandListImmediate& RHICmdList, FRHIVertexBuffer* VertexBuffer, uint32 Offset, uint32 SizeRHI, EResourceLockMode LockMode) override final
 	{
 		return RHI->LockVertexBuffer_RenderThread(RHICmdList, VertexBuffer, Offset, SizeRHI, LockMode);
 	}
 
-	virtual void UnlockVertexBuffer_RenderThread(class FRHICommandListImmediate& RHICmdList, FRHIVertexBuffer* VertexBuffer) override
+	virtual void UnlockVertexBuffer_RenderThread(class FRHICommandListImmediate& RHICmdList, FRHIVertexBuffer* VertexBuffer) override final
 	{
 		return RHI->UnlockVertexBuffer_RenderThread(RHICmdList, VertexBuffer);
 	}
 
-	virtual FTexture2DRHIRef AsyncReallocateTexture2D_RenderThread(class FRHICommandListImmediate& RHICmdList, FRHITexture2D* Texture2D, int32 NewMipCount, int32 NewSizeX, int32 NewSizeY, FThreadSafeCounter* RequestStatus) override
+	virtual FTexture2DRHIRef AsyncReallocateTexture2D_RenderThread(class FRHICommandListImmediate& RHICmdList, FRHITexture2D* Texture2D, int32 NewMipCount, int32 NewSizeX, int32 NewSizeY, FThreadSafeCounter* RequestStatus) override final
 	{
 		return RHI->AsyncReallocateTexture2D_RenderThread(RHICmdList, Texture2D, NewMipCount, NewSizeX, NewSizeY, RequestStatus);
 	}
 
-	virtual ETextureReallocationStatus FinalizeAsyncReallocateTexture2D_RenderThread(class FRHICommandListImmediate& RHICmdList, FRHITexture2D* Texture2D, bool bBlockUntilCompleted) override
+	virtual ETextureReallocationStatus FinalizeAsyncReallocateTexture2D_RenderThread(class FRHICommandListImmediate& RHICmdList, FRHITexture2D* Texture2D, bool bBlockUntilCompleted) override final
 	{
 		return RHI->FinalizeAsyncReallocateTexture2D_RenderThread(RHICmdList, Texture2D, bBlockUntilCompleted);
 	}
 
-	virtual ETextureReallocationStatus CancelAsyncReallocateTexture2D_RenderThread(class FRHICommandListImmediate& RHICmdList, FRHITexture2D* Texture2D, bool bBlockUntilCompleted) override
+	virtual ETextureReallocationStatus CancelAsyncReallocateTexture2D_RenderThread(class FRHICommandListImmediate& RHICmdList, FRHITexture2D* Texture2D, bool bBlockUntilCompleted) override final
 	{
 		return RHI->CancelAsyncReallocateTexture2D_RenderThread(RHICmdList, Texture2D, bBlockUntilCompleted);
 	}
 
-	virtual FIndexBufferRHIRef CreateIndexBuffer_RenderThread(class FRHICommandListImmediate& RHICmdList, uint32 Stride, uint32 Size, uint32 InUsage, FRHIResourceCreateInfo& CreateInfo) override
+	virtual FIndexBufferRHIRef CreateIndexBuffer_RenderThread(class FRHICommandListImmediate& RHICmdList, uint32 Stride, uint32 Size, uint32 InUsage, FRHIResourceCreateInfo& CreateInfo) override final
 	{
 		return RHI->CreateIndexBuffer_RenderThread(RHICmdList, Stride, Size, InUsage, CreateInfo);
 	}
 
-	virtual void* LockIndexBuffer_RenderThread(class FRHICommandListImmediate& RHICmdList, FRHIIndexBuffer* IndexBuffer, uint32 Offset, uint32 SizeRHI, EResourceLockMode LockMode) override
+	virtual void* LockIndexBuffer_RenderThread(class FRHICommandListImmediate& RHICmdList, FRHIIndexBuffer* IndexBuffer, uint32 Offset, uint32 SizeRHI, EResourceLockMode LockMode) override final
 	{
 		return RHI->LockIndexBuffer_RenderThread(RHICmdList, IndexBuffer, Offset, SizeRHI, LockMode);
 	}
 
-	virtual void UnlockIndexBuffer_RenderThread(class FRHICommandListImmediate& RHICmdList, FRHIIndexBuffer* IndexBuffer) override
+	virtual void UnlockIndexBuffer_RenderThread(class FRHICommandListImmediate& RHICmdList, FRHIIndexBuffer* IndexBuffer) override final
 	{
 		RHI->UnlockIndexBuffer_RenderThread(RHICmdList, IndexBuffer);
 	}
 
-	virtual FVertexShaderRHIRef CreateVertexShader_RenderThread(class FRHICommandListImmediate& RHICmdList, const TArray<uint8>& Code) override
+	virtual FVertexShaderRHIRef CreateVertexShader_RenderThread(class FRHICommandListImmediate& RHICmdList, const TArray<uint8>& Code) override final
 	{
 		return RHI->CreateVertexShader_RenderThread(RHICmdList, Code);
 	}
 
-	virtual FVertexShaderRHIRef CreateVertexShader_RenderThread(class FRHICommandListImmediate& RHICmdList, FRHIShaderLibrary* Library, FSHAHash Hash) override
+	virtual FVertexShaderRHIRef CreateVertexShader_RenderThread(class FRHICommandListImmediate& RHICmdList, FRHIShaderLibrary* Library, FSHAHash Hash) override final
 	{
 		return RHI->CreateVertexShader_RenderThread(RHICmdList, Library, Hash);
 	}
 
-	virtual FPixelShaderRHIRef CreatePixelShader_RenderThread(class FRHICommandListImmediate& RHICmdList, const TArray<uint8>& Code) override
+	virtual FPixelShaderRHIRef CreatePixelShader_RenderThread(class FRHICommandListImmediate& RHICmdList, const TArray<uint8>& Code) override final
 	{
 		return RHI->CreatePixelShader_RenderThread(RHICmdList, Code);
 	}
 
-	virtual FPixelShaderRHIRef CreatePixelShader_RenderThread(class FRHICommandListImmediate& RHICmdList, FRHIShaderLibrary* Library, FSHAHash Hash) override
+	virtual FPixelShaderRHIRef CreatePixelShader_RenderThread(class FRHICommandListImmediate& RHICmdList, FRHIShaderLibrary* Library, FSHAHash Hash) override final
 	{
 		return RHI->CreatePixelShader_RenderThread(RHICmdList, Library, Hash);
 	}
 
-	virtual FGeometryShaderRHIRef CreateGeometryShader_RenderThread(class FRHICommandListImmediate& RHICmdList, const TArray<uint8>& Code) override
+	virtual FGeometryShaderRHIRef CreateGeometryShader_RenderThread(class FRHICommandListImmediate& RHICmdList, const TArray<uint8>& Code) override final
 	{
 		return RHI->CreateGeometryShader_RenderThread(RHICmdList, Code);
 	}
 
-	virtual FGeometryShaderRHIRef CreateGeometryShader_RenderThread(class FRHICommandListImmediate& RHICmdList, FRHIShaderLibrary* Library, FSHAHash Hash) override
+	virtual FGeometryShaderRHIRef CreateGeometryShader_RenderThread(class FRHICommandListImmediate& RHICmdList, FRHIShaderLibrary* Library, FSHAHash Hash) override final
 	{
 		return RHI->CreateGeometryShader_RenderThread(RHICmdList, Library, Hash);
 	}
 
-	virtual FGeometryShaderRHIRef CreateGeometryShaderWithStreamOutput_RenderThread(class FRHICommandListImmediate& RHICmdList, const TArray<uint8>& Code, const FStreamOutElementList& ElementList, uint32 NumStrides, const uint32* Strides, int32 RasterizedStream) override
+	virtual FGeometryShaderRHIRef CreateGeometryShaderWithStreamOutput_RenderThread(class FRHICommandListImmediate& RHICmdList, const TArray<uint8>& Code, const FStreamOutElementList& ElementList, uint32 NumStrides, const uint32* Strides, int32 RasterizedStream) override final
 	{
 		return RHI->CreateGeometryShaderWithStreamOutput_RenderThread(RHICmdList, Code, ElementList, NumStrides, Strides, RasterizedStream);
 	}
 
-	virtual FGeometryShaderRHIRef CreateGeometryShaderWithStreamOutput_RenderThread(class FRHICommandListImmediate& RHICmdList, const FStreamOutElementList& ElementList, uint32 NumStrides, const uint32* Strides, int32 RasterizedStream, FRHIShaderLibrary* Library, FSHAHash Hash) override
+	virtual FGeometryShaderRHIRef CreateGeometryShaderWithStreamOutput_RenderThread(class FRHICommandListImmediate& RHICmdList, const FStreamOutElementList& ElementList, uint32 NumStrides, const uint32* Strides, int32 RasterizedStream, FRHIShaderLibrary* Library, FSHAHash Hash) override final
 	{
 		return RHI->CreateGeometryShaderWithStreamOutput_RenderThread(RHICmdList, ElementList, NumStrides, Strides, RasterizedStream, Library, Hash);
 	}
 
-	virtual FComputeShaderRHIRef CreateComputeShader_RenderThread(class FRHICommandListImmediate& RHICmdList, const TArray<uint8>& Code) override
+	virtual FComputeShaderRHIRef CreateComputeShader_RenderThread(class FRHICommandListImmediate& RHICmdList, const TArray<uint8>& Code) override final
 	{
 		return RHI->CreateComputeShader_RenderThread(RHICmdList, Code);
 	}
 
-	virtual FComputeShaderRHIRef CreateComputeShader_RenderThread(class FRHICommandListImmediate& RHICmdList, FRHIShaderLibrary* Library, FSHAHash Hash) override
+	virtual FComputeShaderRHIRef CreateComputeShader_RenderThread(class FRHICommandListImmediate& RHICmdList, FRHIShaderLibrary* Library, FSHAHash Hash) override final
 	{
 		return RHI->CreateComputeShader_RenderThread(RHICmdList, Library, Hash);
 	}
 
-	virtual FHullShaderRHIRef CreateHullShader_RenderThread(class FRHICommandListImmediate& RHICmdList, const TArray<uint8>& Code) override
+	virtual FHullShaderRHIRef CreateHullShader_RenderThread(class FRHICommandListImmediate& RHICmdList, const TArray<uint8>& Code) override final
 	{
 		return RHI->CreateHullShader_RenderThread(RHICmdList, Code);
 	}
 
-	virtual FHullShaderRHIRef CreateHullShader_RenderThread(class FRHICommandListImmediate& RHICmdList, FRHIShaderLibrary* Library, FSHAHash Hash) override
+	virtual FHullShaderRHIRef CreateHullShader_RenderThread(class FRHICommandListImmediate& RHICmdList, FRHIShaderLibrary* Library, FSHAHash Hash) override final
 	{
 		return RHI->CreateHullShader_RenderThread(RHICmdList, Library, Hash);
 	}
 
-	virtual FDomainShaderRHIRef CreateDomainShader_RenderThread(class FRHICommandListImmediate& RHICmdList, const TArray<uint8>& Code) override
+	virtual FDomainShaderRHIRef CreateDomainShader_RenderThread(class FRHICommandListImmediate& RHICmdList, const TArray<uint8>& Code) override final
 	{
 		return RHI->CreateDomainShader_RenderThread(RHICmdList, Code);
 	}
 
-	virtual FDomainShaderRHIRef CreateDomainShader_RenderThread(class FRHICommandListImmediate& RHICmdList, FRHIShaderLibrary* Library, FSHAHash Hash) override
+	virtual FDomainShaderRHIRef CreateDomainShader_RenderThread(class FRHICommandListImmediate& RHICmdList, FRHIShaderLibrary* Library, FSHAHash Hash) override final
 	{
 		return RHI->CreateDomainShader_RenderThread(RHICmdList, Library, Hash);
 	}
 
-	virtual void* LockTexture2D_RenderThread(class FRHICommandListImmediate& RHICmdList, FRHITexture2D* Texture, uint32 MipIndex, EResourceLockMode LockMode, uint32& DestStride, bool bLockWithinMiptail, bool bNeedsDefaultRHIFlush = true) override
+	virtual void* LockTexture2D_RenderThread(class FRHICommandListImmediate& RHICmdList, FRHITexture2D* Texture, uint32 MipIndex, EResourceLockMode LockMode, uint32& DestStride, bool bLockWithinMiptail, bool bNeedsDefaultRHIFlush = true) override final
 	{
 		return RHI->LockTexture2D_RenderThread(RHICmdList, Texture, MipIndex, LockMode, DestStride, bLockWithinMiptail, bNeedsDefaultRHIFlush);
 	}
 
-	virtual void UnlockTexture2D_RenderThread(class FRHICommandListImmediate& RHICmdList, FRHITexture2D* Texture, uint32 MipIndex, bool bLockWithinMiptail, bool bNeedsDefaultRHIFlush = true) override
+	virtual void UnlockTexture2D_RenderThread(class FRHICommandListImmediate& RHICmdList, FRHITexture2D* Texture, uint32 MipIndex, bool bLockWithinMiptail, bool bNeedsDefaultRHIFlush = true) override final
 	{
 		RHI->UnlockTexture2D_RenderThread(RHICmdList, Texture, MipIndex, bLockWithinMiptail, bNeedsDefaultRHIFlush);
 	}
 
-	virtual void UpdateTexture2D_RenderThread(class FRHICommandListImmediate& RHICmdList, FRHITexture2D* Texture, uint32 MipIndex, const struct FUpdateTextureRegion2D& UpdateRegion, uint32 SourcePitch, const uint8* SourceData) override
+	virtual void UpdateTexture2D_RenderThread(class FRHICommandListImmediate& RHICmdList, FRHITexture2D* Texture, uint32 MipIndex, const struct FUpdateTextureRegion2D& UpdateRegion, uint32 SourcePitch, const uint8* SourceData) override final
 	{
 		RHI->UpdateTexture2D_RenderThread(RHICmdList, Texture, MipIndex, UpdateRegion, SourcePitch, SourceData);
 	}
 
-	virtual FUpdateTexture3DData BeginUpdateTexture3D_RenderThread(class FRHICommandListImmediate& RHICmdList, FRHITexture3D* Texture, uint32 MipIndex, const struct FUpdateTextureRegion3D& UpdateRegion) override
+	virtual FUpdateTexture3DData BeginUpdateTexture3D_RenderThread(class FRHICommandListImmediate& RHICmdList, FRHITexture3D* Texture, uint32 MipIndex, const struct FUpdateTextureRegion3D& UpdateRegion) override final
 	{
 		return RHI->BeginUpdateTexture3D_RenderThread(RHICmdList, Texture, MipIndex, UpdateRegion);
 	}
 
-	virtual void EndUpdateTexture3D_RenderThread(class FRHICommandListImmediate& RHICmdList, FUpdateTexture3DData& UpdateData) override
+	virtual void EndUpdateTexture3D_RenderThread(class FRHICommandListImmediate& RHICmdList, FUpdateTexture3DData& UpdateData) override final
 	{
 		RHI->EndUpdateTexture3D_RenderThread(RHICmdList, UpdateData);
 	}
 
-	virtual void EndMultiUpdateTexture3D_RenderThread(class FRHICommandListImmediate& RHICmdList, TArray<FUpdateTexture3DData>& UpdateDataArray) override
+	virtual void EndMultiUpdateTexture3D_RenderThread(class FRHICommandListImmediate& RHICmdList, TArray<FUpdateTexture3DData>& UpdateDataArray) override final
 	{
 		RHI->EndMultiUpdateTexture3D_RenderThread(RHICmdList, UpdateDataArray);
 	}
 
-	virtual void UpdateTexture3D_RenderThread(class FRHICommandListImmediate& RHICmdList, FRHITexture3D* Texture, uint32 MipIndex, const struct FUpdateTextureRegion3D& UpdateRegion, uint32 SourceRowPitch, uint32 SourceDepthPitch, const uint8* SourceData) override
+	virtual void UpdateTexture3D_RenderThread(class FRHICommandListImmediate& RHICmdList, FRHITexture3D* Texture, uint32 MipIndex, const struct FUpdateTextureRegion3D& UpdateRegion, uint32 SourceRowPitch, uint32 SourceDepthPitch, const uint8* SourceData) override final
 	{
 		RHI->UpdateTexture3D_RenderThread(RHICmdList, Texture, MipIndex, UpdateRegion, SourceRowPitch, SourceDepthPitch, SourceData);
 	}
 
-	virtual FRHIShaderLibraryRef RHICreateShaderLibrary_RenderThread(class FRHICommandListImmediate& RHICmdList, EShaderPlatform Platform, FString FilePath, FString Name) override
+	virtual FRHIShaderLibraryRef RHICreateShaderLibrary_RenderThread(class FRHICommandListImmediate& RHICmdList, EShaderPlatform Platform, FString FilePath, FString Name) override final
 	{
 		return RHI->RHICreateShaderLibrary_RenderThread(RHICmdList, Platform, FilePath, Name);
 	}
 
-	virtual FTextureReferenceRHIRef RHICreateTextureReference_RenderThread(class FRHICommandListImmediate& RHICmdList, FLastRenderTimeContainer* LastRenderTime) override
+	virtual FTextureReferenceRHIRef RHICreateTextureReference_RenderThread(class FRHICommandListImmediate& RHICmdList, FLastRenderTimeContainer* LastRenderTime) override final
 	{
 		return RHI->RHICreateTextureReference_RenderThread(RHICmdList, LastRenderTime);
 	}
 
-	virtual FTexture2DRHIRef RHICreateTexture2D_RenderThread(class FRHICommandListImmediate& RHICmdList, uint32 SizeX, uint32 SizeY, uint8 Format, uint32 NumMips, uint32 NumSamples, uint32 Flags, FRHIResourceCreateInfo& CreateInfo) override
+	virtual FTexture2DRHIRef RHICreateTexture2D_RenderThread(class FRHICommandListImmediate& RHICmdList, uint32 SizeX, uint32 SizeY, uint8 Format, uint32 NumMips, uint32 NumSamples, uint32 Flags, FRHIResourceCreateInfo& CreateInfo) override final
 	{
 		return RHI->RHICreateTexture2D_RenderThread(RHICmdList, SizeX, SizeY, Format, NumMips, NumSamples, Flags, CreateInfo);
 	}
 
-	virtual FTexture2DRHIRef RHICreateTextureExternal2D_RenderThread(class FRHICommandListImmediate& RHICmdList, uint32 SizeX, uint32 SizeY, uint8 Format, uint32 NumMips, uint32 NumSamples, uint32 Flags, FRHIResourceCreateInfo& CreateInfo) override
+	virtual FTexture2DRHIRef RHICreateTextureExternal2D_RenderThread(class FRHICommandListImmediate& RHICmdList, uint32 SizeX, uint32 SizeY, uint8 Format, uint32 NumMips, uint32 NumSamples, uint32 Flags, FRHIResourceCreateInfo& CreateInfo) override final
 	{
 		return RHI->RHICreateTextureExternal2D_RenderThread(RHICmdList, SizeX, SizeY, Format, NumMips, NumSamples, Flags, CreateInfo);
 	}
 
-	virtual FTexture2DArrayRHIRef RHICreateTexture2DArray_RenderThread(class FRHICommandListImmediate& RHICmdList, uint32 SizeX, uint32 SizeY, uint32 SizeZ, uint8 Format, uint32 NumMips, uint32 Flags, FRHIResourceCreateInfo& CreateInfo) override
+	virtual FTexture2DArrayRHIRef RHICreateTexture2DArray_RenderThread(class FRHICommandListImmediate& RHICmdList, uint32 SizeX, uint32 SizeY, uint32 SizeZ, uint8 Format, uint32 NumMips, uint32 Flags, FRHIResourceCreateInfo& CreateInfo) override final
 	{
 		return RHI->RHICreateTexture2DArray_RenderThread(RHICmdList, SizeX, SizeY, SizeZ, Format, NumMips, Flags, CreateInfo);
 	}
 
-	virtual FTexture3DRHIRef RHICreateTexture3D_RenderThread(class FRHICommandListImmediate& RHICmdList, uint32 SizeX, uint32 SizeY, uint32 SizeZ, uint8 Format, uint32 NumMips, uint32 Flags, FRHIResourceCreateInfo& CreateInfo) override
+	virtual FTexture3DRHIRef RHICreateTexture3D_RenderThread(class FRHICommandListImmediate& RHICmdList, uint32 SizeX, uint32 SizeY, uint32 SizeZ, uint8 Format, uint32 NumMips, uint32 Flags, FRHIResourceCreateInfo& CreateInfo) override final
 	{
 		return RHI->RHICreateTexture3D_RenderThread(RHICmdList, SizeX, SizeY, SizeZ, Format, NumMips, Flags, CreateInfo);
 	}
 
-	virtual FUnorderedAccessViewRHIRef RHICreateUnorderedAccessView_RenderThread(class FRHICommandListImmediate& RHICmdList, FRHIStructuredBuffer* StructuredBuffer, bool bUseUAVCounter, bool bAppendBuffer) override
+	virtual FUnorderedAccessViewRHIRef RHICreateUnorderedAccessView_RenderThread(class FRHICommandListImmediate& RHICmdList, FRHIStructuredBuffer* StructuredBuffer, bool bUseUAVCounter, bool bAppendBuffer) override final
 	{
 		return RHI->RHICreateUnorderedAccessView_RenderThread(RHICmdList, StructuredBuffer, bUseUAVCounter, bAppendBuffer);
 	}
 
-	virtual FUnorderedAccessViewRHIRef RHICreateUnorderedAccessView_RenderThread(class FRHICommandListImmediate& RHICmdList, FRHITexture* Texture, uint32 MipLevel) override
+	virtual FUnorderedAccessViewRHIRef RHICreateUnorderedAccessView_RenderThread(class FRHICommandListImmediate& RHICmdList, FRHITexture* Texture, uint32 MipLevel) override final
 	{
 		return RHI->RHICreateUnorderedAccessView_RenderThread(RHICmdList, Texture, MipLevel);
 	}
 
-	virtual FUnorderedAccessViewRHIRef RHICreateUnorderedAccessView_RenderThread(class FRHICommandListImmediate& RHICmdList, FRHIVertexBuffer* VertexBuffer, uint8 Format) override
+	virtual FUnorderedAccessViewRHIRef RHICreateUnorderedAccessView_RenderThread(class FRHICommandListImmediate& RHICmdList, FRHIVertexBuffer* VertexBuffer, uint8 Format) override final
 	{
 		return RHI->RHICreateUnorderedAccessView_RenderThread(RHICmdList, VertexBuffer, Format);
 	}
 
-	virtual FUnorderedAccessViewRHIRef RHICreateUnorderedAccessView_RenderThread(class FRHICommandListImmediate& RHICmdList, FRHIIndexBuffer* IndexBuffer, uint8 Format) override
+	virtual FUnorderedAccessViewRHIRef RHICreateUnorderedAccessView_RenderThread(class FRHICommandListImmediate& RHICmdList, FRHIIndexBuffer* IndexBuffer, uint8 Format) override final
 	{
 		return RHI->RHICreateUnorderedAccessView_RenderThread(RHICmdList, IndexBuffer, Format);
 	}
 
-	virtual FShaderResourceViewRHIRef RHICreateShaderResourceView_RenderThread(class FRHICommandListImmediate& RHICmdList, FRHITexture2D* Texture2DRHI, uint8 MipLevel) override
+	virtual FShaderResourceViewRHIRef RHICreateShaderResourceView_RenderThread(class FRHICommandListImmediate& RHICmdList, FRHITexture2D* Texture2DRHI, uint8 MipLevel) override final
 	{
 		return RHI->RHICreateShaderResourceView_RenderThread(RHICmdList, Texture2DRHI, MipLevel);
 	}
 
-	virtual FShaderResourceViewRHIRef RHICreateShaderResourceView_RenderThread(class FRHICommandListImmediate& RHICmdList, FRHITexture2D* Texture2DRHI, uint8 MipLevel, uint8 NumMipLevels, uint8 Format) override
+	virtual FShaderResourceViewRHIRef RHICreateShaderResourceView_RenderThread(class FRHICommandListImmediate& RHICmdList, FRHITexture2D* Texture2DRHI, uint8 MipLevel, uint8 NumMipLevels, uint8 Format) override final
 	{
 		return RHI->RHICreateShaderResourceView_RenderThread(RHICmdList, Texture2DRHI, MipLevel, NumMipLevels, Format);
 	}
 
-	virtual FShaderResourceViewRHIRef RHICreateShaderResourceView_RenderThread(class FRHICommandListImmediate& RHICmdList, FRHITexture3D* Texture3DRHI, uint8 MipLevel) override
+	virtual FShaderResourceViewRHIRef RHICreateShaderResourceView_RenderThread(class FRHICommandListImmediate& RHICmdList, FRHITexture3D* Texture3DRHI, uint8 MipLevel) override final
 	{
 		return RHI->RHICreateShaderResourceView_RenderThread(RHICmdList, Texture3DRHI, MipLevel);
 	}
 
-	virtual FShaderResourceViewRHIRef RHICreateShaderResourceView_RenderThread(class FRHICommandListImmediate& RHICmdList, FRHITexture2DArray* Texture2DArrayRHI, uint8 MipLevel) override
+	virtual FShaderResourceViewRHIRef RHICreateShaderResourceView_RenderThread(class FRHICommandListImmediate& RHICmdList, FRHITexture2DArray* Texture2DArrayRHI, uint8 MipLevel) override final
 	{
 		return RHI->RHICreateShaderResourceView_RenderThread(RHICmdList, Texture2DArrayRHI, MipLevel);
 	}
 
-	virtual FShaderResourceViewRHIRef RHICreateShaderResourceView_RenderThread(class FRHICommandListImmediate& RHICmdList, FRHITextureCube* TextureCubeRHI, uint8 MipLevel) override
+	virtual FShaderResourceViewRHIRef RHICreateShaderResourceView_RenderThread(class FRHICommandListImmediate& RHICmdList, FRHITextureCube* TextureCubeRHI, uint8 MipLevel) override final
 	{
 		return RHI->RHICreateShaderResourceView_RenderThread(RHICmdList, TextureCubeRHI, MipLevel);
 	}
 
-	virtual FShaderResourceViewRHIRef RHICreateShaderResourceView_RenderThread(class FRHICommandListImmediate& RHICmdList, FRHIVertexBuffer* VertexBuffer, uint32 Stride, uint8 Format) override
+	virtual FShaderResourceViewRHIRef RHICreateShaderResourceView_RenderThread(class FRHICommandListImmediate& RHICmdList, FRHIVertexBuffer* VertexBuffer, uint32 Stride, uint8 Format) override final
 	{
 		return RHI->RHICreateShaderResourceView_RenderThread(RHICmdList, VertexBuffer, Stride, Format);
 	}
 
-	virtual FShaderResourceViewRHIRef RHICreateShaderResourceView_RenderThread(class FRHICommandListImmediate& RHICmdList, FRHIIndexBuffer* Buffer) override
+	virtual FShaderResourceViewRHIRef RHICreateShaderResourceView_RenderThread(class FRHICommandListImmediate& RHICmdList, FRHIIndexBuffer* Buffer) override final
 	{
 		return RHI->RHICreateShaderResourceView_RenderThread(RHICmdList, Buffer);
 	}
 
-	virtual FShaderResourceViewRHIRef RHICreateShaderResourceView_RenderThread(class FRHICommandListImmediate& RHICmdList, FRHIStructuredBuffer* StructuredBuffer) override
+	virtual FShaderResourceViewRHIRef RHICreateShaderResourceView_RenderThread(class FRHICommandListImmediate& RHICmdList, FRHIStructuredBuffer* StructuredBuffer) override final
 	{
 		return RHI->RHICreateShaderResourceView_RenderThread(RHICmdList, StructuredBuffer);
 	}
 
-	virtual FTextureCubeRHIRef RHICreateTextureCube_RenderThread(class FRHICommandListImmediate& RHICmdList, uint32 Size, uint8 Format, uint32 NumMips, uint32 Flags, FRHIResourceCreateInfo& CreateInfo) override
+	virtual FTextureCubeRHIRef RHICreateTextureCube_RenderThread(class FRHICommandListImmediate& RHICmdList, uint32 Size, uint8 Format, uint32 NumMips, uint32 Flags, FRHIResourceCreateInfo& CreateInfo) override final
 	{
 		return RHI->RHICreateTextureCube_RenderThread(RHICmdList, Size, Format, NumMips, Flags, CreateInfo);
 	}
 
-	virtual FTextureCubeRHIRef RHICreateTextureCubeArray_RenderThread(class FRHICommandListImmediate& RHICmdList, uint32 Size, uint32 ArraySize, uint8 Format, uint32 NumMips, uint32 Flags, FRHIResourceCreateInfo& CreateInfo) override
+	virtual FTextureCubeRHIRef RHICreateTextureCubeArray_RenderThread(class FRHICommandListImmediate& RHICmdList, uint32 Size, uint32 ArraySize, uint8 Format, uint32 NumMips, uint32 Flags, FRHIResourceCreateInfo& CreateInfo) override final
 	{
 		return RHI->RHICreateTextureCubeArray_RenderThread(RHICmdList, Size, ArraySize, Format, NumMips, Flags, CreateInfo);
 	}
 	
-	virtual FRenderQueryRHIRef RHICreateRenderQuery_RenderThread(class FRHICommandListImmediate& RHICmdList, ERenderQueryType QueryType) override
+	virtual FRenderQueryRHIRef RHICreateRenderQuery_RenderThread(class FRHICommandListImmediate& RHICmdList, ERenderQueryType QueryType) override final
 	{
 		return RHI->RHICreateRenderQuery_RenderThread(RHICmdList, QueryType);
 	}
 
 
-	virtual void* RHILockTextureCubeFace_RenderThread(class FRHICommandListImmediate& RHICmdList, FRHITextureCube* Texture, uint32 FaceIndex, uint32 ArrayIndex, uint32 MipIndex, EResourceLockMode LockMode, uint32& DestStride, bool bLockWithinMiptail) override
+	virtual void* RHILockTextureCubeFace_RenderThread(class FRHICommandListImmediate& RHICmdList, FRHITextureCube* Texture, uint32 FaceIndex, uint32 ArrayIndex, uint32 MipIndex, EResourceLockMode LockMode, uint32& DestStride, bool bLockWithinMiptail) override final
 	{
 		return RHI->RHILockTextureCubeFace_RenderThread(RHICmdList, Texture, FaceIndex, ArrayIndex, MipIndex, LockMode, DestStride, bLockWithinMiptail);
 	}
-	virtual void RHIUnlockTextureCubeFace_RenderThread(class FRHICommandListImmediate& RHICmdList, FRHITextureCube* Texture, uint32 FaceIndex, uint32 ArrayIndex, uint32 MipIndex, bool bLockWithinMiptail) override
+	virtual void RHIUnlockTextureCubeFace_RenderThread(class FRHICommandListImmediate& RHICmdList, FRHITextureCube* Texture, uint32 FaceIndex, uint32 ArrayIndex, uint32 MipIndex, bool bLockWithinMiptail) override final
 	{
 		RHI->RHIUnlockTextureCubeFace_RenderThread(RHICmdList, Texture, FaceIndex, ArrayIndex, MipIndex, bLockWithinMiptail);
 	}
 
-	virtual void RHIAcquireTransientResource_RenderThread(FRHITexture* Texture) override
+	virtual void RHIAcquireTransientResource_RenderThread(FRHITexture* Texture) override final
 	{
 		RHI->RHIAcquireTransientResource_RenderThread(Texture);
 	}
 
-	virtual void RHIDiscardTransientResource_RenderThread(FRHITexture* Texture) override
+	virtual void RHIDiscardTransientResource_RenderThread(FRHITexture* Texture) override final
 	{
 		RHI->RHIDiscardTransientResource_RenderThread(Texture);
 	}
 
-	virtual void RHIAcquireTransientResource_RenderThread(FRHIVertexBuffer* Buffer) override
+	virtual void RHIAcquireTransientResource_RenderThread(FRHIVertexBuffer* Buffer) override final
 	{
 		RHI->RHIAcquireTransientResource_RenderThread(Buffer);
 	}
 
-	virtual void RHIDiscardTransientResource_RenderThread(FRHIVertexBuffer* Buffer) override
+	virtual void RHIDiscardTransientResource_RenderThread(FRHIVertexBuffer* Buffer) override final
 	{
 		RHI->RHIDiscardTransientResource_RenderThread(Buffer);
 	}
 
-	virtual void RHIAcquireTransientResource_RenderThread(FRHIStructuredBuffer* Buffer) override
+	virtual void RHIAcquireTransientResource_RenderThread(FRHIStructuredBuffer* Buffer) override final
 	{
 		RHI->RHIAcquireTransientResource_RenderThread(Buffer);
 	}
 
-	virtual void RHIDiscardTransientResource_RenderThread(FRHIStructuredBuffer* Buffer) override
+	virtual void RHIDiscardTransientResource_RenderThread(FRHIStructuredBuffer* Buffer) override final
 	{
 		RHI->RHIDiscardTransientResource_RenderThread(Buffer);
 	}
 
-	virtual void RHIReadSurfaceFloatData_RenderThread(class FRHICommandListImmediate& RHICmdList, FRHITexture* Texture, FIntRect Rect, TArray<FFloat16Color>& OutData, ECubeFace CubeFace, int32 ArrayIndex, int32 MipIndex) override
+	virtual void RHIReadSurfaceFloatData_RenderThread(class FRHICommandListImmediate& RHICmdList, FRHITexture* Texture, FIntRect Rect, TArray<FFloat16Color>& OutData, ECubeFace CubeFace, int32 ArrayIndex, int32 MipIndex) override final
 	{
 		RHI->RHIReadSurfaceFloatData_RenderThread(RHICmdList, Texture, Rect, OutData, CubeFace, ArrayIndex, MipIndex);
 	}
 
 	//Utilities
-	virtual void EnableIdealGPUCaptureOptions(bool bEnable) override
+	virtual void EnableIdealGPUCaptureOptions(bool bEnable) override final
 	{
 		RHI->EnableIdealGPUCaptureOptions(bEnable);
 	}
 
-	virtual void RHISetResourceAliasability_RenderThread(class FRHICommandListImmediate& RHICmdList, EResourceAliasability AliasMode, FRHITexture** InTextures, int32 NumTextures) override
+	virtual void RHISetResourceAliasability_RenderThread(class FRHICommandListImmediate& RHICmdList, EResourceAliasability AliasMode, FRHITexture** InTextures, int32 NumTextures) override final
 	{
 		RHI->RHISetResourceAliasability_RenderThread(RHICmdList, AliasMode, InTextures, NumTextures);
 	}
 
 	//checks if the GPU is still alive.
-	virtual bool CheckGpuHeartbeat() const override
+	virtual bool CheckGpuHeartbeat() const override final
 	{
 		return RHI->CheckGpuHeartbeat();
 	}
 
-	virtual void VirtualTextureSetFirstMipInMemory_RenderThread(class FRHICommandListImmediate& RHICmdList, FRHITexture2D* Texture, uint32 FirstMip) override
+	virtual void VirtualTextureSetFirstMipInMemory_RenderThread(class FRHICommandListImmediate& RHICmdList, FRHITexture2D* Texture, uint32 FirstMip) override final
 	{
 		RHI->VirtualTextureSetFirstMipInMemory_RenderThread(RHICmdList, Texture, FirstMip);
 	}
 
-	virtual void VirtualTextureSetFirstMipVisible_RenderThread(class FRHICommandListImmediate& RHICmdList, FRHITexture2D* Texture, uint32 FirstMip) override
+	virtual void VirtualTextureSetFirstMipVisible_RenderThread(class FRHICommandListImmediate& RHICmdList, FRHITexture2D* Texture, uint32 FirstMip) override final
 	{
 		RHI->VirtualTextureSetFirstMipVisible_RenderThread(RHICmdList, Texture, FirstMip);
 	}
 
 	/* Copy the source box pixels in the destination box texture, return true if implemented for the current platform*/
-	virtual void RHICopySubTextureRegion_RenderThread(class FRHICommandListImmediate& RHICmdList, FRHITexture2D* SourceTexture, FRHITexture2D* DestinationTexture, FBox2D SourceBox, FBox2D DestinationBox) override
+	virtual void RHICopySubTextureRegion_RenderThread(class FRHICommandListImmediate& RHICmdList, FRHITexture2D* SourceTexture, FRHITexture2D* DestinationTexture, FBox2D SourceBox, FBox2D DestinationBox) override final
 	{
 		return RHI->RHICopySubTextureRegion_RenderThread(RHICmdList, SourceTexture, DestinationTexture, SourceBox, DestinationBox);
 	}
 
-	virtual void RHICopySubTextureRegion(FRHITexture2D* SourceTexture, FRHITexture2D* DestinationTexture, FBox2D SourceBox, FBox2D DestinationBox) override
+	virtual void RHICopySubTextureRegion(FRHITexture2D* SourceTexture, FRHITexture2D* DestinationTexture, FBox2D SourceBox, FBox2D DestinationBox) override final
 	{
 		RHI->RHICopySubTextureRegion(SourceTexture, DestinationTexture, SourceBox, DestinationBox);
 	}
 
-	virtual FRHIFlipDetails RHIWaitForFlip(double TimeoutInSeconds) override
+	virtual FRHIFlipDetails RHIWaitForFlip(double TimeoutInSeconds) override final
 	{
 		return RHI->RHIWaitForFlip(TimeoutInSeconds);
 	}
 
-	virtual void RHISignalFlipEvent() override
+	virtual void RHISignalFlipEvent() override final
 	{
 		RHI->RHISignalFlipEvent();
 	}
 
-	virtual void RHICalibrateTimers() override
+	virtual void RHICalibrateTimers() override final
 	{
 		RHI->RHICalibrateTimers();
 	}
 
-	virtual void RHIPollRenderQueryResults() override
+	virtual void RHIPollRenderQueryResults() override final
 	{
 		RHI->RHIPollRenderQueryResults();
 	}
 
-	virtual bool RHIIsTypedUAVLoadSupported(EPixelFormat PixelFormat) override
+	virtual bool RHIIsTypedUAVLoadSupported(EPixelFormat PixelFormat) override final
 	{
 		return RHI->RHIIsTypedUAVLoadSupported(PixelFormat);
 	}
 
-	virtual uint16 RHIGetPlatformTextureMaxSampleCount() override
+	virtual uint16 RHIGetPlatformTextureMaxSampleCount() override final
 	{
 		return RHI->RHIGetPlatformTextureMaxSampleCount();
 	};
 
 
 #if RHI_RAYTRACING
-	virtual FRayTracingGeometryRHIRef RHICreateRayTracingGeometry(const FRayTracingGeometryInitializer& Initializer) override
+	virtual FRayTracingGeometryRHIRef RHICreateRayTracingGeometry(const FRayTracingGeometryInitializer& Initializer) override final
 	{
 		return RHI->RHICreateRayTracingGeometry(Initializer);
 	}
 
-	virtual FRayTracingSceneRHIRef RHICreateRayTracingScene(const FRayTracingSceneInitializer& Initializer) override
+	virtual FRayTracingSceneRHIRef RHICreateRayTracingScene(const FRayTracingSceneInitializer& Initializer) override final
 	{
 		return RHI->RHICreateRayTracingScene(Initializer);
 	}
 
-	virtual FRayTracingShaderRHIRef RHICreateRayTracingShader(const TArray<uint8>& Code, EShaderFrequency ShaderFrequency) override
+	virtual FRayTracingShaderRHIRef RHICreateRayTracingShader(const TArray<uint8>& Code, EShaderFrequency ShaderFrequency) override final
 	{
 		return RHI->RHICreateRayTracingShader(Code, ShaderFrequency);
 	}
 
-	virtual FRayTracingPipelineStateRHIRef RHICreateRayTracingPipelineState(const FRayTracingPipelineStateInitializer& Initializer) override
+	virtual FRayTracingPipelineStateRHIRef RHICreateRayTracingPipelineState(const FRayTracingPipelineStateInitializer& Initializer) override final
 	{
 		return RHI->RHICreateRayTracingPipelineState(Initializer);
 	}
@@ -1641,9 +1639,28 @@ public:
 
 
 //protected:
-	FDynamicRHI*		RHI;
-	FValidationContext* Context;
+	FDynamicRHI*				RHI;
+	FValidationContext*			Context;
+	FValidationComputeContext*	AsyncComputeContext;
 
 private:
-	FString RHIName;
+	FString						RHIName;
 };
+
+extern RHI_API FValidationRHI* GValidationRHI;
+
+template <typename TDynamicRHI>
+inline TDynamicRHI* GetDynamicRHI()
+{
+	return GValidationRHI ? static_cast<TDynamicRHI*>GValidationRHI->RHI : static_cast<TDynamicRHI*>GDynamicRHI;
+}
+
+#else
+
+template <typename TDynamicRHI>
+inline TDynamicRHI* GetDynamicRHI()
+{
+	return static_cast<TDynamicRHI*>GDynamicRHI;
+}
+
+#endif	// ENABLE_RHI_VALIDATION
diff --git a/Engine/Source/Runtime/RHI/Public/ValidationRHICommon.h b/Engine/Source/Runtime/RHI/Public/ValidationRHICommon.h
new file mode 100644
--- /dev/null
+++ b/Engine/Source/Runtime/RHI/Public/ValidationRHICommon.h
@@ -0,0 +1,16 @@
+// Copyright 1998-2019 Epic Games, Inc. All Rights Reserved.
+
+/*=============================================================================
+	ValidationRHICommon.h: Public Valdation RHI definitions.
+=============================================================================*/
+
+#pragma once 
+
+#ifndef ENABLE_RHI_VALIDATION
+	#define ENABLE_RHI_VALIDATION	(UE_BUILD_DEBUG || UE_BUILD_DEVELOPMENT)
+#endif
+
+#include "RHI.h"
+#include "RHIResources.h"
+#include "RHIContext.h"
+#include "DynamicRHI.h"
diff --git a/Engine/Source/Runtime/VulkanRHI/Private/VulkanRHI.cpp b/Engine/Source/Runtime/VulkanRHI/Private/VulkanRHI.cpp
--- a/Engine/Source/Runtime/VulkanRHI/Private/VulkanRHI.cpp
+++ b/Engine/Source/Runtime/VulkanRHI/Private/VulkanRHI.cpp
@@ -21,8 +21,6 @@
 #include "ValidationRHI.h"
 
 
-FValidationRHI*		GValidationRHI = nullptr;
-
 extern RHI_API bool GUseTexture3DBulkDataRHI;
 
 TAtomic<uint64> GVulkanBufferHandleIdCounter{ 0 };
@@ -96,7 +94,7 @@ FDynamicRHI* FVulkanDynamicRHIModule::CreateRHI(ERHIFeatureLevel::Type InRequest
 	check(RHISupportsSeparateMSAAAndResolveTextures(GMaxRHIShaderPlatform) == (!VULKAN_USE_MSAA_RESOLVE_ATTACHMENTS));
 
 	GVulkanRHI = new FVulkanDynamicRHI();
-
+#if ENABLE_RHI_VALIDATION
 	if (FParse::Param(FCommandLine::Get(), TEXT("RHIValidation")))
 	{
 		GValidationRHI = new FValidationRHI(GVulkanRHI);
@@ -107,6 +105,9 @@ FDynamicRHI* FVulkanDynamicRHIModule::CreateRHI(ERHIFeatureLevel::Type InRequest
 	}
 
 	return GValidationRHI ? (FDynamicRHI*)GValidationRHI : (FDynamicRHI*)GVulkanRHI;
+#else
+	return GVulkanRHI;
+#endif
 }
 
 IMPLEMENT_MODULE(FVulkanDynamicRHIModule, VulkanRHI);
@@ -698,8 +699,6 @@ void FVulkanDynamicRHI::InitInstance()
 
 		GUseTexture3DBulkDataRHI = true;
 
-		GDynamicRHI = GValidationRHI ? (FDynamicRHI*)GValidationRHI : (FDynamicRHI*)this;
-
 		// Notify all initialized FRenderResources that there's a valid RHI device to create their RHI resources for now.
 		for (TLinkedList<FRenderResource*>::TIterator ResourceIt(FRenderResource::GetResourceList()); ResourceIt; ResourceIt.Next())
 		{
diff --git a/Engine/Source/Runtime/VulkanRHI/Private/VulkanRHIPrivate.h b/Engine/Source/Runtime/VulkanRHI/Private/VulkanRHIPrivate.h
--- a/Engine/Source/Runtime/VulkanRHI/Private/VulkanRHIPrivate.h
+++ b/Engine/Source/Runtime/VulkanRHI/Private/VulkanRHIPrivate.h
@@ -11,6 +11,7 @@
 #include "Misc/ScopeLock.h"
 #include "RHI.h"
 #include "RenderUtils.h"
+#include "ValidationRHI.h"
 
 // let the platform set up the headers and some defines
 #include "VulkanPlatform.h"
@@ -291,8 +292,6 @@ private:
 	FVulkanDevice&				Device;
 };
 
-extern class FValidationRHI*	GValidationRHI;
-
 
 namespace VulkanRHI
 {
@@ -840,11 +839,12 @@ namespace VulkanRHI
 
 	inline FVulkanCommandListContext& GetVulkanContext(IRHICommandContext& CmdContext)
 	{
+#if ENABLE_RHI_VALIDATION
 		if (GValidationRHI)
 		{
 			return GetVulkanContext((FValidationContext&)CmdContext);
 		}
-
+#endif
 		return (FVulkanCommandListContext&)CmdContext;
 	}
 }
diff --git a/Engine/Source/Runtime/VulkanRHI/Private/VulkanUtil.cpp b/Engine/Source/Runtime/VulkanRHI/Private/VulkanUtil.cpp
--- a/Engine/Source/Runtime/VulkanRHI/Private/VulkanUtil.cpp
+++ b/Engine/Source/Runtime/VulkanRHI/Private/VulkanUtil.cpp
@@ -592,6 +592,7 @@ namespace VulkanRHIBridge
 
 namespace VulkanRHI
 {
+#if ENABLE_RHI_VALIDATION
 	FVulkanCommandListContext& GetVulkanContext(class FValidationContext& CmdContext)
 	{
 		if (GValidationRHI && GValidationRHI->Context == &CmdContext)
@@ -601,7 +602,7 @@ namespace VulkanRHI
 
 		return (FVulkanCommandListContext&)CmdContext;
 	}
-
+#endif
 	VkBuffer CreateBuffer(FVulkanDevice* InDevice, VkDeviceSize Size, VkBufferUsageFlags BufferUsageFlags, VkMemoryRequirements& OutMemoryRequirements)
 	{
 		VkDevice Device = InDevice->GetInstanceHandle();
diff --git a/Engine/Source/Runtime/Windows/D3D11RHI/Private/D3D11Query.cpp b/Engine/Source/Runtime/Windows/D3D11RHI/Private/D3D11Query.cpp
--- a/Engine/Source/Runtime/Windows/D3D11RHI/Private/D3D11Query.cpp
+++ b/Engine/Source/Runtime/Windows/D3D11RHI/Private/D3D11Query.cpp
@@ -55,8 +55,6 @@ public:
 	void PollQueryResults()
 	{
 		QUICK_SCOPE_CYCLE_COUNTER(STAT_PollQueryResults);
-		FD3D11DynamicRHI* D3D11RHI = static_cast<FD3D11DynamicRHI*>(GDynamicRHI);
-
 		for (int32 BatchIdx = 0; BatchIdx < ActiveBatches.Num(); ++BatchIdx)
 		{
 			TArray<FRenderQueryRHIRef>& Queries = ActiveBatches[BatchIdx].Queries;
@@ -69,7 +67,7 @@ public:
 				{
 					Queries.RemoveAtSwap(QueryIdx--);
 				}
-				else if (D3D11RHI->GetQueryData(Query->Resource, &Query->Result, sizeof(Query->Result), Query->QueryType, false, false))
+				else if (GD3D11RHI->GetQueryData(Query->Resource, &Query->Result, sizeof(Query->Result), Query->QueryType, false, false))
 				{
 					Query->bResultIsCached = true;
 					Queries.RemoveAtSwap(QueryIdx--);
diff --git a/Engine/Source/Runtime/Windows/D3D11RHI/Private/D3D11RHIPrivate.h b/Engine/Source/Runtime/Windows/D3D11RHI/Private/D3D11RHIPrivate.h
--- a/Engine/Source/Runtime/Windows/D3D11RHI/Private/D3D11RHIPrivate.h
+++ b/Engine/Source/Runtime/Windows/D3D11RHI/Private/D3D11RHIPrivate.h
@@ -28,6 +28,7 @@ DECLARE_LOG_CATEGORY_EXTERN(LogD3D11RHI, Log, All);
 #include "D3D11Viewport.h"
 #include "D3D11ConstantBuffer.h"
 #include "D3D11StateCache.h"
+#include "ValidationRHICommon.h"
 
 #ifndef WITH_DX_PERF
 #define WITH_DX_PERF	1
@@ -1556,3 +1557,5 @@ private:
 	// ordered from small to large (for efficient compactening)
 	TArray<FRange> Entries;
 };
+
+extern FD3D11DynamicRHI*	GD3D11RHI;
diff --git a/Engine/Source/Runtime/Windows/D3D11RHI/Private/D3D11Texture.cpp b/Engine/Source/Runtime/Windows/D3D11RHI/Private/D3D11Texture.cpp
--- a/Engine/Source/Runtime/Windows/D3D11RHI/Private/D3D11Texture.cpp
+++ b/Engine/Source/Runtime/Windows/D3D11RHI/Private/D3D11Texture.cpp
@@ -1930,7 +1930,7 @@ void FD3D11DynamicRHI::EndUpdateTexture3D_RenderThread(class FRHICommandListImme
 		RunOnRHIThread(
 			[UpdateData]()
 		{
-			static_cast<FD3D11DynamicRHI*>(GDynamicRHI)->RHIUpdateTexture3D(
+			GD3D11RHI->RHIUpdateTexture3D(
 				UpdateData.Texture,
 				UpdateData.MipIndex,
 				UpdateData.UpdateRegion,
diff --git a/Engine/Source/Runtime/Windows/D3D11RHI/Private/Windows/WindowsD3D11Device.cpp b/Engine/Source/Runtime/Windows/D3D11RHI/Private/Windows/WindowsD3D11Device.cpp
--- a/Engine/Source/Runtime/Windows/D3D11RHI/Private/Windows/WindowsD3D11Device.cpp
+++ b/Engine/Source/Runtime/Windows/D3D11RHI/Private/Windows/WindowsD3D11Device.cpp
@@ -19,6 +19,7 @@
 #include "GenericPlatform/GenericPlatformDriver.h"			// FGPUDriverInfo
 
 #include "dxgi1_3.h"
+#include "ValidationRHI.h"
 
 #if NV_AFTERMATH
 bool GDX11NVAfterMathEnabled = false;
@@ -29,6 +30,8 @@ bool GNVAftermathModuleLoaded = false;
 bool GDX11IntelMetricsDiscoveryEnabled = false;
 #endif
 
+FD3D11DynamicRHI*	GD3D11RHI = nullptr;
+
 extern bool D3D11RHI_ShouldCreateWithD3DDebug();
 extern bool D3D11RHI_ShouldAllowAsyncResourceCreation();
 
@@ -907,7 +910,22 @@ FDynamicRHI* FD3D11DynamicRHIModule::CreateRHI(ERHIFeatureLevel::Type RequestedF
 	TRefCountPtr<IDXGIFactory1> DXGIFactory1;
 	SafeCreateDXGIFactory(DXGIFactory1.GetInitReference());
 	check(DXGIFactory1);
-	return new FD3D11DynamicRHI(DXGIFactory1,ChosenAdapter.MaxSupportedFeatureLevel,ChosenAdapter.AdapterIndex,ChosenDescription);
+
+	GD3D11RHI = new FD3D11DynamicRHI(DXGIFactory1,ChosenAdapter.MaxSupportedFeatureLevel,ChosenAdapter.AdapterIndex,ChosenDescription);
+#if ENABLE_RHI_VALIDATION
+	if (FParse::Param(FCommandLine::Get(), TEXT("RHIValidation")))
+	{
+		GValidationRHI = new FValidationRHI(GD3D11RHI);
+	}
+	else
+	{
+		check(!GValidationRHI);
+	}
+
+	return GValidationRHI ? (FDynamicRHI*)GValidationRHI : (FDynamicRHI*)GD3D11RHI;
+#else
+	return GD3D11RHI;
+#endif
 }
 
 void FD3D11DynamicRHI::Init()
diff --git a/Engine/Source/Runtime/Windows/D3D11RHI/Public/D3D11Util.h b/Engine/Source/Runtime/Windows/D3D11RHI/Public/D3D11Util.h
--- a/Engine/Source/Runtime/Windows/D3D11RHI/Public/D3D11Util.h
+++ b/Engine/Source/Runtime/Windows/D3D11RHI/Public/D3D11Util.h
@@ -27,8 +27,8 @@
 
 #endif
 
-#define D3D11RHI_IMMEDIATE_CONTEXT (static_cast<FD3D11DynamicRHI*>(GDynamicRHI)->GetDeviceContext())
-#define D3D11RHI_DEVICE (static_cast<FD3D11DynamicRHI*>(GDynamicRHI)->GetDevice())
+#define D3D11RHI_IMMEDIATE_CONTEXT	(GD3D11RHI->GetDeviceContext())
+#define D3D11RHI_DEVICE				(GD3D11RHI->GetDevice())
 
 /**
  * Checks that the given result isn't a failure.  If it is, the application exits with an appropriate error message.
