Fixed compile error due to incorrect usage of static_cast in GetDynamicRHI()

diff --git a/Engine/Source/Runtime/RHI/Public/ValidationRHI.h b/Engine/Source/Runtime/RHI/Public/ValidationRHI.h
--- a/Engine/Source/Runtime/RHI/Public/ValidationRHI.h
+++ b/Engine/Source/Runtime/RHI/Public/ValidationRHI.h
@@ -1652,7 +1652,7 @@ extern RHI_API FValidationRHI* GValidationRHI;
 template <typename TDynamicRHI>
 inline TDynamicRHI* GetDynamicRHI()
 {
-	return GValidationRHI ? static_cast<TDynamicRHI*>GValidationRHI->RHI : static_cast<TDynamicRHI*>GDynamicRHI;
+	return GValidationRHI ? static_cast<TDynamicRHI*>(GValidationRHI->RHI) : static_cast<TDynamicRHI*>(GDynamicRHI);
 }
 
 #else
