DR - Missing new function

diff --git a/Engine/Source/Runtime/RHI/Public/ValidationContext.h b/Engine/Source/Runtime/RHI/Public/ValidationContext.h
--- a/Engine/Source/Runtime/RHI/Public/ValidationContext.h
+++ b/Engine/Source/Runtime/RHI/Public/ValidationContext.h
@@ -147,6 +147,11 @@ public:
 		RHIContext->RHIPopEvent();
 	}
 
+	virtual void RHIWriteGPUFence(FGPUFenceRHIParamRef FenceRHI) override final
+	{
+		RHIContext->RHIWriteGPUFence(FenceRHI);
+	}
+
 	IRHIComputeContext* RHIContext;
 	FValidationRHI*		RHI;
 
@@ -776,6 +781,11 @@ public:
 		State.bInsideComputePass = false;
 	}
 
+	virtual void RHIWriteGPUFence(FGPUFenceRHIParamRef FenceRHI) override final
+	{
+		RHIContext->RHIWriteGPUFence(FenceRHI);
+	}
+
 	virtual void RHICopyTexture(FTextureRHIParamRef SourceTexture, FTextureRHIParamRef DestTexture, const FRHICopyTextureInfo& CopyInfo) override final
 	{
 		ensureMsgf(!State.bInsideBeginRenderPass, TEXT("Copying inside a RenderPass is not efficient!"));
