Fixed compile error in Test config due to incorrect usage of static_cast in GetDynamicRHI()

diff --git a/Engine/Source/Runtime/RHI/Public/ValidationRHI.h b/Engine/Source/Runtime/RHI/Public/ValidationRHI.h
--- a/Engine/Source/Runtime/RHI/Public/ValidationRHI.h
+++ b/Engine/Source/Runtime/RHI/Public/ValidationRHI.h
@@ -1660,7 +1660,7 @@ inline TDynamicRHI* GetDynamicRHI()
 template <typename TDynamicRHI>
 inline TDynamicRHI* GetDynamicRHI()
 {
-	return static_cast<TDynamicRHI*>GDynamicRHI;
+	return static_cast<TDynamicRHI*>(GDynamicRHI);
 }
 
 #endif	// ENABLE_RHI_VALIDATION
