DR - Initial validations (PSOs using depth/stencil w/o setting one up)
#rnx

diff --git a/Engine/Source/Runtime/RHI/Private/ValidationRHI.cpp b/Engine/Source/Runtime/RHI/Private/ValidationRHI.cpp
--- a/Engine/Source/Runtime/RHI/Private/ValidationRHI.cpp
+++ b/Engine/Source/Runtime/RHI/Private/ValidationRHI.cpp
@@ -47,6 +47,46 @@ IRHIComputeContext* FValidationRHI::RHIGetDefaultAsyncComputeContext()
 	return AsyncComputeContext;
 }
 
+void FValidationRHI::ValidatePipeline(const FGraphicsPipelineStateInitializer& PSOInitializer)
+{
+	{
+		// Verify depth/stencil access/usage
+		bool bHasDepth = IsDepthOrStencilFormat(PSOInitializer.DepthStencilTargetFormat);
+		bool bHasStencil = IsStencilFormat(PSOInitializer.DepthStencilTargetFormat);
+		const FDepthStencilStateInitializerRHI& Initializer = GValidationRHI->DepthStencilStates.FindChecked(PSOInitializer.DepthStencilState);
+		if (bHasDepth)
+		{
+			if (!bHasStencil)
+			{
+				ensureMsgf(!Initializer.bEnableFrontFaceStencil
+					&& Initializer.FrontFaceStencilTest == CF_Always
+					&& Initializer.FrontFaceStencilFailStencilOp == SO_Keep
+					&& Initializer.FrontFaceDepthFailStencilOp == SO_Keep
+					&& Initializer.FrontFacePassStencilOp == SO_Keep
+					&& !Initializer.bEnableBackFaceStencil
+					&& Initializer.BackFaceStencilTest == CF_Always
+					&& Initializer.BackFaceStencilFailStencilOp == SO_Keep
+					&& Initializer.BackFaceDepthFailStencilOp == SO_Keep
+					&& Initializer.BackFacePassStencilOp == SO_Keep, TEXT("No stencil render target set, yet PSO wants to use stencil operations!"));
+				ensureMsgf(PSOInitializer.StencilTargetLoadAction == ERenderTargetLoadAction::ENoAction,
+					TEXT("No stencil target set, yet PSO wants to load from it!"));
+				ensureMsgf(PSOInitializer.StencilTargetStoreAction == ERenderTargetStoreAction::ENoAction,
+					TEXT("No stencil target set, yet PSO wants to store into it!"));
+			}
+		}
+		else
+		{
+			ensureMsgf(!Initializer.bEnableDepthWrite && Initializer.DepthTest == CF_Always, TEXT("No depth render target set, yet PSO wants to use depth operations!"));
+			ensureMsgf(PSOInitializer.DepthTargetLoadAction == ERenderTargetLoadAction::ENoAction
+				&& PSOInitializer.StencilTargetLoadAction == ERenderTargetLoadAction::ENoAction,
+				TEXT("No depth/stencil target set, yet PSO wants to load from it!"));
+			ensureMsgf(PSOInitializer.DepthTargetStoreAction == ERenderTargetStoreAction::ENoAction
+				&& PSOInitializer.StencilTargetStoreAction == ERenderTargetStoreAction::ENoAction,
+				TEXT("No depth/stencil target set, yet PSO wants to store into it!"));
+		}
+	}
+}
+
 
 FValidationComputeContext::FValidationComputeContext(FValidationRHI* InRHI)
 	: RHIContext(nullptr)
diff --git a/Engine/Source/Runtime/RHI/Public/RHIUtilities.h b/Engine/Source/Runtime/RHI/Public/RHIUtilities.h
--- a/Engine/Source/Runtime/RHI/Public/RHIUtilities.h
+++ b/Engine/Source/Runtime/RHI/Public/RHIUtilities.h
@@ -40,6 +40,21 @@ static inline bool IsDepthOrStencilFormat(EPixelFormat Format)
 	return false;
 }
 
+static inline bool IsStencilFormat(EPixelFormat Format)
+{
+	switch (Format)
+	{
+	case PF_DepthStencil:
+	case PF_X24_G8:
+		return true;
+
+	default:
+		break;
+	}
+
+	return false;
+}
+
 
 /** Encapsulates a GPU read/write buffer with its UAV and SRV. */
 struct FRWBuffer
diff --git a/Engine/Source/Runtime/RHI/Public/ValidationContext.h b/Engine/Source/Runtime/RHI/Public/ValidationContext.h
--- a/Engine/Source/Runtime/RHI/Public/ValidationContext.h
+++ b/Engine/Source/Runtime/RHI/Public/ValidationContext.h
@@ -152,6 +152,12 @@ public:
 		RHIContext->RHIWriteGPUFence(FenceRHI);
 	}
 
+	virtual void RHICopyToStagingBuffer(FVertexBufferRHIParamRef SourceBufferRHI, FStagingBufferRHIParamRef DestinationStagingBufferRHI, uint32 InOffset, uint32 InNumBytes) override final
+	{
+		RHIContext->RHICopyToStagingBuffer(SourceBufferRHI, DestinationStagingBufferRHI, InOffset, InNumBytes);
+	}
+
+
 	IRHIComputeContext* RHIContext;
 	FValidationRHI*		RHI;
 
@@ -786,6 +792,11 @@ public:
 		RHIContext->RHIWriteGPUFence(FenceRHI);
 	}
 
+	virtual void RHICopyToStagingBuffer(FVertexBufferRHIParamRef SourceBufferRHI, FStagingBufferRHIParamRef DestinationStagingBufferRHI, uint32 InOffset, uint32 InNumBytes) override final
+	{
+		RHIContext->RHICopyToStagingBuffer(SourceBufferRHI, DestinationStagingBufferRHI, InOffset, InNumBytes);
+	}
+
 	virtual void RHICopyTexture(FTextureRHIParamRef SourceTexture, FTextureRHIParamRef DestTexture, const FRHICopyTextureInfo& CopyInfo) override final
 	{
 		ensureMsgf(!State.bInsideBeginRenderPass, TEXT("Copying inside a RenderPass is not efficient!"));
diff --git a/Engine/Source/Runtime/RHI/Public/ValidationRHI.h b/Engine/Source/Runtime/RHI/Public/ValidationRHI.h
--- a/Engine/Source/Runtime/RHI/Public/ValidationRHI.h
+++ b/Engine/Source/Runtime/RHI/Public/ValidationRHI.h
@@ -62,7 +62,9 @@ public:
 	// FlushType: Thread safe
 	virtual FDepthStencilStateRHIRef RHICreateDepthStencilState(const FDepthStencilStateInitializerRHI& Initializer) override final
 	{
-		return RHI->RHICreateDepthStencilState(Initializer);
+		FDepthStencilStateRHIRef State = RHI->RHICreateDepthStencilState(Initializer);
+		DepthStencilStates.FindOrAdd(State.GetReference()) = Initializer;
+		return State;
 	}
 
 	// FlushType: Thread safe
@@ -284,6 +286,7 @@ public:
 	// TODO: [PSO API] Make pure virtual
 	virtual FGraphicsPipelineStateRHIRef RHICreateGraphicsPipelineState(const FGraphicsPipelineStateInitializer& Initializer) override final
 	{
+		ValidatePipeline(Initializer);
 		return RHI->RHICreateGraphicsPipelineState(Initializer);
 	}
 
@@ -294,6 +297,7 @@ public:
 
 	virtual FGraphicsPipelineStateRHIRef RHICreateGraphicsPipelineState(const FGraphicsPipelineStateInitializer& Initializer, FRHIPipelineBinaryLibrary* PipelineBinary) override final
 	{
+		ValidatePipeline(Initializer);
 		return RHI->RHICreateGraphicsPipelineState(Initializer);
 	}
 
@@ -1643,8 +1647,12 @@ public:
 	FValidationContext*			Context;
 	FValidationComputeContext*	AsyncComputeContext;
 
+	TMap<FRHIDepthStencilState*, FDepthStencilStateInitializerRHI> DepthStencilStates;
+
 private:
 	FString						RHIName;
+
+	void ValidatePipeline(const FGraphicsPipelineStateInitializer& Initializer);
 };
 
 extern RHI_API FValidationRHI* GValidationRHI;
