DR - Added missing platforms to RHIHasTiledGPU
#rnx

diff --git a/Engine/Source/Runtime/RHI/Public/RHIDefinitions.h b/Engine/Source/Runtime/RHI/Public/RHIDefinitions.h
--- a/Engine/Source/Runtime/RHI/Public/RHIDefinitions.h
+++ b/Engine/Source/Runtime/RHI/Public/RHIDefinitions.h
@@ -1094,7 +1094,9 @@ inline bool RHISupportsGeometryShaders(const EShaderPlatform Platform)
 inline bool RHIHasTiledGPU(const EShaderPlatform Platform)
 {
 	// @todo MetalMRT Technically we should include (Platform == SP_METAL_MRT) but this would disable depth-pre-pass which is currently required.
-	return Platform == SP_METAL || Platform == SP_METAL_TVOS || Platform == SP_OPENGL_ES2_IOS || Platform == SP_OPENGL_ES2_ANDROID || Platform == SP_OPENGL_ES3_1_ANDROID;
+	return Platform == SP_METAL || Platform == SP_METAL_TVOS || Platform == SP_METAL_MRT_TVOS
+		|| Platform == SP_OPENGL_ES2_IOS || Platform == SP_OPENGL_ES2_ANDROID || Platform == SP_OPENGL_ES3_1_ANDROID
+		|| Platform == SP_VULKAN_ES3_1_ANDROID;
 }
 
 inline bool RHISupportsMobileMultiView(const EShaderPlatform Platform)
