DR - Remove some deprecated RHICmdList functions
#rnx

diff --git a/Engine/Source/Runtime/RHI/Public/RHICommandList.h b/Engine/Source/Runtime/RHI/Public/RHICommandList.h
--- a/Engine/Source/Runtime/RHI/Public/RHICommandList.h
+++ b/Engine/Source/Runtime/RHI/Public/RHICommandList.h
@@ -2197,12 +2197,6 @@ public:
 		ALLOC_COMMAND(FRHICommandDrawPrimitive)(BaseVertexIndex, NumPrimitives, NumInstances);
 	}
 
-	UE_DEPRECATED(4.22, "PrimitiveType is now a part of the Graphics Pipeline State.")
-	FORCEINLINE_DEBUGGABLE void DrawPrimitive(uint32 PrimitiveType, uint32 BaseVertexIndex, uint32 NumPrimitives, uint32 NumInstances)
-	{
-		DrawPrimitive(BaseVertexIndex, NumPrimitives, NumInstances);
-	}
-
 	FORCEINLINE_DEBUGGABLE void DrawIndexedPrimitive(FIndexBufferRHIParamRef IndexBuffer, int32 BaseVertexIndex, uint32 FirstInstance, uint32 NumVertices, uint32 StartIndex, uint32 NumPrimitives, uint32 NumInstances)
 	{
 		if (!IndexBuffer)
@@ -2219,12 +2213,6 @@ public:
 		ALLOC_COMMAND(FRHICommandDrawIndexedPrimitive)(IndexBuffer, BaseVertexIndex, FirstInstance, NumVertices, StartIndex, NumPrimitives, NumInstances);
 	}
 
-	UE_DEPRECATED(4.22, "PrimitiveType is now a part of the Graphics Pipeline State.")
-	FORCEINLINE_DEBUGGABLE void DrawIndexedPrimitive(FIndexBufferRHIParamRef IndexBuffer, uint32 PrimitiveType, int32 BaseVertexIndex, uint32 FirstInstance, uint32 NumVertices, uint32 StartIndex, uint32 NumPrimitives, uint32 NumInstances)
-	{
-		DrawIndexedPrimitive(IndexBuffer, BaseVertexIndex, FirstInstance, NumVertices, StartIndex, NumPrimitives, NumInstances);
-	}
-
 	FORCEINLINE_DEBUGGABLE void SetStreamSource(uint32 StreamIndex, FVertexBufferRHIParamRef VertexBuffer, uint32 Offset)
 	{
 		if (Bypass())
@@ -2476,13 +2464,6 @@ public:
 		ALLOC_COMMAND(FRHICommandDrawPrimitiveIndirect)(ArgumentBuffer, ArgumentOffset);
 	}
 
-	UE_DEPRECATED(4.22, "PrimitiveType is now a part of the Graphics Pipeline State.")
-	FORCEINLINE_DEBUGGABLE void DrawPrimitiveIndirect(uint32 PrimitiveType, FVertexBufferRHIParamRef ArgumentBuffer, uint32 ArgumentOffset)
-	{
-		DrawPrimitiveIndirect(ArgumentBuffer, ArgumentOffset);
-	}
-
-
 	FORCEINLINE_DEBUGGABLE void DrawIndexedIndirect(FIndexBufferRHIParamRef IndexBufferRHI, FStructuredBufferRHIParamRef ArgumentsBufferRHI, uint32 DrawArgumentsIndex, uint32 NumInstances)
 	{
 		//check(IsOutsideRenderPass());
@@ -2494,12 +2475,6 @@ public:
 		ALLOC_COMMAND(FRHICommandDrawIndexedIndirect)(IndexBufferRHI, ArgumentsBufferRHI, DrawArgumentsIndex, NumInstances);
 	}
 
-	UE_DEPRECATED(4.22, "PrimitiveType is now a part of the Graphics Pipeline State.")
-	FORCEINLINE_DEBUGGABLE void DrawIndexedIndirect(FIndexBufferRHIParamRef IndexBufferRHI, uint32 PrimitiveType, FStructuredBufferRHIParamRef ArgumentsBufferRHI, uint32 DrawArgumentsIndex, uint32 NumInstances)
-	{
-		DrawIndexedIndirect(IndexBufferRHI, ArgumentsBufferRHI, DrawArgumentsIndex, NumInstances);
-	}
-
 	FORCEINLINE_DEBUGGABLE void DrawIndexedPrimitiveIndirect(FIndexBufferRHIParamRef IndexBuffer, FVertexBufferRHIParamRef ArgumentsBuffer, uint32 ArgumentOffset)
 	{
 		//check(IsOutsideRenderPass());
@@ -2511,12 +2486,6 @@ public:
 		ALLOC_COMMAND(FRHICommandDrawIndexedPrimitiveIndirect)(IndexBuffer, ArgumentsBuffer, ArgumentOffset);
 	}
 
-	UE_DEPRECATED(4.22, "PrimitiveType is now a part of the Graphics Pipeline State.")
-	FORCEINLINE_DEBUGGABLE void DrawIndexedPrimitiveIndirect(uint32 PrimitiveType, FIndexBufferRHIParamRef IndexBuffer, FVertexBufferRHIParamRef ArgumentsBuffer, uint32 ArgumentOffset)
-	{
-		DrawIndexedPrimitiveIndirect(IndexBuffer, ArgumentsBuffer, ArgumentOffset);
-	}
-
 	FORCEINLINE_DEBUGGABLE void SetDepthBounds(float MinDepth, float MaxDepth)
 	{
 		//check(IsOutsideRenderPass());
@@ -3424,12 +3393,6 @@ public:
 		return RHICreateBlendState(Initializer);
 	}
 	
-	UE_DEPRECATED(4.22, "Use PipelineStateCache::GetOrCreateVertexDeclaration() instead.")
-	FORCEINLINE FVertexDeclarationRHIRef CreateVertexDeclaration(const FVertexDeclarationElementList& Elements)
-	{
-		return GDynamicRHI->RHICreateVertexDeclaration(Elements);
-	}
-	
 	FORCEINLINE FPixelShaderRHIRef CreatePixelShader(const TArray<uint8>& Code)
 	{
 		LLM_SCOPE(ELLMTag::Shaders);
