DR - Validate trying to use incompatible APIs on incompatible feature leveles
#rnx

diff --git a/Engine/Source/Runtime/RHI/Public/ValidationRHI.h b/Engine/Source/Runtime/RHI/Public/ValidationRHI.h
--- a/Engine/Source/Runtime/RHI/Public/ValidationRHI.h
+++ b/Engine/Source/Runtime/RHI/Public/ValidationRHI.h
@@ -106,36 +106,42 @@ public:
 	// FlushType: Wait RHI Thread
 	virtual FHullShaderRHIRef RHICreateHullShader(const TArray<uint8>& Code) override final
 	{
+		check(RHISupportsTessellation(GMaxRHIShaderPlatform));
 		return RHI->RHICreateHullShader(Code);
 	}
 
 	// FlushType: Wait RHI Thread
 	virtual FHullShaderRHIRef RHICreateHullShader(FRHIShaderLibrary* Library, FSHAHash Hash) override final
 	{
+		check(RHISupportsTessellation(GMaxRHIShaderPlatform));
 		return RHI->RHICreateHullShader(Library, Hash);
 	}
 
 	// FlushType: Wait RHI Thread
 	virtual FDomainShaderRHIRef RHICreateDomainShader(const TArray<uint8>& Code) override final
 	{
+		check(RHISupportsTessellation(GMaxRHIShaderPlatform));
 		return RHI->RHICreateDomainShader(Code);
 	}
 
 	// FlushType: Wait RHI Thread
 	virtual FDomainShaderRHIRef RHICreateDomainShader(FRHIShaderLibrary* Library, FSHAHash Hash) override final
 	{
+		check(RHISupportsTessellation(GMaxRHIShaderPlatform));
 		return RHI->RHICreateDomainShader(Library, Hash);
 	}
 
 	// FlushType: Wait RHI Thread
 	virtual FGeometryShaderRHIRef RHICreateGeometryShader(const TArray<uint8>& Code) override final
 	{
+		check(RHISupportsGeometryShaders(GMaxRHIShaderPlatform));
 		return RHI->RHICreateGeometryShader(Code);
 	}
 
 	// FlushType: Wait RHI Thread
 	virtual FGeometryShaderRHIRef RHICreateGeometryShader(FRHIShaderLibrary* Library, FSHAHash Hash) override final
 	{
+		check(RHISupportsGeometryShaders(GMaxRHIShaderPlatform));
 		return RHI->RHICreateGeometryShader(Library, Hash);
 	}
 
@@ -143,6 +149,7 @@ public:
 	// FlushType: Wait RHI Thread
 	virtual FGeometryShaderRHIRef RHICreateGeometryShaderWithStreamOutput(const TArray<uint8>& Code, const FStreamOutElementList& ElementList, uint32 NumStrides, const uint32* Strides, int32 RasterizedStream) override final
 	{
+		check(RHISupportsGeometryShaders(GMaxRHIShaderPlatform));
 		return RHI->RHICreateGeometryShaderWithStreamOutput(Code, ElementList, NumStrides, Strides, RasterizedStream);
 	}
 
@@ -150,6 +157,7 @@ public:
 	// FlushType: Wait RHI Thread
 	virtual FGeometryShaderRHIRef RHICreateGeometryShaderWithStreamOutput(const FStreamOutElementList& ElementList, uint32 NumStrides, const uint32* Strides, int32 RasterizedStream, FRHIShaderLibrary* Library, FSHAHash Hash) override final
 	{
+		check(RHISupportsGeometryShaders(GMaxRHIShaderPlatform));
 		return RHI->RHICreateGeometryShaderWithStreamOutput(ElementList, NumStrides, Strides, RasterizedStream, Library, Hash);
 	}
 
@@ -162,11 +170,13 @@ public:
 	// FlushType: Wait RHI Thread
 	virtual FComputeShaderRHIRef RHICreateComputeShader(const TArray<uint8>& Code) override final
 	{
+		check(RHISupportsComputeShaders(GMaxRHIShaderPlatform));
 		return RHI->RHICreateComputeShader(Code);
 	}
 	// FlushType: Wait RHI Thread
 	virtual FComputeShaderRHIRef RHICreateComputeShader(FRHIShaderLibrary* Library, FSHAHash Hash) override final
 	{
+		check(RHISupportsComputeShaders(GMaxRHIShaderPlatform));
 		return RHI->RHICreateComputeShader(Library, Hash);
 	}
 
@@ -192,6 +202,7 @@ public:
 	// FlushType: Thread safe, but varies depending on the RHI	
 	virtual FComputeFenceRHIRef RHICreateComputeFence(const FName& Name) override final
 	{
+		check(RHISupportsComputeShaders(GMaxRHIShaderPlatform));
 		return RHI->RHICreateComputeFence(Name);
 	}
 
@@ -303,6 +314,7 @@ public:
 
 	virtual TRefCountPtr<FRHIComputePipelineState> RHICreateComputePipelineState(FRHIComputeShader* ComputeShader, FRHIPipelineBinaryLibrary* PipelineBinary) override final
 	{
+		check(RHISupportsComputeShaders(GMaxRHIShaderPlatform));
 		return RHI->RHICreateComputePipelineState(ComputeShader, PipelineBinary);
 	}
 
@@ -376,18 +388,21 @@ public:
 	// FlushType: Wait RHI Thread
 	virtual FStructuredBufferRHIRef RHICreateStructuredBuffer(uint32 Stride, uint32 Size, uint32 InUsage, FRHIResourceCreateInfo& CreateInfo) override final
 	{
+		check(IsFeatureLevelSupported(GMaxRHIShaderPlatform, ERHIFeatureLevel::SM5));
 		return RHI->RHICreateStructuredBuffer(Stride, Size, InUsage, CreateInfo);
 	}
 
 	// FlushType: Flush RHI Thread
 	virtual void* RHILockStructuredBuffer(FRHIStructuredBuffer* StructuredBuffer, uint32 Offset, uint32 SizeRHI, EResourceLockMode LockMode) override final
 	{
+		check(IsFeatureLevelSupported(GMaxRHIShaderPlatform, ERHIFeatureLevel::SM5));
 		return RHI->RHILockStructuredBuffer(StructuredBuffer, Offset, SizeRHI, LockMode);
 	}
 
 	// FlushType: Flush RHI Thread
 	virtual void RHIUnlockStructuredBuffer(FRHIStructuredBuffer* StructuredBuffer) override final
 	{
+		check(IsFeatureLevelSupported(GMaxRHIShaderPlatform, ERHIFeatureLevel::SM5));
 		RHI->RHIUnlockStructuredBuffer(StructuredBuffer);
 	}
 
@@ -395,6 +410,7 @@ public:
 	// FlushType: Wait RHI Thread
 	virtual FUnorderedAccessViewRHIRef RHICreateUnorderedAccessView(FRHIStructuredBuffer* StructuredBuffer, bool bUseUAVCounter, bool bAppendBuffer) override final
 	{
+		check(IsFeatureLevelSupported(GMaxRHIShaderPlatform, ERHIFeatureLevel::SM5));
 		return RHI->RHICreateUnorderedAccessView(StructuredBuffer, bUseUAVCounter, bAppendBuffer);
 	}
 
@@ -402,6 +418,7 @@ public:
 	// FlushType: Wait RHI Thread
 	virtual FUnorderedAccessViewRHIRef RHICreateUnorderedAccessView(FRHITexture* Texture, uint32 MipLevel) override final
 	{
+		check(IsFeatureLevelSupported(GMaxRHIShaderPlatform, ERHIFeatureLevel::SM5));
 		return RHI->RHICreateUnorderedAccessView(Texture, MipLevel);
 	}
 
@@ -409,6 +426,7 @@ public:
 	// FlushType: Wait RHI Thread
 	virtual FUnorderedAccessViewRHIRef RHICreateUnorderedAccessView(FRHIVertexBuffer* VertexBuffer, uint8 Format) override final
 	{
+		check(IsFeatureLevelSupported(GMaxRHIShaderPlatform, ERHIFeatureLevel::SM5));
 		return RHI->RHICreateUnorderedAccessView(VertexBuffer, Format);
 	}
 
@@ -416,6 +434,7 @@ public:
 	// FlushType: Wait RHI Thread
 	virtual FUnorderedAccessViewRHIRef RHICreateUnorderedAccessView(FRHIIndexBuffer* IndexBuffer, uint8 Format) override final
 	{
+		check(IsFeatureLevelSupported(GMaxRHIShaderPlatform, ERHIFeatureLevel::SM5));
 		return RHI->RHICreateUnorderedAccessView(IndexBuffer, Format);
 	}
 
@@ -423,6 +442,7 @@ public:
 	// FlushType: Wait RHI Thread
 	virtual FShaderResourceViewRHIRef RHICreateShaderResourceView(FRHIStructuredBuffer* StructuredBuffer) override final
 	{
+		check(RHISupportsComputeShaders(GMaxRHIShaderPlatform));
 		return RHI->RHICreateShaderResourceView(StructuredBuffer);
 	}
 
@@ -993,6 +1013,7 @@ public:
 
 	virtual FUnorderedAccessViewRHIRef RHIGetViewportBackBufferUAV(FRHIViewport* ViewportRHI) override final
 	{
+		check(IsFeatureLevelSupported(GMaxRHIShaderPlatform, ERHIFeatureLevel::SM5));
 		return RHI->RHIGetViewportBackBufferUAV(ViewportRHI);
 	}
 
@@ -1243,6 +1264,7 @@ public:
 
 	virtual FStructuredBufferRHIRef CreateStructuredBuffer_RenderThread(class FRHICommandListImmediate& RHICmdList, uint32 Stride, uint32 Size, uint32 InUsage, FRHIResourceCreateInfo& CreateInfo) override final
 	{
+		check(IsFeatureLevelSupported(GMaxRHIShaderPlatform, ERHIFeatureLevel::SM5));
 		return RHI->CreateStructuredBuffer_RenderThread(RHICmdList, Stride, Size, InUsage, CreateInfo);
 	}
 
@@ -1318,51 +1340,61 @@ public:
 
 	virtual FGeometryShaderRHIRef CreateGeometryShader_RenderThread(class FRHICommandListImmediate& RHICmdList, const TArray<uint8>& Code) override final
 	{
+		check(RHISupportsGeometryShaders(GMaxRHIShaderPlatform));
 		return RHI->CreateGeometryShader_RenderThread(RHICmdList, Code);
 	}
 
 	virtual FGeometryShaderRHIRef CreateGeometryShader_RenderThread(class FRHICommandListImmediate& RHICmdList, FRHIShaderLibrary* Library, FSHAHash Hash) override final
 	{
+		check(RHISupportsGeometryShaders(GMaxRHIShaderPlatform));
 		return RHI->CreateGeometryShader_RenderThread(RHICmdList, Library, Hash);
 	}
 
 	virtual FGeometryShaderRHIRef CreateGeometryShaderWithStreamOutput_RenderThread(class FRHICommandListImmediate& RHICmdList, const TArray<uint8>& Code, const FStreamOutElementList& ElementList, uint32 NumStrides, const uint32* Strides, int32 RasterizedStream) override final
 	{
+		check(RHISupportsGeometryShaders(GMaxRHIShaderPlatform));
 		return RHI->CreateGeometryShaderWithStreamOutput_RenderThread(RHICmdList, Code, ElementList, NumStrides, Strides, RasterizedStream);
 	}
 
 	virtual FGeometryShaderRHIRef CreateGeometryShaderWithStreamOutput_RenderThread(class FRHICommandListImmediate& RHICmdList, const FStreamOutElementList& ElementList, uint32 NumStrides, const uint32* Strides, int32 RasterizedStream, FRHIShaderLibrary* Library, FSHAHash Hash) override final
 	{
+		check(RHISupportsGeometryShaders(GMaxRHIShaderPlatform));
 		return RHI->CreateGeometryShaderWithStreamOutput_RenderThread(RHICmdList, ElementList, NumStrides, Strides, RasterizedStream, Library, Hash);
 	}
 
 	virtual FComputeShaderRHIRef CreateComputeShader_RenderThread(class FRHICommandListImmediate& RHICmdList, const TArray<uint8>& Code) override final
 	{
+		check(RHISupportsComputeShaders(GMaxRHIShaderPlatform));
 		return RHI->CreateComputeShader_RenderThread(RHICmdList, Code);
 	}
 
 	virtual FComputeShaderRHIRef CreateComputeShader_RenderThread(class FRHICommandListImmediate& RHICmdList, FRHIShaderLibrary* Library, FSHAHash Hash) override final
 	{
+		check(RHISupportsComputeShaders(GMaxRHIShaderPlatform));
 		return RHI->CreateComputeShader_RenderThread(RHICmdList, Library, Hash);
 	}
 
 	virtual FHullShaderRHIRef CreateHullShader_RenderThread(class FRHICommandListImmediate& RHICmdList, const TArray<uint8>& Code) override final
 	{
+		check(RHISupportsTessellation(GMaxRHIShaderPlatform));
 		return RHI->CreateHullShader_RenderThread(RHICmdList, Code);
 	}
 
 	virtual FHullShaderRHIRef CreateHullShader_RenderThread(class FRHICommandListImmediate& RHICmdList, FRHIShaderLibrary* Library, FSHAHash Hash) override final
 	{
+		check(RHISupportsTessellation(GMaxRHIShaderPlatform));
 		return RHI->CreateHullShader_RenderThread(RHICmdList, Library, Hash);
 	}
 
 	virtual FDomainShaderRHIRef CreateDomainShader_RenderThread(class FRHICommandListImmediate& RHICmdList, const TArray<uint8>& Code) override final
 	{
+		check(RHISupportsTessellation(GMaxRHIShaderPlatform));
 		return RHI->CreateDomainShader_RenderThread(RHICmdList, Code);
 	}
 
 	virtual FDomainShaderRHIRef CreateDomainShader_RenderThread(class FRHICommandListImmediate& RHICmdList, FRHIShaderLibrary* Library, FSHAHash Hash) override final
 	{
+		check(RHISupportsTessellation(GMaxRHIShaderPlatform));
 		return RHI->CreateDomainShader_RenderThread(RHICmdList, Library, Hash);
 	}
 
@@ -1433,21 +1465,25 @@ public:
 
 	virtual FUnorderedAccessViewRHIRef RHICreateUnorderedAccessView_RenderThread(class FRHICommandListImmediate& RHICmdList, FRHIStructuredBuffer* StructuredBuffer, bool bUseUAVCounter, bool bAppendBuffer) override final
 	{
+		check(IsFeatureLevelSupported(GMaxRHIShaderPlatform, ERHIFeatureLevel::SM5));
 		return RHI->RHICreateUnorderedAccessView_RenderThread(RHICmdList, StructuredBuffer, bUseUAVCounter, bAppendBuffer);
 	}
 
 	virtual FUnorderedAccessViewRHIRef RHICreateUnorderedAccessView_RenderThread(class FRHICommandListImmediate& RHICmdList, FRHITexture* Texture, uint32 MipLevel) override final
 	{
+		check(IsFeatureLevelSupported(GMaxRHIShaderPlatform, ERHIFeatureLevel::SM5));
 		return RHI->RHICreateUnorderedAccessView_RenderThread(RHICmdList, Texture, MipLevel);
 	}
 
 	virtual FUnorderedAccessViewRHIRef RHICreateUnorderedAccessView_RenderThread(class FRHICommandListImmediate& RHICmdList, FRHIVertexBuffer* VertexBuffer, uint8 Format) override final
 	{
+		check(IsFeatureLevelSupported(GMaxRHIShaderPlatform, ERHIFeatureLevel::SM5));
 		return RHI->RHICreateUnorderedAccessView_RenderThread(RHICmdList, VertexBuffer, Format);
 	}
 
 	virtual FUnorderedAccessViewRHIRef RHICreateUnorderedAccessView_RenderThread(class FRHICommandListImmediate& RHICmdList, FRHIIndexBuffer* IndexBuffer, uint8 Format) override final
 	{
+		check(IsFeatureLevelSupported(GMaxRHIShaderPlatform, ERHIFeatureLevel::SM5));
 		return RHI->RHICreateUnorderedAccessView_RenderThread(RHICmdList, IndexBuffer, Format);
 	}
 
@@ -1488,6 +1524,7 @@ public:
 
 	virtual FShaderResourceViewRHIRef RHICreateShaderResourceView_RenderThread(class FRHICommandListImmediate& RHICmdList, FRHIStructuredBuffer* StructuredBuffer) override final
 	{
+		check(IsFeatureLevelSupported(GMaxRHIShaderPlatform, ERHIFeatureLevel::SM5));
 		return RHI->RHICreateShaderResourceView_RenderThread(RHICmdList, StructuredBuffer);
 	}
 
