Fixed use-after-free in FRHIGPUTextureReadback::EnqueueCopy.

FRHIResourceCreateInfo debug name was initialized using a temporary FString, which falls out of scope by the time it's used in RHICreateTexture2D.

diff --git a/Engine/Source/Runtime/RHI/Private/RHIGPUReadback.cpp b/Engine/Source/Runtime/RHI/Private/RHIGPUReadback.cpp
--- a/Engine/Source/Runtime/RHI/Private/RHIGPUReadback.cpp
+++ b/Engine/Source/Runtime/RHI/Private/RHIGPUReadback.cpp
@@ -112,7 +112,8 @@ void FRHIGPUTextureReadback::EnqueueCopy(FRHICommandList& RHICmdList, FTextureRH
 		{
 			FIntVector TextureSize = SourceTexture->GetSizeXYZ();
 
-			FRHIResourceCreateInfo CreateInfo(*Fence->GetFName().ToString());
+			FString FenceName = Fence->GetFName().ToString();
+			FRHIResourceCreateInfo CreateInfo(*FenceName);
 			DestinationStagingBuffer = RHICreateTexture2D(TextureSize.X, TextureSize.Y, SourceTexture->GetFormat(), 1, 1, TexCreate_CPUReadback | TexCreate_HideInVisualizeTexture, CreateInfo);
 		}
 
