DR - Validation RHI
- CopyTexture restriction
- Disable temporarily load/store on missing depth stencil surfaces
#rnx

diff --git a/Engine/Source/Runtime/RHI/Private/ValidationRHI.cpp b/Engine/Source/Runtime/RHI/Private/ValidationRHI.cpp
--- a/Engine/Source/Runtime/RHI/Private/ValidationRHI.cpp
+++ b/Engine/Source/Runtime/RHI/Private/ValidationRHI.cpp
@@ -68,10 +68,12 @@ void FValidationRHI::ValidatePipeline(const FGraphicsPipelineStateInitializer& P
 					&& Initializer.BackFaceStencilFailStencilOp == SO_Keep
 					&& Initializer.BackFaceDepthFailStencilOp == SO_Keep
 					&& Initializer.BackFacePassStencilOp == SO_Keep, TEXT("No stencil render target set, yet PSO wants to use stencil operations!"));
+/*
 				ensureMsgf(PSOInitializer.StencilTargetLoadAction == ERenderTargetLoadAction::ENoAction,
 					TEXT("No stencil target set, yet PSO wants to load from it!"));
 				ensureMsgf(PSOInitializer.StencilTargetStoreAction == ERenderTargetStoreAction::ENoAction,
 					TEXT("No stencil target set, yet PSO wants to store into it!"));
+*/
 			}
 		}
 		else
diff --git a/Engine/Source/Runtime/RHI/Public/ValidationContext.h b/Engine/Source/Runtime/RHI/Public/ValidationContext.h
--- a/Engine/Source/Runtime/RHI/Public/ValidationContext.h
+++ b/Engine/Source/Runtime/RHI/Public/ValidationContext.h
@@ -800,6 +800,9 @@ public:
 	virtual void RHICopyTexture(FTextureRHIParamRef SourceTexture, FTextureRHIParamRef DestTexture, const FRHICopyTextureInfo& CopyInfo) override final
 	{
 		ensureMsgf(!State.bInsideBeginRenderPass, TEXT("Copying inside a RenderPass is not efficient!"));
+		check(SourceTexture);
+		check(DestTexture);
+		checkf(SourceTexture->GetFormat() == DestTexture->GetFormat(), TEXT("Some RHIs do not allow format conversion by the GPU for transfer operations!"));
 		RHIContext->RHICopyTexture(SourceTexture, DestTexture, CopyInfo);
 	}
 
