Fixed Clang compile error

diff --git a/Engine/Source/Runtime/RHI/Public/RHIResources.h b/Engine/Source/Runtime/RHI/Public/RHIResources.h
--- a/Engine/Source/Runtime/RHI/Public/RHIResources.h
+++ b/Engine/Source/Runtime/RHI/Public/RHIResources.h
@@ -2094,7 +2094,7 @@ public:
 			&& CallableHash == rhs.CallableHash;
 	}
 
-	static friend uint32 GetTypeHash(const FRayTracingPipelineStateInitializer& Initializer)
+	friend uint32 GetTypeHash(const FRayTracingPipelineStateInitializer& Initializer)
 	{
 		return GetTypeHash(Initializer.MaxPayloadSizeInBytes) ^
 			GetTypeHash(Initializer.bAllowHitGroupIndexing) ^
