DR - Rename ValidationRHI files (second pass)
#rnx

diff --git a/Engine/Source/Runtime/RHI/Private/RHIValidation.cpp b/Engine/Source/Runtime/RHI/Private/RHIValidation.cpp
new file mode 100644
--- /dev/null
+++ b/Engine/Source/Runtime/RHI/Private/RHIValidation.cpp
@@ -0,0 +1,124 @@
+// Copyright 1998-2019 Epic Games, Inc. All Rights Reserved.
+
+/*=============================================================================
+	RHIValidation.cpp: Public RHI Validation layer definitions.
+=============================================================================*/
+
+#include "RHIValidation.h"
+
+#if ENABLE_RHI_VALIDATION
+#include "RHI.h"
+#include "RHIValidationContext.h"
+
+FValidationRHI* GValidationRHI = nullptr;
+
+FValidationRHI::FValidationRHI(FDynamicRHI* InRHI)
+	: RHI(InRHI)
+{
+	check(RHI);
+
+	Context = new FValidationContext(this);
+	AsyncComputeContext = new FValidationComputeContext(this);
+}
+
+FValidationRHI::~FValidationRHI()
+{
+	delete AsyncComputeContext;
+	delete Context;
+}
+
+IRHICommandContext* FValidationRHI::RHIGetDefaultContext()
+{
+	if (!Context->RHIContext)
+	{
+		Context->RHIContext = RHI->RHIGetDefaultContext();
+	}
+
+	return Context;
+}
+
+IRHIComputeContext* FValidationRHI::RHIGetDefaultAsyncComputeContext()
+{
+	if (!AsyncComputeContext->RHIContext)
+	{
+		AsyncComputeContext->RHIContext = RHI->RHIGetDefaultAsyncComputeContext();
+	}
+
+	return AsyncComputeContext;
+}
+
+void FValidationRHI::ValidatePipeline(const FGraphicsPipelineStateInitializer& PSOInitializer)
+{
+	{
+		// Verify depth/stencil access/usage
+		bool bHasDepth = IsDepthOrStencilFormat(PSOInitializer.DepthStencilTargetFormat);
+		bool bHasStencil = IsStencilFormat(PSOInitializer.DepthStencilTargetFormat);
+		const FDepthStencilStateInitializerRHI& Initializer = GValidationRHI->DepthStencilStates.FindChecked(PSOInitializer.DepthStencilState);
+		if (bHasDepth)
+		{
+			if (!bHasStencil)
+			{
+				ensureMsgf(!Initializer.bEnableFrontFaceStencil
+					&& Initializer.FrontFaceStencilTest == CF_Always
+					&& Initializer.FrontFaceStencilFailStencilOp == SO_Keep
+					&& Initializer.FrontFaceDepthFailStencilOp == SO_Keep
+					&& Initializer.FrontFacePassStencilOp == SO_Keep
+					&& !Initializer.bEnableBackFaceStencil
+					&& Initializer.BackFaceStencilTest == CF_Always
+					&& Initializer.BackFaceStencilFailStencilOp == SO_Keep
+					&& Initializer.BackFaceDepthFailStencilOp == SO_Keep
+					&& Initializer.BackFacePassStencilOp == SO_Keep, TEXT("No stencil render target set, yet PSO wants to use stencil operations!"));
+/*
+				ensureMsgf(PSOInitializer.StencilTargetLoadAction == ERenderTargetLoadAction::ENoAction,
+					TEXT("No stencil target set, yet PSO wants to load from it!"));
+				ensureMsgf(PSOInitializer.StencilTargetStoreAction == ERenderTargetStoreAction::ENoAction,
+					TEXT("No stencil target set, yet PSO wants to store into it!"));
+*/
+			}
+		}
+		else
+		{
+			ensureMsgf(!Initializer.bEnableDepthWrite && Initializer.DepthTest == CF_Always, TEXT("No depth render target set, yet PSO wants to use depth operations!"));
+			ensureMsgf(PSOInitializer.DepthTargetLoadAction == ERenderTargetLoadAction::ENoAction
+				&& PSOInitializer.StencilTargetLoadAction == ERenderTargetLoadAction::ENoAction,
+				TEXT("No depth/stencil target set, yet PSO wants to load from it!"));
+			ensureMsgf(PSOInitializer.DepthTargetStoreAction == ERenderTargetStoreAction::ENoAction
+				&& PSOInitializer.StencilTargetStoreAction == ERenderTargetStoreAction::ENoAction,
+				TEXT("No depth/stencil target set, yet PSO wants to store into it!"));
+		}
+	}
+}
+
+
+FValidationComputeContext::FValidationComputeContext(FValidationRHI* InRHI)
+	: RHIContext(nullptr)
+	, RHI(InRHI)
+{
+	State.Reset();
+}
+
+void FValidationComputeContext::FState::Reset()
+{
+	ComputePassName.Reset();
+	bComputeShaderSet = false;
+}
+
+
+FValidationContext::FValidationContext(FValidationRHI* InRHI)
+	: RHIContext(nullptr)
+	, RHI(InRHI)
+{
+	State.Reset();
+}
+
+void FValidationContext::FState::Reset()
+{
+	bInsideBeginRenderPass = false;
+	bGfxPSOSet = false;
+	RenderPassName.Reset();
+	PreviousRenderPassName.Reset();
+	bInsideComputePass = false;
+	ComputePassName.Reset();
+	bComputeShaderSet = false;
+}
+#endif	// ENABLE_RHI_VALIDATION
diff --git a/Engine/Source/Runtime/RHI/Private/ValidationRHI.cpp b/Engine/Source/Runtime/RHI/Private/ValidationRHI.cpp
deleted file mode 100644
diff --git a/Engine/Source/Runtime/RHI/Public/RHIValidation.h b/Engine/Source/Runtime/RHI/Public/RHIValidation.h
--- a/Engine/Source/Runtime/RHI/Public/RHIValidation.h
+++ b/Engine/Source/Runtime/RHI/Public/RHIValidation.h
@@ -6,7 +6,7 @@
 
 #pragma once 
 
-#include "ValidationRHICommon.h"
+#include "RHIValidationCommon.h"
 
 #if ENABLE_RHI_VALIDATION
 class FValidationComputeContext;
diff --git a/Engine/Source/Runtime/RHI/Public/RHIValidationCommon.h b/Engine/Source/Runtime/RHI/Public/RHIValidationCommon.h
new file mode 100644
--- /dev/null
+++ b/Engine/Source/Runtime/RHI/Public/RHIValidationCommon.h
@@ -0,0 +1,16 @@
+// Copyright 1998-2019 Epic Games, Inc. All Rights Reserved.
+
+/*=============================================================================
+	ValidationRHICommon.h: Public Valdation RHI definitions.
+=============================================================================*/
+
+#pragma once 
+
+#ifndef ENABLE_RHI_VALIDATION
+	#define ENABLE_RHI_VALIDATION	(UE_BUILD_DEBUG || UE_BUILD_DEVELOPMENT)
+#endif
+
+#include "RHI.h"
+#include "RHIResources.h"
+#include "RHIContext.h"
+#include "DynamicRHI.h"
diff --git a/Engine/Source/Runtime/RHI/Public/RHIValidationContext.h b/Engine/Source/Runtime/RHI/Public/RHIValidationContext.h
--- a/Engine/Source/Runtime/RHI/Public/RHIValidationContext.h
+++ b/Engine/Source/Runtime/RHI/Public/RHIValidationContext.h
@@ -1,12 +1,12 @@
 // Copyright 1998-2019 Epic Games, Inc. All Rights Reserved.
 
 /*=============================================================================
-	ValidationContext.h: Public Validation Context definitions.
+	RHIValidationContext.h: Public RHI Validation Context definitions.
 =============================================================================*/
 
 #pragma once
 
-#include "ValidationRHICommon.h"
+#include "RHIValidationCommon.h"
 
 #if ENABLE_RHI_VALIDATION
 class FValidationRHI;
diff --git a/Engine/Source/Runtime/RHI/Public/ValidationRHICommon.h b/Engine/Source/Runtime/RHI/Public/ValidationRHICommon.h
deleted file mode 100644
diff --git a/Engine/Source/Runtime/Windows/D3D11RHI/Private/D3D11RHIPrivate.h b/Engine/Source/Runtime/Windows/D3D11RHI/Private/D3D11RHIPrivate.h
--- a/Engine/Source/Runtime/Windows/D3D11RHI/Private/D3D11RHIPrivate.h
+++ b/Engine/Source/Runtime/Windows/D3D11RHI/Private/D3D11RHIPrivate.h
@@ -28,7 +28,7 @@ DECLARE_LOG_CATEGORY_EXTERN(LogD3D11RHI, Log, All);
 #include "D3D11Viewport.h"
 #include "D3D11ConstantBuffer.h"
 #include "D3D11StateCache.h"
-#include "ValidationRHICommon.h"
+#include "RHIValidationCommon.h"
 
 #ifndef WITH_DX_PERF
 #define WITH_DX_PERF	1
