DR - Updated copy texture validation; removed some individual tests
#rnx

diff --git a/Engine/Source/Runtime/RHI/Public/RHIValidationContext.h b/Engine/Source/Runtime/RHI/Public/RHIValidationContext.h
--- a/Engine/Source/Runtime/RHI/Public/RHIValidationContext.h
+++ b/Engine/Source/Runtime/RHI/Public/RHIValidationContext.h
@@ -664,13 +664,13 @@ public:
 
 	virtual void RHISetStencilRef(uint32 StencilRef) override final
 	{
-		checkf(State.bGfxPSOSet, TEXT("A Graphics PSO has to be set to change stencil ref!"));
+		//checkf(State.bGfxPSOSet, TEXT("A Graphics PSO has to be set to change stencil ref!"));
 		RHIContext->RHISetStencilRef(StencilRef);
 	}
 
 	virtual void RHISetBlendFactor(const FLinearColor& BlendFactor) override final
 	{
-		checkf(State.bGfxPSOSet, TEXT("A Graphics PSO has to be set to change blend factor!"));
+		//checkf(State.bGfxPSOSet, TEXT("A Graphics PSO has to be set to change blend factor!"));
 		RHIContext->RHISetBlendFactor(BlendFactor);
 	}
 
@@ -803,6 +803,25 @@ public:
 		check(SourceTexture);
 		check(DestTexture);
 		checkf(SourceTexture->GetFormat() == DestTexture->GetFormat(), TEXT("Some RHIs do not allow format conversion by the GPU for transfer operations!"));
+
+		FIntVector CopySize = CopyInfo.Size;
+		FIntVector SrcSize = SourceTexture->GetSizeXYZ();
+		FIntVector DestSize = DestTexture->GetSizeXYZ();
+		if (CopySize == FIntVector::ZeroValue)
+		{
+			CopySize = SrcSize;
+		}
+
+		checkf(SrcSize.X <= DestSize.X && SrcSize.Y <= DestSize.Y, TEXT("Some RHIs can't perform scaling operations [%dx%d to %dx%d] during copies!"), SrcSize.X, SrcSize.Y, DestSize.X, DestSize.Y);
+		check(CopyInfo.SourcePosition.X >= 0 && CopyInfo.SourcePosition.Y >= 0 && CopyInfo.SourcePosition.Z >= 0);
+		check(CopyInfo.DestPosition.X >= 0 && CopyInfo.DestPosition.Y >= 0 && CopyInfo.DestPosition.Z >= 0);
+		check(CopyInfo.SourcePosition.X + CopySize.X <= CopyInfo.DestPosition.X + DestSize.X);
+		check(CopyInfo.SourcePosition.Y + CopySize.Y <= CopyInfo.DestPosition.Y + DestSize.Y);
+		if (SourceTexture->GetTexture3D() && DestTexture->GetTexture3D())
+		{
+			check(CopyInfo.SourcePosition.Z + CopySize.Z <= CopyInfo.DestPosition.Z + DestSize.Z);
+		}
+
 		RHIContext->RHICopyTexture(SourceTexture, DestTexture, CopyInfo);
 	}
 
