DR - Remove deprecated utility
#rnx

diff --git a/Engine/Source/Runtime/RHI/Public/RHIUtilities.h b/Engine/Source/Runtime/RHI/Public/RHIUtilities.h
--- a/Engine/Source/Runtime/RHI/Public/RHIUtilities.h
+++ b/Engine/Source/Runtime/RHI/Public/RHIUtilities.h
@@ -471,23 +471,6 @@ inline void SetRenderTarget(FRHICommandList& RHICmdList, FTextureRHIParamRef New
 	PRAGMA_ENABLE_DEPRECATION_WARNINGS
 }
 
-/** Helper for the common case of using a single color and depth render target, with a mip index for the color target. */
-UE_DEPRECATED(4.22, "SetRenderTarget API is deprecated; please use RHIBegin/EndRenderPass instead.")
-inline void SetRenderTarget(FRHICommandList& RHICmdList, FTextureRHIParamRef NewRenderTarget, int32 MipIndex, FTextureRHIParamRef NewDepthStencilTarget, bool bWritableBarrier = false)
-{
-	FRHIRenderTargetView RTV(NewRenderTarget, ERenderTargetLoadAction::ELoad, MipIndex, -1);
-	FRHIDepthRenderTargetView DepthRTV(NewDepthStencilTarget, ERenderTargetLoadAction::ELoad, ERenderTargetStoreAction::EStore);
-	
-	//make these rendertargets safely writable
-	if (bWritableBarrier)
-	{
-		TransitionSetRenderTargetsHelper(RHICmdList, NewRenderTarget, NewDepthStencilTarget, FExclusiveDepthStencil::DepthWrite_StencilWrite);
-	}
-	PRAGMA_DISABLE_DEPRECATION_WARNINGS
-	RHICmdList.SetRenderTargets(1, &RTV, &DepthRTV, 0, nullptr);
-	PRAGMA_ENABLE_DEPRECATION_WARNINGS
-}
-
 /** Helper for the common case of using a single color and depth render target, with a mip index for the color target. */
 UE_DEPRECATED(4.22, "SetRenderTarget API is deprecated; please use RHIBegin/EndRenderPass instead.")
 inline void SetRenderTarget(FRHICommandList& RHICmdList, FTextureRHIParamRef NewRenderTarget, int32 MipIndex, int32 ArraySliceIndex, FTextureRHIParamRef NewDepthStencilTarget, bool bWritableBarrier = false)
