DR - Remove hard to validate case in Occlusion Queriy pools.

diff --git a/Engine/Source/Runtime/RHI/Private/DynamicRHI.cpp b/Engine/Source/Runtime/RHI/Private/DynamicRHI.cpp
--- a/Engine/Source/Runtime/RHI/Private/DynamicRHI.cpp
+++ b/Engine/Source/Runtime/RHI/Private/DynamicRHI.cpp
@@ -383,7 +383,8 @@ FRHIPooledRenderQuery FDefaultRHIRenderQueryPool::AllocateQuery()
 
 void FDefaultRHIRenderQueryPool::ReleaseQuery(TRefCountPtr<FRHIRenderQuery>&& Query)
 {
-	checkf(Query.IsValid() && Query.GetRefCount() <= 2, TEXT("Query has been released but reference still held: use FRHIPooledRenderQuery::GetQueryRef() with extreme caution"));
+	//Hard to validate because of Resource resurrection, better to remove GetQueryRef entirely
+	//checkf(Query.IsValid() && Query.GetRefCount() <= 2, TEXT("Query has been released but reference still held: use FRHIPooledRenderQuery::GetQueryRef() with extreme caution"));
 	Queries.Push(MoveTemp(Query));
 	check(!Query.IsValid());
 }
