Added argument checks to RHICopySubTextureRegion function in ValidationRHI.

#rnx

diff --git a/Engine/Source/Runtime/RHI/Private/RHIValidation.cpp b/Engine/Source/Runtime/RHI/Private/RHIValidation.cpp
--- a/Engine/Source/Runtime/RHI/Private/RHIValidation.cpp
+++ b/Engine/Source/Runtime/RHI/Private/RHIValidation.cpp
@@ -89,6 +89,18 @@ void FValidationRHI::ValidatePipeline(const FGraphicsPipelineStateInitializer& P
 	}
 }
 
+void FValidationRHI::ValidateCopySubTextureRegion(FRHITexture2D* SourceTexture, FRHITexture2D* DestinationTexture, FBox2D SourceBox, FBox2D DestinationBox)
+{
+	const int32 SrcWidth = (int32)(SourceBox.Max.X - SourceBox.Min.X);
+	const int32 SrcHeight = (int32)(SourceBox.Max.Y - SourceBox.Min.Y);
+	const int32 DestWidth = (int32)(DestinationBox.Max.X - DestinationBox.Min.X);
+	const int32 DestHeight = (int32)(DestinationBox.Max.Y - DestinationBox.Min.Y);
+
+	check(SourceTexture != nullptr && DestinationTexture != nullptr);
+	//	check(SrcWidth > 0 && SrcHeight > 0);
+	check(DestWidth == SrcWidth && DestHeight == SrcHeight);
+}
+
 
 FValidationComputeContext::FValidationComputeContext(FValidationRHI* InRHI)
 	: RHIContext(nullptr)
diff --git a/Engine/Source/Runtime/RHI/Public/RHIValidation.h b/Engine/Source/Runtime/RHI/Public/RHIValidation.h
--- a/Engine/Source/Runtime/RHI/Public/RHIValidation.h
+++ b/Engine/Source/Runtime/RHI/Public/RHIValidation.h
@@ -1618,6 +1618,7 @@ public:
 
 	virtual void RHICopySubTextureRegion(FRHITexture2D* SourceTexture, FRHITexture2D* DestinationTexture, FBox2D SourceBox, FBox2D DestinationBox) override final
 	{
+		ValidateCopySubTextureRegion(SourceTexture, DestinationTexture, SourceBox, DestinationBox);
 		RHI->RHICopySubTextureRegion(SourceTexture, DestinationTexture, SourceBox, DestinationBox);
 	}
 
@@ -1686,6 +1687,7 @@ private:
 	FString						RHIName;
 
 	void ValidatePipeline(const FGraphicsPipelineStateInitializer& Initializer);
+	void ValidateCopySubTextureRegion(FRHITexture2D* SourceTexture, FRHITexture2D* DestinationTexture, FBox2D SourceBox, FBox2D DestinationBox);
 };
 
 extern RHI_API FValidationRHI* GValidationRHI;
