Added uniform buffer references to FRHIUniformBufferLayout::Resources array. This enables render graph to release references without knowledge of the specific type in FRDGPass.

diff --git a/Engine/Source/Runtime/RHI/Public/RHIDefinitions.h b/Engine/Source/Runtime/RHI/Public/RHIDefinitions.h
--- a/Engine/Source/Runtime/RHI/Public/RHIDefinitions.h
+++ b/Engine/Source/Runtime/RHI/Public/RHIDefinitions.h
@@ -1177,18 +1177,20 @@ inline bool IsShaderParameterTypeForUniformBufferLayout(EUniformBufferBaseType B
 		// RHI is able to dereference the RHI resource allocated in FRDGResource::CachedRHI to avoid pain in the high level to passdown.
 		IsRDGResourceReferenceShaderParameterType(BaseType) ||
 
-		// #yuriy_todo: RHI is able to dereference uniform buffer in root shader parameter structures
-		// BaseType == UBMT_REFERENCED_STRUCT ||
-
 		// Render graph uses FRHIUniformBufferLayout to walk pass' parameters.
+		BaseType == UBMT_REFERENCED_STRUCT ||
 		BaseType == UBMT_RENDER_TARGET_BINDING_SLOTS;
 }
 
 /** Returns whether the shader parameter type in FRHIUniformBufferLayout is actually ignored by the RHI. */
 inline bool IsShaderParameterTypeIgnoredByRHI(EUniformBufferBaseType BaseType)
 {
-	// Render targets bindings slots needs to be in FRHIUniformBufferLayout for render graph, but the RHI does not actually need to know about it.
-	return BaseType == UBMT_RENDER_TARGET_BINDING_SLOTS;
+	return
+		// Render targets bindings slots needs to be in FRHIUniformBufferLayout for render graph, but the RHI does not actually need to know about it.
+		BaseType == UBMT_RENDER_TARGET_BINDING_SLOTS || 
+
+		// #yuriy_todo: RHI is able to dereference uniform buffer in root shader parameter structures
+		BaseType == UBMT_REFERENCED_STRUCT;
 }
 
 inline const TCHAR* GetShaderFrequencyString(EShaderFrequency Frequency, bool bIncludePrefix = true)
