Fix PS4 implementation of RHICopyToResolveTarget to support BC compressed format copies.
On PS4 RHICopyTexture is implemented in terms of RHICopyToResolveTarget which is why the fix ends up here (this seems backwards, but maybe that's just me!)
Change runtime virtual texture to use CopyTexture instead of the newly deprecated CopySubTextureRegion.
#rnx

diff --git a/Engine/Source/Runtime/RHI/Public/RHIContext.h b/Engine/Source/Runtime/RHI/Public/RHIContext.h
--- a/Engine/Source/Runtime/RHI/Public/RHIContext.h
+++ b/Engine/Source/Runtime/RHI/Public/RHIContext.h
@@ -548,21 +548,42 @@ public:
 
 	virtual void RHICopyTexture(FTextureRHIParamRef SourceTexture, FTextureRHIParamRef DestTexture, const FRHICopyTextureInfo& CopyInfo)
 	{
+		// Implement in terms of RHICopyToResolveTarget()
 		const bool bIsCube = SourceTexture->GetTextureCube() != nullptr;
 		const bool bAllCubeFaces = bIsCube && (CopyInfo.NumSlices % 6) == 0;
 		const int32 NumArraySlices = bAllCubeFaces ? CopyInfo.NumSlices / 6 : CopyInfo.NumSlices;
 		const int32 NumFaces = bAllCubeFaces ? 6 : 1;
+		
+		// Default FResolveRect objects will copy full source texture to (0,0) in dest texture
+		FResolveRect SourceRect;
+		FResolveRect DestRect;
+
+		if (CopyInfo.Size != FIntVector::ZeroValue)
+		{
+			// Note that specifying size in both rects seems redundant/confusing! FRHICopyTextureInfo is a much more explicit structure.
+			// But all RHICopyToResolveTarget implementations use the size from the source rect, and fall back to source texture size if it isn't specified.
+			SourceRect = FResolveRect(CopyInfo.SourcePosition.X, CopyInfo.SourcePosition.Y, CopyInfo.SourcePosition.X + CopyInfo.Size.X, CopyInfo.SourcePosition.Y + CopyInfo.Size.Y);
+			DestRect = FResolveRect(CopyInfo.DestPosition.X, CopyInfo.DestPosition.Y, CopyInfo.DestPosition.X + CopyInfo.Size.X, CopyInfo.DestPosition.Y + CopyInfo.Size.Y);
+		}
+		else if (CopyInfo.DestPosition != FIntVector::ZeroValue)
+		{
+			// Specify dest position only in case where we want to copy full source texture to location within dest texture
+			DestRect = FResolveRect(CopyInfo.DestPosition.X, CopyInfo.DestPosition.Y, CopyInfo.DestPosition.X + CopyInfo.Size.X, CopyInfo.DestPosition.Y + CopyInfo.Size.Y);
+		}
+
 		for (int32 ArrayIndex = 0; ArrayIndex < NumArraySlices; ++ArrayIndex)
 		{
 			int32 SourceArrayIndex = CopyInfo.SourceSliceIndex + ArrayIndex;
 			int32 DestArrayIndex = CopyInfo.DestSliceIndex + ArrayIndex;
 			for (int32 FaceIndex = 0; FaceIndex < NumFaces; ++FaceIndex)
 			{
-				FResolveParams ResolveParams(FResolveRect(),
+				FResolveParams ResolveParams(
+					SourceRect,
 					bIsCube ? (ECubeFace)FaceIndex : CubeFace_PosX,
 					CopyInfo.SourceMipIndex,
 					SourceArrayIndex,
-					DestArrayIndex
+					DestArrayIndex,
+					DestRect
 				);
 				RHICopyToResolveTarget(SourceTexture, DestTexture, ResolveParams);
 			}
diff --git a/Engine/Source/Runtime/Renderer/Private/VT/RuntimeVirtualTextureRender.cpp b/Engine/Source/Runtime/Renderer/Private/VT/RuntimeVirtualTextureRender.cpp
--- a/Engine/Source/Runtime/Renderer/Private/VT/RuntimeVirtualTextureRender.cpp
+++ b/Engine/Source/Runtime/Renderer/Private/VT/RuntimeVirtualTextureRender.cpp
@@ -675,20 +675,20 @@ namespace RuntimeVirtualTexture
 		// Copy to final destination
 		if (GraphSetup.OutputAlias0 != nullptr)
 		{
-			RHICmdList.CopySubTextureRegion(
-				GraphOutputTexture0->GetRenderTargetItem().ShaderResourceTexture->GetTexture2D(),
-				OutputTexture0->GetTexture2D(),
-				FBox2D(FVector2D(0, 0), FVector2D(GraphOutputSize0.X, GraphOutputSize0.Y)),
-				DestBox0);
+			FRHICopyTextureInfo Info;
+			Info.Size = GraphOutputSize0;
+			Info.DestPosition = FIntVector(DestBox0.Min.X, DestBox0.Min.Y, 0);
+
+			RHICmdList.CopyTexture(GraphOutputTexture0->GetRenderTargetItem().ShaderResourceTexture->GetTexture2D(), OutputTexture0->GetTexture2D(), Info);
 		}
 
 		if (GraphSetup.OutputAlias1 != nullptr)
 		{
-			RHICmdList.CopySubTextureRegion(
-				GraphOutputTexture1->GetRenderTargetItem().ShaderResourceTexture->GetTexture2D(),
-				OutputTexture1->GetTexture2D(),
-				FBox2D(FVector2D(0, 0), FVector2D(GraphOutputSize1.X, GraphOutputSize1.Y)),
-				DestBox1);
+			FRHICopyTextureInfo Info;
+			Info.Size = GraphOutputSize1;
+			Info.DestPosition = FIntVector(DestBox1.Min.X, DestBox1.Min.Y, 0);
+
+			RHICmdList.CopyTexture(GraphOutputTexture1->GetRenderTargetItem().ShaderResourceTexture->GetTexture2D(), OutputTexture1->GetTexture2D(), Info);
 		}
 	}
 }
