Move change to RHICopyTexture in CL6412538 down to PS4 only implementation
Safer then potentially triggering an issue on other platforms!
Also tidied dependencies between RHICopyTexture and RHICopyToResolveTarget
#rnx

diff --git a/Engine/Source/Runtime/RHI/Public/RHIContext.h b/Engine/Source/Runtime/RHI/Public/RHIContext.h
--- a/Engine/Source/Runtime/RHI/Public/RHIContext.h
+++ b/Engine/Source/Runtime/RHI/Public/RHIContext.h
@@ -548,42 +548,21 @@ public:
 
 	virtual void RHICopyTexture(FTextureRHIParamRef SourceTexture, FTextureRHIParamRef DestTexture, const FRHICopyTextureInfo& CopyInfo)
 	{
-		// Implement in terms of RHICopyToResolveTarget()
 		const bool bIsCube = SourceTexture->GetTextureCube() != nullptr;
 		const bool bAllCubeFaces = bIsCube && (CopyInfo.NumSlices % 6) == 0;
 		const int32 NumArraySlices = bAllCubeFaces ? CopyInfo.NumSlices / 6 : CopyInfo.NumSlices;
 		const int32 NumFaces = bAllCubeFaces ? 6 : 1;
-		
-		// Default FResolveRect objects will copy full source texture to (0,0) in dest texture
-		FResolveRect SourceRect;
-		FResolveRect DestRect;
-
-		if (CopyInfo.Size != FIntVector::ZeroValue)
-		{
-			// Note that specifying size in both rects seems redundant/confusing! FRHICopyTextureInfo is a much more explicit structure.
-			// But all RHICopyToResolveTarget implementations use the size from the source rect, and fall back to source texture size if it isn't specified.
-			SourceRect = FResolveRect(CopyInfo.SourcePosition.X, CopyInfo.SourcePosition.Y, CopyInfo.SourcePosition.X + CopyInfo.Size.X, CopyInfo.SourcePosition.Y + CopyInfo.Size.Y);
-			DestRect = FResolveRect(CopyInfo.DestPosition.X, CopyInfo.DestPosition.Y, CopyInfo.DestPosition.X + CopyInfo.Size.X, CopyInfo.DestPosition.Y + CopyInfo.Size.Y);
-		}
-		else if (CopyInfo.DestPosition != FIntVector::ZeroValue)
-		{
-			// Specify dest position only in case where we want to copy full source texture to location within dest texture
-			DestRect = FResolveRect(CopyInfo.DestPosition.X, CopyInfo.DestPosition.Y, CopyInfo.DestPosition.X + CopyInfo.Size.X, CopyInfo.DestPosition.Y + CopyInfo.Size.Y);
-		}
-
 		for (int32 ArrayIndex = 0; ArrayIndex < NumArraySlices; ++ArrayIndex)
 		{
 			int32 SourceArrayIndex = CopyInfo.SourceSliceIndex + ArrayIndex;
 			int32 DestArrayIndex = CopyInfo.DestSliceIndex + ArrayIndex;
 			for (int32 FaceIndex = 0; FaceIndex < NumFaces; ++FaceIndex)
 			{
-				FResolveParams ResolveParams(
-					SourceRect,
+				FResolveParams ResolveParams(FResolveRect(),
 					bIsCube ? (ECubeFace)FaceIndex : CubeFace_PosX,
 					CopyInfo.SourceMipIndex,
 					SourceArrayIndex,
-					DestArrayIndex,
-					DestRect
+					DestArrayIndex
 				);
 				RHICopyToResolveTarget(SourceTexture, DestTexture, ResolveParams);
 			}
