DR - More deprecated warnings
#rnx

diff --git a/Engine/Source/Runtime/RHI/Public/RHIUtilities.h b/Engine/Source/Runtime/RHI/Public/RHIUtilities.h
--- a/Engine/Source/Runtime/RHI/Public/RHIUtilities.h
+++ b/Engine/Source/Runtime/RHI/Public/RHIUtilities.h
@@ -369,6 +369,7 @@ inline void DecodeRenderTargetMode(ESimpleRenderTargetMode Mode, ERenderTargetLo
 	}
 }
 
+UE_DEPRECATED(4.22, "SetRenderTargets API is deprecated; please use RHIBegin/EndRenderPass instead.")
 inline void TransitionSetRenderTargetsHelper(FRHICommandList& RHICmdList, FTextureRHIParamRef NewRenderTarget, FTextureRHIParamRef NewDepthStencilTarget, FExclusiveDepthStencil DepthStencilAccess)
 {
 	int32 TransitionIndex = 0;
@@ -386,6 +387,7 @@ inline void TransitionSetRenderTargetsHelper(FRHICommandList& RHICmdList, FTextu
 	RHICmdList.TransitionResources(EResourceTransitionAccess::EWritable, Transitions, TransitionIndex);
 }
 
+UE_DEPRECATED(4.22, "SetRenderTargets API is deprecated; please use RHIBegin/EndRenderPass instead.")
 inline void TransitionSetRenderTargetsHelper(FRHICommandList& RHICmdList, uint32 NumRenderTargets, const FTextureRHIParamRef* NewRenderTargetsRHI, const FTextureRHIParamRef NewDepthStencilTargetRHI, FExclusiveDepthStencil DepthStencilAccess)
 {
 	FTextureRHIParamRef Transitions[MaxSimultaneousRenderTargets + 1];
@@ -440,11 +442,11 @@ inline void SetRenderTarget(FRHICommandList& RHICmdList, FTextureRHIParamRef New
 	FRHIDepthRenderTargetView DepthRTV(NewDepthStencilTarget, ERenderTargetLoadAction::ELoad, ERenderTargetStoreAction::EStore);
 
 	//make these rendertargets safely writable
+	PRAGMA_DISABLE_DEPRECATION_WARNINGS
 	if (bWritableBarrier)
 	{
 		TransitionSetRenderTargetsHelper(RHICmdList, NewRenderTarget, NewDepthStencilTarget, FExclusiveDepthStencil::DepthWrite_StencilWrite);
 	}
-	PRAGMA_DISABLE_DEPRECATION_WARNINGS
 	RHICmdList.SetRenderTargets(1, &RTV, &DepthRTV, 0, NULL);
 	PRAGMA_ENABLE_DEPRECATION_WARNINGS
 }
@@ -458,6 +460,7 @@ inline void SetRenderTarget(FRHICommandList& RHICmdList, FTextureRHIParamRef New
 	DecodeRenderTargetMode(Mode, ColorLoadAction, ColorStoreAction, DepthLoadAction, DepthStoreAction, StencilLoadAction, StencilStoreAction, DepthStencilAccess);
 
 	//make these rendertargets safely writable
+	PRAGMA_DISABLE_DEPRECATION_WARNINGS
 	if (bWritableBarrier)
 	{
 		TransitionSetRenderTargetsHelper(RHICmdList, NewRenderTarget, NewDepthStencilTarget, DepthStencilAccess);
@@ -466,7 +469,6 @@ inline void SetRenderTarget(FRHICommandList& RHICmdList, FTextureRHIParamRef New
 	// now make the FRHISetRenderTargetsInfo that encapsulates all of the info
 	FRHIRenderTargetView ColorView(NewRenderTarget, 0, -1, ColorLoadAction, ColorStoreAction);
 	FRHISetRenderTargetsInfo Info(1, &ColorView, FRHIDepthRenderTargetView(NewDepthStencilTarget, DepthLoadAction, DepthStoreAction, StencilLoadAction, StencilStoreAction, DepthStencilAccess));
-	PRAGMA_DISABLE_DEPRECATION_WARNINGS
 	RHICmdList.SetRenderTargetsAndClear(Info);
 	PRAGMA_ENABLE_DEPRECATION_WARNINGS
 }
@@ -479,11 +481,11 @@ inline void SetRenderTarget(FRHICommandList& RHICmdList, FTextureRHIParamRef New
 	FRHIDepthRenderTargetView DepthRTV(NewDepthStencilTarget, ERenderTargetLoadAction::ELoad, ERenderTargetStoreAction::EStore);
 
 	//make these rendertargets safely writable
+	PRAGMA_DISABLE_DEPRECATION_WARNINGS
 	if (bWritableBarrier)
 	{
 		TransitionSetRenderTargetsHelper(RHICmdList, NewRenderTarget, NewDepthStencilTarget, FExclusiveDepthStencil::DepthWrite_StencilWrite);
 	}
-	PRAGMA_DISABLE_DEPRECATION_WARNINGS
 	RHICmdList.SetRenderTargets(1, &RTV, &DepthRTV, 0, nullptr);
 	PRAGMA_ENABLE_DEPRECATION_WARNINGS
 }
@@ -508,13 +510,13 @@ inline void SetRenderTargets(
 	}
 
 	//make these rendertargets safely writable
+	PRAGMA_DISABLE_DEPRECATION_WARNINGS
 	if (bWritableBarrier)
 	{
 		TransitionSetRenderTargetsHelper(RHICmdList, NewNumSimultaneousRenderTargets, NewRenderTargetsRHI, NewDepthStencilTargetRHI, FExclusiveDepthStencil::DepthWrite_StencilWrite);
 	}
 
 	FRHIDepthRenderTargetView DepthRTV(NewDepthStencilTargetRHI, ERenderTargetLoadAction::ELoad, ERenderTargetStoreAction::EStore);
-	PRAGMA_DISABLE_DEPRECATION_WARNINGS
 	RHICmdList.SetRenderTargets(NewNumSimultaneousRenderTargets, RTVs, &DepthRTV, NewNumUAVs, UAVs);
 	PRAGMA_ENABLE_DEPRECATION_WARNINGS
 }
@@ -543,13 +545,13 @@ inline void SetRenderTargets(
 	}
 
 	//make these rendertargets safely writable
+	PRAGMA_DISABLE_DEPRECATION_WARNINGS
 	if (bWritableBarrier)
 	{
 		TransitionSetRenderTargetsHelper(RHICmdList, NewNumSimultaneousRenderTargets, NewRenderTargetsRHI, NewDepthStencilTargetRHI, DepthStencilAccess);
 	}
 
 	FRHIDepthRenderTargetView DepthRTV(NewDepthStencilTargetRHI, DepthLoadAction, DepthStoreAction, StencilLoadAction, StencilStoreAction, DepthStencilAccess);
-	PRAGMA_DISABLE_DEPRECATION_WARNINGS
 	RHICmdList.SetRenderTargets(NewNumSimultaneousRenderTargets, RTVs, &DepthRTV, 0, nullptr);
 	PRAGMA_ENABLE_DEPRECATION_WARNINGS
 }
