Implemented basic async ray tracing pipeline creation.

Render thread now immediately returns after requesting a pipeline and actual compilation happens on a separate task.
RHI command list that's associated with RTPSO request is set to depend on the creation task.

diff --git a/Engine/Source/Runtime/RHI/Private/PipelineStateCache.cpp b/Engine/Source/Runtime/RHI/Private/PipelineStateCache.cpp
--- a/Engine/Source/Runtime/RHI/Private/PipelineStateCache.cpp
+++ b/Engine/Source/Runtime/RHI/Private/PipelineStateCache.cpp
@@ -266,6 +266,14 @@ public:
 #endif
 };
 
+//extern RHI_API FRHIRayTracingPipelineState* GetRHIRayTracingPipelineState(FRayTracingPipelineState* PipelineState);
+RHI_API FRHIRayTracingPipelineState* GetRHIRayTracingPipelineState(FRayTracingPipelineState* PipelineState)
+{
+	ensure(PipelineState->RHIPipeline);
+	PipelineState->CompletionEvent = nullptr;
+	return PipelineState->RHIPipeline;
+}
+
 #endif // RHI_RAYTRACING
 
 void SetGraphicsPipelineState(FRHICommandList& RHICmdList, const FGraphicsPipelineStateInitializer& Initializer, EApplyRendertargetOption ApplyFlags)
@@ -583,7 +591,6 @@ FAutoConsoleTaskPriority CPrio_FCompilePipelineStateTask(
 // This cache uses a single internal lock and therefore is not designed for highly concurrent operations.
 class FRayTracingPipelineCache
 {
-	// #dxr_todo UE-68234: This needs to support fully asynchronous, non-blocking pipeline creation with explicit completion query mechanism.
 	// #dxr_todo: UE-72566 Should support eviction of stale pipelines.
 	// #dxr_todo: we will likely also need an explicit ray tracing pipeline sub-object cache to hold closest hit shaders
 
@@ -897,12 +904,102 @@ FComputePipelineState* PipelineStateCache::GetAndOrCreateComputePipelineState(FR
 }
 
 #if RHI_RAYTRACING
-FRHIRayTracingPipelineState* PipelineStateCache::GetAndOrCreateRayTracingPipelineState(const FRayTracingPipelineStateInitializer& Initializer)
+
+class FCompileRayTracingPipelineStateTask
+{
+public:
+
+	UE_NONCOPYABLE(FCompileRayTracingPipelineStateTask)
+
+	FPipelineState* Pipeline;
+
+	FRayTracingPipelineStateInitializer Initializer;
+
+	FCompileRayTracingPipelineStateTask(FPipelineState* InPipeline, const FRayTracingPipelineStateInitializer& InInitializer)
+		: Pipeline(InPipeline)
+		, Initializer(InInitializer)
+	{
+		// Copy all referenced shaders and AddRef them while the task is alive
+
+		RayGenTable   = CopyShaderTable(InInitializer.GetRayGenTable());
+		MissTable     = CopyShaderTable(InInitializer.GetMissTable());
+		HitGroupTable = CopyShaderTable(InInitializer.GetHitGroupTable());
+		CallableTable = CopyShaderTable(InInitializer.GetCallableTable());
+
+		// Point initializer to shader tables owned by this task
+
+		Initializer.SetRayGenShaderTable(RayGenTable, InInitializer.GetRayGenHash());
+		Initializer.SetMissShaderTable(MissTable, InInitializer.GetRayMissHash());
+		Initializer.SetHitGroupTable(HitGroupTable, InInitializer.GetHitGroupHash());
+		Initializer.SetCallableTable(CallableTable, InInitializer.GetCallableHash());
+	}
+
+	static ESubsequentsMode::Type GetSubsequentsMode() { return ESubsequentsMode::TrackSubsequents; }
+
+	void DoTask(ENamedThreads::Type CurrentThread, const FGraphEventRef& MyCompletionGraphEvent)
+	{
+		FRayTracingPipelineState* RayTracingPipeline = static_cast<FRayTracingPipelineState*>(Pipeline);
+		RayTracingPipeline->RHIPipeline = RHICreateRayTracingPipelineState(Initializer);
+
+		// References to shaders no longer need to be held by this task
+
+		ReleaseShaders(CallableTable);
+		ReleaseShaders(HitGroupTable);
+		ReleaseShaders(MissTable);
+		ReleaseShaders(RayGenTable);
+
+		Initializer = FRayTracingPipelineStateInitializer();
+	}
+
+	FORCEINLINE TStatId GetStatId() const
+	{
+		RETURN_QUICK_DECLARE_CYCLE_STAT(FCompileRayTracingPipelineStateTask, STATGROUP_TaskGraphTasks);
+	}
+
+	ENamedThreads::Type GetDesiredThread()
+	{
+		return CPrio_FCompilePipelineStateTask.Get();
+	}
+
+private:
+
+	void AddRefShaders(TArray<FRayTracingShaderRHIParamRef>& ShaderTable)
+	{
+		for (FRayTracingShaderRHIParamRef Ptr : ShaderTable)
+		{
+			Ptr->AddRef();
+		}
+	}
+
+	void ReleaseShaders(TArray<FRayTracingShaderRHIParamRef>& ShaderTable)
+	{
+		for (FRayTracingShaderRHIParamRef Ptr : ShaderTable)
+		{
+			Ptr->Release();
+		}
+	}
+
+	TArray<FRayTracingShaderRHIParamRef> CopyShaderTable(const TArrayView<const FRayTracingShaderRHIParamRef>& Source)
+	{
+		TArray<FRayTracingShaderRHIParamRef> Result(Source.GetData(), Source.Num());
+		AddRefShaders(Result);
+		return Result;
+	}
+
+	TArray<FRayTracingShaderRHIParamRef> RayGenTable;
+	TArray<FRayTracingShaderRHIParamRef> MissTable;
+	TArray<FRayTracingShaderRHIParamRef> HitGroupTable;
+	TArray<FRayTracingShaderRHIParamRef> CallableTable;
+};
+
+FRayTracingPipelineState* PipelineStateCache::GetAndOrCreateRayTracingPipelineState(FRHICommandList& RHICmdList, const FRayTracingPipelineStateInitializer& Initializer)
 {
 	LLM_SCOPE(ELLMTag::PSO);
 
 	check(IsInRenderingThread() || IsInParallelRenderingThread());
 
+	const bool DoAsyncCompile = IsAsyncCompilationAllowed(RHICmdList);
+
 	FRayTracingPipelineState* OutCachedState = nullptr;
 
 	bool bWasFound = GRayTracingPipelineCache.Find(Initializer, OutCachedState);
@@ -910,11 +1007,21 @@ FRHIRayTracingPipelineState* PipelineStateCache::GetAndOrCreateRayTracingPipelin
 	if (bWasFound == false)
 	{
 		// #dxr_todo UE-68235: RT PSO disk caching
-		// #dxr_todo UE-68234: asynchronous PSO creation
 
 		OutCachedState = new FRayTracingPipelineState();
 
-		OutCachedState->RHIPipeline = RHICreateRayTracingPipelineState(Initializer);
+		if (DoAsyncCompile)
+		{
+			OutCachedState->CompletionEvent = TGraphTask<FCompileRayTracingPipelineStateTask>::CreateTask().ConstructAndDispatchWhenReady(
+				OutCachedState,
+				Initializer);
+
+			RHICmdList.AddDispatchPrerequisite(OutCachedState->CompletionEvent);
+		}
+		else
+		{
+			OutCachedState->RHIPipeline = RHICreateRayTracingPipelineState(Initializer);
+		}
 
 		GRayTracingPipelineCache.Add(Initializer, OutCachedState);
 	}
@@ -925,8 +1032,7 @@ FRHIRayTracingPipelineState* PipelineStateCache::GetAndOrCreateRayTracingPipelin
 #endif
 	}
 
-	// return the state pointer
-	return OutCachedState->RHIPipeline;
+	return OutCachedState;
 }
 #endif // RHI_RAYTRACING
 
diff --git a/Engine/Source/Runtime/RHI/Public/PipelineStateCache.h b/Engine/Source/Runtime/RHI/Public/PipelineStateCache.h
--- a/Engine/Source/Runtime/RHI/Public/PipelineStateCache.h
+++ b/Engine/Source/Runtime/RHI/Public/PipelineStateCache.h
@@ -12,7 +12,7 @@
 
 class FComputePipelineState;
 class FGraphicsPipelineState;
-class FRHIRayTracingPipelineState;
+class FRayTracingPipelineState;
 
 // Utility flags for modifying render target behavior on a PSO
 enum class EApplyRendertargetOption : int
@@ -36,8 +36,8 @@ namespace PipelineStateCache
 	extern RHI_API FRHIVertexDeclaration*	GetOrCreateVertexDeclaration(const FVertexDeclarationElementList& Elements);
 
 #if RHI_RAYTRACING
-	extern RHI_API FRHIRayTracingPipelineState* GetAndOrCreateRayTracingPipelineState(const FRayTracingPipelineStateInitializer& Initializer);
-#endif
+	extern RHI_API FRayTracingPipelineState* GetAndOrCreateRayTracingPipelineState(FRHICommandList& RHICmdList, const FRayTracingPipelineStateInitializer& Initializer);
+#endif // RHI_RAYTRACING
 
 	/* Evicts unused state entries based on r.pso.evictiontime time. Called in RHICommandList::BeginFrame */
 	extern RHI_API void FlushResources();
diff --git a/Engine/Source/Runtime/RHI/Public/RHICommandList.h b/Engine/Source/Runtime/RHI/Public/RHICommandList.h
--- a/Engine/Source/Runtime/RHI/Public/RHICommandList.h
+++ b/Engine/Source/Runtime/RHI/Public/RHICommandList.h
@@ -1944,14 +1944,14 @@ struct FRHICommandRayTraceIntersection final : public FRHICommand<FRHICommandRay
 
 struct FRHICommandRayTraceDispatch final : public FRHICommand<FRHICommandRayTraceDispatch>
 {
-	FRayTracingPipelineStateRHIParamRef Pipeline;
+	FRayTracingPipelineState* Pipeline;
 	FRayTracingSceneRHIParamRef Scene;
 	FRayTracingShaderBindings GlobalResourceBindings;
 	FRayTracingShaderRHIParamRef RayGenShader;
 	uint32 Width;
 	uint32 Height;
 
-	FRHICommandRayTraceDispatch(FRayTracingPipelineStateRHIParamRef InPipeline, FRayTracingShaderRHIParamRef InRayGenShader, FRayTracingSceneRHIParamRef InScene, const FRayTracingShaderBindings& InGlobalResourceBindings, uint32 InWidth, uint32 InHeight)
+	FRHICommandRayTraceDispatch(FRayTracingPipelineState* InPipeline, FRayTracingShaderRHIParamRef InRayGenShader, FRayTracingSceneRHIParamRef InScene, const FRayTracingShaderBindings& InGlobalResourceBindings, uint32 InWidth, uint32 InHeight)
 		: Pipeline(InPipeline)
 		, Scene(InScene)
 		, GlobalResourceBindings(InGlobalResourceBindings)
@@ -1975,7 +1975,7 @@ struct FRHICommandSetRayTracingBindings final : public FRHICommand<FRHICommandSe
 	uint32 InstanceIndex;
 	uint32 SegmentIndex;
 	uint32 ShaderSlot;
-	FRayTracingPipelineStateRHIParamRef Pipeline;
+	FRayTracingPipelineState* Pipeline;
 	uint32 ShaderIndex;
 	uint32 NumUniformBuffers;
 	const FUniformBufferRHIParamRef* UniformBuffers; // Pointer to an array of uniform buffers, allocated inline within the command list
@@ -1983,7 +1983,7 @@ struct FRHICommandSetRayTracingBindings final : public FRHICommand<FRHICommandSe
 	EBindingType BindingType;
 
 	FRHICommandSetRayTracingBindings(FRayTracingSceneRHIParamRef InScene, uint32 InInstanceIndex, uint32 InSegmentIndex, uint32 InShaderSlot,
-		FRayTracingPipelineStateRHIParamRef InPipeline, uint32 InHitGroupIndex, uint32 InNumUniformBuffers, const FUniformBufferRHIParamRef* InUniformBuffers,
+		FRayTracingPipelineState* InPipeline, uint32 InHitGroupIndex, uint32 InNumUniformBuffers, const FUniformBufferRHIParamRef* InUniformBuffers,
 		uint32 InUserData)
 		: Scene(InScene)
 		, InstanceIndex(InInstanceIndex)
@@ -1999,7 +1999,7 @@ struct FRHICommandSetRayTracingBindings final : public FRHICommand<FRHICommandSe
 	}
 
 	FRHICommandSetRayTracingBindings(FRayTracingSceneRHIParamRef InScene, uint32 InShaderSlot,
-		FRayTracingPipelineStateRHIParamRef InPipeline, uint32 InShaderIndex, 
+		FRayTracingPipelineState* InPipeline, uint32 InShaderIndex,
 		uint32 InNumUniformBuffers, const FUniformBufferRHIParamRef* InUniformBuffers,
 		uint32 InUserData)
 		: Scene(InScene)
@@ -3065,11 +3065,12 @@ public:
 		}
 	}
 
-	FORCEINLINE_DEBUGGABLE void RayTraceDispatch(FRayTracingPipelineStateRHIParamRef Pipeline, FRayTracingShaderRHIParamRef RayGenShader, FRayTracingSceneRHIParamRef Scene, const FRayTracingShaderBindings& GlobalResourceBindings, uint32 Width, uint32 Height)
+	FORCEINLINE_DEBUGGABLE void RayTraceDispatch(FRayTracingPipelineState* Pipeline, FRayTracingShaderRHIParamRef RayGenShader, FRayTracingSceneRHIParamRef Scene, const FRayTracingShaderBindings& GlobalResourceBindings, uint32 Width, uint32 Height)
 	{
 		if (Bypass())
 		{
-			GetContext().RHIRayTraceDispatch(Pipeline, RayGenShader, Scene, GlobalResourceBindings, Width, Height);
+			extern RHI_API FRHIRayTracingPipelineState* GetRHIRayTracingPipelineState(FRayTracingPipelineState*);
+			GetContext().RHIRayTraceDispatch(GetRHIRayTracingPipelineState(Pipeline), RayGenShader, Scene, GlobalResourceBindings, Width, Height);
 		}
 		else
 		{
@@ -3079,13 +3080,14 @@ public:
 
 	FORCEINLINE_DEBUGGABLE void SetRayTracingHitGroup(
 		FRayTracingSceneRHIParamRef Scene, uint32 InstanceIndex, uint32 SegmentIndex, uint32 ShaderSlot,
-		FRayTracingPipelineStateRHIParamRef Pipeline, uint32 HitGroupIndex,
+		FRayTracingPipelineState* Pipeline, uint32 HitGroupIndex,
 		uint32 NumUniformBuffers, const FUniformBufferRHIParamRef* UniformBuffers,
 		uint32 UserData)
 	{
 		if (Bypass())
 		{
-			GetContext().RHISetRayTracingHitGroup(Scene, InstanceIndex, SegmentIndex, ShaderSlot, Pipeline, HitGroupIndex, NumUniformBuffers, UniformBuffers, UserData);
+			extern RHI_API FRHIRayTracingPipelineState* GetRHIRayTracingPipelineState(FRayTracingPipelineState*);
+			GetContext().RHISetRayTracingHitGroup(Scene, InstanceIndex, SegmentIndex, ShaderSlot, GetRHIRayTracingPipelineState(Pipeline), HitGroupIndex, NumUniformBuffers, UniformBuffers, UserData);
 		}
 		else
 		{
@@ -3105,13 +3107,14 @@ public:
 
 	FORCEINLINE_DEBUGGABLE void SetRayTracingCallableShader(
 		FRayTracingSceneRHIParamRef Scene, uint32 ShaderSlotInScene,
-		FRayTracingPipelineStateRHIParamRef Pipeline, uint32 ShaderIndexInPipeline,
+		FRayTracingPipelineState* Pipeline, uint32 ShaderIndexInPipeline,
 		uint32 NumUniformBuffers, const FUniformBufferRHIParamRef* UniformBuffers,
 		uint32 UserData)
 	{
 		if (Bypass())
 		{
-			GetContext().RHISetRayTracingCallableShader(Scene, ShaderSlotInScene, Pipeline, ShaderIndexInPipeline, NumUniformBuffers, UniformBuffers, UserData);
+			extern RHI_API FRHIRayTracingPipelineState* GetRHIRayTracingPipelineState(FRayTracingPipelineState*);
+			GetContext().RHISetRayTracingCallableShader(Scene, ShaderSlotInScene, GetRHIRayTracingPipelineState(Pipeline), ShaderIndexInPipeline, NumUniformBuffers, UniformBuffers, UserData);
 		}
 		else
 		{
diff --git a/Engine/Source/Runtime/RHI/Public/RHICommandListCommandExecutes.inl b/Engine/Source/Runtime/RHI/Public/RHICommandListCommandExecutes.inl
--- a/Engine/Source/Runtime/RHI/Public/RHICommandListCommandExecutes.inl
+++ b/Engine/Source/Runtime/RHI/Public/RHICommandListCommandExecutes.inl
@@ -578,19 +578,21 @@ void FRHICommandRayTraceIntersection::Execute(FRHICommandListBase& CmdList)
 void FRHICommandRayTraceDispatch::Execute(FRHICommandListBase& CmdList)
 {
 	RHISTAT(RayTraceDispatch);
-	INTERNAL_DECORATOR(RHIRayTraceDispatch)(Pipeline, RayGenShader, Scene, GlobalResourceBindings, Width, Height);
+	extern RHI_API FRHIRayTracingPipelineState* GetRHIRayTracingPipelineState(FRayTracingPipelineState*);
+	INTERNAL_DECORATOR(RHIRayTraceDispatch)(GetRHIRayTracingPipelineState(Pipeline), RayGenShader, Scene, GlobalResourceBindings, Width, Height);
 }
 
 void FRHICommandSetRayTracingBindings::Execute(FRHICommandListBase& CmdList)
 {
 	RHISTAT(SetRayTracingHitGroup);
+	extern RHI_API FRHIRayTracingPipelineState* GetRHIRayTracingPipelineState(FRayTracingPipelineState*);
 	if (BindingType == EBindingType_HitGroup)
 	{
-		INTERNAL_DECORATOR(RHISetRayTracingHitGroup)(Scene, InstanceIndex, SegmentIndex, ShaderSlot, Pipeline, ShaderIndex, NumUniformBuffers, UniformBuffers, UserData);
+		INTERNAL_DECORATOR(RHISetRayTracingHitGroup)(Scene, InstanceIndex, SegmentIndex, ShaderSlot, GetRHIRayTracingPipelineState(Pipeline), ShaderIndex, NumUniformBuffers, UniformBuffers, UserData);
 	}
 	else
 	{
-		INTERNAL_DECORATOR(RHISetRayTracingCallableShader)(Scene, ShaderSlot, Pipeline, ShaderIndex, NumUniformBuffers, UniformBuffers, UserData);
+		INTERNAL_DECORATOR(RHISetRayTracingCallableShader)(Scene, ShaderSlot, GetRHIRayTracingPipelineState(Pipeline), ShaderIndex, NumUniformBuffers, UniformBuffers, UserData);
 	}
 }
 
diff --git a/Engine/Source/Runtime/RHI/Public/RHIResources.h b/Engine/Source/Runtime/RHI/Public/RHIResources.h
--- a/Engine/Source/Runtime/RHI/Public/RHIResources.h
+++ b/Engine/Source/Runtime/RHI/Public/RHIResources.h
@@ -2195,35 +2195,35 @@ public:
 	const TArrayView<const FRayTracingShaderRHIParamRef>& GetCallableTable() const { return CallableTable; }
 
 	// Shaders used as entry point to ray tracing work. At least one RayGen shader must be provided.
-	void SetRayGenShaderTable(const TArrayView<const FRayTracingShaderRHIParamRef>& InRayGenShaders)
+	void SetRayGenShaderTable(const TArrayView<const FRayTracingShaderRHIParamRef>& InRayGenShaders, uint64 Hash = 0)
 	{
 		RayGenTable = InRayGenShaders;
-		RayGenHash = ComputeShaderTableHash(InRayGenShaders);
+		RayGenHash = Hash ? Hash : ComputeShaderTableHash(InRayGenShaders);
 	}
 
 	// Shaders that will be invoked if a ray misses all geometry.
 	// If this table is empty, then a built-in default miss shader will be used that sets HitT member of FDefaultPayload to -1.
 	// Desired miss shader can be selected by providing MissShaderIndex to TraceRay() function.
-	void SetMissShaderTable(const TArrayView<const FRayTracingShaderRHIParamRef>& InMissShaders)
+	void SetMissShaderTable(const TArrayView<const FRayTracingShaderRHIParamRef>& InMissShaders, uint64 Hash = 0)
 	{
 		MissTable = InMissShaders;
-		MissHash = ComputeShaderTableHash(InMissShaders);
+		MissHash = Hash ? Hash : ComputeShaderTableHash(InMissShaders);
 	}
 
 	// Shaders that will be invoked when ray intersects geometry.
 	// If this table is empty, then a built-in default shader will be used for all geometry, using FDefaultPayload.
-	void SetHitGroupTable(const TArrayView<const FRayTracingShaderRHIParamRef>& InHitGroups)
+	void SetHitGroupTable(const TArrayView<const FRayTracingShaderRHIParamRef>& InHitGroups, uint64 Hash = 0)
 	{
 		HitGroupTable = InHitGroups;
-		HitGroupHash = ComputeShaderTableHash(HitGroupTable);
+		HitGroupHash = Hash ? Hash : ComputeShaderTableHash(HitGroupTable);
 	}
 
 	// Shaders that can be explicitly invoked from RayGen shaders by their Shader Binding Table (SBT) index.
 	// SetRayTracingCallableShader() command must be used to fill SBT slots before a shader can be called.
-	void SetCallableTable(const TArrayView<const FRayTracingShaderRHIParamRef>& InCallableShaders)
+	void SetCallableTable(const TArrayView<const FRayTracingShaderRHIParamRef>& InCallableShaders, uint64 Hash = 0)
 	{
 		CallableTable = InCallableShaders;
-		CallableHash = ComputeShaderTableHash(CallableTable);
+		CallableHash = Hash ? Hash : ComputeShaderTableHash(CallableTable);
 	}
 
 	uint64 GetHitGroupHash() const { return HitGroupHash; }
diff --git a/Engine/Source/Runtime/Renderer/Private/DeferredShadingRenderer.h b/Engine/Source/Runtime/Renderer/Private/DeferredShadingRenderer.h
--- a/Engine/Source/Runtime/Renderer/Private/DeferredShadingRenderer.h
+++ b/Engine/Source/Runtime/Renderer/Private/DeferredShadingRenderer.h
@@ -561,8 +561,8 @@ private:
 
 	bool GatherRayTracingWorldInstances(FRHICommandListImmediate& RHICmdList);
 	bool DispatchRayTracingWorldUpdates(FRHICommandListImmediate& RHICmdList);
-	FRHIRayTracingPipelineState* BindRayTracingMaterialPipeline(FRHICommandList& RHICmdList, const FViewInfo& View, const TArrayView<const FRayTracingShaderRHIParamRef>& RayGenShaderTable, FRayTracingShaderRHIParamRef MissShader, FRayTracingShaderRHIParamRef DefaultClosestHitShader);
-	FRHIRayTracingPipelineState* BindRayTracingDeferredMaterialGatherPipeline(FRHICommandList& RHICmdList, const FViewInfo& View, FRayTracingShaderRHIParamRef RayGenShader);
+	FRayTracingPipelineState* BindRayTracingMaterialPipeline(FRHICommandList& RHICmdList, const FViewInfo& View, const TArrayView<const FRayTracingShaderRHIParamRef>& RayGenShaderTable, FRayTracingShaderRHIParamRef MissShader, FRayTracingShaderRHIParamRef DefaultClosestHitShader);
+	FRayTracingPipelineState* BindRayTracingDeferredMaterialGatherPipeline(FRHICommandList& RHICmdList, const FViewInfo& View, FRayTracingShaderRHIParamRef RayGenShader);
 
 	// #dxr_todo: UE-72565: refactor ray tracing effects to not be member functions of DeferredShadingRenderer. Register each effect at startup and just loop over them automatically
 	static void PrepareRayTracingReflections(const FViewInfo& View, TArray<FRayTracingShaderRHIParamRef>& OutRayGenShaders);
diff --git a/Engine/Source/Runtime/Renderer/Private/MeshPassProcessor.cpp b/Engine/Source/Runtime/Renderer/Private/MeshPassProcessor.cpp
--- a/Engine/Source/Runtime/Renderer/Private/MeshPassProcessor.cpp
+++ b/Engine/Source/Runtime/Renderer/Private/MeshPassProcessor.cpp
@@ -272,7 +272,7 @@ void FMeshDrawShaderBindings::SetRayTracingShaderBindingsForHitGroup(
 	FRayTracingSceneRHIParamRef Scene, 
 	uint32 InstanceIndex, 
 	uint32 SegmentIndex,
-	FRayTracingPipelineStateRHIParamRef PipelineState, 
+	FRayTracingPipelineState* PipelineState,
 	uint32 HitGroupIndex,
 	uint32 ShaderSlot) const
 {
diff --git a/Engine/Source/Runtime/Renderer/Private/RayTracing/RayTracingAmbientOcclusion.cpp b/Engine/Source/Runtime/Renderer/Private/RayTracing/RayTracingAmbientOcclusion.cpp
--- a/Engine/Source/Runtime/Renderer/Private/RayTracing/RayTracingAmbientOcclusion.cpp
+++ b/Engine/Source/Runtime/Renderer/Private/RayTracing/RayTracingAmbientOcclusion.cpp
@@ -208,7 +208,7 @@ void FDeferredShadingSceneRenderer::RenderRayTracingAmbientOcclusion(
 		SetShaderParameters(GlobalResources, *RayGenerationShader, *PassParameters);
 
 		// TODO: Provide material support for opacity mask
-		FRHIRayTracingPipelineState* Pipeline = View.RayTracingMaterialPipeline;
+		FRayTracingPipelineState* Pipeline = View.RayTracingMaterialPipeline;
 		if (CVarRayTracingAmbientOcclusionEnableMaterials.GetValueOnRenderThread() == 0)
 		{
 			// Declare default pipeline
@@ -224,7 +224,7 @@ void FDeferredShadingSceneRenderer::RenderRayTracingAmbientOcclusion(
 			Initializer.SetHitGroupTable(HitGroupTable);
 			Initializer.bAllowHitGroupIndexing = false; // Use the same hit shader for all geometry in the scene by disabling SBT indexing.
 
-			Pipeline = PipelineStateCache::GetAndOrCreateRayTracingPipelineState(Initializer);
+			Pipeline = PipelineStateCache::GetAndOrCreateRayTracingPipelineState(RHICmdList, Initializer);
 		}
 
 		FRayTracingSceneRHIParamRef RayTracingSceneRHI = View.RayTracingScene.RayTracingSceneRHI;
diff --git a/Engine/Source/Runtime/Renderer/Private/RayTracing/RayTracingBarycentrics.cpp b/Engine/Source/Runtime/Renderer/Private/RayTracing/RayTracingBarycentrics.cpp
--- a/Engine/Source/Runtime/Renderer/Private/RayTracing/RayTracingBarycentrics.cpp
+++ b/Engine/Source/Runtime/Renderer/Private/RayTracing/RayTracingBarycentrics.cpp
@@ -96,7 +96,7 @@ void FDeferredShadingSceneRenderer::RenderRayTracingBarycentrics(FRHICommandList
 	Initializer.SetHitGroupTable(HitGroupTable);
 	Initializer.bAllowHitGroupIndexing = false; // Use the same hit shader for all geometry in the scene by disabling SBT indexing.
 
-	FRHIRayTracingPipelineState* Pipeline = PipelineStateCache::GetAndOrCreateRayTracingPipelineState(Initializer); 
+	FRayTracingPipelineState* Pipeline = PipelineStateCache::GetAndOrCreateRayTracingPipelineState(RHICmdList, Initializer);
 
 	FRayTracingSceneRHIParamRef RayTracingSceneRHI = View.RayTracingScene.RayTracingSceneRHI;
 
diff --git a/Engine/Source/Runtime/Renderer/Private/RayTracing/RayTracingDebug.cpp b/Engine/Source/Runtime/Renderer/Private/RayTracing/RayTracingDebug.cpp
--- a/Engine/Source/Runtime/Renderer/Private/RayTracing/RayTracingDebug.cpp
+++ b/Engine/Source/Runtime/Renderer/Private/RayTracing/RayTracingDebug.cpp
@@ -124,7 +124,7 @@ void FDeferredShadingSceneRenderer::RenderRayTracingDebug(FRHICommandListImmedia
 	auto ClosestHitShader = ShaderMap->GetShader<FRayTracingDebugCHS>();
 	auto MissShader = ShaderMap->GetShader<FRayTracingDebugMS>();
 
-	FRHIRayTracingPipelineState* Pipeline = View.RayTracingMaterialPipeline;
+	FRayTracingPipelineState* Pipeline = View.RayTracingMaterialPipeline;
 
 	FRayTracingSceneRHIParamRef RayTracingSceneRHI = View.RayTracingScene.RayTracingSceneRHI;
 
diff --git a/Engine/Source/Runtime/Renderer/Private/RayTracing/RayTracingDeferredMaterials.cpp b/Engine/Source/Runtime/Renderer/Private/RayTracing/RayTracingDeferredMaterials.cpp
--- a/Engine/Source/Runtime/Renderer/Private/RayTracing/RayTracingDeferredMaterials.cpp
+++ b/Engine/Source/Runtime/Renderer/Private/RayTracing/RayTracingDeferredMaterials.cpp
@@ -40,7 +40,7 @@ class FRayTracingDeferredMaterialMS : public FGlobalShader
 IMPLEMENT_GLOBAL_SHADER(FRayTracingDeferredMaterialCHS, "/Engine/Private/RayTracing/RayTracingDeferredMaterials.usf", "DeferredMaterialCHS", SF_RayHitGroup);
 IMPLEMENT_GLOBAL_SHADER(FRayTracingDeferredMaterialMS,  "/Engine/Private/RayTracing/RayTracingDeferredMaterials.usf", "DeferredMaterialMS",  SF_RayMiss);
 
-FRHIRayTracingPipelineState* FDeferredShadingSceneRenderer::BindRayTracingDeferredMaterialGatherPipeline(FRHICommandList& RHICmdList, const FViewInfo& View, FRayTracingShaderRHIParamRef RayGenShader)
+FRayTracingPipelineState* FDeferredShadingSceneRenderer::BindRayTracingDeferredMaterialGatherPipeline(FRHICommandList& RHICmdList, const FViewInfo& View, FRayTracingShaderRHIParamRef RayGenShader)
 {
 	SCOPE_CYCLE_COUNTER(STAT_BindRayTracingPipeline);
 
@@ -60,7 +60,7 @@ FRHIRayTracingPipelineState* FDeferredShadingSceneRenderer::BindRayTracingDeferr
 	FRayTracingShaderRHIParamRef HitShaderTable[] = { ClosestHitShader->GetRayTracingShader() };
 	Initializer.SetHitGroupTable(HitShaderTable);
 
-	FRHIRayTracingPipelineState* PipelineState = PipelineStateCache::GetAndOrCreateRayTracingPipelineState(Initializer);
+	FRayTracingPipelineState* PipelineState = PipelineStateCache::GetAndOrCreateRayTracingPipelineState(RHICmdList, Initializer);
 
 	const FViewInfo& ReferenceView = Views[0];
 
diff --git a/Engine/Source/Runtime/Renderer/Private/RayTracing/RayTracingMaterialHitShaders.cpp b/Engine/Source/Runtime/Renderer/Private/RayTracing/RayTracingMaterialHitShaders.cpp
--- a/Engine/Source/Runtime/Renderer/Private/RayTracing/RayTracingMaterialHitShaders.cpp
+++ b/Engine/Source/Runtime/Renderer/Private/RayTracing/RayTracingMaterialHitShaders.cpp
@@ -433,7 +433,7 @@ void FRayTracingMeshProcessor::AddMeshBatch(const FMeshBatch& RESTRICT MeshBatch
 	}
 }
 
-FRHIRayTracingPipelineState* FDeferredShadingSceneRenderer::BindRayTracingMaterialPipeline(
+FRayTracingPipelineState* FDeferredShadingSceneRenderer::BindRayTracingMaterialPipeline(
 	FRHICommandList& RHICmdList,
 	const FViewInfo& View,
 	const TArrayView<const FRayTracingShaderRHIParamRef>& RayGenShaderTable,
@@ -443,7 +443,7 @@ FRHIRayTracingPipelineState* FDeferredShadingSceneRenderer::BindRayTracingMateri
 {
 	SCOPE_CYCLE_COUNTER(STAT_BindRayTracingPipeline);
 
-	FRHIRayTracingPipelineState* PipelineState = nullptr;
+	FRayTracingPipelineState* PipelineState = nullptr;
 
 	FRayTracingPipelineStateInitializer Initializer;
 
@@ -473,7 +473,7 @@ FRHIRayTracingPipelineState* FDeferredShadingSceneRenderer::BindRayTracingMateri
 
 	Initializer.SetHitGroupTable(RayTracingMaterialLibrary);
 
-	PipelineState = PipelineStateCache::GetAndOrCreateRayTracingPipelineState(Initializer);
+	PipelineState = PipelineStateCache::GetAndOrCreateRayTracingPipelineState(RHICmdList, Initializer);
 
 	const FViewInfo& ReferenceView = Views[0];
 
diff --git a/Engine/Source/Runtime/Renderer/Private/RayTracing/RayTracingRectLight.cpp b/Engine/Source/Runtime/Renderer/Private/RayTracing/RayTracingRectLight.cpp
--- a/Engine/Source/Runtime/Renderer/Private/RayTracing/RayTracingRectLight.cpp
+++ b/Engine/Source/Runtime/Renderer/Private/RayTracing/RayTracingRectLight.cpp
@@ -272,7 +272,7 @@ public:
 		FRayTracingShaderRHIParamRef RayGenShaderTable[] = { GetRayTracingShader() };
 		Initializer.SetRayGenShaderTable(RayGenShaderTable);
 
-		FRHIRayTracingPipelineState* Pipeline = PipelineStateCache::GetAndOrCreateRayTracingPipelineState(Initializer);
+		FRayTracingPipelineState* Pipeline = PipelineStateCache::GetAndOrCreateRayTracingPipelineState(RHICmdList, Initializer);
 
 		FRayTracingShaderBindingsWriter GlobalResources;
 		GlobalResources.Set(TLASParameter, RayTracingScene.RayTracingSceneRHI->GetShaderResourceView());
diff --git a/Engine/Source/Runtime/Renderer/Private/RayTracing/RayTracingReflections.cpp b/Engine/Source/Runtime/Renderer/Private/RayTracing/RayTracingReflections.cpp
--- a/Engine/Source/Runtime/Renderer/Private/RayTracing/RayTracingReflections.cpp
+++ b/Engine/Source/Runtime/Renderer/Private/RayTracing/RayTracingReflections.cpp
@@ -326,7 +326,7 @@ void FDeferredShadingSceneRenderer::RenderRayTracingReflections(
 				ERenderGraphPassFlags::Compute,
 				[PassParameters, this, &View, RayGenShader, TileAlignedResolution](FRHICommandList& RHICmdList)
 			{
-				FRHIRayTracingPipelineState* Pipeline = BindRayTracingDeferredMaterialGatherPipeline(RHICmdList, View, RayGenShader->GetRayTracingShader());
+				FRayTracingPipelineState* Pipeline = BindRayTracingDeferredMaterialGatherPipeline(RHICmdList, View, RayGenShader->GetRayTracingShader());
 
 				FRayTracingShaderBindingsWriter GlobalResources;
 				SetShaderParameters(GlobalResources, RayGenShader, *PassParameters);
diff --git a/Engine/Source/Runtime/Renderer/Private/RayTracing/RayTracingShadows.cpp b/Engine/Source/Runtime/Renderer/Private/RayTracing/RayTracingShadows.cpp
--- a/Engine/Source/Runtime/Renderer/Private/RayTracing/RayTracingShadows.cpp
+++ b/Engine/Source/Runtime/Renderer/Private/RayTracing/RayTracingShadows.cpp
@@ -233,7 +233,7 @@ void FDeferredShadingSceneRenderer::RenderRayTracingShadows(
 				Initializer.SetHitGroupTable(HitGroupTable);
 				Initializer.bAllowHitGroupIndexing = false; // Use the same hit shader for all geometry in the scene by disabling SBT indexing.
 
-				FRHIRayTracingPipelineState* Pipeline = PipelineStateCache::GetAndOrCreateRayTracingPipelineState(Initializer);
+				FRayTracingPipelineState* Pipeline = PipelineStateCache::GetAndOrCreateRayTracingPipelineState(RHICmdList, Initializer);
 
 				RHICmdList.RayTraceDispatch(Pipeline, RayGenerationShader->GetRayTracingShader(), RayTracingSceneRHI, GlobalResources, Resolution.X, Resolution.Y);
 			}
diff --git a/Engine/Source/Runtime/Renderer/Private/RayTracing/RayTracingTranslucency.cpp b/Engine/Source/Runtime/Renderer/Private/RayTracing/RayTracingTranslucency.cpp
--- a/Engine/Source/Runtime/Renderer/Private/RayTracing/RayTracingTranslucency.cpp
+++ b/Engine/Source/Runtime/Renderer/Private/RayTracing/RayTracingTranslucency.cpp
@@ -407,7 +407,7 @@ void FDeferredShadingSceneRenderer::RenderRayTracingTranslucencyView(
 		ERenderGraphPassFlags::Compute,
 		[PassParameters, this, &View, RayGenShader, RayTracingResolution](FRHICommandList& RHICmdList)
 	{
-		FRHIRayTracingPipelineState* Pipeline = View.RayTracingMaterialPipeline;
+		FRayTracingPipelineState* Pipeline = View.RayTracingMaterialPipeline;
 
 		FRayTracingShaderBindingsWriter GlobalResources;
 		SetShaderParameters(GlobalResources, RayGenShader, *PassParameters);
diff --git a/Engine/Source/Runtime/Renderer/Private/RayTracing/RaytracingSkylight.cpp b/Engine/Source/Runtime/Renderer/Private/RayTracing/RaytracingSkylight.cpp
--- a/Engine/Source/Runtime/Renderer/Private/RayTracing/RaytracingSkylight.cpp
+++ b/Engine/Source/Runtime/Renderer/Private/RayTracing/RaytracingSkylight.cpp
@@ -877,7 +877,7 @@ void FDeferredShadingSceneRenderer::RenderRayTracingSkyLight(
 			FRayTracingShaderBindingsWriter GlobalResources;
 			SetShaderParameters(GlobalResources, *RayGenerationShader, *PassParameters);
 
-			FRHIRayTracingPipelineState* Pipeline = View.RayTracingMaterialPipeline;
+			FRayTracingPipelineState* Pipeline = View.RayTracingMaterialPipeline;
 			if (CVarRayTracingSkyLightEnableMaterials.GetValueOnRenderThread() == 0)
 			{
 				// Declare default pipeline
@@ -893,7 +893,7 @@ void FDeferredShadingSceneRenderer::RenderRayTracingSkyLight(
 				Initializer.SetHitGroupTable(HitGroupTable);
 				Initializer.bAllowHitGroupIndexing = false; // Use the same hit shader for all geometry in the scene by disabling SBT indexing.
 
-				Pipeline = PipelineStateCache::GetAndOrCreateRayTracingPipelineState(Initializer);
+				Pipeline = PipelineStateCache::GetAndOrCreateRayTracingPipelineState(RHICmdList, Initializer);
 			}
 
 			FRayTracingSceneRHIParamRef RayTracingSceneRHI = View.RayTracingScene.RayTracingSceneRHI;
diff --git a/Engine/Source/Runtime/Renderer/Private/SceneRendering.h b/Engine/Source/Runtime/Renderer/Private/SceneRendering.h
--- a/Engine/Source/Runtime/Renderer/Private/SceneRendering.h
+++ b/Engine/Source/Runtime/Renderer/Private/SceneRendering.h
@@ -1079,7 +1079,7 @@ public:
 
 	// Primary pipeline state object to be used with the ray tracing scene for this view.
 	// Material shaders are only available when using this pipeline.
-	FRHIRayTracingPipelineState* RayTracingMaterialPipeline = nullptr;
+	FRayTracingPipelineState* RayTracingMaterialPipeline = nullptr;
 #endif // RHI_RAYTRACING
 
 	/** 
diff --git a/Engine/Source/Runtime/Renderer/Public/MeshPassProcessor.h b/Engine/Source/Runtime/Renderer/Public/MeshPassProcessor.h
--- a/Engine/Source/Runtime/Renderer/Public/MeshPassProcessor.h
+++ b/Engine/Source/Runtime/Renderer/Public/MeshPassProcessor.h
@@ -347,7 +347,7 @@ public:
 	void SetOnCommandListForCompute(FRHICommandList& RHICmdList, FComputeShaderRHIParamRef Shader) const;
 
 #if RHI_RAYTRACING
-	void SetRayTracingShaderBindingsForHitGroup(FRHICommandList& RHICmdList, FRayTracingSceneRHIParamRef Scene, uint32 InstanceIndex, uint32 SegmentIndex, FRayTracingPipelineStateRHIParamRef Pipeline, uint32 HitGroupIndex, uint32 ShaderSlot) const;
+	void SetRayTracingShaderBindingsForHitGroup(FRHICommandList& RHICmdList, FRayTracingSceneRHIParamRef Scene, uint32 InstanceIndex, uint32 SegmentIndex, FRayTracingPipelineState* Pipeline, uint32 HitGroupIndex, uint32 ShaderSlot) const;
 #endif // RHI_RAYTRACING
 
 	/** Returns whether this set of shader bindings can be merged into an instanced draw call with another. */
