Fixed invalid format in FRHICommandListImmediate::CreateShaderResourceView (crash on Mac).

#rnx

diff --git a/Engine/Source/Runtime/RHI/Public/RHICommandList.h b/Engine/Source/Runtime/RHI/Public/RHICommandList.h
--- a/Engine/Source/Runtime/RHI/Public/RHICommandList.h
+++ b/Engine/Source/Runtime/RHI/Public/RHICommandList.h
@@ -3872,7 +3872,7 @@ public:
 	FORCEINLINE FShaderResourceViewRHIRef CreateShaderResourceView(FTextureRHIParamRef Texture, uint8 MipLevel)
 	{
 		LLM_SCOPE(ELLMTag::RHIMisc);
-		const FRHITextureSRVCreateInfo CreateInfo(MipLevel);
+		const FRHITextureSRVCreateInfo CreateInfo(MipLevel, 1, Texture->GetFormat());
 		return GDynamicRHI->RHICreateShaderResourceView_RenderThread(*this, Texture, CreateInfo);
 	}
 	
