DR - More removal of SetRenderTargets
#rnx

diff --git a/Engine/Source/Runtime/RHI/Public/RHIUtilities.h b/Engine/Source/Runtime/RHI/Public/RHIUtilities.h
--- a/Engine/Source/Runtime/RHI/Public/RHIUtilities.h
+++ b/Engine/Source/Runtime/RHI/Public/RHIUtilities.h
@@ -369,7 +369,7 @@ inline void DecodeRenderTargetMode(ESimpleRenderTargetMode Mode, ERenderTargetLo
 	}
 }
 
-UE_DEPRECATED(4.22, "SetRenderTargets API is deprecated; please use RHIBegin/EndRenderPass instead.")
+UE_DEPRECATED(4.22, "TransitionSetRenderTargetsHelper API is deprecated; please use RHITransitionResources directly instead.")
 inline void TransitionSetRenderTargetsHelper(FRHICommandList& RHICmdList, FTextureRHIParamRef NewRenderTarget, FTextureRHIParamRef NewDepthStencilTarget, FExclusiveDepthStencil DepthStencilAccess)
 {
 	int32 TransitionIndex = 0;
@@ -387,7 +387,7 @@ inline void TransitionSetRenderTargetsHelper(FRHICommandList& RHICmdList, FTextu
 	RHICmdList.TransitionResources(EResourceTransitionAccess::EWritable, Transitions, TransitionIndex);
 }
 
-UE_DEPRECATED(4.22, "SetRenderTargets API is deprecated; please use RHIBegin/EndRenderPass instead.")
+UE_DEPRECATED(4.22, "TransitionSetRenderTargetsHelper API is deprecated; please use RHITransitionResources directly instead.")
 inline void TransitionSetRenderTargetsHelper(FRHICommandList& RHICmdList, uint32 NumRenderTargets, const FTextureRHIParamRef* NewRenderTargetsRHI, const FTextureRHIParamRef NewDepthStencilTargetRHI, FExclusiveDepthStencil DepthStencilAccess)
 {
 	FTextureRHIParamRef Transitions[MaxSimultaneousRenderTargets + 1];
@@ -435,7 +435,7 @@ inline void TransitionRenderPassTargets(FRHICommandList& RHICmdList, const FRHIR
 }
 
 /** Helper for the common case of using a single color and depth render target. */
-UE_DEPRECATED(4.22, "SetRenderTargets API is deprecated; please use RHIBegin/EndRenderPass instead.")
+UE_DEPRECATED(4.22, "SetRenderTarget API is deprecated; please use RHIBegin/EndRenderPass instead.")
 inline void SetRenderTarget(FRHICommandList& RHICmdList, FTextureRHIParamRef NewRenderTarget, FTextureRHIParamRef NewDepthStencilTarget, bool bWritableBarrier = false)
 {
 	FRHIRenderTargetView RTV(NewRenderTarget, ERenderTargetLoadAction::ELoad);
@@ -451,76 +451,6 @@ inline void SetRenderTarget(FRHICommandList& RHICmdList, FTextureRHIParamRef New
 	PRAGMA_ENABLE_DEPRECATION_WARNINGS
 }
 
-/** Helper for the common case of using a single color and depth render target. */
-UE_DEPRECATED(4.22, "SetRenderTargets API is deprecated; please use RHIBegin/EndRenderPass instead.")
-inline void SetRenderTarget(FRHICommandList& RHICmdList, FTextureRHIParamRef NewRenderTarget, FTextureRHIParamRef NewDepthStencilTarget, ESimpleRenderTargetMode Mode, FExclusiveDepthStencil DepthStencilAccess = FExclusiveDepthStencil::DepthWrite_StencilWrite, bool bWritableBarrier = false)
-{
-	ERenderTargetLoadAction ColorLoadAction, DepthLoadAction, StencilLoadAction;
-	ERenderTargetStoreAction ColorStoreAction, DepthStoreAction, StencilStoreAction;
-	DecodeRenderTargetMode(Mode, ColorLoadAction, ColorStoreAction, DepthLoadAction, DepthStoreAction, StencilLoadAction, StencilStoreAction, DepthStencilAccess);
-
-	//make these rendertargets safely writable
-	PRAGMA_DISABLE_DEPRECATION_WARNINGS
-	if (bWritableBarrier)
-	{
-		TransitionSetRenderTargetsHelper(RHICmdList, NewRenderTarget, NewDepthStencilTarget, DepthStencilAccess);
-	}
-
-	// now make the FRHISetRenderTargetsInfo that encapsulates all of the info
-	FRHIRenderTargetView ColorView(NewRenderTarget, 0, -1, ColorLoadAction, ColorStoreAction);
-	FRHISetRenderTargetsInfo Info(1, &ColorView, FRHIDepthRenderTargetView(NewDepthStencilTarget, DepthLoadAction, DepthStoreAction, StencilLoadAction, StencilStoreAction, DepthStencilAccess));
-	RHICmdList.SetRenderTargetsAndClear(Info);
-	PRAGMA_ENABLE_DEPRECATION_WARNINGS
-}
-
-/** Helper for the common case of using a single color and depth render target, with a mip index for the color target. */
-UE_DEPRECATED(4.22, "SetRenderTarget API is deprecated; please use RHIBegin/EndRenderPass instead.")
-inline void SetRenderTarget(FRHICommandList& RHICmdList, FTextureRHIParamRef NewRenderTarget, int32 MipIndex, int32 ArraySliceIndex, FTextureRHIParamRef NewDepthStencilTarget, bool bWritableBarrier = false)
-{
-	FRHIRenderTargetView RTV(NewRenderTarget, ERenderTargetLoadAction::ELoad, MipIndex, ArraySliceIndex);
-	FRHIDepthRenderTargetView DepthRTV(NewDepthStencilTarget, ERenderTargetLoadAction::ELoad, ERenderTargetStoreAction::EStore);
-
-	//make these rendertargets safely writable
-	PRAGMA_DISABLE_DEPRECATION_WARNINGS
-	if (bWritableBarrier)
-	{
-		TransitionSetRenderTargetsHelper(RHICmdList, NewRenderTarget, NewDepthStencilTarget, FExclusiveDepthStencil::DepthWrite_StencilWrite);
-	}
-	RHICmdList.SetRenderTargets(1, &RTV, &DepthRTV, 0, nullptr);
-	PRAGMA_ENABLE_DEPRECATION_WARNINGS
-}
-
-/** Helper that converts FTextureRHIParamRef's into FRHIRenderTargetView's. */
-UE_DEPRECATED(4.22, "SetRenderTargets API is deprecated; please use RHIBegin/EndRenderPass instead.")
-inline void SetRenderTargets(
-	FRHICommandList& RHICmdList,
-	uint32 NewNumSimultaneousRenderTargets, 
-	const FTextureRHIParamRef* NewRenderTargetsRHI,
-	FTextureRHIParamRef NewDepthStencilTargetRHI,
-	uint32 NewNumUAVs,
-	const FUnorderedAccessViewRHIParamRef* UAVs,
-	bool bWritableBarrier = false
-	)
-{
-	FRHIRenderTargetView RTVs[MaxSimultaneousRenderTargets];
-	for (uint32 Index = 0; Index < NewNumSimultaneousRenderTargets; Index++)
-	{
-		RTVs[Index] = FRHIRenderTargetView(NewRenderTargetsRHI[Index], ERenderTargetLoadAction::ELoad);
-	
-	}
-
-	//make these rendertargets safely writable
-	PRAGMA_DISABLE_DEPRECATION_WARNINGS
-	if (bWritableBarrier)
-	{
-		TransitionSetRenderTargetsHelper(RHICmdList, NewNumSimultaneousRenderTargets, NewRenderTargetsRHI, NewDepthStencilTargetRHI, FExclusiveDepthStencil::DepthWrite_StencilWrite);
-	}
-
-	FRHIDepthRenderTargetView DepthRTV(NewDepthStencilTargetRHI, ERenderTargetLoadAction::ELoad, ERenderTargetStoreAction::EStore);
-	RHICmdList.SetRenderTargets(NewNumSimultaneousRenderTargets, RTVs, &DepthRTV, NewNumUAVs, UAVs);
-	PRAGMA_ENABLE_DEPRECATION_WARNINGS
-}
-
 /** Helper that converts FTextureRHIParamRef's into FRHIRenderTargetView's. */
 UE_DEPRECATED(4.22, "SetRenderTargets API is deprecated; please use RHIBegin/EndRenderPass instead.")
 inline void SetRenderTargets(
@@ -561,7 +491,9 @@ inline void UnbindRenderTargets(FRHICommandList& RHICmdList)
 {
 	check(RHICmdList.IsOutsideRenderPass());
 PRAGMA_DISABLE_DEPRECATION_WARNINGS
-	SetRenderTarget(RHICmdList, nullptr, nullptr);
+	FRHIRenderTargetView RTV(nullptr, ERenderTargetLoadAction::ENoAction);
+	FRHIDepthRenderTargetView DepthRTV(nullptr, ERenderTargetLoadAction::ENoAction, ERenderTargetStoreAction::ENoAction);
+	RHICmdList.SetRenderTargets(1, &RTV, &DepthRTV, 0, nullptr);
 PRAGMA_ENABLE_DEPRECATION_WARNINGS
 }
 
diff --git a/Engine/Source/Runtime/RenderCore/Private/RenderTargetPool.cpp b/Engine/Source/Runtime/RenderCore/Private/RenderTargetPool.cpp
--- a/Engine/Source/Runtime/RenderCore/Private/RenderTargetPool.cpp
+++ b/Engine/Source/Runtime/RenderCore/Private/RenderTargetPool.cpp
@@ -268,7 +268,9 @@ static void ClobberAllocatedRenderTarget(FRHICommandList& RHICmdList, const FPoo
 
 	if (Out->GetDesc().TargetableFlags & TexCreate_RenderTargetable)
 	{
-		SetRenderTarget(RHICmdList, Out->GetRenderTargetItem().TargetableTexture, FTextureRHIRef());
+		// Needs conversion to Render Passes
+		check(0);
+		//SetRenderTarget(RHICmdList, Out->GetRenderTargetItem().TargetableTexture, FTextureRHIRef());
 		DrawClearQuad(RHICmdList, Color);
 	}
 	else if (Out->GetDesc().TargetableFlags & TexCreate_UAV)
@@ -278,7 +280,9 @@ static void ClobberAllocatedRenderTarget(FRHICommandList& RHICmdList, const FPoo
 
 	if (Desc.TargetableFlags & TexCreate_DepthStencilTargetable)
 	{
-		SetRenderTarget(RHICmdList, FTextureRHIRef(), Out->GetRenderTargetItem().TargetableTexture);
+		// Needs conversion to Render Passes
+		check(0);
+		//SetRenderTarget(RHICmdList, FTextureRHIRef(), Out->GetRenderTargetItem().TargetableTexture);
 		DrawClearQuad(RHICmdList, false, FLinearColor::Black, true, 0.0f, true, 0);
 	}
 }
diff --git a/Engine/Source/Runtime/Renderer/Private/SceneRendering.cpp b/Engine/Source/Runtime/Renderer/Private/SceneRendering.cpp
--- a/Engine/Source/Runtime/Renderer/Private/SceneRendering.cpp
+++ b/Engine/Source/Runtime/Renderer/Private/SceneRendering.cpp
@@ -3280,7 +3280,8 @@ void FRendererModule::CreateAndInitSingleView(FRHICommandListImmediate& RHICmdLi
 	FViewInfo* NewView = new FViewInfo(*ViewInitOptions);
 	ViewFamily->Views.Add(NewView);
 	PRAGMA_DISABLE_DEPRECATION_WARNINGS
-	SetRenderTarget(RHICmdList, ViewFamily->RenderTarget->GetRenderTargetTexture(), nullptr, ESimpleRenderTargetMode::EClearColorExistingDepth);
+	FRHIRenderTargetView RTV(ViewFamily->RenderTarget->GetRenderTargetTexture(), ERenderTargetLoadAction::EClear);
+	RHICmdList.SetRenderTargets(1, &RTV, nullptr, 0, nullptr);
 	PRAGMA_ENABLE_DEPRECATION_WARNINGS
 	FViewInfo* View = (FViewInfo*)ViewFamily->Views[0];
 	View->ViewRect = View->UnscaledViewRect;
@@ -3558,7 +3559,8 @@ static void DisplayInternals(FRHICommandListImmediate& RHICmdList, FViewInfo& In
 		Canvas.SetRenderTargetRect(FIntRect(0, 0, Family->RenderTarget->GetSizeXY().X, Family->RenderTarget->GetSizeXY().Y));
 
 		PRAGMA_DISABLE_DEPRECATION_WARNINGS
-		SetRenderTarget(RHICmdList, Family->RenderTarget->GetRenderTargetTexture(), FTextureRHIRef());
+		FRHIRenderTargetView RTV(Family->RenderTarget->GetRenderTargetTexture(), ERenderTargetLoadAction::ELoad);
+		RHICmdList.SetRenderTargets(1, &RTV, nullptr, 0, nullptr);
 		PRAGMA_ENABLE_DEPRECATION_WARNINGS
 
 		// further down to not intersect with "LIGHTING NEEDS TO BE REBUILT"
