Added FirstArraySlice and NumArraySlices to FRHITextureSRVCreateInfo struct.

diff --git a/Engine/Source/Runtime/D3D12RHI/Private/D3D12View.cpp b/Engine/Source/Runtime/D3D12RHI/Private/D3D12View.cpp
--- a/Engine/Source/Runtime/D3D12RHI/Private/D3D12View.cpp
+++ b/Engine/Source/Runtime/D3D12RHI/Private/D3D12View.cpp
@@ -102,8 +102,8 @@ FShaderResourceViewRHIRef FD3D12DynamicRHI::RHICreateShaderResourceView(FTexture
 		BaseTextureFormat = TextureDesc.Format;
 		BaseTexture = Texture2DArray;
 		SRVDesc.ViewDimension = D3D12_SRV_DIMENSION_TEXTURE2DARRAY;
-		SRVDesc.Texture2DArray.ArraySize = TextureDesc.DepthOrArraySize;
-		SRVDesc.Texture2DArray.FirstArraySlice = 0;
+		SRVDesc.Texture2DArray.ArraySize = (CreateInfo.NumArraySlices == 0 ? TextureDesc.DepthOrArraySize : CreateInfo.NumArraySlices);
+		SRVDesc.Texture2DArray.FirstArraySlice = CreateInfo.FirstArraySlice;
 		SRVDesc.Texture2DArray.MipLevels = CreateInfo.NumMipLevels;
 		SRVDesc.Texture2DArray.MostDetailedMip = CreateInfo.MipLevel;
 	}
diff --git a/Engine/Source/Runtime/RHI/Public/RHI.h b/Engine/Source/Runtime/RHI/Public/RHI.h
--- a/Engine/Source/Runtime/RHI/Public/RHI.h
+++ b/Engine/Source/Runtime/RHI/Public/RHI.h
@@ -1232,6 +1232,17 @@ struct FRHITextureSRVCreateInfo
 		, SRGBOverride(SRGBO_Default)
 		, MipLevel(InMipLevel)
 		, NumMipLevels(InNumMipLevels)
+		, FirstArraySlice(0)
+		, NumArraySlices(0)
+	{}
+
+	explicit FRHITextureSRVCreateInfo(uint8 InMipLevel, uint8 InNumMipLevels, uint32 InFirstArraySlice, uint32 InNumArraySlices, uint8 InFormat = PF_Unknown)
+		: Format(InFormat)
+		, SRGBOverride(SRGBO_Default)
+		, MipLevel(InMipLevel)
+		, NumMipLevels(InNumMipLevels)
+		, FirstArraySlice(InFirstArraySlice)
+		, NumArraySlices(InNumArraySlices)
 	{}
 
 	/** View the texture with a different format. Leave as PF_Unknown to use original format. Useful when sampling stencil */
@@ -1245,6 +1256,12 @@ struct FRHITextureSRVCreateInfo
 
 	/** Create a view to a single, or multiple mip levels */
 	uint8 NumMipLevels;
+
+	/** Specify first array slice index. By default 0. */
+	uint32 FirstArraySlice;
+
+	/** Specify number of array slices. If FirstArraySlice and NumArraySlices are both zero, the SRV is created for all array slices. By default 0. */
+	uint32 NumArraySlices;
 };
 
 // Forward-declaration.
diff --git a/Engine/Source/Runtime/VulkanRHI/Private/VulkanResources.h b/Engine/Source/Runtime/VulkanRHI/Private/VulkanResources.h
--- a/Engine/Source/Runtime/VulkanRHI/Private/VulkanResources.h
+++ b/Engine/Source/Runtime/VulkanRHI/Private/VulkanResources.h
@@ -1322,6 +1322,8 @@ public:
 		, SourceStructuredBuffer(nullptr)
 		, MipLevel(InCreateInfo.MipLevel)
 		, NumMips(InCreateInfo.NumMipLevels)
+		, FirstArraySlice(InCreateInfo.FirstArraySlice)
+		, NumArraySlices(InCreateInfo.NumArraySlices)
 		, Size(0)
 		, SourceBuffer(nullptr)
 		, VolatileLockCounter(MAX_uint32)
@@ -1336,6 +1338,8 @@ public:
 		, SourceStructuredBuffer(InStructuredBuffer)
 		, MipLevel(0)
 		, NumMips(0)
+		, FirstArraySlice(0)
+		, NumArraySlices(0)
 		, Size(InStructuredBuffer->GetSize())
 		, SourceBuffer(nullptr)
 		, VolatileLockCounter(MAX_uint32)
@@ -1362,6 +1366,8 @@ public:
 	FVulkanStructuredBuffer* SourceStructuredBuffer;
 	uint32 MipLevel;
 	uint32 NumMips;
+	uint32 FirstArraySlice;
+	uint32 NumArraySlices;
 
 	~FVulkanShaderResourceView();
 
diff --git a/Engine/Source/Runtime/VulkanRHI/Private/VulkanUAV.cpp b/Engine/Source/Runtime/VulkanRHI/Private/VulkanUAV.cpp
--- a/Engine/Source/Runtime/VulkanRHI/Private/VulkanUAV.cpp
+++ b/Engine/Source/Runtime/VulkanRHI/Private/VulkanUAV.cpp
@@ -124,7 +124,18 @@ void FVulkanShaderResourceView::UpdateView()
 			else if (FRHITexture2DArray* Tex2DArray = SourceTexture->GetTexture2DArray())
 			{
 				FVulkanTexture2DArray* VTex2DArray = ResourceCast(Tex2DArray);
-				TextureView.Create(*Device, VTex2DArray->Surface.Image, VK_IMAGE_VIEW_TYPE_2D_ARRAY, VTex2DArray->Surface.GetPartialAspectMask(), Format, UEToVkTextureFormat(Format, false), MipLevel, NumMips, 0, VTex2DArray->GetSizeZ());
+				TextureView.Create(
+					*Device,
+					VTex2DArray->Surface.Image,
+					VK_IMAGE_VIEW_TYPE_2D_ARRAY,
+					VTex2DArray->Surface.GetPartialAspectMask(),
+					Format,
+					UEToVkTextureFormat(Format, false),
+					MipLevel,
+					NumMips,
+					FirstArraySlice,
+					(NumArraySlices == 0 ? VTex2DArray->GetSizeZ() : NumArraySlices)
+				);
 			}
 			else
 			{
diff --git a/Engine/Source/Runtime/Windows/D3D11RHI/Private/D3D11Texture.cpp b/Engine/Source/Runtime/Windows/D3D11RHI/Private/D3D11Texture.cpp
--- a/Engine/Source/Runtime/Windows/D3D11RHI/Private/D3D11Texture.cpp
+++ b/Engine/Source/Runtime/Windows/D3D11RHI/Private/D3D11Texture.cpp
@@ -1334,8 +1334,8 @@ FShaderResourceViewRHIRef FD3D11DynamicRHI::RHICreateShaderResourceView(FTexture
 		SRVDesc.ViewDimension = D3D11_SRV_DIMENSION_TEXTURE2DARRAY;
 		SRVDesc.Texture2DArray.MostDetailedMip = CreateInfo.MipLevel;
 		SRVDesc.Texture2DArray.MipLevels = CreateInfo.NumMipLevels;
-		SRVDesc.Texture2DArray.FirstArraySlice = 0;
-		SRVDesc.Texture2DArray.ArraySize = TextureDesc.ArraySize;
+		SRVDesc.Texture2DArray.FirstArraySlice = CreateInfo.FirstArraySlice;
+		SRVDesc.Texture2DArray.ArraySize = (CreateInfo.NumArraySlices == 0 ? TextureDesc.ArraySize : CreateInfo.NumArraySlices);
 	}
 	else if (TextureRHI->GetTextureCube() != NULL)
 	{
