DR - Zeromem the hotfix pointers for Vulkan
Minor comment fix
#rnx

diff --git a/Engine/Source/Runtime/RHI/Public/RHIValidation.h b/Engine/Source/Runtime/RHI/Public/RHIValidation.h
--- a/Engine/Source/Runtime/RHI/Public/RHIValidation.h
+++ b/Engine/Source/Runtime/RHI/Public/RHIValidation.h
@@ -1,7 +1,7 @@
 // Copyright 1998-2019 Epic Games, Inc. All Rights Reserved.
 
 /*=============================================================================
-	ValidationRHI.h: Public Valdation RHI definitions.
+	RHIValidation.h: Public Valdation RHI definitions.
 =============================================================================*/
 
 #pragma once 
diff --git a/Engine/Source/Runtime/VulkanRHI/Private/VulkanContext.h b/Engine/Source/Runtime/VulkanRHI/Private/VulkanContext.h
--- a/Engine/Source/Runtime/VulkanRHI/Private/VulkanContext.h
+++ b/Engine/Source/Runtime/VulkanRHI/Private/VulkanContext.h
@@ -367,7 +367,7 @@ public:
 
 	void PrepareParallelFromBase(const FVulkanCommandListContext& BaseContext);
 
-	void* Hotfix;
+	void* Hotfix = nullptr;
 
 protected:
 	FVulkanDynamicRHI* RHI;
diff --git a/Engine/Source/Runtime/VulkanRHI/Private/VulkanDevice.h b/Engine/Source/Runtime/VulkanRHI/Private/VulkanDevice.h
--- a/Engine/Source/Runtime/VulkanRHI/Private/VulkanDevice.h
+++ b/Engine/Source/Runtime/VulkanRHI/Private/VulkanDevice.h
@@ -305,7 +305,7 @@ public:
 	VkSamplerYcbcrConversion CreateSamplerColorConversion(const VkSamplerYcbcrConversionCreateInfo& CreateInfo);
 #endif
 
-	void*	Hotfix;
+	void*	Hotfix = nullptr;
 
 private:
 	void MapFormatSupport(EPixelFormat UEFormat, VkFormat VulkanFormat);
diff --git a/Engine/Source/Runtime/VulkanRHI/Public/VulkanMemory.h b/Engine/Source/Runtime/VulkanRHI/Public/VulkanMemory.h
--- a/Engine/Source/Runtime/VulkanRHI/Public/VulkanMemory.h
+++ b/Engine/Source/Runtime/VulkanRHI/Public/VulkanMemory.h
@@ -850,7 +850,7 @@ namespace VulkanRHI
 		FBufferSuballocation* AllocUniformBuffer(uint32 Size, const void* Contents);
 		void ReleaseUniformBuffer(FBufferSuballocation* UBAlloc);
 
-		void* Hotfix;
+		void* Hotfix = nullptr;
 
 	protected:
 		FDeviceMemoryManager* DeviceMemoryManager;
