DR - Deprecate FGraphicsPipelineStateRHIParamRef, FDepthStencilStateRHIParamRef and FBlendStateRHIParamRef

diff --git a/Engine/Source/Runtime/Apple/MetalRHI/Private/MetalCommands.cpp b/Engine/Source/Runtime/Apple/MetalRHI/Private/MetalCommands.cpp
--- a/Engine/Source/Runtime/Apple/MetalRHI/Private/MetalCommands.cpp
+++ b/Engine/Source/Runtime/Apple/MetalRHI/Private/MetalCommands.cpp
@@ -231,7 +231,7 @@ void FMetalRHICommandContext::RHISetScissorRect(bool bEnable,uint32 MinX,uint32
 	}
 }
 
-void FMetalRHICommandContext::RHISetGraphicsPipelineState(FGraphicsPipelineStateRHIParamRef GraphicsState)
+void FMetalRHICommandContext::RHISetGraphicsPipelineState(FRHIGraphicsPipelineState* GraphicsState)
 {
 	@autoreleasepool {
 		FMetalGraphicsPipelineState* PipelineState = ResourceCast(GraphicsState);
diff --git a/Engine/Source/Runtime/Apple/MetalRHI/Public/MetalRHIContext.h b/Engine/Source/Runtime/Apple/MetalRHI/Public/MetalRHIContext.h
--- a/Engine/Source/Runtime/Apple/MetalRHI/Public/MetalRHIContext.h
+++ b/Engine/Source/Runtime/Apple/MetalRHI/Public/MetalRHIContext.h
@@ -111,7 +111,7 @@ public:
 	// @param MaxY excluding like Win32 RECT
 	virtual void RHISetScissorRect(bool bEnable, uint32 MinX, uint32 MinY, uint32 MaxX, uint32 MaxY) final override;
 	
-	virtual void RHISetGraphicsPipelineState(FGraphicsPipelineStateRHIParamRef GraphicsState) final override;
+	virtual void RHISetGraphicsPipelineState(FRHIGraphicsPipelineState* GraphicsState) final override;
 	
 	/** Set the shader resource view of a surface.  This is used for binding TextureMS parameter types that need a multi sampled view. */
 	virtual void RHISetShaderTexture(FVertexShaderRHIParamRef VertexShader, uint32 TextureIndex, FTextureRHIParamRef NewTexture) final override;
diff --git a/Engine/Source/Runtime/D3D12RHI/Private/D3D12CommandContext.h b/Engine/Source/Runtime/D3D12RHI/Private/D3D12CommandContext.h
--- a/Engine/Source/Runtime/D3D12RHI/Private/D3D12CommandContext.h
+++ b/Engine/Source/Runtime/D3D12RHI/Private/D3D12CommandContext.h
@@ -263,7 +263,7 @@ public:
 	virtual void RHISetStreamSource(uint32 StreamIndex, FVertexBufferRHIParamRef VertexBuffer, uint32 Offset) final override;
 	virtual void RHISetViewport(uint32 MinX, uint32 MinY, float MinZ, uint32 MaxX, uint32 MaxY, float MaxZ) final override;
 	virtual void RHISetScissorRect(bool bEnable, uint32 MinX, uint32 MinY, uint32 MaxX, uint32 MaxY) final override;
-	virtual void RHISetGraphicsPipelineState(FGraphicsPipelineStateRHIParamRef GraphicsPipelineState) final override;
+	virtual void RHISetGraphicsPipelineState(FRHIGraphicsPipelineState* GraphicsPipelineState) final override;
 	virtual void RHISetShaderTexture(FVertexShaderRHIParamRef VertexShader, uint32 TextureIndex, FTextureRHIParamRef NewTexture) final override;
 	virtual void RHISetShaderTexture(FHullShaderRHIParamRef HullShader, uint32 TextureIndex, FTextureRHIParamRef NewTexture) final override;
 	virtual void RHISetShaderTexture(FDomainShaderRHIParamRef DomainShader, uint32 TextureIndex, FTextureRHIParamRef NewTexture) final override;
@@ -564,7 +564,7 @@ public:
 	{
 		ContextRedirect(RHISetScissorRect(bEnable, MinX, MinY, MaxX, MaxY));
 	}
-	FORCEINLINE void RHISetGraphicsPipelineState(FGraphicsPipelineStateRHIParamRef GraphicsPipelineState) final override
+	FORCEINLINE void RHISetGraphicsPipelineState(FRHIGraphicsPipelineState* GraphicsPipelineState) final override
 	{
 		ContextRedirect(RHISetGraphicsPipelineState(GraphicsPipelineState));
 	}
diff --git a/Engine/Source/Runtime/D3D12RHI/Private/D3D12Commands.cpp b/Engine/Source/Runtime/D3D12RHI/Private/D3D12Commands.cpp
--- a/Engine/Source/Runtime/D3D12RHI/Private/D3D12Commands.cpp
+++ b/Engine/Source/Runtime/D3D12RHI/Private/D3D12Commands.cpp
@@ -549,7 +549,7 @@ static inline D3D_PRIMITIVE_TOPOLOGY GetD3D12PrimitiveType(EPrimitiveType Primit
 	return D3D_PRIMITIVE_TOPOLOGY_TRIANGLELIST;
 }
 
-void FD3D12CommandContext::RHISetGraphicsPipelineState(FGraphicsPipelineStateRHIParamRef GraphicsState)
+void FD3D12CommandContext::RHISetGraphicsPipelineState(FRHIGraphicsPipelineState* GraphicsState)
 {
 	FD3D12GraphicsPipelineState* GraphicsPipelineState = FD3D12DynamicRHI::ResourceCast(GraphicsState);
 
diff --git a/Engine/Source/Runtime/EmptyRHI/Private/EmptyCommands.cpp b/Engine/Source/Runtime/EmptyRHI/Private/EmptyCommands.cpp
--- a/Engine/Source/Runtime/EmptyRHI/Private/EmptyCommands.cpp
+++ b/Engine/Source/Runtime/EmptyRHI/Private/EmptyCommands.cpp
@@ -238,13 +238,13 @@ void FEmptyDynamicRHI::RHISetShaderUniformBuffer(FComputeShaderRHIParamRef Compu
 }
 
 
-void FEmptyDynamicRHI::RHISetDepthStencilState(FDepthStencilStateRHIParamRef NewStateRHI, uint32 StencilRef)
+void FEmptyDynamicRHI::RHISetDepthStencilState(FRHIDepthStencilState* NewStateRHI, uint32 StencilRef)
 {
 	FEmptyDepthStencilState* NewState = ResourceCast(NewStateRHI);
 
 }
 
-void FEmptyDynamicRHI::RHISetBlendState(FBlendStateRHIParamRef NewStateRHI, const FLinearColor& BlendFactor)
+void FEmptyDynamicRHI::RHISetBlendState(FRHIBlendState* NewStateRHI, const FLinearColor& BlendFactor)
 {
 	FEmptyBlendState* NewState = ResourceCast(NewStateRHI);
 
diff --git a/Engine/Source/Runtime/EmptyRHI/Public/EmptyRHI.h b/Engine/Source/Runtime/EmptyRHI/Public/EmptyRHI.h
--- a/Engine/Source/Runtime/EmptyRHI/Public/EmptyRHI.h
+++ b/Engine/Source/Runtime/EmptyRHI/Public/EmptyRHI.h
@@ -187,8 +187,8 @@ public:
 	virtual void RHISetShaderParameter(FDomainShaderRHIParamRef DomainShader, uint32 BufferIndex, uint32 BaseIndex, uint32 NumBytes, const void* NewValue) final override;
 	virtual void RHISetShaderParameter(FGeometryShaderRHIParamRef GeometryShader, uint32 BufferIndex, uint32 BaseIndex, uint32 NumBytes, const void* NewValue) final override;
 	virtual void RHISetShaderParameter(FComputeShaderRHIParamRef ComputeShader, uint32 BufferIndex, uint32 BaseIndex, uint32 NumBytes, const void* NewValue) final override;
-	virtual void RHISetDepthStencilState(FDepthStencilStateRHIParamRef NewState, uint32 StencilRef) final override;
-	virtual void RHISetBlendState(FBlendStateRHIParamRef NewState, const FLinearColor& BlendFactor) final override;
+	virtual void RHISetDepthStencilState(FRHIDepthStencilState* NewState, uint32 StencilRef) final override;
+	virtual void RHISetBlendState(FRHIBlendState* NewState, const FLinearColor& BlendFactor) final override;
 	virtual void RHISetRenderTargets(uint32 NumSimultaneousRenderTargets, const FRHIRenderTargetView* NewRenderTargets, const FRHIDepthRenderTargetView* NewDepthStencilTarget, uint32 NumUAVs, const FUnorderedAccessViewRHIParamRef* UAVs) final override;
 	virtual void RHISetRenderTargetsAndClear(const FRHISetRenderTargetsInfo& RenderTargetsInfo) final override;
 	virtual void RHIDrawPrimitive(uint32 PrimitiveType, uint32 BaseVertexIndex, uint32 NumPrimitives, uint32 NumInstances) final override;
diff --git a/Engine/Source/Runtime/Engine/Public/DebugViewModeInterface.h b/Engine/Source/Runtime/Engine/Public/DebugViewModeInterface.h
--- a/Engine/Source/Runtime/Engine/Public/DebugViewModeInterface.h
+++ b/Engine/Source/Runtime/Engine/Public/DebugViewModeInterface.h
@@ -29,8 +29,8 @@ public:
 			: BlendState(nullptr)
 			, DepthStencilState(nullptr)
 		{}
-		FBlendStateRHIParamRef			BlendState;
-		FDepthStencilStateRHIParamRef	DepthStencilState;
+		FRHIBlendState*			BlendState;
+		FRHIDepthStencilState*	DepthStencilState;
 	};
 
 	FDebugViewModeInterface(
diff --git a/Engine/Source/Runtime/NullDrv/Public/NullRHI.h b/Engine/Source/Runtime/NullDrv/Public/NullRHI.h
--- a/Engine/Source/Runtime/NullDrv/Public/NullRHI.h
+++ b/Engine/Source/Runtime/NullDrv/Public/NullRHI.h
@@ -704,15 +704,16 @@ public:
 	}
 
 
-	virtual void RHISetDepthStencilState(FDepthStencilStateRHIParamRef NewState, uint32 StencilRef) final override
+	virtual void RHISetDepthStencilState(FRHIDepthStencilState* NewState, uint32 StencilRef) final override
 	{
 
 	}
 
-	virtual void RHISetBlendState(FBlendStateRHIParamRef NewState, const FLinearColor& BlendFactor) final override
+	virtual void RHISetBlendState(FRHIBlendState* NewState, const FLinearColor& BlendFactor) final override
 	{
 
 	}
+
 	virtual void RHISetRenderTargets(uint32 NumSimultaneousRenderTargets, const FRHIRenderTargetView* NewRenderTargets, const FRHIDepthRenderTargetView* NewDepthStencilTarget, uint32 NumUAVs, const FUnorderedAccessViewRHIParamRef* UAVs) final override
 	{
 
diff --git a/Engine/Source/Runtime/OpenGLDrv/Private/OpenGLCommands.cpp b/Engine/Source/Runtime/OpenGLDrv/Private/OpenGLCommands.cpp
--- a/Engine/Source/Runtime/OpenGLDrv/Private/OpenGLCommands.cpp
+++ b/Engine/Source/Runtime/OpenGLDrv/Private/OpenGLCommands.cpp
@@ -1252,7 +1252,7 @@ void FOpenGLDynamicRHI::RHISetShaderParameter(FComputeShaderRHIParamRef ComputeS
 	PendingState.LinkedProgramAndDirtyFlag = nullptr;
 }
 
-void FOpenGLDynamicRHI::RHISetDepthStencilState(FDepthStencilStateRHIParamRef NewStateRHI,uint32 StencilRef)
+void FOpenGLDynamicRHI::RHISetDepthStencilState(FRHIDepthStencilState* NewStateRHI,uint32 StencilRef)
 {
 	VERIFY_GL_SCOPE();
 	FOpenGLDepthStencilState* NewState = ResourceCast(NewStateRHI);
@@ -1609,7 +1609,7 @@ void FOpenGLDynamicRHI::SetPendingBlendStateForActiveRenderTargets( FOpenGLConte
 	}
 }
 
-void FOpenGLDynamicRHI::RHISetBlendState(FBlendStateRHIParamRef NewStateRHI,const FLinearColor& BlendFactor)
+void FOpenGLDynamicRHI::RHISetBlendState(FRHIBlendState* NewStateRHI,const FLinearColor& BlendFactor)
 {
 	VERIFY_GL_SCOPE();
 	FOpenGLBlendState* NewState = ResourceCast(NewStateRHI);
diff --git a/Engine/Source/Runtime/OpenGLDrv/Public/OpenGLDrv.h b/Engine/Source/Runtime/OpenGLDrv/Public/OpenGLDrv.h
--- a/Engine/Source/Runtime/OpenGLDrv/Public/OpenGLDrv.h
+++ b/Engine/Source/Runtime/OpenGLDrv/Public/OpenGLDrv.h
@@ -499,9 +499,9 @@ public:
 	virtual void RHISetShaderParameter(FDomainShaderRHIParamRef DomainShader, uint32 BufferIndex, uint32 BaseIndex, uint32 NumBytes, const void* NewValue) final override;
 	virtual void RHISetShaderParameter(FGeometryShaderRHIParamRef GeometryShader, uint32 BufferIndex, uint32 BaseIndex, uint32 NumBytes, const void* NewValue) final override;
 	virtual void RHISetShaderParameter(FComputeShaderRHIParamRef ComputeShader, uint32 BufferIndex, uint32 BaseIndex, uint32 NumBytes, const void* NewValue) final override;
-	virtual void RHISetDepthStencilState(FDepthStencilStateRHIParamRef NewState, uint32 StencilRef) final override;
+	virtual void RHISetDepthStencilState(FRHIDepthStencilState* NewState, uint32 StencilRef) final override;
 	virtual void RHISetStencilRef(uint32 StencilRef) final override;
-	virtual void RHISetBlendState(FBlendStateRHIParamRef NewState, const FLinearColor& BlendFactor) final override;
+	virtual void RHISetBlendState(FRHIBlendState* NewState, const FLinearColor& BlendFactor) final override;
 	virtual void RHISetBlendFactor(const FLinearColor& BlendFactor) final override
 	{
 		// Currently ignored as well as on RHISetBlendState()...
@@ -967,7 +967,7 @@ public:
 		RHITHREAD_GLCOMMAND_EPILOGUE();
 	}
 
-	virtual void RHISetGraphicsPipelineState(FGraphicsPipelineStateRHIParamRef GraphicsState) final override
+	virtual void RHISetGraphicsPipelineState(FRHIGraphicsPipelineState* GraphicsState) final override
 	{
 		FRHIGraphicsPipelineStateFallBack* FallbackGraphicsState = static_cast<FRHIGraphicsPipelineStateFallBack*>(GraphicsState);
 
diff --git a/Engine/Source/Runtime/RHI/Public/RHIContext.h b/Engine/Source/Runtime/RHI/Public/RHIContext.h
--- a/Engine/Source/Runtime/RHI/Public/RHIContext.h
+++ b/Engine/Source/Runtime/RHI/Public/RHIContext.h
@@ -326,7 +326,7 @@ public:
 	// @param MaxY excluding like Win32 RECT
 	virtual void RHISetScissorRect(bool bEnable, uint32 MinX, uint32 MinY, uint32 MaxX, uint32 MaxY) = 0;
 
-	virtual void RHISetGraphicsPipelineState(FGraphicsPipelineStateRHIParamRef GraphicsState) = 0;
+	virtual void RHISetGraphicsPipelineState(FRHIGraphicsPipelineState* GraphicsState) = 0;
 
 	/** Set the shader resource view of a surface.  This is used for binding TextureMS parameter types that need a multi sampled view. */
 	virtual void RHISetShaderTexture(FVertexShaderRHIParamRef VertexShader, uint32 TextureIndex, FTextureRHIParamRef NewTexture) = 0;
@@ -668,11 +668,11 @@ public:
 	*/
 	virtual void RHISetBoundShaderState(FBoundShaderStateRHIParamRef BoundShaderState) = 0;
 
-	virtual void RHISetDepthStencilState(FDepthStencilStateRHIParamRef NewState, uint32 StencilRef) = 0;
+	virtual void RHISetDepthStencilState(FRHIDepthStencilState* NewState, uint32 StencilRef) = 0;
 
 	virtual void RHISetRasterizerState(FRHIRasterizerState* NewState) = 0;
 
-	virtual void RHISetBlendState(FBlendStateRHIParamRef NewState, const FLinearColor& BlendFactor) = 0;
+	virtual void RHISetBlendState(FRHIBlendState* NewState, const FLinearColor& BlendFactor) = 0;
 
 	virtual void RHIEnableDepthBoundsTest(bool bEnable) = 0;
 
@@ -680,7 +680,7 @@ public:
 	* This will set most relevant pipeline state. Legacy APIs are expected to set corresponding disjoint state as well.
 	* @param GraphicsShaderState - the graphics pipeline state
 	*/
-	virtual void RHISetGraphicsPipelineState(FGraphicsPipelineStateRHIParamRef GraphicsState) override
+	virtual void RHISetGraphicsPipelineState(FRHIGraphicsPipelineState* GraphicsState) override
 	{
 		FRHIGraphicsPipelineStateFallBack* FallbackGraphicsState = static_cast<FRHIGraphicsPipelineStateFallBack*>(GraphicsState);
 		FGraphicsPipelineStateInitializer& PsoInit = FallbackGraphicsState->Initializer;
diff --git a/Engine/Source/Runtime/RHI/Public/RHIResources.h b/Engine/Source/Runtime/RHI/Public/RHIResources.h
--- a/Engine/Source/Runtime/RHI/Public/RHIResources.h
+++ b/Engine/Source/Runtime/RHI/Public/RHIResources.h
@@ -1105,9 +1105,11 @@ UE_DEPRECATED(4.23, "FRasterizerStateRHIParamRef typedef is deprecated; please u
 typedef FRHIRasterizerState*              FRasterizerStateRHIParamRef;
 typedef TRefCountPtr<FRHIRasterizerState> FRasterizerStateRHIRef;
 
+UE_DEPRECATED(4.23, "FDepthStencilStateRHIParamRef typedef is deprecated; please use FRHIDepthStencilState* directly instead.")
 typedef FRHIDepthStencilState*              FDepthStencilStateRHIParamRef;
 typedef TRefCountPtr<FRHIDepthStencilState> FDepthStencilStateRHIRef;
 
+UE_DEPRECATED(4.23, "FBlendStateRHIParamRef typedef is deprecated; please use FRHIBlendState* directly instead.")
 typedef FRHIBlendState*              FBlendStateRHIParamRef;
 typedef TRefCountPtr<FRHIBlendState> FBlendStateRHIRef;
 
@@ -1190,6 +1192,7 @@ typedef TRefCountPtr<FRHIUnorderedAccessView> FUnorderedAccessViewRHIRef;
 typedef FRHIShaderResourceView*              FShaderResourceViewRHIParamRef;
 typedef TRefCountPtr<FRHIShaderResourceView> FShaderResourceViewRHIRef;
 
+UE_DEPRECATED(4.23, "FGraphicsPipelineStateRHIParamRef typedef is deprecated; please use FRHIGraphicsPipelineState* directly instead.")
 typedef FRHIGraphicsPipelineState*              FGraphicsPipelineStateRHIParamRef;
 typedef TRefCountPtr<FRHIGraphicsPipelineState> FGraphicsPipelineStateRHIRef;
 
@@ -1832,12 +1835,12 @@ public:
 	}
 
 	FGraphicsMinimalPipelineStateInitializer(
-		FBoundShaderStateInput				InBoundShaderState,
-		FBlendStateRHIParamRef				InBlendState,
-		FRHIRasterizerState*				InRasterizerState,
-		FDepthStencilStateRHIParamRef		InDepthStencilState,
-		FImmutableSamplerState				InImmutableSamplerState,
-		EPrimitiveType						InPrimitiveType
+		FBoundShaderStateInput		InBoundShaderState,
+		FRHIBlendState*				InBlendState,
+		FRHIRasterizerState*		InRasterizerState,
+		FRHIDepthStencilState*		InDepthStencilState,
+		FImmutableSamplerState		InImmutableSamplerState,
+		EPrimitiveType				InPrimitiveType
 		)
 		: BoundShaderState(InBoundShaderState)
 		, BlendState(InBlendState)
@@ -1962,19 +1965,19 @@ public:
 
 	// TODO: [PSO API] - As we migrate reuse existing API objects, but eventually we can move to the direct initializers. 
 	// When we do that work, move this to RHI.h as its more appropriate there, but here for now since dependent typdefs are here.
-	FBoundShaderStateInput			BoundShaderState;
-	FBlendStateRHIParamRef			BlendState;
-	FRHIRasterizerState*			RasterizerState;
-	FDepthStencilStateRHIParamRef	DepthStencilState;
-	FImmutableSamplerState			ImmutableSamplerState;
+	FBoundShaderStateInput	BoundShaderState;
+	FRHIBlendState*			BlendState;
+	FRHIRasterizerState*	RasterizerState;
+	FRHIDepthStencilState*	DepthStencilState;
+	FImmutableSamplerState	ImmutableSamplerState;
 
 	// Note: FGraphicsMinimalPipelineStateInitializer is 8-byte aligned and can't have any implicit padding,
 	// as it is sometimes hashed and compared as raw bytes. Explicit padding is therefore required between
 	// all data members and at the end of the structure.
-	bool							bDepthBounds = false;
-	uint8							Padding[3] = {};
+	bool					bDepthBounds = false;
+	uint8					Padding[3] = {};
 
-	EPrimitiveType					PrimitiveType;
+	EPrimitiveType			PrimitiveType;
 };
 
 // Hints for some RHIs that support subpasses
@@ -2032,26 +2035,26 @@ public:
 	}
 
 	FGraphicsPipelineStateInitializer(
-		FBoundShaderStateInput				InBoundShaderState,
-		FBlendStateRHIParamRef				InBlendState,
-		FRHIRasterizerState*				InRasterizerState,
-		FDepthStencilStateRHIParamRef		InDepthStencilState,
-		FImmutableSamplerState				InImmutableSamplerState,
-		EPrimitiveType						InPrimitiveType,
-		uint32								InRenderTargetsEnabled,
-		const TRenderTargetFormats&			InRenderTargetFormats,
-		const TRenderTargetFlags&			InRenderTargetFlags,
-		EPixelFormat						InDepthStencilTargetFormat,
-		uint32								InDepthStencilTargetFlag,
-		ERenderTargetLoadAction				InDepthTargetLoadAction,
-		ERenderTargetStoreAction			InDepthTargetStoreAction,
-		ERenderTargetLoadAction				InStencilTargetLoadAction,
-		ERenderTargetStoreAction			InStencilTargetStoreAction,
-		FExclusiveDepthStencil				InDepthStencilAccess,
-		uint32								InNumSamples,
-		ESubpassHint						InSubpassHint,
-		uint8								InSubpassIndex,
-		uint16								InFlags
+		FBoundShaderStateInput		InBoundShaderState,
+		FRHIBlendState*				InBlendState,
+		FRHIRasterizerState*		InRasterizerState,
+		FRHIDepthStencilState*		InDepthStencilState,
+		FImmutableSamplerState		InImmutableSamplerState,
+		EPrimitiveType				InPrimitiveType,
+		uint32						InRenderTargetsEnabled,
+		const TRenderTargetFormats&	InRenderTargetFormats,
+		const TRenderTargetFlags&	InRenderTargetFlags,
+		EPixelFormat				InDepthStencilTargetFormat,
+		uint32						InDepthStencilTargetFlag,
+		ERenderTargetLoadAction		InDepthTargetLoadAction,
+		ERenderTargetStoreAction	InDepthTargetStoreAction,
+		ERenderTargetLoadAction		InStencilTargetLoadAction,
+		ERenderTargetStoreAction	InStencilTargetStoreAction,
+		FExclusiveDepthStencil		InDepthStencilAccess,
+		uint32						InNumSamples,
+		ESubpassHint				InSubpassHint,
+		uint8						InSubpassIndex,
+		uint16						InFlags
 		)
 		: FGraphicsMinimalPipelineStateInitializer(InBoundShaderState, InBlendState, InRasterizerState, InDepthStencilState, InImmutableSamplerState, InPrimitiveType)
 		, RenderTargetsEnabled(InRenderTargetsEnabled)
diff --git a/Engine/Source/Runtime/RHI/Public/RHIStaticStates.h b/Engine/Source/Runtime/RHI/Public/RHIStaticStates.h
--- a/Engine/Source/Runtime/RHI/Public/RHIStaticStates.h
+++ b/Engine/Source/Runtime/RHI/Public/RHIStaticStates.h
@@ -258,7 +258,7 @@ class TStaticDepthStencilState : public TStaticStateRHI<
 		StencilWriteMask
 		>,
 	FDepthStencilStateRHIRef,
-	FDepthStencilStateRHIParamRef
+	FRHIDepthStencilState*
 	>
 {
 public:
@@ -370,7 +370,7 @@ class TStaticBlendState : public TStaticStateRHI<
 		RT7ColorWriteMask,RT7ColorBlendOp,RT7ColorSrcBlend,RT7ColorDestBlend,RT7AlphaBlendOp,RT7AlphaSrcBlend,RT7AlphaDestBlend
 		>,
 	FBlendStateRHIRef,
-	FBlendStateRHIParamRef
+	FRHIBlendState*
 	>
 {
 public:
diff --git a/Engine/Source/Runtime/RHI/Public/RHIValidationContext.h b/Engine/Source/Runtime/RHI/Public/RHIValidationContext.h
--- a/Engine/Source/Runtime/RHI/Public/RHIValidationContext.h
+++ b/Engine/Source/Runtime/RHI/Public/RHIValidationContext.h
@@ -406,7 +406,7 @@ public:
 		RHIContext->RHISetScissorRect(bEnable, MinX, MinY, MaxX, MaxY);
 	}
 
-	virtual void RHISetGraphicsPipelineState(FGraphicsPipelineStateRHIParamRef GraphicsState) override final
+	virtual void RHISetGraphicsPipelineState(FRHIGraphicsPipelineState* GraphicsState) override final
 	{
 		checkf(State.bInsideBeginRenderPass, TEXT("Graphics PSOs can only be set inside a RenderPass!"));
 		State.bGfxPSOSet = true;
diff --git a/Engine/Source/Runtime/Renderer/Private/BasePassRendering.cpp b/Engine/Source/Runtime/Renderer/Private/BasePassRendering.cpp
--- a/Engine/Source/Runtime/Renderer/Private/BasePassRendering.cpp
+++ b/Engine/Source/Runtime/Renderer/Private/BasePassRendering.cpp
@@ -261,7 +261,7 @@ FMeshDrawCommandSortKey CalculateBasePassMeshStaticSortKey(EDepthDrawingMode Ear
 	return SortKey;
 }
 
-void SetDepthStencilStateForBasePass(FMeshPassProcessorRenderState& DrawRenderState, ERHIFeatureLevel::Type FeatureLevel, const FMeshBatch& Mesh, const FPrimitiveSceneProxy* PrimitiveSceneProxy, bool bEnableReceiveDecalOutput, bool bUseDebugViewPS, FDepthStencilStateRHIParamRef LodFadeOverrideDepthStencilState)
+void SetDepthStencilStateForBasePass(FMeshPassProcessorRenderState& DrawRenderState, ERHIFeatureLevel::Type FeatureLevel, const FMeshBatch& Mesh, const FPrimitiveSceneProxy* PrimitiveSceneProxy, bool bEnableReceiveDecalOutput, bool bUseDebugViewPS, FRHIDepthStencilState* LodFadeOverrideDepthStencilState)
 {
 	static IConsoleVariable* EarlyZPassOnlyMaterialMaskingCVar = IConsoleManager::Get().FindConsoleVariable(TEXT("r.EarlyZPassOnlyMaterialMasking"));
 	bool bMaskInEarlyPass = (EarlyZPassOnlyMaterialMaskingCVar && Mesh.MaterialRenderProxy->GetMaterial(FeatureLevel)->IsMasked() && EarlyZPassOnlyMaterialMaskingCVar->GetInt());
diff --git a/Engine/Source/Runtime/Renderer/Private/BasePassRendering.h b/Engine/Source/Runtime/Renderer/Private/BasePassRendering.h
--- a/Engine/Source/Runtime/Renderer/Private/BasePassRendering.h
+++ b/Engine/Source/Runtime/Renderer/Private/BasePassRendering.h
@@ -654,6 +654,6 @@ private:
 
 ENUM_CLASS_FLAGS(FBasePassMeshProcessor::EFlags);
 
-extern void SetDepthStencilStateForBasePass(FMeshPassProcessorRenderState& DrawRenderState, ERHIFeatureLevel::Type FeatureLevel, const FMeshBatch& Mesh, const FPrimitiveSceneProxy* PrimitiveSceneProxy, bool bEnableReceiveDecalOutput, bool bUseDebugViewPS, FDepthStencilStateRHIParamRef LodFadeOverrideDepthStencilState);
+extern void SetDepthStencilStateForBasePass(FMeshPassProcessorRenderState& DrawRenderState, ERHIFeatureLevel::Type FeatureLevel, const FMeshBatch& Mesh, const FPrimitiveSceneProxy* PrimitiveSceneProxy, bool bEnableReceiveDecalOutput, bool bUseDebugViewPS, FRHIDepthStencilState* LodFadeOverrideDepthStencilState);
 extern void SetupBasePassState(FExclusiveDepthStencil::Type BasePassDepthStencilAccess, const bool bShaderComplexity, FMeshPassProcessorRenderState& DrawRenderState);
 extern FMeshDrawCommandSortKey CalculateTranslucentMeshStaticSortKey(const FPrimitiveSceneProxy* RESTRICT PrimitiveSceneProxy, uint16 MeshIdInPrimitive);
\ No newline at end of file
diff --git a/Engine/Source/Runtime/Renderer/Private/CompositionLighting/PostProcessDeferredDecals.cpp b/Engine/Source/Runtime/Renderer/Private/CompositionLighting/PostProcessDeferredDecals.cpp
--- a/Engine/Source/Runtime/Renderer/Private/CompositionLighting/PostProcessDeferredDecals.cpp
+++ b/Engine/Source/Runtime/Renderer/Private/CompositionLighting/PostProcessDeferredDecals.cpp
@@ -54,7 +54,7 @@ struct FDecalDepthState
 };
 
 // @param RenderState 0:before BasePass, 1:before lighting, (later we could add "after lighting" and multiply)
-FBlendStateRHIParamRef GetDecalBlendState(const ERHIFeatureLevel::Type SMFeatureLevel, EDecalRenderStage InDecalRenderStage, EDecalBlendMode DecalBlendMode, bool bHasNormal)
+FRHIBlendState* GetDecalBlendState(const ERHIFeatureLevel::Type SMFeatureLevel, EDecalRenderStage InDecalRenderStage, EDecalBlendMode DecalBlendMode, bool bHasNormal)
 {
 	if (InDecalRenderStage == DRS_BeforeBasePass)
 	{
@@ -404,7 +404,7 @@ static FDecalDepthState ComputeDecalDepthState(EDecalRenderStage LocalDecalStage
 	return Ret;
 }
 
-static FDepthStencilStateRHIParamRef GetDecalDepthState(uint32& StencilRef, FDecalDepthState DecalDepthState)
+static FRHIDepthStencilState* GetDecalDepthState(uint32& StencilRef, FDecalDepthState DecalDepthState)
 {
 	switch (DecalDepthState.DepthTest)
 	{
diff --git a/Engine/Source/Runtime/Renderer/Private/CompositionLighting/PostProcessDeferredDecals.h b/Engine/Source/Runtime/Renderer/Private/CompositionLighting/PostProcessDeferredDecals.h
--- a/Engine/Source/Runtime/Renderer/Private/CompositionLighting/PostProcessDeferredDecals.h
+++ b/Engine/Source/Runtime/Renderer/Private/CompositionLighting/PostProcessDeferredDecals.h
@@ -82,6 +82,6 @@ struct FDecalRenderTargetManager
 	void SetRenderTargetMode(FDecalRenderingCommon::ERenderTargetMode CurrentRenderTargetMode, bool bHasNormal, bool bPerPixelDBufferMask);
 };
 
-extern FBlendStateRHIParamRef GetDecalBlendState(const ERHIFeatureLevel::Type SMFeatureLevel, EDecalRenderStage InDecalRenderStage, EDecalBlendMode DecalBlendMode, bool bHasNormal);
+extern FRHIBlendState* GetDecalBlendState(const ERHIFeatureLevel::Type SMFeatureLevel, EDecalRenderStage InDecalRenderStage, EDecalBlendMode DecalBlendMode, bool bHasNormal);
 
 extern void RenderMeshDecals(FRenderingCompositePassContext& Context, EDecalRenderStage CurrentDecalStage);
diff --git a/Engine/Source/Runtime/Renderer/Private/CustomDepthRendering.cpp b/Engine/Source/Runtime/Renderer/Private/CustomDepthRendering.cpp
--- a/Engine/Source/Runtime/Renderer/Private/CustomDepthRendering.cpp
+++ b/Engine/Source/Runtime/Renderer/Private/CustomDepthRendering.cpp
@@ -74,7 +74,7 @@ void FCustomDepthPassMeshProcessor::AddMeshBatch(const FMeshBatch& RESTRICT Mesh
 		if (bWriteCustomStencilValues)
 		{
 			const uint32 CustomDepthStencilValue = PrimitiveSceneProxy->GetCustomDepthStencilValue();
-			const static FDepthStencilStateRHIParamRef StencilStates[EStencilMask::SM_Count] =
+			static FRHIDepthStencilState* StencilStates[EStencilMask::SM_Count] =
 			{
 				TStaticDepthStencilState<true, CF_DepthNearOrEqual, true, CF_Always, SO_Keep, SO_Keep, SO_Replace, false, CF_Always, SO_Keep, SO_Keep, SO_Keep, 255, 255>::GetRHI(),
 				TStaticDepthStencilState<true, CF_DepthNearOrEqual, true, CF_Always, SO_Keep, SO_Replace, SO_Replace, false, CF_Always, SO_Keep, SO_Keep, SO_Keep, 255, 255>::GetRHI(),
diff --git a/Engine/Source/Runtime/Renderer/Private/PostProcess/PostProcessMaterial.cpp b/Engine/Source/Runtime/Renderer/Private/PostProcess/PostProcessMaterial.cpp
--- a/Engine/Source/Runtime/Renderer/Private/PostProcess/PostProcessMaterial.cpp
+++ b/Engine/Source/Runtime/Renderer/Private/PostProcess/PostProcessMaterial.cpp
@@ -310,7 +310,7 @@ void FRCPassPostProcessMaterial::Process(FRenderingCompositePassContext& Context
 
 	const FSceneRenderTargetItem* CustomDepthStencilTarget = nullptr;
 
-	FDepthStencilStateRHIParamRef DepthStencilState;
+	FRHIDepthStencilState* DepthStencilState;
 	uint32 StencilRefValue = 0;
 
 	if (bDoStencilTest)
@@ -339,7 +339,7 @@ void FRCPassPostProcessMaterial::Process(FRenderingCompositePassContext& Context
 			Swap(SceneContext.CustomDepth, CustomDepthStencilCopy);
 		}
 
-		static const FDepthStencilStateRHIParamRef StencilStates[] =
+		static FRHIDepthStencilState* StencilStates[] =
 		{
 			TStaticDepthStencilState<false, CF_Always, true, CF_Less>::GetRHI(),
 			TStaticDepthStencilState<false, CF_Always, true, CF_LessEqual>::GetRHI(),
@@ -360,11 +360,11 @@ void FRCPassPostProcessMaterial::Process(FRenderingCompositePassContext& Context
 		DepthStencilState = TStaticDepthStencilState<false, CF_Always>::GetRHI();
 	}
 
-	FBlendStateRHIParamRef BlendState = TStaticBlendState<>::GetRHI();
+	FRHIBlendState* BlendState = TStaticBlendState<>::GetRHI();
 	bool bDoOutputBlend = Material->GetBlendableOutputAlpha() && CVarPostProcessAllowBlendModes.GetValueOnRenderThread() != 0;
 	if (bDoOutputBlend)
 	{
-		static const FBlendStateRHIParamRef BlendStates[] =
+		static FRHIBlendState* BlendStates[] =
 		{
 			TStaticBlendState<>::GetRHI(),
 			TStaticBlendState<>::GetRHI(),
diff --git a/Engine/Source/Runtime/Renderer/Private/PostProcess/PostProcessTemporalAA.cpp b/Engine/Source/Runtime/Renderer/Private/PostProcess/PostProcessTemporalAA.cpp
--- a/Engine/Source/Runtime/Renderer/Private/PostProcess/PostProcessTemporalAA.cpp
+++ b/Engine/Source/Runtime/Renderer/Private/PostProcess/PostProcessTemporalAA.cpp
@@ -734,7 +734,7 @@ FTAAOutputs FTAAPassParameters::AddTemporalAAPass(
 			FTemporalAAPS::FPermutationDomain PermutationVector = BasePermutationVector;
 
 			// Lambda to draw pixel shader.
-			auto DrawTAAPixelShader = [&](FDepthStencilStateRHIParamRef DepthStencilState)
+			auto DrawTAAPixelShader = [&](FRHIDepthStencilState* DepthStencilState)
 			{
 				TShaderMapRef<FTemporalAAPS> TAAPixelShader(View.ShaderMap, PermutationVector);
 
@@ -752,7 +752,7 @@ FTAAOutputs FTAAPassParameters::AddTemporalAAPass(
 			{
 				// Normal temporal feedback
 				// Draw to pixels where stencil == 0
-				FDepthStencilStateRHIParamRef DepthStencilState = TStaticDepthStencilState<
+				FRHIDepthStencilState* DepthStencilState = TStaticDepthStencilState<
 					false, CF_Always,
 					true, CF_Equal, SO_Keep, SO_Keep, SO_Keep,
 					false, CF_Always, SO_Keep, SO_Keep, SO_Keep,
diff --git a/Engine/Source/Runtime/Renderer/Private/ShadowRendering.cpp b/Engine/Source/Runtime/Renderer/Private/ShadowRendering.cpp
--- a/Engine/Source/Runtime/Renderer/Private/ShadowRendering.cpp
+++ b/Engine/Source/Runtime/Renderer/Private/ShadowRendering.cpp
@@ -469,7 +469,7 @@ void FProjectedShadowInfo::SetBlendStateForProjection(
 
 	if (bProjectingForForwardShading)
 	{
-		FBlendStateRHIParamRef BlendState = NULL;
+		FRHIBlendState* BlendState = nullptr;
 
 		if (bUseFadePlane)
 		{
diff --git a/Engine/Source/Runtime/Renderer/Private/VT/RuntimeVirtualTextureRender.cpp b/Engine/Source/Runtime/Renderer/Private/VT/RuntimeVirtualTextureRender.cpp
--- a/Engine/Source/Runtime/Renderer/Private/VT/RuntimeVirtualTextureRender.cpp
+++ b/Engine/Source/Runtime/Renderer/Private/VT/RuntimeVirtualTextureRender.cpp
@@ -72,7 +72,7 @@ namespace RuntimeVirtualTexture
 			OutEnvironment.SetDefine(TEXT("OUT_BASECOLOR"), 1);
 		}
 
-		static FBlendStateRHIParamRef GetBlendState()
+		static FRHIBlendState* GetBlendState()
 		{
 			return TStaticBlendState< CW_RGBA, BO_Add, BF_One, BF_InverseSourceAlpha, BO_Add, BF_Zero, BF_One >::GetRHI();
 		}
@@ -87,7 +87,7 @@ namespace RuntimeVirtualTexture
 			OutEnvironment.SetDefine(TEXT("OUT_BASECOLOR_NORMAL"), 1);
 		}
 
-		static FBlendStateRHIParamRef GetBlendState()
+		static FRHIBlendState* GetBlendState()
 		{
 			return TStaticBlendState<
 				CW_RGBA, BO_Add, BF_One, BF_InverseSourceAlpha, BO_Add, BF_Zero, BF_One,
@@ -104,7 +104,7 @@ namespace RuntimeVirtualTexture
 			OutEnvironment.SetDefine(TEXT("OUT_BASECOLOR_NORMAL_SPECULAR"), 1);
 		}
 
-		static FBlendStateRHIParamRef GetBlendState()
+		static FRHIBlendState* GetBlendState()
 		{
 			return TStaticBlendState<
 				CW_RGBA, BO_Add, BF_One, BF_InverseSourceAlpha, BO_Add, BF_Zero, BF_One,
diff --git a/Engine/Source/Runtime/Renderer/Public/MeshPassProcessor.h b/Engine/Source/Runtime/Renderer/Public/MeshPassProcessor.h
--- a/Engine/Source/Runtime/Renderer/Public/MeshPassProcessor.h
+++ b/Engine/Source/Runtime/Renderer/Public/MeshPassProcessor.h
@@ -961,17 +961,17 @@ struct FMeshPassProcessorRenderState
 	}
 
 public:
-	FORCEINLINE_DEBUGGABLE void SetBlendState(FBlendStateRHIParamRef InBlendState)
+	FORCEINLINE_DEBUGGABLE void SetBlendState(FRHIBlendState* InBlendState)
 	{
 		BlendState = InBlendState;
 	}
 
-	FORCEINLINE_DEBUGGABLE const FBlendStateRHIParamRef GetBlendState() const
+	FORCEINLINE_DEBUGGABLE FRHIBlendState* GetBlendState() const
 	{
 		return BlendState;
 	}
 
-	FORCEINLINE_DEBUGGABLE void SetDepthStencilState(FDepthStencilStateRHIParamRef InDepthStencilState)
+	FORCEINLINE_DEBUGGABLE void SetDepthStencilState(FRHIDepthStencilState* InDepthStencilState)
 	{
 		DepthStencilState = InDepthStencilState;
 		StencilRef = 0;
@@ -982,7 +982,7 @@ public:
 		StencilRef = InStencilRef;
 	}
 
-	FORCEINLINE_DEBUGGABLE const FDepthStencilStateRHIParamRef GetDepthStencilState() const
+	FORCEINLINE_DEBUGGABLE FRHIDepthStencilState* GetDepthStencilState() const
 	{
 		return DepthStencilState;
 	}
@@ -1049,8 +1049,8 @@ public:
 	}
 
 private:
-	FBlendStateRHIParamRef			BlendState;
-	FDepthStencilStateRHIParamRef	DepthStencilState;
+	FRHIBlendState*					BlendState;
+	FRHIDepthStencilState*			DepthStencilState;
 	FExclusiveDepthStencil::Type	DepthStencilAccess;
 
 	TUniformBufferRef<FViewUniformShaderParameters>	ViewUniformBuffer;
diff --git a/Engine/Source/Runtime/SlateRHIRenderer/Private/SlateRHIRenderingPolicy.cpp b/Engine/Source/Runtime/SlateRHIRenderer/Private/SlateRHIRenderingPolicy.cpp
--- a/Engine/Source/Runtime/SlateRHIRenderer/Private/SlateRHIRenderingPolicy.cpp
+++ b/Engine/Source/Runtime/SlateRHIRenderer/Private/SlateRHIRenderingPolicy.cpp
@@ -554,7 +554,7 @@ static bool UpdateScissorRect(
 				// because we're about to go back to rendering widgets "normally", but with the added effect that now
 				// we have the stencil buffer bound with a bunch of clipping zones rendered into it.
 				{
-					FDepthStencilStateRHIParamRef DSMaskRead =
+					FRHIDepthStencilState* DSMaskRead =
 						TStaticDepthStencilState<
 						/*bEnableDepthWrite*/ false
 						, /*DepthTest*/ CF_Always
diff --git a/Engine/Source/Runtime/UtilityShaders/Private/ClearQuad.cpp b/Engine/Source/Runtime/UtilityShaders/Private/ClearQuad.cpp
--- a/Engine/Source/Runtime/UtilityShaders/Private/ClearQuad.cpp
+++ b/Engine/Source/Runtime/UtilityShaders/Private/ClearQuad.cpp
@@ -28,13 +28,11 @@ static void ClearQuadSetup( FRHICommandList& RHICmdList, bool bClearColor, int32
 	}
 
 	// Set new states
-	FBlendStateRHIParamRef BlendStateRHI;
-		
-	BlendStateRHI = bClearColor
+	FRHIBlendState* BlendStateRHI = bClearColor
 		? TStaticBlendState<>::GetRHI()
 		: TStaticBlendStateWriteMask<CW_NONE,CW_NONE,CW_NONE,CW_NONE,CW_NONE,CW_NONE,CW_NONE,CW_NONE>::GetRHI();
 	
-	const FDepthStencilStateRHIParamRef DepthStencilStateRHI = 
+	FRHIDepthStencilState* DepthStencilStateRHI =
 		(bClearDepth && bClearStencil)
 			? TStaticDepthStencilState<
 				true, CF_Always,
@@ -60,8 +58,7 @@ static void ClearQuadSetup( FRHICommandList& RHICmdList, bool bClearColor, int32
 	GraphicsPSOInit.BlendState = BlendStateRHI;
 	GraphicsPSOInit.DepthStencilState = DepthStencilStateRHI;
 
-	auto ShaderMap = GetGlobalShaderMap(GMaxRHIFeatureLevel);
-
+	auto* ShaderMap = GetGlobalShaderMap(GMaxRHIFeatureLevel);
 
 	// Set the new shaders
 	TShaderMapRef<TOneColorVS<true> > VertexShader(ShaderMap);
diff --git a/Engine/Source/Runtime/VulkanRHI/Private/VulkanContext.h b/Engine/Source/Runtime/VulkanRHI/Private/VulkanContext.h
--- a/Engine/Source/Runtime/VulkanRHI/Private/VulkanContext.h
+++ b/Engine/Source/Runtime/VulkanRHI/Private/VulkanContext.h
@@ -172,7 +172,7 @@ public:
 	virtual void RHISetStreamSource(uint32 StreamIndex, FVertexBufferRHIParamRef VertexBuffer, uint32 Offset) final override;
 	virtual void RHISetViewport(uint32 MinX, uint32 MinY, float MinZ, uint32 MaxX, uint32 MaxY, float MaxZ) final override;
 	virtual void RHISetScissorRect(bool bEnable, uint32 MinX, uint32 MinY, uint32 MaxX, uint32 MaxY) final override;
-	virtual void RHISetGraphicsPipelineState(FGraphicsPipelineStateRHIParamRef GraphicsState) final override;
+	virtual void RHISetGraphicsPipelineState(FRHIGraphicsPipelineState* GraphicsState) final override;
 	virtual void RHISetShaderTexture(FVertexShaderRHIParamRef VertexShader, uint32 TextureIndex, FTextureRHIParamRef NewTexture) final override;
 	virtual void RHISetShaderTexture(FHullShaderRHIParamRef HullShader, uint32 TextureIndex, FTextureRHIParamRef NewTexture) final override;
 	virtual void RHISetShaderTexture(FDomainShaderRHIParamRef DomainShader, uint32 TextureIndex, FTextureRHIParamRef NewTexture) final override;
diff --git a/Engine/Source/Runtime/VulkanRHI/Private/VulkanPipelineState.cpp b/Engine/Source/Runtime/VulkanRHI/Private/VulkanPipelineState.cpp
--- a/Engine/Source/Runtime/VulkanRHI/Private/VulkanPipelineState.cpp
+++ b/Engine/Source/Runtime/VulkanRHI/Private/VulkanPipelineState.cpp
@@ -338,7 +338,7 @@ bool FVulkanGraphicsPipelineDescriptorState::InternalUpdateDescriptorSets(FVulka
 }
 
 
-void FVulkanCommandListContext::RHISetGraphicsPipelineState(FGraphicsPipelineStateRHIParamRef GraphicsState)
+void FVulkanCommandListContext::RHISetGraphicsPipelineState(FRHIGraphicsPipelineState* GraphicsState)
 {
 	FVulkanRHIGraphicsPipelineState* Pipeline = ResourceCast(GraphicsState);
 	
diff --git a/Engine/Source/Runtime/Windows/D3D11RHI/Private/D3D11Commands.cpp b/Engine/Source/Runtime/Windows/D3D11RHI/Private/D3D11Commands.cpp
--- a/Engine/Source/Runtime/Windows/D3D11RHI/Private/D3D11Commands.cpp
+++ b/Engine/Source/Runtime/Windows/D3D11RHI/Private/D3D11Commands.cpp
@@ -813,7 +813,7 @@ void FD3D11DynamicRHI::ValidateExclusiveDepthStencilAccess(FExclusiveDepthStenci
 	}
 }
 
-void FD3D11DynamicRHI::RHISetDepthStencilState(FDepthStencilStateRHIParamRef NewStateRHI,uint32 StencilRef)
+void FD3D11DynamicRHI::RHISetDepthStencilState(FRHIDepthStencilState* NewStateRHI,uint32 StencilRef)
 {
 	FD3D11DepthStencilState* NewState = ResourceCast(NewStateRHI);
 
@@ -827,7 +827,7 @@ void FD3D11DynamicRHI::RHISetStencilRef(uint32 StencilRef)
 	StateCache.SetStencilRef(StencilRef);
 }
 
-void FD3D11DynamicRHI::RHISetBlendState(FBlendStateRHIParamRef NewStateRHI,const FLinearColor& BlendFactor)
+void FD3D11DynamicRHI::RHISetBlendState(FRHIBlendState* NewStateRHI,const FLinearColor& BlendFactor)
 {
 	FD3D11BlendState* NewState = ResourceCast(NewStateRHI);
 	StateCache.SetBlendState(NewState->Resource, (const float*)&BlendFactor, 0xffffffff);
diff --git a/Engine/Source/Runtime/Windows/D3D11RHI/Private/D3D11RHIPrivate.h b/Engine/Source/Runtime/Windows/D3D11RHI/Private/D3D11RHIPrivate.h
--- a/Engine/Source/Runtime/Windows/D3D11RHI/Private/D3D11RHIPrivate.h
+++ b/Engine/Source/Runtime/Windows/D3D11RHI/Private/D3D11RHIPrivate.h
@@ -521,7 +521,7 @@ public:
 	virtual void RHISetViewport(uint32 MinX, uint32 MinY, float MinZ, uint32 MaxX, uint32 MaxY, float MaxZ) final override;
 	virtual void RHISetScissorRect(bool bEnable, uint32 MinX, uint32 MinY, uint32 MaxX, uint32 MaxY) final override;
 	virtual void RHISetBoundShaderState(FBoundShaderStateRHIParamRef BoundShaderState) final override;
-	virtual void RHISetGraphicsPipelineState(FGraphicsPipelineStateRHIParamRef GraphicsState) final override
+	virtual void RHISetGraphicsPipelineState(FRHIGraphicsPipelineState* GraphicsState) final override
 	{
 		FRHIGraphicsPipelineStateFallBack* FallbackGraphicsState = static_cast<FRHIGraphicsPipelineStateFallBack*>(GraphicsState);
 		IRHICommandContextPSOFallback::RHISetGraphicsPipelineState(GraphicsState);
@@ -561,9 +561,9 @@ public:
 	virtual void RHISetShaderParameter(FDomainShaderRHIParamRef DomainShader, uint32 BufferIndex, uint32 BaseIndex, uint32 NumBytes, const void* NewValue) final override;
 	virtual void RHISetShaderParameter(FGeometryShaderRHIParamRef GeometryShader, uint32 BufferIndex, uint32 BaseIndex, uint32 NumBytes, const void* NewValue) final override;
 	virtual void RHISetShaderParameter(FComputeShaderRHIParamRef ComputeShader, uint32 BufferIndex, uint32 BaseIndex, uint32 NumBytes, const void* NewValue) final override;
-	virtual void RHISetDepthStencilState(FDepthStencilStateRHIParamRef NewState, uint32 StencilRef) final override;
+	virtual void RHISetDepthStencilState(FRHIDepthStencilState* NewState, uint32 StencilRef) final override;
 	virtual void RHISetStencilRef(uint32 StencilRef) final override;
-	virtual void RHISetBlendState(FBlendStateRHIParamRef NewState, const FLinearColor& BlendFactor) final override;
+	virtual void RHISetBlendState(FRHIBlendState* NewState, const FLinearColor& BlendFactor) final override;
 	virtual void RHISetBlendFactor(const FLinearColor& BlendFactor) final override;
 	virtual void RHISetRenderTargets(uint32 NumSimultaneousRenderTargets, const FRHIRenderTargetView* NewRenderTargets, const FRHIDepthRenderTargetView* NewDepthStencilTarget, uint32 NumUAVs, const FUnorderedAccessViewRHIParamRef* UAVs) final override;
 	virtual void RHISetRenderTargetsAndClear(const FRHISetRenderTargetsInfo& RenderTargetsInfo) final override;
