DR - Deprecate FBoundShaderStateRHIParamRef

diff --git a/Engine/Source/Runtime/EmptyRHI/Private/EmptyCommands.cpp b/Engine/Source/Runtime/EmptyRHI/Private/EmptyCommands.cpp
--- a/Engine/Source/Runtime/EmptyRHI/Private/EmptyCommands.cpp
+++ b/Engine/Source/Runtime/EmptyRHI/Private/EmptyCommands.cpp
@@ -55,7 +55,7 @@ void FEmptyDynamicRHI::RHISetScissorRect(bool bEnable,uint32 MinX,uint32 MinY,ui
 
 }
 
-void FEmptyDynamicRHI::RHISetBoundShaderState( FBoundShaderStateRHIParamRef BoundShaderStateRHI)
+void FEmptyDynamicRHI::RHISetBoundShaderState(FRHIBoundShaderState* BoundShaderStateRHI)
 {
 	FEmptyBoundShaderState* BoundShaderState = ResourceCast(BoundShaderStateRHI);
 
diff --git a/Engine/Source/Runtime/EmptyRHI/Public/EmptyRHI.h b/Engine/Source/Runtime/EmptyRHI/Public/EmptyRHI.h
--- a/Engine/Source/Runtime/EmptyRHI/Public/EmptyRHI.h
+++ b/Engine/Source/Runtime/EmptyRHI/Public/EmptyRHI.h
@@ -154,7 +154,7 @@ public:
 	virtual void RHISetRasterizerState(FRHIRasterizerState* NewState) final override;
 	virtual void RHISetViewport(uint32 MinX, uint32 MinY, float MinZ, uint32 MaxX, uint32 MaxY, float MaxZ) final override;
 	virtual void RHISetScissorRect(bool bEnable, uint32 MinX, uint32 MinY, uint32 MaxX, uint32 MaxY) final override;
-	virtual void RHISetBoundShaderState(FBoundShaderStateRHIParamRef BoundShaderState) final override;
+	virtual void RHISetBoundShaderState(FRHIBoundShaderState* BoundShaderState) final override;
 	virtual void RHISetShaderTexture(FVertexShaderRHIParamRef VertexShader, uint32 TextureIndex, FTextureRHIParamRef NewTexture) final override;
 	virtual void RHISetShaderTexture(FHullShaderRHIParamRef HullShader, uint32 TextureIndex, FTextureRHIParamRef NewTexture) final override;
 	virtual void RHISetShaderTexture(FDomainShaderRHIParamRef DomainShader, uint32 TextureIndex, FTextureRHIParamRef NewTexture) final override;
diff --git a/Engine/Source/Runtime/NullDrv/Public/NullRHI.h b/Engine/Source/Runtime/NullDrv/Public/NullRHI.h
--- a/Engine/Source/Runtime/NullDrv/Public/NullRHI.h
+++ b/Engine/Source/Runtime/NullDrv/Public/NullRHI.h
@@ -532,9 +532,8 @@ public:
 	{
 	}
 
-	virtual void RHISetBoundShaderState(FBoundShaderStateRHIParamRef BoundShaderState) final override
+	virtual void RHISetBoundShaderState(FRHIBoundShaderState* BoundShaderState) final override
 	{
-
 	}
 
 
diff --git a/Engine/Source/Runtime/OpenGLDrv/Private/OpenGLCommands.cpp b/Engine/Source/Runtime/OpenGLDrv/Private/OpenGLCommands.cpp
--- a/Engine/Source/Runtime/OpenGLDrv/Private/OpenGLCommands.cpp
+++ b/Engine/Source/Runtime/OpenGLDrv/Private/OpenGLCommands.cpp
@@ -434,7 +434,7 @@ inline void FOpenGLDynamicRHI::UpdateScissorRectInOpenGLContext( FOpenGLContextS
 * Set bound shader state. This will set the vertex decl/shader, and pixel shader
 * @param BoundShaderState - state resource
 */
-void FOpenGLDynamicRHI::RHISetBoundShaderState( FBoundShaderStateRHIParamRef BoundShaderStateRHI)
+void FOpenGLDynamicRHI::RHISetBoundShaderState(FRHIBoundShaderState* BoundShaderStateRHI)
 {
 	VERIFY_GL_SCOPE();
 	FOpenGLBoundShaderState* BoundShaderState = ResourceCast(BoundShaderStateRHI);
diff --git a/Engine/Source/Runtime/OpenGLDrv/Public/OpenGLDrv.h b/Engine/Source/Runtime/OpenGLDrv/Public/OpenGLDrv.h
--- a/Engine/Source/Runtime/OpenGLDrv/Public/OpenGLDrv.h
+++ b/Engine/Source/Runtime/OpenGLDrv/Public/OpenGLDrv.h
@@ -466,7 +466,7 @@ public:
 	virtual void RHISetRasterizerState(FRHIRasterizerState* NewState) final override;
 	virtual void RHISetViewport(uint32 MinX, uint32 MinY, float MinZ, uint32 MaxX, uint32 MaxY, float MaxZ) final override;
 	virtual void RHISetScissorRect(bool bEnable, uint32 MinX, uint32 MinY, uint32 MaxX, uint32 MaxY) final override;
-	virtual void RHISetBoundShaderState(FBoundShaderStateRHIParamRef BoundShaderState) final override;
+	virtual void RHISetBoundShaderState(FRHIBoundShaderState* BoundShaderState) final override;
 	virtual void RHISetShaderTexture(FVertexShaderRHIParamRef VertexShader, uint32 TextureIndex, FTextureRHIParamRef NewTexture) final override;
 	virtual void RHISetShaderTexture(FHullShaderRHIParamRef HullShader, uint32 TextureIndex, FTextureRHIParamRef NewTexture) final override;
 	virtual void RHISetShaderTexture(FDomainShaderRHIParamRef DomainShader, uint32 TextureIndex, FTextureRHIParamRef NewTexture) final override;
diff --git a/Engine/Source/Runtime/RHI/Private/BoundShaderStateCache.cpp b/Engine/Source/Runtime/RHI/Private/BoundShaderStateCache.cpp
--- a/Engine/Source/Runtime/RHI/Private/BoundShaderStateCache.cpp
+++ b/Engine/Source/Runtime/RHI/Private/BoundShaderStateCache.cpp
@@ -36,7 +36,7 @@ FCachedBoundShaderStateLink::FCachedBoundShaderStateLink(
 	FHullShaderRHIParamRef HullShader,
 	FDomainShaderRHIParamRef DomainShader,
 	FGeometryShaderRHIParamRef GeometryShader,
-	FBoundShaderStateRHIParamRef InBoundShaderState,
+	FRHIBoundShaderState* InBoundShaderState,
 	bool bAddToSingleThreadedCache
 	):
 	BoundShaderState(InBoundShaderState),
@@ -53,7 +53,7 @@ FCachedBoundShaderStateLink::FCachedBoundShaderStateLink(
 	FVertexDeclarationRHIParamRef VertexDeclaration,
 	FVertexShaderRHIParamRef VertexShader,
 	FPixelShaderRHIParamRef PixelShader,
-	FBoundShaderStateRHIParamRef InBoundShaderState,
+	FRHIBoundShaderState* InBoundShaderState,
 	bool bAddToSingleThreadedCache
 	):
 	BoundShaderState(InBoundShaderState),
diff --git a/Engine/Source/Runtime/RHI/Public/BoundShaderStateCache.h b/Engine/Source/Runtime/RHI/Public/BoundShaderStateCache.h
--- a/Engine/Source/Runtime/RHI/Public/BoundShaderStateCache.h
+++ b/Engine/Source/Runtime/RHI/Public/BoundShaderStateCache.h
@@ -160,14 +160,14 @@ public:
 	 * The cached bound shader state.  This is not a reference counted pointer because we rely on the RHI to destruct this object
 	 * when the bound shader state this references is destructed.
 	 */
-	FBoundShaderStateRHIParamRef BoundShaderState;
+	FRHIBoundShaderState* BoundShaderState;
 
 	/** Adds the bound shader state to the cache. */
 	FCachedBoundShaderStateLink(
 		FVertexDeclarationRHIParamRef VertexDeclaration,
 		FVertexShaderRHIParamRef VertexShader,
 		FPixelShaderRHIParamRef PixelShader,
-		FBoundShaderStateRHIParamRef InBoundShaderState,
+		FRHIBoundShaderState* InBoundShaderState,
 		bool bAddToSingleThreadedCache = true
 		);
 
@@ -179,7 +179,7 @@ public:
 		FHullShaderRHIParamRef HullShader,
 		FDomainShaderRHIParamRef DomainShader,
 		FGeometryShaderRHIParamRef GeometryShader,
-		FBoundShaderStateRHIParamRef InBoundShaderState,
+		FRHIBoundShaderState* InBoundShaderState,
 		bool bAddToSingleThreadedCache = true
 		);
 
@@ -229,7 +229,7 @@ public:
 		FVertexDeclarationRHIParamRef VertexDeclaration,
 		FVertexShaderRHIParamRef VertexShader,
 		FPixelShaderRHIParamRef PixelShader,
-		FBoundShaderStateRHIParamRef InBoundShaderState
+		FRHIBoundShaderState* InBoundShaderState
 		)
 		: FCachedBoundShaderStateLink(VertexDeclaration, VertexShader, PixelShader, InBoundShaderState, false)
 	{
@@ -243,7 +243,7 @@ public:
 		FHullShaderRHIParamRef HullShader,
 		FDomainShaderRHIParamRef DomainShader,
 		FGeometryShaderRHIParamRef GeometryShader,
-		FBoundShaderStateRHIParamRef InBoundShaderState
+		FRHIBoundShaderState* InBoundShaderState
 		)
 		: FCachedBoundShaderStateLink(VertexDeclaration, VertexShader, PixelShader, HullShader, DomainShader, GeometryShader, InBoundShaderState, false)
 	{
diff --git a/Engine/Source/Runtime/RHI/Public/RHIContext.h b/Engine/Source/Runtime/RHI/Public/RHIContext.h
--- a/Engine/Source/Runtime/RHI/Public/RHIContext.h
+++ b/Engine/Source/Runtime/RHI/Public/RHIContext.h
@@ -666,7 +666,7 @@ public:
 	* Set bound shader state. This will set the vertex decl/shader, and pixel shader
 	* @param BoundShaderState - state resource
 	*/
-	virtual void RHISetBoundShaderState(FBoundShaderStateRHIParamRef BoundShaderState) = 0;
+	virtual void RHISetBoundShaderState(FRHIBoundShaderState* BoundShaderState) = 0;
 
 	virtual void RHISetDepthStencilState(FRHIDepthStencilState* NewState, uint32 StencilRef) = 0;
 
diff --git a/Engine/Source/Runtime/RHI/Public/RHIResources.h b/Engine/Source/Runtime/RHI/Public/RHIResources.h
--- a/Engine/Source/Runtime/RHI/Public/RHIResources.h
+++ b/Engine/Source/Runtime/RHI/Public/RHIResources.h
@@ -1140,6 +1140,7 @@ typedef TRefCountPtr<FRHIRayTracingShader>          FRayTracingShaderRHIRef;
 typedef FRHIComputeFence*				FComputeFenceRHIParamRef;
 typedef TRefCountPtr<FRHIComputeFence>	FComputeFenceRHIRef;
 
+UE_DEPRECATED(4.23, "FBoundShaderStateRHIParamRef typedef is deprecated; please use FRHIBoundShaderState* directly instead.")
 typedef FRHIBoundShaderState*              FBoundShaderStateRHIParamRef;
 typedef TRefCountPtr<FRHIBoundShaderState> FBoundShaderStateRHIRef;
 
diff --git a/Engine/Source/Runtime/RenderCore/Private/StaticBoundShaderState.cpp b/Engine/Source/Runtime/RenderCore/Private/StaticBoundShaderState.cpp
--- a/Engine/Source/Runtime/RenderCore/Private/StaticBoundShaderState.cpp
+++ b/Engine/Source/Runtime/RenderCore/Private/StaticBoundShaderState.cpp
@@ -50,7 +50,7 @@ FGlobalBoundShaderStateResource::~FGlobalBoundShaderStateResource()
 /**
  * Initializes a global bound shader state with a vanilla bound shader state and required information.
  */
-FBoundShaderStateRHIParamRef FGlobalBoundShaderStateResource::GetInitializedRHI(
+FRHIBoundShaderState* FGlobalBoundShaderStateResource::GetInitializedRHI(
 	FVertexDeclarationRHIParamRef VertexDeclaration,
 	FVertexShaderRHIParamRef VertexShader, 
 	FPixelShaderRHIParamRef PixelShader,
@@ -94,7 +94,7 @@ FBoundShaderStateRHIParamRef FGlobalBoundShaderStateResource::GetInitializedRHI(
 	return BoundShaderState;
 }
 
-FBoundShaderStateRHIParamRef FGlobalBoundShaderStateResource::GetPreinitializedRHI()
+FRHIBoundShaderState* FGlobalBoundShaderStateResource::GetPreinitializedRHI()
 {
 	check(IsInitialized());
 	check(GIsRHIInitialized);
diff --git a/Engine/Source/Runtime/RenderCore/Public/RenderResource.h b/Engine/Source/Runtime/RenderCore/Public/RenderResource.h
--- a/Engine/Source/Runtime/RenderCore/Public/RenderResource.h
+++ b/Engine/Source/Runtime/RenderCore/Public/RenderResource.h
@@ -768,7 +768,7 @@ public:
 	{}
 
 	/** Adds a bound shader state to the history. */
-	FORCEINLINE void Add(FBoundShaderStateRHIParamRef BoundShaderState)
+	FORCEINLINE void Add(FRHIBoundShaderState* BoundShaderState)
 	{
 		if (TThreadSafe && GRHISupportsParallelRHIExecute)
 		{
@@ -782,7 +782,7 @@ public:
 		}
 	}
 
-	FBoundShaderStateRHIParamRef GetLast()
+	FRHIBoundShaderState* GetLast()
 	{
 		check(!GRHISupportsParallelRHIExecute);
 		// % doesn't work as we want on negative numbers, so handle the wraparound manually
diff --git a/Engine/Source/Runtime/RenderCore/Public/StaticBoundShaderState.h b/Engine/Source/Runtime/RenderCore/Public/StaticBoundShaderState.h
--- a/Engine/Source/Runtime/RenderCore/Public/StaticBoundShaderState.h
+++ b/Engine/Source/Runtime/RenderCore/Public/StaticBoundShaderState.h
@@ -34,7 +34,7 @@ public:
 	 * If this global bound shader state hasn't been initialized yet, initialize it.
 	 * @return The bound shader state RHI.
 	 */
-	FBoundShaderStateRHIParamRef GetInitializedRHI(
+	FRHIBoundShaderState* GetInitializedRHI(
 		FVertexDeclarationRHIParamRef VertexDeclaration, 
 		FVertexShaderRHIParamRef VertexShader, 
 		FPixelShaderRHIParamRef PixelShader,
@@ -45,7 +45,7 @@ public:
 	 * If this global bound shader state has been initialized return it, otherwise return null. Can be called from any thread.
 	 * @return The bound shader state RHI.
 	 */
-	FBoundShaderStateRHIParamRef GetPreinitializedRHI();
+	FRHIBoundShaderState* GetPreinitializedRHI();
 
 private:
 
diff --git a/Engine/Source/Runtime/Windows/D3D11RHI/Private/D3D11Commands.cpp b/Engine/Source/Runtime/Windows/D3D11RHI/Private/D3D11Commands.cpp
--- a/Engine/Source/Runtime/Windows/D3D11RHI/Private/D3D11Commands.cpp
+++ b/Engine/Source/Runtime/Windows/D3D11RHI/Private/D3D11Commands.cpp
@@ -278,7 +278,7 @@ void FD3D11DynamicRHI::RHISetScissorRect(bool bEnable,uint32 MinX,uint32 MinY,ui
 * Set bound shader state. This will set the vertex decl/shader, and pixel shader
 * @param BoundShaderState - state resource
 */
-void FD3D11DynamicRHI::RHISetBoundShaderState( FBoundShaderStateRHIParamRef BoundShaderStateRHI)
+void FD3D11DynamicRHI::RHISetBoundShaderState(FRHIBoundShaderState* BoundShaderStateRHI)
 {
 	FD3D11BoundShaderState* BoundShaderState = ResourceCast(BoundShaderStateRHI);
 
diff --git a/Engine/Source/Runtime/Windows/D3D11RHI/Private/D3D11RHIPrivate.h b/Engine/Source/Runtime/Windows/D3D11RHI/Private/D3D11RHIPrivate.h
--- a/Engine/Source/Runtime/Windows/D3D11RHI/Private/D3D11RHIPrivate.h
+++ b/Engine/Source/Runtime/Windows/D3D11RHI/Private/D3D11RHIPrivate.h
@@ -520,7 +520,7 @@ public:
 	virtual void RHISetRasterizerState(FRHIRasterizerState* NewState) final override;
 	virtual void RHISetViewport(uint32 MinX, uint32 MinY, float MinZ, uint32 MaxX, uint32 MaxY, float MaxZ) final override;
 	virtual void RHISetScissorRect(bool bEnable, uint32 MinX, uint32 MinY, uint32 MaxX, uint32 MaxY) final override;
-	virtual void RHISetBoundShaderState(FBoundShaderStateRHIParamRef BoundShaderState) final override;
+	virtual void RHISetBoundShaderState(FRHIBoundShaderState* BoundShaderState) final override;
 	virtual void RHISetGraphicsPipelineState(FRHIGraphicsPipelineState* GraphicsState) final override
 	{
 		FRHIGraphicsPipelineStateFallBack* FallbackGraphicsState = static_cast<FRHIGraphicsPipelineStateFallBack*>(GraphicsState);
