DR - Deprecate FVertexDeclarationRHIParamRef and FComputeFenceRHIParamRef

diff --git a/Engine/Source/Runtime/Apple/MetalRHI/Private/MetalShaders.cpp b/Engine/Source/Runtime/Apple/MetalRHI/Private/MetalShaders.cpp
--- a/Engine/Source/Runtime/Apple/MetalRHI/Private/MetalShaders.cpp
+++ b/Engine/Source/Runtime/Apple/MetalRHI/Private/MetalShaders.cpp
@@ -1584,7 +1584,7 @@ FRHIShaderLibraryRef FMetalDynamicRHI::RHICreateShaderLibrary(EShaderPlatform Pl
 }
 
 FBoundShaderStateRHIRef FMetalDynamicRHI::RHICreateBoundShaderState(
-	FVertexDeclarationRHIParamRef VertexDeclarationRHI, 
+	FRHIVertexDeclaration* VertexDeclarationRHI,
 	FVertexShaderRHIParamRef VertexShaderRHI, 
 	FHullShaderRHIParamRef HullShaderRHI, 
 	FDomainShaderRHIParamRef DomainShaderRHI, 
diff --git a/Engine/Source/Runtime/Apple/MetalRHI/Private/MetalUAV.cpp b/Engine/Source/Runtime/Apple/MetalRHI/Private/MetalUAV.cpp
--- a/Engine/Source/Runtime/Apple/MetalRHI/Private/MetalUAV.cpp
+++ b/Engine/Source/Runtime/Apple/MetalRHI/Private/MetalUAV.cpp
@@ -715,7 +715,7 @@ void FMetalComputeFence::Reset()
 	Fence = nullptr;
 }
 
-void FMetalRHICommandContext::RHITransitionResources(EResourceTransitionAccess TransitionType, EResourceTransitionPipeline TransitionPipeline, FUnorderedAccessViewRHIParamRef* InUAVs, int32 NumUAVs, FComputeFenceRHIParamRef WriteComputeFence)
+void FMetalRHICommandContext::RHITransitionResources(EResourceTransitionAccess TransitionType, EResourceTransitionPipeline TransitionPipeline, FUnorderedAccessViewRHIParamRef* InUAVs, int32 NumUAVs, FRHIComputeFence* WriteComputeFence)
 {
 	@autoreleasepool
 	{
@@ -761,7 +761,7 @@ void FMetalRHICommandContext::RHITransitionResources(EResourceTransitionAccess T
 	}
 }
 
-void FMetalRHICommandContext::RHIWaitComputeFence(FComputeFenceRHIParamRef InFence)
+void FMetalRHICommandContext::RHIWaitComputeFence(FRHIComputeFence* InFence)
 {
 	@autoreleasepool {
 	if (InFence)
diff --git a/Engine/Source/Runtime/Apple/MetalRHI/Public/MetalRHI.h b/Engine/Source/Runtime/Apple/MetalRHI/Public/MetalRHI.h
--- a/Engine/Source/Runtime/Apple/MetalRHI/Public/MetalRHI.h
+++ b/Engine/Source/Runtime/Apple/MetalRHI/Public/MetalRHI.h
@@ -59,7 +59,7 @@ public:
 	virtual FGeometryShaderRHIRef RHICreateGeometryShaderWithStreamOutput(const FStreamOutElementList& ElementList, uint32 NumStrides, const uint32* Strides, int32 RasterizedStream, FRHIShaderLibraryParamRef Library, FSHAHash Hash) final override;
 	virtual FComputeShaderRHIRef RHICreateComputeShader(const TArray<uint8>& Code) final override;
 	virtual FComputeShaderRHIRef RHICreateComputeShader(FRHIShaderLibraryParamRef Library, FSHAHash Hash) final override;
-	virtual FBoundShaderStateRHIRef RHICreateBoundShaderState(FVertexDeclarationRHIParamRef VertexDeclaration, FVertexShaderRHIParamRef VertexShader, FHullShaderRHIParamRef HullShader, FDomainShaderRHIParamRef DomainShader, FPixelShaderRHIParamRef PixelShader, FGeometryShaderRHIParamRef GeometryShader) final override;
+	virtual FBoundShaderStateRHIRef RHICreateBoundShaderState(FRHIVertexDeclaration* VertexDeclaration, FVertexShaderRHIParamRef VertexShader, FHullShaderRHIParamRef HullShader, FDomainShaderRHIParamRef DomainShader, FPixelShaderRHIParamRef PixelShader, FGeometryShaderRHIParamRef GeometryShader) final override;
 	virtual FGraphicsPipelineStateRHIRef RHICreateGraphicsPipelineState(const FGraphicsPipelineStateInitializer& Initializer) final override;
 	virtual TRefCountPtr<FRHIComputePipelineState> RHICreateComputePipelineState(FRHIComputeShader* ComputeShader) final override;
 	virtual FUniformBufferRHIRef RHICreateUniformBuffer(const void* Contents, const FRHIUniformBufferLayout& Layout, EUniformBufferUsage Usage, EUniformBufferValidation Validation) final override;
diff --git a/Engine/Source/Runtime/Apple/MetalRHI/Public/MetalRHIContext.h b/Engine/Source/Runtime/Apple/MetalRHI/Public/MetalRHIContext.h
--- a/Engine/Source/Runtime/Apple/MetalRHI/Public/MetalRHIContext.h
+++ b/Engine/Source/Runtime/Apple/MetalRHI/Public/MetalRHIContext.h
@@ -265,7 +265,7 @@ public:
 	 * @param NumUAVs - number of UAVs to transition
 	 * @param WriteComputeFence - Optional ComputeFence to write as part of this transition
 	 */
-	virtual void RHITransitionResources(EResourceTransitionAccess TransitionType, EResourceTransitionPipeline TransitionPipeline, FUnorderedAccessViewRHIParamRef* InUAVs, int32 NumUAVs, FComputeFenceRHIParamRef WriteComputeFence) final override;
+	virtual void RHITransitionResources(EResourceTransitionAccess TransitionType, EResourceTransitionPipeline TransitionPipeline, FUnorderedAccessViewRHIParamRef* InUAVs, int32 NumUAVs, FRHIComputeFence* WriteComputeFence) final override;
 	
 	virtual void RHICopyToStagingBuffer(FVertexBufferRHIParamRef SourceBufferRHI, FStagingBufferRHIParamRef DestinationStagingBufferRHI, uint32 Offset, uint32 NumBytes) final override;
 	virtual void RHIWriteGPUFence(FGPUFenceRHIParamRef FenceRHI) final override;
@@ -284,7 +284,7 @@ public:
 	/**
 	 * Compute queue will wait for the fence to be written before continuing.
 	 */
-	virtual void RHIWaitComputeFence(FComputeFenceRHIParamRef InFence) final override;
+	virtual void RHIWaitComputeFence(FRHIComputeFence* InFence) final override;
 
 	virtual void RHIBeginRenderPass(const FRHIRenderPassInfo& InInfo, const TCHAR* InName) final override;
 
diff --git a/Engine/Source/Runtime/D3D12RHI/Private/D3D12CommandContext.cpp b/Engine/Source/Runtime/D3D12RHI/Private/D3D12CommandContext.cpp
--- a/Engine/Source/Runtime/D3D12RHI/Private/D3D12CommandContext.cpp
+++ b/Engine/Source/Runtime/D3D12RHI/Private/D3D12CommandContext.cpp
@@ -688,7 +688,7 @@ FD3D12CommandContextRedirector::FD3D12CommandContextRedirector(class FD3D12Adapt
 	FMemory::Memzero(PhysicalContexts, sizeof(PhysicalContexts[0]) * MAX_NUM_GPUS);
 }
 
-void FD3D12CommandContextRedirector::RHITransitionResources(EResourceTransitionAccess TransitionType, EResourceTransitionPipeline TransitionPipeline, FUnorderedAccessViewRHIParamRef* InUAVs, int32 NumUAVs, FComputeFenceRHIParamRef WriteComputeFenceRHI)
+void FD3D12CommandContextRedirector::RHITransitionResources(EResourceTransitionAccess TransitionType, EResourceTransitionPipeline TransitionPipeline, FUnorderedAccessViewRHIParamRef* InUAVs, int32 NumUAVs, FRHIComputeFence* WriteComputeFenceRHI)
 {
 	ContextRedirect(RHITransitionResources(TransitionType, TransitionPipeline, InUAVs, NumUAVs, nullptr));
 
diff --git a/Engine/Source/Runtime/D3D12RHI/Private/D3D12CommandContext.h b/Engine/Source/Runtime/D3D12RHI/Private/D3D12CommandContext.h
--- a/Engine/Source/Runtime/D3D12RHI/Private/D3D12CommandContext.h
+++ b/Engine/Source/Runtime/D3D12RHI/Private/D3D12CommandContext.h
@@ -41,7 +41,7 @@ public:
 	void RHIBeginFrame() final override;
 	void RHIEndFrame() final override;
 
-	void RHIWaitComputeFence(FComputeFenceRHIParamRef InFence) final override;
+	void RHIWaitComputeFence(FRHIComputeFence* InFence) final override;
 
 	virtual void UpdateMemoryStats();
 
@@ -232,7 +232,7 @@ public:
 	virtual void RHISetComputePipelineState(FRHIComputePipelineState* ComputePipelineState) final override;
 	virtual void RHIDispatchComputeShader(uint32 ThreadGroupCountX, uint32 ThreadGroupCountY, uint32 ThreadGroupCountZ) final override;
 	virtual void RHIDispatchIndirectComputeShader(FVertexBufferRHIParamRef ArgumentBuffer, uint32 ArgumentOffset) final override;
-	virtual void RHITransitionResources(EResourceTransitionAccess TransitionType, EResourceTransitionPipeline TransitionPipeline, FUnorderedAccessViewRHIParamRef* InUAVs, int32 NumUAVs, FComputeFenceRHIParamRef WriteComputeFenceRHI) final override;
+	virtual void RHITransitionResources(EResourceTransitionAccess TransitionType, EResourceTransitionPipeline TransitionPipeline, FUnorderedAccessViewRHIParamRef* InUAVs, int32 NumUAVs, FRHIComputeFence* WriteComputeFenceRHI) final override;
 	virtual void RHISetShaderTexture(FComputeShaderRHIParamRef PixelShader, uint32 TextureIndex, FTextureRHIParamRef NewTexture) final override;
 	virtual void RHISetShaderSampler(FComputeShaderRHIParamRef ComputeShader, uint32 SamplerIndex, FRHISamplerState* NewState) final override;
 	virtual void RHISetUAVParameter(FComputeShaderRHIParamRef ComputeShader, uint32 UAVIndex, FUnorderedAccessViewRHIParamRef UAV) final override;
@@ -460,7 +460,7 @@ public:
 		ContextRedirect(RHIDispatchIndirectComputeShader(ArgumentBuffer, ArgumentOffset));
 	}
 	// Special implementation that only signal the fence once.
-	virtual void RHITransitionResources(EResourceTransitionAccess TransitionType, EResourceTransitionPipeline TransitionPipeline, FUnorderedAccessViewRHIParamRef* InUAVs, int32 NumUAVs, FComputeFenceRHIParamRef WriteComputeFenceRHI) final override;
+	virtual void RHITransitionResources(EResourceTransitionAccess TransitionType, EResourceTransitionPipeline TransitionPipeline, FUnorderedAccessViewRHIParamRef* InUAVs, int32 NumUAVs, FRHIComputeFence* WriteComputeFenceRHI) final override;
 
 	FORCEINLINE virtual void RHICopyToStagingBuffer(FVertexBufferRHIParamRef SourceBuffer, FStagingBufferRHIParamRef DestinationStagingBuffer, uint32 Offset, uint32 NumBytes) final override
 	{
diff --git a/Engine/Source/Runtime/D3D12RHI/Private/D3D12Commands.cpp b/Engine/Source/Runtime/D3D12RHI/Private/D3D12Commands.cpp
--- a/Engine/Source/Runtime/D3D12RHI/Private/D3D12Commands.cpp
+++ b/Engine/Source/Runtime/D3D12RHI/Private/D3D12Commands.cpp
@@ -133,7 +133,7 @@ void FD3D12CommandContext::RHISetComputeShader(FComputeShaderRHIParamRef Compute
 	RHISetComputePipelineState(ComputePipelineState);
 }
 
-void FD3D12CommandContextBase::RHIWaitComputeFence(FComputeFenceRHIParamRef InFenceRHI)
+void FD3D12CommandContextBase::RHIWaitComputeFence(FRHIComputeFence* InFenceRHI)
 {
 	FD3D12Fence* Fence = FD3D12DynamicRHI::ResourceCast(InFenceRHI);
 
@@ -278,7 +278,7 @@ void FD3D12CommandContext::RHITransitionResources(EResourceTransitionAccess Tran
 }
 
 
-void FD3D12CommandContext::RHITransitionResources(EResourceTransitionAccess TransitionType, EResourceTransitionPipeline TransitionPipeline, FUnorderedAccessViewRHIParamRef* InUAVs, int32 InNumUAVs, FComputeFenceRHIParamRef WriteComputeFenceRHI)
+void FD3D12CommandContext::RHITransitionResources(EResourceTransitionAccess TransitionType, EResourceTransitionPipeline TransitionPipeline, FUnorderedAccessViewRHIParamRef* InUAVs, int32 InNumUAVs, FRHIComputeFence* WriteComputeFenceRHI)
 {
 	static IConsoleVariable* CVarShowTransitions = IConsoleManager::Get().FindConsoleVariable(TEXT("r.ProfileGPU.ShowTransitions"));
 	const bool bShowTransitionEvents = CVarShowTransitions->GetInt() != 0;
diff --git a/Engine/Source/Runtime/D3D12RHI/Private/D3D12RHIPrivate.h b/Engine/Source/Runtime/D3D12RHI/Private/D3D12RHIPrivate.h
--- a/Engine/Source/Runtime/D3D12RHI/Private/D3D12RHIPrivate.h
+++ b/Engine/Source/Runtime/D3D12RHI/Private/D3D12RHIPrivate.h
@@ -275,7 +275,7 @@ public:
 	virtual FStagingBufferRHIRef RHICreateStagingBuffer() final override;
 	virtual void* RHILockStagingBuffer(FStagingBufferRHIParamRef StagingBuffer, uint32 Offset, uint32 SizeRHI) final override;
     virtual void RHIUnlockStagingBuffer(FStagingBufferRHIParamRef StagingBuffer) final override;
-	virtual FBoundShaderStateRHIRef RHICreateBoundShaderState(FVertexDeclarationRHIParamRef VertexDeclaration, FVertexShaderRHIParamRef VertexShader, FHullShaderRHIParamRef HullShader, FDomainShaderRHIParamRef DomainShader, FPixelShaderRHIParamRef PixelShader, FGeometryShaderRHIParamRef GeometryShader) final override;
+	virtual FBoundShaderStateRHIRef RHICreateBoundShaderState(FRHIVertexDeclaration* VertexDeclaration, FVertexShaderRHIParamRef VertexShader, FHullShaderRHIParamRef HullShader, FDomainShaderRHIParamRef DomainShader, FPixelShaderRHIParamRef PixelShader, FGeometryShaderRHIParamRef GeometryShader) final override;
 	virtual FGraphicsPipelineStateRHIRef RHICreateGraphicsPipelineState(const FGraphicsPipelineStateInitializer& Initializer) final override;
 	virtual TRefCountPtr<FRHIComputePipelineState> RHICreateComputePipelineState(FRHIComputeShader* ComputeShader) final override;
 	virtual FUniformBufferRHIRef RHICreateUniformBuffer(const void* Contents, const FRHIUniformBufferLayout& Layout, EUniformBufferUsage Usage, EUniformBufferValidation Validation) final override;
diff --git a/Engine/Source/Runtime/D3D12RHI/Private/D3D12Shader.h b/Engine/Source/Runtime/D3D12RHI/Private/D3D12Shader.h
--- a/Engine/Source/Runtime/D3D12RHI/Private/D3D12Shader.h
+++ b/Engine/Source/Runtime/D3D12RHI/Private/D3D12Shader.h
@@ -183,7 +183,7 @@ public:
 
 	/** Initialization constructor. */
 	FD3D12BoundShaderState(
-		FVertexDeclarationRHIParamRef InVertexDeclarationRHI,
+		FRHIVertexDeclaration* InVertexDeclarationRHI,
 		FVertexShaderRHIParamRef InVertexShaderRHI,
 		FPixelShaderRHIParamRef InPixelShaderRHI,
 		FHullShaderRHIParamRef InHullShaderRHI,
diff --git a/Engine/Source/Runtime/D3D12RHI/Private/D3D12Shaders.cpp b/Engine/Source/Runtime/D3D12RHI/Private/D3D12Shaders.cpp
--- a/Engine/Source/Runtime/D3D12RHI/Private/D3D12Shaders.cpp
+++ b/Engine/Source/Runtime/D3D12RHI/Private/D3D12Shaders.cpp
@@ -324,7 +324,7 @@ void FD3D12CommandContext::RHISetMultipleViewports(uint32 Count, const FViewport
 }
 
 FD3D12BoundShaderState::FD3D12BoundShaderState(
-	FVertexDeclarationRHIParamRef InVertexDeclarationRHI,
+	FRHIVertexDeclaration* InVertexDeclarationRHI,
 	FVertexShaderRHIParamRef InVertexShaderRHI,
 	FPixelShaderRHIParamRef InPixelShaderRHI,
 	FHullShaderRHIParamRef InHullShaderRHI,
@@ -371,7 +371,7 @@ FD3D12BoundShaderState::~FD3D12BoundShaderState()
 * @param GeometryShader - existing geometry shader
 */
 FBoundShaderStateRHIRef FD3D12DynamicRHI::RHICreateBoundShaderState(
-	FVertexDeclarationRHIParamRef VertexDeclarationRHI,
+	FRHIVertexDeclaration* VertexDeclarationRHI,
 	FVertexShaderRHIParamRef VertexShaderRHI,
 	FHullShaderRHIParamRef HullShaderRHI,
 	FDomainShaderRHIParamRef DomainShaderRHI,
diff --git a/Engine/Source/Runtime/EmptyRHI/Private/EmptyShaders.cpp b/Engine/Source/Runtime/EmptyRHI/Private/EmptyShaders.cpp
--- a/Engine/Source/Runtime/EmptyRHI/Private/EmptyShaders.cpp
+++ b/Engine/Source/Runtime/EmptyRHI/Private/EmptyShaders.cpp
@@ -68,7 +68,7 @@ FComputeShaderRHIRef FEmptyDynamicRHI::RHICreateComputeShader(const TArray<uint8
 
 
 FEmptyBoundShaderState::FEmptyBoundShaderState(
-			FVertexDeclarationRHIParamRef InVertexDeclarationRHI,
+			FRHIVertexDeclaration* InVertexDeclarationRHI,
 			FVertexShaderRHIParamRef InVertexShaderRHI,
 			FPixelShaderRHIParamRef InPixelShaderRHI,
 			FHullShaderRHIParamRef InHullShaderRHI,
@@ -98,7 +98,7 @@ FEmptyBoundShaderState::~FEmptyBoundShaderState()
 }
 
 FBoundShaderStateRHIRef FEmptyDynamicRHI::RHICreateBoundShaderState(
-	FVertexDeclarationRHIParamRef VertexDeclarationRHI, 
+	FRHIVertexDeclaration* VertexDeclarationRHI,
 	FVertexShaderRHIParamRef VertexShaderRHI, 
 	FHullShaderRHIParamRef HullShaderRHI, 
 	FDomainShaderRHIParamRef DomainShaderRHI, 
diff --git a/Engine/Source/Runtime/EmptyRHI/Public/EmptyRHI.h b/Engine/Source/Runtime/EmptyRHI/Public/EmptyRHI.h
--- a/Engine/Source/Runtime/EmptyRHI/Public/EmptyRHI.h
+++ b/Engine/Source/Runtime/EmptyRHI/Public/EmptyRHI.h
@@ -48,7 +48,7 @@ public:
 	virtual FGeometryShaderRHIRef RHICreateGeometryShader(const TArray<uint8>& Code) final override;
 	virtual FGeometryShaderRHIRef RHICreateGeometryShaderWithStreamOutput(const TArray<uint8>& Code, const FStreamOutElementList& ElementList, uint32 NumStrides, const uint32* Strides, int32 RasterizedStream) final override;
 	virtual FComputeShaderRHIRef RHICreateComputeShader(const TArray<uint8>& Code) final override;
-	virtual FBoundShaderStateRHIRef RHICreateBoundShaderState(FVertexDeclarationRHIParamRef VertexDeclaration, FVertexShaderRHIParamRef VertexShader, FHullShaderRHIParamRef HullShader, FDomainShaderRHIParamRef DomainShader, FPixelShaderRHIParamRef PixelShader, FGeometryShaderRHIParamRef GeometryShader) final override;
+	virtual FBoundShaderStateRHIRef RHICreateBoundShaderState(FRHIVertexDeclaration* VertexDeclaration, FVertexShaderRHIParamRef VertexShader, FHullShaderRHIParamRef HullShader, FDomainShaderRHIParamRef DomainShader, FPixelShaderRHIParamRef PixelShader, FGeometryShaderRHIParamRef GeometryShader) final override;
 	virtual FUniformBufferRHIRef RHICreateUniformBuffer(const void* Contents, const FRHIUniformBufferLayout& Layout, EUniformBufferUsage Usage, EUniformBufferValidation Validation) final override;
 	virtual FIndexBufferRHIRef RHICreateIndexBuffer(uint32 Stride, uint32 Size, uint32 InUsage, FRHIResourceCreateInfo& CreateInfo) final override;
 	virtual void* RHILockIndexBuffer(FIndexBufferRHIParamRef IndexBuffer, uint32 Offset, uint32 Size, EResourceLockMode LockMode) final override;
diff --git a/Engine/Source/Runtime/EmptyRHI/Public/EmptyResources.h b/Engine/Source/Runtime/EmptyRHI/Public/EmptyResources.h
--- a/Engine/Source/Runtime/EmptyRHI/Public/EmptyResources.h
+++ b/Engine/Source/Runtime/EmptyRHI/Public/EmptyResources.h
@@ -43,7 +43,7 @@ public:
 
 	/** Initialization constructor. */
 	FEmptyBoundShaderState(
-		FVertexDeclarationRHIParamRef InVertexDeclarationRHI,
+		FRHIVertexDeclaration* InVertexDeclarationRHI,
 		FVertexShaderRHIParamRef InVertexShaderRHI,
 		FPixelShaderRHIParamRef InPixelShaderRHI,
 		FHullShaderRHIParamRef InHullShaderRHI,
diff --git a/Engine/Source/Runtime/Engine/Private/Particles/ParticleGpuSimulation.cpp b/Engine/Source/Runtime/Engine/Private/Particles/ParticleGpuSimulation.cpp
--- a/Engine/Source/Runtime/Engine/Private/Particles/ParticleGpuSimulation.cpp
+++ b/Engine/Source/Runtime/Engine/Private/Particles/ParticleGpuSimulation.cpp
@@ -1360,7 +1360,7 @@ public:
 
 TGlobalResource<FInstancedParticleTileVertexDeclaration> GInstancedParticleTileVertexDeclaration;
 
-static FVertexDeclarationRHIParamRef GetParticleTileVertexDeclaration(ERHIFeatureLevel::Type FeatureLevel)
+static FRHIVertexDeclaration* GetParticleTileVertexDeclaration(ERHIFeatureLevel::Type FeatureLevel)
 {
 	if (FeatureLevel <= ERHIFeatureLevel::ES3_1)
 	{
diff --git a/Engine/Source/Runtime/NullDrv/Public/NullRHI.h b/Engine/Source/Runtime/NullDrv/Public/NullRHI.h
--- a/Engine/Source/Runtime/NullDrv/Public/NullRHI.h
+++ b/Engine/Source/Runtime/NullDrv/Public/NullRHI.h
@@ -78,7 +78,7 @@ public:
 	}
 
 
-	virtual FBoundShaderStateRHIRef RHICreateBoundShaderState(FVertexDeclarationRHIParamRef VertexDeclaration, FVertexShaderRHIParamRef VertexShader, FHullShaderRHIParamRef HullShader, FDomainShaderRHIParamRef DomainShader, FPixelShaderRHIParamRef PixelShader, FGeometryShaderRHIParamRef GeometryShader) final override
+	virtual FBoundShaderStateRHIRef RHICreateBoundShaderState(FRHIVertexDeclaration* VertexDeclaration, FVertexShaderRHIParamRef VertexShader, FHullShaderRHIParamRef HullShader, FDomainShaderRHIParamRef DomainShader, FPixelShaderRHIParamRef PixelShader, FGeometryShaderRHIParamRef GeometryShader) final override
 	{ 
 		return new FRHIBoundShaderState(); 
 	}
diff --git a/Engine/Source/Runtime/OpenGLDrv/Private/OpenGLShaders.cpp b/Engine/Source/Runtime/OpenGLDrv/Private/OpenGLShaders.cpp
--- a/Engine/Source/Runtime/OpenGLDrv/Private/OpenGLShaders.cpp
+++ b/Engine/Source/Runtime/OpenGLDrv/Private/OpenGLShaders.cpp
@@ -3521,7 +3521,7 @@ static void BindShaderStage(FOpenGLLinkedProgramConfiguration& Config, CrossComp
 // ============================================================================================================================
 
 FBoundShaderStateRHIRef FOpenGLDynamicRHI::RHICreateBoundShaderState_OnThisThread(
-	FVertexDeclarationRHIParamRef VertexDeclarationRHI, 
+	FRHIVertexDeclaration* VertexDeclarationRHI,
 	FVertexShaderRHIParamRef VertexShaderRHI, 
 	FHullShaderRHIParamRef HullShaderRHI,
 	FDomainShaderRHIParamRef DomainShaderRHI, 
@@ -3967,7 +3967,7 @@ void FOpenGLDynamicRHI::BindPendingShaderState( FOpenGLContextState& ContextStat
 
 FOpenGLBoundShaderState::FOpenGLBoundShaderState(
 	FOpenGLLinkedProgram* InLinkedProgram,
-	FVertexDeclarationRHIParamRef InVertexDeclarationRHI,
+	FRHIVertexDeclaration* InVertexDeclarationRHI,
 	FVertexShaderRHIParamRef InVertexShaderRHI,
 	FPixelShaderRHIParamRef InPixelShaderRHI,
 	FGeometryShaderRHIParamRef InGeometryShaderRHI,
diff --git a/Engine/Source/Runtime/OpenGLDrv/Public/OpenGLDrv.h b/Engine/Source/Runtime/OpenGLDrv/Public/OpenGLDrv.h
--- a/Engine/Source/Runtime/OpenGLDrv/Public/OpenGLDrv.h
+++ b/Engine/Source/Runtime/OpenGLDrv/Public/OpenGLDrv.h
@@ -998,7 +998,7 @@ public:
 	}
 
 	FBoundShaderStateRHIRef RHICreateBoundShaderState_internal(
-		FVertexDeclarationRHIParamRef VertexDeclarationRHI,
+		FRHIVertexDeclaration* VertexDeclarationRHI,
 		FVertexShaderRHIParamRef VertexShaderRHI,
 		FHullShaderRHIParamRef HullShaderRHI,
 		FDomainShaderRHIParamRef DomainShaderRHI,
@@ -1021,7 +1021,7 @@ public:
 	}
 
 	virtual FBoundShaderStateRHIRef RHICreateBoundShaderState(
-		FVertexDeclarationRHIParamRef VertexDeclarationRHI,
+		FRHIVertexDeclaration* VertexDeclarationRHI,
 		FVertexShaderRHIParamRef VertexShaderRHI,
 		FHullShaderRHIParamRef HullShaderRHI,
 		FDomainShaderRHIParamRef DomainShaderRHI,
@@ -1040,7 +1040,7 @@ public:
 	}
 
 
-	FBoundShaderStateRHIRef RHICreateBoundShaderState_OnThisThread(FVertexDeclarationRHIParamRef VertexDeclaration, FVertexShaderRHIParamRef VertexShader, FHullShaderRHIParamRef HullShader, FDomainShaderRHIParamRef DomainShader, FPixelShaderRHIParamRef PixelShader, FGeometryShaderRHIParamRef GeometryShader, bool FromPSOFileCache);
+	FBoundShaderStateRHIRef RHICreateBoundShaderState_OnThisThread(FRHIVertexDeclaration* VertexDeclaration, FVertexShaderRHIParamRef VertexShader, FHullShaderRHIParamRef HullShader, FDomainShaderRHIParamRef DomainShader, FPixelShaderRHIParamRef PixelShader, FGeometryShaderRHIParamRef GeometryShader, bool FromPSOFileCache);
 	void RHIPerFrameRHIFlushComplete();
 
 	FOpenGLGPUProfiler& GetGPUProfilingData() {
diff --git a/Engine/Source/Runtime/OpenGLDrv/Public/OpenGLResources.h b/Engine/Source/Runtime/OpenGLDrv/Public/OpenGLResources.h
--- a/Engine/Source/Runtime/OpenGLDrv/Public/OpenGLResources.h
+++ b/Engine/Source/Runtime/OpenGLDrv/Public/OpenGLResources.h
@@ -1166,7 +1166,7 @@ public:
 	/** Initialization constructor. */
 	FOpenGLBoundShaderState(
 		FOpenGLLinkedProgram* InLinkedProgram,
-		FVertexDeclarationRHIParamRef InVertexDeclarationRHI,
+		FRHIVertexDeclaration* InVertexDeclarationRHI,
 		FVertexShaderRHIParamRef InVertexShaderRHI,
 		FPixelShaderRHIParamRef InPixelShaderRHI,
 		FGeometryShaderRHIParamRef InGeometryShaderRHI,
diff --git a/Engine/Source/Runtime/RHI/Private/BoundShaderStateCache.cpp b/Engine/Source/Runtime/RHI/Private/BoundShaderStateCache.cpp
--- a/Engine/Source/Runtime/RHI/Private/BoundShaderStateCache.cpp
+++ b/Engine/Source/Runtime/RHI/Private/BoundShaderStateCache.cpp
@@ -30,7 +30,7 @@ static FCriticalSection BoundShaderStateCacheLock;
 
 
 FCachedBoundShaderStateLink::FCachedBoundShaderStateLink(
-	FVertexDeclarationRHIParamRef VertexDeclaration,
+	FRHIVertexDeclaration* VertexDeclaration,
 	FVertexShaderRHIParamRef VertexShader,
 	FPixelShaderRHIParamRef PixelShader,
 	FHullShaderRHIParamRef HullShader,
@@ -50,7 +50,7 @@ FCachedBoundShaderStateLink::FCachedBoundShaderStateLink(
 }
 
 FCachedBoundShaderStateLink::FCachedBoundShaderStateLink(
-	FVertexDeclarationRHIParamRef VertexDeclaration,
+	FRHIVertexDeclaration* VertexDeclaration,
 	FVertexShaderRHIParamRef VertexShader,
 	FPixelShaderRHIParamRef PixelShader,
 	FRHIBoundShaderState* InBoundShaderState,
@@ -76,7 +76,7 @@ FCachedBoundShaderStateLink::~FCachedBoundShaderStateLink()
 }
 
 FCachedBoundShaderStateLink* GetCachedBoundShaderState(
-	FVertexDeclarationRHIParamRef VertexDeclaration,
+	FRHIVertexDeclaration* VertexDeclaration,
 	FVertexShaderRHIParamRef VertexShader,
 	FPixelShaderRHIParamRef PixelShader,
 	FHullShaderRHIParamRef HullShader,
@@ -104,7 +104,7 @@ void FCachedBoundShaderStateLink_Threadsafe::RemoveFromCache()
 
 
 FBoundShaderStateRHIRef GetCachedBoundShaderState_Threadsafe(
-	FVertexDeclarationRHIParamRef VertexDeclaration,
+	FRHIVertexDeclaration* VertexDeclaration,
 	FVertexShaderRHIParamRef VertexShader,
 	FPixelShaderRHIParamRef PixelShader,
 	FHullShaderRHIParamRef HullShader,
diff --git a/Engine/Source/Runtime/RHI/Public/BoundShaderStateCache.h b/Engine/Source/Runtime/RHI/Public/BoundShaderStateCache.h
--- a/Engine/Source/Runtime/RHI/Public/BoundShaderStateCache.h
+++ b/Engine/Source/Runtime/RHI/Public/BoundShaderStateCache.h
@@ -18,7 +18,7 @@ class FBoundShaderStateKey
 public:
 	/** Initialization constructor. */
 	FBoundShaderStateKey(
-		FVertexDeclarationRHIParamRef InVertexDeclaration, 
+		FRHIVertexDeclaration* InVertexDeclaration,
 		FVertexShaderRHIParamRef InVertexShader, 
 		FPixelShaderRHIParamRef InPixelShader,
 		FHullShaderRHIParamRef InHullShader = NULL,
@@ -45,7 +45,7 @@ public:
 	/**
 	* Get the RHI vertex declaration.
 	*/
-	FORCEINLINE FVertexDeclarationRHIParamRef GetVertexDeclaration() const { return VertexDeclaration; }
+	FORCEINLINE FRHIVertexDeclaration* GetVertexDeclaration() const { return VertexDeclaration; }
 
 private:
 	/**
@@ -164,7 +164,7 @@ public:
 
 	/** Adds the bound shader state to the cache. */
 	FCachedBoundShaderStateLink(
-		FVertexDeclarationRHIParamRef VertexDeclaration,
+		FRHIVertexDeclaration* VertexDeclaration,
 		FVertexShaderRHIParamRef VertexShader,
 		FPixelShaderRHIParamRef PixelShader,
 		FRHIBoundShaderState* InBoundShaderState,
@@ -173,7 +173,7 @@ public:
 
 	/** Adds the bound shader state to the cache. */
 	FCachedBoundShaderStateLink(
-		FVertexDeclarationRHIParamRef VertexDeclaration,
+		FRHIVertexDeclaration* VertexDeclaration,
 		FVertexShaderRHIParamRef VertexShader,
 		FPixelShaderRHIParamRef PixelShader,
 		FHullShaderRHIParamRef HullShader,
@@ -198,7 +198,7 @@ public:
 	/**
 	* Get the RHI vertex declaration.
 	*/
-	FORCEINLINE FVertexDeclarationRHIParamRef GetVertexDeclaration() const { return Key.GetVertexDeclaration(); }
+	FORCEINLINE FRHIVertexDeclaration* GetVertexDeclaration() const { return Key.GetVertexDeclaration(); }
 
 protected:
 	FBoundShaderStateKey Key;
@@ -211,7 +211,7 @@ protected:
  * @return If a bound shader state matching the parameters is cached, it is returned; otherwise NULL is returned.
  */
 extern RHI_API FCachedBoundShaderStateLink* GetCachedBoundShaderState(
-	FVertexDeclarationRHIParamRef VertexDeclaration,
+	FRHIVertexDeclaration* VertexDeclaration,
 	FVertexShaderRHIParamRef VertexShader,
 	FPixelShaderRHIParamRef PixelShader,
 	FHullShaderRHIParamRef HullShader = NULL,
@@ -226,7 +226,7 @@ class RHI_API FCachedBoundShaderStateLink_Threadsafe : public FCachedBoundShader
 public:
 	/** Adds the bound shader state to the cache. */
 	FCachedBoundShaderStateLink_Threadsafe(
-		FVertexDeclarationRHIParamRef VertexDeclaration,
+		FRHIVertexDeclaration* VertexDeclaration,
 		FVertexShaderRHIParamRef VertexShader,
 		FPixelShaderRHIParamRef PixelShader,
 		FRHIBoundShaderState* InBoundShaderState
@@ -237,7 +237,7 @@ public:
 
 	/** Adds the bound shader state to the cache. */
 	FCachedBoundShaderStateLink_Threadsafe(
-		FVertexDeclarationRHIParamRef VertexDeclaration,
+		FRHIVertexDeclaration* VertexDeclaration,
 		FVertexShaderRHIParamRef VertexShader,
 		FPixelShaderRHIParamRef PixelShader,
 		FHullShaderRHIParamRef HullShader,
@@ -258,7 +258,7 @@ public:
  * @return If a bound shader state matching the parameters is cached, it is returned; otherwise NULL is returned.
  */
 extern RHI_API FBoundShaderStateRHIRef GetCachedBoundShaderState_Threadsafe(
-	FVertexDeclarationRHIParamRef VertexDeclaration,
+	FRHIVertexDeclaration* VertexDeclaration,
 	FVertexShaderRHIParamRef VertexShader,
 	FPixelShaderRHIParamRef PixelShader,
 	FHullShaderRHIParamRef HullShader = NULL,
diff --git a/Engine/Source/Runtime/RHI/Public/DynamicRHI.h b/Engine/Source/Runtime/RHI/Public/DynamicRHI.h
--- a/Engine/Source/Runtime/RHI/Public/DynamicRHI.h
+++ b/Engine/Source/Runtime/RHI/Public/DynamicRHI.h
@@ -393,7 +393,7 @@ public:
 	* @param PixelShader - existing pixel shader
 	*/
 	// FlushType: Thread safe, but varies depending on the RHI
-	virtual FBoundShaderStateRHIRef RHICreateBoundShaderState(FVertexDeclarationRHIParamRef VertexDeclaration, FVertexShaderRHIParamRef VertexShader, FHullShaderRHIParamRef HullShader, FDomainShaderRHIParamRef DomainShader, FPixelShaderRHIParamRef PixelShader, FGeometryShaderRHIParamRef GeometryShader) = 0;
+	virtual FBoundShaderStateRHIRef RHICreateBoundShaderState(FRHIVertexDeclaration* VertexDeclaration, FVertexShaderRHIParamRef VertexShader, FHullShaderRHIParamRef HullShader, FDomainShaderRHIParamRef DomainShader, FPixelShaderRHIParamRef PixelShader, FGeometryShaderRHIParamRef GeometryShader) = 0;
 
 	/**
 	* Creates a graphics pipeline state object (PSO) that represents a complete gpu pipeline for rendering.
@@ -1241,7 +1241,7 @@ FORCEINLINE FBlendStateRHIRef RHICreateBlendState(const FBlendStateInitializerRH
 	return GDynamicRHI->RHICreateBlendState(Initializer);
 }
 
-FORCEINLINE FBoundShaderStateRHIRef RHICreateBoundShaderState(FVertexDeclarationRHIParamRef VertexDeclaration, FVertexShaderRHIParamRef VertexShader, FHullShaderRHIParamRef HullShader, FDomainShaderRHIParamRef DomainShader, FPixelShaderRHIParamRef PixelShader, FGeometryShaderRHIParamRef GeometryShader)
+FORCEINLINE FBoundShaderStateRHIRef RHICreateBoundShaderState(FRHIVertexDeclaration* VertexDeclaration, FVertexShaderRHIParamRef VertexShader, FHullShaderRHIParamRef HullShader, FDomainShaderRHIParamRef DomainShader, FPixelShaderRHIParamRef PixelShader, FGeometryShaderRHIParamRef GeometryShader)
 {
 	return GDynamicRHI->RHICreateBoundShaderState(VertexDeclaration, VertexShader, HullShader, DomainShader, PixelShader, GeometryShader);
 }
diff --git a/Engine/Source/Runtime/RHI/Public/RHICommandList.h b/Engine/Source/Runtime/RHI/Public/RHICommandList.h
--- a/Engine/Source/Runtime/RHI/Public/RHICommandList.h
+++ b/Engine/Source/Runtime/RHI/Public/RHICommandList.h
@@ -1372,9 +1372,9 @@ struct FRHICommandTransitionUAVs final : public FRHICommand<FRHICommandTransitio
 	FUnorderedAccessViewRHIParamRef* UAVs; // Pointer to an array of UAVs, allocated inline with the command list
 	EResourceTransitionAccess TransitionType;
 	EResourceTransitionPipeline TransitionPipeline;
-	FComputeFenceRHIParamRef WriteFence;
+	FRHIComputeFence* WriteFence;
 
-	FORCEINLINE_DEBUGGABLE FRHICommandTransitionUAVs(EResourceTransitionAccess InTransitionType, EResourceTransitionPipeline InTransitionPipeline, FUnorderedAccessViewRHIParamRef* InUAVs, int32 InNumUAVs, FComputeFenceRHIParamRef InWriteFence)
+	FORCEINLINE_DEBUGGABLE FRHICommandTransitionUAVs(EResourceTransitionAccess InTransitionType, EResourceTransitionPipeline InTransitionPipeline, FUnorderedAccessViewRHIParamRef* InUAVs, int32 InNumUAVs, FRHIComputeFence* InWriteFence)
 		: NumUAVs(InNumUAVs)
 		, UAVs(InUAVs)
 		, TransitionType(InTransitionType)
@@ -1400,9 +1400,9 @@ struct FRHICommandSetAsyncComputeBudget final : public FRHICommand<FRHICommandSe
 template<ECmdList CmdListType>
 struct FRHICommandWaitComputeFence final : public FRHICommand<FRHICommandWaitComputeFence<CmdListType>>
 {
-	FComputeFenceRHIParamRef WaitFence;
+	FRHIComputeFence* WaitFence;
 
-	FORCEINLINE_DEBUGGABLE FRHICommandWaitComputeFence(FComputeFenceRHIParamRef InWaitFence)
+	FORCEINLINE_DEBUGGABLE FRHICommandWaitComputeFence(FRHIComputeFence* InWaitFence)
 		: WaitFence(InWaitFence)
 	{		
 	}
@@ -2724,7 +2724,7 @@ public:
 		ALLOC_COMMAND(FRHICommandTransitionTexturesArray)(TransitionType, InTextures);
 	}
 
-	FORCEINLINE_DEBUGGABLE void TransitionResource(EResourceTransitionAccess TransitionType, EResourceTransitionPipeline TransitionPipeline, FUnorderedAccessViewRHIParamRef InUAV, FComputeFenceRHIParamRef WriteFence)
+	FORCEINLINE_DEBUGGABLE void TransitionResource(EResourceTransitionAccess TransitionType, EResourceTransitionPipeline TransitionPipeline, FUnorderedAccessViewRHIParamRef InUAV, FRHIComputeFence* WriteFence)
 	{
 		check(InUAV == nullptr || InUAV->IsCommitted());
 		FUnorderedAccessViewRHIParamRef UAV = InUAV;
@@ -2745,7 +2745,7 @@ public:
 		TransitionResource(TransitionType, TransitionPipeline, InUAV, nullptr);
 	}
 
-	FORCEINLINE_DEBUGGABLE void TransitionResources(EResourceTransitionAccess TransitionType, EResourceTransitionPipeline TransitionPipeline, FUnorderedAccessViewRHIParamRef* InUAVs, int32 NumUAVs, FComputeFenceRHIParamRef WriteFence)
+	FORCEINLINE_DEBUGGABLE void TransitionResources(EResourceTransitionAccess TransitionType, EResourceTransitionPipeline TransitionPipeline, FUnorderedAccessViewRHIParamRef* InUAVs, int32 NumUAVs, FRHIComputeFence* WriteFence)
 	{
 		if (Bypass())
 		{
@@ -2768,7 +2768,7 @@ public:
 		TransitionResources(TransitionType, TransitionPipeline, InUAVs, NumUAVs, nullptr);
 	}
 
-	FORCEINLINE_DEBUGGABLE void WaitComputeFence(FComputeFenceRHIParamRef WaitFence)
+	FORCEINLINE_DEBUGGABLE void WaitComputeFence(FRHIComputeFence* WaitFence)
 	{
 		if (Bypass())
 		{
@@ -3308,7 +3308,7 @@ public:
 		ALLOC_COMMAND(FRHICommandDispatchIndirectComputeShader<ECmdList::ECompute>)(ArgumentBuffer, ArgumentOffset);
 	}
 
-	FORCEINLINE_DEBUGGABLE void TransitionResource(EResourceTransitionAccess TransitionType, EResourceTransitionPipeline TransitionPipeline, FUnorderedAccessViewRHIParamRef InUAV, FComputeFenceRHIParamRef WriteFence)
+	FORCEINLINE_DEBUGGABLE void TransitionResource(EResourceTransitionAccess TransitionType, EResourceTransitionPipeline TransitionPipeline, FUnorderedAccessViewRHIParamRef InUAV, FRHIComputeFence* WriteFence)
 	{
 		FUnorderedAccessViewRHIParamRef UAV = InUAV;
 		if (Bypass())
@@ -3328,7 +3328,7 @@ public:
 		TransitionResource(TransitionType, TransitionPipeline, InUAV, nullptr);
 	}
 
-	FORCEINLINE_DEBUGGABLE void TransitionResources(EResourceTransitionAccess TransitionType, EResourceTransitionPipeline TransitionPipeline, FUnorderedAccessViewRHIParamRef* InUAVs, int32 NumUAVs, FComputeFenceRHIParamRef WriteFence)
+	FORCEINLINE_DEBUGGABLE void TransitionResources(EResourceTransitionAccess TransitionType, EResourceTransitionPipeline TransitionPipeline, FUnorderedAccessViewRHIParamRef* InUAVs, int32 NumUAVs, FRHIComputeFence* WriteFence)
 	{
 		if (Bypass())
 		{
@@ -3397,7 +3397,7 @@ public:
 		ALLOC_COMMAND(FRHICommandSubmitCommandsHint<ECmdList::ECompute>)();
 	}
 
-	FORCEINLINE_DEBUGGABLE void WaitComputeFence(FComputeFenceRHIParamRef WaitFence)
+	FORCEINLINE_DEBUGGABLE void WaitComputeFence(FRHIComputeFence* WaitFence)
 	{
 		if (Bypass())
 		{
@@ -3642,7 +3642,7 @@ public:
 		return GDynamicRHI->RHICreateStagingBuffer();
 	}
 
-	FORCEINLINE FBoundShaderStateRHIRef CreateBoundShaderState(FVertexDeclarationRHIParamRef VertexDeclaration, FVertexShaderRHIParamRef VertexShader, FHullShaderRHIParamRef HullShader, FDomainShaderRHIParamRef DomainShader, FPixelShaderRHIParamRef PixelShader, FGeometryShaderRHIParamRef GeometryShader)
+	FORCEINLINE FBoundShaderStateRHIRef CreateBoundShaderState(FRHIVertexDeclaration* VertexDeclaration, FVertexShaderRHIParamRef VertexShader, FHullShaderRHIParamRef HullShader, FDomainShaderRHIParamRef DomainShader, FPixelShaderRHIParamRef PixelShader, FGeometryShaderRHIParamRef GeometryShader)
 	{
 		LLM_SCOPE(ELLMTag::Shaders);
 		return RHICreateBoundShaderState(VertexDeclaration, VertexShader, HullShader, DomainShader, PixelShader, GeometryShader);
diff --git a/Engine/Source/Runtime/RHI/Public/RHIContext.h b/Engine/Source/Runtime/RHI/Public/RHIContext.h
--- a/Engine/Source/Runtime/RHI/Public/RHIContext.h
+++ b/Engine/Source/Runtime/RHI/Public/RHIContext.h
@@ -35,7 +35,7 @@ public:
 	/**
 	* Compute queue will wait for the fence to be written before continuing.
 	*/
-	virtual void RHIWaitComputeFence(FComputeFenceRHIParamRef InFence) = 0;
+	virtual void RHIWaitComputeFence(FRHIComputeFence* InFence) = 0;
 
 	/**
 	*Sets the current compute shader.
@@ -68,7 +68,7 @@ public:
 	* @param NumUAVs - number of UAVs to transition
 	* @param WriteComputeFence - Optional ComputeFence to write as part of this transition
 	*/
-	virtual void RHITransitionResources(EResourceTransitionAccess TransitionType, EResourceTransitionPipeline TransitionPipeline, FUnorderedAccessViewRHIParamRef* InUAVs, int32 NumUAVs, FComputeFenceRHIParamRef WriteComputeFence) = 0;
+	virtual void RHITransitionResources(EResourceTransitionAccess TransitionType, EResourceTransitionPipeline TransitionPipeline, FUnorderedAccessViewRHIParamRef* InUAVs, int32 NumUAVs, FRHIComputeFence* WriteComputeFence) = 0;
 
 	/** Set the shader resource view of a surface.  This is used for binding TextureMS parameter types that need a multi sampled view. */
 	virtual void RHISetShaderTexture(FComputeShaderRHIParamRef PixelShader, uint32 TextureIndex, FTextureRHIParamRef NewTexture) = 0;
@@ -167,7 +167,7 @@ public:
 	/**
 	* Compute queue will wait for the fence to be written before continuing.
 	*/
-	virtual void RHIWaitComputeFence(FComputeFenceRHIParamRef InFence) override
+	virtual void RHIWaitComputeFence(FRHIComputeFence* InFence) override
 	{
 		if (InFence)
 		{
@@ -232,7 +232,7 @@ public:
 	* @param NumUAVs - number of UAVs to transition
 	* @param WriteComputeFence - Optional ComputeFence to write as part of this transition
 	*/
-	virtual void RHITransitionResources(EResourceTransitionAccess TransitionType, EResourceTransitionPipeline TransitionPipeline, FUnorderedAccessViewRHIParamRef* InUAVs, int32 NumUAVs, FComputeFenceRHIParamRef WriteComputeFence)
+	virtual void RHITransitionResources(EResourceTransitionAccess TransitionType, EResourceTransitionPipeline TransitionPipeline, FUnorderedAccessViewRHIParamRef* InUAVs, int32 NumUAVs, FRHIComputeFence* WriteComputeFence)
 	{
 		if (WriteComputeFence)
 		{
@@ -649,7 +649,7 @@ public:
 
 
 FORCEINLINE FBoundShaderStateRHIRef RHICreateBoundShaderState(
-	FVertexDeclarationRHIParamRef VertexDeclaration,
+	FRHIVertexDeclaration* VertexDeclaration,
 	FVertexShaderRHIParamRef VertexShader,
 	FHullShaderRHIParamRef HullShader,
 	FDomainShaderRHIParamRef DomainShader,
diff --git a/Engine/Source/Runtime/RHI/Public/RHIResources.h b/Engine/Source/Runtime/RHI/Public/RHIResources.h
--- a/Engine/Source/Runtime/RHI/Public/RHIResources.h
+++ b/Engine/Source/Runtime/RHI/Public/RHIResources.h
@@ -1113,6 +1113,7 @@ UE_DEPRECATED(4.23, "FBlendStateRHIParamRef typedef is deprecated; please use FR
 typedef FRHIBlendState*              FBlendStateRHIParamRef;
 typedef TRefCountPtr<FRHIBlendState> FBlendStateRHIRef;
 
+UE_DEPRECATED(4.23, "FVertexDeclarationRHIParamRef typedef is deprecated; please use FRHIVertexDeclaration* directly instead.")
 typedef FRHIVertexDeclaration*              FVertexDeclarationRHIParamRef;
 typedef TRefCountPtr<FRHIVertexDeclaration> FVertexDeclarationRHIRef;
 
@@ -1137,6 +1138,7 @@ typedef TRefCountPtr<FRHIComputeShader> FComputeShaderRHIRef;
 typedef FRHIRayTracingShader*                       FRayTracingShaderRHIParamRef;
 typedef TRefCountPtr<FRHIRayTracingShader>          FRayTracingShaderRHIRef;
 
+UE_DEPRECATED(4.23, "FComputeFenceRHIParamRef typedef is deprecated; please use FRHIComputeFence* directly instead.")
 typedef FRHIComputeFence*				FComputeFenceRHIParamRef;
 typedef TRefCountPtr<FRHIComputeFence>	FComputeFenceRHIRef;
 
@@ -1752,7 +1754,7 @@ struct FBoundShaderStateInput
 
 	inline FBoundShaderStateInput
 	(
-		FVertexDeclarationRHIParamRef InVertexDeclarationRHI
+		FRHIVertexDeclaration* InVertexDeclarationRHI
 		, FVertexShaderRHIParamRef InVertexShaderRHI
 #if PLATFORM_SUPPORTS_TESSELLATION_SHADERS
 		, FHullShaderRHIParamRef InHullShaderRHI
@@ -1776,7 +1778,7 @@ struct FBoundShaderStateInput
 	{
 	}
 
-	FVertexDeclarationRHIParamRef VertexDeclarationRHI = nullptr;
+	FRHIVertexDeclaration* VertexDeclarationRHI = nullptr;
 	FVertexShaderRHIParamRef VertexShaderRHI = nullptr;
 	FHullShaderRHIParamRef HullShaderRHI = nullptr;
 	FDomainShaderRHIParamRef DomainShaderRHI = nullptr;
diff --git a/Engine/Source/Runtime/RHI/Public/RHIValidationContext.h b/Engine/Source/Runtime/RHI/Public/RHIValidationContext.h
--- a/Engine/Source/Runtime/RHI/Public/RHIValidationContext.h
+++ b/Engine/Source/Runtime/RHI/Public/RHIValidationContext.h
@@ -24,7 +24,7 @@ public:
 	/**
 	* Compute queue will wait for the fence to be written before continuing.
 	*/
-	virtual void RHIWaitComputeFence(FComputeFenceRHIParamRef InFence) override final
+	virtual void RHIWaitComputeFence(FRHIComputeFence* InFence) override final
 	{
 		RHIContext->RHIWaitComputeFence(InFence);
 	}
@@ -65,7 +65,7 @@ public:
 	* @param NumUAVs - number of UAVs to transition
 	* @param WriteComputeFence - Optional ComputeFence to write as part of this transition
 	*/
-	virtual void RHITransitionResources(EResourceTransitionAccess TransitionType, EResourceTransitionPipeline TransitionPipeline, FUnorderedAccessViewRHIParamRef* InUAVs, int32 NumUAVs, FComputeFenceRHIParamRef WriteComputeFence) override final
+	virtual void RHITransitionResources(EResourceTransitionAccess TransitionType, EResourceTransitionPipeline TransitionPipeline, FUnorderedAccessViewRHIParamRef* InUAVs, int32 NumUAVs, FRHIComputeFence* WriteComputeFence) override final
 	{
 		RHIContext->RHITransitionResources(TransitionType, TransitionPipeline, InUAVs, NumUAVs, WriteComputeFence);
 	}
@@ -182,7 +182,7 @@ public:
 	/**
 	* Compute queue will wait for the fence to be written before continuing.
 	*/
-	virtual void RHIWaitComputeFence(FComputeFenceRHIParamRef InFence) override final
+	virtual void RHIWaitComputeFence(FRHIComputeFence* InFence) override final
 	{
 		RHIContext->RHIWaitComputeFence(InFence);
 	}
@@ -273,7 +273,7 @@ public:
 	* @param NumUAVs - number of UAVs to transition
 	* @param WriteComputeFence - Optional ComputeFence to write as part of this transition
 	*/
-	virtual void RHITransitionResources(EResourceTransitionAccess TransitionType, EResourceTransitionPipeline TransitionPipeline, FUnorderedAccessViewRHIParamRef* InUAVs, int32 NumUAVs, FComputeFenceRHIParamRef WriteComputeFence) override final
+	virtual void RHITransitionResources(EResourceTransitionAccess TransitionType, EResourceTransitionPipeline TransitionPipeline, FUnorderedAccessViewRHIParamRef* InUAVs, int32 NumUAVs, FRHIComputeFence* WriteComputeFence) override final
 	{
 		RHIContext->RHITransitionResources(TransitionType, TransitionPipeline, InUAVs, NumUAVs, WriteComputeFence);
 	}
diff --git a/Engine/Source/Runtime/RenderCore/Private/StaticBoundShaderState.cpp b/Engine/Source/Runtime/RenderCore/Private/StaticBoundShaderState.cpp
--- a/Engine/Source/Runtime/RenderCore/Private/StaticBoundShaderState.cpp
+++ b/Engine/Source/Runtime/RenderCore/Private/StaticBoundShaderState.cpp
@@ -51,7 +51,7 @@ FGlobalBoundShaderStateResource::~FGlobalBoundShaderStateResource()
  * Initializes a global bound shader state with a vanilla bound shader state and required information.
  */
 FRHIBoundShaderState* FGlobalBoundShaderStateResource::GetInitializedRHI(
-	FVertexDeclarationRHIParamRef VertexDeclaration,
+	FRHIVertexDeclaration* VertexDeclaration,
 	FVertexShaderRHIParamRef VertexShader, 
 	FPixelShaderRHIParamRef PixelShader,
 	FGeometryShaderRHIParamRef GeometryShader
diff --git a/Engine/Source/Runtime/RenderCore/Public/StaticBoundShaderState.h b/Engine/Source/Runtime/RenderCore/Public/StaticBoundShaderState.h
--- a/Engine/Source/Runtime/RenderCore/Public/StaticBoundShaderState.h
+++ b/Engine/Source/Runtime/RenderCore/Public/StaticBoundShaderState.h
@@ -35,7 +35,7 @@ public:
 	 * @return The bound shader state RHI.
 	 */
 	FRHIBoundShaderState* GetInitializedRHI(
-		FVertexDeclarationRHIParamRef VertexDeclaration, 
+		FRHIVertexDeclaration* VertexDeclaration,
 		FVertexShaderRHIParamRef VertexShader, 
 		FPixelShaderRHIParamRef PixelShader,
 		FGeometryShaderRHIParamRef GeometryShader
@@ -59,7 +59,7 @@ private:
 	RENDERCORE_API virtual void ReleaseRHI();
 
 #if DO_CHECK
-	FVertexDeclarationRHIParamRef BoundVertexDeclaration;
+	FRHIVertexDeclaration* BoundVertexDeclaration;
 	FVertexShaderRHIParamRef BoundVertexShader;
 	FPixelShaderRHIParamRef BoundPixelShader;
 	FGeometryShaderRHIParamRef BoundGeometryShader;
@@ -71,7 +71,7 @@ typedef TGlobalResource<FGlobalBoundShaderStateResource> FGlobalBoundShaderState
 
 struct FGlobalBoundShaderStateArgs
 {
-	FVertexDeclarationRHIParamRef VertexDeclarationRHI;
+	FRHIVertexDeclaration* VertexDeclarationRHI;
 	FShader* VertexShader;
 	FShader* PixelShader;
 	FShader* GeometryShader;
diff --git a/Engine/Source/Runtime/Renderer/Private/DebugViewModeRendering.cpp b/Engine/Source/Runtime/Renderer/Private/DebugViewModeRendering.cpp
--- a/Engine/Source/Runtime/Renderer/Private/DebugViewModeRendering.cpp
+++ b/Engine/Source/Runtime/Renderer/Private/DebugViewModeRendering.cpp
@@ -235,7 +235,7 @@ void FDeferredShadingSceneRenderer::DoDebugViewModePostProcessing(FRHICommandLis
 		// May need to wait on the final pass to complete
 		if (Context.FinalOutput.IsAsyncComputePass())
 		{
-			FComputeFenceRHIParamRef ComputeFinalizeFence = Context.FinalOutput.GetComputePassEndFence();
+			FRHIComputeFence* ComputeFinalizeFence = Context.FinalOutput.GetComputePassEndFence();
 			if (ComputeFinalizeFence)
 			{
 				Context.RHICmdList.WaitComputeFence(ComputeFinalizeFence);
diff --git a/Engine/Source/Runtime/Renderer/Private/GPUBenchmark.cpp b/Engine/Source/Runtime/Renderer/Private/GPUBenchmark.cpp
--- a/Engine/Source/Runtime/Renderer/Private/GPUBenchmark.cpp
+++ b/Engine/Source/Runtime/Renderer/Private/GPUBenchmark.cpp
@@ -205,7 +205,7 @@ void RunBenchmarkShader(FRHICommandList& RHICmdList, FVertexBufferRHIParamRef Ve
 	TShaderMapRef<FPostProcessBenchmarkPS<PsMethod>> PixelShader(ShaderMap);
 
 	bool bVertexTest = VsMethod != 0;
-	FVertexDeclarationRHIParamRef VertexDeclaration = bVertexTest
+	FRHIVertexDeclaration* VertexDeclaration = bVertexTest
 		? GVertexThroughputDeclaration.DeclRHI
 		: GFilterVertexDeclaration.VertexDeclarationRHI;
 
diff --git a/Engine/Source/Runtime/Renderer/Private/MeshPassProcessor.cpp b/Engine/Source/Runtime/Renderer/Private/MeshPassProcessor.cpp
--- a/Engine/Source/Runtime/Renderer/Private/MeshPassProcessor.cpp
+++ b/Engine/Source/Runtime/Renderer/Private/MeshPassProcessor.cpp
@@ -706,7 +706,7 @@ void FMeshDrawShaderBindings::CopyFrom(const FMeshDrawShaderBindings& Other)
 #endif
 }
 
-void FMeshDrawCommand::SetShaders(FVertexDeclarationRHIParamRef VertexDeclaration, const FMeshProcessorShaders& Shaders, FGraphicsMinimalPipelineStateInitializer& PipelineState)
+void FMeshDrawCommand::SetShaders(FRHIVertexDeclaration* VertexDeclaration, const FMeshProcessorShaders& Shaders, FGraphicsMinimalPipelineStateInitializer& PipelineState)
 {
 	PipelineState.BoundShaderState = FBoundShaderStateInput(
 		VertexDeclaration
diff --git a/Engine/Source/Runtime/Renderer/Private/PathTracing/PathCompactionCompute.cpp b/Engine/Source/Runtime/Renderer/Private/PathTracing/PathCompactionCompute.cpp
--- a/Engine/Source/Runtime/Renderer/Private/PathTracing/PathCompactionCompute.cpp
+++ b/Engine/Source/Runtime/Renderer/Private/PathTracing/PathCompactionCompute.cpp
@@ -89,7 +89,7 @@ public:
 		FUnorderedAccessViewRHIParamRef RadianceSortedBlueUAV,
 		FUnorderedAccessViewRHIParamRef RadianceSortedAlphaUAV,
 		FUnorderedAccessViewRHIParamRef SampleCountSortedUAV,
-		FComputeFenceRHIParamRef Fence)
+		FRHIComputeFence* Fence)
 	{
 		FComputeShaderRHIParamRef ShaderRHI = GetComputeShader();
 
diff --git a/Engine/Source/Runtime/Renderer/Private/PathTracing/PathTracingTestbed.cpp b/Engine/Source/Runtime/Renderer/Private/PathTracing/PathTracingTestbed.cpp
--- a/Engine/Source/Runtime/Renderer/Private/PathTracing/PathTracingTestbed.cpp
+++ b/Engine/Source/Runtime/Renderer/Private/PathTracing/PathTracingTestbed.cpp
@@ -67,7 +67,7 @@ public:
 		EResourceTransitionPipeline TransitionPipeline,
 		FRWBuffer& ResultsBuffer,
 		FRWBuffer& FloatResultsBuffer,
-		FComputeFenceRHIParamRef Fence)
+		FRHIComputeFence* Fence)
 	{
 		FComputeShaderRHIParamRef ShaderRHI = GetComputeShader();
 
@@ -152,7 +152,7 @@ public:
 		EResourceTransitionAccess TransitionAccess,
 		EResourceTransitionPipeline TransitionPipeline,
 		FRWBuffer& PdfsResultsBuffer,
-		FComputeFenceRHIParamRef Fence)
+		FRHIComputeFence* Fence)
 	{
 		FComputeShaderRHIParamRef ShaderRHI = GetComputeShader();
 
@@ -241,7 +241,7 @@ public:
 		EResourceTransitionAccess TransitionAccess,
 		EResourceTransitionPipeline TransitionPipeline,
 		FRWBuffer& PdfsResultsBuffer,
-		FComputeFenceRHIParamRef Fence)
+		FRHIComputeFence* Fence)
 	{
 		FComputeShaderRHIParamRef ShaderRHI = GetComputeShader();
 
@@ -329,7 +329,7 @@ public:
 		EResourceTransitionAccess TransitionAccess,
 		EResourceTransitionPipeline TransitionPipeline,
 		FRWBuffer& PdfsResultsBuffer,
-		FComputeFenceRHIParamRef Fence)
+		FRHIComputeFence* Fence)
 	{
 		FComputeShaderRHIParamRef ShaderRHI = GetComputeShader();
 
diff --git a/Engine/Source/Runtime/Renderer/Private/PathTracing/RayCounterCompute.cpp b/Engine/Source/Runtime/Renderer/Private/PathTracing/RayCounterCompute.cpp
--- a/Engine/Source/Runtime/Renderer/Private/PathTracing/RayCounterCompute.cpp
+++ b/Engine/Source/Runtime/Renderer/Private/PathTracing/RayCounterCompute.cpp
@@ -65,7 +65,7 @@ public:
 		EResourceTransitionAccess TransitionAccess,
 		EResourceTransitionPipeline TransitionPipeline,
 		FRWBuffer& TotalRayCountBuffer,
-		FComputeFenceRHIParamRef Fence)
+		FRHIComputeFence* Fence)
 	{
 		FComputeShaderRHIParamRef ShaderRHI = GetComputeShader();
 
diff --git a/Engine/Source/Runtime/Renderer/Private/PathTracing/SkyLightCdfBuild.cpp b/Engine/Source/Runtime/Renderer/Private/PathTracing/SkyLightCdfBuild.cpp
--- a/Engine/Source/Runtime/Renderer/Private/PathTracing/SkyLightCdfBuild.cpp
+++ b/Engine/Source/Runtime/Renderer/Private/PathTracing/SkyLightCdfBuild.cpp
@@ -70,7 +70,7 @@ public:
 	}
 
 	void UnsetParameters(FRHICommandList& RHICmdList, EResourceTransitionAccess TransitionAccess, EResourceTransitionPipeline TransitionPipeline,
-		const FRWBuffer& Buffer, FComputeFenceRHIParamRef Fence)
+		const FRWBuffer& Buffer, FRHIComputeFence* Fence)
 	{
 		FComputeShaderRHIParamRef ShaderRHI = GetComputeShader();
 		RowCdfParameter.UnsetUAV(RHICmdList, ShaderRHI);
@@ -157,7 +157,7 @@ public:
 	}
 
 	void UnsetParameters(FRHICommandList& RHICmdList, EResourceTransitionAccess TransitionAccess, EResourceTransitionPipeline TransitionPipeline,
-		const FRWBuffer& ColumnCdf, FComputeFenceRHIParamRef Fence)
+		const FRWBuffer& ColumnCdf, FRHIComputeFence* Fence)
 	{
 		FComputeShaderRHIParamRef ShaderRHI = GetComputeShader();
 		ColumnCdfParameter.UnsetUAV(RHICmdList, ShaderRHI);
@@ -236,7 +236,7 @@ public:
 	}
 
 	void UnsetParameters(FRHICommandList& RHICmdList, EResourceTransitionAccess TransitionAccess, EResourceTransitionPipeline TransitionPipeline,
-		const FRWBuffer& CubeFaceCdf, FComputeFenceRHIParamRef Fence)
+		const FRWBuffer& CubeFaceCdf, FRHIComputeFence* Fence)
 	{
 		FComputeShaderRHIParamRef ShaderRHI = GetComputeShader();
 		CubeFaceCdfParameter.UnsetUAV(RHICmdList, ShaderRHI);
diff --git a/Engine/Source/Runtime/Renderer/Private/PathTracing/VarianceMipTreeBuild.cpp b/Engine/Source/Runtime/Renderer/Private/PathTracing/VarianceMipTreeBuild.cpp
--- a/Engine/Source/Runtime/Renderer/Private/PathTracing/VarianceMipTreeBuild.cpp
+++ b/Engine/Source/Runtime/Renderer/Private/PathTracing/VarianceMipTreeBuild.cpp
@@ -71,7 +71,7 @@ public:
 		EResourceTransitionAccess TransitionAccess,
 		EResourceTransitionPipeline TransitionPipeline,
 		FRWBuffer& VarianceMap,
-		FComputeFenceRHIParamRef Fence)
+		FRHIComputeFence* Fence)
 	{
 		FComputeShaderRHIParamRef ShaderRHI = GetComputeShader();
 
diff --git a/Engine/Source/Runtime/Renderer/Private/PostProcess/PostProcessBloomSetup.h b/Engine/Source/Runtime/Renderer/Private/PostProcess/PostProcessBloomSetup.h
--- a/Engine/Source/Runtime/Renderer/Private/PostProcess/PostProcessBloomSetup.h
+++ b/Engine/Source/Runtime/Renderer/Private/PostProcess/PostProcessBloomSetup.h
@@ -27,7 +27,7 @@ public:
 	virtual FPooledRenderTargetDesc ComputeOutputDesc(EPassOutputId InPassOutputId) const override;
 	virtual void Release() override { delete this; }
 
-	virtual FComputeFenceRHIParamRef GetComputePassEndFence() const override { return AsyncEndFence; }
+	virtual FRHIComputeFence* GetComputePassEndFence() const override { return AsyncEndFence; }
 
 private:
 	template <typename TRHICmdList>
diff --git a/Engine/Source/Runtime/Renderer/Private/PostProcess/PostProcessCombineLUTs.h b/Engine/Source/Runtime/Renderer/Private/PostProcess/PostProcessCombineLUTs.h
--- a/Engine/Source/Runtime/Renderer/Private/PostProcess/PostProcessCombineLUTs.h
+++ b/Engine/Source/Runtime/Renderer/Private/PostProcess/PostProcessCombineLUTs.h
@@ -39,7 +39,7 @@ public:
 	/** @return 0xffffffff if not found */
 	uint32 FindIndex(const FFinalPostProcessSettings& Settings, UTexture* Tex) const;
 
-	virtual FComputeFenceRHIParamRef GetComputePassEndFence() const override { return AsyncEndFence; }
+	virtual FRHIComputeFence* GetComputePassEndFence() const override { return AsyncEndFence; }
 
 private:
 	template <typename TRHICmdList>
diff --git a/Engine/Source/Runtime/Renderer/Private/PostProcess/PostProcessDownsample.h b/Engine/Source/Runtime/Renderer/Private/PostProcess/PostProcessDownsample.h
--- a/Engine/Source/Runtime/Renderer/Private/PostProcess/PostProcessDownsample.h
+++ b/Engine/Source/Runtime/Renderer/Private/PostProcess/PostProcessDownsample.h
@@ -30,7 +30,7 @@ public:
 	virtual void Release() override { delete this; }
 	virtual FPooledRenderTargetDesc ComputeOutputDesc(EPassOutputId InPassOutputId) const override;
 
-	virtual FComputeFenceRHIParamRef GetComputePassEndFence() const override { return AsyncEndFence; }
+	virtual FRHIComputeFence* GetComputePassEndFence() const override { return AsyncEndFence; }
 
 private:
 	template <uint32 Method, uint32 ManuallyClampUV>
diff --git a/Engine/Source/Runtime/Renderer/Private/PostProcess/PostProcessEyeAdaptation.h b/Engine/Source/Runtime/Renderer/Private/PostProcess/PostProcessEyeAdaptation.h
--- a/Engine/Source/Runtime/Renderer/Private/PostProcess/PostProcessEyeAdaptation.h
+++ b/Engine/Source/Runtime/Renderer/Private/PostProcess/PostProcessEyeAdaptation.h
@@ -59,7 +59,7 @@ public:
 	virtual void Release() override { delete this; }
 	virtual FPooledRenderTargetDesc ComputeOutputDesc(EPassOutputId InPassOutputId) const override;
 
-	virtual FComputeFenceRHIParamRef GetComputePassEndFence() const override { return AsyncEndFence; }
+	virtual FRHIComputeFence* GetComputePassEndFence() const override { return AsyncEndFence; }
 
 private:
 	template <typename TRHICmdList>
diff --git a/Engine/Source/Runtime/Renderer/Private/PostProcess/PostProcessTemporalAA.h b/Engine/Source/Runtime/Renderer/Private/PostProcess/PostProcessTemporalAA.h
--- a/Engine/Source/Runtime/Renderer/Private/PostProcess/PostProcessTemporalAA.h
+++ b/Engine/Source/Runtime/Renderer/Private/PostProcess/PostProcessTemporalAA.h
@@ -170,7 +170,7 @@ public:
 	virtual void Release() override { delete this; }
 	virtual FPooledRenderTargetDesc ComputeOutputDesc(EPassOutputId InPassOutputId) const override;
 
-	virtual FComputeFenceRHIParamRef GetComputePassEndFence() const override { return AsyncEndFence; }
+	virtual FRHIComputeFence* GetComputePassEndFence() const override { return AsyncEndFence; }
 
 private:
 	const FTAAPassParameters SavedParameters;
diff --git a/Engine/Source/Runtime/Renderer/Private/PostProcess/PostProcessTonemap.h b/Engine/Source/Runtime/Renderer/Private/PostProcess/PostProcessTonemap.h
--- a/Engine/Source/Runtime/Renderer/Private/PostProcess/PostProcessTonemap.h
+++ b/Engine/Source/Runtime/Renderer/Private/PostProcess/PostProcessTonemap.h
@@ -57,7 +57,7 @@ public:
 	virtual void Release() override { delete this; }
 	virtual FPooledRenderTargetDesc ComputeOutputDesc(EPassOutputId InPassOutputId) const override;
 
-	virtual FComputeFenceRHIParamRef GetComputePassEndFence() const override { return AsyncEndFence; }
+	virtual FRHIComputeFence* GetComputePassEndFence() const override { return AsyncEndFence; }
 
 	bool bDoGammaOnly;
 	bool bDoScreenPercentageInTonemapper;
diff --git a/Engine/Source/Runtime/Renderer/Private/PostProcess/PostProcessWeightedSampleSum.h b/Engine/Source/Runtime/Renderer/Private/PostProcess/PostProcessWeightedSampleSum.h
--- a/Engine/Source/Runtime/Renderer/Private/PostProcess/PostProcessWeightedSampleSum.h
+++ b/Engine/Source/Runtime/Renderer/Private/PostProcess/PostProcessWeightedSampleSum.h
@@ -53,7 +53,7 @@ public:
 	// @param InCrossCenterWeight >=0
 	void SetCrossCenterWeight(float InCrossCenterWeight) { check(InCrossCenterWeight >= 0.0f); CrossCenterWeight = InCrossCenterWeight; }
 
-	virtual FComputeFenceRHIParamRef GetComputePassEndFence() const override { return AsyncEndFence; }
+	virtual FRHIComputeFence* GetComputePassEndFence() const override { return AsyncEndFence; }
 
 private:
 	void AdjustRectsForFastBlur(FIntRect& SrcRect, FIntRect& DestRect) const;
diff --git a/Engine/Source/Runtime/Renderer/Private/PostProcess/PostProcessing.cpp b/Engine/Source/Runtime/Renderer/Private/PostProcess/PostProcessing.cpp
--- a/Engine/Source/Runtime/Renderer/Private/PostProcess/PostProcessing.cpp
+++ b/Engine/Source/Runtime/Renderer/Private/PostProcess/PostProcessing.cpp
@@ -2086,7 +2086,7 @@ void FPostProcessing::Process(FRHICommandListImmediate& RHICmdList, const FViewI
 			// May need to wait on the final pass to complete
 			if (Context.FinalOutput.IsAsyncComputePass())
 			{
-				FComputeFenceRHIParamRef ComputeFinalizeFence = Context.FinalOutput.GetComputePassEndFence();
+				FRHIComputeFence* ComputeFinalizeFence = Context.FinalOutput.GetComputePassEndFence();
 				if (ComputeFinalizeFence)
 				{
 					Context.RHICmdList.WaitComputeFence(ComputeFinalizeFence);
diff --git a/Engine/Source/Runtime/Renderer/Private/PostProcess/RenderingCompositionGraph.h b/Engine/Source/Runtime/Renderer/Private/PostProcess/RenderingCompositionGraph.h
--- a/Engine/Source/Runtime/Renderer/Private/PostProcess/RenderingCompositionGraph.h
+++ b/Engine/Source/Runtime/Renderer/Private/PostProcess/RenderingCompositionGraph.h
@@ -402,7 +402,7 @@ struct FRenderingCompositePass
 		return bIsComputePass && bPreferAsyncCompute && GSupportsEfficientAsyncCompute;
 #endif
 	};
-	virtual FComputeFenceRHIParamRef GetComputePassEndFence() const { return nullptr; }
+	virtual FRHIComputeFence* GetComputePassEndFence() const { return nullptr; }
 
 protected:
 	/** to avoid wasteful recomputation and to support graph/DAG traversal, if ComputeOutputDesc() was called */
@@ -448,7 +448,7 @@ struct FRenderingCompositeOutputRef
 		return IsValid() && Source->IsAsyncComputePass();
 	}
 
-	FComputeFenceRHIParamRef GetComputePassEndFence() const
+	FRHIComputeFence* GetComputePassEndFence() const
 	{
 		return IsValid() ? Source->GetComputePassEndFence() : nullptr;
 	}
@@ -688,7 +688,7 @@ protected:
 		{
 			if (IsAsyncComputePass() != Input.IsAsyncComputePass())
 			{
-				FComputeFenceRHIParamRef InputComputePassEndFence = Input.GetComputePassEndFence();
+				FRHIComputeFence* InputComputePassEndFence = Input.GetComputePassEndFence();
 				if (InputComputePassEndFence)
 				{
 					RHICmdList.WaitComputeFence(InputComputePassEndFence);
diff --git a/Engine/Source/Runtime/Renderer/Private/RayTracing/RayTracingRectLight.cpp b/Engine/Source/Runtime/Renderer/Private/RayTracing/RayTracingRectLight.cpp
--- a/Engine/Source/Runtime/Renderer/Private/RayTracing/RayTracingRectLight.cpp
+++ b/Engine/Source/Runtime/Renderer/Private/RayTracing/RayTracingRectLight.cpp
@@ -134,7 +134,7 @@ public:
 		EResourceTransitionAccess TransitionAccess,
 		EResourceTransitionPipeline TransitionPipeline,
 		FRWBuffer& MipTree,
-		FComputeFenceRHIParamRef Fence)
+		FRHIComputeFence* Fence)
 	{
 		FComputeShaderRHIParamRef ShaderRHI = GetComputeShader();
 
diff --git a/Engine/Source/Runtime/Renderer/Private/RayTracing/RaytracingSkylight.cpp b/Engine/Source/Runtime/Renderer/Private/RayTracing/RaytracingSkylight.cpp
--- a/Engine/Source/Runtime/Renderer/Private/RayTracing/RaytracingSkylight.cpp
+++ b/Engine/Source/Runtime/Renderer/Private/RayTracing/RaytracingSkylight.cpp
@@ -400,7 +400,7 @@ public:
 		EResourceTransitionAccess TransitionAccess,
 		EResourceTransitionPipeline TransitionPipeline,
 		FRWBuffer& MipTreePdf,
-		FComputeFenceRHIParamRef Fence
+		FRHIComputeFence* Fence
 	)
 	{
 		FComputeShaderRHIParamRef ShaderRHI = GetComputeShader();
diff --git a/Engine/Source/Runtime/Renderer/Public/MeshPassProcessor.h b/Engine/Source/Runtime/Renderer/Public/MeshPassProcessor.h
--- a/Engine/Source/Runtime/Renderer/Public/MeshPassProcessor.h
+++ b/Engine/Source/Runtime/Renderer/Public/MeshPassProcessor.h
@@ -530,7 +530,7 @@ public:
 	}
 
 	/** Sets shaders on the mesh draw command and allocates room for the shader bindings. */
-	RENDERER_API void SetShaders(FVertexDeclarationRHIParamRef VertexDeclaration, const FMeshProcessorShaders& Shaders, FGraphicsMinimalPipelineStateInitializer& PipelineState);
+	RENDERER_API void SetShaders(FRHIVertexDeclaration* VertexDeclaration, const FMeshProcessorShaders& Shaders, FGraphicsMinimalPipelineStateInitializer& PipelineState);
 
 	inline void SetStencilRef(uint32 InStencilRef)
 	{
diff --git a/Engine/Source/Runtime/Renderer/Public/MeshPassProcessor.inl b/Engine/Source/Runtime/Renderer/Public/MeshPassProcessor.inl
--- a/Engine/Source/Runtime/Renderer/Public/MeshPassProcessor.inl
+++ b/Engine/Source/Runtime/Renderer/Public/MeshPassProcessor.inl
@@ -37,7 +37,7 @@ void FMeshPassProcessor::BuildMeshDrawCommands(
 	if ((MeshPassFeatures & EMeshPassFeatures::PositionAndNormalOnly) != EMeshPassFeatures::Default)	InputStreamType = EVertexInputStreamType::PositionAndNormalOnly;
 
 	check(VertexFactory && VertexFactory->IsInitialized());
-	FVertexDeclarationRHIParamRef VertexDeclaration = VertexFactory->GetDeclaration(InputStreamType);
+	FRHIVertexDeclaration* VertexDeclaration = VertexFactory->GetDeclaration(InputStreamType);
 
 	check(!VertexFactory->NeedsDeclaration() || VertexDeclaration);
 
diff --git a/Engine/Source/Runtime/VulkanRHI/Private/VulkanContext.h b/Engine/Source/Runtime/VulkanRHI/Private/VulkanContext.h
--- a/Engine/Source/Runtime/VulkanRHI/Private/VulkanContext.h
+++ b/Engine/Source/Runtime/VulkanRHI/Private/VulkanContext.h
@@ -219,7 +219,7 @@ public:
 
 	virtual void RHISetComputeShader(FComputeShaderRHIParamRef ComputeShader) final override;
 	virtual void RHISetComputePipelineState(FRHIComputePipelineState* ComputePipelineState) final override;
-	virtual void RHIWaitComputeFence(FComputeFenceRHIParamRef InFence) final override;
+	virtual void RHIWaitComputeFence(FRHIComputeFence* InFence) final override;
 	virtual void RHIDispatchComputeShader(uint32 ThreadGroupCountX, uint32 ThreadGroupCountY, uint32 ThreadGroupCountZ) final override;
 	virtual void RHIDispatchIndirectComputeShader(FVertexBufferRHIParamRef ArgumentBuffer, uint32 ArgumentOffset) final override;
 	virtual void RHIAutomaticCacheFlushAfterComputeShader(bool bEnable) final override;
@@ -230,7 +230,7 @@ public:
 	virtual void RHICopyToResolveTarget(FTextureRHIParamRef SourceTexture, FTextureRHIParamRef DestTexture, const FResolveParams& ResolveParams) final override;
 	virtual void RHICopyTexture(FTextureRHIParamRef SourceTexture, FTextureRHIParamRef DestTexture, const FRHICopyTextureInfo& CopyInfo) final override;
 	virtual void RHITransitionResources(EResourceTransitionAccess TransitionType, FTextureRHIParamRef* InRenderTargets, int32 NumTextures) final override;
-	virtual void RHITransitionResources(EResourceTransitionAccess TransitionType, EResourceTransitionPipeline TransitionPipeline, FUnorderedAccessViewRHIParamRef* InUAVs, int32 NumUAVs, FComputeFenceRHIParamRef WriteComputeFence) final override;
+	virtual void RHITransitionResources(EResourceTransitionAccess TransitionType, EResourceTransitionPipeline TransitionPipeline, FUnorderedAccessViewRHIParamRef* InUAVs, int32 NumUAVs, FRHIComputeFence* WriteComputeFence) final override;
 	virtual void RHICopyToStagingBuffer(FVertexBufferRHIParamRef SourceBuffer, FStagingBufferRHIParamRef DestinationStagingBuffer, uint32 Offset, uint32 NumBytes) final override;
 	virtual void RHIWriteGPUFence(FGPUFenceRHIParamRef Fence) final override;
 
diff --git a/Engine/Source/Runtime/VulkanRHI/Private/VulkanRenderTarget.cpp b/Engine/Source/Runtime/VulkanRHI/Private/VulkanRenderTarget.cpp
--- a/Engine/Source/Runtime/VulkanRHI/Private/VulkanRenderTarget.cpp
+++ b/Engine/Source/Runtime/VulkanRHI/Private/VulkanRenderTarget.cpp
@@ -1204,7 +1204,7 @@ void FVulkanDynamicRHI::RHIRead3DSurfaceFloatData(FTextureRHIParamRef TextureRHI
 	Device->GetImmediateContext().GetCommandBufferManager()->PrepareForNewActiveCommandBuffer();
 }
 
-void FVulkanCommandListContext::RHITransitionResources(EResourceTransitionAccess TransitionType, EResourceTransitionPipeline TransitionPipeline, FUnorderedAccessViewRHIParamRef* InUAVs, int32 NumUAVs, FComputeFenceRHIParamRef WriteComputeFenceRHI)
+void FVulkanCommandListContext::RHITransitionResources(EResourceTransitionAccess TransitionType, EResourceTransitionPipeline TransitionPipeline, FUnorderedAccessViewRHIParamRef* InUAVs, int32 NumUAVs, FRHIComputeFence* WriteComputeFenceRHI)
 {
 	FPendingTransition PendingTransition;
 	if (NumUAVs > 0)
diff --git a/Engine/Source/Runtime/VulkanRHI/Private/VulkanResources.h b/Engine/Source/Runtime/VulkanRHI/Private/VulkanResources.h
--- a/Engine/Source/Runtime/VulkanRHI/Private/VulkanResources.h
+++ b/Engine/Source/Runtime/VulkanRHI/Private/VulkanResources.h
@@ -220,7 +220,7 @@ class FVulkanBoundShaderState : public FRHIBoundShaderState
 {
 public:
 	FVulkanBoundShaderState(
-		FVertexDeclarationRHIParamRef InVertexDeclarationRHI,
+		FRHIVertexDeclaration* InVertexDeclarationRHI,
 		FVertexShaderRHIParamRef InVertexShaderRHI,
 		FPixelShaderRHIParamRef InPixelShaderRHI,
 		FHullShaderRHIParamRef InHullShaderRHI,
diff --git a/Engine/Source/Runtime/VulkanRHI/Private/VulkanShaders.cpp b/Engine/Source/Runtime/VulkanRHI/Private/VulkanShaders.cpp
--- a/Engine/Source/Runtime/VulkanRHI/Private/VulkanShaders.cpp
+++ b/Engine/Source/Runtime/VulkanRHI/Private/VulkanShaders.cpp
@@ -638,7 +638,7 @@ void FVulkanGfxPipelineDescriptorInfo::Initialize(const FDescriptorSetRemappingI
 }
 
 
-FVulkanBoundShaderState::FVulkanBoundShaderState(FVertexDeclarationRHIParamRef InVertexDeclarationRHI, FVertexShaderRHIParamRef InVertexShaderRHI,
+FVulkanBoundShaderState::FVulkanBoundShaderState(FRHIVertexDeclaration* InVertexDeclarationRHI, FVertexShaderRHIParamRef InVertexShaderRHI,
 	FPixelShaderRHIParamRef InPixelShaderRHI, FHullShaderRHIParamRef InHullShaderRHI,
 	FDomainShaderRHIParamRef InDomainShaderRHI, FGeometryShaderRHIParamRef InGeometryShaderRHI)
 	: CacheLink(InVertexDeclarationRHI, InVertexShaderRHI, InPixelShaderRHI, InHullShaderRHI, InDomainShaderRHI, InGeometryShaderRHI, this)
@@ -652,7 +652,7 @@ FVulkanBoundShaderState::~FVulkanBoundShaderState()
 }
 
 FBoundShaderStateRHIRef FVulkanDynamicRHI::RHICreateBoundShaderState(
-	FVertexDeclarationRHIParamRef VertexDeclarationRHI, 
+	FRHIVertexDeclaration* VertexDeclarationRHI,
 	FVertexShaderRHIParamRef VertexShaderRHI, 
 	FHullShaderRHIParamRef HullShaderRHI, 
 	FDomainShaderRHIParamRef DomainShaderRHI, 
diff --git a/Engine/Source/Runtime/VulkanRHI/Private/VulkanUAV.cpp b/Engine/Source/Runtime/VulkanRHI/Private/VulkanUAV.cpp
--- a/Engine/Source/Runtime/VulkanRHI/Private/VulkanUAV.cpp
+++ b/Engine/Source/Runtime/VulkanRHI/Private/VulkanUAV.cpp
@@ -434,7 +434,7 @@ FGPUFenceRHIRef FVulkanDynamicRHI::RHICreateGPUFence(const FName& Name)
 	return new FVulkanGPUFence(Name);
 }
 
-void FVulkanCommandListContext::RHIWaitComputeFence(FComputeFenceRHIParamRef InFence)
+void FVulkanCommandListContext::RHIWaitComputeFence(FRHIComputeFence* InFence)
 {
 	FVulkanComputeFence* Fence = ResourceCast(InFence);
 	FVulkanCmdBuffer* CmdBuffer = CommandBufferManager->GetActiveCmdBuffer();
diff --git a/Engine/Source/Runtime/VulkanRHI/Public/VulkanDynamicRHI.h b/Engine/Source/Runtime/VulkanRHI/Public/VulkanDynamicRHI.h
--- a/Engine/Source/Runtime/VulkanRHI/Public/VulkanDynamicRHI.h
+++ b/Engine/Source/Runtime/VulkanRHI/Public/VulkanDynamicRHI.h
@@ -55,7 +55,7 @@ public:
 	virtual FComputeFenceRHIRef RHICreateComputeFence(const FName& Name) final override;
 	virtual FGPUFenceRHIRef RHICreateGPUFence(const FName &Name) final override;
 	virtual FStagingBufferRHIRef RHICreateStagingBuffer() final override;
-	virtual FBoundShaderStateRHIRef RHICreateBoundShaderState(FVertexDeclarationRHIParamRef VertexDeclaration, FVertexShaderRHIParamRef VertexShader, FHullShaderRHIParamRef HullShader, FDomainShaderRHIParamRef DomainShader, FPixelShaderRHIParamRef PixelShader, FGeometryShaderRHIParamRef GeometryShader) final override;
+	virtual FBoundShaderStateRHIRef RHICreateBoundShaderState(FRHIVertexDeclaration* VertexDeclaration, FVertexShaderRHIParamRef VertexShader, FHullShaderRHIParamRef HullShader, FDomainShaderRHIParamRef DomainShader, FPixelShaderRHIParamRef PixelShader, FGeometryShaderRHIParamRef GeometryShader) final override;
 	virtual FGraphicsPipelineStateRHIRef RHICreateGraphicsPipelineState(const FGraphicsPipelineStateInitializer& Initializer) final override;
 	virtual FUniformBufferRHIRef RHICreateUniformBuffer(const void* Contents, const FRHIUniformBufferLayout& Layout, EUniformBufferUsage Usage, EUniformBufferValidation Validation) final override;
 	virtual void RHIUpdateUniformBuffer(FUniformBufferRHIParamRef UniformBufferRHI, const void* Contents) final override;
diff --git a/Engine/Source/Runtime/Windows/D3D11RHI/Private/D3D11Commands.cpp b/Engine/Source/Runtime/Windows/D3D11RHI/Private/D3D11Commands.cpp
--- a/Engine/Source/Runtime/Windows/D3D11RHI/Private/D3D11Commands.cpp
+++ b/Engine/Source/Runtime/Windows/D3D11RHI/Private/D3D11Commands.cpp
@@ -1910,7 +1910,7 @@ void FD3D11DynamicRHI::RHITransitionResources(EResourceTransitionAccess Transiti
 	}
 }
 
-void FD3D11DynamicRHI::RHITransitionResources(EResourceTransitionAccess TransitionType, EResourceTransitionPipeline TransitionPipeline, FUnorderedAccessViewRHIParamRef* InUAVs, int32 InNumUAVs, FComputeFenceRHIParamRef WriteFence)
+void FD3D11DynamicRHI::RHITransitionResources(EResourceTransitionAccess TransitionType, EResourceTransitionPipeline TransitionPipeline, FUnorderedAccessViewRHIParamRef* InUAVs, int32 InNumUAVs, FRHIComputeFence* WriteFence)
 {
 	for (int32 i = 0; i < InNumUAVs; ++i)
 	{
diff --git a/Engine/Source/Runtime/Windows/D3D11RHI/Private/D3D11RHIPrivate.h b/Engine/Source/Runtime/Windows/D3D11RHI/Private/D3D11RHIPrivate.h
--- a/Engine/Source/Runtime/Windows/D3D11RHI/Private/D3D11RHIPrivate.h
+++ b/Engine/Source/Runtime/Windows/D3D11RHI/Private/D3D11RHIPrivate.h
@@ -409,7 +409,7 @@ public:
 	virtual void RHIWriteGPUFence(FGPUFenceRHIParamRef Fence) final override;
 	virtual void* RHILockStagingBuffer(FStagingBufferRHIParamRef StagingBuffer, uint32 Offset, uint32 SizeRHI) final override;
     virtual void RHIUnlockStagingBuffer(FStagingBufferRHIParamRef StagingBuffer) final override;
-	virtual FBoundShaderStateRHIRef RHICreateBoundShaderState(FVertexDeclarationRHIParamRef VertexDeclaration, FVertexShaderRHIParamRef VertexShader, FHullShaderRHIParamRef HullShader, FDomainShaderRHIParamRef DomainShader, FPixelShaderRHIParamRef PixelShader, FGeometryShaderRHIParamRef GeometryShader) final override;
+	virtual FBoundShaderStateRHIRef RHICreateBoundShaderState(FRHIVertexDeclaration* VertexDeclaration, FVertexShaderRHIParamRef VertexShader, FHullShaderRHIParamRef HullShader, FDomainShaderRHIParamRef DomainShader, FPixelShaderRHIParamRef PixelShader, FGeometryShaderRHIParamRef GeometryShader) final override;
 	virtual FUniformBufferRHIRef RHICreateUniformBuffer(const void* Contents, const FRHIUniformBufferLayout& Layout, EUniformBufferUsage Usage, EUniformBufferValidation Validation) final override;
 	virtual void RHIUpdateUniformBuffer(FUniformBufferRHIParamRef UniformBufferRHI, const void* Contents) final override;
 	virtual FIndexBufferRHIRef RHICreateIndexBuffer(uint32 Stride, uint32 Size, uint32 InUsage, FRHIResourceCreateInfo& CreateInfo) final override;
@@ -504,7 +504,7 @@ public:
 	virtual void RHICopyToResolveTarget(FTextureRHIParamRef SourceTexture, FTextureRHIParamRef DestTexture, const FResolveParams& ResolveParams) final override;
 	virtual void RHICopyTexture(FTextureRHIParamRef SourceTexture, FTextureRHIParamRef DestTexture, const FRHICopyTextureInfo& CopyInfo) final override;
 	virtual void RHITransitionResources(EResourceTransitionAccess TransitionType, FTextureRHIParamRef* InTextures, int32 NumTextures) final override;
-	virtual void RHITransitionResources(EResourceTransitionAccess TransitionType, EResourceTransitionPipeline TransitionPipeline, FUnorderedAccessViewRHIParamRef* InUAVs, int32 NumUAVs, FComputeFenceRHIParamRef WriteFence) final override;
+	virtual void RHITransitionResources(EResourceTransitionAccess TransitionType, EResourceTransitionPipeline TransitionPipeline, FUnorderedAccessViewRHIParamRef* InUAVs, int32 NumUAVs, FRHIComputeFence* WriteFence) final override;
 	virtual void RHIBeginRenderQuery(FRenderQueryRHIParamRef RenderQuery) final override;
 	virtual void RHIEndRenderQuery(FRenderQueryRHIParamRef RenderQuery) final override;
 	void RHIBeginOcclusionQueryBatch(uint32 NumQueriesInBatch);
diff --git a/Engine/Source/Runtime/Windows/D3D11RHI/Private/D3D11Shaders.cpp b/Engine/Source/Runtime/Windows/D3D11RHI/Private/D3D11Shaders.cpp
--- a/Engine/Source/Runtime/Windows/D3D11RHI/Private/D3D11Shaders.cpp
+++ b/Engine/Source/Runtime/Windows/D3D11RHI/Private/D3D11Shaders.cpp
@@ -265,7 +265,7 @@ void FD3D11DynamicRHI::RHISetMultipleViewports(uint32 Count, const FViewportBoun
 }
 
 FD3D11BoundShaderState::FD3D11BoundShaderState(
-	FVertexDeclarationRHIParamRef InVertexDeclarationRHI,
+	FRHIVertexDeclaration* InVertexDeclarationRHI,
 	FVertexShaderRHIParamRef InVertexShaderRHI,
 	FPixelShaderRHIParamRef InPixelShaderRHI,
 	FHullShaderRHIParamRef InHullShaderRHI,
@@ -344,7 +344,7 @@ FD3D11BoundShaderState::~FD3D11BoundShaderState()
 * @param GeometryShader - existing geometry shader
 */
 FBoundShaderStateRHIRef FD3D11DynamicRHI::RHICreateBoundShaderState(
-	FVertexDeclarationRHIParamRef VertexDeclarationRHI, 
+	FRHIVertexDeclaration* VertexDeclarationRHI,
 	FVertexShaderRHIParamRef VertexShaderRHI, 
 	FHullShaderRHIParamRef HullShaderRHI, 
 	FDomainShaderRHIParamRef DomainShaderRHI, 
diff --git a/Engine/Source/Runtime/Windows/D3D11RHI/Public/D3D11Resources.h b/Engine/Source/Runtime/Windows/D3D11RHI/Public/D3D11Resources.h
--- a/Engine/Source/Runtime/Windows/D3D11RHI/Public/D3D11Resources.h
+++ b/Engine/Source/Runtime/Windows/D3D11RHI/Public/D3D11Resources.h
@@ -128,7 +128,7 @@ public:
 
 	/** Initialization constructor. */
 	FD3D11BoundShaderState(
-		FVertexDeclarationRHIParamRef InVertexDeclarationRHI,
+		FRHIVertexDeclaration* InVertexDeclarationRHI,
 		FVertexShaderRHIParamRef InVertexShaderRHI,
 		FPixelShaderRHIParamRef InPixelShaderRHI,
 		FHullShaderRHIParamRef InHullShaderRHI,
