DR - Deprecate FShaderResourceViewRHIParamRef and FUnorderedAccessViewRHIParamRef

diff --git a/Engine/Source/Runtime/NullDrv/Public/NullRHI.h b/Engine/Source/Runtime/NullDrv/Public/NullRHI.h
--- a/Engine/Source/Runtime/NullDrv/Public/NullRHI.h
+++ b/Engine/Source/Runtime/NullDrv/Public/NullRHI.h
@@ -232,7 +232,7 @@ public:
 		return new FRHIShaderResourceView(); 
 	}
 
-	virtual void RHIClearTinyUAV(FUnorderedAccessViewRHIParamRef UnorderedAccessViewRHI, const uint32* Values) final override
+	virtual void RHIClearTinyUAV(FRHIUnorderedAccessView* UnorderedAccessViewRHI, const uint32* Values) final override
 	{
 
 	}
@@ -601,47 +601,47 @@ public:
 	}
 
 
-	virtual void RHISetUAVParameter(FComputeShaderRHIParamRef ComputeShader, uint32 UAVIndex, FUnorderedAccessViewRHIParamRef UAV) final override
+	virtual void RHISetUAVParameter(FComputeShaderRHIParamRef ComputeShader, uint32 UAVIndex, FRHIUnorderedAccessView* UAV) final override
 	{
 
 	}
 
 
-	virtual void RHISetUAVParameter(FComputeShaderRHIParamRef ComputeShader, uint32 UAVIndex, FUnorderedAccessViewRHIParamRef UAV, uint32 InitialCount) final override
+	virtual void RHISetUAVParameter(FComputeShaderRHIParamRef ComputeShader, uint32 UAVIndex, FRHIUnorderedAccessView* UAV, uint32 InitialCount) final override
 	{
 
 	}
 
 
-	virtual void RHISetShaderResourceViewParameter(FPixelShaderRHIParamRef PixelShader, uint32 SamplerIndex, FShaderResourceViewRHIParamRef SRV) final override
+	virtual void RHISetShaderResourceViewParameter(FPixelShaderRHIParamRef PixelShader, uint32 SamplerIndex, FRHIShaderResourceView* SRV) final override
 	{
 
 	}
 
 
-	virtual void RHISetShaderResourceViewParameter(FVertexShaderRHIParamRef VertexShader, uint32 SamplerIndex, FShaderResourceViewRHIParamRef SRV) final override
+	virtual void RHISetShaderResourceViewParameter(FVertexShaderRHIParamRef VertexShader, uint32 SamplerIndex, FRHIShaderResourceView* SRV) final override
 	{
 
 	}
 
 
-	virtual void RHISetShaderResourceViewParameter(FComputeShaderRHIParamRef ComputeShader, uint32 SamplerIndex, FShaderResourceViewRHIParamRef SRV) final override
+	virtual void RHISetShaderResourceViewParameter(FComputeShaderRHIParamRef ComputeShader, uint32 SamplerIndex, FRHIShaderResourceView* SRV) final override
 	{
 
 	}
 
 
-	virtual void RHISetShaderResourceViewParameter(FHullShaderRHIParamRef HullShader, uint32 SamplerIndex, FShaderResourceViewRHIParamRef SRV) final override
+	virtual void RHISetShaderResourceViewParameter(FHullShaderRHIParamRef HullShader, uint32 SamplerIndex, FRHIShaderResourceView* SRV) final override
 	{
 
 	}
 
 
-	virtual void RHISetShaderResourceViewParameter(FDomainShaderRHIParamRef DomainShader, uint32 SamplerIndex, FShaderResourceViewRHIParamRef SRV) final override
+	virtual void RHISetShaderResourceViewParameter(FDomainShaderRHIParamRef DomainShader, uint32 SamplerIndex, FRHIShaderResourceView* SRV) final override
 	{
 
 	}
-	virtual void RHISetShaderResourceViewParameter(FGeometryShaderRHIParamRef GeometryShader, uint32 SamplerIndex, FShaderResourceViewRHIParamRef SRV) final override
+	virtual void RHISetShaderResourceViewParameter(FGeometryShaderRHIParamRef GeometryShader, uint32 SamplerIndex, FRHIShaderResourceView* SRV) final override
 	{
 
 	}
@@ -713,7 +713,7 @@ public:
 
 	}
 
-	virtual void RHISetRenderTargets(uint32 NumSimultaneousRenderTargets, const FRHIRenderTargetView* NewRenderTargets, const FRHIDepthRenderTargetView* NewDepthStencilTarget, uint32 NumUAVs, const FUnorderedAccessViewRHIParamRef* UAVs) final override
+	virtual void RHISetRenderTargets(uint32 NumSimultaneousRenderTargets, const FRHIRenderTargetView* NewRenderTargets, const FRHIDepthRenderTargetView* NewDepthStencilTarget, uint32 NumUAVs, FRHIUnorderedAccessView* const* UAVs) final override
 	{
 
 	}
diff --git a/Engine/Source/Runtime/OpenGLDrv/Private/OpenGLCommands.cpp b/Engine/Source/Runtime/OpenGLDrv/Private/OpenGLCommands.cpp
--- a/Engine/Source/Runtime/OpenGLDrv/Private/OpenGLCommands.cpp
+++ b/Engine/Source/Runtime/OpenGLDrv/Private/OpenGLCommands.cpp
@@ -445,7 +445,7 @@ void FOpenGLDynamicRHI::RHISetBoundShaderState(FRHIBoundShaderState* BoundShader
 	BoundShaderStateHistory.Add(BoundShaderState);
 }
 
-void FOpenGLDynamicRHI::RHISetUAVParameter(FComputeShaderRHIParamRef ComputeShaderRHI,uint32 UAVIndex,FUnorderedAccessViewRHIParamRef UnorderedAccessViewRHI)
+void FOpenGLDynamicRHI::RHISetUAVParameter(FComputeShaderRHIParamRef ComputeShaderRHI,uint32 UAVIndex, FRHIUnorderedAccessView* UnorderedAccessViewRHI)
 {
 	check(FOpenGL::SupportsComputeShaders());
 	
@@ -461,7 +461,7 @@ void FOpenGLDynamicRHI::RHISetUAVParameter(FComputeShaderRHIParamRef ComputeShad
 	}
 }
 
-void FOpenGLDynamicRHI::RHISetUAVParameter(FComputeShaderRHIParamRef ComputeShaderRHI,uint32 UAVIndex,FUnorderedAccessViewRHIParamRef UAVRHI, uint32 InitialCount )
+void FOpenGLDynamicRHI::RHISetUAVParameter(FComputeShaderRHIParamRef ComputeShaderRHI,uint32 UAVIndex, FRHIUnorderedAccessView* UAVRHI, uint32 InitialCount )
 {
 	// TODO: Implement for OpenGL
 	check(0);
@@ -843,7 +843,7 @@ void FOpenGLDynamicRHI::UpdateSRV(FOpenGLShaderResourceView* SRV)
 #endif
 }
 
-void FOpenGLDynamicRHI::RHISetShaderResourceViewParameter(FPixelShaderRHIParamRef PixelShaderRHI,uint32 TextureIndex,FShaderResourceViewRHIParamRef SRVRHI)
+void FOpenGLDynamicRHI::RHISetShaderResourceViewParameter(FPixelShaderRHIParamRef PixelShaderRHI,uint32 TextureIndex, FRHIShaderResourceView* SRVRHI)
 {
 	VERIFY_GL_SCOPE();
 	VALIDATE_BOUND_SHADER(PixelShaderRHI);
@@ -869,7 +869,7 @@ void FOpenGLDynamicRHI::RHISetShaderResourceViewParameter(FPixelShaderRHIParamRe
 	RHISetShaderSampler(PixelShaderRHI,TextureIndex,PointSamplerState);
 }
 
-void FOpenGLDynamicRHI::RHISetShaderResourceViewParameter(FVertexShaderRHIParamRef VertexShaderRHI,uint32 TextureIndex,FShaderResourceViewRHIParamRef SRVRHI)
+void FOpenGLDynamicRHI::RHISetShaderResourceViewParameter(FVertexShaderRHIParamRef VertexShaderRHI,uint32 TextureIndex, FRHIShaderResourceView* SRVRHI)
 {
 	VALIDATE_BOUND_SHADER(VertexShaderRHI);
 
@@ -891,7 +891,7 @@ void FOpenGLDynamicRHI::RHISetShaderResourceViewParameter(FVertexShaderRHIParamR
 	RHISetShaderSampler(VertexShaderRHI,TextureIndex,PointSamplerState);
 }
 
-void FOpenGLDynamicRHI::RHISetShaderResourceViewParameter(FComputeShaderRHIParamRef ComputeShaderRHI,uint32 TextureIndex,FShaderResourceViewRHIParamRef SRVRHI)
+void FOpenGLDynamicRHI::RHISetShaderResourceViewParameter(FComputeShaderRHIParamRef ComputeShaderRHI,uint32 TextureIndex, FRHIShaderResourceView* SRVRHI)
 {
 	VERIFY_GL_SCOPE();
 	check(FOpenGL::SupportsComputeShaders());
@@ -912,7 +912,7 @@ void FOpenGLDynamicRHI::RHISetShaderResourceViewParameter(FComputeShaderRHIParam
 	RHISetShaderSampler(ComputeShaderRHI,TextureIndex,PointSamplerState);
 }
 
-void FOpenGLDynamicRHI::RHISetShaderResourceViewParameter(FHullShaderRHIParamRef HullShaderRHI,uint32 TextureIndex,FShaderResourceViewRHIParamRef SRVRHI)
+void FOpenGLDynamicRHI::RHISetShaderResourceViewParameter(FHullShaderRHIParamRef HullShaderRHI,uint32 TextureIndex, FRHIShaderResourceView* SRVRHI)
 {
 	VALIDATE_BOUND_SHADER(HullShaderRHI);
 
@@ -934,7 +934,7 @@ void FOpenGLDynamicRHI::RHISetShaderResourceViewParameter(FHullShaderRHIParamRef
 	InternalSetShaderTexture(NULL, SRV, FOpenGL::GetFirstHullTextureUnit() + TextureIndex, Target, Resource, 0, LimitMip);
 }
 
-void FOpenGLDynamicRHI::RHISetShaderResourceViewParameter(FDomainShaderRHIParamRef DomainShaderRHI,uint32 TextureIndex,FShaderResourceViewRHIParamRef SRVRHI)
+void FOpenGLDynamicRHI::RHISetShaderResourceViewParameter(FDomainShaderRHIParamRef DomainShaderRHI,uint32 TextureIndex, FRHIShaderResourceView* SRVRHI)
 {
 	VALIDATE_BOUND_SHADER(DomainShaderRHI);
 
@@ -956,7 +956,7 @@ void FOpenGLDynamicRHI::RHISetShaderResourceViewParameter(FDomainShaderRHIParamR
 	InternalSetShaderTexture(NULL, SRV, FOpenGL::GetFirstDomainTextureUnit() + TextureIndex, Target, Resource, 0, LimitMip);
 }
 
-void FOpenGLDynamicRHI::RHISetShaderResourceViewParameter(FGeometryShaderRHIParamRef GeometryShaderRHI,uint32 TextureIndex,FShaderResourceViewRHIParamRef SRVRHI)
+void FOpenGLDynamicRHI::RHISetShaderResourceViewParameter(FGeometryShaderRHIParamRef GeometryShaderRHI,uint32 TextureIndex, FRHIShaderResourceView* SRVRHI)
 {
 	VALIDATE_BOUND_SHADER(GeometryShaderRHI);
 
@@ -1621,7 +1621,7 @@ void FOpenGLDynamicRHI::RHISetRenderTargets(
 	const FRHIRenderTargetView* NewRenderTargetsRHI,
 	const FRHIDepthRenderTargetView* NewDepthStencilTargetRHI,
 	uint32 NumUAVs,
-	const FUnorderedAccessViewRHIParamRef* UAVs
+	FRHIUnorderedAccessView* const* UAVs
 	)
 {
 	VERIFY_GL_SCOPE();
diff --git a/Engine/Source/Runtime/OpenGLDrv/Private/OpenGLTexture.cpp b/Engine/Source/Runtime/OpenGLDrv/Private/OpenGLTexture.cpp
--- a/Engine/Source/Runtime/OpenGLDrv/Private/OpenGLTexture.cpp
+++ b/Engine/Source/Runtime/OpenGLDrv/Private/OpenGLTexture.cpp
@@ -1860,7 +1860,7 @@ void FOpenGLDynamicRHI::RHIGetResourceInfo(FTextureRHIParamRef Ref, FRHIResource
 FShaderResourceViewRHIRef FOpenGLDynamicRHI::RHICreateShaderResourceView(FTextureRHIParamRef Texture, const FRHITextureSRVCreateInfo& CreateInfo)
 {
 #if 0
-	FOpenGLShaderResourceViewProxy *ViewProxy = new FOpenGLShaderResourceViewProxy([=](FShaderResourceViewRHIParamRef OwnerRHI)
+	FOpenGLShaderResourceViewProxy *ViewProxy = new FOpenGLShaderResourceViewProxy([=](FRHIShaderResourceView* OwnerRHI)
 	{
 		FOpenGLTexture2D* Texture2D = ResourceCast(Texture2DRHI);
 
@@ -1893,7 +1893,7 @@ FShaderResourceViewRHIRef FOpenGLDynamicRHI::RHICreateShaderResourceView(FTextur
 #if 0
 FShaderResourceViewRHIRef FOpenGLDynamicRHI::RHICreateShaderResourceView(FTexture2DRHIParamRef Texture2DRHI, uint8 MipLevel, uint8 NumMipLevels, uint8 Format)
 {
-	FOpenGLShaderResourceViewProxy *ViewProxy = new FOpenGLShaderResourceViewProxy([=](FShaderResourceViewRHIParamRef OwnerRHI)
+	FOpenGLShaderResourceViewProxy *ViewProxy = new FOpenGLShaderResourceViewProxy([=](FRHIShaderResourceView* OwnerRHI)
 	{
 		FOpenGLTexture2D* Texture2D = ResourceCast(Texture2DRHI);
 		FOpenGLShaderResourceView *View = 0;
@@ -2009,7 +2009,7 @@ FShaderResourceViewRHIRef FOpenGLDynamicRHI::RHICreateShaderResourceView(FTextur
 
 FShaderResourceViewRHIRef FOpenGLDynamicRHI::RHICreateShaderResourceView(FTexture3DRHIParamRef Texture3DRHI, uint8 MipLevel)
 {
-	FOpenGLShaderResourceViewProxy *ViewProxy = new FOpenGLShaderResourceViewProxy([=](FShaderResourceViewRHIParamRef OwnerRHI)
+	FOpenGLShaderResourceViewProxy *ViewProxy = new FOpenGLShaderResourceViewProxy([=](FRHIShaderResourceView* OwnerRHI)
 	{
 		FOpenGLTexture3D* Texture3D = ResourceCast(Texture3DRHI);
 
@@ -2040,7 +2040,7 @@ FShaderResourceViewRHIRef FOpenGLDynamicRHI::RHICreateShaderResourceView(FTextur
 
 FShaderResourceViewRHIRef FOpenGLDynamicRHI::RHICreateShaderResourceView(FTexture2DArrayRHIParamRef Texture2DArrayRHI, uint8 MipLevel)
 {
-	FOpenGLShaderResourceViewProxy *ViewProxy = new FOpenGLShaderResourceViewProxy([=](FShaderResourceViewRHIParamRef OwnerRHI)
+	FOpenGLShaderResourceViewProxy *ViewProxy = new FOpenGLShaderResourceViewProxy([=](FRHIShaderResourceView* OwnerRHI)
 	{
 		FOpenGLTexture2DArray* Texture2DArray = ResourceCast(Texture2DArrayRHI);
 		FOpenGLShaderResourceView *View = 0;
@@ -2071,7 +2071,7 @@ FShaderResourceViewRHIRef FOpenGLDynamicRHI::RHICreateShaderResourceView(FTextur
 FShaderResourceViewRHIRef FOpenGLDynamicRHI::RHICreateShaderResourceView(FTextureCubeRHIParamRef TextureCubeRHI, uint8 MipLevel)
 {
 
-	FOpenGLShaderResourceViewProxy *ViewProxy = new FOpenGLShaderResourceViewProxy([=](FShaderResourceViewRHIParamRef OwnerRHI)
+	FOpenGLShaderResourceViewProxy *ViewProxy = new FOpenGLShaderResourceViewProxy([=](FRHIShaderResourceView* OwnerRHI)
 	{
 		FOpenGLTextureCube* TextureCube = ResourceCast(TextureCubeRHI);
 		if (FOpenGL::SupportsTextureView())
diff --git a/Engine/Source/Runtime/OpenGLDrv/Private/OpenGLUAV.cpp b/Engine/Source/Runtime/OpenGLDrv/Private/OpenGLUAV.cpp
--- a/Engine/Source/Runtime/OpenGLDrv/Private/OpenGLUAV.cpp
+++ b/Engine/Source/Runtime/OpenGLDrv/Private/OpenGLUAV.cpp
@@ -8,7 +8,7 @@
 
 FShaderResourceViewRHIRef FOpenGLDynamicRHI::RHICreateShaderResourceView(FVertexBufferRHIParamRef VertexBufferRHI, uint32 Stride, uint8 Format)
 {
-	FShaderResourceViewRHIRef Result = new FOpenGLShaderResourceViewProxy([=, OGLRHI = this](FShaderResourceViewRHIParamRef OwnerRHI)
+	FShaderResourceViewRHIRef Result = new FOpenGLShaderResourceViewProxy([=, OGLRHI = this](FRHIShaderResourceView* OwnerRHI)
 	{ 
 		VERIFY_GL_SCOPE();
 		GLuint TextureID = 0;
@@ -47,7 +47,7 @@ FShaderResourceViewRHIRef FOpenGLDynamicRHI::RHICreateShaderResourceView(FVertex
 	return Result;
 }
 
-void FOpenGLDynamicRHI::RHIUpdateShaderResourceView(FShaderResourceViewRHIParamRef SRV, FVertexBufferRHIParamRef VertexBuffer, uint32 Stride, uint8 Format)
+void FOpenGLDynamicRHI::RHIUpdateShaderResourceView(FRHIShaderResourceView* SRV, FVertexBufferRHIParamRef VertexBuffer, uint32 Stride, uint8 Format)
 {
 	if (!FOpenGL::SupportsResourceView())
 	{
@@ -79,7 +79,7 @@ void FOpenGLDynamicRHI::RHIUpdateShaderResourceView(FShaderResourceViewRHIParamR
 
 FShaderResourceViewRHIRef FOpenGLDynamicRHI::RHICreateShaderResourceView(FIndexBufferRHIParamRef BufferRHI)
 {
-	return new FOpenGLShaderResourceViewProxy([=](FShaderResourceViewRHIParamRef OwnerRHI)
+	return new FOpenGLShaderResourceViewProxy([=](FRHIShaderResourceView* OwnerRHI)
 	{
 		VERIFY_GL_SCOPE();
 		GLuint TextureID = 0;
@@ -203,7 +203,7 @@ FUnorderedAccessViewRHIRef FOpenGLDynamicRHI::RHICreateUnorderedAccessView(FInde
 
 FShaderResourceViewRHIRef FOpenGLDynamicRHI::RHICreateShaderResourceView(FStructuredBufferRHIParamRef StructuredBufferRHI)
 {
-	return new FOpenGLShaderResourceViewProxy([=](FShaderResourceViewRHIParamRef OwnerRHI)
+	return new FOpenGLShaderResourceViewProxy([=](FRHIShaderResourceView* OwnerRHI)
 	{
 		VERIFY_GL_SCOPE();
 		GLuint TextureID = 0;
@@ -219,7 +219,7 @@ FShaderResourceViewRHIRef FOpenGLDynamicRHI::RHICreateShaderResourceView(FStruct
 	});
 }
 
-void FOpenGLDynamicRHI::RHIClearTinyUAV(FUnorderedAccessViewRHIParamRef UnorderedAccessViewRHI, const uint32* Values)
+void FOpenGLDynamicRHI::RHIClearTinyUAV(FRHIUnorderedAccessView* UnorderedAccessViewRHI, const uint32* Values)
 {
 	FOpenGLUnorderedAccessView* Texture = ResourceCast(UnorderedAccessViewRHI);
 
diff --git a/Engine/Source/Runtime/OpenGLDrv/Public/OpenGLDrv.h b/Engine/Source/Runtime/OpenGLDrv/Public/OpenGLDrv.h
--- a/Engine/Source/Runtime/OpenGLDrv/Public/OpenGLDrv.h
+++ b/Engine/Source/Runtime/OpenGLDrv/Public/OpenGLDrv.h
@@ -380,7 +380,7 @@ public:
 	virtual FShaderResourceViewRHIRef RHICreateShaderResourceView(FStructuredBufferRHIParamRef StructuredBuffer) final override;
 	virtual FShaderResourceViewRHIRef RHICreateShaderResourceView(FVertexBufferRHIParamRef VertexBuffer, uint32 Stride, uint8 Format) final override;
 	virtual FShaderResourceViewRHIRef RHICreateShaderResourceView(FIndexBufferRHIParamRef Buffer) final override;
-	virtual void RHIUpdateShaderResourceView(FShaderResourceViewRHIParamRef SRV, FVertexBufferRHIParamRef VertexBuffer, uint32 Stride, uint8 Format) final override;
+	virtual void RHIUpdateShaderResourceView(FRHIShaderResourceView* SRV, FVertexBufferRHIParamRef VertexBuffer, uint32 Stride, uint8 Format) final override;
 	virtual uint64 RHICalcTexture2DPlatformSize(uint32 SizeX, uint32 SizeY, uint8 Format, uint32 NumMips, uint32 NumSamples, uint32 Flags, uint32& OutAlign) final override;
 	virtual uint64 RHICalcTexture3DPlatformSize(uint32 SizeX, uint32 SizeY, uint32 SizeZ, uint8 Format, uint32 NumMips, uint32 Flags, uint32& OutAlign) final override;
 	virtual uint64 RHICalcTextureCubePlatformSize(uint32 Size, uint8 Format, uint32 NumMips, uint32 Flags, uint32& OutAlign) final override;
@@ -451,7 +451,7 @@ public:
 	virtual void RHIAutomaticCacheFlushAfterComputeShader(bool bEnable) final override;
 	virtual void RHIFlushComputeShaderCache() final override;
 	virtual void RHISetMultipleViewports(uint32 Count, const FViewportBounds* Data) final override;
-	virtual void RHIClearTinyUAV(FUnorderedAccessViewRHIParamRef UnorderedAccessViewRHI, const uint32* Values) final override;
+	virtual void RHIClearTinyUAV(FRHIUnorderedAccessView* UnorderedAccessViewRHI, const uint32* Values) final override;
 	virtual void RHICopyToResolveTarget(FTextureRHIParamRef SourceTexture, FTextureRHIParamRef DestTexture, const FResolveParams& ResolveParams) final override;
 	virtual void RHIBeginRenderQuery(FRenderQueryRHIParamRef RenderQuery) final override;
 	virtual void RHIEndRenderQuery(FRenderQueryRHIParamRef RenderQuery) final override;
@@ -479,14 +479,14 @@ public:
 	virtual void RHISetShaderSampler(FDomainShaderRHIParamRef DomainShader, uint32 SamplerIndex, FRHISamplerState* NewState) final override;
 	virtual void RHISetShaderSampler(FHullShaderRHIParamRef HullShader, uint32 SamplerIndex, FRHISamplerState* NewState) final override;
 	virtual void RHISetShaderSampler(FPixelShaderRHIParamRef PixelShader, uint32 SamplerIndex, FRHISamplerState* NewState) final override;
-	virtual void RHISetUAVParameter(FComputeShaderRHIParamRef ComputeShader, uint32 UAVIndex, FUnorderedAccessViewRHIParamRef UAV) final override;
-	virtual void RHISetUAVParameter(FComputeShaderRHIParamRef ComputeShader, uint32 UAVIndex, FUnorderedAccessViewRHIParamRef UAV, uint32 InitialCount) final override;
-	virtual void RHISetShaderResourceViewParameter(FPixelShaderRHIParamRef PixelShader, uint32 SamplerIndex, FShaderResourceViewRHIParamRef SRV) final override;
-	virtual void RHISetShaderResourceViewParameter(FVertexShaderRHIParamRef VertexShader, uint32 SamplerIndex, FShaderResourceViewRHIParamRef SRV) final override;
-	virtual void RHISetShaderResourceViewParameter(FComputeShaderRHIParamRef ComputeShader, uint32 SamplerIndex, FShaderResourceViewRHIParamRef SRV) final override;
-	virtual void RHISetShaderResourceViewParameter(FHullShaderRHIParamRef HullShader, uint32 SamplerIndex, FShaderResourceViewRHIParamRef SRV) final override;
-	virtual void RHISetShaderResourceViewParameter(FDomainShaderRHIParamRef DomainShader, uint32 SamplerIndex, FShaderResourceViewRHIParamRef SRV) final override;
-	virtual void RHISetShaderResourceViewParameter(FGeometryShaderRHIParamRef GeometryShader, uint32 SamplerIndex, FShaderResourceViewRHIParamRef SRV) final override;
+	virtual void RHISetUAVParameter(FComputeShaderRHIParamRef ComputeShader, uint32 UAVIndex, FRHIUnorderedAccessView* UAV) final override;
+	virtual void RHISetUAVParameter(FComputeShaderRHIParamRef ComputeShader, uint32 UAVIndex, FRHIUnorderedAccessView* UAV, uint32 InitialCount) final override;
+	virtual void RHISetShaderResourceViewParameter(FPixelShaderRHIParamRef PixelShader, uint32 SamplerIndex, FRHIShaderResourceView* SRV) final override;
+	virtual void RHISetShaderResourceViewParameter(FVertexShaderRHIParamRef VertexShader, uint32 SamplerIndex, FRHIShaderResourceView* SRV) final override;
+	virtual void RHISetShaderResourceViewParameter(FComputeShaderRHIParamRef ComputeShader, uint32 SamplerIndex, FRHIShaderResourceView* SRV) final override;
+	virtual void RHISetShaderResourceViewParameter(FHullShaderRHIParamRef HullShader, uint32 SamplerIndex, FRHIShaderResourceView* SRV) final override;
+	virtual void RHISetShaderResourceViewParameter(FDomainShaderRHIParamRef DomainShader, uint32 SamplerIndex, FRHIShaderResourceView* SRV) final override;
+	virtual void RHISetShaderResourceViewParameter(FGeometryShaderRHIParamRef GeometryShader, uint32 SamplerIndex, FRHIShaderResourceView* SRV) final override;
 	virtual void RHISetShaderUniformBuffer(FVertexShaderRHIParamRef VertexShader, uint32 BufferIndex, FUniformBufferRHIParamRef Buffer) final override;
 	virtual void RHISetShaderUniformBuffer(FHullShaderRHIParamRef HullShader, uint32 BufferIndex, FUniformBufferRHIParamRef Buffer) final override;
 	virtual void RHISetShaderUniformBuffer(FDomainShaderRHIParamRef DomainShader, uint32 BufferIndex, FUniformBufferRHIParamRef Buffer) final override;
@@ -506,7 +506,7 @@ public:
 	{
 		// Currently ignored as well as on RHISetBlendState()...
 	}
-	virtual void RHISetRenderTargets(uint32 NumSimultaneousRenderTargets, const FRHIRenderTargetView* NewRenderTargets, const FRHIDepthRenderTargetView* NewDepthStencilTarget, uint32 NumUAVs, const FUnorderedAccessViewRHIParamRef* UAVs) final override;
+	virtual void RHISetRenderTargets(uint32 NumSimultaneousRenderTargets, const FRHIRenderTargetView* NewRenderTargets, const FRHIDepthRenderTargetView* NewDepthStencilTarget, uint32 NumUAVs, FRHIUnorderedAccessView* const* UAVs) final override;
 	virtual void RHISetRenderTargetsAndClear(const FRHISetRenderTargetsInfo& RenderTargetsInfo) final override;
 	virtual void RHIDrawPrimitive(uint32 BaseVertexIndex, uint32 NumPrimitives, uint32 NumInstances) final override;
 	virtual void RHIDrawPrimitiveIndirect(FVertexBufferRHIParamRef ArgumentBuffer, uint32 ArgumentOffset) final override;
diff --git a/Engine/Source/Runtime/OpenGLDrv/Public/OpenGLResources.h b/Engine/Source/Runtime/OpenGLDrv/Public/OpenGLResources.h
--- a/Engine/Source/Runtime/OpenGLDrv/Public/OpenGLResources.h
+++ b/Engine/Source/Runtime/OpenGLDrv/Public/OpenGLResources.h
@@ -1932,7 +1932,7 @@ protected:
 class FOpenGLShaderResourceViewProxy : public TOpenGLResourceProxy<FRHIShaderResourceView, FOpenGLShaderResourceView>
 {
 public:
-	FOpenGLShaderResourceViewProxy(TFunction<FOpenGLShaderResourceView*(FShaderResourceViewRHIParamRef)> CreateFunc)
+	FOpenGLShaderResourceViewProxy(TFunction<FOpenGLShaderResourceView*(FRHIShaderResourceView*)> CreateFunc)
 		: TOpenGLResourceProxy<FRHIShaderResourceView, FOpenGLShaderResourceView>(CreateFunc)
 	{}
 
diff --git a/Engine/Source/Runtime/RHI/Private/DynamicRHI.cpp b/Engine/Source/Runtime/RHI/Private/DynamicRHI.cpp
--- a/Engine/Source/Runtime/RHI/Private/DynamicRHI.cpp
+++ b/Engine/Source/Runtime/RHI/Private/DynamicRHI.cpp
@@ -342,12 +342,12 @@ void FDynamicRHI::RHITransferVertexBufferUnderlyingResource(FVertexBufferRHIPara
 	UE_LOG(LogRHI, Fatal, TEXT("RHITransferVertexBufferUnderlyingResource isn't implemented for the current RHI"));
 }
 
-void FDynamicRHI::RHIUpdateShaderResourceView(FShaderResourceViewRHIParamRef SRV, FVertexBufferRHIParamRef VertexBuffer, uint32 Stride, uint8 Format)
+void FDynamicRHI::RHIUpdateShaderResourceView(FRHIShaderResourceView* SRV, FVertexBufferRHIParamRef VertexBuffer, uint32 Stride, uint8 Format)
 {
 	UE_LOG(LogRHI, Fatal, TEXT("RHIUpdateShaderResourceView isn't implemented for the current RHI"));
 }
 
-void FDynamicRHI::RHIUpdateShaderResourceView(FShaderResourceViewRHIParamRef SRV, FIndexBufferRHIParamRef IndexBuffer)
+void FDynamicRHI::RHIUpdateShaderResourceView(FRHIShaderResourceView* SRV, FIndexBufferRHIParamRef IndexBuffer)
 {
 	UE_LOG(LogRHI, Fatal, TEXT("RHIUpdateShaderResourceView isn't implemented for the current RHI"));
 }
diff --git a/Engine/Source/Runtime/RHI/Public/DynamicRHI.h b/Engine/Source/Runtime/RHI/Public/DynamicRHI.h
--- a/Engine/Source/Runtime/RHI/Public/DynamicRHI.h
+++ b/Engine/Source/Runtime/RHI/Public/DynamicRHI.h
@@ -519,9 +519,9 @@ public:
 
 	// Must be called on RHI thread timeline
 	// Make sure to call RHIThreadFence(true) afterwards so that parallel translation doesn't refer old resources
-	virtual void RHIUpdateShaderResourceView(FShaderResourceViewRHIParamRef SRV, FVertexBufferRHIParamRef VertexBuffer, uint32 Stride, uint8 Format);
+	virtual void RHIUpdateShaderResourceView(FRHIShaderResourceView* SRV, FVertexBufferRHIParamRef VertexBuffer, uint32 Stride, uint8 Format);
 
-	virtual void RHIUpdateShaderResourceView(FShaderResourceViewRHIParamRef SRV, FIndexBufferRHIParamRef IndexBuffer);
+	virtual void RHIUpdateShaderResourceView(FRHIShaderResourceView* SRV, FIndexBufferRHIParamRef IndexBuffer);
 
 	/**
 	* Computes the total size of a 2D texture with the specified parameters.
@@ -873,7 +873,7 @@ public:
 
 	// FlushType: Thread safe
 	virtual void RHIBindDebugLabelName(FTextureRHIParamRef Texture, const TCHAR* Name) = 0;
-	virtual void RHIBindDebugLabelName(FUnorderedAccessViewRHIParamRef UnorderedAccessViewRHI, const TCHAR* Name) {}
+	virtual void RHIBindDebugLabelName(FRHIUnorderedAccessView* UnorderedAccessViewRHI, const TCHAR* Name) {}
 
 	/**
 	* Reads the contents of a texture to an output buffer (non MSAA and MSAA) and returns it as a FColor array.
@@ -1315,7 +1315,7 @@ FORCEINLINE void RHIBindDebugLabelName(FTextureRHIParamRef Texture, const TCHAR*
 	GDynamicRHI->RHIBindDebugLabelName(Texture, Name);
 }
 
-FORCEINLINE void RHIBindDebugLabelName(FUnorderedAccessViewRHIParamRef UnorderedAccessViewRHI, const TCHAR* Name)
+FORCEINLINE void RHIBindDebugLabelName(FRHIUnorderedAccessView* UnorderedAccessViewRHI, const TCHAR* Name)
 {
 	GDynamicRHI->RHIBindDebugLabelName(UnorderedAccessViewRHI, Name);
 }
diff --git a/Engine/Source/Runtime/RHI/Public/RHICommandList.h b/Engine/Source/Runtime/RHI/Public/RHICommandList.h
--- a/Engine/Source/Runtime/RHI/Public/RHICommandList.h
+++ b/Engine/Source/Runtime/RHI/Public/RHICommandList.h
@@ -113,10 +113,10 @@ extern RHI_API TAutoConsoleVariable<int32> CVarRHICmdFlushRenderThreadTasks;
 struct FRayTracingShaderBindings
 {
 	FTextureRHIParamRef Textures[32] = {};
-	FShaderResourceViewRHIParamRef SRVs[32] = {};
+	FRHIShaderResourceView* SRVs[32] = {};
 	FUniformBufferRHIParamRef UniformBuffers[8] = {};
 	FRHISamplerState* Samplers[16] = {};
-	FUnorderedAccessViewRHIParamRef UAVs[8] = {};
+	FRHIUnorderedAccessView* UAVs[8] = {};
 };
 
 // C++ counter-part of FBasicRayData declared in RayTracingCommon.ush
@@ -701,8 +701,8 @@ struct  FRHICommandEndUpdateMultiFrameResource final : public FRHICommand<FRHICo
 
 struct  FRHICommandBeginUpdateMultiFrameUAV final : public FRHICommand<FRHICommandBeginUpdateMultiFrameResource>
 {
-	FUnorderedAccessViewRHIParamRef UAV;
-	FORCEINLINE_DEBUGGABLE FRHICommandBeginUpdateMultiFrameUAV(FUnorderedAccessViewRHIParamRef InUAV)
+	FRHIUnorderedAccessView* UAV;
+	FORCEINLINE_DEBUGGABLE FRHICommandBeginUpdateMultiFrameUAV(FRHIUnorderedAccessView* InUAV)
 		: UAV(InUAV)
 	{
 	}
@@ -711,8 +711,8 @@ struct  FRHICommandBeginUpdateMultiFrameUAV final : public FRHICommand<FRHIComma
 
 struct  FRHICommandEndUpdateMultiFrameUAV final : public FRHICommand<FRHICommandEndUpdateMultiFrameResource>
 {
-	FUnorderedAccessViewRHIParamRef UAV;
-	FORCEINLINE_DEBUGGABLE FRHICommandEndUpdateMultiFrameUAV(FUnorderedAccessViewRHIParamRef InUAV)
+	FRHIUnorderedAccessView* UAV;
+	FORCEINLINE_DEBUGGABLE FRHICommandEndUpdateMultiFrameUAV(FRHIUnorderedAccessView* InUAV)
 		: UAV(InUAV)
 	{
 	}
@@ -783,8 +783,8 @@ struct FRHICommandSetShaderResourceViewParameter final : public FRHICommand<FRHI
 {
 	TShaderRHIParamRef Shader;
 	uint32 SamplerIndex;
-	FShaderResourceViewRHIParamRef SRV;
-	FORCEINLINE_DEBUGGABLE FRHICommandSetShaderResourceViewParameter(TShaderRHIParamRef InShader, uint32 InSamplerIndex, FShaderResourceViewRHIParamRef InSRV)
+	FRHIShaderResourceView* SRV;
+	FORCEINLINE_DEBUGGABLE FRHICommandSetShaderResourceViewParameter(TShaderRHIParamRef InShader, uint32 InSamplerIndex, FRHIShaderResourceView* InSRV)
 		: Shader(InShader)
 		, SamplerIndex(InSamplerIndex)
 		, SRV(InSRV)
@@ -798,8 +798,8 @@ struct FRHICommandSetUAVParameter final : public FRHICommand<FRHICommandSetUAVPa
 {
 	TShaderRHIParamRef Shader;
 	uint32 UAVIndex;
-	FUnorderedAccessViewRHIParamRef UAV;
-	FORCEINLINE_DEBUGGABLE FRHICommandSetUAVParameter(TShaderRHIParamRef InShader, uint32 InUAVIndex, FUnorderedAccessViewRHIParamRef InUAV)
+	FRHIUnorderedAccessView* UAV;
+	FORCEINLINE_DEBUGGABLE FRHICommandSetUAVParameter(TShaderRHIParamRef InShader, uint32 InUAVIndex, FRHIUnorderedAccessView* InUAV)
 		: Shader(InShader)
 		, UAVIndex(InUAVIndex)
 		, UAV(InUAV)
@@ -813,9 +813,9 @@ struct FRHICommandSetUAVParameter_IntialCount final : public FRHICommand<FRHICom
 {
 	TShaderRHIParamRef Shader;
 	uint32 UAVIndex;
-	FUnorderedAccessViewRHIParamRef UAV;
+	FRHIUnorderedAccessView* UAV;
 	uint32 InitialCount;
-	FORCEINLINE_DEBUGGABLE FRHICommandSetUAVParameter_IntialCount(TShaderRHIParamRef InShader, uint32 InUAVIndex, FUnorderedAccessViewRHIParamRef InUAV, uint32 InInitialCount)
+	FORCEINLINE_DEBUGGABLE FRHICommandSetUAVParameter_IntialCount(TShaderRHIParamRef InShader, uint32 InUAVIndex, FRHIUnorderedAccessView* InUAV, uint32 InInitialCount)
 		: Shader(InShader)
 		, UAVIndex(InUAVIndex)
 		, UAV(InUAV)
@@ -972,14 +972,14 @@ struct FRHICommandSetRenderTargets final : public FRHICommand<FRHICommandSetRend
 	FRHIRenderTargetView NewRenderTargetsRHI[MaxSimultaneousRenderTargets];
 	FRHIDepthRenderTargetView NewDepthStencilTarget;
 	uint32 NewNumUAVs;
-	FUnorderedAccessViewRHIParamRef UAVs[MaxSimultaneousUAVs];
+	FRHIUnorderedAccessView* UAVs[MaxSimultaneousUAVs];
 
 	FORCEINLINE_DEBUGGABLE FRHICommandSetRenderTargets(
 		uint32 InNewNumSimultaneousRenderTargets,
 		const FRHIRenderTargetView* InNewRenderTargetsRHI,
 		const FRHIDepthRenderTargetView* InNewDepthStencilTargetRHI,
 		uint32 InNewNumUAVs,
-		const FUnorderedAccessViewRHIParamRef* InUAVs
+		FRHIUnorderedAccessView* const* InUAVs
 		)
 		: NewNumSimultaneousRenderTargets(InNewNumSimultaneousRenderTargets)
 		, NewNumUAVs(InNewNumUAVs)
@@ -1287,10 +1287,10 @@ struct FRHICommandSetDepthBounds final : public FRHICommand<FRHICommandSetDepthB
 
 struct FRHICommandClearTinyUAV final : public FRHICommand<FRHICommandClearTinyUAV>
 {
-	FUnorderedAccessViewRHIParamRef UnorderedAccessViewRHI;
+	FRHIUnorderedAccessView* UnorderedAccessViewRHI;
 	uint32 Values[4];
 
-	FORCEINLINE_DEBUGGABLE FRHICommandClearTinyUAV(FUnorderedAccessViewRHIParamRef InUnorderedAccessViewRHI, const uint32* InValues)
+	FORCEINLINE_DEBUGGABLE FRHICommandClearTinyUAV(FRHIUnorderedAccessView* InUnorderedAccessViewRHI, const uint32* InValues)
 		: UnorderedAccessViewRHI(InUnorderedAccessViewRHI)
 	{
 		Values[0] = InValues[0];
@@ -1369,12 +1369,12 @@ template<ECmdList CmdListType>
 struct FRHICommandTransitionUAVs final : public FRHICommand<FRHICommandTransitionUAVs<CmdListType>>
 {
 	int32 NumUAVs;
-	FUnorderedAccessViewRHIParamRef* UAVs; // Pointer to an array of UAVs, allocated inline with the command list
+	FRHIUnorderedAccessView** UAVs; // Pointer to an array of UAVs, allocated inline with the command list
 	EResourceTransitionAccess TransitionType;
 	EResourceTransitionPipeline TransitionPipeline;
 	FRHIComputeFence* WriteFence;
 
-	FORCEINLINE_DEBUGGABLE FRHICommandTransitionUAVs(EResourceTransitionAccess InTransitionType, EResourceTransitionPipeline InTransitionPipeline, FUnorderedAccessViewRHIParamRef* InUAVs, int32 InNumUAVs, FRHIComputeFence* InWriteFence)
+	FORCEINLINE_DEBUGGABLE FRHICommandTransitionUAVs(EResourceTransitionAccess InTransitionType, EResourceTransitionPipeline InTransitionPipeline, FRHIUnorderedAccessView** InUAVs, int32 InNumUAVs, FRHIComputeFence* InWriteFence)
 		: NumUAVs(InNumUAVs)
 		, UAVs(InUAVs)
 		, TransitionType(InTransitionType)
@@ -1769,7 +1769,7 @@ struct FRHICommandUpdateTextureReference final : public FRHICommand<FRHICommandU
 
 struct FRHIShaderResourceViewUpdateInfo_VB
 {
-	FShaderResourceViewRHIParamRef SRV;
+	FRHIShaderResourceView* SRV;
 	FVertexBufferRHIParamRef VertexBuffer;
 	uint32 Stride;
 	uint8 Format;
@@ -1905,13 +1905,13 @@ struct FRHICommandBuildAccelerationStructures final : public FRHICommand<FRHICom
 struct FRHICommandRayTraceOcclusion final : public FRHICommand<FRHICommandRayTraceOcclusion>
 {
 	FRayTracingSceneRHIParamRef Scene;
-	FShaderResourceViewRHIParamRef Rays;
-	FUnorderedAccessViewRHIParamRef Output;
+	FRHIShaderResourceView* Rays;
+	FRHIUnorderedAccessView* Output;
 	uint32 NumRays;
 
 	FRHICommandRayTraceOcclusion(FRayTracingSceneRHIParamRef InScene,
-		FShaderResourceViewRHIParamRef InRays,
-		FUnorderedAccessViewRHIParamRef InOutput,
+		FRHIShaderResourceView* InRays,
+		FRHIUnorderedAccessView* InOutput,
 		uint32 InNumRays)
 		: Scene(InScene)
 		, Rays(InRays)
@@ -1925,13 +1925,13 @@ struct FRHICommandRayTraceOcclusion final : public FRHICommand<FRHICommandRayTra
 struct FRHICommandRayTraceIntersection final : public FRHICommand<FRHICommandRayTraceIntersection>
 {
 	FRayTracingSceneRHIParamRef Scene;
-	FShaderResourceViewRHIParamRef Rays;
-	FUnorderedAccessViewRHIParamRef Output;
+	FRHIShaderResourceView* Rays;
+	FRHIUnorderedAccessView* Output;
 	uint32 NumRays;
 
 	FRHICommandRayTraceIntersection(FRayTracingSceneRHIParamRef InScene,
-		FShaderResourceViewRHIParamRef InRays,
-		FUnorderedAccessViewRHIParamRef InOutput,
+		FRHIShaderResourceView* InRays,
+		FRHIUnorderedAccessView* InOutput,
 		uint32 InNumRays)
 		: Scene(InScene)
 		, Rays(InRays)
@@ -2074,7 +2074,7 @@ public:
 		ALLOC_COMMAND(FRHICommandEndUpdateMultiFrameResource)(Texture);
 	}
 
-	FORCEINLINE_DEBUGGABLE void BeginUpdateMultiFrameResource(FUnorderedAccessViewRHIParamRef UAV)
+	FORCEINLINE_DEBUGGABLE void BeginUpdateMultiFrameResource(FRHIUnorderedAccessView* UAV)
 	{
 		//check(IsOutsideRenderPass());
 		if (Bypass())
@@ -2085,7 +2085,7 @@ public:
 		ALLOC_COMMAND(FRHICommandBeginUpdateMultiFrameUAV)(UAV);
 	}
 
-	FORCEINLINE_DEBUGGABLE void EndUpdateMultiFrameResource(FUnorderedAccessViewRHIParamRef UAV)
+	FORCEINLINE_DEBUGGABLE void EndUpdateMultiFrameResource(FRHIUnorderedAccessView* UAV)
 	{
 		//check(IsOutsideRenderPass());
 		if (Bypass())
@@ -2186,7 +2186,7 @@ public:
 	}
 
 	template <typename TShaderRHIParamRef>
-	FORCEINLINE_DEBUGGABLE void SetShaderResourceViewParameter(TShaderRHIParamRef Shader, uint32 SamplerIndex, FShaderResourceViewRHIParamRef SRV)
+	FORCEINLINE_DEBUGGABLE void SetShaderResourceViewParameter(TShaderRHIParamRef Shader, uint32 SamplerIndex, FRHIShaderResourceView* SRV)
 	{
 		//check(IsOutsideRenderPass());
 		if (Bypass())
@@ -2198,7 +2198,7 @@ public:
 	}
 
 	template <typename TShaderRHI>
-	FORCEINLINE_DEBUGGABLE void SetShaderResourceViewParameter(TRefCountPtr<TShaderRHI>& Shader, uint32 SamplerIndex, FShaderResourceViewRHIParamRef SRV)
+	FORCEINLINE_DEBUGGABLE void SetShaderResourceViewParameter(TRefCountPtr<TShaderRHI>& Shader, uint32 SamplerIndex, FRHIShaderResourceView* SRV)
 	{
 		SetShaderResourceViewParameter(Shader.GetReference(), SamplerIndex, SRV);
 	}
@@ -2229,7 +2229,7 @@ public:
 		SetShaderSampler(Shader.GetReference(), SamplerIndex, State);
 	}
 
-	FORCEINLINE_DEBUGGABLE void SetUAVParameter(FComputeShaderRHIParamRef Shader, uint32 UAVIndex, FUnorderedAccessViewRHIParamRef UAV)
+	FORCEINLINE_DEBUGGABLE void SetUAVParameter(FComputeShaderRHIParamRef Shader, uint32 UAVIndex, FRHIUnorderedAccessView* UAV)
 	{
 		if (Bypass())
 		{
@@ -2239,12 +2239,12 @@ public:
 		ALLOC_COMMAND(FRHICommandSetUAVParameter<FComputeShaderRHIParamRef, ECmdList::EGfx>)(Shader, UAVIndex, UAV);
 	}
 
-	FORCEINLINE_DEBUGGABLE void SetUAVParameter(TRefCountPtr<FRHIComputeShader>& Shader, uint32 UAVIndex, FUnorderedAccessViewRHIParamRef UAV)
+	FORCEINLINE_DEBUGGABLE void SetUAVParameter(TRefCountPtr<FRHIComputeShader>& Shader, uint32 UAVIndex, FRHIUnorderedAccessView* UAV)
 	{
 		SetUAVParameter(Shader.GetReference(), UAVIndex, UAV);
 	}
 
-	FORCEINLINE_DEBUGGABLE void SetUAVParameter(FComputeShaderRHIParamRef Shader, uint32 UAVIndex, FUnorderedAccessViewRHIParamRef UAV, uint32 InitialCount)
+	FORCEINLINE_DEBUGGABLE void SetUAVParameter(FComputeShaderRHIParamRef Shader, uint32 UAVIndex, FRHIUnorderedAccessView* UAV, uint32 InitialCount)
 	{
 		if (Bypass())
 		{
@@ -2254,7 +2254,7 @@ public:
 		ALLOC_COMMAND(FRHICommandSetUAVParameter_IntialCount<FComputeShaderRHIParamRef, ECmdList::EGfx>)(Shader, UAVIndex, UAV, InitialCount);
 	}
 
-	FORCEINLINE_DEBUGGABLE void SetUAVParameter(TRefCountPtr<FRHIComputeShader>& Shader, uint32 UAVIndex, FUnorderedAccessViewRHIParamRef UAV, uint32 InitialCount)
+	FORCEINLINE_DEBUGGABLE void SetUAVParameter(TRefCountPtr<FRHIComputeShader>& Shader, uint32 UAVIndex, FRHIUnorderedAccessView* UAV, uint32 InitialCount)
 	{
 		SetUAVParameter(Shader.GetReference(), UAVIndex, UAV, InitialCount);
 	}
@@ -2407,7 +2407,7 @@ public:
 		const FRHIRenderTargetView* NewRenderTargetsRHI,
 		const FRHIDepthRenderTargetView* NewDepthStencilTargetRHI,
 		uint32 NewNumUAVs,
-		const FUnorderedAccessViewRHIParamRef* UAVs
+		FRHIUnorderedAccessView* const* UAVs
 		)
 	{
 		check(IsOutsideRenderPass());
@@ -2630,7 +2630,7 @@ public:
 		}
 	}
 
-	FORCEINLINE_DEBUGGABLE void ClearTinyUAV(FUnorderedAccessViewRHIParamRef UnorderedAccessViewRHI, const uint32(&Values)[4])
+	FORCEINLINE_DEBUGGABLE void ClearTinyUAV(FRHIUnorderedAccessView* UnorderedAccessViewRHI, const uint32(&Values)[4])
 	{
 		//check(IsOutsideRenderPass());
 		if (Bypass())
@@ -2724,10 +2724,10 @@ public:
 		ALLOC_COMMAND(FRHICommandTransitionTexturesArray)(TransitionType, InTextures);
 	}
 
-	FORCEINLINE_DEBUGGABLE void TransitionResource(EResourceTransitionAccess TransitionType, EResourceTransitionPipeline TransitionPipeline, FUnorderedAccessViewRHIParamRef InUAV, FRHIComputeFence* WriteFence)
+	FORCEINLINE_DEBUGGABLE void TransitionResource(EResourceTransitionAccess TransitionType, EResourceTransitionPipeline TransitionPipeline, FRHIUnorderedAccessView* InUAV, FRHIComputeFence* WriteFence)
 	{
 		check(InUAV == nullptr || InUAV->IsCommitted());
-		FUnorderedAccessViewRHIParamRef UAV = InUAV;
+		FRHIUnorderedAccessView* UAV = InUAV;
 		if (Bypass())
 		{
 			GetContext().RHITransitionResources(TransitionType, TransitionPipeline, &UAV, 1, WriteFence);
@@ -2735,17 +2735,17 @@ public:
 		}
 
 		// Allocate space to hold the single UAV pointer inline in the command list itself.
-		FUnorderedAccessViewRHIParamRef* UAVArray = (FUnorderedAccessViewRHIParamRef*)Alloc(sizeof(FUnorderedAccessViewRHIParamRef), alignof(FUnorderedAccessViewRHIParamRef));
+		FRHIUnorderedAccessView** UAVArray = (FRHIUnorderedAccessView**)Alloc(sizeof(FRHIUnorderedAccessView*), alignof(FRHIUnorderedAccessView*));
 		UAVArray[0] = UAV;
 		ALLOC_COMMAND(FRHICommandTransitionUAVs<ECmdList::EGfx>)(TransitionType, TransitionPipeline, UAVArray, 1, WriteFence);
 	}
 
-	FORCEINLINE_DEBUGGABLE void TransitionResource(EResourceTransitionAccess TransitionType, EResourceTransitionPipeline TransitionPipeline, FUnorderedAccessViewRHIParamRef InUAV)
+	FORCEINLINE_DEBUGGABLE void TransitionResource(EResourceTransitionAccess TransitionType, EResourceTransitionPipeline TransitionPipeline, FRHIUnorderedAccessView* InUAV)
 	{
 		TransitionResource(TransitionType, TransitionPipeline, InUAV, nullptr);
 	}
 
-	FORCEINLINE_DEBUGGABLE void TransitionResources(EResourceTransitionAccess TransitionType, EResourceTransitionPipeline TransitionPipeline, FUnorderedAccessViewRHIParamRef* InUAVs, int32 NumUAVs, FRHIComputeFence* WriteFence)
+	FORCEINLINE_DEBUGGABLE void TransitionResources(EResourceTransitionAccess TransitionType, EResourceTransitionPipeline TransitionPipeline, FRHIUnorderedAccessView** InUAVs, int32 NumUAVs, FRHIComputeFence* WriteFence)
 	{
 		if (Bypass())
 		{
@@ -2754,7 +2754,7 @@ public:
 		}
 
 		// Allocate space to hold the list UAV pointers inline in the command list itself.
-		FUnorderedAccessViewRHIParamRef* UAVArray = (FUnorderedAccessViewRHIParamRef*)Alloc(sizeof(FUnorderedAccessViewRHIParamRef) * NumUAVs, alignof(FUnorderedAccessViewRHIParamRef));
+		FRHIUnorderedAccessView** UAVArray = (FRHIUnorderedAccessView**)Alloc(sizeof(FRHIUnorderedAccessView*) * NumUAVs, alignof(FRHIUnorderedAccessView*));
 		for (int32 Index = 0; Index < NumUAVs; ++Index)
 		{
 			UAVArray[Index] = InUAVs[Index];
@@ -2763,7 +2763,7 @@ public:
 		ALLOC_COMMAND(FRHICommandTransitionUAVs<ECmdList::EGfx>)(TransitionType, TransitionPipeline, UAVArray, NumUAVs, WriteFence);
 	}
 
-	FORCEINLINE_DEBUGGABLE void TransitionResources(EResourceTransitionAccess TransitionType, EResourceTransitionPipeline TransitionPipeline, FUnorderedAccessViewRHIParamRef* InUAVs, int32 NumUAVs)
+	FORCEINLINE_DEBUGGABLE void TransitionResources(EResourceTransitionAccess TransitionType, EResourceTransitionPipeline TransitionPipeline, FRHIUnorderedAccessView** InUAVs, int32 NumUAVs)
 	{
 		TransitionResources(TransitionType, TransitionPipeline, InUAVs, NumUAVs, nullptr);
 	}
@@ -3039,8 +3039,8 @@ public:
 	 * 0xFFFFFFFF is written if ray intersects any scene triangle, 0 otherwise.
 	 */
 	FORCEINLINE_DEBUGGABLE void RayTraceOcclusion(FRayTracingSceneRHIParamRef Scene,
-		FShaderResourceViewRHIParamRef Rays,
-		FUnorderedAccessViewRHIParamRef Output,
+		FRHIShaderResourceView* Rays,
+		FRHIUnorderedAccessView* Output,
 		uint32 NumRays)
 	{
 		if (Bypass())
@@ -3058,8 +3058,8 @@ public:
 	 * Primitive intersection results are written to output buffer as FBasicRayIntersectionData.
 	 */
 	FORCEINLINE_DEBUGGABLE void RayTraceIntersection(FRayTracingSceneRHIParamRef Scene,
-		FShaderResourceViewRHIParamRef Rays,
-		FUnorderedAccessViewRHIParamRef Output,
+		FRHIShaderResourceView* Rays,
+		FRHIUnorderedAccessView* Output,
 		uint32 NumRays)
 	{
 		if (Bypass())
@@ -3208,7 +3208,7 @@ public:
 		ALLOC_COMMAND(FRHICommandSetShaderTexture<FComputeShaderRHIParamRef, ECmdList::ECompute>)(Shader, TextureIndex, Texture);
 	}
 	
-	FORCEINLINE_DEBUGGABLE void SetShaderResourceViewParameter(FComputeShaderRHIParamRef Shader, uint32 SamplerIndex, FShaderResourceViewRHIParamRef SRV)
+	FORCEINLINE_DEBUGGABLE void SetShaderResourceViewParameter(FComputeShaderRHIParamRef Shader, uint32 SamplerIndex, FRHIShaderResourceView* SRV)
 	{
 		if (Bypass())
 		{
@@ -3235,7 +3235,7 @@ public:
 		ALLOC_COMMAND(FRHICommandSetShaderSampler<FComputeShaderRHIParamRef, ECmdList::ECompute>)(Shader, SamplerIndex, State);
 	}
 
-	FORCEINLINE_DEBUGGABLE void SetUAVParameter(FComputeShaderRHIParamRef Shader, uint32 UAVIndex, FUnorderedAccessViewRHIParamRef UAV)
+	FORCEINLINE_DEBUGGABLE void SetUAVParameter(FComputeShaderRHIParamRef Shader, uint32 UAVIndex, FRHIUnorderedAccessView* UAV)
 	{
 		if (Bypass())
 		{
@@ -3245,7 +3245,7 @@ public:
 		ALLOC_COMMAND(FRHICommandSetUAVParameter<FComputeShaderRHIParamRef, ECmdList::ECompute>)(Shader, UAVIndex, UAV);
 	}
 
-	FORCEINLINE_DEBUGGABLE void SetUAVParameter(FComputeShaderRHIParamRef Shader, uint32 UAVIndex, FUnorderedAccessViewRHIParamRef UAV, uint32 InitialCount)
+	FORCEINLINE_DEBUGGABLE void SetUAVParameter(FComputeShaderRHIParamRef Shader, uint32 UAVIndex, FRHIUnorderedAccessView* UAV, uint32 InitialCount)
 	{
 		if (Bypass())
 		{
@@ -3308,9 +3308,9 @@ public:
 		ALLOC_COMMAND(FRHICommandDispatchIndirectComputeShader<ECmdList::ECompute>)(ArgumentBuffer, ArgumentOffset);
 	}
 
-	FORCEINLINE_DEBUGGABLE void TransitionResource(EResourceTransitionAccess TransitionType, EResourceTransitionPipeline TransitionPipeline, FUnorderedAccessViewRHIParamRef InUAV, FRHIComputeFence* WriteFence)
+	FORCEINLINE_DEBUGGABLE void TransitionResource(EResourceTransitionAccess TransitionType, EResourceTransitionPipeline TransitionPipeline, FRHIUnorderedAccessView* InUAV, FRHIComputeFence* WriteFence)
 	{
-		FUnorderedAccessViewRHIParamRef UAV = InUAV;
+		FRHIUnorderedAccessView* UAV = InUAV;
 		if (Bypass())
 		{
 			GetComputeContext().RHITransitionResources(TransitionType, TransitionPipeline, &UAV, 1, WriteFence);
@@ -3318,17 +3318,17 @@ public:
 		}
 
 		// Allocate space to hold the single UAV pointer inline in the command list itself.
-		FUnorderedAccessViewRHIParamRef* UAVArray = (FUnorderedAccessViewRHIParamRef*)Alloc(sizeof(FUnorderedAccessViewRHIParamRef), alignof(FUnorderedAccessViewRHIParamRef));
+		FRHIUnorderedAccessView** UAVArray = (FRHIUnorderedAccessView**)Alloc(sizeof(FRHIUnorderedAccessView*), alignof(FRHIUnorderedAccessView*));
 		UAVArray[0] = UAV;
 		ALLOC_COMMAND(FRHICommandTransitionUAVs<ECmdList::ECompute>)(TransitionType, TransitionPipeline, UAVArray, 1, WriteFence);
 	}
 
-	FORCEINLINE_DEBUGGABLE void TransitionResource(EResourceTransitionAccess TransitionType, EResourceTransitionPipeline TransitionPipeline, FUnorderedAccessViewRHIParamRef InUAV)
+	FORCEINLINE_DEBUGGABLE void TransitionResource(EResourceTransitionAccess TransitionType, EResourceTransitionPipeline TransitionPipeline, FRHIUnorderedAccessView* InUAV)
 	{
 		TransitionResource(TransitionType, TransitionPipeline, InUAV, nullptr);
 	}
 
-	FORCEINLINE_DEBUGGABLE void TransitionResources(EResourceTransitionAccess TransitionType, EResourceTransitionPipeline TransitionPipeline, FUnorderedAccessViewRHIParamRef* InUAVs, int32 NumUAVs, FRHIComputeFence* WriteFence)
+	FORCEINLINE_DEBUGGABLE void TransitionResources(EResourceTransitionAccess TransitionType, EResourceTransitionPipeline TransitionPipeline, FRHIUnorderedAccessView** InUAVs, int32 NumUAVs, FRHIComputeFence* WriteFence)
 	{
 		if (Bypass())
 		{
@@ -3337,7 +3337,7 @@ public:
 		}
 
 		// Allocate space to hold the list UAV pointers inline in the command list itself.
-		FUnorderedAccessViewRHIParamRef* UAVArray = (FUnorderedAccessViewRHIParamRef*)Alloc(sizeof(FUnorderedAccessViewRHIParamRef) * NumUAVs, alignof(FUnorderedAccessViewRHIParamRef));
+		FRHIUnorderedAccessView** UAVArray = (FRHIUnorderedAccessView**)Alloc(sizeof(FRHIUnorderedAccessView*) * NumUAVs, alignof(FRHIUnorderedAccessView*));
 		for (int32 Index = 0; Index < NumUAVs; ++Index)
 		{
 			UAVArray[Index] = InUAVs[Index];
@@ -3346,7 +3346,7 @@ public:
 		ALLOC_COMMAND(FRHICommandTransitionUAVs<ECmdList::ECompute>)(TransitionType, TransitionPipeline, UAVArray, NumUAVs, WriteFence);
 	}
 
-	FORCEINLINE_DEBUGGABLE void TransitionResources(EResourceTransitionAccess TransitionType, EResourceTransitionPipeline TransitionPipeline, FUnorderedAccessViewRHIParamRef* InUAVs, int32 NumUAVs)
+	FORCEINLINE_DEBUGGABLE void TransitionResources(EResourceTransitionAccess TransitionType, EResourceTransitionPipeline TransitionPipeline, FRHIUnorderedAccessView** InUAVs, int32 NumUAVs)
 	{
 		TransitionResources(TransitionType, TransitionPipeline, InUAVs, NumUAVs, nullptr);
 	}
@@ -4034,7 +4034,7 @@ public:
 		RHIBindDebugLabelName(Texture, Name);
 	}
 
-	FORCEINLINE void BindDebugLabelName(FUnorderedAccessViewRHIParamRef UnorderedAccessViewRHI, const TCHAR* Name)
+	FORCEINLINE void BindDebugLabelName(FRHIUnorderedAccessView* UnorderedAccessViewRHI, const TCHAR* Name)
 	{
 		RHIBindDebugLabelName(UnorderedAccessViewRHI, Name);
 	}
@@ -4958,7 +4958,7 @@ struct TRHIResourceUpdateBatcher
 		}
 	}
 
-	void QueueUpdateRequest(FShaderResourceViewRHIParamRef SRV, FVertexBufferRHIParamRef VertexBuffer, uint32 Stride, uint8 Format)
+	void QueueUpdateRequest(FRHIShaderResourceView* SRV, FVertexBufferRHIParamRef VertexBuffer, uint32 Stride, uint8 Format)
 	{
 		FRHIResourceUpdateInfo& UpdateInfo = GetNextUpdateInfo();
 		UpdateInfo.Type = FRHIResourceUpdateInfo::UT_VertexBufferSRV;
@@ -4970,7 +4970,7 @@ struct TRHIResourceUpdateBatcher
 		}
 	}
 
-	void QueueUpdateRequest(FShaderResourceViewRHIParamRef SRV, FIndexBufferRHIParamRef IndexBuffer)
+	void QueueUpdateRequest(FRHIShaderResourceView* SRV, FIndexBufferRHIParamRef IndexBuffer)
 	{
 		// TODO
 	}
diff --git a/Engine/Source/Runtime/RHI/Public/RHIContext.h b/Engine/Source/Runtime/RHI/Public/RHIContext.h
--- a/Engine/Source/Runtime/RHI/Public/RHIContext.h
+++ b/Engine/Source/Runtime/RHI/Public/RHIContext.h
@@ -68,7 +68,7 @@ public:
 	* @param NumUAVs - number of UAVs to transition
 	* @param WriteComputeFence - Optional ComputeFence to write as part of this transition
 	*/
-	virtual void RHITransitionResources(EResourceTransitionAccess TransitionType, EResourceTransitionPipeline TransitionPipeline, FUnorderedAccessViewRHIParamRef* InUAVs, int32 NumUAVs, FRHIComputeFence* WriteComputeFence) = 0;
+	virtual void RHITransitionResources(EResourceTransitionAccess TransitionType, EResourceTransitionPipeline TransitionPipeline, FRHIUnorderedAccessView** InUAVs, int32 NumUAVs, FRHIComputeFence* WriteComputeFence) = 0;
 
 	/** Set the shader resource view of a surface.  This is used for binding TextureMS parameter types that need a multi sampled view. */
 	virtual void RHISetShaderTexture(FComputeShaderRHIParamRef PixelShader, uint32 TextureIndex, FTextureRHIParamRef NewTexture) = 0;
@@ -87,7 +87,7 @@ public:
 	* @param UAVIndex		The index of the UAVIndex.
 	* @param UAV			The new UAV.
 	*/
-	virtual void RHISetUAVParameter(FComputeShaderRHIParamRef ComputeShader, uint32 UAVIndex, FUnorderedAccessViewRHIParamRef UAV) = 0;
+	virtual void RHISetUAVParameter(FComputeShaderRHIParamRef ComputeShader, uint32 UAVIndex, FRHIUnorderedAccessView* UAV) = 0;
 
 	/**
 	* Sets a compute shader counted UAV parameter and initial count
@@ -96,9 +96,9 @@ public:
 	* @param UAV			The new UAV.
 	* @param InitialCount	The initial number of items in the UAV.
 	*/
-	virtual void RHISetUAVParameter(FComputeShaderRHIParamRef ComputeShader, uint32 UAVIndex, FUnorderedAccessViewRHIParamRef UAV, uint32 InitialCount) = 0;
+	virtual void RHISetUAVParameter(FComputeShaderRHIParamRef ComputeShader, uint32 UAVIndex, FRHIUnorderedAccessView* UAV, uint32 InitialCount) = 0;
 
-	virtual void RHISetShaderResourceViewParameter(FComputeShaderRHIParamRef ComputeShader, uint32 SamplerIndex, FShaderResourceViewRHIParamRef SRV) = 0;
+	virtual void RHISetShaderResourceViewParameter(FComputeShaderRHIParamRef ComputeShader, uint32 SamplerIndex, FRHIShaderResourceView* SRV) = 0;
 
 	virtual void RHISetShaderUniformBuffer(FComputeShaderRHIParamRef ComputeShader, uint32 BufferIndex, FUniformBufferRHIParamRef Buffer) = 0;
 
@@ -189,7 +189,7 @@ public:
 	virtual void RHISetMultipleViewports(uint32 Count, const FViewportBounds* Data) = 0;
 
 	/** Clears a UAV to the multi-component value provided. */
-	virtual void RHIClearTinyUAV(FUnorderedAccessViewRHIParamRef UnorderedAccessViewRHI, const uint32* Values) = 0;
+	virtual void RHIClearTinyUAV(FRHIUnorderedAccessView* UnorderedAccessViewRHI, const uint32* Values) = 0;
 
 	/**
 	* Resolves from one texture to another.
@@ -232,7 +232,7 @@ public:
 	* @param NumUAVs - number of UAVs to transition
 	* @param WriteComputeFence - Optional ComputeFence to write as part of this transition
 	*/
-	virtual void RHITransitionResources(EResourceTransitionAccess TransitionType, EResourceTransitionPipeline TransitionPipeline, FUnorderedAccessViewRHIParamRef* InUAVs, int32 NumUAVs, FRHIComputeFence* WriteComputeFence)
+	virtual void RHITransitionResources(EResourceTransitionAccess TransitionType, EResourceTransitionPipeline TransitionPipeline, FRHIUnorderedAccessView** InUAVs, int32 NumUAVs, FRHIComputeFence* WriteComputeFence)
 	{
 		if (WriteComputeFence)
 		{
@@ -240,7 +240,7 @@ public:
 		}
 	}
 
-	void RHITransitionResources(EResourceTransitionAccess TransitionType, EResourceTransitionPipeline TransitionPipeline, FUnorderedAccessViewRHIParamRef* InUAVs, int32 NumUAVs)
+	void RHITransitionResources(EResourceTransitionAccess TransitionType, EResourceTransitionPipeline TransitionPipeline, FRHIUnorderedAccessView** InUAVs, int32 NumUAVs)
 	{
 		RHITransitionResources(TransitionType, TransitionPipeline, InUAVs, NumUAVs, nullptr);
 	}
@@ -297,12 +297,12 @@ public:
 		/* empty default implementation */
 	}
 
-	virtual void RHIBeginUpdateMultiFrameResource(FUnorderedAccessViewRHIParamRef UAV)
+	virtual void RHIBeginUpdateMultiFrameResource(FRHIUnorderedAccessView* UAV)
 	{
 		/* empty default implementation */
 	}
 
-	virtual void RHIEndUpdateMultiFrameResource(FUnorderedAccessViewRHIParamRef UAV)
+	virtual void RHIEndUpdateMultiFrameResource(FRHIUnorderedAccessView* UAV)
 	{
 		/* empty default implementation */
 	}
@@ -400,7 +400,7 @@ public:
 	* @param UAVIndex		The index of the UAVIndex.
 	* @param UAV			The new UAV.
 	*/
-	virtual void RHISetUAVParameter(FComputeShaderRHIParamRef ComputeShader, uint32 UAVIndex, FUnorderedAccessViewRHIParamRef UAV) = 0;
+	virtual void RHISetUAVParameter(FComputeShaderRHIParamRef ComputeShader, uint32 UAVIndex, FRHIUnorderedAccessView* UAV) = 0;
 
 	/**
 	* Sets a compute shader counted UAV parameter and initial count
@@ -409,19 +409,19 @@ public:
 	* @param UAV			The new UAV.
 	* @param InitialCount	The initial number of items in the UAV.
 	*/
-	virtual void RHISetUAVParameter(FComputeShaderRHIParamRef ComputeShader, uint32 UAVIndex, FUnorderedAccessViewRHIParamRef UAV, uint32 InitialCount) = 0;
+	virtual void RHISetUAVParameter(FComputeShaderRHIParamRef ComputeShader, uint32 UAVIndex, FRHIUnorderedAccessView* UAV, uint32 InitialCount) = 0;
 
-	virtual void RHISetShaderResourceViewParameter(FPixelShaderRHIParamRef PixelShader, uint32 SamplerIndex, FShaderResourceViewRHIParamRef SRV) = 0;
+	virtual void RHISetShaderResourceViewParameter(FPixelShaderRHIParamRef PixelShader, uint32 SamplerIndex, FRHIShaderResourceView* SRV) = 0;
 
-	virtual void RHISetShaderResourceViewParameter(FVertexShaderRHIParamRef VertexShader, uint32 SamplerIndex, FShaderResourceViewRHIParamRef SRV) = 0;
+	virtual void RHISetShaderResourceViewParameter(FVertexShaderRHIParamRef VertexShader, uint32 SamplerIndex, FRHIShaderResourceView* SRV) = 0;
 
-	virtual void RHISetShaderResourceViewParameter(FComputeShaderRHIParamRef ComputeShader, uint32 SamplerIndex, FShaderResourceViewRHIParamRef SRV) = 0;
+	virtual void RHISetShaderResourceViewParameter(FComputeShaderRHIParamRef ComputeShader, uint32 SamplerIndex, FRHIShaderResourceView* SRV) = 0;
 
-	virtual void RHISetShaderResourceViewParameter(FHullShaderRHIParamRef HullShader, uint32 SamplerIndex, FShaderResourceViewRHIParamRef SRV) = 0;
+	virtual void RHISetShaderResourceViewParameter(FHullShaderRHIParamRef HullShader, uint32 SamplerIndex, FRHIShaderResourceView* SRV) = 0;
 
-	virtual void RHISetShaderResourceViewParameter(FDomainShaderRHIParamRef DomainShader, uint32 SamplerIndex, FShaderResourceViewRHIParamRef SRV) = 0;
+	virtual void RHISetShaderResourceViewParameter(FDomainShaderRHIParamRef DomainShader, uint32 SamplerIndex, FRHIShaderResourceView* SRV) = 0;
 
-	virtual void RHISetShaderResourceViewParameter(FGeometryShaderRHIParamRef GeometryShader, uint32 SamplerIndex, FShaderResourceViewRHIParamRef SRV) = 0;
+	virtual void RHISetShaderResourceViewParameter(FGeometryShaderRHIParamRef GeometryShader, uint32 SamplerIndex, FRHIShaderResourceView* SRV) = 0;
 
 	virtual void RHISetShaderUniformBuffer(FVertexShaderRHIParamRef VertexShader, uint32 BufferIndex, FUniformBufferRHIParamRef Buffer) = 0;
 
@@ -451,7 +451,7 @@ public:
 
 	virtual void RHISetBlendFactor(const FLinearColor& BlendFactor) {}
 
-	virtual void RHISetRenderTargets(uint32 NumSimultaneousRenderTargets, const FRHIRenderTargetView* NewRenderTargets, const FRHIDepthRenderTargetView* NewDepthStencilTarget, uint32 NumUAVs, const FUnorderedAccessViewRHIParamRef* UAVs) = 0;
+	virtual void RHISetRenderTargets(uint32 NumSimultaneousRenderTargets, const FRHIRenderTargetView* NewRenderTargets, const FRHIDepthRenderTargetView* NewDepthStencilTarget, uint32 NumUAVs, FRHIUnorderedAccessView* const* UAVs) = 0;
 
 	virtual void RHISetRenderTargetsAndClear(const FRHISetRenderTargetsInfo& RenderTargetsInfo) = 0;
 
@@ -600,16 +600,16 @@ public:
 	}
 
 	virtual void RHIRayTraceOcclusion(FRayTracingSceneRHIParamRef Scene,
-		FShaderResourceViewRHIParamRef Rays,
-		FUnorderedAccessViewRHIParamRef Output,
+		FRHIShaderResourceView* Rays,
+		FRHIUnorderedAccessView* Output,
 		uint32 NumRays)
 	{
 		checkNoEntry();
 	}
 
 	virtual void RHIRayTraceIntersection(FRayTracingSceneRHIParamRef Scene,
-		FShaderResourceViewRHIParamRef Rays,
-		FUnorderedAccessViewRHIParamRef Output,
+		FRHIShaderResourceView* Rays,
+		FRHIUnorderedAccessView* Output,
 		uint32 NumRays)
 	{
 		checkNoEntry();
diff --git a/Engine/Source/Runtime/RHI/Public/RHIResources.h b/Engine/Source/Runtime/RHI/Public/RHIResources.h
--- a/Engine/Source/Runtime/RHI/Public/RHIResources.h
+++ b/Engine/Source/Runtime/RHI/Public/RHIResources.h
@@ -1189,9 +1189,11 @@ typedef TRefCountPtr<FRHIGPUFence>	FGPUFenceRHIRef;
 typedef FRHIViewport*              FViewportRHIParamRef;
 typedef TRefCountPtr<FRHIViewport> FViewportRHIRef;
 
+UE_DEPRECATED(4.23, "FUnorderedAccessViewRHIParamRef typedef is deprecated; please use FRHIUnorderedAccessView* directly instead.")
 typedef FRHIUnorderedAccessView*              FUnorderedAccessViewRHIParamRef;
 typedef TRefCountPtr<FRHIUnorderedAccessView> FUnorderedAccessViewRHIRef;
 
+UE_DEPRECATED(4.23, "FShaderResourceViewRHIParamRef typedef is deprecated; please use FRHIShaderResourceView* directly instead.")
 typedef FRHIShaderResourceView*              FShaderResourceViewRHIParamRef;
 typedef TRefCountPtr<FRHIShaderResourceView> FShaderResourceViewRHIRef;
 
@@ -1218,7 +1220,7 @@ typedef TRefCountPtr<FRHIRayTracingGeometry>     FRayTracingGeometryRHIRef;
 class FRHIRayTracingScene : public FRHIResource
 {
 public:
-	FShaderResourceViewRHIParamRef GetShaderResourceView() { return ShaderResourceView; }
+	FRHIShaderResourceView* GetShaderResourceView() { return ShaderResourceView; }
 protected:
 	FShaderResourceViewRHIRef ShaderResourceView;
 };
@@ -2679,7 +2681,7 @@ struct FRHIRenderPassInfo
 		FMemory::Memzero(&ColorRenderTargets[1], sizeof(FColorEntry) * (MaxSimultaneousRenderTargets - 1));
 	}
 
-	explicit FRHIRenderPassInfo(int32 InNumUAVs, FUnorderedAccessViewRHIParamRef* InUAVs)
+	explicit FRHIRenderPassInfo(int32 InNumUAVs, FRHIUnorderedAccessView** InUAVs)
 	{
 		if (InNumUAVs > MaxSimultaneousUAVs)
 		{
diff --git a/Engine/Source/Runtime/RHI/Public/RHIValidationContext.h b/Engine/Source/Runtime/RHI/Public/RHIValidationContext.h
--- a/Engine/Source/Runtime/RHI/Public/RHIValidationContext.h
+++ b/Engine/Source/Runtime/RHI/Public/RHIValidationContext.h
@@ -65,7 +65,7 @@ public:
 	* @param NumUAVs - number of UAVs to transition
 	* @param WriteComputeFence - Optional ComputeFence to write as part of this transition
 	*/
-	virtual void RHITransitionResources(EResourceTransitionAccess TransitionType, EResourceTransitionPipeline TransitionPipeline, FUnorderedAccessViewRHIParamRef* InUAVs, int32 NumUAVs, FRHIComputeFence* WriteComputeFence) override final
+	virtual void RHITransitionResources(EResourceTransitionAccess TransitionType, EResourceTransitionPipeline TransitionPipeline, FRHIUnorderedAccessView** InUAVs, int32 NumUAVs, FRHIComputeFence* WriteComputeFence) override final
 	{
 		RHIContext->RHITransitionResources(TransitionType, TransitionPipeline, InUAVs, NumUAVs, WriteComputeFence);
 	}
@@ -100,7 +100,7 @@ public:
 	* @param UAVIndex		The index of the UAVIndex.
 	* @param UAV			The new UAV.
 	*/
-	virtual void RHISetUAVParameter(FComputeShaderRHIParamRef Shader, uint32 UAVIndex, FUnorderedAccessViewRHIParamRef UAV) override final
+	virtual void RHISetUAVParameter(FComputeShaderRHIParamRef Shader, uint32 UAVIndex, FRHIUnorderedAccessView* UAV) override final
 	{
 		checkf(State.bComputeShaderSet, TEXT("A Compute shader has to be set to set resources into a shader!"));
 		RHIContext->RHISetUAVParameter(Shader, UAVIndex, UAV);
@@ -113,13 +113,13 @@ public:
 	* @param UAV			The new UAV.
 	* @param InitialCount	The initial number of items in the UAV.
 	*/
-	virtual void RHISetUAVParameter(FComputeShaderRHIParamRef Shader, uint32 UAVIndex, FUnorderedAccessViewRHIParamRef UAV, uint32 InitialCount) override final
+	virtual void RHISetUAVParameter(FComputeShaderRHIParamRef Shader, uint32 UAVIndex, FRHIUnorderedAccessView* UAV, uint32 InitialCount) override final
 	{
 		checkf(State.bComputeShaderSet, TEXT("A Compute shader has to be set to set resources into a shader!"));
 		RHIContext->RHISetUAVParameter(Shader, UAVIndex, UAV, InitialCount);
 	}
 
-	virtual void RHISetShaderResourceViewParameter(FComputeShaderRHIParamRef Shader, uint32 SamplerIndex, FShaderResourceViewRHIParamRef SRV) override final
+	virtual void RHISetShaderResourceViewParameter(FComputeShaderRHIParamRef Shader, uint32 SamplerIndex, FRHIShaderResourceView* SRV) override final
 	{
 		checkf(State.bComputeShaderSet, TEXT("A Compute shader has to be set to set resources into a shader!"));
 		RHIContext->RHISetShaderResourceViewParameter(Shader, SamplerIndex, SRV);
@@ -232,7 +232,7 @@ public:
 	}
 
 	/** Clears a UAV to the multi-component value provided. */
-	virtual void RHIClearTinyUAV(FUnorderedAccessViewRHIParamRef UnorderedAccessViewRHI, const uint32* Values) override final
+	virtual void RHIClearTinyUAV(FRHIUnorderedAccessView* UnorderedAccessViewRHI, const uint32* Values) override final
 	{
 		RHIContext->RHIClearTinyUAV(UnorderedAccessViewRHI, Values);
 	}
@@ -273,7 +273,7 @@ public:
 	* @param NumUAVs - number of UAVs to transition
 	* @param WriteComputeFence - Optional ComputeFence to write as part of this transition
 	*/
-	virtual void RHITransitionResources(EResourceTransitionAccess TransitionType, EResourceTransitionPipeline TransitionPipeline, FUnorderedAccessViewRHIParamRef* InUAVs, int32 NumUAVs, FRHIComputeFence* WriteComputeFence) override final
+	virtual void RHITransitionResources(EResourceTransitionAccess TransitionType, EResourceTransitionPipeline TransitionPipeline, FRHIUnorderedAccessView** InUAVs, int32 NumUAVs, FRHIComputeFence* WriteComputeFence) override final
 	{
 		RHIContext->RHITransitionResources(TransitionType, TransitionPipeline, InUAVs, NumUAVs, WriteComputeFence);
 	}
@@ -364,12 +364,12 @@ public:
 		RHIContext->RHIEndUpdateMultiFrameResource(Texture);
 	}
 
-	virtual void RHIBeginUpdateMultiFrameResource(FUnorderedAccessViewRHIParamRef UAV) override final
+	virtual void RHIBeginUpdateMultiFrameResource(FRHIUnorderedAccessView* UAV) override final
 	{
 		RHIContext->RHIBeginUpdateMultiFrameResource(UAV);
 	}
 
-	virtual void RHIEndUpdateMultiFrameResource(FUnorderedAccessViewRHIParamRef UAV) override final
+	virtual void RHIEndUpdateMultiFrameResource(FRHIUnorderedAccessView* UAV) override final
 	{
 		RHIContext->RHIEndUpdateMultiFrameResource(UAV);
 	}
@@ -535,7 +535,7 @@ public:
 	* @param UAVIndex		The index of the UAVIndex.
 	* @param UAV			The new UAV.
 	*/
-	virtual void RHISetUAVParameter(FComputeShaderRHIParamRef Shader, uint32 UAVIndex, FUnorderedAccessViewRHIParamRef UAV) override final
+	virtual void RHISetUAVParameter(FComputeShaderRHIParamRef Shader, uint32 UAVIndex, FRHIUnorderedAccessView* UAV) override final
 	{
 		checkf(State.bComputeShaderSet, TEXT("A Compute shader has to be set to set resources into a shader!"));
 		RHIContext->RHISetUAVParameter(Shader, UAVIndex, UAV);
@@ -548,43 +548,43 @@ public:
 	* @param UAV			The new UAV.
 	* @param InitialCount	The initial number of items in the UAV.
 	*/
-	virtual void RHISetUAVParameter(FComputeShaderRHIParamRef Shader, uint32 UAVIndex, FUnorderedAccessViewRHIParamRef UAV, uint32 InitialCount) override final
+	virtual void RHISetUAVParameter(FComputeShaderRHIParamRef Shader, uint32 UAVIndex, FRHIUnorderedAccessView* UAV, uint32 InitialCount) override final
 	{
 		checkf(State.bComputeShaderSet, TEXT("A Compute shader has to be set to set resources into a shader!"));
 		RHIContext->RHISetUAVParameter(Shader, UAVIndex, UAV, InitialCount);
 	}
 
-	virtual void RHISetShaderResourceViewParameter(FPixelShaderRHIParamRef Shader, uint32 SamplerIndex, FShaderResourceViewRHIParamRef SRV) override final
+	virtual void RHISetShaderResourceViewParameter(FPixelShaderRHIParamRef Shader, uint32 SamplerIndex, FRHIShaderResourceView* SRV) override final
 	{
 		checkf(State.bGfxPSOSet, TEXT("A Graphics PSO has to be set to set resources into a shader!"));
 		RHIContext->RHISetShaderResourceViewParameter(Shader, SamplerIndex, SRV);
 	}
 
-	virtual void RHISetShaderResourceViewParameter(FVertexShaderRHIParamRef Shader, uint32 SamplerIndex, FShaderResourceViewRHIParamRef SRV) override final
+	virtual void RHISetShaderResourceViewParameter(FVertexShaderRHIParamRef Shader, uint32 SamplerIndex, FRHIShaderResourceView* SRV) override final
 	{
 		checkf(State.bGfxPSOSet, TEXT("A Graphics PSO has to be set to set resources into a shader!"));
 		RHIContext->RHISetShaderResourceViewParameter(Shader, SamplerIndex, SRV);
 	}
 
-	virtual void RHISetShaderResourceViewParameter(FComputeShaderRHIParamRef Shader, uint32 SamplerIndex, FShaderResourceViewRHIParamRef SRV) override final
+	virtual void RHISetShaderResourceViewParameter(FComputeShaderRHIParamRef Shader, uint32 SamplerIndex, FRHIShaderResourceView* SRV) override final
 	{
 		checkf(State.bComputeShaderSet, TEXT("A Compute shader has to be set to set resources into a shader!"));
 		RHIContext->RHISetShaderResourceViewParameter(Shader, SamplerIndex, SRV);
 	}
 
-	virtual void RHISetShaderResourceViewParameter(FHullShaderRHIParamRef Shader, uint32 SamplerIndex, FShaderResourceViewRHIParamRef SRV) override final
+	virtual void RHISetShaderResourceViewParameter(FHullShaderRHIParamRef Shader, uint32 SamplerIndex, FRHIShaderResourceView* SRV) override final
 	{
 		checkf(State.bGfxPSOSet, TEXT("A Graphics PSO has to be set to set resources into a shader!"));
 		RHIContext->RHISetShaderResourceViewParameter(Shader, SamplerIndex, SRV);
 	}
 
-	virtual void RHISetShaderResourceViewParameter(FDomainShaderRHIParamRef Shader, uint32 SamplerIndex, FShaderResourceViewRHIParamRef SRV) override final
+	virtual void RHISetShaderResourceViewParameter(FDomainShaderRHIParamRef Shader, uint32 SamplerIndex, FRHIShaderResourceView* SRV) override final
 	{
 		checkf(State.bGfxPSOSet, TEXT("A Graphics PSO has to be set to set resources into a shader!"));
 		RHIContext->RHISetShaderResourceViewParameter(Shader, SamplerIndex, SRV);
 	}
 
-	virtual void RHISetShaderResourceViewParameter(FGeometryShaderRHIParamRef Shader, uint32 SamplerIndex, FShaderResourceViewRHIParamRef SRV) override final
+	virtual void RHISetShaderResourceViewParameter(FGeometryShaderRHIParamRef Shader, uint32 SamplerIndex, FRHIShaderResourceView* SRV) override final
 	{
 		checkf(State.bGfxPSOSet, TEXT("A Graphics PSO has to be set to set resources into a shader!"));
 		RHIContext->RHISetShaderResourceViewParameter(Shader, SamplerIndex, SRV);
@@ -674,7 +674,7 @@ public:
 		RHIContext->RHISetBlendFactor(BlendFactor);
 	}
 
-	virtual void RHISetRenderTargets(uint32 NumSimultaneousRenderTargets, const FRHIRenderTargetView* NewRenderTargets, const FRHIDepthRenderTargetView* NewDepthStencilTarget, uint32 NumUAVs, const FUnorderedAccessViewRHIParamRef* UAVs) override final
+	virtual void RHISetRenderTargets(uint32 NumSimultaneousRenderTargets, const FRHIRenderTargetView* NewRenderTargets, const FRHIDepthRenderTargetView* NewDepthStencilTarget, uint32 NumUAVs, FRHIUnorderedAccessView* const* UAVs) override final
 	{
 		RHIContext->RHISetRenderTargets(NumSimultaneousRenderTargets, NewRenderTargets, NewDepthStencilTarget, NumUAVs, UAVs);
 	}
@@ -825,16 +825,16 @@ public:
 	}
 
 	virtual void RHIRayTraceOcclusion(FRayTracingSceneRHIParamRef Scene,
-		FShaderResourceViewRHIParamRef Rays,
-		FUnorderedAccessViewRHIParamRef Output,
+		FRHIShaderResourceView* Rays,
+		FRHIUnorderedAccessView* Output,
 		uint32 NumRays) override final
 	{
 		RHIContext->RHIRayTraceOcclusion(Scene, Rays, Output, NumRays);
 	}
 
 	virtual void RHIRayTraceIntersection(FRayTracingSceneRHIParamRef Scene,
-		FShaderResourceViewRHIParamRef Rays,
-		FUnorderedAccessViewRHIParamRef Output,
+		FRHIShaderResourceView* Rays,
+		FRHIUnorderedAccessView* Output,
 		uint32 NumRays) override final
 	{
 		RHIContext->RHIRayTraceIntersection(Scene, Rays, Output, NumRays);
diff --git a/Engine/Source/Runtime/RenderCore/Private/RenderGraphBuilder.cpp b/Engine/Source/Runtime/RenderCore/Private/RenderGraphBuilder.cpp
--- a/Engine/Source/Runtime/RenderCore/Private/RenderGraphBuilder.cpp
+++ b/Engine/Source/Runtime/RenderCore/Private/RenderGraphBuilder.cpp
@@ -802,7 +802,7 @@ void FRDGBuilder::TransitionTexture(FRDGTexture* Texture, ERDGPassAccess PassAcc
 }
 
 void FRDGBuilder::TransitionUAV(
-	FUnorderedAccessViewRHIParamRef UAV,
+	FRHIUnorderedAccessView* UAV,
 	FRDGTrackedResource* UnderlyingResource,
 	ERDGPassAccess PassAccessAfter,
 	ERDGPassPipeline PassPipelineAfter) const
@@ -961,7 +961,7 @@ void FRDGBuilder::PrepareResourcesForExecute(const FRDGPass* Pass, struct FRHIRe
 			if (FRDGTextureUAVRef UAV = Parameter.GetAsTextureUAV())
 			{
 				FRDGTextureRef Texture = UAV->Desc.Texture;
-				FUnorderedAccessViewRHIParamRef UAVRHI = UAV->GetRHI();
+				FRHIUnorderedAccessView* UAVRHI = UAV->GetRHI();
 
 				if (PassPipeline == ERDGPassPipeline::Graphics)
 				{
@@ -994,7 +994,7 @@ void FRDGBuilder::PrepareResourcesForExecute(const FRDGPass* Pass, struct FRHIRe
 				// TODO(RDG): super hacky, find the UAV and transition it. Hopefully there is one...
 				check(Buffer->PooledBuffer);
 				check(Buffer->PooledBuffer->UAVs.Num() == 1);
-				FUnorderedAccessViewRHIParamRef UAV = Buffer->PooledBuffer->UAVs.CreateIterator().Value();
+				FRHIUnorderedAccessView* UAV = Buffer->PooledBuffer->UAVs.CreateIterator().Value();
 
 				TransitionUAV(UAV, Buffer, ERDGPassAccess::Write, PassPipeline);
 			}
@@ -1017,7 +1017,7 @@ void FRDGBuilder::PrepareResourcesForExecute(const FRDGPass* Pass, struct FRHIRe
 				// TODO(RDG): super hacky, find the UAV and transition it. Hopefully there is one...
 				check(Buffer->PooledBuffer);
 				check(Buffer->PooledBuffer->UAVs.Num() == 1);
-				FUnorderedAccessViewRHIParamRef UAV = Buffer->PooledBuffer->UAVs.CreateIterator().Value();
+				FRHIUnorderedAccessView* UAV = Buffer->PooledBuffer->UAVs.CreateIterator().Value();
 
 				AllocateRHIBufferSRVIfNeeded(SRV);
 				TransitionUAV(UAV, Buffer, ERDGPassAccess::Read, PassPipeline);
@@ -1029,7 +1029,7 @@ void FRDGBuilder::PrepareResourcesForExecute(const FRDGPass* Pass, struct FRHIRe
 			if (FRDGBufferUAVRef UAV = Parameter.GetAsBufferUAV())
 			{
 				FRDGBufferRef Buffer = UAV->Desc.Buffer;
-				FUnorderedAccessViewRHIParamRef UAVRHI = UAV->GetRHI();
+				FRHIUnorderedAccessView* UAVRHI = UAV->GetRHI();
 
 				if (PassPipeline == ERDGPassPipeline::Graphics)
 				{
diff --git a/Engine/Source/Runtime/RenderCore/Private/ShaderParameterStruct.cpp b/Engine/Source/Runtime/RenderCore/Private/ShaderParameterStruct.cpp
--- a/Engine/Source/Runtime/RenderCore/Private/ShaderParameterStruct.cpp
+++ b/Engine/Source/Runtime/RenderCore/Private/ShaderParameterStruct.cpp
@@ -410,7 +410,7 @@ void ValidateShaderParameters(const FShader* Shader, const FShaderParametersMeta
 	// SRVs
 	for (const FShaderParameterBindings::FResourceParameter& ParameterBinding : Bindings.SRVs)
 	{
-		auto ShaderParameterRef = *reinterpret_cast<const FShaderResourceViewRHIParamRef*>(Base + ParameterBinding.ByteOffset);
+		FRHIShaderResourceView* ShaderParameterRef = *(FRHIShaderResourceView**)(Base + ParameterBinding.ByteOffset);
 		if (!ShaderParameterRef)
 		{
 			EmitNullShaderParameterFatalError(Shader, ParametersMetadata, ParameterBinding.ByteOffset);
@@ -420,7 +420,7 @@ void ValidateShaderParameters(const FShader* Shader, const FShaderParametersMeta
 	// Samplers
 	for (const FShaderParameterBindings::FResourceParameter& ParameterBinding : Bindings.Samplers)
 	{
-		FRHISamplerState* ShaderParameterRef = (FRHISamplerState*)(Base + ParameterBinding.ByteOffset);
+		FRHISamplerState* ShaderParameterRef = *(FRHISamplerState**)(Base + ParameterBinding.ByteOffset);
 		if (!ShaderParameterRef)
 		{
 			EmitNullShaderParameterFatalError(Shader, ParametersMetadata, ParameterBinding.ByteOffset);
diff --git a/Engine/Source/Runtime/RenderCore/Public/RenderGraphBuilder.h b/Engine/Source/Runtime/RenderCore/Public/RenderGraphBuilder.h
--- a/Engine/Source/Runtime/RenderCore/Public/RenderGraphBuilder.h
+++ b/Engine/Source/Runtime/RenderCore/Public/RenderGraphBuilder.h
@@ -203,7 +203,7 @@ private:
 		ERDGPassPipeline PassPipeline) const;
 
 	void TransitionUAV(
-		FUnorderedAccessViewRHIParamRef UAV,
+		FRHIUnorderedAccessView* UAV,
 		FRDGTrackedResource* UnderlyingResource,
 		ERDGPassAccess PassAccess,
 		ERDGPassPipeline PassPipeline) const;
diff --git a/Engine/Source/Runtime/RenderCore/Public/RenderGraphResources.h b/Engine/Source/Runtime/RenderCore/Public/RenderGraphResources.h
--- a/Engine/Source/Runtime/RenderCore/Public/RenderGraphResources.h
+++ b/Engine/Source/Runtime/RenderCore/Public/RenderGraphResources.h
@@ -240,9 +240,9 @@ class FRDGShaderResourceView
 {
 public:
 	/** Returns the allocated RHI SRV. */
-	FShaderResourceViewRHIParamRef GetRHI() const
+	FRHIShaderResourceView* GetRHI() const
 	{
-		return static_cast<FShaderResourceViewRHIParamRef>(FRDGResource::GetRHI());
+		return static_cast<FRHIShaderResourceView*>(FRDGResource::GetRHI());
 	}
 
 protected:
@@ -251,9 +251,9 @@ protected:
 	{}
 
 	/** Returns the allocated RHI SRV without access checks. */
-	FShaderResourceViewRHIParamRef GetRHIUnchecked() const
+	FRHIShaderResourceView* GetRHIUnchecked() const
 	{
-		return static_cast<FShaderResourceViewRHIParamRef>(FRDGResource::GetRHIUnchecked());
+		return static_cast<FRHIShaderResourceView*>(FRDGResource::GetRHIUnchecked());
 	}
 };
 
@@ -263,9 +263,9 @@ class FRDGUnorderedAccessView
 {
 public:
 	/** Returns the allocated RHI UAV. */
-	FUnorderedAccessViewRHIParamRef GetRHI() const
+	FRHIUnorderedAccessView* GetRHI() const
 	{
-		return static_cast<FUnorderedAccessViewRHIParamRef>(FRDGResource::GetRHI());
+		return static_cast<FRHIUnorderedAccessView*>(FRDGResource::GetRHI());
 	}
 
 protected:
@@ -274,9 +274,9 @@ protected:
 	{}
 
 	/** Returns the allocated RHI UAV without access checks. */
-	FUnorderedAccessViewRHIParamRef GetRHIUnchecked() const
+	FRHIUnorderedAccessView* GetRHIUnchecked() const
 	{
-		return static_cast<FUnorderedAccessViewRHIParamRef>(FRDGResource::GetRHIUnchecked());
+		return static_cast<FRHIUnorderedAccessView*>(FRDGResource::GetRHIUnchecked());
 	}
 };
 
diff --git a/Engine/Source/Runtime/RenderCore/Public/ShaderParameterMacros.h b/Engine/Source/Runtime/RenderCore/Public/ShaderParameterMacros.h
--- a/Engine/Source/Runtime/RenderCore/Public/ShaderParameterMacros.h
+++ b/Engine/Source/Runtime/RenderCore/Public/ShaderParameterMacros.h
@@ -841,10 +841,10 @@ extern RENDERCORE_API FShaderParametersMetadata* FindUniformBufferStructByFName(
  *	SHADER_PARAMETER_SRV_ARRAY(Texture2D, MyArrayOfSRVs, [8])
  */
 #define SHADER_PARAMETER_SRV(ShaderType,MemberName) \
-	INTERNAL_SHADER_PARAMETER_EXPLICIT(UBMT_SRV, TShaderResourceParameterTypeInfo<FShaderResourceViewRHIParamRef>, FShaderResourceViewRHIParamRef,MemberName,, = nullptr,EShaderPrecisionModifier::Float,TEXT(#ShaderType),false)
+	INTERNAL_SHADER_PARAMETER_EXPLICIT(UBMT_SRV, TShaderResourceParameterTypeInfo<FRHIShaderResourceView*>, FRHIShaderResourceView*,MemberName,, = nullptr,EShaderPrecisionModifier::Float,TEXT(#ShaderType),false)
 
 #define SHADER_PARAMETER_SRV_ARRAY(ShaderType,MemberName, ArrayDecl) \
-	INTERNAL_SHADER_PARAMETER_EXPLICIT(UBMT_SRV, TShaderResourceParameterTypeInfo<FShaderResourceViewRHIParamRef ArrayDecl>, FShaderResourceViewRHIParamRef,MemberName,ArrayDecl,,EShaderPrecisionModifier::Float,TEXT(#ShaderType),false)
+	INTERNAL_SHADER_PARAMETER_EXPLICIT(UBMT_SRV, TShaderResourceParameterTypeInfo<FRHIShaderResourceView* ArrayDecl>, FRHIShaderResourceView*,MemberName,ArrayDecl,,EShaderPrecisionModifier::Float,TEXT(#ShaderType),false)
 
 /** Adds an unordered access view.
  *
@@ -852,7 +852,7 @@ extern RENDERCORE_API FShaderParametersMetadata* FindUniformBufferStructByFName(
  *	SHADER_PARAMETER_UAV(Texture2D, MyUAV)
  */
 #define SHADER_PARAMETER_UAV(ShaderType,MemberName) \
-	INTERNAL_SHADER_PARAMETER_EXPLICIT(UBMT_UAV, TShaderResourceParameterTypeInfo<FUnorderedAccessViewRHIParamRef>, FUnorderedAccessViewRHIParamRef,MemberName,, = nullptr,EShaderPrecisionModifier::Float,TEXT(#ShaderType),false)
+	INTERNAL_SHADER_PARAMETER_EXPLICIT(UBMT_UAV, TShaderResourceParameterTypeInfo<FRHIUnorderedAccessView*>, FRHIUnorderedAccessView*,MemberName,, = nullptr,EShaderPrecisionModifier::Float,TEXT(#ShaderType),false)
 
 /** Adds a sampler.
  *
diff --git a/Engine/Source/Runtime/RenderCore/Public/ShaderParameterStruct.h b/Engine/Source/Runtime/RenderCore/Public/ShaderParameterStruct.h
--- a/Engine/Source/Runtime/RenderCore/Public/ShaderParameterStruct.h
+++ b/Engine/Source/Runtime/RenderCore/Public/ShaderParameterStruct.h
@@ -103,7 +103,7 @@ inline void SetShaderUAVs(TRHICmdList& RHICmdList, const TShaderClass* Shader, F
 	// UAVs
 	for (const FShaderParameterBindings::FResourceParameter& ParameterBinding : Bindings.UAVs)
 	{
-		auto ShaderParameterRef = *reinterpret_cast<const FUnorderedAccessViewRHIParamRef*>(Base + ParameterBinding.ByteOffset);
+		FRHIUnorderedAccessView* ShaderParameterRef = *(FRHIUnorderedAccessView**)(Base + ParameterBinding.ByteOffset);
 
 		if (DO_CHECK && !ShaderParameterRef)
 		{
@@ -192,7 +192,7 @@ inline void SetShaderParameters(TRHICmdList& RHICmdList, const TShaderClass* Sha
 	// SRVs
 	for (const FShaderParameterBindings::FResourceParameter& ParameterBinding : Bindings.SRVs)
 	{
-		auto ShaderParameterRef = *reinterpret_cast<const FShaderResourceViewRHIParamRef*>(Base + ParameterBinding.ByteOffset);
+		FRHIShaderResourceView* ShaderParameterRef = *(FRHIShaderResourceView**)(Base + ParameterBinding.ByteOffset);
 
 		if (DO_CHECK && !ShaderParameterRef)
 		{
@@ -304,7 +304,7 @@ void SetShaderParameters(FRayTracingShaderBindingsWriter& RTBindingsWriter, cons
 	// SRVs
 	for (const FShaderParameterBindings::FResourceParameter& ParameterBinding : Bindings.SRVs)
 	{
-		auto ShaderParameterRef = *reinterpret_cast<const FShaderResourceViewRHIParamRef*>(Base + ParameterBinding.ByteOffset);
+		FRHIShaderResourceView* ShaderParameterRef = *(FRHIShaderResourceView**)(Base + ParameterBinding.ByteOffset);
 
 		if (DO_CHECK && !ShaderParameterRef)
 		{
@@ -317,7 +317,7 @@ void SetShaderParameters(FRayTracingShaderBindingsWriter& RTBindingsWriter, cons
 	// UAVs
 	for (const FShaderParameterBindings::FResourceParameter& ParameterBinding : Bindings.UAVs)
 	{
-		auto ShaderParameterRef = *reinterpret_cast<const FUnorderedAccessViewRHIParamRef*>(Base + ParameterBinding.ByteOffset);
+		FRHIUnorderedAccessView* ShaderParameterRef = *(FRHIUnorderedAccessView**)(Base + ParameterBinding.ByteOffset);
 
 		if (DO_CHECK && !ShaderParameterRef)
 		{
diff --git a/Engine/Source/Runtime/RenderCore/Public/ShaderParameterUtils.h b/Engine/Source/Runtime/RenderCore/Public/ShaderParameterUtils.h
--- a/Engine/Source/Runtime/RenderCore/Public/ShaderParameterUtils.h
+++ b/Engine/Source/Runtime/RenderCore/Public/ShaderParameterUtils.h
@@ -316,7 +316,7 @@ FORCEINLINE void SetSRVParameter(
 	TRHICmdList& RHICmdList,
 	ShaderTypeRHIParamRef Shader,
 	const FShaderResourceParameter& Parameter,
-	FShaderResourceViewRHIParamRef NewShaderResourceViewRHI
+	FRHIShaderResourceView* NewShaderResourceViewRHI
 	)
 {
 	if(Parameter.IsBound())
@@ -337,7 +337,7 @@ FORCEINLINE void SetUAVParameter(
 TRHICmdList& RHICmdList,
 	FComputeShaderRHIParamRef ComputeShader,
 	const FShaderResourceParameter& Parameter,
-	FUnorderedAccessViewRHIParamRef NewUnorderedAccessViewRHI
+	FRHIUnorderedAccessView* NewUnorderedAccessViewRHI
 	)
 {
 	if(Parameter.IsBound())
@@ -352,37 +352,37 @@ TRHICmdList& RHICmdList,
 
 
 template<typename TRHICmdList>
-inline bool SetUAVParameterIfCS(TRHICmdList& RHICmdList, const FVertexShaderRHIParamRef Shader, const FShaderResourceParameter& UAVParameter, FUnorderedAccessViewRHIParamRef UAV)
+inline bool SetUAVParameterIfCS(TRHICmdList& RHICmdList, const FVertexShaderRHIParamRef Shader, const FShaderResourceParameter& UAVParameter, FRHIUnorderedAccessView* UAV)
 {
 	return false;
 }
 
 template<typename TRHICmdList>
-inline bool SetUAVParameterIfCS(TRHICmdList& RHICmdList, const FPixelShaderRHIParamRef Shader, const FShaderResourceParameter& UAVParameter, FUnorderedAccessViewRHIParamRef UAV)
+inline bool SetUAVParameterIfCS(TRHICmdList& RHICmdList, const FPixelShaderRHIParamRef Shader, const FShaderResourceParameter& UAVParameter, FRHIUnorderedAccessView* UAV)
 {
 	return false;
 }
 
 template<typename TRHICmdList>
-inline bool SetUAVParameterIfCS(TRHICmdList& RHICmdList, const FHullShaderRHIParamRef Shader, const FShaderResourceParameter& UAVParameter, FUnorderedAccessViewRHIParamRef UAV)
+inline bool SetUAVParameterIfCS(TRHICmdList& RHICmdList, const FHullShaderRHIParamRef Shader, const FShaderResourceParameter& UAVParameter, FRHIUnorderedAccessView* UAV)
 {
 	return false;
 }
 
 template<typename TRHICmdList>
-inline bool SetUAVParameterIfCS(TRHICmdList& RHICmdList, const FDomainShaderRHIParamRef Shader, const FShaderResourceParameter& UAVParameter, FUnorderedAccessViewRHIParamRef UAV)
+inline bool SetUAVParameterIfCS(TRHICmdList& RHICmdList, const FDomainShaderRHIParamRef Shader, const FShaderResourceParameter& UAVParameter, FRHIUnorderedAccessView* UAV)
 {
 	return false;
 }
 
 template<typename TRHICmdList>
-inline bool SetUAVParameterIfCS(TRHICmdList& RHICmdList, const FGeometryShaderRHIParamRef Shader, const FShaderResourceParameter& UAVParameter, FUnorderedAccessViewRHIParamRef UAV)
+inline bool SetUAVParameterIfCS(TRHICmdList& RHICmdList, const FGeometryShaderRHIParamRef Shader, const FShaderResourceParameter& UAVParameter, FRHIUnorderedAccessView* UAV)
 {
 	return false;
 }
 
 template<typename TRHICmdList>
-inline bool SetUAVParameterIfCS(TRHICmdList& RHICmdList, const FComputeShaderRHIParamRef Shader, const FShaderResourceParameter& UAVParameter, FUnorderedAccessViewRHIParamRef UAV)
+inline bool SetUAVParameterIfCS(TRHICmdList& RHICmdList, const FComputeShaderRHIParamRef Shader, const FShaderResourceParameter& UAVParameter, FRHIUnorderedAccessView* UAV)
 {
 	SetUAVParameter(RHICmdList, Shader, UAVParameter, UAV);
 	return UAVParameter.IsBound();
@@ -407,7 +407,7 @@ inline void FRWShaderParameter::SetBuffer(TRHICmdList& RHICmdList, TShaderRHIRef
 }
 
 template<typename TShaderRHIRef, typename TRHICmdList>
-inline void FRWShaderParameter::SetTexture(TRHICmdList& RHICmdList, TShaderRHIRef Shader, const FTextureRHIParamRef Texture, FUnorderedAccessViewRHIParamRef UAV) const
+inline void FRWShaderParameter::SetTexture(TRHICmdList& RHICmdList, TShaderRHIRef Shader, const FTextureRHIParamRef Texture, FRHIUnorderedAccessView* UAV) const
 {
 	if (!SetUAVParameterIfCS(RHICmdList, Shader, UAVParameter, UAV))
 	{
diff --git a/Engine/Source/Runtime/RenderCore/Public/ShaderParameters.h b/Engine/Source/Runtime/RenderCore/Public/ShaderParameters.h
--- a/Engine/Source/Runtime/RenderCore/Public/ShaderParameters.h
+++ b/Engine/Source/Runtime/RenderCore/Public/ShaderParameters.h
@@ -145,7 +145,7 @@ public:
 	inline void SetBuffer(TRHICmdList& RHICmdList, TShaderRHIRef Shader, const FRWBufferStructured& RWBuffer) const;
 
 	template<typename TShaderRHIRef, typename TRHICmdList>
-	inline void SetTexture(TRHICmdList& RHICmdList, TShaderRHIRef Shader, const FTextureRHIParamRef Texture, FUnorderedAccessViewRHIParamRef UAV) const;
+	inline void SetTexture(TRHICmdList& RHICmdList, TShaderRHIRef Shader, const FTextureRHIParamRef Texture, FRHIUnorderedAccessView* UAV) const;
 
 	template<typename TRHICmdList>
 	inline void UnsetUAV(TRHICmdList& RHICmdList, FComputeShaderRHIParamRef ComputeShader) const;
@@ -249,7 +249,7 @@ struct FRayTracingShaderBindingsWriter : FRayTracingShaderBindings
 		}
 	}
 
-	void Set(const FShaderResourceParameter& Param, FShaderResourceViewRHIParamRef Value)
+	void Set(const FShaderResourceParameter& Param, FRHIShaderResourceView* Value)
 	{
 		if (Param.IsBound())
 		{
@@ -266,7 +266,7 @@ struct FRayTracingShaderBindingsWriter : FRayTracingShaderBindings
 		}
 	}
 
-	void Set(const FShaderResourceParameter& Param, FUnorderedAccessViewRHIParamRef Value)
+	void Set(const FShaderResourceParameter& Param, FRHIUnorderedAccessView* Value)
 	{
 		if (Param.IsBound())
 		{
@@ -290,7 +290,7 @@ struct FRayTracingShaderBindingsWriter : FRayTracingShaderBindings
 		Textures[BaseIndex] = Value;
 	}
 
-	void SetSRV(uint16 BaseIndex, FShaderResourceViewRHIParamRef Value)
+	void SetSRV(uint16 BaseIndex, FRHIShaderResourceView* Value)
 	{
 		checkSlow(BaseIndex < ARRAY_COUNT(SRVs));
 		SRVs[BaseIndex] = Value;
@@ -302,7 +302,7 @@ struct FRayTracingShaderBindingsWriter : FRayTracingShaderBindings
 		Samplers[BaseIndex] = Value;
 	}
 
-	void SetUAV(uint16 BaseIndex, FUnorderedAccessViewRHIParamRef Value)
+	void SetUAV(uint16 BaseIndex, FRHIUnorderedAccessView* Value)
 	{
 		checkSlow(BaseIndex < ARRAY_COUNT(UAVs));
 		UAVs[BaseIndex] = Value;
diff --git a/Engine/Source/Runtime/Renderer/Private/CapsuleShadowRendering.cpp b/Engine/Source/Runtime/Renderer/Private/CapsuleShadowRendering.cpp
--- a/Engine/Source/Runtime/Renderer/Private/CapsuleShadowRendering.cpp
+++ b/Engine/Source/Runtime/Renderer/Private/CapsuleShadowRendering.cpp
@@ -149,14 +149,14 @@ public:
 		FScene* Scene,
 		const FSceneView& View, 
 		int32 NumLightDirectionDataValue,
-		FShaderResourceViewRHIParamRef LightDirectionDataSRV,
+		FRHIShaderResourceView* LightDirectionDataSRV,
 		const FRWBuffer& ComputedLightDirectionDataBuffer)
 	{
 		const FComputeShaderRHIParamRef ShaderRHI = GetComputeShader();
 
 		FGlobalShader::SetParameters<FViewUniformShaderParameters>(RHICmdList, ShaderRHI, View.ViewUniformBuffer);
 
-		FUnorderedAccessViewRHIParamRef OutUAVs[1];
+		FRHIUnorderedAccessView* OutUAVs[1];
 		OutUAVs[0] = ComputedLightDirectionDataBuffer.UAV;
 		RHICmdList.TransitionResources(EResourceTransitionAccess::EWritable, EResourceTransitionPipeline::EComputeToCompute, OutUAVs, ARRAY_COUNT(OutUAVs));
 
@@ -176,7 +176,7 @@ public:
 	{
 		ComputedLightDirectionData.UnsetUAV(RHICmdList, GetComputeShader());
 
-		FUnorderedAccessViewRHIParamRef OutUAVs[1];
+		FRHIUnorderedAccessView* OutUAVs[1];
 		OutUAVs[0] = ComputedLightDirectionDataBuffer.UAV;
 		RHICmdList.TransitionResources(EResourceTransitionAccess::EReadable, EResourceTransitionPipeline::EComputeToCompute, OutUAVs, ARRAY_COUNT(OutUAVs) );
 	}
@@ -323,17 +323,17 @@ public:
 		const FIntRect& ScissorRect,
 		int32 DownsampleFactorValue,
 		int32 NumShadowCapsulesValue,
-		FShaderResourceViewRHIParamRef ShadowCapsuleShapesSRV,
+		FRHIShaderResourceView* ShadowCapsuleShapesSRV,
 		int32 NumMeshDistanceFieldCastersValue,
-		FShaderResourceViewRHIParamRef MeshDistanceFieldCasterIndicesSRV,
-		FShaderResourceViewRHIParamRef LightDirectionDataSRV,
+		FRHIShaderResourceView* MeshDistanceFieldCasterIndicesSRV,
+		FRHIShaderResourceView* LightDirectionDataSRV,
 		FTextureRHIParamRef ReceiverBentNormalTextureValue)
 	{
 		const FComputeShaderRHIParamRef ShaderRHI = GetComputeShader();
 
 		FGlobalShader::SetParameters<FViewUniformShaderParameters>(RHICmdList, ShaderRHI, View.ViewUniformBuffer);
 
-		FUnorderedAccessViewRHIParamRef OutUAVs[2];
+		FRHIUnorderedAccessView* OutUAVs[2];
 		OutUAVs[0] = OutputTexture.UAV;
 
 		if (TileIntersectionCountsBuffer)
@@ -440,7 +440,7 @@ public:
 		BentNormalTexture.UnsetUAV(RHICmdList, GetComputeShader());
 		TileIntersectionCounts.UnsetUAV(RHICmdList, GetComputeShader());
 
-		FUnorderedAccessViewRHIParamRef OutUAVs[2];
+		FRHIUnorderedAccessView* OutUAVs[2];
 		OutUAVs[0] = OutputTexture.UAV;
 
 		if (TileIntersectionCountsBuffer)
@@ -1080,7 +1080,7 @@ void FDeferredShadingSceneRenderer::SetupIndirectCapsuleShadows(
 	int32& NumCapsuleShapes, 
 	int32& NumMeshesWithCapsules, 
 	int32& NumMeshDistanceFieldCasters,
-	FShaderResourceViewRHIParamRef& IndirectShadowLightDirectionSRV) const
+	FRHIShaderResourceView*& IndirectShadowLightDirectionSRV) const
 {
 	const float CosFadeStartAngle = FMath::Cos(GCapsuleShadowFadeAngleFromVertical);
 	const FSkyLightSceneProxy* SkyLight = Scene ? Scene->SkyLight : NULL;
@@ -1385,7 +1385,7 @@ void FDeferredShadingSceneRenderer::RenderIndirectCapsuleShadows(
 					int32 NumCapsuleShapes = 0;
 					int32 NumMeshesWithCapsules = 0;
 					int32 NumMeshDistanceFieldCasters = 0;
-					FShaderResourceViewRHIParamRef IndirectShadowLightDirectionSRV = NULL;
+					FRHIShaderResourceView* IndirectShadowLightDirectionSRV = NULL;
 					SetupIndirectCapsuleShadows(RHICmdList, View, NumCapsuleShapes, NumMeshesWithCapsules, NumMeshDistanceFieldCasters, IndirectShadowLightDirectionSRV);
 
 					if (NumCapsuleShapes > 0 || NumMeshDistanceFieldCasters > 0)
@@ -1607,7 +1607,7 @@ void FDeferredShadingSceneRenderer::RenderCapsuleShadowsForMovableSkylight(FRHIC
 					int32 NumCapsuleShapes = 0;
 					int32 NumMeshesWithCapsules = 0;
 					int32 NumMeshDistanceFieldCasters = 0;
-					FShaderResourceViewRHIParamRef IndirectShadowLightDirectionSRV = NULL;
+					FRHIShaderResourceView* IndirectShadowLightDirectionSRV = NULL;
 					SetupIndirectCapsuleShadows(RHICmdList, View, NumCapsuleShapes, NumMeshesWithCapsules, NumMeshDistanceFieldCasters, IndirectShadowLightDirectionSRV);
 
 					// Don't render indirect occlusion from mesh distance fields when operating on a movable skylight,
diff --git a/Engine/Source/Runtime/Renderer/Private/CompositionLighting/PostProcessAmbientOcclusion.cpp b/Engine/Source/Runtime/Renderer/Private/CompositionLighting/PostProcessAmbientOcclusion.cpp
--- a/Engine/Source/Runtime/Renderer/Private/CompositionLighting/PostProcessAmbientOcclusion.cpp
+++ b/Engine/Source/Runtime/Renderer/Private/CompositionLighting/PostProcessAmbientOcclusion.cpp
@@ -522,7 +522,7 @@ public:
 		TRHICmdList& RHICmdList,
 		const FRenderingCompositePassContext& Context,
 		const FIntRect& OutputRect,
-		FUnorderedAccessViewRHIParamRef OutUAV)
+		FRHIUnorderedAccessView* OutUAV)
 	{
 		const FComputeShaderRHIParamRef ShaderRHI = GetComputeShader();
 
@@ -574,7 +574,7 @@ void FRCPassPostProcessAmbientOcclusionSmooth::DispatchCS(
 	TRHICmdList& RHICmdList,
 	const FRenderingCompositePassContext& Context,
 	const FIntRect& OutputRect,
-	FUnorderedAccessViewRHIParamRef OutUAV) const
+	FRHIUnorderedAccessView* OutUAV) const
 {
 	TShaderMapRef<FPostProcessAmbientOcclusionSmoothCS> ComputeShader(Context.GetShaderMap());
 	RHICmdList.SetComputeShader(ComputeShader->GetComputeShader());
@@ -733,7 +733,7 @@ public:
 	}
 	
 	template <typename TRHICmdList>
-	void SetParametersCompute(TRHICmdList& RHICmdList, const FRenderingCompositePassContext& Context, FIntPoint InputTextureSize, FUnorderedAccessViewRHIParamRef OutUAV)
+	void SetParametersCompute(TRHICmdList& RHICmdList, const FRenderingCompositePassContext& Context, FIntPoint InputTextureSize, FRHIUnorderedAccessView* OutUAV)
 	{
 		const FViewInfo& View = Context.View;
 		const FVector4 HZBRemappingValue = GetHZBValue(View);				
@@ -753,7 +753,7 @@ public:
 	}
 
 	
-	void SetParametersGfx(FRHICommandList& RHICmdList, const FRenderingCompositePassContext& Context, FIntPoint InputTextureSize, FUnorderedAccessViewRHIParamRef OutUAV)
+	void SetParametersGfx(FRHICommandList& RHICmdList, const FRenderingCompositePassContext& Context, FIntPoint InputTextureSize, FRHIUnorderedAccessView* OutUAV)
 	{
 		const FViewInfo& View = Context.View;
 		const FVector4 HZBRemappingValue = GetHZBValue(View);
@@ -843,7 +843,7 @@ FShader* FRCPassPostProcessAmbientOcclusion::SetShaderTemplPS(const FRenderingCo
 
 
 template <uint32 bTAOSetupAsInput, uint32 bDoUpsample, uint32 ShaderQuality, typename TRHICmdList>
-void FRCPassPostProcessAmbientOcclusion::DispatchCS(TRHICmdList& RHICmdList, const FRenderingCompositePassContext& Context, const FIntPoint& TexSize, FUnorderedAccessViewRHIParamRef OutUAV)
+void FRCPassPostProcessAmbientOcclusion::DispatchCS(TRHICmdList& RHICmdList, const FRenderingCompositePassContext& Context, const FIntPoint& TexSize, FRHIUnorderedAccessView* OutUAV)
 {
 	TShaderMapRef<FPostProcessAmbientOcclusionPSandCS<bTAOSetupAsInput, bDoUpsample, ShaderQuality, true> > ComputeShader(Context.GetShaderMap());
 
diff --git a/Engine/Source/Runtime/Renderer/Private/CompositionLighting/PostProcessAmbientOcclusion.h b/Engine/Source/Runtime/Renderer/Private/CompositionLighting/PostProcessAmbientOcclusion.h
--- a/Engine/Source/Runtime/Renderer/Private/CompositionLighting/PostProcessAmbientOcclusion.h
+++ b/Engine/Source/Runtime/Renderer/Private/CompositionLighting/PostProcessAmbientOcclusion.h
@@ -85,7 +85,7 @@ private:
 		TRHICmdList& RHICmdList,
 		const FRenderingCompositePassContext& Context,
 		const FIntRect& OutputRect,
-		FUnorderedAccessViewRHIParamRef OutUAV) const;
+		FRHIUnorderedAccessView* OutUAV) const;
 
 	const ESSAOType AOType;
 	const bool bDirectOutput;
@@ -116,7 +116,7 @@ private:
 	FShader* SetShaderTemplPS(const FRenderingCompositePassContext& Context, FGraphicsPipelineStateInitializer& GraphicsPSOInit);
 
 	template <uint32 bAOSetupAsInput, uint32 bDoUpsample, uint32 SampleSetQuality, typename TRHICmdList>
-	void DispatchCS(TRHICmdList& RHICmdList, const FRenderingCompositePassContext& Context, const FIntPoint& TexSize, FUnorderedAccessViewRHIParamRef OutTextureUAV);
+	void DispatchCS(TRHICmdList& RHICmdList, const FRenderingCompositePassContext& Context, const FIntPoint& TexSize, FRHIUnorderedAccessView* OutTextureUAV);
 	
 	const ESSAOType AOType;
 	const EPixelFormat IntermediateFormatOverride;
diff --git a/Engine/Source/Runtime/Renderer/Private/ConvertToUniformMesh.cpp b/Engine/Source/Runtime/Renderer/Private/ConvertToUniformMesh.cpp
--- a/Engine/Source/Runtime/Renderer/Private/ConvertToUniformMesh.cpp
+++ b/Engine/Source/Runtime/Renderer/Private/ConvertToUniformMesh.cpp
@@ -347,7 +347,7 @@ public:
 
 		const FScene* Scene = (const FScene*)View.Family->Scene;
 
-		FUnorderedAccessViewRHIParamRef UniformMeshUAVs[1];
+		FRHIUnorderedAccessView* UniformMeshUAVs[1];
 		UniformMeshUAVs[0] = Scene->DistanceFieldSceneData.SurfelBuffers->Surfels.UAV;
 		RHICmdList.TransitionResources(EResourceTransitionAccess::ERWBarrier, EResourceTransitionPipeline::EComputeToCompute, UniformMeshUAVs, ARRAY_COUNT(UniformMeshUAVs));
 
@@ -365,7 +365,7 @@ public:
 		SurfelBufferParameters.UnsetParameters(RHICmdList, ShaderRHI);
 
 		const FScene* Scene = (const FScene*)View.Family->Scene;
-		FUnorderedAccessViewRHIParamRef UniformMeshUAVs[1];
+		FRHIUnorderedAccessView* UniformMeshUAVs[1];
 		UniformMeshUAVs[0] = Scene->DistanceFieldSceneData.SurfelBuffers->Surfels.UAV;
 		RHICmdList.TransitionResources(EResourceTransitionAccess::EReadable, EResourceTransitionPipeline::EComputeToCompute, UniformMeshUAVs, ARRAY_COUNT(UniformMeshUAVs));
 	}
diff --git a/Engine/Source/Runtime/Renderer/Private/DeferredShadingRenderer.h b/Engine/Source/Runtime/Renderer/Private/DeferredShadingRenderer.h
--- a/Engine/Source/Runtime/Renderer/Private/DeferredShadingRenderer.h
+++ b/Engine/Source/Runtime/Renderer/Private/DeferredShadingRenderer.h
@@ -324,7 +324,7 @@ private:
 		int32& NumCapsuleShapes, 
 		int32& NumMeshesWithCapsules, 
 		int32& NumMeshDistanceFieldCasters,
-		FShaderResourceViewRHIParamRef& IndirectShadowLightDirectionSRV) const;
+		FRHIShaderResourceView*& IndirectShadowLightDirectionSRV) const;
 
 	/** Renders indirect shadows from capsules modulated onto scene color. */
 	void RenderIndirectCapsuleShadows(
@@ -548,7 +548,7 @@ private:
 	void VisualizeVarianceMipTree(FRHICommandListImmediate& RHICmdList, const FViewInfo& View, const FRWBuffer& VarianceMipTree, FIntVector VarianceMipTreeDimensions);
 
 	void ComputePathCompaction(FRHICommandListImmediate& RHICmdList, const FViewInfo& View, FTextureRHIParamRef RadianceTexture, FTextureRHIParamRef SampleCountTexture, FTextureRHIParamRef PixelPositionTexture,
-		FUnorderedAccessViewRHIParamRef RadianceSortedRedUAV, FUnorderedAccessViewRHIParamRef RadianceSortedGreenUAV, FUnorderedAccessViewRHIParamRef RadianceSortedBlueUAV, FUnorderedAccessViewRHIParamRef RadianceSortedAlphaUAV, FUnorderedAccessViewRHIParamRef SampleCountSortedUAV);
+		FRHIUnorderedAccessView* RadianceSortedRedUAV, FRHIUnorderedAccessView* RadianceSortedGreenUAV, FRHIUnorderedAccessView* RadianceSortedBlueUAV, FRHIUnorderedAccessView* RadianceSortedAlphaUAV, FRHIUnorderedAccessView* SampleCountSortedUAV);
 
 	void BuildSkyLightCdf(FRHICommandListImmediate& RHICmdList, const FViewInfo& View, const FTexture& SkyLightTextureCube, FRWBuffer& RowCdf, FRWBuffer& ColumnCdf, FRWBuffer& CubeFaceCdf);
 	void VisualizeSkyLightCdf(FRHICommandListImmediate& RHICmdList, const FViewInfo& View, FIntVector Dimensions, const FRWBuffer& RowCdf, const FRWBuffer& ColumnCdf, const FRWBuffer& CubeFaceCdf);
diff --git a/Engine/Source/Runtime/Renderer/Private/DistanceFieldAmbientOcclusion.h b/Engine/Source/Runtime/Renderer/Private/DistanceFieldAmbientOcclusion.h
--- a/Engine/Source/Runtime/Renderer/Private/DistanceFieldAmbientOcclusion.h
+++ b/Engine/Source/Runtime/Renderer/Private/DistanceFieldAmbientOcclusion.h
@@ -155,7 +155,7 @@ public:
 		ObjectTilesIndirectArguments.SetBuffer(RHICmdList, ShaderRHI, TileIntersectionResources.ObjectTilesIndirectArguments);
 	}
 
-	void GetUAVs(FTileIntersectionResources& TileIntersectionResources, TArray<FUnorderedAccessViewRHIParamRef>& UAVs)
+	void GetUAVs(FTileIntersectionResources& TileIntersectionResources, TArray<FRHIUnorderedAccessView*>& UAVs)
 	{
 		uint32 MaxIndex = 0;
 		
diff --git a/Engine/Source/Runtime/Renderer/Private/DistanceFieldGlobalIllumination.cpp b/Engine/Source/Runtime/Renderer/Private/DistanceFieldGlobalIllumination.cpp
--- a/Engine/Source/Runtime/Renderer/Private/DistanceFieldGlobalIllumination.cpp
+++ b/Engine/Source/Runtime/Renderer/Private/DistanceFieldGlobalIllumination.cpp
@@ -149,7 +149,7 @@ public:
 		SetShaderValue(RHICmdList, ShaderRHI, LightDirectionAndTraceDistance, FVector4(LightSceneProxy->GetDirection(), GVPLDirectionalLightTraceDistance));
 		SetShaderValue(RHICmdList, ShaderRHI, LightColor, LightSceneProxy->GetColor() * LightSceneProxy->GetIndirectLightingScale());
 
-		FUnorderedAccessViewRHIParamRef OutUAVs[2];
+		FRHIUnorderedAccessView* OutUAVs[2];
 		OutUAVs[0] = GVPLResources.VPLParameterBuffer.UAV;
 		OutUAVs[1] = GVPLResources.VPLData.UAV;
 		RHICmdList.TransitionResources(EResourceTransitionAccess::ERWBarrier, EResourceTransitionPipeline::EComputeToCompute, OutUAVs, ARRAY_COUNT(OutUAVs));
@@ -172,7 +172,7 @@ public:
 		VPLParameterBuffer.UnsetUAV(RHICmdList, GetComputeShader());
 		VPLData.UnsetUAV(RHICmdList, GetComputeShader());
 
-		FUnorderedAccessViewRHIParamRef OutUAVs[2];
+		FRHIUnorderedAccessView* OutUAVs[2];
 		OutUAVs[0] = GVPLResources.VPLParameterBuffer.UAV;
 		OutUAVs[1] = GVPLResources.VPLData.UAV;
 		RHICmdList.TransitionResources(EResourceTransitionAccess::EReadable, EResourceTransitionPipeline::EComputeToCompute, OutUAVs, ARRAY_COUNT(OutUAVs));
@@ -311,7 +311,7 @@ public:
 		FComputeShaderRHIParamRef ShaderRHI = GetComputeShader();
 		FGlobalShader::SetParameters<FViewUniformShaderParameters>(RHICmdList, ShaderRHI, View.ViewUniformBuffer);
 
-		FUnorderedAccessViewRHIParamRef OutUAVs[2];
+		FRHIUnorderedAccessView* OutUAVs[2];
 		OutUAVs[0] = GCulledVPLResources.VPLParameterBuffer.UAV;
 		OutUAVs[1] = GCulledVPLResources.VPLData.UAV;
 		RHICmdList.TransitionResources(EResourceTransitionAccess::ERWBarrier, EResourceTransitionPipeline::EComputeToCompute, OutUAVs, ARRAY_COUNT(OutUAVs));
