DR - Deprecate FViewportRHIParamRef, FRayTracingSceneRHIParamRef,  FRayTracingGeometryRHIParamRef and FStagingBufferRHIParamRef
#rnx

diff --git a/Engine/Source/Runtime/Apple/MetalRHI/Private/MetalUAV.cpp b/Engine/Source/Runtime/Apple/MetalRHI/Private/MetalUAV.cpp
--- a/Engine/Source/Runtime/Apple/MetalRHI/Private/MetalUAV.cpp
+++ b/Engine/Source/Runtime/Apple/MetalRHI/Private/MetalUAV.cpp
@@ -779,7 +779,7 @@ void FMetalGPUFence::WriteInternal(mtlpp::CommandBuffer& CmdBuffer)
 	check(Fence);
 }
 
-void FMetalRHICommandContext::RHICopyToStagingBuffer(FVertexBufferRHIParamRef SourceBufferRHI, FStagingBufferRHIParamRef DestinationStagingBufferRHI, uint32 Offset, uint32 NumBytes)
+void FMetalRHICommandContext::RHICopyToStagingBuffer(FVertexBufferRHIParamRef SourceBufferRHI, FRHIStagingBuffer* DestinationStagingBufferRHI, uint32 Offset, uint32 NumBytes)
 {
 	@autoreleasepool {
 		check(DestinationStagingBufferRHI);
diff --git a/Engine/Source/Runtime/Apple/MetalRHI/Private/MetalVertexBuffer.cpp b/Engine/Source/Runtime/Apple/MetalRHI/Private/MetalVertexBuffer.cpp
--- a/Engine/Source/Runtime/Apple/MetalRHI/Private/MetalVertexBuffer.cpp
+++ b/Engine/Source/Runtime/Apple/MetalRHI/Private/MetalVertexBuffer.cpp
@@ -758,25 +758,25 @@ void FMetalDynamicRHI::RHITransferVertexBufferUnderlyingResource(FVertexBufferRH
 	}
 }
 
-void* FMetalDynamicRHI::RHILockStagingBuffer(FStagingBufferRHIParamRef StagingBuffer, uint32 Offset, uint32 SizeRHI)
+void* FMetalDynamicRHI::RHILockStagingBuffer(FRHIStagingBuffer* StagingBuffer, uint32 Offset, uint32 SizeRHI)
 {
 	FMetalStagingBuffer* Buffer = ResourceCast(StagingBuffer);
 	return Buffer->Lock(Offset, SizeRHI);
 }
-void FMetalDynamicRHI::RHIUnlockStagingBuffer(FStagingBufferRHIParamRef StagingBuffer)
+void FMetalDynamicRHI::RHIUnlockStagingBuffer(FRHIStagingBuffer* StagingBuffer)
 {
 	FMetalStagingBuffer* Buffer = ResourceCast(StagingBuffer);
 	Buffer->Unlock();
 }
 
-void* FMetalDynamicRHI::LockStagingBuffer_RenderThread(class FRHICommandListImmediate& RHICmdList, FStagingBufferRHIParamRef StagingBuffer, uint32 Offset, uint32 SizeRHI)
+void* FMetalDynamicRHI::LockStagingBuffer_RenderThread(class FRHICommandListImmediate& RHICmdList, FRHIStagingBuffer* StagingBuffer, uint32 Offset, uint32 SizeRHI)
 {
 	QUICK_SCOPE_CYCLE_COUNTER(STAT_FMetalDynamicRHI_LockStagingBuffer_RenderThread);
 	check(IsInRenderingThread());
 	
 	return RHILockStagingBuffer(StagingBuffer, Offset, SizeRHI);
 }
-void FMetalDynamicRHI::UnlockStagingBuffer_RenderThread(class FRHICommandListImmediate& RHICmdList, FStagingBufferRHIParamRef StagingBuffer)
+void FMetalDynamicRHI::UnlockStagingBuffer_RenderThread(class FRHICommandListImmediate& RHICmdList, FRHIStagingBuffer* StagingBuffer)
 {
 	QUICK_SCOPE_CYCLE_COUNTER(STAT_FMetalDynamicRHI_UnlockStagingBuffer_RenderThread);
 	check(IsInRenderingThread());
diff --git a/Engine/Source/Runtime/Apple/MetalRHI/Private/MetalViewport.cpp b/Engine/Source/Runtime/Apple/MetalRHI/Private/MetalViewport.cpp
--- a/Engine/Source/Runtime/Apple/MetalRHI/Private/MetalViewport.cpp
+++ b/Engine/Source/Runtime/Apple/MetalRHI/Private/MetalViewport.cpp
@@ -593,12 +593,12 @@ FViewportRHIRef FMetalDynamicRHI::RHICreateViewport(void* WindowHandle,uint32 Si
 	}
 }
 
-void FMetalDynamicRHI::RHIResizeViewport(FViewportRHIParamRef Viewport, uint32 SizeX, uint32 SizeY, bool bIsFullscreen)
+void FMetalDynamicRHI::RHIResizeViewport(FRHIViewport* Viewport, uint32 SizeX, uint32 SizeY, bool bIsFullscreen)
 {
 	RHIResizeViewport(Viewport, SizeX, SizeY, bIsFullscreen, PF_Unknown);
 }
 
-void FMetalDynamicRHI::RHIResizeViewport(FViewportRHIParamRef ViewportRHI,uint32 SizeX,uint32 SizeY,bool bIsFullscreen,EPixelFormat Format)
+void FMetalDynamicRHI::RHIResizeViewport(FRHIViewport* ViewportRHI,uint32 SizeX,uint32 SizeY,bool bIsFullscreen,EPixelFormat Format)
 {
 	@autoreleasepool {
 	check( IsInGameThread() );
@@ -617,12 +617,12 @@ void FMetalDynamicRHI::RHITick( float DeltaTime )
  *	Viewport functions.
  *=============================================================================*/
 
-void FMetalRHICommandContext::RHIBeginDrawingViewport(FViewportRHIParamRef ViewportRHI, FTextureRHIParamRef RenderTargetRHI)
+void FMetalRHICommandContext::RHIBeginDrawingViewport(FRHIViewport* ViewportRHI, FTextureRHIParamRef RenderTargetRHI)
 {
 	check(false);
 }
 
-void FMetalRHIImmediateCommandContext::RHIBeginDrawingViewport(FViewportRHIParamRef ViewportRHI, FTextureRHIParamRef RenderTargetRHI)
+void FMetalRHIImmediateCommandContext::RHIBeginDrawingViewport(FRHIViewport* ViewportRHI, FTextureRHIParamRef RenderTargetRHI)
 {
 	@autoreleasepool {
 	FMetalViewport* Viewport = ResourceCast(ViewportRHI);
@@ -644,12 +644,12 @@ void FMetalRHIImmediateCommandContext::RHIBeginDrawingViewport(FViewportRHIParam
 	}
 }
 
-void FMetalRHICommandContext::RHIEndDrawingViewport(FViewportRHIParamRef ViewportRHI,bool bPresent,bool bLockToVsync)
+void FMetalRHICommandContext::RHIEndDrawingViewport(FRHIViewport* ViewportRHI,bool bPresent,bool bLockToVsync)
 {
 	check(false);
 }
 
-void FMetalRHIImmediateCommandContext::RHIEndDrawingViewport(FViewportRHIParamRef ViewportRHI,bool bPresent,bool bLockToVsync)
+void FMetalRHIImmediateCommandContext::RHIEndDrawingViewport(FRHIViewport* ViewportRHI,bool bPresent,bool bLockToVsync)
 {
 	@autoreleasepool {
 	FMetalViewport* Viewport = ResourceCast(ViewportRHI);
@@ -657,7 +657,7 @@ void FMetalRHIImmediateCommandContext::RHIEndDrawingViewport(FViewportRHIParamRe
 	}
 }
 
-FTexture2DRHIRef FMetalDynamicRHI::RHIGetViewportBackBuffer(FViewportRHIParamRef ViewportRHI)
+FTexture2DRHIRef FMetalDynamicRHI::RHIGetViewportBackBuffer(FRHIViewport* ViewportRHI)
 {
 	@autoreleasepool {
 	FMetalViewport* Viewport = ResourceCast(ViewportRHI);
@@ -665,7 +665,7 @@ FTexture2DRHIRef FMetalDynamicRHI::RHIGetViewportBackBuffer(FViewportRHIParamRef
 	}
 }
 
-void FMetalDynamicRHI::RHIAdvanceFrameForGetViewportBackBuffer(FViewportRHIParamRef ViewportRHI)
+void FMetalDynamicRHI::RHIAdvanceFrameForGetViewportBackBuffer(FRHIViewport* ViewportRHI)
 {
 	if (GMetalSeparatePresentThread && (GRHISupportsRHIThread && IsRunningRHIInSeparateThread()))
 	{
diff --git a/Engine/Source/Runtime/Apple/MetalRHI/Public/MetalRHI.h b/Engine/Source/Runtime/Apple/MetalRHI/Public/MetalRHI.h
--- a/Engine/Source/Runtime/Apple/MetalRHI/Public/MetalRHI.h
+++ b/Engine/Source/Runtime/Apple/MetalRHI/Public/MetalRHI.h
@@ -122,15 +122,15 @@ public:
 	virtual void RHIRead3DSurfaceFloatData(FTextureRHIParamRef Texture,FIntRect Rect,FIntPoint ZMinMax,TArray<FFloat16Color>& OutData) final override;
 	virtual FRenderQueryRHIRef RHICreateRenderQuery(ERenderQueryType QueryType) final override;
 	virtual bool RHIGetRenderQueryResult(FRHIRenderQuery* RenderQuery, uint64& OutResult, bool bWait) final override;
-	virtual FTexture2DRHIRef RHIGetViewportBackBuffer(FViewportRHIParamRef Viewport) final override;
-	virtual void RHIAdvanceFrameForGetViewportBackBuffer(FViewportRHIParamRef Viewport) final override;
+	virtual FTexture2DRHIRef RHIGetViewportBackBuffer(FRHIViewport* Viewport) final override;
+	virtual void RHIAdvanceFrameForGetViewportBackBuffer(FRHIViewport* Viewport) final override;
 	virtual void RHIAcquireThreadOwnership() final override;
 	virtual void RHIReleaseThreadOwnership() final override;
 	virtual void RHIFlushResources() final override;
 	virtual uint32 RHIGetGPUFrameCycles() final override;
 	virtual FViewportRHIRef RHICreateViewport(void* WindowHandle, uint32 SizeX, uint32 SizeY, bool bIsFullscreen, EPixelFormat PreferredPixelFormat) final override;
-	virtual void RHIResizeViewport(FViewportRHIParamRef Viewport, uint32 SizeX, uint32 SizeY, bool bIsFullscreen) final override;
-	virtual void RHIResizeViewport(FViewportRHIParamRef Viewport, uint32 SizeX, uint32 SizeY, bool bIsFullscreen, EPixelFormat PreferredPixelFormat) final override;
+	virtual void RHIResizeViewport(FRHIViewport* Viewport, uint32 SizeX, uint32 SizeY, bool bIsFullscreen) final override;
+	virtual void RHIResizeViewport(FRHIViewport* Viewport, uint32 SizeX, uint32 SizeY, bool bIsFullscreen, EPixelFormat PreferredPixelFormat) final override;
 	virtual void RHITick(float DeltaTime) final override;
 	virtual void RHISetStreamOutTargets(uint32 NumTargets,const FVertexBufferRHIParamRef* VertexBuffers,const uint32* Offsets) final override;
 	virtual void RHIBlockUntilGPUIdle() final override;
@@ -184,10 +184,10 @@ public:
 	virtual FGPUFenceRHIRef RHICreateGPUFence(const FName &Name) final override;
 
 	virtual FStagingBufferRHIRef RHICreateStagingBuffer() final override;
-	virtual void* RHILockStagingBuffer(FStagingBufferRHIParamRef StagingBuffer, uint32 Offset, uint32 SizeRHI) final override;
-	virtual void RHIUnlockStagingBuffer(FStagingBufferRHIParamRef StagingBuffer) final override;
-	virtual void* LockStagingBuffer_RenderThread(class FRHICommandListImmediate& RHICmdList, FStagingBufferRHIParamRef StagingBuffer, uint32 Offset, uint32 SizeRHI) final override;
-	virtual void UnlockStagingBuffer_RenderThread(class FRHICommandListImmediate& RHICmdList, FStagingBufferRHIParamRef StagingBuffer) final override;	
+	virtual void* RHILockStagingBuffer(FRHIStagingBuffer* StagingBuffer, uint32 Offset, uint32 SizeRHI) final override;
+	virtual void RHIUnlockStagingBuffer(FRHIStagingBuffer* StagingBuffer) final override;
+	virtual void* LockStagingBuffer_RenderThread(class FRHICommandListImmediate& RHICmdList, FRHIStagingBuffer* StagingBuffer, uint32 Offset, uint32 SizeRHI) final override;
+	virtual void UnlockStagingBuffer_RenderThread(class FRHICommandListImmediate& RHICmdList, FRHIStagingBuffer* StagingBuffer) final override;
 	
 	virtual void RHIAcquireTransientResource_RenderThread(FTextureRHIParamRef Texture) final override;
 	virtual void RHIDiscardTransientResource_RenderThread(FTextureRHIParamRef Texture) final override;
diff --git a/Engine/Source/Runtime/Apple/MetalRHI/Public/MetalRHIContext.h b/Engine/Source/Runtime/Apple/MetalRHI/Public/MetalRHIContext.h
--- a/Engine/Source/Runtime/Apple/MetalRHI/Public/MetalRHIContext.h
+++ b/Engine/Source/Runtime/Apple/MetalRHI/Public/MetalRHIContext.h
@@ -70,10 +70,10 @@ public:
 	virtual void RHIDiscardRenderTargets(bool Depth, bool Stencil, uint32 ColorBitMask) final override;
 	
 	// This method is queued with an RHIThread, otherwise it will flush after it is queued; without an RHI thread there is no benefit to queuing this frame advance commands
-	virtual void RHIBeginDrawingViewport(FViewportRHIParamRef Viewport, FTextureRHIParamRef RenderTargetRHI) override;
+	virtual void RHIBeginDrawingViewport(FRHIViewport* Viewport, FTextureRHIParamRef RenderTargetRHI) override;
 	
 	// This method is queued with an RHIThread, otherwise it will flush after it is queued; without an RHI thread there is no benefit to queuing this frame advance commands
-	virtual void RHIEndDrawingViewport(FViewportRHIParamRef Viewport, bool bPresent, bool bLockToVsync) override;
+	virtual void RHIEndDrawingViewport(FRHIViewport* Viewport, bool bPresent, bool bLockToVsync) override;
 	
 	// This method is queued with an RHIThread, otherwise it will flush after it is queued; without an RHI thread there is no benefit to queuing this frame advance commands
 	virtual void RHIBeginFrame() override;
@@ -267,7 +267,7 @@ public:
 	 */
 	virtual void RHITransitionResources(EResourceTransitionAccess TransitionType, EResourceTransitionPipeline TransitionPipeline, FRHIUnorderedAccessView** InUAVs, int32 NumUAVs, FRHIComputeFence* WriteComputeFence) final override;
 	
-	virtual void RHICopyToStagingBuffer(FVertexBufferRHIParamRef SourceBufferRHI, FStagingBufferRHIParamRef DestinationStagingBufferRHI, uint32 Offset, uint32 NumBytes) final override;
+	virtual void RHICopyToStagingBuffer(FVertexBufferRHIParamRef SourceBufferRHI, FRHIStagingBuffer* DestinationStagingBufferRHI, uint32 Offset, uint32 NumBytes) final override;
 	virtual void RHIWriteGPUFence(FRHIGPUFence* FenceRHI) final override;
 
 	/**
@@ -338,8 +338,8 @@ public:
 	FMetalRHIImmediateCommandContext(class FMetalProfiler* InProfiler, FMetalContext* WrapContext);
 
 	// FRHICommandContext API accessible only on the immediate device context
-	virtual void RHIBeginDrawingViewport(FViewportRHIParamRef Viewport, FTextureRHIParamRef RenderTargetRHI) final override;
-	virtual void RHIEndDrawingViewport(FViewportRHIParamRef Viewport, bool bPresent, bool bLockToVsync) final override;
+	virtual void RHIBeginDrawingViewport(FRHIViewport* Viewport, FTextureRHIParamRef RenderTargetRHI) final override;
+	virtual void RHIEndDrawingViewport(FRHIViewport* Viewport, bool bPresent, bool bLockToVsync) final override;
 	virtual void RHIBeginFrame() final override;
 	virtual void RHIEndFrame() final override;
 	virtual void RHIBeginScene() final override;
diff --git a/Engine/Source/Runtime/D3D12RHI/Private/D3D12CommandContext.h b/Engine/Source/Runtime/D3D12RHI/Private/D3D12CommandContext.h
--- a/Engine/Source/Runtime/D3D12RHI/Private/D3D12CommandContext.h
+++ b/Engine/Source/Runtime/D3D12RHI/Private/D3D12CommandContext.h
@@ -35,8 +35,8 @@ public:
 
 	FD3D12CommandContextBase(class FD3D12Adapter* InParent, FRHIGPUMask InGPUMask, bool InIsDefaultContext, bool InIsAsyncComputeContext);
 
-	void RHIBeginDrawingViewport(FViewportRHIParamRef Viewport, FTextureRHIParamRef RenderTargetRHI) final override;
-	void RHIEndDrawingViewport(FViewportRHIParamRef Viewport, bool bPresent, bool bLockToVsync) final override;
+	void RHIBeginDrawingViewport(FRHIViewport* Viewport, FTextureRHIParamRef RenderTargetRHI) final override;
+	void RHIEndDrawingViewport(FRHIViewport* Viewport, bool bPresent, bool bLockToVsync) final override;
 
 	void RHIBeginFrame() final override;
 	void RHIEndFrame() final override;
@@ -252,7 +252,7 @@ public:
 	virtual void RHICopyToResolveTarget(FTextureRHIParamRef SourceTexture, FTextureRHIParamRef DestTexture, const FResolveParams& ResolveParams) override;
 	virtual void RHICopyTexture(FTextureRHIParamRef SourceTexture, FTextureRHIParamRef DestTexture, const FRHICopyTextureInfo& CopyInfo) final override;
 	virtual void RHITransitionResources(EResourceTransitionAccess TransitionType, FTextureRHIParamRef* InTextures, int32 NumTextures) final override;
-	virtual void RHICopyToStagingBuffer(FVertexBufferRHIParamRef SourceBuffer, FStagingBufferRHIParamRef DestinationStagingBuffer, uint32 Offset, uint32 NumBytes) final override;
+	virtual void RHICopyToStagingBuffer(FVertexBufferRHIParamRef SourceBuffer, FRHIStagingBuffer* DestinationStagingBuffer, uint32 Offset, uint32 NumBytes) final override;
 	virtual void RHIWriteGPUFence(FRHIGPUFence* Fence) final override;
 	virtual void RHIBeginRenderQuery(FRHIRenderQuery* RenderQuery) final override;
 	virtual void RHIEndRenderQuery(FRHIRenderQuery* RenderQuery) final override;
@@ -334,30 +334,30 @@ public:
 #if D3D12_RHI_RAYTRACING
 	virtual void RHICopyBufferRegion(FVertexBufferRHIParamRef DestBuffer, uint64 DstOffset, FVertexBufferRHIParamRef SourceBuffer, uint64 SrcOffset, uint64 NumBytes) final override;
 	virtual void RHICopyBufferRegions(const TArrayView<const FCopyBufferRegionParams> Params) final override;
-	virtual void RHIBuildAccelerationStructure(FRayTracingGeometryRHIParamRef Geometry) final override;
+	virtual void RHIBuildAccelerationStructure(FRHIRayTracingGeometry* Geometry) final override;
 	virtual void RHIUpdateAccelerationStructures(const TArrayView<const FAccelerationStructureUpdateParams> Params) final override;
 	virtual void RHIBuildAccelerationStructures(const TArrayView<const FAccelerationStructureUpdateParams> Params) final override;
-	virtual void RHIBuildAccelerationStructure(FRayTracingSceneRHIParamRef Scene) final override;
-	virtual void RHIRayTraceOcclusion(FRayTracingSceneRHIParamRef Scene,
+	virtual void RHIBuildAccelerationStructure(FRHIRayTracingScene* Scene) final override;
+	virtual void RHIRayTraceOcclusion(FRHIRayTracingScene* Scene,
 		FRHIShaderResourceView* Rays,
 		FRHIUnorderedAccessView* Output,
 		uint32 NumRays) final override;
-	virtual void RHIRayTraceIntersection(FRayTracingSceneRHIParamRef Scene,
+	virtual void RHIRayTraceIntersection(FRHIRayTracingScene* Scene,
 		FRHIShaderResourceView* Rays,
 		FRHIUnorderedAccessView* Output,
 		uint32 NumRays) final override;
 	virtual void RHIRayTraceDispatch(FRHIRayTracingPipelineState* RayTracingPipelineState, FRHIRayTracingShader* RayGenShader,
-		FRayTracingSceneRHIParamRef Scene,
+		FRHIRayTracingScene* Scene,
 		const FRayTracingShaderBindings& GlobalResourceBindings,
 		uint32 Width, uint32 Height) final override;
 	virtual void RHISetRayTracingHitGroup(
-		FRayTracingSceneRHIParamRef Scene, uint32 InstanceIndex, uint32 SegmentIndex, uint32 ShaderSlot,
+		FRHIRayTracingScene* Scene, uint32 InstanceIndex, uint32 SegmentIndex, uint32 ShaderSlot,
 		FRHIRayTracingPipelineState* Pipeline, uint32 HitGroupIndex,
 		uint32 NumUniformBuffers, const FUniformBufferRHIParamRef* UniformBuffers,
 		uint32 LooseParameterDataSize, const void* LooseParameterData,
 		uint32 UserData) final override;
 	virtual void RHISetRayTracingCallableShader(
-		FRayTracingSceneRHIParamRef Scene, uint32 ShaderSlotInScene,
+		FRHIRayTracingScene* Scene, uint32 ShaderSlotInScene,
 		FRHIRayTracingPipelineState* Pipeline, uint32 ShaderIndexInPipeline,
 		uint32 NumUniformBuffers, const FUniformBufferRHIParamRef* UniformBuffers,
 		uint32 UserData) final override;
@@ -462,7 +462,7 @@ public:
 	// Special implementation that only signal the fence once.
 	virtual void RHITransitionResources(EResourceTransitionAccess TransitionType, EResourceTransitionPipeline TransitionPipeline, FRHIUnorderedAccessView** InUAVs, int32 NumUAVs, FRHIComputeFence* WriteComputeFenceRHI) final override;
 
-	FORCEINLINE virtual void RHICopyToStagingBuffer(FVertexBufferRHIParamRef SourceBuffer, FStagingBufferRHIParamRef DestinationStagingBuffer, uint32 Offset, uint32 NumBytes) final override
+	FORCEINLINE virtual void RHICopyToStagingBuffer(FVertexBufferRHIParamRef SourceBuffer, FRHIStagingBuffer* DestinationStagingBuffer, uint32 Offset, uint32 NumBytes) final override
 	{
 		ContextRedirect(RHICopyToStagingBuffer(SourceBuffer, DestinationStagingBuffer, Offset, NumBytes));
 	}
@@ -755,7 +755,7 @@ public:
 		ContextRedirect(RHICopyBufferRegions(Params));
 	}
 #endif
-	virtual void RHIBuildAccelerationStructure(FRayTracingGeometryRHIParamRef Geometry) final override
+	virtual void RHIBuildAccelerationStructure(FRHIRayTracingGeometry* Geometry) final override
 	{
 		ContextRedirect(RHIBuildAccelerationStructure(Geometry));
 	}
@@ -770,12 +770,12 @@ public:
 		ContextRedirect(RHIBuildAccelerationStructures(Params));
 	}
 
-	virtual void RHIBuildAccelerationStructure(FRayTracingSceneRHIParamRef Scene) final override
+	virtual void RHIBuildAccelerationStructure(FRHIRayTracingScene* Scene) final override
 	{
 		ContextRedirect(RHIBuildAccelerationStructure(Scene));
 	}
 
-	virtual void RHIRayTraceOcclusion(FRayTracingSceneRHIParamRef Scene,
+	virtual void RHIRayTraceOcclusion(FRHIRayTracingScene* Scene,
 		FRHIShaderResourceView* Rays,
 		FRHIUnorderedAccessView* Output,
 		uint32 NumRays) final override
@@ -783,7 +783,7 @@ public:
 		ContextRedirect(RHIRayTraceOcclusion(Scene, Rays, Output, NumRays));
 	}
 
-	virtual void RHIRayTraceIntersection(FRayTracingSceneRHIParamRef Scene,
+	virtual void RHIRayTraceIntersection(FRHIRayTracingScene* Scene,
 		FRHIShaderResourceView* Rays,
 		FRHIUnorderedAccessView* Output,
 		uint32 NumRays) final override
@@ -792,7 +792,7 @@ public:
 	}
 
 	virtual void RHIRayTraceDispatch(FRHIRayTracingPipelineState* RayTracingPipelineState, FRHIRayTracingShader* RayGenShader,
-		FRayTracingSceneRHIParamRef Scene,
+		FRHIRayTracingScene* Scene,
 		const FRayTracingShaderBindings& GlobalResourceBindings,
 		uint32 Width, uint32 Height) final override
 	{
@@ -800,7 +800,7 @@ public:
 	}
 
 	virtual void RHISetRayTracingHitGroup(
-		FRayTracingSceneRHIParamRef Scene, uint32 InstanceIndex, uint32 SegmentIndex, uint32 ShaderSlot,
+		FRHIRayTracingScene* Scene, uint32 InstanceIndex, uint32 SegmentIndex, uint32 ShaderSlot,
 		FRHIRayTracingPipelineState* Pipeline, uint32 HitGroupIndex,
 		uint32 NumUniformBuffers, const FUniformBufferRHIParamRef* UniformBuffers,
 		uint32 LooseParameterDataSize, const void* LooseParameterData,
@@ -810,7 +810,7 @@ public:
 	}
 
 	virtual void RHISetRayTracingCallableShader(
-		FRayTracingSceneRHIParamRef Scene, uint32 ShaderSlotInScene,
+		FRHIRayTracingScene* Scene, uint32 ShaderSlotInScene,
 		FRHIRayTracingPipelineState* Pipeline, uint32 ShaderIndexInPipeline,
 		uint32 NumUniformBuffers, const FUniformBufferRHIParamRef* UniformBuffers,
 		uint32 UserData) final override
diff --git a/Engine/Source/Runtime/D3D12RHI/Private/D3D12Commands.cpp b/Engine/Source/Runtime/D3D12RHI/Private/D3D12Commands.cpp
--- a/Engine/Source/Runtime/D3D12RHI/Private/D3D12Commands.cpp
+++ b/Engine/Source/Runtime/D3D12RHI/Private/D3D12Commands.cpp
@@ -391,7 +391,7 @@ void FD3D12CommandContext::RHITransitionResources(EResourceTransitionAccess Tran
 }
 
 
-void FD3D12CommandContext::RHICopyToStagingBuffer(FVertexBufferRHIParamRef SourceBufferRHI, FStagingBufferRHIParamRef StagingBufferRHI, uint32 Offset, uint32 NumBytes)
+void FD3D12CommandContext::RHICopyToStagingBuffer(FVertexBufferRHIParamRef SourceBufferRHI, FRHIStagingBuffer* StagingBufferRHI, uint32 Offset, uint32 NumBytes)
 {
 	FD3D12StagingBuffer* StagingBuffer = FD3D12DynamicRHI::ResourceCast(StagingBufferRHI);
 	check(StagingBuffer);
diff --git a/Engine/Source/Runtime/D3D12RHI/Private/D3D12DirectCommandListManager.cpp b/Engine/Source/Runtime/D3D12RHI/Private/D3D12DirectCommandListManager.cpp
--- a/Engine/Source/Runtime/D3D12RHI/Private/D3D12DirectCommandListManager.cpp
+++ b/Engine/Source/Runtime/D3D12RHI/Private/D3D12DirectCommandListManager.cpp
@@ -44,7 +44,7 @@ FStagingBufferRHIRef FD3D12DynamicRHI::RHICreateStagingBuffer()
 	return new FD3D12StagingBuffer();
 }
 
-void* FD3D12DynamicRHI::RHILockStagingBuffer(FStagingBufferRHIParamRef StagingBufferRHI, uint32 Offset, uint32 SizeRHI)
+void* FD3D12DynamicRHI::RHILockStagingBuffer(FRHIStagingBuffer* StagingBufferRHI, uint32 Offset, uint32 SizeRHI)
 {
 	FD3D12StagingBuffer* StagingBuffer = FD3D12DynamicRHI::ResourceCast(StagingBufferRHI);
 	check(StagingBuffer);
@@ -52,7 +52,7 @@ void* FD3D12DynamicRHI::RHILockStagingBuffer(FStagingBufferRHIParamRef StagingBu
 	return StagingBuffer->Lock(Offset, SizeRHI);
 }
 
-void FD3D12DynamicRHI::RHIUnlockStagingBuffer(FStagingBufferRHIParamRef StagingBufferRHI)
+void FD3D12DynamicRHI::RHIUnlockStagingBuffer(FRHIStagingBuffer* StagingBufferRHI)
 {
 	FD3D12StagingBuffer* StagingBuffer = FD3D12DynamicRHI::ResourceCast(StagingBufferRHI);
 	check(StagingBuffer);
diff --git a/Engine/Source/Runtime/D3D12RHI/Private/D3D12RHIPrivate.h b/Engine/Source/Runtime/D3D12RHI/Private/D3D12RHIPrivate.h
--- a/Engine/Source/Runtime/D3D12RHI/Private/D3D12RHIPrivate.h
+++ b/Engine/Source/Runtime/D3D12RHI/Private/D3D12RHIPrivate.h
@@ -273,8 +273,8 @@ public:
 	virtual FComputeFenceRHIRef RHICreateComputeFence(const FName& Name) final override;
 	virtual FGPUFenceRHIRef RHICreateGPUFence(const FName& Name) final override;
 	virtual FStagingBufferRHIRef RHICreateStagingBuffer() final override;
-	virtual void* RHILockStagingBuffer(FStagingBufferRHIParamRef StagingBuffer, uint32 Offset, uint32 SizeRHI) final override;
-    virtual void RHIUnlockStagingBuffer(FStagingBufferRHIParamRef StagingBuffer) final override;
+	virtual void* RHILockStagingBuffer(FRHIStagingBuffer* StagingBuffer, uint32 Offset, uint32 SizeRHI) final override;
+    virtual void RHIUnlockStagingBuffer(FRHIStagingBuffer* StagingBuffer) final override;
 	virtual FBoundShaderStateRHIRef RHICreateBoundShaderState(FRHIVertexDeclaration* VertexDeclaration, FVertexShaderRHIParamRef VertexShader, FHullShaderRHIParamRef HullShader, FDomainShaderRHIParamRef DomainShader, FPixelShaderRHIParamRef PixelShader, FGeometryShaderRHIParamRef GeometryShader) final override;
 	virtual FGraphicsPipelineStateRHIRef RHICreateGraphicsPipelineState(const FGraphicsPipelineStateInitializer& Initializer) final override;
 	virtual TRefCountPtr<FRHIComputePipelineState> RHICreateComputePipelineState(FRHIComputeShader* ComputeShader) final override;
@@ -342,16 +342,16 @@ public:
 	virtual void RHIRead3DSurfaceFloatData(FTextureRHIParamRef Texture, FIntRect Rect, FIntPoint ZMinMax, TArray<FFloat16Color>& OutData) final override;
 	virtual FRenderQueryRHIRef RHICreateRenderQuery(ERenderQueryType QueryType) final override;
 	virtual bool RHIGetRenderQueryResult(FRHIRenderQuery* RenderQuery, uint64& OutResult, bool bWait) final override;
-	virtual uint32 RHIGetViewportNextPresentGPUIndex(FViewportRHIParamRef Viewport) final override;
-	virtual FTexture2DRHIRef RHIGetViewportBackBuffer(FViewportRHIParamRef Viewport) final override;
-	virtual void RHIAdvanceFrameForGetViewportBackBuffer(FViewportRHIParamRef Viewport) final override;
+	virtual uint32 RHIGetViewportNextPresentGPUIndex(FRHIViewport* Viewport) final override;
+	virtual FTexture2DRHIRef RHIGetViewportBackBuffer(FRHIViewport* Viewport) final override;
+	virtual void RHIAdvanceFrameForGetViewportBackBuffer(FRHIViewport* Viewport) final override;
 	virtual void RHIAcquireThreadOwnership() final override;
 	virtual void RHIReleaseThreadOwnership() final override;
 	virtual void RHIFlushResources() final override;
 	virtual uint32 RHIGetGPUFrameCycles() final override;
 	virtual FViewportRHIRef RHICreateViewport(void* WindowHandle, uint32 SizeX, uint32 SizeY, bool bIsFullscreen, EPixelFormat PreferredPixelFormat) final override;
-	virtual void RHIResizeViewport(FViewportRHIParamRef Viewport, uint32 SizeX, uint32 SizeY, bool bIsFullscreen) final override;
-	virtual void RHIResizeViewport(FViewportRHIParamRef ViewportRHI, uint32 SizeX, uint32 SizeY, bool bIsFullscreen, EPixelFormat PreferredPixelFormat) final override;
+	virtual void RHIResizeViewport(FRHIViewport* Viewport, uint32 SizeX, uint32 SizeY, bool bIsFullscreen) final override;
+	virtual void RHIResizeViewport(FRHIViewport* ViewportRHI, uint32 SizeX, uint32 SizeY, bool bIsFullscreen, EPixelFormat PreferredPixelFormat) final override;
 	virtual void RHITick(float DeltaTime) final override;
 	virtual void RHISetStreamOutTargets(uint32 NumTargets, const FVertexBufferRHIParamRef* VertexBuffers, const uint32* Offsets) final override;
 	virtual void RHIBlockUntilGPUIdle() final override;
diff --git a/Engine/Source/Runtime/D3D12RHI/Private/D3D12RayTracing.cpp b/Engine/Source/Runtime/D3D12RHI/Private/D3D12RayTracing.cpp
--- a/Engine/Source/Runtime/D3D12RHI/Private/D3D12RayTracing.cpp
+++ b/Engine/Source/Runtime/D3D12RHI/Private/D3D12RayTracing.cpp
@@ -2477,7 +2477,7 @@ FD3D12RayTracingShaderTable* FD3D12RayTracingScene::FindOrCreateShaderTable(cons
 	return CreatedShaderTable;
 }
 
-void FD3D12CommandContext::RHIBuildAccelerationStructure(FRayTracingGeometryRHIParamRef InGeometry)
+void FD3D12CommandContext::RHIBuildAccelerationStructure(FRHIRayTracingGeometry* InGeometry)
 {
 	FD3D12RayTracingGeometry* Geometry = FD3D12DynamicRHI::ResourceCast(InGeometry);
 	Geometry->TransitionBuffers(*this);
@@ -2531,7 +2531,7 @@ void FD3D12CommandContext::RHIBuildAccelerationStructures(const TArrayView<const
 	}
 }
 
-void FD3D12CommandContext::RHIBuildAccelerationStructure(FRayTracingSceneRHIParamRef InScene)
+void FD3D12CommandContext::RHIBuildAccelerationStructure(FRHIRayTracingScene* InScene)
 {
 	FD3D12RayTracingScene* Scene = FD3D12DynamicRHI::ResourceCast(InScene);
 	Scene->BuildAccelerationStructure(*this, D3D12_RAYTRACING_ACCELERATION_STRUCTURE_BUILD_FLAG_PREFER_FAST_TRACE);
@@ -3003,7 +3003,7 @@ static void DispatchRays(FD3D12CommandContext& CommandContext,
 }
 
 
-void FD3D12CommandContext::RHIRayTraceOcclusion(FRayTracingSceneRHIParamRef InScene,
+void FD3D12CommandContext::RHIRayTraceOcclusion(FRHIRayTracingScene* InScene,
 	FRHIShaderResourceView* Rays,
 	FRHIUnorderedAccessView* Output,
 	uint32 NumRays)
@@ -3038,7 +3038,7 @@ void FD3D12CommandContext::RHIRayTraceOcclusion(FRayTracingSceneRHIParamRef InSc
 	DispatchRays(*this, Bindings, Pipeline, 0, nullptr, DispatchDesc);
 }
 
-void FD3D12CommandContext::RHIRayTraceIntersection(FRayTracingSceneRHIParamRef InScene,
+void FD3D12CommandContext::RHIRayTraceIntersection(FRHIRayTracingScene* InScene,
 	FRHIShaderResourceView* InRays,
 	FRHIUnorderedAccessView* InOutput,
 	uint32 NumRays)
@@ -3079,7 +3079,7 @@ void FD3D12CommandContext::RHIRayTraceIntersection(FRayTracingSceneRHIParamRef I
 }
 
 void FD3D12CommandContext::RHIRayTraceDispatch(FRHIRayTracingPipelineState* InRayTracingPipelineState, FRHIRayTracingShader* RayGenShaderRHI,
-	FRayTracingSceneRHIParamRef InScene,
+	FRHIRayTracingScene* InScene,
 	const FRayTracingShaderBindings& GlobalResourceBindings,
 	uint32 Width, uint32 Height)
 {
@@ -3111,7 +3111,7 @@ void FD3D12CommandContext::RHIRayTraceDispatch(FRHIRayTracingPipelineState* InRa
 }
 
 void FD3D12CommandContext::RHISetRayTracingHitGroup(
-	FRayTracingSceneRHIParamRef InScene, uint32 InstanceIndex, uint32 SegmentIndex, uint32 ShaderSlot,
+	FRHIRayTracingScene* InScene, uint32 InstanceIndex, uint32 SegmentIndex, uint32 ShaderSlot,
 	FRHIRayTracingPipelineState* InPipeline, uint32 HitGroupIndex,
 	uint32 NumUniformBuffers, const FUniformBufferRHIParamRef* UniformBuffers,
 	uint32 LooseParameterDataSize, const void* LooseParameterData,
@@ -3148,7 +3148,7 @@ void FD3D12CommandContext::RHISetRayTracingHitGroup(
 
 
 void FD3D12CommandContext::RHISetRayTracingCallableShader(
-	FRayTracingSceneRHIParamRef InScene, uint32 ShaderSlotInScene,
+	FRHIRayTracingScene* InScene, uint32 ShaderSlotInScene,
 	FRHIRayTracingPipelineState* InPipeline, uint32 ShaderIndexInPipeline,
 	uint32 NumUniformBuffers, const FUniformBufferRHIParamRef* UniformBuffers,
 	uint32 UserData)
diff --git a/Engine/Source/Runtime/D3D12RHI/Private/D3D12Viewport.cpp b/Engine/Source/Runtime/D3D12RHI/Private/D3D12Viewport.cpp
--- a/Engine/Source/Runtime/D3D12RHI/Private/D3D12Viewport.cpp
+++ b/Engine/Source/Runtime/D3D12RHI/Private/D3D12Viewport.cpp
@@ -767,7 +767,7 @@ FViewportRHIRef FD3D12DynamicRHI::RHICreateViewport(void* WindowHandle, uint32 S
 	return RenderingViewport;
 }
 
-void FD3D12DynamicRHI::RHIResizeViewport(FViewportRHIParamRef ViewportRHI, uint32 SizeX, uint32 SizeY, bool bIsFullscreen)
+void FD3D12DynamicRHI::RHIResizeViewport(FRHIViewport* ViewportRHI, uint32 SizeX, uint32 SizeY, bool bIsFullscreen)
 {
 	check(IsInGameThread());
 
@@ -775,7 +775,7 @@ void FD3D12DynamicRHI::RHIResizeViewport(FViewportRHIParamRef ViewportRHI, uint3
 	Viewport->Resize(SizeX, SizeY, bIsFullscreen, PF_Unknown);
 }
 
-void FD3D12DynamicRHI::RHIResizeViewport(FViewportRHIParamRef ViewportRHI, uint32 SizeX, uint32 SizeY, bool bIsFullscreen, EPixelFormat PreferredPixelFormat)
+void FD3D12DynamicRHI::RHIResizeViewport(FRHIViewport* ViewportRHI, uint32 SizeX, uint32 SizeY, bool bIsFullscreen, EPixelFormat PreferredPixelFormat)
 {
 	check(IsInGameThread());
 
@@ -806,7 +806,7 @@ void FD3D12DynamicRHI::RHITick(float DeltaTime)
  *	Viewport functions.
  *=============================================================================*/
 
-void FD3D12CommandContextBase::RHIBeginDrawingViewport(FViewportRHIParamRef ViewportRHI, FTextureRHIParamRef RenderTargetRHI)
+void FD3D12CommandContextBase::RHIBeginDrawingViewport(FRHIViewport* ViewportRHI, FTextureRHIParamRef RenderTargetRHI)
 {
 	FD3D12Viewport* Viewport = FD3D12DynamicRHI::ResourceCast(ViewportRHI);
 
@@ -831,7 +831,7 @@ void FD3D12CommandContextBase::RHIBeginDrawingViewport(FViewportRHIParamRef View
 	RHISetRenderTargets(1, &RTView, nullptr, 0, nullptr);
 }
 
-void FD3D12CommandContextBase::RHIEndDrawingViewport(FViewportRHIParamRef ViewportRHI, bool bPresent, bool bLockToVsync)
+void FD3D12CommandContextBase::RHIEndDrawingViewport(FRHIViewport* ViewportRHI, bool bPresent, bool bLockToVsync)
 {
 	FD3D12DynamicRHI& RHI = *ParentAdapter->GetOwningRHI();
 	FD3D12Viewport* Viewport = FD3D12DynamicRHI::ResourceCast(ViewportRHI);
@@ -878,7 +878,7 @@ void FD3D12CommandContextBase::RHIEndDrawingViewport(FViewportRHIParamRef Viewpo
 	}
 }
 
-void FD3D12DynamicRHI::RHIAdvanceFrameForGetViewportBackBuffer(FViewportRHIParamRef ViewportRHI)
+void FD3D12DynamicRHI::RHIAdvanceFrameForGetViewportBackBuffer(FRHIViewport* ViewportRHI)
 {
 	check(IsInRenderingThread());
 
@@ -892,7 +892,7 @@ void FD3D12DynamicRHI::RHIAdvanceFrameForGetViewportBackBuffer(FViewportRHIParam
 	Viewport->AdvanceBackBufferFrame_RenderThread();
 }
 
-uint32 FD3D12DynamicRHI::RHIGetViewportNextPresentGPUIndex(FViewportRHIParamRef ViewportRHI)
+uint32 FD3D12DynamicRHI::RHIGetViewportNextPresentGPUIndex(FRHIViewport* ViewportRHI)
 {
 	check(IsInRenderingThread());
 #if WITH_MGPU
@@ -905,7 +905,7 @@ uint32 FD3D12DynamicRHI::RHIGetViewportNextPresentGPUIndex(FViewportRHIParamRef
 	return 0;
 }
 
-FTexture2DRHIRef FD3D12DynamicRHI::RHIGetViewportBackBuffer(FViewportRHIParamRef ViewportRHI)
+FTexture2DRHIRef FD3D12DynamicRHI::RHIGetViewportBackBuffer(FRHIViewport* ViewportRHI)
 {
 	check(IsInRenderingThread());
 
diff --git a/Engine/Source/Runtime/EmptyRHI/Private/EmptyViewport.cpp b/Engine/Source/Runtime/EmptyRHI/Private/EmptyViewport.cpp
--- a/Engine/Source/Runtime/EmptyRHI/Private/EmptyViewport.cpp
+++ b/Engine/Source/Runtime/EmptyRHI/Private/EmptyViewport.cpp
@@ -26,7 +26,7 @@ FViewportRHIRef FEmptyDynamicRHI::RHICreateViewport(void* WindowHandle,uint32 Si
 	return new FEmptyViewport(WindowHandle, SizeX, SizeY, bIsFullscreen);
 }
 
-void FEmptyDynamicRHI::RHIResizeViewport(FViewportRHIParamRef ViewportRHI,uint32 SizeX,uint32 SizeY,bool bIsFullscreen)
+void FEmptyDynamicRHI::RHIResizeViewport(FRHIViewport* ViewportRHI,uint32 SizeX,uint32 SizeY,bool bIsFullscreen)
 {
 	check( IsInGameThread() );
 
@@ -43,25 +43,25 @@ void FEmptyDynamicRHI::RHITick( float DeltaTime )
  *	Viewport functions.
  *=============================================================================*/
 
-void FEmptyDynamicRHI::RHIBeginDrawingViewport(FViewportRHIParamRef ViewportRHI, FTextureRHIParamRef RenderTargetRHI)
+void FEmptyDynamicRHI::RHIBeginDrawingViewport(FRHIViewport* ViewportRHI, FTextureRHIParamRef RenderTargetRHI)
 {
 	FEmptyViewport* Viewport = ResourceCast(ViewportRHI);
 
 	//RHISetRenderTarget(RHIGetViewportBackBuffer(ViewportRHI), NULL);
 }
 
-void FEmptyDynamicRHI::RHIEndDrawingViewport(FViewportRHIParamRef ViewportRHI,bool bPresent,bool bLockToVsync)
+void FEmptyDynamicRHI::RHIEndDrawingViewport(FRHIViewport* ViewportRHI,bool bPresent,bool bLockToVsync)
 {
 	FEmptyViewport* Viewport = ResourceCast(ViewportRHI);
 }
 
-FTexture2DRHIRef FEmptyDynamicRHI::RHIGetViewportBackBuffer(FViewportRHIParamRef ViewportRHI)
+FTexture2DRHIRef FEmptyDynamicRHI::RHIGetViewportBackBuffer(FRHIViewport* ViewportRHI)
 {
 	FEmptyViewport* Viewport = ResourceCast(ViewportRHI);
 
 	return FTexture2DRHIRef();
 }
 
-void FEmptyDynamicRHI::RHIAdvanceFrameForGetViewportBackBuffer(FViewportRHIParamRef Viewport)
+void FEmptyDynamicRHI::RHIAdvanceFrameForGetViewportBackBuffer(FRHIViewport* Viewport)
 {
 }
diff --git a/Engine/Source/Runtime/EmptyRHI/Public/EmptyRHI.h b/Engine/Source/Runtime/EmptyRHI/Public/EmptyRHI.h
--- a/Engine/Source/Runtime/EmptyRHI/Public/EmptyRHI.h
+++ b/Engine/Source/Runtime/EmptyRHI/Public/EmptyRHI.h
@@ -110,14 +110,14 @@ public:
 	virtual void RHIRead3DSurfaceFloatData(FTextureRHIParamRef Texture,FIntRect Rect,FIntPoint ZMinMax,TArray<FFloat16Color>& OutData) final override;
 	virtual FRenderQueryRHIRef RHICreateRenderQuery(ERenderQueryType QueryType) final override;
 	virtual bool RHIGetRenderQueryResult(FRHIRenderQuery* RenderQuery, uint64& OutResult, bool bWait) final override;
-	virtual FTexture2DRHIRef RHIGetViewportBackBuffer(FViewportRHIParamRef Viewport) final override;
-	virtual void RHIAdvanceFrameForGetViewportBackBuffer(FViewportRHIParamRef Viewport) final override;
+	virtual FTexture2DRHIRef RHIGetViewportBackBuffer(FRHIViewport* Viewport) final override;
+	virtual void RHIAdvanceFrameForGetViewportBackBuffer(FRHIViewport* Viewport) final override;
 	virtual void RHIAcquireThreadOwnership() final override;
 	virtual void RHIReleaseThreadOwnership() final override;
 	virtual void RHIFlushResources() final override;
 	virtual uint32 RHIGetGPUFrameCycles() final override;
 	virtual FViewportRHIRef RHICreateViewport(void* WindowHandle, uint32 SizeX, uint32 SizeY, bool bIsFullscreen, EPixelFormat PreferredPixelFormat) final override;
-	virtual void RHIResizeViewport(FViewportRHIParamRef Viewport, uint32 SizeX, uint32 SizeY, bool bIsFullscreen) final override;
+	virtual void RHIResizeViewport(FRHIViewport* Viewport, uint32 SizeX, uint32 SizeY, bool bIsFullscreen) final override;
 	virtual void RHITick(float DeltaTime) final override;
 	virtual void RHISetStreamOutTargets(uint32 NumTargets,const FVertexBufferRHIParamRef* VertexBuffers,const uint32* Offsets) final override;
 	virtual void RHIBlockUntilGPUIdle() final override;
@@ -144,8 +144,8 @@ public:
 	virtual void RHIBeginRenderQuery(FRHIRenderQuery* RenderQuery) final override;
 	virtual void RHIEndRenderQuery(FRHIRenderQuery* RenderQuery) final override;
 	virtual void RHISubmitCommandsHint() final override;
-	virtual void RHIBeginDrawingViewport(FViewportRHIParamRef Viewport, FTextureRHIParamRef RenderTargetRHI) final override;
-	virtual void RHIEndDrawingViewport(FViewportRHIParamRef Viewport, bool bPresent, bool bLockToVsync) final override;
+	virtual void RHIBeginDrawingViewport(FRHIViewport* Viewport, FTextureRHIParamRef RenderTargetRHI) final override;
+	virtual void RHIEndDrawingViewport(FRHIViewport* Viewport, bool bPresent, bool bLockToVsync) final override;
 	virtual void RHIBeginFrame() final override;
 	virtual void RHIEndFrame() final override;
 	virtual void RHIBeginScene() final override;
diff --git a/Engine/Source/Runtime/NullDrv/Public/NullRHI.h b/Engine/Source/Runtime/NullDrv/Public/NullRHI.h
--- a/Engine/Source/Runtime/NullDrv/Public/NullRHI.h
+++ b/Engine/Source/Runtime/NullDrv/Public/NullRHI.h
@@ -442,15 +442,15 @@ public:
 	}
 
 
-	virtual void RHIBeginDrawingViewport(FViewportRHIParamRef Viewport, FTextureRHIParamRef RenderTargetRHI) final override
+	virtual void RHIBeginDrawingViewport(FRHIViewport* Viewport, FTextureRHIParamRef RenderTargetRHI) final override
 	{
 	}
 
-	virtual void RHIEndDrawingViewport(FViewportRHIParamRef Viewport, bool bPresent, bool bLockToVsync) final override
+	virtual void RHIEndDrawingViewport(FRHIViewport* Viewport, bool bPresent, bool bLockToVsync) final override
 	{
 	}
 
-	virtual FTexture2DRHIRef RHIGetViewportBackBuffer(FViewportRHIParamRef Viewport) final override
+	virtual FTexture2DRHIRef RHIGetViewportBackBuffer(FRHIViewport* Viewport) final override
 	{ 
 		return new FRHITexture2D(1,1,1,1,PF_B8G8R8A8,TexCreate_RenderTargetable, FClearValueBinding()); 
 	}
@@ -473,7 +473,7 @@ public:
 	{
 
 	}
-	virtual void RHIAdvanceFrameForGetViewportBackBuffer(FViewportRHIParamRef Viewport) final override
+	virtual void RHIAdvanceFrameForGetViewportBackBuffer(FRHIViewport* Viewport) final override
 	{
 
 	}
@@ -501,7 +501,7 @@ public:
 	{ 
 		return new FRHIViewport(); 
 	}
-	virtual void RHIResizeViewport(FViewportRHIParamRef Viewport, uint32 SizeX, uint32 SizeY, bool bIsFullscreen) final override
+	virtual void RHIResizeViewport(FRHIViewport* Viewport, uint32 SizeX, uint32 SizeY, bool bIsFullscreen) final override
 	{
 
 	}
diff --git a/Engine/Source/Runtime/OpenGLDrv/Private/OpenGLCommands.cpp b/Engine/Source/Runtime/OpenGLDrv/Private/OpenGLCommands.cpp
--- a/Engine/Source/Runtime/OpenGLDrv/Private/OpenGLCommands.cpp
+++ b/Engine/Source/Runtime/OpenGLDrv/Private/OpenGLCommands.cpp
@@ -3325,7 +3325,7 @@ void FOpenGLDynamicRHI::RHIInvalidateCachedState()
 	SharedContextState.InitializeResources(FOpenGL::GetMaxCombinedTextureImageUnits(), OGL_MAX_COMPUTE_STAGE_UAV_UNITS);
 }
 
-void FOpenGLDynamicRHI::RHICopyToStagingBuffer(FVertexBufferRHIParamRef SourceBufferRHI, FStagingBufferRHIParamRef DestinationStagingBufferRHI, uint32 InOffset, uint32 InNumBytes)
+void FOpenGLDynamicRHI::RHICopyToStagingBuffer(FVertexBufferRHIParamRef SourceBufferRHI, FRHIStagingBuffer* DestinationStagingBufferRHI, uint32 InOffset, uint32 InNumBytes)
 {
 #if OPENGL_GL3
 	VERIFY_GL_SCOPE();
diff --git a/Engine/Source/Runtime/OpenGLDrv/Private/OpenGLVertexBuffer.cpp b/Engine/Source/Runtime/OpenGLDrv/Private/OpenGLVertexBuffer.cpp
--- a/Engine/Source/Runtime/OpenGLDrv/Private/OpenGLVertexBuffer.cpp
+++ b/Engine/Source/Runtime/OpenGLDrv/Private/OpenGLVertexBuffer.cpp
@@ -320,24 +320,24 @@ void FOpenGLStagingBuffer::Unlock()
 #endif
 }
 
-void* FOpenGLDynamicRHI::RHILockStagingBuffer(FStagingBufferRHIParamRef StagingBuffer, uint32 Offset, uint32 SizeRHI)
+void* FOpenGLDynamicRHI::RHILockStagingBuffer(FRHIStagingBuffer* StagingBuffer, uint32 Offset, uint32 SizeRHI)
 {
 	FOpenGLStagingBuffer* Buffer = ResourceCast(StagingBuffer);
 	return Buffer->Lock(Offset, SizeRHI);	
 }
 
-void FOpenGLDynamicRHI::RHIUnlockStagingBuffer(FStagingBufferRHIParamRef StagingBuffer)
+void FOpenGLDynamicRHI::RHIUnlockStagingBuffer(FRHIStagingBuffer* StagingBuffer)
 {
 	FOpenGLStagingBuffer* Buffer = ResourceCast(StagingBuffer);
 	Buffer->Unlock();
 }
 
-void* FOpenGLDynamicRHI::LockStagingBuffer_RenderThread(class FRHICommandListImmediate& RHICmdList, FStagingBufferRHIParamRef StagingBuffer, uint32 Offset, uint32 SizeRHI)
+void* FOpenGLDynamicRHI::LockStagingBuffer_RenderThread(class FRHICommandListImmediate& RHICmdList, FRHIStagingBuffer* StagingBuffer, uint32 Offset, uint32 SizeRHI)
 {
 	return RHILockStagingBuffer(StagingBuffer, Offset, SizeRHI);
 }
 
-void FOpenGLDynamicRHI::UnlockStagingBuffer_RenderThread(class FRHICommandListImmediate& RHICmdList, FStagingBufferRHIParamRef StagingBuffer)
+void FOpenGLDynamicRHI::UnlockStagingBuffer_RenderThread(class FRHICommandListImmediate& RHICmdList, FRHIStagingBuffer* StagingBuffer)
 {
 	RHIUnlockStagingBuffer(StagingBuffer);
 }
diff --git a/Engine/Source/Runtime/OpenGLDrv/Private/OpenGLViewport.cpp b/Engine/Source/Runtime/OpenGLDrv/Private/OpenGLViewport.cpp
--- a/Engine/Source/Runtime/OpenGLDrv/Private/OpenGLViewport.cpp
+++ b/Engine/Source/Runtime/OpenGLDrv/Private/OpenGLViewport.cpp
@@ -71,7 +71,7 @@ FViewportRHIRef FOpenGLDynamicRHI::RHICreateViewport(void* WindowHandle,uint32 S
 	return new FOpenGLViewport(this,WindowHandle,SizeX,SizeY,bIsFullscreen,PreferredPixelFormat);
 }
 
-void FOpenGLDynamicRHI::RHIResizeViewport(FViewportRHIParamRef ViewportRHI,uint32 SizeX,uint32 SizeY,bool bIsFullscreen)
+void FOpenGLDynamicRHI::RHIResizeViewport(FRHIViewport* ViewportRHI,uint32 SizeX,uint32 SizeY,bool bIsFullscreen)
 {
 	FOpenGLViewport* Viewport = ResourceCast(ViewportRHI);
 	check( IsInGameThread() );
@@ -89,7 +89,7 @@ void FOpenGLDynamicRHI::RHITick( float DeltaTime )
  *	Viewport functions.
  *=============================================================================*/
 
-void FOpenGLDynamicRHI::RHIBeginDrawingViewport(FViewportRHIParamRef ViewportRHI, FTextureRHIParamRef RenderTarget)
+void FOpenGLDynamicRHI::RHIBeginDrawingViewport(FRHIViewport* ViewportRHI, FTextureRHIParamRef RenderTarget)
 {
 	VERIFY_GL_SCOPE();
 
@@ -129,7 +129,7 @@ void FOpenGLDynamicRHI::RHIBeginDrawingViewport(FViewportRHIParamRef ViewportRHI
 	}
 }
 
-void FOpenGLDynamicRHI::RHIEndDrawingViewport(FViewportRHIParamRef ViewportRHI,bool bPresent,bool bLockToVsync)
+void FOpenGLDynamicRHI::RHIEndDrawingViewport(FRHIViewport* ViewportRHI,bool bPresent,bool bLockToVsync)
 {
 	VERIFY_GL_SCOPE();
 
@@ -209,13 +209,13 @@ void FOpenGLDynamicRHI::RHIEndDrawingViewport(FViewportRHIParamRef ViewportRHI,b
 }
 
 
-FTexture2DRHIRef FOpenGLDynamicRHI::RHIGetViewportBackBuffer(FViewportRHIParamRef ViewportRHI)
+FTexture2DRHIRef FOpenGLDynamicRHI::RHIGetViewportBackBuffer(FRHIViewport* ViewportRHI)
 {
 	FOpenGLViewport* Viewport = ResourceCast(ViewportRHI);
 	return Viewport->GetBackBuffer();
 }
 
-void FOpenGLDynamicRHI::RHIAdvanceFrameForGetViewportBackBuffer(FViewportRHIParamRef Viewport)
+void FOpenGLDynamicRHI::RHIAdvanceFrameForGetViewportBackBuffer(FRHIViewport* Viewport)
 {
 }
 
diff --git a/Engine/Source/Runtime/OpenGLDrv/Public/OpenGLDrv.h b/Engine/Source/Runtime/OpenGLDrv/Public/OpenGLDrv.h
--- a/Engine/Source/Runtime/OpenGLDrv/Public/OpenGLDrv.h
+++ b/Engine/Source/Runtime/OpenGLDrv/Public/OpenGLDrv.h
@@ -420,14 +420,14 @@ public:
 	virtual void RHIRead3DSurfaceFloatData(FTextureRHIParamRef Texture,FIntRect Rect,FIntPoint ZMinMax,TArray<FFloat16Color>& OutData) final override;
 	virtual FRenderQueryRHIRef RHICreateRenderQuery(ERenderQueryType QueryType) final override;
 	virtual bool RHIGetRenderQueryResult(FRHIRenderQuery* RenderQuery, uint64& OutResult, bool bWait) final override;
-	virtual FTexture2DRHIRef RHIGetViewportBackBuffer(FViewportRHIParamRef Viewport) final override;
-	virtual void RHIAdvanceFrameForGetViewportBackBuffer(FViewportRHIParamRef Viewport) final override;
+	virtual FTexture2DRHIRef RHIGetViewportBackBuffer(FRHIViewport* Viewport) final override;
+	virtual void RHIAdvanceFrameForGetViewportBackBuffer(FRHIViewport* Viewport) final override;
 	virtual void RHIAcquireThreadOwnership() final override;
 	virtual void RHIReleaseThreadOwnership() final override;
 	virtual void RHIFlushResources() final override;
 	virtual uint32 RHIGetGPUFrameCycles() final override;
 	virtual FViewportRHIRef RHICreateViewport(void* WindowHandle, uint32 SizeX, uint32 SizeY, bool bIsFullscreen, EPixelFormat PreferredPixelFormat) final override;
-	virtual void RHIResizeViewport(FViewportRHIParamRef Viewport, uint32 SizeX, uint32 SizeY, bool bIsFullscreen) final override;
+	virtual void RHIResizeViewport(FRHIViewport* Viewport, uint32 SizeX, uint32 SizeY, bool bIsFullscreen) final override;
 	virtual void RHITick(float DeltaTime) final override;
 	virtual void RHISetStreamOutTargets(uint32 NumTargets,const FVertexBufferRHIParamRef* VertexBuffers,const uint32* Offsets) final override;
 	virtual void RHIBlockUntilGPUIdle() final override;
@@ -456,8 +456,8 @@ public:
 	virtual void RHIBeginRenderQuery(FRHIRenderQuery* RenderQuery) final override;
 	virtual void RHIEndRenderQuery(FRHIRenderQuery* RenderQuery) final override;
 	virtual void RHISubmitCommandsHint() final override;
-	virtual void RHIBeginDrawingViewport(FViewportRHIParamRef Viewport, FTextureRHIParamRef RenderTargetRHI) final override;
-	virtual void RHIEndDrawingViewport(FViewportRHIParamRef Viewport, bool bPresent, bool bLockToVsync) final override;
+	virtual void RHIBeginDrawingViewport(FRHIViewport* Viewport, FTextureRHIParamRef RenderTargetRHI) final override;
+	virtual void RHIEndDrawingViewport(FRHIViewport* Viewport, bool bPresent, bool bLockToVsync) final override;
 	virtual void RHIBeginFrame() final override;
 	virtual void RHIEndFrame() final override;
 	virtual void RHIBeginScene() final override;
@@ -530,12 +530,12 @@ public:
 	virtual void RHIAliasTextureResources(FTextureRHIParamRef DestTexture, FTextureRHIParamRef SrcTexture);
 
 	// Inline copy
-	virtual void RHICopyToStagingBuffer(FVertexBufferRHIParamRef SourceBufferRHI, FStagingBufferRHIParamRef DestinationStagingBufferRHI, uint32 InOffset, uint32 InNumBytes) final override;
+	virtual void RHICopyToStagingBuffer(FVertexBufferRHIParamRef SourceBufferRHI, FRHIStagingBuffer* DestinationStagingBufferRHI, uint32 InOffset, uint32 InNumBytes) final override;
 	virtual void RHIWriteGPUFence(FRHIGPUFence* FenceRHI) final override;
-	virtual void* RHILockStagingBuffer(FStagingBufferRHIParamRef StagingBuffer, uint32 Offset, uint32 SizeRHI) final override;
-	virtual void RHIUnlockStagingBuffer(FStagingBufferRHIParamRef StagingBuffer) final override;
-	virtual void* LockStagingBuffer_RenderThread(class FRHICommandListImmediate& RHICmdList, FStagingBufferRHIParamRef StagingBuffer, uint32 Offset, uint32 SizeRHI) final override;
-	virtual	void UnlockStagingBuffer_RenderThread(class FRHICommandListImmediate& RHICmdList, FStagingBufferRHIParamRef StagingBuffer) final override;
+	virtual void* RHILockStagingBuffer(FRHIStagingBuffer* StagingBuffer, uint32 Offset, uint32 SizeRHI) final override;
+	virtual void RHIUnlockStagingBuffer(FRHIStagingBuffer* StagingBuffer) final override;
+	virtual void* LockStagingBuffer_RenderThread(class FRHICommandListImmediate& RHICmdList, FRHIStagingBuffer* StagingBuffer, uint32 Offset, uint32 SizeRHI) final override;
+	virtual	void UnlockStagingBuffer_RenderThread(class FRHICommandListImmediate& RHICmdList, FRHIStagingBuffer* StagingBuffer) final override;
 
 	virtual FStagingBufferRHIRef RHICreateStagingBuffer() final override;
 	virtual FGPUFenceRHIRef RHICreateGPUFence(const FName &Name) final override;
diff --git a/Engine/Source/Runtime/RHI/Private/RHICommandList.cpp b/Engine/Source/Runtime/RHI/Private/RHICommandList.cpp
--- a/Engine/Source/Runtime/RHI/Private/RHICommandList.cpp
+++ b/Engine/Source/Runtime/RHI/Private/RHICommandList.cpp
@@ -1488,7 +1488,7 @@ void FRHICommandList::EndScene()
 }
 
 
-void FRHICommandList::BeginDrawingViewport(FViewportRHIParamRef Viewport, FTextureRHIParamRef RenderTargetRHI)
+void FRHICommandList::BeginDrawingViewport(FRHIViewport* Viewport, FTextureRHIParamRef RenderTargetRHI)
 {
 	check(IsImmediate() && IsInRenderingThread());
 	if (Bypass())
@@ -1505,7 +1505,7 @@ void FRHICommandList::BeginDrawingViewport(FViewportRHIParamRef Viewport, FTextu
 	}
 }
 
-void FRHICommandList::EndDrawingViewport(FViewportRHIParamRef Viewport, bool bPresent, bool bLockToVsync)
+void FRHICommandList::EndDrawingViewport(FRHIViewport* Viewport, bool bPresent, bool bLockToVsync)
 {
 	check(IsImmediate() && IsInRenderingThread());
 	if (Bypass())
@@ -2088,26 +2088,26 @@ void FDynamicRHI::UnlockIndexBuffer_RenderThread(class FRHICommandListImmediate&
 }
 
 // @todo-mattc-staging Default implementation
-void* FDynamicRHI::RHILockStagingBuffer(FStagingBufferRHIParamRef StagingBuffer, uint32 Offset, uint32 SizeRHI)
+void* FDynamicRHI::RHILockStagingBuffer(FRHIStagingBuffer* StagingBuffer, uint32 Offset, uint32 SizeRHI)
 {
 	check(false);
 	return nullptr;
 	//return GDynamicRHI->RHILockVertexBuffer(StagingBuffer->GetSourceBuffer(), Offset, SizeRHI, RLM_ReadOnly);
 }
-void FDynamicRHI::RHIUnlockStagingBuffer(FStagingBufferRHIParamRef StagingBuffer)
+void FDynamicRHI::RHIUnlockStagingBuffer(FRHIStagingBuffer* StagingBuffer)
 {
 	check(false);
 	//GDynamicRHI->RHIUnlockVertexBuffer(StagingBuffer->GetSourceBuffer());
 }
 
-void* FDynamicRHI::LockStagingBuffer_RenderThread(class FRHICommandListImmediate& RHICmdList, FStagingBufferRHIParamRef StagingBuffer, uint32 Offset, uint32 SizeRHI)
+void* FDynamicRHI::LockStagingBuffer_RenderThread(class FRHICommandListImmediate& RHICmdList, FRHIStagingBuffer* StagingBuffer, uint32 Offset, uint32 SizeRHI)
 {
 	QUICK_SCOPE_CYCLE_COUNTER(STAT_FDynamicRHI_LockStagingBuffer_RenderThread);
 	check(IsInRenderingThread());
 	RHICmdList.ImmediateFlush(EImmediateFlushType::FlushRHIThread);
 	return GDynamicRHI->RHILockStagingBuffer(StagingBuffer, Offset, SizeRHI);
 }
-void FDynamicRHI::UnlockStagingBuffer_RenderThread(class FRHICommandListImmediate& RHICmdList, FStagingBufferRHIParamRef StagingBuffer)
+void FDynamicRHI::UnlockStagingBuffer_RenderThread(class FRHICommandListImmediate& RHICmdList, FRHIStagingBuffer* StagingBuffer)
 {
 	QUICK_SCOPE_CYCLE_COUNTER(STAT_FDynamicRHI_UnlockStagingBuffer_RenderThread);
 	check(IsInRenderingThread());
diff --git a/Engine/Source/Runtime/RHI/Public/DynamicRHI.h b/Engine/Source/Runtime/RHI/Public/DynamicRHI.h
--- a/Engine/Source/Runtime/RHI/Public/DynamicRHI.h
+++ b/Engine/Source/Runtime/RHI/Public/DynamicRHI.h
@@ -78,7 +78,7 @@ struct FRHIFlipDetails
 
 struct FRayTracingGeometryInstance
 {
-	FRayTracingGeometryRHIParamRef GeometryRHI = nullptr;
+	FRHIRayTracingGeometry* GeometryRHI = nullptr;
 	FMatrix Transform = FMatrix(EForceInit::ForceInitToZero);
 	uint32 UserData = 0;
 	uint8 Mask = 0xFF;
@@ -355,13 +355,13 @@ public:
 	 * @param SizeRHI The length of the region in the buffer to lock.
 	 * @returns A pointer to the data starting at 'Offset' and of length 'SizeRHI' from 'StagingBuffer', or nullptr when there is an error.
 	 */
-	virtual void* RHILockStagingBuffer(FStagingBufferRHIParamRef StagingBuffer, uint32 Offset, uint32 SizeRHI);
+	virtual void* RHILockStagingBuffer(FRHIStagingBuffer* StagingBuffer, uint32 Offset, uint32 SizeRHI);
 
 	/**
 	 * Unlock a staging buffer previously locked with RHILockStagingBuffer.
 	 * @param StagingBuffer The buffer that was previously locked.
 	 */
-	virtual void RHIUnlockStagingBuffer(FStagingBufferRHIParamRef StagingBuffer);
+	virtual void RHIUnlockStagingBuffer(FRHIStagingBuffer* StagingBuffer);
 
 	/**
 	 * Lock a staging buffer to read contents on the CPU that were written by the GPU, without having to stall.
@@ -372,14 +372,14 @@ public:
 	 * @param SizeRHI The length of the region in the buffer to lock.
 	 * @returns A pointer to the data starting at 'Offset' and of length 'SizeRHI' from 'StagingBuffer', or nullptr when there is an error.
 	 */
-	virtual void* LockStagingBuffer_RenderThread(class FRHICommandListImmediate& RHICmdList, FStagingBufferRHIParamRef StagingBuffer, uint32 Offset, uint32 SizeRHI);
+	virtual void* LockStagingBuffer_RenderThread(class FRHICommandListImmediate& RHICmdList, FRHIStagingBuffer* StagingBuffer, uint32 Offset, uint32 SizeRHI);
 
 	/**
 	 * Unlock a staging buffer previously locked with LockStagingBuffer_RenderThread.
 	 * @param RHICmdList The command-list to execute on or synchronize with.
 	 * @param StagingBuffer The buffer what was previously locked.
 	 */
-	virtual void UnlockStagingBuffer_RenderThread(class FRHICommandListImmediate& RHICmdList, FStagingBufferRHIParamRef StagingBuffer);
+	virtual void UnlockStagingBuffer_RenderThread(class FRHICommandListImmediate& RHICmdList, FRHIStagingBuffer* StagingBuffer);
 
 	/**
 	* Creates a bound shader state instance which encapsulates a decl, vertex shader, hull shader, domain shader and pixel shader
@@ -916,16 +916,16 @@ public:
 	virtual bool RHIGetRenderQueryResult(FRHIRenderQuery* RenderQuery, uint64& OutResult, bool bWait) = 0;
 
 	// FlushType: Thread safe
-	virtual uint32 RHIGetViewportNextPresentGPUIndex(FViewportRHIParamRef Viewport)
+	virtual uint32 RHIGetViewportNextPresentGPUIndex(FRHIViewport* Viewport)
 	{
 		return 0; // By default, viewport need to be rendered on GPU0.
 	}
 
 	// With RHI thread, this is the current backbuffer from the perspective of the render thread.
 	// FlushType: Thread safe
-	virtual FTexture2DRHIRef RHIGetViewportBackBuffer(FViewportRHIParamRef Viewport) = 0;
+	virtual FTexture2DRHIRef RHIGetViewportBackBuffer(FRHIViewport* Viewport) = 0;
 
-	virtual FUnorderedAccessViewRHIRef RHIGetViewportBackBufferUAV(FViewportRHIParamRef ViewportRHI)
+	virtual FUnorderedAccessViewRHIRef RHIGetViewportBackBufferUAV(FRHIViewport* ViewportRHI)
 	{
 		return FUnorderedAccessViewRHIRef();
 	}
@@ -937,7 +937,7 @@ public:
 
 	// Only relevant with an RHI thread, this advances the backbuffer for the purpose of GetViewportBackBuffer
 	// FlushType: Thread safe
-	virtual void RHIAdvanceFrameForGetViewportBackBuffer(FViewportRHIParamRef Viewport) = 0;
+	virtual void RHIAdvanceFrameForGetViewportBackBuffer(FRHIViewport* Viewport) = 0;
 
 	/*
 	* Acquires or releases ownership of the platform-specific rendering context for the calling thread
@@ -964,9 +964,9 @@ public:
 
 	//  must be called from the main thread.
 	// FlushType: Thread safe
-	virtual void RHIResizeViewport(FViewportRHIParamRef Viewport, uint32 SizeX, uint32 SizeY, bool bIsFullscreen) = 0;
+	virtual void RHIResizeViewport(FRHIViewport* Viewport, uint32 SizeX, uint32 SizeY, bool bIsFullscreen) = 0;
 
-	virtual void RHIResizeViewport(FViewportRHIParamRef Viewport, uint32 SizeX, uint32 SizeY, bool bIsFullscreen, EPixelFormat PreferredPixelFormat)
+	virtual void RHIResizeViewport(FRHIViewport* Viewport, uint32 SizeX, uint32 SizeY, bool bIsFullscreen, EPixelFormat PreferredPixelFormat)
 	{
 		// Default implementation for RHIs that cannot change formats on the fly
 		RHIResizeViewport(Viewport, SizeX, SizeY, bIsFullscreen);
@@ -1325,12 +1325,12 @@ FORCEINLINE bool RHIGetRenderQueryResult(FRHIRenderQuery* RenderQuery, uint64& O
 	return GDynamicRHI->RHIGetRenderQueryResult(RenderQuery, OutResult, bWait);
 }
 
-FORCEINLINE uint32 RHIGetViewportNextPresentGPUIndex(FViewportRHIParamRef Viewport)
+FORCEINLINE uint32 RHIGetViewportNextPresentGPUIndex(FRHIViewport* Viewport)
 {
 	return GDynamicRHI->RHIGetViewportNextPresentGPUIndex(Viewport);
 }
 
-FORCEINLINE FTexture2DRHIRef RHIGetViewportBackBuffer(FViewportRHIParamRef Viewport)
+FORCEINLINE FTexture2DRHIRef RHIGetViewportBackBuffer(FRHIViewport* Viewport)
 {
 	return GDynamicRHI->RHIGetViewportBackBuffer(Viewport);
 }
@@ -1340,7 +1340,7 @@ FORCEINLINE FTexture2DRHIRef RHIGetFMaskTexture(FTextureRHIParamRef SourceTextur
 	return GDynamicRHI->RHIGetFMaskTexture(SourceTextureRHI);
 }
 
-FORCEINLINE void RHIAdvanceFrameForGetViewportBackBuffer(FViewportRHIParamRef Viewport)
+FORCEINLINE void RHIAdvanceFrameForGetViewportBackBuffer(FRHIViewport* Viewport)
 {
 	return GDynamicRHI->RHIAdvanceFrameForGetViewportBackBuffer(Viewport);
 }
@@ -1355,7 +1355,7 @@ FORCEINLINE FViewportRHIRef RHICreateViewport(void* WindowHandle, uint32 SizeX,
 	return GDynamicRHI->RHICreateViewport(WindowHandle, SizeX, SizeY, bIsFullscreen, PreferredPixelFormat);
 }
 
-FORCEINLINE void RHIResizeViewport(FViewportRHIParamRef Viewport, uint32 SizeX, uint32 SizeY, bool bIsFullscreen, EPixelFormat PreferredPixelFormat)
+FORCEINLINE void RHIResizeViewport(FRHIViewport* Viewport, uint32 SizeX, uint32 SizeY, bool bIsFullscreen, EPixelFormat PreferredPixelFormat)
 {
 	GDynamicRHI->RHIResizeViewport(Viewport, SizeX, SizeY, bIsFullscreen, PreferredPixelFormat);
 }
diff --git a/Engine/Source/Runtime/RHI/Public/RHICommandList.h b/Engine/Source/Runtime/RHI/Public/RHICommandList.h
--- a/Engine/Source/Runtime/RHI/Public/RHICommandList.h
+++ b/Engine/Source/Runtime/RHI/Public/RHICommandList.h
@@ -1413,11 +1413,11 @@ template<ECmdList CmdListType>
 struct FRHICommandCopyToStagingBuffer final : public FRHICommand<FRHICommandCopyToStagingBuffer<CmdListType>>
 {
 	FVertexBufferRHIParamRef SourceBuffer;
-	FStagingBufferRHIParamRef DestinationStagingBuffer;
+	FRHIStagingBuffer* DestinationStagingBuffer;
 	uint32 Offset;
 	uint32 NumBytes;
 
-	FORCEINLINE_DEBUGGABLE FRHICommandCopyToStagingBuffer(FVertexBufferRHIParamRef InSourceBuffer, FStagingBufferRHIParamRef InDestinationStagingBuffer, uint32 InOffset, uint32 InNumBytes)
+	FORCEINLINE_DEBUGGABLE FRHICommandCopyToStagingBuffer(FVertexBufferRHIParamRef InSourceBuffer, FRHIStagingBuffer* InDestinationStagingBuffer, uint32 InOffset, uint32 InNumBytes)
 		: SourceBuffer(InSourceBuffer)
 		, DestinationStagingBuffer(InDestinationStagingBuffer)
 		, Offset(InOffset)
@@ -1667,10 +1667,10 @@ struct FRHICommandEndFrame final : public FRHICommand<FRHICommandEndFrame>
 
 struct FRHICommandBeginDrawingViewport final : public FRHICommand<FRHICommandBeginDrawingViewport>
 {
-	FViewportRHIParamRef Viewport;
+	FRHIViewport* Viewport;
 	FTextureRHIParamRef RenderTargetRHI;
 
-	FORCEINLINE_DEBUGGABLE FRHICommandBeginDrawingViewport(FViewportRHIParamRef InViewport, FTextureRHIParamRef InRenderTargetRHI)
+	FORCEINLINE_DEBUGGABLE FRHICommandBeginDrawingViewport(FRHIViewport* InViewport, FTextureRHIParamRef InRenderTargetRHI)
 		: Viewport(InViewport)
 		, RenderTargetRHI(InRenderTargetRHI)
 	{
@@ -1680,11 +1680,11 @@ struct FRHICommandBeginDrawingViewport final : public FRHICommand<FRHICommandBeg
 
 struct FRHICommandEndDrawingViewport final : public FRHICommand<FRHICommandEndDrawingViewport>
 {
-	FViewportRHIParamRef Viewport;
+	FRHIViewport* Viewport;
 	bool bPresent;
 	bool bLockToVsync;
 
-	FORCEINLINE_DEBUGGABLE FRHICommandEndDrawingViewport(FViewportRHIParamRef InViewport, bool InbPresent, bool InbLockToVsync)
+	FORCEINLINE_DEBUGGABLE FRHICommandEndDrawingViewport(FRHIViewport* InViewport, bool InbPresent, bool InbLockToVsync)
 		: Viewport(InViewport)
 		, bPresent(InbPresent)
 		, bLockToVsync(InbLockToVsync)
@@ -1864,15 +1864,15 @@ struct FRHICommandCopyBufferRegions final : public FRHICommand<FRHICommandCopyBu
 
 struct FRHICommandBuildAccelerationStructure final : public FRHICommand<FRHICommandBuildAccelerationStructure>
 {
-	FRayTracingGeometryRHIParamRef Geometry;
-	FRayTracingSceneRHIParamRef Scene;
+	FRHIRayTracingGeometry* Geometry;
+	FRHIRayTracingScene* Scene;
 
-	explicit FRHICommandBuildAccelerationStructure(FRayTracingGeometryRHIParamRef InGeometry)
+	explicit FRHICommandBuildAccelerationStructure(FRHIRayTracingGeometry* InGeometry)
 		: Geometry(InGeometry)
 		, Scene(nullptr)
 	{}
 
-	explicit FRHICommandBuildAccelerationStructure(FRayTracingSceneRHIParamRef InScene)
+	explicit FRHICommandBuildAccelerationStructure(FRHIRayTracingScene* InScene)
 		: Geometry(nullptr)
 		, Scene(InScene)
 	{}
@@ -1904,12 +1904,12 @@ struct FRHICommandBuildAccelerationStructures final : public FRHICommand<FRHICom
 
 struct FRHICommandRayTraceOcclusion final : public FRHICommand<FRHICommandRayTraceOcclusion>
 {
-	FRayTracingSceneRHIParamRef Scene;
+	FRHIRayTracingScene* Scene;
 	FRHIShaderResourceView* Rays;
 	FRHIUnorderedAccessView* Output;
 	uint32 NumRays;
 
-	FRHICommandRayTraceOcclusion(FRayTracingSceneRHIParamRef InScene,
+	FRHICommandRayTraceOcclusion(FRHIRayTracingScene* InScene,
 		FRHIShaderResourceView* InRays,
 		FRHIUnorderedAccessView* InOutput,
 		uint32 InNumRays)
@@ -1924,12 +1924,12 @@ struct FRHICommandRayTraceOcclusion final : public FRHICommand<FRHICommandRayTra
 
 struct FRHICommandRayTraceIntersection final : public FRHICommand<FRHICommandRayTraceIntersection>
 {
-	FRayTracingSceneRHIParamRef Scene;
+	FRHIRayTracingScene* Scene;
 	FRHIShaderResourceView* Rays;
 	FRHIUnorderedAccessView* Output;
 	uint32 NumRays;
 
-	FRHICommandRayTraceIntersection(FRayTracingSceneRHIParamRef InScene,
+	FRHICommandRayTraceIntersection(FRHIRayTracingScene* InScene,
 		FRHIShaderResourceView* InRays,
 		FRHIUnorderedAccessView* InOutput,
 		uint32 InNumRays)
@@ -1945,13 +1945,13 @@ struct FRHICommandRayTraceIntersection final : public FRHICommand<FRHICommandRay
 struct FRHICommandRayTraceDispatch final : public FRHICommand<FRHICommandRayTraceDispatch>
 {
 	FRayTracingPipelineState* Pipeline;
-	FRayTracingSceneRHIParamRef Scene;
+	FRHIRayTracingScene* Scene;
 	FRayTracingShaderBindings GlobalResourceBindings;
 	FRHIRayTracingShader* RayGenShader;
 	uint32 Width;
 	uint32 Height;
 
-	FRHICommandRayTraceDispatch(FRayTracingPipelineState* InPipeline, FRHIRayTracingShader* InRayGenShader, FRayTracingSceneRHIParamRef InScene, const FRayTracingShaderBindings& InGlobalResourceBindings, uint32 InWidth, uint32 InHeight)
+	FRHICommandRayTraceDispatch(FRayTracingPipelineState* InPipeline, FRHIRayTracingShader* InRayGenShader, FRHIRayTracingScene* InScene, const FRayTracingShaderBindings& InGlobalResourceBindings, uint32 InWidth, uint32 InHeight)
 		: Pipeline(InPipeline)
 		, Scene(InScene)
 		, GlobalResourceBindings(InGlobalResourceBindings)
@@ -1971,7 +1971,7 @@ struct FRHICommandSetRayTracingBindings final : public FRHICommand<FRHICommandSe
 		EBindingType_CallableShader,
 	};
 
-	FRayTracingSceneRHIParamRef Scene;
+	FRHIRayTracingScene* Scene;
 	uint32 InstanceIndex;
 	uint32 SegmentIndex;
 	uint32 ShaderSlot;
@@ -1985,7 +1985,7 @@ struct FRHICommandSetRayTracingBindings final : public FRHICommand<FRHICommandSe
 	EBindingType BindingType;
 
 	// Hit group bindings
-	FRHICommandSetRayTracingBindings(FRayTracingSceneRHIParamRef InScene, uint32 InInstanceIndex, uint32 InSegmentIndex, uint32 InShaderSlot,
+	FRHICommandSetRayTracingBindings(FRHIRayTracingScene* InScene, uint32 InInstanceIndex, uint32 InSegmentIndex, uint32 InShaderSlot,
 		FRayTracingPipelineState* InPipeline, uint32 InHitGroupIndex, uint32 InNumUniformBuffers, const FUniformBufferRHIParamRef* InUniformBuffers,
 		uint32 InLooseParameterDataSize, const void* InLooseParameterData,
 		uint32 InUserData)
@@ -2005,7 +2005,7 @@ struct FRHICommandSetRayTracingBindings final : public FRHICommand<FRHICommandSe
 	}
 
 	// Callable shader bindings
-	FRHICommandSetRayTracingBindings(FRayTracingSceneRHIParamRef InScene, uint32 InShaderSlot,
+	FRHICommandSetRayTracingBindings(FRHIRayTracingScene* InScene, uint32 InShaderSlot,
 		FRayTracingPipelineState* InPipeline, uint32 InShaderIndex,
 		uint32 InNumUniformBuffers, const FUniformBufferRHIParamRef* InUniformBuffers,
 		uint32 InUserData)
@@ -2778,7 +2778,7 @@ public:
 		ALLOC_COMMAND(FRHICommandWaitComputeFence<ECmdList::EGfx>)(WaitFence);
 	}
 
-	FORCEINLINE_DEBUGGABLE void CopyToStagingBuffer(FVertexBufferRHIParamRef SourceBuffer, FStagingBufferRHIParamRef DestinationStagingBuffer, uint32 Offset, uint32 NumBytes)
+	FORCEINLINE_DEBUGGABLE void CopyToStagingBuffer(FVertexBufferRHIParamRef SourceBuffer, FRHIStagingBuffer* DestinationStagingBuffer, uint32 Offset, uint32 NumBytes)
 	{
 		if (Bypass())
 		{
@@ -2892,8 +2892,8 @@ public:
 	// These 6 are special in that they must be called on the immediate command list and they force a flush only when we are not doing RHI thread
 	void BeginScene();
 	void EndScene();
-	void BeginDrawingViewport(FViewportRHIParamRef Viewport, FTextureRHIParamRef RenderTargetRHI);
-	void EndDrawingViewport(FViewportRHIParamRef Viewport, bool bPresent, bool bLockToVsync);
+	void BeginDrawingViewport(FRHIViewport* Viewport, FTextureRHIParamRef RenderTargetRHI);
+	void EndDrawingViewport(FRHIViewport* Viewport, bool bPresent, bool bLockToVsync);
 	void BeginFrame();
 	void EndFrame();
 
@@ -2985,7 +2985,7 @@ public:
 		}
 	}
 
-	FORCEINLINE_DEBUGGABLE void BuildAccelerationStructure(FRayTracingGeometryRHIParamRef Geometry)
+	FORCEINLINE_DEBUGGABLE void BuildAccelerationStructure(FRHIRayTracingGeometry* Geometry)
 	{
 		if (Bypass())
 		{
@@ -3021,7 +3021,7 @@ public:
 		}
 	}
 
-	FORCEINLINE_DEBUGGABLE void BuildAccelerationStructure(FRayTracingSceneRHIParamRef Scene)
+	FORCEINLINE_DEBUGGABLE void BuildAccelerationStructure(FRHIRayTracingScene* Scene)
 	{
 		if (Bypass())
 		{
@@ -3038,7 +3038,7 @@ public:
 	 * Binary intersection results are written to output buffer as R32_UINTs.
 	 * 0xFFFFFFFF is written if ray intersects any scene triangle, 0 otherwise.
 	 */
-	FORCEINLINE_DEBUGGABLE void RayTraceOcclusion(FRayTracingSceneRHIParamRef Scene,
+	FORCEINLINE_DEBUGGABLE void RayTraceOcclusion(FRHIRayTracingScene* Scene,
 		FRHIShaderResourceView* Rays,
 		FRHIUnorderedAccessView* Output,
 		uint32 NumRays)
@@ -3057,7 +3057,7 @@ public:
 	 * Trace rays from an input buffer of FBasicRayData.
 	 * Primitive intersection results are written to output buffer as FBasicRayIntersectionData.
 	 */
-	FORCEINLINE_DEBUGGABLE void RayTraceIntersection(FRayTracingSceneRHIParamRef Scene,
+	FORCEINLINE_DEBUGGABLE void RayTraceIntersection(FRHIRayTracingScene* Scene,
 		FRHIShaderResourceView* Rays,
 		FRHIUnorderedAccessView* Output,
 		uint32 NumRays)
@@ -3072,7 +3072,7 @@ public:
 		}
 	}
 
-	FORCEINLINE_DEBUGGABLE void RayTraceDispatch(FRayTracingPipelineState* Pipeline, FRHIRayTracingShader* RayGenShader, FRayTracingSceneRHIParamRef Scene, const FRayTracingShaderBindings& GlobalResourceBindings, uint32 Width, uint32 Height)
+	FORCEINLINE_DEBUGGABLE void RayTraceDispatch(FRayTracingPipelineState* Pipeline, FRHIRayTracingShader* RayGenShader, FRHIRayTracingScene* Scene, const FRayTracingShaderBindings& GlobalResourceBindings, uint32 Width, uint32 Height)
 	{
 		if (Bypass())
 		{
@@ -3086,7 +3086,7 @@ public:
 	}
 
 	FORCEINLINE_DEBUGGABLE void SetRayTracingHitGroup(
-		FRayTracingSceneRHIParamRef Scene, uint32 InstanceIndex, uint32 SegmentIndex, uint32 ShaderSlot,
+		FRHIRayTracingScene* Scene, uint32 InstanceIndex, uint32 SegmentIndex, uint32 ShaderSlot,
 		FRayTracingPipelineState* Pipeline, uint32 HitGroupIndex,
 		uint32 NumUniformBuffers, const FUniformBufferRHIParamRef* UniformBuffers,
 		uint32 LooseParameterDataSize, const void* LooseParameterData,
@@ -3127,7 +3127,7 @@ public:
 	}
 
 	FORCEINLINE_DEBUGGABLE void SetRayTracingCallableShader(
-		FRayTracingSceneRHIParamRef Scene, uint32 ShaderSlotInScene,
+		FRHIRayTracingScene* Scene, uint32 ShaderSlotInScene,
 		FRayTracingPipelineState* Pipeline, uint32 ShaderIndexInPipeline,
 		uint32 NumUniformBuffers, const FUniformBufferRHIParamRef* UniformBuffers,
 		uint32 UserData)
@@ -3407,7 +3407,7 @@ public:
 		ALLOC_COMMAND(FRHICommandWaitComputeFence<ECmdList::ECompute>)(WaitFence);
 	}
 
-	FORCEINLINE_DEBUGGABLE void CopyToStagingBuffer(FVertexBufferRHIParamRef SourceBuffer, FStagingBufferRHIParamRef DestinationStagingBuffer, uint32 Offset, uint32 NumBytes)
+	FORCEINLINE_DEBUGGABLE void CopyToStagingBuffer(FVertexBufferRHIParamRef SourceBuffer, FRHIStagingBuffer* DestinationStagingBuffer, uint32 Offset, uint32 NumBytes)
 	{
 		if (Bypass())
 		{
@@ -3686,12 +3686,12 @@ public:
 		GDynamicRHI->UnlockIndexBuffer_RenderThread(*this, IndexBuffer);
 	}
 	
-	FORCEINLINE void* LockStagingBuffer(FStagingBufferRHIParamRef StagingBuffer, uint32 Offset, uint32 SizeRHI)
+	FORCEINLINE void* LockStagingBuffer(FRHIStagingBuffer* StagingBuffer, uint32 Offset, uint32 SizeRHI)
 	{
 		return GDynamicRHI->LockStagingBuffer_RenderThread(*this, StagingBuffer, Offset, SizeRHI);
 	}
 	
-	FORCEINLINE void UnlockStagingBuffer(FStagingBufferRHIParamRef StagingBuffer)
+	FORCEINLINE void UnlockStagingBuffer(FRHIStagingBuffer* StagingBuffer)
 	{
 		GDynamicRHI->UnlockStagingBuffer_RenderThread(*this, StagingBuffer);
 	}
@@ -4167,17 +4167,17 @@ public:
 		return RHIGetRenderQueryResult(RenderQuery, OutResult, bWait);
 	}
 	
-	FORCEINLINE uint32 GetViewportNextPresentGPUIndex(FViewportRHIParamRef Viewport)
+	FORCEINLINE uint32 GetViewportNextPresentGPUIndex(FRHIViewport* Viewport)
 	{
 		return GDynamicRHI->RHIGetViewportNextPresentGPUIndex(Viewport);
 	}
 
-	FORCEINLINE FTexture2DRHIRef GetViewportBackBuffer(FViewportRHIParamRef Viewport)
+	FORCEINLINE FTexture2DRHIRef GetViewportBackBuffer(FRHIViewport* Viewport)
 	{
 		return RHIGetViewportBackBuffer(Viewport);
 	}
 	
-	FORCEINLINE void AdvanceFrameForGetViewportBackBuffer(FViewportRHIParamRef Viewport)
+	FORCEINLINE void AdvanceFrameForGetViewportBackBuffer(FRHIViewport* Viewport)
 	{
 		return RHIAdvanceFrameForGetViewportBackBuffer(Viewport);
 	}
@@ -4217,7 +4217,7 @@ public:
 		return RHICreateViewport(WindowHandle, SizeX, SizeY, bIsFullscreen, PreferredPixelFormat);
 	}
 	
-	FORCEINLINE void ResizeViewport(FViewportRHIParamRef Viewport, uint32 SizeX, uint32 SizeY, bool bIsFullscreen, EPixelFormat PreferredPixelFormat)
+	FORCEINLINE void ResizeViewport(FRHIViewport* Viewport, uint32 SizeX, uint32 SizeY, bool bIsFullscreen, EPixelFormat PreferredPixelFormat)
 	{
 		LLM_SCOPE(ELLMTag::RenderTargets);
 		RHIResizeViewport(Viewport, SizeX, SizeY, bIsFullscreen, PreferredPixelFormat);
@@ -4900,12 +4900,12 @@ FORCEINLINE FRHIShaderLibraryRef RHICreateShaderLibrary(EShaderPlatform Platform
     return GDynamicRHI->RHICreateShaderLibrary(Platform, FilePath, Name);
 }
 
-FORCEINLINE void* RHILockStagingBuffer(FStagingBufferRHIParamRef StagingBuffer, uint32 Offset, uint32 Size)
+FORCEINLINE void* RHILockStagingBuffer(FRHIStagingBuffer* StagingBuffer, uint32 Offset, uint32 Size)
 {
 	return FRHICommandListExecutor::GetImmediateCommandList().LockStagingBuffer(StagingBuffer, Offset, Size);
 }
 
-FORCEINLINE void RHIUnlockStagingBuffer(FStagingBufferRHIParamRef StagingBuffer)
+FORCEINLINE void RHIUnlockStagingBuffer(FRHIStagingBuffer* StagingBuffer)
 {
 	 FRHICommandListExecutor::GetImmediateCommandList().UnlockStagingBuffer(StagingBuffer);
 }
diff --git a/Engine/Source/Runtime/RHI/Public/RHIContext.h b/Engine/Source/Runtime/RHI/Public/RHIContext.h
--- a/Engine/Source/Runtime/RHI/Public/RHIContext.h
+++ b/Engine/Source/Runtime/RHI/Public/RHIContext.h
@@ -126,7 +126,7 @@ public:
 	 * @param Offset The start of the data in 'SourceBuffer'
 	 * @param NumBytes The number of bytes to copy out of 'SourceBuffer'
 	 */
-	virtual void RHICopyToStagingBuffer(FVertexBufferRHIParamRef SourceBufferRHI, FStagingBufferRHIParamRef DestinationStagingBufferRHI, uint32 InOffset, uint32 InNumBytes)
+	virtual void RHICopyToStagingBuffer(FVertexBufferRHIParamRef SourceBufferRHI, FRHIStagingBuffer* DestinationStagingBufferRHI, uint32 InOffset, uint32 InNumBytes)
 	{
 		check(false);
 	}
@@ -259,10 +259,10 @@ public:
 	virtual void RHIDiscardRenderTargets(bool Depth, bool Stencil, uint32 ColorBitMask) {}
 
 	// This method is queued with an RHIThread, otherwise it will flush after it is queued; without an RHI thread there is no benefit to queuing this frame advance commands
-	virtual void RHIBeginDrawingViewport(FViewportRHIParamRef Viewport, FTextureRHIParamRef RenderTargetRHI) = 0;
+	virtual void RHIBeginDrawingViewport(FRHIViewport* Viewport, FTextureRHIParamRef RenderTargetRHI) = 0;
 
 	// This method is queued with an RHIThread, otherwise it will flush after it is queued; without an RHI thread there is no benefit to queuing this frame advance commands
-	virtual void RHIEndDrawingViewport(FViewportRHIParamRef Viewport, bool bPresent, bool bLockToVsync) = 0;
+	virtual void RHIEndDrawingViewport(FRHIViewport* Viewport, bool bPresent, bool bLockToVsync) = 0;
 
 	// This method is queued with an RHIThread, otherwise it will flush after it is queued; without an RHI thread there is no benefit to queuing this frame advance commands
 	virtual void RHIBeginFrame() = 0;
@@ -579,7 +579,7 @@ public:
 		checkNoEntry();
 	}
 #endif
-	virtual void RHIBuildAccelerationStructure(FRayTracingGeometryRHIParamRef Geometry)
+	virtual void RHIBuildAccelerationStructure(FRHIRayTracingGeometry* Geometry)
 	{
 		checkNoEntry();
 	}
@@ -594,12 +594,12 @@ public:
 		checkNoEntry();
 	}
 
-	virtual void RHIBuildAccelerationStructure(FRayTracingSceneRHIParamRef Scene)
+	virtual void RHIBuildAccelerationStructure(FRHIRayTracingScene* Scene)
 	{
 		checkNoEntry();
 	}
 
-	virtual void RHIRayTraceOcclusion(FRayTracingSceneRHIParamRef Scene,
+	virtual void RHIRayTraceOcclusion(FRHIRayTracingScene* Scene,
 		FRHIShaderResourceView* Rays,
 		FRHIUnorderedAccessView* Output,
 		uint32 NumRays)
@@ -607,7 +607,7 @@ public:
 		checkNoEntry();
 	}
 
-	virtual void RHIRayTraceIntersection(FRayTracingSceneRHIParamRef Scene,
+	virtual void RHIRayTraceIntersection(FRHIRayTracingScene* Scene,
 		FRHIShaderResourceView* Rays,
 		FRHIUnorderedAccessView* Output,
 		uint32 NumRays)
@@ -616,7 +616,7 @@ public:
 	}
 
 	virtual void RHIRayTraceDispatch(FRHIRayTracingPipelineState* RayTracingPipelineState, FRHIRayTracingShader* RayGenShader,
-		FRayTracingSceneRHIParamRef Scene, 
+		FRHIRayTracingScene* Scene,
 		const FRayTracingShaderBindings& GlobalResourceBindings,
 		uint32 Width, uint32 Height)
 	{
@@ -624,7 +624,7 @@ public:
 	}
 
 	virtual void RHISetRayTracingHitGroup(
-		FRayTracingSceneRHIParamRef Scene, uint32 InstanceIndex, uint32 SegmentIndex, uint32 ShaderSlot,
+		FRHIRayTracingScene* Scene, uint32 InstanceIndex, uint32 SegmentIndex, uint32 ShaderSlot,
 		FRHIRayTracingPipelineState* Pipeline, uint32 HitGroupIndex,
 		uint32 NumUniformBuffers, const FUniformBufferRHIParamRef* UniformBuffers,
 		uint32 LooseParameterDataSize, const void* LooseParameterData,
@@ -634,7 +634,7 @@ public:
 	}
 
 	virtual void RHISetRayTracingCallableShader(
-		FRayTracingSceneRHIParamRef Scene, uint32 ShaderSlotInScene,
+		FRHIRayTracingScene* Scene, uint32 ShaderSlotInScene,
 		FRHIRayTracingPipelineState* Pipeline, uint32 ShaderIndexInPipeline,
 		uint32 NumUniformBuffers, const FUniformBufferRHIParamRef* UniformBuffers,
 		uint32 UserData)
diff --git a/Engine/Source/Runtime/RHI/Public/RHIResources.h b/Engine/Source/Runtime/RHI/Public/RHIResources.h
--- a/Engine/Source/Runtime/RHI/Public/RHIResources.h
+++ b/Engine/Source/Runtime/RHI/Public/RHIResources.h
@@ -1188,6 +1188,7 @@ UE_DEPRECATED(4.23, "FGPUFenceRHIParamRef typedef is deprecated; please use FRHI
 typedef FRHIGPUFence*				FGPUFenceRHIParamRef;
 typedef TRefCountPtr<FRHIGPUFence>	FGPUFenceRHIRef;
 
+UE_DEPRECATED(4.23, "FViewportRHIParamRef typedef is deprecated; please use FRHIViewport* directly instead.")
 typedef FRHIViewport*              FViewportRHIParamRef;
 typedef TRefCountPtr<FRHIViewport> FViewportRHIRef;
 
@@ -1215,6 +1216,7 @@ typedef TRefCountPtr<FRHIRayTracingPipelineState> FRayTracingPipelineStateRHIRef
 /** Bottom level ray tracing acceleration structure (contains triangles). */
 class FRHIRayTracingGeometry : public FRHIResource {};
 
+UE_DEPRECATED(4.23, "FRayTracingGeometryRHIParamRef typedef is deprecated; please use FRHIRayTracingGeometry* directly instead.")
 typedef FRHIRayTracingGeometry*                  FRayTracingGeometryRHIParamRef;
 typedef TRefCountPtr<FRHIRayTracingGeometry>     FRayTracingGeometryRHIRef;
 
@@ -1227,6 +1229,7 @@ protected:
 	FShaderResourceViewRHIRef ShaderResourceView;
 };
 
+UE_DEPRECATED(4.23, "FRayTracingSceneRHIParamRef typedef is deprecated; please use FRHIRayTracingScene* directly instead.")
 typedef FRHIRayTracingScene*                     FRayTracingSceneRHIParamRef;
 typedef TRefCountPtr<FRHIRayTracingScene>        FRayTracingSceneRHIRef;
 
@@ -1264,6 +1267,7 @@ public:
 	uint32 Offset;
 };
 
+UE_DEPRECATED(4.23, "FStagingBufferRHIParamRef typedef is deprecated; please use FRHIStagingBuffer* directly instead.")
 typedef FRHIStagingBuffer*				FStagingBufferRHIParamRef;
 typedef TRefCountPtr<FRHIStagingBuffer>	FStagingBufferRHIRef;
 
diff --git a/Engine/Source/Runtime/RHI/Public/RHIValidationContext.h b/Engine/Source/Runtime/RHI/Public/RHIValidationContext.h
--- a/Engine/Source/Runtime/RHI/Public/RHIValidationContext.h
+++ b/Engine/Source/Runtime/RHI/Public/RHIValidationContext.h
@@ -152,7 +152,7 @@ public:
 		RHIContext->RHIWriteGPUFence(FenceRHI);
 	}
 
-	virtual void RHICopyToStagingBuffer(FVertexBufferRHIParamRef SourceBufferRHI, FStagingBufferRHIParamRef DestinationStagingBufferRHI, uint32 InOffset, uint32 InNumBytes) override final
+	virtual void RHICopyToStagingBuffer(FVertexBufferRHIParamRef SourceBufferRHI, FRHIStagingBuffer* DestinationStagingBufferRHI, uint32 InOffset, uint32 InNumBytes) override final
 	{
 		RHIContext->RHICopyToStagingBuffer(SourceBufferRHI, DestinationStagingBufferRHI, InOffset, InNumBytes);
 	}
@@ -307,13 +307,13 @@ public:
 	}
 
 	// This method is queued with an RHIThread, otherwise it will flush after it is queued; without an RHI thread there is no benefit to queuing this frame advance commands
-	virtual void RHIBeginDrawingViewport(FViewportRHIParamRef Viewport, FTextureRHIParamRef RenderTargetRHI) override final
+	virtual void RHIBeginDrawingViewport(FRHIViewport* Viewport, FTextureRHIParamRef RenderTargetRHI) override final
 	{
 		RHIContext->RHIBeginDrawingViewport(Viewport, RenderTargetRHI);
 	}
 
 	// This method is queued with an RHIThread, otherwise it will flush after it is queued; without an RHI thread there is no benefit to queuing this frame advance commands
-	virtual void RHIEndDrawingViewport(FViewportRHIParamRef Viewport, bool bPresent, bool bLockToVsync) override final
+	virtual void RHIEndDrawingViewport(FRHIViewport* Viewport, bool bPresent, bool bLockToVsync) override final
 	{
 		RHIContext->RHIEndDrawingViewport(Viewport, bPresent, bLockToVsync);
 	}
@@ -792,7 +792,7 @@ public:
 		RHIContext->RHIWriteGPUFence(FenceRHI);
 	}
 
-	virtual void RHICopyToStagingBuffer(FVertexBufferRHIParamRef SourceBufferRHI, FStagingBufferRHIParamRef DestinationStagingBufferRHI, uint32 InOffset, uint32 InNumBytes) override final
+	virtual void RHICopyToStagingBuffer(FVertexBufferRHIParamRef SourceBufferRHI, FRHIStagingBuffer* DestinationStagingBufferRHI, uint32 InOffset, uint32 InNumBytes) override final
 	{
 		RHIContext->RHICopyToStagingBuffer(SourceBufferRHI, DestinationStagingBufferRHI, InOffset, InNumBytes);
 	}
@@ -804,7 +804,7 @@ public:
 		RHIContext->RHICopyTexture(SourceTexture, DestTexture, CopyInfo);
 	}
 
-	virtual void RHIBuildAccelerationStructure(FRayTracingGeometryRHIParamRef Geometry) override final
+	virtual void RHIBuildAccelerationStructure(FRHIRayTracingGeometry* Geometry) override final
 	{
 		RHIContext->RHIBuildAccelerationStructure(Geometry);
 	}
@@ -819,12 +819,12 @@ public:
 		RHIContext->RHIBuildAccelerationStructures(Params);
 	}
 
-	virtual void RHIBuildAccelerationStructure(FRayTracingSceneRHIParamRef Scene) override final
+	virtual void RHIBuildAccelerationStructure(FRHIRayTracingScene* Scene) override final
 	{
 		RHIContext->RHIBuildAccelerationStructure(Scene);
 	}
 
-	virtual void RHIRayTraceOcclusion(FRayTracingSceneRHIParamRef Scene,
+	virtual void RHIRayTraceOcclusion(FRHIRayTracingScene* Scene,
 		FRHIShaderResourceView* Rays,
 		FRHIUnorderedAccessView* Output,
 		uint32 NumRays) override final
@@ -832,7 +832,7 @@ public:
 		RHIContext->RHIRayTraceOcclusion(Scene, Rays, Output, NumRays);
 	}
 
-	virtual void RHIRayTraceIntersection(FRayTracingSceneRHIParamRef Scene,
+	virtual void RHIRayTraceIntersection(FRHIRayTracingScene* Scene,
 		FRHIShaderResourceView* Rays,
 		FRHIUnorderedAccessView* Output,
 		uint32 NumRays) override final
@@ -841,7 +841,7 @@ public:
 	}
 
 	virtual void RHIRayTraceDispatch(FRHIRayTracingPipelineState* RayTracingPipelineState, FRHIRayTracingShader* RayGenShader,
-		FRayTracingSceneRHIParamRef Scene,
+		FRHIRayTracingScene* Scene,
 		const FRayTracingShaderBindings& GlobalResourceBindings,
 		uint32 Width, uint32 Height) override final
 	{
@@ -849,7 +849,7 @@ public:
 	}
 
 	virtual void RHISetRayTracingHitGroup(
-		FRayTracingSceneRHIParamRef Scene, uint32 InstanceIndex, uint32 SegmentIndex, uint32 ShaderSlot,
+		FRHIRayTracingScene* Scene, uint32 InstanceIndex, uint32 SegmentIndex, uint32 ShaderSlot,
 		FRHIRayTracingPipelineState* Pipeline, uint32 HitGroupIndex,
 		uint32 NumUniformBuffers, const FUniformBufferRHIParamRef* UniformBuffers,
 		uint32 LooseParameterDataSize, const void* LooseParameterData,
diff --git a/Engine/Source/Runtime/Renderer/Private/MeshPassProcessor.cpp b/Engine/Source/Runtime/Renderer/Private/MeshPassProcessor.cpp
--- a/Engine/Source/Runtime/Renderer/Private/MeshPassProcessor.cpp
+++ b/Engine/Source/Runtime/Renderer/Private/MeshPassProcessor.cpp
@@ -269,7 +269,7 @@ void FMeshDrawShaderBindings::SetShaderBindings(
 #if RHI_RAYTRACING
 void FMeshDrawShaderBindings::SetRayTracingShaderBindingsForHitGroup(
 	FRHICommandList& RHICmdList, 
-	FRayTracingSceneRHIParamRef Scene, 
+	FRHIRayTracingScene* Scene,
 	uint32 InstanceIndex, 
 	uint32 SegmentIndex,
 	FRayTracingPipelineState* PipelineState,
diff --git a/Engine/Source/Runtime/Renderer/Private/PathTracing.cpp b/Engine/Source/Runtime/Renderer/Private/PathTracing.cpp
--- a/Engine/Source/Runtime/Renderer/Private/PathTracing.cpp
+++ b/Engine/Source/Runtime/Renderer/Private/PathTracing.cpp
@@ -509,7 +509,7 @@ void FDeferredShadingSceneRenderer::RenderPathTracing(FRHICommandListImmediate&
 		RayCountPerPixelRT->GetRenderTargetItem().UAV
 	);
 
-	FRayTracingSceneRHIParamRef RayTracingSceneRHI = View.RayTracingScene.RayTracingSceneRHI;
+	FRHIRayTracingScene* RayTracingSceneRHI = View.RayTracingScene.RayTracingSceneRHI;
 	RHICmdList.RayTraceDispatch(View.RayTracingMaterialPipeline, RayGenShader->GetRayTracingShader(), RayTracingSceneRHI, GlobalResources, View.ViewRect.Size().X, View.ViewRect.Size().Y);
 
 	// Save RayTracingIndirect for compositing
diff --git a/Engine/Source/Runtime/Renderer/Private/RayTracing/RayTracingAmbientOcclusion.cpp b/Engine/Source/Runtime/Renderer/Private/RayTracing/RayTracingAmbientOcclusion.cpp
--- a/Engine/Source/Runtime/Renderer/Private/RayTracing/RayTracingAmbientOcclusion.cpp
+++ b/Engine/Source/Runtime/Renderer/Private/RayTracing/RayTracingAmbientOcclusion.cpp
@@ -227,7 +227,7 @@ void FDeferredShadingSceneRenderer::RenderRayTracingAmbientOcclusion(
 			Pipeline = PipelineStateCache::GetAndOrCreateRayTracingPipelineState(RHICmdList, Initializer);
 		}
 
-		FRayTracingSceneRHIParamRef RayTracingSceneRHI = View.RayTracingScene.RayTracingSceneRHI;
+		FRHIRayTracingScene* RayTracingSceneRHI = View.RayTracingScene.RayTracingSceneRHI;
 		RHICmdList.RayTraceDispatch(Pipeline, RayGenerationShader->GetRayTracingShader(), RayTracingSceneRHI, GlobalResources, RayTracingResolution.X, RayTracingResolution.Y);
 	});
 
diff --git a/Engine/Source/Runtime/Renderer/Private/RayTracing/RayTracingBarycentrics.cpp b/Engine/Source/Runtime/Renderer/Private/RayTracing/RayTracingBarycentrics.cpp
--- a/Engine/Source/Runtime/Renderer/Private/RayTracing/RayTracingBarycentrics.cpp
+++ b/Engine/Source/Runtime/Renderer/Private/RayTracing/RayTracingBarycentrics.cpp
@@ -98,7 +98,7 @@ void FDeferredShadingSceneRenderer::RenderRayTracingBarycentrics(FRHICommandList
 
 	FRayTracingPipelineState* Pipeline = PipelineStateCache::GetAndOrCreateRayTracingPipelineState(RHICmdList, Initializer);
 
-	FRayTracingSceneRHIParamRef RayTracingSceneRHI = View.RayTracingScene.RayTracingSceneRHI;
+	FRHIRayTracingScene* RayTracingSceneRHI = View.RayTracingScene.RayTracingSceneRHI;
 
 	FRayTracingBarycentricsRGS::FParameters* RayGenParameters = GraphBuilder.AllocParameters<FRayTracingBarycentricsRGS::FParameters>();
 
diff --git a/Engine/Source/Runtime/Renderer/Private/RayTracing/RayTracingDebug.cpp b/Engine/Source/Runtime/Renderer/Private/RayTracing/RayTracingDebug.cpp
--- a/Engine/Source/Runtime/Renderer/Private/RayTracing/RayTracingDebug.cpp
+++ b/Engine/Source/Runtime/Renderer/Private/RayTracing/RayTracingDebug.cpp
@@ -126,7 +126,7 @@ void FDeferredShadingSceneRenderer::RenderRayTracingDebug(FRHICommandListImmedia
 
 	FRayTracingPipelineState* Pipeline = View.RayTracingMaterialPipeline;
 
-	FRayTracingSceneRHIParamRef RayTracingSceneRHI = View.RayTracingScene.RayTracingSceneRHI;
+	FRHIRayTracingScene* RayTracingSceneRHI = View.RayTracingScene.RayTracingSceneRHI;
 
 	FRayTracingDebugRGS::FParameters* RayGenParameters = GraphBuilder.AllocParameters<FRayTracingDebugRGS::FParameters>();
 
diff --git a/Engine/Source/Runtime/Renderer/Private/RayTracing/RayTracingGlobalIllumination.cpp b/Engine/Source/Runtime/Renderer/Private/RayTracing/RayTracingGlobalIllumination.cpp
--- a/Engine/Source/Runtime/Renderer/Private/RayTracing/RayTracingGlobalIllumination.cpp
+++ b/Engine/Source/Runtime/Renderer/Private/RayTracing/RayTracingGlobalIllumination.cpp
@@ -515,7 +515,7 @@ void FDeferredShadingSceneRenderer::RenderRayTracingGlobalIllumination(
 			FRayTracingShaderBindingsWriter GlobalResources;
 			SetShaderParameters(GlobalResources, *RayGenerationShader, *PassParameters);
 
-			FRayTracingSceneRHIParamRef RayTracingSceneRHI = View.RayTracingScene.RayTracingSceneRHI;
+			FRHIRayTracingScene* RayTracingSceneRHI = View.RayTracingScene.RayTracingSceneRHI;
 			RHICmdList.RayTraceDispatch(View.RayTracingMaterialPipeline, RayGenerationShader->GetRayTracingShader(), RayTracingSceneRHI, GlobalResources, RayTracingResolution.X, RayTracingResolution.Y);
 		});
 	}
diff --git a/Engine/Source/Runtime/Renderer/Private/RayTracing/RayTracingReflections.cpp b/Engine/Source/Runtime/Renderer/Private/RayTracing/RayTracingReflections.cpp
--- a/Engine/Source/Runtime/Renderer/Private/RayTracing/RayTracingReflections.cpp
+++ b/Engine/Source/Runtime/Renderer/Private/RayTracing/RayTracingReflections.cpp
@@ -336,7 +336,7 @@ void FDeferredShadingSceneRenderer::RenderRayTracingReflections(
 				FRayTracingShaderBindingsWriter GlobalResources;
 				SetShaderParameters(GlobalResources, RayGenShader, *PassParameters);
 
-				FRayTracingSceneRHIParamRef RayTracingSceneRHI = View.RayTracingScene.RayTracingSceneRHI;
+				FRHIRayTracingScene* RayTracingSceneRHI = View.RayTracingScene.RayTracingSceneRHI;
 				RHICmdList.RayTraceDispatch(Pipeline, RayGenShader->GetRayTracingShader(), RayTracingSceneRHI, GlobalResources, TileAlignedResolution.X, TileAlignedResolution.Y);
 			});
 
@@ -358,7 +358,7 @@ void FDeferredShadingSceneRenderer::RenderRayTracingReflections(
 				FRayTracingShaderBindingsWriter GlobalResources;
 				SetShaderParameters(GlobalResources, RayGenShader, *PassParameters);
 
-				FRayTracingSceneRHIParamRef RayTracingSceneRHI = View.RayTracingScene.RayTracingSceneRHI;
+				FRHIRayTracingScene* RayTracingSceneRHI = View.RayTracingScene.RayTracingSceneRHI;
 
 				if (DeferredMaterialMode == EDeferredMaterialMode::Shade)
 				{
diff --git a/Engine/Source/Runtime/Renderer/Private/RayTracing/RayTracingShadows.cpp b/Engine/Source/Runtime/Renderer/Private/RayTracing/RayTracingShadows.cpp
--- a/Engine/Source/Runtime/Renderer/Private/RayTracing/RayTracingShadows.cpp
+++ b/Engine/Source/Runtime/Renderer/Private/RayTracing/RayTracingShadows.cpp
@@ -211,7 +211,7 @@ void FDeferredShadingSceneRenderer::RenderRayTracingShadows(
 			FRayTracingShaderBindingsWriter GlobalResources;
 			SetShaderParameters(GlobalResources, *RayGenerationShader, *PassParameters);
 
-			FRayTracingSceneRHIParamRef RayTracingSceneRHI = View.RayTracingScene.RayTracingSceneRHI;
+			FRHIRayTracingScene* RayTracingSceneRHI = View.RayTracingScene.RayTracingSceneRHI;
 
 			if (GRayTracingShadowsEnableMaterials)
 			{
diff --git a/Engine/Source/Runtime/Renderer/Private/RayTracing/RayTracingTranslucency.cpp b/Engine/Source/Runtime/Renderer/Private/RayTracing/RayTracingTranslucency.cpp
--- a/Engine/Source/Runtime/Renderer/Private/RayTracing/RayTracingTranslucency.cpp
+++ b/Engine/Source/Runtime/Renderer/Private/RayTracing/RayTracingTranslucency.cpp
@@ -417,7 +417,7 @@ void FDeferredShadingSceneRenderer::RenderRayTracingTranslucencyView(
 		FRayTracingShaderBindingsWriter GlobalResources;
 		SetShaderParameters(GlobalResources, RayGenShader, *PassParameters);
 
-		FRayTracingSceneRHIParamRef RayTracingSceneRHI = View.RayTracingScene.RayTracingSceneRHI;
+		FRHIRayTracingScene* RayTracingSceneRHI = View.RayTracingScene.RayTracingSceneRHI;
 		RHICmdList.RayTraceDispatch(Pipeline, RayGenShader->GetRayTracingShader(), RayTracingSceneRHI, GlobalResources, RayTracingResolution.X, RayTracingResolution.Y);
 	});
 }
diff --git a/Engine/Source/Runtime/Renderer/Private/RayTracing/RaytracingSkylight.cpp b/Engine/Source/Runtime/Renderer/Private/RayTracing/RaytracingSkylight.cpp
--- a/Engine/Source/Runtime/Renderer/Private/RayTracing/RaytracingSkylight.cpp
+++ b/Engine/Source/Runtime/Renderer/Private/RayTracing/RaytracingSkylight.cpp
@@ -896,7 +896,7 @@ void FDeferredShadingSceneRenderer::RenderRayTracingSkyLight(
 				Pipeline = PipelineStateCache::GetAndOrCreateRayTracingPipelineState(RHICmdList, Initializer);
 			}
 
-			FRayTracingSceneRHIParamRef RayTracingSceneRHI = View.RayTracingScene.RayTracingSceneRHI;
+			FRHIRayTracingScene* RayTracingSceneRHI = View.RayTracingScene.RayTracingSceneRHI;
 			RHICmdList.RayTraceDispatch(Pipeline, RayGenerationShader->GetRayTracingShader(), RayTracingSceneRHI, GlobalResources, RayTracingResolution.X, RayTracingResolution.Y);
 		});
 
diff --git a/Engine/Source/Runtime/Renderer/Private/Tests/RayTracingTestbed.cpp b/Engine/Source/Runtime/Renderer/Private/Tests/RayTracingTestbed.cpp
--- a/Engine/Source/Runtime/Renderer/Private/Tests/RayTracingTestbed.cpp
+++ b/Engine/Source/Runtime/Renderer/Private/Tests/RayTracingTestbed.cpp
@@ -112,7 +112,7 @@ bool RunRayTracingTestbed_RenderThread(const FString& Parameters)
 	FRayTracingSceneInitializer Initializer;
 	Initializer.Instances = Instances;
 	Initializer.ShaderSlotsPerGeometrySegment = RAY_TRACING_NUM_SHADER_SLOTS;
-	FRayTracingSceneRHIParamRef Scene = RHICreateRayTracingScene(Initializer);
+	FRHIRayTracingScene* Scene = RHICreateRayTracingScene(Initializer);
 
 	FRHICommandListImmediate& RHICmdList = FRHICommandListExecutor::GetImmediateCommandList();
 
diff --git a/Engine/Source/Runtime/Renderer/Public/MeshPassProcessor.h b/Engine/Source/Runtime/Renderer/Public/MeshPassProcessor.h
--- a/Engine/Source/Runtime/Renderer/Public/MeshPassProcessor.h
+++ b/Engine/Source/Runtime/Renderer/Public/MeshPassProcessor.h
@@ -337,7 +337,7 @@ public:
 	void SetOnCommandListForCompute(FRHICommandList& RHICmdList, FComputeShaderRHIParamRef Shader) const;
 
 #if RHI_RAYTRACING
-	void SetRayTracingShaderBindingsForHitGroup(FRHICommandList& RHICmdList, FRayTracingSceneRHIParamRef Scene, uint32 InstanceIndex, uint32 SegmentIndex, FRayTracingPipelineState* Pipeline, uint32 HitGroupIndex, uint32 ShaderSlot) const;
+	void SetRayTracingShaderBindingsForHitGroup(FRHICommandList& RHICmdList, FRHIRayTracingScene* Scene, uint32 InstanceIndex, uint32 SegmentIndex, FRayTracingPipelineState* Pipeline, uint32 HitGroupIndex, uint32 ShaderSlot) const;
 #endif // RHI_RAYTRACING
 
 	/** Returns whether this set of shader bindings can be merged into an instanced draw call with another. */
diff --git a/Engine/Source/Runtime/VulkanRHI/Private/VulkanCommands.cpp b/Engine/Source/Runtime/VulkanRHI/Private/VulkanCommands.cpp
--- a/Engine/Source/Runtime/VulkanRHI/Private/VulkanCommands.cpp
+++ b/Engine/Source/Runtime/VulkanRHI/Private/VulkanCommands.cpp
@@ -1105,7 +1105,7 @@ void FVulkanCommandListContext::PrepareParallelFromBase(const FVulkanCommandList
 	TransitionAndLayoutManager.TempCopy(BaseContext.TransitionAndLayoutManager);
 }
 
-void FVulkanCommandListContext::RHICopyToStagingBuffer(FVertexBufferRHIParamRef SourceBufferRHI, FStagingBufferRHIParamRef StagingBufferRHI, uint32 Offset, uint32 NumBytes)
+void FVulkanCommandListContext::RHICopyToStagingBuffer(FVertexBufferRHIParamRef SourceBufferRHI, FRHIStagingBuffer* StagingBufferRHI, uint32 Offset, uint32 NumBytes)
 {
 	FVulkanCmdBuffer* CmdBuffer = CommandBufferManager->GetActiveCmdBuffer();
 	FVulkanVertexBuffer* VertexBuffer = ResourceCast(SourceBufferRHI);
diff --git a/Engine/Source/Runtime/VulkanRHI/Private/VulkanContext.h b/Engine/Source/Runtime/VulkanRHI/Private/VulkanContext.h
--- a/Engine/Source/Runtime/VulkanRHI/Private/VulkanContext.h
+++ b/Engine/Source/Runtime/VulkanRHI/Private/VulkanContext.h
@@ -231,7 +231,7 @@ public:
 	virtual void RHICopyTexture(FTextureRHIParamRef SourceTexture, FTextureRHIParamRef DestTexture, const FRHICopyTextureInfo& CopyInfo) final override;
 	virtual void RHITransitionResources(EResourceTransitionAccess TransitionType, FTextureRHIParamRef* InRenderTargets, int32 NumTextures) final override;
 	virtual void RHITransitionResources(EResourceTransitionAccess TransitionType, EResourceTransitionPipeline TransitionPipeline, FRHIUnorderedAccessView** InUAVs, int32 NumUAVs, FRHIComputeFence* WriteComputeFence) final override;
-	virtual void RHICopyToStagingBuffer(FVertexBufferRHIParamRef SourceBuffer, FStagingBufferRHIParamRef DestinationStagingBuffer, uint32 Offset, uint32 NumBytes) final override;
+	virtual void RHICopyToStagingBuffer(FVertexBufferRHIParamRef SourceBuffer, FRHIStagingBuffer* DestinationStagingBuffer, uint32 Offset, uint32 NumBytes) final override;
 	virtual void RHIWriteGPUFence(FRHIGPUFence* Fence) final override;
 
 	// Render time measurement
@@ -242,8 +242,8 @@ public:
 
 	virtual void RHISubmitCommandsHint() final override;
 
-	virtual void RHIBeginDrawingViewport(FViewportRHIParamRef Viewport, FTextureRHIParamRef RenderTargetRHI) final override;
-	virtual void RHIEndDrawingViewport(FViewportRHIParamRef Viewport, bool bPresent, bool bLockToVsync) final override;
+	virtual void RHIBeginDrawingViewport(FRHIViewport* Viewport, FTextureRHIParamRef RenderTargetRHI) final override;
+	virtual void RHIEndDrawingViewport(FRHIViewport* Viewport, bool bPresent, bool bLockToVsync) final override;
 
 	virtual void RHIBeginFrame() final override;
 	virtual void RHIEndFrame() final override;
diff --git a/Engine/Source/Runtime/VulkanRHI/Private/VulkanRHI.cpp b/Engine/Source/Runtime/VulkanRHI/Private/VulkanRHI.cpp
--- a/Engine/Source/Runtime/VulkanRHI/Private/VulkanRHI.cpp
+++ b/Engine/Source/Runtime/VulkanRHI/Private/VulkanRHI.cpp
@@ -807,7 +807,7 @@ void FVulkanCommandListContext::RHIEndScene()
 	//FRCLog::Printf(FString::Printf(TEXT("FVulkanCommandListContext::RHIEndScene()")));
 }
 
-void FVulkanCommandListContext::RHIBeginDrawingViewport(FViewportRHIParamRef ViewportRHI, FTextureRHIParamRef RenderTargetRHI)
+void FVulkanCommandListContext::RHIBeginDrawingViewport(FRHIViewport* ViewportRHI, FTextureRHIParamRef RenderTargetRHI)
 {
 	//FRCLog::Printf(FString::Printf(TEXT("FVulkanCommandListContext::RHIBeginDrawingViewport\n")));
 	check(ViewportRHI);
@@ -815,7 +815,7 @@ void FVulkanCommandListContext::RHIBeginDrawingViewport(FViewportRHIParamRef Vie
 	RHI->DrawingViewport = Viewport;
 }
 
-void FVulkanCommandListContext::RHIEndDrawingViewport(FViewportRHIParamRef ViewportRHI, bool bPresent, bool bLockToVsync)
+void FVulkanCommandListContext::RHIEndDrawingViewport(FRHIViewport* ViewportRHI, bool bPresent, bool bLockToVsync)
 {
 	LLM_SCOPE_VULKAN(ELLMTagVulkan::VulkanMisc);
 	//FRCLog::Printf(FString::Printf(TEXT("FVulkanCommandListContext::RHIEndDrawingViewport()")));
diff --git a/Engine/Source/Runtime/VulkanRHI/Private/VulkanUtil.cpp b/Engine/Source/Runtime/VulkanRHI/Private/VulkanUtil.cpp
--- a/Engine/Source/Runtime/VulkanRHI/Private/VulkanUtil.cpp
+++ b/Engine/Source/Runtime/VulkanRHI/Private/VulkanUtil.cpp
@@ -119,13 +119,13 @@ FStagingBufferRHIRef FVulkanDynamicRHI::RHICreateStagingBuffer()
 	return new FVulkanStagingBuffer();
 }
 
-void* FVulkanDynamicRHI::RHILockStagingBuffer(FStagingBufferRHIParamRef StagingBufferRHI, uint32 Offset, uint32 NumBytes)
+void* FVulkanDynamicRHI::RHILockStagingBuffer(FRHIStagingBuffer* StagingBufferRHI, uint32 Offset, uint32 NumBytes)
 {
 	FVulkanStagingBuffer* StagingBuffer = ResourceCast(StagingBufferRHI);
 	return StagingBuffer->Lock(Offset, NumBytes);
 }
 
-void FVulkanDynamicRHI::RHIUnlockStagingBuffer(FStagingBufferRHIParamRef StagingBufferRHI)
+void FVulkanDynamicRHI::RHIUnlockStagingBuffer(FRHIStagingBuffer* StagingBufferRHI)
 {
 	FVulkanStagingBuffer* StagingBuffer = ResourceCast(StagingBufferRHI);
 	StagingBuffer->Unlock();
diff --git a/Engine/Source/Runtime/VulkanRHI/Private/VulkanViewport.cpp b/Engine/Source/Runtime/VulkanRHI/Private/VulkanViewport.cpp
--- a/Engine/Source/Runtime/VulkanRHI/Private/VulkanViewport.cpp
+++ b/Engine/Source/Runtime/VulkanRHI/Private/VulkanViewport.cpp
@@ -898,7 +898,7 @@ FViewportRHIRef FVulkanDynamicRHI::RHICreateViewport(void* WindowHandle, uint32
 	return new FVulkanViewport(this, Device, WindowHandle, SizeX, SizeY, bIsFullscreen, PreferredPixelFormat);
 }
 
-void FVulkanDynamicRHI::RHIResizeViewport(FViewportRHIParamRef ViewportRHI, uint32 SizeX, uint32 SizeY, bool bIsFullscreen, EPixelFormat PreferredPixelFormat)
+void FVulkanDynamicRHI::RHIResizeViewport(FRHIViewport* ViewportRHI, uint32 SizeX, uint32 SizeY, bool bIsFullscreen, EPixelFormat PreferredPixelFormat)
 {
 	check(IsInGameThread());
 	FVulkanViewport* Viewport = ResourceCast(ViewportRHI);
@@ -923,7 +923,7 @@ void FVulkanDynamicRHI::RHIResizeViewport(FViewportRHIParamRef ViewportRHI, uint
 	}
 }
 
-void FVulkanDynamicRHI::RHIResizeViewport(FViewportRHIParamRef ViewportRHI, uint32 SizeX, uint32 SizeY, bool bIsFullscreen)
+void FVulkanDynamicRHI::RHIResizeViewport(FRHIViewport* ViewportRHI, uint32 SizeX, uint32 SizeY, bool bIsFullscreen)
 {
 	check(IsInGameThread());
 	FVulkanViewport* Viewport = ResourceCast(ViewportRHI);
@@ -971,7 +971,7 @@ void FVulkanDynamicRHI::RHITick(float DeltaTime)
 	}
 }
 
-FTexture2DRHIRef FVulkanDynamicRHI::RHIGetViewportBackBuffer(FViewportRHIParamRef ViewportRHI)
+FTexture2DRHIRef FVulkanDynamicRHI::RHIGetViewportBackBuffer(FRHIViewport* ViewportRHI)
 {
 	check(IsInRenderingThread());
 	check(ViewportRHI);
@@ -985,7 +985,7 @@ FTexture2DRHIRef FVulkanDynamicRHI::RHIGetViewportBackBuffer(FViewportRHIParamRe
 	return Viewport->GetBackBuffer(FRHICommandListExecutor::GetImmediateCommandList());
 }
 
-void FVulkanDynamicRHI::RHIAdvanceFrameForGetViewportBackBuffer(FViewportRHIParamRef ViewportRHI)
+void FVulkanDynamicRHI::RHIAdvanceFrameForGetViewportBackBuffer(FRHIViewport* ViewportRHI)
 {
 	check(IsInRenderingThread());
 	check(ViewportRHI);
diff --git a/Engine/Source/Runtime/VulkanRHI/Public/VulkanDynamicRHI.h b/Engine/Source/Runtime/VulkanRHI/Public/VulkanDynamicRHI.h
--- a/Engine/Source/Runtime/VulkanRHI/Public/VulkanDynamicRHI.h
+++ b/Engine/Source/Runtime/VulkanRHI/Public/VulkanDynamicRHI.h
@@ -136,15 +136,15 @@ public:
 	virtual void RHIRead3DSurfaceFloatData(FTextureRHIParamRef Texture, FIntRect Rect, FIntPoint ZMinMax, TArray<FFloat16Color>& OutData) final override;
 	virtual FRenderQueryRHIRef RHICreateRenderQuery(ERenderQueryType QueryType) final override;
 	virtual bool RHIGetRenderQueryResult(FRHIRenderQuery* RenderQuery, uint64& OutResult, bool bWait) final override;
-	virtual FTexture2DRHIRef RHIGetViewportBackBuffer(FViewportRHIParamRef Viewport) final override;
-	virtual void RHIAdvanceFrameForGetViewportBackBuffer(FViewportRHIParamRef Viewport) final override;
+	virtual FTexture2DRHIRef RHIGetViewportBackBuffer(FRHIViewport* Viewport) final override;
+	virtual void RHIAdvanceFrameForGetViewportBackBuffer(FRHIViewport* Viewport) final override;
 	virtual void RHIAcquireThreadOwnership() final override;
 	virtual void RHIReleaseThreadOwnership() final override;
 	virtual void RHIFlushResources() final override;
 	virtual uint32 RHIGetGPUFrameCycles() final override;
 	virtual FViewportRHIRef RHICreateViewport(void* WindowHandle, uint32 SizeX, uint32 SizeY, bool bIsFullscreen, EPixelFormat PreferredPixelFormat) final override;
-	virtual void RHIResizeViewport(FViewportRHIParamRef Viewport, uint32 SizeX, uint32 SizeY, bool bIsFullscreen) final override;
-	virtual void RHIResizeViewport(FViewportRHIParamRef Viewport, uint32 SizeX, uint32 SizeY, bool bIsFullscreen, EPixelFormat PreferredPixelFormat) final override;
+	virtual void RHIResizeViewport(FRHIViewport* Viewport, uint32 SizeX, uint32 SizeY, bool bIsFullscreen) final override;
+	virtual void RHIResizeViewport(FRHIViewport* Viewport, uint32 SizeX, uint32 SizeY, bool bIsFullscreen, EPixelFormat PreferredPixelFormat) final override;
 	virtual void RHITick(float DeltaTime) final override;
 	virtual void RHISetStreamOutTargets(uint32 NumTargets, const FVertexBufferRHIParamRef* VertexBuffers, const uint32* Offsets) final override;
 	virtual void RHIBlockUntilGPUIdle() final override;
@@ -370,8 +370,8 @@ public:
 
 	virtual void VulkanSetImageLayout( VkCommandBuffer CmdBuffer, VkImage Image, VkImageLayout OldLayout, VkImageLayout NewLayout, const VkImageSubresourceRange& SubresourceRange );
 	
-	virtual void* RHILockStagingBuffer(FStagingBufferRHIParamRef StagingBuffer, uint32 Offset, uint32 SizeRHI) final override;
-	virtual void RHIUnlockStagingBuffer(FStagingBufferRHIParamRef StagingBuffer) final override;
+	virtual void* RHILockStagingBuffer(FRHIStagingBuffer* StagingBuffer, uint32 Offset, uint32 SizeRHI) final override;
+	virtual void RHIUnlockStagingBuffer(FRHIStagingBuffer* StagingBuffer) final override;
 
 public:
 	static void SavePipelineCache();
diff --git a/Engine/Source/Runtime/Windows/D3D11RHI/Private/D3D11Commands.cpp b/Engine/Source/Runtime/Windows/D3D11RHI/Private/D3D11Commands.cpp
--- a/Engine/Source/Runtime/Windows/D3D11RHI/Private/D3D11Commands.cpp
+++ b/Engine/Source/Runtime/Windows/D3D11RHI/Private/D3D11Commands.cpp
@@ -2047,7 +2047,7 @@ FStagingBufferRHIRef FD3D11DynamicRHI::RHICreateStagingBuffer()
 	return new FD3D11StagingBuffer();
 }
 
-void FD3D11DynamicRHI::RHICopyToStagingBuffer(FVertexBufferRHIParamRef SourceBufferRHI, FStagingBufferRHIParamRef StagingBufferRHI, uint32 Offset, uint32 NumBytes)
+void FD3D11DynamicRHI::RHICopyToStagingBuffer(FVertexBufferRHIParamRef SourceBufferRHI, FRHIStagingBuffer* StagingBufferRHI, uint32 Offset, uint32 NumBytes)
 {
 	FD3D11VertexBuffer* SourceBuffer = ResourceCast(SourceBufferRHI);
 	FD3D11StagingBuffer* StagingBuffer = ResourceCast(StagingBufferRHI);
@@ -2099,14 +2099,14 @@ void FD3D11DynamicRHI::RHIWriteGPUFence(FRHIGPUFence* FenceRHI)
 	Fence->WriteInternal();
 }
 
-void* FD3D11DynamicRHI::RHILockStagingBuffer(FStagingBufferRHIParamRef StagingBufferRHI, uint32 Offset, uint32 SizeRHI)
+void* FD3D11DynamicRHI::RHILockStagingBuffer(FRHIStagingBuffer* StagingBufferRHI, uint32 Offset, uint32 SizeRHI)
 {
 	check(StagingBufferRHI);
 	FD3D11StagingBuffer* StagingBuffer = ResourceCast(StagingBufferRHI);
 	return StagingBuffer->Lock(Offset, SizeRHI);
 }
 
-void FD3D11DynamicRHI::RHIUnlockStagingBuffer(FStagingBufferRHIParamRef StagingBufferRHI)
+void FD3D11DynamicRHI::RHIUnlockStagingBuffer(FRHIStagingBuffer* StagingBufferRHI)
 {
 	FD3D11StagingBuffer* StagingBuffer = ResourceCast(StagingBufferRHI);
 	StagingBuffer->Unlock();
diff --git a/Engine/Source/Runtime/Windows/D3D11RHI/Private/D3D11RHIPrivate.h b/Engine/Source/Runtime/Windows/D3D11RHI/Private/D3D11RHIPrivate.h
--- a/Engine/Source/Runtime/Windows/D3D11RHI/Private/D3D11RHIPrivate.h
+++ b/Engine/Source/Runtime/Windows/D3D11RHI/Private/D3D11RHIPrivate.h
@@ -405,10 +405,10 @@ public:
 	virtual FGeometryShaderRHIRef RHICreateGeometryShaderWithStreamOutput(const TArray<uint8>& Code, const FStreamOutElementList& ElementList, uint32 NumStrides, const uint32* Strides, int32 RasterizedStream) final override;
 	virtual FComputeShaderRHIRef RHICreateComputeShader(const TArray<uint8>& Code) final override;
 	virtual FStagingBufferRHIRef RHICreateStagingBuffer() final override;
-	virtual void RHICopyToStagingBuffer(FVertexBufferRHIParamRef SourceBuffer, FStagingBufferRHIParamRef DestinationStagingBuffer, uint32 Offset, uint32 NumBytes) final override;
+	virtual void RHICopyToStagingBuffer(FVertexBufferRHIParamRef SourceBuffer, FRHIStagingBuffer* DestinationStagingBuffer, uint32 Offset, uint32 NumBytes) final override;
 	virtual void RHIWriteGPUFence(FRHIGPUFence* Fence) final override;
-	virtual void* RHILockStagingBuffer(FStagingBufferRHIParamRef StagingBuffer, uint32 Offset, uint32 SizeRHI) final override;
-    virtual void RHIUnlockStagingBuffer(FStagingBufferRHIParamRef StagingBuffer) final override;
+	virtual void* RHILockStagingBuffer(FRHIStagingBuffer* StagingBuffer, uint32 Offset, uint32 SizeRHI) final override;
+    virtual void RHIUnlockStagingBuffer(FRHIStagingBuffer* StagingBuffer) final override;
 	virtual FBoundShaderStateRHIRef RHICreateBoundShaderState(FRHIVertexDeclaration* VertexDeclaration, FVertexShaderRHIParamRef VertexShader, FHullShaderRHIParamRef HullShader, FDomainShaderRHIParamRef DomainShader, FPixelShaderRHIParamRef PixelShader, FGeometryShaderRHIParamRef GeometryShader) final override;
 	virtual FUniformBufferRHIRef RHICreateUniformBuffer(const void* Contents, const FRHIUniformBufferLayout& Layout, EUniformBufferUsage Usage, EUniformBufferValidation Validation) final override;
 	virtual void RHIUpdateUniformBuffer(FUniformBufferRHIParamRef UniformBufferRHI, const void* Contents) final override;
@@ -473,15 +473,15 @@ public:
 	virtual void RHIRead3DSurfaceFloatData(FTextureRHIParamRef Texture,FIntRect Rect,FIntPoint ZMinMax,TArray<FFloat16Color>& OutData) final override;
 	virtual FRenderQueryRHIRef RHICreateRenderQuery(ERenderQueryType QueryType) final override;
 	virtual bool RHIGetRenderQueryResult(FRHIRenderQuery* RenderQuery, uint64& OutResult, bool bWait) final override;
-	virtual FTexture2DRHIRef RHIGetViewportBackBuffer(FViewportRHIParamRef Viewport) final override;
-	virtual void RHIAdvanceFrameForGetViewportBackBuffer(FViewportRHIParamRef Viewport) final override;
+	virtual FTexture2DRHIRef RHIGetViewportBackBuffer(FRHIViewport* Viewport) final override;
+	virtual void RHIAdvanceFrameForGetViewportBackBuffer(FRHIViewport* Viewport) final override;
 	virtual void RHIAcquireThreadOwnership() final override;
 	virtual void RHIReleaseThreadOwnership() final override;
 	virtual void RHIFlushResources() final override;
 	virtual uint32 RHIGetGPUFrameCycles() final override;
 	virtual FViewportRHIRef RHICreateViewport(void* WindowHandle, uint32 SizeX, uint32 SizeY, bool bIsFullscreen, EPixelFormat PreferredPixelFormat) final override;
-	virtual void RHIResizeViewport(FViewportRHIParamRef Viewport, uint32 SizeX, uint32 SizeY, bool bIsFullscreen) final override;
-	virtual void RHIResizeViewport(FViewportRHIParamRef Viewport, uint32 SizeX, uint32 SizeY, bool bIsFullscreen, EPixelFormat PreferredPixelFormat) final override;
+	virtual void RHIResizeViewport(FRHIViewport* Viewport, uint32 SizeX, uint32 SizeY, bool bIsFullscreen) final override;
+	virtual void RHIResizeViewport(FRHIViewport* Viewport, uint32 SizeX, uint32 SizeY, bool bIsFullscreen, EPixelFormat PreferredPixelFormat) final override;
 	virtual void RHITick(float DeltaTime) final override;
 	virtual void RHISetStreamOutTargets(uint32 NumTargets,const FVertexBufferRHIParamRef* VertexBuffers,const uint32* Offsets) final override;
 	virtual void RHIBlockUntilGPUIdle() final override;
@@ -510,8 +510,8 @@ public:
 	void RHIBeginOcclusionQueryBatch(uint32 NumQueriesInBatch);
 	void RHIEndOcclusionQueryBatch();
 	virtual void RHISubmitCommandsHint() final override;
-	virtual void RHIBeginDrawingViewport(FViewportRHIParamRef Viewport, FTextureRHIParamRef RenderTargetRHI) final override;
-	virtual void RHIEndDrawingViewport(FViewportRHIParamRef Viewport, bool bPresent, bool bLockToVsync) final override;
+	virtual void RHIBeginDrawingViewport(FRHIViewport* Viewport, FTextureRHIParamRef RenderTargetRHI) final override;
+	virtual void RHIEndDrawingViewport(FRHIViewport* Viewport, bool bPresent, bool bLockToVsync) final override;
 	virtual void RHIBeginFrame() final override;
 	virtual void RHIEndFrame() final override;
 	virtual void RHIBeginScene() final override;
diff --git a/Engine/Source/Runtime/Windows/D3D11RHI/Private/D3D11Viewport.cpp b/Engine/Source/Runtime/Windows/D3D11RHI/Private/D3D11Viewport.cpp
--- a/Engine/Source/Runtime/Windows/D3D11RHI/Private/D3D11Viewport.cpp
+++ b/Engine/Source/Runtime/Windows/D3D11RHI/Private/D3D11Viewport.cpp
@@ -555,7 +555,7 @@ FViewportRHIRef FD3D11DynamicRHI::RHICreateViewport(void* WindowHandle,uint32 Si
 	return new FD3D11Viewport(this,(HWND)WindowHandle,SizeX,SizeY,bIsFullscreen,PreferredPixelFormat);
 }
 
-void FD3D11DynamicRHI::RHIResizeViewport(FViewportRHIParamRef ViewportRHI,uint32 SizeX,uint32 SizeY,bool bIsFullscreen)
+void FD3D11DynamicRHI::RHIResizeViewport(FRHIViewport* ViewportRHI,uint32 SizeX,uint32 SizeY,bool bIsFullscreen)
 {
 	FD3D11Viewport* Viewport = ResourceCast(ViewportRHI);
 
@@ -563,7 +563,7 @@ void FD3D11DynamicRHI::RHIResizeViewport(FViewportRHIParamRef ViewportRHI,uint32
 	Viewport->Resize(SizeX,SizeY,bIsFullscreen, PF_Unknown);
 }
 
-void FD3D11DynamicRHI::RHIResizeViewport(FViewportRHIParamRef ViewportRHI, uint32 SizeX, uint32 SizeY, bool bIsFullscreen, EPixelFormat PreferredPixelFormat)
+void FD3D11DynamicRHI::RHIResizeViewport(FRHIViewport* ViewportRHI, uint32 SizeX, uint32 SizeY, bool bIsFullscreen, EPixelFormat PreferredPixelFormat)
 {
 	check(IsInGameThread());
 
@@ -593,7 +593,7 @@ void FD3D11DynamicRHI::RHITick( float DeltaTime )
  *	Viewport functions.
  *=============================================================================*/
 
-void FD3D11DynamicRHI::RHIBeginDrawingViewport(FViewportRHIParamRef ViewportRHI, FTextureRHIParamRef RenderTarget)
+void FD3D11DynamicRHI::RHIBeginDrawingViewport(FRHIViewport* ViewportRHI, FTextureRHIParamRef RenderTarget)
 {
 	FD3D11Viewport* Viewport = ResourceCast(ViewportRHI);
 
@@ -615,7 +615,7 @@ void FD3D11DynamicRHI::RHIBeginDrawingViewport(FViewportRHIParamRef ViewportRHI,
 	RHISetScissorRect(false,0,0,0,0);
 }
 
-void FD3D11DynamicRHI::RHIEndDrawingViewport(FViewportRHIParamRef ViewportRHI,bool bPresent,bool bLockToVsync)
+void FD3D11DynamicRHI::RHIEndDrawingViewport(FRHIViewport* ViewportRHI,bool bPresent,bool bLockToVsync)
 {
 	++PresentCounter;
 	FD3D11Viewport* Viewport = ResourceCast(ViewportRHI);
@@ -704,11 +704,11 @@ void FD3D11DynamicRHI::RHIEndDrawingViewport(FViewportRHIParamRef ViewportRHI,bo
 #endif
 }
 
-void FD3D11DynamicRHI::RHIAdvanceFrameForGetViewportBackBuffer(FViewportRHIParamRef Viewport)
+void FD3D11DynamicRHI::RHIAdvanceFrameForGetViewportBackBuffer(FRHIViewport* Viewport)
 {
 }
 
-FTexture2DRHIRef FD3D11DynamicRHI::RHIGetViewportBackBuffer(FViewportRHIParamRef ViewportRHI)
+FTexture2DRHIRef FD3D11DynamicRHI::RHIGetViewportBackBuffer(FRHIViewport* ViewportRHI)
 {
 	FD3D11Viewport* Viewport = ResourceCast(ViewportRHI);
 
