DR - Deprecate FCustomPresentRHIParamRef, FRHIShaderLibraryParamRef and FRHIPipelineBinaryLibraryParamRef
#rnx

diff --git a/Engine/Source/Runtime/Apple/MetalRHI/Private/MetalShaders.cpp b/Engine/Source/Runtime/Apple/MetalRHI/Private/MetalShaders.cpp
--- a/Engine/Source/Runtime/Apple/MetalRHI/Private/MetalShaders.cpp
+++ b/Engine/Source/Runtime/Apple/MetalRHI/Private/MetalShaders.cpp
@@ -1123,7 +1123,7 @@ FVertexShaderRHIRef FMetalDynamicRHI::RHICreateVertexShader(const TArray<uint8>&
 	}
 }
 
-FVertexShaderRHIRef FMetalDynamicRHI::RHICreateVertexShader(FRHIShaderLibraryParamRef Library, FSHAHash Hash)
+FVertexShaderRHIRef FMetalDynamicRHI::RHICreateVertexShader(FRHIShaderLibrary* Library, FSHAHash Hash)
 {
 	@autoreleasepool {
 
@@ -1143,7 +1143,7 @@ FPixelShaderRHIRef FMetalDynamicRHI::RHICreatePixelShader(const TArray<uint8>& C
 	}
 }
 
-FPixelShaderRHIRef FMetalDynamicRHI::RHICreatePixelShader(FRHIShaderLibraryParamRef Library, FSHAHash Hash)
+FPixelShaderRHIRef FMetalDynamicRHI::RHICreatePixelShader(FRHIShaderLibrary* Library, FSHAHash Hash)
 {
 	@autoreleasepool {
 	
@@ -1169,7 +1169,7 @@ FHullShaderRHIRef FMetalDynamicRHI::RHICreateHullShader(const TArray<uint8>& Cod
 	}
 }
 
-FHullShaderRHIRef FMetalDynamicRHI::RHICreateHullShader(FRHIShaderLibraryParamRef Library, FSHAHash Hash)
+FHullShaderRHIRef FMetalDynamicRHI::RHICreateHullShader(FRHIShaderLibrary* Library, FSHAHash Hash)
 {
 	@autoreleasepool {
 	
@@ -1195,7 +1195,7 @@ FDomainShaderRHIRef FMetalDynamicRHI::RHICreateDomainShader(const TArray<uint8>&
 	}
 }
 
-FDomainShaderRHIRef FMetalDynamicRHI::RHICreateDomainShader(FRHIShaderLibraryParamRef Library, FSHAHash Hash)
+FDomainShaderRHIRef FMetalDynamicRHI::RHICreateDomainShader(FRHIShaderLibrary* Library, FSHAHash Hash)
 {
 	@autoreleasepool {
 	
@@ -1217,7 +1217,7 @@ FGeometryShaderRHIRef FMetalDynamicRHI::RHICreateGeometryShader(const TArray<uin
 	}
 }
 
-FGeometryShaderRHIRef FMetalDynamicRHI::RHICreateGeometryShader(FRHIShaderLibraryParamRef Library, FSHAHash Hash)
+FGeometryShaderRHIRef FMetalDynamicRHI::RHICreateGeometryShader(FRHIShaderLibrary* Library, FSHAHash Hash)
 {
 	@autoreleasepool {
 	
@@ -1237,7 +1237,7 @@ FGeometryShaderRHIRef FMetalDynamicRHI::RHICreateGeometryShaderWithStreamOutput(
 }
 
 FGeometryShaderRHIRef FMetalDynamicRHI::RHICreateGeometryShaderWithStreamOutput(const FStreamOutElementList& ElementList,
-	uint32 NumStrides, const uint32* Strides, int32 RasterizedStream, FRHIShaderLibraryParamRef Library, FSHAHash Hash)
+	uint32 NumStrides, const uint32* Strides, int32 RasterizedStream, FRHIShaderLibrary* Library, FSHAHash Hash)
 {
 	@autoreleasepool {
 	
@@ -1256,7 +1256,7 @@ FComputeShaderRHIRef FMetalDynamicRHI::RHICreateComputeShader(const TArray<uint8
 	}
 }
 
-FComputeShaderRHIRef FMetalDynamicRHI::RHICreateComputeShader(FRHIShaderLibraryParamRef Library, FSHAHash Hash) 
+FComputeShaderRHIRef FMetalDynamicRHI::RHICreateComputeShader(FRHIShaderLibrary* Library, FSHAHash Hash)
 { 
 	@autoreleasepool {
 	
@@ -1272,7 +1272,7 @@ FVertexShaderRHIRef FMetalDynamicRHI::CreateVertexShader_RenderThread(class FRHI
 {
 	return RHICreateVertexShader(Code);
 }
-FVertexShaderRHIRef FMetalDynamicRHI::CreateVertexShader_RenderThread(class FRHICommandListImmediate& RHICmdList, FRHIShaderLibraryParamRef Library, FSHAHash Hash)
+FVertexShaderRHIRef FMetalDynamicRHI::CreateVertexShader_RenderThread(class FRHICommandListImmediate& RHICmdList, FRHIShaderLibrary* Library, FSHAHash Hash)
 {
 	return RHICreateVertexShader(Library, Hash);
 }
@@ -1280,7 +1280,7 @@ FPixelShaderRHIRef FMetalDynamicRHI::CreatePixelShader_RenderThread(class FRHICo
 {
 	return RHICreatePixelShader(Code);
 }
-FPixelShaderRHIRef FMetalDynamicRHI::CreatePixelShader_RenderThread(class FRHICommandListImmediate& RHICmdList, FRHIShaderLibraryParamRef Library, FSHAHash Hash)
+FPixelShaderRHIRef FMetalDynamicRHI::CreatePixelShader_RenderThread(class FRHICommandListImmediate& RHICmdList, FRHIShaderLibrary* Library, FSHAHash Hash)
 {
 	return RHICreatePixelShader(Library, Hash);
 }
@@ -1288,7 +1288,7 @@ FGeometryShaderRHIRef FMetalDynamicRHI::CreateGeometryShader_RenderThread(class
 {
 	return RHICreateGeometryShader(Code);
 }
-FGeometryShaderRHIRef FMetalDynamicRHI::CreateGeometryShader_RenderThread(class FRHICommandListImmediate& RHICmdList, FRHIShaderLibraryParamRef Library, FSHAHash Hash)
+FGeometryShaderRHIRef FMetalDynamicRHI::CreateGeometryShader_RenderThread(class FRHICommandListImmediate& RHICmdList, FRHIShaderLibrary* Library, FSHAHash Hash)
 {
 	return RHICreateGeometryShader(Library, Hash);
 }
@@ -1296,7 +1296,7 @@ FGeometryShaderRHIRef FMetalDynamicRHI::CreateGeometryShaderWithStreamOutput_Ren
 {
 	return RHICreateGeometryShaderWithStreamOutput(Code, ElementList, NumStrides, Strides, RasterizedStream);
 }
-FGeometryShaderRHIRef FMetalDynamicRHI::CreateGeometryShaderWithStreamOutput_RenderThread(class FRHICommandListImmediate& RHICmdList, const FStreamOutElementList& ElementList, uint32 NumStrides, const uint32* Strides, int32 RasterizedStream, FRHIShaderLibraryParamRef Library, FSHAHash Hash)
+FGeometryShaderRHIRef FMetalDynamicRHI::CreateGeometryShaderWithStreamOutput_RenderThread(class FRHICommandListImmediate& RHICmdList, const FStreamOutElementList& ElementList, uint32 NumStrides, const uint32* Strides, int32 RasterizedStream, FRHIShaderLibrary* Library, FSHAHash Hash)
 {
 	return RHICreateGeometryShaderWithStreamOutput(ElementList, NumStrides, Strides, RasterizedStream, Library, Hash);
 }
@@ -1304,7 +1304,7 @@ FComputeShaderRHIRef FMetalDynamicRHI::CreateComputeShader_RenderThread(class FR
 {
 	return RHICreateComputeShader(Code);
 }
-FComputeShaderRHIRef FMetalDynamicRHI::CreateComputeShader_RenderThread(class FRHICommandListImmediate& RHICmdList, FRHIShaderLibraryParamRef Library, FSHAHash Hash)
+FComputeShaderRHIRef FMetalDynamicRHI::CreateComputeShader_RenderThread(class FRHICommandListImmediate& RHICmdList, FRHIShaderLibrary* Library, FSHAHash Hash)
 {
 	return RHICreateComputeShader(Library, Hash);
 }
@@ -1312,7 +1312,7 @@ FHullShaderRHIRef FMetalDynamicRHI::CreateHullShader_RenderThread(class FRHIComm
 {
 	return RHICreateHullShader(Code);
 }
-FHullShaderRHIRef FMetalDynamicRHI::CreateHullShader_RenderThread(class FRHICommandListImmediate& RHICmdList, FRHIShaderLibraryParamRef Library, FSHAHash Hash)
+FHullShaderRHIRef FMetalDynamicRHI::CreateHullShader_RenderThread(class FRHICommandListImmediate& RHICmdList, FRHIShaderLibrary* Library, FSHAHash Hash)
 {
 	return RHICreateHullShader(Library, Hash);
 }
@@ -1320,13 +1320,13 @@ FDomainShaderRHIRef FMetalDynamicRHI::CreateDomainShader_RenderThread(class FRHI
 {
 	return RHICreateDomainShader(Code);
 }
-FDomainShaderRHIRef FMetalDynamicRHI::CreateDomainShader_RenderThread(class FRHICommandListImmediate& RHICmdList, FRHIShaderLibraryParamRef Library, FSHAHash Hash)
+FDomainShaderRHIRef FMetalDynamicRHI::CreateDomainShader_RenderThread(class FRHICommandListImmediate& RHICmdList, FRHIShaderLibrary* Library, FSHAHash Hash)
 {
 	return RHICreateDomainShader(Library, Hash);
 }
 
 static FCriticalSection LoadedShaderLibraryMutex;
-static TMap<FString, FRHIShaderLibraryParamRef> LoadedShaderLibraryMap;
+static TMap<FString, FRHIShaderLibrary*> LoadedShaderLibraryMap;
 
 FMetalShaderLibrary::FMetalShaderLibrary(EShaderPlatform InPlatform, FString const& Name, TArray<mtlpp::Library> InLibrary, FMetalShaderMap const& InMap, const FString& InShaderLibraryFilename)
 : FRHIShaderLibrary(InPlatform, Name)
@@ -1526,7 +1526,7 @@ FRHIShaderLibraryRef FMetalDynamicRHI::RHICreateShaderLibrary(EShaderPlatform Pl
 	}
 
 	FScopeLock Lock(&LoadedShaderLibraryMutex);
-	FRHIShaderLibraryParamRef* FoundShaderLibrary = LoadedShaderLibraryMap.Find(BinaryShaderFile);
+	FRHIShaderLibrary** FoundShaderLibrary = LoadedShaderLibraryMap.Find(BinaryShaderFile);
 	if (FoundShaderLibrary)
 	{
 		return *FoundShaderLibrary;
diff --git a/Engine/Source/Runtime/Apple/MetalRHI/Public/MetalRHI.h b/Engine/Source/Runtime/Apple/MetalRHI/Public/MetalRHI.h
--- a/Engine/Source/Runtime/Apple/MetalRHI/Public/MetalRHI.h
+++ b/Engine/Source/Runtime/Apple/MetalRHI/Public/MetalRHI.h
@@ -46,19 +46,19 @@ public:
 	virtual FBlendStateRHIRef RHICreateBlendState(const FBlendStateInitializerRHI& Initializer) final override;
 	virtual FVertexDeclarationRHIRef RHICreateVertexDeclaration(const FVertexDeclarationElementList& Elements) final override;
 	virtual FPixelShaderRHIRef RHICreatePixelShader(const TArray<uint8>& Code) final override;
-	virtual FPixelShaderRHIRef RHICreatePixelShader(FRHIShaderLibraryParamRef Library, FSHAHash Hash) final override;
+	virtual FPixelShaderRHIRef RHICreatePixelShader(FRHIShaderLibrary* Library, FSHAHash Hash) final override;
 	virtual FVertexShaderRHIRef RHICreateVertexShader(const TArray<uint8>& Code) final override;
-	virtual FVertexShaderRHIRef RHICreateVertexShader(FRHIShaderLibraryParamRef Library, FSHAHash Hash) final override;
+	virtual FVertexShaderRHIRef RHICreateVertexShader(FRHIShaderLibrary* Library, FSHAHash Hash) final override;
 	virtual FHullShaderRHIRef RHICreateHullShader(const TArray<uint8>& Code) final override;
-	virtual FHullShaderRHIRef RHICreateHullShader(FRHIShaderLibraryParamRef Library, FSHAHash Hash) final override;
+	virtual FHullShaderRHIRef RHICreateHullShader(FRHIShaderLibrary* Library, FSHAHash Hash) final override;
 	virtual FDomainShaderRHIRef RHICreateDomainShader(const TArray<uint8>& Code) final override;
-	virtual FDomainShaderRHIRef RHICreateDomainShader(FRHIShaderLibraryParamRef Library, FSHAHash Hash) final override;
+	virtual FDomainShaderRHIRef RHICreateDomainShader(FRHIShaderLibrary* Library, FSHAHash Hash) final override;
 	virtual FGeometryShaderRHIRef RHICreateGeometryShader(const TArray<uint8>& Code) final override;
-	virtual FGeometryShaderRHIRef RHICreateGeometryShader(FRHIShaderLibraryParamRef Library, FSHAHash Hash) final override;
+	virtual FGeometryShaderRHIRef RHICreateGeometryShader(FRHIShaderLibrary* Library, FSHAHash Hash) final override;
 	virtual FGeometryShaderRHIRef RHICreateGeometryShaderWithStreamOutput(const TArray<uint8>& Code, const FStreamOutElementList& ElementList, uint32 NumStrides, const uint32* Strides, int32 RasterizedStream) final override;
-	virtual FGeometryShaderRHIRef RHICreateGeometryShaderWithStreamOutput(const FStreamOutElementList& ElementList, uint32 NumStrides, const uint32* Strides, int32 RasterizedStream, FRHIShaderLibraryParamRef Library, FSHAHash Hash) final override;
+	virtual FGeometryShaderRHIRef RHICreateGeometryShaderWithStreamOutput(const FStreamOutElementList& ElementList, uint32 NumStrides, const uint32* Strides, int32 RasterizedStream, FRHIShaderLibrary* Library, FSHAHash Hash) final override;
 	virtual FComputeShaderRHIRef RHICreateComputeShader(const TArray<uint8>& Code) final override;
-	virtual FComputeShaderRHIRef RHICreateComputeShader(FRHIShaderLibraryParamRef Library, FSHAHash Hash) final override;
+	virtual FComputeShaderRHIRef RHICreateComputeShader(FRHIShaderLibrary* Library, FSHAHash Hash) final override;
 	virtual FBoundShaderStateRHIRef RHICreateBoundShaderState(FRHIVertexDeclaration* VertexDeclaration, FVertexShaderRHIParamRef VertexShader, FHullShaderRHIParamRef HullShader, FDomainShaderRHIParamRef DomainShader, FPixelShaderRHIParamRef PixelShader, FGeometryShaderRHIParamRef GeometryShader) final override;
 	virtual FGraphicsPipelineStateRHIRef RHICreateGraphicsPipelineState(const FGraphicsPipelineStateInitializer& Initializer) final override;
 	virtual TRefCountPtr<FRHIComputePipelineState> RHICreateComputePipelineState(FRHIComputeShader* ComputeShader) final override;
@@ -201,19 +201,19 @@ public:
 	virtual void RHICopySubTextureRegion(FTexture2DRHIParamRef SourceTexture, FTexture2DRHIParamRef DestinationTexture, FBox2D SourceBox, FBox2D DestinationBox) final override;
 
 	virtual FVertexShaderRHIRef CreateVertexShader_RenderThread(class FRHICommandListImmediate& RHICmdList, const TArray<uint8>& Code) final override;
-	virtual FVertexShaderRHIRef CreateVertexShader_RenderThread(class FRHICommandListImmediate& RHICmdList, FRHIShaderLibraryParamRef Library, FSHAHash Hash) final override;
+	virtual FVertexShaderRHIRef CreateVertexShader_RenderThread(class FRHICommandListImmediate& RHICmdList, FRHIShaderLibrary* Library, FSHAHash Hash) final override;
 	virtual FPixelShaderRHIRef CreatePixelShader_RenderThread(class FRHICommandListImmediate& RHICmdList, const TArray<uint8>& Code) final override;
-	virtual FPixelShaderRHIRef CreatePixelShader_RenderThread(class FRHICommandListImmediate& RHICmdList, FRHIShaderLibraryParamRef Library, FSHAHash Hash) final override;
+	virtual FPixelShaderRHIRef CreatePixelShader_RenderThread(class FRHICommandListImmediate& RHICmdList, FRHIShaderLibrary* Library, FSHAHash Hash) final override;
 	virtual FGeometryShaderRHIRef CreateGeometryShader_RenderThread(class FRHICommandListImmediate& RHICmdList, const TArray<uint8>& Code) final override;
-	virtual FGeometryShaderRHIRef CreateGeometryShader_RenderThread(class FRHICommandListImmediate& RHICmdList, FRHIShaderLibraryParamRef Library, FSHAHash Hash) final override;
+	virtual FGeometryShaderRHIRef CreateGeometryShader_RenderThread(class FRHICommandListImmediate& RHICmdList, FRHIShaderLibrary* Library, FSHAHash Hash) final override;
 	virtual FGeometryShaderRHIRef CreateGeometryShaderWithStreamOutput_RenderThread(class FRHICommandListImmediate& RHICmdList, const TArray<uint8>& Code, const FStreamOutElementList& ElementList, uint32 NumStrides, const uint32* Strides, int32 RasterizedStream) final override;
-	virtual FGeometryShaderRHIRef CreateGeometryShaderWithStreamOutput_RenderThread(class FRHICommandListImmediate& RHICmdList, const FStreamOutElementList& ElementList, uint32 NumStrides, const uint32* Strides, int32 RasterizedStream, FRHIShaderLibraryParamRef Library, FSHAHash Hash) final override;
+	virtual FGeometryShaderRHIRef CreateGeometryShaderWithStreamOutput_RenderThread(class FRHICommandListImmediate& RHICmdList, const FStreamOutElementList& ElementList, uint32 NumStrides, const uint32* Strides, int32 RasterizedStream, FRHIShaderLibrary* Library, FSHAHash Hash) final override;
 	virtual FComputeShaderRHIRef CreateComputeShader_RenderThread(class FRHICommandListImmediate& RHICmdList, const TArray<uint8>& Code) final override;
-	virtual FComputeShaderRHIRef CreateComputeShader_RenderThread(class FRHICommandListImmediate& RHICmdList, FRHIShaderLibraryParamRef Library, FSHAHash Hash) final override;
+	virtual FComputeShaderRHIRef CreateComputeShader_RenderThread(class FRHICommandListImmediate& RHICmdList, FRHIShaderLibrary* Library, FSHAHash Hash) final override;
 	virtual FHullShaderRHIRef CreateHullShader_RenderThread(class FRHICommandListImmediate& RHICmdList, const TArray<uint8>& Code) final override;
-	virtual FHullShaderRHIRef CreateHullShader_RenderThread(class FRHICommandListImmediate& RHICmdList, FRHIShaderLibraryParamRef Library, FSHAHash Hash) final override;
+	virtual FHullShaderRHIRef CreateHullShader_RenderThread(class FRHICommandListImmediate& RHICmdList, FRHIShaderLibrary* Library, FSHAHash Hash) final override;
 	virtual FDomainShaderRHIRef CreateDomainShader_RenderThread(class FRHICommandListImmediate& RHICmdList, const TArray<uint8>& Code) final override;
-	virtual FDomainShaderRHIRef CreateDomainShader_RenderThread(class FRHICommandListImmediate& RHICmdList, FRHIShaderLibraryParamRef Library, FSHAHash Hash) final override;
+	virtual FDomainShaderRHIRef CreateDomainShader_RenderThread(class FRHICommandListImmediate& RHICmdList, FRHIShaderLibrary* Library, FSHAHash Hash) final override;
 	virtual FRHIShaderLibraryRef RHICreateShaderLibrary_RenderThread(class FRHICommandListImmediate& RHICmdList, EShaderPlatform Platform, FString FilePath, FString Name) final override;
 
 	virtual void RHIUpdateUniformBuffer(FUniformBufferRHIParamRef UniformBufferRHI, const void* Contents) final override;
diff --git a/Engine/Source/Runtime/OpenGLDrv/Private/OpenGLShaders.cpp b/Engine/Source/Runtime/OpenGLDrv/Private/OpenGLShaders.cpp
--- a/Engine/Source/Runtime/OpenGLDrv/Private/OpenGLShaders.cpp
+++ b/Engine/Source/Runtime/OpenGLDrv/Private/OpenGLShaders.cpp
@@ -885,7 +885,7 @@ ShaderType* CompileOpenGLShader(const TArray<uint8>& InShaderCode, const FSHAHas
 }
 
 template <typename ShaderType>
-ShaderType* CompileOpenGLShader(FRHIShaderLibraryParamRef Library, FSHAHash LibraryHash, FRHIShader* RHIShader = nullptr)
+ShaderType* CompileOpenGLShader(FRHIShaderLibrary* Library, FSHAHash LibraryHash, FRHIShader* RHIShader = nullptr)
 {
 	FLibraryShaderCacheValue *Val = GetOpenGLCompiledLibraryShaderCache().Find(LibraryHash);
 	ShaderType* Shader = nullptr;
@@ -1498,7 +1498,7 @@ FGeometryShaderRHIRef FOpenGLDynamicRHI::RHICreateGeometryShaderWithStreamOutput
 }
 
 template<typename RHIType, typename TOGLProxyType>
-RHIType* CreateProxyShader(FRHIShaderLibraryParamRef Library, FSHAHash Hash)
+RHIType* CreateProxyShader(FRHIShaderLibrary* Library, FSHAHash Hash)
 {
 	FRHICommandListImmediate& RHICmdList = FRHICommandListExecutor::GetImmediateCommandList();
 	if (ShouldRunGLRenderContextOpOnThisThread(RHICmdList))
@@ -1518,34 +1518,34 @@ RHIType* CreateProxyShader(FRHIShaderLibraryParamRef Library, FSHAHash Hash)
 	}
 }
 
-FVertexShaderRHIRef FOpenGLDynamicRHI::RHICreateVertexShader(FRHIShaderLibraryParamRef Library, FSHAHash Hash)
+FVertexShaderRHIRef FOpenGLDynamicRHI::RHICreateVertexShader(FRHIShaderLibrary* Library, FSHAHash Hash)
 {
 	return CreateProxyShader<FRHIVertexShader, FOpenGLVertexShaderProxy>(Library, Hash);
 }
 
-FPixelShaderRHIRef FOpenGLDynamicRHI::RHICreatePixelShader(FRHIShaderLibraryParamRef Library, FSHAHash Hash)
+FPixelShaderRHIRef FOpenGLDynamicRHI::RHICreatePixelShader(FRHIShaderLibrary* Library, FSHAHash Hash)
 {
 	return CreateProxyShader<FRHIPixelShader, FOpenGLPixelShaderProxy>(Library, Hash);
 }
 
-FGeometryShaderRHIRef FOpenGLDynamicRHI::RHICreateGeometryShader(FRHIShaderLibraryParamRef Library, FSHAHash Hash)
+FGeometryShaderRHIRef FOpenGLDynamicRHI::RHICreateGeometryShader(FRHIShaderLibrary* Library, FSHAHash Hash)
 {
 	return CreateProxyShader<FRHIGeometryShader, FOpenGLGeometryShaderProxy>(Library, Hash);
 }
 
-FHullShaderRHIRef FOpenGLDynamicRHI::RHICreateHullShader(FRHIShaderLibraryParamRef Library, FSHAHash Hash)
+FHullShaderRHIRef FOpenGLDynamicRHI::RHICreateHullShader(FRHIShaderLibrary* Library, FSHAHash Hash)
 {
 	check(GMaxRHIFeatureLevel >= ERHIFeatureLevel::SM5);
 	return CreateProxyShader<FRHIHullShader, FOpenGLHullShaderProxy>(Library, Hash);
 }
 
-FDomainShaderRHIRef FOpenGLDynamicRHI::RHICreateDomainShader(FRHIShaderLibraryParamRef Library, FSHAHash Hash)
+FDomainShaderRHIRef FOpenGLDynamicRHI::RHICreateDomainShader(FRHIShaderLibrary* Library, FSHAHash Hash)
 {
 	check(GMaxRHIFeatureLevel >= ERHIFeatureLevel::SM5);
 	return CreateProxyShader<FRHIDomainShader, FOpenGLDomainShaderProxy>(Library, Hash);
 }
 
-FGeometryShaderRHIRef FOpenGLDynamicRHI::RHICreateGeometryShaderWithStreamOutput(const FStreamOutElementList& ElementList, uint32 NumStrides, const uint32* Strides, int32 RasterizedStream, FRHIShaderLibraryParamRef Library, FSHAHash Hash)
+FGeometryShaderRHIRef FOpenGLDynamicRHI::RHICreateGeometryShaderWithStreamOutput(const FStreamOutElementList& ElementList, uint32 NumStrides, const uint32* Strides, int32 RasterizedStream, FRHIShaderLibrary* Library, FSHAHash Hash)
 {
 	UE_LOG(LogRHI, Fatal, TEXT("OpenGL Render path does not support stream output!"));
 	return NULL;
@@ -3276,7 +3276,7 @@ FRHIComputeShader* CreateProxyShader<FRHIComputeShader, FOpenGLComputeShaderProx
 }
 
 template<>
-FRHIComputeShader* CreateProxyShader<FRHIComputeShader, FOpenGLComputeShaderProxy>(FRHIShaderLibraryParamRef Library, FSHAHash Hash)
+FRHIComputeShader* CreateProxyShader<FRHIComputeShader, FOpenGLComputeShaderProxy>(FRHIShaderLibrary* Library, FSHAHash Hash)
 {
 	FRHICommandListImmediate& RHICmdList = FRHICommandListExecutor::GetImmediateCommandList();
 	if (ShouldRunGLRenderContextOpOnThisThread(RHICmdList))
@@ -3300,7 +3300,7 @@ FRHIComputeShader* CreateProxyShader<FRHIComputeShader, FOpenGLComputeShaderProx
 	}
 }
 
-FComputeShaderRHIRef FOpenGLDynamicRHI::RHICreateComputeShader(FRHIShaderLibraryParamRef Library, FSHAHash Hash)
+FComputeShaderRHIRef FOpenGLDynamicRHI::RHICreateComputeShader(FRHIShaderLibrary* Library, FSHAHash Hash)
 {
 	check(RHISupportsComputeShaders(GMaxRHIShaderPlatform));
 	return CreateProxyShader<FRHIComputeShader, FOpenGLComputeShaderProxy>(Library, Hash);
diff --git a/Engine/Source/Runtime/OpenGLDrv/Public/OpenGLDrv.h b/Engine/Source/Runtime/OpenGLDrv/Public/OpenGLDrv.h
--- a/Engine/Source/Runtime/OpenGLDrv/Public/OpenGLDrv.h
+++ b/Engine/Source/Runtime/OpenGLDrv/Public/OpenGLDrv.h
@@ -353,13 +353,13 @@ public:
 	virtual FGeometryShaderRHIRef RHICreateGeometryShaderWithStreamOutput(const TArray<uint8>& Code, const FStreamOutElementList& ElementList, uint32 NumStrides, const uint32* Strides, int32 RasterizedStream) final override;
 	virtual FComputeShaderRHIRef RHICreateComputeShader(const TArray<uint8>& Code) final override;
 
-	virtual FPixelShaderRHIRef RHICreatePixelShader(FRHIShaderLibraryParamRef Library, FSHAHash Hash) final override;
-	virtual FVertexShaderRHIRef RHICreateVertexShader(FRHIShaderLibraryParamRef Library, FSHAHash Hash) final override;
-	virtual FHullShaderRHIRef RHICreateHullShader(FRHIShaderLibraryParamRef Library, FSHAHash Hash) final override;
-	virtual FDomainShaderRHIRef RHICreateDomainShader(FRHIShaderLibraryParamRef Library, FSHAHash Hash) final override;
-	virtual FGeometryShaderRHIRef RHICreateGeometryShader(FRHIShaderLibraryParamRef Library, FSHAHash Hash) final override;
-	virtual FGeometryShaderRHIRef RHICreateGeometryShaderWithStreamOutput(const FStreamOutElementList& ElementList, uint32 NumStrides, const uint32* Strides, int32 RasterizedStream, FRHIShaderLibraryParamRef Library, FSHAHash Hash) final override;
-	virtual FComputeShaderRHIRef RHICreateComputeShader(FRHIShaderLibraryParamRef Library, FSHAHash Hash) final override;
+	virtual FPixelShaderRHIRef RHICreatePixelShader(FRHIShaderLibrary* Library, FSHAHash Hash) final override;
+	virtual FVertexShaderRHIRef RHICreateVertexShader(FRHIShaderLibrary* Library, FSHAHash Hash) final override;
+	virtual FHullShaderRHIRef RHICreateHullShader(FRHIShaderLibrary* Library, FSHAHash Hash) final override;
+	virtual FDomainShaderRHIRef RHICreateDomainShader(FRHIShaderLibrary* Library, FSHAHash Hash) final override;
+	virtual FGeometryShaderRHIRef RHICreateGeometryShader(FRHIShaderLibrary* Library, FSHAHash Hash) final override;
+	virtual FGeometryShaderRHIRef RHICreateGeometryShaderWithStreamOutput(const FStreamOutElementList& ElementList, uint32 NumStrides, const uint32* Strides, int32 RasterizedStream, FRHIShaderLibrary* Library, FSHAHash Hash) final override;
+	virtual FComputeShaderRHIRef RHICreateComputeShader(FRHIShaderLibrary* Library, FSHAHash Hash) final override;
 
 	virtual FUniformBufferRHIRef RHICreateUniformBuffer(const void* Contents, const FRHIUniformBufferLayout& Layout, EUniformBufferUsage Usage, EUniformBufferValidation Validation) final override;
 	virtual void RHIUpdateUniformBuffer(FUniformBufferRHIParamRef UniformBufferRHI, const void* Contents) final override;
diff --git a/Engine/Source/Runtime/RHI/Private/RHICommandList.cpp b/Engine/Source/Runtime/RHI/Private/RHICommandList.cpp
--- a/Engine/Source/Runtime/RHI/Private/RHICommandList.cpp
+++ b/Engine/Source/Runtime/RHI/Private/RHICommandList.cpp
@@ -2143,7 +2143,7 @@ FVertexShaderRHIRef FDynamicRHI::CreateVertexShader_RenderThread(class FRHIComma
 	return GDynamicRHI->RHICreateVertexShader(Code);
 }
 
-FVertexShaderRHIRef FDynamicRHI::CreateVertexShader_RenderThread(class FRHICommandListImmediate& RHICmdList, FRHIShaderLibraryParamRef Library, FSHAHash Hash)
+FVertexShaderRHIRef FDynamicRHI::CreateVertexShader_RenderThread(class FRHICommandListImmediate& RHICmdList, FRHIShaderLibrary* Library, FSHAHash Hash)
 {
 	CSV_SCOPED_TIMING_STAT(RHITStalls, CreateVertexShader_RenderThread);
 	FScopedRHIThreadStaller StallRHIThread(RHICmdList);
@@ -2157,7 +2157,7 @@ FPixelShaderRHIRef FDynamicRHI::CreatePixelShader_RenderThread(class FRHICommand
 	return GDynamicRHI->RHICreatePixelShader(Code);
 }
 
-FPixelShaderRHIRef FDynamicRHI::CreatePixelShader_RenderThread(class FRHICommandListImmediate& RHICmdList, FRHIShaderLibraryParamRef Library, FSHAHash Hash)
+FPixelShaderRHIRef FDynamicRHI::CreatePixelShader_RenderThread(class FRHICommandListImmediate& RHICmdList, FRHIShaderLibrary* Library, FSHAHash Hash)
 {
 	CSV_SCOPED_TIMING_STAT(RHITStalls, CreatePixelShader_RenderThread);
 	FScopedRHIThreadStaller StallRHIThread(RHICmdList);
@@ -2171,7 +2171,7 @@ FGeometryShaderRHIRef FDynamicRHI::CreateGeometryShader_RenderThread(class FRHIC
 	return GDynamicRHI->RHICreateGeometryShader(Code);
 }
 
-FGeometryShaderRHIRef FDynamicRHI::CreateGeometryShader_RenderThread(class FRHICommandListImmediate& RHICmdList, FRHIShaderLibraryParamRef Library, FSHAHash Hash)
+FGeometryShaderRHIRef FDynamicRHI::CreateGeometryShader_RenderThread(class FRHICommandListImmediate& RHICmdList, FRHIShaderLibrary* Library, FSHAHash Hash)
 {
 	CSV_SCOPED_TIMING_STAT(RHITStalls, CreateGeometryShader_RenderThread);
 	FScopedRHIThreadStaller StallRHIThread(RHICmdList);
@@ -2185,7 +2185,7 @@ FGeometryShaderRHIRef FDynamicRHI::CreateGeometryShaderWithStreamOutput_RenderTh
 	return GDynamicRHI->RHICreateGeometryShaderWithStreamOutput(Code, ElementList, NumStrides, Strides, RasterizedStream);
 }
 
-FGeometryShaderRHIRef FDynamicRHI::CreateGeometryShaderWithStreamOutput_RenderThread(class FRHICommandListImmediate& RHICmdList, const FStreamOutElementList& ElementList, uint32 NumStrides, const uint32* Strides, int32 RasterizedStream, FRHIShaderLibraryParamRef Library, FSHAHash Hash)
+FGeometryShaderRHIRef FDynamicRHI::CreateGeometryShaderWithStreamOutput_RenderThread(class FRHICommandListImmediate& RHICmdList, const FStreamOutElementList& ElementList, uint32 NumStrides, const uint32* Strides, int32 RasterizedStream, FRHIShaderLibrary* Library, FSHAHash Hash)
 {
 	CSV_SCOPED_TIMING_STAT(RHITStalls, CreateGeometryShader_RenderThread);
 	FScopedRHIThreadStaller StallRHIThread(RHICmdList);
@@ -2199,7 +2199,7 @@ FComputeShaderRHIRef FDynamicRHI::CreateComputeShader_RenderThread(class FRHICom
 	return GDynamicRHI->RHICreateComputeShader(Code);
 }
 
-FComputeShaderRHIRef FDynamicRHI::CreateComputeShader_RenderThread(class FRHICommandListImmediate& RHICmdList, FRHIShaderLibraryParamRef Library, FSHAHash Hash)
+FComputeShaderRHIRef FDynamicRHI::CreateComputeShader_RenderThread(class FRHICommandListImmediate& RHICmdList, FRHIShaderLibrary* Library, FSHAHash Hash)
 {
 	CSV_SCOPED_TIMING_STAT(RHITStalls, CreateComputeShader_RenderThread);
 	FScopedRHIThreadStaller StallRHIThread(RHICmdList);
@@ -2213,7 +2213,7 @@ FHullShaderRHIRef FDynamicRHI::CreateHullShader_RenderThread(class FRHICommandLi
 	return GDynamicRHI->RHICreateHullShader(Code);
 }
 
-FHullShaderRHIRef FDynamicRHI::CreateHullShader_RenderThread(class FRHICommandListImmediate& RHICmdList, FRHIShaderLibraryParamRef Library, FSHAHash Hash)
+FHullShaderRHIRef FDynamicRHI::CreateHullShader_RenderThread(class FRHICommandListImmediate& RHICmdList, FRHIShaderLibrary* Library, FSHAHash Hash)
 {
 	CSV_SCOPED_TIMING_STAT(RHITStalls, CreateHullShader_RenderThread);
 	FScopedRHIThreadStaller StallRHIThread(RHICmdList);
@@ -2227,7 +2227,7 @@ FDomainShaderRHIRef FDynamicRHI::CreateDomainShader_RenderThread(class FRHIComma
 	return GDynamicRHI->RHICreateDomainShader(Code);
 }
 
-FDomainShaderRHIRef FDynamicRHI::CreateDomainShader_RenderThread(class FRHICommandListImmediate& RHICmdList, FRHIShaderLibraryParamRef Library, FSHAHash Hash)
+FDomainShaderRHIRef FDynamicRHI::CreateDomainShader_RenderThread(class FRHICommandListImmediate& RHICmdList, FRHIShaderLibrary* Library, FSHAHash Hash)
 {
 	CSV_SCOPED_TIMING_STAT(RHITStalls, CreateDomainShader_RenderThread);
 	FScopedRHIThreadStaller StallRHIThread(RHICmdList);
diff --git a/Engine/Source/Runtime/RHI/Public/DynamicRHI.h b/Engine/Source/Runtime/RHI/Public/DynamicRHI.h
--- a/Engine/Source/Runtime/RHI/Public/DynamicRHI.h
+++ b/Engine/Source/Runtime/RHI/Public/DynamicRHI.h
@@ -233,7 +233,7 @@ public:
 	virtual FPixelShaderRHIRef RHICreatePixelShader(const TArray<uint8>& Code) = 0;
 
 	// FlushType: Wait RHI Thread
-	virtual FPixelShaderRHIRef RHICreatePixelShader(FRHIShaderLibraryParamRef Library, FSHAHash Hash)
+	virtual FPixelShaderRHIRef RHICreatePixelShader(FRHIShaderLibrary* Library, FSHAHash Hash)
 	{
 		return nullptr;
 	}
@@ -242,7 +242,7 @@ public:
 	virtual FVertexShaderRHIRef RHICreateVertexShader(const TArray<uint8>& Code) = 0;
 
 	// FlushType: Wait RHI Thread
-	virtual FVertexShaderRHIRef RHICreateVertexShader(FRHIShaderLibraryParamRef Library, FSHAHash Hash)
+	virtual FVertexShaderRHIRef RHICreateVertexShader(FRHIShaderLibrary* Library, FSHAHash Hash)
 	{
 		return nullptr;
 	}
@@ -251,7 +251,7 @@ public:
 	virtual FHullShaderRHIRef RHICreateHullShader(const TArray<uint8>& Code) = 0;
 
 	// FlushType: Wait RHI Thread
-	virtual FHullShaderRHIRef RHICreateHullShader(FRHIShaderLibraryParamRef Library, FSHAHash Hash)
+	virtual FHullShaderRHIRef RHICreateHullShader(FRHIShaderLibrary* Library, FSHAHash Hash)
 	{
 		return nullptr;
 	}
@@ -260,7 +260,7 @@ public:
 	virtual FDomainShaderRHIRef RHICreateDomainShader(const TArray<uint8>& Code) = 0;
 
 	// FlushType: Wait RHI Thread
-	virtual FDomainShaderRHIRef RHICreateDomainShader(FRHIShaderLibraryParamRef Library, FSHAHash Hash)
+	virtual FDomainShaderRHIRef RHICreateDomainShader(FRHIShaderLibrary* Library, FSHAHash Hash)
 	{
 		return nullptr;
 	}
@@ -269,7 +269,7 @@ public:
 	virtual FGeometryShaderRHIRef RHICreateGeometryShader(const TArray<uint8>& Code) = 0;
 
 	// FlushType: Wait RHI Thread
-	virtual FGeometryShaderRHIRef RHICreateGeometryShader(FRHIShaderLibraryParamRef Library, FSHAHash Hash)
+	virtual FGeometryShaderRHIRef RHICreateGeometryShader(FRHIShaderLibrary* Library, FSHAHash Hash)
 	{
 		return nullptr;
 	}
@@ -280,7 +280,7 @@ public:
 
 	/** Creates a geometry shader with stream output ability, defined by ElementList. */
 	// FlushType: Wait RHI Thread
-	virtual FGeometryShaderRHIRef RHICreateGeometryShaderWithStreamOutput(const FStreamOutElementList& ElementList, uint32 NumStrides, const uint32* Strides, int32 RasterizedStream, FRHIShaderLibraryParamRef Library, FSHAHash Hash)
+	virtual FGeometryShaderRHIRef RHICreateGeometryShaderWithStreamOutput(const FStreamOutElementList& ElementList, uint32 NumStrides, const uint32* Strides, int32 RasterizedStream, FRHIShaderLibrary* Library, FSHAHash Hash)
 	{
 		return nullptr;
 	}
@@ -292,7 +292,7 @@ public:
 	virtual FComputeShaderRHIRef RHICreateComputeShader(const TArray<uint8>& Code) = 0;
 
 	// FlushType: Wait RHI Thread
-	virtual FComputeShaderRHIRef RHICreateComputeShader(FRHIShaderLibraryParamRef Library, FSHAHash Hash)
+	virtual FComputeShaderRHIRef RHICreateComputeShader(FRHIShaderLibrary* Library, FSHAHash Hash)
 	{
 		return nullptr;
 	}
@@ -416,12 +416,12 @@ public:
 		return new FRHIComputePipelineStateFallback(ComputeShader);
 	}
 
-	virtual FGraphicsPipelineStateRHIRef RHICreateGraphicsPipelineState(const FGraphicsPipelineStateInitializer& Initializer, FRHIPipelineBinaryLibraryParamRef PipelineBinary)
+	virtual FGraphicsPipelineStateRHIRef RHICreateGraphicsPipelineState(const FGraphicsPipelineStateInitializer& Initializer, FRHIPipelineBinaryLibrary* PipelineBinary)
 	{
 		return RHICreateGraphicsPipelineState(Initializer);
 	}
 
-	virtual TRefCountPtr<FRHIComputePipelineState> RHICreateComputePipelineState(FRHIComputeShader* ComputeShader, FRHIPipelineBinaryLibraryParamRef PipelineBinary)
+	virtual TRefCountPtr<FRHIComputePipelineState> RHICreateComputePipelineState(FRHIComputeShader* ComputeShader, FRHIPipelineBinaryLibrary* PipelineBinary)
 	{
 		return RHICreateComputePipelineState(ComputeShader);
 	}
@@ -1104,19 +1104,19 @@ public:
 	virtual void* LockIndexBuffer_RenderThread(class FRHICommandListImmediate& RHICmdList, FIndexBufferRHIParamRef IndexBuffer, uint32 Offset, uint32 SizeRHI, EResourceLockMode LockMode);
 	virtual void UnlockIndexBuffer_RenderThread(class FRHICommandListImmediate& RHICmdList, FIndexBufferRHIParamRef IndexBuffer);
 	virtual FVertexShaderRHIRef CreateVertexShader_RenderThread(class FRHICommandListImmediate& RHICmdList, const TArray<uint8>& Code);
-	virtual FVertexShaderRHIRef CreateVertexShader_RenderThread(class FRHICommandListImmediate& RHICmdList, FRHIShaderLibraryParamRef Library, FSHAHash Hash);
+	virtual FVertexShaderRHIRef CreateVertexShader_RenderThread(class FRHICommandListImmediate& RHICmdList, FRHIShaderLibrary* Library, FSHAHash Hash);
 	virtual FPixelShaderRHIRef CreatePixelShader_RenderThread(class FRHICommandListImmediate& RHICmdList, const TArray<uint8>& Code);
-	virtual FPixelShaderRHIRef CreatePixelShader_RenderThread(class FRHICommandListImmediate& RHICmdList, FRHIShaderLibraryParamRef Library, FSHAHash Hash);
+	virtual FPixelShaderRHIRef CreatePixelShader_RenderThread(class FRHICommandListImmediate& RHICmdList, FRHIShaderLibrary* Library, FSHAHash Hash);
 	virtual FGeometryShaderRHIRef CreateGeometryShader_RenderThread(class FRHICommandListImmediate& RHICmdList, const TArray<uint8>& Code);
-	virtual FGeometryShaderRHIRef CreateGeometryShader_RenderThread(class FRHICommandListImmediate& RHICmdList, FRHIShaderLibraryParamRef Library, FSHAHash Hash);
+	virtual FGeometryShaderRHIRef CreateGeometryShader_RenderThread(class FRHICommandListImmediate& RHICmdList, FRHIShaderLibrary* Library, FSHAHash Hash);
 	virtual FGeometryShaderRHIRef CreateGeometryShaderWithStreamOutput_RenderThread(class FRHICommandListImmediate& RHICmdList, const TArray<uint8>& Code, const FStreamOutElementList& ElementList, uint32 NumStrides, const uint32* Strides, int32 RasterizedStream);
-	virtual FGeometryShaderRHIRef CreateGeometryShaderWithStreamOutput_RenderThread(class FRHICommandListImmediate& RHICmdList, const FStreamOutElementList& ElementList, uint32 NumStrides, const uint32* Strides, int32 RasterizedStream, FRHIShaderLibraryParamRef Library, FSHAHash Hash);
+	virtual FGeometryShaderRHIRef CreateGeometryShaderWithStreamOutput_RenderThread(class FRHICommandListImmediate& RHICmdList, const FStreamOutElementList& ElementList, uint32 NumStrides, const uint32* Strides, int32 RasterizedStream, FRHIShaderLibrary* Library, FSHAHash Hash);
 	virtual FComputeShaderRHIRef CreateComputeShader_RenderThread(class FRHICommandListImmediate& RHICmdList, const TArray<uint8>& Code);
-	virtual FComputeShaderRHIRef CreateComputeShader_RenderThread(class FRHICommandListImmediate& RHICmdList, FRHIShaderLibraryParamRef Library, FSHAHash Hash);
+	virtual FComputeShaderRHIRef CreateComputeShader_RenderThread(class FRHICommandListImmediate& RHICmdList, FRHIShaderLibrary* Library, FSHAHash Hash);
 	virtual FHullShaderRHIRef CreateHullShader_RenderThread(class FRHICommandListImmediate& RHICmdList, const TArray<uint8>& Code);
-	virtual FHullShaderRHIRef CreateHullShader_RenderThread(class FRHICommandListImmediate& RHICmdList, FRHIShaderLibraryParamRef Library, FSHAHash Hash);
+	virtual FHullShaderRHIRef CreateHullShader_RenderThread(class FRHICommandListImmediate& RHICmdList, FRHIShaderLibrary* Library, FSHAHash Hash);
 	virtual FDomainShaderRHIRef CreateDomainShader_RenderThread(class FRHICommandListImmediate& RHICmdList, const TArray<uint8>& Code);
-	virtual FDomainShaderRHIRef CreateDomainShader_RenderThread(class FRHICommandListImmediate& RHICmdList, FRHIShaderLibraryParamRef Library, FSHAHash Hash);
+	virtual FDomainShaderRHIRef CreateDomainShader_RenderThread(class FRHICommandListImmediate& RHICmdList, FRHIShaderLibrary* Library, FSHAHash Hash);
 	virtual void* LockTexture2D_RenderThread(class FRHICommandListImmediate& RHICmdList, FTexture2DRHIParamRef Texture, uint32 MipIndex, EResourceLockMode LockMode, uint32& DestStride, bool bLockWithinMiptail, bool bNeedsDefaultRHIFlush = true);
 	virtual void UnlockTexture2D_RenderThread(class FRHICommandListImmediate& RHICmdList, FTexture2DRHIParamRef Texture, uint32 MipIndex, bool bLockWithinMiptail, bool bNeedsDefaultRHIFlush = true);
 	virtual void UpdateTexture2D_RenderThread(class FRHICommandListImmediate& RHICmdList, FTexture2DRHIParamRef Texture, uint32 MipIndex, const struct FUpdateTextureRegion2D& UpdateRegion, uint32 SourcePitch, const uint8* SourceData);
diff --git a/Engine/Source/Runtime/RHI/Public/RHICommandList.h b/Engine/Source/Runtime/RHI/Public/RHICommandList.h
--- a/Engine/Source/Runtime/RHI/Public/RHICommandList.h
+++ b/Engine/Source/Runtime/RHI/Public/RHICommandList.h
@@ -3549,7 +3549,7 @@ public:
 		return GDynamicRHI->CreatePixelShader_RenderThread(*this, Code);
 	}
 	
-	FORCEINLINE FPixelShaderRHIRef CreatePixelShader(FRHIShaderLibraryParamRef Library, FSHAHash Hash)
+	FORCEINLINE FPixelShaderRHIRef CreatePixelShader(FRHIShaderLibrary* Library, FSHAHash Hash)
 	{
 		LLM_SCOPE(ELLMTag::Shaders);
 		return GDynamicRHI->CreatePixelShader_RenderThread(*this, Library, Hash);
@@ -3561,7 +3561,7 @@ public:
 		return GDynamicRHI->CreateVertexShader_RenderThread(*this, Code);
 	}
 	
-	FORCEINLINE FVertexShaderRHIRef CreateVertexShader(FRHIShaderLibraryParamRef Library, FSHAHash Hash)
+	FORCEINLINE FVertexShaderRHIRef CreateVertexShader(FRHIShaderLibrary* Library, FSHAHash Hash)
 	{
 		LLM_SCOPE(ELLMTag::Shaders);
 		return GDynamicRHI->CreateVertexShader_RenderThread(*this, Library, Hash);
@@ -3573,7 +3573,7 @@ public:
 		return GDynamicRHI->CreateHullShader_RenderThread(*this, Code);
 	}
 	
-	FORCEINLINE FHullShaderRHIRef CreateHullShader(FRHIShaderLibraryParamRef Library, FSHAHash Hash)
+	FORCEINLINE FHullShaderRHIRef CreateHullShader(FRHIShaderLibrary* Library, FSHAHash Hash)
 	{
 		LLM_SCOPE(ELLMTag::Shaders);
 		return GDynamicRHI->CreateHullShader_RenderThread(*this, Library, Hash);
@@ -3585,7 +3585,7 @@ public:
 		return GDynamicRHI->CreateDomainShader_RenderThread(*this, Code);
 	}
 	
-	FORCEINLINE FDomainShaderRHIRef CreateDomainShader(FRHIShaderLibraryParamRef Library, FSHAHash Hash)
+	FORCEINLINE FDomainShaderRHIRef CreateDomainShader(FRHIShaderLibrary* Library, FSHAHash Hash)
 	{
 		LLM_SCOPE(ELLMTag::Shaders);
 		return GDynamicRHI->CreateDomainShader_RenderThread(*this, Library, Hash);
@@ -3597,7 +3597,7 @@ public:
 		return GDynamicRHI->CreateGeometryShader_RenderThread(*this, Code);
 	}
 	
-	FORCEINLINE FGeometryShaderRHIRef CreateGeometryShader(FRHIShaderLibraryParamRef Library, FSHAHash Hash)
+	FORCEINLINE FGeometryShaderRHIRef CreateGeometryShader(FRHIShaderLibrary* Library, FSHAHash Hash)
 	{
 		LLM_SCOPE(ELLMTag::Shaders);
 		return GDynamicRHI->CreateGeometryShader_RenderThread(*this, Library, Hash);
@@ -3609,7 +3609,7 @@ public:
 		return GDynamicRHI->CreateGeometryShaderWithStreamOutput_RenderThread(*this, Code, ElementList, NumStrides, Strides, RasterizedStream);
 	}
 	
-	FORCEINLINE FGeometryShaderRHIRef CreateGeometryShaderWithStreamOutput(const FStreamOutElementList& ElementList, uint32 NumStrides, const uint32* Strides, int32 RasterizedStream, FRHIShaderLibraryParamRef Library, FSHAHash Hash)
+	FORCEINLINE FGeometryShaderRHIRef CreateGeometryShaderWithStreamOutput(const FStreamOutElementList& ElementList, uint32 NumStrides, const uint32* Strides, int32 RasterizedStream, FRHIShaderLibrary* Library, FSHAHash Hash)
 	{
 		LLM_SCOPE(ELLMTag::Shaders);
 		return GDynamicRHI->CreateGeometryShaderWithStreamOutput_RenderThread(*this, ElementList, NumStrides, Strides, RasterizedStream, Library, Hash);
@@ -3621,7 +3621,7 @@ public:
 		return GDynamicRHI->CreateComputeShader_RenderThread(*this, Code);
 	}
 	
-	FORCEINLINE FComputeShaderRHIRef CreateComputeShader(FRHIShaderLibraryParamRef Library, FSHAHash Hash)
+	FORCEINLINE FComputeShaderRHIRef CreateComputeShader(FRHIShaderLibrary* Library, FSHAHash Hash)
 	{
 		LLM_SCOPE(ELLMTag::Shaders);
 		return GDynamicRHI->CreateComputeShader_RenderThread(*this, Library, Hash);
@@ -4497,7 +4497,7 @@ FORCEINLINE FPixelShaderRHIRef RHICreatePixelShader(const TArray<uint8>& Code)
 	return FRHICommandListExecutor::GetImmediateCommandList().CreatePixelShader(Code);
 }
 
-FORCEINLINE FPixelShaderRHIRef RHICreatePixelShader(FRHIShaderLibraryParamRef Library, FSHAHash Hash)
+FORCEINLINE FPixelShaderRHIRef RHICreatePixelShader(FRHIShaderLibrary* Library, FSHAHash Hash)
 {
 	return FRHICommandListExecutor::GetImmediateCommandList().CreatePixelShader(Library, Hash);
 }
@@ -4507,7 +4507,7 @@ FORCEINLINE FVertexShaderRHIRef RHICreateVertexShader(const TArray<uint8>& Code)
 	return FRHICommandListExecutor::GetImmediateCommandList().CreateVertexShader(Code);
 }
 
-FORCEINLINE FVertexShaderRHIRef RHICreateVertexShader(FRHIShaderLibraryParamRef Library, FSHAHash Hash)
+FORCEINLINE FVertexShaderRHIRef RHICreateVertexShader(FRHIShaderLibrary* Library, FSHAHash Hash)
 {
 	return FRHICommandListExecutor::GetImmediateCommandList().CreateVertexShader(Library, Hash);
 }
@@ -4517,7 +4517,7 @@ FORCEINLINE FHullShaderRHIRef RHICreateHullShader(const TArray<uint8>& Code)
 	return FRHICommandListExecutor::GetImmediateCommandList().CreateHullShader(Code);
 }
 
-FORCEINLINE FHullShaderRHIRef RHICreateHullShader(FRHIShaderLibraryParamRef Library, FSHAHash Hash)
+FORCEINLINE FHullShaderRHIRef RHICreateHullShader(FRHIShaderLibrary* Library, FSHAHash Hash)
 {
 	return FRHICommandListExecutor::GetImmediateCommandList().CreateHullShader(Library, Hash);
 }
@@ -4527,7 +4527,7 @@ FORCEINLINE FDomainShaderRHIRef RHICreateDomainShader(const TArray<uint8>& Code)
 	return FRHICommandListExecutor::GetImmediateCommandList().CreateDomainShader(Code);
 }
 
-FORCEINLINE FDomainShaderRHIRef RHICreateDomainShader(FRHIShaderLibraryParamRef Library, FSHAHash Hash)
+FORCEINLINE FDomainShaderRHIRef RHICreateDomainShader(FRHIShaderLibrary* Library, FSHAHash Hash)
 {
 	return FRHICommandListExecutor::GetImmediateCommandList().CreateDomainShader(Library, Hash);
 }
@@ -4537,7 +4537,7 @@ FORCEINLINE FGeometryShaderRHIRef RHICreateGeometryShader(const TArray<uint8>& C
 	return FRHICommandListExecutor::GetImmediateCommandList().CreateGeometryShader(Code);
 }
 
-FORCEINLINE FGeometryShaderRHIRef RHICreateGeometryShader(FRHIShaderLibraryParamRef Library, FSHAHash Hash)
+FORCEINLINE FGeometryShaderRHIRef RHICreateGeometryShader(FRHIShaderLibrary* Library, FSHAHash Hash)
 {
 	return FRHICommandListExecutor::GetImmediateCommandList().CreateGeometryShader(Library, Hash);
 }
@@ -4547,7 +4547,7 @@ FORCEINLINE FGeometryShaderRHIRef RHICreateGeometryShaderWithStreamOutput(const
 	return FRHICommandListExecutor::GetImmediateCommandList().CreateGeometryShaderWithStreamOutput(Code, ElementList, NumStrides, Strides, RasterizedStream);
 }
 
-FORCEINLINE FGeometryShaderRHIRef RHICreateGeometryShaderWithStreamOutput(const FStreamOutElementList& ElementList, uint32 NumStrides, const uint32* Strides, int32 RasterizedStream, FRHIShaderLibraryParamRef Library, FSHAHash Hash)
+FORCEINLINE FGeometryShaderRHIRef RHICreateGeometryShaderWithStreamOutput(const FStreamOutElementList& ElementList, uint32 NumStrides, const uint32* Strides, int32 RasterizedStream, FRHIShaderLibrary* Library, FSHAHash Hash)
 {
 	return FRHICommandListExecutor::GetImmediateCommandList().CreateGeometryShaderWithStreamOutput(ElementList, NumStrides, Strides, RasterizedStream, Library, Hash);
 }
@@ -4557,7 +4557,7 @@ FORCEINLINE FComputeShaderRHIRef RHICreateComputeShader(const TArray<uint8>& Cod
 	return FRHICommandListExecutor::GetImmediateCommandList().CreateComputeShader(Code);
 }
 
-FORCEINLINE FComputeShaderRHIRef RHICreateComputeShader(FRHIShaderLibraryParamRef Library, FSHAHash Hash)
+FORCEINLINE FComputeShaderRHIRef RHICreateComputeShader(FRHIShaderLibrary* Library, FSHAHash Hash)
 {
 	return FRHICommandListExecutor::GetImmediateCommandList().CreateComputeShader(Library, Hash);
 }
diff --git a/Engine/Source/Runtime/RHI/Public/RHIResources.h b/Engine/Source/Runtime/RHI/Public/RHIResources.h
--- a/Engine/Source/Runtime/RHI/Public/RHIResources.h
+++ b/Engine/Source/Runtime/RHI/Public/RHIResources.h
@@ -1732,6 +1732,7 @@ public:
 };
 
 
+UE_DEPRECATED(4.23, "FCustomPresentRHIParamRef typedef is deprecated; please use FRHICustomPresent* directly instead.")
 typedef FRHICustomPresent*              FCustomPresentRHIParamRef;
 typedef TRefCountPtr<FRHICustomPresent> FCustomPresentRHIRef;
 
@@ -2377,6 +2378,7 @@ protected:
 	uint32 LibraryId;
 };
 
+UE_DEPRECATED(4.23, "FRHIShaderLibraryParamRef typedef is deprecated; please use FRHIShaderLibrary* directly instead.")
 typedef FRHIShaderLibrary*				FRHIShaderLibraryParamRef;
 typedef TRefCountPtr<FRHIShaderLibrary>	FRHIShaderLibraryRef;
 
@@ -2391,6 +2393,8 @@ public:
 protected:
 	EShaderPlatform Platform;
 };
+
+UE_DEPRECATED(4.23, "FRHIPipelineBinaryLibraryParamRef typedef is deprecated; please use FRHIPipelineBinaryLibrary* directly instead.")
 typedef FRHIPipelineBinaryLibrary*				FRHIPipelineBinaryLibraryParamRef;
 typedef TRefCountPtr<FRHIPipelineBinaryLibrary>	FRHIPipelineBinaryLibraryRef;
 
diff --git a/Engine/Source/Runtime/RenderCore/Private/ShaderCodeLibrary.cpp b/Engine/Source/Runtime/RenderCore/Private/ShaderCodeLibrary.cpp
--- a/Engine/Source/Runtime/RenderCore/Private/ShaderCodeLibrary.cpp
+++ b/Engine/Source/Runtime/RenderCore/Private/ShaderCodeLibrary.cpp
@@ -1668,7 +1668,7 @@ public:
 			FRWScopeLock(LibraryMutex, SLT_Write);
 			for (uint32 i = ShaderCodeArchiveStack.Num(); i > 0; i--)
 			{
-				FRHIShaderLibraryParamRef ShaderCodeArchive = ShaderCodeArchiveStack[i - 1];
+				FRHIShaderLibrary* ShaderCodeArchive = ShaderCodeArchiveStack[i - 1];
 				if (ShaderCodeArchive->GetName() == Name)
 				{
 					ShaderCodeArchiveStack.RemoveAt(i - 1);
@@ -1754,7 +1754,7 @@ public:
 		checkSlow(Platform == GetRuntimeShaderPlatform());
 
 		FVertexShaderRHIRef Result;
-		FRHIShaderLibraryParamRef ShaderCodeArchive = FindShaderLibrary(Hash);
+		FRHIShaderLibrary* ShaderCodeArchive = FindShaderLibrary(Hash);
 		if (ShaderCodeArchive)
 		{
 			if (bNativeFormat || GRHILazyShaderCodeLoading)
@@ -1770,7 +1770,7 @@ public:
 		checkSlow(Platform == GetRuntimeShaderPlatform());
 
 		FPixelShaderRHIRef Result;
-		FRHIShaderLibraryParamRef ShaderCodeArchive = FindShaderLibrary(Hash);
+		FRHIShaderLibrary* ShaderCodeArchive = FindShaderLibrary(Hash);
 		if (ShaderCodeArchive)
 		{
 			if (bNativeFormat || GRHILazyShaderCodeLoading)
@@ -1786,7 +1786,7 @@ public:
 		checkSlow(Platform == GetRuntimeShaderPlatform());
 
 		FGeometryShaderRHIRef Result;
-		FRHIShaderLibraryParamRef ShaderCodeArchive = FindShaderLibrary(Hash);
+		FRHIShaderLibrary* ShaderCodeArchive = FindShaderLibrary(Hash);
 		if (ShaderCodeArchive)
 		{
 			if (bNativeFormat || GRHILazyShaderCodeLoading)
@@ -1802,7 +1802,7 @@ public:
 		checkSlow(Platform == GetRuntimeShaderPlatform());
 
 		FGeometryShaderRHIRef Result;
-		FRHIShaderLibraryParamRef ShaderCodeArchive = FindShaderLibrary(Hash);
+		FRHIShaderLibrary* ShaderCodeArchive = FindShaderLibrary(Hash);
 		if (ShaderCodeArchive)
 		{
 			if (bNativeFormat || GRHILazyShaderCodeLoading)
@@ -1818,7 +1818,7 @@ public:
 		checkSlow(Platform == GetRuntimeShaderPlatform());
 
 		FHullShaderRHIRef Result;
-		FRHIShaderLibraryParamRef ShaderCodeArchive = FindShaderLibrary(Hash);
+		FRHIShaderLibrary* ShaderCodeArchive = FindShaderLibrary(Hash);
 		if (ShaderCodeArchive)
 		{
 			if (bNativeFormat || GRHILazyShaderCodeLoading)
@@ -1834,7 +1834,7 @@ public:
 		checkSlow(Platform == GetRuntimeShaderPlatform());
 
 		FDomainShaderRHIRef Result;
-		FRHIShaderLibraryParamRef ShaderCodeArchive = FindShaderLibrary(Hash);
+		FRHIShaderLibrary* ShaderCodeArchive = FindShaderLibrary(Hash);
 		if (ShaderCodeArchive)
 		{
 			if (bNativeFormat || GRHILazyShaderCodeLoading)
@@ -1850,7 +1850,7 @@ public:
 		checkSlow(Platform == GetRuntimeShaderPlatform());
 
 		FComputeShaderRHIRef Result;
-		FRHIShaderLibraryParamRef ShaderCodeArchive = FindShaderLibrary(Hash);
+		FRHIShaderLibrary* ShaderCodeArchive = FindShaderLibrary(Hash);
 		if (ShaderCodeArchive)
 		{
 			if (bNativeFormat || GRHILazyShaderCodeLoading)
@@ -1887,15 +1887,15 @@ public:
 		return nullptr;
 	}
 
-	FRHIShaderLibraryParamRef FindShaderLibrary(const FSHAHash& Hash)
+	FRHIShaderLibrary* FindShaderLibrary(const FSHAHash& Hash)
 	{
 		FRWScopeLock(LibraryMutex, SLT_ReadOnly);
-		FRHIShaderLibraryParamRef Result = nullptr;
+		FRHIShaderLibrary* Result = nullptr;
 
 		// Search in library opened order
 		for (int32 i = 0; i < ShaderCodeArchiveStack.Num(); ++i)
 		{
-			FRHIShaderLibraryParamRef ShaderCodeArchive = ShaderCodeArchiveStack[i];
+			FRHIShaderLibrary* ShaderCodeArchive = ShaderCodeArchiveStack[i];
 			if (ShaderCodeArchive->ContainsEntry(Hash))
 			{
 				Result = ShaderCodeArchive;
@@ -1907,7 +1907,7 @@ public:
 
 	bool ContainsShaderCode(const FSHAHash& Hash)
 	{
-		FRHIShaderLibraryParamRef ShaderCodeArchive = FindShaderLibrary(Hash);
+		FRHIShaderLibrary* ShaderCodeArchive = FindShaderLibrary(Hash);
 		if (ShaderCodeArchive)
 			return true;
 		else
@@ -1916,7 +1916,7 @@ public:
 
 	bool RequestShaderCode(const FSHAHash& Hash, FArchive* Ar)
 	{
-		FRHIShaderLibraryParamRef ShaderCodeArchive = FindShaderLibrary(Hash);
+		FRHIShaderLibrary* ShaderCodeArchive = FindShaderLibrary(Hash);
 		if (ShaderCodeArchive)
 			return ShaderCodeArchive->RequestEntry(Hash, Ar);
 		else
@@ -1927,7 +1927,7 @@ public:
 	{
 		if (!bNativeFormat)
 		{
-			FRHIShaderLibraryParamRef ShaderCodeArchive = FindShaderLibrary(Hash);
+			FRHIShaderLibrary* ShaderCodeArchive = FindShaderLibrary(Hash);
 			if (ShaderCodeArchive)
 				((FShaderCodeArchive*)ShaderCodeArchive)->ReleaseShaderCode(Hash);
 		}
diff --git a/Engine/Source/Runtime/VulkanRHI/Public/VulkanDynamicRHI.h b/Engine/Source/Runtime/VulkanRHI/Public/VulkanDynamicRHI.h
--- a/Engine/Source/Runtime/VulkanRHI/Public/VulkanDynamicRHI.h
+++ b/Engine/Source/Runtime/VulkanRHI/Public/VulkanDynamicRHI.h
@@ -194,37 +194,37 @@ public:
 	{
 		return RHICreateVertexShader(Code);
 	}
-	//virtual FVertexShaderRHIRef CreateVertexShader_RenderThread(class FRHICommandListImmediate& RHICmdList, FRHIShaderLibraryParamRef Library, FSHAHash Hash) override final;
+	//virtual FVertexShaderRHIRef CreateVertexShader_RenderThread(class FRHICommandListImmediate& RHICmdList, FRHIShaderLibrary* Library, FSHAHash Hash) override final;
 	virtual FPixelShaderRHIRef CreatePixelShader_RenderThread(class FRHICommandListImmediate& RHICmdList, const TArray<uint8>& Code) override final
 	{
 		return RHICreatePixelShader(Code);
 	}
-	//virtual FPixelShaderRHIRef CreatePixelShader_RenderThread(class FRHICommandListImmediate& RHICmdList, FRHIShaderLibraryParamRef Library, FSHAHash Hash) override final;
+	//virtual FPixelShaderRHIRef CreatePixelShader_RenderThread(class FRHICommandListImmediate& RHICmdList, FRHIShaderLibrary* Library, FSHAHash Hash) override final;
 	virtual FGeometryShaderRHIRef CreateGeometryShader_RenderThread(class FRHICommandListImmediate& RHICmdList, const TArray<uint8>& Code) override final
 	{
 		return RHICreateGeometryShader(Code);
 	}
-	//virtual FGeometryShaderRHIRef CreateGeometryShader_RenderThread(class FRHICommandListImmediate& RHICmdList, FRHIShaderLibraryParamRef Library, FSHAHash Hash) override final;
+	//virtual FGeometryShaderRHIRef CreateGeometryShader_RenderThread(class FRHICommandListImmediate& RHICmdList, FRHIShaderLibrary* Library, FSHAHash Hash) override final;
 	virtual FGeometryShaderRHIRef CreateGeometryShaderWithStreamOutput_RenderThread(class FRHICommandListImmediate& RHICmdList, const TArray<uint8>& Code, const FStreamOutElementList& ElementList, uint32 NumStrides, const uint32* Strides, int32 RasterizedStream) override final
 	{
 		return RHICreateGeometryShaderWithStreamOutput(Code, ElementList, NumStrides, Strides, RasterizedStream);
 	}
-	//virtual FGeometryShaderRHIRef CreateGeometryShaderWithStreamOutput_RenderThread(class FRHICommandListImmediate& RHICmdList, const FStreamOutElementList& ElementList, uint32 NumStrides, const uint32* Strides, int32 RasterizedStream, FRHIShaderLibraryParamRef Library, FSHAHash Hash) override final;
+	//virtual FGeometryShaderRHIRef CreateGeometryShaderWithStreamOutput_RenderThread(class FRHICommandListImmediate& RHICmdList, const FStreamOutElementList& ElementList, uint32 NumStrides, const uint32* Strides, int32 RasterizedStream, FRHIShaderLibrary* Library, FSHAHash Hash) override final;
 	virtual FComputeShaderRHIRef CreateComputeShader_RenderThread(class FRHICommandListImmediate& RHICmdList, const TArray<uint8>& Code) override final
 	{
 		return RHICreateComputeShader(Code);
 	}
-	//virtual FComputeShaderRHIRef CreateComputeShader_RenderThread(class FRHICommandListImmediate& RHICmdList, FRHIShaderLibraryParamRef Library, FSHAHash Hash) override final;
+	//virtual FComputeShaderRHIRef CreateComputeShader_RenderThread(class FRHICommandListImmediate& RHICmdList, FRHIShaderLibrary* Library, FSHAHash Hash) override final;
 	virtual FHullShaderRHIRef CreateHullShader_RenderThread(class FRHICommandListImmediate& RHICmdList, const TArray<uint8>& Code) override final
 	{
 		return RHICreateHullShader(Code);
 	}
-	//virtual FHullShaderRHIRef CreateHullShader_RenderThread(class FRHICommandListImmediate& RHICmdList, FRHIShaderLibraryParamRef Library, FSHAHash Hash) override final;
+	//virtual FHullShaderRHIRef CreateHullShader_RenderThread(class FRHICommandListImmediate& RHICmdList, FRHIShaderLibrary* Library, FSHAHash Hash) override final;
 	virtual FDomainShaderRHIRef CreateDomainShader_RenderThread(class FRHICommandListImmediate& RHICmdList, const TArray<uint8>& Code) override final
 	{
 		return RHICreateDomainShader(Code);
 	}
-	//virtual FDomainShaderRHIRef CreateDomainShader_RenderThread(class FRHICommandListImmediate& RHICmdList, FRHIShaderLibraryParamRef Library, FSHAHash Hash) override final;
+	//virtual FDomainShaderRHIRef CreateDomainShader_RenderThread(class FRHICommandListImmediate& RHICmdList, FRHIShaderLibrary* Library, FSHAHash Hash) override final;
 	virtual void* LockTexture2D_RenderThread(class FRHICommandListImmediate& RHICmdList, FTexture2DRHIParamRef Texture, uint32 MipIndex, EResourceLockMode LockMode, uint32& DestStride, bool bLockWithinMiptail, bool bNeedsDefaultRHIFlush = true) override final
 	{
 		return RHILockTexture2D(Texture, MipIndex, LockMode, DestStride, bLockWithinMiptail);
