DR - Deprecate FVertexShaderRHIParamRef, FHullShaderRHIParamRef, FDomainShaderRHIParamRef, FPixelShaderRHIParamRef, FGeometryShaderRHIParamRef and FComputeShaderRHIParamRef

diff --git a/Engine/Source/Runtime/RHI/Public/BoundShaderStateCache.h b/Engine/Source/Runtime/RHI/Public/BoundShaderStateCache.h
--- a/Engine/Source/Runtime/RHI/Public/BoundShaderStateCache.h
+++ b/Engine/Source/Runtime/RHI/Public/BoundShaderStateCache.h
@@ -19,11 +19,11 @@ public:
 	/** Initialization constructor. */
 	FBoundShaderStateKey(
 		FRHIVertexDeclaration* InVertexDeclaration,
-		FVertexShaderRHIParamRef InVertexShader, 
-		FPixelShaderRHIParamRef InPixelShader,
-		FHullShaderRHIParamRef InHullShader = NULL,
-		FDomainShaderRHIParamRef InDomainShader = NULL,
-		FGeometryShaderRHIParamRef InGeometryShader = NULL
+		FRHIVertexShader* InVertexShader,
+		FRHIPixelShader* InPixelShader,
+		FRHIHullShader* InHullShader = nullptr,
+		FRHIDomainShader* InDomainShader = nullptr,
+		FRHIGeometryShader* InGeometryShader = nullptr
 		)
 		: VertexDeclaration(InVertexDeclaration)
 		, VertexShader(InVertexShader)
@@ -36,11 +36,11 @@ public:
 	/**
 	 * Get the RHI shader for the given frequency.
 	 */
-	FORCEINLINE FVertexShaderRHIParamRef   GetVertexShader() const   { return VertexShader; }
-	FORCEINLINE FPixelShaderRHIParamRef    GetPixelShader() const    { return PixelShader; }
-	FORCEINLINE FHullShaderRHIParamRef     GetHullShader() const     { return HullShader; }
-	FORCEINLINE FDomainShaderRHIParamRef   GetDomainShader() const   { return DomainShader; }
-	FORCEINLINE FGeometryShaderRHIParamRef GetGeometryShader() const { return GeometryShader; }
+	FORCEINLINE FRHIVertexShader*   GetVertexShader() const   { return VertexShader; }
+	FORCEINLINE FRHIPixelShader*    GetPixelShader() const    { return PixelShader; }
+	FORCEINLINE FRHIHullShader*     GetHullShader() const     { return HullShader; }
+	FORCEINLINE FRHIDomainShader*   GetDomainShader() const   { return DomainShader; }
+	FORCEINLINE FRHIGeometryShader* GetGeometryShader() const { return GeometryShader; }
 
 	/**
 	* Get the RHI vertex declaration.
@@ -95,9 +95,9 @@ public:
 		FRHIVertexDeclaration* InVertexDeclaration,
 		FRHIVertexShader* InVertexShader,
 		FRHIPixelShader* InPixelShader,
-		FRHIHullShader* InHullShader = NULL,
-		FRHIDomainShader* InDomainShader = NULL,
-		FRHIGeometryShader* InGeometryShader = NULL
+		FRHIHullShader* InHullShader = nullptr,
+		FRHIDomainShader* InDomainShader = nullptr,
+		FRHIGeometryShader* InGeometryShader = nullptr
 	)
 		: VertexDeclaration(InVertexDeclaration)
 		, VertexShader(InVertexShader)
@@ -165,8 +165,8 @@ public:
 	/** Adds the bound shader state to the cache. */
 	FCachedBoundShaderStateLink(
 		FRHIVertexDeclaration* VertexDeclaration,
-		FVertexShaderRHIParamRef VertexShader,
-		FPixelShaderRHIParamRef PixelShader,
+		FRHIVertexShader* VertexShader,
+		FRHIPixelShader* PixelShader,
 		FRHIBoundShaderState* InBoundShaderState,
 		bool bAddToSingleThreadedCache = true
 		);
@@ -174,11 +174,11 @@ public:
 	/** Adds the bound shader state to the cache. */
 	FCachedBoundShaderStateLink(
 		FRHIVertexDeclaration* VertexDeclaration,
-		FVertexShaderRHIParamRef VertexShader,
-		FPixelShaderRHIParamRef PixelShader,
-		FHullShaderRHIParamRef HullShader,
-		FDomainShaderRHIParamRef DomainShader,
-		FGeometryShaderRHIParamRef GeometryShader,
+		FRHIVertexShader* VertexShader,
+		FRHIPixelShader* PixelShader,
+		FRHIHullShader* HullShader,
+		FRHIDomainShader* DomainShader,
+		FRHIGeometryShader* GeometryShader,
 		FRHIBoundShaderState* InBoundShaderState,
 		bool bAddToSingleThreadedCache = true
 		);
@@ -189,11 +189,11 @@ public:
 	/**
 	 * Get the RHI shader for the given frequency.
 	 */
-	FORCEINLINE FVertexShaderRHIParamRef   GetVertexShader() const   { return Key.GetVertexShader(); }
-	FORCEINLINE FPixelShaderRHIParamRef    GetPixelShader() const    { return Key.GetPixelShader(); }
-	FORCEINLINE FHullShaderRHIParamRef     GetHullShader() const     { return Key.GetHullShader(); }
-	FORCEINLINE FDomainShaderRHIParamRef   GetDomainShader() const   { return Key.GetDomainShader(); }
-	FORCEINLINE FGeometryShaderRHIParamRef GetGeometryShader() const { return Key.GetGeometryShader(); }
+	FORCEINLINE FRHIVertexShader*   GetVertexShader() const   { return Key.GetVertexShader(); }
+	FORCEINLINE FRHIPixelShader*    GetPixelShader() const    { return Key.GetPixelShader(); }
+	FORCEINLINE FRHIHullShader*     GetHullShader() const     { return Key.GetHullShader(); }
+	FORCEINLINE FRHIDomainShader*   GetDomainShader() const   { return Key.GetDomainShader(); }
+	FORCEINLINE FRHIGeometryShader* GetGeometryShader() const { return Key.GetGeometryShader(); }
 
 	/**
 	* Get the RHI vertex declaration.
@@ -212,11 +212,11 @@ protected:
  */
 extern RHI_API FCachedBoundShaderStateLink* GetCachedBoundShaderState(
 	FRHIVertexDeclaration* VertexDeclaration,
-	FVertexShaderRHIParamRef VertexShader,
-	FPixelShaderRHIParamRef PixelShader,
-	FHullShaderRHIParamRef HullShader = NULL,
-	FDomainShaderRHIParamRef DomainShader = NULL,
-	FGeometryShaderRHIParamRef GeometryShader = NULL
+	FRHIVertexShader* VertexShader,
+	FRHIPixelShader* PixelShader,
+	FRHIHullShader* HullShader = nullptr,
+	FRHIDomainShader* DomainShader = nullptr,
+	FRHIGeometryShader* GeometryShader = nullptr
 	);
 
 extern RHI_API void EmptyCachedBoundShaderStates();
@@ -227,8 +227,8 @@ public:
 	/** Adds the bound shader state to the cache. */
 	FCachedBoundShaderStateLink_Threadsafe(
 		FRHIVertexDeclaration* VertexDeclaration,
-		FVertexShaderRHIParamRef VertexShader,
-		FPixelShaderRHIParamRef PixelShader,
+		FRHIVertexShader* VertexShader,
+		FRHIPixelShader* PixelShader,
 		FRHIBoundShaderState* InBoundShaderState
 		)
 		: FCachedBoundShaderStateLink(VertexDeclaration, VertexShader, PixelShader, InBoundShaderState, false)
@@ -238,11 +238,11 @@ public:
 	/** Adds the bound shader state to the cache. */
 	FCachedBoundShaderStateLink_Threadsafe(
 		FRHIVertexDeclaration* VertexDeclaration,
-		FVertexShaderRHIParamRef VertexShader,
-		FPixelShaderRHIParamRef PixelShader,
-		FHullShaderRHIParamRef HullShader,
-		FDomainShaderRHIParamRef DomainShader,
-		FGeometryShaderRHIParamRef GeometryShader,
+		FRHIVertexShader* VertexShader,
+		FRHIPixelShader* PixelShader,
+		FRHIHullShader* HullShader,
+		FRHIDomainShader* DomainShader,
+		FRHIGeometryShader* GeometryShader,
 		FRHIBoundShaderState* InBoundShaderState
 		)
 		: FCachedBoundShaderStateLink(VertexDeclaration, VertexShader, PixelShader, HullShader, DomainShader, GeometryShader, InBoundShaderState, false)
@@ -259,10 +259,10 @@ public:
  */
 extern RHI_API FBoundShaderStateRHIRef GetCachedBoundShaderState_Threadsafe(
 	FRHIVertexDeclaration* VertexDeclaration,
-	FVertexShaderRHIParamRef VertexShader,
-	FPixelShaderRHIParamRef PixelShader,
-	FHullShaderRHIParamRef HullShader = NULL,
-	FDomainShaderRHIParamRef DomainShader = NULL,
-	FGeometryShaderRHIParamRef GeometryShader = NULL
+	FRHIVertexShader* VertexShader,
+	FRHIPixelShader* PixelShader,
+	FRHIHullShader* HullShader = nullptr,
+	FRHIDomainShader* DomainShader = nullptr,
+	FRHIGeometryShader* GeometryShader = nullptr
 	);
 
diff --git a/Engine/Source/Runtime/RHI/Public/DynamicRHI.h b/Engine/Source/Runtime/RHI/Public/DynamicRHI.h
--- a/Engine/Source/Runtime/RHI/Public/DynamicRHI.h
+++ b/Engine/Source/Runtime/RHI/Public/DynamicRHI.h
@@ -393,7 +393,7 @@ public:
 	* @param PixelShader - existing pixel shader
 	*/
 	// FlushType: Thread safe, but varies depending on the RHI
-	virtual FBoundShaderStateRHIRef RHICreateBoundShaderState(FRHIVertexDeclaration* VertexDeclaration, FVertexShaderRHIParamRef VertexShader, FHullShaderRHIParamRef HullShader, FDomainShaderRHIParamRef DomainShader, FPixelShaderRHIParamRef PixelShader, FGeometryShaderRHIParamRef GeometryShader) = 0;
+	virtual FBoundShaderStateRHIRef RHICreateBoundShaderState(FRHIVertexDeclaration* VertexDeclaration, FRHIVertexShader* VertexShader, FRHIHullShader* HullShader, FRHIDomainShader* DomainShader, FRHIPixelShader* PixelShader, FRHIGeometryShader* GeometryShader) = 0;
 
 	/**
 	* Creates a graphics pipeline state object (PSO) that represents a complete gpu pipeline for rendering.
@@ -1241,7 +1241,7 @@ FORCEINLINE FBlendStateRHIRef RHICreateBlendState(const FBlendStateInitializerRH
 	return GDynamicRHI->RHICreateBlendState(Initializer);
 }
 
-FORCEINLINE FBoundShaderStateRHIRef RHICreateBoundShaderState(FRHIVertexDeclaration* VertexDeclaration, FVertexShaderRHIParamRef VertexShader, FHullShaderRHIParamRef HullShader, FDomainShaderRHIParamRef DomainShader, FPixelShaderRHIParamRef PixelShader, FGeometryShaderRHIParamRef GeometryShader)
+FORCEINLINE FBoundShaderStateRHIRef RHICreateBoundShaderState(FRHIVertexDeclaration* VertexDeclaration, FRHIVertexShader* VertexShader, FRHIHullShader* HullShader, FRHIDomainShader* DomainShader, FRHIPixelShader* PixelShader, FRHIGeometryShader* GeometryShader)
 {
 	return GDynamicRHI->RHICreateBoundShaderState(VertexDeclaration, VertexShader, HullShader, DomainShader, PixelShader, GeometryShader);
 }
diff --git a/Engine/Source/Runtime/RHI/Public/RHI.h b/Engine/Source/Runtime/RHI/Public/RHI.h
--- a/Engine/Source/Runtime/RHI/Public/RHI.h
+++ b/Engine/Source/Runtime/RHI/Public/RHI.h
@@ -1627,12 +1627,12 @@ extern RHI_API void RHIExit();
 
 
 // the following helper macros allow to safely convert shader types without much code clutter
-#define GETSAFERHISHADER_PIXEL(Shader) ((Shader) ? (Shader)->GetPixelShader() : (FPixelShaderRHIParamRef)FPixelShaderRHIRef())
-#define GETSAFERHISHADER_VERTEX(Shader) ((Shader) ? (Shader)->GetVertexShader() : (FVertexShaderRHIParamRef)FVertexShaderRHIRef())
-#define GETSAFERHISHADER_HULL(Shader) ((Shader) ? (Shader)->GetHullShader() : (FHullShaderRHIParamRef)FHullShaderRHIRef())
-#define GETSAFERHISHADER_DOMAIN(Shader) ((Shader) ? (Shader)->GetDomainShader() : (FDomainShaderRHIParamRef)FDomainShaderRHIRef())
-#define GETSAFERHISHADER_GEOMETRY(Shader) ((Shader) ? (Shader)->GetGeometryShader() : (FGeometryShaderRHIParamRef)FGeometryShaderRHIRef())
-#define GETSAFERHISHADER_COMPUTE(Shader) ((Shader) ? (Shader)->GetComputeShader() : (FComputeShaderRHIParamRef)FComputeShaderRHIRef())
+#define GETSAFERHISHADER_PIXEL(Shader) ((Shader) ? (Shader)->GetPixelShader() : nullptr)
+#define GETSAFERHISHADER_VERTEX(Shader) ((Shader) ? (Shader)->GetVertexShader() : nullptr)
+#define GETSAFERHISHADER_HULL(Shader) ((Shader) ? (Shader)->GetHullShader() : nullptr)
+#define GETSAFERHISHADER_DOMAIN(Shader) ((Shader) ? (Shader)->GetDomainShader() : nullptr)
+#define GETSAFERHISHADER_GEOMETRY(Shader) ((Shader) ? (Shader)->GetGeometryShader() : (FRHIGeometryShader*)FGeometryShaderRHIRef())
+#define GETSAFERHISHADER_COMPUTE(Shader) ((Shader) ? (Shader)->GetComputeShader() : nullptr)
 #define GETSAFERHISHADER_RAYTRACING(Shader) ((Shader) ? (Shader)->GetRayTracingShader() : (FRHIRayTracingShader*)FRayTracingShaderRHIRef())
 
 
diff --git a/Engine/Source/Runtime/RHI/Public/RHICommandList.h b/Engine/Source/Runtime/RHI/Public/RHICommandList.h
--- a/Engine/Source/Runtime/RHI/Public/RHICommandList.h
+++ b/Engine/Source/Runtime/RHI/Public/RHICommandList.h
@@ -729,15 +729,15 @@ struct FRHICommandSetStencilRef final : public FRHICommand<FRHICommandSetStencil
 	RHI_API void Execute(FRHICommandListBase& CmdList);
 };
 
-template <typename TShaderRHIParamRef, ECmdList CmdListType>
-struct FRHICommandSetShaderParameter final : public FRHICommand<FRHICommandSetShaderParameter<TShaderRHIParamRef, CmdListType> >
+template <typename TRHIShader, ECmdList CmdListType>
+struct FRHICommandSetShaderParameter final : public FRHICommand<FRHICommandSetShaderParameter<TRHIShader, CmdListType> >
 {
-	TShaderRHIParamRef Shader;
+	TRHIShader* Shader;
 	const void* NewValue;
 	uint32 BufferIndex;
 	uint32 BaseIndex;
 	uint32 NumBytes;
-	FORCEINLINE_DEBUGGABLE FRHICommandSetShaderParameter(TShaderRHIParamRef InShader, uint32 InBufferIndex, uint32 InBaseIndex, uint32 InNumBytes, const void* InNewValue)
+	FORCEINLINE_DEBUGGABLE FRHICommandSetShaderParameter(TRHIShader* InShader, uint32 InBufferIndex, uint32 InBaseIndex, uint32 InNumBytes, const void* InNewValue)
 		: Shader(InShader)
 		, NewValue(InNewValue)
 		, BufferIndex(InBufferIndex)
@@ -748,13 +748,13 @@ struct FRHICommandSetShaderParameter final : public FRHICommand<FRHICommandSetSh
 	RHI_API void Execute(FRHICommandListBase& CmdList);
 };
 
-template <typename TShaderRHIParamRef, ECmdList CmdListType>
-struct FRHICommandSetShaderUniformBuffer final : public FRHICommand<FRHICommandSetShaderUniformBuffer<TShaderRHIParamRef, CmdListType> >
+template <typename TRHIShader, ECmdList CmdListType>
+struct FRHICommandSetShaderUniformBuffer final : public FRHICommand<FRHICommandSetShaderUniformBuffer<TRHIShader, CmdListType> >
 {
-	TShaderRHIParamRef Shader;
+	TRHIShader* Shader;
 	uint32 BaseIndex;
 	FUniformBufferRHIParamRef UniformBuffer;
-	FORCEINLINE_DEBUGGABLE FRHICommandSetShaderUniformBuffer(TShaderRHIParamRef InShader, uint32 InBaseIndex, FUniformBufferRHIParamRef InUniformBuffer)
+	FORCEINLINE_DEBUGGABLE FRHICommandSetShaderUniformBuffer(TRHIShader* InShader, uint32 InBaseIndex, FUniformBufferRHIParamRef InUniformBuffer)
 		: Shader(InShader)
 		, BaseIndex(InBaseIndex)
 		, UniformBuffer(InUniformBuffer)
@@ -763,13 +763,13 @@ struct FRHICommandSetShaderUniformBuffer final : public FRHICommand<FRHICommandS
 	RHI_API void Execute(FRHICommandListBase& CmdList);
 };
 
-template <typename TShaderRHIParamRef, ECmdList CmdListType>
-struct FRHICommandSetShaderTexture final : public FRHICommand<FRHICommandSetShaderTexture<TShaderRHIParamRef, CmdListType> >
+template <typename TRHIShader, ECmdList CmdListType>
+struct FRHICommandSetShaderTexture final : public FRHICommand<FRHICommandSetShaderTexture<TRHIShader, CmdListType> >
 {
-	TShaderRHIParamRef Shader;
+	TRHIShader* Shader;
 	uint32 TextureIndex;
 	FTextureRHIParamRef Texture;
-	FORCEINLINE_DEBUGGABLE FRHICommandSetShaderTexture(TShaderRHIParamRef InShader, uint32 InTextureIndex, FTextureRHIParamRef InTexture)
+	FORCEINLINE_DEBUGGABLE FRHICommandSetShaderTexture(TRHIShader* InShader, uint32 InTextureIndex, FTextureRHIParamRef InTexture)
 		: Shader(InShader)
 		, TextureIndex(InTextureIndex)
 		, Texture(InTexture)
@@ -778,13 +778,13 @@ struct FRHICommandSetShaderTexture final : public FRHICommand<FRHICommandSetShad
 	RHI_API void Execute(FRHICommandListBase& CmdList);
 };
 
-template <typename TShaderRHIParamRef, ECmdList CmdListType>
-struct FRHICommandSetShaderResourceViewParameter final : public FRHICommand<FRHICommandSetShaderResourceViewParameter<TShaderRHIParamRef, CmdListType> >
+template <typename TRHIShader, ECmdList CmdListType>
+struct FRHICommandSetShaderResourceViewParameter final : public FRHICommand<FRHICommandSetShaderResourceViewParameter<TRHIShader, CmdListType> >
 {
-	TShaderRHIParamRef Shader;
+	TRHIShader* Shader;
 	uint32 SamplerIndex;
 	FRHIShaderResourceView* SRV;
-	FORCEINLINE_DEBUGGABLE FRHICommandSetShaderResourceViewParameter(TShaderRHIParamRef InShader, uint32 InSamplerIndex, FRHIShaderResourceView* InSRV)
+	FORCEINLINE_DEBUGGABLE FRHICommandSetShaderResourceViewParameter(TRHIShader* InShader, uint32 InSamplerIndex, FRHIShaderResourceView* InSRV)
 		: Shader(InShader)
 		, SamplerIndex(InSamplerIndex)
 		, SRV(InSRV)
@@ -793,13 +793,13 @@ struct FRHICommandSetShaderResourceViewParameter final : public FRHICommand<FRHI
 	RHI_API void Execute(FRHICommandListBase& CmdList);
 };
 
-template <typename TShaderRHIParamRef, ECmdList CmdListType>
-struct FRHICommandSetUAVParameter final : public FRHICommand<FRHICommandSetUAVParameter<TShaderRHIParamRef, CmdListType> >
+template <typename TRHIShader, ECmdList CmdListType>
+struct FRHICommandSetUAVParameter final : public FRHICommand<FRHICommandSetUAVParameter<TRHIShader, CmdListType> >
 {
-	TShaderRHIParamRef Shader;
+	TRHIShader* Shader;
 	uint32 UAVIndex;
 	FRHIUnorderedAccessView* UAV;
-	FORCEINLINE_DEBUGGABLE FRHICommandSetUAVParameter(TShaderRHIParamRef InShader, uint32 InUAVIndex, FRHIUnorderedAccessView* InUAV)
+	FORCEINLINE_DEBUGGABLE FRHICommandSetUAVParameter(TRHIShader* InShader, uint32 InUAVIndex, FRHIUnorderedAccessView* InUAV)
 		: Shader(InShader)
 		, UAVIndex(InUAVIndex)
 		, UAV(InUAV)
@@ -808,14 +808,14 @@ struct FRHICommandSetUAVParameter final : public FRHICommand<FRHICommandSetUAVPa
 	RHI_API void Execute(FRHICommandListBase& CmdList);
 };
 
-template <typename TShaderRHIParamRef, ECmdList CmdListType>
-struct FRHICommandSetUAVParameter_IntialCount final : public FRHICommand<FRHICommandSetUAVParameter_IntialCount<TShaderRHIParamRef, CmdListType> >
+template <typename TRHIShader, ECmdList CmdListType>
+struct FRHICommandSetUAVParameter_IntialCount final : public FRHICommand<FRHICommandSetUAVParameter_IntialCount<TRHIShader, CmdListType> >
 {
-	TShaderRHIParamRef Shader;
+	TRHIShader* Shader;
 	uint32 UAVIndex;
 	FRHIUnorderedAccessView* UAV;
 	uint32 InitialCount;
-	FORCEINLINE_DEBUGGABLE FRHICommandSetUAVParameter_IntialCount(TShaderRHIParamRef InShader, uint32 InUAVIndex, FRHIUnorderedAccessView* InUAV, uint32 InInitialCount)
+	FORCEINLINE_DEBUGGABLE FRHICommandSetUAVParameter_IntialCount(TRHIShader* InShader, uint32 InUAVIndex, FRHIUnorderedAccessView* InUAV, uint32 InInitialCount)
 		: Shader(InShader)
 		, UAVIndex(InUAVIndex)
 		, UAV(InUAV)
@@ -825,13 +825,13 @@ struct FRHICommandSetUAVParameter_IntialCount final : public FRHICommand<FRHICom
 	RHI_API void Execute(FRHICommandListBase& CmdList);
 };
 
-template <typename TShaderRHIParamRef, ECmdList CmdListType>
-struct FRHICommandSetShaderSampler final : public FRHICommand<FRHICommandSetShaderSampler<TShaderRHIParamRef, CmdListType> >
+template <typename TRHIShader, ECmdList CmdListType>
+struct FRHICommandSetShaderSampler final : public FRHICommand<FRHICommandSetShaderSampler<TRHIShader, CmdListType> >
 {
-	TShaderRHIParamRef Shader;
+	TRHIShader* Shader;
 	uint32 SamplerIndex;
 	FRHISamplerState* Sampler;
-	FORCEINLINE_DEBUGGABLE FRHICommandSetShaderSampler(TShaderRHIParamRef InShader, uint32 InSamplerIndex, FRHISamplerState* InSampler)
+	FORCEINLINE_DEBUGGABLE FRHICommandSetShaderSampler(TRHIShader* InShader, uint32 InSamplerIndex, FRHISamplerState* InSampler)
 		: Shader(InShader)
 		, SamplerIndex(InSamplerIndex)
 		, Sampler(InSampler)
@@ -1156,8 +1156,8 @@ struct FRHICommandBindClearMRTValues final : public FRHICommand<FRHICommandBindC
 template<ECmdList CmdListType>
 struct FRHICommandSetComputeShader final : public FRHICommand<FRHICommandSetComputeShader<CmdListType>>
 {
-	FComputeShaderRHIParamRef ComputeShader;
-	FORCEINLINE_DEBUGGABLE FRHICommandSetComputeShader(FComputeShaderRHIParamRef InComputeShader)
+	FRHIComputeShader* ComputeShader;
+	FORCEINLINE_DEBUGGABLE FRHICommandSetComputeShader(FRHIComputeShader* InComputeShader)
 		: ComputeShader(InComputeShader)
 	{
 	}
@@ -1579,13 +1579,13 @@ struct FRHICommandBuildLocalUniformBuffer final : public FRHICommand<FRHICommand
 	RHI_API void Execute(FRHICommandListBase& CmdList);
 };
 
-template <typename TShaderRHIParamRef>
-struct FRHICommandSetLocalUniformBuffer final : public FRHICommand<FRHICommandSetLocalUniformBuffer<TShaderRHIParamRef> >
+template <typename TRHIShader>
+struct FRHICommandSetLocalUniformBuffer final : public FRHICommand<FRHICommandSetLocalUniformBuffer<TRHIShader> >
 {
-	TShaderRHIParamRef Shader;
+	TRHIShader* Shader;
 	uint32 BaseIndex;
 	FLocalUniformBuffer LocalUniformBuffer;
-	FORCEINLINE_DEBUGGABLE FRHICommandSetLocalUniformBuffer(FRHICommandListBase* CheckCmdList, TShaderRHIParamRef InShader, uint32 InBaseIndex, const FLocalUniformBuffer& InLocalUniformBuffer)
+	FORCEINLINE_DEBUGGABLE FRHICommandSetLocalUniformBuffer(FRHICommandListBase* CheckCmdList, TRHIShader* InShader, uint32 InBaseIndex, const FLocalUniformBuffer& InLocalUniformBuffer)
 		: Shader(InShader)
 		, BaseIndex(InBaseIndex)
 		, LocalUniformBuffer(InLocalUniformBuffer)
@@ -2030,11 +2030,11 @@ struct FRHICommandSetRayTracingBindings final : public FRHICommand<FRHICommandSe
 #define ALLOC_COMMAND(...) new ( AllocCommand(sizeof(__VA_ARGS__), alignof(__VA_ARGS__)) ) __VA_ARGS__
 #define ALLOC_COMMAND_CL(RHICmdList, ...) new ( (RHICmdList).AllocCommand(sizeof(__VA_ARGS__), alignof(__VA_ARGS__)) ) __VA_ARGS__
 
-template<> void FRHICommandSetShaderParameter<FComputeShaderRHIParamRef, ECmdList::ECompute>::Execute(FRHICommandListBase& CmdList);
-template<> void FRHICommandSetShaderUniformBuffer<FComputeShaderRHIParamRef, ECmdList::ECompute>::Execute(FRHICommandListBase& CmdList);
-template<> void FRHICommandSetShaderTexture<FComputeShaderRHIParamRef, ECmdList::ECompute>::Execute(FRHICommandListBase& CmdList);
-template<> void FRHICommandSetShaderResourceViewParameter<FComputeShaderRHIParamRef, ECmdList::ECompute>::Execute(FRHICommandListBase& CmdList);
-template<> void FRHICommandSetShaderSampler<FComputeShaderRHIParamRef, ECmdList::ECompute>::Execute(FRHICommandListBase& CmdList);
+template<> void FRHICommandSetShaderParameter<FRHIComputeShader, ECmdList::ECompute>::Execute(FRHICommandListBase& CmdList);
+template<> void FRHICommandSetShaderUniformBuffer<FRHIComputeShader, ECmdList::ECompute>::Execute(FRHICommandListBase& CmdList);
+template<> void FRHICommandSetShaderTexture<FRHIComputeShader, ECmdList::ECompute>::Execute(FRHICommandListBase& CmdList);
+template<> void FRHICommandSetShaderResourceViewParameter<FRHIComputeShader, ECmdList::ECompute>::Execute(FRHICommandListBase& CmdList);
+template<> void FRHICommandSetShaderSampler<FRHIComputeShader, ECmdList::ECompute>::Execute(FRHICommandListBase& CmdList);
 
 
 class RHI_API FRHICommandList : public FRHICommandListBase
@@ -2113,8 +2113,8 @@ public:
 		return Result;
 	}
 
-	template <typename TShaderRHIParamRef>
-	FORCEINLINE_DEBUGGABLE void SetLocalShaderUniformBuffer(TShaderRHIParamRef Shader, uint32 BaseIndex, const FLocalUniformBuffer& UniformBuffer)
+	template <typename TRHIShader>
+	FORCEINLINE_DEBUGGABLE void SetLocalShaderUniformBuffer(TRHIShader* Shader, uint32 BaseIndex, const FLocalUniformBuffer& UniformBuffer)
 	{
 		//check(IsOutsideRenderPass());
 		if (Bypass())
@@ -2122,7 +2122,7 @@ public:
 			GetContext().RHISetShaderUniformBuffer(Shader, BaseIndex, UniformBuffer.BypassUniform);
 			return;
 		}
-		ALLOC_COMMAND(FRHICommandSetLocalUniformBuffer<TShaderRHIParamRef>)(this, Shader, BaseIndex, UniformBuffer);
+		ALLOC_COMMAND(FRHICommandSetLocalUniformBuffer<TRHIShader>)(this, Shader, BaseIndex, UniformBuffer);
 	}
 
 	template <typename TShaderRHI>
@@ -2140,7 +2140,7 @@ public:
 			GetContext().RHISetShaderUniformBuffer(Shader, BaseIndex, UniformBuffer);
 			return;
 		}
-		ALLOC_COMMAND(FRHICommandSetShaderUniformBuffer<TShaderRHI*, ECmdList::EGfx>)(Shader, BaseIndex, UniformBuffer);
+		ALLOC_COMMAND(FRHICommandSetShaderUniformBuffer<TShaderRHI, ECmdList::EGfx>)(Shader, BaseIndex, UniformBuffer);
 	}
 	template <typename TShaderRHI>
 	FORCEINLINE void SetShaderUniformBuffer(TRefCountPtr<TShaderRHI>& Shader, uint32 BaseIndex, FUniformBufferRHIParamRef UniformBuffer)
@@ -2159,7 +2159,7 @@ public:
 		}
 		void* UseValue = Alloc(NumBytes, 16);
 		FMemory::Memcpy(UseValue, NewValue, NumBytes);
-		ALLOC_COMMAND(FRHICommandSetShaderParameter<TShaderRHI*, ECmdList::EGfx>)(Shader, BufferIndex, BaseIndex, NumBytes, UseValue);
+		ALLOC_COMMAND(FRHICommandSetShaderParameter<TShaderRHI, ECmdList::EGfx>)(Shader, BufferIndex, BaseIndex, NumBytes, UseValue);
 	}
 	template <typename TShaderRHI>
 	FORCEINLINE void SetShaderParameter(TRefCountPtr<TShaderRHI>& Shader, uint32 BufferIndex, uint32 BaseIndex, uint32 NumBytes, const void* NewValue)
@@ -2167,8 +2167,8 @@ public:
 		SetShaderParameter(Shader.GetReference(), BufferIndex, BaseIndex, NumBytes, NewValue);
 	}
 
-	template <typename TShaderRHIParamRef>
-	FORCEINLINE_DEBUGGABLE void SetShaderTexture(TShaderRHIParamRef Shader, uint32 TextureIndex, FTextureRHIParamRef Texture)
+	template <typename TRHIShader>
+	FORCEINLINE_DEBUGGABLE void SetShaderTexture(TRHIShader* Shader, uint32 TextureIndex, FTextureRHIParamRef Texture)
 	{
 		//check(IsOutsideRenderPass());
 		if (Bypass())
@@ -2176,7 +2176,7 @@ public:
 			GetContext().RHISetShaderTexture(Shader, TextureIndex, Texture);
 			return;
 		}
-		ALLOC_COMMAND(FRHICommandSetShaderTexture<TShaderRHIParamRef, ECmdList::EGfx>)(Shader, TextureIndex, Texture);
+		ALLOC_COMMAND(FRHICommandSetShaderTexture<TRHIShader, ECmdList::EGfx>)(Shader, TextureIndex, Texture);
 	}
 
 	template <typename TShaderRHI>
@@ -2185,8 +2185,14 @@ public:
 		SetShaderTexture(Shader.GetReference(), TextureIndex, Texture);
 	}
 
-	template <typename TShaderRHIParamRef>
-	FORCEINLINE_DEBUGGABLE void SetShaderResourceViewParameter(TShaderRHIParamRef Shader, uint32 SamplerIndex, FRHIShaderResourceView* SRV)
+	template <typename TShaderRHI>
+	FORCEINLINE_DEBUGGABLE void SetShaderTexture(const TShaderRHI* Shader, uint32 TextureIndex, FTextureRHIParamRef Texture)
+	{
+		SetShaderTexture((TShaderRHI*)Shader, TextureIndex, Texture);
+	}
+
+	template <typename TRHIShader>
+	FORCEINLINE_DEBUGGABLE void SetShaderResourceViewParameter(TRHIShader* Shader, uint32 SamplerIndex, FRHIShaderResourceView* SRV)
 	{
 		//check(IsOutsideRenderPass());
 		if (Bypass())
@@ -2194,7 +2200,7 @@ public:
 			GetContext().RHISetShaderResourceViewParameter(Shader, SamplerIndex, SRV);
 			return;
 		}
-		ALLOC_COMMAND(FRHICommandSetShaderResourceViewParameter<TShaderRHIParamRef, ECmdList::EGfx>)(Shader, SamplerIndex, SRV);
+		ALLOC_COMMAND(FRHICommandSetShaderResourceViewParameter<TRHIShader, ECmdList::EGfx>)(Shader, SamplerIndex, SRV);
 	}
 
 	template <typename TShaderRHI>
@@ -2203,8 +2209,8 @@ public:
 		SetShaderResourceViewParameter(Shader.GetReference(), SamplerIndex, SRV);
 	}
 
-	template <typename TShaderRHIParamRef>
-	FORCEINLINE_DEBUGGABLE void SetShaderSampler(TShaderRHIParamRef Shader, uint32 SamplerIndex, FRHISamplerState* State)
+	template <typename TRHIShader>
+	FORCEINLINE_DEBUGGABLE void SetShaderSampler(TRHIShader* Shader, uint32 SamplerIndex, FRHISamplerState* State)
 	{
 		//check(IsOutsideRenderPass());
 		
@@ -2220,7 +2226,7 @@ public:
 			GetContext().RHISetShaderSampler(Shader, SamplerIndex, State);
 			return;
 		}
-		ALLOC_COMMAND(FRHICommandSetShaderSampler<TShaderRHIParamRef, ECmdList::EGfx>)(Shader, SamplerIndex, State);
+		ALLOC_COMMAND(FRHICommandSetShaderSampler<TRHIShader, ECmdList::EGfx>)(Shader, SamplerIndex, State);
 	}
 
 	template <typename TShaderRHI>
@@ -2229,14 +2235,14 @@ public:
 		SetShaderSampler(Shader.GetReference(), SamplerIndex, State);
 	}
 
-	FORCEINLINE_DEBUGGABLE void SetUAVParameter(FComputeShaderRHIParamRef Shader, uint32 UAVIndex, FRHIUnorderedAccessView* UAV)
+	FORCEINLINE_DEBUGGABLE void SetUAVParameter(FRHIComputeShader* Shader, uint32 UAVIndex, FRHIUnorderedAccessView* UAV)
 	{
 		if (Bypass())
 		{
 			GetContext().RHISetUAVParameter(Shader, UAVIndex, UAV);
 			return;
 		}
-		ALLOC_COMMAND(FRHICommandSetUAVParameter<FComputeShaderRHIParamRef, ECmdList::EGfx>)(Shader, UAVIndex, UAV);
+		ALLOC_COMMAND(FRHICommandSetUAVParameter<FRHIComputeShader, ECmdList::EGfx>)(Shader, UAVIndex, UAV);
 	}
 
 	FORCEINLINE_DEBUGGABLE void SetUAVParameter(TRefCountPtr<FRHIComputeShader>& Shader, uint32 UAVIndex, FRHIUnorderedAccessView* UAV)
@@ -2244,14 +2250,14 @@ public:
 		SetUAVParameter(Shader.GetReference(), UAVIndex, UAV);
 	}
 
-	FORCEINLINE_DEBUGGABLE void SetUAVParameter(FComputeShaderRHIParamRef Shader, uint32 UAVIndex, FRHIUnorderedAccessView* UAV, uint32 InitialCount)
+	FORCEINLINE_DEBUGGABLE void SetUAVParameter(FRHIComputeShader* Shader, uint32 UAVIndex, FRHIUnorderedAccessView* UAV, uint32 InitialCount)
 	{
 		if (Bypass())
 		{
 			GetContext().RHISetUAVParameter(Shader, UAVIndex, UAV, InitialCount);
 			return;
 		}
-		ALLOC_COMMAND(FRHICommandSetUAVParameter_IntialCount<FComputeShaderRHIParamRef, ECmdList::EGfx>)(Shader, UAVIndex, UAV, InitialCount);
+		ALLOC_COMMAND(FRHICommandSetUAVParameter_IntialCount<FRHIComputeShader, ECmdList::EGfx>)(Shader, UAVIndex, UAV, InitialCount);
 	}
 
 	FORCEINLINE_DEBUGGABLE void SetUAVParameter(TRefCountPtr<FRHIComputeShader>& Shader, uint32 UAVIndex, FRHIUnorderedAccessView* UAV, uint32 InitialCount)
@@ -2464,7 +2470,7 @@ public:
 		ALLOC_COMMAND(FRHICommandBindClearMRTValues)(bClearColor, bClearDepth, bClearStencil);
 	}	
 
-	FORCEINLINE_DEBUGGABLE void SetComputeShader(FComputeShaderRHIParamRef ComputeShader)
+	FORCEINLINE_DEBUGGABLE void SetComputeShader(FRHIComputeShader* ComputeShader)
 	{
 		ComputeShader->UpdateStats();
 		if (Bypass())
@@ -3166,14 +3172,14 @@ public:
 	void* operator new(size_t Size);
 	void operator delete(void *RawMemory);
 
-	FORCEINLINE_DEBUGGABLE void SetShaderUniformBuffer(FComputeShaderRHIParamRef Shader, uint32 BaseIndex, FUniformBufferRHIParamRef UniformBuffer)
+	FORCEINLINE_DEBUGGABLE void SetShaderUniformBuffer(FRHIComputeShader* Shader, uint32 BaseIndex, FUniformBufferRHIParamRef UniformBuffer)
 	{
 		if (Bypass())
 		{
 			GetComputeContext().RHISetShaderUniformBuffer(Shader, BaseIndex, UniformBuffer);
 			return;
 		}
-		ALLOC_COMMAND(FRHICommandSetShaderUniformBuffer<FComputeShaderRHIParamRef, ECmdList::ECompute>)(Shader, BaseIndex, UniformBuffer);		
+		ALLOC_COMMAND(FRHICommandSetShaderUniformBuffer<FRHIComputeShader, ECmdList::ECompute>)(Shader, BaseIndex, UniformBuffer);
 	}
 	
 	FORCEINLINE void SetShaderUniformBuffer(FComputeShaderRHIRef& Shader, uint32 BaseIndex, FUniformBufferRHIParamRef UniformBuffer)
@@ -3181,7 +3187,7 @@ public:
 		SetShaderUniformBuffer(Shader.GetReference(), BaseIndex, UniformBuffer);
 	}
 
-	FORCEINLINE_DEBUGGABLE void SetShaderParameter(FComputeShaderRHIParamRef Shader, uint32 BufferIndex, uint32 BaseIndex, uint32 NumBytes, const void* NewValue)
+	FORCEINLINE_DEBUGGABLE void SetShaderParameter(FRHIComputeShader* Shader, uint32 BufferIndex, uint32 BaseIndex, uint32 NumBytes, const void* NewValue)
 	{
 		if (Bypass())
 		{
@@ -3190,7 +3196,7 @@ public:
 		}
 		void* UseValue = Alloc(NumBytes, 16);
 		FMemory::Memcpy(UseValue, NewValue, NumBytes);
-		ALLOC_COMMAND(FRHICommandSetShaderParameter<FComputeShaderRHIParamRef, ECmdList::ECompute>)(Shader, BufferIndex, BaseIndex, NumBytes, UseValue);
+		ALLOC_COMMAND(FRHICommandSetShaderParameter<FRHIComputeShader, ECmdList::ECompute>)(Shader, BufferIndex, BaseIndex, NumBytes, UseValue);
 	}
 	
 	FORCEINLINE void SetShaderParameter(FComputeShaderRHIRef& Shader, uint32 BufferIndex, uint32 BaseIndex, uint32 NumBytes, const void* NewValue)
@@ -3198,27 +3204,27 @@ public:
 		SetShaderParameter(Shader.GetReference(), BufferIndex, BaseIndex, NumBytes, NewValue);
 	}
 	
-	FORCEINLINE_DEBUGGABLE void SetShaderTexture(FComputeShaderRHIParamRef Shader, uint32 TextureIndex, FTextureRHIParamRef Texture)
+	FORCEINLINE_DEBUGGABLE void SetShaderTexture(FRHIComputeShader* Shader, uint32 TextureIndex, FTextureRHIParamRef Texture)
 	{
 		if (Bypass())
 		{
 			GetComputeContext().RHISetShaderTexture(Shader, TextureIndex, Texture);
 			return;
 		}
-		ALLOC_COMMAND(FRHICommandSetShaderTexture<FComputeShaderRHIParamRef, ECmdList::ECompute>)(Shader, TextureIndex, Texture);
+		ALLOC_COMMAND(FRHICommandSetShaderTexture<FRHIComputeShader, ECmdList::ECompute>)(Shader, TextureIndex, Texture);
 	}
 	
-	FORCEINLINE_DEBUGGABLE void SetShaderResourceViewParameter(FComputeShaderRHIParamRef Shader, uint32 SamplerIndex, FRHIShaderResourceView* SRV)
+	FORCEINLINE_DEBUGGABLE void SetShaderResourceViewParameter(FRHIComputeShader* Shader, uint32 SamplerIndex, FRHIShaderResourceView* SRV)
 	{
 		if (Bypass())
 		{
 			GetComputeContext().RHISetShaderResourceViewParameter(Shader, SamplerIndex, SRV);
 			return;
 		}
-		ALLOC_COMMAND(FRHICommandSetShaderResourceViewParameter<FComputeShaderRHIParamRef, ECmdList::ECompute>)(Shader, SamplerIndex, SRV);
+		ALLOC_COMMAND(FRHICommandSetShaderResourceViewParameter<FRHIComputeShader, ECmdList::ECompute>)(Shader, SamplerIndex, SRV);
 	}
 	
-	FORCEINLINE_DEBUGGABLE void SetShaderSampler(FComputeShaderRHIParamRef Shader, uint32 SamplerIndex, FRHISamplerState* State)
+	FORCEINLINE_DEBUGGABLE void SetShaderSampler(FRHIComputeShader* Shader, uint32 SamplerIndex, FRHISamplerState* State)
 	{
 		// Immutable samplers can't be set dynamically
 		check(!State->IsImmutable());
@@ -3232,30 +3238,30 @@ public:
 			GetComputeContext().RHISetShaderSampler(Shader, SamplerIndex, State);
 			return;
 		}
-		ALLOC_COMMAND(FRHICommandSetShaderSampler<FComputeShaderRHIParamRef, ECmdList::ECompute>)(Shader, SamplerIndex, State);
+		ALLOC_COMMAND(FRHICommandSetShaderSampler<FRHIComputeShader, ECmdList::ECompute>)(Shader, SamplerIndex, State);
 	}
 
-	FORCEINLINE_DEBUGGABLE void SetUAVParameter(FComputeShaderRHIParamRef Shader, uint32 UAVIndex, FRHIUnorderedAccessView* UAV)
+	FORCEINLINE_DEBUGGABLE void SetUAVParameter(FRHIComputeShader* Shader, uint32 UAVIndex, FRHIUnorderedAccessView* UAV)
 	{
 		if (Bypass())
 		{
 			GetComputeContext().RHISetUAVParameter(Shader, UAVIndex, UAV);
 			return;
 		}
-		ALLOC_COMMAND(FRHICommandSetUAVParameter<FComputeShaderRHIParamRef, ECmdList::ECompute>)(Shader, UAVIndex, UAV);
+		ALLOC_COMMAND(FRHICommandSetUAVParameter<FRHIComputeShader, ECmdList::ECompute>)(Shader, UAVIndex, UAV);
 	}
 
-	FORCEINLINE_DEBUGGABLE void SetUAVParameter(FComputeShaderRHIParamRef Shader, uint32 UAVIndex, FRHIUnorderedAccessView* UAV, uint32 InitialCount)
+	FORCEINLINE_DEBUGGABLE void SetUAVParameter(FRHIComputeShader* Shader, uint32 UAVIndex, FRHIUnorderedAccessView* UAV, uint32 InitialCount)
 	{
 		if (Bypass())
 		{
 			GetComputeContext().RHISetUAVParameter(Shader, UAVIndex, UAV, InitialCount);
 			return;
 		}
-		ALLOC_COMMAND(FRHICommandSetUAVParameter_IntialCount<FComputeShaderRHIParamRef, ECmdList::ECompute>)(Shader, UAVIndex, UAV, InitialCount);
+		ALLOC_COMMAND(FRHICommandSetUAVParameter_IntialCount<FRHIComputeShader, ECmdList::ECompute>)(Shader, UAVIndex, UAV, InitialCount);
 	}
 	
-	FORCEINLINE_DEBUGGABLE void SetComputeShader(FComputeShaderRHIParamRef ComputeShader)
+	FORCEINLINE_DEBUGGABLE void SetComputeShader(FRHIComputeShader* ComputeShader)
 	{
 		ComputeShader->UpdateStats();
 		if (Bypass())
@@ -3642,7 +3648,7 @@ public:
 		return GDynamicRHI->RHICreateStagingBuffer();
 	}
 
-	FORCEINLINE FBoundShaderStateRHIRef CreateBoundShaderState(FRHIVertexDeclaration* VertexDeclaration, FVertexShaderRHIParamRef VertexShader, FHullShaderRHIParamRef HullShader, FDomainShaderRHIParamRef DomainShader, FPixelShaderRHIParamRef PixelShader, FGeometryShaderRHIParamRef GeometryShader)
+	FORCEINLINE FBoundShaderStateRHIRef CreateBoundShaderState(FRHIVertexDeclaration* VertexDeclaration, FRHIVertexShader* VertexShader, FRHIHullShader* HullShader, FRHIDomainShader* DomainShader, FRHIPixelShader* PixelShader, FRHIGeometryShader* GeometryShader)
 	{
 		LLM_SCOPE(ELLMTag::Shaders);
 		return RHICreateBoundShaderState(VertexDeclaration, VertexShader, HullShader, DomainShader, PixelShader, GeometryShader);
diff --git a/Engine/Source/Runtime/RHI/Public/RHICommandListCommandExecutes.inl b/Engine/Source/Runtime/RHI/Public/RHICommandListCommandExecutes.inl
--- a/Engine/Source/Runtime/RHI/Public/RHICommandListCommandExecutes.inl
+++ b/Engine/Source/Runtime/RHI/Public/RHICommandListCommandExecutes.inl
@@ -69,7 +69,7 @@ struct FRHICommandRayTraceDispatch;
 struct FRHICommandSetRayTracingBindings;
 
 enum class ECmdList;
-template <typename TShaderRHIParamRef> struct FRHICommandSetLocalUniformBuffer;
+template <typename TRHIShader> struct FRHICommandSetLocalUniformBuffer;
 
 void FRHICommandBeginUpdateMultiFrameResource::Execute(FRHICommandListBase& CmdList)
 {
@@ -101,109 +101,109 @@ void FRHICommandSetStencilRef::Execute(FRHICommandListBase& CmdList)
 	INTERNAL_DECORATOR(RHISetStencilRef)(StencilRef);
 }
 
-template <typename TShaderRHIParamRef, ECmdList CmdListType>
-void FRHICommandSetShaderParameter<TShaderRHIParamRef, CmdListType>::Execute(FRHICommandListBase& CmdList)
+template <typename TRHIShader, ECmdList CmdListType>
+void FRHICommandSetShaderParameter<TRHIShader, CmdListType>::Execute(FRHICommandListBase& CmdList)
 {
 	RHISTAT(SetShaderParameter);
 	INTERNAL_DECORATOR(RHISetShaderParameter)(Shader, BufferIndex, BaseIndex, NumBytes, NewValue); 
 }
-template struct FRHICommandSetShaderParameter<FVertexShaderRHIParamRef, ECmdList::EGfx>;
-template struct FRHICommandSetShaderParameter<FHullShaderRHIParamRef, ECmdList::EGfx>;
-template struct FRHICommandSetShaderParameter<FDomainShaderRHIParamRef, ECmdList::EGfx>;
-template struct FRHICommandSetShaderParameter<FGeometryShaderRHIParamRef, ECmdList::EGfx>;
-template struct FRHICommandSetShaderParameter<FPixelShaderRHIParamRef, ECmdList::EGfx>;
-template struct FRHICommandSetShaderParameter<FComputeShaderRHIParamRef, ECmdList::EGfx>;
-template<> void FRHICommandSetShaderParameter<FComputeShaderRHIParamRef, ECmdList::ECompute>::Execute(FRHICommandListBase& CmdList)
+template struct FRHICommandSetShaderParameter<FRHIVertexShader, ECmdList::EGfx>;
+template struct FRHICommandSetShaderParameter<FRHIHullShader, ECmdList::EGfx>;
+template struct FRHICommandSetShaderParameter<FRHIDomainShader, ECmdList::EGfx>;
+template struct FRHICommandSetShaderParameter<FRHIGeometryShader, ECmdList::EGfx>;
+template struct FRHICommandSetShaderParameter<FRHIPixelShader, ECmdList::EGfx>;
+template struct FRHICommandSetShaderParameter<FRHIComputeShader, ECmdList::EGfx>;
+template<> void FRHICommandSetShaderParameter<FRHIComputeShader, ECmdList::ECompute>::Execute(FRHICommandListBase& CmdList)
 {
 	RHISTAT(SetShaderParameter);
 	INTERNAL_DECORATOR_COMPUTE(RHISetShaderParameter)(Shader, BufferIndex, BaseIndex, NumBytes, NewValue);
 }
 
-template <typename TShaderRHIParamRef, ECmdList CmdListType>
-void FRHICommandSetShaderUniformBuffer<TShaderRHIParamRef, CmdListType>::Execute(FRHICommandListBase& CmdList)
+template <typename TRHIShader, ECmdList CmdListType>
+void FRHICommandSetShaderUniformBuffer<TRHIShader, CmdListType>::Execute(FRHICommandListBase& CmdList)
 {
 	RHISTAT(SetShaderUniformBuffer);
 	INTERNAL_DECORATOR(RHISetShaderUniformBuffer)(Shader, BaseIndex, UniformBuffer);
 }
-template struct FRHICommandSetShaderUniformBuffer<FVertexShaderRHIParamRef, ECmdList::EGfx>;
-template struct FRHICommandSetShaderUniformBuffer<FHullShaderRHIParamRef, ECmdList::EGfx>;
-template struct FRHICommandSetShaderUniformBuffer<FDomainShaderRHIParamRef, ECmdList::EGfx>;
-template struct FRHICommandSetShaderUniformBuffer<FGeometryShaderRHIParamRef, ECmdList::EGfx>;
-template struct FRHICommandSetShaderUniformBuffer<FPixelShaderRHIParamRef, ECmdList::EGfx>;
-template struct FRHICommandSetShaderUniformBuffer<FComputeShaderRHIParamRef, ECmdList::EGfx>;
-template<> void FRHICommandSetShaderUniformBuffer<FComputeShaderRHIParamRef, ECmdList::ECompute>::Execute(FRHICommandListBase& CmdList)
+template struct FRHICommandSetShaderUniformBuffer<FRHIVertexShader, ECmdList::EGfx>;
+template struct FRHICommandSetShaderUniformBuffer<FRHIHullShader, ECmdList::EGfx>;
+template struct FRHICommandSetShaderUniformBuffer<FRHIDomainShader, ECmdList::EGfx>;
+template struct FRHICommandSetShaderUniformBuffer<FRHIGeometryShader, ECmdList::EGfx>;
+template struct FRHICommandSetShaderUniformBuffer<FRHIPixelShader, ECmdList::EGfx>;
+template struct FRHICommandSetShaderUniformBuffer<FRHIComputeShader, ECmdList::EGfx>;
+template<> void FRHICommandSetShaderUniformBuffer<FRHIComputeShader, ECmdList::ECompute>::Execute(FRHICommandListBase& CmdList)
 {
 	RHISTAT(SetShaderUniformBuffer);
 	INTERNAL_DECORATOR_COMPUTE(RHISetShaderUniformBuffer)(Shader, BaseIndex, UniformBuffer);
 }
 
-template <typename TShaderRHIParamRef, ECmdList CmdListType>
-void FRHICommandSetShaderTexture<TShaderRHIParamRef, CmdListType>::Execute(FRHICommandListBase& CmdList)
+template <typename TRHIShader, ECmdList CmdListType>
+void FRHICommandSetShaderTexture<TRHIShader, CmdListType>::Execute(FRHICommandListBase& CmdList)
 {
 	RHISTAT(SetShaderTexture);
 	INTERNAL_DECORATOR(RHISetShaderTexture)(Shader, TextureIndex, Texture);
 }
-template struct FRHICommandSetShaderTexture<FVertexShaderRHIParamRef, ECmdList::EGfx>;
-template struct FRHICommandSetShaderTexture<FHullShaderRHIParamRef, ECmdList::EGfx>;
-template struct FRHICommandSetShaderTexture<FDomainShaderRHIParamRef, ECmdList::EGfx>;
-template struct FRHICommandSetShaderTexture<FGeometryShaderRHIParamRef, ECmdList::EGfx>;
-template struct FRHICommandSetShaderTexture<FPixelShaderRHIParamRef, ECmdList::EGfx>;
-template struct FRHICommandSetShaderTexture<FComputeShaderRHIParamRef, ECmdList::EGfx>;
-template<> void FRHICommandSetShaderTexture<FComputeShaderRHIParamRef, ECmdList::ECompute>::Execute(FRHICommandListBase& CmdList)
+template struct FRHICommandSetShaderTexture<FRHIVertexShader, ECmdList::EGfx>;
+template struct FRHICommandSetShaderTexture<FRHIHullShader, ECmdList::EGfx>;
+template struct FRHICommandSetShaderTexture<FRHIDomainShader, ECmdList::EGfx>;
+template struct FRHICommandSetShaderTexture<FRHIGeometryShader, ECmdList::EGfx>;
+template struct FRHICommandSetShaderTexture<FRHIPixelShader, ECmdList::EGfx>;
+template struct FRHICommandSetShaderTexture<FRHIComputeShader, ECmdList::EGfx>;
+template<> void FRHICommandSetShaderTexture<FRHIComputeShader, ECmdList::ECompute>::Execute(FRHICommandListBase& CmdList)
 {
 	RHISTAT(SetShaderTexture);
 	INTERNAL_DECORATOR_COMPUTE(RHISetShaderTexture)(Shader, TextureIndex, Texture);
 }
 
-template <typename TShaderRHIParamRef, ECmdList CmdListType>
-void FRHICommandSetShaderResourceViewParameter<TShaderRHIParamRef, CmdListType>::Execute(FRHICommandListBase& CmdList)
+template <typename TRHIShader, ECmdList CmdListType>
+void FRHICommandSetShaderResourceViewParameter<TRHIShader, CmdListType>::Execute(FRHICommandListBase& CmdList)
 {
 	RHISTAT(SetShaderResourceViewParameter);
 	INTERNAL_DECORATOR(RHISetShaderResourceViewParameter)(Shader, SamplerIndex, SRV);
 }
-template struct FRHICommandSetShaderResourceViewParameter<FVertexShaderRHIParamRef, ECmdList::EGfx>;
-template struct FRHICommandSetShaderResourceViewParameter<FHullShaderRHIParamRef, ECmdList::EGfx>;
-template struct FRHICommandSetShaderResourceViewParameter<FDomainShaderRHIParamRef, ECmdList::EGfx>;
-template struct FRHICommandSetShaderResourceViewParameter<FGeometryShaderRHIParamRef, ECmdList::EGfx>;
-template struct FRHICommandSetShaderResourceViewParameter<FPixelShaderRHIParamRef, ECmdList::EGfx>;
-template struct FRHICommandSetShaderResourceViewParameter<FComputeShaderRHIParamRef, ECmdList::EGfx>;
-template<> void FRHICommandSetShaderResourceViewParameter<FComputeShaderRHIParamRef, ECmdList::ECompute>::Execute(FRHICommandListBase& CmdList)
+template struct FRHICommandSetShaderResourceViewParameter<FRHIVertexShader, ECmdList::EGfx>;
+template struct FRHICommandSetShaderResourceViewParameter<FRHIHullShader, ECmdList::EGfx>;
+template struct FRHICommandSetShaderResourceViewParameter<FRHIDomainShader, ECmdList::EGfx>;
+template struct FRHICommandSetShaderResourceViewParameter<FRHIGeometryShader, ECmdList::EGfx>;
+template struct FRHICommandSetShaderResourceViewParameter<FRHIPixelShader, ECmdList::EGfx>;
+template struct FRHICommandSetShaderResourceViewParameter<FRHIComputeShader, ECmdList::EGfx>;
+template<> void FRHICommandSetShaderResourceViewParameter<FRHIComputeShader, ECmdList::ECompute>::Execute(FRHICommandListBase& CmdList)
 {
 	RHISTAT(SetShaderResourceViewParameter);
 	INTERNAL_DECORATOR_COMPUTE(RHISetShaderResourceViewParameter)(Shader, SamplerIndex, SRV);
 }
 
-template <typename TShaderRHIParamRef, ECmdList CmdListType>
-void FRHICommandSetUAVParameter<TShaderRHIParamRef, CmdListType>::Execute(FRHICommandListBase& CmdList)
+template <typename TRHIShader, ECmdList CmdListType>
+void FRHICommandSetUAVParameter<TRHIShader, CmdListType>::Execute(FRHICommandListBase& CmdList)
 {
 	RHISTAT(SetUAVParameter);
 	INTERNAL_DECORATOR_CONTEXT(RHISetUAVParameter)(Shader, UAVIndex, UAV);
 }
-template struct FRHICommandSetUAVParameter<FComputeShaderRHIParamRef, ECmdList::EGfx>;
-template struct FRHICommandSetUAVParameter<FComputeShaderRHIParamRef, ECmdList::ECompute>;
+template struct FRHICommandSetUAVParameter<FRHIComputeShader, ECmdList::EGfx>;
+template struct FRHICommandSetUAVParameter<FRHIComputeShader, ECmdList::ECompute>;
 
-template <typename TShaderRHIParamRef, ECmdList CmdListType>
-void FRHICommandSetUAVParameter_IntialCount<TShaderRHIParamRef, CmdListType>::Execute(FRHICommandListBase& CmdList)
+template <typename TRHIShader, ECmdList CmdListType>
+void FRHICommandSetUAVParameter_IntialCount<TRHIShader, CmdListType>::Execute(FRHICommandListBase& CmdList)
 {
 	RHISTAT(SetUAVParameter);
 	INTERNAL_DECORATOR_CONTEXT(RHISetUAVParameter)(Shader, UAVIndex, UAV, InitialCount);
 }
-template struct FRHICommandSetUAVParameter_IntialCount<FComputeShaderRHIParamRef, ECmdList::EGfx>;
-template struct FRHICommandSetUAVParameter_IntialCount<FComputeShaderRHIParamRef, ECmdList::ECompute>;
+template struct FRHICommandSetUAVParameter_IntialCount<FRHIComputeShader, ECmdList::EGfx>;
+template struct FRHICommandSetUAVParameter_IntialCount<FRHIComputeShader, ECmdList::ECompute>;
 
-template <typename TShaderRHIParamRef, ECmdList CmdListType>
-void FRHICommandSetShaderSampler<TShaderRHIParamRef, CmdListType>::Execute(FRHICommandListBase& CmdList)
+template <typename TRHIShader, ECmdList CmdListType>
+void FRHICommandSetShaderSampler<TRHIShader, CmdListType>::Execute(FRHICommandListBase& CmdList)
 {
 	RHISTAT(SetShaderSampler);
 	INTERNAL_DECORATOR(RHISetShaderSampler)(Shader, SamplerIndex, Sampler);
 }
-template struct FRHICommandSetShaderSampler<FVertexShaderRHIParamRef, ECmdList::EGfx>;
-template struct FRHICommandSetShaderSampler<FHullShaderRHIParamRef, ECmdList::EGfx>;
-template struct FRHICommandSetShaderSampler<FDomainShaderRHIParamRef, ECmdList::EGfx>;
-template struct FRHICommandSetShaderSampler<FGeometryShaderRHIParamRef, ECmdList::EGfx>;
-template struct FRHICommandSetShaderSampler<FPixelShaderRHIParamRef, ECmdList::EGfx>;
-template struct FRHICommandSetShaderSampler<FComputeShaderRHIParamRef, ECmdList::EGfx>;
-template<> void FRHICommandSetShaderSampler<FComputeShaderRHIParamRef, ECmdList::ECompute>::Execute(FRHICommandListBase& CmdList)
+template struct FRHICommandSetShaderSampler<FRHIVertexShader, ECmdList::EGfx>;
+template struct FRHICommandSetShaderSampler<FRHIHullShader, ECmdList::EGfx>;
+template struct FRHICommandSetShaderSampler<FRHIDomainShader, ECmdList::EGfx>;
+template struct FRHICommandSetShaderSampler<FRHIGeometryShader, ECmdList::EGfx>;
+template struct FRHICommandSetShaderSampler<FRHIPixelShader, ECmdList::EGfx>;
+template struct FRHICommandSetShaderSampler<FRHIComputeShader, ECmdList::EGfx>;
+template<> void FRHICommandSetShaderSampler<FRHIComputeShader, ECmdList::ECompute>::Execute(FRHICommandListBase& CmdList)
 {
 	RHISTAT(SetShaderSampler);
 	INTERNAL_DECORATOR_COMPUTE(RHISetShaderSampler)(Shader, SamplerIndex, Sampler);
@@ -481,8 +481,8 @@ void FRHICommandBuildLocalUniformBuffer::Execute(FRHICommandListBase& CmdList)
 	WorkArea.Contents = nullptr;
 }
 
-template <typename TShaderRHIParamRef>
-void FRHICommandSetLocalUniformBuffer<TShaderRHIParamRef>::Execute(FRHICommandListBase& CmdList)
+template <typename TRHIShader>
+void FRHICommandSetLocalUniformBuffer<TRHIShader>::Execute(FRHICommandListBase& CmdList)
 {
 	RHISTAT(SetLocalUniformBuffer);
 	check(LocalUniformBuffer.WorkArea->ComputedUniformBuffer->UseCount > 0 && IsValidRef(LocalUniformBuffer.WorkArea->ComputedUniformBuffer->UniformBuffer)); // this should have been created and should have uses outstanding
@@ -492,12 +492,12 @@ void FRHICommandSetLocalUniformBuffer<TShaderRHIParamRef>::Execute(FRHICommandLi
 		LocalUniformBuffer.WorkArea->ComputedUniformBuffer->~FComputedUniformBuffer();
 	}
 }
-template struct FRHICommandSetLocalUniformBuffer<FVertexShaderRHIParamRef>;
-template struct FRHICommandSetLocalUniformBuffer<FHullShaderRHIParamRef>;
-template struct FRHICommandSetLocalUniformBuffer<FDomainShaderRHIParamRef>;
-template struct FRHICommandSetLocalUniformBuffer<FGeometryShaderRHIParamRef>;
-template struct FRHICommandSetLocalUniformBuffer<FPixelShaderRHIParamRef>;
-template struct FRHICommandSetLocalUniformBuffer<FComputeShaderRHIParamRef>;
+template struct FRHICommandSetLocalUniformBuffer<FRHIVertexShader>;
+template struct FRHICommandSetLocalUniformBuffer<FRHIHullShader>;
+template struct FRHICommandSetLocalUniformBuffer<FRHIDomainShader>;
+template struct FRHICommandSetLocalUniformBuffer<FRHIGeometryShader>;
+template struct FRHICommandSetLocalUniformBuffer<FRHIPixelShader>;
+template struct FRHICommandSetLocalUniformBuffer<FRHIComputeShader>;
 
 void FRHICommandBeginRenderQuery::Execute(FRHICommandListBase& CmdList)
 {
diff --git a/Engine/Source/Runtime/RHI/Public/RHIContext.h b/Engine/Source/Runtime/RHI/Public/RHIContext.h
--- a/Engine/Source/Runtime/RHI/Public/RHIContext.h
+++ b/Engine/Source/Runtime/RHI/Public/RHIContext.h
@@ -40,7 +40,7 @@ public:
 	/**
 	*Sets the current compute shader.
 	*/
-	virtual void RHISetComputeShader(FComputeShaderRHIParamRef ComputeShader) = 0;
+	virtual void RHISetComputeShader(FRHIComputeShader* ComputeShader) = 0;
 
 	virtual void RHISetComputePipelineState(FRHIComputePipelineState* ComputePipelineState)
 	{
@@ -71,7 +71,7 @@ public:
 	virtual void RHITransitionResources(EResourceTransitionAccess TransitionType, EResourceTransitionPipeline TransitionPipeline, FRHIUnorderedAccessView** InUAVs, int32 NumUAVs, FRHIComputeFence* WriteComputeFence) = 0;
 
 	/** Set the shader resource view of a surface.  This is used for binding TextureMS parameter types that need a multi sampled view. */
-	virtual void RHISetShaderTexture(FComputeShaderRHIParamRef PixelShader, uint32 TextureIndex, FTextureRHIParamRef NewTexture) = 0;
+	virtual void RHISetShaderTexture(FRHIComputeShader* PixelShader, uint32 TextureIndex, FTextureRHIParamRef NewTexture) = 0;
 
 	/**
 	* Sets sampler state.
@@ -79,7 +79,7 @@ public:
 	* @param SamplerIndex		The index of the sampler.
 	* @param NewState			The new sampler state.
 	*/
-	virtual void RHISetShaderSampler(FComputeShaderRHIParamRef ComputeShader, uint32 SamplerIndex, FRHISamplerState* NewState) = 0;
+	virtual void RHISetShaderSampler(FRHIComputeShader* ComputeShader, uint32 SamplerIndex, FRHISamplerState* NewState) = 0;
 
 	/**
 	* Sets a compute shader UAV parameter.
@@ -87,7 +87,7 @@ public:
 	* @param UAVIndex		The index of the UAVIndex.
 	* @param UAV			The new UAV.
 	*/
-	virtual void RHISetUAVParameter(FComputeShaderRHIParamRef ComputeShader, uint32 UAVIndex, FRHIUnorderedAccessView* UAV) = 0;
+	virtual void RHISetUAVParameter(FRHIComputeShader* ComputeShader, uint32 UAVIndex, FRHIUnorderedAccessView* UAV) = 0;
 
 	/**
 	* Sets a compute shader counted UAV parameter and initial count
@@ -96,13 +96,13 @@ public:
 	* @param UAV			The new UAV.
 	* @param InitialCount	The initial number of items in the UAV.
 	*/
-	virtual void RHISetUAVParameter(FComputeShaderRHIParamRef ComputeShader, uint32 UAVIndex, FRHIUnorderedAccessView* UAV, uint32 InitialCount) = 0;
+	virtual void RHISetUAVParameter(FRHIComputeShader* ComputeShader, uint32 UAVIndex, FRHIUnorderedAccessView* UAV, uint32 InitialCount) = 0;
 
-	virtual void RHISetShaderResourceViewParameter(FComputeShaderRHIParamRef ComputeShader, uint32 SamplerIndex, FRHIShaderResourceView* SRV) = 0;
+	virtual void RHISetShaderResourceViewParameter(FRHIComputeShader* ComputeShader, uint32 SamplerIndex, FRHIShaderResourceView* SRV) = 0;
 
-	virtual void RHISetShaderUniformBuffer(FComputeShaderRHIParamRef ComputeShader, uint32 BufferIndex, FUniformBufferRHIParamRef Buffer) = 0;
+	virtual void RHISetShaderUniformBuffer(FRHIComputeShader* ComputeShader, uint32 BufferIndex, FUniformBufferRHIParamRef Buffer) = 0;
 
-	virtual void RHISetShaderParameter(FComputeShaderRHIParamRef ComputeShader, uint32 BufferIndex, uint32 BaseIndex, uint32 NumBytes, const void* NewValue) = 0;
+	virtual void RHISetShaderParameter(FRHIComputeShader* ComputeShader, uint32 BufferIndex, uint32 BaseIndex, uint32 NumBytes, const void* NewValue) = 0;
 
 	virtual void RHIPushEvent(const TCHAR* Name, FColor Color) = 0;
 
@@ -329,22 +329,22 @@ public:
 	virtual void RHISetGraphicsPipelineState(FRHIGraphicsPipelineState* GraphicsState) = 0;
 
 	/** Set the shader resource view of a surface.  This is used for binding TextureMS parameter types that need a multi sampled view. */
-	virtual void RHISetShaderTexture(FVertexShaderRHIParamRef VertexShader, uint32 TextureIndex, FTextureRHIParamRef NewTexture) = 0;
+	virtual void RHISetShaderTexture(FRHIVertexShader* VertexShader, uint32 TextureIndex, FTextureRHIParamRef NewTexture) = 0;
 
 	/** Set the shader resource view of a surface.  This is used for binding TextureMS parameter types that need a multi sampled view. */
-	virtual void RHISetShaderTexture(FHullShaderRHIParamRef HullShader, uint32 TextureIndex, FTextureRHIParamRef NewTexture) = 0;
+	virtual void RHISetShaderTexture(FRHIHullShader* HullShader, uint32 TextureIndex, FTextureRHIParamRef NewTexture) = 0;
 
 	/** Set the shader resource view of a surface.  This is used for binding TextureMS parameter types that need a multi sampled view. */
-	virtual void RHISetShaderTexture(FDomainShaderRHIParamRef DomainShader, uint32 TextureIndex, FTextureRHIParamRef NewTexture) = 0;
+	virtual void RHISetShaderTexture(FRHIDomainShader* DomainShader, uint32 TextureIndex, FTextureRHIParamRef NewTexture) = 0;
 
 	/** Set the shader resource view of a surface.  This is used for binding TextureMS parameter types that need a multi sampled view. */
-	virtual void RHISetShaderTexture(FGeometryShaderRHIParamRef GeometryShader, uint32 TextureIndex, FTextureRHIParamRef NewTexture) = 0;
+	virtual void RHISetShaderTexture(FRHIGeometryShader* GeometryShader, uint32 TextureIndex, FTextureRHIParamRef NewTexture) = 0;
 
 	/** Set the shader resource view of a surface.  This is used for binding TextureMS parameter types that need a multi sampled view. */
-	virtual void RHISetShaderTexture(FPixelShaderRHIParamRef PixelShader, uint32 TextureIndex, FTextureRHIParamRef NewTexture) = 0;
+	virtual void RHISetShaderTexture(FRHIPixelShader* PixelShader, uint32 TextureIndex, FTextureRHIParamRef NewTexture) = 0;
 
 	/** Set the shader resource view of a surface.  This is used for binding TextureMS parameter types that need a multi sampled view. */
-	virtual void RHISetShaderTexture(FComputeShaderRHIParamRef PixelShader, uint32 TextureIndex, FTextureRHIParamRef NewTexture) = 0;
+	virtual void RHISetShaderTexture(FRHIComputeShader* PixelShader, uint32 TextureIndex, FTextureRHIParamRef NewTexture) = 0;
 
 	/**
 	* Sets sampler state.
@@ -352,7 +352,7 @@ public:
 	* @param SamplerIndex		The index of the sampler.
 	* @param NewState			The new sampler state.
 	*/
-	virtual void RHISetShaderSampler(FComputeShaderRHIParamRef ComputeShader, uint32 SamplerIndex, FRHISamplerState* NewState) = 0;
+	virtual void RHISetShaderSampler(FRHIComputeShader* ComputeShader, uint32 SamplerIndex, FRHISamplerState* NewState) = 0;
 
 	/**
 	* Sets sampler state.
@@ -360,7 +360,7 @@ public:
 	* @param SamplerIndex		The index of the sampler.
 	* @param NewState			The new sampler state.
 	*/
-	virtual void RHISetShaderSampler(FVertexShaderRHIParamRef VertexShader, uint32 SamplerIndex, FRHISamplerState* NewState) = 0;
+	virtual void RHISetShaderSampler(FRHIVertexShader* VertexShader, uint32 SamplerIndex, FRHISamplerState* NewState) = 0;
 
 	/**
 	* Sets sampler state.
@@ -368,7 +368,7 @@ public:
 	* @param SamplerIndex		The index of the sampler.
 	* @param NewState			The new sampler state.
 	*/
-	virtual void RHISetShaderSampler(FGeometryShaderRHIParamRef GeometryShader, uint32 SamplerIndex, FRHISamplerState* NewState) = 0;
+	virtual void RHISetShaderSampler(FRHIGeometryShader* GeometryShader, uint32 SamplerIndex, FRHISamplerState* NewState) = 0;
 
 	/**
 	* Sets sampler state.
@@ -376,7 +376,7 @@ public:
 	* @param SamplerIndex		The index of the sampler.
 	* @param NewState			The new sampler state.
 	*/
-	virtual void RHISetShaderSampler(FDomainShaderRHIParamRef DomainShader, uint32 SamplerIndex, FRHISamplerState* NewState) = 0;
+	virtual void RHISetShaderSampler(FRHIDomainShader* DomainShader, uint32 SamplerIndex, FRHISamplerState* NewState) = 0;
 
 	/**
 	* Sets sampler state.
@@ -384,7 +384,7 @@ public:
 	* @param SamplerIndex		The index of the sampler.
 	* @param NewState			The new sampler state.
 	*/
-	virtual void RHISetShaderSampler(FHullShaderRHIParamRef HullShader, uint32 SamplerIndex, FRHISamplerState* NewState) = 0;
+	virtual void RHISetShaderSampler(FRHIHullShader* HullShader, uint32 SamplerIndex, FRHISamplerState* NewState) = 0;
 
 	/**
 	* Sets sampler state.
@@ -392,7 +392,7 @@ public:
 	* @param SamplerIndex		The index of the sampler.
 	* @param NewState			The new sampler state.
 	*/
-	virtual void RHISetShaderSampler(FPixelShaderRHIParamRef PixelShader, uint32 SamplerIndex, FRHISamplerState* NewState) = 0;
+	virtual void RHISetShaderSampler(FRHIPixelShader* PixelShader, uint32 SamplerIndex, FRHISamplerState* NewState) = 0;
 
 	/**
 	* Sets a compute shader UAV parameter.
@@ -400,7 +400,7 @@ public:
 	* @param UAVIndex		The index of the UAVIndex.
 	* @param UAV			The new UAV.
 	*/
-	virtual void RHISetUAVParameter(FComputeShaderRHIParamRef ComputeShader, uint32 UAVIndex, FRHIUnorderedAccessView* UAV) = 0;
+	virtual void RHISetUAVParameter(FRHIComputeShader* ComputeShader, uint32 UAVIndex, FRHIUnorderedAccessView* UAV) = 0;
 
 	/**
 	* Sets a compute shader counted UAV parameter and initial count
@@ -409,43 +409,43 @@ public:
 	* @param UAV			The new UAV.
 	* @param InitialCount	The initial number of items in the UAV.
 	*/
-	virtual void RHISetUAVParameter(FComputeShaderRHIParamRef ComputeShader, uint32 UAVIndex, FRHIUnorderedAccessView* UAV, uint32 InitialCount) = 0;
+	virtual void RHISetUAVParameter(FRHIComputeShader* ComputeShader, uint32 UAVIndex, FRHIUnorderedAccessView* UAV, uint32 InitialCount) = 0;
 
-	virtual void RHISetShaderResourceViewParameter(FPixelShaderRHIParamRef PixelShader, uint32 SamplerIndex, FRHIShaderResourceView* SRV) = 0;
+	virtual void RHISetShaderResourceViewParameter(FRHIPixelShader* PixelShader, uint32 SamplerIndex, FRHIShaderResourceView* SRV) = 0;
 
-	virtual void RHISetShaderResourceViewParameter(FVertexShaderRHIParamRef VertexShader, uint32 SamplerIndex, FRHIShaderResourceView* SRV) = 0;
+	virtual void RHISetShaderResourceViewParameter(FRHIVertexShader* VertexShader, uint32 SamplerIndex, FRHIShaderResourceView* SRV) = 0;
 
-	virtual void RHISetShaderResourceViewParameter(FComputeShaderRHIParamRef ComputeShader, uint32 SamplerIndex, FRHIShaderResourceView* SRV) = 0;
+	virtual void RHISetShaderResourceViewParameter(FRHIComputeShader* ComputeShader, uint32 SamplerIndex, FRHIShaderResourceView* SRV) = 0;
 
-	virtual void RHISetShaderResourceViewParameter(FHullShaderRHIParamRef HullShader, uint32 SamplerIndex, FRHIShaderResourceView* SRV) = 0;
+	virtual void RHISetShaderResourceViewParameter(FRHIHullShader* HullShader, uint32 SamplerIndex, FRHIShaderResourceView* SRV) = 0;
 
-	virtual void RHISetShaderResourceViewParameter(FDomainShaderRHIParamRef DomainShader, uint32 SamplerIndex, FRHIShaderResourceView* SRV) = 0;
+	virtual void RHISetShaderResourceViewParameter(FRHIDomainShader* DomainShader, uint32 SamplerIndex, FRHIShaderResourceView* SRV) = 0;
 
-	virtual void RHISetShaderResourceViewParameter(FGeometryShaderRHIParamRef GeometryShader, uint32 SamplerIndex, FRHIShaderResourceView* SRV) = 0;
+	virtual void RHISetShaderResourceViewParameter(FRHIGeometryShader* GeometryShader, uint32 SamplerIndex, FRHIShaderResourceView* SRV) = 0;
 
-	virtual void RHISetShaderUniformBuffer(FVertexShaderRHIParamRef VertexShader, uint32 BufferIndex, FUniformBufferRHIParamRef Buffer) = 0;
+	virtual void RHISetShaderUniformBuffer(FRHIVertexShader* VertexShader, uint32 BufferIndex, FUniformBufferRHIParamRef Buffer) = 0;
 
-	virtual void RHISetShaderUniformBuffer(FHullShaderRHIParamRef HullShader, uint32 BufferIndex, FUniformBufferRHIParamRef Buffer) = 0;
+	virtual void RHISetShaderUniformBuffer(FRHIHullShader* HullShader, uint32 BufferIndex, FUniformBufferRHIParamRef Buffer) = 0;
 
-	virtual void RHISetShaderUniformBuffer(FDomainShaderRHIParamRef DomainShader, uint32 BufferIndex, FUniformBufferRHIParamRef Buffer) = 0;
+	virtual void RHISetShaderUniformBuffer(FRHIDomainShader* DomainShader, uint32 BufferIndex, FUniformBufferRHIParamRef Buffer) = 0;
 
-	virtual void RHISetShaderUniformBuffer(FGeometryShaderRHIParamRef GeometryShader, uint32 BufferIndex, FUniformBufferRHIParamRef Buffer) = 0;
+	virtual void RHISetShaderUniformBuffer(FRHIGeometryShader* GeometryShader, uint32 BufferIndex, FUniformBufferRHIParamRef Buffer) = 0;
 
-	virtual void RHISetShaderUniformBuffer(FPixelShaderRHIParamRef PixelShader, uint32 BufferIndex, FUniformBufferRHIParamRef Buffer) = 0;
+	virtual void RHISetShaderUniformBuffer(FRHIPixelShader* PixelShader, uint32 BufferIndex, FUniformBufferRHIParamRef Buffer) = 0;
 
-	virtual void RHISetShaderUniformBuffer(FComputeShaderRHIParamRef ComputeShader, uint32 BufferIndex, FUniformBufferRHIParamRef Buffer) = 0;
+	virtual void RHISetShaderUniformBuffer(FRHIComputeShader* ComputeShader, uint32 BufferIndex, FUniformBufferRHIParamRef Buffer) = 0;
 
-	virtual void RHISetShaderParameter(FVertexShaderRHIParamRef VertexShader, uint32 BufferIndex, uint32 BaseIndex, uint32 NumBytes, const void* NewValue) = 0;
+	virtual void RHISetShaderParameter(FRHIVertexShader* VertexShader, uint32 BufferIndex, uint32 BaseIndex, uint32 NumBytes, const void* NewValue) = 0;
 
-	virtual void RHISetShaderParameter(FPixelShaderRHIParamRef PixelShader, uint32 BufferIndex, uint32 BaseIndex, uint32 NumBytes, const void* NewValue) = 0;
+	virtual void RHISetShaderParameter(FRHIPixelShader* PixelShader, uint32 BufferIndex, uint32 BaseIndex, uint32 NumBytes, const void* NewValue) = 0;
 
-	virtual void RHISetShaderParameter(FHullShaderRHIParamRef HullShader, uint32 BufferIndex, uint32 BaseIndex, uint32 NumBytes, const void* NewValue) = 0;
+	virtual void RHISetShaderParameter(FRHIHullShader* HullShader, uint32 BufferIndex, uint32 BaseIndex, uint32 NumBytes, const void* NewValue) = 0;
 
-	virtual void RHISetShaderParameter(FDomainShaderRHIParamRef DomainShader, uint32 BufferIndex, uint32 BaseIndex, uint32 NumBytes, const void* NewValue) = 0;
+	virtual void RHISetShaderParameter(FRHIDomainShader* DomainShader, uint32 BufferIndex, uint32 BaseIndex, uint32 NumBytes, const void* NewValue) = 0;
 
-	virtual void RHISetShaderParameter(FGeometryShaderRHIParamRef GeometryShader, uint32 BufferIndex, uint32 BaseIndex, uint32 NumBytes, const void* NewValue) = 0;
+	virtual void RHISetShaderParameter(FRHIGeometryShader* GeometryShader, uint32 BufferIndex, uint32 BaseIndex, uint32 NumBytes, const void* NewValue) = 0;
 
-	virtual void RHISetShaderParameter(FComputeShaderRHIParamRef ComputeShader, uint32 BufferIndex, uint32 BaseIndex, uint32 NumBytes, const void* NewValue) = 0;
+	virtual void RHISetShaderParameter(FRHIComputeShader* ComputeShader, uint32 BufferIndex, uint32 BaseIndex, uint32 NumBytes, const void* NewValue) = 0;
 
 	virtual void RHISetStencilRef(uint32 StencilRef) {}
 
@@ -650,11 +650,11 @@ public:
 
 FORCEINLINE FBoundShaderStateRHIRef RHICreateBoundShaderState(
 	FRHIVertexDeclaration* VertexDeclaration,
-	FVertexShaderRHIParamRef VertexShader,
-	FHullShaderRHIParamRef HullShader,
-	FDomainShaderRHIParamRef DomainShader,
-	FPixelShaderRHIParamRef PixelShader,
-	FGeometryShaderRHIParamRef GeometryShader
+	FRHIVertexShader* VertexShader,
+	FRHIHullShader* HullShader,
+	FRHIDomainShader* DomainShader,
+	FRHIPixelShader* PixelShader,
+	FRHIGeometryShader* GeometryShader
 );
 
 
diff --git a/Engine/Source/Runtime/RHI/Public/RHIResources.h b/Engine/Source/Runtime/RHI/Public/RHIResources.h
--- a/Engine/Source/Runtime/RHI/Public/RHIResources.h
+++ b/Engine/Source/Runtime/RHI/Public/RHIResources.h
@@ -1,6 +1,5 @@
 // Copyright 1998-2019 Epic Games, Inc. All Rights Reserved.
 
-
 #pragma once
 
 #include "CoreTypes.h"
@@ -1117,21 +1116,27 @@ UE_DEPRECATED(4.23, "FVertexDeclarationRHIParamRef typedef is deprecated; please
 typedef FRHIVertexDeclaration*              FVertexDeclarationRHIParamRef;
 typedef TRefCountPtr<FRHIVertexDeclaration> FVertexDeclarationRHIRef;
 
+UE_DEPRECATED(4.23, "FVertexShaderRHIParamRef typedef is deprecated; please use FRHIVertexShader* directly instead.")
 typedef FRHIVertexShader*              FVertexShaderRHIParamRef;
 typedef TRefCountPtr<FRHIVertexShader> FVertexShaderRHIRef;
 
+UE_DEPRECATED(4.23, "FHullShaderRHIParamRef typedef is deprecated; please use FRHIHullShader* directly instead.")
 typedef FRHIHullShader*              FHullShaderRHIParamRef;
 typedef TRefCountPtr<FRHIHullShader> FHullShaderRHIRef;
 
+UE_DEPRECATED(4.23, "FDomainShaderRHIParamRef typedef is deprecated; please use FRHIDomainShader* directly instead.")
 typedef FRHIDomainShader*              FDomainShaderRHIParamRef;
 typedef TRefCountPtr<FRHIDomainShader> FDomainShaderRHIRef;
 
+UE_DEPRECATED(4.23, "FPixelShaderRHIParamRef typedef is deprecated; please use FRHIPixelShader* directly instead.")
 typedef FRHIPixelShader*              FPixelShaderRHIParamRef;
 typedef TRefCountPtr<FRHIPixelShader> FPixelShaderRHIRef;
 
+UE_DEPRECATED(4.23, "FGeometryShaderRHIParamRef typedef is deprecated; please use FRHIGeometryShader* directly instead.")
 typedef FRHIGeometryShader*              FGeometryShaderRHIParamRef;
 typedef TRefCountPtr<FRHIGeometryShader> FGeometryShaderRHIRef;
 
+UE_DEPRECATED(4.23, "FComputeShaderRHIParamRef typedef is deprecated; please use FRHIComputeShader* directly instead.")
 typedef FRHIComputeShader*              FComputeShaderRHIParamRef;
 typedef TRefCountPtr<FRHIComputeShader> FComputeShaderRHIRef;
 
@@ -1738,18 +1743,18 @@ typedef TRefCountPtr<FRHICustomPresent> FCustomPresentRHIRef;
 
 // Template magic to convert an FRHI*Shader to its enum
 template<typename TRHIShader> struct TRHIShaderToEnum {};
-template<> struct TRHIShaderToEnum<FRHIVertexShader>	{ enum { ShaderFrequency = SF_Vertex }; };
-template<> struct TRHIShaderToEnum<FRHIHullShader>		{ enum { ShaderFrequency = SF_Hull }; };
-template<> struct TRHIShaderToEnum<FRHIDomainShader>	{ enum { ShaderFrequency = SF_Domain }; };
-template<> struct TRHIShaderToEnum<FRHIPixelShader>		{ enum { ShaderFrequency = SF_Pixel }; };
-template<> struct TRHIShaderToEnum<FRHIGeometryShader>	{ enum { ShaderFrequency = SF_Geometry }; };
-template<> struct TRHIShaderToEnum<FRHIComputeShader>	{ enum { ShaderFrequency = SF_Compute }; };
-template<> struct TRHIShaderToEnum<FVertexShaderRHIParamRef>	{ enum { ShaderFrequency = SF_Vertex }; };
-template<> struct TRHIShaderToEnum<FHullShaderRHIParamRef>		{ enum { ShaderFrequency = SF_Hull }; };
-template<> struct TRHIShaderToEnum<FDomainShaderRHIParamRef>	{ enum { ShaderFrequency = SF_Domain }; };
-template<> struct TRHIShaderToEnum<FPixelShaderRHIParamRef>		{ enum { ShaderFrequency = SF_Pixel }; };
-template<> struct TRHIShaderToEnum<FGeometryShaderRHIParamRef>	{ enum { ShaderFrequency = SF_Geometry }; };
-template<> struct TRHIShaderToEnum<FComputeShaderRHIParamRef>	{ enum { ShaderFrequency = SF_Compute }; };
+template<> struct TRHIShaderToEnum<FRHIVertexShader>		{ enum { ShaderFrequency = SF_Vertex }; };
+template<> struct TRHIShaderToEnum<FRHIHullShader>			{ enum { ShaderFrequency = SF_Hull }; };
+template<> struct TRHIShaderToEnum<FRHIDomainShader>		{ enum { ShaderFrequency = SF_Domain }; };
+template<> struct TRHIShaderToEnum<FRHIPixelShader>			{ enum { ShaderFrequency = SF_Pixel }; };
+template<> struct TRHIShaderToEnum<FRHIGeometryShader>		{ enum { ShaderFrequency = SF_Geometry }; };
+template<> struct TRHIShaderToEnum<FRHIComputeShader>		{ enum { ShaderFrequency = SF_Compute }; };
+template<> struct TRHIShaderToEnum<FRHIVertexShader*>		{ enum { ShaderFrequency = SF_Vertex }; };
+template<> struct TRHIShaderToEnum<FRHIHullShader*>			{ enum { ShaderFrequency = SF_Hull }; };
+template<> struct TRHIShaderToEnum<FRHIDomainShader*>		{ enum { ShaderFrequency = SF_Domain }; };
+template<> struct TRHIShaderToEnum<FRHIPixelShader*>		{ enum { ShaderFrequency = SF_Pixel }; };
+template<> struct TRHIShaderToEnum<FRHIGeometryShader*>		{ enum { ShaderFrequency = SF_Geometry }; };
+template<> struct TRHIShaderToEnum<FRHIComputeShader*>		{ enum { ShaderFrequency = SF_Compute }; };
 template<> struct TRHIShaderToEnum<FVertexShaderRHIRef>		{ enum { ShaderFrequency = SF_Vertex }; };
 template<> struct TRHIShaderToEnum<FHullShaderRHIRef>		{ enum { ShaderFrequency = SF_Hull }; };
 template<> struct TRHIShaderToEnum<FDomainShaderRHIRef>		{ enum { ShaderFrequency = SF_Domain }; };
@@ -1764,14 +1769,14 @@ struct FBoundShaderStateInput
 	inline FBoundShaderStateInput
 	(
 		FRHIVertexDeclaration* InVertexDeclarationRHI
-		, FVertexShaderRHIParamRef InVertexShaderRHI
+		, FRHIVertexShader* InVertexShaderRHI
 #if PLATFORM_SUPPORTS_TESSELLATION_SHADERS
-		, FHullShaderRHIParamRef InHullShaderRHI
-		, FDomainShaderRHIParamRef InDomainShaderRHI
+		, FRHIHullShader* InHullShaderRHI
+		, FRHIDomainShader* InDomainShaderRHI
 #endif
-		, FPixelShaderRHIParamRef InPixelShaderRHI
+		, FRHIPixelShader* InPixelShaderRHI
 #if PLATFORM_SUPPORTS_GEOMETRY_SHADERS
-		, FGeometryShaderRHIParamRef InGeometryShaderRHI
+		, FRHIGeometryShader* InGeometryShaderRHI
 #endif
 	)
 		: VertexDeclarationRHI(InVertexDeclarationRHI)
@@ -1788,11 +1793,11 @@ struct FBoundShaderStateInput
 	}
 
 	FRHIVertexDeclaration* VertexDeclarationRHI = nullptr;
-	FVertexShaderRHIParamRef VertexShaderRHI = nullptr;
-	FHullShaderRHIParamRef HullShaderRHI = nullptr;
-	FDomainShaderRHIParamRef DomainShaderRHI = nullptr;
-	FPixelShaderRHIParamRef PixelShaderRHI = nullptr;
-	FGeometryShaderRHIParamRef GeometryShaderRHI = nullptr;
+	FRHIVertexShader* VertexShaderRHI = nullptr;
+	FRHIHullShader* HullShaderRHI = nullptr;
+	FRHIDomainShader* DomainShaderRHI = nullptr;
+	FRHIPixelShader* PixelShaderRHI = nullptr;
+	FRHIGeometryShader* GeometryShaderRHI = nullptr;
 };
 
 struct FImmutableSamplerState
diff --git a/Engine/Source/Runtime/RHI/Public/RHIValidationContext.h b/Engine/Source/Runtime/RHI/Public/RHIValidationContext.h
--- a/Engine/Source/Runtime/RHI/Public/RHIValidationContext.h
+++ b/Engine/Source/Runtime/RHI/Public/RHIValidationContext.h
@@ -33,7 +33,7 @@ public:
 	*Sets the current compute shader.  Mostly for compliance with platforms
 	*that require shader setting before resource binding.
 	*/
-	virtual void RHISetComputeShader(FComputeShaderRHIParamRef Shader) override final
+	virtual void RHISetComputeShader(FRHIComputeShader* Shader) override final
 	{
 		State.bComputeShaderSet = true;
 		RHIContext->RHISetComputeShader(Shader);
@@ -76,7 +76,7 @@ public:
 	}
 
 	/** Set the shader resource view of a surface.  This is used for binding TextureMS parameter types that need a multi sampled view. */
-	virtual void RHISetShaderTexture(FComputeShaderRHIParamRef Shader, uint32 TextureIndex, FTextureRHIParamRef NewTexture) override final
+	virtual void RHISetShaderTexture(FRHIComputeShader* Shader, uint32 TextureIndex, FTextureRHIParamRef NewTexture) override final
 	{
 		checkf(State.bComputeShaderSet, TEXT("A Compute shader has to be set to set resources into a shader!"));
 		RHIContext->RHISetShaderTexture(Shader, TextureIndex, NewTexture);
@@ -88,7 +88,7 @@ public:
 	* @param SamplerIndex		The index of the sampler.
 	* @param NewState			The new sampler state.
 	*/
-	virtual void RHISetShaderSampler(FComputeShaderRHIParamRef Shader, uint32 SamplerIndex, FRHISamplerState* NewState) override final
+	virtual void RHISetShaderSampler(FRHIComputeShader* Shader, uint32 SamplerIndex, FRHISamplerState* NewState) override final
 	{
 		checkf(State.bComputeShaderSet, TEXT("A Compute shader has to be set to set resources into a shader!"));
 		RHIContext->RHISetShaderSampler(Shader, SamplerIndex, NewState);
@@ -100,7 +100,7 @@ public:
 	* @param UAVIndex		The index of the UAVIndex.
 	* @param UAV			The new UAV.
 	*/
-	virtual void RHISetUAVParameter(FComputeShaderRHIParamRef Shader, uint32 UAVIndex, FRHIUnorderedAccessView* UAV) override final
+	virtual void RHISetUAVParameter(FRHIComputeShader* Shader, uint32 UAVIndex, FRHIUnorderedAccessView* UAV) override final
 	{
 		checkf(State.bComputeShaderSet, TEXT("A Compute shader has to be set to set resources into a shader!"));
 		RHIContext->RHISetUAVParameter(Shader, UAVIndex, UAV);
@@ -113,25 +113,25 @@ public:
 	* @param UAV			The new UAV.
 	* @param InitialCount	The initial number of items in the UAV.
 	*/
-	virtual void RHISetUAVParameter(FComputeShaderRHIParamRef Shader, uint32 UAVIndex, FRHIUnorderedAccessView* UAV, uint32 InitialCount) override final
+	virtual void RHISetUAVParameter(FRHIComputeShader* Shader, uint32 UAVIndex, FRHIUnorderedAccessView* UAV, uint32 InitialCount) override final
 	{
 		checkf(State.bComputeShaderSet, TEXT("A Compute shader has to be set to set resources into a shader!"));
 		RHIContext->RHISetUAVParameter(Shader, UAVIndex, UAV, InitialCount);
 	}
 
-	virtual void RHISetShaderResourceViewParameter(FComputeShaderRHIParamRef Shader, uint32 SamplerIndex, FRHIShaderResourceView* SRV) override final
+	virtual void RHISetShaderResourceViewParameter(FRHIComputeShader* Shader, uint32 SamplerIndex, FRHIShaderResourceView* SRV) override final
 	{
 		checkf(State.bComputeShaderSet, TEXT("A Compute shader has to be set to set resources into a shader!"));
 		RHIContext->RHISetShaderResourceViewParameter(Shader, SamplerIndex, SRV);
 	}
 
-	virtual void RHISetShaderUniformBuffer(FComputeShaderRHIParamRef Shader, uint32 BufferIndex, FUniformBufferRHIParamRef Buffer) override final
+	virtual void RHISetShaderUniformBuffer(FRHIComputeShader* Shader, uint32 BufferIndex, FUniformBufferRHIParamRef Buffer) override final
 	{
 		checkf(State.bComputeShaderSet, TEXT("A Compute shader has to be set to set resources into a shader!"));
 		RHIContext->RHISetShaderUniformBuffer(Shader, BufferIndex, Buffer);
 	}
 
-	virtual void RHISetShaderParameter(FComputeShaderRHIParamRef Shader, uint32 BufferIndex, uint32 BaseIndex, uint32 NumBytes, const void* NewValue) override final
+	virtual void RHISetShaderParameter(FRHIComputeShader* Shader, uint32 BufferIndex, uint32 BaseIndex, uint32 NumBytes, const void* NewValue) override final
 	{
 		checkf(State.bComputeShaderSet, TEXT("A Compute shader has to be set to set resources into a shader!"));
 		RHIContext->RHISetShaderParameter(Shader, BufferIndex, BaseIndex, NumBytes, NewValue);
@@ -191,7 +191,7 @@ public:
 	*Sets the current compute shader.  Mostly for compliance with platforms
 	*that require shader setting before resource binding.
 	*/
-	virtual void RHISetComputeShader(FComputeShaderRHIParamRef Shader) override final
+	virtual void RHISetComputeShader(FRHIComputeShader* Shader) override final
 	{
 		State.bComputeShaderSet = true;
 		State.bGfxPSOSet = false;
@@ -415,42 +415,42 @@ public:
 	}
 
 	/** Set the shader resource view of a surface.  This is used for binding TextureMS parameter types that need a multi sampled view. */
-	virtual void RHISetShaderTexture(FVertexShaderRHIParamRef Shader, uint32 TextureIndex, FTextureRHIParamRef NewTexture) override final
+	virtual void RHISetShaderTexture(FRHIVertexShader* Shader, uint32 TextureIndex, FTextureRHIParamRef NewTexture) override final
 	{
 		checkf(State.bGfxPSOSet, TEXT("A Graphics PSO has to be set to set resources into a shader!"));
 		RHIContext->RHISetShaderTexture(Shader, TextureIndex, NewTexture);
 	}
 
 	/** Set the shader resource view of a surface.  This is used for binding TextureMS parameter types that need a multi sampled view. */
-	virtual void RHISetShaderTexture(FHullShaderRHIParamRef Shader, uint32 TextureIndex, FTextureRHIParamRef NewTexture) override final
+	virtual void RHISetShaderTexture(FRHIHullShader* Shader, uint32 TextureIndex, FTextureRHIParamRef NewTexture) override final
 	{
 		checkf(State.bGfxPSOSet, TEXT("A Graphics PSO has to be set to set resources into a shader!"));
 		RHIContext->RHISetShaderTexture(Shader, TextureIndex, NewTexture);
 	}
 
 	/** Set the shader resource view of a surface.  This is used for binding TextureMS parameter types that need a multi sampled view. */
-	virtual void RHISetShaderTexture(FDomainShaderRHIParamRef Shader, uint32 TextureIndex, FTextureRHIParamRef NewTexture) override final
+	virtual void RHISetShaderTexture(FRHIDomainShader* Shader, uint32 TextureIndex, FTextureRHIParamRef NewTexture) override final
 	{
 		checkf(State.bGfxPSOSet, TEXT("A Graphics PSO has to be set to set resources into a shader!"));
 		RHIContext->RHISetShaderTexture(Shader, TextureIndex, NewTexture);
 	}
 
 	/** Set the shader resource view of a surface.  This is used for binding TextureMS parameter types that need a multi sampled view. */
-	virtual void RHISetShaderTexture(FGeometryShaderRHIParamRef Shader, uint32 TextureIndex, FTextureRHIParamRef NewTexture) override final
+	virtual void RHISetShaderTexture(FRHIGeometryShader* Shader, uint32 TextureIndex, FTextureRHIParamRef NewTexture) override final
 	{
 		checkf(State.bGfxPSOSet, TEXT("A Graphics PSO has to be set to set resources into a shader!"));
 		RHIContext->RHISetShaderTexture(Shader, TextureIndex, NewTexture);
 	}
 
 	/** Set the shader resource view of a surface.  This is used for binding TextureMS parameter types that need a multi sampled view. */
-	virtual void RHISetShaderTexture(FPixelShaderRHIParamRef Shader, uint32 TextureIndex, FTextureRHIParamRef NewTexture) override final
+	virtual void RHISetShaderTexture(FRHIPixelShader* Shader, uint32 TextureIndex, FTextureRHIParamRef NewTexture) override final
 	{
 		checkf(State.bGfxPSOSet, TEXT("A Graphics PSO has to be set to set resources into a shader!"));
 		RHIContext->RHISetShaderTexture(Shader, TextureIndex, NewTexture);
 	}
 
 	/** Set the shader resource view of a surface.  This is used for binding TextureMS parameter types that need a multi sampled view. */
-	virtual void RHISetShaderTexture(FComputeShaderRHIParamRef Shader, uint32 TextureIndex, FTextureRHIParamRef NewTexture) override final
+	virtual void RHISetShaderTexture(FRHIComputeShader* Shader, uint32 TextureIndex, FTextureRHIParamRef NewTexture) override final
 	{
 		checkf(State.bComputeShaderSet, TEXT("A Compute shader has to be set to set resources into a shader!"));
 		RHIContext->RHISetShaderTexture(Shader, TextureIndex, NewTexture);
@@ -462,7 +462,7 @@ public:
 	* @param SamplerIndex		The index of the sampler.
 	* @param NewState			The new sampler state.
 	*/
-	virtual void RHISetShaderSampler(FComputeShaderRHIParamRef Shader, uint32 SamplerIndex, FRHISamplerState* NewState) override final
+	virtual void RHISetShaderSampler(FRHIComputeShader* Shader, uint32 SamplerIndex, FRHISamplerState* NewState) override final
 	{
 		checkf(State.bComputeShaderSet, TEXT("A Compute shader has to be set to set resources into a shader!"));
 		RHIContext->RHISetShaderSampler(Shader, SamplerIndex, NewState);
@@ -474,7 +474,7 @@ public:
 	* @param SamplerIndex		The index of the sampler.
 	* @param NewState			The new sampler state.
 	*/
-	virtual void RHISetShaderSampler(FVertexShaderRHIParamRef Shader, uint32 SamplerIndex, FRHISamplerState* NewState) override final
+	virtual void RHISetShaderSampler(FRHIVertexShader* Shader, uint32 SamplerIndex, FRHISamplerState* NewState) override final
 	{
 		checkf(State.bGfxPSOSet, TEXT("A Graphics PSO has to be set to set resources into a shader!"));
 		RHIContext->RHISetShaderSampler(Shader, SamplerIndex, NewState);
@@ -487,7 +487,7 @@ public:
 	* @param SamplerIndex		The index of the sampler.
 	* @param NewState			The new sampler state.
 	*/
-	virtual void RHISetShaderSampler(FGeometryShaderRHIParamRef Shader, uint32 SamplerIndex, FRHISamplerState* NewState) override final
+	virtual void RHISetShaderSampler(FRHIGeometryShader* Shader, uint32 SamplerIndex, FRHISamplerState* NewState) override final
 	{
 		checkf(State.bGfxPSOSet, TEXT("A Graphics PSO has to be set to set resources into a shader!"));
 		RHIContext->RHISetShaderSampler(Shader, SamplerIndex, NewState);
@@ -499,7 +499,7 @@ public:
 	* @param SamplerIndex		The index of the sampler.
 	* @param NewState			The new sampler state.
 	*/
-	virtual void RHISetShaderSampler(FDomainShaderRHIParamRef Shader, uint32 SamplerIndex, FRHISamplerState* NewState) override final
+	virtual void RHISetShaderSampler(FRHIDomainShader* Shader, uint32 SamplerIndex, FRHISamplerState* NewState) override final
 	{
 		checkf(State.bGfxPSOSet, TEXT("A Graphics PSO has to be set to set resources into a shader!"));
 		RHIContext->RHISetShaderSampler(Shader, SamplerIndex, NewState);
@@ -511,7 +511,7 @@ public:
 	* @param SamplerIndex		The index of the sampler.
 	* @param NewState			The new sampler state.
 	*/
-	virtual void RHISetShaderSampler(FHullShaderRHIParamRef Shader, uint32 SamplerIndex, FRHISamplerState* NewState) override final
+	virtual void RHISetShaderSampler(FRHIHullShader* Shader, uint32 SamplerIndex, FRHISamplerState* NewState) override final
 	{
 		checkf(State.bGfxPSOSet, TEXT("A Graphics PSO has to be set to set resources into a shader!"));
 		RHIContext->RHISetShaderSampler(Shader, SamplerIndex, NewState);
@@ -523,7 +523,7 @@ public:
 	* @param SamplerIndex		The index of the sampler.
 	* @param NewState			The new sampler state.
 	*/
-	virtual void RHISetShaderSampler(FPixelShaderRHIParamRef Shader, uint32 SamplerIndex, FRHISamplerState* NewState) override final
+	virtual void RHISetShaderSampler(FRHIPixelShader* Shader, uint32 SamplerIndex, FRHISamplerState* NewState) override final
 	{
 		checkf(State.bGfxPSOSet, TEXT("A Graphics PSO has to be set to set resources into a shader!"));
 		RHIContext->RHISetShaderSampler(Shader, SamplerIndex, NewState);
@@ -535,7 +535,7 @@ public:
 	* @param UAVIndex		The index of the UAVIndex.
 	* @param UAV			The new UAV.
 	*/
-	virtual void RHISetUAVParameter(FComputeShaderRHIParamRef Shader, uint32 UAVIndex, FRHIUnorderedAccessView* UAV) override final
+	virtual void RHISetUAVParameter(FRHIComputeShader* Shader, uint32 UAVIndex, FRHIUnorderedAccessView* UAV) override final
 	{
 		checkf(State.bComputeShaderSet, TEXT("A Compute shader has to be set to set resources into a shader!"));
 		RHIContext->RHISetUAVParameter(Shader, UAVIndex, UAV);
@@ -548,115 +548,115 @@ public:
 	* @param UAV			The new UAV.
 	* @param InitialCount	The initial number of items in the UAV.
 	*/
-	virtual void RHISetUAVParameter(FComputeShaderRHIParamRef Shader, uint32 UAVIndex, FRHIUnorderedAccessView* UAV, uint32 InitialCount) override final
+	virtual void RHISetUAVParameter(FRHIComputeShader* Shader, uint32 UAVIndex, FRHIUnorderedAccessView* UAV, uint32 InitialCount) override final
 	{
 		checkf(State.bComputeShaderSet, TEXT("A Compute shader has to be set to set resources into a shader!"));
 		RHIContext->RHISetUAVParameter(Shader, UAVIndex, UAV, InitialCount);
 	}
 
-	virtual void RHISetShaderResourceViewParameter(FPixelShaderRHIParamRef Shader, uint32 SamplerIndex, FRHIShaderResourceView* SRV) override final
+	virtual void RHISetShaderResourceViewParameter(FRHIPixelShader* Shader, uint32 SamplerIndex, FRHIShaderResourceView* SRV) override final
 	{
 		checkf(State.bGfxPSOSet, TEXT("A Graphics PSO has to be set to set resources into a shader!"));
 		RHIContext->RHISetShaderResourceViewParameter(Shader, SamplerIndex, SRV);
 	}
 
-	virtual void RHISetShaderResourceViewParameter(FVertexShaderRHIParamRef Shader, uint32 SamplerIndex, FRHIShaderResourceView* SRV) override final
+	virtual void RHISetShaderResourceViewParameter(FRHIVertexShader* Shader, uint32 SamplerIndex, FRHIShaderResourceView* SRV) override final
 	{
 		checkf(State.bGfxPSOSet, TEXT("A Graphics PSO has to be set to set resources into a shader!"));
 		RHIContext->RHISetShaderResourceViewParameter(Shader, SamplerIndex, SRV);
 	}
 
-	virtual void RHISetShaderResourceViewParameter(FComputeShaderRHIParamRef Shader, uint32 SamplerIndex, FRHIShaderResourceView* SRV) override final
+	virtual void RHISetShaderResourceViewParameter(FRHIComputeShader* Shader, uint32 SamplerIndex, FRHIShaderResourceView* SRV) override final
 	{
 		checkf(State.bComputeShaderSet, TEXT("A Compute shader has to be set to set resources into a shader!"));
 		RHIContext->RHISetShaderResourceViewParameter(Shader, SamplerIndex, SRV);
 	}
 
-	virtual void RHISetShaderResourceViewParameter(FHullShaderRHIParamRef Shader, uint32 SamplerIndex, FRHIShaderResourceView* SRV) override final
+	virtual void RHISetShaderResourceViewParameter(FRHIHullShader* Shader, uint32 SamplerIndex, FRHIShaderResourceView* SRV) override final
 	{
 		checkf(State.bGfxPSOSet, TEXT("A Graphics PSO has to be set to set resources into a shader!"));
 		RHIContext->RHISetShaderResourceViewParameter(Shader, SamplerIndex, SRV);
 	}
 
-	virtual void RHISetShaderResourceViewParameter(FDomainShaderRHIParamRef Shader, uint32 SamplerIndex, FRHIShaderResourceView* SRV) override final
+	virtual void RHISetShaderResourceViewParameter(FRHIDomainShader* Shader, uint32 SamplerIndex, FRHIShaderResourceView* SRV) override final
 	{
 		checkf(State.bGfxPSOSet, TEXT("A Graphics PSO has to be set to set resources into a shader!"));
 		RHIContext->RHISetShaderResourceViewParameter(Shader, SamplerIndex, SRV);
 	}
 
-	virtual void RHISetShaderResourceViewParameter(FGeometryShaderRHIParamRef Shader, uint32 SamplerIndex, FRHIShaderResourceView* SRV) override final
+	virtual void RHISetShaderResourceViewParameter(FRHIGeometryShader* Shader, uint32 SamplerIndex, FRHIShaderResourceView* SRV) override final
 	{
 		checkf(State.bGfxPSOSet, TEXT("A Graphics PSO has to be set to set resources into a shader!"));
 		RHIContext->RHISetShaderResourceViewParameter(Shader, SamplerIndex, SRV);
 	}
 
-	virtual void RHISetShaderUniformBuffer(FVertexShaderRHIParamRef Shader, uint32 BufferIndex, FUniformBufferRHIParamRef Buffer) override final
+	virtual void RHISetShaderUniformBuffer(FRHIVertexShader* Shader, uint32 BufferIndex, FUniformBufferRHIParamRef Buffer) override final
 	{
 		checkf(State.bGfxPSOSet, TEXT("A Graphics PSO has to be set to set resources into a shader!"));
 		RHIContext->RHISetShaderUniformBuffer(Shader, BufferIndex, Buffer);
 	}
 
-	virtual void RHISetShaderUniformBuffer(FHullShaderRHIParamRef Shader, uint32 BufferIndex, FUniformBufferRHIParamRef Buffer) override final
+	virtual void RHISetShaderUniformBuffer(FRHIHullShader* Shader, uint32 BufferIndex, FUniformBufferRHIParamRef Buffer) override final
 	{
 		checkf(State.bGfxPSOSet, TEXT("A Graphics PSO has to be set to set resources into a shader!"));
 		RHIContext->RHISetShaderUniformBuffer(Shader, BufferIndex, Buffer);
 	}
 
-	virtual void RHISetShaderUniformBuffer(FDomainShaderRHIParamRef Shader, uint32 BufferIndex, FUniformBufferRHIParamRef Buffer) override final
+	virtual void RHISetShaderUniformBuffer(FRHIDomainShader* Shader, uint32 BufferIndex, FUniformBufferRHIParamRef Buffer) override final
 	{
 		checkf(State.bGfxPSOSet, TEXT("A Graphics PSO has to be set to set resources into a shader!"));
 		RHIContext->RHISetShaderUniformBuffer(Shader, BufferIndex, Buffer);
 	}
 
-	virtual void RHISetShaderUniformBuffer(FGeometryShaderRHIParamRef Shader, uint32 BufferIndex, FUniformBufferRHIParamRef Buffer) override final
+	virtual void RHISetShaderUniformBuffer(FRHIGeometryShader* Shader, uint32 BufferIndex, FUniformBufferRHIParamRef Buffer) override final
 	{
 		checkf(State.bGfxPSOSet, TEXT("A Graphics PSO has to be set to set resources into a shader!"));
 		RHIContext->RHISetShaderUniformBuffer(Shader, BufferIndex, Buffer);
 	}
 
-	virtual void RHISetShaderUniformBuffer(FPixelShaderRHIParamRef Shader, uint32 BufferIndex, FUniformBufferRHIParamRef Buffer) override final
+	virtual void RHISetShaderUniformBuffer(FRHIPixelShader* Shader, uint32 BufferIndex, FUniformBufferRHIParamRef Buffer) override final
 	{
 		checkf(State.bGfxPSOSet, TEXT("A Graphics PSO has to be set to set resources into a shader!"));
 		RHIContext->RHISetShaderUniformBuffer(Shader, BufferIndex, Buffer);
 	}
 
-	virtual void RHISetShaderUniformBuffer(FComputeShaderRHIParamRef Shader, uint32 BufferIndex, FUniformBufferRHIParamRef Buffer) override final
+	virtual void RHISetShaderUniformBuffer(FRHIComputeShader* Shader, uint32 BufferIndex, FUniformBufferRHIParamRef Buffer) override final
 	{
 		checkf(State.bComputeShaderSet, TEXT("A Compute shader has to be set to set resources into a shader!"));
 		RHIContext->RHISetShaderUniformBuffer(Shader, BufferIndex, Buffer);
 	}
 
-	virtual void RHISetShaderParameter(FVertexShaderRHIParamRef Shader, uint32 BufferIndex, uint32 BaseIndex, uint32 NumBytes, const void* NewValue) override final
+	virtual void RHISetShaderParameter(FRHIVertexShader* Shader, uint32 BufferIndex, uint32 BaseIndex, uint32 NumBytes, const void* NewValue) override final
 	{
 		checkf(State.bGfxPSOSet, TEXT("A Graphics PSO has to be set to set resources into a shader!"));
 		RHIContext->RHISetShaderParameter(Shader, BufferIndex, BaseIndex, NumBytes, NewValue);
 	}
 
-	virtual void RHISetShaderParameter(FPixelShaderRHIParamRef Shader, uint32 BufferIndex, uint32 BaseIndex, uint32 NumBytes, const void* NewValue) override final
+	virtual void RHISetShaderParameter(FRHIPixelShader* Shader, uint32 BufferIndex, uint32 BaseIndex, uint32 NumBytes, const void* NewValue) override final
 	{
 		checkf(State.bGfxPSOSet, TEXT("A Graphics PSO has to be set to set resources into a shader!"));
 		RHIContext->RHISetShaderParameter(Shader, BufferIndex, BaseIndex, NumBytes, NewValue);
 	}
 
-	virtual void RHISetShaderParameter(FHullShaderRHIParamRef Shader, uint32 BufferIndex, uint32 BaseIndex, uint32 NumBytes, const void* NewValue) override final
+	virtual void RHISetShaderParameter(FRHIHullShader* Shader, uint32 BufferIndex, uint32 BaseIndex, uint32 NumBytes, const void* NewValue) override final
 	{
 		checkf(State.bGfxPSOSet, TEXT("A Graphics PSO has to be set to set resources into a shader!"));
 		RHIContext->RHISetShaderParameter(Shader, BufferIndex, BaseIndex, NumBytes, NewValue);
 	}
 
-	virtual void RHISetShaderParameter(FDomainShaderRHIParamRef Shader, uint32 BufferIndex, uint32 BaseIndex, uint32 NumBytes, const void* NewValue) override final
+	virtual void RHISetShaderParameter(FRHIDomainShader* Shader, uint32 BufferIndex, uint32 BaseIndex, uint32 NumBytes, const void* NewValue) override final
 	{
 		checkf(State.bGfxPSOSet, TEXT("A Graphics PSO has to be set to set resources into a shader!"));
 		RHIContext->RHISetShaderParameter(Shader, BufferIndex, BaseIndex, NumBytes, NewValue);
 	}
 
-	virtual void RHISetShaderParameter(FGeometryShaderRHIParamRef Shader, uint32 BufferIndex, uint32 BaseIndex, uint32 NumBytes, const void* NewValue) override final
+	virtual void RHISetShaderParameter(FRHIGeometryShader* Shader, uint32 BufferIndex, uint32 BaseIndex, uint32 NumBytes, const void* NewValue) override final
 	{
 		checkf(State.bGfxPSOSet, TEXT("A Graphics PSO has to be set to set resources into a shader!"));
 		RHIContext->RHISetShaderParameter(Shader, BufferIndex, BaseIndex, NumBytes, NewValue);
 	}
 
-	virtual void RHISetShaderParameter(FComputeShaderRHIParamRef Shader, uint32 BufferIndex, uint32 BaseIndex, uint32 NumBytes, const void* NewValue) override final
+	virtual void RHISetShaderParameter(FRHIComputeShader* Shader, uint32 BufferIndex, uint32 BaseIndex, uint32 NumBytes, const void* NewValue) override final
 	{
 		checkf(State.bComputeShaderSet, TEXT("A Compute shader has to be set to set resources into a shader!"));
 		RHIContext->RHISetShaderParameter(Shader, BufferIndex, BaseIndex, NumBytes, NewValue);
diff --git a/Engine/Source/Runtime/RenderCore/Private/StaticBoundShaderState.cpp b/Engine/Source/Runtime/RenderCore/Private/StaticBoundShaderState.cpp
--- a/Engine/Source/Runtime/RenderCore/Private/StaticBoundShaderState.cpp
+++ b/Engine/Source/Runtime/RenderCore/Private/StaticBoundShaderState.cpp
@@ -52,9 +52,9 @@ FGlobalBoundShaderStateResource::~FGlobalBoundShaderStateResource()
  */
 FRHIBoundShaderState* FGlobalBoundShaderStateResource::GetInitializedRHI(
 	FRHIVertexDeclaration* VertexDeclaration,
-	FVertexShaderRHIParamRef VertexShader, 
-	FPixelShaderRHIParamRef PixelShader,
-	FGeometryShaderRHIParamRef GeometryShader
+	FRHIVertexShader* VertexShader,
+	FRHIPixelShader* PixelShader,
+	FRHIGeometryShader* GeometryShader
 	)
 {
 	check(IsInitialized());
diff --git a/Engine/Source/Runtime/RenderCore/Public/Shader.h b/Engine/Source/Runtime/RenderCore/Public/Shader.h
--- a/Engine/Source/Runtime/RenderCore/Public/Shader.h
+++ b/Engine/Source/Runtime/RenderCore/Public/Shader.h
@@ -312,7 +312,7 @@ public:
 	RENDERCORE_API void Register();
 
 	/** @return the shader's vertex shader */
-	FORCEINLINE const FVertexShaderRHIParamRef GetVertexShader()
+	FORCEINLINE FRHIVertexShader* GetVertexShader()
 	{
 		checkSlow(Target.Frequency == SF_Vertex);
 		if (!IsInitialized())
@@ -322,7 +322,7 @@ public:
 		return (FRHIVertexShader*)Shader.GetReference();
 	}
 	/** @return the shader's pixel shader */
-	FORCEINLINE const FPixelShaderRHIParamRef GetPixelShader()
+	FORCEINLINE FRHIPixelShader* GetPixelShader()
 	{
 		checkSlow(Target.Frequency == SF_Pixel);
 		if (!IsInitialized())
@@ -332,7 +332,7 @@ public:
 		return (FRHIPixelShader*)Shader.GetReference();
 	}
 	/** @return the shader's hull shader */
-	FORCEINLINE const FHullShaderRHIParamRef GetHullShader()
+	FORCEINLINE FRHIHullShader* GetHullShader()
 	{
 		checkSlow(Target.Frequency == SF_Hull);
 		if (!IsInitialized())
@@ -342,7 +342,7 @@ public:
 		return (FRHIHullShader*)Shader.GetReference();
 	}
 	/** @return the shader's domain shader */
-	FORCEINLINE const FDomainShaderRHIParamRef GetDomainShader()
+	FORCEINLINE FRHIDomainShader* GetDomainShader()
 	{
 		checkSlow(Target.Frequency == SF_Domain);
 		if (!IsInitialized())
@@ -352,7 +352,7 @@ public:
 		return (FRHIDomainShader*)Shader.GetReference();
 	}
 	/** @return the shader's geometry shader */
-	FORCEINLINE const FGeometryShaderRHIParamRef GetGeometryShader()
+	FORCEINLINE FRHIGeometryShader* GetGeometryShader()
 	{
 		checkSlow(Target.Frequency == SF_Geometry);
 		if (!IsInitialized())
@@ -362,7 +362,7 @@ public:
 		return (FRHIGeometryShader*)Shader.GetReference();
 	}
 	/** @return the shader's compute shader */
-	FORCEINLINE const FComputeShaderRHIParamRef GetComputeShader()
+	FORCEINLINE FRHIComputeShader* GetComputeShader()
 	{
 		checkSlow(Target.Frequency == SF_Compute);
 		if (!IsInitialized())
@@ -907,32 +907,32 @@ public:
 	virtual const FVertexFactoryParameterRef* GetVertexFactoryParameterRef() const { return NULL; }
 
 	/** @return the shader's vertex shader */
-	inline const FVertexShaderRHIParamRef GetVertexShader() const
+	inline FRHIVertexShader* GetVertexShader() const
 	{
 		return Resource->GetVertexShader();
 	}
 	/** @return the shader's pixel shader */
-	inline const FPixelShaderRHIParamRef GetPixelShader() const
+	inline FRHIPixelShader* GetPixelShader() const
 	{
 		return Resource->GetPixelShader();
 	}
 	/** @return the shader's hull shader */
-	inline const FHullShaderRHIParamRef GetHullShader() const
+	inline FRHIHullShader* GetHullShader() const
 	{
 		return Resource->GetHullShader();
 	}
 	/** @return the shader's domain shader */
-	inline const FDomainShaderRHIParamRef GetDomainShader() const
+	inline FRHIDomainShader* GetDomainShader() const
 	{
 		return Resource->GetDomainShader();
 	}
 	/** @return the shader's geometry shader */
-	inline const FGeometryShaderRHIParamRef GetGeometryShader() const
+	inline FRHIGeometryShader* GetGeometryShader() const
 	{
 		return Resource->GetGeometryShader();
 	}
 	/** @return the shader's compute shader */
-	inline const FComputeShaderRHIParamRef GetComputeShader() const
+	inline FRHIComputeShader* GetComputeShader() const
 	{
 		return Resource->GetComputeShader();
 	}
diff --git a/Engine/Source/Runtime/RenderCore/Public/ShaderParameterStruct.h b/Engine/Source/Runtime/RenderCore/Public/ShaderParameterStruct.h
--- a/Engine/Source/Runtime/RenderCore/Public/ShaderParameterStruct.h
+++ b/Engine/Source/Runtime/RenderCore/Public/ShaderParameterStruct.h
@@ -87,13 +87,13 @@ inline void SetShaderUAVs(TRHICmdList& RHICmdList, const TShaderClass* Shader, T
 }
 
 template<typename TRHICmdList, typename TShaderClass>
-inline void SetShaderUAVs(TRHICmdList& RHICmdList, const TShaderClass* Shader, FPixelShaderRHIParamRef ShadeRHI, const typename TShaderClass::FParameters& Parameters)
+inline void SetShaderUAVs(TRHICmdList& RHICmdList, const TShaderClass* Shader, FRHIPixelShader* ShadeRHI, const typename TShaderClass::FParameters& Parameters)
 {
 	// Pixelshader UAVs are bound together with rendertargets using BeginRenderPass
 }
 
 template<typename TRHICmdList, typename TShaderClass>
-inline void SetShaderUAVs(TRHICmdList& RHICmdList, const TShaderClass* Shader, FComputeShaderRHIParamRef ShadeRHI, const typename TShaderClass::FParameters& Parameters)
+inline void SetShaderUAVs(TRHICmdList& RHICmdList, const TShaderClass* Shader, FRHIComputeShader* ShadeRHI, const typename TShaderClass::FParameters& Parameters)
 {
 	const FShaderParameterBindings& Bindings = Shader->Bindings;
 
@@ -138,7 +138,7 @@ inline void SetShaderUAVs(TRHICmdList& RHICmdList, const TShaderClass* Shader, F
 
 /** Unset compute shader UAVs. */
 template<typename TRHICmdList, typename TShaderClass>
-inline void UnsetShaderUAVs(TRHICmdList& RHICmdList, const TShaderClass* Shader, FComputeShaderRHIParamRef ShadeRHI)
+inline void UnsetShaderUAVs(TRHICmdList& RHICmdList, const TShaderClass* Shader, FRHIComputeShader* ShadeRHI)
 {
 	// TODO(RDG): Once all shader sets their parameter through this, can refactor RHI so all UAVs of a shader get unset through a single RHI function call.
 	const FShaderParameterBindings& Bindings = Shader->Bindings;
