Precises what Source and Target means on TStaticBlendStateRHI's documentation.

diff --git a/Engine/Source/Runtime/RHI/Public/RHIStaticStates.h b/Engine/Source/Runtime/RHI/Public/RHIStaticStates.h
--- a/Engine/Source/Runtime/RHI/Public/RHIStaticStates.h
+++ b/Engine/Source/Runtime/RHI/Public/RHIStaticStates.h
@@ -296,6 +296,8 @@ public:
  * else
  *		Alpha blended the same way as rgb
  * 
+ * Where source is the color coming from the pixel shader, and target is the color in the render target.
+ *
  * So for example, TStaticBlendState<BO_Add,BF_SourceAlpha,BF_InverseSourceAlpha,BO_Add,BF_Zero,BF_One> produces:
  * FinalColor.rgb = SourceColor * SourceAlpha + DestColor * (1 - SourceAlpha);
  * FinalColor.a = SourceAlpha * 0 + DestAlpha * 1;
