[RENDERING] [+] New RHI buffer lock/unlock API. These changes were made to remove the need for the RHI to call functions such as IsInRenderingThread() and IsInRHIThread() in their buffer lock/unlock implementations to determine whether buffer renames etc should be deferred or not.

RHIs can now make use of the RHICmdList.IsTopOfPipe() and RHICmdList.IsBottomOfPipe() functions to determine if lock/unlock operations should be deferred or not:

 - If IsTopOfPipe() returns true, the code is executing on the rendering thread (RT), whilst the RT is recording commands into the immediate RHI command list. Buffer renames etc may need to be deferred down the pipe using RHICmdList.EnqueueLambda().

 - If IsBottomOfPipe() returns true, the code is replaying previously recorded commands (or we are Bypass()ing the command list recording). As such, we should behave like we are running on the RHI thread (RHIT) and access GPU resources directly. It is possible for IsBottomOfPipe() to return true whilst we are on the RT, if running in "-norhithread" mode, since the RT both records and plays back the immediate command list. This is why it is important to use the IsTopOfPipe/IsBottomOfPipe functions to provide the execution context in the RHI, rather than using threading functions like IsInRenderingThread(). It is also more optimal to use these new functions, as the IsIn*Thread() functions require accessing TLS data which can be slow on some platforms.

The rest of this change refactors the RHI lock/unlock API to allow the FRHICommandListImmediate to be passed all the way into the RHI implementation. RHIs now have a choice of which buffer lock functions to implement. Either:

 - Implement the FDynamicRHI::RHILock*Buffer and FDynamicRHI::RHIUnlock*Buffer functions. These can be called on *both* the top and bottom of pipe. The RHI is required to use the passed-in RHICmdList to determine execution context. This is the most optimal approach as it gives the RHI the opportunity to return GPU memory (e.g. upload heap memory) in the Lock function directly, minimizing the number of memcpy()s needed to update a buffer. It is however the more complex solution as the RHI needs to use the RHICmdList.IsTopOfPipe/IsBottomOfPipe functions to determine if lock operations should be deferred with EnqueueLambda().

 - Implement the FDynamicRHI::Lock*Buffer_BottomOfPipe and FDynamicRHI::Unlock*Buffer_BottomOfPipe functions. This is the simplest solution. RHIs can assume these functions will always be called at the bottom of the pipe (hence the name), and so can access GPU resources directly without having to defer operations via EnqueueLambda, nor check the IsTopOfPipe/IsBottomOfPipe functions. It is less optimal though, as its implementation relies on the base class implementation of FDynamicRHI::RHILock*Buffer, which returns temporary CPU memory on Lock(), and does a memcpy() and free on Unlock.

This change also modifies the global resolve shader. Various platform RHIs were doing what was effectively a manual draw user primitives, by creating, locking and unlocking a vertex buffer just to fill it with vertices for a full screen resolve. This is now handled with a global, immutable vertex buffer and uniform parameters.

diff --git a/Engine/Plugins/Experimental/GeometryCache/Source/GeometryCache/Private/GeometryCacheSceneProxy.cpp b/Engine/Plugins/Experimental/GeometryCache/Source/GeometryCache/Private/GeometryCacheSceneProxy.cpp
--- a/Engine/Plugins/Experimental/GeometryCache/Source/GeometryCache/Private/GeometryCacheSceneProxy.cpp
+++ b/Engine/Plugins/Experimental/GeometryCache/Source/GeometryCache/Private/GeometryCacheSceneProxy.cpp
@@ -210,6 +210,7 @@ FGeometryCacheSceneProxy::~FGeometryCacheSceneProxy()
 	Tracks.Empty();
 }
 
+#if 0
 struct FRHICommandUpdateGeometryCacheBuffer : public FRHICommand<FRHICommandUpdateGeometryCacheBuffer>
 {
 	FGraphEventRef BufferGenerationCompleteFence;
@@ -270,6 +271,7 @@ struct FRHICommandUpdateGeometryCacheBuffer : public FRHICommand<FRHICommandUpda
 		BufferGenerationCompleteFence = nullptr;
 	}
 };
+#endif
 
 class FGeometryCacheVertexFactoryUserDataWrapper : public FOneFrameResource
 {
diff --git a/Engine/Shaders/Private/ResolveVertexShader.usf b/Engine/Shaders/Private/ResolveVertexShader.usf
--- a/Engine/Shaders/Private/ResolveVertexShader.usf
+++ b/Engine/Shaders/Private/ResolveVertexShader.usf
@@ -6,13 +6,15 @@
 
 #include "Common.ush"
 
+float4 PositionMinMax; // xy = min [x, y], zw = max [x, y]
+float4 UVMinMax;       // xy = min [u, v], zw = max [u, v]
+
 void Main(
 	float2 InPosition : ATTRIBUTE0,
-	float2 InUV       : ATTRIBUTE1,
 	out float2 OutUV       : TEXCOORD0,
 	out float4 OutPosition : SV_POSITION
 	)
 {
-	OutPosition = float4( InPosition, 0, 1 );
-    OutUV = InUV;
+	OutPosition = float4(lerp(PositionMinMax.xy, PositionMinMax.zw, InPosition), 0, 1);
+    OutUV = lerp(UVMinMax.xy, UVMinMax.zw, InPosition);
 }
diff --git a/Engine/Source/Runtime/Apple/MetalRHI/Private/MetalIndexBuffer.cpp b/Engine/Source/Runtime/Apple/MetalRHI/Private/MetalIndexBuffer.cpp
--- a/Engine/Source/Runtime/Apple/MetalRHI/Private/MetalIndexBuffer.cpp
+++ b/Engine/Source/Runtime/Apple/MetalRHI/Private/MetalIndexBuffer.cpp
@@ -118,7 +118,7 @@ void FMetalDynamicRHI::RHITransferIndexBufferUnderlyingResource(FIndexBufferRHIP
 	}
 }
 
-void* FMetalDynamicRHI::RHILockIndexBuffer(FIndexBufferRHIParamRef IndexBufferRHI, uint32 Offset, uint32 Size, EResourceLockMode LockMode)
+void* FMetalDynamicRHI::LockIndexBuffer_BottomOfPipe(FRHICommandListImmediate& RHICmdList, FIndexBufferRHIParamRef IndexBufferRHI, uint32 Offset, uint32 Size, EResourceLockMode LockMode)
 {
 	@autoreleasepool {
 	FMetalIndexBuffer* IndexBuffer = ResourceCast(IndexBufferRHI);
@@ -127,7 +127,7 @@ void* FMetalDynamicRHI::RHILockIndexBuffer(FIndexBufferRHIParamRef IndexBufferRH
 	}
 }
 
-void FMetalDynamicRHI::RHIUnlockIndexBuffer(FIndexBufferRHIParamRef IndexBufferRHI)
+void FMetalDynamicRHI::UnlockIndexBuffer_BottomOfPipe(FRHICommandListImmediate& RHICmdList, FIndexBufferRHIParamRef IndexBufferRHI)
 {
 	@autoreleasepool {
 	FMetalIndexBuffer* IndexBuffer = ResourceCast(IndexBufferRHI);
diff --git a/Engine/Source/Runtime/Apple/MetalRHI/Private/MetalStructuredBuffer.cpp b/Engine/Source/Runtime/Apple/MetalRHI/Private/MetalStructuredBuffer.cpp
--- a/Engine/Source/Runtime/Apple/MetalRHI/Private/MetalStructuredBuffer.cpp
+++ b/Engine/Source/Runtime/Apple/MetalRHI/Private/MetalStructuredBuffer.cpp
@@ -54,7 +54,7 @@ FStructuredBufferRHIRef FMetalDynamicRHI::RHICreateStructuredBuffer(uint32 Strid
 	}
 }
 
-void* FMetalDynamicRHI::RHILockStructuredBuffer(FStructuredBufferRHIParamRef StructuredBufferRHI,uint32 Offset,uint32 Size,EResourceLockMode LockMode)
+void* FMetalDynamicRHI::LockStructuredBuffer_BottomOfPipe(FRHICommandListImmediate& RHICmdList, FStructuredBufferRHIParamRef StructuredBufferRHI, uint32 Offset, uint32 Size, EResourceLockMode LockMode)
 {
 	@autoreleasepool {
 	FMetalStructuredBuffer* StructuredBuffer = ResourceCast(StructuredBufferRHI);
@@ -64,7 +64,7 @@ void* FMetalDynamicRHI::RHILockStructuredBuffer(FStructuredBufferRHIParamRef Str
 	}
 }
 
-void FMetalDynamicRHI::RHIUnlockStructuredBuffer(FStructuredBufferRHIParamRef StructuredBufferRHI)
+void FMetalDynamicRHI::UnlockStructuredBuffer_BottomOfPipe(FRHICommandListImmediate& RHICmdList, FStructuredBufferRHIParamRef StructuredBufferRHI)
 {
 	@autoreleasepool {
 	FMetalStructuredBuffer* StructuredBuffer = ResourceCast(StructuredBufferRHI);
diff --git a/Engine/Source/Runtime/Apple/MetalRHI/Private/MetalVertexBuffer.cpp b/Engine/Source/Runtime/Apple/MetalRHI/Private/MetalVertexBuffer.cpp
--- a/Engine/Source/Runtime/Apple/MetalRHI/Private/MetalVertexBuffer.cpp
+++ b/Engine/Source/Runtime/Apple/MetalRHI/Private/MetalVertexBuffer.cpp
@@ -584,7 +584,7 @@ FVertexBufferRHIRef FMetalDynamicRHI::RHICreateVertexBuffer(uint32 Size, uint32
 	}
 }
 
-void* FMetalDynamicRHI::RHILockVertexBuffer(FVertexBufferRHIParamRef VertexBufferRHI, uint32 Offset, uint32 Size, EResourceLockMode LockMode)
+void* FMetalDynamicRHI::LockVertexBuffer_BottomOfPipe(FRHICommandListImmediate& RHICmdList, FVertexBufferRHIParamRef VertexBufferRHI, uint32 Offset, uint32 Size, EResourceLockMode LockMode)
 {
 	@autoreleasepool {
 	FMetalVertexBuffer* VertexBuffer = ResourceCast(VertexBufferRHI);
@@ -594,7 +594,7 @@ void* FMetalDynamicRHI::RHILockVertexBuffer(FVertexBufferRHIParamRef VertexBuffe
 	}
 }
 
-void FMetalDynamicRHI::RHIUnlockVertexBuffer(FVertexBufferRHIParamRef VertexBufferRHI)
+void FMetalDynamicRHI::UnlockVertexBuffer_BottomOfPipe(FRHICommandListImmediate& RHICmdList, FVertexBufferRHIParamRef VertexBufferRHI)
 {
 	@autoreleasepool {
 	FMetalVertexBuffer* VertexBuffer = ResourceCast(VertexBufferRHI);
diff --git a/Engine/Source/Runtime/Apple/MetalRHI/Public/MetalRHI.h b/Engine/Source/Runtime/Apple/MetalRHI/Public/MetalRHI.h
--- a/Engine/Source/Runtime/Apple/MetalRHI/Public/MetalRHI.h
+++ b/Engine/Source/Runtime/Apple/MetalRHI/Public/MetalRHI.h
@@ -64,17 +64,17 @@ public:
 	virtual TRefCountPtr<FRHIComputePipelineState> RHICreateComputePipelineState(FRHIComputeShader* ComputeShader) final override;
 	virtual FUniformBufferRHIRef RHICreateUniformBuffer(const void* Contents, const FRHIUniformBufferLayout& Layout, EUniformBufferUsage Usage, EUniformBufferValidation Validation) final override;
 	virtual FIndexBufferRHIRef RHICreateIndexBuffer(uint32 Stride, uint32 Size, uint32 InUsage, FRHIResourceCreateInfo& CreateInfo) final override;
-	virtual void* RHILockIndexBuffer(FIndexBufferRHIParamRef IndexBuffer, uint32 Offset, uint32 Size, EResourceLockMode LockMode) final override;
-	virtual void RHIUnlockIndexBuffer(FIndexBufferRHIParamRef IndexBuffer) final override;
+	virtual void* LockIndexBuffer_BottomOfPipe(FRHICommandListImmediate& RHICmdList, FIndexBufferRHIParamRef IndexBuffer, uint32 Offset, uint32 Size, EResourceLockMode LockMode) final override;
+	virtual void UnlockIndexBuffer_BottomOfPipe(FRHICommandListImmediate& RHICmdList, FIndexBufferRHIParamRef IndexBuffer) final override;
 	virtual void RHITransferIndexBufferUnderlyingResource(FIndexBufferRHIParamRef DestIndexBuffer, FIndexBufferRHIParamRef SrcIndexBuffer) final override;
 	virtual FVertexBufferRHIRef RHICreateVertexBuffer(uint32 Size, uint32 InUsage, FRHIResourceCreateInfo& CreateInfo) final override;
-	virtual void* RHILockVertexBuffer(FVertexBufferRHIParamRef VertexBuffer, uint32 Offset, uint32 SizeRHI, EResourceLockMode LockMode) final override;
-	virtual void RHIUnlockVertexBuffer(FVertexBufferRHIParamRef VertexBuffer) final override;
+	virtual void* LockVertexBuffer_BottomOfPipe(FRHICommandListImmediate& RHICmdList, FVertexBufferRHIParamRef VertexBuffer, uint32 Offset, uint32 SizeRHI, EResourceLockMode LockMode) final override;
+	virtual void UnlockVertexBuffer_BottomOfPipe(FRHICommandListImmediate& RHICmdList, FVertexBufferRHIParamRef VertexBuffer) final override;
 	virtual void RHICopyVertexBuffer(FVertexBufferRHIParamRef SourceBuffer,FVertexBufferRHIParamRef DestBuffer) final override;
 	virtual void RHITransferVertexBufferUnderlyingResource(FVertexBufferRHIParamRef DestVertexBuffer, FVertexBufferRHIParamRef SrcVertexBuffer) final override;
 	virtual FStructuredBufferRHIRef RHICreateStructuredBuffer(uint32 Stride, uint32 Size, uint32 InUsage, FRHIResourceCreateInfo& CreateInfo) final override;
-	virtual void* RHILockStructuredBuffer(FStructuredBufferRHIParamRef StructuredBuffer, uint32 Offset, uint32 SizeRHI, EResourceLockMode LockMode) final override;
-	virtual void RHIUnlockStructuredBuffer(FStructuredBufferRHIParamRef StructuredBuffer) final override;
+	virtual void* LockStructuredBuffer_BottomOfPipe(FRHICommandListImmediate& RHICmdList, FStructuredBufferRHIParamRef StructuredBuffer, uint32 Offset, uint32 SizeRHI, EResourceLockMode LockMode) final override;
+	virtual void UnlockStructuredBuffer_BottomOfPipe(FRHICommandListImmediate& RHICmdList, FStructuredBufferRHIParamRef StructuredBuffer) final override;
 	virtual FUnorderedAccessViewRHIRef RHICreateUnorderedAccessView(FStructuredBufferRHIParamRef StructuredBuffer, bool bUseUAVCounter, bool bAppendBuffer) final override;
 	virtual FUnorderedAccessViewRHIRef RHICreateUnorderedAccessView(FTextureRHIParamRef Texture, uint32 MipLevel) final override;
 	virtual FUnorderedAccessViewRHIRef RHICreateUnorderedAccessView(FVertexBufferRHIParamRef VertexBuffer, uint8 Format) final override;
diff --git a/Engine/Source/Runtime/D3D12RHI/Private/D3D12IndexBuffer.cpp b/Engine/Source/Runtime/D3D12RHI/Private/D3D12IndexBuffer.cpp
--- a/Engine/Source/Runtime/D3D12RHI/Private/D3D12IndexBuffer.cpp
+++ b/Engine/Source/Runtime/D3D12RHI/Private/D3D12IndexBuffer.cpp
@@ -104,14 +104,14 @@ FIndexBufferRHIRef FD3D12DynamicRHI::RHICreateIndexBuffer(uint32 Stride, uint32
 	return Buffer;
 }
 
-void* FD3D12DynamicRHI::RHILockIndexBuffer(FIndexBufferRHIParamRef IndexBufferRHI, uint32 Offset, uint32 Size, EResourceLockMode LockMode)
+void* FD3D12DynamicRHI::RHILockIndexBuffer(FRHICommandListImmediate& RHICmdList, FIndexBufferRHIParamRef IndexBufferRHI, uint32 Offset, uint32 Size, EResourceLockMode LockMode)
 {
-	return LockBuffer(nullptr, FD3D12DynamicRHI::ResourceCast(IndexBufferRHI), Offset, Size, LockMode);
+	return LockBuffer(&RHICmdList, FD3D12DynamicRHI::ResourceCast(IndexBufferRHI), Offset, Size, LockMode);
 }
 
-void FD3D12DynamicRHI::RHIUnlockIndexBuffer(FIndexBufferRHIParamRef IndexBufferRHI)
+void FD3D12DynamicRHI::RHIUnlockIndexBuffer(FRHICommandListImmediate& RHICmdList, FIndexBufferRHIParamRef IndexBufferRHI)
 {
-	UnlockBuffer(nullptr, FD3D12DynamicRHI::ResourceCast(IndexBufferRHI));
+	UnlockBuffer(&RHICmdList, FD3D12DynamicRHI::ResourceCast(IndexBufferRHI));
 }
 
 void FD3D12DynamicRHI::RHITransferIndexBufferUnderlyingResource(FIndexBufferRHIParamRef DestIndexBuffer, FIndexBufferRHIParamRef SrcIndexBuffer)
@@ -149,20 +149,6 @@ FIndexBufferRHIRef FD3D12DynamicRHI::CreateIndexBuffer_RenderThread(class FRHICo
 	return Buffer;
 }
 
-void* FD3D12DynamicRHI::LockIndexBuffer_RenderThread(class FRHICommandListImmediate& RHICmdList, FIndexBufferRHIParamRef IndexBufferRHI, uint32 Offset, uint32 SizeRHI, EResourceLockMode LockMode)
-{
-	return LockBuffer(&RHICmdList, FD3D12DynamicRHI::ResourceCast(IndexBufferRHI), Offset, SizeRHI, LockMode);
-}
-
-void FD3D12DynamicRHI::UnlockIndexBuffer_RenderThread(class FRHICommandListImmediate& RHICmdList, FIndexBufferRHIParamRef IndexBufferRHI)
-{
-	// Pull down the above RHI implementation so that we can flush only when absolutely necessary
-	QUICK_SCOPE_CYCLE_COUNTER(STAT_FDynamicRHI_UnlockIndexBuffer_RenderThread);
-	check(IsInRenderingThread());
-
-	UnlockBuffer(&RHICmdList, FD3D12DynamicRHI::ResourceCast(IndexBufferRHI));
-}
-
 FIndexBufferRHIRef FD3D12DynamicRHI::CreateAndLockIndexBuffer_RenderThread(class FRHICommandListImmediate& RHICmdList, uint32 Stride, uint32 Size, uint32 InUsage, FRHIResourceCreateInfo& CreateInfo, void*& OutDataBuffer)
 {
 	const D3D12_RESOURCE_DESC Desc = CreateIndexBufferResourceDesc(Size, InUsage);
@@ -175,7 +161,7 @@ FIndexBufferRHIRef FD3D12DynamicRHI::CreateAndLockIndexBuffer_RenderThread(class
 		Buffer->SetCommitted(false);
 	}
 
-	OutDataBuffer = LockIndexBuffer_RenderThread(RHICmdList, Buffer, 0, Size, RLM_WriteOnly);
+	OutDataBuffer = LockBuffer(&RHICmdList, Buffer, 0, Size, RLM_WriteOnly);
 
 	return Buffer;
 }
diff --git a/Engine/Source/Runtime/D3D12RHI/Private/D3D12RHIPrivate.h b/Engine/Source/Runtime/D3D12RHI/Private/D3D12RHIPrivate.h
--- a/Engine/Source/Runtime/D3D12RHI/Private/D3D12RHIPrivate.h
+++ b/Engine/Source/Runtime/D3D12RHI/Private/D3D12RHIPrivate.h
@@ -283,17 +283,17 @@ public:
 	virtual FUniformBufferRHIRef RHICreateUniformBuffer(const void* Contents, const FRHIUniformBufferLayout& Layout, EUniformBufferUsage Usage, EUniformBufferValidation Validation) final override;
 	virtual void RHIUpdateUniformBuffer(FUniformBufferRHIParamRef UniformBufferRHI, const void* Contents) final override;
 	virtual FIndexBufferRHIRef RHICreateIndexBuffer(uint32 Stride, uint32 Size, uint32 InUsage, FRHIResourceCreateInfo& CreateInfo) final override;
-	virtual void* RHILockIndexBuffer(FIndexBufferRHIParamRef IndexBuffer, uint32 Offset, uint32 Size, EResourceLockMode LockMode) final override;
-	virtual void RHIUnlockIndexBuffer(FIndexBufferRHIParamRef IndexBuffer) final override;
+	virtual void* RHILockIndexBuffer(FRHICommandListImmediate& RHICmdList, FIndexBufferRHIParamRef IndexBuffer, uint32 Offset, uint32 Size, EResourceLockMode LockMode) final override;
+	virtual void RHIUnlockIndexBuffer(FRHICommandListImmediate& RHICmdList, FIndexBufferRHIParamRef IndexBuffer) final override;
 	virtual void RHITransferIndexBufferUnderlyingResource(FIndexBufferRHIParamRef DestIndexBuffer, FIndexBufferRHIParamRef SrcIndexBuffer) final override;
 	virtual FVertexBufferRHIRef RHICreateVertexBuffer(uint32 Size, uint32 InUsage, FRHIResourceCreateInfo& CreateInfo) final override;
-	virtual void* RHILockVertexBuffer(FVertexBufferRHIParamRef VertexBuffer, uint32 Offset, uint32 SizeRHI, EResourceLockMode LockMode) final override;
-	virtual void RHIUnlockVertexBuffer(FVertexBufferRHIParamRef VertexBuffer) final override;
+	virtual void* RHILockVertexBuffer(FRHICommandListImmediate& RHICmdList, FVertexBufferRHIParamRef VertexBuffer, uint32 Offset, uint32 SizeRHI, EResourceLockMode LockMode) final override;
+	virtual void RHIUnlockVertexBuffer(FRHICommandListImmediate& RHICmdList, FVertexBufferRHIParamRef VertexBuffer) final override;
 	virtual void RHICopyVertexBuffer(FVertexBufferRHIParamRef SourceBuffer, FVertexBufferRHIParamRef DestBuffer) final override;
 	virtual void RHITransferVertexBufferUnderlyingResource(FVertexBufferRHIParamRef DestVertexBuffer, FVertexBufferRHIParamRef SrcVertexBuffer) final override;
 	virtual FStructuredBufferRHIRef RHICreateStructuredBuffer(uint32 Stride, uint32 Size, uint32 InUsage, FRHIResourceCreateInfo& CreateInfo) final override;
-	virtual void* RHILockStructuredBuffer(FStructuredBufferRHIParamRef StructuredBuffer, uint32 Offset, uint32 SizeRHI, EResourceLockMode LockMode) final override;
-	virtual void RHIUnlockStructuredBuffer(FStructuredBufferRHIParamRef StructuredBuffer) final override;
+	virtual void* RHILockStructuredBuffer(FRHICommandListImmediate& RHICmdList, FStructuredBufferRHIParamRef StructuredBuffer, uint32 Offset, uint32 SizeRHI, EResourceLockMode LockMode) final override;
+	virtual void RHIUnlockStructuredBuffer(FRHICommandListImmediate& RHICmdList, FStructuredBufferRHIParamRef StructuredBuffer) final override;
 	virtual FUnorderedAccessViewRHIRef RHICreateUnorderedAccessView(FStructuredBufferRHIParamRef StructuredBuffer, bool bUseUAVCounter, bool bAppendBuffer) final override;
 	virtual FUnorderedAccessViewRHIRef RHICreateUnorderedAccessView(FTextureRHIParamRef Texture, uint32 MipLevel) final override;
 	virtual FUnorderedAccessViewRHIRef RHICreateUnorderedAccessView(FVertexBufferRHIParamRef VertexBuffer, uint8 Format) final override;
@@ -422,16 +422,12 @@ public:
 
 	virtual FVertexBufferRHIRef CreateVertexBuffer_RenderThread(class FRHICommandListImmediate& RHICmdList, uint32 Size, uint32 InUsage, FRHIResourceCreateInfo& CreateInfo);
 	virtual FStructuredBufferRHIRef CreateStructuredBuffer_RenderThread(class FRHICommandListImmediate& RHICmdList, uint32 Stride, uint32 Size, uint32 InUsage, FRHIResourceCreateInfo& CreateInfo);
-	virtual void* LockVertexBuffer_RenderThread(class FRHICommandListImmediate& RHICmdList, FVertexBufferRHIParamRef VertexBuffer, uint32 Offset, uint32 SizeRHI, EResourceLockMode LockMode);
-	virtual void UnlockVertexBuffer_RenderThread(class FRHICommandListImmediate& RHICmdList, FVertexBufferRHIParamRef VertexBuffer);
 	virtual FVertexBufferRHIRef CreateAndLockVertexBuffer_RenderThread(class FRHICommandListImmediate& RHICmdList, uint32 Size, uint32 InUsage, FRHIResourceCreateInfo& CreateInfo, void*& OutDataBuffer);
 	virtual FIndexBufferRHIRef CreateAndLockIndexBuffer_RenderThread(class FRHICommandListImmediate& RHICmdList, uint32 Stride, uint32 Size, uint32 InUsage, FRHIResourceCreateInfo& CreateInfo, void*& OutDataBuffer);
 	// virtual FTexture2DRHIRef AsyncReallocateTexture2D_RenderThread(class FRHICommandListImmediate& RHICmdList, FTexture2DRHIParamRef Texture2D, int32 NewMipCount, int32 NewSizeX, int32 NewSizeY, FThreadSafeCounter* RequestStatus);
 	// virtual ETextureReallocationStatus FinalizeAsyncReallocateTexture2D_RenderThread(class FRHICommandListImmediate& RHICmdList, FTexture2DRHIParamRef Texture2D, bool bBlockUntilCompleted);
 	// virtual ETextureReallocationStatus CancelAsyncReallocateTexture2D_RenderThread(class FRHICommandListImmediate& RHICmdList, FTexture2DRHIParamRef Texture2D, bool bBlockUntilCompleted);
 	virtual FIndexBufferRHIRef CreateIndexBuffer_RenderThread(class FRHICommandListImmediate& RHICmdList, uint32 Stride, uint32 Size, uint32 InUsage, FRHIResourceCreateInfo& CreateInfo);
-	virtual void* LockIndexBuffer_RenderThread(class FRHICommandListImmediate& RHICmdList, FIndexBufferRHIParamRef IndexBuffer, uint32 Offset, uint32 SizeRHI, EResourceLockMode LockMode);
-	virtual void UnlockIndexBuffer_RenderThread(class FRHICommandListImmediate& RHICmdList, FIndexBufferRHIParamRef IndexBuffer);
 	virtual void UpdateTexture2D_RenderThread(class FRHICommandListImmediate& RHICmdList, FTexture2DRHIParamRef Texture, uint32 MipIndex, const struct FUpdateTextureRegion2D& UpdateRegion, uint32 SourcePitch, const uint8* SourceData);
 	virtual void* LockTexture2D_RenderThread(class FRHICommandListImmediate& RHICmdList, FTexture2DRHIParamRef Texture, uint32 MipIndex, EResourceLockMode LockMode, uint32& DestStride, bool bLockWithinMiptail, bool bNeedsDefaultRHIFlush = true);
 	virtual void UnlockTexture2D_RenderThread(class FRHICommandListImmediate& RHICmdList, FTexture2DRHIParamRef Texture, uint32 MipIndex, bool bLockWithinMiptail, bool bNeedsDefaultRHIFlush = true);
@@ -480,14 +476,14 @@ public:
 	template<class BufferType>
 	void UnlockBuffer(FRHICommandListImmediate* RHICmdList, BufferType* Buffer);
 
-	static inline bool ShouldDeferBufferLockOperation(FRHICommandList* RHICmdList)
+	static inline bool ShouldDeferBufferLockOperation(FRHICommandListImmediate* RHICmdList)
 	{
 		if (RHICmdList == nullptr)
 		{
 			return false;
 		}
 
-		if (RHICmdList->Bypass() || !IsRunningRHIInSeparateThread())
+		if (RHICmdList->IsBottomOfPipe())
 		{
 			return false;
 		}
diff --git a/Engine/Source/Runtime/D3D12RHI/Private/D3D12RenderTarget.cpp b/Engine/Source/Runtime/D3D12RHI/Private/D3D12RenderTarget.cpp
--- a/Engine/Source/Runtime/D3D12RHI/Private/D3D12RenderTarget.cpp
+++ b/Engine/Source/Runtime/D3D12RHI/Private/D3D12RenderTarget.cpp
@@ -166,7 +166,7 @@ void FD3D12CommandContext::ResolveTextureUsingShader(
 	TShaderMapRef<FResolveVS> ResolveVertexShader(ShaderMap);
 	TShaderMapRef<TPixelShader> ResolvePixelShader(ShaderMap);
 
-	GraphicsPSOInit.BoundShaderState.VertexDeclarationRHI = GScreenVertexDeclaration.VertexDeclarationRHI;
+	GraphicsPSOInit.BoundShaderState.VertexDeclarationRHI = GResolveVertexBuffer.VertexDeclarationRHI;
 	GraphicsPSOInit.BoundShaderState.VertexShaderRHI = GETSAFERHISHADER_VERTEX(*ResolveVertexShader);
 	GraphicsPSOInit.BoundShaderState.PixelShaderRHI = GETSAFERHISHADER_PIXEL(*ResolvePixelShader);
 	GraphicsPSOInit.PrimitiveType = PT_TriangleStrip;
@@ -175,43 +175,14 @@ void FD3D12CommandContext::ResolveTextureUsingShader(
 	RHICmdList.SetBlendFactor(FLinearColor::White);
 
 	ResolvePixelShader->SetParameters(RHICmdList, PixelShaderParameter);
+	ResolveVertexShader->SetParameters(RHICmdList, FVector4(MinX, MinY, MaxX, MaxY), FVector4(MinU, MinV, MaxU, MaxV));
 	RHICmdList.Flush(); // always call flush when using a command list in RHI implementations before doing anything else. This is super hazardous.
 
 	// Set the source texture.
 	const uint32 TextureIndex = ResolvePixelShader->UnresolvedSurface.GetBaseIndex();
 	StateCache.SetShaderResourceView<SF_Pixel>(SourceTexture->GetShaderResourceView(), TextureIndex);
 
-	FRHIResourceCreateInfo CreateInfo;
-	check(IsInRHIThread() || !IsRunningRHIInSeparateThread());
-	FVertexBufferRHIRef VertexBufferRHI = GDynamicRHI->RHICreateVertexBuffer(sizeof(FScreenVertex) * 4, BUF_Volatile, CreateInfo);
-	void* VoidPtr = GDynamicRHI->RHILockVertexBuffer(VertexBufferRHI, 0, sizeof(FScreenVertex) * 4, RLM_WriteOnly);
-
-	// Generate the vertices used
-	FScreenVertex* Vertices = (FScreenVertex*)VoidPtr;
-
-	Vertices[0].Position.X = MaxX;
-	Vertices[0].Position.Y = MinY;
-	Vertices[0].UV.X = MaxU;
-	Vertices[0].UV.Y = MinV;
-
-	Vertices[1].Position.X = MaxX;
-	Vertices[1].Position.Y = MaxY;
-	Vertices[1].UV.X = MaxU;
-	Vertices[1].UV.Y = MaxV;
-
-	Vertices[2].Position.X = MinX;
-	Vertices[2].Position.Y = MinY;
-	Vertices[2].UV.X = MinU;
-	Vertices[2].UV.Y = MinV;
-
-	Vertices[3].Position.X = MinX;
-	Vertices[3].Position.Y = MaxY;
-	Vertices[3].UV.X = MinU;
-	Vertices[3].UV.Y = MaxV;
-
-	GDynamicRHI->RHIUnlockVertexBuffer(VertexBufferRHI);
-
-	RHICmdList.SetStreamSource(0, VertexBufferRHI, 0);
+	RHICmdList.SetStreamSource(0, GResolveVertexBuffer.VertexBufferRHI, 0);
 	RHICmdList.DrawPrimitive(0, 2, 1);
 
 	RHICmdList.Flush(); // always call flush when using a command list in RHI implementations before doing anything else. This is super hazardous.
diff --git a/Engine/Source/Runtime/D3D12RHI/Private/D3D12StructuredBuffer.cpp b/Engine/Source/Runtime/D3D12RHI/Private/D3D12StructuredBuffer.cpp
--- a/Engine/Source/Runtime/D3D12RHI/Private/D3D12StructuredBuffer.cpp
+++ b/Engine/Source/Runtime/D3D12RHI/Private/D3D12StructuredBuffer.cpp
@@ -109,12 +109,12 @@ void FD3D12StructuredBuffer::RenameLDAChain(FD3D12ResourceLocation& NewLocation)
 	}
 }
 
-void* FD3D12DynamicRHI::RHILockStructuredBuffer(FStructuredBufferRHIParamRef StructuredBufferRHI, uint32 Offset, uint32 Size, EResourceLockMode LockMode)
+void* FD3D12DynamicRHI::RHILockStructuredBuffer(FRHICommandListImmediate& RHICmdList, FStructuredBufferRHIParamRef StructuredBufferRHI, uint32 Offset, uint32 Size, EResourceLockMode LockMode)
 {
-	return LockBuffer(nullptr, FD3D12DynamicRHI::ResourceCast(StructuredBufferRHI), Offset, Size, LockMode);
+	return LockBuffer(&RHICmdList, FD3D12DynamicRHI::ResourceCast(StructuredBufferRHI), Offset, Size, LockMode);
 }
 
-void FD3D12DynamicRHI::RHIUnlockStructuredBuffer(FStructuredBufferRHIParamRef StructuredBufferRHI)
+void FD3D12DynamicRHI::RHIUnlockStructuredBuffer(FRHICommandListImmediate& RHICmdList, FStructuredBufferRHIParamRef StructuredBufferRHI)
 {
-	UnlockBuffer(nullptr, FD3D12DynamicRHI::ResourceCast(StructuredBufferRHI));
+	UnlockBuffer(&RHICmdList, FD3D12DynamicRHI::ResourceCast(StructuredBufferRHI));
 }
diff --git a/Engine/Source/Runtime/D3D12RHI/Private/D3D12VertexBuffer.cpp b/Engine/Source/Runtime/D3D12RHI/Private/D3D12VertexBuffer.cpp
--- a/Engine/Source/Runtime/D3D12RHI/Private/D3D12VertexBuffer.cpp
+++ b/Engine/Source/Runtime/D3D12RHI/Private/D3D12VertexBuffer.cpp
@@ -121,14 +121,14 @@ FVertexBufferRHIRef FD3D12DynamicRHI::RHICreateVertexBuffer(uint32 Size, uint32
 	return Buffer;
 }
 
-void* FD3D12DynamicRHI::RHILockVertexBuffer(FVertexBufferRHIParamRef VertexBufferRHI, uint32 Offset, uint32 Size, EResourceLockMode LockMode)
+void* FD3D12DynamicRHI::RHILockVertexBuffer(FRHICommandListImmediate& RHICmdList, FVertexBufferRHIParamRef VertexBufferRHI, uint32 Offset, uint32 Size, EResourceLockMode LockMode)
 {
-	return LockBuffer(nullptr, FD3D12DynamicRHI::ResourceCast(VertexBufferRHI), Offset, Size, LockMode);
+	return LockBuffer(&RHICmdList, FD3D12DynamicRHI::ResourceCast(VertexBufferRHI), Offset, Size, LockMode);
 }
 
-void FD3D12DynamicRHI::RHIUnlockVertexBuffer(FVertexBufferRHIParamRef VertexBufferRHI)
+void FD3D12DynamicRHI::RHIUnlockVertexBuffer(FRHICommandListImmediate& RHICmdList, FVertexBufferRHIParamRef VertexBufferRHI)
 {
-	UnlockBuffer(nullptr, FD3D12DynamicRHI::ResourceCast(VertexBufferRHI));
+	UnlockBuffer(&RHICmdList, FD3D12DynamicRHI::ResourceCast(VertexBufferRHI));
 }
 
 FVertexBufferRHIRef FD3D12DynamicRHI::CreateVertexBuffer_RenderThread(FRHICommandListImmediate& RHICmdList, uint32 Size, uint32 InUsage, FRHIResourceCreateInfo& CreateInfo)
@@ -151,24 +151,6 @@ FVertexBufferRHIRef FD3D12DynamicRHI::CreateVertexBuffer_RenderThread(FRHIComman
 	return Buffer;
 }
 
-void* FD3D12DynamicRHI::LockVertexBuffer_RenderThread(class FRHICommandListImmediate& RHICmdList, FVertexBufferRHIParamRef VertexBufferRHI, uint32 Offset, uint32 SizeRHI, EResourceLockMode LockMode)
-{
-	// Pull down the above RHI implementation so that we can flush only when absolutely necessary
-	QUICK_SCOPE_CYCLE_COUNTER(STAT_FDynamicRHI_LockVertexBuffer_RenderThread);
-	check(IsInRenderingThread());
-
-	return LockBuffer(&RHICmdList, FD3D12DynamicRHI::ResourceCast(VertexBufferRHI), Offset, SizeRHI, LockMode);
-}
-
-void FD3D12DynamicRHI::UnlockVertexBuffer_RenderThread(FRHICommandListImmediate& RHICmdList, FVertexBufferRHIParamRef VertexBufferRHI)
-{
-	// Pull down the above RHI implementation so that we can flush only when absolutely necessary
-	QUICK_SCOPE_CYCLE_COUNTER(STAT_FDynamicRHI_UnlockVertexBuffer_RenderThread);
-	check(IsInRenderingThread());
-
-	UnlockBuffer(&RHICmdList, FD3D12DynamicRHI::ResourceCast(VertexBufferRHI));
-}
-
 void FD3D12DynamicRHI::RHICopyVertexBuffer(FVertexBufferRHIParamRef SourceBufferRHI, FVertexBufferRHIParamRef DestBufferRHI)
 {
 	FD3D12VertexBuffer*  SourceBuffer = FD3D12DynamicRHI::ResourceCast(SourceBufferRHI);
@@ -402,7 +384,7 @@ FVertexBufferRHIRef FD3D12DynamicRHI::CreateAndLockVertexBuffer_RenderThread(FRH
 		// TODO: this should ideally be set in platform-independent code, since this tracking is for the high level
 		Buffer->SetCommitted(false);
 	}
-	OutDataBuffer = LockVertexBuffer_RenderThread(RHICmdList, Buffer, 0, Size, RLM_WriteOnly);
+	OutDataBuffer = LockBuffer(&RHICmdList, Buffer, 0, Size, RLM_WriteOnly);
 
 	return Buffer;
 }
diff --git a/Engine/Source/Runtime/EmptyRHI/Private/EmptyIndexBuffer.cpp b/Engine/Source/Runtime/EmptyRHI/Private/EmptyIndexBuffer.cpp
--- a/Engine/Source/Runtime/EmptyRHI/Private/EmptyIndexBuffer.cpp
+++ b/Engine/Source/Runtime/EmptyRHI/Private/EmptyIndexBuffer.cpp
@@ -52,14 +52,14 @@ FIndexBufferRHIRef FEmptyDynamicRHI::RHICreateIndexBuffer(uint32 Stride, uint32
 	return IndexBuffer;
 }
 
-void* FEmptyDynamicRHI::RHILockIndexBuffer(FIndexBufferRHIParamRef IndexBufferRHI, uint32 Offset, uint32 Size, EResourceLockMode LockMode)
+void* FEmptyDynamicRHI::LockIndexBuffer_BottomOfPipe(FRHICommandListImmediate& RHICmdList, FIndexBufferRHIParamRef IndexBufferRHI, uint32 Offset, uint32 Size, EResourceLockMode LockMode)
 {
 	FEmptyIndexBuffer* IndexBuffer = ResourceCast(IndexBufferRHI);
 
 	return (uint8*)IndexBuffer->Lock(LockMode, Size) + Offset;
 }
 
-void FEmptyDynamicRHI::RHIUnlockIndexBuffer(FIndexBufferRHIParamRef IndexBufferRHI)
+void FEmptyDynamicRHI::UnlockIndexBuffer_BottomOfPipe(FRHICommandListImmediate& RHICmdList, FIndexBufferRHIParamRef IndexBufferRHI)
 {
 	FEmptyIndexBuffer* IndexBuffer = ResourceCast(IndexBufferRHI);
 
diff --git a/Engine/Source/Runtime/EmptyRHI/Private/EmptyStructuredBuffer.cpp b/Engine/Source/Runtime/EmptyRHI/Private/EmptyStructuredBuffer.cpp
--- a/Engine/Source/Runtime/EmptyRHI/Private/EmptyStructuredBuffer.cpp
+++ b/Engine/Source/Runtime/EmptyRHI/Private/EmptyStructuredBuffer.cpp
@@ -30,14 +30,14 @@ FStructuredBufferRHIRef FEmptyDynamicRHI::RHICreateStructuredBuffer(uint32 Strid
 	return nullptr;//new FEmptyStructuredBuffer(Stride, Size, ResourceArray, InUsage);
 }
 
-void* FEmptyDynamicRHI::RHILockStructuredBuffer(FStructuredBufferRHIParamRef StructuredBufferRHI,uint32 Offset,uint32 Size,EResourceLockMode LockMode)
+void* FEmptyDynamicRHI::LockStructuredBuffer_BottomOfPipe(FRHICommandListImmediate& RHICmdList, FStructuredBufferRHIParamRef StructuredBufferRHI, uint32 Offset, uint32 Size, EResourceLockMode LockMode)
 {
 	FEmptyStructuredBuffer* StructuredBuffer = ResourceCast(StructuredBufferRHI);
 
 	return nullptr;
 }
 
-void FEmptyDynamicRHI::RHIUnlockStructuredBuffer(FStructuredBufferRHIParamRef StructuredBufferRHI)
+void FEmptyDynamicRHI::UnlockStructuredBuffer_BottomOfPipe(FRHICommandListImmediate& RHICmdList, FStructuredBufferRHIParamRef StructuredBufferRHI)
 {
 
 }
diff --git a/Engine/Source/Runtime/EmptyRHI/Private/EmptyVertexBuffer.cpp b/Engine/Source/Runtime/EmptyRHI/Private/EmptyVertexBuffer.cpp
--- a/Engine/Source/Runtime/EmptyRHI/Private/EmptyVertexBuffer.cpp
+++ b/Engine/Source/Runtime/EmptyRHI/Private/EmptyVertexBuffer.cpp
@@ -51,7 +51,7 @@ FVertexBufferRHIRef FEmptyDynamicRHI::RHICreateVertexBuffer(uint32 Size, uint32
 	return VertexBuffer;
 }
 
-void* FEmptyDynamicRHI::RHILockVertexBuffer(FVertexBufferRHIParamRef VertexBufferRHI, uint32 Offset, uint32 Size, EResourceLockMode LockMode)
+void* FEmptyDynamicRHI::LockVertexBuffer_BottomOfPipe(FRHICommandListImmediate& RHICmdList, FVertexBufferRHIParamRef VertexBufferRHI, uint32 Offset, uint32 Size, EResourceLockMode LockMode)
 {
 	FEmptyVertexBuffer* VertexBuffer = ResourceCast(VertexBufferRHI);
 
@@ -59,7 +59,7 @@ void* FEmptyDynamicRHI::RHILockVertexBuffer(FVertexBufferRHIParamRef VertexBuffe
 	return (uint8*)VertexBuffer->Lock(LockMode, Size) + Offset;
 }
 
-void FEmptyDynamicRHI::RHIUnlockVertexBuffer(FVertexBufferRHIParamRef VertexBufferRHI)
+void FEmptyDynamicRHI::UnlockVertexBuffer_BottomOfPipe(FRHICommandListImmediate& RHICmdList, FVertexBufferRHIParamRef VertexBufferRHI)
 {
 	FEmptyVertexBuffer* VertexBuffer = ResourceCast(VertexBufferRHI);
 
diff --git a/Engine/Source/Runtime/EmptyRHI/Public/EmptyRHI.h b/Engine/Source/Runtime/EmptyRHI/Public/EmptyRHI.h
--- a/Engine/Source/Runtime/EmptyRHI/Public/EmptyRHI.h
+++ b/Engine/Source/Runtime/EmptyRHI/Public/EmptyRHI.h
@@ -51,17 +51,17 @@ public:
 	virtual FBoundShaderStateRHIRef RHICreateBoundShaderState(FRHIVertexDeclaration* VertexDeclaration, FRHIVertexShader* VertexShader, FRHIHullShader* HullShader, FRHIDomainShaderFRHIDomainShader*, FRHIPixelShader* PixelShader, FRHIGeometryShader* GeometryShader) final override;
 	virtual FUniformBufferRHIRef RHICreateUniformBuffer(const void* Contents, const FRHIUniformBufferLayout& Layout, EUniformBufferUsage Usage, EUniformBufferValidation Validation) final override;
 	virtual FIndexBufferRHIRef RHICreateIndexBuffer(uint32 Stride, uint32 Size, uint32 InUsage, FRHIResourceCreateInfo& CreateInfo) final override;
-	virtual void* RHILockIndexBuffer(FIndexBufferRHIParamRef IndexBuffer, uint32 Offset, uint32 Size, EResourceLockMode LockMode) final override;
-	virtual void RHIUnlockIndexBuffer(FIndexBufferRHIParamRef IndexBuffer) final override;
+	virtual void* LockIndexBuffer_BottomOfPipe(FRHICommandListImmediate& RHICmdList, FIndexBufferRHIParamRef IndexBuffer, uint32 Offset, uint32 Size, EResourceLockMode LockMode) final override;
+	virtual void UnlockIndexBuffer_BottomOfPipe(FRHICommandListImmediate& RHICmdList, FIndexBufferRHIParamRef IndexBuffer) final override;
 	virtual void RHITransferIndexBufferUnderlyingResource(FIndexBufferRHIParamRef DestIndexBuffer, FIndexBufferRHIParamRef SrcIndexBuffer) final override;
 	virtual FVertexBufferRHIRef RHICreateVertexBuffer(uint32 Size, uint32 InUsage, FRHIResourceCreateInfo& CreateInfo) final override;
-	virtual void* RHILockVertexBuffer(FVertexBufferRHIParamRef VertexBuffer, uint32 Offset, uint32 SizeRHI, EResourceLockMode LockMode) final override;
-	virtual void RHIUnlockVertexBuffer(FVertexBufferRHIParamRef VertexBuffer) final override;
+	virtual void* LockVertexBuffer_BottomOfPipe(FRHICommandListImmediate& RHICmdList, FVertexBufferRHIParamRef VertexBuffer, uint32 Offset, uint32 SizeRHI, EResourceLockMode LockMode) final override;
+	virtual void UnlockVertexBuffer_BottomOfPipe(FRHICommandListImmediate& RHICmdList, FVertexBufferRHIParamRef VertexBuffer) final override;
 	virtual void RHICopyVertexBuffer(FVertexBufferRHIParamRef SourceBuffer,FVertexBufferRHIParamRef DestBuffer) final override;
 	virtual void RHITransferVertexBufferUnderlyingResource(FVertexBufferRHIParamRef DestVertexBuffer, FVertexBufferRHIParamRef SrcVertexBuffer) final override;
 	virtual FStructuredBufferRHIRef RHICreateStructuredBuffer(uint32 Stride, uint32 Size, uint32 InUsage, FRHIResourceCreateInfo& CreateInfo) final override;
-	virtual void* RHILockStructuredBuffer(FStructuredBufferRHIParamRef StructuredBuffer, uint32 Offset, uint32 SizeRHI, EResourceLockMode LockMode) final override;
-	virtual void RHIUnlockStructuredBuffer(FStructuredBufferRHIParamRef StructuredBuffer) final override;
+	virtual void* LockStructuredBuffer_BottomOfPipe(FRHICommandListImmediate& RHICmdList, FStructuredBufferRHIParamRef StructuredBuffer, uint32 Offset, uint32 SizeRHI, EResourceLockMode LockMode) final override;
+	virtual void UnlockStructuredBuffer_BottomOfPipe(FRHICommandListImmediate& RHICmdList, FStructuredBufferRHIParamRef StructuredBuffer) final override;
 	virtual FUnorderedAccessViewRHIRef RHICreateUnorderedAccessView(FStructuredBufferRHIParamRef StructuredBuffer, bool bUseUAVCounter, bool bAppendBuffer) final override;
 	virtual FUnorderedAccessViewRHIRef RHICreateUnorderedAccessView(FTextureRHIParamRef Texture, uint32 MipLevel) final override;
 	virtual FUnorderedAccessViewRHIRef RHICreateUnorderedAccessView(FVertexBufferRHIParamRef VertexBuffer, uint8 Format) final override;
diff --git a/Engine/Source/Runtime/Engine/Private/GPUSkinVertexFactory.cpp b/Engine/Source/Runtime/Engine/Private/GPUSkinVertexFactory.cpp
--- a/Engine/Source/Runtime/Engine/Private/GPUSkinVertexFactory.cpp
+++ b/Engine/Source/Runtime/Engine/Private/GPUSkinVertexFactory.cpp
@@ -197,47 +197,6 @@ static bool DeferSkeletalLockAndFillToRHIThread()
 	return IsRunningRHIInSeparateThread() && CVarRHICmdDeferSkeletalLockAndFillToRHIThread.GetValueOnRenderThread() > 0;
 }
 
-struct FRHICommandUpdateBoneBuffer final : public FRHICommand<FRHICommandUpdateBoneBuffer>
-{
-	FVertexBufferRHIParamRef VertexBuffer;
-	uint32 BufferSize;
-	const TArray<FMatrix>& ReferenceToLocalMatrices;
-	const TArray<FBoneIndexType>& BoneMap;
-
-
-	FORCEINLINE_DEBUGGABLE FRHICommandUpdateBoneBuffer(FVertexBufferRHIParamRef InVertexBuffer, uint32 InBufferSize, const TArray<FMatrix>& InReferenceToLocalMatrices, const TArray<FBoneIndexType>& InBoneMap)
-		: VertexBuffer(InVertexBuffer)
-		, BufferSize(InBufferSize)
-		, ReferenceToLocalMatrices(InReferenceToLocalMatrices)
-		, BoneMap(InBoneMap)
-	{
-	}
-	void Execute(FRHICommandListBase& CmdList)
-	{
-		QUICK_SCOPE_CYCLE_COUNTER(STAT_FRHICommandUpdateBoneBuffer_Execute);
-		FMatrix3x4* ChunkMatrices = (FMatrix3x4*)GDynamicRHI->RHILockVertexBuffer(VertexBuffer, 0, BufferSize, RLM_WriteOnly);
-		//FMatrix3x4 is sizeof() == 48
-		// PLATFORM_CACHE_LINE_SIZE (128) / 48 = 2.6
-		//  sizeof(FMatrix) == 64
-		// PLATFORM_CACHE_LINE_SIZE (128) / 64 = 2
-		const uint32 NumBones = BoneMap.Num();
-		check(NumBones > 0 && NumBones < 256); // otherwise maybe some bad threading on BoneMap, maybe we need to copy that
-		const int32 PreFetchStride = 2; // FPlatformMisc::Prefetch stride
-		for (uint32 BoneIdx = 0; BoneIdx < NumBones; BoneIdx++)
-		{
-			const FBoneIndexType RefToLocalIdx = BoneMap[BoneIdx];
-			check(ReferenceToLocalMatrices.IsValidIndex(RefToLocalIdx)); // otherwise maybe some bad threading on BoneMap, maybe we need to copy that
-			FPlatformMisc::Prefetch( ReferenceToLocalMatrices.GetData() + RefToLocalIdx + PreFetchStride );
-			FPlatformMisc::Prefetch( ReferenceToLocalMatrices.GetData() + RefToLocalIdx + PreFetchStride, PLATFORM_CACHE_LINE_SIZE );
-
-			FMatrix3x4& BoneMat = ChunkMatrices[BoneIdx];
-			const FMatrix& RefToLocal = ReferenceToLocalMatrices[RefToLocalIdx];
-			RefToLocal.To3x4MatrixTranspose( (float*)BoneMat.M );
-		}
-		GDynamicRHI->RHIUnlockVertexBuffer(VertexBuffer);
-	}
-};
-
 bool FGPUBaseSkinVertexFactory::FShaderDataType::UpdateBoneData(FRHICommandListImmediate& RHICmdList, const TArray<FMatrix>& ReferenceToLocalMatrices,
 	const TArray<FBoneIndexType>& BoneMap, uint32 RevisionNumber, bool bPrevious, ERHIFeatureLevel::Type InFeatureLevel, bool bUseSkinCache)
 {
@@ -276,7 +235,34 @@ bool FGPUBaseSkinVertexFactory::FShaderDataType::UpdateBoneData(FRHICommandListI
 		{
 			if (!bUseSkinCache && DeferSkeletalLockAndFillToRHIThread())
 			{
-				ALLOC_COMMAND_CL(RHICmdList, FRHICommandUpdateBoneBuffer)(CurrentBoneBuffer->VertexBufferRHI, VectorArraySize, ReferenceToLocalMatrices, BoneMap);
+				FVertexBufferRHIParamRef VertexBuffer = CurrentBoneBuffer->VertexBufferRHI;
+				RHICmdList.EnqueueLambda([VertexBuffer, VectorArraySize, &ReferenceToLocalMatrices, &BoneMap](FRHICommandListImmediate& RHICmdList)
+				{
+					QUICK_SCOPE_CYCLE_COUNTER(STAT_FRHICommandUpdateBoneBuffer_Execute);
+					FMatrix3x4* ChunkMatrices = (FMatrix3x4*)RHICmdList.LockVertexBuffer(VertexBuffer, 0, VectorArraySize, RLM_WriteOnly);
+					//FMatrix3x4 is sizeof() == 48
+					// PLATFORM_CACHE_LINE_SIZE (128) / 48 = 2.6
+					//  sizeof(FMatrix) == 64
+					// PLATFORM_CACHE_LINE_SIZE (128) / 64 = 2
+					const uint32 NumBones = BoneMap.Num();
+					check(NumBones > 0 && NumBones < 256); // otherwise maybe some bad threading on BoneMap, maybe we need to copy that
+					const int32 PreFetchStride = 2; // FPlatformMisc::Prefetch stride
+					for (uint32 BoneIdx = 0; BoneIdx < NumBones; BoneIdx++)
+					{
+						const FBoneIndexType RefToLocalIdx = BoneMap[BoneIdx];
+						check(ReferenceToLocalMatrices.IsValidIndex(RefToLocalIdx)); // otherwise maybe some bad threading on BoneMap, maybe we need to copy that
+						FPlatformMisc::Prefetch(ReferenceToLocalMatrices.GetData() + RefToLocalIdx + PreFetchStride);
+						FPlatformMisc::Prefetch(ReferenceToLocalMatrices.GetData() + RefToLocalIdx + PreFetchStride, PLATFORM_CACHE_LINE_SIZE);
+
+						FMatrix3x4& BoneMat = ChunkMatrices[BoneIdx];
+						const FMatrix& RefToLocal = ReferenceToLocalMatrices[RefToLocalIdx];
+						RefToLocal.To3x4MatrixTranspose((float*)BoneMat.M);
+					}
+					RHICmdList.UnlockVertexBuffer(VertexBuffer);
+				});
+
+				RHICmdList.RHIThreadFence(true);
+
 				return true;
 			}
 			ChunkMatrices = (FMatrix3x4*)RHILockVertexBuffer(CurrentBoneBuffer->VertexBufferRHI, 0, VectorArraySize, RLM_WriteOnly);
@@ -900,44 +886,6 @@ protected:
 	TGPUSkinAPEXClothVertexFactory::ClothShaderType
 -----------------------------------------------------------------------------*/
 
-struct FRHICommandUpdateClothBuffer final : public FRHICommand<FRHICommandUpdateClothBuffer>
-{
-	FVertexBufferRHIParamRef VertexBuffer;
-	uint32 BufferSize;
-	const TArray<FVector>& SimulPositions;
-	const TArray<FVector>& SimulNormals;
-
-
-	FORCEINLINE_DEBUGGABLE FRHICommandUpdateClothBuffer(FVertexBufferRHIParamRef InVertexBuffer, uint32 InBufferSize, const TArray<FVector>& InSimulPositions, const TArray<FVector>& InSimulNormals)
-		: VertexBuffer(InVertexBuffer)
-		, BufferSize(InBufferSize)
-		, SimulPositions(InSimulPositions)
-		, SimulNormals(InSimulNormals)
-	{
-	}
-	void Execute(FRHICommandListBase& CmdList)
-	{
-		QUICK_SCOPE_CYCLE_COUNTER(STAT_FRHICommandUpdateBoneBuffer_Execute);
-		float* RESTRICT Data = (float* RESTRICT)GDynamicRHI->RHILockVertexBuffer(VertexBuffer, 0, BufferSize, RLM_WriteOnly);
-		uint32 NumSimulVerts = SimulPositions.Num();
-		check(NumSimulVerts > 0 && NumSimulVerts <= MAX_APEXCLOTH_VERTICES_FOR_VB);
-		float* RESTRICT Pos = (float* RESTRICT) &SimulPositions[0].X;
-		float* RESTRICT Normal = (float* RESTRICT) &SimulNormals[0].X;
-		for (uint32 Index = 0; Index < NumSimulVerts; Index++)
-		{
-			FPlatformMisc::Prefetch(Pos + PLATFORM_CACHE_LINE_SIZE);
-			FPlatformMisc::Prefetch(Normal + PLATFORM_CACHE_LINE_SIZE);
-
-			FMemory::Memcpy(Data, Pos, sizeof(float) * 3);
-			FMemory::Memcpy(Data + 3, Normal, sizeof(float) * 3);
-			Data += 6;
-			Pos += 3;
-			Normal += 3;
-		}
-		GDynamicRHI->RHIUnlockVertexBuffer(VertexBuffer);
-	}
-};
-
 bool FGPUBaseSkinAPEXClothVertexFactory::ClothShaderType::UpdateClothSimulData(FRHICommandListImmediate& RHICmdList, const TArray<FVector>& InSimulPositions,
 	const TArray<FVector>& InSimulNormals, uint32 FrameNumberToPrepare, ERHIFeatureLevel::Type FeatureLevel)
 {
@@ -971,7 +919,31 @@ bool FGPUBaseSkinAPEXClothVertexFactory::ClothShaderType::UpdateClothSimulData(F
 		{
 			if (DeferSkeletalLockAndFillToRHIThread())
 			{
-				ALLOC_COMMAND_CL(RHICmdList, FRHICommandUpdateClothBuffer)(CurrentClothBuffer->VertexBufferRHI, VectorArraySize, InSimulPositions, InSimulNormals);
+				FVertexBufferRHIParamRef VertexBuffer = CurrentClothBuffer->VertexBufferRHI;
+				RHICmdList.EnqueueLambda([VertexBuffer, VectorArraySize, &InSimulPositions, &InSimulNormals](FRHICommandListImmediate& RHICmdList)
+				{
+					QUICK_SCOPE_CYCLE_COUNTER(STAT_FRHICommandUpdateBoneBuffer_Execute);
+					float* RESTRICT Data = (float* RESTRICT)RHICmdList.LockVertexBuffer(VertexBuffer, 0, VectorArraySize, RLM_WriteOnly);
+					uint32 NumSimulVerts = InSimulPositions.Num();
+					check(NumSimulVerts > 0 && NumSimulVerts <= MAX_APEXCLOTH_VERTICES_FOR_VB);
+					float* RESTRICT Pos = (float* RESTRICT) &InSimulPositions[0].X;
+					float* RESTRICT Normal = (float* RESTRICT) &InSimulNormals[0].X;
+					for (uint32 Index = 0; Index < NumSimulVerts; Index++)
+					{
+						FPlatformMisc::Prefetch(Pos + PLATFORM_CACHE_LINE_SIZE);
+						FPlatformMisc::Prefetch(Normal + PLATFORM_CACHE_LINE_SIZE);
+
+						FMemory::Memcpy(Data, Pos, sizeof(float) * 3);
+						FMemory::Memcpy(Data + 3, Normal, sizeof(float) * 3);
+						Data += 6;
+						Pos += 3;
+						Normal += 3;
+					}
+					RHICmdList.UnlockVertexBuffer(VertexBuffer);
+				});
+
+				RHICmdList.RHIThreadFence(true);
+
 				return true;
 			}
 			float* RESTRICT Data = (float* RESTRICT)RHILockVertexBuffer(CurrentClothBuffer->VertexBufferRHI, 0, VectorArraySize, RLM_WriteOnly);
diff --git a/Engine/Source/Runtime/NullDrv/Public/NullRHI.h b/Engine/Source/Runtime/NullDrv/Public/NullRHI.h
--- a/Engine/Source/Runtime/NullDrv/Public/NullRHI.h
+++ b/Engine/Source/Runtime/NullDrv/Public/NullRHI.h
@@ -132,11 +132,11 @@ public:
 		return new FRHIIndexBuffer(Stride,Size,InUsage); 
 	}
 
-	virtual void* RHILockIndexBuffer(FIndexBufferRHIParamRef IndexBuffer, uint32 Offset, uint32 Size, EResourceLockMode LockMode) final override
+	virtual void* LockIndexBuffer_BottomOfPipe(FRHICommandListImmediate& RHICmdList, FIndexBufferRHIParamRef IndexBuffer, uint32 Offset, uint32 Size, EResourceLockMode LockMode) final override
 	{ 
 		return GetStaticBuffer(); 
 	}
-	virtual void RHIUnlockIndexBuffer(FIndexBufferRHIParamRef IndexBuffer) final override
+	virtual void UnlockIndexBuffer_BottomOfPipe(FRHICommandListImmediate& RHICmdList, FIndexBufferRHIParamRef IndexBuffer) final override
 	{
 
 	}
@@ -155,11 +155,11 @@ public:
 		return new FRHIVertexBuffer(Size,InUsage); 
 	}
 
-	virtual void* RHILockVertexBuffer(FVertexBufferRHIParamRef VertexBuffer, uint32 Offset, uint32 SizeRHI, EResourceLockMode LockMode) final override
+	virtual void* LockVertexBuffer_BottomOfPipe(FRHICommandListImmediate& RHICmdList, FVertexBufferRHIParamRef VertexBuffer, uint32 Offset, uint32 SizeRHI, EResourceLockMode LockMode) final override
 	{ 
 		return GetStaticBuffer(); 
 	}
-	virtual void RHIUnlockVertexBuffer(FVertexBufferRHIParamRef VertexBuffer) final override
+	virtual void UnlockVertexBuffer_BottomOfPipe(FRHICommandListImmediate& RHICmdList, FVertexBufferRHIParamRef VertexBuffer) final override
 	{
 
 	}
@@ -184,11 +184,11 @@ public:
 		return new FRHIStructuredBuffer(Stride,Size,InUsage); 
 	}
 
-	virtual void* RHILockStructuredBuffer(FStructuredBufferRHIParamRef StructuredBuffer, uint32 Offset, uint32 SizeRHI, EResourceLockMode LockMode) final override
+	virtual void* LockStructuredBuffer_BottomOfPipe(FRHICommandListImmediate& RHICmdList, FStructuredBufferRHIParamRef StructuredBuffer, uint32 Offset, uint32 SizeRHI, EResourceLockMode LockMode) final override
 	{ 
 		return GetStaticBuffer(); 
 	}
-	virtual void RHIUnlockStructuredBuffer(FStructuredBufferRHIParamRef StructuredBuffer) final override
+	virtual void UnlockStructuredBuffer_BottomOfPipe(FRHICommandListImmediate& RHICmdList, FStructuredBufferRHIParamRef StructuredBuffer) final override
 	{
 
 	}
diff --git a/Engine/Source/Runtime/OpenGLDrv/Private/OpenGLIndexBuffer.cpp b/Engine/Source/Runtime/OpenGLDrv/Private/OpenGLIndexBuffer.cpp
--- a/Engine/Source/Runtime/OpenGLDrv/Private/OpenGLIndexBuffer.cpp
+++ b/Engine/Source/Runtime/OpenGLDrv/Private/OpenGLIndexBuffer.cpp
@@ -34,9 +34,8 @@ FIndexBufferRHIRef FOpenGLDynamicRHI::CreateIndexBuffer_RenderThread(class FRHIC
 	return this->RHICreateIndexBuffer(Stride, Size, InUsage, CreateInfo);
 }
 
-void* FOpenGLDynamicRHI::RHILockIndexBuffer(FIndexBufferRHIParamRef IndexBufferRHI,uint32 Offset,uint32 Size,EResourceLockMode LockMode)
+void* FOpenGLDynamicRHI::LockIndexBuffer_BottomOfPipe(FRHICommandListImmediate& RHICmdList, FIndexBufferRHIParamRef IndexBufferRHI, uint32 Offset, uint32 Size, EResourceLockMode LockMode)
 {
-	FRHICommandListImmediate& RHICmdList = FRHICommandListExecutor::GetImmediateCommandList();
 	RHITHREAD_GLCOMMAND_PROLOGUE();
 	VERIFY_GL_SCOPE();
 	FOpenGLIndexBuffer* IndexBuffer = ResourceCast(IndexBufferRHI);
@@ -44,9 +43,8 @@ void* FOpenGLDynamicRHI::RHILockIndexBuffer(FIndexBufferRHIParamRef IndexBufferR
 	RHITHREAD_GLCOMMAND_EPILOGUE_RETURN(void*);
 }
 
-void FOpenGLDynamicRHI::RHIUnlockIndexBuffer(FIndexBufferRHIParamRef IndexBufferRHI)
+void FOpenGLDynamicRHI::UnlockIndexBuffer_BottomOfPipe(FRHICommandListImmediate& RHICmdList, FIndexBufferRHIParamRef IndexBufferRHI)
 {
-	FRHICommandListImmediate& RHICmdList = FRHICommandListExecutor::GetImmediateCommandList();
 	RHITHREAD_GLCOMMAND_PROLOGUE();
 	VERIFY_GL_SCOPE();
 	FOpenGLIndexBuffer* IndexBuffer = ResourceCast(IndexBufferRHI);
diff --git a/Engine/Source/Runtime/OpenGLDrv/Private/OpenGLStructuredBuffer.cpp b/Engine/Source/Runtime/OpenGLDrv/Private/OpenGLStructuredBuffer.cpp
--- a/Engine/Source/Runtime/OpenGLDrv/Private/OpenGLStructuredBuffer.cpp
+++ b/Engine/Source/Runtime/OpenGLDrv/Private/OpenGLStructuredBuffer.cpp
@@ -25,14 +25,14 @@ FStructuredBufferRHIRef FOpenGLDynamicRHI::RHICreateStructuredBuffer(uint32 Stri
 	return StructuredBuffer.GetReference();
 }
 
-void* FOpenGLDynamicRHI::RHILockStructuredBuffer(FStructuredBufferRHIParamRef StructuredBufferRHI,uint32 Offset,uint32 Size,EResourceLockMode LockMode)
+void* FOpenGLDynamicRHI::LockStructuredBuffer_BottomOfPipe(FRHICommandListImmediate& RHICmdList, FStructuredBufferRHIParamRef StructuredBufferRHI, uint32 Offset, uint32 Size, EResourceLockMode LockMode)
 {
 	VERIFY_GL_SCOPE();
 	FOpenGLStructuredBuffer* StructuredBuffer = ResourceCast(StructuredBufferRHI);
 	return StructuredBuffer->Lock(Offset, Size, LockMode == RLM_ReadOnly, StructuredBuffer->IsDynamic());
 }
 
-void FOpenGLDynamicRHI::RHIUnlockStructuredBuffer(FStructuredBufferRHIParamRef StructuredBufferRHI)
+void FOpenGLDynamicRHI::UnlockStructuredBuffer_BottomOfPipe(FRHICommandListImmediate& RHICmdList, FStructuredBufferRHIParamRef StructuredBufferRHI)
 {
 	VERIFY_GL_SCOPE();
 	FOpenGLStructuredBuffer* StructuredBuffer = ResourceCast(StructuredBufferRHI);
diff --git a/Engine/Source/Runtime/OpenGLDrv/Private/OpenGLVertexBuffer.cpp b/Engine/Source/Runtime/OpenGLDrv/Private/OpenGLVertexBuffer.cpp
--- a/Engine/Source/Runtime/OpenGLDrv/Private/OpenGLVertexBuffer.cpp
+++ b/Engine/Source/Runtime/OpenGLDrv/Private/OpenGLVertexBuffer.cpp
@@ -192,10 +192,9 @@ FVertexBufferRHIRef FOpenGLDynamicRHI::RHICreateVertexBuffer(uint32 Size, uint32
 	return VertexBuffer.GetReference();
 }
 
-void* FOpenGLDynamicRHI::RHILockVertexBuffer(FVertexBufferRHIParamRef VertexBufferRHI,uint32 Offset,uint32 Size,EResourceLockMode LockMode)
+void* FOpenGLDynamicRHI::LockVertexBuffer_BottomOfPipe(FRHICommandListImmediate& RHICmdList, FVertexBufferRHIParamRef VertexBufferRHI, uint32 Offset, uint32 Size, EResourceLockMode LockMode)
 {
 	check(Size > 0);
-	FRHICommandListImmediate& RHICmdList = FRHICommandListExecutor::GetImmediateCommandList();
 	RHITHREAD_GLCOMMAND_PROLOGUE();
 
 	VERIFY_GL_SCOPE();
@@ -222,9 +221,8 @@ void* FOpenGLDynamicRHI::RHILockVertexBuffer(FVertexBufferRHIParamRef VertexBuff
 	RHITHREAD_GLCOMMAND_EPILOGUE_RETURN(void*);
 }
 
-void FOpenGLDynamicRHI::RHIUnlockVertexBuffer(FVertexBufferRHIParamRef VertexBufferRHI)
+void FOpenGLDynamicRHI::UnlockVertexBuffer_BottomOfPipe(FRHICommandListImmediate& RHICmdList, FVertexBufferRHIParamRef VertexBufferRHI)
 {
-	FRHICommandListImmediate& RHICmdList = FRHICommandListExecutor::GetImmediateCommandList();
 	RHITHREAD_GLCOMMAND_PROLOGUE();
 	VERIFY_GL_SCOPE();
 	FOpenGLVertexBuffer* VertexBuffer = ResourceCast(VertexBufferRHI);
diff --git a/Engine/Source/Runtime/OpenGLDrv/Public/OpenGLDrv.h b/Engine/Source/Runtime/OpenGLDrv/Public/OpenGLDrv.h
--- a/Engine/Source/Runtime/OpenGLDrv/Public/OpenGLDrv.h
+++ b/Engine/Source/Runtime/OpenGLDrv/Public/OpenGLDrv.h
@@ -364,15 +364,15 @@ public:
 	virtual FUniformBufferRHIRef RHICreateUniformBuffer(const void* Contents, const FRHIUniformBufferLayout& Layout, EUniformBufferUsage Usage, EUniformBufferValidation Validation) final override;
 	virtual void RHIUpdateUniformBuffer(FUniformBufferRHIParamRef UniformBufferRHI, const void* Contents) final override;
 	virtual FIndexBufferRHIRef RHICreateIndexBuffer(uint32 Stride, uint32 Size, uint32 InUsage, FRHIResourceCreateInfo& CreateInfo) final override;
-	virtual void* RHILockIndexBuffer(FIndexBufferRHIParamRef IndexBuffer, uint32 Offset, uint32 Size, EResourceLockMode LockMode) final override;
-	virtual void RHIUnlockIndexBuffer(FIndexBufferRHIParamRef IndexBuffer) final override;
+	virtual void* LockIndexBuffer_BottomOfPipe(FRHICommandListImmediate& RHICmdList, FIndexBufferRHIParamRef IndexBuffer, uint32 Offset, uint32 Size, EResourceLockMode LockMode) final override;
+	virtual void UnlockIndexBuffer_BottomOfPipe(FRHICommandListImmediate& RHICmdList, FIndexBufferRHIParamRef IndexBuffer) final override;
 	virtual FVertexBufferRHIRef RHICreateVertexBuffer(uint32 Size, uint32 InUsage, FRHIResourceCreateInfo& CreateInfo) final override;
-	virtual void* RHILockVertexBuffer(FVertexBufferRHIParamRef VertexBuffer, uint32 Offset, uint32 SizeRHI, EResourceLockMode LockMode) final override;
-	virtual void RHIUnlockVertexBuffer(FVertexBufferRHIParamRef VertexBuffer) final override;
+	virtual void* LockVertexBuffer_BottomOfPipe(FRHICommandListImmediate& RHICmdList, FVertexBufferRHIParamRef VertexBuffer, uint32 Offset, uint32 SizeRHI, EResourceLockMode LockMode) final override;
+	virtual void UnlockVertexBuffer_BottomOfPipe(FRHICommandListImmediate& RHICmdList, FVertexBufferRHIParamRef VertexBuffer) final override;
 	virtual void RHICopyVertexBuffer(FVertexBufferRHIParamRef SourceBuffer,FVertexBufferRHIParamRef DestBuffer) final override;
 	virtual FStructuredBufferRHIRef RHICreateStructuredBuffer(uint32 Stride, uint32 Size, uint32 InUsage, FRHIResourceCreateInfo& CreateInfo) final override;
-	virtual void* RHILockStructuredBuffer(FStructuredBufferRHIParamRef StructuredBuffer, uint32 Offset, uint32 SizeRHI, EResourceLockMode LockMode) final override;
-	virtual void RHIUnlockStructuredBuffer(FStructuredBufferRHIParamRef StructuredBuffer) final override;
+	virtual void* LockStructuredBuffer_BottomOfPipe(FRHICommandListImmediate& RHICmdList, FStructuredBufferRHIParamRef StructuredBuffer, uint32 Offset, uint32 SizeRHI, EResourceLockMode LockMode) final override;
+	virtual void UnlockStructuredBuffer_BottomOfPipe(FRHICommandListImmediate& RHICmdList, FStructuredBufferRHIParamRef StructuredBuffer) final override;
 	virtual FUnorderedAccessViewRHIRef RHICreateUnorderedAccessView(FStructuredBufferRHIParamRef StructuredBuffer, bool bUseUAVCounter, bool bAppendBuffer) final override;
 	virtual FUnorderedAccessViewRHIRef RHICreateUnorderedAccessView(FTextureRHIParamRef Texture, uint32 MipLevel) final override;
 	virtual FUnorderedAccessViewRHIRef RHICreateUnorderedAccessView(FVertexBufferRHIParamRef VertexBuffer, uint8 Format) final override;
diff --git a/Engine/Source/Runtime/RHI/Private/RHICommandList.cpp b/Engine/Source/Runtime/RHI/Private/RHICommandList.cpp
--- a/Engine/Source/Runtime/RHI/Private/RHICommandList.cpp
+++ b/Engine/Source/Runtime/RHI/Private/RHICommandList.cpp
@@ -1870,7 +1870,7 @@ void FRHICommandListBase::operator delete(void *RawMemory)
 FVertexBufferRHIRef FDynamicRHI::CreateAndLockVertexBuffer_RenderThread(class FRHICommandListImmediate& RHICmdList, uint32 Size, uint32 InUsage, FRHIResourceCreateInfo& CreateInfo, void*& OutDataBuffer)
 {
 	FVertexBufferRHIRef VertexBuffer = CreateVertexBuffer_RenderThread(RHICmdList, Size, InUsage, CreateInfo);
-	OutDataBuffer = LockVertexBuffer_RenderThread(RHICmdList, VertexBuffer, 0, Size, RLM_WriteOnly);
+	OutDataBuffer = RHILockVertexBuffer(RHICmdList, VertexBuffer, 0, Size, RLM_WriteOnly);
 
 	return VertexBuffer;
 }
@@ -1878,7 +1878,7 @@ FVertexBufferRHIRef FDynamicRHI::CreateAndLockVertexBuffer_RenderThread(class FR
 FIndexBufferRHIRef FDynamicRHI::CreateAndLockIndexBuffer_RenderThread(class FRHICommandListImmediate& RHICmdList, uint32 Stride, uint32 Size, uint32 InUsage, FRHIResourceCreateInfo& CreateInfo, void*& OutDataBuffer)
 {
 	FIndexBufferRHIRef IndexBuffer = CreateIndexBuffer_RenderThread(RHICmdList, Stride, Size, InUsage, CreateInfo);
-	OutDataBuffer = LockIndexBuffer_RenderThread(RHICmdList, IndexBuffer, 0, Size, RLM_WriteOnly);
+	OutDataBuffer = RHILockIndexBuffer(RHICmdList, IndexBuffer, 0, Size, RLM_WriteOnly);
 	return IndexBuffer;
 }
 
@@ -1918,173 +1918,240 @@ FShaderResourceViewRHIRef FDynamicRHI::CreateShaderResourceView_RenderThread(cla
 	return GDynamicRHI->RHICreateShaderResourceView(Buffer);
 }
 
-struct FRHICommandUpdateVertexBuffer final : public FRHICommand<FRHICommandUpdateVertexBuffer>
+static FLockTracker GLockTracker;
+
+void* FDynamicRHI::RHILockVertexBuffer(class FRHICommandListImmediate& RHICmdList, FVertexBufferRHIParamRef VertexBuffer, uint32 Offset, uint32 SizeRHI, EResourceLockMode LockMode)
 {
-	FVertexBufferRHIParamRef VertexBuffer;
-	void* Buffer;
-	uint32 BufferSize;
-	uint32 Offset;
+	QUICK_SCOPE_CYCLE_COUNTER(STAT_FDynamicRHI_LockVertexBuffer);
 
-	FORCEINLINE_DEBUGGABLE FRHICommandUpdateVertexBuffer(FVertexBufferRHIParamRef InVertexBuffer, void* InBuffer, uint32 InOffset, uint32 InBufferSize)
-		: VertexBuffer(InVertexBuffer)
-		, Buffer(InBuffer)
-		, BufferSize(InBufferSize)
-		, Offset(InOffset)
+	void* Result;
+	if (RHICmdList.IsTopOfPipe())
 	{
+		bool bBuffer = CVarRHICmdBufferWriteLocks.GetValueOnRenderThread() > 0;
+		if (!bBuffer || LockMode != RLM_WriteOnly)
+		{
+			QUICK_SCOPE_CYCLE_COUNTER(STAT_RHIMETHOD_LockVertexBuffer_FlushAndLock);
+
+			FRHICommandListScopedFlushAndExecute Flush(RHICmdList);
+			Result = GDynamicRHI->LockVertexBuffer_BottomOfPipe(RHICmdList, VertexBuffer, Offset, SizeRHI, LockMode);
+		}
+		else
+		{
+			QUICK_SCOPE_CYCLE_COUNTER(STAT_RHIMETHOD_LockVertexBuffer_Malloc);
+			Result = FMemory::Malloc(SizeRHI, 16);
+		}
+
+		// Only use the lock tracker at the top of the pipe. There's no need to track locks
+		// at the bottom of the pipe, and doing so would require a critical section.
+		GLockTracker.Lock(VertexBuffer, Result, Offset, SizeRHI, LockMode);
 	}
-	void Execute(FRHICommandListBase& CmdList)
+	else
 	{
-		QUICK_SCOPE_CYCLE_COUNTER(STAT_FRHICommandUpdateVertexBuffer_Execute);
-		void* Data = GDynamicRHI->RHILockVertexBuffer(VertexBuffer, Offset, BufferSize, RLM_WriteOnly);
-		FMemory::Memcpy(Data, Buffer, BufferSize);
-		FMemory::Free(Buffer);
-		GDynamicRHI->RHIUnlockVertexBuffer(VertexBuffer);
+		Result = GDynamicRHI->LockVertexBuffer_BottomOfPipe(RHICmdList, VertexBuffer, Offset, SizeRHI, LockMode);
 	}
-};
 
-struct FRHICommandUpdateIndexBuffer final : public FRHICommand<FRHICommandUpdateIndexBuffer>
+	check(Result);
+	return Result;	
+}
+
+void FDynamicRHI::RHIUnlockVertexBuffer(class FRHICommandListImmediate& RHICmdList, FVertexBufferRHIParamRef VertexBuffer)
 {
-	FIndexBufferRHIParamRef IndexBuffer;
-	void* Buffer;
-	uint32 BufferSize;
-	uint32 Offset;
+	QUICK_SCOPE_CYCLE_COUNTER(STAT_FDynamicRHI_UnlockVertexBuffer_RenderThread);
 
-	FORCEINLINE_DEBUGGABLE FRHICommandUpdateIndexBuffer(FIndexBufferRHIParamRef InIndexBuffer, void* InBuffer, uint32 InOffset, uint32 InBufferSize)
-		: IndexBuffer(InIndexBuffer)
-		, Buffer(InBuffer)
-		, BufferSize(InBufferSize)
-		, Offset(InOffset)
+	if (RHICmdList.IsTopOfPipe())
 	{
+		FLockTracker::FLockParams Params = GLockTracker.Unlock(VertexBuffer);
+
+		bool bBuffer = CVarRHICmdBufferWriteLocks.GetValueOnRenderThread() > 0;
+		if (!bBuffer || Params.LockMode != RLM_WriteOnly)
+		{
+			QUICK_SCOPE_CYCLE_COUNTER(STAT_RHIMETHOD_UnlockVertexBuffer_FlushAndUnlock);
+
+			FRHICommandListScopedFlushAndExecute Flush(RHICmdList);
+			GDynamicRHI->UnlockVertexBuffer_BottomOfPipe(RHICmdList, VertexBuffer);
+			GLockTracker.TotalMemoryOutstanding = 0;
+		}
+		else
+		{
+			RHICmdList.EnqueueLambda([VertexBuffer, Params](FRHICommandListImmediate& RHICmdList)
+			{
+				QUICK_SCOPE_CYCLE_COUNTER(STAT_FRHICommandUpdateVertexBuffer_Execute);
+				void* Data = GDynamicRHI->LockVertexBuffer_BottomOfPipe(RHICmdList, VertexBuffer, Params.Offset, Params.BufferSize, RLM_WriteOnly);
+				FMemory::Memcpy(Data, Params.Buffer, Params.BufferSize);
+				FMemory::Free(Params.Buffer);
+				GDynamicRHI->UnlockVertexBuffer_BottomOfPipe(RHICmdList, VertexBuffer);
+			});
+			RHICmdList.RHIThreadFence(true);
+		}
+
+		if (GLockTracker.TotalMemoryOutstanding > 256 * 1024)
+		{
+			QUICK_SCOPE_CYCLE_COUNTER(STAT_RHIMETHOD_UnlockVertexBuffer_FlushForMem);
+			// we could be loading a level or something, lets get this stuff going
+			RHICmdList.ImmediateFlush(EImmediateFlushType::DispatchToRHIThread);
+			GLockTracker.TotalMemoryOutstanding = 0;
+		}
 	}
-	void Execute(FRHICommandListBase& CmdList)
+	else
 	{
-		QUICK_SCOPE_CYCLE_COUNTER(STAT_FRHICommandUpdateIndexBuffer_Execute);
-		void* Data = GDynamicRHI->RHILockIndexBuffer(IndexBuffer, Offset, BufferSize, RLM_WriteOnly);
-		FMemory::Memcpy(Data, Buffer, BufferSize);
-		FMemory::Free(Buffer);
-		GDynamicRHI->RHIUnlockIndexBuffer(IndexBuffer);
+		GDynamicRHI->UnlockVertexBuffer_BottomOfPipe(RHICmdList, VertexBuffer);
 	}
-};
-
-static FLockTracker GLockTracker;
-
+}
 
-void* FDynamicRHI::LockVertexBuffer_RenderThread(class FRHICommandListImmediate& RHICmdList, FVertexBufferRHIParamRef VertexBuffer, uint32 Offset, uint32 SizeRHI, EResourceLockMode LockMode)
+void* FDynamicRHI::RHILockIndexBuffer(class FRHICommandListImmediate& RHICmdList, FIndexBufferRHIParamRef IndexBuffer, uint32 Offset, uint32 SizeRHI, EResourceLockMode LockMode)
 {
-	QUICK_SCOPE_CYCLE_COUNTER(STAT_FDynamicRHI_LockVertexBuffer_RenderThread);
-	check(IsInRenderingThread());
-	bool bBuffer = CVarRHICmdBufferWriteLocks.GetValueOnRenderThread() > 0;
+	QUICK_SCOPE_CYCLE_COUNTER(STAT_FDynamicRHI_LockIndexBuffer);
+
 	void* Result;
-	if (!bBuffer || LockMode != RLM_WriteOnly || RHICmdList.Bypass() || !IsRunningRHIInSeparateThread())
+	if (RHICmdList.IsTopOfPipe())
 	{
+		bool bBuffer = CVarRHICmdBufferWriteLocks.GetValueOnRenderThread() > 0;
+		if (!bBuffer || LockMode != RLM_WriteOnly)
 		{
-			QUICK_SCOPE_CYCLE_COUNTER(STAT_RHIMETHOD_LockVertexBuffer_Flush);
-			RHICmdList.ImmediateFlush(EImmediateFlushType::FlushRHIThread);
+			QUICK_SCOPE_CYCLE_COUNTER(STAT_RHIMETHOD_LockIndexBuffer_FlushAndLock);
+
+			FRHICommandListScopedFlushAndExecute Flush(RHICmdList);
+			Result = GDynamicRHI->LockIndexBuffer_BottomOfPipe(RHICmdList, IndexBuffer, Offset, SizeRHI, LockMode);
 		}
+		else
 		{
-			QUICK_SCOPE_CYCLE_COUNTER(STAT_RHIMETHOD_LockVertexBuffer_Lock);
-			Result = GDynamicRHI->RHILockVertexBuffer(VertexBuffer, Offset, SizeRHI, LockMode);
+			QUICK_SCOPE_CYCLE_COUNTER(STAT_RHIMETHOD_LockIndexBuffer_Malloc);
+			Result = FMemory::Malloc(SizeRHI, 16);
 		}
+
+		// Only use the lock tracker at the top of the pipe. There's no need to track locks
+		// at the bottom of the pipe, and doing so would require a critical section.
+		GLockTracker.Lock(IndexBuffer, Result, Offset, SizeRHI, LockMode);
 	}
 	else
 	{
-		QUICK_SCOPE_CYCLE_COUNTER(STAT_RHIMETHOD_LockVertexBuffer_Malloc);
-		Result = FMemory::Malloc(SizeRHI, 16);
+		Result = GDynamicRHI->LockIndexBuffer_BottomOfPipe(RHICmdList, IndexBuffer, Offset, SizeRHI, LockMode);
 	}
+
 	check(Result);
-	GLockTracker.Lock(VertexBuffer, Result, Offset, SizeRHI, LockMode);
 	return Result;
 }
 
-void FDynamicRHI::UnlockVertexBuffer_RenderThread(class FRHICommandListImmediate& RHICmdList, FVertexBufferRHIParamRef VertexBuffer)
+void FDynamicRHI::RHIUnlockIndexBuffer(class FRHICommandListImmediate& RHICmdList, FIndexBufferRHIParamRef IndexBuffer)
 {
-	QUICK_SCOPE_CYCLE_COUNTER(STAT_FDynamicRHI_UnlockVertexBuffer_RenderThread);
-	check(IsInRenderingThread());
-	bool bBuffer = CVarRHICmdBufferWriteLocks.GetValueOnRenderThread() > 0;
-	FLockTracker::FLockParams Params = GLockTracker.Unlock(VertexBuffer);
-	if (!bBuffer || Params.LockMode != RLM_WriteOnly || RHICmdList.Bypass() || !IsRunningRHIInSeparateThread())
+	QUICK_SCOPE_CYCLE_COUNTER(STAT_FDynamicRHI_UnlockIndexBuffer_RenderThread);
+
+	if (RHICmdList.IsTopOfPipe())
 	{
+		FLockTracker::FLockParams Params = GLockTracker.Unlock(IndexBuffer);
+
+		bool bBuffer = CVarRHICmdBufferWriteLocks.GetValueOnRenderThread() > 0;
+		if (!bBuffer || Params.LockMode != RLM_WriteOnly)
 		{
-			QUICK_SCOPE_CYCLE_COUNTER(STAT_RHIMETHOD_UnlockVertexBuffer_Flush);
-			RHICmdList.ImmediateFlush(EImmediateFlushType::FlushRHIThread);
+			QUICK_SCOPE_CYCLE_COUNTER(STAT_RHIMETHOD_UnlockIndexBuffer_FlushAndUnlock);
+
+			FRHICommandListScopedFlushAndExecute Flush(RHICmdList);
+			GDynamicRHI->UnlockIndexBuffer_BottomOfPipe(RHICmdList, IndexBuffer);
+			GLockTracker.TotalMemoryOutstanding = 0;
 		}
+		else
 		{
-			QUICK_SCOPE_CYCLE_COUNTER(STAT_RHIMETHOD_UnlockVertexBuffer_Unlock);
-			GDynamicRHI->RHIUnlockVertexBuffer(VertexBuffer);
-			GLockTracker.TotalMemoryOutstanding = 0;
+			RHICmdList.EnqueueLambda([IndexBuffer, Params](FRHICommandListImmediate& RHICmdList)
+			{
+				QUICK_SCOPE_CYCLE_COUNTER(STAT_FRHICommandUpdateIndexBuffer_Execute);
+				void* Data = GDynamicRHI->LockIndexBuffer_BottomOfPipe(RHICmdList, IndexBuffer, Params.Offset, Params.BufferSize, RLM_WriteOnly);
+				FMemory::Memcpy(Data, Params.Buffer, Params.BufferSize);
+				FMemory::Free(Params.Buffer);
+				GDynamicRHI->UnlockIndexBuffer_BottomOfPipe(RHICmdList, IndexBuffer);
+			});
+			RHICmdList.RHIThreadFence(true);
 		}
-	}	
-	else
-	{
-		ALLOC_COMMAND_CL(RHICmdList, FRHICommandUpdateVertexBuffer)(VertexBuffer, Params.Buffer, Params.Offset, Params.BufferSize);
-		RHICmdList.RHIThreadFence(true);
+
 		if (GLockTracker.TotalMemoryOutstanding > 256 * 1024)
 		{
-			QUICK_SCOPE_CYCLE_COUNTER(STAT_RHIMETHOD_UnlockVertexBuffer_FlushForMem);
+			QUICK_SCOPE_CYCLE_COUNTER(STAT_RHIMETHOD_UnlockIndexBuffer_FlushForMem);
 			// we could be loading a level or something, lets get this stuff going
-			RHICmdList.ImmediateFlush(EImmediateFlushType::DispatchToRHIThread); 
+			RHICmdList.ImmediateFlush(EImmediateFlushType::DispatchToRHIThread);
 			GLockTracker.TotalMemoryOutstanding = 0;
 		}
 	}
+	else
+	{
+		GDynamicRHI->UnlockIndexBuffer_BottomOfPipe(RHICmdList, IndexBuffer);
+	}
 }
 
-void* FDynamicRHI::LockIndexBuffer_RenderThread(class FRHICommandListImmediate& RHICmdList, FIndexBufferRHIParamRef IndexBuffer, uint32 Offset, uint32 SizeRHI, EResourceLockMode LockMode)
+void* FDynamicRHI::RHILockStructuredBuffer(class FRHICommandListImmediate& RHICmdList, FStructuredBufferRHIParamRef StructuredBuffer, uint32 Offset, uint32 SizeRHI, EResourceLockMode LockMode)
 {
-	QUICK_SCOPE_CYCLE_COUNTER(STAT_FDynamicRHI_LockIndexBuffer_RenderThread);
-	check(IsInRenderingThread());
-	bool bBuffer = CVarRHICmdBufferWriteLocks.GetValueOnRenderThread() > 0;
+	QUICK_SCOPE_CYCLE_COUNTER(STAT_FDynamicRHI_LockStructuredBuffer);
+
 	void* Result;
-	if (!bBuffer || LockMode != RLM_WriteOnly || RHICmdList.Bypass() || !IsRunningRHIInSeparateThread())
+	if (RHICmdList.IsTopOfPipe())
 	{
+		bool bBuffer = CVarRHICmdBufferWriteLocks.GetValueOnRenderThread() > 0;
+		if (!bBuffer || LockMode != RLM_WriteOnly)
 		{
-			QUICK_SCOPE_CYCLE_COUNTER(STAT_RHIMETHOD_LockIndexBuffer_Flush);
-			RHICmdList.ImmediateFlush(EImmediateFlushType::FlushRHIThread);
+			QUICK_SCOPE_CYCLE_COUNTER(STAT_RHIMETHOD_LockStructuredBuffer_FlushAndLock);
+
+			FRHICommandListScopedFlushAndExecute Flush(RHICmdList);
+			Result = GDynamicRHI->LockStructuredBuffer_BottomOfPipe(RHICmdList, StructuredBuffer, Offset, SizeRHI, LockMode);
 		}
+		else
 		{
-			QUICK_SCOPE_CYCLE_COUNTER(STAT_RHIMETHOD_LockIndexBuffer_Lock);
-			Result = GDynamicRHI->RHILockIndexBuffer(IndexBuffer, Offset, SizeRHI, LockMode);
+			QUICK_SCOPE_CYCLE_COUNTER(STAT_RHIMETHOD_LockStructuredBuffer_Malloc);
+			Result = FMemory::Malloc(SizeRHI, 16);
 		}
+		
+		// Only use the lock tracker at the top of the pipe. There's no need to track locks
+		// at the bottom of the pipe, and doing so would require a critical section.
+		GLockTracker.Lock(StructuredBuffer, Result, Offset, SizeRHI, LockMode);
 	}
 	else
 	{
-		QUICK_SCOPE_CYCLE_COUNTER(STAT_RHIMETHOD_LockIndexBuffer_Malloc);
-		Result = FMemory::Malloc(SizeRHI, 16);
+		Result = GDynamicRHI->LockStructuredBuffer_BottomOfPipe(RHICmdList, StructuredBuffer, Offset, SizeRHI, LockMode);
 	}
+
 	check(Result);
-	GLockTracker.Lock(IndexBuffer, Result, Offset, SizeRHI, LockMode);
 	return Result;
 }
 
-void FDynamicRHI::UnlockIndexBuffer_RenderThread(class FRHICommandListImmediate& RHICmdList, FIndexBufferRHIParamRef IndexBuffer)
+void FDynamicRHI::RHIUnlockStructuredBuffer(class FRHICommandListImmediate& RHICmdList, FStructuredBufferRHIParamRef StructuredBuffer)
 {
-	QUICK_SCOPE_CYCLE_COUNTER(STAT_FDynamicRHI_UnlockIndexBuffer_RenderThread);
-	check(IsInRenderingThread());
-	bool bBuffer = CVarRHICmdBufferWriteLocks.GetValueOnRenderThread() > 0;
-	FLockTracker::FLockParams Params = GLockTracker.Unlock(IndexBuffer);
-	if (!bBuffer || Params.LockMode != RLM_WriteOnly || RHICmdList.Bypass() || !IsRunningRHIInSeparateThread())
+	QUICK_SCOPE_CYCLE_COUNTER(STAT_FDynamicRHI_UnlockStructuredBuffer_RenderThread);
+
+	if (RHICmdList.IsTopOfPipe())
 	{
+		FLockTracker::FLockParams Params = GLockTracker.Unlock(StructuredBuffer);
+
+		bool bBuffer = CVarRHICmdBufferWriteLocks.GetValueOnRenderThread() > 0;
+		if (!bBuffer || Params.LockMode != RLM_WriteOnly)
 		{
-			QUICK_SCOPE_CYCLE_COUNTER(STAT_RHIMETHOD_UnlockIndexBuffer_Flush);
-			RHICmdList.ImmediateFlush(EImmediateFlushType::FlushRHIThread);
+			QUICK_SCOPE_CYCLE_COUNTER(STAT_RHIMETHOD_UnlockStructuredBuffer_FlushAndUnlock);
+
+			FRHICommandListScopedFlushAndExecute Flush(RHICmdList);
+			GDynamicRHI->UnlockStructuredBuffer_BottomOfPipe(RHICmdList, StructuredBuffer);
+			GLockTracker.TotalMemoryOutstanding = 0;
 		}
+		else
 		{
-			QUICK_SCOPE_CYCLE_COUNTER(STAT_RHIMETHOD_UnlockIndexBuffer_Unlock);
-			GDynamicRHI->RHIUnlockIndexBuffer(IndexBuffer);
-			GLockTracker.TotalMemoryOutstanding = 0;
+			RHICmdList.EnqueueLambda([StructuredBuffer, Params](FRHICommandListImmediate& RHICmdList)
+			{
+				QUICK_SCOPE_CYCLE_COUNTER(STAT_FRHICommandUpdateStructuredBuffer_Execute);
+				void* Data = GDynamicRHI->LockStructuredBuffer_BottomOfPipe(RHICmdList, StructuredBuffer, Params.Offset, Params.BufferSize, RLM_WriteOnly);
+				FMemory::Memcpy(Data, Params.Buffer, Params.BufferSize);
+				FMemory::Free(Params.Buffer);
+				GDynamicRHI->UnlockStructuredBuffer_BottomOfPipe(RHICmdList, StructuredBuffer);
+			});
+			RHICmdList.RHIThreadFence(true);
 		}
-	}	
-	else
-	{
-		ALLOC_COMMAND_CL(RHICmdList, FRHICommandUpdateIndexBuffer)(IndexBuffer, Params.Buffer, Params.Offset, Params.BufferSize);
-		RHICmdList.RHIThreadFence(true);
+
 		if (GLockTracker.TotalMemoryOutstanding > 256 * 1024)
 		{
-			QUICK_SCOPE_CYCLE_COUNTER(STAT_RHIMETHOD_UnlockIndexBuffer_FlushForMem);
+			QUICK_SCOPE_CYCLE_COUNTER(STAT_RHIMETHOD_UnlockStructuredBuffer_FlushForMem);
 			// we could be loading a level or something, lets get this stuff going
-			RHICmdList.ImmediateFlush(EImmediateFlushType::DispatchToRHIThread); 
+			RHICmdList.ImmediateFlush(EImmediateFlushType::DispatchToRHIThread);
 			GLockTracker.TotalMemoryOutstanding = 0;
 		}
 	}
+	else
+	{
+		GDynamicRHI->UnlockStructuredBuffer_BottomOfPipe(RHICmdList, StructuredBuffer);
+	}
 }
 
 // @todo-mattc-staging Default implementation
diff --git a/Engine/Source/Runtime/RHI/Public/DynamicRHI.h b/Engine/Source/Runtime/RHI/Public/DynamicRHI.h
--- a/Engine/Source/Runtime/RHI/Public/DynamicRHI.h
+++ b/Engine/Source/Runtime/RHI/Public/DynamicRHI.h
@@ -442,10 +442,10 @@ public:
 	virtual FIndexBufferRHIRef RHICreateIndexBuffer(uint32 Stride, uint32 Size, uint32 InUsage, FRHIResourceCreateInfo& CreateInfo) = 0;
 
 	// FlushType: Flush RHI Thread
-	virtual void* RHILockIndexBuffer(FIndexBufferRHIParamRef IndexBuffer, uint32 Offset, uint32 Size, EResourceLockMode LockMode) = 0;
+	virtual void* RHILockIndexBuffer(FRHICommandListImmediate& RHICmdList, FIndexBufferRHIParamRef IndexBuffer, uint32 Offset, uint32 Size, EResourceLockMode LockMode);
 
 	// FlushType: Flush RHI Thread
-	virtual void RHIUnlockIndexBuffer(FIndexBufferRHIParamRef IndexBuffer) = 0;
+	virtual void RHIUnlockIndexBuffer(FRHICommandListImmediate& RHICmdList, FIndexBufferRHIParamRef IndexBuffer);
 
 	/**
 	* Transfer metadata and underlying resource from src to dest and release any resource owned by dest.
@@ -461,10 +461,10 @@ public:
 	virtual FVertexBufferRHIRef RHICreateVertexBuffer(uint32 Size, uint32 InUsage, FRHIResourceCreateInfo& CreateInfo) = 0;
 
 	// FlushType: Flush RHI Thread
-	virtual void* RHILockVertexBuffer(FVertexBufferRHIParamRef VertexBuffer, uint32 Offset, uint32 SizeRHI, EResourceLockMode LockMode) = 0;
+	virtual void* RHILockVertexBuffer(FRHICommandListImmediate& RHICmdList, FVertexBufferRHIParamRef VertexBuffer, uint32 Offset, uint32 SizeRHI, EResourceLockMode LockMode);
 
 	// FlushType: Flush RHI Thread
-	virtual void RHIUnlockVertexBuffer(FVertexBufferRHIParamRef VertexBuffer) = 0;
+	virtual void RHIUnlockVertexBuffer(FRHICommandListImmediate& RHICmdList, FVertexBufferRHIParamRef VertexBuffer);
 
 	/** Copies the contents of one vertex buffer to another vertex buffer.  They must have identical sizes. */
 	// FlushType: Flush Immediate (seems dangerous)
@@ -484,10 +484,10 @@ public:
 	virtual FStructuredBufferRHIRef RHICreateStructuredBuffer(uint32 Stride, uint32 Size, uint32 InUsage, FRHIResourceCreateInfo& CreateInfo) = 0;
 
 	// FlushType: Flush RHI Thread
-	virtual void* RHILockStructuredBuffer(FStructuredBufferRHIParamRef StructuredBuffer, uint32 Offset, uint32 SizeRHI, EResourceLockMode LockMode) = 0;
+	virtual void* RHILockStructuredBuffer(FRHICommandListImmediate& RHICmdList, FStructuredBufferRHIParamRef StructuredBuffer, uint32 Offset, uint32 SizeRHI, EResourceLockMode LockMode);
 
 	// FlushType: Flush RHI Thread
-	virtual void RHIUnlockStructuredBuffer(FStructuredBufferRHIParamRef StructuredBuffer) = 0;
+	virtual void RHIUnlockStructuredBuffer(FRHICommandListImmediate& RHICmdList, FStructuredBufferRHIParamRef StructuredBuffer);
 
 	/** Creates an unordered access view of the given structured buffer. */
 	// FlushType: Wait RHI Thread
@@ -1119,14 +1119,10 @@ public:
 	virtual FStructuredBufferRHIRef CreateStructuredBuffer_RenderThread(class FRHICommandListImmediate& RHICmdList, uint32 Stride, uint32 Size, uint32 InUsage, FRHIResourceCreateInfo& CreateInfo);
 	virtual FShaderResourceViewRHIRef CreateShaderResourceView_RenderThread(class FRHICommandListImmediate& RHICmdList, FVertexBufferRHIParamRef VertexBuffer, uint32 Stride, uint8 Format);
 	virtual FShaderResourceViewRHIRef CreateShaderResourceView_RenderThread(class FRHICommandListImmediate& RHICmdList, FIndexBufferRHIParamRef Buffer);
-	virtual void* LockVertexBuffer_RenderThread(class FRHICommandListImmediate& RHICmdList, FVertexBufferRHIParamRef VertexBuffer, uint32 Offset, uint32 SizeRHI, EResourceLockMode LockMode);
-	virtual void UnlockVertexBuffer_RenderThread(class FRHICommandListImmediate& RHICmdList, FVertexBufferRHIParamRef VertexBuffer);
 	virtual FTexture2DRHIRef AsyncReallocateTexture2D_RenderThread(class FRHICommandListImmediate& RHICmdList, FTexture2DRHIParamRef Texture2D, int32 NewMipCount, int32 NewSizeX, int32 NewSizeY, FThreadSafeCounter* RequestStatus);
 	virtual ETextureReallocationStatus FinalizeAsyncReallocateTexture2D_RenderThread(class FRHICommandListImmediate& RHICmdList, FTexture2DRHIParamRef Texture2D, bool bBlockUntilCompleted);
 	virtual ETextureReallocationStatus CancelAsyncReallocateTexture2D_RenderThread(class FRHICommandListImmediate& RHICmdList, FTexture2DRHIParamRef Texture2D, bool bBlockUntilCompleted);
 	virtual FIndexBufferRHIRef CreateIndexBuffer_RenderThread(class FRHICommandListImmediate& RHICmdList, uint32 Stride, uint32 Size, uint32 InUsage, FRHIResourceCreateInfo& CreateInfo);
-	virtual void* LockIndexBuffer_RenderThread(class FRHICommandListImmediate& RHICmdList, FIndexBufferRHIParamRef IndexBuffer, uint32 Offset, uint32 SizeRHI, EResourceLockMode LockMode);
-	virtual void UnlockIndexBuffer_RenderThread(class FRHICommandListImmediate& RHICmdList, FIndexBufferRHIParamRef IndexBuffer);
 	virtual FVertexShaderRHIRef CreateVertexShader_RenderThread(class FRHICommandListImmediate& RHICmdList, const TArray<uint8>& Code);
 	virtual FVertexShaderRHIRef CreateVertexShader_RenderThread(class FRHICommandListImmediate& RHICmdList, FRHIShaderLibrary* Library, FSHAHash Hash);
 	virtual FPixelShaderRHIRef CreatePixelShader_RenderThread(class FRHICommandListImmediate& RHICmdList, const TArray<uint8>& Code);
@@ -1194,6 +1190,42 @@ public:
 	virtual void RHIUnmapStagingSurface_RenderThread(class FRHICommandListImmediate& RHICmdList, FTextureRHIParamRef Texture);
 	virtual void RHIReadSurfaceFloatData_RenderThread(class FRHICommandListImmediate& RHICmdList, FTextureRHIParamRef Texture, FIntRect Rect, TArray<FFloat16Color>& OutData, ECubeFace CubeFace, int32 ArrayIndex, int32 MipIndex);
 
+	// Buffer Lock/Unlock
+	virtual void* LockStructuredBuffer_BottomOfPipe(class FRHICommandListImmediate& RHICmdList, FStructuredBufferRHIParamRef StructuredBuffer, uint32 Offset, uint32 SizeRHI, EResourceLockMode LockMode)
+	{
+		// Either this function or RHILockStructuredBuffer must be implemented by the platform RHI.
+		checkNoEntry();
+		return nullptr;
+	}
+	virtual void* LockVertexBuffer_BottomOfPipe(class FRHICommandListImmediate& RHICmdList, FVertexBufferRHIParamRef VertexBuffer, uint32 Offset, uint32 SizeRHI, EResourceLockMode LockMode)
+	{
+		// Either this function or RHILockVertexBuffer must be implemented by the platform RHI.
+		checkNoEntry();
+		return nullptr;
+	}
+	virtual void* LockIndexBuffer_BottomOfPipe(class FRHICommandListImmediate& RHICmdList, FIndexBufferRHIParamRef IndexBuffer, uint32 Offset, uint32 SizeRHI, EResourceLockMode LockMode)
+	{
+		// Either this function or RHILockIndexBuffer must be implemented by the platform RHI.
+		checkNoEntry();
+		return nullptr;
+	}
+
+	virtual void UnlockStructuredBuffer_BottomOfPipe(class FRHICommandListImmediate& RHICmdList, FStructuredBufferRHIParamRef StructuredBuffer)
+	{
+		// Either this function or RHIUnlockStructuredBuffer must be implemented by the platform RHI.
+		checkNoEntry();
+	}
+	virtual void UnlockVertexBuffer_BottomOfPipe(class FRHICommandListImmediate& RHICmdList, FVertexBufferRHIParamRef VertexBuffer)
+	{
+		// Either this function or RHIUnlockVertexBuffer must be implemented by the platform RHI.
+		checkNoEntry();
+	}
+	virtual void UnlockIndexBuffer_BottomOfPipe(class FRHICommandListImmediate& RHICmdList, FIndexBufferRHIParamRef IndexBuffer)
+	{
+		// Either this function or RHIUnlockIndexBuffer must be implemented by the platform RHI.
+		checkNoEntry();
+	}
+
 	//Utilities
 	virtual void EnableIdealGPUCaptureOptions(bool bEnable);
 	
diff --git a/Engine/Source/Runtime/RHI/Public/RHICommandList.h b/Engine/Source/Runtime/RHI/Public/RHICommandList.h
--- a/Engine/Source/Runtime/RHI/Public/RHICommandList.h
+++ b/Engine/Source/Runtime/RHI/Public/RHICommandList.h
@@ -537,6 +537,7 @@ private:
 
 	friend class FRHICommandListExecutor;
 	friend class FRHICommandListIterator;
+	friend class FRHICommandListScopedFlushAndExecute;
 
 protected:
 	bool bAsyncPSOCompileAllowed;
@@ -3508,27 +3509,29 @@ public:
 	//Queue the given async compute commandlists in order with the current immediate commandlist
 	void QueueAsyncCompute(FRHIAsyncComputeCommandList& RHIComputeCmdList);
 
+	FORCEINLINE bool IsBottomOfPipe()
+	{
+		return Bypass() || IsExecuting();
+	}
+
+	FORCEINLINE bool IsTopOfPipe()
+	{
+		return !IsBottomOfPipe();
+	}
+
 	template <typename LAMBDA>
-	FORCEINLINE_DEBUGGABLE bool EnqueueLambda(bool bRunOnCurrentThread, LAMBDA&& Lambda)
+	FORCEINLINE_DEBUGGABLE void EnqueueLambda(LAMBDA&& Lambda)
 	{
-		if (bRunOnCurrentThread)
+		if (IsBottomOfPipe())
 		{
 			Lambda(*this);
-			return false;
 		}
 		else
 		{
 			ALLOC_COMMAND(TRHILambdaCommand<LAMBDA>)(Forward<LAMBDA>(Lambda));
-			return true;
 		}
 	}
 
-	template <typename LAMBDA>
-	FORCEINLINE_DEBUGGABLE bool EnqueueLambda(LAMBDA&& Lambda)
-	{
-		return EnqueueLambda(Bypass(), Forward<LAMBDA>(Lambda));
-	}
-
 	FORCEINLINE FSamplerStateRHIRef CreateSamplerState(const FSamplerStateInitializerRHI& Initializer)
 	{
 		LLM_SCOPE(ELLMTag::RHIMisc);
@@ -3690,12 +3693,12 @@ public:
 	
 	FORCEINLINE void* LockIndexBuffer(FIndexBufferRHIParamRef IndexBuffer, uint32 Offset, uint32 SizeRHI, EResourceLockMode LockMode)
 	{
-		return GDynamicRHI->LockIndexBuffer_RenderThread(*this, IndexBuffer, Offset, SizeRHI, LockMode);
+		return GDynamicRHI->RHILockIndexBuffer(*this, IndexBuffer, Offset, SizeRHI, LockMode);
 	}
 	
 	FORCEINLINE void UnlockIndexBuffer(FIndexBufferRHIParamRef IndexBuffer)
 	{
-		GDynamicRHI->UnlockIndexBuffer_RenderThread(*this, IndexBuffer);
+		GDynamicRHI->RHIUnlockIndexBuffer(*this, IndexBuffer);
 	}
 	
 	FORCEINLINE void* LockStagingBuffer(FRHIStagingBuffer* StagingBuffer, uint32 Offset, uint32 SizeRHI)
@@ -3720,12 +3723,12 @@ public:
 	
 	FORCEINLINE void* LockVertexBuffer(FVertexBufferRHIParamRef VertexBuffer, uint32 Offset, uint32 SizeRHI, EResourceLockMode LockMode)
 	{
-		return GDynamicRHI->LockVertexBuffer_RenderThread(*this, VertexBuffer, Offset, SizeRHI, LockMode);
+		return GDynamicRHI->RHILockVertexBuffer(*this, VertexBuffer, Offset, SizeRHI, LockMode);
 	}
 	
 	FORCEINLINE void UnlockVertexBuffer(FVertexBufferRHIParamRef VertexBuffer)
 	{
-		GDynamicRHI->UnlockVertexBuffer_RenderThread(*this, VertexBuffer);
+		GDynamicRHI->RHIUnlockVertexBuffer(*this, VertexBuffer);
 	}
 	
 	FORCEINLINE void CopyVertexBuffer(FVertexBufferRHIParamRef SourceBuffer,FVertexBufferRHIParamRef DestBuffer)
@@ -3744,19 +3747,13 @@ public:
 	FORCEINLINE void* LockStructuredBuffer(FStructuredBufferRHIParamRef StructuredBuffer, uint32 Offset, uint32 SizeRHI, EResourceLockMode LockMode)
 	{
 		LLM_SCOPE(ELLMTag::RHIMisc);
-		QUICK_SCOPE_CYCLE_COUNTER(STAT_RHIMETHOD_LockStructuredBuffer_Flush);
-		ImmediateFlush(EImmediateFlushType::FlushRHIThread); 
-		 
-		return GDynamicRHI->RHILockStructuredBuffer(StructuredBuffer, Offset, SizeRHI, LockMode);
+		return GDynamicRHI->RHILockStructuredBuffer(*this, StructuredBuffer, Offset, SizeRHI, LockMode);
 	}
 	
 	FORCEINLINE void UnlockStructuredBuffer(FStructuredBufferRHIParamRef StructuredBuffer)
 	{
 		LLM_SCOPE(ELLMTag::RHIMisc);
-		QUICK_SCOPE_CYCLE_COUNTER(STAT_RHIMETHOD_UnlockStructuredBuffer_Flush);
-		ImmediateFlush(EImmediateFlushType::FlushRHIThread); 
-		  
-		GDynamicRHI->RHIUnlockStructuredBuffer(StructuredBuffer);
+		GDynamicRHI->RHIUnlockStructuredBuffer(*this, StructuredBuffer);
 	}
 	
 	FORCEINLINE FUnorderedAccessViewRHIRef CreateUnorderedAccessView(FStructuredBufferRHIParamRef StructuredBuffer, bool bUseUAVCounter, bool bAppendBuffer)
@@ -4360,7 +4357,25 @@ public:
 	void UpdateRHIResources(FRHIResourceUpdateInfo* UpdateInfos, int32 Num, bool bNeedReleaseRefs);
 };
 
- struct FScopedGPUMask
+class FRHICommandListScopedFlushAndExecute
+{
+	FRHICommandListImmediate& RHICmdList;
+
+public:
+	FRHICommandListScopedFlushAndExecute(FRHICommandListImmediate& InRHICmdList)
+		: RHICmdList(InRHICmdList)
+	{
+		check(RHICmdList.IsTopOfPipe());
+		RHICmdList.ImmediateFlush(EImmediateFlushType::FlushRHIThread);
+		RHICmdList.bExecuting = true;
+	}
+	~FRHICommandListScopedFlushAndExecute()
+	{
+		RHICmdList.bExecuting = false;
+	}
+};
+
+struct FScopedGPUMask
 {
 	FRHICommandListImmediate& RHICmdList;
 	FRHIGPUMask PrevGPUMask;
diff --git a/Engine/Source/Runtime/RHI/Public/RHIValidation.h b/Engine/Source/Runtime/RHI/Public/RHIValidation.h
--- a/Engine/Source/Runtime/RHI/Public/RHIValidation.h
+++ b/Engine/Source/Runtime/RHI/Public/RHIValidation.h
@@ -343,15 +343,15 @@ public:
 	}
 
 	// FlushType: Flush RHI Thread
-	virtual void* RHILockIndexBuffer(FRHIIndexBuffer* IndexBuffer, uint32 Offset, uint32 Size, EResourceLockMode LockMode) override final
+	virtual void* RHILockIndexBuffer(class FRHICommandListImmediate& RHICmdList, FRHIIndexBuffer* IndexBuffer, uint32 Offset, uint32 SizeRHI, EResourceLockMode LockMode) override final
 	{
-		return RHI->RHILockIndexBuffer(IndexBuffer, Offset, Size, LockMode);
+		return RHI->RHILockIndexBuffer(RHICmdList, IndexBuffer, Offset, SizeRHI, LockMode);
 	}
 
 	// FlushType: Flush RHI Thread
-	virtual void RHIUnlockIndexBuffer(FRHIIndexBuffer* IndexBuffer) override final
+	virtual void RHIUnlockIndexBuffer(class FRHICommandListImmediate& RHICmdList, FRHIIndexBuffer* IndexBuffer) override final
 	{
-		RHI->RHIUnlockIndexBuffer(IndexBuffer);
+		RHI->RHIUnlockIndexBuffer(RHICmdList, IndexBuffer);
 	}
 
 	/**
@@ -364,15 +364,15 @@ public:
 	}
 
 	// FlushType: Flush RHI Thread
-	virtual void* RHILockVertexBuffer(FRHIVertexBuffer* VertexBuffer, uint32 Offset, uint32 SizeRHI, EResourceLockMode LockMode) override final
+	virtual void* RHILockVertexBuffer(class FRHICommandListImmediate& RHICmdList, FRHIVertexBuffer* VertexBuffer, uint32 Offset, uint32 SizeRHI, EResourceLockMode LockMode) override final
 	{
-		return RHI->RHILockVertexBuffer(VertexBuffer, Offset, SizeRHI, LockMode);
+		return RHI->RHILockVertexBuffer(RHICmdList, VertexBuffer, Offset, SizeRHI, LockMode);
 	}
 
 	// FlushType: Flush RHI Thread
-	virtual void RHIUnlockVertexBuffer(FRHIVertexBuffer* VertexBuffer) override final
+	virtual void RHIUnlockVertexBuffer(class FRHICommandListImmediate& RHICmdList, FRHIVertexBuffer* VertexBuffer) override final
 	{
-		RHI->RHIUnlockVertexBuffer(VertexBuffer);
+		RHI->RHIUnlockVertexBuffer(RHICmdList, VertexBuffer);
 	}
 
 	/** Copies the contents of one vertex buffer to another vertex buffer.  They must have identical sizes. */
@@ -393,17 +393,17 @@ public:
 	}
 
 	// FlushType: Flush RHI Thread
-	virtual void* RHILockStructuredBuffer(FRHIStructuredBuffer* StructuredBuffer, uint32 Offset, uint32 SizeRHI, EResourceLockMode LockMode) override final
+	virtual void* RHILockStructuredBuffer(class FRHICommandListImmediate& RHICmdList, FRHIStructuredBuffer* StructuredBuffer, uint32 Offset, uint32 SizeRHI, EResourceLockMode LockMode) override final
 	{
 		check(IsFeatureLevelSupported(GMaxRHIShaderPlatform, ERHIFeatureLevel::SM5));
-		return RHI->RHILockStructuredBuffer(StructuredBuffer, Offset, SizeRHI, LockMode);
+		return RHI->RHILockStructuredBuffer(RHICmdList, StructuredBuffer, Offset, SizeRHI, LockMode);
 	}
 
 	// FlushType: Flush RHI Thread
-	virtual void RHIUnlockStructuredBuffer(FRHIStructuredBuffer* StructuredBuffer) override final
+	virtual void RHIUnlockStructuredBuffer(class FRHICommandListImmediate& RHICmdList, FRHIStructuredBuffer* StructuredBuffer) override final
 	{
 		check(IsFeatureLevelSupported(GMaxRHIShaderPlatform, ERHIFeatureLevel::SM5));
-		RHI->RHIUnlockStructuredBuffer(StructuredBuffer);
+		RHI->RHIUnlockStructuredBuffer(RHICmdList, StructuredBuffer);
 	}
 
 	/** Creates an unordered access view of the given structured buffer. */
@@ -1237,14 +1237,14 @@ public:
 		return RHI->CreateShaderResourceView_RenderThread(RHICmdList, Buffer);
 	}
 
-	virtual void* LockVertexBuffer_RenderThread(class FRHICommandListImmediate& RHICmdList, FRHIVertexBuffer* VertexBuffer, uint32 Offset, uint32 SizeRHI, EResourceLockMode LockMode) override final
+	virtual void* LockVertexBuffer_BottomOfPipe(class FRHICommandListImmediate& RHICmdList, FRHIVertexBuffer* VertexBuffer, uint32 Offset, uint32 SizeRHI, EResourceLockMode LockMode) override final
 	{
-		return RHI->LockVertexBuffer_RenderThread(RHICmdList, VertexBuffer, Offset, SizeRHI, LockMode);
+		return RHI->LockVertexBuffer_BottomOfPipe(RHICmdList, VertexBuffer, Offset, SizeRHI, LockMode);
 	}
 
-	virtual void UnlockVertexBuffer_RenderThread(class FRHICommandListImmediate& RHICmdList, FRHIVertexBuffer* VertexBuffer) override final
+	virtual void UnlockVertexBuffer_BottomOfPipe(class FRHICommandListImmediate& RHICmdList, FRHIVertexBuffer* VertexBuffer) override final
 	{
-		return RHI->UnlockVertexBuffer_RenderThread(RHICmdList, VertexBuffer);
+		return RHI->UnlockVertexBuffer_BottomOfPipe(RHICmdList, VertexBuffer);
 	}
 
 	virtual FTexture2DRHIRef AsyncReallocateTexture2D_RenderThread(class FRHICommandListImmediate& RHICmdList, FRHITexture2D* Texture2D, int32 NewMipCount, int32 NewSizeX, int32 NewSizeY, FThreadSafeCounter* RequestStatus) override final
@@ -1267,14 +1267,24 @@ public:
 		return RHI->CreateIndexBuffer_RenderThread(RHICmdList, Stride, Size, InUsage, CreateInfo);
 	}
 
-	virtual void* LockIndexBuffer_RenderThread(class FRHICommandListImmediate& RHICmdList, FRHIIndexBuffer* IndexBuffer, uint32 Offset, uint32 SizeRHI, EResourceLockMode LockMode) override final
+	virtual void* LockIndexBuffer_BottomOfPipe(class FRHICommandListImmediate& RHICmdList, FRHIIndexBuffer* IndexBuffer, uint32 Offset, uint32 SizeRHI, EResourceLockMode LockMode) override final
 	{
-		return RHI->LockIndexBuffer_RenderThread(RHICmdList, IndexBuffer, Offset, SizeRHI, LockMode);
+		return RHI->LockIndexBuffer_BottomOfPipe(RHICmdList, IndexBuffer, Offset, SizeRHI, LockMode);
 	}
 
-	virtual void UnlockIndexBuffer_RenderThread(class FRHICommandListImmediate& RHICmdList, FRHIIndexBuffer* IndexBuffer) override final
+	virtual void UnlockIndexBuffer_BottomOfPipe(class FRHICommandListImmediate& RHICmdList, FRHIIndexBuffer* IndexBuffer) override final
 	{
-		RHI->UnlockIndexBuffer_RenderThread(RHICmdList, IndexBuffer);
+		RHI->UnlockIndexBuffer_BottomOfPipe(RHICmdList, IndexBuffer);
+	}
+
+	virtual void* LockStructuredBuffer_BottomOfPipe(class FRHICommandListImmediate& RHICmdList, FRHIStructuredBuffer* StructuredBuffer, uint32 Offset, uint32 SizeRHI, EResourceLockMode LockMode) override final
+	{
+		return RHI->LockStructuredBuffer_BottomOfPipe(RHICmdList, StructuredBuffer, Offset, SizeRHI, LockMode);
+	}
+	
+	virtual void UnlockStructuredBuffer_BottomOfPipe(class FRHICommandListImmediate& RHICmdList, FRHIStructuredBuffer* StructuredBuffer) override final
+	{
+		RHI->UnlockStructuredBuffer_BottomOfPipe(RHICmdList, StructuredBuffer);
 	}
 
 	virtual FVertexShaderRHIRef CreateVertexShader_RenderThread(class FRHICommandListImmediate& RHICmdList, const TArray<uint8>& Code) override final
diff --git a/Engine/Source/Runtime/Renderer/Private/MeshDrawCommands.cpp b/Engine/Source/Runtime/Renderer/Private/MeshDrawCommands.cpp
--- a/Engine/Source/Runtime/Renderer/Private/MeshDrawCommands.cpp
+++ b/Engine/Source/Runtime/Renderer/Private/MeshDrawCommands.cpp
@@ -1128,38 +1128,6 @@ public:
 	}
 };
 
-/*
- * Copies provided vertex data (assumed to be on MemStack) to a vertex buffer.
- */
-struct FRHICommandUpdatePrimitiveIdBuffer : public FRHICommand<FRHICommandUpdatePrimitiveIdBuffer>
-{
-	FVertexBufferRHIParamRef VertexBuffer;
-	void* VertexBufferData;
-	int32 VertexBufferDataSize;
-
-	virtual ~FRHICommandUpdatePrimitiveIdBuffer() {}
-
-	FORCEINLINE_DEBUGGABLE FRHICommandUpdatePrimitiveIdBuffer(
-		FVertexBufferRHIParamRef InVertexBuffer,
-		void* InVertexBufferData,
-		int32 InVertexBufferDataSize)
-		: VertexBuffer(InVertexBuffer)
-		, VertexBufferData(InVertexBufferData)
-		, VertexBufferDataSize(InVertexBufferDataSize)
-	{
-	}
-
-	void Execute(FRHICommandListBase& CmdList)
-	{
-		// Upload vertex buffer data.
-		void* RESTRICT Data = (void* RESTRICT)GDynamicRHI->RHILockVertexBuffer(VertexBuffer, 0, VertexBufferDataSize, RLM_WriteOnly);
-		FMemory::Memcpy(Data, VertexBufferData, VertexBufferDataSize);
-		GDynamicRHI->RHIUnlockVertexBuffer(VertexBuffer);
-
-		FMemory::Free(VertexBufferData);
-	}
-};
-
 void FParallelMeshDrawCommandPass::DispatchDraw(FParallelCommandListSet* ParallelCommandListSet, FRHICommandList& RHICmdList) const
 {
 	if (MaxNumDraws <= 0)
@@ -1182,11 +1150,18 @@ void FParallelMeshDrawCommandPass::DispatchDraw(FParallelCommandListSet* Paralle
 				RHICommandList.AddDispatchPrerequisite(TaskEventRef);
 			}
 
-			new (RHICommandList.AllocCommand<FRHICommandUpdatePrimitiveIdBuffer>())FRHICommandUpdatePrimitiveIdBuffer(
-				PrimitiveIdVertexBufferRHI,
-				TaskContext.PrimitiveIdBufferData,
-				TaskContext.PrimitiveIdBufferDataSize
-			);
+			RHICommandList.EnqueueLambda([
+				VertexBuffer = PrimitiveIdsBuffer,
+				VertexBufferData = TaskContext.PrimitiveIdBufferData, 
+				VertexBufferDataSize = TaskContext.PrimitiveIdBufferDataSize](FRHICommandListImmediate& CmdList)
+			{
+				// Upload vertex buffer data.
+				void* RESTRICT Data = (void* RESTRICT)CmdList.LockVertexBuffer(VertexBuffer, 0, VertexBufferDataSize, RLM_WriteOnly);
+				FMemory::Memcpy(Data, VertexBufferData, VertexBufferDataSize);
+				CmdList.UnlockVertexBuffer(VertexBuffer);
+
+				FMemory::Free(VertexBufferData);
+			});
 
 			RHICommandList.RHIThreadFence(true);
 
diff --git a/Engine/Source/Runtime/Renderer/Private/PostProcess/SceneRenderTargets.cpp b/Engine/Source/Runtime/Renderer/Private/PostProcess/SceneRenderTargets.cpp
--- a/Engine/Source/Runtime/Renderer/Private/PostProcess/SceneRenderTargets.cpp
+++ b/Engine/Source/Runtime/Renderer/Private/PostProcess/SceneRenderTargets.cpp
@@ -1828,7 +1828,7 @@ void FSceneRenderTargets::ResolveDepthTexture(FRHICommandList& RHICmdList, const
 		break;
 	}
 
-	GraphicsPSOInit.BoundShaderState.VertexDeclarationRHI = GScreenVertexDeclaration.VertexDeclarationRHI;
+	GraphicsPSOInit.BoundShaderState.VertexDeclarationRHI = GResolveVertexBuffer.VertexDeclarationRHI;
 	GraphicsPSOInit.BoundShaderState.VertexShaderRHI = GETSAFERHISHADER_VERTEX(*ResolveVertexShader);
 	GraphicsPSOInit.BoundShaderState.PixelShaderRHI = ResolvePixelShader;
 	GraphicsPSOInit.PrimitiveType = PT_TriangleStrip;
@@ -1842,35 +1842,9 @@ void FSceneRenderTargets::ResolveDepthTexture(FRHICommandList& RHICmdList, const
 		RHICmdList.SetShaderTexture(ResolvePixelShader, TextureIndex, SourceTexture);
 	}
 
-	FRHIResourceCreateInfo CreateInfo;
-	FVertexBufferRHIRef VertexBufferRHI = RHICreateVertexBuffer(sizeof(FScreenVertex) * 4, BUF_Volatile, CreateInfo);
-	void* VoidPtr = RHILockVertexBuffer(VertexBufferRHI, 0, sizeof(FScreenVertex) * 4, RLM_WriteOnly);
-
-	// Generate the vertices used
-	FScreenVertex* Vertices = (FScreenVertex*)VoidPtr;
-
-	Vertices[0].Position.X = MaxX;
-	Vertices[0].Position.Y = MinY;
-	Vertices[0].UV.X = MaxU;
-	Vertices[0].UV.Y = MinV;
-
-	Vertices[1].Position.X = MaxX;
-	Vertices[1].Position.Y = MaxY;
-	Vertices[1].UV.X = MaxU;
-	Vertices[1].UV.Y = MaxV;
-
-	Vertices[2].Position.X = MinX;
-	Vertices[2].Position.Y = MinY;
-	Vertices[2].UV.X = MinU;
-	Vertices[2].UV.Y = MinV;
-
-	Vertices[3].Position.X = MinX;
-	Vertices[3].Position.Y = MaxY;
-	Vertices[3].UV.X = MinU;
-	Vertices[3].UV.Y = MaxV;
+	ResolveVertexShader->SetParameters(RHICmdList, FVector4(MinX, MinY, MaxX, MaxY), FVector4(MinU, MinV, MaxU, MaxV));
 
-	RHIUnlockVertexBuffer(VertexBufferRHI);
-	RHICmdList.SetStreamSource(0, VertexBufferRHI, 0);
+	RHICmdList.SetStreamSource(0, GResolveVertexBuffer.VertexBufferRHI, 0);
 	RHICmdList.DrawPrimitive(0, 2, 1);
 }
 	RHICmdList.EndRenderPass();
diff --git a/Engine/Source/Runtime/SlateRHIRenderer/Private/SlateRHIRenderingPolicy.cpp b/Engine/Source/Runtime/SlateRHIRenderer/Private/SlateRHIRenderingPolicy.cpp
--- a/Engine/Source/Runtime/SlateRHIRenderer/Private/SlateRHIRenderingPolicy.cpp
+++ b/Engine/Source/Runtime/SlateRHIRenderer/Private/SlateRHIRenderingPolicy.cpp
@@ -130,53 +130,6 @@ void FSlateRHIRenderingPolicy::EndDrawingWindows()
 	check( IsInParallelRenderingThread() );
 }
 
-struct FSlateUpdateVertexAndIndexBuffers final : public FRHICommand<FSlateUpdateVertexAndIndexBuffers>
-{
-	FVertexBufferRHIRef VertexBufferRHI;
-	FIndexBufferRHIRef IndexBufferRHI;
-	FSlateBatchData& BatchData;
-	bool bAbsoluteIndices;
-
-	FSlateUpdateVertexAndIndexBuffers(TSlateElementVertexBuffer<FSlateVertex>& InVertexBuffer, FSlateElementIndexBuffer& InIndexBuffer, FSlateBatchData& InBatchData, bool bInAbsoluteIndices)
-		: VertexBufferRHI(InVertexBuffer.VertexBufferRHI)
-		, IndexBufferRHI(InIndexBuffer.IndexBufferRHI)
-		, BatchData(InBatchData)
-		, bAbsoluteIndices(bInAbsoluteIndices)
-	{
-		check(IsInRenderingThread());
-	}
-
-	void Execute(FRHICommandListBase& CmdList)
-	{
-		SCOPE_CYCLE_COUNTER(STAT_SlateUpdateBufferRTTime);
-
-		const int32 NumBatchedVertices = BatchData.GetNumBatchedVertices();
-		const int32 NumBatchedIndices = BatchData.GetNumBatchedIndices();
-
-		int32 RequiredVertexBufferSize = NumBatchedVertices*sizeof(FSlateVertex);
-		uint8* VertexBufferData = (uint8*)GDynamicRHI->RHILockVertexBuffer( VertexBufferRHI, 0, RequiredVertexBufferSize, RLM_WriteOnly );
-
-		uint32 RequiredIndexBufferSize = NumBatchedIndices*sizeof(SlateIndex);		
-		uint8* IndexBufferData = (uint8*)GDynamicRHI->RHILockIndexBuffer( IndexBufferRHI, 0, RequiredIndexBufferSize, RLM_WriteOnly );
-
-		//Early out if we have an invalid buffer (might have lost context and now have invalid buffers)
-		if ((nullptr != VertexBufferData) && (nullptr != IndexBufferData))
-		{
-			BatchData.FillVertexAndIndexBuffer(VertexBufferData, IndexBufferData, bAbsoluteIndices);
-		}
-
-		if (nullptr != VertexBufferData)
-		{
-			GDynamicRHI->RHIUnlockVertexBuffer(VertexBufferRHI);
-		}
-
-		if (nullptr != IndexBufferData)
-		{
-			GDynamicRHI->RHIUnlockIndexBuffer(IndexBufferRHI);
-		}
-	}
-};
-
 void FSlateRHIRenderingPolicy::UpdateVertexAndIndexBuffers(FRHICommandListImmediate& RHICmdList, FSlateBatchData& InBatchData)
 {
 	UpdateVertexAndIndexBuffers(RHICmdList, InBatchData, VertexBuffers, IndexBuffers);
@@ -238,7 +191,29 @@ void FSlateRHIRenderingPolicy::UpdateVertexAndIndexBuffers(FRHICommandListImmedi
 		}
 		else
 		{
-			ALLOC_COMMAND_CL(RHICmdList, FSlateUpdateVertexAndIndexBuffers)(VertexBuffer, IndexBuffer, InBatchData, bAbsoluteIndices);
+			FVertexBufferRHIParamRef VertexBufferRHI = VertexBuffer.VertexBufferRHI;
+			FIndexBufferRHIParamRef IndexBufferRHI = IndexBuffer.IndexBufferRHI;
+
+			RHICmdList.EnqueueLambda([VertexBufferRHI, IndexBufferRHI, &InBatchData, bAbsoluteIndices](FRHICommandListImmediate& RHICmdList)
+			{
+				SCOPE_CYCLE_COUNTER(STAT_SlateUpdateBufferRTTime);
+
+				const int32 NumBatchedVertices = InBatchData.GetNumBatchedVertices();
+				const int32 NumBatchedIndices = InBatchData.GetNumBatchedIndices();
+
+				int32 RequiredVertexBufferSize = NumBatchedVertices * sizeof(FSlateVertex);
+				uint8* VertexBufferData = (uint8*)RHICmdList.LockVertexBuffer(VertexBufferRHI, 0, RequiredVertexBufferSize, RLM_WriteOnly);
+
+				uint32 RequiredIndexBufferSize = NumBatchedIndices * sizeof(SlateIndex);
+				uint8* IndexBufferData = (uint8*)RHICmdList.LockIndexBuffer(IndexBufferRHI, 0, RequiredIndexBufferSize, RLM_WriteOnly);
+
+				InBatchData.FillVertexAndIndexBuffer(VertexBufferData, IndexBufferData, bAbsoluteIndices);
+
+				RHICmdList.UnlockVertexBuffer(VertexBufferRHI);
+				RHICmdList.UnlockIndexBuffer(IndexBufferRHI);
+			});
+
+			RHICmdList.RHIThreadFence(true);
 		}
 	}
 
diff --git a/Engine/Source/Runtime/SlateRHIRenderer/Private/SlateUpdatableBuffer.cpp b/Engine/Source/Runtime/SlateRHIRenderer/Private/SlateUpdatableBuffer.cpp
--- a/Engine/Source/Runtime/SlateRHIRenderer/Private/SlateUpdatableBuffer.cpp
+++ b/Engine/Source/Runtime/SlateRHIRenderer/Private/SlateUpdatableBuffer.cpp
@@ -5,29 +5,6 @@
 
 DECLARE_CYCLE_STAT(TEXT("UpdateInstanceBuffer Time"), STAT_SlateUpdateInstanceBuffer, STATGROUP_Slate);
 
-struct FSlateUpdateInstanceBufferCommand final : public FRHICommand<FSlateUpdateInstanceBufferCommand>
-{
-	FVertexBufferRHIRef VertexBufferRHI;
-	const TArray<FVector4>& InstanceData;
-
-	FSlateUpdateInstanceBufferCommand(TSlateElementVertexBuffer<FVector4>& InInstanceBuffer, const TArray<FVector4>& InInstanceData )
-		: VertexBufferRHI(InInstanceBuffer.VertexBufferRHI)
-		, InstanceData(InInstanceData)
-	{}
-
-	void Execute(FRHICommandListBase& CmdList)
-	{
-		SCOPE_CYCLE_COUNTER( STAT_SlateUpdateInstanceBuffer );
-
-		int32 RequiredVertexBufferSize = InstanceData.Num()*sizeof(FVector4);
-		uint8* InstanceBufferData = (uint8*)GDynamicRHI->RHILockVertexBuffer( VertexBufferRHI, 0, RequiredVertexBufferSize, RLM_WriteOnly );
-
-		FMemory::Memcpy( InstanceBufferData, InstanceData.GetData(), InstanceData.Num()*sizeof(FVector4) );
-	
-		GDynamicRHI->RHIUnlockVertexBuffer( VertexBufferRHI );
-	}
-};
-
 FSlateUpdatableInstanceBuffer::FSlateUpdatableInstanceBuffer( int32 InitialInstanceCount )
 	: FreeBufferIndex(0)
 {
@@ -100,7 +77,20 @@ void FSlateUpdatableInstanceBuffer::UpdateRenderingData_RenderThread(FRHICommand
 	}
 	else
 	{
-		ALLOC_COMMAND_CL(RHICmdList, FSlateUpdateInstanceBufferCommand)(InstanceBufferResource, RenderThreadBufferData);
+		FVertexBufferRHIParamRef VertexBufferRHI = InstanceBufferResource.VertexBufferRHI;
+		RHICmdList.EnqueueLambda([VertexBufferRHI, &InstanceData = RenderThreadBufferData](FRHICommandListImmediate& RHICmdList)
+		{
+			SCOPE_CYCLE_COUNTER(STAT_SlateUpdateInstanceBuffer);
+
+			int32 RequiredVertexBufferSize = InstanceData.Num() * sizeof(FVector4);
+			uint8* InstanceBufferData = (uint8*)RHICmdList.LockVertexBuffer(VertexBufferRHI, 0, RequiredVertexBufferSize, RLM_WriteOnly);
+
+			FMemory::Memcpy(InstanceBufferData, InstanceData.GetData(), InstanceData.Num() * sizeof(FVector4));
+
+			RHICmdList.UnlockVertexBuffer(VertexBufferRHI);
+		});
+
+		RHICmdList.RHIThreadFence(true);
 	}
 }
 
diff --git a/Engine/Source/Runtime/UtilityShaders/Private/ResolveShader.cpp b/Engine/Source/Runtime/UtilityShaders/Private/ResolveShader.cpp
--- a/Engine/Source/Runtime/UtilityShaders/Private/ResolveShader.cpp
+++ b/Engine/Source/Runtime/UtilityShaders/Private/ResolveShader.cpp
@@ -12,7 +12,39 @@ IMPLEMENT_SHADER_TYPE(, FResolveDepthNonMSPS, TEXT("/Engine/Private/ResolvePixel
 IMPLEMENT_SHADER_TYPE(, FResolveSingleSamplePS, TEXT("/Engine/Private/ResolvePixelShader.usf"), TEXT("MainSingleSample"), SF_Pixel);
 IMPLEMENT_SHADER_TYPE(, FResolveVS, TEXT("/Engine/Private/ResolveVertexShader.usf"), TEXT("Main"), SF_Vertex);
 
+UTILITYSHADERS_API TGlobalResource<FResolveVertexBuffer> GResolveVertexBuffer;
+
 void FResolveSingleSamplePS::SetParameters(FRHICommandList& RHICmdList, uint32 SingleSampleIndexValue)
 {
 	SetShaderValue(RHICmdList, GetPixelShader(),SingleSampleIndex,SingleSampleIndexValue);
 }
+
+void FResolveVS::SetParameters(FRHICommandList& RHICmdList, FVector4 InPositionMinMax, FVector4 InUVMinMax)
+{
+	SetShaderValue(RHICmdList, GetVertexShader(), PositionMinMax, InPositionMinMax);
+	SetShaderValue(RHICmdList, GetVertexShader(), UVMinMax, InUVMinMax);
+}
+
+void FResolveVertexBuffer::InitDynamicRHI()
+{
+	FVertexDeclarationElementList Elements;
+	Elements.Add(FVertexElement(0, 0, VET_Float2, 0, sizeof(FVector2D)));
+	VertexDeclarationRHI = RHICreateVertexDeclaration(Elements);
+
+	void* VoidPtr;
+	FRHIResourceCreateInfo CreateInfo;
+	VertexBufferRHI = RHICreateAndLockVertexBuffer(sizeof(FVector2D) * 4, BUF_Static, CreateInfo, VoidPtr);
+	// Generate the vertices used
+	FVector2D* Vertices = (FVector2D*)VoidPtr;
+	Vertices[0] = { 0.0f, 0.0f };
+	Vertices[1] = { 1.0f, 0.0f };
+	Vertices[2] = { 0.0f, 1.0f };
+	Vertices[3] = { 1.0f, 1.0f };
+	RHIUnlockVertexBuffer(VertexBufferRHI);
+}
+
+void FResolveVertexBuffer::ReleaseDynamicRHI()
+{
+	VertexBufferRHI.SafeRelease();
+	VertexDeclarationRHI.SafeRelease();
+}
diff --git a/Engine/Source/Runtime/UtilityShaders/Public/ResolveShader.h b/Engine/Source/Runtime/UtilityShaders/Public/ResolveShader.h
--- a/Engine/Source/Runtime/UtilityShaders/Public/ResolveShader.h
+++ b/Engine/Source/Runtime/UtilityShaders/Public/ResolveShader.h
@@ -249,8 +249,35 @@ public:
 	
 	static bool ShouldCompilePermutation(const FGlobalShaderPermutationParameters& Parameters) { return true; }
 	
-	FResolveVS(const ShaderMetaType::CompiledShaderInitializerType& Initializer):
-	FGlobalShader(Initializer)
-	{}
+	FResolveVS(const ShaderMetaType::CompiledShaderInitializerType& Initializer)
+		: FGlobalShader(Initializer)
+	{
+		PositionMinMax.Bind(Initializer.ParameterMap, TEXT("PositionMinMax"), SPF_Mandatory);
+		UVMinMax.Bind(Initializer.ParameterMap, TEXT("UVMinMax"), SPF_Mandatory);
+	}
 	FResolveVS() {}
+
+	UTILITYSHADERS_API void SetParameters(FRHICommandList& RHICmdList, FVector4 PositionMinMax, FVector4 UVMinMax);
+
+	virtual bool Serialize(FArchive& Ar) override
+	{
+		bool bShaderHasOutdatedParameters = FGlobalShader::Serialize(Ar);
+		Ar << PositionMinMax;
+		Ar << UVMinMax;
+		return bShaderHasOutdatedParameters;
+	}
+
+	FShaderParameter PositionMinMax;
+	FShaderParameter UVMinMax;
+};
+
+struct FResolveVertexBuffer : public FRenderResource
+{
+	FVertexBufferRHIRef VertexBufferRHI;
+	FVertexDeclarationRHIRef VertexDeclarationRHI;
+
+	virtual void InitDynamicRHI() override;
+	virtual void ReleaseDynamicRHI() override;
 };
+
+extern UTILITYSHADERS_API TGlobalResource<FResolveVertexBuffer> GResolveVertexBuffer;
diff --git a/Engine/Source/Runtime/VulkanRHI/Private/VulkanIndexBuffer.cpp b/Engine/Source/Runtime/VulkanRHI/Private/VulkanIndexBuffer.cpp
--- a/Engine/Source/Runtime/VulkanRHI/Private/VulkanIndexBuffer.cpp
+++ b/Engine/Source/Runtime/VulkanRHI/Private/VulkanIndexBuffer.cpp
@@ -364,7 +364,7 @@ FIndexBufferRHIRef FVulkanDynamicRHI::RHICreateIndexBuffer(uint32 Stride, uint32
 	return new FVulkanIndexBuffer(Device, Stride, Size, InUsage, CreateInfo, nullptr);
 }
 
-void* FVulkanDynamicRHI::RHILockIndexBuffer(FIndexBufferRHIParamRef IndexBufferRHI, uint32 Offset, uint32 Size, EResourceLockMode LockMode)
+void* FVulkanDynamicRHI::LockIndexBuffer_BottomOfPipe(FRHICommandListImmediate& RHICmdList, FIndexBufferRHIParamRef IndexBufferRHI, uint32 Offset, uint32 Size, EResourceLockMode LockMode)
 {
 	FVulkanIndexBuffer* IndexBuffer = ResourceCast(IndexBufferRHI);
 	return IndexBuffer->Lock(false, LockMode, Size, Offset);
@@ -377,7 +377,7 @@ void* FVulkanDynamicRHI::LockIndexBuffer_RenderThread(class FRHICommandListImmed
 }
 #endif
 
-void FVulkanDynamicRHI::RHIUnlockIndexBuffer(FIndexBufferRHIParamRef IndexBufferRHI)
+void FVulkanDynamicRHI::UnlockIndexBuffer_BottomOfPipe(FRHICommandListImmediate& RHICmdList, FIndexBufferRHIParamRef IndexBufferRHI)
 {
 	FVulkanIndexBuffer* IndexBuffer = ResourceCast(IndexBufferRHI);
 	IndexBuffer->Unlock(false);
diff --git a/Engine/Source/Runtime/VulkanRHI/Private/VulkanStructuredBuffer.cpp b/Engine/Source/Runtime/VulkanRHI/Private/VulkanStructuredBuffer.cpp
--- a/Engine/Source/Runtime/VulkanRHI/Private/VulkanStructuredBuffer.cpp
+++ b/Engine/Source/Runtime/VulkanRHI/Private/VulkanStructuredBuffer.cpp
@@ -20,14 +20,14 @@ FStructuredBufferRHIRef FVulkanDynamicRHI::RHICreateStructuredBuffer(uint32 InSt
 	return new FVulkanStructuredBuffer(Device, InStride, InSize, CreateInfo, InUsage);
 }
 
-void* FVulkanDynamicRHI::RHILockStructuredBuffer(FStructuredBufferRHIParamRef StructuredBufferRHI,uint32 Offset,uint32 Size,EResourceLockMode LockMode)
+void* FVulkanDynamicRHI::LockStructuredBuffer_BottomOfPipe(FRHICommandListImmediate& RHICmdList, FStructuredBufferRHIParamRef StructuredBufferRHI,uint32 Offset,uint32 Size,EResourceLockMode LockMode)
 {
 	FVulkanStructuredBuffer* StructuredBuffer = ResourceCast(StructuredBufferRHI);
 
 	return StructuredBuffer->Lock(false, LockMode, Size, Offset);
 }
 
-void FVulkanDynamicRHI::RHIUnlockStructuredBuffer(FStructuredBufferRHIParamRef StructuredBufferRHI)
+void FVulkanDynamicRHI::UnlockStructuredBuffer_BottomOfPipe(FRHICommandListImmediate& RHICmdList, FStructuredBufferRHIParamRef StructuredBufferRHI)
 {
 	FVulkanStructuredBuffer* StructuredBuffer = ResourceCast(StructuredBufferRHI);
 
diff --git a/Engine/Source/Runtime/VulkanRHI/Private/VulkanVertexBuffer.cpp b/Engine/Source/Runtime/VulkanRHI/Private/VulkanVertexBuffer.cpp
--- a/Engine/Source/Runtime/VulkanRHI/Private/VulkanVertexBuffer.cpp
+++ b/Engine/Source/Runtime/VulkanRHI/Private/VulkanVertexBuffer.cpp
@@ -30,7 +30,7 @@ FVertexBufferRHIRef FVulkanDynamicRHI::RHICreateVertexBuffer(uint32 Size, uint32
 	return VertexBuffer;
 }
 
-void* FVulkanDynamicRHI::RHILockVertexBuffer(FVertexBufferRHIParamRef VertexBufferRHI, uint32 Offset, uint32 Size, EResourceLockMode LockMode)
+void* FVulkanDynamicRHI::LockVertexBuffer_BottomOfPipe(FRHICommandListImmediate& RHICmdList, FVertexBufferRHIParamRef VertexBufferRHI, uint32 Offset, uint32 Size, EResourceLockMode LockMode)
 {
 	LLM_SCOPE_VULKAN(ELLMTagVulkan::VulkanVertexBuffers);
 	FVulkanVertexBuffer* VertexBuffer = ResourceCast(VertexBufferRHI);
@@ -44,7 +44,7 @@ void* FVulkanDynamicRHI::LockVertexBuffer_RenderThread(class FRHICommandListImme
 }
 #endif
 
-void FVulkanDynamicRHI::RHIUnlockVertexBuffer(FVertexBufferRHIParamRef VertexBufferRHI)
+void FVulkanDynamicRHI::UnlockVertexBuffer_BottomOfPipe(FRHICommandListImmediate& RHICmdList, FVertexBufferRHIParamRef VertexBufferRHI)
 {
 	LLM_SCOPE_VULKAN(ELLMTagVulkan::VulkanVertexBuffers);
 	FVulkanVertexBuffer* VertexBuffer = ResourceCast(VertexBufferRHI);
diff --git a/Engine/Source/Runtime/VulkanRHI/Public/VulkanDynamicRHI.h b/Engine/Source/Runtime/VulkanRHI/Public/VulkanDynamicRHI.h
--- a/Engine/Source/Runtime/VulkanRHI/Public/VulkanDynamicRHI.h
+++ b/Engine/Source/Runtime/VulkanRHI/Public/VulkanDynamicRHI.h
@@ -60,17 +60,17 @@ public:
 	virtual FUniformBufferRHIRef RHICreateUniformBuffer(const void* Contents, const FRHIUniformBufferLayout& Layout, EUniformBufferUsage Usage, EUniformBufferValidation Validation) final override;
 	virtual void RHIUpdateUniformBuffer(FUniformBufferRHIParamRef UniformBufferRHI, const void* Contents) final override;
 	virtual FIndexBufferRHIRef RHICreateIndexBuffer(uint32 Stride, uint32 Size, uint32 InUsage, FRHIResourceCreateInfo& CreateInfo) final override;
-	virtual void* RHILockIndexBuffer(FIndexBufferRHIParamRef IndexBuffer, uint32 Offset, uint32 Size, EResourceLockMode LockMode) final override;
-	virtual void RHIUnlockIndexBuffer(FIndexBufferRHIParamRef IndexBuffer) final override;
+	virtual void* LockIndexBuffer_BottomOfPipe(FRHICommandListImmediate& RHICmdList, FIndexBufferRHIParamRef IndexBuffer, uint32 Offset, uint32 Size, EResourceLockMode LockMode) final override;
+	virtual void UnlockIndexBuffer_BottomOfPipe(FRHICommandListImmediate& RHICmdList, FIndexBufferRHIParamRef IndexBuffer) final override;
 	virtual void RHITransferIndexBufferUnderlyingResource(FIndexBufferRHIParamRef DestIndexBuffer, FIndexBufferRHIParamRef SrcIndexBuffer) final override;
 	virtual FVertexBufferRHIRef RHICreateVertexBuffer(uint32 Size, uint32 InUsage, FRHIResourceCreateInfo& CreateInfo) final override;
-	virtual void* RHILockVertexBuffer(FVertexBufferRHIParamRef VertexBuffer, uint32 Offset, uint32 SizeRHI, EResourceLockMode LockMode) final override;
-	virtual void RHIUnlockVertexBuffer(FVertexBufferRHIParamRef VertexBuffer) final override;
+	virtual void* LockVertexBuffer_BottomOfPipe(FRHICommandListImmediate& RHICmdList, FVertexBufferRHIParamRef VertexBuffer, uint32 Offset, uint32 SizeRHI, EResourceLockMode LockMode) final override;
+	virtual void UnlockVertexBuffer_BottomOfPipe(FRHICommandListImmediate& RHICmdList, FVertexBufferRHIParamRef VertexBuffer) final override;
 	virtual void RHICopyVertexBuffer(FVertexBufferRHIParamRef SourceBuffer, FVertexBufferRHIParamRef DestBuffer) final override;
 	virtual void RHITransferVertexBufferUnderlyingResource(FVertexBufferRHIParamRef DestVertexBuffer, FVertexBufferRHIParamRef SrcVertexBuffer) final override;
 	virtual FStructuredBufferRHIRef RHICreateStructuredBuffer(uint32 Stride, uint32 Size, uint32 InUsage, FRHIResourceCreateInfo& CreateInfo) final override;
-	virtual void* RHILockStructuredBuffer(FStructuredBufferRHIParamRef StructuredBuffer, uint32 Offset, uint32 SizeRHI, EResourceLockMode LockMode) final override;
-	virtual void RHIUnlockStructuredBuffer(FStructuredBufferRHIParamRef StructuredBuffer) final override;
+	virtual void* LockStructuredBuffer_BottomOfPipe(FRHICommandListImmediate& RHICmdList, FStructuredBufferRHIParamRef StructuredBuffer, uint32 Offset, uint32 SizeRHI, EResourceLockMode LockMode) final override;
+	virtual void UnlockStructuredBuffer_BottomOfPipe(FRHICommandListImmediate& RHICmdList, FStructuredBufferRHIParamRef StructuredBuffer) final override;
 	virtual FUnorderedAccessViewRHIRef RHICreateUnorderedAccessView(FStructuredBufferRHIParamRef StructuredBuffer, bool bUseUAVCounter, bool bAppendBuffer) final override;
 	virtual FUnorderedAccessViewRHIRef RHICreateUnorderedAccessView(FTextureRHIParamRef Texture, uint32 MipLevel) final override;
 	virtual FUnorderedAccessViewRHIRef RHICreateUnorderedAccessView(FVertexBufferRHIParamRef VertexBuffer, uint8 Format) final override;
@@ -441,21 +441,6 @@ protected:
 	void RemoveDebugLayerCallback();
 #endif
 
-	inline bool ShouldDeferBufferLockOperation(FRHICommandList* RHICmdList)
-	{
-		if (RHICmdList == nullptr)
-		{
-			return false;
-		}
-
-		if (RHICmdList->Bypass() || !IsRunningRHIInSeparateThread())
-		{
-			return false;
-		}
-
-		return true;
-	}
-
 	FCriticalSection LockBufferCS;
 
 	void InternalUnlockTexture2D(bool bFromRenderingThread, FTexture2DRHIParamRef Texture, uint32 MipIndex, bool bLockWithinMiptail);
diff --git a/Engine/Source/Runtime/Windows/D3D11RHI/Private/D3D11IndexBuffer.cpp b/Engine/Source/Runtime/Windows/D3D11RHI/Private/D3D11IndexBuffer.cpp
--- a/Engine/Source/Runtime/Windows/D3D11RHI/Private/D3D11IndexBuffer.cpp
+++ b/Engine/Source/Runtime/Windows/D3D11RHI/Private/D3D11IndexBuffer.cpp
@@ -76,7 +76,7 @@ FIndexBufferRHIRef FD3D11DynamicRHI::CreateIndexBuffer_RenderThread(
 	return RHICreateIndexBuffer(Stride, Size, InUsage, CreateInfo);
 }
 
-void* FD3D11DynamicRHI::RHILockIndexBuffer(FIndexBufferRHIParamRef IndexBufferRHI,uint32 Offset,uint32 Size,EResourceLockMode LockMode)
+void* FD3D11DynamicRHI::LockIndexBuffer_BottomOfPipe(FRHICommandListImmediate& RHICmdList, FIndexBufferRHIParamRef IndexBufferRHI, uint32 Offset, uint32 Size, EResourceLockMode LockMode)
 {
 	FD3D11IndexBuffer* IndexBuffer = ResourceCast(IndexBufferRHI);
 	// If this resource is bound to the device, unbind it
@@ -140,7 +140,7 @@ void* FD3D11DynamicRHI::RHILockIndexBuffer(FIndexBufferRHIParamRef IndexBufferRH
 	return (void*)((uint8*)LockedData.GetData() + Offset);
 }
 
-void FD3D11DynamicRHI::RHIUnlockIndexBuffer(FIndexBufferRHIParamRef IndexBufferRHI)
+void FD3D11DynamicRHI::UnlockIndexBuffer_BottomOfPipe(FRHICommandListImmediate& RHICmdList, FIndexBufferRHIParamRef IndexBufferRHI)
 {
 	FD3D11IndexBuffer* IndexBuffer = ResourceCast(IndexBufferRHI);
 
diff --git a/Engine/Source/Runtime/Windows/D3D11RHI/Private/D3D11RHIPrivate.h b/Engine/Source/Runtime/Windows/D3D11RHI/Private/D3D11RHIPrivate.h
--- a/Engine/Source/Runtime/Windows/D3D11RHI/Private/D3D11RHIPrivate.h
+++ b/Engine/Source/Runtime/Windows/D3D11RHI/Private/D3D11RHIPrivate.h
@@ -416,17 +416,17 @@ public:
 	virtual FUniformBufferRHIRef RHICreateUniformBuffer(const void* Contents, const FRHIUniformBufferLayout& Layout, EUniformBufferUsage Usage, EUniformBufferValidation Validation) final override;
 	virtual void RHIUpdateUniformBuffer(FUniformBufferRHIParamRef UniformBufferRHI, const void* Contents) final override;
 	virtual FIndexBufferRHIRef RHICreateIndexBuffer(uint32 Stride, uint32 Size, uint32 InUsage, FRHIResourceCreateInfo& CreateInfo) final override;
-	virtual void* RHILockIndexBuffer(FIndexBufferRHIParamRef IndexBuffer, uint32 Offset, uint32 Size, EResourceLockMode LockMode) final override;
-	virtual void RHIUnlockIndexBuffer(FIndexBufferRHIParamRef IndexBuffer) final override;
+	virtual void* LockIndexBuffer_BottomOfPipe(FRHICommandListImmediate& RHICmdList, FIndexBufferRHIParamRef IndexBuffer, uint32 Offset, uint32 Size, EResourceLockMode LockMode) final override;
+	virtual void UnlockIndexBuffer_BottomOfPipe(FRHICommandListImmediate& RHICmdList, FIndexBufferRHIParamRef IndexBuffer) final override;
 	virtual void RHITransferIndexBufferUnderlyingResource(FIndexBufferRHIParamRef DestIndexBuffer, FIndexBufferRHIParamRef SrcIndexBuffer) final override;
 	virtual FVertexBufferRHIRef RHICreateVertexBuffer(uint32 Size, uint32 InUsage, FRHIResourceCreateInfo& CreateInfo) final override;
-	virtual void* RHILockVertexBuffer(FVertexBufferRHIParamRef VertexBuffer, uint32 Offset, uint32 SizeRHI, EResourceLockMode LockMode) final override;
-	virtual void RHIUnlockVertexBuffer(FVertexBufferRHIParamRef VertexBuffer) final override;
+	virtual void* LockVertexBuffer_BottomOfPipe(FRHICommandListImmediate& RHICmdList, FVertexBufferRHIParamRef VertexBuffer, uint32 Offset, uint32 SizeRHI, EResourceLockMode LockMode) final override;
+	virtual void UnlockVertexBuffer_BottomOfPipe(FRHICommandListImmediate& RHICmdList, FVertexBufferRHIParamRef VertexBuffer) final override;
 	virtual void RHICopyVertexBuffer(FVertexBufferRHIParamRef SourceBuffer,FVertexBufferRHIParamRef DestBuffer) final override;
 	virtual void RHITransferVertexBufferUnderlyingResource(FVertexBufferRHIParamRef DestVertexBuffer, FVertexBufferRHIParamRef SrcVertexBuffer) final override;
 	virtual FStructuredBufferRHIRef RHICreateStructuredBuffer(uint32 Stride, uint32 Size, uint32 InUsage, FRHIResourceCreateInfo& CreateInfo) final override;
-	virtual void* RHILockStructuredBuffer(FStructuredBufferRHIParamRef StructuredBuffer, uint32 Offset, uint32 SizeRHI, EResourceLockMode LockMode) final override;
-	virtual void RHIUnlockStructuredBuffer(FStructuredBufferRHIParamRef StructuredBuffer) final override;
+	virtual void* LockStructuredBuffer_BottomOfPipe(FRHICommandListImmediate& RHICmdList, FStructuredBufferRHIParamRef StructuredBuffer, uint32 Offset, uint32 SizeRHI, EResourceLockMode LockMode) final override;
+	virtual void UnlockStructuredBuffer_BottomOfPipe(FRHICommandListImmediate& RHICmdList, FStructuredBufferRHIParamRef StructuredBuffer) final override;
 	virtual FUnorderedAccessViewRHIRef RHICreateUnorderedAccessView(FStructuredBufferRHIParamRef StructuredBuffer, bool bUseUAVCounter, bool bAppendBuffer) final override;
 	virtual FUnorderedAccessViewRHIRef RHICreateUnorderedAccessView(FTextureRHIParamRef Texture, uint32 MipLevel) final override;
 	virtual FUnorderedAccessViewRHIRef RHICreateUnorderedAccessView(FVertexBufferRHIParamRef VertexBuffer, uint8 Format) final override;
diff --git a/Engine/Source/Runtime/Windows/D3D11RHI/Private/D3D11RenderTarget.cpp b/Engine/Source/Runtime/Windows/D3D11RHI/Private/D3D11RenderTarget.cpp
--- a/Engine/Source/Runtime/Windows/D3D11RHI/Private/D3D11RenderTarget.cpp
+++ b/Engine/Source/Runtime/Windows/D3D11RHI/Private/D3D11RenderTarget.cpp
@@ -142,7 +142,7 @@ void FD3D11DynamicRHI::ResolveTextureUsingShader(
 	TShaderMapRef<FResolveVS> ResolveVertexShader(ShaderMap);
 	TShaderMapRef<TPixelShader> ResolvePixelShader(ShaderMap);
 
-	GraphicsPSOInit.BoundShaderState.VertexDeclarationRHI = GScreenVertexDeclaration.VertexDeclarationRHI;
+	GraphicsPSOInit.BoundShaderState.VertexDeclarationRHI = GResolveVertexBuffer.VertexDeclarationRHI;
 	GraphicsPSOInit.BoundShaderState.VertexShaderRHI = GETSAFERHISHADER_VERTEX(*ResolveVertexShader);
 	GraphicsPSOInit.BoundShaderState.PixelShaderRHI = GETSAFERHISHADER_PIXEL(*ResolvePixelShader);
 	GraphicsPSOInit.PrimitiveType = PT_TriangleStrip;
@@ -152,6 +152,7 @@ void FD3D11DynamicRHI::ResolveTextureUsingShader(
 	RHICmdList.SetBlendFactor(FLinearColor::White);
 
 	ResolvePixelShader->SetParameters(RHICmdList, PixelShaderParameter);
+	ResolveVertexShader->SetParameters(RHICmdList, FVector4(MinX, MinY, MaxX, MaxY), FVector4(MinU, MinV, MaxU, MaxV));
 	RHICmdList.Flush(); // always call flush when using a command list in RHI implementations before doing anything else. This is super hazardous.
 
 	// Set the source texture.
@@ -162,35 +163,7 @@ void FD3D11DynamicRHI::ResolveTextureUsingShader(
 		SetShaderResourceView<SF_Pixel>(SourceTexture, SourceTexture->GetShaderResourceView(), TextureIndex, SourceTexture->GetName());
 	}
 
-	FRHIResourceCreateInfo CreateInfo;
-	FVertexBufferRHIRef VertexBufferRHI = RHICreateVertexBuffer(sizeof(FScreenVertex) * 4, BUF_Volatile, CreateInfo);
-	void* VoidPtr = RHILockVertexBuffer(VertexBufferRHI, 0, sizeof(FScreenVertex) * 4, RLM_WriteOnly);
-
-	// Generate the vertices used
-	FScreenVertex* Vertices = (FScreenVertex*)VoidPtr;
-
-	Vertices[0].Position.X = MaxX;
-	Vertices[0].Position.Y = MinY;
-	Vertices[0].UV.X       = MaxU;
-	Vertices[0].UV.Y       = MinV;
-
-	Vertices[1].Position.X = MaxX;
-	Vertices[1].Position.Y = MaxY;
-	Vertices[1].UV.X       = MaxU;
-	Vertices[1].UV.Y       = MaxV;
-
-	Vertices[2].Position.X = MinX;
-	Vertices[2].Position.Y = MinY;
-	Vertices[2].UV.X       = MinU;
-	Vertices[2].UV.Y       = MinV;
-
-	Vertices[3].Position.X = MinX;
-	Vertices[3].Position.Y = MaxY;
-	Vertices[3].UV.X       = MinU;
-	Vertices[3].UV.Y       = MaxV;
-
-	RHIUnlockVertexBuffer(VertexBufferRHI);
-	RHICmdList.SetStreamSource(0, VertexBufferRHI, 0);
+	RHICmdList.SetStreamSource(0, GResolveVertexBuffer.VertexBufferRHI, 0);
 	RHICmdList.DrawPrimitive(0, 2, 1);
 
 	RHICmdList.Flush(); // always call flush when using a command list in RHI implementations before doing anything else. This is super hazardous.
diff --git a/Engine/Source/Runtime/Windows/D3D11RHI/Private/D3D11StructuredBuffer.cpp b/Engine/Source/Runtime/Windows/D3D11RHI/Private/D3D11StructuredBuffer.cpp
--- a/Engine/Source/Runtime/Windows/D3D11RHI/Private/D3D11StructuredBuffer.cpp
+++ b/Engine/Source/Runtime/Windows/D3D11RHI/Private/D3D11StructuredBuffer.cpp
@@ -102,7 +102,7 @@ FStructuredBufferRHIRef FD3D11DynamicRHI::CreateStructuredBuffer_RenderThread(
 	return RHICreateStructuredBuffer(Stride, Size, InUsage, CreateInfo);
 }
 
-void* FD3D11DynamicRHI::RHILockStructuredBuffer(FStructuredBufferRHIParamRef StructuredBufferRHI,uint32 Offset,uint32 Size,EResourceLockMode LockMode)
+void* FD3D11DynamicRHI::LockStructuredBuffer_BottomOfPipe(FRHICommandListImmediate& RHICmdList, FStructuredBufferRHIParamRef StructuredBufferRHI,uint32 Offset,uint32 Size,EResourceLockMode LockMode)
 {
 	FD3D11StructuredBuffer* StructuredBuffer = ResourceCast(StructuredBufferRHI);
 	
@@ -167,7 +167,7 @@ void* FD3D11DynamicRHI::RHILockStructuredBuffer(FStructuredBufferRHIParamRef Str
 	return (void*)((uint8*)LockedData.GetData() + Offset);
 }
 
-void FD3D11DynamicRHI::RHIUnlockStructuredBuffer(FStructuredBufferRHIParamRef StructuredBufferRHI)
+void FD3D11DynamicRHI::UnlockStructuredBuffer_BottomOfPipe(FRHICommandListImmediate& RHICmdList, FStructuredBufferRHIParamRef StructuredBufferRHI)
 {
 	FD3D11StructuredBuffer* StructuredBuffer = ResourceCast(StructuredBufferRHI);
 
diff --git a/Engine/Source/Runtime/Windows/D3D11RHI/Private/D3D11VertexBuffer.cpp b/Engine/Source/Runtime/Windows/D3D11RHI/Private/D3D11VertexBuffer.cpp
--- a/Engine/Source/Runtime/Windows/D3D11RHI/Private/D3D11VertexBuffer.cpp
+++ b/Engine/Source/Runtime/Windows/D3D11RHI/Private/D3D11VertexBuffer.cpp
@@ -108,7 +108,7 @@ FVertexBufferRHIRef FD3D11DynamicRHI::CreateVertexBuffer_RenderThread(
 	return RHICreateVertexBuffer(Size, InUsage, CreateInfo);
 }
 
-void* FD3D11DynamicRHI::RHILockVertexBuffer(FVertexBufferRHIParamRef VertexBufferRHI,uint32 Offset,uint32 Size,EResourceLockMode LockMode)
+void* FD3D11DynamicRHI::LockVertexBuffer_BottomOfPipe(FRHICommandListImmediate& RHICmdList, FVertexBufferRHIParamRef VertexBufferRHI,uint32 Offset,uint32 Size,EResourceLockMode LockMode)
 {
 	check(Size > 0);
 
@@ -189,7 +189,7 @@ void* FD3D11DynamicRHI::RHILockVertexBuffer(FVertexBufferRHIParamRef VertexBuffe
 	return (void*)((uint8*)LockedData.GetData() + Offset);
 }
 
-void FD3D11DynamicRHI::RHIUnlockVertexBuffer(FVertexBufferRHIParamRef VertexBufferRHI)
+void FD3D11DynamicRHI::UnlockVertexBuffer_BottomOfPipe(FRHICommandListImmediate& RHICmdList, FVertexBufferRHIParamRef VertexBufferRHI)
 {
 	FD3D11VertexBuffer* VertexBuffer = ResourceCast(VertexBufferRHI);
 
