DR - Fix clang warnings
#rnx

diff --git a/Engine/Source/Runtime/Engine/Private/GPUSkinVertexFactory.cpp b/Engine/Source/Runtime/Engine/Private/GPUSkinVertexFactory.cpp
--- a/Engine/Source/Runtime/Engine/Private/GPUSkinVertexFactory.cpp
+++ b/Engine/Source/Runtime/Engine/Private/GPUSkinVertexFactory.cpp
@@ -244,10 +244,10 @@ bool FGPUBaseSkinVertexFactory::FShaderDataType::UpdateBoneData(FRHICommandListI
 					// PLATFORM_CACHE_LINE_SIZE (128) / 48 = 2.6
 					//  sizeof(FMatrix) == 64
 					// PLATFORM_CACHE_LINE_SIZE (128) / 64 = 2
-					const uint32 NumBones = BoneMap.Num();
-					check(NumBones > 0 && NumBones < 256); // otherwise maybe some bad threading on BoneMap, maybe we need to copy that
+					const uint32 LocalNumBones = BoneMap.Num();
+					check(LocalNumBones > 0 && LocalNumBones < 256); // otherwise maybe some bad threading on BoneMap, maybe we need to copy that
 					const int32 PreFetchStride = 2; // FPlatformMisc::Prefetch stride
-					for (uint32 BoneIdx = 0; BoneIdx < NumBones; BoneIdx++)
+					for (uint32 BoneIdx = 0; BoneIdx < LocalNumBones; BoneIdx++)
 					{
 						const FBoneIndexType RefToLocalIdx = BoneMap[BoneIdx];
 						check(ReferenceToLocalMatrices.IsValidIndex(RefToLocalIdx)); // otherwise maybe some bad threading on BoneMap, maybe we need to copy that
@@ -920,10 +920,10 @@ bool FGPUBaseSkinAPEXClothVertexFactory::ClothShaderType::UpdateClothSimulData(F
 			if (DeferSkeletalLockAndFillToRHIThread())
 			{
 				FRHIVertexBuffer* VertexBuffer = CurrentClothBuffer->VertexBufferRHI;
-				RHICmdList.EnqueueLambda([VertexBuffer, VectorArraySize, &InSimulPositions, &InSimulNormals](FRHICommandListImmediate& RHICmdList)
+				RHICmdList.EnqueueLambda([VertexBuffer, VectorArraySize, &InSimulPositions, &InSimulNormals](FRHICommandListImmediate& InRHICmdList)
 				{
 					QUICK_SCOPE_CYCLE_COUNTER(STAT_FRHICommandUpdateBoneBuffer_Execute);
-					float* RESTRICT Data = (float* RESTRICT)RHICmdList.LockVertexBuffer(VertexBuffer, 0, VectorArraySize, RLM_WriteOnly);
+					float* RESTRICT Data = (float* RESTRICT)InRHICmdList.LockVertexBuffer(VertexBuffer, 0, VectorArraySize, RLM_WriteOnly);
 					uint32 NumSimulVerts = InSimulPositions.Num();
 					check(NumSimulVerts > 0 && NumSimulVerts <= MAX_APEXCLOTH_VERTICES_FOR_VB);
 					float* RESTRICT Pos = (float* RESTRICT) &InSimulPositions[0].X;
@@ -939,7 +939,7 @@ bool FGPUBaseSkinAPEXClothVertexFactory::ClothShaderType::UpdateClothSimulData(F
 						Pos += 3;
 						Normal += 3;
 					}
-					RHICmdList.UnlockVertexBuffer(VertexBuffer);
+					InRHICmdList.UnlockVertexBuffer(VertexBuffer);
 				});
 
 				RHICmdList.RHIThreadFence(true);
diff --git a/Engine/Source/Runtime/RHI/Private/RHICommandList.cpp b/Engine/Source/Runtime/RHI/Private/RHICommandList.cpp
--- a/Engine/Source/Runtime/RHI/Private/RHICommandList.cpp
+++ b/Engine/Source/Runtime/RHI/Private/RHICommandList.cpp
@@ -1973,13 +1973,13 @@ void FDynamicRHI::RHIUnlockVertexBuffer(class FRHICommandListImmediate& RHICmdLi
 		}
 		else
 		{
-			RHICmdList.EnqueueLambda([VertexBuffer, Params](FRHICommandListImmediate& RHICmdList)
+			RHICmdList.EnqueueLambda([VertexBuffer, Params](FRHICommandListImmediate& InRHICmdList)
 			{
 				QUICK_SCOPE_CYCLE_COUNTER(STAT_FRHICommandUpdateVertexBuffer_Execute);
-				void* Data = GDynamicRHI->LockVertexBuffer_BottomOfPipe(RHICmdList, VertexBuffer, Params.Offset, Params.BufferSize, RLM_WriteOnly);
+				void* Data = GDynamicRHI->LockVertexBuffer_BottomOfPipe(InRHICmdList, VertexBuffer, Params.Offset, Params.BufferSize, RLM_WriteOnly);
 				FMemory::Memcpy(Data, Params.Buffer, Params.BufferSize);
 				FMemory::Free(Params.Buffer);
-				GDynamicRHI->UnlockVertexBuffer_BottomOfPipe(RHICmdList, VertexBuffer);
+				GDynamicRHI->UnlockVertexBuffer_BottomOfPipe(InRHICmdList, VertexBuffer);
 			});
 			RHICmdList.RHIThreadFence(true);
 		}
@@ -2051,13 +2051,13 @@ void FDynamicRHI::RHIUnlockIndexBuffer(class FRHICommandListImmediate& RHICmdLis
 		}
 		else
 		{
-			RHICmdList.EnqueueLambda([IndexBuffer, Params](FRHICommandListImmediate& RHICmdList)
+			RHICmdList.EnqueueLambda([IndexBuffer, Params](FRHICommandListImmediate& InRHICmdList)
 			{
 				QUICK_SCOPE_CYCLE_COUNTER(STAT_FRHICommandUpdateIndexBuffer_Execute);
-				void* Data = GDynamicRHI->LockIndexBuffer_BottomOfPipe(RHICmdList, IndexBuffer, Params.Offset, Params.BufferSize, RLM_WriteOnly);
+				void* Data = GDynamicRHI->LockIndexBuffer_BottomOfPipe(InRHICmdList, IndexBuffer, Params.Offset, Params.BufferSize, RLM_WriteOnly);
 				FMemory::Memcpy(Data, Params.Buffer, Params.BufferSize);
 				FMemory::Free(Params.Buffer);
-				GDynamicRHI->UnlockIndexBuffer_BottomOfPipe(RHICmdList, IndexBuffer);
+				GDynamicRHI->UnlockIndexBuffer_BottomOfPipe(InRHICmdList, IndexBuffer);
 			});
 			RHICmdList.RHIThreadFence(true);
 		}
@@ -2129,13 +2129,13 @@ void FDynamicRHI::RHIUnlockStructuredBuffer(class FRHICommandListImmediate& RHIC
 		}
 		else
 		{
-			RHICmdList.EnqueueLambda([StructuredBuffer, Params](FRHICommandListImmediate& RHICmdList)
+			RHICmdList.EnqueueLambda([StructuredBuffer, Params](FRHICommandListImmediate& InRHICmdList)
 			{
 				QUICK_SCOPE_CYCLE_COUNTER(STAT_FRHICommandUpdateStructuredBuffer_Execute);
-				void* Data = GDynamicRHI->LockStructuredBuffer_BottomOfPipe(RHICmdList, StructuredBuffer, Params.Offset, Params.BufferSize, RLM_WriteOnly);
+				void* Data = GDynamicRHI->LockStructuredBuffer_BottomOfPipe(InRHICmdList, StructuredBuffer, Params.Offset, Params.BufferSize, RLM_WriteOnly);
 				FMemory::Memcpy(Data, Params.Buffer, Params.BufferSize);
 				FMemory::Free(Params.Buffer);
-				GDynamicRHI->UnlockStructuredBuffer_BottomOfPipe(RHICmdList, StructuredBuffer);
+				GDynamicRHI->UnlockStructuredBuffer_BottomOfPipe(InRHICmdList, StructuredBuffer);
 			});
 			RHICmdList.RHIThreadFence(true);
 		}
diff --git a/Engine/Source/Runtime/SlateRHIRenderer/Private/SlateRHIRenderingPolicy.cpp b/Engine/Source/Runtime/SlateRHIRenderer/Private/SlateRHIRenderingPolicy.cpp
--- a/Engine/Source/Runtime/SlateRHIRenderer/Private/SlateRHIRenderingPolicy.cpp
+++ b/Engine/Source/Runtime/SlateRHIRenderer/Private/SlateRHIRenderingPolicy.cpp
@@ -194,7 +194,7 @@ void FSlateRHIRenderingPolicy::UpdateVertexAndIndexBuffers(FRHICommandListImmedi
 			FRHIVertexBuffer* VertexBufferRHI = VertexBuffer.VertexBufferRHI;
 			FRHIIndexBuffer* IndexBufferRHI = IndexBuffer.IndexBufferRHI;
 
-			RHICmdList.EnqueueLambda([VertexBufferRHI, IndexBufferRHI, &InBatchData, bAbsoluteIndices](FRHICommandListImmediate& RHICmdList)
+			RHICmdList.EnqueueLambda([VertexBufferRHI, IndexBufferRHI, &InBatchData, bAbsoluteIndices](FRHICommandListImmediate& InRHICmdList)
 			{
 				SCOPE_CYCLE_COUNTER(STAT_SlateUpdateBufferRTTime);
 
@@ -202,15 +202,15 @@ void FSlateRHIRenderingPolicy::UpdateVertexAndIndexBuffers(FRHICommandListImmedi
 				const int32 NumBatchedIndices = InBatchData.GetNumBatchedIndices();
 
 				int32 RequiredVertexBufferSize = NumBatchedVertices * sizeof(FSlateVertex);
-				uint8* VertexBufferData = (uint8*)RHICmdList.LockVertexBuffer(VertexBufferRHI, 0, RequiredVertexBufferSize, RLM_WriteOnly);
+				uint8* VertexBufferData = (uint8*)InRHICmdList.LockVertexBuffer(VertexBufferRHI, 0, RequiredVertexBufferSize, RLM_WriteOnly);
 
 				uint32 RequiredIndexBufferSize = NumBatchedIndices * sizeof(SlateIndex);
-				uint8* IndexBufferData = (uint8*)RHICmdList.LockIndexBuffer(IndexBufferRHI, 0, RequiredIndexBufferSize, RLM_WriteOnly);
+				uint8* IndexBufferData = (uint8*)InRHICmdList.LockIndexBuffer(IndexBufferRHI, 0, RequiredIndexBufferSize, RLM_WriteOnly);
 
 				InBatchData.FillVertexAndIndexBuffer(VertexBufferData, IndexBufferData, bAbsoluteIndices);
 
-				RHICmdList.UnlockVertexBuffer(VertexBufferRHI);
-				RHICmdList.UnlockIndexBuffer(IndexBufferRHI);
+				InRHICmdList.UnlockVertexBuffer(VertexBufferRHI);
+				InRHICmdList.UnlockIndexBuffer(IndexBufferRHI);
 			});
 
 			RHICmdList.RHIThreadFence(true);
diff --git a/Engine/Source/Runtime/SlateRHIRenderer/Private/SlateUpdatableBuffer.cpp b/Engine/Source/Runtime/SlateRHIRenderer/Private/SlateUpdatableBuffer.cpp
--- a/Engine/Source/Runtime/SlateRHIRenderer/Private/SlateUpdatableBuffer.cpp
+++ b/Engine/Source/Runtime/SlateRHIRenderer/Private/SlateUpdatableBuffer.cpp
@@ -78,16 +78,16 @@ void FSlateUpdatableInstanceBuffer::UpdateRenderingData_RenderThread(FRHICommand
 	else
 	{
 		FRHIVertexBuffer* VertexBufferRHI = InstanceBufferResource.VertexBufferRHI;
-		RHICmdList.EnqueueLambda([VertexBufferRHI, &InstanceData = RenderThreadBufferData](FRHICommandListImmediate& RHICmdList)
+		RHICmdList.EnqueueLambda([VertexBufferRHI, &InstanceData = RenderThreadBufferData](FRHICommandListImmediate& InRHICmdList)
 		{
 			SCOPE_CYCLE_COUNTER(STAT_SlateUpdateInstanceBuffer);
 
 			int32 RequiredVertexBufferSize = InstanceData.Num() * sizeof(FVector4);
-			uint8* InstanceBufferData = (uint8*)RHICmdList.LockVertexBuffer(VertexBufferRHI, 0, RequiredVertexBufferSize, RLM_WriteOnly);
+			uint8* InstanceBufferData = (uint8*)InRHICmdList.LockVertexBuffer(VertexBufferRHI, 0, RequiredVertexBufferSize, RLM_WriteOnly);
 
 			FMemory::Memcpy(InstanceBufferData, InstanceData.GetData(), InstanceData.Num() * sizeof(FVector4));
 
-			RHICmdList.UnlockVertexBuffer(VertexBufferRHI);
+			InRHICmdList.UnlockVertexBuffer(VertexBufferRHI);
 		});
 
 		RHICmdList.RHIThreadFence(true);
