Undo 7109510 that collided hard with 7039544 to unblock Dev-Rendering

diff --git a/Engine/Source/Runtime/RHI/Public/DynamicRHI.h b/Engine/Source/Runtime/RHI/Public/DynamicRHI.h
--- a/Engine/Source/Runtime/RHI/Public/DynamicRHI.h
+++ b/Engine/Source/Runtime/RHI/Public/DynamicRHI.h
@@ -1141,8 +1141,6 @@ public:
 	virtual void UnlockTexture2D_RenderThread(class FRHICommandListImmediate& RHICmdList, FRHITexture2D* Texture, uint32 MipIndex, bool bLockWithinMiptail, bool bNeedsDefaultRHIFlush = true);
 	virtual void UpdateTexture2D_RenderThread(class FRHICommandListImmediate& RHICmdList, FRHITexture2D* Texture, uint32 MipIndex, const struct FUpdateTextureRegion2D& UpdateRegion, uint32 SourcePitch, const uint8* SourceData);
 	virtual void UpdateFromBufferTexture2D_RenderThread(class FRHICommandListImmediate& RHICmdList, FRHITexture2D* Texture, uint32 MipIndex, const struct FUpdateTextureRegion2D& UpdateRegion, uint32 SourcePitch, FRHIStructuredBuffer* Buffer, uint32 BufferOffset);
-	virtual void* LockStructuredBuffer_RenderThread(class FRHICommandListImmediate& RHICmdList, FRHIStructuredBuffer* StructuredBuffer, uint32 Offset, uint32 SizeRHI, EResourceLockMode LockMode);
-	virtual void UnlockStructuredBuffer_RenderThread(class FRHICommandListImmediate& RHICmdList, FRHIStructuredBuffer* StructuredBuffer);
 
 	virtual FUpdateTexture3DData BeginUpdateTexture3D_RenderThread(class FRHICommandListImmediate& RHICmdList, FRHITexture3D* Texture, uint32 MipIndex, const struct FUpdateTextureRegion3D& UpdateRegion);
 	virtual void EndUpdateTexture3D_RenderThread(class FRHICommandListImmediate& RHICmdList, FUpdateTexture3DData& UpdateData);
diff --git a/Engine/Source/Runtime/RHI/Public/RHICommandList.h b/Engine/Source/Runtime/RHI/Public/RHICommandList.h
--- a/Engine/Source/Runtime/RHI/Public/RHICommandList.h
+++ b/Engine/Source/Runtime/RHI/Public/RHICommandList.h
@@ -3741,13 +3741,19 @@ public:
 	FORCEINLINE void* LockStructuredBuffer(FRHIStructuredBuffer* StructuredBuffer, uint32 Offset, uint32 SizeRHI, EResourceLockMode LockMode)
 	{
 		LLM_SCOPE(ELLMTag::RHIMisc);
-		return GDynamicRHI->LockStructuredBuffer_RenderThread(*this, StructuredBuffer, Offset, SizeRHI, LockMode);
+		QUICK_SCOPE_CYCLE_COUNTER(STAT_RHIMETHOD_LockStructuredBuffer_Flush);
+		ImmediateFlush(EImmediateFlushType::FlushRHIThread);
+
+		return GDynamicRHI->RHILockStructuredBuffer(*this, StructuredBuffer, Offset, SizeRHI, LockMode);
 	}
 	
 	FORCEINLINE void UnlockStructuredBuffer(FRHIStructuredBuffer* StructuredBuffer)
 	{
 		LLM_SCOPE(ELLMTag::RHIMisc);
-		GDynamicRHI->UnlockStructuredBuffer_RenderThread(*this, StructuredBuffer);
+		QUICK_SCOPE_CYCLE_COUNTER(STAT_RHIMETHOD_UnlockStructuredBuffer_Flush);
+		ImmediateFlush(EImmediateFlushType::FlushRHIThread);
+
+		GDynamicRHI->RHIUnlockStructuredBuffer(*this, StructuredBuffer);
 	}
 	
 	FORCEINLINE FUnorderedAccessViewRHIRef CreateUnorderedAccessView(FRHIStructuredBuffer* StructuredBuffer, bool bUseUAVCounter, bool bAppendBuffer)
