Mac Metal can *always* support volumetric rendering and on iOS it is a *cook* time option. Brings Mac back into line with D3D.

diff --git a/Engine/Source/Runtime/RHI/Public/RHI.h b/Engine/Source/Runtime/RHI/Public/RHI.h
--- a/Engine/Source/Runtime/RHI/Public/RHI.h
+++ b/Engine/Source/Runtime/RHI/Public/RHI.h
@@ -165,6 +165,16 @@ inline bool RHISupportsVertexShaderLayer(const EShaderPlatform Platform)
 	return IsFeatureLevelSupported(Platform, ERHIFeatureLevel::SM4) && IsMetalPlatform(Platform) && (IsPCPlatform(Platform) || (Platform == SP_METAL_MRT && RHIGetShaderLanguageVersion(Platform) >= 4));
 }
 
+/** Return true if and only if the GPU support rendering to volume textures (2D Array, 3D) is guaranteed supported for a target platform.
+	if PipelineVolumeTextureLUTSupportGuaranteedAtRuntime is true then it is guaranteed that GSupportsVolumeTextureRendering is true at runtime.
+*/
+inline bool RHIVolumeTextureRenderingSupportGuaranteed(const EShaderPlatform Platform)
+{
+	return IsFeatureLevelSupported(Platform, ERHIFeatureLevel::SM4)
+		&& (!IsMetalPlatform(Platform) || RHISupportsVertexShaderLayer(Platform)) // For Metal only shader platforms & versions that support vertex-shader-layer can render to volume textures - this is a compile/cook time check.
+		&& !IsOpenGLPlatform(Platform);		// Apparently, some OpenGL 3.3 cards support SM4 but can't render to volume textures
+}
+
 inline bool RHISupports4ComponentUAVReadWrite(EShaderPlatform Platform)
 {
 	// Must match usf PLATFORM_SUPPORTS_4COMPONENT_UAV_READ_WRITE
diff --git a/Engine/Source/Runtime/RHI/Public/RHIDefinitions.h b/Engine/Source/Runtime/RHI/Public/RHIDefinitions.h
--- a/Engine/Source/Runtime/RHI/Public/RHIDefinitions.h
+++ b/Engine/Source/Runtime/RHI/Public/RHIDefinitions.h
@@ -1170,16 +1170,6 @@ inline bool RHISupportsGeometryShaders(const EShaderPlatform Platform)
 	return IsFeatureLevelSupported(Platform, ERHIFeatureLevel::SM4) && !IsMetalPlatform(Platform) && !IsVulkanMobilePlatform(Platform);
 }
 
-/** Return true if and only if the GPU support rendering to volume textures (2D Array, 3D) is guaranteed supported for a target platform.
-	if PipelineVolumeTextureLUTSupportGuaranteedAtRuntime is true then it is guaranteed that GSupportsVolumeTextureRendering is true at runtime.
-*/
-inline bool RHIVolumeTextureRenderingSupportGuaranteed(const EShaderPlatform Platform)
-{
-	return IsFeatureLevelSupported(Platform, ERHIFeatureLevel::SM4)
-		&& !IsMetalPlatform(Platform)		// EMetalFeaturesLayeredRendering supports needs to be checked at runtime
-		&& !IsOpenGLPlatform(Platform);		// Apparently, some OpenGL 3.3 cards support SM4 but can't render to volume textures 
-}
-
 inline bool RHIHasTiledGPU(const EShaderPlatform Platform)
 {
 	// @todo MetalMRT Technically we should include (Platform == SP_METAL_MRT) but this would disable depth-pre-pass which is currently required.
