Adding the ConvertRawR32G32B32A32DataToFColor function fix to 4.24.

diff --git a/Engine/Source/Runtime/RHI/Public/RHISurfaceDataConversion.h b/Engine/Source/Runtime/RHI/Public/RHISurfaceDataConversion.h
--- a/Engine/Source/Runtime/RHI/Public/RHISurfaceDataConversion.h
+++ b/Engine/Source/Runtime/RHI/Public/RHISurfaceDataConversion.h
@@ -191,7 +191,7 @@ static inline void ConvertRawR32G32B32A32DataToFColor(uint32 Width, uint32 Heigh
 
 	for (uint32 Y = 0; Y < Height; Y++)
 	{
-		float* SrcPtr = (float*)In;
+		float* SrcPtr = (float*)(In + Y * SrcPitch);
 		FColor* DestPtr = Out + Y * Width;
 
 		for (uint32 X = 0; X < Width; X++)
