DR - Remove deprecated function

diff --git a/Engine/Source/Runtime/RHI/Public/RHICommandList.h b/Engine/Source/Runtime/RHI/Public/RHICommandList.h
--- a/Engine/Source/Runtime/RHI/Public/RHICommandList.h
+++ b/Engine/Source/Runtime/RHI/Public/RHICommandList.h
@@ -3892,12 +3892,6 @@ public:
 		return GDynamicRHI->RHICreateTexture2DArray_RenderThread(*this, SizeX, SizeY, SizeZ, Format, NumMips, NumSamples, Flags, CreateInfo);
 	}
 
-	UE_DEPRECATED(4.23, "CreateTexture2DArray now takes NumSamples")
-	FORCEINLINE FTexture2DArrayRHIRef CreateTexture2DArray(uint32 SizeX, uint32 SizeY, uint32 SizeZ, uint8 Format, uint32 NumMips, uint32 Flags, FRHIResourceCreateInfo& CreateInfo)
-	{
-		return CreateTexture2DArray(SizeX, SizeY, SizeZ, Format, NumMips, 1, Flags, CreateInfo);
-	}
-	
 	FORCEINLINE FTexture3DRHIRef CreateTexture3D(uint32 SizeX, uint32 SizeY, uint32 SizeZ, uint8 Format, uint32 NumMips, uint32 Flags, FRHIResourceCreateInfo& CreateInfo)
 	{
 		LLM_SCOPE((Flags & (TexCreate_RenderTargetable | TexCreate_DepthStencilTargetable)) != 0 ? ELLMTag::RenderTargets : ELLMTag::Textures);
@@ -4772,12 +4766,6 @@ FORCEINLINE FTexture2DArrayRHIRef RHICreateTexture2DArray(uint32 SizeX, uint32 S
 	return FRHICommandListExecutor::GetImmediateCommandList().CreateTexture2DArray(SizeX, SizeY, SizeZ, Format, NumMips, NumSamples, Flags, CreateInfo);
 }
 
-UE_DEPRECATED(4.23, "RHICreateTexture2DArray now takes NumSamples")
-FORCEINLINE FTexture2DArrayRHIRef RHICreateTexture2DArray(uint32 SizeX, uint32 SizeY, uint32 SizeZ, uint8 Format, uint32 NumMips, uint32 Flags, FRHIResourceCreateInfo& CreateInfo)
-{
-	return RHICreateTexture2DArray(SizeX, SizeY, SizeZ, Format, NumMips, 1, Flags, CreateInfo);
-}
-
 FORCEINLINE FTexture3DRHIRef RHICreateTexture3D(uint32 SizeX, uint32 SizeY, uint32 SizeZ, uint8 Format, uint32 NumMips, uint32 Flags, FRHIResourceCreateInfo& CreateInfo)
 {
 	return FRHICommandListExecutor::GetImmediateCommandList().CreateTexture3D(SizeX, SizeY, SizeZ, Format, NumMips, Flags, CreateInfo);
