Workaround to allow render target with single texture when using MSAA.

diff --git a/Engine/Source/Runtime/RHI/Public/RHIUtilities.h b/Engine/Source/Runtime/RHI/Public/RHIUtilities.h
--- a/Engine/Source/Runtime/RHI/Public/RHIUtilities.h
+++ b/Engine/Source/Runtime/RHI/Public/RHIUtilities.h
@@ -490,29 +490,33 @@ inline void RHICreateTargetableShaderResource2D(
 	uint32 Flags,
 	uint32 TargetableTextureFlags,
 	bool bForceSeparateTargetAndShaderResource,
+	bool bForceSharedTargetAndShaderResource,
 	FRHIResourceCreateInfo& CreateInfo,
 	FTexture2DRHIRef& OutTargetableTexture,
 	FTexture2DRHIRef& OutShaderResourceTexture,
-	uint32 NumSamples=1
-	)
+	uint32 NumSamples = 1
+)
 {
 	// Ensure none of the usage flags are passed in.
 	check(!(Flags & TexCreate_RenderTargetable));
 	check(!(Flags & TexCreate_ResolveTargetable));
 	check(!(Flags & TexCreate_ShaderResource));
 
+	// Ensure we aren't forcing separate and shared textures at the same time.
+	check(!(bForceSeparateTargetAndShaderResource && bForceSharedTargetAndShaderResource));
+
 	// Ensure that all of the flags provided for the targetable texture are not already passed in Flags.
 	check(!(Flags & TargetableTextureFlags));
 
 	// Ensure that the targetable texture is either render or depth-stencil targetable.
 	check(TargetableTextureFlags & (TexCreate_RenderTargetable | TexCreate_DepthStencilTargetable | TexCreate_UAV));
 
-	if (NumSamples > 1)
+	if (NumSamples > 1 && !bForceSharedTargetAndShaderResource)
 	{
 		bForceSeparateTargetAndShaderResource = RHISupportsSeparateMSAAAndResolveTextures(GMaxRHIShaderPlatform);
 	}
 
-	if (!bForceSeparateTargetAndShaderResource/* && GSupportsRenderDepthTargetableShaderResources*/)
+	if (!bForceSeparateTargetAndShaderResource)
 	{
 		// Create a single texture that has both TargetableTextureFlags and TexCreate_ShaderResource set.
 		OutTargetableTexture = OutShaderResourceTexture = RHICreateTexture2D(SizeX, SizeY, Format, NumMips, NumSamples, Flags | TargetableTextureFlags | TexCreate_ShaderResource, CreateInfo);
@@ -530,6 +534,22 @@ inline void RHICreateTargetableShaderResource2D(
 	}
 }
 
+inline void RHICreateTargetableShaderResource2D(
+	uint32 SizeX,
+	uint32 SizeY,
+	uint8 Format,
+	uint32 NumMips,
+	uint32 Flags,
+	uint32 TargetableTextureFlags,
+	bool bForceSeparateTargetAndShaderResource,
+	FRHIResourceCreateInfo& CreateInfo,
+	FTexture2DRHIRef& OutTargetableTexture,
+	FTexture2DRHIRef& OutShaderResourceTexture,
+	uint32 NumSamples = 1)
+{
+	RHICreateTargetableShaderResource2D(SizeX, SizeY, Format, NumMips, Flags, TargetableTextureFlags, bForceSeparateTargetAndShaderResource, false, CreateInfo, OutTargetableTexture, OutShaderResourceTexture, NumSamples);
+}
+
 inline void RHICreateTargetableShaderResource2DArray(
 	uint32 SizeX,
 	uint32 SizeY,
@@ -538,6 +558,8 @@ inline void RHICreateTargetableShaderResource2DArray(
 	uint32 NumMips,
 	uint32 Flags,
 	uint32 TargetableTextureFlags,
+	bool bForceSeparateTargetAndShaderResource,
+	bool bForceSharedTargetAndShaderResource,
 	FRHIResourceCreateInfo& CreateInfo,
 	FTexture2DArrayRHIRef& OutTargetableTexture,
 	FTexture2DArrayRHIRef& OutShaderResourceTexture,
@@ -549,13 +571,16 @@ inline void RHICreateTargetableShaderResource2DArray(
 	check(!(Flags & TexCreate_ResolveTargetable));
 	check(!(Flags & TexCreate_ShaderResource));
 
+	// Ensure we aren't forcing separate and shared textures at the same time.
+	check(!(bForceSeparateTargetAndShaderResource && bForceSharedTargetAndShaderResource));
+
 	// Ensure that all of the flags provided for the targetable texture are not already passed in Flags.
 	check(!(Flags & TargetableTextureFlags));
 
 	// Ensure that the targetable texture is either render or depth-stencil targetable.
 	check(TargetableTextureFlags & (TexCreate_RenderTargetable | TexCreate_DepthStencilTargetable));
-	bool bForceSeparateTargetAndShaderResource = false;
-	if (NumSamples > 1)
+
+	if (NumSamples > 1 && !bForceSharedTargetAndShaderResource)
 	{
 		bForceSeparateTargetAndShaderResource = RHISupportsSeparateMSAAAndResolveTextures(GMaxRHIShaderPlatform);
 	}
@@ -578,6 +603,22 @@ inline void RHICreateTargetableShaderResource2DArray(
 	}
 }
 
+inline void RHICreateTargetableShaderResource2DArray(
+	uint32 SizeX,
+	uint32 SizeY,
+	uint32 SizeZ,
+	uint8 Format,
+	uint32 NumMips,
+	uint32 Flags,
+	uint32 TargetableTextureFlags,
+	FRHIResourceCreateInfo& CreateInfo,
+	FTexture2DArrayRHIRef& OutTargetableTexture,
+	FTexture2DArrayRHIRef& OutShaderResourceTexture,
+	uint32 NumSamples = 1)
+{
+	RHICreateTargetableShaderResource2DArray(SizeX, SizeY, SizeZ, Format, NumMips, Flags, TargetableTextureFlags, false, false, CreateInfo, OutTargetableTexture, OutShaderResourceTexture, NumSamples);
+}
+
 /**
  * Creates 1 or 2 textures with the same dimensions/format.
  * If the RHI supports textures that can be used as both shader resources and render targets,
diff --git a/Engine/Source/Runtime/RenderCore/Private/RenderTargetPool.cpp b/Engine/Source/Runtime/RenderCore/Private/RenderTargetPool.cpp
--- a/Engine/Source/Runtime/RenderCore/Private/RenderTargetPool.cpp
+++ b/Engine/Source/Runtime/RenderCore/Private/RenderTargetPool.cpp
@@ -362,6 +362,7 @@ Done:
 						Desc.Flags,
 						Desc.TargetableFlags,
 						Desc.bForceSeparateTargetAndShaderResource,
+						Desc.bForceSharedTargetAndShaderResource,
 						CreateInfo,
 						(FTexture2DRHIRef&)Found->RenderTargetItem.TargetableTexture,
 						(FTexture2DRHIRef&)Found->RenderTargetItem.ShaderResourceTexture,
@@ -378,6 +379,8 @@ Done:
 						Desc.NumMips,
 						Desc.Flags,
 						Desc.TargetableFlags,
+						Desc.bForceSeparateTargetAndShaderResource,
+						Desc.bForceSharedTargetAndShaderResource,
 						CreateInfo,
 						(FTexture2DArrayRHIRef&)Found->RenderTargetItem.TargetableTexture,
 						(FTexture2DArrayRHIRef&)Found->RenderTargetItem.ShaderResourceTexture,
diff --git a/Engine/Source/Runtime/RenderCore/Public/RendererInterface.h b/Engine/Source/Runtime/RenderCore/Public/RendererInterface.h
--- a/Engine/Source/Runtime/RenderCore/Public/RendererInterface.h
+++ b/Engine/Source/Runtime/RenderCore/Public/RendererInterface.h
@@ -65,6 +65,7 @@ public:
 		, Flags(TexCreate_None)
 		, TargetableFlags(TexCreate_None)
 		, bForceSeparateTargetAndShaderResource(false)
+		, bForceSharedTargetAndShaderResource(false)
 		, DebugName(TEXT("UnknownTexture"))
 		, AutoWritable(true)
 		, bCreateRenderTargetWriteMask(false)
@@ -209,6 +210,7 @@ public:
 			&& LhsFlags == RhsFlags
 			&& TargetableFlags == rhs.TargetableFlags
 			&& bForceSeparateTargetAndShaderResource == rhs.bForceSeparateTargetAndShaderResource
+			&& bForceSharedTargetAndShaderResource == rhs.bForceSharedTargetAndShaderResource
 			&& ClearValue == rhs.ClearValue
 			&& AutoWritable == rhs.AutoWritable;
 	}
@@ -330,6 +332,7 @@ public:
 		NumSamples = 1;
 
 		bForceSeparateTargetAndShaderResource = false;
+		bForceSharedTargetAndShaderResource = false;
 		AutoWritable = true;
 
 		// Remove UAV flag for rendertargets that don't need it (some formats are incompatible)
@@ -363,6 +366,8 @@ public:
 	uint32 TargetableFlags;
 	/** Whether the shader-resource and targetable texture must be separate textures. */
 	bool bForceSeparateTargetAndShaderResource;
+	/** Whether the shader-resource and targetable texture must be the same resource. */
+	bool bForceSharedTargetAndShaderResource;
 	/** only set a pointer to memory that never gets released */
 	const TCHAR* DebugName;
 	/** automatically set to writable via barrier during */
