DR - Implementing RHI Scopes for Unreal Insights.

diff --git a/Engine/Plugins/Experimental/GeometryCache/Source/GeometryCache/Private/GeometryCacheSceneProxy.cpp b/Engine/Plugins/Experimental/GeometryCache/Source/GeometryCache/Private/GeometryCacheSceneProxy.cpp
--- a/Engine/Plugins/Experimental/GeometryCache/Source/GeometryCache/Private/GeometryCacheSceneProxy.cpp
+++ b/Engine/Plugins/Experimental/GeometryCache/Source/GeometryCache/Private/GeometryCacheSceneProxy.cpp
@@ -211,7 +211,7 @@ FGeometryCacheSceneProxy::~FGeometryCacheSceneProxy()
 }
 
 #if 0
-struct FRHICommandUpdateGeometryCacheBuffer : public FRHICommand<FRHICommandUpdateGeometryCacheBuffer>
+FRHICOMMAND_MACRO(FRHICommandUpdateGeometryCacheBuffer)
 {
 	FGraphEventRef BufferGenerationCompleteFence;
 
diff --git a/Engine/Plugins/Experimental/PixelStreaming/Source/PixelStreaming/Private/PixelStreamingNvVideoEncoder.cpp b/Engine/Plugins/Experimental/PixelStreaming/Source/PixelStreaming/Private/PixelStreamingNvVideoEncoder.cpp
--- a/Engine/Plugins/Experimental/PixelStreaming/Source/PixelStreaming/Private/PixelStreamingNvVideoEncoder.cpp
+++ b/Engine/Plugins/Experimental/PixelStreaming/Source/PixelStreaming/Private/PixelStreamingNvVideoEncoder.cpp
@@ -78,7 +78,7 @@ private:
 		FThreadSafeBool bEncoding = false;
 	};
 
-	struct FRHITransferRenderTargetToNvEnc final : public FRHICommand<FRHITransferRenderTargetToNvEnc>
+	FRHICOMMAND_MACRO(FRHITransferRenderTargetToNvEnc)
 	{
 		FPixelStreamingNvVideoEncoder::FPixelStreamingNvVideoEncoderImpl* Encoder;
 		FFrame* Frame;
diff --git a/Engine/Plugins/Media/WmfMedia/Source/WmfMedia/Private/Wmf/WmfMediaHardwareVideoDecodingRendering.cpp b/Engine/Plugins/Media/WmfMedia/Source/WmfMedia/Private/Wmf/WmfMediaHardwareVideoDecodingRendering.cpp
--- a/Engine/Plugins/Media/WmfMedia/Source/WmfMedia/Private/Wmf/WmfMediaHardwareVideoDecodingRendering.cpp
+++ b/Engine/Plugins/Media/WmfMedia/Source/WmfMedia/Private/Wmf/WmfMediaHardwareVideoDecodingRendering.cpp
@@ -29,7 +29,7 @@
 
 #include "WmfMediaHardwareVideoDecodingShaders.h"
 
-struct FRHICommandCopyResource final : public FRHICommand<FRHICommandCopyResource>
+FRHICOMMAND_MACRO(FRHICommandCopyResource)
 {
 	TComPtr<ID3D11Texture2D> SampleTexture;
 	FTexture2DRHIRef SampleDestinationTexture;
diff --git a/Engine/Source/Runtime/GameplayMediaEncoder/Private/GameplayMediaEncoderCommon.h b/Engine/Source/Runtime/GameplayMediaEncoder/Private/GameplayMediaEncoderCommon.h
--- a/Engine/Source/Runtime/GameplayMediaEncoder/Private/GameplayMediaEncoderCommon.h
+++ b/Engine/Source/Runtime/GameplayMediaEncoder/Private/GameplayMediaEncoderCommon.h
@@ -286,7 +286,7 @@ inline void ExecuteRHICommand(F&& Functor)
 	}
 	else
 	{
-		struct FLocalRHICommand final : public FRHICommand<FLocalRHICommand>
+		FRHICOMMAND_MACRO(FLocalRHICommand)
 		{
 			F Functor;
 
diff --git a/Engine/Source/Runtime/HeadMountedDisplay/Private/DefaultSpectatorScreenController.cpp b/Engine/Source/Runtime/HeadMountedDisplay/Private/DefaultSpectatorScreenController.cpp
--- a/Engine/Source/Runtime/HeadMountedDisplay/Private/DefaultSpectatorScreenController.cpp
+++ b/Engine/Source/Runtime/HeadMountedDisplay/Private/DefaultSpectatorScreenController.cpp
@@ -76,7 +76,7 @@ FSpectatorScreenRenderDelegate* FDefaultSpectatorScreenController::GetSpectatorS
 	return &SpectatorScreenDelegate_RenderThread;
 }
 
-struct FRHISetSpectatorScreenTexture final : public FRHICommand<FRHISetSpectatorScreenTexture>
+FRHICOMMAND_MACRO(FRHISetSpectatorScreenTexture)
 {
 	FDefaultSpectatorScreenController* SpectatorScreenController;
 	FTexture2DRHIRef Texture;
@@ -131,7 +131,7 @@ void FDefaultSpectatorScreenController::SetSpectatorScreenTexture_RenderThread(F
 }
 
 
-struct FRHISetSpectatorScreenModeTexturePlusEyeLayout final : public FRHICommand<FRHISetSpectatorScreenModeTexturePlusEyeLayout>
+FRHICOMMAND_MACRO(FRHISetSpectatorScreenModeTexturePlusEyeLayout)
 {
 	FDefaultSpectatorScreenController* SpectatorScreenController;
 	FSpectatorScreenModeTexturePlusEyeLayout Layout;
diff --git a/Engine/Source/Runtime/RHI/Private/RHICommandList.cpp b/Engine/Source/Runtime/RHI/Private/RHICommandList.cpp
--- a/Engine/Source/Runtime/RHI/Private/RHICommandList.cpp
+++ b/Engine/Source/Runtime/RHI/Private/RHICommandList.cpp
@@ -158,7 +158,7 @@ RHI_API FAutoConsoleTaskPriority CPrio_SceneRenderingTask(
 	ENamedThreads::HighTaskPriority 
 	);
 
-struct FRHICommandStat final : public FRHICommand<FRHICommandStat>
+FRHICOMMAND_MACRO(FRHICommandStat)
 {
 	TStatId CurrentExecuteStat;
 	FORCEINLINE_DEBUGGABLE FRHICommandStat(TStatId InCurrentExecuteStat)
@@ -748,7 +748,7 @@ bool FRHICommandListExecutor::IsRHIThreadCompletelyFlushed()
 	return !RenderThreadSublistDispatchTask;
 }
 
-struct FRHICommandRHIThreadFence final : public FRHICommand<FRHICommandRHIThreadFence>
+FRHICOMMAND_MACRO(FRHICommandRHIThreadFence)
 {
 	FGraphEventRef Fence;
 	FORCEINLINE_DEBUGGABLE FRHICommandRHIThreadFence()
@@ -783,7 +783,7 @@ FGraphEventRef FRHICommandListImmediate::RHIThreadFence(bool bSetLockFence)
 }		
 
 DECLARE_CYCLE_STAT(TEXT("Async Compute CmdList Execute"), STAT_AsyncComputeExecute, STATGROUP_RHICMDLIST);
-struct FRHIAsyncComputeSubmitList final : public FRHICommand<FRHIAsyncComputeSubmitList>
+FRHICOMMAND_MACRO(FRHIAsyncComputeSubmitList)
 {
 	FRHIAsyncComputeCommandList* RHICmdList;
 	FORCEINLINE_DEBUGGABLE FRHIAsyncComputeSubmitList(FRHIAsyncComputeCommandList* InRHICmdList)
@@ -957,7 +957,7 @@ DECLARE_DWORD_COUNTER_STAT(TEXT("Num Parallel Async Chains Links"), STAT_Paralle
 DECLARE_CYCLE_STAT(TEXT("Wait for Parallel Async CmdList"), STAT_ParallelChainWait, STATGROUP_RHICMDLIST);
 DECLARE_CYCLE_STAT(TEXT("Parallel Async Chain Execute"), STAT_ParallelChainExecute, STATGROUP_RHICMDLIST);
 
-struct FRHICommandWaitForAndSubmitSubListParallel final : public FRHICommand<FRHICommandWaitForAndSubmitSubListParallel>
+FRHICOMMAND_MACRO(FRHICommandWaitForAndSubmitSubListParallel)
 {
 	FGraphEventRef TranslateCompletionEvent;
 	IRHICommandContextContainer* ContextContainer;
@@ -1008,7 +1008,7 @@ DECLARE_CYCLE_STAT(TEXT("Async Chain Execute"), STAT_ChainExecute, STATGROUP_RHI
 
 FGraphEvent* GEventToWaitFor = nullptr;
 
-struct FRHICommandWaitForAndSubmitSubList final : public FRHICommand<FRHICommandWaitForAndSubmitSubList>
+FRHICOMMAND_MACRO(FRHICommandWaitForAndSubmitSubList)
 {
 	FGraphEventRef EventToWaitFor;
 	FRHICommandListBase* RHICmdList;
@@ -1373,7 +1373,7 @@ DECLARE_DWORD_COUNTER_STAT(TEXT("Num RT Chains Links"), STAT_RTChainLinkCount, S
 DECLARE_CYCLE_STAT(TEXT("Wait for RT CmdList"), STAT_RTChainWait, STATGROUP_RHICMDLIST);
 DECLARE_CYCLE_STAT(TEXT("RT Chain Execute"), STAT_RTChainExecute, STATGROUP_RHICMDLIST);
 
-struct FRHICommandWaitForAndSubmitRTSubList final : public FRHICommand<FRHICommandWaitForAndSubmitRTSubList>
+FRHICOMMAND_MACRO(FRHICommandWaitForAndSubmitRTSubList)
 {
 	FGraphEventRef EventToWaitFor;
 	FRHICommandList* RHICmdList;
@@ -1434,7 +1434,7 @@ void FRHICommandListBase::AddDispatchPrerequisite(const FGraphEventRef& Prereq)
 	}
 }
 
-struct FRHICommandSubmitSubList final : public FRHICommand<FRHICommandSubmitSubList>
+FRHICOMMAND_MACRO(FRHICommandSubmitSubList)
 {
 	FRHICommandList* RHICmdList;
 
diff --git a/Engine/Source/Runtime/RHI/Public/RHICommandList.h b/Engine/Source/Runtime/RHI/Public/RHICommandList.h
--- a/Engine/Source/Runtime/RHI/Public/RHICommandList.h
+++ b/Engine/Source/Runtime/RHI/Public/RHICommandList.h
@@ -26,6 +26,7 @@
 #include "Async/TaskGraphInterfaces.h"
 #include "HAL/LowLevelMemTracker.h"
 #include "ProfilingDebugging/CsvProfiler.h"
+#include "ProfilingDebugging/CpuProfilerTrace.h"
 
 CSV_DECLARE_CATEGORY_MODULE_EXTERN(RHI_API, RHITStalls);
 CSV_DECLARE_CATEGORY_MODULE_EXTERN(RHI_API, RHITFlushes);
@@ -658,7 +659,12 @@ public:
 	FCommonData Data;
 };
 
-template<typename TCmd>
+struct FUnnamedRhiCommand
+{
+	static const TCHAR* TStr() { return TEXT("FUnnamedRhiCommand"); }
+};
+
+template<typename TCmd, typename NameType = FUnnamedRhiCommand>
 struct FRHICommand : public FRHICommandBase
 {
 #if RHICOMMAND_CALLSTACK
@@ -672,6 +678,15 @@ struct FRHICommand : public FRHICommandBase
 
 	void ExecuteAndDestruct(FRHICommandListBase& CmdList, FRHICommandListDebugContext& Context) override final
 	{
+#if CPUPROFILERTRACE_ENABLED
+		static uint16 __CpuProfilerEventSpecId;
+		if (__CpuProfilerEventSpecId == 0)
+		{
+			__CpuProfilerEventSpecId = FCpuProfilerTrace::OutputEventType(TEXT("TRHILambdaCommand"), CpuProfilerGroup_Default);
+		}
+		FCpuProfilerTrace::FEventScope __CpuProfilerEventScope(__CpuProfilerEventSpecId);
+#endif
+
 		TCmd *ThisCmd = static_cast<TCmd*>(this);
 #if RHI_COMMAND_LIST_DEBUG_TRACES
 		ThisCmd->StoreDebugInfo(Context);
@@ -683,7 +698,14 @@ struct FRHICommand : public FRHICommandBase
 	virtual void StoreDebugInfo(FRHICommandListDebugContext& Context) {};
 };
 
-struct  FRHICommandBeginUpdateMultiFrameResource final : public FRHICommand<FRHICommandBeginUpdateMultiFrameResource>
+#define FRHICOMMAND_MACRO(CommandName)								\
+struct PREPROCESSOR_JOIN(CommandName##String, __LINE__)				\
+{																	\
+	static const TCHAR* TStr() { return TEXT(#CommandName); }		\
+};																	\
+struct CommandName final : public FRHICommand<CommandName, PREPROCESSOR_JOIN(CommandName##String, __LINE__)>
+
+FRHICOMMAND_MACRO(FRHICommandBeginUpdateMultiFrameResource)
 {
 	FRHITexture* Texture;
 	FORCEINLINE_DEBUGGABLE FRHICommandBeginUpdateMultiFrameResource(FRHITexture* InTexture)
@@ -693,7 +715,7 @@ struct  FRHICommandBeginUpdateMultiFrameResource final : public FRHICommand<FRHI
 	RHI_API void Execute(FRHICommandListBase& CmdList);
 };
 
-struct  FRHICommandEndUpdateMultiFrameResource final : public FRHICommand<FRHICommandEndUpdateMultiFrameResource>
+FRHICOMMAND_MACRO(FRHICommandEndUpdateMultiFrameResource)
 {
 	FRHITexture* Texture;
 	FORCEINLINE_DEBUGGABLE FRHICommandEndUpdateMultiFrameResource(FRHITexture* InTexture)
@@ -703,7 +725,7 @@ struct  FRHICommandEndUpdateMultiFrameResource final : public FRHICommand<FRHICo
 	RHI_API void Execute(FRHICommandListBase& CmdList);
 };
 
-struct  FRHICommandBeginUpdateMultiFrameUAV final : public FRHICommand<FRHICommandBeginUpdateMultiFrameResource>
+FRHICOMMAND_MACRO(FRHICommandBeginUpdateMultiFrameUAV)
 {
 	FRHIUnorderedAccessView* UAV;
 	FORCEINLINE_DEBUGGABLE FRHICommandBeginUpdateMultiFrameUAV(FRHIUnorderedAccessView* InUAV)
@@ -713,7 +735,7 @@ struct  FRHICommandBeginUpdateMultiFrameUAV final : public FRHICommand<FRHIComma
 	RHI_API void Execute(FRHICommandListBase& CmdList);
 };
 
-struct  FRHICommandEndUpdateMultiFrameUAV final : public FRHICommand<FRHICommandEndUpdateMultiFrameResource>
+FRHICOMMAND_MACRO(FRHICommandEndUpdateMultiFrameUAV)
 {
 	FRHIUnorderedAccessView* UAV;
 	FORCEINLINE_DEBUGGABLE FRHICommandEndUpdateMultiFrameUAV(FRHIUnorderedAccessView* InUAV)
@@ -723,7 +745,7 @@ struct  FRHICommandEndUpdateMultiFrameUAV final : public FRHICommand<FRHICommand
 	RHI_API void Execute(FRHICommandListBase& CmdList);
 };
 
-struct FRHICommandSetStencilRef final : public FRHICommand<FRHICommandSetStencilRef>
+FRHICOMMAND_MACRO(FRHICommandSetStencilRef)
 {
 	uint32 StencilRef;
 	FORCEINLINE_DEBUGGABLE FRHICommandSetStencilRef(uint32 InStencilRef)
@@ -733,8 +755,12 @@ struct FRHICommandSetStencilRef final : public FRHICommand<FRHICommandSetStencil
 	RHI_API void Execute(FRHICommandListBase& CmdList);
 };
 
+struct FRHICommandSetShaderParameterString
+{
+	static const TCHAR* TStr() { return TEXT("FRHICommandSetShaderParameter"); }
+};
 template <typename TRHIShader, ECmdList CmdListType>
-struct FRHICommandSetShaderParameter final : public FRHICommand<FRHICommandSetShaderParameter<TRHIShader, CmdListType> >
+struct FRHICommandSetShaderParameter final : public FRHICommand<FRHICommandSetShaderParameter<TRHIShader, CmdListType>, FRHICommandSetShaderParameterString>
 {
 	TRHIShader* Shader;
 	const void* NewValue;
@@ -752,8 +778,12 @@ struct FRHICommandSetShaderParameter final : public FRHICommand<FRHICommandSetSh
 	RHI_API void Execute(FRHICommandListBase& CmdList);
 };
 
+struct FRHICommandSetShaderUniformBufferString
+{
+	static const TCHAR* TStr() { return TEXT("FRHICommandSetShaderUniformBuffer"); }
+};
 template <typename TRHIShader, ECmdList CmdListType>
-struct FRHICommandSetShaderUniformBuffer final : public FRHICommand<FRHICommandSetShaderUniformBuffer<TRHIShader, CmdListType> >
+struct FRHICommandSetShaderUniformBuffer final : public FRHICommand<FRHICommandSetShaderUniformBuffer<TRHIShader, CmdListType>, FRHICommandSetShaderUniformBufferString>
 {
 	TRHIShader* Shader;
 	uint32 BaseIndex;
@@ -767,8 +797,12 @@ struct FRHICommandSetShaderUniformBuffer final : public FRHICommand<FRHICommandS
 	RHI_API void Execute(FRHICommandListBase& CmdList);
 };
 
+struct FRHICommandSetShaderTextureString
+{
+	static const TCHAR* TStr() { return TEXT("FRHICommandSetShaderTexture"); }
+};
 template <typename TRHIShader, ECmdList CmdListType>
-struct FRHICommandSetShaderTexture final : public FRHICommand<FRHICommandSetShaderTexture<TRHIShader, CmdListType> >
+struct FRHICommandSetShaderTexture final : public FRHICommand<FRHICommandSetShaderTexture<TRHIShader, CmdListType>, FRHICommandSetShaderTextureString >
 {
 	TRHIShader* Shader;
 	uint32 TextureIndex;
@@ -782,8 +816,12 @@ struct FRHICommandSetShaderTexture final : public FRHICommand<FRHICommandSetShad
 	RHI_API void Execute(FRHICommandListBase& CmdList);
 };
 
+struct FRHICommandSetShaderResourceViewParameterString
+{
+	static const TCHAR* TStr() { return TEXT("FRHICommandSetShaderResourceViewParameter"); }
+};
 template <typename TRHIShader, ECmdList CmdListType>
-struct FRHICommandSetShaderResourceViewParameter final : public FRHICommand<FRHICommandSetShaderResourceViewParameter<TRHIShader, CmdListType> >
+struct FRHICommandSetShaderResourceViewParameter final : public FRHICommand<FRHICommandSetShaderResourceViewParameter<TRHIShader, CmdListType>, FRHICommandSetShaderResourceViewParameterString >
 {
 	TRHIShader* Shader;
 	uint32 SamplerIndex;
@@ -797,8 +835,12 @@ struct FRHICommandSetShaderResourceViewParameter final : public FRHICommand<FRHI
 	RHI_API void Execute(FRHICommandListBase& CmdList);
 };
 
+struct FRHICommandSetUAVParameterString
+{
+	static const TCHAR* TStr() { return TEXT("FRHICommandSetUAVParameter"); }
+};
 template <typename TRHIShader, ECmdList CmdListType>
-struct FRHICommandSetUAVParameter final : public FRHICommand<FRHICommandSetUAVParameter<TRHIShader, CmdListType> >
+struct FRHICommandSetUAVParameter final : public FRHICommand<FRHICommandSetUAVParameter<TRHIShader, CmdListType>, FRHICommandSetUAVParameterString >
 {
 	TRHIShader* Shader;
 	uint32 UAVIndex;
@@ -812,8 +854,12 @@ struct FRHICommandSetUAVParameter final : public FRHICommand<FRHICommandSetUAVPa
 	RHI_API void Execute(FRHICommandListBase& CmdList);
 };
 
+struct FRHICommandSetUAVParameter_IntialCountString
+{
+	static const TCHAR* TStr() { return TEXT("FRHICommandSetUAVParameter_IntialCount"); }
+};
 template <typename TRHIShader, ECmdList CmdListType>
-struct FRHICommandSetUAVParameter_IntialCount final : public FRHICommand<FRHICommandSetUAVParameter_IntialCount<TRHIShader, CmdListType> >
+struct FRHICommandSetUAVParameter_IntialCount final : public FRHICommand<FRHICommandSetUAVParameter_IntialCount<TRHIShader, CmdListType>, FRHICommandSetUAVParameter_IntialCountString >
 {
 	TRHIShader* Shader;
 	uint32 UAVIndex;
@@ -829,8 +875,12 @@ struct FRHICommandSetUAVParameter_IntialCount final : public FRHICommand<FRHICom
 	RHI_API void Execute(FRHICommandListBase& CmdList);
 };
 
+struct FRHICommandSetShaderSamplerString
+{
+	static const TCHAR* TStr() { return TEXT("FRHICommandSetShaderSampler"); }
+};
 template <typename TRHIShader, ECmdList CmdListType>
-struct FRHICommandSetShaderSampler final : public FRHICommand<FRHICommandSetShaderSampler<TRHIShader, CmdListType> >
+struct FRHICommandSetShaderSampler final : public FRHICommand<FRHICommandSetShaderSampler<TRHIShader, CmdListType>, FRHICommandSetShaderSamplerString >
 {
 	TRHIShader* Shader;
 	uint32 SamplerIndex;
@@ -844,7 +894,7 @@ struct FRHICommandSetShaderSampler final : public FRHICommand<FRHICommandSetShad
 	RHI_API void Execute(FRHICommandListBase& CmdList);
 };
 
-struct FRHICommandDrawPrimitive final : public FRHICommand<FRHICommandDrawPrimitive>
+FRHICOMMAND_MACRO(FRHICommandDrawPrimitive)
 {
 	uint32 BaseVertexIndex;
 	uint32 NumPrimitives;
@@ -858,7 +908,7 @@ struct FRHICommandDrawPrimitive final : public FRHICommand<FRHICommandDrawPrimit
 	RHI_API void Execute(FRHICommandListBase& CmdList);
 };
 
-struct FRHICommandDrawIndexedPrimitive final : public FRHICommand<FRHICommandDrawIndexedPrimitive>
+FRHICOMMAND_MACRO(FRHICommandDrawIndexedPrimitive)
 {
 	FRHIIndexBuffer* IndexBuffer;
 	int32 BaseVertexIndex;
@@ -880,7 +930,7 @@ struct FRHICommandDrawIndexedPrimitive final : public FRHICommand<FRHICommandDra
 	RHI_API void Execute(FRHICommandListBase& CmdList);
 };
 
-struct FRHICommandSetBlendFactor final : public FRHICommand<FRHICommandSetBlendFactor>
+FRHICOMMAND_MACRO(FRHICommandSetBlendFactor)
 {
 	FLinearColor BlendFactor;
 	FORCEINLINE_DEBUGGABLE FRHICommandSetBlendFactor(const FLinearColor& InBlendFactor)
@@ -890,7 +940,7 @@ struct FRHICommandSetBlendFactor final : public FRHICommand<FRHICommandSetBlendF
 	RHI_API void Execute(FRHICommandListBase& CmdList);
 };
 
-struct FRHICommandSetStreamSource final : public FRHICommand<FRHICommandSetStreamSource>
+FRHICOMMAND_MACRO(FRHICommandSetStreamSource)
 {
 	uint32 StreamIndex;
 	FRHIVertexBuffer* VertexBuffer;
@@ -904,7 +954,7 @@ struct FRHICommandSetStreamSource final : public FRHICommand<FRHICommandSetStrea
 	RHI_API void Execute(FRHICommandListBase& CmdList);
 };
 
-struct FRHICommandSetViewport final : public FRHICommand<FRHICommandSetViewport>
+FRHICOMMAND_MACRO(FRHICommandSetViewport)
 {
 	uint32 MinX;
 	uint32 MinY;
@@ -924,7 +974,7 @@ struct FRHICommandSetViewport final : public FRHICommand<FRHICommandSetViewport>
 	RHI_API void Execute(FRHICommandListBase& CmdList);
 };
 
-struct FRHICommandSetStereoViewport final : public FRHICommand<FRHICommandSetStereoViewport>
+FRHICOMMAND_MACRO(FRHICommandSetStereoViewport)
 {
 	uint32 LeftMinX;
 	uint32 RightMinX;
@@ -952,7 +1002,7 @@ struct FRHICommandSetStereoViewport final : public FRHICommand<FRHICommandSetSte
 	RHI_API void Execute(FRHICommandListBase& CmdList);
 };
 
-struct FRHICommandSetScissorRect final : public FRHICommand<FRHICommandSetScissorRect>
+FRHICOMMAND_MACRO(FRHICommandSetScissorRect)
 {
 	bool bEnable;
 	uint32 MinX;
@@ -970,7 +1020,7 @@ struct FRHICommandSetScissorRect final : public FRHICommand<FRHICommandSetScisso
 	RHI_API void Execute(FRHICommandListBase& CmdList);
 };
 
-struct FRHICommandSetRenderTargets final : public FRHICommand<FRHICommandSetRenderTargets>
+FRHICOMMAND_MACRO(FRHICommandSetRenderTargets)
 {
 	uint32 NewNumSimultaneousRenderTargets;
 	FRHIRenderTargetView NewRenderTargetsRHI[MaxSimultaneousRenderTargets];
@@ -1006,7 +1056,7 @@ struct FRHICommandSetRenderTargets final : public FRHICommand<FRHICommandSetRend
 	RHI_API void Execute(FRHICommandListBase& CmdList);
 };
 
-struct FRHICommandBeginRenderPass final : public FRHICommand<FRHICommandBeginRenderPass>
+FRHICOMMAND_MACRO(FRHICommandBeginRenderPass)
 {
 	FRHIRenderPassInfo Info;
 	const TCHAR* Name;
@@ -1020,7 +1070,7 @@ struct FRHICommandBeginRenderPass final : public FRHICommand<FRHICommandBeginRen
 	RHI_API void Execute(FRHICommandListBase& CmdList);
 };
 
-struct FRHICommandEndRenderPass final : public FRHICommand<FRHICommandEndRenderPass>
+FRHICOMMAND_MACRO(FRHICommandEndRenderPass)
 {
 	FRHICommandEndRenderPass()
 	{
@@ -1029,7 +1079,7 @@ struct FRHICommandEndRenderPass final : public FRHICommand<FRHICommandEndRenderP
 	RHI_API void Execute(FRHICommandListBase& CmdList);
 };
 
-struct FRHICommandNextSubpass final : public FRHICommand<FRHICommandNextSubpass>
+FRHICOMMAND_MACRO(FRHICommandNextSubpass)
 {
 	FRHICommandNextSubpass()
 	{
@@ -1043,7 +1093,7 @@ struct FLocalCmdListParallelRenderPass
 	TRefCountPtr<struct FRHIParallelRenderPass> RenderPass;
 };
 
-struct FRHICommandBeginParallelRenderPass final : public FRHICommand<FRHICommandBeginParallelRenderPass>
+FRHICOMMAND_MACRO(FRHICommandBeginParallelRenderPass)
 {
 	FRHIRenderPassInfo Info;
 	FLocalCmdListParallelRenderPass* LocalRenderPass;
@@ -1059,7 +1109,7 @@ struct FRHICommandBeginParallelRenderPass final : public FRHICommand<FRHICommand
 	RHI_API void Execute(FRHICommandListBase& CmdList);
 };
 
-struct FRHICommandEndParallelRenderPass final : public FRHICommand<FRHICommandEndParallelRenderPass>
+FRHICOMMAND_MACRO(FRHICommandEndParallelRenderPass)
 {
 	FLocalCmdListParallelRenderPass* LocalRenderPass;
 
@@ -1076,7 +1126,7 @@ struct FLocalCmdListRenderSubPass
 	TRefCountPtr<struct FRHIRenderSubPass> RenderSubPass;
 };
 
-struct FRHICommandBeginRenderSubPass final : public FRHICommand<FRHICommandBeginRenderSubPass>
+FRHICOMMAND_MACRO(FRHICommandBeginRenderSubPass)
 {
 	FLocalCmdListParallelRenderPass* LocalRenderPass;
 	FLocalCmdListRenderSubPass* LocalRenderSubPass;
@@ -1090,7 +1140,7 @@ struct FRHICommandBeginRenderSubPass final : public FRHICommand<FRHICommandBegin
 	RHI_API void Execute(FRHICommandListBase& CmdList);
 };
 
-struct FRHICommandEndRenderSubPass final : public FRHICommand<FRHICommandEndRenderSubPass>
+FRHICOMMAND_MACRO(FRHICommandEndRenderSubPass)
 {
 	FLocalCmdListParallelRenderPass* LocalRenderPass;
 	FLocalCmdListRenderSubPass* LocalRenderSubPass;
@@ -1104,7 +1154,7 @@ struct FRHICommandEndRenderSubPass final : public FRHICommand<FRHICommandEndRend
 	RHI_API void Execute(FRHICommandListBase& CmdList);
 };
 
-struct FRHICommandBeginComputePass final : public FRHICommand<FRHICommandBeginComputePass>
+FRHICOMMAND_MACRO(FRHICommandBeginComputePass)
 {
 	const TCHAR* Name;
 
@@ -1116,7 +1166,7 @@ struct FRHICommandBeginComputePass final : public FRHICommand<FRHICommandBeginCo
 	RHI_API void Execute(FRHICommandListBase& CmdList);
 };
 
-struct FRHICommandEndComputePass final : public FRHICommand<FRHICommandEndComputePass>
+FRHICOMMAND_MACRO(FRHICommandEndComputePass)
 {
 	FRHICommandEndComputePass()
 	{
@@ -1125,7 +1175,7 @@ struct FRHICommandEndComputePass final : public FRHICommand<FRHICommandEndComput
 	RHI_API void Execute(FRHICommandListBase& CmdList);
 };
 
-struct FRHICommandBindClearMRTValues final : public FRHICommand<FRHICommandBindClearMRTValues>
+FRHICOMMAND_MACRO(FRHICommandBindClearMRTValues)
 {
 	bool bClearColor;
 	bool bClearDepth;
@@ -1145,8 +1195,12 @@ struct FRHICommandBindClearMRTValues final : public FRHICommand<FRHICommandBindC
 	RHI_API void Execute(FRHICommandListBase& CmdList);
 };
 
+struct FRHICommandSetComputeShaderString
+{
+	static const TCHAR* TStr() { return TEXT("FRHICommandSetComputeShader"); }
+};
 template<ECmdList CmdListType>
-struct FRHICommandSetComputeShader final : public FRHICommand<FRHICommandSetComputeShader<CmdListType>>
+struct FRHICommandSetComputeShader final : public FRHICommand<FRHICommandSetComputeShader<CmdListType>, FRHICommandSetComputeShaderString>
 {
 	FRHIComputeShader* ComputeShader;
 	FORCEINLINE_DEBUGGABLE FRHICommandSetComputeShader(FRHIComputeShader* InComputeShader)
@@ -1156,8 +1210,12 @@ struct FRHICommandSetComputeShader final : public FRHICommand<FRHICommandSetComp
 	RHI_API void Execute(FRHICommandListBase& CmdList);
 };
 
+struct FRHICommandSetComputePipelineStateString
+{
+	static const TCHAR* TStr() { return TEXT("FRHICommandSetComputePipelineState"); }
+};
 template<ECmdList CmdListType>
-struct FRHICommandSetComputePipelineState final : public FRHICommand<FRHICommandSetComputePipelineState<CmdListType>>
+struct FRHICommandSetComputePipelineState final : public FRHICommand<FRHICommandSetComputePipelineState<CmdListType>, FRHICommandSetComputePipelineStateString>
 {
 	FComputePipelineState* ComputePipelineState;
 	FORCEINLINE_DEBUGGABLE FRHICommandSetComputePipelineState(FComputePipelineState* InComputePipelineState)
@@ -1167,7 +1225,7 @@ struct FRHICommandSetComputePipelineState final : public FRHICommand<FRHICommand
 	RHI_API void Execute(FRHICommandListBase& CmdList);
 };
 
-struct FRHICommandSetGraphicsPipelineState final : public FRHICommand<FRHICommandSetGraphicsPipelineState>
+FRHICOMMAND_MACRO(FRHICommandSetGraphicsPipelineState)
 {
 	FGraphicsPipelineState* GraphicsPipelineState;
 	FORCEINLINE_DEBUGGABLE FRHICommandSetGraphicsPipelineState(FGraphicsPipelineState* InGraphicsPipelineState)
@@ -1177,8 +1235,12 @@ struct FRHICommandSetGraphicsPipelineState final : public FRHICommand<FRHIComman
 	RHI_API void Execute(FRHICommandListBase& CmdList);
 };
 
+struct FRHICommandDispatchComputeShaderString
+{
+	static const TCHAR* TStr() { return TEXT("FRHICommandDispatchComputeShader"); }
+};
 template<ECmdList CmdListType>
-struct FRHICommandDispatchComputeShader final : public FRHICommand<FRHICommandDispatchComputeShader<CmdListType>>
+struct FRHICommandDispatchComputeShader final : public FRHICommand<FRHICommandDispatchComputeShader<CmdListType>, FRHICommandDispatchComputeShaderString>
 {
 	uint32 ThreadGroupCountX;
 	uint32 ThreadGroupCountY;
@@ -1192,8 +1254,12 @@ struct FRHICommandDispatchComputeShader final : public FRHICommand<FRHICommandDi
 	RHI_API void Execute(FRHICommandListBase& CmdList);
 };
 
+struct FRHICommandDispatchIndirectComputeShaderString
+{
+	static const TCHAR* TStr() { return TEXT("FRHICommandDispatchIndirectComputeShader"); }
+};
 template<ECmdList CmdListType>
-struct FRHICommandDispatchIndirectComputeShader final : public FRHICommand<FRHICommandDispatchIndirectComputeShader<CmdListType>>
+struct FRHICommandDispatchIndirectComputeShader final : public FRHICommand<FRHICommandDispatchIndirectComputeShader<CmdListType>, FRHICommandDispatchIndirectComputeShaderString>
 {
 	FRHIVertexBuffer* ArgumentBuffer;
 	uint32 ArgumentOffset;
@@ -1205,7 +1271,7 @@ struct FRHICommandDispatchIndirectComputeShader final : public FRHICommand<FRHIC
 	RHI_API void Execute(FRHICommandListBase& CmdList);
 };
 
-struct FRHICommandAutomaticCacheFlushAfterComputeShader final : public FRHICommand<FRHICommandAutomaticCacheFlushAfterComputeShader>
+FRHICOMMAND_MACRO(FRHICommandAutomaticCacheFlushAfterComputeShader)
 {
 	bool bEnable;
 	FORCEINLINE_DEBUGGABLE FRHICommandAutomaticCacheFlushAfterComputeShader(bool InbEnable)
@@ -1215,12 +1281,12 @@ struct FRHICommandAutomaticCacheFlushAfterComputeShader final : public FRHIComma
 	RHI_API void Execute(FRHICommandListBase& CmdList);
 };
 
-struct FRHICommandFlushComputeShaderCache final : public FRHICommand<FRHICommandFlushComputeShaderCache>
+FRHICOMMAND_MACRO(FRHICommandFlushComputeShaderCache)
 {
 	RHI_API void Execute(FRHICommandListBase& CmdList);
 };
 
-struct FRHICommandDrawPrimitiveIndirect final : public FRHICommand<FRHICommandDrawPrimitiveIndirect>
+FRHICOMMAND_MACRO(FRHICommandDrawPrimitiveIndirect)
 {
 	FRHIVertexBuffer* ArgumentBuffer;
 	uint32 ArgumentOffset;
@@ -1232,7 +1298,7 @@ struct FRHICommandDrawPrimitiveIndirect final : public FRHICommand<FRHICommandDr
 	RHI_API void Execute(FRHICommandListBase& CmdList);
 };
 
-struct FRHICommandDrawIndexedIndirect final : public FRHICommand<FRHICommandDrawIndexedIndirect>
+FRHICOMMAND_MACRO(FRHICommandDrawIndexedIndirect)
 {
 	FRHIIndexBuffer* IndexBufferRHI;
 	FRHIStructuredBuffer* ArgumentsBufferRHI;
@@ -1249,7 +1315,7 @@ struct FRHICommandDrawIndexedIndirect final : public FRHICommand<FRHICommandDraw
 	RHI_API void Execute(FRHICommandListBase& CmdList);
 };
 
-struct FRHICommandDrawIndexedPrimitiveIndirect final : public FRHICommand<FRHICommandDrawIndexedPrimitiveIndirect>
+FRHICOMMAND_MACRO(FRHICommandDrawIndexedPrimitiveIndirect)
 {
 	FRHIIndexBuffer* IndexBuffer;
 	FRHIVertexBuffer* ArgumentsBuffer;
@@ -1264,7 +1330,7 @@ struct FRHICommandDrawIndexedPrimitiveIndirect final : public FRHICommand<FRHICo
 	RHI_API void Execute(FRHICommandListBase& CmdList);
 };
 
-struct FRHICommandSetDepthBounds final : public FRHICommand<FRHICommandSetDepthBounds>
+FRHICOMMAND_MACRO(FRHICommandSetDepthBounds)
 {
 	float MinDepth;
 	float MaxDepth;
@@ -1277,7 +1343,7 @@ struct FRHICommandSetDepthBounds final : public FRHICommand<FRHICommandSetDepthB
 	RHI_API void Execute(FRHICommandListBase& CmdList);
 };
 
-struct FRHICommandClearTinyUAV final : public FRHICommand<FRHICommandClearTinyUAV>
+FRHICOMMAND_MACRO(FRHICommandClearTinyUAV)
 {
 	FRHIUnorderedAccessView* UnorderedAccessViewRHI;
 	uint32 Values[4];
@@ -1293,7 +1359,7 @@ struct FRHICommandClearTinyUAV final : public FRHICommand<FRHICommandClearTinyUA
 	RHI_API void Execute(FRHICommandListBase& CmdList);
 };
 
-struct FRHICommandCopyToResolveTarget final : public FRHICommand<FRHICommandCopyToResolveTarget>
+FRHICOMMAND_MACRO(FRHICommandCopyToResolveTarget)
 {
 	FResolveParams ResolveParams;
 	FRHITexture* SourceTexture;
@@ -1312,7 +1378,7 @@ struct FRHICommandCopyToResolveTarget final : public FRHICommand<FRHICommandCopy
 	RHI_API void Execute(FRHICommandListBase& CmdList);
 };
 
-struct FRHICommandCopyTexture final : public FRHICommand<FRHICommandCopyTexture>
+FRHICOMMAND_MACRO(FRHICommandCopyTexture)
 {
 	FRHICopyTextureInfo CopyInfo;
 	FRHITexture* SourceTexture;
@@ -1331,7 +1397,7 @@ struct FRHICommandCopyTexture final : public FRHICommand<FRHICommandCopyTexture>
 	RHI_API void Execute(FRHICommandListBase& CmdList);
 };
 
-struct FRHICommandTransitionTextures final : public FRHICommand<FRHICommandTransitionTextures>
+FRHICOMMAND_MACRO(FRHICommandTransitionTextures)
 {
 	int32 NumTextures;
 	FRHITexture** Textures; // Pointer to an array of textures, allocated inline with the command list
@@ -1345,7 +1411,7 @@ struct FRHICommandTransitionTextures final : public FRHICommand<FRHICommandTrans
 	RHI_API void Execute(FRHICommandListBase& CmdList);
 };
 
-struct FRHICommandTransitionTexturesArray final : public FRHICommand<FRHICommandTransitionTexturesArray>
+FRHICOMMAND_MACRO(FRHICommandTransitionTexturesArray)
 {	
 	TArray<FRHITexture*>& Textures;
 	EResourceTransitionAccess TransitionType;
@@ -1357,8 +1423,12 @@ struct FRHICommandTransitionTexturesArray final : public FRHICommand<FRHICommand
 	RHI_API void Execute(FRHICommandListBase& CmdList);
 };
 
+struct FRHICommandTransitionUAVsString
+{
+	static const TCHAR* TStr() { return TEXT("FRHICommandTransitionUAVs"); }
+};
 template<ECmdList CmdListType>
-struct FRHICommandTransitionUAVs final : public FRHICommand<FRHICommandTransitionUAVs<CmdListType>>
+struct FRHICommandTransitionUAVs final : public FRHICommand<FRHICommandTransitionUAVs<CmdListType>, FRHICommandTransitionUAVsString>
 {
 	int32 NumUAVs;
 	FRHIUnorderedAccessView** UAVs; // Pointer to an array of UAVs, allocated inline with the command list
@@ -1377,8 +1447,12 @@ struct FRHICommandTransitionUAVs final : public FRHICommand<FRHICommandTransitio
 	RHI_API void Execute(FRHICommandListBase& CmdList);
 };
 
+struct FRHICommandSetAsyncComputeBudgetString
+{
+	static const TCHAR* TStr() { return TEXT("FRHICommandSetAsyncComputeBudget"); }
+};
 template<ECmdList CmdListType>
-struct FRHICommandSetAsyncComputeBudget final : public FRHICommand<FRHICommandSetAsyncComputeBudget<CmdListType>>
+struct FRHICommandSetAsyncComputeBudget final : public FRHICommand<FRHICommandSetAsyncComputeBudget<CmdListType>, FRHICommandSetAsyncComputeBudgetString>
 {
 	EAsyncComputeBudget Budget;
 
@@ -1389,8 +1463,12 @@ struct FRHICommandSetAsyncComputeBudget final : public FRHICommand<FRHICommandSe
 	RHI_API void Execute(FRHICommandListBase& CmdList);
 };
 
+struct FRHICommandWaitComputeFenceString
+{
+	static const TCHAR* TStr() { return TEXT("FRHICommandWaitComputeFence"); }
+};
 template<ECmdList CmdListType>
-struct FRHICommandWaitComputeFence final : public FRHICommand<FRHICommandWaitComputeFence<CmdListType>>
+struct FRHICommandWaitComputeFence final : public FRHICommand<FRHICommandWaitComputeFence<CmdListType>, FRHICommandWaitComputeFenceString>
 {
 	FRHIComputeFence* WaitFence;
 
@@ -1401,8 +1479,12 @@ struct FRHICommandWaitComputeFence final : public FRHICommand<FRHICommandWaitCom
 	RHI_API void Execute(FRHICommandListBase& CmdList);
 };
 
+struct FRHICommandCopyToStagingBufferString
+{
+	static const TCHAR* TStr() { return TEXT("FRHICommandCopyToStagingBuffer"); }
+};
 template<ECmdList CmdListType>
-struct FRHICommandCopyToStagingBuffer final : public FRHICommand<FRHICommandCopyToStagingBuffer<CmdListType>>
+struct FRHICommandCopyToStagingBuffer final : public FRHICommand<FRHICommandCopyToStagingBuffer<CmdListType>, FRHICommandCopyToStagingBufferString>
 {
 	FRHIVertexBuffer* SourceBuffer;
 	FRHIStagingBuffer* DestinationStagingBuffer;
@@ -1419,8 +1501,12 @@ struct FRHICommandCopyToStagingBuffer final : public FRHICommand<FRHICommandCopy
 	RHI_API void Execute(FRHICommandListBase& CmdList);
 };
 
+struct FRHICommandWriteGPUFenceString
+{
+	static const TCHAR* TStr() { return TEXT("FRHICommandWriteGPUFence"); }
+};
 template<ECmdList CmdListType>
-struct FRHICommandWriteGPUFence final : public FRHICommand<FRHICommandWriteGPUFence<CmdListType>>
+struct FRHICommandWriteGPUFence final : public FRHICommand<FRHICommandWriteGPUFence<CmdListType>, FRHICommandWriteGPUFenceString>
 {
 	FRHIGPUFence* Fence;
 
@@ -1431,7 +1517,7 @@ struct FRHICommandWriteGPUFence final : public FRHICommand<FRHICommandWriteGPUFe
 	RHI_API void Execute(FRHICommandListBase& CmdList);
 };
 
-struct FRHICommandClearColorTexture final : public FRHICommand<FRHICommandClearColorTexture>
+FRHICOMMAND_MACRO(FRHICommandClearColorTexture)
 {
 	FRHITexture* Texture;
 	FLinearColor Color;
@@ -1447,7 +1533,7 @@ struct FRHICommandClearColorTexture final : public FRHICommand<FRHICommandClearC
 	RHI_API void Execute(FRHICommandListBase& CmdList);
 };
 
-struct FRHICommandClearDepthStencilTexture final : public FRHICommand<FRHICommandClearDepthStencilTexture>
+FRHICOMMAND_MACRO(FRHICommandClearDepthStencilTexture)
 {
 	FRHITexture* Texture;
 	float Depth;
@@ -1469,7 +1555,7 @@ struct FRHICommandClearDepthStencilTexture final : public FRHICommand<FRHIComman
 	RHI_API void Execute(FRHICommandListBase& CmdList);
 };
 
-struct FRHICommandClearColorTextures final : public FRHICommand<FRHICommandClearColorTextures>
+FRHICOMMAND_MACRO(FRHICommandClearColorTextures)
 {
 	FLinearColor ColorArray[MaxSimultaneousRenderTargets];
 	FRHITexture* Textures[MaxSimultaneousRenderTargets];
@@ -1555,7 +1641,7 @@ struct FLocalUniformBuffer
 	}
 };
 
-struct FRHICommandBuildLocalUniformBuffer final : public FRHICommand<FRHICommandBuildLocalUniformBuffer>
+FRHICOMMAND_MACRO(FRHICommandBuildLocalUniformBuffer)
 {
 	FLocalUniformBufferWorkArea WorkArea;
 	FORCEINLINE_DEBUGGABLE FRHICommandBuildLocalUniformBuffer(
@@ -1571,8 +1657,12 @@ struct FRHICommandBuildLocalUniformBuffer final : public FRHICommand<FRHICommand
 	RHI_API void Execute(FRHICommandListBase& CmdList);
 };
 
+struct FRHICommandSetLocalUniformBufferString
+{
+	static const TCHAR* TStr() { return TEXT("FRHICommandSetLocalUniformBuffer"); }
+};
 template <typename TRHIShader>
-struct FRHICommandSetLocalUniformBuffer final : public FRHICommand<FRHICommandSetLocalUniformBuffer<TRHIShader> >
+struct FRHICommandSetLocalUniformBuffer final : public FRHICommand<FRHICommandSetLocalUniformBuffer<TRHIShader>, FRHICommandSetLocalUniformBufferString >
 {
 	TRHIShader* Shader;
 	uint32 BaseIndex;
@@ -1589,7 +1679,7 @@ struct FRHICommandSetLocalUniformBuffer final : public FRHICommand<FRHICommandSe
 	RHI_API void Execute(FRHICommandListBase& CmdList);
 };
 
-struct FRHICommandBeginRenderQuery final : public FRHICommand<FRHICommandBeginRenderQuery>
+FRHICOMMAND_MACRO(FRHICommandBeginRenderQuery)
 {
 	FRHIRenderQuery* RenderQuery;
 
@@ -1600,7 +1690,7 @@ struct FRHICommandBeginRenderQuery final : public FRHICommand<FRHICommandBeginRe
 	RHI_API void Execute(FRHICommandListBase& CmdList);
 };
 
-struct FRHICommandEndRenderQuery final : public FRHICommand<FRHICommandEndRenderQuery>
+FRHICOMMAND_MACRO(FRHICommandEndRenderQuery)
 {
 	FRHIRenderQuery* RenderQuery;
 
@@ -1611,8 +1701,12 @@ struct FRHICommandEndRenderQuery final : public FRHICommand<FRHICommandEndRender
 	RHI_API void Execute(FRHICommandListBase& CmdList);
 };
 
+struct FRHICommandSubmitCommandsHintString
+{
+	static const TCHAR* TStr() { return TEXT("FRHICommandSubmitCommandsHint"); }
+};
 template<ECmdList CmdListType>
-struct FRHICommandSubmitCommandsHint final : public FRHICommand<FRHICommandSubmitCommandsHint<CmdListType>>
+struct FRHICommandSubmitCommandsHint final : public FRHICommand<FRHICommandSubmitCommandsHint<CmdListType>, FRHICommandSubmitCommandsHintString>
 {
 	FORCEINLINE_DEBUGGABLE FRHICommandSubmitCommandsHint()
 	{
@@ -1620,12 +1714,12 @@ struct FRHICommandSubmitCommandsHint final : public FRHICommand<FRHICommandSubmi
 	RHI_API void Execute(FRHICommandListBase& CmdList);
 };
 
-struct FRHICommandPollOcclusionQueries final : public FRHICommand<FRHICommandPollOcclusionQueries>
+FRHICOMMAND_MACRO(FRHICommandPollOcclusionQueries)
 {
 	RHI_API void Execute(FRHICommandListBase& CmdList);
 };
 
-struct FRHICommandBeginScene final : public FRHICommand<FRHICommandBeginScene>
+FRHICOMMAND_MACRO(FRHICommandBeginScene)
 {
 	FORCEINLINE_DEBUGGABLE FRHICommandBeginScene()
 	{
@@ -1633,7 +1727,7 @@ struct FRHICommandBeginScene final : public FRHICommand<FRHICommandBeginScene>
 	RHI_API void Execute(FRHICommandListBase& CmdList);
 };
 
-struct FRHICommandEndScene final : public FRHICommand<FRHICommandEndScene>
+FRHICOMMAND_MACRO(FRHICommandEndScene)
 {
 	FORCEINLINE_DEBUGGABLE FRHICommandEndScene()
 	{
@@ -1641,7 +1735,7 @@ struct FRHICommandEndScene final : public FRHICommand<FRHICommandEndScene>
 	RHI_API void Execute(FRHICommandListBase& CmdList);
 };
 
-struct FRHICommandBeginFrame final : public FRHICommand<FRHICommandBeginFrame>
+FRHICOMMAND_MACRO(FRHICommandBeginFrame)
 {
 	FORCEINLINE_DEBUGGABLE FRHICommandBeginFrame()
 	{
@@ -1649,7 +1743,7 @@ struct FRHICommandBeginFrame final : public FRHICommand<FRHICommandBeginFrame>
 	RHI_API void Execute(FRHICommandListBase& CmdList);
 };
 
-struct FRHICommandEndFrame final : public FRHICommand<FRHICommandEndFrame>
+FRHICOMMAND_MACRO(FRHICommandEndFrame)
 {
 	FORCEINLINE_DEBUGGABLE FRHICommandEndFrame()
 	{
@@ -1657,7 +1751,7 @@ struct FRHICommandEndFrame final : public FRHICommand<FRHICommandEndFrame>
 	RHI_API void Execute(FRHICommandListBase& CmdList);
 };
 
-struct FRHICommandBeginDrawingViewport final : public FRHICommand<FRHICommandBeginDrawingViewport>
+FRHICOMMAND_MACRO(FRHICommandBeginDrawingViewport)
 {
 	FRHIViewport* Viewport;
 	FRHITexture* RenderTargetRHI;
@@ -1670,7 +1764,7 @@ struct FRHICommandBeginDrawingViewport final : public FRHICommand<FRHICommandBeg
 	RHI_API void Execute(FRHICommandListBase& CmdList);
 };
 
-struct FRHICommandEndDrawingViewport final : public FRHICommand<FRHICommandEndDrawingViewport>
+FRHICOMMAND_MACRO(FRHICommandEndDrawingViewport)
 {
 	FRHIViewport* Viewport;
 	bool bPresent;
@@ -1685,8 +1779,12 @@ struct FRHICommandEndDrawingViewport final : public FRHICommand<FRHICommandEndDr
 	RHI_API void Execute(FRHICommandListBase& CmdList);
 };
 
+struct FRHICommandPushEventString
+{
+	static const TCHAR* TStr() { return TEXT("FRHICommandPushEventString"); }
+};
 template<ECmdList CmdListType>
-struct FRHICommandPushEvent final : public FRHICommand<FRHICommandPushEvent<CmdListType>>
+struct FRHICommandPushEvent final : public FRHICommand<FRHICommandPushEvent<CmdListType>, FRHICommandPushEventString>
 {
 	const TCHAR *Name;
 	FColor Color;
@@ -1704,8 +1802,12 @@ struct FRHICommandPushEvent final : public FRHICommand<FRHICommandPushEvent<CmdL
 	};
 };
 
+struct FRHICommandPopEventString
+{
+	static const TCHAR* TStr() { return TEXT("FRHICommandPopEvent"); }
+};
 template<ECmdList CmdListType>
-struct FRHICommandPopEvent final : public FRHICommand<FRHICommandPopEvent<CmdListType>>
+struct FRHICommandPopEvent final : public FRHICommand<FRHICommandPopEvent<CmdListType>, FRHICommandPopEventString>
 {
 	RHI_API void Execute(FRHICommandListBase& CmdList);
 
@@ -1715,12 +1817,12 @@ struct FRHICommandPopEvent final : public FRHICommand<FRHICommandPopEvent<CmdLis
 	};
 };
 
-struct FRHICommandInvalidateCachedState final : public FRHICommand<FRHICommandInvalidateCachedState>
+FRHICOMMAND_MACRO(FRHICommandInvalidateCachedState)
 {
 	RHI_API void Execute(FRHICommandListBase& CmdList);
 };
 
-struct FRHICommandDiscardRenderTargets final : public FRHICommand<FRHICommandDiscardRenderTargets>
+FRHICOMMAND_MACRO(FRHICommandDiscardRenderTargets)
 {
 	uint32 ColorBitMask;
 	bool Depth;
@@ -1736,7 +1838,7 @@ struct FRHICommandDiscardRenderTargets final : public FRHICommand<FRHICommandDis
 	RHI_API void Execute(FRHICommandListBase& CmdList);
 };
 
-struct FRHICommandDebugBreak final : public FRHICommand<FRHICommandDebugBreak>
+FRHICOMMAND_MACRO(FRHICommandDebugBreak)
 {
 	void Execute(FRHICommandListBase& CmdList)
 	{
@@ -1747,7 +1849,7 @@ struct FRHICommandDebugBreak final : public FRHICommand<FRHICommandDebugBreak>
 	}
 };
 
-struct FRHICommandUpdateTextureReference final : public FRHICommand<FRHICommandUpdateTextureReference>
+FRHICOMMAND_MACRO(FRHICommandUpdateTextureReference)
 {
 	FRHITextureReference* TextureRef;
 	FRHITexture* NewTexture;
@@ -1813,7 +1915,7 @@ struct FRHIResourceUpdateInfo
 	void ReleaseRefs();
 };
 
-struct FRHICommandUpdateRHIResources final : public FRHICommand<FRHICommandUpdateRHIResources>
+FRHICOMMAND_MACRO(FRHICommandUpdateRHIResources)
 {
 	FRHIResourceUpdateInfo* UpdateInfos;
 	int32 Num;
@@ -1831,7 +1933,7 @@ struct FRHICommandUpdateRHIResources final : public FRHICommand<FRHICommandUpdat
 };
 
 #if RHI_RAYTRACING
-struct FRHICommandCopyBufferRegion final : public FRHICommand<FRHICommandCopyBufferRegion>
+FRHICOMMAND_MACRO(FRHICommandCopyBufferRegion)
 {
 	FRHIVertexBuffer* DestBuffer;
 	uint64 DstOffset;
@@ -1850,7 +1952,7 @@ struct FRHICommandCopyBufferRegion final : public FRHICommand<FRHICommandCopyBuf
 	RHI_API void Execute(FRHICommandListBase& CmdList);
 };
 
-struct FRHICommandCopyBufferRegions final : public FRHICommand<FRHICommandCopyBufferRegions>
+FRHICOMMAND_MACRO(FRHICommandCopyBufferRegions)
 {
 	const TArrayView<const FCopyBufferRegionParams> Params;
 
@@ -1861,7 +1963,7 @@ struct FRHICommandCopyBufferRegions final : public FRHICommand<FRHICommandCopyBu
 	RHI_API void Execute(FRHICommandListBase& CmdList);
 };
 
-struct FRHICommandBuildAccelerationStructure final : public FRHICommand<FRHICommandBuildAccelerationStructure>
+FRHICOMMAND_MACRO(FRHICommandBuildAccelerationStructure)
 {
 	FRHIRayTracingGeometry* Geometry;
 	FRHIRayTracingScene* Scene;
@@ -1879,7 +1981,7 @@ struct FRHICommandBuildAccelerationStructure final : public FRHICommand<FRHIComm
 	RHI_API void Execute(FRHICommandListBase& CmdList);
 };
 
-struct FRHICommandClearRayTracingBindings final : public FRHICommand<FRHICommandClearRayTracingBindings>
+FRHICOMMAND_MACRO(FRHICommandClearRayTracingBindings)
 {
 	FRHIRayTracingScene* Scene;
 
@@ -1891,7 +1993,7 @@ struct FRHICommandClearRayTracingBindings final : public FRHICommand<FRHICommand
 };
 
 
-struct FRHICommandUpdateAccelerationStructures final : public FRHICommand<FRHICommandUpdateAccelerationStructures>
+FRHICOMMAND_MACRO(FRHICommandUpdateAccelerationStructures)
 {
 	const TArrayView<const FAccelerationStructureUpdateParams> UpdateParams;
 
@@ -1902,7 +2004,7 @@ struct FRHICommandUpdateAccelerationStructures final : public FRHICommand<FRHICo
 	RHI_API void Execute(FRHICommandListBase& CmdList);
 };
 
-struct FRHICommandBuildAccelerationStructures final : public FRHICommand<FRHICommandBuildAccelerationStructures>
+FRHICOMMAND_MACRO(FRHICommandBuildAccelerationStructures)
 {
 	const TArrayView<const FAccelerationStructureUpdateParams> UpdateParams;
 
@@ -1913,7 +2015,7 @@ struct FRHICommandBuildAccelerationStructures final : public FRHICommand<FRHICom
 	RHI_API void Execute(FRHICommandListBase& CmdList);
 };
 
-struct FRHICommandRayTraceOcclusion final : public FRHICommand<FRHICommandRayTraceOcclusion>
+FRHICOMMAND_MACRO(FRHICommandRayTraceOcclusion)
 {
 	FRHIRayTracingScene* Scene;
 	FRHIShaderResourceView* Rays;
@@ -1933,7 +2035,7 @@ struct FRHICommandRayTraceOcclusion final : public FRHICommand<FRHICommandRayTra
 	RHI_API void Execute(FRHICommandListBase& CmdList);
 };
 
-struct FRHICommandRayTraceIntersection final : public FRHICommand<FRHICommandRayTraceIntersection>
+FRHICOMMAND_MACRO(FRHICommandRayTraceIntersection)
 {
 	FRHIRayTracingScene* Scene;
 	FRHIShaderResourceView* Rays;
@@ -1953,7 +2055,7 @@ struct FRHICommandRayTraceIntersection final : public FRHICommand<FRHICommandRay
 	RHI_API void Execute(FRHICommandListBase& CmdList);
 };
 
-struct FRHICommandRayTraceDispatch final : public FRHICommand<FRHICommandRayTraceDispatch>
+FRHICOMMAND_MACRO(FRHICommandRayTraceDispatch)
 {
 	FRayTracingPipelineState* Pipeline;
 	FRHIRayTracingScene* Scene;
@@ -1974,7 +2076,7 @@ struct FRHICommandRayTraceDispatch final : public FRHICommand<FRHICommandRayTrac
 	RHI_API void Execute(FRHICommandListBase& CmdList);
 };
 
-struct FRHICommandSetRayTracingBindings final : public FRHICommand<FRHICommandSetRayTracingBindings>
+FRHICOMMAND_MACRO(FRHICommandSetRayTracingBindings)
 {
 	enum EBindingType
 	{
@@ -3477,6 +3579,14 @@ class RHI_API FRHICommandListImmediate : public FRHICommandList
 
 		void ExecuteAndDestruct(FRHICommandListBase& CmdList, FRHICommandListDebugContext&) override final
 		{
+#if CPUPROFILERTRACE_ENABLED
+			static uint16 __CpuProfilerEventSpecId;
+			if (__CpuProfilerEventSpecId == 0)
+			{
+				__CpuProfilerEventSpecId = FCpuProfilerTrace::OutputEventType(TEXT("TRHILambdaCommand"), CpuProfilerGroup_Default);
+			}
+			FCpuProfilerTrace::FEventScope __CpuProfilerEventScope(__CpuProfilerEventSpecId);
+#endif
 			Lambda(*static_cast<FRHICommandListImmediate*>(&CmdList));
 			Lambda.~LAMBDA();
 		}
diff --git a/Engine/Source/Runtime/RenderCore/Private/RenderingThread.cpp b/Engine/Source/Runtime/RenderCore/Private/RenderingThread.cpp
--- a/Engine/Source/Runtime/RenderCore/Private/RenderingThread.cpp
+++ b/Engine/Source/Runtime/RenderCore/Private/RenderingThread.cpp
@@ -906,7 +906,7 @@ TAutoConsoleVariable<int32> CVarGTSyncType(
 	TEXT(" 2 - Sync the game thread with the GPU swap chain flip (only on supported platforms).\n"),
 	ECVF_Default);
 
-struct FRHISyncFrameCommand final : public FRHICommand<FRHISyncFrameCommand>
+FRHICOMMAND_MACRO(FRHISyncFrameCommand)
 {
 	FGraphEventRef GraphEvent;
 	int32 GTSyncType;
diff --git a/Engine/Source/Runtime/RenderCore/Public/RenderingThread.h b/Engine/Source/Runtime/RenderCore/Public/RenderingThread.h
--- a/Engine/Source/Runtime/RenderCore/Public/RenderingThread.h
+++ b/Engine/Source/Runtime/RenderCore/Public/RenderingThread.h
@@ -186,6 +186,14 @@ public:
 
 	void DoTask(ENamedThreads::Type CurrentThread, const FGraphEventRef& MyCompletionGraphEvent)
 	{
+#if CPUPROFILERTRACE_ENABLED
+		static uint16 __CpuProfilerEventSpecId;
+		if (__CpuProfilerEventSpecId == 0)
+		{
+			__CpuProfilerEventSpecId = FCpuProfilerTrace::OutputEventType(TEXT("TRHILambdaCommand"), CpuProfilerGroup_Default);
+		}
+		FCpuProfilerTrace::FEventScope __CpuProfilerEventScope(__CpuProfilerEventSpecId);
+#endif
 		FRHICommandListImmediate& RHICmdList = GetImmediateCommandList_ForRenderCommand();
 		Lambda(RHICmdList);
 	}
diff --git a/Engine/Source/Runtime/SlateRHIRenderer/Private/Slate3DRenderer.cpp b/Engine/Source/Runtime/SlateRHIRenderer/Private/Slate3DRenderer.cpp
--- a/Engine/Source/Runtime/SlateRHIRenderer/Private/Slate3DRenderer.cpp
+++ b/Engine/Source/Runtime/SlateRHIRenderer/Private/Slate3DRenderer.cpp
@@ -107,8 +107,13 @@ void FSlate3DRenderer::DrawWindow_GameThread(FSlateDrawBuffer& DrawBuffer)
 	}
 }
 
+struct TKeepAliveCommandString
+{
+	static const TCHAR* TStr() { return TEXT("TKeepAliveCommand"); }
+};
+
 template<typename TKeepAliveType>
-struct TKeepAliveCommand final : public FRHICommand < TKeepAliveCommand<TKeepAliveType> >
+struct TKeepAliveCommand final : public FRHICommand < TKeepAliveCommand<TKeepAliveType>, TKeepAliveCommandString >
 {
 	TKeepAliveType Value;
 	
diff --git a/Engine/Source/Runtime/SlateRHIRenderer/Private/SlateRHIRenderer.cpp b/Engine/Source/Runtime/SlateRHIRenderer/Private/SlateRHIRenderer.cpp
--- a/Engine/Source/Runtime/SlateRHIRenderer/Private/SlateRHIRenderer.cpp
+++ b/Engine/Source/Runtime/SlateRHIRenderer/Private/SlateRHIRenderer.cpp
@@ -1478,7 +1478,7 @@ void FSlateRHIRenderer::ClearScenes()
 }
 
 
-struct FClearCachedRenderingDataCommand final : public FRHICommand < FClearCachedRenderingDataCommand >
+FRHICOMMAND_MACRO(FClearCachedRenderingDataCommand)
 {
 public:
 	FClearCachedRenderingDataCommand(FSlateCachedFastPathRenderingData* InCachedRenderingData)
@@ -1496,7 +1496,7 @@ private:
 	FSlateCachedFastPathRenderingData* CachedRenderingData;
 };
 
-struct FClearCachedElementDataCommand final : public FRHICommand < FClearCachedElementDataCommand >
+FRHICOMMAND_MACRO(FClearCachedElementDataCommand)
 {
 public:
 	FClearCachedElementDataCommand(FSlateCachedElementData* InCachedElementData)
diff --git a/Engine/Source/Runtime/SlateRHIRenderer/Private/SlateRHIRenderer.h b/Engine/Source/Runtime/SlateRHIRenderer/Private/SlateRHIRenderer.h
--- a/Engine/Source/Runtime/SlateRHIRenderer/Private/SlateRHIRenderer.h
+++ b/Engine/Source/Runtime/SlateRHIRenderer/Private/SlateRHIRenderer.h
@@ -20,7 +20,7 @@ class ISlateStyle;
 class SWindow;
 struct Rect;
 
-template<typename TCmd> struct FRHICommand;
+template<typename TCmd, typename NameType> struct FRHICommand;
 
 typedef TArray<FRenderThreadUpdateContext, TInlineAllocator<2>> FDeferredUpdateContextList;
 
@@ -289,7 +289,11 @@ private:
 	uint32 ResourceVersion;
 };
 
-struct FSlateEndDrawingWindowsCommand final : public FRHICommand < FSlateEndDrawingWindowsCommand >
+struct FSlateEndDrawingWindowsCommandString
+{
+	static const TCHAR* TStr() { return TEXT("FSlateEndDrawingWindowsCommand"); }
+};
+struct FSlateEndDrawingWindowsCommand final : public FRHICommand < FSlateEndDrawingWindowsCommand, FSlateEndDrawingWindowsCommandString >
 {
 	FSlateRHIRenderingPolicy& Policy;
 	FSlateDrawBuffer* DrawBuffer;
