DR - Fixing scope names

diff --git a/Engine/Source/Runtime/RHI/Public/RHICommandList.h b/Engine/Source/Runtime/RHI/Public/RHICommandList.h
--- a/Engine/Source/Runtime/RHI/Public/RHICommandList.h
+++ b/Engine/Source/Runtime/RHI/Public/RHICommandList.h
@@ -682,7 +682,7 @@ struct FRHICommand : public FRHICommandBase
 		static uint16 __CpuProfilerEventSpecId;
 		if (__CpuProfilerEventSpecId == 0)
 		{
-			__CpuProfilerEventSpecId = FCpuProfilerTrace::OutputEventType(TEXT("TRHILambdaCommand"), CpuProfilerGroup_Default);
+			__CpuProfilerEventSpecId = FCpuProfilerTrace::OutputEventType(NameType::TStr(), CpuProfilerGroup_Default);
 		}
 		FCpuProfilerTrace::FEventScope __CpuProfilerEventScope(__CpuProfilerEventSpecId);
 #endif
diff --git a/Engine/Source/Runtime/RenderCore/Public/RenderingThread.h b/Engine/Source/Runtime/RenderCore/Public/RenderingThread.h
--- a/Engine/Source/Runtime/RenderCore/Public/RenderingThread.h
+++ b/Engine/Source/Runtime/RenderCore/Public/RenderingThread.h
@@ -190,7 +190,7 @@ public:
 		static uint16 __CpuProfilerEventSpecId;
 		if (__CpuProfilerEventSpecId == 0)
 		{
-			__CpuProfilerEventSpecId = FCpuProfilerTrace::OutputEventType(TEXT("TRHILambdaCommand"), CpuProfilerGroup_Default);
+			__CpuProfilerEventSpecId = FCpuProfilerTrace::OutputEventType(TSTR::TStr(), CpuProfilerGroup_Default);
 		}
 		FCpuProfilerTrace::FEventScope __CpuProfilerEventScope(__CpuProfilerEventSpecId);
 #endif
