Removed references to OpenGLDrv.h in build scripts for iOS platform.

diff --git a/Engine/Source/Runtime/ApplicationCore/Private/IOS/IOSPlatformApplicationMisc.cpp b/Engine/Source/Runtime/ApplicationCore/Private/IOS/IOSPlatformApplicationMisc.cpp
--- a/Engine/Source/Runtime/ApplicationCore/Private/IOS/IOSPlatformApplicationMisc.cpp
+++ b/Engine/Source/Runtime/ApplicationCore/Private/IOS/IOSPlatformApplicationMisc.cpp
@@ -148,7 +148,6 @@ EAppReturnType::Type MessageBoxExtImpl( EAppMsgType::Type MsgType, const TCHAR*
 
 void FIOSPlatformApplicationMisc::LoadPreInitModules()
 {
-	FModuleManager::Get().LoadModule(TEXT("OpenGLDrv"));
 	FModuleManager::Get().LoadModule(TEXT("IOSAudio"));
 	FModuleManager::Get().LoadModule(TEXT("AudioMixerAudioUnit"));
 }
diff --git a/Engine/Source/Runtime/Launch/Launch.Build.cs b/Engine/Source/Runtime/Launch/Launch.Build.cs
--- a/Engine/Source/Runtime/Launch/Launch.Build.cs
+++ b/Engine/Source/Runtime/Launch/Launch.Build.cs
@@ -248,12 +248,6 @@ public class Launch : ModuleRules
 				"IOSRuntimeSettings",
 			});
 
-			// no longer build GL for apps requiring iOS 12 or later
-			if (Target.IOSPlatform.RuntimeVersion < 12.0)
-			{
-				PublicFrameworks.Add("OpenGLES");
-				PrivateDependencyModuleNames.Add("OpenGLDrv");
-			}
 			// needed for Metal layer
 			PublicFrameworks.Add("QuartzCore");
 		}
diff --git a/Engine/Source/Runtime/RHI/Private/Apple/AppleDynamicRHI.cpp b/Engine/Source/Runtime/RHI/Private/Apple/AppleDynamicRHI.cpp
--- a/Engine/Source/Runtime/RHI/Private/Apple/AppleDynamicRHI.cpp
+++ b/Engine/Source/Runtime/RHI/Private/Apple/AppleDynamicRHI.cpp
@@ -119,19 +119,6 @@ FDynamicRHI* PlatformCreateDynamicRHI()
 		}
 		FApp::SetGraphicsRHI(TEXT("Metal"));
 	}
-	else if (bForceOpenGL)
-	{
-		DynamicRHIModule = &FModuleManager::LoadModuleChecked<IDynamicRHIModule>(TEXT("OpenGLDrv"));
-		if (Sum == 1)
-		{
-#if PLATFORM_MAC
-			RequestedFeatureLevel = ERHIFeatureLevel::ES3_1;
-#else
-			RequestedFeatureLevel = ERHIFeatureLevel::ES2;
-#endif
-		}
-		FApp::SetGraphicsRHI(TEXT("OpenGL"));
-	}
 	else
 	{
 		FText Title = NSLOCTEXT("AppleDynamicRHI", "OpenGLNotSupportedTitle","OpenGL Not Supported");
