Fixed OpenGL build errors with -allmodules.
Now setting SupportedPlatformsAttribute to desktop platforms/android/html5
Removed some OGL references from AppleDynamicRHI

diff --git a/Engine/Source/Runtime/OpenGLDrv/OpenGLDrv.Build.cs b/Engine/Source/Runtime/OpenGLDrv/OpenGLDrv.Build.cs
--- a/Engine/Source/Runtime/OpenGLDrv/OpenGLDrv.Build.cs
+++ b/Engine/Source/Runtime/OpenGLDrv/OpenGLDrv.Build.cs
@@ -2,6 +2,7 @@
 
 using UnrealBuildTool;
 
+[SupportedPlatformsAttribute(new string[] {"Win32", "Win64", "Mac", "Linux", "HTML5", "Android", "LinuxAArch64"})]
 public class OpenGLDrv : ModuleRules
 {
 	public OpenGLDrv(ReadOnlyTargetRules Target) : base(Target)
diff --git a/Engine/Source/Runtime/RHI/Private/Apple/AppleDynamicRHI.cpp b/Engine/Source/Runtime/RHI/Private/Apple/AppleDynamicRHI.cpp
--- a/Engine/Source/Runtime/RHI/Private/Apple/AppleDynamicRHI.cpp
+++ b/Engine/Source/Runtime/RHI/Private/Apple/AppleDynamicRHI.cpp
@@ -28,23 +28,15 @@ FDynamicRHI* PlatformCreateDynamicRHI()
 
 	bool const bIsMetalSupported = FPlatformMisc::HasPlatformFeature(TEXT("Metal"));
 	bool const bAllowMetal = (GAppleMetalEnabled && bIsMetalSupported);
-	bool const bAllowOpenGL = !GAppleOpenGLDisabled && !PLATFORM_MAC;
 
 #if PLATFORM_MAC
 	bool bForceMetal = bAllowMetal && (FParse::Param(FCommandLine::Get(),TEXT("metal")) || FParse::Param(FCommandLine::Get(),TEXT("metalsm5")));
-	bool bForceOpenGL = false; // OpenGL is no longer supported on Mac
 #else
 	bool bForceMetal = bAllowMetal && (FParse::Param(FCommandLine::Get(),TEXT("metal")) || FParse::Param(FCommandLine::Get(),TEXT("metalmrt")));
-	bool bForceOpenGL = bAllowOpenGL && FParse::Param(FCommandLine::Get(),TEXT("es2"));
 #endif
 
 	ERHIFeatureLevel::Type RequestedFeatureLevel = ERHIFeatureLevel::Num;
-	int32 Sum = ((bForceMetal ? 1 : 0) + (bForceOpenGL ? 1 : 0));
-	if (Sum > 1)
-	{
-		UE_LOG(LogRHI, Fatal,TEXT("-metal, -metalsm5, and -opengl are mutually exclusive options, but more than one was specified on the command-line."));
-	}
-	else if (Sum == 0)
+	 if (bForceMetal)
 	{
 		// Check the list of targeted shader platforms and decide an RHI based off them
 		TArray<FString> TargetedShaderFormats;
@@ -83,7 +75,6 @@ FDynamicRHI* PlatformCreateDynamicRHI()
 			if (bAllowMetal == true || !IsMetalPlatform(TargetedPlatform))
 			{
 				bForceMetal = IsMetalPlatform(TargetedPlatform);
-				bForceOpenGL = IsOpenGLPlatform(TargetedPlatform) && !PLATFORM_MAC;
 				RequestedFeatureLevel = GetMaxSupportedFeatureLevel(TargetedPlatform);
 				break;
 			}
@@ -95,7 +86,6 @@ FDynamicRHI* PlatformCreateDynamicRHI()
 	{
 		DynamicRHIModule = &FModuleManager::LoadModuleChecked<IDynamicRHIModule>(TEXT("MetalRHI"));
 		
-		if (Sum == 1)
 		{
 #if PLATFORM_MAC
 			if (FParse::Param(FCommandLine::Get(),TEXT("metal")))
