Fat fingered a merge earlier. Fixes issue with mac trying to select opengl every time and failing.

diff --git a/Engine/Source/Runtime/RHI/Private/Apple/AppleDynamicRHI.cpp b/Engine/Source/Runtime/RHI/Private/Apple/AppleDynamicRHI.cpp
--- a/Engine/Source/Runtime/RHI/Private/Apple/AppleDynamicRHI.cpp
+++ b/Engine/Source/Runtime/RHI/Private/Apple/AppleDynamicRHI.cpp
@@ -27,16 +27,25 @@ FDynamicRHI* PlatformCreateDynamicRHI()
 	IDynamicRHIModule* DynamicRHIModule = NULL;
 
 	bool const bIsMetalSupported = FPlatformMisc::HasPlatformFeature(TEXT("Metal"));
-	bool const bAllowMetal = (GAppleMetalEnabled && bIsMetalSupported);
-
+	
+	// Must be Metal!
+	if(!bIsMetalSupported)
+	{
+		FText Title = NSLOCTEXT("AppleDynamicRHI", "OpenGLNotSupportedTitle","OpenGL Not Supported");
 #if PLATFORM_MAC
-	bool bForceMetal = bAllowMetal && (FParse::Param(FCommandLine::Get(),TEXT("metal")) || FParse::Param(FCommandLine::Get(),TEXT("metalsm5")));
+		FMessageDialog::Open(EAppMsgType::Ok, NSLOCTEXT("MacPlatformCreateDynamicRHI", "OpenGLNotSupported.", "You must have a Metal compatible graphics card and be running Mac OS X 10.11.6 or later to launch this process."), &Title);
 #else
-	bool bForceMetal = bAllowMetal && (FParse::Param(FCommandLine::Get(),TEXT("metal")) || FParse::Param(FCommandLine::Get(),TEXT("metalmrt")));
+		FMessageDialog::Open(EAppMsgType::Ok, NSLOCTEXT("AppleDynamicRHI", "OpenGLNotSupported.", "You must have a Metal compatible iOS or tvOS device with iOS 8 or later to launch this app."), &Title);
 #endif
+		FPlatformMisc::RequestExit(true);
+	}
+	
+	if (FParse::Param(FCommandLine::Get(),TEXT("opengl")))
+	{
+		UE_LOG(LogRHI, Log, TEXT("OpenGL command line option ignored; Apple platforms only support Metal."));
+	}
 
 	ERHIFeatureLevel::Type RequestedFeatureLevel = ERHIFeatureLevel::Num;
-	 if (bForceMetal)
 	{
 		// Check the list of targeted shader platforms and decide an RHI based off them
 		TArray<FString> TargetedShaderFormats;
@@ -54,17 +63,14 @@ FDynamicRHI* PlatformCreateDynamicRHI()
 #endif
 		}
 		
-		bool bSupportsMetal = false;
-		GConfig->GetBool(TEXT("/Script/IOSRuntimeSettings.IOSRuntimeSettings"), TEXT("bSupportsMetal"), bSupportsMetal, GEngineIni);
-		if (bSupportsMetal)
-		{
 #if PLATFORM_TVOS
-			TargetedShaderFormats.Add(LegacyShaderPlatformToShaderFormat(SP_METAL_TVOS).ToString());
+		TargetedShaderFormats.Add(LegacyShaderPlatformToShaderFormat(SP_METAL_TVOS).ToString());
 #else
-			TargetedShaderFormats.Add(LegacyShaderPlatformToShaderFormat(SP_METAL).ToString());
-#endif
-		}
+		TargetedShaderFormats.Add(LegacyShaderPlatformToShaderFormat(SP_METAL).ToString());
 #endif
+		
+#endif // else branch of PLATFORM_MAC
+		
 		// Metal is not always available, so don't assume that we can use the first platform
 		for (FString Name : TargetedShaderFormats)
 		{
@@ -72,9 +78,8 @@ FDynamicRHI* PlatformCreateDynamicRHI()
 			EShaderPlatform TargetedPlatform = ShaderFormatToLegacyShaderPlatform(ShaderFormatName);
 			
 			// Instead use the first platform that *could* work
-			if (bAllowMetal == true || !IsMetalPlatform(TargetedPlatform))
+			if (IsMetalPlatform(TargetedPlatform))
 			{
-				bForceMetal = IsMetalPlatform(TargetedPlatform);
 				RequestedFeatureLevel = GetMaxSupportedFeatureLevel(TargetedPlatform);
 				break;
 			}
@@ -82,7 +87,6 @@ FDynamicRHI* PlatformCreateDynamicRHI()
 	}
 
 	// Load the dynamic RHI module.
-	if (bForceMetal)
 	{
 		DynamicRHIModule = &FModuleManager::LoadModuleChecked<IDynamicRHIModule>(TEXT("MetalRHI"));
 		
@@ -109,16 +113,6 @@ FDynamicRHI* PlatformCreateDynamicRHI()
 		}
 		FApp::SetGraphicsRHI(TEXT("Metal"));
 	}
-	else
-	{
-		FText Title = NSLOCTEXT("AppleDynamicRHI", "OpenGLNotSupportedTitle","OpenGL Not Supported");
-#if PLATFORM_MAC
-		FMessageDialog::Open(EAppMsgType::Ok, NSLOCTEXT("MacPlatformCreateDynamicRHI", "OpenGLNotSupported.", "You must have a Metal compatible graphics card and be running Mac OS X 10.11.6 or later to launch this process."), &Title);
-#else
-		FMessageDialog::Open(EAppMsgType::Ok, NSLOCTEXT("AppleDynamicRHI", "OpenGLNotSupported.", "You must have a Metal compatible iOS or tvOS device with iOS 8 or later to launch this app."), &Title);
-#endif
-		FPlatformMisc::RequestExit(true);
-	}
 	
 	// Create the dynamic RHI.
 	DynamicRHI = DynamicRHIModule->CreateRHI(RequestedFeatureLevel);
