properly only count new PSOs once

diff --git a/Engine/Source/Runtime/RHI/Private/PipelineFileCache.cpp b/Engine/Source/Runtime/RHI/Private/PipelineFileCache.cpp
--- a/Engine/Source/Runtime/RHI/Private/PipelineFileCache.cpp
+++ b/Engine/Source/Runtime/RHI/Private/PipelineFileCache.cpp
@@ -136,6 +136,7 @@ TMap<uint32, FPSOUsageData> FPipelineFileCache::RunTimeToPSOUsage;
 TMap<uint32, FPSOUsageData> FPipelineFileCache::NewPSOUsage;
 TMap<uint32, FPipelineStateStats*> FPipelineFileCache::Stats;
 TSet<FPipelineCacheFileFormatPSO> FPipelineFileCache::NewPSOs;
+TSet<uint32> FPipelineFileCache::NewPSOHashes;
 uint32 FPipelineFileCache::NumNewPSOs;
 FPipelineFileCache::PSOOrder FPipelineFileCache::RequestedOrder = FPipelineFileCache::PSOOrder::MostToLeastUsed;
 bool FPipelineFileCache::FileCacheEnabled = false;
@@ -2587,6 +2588,7 @@ void FPipelineFileCache::Shutdown()
 		}
 		Stats.Empty();
 		NewPSOs.Empty();
+		NewPSOHashes.Empty();
         NumNewPSOs = 0;
 		
 		FileCacheEnabled = false;
@@ -2613,6 +2615,7 @@ bool FPipelineFileCache::OpenPipelineFileCache(FString const& Name, EShaderPlatf
 			
 			// File Cache now exists - these caches should be empty for this file otherwise will have false positives from any previous file caching - if not something has been caching when it should not be
 			check(NewPSOs.Num() == 0);
+			check(NewPSOHashes.Num() == 0);
 			check(RunTimeToPSOUsage.Num() == 0);
 		}
 	}
@@ -2681,6 +2684,7 @@ void FPipelineFileCache::ClosePipelineFileCache()
 			RunTimeToPSOUsage.Empty();
 			NewPSOUsage.Empty();
 			NewPSOs.Empty();
+			NewPSOHashes.Empty();
             NumNewPSOs = 0;
 			
 			SET_MEMORY_STAT(STAT_NewCachedPSOMemory, 0);
@@ -2722,8 +2726,8 @@ void FPipelineFileCache::CacheGraphicsPSO(uint32 RunTimeHash, FGraphicsPipelineS
 					
 					if (!FileCache->IsPSOEntryCached(NewEntry, &CurrentUsageData))
 					{
-						bool bActuallyNewPSO = true;
-						if (IsOpenGLPlatform(GMaxRHIShaderPlatform)) // OpenGL is a BSS platform and so we don't report BSS matches as missing.
+						bool bActuallyNewPSO = !NewPSOHashes.Contains(PSOHash);
+						if (bActuallyNewPSO && IsOpenGLPlatform(GMaxRHIShaderPlatform)) // OpenGL is a BSS platform and so we don't report BSS matches as missing.
 						{
 							bActuallyNewPSO = !FileCache->IsBSSEquivalentPSOEntryCached(NewEntry);
 						}
@@ -2740,7 +2744,8 @@ void FPipelineFileCache::CacheGraphicsPSO(uint32 RunTimeHash, FGraphicsPipelineS
 								NewPSOs.Add(NewEntry);
 								INC_MEMORY_STAT_BY(STAT_NewCachedPSOMemory, sizeof(FPipelineCacheFileFormatPSO) + sizeof(uint32) + sizeof(uint32));
 							}
-							
+							NewPSOHashes.Add(PSOHash);
+
 							NumNewPSOs++;
 							INC_DWORD_STAT(STAT_NewGraphicsPipelineStateCount);
 							INC_DWORD_STAT(STAT_TotalGraphicsPipelineStateCount);
@@ -2797,26 +2802,32 @@ void FPipelineFileCache::CacheComputePSO(uint32 RunTimeHash, FRHIComputeShader c
 					
 					if (!FileCache->IsPSOEntryCached(NewEntry, &CurrentUsageData))
 					{
-						CSV_EVENT(PSO, TEXT("Encountered new compute PSO"));
-						UE_LOG(LogRHI, Warning, TEXT("Encountered a new compute PSO: %u"), PSOHash);
-						if (GPSOFileCachePrintNewPSODescriptors > 0)
-						{
-							UE_LOG(LogRHI, Warning, TEXT("New compute PSO (%u) Description: %s"), PSOHash, *NewEntry.ComputeDesc.ComputeShader.ToString());
-						}
-
-						if (LogPSOtoFileCache())
-						{
-							NewPSOs.Add(NewEntry);
-							INC_MEMORY_STAT_BY(STAT_NewCachedPSOMemory, sizeof(FPipelineCacheFileFormatPSO) + sizeof(uint32) + sizeof(uint32));
-						}
-
-						NumNewPSOs++;
-						INC_DWORD_STAT(STAT_NewComputePipelineStateCount);
-						INC_DWORD_STAT(STAT_TotalComputePipelineStateCount);
-						
-						if (ReportNewPSOs() && PSOLoggedEvent.IsBound())
+						bool bActuallyNewPSO = !NewPSOHashes.Contains(PSOHash);
+						if (bActuallyNewPSO)
 						{
-							PSOLoggedEvent.Broadcast(NewEntry);
+							CSV_EVENT(PSO, TEXT("Encountered new compute PSO"));
+							UE_LOG(LogRHI, Warning, TEXT("Encountered a new compute PSO: %u"), PSOHash);
+							if (GPSOFileCachePrintNewPSODescriptors > 0)
+							{
+								UE_LOG(LogRHI, Warning, TEXT("New compute PSO (%u) Description: %s"), PSOHash, *NewEntry.ComputeDesc.ComputeShader.ToString());
+							}
+							
+							if (LogPSOtoFileCache())
+							{
+								NewPSOs.Add(NewEntry);
+								INC_MEMORY_STAT_BY(STAT_NewCachedPSOMemory, sizeof(FPipelineCacheFileFormatPSO) + sizeof(uint32) + sizeof(uint32));
+							}
+							
+							NewPSOHashes.Add(PSOHash);
+							
+							NumNewPSOs++;
+							INC_DWORD_STAT(STAT_NewComputePipelineStateCount);
+							INC_DWORD_STAT(STAT_TotalComputePipelineStateCount);
+							
+							if (ReportNewPSOs() && PSOLoggedEvent.IsBound())
+							{
+								PSOLoggedEvent.Broadcast(NewEntry);
+							}
 						}
 					}
 					
diff --git a/Engine/Source/Runtime/RHI/Public/PipelineFileCache.h b/Engine/Source/Runtime/RHI/Public/PipelineFileCache.h
--- a/Engine/Source/Runtime/RHI/Public/PipelineFileCache.h
+++ b/Engine/Source/Runtime/RHI/Public/PipelineFileCache.h
@@ -336,6 +336,7 @@ private:
 	static TMap<uint32, FPSOUsageData> NewPSOUsage;				// For mask or engine updates - Merged + Saved (Our internal PSO hash to latest usage data) - temp working scratch, only holds updates since last "save" so is not the authority on state
 	static TMap<uint32, FPipelineStateStats*> Stats;
 	static TSet<FPipelineCacheFileFormatPSO> NewPSOs;
+	static TSet<uint32> NewPSOHashes;
     static uint32 NumNewPSOs;
 	static PSOOrder RequestedOrder;
 	static bool FileCacheEnabled;
