Fixup minor error discovered by the AddressSanitizer: stack-use-after-scope.

#rnx

diff --git a/Engine/Source/Runtime/RHI/Private/PipelineFileCache.cpp b/Engine/Source/Runtime/RHI/Private/PipelineFileCache.cpp
--- a/Engine/Source/Runtime/RHI/Private/PipelineFileCache.cpp
+++ b/Engine/Source/Runtime/RHI/Private/PipelineFileCache.cpp
@@ -2475,8 +2475,8 @@ void FPipelineFileCache::PreCompileComplete()
 	static FString PrivateWritePathBase = FString([NSSearchPathForDirectoriesInDomains(NSLibraryDirectory, NSUserDomainMask, YES) objectAtIndex:0]) + TEXT("/");
 	FString Result = PrivateWritePathBase + FString([NSString stringWithFormat:@"/Caches/%@/com.apple.metal/usecache.txt", [NSBundle mainBundle].bundleIdentifier]);
 	int32 Handle = open(TCHAR_TO_UTF8(*Result), O_CREAT, S_IRUSR | S_IWUSR | S_IRGRP | S_IWGRP | S_IROTH | S_IWOTH);
-	char* Version = TCHAR_TO_ANSI(*FEngineVersion::Current().ToString());
-	write(Handle, Version, strlen(Version));
+	FString Version = FEngineVersion::Current().ToString();
+	write(Handle, TCHAR_TO_ANSI(*Version), Version.Len());
 	close(Handle);
 #endif
 }
