Added CVar to control fine-grained RHI command profiling events. Use r.RHICmdTraceEvents=1 to enable tracing events (off by default).

(ushell-p4-cherrypick of 9678980 by [REDACTED_PERSON])

diff --git a/Engine/Source/Runtime/RHI/Private/RHICommandList.cpp b/Engine/Source/Runtime/RHI/Private/RHICommandList.cpp
--- a/Engine/Source/Runtime/RHI/Private/RHICommandList.cpp
+++ b/Engine/Source/Runtime/RHI/Private/RHICommandList.cpp
@@ -117,6 +117,12 @@ static TAutoConsoleVariable<int32> CVarRHICmdMinCmdlistSizeForParallelTranslate(
 	32,
 	TEXT("In kilobytes. Cmdlists are merged into one parallel translate until we have at least this much memory to process. For a given pass, we won't do more translates than we have task threads. Only relevant if r.RHICmdBalanceTranslatesAfterTasks is on."));
 
+RHI_API int32 GRHICmdTraceEvents = 0;
+static FAutoConsoleVariableRef CVarRHICmdTraceEvents(
+	TEXT("r.RHICmdTraceEvents"),
+	GRHICmdTraceEvents,
+	TEXT("Enable tracing profiler events for every RHI command. (default = 0)")
+);
 
 bool GUseRHIThread_InternalUseOnly = false;
 bool GUseRHITaskThreads_InternalUseOnly = false;
diff --git a/Engine/Source/Runtime/RHI/Public/RHICommandList.h b/Engine/Source/Runtime/RHI/Public/RHICommandList.h
--- a/Engine/Source/Runtime/RHI/Public/RHICommandList.h
+++ b/Engine/Source/Runtime/RHI/Public/RHICommandList.h
@@ -703,8 +703,21 @@ struct FRHICommand : public FRHICommandBase
 		{
 			__CpuProfilerEventSpecId = FCpuProfilerTrace::OutputEventType(NameType::TStr(), CpuProfilerGroup_Default);
 		}
-		FCpuProfilerTrace::FEventScope __CpuProfilerEventScope(__CpuProfilerEventSpecId);
-#endif
+
+		extern RHI_API int32 GRHICmdTraceEvents;
+		struct FConditionalTraceScope
+		{
+			FConditionalTraceScope(const uint16 InSpecId) : SpecId(InSpecId)
+			{
+				if (SpecId) FCpuProfilerTrace::OutputBeginEvent(SpecId);
+			}
+			~FConditionalTraceScope()
+			{
+				if (SpecId) FCpuProfilerTrace::OutputEndEvent();
+			}
+			const uint16 SpecId;
+		} TraceScope(GRHICmdTraceEvents ? __CpuProfilerEventSpecId : 0);
+#endif // CPUPROFILERTRACE_ENABLED
 
 		TCmd *ThisCmd = static_cast<TCmd*>(this);
 #if RHI_COMMAND_LIST_DEBUG_TRACES
