Re-add Bypass() condition in FRHICommandList::SetGPUMask. This was accidentally removed when the MGPU branch was merged to Dev-Rendering.
Fixes the right eye showing up black in MGPU with -emulatestereo if the RHI thread is disabled.
#rnx

diff --git a/Engine/Source/Runtime/RHI/Public/RHICommandList.h b/Engine/Source/Runtime/RHI/Public/RHICommandList.h
--- a/Engine/Source/Runtime/RHI/Public/RHICommandList.h
+++ b/Engine/Source/Runtime/RHI/Public/RHICommandList.h
@@ -3391,7 +3391,7 @@ public:
 		if (GPUMask != InGPUMask)
 		{
 			GPUMask = InGPUMask;
-			if (HasCommands())
+			if (HasCommands() || Bypass())
 			{
 				SetGPUMaskOnContext();
 			}
