Updating "Not Supported" message title.
We'll keep the key as OpenGLNotSupportedTitle for now.

diff --git a/Engine/Source/Runtime/RHI/Private/Apple/AppleDynamicRHI.cpp b/Engine/Source/Runtime/RHI/Private/Apple/AppleDynamicRHI.cpp
--- a/Engine/Source/Runtime/RHI/Private/Apple/AppleDynamicRHI.cpp
+++ b/Engine/Source/Runtime/RHI/Private/Apple/AppleDynamicRHI.cpp
@@ -31,7 +31,7 @@ FDynamicRHI* PlatformCreateDynamicRHI()
 	// Must be Metal!
 	if(!bIsMetalSupported)
 	{
-		FText Title = NSLOCTEXT("AppleDynamicRHI", "OpenGLNotSupportedTitle","OpenGL Not Supported");
+		FText Title = NSLOCTEXT("AppleDynamicRHI", "OpenGLNotSupportedTitle","Metal Not Supported");
 #if PLATFORM_MAC
 		FMessageDialog::Open(EAppMsgType::Ok, NSLOCTEXT("MacPlatformCreateDynamicRHI", "OpenGLNotSupported.", "You must have a Metal compatible graphics card and be running Mac OS X 10.11.6 or later to launch this process."), &Title);
 #else
