Enabled ReportPSO for Windows D3D12 so we get internal playtest and external anaytics on PSO coverage.

diff --git a/Engine/Source/Runtime/RHI/Private/Windows/WindowsDynamicRHI.cpp b/Engine/Source/Runtime/RHI/Private/Windows/WindowsDynamicRHI.cpp
--- a/Engine/Source/Runtime/RHI/Private/Windows/WindowsDynamicRHI.cpp
+++ b/Engine/Source/Runtime/RHI/Private/Windows/WindowsDynamicRHI.cpp
@@ -184,6 +184,9 @@ static IDynamicRHIModule* LoadDynamicRHIModule(ERHIFeatureLevel::Type& DesiredFe
 		// Enable -psocache by default on DX12. Since RHI is selected at runtime we can't set this at compile time with PIPELINE_CACHE_DEFAULT_ENABLED.
 		auto PSOFileCacheEnabledCVar = IConsoleManager::Get().FindTConsoleVariableDataInt(TEXT("r.ShaderPipelineCache.Enabled"));
 		*PSOFileCacheEnabledCVar = 1;
+
+		auto PSOFileCacheReportCVar = IConsoleManager::Get().FindTConsoleVariableDataInt(TEXT("r.ShaderPipelineCache.ReportPSO"));
+		*PSOFileCacheReportCVar = 1;
 #endif
 
 		if (!DynamicRHIModule || !DynamicRHIModule->IsSupported())
