Don't set the output file GUID until all header loading has been completed successfully.  Also add AsyncFile handle access guards and extra logging.

diff --git a/Engine/Source/Runtime/RHI/Private/PipelineFileCache.cpp b/Engine/Source/Runtime/RHI/Private/PipelineFileCache.cpp
--- a/Engine/Source/Runtime/RHI/Private/PipelineFileCache.cpp
+++ b/Engine/Source/Runtime/RHI/Private/PipelineFileCache.cpp
@@ -1293,7 +1293,8 @@ public:
 	
 	bool OpenPipelineFileCache(FString const& FilePath, FGuid& Guid, TSharedPtr<IAsyncReadFileHandle, ESPMode::ThreadSafe>& Handle, FPipelineCacheFileFormatTOC& Content)
 	{
-		bool bOK = false;
+		bool bSuccess = false;
+
 		FArchive* FileReader = IFileManager::Get().CreateFileReader(*FilePath);
 		if (FileReader)
 		{
@@ -1308,19 +1309,18 @@ public:
 				UE_LOG(LogRHI, Log, TEXT("FPipelineCacheFile Header Game Version: %d"), Header.GameVersion);
 				UE_LOG(LogRHI, Log, TEXT("FPipelineCacheFile Header Engine Data Version: %d"), Header.Version);
 				UE_LOG(LogRHI, Log, TEXT("FPipelineCacheFile Header TOC Offset: %llu"), Header.TableOffset);
-				UE_LOG(LogRHI, Log, TEXT("FPipelineCacheFile File Size: %lldBytes"), FileReader->TotalSize());
+				UE_LOG(LogRHI, Log, TEXT("FPipelineCacheFile File Size: %lld Bytes"), FileReader->TotalSize());
 				
 				if(Header.TableOffset < (uint64)FileReader->TotalSize())
 				{
-					TOCOffset = Header.TableOffset;
-					Guid = Header.Guid;
 					FileReader->Seek(Header.TableOffset);
-					
 					*FileReader << Content;
-					bOK = !FileReader->IsError();
+					
+					// FPipelineCacheFileFormatTOC archive read can set the FArchive to error on failure
+					bSuccess = !FileReader->IsError();
 				}
 				
-				if(!bOK)
+				if(!bSuccess)
 				{
 					UE_LOG(LogRHI, Log, TEXT("FPipelineCacheFile: %s is corrupt reading TOC"), *FilePath);
 				}
@@ -1328,14 +1328,27 @@ public:
 			
 			if(!FileReader->Close())
 			{
-				bOK = false;
+				bSuccess = false;
 			}
+			
 			delete FileReader;
+			FileReader = nullptr;
 			
-			if(bOK)
+			if(bSuccess)
 			{
-				UE_LOG(LogRHI, Log, TEXT("Opened FPipelineCacheFile: %s (GUID: %s) with %d entries."), *FilePath, *Guid.ToString(), Content.MetaData.Num());
 				Handle = MakeShareable(FPlatformFileManager::Get().GetPlatformFile().OpenAsyncRead(*FilePath));
+				if(Handle.IsValid())
+				{
+					UE_LOG(LogRHI, Log, TEXT("Opened FPipelineCacheFile: %s (GUID: %s) with %d entries."), *FilePath, *Guid.ToString(), Content.MetaData.Num());
+					
+					Guid = Header.Guid;
+					TOCOffset = Header.TableOffset;
+				}
+				else
+				{
+					UE_LOG(LogRHI, Log, TEXT("Filed to create async read file handle to FPipelineCacheFile: %s (GUID: %s)"), *FilePath, *Guid.ToString());
+					bSuccess = false;
+				}
 			}
 		}
 		else
@@ -1343,7 +1356,7 @@ public:
 			UE_LOG(LogRHI, Log, TEXT("Could not open FPipelineCacheFile: %s"), *FilePath);
 		}
 		
-		return bOK;
+		return bSuccess;
 	}
 	
     bool ShouldDeleteExistingUserCache()
@@ -2328,7 +2341,7 @@ public:
 			if (Meta.FileGuid == GameFileGuid)
 			{
                 FPipelineCacheFileFormatPSOMetaData const* GameMeta = GameTOC.MetaData.Find(Entry->Hash);
-                if (GameMeta)
+                if (GameMeta && GameAsyncFileHandle.IsValid())
                 {
                     Entry->Data.SetNum(GameMeta->FileSize);
                     Entry->ParentFileHandle = GameAsyncFileHandle;
@@ -2336,16 +2349,25 @@ public:
                 }
                 else
                 {
-                    UE_LOG(LogRHI, Verbose, TEXT("Encountered a PSO entry %u that has been removed from the game-content file: %s"), Entry->Hash, *Meta.FileGuid.ToString());
+                    UE_LOG(LogRHI, Verbose, TEXT("Encountered a PSO entry %u that has been removed from the game-content file: %s or no game-content file"), Entry->Hash, *Meta.FileGuid.ToString());
                     Entry->bValid = false;
                     continue;
                 }
 			}
 			else if (Meta.FileGuid == UserFileGuid)
 			{
-                Entry->Data.SetNum(Meta.FileSize);
-				Entry->ParentFileHandle = UserAsyncFileHandle;
-				Entry->ReadRequest = MakeShareable(UserAsyncFileHandle->ReadRequest(Meta.FileOffset, Meta.FileSize, AIOP_Normal, nullptr, Entry->Data.GetData()));
+				if(UserAsyncFileHandle.IsValid())
+				{
+					Entry->Data.SetNum(Meta.FileSize);
+					Entry->ParentFileHandle = UserAsyncFileHandle;
+					Entry->ReadRequest = MakeShareable(UserAsyncFileHandle->ReadRequest(Meta.FileOffset, Meta.FileSize, AIOP_Normal, nullptr, Entry->Data.GetData()));
+				}
+				else
+				{
+					UE_LOG(LogRHI, Verbose, TEXT("Encountered a PSO entry %u that references user content file ID: %s but async handle not valid"), Entry->Hash, *Meta.FileGuid.ToString());
+					Entry->bValid = false;
+                    continue;
+				}
 			}
             else
             {
@@ -2695,7 +2717,7 @@ void FPipelineFileCache::ClosePipelineFileCache()
 
 void FPipelineFileCache::RegisterPSOUsageDataUpdateForNextSave(FPSOUsageData& UsageData)
 {
-	FPSOUsageData& CurrentEntry = NewPSOUsage.FindOrAdd(UsageData.PSOHash);	
+	FPSOUsageData& CurrentEntry = NewPSOUsage.FindOrAdd(UsageData.PSOHash);
 	CurrentEntry.PSOHash = UsageData.PSOHash;
 	CurrentEntry.UsageMask |= UsageData.UsageMask;
 	CurrentEntry.EngineFlags |= UsageData.EngineFlags;
