Clean up memcpy of FColor. Use assignment instead

diff --git a/Engine/Source/Runtime/RHI/Public/RHISurfaceDataConversion.h b/Engine/Source/Runtime/RHI/Public/RHISurfaceDataConversion.h
--- a/Engine/Source/Runtime/RHI/Public/RHISurfaceDataConversion.h
+++ b/Engine/Source/Runtime/RHI/Public/RHISurfaceDataConversion.h
@@ -133,15 +133,15 @@ static inline void ConvertRawR16G16B16A16FDataToFColor(uint32 Width, uint32 Heig
 
 		for (uint32 X = 0; X < Width; X++)
 		{
-			FColor NormalizedColor =
+			*DestPtr =
 				FLinearColor(
 				(SrcPtr[0] - MinValue.X) / (MaxValue.X - MinValue.X),
 					(SrcPtr[1] - MinValue.Y) / (MaxValue.Y - MinValue.Y),
 					(SrcPtr[2] - MinValue.Z) / (MaxValue.Z - MinValue.Z),
 					(SrcPtr[3] - MinValue.W) / (MaxValue.W - MinValue.W)
 				).ToFColor(LinearToGamma);
-			FMemory::Memcpy(DestPtr++, &NormalizedColor, sizeof(FColor));
 			SrcPtr += 4;
+			++DestPtr;
 		}
 	}
 }
@@ -159,9 +159,9 @@ static inline void ConvertRawR11G11B10DataToFColor(uint32 Width, uint32 Height,
 		{
 			FLinearColor Value = (*SrcPtr).ToLinearColor();
 
-			FColor NormalizedColor = Value.ToFColor(LinearToGamma);
-			FMemory::Memcpy(DestPtr++, &NormalizedColor, sizeof(FColor));
+			*DestPtr = Value.ToFColor(LinearToGamma);
 			++SrcPtr;
+			++DestPtr;
 		}
 	}
 }
@@ -196,15 +196,15 @@ static inline void ConvertRawR32G32B32A32DataToFColor(uint32 Width, uint32 Heigh
 
 		for (uint32 X = 0; X < Width; X++)
 		{
-			FColor NormalizedColor =
+			*DestPtr =
 				FLinearColor(
 				(SrcPtr[0] - MinValue.X) / (MaxValue.X - MinValue.X),
 					(SrcPtr[1] - MinValue.Y) / (MaxValue.Y - MinValue.Y),
 					(SrcPtr[2] - MinValue.Z) / (MaxValue.Z - MinValue.Z),
 					(SrcPtr[3] - MinValue.W) / (MaxValue.W - MinValue.W)
 				).ToFColor(LinearToGamma);
-			FMemory::Memcpy(DestPtr++, &NormalizedColor, sizeof(FColor));
 			SrcPtr += 4;
+			++DestPtr;
 		}
 	}
 }
@@ -233,8 +233,9 @@ static inline void ConvertRawR24G8DataToFColor(uint32 Width, uint32 Height, uint
 				NormalizedColor = FLinearColor(LinearValue, LinearValue, LinearValue, 0).ToFColor(bLinearToGamma);
 			}
 
-			FMemory::Memcpy(DestPtr++, &NormalizedColor, sizeof(FColor));
+			*DestPtr = NormalizedColor;
 			++SrcPtr;
+			++DestPtr;
 		}
 	}
 }
@@ -253,9 +254,9 @@ static inline void ConvertRawR32DataToFColor(uint32 Width, uint32 Height, uint8
 
 			float LinearValue = FMath::Min(InFlags.ComputeNormalizedDepth(DeviceZ), 1.0f);
 
-			FColor NormalizedColor = FLinearColor(LinearValue, LinearValue, LinearValue, 0).ToFColor(bLinearToGamma);
-			FMemory::Memcpy(DestPtr++, &NormalizedColor, sizeof(FColor));
+			*DestPtr = FLinearColor(LinearValue, LinearValue, LinearValue, 0).ToFColor(bLinearToGamma);
 			SrcPtr += 1; // todo: copies only depth, need to check how this format is read
+			++DestPtr;
 			UE_LOG(LogRHI, Warning, TEXT("CPU read of R32G8X24 is not tested and may not function."));
 		}
 	}
@@ -329,9 +330,9 @@ static inline void ConvertRawD32S8DataToFColor(uint32 Width, uint32 Height, uint
 			{
 				float DeviceZ = *SrcPtr;
 				float LinearValue = FMath::Min(InFlags.ComputeNormalizedDepth(DeviceZ), 1.0f);
-				FColor NormalizedColor = FLinearColor(LinearValue, LinearValue, LinearValue, 0).ToFColor(bLinearToGamma);
+				*DestPtr = FLinearColor(LinearValue, LinearValue, LinearValue, 0).ToFColor(bLinearToGamma);
 
-				FMemory::Memcpy(DestPtr++, &NormalizedColor, sizeof(FColor));
+				++DestPtr;
 				++SrcPtr;
 			}
 		}
@@ -348,10 +349,10 @@ static inline void ConvertRawD32S8DataToFColor(uint32 Width, uint32 Height, uint
 			for (uint32 X = 0; X < Width; X++)
 			{
 				uint8 DeviceStencil = *SrcPtr;
-				FColor NormalizedColor = FColor(DeviceStencil, DeviceStencil, DeviceStencil, 0xFF);
+				*DestPtr = FColor(DeviceStencil, DeviceStencil, DeviceStencil, 0xFF);
 
-				FMemory::Memcpy(DestPtr++, &NormalizedColor, sizeof(FColor));
 				++SrcPtr;
+				++DestPtr;
 			}
 		}
 	}
@@ -647,4 +648,4 @@ static inline void ConvertRawR16G16DataToFLinearColor(uint32 Width, uint32 Heigh
 			++DestPtr;
 		}
 	}
-}
\ No newline at end of file
+}
