Check for macOS 10.13.6 for the editor for now and 10.14.6 for the game

diff --git a/Engine/Source/Runtime/RHI/Private/DynamicRHI.cpp b/Engine/Source/Runtime/RHI/Private/DynamicRHI.cpp
--- a/Engine/Source/Runtime/RHI/Private/DynamicRHI.cpp
+++ b/Engine/Source/Runtime/RHI/Private/DynamicRHI.cpp
@@ -178,6 +178,16 @@ static void RHIDetectAndWarnOfBadDrivers(bool bHasEditorToken)
 		return;
 	}
 
+#if WITH_EDITOR
+	if (FPlatformMisc::MacOSXVersionCompare(10,13,6) < 0)
+	{
+		const FString BaseName = FApp::HasProjectName() ? FApp::GetProjectName() : TEXT("");
+		// this message can be suppressed with r.WarnOfBadDrivers=0
+		FPlatformMisc::MessageBoxExt(EAppMsgType::Ok,
+									 *NSLOCTEXT("MessageDialog", "UpdateMacOSX_Body", "Please update to the latest version of macOS for best performance and stability.").ToString(),
+									 *NSLOCTEXT("MessageDialog", "UpdateMacOSX_Title", "Update macOS").ToString());
+	}
+#else
 	if (FPlatformMisc::MacOSXVersionCompare(10,14,6) < 0)
 	{
 		const FString BaseName = FApp::HasProjectName() ? FApp::GetProjectName() : TEXT("");
@@ -196,6 +206,7 @@ static void RHIDetectAndWarnOfBadDrivers(bool bHasEditorToken)
 										 *NSLOCTEXT("MessageDialog", "UpdateMacOSX_Title", "Update macOS").ToString());
 		}
 	}
+#endif
 }
 #endif // PLATFORM_WINDOWS
 
