Don't show the old macOS version warning in unattended mode

diff --git a/Engine/Source/Runtime/RHI/Private/DynamicRHI.cpp b/Engine/Source/Runtime/RHI/Private/DynamicRHI.cpp
--- a/Engine/Source/Runtime/RHI/Private/DynamicRHI.cpp
+++ b/Engine/Source/Runtime/RHI/Private/DynamicRHI.cpp
@@ -173,7 +173,7 @@ static void RHIDetectAndWarnOfBadDrivers(bool bHasEditorToken)
 {
 	int32 CVarValue = CVarWarnOfBadDrivers.GetValueOnGameThread();
 
-	if(!GIsRHIInitialized || !CVarValue || GRHIVendorId == 0 || bHasEditorToken)
+	if (!GIsRHIInitialized || !CVarValue || GRHIVendorId == 0 || bHasEditorToken || FApp::IsUnattended())
 	{
 		return;
 	}
