Richard Wallis' change to delete local user PSO cache when using -logpso, plus changes to only enable the local cache in shipping build by default, and enable it on D3D12

diff --git a/Engine/Source/Runtime/RHI/Private/PipelineFileCache.cpp b/Engine/Source/Runtime/RHI/Private/PipelineFileCache.cpp
--- a/Engine/Source/Runtime/RHI/Private/PipelineFileCache.cpp
+++ b/Engine/Source/Runtime/RHI/Private/PipelineFileCache.cpp
@@ -79,7 +79,7 @@ static TAutoConsoleVariable<int32> CVarPSOFileCacheEnabled(
 static TAutoConsoleVariable<int32> CVarPSOFileCacheLogPSO(
 														   TEXT("r.ShaderPipelineCache.LogPSO"),
 														   PIPELINE_CACHE_DEFAULT_ENABLED,
-														   TEXT("1 Logs new PSO entries into the file cache and allow saving, 0 uses existing PSO file cache in readonly mode (if enabled)."),
+														   TEXT("1 Logs new PSO entries into the file cache and allows saving."),
 														   ECVF_Default | ECVF_RenderThreadSafe
 														   );
 
@@ -100,11 +100,7 @@ static FAutoConsoleVariableRef CVarPSOFileCachePrintNewPSODescriptors(
 
 static TAutoConsoleVariable<int32> CVarPSOFileCacheSaveUserCache(
                                                             TEXT("r.ShaderPipelineCache.SaveUserCache"),
-#if PIPELINE_CACHE_DEFAULT_ENABLED
-															(int32)1,
-#else
-															(int32)0,
-#endif
+															PIPELINE_CACHE_DEFAULT_ENABLED && UE_BUILD_SHIPPING,
                                                             TEXT("If > 0 then any missed PSOs will be saved to a writable user cache file for subsequent runs to load and avoid in-game hitches. Enabled by default on macOS only."),
                                                             ECVF_Default | ECVF_RenderThreadSafe
                                                             );
@@ -1366,12 +1362,17 @@ public:
         if (!bOnce)
         {
             bOnce = true;
-            bCmdLineForce = FParse::Param(FCommandLine::Get(), TEXT("deleteuserpsocache"));
+            bCmdLineForce = FParse::Param(FCommandLine::Get(), TEXT("deleteuserpsocache")) || FParse::Param(FCommandLine::Get(), TEXT("logPSO"));
             UE_CLOG(bCmdLineForce, LogRHI, Warning, TEXT("****************************** Deleting user-writable PSO cache as requested on command line"));
         }
         return bCmdLineForce;
     }
     
+    bool ShouldLoadUserCache()
+    {
+		return FPipelineFileCache::LogPSOtoFileCache() && (CVarPSOFileCacheSaveUserCache.GetValueOnAnyThread() > 0);
+    }
+    
 	bool OpenPipelineFileCache(FString const& FileName, EShaderPlatform Platform, FGuid& OutGameFileGuid)
 	{
 		SET_DWORD_STAT(STAT_TotalGraphicsPipelineStateCount, 0);
@@ -1442,7 +1443,7 @@ public:
 
         bool bUserFileOk = false;
         
-        if (FPipelineFileCache::LogPSOtoFileCache() && (CVarPSOFileCacheSaveUserCache.GetValueOnAnyThread() > 0))
+        if (ShouldLoadUserCache())
         {
 			FPipelineCacheFileFormatTOC UserTOC;
             bUserFileOk = OpenPipelineFileCache(FilePath, UserFileGuid, UserAsyncFileHandle, UserTOC);
diff --git a/Engine/Source/Runtime/RHI/Private/Windows/WindowsDynamicRHI.cpp b/Engine/Source/Runtime/RHI/Private/Windows/WindowsDynamicRHI.cpp
--- a/Engine/Source/Runtime/RHI/Private/Windows/WindowsDynamicRHI.cpp
+++ b/Engine/Source/Runtime/RHI/Private/Windows/WindowsDynamicRHI.cpp
@@ -187,6 +187,9 @@ static IDynamicRHIModule* LoadDynamicRHIModule(ERHIFeatureLevel::Type& DesiredFe
 
 		auto PSOFileCacheReportCVar = IConsoleManager::Get().FindTConsoleVariableDataInt(TEXT("r.ShaderPipelineCache.ReportPSO"));
 		*PSOFileCacheReportCVar = 1;
+
+		auto PSOFileCacheUserCacheCVar = IConsoleManager::Get().FindTConsoleVariableDataInt(TEXT("r.ShaderPipelineCache.SaveUserCache"));
+		*PSOFileCacheUserCacheCVar = UE_BUILD_SHIPPING;
 #endif
 
 		if (!DynamicRHIModule || !DynamicRHIModule->IsSupported())
