UI for GPU Crash Debugging

diff --git a/Engine/Source/Runtime/RHI/Private/Windows/WindowsDynamicRHI.cpp b/Engine/Source/Runtime/RHI/Private/Windows/WindowsDynamicRHI.cpp
--- a/Engine/Source/Runtime/RHI/Private/Windows/WindowsDynamicRHI.cpp
+++ b/Engine/Source/Runtime/RHI/Private/Windows/WindowsDynamicRHI.cpp
@@ -46,6 +46,13 @@ static bool ShouldPreferD3D12()
 
 static IDynamicRHIModule* LoadDynamicRHIModule(ERHIFeatureLevel::Type& DesiredFeatureLevel, const TCHAR*& LoadedRHIModuleName)
 {
+	bool bUseGPUCrashDebugging = false;
+	if (!GIsEditor && GConfig->GetBool(TEXT("D3DRHIPreference"), TEXT("bUseGPUCrashDebugging"), bUseGPUCrashDebugging, GGameUserSettingsIni))
+	{
+		auto GPUCrashDebuggingCVar = IConsoleManager::Get().FindTConsoleVariableDataInt(TEXT("r.GPUCrashDebugging"));
+		*GPUCrashDebuggingCVar = bUseGPUCrashDebugging;
+	}
+
 	bool bPreferD3D12 = ShouldPreferD3D12();
 	
 	// command line overrides
