#rnx

diff --git a/Engine/Source/Runtime/RHI/Public/PipelineFileCache.h b/Engine/Source/Runtime/RHI/Public/PipelineFileCache.h
--- a/Engine/Source/Runtime/RHI/Public/PipelineFileCache.h
+++ b/Engine/Source/Runtime/RHI/Public/PipelineFileCache.h
@@ -199,7 +199,7 @@ struct RHI_API FPipelineCacheFileFormatPSORead
 	: Ar(nullptr)
 	, Hash(0)
 	, bReadCompleted(false)
-    , bValid(false)
+	, bValid(false)
 	{}
 	
 	~FPipelineCacheFileFormatPSORead()
@@ -216,10 +216,12 @@ struct RHI_API FPipelineCacheFileFormatPSORead
 	
 	uint32 Hash;
 	bool bReadCompleted;
-    bool bValid;
+	bool bValid;
 	
-	TSharedPtr<class IAsyncReadRequest, ESPMode::ThreadSafe> ReadRequest;
+	// Note that the contract of IAsyncReadFileHandle and IAsyncReadRequest requires that we delete the ReadRequest before deleting its ParentFileHandle. 
+	// We therefore require that ParentFileHandle is declared before ReadRequest, so that the class destructor tears down first ReadRequest then ParentFileHandle.
 	TSharedPtr<class IAsyncReadFileHandle, ESPMode::ThreadSafe> ParentFileHandle;
+	TSharedPtr<class IAsyncReadRequest, ESPMode::ThreadSafe> ReadRequest;
 };
 
 struct RHI_API FPipelineCachePSOHeader
