Fixed crash in command list CopyContext.

diff --git a/Engine/Source/Runtime/RHI/Public/RHICommandList.h b/Engine/Source/Runtime/RHI/Public/RHICommandList.h
--- a/Engine/Source/Runtime/RHI/Public/RHICommandList.h
+++ b/Engine/Source/Runtime/RHI/Public/RHICommandList.h
@@ -527,7 +527,6 @@ public:
 
 	void CopyContext(FRHICommandListBase& ParentCommandList)
 	{
-		check(Context);
 		Context = ParentCommandList.Context;
 		ComputeContext = ParentCommandList.ComputeContext;
 	}
