F - Fix typo
#rnx

diff --git a/Engine/Source/Runtime/D3D12RHI/Private/D3D12View.cpp b/Engine/Source/Runtime/D3D12RHI/Private/D3D12View.cpp
--- a/Engine/Source/Runtime/D3D12RHI/Private/D3D12View.cpp
+++ b/Engine/Source/Runtime/D3D12RHI/Private/D3D12View.cpp
@@ -266,7 +266,6 @@ FShaderResourceViewRHIRef FD3D12DynamicRHI::RHICreateShaderResourceView(const FS
 
 					const uint32 BufferUsage = StructuredBuffer->GetUsage();
 					const bool bByteAccessBuffer = (BufferUsage & BUF_ByteAddressBuffer) != 0;
-					const bool bUINT8Access = (BufferUsage & BUF_UINT8) != 0;
 					// Create a Shader Resource View
 					D3D12_SHADER_RESOURCE_VIEW_DESC SRVDesc = {};
 					SRVDesc.Shader4ComponentMapping = D3D12_DEFAULT_SHADER_4_COMPONENT_MAPPING;
@@ -282,11 +281,6 @@ FShaderResourceViewRHIRef FD3D12DynamicRHI::RHICreateShaderResourceView(const FS
 						SRVDesc.Format = DXGI_FORMAT_R32_TYPELESS;
 						Stride = 4;
 					}
-					else if (bUINT8Access)
-					{
-						SRVDesc.Format = DXGI_FORMAT_R8_UINT;
-						Stride = 1;
-					}
 					else
 					{
 						SRVDesc.Buffer.StructureByteStride = Stride;
diff --git a/Engine/Source/Runtime/RHI/Public/RHIDefinitions.h b/Engine/Source/Runtime/RHI/Public/RHIDefinitions.h
--- a/Engine/Source/Runtime/RHI/Public/RHIDefinitions.h
+++ b/Engine/Source/Runtime/RHI/Public/RHIDefinitions.h
@@ -794,8 +794,13 @@ enum EBufferUsageFlags
 
 	/** Create a byte address buffer, which is basically a structured buffer with a uint32 type. */
 	BUF_ByteAddressBuffer = 0x0020,
+
 	/** Create a structured buffer with an atomic UAV counter. */
-	BUF_UAVCounter        = 0x0040,
+	BUF_UAVCounter UE_DEPRECATED(4.25, "BUF_UAVCounter is deprecated and getting removed; please don't use.") = 0x0040,
+
+	/** Buffer that the GPU will use as a source for a copy. */
+	BUF_SourceCopy			= 0x0040,
+
 	/** Create a buffer that can be bound as a stream output target. */
 	BUF_StreamOutput      = 0x0080,
 	/** Create a buffer which contains the arguments used by DispatchIndirect or DrawIndirect. */
@@ -825,7 +830,7 @@ enum EBufferUsageFlags
 	BUF_Transient		  = 0x2000,
 
 	/** Buffer that should be accessed one byte at a time. */
-	BUF_UINT8             = 0x4000,
+	BUF_UINT8 UE_DEPRECATED(4.25, "BUF_UINT8 is deprecated and getting removed; please don't use.") = 0x4000,
 
 	/**
 	 * Buffer contains opaque ray tracing acceleration structure data.
