Removed unnecessary module dependency on TraceLog from RHI.

#rnx

diff --git a/Engine/Source/Runtime/RHI/RHI.Build.cs b/Engine/Source/Runtime/RHI/RHI.Build.cs
--- a/Engine/Source/Runtime/RHI/RHI.Build.cs
+++ b/Engine/Source/Runtime/RHI/RHI.Build.cs
@@ -9,7 +9,6 @@ public class RHI : ModuleRules
 	{
 		PrivateDependencyModuleNames.Add("Core");
 		PrivateDependencyModuleNames.Add("ApplicationCore");
-		PrivateDependencyModuleNames.Add("TraceLog");
 
 		if (Target.bCompileAgainstEngine)
 		{
