AUTO MERGE

Merging //UE4/Main @ 11112898 to //UE4/Dev-Core (dev-core-l0178)

(ushell-p4-mergedown)

diff --git a/Engine/Source/Developer/PakFileUtilities/Private/PakFileUtilities.cpp b/Engine/Source/Developer/PakFileUtilities/Private/PakFileUtilities.cpp
--- a/Engine/Source/Developer/PakFileUtilities/Private/PakFileUtilities.cpp
+++ b/Engine/Source/Developer/PakFileUtilities/Private/PakFileUtilities.cpp
@@ -1975,7 +1975,7 @@ bool CreatePakFile(const TCHAR* Filename, TArray<FPakInputPair>& FilesToAdd, con
 					// if we still save 64KB it's probably worthwhile compressing, as that saves a file read operation in the runtime.
 								// TODO: drive this threshold from the command line
 					float PercentLess = ((float)CompressedFileBuffer.TotalCompressedSize / (OriginalFileSize / 100.f));
-					const bool bNotEnoughCompression = PercentLess > 90.f && (OriginalFileSize - CompressedFileBuffer.TotalCompressedSize) < 6553;
+					const bool bNotEnoughCompression = (PercentLess > 90.f) && ((OriginalFileSize - CompressedFileBuffer.TotalCompressedSize) < 65536);
 					const bool bIsLastCompressionFormat = MethodIndex == CompressionFormats->Num() - 1;
 					if (bNotEnoughCompression && (!bForceCompress || !bIsLastCompressionFormat))
 					{
diff --git a/Engine/Source/Editor/KismetCompiler/Private/AnimBlueprintCompiler.cpp b/Engine/Source/Editor/KismetCompiler/Private/AnimBlueprintCompiler.cpp
--- a/Engine/Source/Editor/KismetCompiler/Private/AnimBlueprintCompiler.cpp
+++ b/Engine/Source/Editor/KismetCompiler/Private/AnimBlueprintCompiler.cpp
@@ -832,6 +832,10 @@ void FAnimBlueprintCompilerContext::ProcessCustomPropertyNode(UAnimGraphNode_Cus
 					{
 						CustomPropNode->AddSourceTargetProperties(NewProperty->GetFName(), FoundProperty->GetFName());
 					}
+					else
+					{
+						CustomPropNode->AddSourceTargetProperties(NewProperty->GetFName(), Pin->GetFName());
+					}
 				}
 			}
 		}
diff --git a/Engine/Source/Editor/UnrealEd/Classes/Settings/ProjectPackagingSettings.h b/Engine/Source/Editor/UnrealEd/Classes/Settings/ProjectPackagingSettings.h
--- a/Engine/Source/Editor/UnrealEd/Classes/Settings/ProjectPackagingSettings.h
+++ b/Engine/Source/Editor/UnrealEd/Classes/Settings/ProjectPackagingSettings.h
@@ -289,6 +289,10 @@ public:
 	UPROPERTY(config, EditAnywhere, Category=Packaging, AdvancedDisplay)
 	TArray<FString> LocalizationTargetsToChunk;
 
+	/** The chunk ID that should be used as the catch-all chunk for any non-asset localized strings */
+	UPROPERTY(config, EditAnywhere, Category=Packaging, AdvancedDisplay)
+	int32 LocalizationTargetCatchAllChunkId = 0;
+
 	/**
 	 * Cook all things in the project content directory
 	 */
diff --git a/Engine/Source/Editor/UnrealEd/Private/Commandlets/AssetRegistryGenerator.cpp b/Engine/Source/Editor/UnrealEd/Private/Commandlets/AssetRegistryGenerator.cpp
--- a/Engine/Source/Editor/UnrealEd/Private/Commandlets/AssetRegistryGenerator.cpp
+++ b/Engine/Source/Editor/UnrealEd/Private/Commandlets/AssetRegistryGenerator.cpp
@@ -441,7 +441,7 @@ bool FAssetRegistryGenerator::GenerateStreamingInstallManifest(int64 InExtraFlav
 			}
 
 			// Allow the extra data generation steps to run and add their output to the manifest
-			if (ChunkDataGenerators.Num() > 0)
+			if (ChunkDataGenerators.Num() > 0 && SubChunkIndex == 0)
 			{
 				TSet<FName> PackagesInChunk;
 				PackagesInChunk.Reserve(Manifest->Num());
diff --git a/Engine/Source/Editor/UnrealEd/Private/Commandlets/PackageUtilities.cpp b/Engine/Source/Editor/UnrealEd/Private/Commandlets/PackageUtilities.cpp
--- a/Engine/Source/Editor/UnrealEd/Private/Commandlets/PackageUtilities.cpp
+++ b/Engine/Source/Editor/UnrealEd/Private/Commandlets/PackageUtilities.cpp
@@ -37,10 +37,6 @@
 #include "Editor.h"
 #include "FileHelpers.h"
 
-#include "Animation/AnimCompress_BitwiseCompressOnly.h"
-#include "Animation/AnimCompress_Automatic.h"
-
-
 #include "CollectionManagerTypes.h"
 #include "ICollectionManager.h"
 #include "CollectionManagerModule.h"
@@ -1779,12 +1775,8 @@ struct CompressAnimationsFunctor
 
 		const bool bSkipCinematicPackages = Switches.Contains(TEXT("SKIPCINES"));
 		const bool bSkipLongAnimations = Switches.Contains(TEXT("SKIPLONGANIMS"));
-		// Reset compression, don't do incremental compression, start from scratch
-		const bool bResetCompression = Switches.Contains(TEXT("RESETCOMPRESSION"));
 		/** Clear bDoNotOverrideCompression flag in animations */
 		const bool bClearNoCompressionOverride = Switches.Contains(TEXT("CLEARNOCOMPRESSIONOVERRIDE"));
-		/** If we're analyzing, we're not actually going to recompress, so we can skip some significant work. */
-		const bool bAnalyze = Switches.Contains(TEXT("ANALYZE"));
 		// See if we can save this package. If we can't, don't bother...
 		/** if we should auto checkout packages that need to be saved **/
 		const bool bAutoCheckOut = Switches.Contains(TEXT("AUTOCHECKOUTPACKAGES"));
@@ -1792,7 +1784,7 @@ struct CompressAnimationsFunctor
 		FSourceControlStatePtr SourceControlState = SourceControl.GetProvider().GetState(PackageFileName, EStateCacheUsage::ForceUpdate);
 
 		// check to see if we need to check this package out
-		if( !bAnalyze && SourceControlState.IsValid() && SourceControlState->CanCheckout() )
+		if( SourceControlState.IsValid() && SourceControlState->CanCheckout() )
 		{
 			// Cant check out, check to see why
 			if (bAutoCheckOut == true)
@@ -1854,13 +1846,13 @@ struct CompressAnimationsFunctor
 
 			// If animation has already been compressed with the commandlet and version is the same. then skip.
 			// We're only interested in new animations.
-			if( !bAnalyze && !bForceCompression && AnimSeq->CompressCommandletVersion == CompressCommandletVersion )
+			if( !bForceCompression && AnimSeq->CompressCommandletVersion == CompressCommandletVersion )
 			{
 				UE_LOG(LogPackageUtilities, Warning, TEXT("Same CompressCommandletVersion (%i) skip animation: %s (%s)"), CompressCommandletVersion, *AnimSeq->GetName(), *AnimSeq->GetFullName());
 				continue;
 			}
 
-			if( !bAnalyze && !bForceCompression && bSkipLongAnimations && (AnimSeq->GetRawNumberOfFrames() > 300) )
+			if( !bForceCompression && bSkipLongAnimations && (AnimSeq->GetRawNumberOfFrames() > 300) )
 			{
 				UE_LOG(LogPackageUtilities, Warning, TEXT("Animation (%s) has more than 300 frames (%i frames) and SKIPLONGANIMS switch is set. Skipping."), *AnimSeq->GetName(), AnimSeq->GetRawNumberOfFrames());
 				continue;
@@ -1873,349 +1865,6 @@ struct CompressAnimationsFunctor
 				Skeleton->GetLinker()->Preload(Skeleton);
 			}
 
-			if( bAnalyze )
-			{
-				static int32 NumTotalAnimations = 0;
-				static int32 NumTotalSize = 0;
-				static int32 Trans96Savings = 0;
-				static int32 Trans48Savings = 0;
-				static int32 Rot96Savings = 0;
-				static int32 Rot48Savings = 0;
-				static int32 Scale96Savings = 0;
-				static int32 Scale48Savings = 0;
-				static int32 Num96TransTracks = 0;
-				static int32 Num96RotTracks = 0;
-				static int32 Num96ScaleTracks = 0;
-				static int32 Num48TransTracks = 0;
-				static int32 Num48RotTracks = 0;
-				static int32 Num48ScaleTracks = 0;
-				static int32 Num32TransTracks = 0;
-				static int32 Num32ScaleTracks = 0;
-				static int32 UnknownTransTrack = 0;
-				static int32 UnknownRotTrack = 0;
-				static int32 UnknownScaleTrack = 0;
-				static int32 RotationOnlySavings = 0;
-				static int32 RotationOnlyManyKeys = 0;
-
-				NumTotalAnimations++;
-
-				FArchiveCountMem CountBytesSize( AnimSeq );
-				int32 ResourceSize = CountBytesSize.GetNum();
-
-				NumTotalSize += ResourceSize;
-				
-				FUECompressedAnimData& CompressedData = AnimSeq->CompressedData.CompressedDataStructure;
-
-				// Looking for PerTrackCompression using 96bit translation compression.
-				if( CompressedData.KeyEncodingFormat == AKF_PerTrackCompression && CompressedData.CompressedByteStream.Num() > 0 )
-				{
-					bool bCandidate = false;
-
-					for (int32 i = 0; i<AnimSeq->GetCompressedTrackToSkeletonMapTable().Num(); i++)
- 					{
- 						const int32 TrackIndex = i;
-						const int32 BoneTreeIndex = AnimSeq->GetCompressedTrackToSkeletonMapTable()[TrackIndex].BoneTreeIndex;
-						const FName BoneTreeName = Skeleton->GetReferenceSkeleton().GetBoneName(BoneTreeIndex);
-
- 						// Translation
-						{
-							// Use the CompressedTrackOffsets stream to find the data addresses
-							const int32* RESTRICT TrackDataForTransKey = CompressedData.CompressedTrackOffsets.GetData() + (TrackIndex * 2);
-							const int32 TransKeysOffset = TrackDataForTransKey[0];
- 							if( TransKeysOffset != INDEX_NONE )
- 							{
-								const uint8* RESTRICT TrackData = CompressedData.CompressedByteStream.GetData() + TransKeysOffset + 4;
-								const int32 Header = *((int32*)(CompressedData.CompressedByteStream.GetData() + TransKeysOffset));
-	 
- 								int32 KeyFormat;
- 								int32 NumKeys;
- 								int32 FormatFlags;
- 								int32 BytesPerKey;
- 								int32 FixedBytes;
- 								FAnimationCompression_PerTrackUtils::DecomposeHeader(Header, /*OUT*/ KeyFormat, /*OUT*/ NumKeys, /*OUT*/ FormatFlags, /*OUT*/BytesPerKey, /*OUT*/ FixedBytes);
- 								if( KeyFormat == ACF_Float96NoW )
- 								{
-									Num96TransTracks++;
-
- 									// Determine which components we could let go, and bytes we could save.
-									const FBox KeyBounds((FVector*)(TrackData + FixedBytes), NumKeys);
-									const bool bHasX = (FMath::Abs(KeyBounds.Max.X) >= 0.0002f) || (FMath::Abs(KeyBounds.Min.X) >= 0.0002f);
-									const bool bHasY = (FMath::Abs(KeyBounds.Max.Y) >= 0.0002f) || (FMath::Abs(KeyBounds.Min.Y) >= 0.0002f);
-									const bool bHasZ = (FMath::Abs(KeyBounds.Max.Z) >= 0.0002f) || (FMath::Abs(KeyBounds.Min.Z) >= 0.0002f);
-
-									if( !bHasX )
-									{
-										Trans96Savings += (4 * NumKeys);
-										bCandidate = true;
-									}
-									if( !bHasY )
-									{
-										Trans96Savings += (4 * NumKeys);
-										bCandidate = true;
-									}
-									if( !bHasZ )
-									{
-										Trans96Savings += (4 * NumKeys);
-										bCandidate = true;
-									}
-								}
-								// Measure savings on 48bits translations
-								else if( KeyFormat == ACF_Fixed48NoW )
-								{
-									Num48TransTracks++;
-
-									const int32 SavedBytes = (6 - BytesPerKey) * NumKeys;
-									if( SavedBytes > 0 )
-									{
-										bCandidate = true;
-										Trans48Savings += SavedBytes;
-									}
-								}
-								else if( KeyFormat == ACF_IntervalFixed32NoW )
-								{
-									Num32TransTracks++;
-								}
-								else
-								{
-									UnknownTransTrack++;
-								}
-
-								// Measure how much we'd save if we used "rotation only" for compression
-								// root bone is true if BoneTreeIndex == 0 
-								// @todoanim : @fixmelh : AnimRotationOnly fix
-								if( BoneTreeIndex > 0 )
-// 									&& ((AnimSet->UseTranslationBoneNames.Num() > 0 && AnimSet->UseTranslationBoneNames.FindItemIndex(BoneName) == INDEX_NONE) 
-// 										|| (AnimSet->ForceMeshTranslationBoneNames.FindItemIndex(BoneName) != INDEX_NONE)) 
-// 									)
-								{
-									RotationOnlySavings += (BytesPerKey * NumKeys);
-									if( NumKeys > 1 )
-									{
-										const uint8* RESTRICT KeyData0 = TrackData + FixedBytes;
-										FVector V0;
-										FAnimationCompression_PerTrackUtils::DecompressTranslation(KeyFormat, FormatFlags, V0, TrackData, KeyData0);
-
-										float MaxErrorFromFirst = 0.f;
-										float MaxErrorFromDefault = 0.f;
-										const TArray<FTransform> & LocalRefPoses = Skeleton->GetRefLocalPoses();
-										for(int32 KeyIdx=0; KeyIdx<NumKeys; KeyIdx++)
-										{
-											const uint8* RESTRICT KeyDataN = TrackData + FixedBytes + KeyIdx * BytesPerKey;
-											FVector VN;
-											FAnimationCompression_PerTrackUtils::DecompressTranslation(KeyFormat, FormatFlags, VN, TrackData, KeyDataN);
-
-											// @todoanim: we will need more discussion, but we might use Skeleton->RefLocalPoses for compression
-											MaxErrorFromDefault = FMath::Max(MaxErrorFromDefault, FMath::Abs(VN.X - LocalRefPoses[BoneTreeIndex].GetLocation().X));
-											MaxErrorFromDefault = FMath::Max(MaxErrorFromDefault, FMath::Abs(VN.Y - LocalRefPoses[BoneTreeIndex].GetLocation().Y));
-											MaxErrorFromDefault = FMath::Max(MaxErrorFromDefault, FMath::Abs(VN.Z - LocalRefPoses[BoneTreeIndex].GetLocation().Z));
-
-											MaxErrorFromFirst = FMath::Max(MaxErrorFromFirst, FMath::Abs(VN.X - V0.X));
-											MaxErrorFromFirst = FMath::Max(MaxErrorFromFirst, FMath::Abs(VN.Y - V0.Y));
-											MaxErrorFromFirst = FMath::Max(MaxErrorFromFirst, FMath::Abs(VN.Z - V0.Z));
-										}
-
-										UE_LOG(LogPackageUtilities, Warning, TEXT("RotationOnly translation track that is animated! %s, %s (%s) NumKeys: %i, MaxErrorFromDefault: %f, MaxErrorFromFirst: %f"), 
-											*BoneTreeName.ToString(), *AnimSeq->GetName(), *AnimSeq->GetFullName(), NumKeys, MaxErrorFromDefault, MaxErrorFromFirst);
-										RotationOnlyManyKeys += (BytesPerKey * (NumKeys-1));
-									}
-								}
- 							}
-						}
-
-						// Rotation
-						{
-							// Use the CompressedTrackOffsets stream to find the data addresses
-							const int32* RESTRICT TrackDataForRotKey = CompressedData.CompressedTrackOffsets.GetData() + (TrackIndex * 2);
-							const int32 RotKeysOffset = TrackDataForRotKey[1];
-							if( RotKeysOffset != INDEX_NONE )
-							{
-								const uint8* RESTRICT TrackData = CompressedData.CompressedByteStream.GetData() + RotKeysOffset + 4;
-								const int32 Header = *((int32*)(CompressedData.CompressedByteStream.GetData() + RotKeysOffset));
-
-								int32 KeyFormat;
-								int32 NumKeys;
-								int32 FormatFlags;
-								int32 BytesPerKey;
-								int32 FixedBytes;
-								FAnimationCompression_PerTrackUtils::DecomposeHeader(Header, /*OUT*/ KeyFormat, /*OUT*/ NumKeys, /*OUT*/ FormatFlags, /*OUT*/BytesPerKey, /*OUT*/ FixedBytes);
-								if( KeyFormat == ACF_Float96NoW )
-								{
-									Num96RotTracks++;
-
-									// Determine which components we could let go, and bytes we could save.
-									const FBox KeyBounds((FVector*)(TrackData + FixedBytes), NumKeys);
-									const bool bHasX = (FMath::Abs(KeyBounds.Max.X) >= 0.0002f) || (FMath::Abs(KeyBounds.Min.X) >= 0.0002f);
-									const bool bHasY = (FMath::Abs(KeyBounds.Max.Y) >= 0.0002f) || (FMath::Abs(KeyBounds.Min.Y) >= 0.0002f);
-									const bool bHasZ = (FMath::Abs(KeyBounds.Max.Z) >= 0.0002f) || (FMath::Abs(KeyBounds.Min.Z) >= 0.0002f);
-
-									if( !bHasX )
-									{
-										Rot96Savings += (4 * NumKeys);
-										bCandidate = true;
-									}
-									if( !bHasY )
-									{
-										Rot96Savings += (4 * NumKeys);
-										bCandidate = true;
-									}
-									if( !bHasZ )
-									{
-										Rot96Savings += (4 * NumKeys);
-										bCandidate = true;
-									}
-								}
-								// Measure savings on 48bits rotations.
-								else if( KeyFormat == ACF_Fixed48NoW )
-								{
-									Num48RotTracks++;
-
-									const int32 SavedBytes = (6 - BytesPerKey) * NumKeys;
-									if( SavedBytes > 0 )
-									{
-										bCandidate = true;
-										Rot48Savings += SavedBytes;
-									}
-								}
-								else
-								{
-									UnknownRotTrack++;
-								}
-							}
-						}
-
-						// Scale
-						{
-							// Use the CompressedTrackOffsets stream to find the data addresses
-							const int32 ScaleKeysOffset = CompressedData.CompressedScaleOffsets.GetOffsetData(TrackIndex, 0);
-							if( ScaleKeysOffset != INDEX_NONE )
-							{
-								const uint8* RESTRICT TrackData = CompressedData.CompressedByteStream.GetData() + ScaleKeysOffset + 4;
-								const int32 Header = *((int32*)(CompressedData.CompressedByteStream.GetData() + ScaleKeysOffset));
-
-								int32 KeyFormat;
-								int32 NumKeys;
-								int32 FormatFlags;
-								int32 BytesPerKey;
-								int32 FixedBytes;
-								FAnimationCompression_PerTrackUtils::DecomposeHeader(Header, /*OUT*/ KeyFormat, /*OUT*/ NumKeys, /*OUT*/ FormatFlags, /*OUT*/BytesPerKey, /*OUT*/ FixedBytes);
-								if( KeyFormat == ACF_Float96NoW )
-								{
-									Num96ScaleTracks++;
-
-									// Determine which components we could let go, and bytes we could save.
-									const FBox KeyBounds((FVector*)(TrackData + FixedBytes), NumKeys);
-									const bool bHasX = (FMath::Abs(KeyBounds.Max.X) >= 0.0002f) || (FMath::Abs(KeyBounds.Min.X) >= 0.0002f);
-									const bool bHasY = (FMath::Abs(KeyBounds.Max.Y) >= 0.0002f) || (FMath::Abs(KeyBounds.Min.Y) >= 0.0002f);
-									const bool bHasZ = (FMath::Abs(KeyBounds.Max.Z) >= 0.0002f) || (FMath::Abs(KeyBounds.Min.Z) >= 0.0002f);
-
-									if( !bHasX )
-									{
-										Scale96Savings += (4 * NumKeys);
-										bCandidate = true;
-									}
-									if( !bHasY )
-									{
-										Scale96Savings += (4 * NumKeys);
-										bCandidate = true;
-									}
-									if( !bHasZ )
-									{
-										Scale96Savings += (4 * NumKeys);
-										bCandidate = true;
-									}
-								}
-								// Measure savings on 48bits Scales
-								else if( KeyFormat == ACF_Fixed48NoW )
-								{
-									Num48ScaleTracks++;
-
-									const int32 SavedBytes = (6 - BytesPerKey) * NumKeys;
-									if( SavedBytes > 0 )
-									{
-										bCandidate = true;
-										Scale48Savings += SavedBytes;
-									}
-								}
-								else if( KeyFormat == ACF_IntervalFixed32NoW )
-								{
-									Num32ScaleTracks++;
-								}
-								else
-								{
-									UnknownScaleTrack++;
-								}
-
-								// Measure how much we'd save if we used "rotation only" for compression
-								// root bone is true if BoneTreeIndex == 0 
-								// @todoanim : @fixmelh : AnimRotationOnly fix
-								if( BoneTreeIndex > 0 )
-									// 									&& ((AnimSet->UseScaleBoneNames.Num() > 0 && AnimSet->UseScaleBoneNames.FindItemIndex(BoneName) == INDEX_NONE) 
-										// 										|| (AnimSet->ForceMeshScaleBoneNames.FindItemIndex(BoneName) != INDEX_NONE)) 
-											// 									)
-								{
-									RotationOnlySavings += (BytesPerKey * NumKeys);
-									if( NumKeys > 1 )
-									{
-										const uint8* RESTRICT KeyData0 = TrackData + FixedBytes;
-										FVector V0;
-										FAnimationCompression_PerTrackUtils::DecompressScale(KeyFormat, FormatFlags, V0, TrackData, KeyData0);
-
-										float MaxErrorFromFirst = 0.f;
-										float MaxErrorFromDefault = 0.f;
-										const TArray<FTransform> & LocalRefPoses = Skeleton->GetRefLocalPoses();
-										for(int32 KeyIdx=0; KeyIdx<NumKeys; KeyIdx++)
-										{
-											const uint8* RESTRICT KeyDataN = TrackData + FixedBytes + KeyIdx * BytesPerKey;
-											FVector VN;
-											FAnimationCompression_PerTrackUtils::DecompressScale(KeyFormat, FormatFlags, VN, TrackData, KeyDataN);
-
-											// @todoanim: we will need more discussion, but we might use Skeleton->RefLocalPoses for compression
-											MaxErrorFromDefault = FMath::Max(MaxErrorFromDefault, FMath::Abs(VN.X - LocalRefPoses[BoneTreeIndex].GetLocation().X));
-											MaxErrorFromDefault = FMath::Max(MaxErrorFromDefault, FMath::Abs(VN.Y - LocalRefPoses[BoneTreeIndex].GetLocation().Y));
-											MaxErrorFromDefault = FMath::Max(MaxErrorFromDefault, FMath::Abs(VN.Z - LocalRefPoses[BoneTreeIndex].GetLocation().Z));
-
-											MaxErrorFromFirst = FMath::Max(MaxErrorFromFirst, FMath::Abs(VN.X - V0.X));
-											MaxErrorFromFirst = FMath::Max(MaxErrorFromFirst, FMath::Abs(VN.Y - V0.Y));
-											MaxErrorFromFirst = FMath::Max(MaxErrorFromFirst, FMath::Abs(VN.Z - V0.Z));
-										}
-
-										UE_LOG(LogPackageUtilities, Warning, TEXT("RotationOnly Scale track that is animated! %s, %s (%s) NumKeys: %i, MaxErrorFromDefault: %f, MaxErrorFromFirst: %f"), 
-											*BoneTreeName.ToString(), *AnimSeq->GetName(), *AnimSeq->GetFullName(), NumKeys, MaxErrorFromDefault, MaxErrorFromFirst);
-										RotationOnlyManyKeys += (BytesPerKey * (NumKeys-1));
-									}
-								}
-							}
-						}
- 					}
-
-					if( bCandidate )
-					{
-						++AnalyzeCompressionCandidates;
-						UE_LOG(LogPackageUtilities, Warning, TEXT("[%i] Animation could be recompressed: %s (%s), Trans96Savings: %i, Rot96Savings: %i, Scale96Savings: %i, Trans48Savings: %i, Rot48Savings: %i, Scale48Savings: %i, RotationOnlySavings: %i, RotationOnlyManyKeys: %i (bytes)"), 
-							AnalyzeCompressionCandidates, *AnimSeq->GetName(), *AnimSeq->GetFullName(), Trans96Savings, Rot96Savings, Scale96Savings, Trans48Savings, Rot48Savings, Scale48Savings, RotationOnlySavings, RotationOnlyManyKeys);
-						UE_LOG(LogPackageUtilities, Warning, TEXT("Translation Track Count, Num96TransTracks: %i, Num48TransTracks: %i, Num32TransTracks: %i, UnknownTransTrack: %i"), 
-							Num96TransTracks, Num48TransTracks, Num32TransTracks, UnknownTransTrack);
-						UE_LOG(LogPackageUtilities, Warning, TEXT("Rotation Track Count, Num96RotTracks: %i, Num48RotTracks: %i, UnknownRotTrack: %i"), 
-							Num96RotTracks, Num48RotTracks, UnknownRotTrack);
-						UE_LOG(LogPackageUtilities, Warning, TEXT("Scale Track Count, Num96ScaleTracks: %i, Num48ScaleTracks: %i, Num32ScaleTracks: %i, UnknownScaleTrack: %i"), 
-							Num96ScaleTracks, Num48ScaleTracks, Num32ScaleTracks, UnknownScaleTrack);
-					}
-				}
-
-// 				if( AnimSeq->NumFrames > 1 && AnimSeq->KeyEncodingFormat != AKF_PerTrackCompression )
-// 				{
-// 					++AnalyzeCompressionCandidates;
-// 
-// 					FArchiveCountMem CountBytesSize( AnimSeq );
-// 					int32 ResourceSize = CountBytesSize.GetNum();
-// 
-// 					UE_LOG(LogPackageUtilities, Warning, TEXT("[%i] Animation could be recompressed: %s (%s), frames: %i, length: %f, size: %i bytes, compression scheme: %s"), 
-// 						AnalyzeCompressionCandidates, *AnimSeq->GetName(), *AnimSet->GetFullName(),  AnimSeq->NumFrames, AnimSeq->SequenceLength, ResourceSize, AnimSeq->CompressionScheme ? *AnimSeq->CompressionScheme->GetClass()->GetName() : TEXT("NULL"));
-// 				}
-
-				continue;
-			}
-
 			float HighestRatio = 0.f;
 #if 0 // @todoanim: not sure why we need this here
 			USkeletalMesh*	BestSkeletalMeshMatch = NULL;
@@ -2289,24 +1938,9 @@ struct CompressAnimationsFunctor
 				bDirtyPackage = true;
 			}
 
-			// Reset to default compressor
-			if( bResetCompression )
-			{
-				UE_LOG(LogPackageUtilities, Warning, TEXT("%s (%s) Resetting with BitwiseCompressOnly."), *AnimSeq->GetName(), *AnimSeq->GetFullName());
-				UAnimCompress* CompressionAlgorithm = NewObject<UAnimCompress_BitwiseCompressOnly>();
-				CompressionAlgorithm->RotationCompressionFormat = ACF_Float96NoW;
-				CompressionAlgorithm->TranslationCompressionFormat = ACF_None;
-				CompressionAlgorithm->ScaleCompressionFormat = ACF_Float96NoW;
-				AnimSeq->CompressionScheme = CompressionAlgorithm;
-				AnimSeq->RequestSyncAnimRecompression();
-
-				// Force an update.
-				AnimSeq->CompressCommandletVersion = 0;
-			}
-			
-			// Do not perform automatic recompression on animations marked as 'bDoNotOverrideCompression'
-			// Unless they have no compression scheme, or they're using automatic compression.
-			if (AnimSeq->bDoNotOverrideCompression && (AnimSeq->CompressionScheme != nullptr) && !AnimSeq->CompressionScheme->IsA(UAnimCompress_Automatic::StaticClass()))
+			// Do not perform recompression on animations marked as 'bDoNotOverrideCompression'
+			// Unless they have no compression scheme.
+			if (AnimSeq->bDoNotOverrideCompression && AnimSeq->BoneCompressionSettings != nullptr)
 			{
 				continue;
 			}
@@ -2317,16 +1951,13 @@ struct CompressAnimationsFunctor
 				NumAnimationsInPackage,
 				*PackageFileName);
 
-			// First set automatic compressor and call it.
-			// This will run through a bunch of compressors and pick the best.
-			// Problem is this is going to create a DDC key with 'Automatic Compressor'
-			UAnimCompress* CompressionAlgorithm = NewObject<UAnimCompress_Automatic>();
-			AnimSeq->CompressionScheme = static_cast<UAnimCompress*>(StaticDuplicateObject(CompressionAlgorithm, AnimSeq));
+			UE_LOG(LogPackageUtilities, Warning, TEXT("%s (%s) Resetting with to default compression settings."), *AnimSeq->GetName(), *AnimSeq->GetFullName());
+			AnimSeq->BoneCompressionSettings = nullptr;
+			AnimSeq->CurveCompressionSettings = nullptr;
 			AnimSeq->RequestAnimCompression(FRequestAnimCompressionParams(false, true, false));
 
-			// Automatic compression should have picked a suitable compressor that is not UAnimCompress_Automatic
-			// May still be automatic if we read the data from the DDC
-			if (!AnimSeq->CompressionScheme->IsA(UAnimCompress_Automatic::StaticClass()))
+			// Automatic compression should have picked a suitable compressor
+			if (!AnimSeq->IsCompressedDataValid())
 			{
 				// Update CompressCommandletVersion in that case, and create a proper DDC entry
 				// (with actual compressor)
@@ -2394,7 +2025,7 @@ struct CompressAnimationsFunctor
 /*		bDirtyPackage = bDirtyPackage || Package->IsDirty();*/
 
 		// If we need to save package, do so.
-		if( bDirtyPackage && !bAnalyze )
+		if( bDirtyPackage )
 		{
 			bool bCorrectlySaved = false;
 
diff --git a/Engine/Source/Programs/AutomationTool/AutomationUtils/ProjectParams.cs b/Engine/Source/Programs/AutomationTool/AutomationUtils/ProjectParams.cs
--- a/Engine/Source/Programs/AutomationTool/AutomationUtils/ProjectParams.cs
+++ b/Engine/Source/Programs/AutomationTool/AutomationUtils/ProjectParams.cs
@@ -1972,10 +1972,6 @@ namespace AutomationTool
 		{
 			return DetectedTargets.FindAll(Target => Target.Rules.Type == DesiredType).ConvertAll(Target => Target.TargetName);
 		}
-		private List<String> PrimaryTargetNamesOfType(TargetType DesiredType)
-		{
-			return DetectedTargets.FindAll(Target => (Target.Rules.Type == DesiredType && Target.Rules.IsPrimaryTarget)).ConvertAll(Target => Target.TargetName);
-		}
 
 		private String ChooseTarget(List<String> Targets, TargetType Type)
 		{
@@ -2115,7 +2111,7 @@ namespace AutomationTool
 				List<String> AvailableGameTargets = TargetNamesOfType(TargetType.Game);
 				List<String> AvailableClientTargets = TargetNamesOfType(TargetType.Client);
 				List<String> AvailableServerTargets = TargetNamesOfType(TargetType.Server);
-				List<String> AvailableEditorTargets = PrimaryTargetNamesOfType(TargetType.Editor);
+				List<String> AvailableEditorTargets = TargetNamesOfType(TargetType.Editor);
 
 				// That should cover all detected targets; Program targets are handled separately.
 				System.Diagnostics.Debug.Assert(DetectedTargets.Count == (AvailableGameTargets.Count + AvailableClientTargets.Count + AvailableServerTargets.Count + AvailableEditorTargets.Count));
@@ -2137,12 +2133,26 @@ namespace AutomationTool
 
 				if (AvailableEditorTargets.Count > 0)
 				{
-					if (AvailableEditorTargets.Count > 1)
+					string DefaultEditorTarget;
+
+					if (EngineConfigs[BuildHostPlatform.Current.Platform].GetString("/Script/BuildSettings.BuildSettings", "DefaultEditorTarget", out DefaultEditorTarget))
 					{
-						throw new AutomationException("There can be only one primary editor target per project.");
+						if (!AvailableEditorTargets.Contains(DefaultEditorTarget))
+						{
+							throw new AutomationException(string.Format("A default editor target '{0}' was specified in engine.ini but does not exist", DefaultEditorTarget));
+						}
+
+						EditorTarget = DefaultEditorTarget;
 					}
+					else
+					{
+						if (AvailableEditorTargets.Count > 1)
+						{
+							throw new AutomationException("Project contains multiple editor targets but no default is set in engine.ini");
+						}
 
-					EditorTarget = AvailableEditorTargets.First();
+						EditorTarget = AvailableEditorTargets.First();
+					}
 				}
 
 				if (AvailableServerTargets.Count > 0 && (DedicatedServer || Cook || CookOnTheFly)) // only if server is needed
diff --git a/Engine/Source/Programs/AutomationTool/Scripts/CopyBuildToStagingDirectory.Automation.cs b/Engine/Source/Programs/AutomationTool/Scripts/CopyBuildToStagingDirectory.Automation.cs
--- a/Engine/Source/Programs/AutomationTool/Scripts/CopyBuildToStagingDirectory.Automation.cs
+++ b/Engine/Source/Programs/AutomationTool/Scripts/CopyBuildToStagingDirectory.Automation.cs
@@ -1221,6 +1221,7 @@ public partial class Project : CommandUtils
 			case "input.ini":
 			case "scalability.ini":
 			case "runtimeoptions.ini":
+			case "installbundle.ini":
 				return true;
 			case "crypto.ini":
 			case "editor.ini":
@@ -3570,6 +3571,8 @@ public partial class Project : CommandUtils
 				// clean the staging directories first
 				foreach (var SC in DeployContextList)
 				{
+					SC.StageTargetPlatform.PreStage(Params, SC);
+
 					// write out the commandline file now so it can go into the manifest
 					WriteStageCommandline(Params, SC);
 					CreateStagingManifest(Params, SC);
diff --git a/Engine/Source/Programs/UnrealBuildTool/Configuration/ModuleRules.cs b/Engine/Source/Programs/UnrealBuildTool/Configuration/ModuleRules.cs
--- a/Engine/Source/Programs/UnrealBuildTool/Configuration/ModuleRules.cs
+++ b/Engine/Source/Programs/UnrealBuildTool/Configuration/ModuleRules.cs
@@ -607,6 +607,16 @@ namespace UnrealBuildTool
 			set { ShadowVariableWarningLevel = (value ? WarningLevel.Warning : WarningLevel.Off); }
 		}
 
+		/// <summary>
+		/// How to treat unsafe implicit type cast warnings (e.g., double->float or int64->int32)
+		/// </summary>
+		public WarningLevel UnsafeTypeCastWarningLevel
+		{
+			get { return UnsafeTypeCastWarningLevelPrivate ?? Target.UnsafeTypeCastWarningLevel; }
+			set { UnsafeTypeCastWarningLevelPrivate = value; }
+		}
+		private WarningLevel? UnsafeTypeCastWarningLevelPrivate;
+
 		/// <summary>
 		/// Enable warnings for using undefined identifiers in #if expressions
 		/// </summary>
diff --git a/Engine/Source/Programs/UnrealBuildTool/Platform/Windows/VCToolChain.cs b/Engine/Source/Programs/UnrealBuildTool/Platform/Windows/VCToolChain.cs
--- a/Engine/Source/Programs/UnrealBuildTool/Platform/Windows/VCToolChain.cs
+++ b/Engine/Source/Programs/UnrealBuildTool/Platform/Windows/VCToolChain.cs
@@ -499,8 +499,6 @@ namespace UnrealBuildTool
 				}
 
 				Arguments.Add("/wd4463"); // 4463 - overflow; assigning 1 to bit-field that can only hold values from -1 to 0
-
-				Arguments.Add("/wd4838"); // 4838: conversion from 'type1' to 'type2' requires a narrowing conversion
 			}
 
 			if(CompileEnvironment.bEnableUndefinedIdentifierWarnings)
@@ -514,6 +512,30 @@ namespace UnrealBuildTool
 					Arguments.Add("/w44668");
 				}
 			}
+
+			// The unsafe type cast warnings setting controls the following warnings currently:
+			//   4244: conversion from 'type1' to 'type2', possible loss of data
+			//   4838: conversion from 'type1' to 'type2' requires a narrowing conversion
+			//@TODO: FLOATPRECISION: Consider doing the following as well:
+			//   4267: 'var' : conversion from 'size_t' to 'type', possible loss of data
+			//   4305: 'identifier' : truncation from 'type1' to 'type2'
+			WarningLevel EffectiveCastWarningLevel = (Target.Platform == UnrealTargetPlatform.Win64) ? CompileEnvironment.UnsafeTypeCastWarningLevel : WarningLevel.Off;
+			if (EffectiveCastWarningLevel == WarningLevel.Error)
+ 			{
+ 				Arguments.Add("/we4244");
+				Arguments.Add("/we4838");
+ 			}
+ 			else if (EffectiveCastWarningLevel == WarningLevel.Warning)
+ 			{
+				// Note: The extra 4 is not a typo, /wLXXXX sets warning XXXX to level L
+ 				Arguments.Add("/w44244");
+				Arguments.Add("/w44838");
+			}
+			else
+ 			{
+ 				Arguments.Add("/wd4244");
+				Arguments.Add("/wd4838");
+ 			}
 		}
 
 		void AppendCLArguments_CPP(CppCompileEnvironment CompileEnvironment, List<string> Arguments)
@@ -1302,7 +1324,8 @@ namespace UnrealBuildTool
 
 				// Resource tool can run remotely if possible
 				CompileAction.bCanExecuteRemotely = true;
-
+				CompileAction.bCanExecuteRemotelyWithSNDBS = false;	// no tool template for SN-DBS results in warnings
+			
 				List<string> Arguments = new List<string>();
 
 				// Suppress header spew
@@ -1343,7 +1366,7 @@ namespace UnrealBuildTool
 					}
 				}
 
-				// Figure the icon to use. We can only use a custom icon when compiling to a project-specific intemediate directory (and not for the shared editor executable, for example).
+				// Figure the icon to use. We can only use a custom icon when compiling to a project-specific intermediate directory (and not for the shared editor executable, for example).
 				FileReference IconFile;
 				if(Target.ProjectFile != null && !CompileEnvironment.bUseSharedBuildEnvironment)
 				{
diff --git a/Engine/Source/Programs/UnrealBuildTool/System/CppCompileEnvironment.cs b/Engine/Source/Programs/UnrealBuildTool/System/CppCompileEnvironment.cs
--- a/Engine/Source/Programs/UnrealBuildTool/System/CppCompileEnvironment.cs
+++ b/Engine/Source/Programs/UnrealBuildTool/System/CppCompileEnvironment.cs
@@ -186,6 +186,11 @@ namespace UnrealBuildTool
 		/// </summary>
 		public WarningLevel ShadowVariableWarningLevel = WarningLevel.Warning;
 
+		/// <summary>
+		/// How to treat unsafe implicit type cast warnings (e.g., double->float or int64->int32)
+		/// </summary>
+		public WarningLevel UnsafeTypeCastWarningLevel = WarningLevel.Off;
+
 		/// <summary>
 		/// Whether to warn about the use of undefined identifiers in #if expressions
 		/// </summary>
@@ -401,6 +406,7 @@ namespace UnrealBuildTool
 			bEnableExceptions = Other.bEnableExceptions;
 			bEnableObjCExceptions = Other.bEnableObjCExceptions;
 			ShadowVariableWarningLevel = Other.ShadowVariableWarningLevel;
+			UnsafeTypeCastWarningLevel = Other.UnsafeTypeCastWarningLevel;
 			bUndefinedIdentifierWarningsAsErrors = Other.bUndefinedIdentifierWarningsAsErrors;
 			bEnableUndefinedIdentifierWarnings = Other.bEnableUndefinedIdentifierWarnings;
 			bOptimizeCode = Other.bOptimizeCode;
diff --git a/Engine/Source/Programs/UnrealHeaderTool/Private/CodeGenerator.cpp b/Engine/Source/Programs/UnrealHeaderTool/Private/CodeGenerator.cpp
--- a/Engine/Source/Programs/UnrealHeaderTool/Private/CodeGenerator.cpp
+++ b/Engine/Source/Programs/UnrealHeaderTool/Private/CodeGenerator.cpp
@@ -307,7 +307,7 @@ namespace
  */
 const TCHAR* FindIdentifierExactMatch(const TCHAR* StringBegin, const TCHAR* StringEnd, const FString& Identifier)
 {
-	int32 StringLen = StringEnd - StringBegin;
+	int32 StringLen = UE_PTRDIFF_TO_INT32(StringEnd - StringBegin);
 
 	// Check for exact match first.
 	if (FCString::Strncmp(StringBegin, *Identifier, StringLen) == 0)
@@ -375,7 +375,7 @@ int32 FindIdentifierExactMatch(const FString& String, const FString& Identifier)
 		return INDEX_NONE;
 	}
 
-	return IdentifierPtr - *String;
+	return UE_PTRDIFF_TO_INT32(IdentifierPtr - *String);
 }
 
 /**
@@ -2529,8 +2529,6 @@ void FNativeClassHeaderGenerator::ExportNativeGeneratedInitCode(FOutputDevice& O
 
 	if (ClassHasReplicatedProperties(Class))
 	{
-		Class->SetUpUhtReplicationData(nullptr);
-
 		Out.Logf(TEXT(
 			"\r\n"
 			"\tvoid %s::ValidateGeneratedRepEnums(const TArray<struct FRepRecord>& ClassReps) const\r\n"
@@ -3624,6 +3622,7 @@ void FNativeClassHeaderGenerator::ExportGeneratedStructBodyMacros(FOutputDevice&
 	const FString& SingletonName = GetSingletonName(Struct, OutReferenceGatherers.UniqueCrossModuleReferences);
 	const FString ChoppedSingletonName = SingletonName.LeftChop(2);
 
+	const FString RigVMParameterPrefix = TEXT("const FName& RigVMOperatorName, int32 RigVMOperatorIndex");
 	TArray<FString> RigVMVirtualFuncProlog, RigVMStubProlog;
 
 	// for RigVM methods we need to generated a macro used for implementing the static method
@@ -3632,14 +3631,9 @@ void FNativeClassHeaderGenerator::ExportGeneratedStructBodyMacros(FOutputDevice&
 	const FRigVMStructInfo* StructRigVMInfo = FHeaderParser::StructRigVMMap.Find(Struct);
 	if(StructRigVMInfo)
 	{
-		RigVMStubProlog.Add(FString::Printf(TEXT("ensure(RigVMOperands.Num() == %d);"), StructRigVMInfo->Members.Num()));
-		for (int32 ParameterIndex = 0; ParameterIndex < StructRigVMInfo->Members.Num(); ParameterIndex++)
-		{
-			const FRigVMParameter& Parameter = StructRigVMInfo->Members[ParameterIndex];
-			RigVMStubProlog.Add(FString::Printf(TEXT("const FRigVMOperand& %s_Operand = RigVMOperands[%d];"), *Parameter.Name, ParameterIndex));
-		}
-		RigVMStubProlog.Add(FString());
+		//RigVMStubProlog.Add(FString::Printf(TEXT("ensure(RigVMOperandMemory.Num() == %d);"), StructRigVMInfo->Members.Num()));
 
+		int32 OperandIndex = 0;
 		for (int32 ParameterIndex = 0; ParameterIndex < StructRigVMInfo->Members.Num(); ParameterIndex++)
 		{
 			const FRigVMParameter& Parameter = StructRigVMInfo->Members[ParameterIndex];
@@ -3647,7 +3641,7 @@ void FNativeClassHeaderGenerator::ExportGeneratedStructBodyMacros(FOutputDevice&
 			{
 				if (Parameter.IsArray() && !Parameter.IsConst() && !Parameter.MaxArraySize.IsEmpty())
 				{
-					RigVMVirtualFuncProlog.Add(FString::Printf(TEXT("%s.SetNumUninitialized( %s );"), *Parameter.Name, *Parameter.MaxArraySize));
+					RigVMVirtualFuncProlog.Add(FString::Printf(TEXT("%s.SetNum( %s );"), *Parameter.Name, *Parameter.MaxArraySize));
 				}
 				RigVMVirtualFuncProlog.Add(FString::Printf(TEXT("%s %s(%s);"), *Parameter.CastType, *Parameter.CastName, *Parameter.Name));
 			}
@@ -3655,25 +3649,38 @@ void FNativeClassHeaderGenerator::ExportGeneratedStructBodyMacros(FOutputDevice&
 			FString VariableType;
 			FString ExtractedType;
 			const FString& ParamTypeOriginal = Parameter.TypeOriginal(true);
+			const FString& ParamNameOriginal = Parameter.NameOriginal(false);
+
 			if (ParamTypeOriginal.StartsWith(TEXT("TArrayView"), ESearchCase::CaseSensitive))
 			{
 				ExtractedType = Parameter.ExtendedType().LeftChop(1).RightChop(1);
 				VariableType = ParamTypeOriginal;
+				
+				RigVMStubProlog.Add(FString::Printf(TEXT("%s %s = TArrayView<%s>((%s*)RigVMOperandMemory[%d], reinterpret_cast<uint64>(RigVMOperandMemory[%d]));"),
+					*VariableType,
+					*ParamNameOriginal,
+					*ExtractedType,
+					*ExtractedType,
+					OperandIndex,
+					OperandIndex + 1));
+
+				OperandIndex += 2;
 			}
 			else
 			{
 				VariableType = Parameter.TypeVariableRef(true);
 				ExtractedType = Parameter.TypeOriginal();
-			}
+				
+				FString ParameterCast = FString::Printf(TEXT("*(%s*)"), *ExtractedType);
 
-			const FString& ParamNameOriginal = Parameter.NameOriginal(false);
-			RigVMStubProlog.Add(FString::Printf(TEXT("%s %s = RigVMMemoryContainer[%s_Operand.GetContainerIndex()]->%s<%s>(%s_Operand, true);"),
-				*VariableType,
-				*ParamNameOriginal,
-				*ParamNameOriginal,
-				*Parameter.Getter,
-				*ExtractedType,
-				*ParamNameOriginal));
+				RigVMStubProlog.Add(FString::Printf(TEXT("%s %s = %sRigVMOperandMemory[%d];"),
+					*VariableType,
+					*ParamNameOriginal,
+					*ParameterCast,
+					OperandIndex));
+
+				OperandIndex++;
+			}
 		}
 
 		FString StructMembers = StructRigVMInfo->Members.Declarations(false, TEXT(", \\\r\n\t\t"), true, false);
@@ -3682,8 +3689,9 @@ void FNativeClassHeaderGenerator::ExportGeneratedStructBodyMacros(FOutputDevice&
 		for (const FRigVMMethodInfo& MethodInfo : StructRigVMInfo->Methods)
 		{
 			FString ParameterSuffix = MethodInfo.Parameters.Declarations(true, TEXT(", \\\r\n\t\t"));
+			FString RigVMParameterPrefix2 = RigVMParameterPrefix + FString((StructMembers.IsEmpty() && ParameterSuffix.IsEmpty()) ? TEXT("") : TEXT(", \\\r\n\t\t"));
 			OutGeneratedHeaderText.Logf(TEXT("#define %s_%s() \\\r\n"), *StructNameCPP, *MethodInfo.Name);
-			OutGeneratedHeaderText.Logf(TEXT("\t%s %s::Static%s( \\\r\n\t\t%s%s \\\r\n\t)\n"), *MethodInfo.ReturnType, *StructNameCPP, *MethodInfo.Name, *StructMembers, *ParameterSuffix);
+			OutGeneratedHeaderText.Logf(TEXT("\t%s %s::Static%s( \\\r\n\t\t%s%s%s \\\r\n\t)\n"), *MethodInfo.ReturnType, *StructNameCPP, *MethodInfo.Name, *RigVMParameterPrefix2, *StructMembers, *ParameterSuffix);
 		}
 		OutGeneratedHeaderText.Log(TEXT("\n"));
 	}
@@ -3707,9 +3715,48 @@ void FNativeClassHeaderGenerator::ExportGeneratedStructBodyMacros(FOutputDevice&
 			FString StructMembers = StructRigVMInfo->Members.Declarations(false, TEXT(",\r\n\t\t"), true, false);
 			for (const FRigVMMethodInfo& MethodInfo : StructRigVMInfo->Methods)
 			{
+				FString StructMembersForStub = StructRigVMInfo->Members.Names(false, TEXT(",\r\n\t\t\t"), false);
 				FString ParameterSuffix = MethodInfo.Parameters.Declarations(true, TEXT(",\r\n\t\t"));
-				RigVMMethodsDeclarations += FString::Printf(TEXT("\tstatic %s Static%s(\r\n\t\t%s%s\r\n\t);\r\n"), *MethodInfo.ReturnType, *MethodInfo.Name, *StructMembers, *ParameterSuffix);
-				RigVMMethodsDeclarations += FString::Printf(TEXT("\tstatic %s RigVM%s(\r\n\t\tconst FRigVMOperandArray& RigVMOperands,\r\n\t\tFRigVMMemoryContainerPtrArray& RigVMMemoryContainer,\r\n\t\tconst FRigVMUserDataArray& RigVMUserData\r\n\t);\r\n"), *MethodInfo.ReturnType, *MethodInfo.Name);
+				FString ParameterNamesSuffix = MethodInfo.Parameters.Names(true, TEXT(",\r\n\t\t\t"));
+				FString RigVMParameterPrefix2 = RigVMParameterPrefix + FString((StructMembers.IsEmpty() && ParameterSuffix.IsEmpty()) ? TEXT("") : TEXT(",\r\n\t\t"));
+				FString RigVMParameterPrefix4 = FString(TEXT("RigVMOperatorName,\r\n\t\t\tRigVMOperatorIndex")) + FString((StructMembersForStub.IsEmpty() && ParameterSuffix.IsEmpty()) ? TEXT("") : TEXT(",\r\n\t\t\t"));
+
+				RigVMMethodsDeclarations += FString::Printf(TEXT("\tstatic %s Static%s(\r\n\t\t%s%s%s\r\n\t);\r\n"), *MethodInfo.ReturnType, *MethodInfo.Name, *RigVMParameterPrefix2, *StructMembers, *ParameterSuffix);
+				RigVMMethodsDeclarations += FString::Printf(TEXT("\tFORCEINLINE static %s RigVM%s(\r\n\t\t%s,\r\n\t\tFRigVMOperandMemory RigVMOperandMemory,\r\n\t\tconst FRigVMUserDataArray& RigVMUserData\r\n\t)\r\n"), *MethodInfo.ReturnType, *MethodInfo.Name, *RigVMParameterPrefix);
+				RigVMMethodsDeclarations += FString::Printf(TEXT("\t{\r\n"));
+
+				// implement inline stub method body
+				if (MethodInfo.Parameters.Num() > 0)
+				{
+					//RigVMMethodsDeclarations += FString::Printf(TEXT("\t\tensure(RigVMUserData.Num() == %d);\r\n"), MethodInfo.Parameters.Num());
+					for (int32 ParameterIndex = 0; ParameterIndex < MethodInfo.Parameters.Num(); ParameterIndex++)
+					{
+						const FRigVMParameter& Parameter = MethodInfo.Parameters[ParameterIndex];
+						RigVMMethodsDeclarations += FString::Printf(TEXT("\t\t%s = *(%s*)RigVMUserData[%d];\r\n"), *Parameter.Declaration(), *Parameter.TypeNoRef(), ParameterIndex);
+					}
+					RigVMMethodsDeclarations += FString::Printf(TEXT("\t\t\r\n"));
+				}
+
+				if (RigVMStubProlog.Num() > 0)
+				{
+					for (const FString& RigVMStubPrologLine : RigVMStubProlog)
+					{
+						RigVMMethodsDeclarations += FString::Printf(TEXT("\t\t%s\r\n"), *RigVMStubPrologLine);
+					}
+					RigVMMethodsDeclarations += FString::Printf(TEXT("\t\t\r\n"));
+				}
+
+				RigVMMethodsDeclarations += FString::Printf(TEXT("\t\t%sStatic%s(\r\n\t\t\t%s%s%s\r\n\t\t);\r\n"), *MethodInfo.ReturnPrefix(), *MethodInfo.Name, *RigVMParameterPrefix4, *StructMembersForStub, *ParameterNamesSuffix);
+				RigVMMethodsDeclarations += FString::Printf(TEXT("\t}\r\n"));
+			}
+
+			for (const FRigVMParameter& StructMember : StructRigVMInfo->Members)
+			{
+				if (!StructMember.MaxArraySize.IsEmpty())
+				{
+					RigVMMethodsDeclarations += TEXT("\tvirtual int32 GetMaxArraySize(const FName& InMemberName, const FRigVMUserDataArray& RigVMUserData) override;\r\n");
+					break;
+				}
 			}
 		}
 
@@ -3751,7 +3798,7 @@ void FNativeClassHeaderGenerator::ExportGeneratedStructBodyMacros(FOutputDevice&
 		{
 			for (const FRigVMMethodInfo& MethodInfo : StructRigVMInfo->Methods)
 			{
-				Out.Logf(TEXT("\t\tFRigVMRegistry::Get().Register(TEXT(\"%s::%s\"), &%s::RigVM%s);\r\n"),
+				Out.Logf(TEXT("\t\tFRigVMRegistry::Get().Register(TEXT(\"%s::%s\"), &%s::RigVM%s, Singleton);\r\n"),
 					*StructNameCPP, *MethodInfo.Name, *StructNameCPP, *MethodInfo.Name);
 			}
 		}
@@ -3940,7 +3987,6 @@ void FNativeClassHeaderGenerator::ExportGeneratedStructBodyMacros(FOutputDevice&
 	if (StructRigVMInfo)
 	{
 		FString StructMembersForVirtualFunc = StructRigVMInfo->Members.Names(false, TEXT(",\r\n\t\t"), true);
-		FString StructMembersForStub = StructRigVMInfo->Members.Names(false, TEXT(",\r\n\t\t"), false);
 
 		for (const FRigVMMethodInfo& MethodInfo : StructRigVMInfo->Methods)
 		{
@@ -3948,6 +3994,8 @@ void FNativeClassHeaderGenerator::ExportGeneratedStructBodyMacros(FOutputDevice&
 
 			FString ParameterDeclaration = MethodInfo.Parameters.Declarations(false, TEXT(",\r\n\t\t"));
 			FString ParameterSuffix = MethodInfo.Parameters.Names(true, TEXT(",\r\n\t\t"));
+			FString RigVMParameterPrefix2 = RigVMParameterPrefix + FString((StructMembersForVirtualFunc.IsEmpty() && ParameterSuffix.IsEmpty()) ? TEXT("") : TEXT(",\r\n\t\t"));
+			FString RigVMParameterPrefix3 = FString(TEXT("NAME_None,\r\n\t\tINDEX_NONE")) + FString((StructMembersForVirtualFunc.IsEmpty() && ParameterSuffix.IsEmpty()) ? TEXT("") : TEXT(",\r\n\t\t"));
 
 			// implement the virtual function body.
 			Out.Logf(TEXT("%s %s::%s(%s)\r\n"), *MethodInfo.ReturnType, *StructNameCPP, *MethodInfo.Name, *ParameterDeclaration);
@@ -3962,40 +4010,39 @@ void FNativeClassHeaderGenerator::ExportGeneratedStructBodyMacros(FOutputDevice&
 				Out.Log(TEXT("\t\r\n"));
 			}
 
-			Out.Logf(TEXT("    %sStatic%s(\r\n\t\t%s%s\r\n\t);\n"), *MethodInfo.ReturnPrefix(), *MethodInfo.Name, *StructMembersForVirtualFunc, *ParameterSuffix);
+			Out.Logf(TEXT("    %sStatic%s(\r\n\t\t%s%s%s\r\n\t);\n"), *MethodInfo.ReturnPrefix(), *MethodInfo.Name, *RigVMParameterPrefix3, *StructMembersForVirtualFunc, *ParameterSuffix);
 			Out.Log(TEXT("}\r\n"));
+		}
 
-			Out.Log(TEXT("\r\n"));
-
-			// implement stub method body
-			Out.Logf(TEXT("%s %s::RigVM%s(const FRigVMOperandArray& RigVMOperands, FRigVMMemoryContainerPtrArray& RigVMMemoryContainer, const FRigVMUserDataArray& RigVMUserData)\r\n"), *MethodInfo.ReturnType, *StructNameCPP, *MethodInfo.Name);
-			Out.Log(TEXT("{\r\n"));
+		Out.Log(TEXT("\r\n"));
 
-			if (MethodInfo.Parameters.Num() > 0)
+		bool bHasGetMaxArraySize = false;
+		for (const FRigVMParameter& StructMember : StructRigVMInfo->Members)
+		{
+			if (!StructMember.MaxArraySize.IsEmpty())
 			{
-				Out.Logf(TEXT("\tensure(RigVMUserData.Num() == %d);\r\n"), MethodInfo.Parameters.Num());
-				for (int32 ParameterIndex = 0; ParameterIndex < MethodInfo.Parameters.Num(); ParameterIndex++)
-				{
-					const FRigVMParameter& Parameter = MethodInfo.Parameters[ParameterIndex];
-					Out.Logf(TEXT("\t%s = *(%s*)RigVMUserData[%d];\r\n"), *Parameter.Declaration(), *Parameter.TypeNoRef(), ParameterIndex);
-				}
-				Out.Log(TEXT("\t\r\n"));
+				bHasGetMaxArraySize = true;
+				break;
 			}
+		}
 
-			if(RigVMStubProlog.Num() > 0)
+		if (bHasGetMaxArraySize)
+		{
+			Out.Logf(TEXT("int32 %s::GetMaxArraySize(const FName& InMemberName, const FRigVMUserDataArray& RigVMUserData)\r\n"), *StructNameCPP);
+			Out.Log(TEXT("{\r\n"));
+			for (const FRigVMParameter& StructMember : StructRigVMInfo->Members)
 			{
-				for (const FString& RigVMStubPrologLine : RigVMStubProlog)
+				if (!StructMember.MaxArraySize.IsEmpty())
 				{
-					Out.Logf(TEXT("\t%s\r\n"), *RigVMStubPrologLine);
+					Out.Logf(TEXT("\tif(InMemberName == TEXT(\"%s\"))\r\n"), *StructMember.Name);
+					Out.Log(TEXT("\t{\r\n"));
+					Out.Logf(TEXT("\t\treturn %s;\r\n"), *StructMember.MaxArraySize);
+					Out.Log(TEXT("\t}\r\n"));
 				}
-				Out.Log(TEXT("\t\r\n"));
 			}
-
-			Out.Logf(TEXT("\t%sStatic%s(\r\n\t\t%s%s\r\n\t);\r\n"), *MethodInfo.ReturnPrefix(), *MethodInfo.Name, *StructMembersForStub, *ParameterSuffix);
-			Out.Log(TEXT("}\r\n"));
+			Out.Log(TEXT("\treturn INDEX_NONE;\r\n"));
+			Out.Log(TEXT("}\r\n\r\n"));
 		}
-
-		Out.Log(TEXT("\r\n"));
 	}
 }
 
@@ -4999,7 +5046,7 @@ void FNativeClassHeaderGenerator::ExportFunctionThunk(FUHTStringBuilder& RPCWrap
 	const TCHAR* ClassEnd   = ClassRange.End;
 	FString      ClassName  = Function->GetOwnerClass()->GetName();
 
-	FString ClassDefinition(ClassEnd - ClassStart, ClassStart);
+	FString ClassDefinition(UE_PTRDIFF_TO_INT32(ClassEnd - ClassStart), ClassStart);
 
 	bool bHasImplementation = HasIdentifierExactMatch(ClassDefinition, FunctionData.CppImplName);
 	bool bHasValidate = HasIdentifierExactMatch(ClassDefinition, FunctionData.CppValidationImplName);
@@ -5209,10 +5256,10 @@ void FNativeClassHeaderGenerator::ExportNativeFunctions(FOutputDevice& OutGenera
 		{
 			const TCHAR* ClassStart = ClassRange.Start;
 			const TCHAR* ClassEnd   = ClassRange.End;
-			FString ClassDefinition(ClassEnd - ClassStart, ClassStart);
+			FString ClassDefinition(UE_PTRDIFF_TO_INT32(ClassEnd - ClassStart), ClassStart);
 
 			FString FunctionName = Function->GetName();
-			int32 ClassDefinitionStartPosition = ClassStart - *SourceFile.GetContent();
+			int32 ClassDefinitionStartPosition = UE_PTRDIFF_TO_INT32(ClassStart - *SourceFile.GetContent());
 
 			int32 ImplementationPosition = FindIdentifierExactMatch(ClassDefinition, FunctionData.CppImplName);
 			bool bHasImplementation = ImplementationPosition != INDEX_NONE;
@@ -5684,6 +5731,17 @@ FNativeClassHeaderGenerator::FNativeClassHeaderGenerator(
 
 		FGeneratedCPP& GeneratedCPP = GeneratedCPPs.Emplace(SourceFile, MoveTemp(GeneratedSourceFilename));
 		GeneratedCPP.RelativeIncludes.Add(MoveTemp(ModuleRelativeFilename));
+
+		// This needs to be done outside of parallel blocks because it will modify UClass memory.
+		// Later calls to SetUpUhtReplicationData inside parallel blocks should be fine, because
+		// they will see the memory has already been set up, and just return the parent pointer.
+		for (UClass* Class : SourceFile->GetDefinedClasses())
+		{
+			if (ClassHasReplicatedProperties(Class))
+			{
+				Class->SetUpUhtReplicationData(nullptr);
+			}
+		}
 	}
 
 	const FManifestModule* ConstPackageManifest = GetPackageManifest(PackageName);
diff --git a/Engine/Source/Runtime/Core/Private/Containers/String.cpp b/Engine/Source/Runtime/Core/Private/Containers/String.cpp
--- a/Engine/Source/Runtime/Core/Private/Containers/String.cpp
+++ b/Engine/Source/Runtime/Core/Private/Containers/String.cpp
@@ -195,7 +195,7 @@ int32 FString::Find(const TCHAR* SubStr, ESearchCase::Type SearchCase, ESearchDi
 			? FCString::Stristr(Start, SubStr)
 			: FCString::Strstr(Start, SubStr);
 
-		return Tmp ? (Tmp-**this) : INDEX_NONE;
+		return Tmp ? UE_PTRDIFF_TO_INT32(Tmp-**this) : INDEX_NONE;
 	}
 	else
 	{
@@ -738,7 +738,7 @@ bool FString::ToBlob(const FString& Source,uint8* DestBuffer,const uint32 DestSi
 			ConvBuffer[0] = Source[Index];
 			ConvBuffer[1] = Source[Index + 1];
 			ConvBuffer[2] = Source[Index + 2];
-			DestBuffer[WriteIndex] = FCString::Atoi(ConvBuffer);
+			DestBuffer[WriteIndex] = (uint8)FCString::Atoi(ConvBuffer);
 		}
 		return true;
 	}
@@ -773,7 +773,7 @@ bool FString::ToHexBlob( const FString& Source, uint8* DestBuffer, const uint32
 		{
 			ConvBuffer[0] = Source[Index];
 			ConvBuffer[1] = Source[Index + 1];
-			DestBuffer[WriteIndex] = FCString::Strtoi( ConvBuffer, &End, 16 );
+			DestBuffer[WriteIndex] = (uint8)FCString::Strtoi( ConvBuffer, &End, 16 );
 		}
 		return true;
 	}
@@ -920,7 +920,7 @@ int32 FString::ParseIntoArray( TArray<FString>& OutArray, const TCHAR* pchDelim,
 		{
 			if (!InCullEmpty || At-Start)
 			{
-				OutArray.Emplace(At-Start,Start);
+				OutArray.Emplace(UE_PTRDIFF_TO_INT32(At-Start),Start);
 			}
 			Start = At + DelimLength;
 		}
@@ -1076,10 +1076,12 @@ FString FString::Replace(const TCHAR* From, const TCHAR* To, ESearchCase::Type S
 		// look for From in the remaining string
 		const TCHAR* FromLocation = SearchCase == ESearchCase::IgnoreCase ? FCString::Stristr(Travel, From) : FCString::Strstr(Travel, From);
 		if (!FromLocation)
+		{
 			break;
+		}
 
 		// copy everything up to FromLocation
-		Result.AppendChars(Travel, FromLocation - Travel);
+		Result.AppendChars(Travel, UE_PTRDIFF_TO_INT32(FromLocation - Travel));
 
 		// copy over the To
 		Result.AppendChars(To, ToLength);
@@ -1392,7 +1394,7 @@ FArchive& operator<<( FArchive& Ar, FString& A )
 			SaveNum = -SaveNum;
 		}
 
-		int32 MaxSerializeSize = Ar.GetMaxSerializeSize();
+		int64 MaxSerializeSize = Ar.GetMaxSerializeSize();
 		// Protect against network packets allocating too much memory
 		if ((MaxSerializeSize > 0) && (SaveNum > MaxSerializeSize))
 		{
@@ -1521,7 +1523,7 @@ int32 FindMatchingClosingParenthesis(const FString& TargetString, const int32 St
 		// Did we find the matching close parenthesis
 		if (ParenthesisCount == 0 && *(CurrPosition - 1) == TEXT(')'))
 		{
-			return StartSearch + ((CurrPosition - 1) - StartPosition);
+			return StartSearch + UE_PTRDIFF_TO_INT32((CurrPosition - 1) - StartPosition);
 		}
 	}
 
@@ -1557,7 +1559,7 @@ void FTextRange::CalculateLineRangesFromString(const FString& Input, TArray<FTex
 		const bool bIsWindowsNewLine = (*CurrentChar == '\r' && *(CurrentChar + 1) == '\n');
 		if (bIsWindowsNewLine || FChar::IsLinebreak(*CurrentChar))
 		{
-			const int32 LineEndIndex = (CurrentChar - InputStart);
+			const int32 LineEndIndex = UE_PTRDIFF_TO_INT32(CurrentChar - InputStart);
 			check(LineEndIndex >= LineBeginIndex);
 			LineRanges.Emplace(FTextRange(LineBeginIndex, LineEndIndex));
 
@@ -1565,7 +1567,7 @@ void FTextRange::CalculateLineRangesFromString(const FString& Input, TArray<FTex
 			{
 				++CurrentChar; // skip the \n of the \r\n chain
 			}
-			LineBeginIndex = (CurrentChar - InputStart) + 1; // The next line begins after the end of the current line
+			LineBeginIndex = UE_PTRDIFF_TO_INT32(CurrentChar - InputStart) + 1; // The next line begins after the end of the current line
 		}
 	}
 
diff --git a/Engine/Source/Runtime/Core/Private/IO/IoDispatcherFileBackend.cpp b/Engine/Source/Runtime/Core/Private/IO/IoDispatcherFileBackend.cpp
--- a/Engine/Source/Runtime/Core/Private/IO/IoDispatcherFileBackend.cpp
+++ b/Engine/Source/Runtime/Core/Private/IO/IoDispatcherFileBackend.cpp
@@ -176,9 +176,9 @@ EIoStoreResolveResult FFileIoStore::Resolve(FIoRequestImpl* Request)
 					ResolvedRequest.Request->IoBuffer = FIoBuffer(FIoBuffer::Wrap, TargetVa, ResolvedRequest.ResolvedSize);
 				}
 				PlatformImpl.BeginReadsForRequest(ResolvedRequest);
-				uint64 RequestBeginBlockIndex = ResolvedRequest.ResolvedOffset / CacheBlockSize;
-				uint64 RequestEndBlockIndex = (ResolvedRequest.ResolvedOffset + ResolvedRequest.ResolvedSize - 1) / CacheBlockSize + 1;
-				uint64 BlockCount = RequestEndBlockIndex - RequestBeginBlockIndex;
+				const uint32 RequestBeginBlockIndex = (uint32)(ResolvedRequest.ResolvedOffset / CacheBlockSize);
+				const uint32 RequestEndBlockIndex = (uint32)((ResolvedRequest.ResolvedOffset + ResolvedRequest.ResolvedSize - 1) / CacheBlockSize + 1);
+				const uint32 BlockCount = RequestEndBlockIndex - RequestBeginBlockIndex;
 				check(BlockCount > 0);
 				ReadBlockCached(RequestBeginBlockIndex, ResolvedRequest);
 				if (BlockCount > 1)
@@ -282,7 +282,7 @@ void FFileIoStore::ReadBlockCached(uint32 BlockIndex, const FFileIoStoreResolved
 	FFileIoStoreCacheBlockKey Key;
 	Key.FileHandle = ResolvedRequest.ResolvedFileHandle;
 	Key.BlockIndex = BlockIndex;
-	Key.Hash = HashCombine(ResolvedRequest.ResolvedFileHandle, BlockIndex);
+	Key.Hash = HashCombine(GetTypeHash(ResolvedRequest.ResolvedFileHandle), GetTypeHash(BlockIndex));
 	uint64 BlockOffset = uint64(BlockIndex) * uint64(CacheBlockSize);
 	FFileIoStoreReadBlock* CachedBlock = CachedBlocksMap.FindRef(Key);
 	if (!CachedBlock)
@@ -328,7 +328,8 @@ void FFileIoStore::ReadBlockCached(uint32 BlockIndex, const FFileIoStoreResolved
 		Scatter.Request = ResolvedRequest.Request;
 		Scatter.Dst = ResolvedRequest.Request->IoBuffer.Data() + BlockOffsetInRequest;
 		Scatter.Src = CachedBlock->Buffer + RequestStartOffsetInBlock;
-		Scatter.Size = RequestSizeInBlock;
+		check(RequestSizeInBlock <= TNumericLimits<uint32>::Max());
+		Scatter.Size = (uint32)RequestSizeInBlock;
 	}
 }
 
diff --git a/Engine/Source/Runtime/Core/Private/Misc/PathViews.cpp b/Engine/Source/Runtime/Core/Private/Misc/PathViews.cpp
--- a/Engine/Source/Runtime/Core/Private/Misc/PathViews.cpp
+++ b/Engine/Source/Runtime/Core/Private/Misc/PathViews.cpp
@@ -16,7 +16,7 @@ FStringView FPathViews::GetCleanFilename(const FStringView& InPath)
 {
 	if (const TCHAR* StartPos = Algo::FindLastByPredicate(InPath, UE4PathViews_Private::IsSlashOrBackslash))
 	{
-		return InPath.RightChop(StartPos - InPath.GetData() + 1);
+		return InPath.RightChop(UE_PTRDIFF_TO_INT32(StartPos - InPath.GetData() + 1));
 	}
 	return InPath;
 }
@@ -41,7 +41,7 @@ FStringView FPathViews::GetPath(const FStringView& InPath)
 {
 	if (const TCHAR* EndPos = Algo::FindLastByPredicate(InPath, UE4PathViews_Private::IsSlashOrBackslash))
 	{
-		return InPath.Left(EndPos - InPath.GetData());
+		return InPath.Left((FStringView::SizeType)(EndPos - InPath.GetData()));
 	}
 	return FStringView();
 }
@@ -51,7 +51,7 @@ FStringView FPathViews::GetExtension(const FStringView& InPath, bool bIncludeDot
 	if (const TCHAR* Dot = Algo::FindLast(GetCleanFilename(InPath), TEXT('.')))
 	{
 		const TCHAR* Extension = bIncludeDot ? Dot : Dot + 1;
-		return FStringView(Extension, InPath.GetData() + InPath.Len() - Extension);
+		return FStringView(Extension, (FStringView::SizeType)(InPath.GetData() + InPath.Len() - Extension));
 	}
 	return FStringView();
 }
@@ -61,7 +61,7 @@ FStringView FPathViews::GetPathLeaf(const FStringView& InPath)
 	if (const TCHAR* EndPos = Algo::FindLastByPredicate(InPath, UE4PathViews_Private::IsNotSlashOrBackslash))
 	{
 		++EndPos;
-		return GetCleanFilename(InPath.Left(EndPos - InPath.GetData()));
+		return GetCleanFilename(InPath.Left((FStringView::SizeType)(EndPos - InPath.GetData())));
 	}
 	return FStringView();
 }
@@ -70,7 +70,7 @@ void FPathViews::Split(const FStringView& InPath, FStringView& OutPath, FStringV
 {
 	const FStringView CleanName = GetCleanFilename(InPath);
 	const TCHAR* DotPos = Algo::FindLast(CleanName, TEXT('.'));
-	const FStringView::SizeType NameLen = DotPos ? DotPos - CleanName.GetData() : CleanName.Len();
+	const FStringView::SizeType NameLen = DotPos ? (FStringView::SizeType)(DotPos - CleanName.GetData()) : CleanName.Len();
 	OutPath = InPath.LeftChop(CleanName.Len() + 1);
 	OutName = CleanName.Left(NameLen);
 	OutExt = CleanName.RightChop(NameLen + 1);
diff --git a/Engine/Source/Runtime/Core/Private/ProfilingDebugging/CpuProfilerTrace.cpp b/Engine/Source/Runtime/Core/Private/ProfilingDebugging/CpuProfilerTrace.cpp
--- a/Engine/Source/Runtime/Core/Private/ProfilingDebugging/CpuProfilerTrace.cpp
+++ b/Engine/Source/Runtime/Core/Private/ProfilingDebugging/CpuProfilerTrace.cpp
@@ -124,7 +124,7 @@ void FCpuProfilerTraceInternal::FlushThreadBuffer(FThreadBuffer* InThreadBuffer)
 	uint8* BufferPtr = ThreadBuffer->Buffer + ThreadBuffer->BufferSize; \
 	FTraceUtils::Encode7bit((CycleDiff << 1) | 1ull, BufferPtr); \
 	FTraceUtils::Encode7bit(SpecId, BufferPtr); \
-	ThreadBuffer->BufferSize = BufferPtr - ThreadBuffer->Buffer; \
+	ThreadBuffer->BufferSize = (uint16)(BufferPtr - ThreadBuffer->Buffer); \
 	if (ThreadBuffer->BufferSize >= FCpuProfilerTraceInternal::FullBufferThreshold) \
 	{ \
 		FCpuProfilerTraceInternal::FlushThreadBuffer(ThreadBuffer); \
@@ -187,7 +187,7 @@ void FCpuProfilerTrace::OutputEndEvent(const Trace::FChannel& Channel)
 	ThreadBuffer->LastCycle = Cycle;
 	uint8* BufferPtr = ThreadBuffer->Buffer + ThreadBuffer->BufferSize;
 	FTraceUtils::Encode7bit(CycleDiff << 1, BufferPtr);
-	ThreadBuffer->BufferSize = BufferPtr - ThreadBuffer->Buffer;
+	ThreadBuffer->BufferSize = (uint16)(BufferPtr - ThreadBuffer->Buffer);
 	if ((FCpuProfilerTraceInternal::ThreadDepth == 0) | (ThreadBuffer->BufferSize >= FCpuProfilerTraceInternal::FullBufferThreshold))
 	{
 		FCpuProfilerTraceInternal::FlushThreadBuffer(ThreadBuffer);
diff --git a/Engine/Source/Runtime/Core/Private/ProfilingDebugging/MiscTrace.cpp b/Engine/Source/Runtime/Core/Private/ProfilingDebugging/MiscTrace.cpp
--- a/Engine/Source/Runtime/Core/Private/ProfilingDebugging/MiscTrace.cpp
+++ b/Engine/Source/Runtime/Core/Private/ProfilingDebugging/MiscTrace.cpp
@@ -136,7 +136,7 @@ void FMiscTrace::OutputBeginFrame(ETraceFrameType FrameType)
 	uint8 Buffer[9];
 	uint8* BufferPtr = Buffer;
 	FTraceUtils::Encode7bit(CycleDiff, BufferPtr);
-	uint16 BufferSize = BufferPtr - Buffer;
+	uint16 BufferSize = (uint16)(BufferPtr - Buffer);
 	if (FrameType == TraceFrameType_Game)
 	{
 		UE_TRACE_LOG(Misc, BeginGameFrame, FrameChannel, BufferSize)
@@ -161,7 +161,7 @@ void FMiscTrace::OutputEndFrame(ETraceFrameType FrameType)
 	uint8 Buffer[9];
 	uint8* BufferPtr = Buffer;
 	FTraceUtils::Encode7bit(CycleDiff, BufferPtr);
-	uint16 BufferSize = BufferPtr - Buffer;
+	uint16 BufferSize = (uint16)(BufferPtr - Buffer);
 	if (FrameType == TraceFrameType_Game)
 	{
 		UE_TRACE_LOG(Misc, EndGameFrame, FrameChannel, BufferSize)
diff --git a/Engine/Source/Runtime/Core/Private/UObject/UnrealNames.cpp b/Engine/Source/Runtime/Core/Private/UObject/UnrealNames.cpp
--- a/Engine/Source/Runtime/Core/Private/UObject/UnrealNames.cpp
+++ b/Engine/Source/Runtime/Core/Private/UObject/UnrealNames.cpp
@@ -27,6 +27,8 @@
 #include "Hash/CityHash.h"
 #include "Templates/AlignmentTemplates.h"
 
+PRAGMA_DISABLE_UNSAFE_TYPECAST_WARNINGS
+
 // Page protection to catch FNameEntry stomps
 #ifndef FNAME_WRITE_PROTECT_PAGES
 #define FNAME_WRITE_PROTECT_PAGES 0
@@ -1701,7 +1703,7 @@ static int32 GetLengthAndWidth(const WIDECHAR* Str, bool& bOutIsWide)
 
 	bOutIsWide = UserCharBits & 0xffffff80u;
 
-	return It - Str;
+	return UE_PTRDIFF_TO_INT32(It - Str);
 }
 
 static FWideStringViewWithWidth MakeUnconvertedView(const WIDECHAR* Str, int32 Len)
@@ -3171,3 +3173,5 @@ uint8** FNameDebugVisualizer::GetBlocks()
 
 	return ((FNamePool*)(NamePoolData))->GetBlocksForDebugVisualizer();
 }
+
+PRAGMA_ENABLE_UNSAFE_TYPECAST_WARNINGS
diff --git a/Engine/Source/Runtime/Core/Private/Windows/WindowsPlatformMisc.cpp b/Engine/Source/Runtime/Core/Private/Windows/WindowsPlatformMisc.cpp
--- a/Engine/Source/Runtime/Core/Private/Windows/WindowsPlatformMisc.cpp
+++ b/Engine/Source/Runtime/Core/Private/Windows/WindowsPlatformMisc.cpp
@@ -1186,17 +1186,17 @@ int MessageBoxExtInternal( EAppMsgType::Type MsgType, HWND HandleWnd, const TCHA
 		case EAppMsgType::YesNoYesAllNoAll:
 		{
 			GCancelButtonEnabled = false;
-			return DialogBox(GetModuleHandle(NULL), MAKEINTRESOURCE(IDD_YESNO2ALL), HandleWnd, MessageBoxDlgProc);
+			return (int)DialogBox(GetModuleHandle(NULL), MAKEINTRESOURCE(IDD_YESNO2ALL), HandleWnd, MessageBoxDlgProc);
 		}
 		case EAppMsgType::YesNoYesAllNoAllCancel:
 		{
 			GCancelButtonEnabled = true;
-			return DialogBox(GetModuleHandle(NULL), MAKEINTRESOURCE(IDD_YESNO2ALLCANCEL), HandleWnd, MessageBoxDlgProc);
+			return (int)DialogBox(GetModuleHandle(NULL), MAKEINTRESOURCE(IDD_YESNO2ALLCANCEL), HandleWnd, MessageBoxDlgProc);
 		}
 		case EAppMsgType::YesNoYesAll:
 		{
 			GCancelButtonEnabled = false;
-			return DialogBox(GetModuleHandle(NULL), MAKEINTRESOURCE(IDD_YESNOYESTOALL), HandleWnd, MessageBoxDlgProc);
+			return (int)DialogBox(GetModuleHandle(NULL), MAKEINTRESOURCE(IDD_YESNOYESTOALL), HandleWnd, MessageBoxDlgProc);
 		}
 	}
 
diff --git a/Engine/Source/Runtime/Core/Public/Containers/Array.h b/Engine/Source/Runtime/Core/Public/Containers/Array.h
--- a/Engine/Source/Runtime/Core/Public/Containers/Array.h
+++ b/Engine/Source/Runtime/Core/Public/Containers/Array.h
@@ -1591,7 +1591,7 @@ public:
 	FORCEINLINE void RemoveAt(SizeType Index, CountType Count, bool bAllowShrinking = true)
 	{
 		static_assert(!TAreTypesEqual<CountType, bool>::Value, "TArray::RemoveAt: unexpected bool passed as the Count argument");
-		RemoveAtImpl(Index, Count, bAllowShrinking);
+		RemoveAtImpl(Index, (SizeType)Count, bAllowShrinking);
 	}
 
 private:
@@ -2560,7 +2560,7 @@ private:
 	template <typename OtherElementType, typename OtherSizeType>
 	void CopyToEmpty(const OtherElementType* OtherData, OtherSizeType OtherNum, SizeType PrevMax, SizeType ExtraSlack)
 	{
-		SizeType NewNum = OtherNum;
+		SizeType NewNum = (SizeType)OtherNum;
 		checkf((OtherSizeType)NewNum == OtherNum, TEXT("Invalid number of elements to add to this array type: %llu"), (unsigned long long)NewNum);
 
 		checkSlow(ExtraSlack >= 0);
diff --git a/Engine/Source/Runtime/Core/Public/Containers/ArrayView.h b/Engine/Source/Runtime/Core/Public/Containers/ArrayView.h
--- a/Engine/Source/Runtime/Core/Public/Containers/ArrayView.h
+++ b/Engine/Source/Runtime/Core/Public/Containers/ArrayView.h
@@ -121,8 +121,10 @@ public:
 	>
 	FORCEINLINE TArrayView(OtherRangeType&& Other)
 		: DataPtr(ArrayViewPrivate::GetDataHelper(Forward<OtherRangeType>(Other)))
-		, ArrayNum(GetNum(Forward<OtherRangeType>(Other)))
 	{
+		const auto InCount = GetNum(Forward<OtherRangeType>(Other));
+		check((InCount >= 0) && ((sizeof(InCount) < sizeof(int32)) || (InCount <= static_cast<decltype(InCount)>(TNumericLimits<int32>::Max()))));
+		ArrayNum = (int32)InCount;
 	}
 
 	/**
diff --git a/Engine/Source/Runtime/Core/Public/HAL/Platform.h b/Engine/Source/Runtime/Core/Public/HAL/Platform.h
--- a/Engine/Source/Runtime/Core/Public/HAL/Platform.h
+++ b/Engine/Source/Runtime/Core/Public/HAL/Platform.h
@@ -63,6 +63,9 @@
 #if !defined(PLATFORM_LINUX)
 	#define PLATFORM_LINUX 0
 #endif
+#if !defined(PLATFORM_LINUXAARCH64)
+	#define PLATFORM_LINUXAARCH64 0
+#endif
 #if !defined(PLATFORM_SWITCH)
 	#define PLATFORM_SWITCH 0
 #endif
diff --git a/Engine/Source/Runtime/Core/Public/IO/IoDispatcher.h b/Engine/Source/Runtime/Core/Public/IO/IoDispatcher.h
--- a/Engine/Source/Runtime/Core/Public/IO/IoDispatcher.h
+++ b/Engine/Source/Runtime/Core/Public/IO/IoDispatcher.h
@@ -668,7 +668,7 @@ public:
 		return RequestedOffset;
 	}
 
-	uint32 GetSize() const
+	uint64 GetSize() const
 	{
 		return RequestedSize;
 	}
diff --git a/Engine/Source/Runtime/Core/Public/Misc/CString.h b/Engine/Source/Runtime/Core/Public/Misc/CString.h
--- a/Engine/Source/Runtime/Core/Public/Misc/CString.h
+++ b/Engine/Source/Runtime/Core/Public/Misc/CString.h
@@ -777,13 +777,13 @@ int32 TCString<T>::Strspn( const CharType* String, const CharType* Mask )
 			}
 		}
 
-		return StringIt - String;
+		return UE_PTRDIFF_TO_INT32(StringIt - String);
 
 	NextChar:
 		++StringIt;
 	}
 
-	return StringIt - String;
+	return UE_PTRDIFF_TO_INT32(StringIt - String);
 }
 
 template <typename T> FORCEINLINE
diff --git a/Engine/Source/Runtime/Core/Public/Misc/CoreMiscDefines.h b/Engine/Source/Runtime/Core/Public/Misc/CoreMiscDefines.h
--- a/Engine/Source/Runtime/Core/Public/Misc/CoreMiscDefines.h
+++ b/Engine/Source/Runtime/Core/Public/Misc/CoreMiscDefines.h
@@ -246,6 +246,12 @@ struct TStaticDeprecateExpression
 	}; \
 	enum class PREPROCESSOR_JOIN(EDeprecationMsg_, __LINE__) { Value = PREPROCESSOR_JOIN(FDeprecationMsg_, __LINE__)::condition(TStaticDeprecateExpression<!!(bExpression)>()) }
 
+// These defines are used to mark a difference between two pointers as expected to fit into the specified range
+// while still leaving something searchable if the surrounding code is updated to work with a 64 bit count/range
+// in the future
+#define UE_PTRDIFF_TO_INT32(argument) static_cast<int32>(argument)
+#define UE_PTRDIFF_TO_UINT32(argument) static_cast<uint32>(argument)
+
 /**
 * Makes a type non-copyable and non-movable by deleting copy/move constructors and assignment/move operators.
 * The macro should be placed in the public section of the type for better compiler diagnostic messages.
diff --git a/Engine/Source/Runtime/Core/Public/ProfilingDebugging/CpuProfilerTrace.h b/Engine/Source/Runtime/Core/Public/ProfilingDebugging/CpuProfilerTrace.h
--- a/Engine/Source/Runtime/Core/Public/ProfilingDebugging/CpuProfilerTrace.h
+++ b/Engine/Source/Runtime/Core/Public/ProfilingDebugging/CpuProfilerTrace.h
@@ -78,7 +78,7 @@ struct FCpuProfilerTrace
 	FCpuProfilerTrace::Shutdown();
 
 #define TRACE_CPUPROFILER_EVENT_SCOPE_ON_CHANNEL_STR(NameStr, Channel) \
-	static uint16 PREPROCESSOR_JOIN(__CpuProfilerEventSpecId, __LINE__); \
+	static uint32 PREPROCESSOR_JOIN(__CpuProfilerEventSpecId, __LINE__); \
 	if (bool(Channel) && PREPROCESSOR_JOIN(__CpuProfilerEventSpecId, __LINE__) == 0) { \
 		PREPROCESSOR_JOIN(__CpuProfilerEventSpecId, __LINE__) = FCpuProfilerTrace::OutputEventType(NameStr); \
 	} \
diff --git a/Engine/Source/Runtime/Core/Public/Windows/WindowsPlatformMath.h b/Engine/Source/Runtime/Core/Public/Windows/WindowsPlatformMath.h
--- a/Engine/Source/Runtime/Core/Public/Windows/WindowsPlatformMath.h
+++ b/Engine/Source/Runtime/Core/Public/Windows/WindowsPlatformMath.h
@@ -22,7 +22,9 @@ struct FWindowsPlatformMath : public TUnrealPlatformMathSSE4Base<FGenericPlatfor
 {
 #if PLATFORM_ENABLE_VECTORINTRINSICS
 	static FORCEINLINE bool IsNaN( float A ) { return _isnan(A) != 0; }
+	static FORCEINLINE bool IsNaN(double A) { return _isnan(A) != 0; }
 	static FORCEINLINE bool IsFinite( float A ) { return _finite(A) != 0; }
+	static FORCEINLINE bool IsFinite(double A) { return _finite(A) != 0; }
 
 	#pragma intrinsic( _BitScanReverse )
 	static FORCEINLINE uint32 FloorLog2(uint32 Value) 
diff --git a/Engine/Source/Runtime/CoreUObject/Private/UObject/Class.cpp b/Engine/Source/Runtime/CoreUObject/Private/UObject/Class.cpp
--- a/Engine/Source/Runtime/CoreUObject/Private/UObject/Class.cpp
+++ b/Engine/Source/Runtime/CoreUObject/Private/UObject/Class.cpp
@@ -4150,7 +4150,7 @@ void UClass::ValidateRuntimeReplicationData()
 				}
 			}
 
-			UE_LOG(LogClass, Warning, TEXT("Property %s::%s (SourceClass: %s) was not registered in GetLifetimeReplicatedProps. This property will not be replicated."),
+			UE_LOG(LogClass, Warning, TEXT("Property %s::%s (SourceClass: %s) was not registered in GetLifetimeReplicatedProps. This property will not be replicated. Use DISABLE_REPLICATED_PROPERTY if not replicating was intentional."),
 				*GetName(), *RepProp->GetName(), *RepProp->GetOwnerClass()->GetName());
 		}
 	}
diff --git a/Engine/Source/Runtime/CoreUObject/Public/UObject/UnrealType.h b/Engine/Source/Runtime/CoreUObject/Public/UObject/UnrealType.h
--- a/Engine/Source/Runtime/CoreUObject/Public/UObject/UnrealType.h
+++ b/Engine/Source/Runtime/CoreUObject/Public/UObject/UnrealType.h
@@ -1194,7 +1194,8 @@ class COREUOBJECT_API FNumericProperty : public FProperty
 	{
 		if (!TIsFloatingPoint<T>::Value)
 		{
-			return CanHoldDoubleValueInternal(Value);
+			//@TODO: FLOATPRECISION: This feels wrong, it might be losing precision before it tests to see if it's going to lose precision...
+			return CanHoldDoubleValueInternal((double)Value);
 		}
 		else if (TIsSigned<T>::Value)
 		{
@@ -1393,7 +1394,7 @@ public:
 				if (!Tag.EnumName.IsNone())
 				{
 					int64 PreviousValue = this->ReadEnumAsInt64(Slot, DefaultsStruct, Tag);
-					this->SetPropertyValue_InContainer(Data, PreviousValue, Tag.ArrayIndex);
+					this->SetPropertyValue_InContainer(Data, (TCppType)PreviousValue, Tag.ArrayIndex);
 				}
 				else
 				{
@@ -1450,17 +1451,17 @@ public:
 	virtual void SetIntPropertyValue(void* Data, uint64 Value) const override
 	{
 		check(TIsIntegral<TCppType>::Value);
-		TTypeFundamentals::SetPropertyValue(Data, Value);
+		TTypeFundamentals::SetPropertyValue(Data, (TCppType)Value);
 	}
 	virtual void SetIntPropertyValue(void* Data, int64 Value) const override
 	{
 		check(TIsIntegral<TCppType>::Value);
-		TTypeFundamentals::SetPropertyValue(Data, Value);
+		TTypeFundamentals::SetPropertyValue(Data, (TCppType)Value);
 	}
 	virtual void SetFloatingPointPropertyValue(void* Data, double Value) const override
 	{
 		check(TIsFloatingPoint<TCppType>::Value);
-		TTypeFundamentals::SetPropertyValue(Data, Value);
+		TTypeFundamentals::SetPropertyValue(Data, (TCppType)Value);
 	}
 	virtual void SetNumericPropertyValueFromString(void* Data, TCHAR const* Value) const override
 	{
@@ -1473,17 +1474,17 @@ public:
 	virtual int64 GetSignedIntPropertyValue(void const* Data) const override
 	{
 		check(TIsIntegral<TCppType>::Value);
-		return TTypeFundamentals::GetPropertyValue(Data);
+		return (int64)TTypeFundamentals::GetPropertyValue(Data);
 	}
 	virtual uint64 GetUnsignedIntPropertyValue(void const* Data) const override
 	{
 		check(TIsIntegral<TCppType>::Value);
-		return TTypeFundamentals::GetPropertyValue(Data);
+		return (uint64)TTypeFundamentals::GetPropertyValue(Data);
 	}
 	virtual double GetFloatingPointPropertyValue(void const* Data) const override
 	{
 		check(TIsFloatingPoint<TCppType>::Value);
-		return TTypeFundamentals::GetPropertyValue(Data);
+		return (double)TTypeFundamentals::GetPropertyValue(Data);
 	}
 	// End of FNumericProperty interface
 
@@ -3752,7 +3753,7 @@ public:
 			[LocalKeyPropForCapture](const void* A, const void* B) { return LocalKeyPropForCapture->Identical(A, B); }
 			))
 		{
-			int32 Idx = (Entry - (uint8*)Map->GetData(0, MapLayout)) / MapLayout.SetLayout.Size;
+			int32 Idx = (int32)((Entry - (uint8*)Map->GetData(0, MapLayout)) / MapLayout.SetLayout.Size);
 			RemoveAt(Idx);
 			return true;
 		}
diff --git a/Engine/Source/Runtime/Engine/Classes/Materials/Material.h b/Engine/Source/Runtime/Engine/Classes/Materials/Material.h
--- a/Engine/Source/Runtime/Engine/Classes/Materials/Material.h
+++ b/Engine/Source/Runtime/Engine/Classes/Materials/Material.h
@@ -767,6 +767,10 @@ public:
 	UPROPERTY(EditAnywhere, BlueprintReadOnly, Category=Mobile)
 	uint8 bUseLightmapDirectionality : 1;
 
+	/* Forward (including mobile) renderer: use preintegrated GF lut for simple IBL, but will use one more sampler. */
+	UPROPERTY(EditAnywhere, BlueprintReadOnly, Category = ForwardShading, meta = (DisplayName = "PreintegratedGF For Simple IBL"))
+	uint32 bForwardRenderUsePreintegratedGFForSimpleIBL : 1;
+
 	/* 
 	 * Forward renderer: enables multiple parallax-corrected reflection captures that blend together.
 	 * Mobile renderer: blend between nearest 3 reflection captures, but reduces the number of samplers available to the material as two more samplers will be used for reflection cubemaps.
@@ -802,10 +806,6 @@ public:
 	UPROPERTY(EditAnywhere, Category=Material, AdvancedDisplay)
 	uint8 Wireframe : 1;
 
-	/** Skips outputting velocity during the base pass. */
-	UPROPERTY(EditAnywhere, Category=Material, AdvancedDisplay, meta=(DisplayName = "Support accurate velocities from Vertex Deformation"))
-	uint8 bOutputVelocityOnBasePass : 1;
-
 #if WITH_EDITORONLY_DATA
 	UPROPERTY()
 	int32 EditorX;
diff --git a/Engine/Source/Runtime/Engine/Private/Animation/AnimInstanceProxy.cpp b/Engine/Source/Runtime/Engine/Private/Animation/AnimInstanceProxy.cpp
--- a/Engine/Source/Runtime/Engine/Private/Animation/AnimInstanceProxy.cpp
+++ b/Engine/Source/Runtime/Engine/Private/Animation/AnimInstanceProxy.cpp
@@ -2676,4 +2676,44 @@ FAnimBlueprintDebugData* FAnimInstanceProxy::GetAnimBlueprintDebugData() const
 	return nullptr;
 }
 
+void FAnimInstanceProxy::InitializeInputProxy(FAnimInstanceProxy* InputProxy, UAnimInstance* InAnimInstance)
+{
+	if (InAnimInstance && InputProxy)
+	{
+		InputProxy->Initialize(InAnimInstance);
+	}
+}
+
+void FAnimInstanceProxy::GatherInputProxyDebugData(FAnimInstanceProxy* InputProxy, FNodeDebugData& DebugData)
+{
+	if (InputProxy)
+	{
+		InputProxy->GatherDebugData(DebugData);
+	}
+}
+
+void FAnimInstanceProxy::CacheBonesInputProxy(FAnimInstanceProxy* InputProxy)
+{
+	if (InputProxy)
+	{
+		InputProxy->CacheBones();
+	}
+}
+
+void FAnimInstanceProxy::UpdateInputProxy(FAnimInstanceProxy* InputProxy, const FAnimationUpdateContext& Context)
+{
+	if (InputProxy)
+	{
+		InputProxy->UpdateAnimationNode(Context);
+	}
+}
+
+void FAnimInstanceProxy::EvaluateInputProxy(FAnimInstanceProxy* InputProxy, FPoseContext& Output)
+{
+	if (InputProxy)
+	{
+		InputProxy->Evaluate(Output);
+	}
+}
+
 #undef LOCTEXT_NAMESPACE
diff --git a/Engine/Source/Runtime/Engine/Private/Animation/AnimMontage.cpp b/Engine/Source/Runtime/Engine/Private/Animation/AnimMontage.cpp
--- a/Engine/Source/Runtime/Engine/Private/Animation/AnimMontage.cpp
+++ b/Engine/Source/Runtime/Engine/Private/Animation/AnimMontage.cpp
@@ -2509,11 +2509,11 @@ void FAnimMontageInstance::BranchingPointEventHandler(const FBranchingPointMarke
 	}
 }
 
-UAnimMontage* FAnimMontageInstance::InitializeMatineeControl(FName SlotName, USkeletalMeshComponent* SkeletalMeshComponent, UAnimSequenceBase* InAnimSequence, bool bLooping)
+UAnimMontage* FAnimMontageInstance::InitializeMatineeControl(FName SlotName, UAnimInstance* AnimInstance, UAnimSequenceBase* InAnimSequence, bool bLooping)
 {
 	UAnimMontage* MontageToPlay = Cast<UAnimMontage>(InAnimSequence);
 
-	if (UAnimSingleNodeInstance* SingleNodeInst = SkeletalMeshComponent->GetSingleNodeInstance())
+	if (UAnimSingleNodeInstance* SingleNodeInst = Cast<UAnimSingleNodeInstance>(AnimInstance))
 	{
 		// Single node anim instance
 		if (SingleNodeInst->GetCurrentAsset() != InAnimSequence)
@@ -2529,7 +2529,7 @@ UAnimMontage* FAnimMontageInstance::InitializeMatineeControl(FName SlotName, USk
 
 		return MontageToPlay;
 	}
-	else if (UAnimInstance* AnimInst = SkeletalMeshComponent->GetAnimInstance())
+	else if (UAnimInstance* AnimInst = AnimInstance)
 	{
 		UAnimMontage* PreviousMontage = nullptr;
 
@@ -2580,9 +2580,19 @@ UAnimMontage* FAnimMontageInstance::InitializeMatineeControl(FName SlotName, USk
 
 UAnimMontage* FAnimMontageInstance::SetMatineeAnimPositionInner(FName SlotName, USkeletalMeshComponent* SkeletalMeshComponent, UAnimSequenceBase* InAnimSequence, float InPosition, bool bLooping)
 {
-	UAnimMontage* PlayingMontage = InitializeMatineeControl(SlotName, SkeletalMeshComponent, InAnimSequence, bLooping);
-	UAnimInstance* AnimInst = SkeletalMeshComponent->GetAnimInstance();
-	if (UAnimSingleNodeInstance* SingleNodeInst = SkeletalMeshComponent->GetSingleNodeInstance())
+	if (SkeletalMeshComponent)
+	{
+		return SetMatineeAnimPositionInner(SlotName, SkeletalMeshComponent->GetAnimInstance(), InAnimSequence, InPosition, bLooping);
+	}
+
+	return nullptr;
+}
+
+UAnimMontage* FAnimMontageInstance::SetMatineeAnimPositionInner(FName SlotName, UAnimInstance* AnimInstance, UAnimSequenceBase* InAnimSequence, float InPosition, bool bLooping)
+{
+	UAnimMontage* PlayingMontage = InitializeMatineeControl(SlotName, AnimInstance, InAnimSequence, bLooping);
+	UAnimInstance* AnimInst = AnimInstance;
+	if (UAnimSingleNodeInstance* SingleNodeInst = Cast<UAnimSingleNodeInstance>(AnimInst))
 	{
 		if (SingleNodeInst->GetCurrentTime() != InPosition)
 		{
@@ -2613,16 +2623,36 @@ UAnimMontage* FAnimMontageInstance::SetMatineeAnimPositionInner(FName SlotName,
 }
 
 UAnimMontage* FAnimMontageInstance::PreviewMatineeSetAnimPositionInner(FName SlotName, USkeletalMeshComponent* SkeletalMeshComponent, UAnimSequenceBase* InAnimSequence, float InPosition, bool bLooping, bool bFireNotifies, float DeltaTime)
+{
+	if (SkeletalMeshComponent)
+	{
+		return PreviewMatineeSetAnimPositionInner(SlotName, SkeletalMeshComponent, SkeletalMeshComponent->GetAnimInstance(), InAnimSequence, InPosition, bLooping, bFireNotifies, DeltaTime);
+	}
+
+	return nullptr;
+}
+
+UAnimMontage* FAnimMontageInstance::PreviewSequencerMontagePosition(FName SlotName, USkeletalMeshComponent* SkeletalMeshComponent, int32& InOutInstanceId, UAnimSequenceBase* InAnimSequence, float InFromPosition, float InToPosition, float Weight, bool bLooping, bool bFireNotifies, bool bPlaying)
+{
+	if (SkeletalMeshComponent)
+	{
+		return PreviewSequencerMontagePosition(SlotName, SkeletalMeshComponent, SkeletalMeshComponent->GetAnimInstance(), InOutInstanceId, InAnimSequence, InFromPosition, InToPosition, Weight, bLooping, bFireNotifies, bPlaying);
+	}
+
+	return nullptr;
+}
+
+UAnimMontage* FAnimMontageInstance::PreviewMatineeSetAnimPositionInner(FName SlotName, USkeletalMeshComponent* SkeletalMeshComponent, UAnimInstance* AnimInstance, UAnimSequenceBase* InAnimSequence, float InPosition, bool bLooping, bool bFireNotifies, float DeltaTime)
 {
 	// Codepath for updating an animation when the skeletal mesh component is not going to be ticked (ie in editor)
-	UAnimMontage* PlayingMontage = InitializeMatineeControl(SlotName, SkeletalMeshComponent, InAnimSequence, bLooping);
+	UAnimMontage* PlayingMontage = InitializeMatineeControl(SlotName, AnimInstance, InAnimSequence, bLooping);
 
-	UAnimInstance* AnimInst = SkeletalMeshComponent->GetAnimInstance();
+	UAnimInstance* AnimInst = AnimInstance;
 
 	FAnimMontageInstance* MontageInstanceToUpdate = AnimInst && PlayingMontage ? AnimInst->GetActiveInstanceForMontage(PlayingMontage) : nullptr;
 	float PreviousPosition = InPosition;
 
-	if (UAnimSingleNodeInstance* SingleNodeInst = SkeletalMeshComponent->GetSingleNodeInstance())
+	if (UAnimSingleNodeInstance* SingleNodeInst = Cast<UAnimSingleNodeInstance>(AnimInst))
 	{
 		PreviousPosition = SingleNodeInst->GetCurrentTime();
 
@@ -2687,9 +2717,19 @@ UAnimMontage* FAnimMontageInstance::PreviewMatineeSetAnimPositionInner(FName Slo
 	return PlayingMontage;
 }
 
-UAnimMontage* FAnimMontageInstance::SetSequencerMontagePosition(FName SlotName, USkeletalMeshComponent* SkeletalMeshComponent, int32& InOutInstanceId, UAnimSequenceBase* InAnimSequence, float InFromPosition, float InToPosition, float Weight, bool bLooping, bool bInPlaying)
+UAnimMontage* FAnimMontageInstance::SetSequencerMontagePosition(FName SlotName, USkeletalMeshComponent* SkeletalMeshComponent, int32& InOutInstanceId, UAnimSequenceBase* InAnimSequence, float InFromPosition, float InToPosition, float Weight, bool bLooping, bool bPlaying)
+{
+	if (SkeletalMeshComponent)
+	{
+		return SetSequencerMontagePosition(SlotName, SkeletalMeshComponent->GetAnimInstance(), InOutInstanceId, InAnimSequence, InFromPosition, InToPosition, Weight, bLooping, bPlaying);
+	}
+
+	return nullptr;
+}
+
+UAnimMontage* FAnimMontageInstance::SetSequencerMontagePosition(FName SlotName, UAnimInstance* AnimInstance, int32& InOutInstanceId, UAnimSequenceBase* InAnimSequence, float InFromPosition, float InToPosition, float Weight, bool bLooping, bool bInPlaying)
 {
-	UAnimInstance* AnimInst = SkeletalMeshComponent->GetAnimInstance();
+	UAnimInstance* AnimInst = AnimInstance;
 	if (AnimInst)
 	{
 		UAnimMontage* PlayingMontage = nullptr;
@@ -2736,14 +2776,14 @@ UAnimMontage* FAnimMontageInstance::SetSequencerMontagePosition(FName SlotName,
 	return nullptr;
 }
 
-UAnimMontage* FAnimMontageInstance::PreviewSequencerMontagePosition(FName SlotName, USkeletalMeshComponent* SkeletalMeshComponent, int32& InOutInstanceId, UAnimSequenceBase* InAnimSequence, float InFromPosition, float InToPosition, float Weight, bool bLooping, bool bFireNotifies, bool bInPlaying)
+UAnimMontage* FAnimMontageInstance::PreviewSequencerMontagePosition(FName SlotName, USkeletalMeshComponent* SkeletalMeshComponent, UAnimInstance* AnimInstance, int32& InOutInstanceId, UAnimSequenceBase* InAnimSequence, float InFromPosition, float InToPosition, float Weight, bool bLooping, bool bFireNotifies, bool bInPlaying)
 {
-	UAnimInstance* AnimInst = SkeletalMeshComponent->GetAnimInstance();
+	UAnimInstance* AnimInst = AnimInstance;
 	if (AnimInst)
 	{
 		FAnimMontageInstance* MontageInstanceToUpdate = AnimInst->GetMontageInstanceForID(InOutInstanceId);
 
-		UAnimMontage* PlayingMontage = SetSequencerMontagePosition(SlotName, SkeletalMeshComponent, InOutInstanceId, InAnimSequence, InFromPosition, InToPosition, Weight, bLooping, bInPlaying);
+		UAnimMontage* PlayingMontage = SetSequencerMontagePosition(SlotName, AnimInst, InOutInstanceId, InAnimSequence, InFromPosition, InToPosition, Weight, bLooping, bInPlaying);
 		if (PlayingMontage)
 		{
 			// we have to get it again in case if this is new
diff --git a/Engine/Source/Runtime/Engine/Private/Animation/AnimNode_LinkedAnimGraph.cpp b/Engine/Source/Runtime/Engine/Private/Animation/AnimNode_LinkedAnimGraph.cpp
--- a/Engine/Source/Runtime/Engine/Private/Animation/AnimNode_LinkedAnimGraph.cpp
+++ b/Engine/Source/Runtime/Engine/Private/Animation/AnimNode_LinkedAnimGraph.cpp
@@ -208,6 +208,11 @@ void FAnimNode_LinkedAnimGraph::OnInitializeAnimInstance(const FAnimInstanceProx
 		// We have an instance but no instance class
 		TeardownInstance();
 	}
+
+	if(InstanceToRun)
+	{
+		InstanceToRun->GetProxyOnAnyThread<FAnimInstanceProxy>().InitializeObjects(InstanceToRun);
+	}
 }
 
 void FAnimNode_LinkedAnimGraph::TeardownInstance()
diff --git a/Engine/Source/Runtime/Engine/Private/AssetManager.cpp b/Engine/Source/Runtime/Engine/Private/AssetManager.cpp
--- a/Engine/Source/Runtime/Engine/Private/AssetManager.cpp
+++ b/Engine/Source/Runtime/Engine/Private/AssetManager.cpp
@@ -1663,34 +1663,34 @@ bool UAssetManager::FindMissingChunkList(const TArray<FSoftObjectPath>& AssetLis
 		GetAssetDataForPath(Asset, FoundData);
 		TSet<int32> FoundChunks, MissingChunks, ErrorChunks;
 
-		for (int32 ChunkId : FoundData.ChunkIDs)
+		for (int32 PakchunkId : FoundData.ChunkIDs)
 		{
-			if (!ChunkLocationCache.Contains(ChunkId))
+			if (!ChunkLocationCache.Contains(PakchunkId))
 			{
-				EChunkLocation::Type Location = ChunkInstall->GetChunkLocation(ChunkId);
+				EChunkLocation::Type Location = ChunkInstall->GetPakchunkLocation(PakchunkId);
 
 				// If chunk install thinks the chunk is available, we need to double check with the pak system that it isn't
 				// pending decryption
 				if (Location >= EChunkLocation::LocalSlow && Pak->AnyChunksAvailable())
 				{
-					Location = Pak->GetPakChunkLocation(ChunkId);
+					Location = Pak->GetPakChunkLocation(PakchunkId);
 				}
 
-				ChunkLocationCache.Add(ChunkId, Location);
+				ChunkLocationCache.Add(PakchunkId, Location);
 			}
-			EChunkLocation::Type ChunkLocation = ChunkLocationCache[ChunkId];
+			EChunkLocation::Type ChunkLocation = ChunkLocationCache[PakchunkId];
 
 			switch (ChunkLocation)
 			{			
 			case EChunkLocation::DoesNotExist:
-				ErrorChunks.Add(ChunkId);
+				ErrorChunks.Add(PakchunkId);
 				break;
 			case EChunkLocation::NotAvailable:
-				MissingChunks.Add(ChunkId);
+				MissingChunks.Add(PakchunkId);
 				break;
 			case EChunkLocation::LocalSlow:
 			case EChunkLocation::LocalFast:
-				FoundChunks.Add(ChunkId);
+				FoundChunks.Add(PakchunkId);
 				break;
 			}
 		}
@@ -1753,7 +1753,7 @@ void UAssetManager::AcquireChunkList(const TArray<int32>& ChunkList, FAssetManag
 
 	for (int32 MissingChunk : PendingChunkInstall->PendingChunks)
 	{
-		ChunkInstall->PrioritizeChunk(MissingChunk, Priority);
+		ChunkInstall->PrioritizePakchunk(MissingChunk, Priority);
 	}
 }
 
@@ -1847,9 +1847,9 @@ void UAssetManager::OnChunkDownloaded(uint32 ChunkId, bool bSuccess)
 			TArray<int32> NewPendingList;
 			
 			// Check all chunks if they are done or failed
-			for (int32 PendingChunkId : PendingChunkInstall.PendingChunks)
+			for (int32 PendingPakchunkId : PendingChunkInstall.PendingChunks)
 			{
-				EChunkLocation::Type ChunkLocation = ChunkInstall->GetChunkLocation(PendingChunkId);
+				EChunkLocation::Type ChunkLocation = ChunkInstall->GetPakchunkLocation(PendingPakchunkId);
 
 				switch (ChunkLocation)
 				{
@@ -1857,7 +1857,7 @@ void UAssetManager::OnChunkDownloaded(uint32 ChunkId, bool bSuccess)
 					bFailed = true;
 					break;
 				case EChunkLocation::NotAvailable:
-					NewPendingList.Add(PendingChunkId);
+					NewPendingList.Add(PendingPakchunkId);
 					break;
 				}
 			}
diff --git a/Engine/Source/Runtime/Engine/Private/Components/ActorComponent.cpp b/Engine/Source/Runtime/Engine/Private/Components/ActorComponent.cpp
--- a/Engine/Source/Runtime/Engine/Private/Components/ActorComponent.cpp
+++ b/Engine/Source/Runtime/Engine/Private/Components/ActorComponent.cpp
@@ -30,6 +30,7 @@
 #include "ComponentUtils.h"
 #include "Engine/Engine.h"
 #include "HAL/LowLevelMemTracker.h"
+#include "Net/Core/PushModel/PushModel.h"
 
 #if WITH_EDITOR
 #include "Kismet2/ComponentEditorUtils.h"
@@ -1757,7 +1758,7 @@ void UActorComponent::SetIsReplicated(bool bShouldReplicate)
 		{
 PRAGMA_DISABLE_DEPRECATION_WARNINGS
 			bReplicates = bShouldReplicate;
-			// MARK_PROPERTY_DIRTY_FROM_NAME(UActorComponent, bReplicates, this);
+			MARK_PROPERTY_DIRTY_FROM_NAME(UActorComponent, bReplicates, this);
 PRAGMA_ENABLE_DEPRECATION_WARNINGS
 
 			if (AActor* MyOwner = GetOwner())
@@ -1805,10 +1806,11 @@ void UActorComponent::GetLifetimeReplicatedProps( TArray< FLifetimeProperty > &
 		BPClass->GetLifetimeBlueprintReplicationList(OutLifetimeProps);
 	}
 
-PRAGMA_DISABLE_DEPRECATION_WARNINGS
-	DOREPLIFETIME( UActorComponent, bIsActive );
-	DOREPLIFETIME( UActorComponent, bReplicates );
-PRAGMA_ENABLE_DEPRECATION_WARNINGS
+	FDoRepLifetimeParams SharedParams;
+	SharedParams.bIsPushBased = true;
+
+	DOREPLIFETIME_WITH_PARAMS_FAST(UActorComponent, bIsActive, SharedParams);
+	DOREPLIFETIME_WITH_PARAMS_FAST(UActorComponent, bReplicates, SharedParams);
 }
 
 void UActorComponent::OnRep_IsActive()
@@ -1974,7 +1976,7 @@ void UActorComponent::SetIsReplicatedByDefault(const bool bNewReplicates)
 	{
 PRAGMA_DISABLE_DEPRECATION_WARNINGS
 		bReplicates = bNewReplicates;
-		// MARK_PROPERTY_DIRTY_FROM_NAME(UActorComponent, bReplicates, this);
+		MARK_PROPERTY_DIRTY_FROM_NAME(UActorComponent, bReplicates, this);
 PRAGMA_ENABLE_DEPRECATION_WARNINGS
 	}
 	else
@@ -1988,7 +1990,7 @@ void UActorComponent::SetActiveFlag(const bool bNewIsActive)
 {
 PRAGMA_DISABLE_DEPRECATION_WARNINGS
 	bIsActive = bNewIsActive;
-	// MARK_PROPERTY_DIRTY_FROM_NAME(UActorComponent, bIsActive, this);
+	MARK_PROPERTY_DIRTY_FROM_NAME(UActorComponent, bIsActive, this);
 PRAGMA_ENABLE_DEPRECATION_WARNINGS
 }
 
diff --git a/Engine/Source/Runtime/Engine/Private/Materials/Material.cpp b/Engine/Source/Runtime/Engine/Private/Materials/Material.cpp
--- a/Engine/Source/Runtime/Engine/Private/Materials/Material.cpp
+++ b/Engine/Source/Runtime/Engine/Private/Materials/Material.cpp
@@ -916,7 +916,6 @@ UMaterial::UMaterial(const FObjectInitializer& ObjectInitializer)
 	bEnableCrackFreeDisplacement = false;
 	bEnableAdaptiveTessellation = true;
 	MaxDisplacement = 0.0f;
-	bOutputVelocityOnBasePass = true;
 	bEnableSeparateTranslucency = true;
 	bEnableMobileSeparateTranslucency = false;
 	bEnableResponsiveAA = false;
diff --git a/Engine/Source/Runtime/Engine/Private/Materials/MaterialShared.cpp b/Engine/Source/Runtime/Engine/Private/Materials/MaterialShared.cpp
--- a/Engine/Source/Runtime/Engine/Private/Materials/MaterialShared.cpp
+++ b/Engine/Source/Runtime/Engine/Private/Materials/MaterialShared.cpp
@@ -1276,6 +1276,11 @@ bool FMaterialResource::IsUsingFullPrecision() const
 	return Material->bUseFullPrecision;
 }
 
+bool FMaterialResource::IsUsingPreintegratedGFForSimpleIBL() const
+{
+	return Material->bForwardRenderUsePreintegratedGFForSimpleIBL;
+}
+
 bool FMaterialResource::IsUsingHQForwardReflections() const
 {
 	return Material->bUseHQForwardReflections;
@@ -1286,11 +1291,6 @@ bool FMaterialResource::IsUsingPlanarForwardReflections() const
 	return Material->bUsePlanarForwardReflections;
 }
 
-bool FMaterialResource::OutputsVelocityOnBasePass() const
-{
-	return Material->bOutputVelocityOnBasePass && !IsUIMaterial();
-}
-
 bool FMaterialResource::IsNonmetal() const
 {
 	return !Material->bUseMaterialAttributes ?
@@ -1695,6 +1695,7 @@ void FMaterial::SetupMaterialEnvironment(
 	OutEnvironment.SetDefine(TEXT("MATERIAL_TANGENTSPACENORMAL"), IsTangentSpaceNormal());
 	OutEnvironment.SetDefine(TEXT("GENERATE_SPHERICAL_PARTICLE_NORMALS"),ShouldGenerateSphericalParticleNormals());
 	OutEnvironment.SetDefine(TEXT("MATERIAL_USES_SCENE_COLOR_COPY"), RequiresSceneColorCopy_GameThread());
+	OutEnvironment.SetDefine(TEXT("MATERIAL_USE_PREINTEGRATED_GF"), IsUsingPreintegratedGFForSimpleIBL());
 	OutEnvironment.SetDefine(TEXT("MATERIAL_HQ_FORWARD_REFLECTIONS"), IsUsingHQForwardReflections());
 	OutEnvironment.SetDefine(TEXT("MATERIAL_PLANAR_FORWARD_REFLECTIONS"), IsUsingPlanarForwardReflections());
 	OutEnvironment.SetDefine(TEXT("MATERIAL_NONMETAL"), IsNonmetal());
diff --git a/Engine/Source/Runtime/Engine/Private/StaticMesh.cpp b/Engine/Source/Runtime/Engine/Private/StaticMesh.cpp
--- a/Engine/Source/Runtime/Engine/Private/StaticMesh.cpp
+++ b/Engine/Source/Runtime/Engine/Private/StaticMesh.cpp
@@ -236,6 +236,14 @@ FArchive& operator<<(FArchive& Ar, FStaticMeshSection& Section)
 	Ar << Section.bEnableCollision;
 	Ar << Section.bCastShadow;
 
+	if(
+		!(Ar.IsLoading() && 
+		  Ar.CustomVer(FRenderingObjectVersion::GUID) < FRenderingObjectVersion::StaticMeshSectionForceOpaqueField)
+	  )
+	{
+		Ar << Section.bForceOpaque;
+	}
+
 #if WITH_EDITORONLY_DATA
 	if((!Ar.IsCooking() && !Ar.IsFilterEditorOnly()) || (Ar.IsCooking() && Ar.CookingTarget()->HasEditorOnlyData()))
 	{
@@ -1212,6 +1220,7 @@ void FStaticMeshLODResources::InitResources(UStaticMesh* Parent)
 					Segment.VertexBufferOffset = 0;
 					Segment.FirstPrimitive = Section.FirstIndex / 3;
 					Segment.NumPrimitives = Section.NumTriangles;
+					Segment.bForceOpaque = Section.bForceOpaque;
 					GeometrySections.Add(Segment);
 					Initializer.TotalPrimitiveCount += Section.NumTriangles;
 				}
@@ -1661,6 +1670,7 @@ void FStaticMeshRenderData::ResolveSectionInfo(UStaticMesh* Owner)
 			Section.MaterialIndex = Info.MaterialIndex;
 			Section.bEnableCollision = Info.bEnableCollision;
 			Section.bCastShadow = Info.bCastShadow;
+			Section.bForceOpaque = Info.bForceOpaque;
 		}
 
 		// Arbitrary constant used as a base in Pow(K, LODIndex) that achieves much the same progression as a
@@ -2072,7 +2082,7 @@ static void SerializeBuildSettingsForDDC(FArchive& Ar, FMeshBuildSettings& Build
 // differences, etc.) replace the version GUID below with a new one.
 // In case of merge conflicts with DDC versions, you MUST generate a new GUID
 // and set this new GUID as the version.                                       
-#define STATICMESH_DERIVEDDATA_VER TEXT("D819AE82DB6A4CE0891F68BD81CFC2A8")
+#define STATICMESH_DERIVEDDATA_VER TEXT("AC59C886D91E43A1A5EC4DC250ED0E6A")
 
 static const FString& GetStaticMeshDerivedDataVersion()
 {
diff --git a/Engine/Source/Runtime/Engine/Private/Streaming/SkeletalMeshUpdate.cpp b/Engine/Source/Runtime/Engine/Private/Streaming/SkeletalMeshUpdate.cpp
--- a/Engine/Source/Runtime/Engine/Private/Streaming/SkeletalMeshUpdate.cpp
+++ b/Engine/Source/Runtime/Engine/Private/Streaming/SkeletalMeshUpdate.cpp
@@ -89,7 +89,8 @@ void FSkeletalMeshStreamIn::FIntermediateBuffers::SafeRelease()
 	TexCoordVertexBuffer.SafeRelease();
 	PositionVertexBuffer.SafeRelease();
 	ColorVertexBuffer.SafeRelease();
-	SkinWeightVertexBuffer.SafeRelease();
+	SkinWeightVertexBuffer.DataVertexBufferRHI.SafeRelease();
+	SkinWeightVertexBuffer.LookupVertexBufferRHI.SafeRelease();
 	ClothVertexBuffer.SafeRelease();
 	IndexBuffer.SafeRelease();
 	AdjacencyIndexBuffer.SafeRelease();
@@ -117,7 +118,8 @@ void FSkeletalMeshStreamIn::FIntermediateBuffers::CheckIsNull() const
 		&& !TexCoordVertexBuffer
 		&& !PositionVertexBuffer
 		&& !ColorVertexBuffer
-		&& !SkinWeightVertexBuffer
+		&& !SkinWeightVertexBuffer.DataVertexBufferRHI
+		&& !SkinWeightVertexBuffer.LookupVertexBufferRHI
 		&& !ClothVertexBuffer
 		&& !IndexBuffer
 		&& !AdjacencyIndexBuffer
diff --git a/Engine/Source/Runtime/Engine/Private/Texture2D.cpp b/Engine/Source/Runtime/Engine/Private/Texture2D.cpp
--- a/Engine/Source/Runtime/Engine/Private/Texture2D.cpp
+++ b/Engine/Source/Runtime/Engine/Private/Texture2D.cpp
@@ -1009,7 +1009,7 @@ int32 UTexture2D::CalcTextureMemorySize( int32 MipCount ) const
 		else
 		{
 			const FIntPoint MipExtents = CalcMipMapExtent(SizeX, SizeY, Format, FirstMip);
-			Size = (int32)RHICalcTexture2DPlatformSize(MipExtents.X, MipExtents.Y, Format, MipCount, 1, TexCreateFlags, FRHIResourceCreateInfo(PlatformData->GetExtData()), TextureAlign);
+			Size = (int32)RHICalcTexture2DPlatformSize(MipExtents.X, MipExtents.Y, Format, FMath::Max(1, MipCount), 1, TexCreateFlags, FRHIResourceCreateInfo(PlatformData->GetExtData()), TextureAlign);
 		}
 	}
 	return Size;
diff --git a/Engine/Source/Runtime/Engine/Private/UnrealEngine.cpp b/Engine/Source/Runtime/Engine/Private/UnrealEngine.cpp
--- a/Engine/Source/Runtime/Engine/Private/UnrealEngine.cpp
+++ b/Engine/Source/Runtime/Engine/Private/UnrealEngine.cpp
@@ -228,6 +228,7 @@ UnrealEngine.cpp: Implements the UEngine class and helpers.
 
 #include "Particles/ParticleSystemManager.h"
 #include "Components/SkinnedMeshComponent.h"
+#include "StudioAnalytics.h"
 
 DEFINE_LOG_CATEGORY(LogEngine);
 IMPLEMENT_MODULE( FEngineModule, Engine );
@@ -1009,10 +1010,8 @@ public:
 			{
 				check( ImageWrapper->GetWidth() == InWidth );
 				check( ImageWrapper->GetHeight() == InHeight );
-				const TArray<uint8>* RawData = NULL;
-				if ( ImageWrapper->GetRaw( ERGBFormat::RGBA, 8, RawData ) )	// @todo CB: Eliminate image copy here? (decompress straight to buffer)
+				if ( ImageWrapper->GetRaw( ERGBFormat::RGBA, 8, OutUncompressedData) )
 				{
-					OutUncompressedData = *RawData;
 					bSucceeded = true;
 				}
 			}
@@ -1764,6 +1763,7 @@ void UEngine::PreExit()
 	ShutdownRenderingCVarsCaching();
 	const bool bIsEngineShutdown = true;
 	FEngineAnalytics::Shutdown(bIsEngineShutdown);
+	FStudioAnalytics::Shutdown();
 	if (ScreenSaverInhibitor)
 	{
 		// Resume the thread to avoid a deadlock while waiting for finish.
@@ -3543,11 +3543,9 @@ struct FSortedAnimAsset
 	float		SequenceLength;
 	float		RateScale;
 	int32		NumCurves;
-	FString		TranslationFormat;
-	FString		RotationFormat;
-	FString		ScaleFormat;
+	FString		DebugString;
 
-	FSortedAnimAsset(int32 InNumKB, int32 InMaxKB, int32 InResKBExc, int32 InResKBInc, FName InAnimAssetType, FString InName, int32 InNumKeys, float InSequenceLength, float InRateScale, int32 InNumCurves, FString InTranslationFormat, FString InRotationFormat, FString InScaleFormat)
+	FSortedAnimAsset(int32 InNumKB, int32 InMaxKB, int32 InResKBExc, int32 InResKBInc, FName InAnimAssetType, FString InName, int32 InNumKeys, float InSequenceLength, float InRateScale, int32 InNumCurves, const FString& InDebugString)
 		: NumKB(InNumKB)
 		, MaxKB(InMaxKB)
 		, ResKBExc(InResKBExc)
@@ -3558,9 +3556,7 @@ struct FSortedAnimAsset
 		, SequenceLength(InSequenceLength)
 		, RateScale(InRateScale)
 		, NumCurves(InNumCurves)
-		, TranslationFormat(InTranslationFormat)
-		, RotationFormat(InRotationFormat)
-		, ScaleFormat(InScaleFormat)
+		, DebugString(InDebugString)
 	{}
 };
 struct FCompareFSortedAnimAsset
@@ -5610,9 +5606,7 @@ bool UEngine::HandleListAnimsCommand(const TCHAR* Cmd, FOutputDevice& Ar)
 		float SequenceLength = 0.f;
 		float RateScale = 1.f;
 		int32 NumCurves = 0;
-		FString TranslationFormat;
-		FString RotationFormat;
-		FString ScaleFormat;
+		FString DebugString;
 
 		if (UAnimSequence* AnimSeq = Cast<UAnimSequence>(AnimAsset))
 		{
@@ -5621,10 +5615,7 @@ bool UEngine::HandleListAnimsCommand(const TCHAR* Cmd, FOutputDevice& Ar)
 			RateScale = AnimSeq->RateScale;
 			NumCurves = AnimSeq->RawCurveData.FloatCurves.Num();
 
-			const FUECompressedAnimData& CompressedData = AnimSeq->CompressedData.CompressedDataStructure;
-			TranslationFormat = FAnimationUtils::GetAnimationCompressionFormatString(CompressedData.TranslationCompressionFormat);
-			RotationFormat = FAnimationUtils::GetAnimationCompressionFormatString(CompressedData.RotationCompressionFormat);
-			ScaleFormat = FAnimationUtils::GetAnimationCompressionFormatString(CompressedData.ScaleCompressionFormat);
+			DebugString = AnimSeq->CompressedData.CompressedDataStructure->GetDebugString();
 		}
 
 		new(SortedAnimAssets) FSortedAnimAsset(
@@ -5638,9 +5629,7 @@ bool UEngine::HandleListAnimsCommand(const TCHAR* Cmd, FOutputDevice& Ar)
 			SequenceLength,
 			RateScale,
 			NumCurves,
-			TranslationFormat,
-			RotationFormat,
-			ScaleFormat
+			DebugString
 		);
 	}
 
@@ -5658,9 +5647,9 @@ bool UEngine::HandleListAnimsCommand(const TCHAR* Cmd, FOutputDevice& Ar)
 	for (const FSortedAnimAsset& SortedAsset : SortedAnimAssets)
 	{
 
-		Ar.Logf(TEXT(",%i, %i, %i, %i, %s, %s, %f, %f, %d, %s, %s, %s"),
+		Ar.Logf(TEXT(",%i, %i, %i, %i, %s, %s, %f, %f, %d, %s"),
 			SortedAsset.NumKB, SortedAsset.MaxKB, SortedAsset.ResKBExc, SortedAsset.ResKBInc, *SortedAsset.AnimAssetType.ToString(), *SortedAsset.Name, SortedAsset.SequenceLength, 
-			SortedAsset.RateScale, SortedAsset.NumCurves, *SortedAsset.TranslationFormat, *SortedAsset.RotationFormat, *SortedAsset.ScaleFormat);
+			SortedAsset.RateScale, SortedAsset.NumCurves, *SortedAsset.DebugString);
 
 		TotalNumKB += SortedAsset.NumKB;
 		TotalMaxKB += SortedAsset.MaxKB;
@@ -12351,7 +12340,7 @@ bool UEngine::LoadMap( FWorldContext& WorldContext, FURL URL, class UPendingNetG
 	GInitRunaway();
 
 #if !UE_BUILD_SHIPPING
-	const bool bOldWorldWasShowingCollisionForHiddenComponents = WorldContext.World() && WorldContext.World()->bCreateRenderStateForHiddenComponents;
+	const bool bOldWorldWasShowingCollisionForHiddenComponents = WorldContext.World() && WorldContext.World()->bCreateRenderStateForHiddenComponentsWithCollsion;
 #endif
 
 	// Unload the current world
@@ -12665,7 +12654,7 @@ bool UEngine::LoadMap( FWorldContext& WorldContext, FURL URL, class UPendingNetG
 	WorldContext.World()->WorldType = WorldContext.WorldType;
 
 #if !UE_BUILD_SHIPPING
-	GWorld->bCreateRenderStateForHiddenComponents = bOldWorldWasShowingCollisionForHiddenComponents;
+	GWorld->bCreateRenderStateForHiddenComponentsWithCollsion = bOldWorldWasShowingCollisionForHiddenComponents;
 #endif
 
 	// Fixme: hacky but we need to set PackageFlags here if we are in a PIE Context.
diff --git a/Engine/Source/Runtime/Engine/Private/VT/UploadingVirtualTexture.cpp b/Engine/Source/Runtime/Engine/Private/VT/UploadingVirtualTexture.cpp
--- a/Engine/Source/Runtime/Engine/Private/VT/UploadingVirtualTexture.cpp
+++ b/Engine/Source/Runtime/Engine/Private/VT/UploadingVirtualTexture.cpp
@@ -47,9 +47,9 @@ static EAsyncIOPriorityAndFlags GetAsyncIOPriority(EVTRequestPagePriority Priori
 {
 	switch (Priority)
 	{
-	case EVTRequestPagePriority::High: return (EAsyncIOPriorityAndFlags)GVirtualTextureIOPriority_HighPagePri;
-	case EVTRequestPagePriority::Normal: return (EAsyncIOPriorityAndFlags)GVirtualTextureIOPriority_NormalPagePri;
-	default: check(false); return (EAsyncIOPriorityAndFlags)GVirtualTextureIOPriority_NormalPagePri;
+	case EVTRequestPagePriority::High: return (EAsyncIOPriorityAndFlags)GVirtualTextureIOPriority_HighPagePri | EAsyncIOPriorityAndFlags::AIOP_FLAG_DONTCACHE;
+	case EVTRequestPagePriority::Normal: return (EAsyncIOPriorityAndFlags)GVirtualTextureIOPriority_NormalPagePri | EAsyncIOPriorityAndFlags::AIOP_FLAG_DONTCACHE;
+	default: check(false); return (EAsyncIOPriorityAndFlags)GVirtualTextureIOPriority_NormalPagePri | EAsyncIOPriorityAndFlags::AIOP_FLAG_DONTCACHE;
 	}
 }
 
diff --git a/Engine/Source/Runtime/Launch/Private/LaunchEngineLoop.cpp b/Engine/Source/Runtime/Launch/Private/LaunchEngineLoop.cpp
--- a/Engine/Source/Runtime/Launch/Private/LaunchEngineLoop.cpp
+++ b/Engine/Source/Runtime/Launch/Private/LaunchEngineLoop.cpp
@@ -880,7 +880,7 @@ bool LaunchCheckForFileOverride(const TCHAR* CmdLine, bool& OutFileOverrideFound
 void LaunchCheckForCommandLineAliases(bool& bChanged)
 {
 	bChanged = false;
-	if (FPaths::IsProjectFilePathSet())
+	if (FPaths::ProjectDir().Len() > 0)
 	{
 		// Pass GeneratedConfigDir as nullptr instead of FPaths::GeneratedConfigDir() so that saved directory is not cached before -saveddir argument can be added
 		const TCHAR* GeneratedConfigDir = nullptr;
@@ -948,12 +948,26 @@ void LaunchCheckForCmdLineFile(const TCHAR* CmdLine, bool& bChanged)
 				FString FileCmds;
 				if (FFileHelper::LoadFileToString(FileCmds, *InFilePath))
 				{
-					FileCmds = FString(TEXT(" ")) + FileCmds.TrimStartAndEnd();
-					if (FileCmds.Len() > 1)
+					FileCmds = FileCmds.TrimStartAndEnd();
+					if (FileCmds.Len() > 0)
 					{
-						UE_LOG(LogInit, Log, TEXT("Appending commandline from file: %s, %s"), *InFilePath, *FileCmds);
-						FCommandLine::Append(*FileCmds);
-						bChanged = true;
+						UE_LOG(LogInit, Log, TEXT("Inserting commandline from file: %s, %s"), *InFilePath, *FileCmds);
+						FString NewCommandLine = FCommandLine::Get();
+						FString Left;
+						FString Right;
+						if (NewCommandLine.Split(TEXT("-CmdLineFile="), &Left, &Right))
+						{
+							FString NextToken;
+							const TCHAR* Stream = *Right;
+							if (FParse::Token(Stream, NextToken, /*bUseEscape=*/ false))
+							{
+								Right = FString(Stream);
+							}
+
+							NewCommandLine = Left + TEXT(" ") + FileCmds + TEXT(" ") + Right;
+							FCommandLine::Set(*NewCommandLine);
+							bChanged = true;
+						}
 					}
 
 					return true;
@@ -963,7 +977,7 @@ void LaunchCheckForCmdLineFile(const TCHAR* CmdLine, bool& bChanged)
 			};
 
 			bool bFoundFile = TryProcessFile(CmdLineFile);
-			if (!bFoundFile && FPaths::IsProjectFilePathSet())
+			if (!bFoundFile && FPaths::ProjectDir().Len() > 0)
 			{
 				const FString ProjectDir = FPaths::ProjectDir();
 				bFoundFile = TryProcessFile(ProjectDir + CmdLineFile);
@@ -1407,6 +1421,8 @@ int32 FEngineLoop::PreInitPreStartupScreen(const TCHAR* CmdLine)
 	}
 
 	{
+		SCOPED_BOOT_TIMING("InitTrace")
+
 		FString Parameter;
 		if (FParse::Value(CmdLine, TEXT("-tracehost="), Parameter))
 		{
@@ -1627,7 +1643,10 @@ int32 FEngineLoop::PreInitPreStartupScreen(const TCHAR* CmdLine)
 
 	// 
 #if PLATFORM_DESKTOP && !IS_MONOLITHIC
-	FModuleManager::Get().AddExtraBinarySearchPaths();
+	{
+		SCOPED_BOOT_TIMING("AddExtraBinarySearchPaths");
+		FModuleManager::Get().AddExtraBinarySearchPaths();
+	}
 #endif
 
 	// Initialize file manager
@@ -2486,7 +2505,7 @@ int32 FEngineLoop::PreInitPreStartupScreen(const TCHAR* CmdLine)
 
 	{
 		SCOPED_BOOT_TIMING("FUniformBufferStruct::InitializeStructs()");
-		FShaderParametersMetadata::InitializeAllGlobalStructs();
+		FShaderParametersMetadata::InitializeAllUniformBufferStructs();
 	}
 
 	{
@@ -5663,6 +5682,15 @@ void FEngineLoop::PostInitRHI()
 		PixelFormatByteWidth[i] = GPixelFormats[i].BlockBytes;
 	}
 	RHIPostInit(PixelFormatByteWidth);
+
+#if (!UE_BUILD_SHIPPING)
+	if (FParse::Param(FCommandLine::Get(), TEXT("rhiunittest")))
+	{
+		extern ENGINE_API void RunRHIUnitTest();
+		RunRHIUnitTest();
+	}
+#endif //(!UE_BUILD_SHIPPING)
+
 #endif
 }
 
diff --git a/Engine/Source/Runtime/PakFile/Private/IPlatformFilePak.cpp b/Engine/Source/Runtime/PakFile/Private/IPlatformFilePak.cpp
--- a/Engine/Source/Runtime/PakFile/Private/IPlatformFilePak.cpp
+++ b/Engine/Source/Runtime/PakFile/Private/IPlatformFilePak.cpp
@@ -57,7 +57,7 @@ CSV_DECLARE_CATEGORY_MODULE_EXTERN(CORE_API, FileIO);
 
 static FString GMountStartupPaksWildCard = TEXT(MOUNT_STARTUP_PAKS_WILDCARD);
 
-int32 ParseChunkIDFromFilename(const FString& InFilename)
+int32 GetPakchunkIDFromPakFile(const FString& InFilename)
 {
 	FString ChunkIdentifier(TEXT("pakchunk"));
 	FString BaseFilename = FPaths::GetBaseFilename(InFilename);
@@ -181,11 +181,13 @@ FPakCustomEncryptionDelegate& FPakPlatformFile::GetPakCustomEncryptionDelegate()
 	return Delegate;
 }
 
+// Needs to be a global as multiple threads may try to access the failed delegate
+static FPakChunkSignatureCheckFailedHandler SignatureFailedDelegate;
 FPakChunkSignatureCheckFailedHandler& FPakPlatformFile::GetPakChunkSignatureCheckFailedHandler()
 {
-	static FPakChunkSignatureCheckFailedHandler Delegate;
-	return Delegate;
+	return SignatureFailedDelegate;
 }
+
 FPakMasterSignatureTableCheckFailureHandler& FPakPlatformFile::GetPakMasterSignatureTableCheckFailureHandler()
 {
 	static FPakMasterSignatureTableCheckFailureHandler Delegate;
@@ -447,6 +449,13 @@ static FAutoConsoleVariableRef CVar_PakCache_TimeToTrim(
 	TEXT("Controls how long to hold onto a cached but unreferenced block for.")
 );
 
+int32 GPakCache_EnableNoCaching = 0;
+static FAutoConsoleVariableRef CVar_EnableNoCaching(
+	TEXT("pakcache.EnableNoCaching"),
+	GPakCache_EnableNoCaching,
+	TEXT("if > 0, then we'll allow a read requests pak cache memory to be ditched early")
+);
+
 class FPakPrecacher;
 
 typedef uint64 FJoinedOffsetAndPakIndex;
@@ -1157,6 +1166,7 @@ class FPakPrecacher
 			, Index(IntervalTreeInvalidIndex)
 			, Next(IntervalTreeInvalidIndex)
 			, Status(EBlockStatus::InFlight)
+			, TimeNoLongerReferenced(0)
 		{
 		}
 	};
@@ -1789,7 +1799,7 @@ private: // below here we assume CachedFilesScopeLock until we get to the next s
 					break;
 				}
 			}
-			if(GPakCache_TimeToTrim != 0.0f)
+			if (GPakCache_TimeToTrim != 0.0f)
 			{
 				// we'll trim based on time rather than trying to keep within a memory budget
 				double CurrentTime = FPlatformTime::Seconds();
@@ -1808,7 +1818,7 @@ private: // below here we assume CachedFilesScopeLock until we get to the next s
 							uint16 PakIndex = GetRequestPakIndex(OffsetAndPakIndex);
 							int64 Offset = GetRequestOffset(OffsetAndPakIndex);
 							FPakData& Pak = CachedPakData[PakIndex];
-							bool removedall = true;
+							bool bRemovedAll = true;
 							MaybeRemoveOverlappingNodesInIntervalTree<FCacheBlock>(
 								&Pak.CacheBlocks[(int32)EBlockStatus::Complete],
 								CacheBlockAllocator,
@@ -1818,7 +1828,7 @@ private: // below here we assume CachedFilesScopeLock until we get to the next s
 								Pak.MaxNode,
 								Pak.StartShift,
 								Pak.MaxShift,
-								[this, CurrentTime, &removedall](TIntervalTreeIndex BlockIndex) -> bool
+								[this, CurrentTime, &bRemovedAll](TIntervalTreeIndex BlockIndex) -> bool
 							{
 								FCacheBlock &Block = CacheBlockAllocator.Get(BlockIndex);
 								if (!Block.InRequestRefCount && (CurrentTime - Block.TimeNoLongerReferenced >= GPakCache_TimeToTrim))
@@ -1827,11 +1837,11 @@ private: // below here we assume CachedFilesScopeLock until we get to the next s
 									ClearBlock(Block);
 									return true;
 								}
-								removedall = false;
+								bRemovedAll = false;
 								return false;
 							}
 							);
-							if (!removedall)
+							if (!bRemovedAll)
 								break;
 							NumToRemove++;
 						}
@@ -1908,6 +1918,9 @@ private: // below here we assume CachedFilesScopeLock until we get to the next s
 		FPakData& Pak = CachedPakData[PakIndex];
 		check(Offset + Request.Size <= Pak.TotalSize && Request.Size > 0 && Request.GetPriority() >= AIOP_MIN && Request.GetPriority() <= AIOP_MAX && int32(Request.Status) >= 0 && int32(Request.Status) < int32(EInRequestStatus::Num));
 
+		bool RequestDontCache = (Request.PriorityAndFlags & AIOP_FLAG_DONTCACHE) != 0;
+
+
 		if (RemoveFromIntervalTree<FPakInRequest>(&Pak.InRequests[Request.GetPriority()][(int32)Request.Status], InRequestAllocator, Index, Pak.StartShift, Pak.MaxShift))
 		{
 
@@ -1921,7 +1934,7 @@ private: // below here we assume CachedFilesScopeLock until we get to the next s
 				Pak.MaxNode,
 				Pak.StartShift,
 				Pak.MaxShift,
-				[this, OffsetOfLastByte](TIntervalTreeIndex BlockIndex) -> bool
+				[this, OffsetOfLastByte, RequestDontCache](TIntervalTreeIndex BlockIndex) -> bool
 			{
 				FCacheBlock &Block = CacheBlockAllocator.Get(BlockIndex);
 				check(Block.InRequestRefCount);
@@ -1929,11 +1942,23 @@ private: // below here we assume CachedFilesScopeLock until we get to the next s
 				{
 					if (GPakCache_NumUnreferencedBlocksToCache && GetRequestOffset(Block.OffsetAndPakIndex) + Block.Size > OffsetOfLastByte) // last block
 					{
-						uint16 BlocksPakIndex = GetRequestPakIndexLow(Block.OffsetAndPakIndex);
-						int32 BlocksCacheIndex = CachedPakData[BlocksPakIndex].ActualPakFile->GetCacheIndex();
-						Block.TimeNoLongerReferenced = FPlatformTime::Seconds();
-						OffsetAndPakIndexOfSavedBlocked[BlocksCacheIndex].Remove(Block.OffsetAndPakIndex);
-						OffsetAndPakIndexOfSavedBlocked[BlocksCacheIndex].Add(Block.OffsetAndPakIndex);
+						if (RequestDontCache && GPakCache_EnableNoCaching != 0)
+						{
+							uint16 BlocksPakIndex = GetRequestPakIndexLow(Block.OffsetAndPakIndex);
+							int32 BlocksCacheIndex = CachedPakData[BlocksPakIndex].ActualPakFile->GetCacheIndex();
+							Block.TimeNoLongerReferenced = 0.0;
+							OffsetAndPakIndexOfSavedBlocked[BlocksCacheIndex].Remove(Block.OffsetAndPakIndex);
+							ClearBlock(Block);
+							return true;
+						}
+						else
+						{
+							uint16 BlocksPakIndex = GetRequestPakIndexLow(Block.OffsetAndPakIndex);
+							int32 BlocksCacheIndex = CachedPakData[BlocksPakIndex].ActualPakFile->GetCacheIndex();
+							Block.TimeNoLongerReferenced = FPlatformTime::Seconds();
+							OffsetAndPakIndexOfSavedBlocked[BlocksCacheIndex].Remove(Block.OffsetAndPakIndex);
+							OffsetAndPakIndexOfSavedBlocked[BlocksCacheIndex].Add(Block.OffsetAndPakIndex);
+						}
 						return false;
 					}
 					ClearBlock(Block);
@@ -4749,7 +4774,7 @@ FPakFile::FPakFile(const TCHAR* Filename, bool bIsSigned)
 	, bSigned(bIsSigned)
 	, bIsValid(false)
 	, bFilenamesRemoved(false)
-	, ChunkID(ParseChunkIDFromFilename(Filename))
+	, ChunkID(GetPakchunkIDFromPakFile(Filename))
 	, bAttemptedPakEntryShrink(false)
 	, bAttemptedPakFilenameUnload(false)
  	, MappedFileHandle(nullptr)
@@ -4779,7 +4804,7 @@ FPakFile::FPakFile(IPlatformFile* LowerLevel, const TCHAR* Filename, bool bIsSig
 	, bSigned(bIsSigned)
 	, bIsValid(false)
 	, bFilenamesRemoved(false)
-	, ChunkID(ParseChunkIDFromFilename(Filename))
+	, ChunkID(GetPakchunkIDFromPakFile(Filename))
 	, bAttemptedPakEntryShrink(false)
 	, bAttemptedPakFilenameUnload(false)
 	, MappedFileHandle(nullptr)
@@ -6068,10 +6093,10 @@ void FPakPlatformFile::FindPakFilesInDirectory(IPlatformFile* LowLevelFile, cons
 					// if a platform supports chunk style installs, make sure that the chunk a pak file resides in is actually fully installed before accepting pak files from it
 					if (ChunkInstall)
 					{
-						int32 ChunkID = ParseChunkIDFromFilename(Filename);
-						if (ChunkID != INDEX_NONE)
+						int32 PakchunkID = GetPakchunkIDFromPakFile(Filename);
+						if (PakchunkID != INDEX_NONE)
 						{
-							if (ChunkInstall->GetChunkLocation(ChunkID) == EChunkLocation::NotAvailable)
+							if (ChunkInstall->GetPakchunkLocation(PakchunkID) == EChunkLocation::NotAvailable)
 							{
 								return true;
 							}
@@ -6590,7 +6615,7 @@ void FPakPlatformFile::RegisterEncryptionKey(const FGuid& InGuid, const FAES::FA
 				UE_LOG(LogPakFile, Log, TEXT("Successfully mounted deferred pak file '%s'"), *Entry.Filename);
 				NumMounted++;
 
-				int32 ChunkID = ParseChunkIDFromFilename(Entry.Filename);
+				int32 ChunkID = GetPakchunkIDFromPakFile(Entry.Filename);
 				if (ChunkID != INDEX_NONE)
 				{
 					ChunksToNotify.Add(ChunkID);
diff --git a/Engine/Source/Runtime/RHI/Private/RHI.cpp b/Engine/Source/Runtime/RHI/Private/RHI.cpp
--- a/Engine/Source/Runtime/RHI/Private/RHI.cpp
+++ b/Engine/Source/Runtime/RHI/Private/RHI.cpp
@@ -518,6 +518,7 @@ bool GSupportsTexture3D = true;
 bool GSupportsMobileMultiView = false;
 bool GSupportsImageExternal = false;
 bool GSupportsResourceView = true;
+bool GRHISupportsMultithreading = false;
 TRHIGlobal<bool> GSupportsMultipleRenderTargets(true);
 bool GSupportsWideMRT = true;
 float GMinClipZ = 0.0f;
@@ -543,10 +544,11 @@ bool GRHISupportsTextureStreaming = false;
 bool GSupportsDepthBoundsTest = false;
 bool GSupportsEfficientAsyncCompute = false;
 bool GRHISupportsBaseVertexIndex = true;
-TRHIGlobal<bool> GRHISupportsInstancing(true);
 bool GRHISupportsFirstInstance = false;
 bool GRHISupportsDynamicResolution = false;
 bool GRHISupportsRayTracing = false;
+bool GRHISupportsRayTracingMissShaderBindings = false;
+bool GRHISupportsRayTracingAsyncBuildAccelerationStructure = false;
 bool GRHISupportsWaveOperations = false;
 bool GRHISupportsRHIThread = false;
 bool GRHISupportsRHIOnTaskThread = false;
@@ -733,10 +735,7 @@ FName ShaderPlatformToPlatformName(EShaderPlatform Platform)
 	switch (Platform)
 	{
 	case SP_PCD3D_SM5:
-	case SP_OPENGL_SM4:
-	case SP_OPENGL_PCES2:
 	case SP_OPENGL_SM5:
-	case SP_PCD3D_ES2:
 	case SP_PCD3D_ES3_1:
 	case SP_OPENGL_PCES3_1:
 	case SP_VULKAN_PCES3_1:
@@ -759,7 +758,6 @@ FName ShaderPlatformToPlatformName(EShaderPlatform Platform)
 	case SP_METAL_SM5:
 	case SP_METAL_SM5_NOTESS:
 	case SP_METAL_MACES3_1:
-	case SP_METAL_MACES2:
 	case SP_METAL_MRT_MAC:
 		return NAME_PLATFORM_MAC;
 	case SP_SWITCH:
@@ -1003,19 +1001,6 @@ void FRHIRenderPassInfo::ConvertToRenderTargetsInfo(FRHISetRenderTargetsInfo& Ou
 	OutRTInfo.bClearStencil = (StencilLoadAction == ERenderTargetLoadAction::EClear);
 
 	OutRTInfo.FoveationTexture = FoveationTexture;
-
-	if (NumUAVs > 0)
-	{
-		check(UAVIndex != -1);
-		int32 StartingUAVIndex = FMath::Max(UAVIndex, OutRTInfo.NumColorRenderTargets);
-		check((StartingUAVIndex+NumUAVs) <= MaxSimultaneousUAVs);
-		OutRTInfo.NumColorRenderTargets = StartingUAVIndex;
-		for (int32 Index = 0; Index < NumUAVs; ++Index)
-		{
-			OutRTInfo.UnorderedAccessView[Index] = UAVs[Index];
-		}
-		OutRTInfo.NumUAVs = NumUAVs;
-	}
 }
 
 void FRHIRenderPassInfo::OnVerifyNumUAVsFailed(int32 InNumUAVs)
@@ -1158,10 +1143,7 @@ FString LexToString(EShaderPlatform Platform)
 	{
 	case SP_PCD3D_SM5: return TEXT("PCD3D_SM5");
 	case SP_PCD3D_ES3_1: return TEXT("PCD3D_ES3_1");
-	case SP_PCD3D_ES2: return TEXT("PCD3D_ES2");
-	case SP_OPENGL_SM4: return TEXT("OPENGL_SM4");
 	case SP_OPENGL_SM5: return TEXT("OPENGL_SM5");
-	case SP_OPENGL_PCES2: return TEXT("OPENGL_PCES2");
 	case SP_OPENGL_PCES3_1: return TEXT("OPENGL_PCES3_1");
 	case SP_OPENGL_ES2_ANDROID: return TEXT("OPENGL_ES2_ANDROID");
 	case SP_OPENGL_ES2_WEBGL: return TEXT("OPENGL_ES2_WEBGL");
@@ -1179,16 +1161,19 @@ FString LexToString(EShaderPlatform Platform)
 	case SP_METAL_SM5: return TEXT("METAL_SM5");
 	case SP_METAL_SM5_NOTESS: return TEXT("METAL_SM5_NOTESS");
 	case SP_METAL_MACES3_1: return TEXT("METAL_MACES3_1");
-	case SP_METAL_MACES2: return TEXT("METAL_MACES2");
 	case SP_VULKAN_ES3_1_ANDROID: return TEXT("VULKAN_ES3_1_ANDROID");
 	case SP_VULKAN_ES3_1_LUMIN: return TEXT("VULKAN_ES3_1_LUMIN");
 	case SP_VULKAN_PCES3_1: return TEXT("VULKAN_PCES3_1");
 	case SP_VULKAN_SM5: return TEXT("VULKAN_SM5");
 	case SP_VULKAN_SM5_LUMIN: return TEXT("VULKAN_SM5_LUMIN");
 
-	case SP_OPENGL_ES2_IOS_DEPRECATED:
-	case SP_VULKAN_SM4_DEPRECATED:
-	case SP_PCD3D_SM4_DEPRECATED:
+	case SP_OPENGL_ES2_IOS_REMOVED:
+	case SP_VULKAN_SM4_REMOVED:
+	case SP_PCD3D_SM4_REMOVED:
+	case SP_OPENGL_SM4_REMOVED:
+	case SP_PCD3D_ES2_DEPRECATED:
+	case SP_OPENGL_PCES2_DEPRECATED:
+	case SP_METAL_MACES2_DEPRECATED:
 		return TEXT("");
 
 	default:
@@ -1265,7 +1250,9 @@ inline void ParseDataDrivenShaderInfo(const FConfigSection& Section, FDataDriven
 	Info.bSupports4ComponentUAVReadWrite = GetSectionBool(Section, "bSupports4ComponentUAVReadWrite");
 	Info.bSupportsRenderTargetWriteMask = GetSectionBool(Section, "bSupportsRenderTargetWriteMask");
 	Info.bSupportsRayTracing = GetSectionBool(Section, "bSupportsRayTracing");
+	Info.bSupportsRayTracingMissShaderBindings = GetSectionBool(Section, "bSupportsRayTracingMissShaderBindings");
 	Info.bSupportsGPUSkinCache = GetSectionBool(Section, "bSupportsGPUSkinCache");
+	Info.bSupportsByteBufferComputeShaders = GetSectionBool(Section, "bSupportsByteBufferComputeShaders");
 
 	Info.bTargetsTiledGPU = GetSectionBool(Section, "bTargetsTiledGPU");
 	Info.bNeedsOfflineCompiler = GetSectionBool(Section, "bNeedsOfflineCompiler");
diff --git a/Engine/Source/Runtime/RHI/Public/RHI.h b/Engine/Source/Runtime/RHI/Public/RHI.h
--- a/Engine/Source/Runtime/RHI/Public/RHI.h
+++ b/Engine/Source/Runtime/RHI/Public/RHI.h
@@ -82,6 +82,9 @@ extern RHI_API EShaderPlatform GMaxRHIShaderPlatform;
 /** true if the RHI supports SRVs */
 extern RHI_API bool GSupportsResourceView;
 
+/** Whether the RHI can send commands to the device context from multiple threads. Used in the GPU readback to avoid stalling the RHI threads. */
+extern RHI_API bool GRHISupportsMultithreading;
+
 /** 
  * only set if RHI has the information (after init of the RHI and only if RHI has that information, never changes after that)
  * e.g. "NVIDIA GeForce GTX 670"
@@ -166,7 +169,7 @@ inline bool RHISupportsVolumeTextures(ERHIFeatureLevel::Type FeatureLevel)
 
 inline bool RHISupportsVertexShaderLayer(const EShaderPlatform Platform)
 {
-	return IsFeatureLevelSupported(Platform, ERHIFeatureLevel::SM5) && IsMetalPlatform(Platform) && (IsPCPlatform(Platform) || (Platform == SP_METAL_MRT && RHIGetShaderLanguageVersion(Platform) >= 4));
+	return IsFeatureLevelSupported(Platform, ERHIFeatureLevel::SM5) && IsMetalPlatform(Platform) && IsPCPlatform(Platform);
 }
 
 /** Return true if and only if the GPU support rendering to volume textures (2D Array, 3D) is guaranteed supported for a target platform.
@@ -498,9 +501,6 @@ extern RHI_API int32 GCurrentNumPrimitivesDrawnRHI;
 /** Whether or not the RHI can handle a non-zero BaseVertexIndex - extra SetStreamSource calls will be needed if this is false */
 extern RHI_API bool GRHISupportsBaseVertexIndex;
 
-/** True if the RHI supports hardware instancing */
-extern RHI_API TRHIGlobal<bool> GRHISupportsInstancing;
-
 /** True if the RHI supports copying cubemap faces using CopyToResolveTarget */
 extern RHI_API bool GRHISupportsResolveCubemapFaces;
 
@@ -513,6 +513,12 @@ extern RHI_API bool GRHISupportsDynamicResolution;
 /** Whether or not the RHI supports ray tracing on current hardware (acceleration structure building and new ray tracing-specific shader types). */
 extern RHI_API bool GRHISupportsRayTracing;
 
+/** Whether or not the RHI supports binding multiple miss shaders with local resources via RHISetRayTracingMissShader(). */
+extern RHI_API bool GRHISupportsRayTracingMissShaderBindings;
+
+/** Whether or not the RHI supports async building ray tracing acceleration structures. */
+extern RHI_API bool GRHISupportsRayTracingAsyncBuildAccelerationStructure;
+
 /** Whether or not the RHI supports shader wave operations (shader model 6.0). */
 extern RHI_API bool GRHISupportsWaveOperations;
 
diff --git a/Engine/Source/Runtime/RHI/Public/RHICommandList.h b/Engine/Source/Runtime/RHI/Public/RHICommandList.h
--- a/Engine/Source/Runtime/RHI/Public/RHICommandList.h
+++ b/Engine/Source/Runtime/RHI/Public/RHICommandList.h
@@ -117,7 +117,6 @@
 
 
 extern RHI_API bool GEnableAsyncCompute;
-extern RHI_API bool GSupportAsyncComputeRaytracingBuildBVH;
 extern RHI_API TAutoConsoleVariable<int32> CVarRHICmdWidth;
 extern RHI_API TAutoConsoleVariable<int32> CVarRHICmdFlushRenderThreadTasks;
 
@@ -125,10 +124,23 @@
 struct FRayTracingShaderBindings
 {
 	FRHITexture* Textures[32] = {};
-	FRHIShaderResourceView* SRVs[32] = {};
-	FRHIUniformBuffer* UniformBuffers[8] = {};
+	FRHIShaderResourceView* SRVs[48] = {};
+	FRHIUniformBuffer* UniformBuffers[16] = {};
 	FRHISamplerState* Samplers[16] = {};
-	FRHIUnorderedAccessView* UAVs[8] = {};
+	FRHIUnorderedAccessView* UAVs[16] = {};
 };
 
+struct FRayTracingLocalShaderBindings
+{
+	uint32 InstanceIndex = 0;
+	uint32 SegmentIndex = 0;
+	uint32 ShaderSlot = 0;
+	uint32 ShaderIndexInPipeline = 0;
+	uint32 UserData = 0;
+	uint16 NumUniformBuffers = 0;
+	uint16 LooseParameterDataSize = 0;
+	FRHIUniformBuffer** UniformBuffers = nullptr;
+	uint8* LooseParameterData = nullptr;
+};
+
 // C++ counter-part of FBasicRayData declared in RayTracingCommon.ush
@@ -210,12 +222,6 @@
 #define PSO_VERIFY	check
 #endif
 
-enum class ECmdList
-{
-	EGfx,
-	ECompute,	
-};
-
 class IRHICommandContextContainer
 {
 public:
@@ -399,7 +405,6 @@
 class RHI_API FRHICommandListBase : public FNoncopyable
 {
 public:
-	FRHICommandListBase(FRHIGPUMask InGPUMask);
 	~FRHICommandListBase();
 
 	/** Custom new/delete with recycling */
@@ -506,6 +511,7 @@
 	{
 		check(InContext);
 		Context = InContext;
+		ComputeContext = InContext;
 	}
 
 	FORCEINLINE IRHICommandContext& GetContext()
@@ -514,10 +520,10 @@
 		return *Context;
 	}
 
-	void SetComputeContext(IRHIComputeContext* InContext)
+	void SetComputeContext(IRHIComputeContext* InComputeContext)
 	{
-		check(InContext);
-		ComputeContext = InContext;
+		check(Context == nullptr);
+		ComputeContext = InComputeContext;
 	}
 
 	IRHIComputeContext& GetComputeContext()
@@ -544,10 +550,6 @@
 
 	FORCEINLINE const FRHIGPUMask& GetGPUMask() const { return GPUMask; }
 
-protected:
-	void SetGPUMaskOnContext();
-	void SetGPUMaskOnComputeContext();
-
 private:
 	FRHICommandBase* Root;
 	FRHICommandBase** CommandLink;
@@ -564,8 +566,9 @@
 	friend class FRHICommandListScopedFlushAndExecute;
 
 protected:
+	FRHICommandListBase(FRHIGPUMask InGPUMask);
+
 	bool bAsyncPSOCompileAllowed;
-	bool bRecursive;
 	FRHIGPUMask GPUMask;
 	// GPUMask that was set at the time the command list was last Reset. We set
     // this mask on the command contexts immediately before executing the
@@ -814,8 +817,8 @@
 {
 	static const TCHAR* TStr() { return TEXT("FRHICommandSetShaderParameter"); }
 };
-template <typename TRHIShader, ECmdList CmdListType>
-struct FRHICommandSetShaderParameter final : public FRHICommand<FRHICommandSetShaderParameter<TRHIShader, CmdListType>, FRHICommandSetShaderParameterString>
+template <typename TRHIShader>
+struct FRHICommandSetShaderParameter final : public FRHICommand<FRHICommandSetShaderParameter<TRHIShader>, FRHICommandSetShaderParameterString>
 {
 	TRHIShader* Shader;
 	const void* NewValue;
@@ -837,8 +840,8 @@
 {
 	static const TCHAR* TStr() { return TEXT("FRHICommandSetShaderUniformBuffer"); }
 };
-template <typename TRHIShader, ECmdList CmdListType>
-struct FRHICommandSetShaderUniformBuffer final : public FRHICommand<FRHICommandSetShaderUniformBuffer<TRHIShader, CmdListType>, FRHICommandSetShaderUniformBufferString>
+template <typename TRHIShader>
+struct FRHICommandSetShaderUniformBuffer final : public FRHICommand<FRHICommandSetShaderUniformBuffer<TRHIShader>, FRHICommandSetShaderUniformBufferString>
 {
 	TRHIShader* Shader;
 	uint32 BaseIndex;
@@ -856,8 +859,8 @@
 {
 	static const TCHAR* TStr() { return TEXT("FRHICommandSetShaderTexture"); }
 };
-template <typename TRHIShader, ECmdList CmdListType>
-struct FRHICommandSetShaderTexture final : public FRHICommand<FRHICommandSetShaderTexture<TRHIShader, CmdListType>, FRHICommandSetShaderTextureString >
+template <typename TRHIShader>
+struct FRHICommandSetShaderTexture final : public FRHICommand<FRHICommandSetShaderTexture<TRHIShader>, FRHICommandSetShaderTextureString >
 {
 	TRHIShader* Shader;
 	uint32 TextureIndex;
@@ -875,8 +878,8 @@
 {
 	static const TCHAR* TStr() { return TEXT("FRHICommandSetShaderResourceViewParameter"); }
 };
-template <typename TRHIShader, ECmdList CmdListType>
-struct FRHICommandSetShaderResourceViewParameter final : public FRHICommand<FRHICommandSetShaderResourceViewParameter<TRHIShader, CmdListType>, FRHICommandSetShaderResourceViewParameterString >
+template <typename TRHIShader>
+struct FRHICommandSetShaderResourceViewParameter final : public FRHICommand<FRHICommandSetShaderResourceViewParameter<TRHIShader>, FRHICommandSetShaderResourceViewParameterString >
 {
 	TRHIShader* Shader;
 	uint32 SamplerIndex;
@@ -894,8 +897,8 @@
 {
 	static const TCHAR* TStr() { return TEXT("FRHICommandSetUAVParameter"); }
 };
-template <typename TRHIShader, ECmdList CmdListType>
-struct FRHICommandSetUAVParameter final : public FRHICommand<FRHICommandSetUAVParameter<TRHIShader, CmdListType>, FRHICommandSetUAVParameterString >
+template <typename TRHIShader>
+struct FRHICommandSetUAVParameter final : public FRHICommand<FRHICommandSetUAVParameter<TRHIShader>, FRHICommandSetUAVParameterString >
 {
 	TRHIShader* Shader;
 	uint32 UAVIndex;
@@ -909,18 +912,17 @@
 	RHI_API void Execute(FRHICommandListBase& CmdList);
 };
 
-struct FRHICommandSetUAVParameter_IntialCountString
+struct FRHICommandSetUAVParameter_InitialCountString
 {
-	static const TCHAR* TStr() { return TEXT("FRHICommandSetUAVParameter_IntialCount"); }
+	static const TCHAR* TStr() { return TEXT("FRHICommandSetUAVParameter_InitialCount"); }
 };
-template <typename TRHIShader, ECmdList CmdListType>
-struct FRHICommandSetUAVParameter_IntialCount final : public FRHICommand<FRHICommandSetUAVParameter_IntialCount<TRHIShader, CmdListType>, FRHICommandSetUAVParameter_IntialCountString >
+struct FRHICommandSetUAVParameter_InitialCount final : public FRHICommand<FRHICommandSetUAVParameter_InitialCount, FRHICommandSetUAVParameter_InitialCountString >
 {
-	TRHIShader* Shader;
+	FRHIComputeShader* Shader;
 	uint32 UAVIndex;
 	FRHIUnorderedAccessView* UAV;
 	uint32 InitialCount;
-	FORCEINLINE_DEBUGGABLE FRHICommandSetUAVParameter_IntialCount(TRHIShader* InShader, uint32 InUAVIndex, FRHIUnorderedAccessView* InUAV, uint32 InInitialCount)
+	FORCEINLINE_DEBUGGABLE FRHICommandSetUAVParameter_InitialCount(FRHIComputeShader* InShader, uint32 InUAVIndex, FRHIUnorderedAccessView* InUAV, uint32 InInitialCount)
 		: Shader(InShader)
 		, UAVIndex(InUAVIndex)
 		, UAV(InUAV)
@@ -934,8 +936,8 @@
 {
 	static const TCHAR* TStr() { return TEXT("FRHICommandSetShaderSampler"); }
 };
-template <typename TRHIShader, ECmdList CmdListType>
-struct FRHICommandSetShaderSampler final : public FRHICommand<FRHICommandSetShaderSampler<TRHIShader, CmdListType>, FRHICommandSetShaderSamplerString >
+template <typename TRHIShader>
+struct FRHICommandSetShaderSampler final : public FRHICommand<FRHICommandSetShaderSampler<TRHIShader>, FRHICommandSetShaderSamplerString >
 {
 	TRHIShader* Shader;
 	uint32 SamplerIndex;
@@ -1011,13 +1013,13 @@
 
 FRHICOMMAND_MACRO(FRHICommandSetViewport)
 {
-	uint32 MinX;
-	uint32 MinY;
+	float MinX;
+	float MinY;
 	float MinZ;
-	uint32 MaxX;
-	uint32 MaxY;
+	float MaxX;
+	float MaxY;
 	float MaxZ;
-	FORCEINLINE_DEBUGGABLE FRHICommandSetViewport(uint32 InMinX, uint32 InMinY, float InMinZ, uint32 InMaxX, uint32 InMaxY, float InMaxZ)
+	FORCEINLINE_DEBUGGABLE FRHICommandSetViewport(float InMinX, float InMinY, float InMinZ, float InMaxX, float InMaxY, float InMaxZ)
 		: MinX(InMinX)
 		, MinY(InMinY)
 		, MinZ(InMinZ)
@@ -1031,17 +1033,17 @@
 
 FRHICOMMAND_MACRO(FRHICommandSetStereoViewport)
 {
-	uint32 LeftMinX;
-	uint32 RightMinX;
-	uint32 LeftMinY;
-	uint32 RightMinY;
+	float LeftMinX;
+	float RightMinX;
+	float LeftMinY;
+	float RightMinY;
 	float MinZ;
-	uint32 LeftMaxX;
-	uint32 RightMaxX;
-	uint32 LeftMaxY;
-	uint32 RightMaxY;
+	float LeftMaxX;
+	float RightMaxX;
+	float LeftMaxY;
+	float RightMaxY;
 	float MaxZ;
-	FORCEINLINE_DEBUGGABLE FRHICommandSetStereoViewport(uint32 InLeftMinX, uint32 InRightMinX, uint32 InLeftMinY, uint32 InRightMinY, float InMinZ, uint32 InLeftMaxX, uint32 InRightMaxX, uint32 InLeftMaxY, uint32 InRightMaxY, float InMaxZ)
+	FORCEINLINE_DEBUGGABLE FRHICommandSetStereoViewport(float InLeftMinX, float InRightMinX, float InLeftMinY, float InRightMinY, float InMinZ, float InLeftMaxX, float InRightMaxX, float InLeftMaxY, float InRightMaxY, float InMaxZ)
 		: LeftMinX(InLeftMinX)
 		, RightMinX(InRightMinX)
 		, LeftMinY(InLeftMinY)
@@ -1080,29 +1082,20 @@
 	uint32 NewNumSimultaneousRenderTargets;
 	FRHIRenderTargetView NewRenderTargetsRHI[MaxSimultaneousRenderTargets];
 	FRHIDepthRenderTargetView NewDepthStencilTarget;
-	uint32 NewNumUAVs;
-	FRHIUnorderedAccessView* UAVs[MaxSimultaneousUAVs];
 
 	FORCEINLINE_DEBUGGABLE FRHICommandSetRenderTargets(
 		uint32 InNewNumSimultaneousRenderTargets,
 		const FRHIRenderTargetView* InNewRenderTargetsRHI,
-		const FRHIDepthRenderTargetView* InNewDepthStencilTargetRHI,
-		uint32 InNewNumUAVs,
-		FRHIUnorderedAccessView* const* InUAVs
+		const FRHIDepthRenderTargetView* InNewDepthStencilTargetRHI
 		)
 		: NewNumSimultaneousRenderTargets(InNewNumSimultaneousRenderTargets)
-		, NewNumUAVs(InNewNumUAVs)
 
 	{
-		check(InNewNumSimultaneousRenderTargets <= MaxSimultaneousRenderTargets && InNewNumUAVs <= MaxSimultaneousUAVs);
+		check(InNewNumSimultaneousRenderTargets <= MaxSimultaneousRenderTargets);
 		for (uint32 Index = 0; Index < NewNumSimultaneousRenderTargets; Index++)
 		{
 			NewRenderTargetsRHI[Index] = InNewRenderTargetsRHI[Index];
 		}
-		for (uint32 Index = 0; Index < NewNumUAVs; Index++)
-		{
-			UAVs[Index] = InUAVs[Index];
-		}
 		if (InNewDepthStencilTargetRHI)
 		{
 			NewDepthStencilTarget = *InNewDepthStencilTargetRHI;
@@ -1254,8 +1247,8 @@
 {
 	static const TCHAR* TStr() { return TEXT("FRHICommandSetComputeShader"); }
 };
-template<ECmdList CmdListType>
-struct FRHICommandSetComputeShader final : public FRHICommand<FRHICommandSetComputeShader<CmdListType>, FRHICommandSetComputeShaderString>
+
+struct FRHICommandSetComputeShader final : public FRHICommand<FRHICommandSetComputeShader, FRHICommandSetComputeShaderString>
 {
 	FRHIComputeShader* ComputeShader;
 	FORCEINLINE_DEBUGGABLE FRHICommandSetComputeShader(FRHIComputeShader* InComputeShader)
@@ -1269,8 +1262,8 @@
 {
 	static const TCHAR* TStr() { return TEXT("FRHICommandSetComputePipelineState"); }
 };
-template<ECmdList CmdListType>
-struct FRHICommandSetComputePipelineState final : public FRHICommand<FRHICommandSetComputePipelineState<CmdListType>, FRHICommandSetComputePipelineStateString>
+
+struct FRHICommandSetComputePipelineState final : public FRHICommand<FRHICommandSetComputePipelineState, FRHICommandSetComputePipelineStateString>
 {
 	FComputePipelineState* ComputePipelineState;
 	FORCEINLINE_DEBUGGABLE FRHICommandSetComputePipelineState(FComputePipelineState* InComputePipelineState)
@@ -1294,8 +1287,8 @@
 {
 	static const TCHAR* TStr() { return TEXT("FRHICommandDispatchComputeShader"); }
 };
-template<ECmdList CmdListType>
-struct FRHICommandDispatchComputeShader final : public FRHICommand<FRHICommandDispatchComputeShader<CmdListType>, FRHICommandDispatchComputeShaderString>
+
+struct FRHICommandDispatchComputeShader final : public FRHICommand<FRHICommandDispatchComputeShader, FRHICommandDispatchComputeShaderString>
 {
 	uint32 ThreadGroupCountX;
 	uint32 ThreadGroupCountY;
@@ -1313,8 +1306,8 @@
 {
 	static const TCHAR* TStr() { return TEXT("FRHICommandDispatchIndirectComputeShader"); }
 };
-template<ECmdList CmdListType>
-struct FRHICommandDispatchIndirectComputeShader final : public FRHICommand<FRHICommandDispatchIndirectComputeShader<CmdListType>, FRHICommandDispatchIndirectComputeShaderString>
+
+struct FRHICommandDispatchIndirectComputeShader final : public FRHICommand<FRHICommandDispatchIndirectComputeShader, FRHICommandDispatchIndirectComputeShaderString>
 {
 	FRHIVertexBuffer* ArgumentBuffer;
 	uint32 ArgumentOffset;
@@ -1398,18 +1391,28 @@
 	RHI_API void Execute(FRHICommandListBase& CmdList);
 };
 
-FRHICOMMAND_MACRO(FRHICommandClearTinyUAV)
+FRHICOMMAND_MACRO(FRHICommandClearUAVFloat)
 {
 	FRHIUnorderedAccessView* UnorderedAccessViewRHI;
-	uint32 Values[4];
+	FVector4 Values;
 
-	FORCEINLINE_DEBUGGABLE FRHICommandClearTinyUAV(FRHIUnorderedAccessView* InUnorderedAccessViewRHI, const uint32* InValues)
+	FORCEINLINE_DEBUGGABLE FRHICommandClearUAVFloat(FRHIUnorderedAccessView* InUnorderedAccessViewRHI, const FVector4& InValues)
 		: UnorderedAccessViewRHI(InUnorderedAccessViewRHI)
+		, Values(InValues)
 	{
-		Values[0] = InValues[0];
-		Values[1] = InValues[1];
-		Values[2] = InValues[2];
-		Values[3] = InValues[3];
 	}
 	RHI_API void Execute(FRHICommandListBase& CmdList);
 };
+
+FRHICOMMAND_MACRO(FRHICommandClearUAVUint)
+{
+	FRHIUnorderedAccessView* UnorderedAccessViewRHI;
+	FUintVector4 Values;
+
+	FORCEINLINE_DEBUGGABLE FRHICommandClearUAVUint(FRHIUnorderedAccessView* InUnorderedAccessViewRHI, const FUintVector4& InValues)
+		: UnorderedAccessViewRHI(InUnorderedAccessViewRHI)
+		, Values(InValues)
+	{
+	}
+	RHI_API void Execute(FRHICommandListBase& CmdList);
+};
@@ -1506,8 +1509,7 @@
 {
 	static const TCHAR* TStr() { return TEXT("FRHICommandTransitionUAVs"); }
 };
-template<ECmdList CmdListType>
-struct FRHICommandTransitionUAVs final : public FRHICommand<FRHICommandTransitionUAVs<CmdListType>, FRHICommandTransitionUAVsString>
+struct FRHICommandTransitionUAVs final : public FRHICommand<FRHICommandTransitionUAVs, FRHICommandTransitionUAVsString>
 {
 	int32 NumUAVs;
 	FRHIUnorderedAccessView** UAVs; // Pointer to an array of UAVs, allocated inline with the command list
@@ -1530,8 +1532,8 @@
 {
 	static const TCHAR* TStr() { return TEXT("FRHICommandSetAsyncComputeBudget"); }
 };
-template<ECmdList CmdListType>
-struct FRHICommandSetAsyncComputeBudget final : public FRHICommand<FRHICommandSetAsyncComputeBudget<CmdListType>, FRHICommandSetAsyncComputeBudgetString>
+
+struct FRHICommandSetAsyncComputeBudget final : public FRHICommand<FRHICommandSetAsyncComputeBudget, FRHICommandSetAsyncComputeBudgetString>
 {
 	EAsyncComputeBudget Budget;
 
@@ -1546,8 +1548,7 @@
 {
 	static const TCHAR* TStr() { return TEXT("FRHICommandWaitComputeFence"); }
 };
-template<ECmdList CmdListType>
-struct FRHICommandWaitComputeFence final : public FRHICommand<FRHICommandWaitComputeFence<CmdListType>, FRHICommandWaitComputeFenceString>
+struct FRHICommandWaitComputeFence final : public FRHICommand<FRHICommandWaitComputeFence, FRHICommandWaitComputeFenceString>
 {
 	FRHIComputeFence* WaitFence;
 
@@ -1562,8 +1563,8 @@
 {
 	static const TCHAR* TStr() { return TEXT("FRHICommandCopyToStagingBuffer"); }
 };
-template<ECmdList CmdListType>
-struct FRHICommandCopyToStagingBuffer final : public FRHICommand<FRHICommandCopyToStagingBuffer<CmdListType>, FRHICommandCopyToStagingBufferString>
+
+struct FRHICommandCopyToStagingBuffer final : public FRHICommand<FRHICommandCopyToStagingBuffer, FRHICommandCopyToStagingBufferString>
 {
 	FRHIVertexBuffer* SourceBuffer;
 	FRHIStagingBuffer* DestinationStagingBuffer;
@@ -1584,8 +1585,8 @@
 {
 	static const TCHAR* TStr() { return TEXT("FRHICommandWriteGPUFence"); }
 };
-template<ECmdList CmdListType>
-struct FRHICommandWriteGPUFence final : public FRHICommand<FRHICommandWriteGPUFence<CmdListType>, FRHICommandWriteGPUFenceString>
+
+struct FRHICommandWriteGPUFence final : public FRHICommand<FRHICommandWriteGPUFence, FRHICommandWriteGPUFenceString>
 {
 	FRHIGPUFence* Fence;
 
@@ -1657,6 +1658,16 @@
 	RHI_API void Execute(FRHICommandListBase& CmdList);
 };
 
+FRHICOMMAND_MACRO(FRHICommandSetGlobalUniformBuffers)
+{
+	FUniformBufferStaticBindings UniformBuffers;
+
+	FORCEINLINE_DEBUGGABLE FRHICommandSetGlobalUniformBuffers(const FUniformBufferStaticBindings& InUniformBuffers)
+		: UniformBuffers(InUniformBuffers)
+	{}
+	RHI_API void Execute(FRHICommandListBase& CmdList);
+};
+
 struct FComputedGraphicsPipelineState
 {
 	FGraphicsPipelineStateRHIRef GraphicsPipelineState;
@@ -1784,8 +1795,8 @@
 {
 	static const TCHAR* TStr() { return TEXT("FRHICommandSubmitCommandsHint"); }
 };
-template<ECmdList CmdListType>
-struct FRHICommandSubmitCommandsHint final : public FRHICommand<FRHICommandSubmitCommandsHint<CmdListType>, FRHICommandSubmitCommandsHintString>
+
+struct FRHICommandSubmitCommandsHint final : public FRHICommand<FRHICommandSubmitCommandsHint, FRHICommandSubmitCommandsHintString>
 {
 	FORCEINLINE_DEBUGGABLE FRHICommandSubmitCommandsHint()
 	{
@@ -1862,8 +1873,8 @@
 {
 	static const TCHAR* TStr() { return TEXT("FRHICommandPushEventString"); }
 };
-template<ECmdList CmdListType>
-struct FRHICommandPushEvent final : public FRHICommand<FRHICommandPushEvent<CmdListType>, FRHICommandPushEventString>
+
+struct FRHICommandPushEvent final : public FRHICommand<FRHICommandPushEvent, FRHICommandPushEventString>
 {
 	const TCHAR *Name;
 	FColor Color;
@@ -1885,8 +1896,8 @@
 {
 	static const TCHAR* TStr() { return TEXT("FRHICommandPopEvent"); }
 };
-template<ECmdList CmdListType>
-struct FRHICommandPopEvent final : public FRHICommand<FRHICommandPopEvent<CmdListType>, FRHICommandPopEventString>
+
+struct FRHICommandPopEvent final : public FRHICommand<FRHICommandPopEvent, FRHICommandPopEventString>
 {
 	RHI_API void Execute(FRHICommandListBase& CmdList);
 
@@ -2042,8 +2053,7 @@
 	RHI_API void Execute(FRHICommandListBase& CmdList);
 };
 
-template <ECmdList CmdListType>
-struct FRHICommandBuildAccelerationStructure final : public FRHICommand<FRHICommandBuildAccelerationStructure<CmdListType>>
+struct FRHICommandBuildAccelerationStructure final : public FRHICommand<FRHICommandBuildAccelerationStructure>
 {
 	FRHIRayTracingScene* Scene;
 
@@ -2065,8 +2075,7 @@
 	RHI_API void Execute(FRHICommandListBase& CmdList);
 };
 
-template <ECmdList CmdListType>
-struct FRHICommandBuildAccelerationStructures final : public FRHICommand<FRHICommandBuildAccelerationStructures<CmdListType>>
+struct FRHICommandBuildAccelerationStructures final : public FRHICommand<FRHICommandBuildAccelerationStructures>
 {
 	const TArrayView<const FAccelerationStructureBuildParams> Params;
 
@@ -2144,20 +2153,26 @@
 	{
 		EBindingType_HitGroup,
 		EBindingType_CallableShader,
+		EBindingType_MissShader,
+		EBindingType_HitGroupBatch,
 	};
 
-	FRHIRayTracingScene* Scene;
-	uint32 InstanceIndex;
-	uint32 SegmentIndex;
-	uint32 ShaderSlot;
-	FRayTracingPipelineState* Pipeline;
-	uint32 ShaderIndex;
-	uint32 NumUniformBuffers;
-	FRHIUniformBuffer* const* UniformBuffers; // Pointer to an array of uniform buffers, allocated inline within the command list
-	uint32 LooseParameterDataSize;
-	const void* LooseParameterData;
-	uint32 UserData;
-	EBindingType BindingType;
+	FRHIRayTracingScene* Scene = nullptr;
+	EBindingType BindingType = EBindingType_HitGroup;
+	uint32 InstanceIndex = 0;
+	uint32 SegmentIndex = 0;
+	uint32 ShaderSlot = 0;
+	FRayTracingPipelineState* Pipeline = nullptr;
+	uint32 ShaderIndex = 0;
+	uint32 NumUniformBuffers = 0;
+	FRHIUniformBuffer* const* UniformBuffers = nullptr; // Pointer to an array of uniform buffers, allocated inline within the command list
+	uint32 LooseParameterDataSize = 0;
+	const void* LooseParameterData = nullptr;
+	uint32 UserData = 0;
 
+	// Batched bindings
+	uint32 NumBindings = 0;
+	const FRayTracingLocalShaderBindings* Bindings = nullptr;
+
 	// Hit group bindings
 	FRHICommandSetRayTracingBindings(FRHIRayTracingScene* InScene, uint32 InInstanceIndex, uint32 InSegmentIndex, uint32 InShaderSlot,
@@ -2165,6 +2180,7 @@
 		uint32 InLooseParameterDataSize, const void* InLooseParameterData,
 		uint32 InUserData)
 		: Scene(InScene)
+		, BindingType(EBindingType_HitGroup)
 		, InstanceIndex(InInstanceIndex)
 		, SegmentIndex(InSegmentIndex)
 		, ShaderSlot(InShaderSlot)
@@ -2175,16 +2191,26 @@
 		, LooseParameterDataSize(InLooseParameterDataSize)
 		, LooseParameterData(InLooseParameterData)
 		, UserData(InUserData)
-		, BindingType(EBindingType_HitGroup)
 	{
 	}
 
-	// Callable shader bindings
+	// Batched hit group bindings
+	FRHICommandSetRayTracingBindings(FRHIRayTracingScene* InScene, FRayTracingPipelineState* InPipeline, uint32 InNumBindings, const FRayTracingLocalShaderBindings* InBindings)
+		: Scene(InScene)
+		, BindingType(EBindingType_HitGroupBatch)
+		, Pipeline(InPipeline)
+		, NumBindings(InNumBindings)
+		, Bindings(InBindings)
+	{
+	}
+
+	// Callable and Miss shader bindings
 	FRHICommandSetRayTracingBindings(FRHIRayTracingScene* InScene, uint32 InShaderSlot,
 		FRayTracingPipelineState* InPipeline, uint32 InShaderIndex,
 		uint32 InNumUniformBuffers, FRHIUniformBuffer* const* InUniformBuffers,
-		uint32 InUserData)
+		uint32 InUserData, EBindingType InBindingType)
 		: Scene(InScene)
+		, BindingType(InBindingType)
 		, InstanceIndex(0)
 		, SegmentIndex(0)
 		, ShaderSlot(InShaderSlot)
@@ -2193,7 +2219,6 @@
 		, NumUniformBuffers(InNumUniformBuffers)
 		, UniformBuffers(InUniformBuffers)
 		, UserData(InUserData)
-		, BindingType(EBindingType_CallableShader)
 	{
 	}
 
@@ -2205,19 +2230,379 @@
 #define ALLOC_COMMAND(...) new ( AllocCommand(sizeof(__VA_ARGS__), alignof(__VA_ARGS__)) ) __VA_ARGS__
 #define ALLOC_COMMAND_CL(RHICmdList, ...) new ( (RHICmdList).AllocCommand(sizeof(__VA_ARGS__), alignof(__VA_ARGS__)) ) __VA_ARGS__
 
-template<> void FRHICommandSetShaderParameter<FRHIComputeShader, ECmdList::ECompute>::Execute(FRHICommandListBase& CmdList);
-template<> void FRHICommandSetShaderUniformBuffer<FRHIComputeShader, ECmdList::ECompute>::Execute(FRHICommandListBase& CmdList);
-template<> void FRHICommandSetShaderTexture<FRHIComputeShader, ECmdList::ECompute>::Execute(FRHICommandListBase& CmdList);
-template<> void FRHICommandSetShaderResourceViewParameter<FRHIComputeShader, ECmdList::ECompute>::Execute(FRHICommandListBase& CmdList);
-template<> void FRHICommandSetShaderSampler<FRHIComputeShader, ECmdList::ECompute>::Execute(FRHICommandListBase& CmdList);
+template<> void FRHICommandSetShaderParameter<FRHIComputeShader>::Execute(FRHICommandListBase& CmdList);
+template<> void FRHICommandSetShaderUniformBuffer<FRHIComputeShader>::Execute(FRHICommandListBase& CmdList);
+template<> void FRHICommandSetShaderTexture<FRHIComputeShader>::Execute(FRHICommandListBase& CmdList);
+template<> void FRHICommandSetShaderResourceViewParameter<FRHIComputeShader>::Execute(FRHICommandListBase& CmdList);
+template<> void FRHICommandSetShaderSampler<FRHIComputeShader>::Execute(FRHICommandListBase& CmdList);
 
 
-class RHI_API FRHICommandList : public FRHICommandListBase
+class RHI_API FRHIComputeCommandList : public FRHICommandListBase
 {
 public:
+	FRHIComputeCommandList(FRHIGPUMask GPUMask) : FRHICommandListBase(GPUMask) {}
 
-	FORCEINLINE FRHICommandList(FRHIGPUMask GPUMask) : FRHICommandListBase(GPUMask) {}
+	/** Custom new/delete with recycling */
+	void* operator new(size_t Size);
+	void operator delete(void *RawMemory);
 
+	FORCEINLINE_DEBUGGABLE void SetGlobalUniformBuffers(const FUniformBufferStaticBindings& UniformBuffers)
+	{
+		if (Bypass())
+		{
+			GetComputeContext().RHISetGlobalUniformBuffers(UniformBuffers);
+			return;
+		}
+		ALLOC_COMMAND(FRHICommandSetGlobalUniformBuffers)(UniformBuffers);
+	}
+
+	FORCEINLINE_DEBUGGABLE void SetShaderUniformBuffer(FRHIComputeShader* Shader, uint32 BaseIndex, FRHIUniformBuffer* UniformBuffer)
+	{
+		if (Bypass())
+		{
+			GetComputeContext().RHISetShaderUniformBuffer(Shader, BaseIndex, UniformBuffer);
+			return;
+		}
+		ALLOC_COMMAND(FRHICommandSetShaderUniformBuffer<FRHIComputeShader>)(Shader, BaseIndex, UniformBuffer);
+	}
+
+	FORCEINLINE void SetShaderUniformBuffer(FComputeShaderRHIRef& Shader, uint32 BaseIndex, FRHIUniformBuffer* UniformBuffer)
+	{
+		SetShaderUniformBuffer(Shader.GetReference(), BaseIndex, UniformBuffer);
+	}
+
+	FORCEINLINE_DEBUGGABLE void SetShaderParameter(FRHIComputeShader* Shader, uint32 BufferIndex, uint32 BaseIndex, uint32 NumBytes, const void* NewValue)
+	{
+		if (Bypass())
+		{
+			GetComputeContext().RHISetShaderParameter(Shader, BufferIndex, BaseIndex, NumBytes, NewValue);
+			return;
+		}
+		void* UseValue = Alloc(NumBytes, 16);
+		FMemory::Memcpy(UseValue, NewValue, NumBytes);
+		ALLOC_COMMAND(FRHICommandSetShaderParameter<FRHIComputeShader>)(Shader, BufferIndex, BaseIndex, NumBytes, UseValue);
+	}
+
+	FORCEINLINE void SetShaderParameter(FComputeShaderRHIRef& Shader, uint32 BufferIndex, uint32 BaseIndex, uint32 NumBytes, const void* NewValue)
+	{
+		SetShaderParameter(Shader.GetReference(), BufferIndex, BaseIndex, NumBytes, NewValue);
+	}
+
+	FORCEINLINE_DEBUGGABLE void SetShaderTexture(FRHIComputeShader* Shader, uint32 TextureIndex, FRHITexture* Texture)
+	{
+		if (Bypass())
+		{
+			GetComputeContext().RHISetShaderTexture(Shader, TextureIndex, Texture);
+			return;
+		}
+		ALLOC_COMMAND(FRHICommandSetShaderTexture<FRHIComputeShader>)(Shader, TextureIndex, Texture);
+	}
+
+	FORCEINLINE_DEBUGGABLE void SetShaderResourceViewParameter(FRHIComputeShader* Shader, uint32 SamplerIndex, FRHIShaderResourceView* SRV)
+	{
+		if (Bypass())
+		{
+			GetComputeContext().RHISetShaderResourceViewParameter(Shader, SamplerIndex, SRV);
+			return;
+		}
+		ALLOC_COMMAND(FRHICommandSetShaderResourceViewParameter<FRHIComputeShader>)(Shader, SamplerIndex, SRV);
+	}
+
+	FORCEINLINE_DEBUGGABLE void SetShaderSampler(FRHIComputeShader* Shader, uint32 SamplerIndex, FRHISamplerState* State)
+	{
+		// Immutable samplers can't be set dynamically
+		check(!State->IsImmutable());
+		if (State->IsImmutable())
+		{
+			return;
+		}
+
+		if (Bypass())
+		{
+			GetComputeContext().RHISetShaderSampler(Shader, SamplerIndex, State);
+			return;
+		}
+		ALLOC_COMMAND(FRHICommandSetShaderSampler<FRHIComputeShader>)(Shader, SamplerIndex, State);
+	}
+
+	FORCEINLINE_DEBUGGABLE void SetUAVParameter(FRHIComputeShader* Shader, uint32 UAVIndex, FRHIUnorderedAccessView* UAV)
+	{
+		if (Bypass())
+		{
+			GetComputeContext().RHISetUAVParameter(Shader, UAVIndex, UAV);
+			return;
+		}
+		ALLOC_COMMAND(FRHICommandSetUAVParameter<FRHIComputeShader>)(Shader, UAVIndex, UAV);
+	}
+
+	FORCEINLINE_DEBUGGABLE void SetUAVParameter(FRHIComputeShader* Shader, uint32 UAVIndex, FRHIUnorderedAccessView* UAV, uint32 InitialCount)
+	{
+		if (Bypass())
+		{
+			GetComputeContext().RHISetUAVParameter(Shader, UAVIndex, UAV, InitialCount);
+			return;
+		}
+		ALLOC_COMMAND(FRHICommandSetUAVParameter_InitialCount)(Shader, UAVIndex, UAV, InitialCount);
+	}
+
+	FORCEINLINE_DEBUGGABLE void SetComputeShader(FRHIComputeShader* ComputeShader)
+	{
+		ComputeShader->UpdateStats();
+		if (Bypass())
+		{
+			GetComputeContext().RHISetComputeShader(ComputeShader);
+			return;
+		}
+		ALLOC_COMMAND(FRHICommandSetComputeShader)(ComputeShader);
+	}
+
+	FORCEINLINE_DEBUGGABLE void SetComputePipelineState(FComputePipelineState* ComputePipelineState)
+	{
+		if (Bypass())
+		{
+			extern RHI_API FRHIComputePipelineState* ExecuteSetComputePipelineState(FComputePipelineState* ComputePipelineState);
+			FRHIComputePipelineState* RHIComputePipelineState = ExecuteSetComputePipelineState(ComputePipelineState);
+			GetComputeContext().RHISetComputePipelineState(RHIComputePipelineState);
+			return;
+		}
+		ALLOC_COMMAND(FRHICommandSetComputePipelineState)(ComputePipelineState);
+	}
+
+	FORCEINLINE_DEBUGGABLE void SetAsyncComputeBudget(EAsyncComputeBudget Budget)
+	{
+		if (Bypass())
+		{
+			GetComputeContext().RHISetAsyncComputeBudget(Budget);
+			return;
+		}
+		ALLOC_COMMAND(FRHICommandSetAsyncComputeBudget)(Budget);
+	}
+
+	FORCEINLINE_DEBUGGABLE void DispatchComputeShader(uint32 ThreadGroupCountX, uint32 ThreadGroupCountY, uint32 ThreadGroupCountZ)
+	{
+		if (Bypass())
+		{
+			GetComputeContext().RHIDispatchComputeShader(ThreadGroupCountX, ThreadGroupCountY, ThreadGroupCountZ);
+			return;
+		}
+		ALLOC_COMMAND(FRHICommandDispatchComputeShader)(ThreadGroupCountX, ThreadGroupCountY, ThreadGroupCountZ);
+	}
+
+	FORCEINLINE_DEBUGGABLE void DispatchIndirectComputeShader(FRHIVertexBuffer* ArgumentBuffer, uint32 ArgumentOffset)
+	{
+		if (Bypass())
+		{
+			GetComputeContext().RHIDispatchIndirectComputeShader(ArgumentBuffer, ArgumentOffset);
+			return;
+		}
+		ALLOC_COMMAND(FRHICommandDispatchIndirectComputeShader)(ArgumentBuffer, ArgumentOffset);
+	}
+
+	FORCEINLINE_DEBUGGABLE void ClearUAVFloat(FRHIUnorderedAccessView* UnorderedAccessViewRHI, const FVector4& Values)
+	{
+		if (Bypass())
+		{
+			GetComputeContext().RHIClearUAVFloat(UnorderedAccessViewRHI, Values);
+			return;
+		}
+		ALLOC_COMMAND(FRHICommandClearUAVFloat)(UnorderedAccessViewRHI, Values);
+	}
+
+	FORCEINLINE_DEBUGGABLE void ClearUAVUint(FRHIUnorderedAccessView* UnorderedAccessViewRHI, const FUintVector4& Values)
+	{
+		if (Bypass())
+		{
+			GetComputeContext().RHIClearUAVUint(UnorderedAccessViewRHI, Values);
+			return;
+		}
+		ALLOC_COMMAND(FRHICommandClearUAVUint)(UnorderedAccessViewRHI, Values);
+	}
+
+	FORCEINLINE_DEBUGGABLE void TransitionResource(EResourceTransitionAccess TransitionType, EResourceTransitionPipeline TransitionPipeline, FRHIUnorderedAccessView* InUAV, FRHIComputeFence* WriteFence)
+	{
+		check(InUAV == nullptr || InUAV->IsCommitted());
+		FRHIUnorderedAccessView* UAV = InUAV;
+		if (Bypass())
+		{
+			GetComputeContext().RHITransitionResources(TransitionType, TransitionPipeline, &UAV, 1, WriteFence);
+			return;
+		}
+
+		// Allocate space to hold the single UAV pointer inline in the command list itself.
+		FRHIUnorderedAccessView** UAVArray = (FRHIUnorderedAccessView**)Alloc(sizeof(FRHIUnorderedAccessView*), alignof(FRHIUnorderedAccessView*));
+		UAVArray[0] = UAV;
+		ALLOC_COMMAND(FRHICommandTransitionUAVs)(TransitionType, TransitionPipeline, UAVArray, 1, WriteFence);
+	}
+
+	FORCEINLINE_DEBUGGABLE void TransitionResource(EResourceTransitionAccess TransitionType, EResourceTransitionPipeline TransitionPipeline, FRHIUnorderedAccessView* InUAV)
+	{
+		TransitionResource(TransitionType, TransitionPipeline, InUAV, nullptr);
+	}
+
+	FORCEINLINE_DEBUGGABLE void TransitionResources(EResourceTransitionAccess TransitionType, EResourceTransitionPipeline TransitionPipeline, FRHIUnorderedAccessView** InUAVs, int32 NumUAVs, FRHIComputeFence* WriteFence)
+	{
+		if (Bypass())
+		{
+			GetComputeContext().RHITransitionResources(TransitionType, TransitionPipeline, InUAVs, NumUAVs, WriteFence);
+			return;
+		}
+
+		// Allocate space to hold the list UAV pointers inline in the command list itself.
+		FRHIUnorderedAccessView** UAVArray = (FRHIUnorderedAccessView**)Alloc(sizeof(FRHIUnorderedAccessView*) * NumUAVs, alignof(FRHIUnorderedAccessView*));
+		for (int32 Index = 0; Index < NumUAVs; ++Index)
+		{
+			UAVArray[Index] = InUAVs[Index];
+		}
+
+		ALLOC_COMMAND(FRHICommandTransitionUAVs)(TransitionType, TransitionPipeline, UAVArray, NumUAVs, WriteFence);
+	}
+
+	FORCEINLINE_DEBUGGABLE void TransitionResources(EResourceTransitionAccess TransitionType, EResourceTransitionPipeline TransitionPipeline, FRHIUnorderedAccessView** InUAVs, int32 NumUAVs)
+	{
+		TransitionResources(TransitionType, TransitionPipeline, InUAVs, NumUAVs, nullptr);
+	}
+
+	FORCEINLINE_DEBUGGABLE void WaitComputeFence(FRHIComputeFence* WaitFence)
+	{
+		if (Bypass())
+		{
+			GetComputeContext().RHIWaitComputeFence(WaitFence);
+			return;
+		}
+		ALLOC_COMMAND(FRHICommandWaitComputeFence)(WaitFence);
+	}
+
+	FORCEINLINE_DEBUGGABLE void PushEvent(const TCHAR* Name, FColor Color)
+	{
+		if (Bypass())
+		{
+			GetComputeContext().RHIPushEvent(Name, Color);
+			return;
+		}
+		TCHAR* NameCopy = AllocString(Name);
+		ALLOC_COMMAND(FRHICommandPushEvent)(NameCopy, Color);
+	}
+
+	FORCEINLINE_DEBUGGABLE void PopEvent()
+	{
+		if (Bypass())
+		{
+			GetComputeContext().RHIPopEvent();
+			return;
+		}
+		ALLOC_COMMAND(FRHICommandPopEvent)();
+	}
+
+	FORCEINLINE_DEBUGGABLE void BreakPoint()
+	{
+#if !UE_BUILD_SHIPPING
+		if (Bypass())
+		{
+			if (FPlatformMisc::IsDebuggerPresent())
+			{
+				UE_DEBUG_BREAK();
+			}
+			return;
+		}
+		ALLOC_COMMAND(FRHICommandDebugBreak)();
+#endif
+	}
+
+	FORCEINLINE_DEBUGGABLE void SubmitCommandsHint()
+	{
+		if (Bypass())
+		{
+			GetComputeContext().RHISubmitCommandsHint();
+			return;
+		}
+		ALLOC_COMMAND(FRHICommandSubmitCommandsHint)();
+	}
+
+	FORCEINLINE_DEBUGGABLE void CopyToStagingBuffer(FRHIVertexBuffer* SourceBuffer, FRHIStagingBuffer* DestinationStagingBuffer, uint32 Offset, uint32 NumBytes)
+	{
+		if (Bypass())
+		{
+			GetComputeContext().RHICopyToStagingBuffer(SourceBuffer, DestinationStagingBuffer, Offset, NumBytes);
+			return;
+		}
+		ALLOC_COMMAND(FRHICommandCopyToStagingBuffer)(SourceBuffer, DestinationStagingBuffer, Offset, NumBytes);
+	}
+
+	FORCEINLINE_DEBUGGABLE void WriteGPUFence(FRHIGPUFence* Fence)
+	{
+		if (Bypass())
+		{
+			GetComputeContext().RHIWriteGPUFence(Fence);
+			return;
+		}
+		ALLOC_COMMAND(FRHICommandWriteGPUFence)(Fence);
+	}
+
+	FORCEINLINE_DEBUGGABLE void SetGPUMask(FRHIGPUMask InGPUMask)
+	{
+		if (GPUMask != InGPUMask)
+		{
+			GPUMask = InGPUMask;
+#if WITH_MGPU
+			if (!HasCommands())
+			{
+				// Update even in Bypass mode to make sure it has the correct value after a toggle.
+				InitialGPUMask = GPUMask;
+				if (Bypass())
+				{
+					GetComputeContext().RHISetGPUMask(GPUMask);
+					return;
+				}
+			}
+			else
+			{
+				checkSlow(!Bypass());
+				ALLOC_COMMAND(FRHICommandSetGPUMask)(GPUMask);
+			}
+#endif
+		}
+	}
+
+#if RHI_RAYTRACING
+	FORCEINLINE_DEBUGGABLE void BuildAccelerationStructure(FRHIRayTracingGeometry* Geometry)
+	{
+		FAccelerationStructureBuildParams Params;
+		Params.Geometry = Geometry;
+		Params.BuildMode = EAccelerationStructureBuildMode::Build;
+		BuildAccelerationStructures(MakeArrayView(&Params, 1));
+	}
+
+	FORCEINLINE_DEBUGGABLE void BuildAccelerationStructures(const TArrayView<const FAccelerationStructureBuildParams> Params)
+	{
+		if (Bypass())
+		{
+			GetComputeContext().RHIBuildAccelerationStructures(Params);
+		}
+		else
+		{
+			ALLOC_COMMAND(FRHICommandBuildAccelerationStructures)(AllocArray(Params));
+		}
+	}
+
+	FORCEINLINE_DEBUGGABLE void BuildAccelerationStructure(FRHIRayTracingScene* Scene)
+	{
+		if (Bypass())
+		{
+			GetComputeContext().RHIBuildAccelerationStructure(Scene);
+		}
+		else
+		{
+			ALLOC_COMMAND(FRHICommandBuildAccelerationStructure)(Scene);
+		}
+	}
+#endif
+};
+
+class RHI_API FRHICommandList : public FRHIComputeCommandList
+{
+public:
+	FRHICommandList(FRHIGPUMask GPUMask) : FRHIComputeCommandList(GPUMask) {}
+
 	bool AsyncPSOCompileAllowed() const
 	{
 		return bAsyncPSOCompileAllowed;
@@ -2339,8 +2724,9 @@
 		SetLocalShaderUniformBuffer(Shader.GetReference(), BaseIndex, UniformBuffer);
 	}
 
-	template <typename TShaderRHI>
-	FORCEINLINE_DEBUGGABLE void SetShaderUniformBuffer(TShaderRHI* Shader, uint32 BaseIndex, FRHIUniformBuffer* UniformBuffer)
+	using FRHIComputeCommandList::SetShaderUniformBuffer;
+
+	FORCEINLINE_DEBUGGABLE void SetShaderUniformBuffer(FRHIGraphicsShader* Shader, uint32 BaseIndex, FRHIUniformBuffer* UniformBuffer)
 	{
 		//check(IsOutsideRenderPass());
 		if (Bypass())
@@ -2348,16 +2734,18 @@
 			GetContext().RHISetShaderUniformBuffer(Shader, BaseIndex, UniformBuffer);
 			return;
 		}
-		ALLOC_COMMAND(FRHICommandSetShaderUniformBuffer<TShaderRHI, ECmdList::EGfx>)(Shader, BaseIndex, UniformBuffer);
+		ALLOC_COMMAND(FRHICommandSetShaderUniformBuffer<FRHIGraphicsShader>)(Shader, BaseIndex, UniformBuffer);
 	}
+
 	template <typename TShaderRHI>
 	FORCEINLINE void SetShaderUniformBuffer(const TRefCountPtr<TShaderRHI>& Shader, uint32 BaseIndex, FRHIUniformBuffer* UniformBuffer)
 	{
 		SetShaderUniformBuffer(Shader.GetReference(), BaseIndex, UniformBuffer);
 	}
 
-	template <typename TShaderRHI>
-	FORCEINLINE_DEBUGGABLE void SetShaderParameter(TShaderRHI* Shader, uint32 BufferIndex, uint32 BaseIndex, uint32 NumBytes, const void* NewValue)
+	using FRHIComputeCommandList::SetShaderParameter;
+
+	FORCEINLINE_DEBUGGABLE void SetShaderParameter(FRHIGraphicsShader* Shader, uint32 BufferIndex, uint32 BaseIndex, uint32 NumBytes, const void* NewValue)
 	{
 		//check(IsOutsideRenderPass());
 		if (Bypass())
@@ -2367,16 +2755,18 @@
 		}
 		void* UseValue = Alloc(NumBytes, 16);
 		FMemory::Memcpy(UseValue, NewValue, NumBytes);
-		ALLOC_COMMAND(FRHICommandSetShaderParameter<TShaderRHI, ECmdList::EGfx>)(Shader, BufferIndex, BaseIndex, NumBytes, UseValue);
+		ALLOC_COMMAND(FRHICommandSetShaderParameter<FRHIGraphicsShader>)(Shader, BufferIndex, BaseIndex, NumBytes, UseValue);
 	}
+
 	template <typename TShaderRHI>
 	FORCEINLINE void SetShaderParameter(const TRefCountPtr<TShaderRHI>& Shader, uint32 BufferIndex, uint32 BaseIndex, uint32 NumBytes, const void* NewValue)
 	{
 		SetShaderParameter(Shader.GetReference(), BufferIndex, BaseIndex, NumBytes, NewValue);
 	}
 
-	template <typename TRHIShader>
-	FORCEINLINE_DEBUGGABLE void SetShaderTexture(TRHIShader* Shader, uint32 TextureIndex, FRHITexture* Texture)
+	using FRHIComputeCommandList::SetShaderTexture;
+
+	FORCEINLINE_DEBUGGABLE void SetShaderTexture(FRHIGraphicsShader* Shader, uint32 TextureIndex, FRHITexture* Texture)
 	{
 		//check(IsOutsideRenderPass());
 		if (Bypass())
@@ -2384,7 +2774,7 @@
 			GetContext().RHISetShaderTexture(Shader, TextureIndex, Texture);
 			return;
 		}
-		ALLOC_COMMAND(FRHICommandSetShaderTexture<TRHIShader, ECmdList::EGfx>)(Shader, TextureIndex, Texture);
+		ALLOC_COMMAND(FRHICommandSetShaderTexture<FRHIGraphicsShader>)(Shader, TextureIndex, Texture);
 	}
 
 	template <typename TShaderRHI>
@@ -2393,14 +2783,9 @@
 		SetShaderTexture(Shader.GetReference(), TextureIndex, Texture);
 	}
 
-	template <typename TShaderRHI>
-	FORCEINLINE_DEBUGGABLE void SetShaderTexture(const TShaderRHI* Shader, uint32 TextureIndex, FRHITexture* Texture)
-	{
-		SetShaderTexture((TShaderRHI*)Shader, TextureIndex, Texture);
-	}
+	using FRHIComputeCommandList::SetShaderResourceViewParameter;
 
-	template <typename TRHIShader>
-	FORCEINLINE_DEBUGGABLE void SetShaderResourceViewParameter(TRHIShader* Shader, uint32 SamplerIndex, FRHIShaderResourceView* SRV)
+	FORCEINLINE_DEBUGGABLE void SetShaderResourceViewParameter(FRHIGraphicsShader* Shader, uint32 SamplerIndex, FRHIShaderResourceView* SRV)
 	{
 		//check(IsOutsideRenderPass());
 		if (Bypass())
@@ -2408,7 +2793,7 @@
 			GetContext().RHISetShaderResourceViewParameter(Shader, SamplerIndex, SRV);
 			return;
 		}
-		ALLOC_COMMAND(FRHICommandSetShaderResourceViewParameter<TRHIShader, ECmdList::EGfx>)(Shader, SamplerIndex, SRV);
+		ALLOC_COMMAND(FRHICommandSetShaderResourceViewParameter<FRHIGraphicsShader>)(Shader, SamplerIndex, SRV);
 	}
 
 	template <typename TShaderRHI>
@@ -2417,8 +2802,9 @@
 		SetShaderResourceViewParameter(Shader.GetReference(), SamplerIndex, SRV);
 	}
 
-	template <typename TRHIShader>
-	FORCEINLINE_DEBUGGABLE void SetShaderSampler(TRHIShader* Shader, uint32 SamplerIndex, FRHISamplerState* State)
+	using FRHIComputeCommandList::SetShaderSampler;
+
+	FORCEINLINE_DEBUGGABLE void SetShaderSampler(FRHIGraphicsShader* Shader, uint32 SamplerIndex, FRHISamplerState* State)
 	{
 		//check(IsOutsideRenderPass());
 		
@@ -2434,7 +2820,7 @@
 			GetContext().RHISetShaderSampler(Shader, SamplerIndex, State);
 			return;
 		}
-		ALLOC_COMMAND(FRHICommandSetShaderSampler<TRHIShader, ECmdList::EGfx>)(Shader, SamplerIndex, State);
+		ALLOC_COMMAND(FRHICommandSetShaderSampler<FRHIGraphicsShader>)(Shader, SamplerIndex, State);
 	}
 
 	template <typename TShaderRHI>
@@ -2443,36 +2829,23 @@
 		SetShaderSampler(Shader.GetReference(), SamplerIndex, State);
 	}
 
-	FORCEINLINE_DEBUGGABLE void SetUAVParameter(FRHIComputeShader* Shader, uint32 UAVIndex, FRHIUnorderedAccessView* UAV)
+	using FRHIComputeCommandList::SetUAVParameter;
+
+	FORCEINLINE_DEBUGGABLE void SetUAVParameter(FRHIPixelShader* Shader, uint32 UAVIndex, FRHIUnorderedAccessView* UAV)
 	{
 		if (Bypass())
 		{
 			GetContext().RHISetUAVParameter(Shader, UAVIndex, UAV);
 			return;
 		}
-		ALLOC_COMMAND(FRHICommandSetUAVParameter<FRHIComputeShader, ECmdList::EGfx>)(Shader, UAVIndex, UAV);
+		ALLOC_COMMAND(FRHICommandSetUAVParameter<FRHIPixelShader>)(Shader, UAVIndex, UAV);
 	}
 
-	FORCEINLINE_DEBUGGABLE void SetUAVParameter(const TRefCountPtr<FRHIComputeShader>& Shader, uint32 UAVIndex, FRHIUnorderedAccessView* UAV)
+	FORCEINLINE_DEBUGGABLE void SetUAVParameter(const TRefCountPtr<FRHIPixelShader>& Shader, uint32 UAVIndex, FRHIUnorderedAccessView* UAV)
 	{
 		SetUAVParameter(Shader.GetReference(), UAVIndex, UAV);
 	}
 
-	FORCEINLINE_DEBUGGABLE void SetUAVParameter(FRHIComputeShader* Shader, uint32 UAVIndex, FRHIUnorderedAccessView* UAV, uint32 InitialCount)
-	{
-		if (Bypass())
-		{
-			GetContext().RHISetUAVParameter(Shader, UAVIndex, UAV, InitialCount);
-			return;
-		}
-		ALLOC_COMMAND(FRHICommandSetUAVParameter_IntialCount<FRHIComputeShader, ECmdList::EGfx>)(Shader, UAVIndex, UAV, InitialCount);
-	}
-
-	FORCEINLINE_DEBUGGABLE void SetUAVParameter(const TRefCountPtr<FRHIComputeShader>& Shader, uint32 UAVIndex, FRHIUnorderedAccessView* UAV, uint32 InitialCount)
-	{
-		SetUAVParameter(Shader.GetReference(), UAVIndex, UAV, InitialCount);
-	}
-
 	FORCEINLINE_DEBUGGABLE void SetBlendFactor(const FLinearColor& BlendFactor = FLinearColor::White)
 	{
 		//check(IsOutsideRenderPass());
@@ -2533,7 +2906,7 @@
 		ALLOC_COMMAND(FRHICommandSetStencilRef)(StencilRef);
 	}
 
-	FORCEINLINE_DEBUGGABLE void SetViewport(uint32 MinX, uint32 MinY, float MinZ, uint32 MaxX, uint32 MaxY, float MaxZ)
+	FORCEINLINE_DEBUGGABLE void SetViewport(float MinX, float MinY, float MinZ, float MaxX, float MaxY, float MaxZ)
 	{
 		//check(IsOutsideRenderPass());
 		if (Bypass())
@@ -2544,7 +2917,7 @@
 		ALLOC_COMMAND(FRHICommandSetViewport)(MinX, MinY, MinZ, MaxX, MaxY, MaxZ);
 	}
 
-	FORCEINLINE_DEBUGGABLE void SetStereoViewport(uint32 LeftMinX, uint32 RightMinX, uint32 LeftMinY, uint32 RightMinY, float MinZ, uint32 LeftMaxX, uint32 RightMaxX, uint32 LeftMaxY, uint32 RightMaxY, float MaxZ)
+	FORCEINLINE_DEBUGGABLE void SetStereoViewport(float LeftMinX, float RightMinX, float LeftMinY, float RightMinY, float MinZ, float LeftMaxX, float RightMaxX, float LeftMaxY, float RightMaxY, float MaxZ)
 	{
 		//check(IsOutsideRenderPass());
 		if (Bypass())
@@ -2623,10 +2996,7 @@
 	FORCEINLINE_DEBUGGABLE void SetRenderTargets(
 		uint32 NewNumSimultaneousRenderTargets,
 		const FRHIRenderTargetView* NewRenderTargetsRHI,
-		const FRHIDepthRenderTargetView* NewDepthStencilTargetRHI,
-		uint32 NewNumUAVs,
-		FRHIUnorderedAccessView* const* UAVs
-		)
+		const FRHIDepthRenderTargetView* NewDepthStencilTargetRHI)
 	{
 		check(IsOutsideRenderPass());
 		CacheActiveRenderTargets(
@@ -2640,17 +3010,13 @@
 			GetContext().RHISetRenderTargets(
 				NewNumSimultaneousRenderTargets,
 				NewRenderTargetsRHI,
-				NewDepthStencilTargetRHI,
-				NewNumUAVs,
-				UAVs);
+				NewDepthStencilTargetRHI);
 			return;
 		}
 		ALLOC_COMMAND(FRHICommandSetRenderTargets)(
 			NewNumSimultaneousRenderTargets,
 			NewRenderTargetsRHI,
-			NewDepthStencilTargetRHI,
-			NewNumUAVs,
-			UAVs);
+			NewDepthStencilTargetRHI);
 	}
 
 	FORCEINLINE_DEBUGGABLE void BindClearMRTValues(bool bClearColor, bool bClearDepth, bool bClearStencil)
@@ -2664,29 +3030,6 @@
 		ALLOC_COMMAND(FRHICommandBindClearMRTValues)(bClearColor, bClearDepth, bClearStencil);
 	}	
 
-	FORCEINLINE_DEBUGGABLE void SetComputeShader(FRHIComputeShader* ComputeShader)
-	{
-		ComputeShader->UpdateStats();
-		if (Bypass())
-		{
-			GetContext().RHISetComputeShader(ComputeShader);
-			return;
-		}
-		ALLOC_COMMAND(FRHICommandSetComputeShader<ECmdList::EGfx>)(ComputeShader);
-	}
-
-	FORCEINLINE_DEBUGGABLE void SetComputePipelineState(class FComputePipelineState* ComputePipelineState)
-	{
-		if (Bypass())
-		{
-			extern RHI_API FRHIComputePipelineState* ExecuteSetComputePipelineState(class FComputePipelineState* ComputePipelineState);
-			FRHIComputePipelineState* RHIComputePipelineState = ExecuteSetComputePipelineState(ComputePipelineState);
-			GetContext().RHISetComputePipelineState(RHIComputePipelineState);
-			return;
-		}
-		ALLOC_COMMAND(FRHICommandSetComputePipelineState<ECmdList::EGfx>)(ComputePipelineState);
-	}
-
 	FORCEINLINE_DEBUGGABLE void SetGraphicsPipelineState(class FGraphicsPipelineState* GraphicsPipelineState)
 	{
 		//check(IsOutsideRenderPass());
@@ -2700,26 +3043,6 @@
 		ALLOC_COMMAND(FRHICommandSetGraphicsPipelineState)(GraphicsPipelineState);
 	}
 
-	FORCEINLINE_DEBUGGABLE void DispatchComputeShader(uint32 ThreadGroupCountX, uint32 ThreadGroupCountY, uint32 ThreadGroupCountZ)
-	{
-		if (Bypass())
-		{
-			GetContext().RHIDispatchComputeShader(ThreadGroupCountX, ThreadGroupCountY, ThreadGroupCountZ);
-			return;
-		}
-		ALLOC_COMMAND(FRHICommandDispatchComputeShader<ECmdList::EGfx>)(ThreadGroupCountX, ThreadGroupCountY, ThreadGroupCountZ);
-	}
-
-	FORCEINLINE_DEBUGGABLE void DispatchIndirectComputeShader(FRHIVertexBuffer* ArgumentBuffer, uint32 ArgumentOffset)
-	{
-		if (Bypass())
-		{
-			GetContext().RHIDispatchIndirectComputeShader(ArgumentBuffer, ArgumentOffset);
-			return;
-		}
-		ALLOC_COMMAND(FRHICommandDispatchIndirectComputeShader<ECmdList::EGfx>)(ArgumentBuffer, ArgumentOffset);
-	}
-
 	FORCEINLINE_DEBUGGABLE void AutomaticCacheFlushAfterComputeShader(bool bEnable)
 	{
 		if (Bypass())
@@ -2840,15 +3163,11 @@
 		ALLOC_COMMAND(FRHICommandResummarizeHTile)(DepthTexture);
 	}
 
+	UE_DEPRECATED(4.25, "RHIClearTinyUAV is deprecated. Use RHIClearUAVUint or RHIClearUAVFloat instead.")
 	FORCEINLINE_DEBUGGABLE void ClearTinyUAV(FRHIUnorderedAccessView* UnorderedAccessViewRHI, const uint32(&Values)[4])
 	{
-		//check(IsOutsideRenderPass());
-		if (Bypass())
-		{
-			GetContext().RHIClearTinyUAV(UnorderedAccessViewRHI, Values);
-			return;
-		}
-		ALLOC_COMMAND(FRHICommandClearTinyUAV)(UnorderedAccessViewRHI, Values);
+		// Forward to the new uint clear implementation.
+		ClearUAVUint(UnorderedAccessViewRHI, FUintVector4(Values[0], Values[1], Values[2], Values[3]));
 	}
 
 	FORCEINLINE_DEBUGGABLE void BeginRenderQuery(FRHIRenderQuery* RenderQuery)
@@ -2870,16 +3189,6 @@
 		ALLOC_COMMAND(FRHICommandEndRenderQuery)(RenderQuery);
 	}
 
-	FORCEINLINE_DEBUGGABLE void SubmitCommandsHint()
-	{
-		if (Bypass())
-		{
-			GetContext().RHISubmitCommandsHint();
-			return;
-		}
-		ALLOC_COMMAND(FRHICommandSubmitCommandsHint<ECmdList::EGfx>)();
-	}
-
 	FORCEINLINE_DEBUGGABLE void PollOcclusionQueries()
 	{
 		if (Bypass())
@@ -2890,6 +3199,9 @@
 		ALLOC_COMMAND(FRHICommandPollOcclusionQueries)();
 	}
 
+	using FRHIComputeCommandList::TransitionResource;
+	using FRHIComputeCommandList::TransitionResources;
+
 	FORCEINLINE_DEBUGGABLE void TransitionResource(FExclusiveDepthStencil DepthStencilMode, FRHITexture* DepthTexture)
 	{
 		if (Bypass())
@@ -2944,80 +3256,6 @@
 		ALLOC_COMMAND(FRHICommandTransitionTexturesArray)(TransitionType, InTextures);
 	}
 
-	FORCEINLINE_DEBUGGABLE void TransitionResource(EResourceTransitionAccess TransitionType, EResourceTransitionPipeline TransitionPipeline, FRHIUnorderedAccessView* InUAV, FRHIComputeFence* WriteFence)
-	{
-		check(InUAV == nullptr || InUAV->IsCommitted());
-		FRHIUnorderedAccessView* UAV = InUAV;
-		if (Bypass())
-		{
-			GetContext().RHITransitionResources(TransitionType, TransitionPipeline, &UAV, 1, WriteFence);
-			return;
-		}
-
-		// Allocate space to hold the single UAV pointer inline in the command list itself.
-		FRHIUnorderedAccessView** UAVArray = (FRHIUnorderedAccessView**)Alloc(sizeof(FRHIUnorderedAccessView*), alignof(FRHIUnorderedAccessView*));
-		UAVArray[0] = UAV;
-		ALLOC_COMMAND(FRHICommandTransitionUAVs<ECmdList::EGfx>)(TransitionType, TransitionPipeline, UAVArray, 1, WriteFence);
-	}
-
-	FORCEINLINE_DEBUGGABLE void TransitionResource(EResourceTransitionAccess TransitionType, EResourceTransitionPipeline TransitionPipeline, FRHIUnorderedAccessView* InUAV)
-	{
-		TransitionResource(TransitionType, TransitionPipeline, InUAV, nullptr);
-	}
-
-	FORCEINLINE_DEBUGGABLE void TransitionResources(EResourceTransitionAccess TransitionType, EResourceTransitionPipeline TransitionPipeline, FRHIUnorderedAccessView** InUAVs, int32 NumUAVs, FRHIComputeFence* WriteFence)
-	{
-		if (Bypass())
-		{
-			GetContext().RHITransitionResources(TransitionType, TransitionPipeline, InUAVs, NumUAVs, WriteFence);
-			return;
-		}
-
-		// Allocate space to hold the list UAV pointers inline in the command list itself.
-		FRHIUnorderedAccessView** UAVArray = (FRHIUnorderedAccessView**)Alloc(sizeof(FRHIUnorderedAccessView*) * NumUAVs, alignof(FRHIUnorderedAccessView*));
-		for (int32 Index = 0; Index < NumUAVs; ++Index)
-		{
-			UAVArray[Index] = InUAVs[Index];
-		}
-
-		ALLOC_COMMAND(FRHICommandTransitionUAVs<ECmdList::EGfx>)(TransitionType, TransitionPipeline, UAVArray, NumUAVs, WriteFence);
-	}
-
-	FORCEINLINE_DEBUGGABLE void TransitionResources(EResourceTransitionAccess TransitionType, EResourceTransitionPipeline TransitionPipeline, FRHIUnorderedAccessView** InUAVs, int32 NumUAVs)
-	{
-		TransitionResources(TransitionType, TransitionPipeline, InUAVs, NumUAVs, nullptr);
-	}
-
-	FORCEINLINE_DEBUGGABLE void WaitComputeFence(FRHIComputeFence* WaitFence)
-	{
-		if (Bypass())
-		{
-			GetContext().RHIWaitComputeFence(WaitFence);
-			return;
-		}
-		ALLOC_COMMAND(FRHICommandWaitComputeFence<ECmdList::EGfx>)(WaitFence);
-	}
-
-	FORCEINLINE_DEBUGGABLE void CopyToStagingBuffer(FRHIVertexBuffer* SourceBuffer, FRHIStagingBuffer* DestinationStagingBuffer, uint32 Offset, uint32 NumBytes)
-	{
-		if (Bypass())
-		{
-			GetContext().RHICopyToStagingBuffer(SourceBuffer, DestinationStagingBuffer, Offset, NumBytes);
-			return;
-		}
-		ALLOC_COMMAND(FRHICommandCopyToStagingBuffer<ECmdList::EGfx>)(SourceBuffer, DestinationStagingBuffer, Offset, NumBytes);
-	}
-
-	FORCEINLINE_DEBUGGABLE void WriteGPUFence(FRHIGPUFence* Fence)
-	{
-		if (Bypass())
-		{
-			GetContext().RHIWriteGPUFence(Fence);
-			return;
-		}
-		ALLOC_COMMAND(FRHICommandWriteGPUFence<ECmdList::EGfx>)(Fence);
-	}
-
 	FORCEINLINE_DEBUGGABLE void BeginRenderPass(const FRHIRenderPassInfo& InInfo, const TCHAR* Name)
 	{
 		check(!IsInsideRenderPass());
@@ -3090,8 +3328,6 @@
 			ALLOC_COMMAND(FRHICommandBeginComputePass)(NameCopy);
 		}
 		Data.bInsideComputePass = true;
-
-		Data.bInsideComputePass = true;
 	}
 
 	void EndComputePass()
@@ -3117,27 +3353,6 @@
 	void BeginFrame();
 	void EndFrame();
 
-	FORCEINLINE_DEBUGGABLE void PushEvent(const TCHAR* Name, FColor Color)
-	{
-		if (Bypass())
-		{
-			GetContext().RHIPushEvent(Name, Color);
-			return;
-		}
-		TCHAR* NameCopy  = AllocString(Name);
-		ALLOC_COMMAND(FRHICommandPushEvent<ECmdList::EGfx>)(NameCopy, Color);
-	}
-
-	FORCEINLINE_DEBUGGABLE void PopEvent()
-	{
-		if (Bypass())
-		{
-			GetContext().RHIPopEvent();
-			return;
-		}
-		ALLOC_COMMAND(FRHICommandPopEvent<ECmdList::EGfx>)();
-	}
-	
 	FORCEINLINE_DEBUGGABLE void RHIInvalidateCachedState()
 	{
 		if (Bypass())
@@ -3158,21 +3373,6 @@
 		ALLOC_COMMAND(FRHICommandDiscardRenderTargets)(Depth, Stencil, ColorBitMask);
 	}
 	
-	FORCEINLINE_DEBUGGABLE void BreakPoint()
-	{
-#if !UE_BUILD_SHIPPING
-		if (Bypass())
-		{
-			if (FPlatformMisc::IsDebuggerPresent())
-			{
-				UE_DEBUG_BREAK();
-			}
-			return;
-		}
-		ALLOC_COMMAND(FRHICommandDebugBreak)();
-#endif
-	}
-
 #if RHI_RAYTRACING
 	// Ray tracing API
 	UE_DEPRECATED(4.25, "CopyBufferRegion API is deprecated. Use an explicit compute shader copy dispatch instead.")
@@ -3207,38 +3407,6 @@
 		}
 	}
 
-	FORCEINLINE_DEBUGGABLE void BuildAccelerationStructures(const TArrayView<const FAccelerationStructureBuildParams> Params)
-	{
-		if (Bypass())
-		{
-			GetContext().RHIBuildAccelerationStructures(Params);
-		}
-		else
-		{
-			ALLOC_COMMAND(FRHICommandBuildAccelerationStructures<ECmdList::EGfx>)(AllocArray(Params));
-		}
-	}
-
-	FORCEINLINE_DEBUGGABLE void BuildAccelerationStructure(FRHIRayTracingGeometry* Geometry)
-	{
-		FAccelerationStructureBuildParams Params;
-		Params.Geometry = Geometry;
-		Params.BuildMode = EAccelerationStructureBuildMode::Build;
-		BuildAccelerationStructures(MakeArrayView(&Params, 1));
-	}
-
-	FORCEINLINE_DEBUGGABLE void BuildAccelerationStructure(FRHIRayTracingScene* Scene)
-	{
-		if (Bypass())
-		{
-			GetContext().RHIBuildAccelerationStructure(Scene);
-		}
-		else
-		{
-			ALLOC_COMMAND(FRHICommandBuildAccelerationStructure<ECmdList::EGfx>)(Scene);
-		}
-	}
-
 	FORCEINLINE_DEBUGGABLE void ClearRayTracingBindings(FRHIRayTracingScene* Scene)
 	{
 		if (Bypass())
@@ -3303,6 +3471,58 @@
 		}
 	}
 
+	FORCEINLINE_DEBUGGABLE void SetRayTracingHitGroups(
+		FRHIRayTracingScene* Scene, FRayTracingPipelineState* Pipeline,
+		uint32 NumBindings, const FRayTracingLocalShaderBindings* Bindings,
+		bool bCopyDataToInlineStorage = true)
+	{
+		if (Bypass())
+		{
+			extern RHI_API FRHIRayTracingPipelineState* GetRHIRayTracingPipelineState(FRayTracingPipelineState*);
+			GetContext().RHISetRayTracingHitGroups(Scene, GetRHIRayTracingPipelineState(Pipeline), NumBindings, Bindings);
+		}
+		else
+		{
+			FRayTracingLocalShaderBindings* InlineBindings = nullptr;
+
+			// By default all batch binding data is stored in the command list memory.
+			// However, user may skip this copy if they take responsibility for keeping data alive until this command is executed.
+			if (bCopyDataToInlineStorage)
+			{
+				if (NumBindings)
+				{
+					uint32 Size = sizeof(FRayTracingLocalShaderBindings) * NumBindings;
+					InlineBindings = (FRayTracingLocalShaderBindings*)Alloc(Size, alignof(FRayTracingLocalShaderBindings));
+					FMemory::Memcpy(InlineBindings, Bindings, Size);
+				}
+
+				for (uint32 i = 0; i < NumBindings; ++i)
+				{
+					if (InlineBindings[i].NumUniformBuffers)
+					{
+						InlineBindings[i].UniformBuffers = (FRHIUniformBuffer**)Alloc(sizeof(FRHIUniformBuffer*) * InlineBindings[i].NumUniformBuffers, alignof(FRHIUniformBuffer*));
+						for (uint32 Index = 0; Index < InlineBindings[i].NumUniformBuffers; ++Index)
+						{
+							InlineBindings[i].UniformBuffers[Index] = Bindings[i].UniformBuffers[Index];
+						}
+					}
+
+					if (InlineBindings[i].LooseParameterDataSize)
+					{
+						InlineBindings[i].LooseParameterData = (uint8*)Alloc(InlineBindings[i].LooseParameterDataSize, 16);
+						FMemory::Memcpy(InlineBindings[i].LooseParameterData, Bindings[i].LooseParameterData, InlineBindings[i].LooseParameterDataSize);
+					}
+				}
+
+				ALLOC_COMMAND(FRHICommandSetRayTracingBindings)(Scene, Pipeline, NumBindings, InlineBindings);
+			}
+			else
+			{
+				ALLOC_COMMAND(FRHICommandSetRayTracingBindings)(Scene, Pipeline, NumBindings, Bindings);
+			}
+		}
+	}
+
 	FORCEINLINE_DEBUGGABLE void SetRayTracingHitGroup(
 		FRHIRayTracingScene* Scene, uint32 InstanceIndex, uint32 SegmentIndex, uint32 ShaderSlot,
 		FRayTracingPipelineState* Pipeline, uint32 HitGroupIndex,
@@ -3367,349 +3587,38 @@
 				}
 			}
 
-			ALLOC_COMMAND(FRHICommandSetRayTracingBindings)(Scene, ShaderSlotInScene, Pipeline, ShaderIndexInPipeline, NumUniformBuffers, InlineUniformBuffers, UserData);
+			ALLOC_COMMAND(FRHICommandSetRayTracingBindings)(Scene, ShaderSlotInScene, Pipeline, ShaderIndexInPipeline, NumUniformBuffers, InlineUniformBuffers, UserData, FRHICommandSetRayTracingBindings::EBindingType::EBindingType_CallableShader);
 		}
 	}
 
+	FORCEINLINE_DEBUGGABLE void SetRayTracingMissShader(
+		FRHIRayTracingScene* Scene, uint32 ShaderSlotInScene,
+		FRayTracingPipelineState* Pipeline, uint32 ShaderIndexInPipeline,
+		uint32 NumUniformBuffers, FRHIUniformBuffer* const* UniformBuffers,
+		uint32 UserData)
+	{
+		if (Bypass())
+		{
+			extern RHI_API FRHIRayTracingPipelineState* GetRHIRayTracingPipelineState(FRayTracingPipelineState*);
+			GetContext().RHISetRayTracingMissShader(Scene, ShaderSlotInScene, GetRHIRayTracingPipelineState(Pipeline), ShaderIndexInPipeline, NumUniformBuffers, UniformBuffers, UserData);
+		}
+		else
+		{
+			FRHIUniformBuffer** InlineUniformBuffers = nullptr;
+			if (NumUniformBuffers)
+			{
+				InlineUniformBuffers = (FRHIUniformBuffer**)Alloc(sizeof(FRHIUniformBuffer*) * NumUniformBuffers, alignof(FRHIUniformBuffer*));
+				for (uint32 Index = 0; Index < NumUniformBuffers; ++Index)
+				{
+					InlineUniformBuffers[Index] = UniformBuffers[Index];
+				}
+			}
+
+			ALLOC_COMMAND(FRHICommandSetRayTracingBindings)(Scene, ShaderSlotInScene, Pipeline, ShaderIndexInPipeline, NumUniformBuffers, InlineUniformBuffers, UserData, FRHICommandSetRayTracingBindings::EBindingType::EBindingType_MissShader);
+		}
+	}
+
 #endif // RHI_RAYTRACING
-
-	FORCEINLINE_DEBUGGABLE void SetGPUMask(FRHIGPUMask InGPUMask)
-	{
-		if (GPUMask != InGPUMask)
-		{
-			GPUMask = InGPUMask;
-			if (HasCommands() || Bypass())
-			{
-				SetGPUMaskOnContext();
-			}
-			else
-			{
-				InitialGPUMask = GPUMask;
-			}
-		}
-	}
 };
 
-class RHI_API FRHIAsyncComputeCommandList : public FRHICommandListBase
-{
-public:
-
-	FRHIAsyncComputeCommandList() : FRHICommandListBase(FRHIGPUMask::All()) {}
-
-	/** Custom new/delete with recycling */
-	void* operator new(size_t Size);
-	void operator delete(void *RawMemory);
-
-	FORCEINLINE_DEBUGGABLE void SetShaderUniformBuffer(FRHIComputeShader* Shader, uint32 BaseIndex, FRHIUniformBuffer* UniformBuffer)
-	{
-		if (Bypass())
-		{
-			GetComputeContext().RHISetShaderUniformBuffer(Shader, BaseIndex, UniformBuffer);
-			return;
-		}
-		ALLOC_COMMAND(FRHICommandSetShaderUniformBuffer<FRHIComputeShader, ECmdList::ECompute>)(Shader, BaseIndex, UniformBuffer);
-	}
-	
-	FORCEINLINE void SetShaderUniformBuffer(FComputeShaderRHIRef& Shader, uint32 BaseIndex, FRHIUniformBuffer* UniformBuffer)
-	{
-		SetShaderUniformBuffer(Shader.GetReference(), BaseIndex, UniformBuffer);
-	}
-
-	FORCEINLINE_DEBUGGABLE void SetShaderParameter(FRHIComputeShader* Shader, uint32 BufferIndex, uint32 BaseIndex, uint32 NumBytes, const void* NewValue)
-	{
-		if (Bypass())
-		{
-			GetComputeContext().RHISetShaderParameter(Shader, BufferIndex, BaseIndex, NumBytes, NewValue);
-			return;
-		}
-		void* UseValue = Alloc(NumBytes, 16);
-		FMemory::Memcpy(UseValue, NewValue, NumBytes);
-		ALLOC_COMMAND(FRHICommandSetShaderParameter<FRHIComputeShader, ECmdList::ECompute>)(Shader, BufferIndex, BaseIndex, NumBytes, UseValue);
-	}
-	
-	FORCEINLINE void SetShaderParameter(FComputeShaderRHIRef& Shader, uint32 BufferIndex, uint32 BaseIndex, uint32 NumBytes, const void* NewValue)
-	{
-		SetShaderParameter(Shader.GetReference(), BufferIndex, BaseIndex, NumBytes, NewValue);
-	}
-	
-	FORCEINLINE_DEBUGGABLE void SetShaderTexture(FRHIComputeShader* Shader, uint32 TextureIndex, FRHITexture* Texture)
-	{
-		if (Bypass())
-		{
-			GetComputeContext().RHISetShaderTexture(Shader, TextureIndex, Texture);
-			return;
-		}
-		ALLOC_COMMAND(FRHICommandSetShaderTexture<FRHIComputeShader, ECmdList::ECompute>)(Shader, TextureIndex, Texture);
-	}
-	
-	FORCEINLINE_DEBUGGABLE void SetShaderResourceViewParameter(FRHIComputeShader* Shader, uint32 SamplerIndex, FRHIShaderResourceView* SRV)
-	{
-		if (Bypass())
-		{
-			GetComputeContext().RHISetShaderResourceViewParameter(Shader, SamplerIndex, SRV);
-			return;
-		}
-		ALLOC_COMMAND(FRHICommandSetShaderResourceViewParameter<FRHIComputeShader, ECmdList::ECompute>)(Shader, SamplerIndex, SRV);
-	}
-	
-	FORCEINLINE_DEBUGGABLE void SetShaderSampler(FRHIComputeShader* Shader, uint32 SamplerIndex, FRHISamplerState* State)
-	{
-		// Immutable samplers can't be set dynamically
-		check(!State->IsImmutable());
-		if (State->IsImmutable())
-		{
-			return;
-		}
-		
-		if (Bypass())
-		{
-			GetComputeContext().RHISetShaderSampler(Shader, SamplerIndex, State);
-			return;
-		}
-		ALLOC_COMMAND(FRHICommandSetShaderSampler<FRHIComputeShader, ECmdList::ECompute>)(Shader, SamplerIndex, State);
-	}
-
-	FORCEINLINE_DEBUGGABLE void SetUAVParameter(FRHIComputeShader* Shader, uint32 UAVIndex, FRHIUnorderedAccessView* UAV)
-	{
-		if (Bypass())
-		{
-			GetComputeContext().RHISetUAVParameter(Shader, UAVIndex, UAV);
-			return;
-		}
-		ALLOC_COMMAND(FRHICommandSetUAVParameter<FRHIComputeShader, ECmdList::ECompute>)(Shader, UAVIndex, UAV);
-	}
-
-	FORCEINLINE_DEBUGGABLE void SetUAVParameter(FRHIComputeShader* Shader, uint32 UAVIndex, FRHIUnorderedAccessView* UAV, uint32 InitialCount)
-	{
-		if (Bypass())
-		{
-			GetComputeContext().RHISetUAVParameter(Shader, UAVIndex, UAV, InitialCount);
-			return;
-		}
-		ALLOC_COMMAND(FRHICommandSetUAVParameter_IntialCount<FRHIComputeShader, ECmdList::ECompute>)(Shader, UAVIndex, UAV, InitialCount);
-	}
-	
-	FORCEINLINE_DEBUGGABLE void SetComputeShader(FRHIComputeShader* ComputeShader)
-	{
-		ComputeShader->UpdateStats();
-		if (Bypass())
-		{
-			GetComputeContext().RHISetComputeShader(ComputeShader);
-			return;
-		}
-		ALLOC_COMMAND(FRHICommandSetComputeShader<ECmdList::ECompute>)(ComputeShader);
-	}
-
-	FORCEINLINE_DEBUGGABLE void SetComputePipelineState(FComputePipelineState* ComputePipelineState)
-	{
-		if (Bypass())
-		{
-			extern FRHIComputePipelineState* ExecuteSetComputePipelineState(FComputePipelineState* ComputePipelineState);
-			FRHIComputePipelineState* RHIComputePipelineState = ExecuteSetComputePipelineState(ComputePipelineState);
-			GetComputeContext().RHISetComputePipelineState(RHIComputePipelineState);
-			return;
-		}
-		ALLOC_COMMAND(FRHICommandSetComputePipelineState<ECmdList::ECompute>)(ComputePipelineState);
-	}
-
-	FORCEINLINE_DEBUGGABLE void SetAsyncComputeBudget(EAsyncComputeBudget Budget)
-	{
-		if (Bypass())
-		{
-			GetComputeContext().RHISetAsyncComputeBudget(Budget);
-			return;
-		}
-		ALLOC_COMMAND(FRHICommandSetAsyncComputeBudget<ECmdList::ECompute>)(Budget);
-	}
-
-	FORCEINLINE_DEBUGGABLE void DispatchComputeShader(uint32 ThreadGroupCountX, uint32 ThreadGroupCountY, uint32 ThreadGroupCountZ)
-	{
-		if (Bypass())
-		{
-			GetComputeContext().RHIDispatchComputeShader(ThreadGroupCountX, ThreadGroupCountY, ThreadGroupCountZ);
-			return;
-		}
-		ALLOC_COMMAND(FRHICommandDispatchComputeShader<ECmdList::ECompute>)(ThreadGroupCountX, ThreadGroupCountY, ThreadGroupCountZ);
-	}
-
-	FORCEINLINE_DEBUGGABLE void DispatchIndirectComputeShader(FRHIVertexBuffer* ArgumentBuffer, uint32 ArgumentOffset)
-	{
-		if (Bypass())
-		{
-			GetComputeContext().RHIDispatchIndirectComputeShader(ArgumentBuffer, ArgumentOffset);
-			return;
-		}
-		ALLOC_COMMAND(FRHICommandDispatchIndirectComputeShader<ECmdList::ECompute>)(ArgumentBuffer, ArgumentOffset);
-	}
-
-	FORCEINLINE_DEBUGGABLE void TransitionResource(EResourceTransitionAccess TransitionType, EResourceTransitionPipeline TransitionPipeline, FRHIUnorderedAccessView* InUAV, FRHIComputeFence* WriteFence)
-	{
-		FRHIUnorderedAccessView* UAV = InUAV;
-		if (Bypass())
-		{
-			GetComputeContext().RHITransitionResources(TransitionType, TransitionPipeline, &UAV, 1, WriteFence);
-			return;
-		}
-
-		// Allocate space to hold the single UAV pointer inline in the command list itself.
-		FRHIUnorderedAccessView** UAVArray = (FRHIUnorderedAccessView**)Alloc(sizeof(FRHIUnorderedAccessView*), alignof(FRHIUnorderedAccessView*));
-		UAVArray[0] = UAV;
-		ALLOC_COMMAND(FRHICommandTransitionUAVs<ECmdList::ECompute>)(TransitionType, TransitionPipeline, UAVArray, 1, WriteFence);
-	}
-
-	FORCEINLINE_DEBUGGABLE void TransitionResource(EResourceTransitionAccess TransitionType, EResourceTransitionPipeline TransitionPipeline, FRHIUnorderedAccessView* InUAV)
-	{
-		TransitionResource(TransitionType, TransitionPipeline, InUAV, nullptr);
-	}
-
-	FORCEINLINE_DEBUGGABLE void TransitionResources(EResourceTransitionAccess TransitionType, EResourceTransitionPipeline TransitionPipeline, FRHIUnorderedAccessView** InUAVs, int32 NumUAVs, FRHIComputeFence* WriteFence)
-	{
-		if (Bypass())
-		{
-			GetComputeContext().RHITransitionResources(TransitionType, TransitionPipeline, InUAVs, NumUAVs, WriteFence);
-			return;
-		}
-
-		// Allocate space to hold the list UAV pointers inline in the command list itself.
-		FRHIUnorderedAccessView** UAVArray = (FRHIUnorderedAccessView**)Alloc(sizeof(FRHIUnorderedAccessView*) * NumUAVs, alignof(FRHIUnorderedAccessView*));
-		for (int32 Index = 0; Index < NumUAVs; ++Index)
-		{
-			UAVArray[Index] = InUAVs[Index];
-		}
-
-		ALLOC_COMMAND(FRHICommandTransitionUAVs<ECmdList::ECompute>)(TransitionType, TransitionPipeline, UAVArray, NumUAVs, WriteFence);
-	}
-
-	FORCEINLINE_DEBUGGABLE void TransitionResources(EResourceTransitionAccess TransitionType, EResourceTransitionPipeline TransitionPipeline, FRHIUnorderedAccessView** InUAVs, int32 NumUAVs)
-	{
-		TransitionResources(TransitionType, TransitionPipeline, InUAVs, NumUAVs, nullptr);
-	}
-	
-	FORCEINLINE_DEBUGGABLE void PushEvent(const TCHAR* Name, FColor Color)
-	{
-		if (Bypass())
-		{
-			GetComputeContext().RHIPushEvent(Name, Color);
-			return;
-		}
-		TCHAR* NameCopy  = AllocString(Name);
-		ALLOC_COMMAND(FRHICommandPushEvent<ECmdList::ECompute>)(NameCopy, Color);
-	}
-
-	FORCEINLINE_DEBUGGABLE void PopEvent()
-	{
-		if (Bypass())
-		{
-			GetComputeContext().RHIPopEvent();
-			return;
-		}
-		ALLOC_COMMAND(FRHICommandPopEvent<ECmdList::ECompute>)();
-	}
-
-	FORCEINLINE_DEBUGGABLE void BreakPoint()
-	{
-#if !UE_BUILD_SHIPPING
-		if (Bypass())
-		{
-			if (FPlatformMisc::IsDebuggerPresent())
-			{
-				UE_DEBUG_BREAK();
-			}
-			return;
-		}
-		ALLOC_COMMAND(FRHICommandDebugBreak)();
-#endif
-	}
-
-	FORCEINLINE_DEBUGGABLE void SubmitCommandsHint()
-	{
-		if (Bypass())
-		{
-			GetComputeContext().RHISubmitCommandsHint();
-			return;
-		}
-		ALLOC_COMMAND(FRHICommandSubmitCommandsHint<ECmdList::ECompute>)();
-	}
-
-	FORCEINLINE_DEBUGGABLE void WaitComputeFence(FRHIComputeFence* WaitFence)
-	{
-		if (Bypass())
-		{
-			GetComputeContext().RHIWaitComputeFence(WaitFence);
-			return;
-		}
-		ALLOC_COMMAND(FRHICommandWaitComputeFence<ECmdList::ECompute>)(WaitFence);
-	}
-
-	FORCEINLINE_DEBUGGABLE void CopyToStagingBuffer(FRHIVertexBuffer* SourceBuffer, FRHIStagingBuffer* DestinationStagingBuffer, uint32 Offset, uint32 NumBytes)
-	{
-		if (Bypass())
-		{
-			GetComputeContext().RHICopyToStagingBuffer(SourceBuffer, DestinationStagingBuffer, Offset, NumBytes);
-			return;
-		}
-		ALLOC_COMMAND(FRHICommandCopyToStagingBuffer<ECmdList::ECompute>)(SourceBuffer, DestinationStagingBuffer, Offset, NumBytes);
-	}
-
-	FORCEINLINE_DEBUGGABLE void WriteGPUFence(FRHIGPUFence* Fence)
-	{
-		if (Bypass())
-		{
-			GetComputeContext().RHIWriteGPUFence(Fence);
-			return;
-		}
-		ALLOC_COMMAND(FRHICommandWriteGPUFence<ECmdList::ECompute>)(Fence);
-	}
-
-	FORCEINLINE_DEBUGGABLE void SetGPUMask(FRHIGPUMask InGPUMask)
-	{
-		if (GPUMask != InGPUMask)
-		{
-			GPUMask = InGPUMask;
-			if (HasCommands() || Bypass())
-			{
-				SetGPUMaskOnComputeContext();
-			}
-			else
-			{
-				InitialGPUMask = GPUMask;
-			}
-		}
-	}
-
-#if RHI_RAYTRACING
-	FORCEINLINE_DEBUGGABLE void BuildAccelerationStructure(FRHIRayTracingGeometry* Geometry)
-	{
-		FAccelerationStructureBuildParams Params;
-		Params.Geometry = Geometry;
-		Params.BuildMode = EAccelerationStructureBuildMode::Build;
-		BuildAccelerationStructures(MakeArrayView(&Params, 1));
-	}
-
-	FORCEINLINE_DEBUGGABLE void BuildAccelerationStructures(const TArrayView<const FAccelerationStructureBuildParams> Params)
-	{
-		if (Bypass())
-		{
-			GetComputeContext().RHIBuildAccelerationStructures(Params);
-		}
-		else
-		{
-			ALLOC_COMMAND(FRHICommandBuildAccelerationStructures<ECmdList::ECompute>)(AllocArray(Params));
-		}
-	}
-
-	FORCEINLINE_DEBUGGABLE void BuildAccelerationStructure(FRHIRayTracingScene* Scene)
-	{
-		if (Bypass())
-		{
-			GetComputeContext().RHIBuildAccelerationStructure(Scene);
-		}
-		else
-		{
-			ALLOC_COMMAND(FRHICommandBuildAccelerationStructure<ECmdList::ECompute>)(Scene);
-		}
-	}
-#endif
-};
-
 namespace EImmediateFlushType
@@ -3778,7 +3687,7 @@
 	FGraphEventRef RHIThreadFence(bool bSetLockFence = false);
 
 	//Queue the given async compute commandlists in order with the current immediate commandlist
-	void QueueAsyncCompute(FRHIAsyncComputeCommandList& RHIComputeCmdList);
+	void QueueAsyncCompute(FRHIComputeCommandList& RHIComputeCmdList);
 
 	FORCEINLINE bool IsBottomOfPipe()
 	{
@@ -3958,9 +3867,9 @@
 		GDynamicRHI->RHIUnlockIndexBuffer(*this, IndexBuffer);
 	}
 	
-	FORCEINLINE void* LockStagingBuffer(FRHIStagingBuffer* StagingBuffer, uint32 Offset, uint32 SizeRHI)
+	FORCEINLINE void* LockStagingBuffer(FRHIStagingBuffer* StagingBuffer, FRHIGPUFence* Fence, uint32 Offset, uint32 SizeRHI)
 	{
-		return GDynamicRHI->LockStagingBuffer_RenderThread(*this, StagingBuffer, Offset, SizeRHI);
+		return GDynamicRHI->LockStagingBuffer_RenderThread(*this, StagingBuffer, Fence, Offset, SizeRHI);
 	}
 	
 	FORCEINLINE void UnlockStagingBuffer(FRHIStagingBuffer* StagingBuffer)
@@ -4055,6 +3964,12 @@
 		return GDynamicRHI->CreateShaderResourceView_RenderThread(*this, VertexBuffer, Stride, Format);
 	}
 	
+	FORCEINLINE FShaderResourceViewRHIRef CreateShaderResourceView(const FShaderResourceViewInitializer& Initializer)
+	{
+		LLM_SCOPE(ELLMTag::RHIMisc);
+		return GDynamicRHI->CreateShaderResourceView_RenderThread(*this, Initializer);
+	}
+
 	FORCEINLINE FShaderResourceViewRHIRef CreateShaderResourceView(FRHIIndexBuffer* Buffer)
 	{
 		LLM_SCOPE(ELLMTag::RHIMisc);
@@ -4446,11 +4361,11 @@
 		}
 	}
 
-	FORCEINLINE bool GetRenderQueryResult(FRHIRenderQuery* RenderQuery, uint64& OutResult, bool bWait)
+	FORCEINLINE bool GetRenderQueryResult(FRHIRenderQuery* RenderQuery, uint64& OutResult, bool bWait, uint32 GPUIndex = INDEX_NONE)
 	{
-		return RHIGetRenderQueryResult(RenderQuery, OutResult, bWait);
+		return RHIGetRenderQueryResult(RenderQuery, OutResult, bWait, GPUIndex);
 	}
-	
+
 	FORCEINLINE uint32 GetViewportNextPresentGPUIndex(FRHIViewport* Viewport)
 	{
 		return GDynamicRHI->RHIGetViewportNextPresentGPUIndex(Viewport);
@@ -4651,24 +4566,23 @@
 	}
 };
 
-template <typename T>
-struct TScopedGPUMask
+struct FScopedGPUMask
 {
-	T* RHICmdList;
+	FRHIComputeCommandList* RHICmdList;
 	FRHIGPUMask PrevGPUMask;
-	FORCEINLINE TScopedGPUMask(T& InRHICmdList, FRHIGPUMask InGPUMask)
+	FORCEINLINE FScopedGPUMask(FRHIComputeCommandList& InRHICmdList, FRHIGPUMask InGPUMask)
 		: RHICmdList(&InRHICmdList)
 		, PrevGPUMask(InRHICmdList.GetGPUMask())
 	{
 		InRHICmdList.SetGPUMask(InGPUMask);
 	}
-	FORCEINLINE TScopedGPUMask(TScopedGPUMask&& Other)
+	FORCEINLINE FScopedGPUMask(FScopedGPUMask&& Other)
 		: RHICmdList(Other.RHICmdList)
 		, PrevGPUMask(Other.PrevGPUMask)
 	{
 		Other.RHICmdList = nullptr;
 	}
-	FORCEINLINE ~TScopedGPUMask()
+	FORCEINLINE ~FScopedGPUMask()
 	{
 		if (RHICmdList != nullptr)
 		{
@@ -4677,27 +4591,61 @@
 	}
 };
 
-FORCEINLINE TScopedGPUMask<FRHICommandList> MakeScopedGPUMask(FRHICommandList& RHICmdList, FRHIGPUMask GPUMask) { return TScopedGPUMask<FRHICommandList>(RHICmdList, GPUMask); }
-FORCEINLINE TScopedGPUMask<FRHIAsyncComputeCommandList> MakeScopedGPUMask(FRHIAsyncComputeCommandList& RHICmdList, FRHIGPUMask GPUMask) { return TScopedGPUMask<FRHIAsyncComputeCommandList>(RHICmdList, GPUMask); }
-
 #if WITH_MGPU
-	#define SCOPED_GPU_MASK(RHICmdList, GPUMask) auto PREPROCESSOR_JOIN(ScopedGPUMask, __LINE__){ MakeScopedGPUMask(RHICmdList, GPUMask) }
+	#define SCOPED_GPU_MASK(RHICmdList, GPUMask) FScopedGPUMask PREPROCESSOR_JOIN(ScopedGPUMask, __LINE__){ RHICmdList, GPUMask }
 #else
 	#define SCOPED_GPU_MASK(RHICmdList, GPUMask)
 #endif // WITH_MGPU
 
+struct RHI_API FScopedUniformBufferGlobalBindings
+{
+	FScopedUniformBufferGlobalBindings(FRHIComputeCommandList& InRHICmdList, FUniformBufferStaticBindings UniformBuffers)
+		: RHICmdList(InRHICmdList)
+	{
+#if VALIDATE_UNIFORM_BUFFER_GLOBAL_BINDINGS
+		checkf(!bRecursionGuard, TEXT("Uniform buffer global binding scope has been called recursively!"));
+		bRecursionGuard = true;
+#endif
+
+		RHICmdList.SetGlobalUniformBuffers(UniformBuffers);
+	}
+
+	template <typename... TArgs>
+	FScopedUniformBufferGlobalBindings(FRHIComputeCommandList& InRHICmdList, TArgs... Args)
+		: FScopedUniformBufferGlobalBindings(InRHICmdList, FUniformBufferStaticBindings{ Args... })
+	{}
+
+	~FScopedUniformBufferGlobalBindings()
+	{
+		RHICmdList.SetGlobalUniformBuffers(FUniformBufferStaticBindings());
+
+#if VALIDATE_UNIFORM_BUFFER_GLOBAL_BINDINGS
+		bRecursionGuard = false;
+#endif
+	}
+
+	FRHIComputeCommandList& RHICmdList;
+
+#if VALIDATE_UNIFORM_BUFFER_GLOBAL_BINDINGS
+	static bool bRecursionGuard;
+#endif
+};
+
+#define SCOPED_UNIFORM_BUFFER_GLOBAL_BINDINGS(RHICmdList, UniformBuffers) FScopedUniformBufferGlobalBindings PREPROCESSOR_JOIN(UniformBuffers, __LINE__){ RHICmdList, UniformBuffers }
+
 // Single commandlist for async compute generation.  In the future we may expand this to allow async compute command generation
 // on multiple threads at once.
-class RHI_API FRHIAsyncComputeCommandListImmediate : public FRHIAsyncComputeCommandList
+class RHI_API FRHIAsyncComputeCommandListImmediate : public FRHIComputeCommandList
 {
 public:
+	FRHIAsyncComputeCommandListImmediate()
+		: FRHIComputeCommandList(FRHIGPUMask::All())
+	{}
 
 	//If RHIThread is enabled this will dispatch all current commands to the RHI Thread.  If RHI thread is disabled
 	//this will immediately execute the current commands.
 	//This also queues a GPU Submission command as the final command in the dispatch.
 	static void ImmediateDispatch(FRHIAsyncComputeCommandListImmediate& RHIComputeCmdList);
-
-private:
 };
 
 // typedef to mark the recursive use of commandlists in the RHI implementations
@@ -4710,7 +4658,6 @@
 	{
 		SetContext(Context);
 		bAsyncPSOCompileAllowed = false;
-		bRecursive = true;
 	}
 };
 
@@ -4741,8 +4688,8 @@
 	};
 
 public:
-	TRHICommandList_RecursiveHazardous(ContextType *Context)
-		: FRHICommandList_RecursiveHazardous(Context)
+	TRHICommandList_RecursiveHazardous(ContextType *Context, FRHIGPUMask GPUMask = FRHIGPUMask::All())
+		: FRHICommandList_RecursiveHazardous(Context, GPUMask)
 	{}
 
 	template <typename LAMBDA>
@@ -4761,8 +4708,8 @@
 #else
 
 public:
-	TRHICommandList_RecursiveHazardous(ContextType *Context)
-		: FRHICommandList_RecursiveHazardous(Context)
+	TRHICommandList_RecursiveHazardous(ContextType *Context, FRHIGPUMask GPUMask = FRHIGPUMask::All())
+		: FRHICommandList_RecursiveHazardous(Context, GPUMask)
 	{}
 
 	template <typename LAMBDA>
@@ -5050,6 +4997,11 @@
 	return FRHICommandListExecutor::GetImmediateCommandList().CreateShaderResourceView(VertexBuffer, Stride, Format);
 }
 
+FORCEINLINE FShaderResourceViewRHIRef RHICreateShaderResourceView(const FShaderResourceViewInitializer& Initializer)
+{
+	return FRHICommandListExecutor::GetImmediateCommandList().CreateShaderResourceView(Initializer);
+}
+
 FORCEINLINE FShaderResourceViewRHIRef RHICreateShaderResourceView(FRHIIndexBuffer* Buffer)
 {
 	return FRHICommandListExecutor::GetImmediateCommandList().CreateShaderResourceView(Buffer);
@@ -5291,7 +5243,12 @@
 
 FORCEINLINE void* RHILockStagingBuffer(FRHIStagingBuffer* StagingBuffer, uint32 Offset, uint32 Size)
 {
-	return FRHICommandListExecutor::GetImmediateCommandList().LockStagingBuffer(StagingBuffer, Offset, Size);
+	return FRHICommandListExecutor::GetImmediateCommandList().LockStagingBuffer(StagingBuffer, nullptr, Offset, Size);
 }
 
+FORCEINLINE void* RHILockStagingBuffer(FRHIStagingBuffer* StagingBuffer, FRHIGPUFence* Fence, uint32 Offset, uint32 Size)
+{
+	return FRHICommandListExecutor::GetImmediateCommandList().LockStagingBuffer(StagingBuffer, Fence, Offset, Size);
+}
+
 FORCEINLINE void RHIUnlockStagingBuffer(FRHIStagingBuffer* StagingBuffer)
diff --git a/Engine/Source/Runtime/RenderCore/Private/ShaderCore.cpp b/Engine/Source/Runtime/RenderCore/Private/ShaderCore.cpp
--- a/Engine/Source/Runtime/RenderCore/Private/ShaderCore.cpp
+++ b/Engine/Source/Runtime/RenderCore/Private/ShaderCore.cpp
@@ -25,6 +25,8 @@
 #include "Misc/StringBuilder.h"
 #endif
 
+#define VALIDATE_GLOBAL_UNIFORM_BUFFERS (!UE_BUILD_SHIPPING && !UE_BUILD_TEST)
+
 static TAutoConsoleVariable<int32> CVarShaderDevelopmentMode(
 	TEXT("r.ShaderDevelopmentMode"),
 	0,
@@ -355,6 +357,37 @@ bool ShouldExportShaderDebugInfo(EShaderPlatform Platform)
 	return (GExportDebugInfoPlatforms & (uint64(1) << Platform)) != 0;
 }
 
+void ValidateStaticUniformBuffer(FRHIUniformBuffer* UniformBuffer, FUniformBufferStaticSlot Slot, uint32 ExpectedHash)
+{
+#if VALIDATE_GLOBAL_UNIFORM_BUFFERS
+	FUniformBufferStaticSlotRegistry& SlotRegistry = FUniformBufferStaticSlotRegistry::Get();
+
+	if (!UniformBuffer)
+	{
+		UE_LOG(LogShaders, Warning, TEXT("Shader requested a global uniform buffer at static slot %s, but it was null."), *SlotRegistry.GetDebugDescription(Slot));
+	}
+	else
+	{
+		const FRHIUniformBufferLayout& Layout = UniformBuffer->GetLayout();
+
+		if (Layout.GetHash() != ExpectedHash)
+		{
+			const FShaderParametersMetadata* ExpectedStructMetadata = FindUniformBufferStructByLayoutHash(ExpectedHash);
+
+			checkf(
+				ExpectedStructMetadata,
+				TEXT("Shader is requesting uniform buffer '%s' at slot %s with hash '%u', but a reverse lookup of the hash can't find it. The shader cache may be out of date."),
+				*Layout.GetDebugName().ToString(), *SlotRegistry.GetDebugDescription(Slot), ExpectedHash);
+
+			checkf(
+				false,
+				TEXT("Shader attempted to bind uniform buffer '%s' at slot %s with hash '%u', but the shader expected '%s' with hash '%u'."),
+				*Layout.GetDebugName().ToString(), *SlotRegistry.GetDebugDescription(Slot), ExpectedHash, ExpectedStructMetadata->GetShaderVariableName(), Layout.GetHash());
+		}
+	}
+#endif
+}
+
 bool FShaderParameterMap::FindParameterAllocation(const TCHAR* ParameterName,uint16& OutBufferIndex,uint16& OutBaseIndex,uint16& OutSize) const
 {
 	const FParameterAllocation* Allocation = ParameterMap.Find(ParameterName);
diff --git a/Engine/Source/Runtime/TraceLog/Private/Trace/Writer.cpp b/Engine/Source/Runtime/TraceLog/Private/Trace/Writer.cpp
--- a/Engine/Source/Runtime/TraceLog/Private/Trace/Writer.cpp
+++ b/Engine/Source/Runtime/TraceLog/Private/Trace/Writer.cpp
@@ -249,7 +249,7 @@ static FWriteBuffer* Writer_NextBufferInternal(uint32 PageGrowth)
 }
 
 ////////////////////////////////////////////////////////////////////////////////
-TRACELOG_API FWriteBuffer* Writer_NextBuffer(uint16 Size)
+TRACELOG_API FWriteBuffer* Writer_NextBuffer(int32 Size)
 {
 	if (Size >= GPoolBlockSize - sizeof(FWriteBuffer))
 	{
@@ -1077,7 +1077,7 @@ void Writer_EventCreate(
 	Target->bImportant = Flags & FEventDef::Flag_Important;
 
 	// Calculate the number of fields and size of name data.
-	int NamesSize = LoggerName.Length + EventName.Length;
+	uint16 NamesSize = LoggerName.Length + EventName.Length;
 	for (uint32 i = 0; i < FieldCount; ++i)
 	{
 		NamesSize += FieldDescs[i].NameSize;
