Copying Private-LoadTimes-4.24 stream to Main. Biggest changes are in Materials/Shader memory freezing.

diff --git a/Engine/Source/Runtime/RHI/Public/RHIDefinitions.h b/Engine/Source/Runtime/RHI/Public/RHIDefinitions.h
--- a/Engine/Source/Runtime/RHI/Public/RHIDefinitions.h
+++ b/Engine/Source/Runtime/RHI/Public/RHIDefinitions.h
@@ -10,6 +10,11 @@
 #include "CoreMinimal.h"
 #include "PixelFormat.h"
 #include "HAL/IConsoleManager.h"
+#include "Serialization/MemoryLayout.h"
+
+#ifndef USE_STATIC_SHADER_PLATFORM_ENUMS
+#define USE_STATIC_SHADER_PLATFORM_ENUMS 0
+#endif
 
 #ifndef RHI_RAYTRACING
 #if (PLATFORM_WINDOWS && PLATFORM_64BITS)
@@ -19,7 +24,7 @@
 #endif
 #endif
 
-enum EShaderFrequency
+enum EShaderFrequency : uint8
 {
 	SF_Vertex			= 0,
 	SF_Hull				= 1,
@@ -103,6 +108,32 @@ enum EShaderPlatform
 };
 static_assert(SP_NumPlatforms <= (1 << SP_NumBits), "SP_NumPlatforms will not fit on SP_NumBits");
 
+struct FGenericStaticShaderPlatform final
+{
+	inline FGenericStaticShaderPlatform(const EShaderPlatform InPlatform) : Platform(InPlatform) {}
+	inline operator EShaderPlatform() const
+	{
+		return Platform;
+	}
+
+	inline bool operator == (const EShaderPlatform Other) const
+	{
+		return Other == Platform;
+	}
+	inline bool operator != (const EShaderPlatform Other) const
+	{
+		return Other != Platform;
+	}
+private:
+	const EShaderPlatform Platform;
+};
+
+#if USE_STATIC_SHADER_PLATFORM_ENUMS
+#include COMPILED_PLATFORM_HEADER(StaticShaderPlatform.inl)
+#else
+using FStaticShaderPlatform = FGenericStaticShaderPlatform;
+#endif
+
 class FStaticShaderPlatformNames
 {
 private:
@@ -212,66 +243,300 @@ namespace ERHIFeatureLevel
 		Num
 	};
 };
+DECLARE_INTRINSIC_TYPE_LAYOUT(ERHIFeatureLevel::Type);
+struct FGenericStaticFeatureLevel
+{
+	inline FGenericStaticFeatureLevel(const ERHIFeatureLevel::Type InFeatureLevel) : FeatureLevel(InFeatureLevel) {}
+	inline FGenericStaticFeatureLevel(const TEnumAsByte<ERHIFeatureLevel::Type> InFeatureLevel) : FeatureLevel(InFeatureLevel) {}
+
+	inline operator ERHIFeatureLevel::Type() const
+	{
+		return FeatureLevel;
+	}
+
+	inline bool operator == (const ERHIFeatureLevel::Type Other) const
+	{
+		return Other == FeatureLevel;
+	}
+
+	inline bool operator != (const ERHIFeatureLevel::Type Other) const
+	{
+		return Other != FeatureLevel;
+	}
+
+	inline bool operator <= (const ERHIFeatureLevel::Type Other) const
+	{
+		return FeatureLevel <= Other;
+	}
+
+	inline bool operator < (const ERHIFeatureLevel::Type Other) const
+	{
+		return FeatureLevel < Other;
+	}
+
+	inline bool operator >= (const ERHIFeatureLevel::Type Other) const
+	{
+		return FeatureLevel >= Other;
+	}
+
+	inline bool operator > (const ERHIFeatureLevel::Type Other) const
+	{
+		return FeatureLevel > Other;
+	}
+
+private:
+	ERHIFeatureLevel::Type FeatureLevel;
+};
 
+#if USE_STATIC_SHADER_PLATFORM_ENUMS
+#include COMPILED_PLATFORM_HEADER(StaticFeatureLevel.inl)
+#else
+using FStaticFeatureLevel = FGenericStaticFeatureLevel;
+#endif
 
-extern RHI_API FName LANGUAGE_D3D;
-extern RHI_API FName LANGUAGE_Metal;
-extern RHI_API FName LANGUAGE_OpenGL;
-extern RHI_API FName LANGUAGE_Vulkan;
-extern RHI_API FName LANGUAGE_Sony;
-extern RHI_API FName LANGUAGE_Nintendo;
+extern RHI_API const FName LANGUAGE_D3D;
+extern RHI_API const FName LANGUAGE_Metal;
+extern RHI_API const FName LANGUAGE_OpenGL;
+extern RHI_API const FName LANGUAGE_Vulkan;
+extern RHI_API const FName LANGUAGE_Sony;
+extern RHI_API const FName LANGUAGE_Nintendo;
 
-struct RHI_API FDataDrivenShaderPlatformInfo
+class RHI_API FGenericDataDrivenShaderPlatformInfo
 {
 	FName Language;
 	ERHIFeatureLevel::Type MaxFeatureLevel;
-	bool bIsMobile;
-	bool bIsMetalMRT;
-	bool bIsPC;
-	bool bIsConsole;
-	bool bIsAndroidOpenGLES;
-
-	bool bSupportsDrawIndirect;
-	bool bSupportsMobileMultiView;
-	bool bSupportsVolumeTextureCompression;
-	bool bSupportsDistanceFields; // used for DFShadows and DFAO - since they had the same checks
-	bool bSupportsDiaphragmDOF;
-	bool bSupportsRGBColorBuffer;
-	bool bSupportsCapsuleShadows;
-	bool bSupportsVolumetricFog; // also used for FVVoxelization
-	bool bSupportsIndexBufferUAVs;
-	bool bSupportsInstancedStereo;
-	bool bSupportsMultiView;
-	bool bSupportsMSAA;
-	bool bSupports4ComponentUAVReadWrite;
-	bool bSupportsRenderTargetWriteMask;
-	bool bSupportsRayTracing;
-	bool bSupportsRayTracingMissShaderBindings; // Whether resources could be bound for miss shaders
-	bool bSupportsGPUSkinCache;
-	bool bSupportsByteBufferComputeShaders;
-
-	bool bTargetsTiledGPU;
-	bool bNeedsOfflineCompiler;
+	uint32 bIsMobile: 1;
+	uint32 bIsMetalMRT: 1;
+	uint32 bIsPC: 1;
+	uint32 bIsConsole: 1;
+	uint32 bIsAndroidOpenGLES: 1;
+
+	uint32 bSupportsDrawIndirect: 1;
+	uint32 bSupportsMobileMultiView: 1;
+	uint32 bSupportsVolumeTextureCompression: 1;
+	uint32 bSupportsDistanceFields: 1; // used for DFShadows and DFAO - since they had the same checks
+	uint32 bSupportsDiaphragmDOF: 1;
+	uint32 bSupportsRGBColorBuffer: 1;
+	uint32 bSupportsCapsuleShadows: 1;
+	uint32 bSupportsVolumetricFog: 1; // also used for FVVoxelization
+	uint32 bSupportsIndexBufferUAVs: 1;
+	uint32 bSupportsInstancedStereo: 1;
+	uint32 bSupportsMultiView: 1;
+	uint32 bSupportsMSAA: 1;
+	uint32 bSupports4ComponentUAVReadWrite: 1;
+	uint32 bSupportsRenderTargetWriteMask: 1;
+	uint32 bSupportsRayTracing: 1;
+	uint32 bSupportsRayTracingMissShaderBindings : 1; // Whether resources could be bound for miss shaders
+	uint32 bSupportsGPUSkinCache: 1;
+	uint32 bSupportsGPUScene : 1;
+	uint32 bSupportsByteBufferComputeShaders : 1;
+
+	uint32 bTargetsTiledGPU: 1;
+	uint32 bNeedsOfflineCompiler: 1;
 
 	// NOTE: WHen adding fields, you must also add to ParseDataDrivenShaderInfo!
+	uint32 bContainsValidPlatformInfo : 1;
 
-	FDataDrivenShaderPlatformInfo()
+	FGenericDataDrivenShaderPlatformInfo()
 	{
 		FMemory::Memzero(this, sizeof(*this));
 		MaxFeatureLevel = ERHIFeatureLevel::Num;
 	}
 
+public:
 	static void Initialize();
+	static void ParseDataDrivenShaderInfo(const FConfigSection& Section, FGenericDataDrivenShaderPlatformInfo& Info);
+
+	static FORCEINLINE_DEBUGGABLE const bool GetIsLanguageD3D(const FStaticShaderPlatform Platform)
+	{
+		return Infos[Platform].Language == LANGUAGE_D3D;
+	}
+
+	static FORCEINLINE_DEBUGGABLE const bool GetIsLanguageMetal(const FStaticShaderPlatform Platform)
+	{
+		return Infos[Platform].Language == LANGUAGE_Metal;
+	}
+
+	static FORCEINLINE_DEBUGGABLE const bool GetIsLanguageOpenGL(const FStaticShaderPlatform Platform)
+	{
+		return Infos[Platform].Language == LANGUAGE_OpenGL;
+	}
+
+	static FORCEINLINE_DEBUGGABLE const bool GetIsLanguageVulkan(const FStaticShaderPlatform Platform)
+	{
+		return Infos[Platform].Language == LANGUAGE_Vulkan;
+	}
+
+	static FORCEINLINE_DEBUGGABLE const bool GetIsLanguageSony(const FStaticShaderPlatform Platform)
+	{
+		return Infos[Platform].Language == LANGUAGE_Sony;
+	}
+
+	static FORCEINLINE_DEBUGGABLE const bool GetIsLanguageNintendo(const FStaticShaderPlatform Platform)
+	{
+		return Infos[Platform].Language == LANGUAGE_Nintendo;
+	}
+
+	static FORCEINLINE_DEBUGGABLE const ERHIFeatureLevel::Type GetMaxFeatureLevel(const FStaticShaderPlatform Platform)
+	{
+		return Infos[Platform].MaxFeatureLevel;
+	}
+
+	static FORCEINLINE_DEBUGGABLE const bool GetIsMobile(const FStaticShaderPlatform Platform)
+	{
+		return Infos[Platform].bIsMobile;
+	}
+
+	static FORCEINLINE_DEBUGGABLE const bool GetIsMetalMRT(const FStaticShaderPlatform Platform)
+	{
+		return Infos[Platform].bIsMetalMRT;
+	}
+
+	static FORCEINLINE_DEBUGGABLE const bool GetIsPC(const FStaticShaderPlatform Platform)
+	{
+		return Infos[Platform].bIsPC;
+	}
+
+	static FORCEINLINE_DEBUGGABLE const bool GetIsConsole(const FStaticShaderPlatform Platform)
+	{
+		return Infos[Platform].bIsConsole;
+	}
+
+	static FORCEINLINE_DEBUGGABLE const bool GetIsAndroidOpenGLES(const FStaticShaderPlatform Platform)
+	{
+		return Infos[Platform].bIsAndroidOpenGLES;
+	}
+
+	static FORCEINLINE_DEBUGGABLE const bool GetSupportsDrawIndirect(const FStaticShaderPlatform Platform)
+	{
+		return Infos[Platform].bSupportsDrawIndirect;
+	}
+
+	static FORCEINLINE_DEBUGGABLE const bool GetSupportsMobileMultiView(const FStaticShaderPlatform Platform)
+	{
+		return Infos[Platform].bSupportsMobileMultiView;
+	}
+
+	static FORCEINLINE_DEBUGGABLE const bool GetSupportsVolumeTextureCompression(const FStaticShaderPlatform Platform)
+	{
+		return Infos[Platform].bSupportsVolumeTextureCompression;
+	}
+
+	static FORCEINLINE_DEBUGGABLE const bool GetSupportsDistanceFields(const FStaticShaderPlatform Platform)
+	{
+		return Infos[Platform].bSupportsDistanceFields;
+	}
+
+	static FORCEINLINE_DEBUGGABLE const bool GetSupportsDiaphragmDOF(const FStaticShaderPlatform Platform)
+	{
+		return Infos[Platform].bSupportsDiaphragmDOF;
+	}
+
+	static FORCEINLINE_DEBUGGABLE const bool GetSupportsRGBColorBuffer(const FStaticShaderPlatform Platform)
+	{
+		return Infos[Platform].bSupportsRGBColorBuffer;
+	}
+
+	static FORCEINLINE_DEBUGGABLE const bool GetSupportsCapsuleShadows(const FStaticShaderPlatform Platform)
+	{
+		return Infos[Platform].bSupportsCapsuleShadows;
+	}
+
+	static FORCEINLINE_DEBUGGABLE const bool GetSupportsVolumetricFog(const FStaticShaderPlatform Platform)
+	{
+		return Infos[Platform].bSupportsVolumetricFog;
+	}
+
+	static FORCEINLINE_DEBUGGABLE const bool GetSupportsIndexBufferUAVs(const FStaticShaderPlatform Platform)
+	{
+		return Infos[Platform].bSupportsIndexBufferUAVs;
+	}
+
+	static FORCEINLINE_DEBUGGABLE const bool GetSupportsInstancedStereo(const FStaticShaderPlatform Platform)
+	{
+		return Infos[Platform].bSupportsInstancedStereo;
+	}
+
+	static FORCEINLINE_DEBUGGABLE const bool GetSupportsMultiView(const FStaticShaderPlatform Platform)
+	{
+		return Infos[Platform].bSupportsMultiView;
+	}
+
+	static FORCEINLINE_DEBUGGABLE const bool GetSupportsMSAA(const FStaticShaderPlatform Platform)
+	{
+		return Infos[Platform].bSupportsMSAA;
+	}
+
+	static FORCEINLINE_DEBUGGABLE const bool GetSupports4ComponentUAVReadWrite(const FStaticShaderPlatform Platform)
+	{
+		return Infos[Platform].bSupports4ComponentUAVReadWrite;
+	}
+
+	static FORCEINLINE_DEBUGGABLE const bool GetSupportsRenderTargetWriteMask(const FStaticShaderPlatform Platform)
+	{
+		return Infos[Platform].bSupportsRenderTargetWriteMask;
+	}
+
+	static FORCEINLINE_DEBUGGABLE const bool GetSupportsRayTracing(const FStaticShaderPlatform Platform)
+	{
+		return Infos[Platform].bSupportsRayTracing;
+	}
+
+	static FORCEINLINE_DEBUGGABLE const bool GetSupportsRayTracingMissShaderBindings(const FStaticShaderPlatform Platform)
+	{
+		return Infos[Platform].bSupportsRayTracingMissShaderBindings;
+	}
+
+	static FORCEINLINE_DEBUGGABLE const bool GetSupportsGPUSkinCache(const FStaticShaderPlatform Platform)
+	{
+		return Infos[Platform].bSupportsGPUSkinCache;
+	}
+
+	static FORCEINLINE_DEBUGGABLE const bool GetTargetsTiledGPU(const FStaticShaderPlatform Platform)
+	{
+		return Infos[Platform].bTargetsTiledGPU;
+	}
+
+	static FORCEINLINE_DEBUGGABLE const bool GetNeedsOfflineCompiler(const FStaticShaderPlatform Platform)
+	{
+		return Infos[Platform].bNeedsOfflineCompiler;
+	}
 
-	static const FDataDrivenShaderPlatformInfo& GetInfo(EShaderPlatform Platform)
+	static FORCEINLINE_DEBUGGABLE const bool GetSupportsGPUScene(const FStaticShaderPlatform Platform)
 	{
-		return Infos[Platform];
+		return Infos[Platform].bSupportsGPUScene;
 	}
+
+	static FORCEINLINE_DEBUGGABLE const bool GetSupportsByteBufferComputeShaders(const FStaticShaderPlatform Platform)
+	{
+		return Infos[Platform].bSupportsByteBufferComputeShaders;
+	}
+
 private:
-	static FDataDrivenShaderPlatformInfo Infos[SP_NumPlatforms];
+	static FGenericDataDrivenShaderPlatformInfo Infos[SP_NumPlatforms];
+
+public:
+	static bool IsValid(const FStaticShaderPlatform Platform)
+	{
+		return Infos[Platform].bContainsValidPlatformInfo;
+	}
 };
 
+#if USE_STATIC_SHADER_PLATFORM_ENUMS
+
+#define IMPLEMENT_DDPSPI_SETTING_WITH_RETURN_TYPE(ReturnType, Function, Value) \
+	static FORCEINLINE_DEBUGGABLE const ReturnType Function(const FStaticShaderPlatform Platform) \
+	{ \
+		checkSlow(!FGenericDataDrivenShaderPlatformInfo::IsValid(Platform) || FGenericDataDrivenShaderPlatformInfo::Function(Platform) == Value); \
+		return Value; \
+	}
+#define IMPLEMENT_DDPSPI_SETTING(Function, Value) IMPLEMENT_DDPSPI_SETTING_WITH_RETURN_TYPE(bool, Function, Value)
 
+#include COMPILED_PLATFORM_HEADER(DataDrivenShaderPlatformInfo.inl)
+
+#else
+using FDataDrivenShaderPlatformInfo = FGenericDataDrivenShaderPlatformInfo;
+#endif
 
 enum ERenderQueryType
 {
@@ -532,6 +797,7 @@ enum EVertexElementType
 	VET_NumBits = 5,
 };
 static_assert(VET_MAX <= (1 << VET_NumBits), "VET_MAX will not fit on VET_NumBits");
+DECLARE_INTRINSIC_TYPE_LAYOUT(EVertexElementType);
 
 enum ECubeFace
 {
@@ -606,6 +872,7 @@ enum EUniformBufferBaseType : uint8
 	EUniformBufferBaseType_NumBits = 5,
 };
 static_assert(EUniformBufferBaseType_Num <= (1 << EUniformBufferBaseType_NumBits), "EUniformBufferBaseType_Num will not fit on EUniformBufferBaseType_NumBits");
+DECLARE_INTRINSIC_TYPE_LAYOUT(EUniformBufferBaseType);
 
 /** Numerical type used to store the static slot indices. */
 using FUniformBufferStaticSlot = uint8;
@@ -1047,25 +1314,25 @@ enum class EAsyncComputeBudget
 	EAll_4,				//Async compute can use the entire GPU.
 };
 
-inline bool IsPCPlatform(const EShaderPlatform Platform)
+inline bool IsPCPlatform(const FStaticShaderPlatform Platform)
 {
 	return Platform == SP_PCD3D_SM5 || Platform == SP_PCD3D_ES3_1 ||
 		Platform == SP_OPENGL_SM5 || Platform == SP_OPENGL_PCES3_1 ||
 		Platform == SP_METAL_SM5_NOTESS || Platform == SP_METAL_SM5 ||
 		Platform == SP_VULKAN_PCES3_1 || Platform == SP_VULKAN_SM5 || Platform == SP_METAL_MACES3_1 || Platform == SP_METAL_MRT_MAC 
-		|| FDataDrivenShaderPlatformInfo::GetInfo(Platform).bIsPC;
+		|| FDataDrivenShaderPlatformInfo::GetIsPC(Platform);
 }
 
 /** Whether the shader platform corresponds to the ES2 feature level. */
 UE_DEPRECATED(4.23, "Feature level ES2 is getting deprecated. Please use ES3.1.")
-inline bool IsES2Platform(const EShaderPlatform Platform)
+inline bool IsES2Platform(const FStaticShaderPlatform Platform)
 {
 	return Platform == SP_OPENGL_ES2_ANDROID || Platform == SP_OPENGL_ES2_WEBGL
-		|| FDataDrivenShaderPlatformInfo::GetInfo(Platform).MaxFeatureLevel == ERHIFeatureLevel::ES2;
+		|| FDataDrivenShaderPlatformInfo::GetMaxFeatureLevel(Platform) == ERHIFeatureLevel::ES2;
 }
 
 /** Whether the shader platform corresponds to the ES2/ES3.1 feature level. */
-inline bool IsMobilePlatform(const EShaderPlatform Platform)
+inline bool IsMobilePlatform(const FStaticShaderPlatform Platform)
 {
 	return 
 		PRAGMA_DISABLE_DEPRECATION_WARNINGS
@@ -1076,84 +1343,84 @@ inline bool IsMobilePlatform(const EShaderPlatform Platform)
 		|| Platform == SP_OPENGL_PCES3_1 || Platform == SP_OPENGL_ES3_1_ANDROID
 		|| Platform == SP_VULKAN_ES3_1_ANDROID || Platform == SP_VULKAN_PCES3_1 || Platform == SP_VULKAN_ES3_1_LUMIN
 		|| Platform == SP_SWITCH_FORWARD
-		|| FDataDrivenShaderPlatformInfo::GetInfo(Platform).bIsMobile;
+		|| FDataDrivenShaderPlatformInfo::GetIsMobile(Platform);
 }
 
-inline bool IsOpenGLPlatform(const EShaderPlatform Platform)
+inline bool IsOpenGLPlatform(const FStaticShaderPlatform Platform)
 {
 	return Platform == SP_OPENGL_SM5 || Platform == SP_OPENGL_PCES3_1
 		|| Platform == SP_OPENGL_ES2_ANDROID || Platform == SP_OPENGL_ES2_WEBGL || Platform == SP_OPENGL_ES31_EXT
 		|| Platform == SP_OPENGL_ES3_1_ANDROID
-		|| FDataDrivenShaderPlatformInfo::GetInfo(Platform).Language == LANGUAGE_OpenGL;
+		|| FDataDrivenShaderPlatformInfo::GetIsLanguageOpenGL(Platform);
 }
 
-inline bool IsMetalPlatform(const EShaderPlatform Platform)
+inline bool IsMetalPlatform(const FStaticShaderPlatform Platform)
 {
 	return Platform == SP_METAL || Platform == SP_METAL_MRT || Platform == SP_METAL_TVOS || Platform == SP_METAL_MRT_TVOS || Platform == SP_METAL_SM5_NOTESS || Platform == SP_METAL_SM5 || Platform == SP_METAL_MACES3_1 || Platform == SP_METAL_MRT_MAC
-		|| FDataDrivenShaderPlatformInfo::GetInfo(Platform).Language == LANGUAGE_Metal;
+		|| FDataDrivenShaderPlatformInfo::GetIsLanguageMetal(Platform);
 }
 
-inline bool IsMetalMobilePlatform(const EShaderPlatform Platform)
+inline bool IsMetalMobilePlatform(const FStaticShaderPlatform Platform)
 {
 	return Platform == SP_METAL || Platform == SP_METAL_TVOS
-		|| (FDataDrivenShaderPlatformInfo::GetInfo(Platform).Language == LANGUAGE_OpenGL && FDataDrivenShaderPlatformInfo::GetInfo(Platform).bIsMobile);
+		|| (FDataDrivenShaderPlatformInfo::GetIsLanguageOpenGL(Platform) && FDataDrivenShaderPlatformInfo::GetIsMobile(Platform));
 }
 
-inline bool IsMetalMRTPlatform(const EShaderPlatform Platform)
+inline bool IsMetalMRTPlatform(const FStaticShaderPlatform Platform)
 {
 	return Platform == SP_METAL_MRT || Platform == SP_METAL_MRT_TVOS || Platform == SP_METAL_MRT_MAC
-		|| FDataDrivenShaderPlatformInfo::GetInfo(Platform).bIsMetalMRT;
+		|| FDataDrivenShaderPlatformInfo::GetIsMetalMRT(Platform);
 }
 
-inline bool IsMetalSM5Platform(const EShaderPlatform Platform)
+inline bool IsMetalSM5Platform(const FStaticShaderPlatform Platform)
 {
 	return Platform == SP_METAL_MRT || Platform == SP_METAL_MRT_TVOS || Platform == SP_METAL_SM5_NOTESS || Platform == SP_METAL_SM5 || Platform == SP_METAL_MRT_MAC
-		|| (FDataDrivenShaderPlatformInfo::GetInfo(Platform).Language == LANGUAGE_Metal && FDataDrivenShaderPlatformInfo::GetInfo(Platform).MaxFeatureLevel == ERHIFeatureLevel::SM5);
+		|| (FDataDrivenShaderPlatformInfo::GetIsLanguageMetal(Platform) && FDataDrivenShaderPlatformInfo::GetMaxFeatureLevel(Platform) == ERHIFeatureLevel::SM5);
 }
 
-inline bool IsConsolePlatform(const EShaderPlatform Platform)
+inline bool IsConsolePlatform(const FStaticShaderPlatform Platform)
 {
 	return Platform == SP_PS4 || Platform == SP_XBOXONE_D3D12
-		|| FDataDrivenShaderPlatformInfo::GetInfo(Platform).bIsConsole;
+		|| FDataDrivenShaderPlatformInfo::GetIsConsole(Platform);
 }
 
-inline bool IsSwitchPlatform(const EShaderPlatform Platform)
+inline bool IsSwitchPlatform(const FStaticShaderPlatform Platform)
 {
 	return Platform == SP_SWITCH || Platform == SP_SWITCH_FORWARD
-		|| FDataDrivenShaderPlatformInfo::GetInfo(Platform).Language == LANGUAGE_Nintendo;
+		|| FDataDrivenShaderPlatformInfo::GetIsLanguageNintendo(Platform);
 }
 
-inline bool IsPS4Platform(const EShaderPlatform Platform)
+inline bool IsPS4Platform(const FStaticShaderPlatform Platform)
 {
 	return Platform == SP_PS4
-		|| FDataDrivenShaderPlatformInfo::GetInfo(Platform).Language == LANGUAGE_Sony;
+		|| FDataDrivenShaderPlatformInfo::GetIsLanguageSony(Platform);
 }
 
-inline bool IsVulkanPlatform(const EShaderPlatform Platform)
+inline bool IsVulkanPlatform(const FStaticShaderPlatform Platform)
 {
 	return Platform == SP_VULKAN_SM5 || Platform == SP_VULKAN_SM5_LUMIN || Platform == SP_VULKAN_PCES3_1 || Platform == SP_VULKAN_ES3_1_ANDROID || Platform == SP_VULKAN_ES3_1_LUMIN
-		|| FDataDrivenShaderPlatformInfo::GetInfo(Platform).Language == LANGUAGE_Vulkan;
+		|| FDataDrivenShaderPlatformInfo::GetIsLanguageVulkan(Platform);
 }
 
-inline bool IsVulkanSM5Platform(const EShaderPlatform Platform)
+inline bool IsVulkanSM5Platform(const FStaticShaderPlatform Platform)
 {
 	return Platform == SP_VULKAN_SM5 || Platform == SP_VULKAN_SM5_LUMIN
-		|| (FDataDrivenShaderPlatformInfo::GetInfo(Platform).Language == LANGUAGE_Vulkan && FDataDrivenShaderPlatformInfo::GetInfo(Platform).MaxFeatureLevel == ERHIFeatureLevel::SM5);
+		|| (FDataDrivenShaderPlatformInfo::GetIsLanguageVulkan(Platform) && FDataDrivenShaderPlatformInfo::GetMaxFeatureLevel(Platform) == ERHIFeatureLevel::SM5);
 }
 
-inline bool IsAndroidOpenGLESPlatform(const EShaderPlatform Platform)
+inline bool IsAndroidOpenGLESPlatform(const FStaticShaderPlatform Platform)
 {
 	return Platform == SP_OPENGL_ES2_ANDROID || Platform == SP_OPENGL_ES3_1_ANDROID
-		|| FDataDrivenShaderPlatformInfo::GetInfo(Platform).bIsAndroidOpenGLES;
+		|| FDataDrivenShaderPlatformInfo::GetIsAndroidOpenGLES(Platform);
 }
 
-inline bool IsVulkanMobilePlatform(const EShaderPlatform Platform)
+inline bool IsVulkanMobilePlatform(const FStaticShaderPlatform Platform)
 {
 	return Platform == SP_VULKAN_PCES3_1 || Platform == SP_VULKAN_ES3_1_ANDROID || Platform == SP_VULKAN_ES3_1_LUMIN
-		|| (FDataDrivenShaderPlatformInfo::GetInfo(Platform).Language == LANGUAGE_Vulkan && FDataDrivenShaderPlatformInfo::GetInfo(Platform).bIsMobile);
+		|| (FDataDrivenShaderPlatformInfo::GetIsLanguageVulkan(Platform) && FDataDrivenShaderPlatformInfo::GetIsMobile(Platform));
 }
 
-inline bool IsD3DPlatform(const EShaderPlatform Platform, bool bIncludeXboxOne)
+inline bool IsD3DPlatform(const FStaticShaderPlatform Platform, bool bIncludeXboxOne)
 {
 	switch (Platform)
 	{
@@ -1163,18 +1430,18 @@ inline bool IsD3DPlatform(const EShaderPlatform Platform, bool bIncludeXboxOne)
 	case SP_XBOXONE_D3D12:
 		return bIncludeXboxOne;
 	default:
-		return FDataDrivenShaderPlatformInfo::GetInfo(Platform).Language == LANGUAGE_D3D;
+		return FDataDrivenShaderPlatformInfo::GetIsLanguageD3D(Platform);
 	}
 
 	return false;
 }
 
-inline bool IsHlslccShaderPlatform(const EShaderPlatform Platform)
+inline bool IsHlslccShaderPlatform(const FStaticShaderPlatform Platform)
 {
 	return IsMetalPlatform(Platform) || IsVulkanPlatform(Platform) || IsSwitchPlatform(Platform) || IsOpenGLPlatform(Platform);
 }
 
-inline bool IsDeprecatedShaderPlatform(const EShaderPlatform ShaderPlatform)
+inline bool IsDeprecatedShaderPlatform(const FStaticShaderPlatform ShaderPlatform)
 {
 	return ShaderPlatform == SP_OPENGL_SM5 || ShaderPlatform == SP_PCD3D_SM4_REMOVED || ShaderPlatform == SP_OPENGL_ES2_IOS_REMOVED ||
 		ShaderPlatform == SP_PCD3D_ES2_DEPRECATED || ShaderPlatform == SP_METAL_MACES2_DEPRECATED || ShaderPlatform == SP_OPENGL_PCES2_DEPRECATED ||
@@ -1215,12 +1482,12 @@ inline ERHIFeatureLevel::Type GetMaxSupportedFeatureLevel(EShaderPlatform InShad
 	case SP_SWITCH_FORWARD:
 		return ERHIFeatureLevel::ES3_1;
 	default:
-		return FDataDrivenShaderPlatformInfo::GetInfo(InShaderPlatform).MaxFeatureLevel;
+		return FDataDrivenShaderPlatformInfo::GetMaxFeatureLevel(InShaderPlatform);
 	}
 }
 
 /* Returns true if the shader platform Platform is used to simulate a mobile feature level on a PC platform. */
-inline bool IsSimulatedPlatform(EShaderPlatform Platform)
+inline bool IsSimulatedPlatform(const FStaticShaderPlatform Platform)
 {
 	switch (Platform)
 	{
@@ -1257,12 +1524,19 @@ inline EShaderPlatform GetSimulatedPlatform(EShaderPlatform Platform)
 }
 
 /** Returns true if the feature level is supported by the shader platform. */
-inline bool IsFeatureLevelSupported(EShaderPlatform InShaderPlatform, ERHIFeatureLevel::Type InFeatureLevel)
+inline bool IsFeatureLevelSupported(FStaticShaderPlatform InShaderPlatform, const FStaticFeatureLevel InFeatureLevel)
 {
 	return InFeatureLevel <= GetMaxSupportedFeatureLevel(InShaderPlatform);
 }
 
-inline bool RHINeedsToSwitchVerticalAxis(EShaderPlatform Platform)
+/** Returns true if the feature level is supported by the shader platform. */
+inline bool IsFeatureLevelSupported(FStaticShaderPlatform InShaderPlatform, const ERHIFeatureLevel::Type&& InFeatureLevel)
+{
+	return InFeatureLevel <= GetMaxSupportedFeatureLevel(InShaderPlatform);
+}
+
+
+inline bool RHINeedsToSwitchVerticalAxis(const FStaticShaderPlatform Platform)
 {
 #if WITH_EDITOR
 	static const auto CVar = IConsoleManager::Get().FindTConsoleVariableDataInt(TEXT("r.Mobile.ForceRHISwitchVerticalAxis"));
@@ -1278,49 +1552,49 @@ inline bool RHINeedsToSwitchVerticalAxis(EShaderPlatform Platform)
 		   && Platform != SP_SWITCH && Platform != SP_SWITCH_FORWARD;
 }
 
-inline bool RHISupportsSeparateMSAAAndResolveTextures(const EShaderPlatform Platform)
+inline bool RHISupportsSeparateMSAAAndResolveTextures(const FStaticShaderPlatform Platform)
 {
 	// Metal mobile devices and Android ES3.1 need to handle MSAA and resolve textures internally (unless RHICreateTexture2D was changed to take an optional resolve target)
 	return !IsMetalMobilePlatform(Platform) && !IsAndroidOpenGLESPlatform(Platform);
 }
 
-inline bool RHISupportsComputeShaders(const EShaderPlatform Platform)
+inline bool RHISupportsComputeShaders(const FStaticShaderPlatform Platform)
 {
 	return IsFeatureLevelSupported(Platform, ERHIFeatureLevel::SM5) 
 		|| (GetMaxSupportedFeatureLevel(Platform) == ERHIFeatureLevel::ES3_1 && !IsSwitchPlatform(Platform));
 }
 
-inline bool RHISupportsGeometryShaders(const EShaderPlatform Platform)
+inline bool RHISupportsGeometryShaders(const FStaticShaderPlatform Platform)
 {
 	return IsFeatureLevelSupported(Platform, ERHIFeatureLevel::SM5) && !IsMetalPlatform(Platform) && !IsVulkanMobilePlatform(Platform);
 }
 
-inline bool RHIHasTiledGPU(const EShaderPlatform Platform)
+inline bool RHIHasTiledGPU(const FStaticShaderPlatform Platform)
 {
 	// @todo MetalMRT Technically we should include (Platform == SP_METAL_MRT) but this would disable depth-pre-pass which is currently required.
 	return Platform == SP_METAL || Platform == SP_METAL_TVOS
 		|| Platform == SP_OPENGL_ES2_ANDROID || Platform == SP_OPENGL_ES3_1_ANDROID
 		|| Platform == SP_VULKAN_ES3_1_ANDROID
-		|| FDataDrivenShaderPlatformInfo::GetInfo(Platform).bTargetsTiledGPU;
+		|| FDataDrivenShaderPlatformInfo::GetTargetsTiledGPU(Platform);
 }
 
-inline bool RHISupportsMobileMultiView(const EShaderPlatform Platform)
+inline bool RHISupportsMobileMultiView(const FStaticShaderPlatform Platform)
 {
 	return (Platform == EShaderPlatform::SP_OPENGL_ES3_1_ANDROID || Platform == EShaderPlatform::SP_OPENGL_ES2_ANDROID) || IsVulkanMobilePlatform(Platform)
-		|| FDataDrivenShaderPlatformInfo::GetInfo(Platform).bSupportsMobileMultiView;
+		|| FDataDrivenShaderPlatformInfo::GetSupportsMobileMultiView(Platform);
 }
 
-inline bool RHISupportsNativeShaderLibraries(const EShaderPlatform Platform)
+inline bool RHISupportsNativeShaderLibraries(const FStaticShaderPlatform Platform)
 {
 	return IsMetalPlatform(Platform);
 }
 
-inline bool RHISupportsShaderPipelines(EShaderPlatform Platform)
+inline bool RHISupportsShaderPipelines(const FStaticShaderPlatform Platform)
 {
 	return !IsMobilePlatform(Platform);
 }
 
-inline bool RHISupportsDualSourceBlending(EShaderPlatform Platform)
+inline bool RHISupportsDualSourceBlending(const FStaticShaderPlatform Platform)
 {
 	// For now only enable support for SM5
 	return IsFeatureLevelSupported(Platform, ERHIFeatureLevel::SM5) && (IsD3DPlatform(Platform, true) || IsPS4Platform(Platform) || IsVulkanPlatform(Platform) || IsMetalPlatform(Platform));
@@ -1329,7 +1603,7 @@ inline bool RHISupportsDualSourceBlending(EShaderPlatform Platform)
 // Return what the expected number of samplers will be supported by a feature level
 // Note that since the Feature Level is pretty orthogonal to the RHI/HW, this is not going to be perfect
 // If should only be used for a guess at the limit, the real limit will not be known until runtime
-inline uint32 GetExpectedFeatureLevelMaxTextureSamplers(ERHIFeatureLevel::Type FeatureLevel)
+inline uint32 GetExpectedFeatureLevelMaxTextureSamplers(const FStaticFeatureLevel FeatureLevel)
 {
 	if (FeatureLevel == ERHIFeatureLevel::ES2)
 	{
@@ -1341,7 +1615,7 @@ inline uint32 GetExpectedFeatureLevelMaxTextureSamplers(ERHIFeatureLevel::Type F
 	}
 }
 
-inline int32 GetFeatureLevelMaxNumberOfBones(ERHIFeatureLevel::Type FeatureLevel)
+inline int32 GetFeatureLevelMaxNumberOfBones(const FStaticFeatureLevel FeatureLevel)
 {
 	switch (FeatureLevel)
 	{
diff --git a/Engine/Source/Runtime/RHI/Public/RHIResources.h b/Engine/Source/Runtime/RHI/Public/RHIResources.h
--- a/Engine/Source/Runtime/RHI/Public/RHIResources.h
+++ b/Engine/Source/Runtime/RHI/Public/RHIResources.h
@@ -17,6 +17,9 @@
 #include "PixelFormat.h"
 #include "Containers/LockFreeList.h"
 #include "Misc/SecureHash.h"
+#include "Hash/CityHash.h"
+#include "Async/TaskGraphInterfaces.h"
+#include "Serialization/MemoryImage.h"
 
 #define DISABLE_RHI_DEFFERED_DELETE 0
 
@@ -198,100 +201,67 @@ public:
 	FString ShaderName;
 #endif
 
-	enum class EType : uint8
-	{
-		Vertex,
-		Hull,
-		Domain,
-		Pixel,
-		Geometry,
-		RayTracing,
-		Compute,
-	};
-
-	FRHIShader(EType InType)
-		: Type(InType)
+	explicit FRHIShader(EShaderFrequency InFrequency)
+		: Frequency(InFrequency)
 	{
 	}
 
-	inline EType GetType() const
+	inline EShaderFrequency GetFrequency() const
 	{
-		return Type;
-	}
-
-	static inline EShaderFrequency ConvertTypeToShaderFrequency(EType Type)
-	{
-		switch (Type)
-		{
-		case EType::Vertex:		return SF_Vertex;
-		case EType::Hull:		return SF_Hull;
-		case EType::Domain:		return SF_Domain;
-		case EType::Pixel:		return SF_Pixel;
-		case EType::Geometry:	return SF_Geometry;
-		case EType::RayTracing:	return SF_RayGen;	// Not really supported in this API...
-		case EType::Compute:	return SF_Compute;
-		default: break;
-		}
-
-		return SF_NumFrequencies;
-	}
-
-	inline EShaderFrequency GetShaderFrequency() const
-	{
-		return ConvertTypeToShaderFrequency(Type);
+		return Frequency;
 	}
 
 private:
 	FSHAHash Hash;
-	EType Type;
+	EShaderFrequency Frequency;
 };
 
 class RHI_VTABLE FRHIGraphicsShader : public FRHIShader
 {
 public:
-	FRHIGraphicsShader(FRHIShader::EType InType) : FRHIShader(InType) {}
+	explicit FRHIGraphicsShader(EShaderFrequency InFrequency) : FRHIShader(InFrequency) {}
 };
 
 class RHI_VTABLE FRHIVertexShader : public FRHIGraphicsShader
 {
 public:
-	FRHIVertexShader() : FRHIGraphicsShader(FRHIShader::EType::Vertex) {}
+	FRHIVertexShader() : FRHIGraphicsShader(SF_Vertex) {}
 };
 
 class RHI_VTABLE FRHIHullShader : public FRHIGraphicsShader
 {
 public:
-	FRHIHullShader() : FRHIGraphicsShader(FRHIShader::EType::Hull) {}
+	FRHIHullShader() : FRHIGraphicsShader(SF_Hull) {}
 };
 
 class RHI_VTABLE FRHIDomainShader : public FRHIGraphicsShader
 {
 public:
-	FRHIDomainShader() : FRHIGraphicsShader(FRHIShader::EType::Domain) {}
+	FRHIDomainShader() : FRHIGraphicsShader(SF_Domain) {}
 };
 
 class RHI_VTABLE FRHIPixelShader : public FRHIGraphicsShader
 {
 public:
-	FRHIPixelShader() : FRHIGraphicsShader(FRHIShader::EType::Pixel) {}
+	FRHIPixelShader() : FRHIGraphicsShader(SF_Pixel) {}
 };
 
 class RHI_VTABLE FRHIGeometryShader : public FRHIGraphicsShader
 {
 public:
-	FRHIGeometryShader() : FRHIGraphicsShader(FRHIShader::EType::Geometry) {}
+	FRHIGeometryShader() : FRHIGraphicsShader(SF_Geometry) {}
 };
 
 class RHI_VTABLE FRHIRayTracingShader : public FRHIShader
 {
 public:
-	FRHIRayTracingShader() : FRHIShader(FRHIShader::EType::RayTracing) {}
+	explicit FRHIRayTracingShader(EShaderFrequency InFrequency) : FRHIShader(InFrequency) {}
 };
 
 class RHI_API FRHIComputeShader : public FRHIShader
 {
 public:
-	FRHIComputeShader() : FRHIShader(FRHIShader::EType::Compute), Stats(nullptr) {}
+	FRHIComputeShader() : FRHIShader(SF_Compute), Stats(nullptr) {}
 	
 	inline void SetStats(struct FPipelineStateStats* Ptr) { Stats = Ptr; }
 	void UpdateStats();
@@ -322,28 +292,29 @@ class FRHIRayTracingPipelineState : public FRHIResource {};
 /** The layout of a uniform buffer in memory. */
 struct FRHIUniformBufferLayout
 {
+public:
 	/** Data structure to store information about resource parameter in a shader parameter structure. */
 	struct FResourceParameter
 	{
-		/** Byte offset to each resource in the uniform buffer memory. */
-		uint16 MemberOffset;
+		DECLARE_EXPORTED_TYPE_LAYOUT(FResourceParameter, RHI_API, NonVirtual);
+	public:
+		friend inline FArchive& operator<<(FArchive& Ar, FResourceParameter& Ref)
+		{
+			uint8 Type = (uint8)Ref.MemberType;
+			Ar << Ref.MemberOffset;
+			Ar << Type;
+			Ref.MemberType = (EUniformBufferBaseType)Type;
+			return Ar;
+		}
 
+		/** Byte offset to each resource in the uniform buffer memory. */
+		LAYOUT_FIELD(uint16, MemberOffset);
 		/** Type of the member that allow (). */
-		EUniformBufferBaseType MemberType;
+		LAYOUT_FIELD(EUniformBufferBaseType, MemberType);
 	};
 
-	/** The size of the constant buffer in bytes. */
-	uint32 ConstantBufferSize;
-
-	/** The static slot (if applicable). */
-	FUniformBufferStaticSlot StaticSlot = MAX_UNIFORM_BUFFER_STATIC_SLOTS;
-
-	/** The list of all resource inlined into the shader parameter structure. */
-	TArray<FResourceParameter> Resources;
-
-#if VALIDATE_UNIFORM_BUFFER_LAYOUT_LIFETIME
-	mutable int32 NumUsesForDebugging = 0;
-#endif
+	DECLARE_EXPORTED_TYPE_LAYOUT(FRHIUniformBufferLayout, RHI_API, NonVirtual);
+public:
 
 	inline uint32 GetHash() const
 	{
@@ -382,7 +353,7 @@ struct FRHIUniformBufferLayout
 		Hash = TmpHash;
 	}
 
-	explicit FRHIUniformBufferLayout(FName InName) :
+	explicit FRHIUniformBufferLayout(const TCHAR* InName) :
 		ConstantBufferSize(0),
 		Name(InName),
 		Hash(0)
@@ -395,7 +366,6 @@ struct FRHIUniformBufferLayout
 	};
 	explicit FRHIUniformBufferLayout(EInit) :
 		ConstantBufferSize(0),
-		Name(FName()),
 		Hash(0)
 	{
 	}
@@ -416,17 +386,36 @@ struct FRHIUniformBufferLayout
 		Hash = Source.Hash;
 	}
 
-	const FName GetDebugName() const { return Name; }
+	const FMemoryImageString& GetDebugName() const { return Name; }
 
 	uint32 NumRenderTargets()	const { return 0; }
 	uint32 NumTextures()		const { return 0; }
 	uint32 NumUAVs()			const { return 0; }
 
+	friend FArchive& operator<<(FArchive& Ar, FRHIUniformBufferLayout& Ref)
+	{
+		Ar << Ref.Name;
+		Ar << Ref.ConstantBufferSize;
+		Ar << Ref.Hash;
+		Ar << Ref.Resources;
+		return Ar;
+	}
+
+	/** The size of the constant buffer in bytes. */
+	LAYOUT_FIELD(uint32, ConstantBufferSize);
+
+	/** The static slot (if applicable). */
+	LAYOUT_FIELD_INITIALIZED(FUniformBufferStaticSlot, StaticSlot, MAX_UNIFORM_BUFFER_STATIC_SLOTS);
+
+	/** The list of all resource inlined into the shader parameter structure. */
+	LAYOUT_FIELD(TMemoryImageArray<FResourceParameter>, Resources);
+
+	LAYOUT_MUTABLE_FIELD(int32, NumUsesForDebugging);
+
 private:
 	// for debugging / error message
-	FName Name;
-
-	uint32 Hash;
+	LAYOUT_FIELD(FMemoryImageString, Name);
+	LAYOUT_FIELD(uint32, Hash);
 };
 
 /** Compare two uniform buffer layouts. */
@@ -1876,179 +1865,6 @@ struct FImmutableSamplerState
 	TImmutableSamplers ImmutableSamplers;
 };
 
-/** 
- * Pipeline state without render target state 
- * Useful for mesh passes where the render target state is not changing between draws.
- * Note: the size of this class affects rendering mesh pass traversal performance. 
- */
-class FGraphicsMinimalPipelineStateInitializer
-{
-public:
-	// Can't use TEnumByte<EPixelFormat> as it changes the struct to be non trivially constructible, breaking memset
-	using TRenderTargetFormats		= TStaticArray<uint8/*EPixelFormat*/, MaxSimultaneousRenderTargets>;
-	using TRenderTargetFlags		= TStaticArray<uint32, MaxSimultaneousRenderTargets>;
-
-	FGraphicsMinimalPipelineStateInitializer()
-		: BlendState(nullptr)
-		, RasterizerState(nullptr)
-		, DepthStencilState(nullptr)
-		, PrimitiveType(PT_Num)
-	{
-		static_assert(sizeof(EPixelFormat) != sizeof(uint8), "Change TRenderTargetFormats's uint8 to EPixelFormat");
-		static_assert(PF_MAX < MAX_uint8, "TRenderTargetFormats assumes EPixelFormat can fit in a uint8!");
-	}
-
-	FGraphicsMinimalPipelineStateInitializer(
-		FBoundShaderStateInput		InBoundShaderState,
-		FRHIBlendState*				InBlendState,
-		FRHIRasterizerState*		InRasterizerState,
-		FRHIDepthStencilState*		InDepthStencilState,
-		FImmutableSamplerState		InImmutableSamplerState,
-		EPrimitiveType				InPrimitiveType
-		)
-		: BoundShaderState(InBoundShaderState)
-		, BlendState(InBlendState)
-		, RasterizerState(InRasterizerState)
-		, DepthStencilState(InDepthStencilState)
-		, ImmutableSamplerState(InImmutableSamplerState)
-		, PrimitiveType(InPrimitiveType)
-	{
-	}
-
-	FGraphicsMinimalPipelineStateInitializer(const FGraphicsMinimalPipelineStateInitializer& InMinimalState)
-		: BoundShaderState(InMinimalState.BoundShaderState)
-		, BlendState(InMinimalState.BlendState)
-		, RasterizerState(InMinimalState.RasterizerState)
-		, DepthStencilState(InMinimalState.DepthStencilState)
-		, ImmutableSamplerState(InMinimalState.ImmutableSamplerState)
-		, bDepthBounds(InMinimalState.bDepthBounds)
-		, bMultiView(InMinimalState.bMultiView)
-		, PrimitiveType(InMinimalState.PrimitiveType)
-	{
-	}
-
-	inline bool operator==(const FGraphicsMinimalPipelineStateInitializer& rhs) const
-	{
-		if (BoundShaderState.VertexDeclarationRHI != rhs.BoundShaderState.VertexDeclarationRHI || 
-			BoundShaderState.VertexShaderRHI != rhs.BoundShaderState.VertexShaderRHI ||
-			BoundShaderState.PixelShaderRHI != rhs.BoundShaderState.PixelShaderRHI ||
-#if PLATFORM_SUPPORTS_GEOMETRY_SHADERS
-			BoundShaderState.GeometryShaderRHI != rhs.BoundShaderState.GeometryShaderRHI ||
-#endif
-#if PLATFORM_SUPPORTS_TESSELLATION_SHADERS
-			BoundShaderState.DomainShaderRHI != rhs.BoundShaderState.DomainShaderRHI ||
-			BoundShaderState.HullShaderRHI != rhs.BoundShaderState.HullShaderRHI ||	
-#endif		
-			BlendState != rhs.BlendState || 
-			RasterizerState != rhs.RasterizerState || 
-			DepthStencilState != rhs.DepthStencilState ||
-			ImmutableSamplerState != rhs.ImmutableSamplerState ||
-			bDepthBounds != rhs.bDepthBounds ||
-			bMultiView != rhs.bMultiView ||
-			PrimitiveType != rhs.PrimitiveType) 
-		{
-			return false;
-		}
-
-		return true;
-	}
-
-	inline bool operator!=(const FGraphicsMinimalPipelineStateInitializer& rhs) const
-	{
-		return !(*this == rhs);
-	}
-
-	inline friend uint32 GetTypeHash(const FGraphicsMinimalPipelineStateInitializer& Initializer)
-	{
-		return PointerHash(Initializer.BoundShaderState.VertexDeclarationRHI, 
-			PointerHash(Initializer.BoundShaderState.VertexShaderRHI, 
-				PointerHash(Initializer.BoundShaderState.PixelShaderRHI, 
-					PointerHash(Initializer.RasterizerState))));
-	}
-
-#define COMPARE_FIELD_BEGIN(Field) \
-		if (Field != rhs.Field) \
-		{ return Field COMPARE_OP rhs.Field; }
-
-#define COMPARE_FIELD(Field) \
-		else if (Field != rhs.Field) \
-		{ return Field COMPARE_OP rhs.Field; }
-
-#define COMPARE_FIELD_END \
-		else { return false; }
-
-	bool operator<(const FGraphicsMinimalPipelineStateInitializer& rhs) const
-	{
-#define COMPARE_OP <
-
-		COMPARE_FIELD_BEGIN(BoundShaderState.VertexDeclarationRHI)
-		COMPARE_FIELD(BoundShaderState.VertexShaderRHI)
-		COMPARE_FIELD(BoundShaderState.PixelShaderRHI)
-#if PLATFORM_SUPPORTS_GEOMETRY_SHADERS
-		COMPARE_FIELD(BoundShaderState.GeometryShaderRHI)
-#endif
-#if PLATFORM_SUPPORTS_TESSELLATION_SHADERS
-		COMPARE_FIELD(BoundShaderState.DomainShaderRHI)
-		COMPARE_FIELD(BoundShaderState.HullShaderRHI)
-#endif
-		COMPARE_FIELD(BlendState)
-		COMPARE_FIELD(RasterizerState)
-		COMPARE_FIELD(DepthStencilState)
-		COMPARE_FIELD(bDepthBounds)
-		COMPARE_FIELD(bMultiView)
-		COMPARE_FIELD(PrimitiveType)
-		COMPARE_FIELD_END;
-
-#undef COMPARE_OP
-	}
-
-	bool operator>(const FGraphicsMinimalPipelineStateInitializer& rhs) const
-	{
-#define COMPARE_OP >
-
-		COMPARE_FIELD_BEGIN(BoundShaderState.VertexDeclarationRHI)
-		COMPARE_FIELD(BoundShaderState.VertexShaderRHI)
-		COMPARE_FIELD(BoundShaderState.PixelShaderRHI)
-#if PLATFORM_SUPPORTS_GEOMETRY_SHADERS
-		COMPARE_FIELD(BoundShaderState.GeometryShaderRHI)
-#endif
-#if PLATFORM_SUPPORTS_TESSELLATION_SHADERS
-		COMPARE_FIELD(BoundShaderState.DomainShaderRHI)
-		COMPARE_FIELD(BoundShaderState.HullShaderRHI)
-#endif
-		COMPARE_FIELD(BlendState)
-		COMPARE_FIELD(RasterizerState)
-		COMPARE_FIELD(DepthStencilState)
-		COMPARE_FIELD(bDepthBounds)
-		COMPARE_FIELD(bMultiView)
-		COMPARE_FIELD(PrimitiveType)
-		COMPARE_FIELD_END;
-
-#undef COMPARE_OP
-	}
-
-#undef COMPARE_FIELD_BEGIN
-#undef COMPARE_FIELD
-#undef COMPARE_FIELD_END
-
-	// TODO: [PSO API] - As we migrate reuse existing API objects, but eventually we can move to the direct initializers. 
-	// When we do that work, move this to RHI.h as its more appropriate there, but here for now since dependent typdefs are here.
-	FBoundShaderStateInput	BoundShaderState;
-	FRHIBlendState*			BlendState;
-	FRHIRasterizerState*	RasterizerState;
-	FRHIDepthStencilState*	DepthStencilState;
-	FImmutableSamplerState	ImmutableSamplerState;
-
-	// Note: FGraphicsMinimalPipelineStateInitializer is 8-byte aligned and can't have any implicit padding,
-	// as it is sometimes hashed and compared as raw bytes. Explicit padding is therefore required between
-	// all data members and at the end of the structure.
-	bool							bDepthBounds = false;
-	bool							bMultiView = false;
-	uint8							Padding[2] = {};
-
-	EPrimitiveType			PrimitiveType;
-};
-
 // Hints for some RHIs that support subpasses
 enum class ESubpassHint : uint8
 {
@@ -2059,7 +1875,7 @@ enum class ESubpassHint : uint8
 	DepthReadSubpass,
 };
 
-class FGraphicsPipelineStateInitializer final : public FGraphicsMinimalPipelineStateInitializer
+class FGraphicsPipelineStateInitializer
 {
 public:
 	// Can't use TEnumByte<EPixelFormat> as it changes the struct to be non trivially constructible, breaking memset
@@ -2079,30 +1895,14 @@ public:
 		, NumSamples(0)
 		, SubpassHint(ESubpassHint::None)
 		, SubpassIndex(0)
+		, bDepthBounds(false)
+		, bMultiView(false)
 		, Flags(0)
 	{
 		static_assert(sizeof(EPixelFormat) != sizeof(uint8), "Change TRenderTargetFormats's uint8 to EPixelFormat");
 		static_assert(PF_MAX < MAX_uint8, "TRenderTargetFormats assumes EPixelFormat can fit in a uint8!");
 	}
 
-	FGraphicsPipelineStateInitializer(const FGraphicsMinimalPipelineStateInitializer& InMinimalState)
-		: FGraphicsMinimalPipelineStateInitializer(InMinimalState)
-		, RenderTargetsEnabled(0)
-		, RenderTargetFormats(PF_Unknown)
-		, RenderTargetFlags(0)
-		, DepthStencilTargetFormat(PF_Unknown)
-		, DepthStencilTargetFlag(0)
-		, DepthTargetLoadAction(ERenderTargetLoadAction::ENoAction)
-		, DepthTargetStoreAction(ERenderTargetStoreAction::ENoAction)
-		, StencilTargetLoadAction(ERenderTargetLoadAction::ENoAction)
-		, StencilTargetStoreAction(ERenderTargetStoreAction::ENoAction)
-		, NumSamples(0)
-		, SubpassHint(ESubpassHint::None)
-		, SubpassIndex(0)
-		, Flags(0)
-	{
-	}
-
 	FGraphicsPipelineStateInitializer(
 		FBoundShaderStateInput		InBoundShaderState,
 		FRHIBlendState*				InBlendState,
@@ -2123,9 +1923,16 @@ public:
 		uint32						InNumSamples,
 		ESubpassHint				InSubpassHint,
 		uint8						InSubpassIndex,
-		uint16						InFlags
+		uint16						InFlags,
+		bool						bInDepthBounds,
+		bool						bInMultiView
 		)
-		: FGraphicsMinimalPipelineStateInitializer(InBoundShaderState, InBlendState, InRasterizerState, InDepthStencilState, InImmutableSamplerState, InPrimitiveType)
+		: BoundShaderState(InBoundShaderState)
+		, BlendState(InBlendState)
+		, RasterizerState(InRasterizerState)
+		, DepthStencilState(InDepthStencilState)
+		, ImmutableSamplerState(InImmutableSamplerState)
+		, PrimitiveType(InPrimitiveType)
 		, RenderTargetsEnabled(InRenderTargetsEnabled)
 		, RenderTargetFormats(InRenderTargetFormats)
 		, RenderTargetFlags(InRenderTargetFlags)
@@ -2139,22 +1946,31 @@ public:
 		, NumSamples(InNumSamples)
 		, SubpassHint(InSubpassHint)
 		, SubpassIndex(InSubpassIndex)
+		, bDepthBounds(bInDepthBounds)
+		, bMultiView(bInMultiView)
 		, Flags(InFlags)
 	{
 	}
 
 	bool operator==(const FGraphicsPipelineStateInitializer& rhs) const
 	{
-		if (!FGraphicsMinimalPipelineStateInitializer::operator ==(rhs) ||
-			VertexShaderHash != rhs.VertexShaderHash ||
-			PixelShaderHash != rhs.PixelShaderHash ||
+		if (BoundShaderState.VertexDeclarationRHI != rhs.BoundShaderState.VertexDeclarationRHI ||
+			BoundShaderState.VertexShaderRHI != rhs.BoundShaderState.VertexShaderRHI ||
+			BoundShaderState.PixelShaderRHI != rhs.BoundShaderState.PixelShaderRHI ||
 #if PLATFORM_SUPPORTS_GEOMETRY_SHADERS
-			GeometryShaderHash != rhs.GeometryShaderHash ||
+			BoundShaderState.GeometryShaderRHI != rhs.BoundShaderState.GeometryShaderRHI ||
 #endif
 #if PLATFORM_SUPPORTS_TESSELLATION_SHADERS
-			HullShaderHash != rhs.HullShaderHash ||
-			DomainShaderHash != rhs.DomainShaderHash ||
+			BoundShaderState.DomainShaderRHI != rhs.BoundShaderState.DomainShaderRHI ||
+			BoundShaderState.HullShaderRHI != rhs.BoundShaderState.HullShaderRHI ||
 #endif
+			BlendState != rhs.BlendState ||
+			RasterizerState != rhs.RasterizerState ||
+			DepthStencilState != rhs.DepthStencilState ||
+			ImmutableSamplerState != rhs.ImmutableSamplerState ||
+			PrimitiveType != rhs.PrimitiveType ||
+			bDepthBounds != rhs.bDepthBounds ||
+			bMultiView != rhs.bMultiView ||
 			RenderTargetsEnabled != rhs.RenderTargetsEnabled ||
 			RenderTargetFormats != rhs.RenderTargetFormats || 
 			RenderTargetFlags != rhs.RenderTargetFlags || 
@@ -2194,15 +2010,13 @@ public:
 		return RenderTargetsEnabled;
 	}
 
-	FSHAHash						VertexShaderHash;
-	FSHAHash						PixelShaderHash;
-#if PLATFORM_SUPPORTS_GEOMETRY_SHADERS
-	FSHAHash						GeometryShaderHash;
-#endif
-#if PLATFORM_SUPPORTS_TESSELLATION_SHADERS
-	FSHAHash						HullShaderHash;
-	FSHAHash						DomainShaderHash;
-#endif
+	FBoundShaderStateInput			BoundShaderState;
+	FRHIBlendState*					BlendState;
+	FRHIRasterizerState*			RasterizerState;
+	FRHIDepthStencilState*			DepthStencilState;
+	FImmutableSamplerState			ImmutableSamplerState;
+
+	EPrimitiveType					PrimitiveType;
 	uint32							RenderTargetsEnabled;
 	TRenderTargetFormats			RenderTargetFormats;
 	TRenderTargetFlags				RenderTargetFlags;
@@ -2216,6 +2030,8 @@ public:
 	uint16							NumSamples;
 	ESubpassHint					SubpassHint;
 	uint8							SubpassIndex;
+	bool							bDepthBounds;
+	bool							bMultiView;
 	
 	// Note: these flags do NOT affect compilation of this PSO.
 	// The resulting object is invariant with respect to whatever is set here, they are
@@ -2379,55 +2195,20 @@ public:
 	virtual ~FRHIShaderLibrary() {}
 	
 	FORCEINLINE EShaderPlatform GetPlatform(void) const { return Platform; }
-	FORCEINLINE FString GetName(void) const { return LibraryName; }
+	FORCEINLINE const FString& GetName(void) const { return LibraryName; }
 	FORCEINLINE uint32 GetId(void) const { return LibraryId; }
 	
 	virtual bool IsNativeLibrary() const = 0;
-	
-	//Library iteration
-	struct FShaderLibraryEntry
-	{
-		FShaderLibraryEntry(): Frequency(SF_NumFrequencies), Platform(SP_NumPlatforms) {}
-		
-		FSHAHash Hash;
-		EShaderFrequency Frequency;
-		EShaderPlatform Platform;
-		
-		bool IsValid() const {return (Frequency < SF_NumFrequencies) && (Platform < SP_NumPlatforms);}
-	};
-	
-	class FShaderLibraryIterator : public FRHIResource
-	{
-	public:
-		FShaderLibraryIterator(FRHIShaderLibrary* ShaderLibrary) : ShaderLibrarySource(ShaderLibrary) {}
-		virtual ~FShaderLibraryIterator() {}
-	
-		//Is the iterator valid
-		virtual bool IsValid() const					 = 0;
-		
-		//Iterator position access
-		virtual FShaderLibraryEntry operator*()	const	 = 0;
-		
-		//Iterator next operation
-		virtual FShaderLibraryIterator& operator++()	 = 0;
-		
-		//Access the library we are iterating through - allow query e.g. GetPlatform from iterator object
-		FRHIShaderLibrary* GetLibrary() const			 {return ShaderLibrarySource;};
-		
-	protected:
-		//Control source object lifetime while iterator is 'active'
-		TRefCountPtr<FRHIShaderLibrary> ShaderLibrarySource;
-	};
-	
-	virtual TRefCountPtr<FShaderLibraryIterator> CreateIterator(void) = 0;
-	virtual bool RequestEntry(const FSHAHash& Hash, FArchive* Ar) = 0;
-	virtual bool RequestEntry(const FSHAHash& Hash, TArray<uint8>& OutRaw)
-	{
-		check(!"This shader code library does not support raw reads!");
-		return false;
-	}
-	virtual bool ContainsEntry(const FSHAHash& Hash) = 0;
-	virtual uint32 GetShaderCount(void) const = 0;
+	virtual int32 GetNumShaderMaps() const = 0;
+	virtual int32 GetNumShaders() const = 0;
+	virtual int32 GetNumShadersForShaderMap(int32 ShaderMapIndex) const = 0;
+	virtual int32 GetShaderIndex(int32 ShaderMapIndex, int32 i) const = 0;
+	virtual int32 FindShaderMapIndex(const FSHAHash& Hash) = 0;
+	virtual int32 FindShaderIndex(const FSHAHash& Hash) = 0;
+	virtual FGraphEventRef PreloadShader(int32 ShaderIndex) { return FGraphEventRef(); }
+	virtual FGraphEventRef PreloadShaderMap(int32 ShaderMapIndex) { return FGraphEventRef(); }
+
+	virtual TRefCountPtr<FRHIShader> CreateShader(int32 ShaderIndex) { return nullptr; }
 
 protected:
 	EShaderPlatform Platform;
diff --git a/Engine/Source/Runtime/RHI/Public/RHIStaticStates.h b/Engine/Source/Runtime/RHI/Public/RHIStaticStates.h
--- a/Engine/Source/Runtime/RHI/Public/RHIStaticStates.h
+++ b/Engine/Source/Runtime/RHI/Public/RHIStaticStates.h
@@ -57,49 +57,13 @@ template<typename InitializerType,typename RHIRefType,typename RHIParamRefType>
 class TStaticStateRHI
 {
 public:
-
-	static void GetRHI_WithNoReturnValue()
-	{
-		GetRHI();
-	}
 	static RHIParamRefType GetRHI()
 	{
-		// This is super-duper nasty. We rely upon the fact that all compilers will assign uninitialized static, POD data (like StaticResource) to a zero seg and not do any construction on it whatsoever.
-		static FStaticStateResource* StaticResource; // Must be left uninitialized!
-
-		if (!StaticResource)
-		{
-			if (GIsRHIInitialized && GRHISupportsRHIThread)
-			{
-				FStaticStateResource* NewStaticResource = new FStaticStateResource();
-				FStaticStateResource* ValueWas = (FStaticStateResource*)FPlatformAtomics::InterlockedCompareExchangePointer((void**)&StaticResource, NewStaticResource, nullptr);
-				if (ValueWas)
-				{
-					// we made a redundant one...leak it
-				}
-			}
-			else
-			{
-				if (!IsInRenderingThread())
-				{
-					check(IsInParallelRenderingThread());
-					{
-						FScopedEvent Event;
-						TGraphTask<FInitStaticResourceRenderThreadTask>::CreateTask().ConstructAndDispatchWhenReady(&GetRHI_WithNoReturnValue, Event);
-					}
-				}
-				else
-				{
-					StaticResource = new FStaticStateResource();
-				}
-			}
-			CA_ASSUME(StaticResource);
-		}
-		return StaticResource->StateRHI;
+		checkSlow(StaticResource.StateRHI);
+		return StaticResource.StateRHI;
 	};
 
 private:
-
 	/** A resource which manages the RHI resource. */
 	class FStaticStateResource : public FRenderResource
 	{
@@ -113,8 +77,7 @@ private:
 			}
 			else
 			{
-				InitResource();
-
+				BeginInitResource(this);
 			}
 		}
 
@@ -136,8 +99,13 @@ private:
 			ReleaseResource();
 		}
 	};
+
+	static FStaticStateResource StaticResource;
 };
 
+template<typename InitializerType, typename RHIRefType, typename RHIParamRefType>
+typename TStaticStateRHI<InitializerType, RHIRefType, RHIParamRefType>::FStaticStateResource TStaticStateRHI<InitializerType, RHIRefType, RHIParamRefType>::StaticResource = TStaticStateRHI<InitializerType, RHIRefType, RHIParamRefType>::FStaticStateResource();
+
 /**
  * A static RHI sampler state resource.
  * TStaticSamplerStateRHI<...>::GetStaticState() will return a FSamplerStateRHIRef to a sampler state with the desired settings.
diff --git a/Engine/Source/Runtime/RHI/Public/RHIUtilities.h b/Engine/Source/Runtime/RHI/Public/RHIUtilities.h
--- a/Engine/Source/Runtime/RHI/Public/RHIUtilities.h
+++ b/Engine/Source/Runtime/RHI/Public/RHIUtilities.h
@@ -880,7 +880,7 @@ inline uint32 GetVertexCountForPrimitiveCount(uint32 NumPrimitives, uint32 Primi
 inline uint32 ComputeAnisotropyRT(int32 InitializerMaxAnisotropy)
 {
 	static const auto CVar = IConsoleManager::Get().FindTConsoleVariableDataInt(TEXT("r.MaxAnisotropy"));
-	int32 CVarValue = CVar->GetValueOnRenderThread();
+	int32 CVarValue = CVar->GetValueOnAnyThread(); // this is sometimes called from main thread during initialization of static RHI states
 
 	return FMath::Clamp(InitializerMaxAnisotropy > 0 ? InitializerMaxAnisotropy : CVarValue, 1, 16);
 }
diff --git a/Engine/Source/Runtime/RHI/Public/RHIValidation.h b/Engine/Source/Runtime/RHI/Public/RHIValidation.h
--- a/Engine/Source/Runtime/RHI/Public/RHIValidation.h
+++ b/Engine/Source/Runtime/RHI/Public/RHIValidation.h
@@ -80,69 +80,36 @@ public:
 	}
 
 	// FlushType: Wait RHI Thread
-	virtual FPixelShaderRHIRef RHICreatePixelShader(const TArray<uint8>& Code) override final
+	virtual FPixelShaderRHIRef RHICreatePixelShader(TArrayView<const uint8> Code, const FSHAHash& Hash) override final
 	{
-		return RHI->RHICreatePixelShader(Code);
+		return RHI->RHICreatePixelShader(Code, Hash);
 	}
 
 	// FlushType: Wait RHI Thread
-	virtual FPixelShaderRHIRef RHICreatePixelShader(FRHIShaderLibrary* Library, FSHAHash Hash) override final
+	virtual FVertexShaderRHIRef RHICreateVertexShader(TArrayView<const uint8> Code, const FSHAHash& Hash) override final
 	{
-		return RHI->RHICreatePixelShader(Library, Hash);
+		return RHI->RHICreateVertexShader(Code, Hash);
 	}
 
 	// FlushType: Wait RHI Thread
-	virtual FVertexShaderRHIRef RHICreateVertexShader(const TArray<uint8>& Code) override final
-	{
-		return RHI->RHICreateVertexShader(Code);
-	}
-
-	// FlushType: Wait RHI Thread
-	virtual FVertexShaderRHIRef RHICreateVertexShader(FRHIShaderLibrary* Library, FSHAHash Hash) override final
-	{
-		return RHI->RHICreateVertexShader(Library, Hash);
-	}
-
-	// FlushType: Wait RHI Thread
-	virtual FHullShaderRHIRef RHICreateHullShader(const TArray<uint8>& Code) override final
-	{
-		check(RHISupportsTessellation(GMaxRHIShaderPlatform));
-		return RHI->RHICreateHullShader(Code);
-	}
-
-	// FlushType: Wait RHI Thread
-	virtual FHullShaderRHIRef RHICreateHullShader(FRHIShaderLibrary* Library, FSHAHash Hash) override final
+	virtual FHullShaderRHIRef RHICreateHullShader(TArrayView<const uint8> Code, const FSHAHash& Hash) override final
 	{
 		check(RHISupportsTessellation(GMaxRHIShaderPlatform));
-		return RHI->RHICreateHullShader(Library, Hash);
+		return RHI->RHICreateHullShader(Code, Hash);
 	}
 
 	// FlushType: Wait RHI Thread
-	virtual FDomainShaderRHIRef RHICreateDomainShader(const TArray<uint8>& Code) override final
+	virtual FDomainShaderRHIRef RHICreateDomainShader(TArrayView<const uint8> Code, const FSHAHash& Hash) override final
 	{
 		check(RHISupportsTessellation(GMaxRHIShaderPlatform));
-		return RHI->RHICreateDomainShader(Code);
+		return RHI->RHICreateDomainShader(Code, Hash);
 	}
 
 	// FlushType: Wait RHI Thread
-	virtual FDomainShaderRHIRef RHICreateDomainShader(FRHIShaderLibrary* Library, FSHAHash Hash) override final
-	{
-		check(RHISupportsTessellation(GMaxRHIShaderPlatform));
-		return RHI->RHICreateDomainShader(Library, Hash);
-	}
-
-	// FlushType: Wait RHI Thread
-	virtual FGeometryShaderRHIRef RHICreateGeometryShader(const TArray<uint8>& Code) override final
-	{
-		check(RHISupportsGeometryShaders(GMaxRHIShaderPlatform));
-		return RHI->RHICreateGeometryShader(Code);
-	}
-
-	// FlushType: Wait RHI Thread
-	virtual FGeometryShaderRHIRef RHICreateGeometryShader(FRHIShaderLibrary* Library, FSHAHash Hash) override final
+	virtual FGeometryShaderRHIRef RHICreateGeometryShader(TArrayView<const uint8> Code, const FSHAHash& Hash) override final
 	{
 		check(RHISupportsGeometryShaders(GMaxRHIShaderPlatform));
-		return RHI->RHICreateGeometryShader(Library, Hash);
+		return RHI->RHICreateGeometryShader(Code, Hash);
 	}
 
 	// Some RHIs can have pending messages/logs for error tracking, or debug modes
@@ -152,16 +119,10 @@ public:
 	}
 
 	// FlushType: Wait RHI Thread
-	virtual FComputeShaderRHIRef RHICreateComputeShader(const TArray<uint8>& Code) override final
+	virtual FComputeShaderRHIRef RHICreateComputeShader(TArrayView<const uint8> Code, const FSHAHash& Hash) override final
 	{
 		check(RHISupportsComputeShaders(GMaxRHIShaderPlatform));
-		return RHI->RHICreateComputeShader(Code);
-	}
-	// FlushType: Wait RHI Thread
-	virtual FComputeShaderRHIRef RHICreateComputeShader(FRHIShaderLibrary* Library, FSHAHash Hash) override final
-	{
-		check(RHISupportsComputeShaders(GMaxRHIShaderPlatform));
-		return RHI->RHICreateComputeShader(Library, Hash);
+		return RHI->RHICreateComputeShader(Code, Hash);
 	}
 
 	/**
@@ -1310,72 +1271,38 @@ public:
 		RHI->UnlockStructuredBuffer_BottomOfPipe(RHICmdList, StructuredBuffer);
 	}
 
-	virtual FVertexShaderRHIRef CreateVertexShader_RenderThread(class FRHICommandListImmediate& RHICmdList, const TArray<uint8>& Code) override final
-	{
-		return RHI->CreateVertexShader_RenderThread(RHICmdList, Code);
-	}
-
-	virtual FVertexShaderRHIRef CreateVertexShader_RenderThread(class FRHICommandListImmediate& RHICmdList, FRHIShaderLibrary* Library, FSHAHash Hash) override final
-	{
-		return RHI->CreateVertexShader_RenderThread(RHICmdList, Library, Hash);
-	}
-
-	virtual FPixelShaderRHIRef CreatePixelShader_RenderThread(class FRHICommandListImmediate& RHICmdList, const TArray<uint8>& Code) override final
-	{
-		return RHI->CreatePixelShader_RenderThread(RHICmdList, Code);
-	}
-
-	virtual FPixelShaderRHIRef CreatePixelShader_RenderThread(class FRHICommandListImmediate& RHICmdList, FRHIShaderLibrary* Library, FSHAHash Hash) override final
+	virtual FVertexShaderRHIRef CreateVertexShader_RenderThread(class FRHICommandListImmediate& RHICmdList, TArrayView<const uint8> Code, const FSHAHash& Hash) override final
 	{
-		return RHI->CreatePixelShader_RenderThread(RHICmdList, Library, Hash);
+		return RHI->CreateVertexShader_RenderThread(RHICmdList, Code, Hash);
 	}
 
-	virtual FGeometryShaderRHIRef CreateGeometryShader_RenderThread(class FRHICommandListImmediate& RHICmdList, const TArray<uint8>& Code) override final
+	virtual FPixelShaderRHIRef CreatePixelShader_RenderThread(class FRHICommandListImmediate& RHICmdList, TArrayView<const uint8> Code, const FSHAHash& Hash) override final
 	{
-		check(RHISupportsGeometryShaders(GMaxRHIShaderPlatform));
-		return RHI->CreateGeometryShader_RenderThread(RHICmdList, Code);
+		return RHI->CreatePixelShader_RenderThread(RHICmdList, Code, Hash);
 	}
 
-	virtual FGeometryShaderRHIRef CreateGeometryShader_RenderThread(class FRHICommandListImmediate& RHICmdList, FRHIShaderLibrary* Library, FSHAHash Hash) override final
+	virtual FGeometryShaderRHIRef CreateGeometryShader_RenderThread(class FRHICommandListImmediate& RHICmdList, TArrayView<const uint8> Code, const FSHAHash& Hash) override final
 	{
 		check(RHISupportsGeometryShaders(GMaxRHIShaderPlatform));
-		return RHI->CreateGeometryShader_RenderThread(RHICmdList, Library, Hash);
+		return RHI->CreateGeometryShader_RenderThread(RHICmdList, Code, Hash);
 	}
 
-	virtual FComputeShaderRHIRef CreateComputeShader_RenderThread(class FRHICommandListImmediate& RHICmdList, const TArray<uint8>& Code) override final
+	virtual FComputeShaderRHIRef CreateComputeShader_RenderThread(class FRHICommandListImmediate& RHICmdList, TArrayView<const uint8> Code, const FSHAHash& Hash) override final
 	{
 		check(RHISupportsComputeShaders(GMaxRHIShaderPlatform));
-		return RHI->CreateComputeShader_RenderThread(RHICmdList, Code);
-	}
-
-	virtual FComputeShaderRHIRef CreateComputeShader_RenderThread(class FRHICommandListImmediate& RHICmdList, FRHIShaderLibrary* Library, FSHAHash Hash) override final
-	{
-		check(RHISupportsComputeShaders(GMaxRHIShaderPlatform));
-		return RHI->CreateComputeShader_RenderThread(RHICmdList, Library, Hash);
-	}
-
-	virtual FHullShaderRHIRef CreateHullShader_RenderThread(class FRHICommandListImmediate& RHICmdList, const TArray<uint8>& Code) override final
-	{
-		check(RHISupportsTessellation(GMaxRHIShaderPlatform));
-		return RHI->CreateHullShader_RenderThread(RHICmdList, Code);
-	}
-
-	virtual FHullShaderRHIRef CreateHullShader_RenderThread(class FRHICommandListImmediate& RHICmdList, FRHIShaderLibrary* Library, FSHAHash Hash) override final
-	{
-		check(RHISupportsTessellation(GMaxRHIShaderPlatform));
-		return RHI->CreateHullShader_RenderThread(RHICmdList, Library, Hash);
+		return RHI->CreateComputeShader_RenderThread(RHICmdList, Code, Hash);
 	}
 
-	virtual FDomainShaderRHIRef CreateDomainShader_RenderThread(class FRHICommandListImmediate& RHICmdList, const TArray<uint8>& Code) override final
+	virtual FHullShaderRHIRef CreateHullShader_RenderThread(class FRHICommandListImmediate& RHICmdList, TArrayView<const uint8> Code, const FSHAHash& Hash) override final
 	{
 		check(RHISupportsTessellation(GMaxRHIShaderPlatform));
-		return RHI->CreateDomainShader_RenderThread(RHICmdList, Code);
+		return RHI->CreateHullShader_RenderThread(RHICmdList, Code, Hash);
 	}
 
-	virtual FDomainShaderRHIRef CreateDomainShader_RenderThread(class FRHICommandListImmediate& RHICmdList, FRHIShaderLibrary* Library, FSHAHash Hash) override final
+	virtual FDomainShaderRHIRef CreateDomainShader_RenderThread(class FRHICommandListImmediate& RHICmdList, TArrayView<const uint8> Code, const FSHAHash& Hash) override final
 	{
 		check(RHISupportsTessellation(GMaxRHIShaderPlatform));
-		return RHI->CreateDomainShader_RenderThread(RHICmdList, Library, Hash);
+		return RHI->CreateDomainShader_RenderThread(RHICmdList, Code, Hash);
 	}
 
 	virtual void* LockTexture2D_RenderThread(class FRHICommandListImmediate& RHICmdList, FRHITexture2D* Texture, uint32 MipIndex, EResourceLockMode LockMode, uint32& DestStride, bool bLockWithinMiptail, bool bNeedsDefaultRHIFlush = true) override final
@@ -1642,9 +1569,9 @@ public:
 		return RHI->RHICreateRayTracingScene(Initializer);
 	}
 
-	virtual FRayTracingShaderRHIRef RHICreateRayTracingShader(const TArray<uint8>& Code, EShaderFrequency ShaderFrequency) override final
+	virtual FRayTracingShaderRHIRef RHICreateRayTracingShader(TArrayView<const uint8> Code, const FSHAHash& Hash, EShaderFrequency ShaderFrequency) override final
 	{
-		return RHI->RHICreateRayTracingShader(Code, ShaderFrequency);
+		return RHI->RHICreateRayTracingShader(Code, Hash, ShaderFrequency);
 	}
 
 	virtual FRayTracingPipelineStateRHIRef RHICreateRayTracingPipelineState(const FRayTracingPipelineStateInitializer& Initializer) override final
diff --git a/Engine/Source/Runtime/RenderCore/Private/ClearQuad.cpp b/Engine/Source/Runtime/RenderCore/Private/ClearQuad.cpp
--- a/Engine/Source/Runtime/RenderCore/Private/ClearQuad.cpp
+++ b/Engine/Source/Runtime/RenderCore/Private/ClearQuad.cpp
@@ -56,54 +56,54 @@ static void ClearQuadSetup( FRHICommandList& RHICmdList, bool bClearColor, int32
 	// Set the new shaders
 	TShaderMapRef<TOneColorVS<true> > VertexShader(ShaderMap);
 
-	FOneColorPS* PixelShader = NULL;
+	TShaderRef<FOneColorPS> PixelShader;
 
 	// Set the shader to write to the appropriate number of render targets
 	// On AMD PC hardware, outputting to a color index in the shader without a matching render target set has a significant performance hit
 	if (NumClearColors <= 1)
 	{
 		TShaderMapRef<TOneColorPixelShaderMRT<1> > MRTPixelShader(ShaderMap);
-		PixelShader = *MRTPixelShader;
+		PixelShader = MRTPixelShader;
 	}
 	else if (NumClearColors == 2)
 	{
 		TShaderMapRef<TOneColorPixelShaderMRT<2> > MRTPixelShader(ShaderMap);
-		PixelShader = *MRTPixelShader;
+		PixelShader = MRTPixelShader;
 	}
 	else if (NumClearColors == 3)
 	{
 		TShaderMapRef<TOneColorPixelShaderMRT<3> > MRTPixelShader(ShaderMap);
-		PixelShader = *MRTPixelShader;
+		PixelShader = MRTPixelShader;
 	}
 	else if (NumClearColors == 4)
 	{
 		TShaderMapRef<TOneColorPixelShaderMRT<4> > MRTPixelShader(ShaderMap);
-		PixelShader = *MRTPixelShader;
+		PixelShader = MRTPixelShader;
 	}
 	else if (NumClearColors == 5)
 	{
 		TShaderMapRef<TOneColorPixelShaderMRT<5> > MRTPixelShader(ShaderMap);
-		PixelShader = *MRTPixelShader;
+		PixelShader = MRTPixelShader;
 	}
 	else if (NumClearColors == 6)
 	{
 		TShaderMapRef<TOneColorPixelShaderMRT<6> > MRTPixelShader(ShaderMap);
-		PixelShader = *MRTPixelShader;
+		PixelShader = MRTPixelShader;
 	}
 	else if (NumClearColors == 7)
 	{
 		TShaderMapRef<TOneColorPixelShaderMRT<7> > MRTPixelShader(ShaderMap);
-		PixelShader = *MRTPixelShader;
+		PixelShader = MRTPixelShader;
 	}
 	else if (NumClearColors == 8)
 	{
 		TShaderMapRef<TOneColorPixelShaderMRT<8> > MRTPixelShader(ShaderMap);
-		PixelShader = *MRTPixelShader;
+		PixelShader = MRTPixelShader;
 	}
 
 	GraphicsPSOInit.BoundShaderState.VertexDeclarationRHI = GetVertexDeclarationFVector4();
-	GraphicsPSOInit.BoundShaderState.VertexShaderRHI = GETSAFERHISHADER_VERTEX(*VertexShader);
-	GraphicsPSOInit.BoundShaderState.PixelShaderRHI = GETSAFERHISHADER_PIXEL(PixelShader);
+	GraphicsPSOInit.BoundShaderState.VertexShaderRHI = VertexShader.GetVertexShader();
+	GraphicsPSOInit.BoundShaderState.PixelShaderRHI = PixelShader.GetPixelShader();
 	GraphicsPSOInit.PrimitiveType = PT_TriangleStrip;
 
 	if (PSOModifier)
diff --git a/Engine/Source/Runtime/RenderCore/Private/CopyTextureShaders.cpp b/Engine/Source/Runtime/RenderCore/Private/CopyTextureShaders.cpp
--- a/Engine/Source/Runtime/RenderCore/Private/CopyTextureShaders.cpp
+++ b/Engine/Source/Runtime/RenderCore/Private/CopyTextureShaders.cpp
@@ -2,6 +2,8 @@
 
 #include "CopyTextureShaders.h"
 
+IMPLEMENT_TYPE_LAYOUT(FCopyTextureCS);
+
 #define IMPLEMENT_COPY_RESOURCE_SHADER(SrcType,DstType,ValueType)\
 	typedef TCopyResourceCS<ECopyTextureResourceType::SrcType, ECopyTextureResourceType::DstType, ECopyTextureValueType::ValueType, 4> FCopyTextureCS_##SrcType##_##DstType##_##ValueType##4;\
 	IMPLEMENT_SHADER_TYPE2_WITH_TEMPLATE_PREFIX(template<>, FCopyTextureCS_##SrcType##_##DstType##_##ValueType##4, SF_Compute);
diff --git a/Engine/Source/Runtime/RenderCore/Private/GenerateMips.cpp b/Engine/Source/Runtime/RenderCore/Private/GenerateMips.cpp
--- a/Engine/Source/Runtime/RenderCore/Private/GenerateMips.cpp
+++ b/Engine/Source/Runtime/RenderCore/Private/GenerateMips.cpp
@@ -70,12 +70,6 @@ public:
 	FMipsShadersVS(const ShaderMetaType::CompiledShaderInitializerType& Initializer)
 		: FGlobalShader(Initializer)
 	{ }
-
-	virtual bool Serialize(FArchive& Ar) override
-	{
-		bool bShaderHasOutdatedParameters = FGlobalShader::Serialize(Ar);
-		return bShaderHasOutdatedParameters;
-	}
 };
 IMPLEMENT_SHADER_TYPE(, FMipsShadersVS, TEXT("/Engine/Private/ComputeGenerateMips.usf"), TEXT("MainVS"), SF_Vertex);
 
@@ -115,7 +109,7 @@ public:
 		}
 
 		TUniformBufferRef<FMipsShadersUB> Data = TUniformBufferRef<FMipsShadersUB>::CreateUniformBufferImmediate(UB, UniformBuffer_SingleFrame);
-		SetUniformBufferParameter(RHICmdList, GetPixelShader(), GetUniformBufferParameter<FMipsShadersUB>(), Data);
+		SetUniformBufferParameter(RHICmdList, RHICmdList.GetBoundPixelShader(), GetUniformBufferParameter<FMipsShadersUB>(), Data);
 	}
 };
 
@@ -204,8 +198,8 @@ void FGenerateMips::RenderMips(FRHICommandListImmediate& CommandList, FRHITextur
 	TShaderMapRef<FMipsShadersPS> CopyShader(ShaderMap);
 
 	GraphicsPSOInit.BoundShaderState.VertexDeclarationRHI = GenMipsStruct->VertexDeclaration;
-	GraphicsPSOInit.BoundShaderState.VertexShaderRHI = GETSAFERHISHADER_VERTEX(*VertexShader);
-	GraphicsPSOInit.BoundShaderState.PixelShaderRHI = GETSAFERHISHADER_PIXEL(*CopyShader);
+	GraphicsPSOInit.BoundShaderState.VertexShaderRHI = VertexShader.GetVertexShader();
+	GraphicsPSOInit.BoundShaderState.PixelShaderRHI = CopyShader.GetPixelShader();
 
 	uint32 NumMips = InTexture->GetNumMips();
 	for (uint32 MipLevel=1; MipLevel<NumMips; ++MipLevel)
@@ -333,7 +327,7 @@ void FGenerateMips::Compute(FRHICommandListImmediate& RHIImmCmdList, FRHITexture
 			FMath::Max((DestTextureSizeY + MIPSSHADER_NUMTHREADS - 1) / MIPSSHADER_NUMTHREADS, 1),
 			1);
 		//Pass added per mip level to be written.
-		ClearUnusedGraphResources(*ComputeShader, PassParameters);
+		ClearUnusedGraphResources(ComputeShader, PassParameters);
 
 		GraphBuilder.AddPass(
 			RDG_EVENT_NAME("Generate2DTextureMips DestMipLevel=%d", MipLevel),
@@ -341,7 +335,7 @@ void FGenerateMips::Compute(FRHICommandListImmediate& RHIImmCmdList, FRHITexture
 			ERDGPassFlags::Compute | ERDGPassFlags::GenerateMips,
 			[PassParameters, ComputeShader, GenMipsGroupCount](FRHICommandList& RHICmdList)
 		{
-			FComputeShaderUtils::Dispatch(RHICmdList, *ComputeShader, *PassParameters, GenMipsGroupCount);
+			FComputeShaderUtils::Dispatch(RHICmdList, ComputeShader, *PassParameters, GenMipsGroupCount);
 		});
 	}
 	GraphBuilder.QueueTextureExtraction(GraphTexture, &GenMipsStruct->RenderTarget);
@@ -433,7 +427,7 @@ void FGenerateMips::Execute(FRDGBuilder* GraphBuilder, FRDGTextureRef InGraphTex
 			1);
 		
 		//Pass added per mip level to be written.
-		ClearUnusedGraphResources(*ComputeShader, PassParameters);
+		ClearUnusedGraphResources(ComputeShader, PassParameters);
 
 		GraphBuilder->AddPass(
 			RDG_EVENT_NAME("Generate2DTextureMips DestMipLevel=%d", MipLevel),
@@ -441,7 +435,7 @@ void FGenerateMips::Execute(FRDGBuilder* GraphBuilder, FRDGTextureRef InGraphTex
 			ERDGPassFlags::Compute | ERDGPassFlags::GenerateMips,
 			[PassParameters, ComputeShader, GenMipsGroupCount](FRHICommandList& RHICmdList)
 		{
-			FComputeShaderUtils::Dispatch(RHICmdList, *ComputeShader, *PassParameters, GenMipsGroupCount);
+			FComputeShaderUtils::Dispatch(RHICmdList, ComputeShader, *PassParameters, GenMipsGroupCount);
 		});
 
 	}
diff --git a/Engine/Source/Runtime/RenderCore/Private/GlobalShader.cpp b/Engine/Source/Runtime/RenderCore/Private/GlobalShader.cpp
--- a/Engine/Source/Runtime/RenderCore/Private/GlobalShader.cpp
+++ b/Engine/Source/Runtime/RenderCore/Private/GlobalShader.cpp
@@ -9,6 +9,7 @@
 #include "Misc/MessageDialog.h"
 #include "HAL/FileManager.h"
 #include "Serialization/NameAsStringProxyArchive.h"
+#include "Interfaces/ITargetPlatform.h"
 #include "Misc/Paths.h"
 #include "Serialization/MemoryWriter.h"
 #include "Serialization/MemoryReader.h"
@@ -16,8 +17,12 @@
 #include "Misc/App.h"
 #include "StaticBoundShaderState.h"
 
+
 /** The global shader map. */
-TShaderMap<FGlobalShaderType>* GGlobalShaderMap[SP_NumPlatforms] = {};
+FGlobalShaderMap* GGlobalShaderMap[SP_NumPlatforms] = {};
+
+IMPLEMENT_TYPE_LAYOUT(FGlobalShader);
+IMPLEMENT_TYPE_LAYOUT(FGlobalShaderMapContent);
 
 IMPLEMENT_SHADER_TYPE(,FNULLPS,TEXT("/Engine/Private/NullPixelShader.usf"),TEXT("Main"),SF_Pixel);
 
@@ -87,7 +92,7 @@ FGlobalShaderMapId::FGlobalShaderMapId(EShaderPlatform Platform)
 		FShaderType* ShaderType = ShaderTypes[TypeIndex];
 		FShaderTypeDependency Dependency(ShaderType, Platform);
 
-		const TCHAR* ShaderFilename = Dependency.ShaderType->GetShaderFilename();
+		const TCHAR* ShaderFilename = ShaderType->GetShaderFilename();
 
 		TArray<FShaderTypeDependency>& Dependencies = ShaderFilenameToDependenciesMap.FindOrAdd(ShaderFilename);
 
@@ -104,31 +109,51 @@ FGlobalShaderMapId::FGlobalShaderMapId(EShaderPlatform Platform)
 	}
 }
 
-void FGlobalShaderMapId::AppendKeyString(FString& KeyString, const TArray<FShaderTypeDependency>& Dependencies) const
+void FGlobalShaderMapId::AppendKeyString(FString& KeyString, const TArray<FShaderTypeDependency>& Dependencies, const ITargetPlatform* TargetPlatform) const
 {
 #if WITH_EDITOR
+
+	FPlatformTypeLayoutParameters LayoutParams;
+	if (TargetPlatform)
+	{
+		LayoutParams.InitializeForPlatform(TargetPlatform->IniPlatformName(), TargetPlatform->HasEditorOnlyData());
+	}
+	else
+	{
+		LayoutParams.InitializeForCurrent();
+	}
+
+	{
+		const FSHAHash LayoutHash = Freeze::HashLayout(StaticGetTypeLayoutDesc<FGlobalShaderMapContent>(), LayoutParams);
+		KeyString += TEXT("_");
+		KeyString += LayoutHash.ToString();
+		KeyString += TEXT("_");
+	}
+
 	TMap<const TCHAR*,FCachedUniformBufferDeclaration> ReferencedUniformBuffers;
 
 	for (int32 ShaderIndex = 0; ShaderIndex < Dependencies.Num(); ShaderIndex++)
 	{
 		const FShaderTypeDependency& ShaderTypeDependency = Dependencies[ShaderIndex];
+		const FShaderType* ShaderType = FindShaderTypeByName(ShaderTypeDependency.ShaderTypeName);
 
 		KeyString += TEXT("_");
-		KeyString += ShaderTypeDependency.ShaderType->GetName();
+		KeyString += ShaderType->GetName();
 		KeyString += FString::Printf(TEXT("%i"), ShaderTypeDependency.PermutationId);
 
 		// Add the type's source hash so that we can invalidate cached shaders when .usf changes are made
 		KeyString += ShaderTypeDependency.SourceHash.ToString();
 
-		if (const FShaderParametersMetadata* ParameterStructMetadata = ShaderTypeDependency.ShaderType->GetRootParametersMetadata())
+		if (const FShaderParametersMetadata* ParameterStructMetadata = ShaderType->GetRootParametersMetadata())
 		{
 			KeyString += FString::Printf(TEXT("%08x"), ParameterStructMetadata->GetLayoutHash());
 		}
 
 		// Add the serialization history to the key string so that we can detect changes to global shader serialization without a corresponding .usf change
-		ShaderTypeDependency.ShaderType->GetSerializationHistory().AppendKeyString(KeyString);
+		const FSHAHash LayoutHash = Freeze::HashLayout(ShaderType->GetLayout(), LayoutParams);
+		KeyString += LayoutHash.ToString();
 
-		const TMap<const TCHAR*, FCachedUniformBufferDeclaration>& ReferencedUniformBufferStructsCache = ShaderTypeDependency.ShaderType->GetReferencedUniformBufferStructsCache();
+		const TMap<const TCHAR*, FCachedUniformBufferDeclaration>& ReferencedUniformBufferStructsCache = ShaderType->GetReferencedUniformBufferStructsCache();
 
 		// Gather referenced uniform buffers
 		for (TMap<const TCHAR*, FCachedUniformBufferDeclaration>::TConstIterator It(ReferencedUniformBufferStructsCache); It; ++It)
@@ -140,14 +165,15 @@ void FGlobalShaderMapId::AppendKeyString(FString& KeyString, const TArray<FShade
 	for (int32 Index = 0; Index < ShaderPipelineTypeDependencies.Num(); ++Index)
 	{
 		const FShaderPipelineTypeDependency& Dependency = ShaderPipelineTypeDependencies[Index];
+		const FShaderPipelineType* ShaderPipelineType = FShaderPipelineType::GetShaderPipelineTypeByName(Dependency.ShaderPipelineTypeName);
 
 		KeyString += TEXT("_");
-		KeyString += Dependency.ShaderPipelineType->GetName();
+		KeyString += ShaderPipelineType->GetName();
 
 		// Add the type's source hash so that we can invalidate cached shaders when .usf changes are made
 		KeyString += Dependency.StagesSourceHash.ToString();
 
-		for (const FShaderType* ShaderType : Dependency.ShaderPipelineType->GetStages())
+		for (const FShaderType* ShaderType : ShaderPipelineType->GetStages())
 		{
 			if (const FShaderParametersMetadata* ParameterStructMetadata = ShaderType->GetRootParametersMetadata())
 			{
@@ -184,6 +210,7 @@ FGlobalShader::FGlobalShader(const ShaderMetaType::CompiledShaderInitializerType
 
 void BackupGlobalShaderMap(FGlobalShaderBackupData& OutGlobalShaderBackup)
 {
+#if 0
 	for (int32 i = (int32)ERHIFeatureLevel::ES2; i < (int32)ERHIFeatureLevel::Num; ++i)
 	{
 		EShaderPlatform ShaderPlatform = GetFeatureLevelShaderPlatform((ERHIFeatureLevel::Type)i);
@@ -192,7 +219,7 @@ void BackupGlobalShaderMap(FGlobalShaderBackupData& OutGlobalShaderBackup)
 			TUniquePtr<TArray<uint8>> ShaderData = MakeUnique<TArray<uint8>>();
 			FMemoryWriter Ar(*ShaderData);
 			GGlobalShaderMap[ShaderPlatform]->SerializeInline(Ar, true, true, false, nullptr);
-			GGlobalShaderMap[ShaderPlatform]->RegisterSerializedShaders(false);
+			//GGlobalShaderMap[ShaderPlatform]->RegisterSerializedShaders(false);
 			GGlobalShaderMap[ShaderPlatform]->Empty();
 			OutGlobalShaderBackup.FeatureLevelShaderData[i] = MoveTemp(ShaderData);
 		}
@@ -203,10 +230,13 @@ void BackupGlobalShaderMap(FGlobalShaderBackupData& OutGlobalShaderBackup)
 	{
 		BeginUpdateResourceRHI(*It);
 	}
+#endif
+	check(0);
 }
 
 void RestoreGlobalShaderMap(const FGlobalShaderBackupData& GlobalShaderBackup)
 {
+#if 0
 	for (int32 i = (int32)ERHIFeatureLevel::ES2; i < (int32)ERHIFeatureLevel::Num; ++i)
 	{
 		EShaderPlatform ShaderPlatform = GetFeatureLevelShaderPlatform((ERHIFeatureLevel::Type)i);		
@@ -216,16 +246,202 @@ void RestoreGlobalShaderMap(const FGlobalShaderBackupData& GlobalShaderBackup)
 		{
 			FMemoryReader Ar(*GlobalShaderBackup.FeatureLevelShaderData[i]);
 			GGlobalShaderMap[ShaderPlatform]->SerializeInline(Ar, true, true, false, nullptr);
-			GGlobalShaderMap[ShaderPlatform]->RegisterSerializedShaders(false);
+			//GGlobalShaderMap[ShaderPlatform]->RegisterSerializedShaders(false);
 		}
 	}
+#endif
+	check(0);
 }
 
 
-TShaderMap<FGlobalShaderType>* GetGlobalShaderMap(EShaderPlatform Platform)
+FGlobalShaderMap* GetGlobalShaderMap(EShaderPlatform Platform)
 {
 	// If the global shader map hasn't been created yet
 	check(GGlobalShaderMap[Platform]);
-
 	return GGlobalShaderMap[Platform];
 }
+
+FGlobalShaderMapSection* FGlobalShaderMapSection::CreateFromArchive(FArchive& Ar)
+{
+	FGlobalShaderMapSection* Section = new FGlobalShaderMapSection();
+	Section->Serialize(Ar);
+	return Section;
+}
+
+void FGlobalShaderMapSection::Serialize(FArchive& Ar)
+{
+	Super::Serialize(Ar, true, false);
+}
+
+TShaderRef<FShader> FGlobalShaderMapSection::GetShader(FShaderType* ShaderType, int32 PermutationId) const
+{
+	return TShaderRef<FShader>(GetContent()->GetShader(ShaderType, PermutationId), *this);
+}
+
+FShaderPipelineRef FGlobalShaderMapSection::GetShaderPipeline(const FShaderPipelineType* PipelineType) const
+{
+	return FShaderPipelineRef(GetContent()->GetShaderPipeline(PipelineType), *this);
+}
+
+FGlobalShaderMap::FGlobalShaderMap(EShaderPlatform InPlatform)
+	: Platform(InPlatform)
+{
+}
+
+FGlobalShaderMap::~FGlobalShaderMap()
+{
+	for (const auto& It : SectionMap)
+	{
+		delete It.Value;
+	}
+}
+
+TShaderRef<FShader> FGlobalShaderMap::GetShader(FShaderType* ShaderType, int32 PermutationId) const
+{
+	const FHashedName HashedFilename(ShaderType->GetShaderFilename());
+	FGlobalShaderMapSection* const* Section = SectionMap.Find(HashedFilename);
+	return Section ? (*Section)->GetShader(ShaderType, PermutationId) : TShaderRef<FShader>();
+}
+
+FShaderPipelineRef FGlobalShaderMap::GetShaderPipeline(const FShaderPipelineType* ShaderPipelineType) const
+{
+	FGlobalShaderMapSection* const* Section = SectionMap.Find(ShaderPipelineType->GetHashedPrimaryShaderFilename());
+	return Section ? (*Section)->GetShaderPipeline(ShaderPipelineType) : FShaderPipelineRef();
+}
+
+void FGlobalShaderMap::BeginCreateAllShaders()
+{
+	for (const auto& It : SectionMap)
+	{
+		It.Value->GetResource()->BeginCreateAllShaders();
+	}
+}
+
+#if WITH_EDITOR
+void FGlobalShaderMap::GetOutdatedTypes(TArray<const FShaderType*>& OutdatedShaderTypes, TArray<const FShaderPipelineType*>& OutdatedShaderPipelineTypes, TArray<const FVertexFactoryType*>& OutdatedFactoryTypes) const
+{
+	for (const auto& It : SectionMap)
+	{
+		It.Value->GetOutdatedTypes(OutdatedShaderTypes, OutdatedShaderPipelineTypes, OutdatedFactoryTypes);
+	}
+}
+
+void FGlobalShaderMap::SaveShaderStableKeys(EShaderPlatform TargetShaderPlatform)
+{
+	FStableShaderKeyAndValue SaveKeyVal;
+	for (const auto& It : SectionMap)
+	{
+		It.Value->SaveShaderStableKeys(TargetShaderPlatform, SaveKeyVal);
+	}
+}
+
+#endif // WITH_EDITOR
+
+bool FGlobalShaderMap::IsEmpty() const
+{
+	for (const auto& It : SectionMap)
+	{
+		if (!It.Value->GetContent()->IsEmpty())
+		{
+			return false;
+		}
+	}
+	return true;
+}
+
+void FGlobalShaderMap::Empty()
+{
+	for (const auto& It : SectionMap)
+	{
+		It.Value->GetMutableContent()->Empty();
+	}
+}
+
+FShader* FGlobalShaderMap::FindOrAddShader(const FShaderType* ShaderType, FShader* Shader)
+{
+	const FHashedName HashedFilename(ShaderType->GetShaderFilename());
+	FGlobalShaderMapSection*& Section = SectionMap.FindOrAdd(HashedFilename);
+	if (!Section)
+	{
+		Section = new FGlobalShaderMapSection(Platform, HashedFilename);
+	}
+	return Section->GetMutableContent()->FindOrAddShader(Shader);
+}
+
+FShaderPipeline* FGlobalShaderMap::FindOrAddShaderPipeline(const FShaderPipelineType* ShaderPipelineType, FShaderPipeline* ShaderPipeline)
+{
+	FGlobalShaderMapSection*& Section = SectionMap.FindOrAdd(ShaderPipelineType->GetHashedPrimaryShaderFilename());
+	if (!Section)
+	{
+		Section = new FGlobalShaderMapSection(Platform, ShaderPipelineType->GetHashedPrimaryShaderFilename());
+	}
+	return Section->GetMutableContent()->FindOrAddShaderPipeline(ShaderPipeline);
+}
+
+void FGlobalShaderMap::RemoveShaderTypePermutaion(const FShaderType* Type, int32 PermutationId)
+{
+	const FHashedName HashedFilename(Type->GetShaderFilename());
+	FGlobalShaderMapSection** Section = SectionMap.Find(HashedFilename);
+	if (Section)
+	{
+		(*Section)->GetMutableContent()->RemoveShaderTypePermutaion(Type, PermutationId);
+	}
+}
+
+void FGlobalShaderMap::RemoveShaderPipelineType(const FShaderPipelineType* ShaderPipelineType)
+{
+	FGlobalShaderMapSection** Section = SectionMap.Find(ShaderPipelineType->GetHashedPrimaryShaderFilename());
+	if (Section)
+	{
+		(*Section)->GetMutableContent()->RemoveShaderPipelineType(ShaderPipelineType);
+	}
+}
+
+void FGlobalShaderMap::AddSection(FGlobalShaderMapSection* InSection)
+{
+	const FGlobalShaderMapContent* Content = InSection->GetContent();
+	const FHashedName& HashedFilename = Content->HashedSourceFilename;
+
+	SectionMap.Add(HashedFilename, InSection);
+}
+
+FGlobalShaderMapSection* FGlobalShaderMap::FindSection(const FHashedName& HashedShaderFilename)
+{
+	FGlobalShaderMapSection* const* Section = SectionMap.Find(HashedShaderFilename);
+	return Section ? *Section : nullptr;
+}
+
+FGlobalShaderMapSection* FGlobalShaderMap::FindOrAddSection(const FShaderType* ShaderType)
+{
+	const FHashedName HashedFilename(ShaderType->GetShaderFilename());
+	FGlobalShaderMapSection* Section = FindSection(HashedFilename);
+	if(!Section)
+	{
+		Section = new FGlobalShaderMapSection(Platform, HashedFilename);
+		AddSection(Section);
+	}
+	return Section;
+}
+
+void FGlobalShaderMap::SaveToGlobalArchive(FArchive& Ar)
+{
+	int32 NumSections = SectionMap.Num();
+	Ar << NumSections;
+
+	for (const auto& It : SectionMap)
+	{
+		It.Value->Serialize(Ar);
+	}
+}
+
+void FGlobalShaderMap::LoadFromGlobalArchive(FArchive& Ar)
+{
+	int32 NumSections = 0;
+	Ar << NumSections;
+
+	for (int32 i = 0; i < NumSections; ++i)
+	{
+		FGlobalShaderMapSection* Section = FGlobalShaderMapSection::CreateFromArchive(Ar);
+		AddSection(Section);
+	}
+}
diff --git a/Engine/Source/Runtime/RenderCore/Private/LongGPUTask.cpp b/Engine/Source/Runtime/RenderCore/Private/LongGPUTask.cpp
--- a/Engine/Source/Runtime/RenderCore/Private/LongGPUTask.cpp
+++ b/Engine/Source/Runtime/RenderCore/Private/LongGPUTask.cpp
@@ -33,8 +33,8 @@ void IssueScalableLongGPUTask(FRHICommandListImmediate& RHICmdList, int32 NumIte
 		TShaderMapRef<FLongGPUTaskPS> PixelShader(ShaderMap);
 
 		GraphicsPSOInit.BoundShaderState.VertexDeclarationRHI = GetVertexDeclarationFVector4();
-		GraphicsPSOInit.BoundShaderState.VertexShaderRHI = VertexShader->GetVertexShader();
-		GraphicsPSOInit.BoundShaderState.PixelShaderRHI = PixelShader->GetPixelShader();
+		GraphicsPSOInit.BoundShaderState.VertexShaderRHI = VertexShader.GetVertexShader();
+		GraphicsPSOInit.BoundShaderState.PixelShaderRHI = PixelShader.GetPixelShader();
 		GraphicsPSOInit.PrimitiveType = PT_TriangleStrip;
 
 		SetGraphicsPipelineState(RHICmdList, GraphicsPSOInit);
diff --git a/Engine/Source/Runtime/RenderCore/Private/MediaShaders.cpp b/Engine/Source/Runtime/RenderCore/Private/MediaShaders.cpp
--- a/Engine/Source/Runtime/RenderCore/Private/MediaShaders.cpp
+++ b/Engine/Source/Runtime/RenderCore/Private/MediaShaders.cpp
@@ -112,7 +112,7 @@ void FAYUVConvertPS::SetParameters(FRHICommandList& CommandList, TRefCountPtr<FR
 	}
 
 	TUniformBufferRef<FAYUVConvertUB> Data = TUniformBufferRef<FAYUVConvertUB>::CreateUniformBufferImmediate(UB, UniformBuffer_SingleFrame);
-	SetUniformBufferParameter(CommandList, GetPixelShader(), GetUniformBufferParameter<FAYUVConvertUB>(), Data);
+	SetUniformBufferParameter(CommandList, CommandList.GetBoundPixelShader(), GetUniformBufferParameter<FAYUVConvertUB>(), Data);
 }
 
 
@@ -141,7 +141,7 @@ void FBMPConvertPS::SetParameters(FRHICommandList& CommandList, TRefCountPtr<FRH
 	}
 
 	TUniformBufferRef<FBMPConvertUB> Data = TUniformBufferRef<FBMPConvertUB>::CreateUniformBufferImmediate(UB, UniformBuffer_SingleFrame);
-	SetUniformBufferParameter(CommandList, GetPixelShader(), GetUniformBufferParameter<FBMPConvertUB>(), Data);
+	SetUniformBufferParameter(CommandList, CommandList.GetBoundPixelShader(), GetUniformBufferParameter<FBMPConvertUB>(), Data);
 }
 
 
@@ -176,7 +176,7 @@ void FNV12ConvertPS::SetParameters(FRHICommandList& CommandList, TRefCountPtr<FR
 	}
 
 	TUniformBufferRef<FNV12ConvertUB> Data = TUniformBufferRef<FNV12ConvertUB>::CreateUniformBufferImmediate(UB, UniformBuffer_SingleFrame);
-	SetUniformBufferParameter(CommandList, GetPixelShader(), GetUniformBufferParameter<FNV12ConvertUB>(), Data);
+	SetUniformBufferParameter(CommandList, CommandList.GetBoundPixelShader(), GetUniformBufferParameter<FNV12ConvertUB>(), Data);
 }
 
 
@@ -211,7 +211,7 @@ void FNV21ConvertPS::SetParameters(FRHICommandList& CommandList, TRefCountPtr<FR
 	}
 
 	TUniformBufferRef<FNV21ConvertUB> Data = TUniformBufferRef<FNV21ConvertUB>::CreateUniformBufferImmediate(UB, UniformBuffer_SingleFrame);
-	SetUniformBufferParameter(CommandList, GetPixelShader(), GetUniformBufferParameter<FNV21ConvertUB>(), Data);
+	SetUniformBufferParameter(CommandList, CommandList.GetBoundPixelShader(), GetUniformBufferParameter<FNV21ConvertUB>(), Data);
 }
 
 
@@ -240,7 +240,7 @@ void FRGBConvertPS::SetParameters(FRHICommandList& CommandList, TRefCountPtr<FRH
 	}
 
 	TUniformBufferRef<FRGBConvertUB> Data = TUniformBufferRef<FRGBConvertUB>::CreateUniformBufferImmediate(UB, UniformBuffer_SingleFrame);
-	SetUniformBufferParameter(CommandList, GetPixelShader(), GetUniformBufferParameter<FRGBConvertUB>(), Data);
+	SetUniformBufferParameter(CommandList, CommandList.GetBoundPixelShader(), GetUniformBufferParameter<FRGBConvertUB>(), Data);
 }
 
 
@@ -277,7 +277,7 @@ void FYCbCrConvertPS::SetParameters(FRHICommandList& CommandList, TRefCountPtr<F
 	}
 
 	TUniformBufferRef<FYCbCrConvertUB> Data = TUniformBufferRef<FYCbCrConvertUB>::CreateUniformBufferImmediate(UB, UniformBuffer_SingleFrame);
-	SetUniformBufferParameter(CommandList, GetPixelShader(), GetUniformBufferParameter<FYCbCrConvertUB>(), Data);	
+	SetUniformBufferParameter(CommandList, CommandList.GetBoundPixelShader(), GetUniformBufferParameter<FYCbCrConvertUB>(), Data);	
 }
 
 
@@ -310,7 +310,7 @@ void FUYVYConvertPS::SetParameters(FRHICommandList& CommandList, TRefCountPtr<FR
 	}
 
 	TUniformBufferRef<FUYVYConvertUB> Data = TUniformBufferRef<FUYVYConvertUB>::CreateUniformBufferImmediate(UB, UniformBuffer_SingleFrame);
-	SetUniformBufferParameter(CommandList, GetPixelShader(), GetUniformBufferParameter<FUYVYConvertUB>(), Data);
+	SetUniformBufferParameter(CommandList, CommandList.GetBoundPixelShader(), GetUniformBufferParameter<FUYVYConvertUB>(), Data);
 }
 
 
@@ -349,7 +349,7 @@ void FYUVConvertPS::SetParameters(FRHICommandList& CommandList, TRefCountPtr<FRH
 	}
 
 	TUniformBufferRef<FYUVConvertUB> Data = TUniformBufferRef<FYUVConvertUB>::CreateUniformBufferImmediate(UB, UniformBuffer_SingleFrame);
-	SetUniformBufferParameter(CommandList, GetPixelShader(), GetUniformBufferParameter<FYUVConvertUB>(), Data);
+	SetUniformBufferParameter(CommandList, CommandList.GetBoundPixelShader(), GetUniformBufferParameter<FYUVConvertUB>(), Data);
 }
 
 
@@ -380,7 +380,7 @@ void FYUVv210ConvertPS::SetParameters(FRHICommandList& CommandList, TRefCountPtr
 	}
 
 	TUniformBufferRef<FYUVv210ConvertUB> Data = TUniformBufferRef<FYUVv210ConvertUB>::CreateUniformBufferImmediate(UB, UniformBuffer_SingleFrame);
-	SetUniformBufferParameter(CommandList, GetPixelShader(), GetUniformBufferParameter<FYUVv210ConvertUB>(), Data);
+	SetUniformBufferParameter(CommandList, CommandList.GetBoundPixelShader(), GetUniformBufferParameter<FYUVv210ConvertUB>(), Data);
 }
 
 
@@ -415,7 +415,7 @@ void FYUY2ConvertPS::SetParameters(FRHICommandList& CommandList, TRefCountPtr<FR
 	}
 
 	TUniformBufferRef<FYUY2ConvertUB> Data = TUniformBufferRef<FYUY2ConvertUB>::CreateUniformBufferImmediate(UB, UniformBuffer_SingleFrame);
-	SetUniformBufferParameter(CommandList, GetPixelShader(), GetUniformBufferParameter<FYUY2ConvertUB>(), Data);
+	SetUniformBufferParameter(CommandList, CommandList.GetBoundPixelShader(), GetUniformBufferParameter<FYUY2ConvertUB>(), Data);
 }
 
 
@@ -448,7 +448,7 @@ void FYVYUConvertPS::SetParameters(FRHICommandList& CommandList, TRefCountPtr<FR
 	}
 
 	TUniformBufferRef<FYVYUConvertUB> Data = TUniformBufferRef<FYVYUConvertUB>::CreateUniformBufferImmediate(UB, UniformBuffer_SingleFrame);
-	SetUniformBufferParameter(CommandList, GetPixelShader(), GetUniformBufferParameter<FYVYUConvertUB>(), Data);
+	SetUniformBufferParameter(CommandList, CommandList.GetBoundPixelShader(), GetUniformBufferParameter<FYVYUConvertUB>(), Data);
 }
 
 
@@ -479,7 +479,7 @@ void FRGB8toUYVY8ConvertPS::SetParameters(FRHICommandList& CommandList, TRefCoun
 	}
 
 	TUniformBufferRef<FRGB8toUYVY8ConvertUB> Data = TUniformBufferRef<FRGB8toUYVY8ConvertUB>::CreateUniformBufferImmediate(UB, UniformBuffer_SingleFrame);
-	SetUniformBufferParameter(CommandList, GetPixelShader(), GetUniformBufferParameter<FRGB8toUYVY8ConvertUB>(), Data);
+	SetUniformBufferParameter(CommandList, CommandList.GetBoundPixelShader(), GetUniformBufferParameter<FRGB8toUYVY8ConvertUB>(), Data);
 }
 
 
@@ -510,7 +510,7 @@ void FRGB10toYUVv210ConvertPS::SetParameters(FRHICommandList& CommandList, TRefC
 	}
 
 	TUniformBufferRef<FRGB10toYUVv210ConvertUB> Data = TUniformBufferRef<FRGB10toYUVv210ConvertUB>::CreateUniformBufferImmediate(UB, UniformBuffer_SingleFrame);
-	SetUniformBufferParameter(CommandList, GetPixelShader(), GetUniformBufferParameter<FRGB10toYUVv210ConvertUB>(), Data);
+	SetUniformBufferParameter(CommandList, CommandList.GetBoundPixelShader(), GetUniformBufferParameter<FRGB10toYUVv210ConvertUB>(), Data);
 }
 
 
@@ -535,7 +535,7 @@ void FInvertAlphaPS::SetParameters(FRHICommandList& CommandList, TRefCountPtr<FR
 	}
 
 	TUniformBufferRef<FInvertAlphaUB> Data = TUniformBufferRef<FInvertAlphaUB>::CreateUniformBufferImmediate(UB, UniformBuffer_SingleFrame);
-	SetUniformBufferParameter(CommandList, GetPixelShader(), GetUniformBufferParameter<FInvertAlphaUB>(), Data);
+	SetUniformBufferParameter(CommandList, CommandList.GetBoundPixelShader(), GetUniformBufferParameter<FInvertAlphaUB>(), Data);
 }
 
 
@@ -560,7 +560,7 @@ void FSetAlphaOnePS::SetParameters(FRHICommandList& CommandList, TRefCountPtr<FR
 	}
 
 	TUniformBufferRef<FSetAlphaOneUB> Data = TUniformBufferRef<FSetAlphaOneUB>::CreateUniformBufferImmediate(UB, UniformBuffer_SingleFrame);
-	SetUniformBufferParameter(CommandList, GetPixelShader(), GetUniformBufferParameter<FSetAlphaOneUB>(), Data);
+	SetUniformBufferParameter(CommandList, CommandList.GetBoundPixelShader(), GetUniformBufferParameter<FSetAlphaOneUB>(), Data);
 }
 
 
@@ -589,5 +589,5 @@ void FReadTextureExternalPS::SetParameters(FRHICommandList& CommandList, FTextur
 	}
 
 	TUniformBufferRef<FReadTextureExternalUB> Data = TUniformBufferRef<FReadTextureExternalUB>::CreateUniformBufferImmediate(UB, UniformBuffer_SingleFrame);
-	SetUniformBufferParameter(CommandList, GetPixelShader(), GetUniformBufferParameter<FReadTextureExternalUB>(), Data);
+	SetUniformBufferParameter(CommandList, CommandList.GetBoundPixelShader(), GetUniformBufferParameter<FReadTextureExternalUB>(), Data);
 }
diff --git a/Engine/Source/Runtime/RenderCore/Private/OneColorShader.cpp b/Engine/Source/Runtime/RenderCore/Private/OneColorShader.cpp
--- a/Engine/Source/Runtime/RenderCore/Private/OneColorShader.cpp
+++ b/Engine/Source/Runtime/RenderCore/Private/OneColorShader.cpp
@@ -14,26 +14,21 @@ void FOneColorPS::SetColors(FRHICommandList& RHICmdList, const FLinearColor* Col
 	check(NumColors <= MaxSimultaneousRenderTargets);
 
 	auto& ClearUBParam = GetUniformBufferParameter<FClearShaderUB>();
-	if (ClearUBParam.IsInitialized())
+	if (ClearUBParam.IsBound())
 	{
-		if (ClearUBParam.IsBound())
+		FClearShaderUB ClearData;
+		FMemory::Memzero(ClearData.DrawColorMRT);
+		for (int32 i = 0; i < NumColors; ++i)
 		{
-			FClearShaderUB ClearData;
-			FMemory::Memzero(ClearData.DrawColorMRT);			
-			for (int32 i = 0; i < NumColors; ++i)
-			{
-				ClearData.DrawColorMRT[i].X = Colors[i].R;
-				ClearData.DrawColorMRT[i].Y = Colors[i].G;
-				ClearData.DrawColorMRT[i].Z = Colors[i].B;
-				ClearData.DrawColorMRT[i].W = Colors[i].A;
-			}
-
-			FLocalUniformBuffer LocalUB = TUniformBufferRef<FClearShaderUB>::CreateLocalUniformBuffer(RHICmdList, ClearData, UniformBuffer_SingleFrame);	
-			RHICmdList.SetLocalShaderUniformBuffer(GetPixelShader(), ClearUBParam.GetBaseIndex(), LocalUB);
+			ClearData.DrawColorMRT[i].X = Colors[i].R;
+			ClearData.DrawColorMRT[i].Y = Colors[i].G;
+			ClearData.DrawColorMRT[i].Z = Colors[i].B;
+			ClearData.DrawColorMRT[i].W = Colors[i].A;
 		}
-	}
 
-	
+		FLocalUniformBuffer LocalUB = TUniformBufferRef<FClearShaderUB>::CreateLocalUniformBuffer(RHICmdList, ClearData, UniformBuffer_SingleFrame);
+		RHICmdList.SetLocalShaderUniformBuffer(RHICmdList.GetBoundPixelShader(), ClearUBParam.GetBaseIndex(), LocalUB);
+	}
 }
 
 // #define avoids a lot of code duplication
diff --git a/Engine/Source/Runtime/RenderCore/Private/PixelShaderUtils.cpp b/Engine/Source/Runtime/RenderCore/Private/PixelShaderUtils.cpp
--- a/Engine/Source/Runtime/RenderCore/Private/PixelShaderUtils.cpp
+++ b/Engine/Source/Runtime/RenderCore/Private/PixelShaderUtils.cpp
@@ -52,8 +52,8 @@ void FPixelShaderUtils::DrawFullscreenQuad(FRHICommandList& RHICmdList, uint32 I
 // static
 void FPixelShaderUtils::InitFullscreenPipelineState(
 	FRHICommandList& RHICmdList,
-	const TShaderMap<FGlobalShaderType>* GlobalShaderMap,
-	const FShader* PixelShader,
+	const FGlobalShaderMap* GlobalShaderMap,
+	const TShaderRef<FShader>& PixelShader,
 	FGraphicsPipelineStateInitializer& GraphicsPSOInit)
 {
 	TShaderMapRef<FScreenVertexShaderVS> VertexShader(GlobalShaderMap);
@@ -64,8 +64,8 @@ void FPixelShaderUtils::InitFullscreenPipelineState(
 	GraphicsPSOInit.DepthStencilState = TStaticDepthStencilState<false, CF_Always>::GetRHI();
 
 	GraphicsPSOInit.BoundShaderState.VertexDeclarationRHI = GFilterVertexDeclaration.VertexDeclarationRHI;
-	GraphicsPSOInit.BoundShaderState.VertexShaderRHI = GETSAFERHISHADER_VERTEX(*VertexShader);
-	GraphicsPSOInit.BoundShaderState.PixelShaderRHI = GETSAFERHISHADER_PIXEL(PixelShader);
+	GraphicsPSOInit.BoundShaderState.VertexShaderRHI = VertexShader.GetVertexShader();
+	GraphicsPSOInit.BoundShaderState.PixelShaderRHI = PixelShader.GetPixelShader();
 	GraphicsPSOInit.PrimitiveType = PT_TriangleList;
 }
 
diff --git a/Engine/Source/Runtime/RenderCore/Private/RenderGraphUtils.cpp b/Engine/Source/Runtime/RenderCore/Private/RenderGraphUtils.cpp
--- a/Engine/Source/Runtime/RenderCore/Private/RenderGraphUtils.cpp
+++ b/Engine/Source/Runtime/RenderCore/Private/RenderGraphUtils.cpp
@@ -477,7 +477,6 @@ void AddClearDepthStencilPass(
 		DrawClearQuad(RHICmdList, false, FLinearColor(), bClearDepth, Depth, bClearStencil, Stencil);
 	});
 }
-
 class FClearUAVUIntCS : public FGlobalShader
 {
 	DECLARE_GLOBAL_SHADER(FClearUAVUIntCS)
@@ -510,7 +509,7 @@ public:
 
 IMPLEMENT_GLOBAL_SHADER(FClearUAVUIntCS, "/Engine/Private/Tools/ClearUAV.usf", "ClearUAVUIntCS", SF_Compute);
 
-void FComputeShaderUtils::ClearUAV(FRDGBuilder& GraphBuilder, TShaderMap<FGlobalShaderType>* ShaderMap, FRDGBufferUAVRef UAV, uint32 ClearValue)
+void FComputeShaderUtils::ClearUAV(FRDGBuilder& GraphBuilder, FGlobalShaderMap* ShaderMap, FRDGBufferUAVRef UAV, uint32 ClearValue)
 {
 	FClearUAVUIntCS::FParameters* PassParameters = GraphBuilder.AllocParameters<FClearUAVUIntCS::FParameters>();
 	PassParameters->UAV = UAV;
@@ -561,7 +560,7 @@ public:
 
 IMPLEMENT_GLOBAL_SHADER(FClearUAVFloatCS, "/Engine/Private/Tools/ClearUAV.usf", "ClearUAVFloatCS", SF_Compute);
 
-void FComputeShaderUtils::ClearUAV(FRDGBuilder& GraphBuilder, TShaderMap<FGlobalShaderType>* ShaderMap, FRDGBufferUAVRef UAV, FVector4 ClearValue)
+void FComputeShaderUtils::ClearUAV(FRDGBuilder& GraphBuilder, FGlobalShaderMap* ShaderMap, FRDGBufferUAVRef UAV, FVector4 ClearValue)
 {
 	FClearUAVFloatCS::FParameters* PassParameters = GraphBuilder.AllocParameters<FClearUAVFloatCS::FParameters>();
 	PassParameters->UAVFloat = UAV;
