Make D3D12 pipeline state control the lifetime of its RHI resources.

diff --git a/Engine/Source/Runtime/D3D12RHI/Private/D3D12PipelineState.h b/Engine/Source/Runtime/D3D12RHI/Private/D3D12PipelineState.h
--- a/Engine/Source/Runtime/D3D12RHI/Private/D3D12PipelineState.h
+++ b/Engine/Source/Runtime/D3D12RHI/Private/D3D12PipelineState.h
@@ -340,7 +340,7 @@ struct FD3D12GraphicsPipelineState : public FRHIGraphicsPipelineState
 	explicit FD3D12GraphicsPipelineState(const FGraphicsPipelineStateInitializer& Initializer, const FD3D12RootSignature* InRootSignature, FD3D12PipelineState* InPipelineState);
 	~FD3D12GraphicsPipelineState();
 
-	const FGraphicsPipelineStateInitializer PipelineStateInitializer;
+	FGraphicsPipelineStateInitializer PipelineStateInitializer;
 	const FD3D12RootSignature* RootSignature;
 	uint16 StreamStrides[MaxVertexElementCount];
 	bool bShaderNeedsGlobalConstantBuffer[SF_NumStandardFrequencies];
diff --git a/Engine/Source/Runtime/D3D12RHI/Private/D3D12Pipelinestate.cpp b/Engine/Source/Runtime/D3D12RHI/Private/D3D12Pipelinestate.cpp
--- a/Engine/Source/Runtime/D3D12RHI/Private/D3D12Pipelinestate.cpp
+++ b/Engine/Source/Runtime/D3D12RHI/Private/D3D12Pipelinestate.cpp
@@ -327,6 +327,9 @@ FD3D12GraphicsPipelineState::FD3D12GraphicsPipelineState(
 	, RootSignature(InRootSignature)
 	, PipelineState(InPipelineState)
 {
+	// hold on to bound RHI resources
+	PipelineStateInitializer.BoundShaderState.AddRefResources();
+
 	if (Initializer.BoundShaderState.VertexDeclarationRHI)
 		FMemory::Memcpy(StreamStrides, ((FD3D12VertexDeclaration*) Initializer.BoundShaderState.VertexDeclarationRHI)->StreamStrides, sizeof(StreamStrides));
 	else
@@ -349,6 +352,9 @@ FD3D12GraphicsPipelineState::~FD3D12GraphicsPipelineState()
 	delete PipelineState;
 	PipelineState = nullptr;
 #endif // D3D12_USE_DERIVED_PSO
+
+	// release bound RHI resources
+	PipelineStateInitializer.BoundShaderState.ReleaseResources();
 }
 
 FD3D12ComputePipelineState::~FD3D12ComputePipelineState()
diff --git a/Engine/Source/Runtime/RHI/Private/PipelineFileCache.cpp b/Engine/Source/Runtime/RHI/Private/PipelineFileCache.cpp
--- a/Engine/Source/Runtime/RHI/Private/PipelineFileCache.cpp
+++ b/Engine/Source/Runtime/RHI/Private/PipelineFileCache.cpp
@@ -1005,6 +1005,7 @@ FPipelineCacheFileFormatPSO::FPipelineCacheFileFormatPSO()
 	FMemory::Memset(&PSO.GraphicsDesc, 0, sizeof(GraphicsDescriptor));
 	
 	check (Init.BoundShaderState.VertexDeclarationRHI);
+	check (Init.BoundShaderState.VertexDeclarationRHI->IsValid());
 	{
 		bOK &= Init.BoundShaderState.VertexDeclarationRHI->GetInitializer(PSO.GraphicsDesc.VertexDescriptor);
 		check(bOK);
diff --git a/Engine/Source/Runtime/RHI/Public/RHIResources.h b/Engine/Source/Runtime/RHI/Public/RHIResources.h
--- a/Engine/Source/Runtime/RHI/Public/RHIResources.h
+++ b/Engine/Source/Runtime/RHI/Public/RHIResources.h
@@ -1828,6 +1828,64 @@ struct FBoundShaderStateInput
 	{
 	}
 
+	void AddRefResources()
+	{
+		check(VertexDeclarationRHI);
+		VertexDeclarationRHI->AddRef();
+
+		check(VertexShaderRHI);
+		VertexShaderRHI->AddRef();
+
+		if (HullShaderRHI)
+		{
+			HullShaderRHI->AddRef();
+		}
+
+		if (DomainShaderRHI)
+		{
+			DomainShaderRHI->AddRef();
+		}
+
+		if (PixelShaderRHI)
+		{
+			PixelShaderRHI->AddRef();
+		}
+
+		if (GeometryShaderRHI)
+		{
+			GeometryShaderRHI->AddRef();
+		}
+	}
+
+	void ReleaseResources()
+	{
+		check(VertexDeclarationRHI);
+		VertexDeclarationRHI->Release();
+
+		check(VertexShaderRHI);
+		VertexShaderRHI->Release();
+
+		if (HullShaderRHI)
+		{
+			HullShaderRHI->Release();
+		}
+
+		if (DomainShaderRHI)
+		{
+			DomainShaderRHI->Release();
+		}
+
+		if (PixelShaderRHI)
+		{
+			PixelShaderRHI->Release();
+		}
+
+		if (GeometryShaderRHI)
+		{
+			GeometryShaderRHI->Release();
+		}
+	}
+
 	FRHIVertexDeclaration* VertexDeclarationRHI = nullptr;
 	FRHIVertexShader* VertexShaderRHI = nullptr;
 	FRHIHullShader* HullShaderRHI = nullptr;
