DRPS - Compile fix
#rnx

diff --git a/Engine/Source/Runtime/RHI/Private/RHI.cpp b/Engine/Source/Runtime/RHI/Private/RHI.cpp
--- a/Engine/Source/Runtime/RHI/Private/RHI.cpp
+++ b/Engine/Source/Runtime/RHI/Private/RHI.cpp
@@ -471,7 +471,6 @@ bool GHardwareHiddenSurfaceRemoval = false;
 bool GRHISupportsAsyncTextureCreation = false;
 bool GRHISupportsQuadTopology = false;
 bool GRHISupportsRectTopology = false;
-bool GRHISupportsDepthUAV = false;
 bool GRHISupportsAtomicUInt64 = false;
 bool GRHISupportsResummarizeHTile = false;
 bool GRHISupportsExplicitHTile = false;
