Fix FrameFence in bypass mode

diff --git a/Engine/Source/Runtime/RHI/Private/RHICommandList.cpp b/Engine/Source/Runtime/RHI/Private/RHICommandList.cpp
--- a/Engine/Source/Runtime/RHI/Private/RHICommandList.cpp
+++ b/Engine/Source/Runtime/RHI/Private/RHICommandList.cpp
@@ -1618,12 +1618,12 @@ void FRHICommandList::EndFrame()
 	if (Bypass())
 	{
 		GetContext().RHIEndFrame();
+		GDynamicRHI->RHIAdvanceFrameFence();
 		return;
 	}
 
 	ALLOC_COMMAND(FRHICommandEndFrame)();
-
-	FRHICommandListExecutor::GetImmediateCommandList().AdvanceFrameFence();
+	GDynamicRHI->RHIAdvanceFrameFence();
 
 	if (!IsRunningRHIInSeparateThread())
 	{
diff --git a/Engine/Source/Runtime/RHI/Public/DynamicRHI.h b/Engine/Source/Runtime/RHI/Public/DynamicRHI.h
--- a/Engine/Source/Runtime/RHI/Public/DynamicRHI.h
+++ b/Engine/Source/Runtime/RHI/Public/DynamicRHI.h
@@ -1419,11 +1419,6 @@ FORCEINLINE FTexture2DRHIRef RHIGetViewportBackBuffer(FRHIViewport* Viewport)
 	return GDynamicRHI->RHIGetViewportBackBuffer(Viewport);
 }
 
-FORCEINLINE void RHIAdvanceFrameFence()
-{
-	return GDynamicRHI->RHIAdvanceFrameFence();
-}
-
 FORCEINLINE FTexture2DRHIRef RHIGetStencilTexture(FRHITexture* SourceTextureRHI)
 {
 	return GDynamicRHI->RHIGetStencilTexture(SourceTextureRHI);
diff --git a/Engine/Source/Runtime/RHI/Public/RHICommandList.h b/Engine/Source/Runtime/RHI/Public/RHICommandList.h
--- a/Engine/Source/Runtime/RHI/Public/RHICommandList.h
+++ b/Engine/Source/Runtime/RHI/Public/RHICommandList.h
@@ -4517,11 +4517,6 @@ public:
 		return RHIGetViewportBackBuffer(Viewport);
 	}
 	
-	FORCEINLINE void AdvanceFrameFence()
-	{
-		RHIAdvanceFrameFence();
-	}
-
 	FORCEINLINE void AdvanceFrameForGetViewportBackBuffer(FRHIViewport* Viewport)
 	{
 		return RHIAdvanceFrameForGetViewportBackBuffer(Viewport);
