Merged command list fix from Dev-Barriers stream.

diff --git a/Engine/Source/Runtime/RHI/Private/RHICommandList.cpp b/Engine/Source/Runtime/RHI/Private/RHICommandList.cpp
--- a/Engine/Source/Runtime/RHI/Private/RHICommandList.cpp
+++ b/Engine/Source/Runtime/RHI/Private/RHICommandList.cpp
@@ -242,6 +242,7 @@ void FRHIAsyncComputeCommandListImmediate::ImmediateDispatch(FRHIAsyncComputeCom
 			static_assert(sizeof(FRHICommandList) == sizeof(FRHIAsyncComputeCommandListImmediate), "We are memswapping FRHICommandList and FRHICommandListImmediate; they need to be swappable.");
 			check(RHIComputeCmdList.IsImmediateAsyncCompute());
 			SwapCmdList->ExchangeCmdList(RHIComputeCmdList);
+			RHIComputeCmdList.CopyContext(*SwapCmdList);
 			RHIComputeCmdList.GPUMask = SwapCmdList->GPUMask;
 			// NB: InitialGPUMask set to GPUMask since exchanging the list
 			// is equivalent to a Reset.
@@ -483,6 +484,7 @@ void FRHICommandListExecutor::ExecuteInner(FRHICommandListBase& CmdList)
 				// we should make command lists virtual and transfer ownership rather than this devious approach
 				static_assert(sizeof(FRHICommandList) == sizeof(FRHICommandListImmediate), "We are memswapping FRHICommandList and FRHICommandListImmediate; they need to be swappable.");
 				SwapCmdList->ExchangeCmdList(CmdList);
+				CmdList.CopyContext(*SwapCmdList);
 				CmdList.GPUMask = SwapCmdList->GPUMask;
 				// NB: InitialGPUMask set to GPUMask since exchanging the list
                 // is equivalent to a Reset.
@@ -841,18 +843,11 @@ FRHICommandListBase::FRHICommandListBase(FRHIGPUMask InGPUMask)
 	, ComputeContext(nullptr)
 	, MemManager(0)
 	, bAsyncPSOCompileAllowed(true)
-	, bRecursive(false)
 	, GPUMask(InGPUMask)
 	, InitialGPUMask(InGPUMask)
 {
 	GRHICommandList.OutstandingCmdListCount.Increment();
 	Reset();
-
-	if (!bRecursive)
-	{
-		Context = GDynamicRHI ? RHIGetDefaultContext() : nullptr;
-		ComputeContext = Context;
-	}
 }
 
 FRHICommandListBase::~FRHICommandListBase()
diff --git a/Engine/Source/Runtime/RHI/Public/RHICommandList.h b/Engine/Source/Runtime/RHI/Public/RHICommandList.h
--- a/Engine/Source/Runtime/RHI/Public/RHICommandList.h
+++ b/Engine/Source/Runtime/RHI/Public/RHICommandList.h
@@ -513,10 +513,10 @@ public:
 		return *Context;
 	}
 
-	void SetComputeContext(IRHIComputeContext* InContext)
+	void SetComputeContext(IRHIComputeContext* InComputeContext)
 	{
-		check(InContext);
-		ComputeContext = InContext;
+		check(Context == nullptr);
+		ComputeContext = InComputeContext;
 	}
 
 	IRHIComputeContext& GetComputeContext()
@@ -562,7 +562,6 @@ protected:
 	FRHICommandListBase(FRHIGPUMask InGPUMask);
 
 	bool bAsyncPSOCompileAllowed;
-	bool bRecursive;
 	FRHIGPUMask GPUMask;
 	// GPUMask that was set at the time the command list was last Reset. We set
     // this mask on the command contexts immediately before executing the
@@ -4673,7 +4672,6 @@ public:
 	{
 		SetContext(Context);
 		bAsyncPSOCompileAllowed = false;
-		bRecursive = true;
 	}
 };
 
