Compile fix

diff --git a/Engine/Source/Runtime/RHI/Public/RHIResources.h b/Engine/Source/Runtime/RHI/Public/RHIResources.h
--- a/Engine/Source/Runtime/RHI/Public/RHIResources.h
+++ b/Engine/Source/Runtime/RHI/Public/RHIResources.h
@@ -201,7 +201,7 @@ public:
 	FString ShaderName;
 	FORCEINLINE const TCHAR* GetShaderName() const { return *ShaderName; }
 #else
-	FORCEINLINE const TCHAR* GetShaderName() const { TEXT(""); }
+	FORCEINLINE const TCHAR* GetShaderName() const { return TEXT(""); }
 #endif
 
 	explicit FRHIShader(EShaderFrequency InFrequency)
