Added a thread safe counter to the FRHIGPUFence, in order to prevent Poll() from returning wrong values while Enqueue() hasn't been processed yet.

diff --git a/Engine/Source/Runtime/RHI/Private/RHICommandList.cpp b/Engine/Source/Runtime/RHI/Private/RHICommandList.cpp
--- a/Engine/Source/Runtime/RHI/Private/RHICommandList.cpp
+++ b/Engine/Source/Runtime/RHI/Private/RHICommandList.cpp
@@ -2218,7 +2218,7 @@ void FDynamicRHI::RHIUnlockStagingBuffer(FRHIStagingBuffer* StagingBuffer)
 void* FDynamicRHI::LockStagingBuffer_RenderThread(class FRHICommandListImmediate& RHICmdList, FRHIStagingBuffer* StagingBuffer, FRHIGPUFence* Fence, uint32 Offset, uint32 SizeRHI)
 {
 	check(IsInRenderingThread());
-	if (Fence == nullptr || !Fence->Poll())
+	if (!Fence || !Fence->Poll() || Fence->NumPendingWriteCommands.GetValue() != 0)
 	{
 		QUICK_SCOPE_CYCLE_COUNTER(STAT_FDynamicRHI_LockStagingBuffer_Flush);
 		RHICmdList.ImmediateFlush(EImmediateFlushType::FlushRHIThread);
@@ -2561,7 +2561,7 @@ void FDynamicRHI::RHIUnlockTextureCubeFace_RenderThread(class FRHICommandListImm
 
 void FDynamicRHI::RHIMapStagingSurface_RenderThread(class FRHICommandListImmediate& RHICmdList, FRHITexture* Texture, FRHIGPUFence* Fence, void*& OutData, int32& OutWidth, int32& OutHeight)
 {
-	if (Fence == nullptr || !Fence->Poll())
+	if (Fence == nullptr || !Fence->Poll() || Fence->NumPendingWriteCommands.GetValue() != 0)
 	{
 		QUICK_SCOPE_CYCLE_COUNTER(STAT_RHIMETHOD_MapStagingSurface_Flush);
 		RHICmdList.ImmediateFlush(EImmediateFlushType::FlushRHIThread);
diff --git a/Engine/Source/Runtime/RHI/Public/RHICommandList.h b/Engine/Source/Runtime/RHI/Public/RHICommandList.h
--- a/Engine/Source/Runtime/RHI/Public/RHICommandList.h
+++ b/Engine/Source/Runtime/RHI/Public/RHICommandList.h
@@ -1643,6 +1643,10 @@ struct FRHICommandWriteGPUFence final : public FRHICommand<FRHICommandWriteGPUFe
 	FORCEINLINE_DEBUGGABLE FRHICommandWriteGPUFence(FRHIGPUFence* InFence)
 		: Fence(InFence)
 	{
+		if (Fence)
+		{
+			Fence->NumPendingWriteCommands.Increment();
+		}
 	}
 	RHI_API void Execute(FRHICommandListBase& CmdList);
 };
diff --git a/Engine/Source/Runtime/RHI/Public/RHICommandListCommandExecutes.inl b/Engine/Source/Runtime/RHI/Public/RHICommandListCommandExecutes.inl
--- a/Engine/Source/Runtime/RHI/Public/RHICommandListCommandExecutes.inl
+++ b/Engine/Source/Runtime/RHI/Public/RHICommandListCommandExecutes.inl
@@ -441,6 +441,10 @@ void FRHICommandWriteGPUFence::Execute(FRHICommandListBase& CmdList)
 {
 	RHISTAT(WriteGPUFence);
 	INTERNAL_DECORATOR_COMPUTE(RHIWriteGPUFence)(Fence);
+	if (Fence)
+	{
+		Fence->NumPendingWriteCommands.Decrement();
+	}
 }
 
 void FRHICommandSetGlobalUniformBuffers::Execute(FRHICommandListBase& CmdList)
diff --git a/Engine/Source/Runtime/RHI/Public/RHIGPUReadback.h b/Engine/Source/Runtime/RHI/Public/RHIGPUReadback.h
--- a/Engine/Source/Runtime/RHI/Public/RHIGPUReadback.h
+++ b/Engine/Source/Runtime/RHI/Public/RHIGPUReadback.h
@@ -29,7 +29,7 @@ public:
 	/** Indicates if the data is in place and ready to be read. */
 	FORCEINLINE bool IsReady()
 	{
-		return !Fence || Fence->Poll();
+		return !Fence || (Fence->NumPendingWriteCommands.GetValue() == 0 && Fence->Poll());
 	}
 
 	/** Indicates if the data is in place and ready to be read on a subset of GPUs. */
diff --git a/Engine/Source/Runtime/RHI/Public/RHIResources.h b/Engine/Source/Runtime/RHI/Public/RHIResources.h
--- a/Engine/Source/Runtime/RHI/Public/RHIResources.h
+++ b/Engine/Source/Runtime/RHI/Public/RHIResources.h
@@ -997,6 +997,8 @@ public:
 
 	const FName& GetFName() const { return FenceName; }
 
+	FThreadSafeCounter NumPendingWriteCommands;
+
 protected:
 	FName FenceName;
 };
