UE4.25 - Do not allow NULL views of RHI resources. Helps track down issues earlier in the frame.
#rnx

diff --git a/Engine/Source/Runtime/RHI/Public/RHICommandList.h b/Engine/Source/Runtime/RHI/Public/RHICommandList.h
--- a/Engine/Source/Runtime/RHI/Public/RHICommandList.h
+++ b/Engine/Source/Runtime/RHI/Public/RHICommandList.h
@@ -4002,42 +4002,49 @@ public:
 	FORCEINLINE FUnorderedAccessViewRHIRef CreateUnorderedAccessView(FRHIStructuredBuffer* StructuredBuffer, bool bUseUAVCounter, bool bAppendBuffer)
 	{
 		LLM_SCOPE(ELLMTag::RHIMisc);
+		checkf(StructuredBuffer, TEXT("Can't create a view off a null resource!"));
 		return GDynamicRHI->RHICreateUnorderedAccessView_RenderThread(*this, StructuredBuffer, bUseUAVCounter, bAppendBuffer);
 	}
 	
 	FORCEINLINE FUnorderedAccessViewRHIRef CreateUnorderedAccessView(FRHITexture* Texture, uint32 MipLevel)
 	{
 		LLM_SCOPE(ELLMTag::RHIMisc);
+		checkf(Texture, TEXT("Can't create a view off a null resource!"));
 		return GDynamicRHI->RHICreateUnorderedAccessView_RenderThread(*this, Texture, MipLevel);
 	}
 	
 	FORCEINLINE FUnorderedAccessViewRHIRef CreateUnorderedAccessView(FRHITexture* Texture, uint32 MipLevel, uint8 Format)
 	{
 		LLM_SCOPE(ELLMTag::RHIMisc);
+		checkf(Texture, TEXT("Can't create a view off a null resource!"));
 		return GDynamicRHI->RHICreateUnorderedAccessView_RenderThread(*this, Texture, MipLevel, Format);
 	}
 
 	FORCEINLINE FUnorderedAccessViewRHIRef CreateUnorderedAccessView(FRHIVertexBuffer* VertexBuffer, uint8 Format)
 	{
 		LLM_SCOPE(ELLMTag::RHIMisc);
+		checkf(VertexBuffer, TEXT("Can't create a view off a null resource!"));
 		return GDynamicRHI->RHICreateUnorderedAccessView_RenderThread(*this, VertexBuffer, Format);
 	}
 
 	FORCEINLINE FUnorderedAccessViewRHIRef CreateUnorderedAccessView(FRHIIndexBuffer* IndexBuffer, uint8 Format)
 	{
 		LLM_SCOPE(ELLMTag::RHIMisc);
+		checkf(IndexBuffer, TEXT("Can't create a view off a null resource!"));
 		return GDynamicRHI->RHICreateUnorderedAccessView_RenderThread(*this, IndexBuffer, Format);
 	}
 
 	FORCEINLINE FShaderResourceViewRHIRef CreateShaderResourceView(FRHIStructuredBuffer* StructuredBuffer)
 	{
 		LLM_SCOPE(ELLMTag::RHIMisc);
+		checkf(StructuredBuffer, TEXT("Can't create a view off a null resource!"));
 		return GDynamicRHI->RHICreateShaderResourceView_RenderThread(*this, StructuredBuffer);
 	}
 	
 	FORCEINLINE FShaderResourceViewRHIRef CreateShaderResourceView(FRHIVertexBuffer* VertexBuffer, uint32 Stride, uint8 Format)
 	{
 		LLM_SCOPE(ELLMTag::RHIMisc);
+		checkf(VertexBuffer, TEXT("Can't create a view off a null resource!"));
 		return GDynamicRHI->CreateShaderResourceView_RenderThread(*this, VertexBuffer, Stride, Format);
 	}
 	
@@ -4050,6 +4057,7 @@ public:
 	FORCEINLINE FShaderResourceViewRHIRef CreateShaderResourceView(FRHIIndexBuffer* Buffer)
 	{
 		LLM_SCOPE(ELLMTag::RHIMisc);
+		checkf(Buffer, TEXT("Can't create a view off a null resource!"));
 		return GDynamicRHI->CreateShaderResourceView_RenderThread(*this, Buffer);
 	}
 	
@@ -4145,12 +4153,14 @@ public:
 	FORCEINLINE FShaderResourceViewRHIRef CreateShaderResourceView(FRHITexture* Texture, const FRHITextureSRVCreateInfo& CreateInfo)
 	{
 		LLM_SCOPE(ELLMTag::RHIMisc);
+		checkf(Texture, TEXT("Can't create a view off a null resource!"));
 		return GDynamicRHI->RHICreateShaderResourceView_RenderThread(*this, Texture, CreateInfo);
 	}
 
 	FORCEINLINE FShaderResourceViewRHIRef CreateShaderResourceView(FRHITexture* Texture, uint8 MipLevel)
 	{
 		LLM_SCOPE(ELLMTag::RHIMisc);
+		checkf(Texture, TEXT("Can't create a view off a null resource!"));
 		const FRHITextureSRVCreateInfo CreateInfo(MipLevel, 1, Texture->GetFormat());
 		return GDynamicRHI->RHICreateShaderResourceView_RenderThread(*this, Texture, CreateInfo);
 	}
@@ -4158,6 +4168,7 @@ public:
 	FORCEINLINE FShaderResourceViewRHIRef CreateShaderResourceView(FRHITexture* Texture, uint8 MipLevel, uint8 NumMipLevels, uint8 Format)
 	{
 		LLM_SCOPE(ELLMTag::RHIMisc);
+		checkf(Texture, TEXT("Can't create a view off a null resource!"));
 		const FRHITextureSRVCreateInfo CreateInfo(MipLevel, NumMipLevels, Format);
 		return GDynamicRHI->RHICreateShaderResourceView_RenderThread(*this, Texture, CreateInfo);
 	}
@@ -4165,12 +4176,14 @@ public:
 	FORCEINLINE FShaderResourceViewRHIRef CreateShaderResourceViewWriteMask(FRHITexture2D* Texture2DRHI)
 	{
 		LLM_SCOPE(ELLMTag::RHIMisc);
+		checkf(Texture2DRHI, TEXT("Can't create a view off a null resource!"));
 		return GDynamicRHI->RHICreateShaderResourceViewWriteMask_RenderThread(*this, Texture2DRHI);
 	}
 
 	FORCEINLINE FShaderResourceViewRHIRef CreateShaderResourceViewFMask(FRHITexture2D* Texture2DRHI)
 	{
 		LLM_SCOPE(ELLMTag::RHIMisc);
+		checkf(Texture2DRHI, TEXT("Can't create a view off a null resource!"));
 		return GDynamicRHI->RHICreateShaderResourceViewFMask_RenderThread(*this, Texture2DRHI);
 	}
 
