Avoid creating tons of Blend, DepthStencil, Rasterizer states when caching PSOs

diff --git a/Engine/Source/Runtime/Apple/MetalRHI/Private/MetalState.cpp b/Engine/Source/Runtime/Apple/MetalRHI/Private/MetalState.cpp
--- a/Engine/Source/Runtime/Apple/MetalRHI/Private/MetalState.cpp
+++ b/Engine/Source/Runtime/Apple/MetalRHI/Private/MetalState.cpp
@@ -196,70 +196,6 @@ static uint32 GetTypeHash(const FSamplerStateInitializerRHI& Initializer)
 	return Hash;
 }
 
-static bool operator==(const FDepthStencilStateInitializerRHI& Left, const FDepthStencilStateInitializerRHI& Right)
-{
-	bool bSame = Left.bEnableDepthWrite == Right.bEnableDepthWrite
-	&& Left.DepthTest == Right.DepthTest
-	&& Left.FrontFaceStencilTest == Right.FrontFaceStencilTest
-	&& Left.FrontFaceStencilFailStencilOp == Right.FrontFaceStencilFailStencilOp
-	&& Left.FrontFaceDepthFailStencilOp == Right.FrontFaceDepthFailStencilOp
-	&& Left.FrontFacePassStencilOp == Right.FrontFacePassStencilOp
-	&& Left.bEnableBackFaceStencil == Right.bEnableBackFaceStencil
-	&& Left.BackFaceStencilTest == Right.BackFaceStencilTest
-	&& Left.BackFaceStencilFailStencilOp == Right.BackFaceStencilFailStencilOp
-	&& Left.BackFaceDepthFailStencilOp == Right.BackFaceDepthFailStencilOp
-	&& Left.BackFacePassStencilOp == Right.BackFacePassStencilOp
-	&& Left.StencilReadMask == Right.StencilReadMask
-	&& Left.StencilWriteMask == Right.StencilWriteMask;
-	
-	return bSame;
-}
-
-static uint32 GetTypeHash(const FDepthStencilStateInitializerRHI& Initializer)
-{
-	uint32 Hash = FCrc::MemCrc32(&Initializer.bEnableDepthWrite, sizeof(Initializer.bEnableDepthWrite));
-	Hash = FCrc::MemCrc32(&Initializer.DepthTest, sizeof(Initializer.DepthTest), Hash);
-	Hash = FCrc::MemCrc32(&Initializer.FrontFaceStencilTest, sizeof(Initializer.FrontFaceStencilTest), Hash);
-	Hash = FCrc::MemCrc32(&Initializer.FrontFaceStencilFailStencilOp, sizeof(Initializer.FrontFaceStencilFailStencilOp), Hash);
-	Hash = FCrc::MemCrc32(&Initializer.FrontFaceDepthFailStencilOp, sizeof(Initializer.FrontFaceDepthFailStencilOp), Hash);
-	Hash = FCrc::MemCrc32(&Initializer.FrontFacePassStencilOp, sizeof(Initializer.FrontFacePassStencilOp), Hash);
-	Hash = FCrc::MemCrc32(&Initializer.bEnableBackFaceStencil, sizeof(Initializer.bEnableBackFaceStencil), Hash);
-	Hash = FCrc::MemCrc32(&Initializer.BackFaceStencilTest, sizeof(Initializer.BackFaceStencilTest), Hash);
-	Hash = FCrc::MemCrc32(&Initializer.BackFaceStencilFailStencilOp, sizeof(Initializer.BackFaceStencilFailStencilOp), Hash);
-	Hash = FCrc::MemCrc32(&Initializer.BackFaceDepthFailStencilOp, sizeof(Initializer.BackFaceDepthFailStencilOp), Hash);
-	Hash = FCrc::MemCrc32(&Initializer.BackFacePassStencilOp, sizeof(Initializer.BackFacePassStencilOp), Hash);
-	Hash = FCrc::MemCrc32(&Initializer.StencilReadMask, sizeof(Initializer.StencilReadMask), Hash);
-	Hash = FCrc::MemCrc32(&Initializer.StencilWriteMask, sizeof(Initializer.StencilWriteMask), Hash);
-	
-	return Hash;
-}
-
-static bool operator==(const FBlendStateInitializerRHI::FRenderTarget& Left, const FBlendStateInitializerRHI::FRenderTarget& Right)
-{
-	bool bSame = Left.ColorBlendOp == Right.ColorBlendOp
-	&& Left.ColorDestBlend == Right.ColorDestBlend
-	&& Left.ColorSrcBlend == Right.ColorSrcBlend
-	&& Left.AlphaBlendOp == Right.AlphaBlendOp
-	&& Left.AlphaDestBlend == Right.AlphaDestBlend
-	&& Left.AlphaSrcBlend == Right.AlphaSrcBlend
-	&& Left.ColorWriteMask == Right.ColorWriteMask;
-	
-	return bSame;
-}
-
-static uint32 GetTypeHash(const FBlendStateInitializerRHI::FRenderTarget& Initializer)
-{
-	uint32 Hash = FCrc::MemCrc32(&Initializer.ColorBlendOp, sizeof(Initializer.ColorBlendOp));
-	Hash = FCrc::MemCrc32(&Initializer.ColorDestBlend, sizeof(Initializer.ColorDestBlend), Hash);
-	Hash = FCrc::MemCrc32(&Initializer.ColorSrcBlend, sizeof(Initializer.ColorSrcBlend), Hash);
-	Hash = FCrc::MemCrc32(&Initializer.AlphaBlendOp, sizeof(Initializer.AlphaBlendOp), Hash);
-	Hash = FCrc::MemCrc32(&Initializer.AlphaDestBlend, sizeof(Initializer.AlphaDestBlend), Hash);
-	Hash = FCrc::MemCrc32(&Initializer.AlphaSrcBlend, sizeof(Initializer.AlphaSrcBlend), Hash);
-	Hash = FCrc::MemCrc32(&Initializer.ColorWriteMask, sizeof(Initializer.ColorWriteMask), Hash);
-	
-	return Hash;
-}
-
 template <typename InitializerType, typename StateType>
 class FMetalStateObjectCache
 {
diff --git a/Engine/Source/Runtime/RHI/Private/RHI.cpp b/Engine/Source/Runtime/RHI/Private/RHI.cpp
--- a/Engine/Source/Runtime/RHI/Private/RHI.cpp
+++ b/Engine/Source/Runtime/RHI/Private/RHI.cpp
@@ -169,6 +169,68 @@ void FVertexElement::FromString(const FStringView& InSrc)
 	check(Parts.GetData() + PartCount == PartIt);
 }
 
+uint32 GetTypeHash(const FRasterizerStateInitializerRHI& Initializer)
+{
+	uint32 Hash = GetTypeHash(Initializer.FillMode);
+	Hash = HashCombine(Hash, GetTypeHash(Initializer.CullMode));
+	Hash = HashCombine(Hash, GetTypeHash(Initializer.DepthBias));
+	Hash = HashCombine(Hash, GetTypeHash(Initializer.SlopeScaleDepthBias));
+	Hash = HashCombine(Hash, GetTypeHash(Initializer.bAllowMSAA));
+	Hash = HashCombine(Hash, GetTypeHash(Initializer.bEnableLineAA));
+	return Hash;
+}
+	
+bool operator== (const FRasterizerStateInitializerRHI& A, const FRasterizerStateInitializerRHI& B)
+{
+	bool bSame = 
+		A.FillMode == B.FillMode && 
+		A.CullMode == B.CullMode && 
+		A.DepthBias == B.DepthBias && 
+		A.SlopeScaleDepthBias == B.SlopeScaleDepthBias && 
+		A.bAllowMSAA == B.bAllowMSAA && 
+		A.bEnableLineAA == B.bEnableLineAA;
+	return bSame;
+}
+
+uint32 GetTypeHash(const FDepthStencilStateInitializerRHI& Initializer)
+{
+	uint32 Hash = GetTypeHash(Initializer.bEnableDepthWrite);
+	Hash = HashCombine(Hash, GetTypeHash(Initializer.DepthTest));
+	Hash = HashCombine(Hash, GetTypeHash(Initializer.bEnableFrontFaceStencil));
+	Hash = HashCombine(Hash, GetTypeHash(Initializer.FrontFaceStencilTest));
+	Hash = HashCombine(Hash, GetTypeHash(Initializer.FrontFaceStencilFailStencilOp));
+	Hash = HashCombine(Hash, GetTypeHash(Initializer.FrontFaceDepthFailStencilOp));
+	Hash = HashCombine(Hash, GetTypeHash(Initializer.FrontFacePassStencilOp));
+	Hash = HashCombine(Hash, GetTypeHash(Initializer.bEnableBackFaceStencil));
+	Hash = HashCombine(Hash, GetTypeHash(Initializer.BackFaceStencilTest));
+	Hash = HashCombine(Hash, GetTypeHash(Initializer.BackFaceStencilFailStencilOp));
+	Hash = HashCombine(Hash, GetTypeHash(Initializer.BackFaceDepthFailStencilOp));
+	Hash = HashCombine(Hash, GetTypeHash(Initializer.BackFacePassStencilOp));
+	Hash = HashCombine(Hash, GetTypeHash(Initializer.StencilReadMask));
+	Hash = HashCombine(Hash, GetTypeHash(Initializer.StencilWriteMask));
+	return Hash;
+}
+
+bool operator== (const FDepthStencilStateInitializerRHI& A, const FDepthStencilStateInitializerRHI& B)
+{
+	bool bSame = 
+		A.bEnableDepthWrite == B.bEnableDepthWrite && 
+		A.DepthTest == B.DepthTest && 
+		A.bEnableFrontFaceStencil == B.bEnableFrontFaceStencil && 
+		A.FrontFaceStencilTest == B.FrontFaceStencilTest && 
+		A.FrontFaceStencilFailStencilOp == B.FrontFaceStencilFailStencilOp && 
+		A.FrontFaceDepthFailStencilOp == B.FrontFaceDepthFailStencilOp && 
+		A.FrontFacePassStencilOp == B.FrontFacePassStencilOp && 
+		A.bEnableBackFaceStencil == B.bEnableBackFaceStencil && 
+		A.BackFaceStencilTest == B.BackFaceStencilTest && 
+		A.BackFaceStencilFailStencilOp == B.BackFaceStencilFailStencilOp && 
+		A.BackFaceDepthFailStencilOp == B.BackFaceDepthFailStencilOp && 
+		A.BackFacePassStencilOp == B.BackFacePassStencilOp && 
+		A.StencilReadMask == B.StencilReadMask && 
+		A.StencilWriteMask == B.StencilWriteMask;
+	return bSame;
+}
+
 FString FDepthStencilStateInitializerRHI::ToString() const
 {
 	return
@@ -269,6 +331,27 @@ void FBlendStateInitializerRHI::FromString(const FStringView& InSrc)
 	check(Parts.GetData() + PartCount == PartIt);
 }
 
+uint32 GetTypeHash(const FBlendStateInitializerRHI& Initializer)
+{
+	uint32 Hash = GetTypeHash(Initializer.bUseIndependentRenderTargetBlendStates);
+	for (int32 i = 0; i < MaxSimultaneousRenderTargets; ++i)
+	{
+		Hash = HashCombine(Hash, GetTypeHash(Initializer.RenderTargets[i]));
+	}
+	
+	return Hash;
+}
+
+bool operator== (const FBlendStateInitializerRHI& A, const FBlendStateInitializerRHI& B)
+{
+	bool bSame = A.bUseIndependentRenderTargetBlendStates == B.bUseIndependentRenderTargetBlendStates;
+	for (int32 i = 0; i < MaxSimultaneousRenderTargets && bSame; ++i)
+	{
+		bSame = bSame && A.RenderTargets[i] == B.RenderTargets[i];
+	}
+	return bSame;
+}
+
 
 FString FBlendStateInitializerRHI::FRenderTarget::ToString() const
 {
@@ -308,6 +391,31 @@ void FBlendStateInitializerRHI::FRenderTarget::FromString(TArrayView<const FStri
 	LexFromString((uint8&)ColorWriteMask, *PartIt++);
 }
 
+uint32 GetTypeHash(const FBlendStateInitializerRHI::FRenderTarget& Initializer)
+{
+	uint32 Hash = GetTypeHash(Initializer.ColorBlendOp);
+	Hash = HashCombine(Hash, GetTypeHash(Initializer.ColorDestBlend));
+	Hash = HashCombine(Hash, GetTypeHash(Initializer.ColorSrcBlend));
+	Hash = HashCombine(Hash, GetTypeHash(Initializer.AlphaBlendOp));
+	Hash = HashCombine(Hash, GetTypeHash(Initializer.AlphaDestBlend));
+	Hash = HashCombine(Hash, GetTypeHash(Initializer.AlphaSrcBlend));
+	Hash = HashCombine(Hash, GetTypeHash(Initializer.ColorWriteMask));
+	return Hash;
+}
+
+bool operator==(const FBlendStateInitializerRHI::FRenderTarget& A, const FBlendStateInitializerRHI::FRenderTarget& B)
+{
+	bool bSame = 
+		A.ColorBlendOp == B.ColorBlendOp && 
+		A.ColorDestBlend == B.ColorDestBlend && 
+		A.ColorSrcBlend == B.ColorSrcBlend && 
+		A.AlphaBlendOp == B.AlphaBlendOp && 
+		A.AlphaDestBlend == B.AlphaDestBlend && 
+		A.AlphaSrcBlend == B.AlphaSrcBlend && 
+		A.ColorWriteMask == B.ColorWriteMask;
+	return bSame;
+}
+
 bool FRHIResource::Bypass()
 {
 	return GRHICommandList.Bypass();
diff --git a/Engine/Source/Runtime/RHI/Public/RHI.h b/Engine/Source/Runtime/RHI/Public/RHI.h
--- a/Engine/Source/Runtime/RHI/Public/RHI.h
+++ b/Engine/Source/Runtime/RHI/Public/RHI.h
@@ -869,6 +869,9 @@ struct FRasterizerStateInitializerRHI
 		Ar << RasterizerStateInitializer.bEnableLineAA;
 		return Ar;
 	}
+
+	RHI_API friend uint32 GetTypeHash(const FRasterizerStateInitializerRHI& Initializer);
+	RHI_API friend bool operator== (const FRasterizerStateInitializerRHI& A, const FRasterizerStateInitializerRHI& B);
 };
 
 struct FDepthStencilStateInitializerRHI
@@ -939,6 +942,10 @@ struct FDepthStencilStateInitializerRHI
 		Ar << DepthStencilStateInitializer.StencilWriteMask;
 		return Ar;
 	}
+
+	RHI_API friend uint32 GetTypeHash(const FDepthStencilStateInitializerRHI& Initializer);
+	RHI_API friend bool operator== (const FDepthStencilStateInitializerRHI& A, const FDepthStencilStateInitializerRHI& B);
+	
 	RHI_API FString ToString() const;
 	RHI_API void FromString(const FString& Src);
 	RHI_API void FromString(const FStringView& Src);
@@ -991,6 +998,10 @@ public:
 			Ar << RenderTarget.ColorWriteMask;
 			return Ar;
 		}
+		
+		friend uint32 GetTypeHash(const FRenderTarget& RenderTarget);
+		friend bool operator== (const FRenderTarget& A, const FRenderTarget& B);
+
 		RHI_API FString ToString() const;
 		RHI_API void FromString(const TArray<FString>& Parts, int32 Index);
 		RHI_API void FromString(TArrayView<const FStringView> Parts);
@@ -1025,6 +1036,10 @@ public:
 		Ar << BlendStateInitializer.bUseIndependentRenderTargetBlendStates;
 		return Ar;
 	}
+
+	RHI_API friend uint32 GetTypeHash(const FBlendStateInitializerRHI& Initializer);
+	RHI_API friend bool operator== (const FBlendStateInitializerRHI& A, const FBlendStateInitializerRHI& B);
+	
 	RHI_API FString ToString() const;
 	RHI_API void FromString(const FString& Src);
 	RHI_API void FromString(const FStringView& Src);
diff --git a/Engine/Source/Runtime/RenderCore/Private/ShaderPipelineCache.cpp b/Engine/Source/Runtime/RenderCore/Private/ShaderPipelineCache.cpp
--- a/Engine/Source/Runtime/RenderCore/Private/ShaderPipelineCache.cpp
+++ b/Engine/Source/Runtime/RenderCore/Private/ShaderPipelineCache.cpp
@@ -628,13 +628,13 @@ bool FShaderPipelineCache::Precompile(FRHICommandListImmediate& RHICmdList, ESha
 				GraphicsInitializer.BoundShaderState.GeometryShaderRHI = GeometryShader;
 			}
 	#endif
-			auto BlendState = RHICmdList.CreateBlendState(PSO.GraphicsDesc.BlendState);
+			auto BlendState = GetOrCreateBlendState(PSO.GraphicsDesc.BlendState);
 			GraphicsInitializer.BlendState = BlendState;
 			
-			auto RasterState = RHICmdList.CreateRasterizerState(PSO.GraphicsDesc.RasterizerState);
+			auto RasterState = GetOrCreateRasterizerState(PSO.GraphicsDesc.RasterizerState);
 			GraphicsInitializer.RasterizerState = RasterState;
 			
-			auto DepthState = RHICmdList.CreateDepthStencilState(PSO.GraphicsDesc.DepthStencilState);
+			auto DepthState = GetOrCreateDepthStencilState(PSO.GraphicsDesc.DepthStencilState);
 			GraphicsInitializer.DepthStencilState = DepthState;
 			
 			for (uint32 i = 0; i < MaxSimultaneousRenderTargets; ++i)
@@ -1513,6 +1513,28 @@ void FShaderPipelineCache::Close(bool bShuttingDown)
 	bOpened = false;
 
 	FPipelineFileCache::ClosePipelineFileCache();
+
+
+	//
+	// Clean up cached RHI resources
+	//
+	for (auto Pair : BlendStateCache)
+	{
+		Pair.Value->Release();
+	}
+	BlendStateCache.Empty();
+	
+	for (auto Pair : RasterizerStateCache)
+	{
+		Pair.Value->Release();
+	}
+	RasterizerStateCache.Empty();
+
+	for (auto Pair : DepthStencilStateCache)
+	{
+		Pair.Value->Release();
+	}
+	DepthStencilStateCache.Empty();
 }
 
 void FShaderPipelineCache::OnShaderLibraryStateChanged(ELibraryState State, EShaderPlatform Platform, FString const& Name)
@@ -1556,3 +1578,52 @@ void FShaderPipelineCache::OnShaderLibraryStateChanged(ELibraryState State, ESha
     FPlatformAtomics::InterlockedExchange(&TotalWaitingTasks, Count);
 	UE_LOG(LogRHI, Display, TEXT("Opened pipeline cache after state change and enqueued %d of %d tasks for precompile."), Count, OrderedCompileTasks.Num());
 }
+
+
+FRHIBlendState* FShaderPipelineCache::GetOrCreateBlendState(const FBlendStateInitializerRHI& Initializer)
+{
+	FRHIBlendState** Found = BlendStateCache.Find(Initializer);
+	if (Found)
+	{
+		return *Found;
+	}
+
+	FBlendStateRHIRef NewState = RHICreateBlendState(Initializer);
+
+	// Add an extra reference so we don't have TRefCountPtr in the maps
+	NewState->AddRef();
+	BlendStateCache.Add(Initializer, NewState);
+	return NewState;
+}
+
+FRHIRasterizerState* FShaderPipelineCache::GetOrCreateRasterizerState(const FRasterizerStateInitializerRHI& Initializer)
+{
+	FRHIRasterizerState** Found = RasterizerStateCache.Find(Initializer);
+	if (Found)
+	{
+		return *Found;
+	}
+
+	FRasterizerStateRHIRef NewState = RHICreateRasterizerState(Initializer);
+
+	// Add an extra reference so we don't have TRefCountPtr in the maps
+	NewState->AddRef();
+	RasterizerStateCache.Add(Initializer, NewState);
+	return NewState;
+}
+
+FRHIDepthStencilState* FShaderPipelineCache::GetOrCreateDepthStencilState(const FDepthStencilStateInitializerRHI& Initializer)
+{
+	FRHIDepthStencilState** Found = DepthStencilStateCache.Find(Initializer);
+	if (Found)
+	{
+		return *Found;
+	}
+
+	FDepthStencilStateRHIRef NewState = RHICreateDepthStencilState(Initializer);
+
+	// Add an extra reference so we don't have TRefCountPtr in the maps
+	NewState->AddRef();
+	DepthStencilStateCache.Add(Initializer, NewState);
+	return NewState;
+}
diff --git a/Engine/Source/Runtime/RenderCore/Public/ShaderPipelineCache.h b/Engine/Source/Runtime/RenderCore/Public/ShaderPipelineCache.h
--- a/Engine/Source/Runtime/RenderCore/Public/ShaderPipelineCache.h
+++ b/Engine/Source/Runtime/RenderCore/Public/ShaderPipelineCache.h
@@ -223,6 +223,10 @@ private:
     
     void OnShaderLibraryStateChanged(ELibraryState State, EShaderPlatform Platform, FString const& Name);
 
+	FRHIBlendState* GetOrCreateBlendState(const FBlendStateInitializerRHI& Initializer);
+	FRHIRasterizerState* GetOrCreateRasterizerState(const FRasterizerStateInitializerRHI& Initializer);
+	FRHIDepthStencilState* GetOrCreateDepthStencilState(const FDepthStencilStateInitializerRHI& Initializer);
+
 private:
 	static FShaderPipelineCache* ShaderPipelineCache;
 	TArray<CompileJob> ReadTasks;
@@ -268,4 +272,8 @@ private:
 	TSet<uint64> CompletedMasks;
 	float TotalPrecompileWallTime;
 	int64 TotalPrecompileTasks;
+
+	TMap<FBlendStateInitializerRHI, FRHIBlendState*> BlendStateCache;
+	TMap<FRasterizerStateInitializerRHI, FRHIRasterizerState*> RasterizerStateCache;
+	TMap<FDepthStencilStateInitializerRHI, FRHIDepthStencilState*> DepthStencilStateCache;
 };
