Change IsMobilePlatform() and GetIsMobile() to take EShaderPlatforms so they can be called with none native shader platforms on platform extension platforms

diff --git a/Engine/Source/Runtime/RHI/Public/RHIDefinitions.h b/Engine/Source/Runtime/RHI/Public/RHIDefinitions.h
--- a/Engine/Source/Runtime/RHI/Public/RHIDefinitions.h
+++ b/Engine/Source/Runtime/RHI/Public/RHIDefinitions.h
@@ -386,7 +386,7 @@ public:
 		return Infos[Platform].MaxFeatureLevel;
 	}
 
-	static FORCEINLINE_DEBUGGABLE const bool GetIsMobile(const FStaticShaderPlatform Platform)
+	static FORCEINLINE_DEBUGGABLE const bool GetIsMobile(const EShaderPlatform Platform)
 	{
 		return Infos[Platform].bIsMobile;
 	}
@@ -1348,7 +1348,7 @@ inline bool IsPCPlatform(const FStaticShaderPlatform Platform)
 }
 
 /** Whether the shader platform corresponds to the ES2/ES3.1 feature level. */
-inline bool IsMobilePlatform(const FStaticShaderPlatform Platform)
+inline bool IsMobilePlatform(const EShaderPlatform Platform)
 {
 	return 
 		Platform == SP_METAL || Platform == SP_METAL_MACES3_1 || Platform == SP_METAL_TVOS
