Fix Mac link errors related to CL 11574745

diff --git a/Engine/Source/Runtime/RHI/Public/RHI.h b/Engine/Source/Runtime/RHI/Public/RHI.h
--- a/Engine/Source/Runtime/RHI/Public/RHI.h
+++ b/Engine/Source/Runtime/RHI/Public/RHI.h
@@ -999,9 +999,6 @@ public:
 			return Ar;
 		}
 		
-		friend uint32 GetTypeHash(const FRenderTarget& RenderTarget);
-		friend bool operator== (const FRenderTarget& A, const FRenderTarget& B);
-
 		RHI_API FString ToString() const;
 		RHI_API void FromString(const TArray<FString>& Parts, int32 Index);
 		RHI_API void FromString(TArrayView<const FStringView> Parts);
@@ -1037,6 +1034,9 @@ public:
 		return Ar;
 	}
 
+	RHI_API friend uint32 GetTypeHash(const FBlendStateInitializerRHI::FRenderTarget& RenderTarget);
+	RHI_API friend bool operator== (const FBlendStateInitializerRHI::FRenderTarget& A, const FBlendStateInitializerRHI::FRenderTarget& B);
+	
 	RHI_API friend uint32 GetTypeHash(const FBlendStateInitializerRHI& Initializer);
 	RHI_API friend bool operator== (const FBlendStateInitializerRHI& A, const FBlendStateInitializerRHI& B);
 	
