Fix Clang CIS

diff --git a/Engine/Source/Runtime/RHI/Public/RHIContext.h b/Engine/Source/Runtime/RHI/Public/RHIContext.h
--- a/Engine/Source/Runtime/RHI/Public/RHIContext.h
+++ b/Engine/Source/Runtime/RHI/Public/RHIContext.h
@@ -411,6 +411,11 @@ public:
 
 	virtual void RHIEndRenderQuery(FRHIRenderQuery* RenderQuery) = 0;
 
+	virtual void RHICalibrateTimers()
+	{
+		/* empty default implementation */
+	}
+
 	virtual void RHICalibrateTimers(FRHITimestampCalibrationQuery* CalibrationQuery)
 	{
 		/* empty default implementation */
