Added Android Desktop Forward Rendering (Experimental)

diff --git a/Engine/Config/BaseEngine.ini b/Engine/Config/BaseEngine.ini
--- a/Engine/Config/BaseEngine.ini
+++ b/Engine/Config/BaseEngine.ini
@@ -2047,6 +2047,7 @@ bBuildForX8664=false
 bBuildForES2=false
 bBuildForES31=true
 bSupportsVulkan=false
+bSupportsVulkanSM5=false
 bDetectVulkanByDefault=true
 bSplitIntoSeparateApks=false
 bPackageDataInsideApk=false
diff --git a/Engine/Source/Developer/Android/AndroidPlatformEditor/Private/AndroidPlatformEditorModule.cpp b/Engine/Source/Developer/Android/AndroidPlatformEditor/Private/AndroidPlatformEditorModule.cpp
--- a/Engine/Source/Developer/Android/AndroidPlatformEditor/Private/AndroidPlatformEditorModule.cpp
+++ b/Engine/Source/Developer/Android/AndroidPlatformEditor/Private/AndroidPlatformEditorModule.cpp
@@ -95,6 +95,15 @@ class FAndroidPlatformEditorModule
 					AndroidMaterialQualitySettings
 				);
 			}
+			{
+				static FName NAME_SF_VULKAN_SM5_ANDROID(TEXT("SF_VULKAN_SM5_ANDROID"));
+				UShaderPlatformQualitySettings* AndroidMaterialQualitySettings = UMaterialShaderQualitySettings::Get()->GetShaderPlatformQualitySettings(NAME_SF_VULKAN_SM5_ANDROID);
+				SettingsModule->RegisterSettings("Project", "Platforms", "AndroidVulkanSM5Quality",
+					LOCTEXT("AndroidVulkanSM5QualitySettingsName", "Android SM5 Material Quality - Vulkan"),
+					LOCTEXT("AndroidVulkanSM5QualitySettingsDescription", "Settings for Android Vulkan SM5 material quality"),
+					AndroidMaterialQualitySettings
+				);
+			}
 
 		}
 
diff --git a/Engine/Source/Developer/Android/AndroidTargetPlatform/Private/AndroidTargetPlatform.cpp b/Engine/Source/Developer/Android/AndroidTargetPlatform/Private/AndroidTargetPlatform.cpp
--- a/Engine/Source/Developer/Android/AndroidTargetPlatform/Private/AndroidTargetPlatform.cpp
+++ b/Engine/Source/Developer/Android/AndroidTargetPlatform/Private/AndroidTargetPlatform.cpp
@@ -247,6 +247,16 @@ bool FAndroidTargetPlatform::SupportsVulkan() const
 	return bSupportsVulkan;
 }
 
+bool FAndroidTargetPlatform::SupportsVulkanSM5() const
+{
+	// default to no support for VulkanSM5
+	bool bSupportsMobileVulkanSM5 = false;
+#if WITH_ENGINE
+	GConfig->GetBool(TEXT("/Script/AndroidRuntimeSettings.AndroidRuntimeSettings"), TEXT("bSupportsVulkanSM5"), bSupportsMobileVulkanSM5, GEngineIni);
+#endif
+	return bSupportsMobileVulkanSM5;
+}
+
 bool FAndroidTargetPlatform::SupportsSoftwareOcclusion() const
 {
 	static auto* CVarMobileAllowSoftwareOcclusion = IConsoleManager::Get().FindTConsoleVariableDataInt(TEXT("r.Mobile.AllowSoftwareOcclusion"));
@@ -343,8 +353,10 @@ bool FAndroidTargetPlatform::SupportsFeature( ETargetPlatformFeatures Feature )
 			return SupportsES31() || SupportsVulkan();
 
 		case ETargetPlatformFeatures::HighQualityLightmaps:
-		case ETargetPlatformFeatures::Tessellation:
 		case ETargetPlatformFeatures::DeferredRendering:
+			return SupportsAEP() || SupportsVulkanSM5();
+
+		case ETargetPlatformFeatures::Tessellation:
 			return SupportsAEP();
 
 		case ETargetPlatformFeatures::SoftwareOcclusion:
@@ -368,10 +380,15 @@ void FAndroidTargetPlatform::GetAllPossibleShaderFormats( TArray<FName>& OutForm
 	static FName NAME_GLSL_310_ES_EXT(TEXT("GLSL_310_ES_EXT"));
 	static FName NAME_SF_VULKAN_ES31_ANDROID(TEXT("SF_VULKAN_ES31_ANDROID_NOUB"));
 	static FName NAME_GLSL_ES3_1_ANDROID(TEXT("GLSL_ES3_1_ANDROID"));
+	static FName NAME_SF_VULKAN_SM5_ANDROID(TEXT("SF_VULKAN_SM5_ANDROID"));
 
 	if (SupportsVulkan())
 	{
 		OutFormats.AddUnique(NAME_SF_VULKAN_ES31_ANDROID);
+		if (SupportsVulkanSM5())
+		{
+			OutFormats.AddUnique(NAME_SF_VULKAN_SM5_ANDROID);
+		}
 	}
 
 	if (SupportsES31())
@@ -576,7 +593,7 @@ void FAndroidTargetPlatform::GetAllTextureFormats(TArray<FName>& OutFormats) con
 
 void FAndroidTargetPlatform::GetReflectionCaptureFormats( TArray<FName>& OutFormats ) const
 {
-	if (SupportsAEP())
+	if (SupportsAEP() || SupportsVulkanSM5())
 	{
 		// use Full HDR with AEP
 		OutFormats.Add(FName(TEXT("FullHDR")));
diff --git a/Engine/Source/Developer/Android/AndroidTargetPlatform/Private/AndroidTargetPlatform.h b/Engine/Source/Developer/Android/AndroidTargetPlatform/Private/AndroidTargetPlatform.h
--- a/Engine/Source/Developer/Android/AndroidTargetPlatform/Private/AndroidTargetPlatform.h
+++ b/Engine/Source/Developer/Android/AndroidTargetPlatform/Private/AndroidTargetPlatform.h
@@ -285,6 +285,7 @@ protected:
 	bool SupportsAEP() const;
 	bool SupportsVulkan() const;
 	bool SupportsSoftwareOcclusion() const;
+	bool SupportsVulkanSM5() const;
 
 #if WITH_ENGINE
 	// Holds the Engine INI settings (for quick access).
diff --git a/Engine/Source/Developer/SharedSettingsWidgets/Private/ShaderFormatsPropertyDetails.cpp b/Engine/Source/Developer/SharedSettingsWidgets/Private/ShaderFormatsPropertyDetails.cpp
--- a/Engine/Source/Developer/SharedSettingsWidgets/Private/ShaderFormatsPropertyDetails.cpp
+++ b/Engine/Source/Developer/SharedSettingsWidgets/Private/ShaderFormatsPropertyDetails.cpp
@@ -71,6 +71,7 @@ FText FShaderFormatsPropertyDetails::GetFriendlyNameFromRHINameMac(const FString
 			break;
 		case SP_VULKAN_SM5:
 		case SP_VULKAN_SM5_LUMIN:
+		case SP_VULKAN_SM5_ANDROID:
 			FriendlyRHIName = LOCTEXT("VulkanSM5", "Vulkan (SM5)");
 			break;
 		case SP_VULKAN_PCES3_1:
diff --git a/Engine/Source/Developer/VulkanShaderFormat/Private/VulkanShaderFormat.cpp b/Engine/Source/Developer/VulkanShaderFormat/Private/VulkanShaderFormat.cpp
--- a/Engine/Source/Developer/VulkanShaderFormat/Private/VulkanShaderFormat.cpp
+++ b/Engine/Source/Developer/VulkanShaderFormat/Private/VulkanShaderFormat.cpp
@@ -19,6 +19,8 @@ static FName NAME_VULKAN_SM5_NOUB(TEXT("SF_VULKAN_SM5_NOUB"));
 static FName NAME_VULKAN_SM5(TEXT("SF_VULKAN_SM5"));
 static FName NAME_VULKAN_SM5_LUMIN(TEXT("SF_VULKAN_SM5_LUMIN"));
 static FName NAME_VULKAN_SM5_LUMIN_NOUB(TEXT("SF_VULKAN_SM5_LUMIN_NOUB"));
+static FName NAME_VULKAN_SM5_ANDROID(TEXT("SF_VULKAN_SM5_ANDROID"));
+static FName NAME_VULKAN_SM5_ANDROID_NOUB(TEXT("SF_VULKAN_SM5_ANDROID_NOUB"));
 
 class FShaderFormatVulkan : public IShaderFormat
 {
@@ -30,7 +32,8 @@ class FShaderFormatVulkan : public IShaderFormat
 
 	int32 InternalGetVersion(FName Format) const
 	{
-		if (Format == NAME_VULKAN_SM5 || Format == NAME_VULKAN_SM5_NOUB || Format == NAME_VULKAN_SM5_LUMIN || Format == NAME_VULKAN_SM5_LUMIN_NOUB)
+		if (Format == NAME_VULKAN_SM5 || Format == NAME_VULKAN_SM5_NOUB || Format == NAME_VULKAN_SM5_LUMIN ||
+			Format == NAME_VULKAN_SM5_LUMIN_NOUB || Format == NAME_VULKAN_SM5_ANDROID || Format == NAME_VULKAN_SM5_ANDROID_NOUB)
 		{
 			return UE_SHADER_VULKAN_SM5_VER;
 		}
@@ -65,6 +68,8 @@ public:
 		OutFormats.Add(NAME_VULKAN_ES3_1_LUMIN_NOUB);
 		OutFormats.Add(NAME_VULKAN_ES3_1_NOUB);
 		OutFormats.Add(NAME_VULKAN_SM5_NOUB);
+		OutFormats.Add(NAME_VULKAN_SM5_ANDROID);
+		OutFormats.Add(NAME_VULKAN_SM5_ANDROID_NOUB);
 	}
 
 	virtual void CompileShader(FName Format, const struct FShaderCompilerInput& Input, struct FShaderCompilerOutput& Output,const FString& WorkingDirectory) const
@@ -86,11 +91,11 @@ public:
 		{
 			DoCompileVulkanShader(Input, Output, WorkingDirectory, EVulkanShaderVersion::ES3_1_ANDROID_NOUB);
 		}
-		else if (Format == NAME_VULKAN_SM5_NOUB || Format == NAME_VULKAN_SM5_LUMIN_NOUB)
+		else if (Format == NAME_VULKAN_SM5_NOUB || Format == NAME_VULKAN_SM5_LUMIN_NOUB || Format == NAME_VULKAN_SM5_ANDROID_NOUB)
 		{
 			DoCompileVulkanShader(Input, Output, WorkingDirectory, EVulkanShaderVersion::SM5_NOUB);
 		}
-		else if (Format == NAME_VULKAN_SM5 || Format == NAME_VULKAN_SM5_LUMIN)
+		else if (Format == NAME_VULKAN_SM5 || Format == NAME_VULKAN_SM5_LUMIN || Format == NAME_VULKAN_SM5_ANDROID)
 		{
 			DoCompileVulkanShader(Input, Output, WorkingDirectory, EVulkanShaderVersion::SM5);
 		}
diff --git a/Engine/Source/Editor/LevelEditor/Private/LevelEditorActions.cpp b/Engine/Source/Editor/LevelEditor/Private/LevelEditorActions.cpp
--- a/Engine/Source/Editor/LevelEditor/Private/LevelEditorActions.cpp
+++ b/Engine/Source/Editor/LevelEditor/Private/LevelEditorActions.cpp
@@ -3339,6 +3339,7 @@ void FLevelEditorCommands::RegisterCommands()
 	UI_COMMAND(PreviewPlatformOverride_SM5, "Shader Model 5", "DirectX 11, OpenGL 4.3+, PS4, XB1", EUserInterfaceActionType::Check, FInputChord());
 	UI_COMMAND(PreviewPlatformOverride_AndroidGLES31, "Android ES 3.1", "Mobile preview using Android ES3.1 quality settings.", EUserInterfaceActionType::Check, FInputChord());
 	UI_COMMAND(PreviewPlatformOverride_AndroidVulkanES31, "Android Vulkan", "Mobile preview using Android Vulkan quality settings.", EUserInterfaceActionType::Check, FInputChord());
+	UI_COMMAND(PreviewPlatformOverride_AndroidVulkanSM5, "Android Vulkan SM5", "Mobile preview using Android Vulkan SM5 quality settings.", EUserInterfaceActionType::Check, FInputChord());
 	UI_COMMAND(PreviewPlatformOverride_IOSMetalES31, "iOS", "Mobile preview using iOS material quality settings.", EUserInterfaceActionType::Check, FInputChord());
 
 	UI_COMMAND( ConnectToSourceControl, "Connect to Source Control...", "Opens a dialog to connect to source control.", EUserInterfaceActionType::Button, FInputChord());
diff --git a/Engine/Source/Editor/LevelEditor/Private/LevelEditorToolBar.cpp b/Engine/Source/Editor/LevelEditor/Private/LevelEditorToolBar.cpp
--- a/Engine/Source/Editor/LevelEditor/Private/LevelEditorToolBar.cpp
+++ b/Engine/Source/Editor/LevelEditor/Private/LevelEditorToolBar.cpp
@@ -1927,6 +1927,14 @@ static void MakeShaderModelPreviewMenu( UToolMenu* InMenu )
 		Section.AddMenuEntry(FLevelEditorCommands::Get().PreviewPlatformOverride_AndroidVulkanES31);
 	}
 
+	bool bAndroidSupportsMobileVulkanSM5 = false;
+	GConfig->GetBool(TEXT("/Script/AndroidRuntimeSettings.AndroidRuntimeSettings"), TEXT("bSupportsAndroidVulkanSM5"), bAndroidSupportsVulkan, GEngineIni);
+	if (bAndroidSupportsMobileVulkanSM5)
+	{
+		Section.AddMenuEntry(FLevelEditorCommands::Get().PreviewPlatformOverride_AndroidVulkanSM5);
+	}
+
+
 	// iOS
 	bool bIOSSupportsMetal = false;
 	GConfig->GetBool(TEXT("/Script/IOSRuntimeSettings.IOSRuntimeSettings"), TEXT("bSupportsMetal"), bIOSSupportsMetal, GEngineIni);
diff --git a/Engine/Source/Editor/LevelEditor/Public/LevelEditorActions.h b/Engine/Source/Editor/LevelEditor/Public/LevelEditorActions.h
--- a/Engine/Source/Editor/LevelEditor/Public/LevelEditorActions.h
+++ b/Engine/Source/Editor/LevelEditor/Public/LevelEditorActions.h
@@ -593,6 +593,7 @@ public:
 	TSharedPtr< FUICommandInfo > PreviewPlatformOverride_SM5;
 	TSharedPtr< FUICommandInfo > PreviewPlatformOverride_AndroidGLES31;
 	TSharedPtr< FUICommandInfo > PreviewPlatformOverride_AndroidVulkanES31;
+	TSharedPtr< FUICommandInfo > PreviewPlatformOverride_AndroidVulkanSM5;
 	TSharedPtr< FUICommandInfo > PreviewPlatformOverride_IOSMetalES31;
 	
 	///**
diff --git a/Engine/Source/Editor/MaterialEditor/Private/MaterialStats.cpp b/Engine/Source/Editor/MaterialEditor/Private/MaterialStats.cpp
--- a/Engine/Source/Editor/MaterialEditor/Private/MaterialStats.cpp
+++ b/Engine/Source/Editor/MaterialEditor/Private/MaterialStats.cpp
@@ -420,6 +420,7 @@ void FMaterialStats::BuildShaderPlatformDB()
 	// Android
 	AddShaderPlatform(EPlatformCategoryType::Android, SP_OPENGL_ES3_1_ANDROID, TEXT("Android GLES 3.1"), true, true, TEXT("Android, OpenGLES 3.1"));
 	AddShaderPlatform(EPlatformCategoryType::Android, SP_VULKAN_ES3_1_ANDROID, TEXT("Android Vulkan"), true, true, TEXT("Android, Vulkan"));
+	AddShaderPlatform(EPlatformCategoryType::Android, SP_VULKAN_SM5_ANDROID, TEXT("Android Vulkan SM5"), true, true, TEXT("Android, Vulkan SM5"));
 
 	// Apple
 	AddShaderPlatform(EPlatformCategoryType::Desktop, SP_METAL_SM5, TEXT("Metal SM5"), false, true, TEXT("macOS, Metal, Shader Model 5"));
diff --git a/Engine/Source/Editor/MaterialEditor/Private/MaterialStatsCommon.cpp b/Engine/Source/Editor/MaterialEditor/Private/MaterialStatsCommon.cpp
--- a/Engine/Source/Editor/MaterialEditor/Private/MaterialStatsCommon.cpp
+++ b/Engine/Source/Editor/MaterialEditor/Private/MaterialStatsCommon.cpp
@@ -153,6 +153,8 @@ FString FMaterialStatsUtils::ShaderPlatformTypeName(const EShaderPlatform Platfo
 			return FString("VULKAN_SM5");
 		case SP_VULKAN_ES3_1_ANDROID:
 			return FString("VULKAN_ES3_1_ANDROID");
+		case SP_VULKAN_SM5_ANDROID:
+			return FString("VULKAN_SM5_ANDROID");
 		case SP_METAL_MACES3_1:
 			return FString("METAL_MACES3_1");
 		case SP_OPENGL_ES3_1_ANDROID:
@@ -179,6 +181,7 @@ FString FMaterialStatsUtils::GetPlatformOfflineCompilerPath(const EShaderPlatfor
 	{
 		case SP_OPENGL_ES3_1_ANDROID:
 		case SP_VULKAN_ES3_1_ANDROID:
+		case SP_VULKAN_SM5_ANDROID:
 			return FPaths::ConvertRelativePathToFull(GetDefault<UMaterialEditorSettings>()->MaliOfflineCompilerPath.FilePath);
 		break;
 
@@ -210,6 +213,7 @@ bool FMaterialStatsUtils::PlatformNeedsOfflineCompiler(const EShaderPlatform Sha
 		case SP_VULKAN_PCES3_1:
 		case SP_VULKAN_SM5:
 		case SP_VULKAN_ES3_1_ANDROID:
+		case SP_VULKAN_SM5_ANDROID:
 		case SP_OPENGL_ES3_1_ANDROID:
 			return true;
 
diff --git a/Engine/Source/Editor/UnrealEd/Private/NormalMapPreview.cpp b/Engine/Source/Editor/UnrealEd/Private/NormalMapPreview.cpp
--- a/Engine/Source/Editor/UnrealEd/Private/NormalMapPreview.cpp
+++ b/Engine/Source/Editor/UnrealEd/Private/NormalMapPreview.cpp
@@ -26,7 +26,7 @@ public:
 	/** Should the shader be cached? Always. */
 	static bool ShouldCompilePermutation(const FGlobalShaderPermutationParameters& Parameters)
 	{
-		return !IsConsolePlatform(Parameters.Platform);
+		return IsPCPlatform(Parameters.Platform);
 	}
 
 	/** Default constructor. */
diff --git a/Engine/Source/Runtime/Android/AndroidRuntimeSettings/Classes/AndroidRuntimeSettings.h b/Engine/Source/Runtime/Android/AndroidRuntimeSettings/Classes/AndroidRuntimeSettings.h
--- a/Engine/Source/Runtime/Android/AndroidRuntimeSettings/Classes/AndroidRuntimeSettings.h
+++ b/Engine/Source/Runtime/Android/AndroidRuntimeSettings/Classes/AndroidRuntimeSettings.h
@@ -391,6 +391,10 @@ public:
 	UPROPERTY(GlobalConfig, EditAnywhere, Category = Build, meta = (DisplayName = "Support Vulkan"))
 	bool bSupportsVulkan;
 
+	// Enable Vulkan SM5 rendering support
+	UPROPERTY(GlobalConfig, EditAnywhere, Category = Build, meta = (DisplayName = "Support Vulkan Desktop [Experimental]"))
+	bool bSupportsVulkanSM5;
+
 	// Whether to detect Vulkan device support by default, if the project is packaged with Vulkan support. If unchecked, the -detectvulkan commandline will enable Vulkan detection.
 	UPROPERTY(GlobalConfig, EditAnywhere, AdvancedDisplay, Category = Build, meta = (DisplayName = "Detect Vulkan device support", EditCondition = "bSupportsVulkan"))
 	bool bDetectVulkanByDefault;
diff --git a/Engine/Source/Runtime/Engine/Private/Components/PointLightComponent.cpp b/Engine/Source/Runtime/Engine/Private/Components/PointLightComponent.cpp
--- a/Engine/Source/Runtime/Engine/Private/Components/PointLightComponent.cpp
+++ b/Engine/Source/Runtime/Engine/Private/Components/PointLightComponent.cpp
@@ -66,7 +66,10 @@ bool FPointLightSceneProxy::GetWholeSceneProjectedShadowInitializer(const FScene
 		OutInitializer.WAxis = FVector4(0,0,1,0);
 		OutInitializer.MinLightW = 0.1f;
 		OutInitializer.MaxDistanceToCastInLightW = Radius;
-		OutInitializer.bOnePassPointLightShadow = true;
+
+		bool bSupportsGeometryShaders = RHISupportsGeometryShaders(GShaderPlatformForFeatureLevel[ViewFamily.GetFeatureLevel()]);
+		OutInitializer.bOnePassPointLightShadow = bSupportsGeometryShaders;
+
 		OutInitializer.bRayTracedDistanceField = UseRayTracedDistanceFieldShadows() && DoesPlatformSupportDistanceFieldShadowing(ViewFamily.GetShaderPlatform());
 		return true;
 	}
diff --git a/Engine/Source/Runtime/Engine/Private/Components/RectLightComponent.cpp b/Engine/Source/Runtime/Engine/Private/Components/RectLightComponent.cpp
--- a/Engine/Source/Runtime/Engine/Private/Components/RectLightComponent.cpp
+++ b/Engine/Source/Runtime/Engine/Private/Components/RectLightComponent.cpp
@@ -258,7 +258,10 @@ bool FRectLightSceneProxy::GetWholeSceneProjectedShadowInitializer(const FSceneV
 		OutInitializer.WAxis = FVector4(0, 0, 1, 0);
 		OutInitializer.MinLightW = 0.1f;
 		OutInitializer.MaxDistanceToCastInLightW = Radius;
-		OutInitializer.bOnePassPointLightShadow = true;
+		
+		bool bSupportsGeometryShaders = RHISupportsGeometryShaders(GShaderPlatformForFeatureLevel[ViewFamily.GetFeatureLevel()]);
+		OutInitializer.bOnePassPointLightShadow = bSupportsGeometryShaders;
+
 		OutInitializer.bRayTracedDistanceField = UseRayTracedDistanceFieldShadows() && DoesPlatformSupportDistanceFieldShadowing(ViewFamily.GetShaderPlatform());
 		return true;
 	}
diff --git a/Engine/Source/Runtime/Engine/Private/Materials/HLSLMaterialTranslator.cpp b/Engine/Source/Runtime/Engine/Private/Materials/HLSLMaterialTranslator.cpp
--- a/Engine/Source/Runtime/Engine/Private/Materials/HLSLMaterialTranslator.cpp
+++ b/Engine/Source/Runtime/Engine/Private/Materials/HLSLMaterialTranslator.cpp
@@ -1335,13 +1335,14 @@ void FHLSLMaterialTranslator::GetMaterialEnvironment(EShaderPlatform InPlatform,
 			}
 		}
 
-		if(ShadingModels.HasShadingModel(MSM_SingleLayerWater) && (IsSwitchPlatform(Platform) || IsPS4Platform(Platform) || Platform == SP_XBOXONE_D3D12))
+		if(ShadingModels.HasShadingModel(MSM_SingleLayerWater) && 
+			(IsSwitchPlatform(Platform) || IsVulkanMobileSM5Platform(Platform) || IsPS4Platform(Platform) || Platform == SP_XBOXONE_D3D12))
 		{
 			OutEnvironment.SetDefine(TEXT("DISABLE_FORWARD_LOCAL_LIGHTS"), TEXT("1"));
 		}
 
 		// This is to have switch use the simple single layer water shading similar to mobile: no dynamic lights, only sun and sky, no distortion, no colored transmittance on background, no custom depth read.
-		const bool bSingleLayerWaterUsesSimpleShading = IsSwitchPlatform(InPlatform) && IsForwardShadingEnabled(InPlatform);
+		const bool bSingleLayerWaterUsesSimpleShading = (IsSwitchPlatform(InPlatform) || IsVulkanMobileSM5Platform(InPlatform)) && IsForwardShadingEnabled(InPlatform);
 		if (ShadingModels.HasShadingModel(MSM_SingleLayerWater) && bSingleLayerWaterUsesSimpleShading)
 		{
 			OutEnvironment.SetDefine(TEXT("SINGLE_LAYER_WATER_SIMPLE_FORWARD"), TEXT("1"));
diff --git a/Engine/Source/Runtime/Engine/Private/Rendering/SubsurfaceProfile.cpp b/Engine/Source/Runtime/Engine/Private/Rendering/SubsurfaceProfile.cpp
--- a/Engine/Source/Runtime/Engine/Private/Rendering/SubsurfaceProfile.cpp
+++ b/Engine/Source/Runtime/Engine/Private/Rendering/SubsurfaceProfile.cpp
@@ -224,7 +224,7 @@ void FSubsurfaceProfileTexture::CreateTexture(FRHICommandListImmediate& RHICmdLi
 	{
 		Desc.Format = PF_A16B16G16R16;
 	}
-
+	
 	GRenderTargetPool.FindFreeElement(RHICmdList, Desc, GSSProfiles, TEXT("SSProfiles"));
 
 	// Write the contents of the texture.
@@ -234,8 +234,10 @@ void FSubsurfaceProfileTexture::CreateTexture(FRHICommandListImmediate& RHICmdLi
 	FLinearColor TextureRow[Width];
 	FMemory::Memzero(TextureRow);
 
-	const float FloatScale = GetNextSmallerPositiveFloat(0x10000);
-	check((int32)GetNextSmallerPositiveFloat(0x10000) == 0xffff);
+	// Android devices have different precision to desktop GPU's, scaling to one less works on Adreno and Mali
+	const float FloatScaleInitial = GMaxRHIShaderPlatform == SP_VULKAN_SM5_ANDROID ? 0xffff : 0x10000;
+	const float FloatScale = GetNextSmallerPositiveFloat(FloatScaleInitial);
+	check((int32)GetNextSmallerPositiveFloat(FloatScaleInitial) == (FloatScaleInitial - 1));
 
 	for (uint32 y = 0; y < Height; ++y)
 	{
diff --git a/Engine/Source/Runtime/Engine/Public/CubemapUnwrapUtils.h b/Engine/Source/Runtime/Engine/Public/CubemapUnwrapUtils.h
--- a/Engine/Source/Runtime/Engine/Public/CubemapUnwrapUtils.h
+++ b/Engine/Source/Runtime/Engine/Public/CubemapUnwrapUtils.h
@@ -48,7 +48,7 @@ class FCubemapTexturePropertiesVS : public FGlobalShader
 	DECLARE_SHADER_TYPE(FCubemapTexturePropertiesVS,Global);
 public:
 
-	static bool ShouldCompilePermutation(const FGlobalShaderPermutationParameters& Parameters) { return !IsConsolePlatform(Parameters.Platform); }
+	static bool ShouldCompilePermutation(const FGlobalShaderPermutationParameters& Parameters) { return IsPCPlatform(Parameters.Platform);}
 
 	FCubemapTexturePropertiesVS(const ShaderMetaType::CompiledShaderInitializerType& Initializer):
 		FGlobalShader(Initializer)
@@ -72,7 +72,7 @@ class FCubemapTexturePropertiesPS : public FGlobalShader
 	DECLARE_SHADER_TYPE(FCubemapTexturePropertiesPS,Global);
 public:
 
-	static bool ShouldCompilePermutation(const FGlobalShaderPermutationParameters& Parameters) { return !IsConsolePlatform(Parameters.Platform); }
+	static bool ShouldCompilePermutation(const FGlobalShaderPermutationParameters& Parameters) { return IsPCPlatform(Parameters.Platform);}
 
 	FCubemapTexturePropertiesPS(const ShaderMetaType::CompiledShaderInitializerType& Initializer)
 		: FGlobalShader(Initializer)
diff --git a/Engine/Source/Runtime/RHI/Private/RHI.cpp b/Engine/Source/Runtime/RHI/Private/RHI.cpp
--- a/Engine/Source/Runtime/RHI/Private/RHI.cpp
+++ b/Engine/Source/Runtime/RHI/Private/RHI.cpp
@@ -893,6 +893,7 @@ FName ShaderPlatformToPlatformName(EShaderPlatform Platform)
 		return NAME_PLATFORM_XBOXONE;
 	case SP_OPENGL_ES31_EXT:
 	case SP_VULKAN_ES3_1_ANDROID:
+	case SP_VULKAN_SM5_ANDROID:
 	case SP_OPENGL_ES3_1_ANDROID:
 		return NAME_PLATFORM_ANDROID;
 	case SP_METAL:
@@ -1297,6 +1298,7 @@ FString LexToString(EShaderPlatform Platform)
 	case SP_VULKAN_PCES3_1: return TEXT("VULKAN_PCES3_1");
 	case SP_VULKAN_SM5: return TEXT("VULKAN_SM5");
 	case SP_VULKAN_SM5_LUMIN: return TEXT("VULKAN_SM5_LUMIN");
+	case SP_VULKAN_SM5_ANDROID: return TEXT("VULKAN_SM5_ANDROID");
 
 	case SP_OPENGL_ES2_ANDROID_REMOVED:
 	case SP_OPENGL_ES2_WEBGL_REMOVED:
diff --git a/Engine/Source/Runtime/RHI/Public/RHIDefinitions.h b/Engine/Source/Runtime/RHI/Public/RHIDefinitions.h
--- a/Engine/Source/Runtime/RHI/Public/RHIDefinitions.h
+++ b/Engine/Source/Runtime/RHI/Public/RHIDefinitions.h
@@ -98,9 +98,10 @@ enum EShaderPlatform
 
 	//  Add new platforms below this line, starting from (SP_StaticPlatform_Last + 1)
 	//---------------------------------------------------------------------------------
+	SP_VULKAN_SM5_ANDROID			= SP_StaticPlatform_Last+1,
 
 	SP_NumPlatforms,
-	SP_NumBits				= 7,
+	SP_NumBits						= 7,
 };
 static_assert(SP_NumPlatforms <= (1 << SP_NumBits), "SP_NumPlatforms will not fit on SP_NumBits");
 
@@ -1369,7 +1370,9 @@ inline bool IsOpenGLPlatform(const FStaticShaderPlatform Platform)
 
 inline bool IsMetalPlatform(const FStaticShaderPlatform Platform)
 {
-	return Platform == SP_METAL || Platform == SP_METAL_MRT || Platform == SP_METAL_TVOS || Platform == SP_METAL_MRT_TVOS || Platform == SP_METAL_SM5_NOTESS || Platform == SP_METAL_SM5 || Platform == SP_METAL_MACES3_1 || Platform == SP_METAL_MRT_MAC
+	return Platform == SP_METAL || Platform == SP_METAL_MRT || Platform == SP_METAL_TVOS || Platform == SP_METAL_MRT_TVOS 
+		|| Platform == SP_METAL_SM5_NOTESS || Platform == SP_METAL_SM5
+		|| Platform == SP_METAL_MACES3_1 || Platform == SP_METAL_MRT_MAC
 		|| FDataDrivenShaderPlatformInfo::GetIsLanguageMetal(Platform);
 }
 
@@ -1411,16 +1414,22 @@ inline bool IsPS4Platform(const FStaticShaderPlatform Platform)
 
 inline bool IsVulkanPlatform(const FStaticShaderPlatform Platform)
 {
-	return Platform == SP_VULKAN_SM5 || Platform == SP_VULKAN_SM5_LUMIN || Platform == SP_VULKAN_PCES3_1 || Platform == SP_VULKAN_ES3_1_ANDROID || Platform == SP_VULKAN_ES3_1_LUMIN
+	return Platform == SP_VULKAN_SM5 || Platform == SP_VULKAN_SM5_LUMIN || Platform == SP_VULKAN_PCES3_1 
+		|| Platform == SP_VULKAN_ES3_1_ANDROID || Platform == SP_VULKAN_ES3_1_LUMIN || Platform == SP_VULKAN_SM5_ANDROID
 		|| FDataDrivenShaderPlatformInfo::GetIsLanguageVulkan(Platform);
 }
 
 inline bool IsVulkanSM5Platform(const FStaticShaderPlatform Platform)
 {
-	return Platform == SP_VULKAN_SM5 || Platform == SP_VULKAN_SM5_LUMIN
+	return Platform == SP_VULKAN_SM5 || Platform == SP_VULKAN_SM5_LUMIN || Platform == SP_VULKAN_SM5_ANDROID
 		|| (FDataDrivenShaderPlatformInfo::GetIsLanguageVulkan(Platform) && FDataDrivenShaderPlatformInfo::GetMaxFeatureLevel(Platform) == ERHIFeatureLevel::SM5);
 }
 
+inline bool IsVulkanMobileSM5Platform(const EShaderPlatform Platform)
+{
+	return Platform == SP_VULKAN_SM5_ANDROID;
+}
+
 inline bool IsAndroidOpenGLESPlatform(const FStaticShaderPlatform Platform)
 {
 	return Platform == SP_OPENGL_ES3_1_ANDROID || FDataDrivenShaderPlatformInfo::GetIsAndroidOpenGLES(Platform);
@@ -1478,6 +1487,7 @@ inline FStaticFeatureLevel GetMaxSupportedFeatureLevel(const FStaticShaderPlatfo
 	case SP_VULKAN_SM5:
 	case SP_VULKAN_SM5_LUMIN:
 	case SP_SWITCH:
+	case SP_VULKAN_SM5_ANDROID:
 		return ERHIFeatureLevel::SM5;
 	case SP_METAL:
 	case SP_METAL_TVOS:
@@ -1568,7 +1578,8 @@ inline bool RHISupportsComputeShaders(const FStaticShaderPlatform Platform)
 
 inline bool RHISupportsGeometryShaders(const FStaticShaderPlatform Platform)
 {
-	return IsFeatureLevelSupported(Platform, ERHIFeatureLevel::SM5) && !IsMetalPlatform(Platform) && !IsVulkanMobilePlatform(Platform);
+	return IsFeatureLevelSupported(Platform, ERHIFeatureLevel::SM5) && !IsMetalPlatform(Platform) &&
+			 !IsVulkanMobilePlatform(Platform) && !IsVulkanMobileSM5Platform(Platform);
 }
 
 inline bool RHIHasTiledGPU(const FStaticShaderPlatform Platform)
diff --git a/Engine/Source/Runtime/RHI/Public/RHIShaderFormatDefinitions.inl b/Engine/Source/Runtime/RHI/Public/RHIShaderFormatDefinitions.inl
--- a/Engine/Source/Runtime/RHI/Public/RHIShaderFormatDefinitions.inl
+++ b/Engine/Source/Runtime/RHI/Public/RHIShaderFormatDefinitions.inl
@@ -42,6 +42,8 @@ static FName NAME_VULKAN_SM5_NOUB(TEXT("SF_VULKAN_SM5_NOUB"));
 static FName NAME_VULKAN_SM5(TEXT("SF_VULKAN_SM5"));
 static FName NAME_VULKAN_SM5_LUMIN(TEXT("SF_VULKAN_SM5_LUMIN"));
 static FName NAME_VULKAN_SM5_LUMIN_NOUB(TEXT("SF_VULKAN_SM5_LUMIN_NOUB"));
+static FName NAME_VULKAN_SM5_ANDROID(TEXT("SF_VULKAN_SM5_ANDROID"));
+static FName NAME_VULKAN_SM5_ANDROID_NOUB(TEXT("SF_VULKAN_SM5_ANDROID_NOUB"));
 
 
 static FName ShaderPlatformToShaderFormatName(EShaderPlatform Platform)
@@ -115,6 +117,11 @@ static FName ShaderPlatformToShaderFormatName(EShaderPlatform Platform)
 		static auto* CVar = IConsoleManager::Get().FindTConsoleVariableDataInt(TEXT("r.Vulkan.UseRealUBs"));
 		return (CVar && CVar->GetValueOnAnyThread() == 0) ? NAME_VULKAN_SM5_LUMIN_NOUB : NAME_VULKAN_SM5_LUMIN;
 	}
+	case SP_VULKAN_SM5_ANDROID:
+	{
+		static auto* CVar = IConsoleManager::Get().FindTConsoleVariableDataInt(TEXT("r.Vulkan.UseRealUBs"));
+		return (CVar && CVar->GetValueOnAnyThread() == 0) ? NAME_VULKAN_SM5_ANDROID_NOUB : NAME_VULKAN_SM5_ANDROID;
+	}
 
 	default:
 		if (FStaticShaderPlatformNames::IsStaticPlatform(Platform))
@@ -165,6 +172,8 @@ static EShaderPlatform ShaderFormatNameToShaderPlatform(FName ShaderFormat)
 	if (ShaderFormat == NAME_VULKAN_SM5)				return SP_VULKAN_SM5;
 	if (ShaderFormat == NAME_VULKAN_SM5_LUMIN)			return SP_VULKAN_SM5_LUMIN;
 	if (ShaderFormat == NAME_VULKAN_SM5_LUMIN_NOUB)		return SP_VULKAN_SM5_LUMIN;
+	if (ShaderFormat == NAME_VULKAN_SM5_ANDROID)		return SP_VULKAN_SM5_ANDROID;
+	if (ShaderFormat == NAME_VULKAN_SM5_ANDROID_NOUB)	return SP_VULKAN_SM5_ANDROID;
 
 	for (int32 StaticPlatform = SP_StaticPlatform_First; StaticPlatform <= SP_StaticPlatform_Last; ++StaticPlatform)
 	{
diff --git a/Engine/Source/Runtime/RHI/Public/RHIShaderPlatformDefinitions.inl b/Engine/Source/Runtime/RHI/Public/RHIShaderPlatformDefinitions.inl
--- a/Engine/Source/Runtime/RHI/Public/RHIShaderPlatformDefinitions.inl
+++ b/Engine/Source/Runtime/RHI/Public/RHIShaderPlatformDefinitions.inl
@@ -103,6 +103,13 @@ static const FText& GetFriendlyShaderPlatformName(const EShaderPlatform InShader
 	}
 	break;
 
+	case SP_VULKAN_SM5_ANDROID:
+	{
+		static const FText Description = NSLOCTEXT("FriendlyShaderPlatformNames", "Android_Vulkan_SM5_loc", "Android Vulkan SM5");
+		return Description;
+	}
+	break;
+
 	default:
 		break;
 	};
diff --git a/Engine/Source/Runtime/RenderCore/Public/RenderUtils.h b/Engine/Source/Runtime/RenderCore/Public/RenderUtils.h
--- a/Engine/Source/Runtime/RenderCore/Public/RenderUtils.h
+++ b/Engine/Source/Runtime/RenderCore/Public/RenderUtils.h
@@ -555,7 +555,7 @@ inline bool IsUsingPerPixelDBufferMask(const FStaticShaderPlatform Platform)
 
 inline bool UseGPUScene(const FStaticShaderPlatform Platform, const FStaticFeatureLevel FeatureLevel)
 {
-	if (FeatureLevel == ERHIFeatureLevel::ES3_1)
+	if (IsMobilePlatform(Platform))
 	{
 		return MobileSupportsGPUScene(Platform);
 	}
@@ -565,6 +565,7 @@ inline bool UseGPUScene(const FStaticShaderPlatform Platform, const FStaticFeatu
 		//@todo - support GPU Scene management compute shaders on these platforms to get dynamic instancing speedups on the Rendering Thread and RHI Thread
 		&& !IsOpenGLPlatform(Platform)
 		&& !IsSwitchPlatform(Platform)
+		&& !IsVulkanMobileSM5Platform(Platform)
 		// we only check DDSPI for platforms that have been read in - IsValid() can go away once ALL platforms are converted over to this system
 		&& (!FDataDrivenShaderPlatformInfo::IsValid(Platform) || FDataDrivenShaderPlatformInfo::GetSupportsGPUScene(Platform));
 }
diff --git a/Engine/Source/Runtime/Renderer/Private/DeferredShadingRenderer.cpp b/Engine/Source/Runtime/Renderer/Private/DeferredShadingRenderer.cpp
--- a/Engine/Source/Runtime/Renderer/Private/DeferredShadingRenderer.cpp
+++ b/Engine/Source/Runtime/Renderer/Private/DeferredShadingRenderer.cpp
@@ -1897,7 +1897,9 @@ void FDeferredShadingSceneRenderer::Render(FRHICommandListImmediate& RHICmdList)
 		 
 		// if we didn't to the prepass above, then we will need to clear now, otherwise, it's already been cleared and rendered to
 		ERenderTargetLoadAction ColorLoadAction = ERenderTargetLoadAction::ELoad;
-		ERenderTargetLoadAction DepthLoadAction = bClearDepth ? ERenderTargetLoadAction::EClear : (!IsMetalPlatform(ShaderPlatform) ? ERenderTargetLoadAction::ENoAction : ERenderTargetLoadAction::ELoad);
+
+		bool bRequiresLoad = IsMetalPlatform(ShaderPlatform) || IsVulkanMobileSM5Platform(ShaderPlatform);
+		ERenderTargetLoadAction DepthLoadAction = bClearDepth ? ERenderTargetLoadAction::EClear : (!bRequiresLoad ? ERenderTargetLoadAction::ENoAction : ERenderTargetLoadAction::ELoad);
 
 		const bool bClearBlack = ViewFamily.EngineShowFlags.ShaderComplexity || ViewFamily.EngineShowFlags.StationaryLightOverlap;
 		const float ClearAlpha = GetSceneColorClearAlpha();
@@ -1951,7 +1953,9 @@ void FDeferredShadingSceneRenderer::Render(FRHICommandListImmediate& RHICmdList)
 	{
 		// Make sure we have began the renderpass
 		ERenderTargetLoadAction DepthLoadAction = bDepthWasCleared ? ERenderTargetLoadAction::ELoad : ERenderTargetLoadAction::EClear;
-		SceneContext.BeginRenderingGBuffer(RHICmdList, (!IsMetalPlatform(ShaderPlatform) ? ERenderTargetLoadAction::ENoAction : ERenderTargetLoadAction::ELoad), DepthLoadAction, BasePassDepthStencilAccess, ViewFamily.EngineShowFlags.ShaderComplexity);
+
+		bool bRequiresLoad = IsMetalPlatform(ShaderPlatform) || IsVulkanMobileSM5Platform(ShaderPlatform);
+		SceneContext.BeginRenderingGBuffer(RHICmdList, (!bRequiresLoad ? ERenderTargetLoadAction::ENoAction : ERenderTargetLoadAction::ELoad), DepthLoadAction, BasePassDepthStencilAccess, ViewFamily.EngineShowFlags.ShaderComplexity);
 	}
 	// Wait for Async SSAO before rendering base pass with forward rendering
 	if (IsForwardShadingEnabled(ShaderPlatform))
@@ -2503,7 +2507,8 @@ void FDeferredShadingSceneRenderer::Render(FRHICommandListImmediate& RHICmdList)
 		SetupSceneTextureUniformParameters(SceneContext, FeatureLevel, ESceneTextureSetupMode::SceneDepth | ESceneTextureSetupMode::GBuffers, SceneTextureParameters);
 		TUniformBufferRef<FSceneTexturesUniformParameters> SceneTextureUniformBuffer = TUniformBufferRef<FSceneTexturesUniformParameters>::CreateUniformBufferImmediate(SceneTextureParameters, UniformBuffer_SingleFrame);
 
-        SceneContext.BeginRenderingSceneColor(RHICmdList, (!IsMetalPlatform(ShaderPlatform) ? ESimpleRenderTargetMode::EUninitializedColorExistingDepth : ESimpleRenderTargetMode::EExistingColorAndDepth));
+		bool bRequiresLoad = IsMetalPlatform(ShaderPlatform) || IsVulkanMobileSM5Platform(ShaderPlatform);
+        SceneContext.BeginRenderingSceneColor(RHICmdList, (!bRequiresLoad ? ESimpleRenderTargetMode::EUninitializedColorExistingDepth : ESimpleRenderTargetMode::EExistingColorAndDepth));
 		for (int32 ViewIndex = 0; ViewIndex < Views.Num(); ++ViewIndex)
 		{
 			const FViewInfo& View = Views[ViewIndex];
diff --git a/Engine/Source/Runtime/Renderer/Private/GlobalDistanceField.cpp b/Engine/Source/Runtime/Renderer/Private/GlobalDistanceField.cpp
--- a/Engine/Source/Runtime/Renderer/Private/GlobalDistanceField.cpp
+++ b/Engine/Source/Runtime/Renderer/Private/GlobalDistanceField.cpp
@@ -612,7 +612,8 @@ public:
 
 	static bool ShouldCompilePermutation(const FGlobalShaderPermutationParameters& Parameters)
 	{
-		return IsFeatureLevelSupported(Parameters.Platform, ERHIFeatureLevel::SM5) && DoesPlatformSupportDistanceFieldAO(Parameters.Platform) && !IsMetalPlatform(Parameters.Platform) && IsUsingDistanceFields(Parameters.Platform);
+		return IsFeatureLevelSupported(Parameters.Platform, ERHIFeatureLevel::SM5) && DoesPlatformSupportDistanceFieldAO(Parameters.Platform) &&
+				!IsMetalPlatform(Parameters.Platform) && !IsVulkanMobileSM5Platform(Parameters.Platform) && IsUsingDistanceFields(Parameters.Platform);
 	}
 
 	static void ModifyCompilationEnvironment(const FGlobalShaderPermutationParameters& Parameters, FShaderCompilerEnvironment& OutEnvironment)
@@ -721,7 +722,8 @@ void FHeightfieldLightingViewInfo::CompositeHeightfieldsIntoGlobalDistanceField(
 		&& GDynamicRHI->RHIIsTypedUAVLoadSupported(ClipmapPixelFormat)
 		&& NumPrimitives > 0
 		&& SupportsDistanceFieldAO(Scene->GetFeatureLevel(), Scene->GetShaderPlatform())
-		&& !IsMetalPlatform(Scene->GetShaderPlatform()))
+		&& !IsMetalPlatform(Scene->GetShaderPlatform())
+		&& !IsVulkanMobileSM5Platform(Scene->GetShaderPlatform()))
 	{
 		FHeightfieldDescription UpdateRegionHeightfield;
 
diff --git a/Engine/Source/Runtime/Renderer/Private/MeshDrawCommands.cpp b/Engine/Source/Runtime/Renderer/Private/MeshDrawCommands.cpp
--- a/Engine/Source/Runtime/Renderer/Private/MeshDrawCommands.cpp
+++ b/Engine/Source/Runtime/Renderer/Private/MeshDrawCommands.cpp
@@ -774,9 +774,11 @@ public:
 	{
 		TRACE_CPUPROFILER_EVENT_SCOPE(MeshDrawCommandPassSetupTask);
 		// Mobile base pass is a special case, as final lists is created from two mesh passes based on CSM visibility.
-		const bool bMobileBasePass = Context.ShadingPath == EShadingPath::Mobile && Context.PassType == EMeshPass::BasePass;
+		const bool bMobileShadingBasePass = Context.ShadingPath == EShadingPath::Mobile && Context.PassType == EMeshPass::BasePass;
+		// On SM5 Mobile platform, still want the same sorting
+		const bool bMobileVulkanSM5BasePass = IsVulkanMobileSM5Platform(Context.ShaderPlatform) && Context.PassType == EMeshPass::BasePass;
 
-		if (bMobileBasePass)
+		if (bMobileShadingBasePass)
 		{
 			MergeMobileBasePassMeshDrawCommands(
 				Context.View->MobileCSMVisibilityInfo,
@@ -841,7 +843,7 @@ public:
 			}
 
 			// Update sort keys.
-			if (bMobileBasePass)
+			if (bMobileShadingBasePass || bMobileVulkanSM5BasePass)
 			{
 				UpdateMobileBasePassMeshSortKeys(
 					Context.ViewOrigin,
diff --git a/Engine/Source/Runtime/Renderer/Private/MobileDistortionPass.cpp b/Engine/Source/Runtime/Renderer/Private/MobileDistortionPass.cpp
--- a/Engine/Source/Runtime/Renderer/Private/MobileDistortionPass.cpp
+++ b/Engine/Source/Runtime/Renderer/Private/MobileDistortionPass.cpp
@@ -79,7 +79,7 @@ class FDistortionMergePS_ES2 : public FGlobalShader
 
 	static bool ShouldCompilePermutation(const FGlobalShaderPermutationParameters& Parameters)
 	{
-		return !IsConsolePlatform(Parameters.Platform);
+		return GetMaxSupportedFeatureLevel(Parameters.Platform) <= ERHIFeatureLevel::ES3_1;
 	}
 
 	FDistortionMergePS_ES2() {}
diff --git a/Engine/Source/Runtime/Renderer/Private/PostProcess/PostProcessTonemap.cpp b/Engine/Source/Runtime/Renderer/Private/PostProcess/PostProcessTonemap.cpp
--- a/Engine/Source/Runtime/Renderer/Private/PostProcess/PostProcessTonemap.cpp
+++ b/Engine/Source/Runtime/Renderer/Private/PostProcess/PostProcessTonemap.cpp
@@ -1317,7 +1317,7 @@ class FPostProcessTonemapVS_ES2 : public FGlobalShader
 			return false;
 		}
 
-		return !IsConsolePlatform(Parameters.Platform);
+		return GetMaxSupportedFeatureLevel(Parameters.Platform) <= ERHIFeatureLevel::ES3_1;
 	}
 
 	static void ModifyCompilationEnvironment(const FGlobalShaderPermutationParameters& Parameters, FShaderCompilerEnvironment& OutEnvironment)
diff --git a/Engine/Source/Runtime/Renderer/Private/SceneOcclusion.cpp b/Engine/Source/Runtime/Renderer/Private/SceneOcclusion.cpp
--- a/Engine/Source/Runtime/Renderer/Private/SceneOcclusion.cpp
+++ b/Engine/Source/Runtime/Renderer/Private/SceneOcclusion.cpp
@@ -74,16 +74,16 @@ int32 FOcclusionQueryHelpers::GetNumBufferedFrames(ERHIFeatureLevel::Type Featur
 	return FMath::Min<int32>(GNumAlternateFrameRenderingGroups, (int32)FOcclusionQueryHelpers::MaxBufferedOcclusionFrames);
 #endif
 	static const auto NumBufferedQueriesVar = IConsoleManager::Get().FindTConsoleVariableDataInt(TEXT("r.NumBufferedOcclusionQueries"));
+	EShaderPlatform ShaderPlatform = GShaderPlatformForFeatureLevel[FeatureLevel];
 
 	int32 NumExtraMobileFrames = 0;
-	if (FeatureLevel <= ERHIFeatureLevel::ES3_1)
+	if (FeatureLevel <= ERHIFeatureLevel::ES3_1 || IsVulkanMobileSM5Platform(ShaderPlatform))
 	{
 		NumExtraMobileFrames++; // the mobile renderer just doesn't do much after the basepass, and hence it will be asking for the query results almost immediately; the results can't possibly be ready in 1 frame.
 		
-		EShaderPlatform ShaderPlatform = GShaderPlatformForFeatureLevel[FeatureLevel];
 		if ((
 			//IsOpenGLPlatform(ShaderPlatform) || 
-			IsVulkanPlatform(ShaderPlatform) || IsSwitchPlatform(ShaderPlatform)) && IsRunningRHIInSeparateThread())
+			IsVulkanPlatform(ShaderPlatform) || IsSwitchPlatform(ShaderPlatform) || IsVulkanMobileSM5Platform(ShaderPlatform)) && IsRunningRHIInSeparateThread())
 		{
 			// Android, unfortunately, requires the RHIThread to mediate the readback of queries. Therefore we need an extra frame to avoid a stall in either thread. 
 			// The RHIT needs to do read back after the queries are ready and before the RT needs them to avoid stalls. The RHIT may be busy when the queries become ready, so this is all very complicated.
@@ -1395,6 +1395,7 @@ void FSceneRenderer::BeginOcclusionTests(FRHICommandListImmediate& RHICmdList, b
 				// Restore default render target
 				// #todo-renderpasses this is not ideal. This pass should be self-contained. Can we refactor this?
 				SceneContext.BeginRenderingSceneColor(RHICmdList, ESimpleRenderTargetMode::EUninitializedColorExistingDepth, FExclusiveDepthStencil::DepthRead_StencilWrite);
+				RHICmdList.EndRenderPass();
 			}
 		}
 	}
diff --git a/Engine/Source/Runtime/Renderer/Private/SingleLayerWaterRendering.cpp b/Engine/Source/Runtime/Renderer/Private/SingleLayerWaterRendering.cpp
--- a/Engine/Source/Runtime/Renderer/Private/SingleLayerWaterRendering.cpp
+++ b/Engine/Source/Runtime/Renderer/Private/SingleLayerWaterRendering.cpp
@@ -95,7 +95,7 @@ static bool ShouldRenderSingleLayerWater(const FViewInfo& View)
 // This is to have switch use the simple single layer water shading similar to mobile: no dynamic lights, only sun and sky, no distortion, no colored transmittance on background, no custom depth read.
 bool SingleLayerWaterUsesSimpleShading(EShaderPlatform ShaderPlatform)
 {
-	return  IsSwitchPlatform(ShaderPlatform) && IsForwardShadingEnabled(ShaderPlatform);
+	return  (IsSwitchPlatform(ShaderPlatform) || IsVulkanMobileSM5Platform(ShaderPlatform)) && IsForwardShadingEnabled(ShaderPlatform);
 }
 
 bool ShouldRenderSingleLayerWater(const TArray<FViewInfo>& Views, const FEngineShowFlags& EngineShowFlags)
diff --git a/Engine/Source/Runtime/TraceLog/Private/Trace/Detail/Android/AndroidTrace.cpp b/Engine/Source/Runtime/TraceLog/Private/Trace/Detail/Android/AndroidTrace.cpp
--- a/Engine/Source/Runtime/TraceLog/Private/Trace/Detail/Android/AndroidTrace.cpp
+++ b/Engine/Source/Runtime/TraceLog/Private/Trace/Detail/Android/AndroidTrace.cpp
@@ -12,6 +12,7 @@
 #include <sys/socket.h>
 #include <sys/stat.h>
 #include <unistd.h>
+#include <errno.h>
 
 namespace Trace {
 namespace Private {
diff --git a/Engine/Source/Runtime/VulkanRHI/Private/Android/VulkanAndroidPlatform.cpp b/Engine/Source/Runtime/VulkanRHI/Private/Android/VulkanAndroidPlatform.cpp
--- a/Engine/Source/Runtime/VulkanRHI/Private/Android/VulkanAndroidPlatform.cpp
+++ b/Engine/Source/Runtime/VulkanRHI/Private/Android/VulkanAndroidPlatform.cpp
@@ -19,6 +19,7 @@
 #define DEFINE_VK_ENTRYPOINTS(Type,Func) Type VulkanDynamicAPI::Func = NULL;
 ENUM_VK_ENTRYPOINTS_ALL(DEFINE_VK_ENTRYPOINTS)
 
+#define VULKAN_MALI_LAYER_NAME "VK_LAYER_ARM_AGA"
 
 void* FVulkanAndroidPlatform::VulkanLib = nullptr;
 bool FVulkanAndroidPlatform::bAttemptedLoad = false;
@@ -152,6 +153,10 @@ void FVulkanAndroidPlatform::GetDeviceExtensions(EGpuVendorId VendorId, TArray<c
 	OutExtensions.Add(VK_KHR_ANDROID_SURFACE_EXTENSION_NAME);
 	OutExtensions.Add(VK_GOOGLE_DISPLAY_TIMING_EXTENSION_NAME);
 	OutExtensions.Add(VK_EXT_FRAGMENT_DENSITY_MAP_EXTENSION_NAME);
+
+#if !UE_BUILD_SHIPPING
+	OutExtensions.Add(VULKAN_MALI_LAYER_NAME);
+#endif
 }
 
 bool FVulkanAndroidPlatform::SupportsStandardSwapchain()
@@ -202,8 +207,19 @@ void FVulkanAndroidPlatform::OverridePlatformHandlers(bool bInit)
 
 void FVulkanAndroidPlatform::SetupMaxRHIFeatureLevelAndShaderPlatform(ERHIFeatureLevel::Type InRequestedFeatureLevel)
 {
-	GMaxRHIFeatureLevel = ERHIFeatureLevel::ES3_1;
-	GMaxRHIShaderPlatform = SP_VULKAN_ES3_1_ANDROID;
+	if (!GIsEditor &&
+		(FVulkanPlatform::RequiresMobileRenderer() || 
+		InRequestedFeatureLevel == ERHIFeatureLevel::ES3_1 ||
+		FParse::Param(FCommandLine::Get(), TEXT("featureleveles31"))))
+	{
+		GMaxRHIFeatureLevel = ERHIFeatureLevel::ES3_1;
+		GMaxRHIShaderPlatform = SP_VULKAN_ES3_1_ANDROID;
+	}
+	else
+	{
+		GMaxRHIFeatureLevel = ERHIFeatureLevel::SM5;
+		GMaxRHIShaderPlatform = SP_VULKAN_SM5_ANDROID;
+	}
 }
 
 #endif
diff --git a/Engine/Source/Runtime/VulkanRHI/Private/Android/VulkanAndroidPlatform.h b/Engine/Source/Runtime/VulkanRHI/Private/Android/VulkanAndroidPlatform.h
--- a/Engine/Source/Runtime/VulkanRHI/Private/Android/VulkanAndroidPlatform.h
+++ b/Engine/Source/Runtime/VulkanRHI/Private/Android/VulkanAndroidPlatform.h
@@ -61,10 +61,20 @@ public:
 
 	static void SetupFeatureLevels()
 	{
-		GShaderPlatformForFeatureLevel[ERHIFeatureLevel::ES2_REMOVED] = SP_NumPlatforms;
-		GShaderPlatformForFeatureLevel[ERHIFeatureLevel::ES3_1] = SP_VULKAN_ES3_1_ANDROID;
-		GShaderPlatformForFeatureLevel[ERHIFeatureLevel::SM4_REMOVED] = SP_NumPlatforms;
-		GShaderPlatformForFeatureLevel[ERHIFeatureLevel::SM5] = SP_NumPlatforms;
+		if (RequiresMobileRenderer())
+		{
+			GShaderPlatformForFeatureLevel[ERHIFeatureLevel::ES2_REMOVED] = SP_NumPlatforms;
+			GShaderPlatformForFeatureLevel[ERHIFeatureLevel::ES3_1] = SP_VULKAN_ES3_1_ANDROID;
+			GShaderPlatformForFeatureLevel[ERHIFeatureLevel::SM4_REMOVED] = SP_NumPlatforms;
+			GShaderPlatformForFeatureLevel[ERHIFeatureLevel::SM5] = SP_NumPlatforms;
+		}
+		else
+		{
+			GShaderPlatformForFeatureLevel[ERHIFeatureLevel::ES2_REMOVED] = SP_NumPlatforms;
+			GShaderPlatformForFeatureLevel[ERHIFeatureLevel::ES3_1] = SP_VULKAN_SM5_ANDROID;
+			GShaderPlatformForFeatureLevel[ERHIFeatureLevel::SM4_REMOVED] = SP_VULKAN_SM5_ANDROID;
+			GShaderPlatformForFeatureLevel[ERHIFeatureLevel::SM5] = SP_VULKAN_SM5_ANDROID;
+		}
 	}
 
 	static bool SupportsStandardSwapchain();
@@ -72,7 +82,13 @@ public:
 
 	static bool SupportsTimestampRenderQueries();
 
-	static bool RequiresMobileRenderer() { return true; }
+	static bool RequiresMobileRenderer()
+	{
+		return !FAndroidMisc::ShouldUseDesktopVulkan();
+	}
+
+	static bool SupportsVolumeTextureRendering() { return false; }
+
 	static void OverridePlatformHandlers(bool bInit);
 
 	//#todo-rco: Detect Mali?
@@ -84,8 +100,7 @@ public:
 
 	static bool UseRealUBsOptimization(bool bCodeHeaderUseRealUBs)
 	{
-		// Android is hard-coded to SF_VULKAN_ES31_ANDROID_NOUB shader format
-		return false;
+		return !RequiresMobileRenderer();
 	}
 
 	// Assume most devices can't use the extra cores for running parallel tasks
diff --git a/Engine/Source/Runtime/VulkanRHI/Private/VulkanGenericPlatform.h b/Engine/Source/Runtime/VulkanRHI/Private/VulkanGenericPlatform.h
--- a/Engine/Source/Runtime/VulkanRHI/Private/VulkanGenericPlatform.h
+++ b/Engine/Source/Runtime/VulkanRHI/Private/VulkanGenericPlatform.h
@@ -95,6 +95,8 @@ public:
 	/** The status quo is false, so the default is chosen to not change it. As platforms opt in it may be better to flip the default. */
 	static bool SupportsDynamicResolution() { return false; }
 
+	static bool SupportsVolumeTextureRendering() { return true; }
+
 	// Allow platforms to add extension features to the DeviceInfo pNext chain
 	static void EnablePhysicalDeviceFeatureExtensions(VkDeviceCreateInfo& DeviceInfo) {}
 
diff --git a/Engine/Source/Runtime/VulkanRHI/Private/VulkanRHI.cpp b/Engine/Source/Runtime/VulkanRHI/Private/VulkanRHI.cpp
--- a/Engine/Source/Runtime/VulkanRHI/Private/VulkanRHI.cpp
+++ b/Engine/Source/Runtime/VulkanRHI/Private/VulkanRHI.cpp
@@ -704,7 +704,7 @@ void FVulkanDynamicRHI::InitInstance()
 		//#todo-rco: Add newer Nvidia also
 		GSupportsEfficientAsyncCompute = (Device->GetVendorId() == EGpuVendorId::Amd) && (GRHIAllowAsyncComputeCvar.GetValueOnAnyThread() > 0) && (Device->ComputeContext != Device->ImmediateContext);
 
-		GSupportsVolumeTextureRendering = true;
+		GSupportsVolumeTextureRendering = FVulkanPlatform::SupportsVolumeTextureRendering();
 
 		// Indicate that the RHI needs to use the engine's deferred deletion queue.
 		GRHINeedsExtraDeletionLatency = true;
