Reduce new PSO message down from warning to display - log and console outputs are still updated as before but warnings treated as errors in certain tests is not what we want here.  These should not have been warnings in the first place - more of an informational thing.

diff --git a/Engine/Source/Runtime/RHI/Private/PipelineFileCache.cpp b/Engine/Source/Runtime/RHI/Private/PipelineFileCache.cpp
--- a/Engine/Source/Runtime/RHI/Private/PipelineFileCache.cpp
+++ b/Engine/Source/Runtime/RHI/Private/PipelineFileCache.cpp
@@ -2778,10 +2778,10 @@ void FPipelineFileCache::CacheGraphicsPSO(uint32 RunTimeHash, FGraphicsPipelineS
 						if (bActuallyNewPSO)
 						{
 							CSV_EVENT(PSO, TEXT("Encountered new graphics PSO"));
-							UE_LOG(LogRHI, Warning, TEXT("Encountered a new graphics PSO: %u"), PSOHash);
+							UE_LOG(LogRHI, Display, TEXT("Encountered a new graphics PSO: %u"), PSOHash);
 							if (GPSOFileCachePrintNewPSODescriptors > 0)
 							{
-								UE_LOG(LogRHI, Warning, TEXT("New Graphics PSO (%u) Description: %s"), PSOHash, *NewEntry.GraphicsDesc.ToString());
+								UE_LOG(LogRHI, Display, TEXT("New Graphics PSO (%u) Description: %s"), PSOHash, *NewEntry.GraphicsDesc.ToString());
 							}
 							if (LogPSOtoFileCache())
 							{
@@ -2850,10 +2850,10 @@ void FPipelineFileCache::CacheComputePSO(uint32 RunTimeHash, FRHIComputeShader c
 						if (bActuallyNewPSO)
 						{
 							CSV_EVENT(PSO, TEXT("Encountered new compute PSO"));
-							UE_LOG(LogRHI, Warning, TEXT("Encountered a new compute PSO: %u"), PSOHash);
+							UE_LOG(LogRHI, Display, TEXT("Encountered a new compute PSO: %u"), PSOHash);
 							if (GPSOFileCachePrintNewPSODescriptors > 0)
 							{
-								UE_LOG(LogRHI, Warning, TEXT("New compute PSO (%u) Description: %s"), PSOHash, *NewEntry.ComputeDesc.ComputeShader.ToString());
+								UE_LOG(LogRHI, Display, TEXT("New compute PSO (%u) Description: %s"), PSOHash, *NewEntry.ComputeDesc.ComputeShader.ToString());
 							}
 							
 							if (LogPSOtoFileCache())
