Fix bug where fragment density map in vulkan (if you use vulkan foveation ) is added in the renderpass, but not in the renderpass attached to the graphics pipeline.

diff --git a/Engine/Source/Runtime/RHI/Public/RHICommandList.h b/Engine/Source/Runtime/RHI/Public/RHICommandList.h
--- a/Engine/Source/Runtime/RHI/Public/RHICommandList.h
+++ b/Engine/Source/Runtime/RHI/Public/RHICommandList.h
@@ -597,7 +597,7 @@ protected:
 		
 		ESubpassHint SubpassHint = ESubpassHint::None;
 		uint8 SubpassIndex = 0;
-
+		bool HasFragmentDensityAttachment = false;
 	} PSOContext;
 
 	FBoundShaderStateInput BoundShaderInput;
@@ -628,7 +628,8 @@ protected:
 	void CacheActiveRenderTargets(
 		uint32 NewNumSimultaneousRenderTargets,
 		const FRHIRenderTargetView* NewRenderTargetsRHI,
-		const FRHIDepthRenderTargetView* NewDepthStencilTargetRHI
+		const FRHIDepthRenderTargetView* NewDepthStencilTargetRHI,
+		const bool HasFragmentDensityAttachment
 		)
 	{
 		PSOContext.CachedNumSimultanousRenderTargets = NewNumSimultaneousRenderTargets;
@@ -639,13 +640,14 @@ protected:
 		}
 
 		PSOContext.CachedDepthStencilTarget = (NewDepthStencilTargetRHI) ? *NewDepthStencilTargetRHI : FRHIDepthRenderTargetView();
+		PSOContext.HasFragmentDensityAttachment = HasFragmentDensityAttachment;
 	}
 
 	void CacheActiveRenderTargets(const FRHIRenderPassInfo& Info)
 	{
 		FRHISetRenderTargetsInfo RTInfo;
 		Info.ConvertToRenderTargetsInfo(RTInfo);
-		CacheActiveRenderTargets(RTInfo.NumColorRenderTargets, RTInfo.ColorRenderTarget, &RTInfo.DepthStencilRenderTarget);
+		CacheActiveRenderTargets(RTInfo.NumColorRenderTargets, RTInfo.ColorRenderTarget, &RTInfo.DepthStencilRenderTarget, RTInfo.FoveationTexture != nullptr);
 	}
 
 	void IncrementSubpass()
@@ -3129,6 +3131,7 @@ public:
 
 		GraphicsPSOInit.SubpassHint = PSOContext.SubpassHint;
 		GraphicsPSOInit.SubpassIndex = PSOContext.SubpassIndex;
+		GraphicsPSOInit.bHasFragmentDensityAttachment = PSOContext.HasFragmentDensityAttachment;
 	}
 
 	UE_DEPRECATED(4.22, "SetRenderTargets API is deprecated; please use RHIBegin/EndRenderPass instead.")
@@ -3141,7 +3144,8 @@ public:
 		CacheActiveRenderTargets(
 			NewNumSimultaneousRenderTargets, 
 			NewRenderTargetsRHI, 
-			NewDepthStencilTargetRHI
+			NewDepthStencilTargetRHI,
+			false
 			);
 
 		if (Bypass())
diff --git a/Engine/Source/Runtime/RHI/Public/RHIResources.h b/Engine/Source/Runtime/RHI/Public/RHIResources.h
--- a/Engine/Source/Runtime/RHI/Public/RHIResources.h
+++ b/Engine/Source/Runtime/RHI/Public/RHIResources.h
@@ -1967,6 +1967,7 @@ public:
 		, SubpassIndex(0)
 		, bDepthBounds(false)
 		, bMultiView(false)
+		, bHasFragmentDensityAttachment(false)
 		, Flags(0)
 	{
 		static_assert(sizeof(EPixelFormat) != sizeof(uint8), "Change TRenderTargetFormats's uint8 to EPixelFormat");
@@ -1995,7 +1996,8 @@ public:
 		uint8						InSubpassIndex,
 		uint16						InFlags,
 		bool						bInDepthBounds,
-		bool						bInMultiView
+		bool						bInMultiView,
+		bool						bHasFragmentDensityAttachment
 		)
 		: BoundShaderState(InBoundShaderState)
 		, BlendState(InBlendState)
@@ -2018,6 +2020,7 @@ public:
 		, SubpassIndex(InSubpassIndex)
 		, bDepthBounds(bInDepthBounds)
 		, bMultiView(bInMultiView)
+		, bHasFragmentDensityAttachment(bHasFragmentDensityAttachment)
 		, Flags(InFlags)
 	{
 	}
@@ -2041,6 +2044,7 @@ public:
 			PrimitiveType != rhs.PrimitiveType ||
 			bDepthBounds != rhs.bDepthBounds ||
 			bMultiView != rhs.bMultiView ||
+			bHasFragmentDensityAttachment != rhs.bHasFragmentDensityAttachment ||
 			RenderTargetsEnabled != rhs.RenderTargetsEnabled ||
 			RenderTargetFormats != rhs.RenderTargetFormats || 
 			RenderTargetFlags != rhs.RenderTargetFlags || 
@@ -2102,6 +2106,7 @@ public:
 	uint8							SubpassIndex;
 	bool							bDepthBounds;
 	bool							bMultiView;
+	bool							bHasFragmentDensityAttachment;
 	
 	// Note: these flags do NOT affect compilation of this PSO.
 	// The resulting object is invariant with respect to whatever is set here, they are
diff --git a/Engine/Source/Runtime/Renderer/Public/MeshPassProcessor.h b/Engine/Source/Runtime/Renderer/Public/MeshPassProcessor.h
--- a/Engine/Source/Runtime/Renderer/Public/MeshPassProcessor.h
+++ b/Engine/Source/Runtime/Renderer/Public/MeshPassProcessor.h
@@ -270,6 +270,7 @@ public:
 			, 0
 			, bDepthBounds
 			, bMultiView
+			, bHasFragmentDensityAttachment
 		);
 	}
 
@@ -296,6 +297,7 @@ public:
 			ImmutableSamplerState != rhs.ImmutableSamplerState ||
 			bDepthBounds != rhs.bDepthBounds ||
 			bMultiView != rhs.bMultiView ||
+			bHasFragmentDensityAttachment != rhs.bHasFragmentDensityAttachment ||
 			PrimitiveType != rhs.PrimitiveType)
 		{
 			return false;
@@ -362,6 +364,7 @@ public:
 			COMPARE_FIELD(DepthStencilState)
 			COMPARE_FIELD(bDepthBounds)
 			COMPARE_FIELD(bMultiView)
+			COMPARE_FIELD(bHasFragmentDensityAttachment)
 			COMPARE_FIELD(PrimitiveType)
 		COMPARE_FIELD_END;
 
@@ -392,6 +395,7 @@ public:
 			COMPARE_FIELD(DepthStencilState)
 			COMPARE_FIELD(bDepthBounds)
 			COMPARE_FIELD(bMultiView)
+			COMPARE_FIELD(bHasFragmentDensityAttachment)
 			COMPARE_FIELD(PrimitiveType)
 			COMPARE_FIELD_END;
 
@@ -415,7 +419,8 @@ public:
 	// all data members and at the end of the structure.
 	bool							bDepthBounds = false;
 	bool							bMultiView = false;
-	uint8							Padding[2] = {};
+	bool							bHasFragmentDensityAttachment = false;
+	uint8							Padding[1] = {};
 
 	EPrimitiveType			PrimitiveType;
 };
diff --git a/Engine/Source/Runtime/VulkanRHI/Private/VulkanRenderTarget.cpp b/Engine/Source/Runtime/VulkanRHI/Private/VulkanRenderTarget.cpp
--- a/Engine/Source/Runtime/VulkanRHI/Private/VulkanRenderTarget.cpp
+++ b/Engine/Source/Runtime/VulkanRHI/Private/VulkanRenderTarget.cpp
@@ -2432,6 +2432,37 @@ FVulkanRenderTargetLayout::FVulkanRenderTargetLayout(const FGraphicsPipelineStat
 		bHasDepthStencil = true;
 	}
 
+	if (Initializer.bHasFragmentDensityAttachment)
+	{
+		VkAttachmentDescription& CurrDesc = Desc[NumAttachmentDescriptions];
+		FMemory::Memzero(CurrDesc);
+
+		const VkImageLayout FragmentDensityLayout = VK_IMAGE_LAYOUT_FRAGMENT_DENSITY_MAP_OPTIMAL_EXT;
+
+		CurrDesc.flags = 0;
+		CurrDesc.format = VK_FORMAT_R8G8_UNORM;
+		CurrDesc.samples = VK_SAMPLE_COUNT_1_BIT;
+		CurrDesc.loadOp = VK_ATTACHMENT_LOAD_OP_DONT_CARE;
+		CurrDesc.storeOp = VK_ATTACHMENT_STORE_OP_DONT_CARE;
+		CurrDesc.stencilLoadOp = VK_ATTACHMENT_LOAD_OP_DONT_CARE;
+		CurrDesc.stencilStoreOp = VK_ATTACHMENT_STORE_OP_DONT_CARE;
+		CurrDesc.initialLayout = FragmentDensityLayout;
+		CurrDesc.finalLayout = FragmentDensityLayout;
+
+		FragmentDensityReference.attachment = NumAttachmentDescriptions;
+		FragmentDensityReference.layout = FragmentDensityLayout;
+
+		FullHashInfo.LoadOps[MaxSimultaneousRenderTargets + 2] = CurrDesc.stencilLoadOp;
+		FullHashInfo.StoreOps[MaxSimultaneousRenderTargets + 2] = CurrDesc.stencilStoreOp;
+#if VULKAN_USE_REAL_RENDERPASS_COMPATIBILITY
+		FullHashInfo.InitialLayout[MaxSimultaneousRenderTargets + 1] = FragmentDensityLayout;
+#endif
+		CompatibleHashInfo.Formats[MaxSimultaneousRenderTargets + 1] = CurrDesc.format;
+
+		++NumAttachmentDescriptions;
+		bHasFragmentDensityAttachment = true;
+	}
+
 	SubpassHint = Initializer.SubpassHint;
 	CompatibleHashInfo.SubpassHint = (uint8)Initializer.SubpassHint;
 
