UE4.25 - Renumber deprecated BUF enums to avoid conflict with new ones
#rnx

diff --git a/Engine/Source/Runtime/RHI/Public/RHIDefinitions.h b/Engine/Source/Runtime/RHI/Public/RHIDefinitions.h
--- a/Engine/Source/Runtime/RHI/Public/RHIDefinitions.h
+++ b/Engine/Source/Runtime/RHI/Public/RHIDefinitions.h
@@ -1088,7 +1088,7 @@ enum EBufferUsageFlags
 	BUF_ByteAddressBuffer = 0x0020,
 
 	/** Create a structured buffer with an atomic UAV counter. */
-	BUF_UAVCounter UE_DEPRECATED(4.25, "BUF_UAVCounter is deprecated and getting removed; please don't use.") = 0x0040,
+	BUF_UAVCounter UE_DEPRECATED(4.25, "BUF_UAVCounter is deprecated and getting removed; please don't use.") = 0x10000,
 
 	/** Buffer that the GPU will use as a source for a copy. */
 	BUF_SourceCopy			= 0x0040,
@@ -1122,7 +1122,7 @@ enum EBufferUsageFlags
 	BUF_Transient		  = 0x2000,
 
 	/** Buffer that should be accessed one byte at a time. */
-	BUF_UINT8 UE_DEPRECATED(4.25, "BUF_UINT8 is deprecated and getting removed; please don't use.") = 0x4000,
+	BUF_UINT8 UE_DEPRECATED(4.25, "BUF_UINT8 is deprecated and getting removed; please don't use.") = 0x20000,
 
 	/**
 	 * Buffer contains opaque ray tracing acceleration structure data.
