REL - Add RHI interface to query maximum supported Buffer size.

diff --git a/Engine/Source/Runtime/OpenGLDrv/Private/OpenGLDevice.cpp b/Engine/Source/Runtime/OpenGLDrv/Private/OpenGLDevice.cpp
--- a/Engine/Source/Runtime/OpenGLDrv/Private/OpenGLDevice.cpp
+++ b/Engine/Source/Runtime/OpenGLDrv/Private/OpenGLDevice.cpp
@@ -736,9 +736,6 @@ static void InitRHICapabilitiesForGL()
 	GMaxOpenGLColorSamples = Value_GL_MAX_COLOR_TEXTURE_SAMPLES;
 	GMaxOpenGLDepthSamples = Value_GL_MAX_DEPTH_TEXTURE_SAMPLES;
 	GMaxOpenGLIntegerSamples = Value_GL_MAX_INTEGER_SAMPLES;
-#if defined(GL_MAX_TEXTURE_BUFFER_SIZE)
-	GMaxTextureBufferSize = Value_GL_MAX_TEXTURE_BUFFER_SIZE;
-#endif
 
 	// Verify some assumptions.
 	// Android seems like reports one color attachment even when it supports MRT
@@ -814,6 +811,11 @@ static void InitRHICapabilitiesForGL()
 	GMaxTextureMipCount = FMath::CeilLogTwo(Value_GL_MAX_TEXTURE_SIZE) + 1;
 	GMaxTextureMipCount = FMath::Min<int32>(MAX_TEXTURE_MIP_COUNT, GMaxTextureMipCount);
 	GMaxTextureDimensions = Value_GL_MAX_TEXTURE_SIZE;
+	
+#if defined(GL_MAX_TEXTURE_BUFFER_SIZE)
+	GMaxBufferDimensions = Value_GL_MAX_TEXTURE_BUFFER_SIZE;
+#endif
+
 	GMaxCubeTextureDimensions = Value_GL_MAX_CUBE_MAP_TEXTURE_SIZE;
 #if defined(GL_MAX_ARRAY_TEXTURE_LAYERS) && GL_MAX_ARRAY_TEXTURE_LAYERS
 	GMaxTextureArrayLayers = Value_GL_MAX_ARRAY_TEXTURE_LAYERS;
diff --git a/Engine/Source/Runtime/RHI/Private/RHI.cpp b/Engine/Source/Runtime/RHI/Private/RHI.cpp
--- a/Engine/Source/Runtime/RHI/Private/RHI.cpp
+++ b/Engine/Source/Runtime/RHI/Private/RHI.cpp
@@ -676,6 +676,7 @@ bool GRHISupportsLazyShaderCodeLoading = false;
 TRHIGlobal<int32> GMaxShadowDepthBufferSizeX(2048);
 TRHIGlobal<int32> GMaxShadowDepthBufferSizeY(2048);
 TRHIGlobal<int32> GMaxTextureDimensions(2048);
+TRHIGlobal<int64> GMaxBufferDimensions(2<<27);
 TRHIGlobal<int32> GMaxVolumeTextureDimensions(2048);
 TRHIGlobal<int32> GMaxCubeTextureDimensions(2048);
 int32 GMaxTextureArrayLayers = 256;
@@ -727,7 +728,6 @@ FVertexElementTypeSupportInfo GVertexElementTypeSupport;
 RHI_API int32 volatile GCurrentTextureMemorySize = 0;
 RHI_API int32 volatile GCurrentRendertargetMemorySize = 0;
 RHI_API int64 GTexturePoolSize = 0 * 1024 * 1024;
-RHI_API int64 GMaxTextureBufferSize = 0;
 RHI_API int32 GPoolSizeVRAMPercentage = 0;
 
 RHI_API EShaderPlatform GShaderPlatformForFeatureLevel[ERHIFeatureLevel::Num] = {SP_NumPlatforms,SP_NumPlatforms,SP_NumPlatforms,SP_NumPlatforms};
diff --git a/Engine/Source/Runtime/RHI/Public/RHI.h b/Engine/Source/Runtime/RHI/Public/RHI.h
--- a/Engine/Source/Runtime/RHI/Public/RHI.h
+++ b/Engine/Source/Runtime/RHI/Public/RHI.h
@@ -420,9 +420,16 @@ extern RHI_API TRHIGlobal<int32> GMaxShadowDepthBufferSizeY;
 /** The maximum size allowed for 2D textures in both dimensions. */
 extern RHI_API TRHIGlobal<int32> GMaxTextureDimensions;
 
+/** The maximum size allowed for 2D textures in both dimensions. */
+extern RHI_API TRHIGlobal<int64> GMaxBufferDimensions;
+
 /** The maximum size allowed for 3D textures in all three dimensions. */
 extern RHI_API TRHIGlobal<int32> GMaxVolumeTextureDimensions;
 
+FORCEINLINE uint64 GetMaxBufferDimension()
+{
+	return GMaxBufferDimensions;
+}
 
 FORCEINLINE uint32 GetMax2DTextureDimension()
 {
@@ -485,8 +492,7 @@ extern RHI_API volatile int32 GCurrentTextureMemorySize;
 extern RHI_API volatile int32 GCurrentRendertargetMemorySize;
 /** Current texture streaming pool size, in bytes. 0 means unlimited. */
 extern RHI_API int64 GTexturePoolSize;
-/** Maximum texture buffer size. Relevant for OpenGL. 0 means we do not take it into account*/
-extern RHI_API int64 GMaxTextureBufferSize;
+
 /** In percent. If non-zero, the texture pool size is a percentage of GTotalGraphicsMemory. */
 extern RHI_API int32 GPoolSizeVRAMPercentage;
 
diff --git a/Engine/Source/Runtime/Renderer/Private/GPUScene.cpp b/Engine/Source/Runtime/Renderer/Private/GPUScene.cpp
--- a/Engine/Source/Runtime/Renderer/Private/GPUScene.cpp
+++ b/Engine/Source/Runtime/Renderer/Private/GPUScene.cpp
@@ -219,7 +219,7 @@ void UpdateUniformResource<FRWBufferStructured>(FViewInfo& View, FScene& Scene,
 
 static int32 GetMaxPrimitivesUpdate(uint32 NumUploads, uint32 InStrideInFloat4s)
 {
-	return GMaxTextureBufferSize == 0 ? NumUploads : FMath::Min((uint32)(GMaxTextureBufferSize / InStrideInFloat4s), NumUploads);
+	return FMath::Min((uint32)(GetMaxBufferDimension() / InStrideInFloat4s), NumUploads);
 }
 
 template<typename ResourceType>
diff --git a/Engine/Source/Runtime/VulkanRHI/Private/VulkanRHI.cpp b/Engine/Source/Runtime/VulkanRHI/Private/VulkanRHI.cpp
--- a/Engine/Source/Runtime/VulkanRHI/Private/VulkanRHI.cpp
+++ b/Engine/Source/Runtime/VulkanRHI/Private/VulkanRHI.cpp
@@ -714,6 +714,7 @@ void FVulkanDynamicRHI::InitInstance()
 		GMaxShadowDepthBufferSizeX =  FPlatformMath::Min<int32>(Props.limits.maxImageDimension2D, GMaxShadowDepthBufferSizeX);
 		GMaxShadowDepthBufferSizeY =  FPlatformMath::Min<int32>(Props.limits.maxImageDimension2D, GMaxShadowDepthBufferSizeY);
 		GMaxTextureDimensions = Props.limits.maxImageDimension2D;
+		GMaxBufferDimensions = Props.limits.maxTexelBufferElements;
 		GMaxTextureMipCount = FPlatformMath::CeilLogTwo( GMaxTextureDimensions ) + 1;
 		GMaxTextureMipCount = FPlatformMath::Min<int32>( MAX_TEXTURE_MIP_COUNT, GMaxTextureMipCount );
 		GMaxCubeTextureDimensions = Props.limits.maxImageDimensionCube;
