Add missing forwarding functon for RHIAdvanceFrameFence()

diff --git a/Engine/Source/Runtime/RHI/Public/RHIValidation.h b/Engine/Source/Runtime/RHI/Public/RHIValidation.h
--- a/Engine/Source/Runtime/RHI/Public/RHIValidation.h
+++ b/Engine/Source/Runtime/RHI/Public/RHIValidation.h
@@ -970,6 +970,11 @@ public:
 		return RHI->RHICreateAliasedTexture(SourceTexture);
 	}
 
+	virtual void RHIAdvanceFrameFence() override final
+	{
+		RHI->RHIAdvanceFrameFence();
+	}
+
 	// Only relevant with an RHI thread, this advances the backbuffer for the purpose of GetViewportBackBuffer
 	// FlushType: Thread safe
 	virtual void RHIAdvanceFrameForGetViewportBackBuffer(FRHIViewport* Viewport) override final
