Fix for lens flares causing corruption in splitscreen in XboxOneX

The problem is when DCC is enabled on a source texture CopyTextureRegion() will ignore the source rect and copy the whole texture to the same sub rect in the destination. This is a bug in the XDK.

The workaround is to disable DCC for the lens flare texture.

diff --git a/Engine/Source/Runtime/RHI/Public/RHIDefinitions.h b/Engine/Source/Runtime/RHI/Public/RHIDefinitions.h
--- a/Engine/Source/Runtime/RHI/Public/RHIDefinitions.h
+++ b/Engine/Source/Runtime/RHI/Public/RHIDefinitions.h
@@ -1218,6 +1218,8 @@ enum ETextureCreateFlags
 	TexCreate_FastVRAMPartialAlloc  = 1<<13,
 	// Do not create associated shader resource view, only applicable to D3D11 and D3D12
 	TexCreate_DisableSRVCreation = 1 << 14,
+	// Do not allow Delta Color Compression (DCC) to be used with this texture
+	TexCreate_DisableDCC		    = 1 << 15,
 	// UnorderedAccessView (DX11 only)
 	// Warning: Causes additional synchronization between draw calls when using a render target allocated with this flag, use sparingly
 	// See: GCNPerformanceTweets.pdf Tip 37
diff --git a/Engine/Source/Runtime/Renderer/Private/PostProcess/PostProcessLensFlares.cpp b/Engine/Source/Runtime/Renderer/Private/PostProcess/PostProcessLensFlares.cpp
--- a/Engine/Source/Runtime/Renderer/Private/PostProcess/PostProcessLensFlares.cpp
+++ b/Engine/Source/Runtime/Renderer/Private/PostProcess/PostProcessLensFlares.cpp
@@ -225,6 +225,7 @@ FScreenPassTexture AddLensFlaresPass(
 	// Initialize the lens flare output texture.
 	{
 		FRDGTextureDesc LensFlareTextureDesc = Inputs.Bloom.Texture->Desc;
+		LensFlareTextureDesc.Flags |= TexCreate_DisableDCC; // DCC will cause errors on this resource when running splitscreen so disable
 		LensFlareTextureDesc.Reset();
 		LensFlareTextureDesc.ClearValue = FClearValueBinding(FLinearColor::Transparent);
 
