Add SP_OPENGL_SM4_REMOVED to IsDeprecatedShaderPlatform()

Fixes assert when getting shader stats (called from FShaderCookerStats::Initialize)

diff --git a/Engine/Source/Runtime/RHI/Public/RHIDefinitions.h b/Engine/Source/Runtime/RHI/Public/RHIDefinitions.h
--- a/Engine/Source/Runtime/RHI/Public/RHIDefinitions.h
+++ b/Engine/Source/Runtime/RHI/Public/RHIDefinitions.h
@@ -1472,7 +1472,7 @@ inline bool IsDeprecatedShaderPlatform(const FStaticShaderPlatform ShaderPlatfor
 	return ShaderPlatform == SP_OPENGL_SM5 || ShaderPlatform == SP_PCD3D_SM4_REMOVED || ShaderPlatform == SP_OPENGL_ES2_IOS_REMOVED ||
 		ShaderPlatform == SP_PCD3D_ES2_REMOVED || ShaderPlatform == SP_METAL_MACES2_REMOVED || ShaderPlatform == SP_OPENGL_PCES2_REMOVED ||
 		ShaderPlatform == SP_OPENGL_ES2_ANDROID_REMOVED || ShaderPlatform == SP_OPENGL_ES2_WEBGL_REMOVED ||
-		ShaderPlatform == SP_VULKAN_SM4_REMOVED;
+		ShaderPlatform == SP_VULKAN_SM4_REMOVED || ShaderPlatform == SP_OPENGL_SM4_REMOVED;
 }
 
 inline FStaticFeatureLevel GetMaxSupportedFeatureLevel(const FStaticShaderPlatform InShaderPlatform)
