Unified handling of element count in RHICreateShaderResourceView.

Each RHI was doing its own validation of the requested SRV size, which resulted in inconsistent behavior: some were clamping the count to make sure it fits in the underlying buffer, some were checking the offset, some weren't doing anything. In particular, Vulkan and Metal didn't know that a size of UINT32_MAX meant map the entire buffer, which led to problems, as the high level code recently started using this feature. Moved all this into a common place to get consistent behavior everywhere (namely, the element count is clamped to make sure the view doesn't extend past the end of the buffer; this was what D3D11 and 12 were doing already).

#rnx

diff --git a/Engine/Source/Runtime/RHI/Private/DynamicRHI.cpp b/Engine/Source/Runtime/RHI/Private/DynamicRHI.cpp
--- a/Engine/Source/Runtime/RHI/Private/DynamicRHI.cpp
+++ b/Engine/Source/Runtime/RHI/Private/DynamicRHI.cpp
@@ -14,6 +14,7 @@
 #include "GenericPlatform/GenericPlatformDriver.h"
 #include "GenericPlatform/GenericPlatformCrashContext.h"
 #include "PipelineStateCache.h"
+#include "RenderUtils.h"
 
 IMPLEMENT_TYPE_LAYOUT(FRayTracingGeometryInitializer);
 IMPLEMENT_TYPE_LAYOUT(FRayTracingGeometrySegment);
@@ -454,52 +455,125 @@ void FDynamicRHI::RHICheckViewportHDRStatus(FRHIViewport* Viewport)
 {
 }
 
+static void CreateVertexBufferSRVInitializer(FShaderResourceViewInitializer::FVertexBufferShaderResourceViewInitializer& Initializer, FRHIVertexBuffer* InVertexBuffer, EPixelFormat InFormat, uint32 InStartOffsetBytes, uint32 InNumElements)
+{
+	uint32 MaxNumElements, MaxOffset;
+	if (InVertexBuffer)
+	{
+		const uint32 BufferSize = InVertexBuffer->GetSize();
+		const uint32 ElementSize = GPixelFormats[InFormat].BlockBytes;
+		MaxNumElements = (BufferSize - InStartOffsetBytes) / ElementSize;
+		const uint32 BufferAlignment = RHIGetMinimumAlignmentForBufferBackedSRV(InFormat);
+		MaxOffset = (BufferSize > 0) ? (BufferSize - 1) / BufferAlignment * BufferAlignment : 0;
+
+		checkf(InStartOffsetBytes % BufferAlignment == 0, TEXT("Requested vertex buffer offset is not aligned correctly. InStartOffsetBytes=%u, BufferAlignment=%u"), InStartOffsetBytes, BufferAlignment);
+
+		// UINT32_MAX is fine, it means the caller wants the view to extend to the end of the buffer. For explicitly provided element counts, make sure we're not requesting more data than available.
+		checkf(InNumElements == UINT32_MAX || InNumElements <= MaxNumElements,
+			TEXT("Requested view extends past the end of the vertex buffer. InNumElements=%u, InFormat=%d, InStartOffsetBytes=%u, BufferSize=%u, ElementSize=%u, MaxNumElements=%u"),
+			InNumElements, (int)InFormat, InStartOffsetBytes, BufferSize, ElementSize, MaxNumElements);
+	}
+	else
+	{
+		MaxNumElements = 0;
+		MaxOffset = 0;
+	}
+
+	Initializer.VertexBuffer = InVertexBuffer;
+	Initializer.StartOffsetBytes = FMath::Min(InStartOffsetBytes, MaxOffset);
+	Initializer.NumElements = FMath::Min(InNumElements, MaxNumElements);
+	Initializer.Format = InFormat;
+}
 
 FShaderResourceViewInitializer::FShaderResourceViewInitializer(FRHIVertexBuffer* InVertexBuffer, EPixelFormat InFormat, uint32 InStartOffsetBytes, uint32 InNumElements)
-	: VertexBufferInitializer({ InVertexBuffer, InStartOffsetBytes, InNumElements, InFormat }), Type(EType::VertexBufferSRV)
+	: Type(EType::VertexBufferSRV)
 {
-	check(InStartOffsetBytes % RHIGetMinimumAlignmentForBufferBackedSRV(InFormat) == 0);
-	/*if (!VertexBufferInitializer.IsWholeResource())
-	{
-		const uint32 Stride = GPixelFormats[InFormat].BlockBytes;
-		check((VertexBufferInitializer.NumElements * Stride + VertexBufferInitializer.StartOffsetBytes)  <= VertexBufferInitializer.VertexBuffer->GetSize());
-	}*/
+	CreateVertexBufferSRVInitializer(VertexBufferInitializer, InVertexBuffer, InFormat, InStartOffsetBytes, InNumElements);
 }
 
 FShaderResourceViewInitializer::FShaderResourceViewInitializer(FRHIVertexBuffer* InVertexBuffer, EPixelFormat InFormat)
-	: VertexBufferInitializer({ InVertexBuffer, 0, UINT32_MAX, InFormat }), Type(EType::VertexBufferSRV) 
+	: Type(EType::VertexBufferSRV) 
 {
+	CreateVertexBufferSRVInitializer(VertexBufferInitializer, InVertexBuffer, InFormat, 0, UINT32_MAX);
 }
 
-FShaderResourceViewInitializer::FShaderResourceViewInitializer(FRHIStructuredBuffer* InStructuredBuffer, uint32 InStartOffsetBytes, uint32 InNumElements)
-	: StructuredBufferInitializer(FStructuredBufferShaderResourceViewInitializer{ InStructuredBuffer, InStartOffsetBytes, InNumElements }), Type(EType::StructuredBufferSRV)
+static void CreateStructuredBufferSRVInitializer(FShaderResourceViewInitializer::FStructuredBufferShaderResourceViewInitializer& Initializer, FRHIStructuredBuffer* InStructuredBuffer, uint32 InStartOffsetBytes, uint32 InNumElements)
 {
-	check(InStartOffsetBytes % InStructuredBuffer->GetStride() == 0);
-	if (!StructuredBufferInitializer.IsWholeResource())
+	uint32 MaxNumElements, MaxOffset;
+	if (InStructuredBuffer)
+	{
+		const bool bIsByteAddressBuffer = (InStructuredBuffer->GetUsage() & BUF_ByteAddressBuffer) != 0;
+		const uint32 BufferSize = InStructuredBuffer->GetSize();
+		const uint32 ElementSize = bIsByteAddressBuffer ? 4 : InStructuredBuffer->GetStride();
+		MaxNumElements = (BufferSize - InStartOffsetBytes) / ElementSize;
+		MaxOffset = (BufferSize > 0) ? (BufferSize - 1) / ElementSize * ElementSize : 0;
+
+		checkf(InStartOffsetBytes % ElementSize == 0, TEXT("Requested structured buffer offset is not aligned correctly. InStartOffsetBytes=%u, ElementSize=%u"), InStartOffsetBytes, ElementSize);
+
+		checkf(InNumElements == UINT32_MAX || InNumElements <= MaxNumElements,
+			TEXT("Requested view extends past the end of the structured buffer. InNumElements=%u, InStartOffsetBytes=%u, BufferSize=%u, bIsByteAddressBuffer=%d, ElementSize=%u, MaxNumElements=%u"),
+			InNumElements, InStartOffsetBytes, BufferSize, (int)bIsByteAddressBuffer, ElementSize, MaxNumElements);
+	}
+	else
 	{
-		const uint32 Stride = StructuredBufferInitializer.StructuredBuffer->GetStride();
-		check((StructuredBufferInitializer.NumElements * Stride + StructuredBufferInitializer.StartOffsetBytes)  <= StructuredBufferInitializer.StructuredBuffer->GetSize());
+		MaxNumElements = 0;
+		MaxOffset = 0;
 	}
+
+	Initializer.StructuredBuffer = InStructuredBuffer;
+	Initializer.StartOffsetBytes = FMath::Min(InStartOffsetBytes, MaxOffset);
+	Initializer.NumElements = FMath::Min(InNumElements, MaxNumElements);
+}
+
+FShaderResourceViewInitializer::FShaderResourceViewInitializer(FRHIStructuredBuffer* InStructuredBuffer, uint32 InStartOffsetBytes, uint32 InNumElements)
+	: Type(EType::StructuredBufferSRV)
+{
+	CreateStructuredBufferSRVInitializer(StructuredBufferInitializer, InStructuredBuffer, InStartOffsetBytes, InNumElements);
 }
 
 FShaderResourceViewInitializer::FShaderResourceViewInitializer(FRHIStructuredBuffer* InStructuredBuffer)
-	: StructuredBufferInitializer(FStructuredBufferShaderResourceViewInitializer{ InStructuredBuffer, 0, UINT32_MAX }), Type(EType::StructuredBufferSRV) 
+	: Type(EType::StructuredBufferSRV) 
 {
+	CreateStructuredBufferSRVInitializer(StructuredBufferInitializer, InStructuredBuffer, 0, UINT32_MAX);
 }
 
-FShaderResourceViewInitializer::FShaderResourceViewInitializer(FRHIIndexBuffer* InIndexBuffer, uint32 InStartOffsetBytes, uint32 InNumElements)
-	: IndexBufferInitializer(FIndexBufferShaderResourceViewInitializer{ InIndexBuffer, InStartOffsetBytes, InNumElements }), Type(EType::IndexBufferSRV)
+static void CreateIndexBufferSRVInitializer(FShaderResourceViewInitializer::FIndexBufferShaderResourceViewInitializer& Initializer, FRHIIndexBuffer* InIndexBuffer, uint32 InStartOffsetBytes, uint32 InNumElements)
 {
-	check(InStartOffsetBytes % RHIGetMinimumAlignmentForBufferBackedSRV(InIndexBuffer->GetStride() == 2 ? PF_R16_UINT : PF_R32_UINT) == 0);
-	if (!IndexBufferInitializer.IsWholeResource())
+	uint32 MaxNumElements, MaxOffset;
+	if (InIndexBuffer)
 	{
-		const uint32 Stride = IndexBufferInitializer.IndexBuffer->GetStride();
-		check((IndexBufferInitializer.NumElements * Stride + IndexBufferInitializer.StartOffsetBytes) <= IndexBufferInitializer.IndexBuffer->GetSize());
+		const uint32 BufferSize = InIndexBuffer->GetSize();
+		const uint32 ElementSize = InIndexBuffer->GetStride();
+		MaxNumElements = (BufferSize - InStartOffsetBytes) / ElementSize;
+		const uint32 BufferAlignment = RHIGetMinimumAlignmentForBufferBackedSRV(ElementSize == 2 ? PF_R16_UINT : PF_R32_UINT);
+		MaxOffset = (BufferSize > 0) ? (BufferSize - 1) / BufferAlignment * BufferAlignment : 0;
+
+		checkf(InStartOffsetBytes % BufferAlignment == 0, TEXT("Requested index buffer offset is not aligned correctly. InStartOffsetBytes=%u, BufferAlignment=%u"), InStartOffsetBytes, BufferAlignment);
+
+		checkf(InNumElements == UINT32_MAX || InNumElements <= MaxNumElements,
+			TEXT("Requested view extends past the end of the index buffer. InNumElements=%u, InStartOffsetBytes=%u, BufferSize=%u, ElementSize=%u, MaxNumElements=%u"),
+			InNumElements, InStartOffsetBytes, BufferSize, ElementSize, MaxNumElements);
 	}
+	else
+	{
+		MaxNumElements = 0;
+		MaxOffset = 0;
+	}
+
+	Initializer.IndexBuffer = InIndexBuffer;
+	Initializer.StartOffsetBytes = FMath::Min(InStartOffsetBytes, MaxOffset);
+	Initializer.NumElements = FMath::Min(InNumElements, MaxNumElements);
+}
+
+FShaderResourceViewInitializer::FShaderResourceViewInitializer(FRHIIndexBuffer* InIndexBuffer, uint32 InStartOffsetBytes, uint32 InNumElements)
+	: Type(EType::IndexBufferSRV)
+{
+	CreateIndexBufferSRVInitializer(IndexBufferInitializer, InIndexBuffer, InStartOffsetBytes, InNumElements);
 }
 
 FShaderResourceViewInitializer::FShaderResourceViewInitializer(FRHIIndexBuffer* InIndexBuffer)
-	: IndexBufferInitializer(FIndexBufferShaderResourceViewInitializer{ InIndexBuffer, 0, UINT32_MAX }), Type(EType::IndexBufferSRV) 
+	: Type(EType::IndexBufferSRV) 
 {
+	CreateIndexBufferSRVInitializer(IndexBufferInitializer, InIndexBuffer, 0, UINT32_MAX);
 }
 
diff --git a/Engine/Source/Runtime/RHI/Public/DynamicRHI.h b/Engine/Source/Runtime/RHI/Public/DynamicRHI.h
--- a/Engine/Source/Runtime/RHI/Public/DynamicRHI.h
+++ b/Engine/Source/Runtime/RHI/Public/DynamicRHI.h
@@ -222,11 +222,6 @@ struct FShaderResourceViewInitializer
 		uint32 StartOffsetBytes;
 		uint32 NumElements;
 		EPixelFormat Format;
-
-		inline bool IsWholeResource() const 
-		{ 
-			return StartOffsetBytes == 0 && NumElements == UINT32_MAX;
-		}
 	};
 
 	RHI_API FShaderResourceViewInitializer(FRHIVertexBuffer* InVertexBuffer, EPixelFormat InFormat, uint32 InStartOffsetBytes, uint32 InNumElements);
@@ -243,11 +238,6 @@ struct FShaderResourceViewInitializer
 		FRHIStructuredBuffer* StructuredBuffer;
 		uint32 StartOffsetBytes;
 		uint32 NumElements;
-
-		inline bool IsWholeResource() const 
-		{
-			return StartOffsetBytes == 0 && NumElements == UINT32_MAX;
-		}
 	};
 
 	RHI_API FShaderResourceViewInitializer(FRHIStructuredBuffer* InStructuredBuffer, uint32 InStartOffsetBytes, uint32 InNumElements);
@@ -264,11 +254,6 @@ struct FShaderResourceViewInitializer
 		FRHIIndexBuffer* IndexBuffer;
 		uint32 StartOffsetBytes;
 		uint32 NumElements;
-
-		inline bool IsWholeResource() const 
-		{ 
-			return StartOffsetBytes == 0 && NumElements == UINT32_MAX;
-		}
 	};
 
 	RHI_API FShaderResourceViewInitializer(FRHIIndexBuffer* InIndexBuffer, uint32 InStartOffsetBytes, uint32 InNumElements);
diff --git a/Engine/Source/Runtime/RHI/Public/RHIResources.h b/Engine/Source/Runtime/RHI/Public/RHIResources.h
--- a/Engine/Source/Runtime/RHI/Public/RHIResources.h
+++ b/Engine/Source/Runtime/RHI/Public/RHIResources.h
@@ -606,7 +606,7 @@ public:
 	, Usage(InUsage)
 	{}
 
-	/** @return The stride in bytes of the structured buffer; must be 2 or 4. */
+	/** @return The stride in bytes of the structured buffer. */
 	uint32 GetStride() const { return Stride; }
 
 	/** @return The number of bytes in the structured buffer. */
diff --git a/Engine/Source/Runtime/RHI/RHI.Build.cs b/Engine/Source/Runtime/RHI/RHI.Build.cs
--- a/Engine/Source/Runtime/RHI/RHI.Build.cs
+++ b/Engine/Source/Runtime/RHI/RHI.Build.cs
@@ -10,6 +10,7 @@ public class RHI : ModuleRules
 		PrivateDependencyModuleNames.Add("Core");
 		PrivateDependencyModuleNames.Add("TraceLog");
 		PrivateDependencyModuleNames.Add("ApplicationCore");
+		PrivateDependencyModuleNames.Add("RenderCore");
 
 		if (Target.bCompileAgainstEngine)
 		{
