Undo changelist 13046654

#rnx

diff --git a/Engine/Source/Runtime/RHI/Private/DynamicRHI.cpp b/Engine/Source/Runtime/RHI/Private/DynamicRHI.cpp
--- a/Engine/Source/Runtime/RHI/Private/DynamicRHI.cpp
+++ b/Engine/Source/Runtime/RHI/Private/DynamicRHI.cpp
@@ -14,7 +14,6 @@
 #include "GenericPlatform/GenericPlatformDriver.h"
 #include "GenericPlatform/GenericPlatformCrashContext.h"
 #include "PipelineStateCache.h"
-#include "RenderUtils.h"
 
 IMPLEMENT_TYPE_LAYOUT(FRayTracingGeometryInitializer);
 IMPLEMENT_TYPE_LAYOUT(FRayTracingGeometrySegment);
@@ -455,125 +454,52 @@ void FDynamicRHI::RHICheckViewportHDRStatus(FRHIViewport* Viewport)
 {
 }
 
-static void CreateVertexBufferSRVInitializer(FShaderResourceViewInitializer::FVertexBufferShaderResourceViewInitializer& Initializer, FRHIVertexBuffer* InVertexBuffer, EPixelFormat InFormat, uint32 InStartOffsetBytes, uint32 InNumElements)
-{
-	uint32 MaxNumElements, MaxOffset;
-	if (InVertexBuffer)
-	{
-		const uint32 BufferSize = InVertexBuffer->GetSize();
-		const uint32 ElementSize = GPixelFormats[InFormat].BlockBytes;
-		MaxNumElements = (BufferSize - InStartOffsetBytes) / ElementSize;
-		const uint32 BufferAlignment = RHIGetMinimumAlignmentForBufferBackedSRV(InFormat);
-		MaxOffset = (BufferSize > 0) ? (BufferSize - 1) / BufferAlignment * BufferAlignment : 0;
-
-		checkf(InStartOffsetBytes % BufferAlignment == 0, TEXT("Requested vertex buffer offset is not aligned correctly. InStartOffsetBytes=%u, BufferAlignment=%u"), InStartOffsetBytes, BufferAlignment);
-
-		// UINT32_MAX is fine, it means the caller wants the view to extend to the end of the buffer. For explicitly provided element counts, make sure we're not requesting more data than available.
-		checkf(InNumElements == UINT32_MAX || InNumElements <= MaxNumElements,
-			TEXT("Requested view extends past the end of the vertex buffer. InNumElements=%u, InFormat=%d, InStartOffsetBytes=%u, BufferSize=%u, ElementSize=%u, MaxNumElements=%u"),
-			InNumElements, (int)InFormat, InStartOffsetBytes, BufferSize, ElementSize, MaxNumElements);
-	}
-	else
-	{
-		MaxNumElements = 0;
-		MaxOffset = 0;
-	}
-
-	Initializer.VertexBuffer = InVertexBuffer;
-	Initializer.StartOffsetBytes = FMath::Min(InStartOffsetBytes, MaxOffset);
-	Initializer.NumElements = FMath::Min(InNumElements, MaxNumElements);
-	Initializer.Format = InFormat;
-}
 
 FShaderResourceViewInitializer::FShaderResourceViewInitializer(FRHIVertexBuffer* InVertexBuffer, EPixelFormat InFormat, uint32 InStartOffsetBytes, uint32 InNumElements)
-	: Type(EType::VertexBufferSRV)
+	: VertexBufferInitializer({ InVertexBuffer, InStartOffsetBytes, InNumElements, InFormat }), Type(EType::VertexBufferSRV)
 {
-	CreateVertexBufferSRVInitializer(VertexBufferInitializer, InVertexBuffer, InFormat, InStartOffsetBytes, InNumElements);
+	check(InStartOffsetBytes % RHIGetMinimumAlignmentForBufferBackedSRV(InFormat) == 0);
+	/*if (!VertexBufferInitializer.IsWholeResource())
+	{
+		const uint32 Stride = GPixelFormats[InFormat].BlockBytes;
+		check((VertexBufferInitializer.NumElements * Stride + VertexBufferInitializer.StartOffsetBytes)  <= VertexBufferInitializer.VertexBuffer->GetSize());
+	}*/
 }
 
 FShaderResourceViewInitializer::FShaderResourceViewInitializer(FRHIVertexBuffer* InVertexBuffer, EPixelFormat InFormat)
-	: Type(EType::VertexBufferSRV) 
+	: VertexBufferInitializer({ InVertexBuffer, 0, UINT32_MAX, InFormat }), Type(EType::VertexBufferSRV) 
 {
-	CreateVertexBufferSRVInitializer(VertexBufferInitializer, InVertexBuffer, InFormat, 0, UINT32_MAX);
 }
 
-static void CreateStructuredBufferSRVInitializer(FShaderResourceViewInitializer::FStructuredBufferShaderResourceViewInitializer& Initializer, FRHIStructuredBuffer* InStructuredBuffer, uint32 InStartOffsetBytes, uint32 InNumElements)
+FShaderResourceViewInitializer::FShaderResourceViewInitializer(FRHIStructuredBuffer* InStructuredBuffer, uint32 InStartOffsetBytes, uint32 InNumElements)
+	: StructuredBufferInitializer(FStructuredBufferShaderResourceViewInitializer{ InStructuredBuffer, InStartOffsetBytes, InNumElements }), Type(EType::StructuredBufferSRV)
 {
-	uint32 MaxNumElements, MaxOffset;
-	if (InStructuredBuffer)
-	{
-		const bool bIsByteAddressBuffer = (InStructuredBuffer->GetUsage() & BUF_ByteAddressBuffer) != 0;
-		const uint32 BufferSize = InStructuredBuffer->GetSize();
-		const uint32 ElementSize = bIsByteAddressBuffer ? 4 : InStructuredBuffer->GetStride();
-		MaxNumElements = (BufferSize - InStartOffsetBytes) / ElementSize;
-		MaxOffset = (BufferSize > 0) ? (BufferSize - 1) / ElementSize * ElementSize : 0;
-
-		checkf(InStartOffsetBytes % ElementSize == 0, TEXT("Requested structured buffer offset is not aligned correctly. InStartOffsetBytes=%u, ElementSize=%u"), InStartOffsetBytes, ElementSize);
-
-		checkf(InNumElements == UINT32_MAX || InNumElements <= MaxNumElements,
-			TEXT("Requested view extends past the end of the structured buffer. InNumElements=%u, InStartOffsetBytes=%u, BufferSize=%u, bIsByteAddressBuffer=%d, ElementSize=%u, MaxNumElements=%u"),
-			InNumElements, InStartOffsetBytes, BufferSize, (int)bIsByteAddressBuffer, ElementSize, MaxNumElements);
-	}
-	else
+	check(InStartOffsetBytes % InStructuredBuffer->GetStride() == 0);
+	if (!StructuredBufferInitializer.IsWholeResource())
 	{
-		MaxNumElements = 0;
-		MaxOffset = 0;
+		const uint32 Stride = StructuredBufferInitializer.StructuredBuffer->GetStride();
+		check((StructuredBufferInitializer.NumElements * Stride + StructuredBufferInitializer.StartOffsetBytes)  <= StructuredBufferInitializer.StructuredBuffer->GetSize());
 	}
-
-	Initializer.StructuredBuffer = InStructuredBuffer;
-	Initializer.StartOffsetBytes = FMath::Min(InStartOffsetBytes, MaxOffset);
-	Initializer.NumElements = FMath::Min(InNumElements, MaxNumElements);
-}
-
-FShaderResourceViewInitializer::FShaderResourceViewInitializer(FRHIStructuredBuffer* InStructuredBuffer, uint32 InStartOffsetBytes, uint32 InNumElements)
-	: Type(EType::StructuredBufferSRV)
-{
-	CreateStructuredBufferSRVInitializer(StructuredBufferInitializer, InStructuredBuffer, InStartOffsetBytes, InNumElements);
 }
 
 FShaderResourceViewInitializer::FShaderResourceViewInitializer(FRHIStructuredBuffer* InStructuredBuffer)
-	: Type(EType::StructuredBufferSRV) 
+	: StructuredBufferInitializer(FStructuredBufferShaderResourceViewInitializer{ InStructuredBuffer, 0, UINT32_MAX }), Type(EType::StructuredBufferSRV) 
 {
-	CreateStructuredBufferSRVInitializer(StructuredBufferInitializer, InStructuredBuffer, 0, UINT32_MAX);
 }
 
-static void CreateIndexBufferSRVInitializer(FShaderResourceViewInitializer::FIndexBufferShaderResourceViewInitializer& Initializer, FRHIIndexBuffer* InIndexBuffer, uint32 InStartOffsetBytes, uint32 InNumElements)
+FShaderResourceViewInitializer::FShaderResourceViewInitializer(FRHIIndexBuffer* InIndexBuffer, uint32 InStartOffsetBytes, uint32 InNumElements)
+	: IndexBufferInitializer(FIndexBufferShaderResourceViewInitializer{ InIndexBuffer, InStartOffsetBytes, InNumElements }), Type(EType::IndexBufferSRV)
 {
-	uint32 MaxNumElements, MaxOffset;
-	if (InIndexBuffer)
+	check(InStartOffsetBytes % RHIGetMinimumAlignmentForBufferBackedSRV(InIndexBuffer->GetStride() == 2 ? PF_R16_UINT : PF_R32_UINT) == 0);
+	if (!IndexBufferInitializer.IsWholeResource())
 	{
-		const uint32 BufferSize = InIndexBuffer->GetSize();
-		const uint32 ElementSize = InIndexBuffer->GetStride();
-		MaxNumElements = (BufferSize - InStartOffsetBytes) / ElementSize;
-		const uint32 BufferAlignment = RHIGetMinimumAlignmentForBufferBackedSRV(ElementSize == 2 ? PF_R16_UINT : PF_R32_UINT);
-		MaxOffset = (BufferSize > 0) ? (BufferSize - 1) / BufferAlignment * BufferAlignment : 0;
-
-		checkf(InStartOffsetBytes % BufferAlignment == 0, TEXT("Requested index buffer offset is not aligned correctly. InStartOffsetBytes=%u, BufferAlignment=%u"), InStartOffsetBytes, BufferAlignment);
-
-		checkf(InNumElements == UINT32_MAX || InNumElements <= MaxNumElements,
-			TEXT("Requested view extends past the end of the index buffer. InNumElements=%u, InStartOffsetBytes=%u, BufferSize=%u, ElementSize=%u, MaxNumElements=%u"),
-			InNumElements, InStartOffsetBytes, BufferSize, ElementSize, MaxNumElements);
+		const uint32 Stride = IndexBufferInitializer.IndexBuffer->GetStride();
+		check((IndexBufferInitializer.NumElements * Stride + IndexBufferInitializer.StartOffsetBytes) <= IndexBufferInitializer.IndexBuffer->GetSize());
 	}
-	else
-	{
-		MaxNumElements = 0;
-		MaxOffset = 0;
-	}
-
-	Initializer.IndexBuffer = InIndexBuffer;
-	Initializer.StartOffsetBytes = FMath::Min(InStartOffsetBytes, MaxOffset);
-	Initializer.NumElements = FMath::Min(InNumElements, MaxNumElements);
-}
-
-FShaderResourceViewInitializer::FShaderResourceViewInitializer(FRHIIndexBuffer* InIndexBuffer, uint32 InStartOffsetBytes, uint32 InNumElements)
-	: Type(EType::IndexBufferSRV)
-{
-	CreateIndexBufferSRVInitializer(IndexBufferInitializer, InIndexBuffer, InStartOffsetBytes, InNumElements);
 }
 
 FShaderResourceViewInitializer::FShaderResourceViewInitializer(FRHIIndexBuffer* InIndexBuffer)
-	: Type(EType::IndexBufferSRV) 
+	: IndexBufferInitializer(FIndexBufferShaderResourceViewInitializer{ InIndexBuffer, 0, UINT32_MAX }), Type(EType::IndexBufferSRV) 
 {
-	CreateIndexBufferSRVInitializer(IndexBufferInitializer, InIndexBuffer, 0, UINT32_MAX);
 }
 
diff --git a/Engine/Source/Runtime/RHI/Public/DynamicRHI.h b/Engine/Source/Runtime/RHI/Public/DynamicRHI.h
--- a/Engine/Source/Runtime/RHI/Public/DynamicRHI.h
+++ b/Engine/Source/Runtime/RHI/Public/DynamicRHI.h
@@ -222,6 +222,11 @@ struct FShaderResourceViewInitializer
 		uint32 StartOffsetBytes;
 		uint32 NumElements;
 		EPixelFormat Format;
+
+		inline bool IsWholeResource() const 
+		{ 
+			return StartOffsetBytes == 0 && NumElements == UINT32_MAX;
+		}
 	};
 
 	RHI_API FShaderResourceViewInitializer(FRHIVertexBuffer* InVertexBuffer, EPixelFormat InFormat, uint32 InStartOffsetBytes, uint32 InNumElements);
@@ -238,6 +243,11 @@ struct FShaderResourceViewInitializer
 		FRHIStructuredBuffer* StructuredBuffer;
 		uint32 StartOffsetBytes;
 		uint32 NumElements;
+
+		inline bool IsWholeResource() const 
+		{
+			return StartOffsetBytes == 0 && NumElements == UINT32_MAX;
+		}
 	};
 
 	RHI_API FShaderResourceViewInitializer(FRHIStructuredBuffer* InStructuredBuffer, uint32 InStartOffsetBytes, uint32 InNumElements);
@@ -254,6 +264,11 @@ struct FShaderResourceViewInitializer
 		FRHIIndexBuffer* IndexBuffer;
 		uint32 StartOffsetBytes;
 		uint32 NumElements;
+
+		inline bool IsWholeResource() const 
+		{ 
+			return StartOffsetBytes == 0 && NumElements == UINT32_MAX;
+		}
 	};
 
 	RHI_API FShaderResourceViewInitializer(FRHIIndexBuffer* InIndexBuffer, uint32 InStartOffsetBytes, uint32 InNumElements);
diff --git a/Engine/Source/Runtime/RHI/Public/RHIResources.h b/Engine/Source/Runtime/RHI/Public/RHIResources.h
--- a/Engine/Source/Runtime/RHI/Public/RHIResources.h
+++ b/Engine/Source/Runtime/RHI/Public/RHIResources.h
@@ -606,7 +606,7 @@ public:
 	, Usage(InUsage)
 	{}
 
-	/** @return The stride in bytes of the structured buffer. */
+	/** @return The stride in bytes of the structured buffer; must be 2 or 4. */
 	uint32 GetStride() const { return Stride; }
 
 	/** @return The number of bytes in the structured buffer. */
diff --git a/Engine/Source/Runtime/RHI/RHI.Build.cs b/Engine/Source/Runtime/RHI/RHI.Build.cs
--- a/Engine/Source/Runtime/RHI/RHI.Build.cs
+++ b/Engine/Source/Runtime/RHI/RHI.Build.cs
@@ -10,7 +10,6 @@ public class RHI : ModuleRules
 		PrivateDependencyModuleNames.Add("Core");
 		PrivateDependencyModuleNames.Add("TraceLog");
 		PrivateDependencyModuleNames.Add("ApplicationCore");
-		PrivateDependencyModuleNames.Add("RenderCore");
 
 		if (Target.bCompileAgainstEngine)
 		{
