Fixed xb1 crash on RequestExit(false).
FRHIFrameOffsetThread is trying to return a FEvent to the synch event pool, but the pool has already been destroyed.

#rnx

diff --git a/Engine/Source/Runtime/RHI/Private/RHIUtilities.cpp b/Engine/Source/Runtime/RHI/Private/RHIUtilities.cpp
--- a/Engine/Source/Runtime/RHI/Private/RHIUtilities.cpp
+++ b/Engine/Source/Runtime/RHI/Private/RHIUtilities.cpp
@@ -227,7 +227,7 @@ public:
 	{
 		if (WaitEvent)
 		{
-			FPlatformProcess::ReturnSynchEventToPool(WaitEvent);
+			delete WaitEvent;
 			WaitEvent = nullptr;
 		}
 	}
@@ -303,7 +303,7 @@ private:
 		// Wait event can't be initialized with the singleton, because it will crash when initialized to early
 		if (WaitEvent == nullptr)
 		{
-			WaitEvent = FPlatformProcess::GetSynchEventFromPool(false);
+			WaitEvent = FPlatformProcess::CreateSynchEvent(true);
 		}
 		return WaitEvent;
 	}
