Back out changelist 13150251. Possible cause of PS4 crash

diff --git a/Engine/Source/Runtime/RHI/Private/RHIUtilities.cpp b/Engine/Source/Runtime/RHI/Private/RHIUtilities.cpp
--- a/Engine/Source/Runtime/RHI/Private/RHIUtilities.cpp
+++ b/Engine/Source/Runtime/RHI/Private/RHIUtilities.cpp
@@ -227,7 +227,7 @@ public:
 	{
 		if (WaitEvent)
 		{
-			delete WaitEvent;
+			FPlatformProcess::ReturnSynchEventToPool(WaitEvent);
 			WaitEvent = nullptr;
 		}
 	}
@@ -303,7 +303,7 @@ private:
 		// Wait event can't be initialized with the singleton, because it will crash when initialized to early
 		if (WaitEvent == nullptr)
 		{
-			WaitEvent = FPlatformProcess::CreateSynchEvent(true);
+			WaitEvent = FPlatformProcess::GetSynchEventFromPool(false);
 		}
 		return WaitEvent;
 	}
