Implements FGenericDataDrivenShaderPlatformInfo::GetSupportsGen5TemporalAA()

diff --git a/Engine/Source/Runtime/RHI/Private/RHI.cpp b/Engine/Source/Runtime/RHI/Private/RHI.cpp
--- a/Engine/Source/Runtime/RHI/Private/RHI.cpp
+++ b/Engine/Source/Runtime/RHI/Private/RHI.cpp
@@ -1546,6 +1546,7 @@ void FGenericDataDrivenShaderPlatformInfo::ParseDataDrivenShaderInfo(const FConf
 	Info.bSupportsGPUScene = GetSectionBool(Section, "bSupportsGPUScene");
 	Info.bSupportsWaveOperations = GetSectionBool(Section, "bSupportsWaveOperations");
 	Info.bRequiresExplicit128bitRT = GetSectionBool(Section, "bRequiresExplicit128bitRT");
+	Info.bSupportsGen5TemporalAA = GetSectionBool(Section, "bSupportsGen5TemporalAA");
 	Info.bTargetsTiledGPU = GetSectionBool(Section, "bTargetsTiledGPU");
 	Info.bNeedsOfflineCompiler = GetSectionBool(Section, "bNeedsOfflineCompiler");
 }
diff --git a/Engine/Source/Runtime/RHI/Public/RHIDefinitions.h b/Engine/Source/Runtime/RHI/Public/RHIDefinitions.h
--- a/Engine/Source/Runtime/RHI/Public/RHIDefinitions.h
+++ b/Engine/Source/Runtime/RHI/Public/RHIDefinitions.h
@@ -343,6 +343,7 @@ class RHI_API FGenericDataDrivenShaderPlatformInfo
 	uint32 bSupportsIntrinsicWaveOnce : 1;
 	uint32 bSupportsConservativeRasterization : 1;
 	uint32 bRequiresExplicit128bitRT : 1;
+	uint32 bSupportsGen5TemporalAA : 1;
 
 	uint32 bTargetsTiledGPU: 1;
 	uint32 bNeedsOfflineCompiler: 1;
@@ -585,6 +586,11 @@ public:
 		return Infos[Platform].bSupportsConservativeRasterization;
 	}
 
+	static FORCEINLINE_DEBUGGABLE const bool GetSupportsGen5TemporalAA(const FStaticShaderPlatform Platform)
+	{
+		return Infos[Platform].bSupportsGen5TemporalAA;
+	}
+
 private:
 	static FGenericDataDrivenShaderPlatformInfo Infos[SP_NumPlatforms];
 
diff --git a/Engine/Source/Runtime/Renderer/Private/PostProcess/TemporalAA.cpp b/Engine/Source/Runtime/Renderer/Private/PostProcess/TemporalAA.cpp
--- a/Engine/Source/Runtime/Renderer/Private/PostProcess/TemporalAA.cpp
+++ b/Engine/Source/Runtime/Renderer/Private/PostProcess/TemporalAA.cpp
@@ -92,8 +92,7 @@ inline bool DoesPlatformSupportTemporalHistoryUpscale(EShaderPlatform Platform)
 
 inline bool DoesPlatformSupportGen5TAA(EShaderPlatform Platform)
 {
-	return Platform == SP_PCD3D_SM5;
-	//return IsFeatureLevelSupported(Platform, ERHIFeatureLevel::SM5);
+	return Platform == SP_PCD3D_SM5 || FDataDrivenShaderPlatformInfo::GetSupportsGen5TemporalAA(Platform);
 }
 
 BEGIN_SHADER_PARAMETER_STRUCT(FTAA2CommonParameters, )
