Added new RHIGetHTilePlatformConfig method to query (hidden) platform details around HTile layout

diff --git a/Engine/Source/Runtime/RHI/Public/DynamicRHI.h b/Engine/Source/Runtime/RHI/Public/DynamicRHI.h
--- a/Engine/Source/Runtime/RHI/Public/DynamicRHI.h
+++ b/Engine/Source/Runtime/RHI/Public/DynamicRHI.h
@@ -1081,6 +1081,11 @@ public:
 		return FUnorderedAccessViewRHIRef();
 	}
 
+	virtual uint32 RHIGetHTilePlatformConfig(uint32 DepthWidth, uint32 DepthHeight) const
+	{
+		return 0;
+	}
+
 	virtual FShaderResourceViewRHIRef RHICreateShaderResourceViewHTile(FRHITexture2D* RenderTarget)
 	{
 		return nullptr;
@@ -1576,6 +1581,11 @@ FORCEINLINE FTexture2DRHIRef RHIGetViewportBackBuffer(FRHIViewport* Viewport)
 	return GDynamicRHI->RHIGetViewportBackBuffer(Viewport);
 }
 
+FORCEINLINE uint32 RHIGetHTilePlatformConfig(uint32 DepthWidth, uint32 DepthHeight)
+{
+	return GDynamicRHI->RHIGetHTilePlatformConfig(DepthWidth, DepthHeight);
+}
+
 FORCEINLINE FShaderResourceViewRHIRef RHICreateShaderResourceViewHTile(FRHITexture2D* RenderTarget)
 {
 	return GDynamicRHI->RHICreateShaderResourceViewHTile(RenderTarget);
diff --git a/Engine/Source/Runtime/RHI/Public/RHIValidation.h b/Engine/Source/Runtime/RHI/Public/RHIValidation.h
--- a/Engine/Source/Runtime/RHI/Public/RHIValidation.h
+++ b/Engine/Source/Runtime/RHI/Public/RHIValidation.h
@@ -1171,6 +1171,11 @@ public:
 		return RHI->RHIGetViewportBackBufferUAV(ViewportRHI);
 	}
 
+	virtual uint32 RHIGetHTilePlatformConfig(uint32 DepthWidth, uint32 DepthHeight) const override final
+	{
+		return RHI->RHIGetHTilePlatformConfig(DepthWidth, DepthHeight);
+	}
+
 	virtual FShaderResourceViewRHIRef RHICreateShaderResourceViewHTile(FRHITexture2D* RenderTarget) override final
 	{
 		FShaderResourceViewRHIRef SRV = RHI->RHICreateShaderResourceViewHTile(RenderTarget);
