Temporarily disable some deprecation warnings to fix CIS.

diff --git a/Engine/Source/Runtime/RHI/Public/RHI.h b/Engine/Source/Runtime/RHI/Public/RHI.h
--- a/Engine/Source/Runtime/RHI/Public/RHI.h
+++ b/Engine/Source/Runtime/RHI/Public/RHI.h
@@ -718,10 +718,10 @@ ENUM_CLASS_FLAGS(ERHIPipeline)
 
 enum class EResourceTransitionPipeline
 {
-	EGfxToCompute     UE_DEPRECATED(4.26, "The RHI resource barrier API has been refactored. Use the new RHITransition API. EResourceTransitionPipeline has been replaced with ERHIPipeline bitmask."),
-	EComputeToGfx     UE_DEPRECATED(4.26, "The RHI resource barrier API has been refactored. Use the new RHITransition API. EResourceTransitionPipeline has been replaced with ERHIPipeline bitmask."),
-	EGfxToGfx         UE_DEPRECATED(4.26, "The RHI resource barrier API has been refactored. Use the new RHITransition API. EResourceTransitionPipeline has been replaced with ERHIPipeline bitmask."),
-	EComputeToCompute UE_DEPRECATED(4.26, "The RHI resource barrier API has been refactored. Use the new RHITransition API. EResourceTransitionPipeline has been replaced with ERHIPipeline bitmask.")
+	EGfxToCompute, //     UE_DEPRECATED(4.26, "The RHI resource barrier API has been refactored. Use the new RHITransition API. EResourceTransitionPipeline has been replaced with ERHIPipeline bitmask."),
+	EComputeToGfx, //     UE_DEPRECATED(4.26, "The RHI resource barrier API has been refactored. Use the new RHITransition API. EResourceTransitionPipeline has been replaced with ERHIPipeline bitmask."),
+	EGfxToGfx,     //     UE_DEPRECATED(4.26, "The RHI resource barrier API has been refactored. Use the new RHITransition API. EResourceTransitionPipeline has been replaced with ERHIPipeline bitmask."),
+	EComputeToCompute // UE_DEPRECATED(4.26, "The RHI resource barrier API has been refactored. Use the new RHITransition API. EResourceTransitionPipeline has been replaced with ERHIPipeline bitmask.")
 };
 
 enum class ERHIAccess
@@ -797,16 +797,16 @@ extern RHI_API ERHIAccess GRHITextureReadAccessMask;
 // Helper namespace to maintain compatibility with old renderer code for one engine release.
 namespace EResourceTransitionAccess
 {
-	UE_DEPRECATED(4.26, "The RHI resource barrier API has been refactored. Use the RHITransition API and ERHIAccess to specify explicit previous/next states during resource transitions.")
+	//UE_DEPRECATED(4.26, "The RHI resource barrier API has been refactored. Use the RHITransition API and ERHIAccess to specify explicit previous/next states during resource transitions.")
 	static constexpr const ERHIAccess EReadable    = ERHIAccess::EReadable;
 
-	UE_DEPRECATED(4.26, "The RHI resource barrier API has been refactored. Use the RHITransition API and ERHIAccess to specify explicit previous/next states during resource transitions.")
+	//UE_DEPRECATED(4.26, "The RHI resource barrier API has been refactored. Use the RHITransition API and ERHIAccess to specify explicit previous/next states during resource transitions.")
 	static constexpr const ERHIAccess EWritable    = ERHIAccess::EWritable;
 
-	UE_DEPRECATED(4.26, "The RHI resource barrier API has been refactored. Use the RHITransition API and ERHIAccess to specify explicit previous/next states during resource transitions.")
+	//UE_DEPRECATED(4.26, "The RHI resource barrier API has been refactored. Use the RHITransition API and ERHIAccess to specify explicit previous/next states during resource transitions.")
 	static constexpr const ERHIAccess ERWBarrier   = ERHIAccess::ERWBarrier;
 
-	UE_DEPRECATED(4.26, "The RHI resource barrier API has been refactored. Use the RHITransition API and ERHIAccess to specify explicit previous/next states during resource transitions.")
+	//UE_DEPRECATED(4.26, "The RHI resource barrier API has been refactored. Use the RHITransition API and ERHIAccess to specify explicit previous/next states during resource transitions.")
 	static constexpr const ERHIAccess ERWNoBarrier = ERHIAccess::ERWNoBarrier;
 };
 
