Added RHI capability query for stencil ref export from pixel shader. Enabled on some non-Windows platforms (Windows needs proper capability detection and testing still).

diff --git a/Engine/Source/Runtime/D3D12RHI/Private/D3D12Device.cpp b/Engine/Source/Runtime/D3D12RHI/Private/D3D12Device.cpp
--- a/Engine/Source/Runtime/D3D12RHI/Private/D3D12Device.cpp
+++ b/Engine/Source/Runtime/D3D12RHI/Private/D3D12Device.cpp
@@ -164,6 +164,7 @@ void FD3D12Device::SetupAfterDeviceCreation()
 	ID3D12Device* Direct3DDevice = GetParentAdapter()->GetD3DDevice();
 
 	GRHISupportsArrayIndexFromAnyShader = true;
+	GRHISupportsStencilRefFromPixelShader = false; // TODO: Sort out DXC shader database SM6.0 usage. DX12 supports this feature, but need to improve DXC support.
 
 #if (PLATFORM_WINDOWS || PLATFORM_HOLOLENS)
 	// Check if we're running under GPU capture
diff --git a/Engine/Source/Runtime/RHI/Private/RHI.cpp b/Engine/Source/Runtime/RHI/Private/RHI.cpp
--- a/Engine/Source/Runtime/RHI/Private/RHI.cpp
+++ b/Engine/Source/Runtime/RHI/Private/RHI.cpp
@@ -875,6 +875,7 @@ FIntVector GRHIMaxDispatchThreadGroupsPerDimension(0, 0, 0);
 uint64 GRHIPresentCounter = 1;
 
 bool GRHISupportsArrayIndexFromAnyShader = false;
+bool GRHISupportsStencilRefFromPixelShader = false;
 bool GRHISupportsConservativeRasterization = false;
 
 /** Whether we are profiling GPU hitches. */
diff --git a/Engine/Source/Runtime/RHI/Public/RHI.h b/Engine/Source/Runtime/RHI/Public/RHI.h
--- a/Engine/Source/Runtime/RHI/Public/RHI.h
+++ b/Engine/Source/Runtime/RHI/Public/RHI.h
@@ -624,6 +624,9 @@ extern RHI_API uint64 GRHIPresentCounter;
 /** True if the RHI supports setting the render target array index from any shader stage */
 extern RHI_API bool GRHISupportsArrayIndexFromAnyShader;
 
+/** True if the RHI supports setting the stencil ref at pixel granularity from the pixel shader */
+extern RHI_API bool GRHISupportsStencilRefFromPixelShader;
+
 /** Whether current RHI supports overestimated conservative rasterization. */
 extern RHI_API bool GRHISupportsConservativeRasterization;
 
