Add RayTracingAccellarionStructure state to ERHIAccess

#b Jonas.Meyer

diff --git a/Engine/Source/Runtime/RHI/Public/RHI.h b/Engine/Source/Runtime/RHI/Public/RHI.h
--- a/Engine/Source/Runtime/RHI/Public/RHI.h
+++ b/Engine/Source/Runtime/RHI/Public/RHI.h
@@ -728,25 +728,26 @@ enum class ERHIAccess
 	Unknown = 0,
 
 	// Read states
-	CPURead             = 1 <<  0,
-	Present             = 1 <<  1,
-	IndirectArgs        = 1 <<  2,
-	VertexOrIndexBuffer = 1 <<  3,
-	SRVCompute          = 1 <<  4,
-	SRVGraphics         = 1 <<  5,
-	CopySrc             = 1 <<  6,
-	ResolveSrc          = 1 <<  7,
-	DSVRead				= 1 <<  8,
+	CPURead             	= 1 <<  0,
+	Present             	= 1 <<  1,
+	IndirectArgs        	= 1 <<  2,
+	VertexOrIndexBuffer 	= 1 <<  3,
+	SRVCompute          	= 1 <<  4,
+	SRVGraphics         	= 1 <<  5,
+	CopySrc             	= 1 <<  6,
+	ResolveSrc          	= 1 <<  7,
+	DSVRead					= 1 <<  8,
 
 	// Read-write states
-	UAVCompute          = 1 <<  9,
-	UAVGraphics         = 1 << 10,
-	RTV                 = 1 << 11,
-	CopyDest            = 1 << 12,
-	ResolveDst          = 1 << 13,
-	DSVWrite            = 1 << 14,
-
-	Last = DSVWrite,
+	UAVCompute          	= 1 <<  9,
+	UAVGraphics         	= 1 << 10,
+	RTV                 	= 1 << 11,
+	CopyDest            	= 1 << 12,
+	ResolveDst          	= 1 << 13,
+	DSVWrite            	= 1 << 14,
+	RTAccelerationStructure = 1 << 15,
+
+	Last = RTAccelerationStructure,
 	None = Unknown,
 	Mask = (Last << 1) - 1,
 
@@ -772,7 +773,7 @@ enum class ERHIAccess
 	WriteOnlyMask = WriteOnlyExclusiveMask | DSVWrite,
 
 	// A mask of all bits representing writable states which may also include readable states.
-	WritableMask = WriteOnlyMask | UAVMask
+	WritableMask = WriteOnlyMask | UAVMask | RTAccelerationStructure
 };
 ENUM_CLASS_FLAGS(ERHIAccess)
 
