- Move RHIDecayResourceAccess to RHIValidation code - only used there
- Don't decay AccessMask when UAV overlap is enabled and it's in UAVMask mode (allow overlap between graphics & compute)

diff --git a/Engine/Source/Runtime/RHI/Private/RHIValidation.cpp b/Engine/Source/Runtime/RHI/Private/RHIValidation.cpp
--- a/Engine/Source/Runtime/RHI/Private/RHIValidation.cpp
+++ b/Engine/Source/Runtime/RHI/Private/RHIValidation.cpp
@@ -170,6 +170,48 @@ void FValidationContext::FState::Reset()
 
 namespace RHIValidation
 {
+	ERHIAccess DecayResourceAccess(ERHIAccess AccessMask, ERHIAccess RequiredAccess, bool bAllowUAVOverlap)
+	{
+		using T = __underlying_type(ERHIAccess);
+		checkf((T(RequiredAccess) & (T(RequiredAccess) - 1)) == 0, TEXT("Only one required access bit may be set at once."));
+		
+		if (EnumHasAnyFlags(RequiredAccess, ERHIAccess::UAVMask))
+		{
+			// UAV writes decay to no allowed resource access when overlaps are disabled. A barrier is always required after the dispatch/draw.
+			// Otherwise keep the same accessmask and don't touch or decay the state
+			return !bAllowUAVOverlap ? ERHIAccess::None : AccessMask;
+		}
+
+		// Handle DSV modes
+		if (EnumHasAnyFlags(RequiredAccess, ERHIAccess::DSVWrite))
+		{
+			constexpr ERHIAccess CompatibleStates =
+				ERHIAccess::DSVRead |
+				ERHIAccess::DSVWrite;
+
+			return AccessMask & CompatibleStates;
+		}
+		if (EnumHasAnyFlags(RequiredAccess, ERHIAccess::DSVRead))
+		{
+			constexpr ERHIAccess CompatibleStates =
+				ERHIAccess::DSVRead |
+				ERHIAccess::DSVWrite |
+				ERHIAccess::SRVGraphics |
+				ERHIAccess::SRVCompute;
+
+			return AccessMask & CompatibleStates;
+		}
+
+		if (EnumHasAnyFlags(RequiredAccess, ERHIAccess::WritableMask))
+		{
+			// Decay to only 1 allowed state for all other writable states.
+			return RequiredAccess;
+		}
+
+		// Else, the state is readable. All readable states are compatible.
+		return AccessMask;
+	}
+
 #define BARRIER_TRACKER_LOG_PREFIX TEXT("\n\n")\
 	TEXT("--------------------------------------------------------------------\n")\
 	TEXT("              RHI Resource Transition Validation Error              \n")\
@@ -530,7 +572,7 @@ namespace RHIValidation
 		}
 
 		// Disable all non-compatible access types
-		CurrentState.Access = RHIDecayResourceAccess(CurrentState.Access, RequiredState.Access, bAllowAllUAVsOverlap || bExplicitAllowUAVOverlap);
+		CurrentState.Access = DecayResourceAccess(CurrentState.Access, RequiredState.Access, bAllowAllUAVsOverlap || bExplicitAllowUAVOverlap);
 		CurrentState.Pipelines = CurrentState.Pipelines & RequiredState.Pipelines;
 	}
 
diff --git a/Engine/Source/Runtime/RHI/Public/RHI.h b/Engine/Source/Runtime/RHI/Public/RHI.h
--- a/Engine/Source/Runtime/RHI/Public/RHI.h
+++ b/Engine/Source/Runtime/RHI/Public/RHI.h
@@ -1691,47 +1691,6 @@ inline constexpr bool IsValidAccess(ERHIAccess Access)
 	return !IsInvalidAccess(Access);
 }
 
-inline ERHIAccess RHIDecayResourceAccess(ERHIAccess AccessMask, ERHIAccess RequiredAccess, bool bAllowUAVOverlap)
-{
-	using T = __underlying_type(ERHIAccess);
-	checkf((T(RequiredAccess) & (T(RequiredAccess) - 1)) == 0, TEXT("Only one required access bit may be set at once."));
-
-	if (!bAllowUAVOverlap && EnumHasAnyFlags(RequiredAccess, ERHIAccess::UAVMask))
-	{
-		// UAV writes decay to no allowed resource access when overlaps are disabled. A barrier is always required after the dispatch/draw.
-		return ERHIAccess::None;
-	}
-	
-	// Handle DSV modes
-	if (EnumHasAnyFlags(RequiredAccess, ERHIAccess::DSVWrite))
-	{
-		constexpr ERHIAccess CompatibleStates =
-			ERHIAccess::DSVRead |
-			ERHIAccess::DSVWrite;
-
-		return AccessMask & CompatibleStates;
-	}
-	if (EnumHasAnyFlags(RequiredAccess, ERHIAccess::DSVRead))
-	{
-		constexpr ERHIAccess CompatibleStates =
-			ERHIAccess::DSVRead |
-			ERHIAccess::DSVWrite |
-			ERHIAccess::SRVGraphics |
-			ERHIAccess::SRVCompute;
-
-		return AccessMask & CompatibleStates;
-	}
-
-	if (EnumHasAnyFlags(RequiredAccess, ERHIAccess::WritableMask))
-	{
-		// Decay to only 1 allowed state for all other writable states.
-		return RequiredAccess; 
-	}
-
-	// Else, the state is readable. All readable states are compatible.
-	return AccessMask;
-}
-
 enum class ERHICreateTransitionFlags
 {
 	None = 0,
