Added RHI API support for multi-pipe transitions.

diff --git a/Engine/Source/Runtime/RHI/Private/RHI.cpp b/Engine/Source/Runtime/RHI/Private/RHI.cpp
--- a/Engine/Source/Runtime/RHI/Private/RHI.cpp
+++ b/Engine/Source/Runtime/RHI/Private/RHI.cpp
@@ -903,6 +903,7 @@ RHI_API uint64 GRHITransitionPrivateData_SizeInBytes = 0;
 RHI_API uint64 GRHITransitionPrivateData_AlignInBytes = 0;
 
 ERHIAccess GRHITextureReadAccessMask = ERHIAccess::ReadOnlyMask;
+ERHIAccess GRHIMultiPipelineAccessMask = ERHIAccess::Unknown;
 
 /** Called once per frame only from within an RHI. */
 void RHIPrivateBeginFrame()
diff --git a/Engine/Source/Runtime/RHI/Private/RHIValidation.cpp b/Engine/Source/Runtime/RHI/Private/RHIValidation.cpp
--- a/Engine/Source/Runtime/RHI/Private/RHIValidation.cpp
+++ b/Engine/Source/Runtime/RHI/Private/RHIValidation.cpp
@@ -124,6 +124,116 @@ void FValidationRHI::ValidatePipeline(const FGraphicsPipelineStateInitializer& P
 	}
 }
 
+void FValidationRHI::RHICreateTransition(FRHITransition* Transition, ERHIPipeline SrcPipelines, ERHIPipeline DstPipelines, ERHICreateTransitionFlags CreateFlags, TArrayView<const FRHITransitionInfo> Infos)
+{
+	using namespace RHIValidation;
+
+	struct FFenceEdge
+	{
+		FFence* Fence = nullptr;
+		ERHIPipeline SrcPipe = ERHIPipeline::None;
+		ERHIPipeline DstPipe = ERHIPipeline::None;
+	};
+
+	TArray<FFenceEdge> Fences;
+
+	if (SrcPipelines != DstPipelines)
+	{
+		for (ERHIPipeline SrcPipe : GetRHIPipelines())
+		{
+			if (!EnumHasAnyFlags(SrcPipelines, SrcPipe))
+			{
+				continue;
+			}
+
+			for (ERHIPipeline DstPipe : GetRHIPipelines())
+			{
+				if (!EnumHasAnyFlags(DstPipelines, DstPipe) || SrcPipe == DstPipe)
+				{
+					continue;
+				}
+
+				FFenceEdge FenceEdge;
+				FenceEdge.Fence = new FFence;
+				FenceEdge.SrcPipe = SrcPipe;
+				FenceEdge.DstPipe = DstPipe;
+				Fences.Add(FenceEdge);
+			}
+		}
+	}
+
+	TArray<FOperation> BeginOps, EndOps;
+	BeginOps.Reserve(Infos.Num() + Fences.Num());
+	EndOps  .Reserve(Infos.Num() + Fences.Num());
+
+	for (const FFenceEdge& FenceEdge : Fences)
+	{
+		EndOps.Emplace(FOperation::Wait(FenceEdge.Fence, FenceEdge.DstPipe));
+	}
+
+	bool bDoTrace = false;
+
+	for (int32 Index = 0; Index < Infos.Num(); ++Index)
+	{
+		const FRHITransitionInfo& Info = Infos[Index];
+		if (!Info.Resource)
+			continue;
+
+		checkf(Info.Type != FRHITransitionInfo::EType::Unknown, TEXT("FRHITransitionInfo::Type cannot be Unknown when creating a resource transition."));
+
+		FResourceIdentity Identity;
+
+		switch (Info.Type)
+		{
+		default: checkNoEntry(); // fall through
+		case FRHITransitionInfo::EType::Texture:
+			Identity = Info.Texture->GetTransitionIdentity(Info);
+			break;
+
+		case FRHITransitionInfo::EType::VertexBuffer:
+			Identity = Info.VertexBuffer->GetWholeResourceIdentity();
+			break;
+
+		case FRHITransitionInfo::EType::IndexBuffer:
+			Identity = Info.IndexBuffer->GetWholeResourceIdentity();
+			break;
+
+		case FRHITransitionInfo::EType::StructuredBuffer:
+			Identity = Info.StructuredBuffer->GetWholeResourceIdentity();
+			break;
+
+		case FRHITransitionInfo::EType::UAV:
+			Identity = Info.UAV->ViewIdentity;
+			break;
+		}
+
+		// Take a backtrace of this transition creation if any of the resources it contains have logging enabled.
+		bDoTrace |= (Identity.Resource->LoggingMode != RHIValidation::ELoggingMode::None);
+
+		FState PreviousState = FState(Info.AccessBefore, SrcPipelines);
+		FState NextState = FState(Info.AccessAfter, DstPipelines);
+
+		BeginOps.Emplace(FOperation::BeginTransitionResource(Identity, PreviousState, NextState, Info.Flags, nullptr));
+		EndOps  .Emplace(FOperation::EndTransitionResource(Identity, PreviousState, NextState, nullptr));
+	}
+
+	if (bDoTrace)
+	{
+		void* Backtrace = CaptureBacktrace();
+		for (FOperation& Op : BeginOps) { Op.Data_BeginTransition.CreateBacktrace = Backtrace; }
+		for (FOperation& Op : EndOps) { Op.Data_EndTransition.CreateBacktrace = Backtrace; }
+	}
+
+	for (const FFenceEdge& FenceEdge : Fences)
+	{
+		BeginOps.Emplace(FOperation::Signal(FenceEdge.Fence, FenceEdge.SrcPipe));
+	}
+
+	Transition->PendingOperationsBegin.Operations = MoveTemp(BeginOps);
+	Transition->PendingOperationsEnd.Operations = MoveTemp(EndOps);
+
+	return RHI->RHICreateTransition(Transition, SrcPipelines, DstPipelines, CreateFlags, Infos);
+}
 
 FValidationComputeContext::FValidationComputeContext()
 	: RHIContext(nullptr)
@@ -358,6 +468,26 @@ namespace RHIValidation
 			*GetRHIPipelineName(CurrentStateFromRHI.Pipelines));
 	}
 
+	static inline FString GetReasonString_MismatchedEndTransition(
+		FResource* Resource, FSubresourceIndex const& SubresourceIndex,
+		const FState& TargetState,
+		const FState& TargetStateFromRHI)
+	{
+		const TCHAR* DebugName = Resource->GetDebugName();
+		return FString::Printf(
+			BARRIER_TRACKER_LOG_PREFIX
+			TEXT("The expected target state \"%s\" on pipe \"%s\" in end transition does not match the tracked target state \"%s\" on pipe \"%s\" for the resource \"%s\" (0x%p) (%s).\n")
+			TEXT("    --- The call to EndTransition() is mismatched with the another BeginTransition() with different states.\n")
+			BARRIER_TRACKER_LOG_SUFFIX,
+			*GetRHIAccessName(TargetStateFromRHI.Access),
+			*GetRHIPipelineName(TargetState.Pipelines),
+			*GetRHIAccessName(TargetState.Access),
+			*GetRHIPipelineName(TargetStateFromRHI.Pipelines),
+			DebugName ? DebugName : TEXT("Unnamed"),
+			Resource,
+			*SubresourceIndex.ToString());
+	}
+
 	static inline FString GetReasonString_UnnecessaryTransition(
 		FResource* Resource, FSubresourceIndex const& SubresourceIndex,
 		const FState& CurrentState)
@@ -477,8 +607,10 @@ namespace RHIValidation
 		return Trace;
 	}
 
-	void FSubresourceState::BeginTransition(FResource* Resource, FSubresourceIndex const& SubresourceIndex, const FState& CurrentStateFromRHI, const FState& TargetState, EResourceTransitionFlags NewFlags, void* CreateTrace)
+	void FSubresourceState::BeginTransition(FResource* Resource, FSubresourceIndex const& SubresourceIndex, const FState& CurrentStateFromRHI, const FState& TargetState, EResourceTransitionFlags NewFlags, ERHIPipeline ExecutingPipeline, void* CreateTrace)
 	{
+		FPipelineState& State = States[ExecutingPipeline];
+
 		void* BeginTrace = nullptr;
 		if (Resource->LoggingMode != ELoggingMode::None 
 #if LOG_UNNAMED_RESOURCES
@@ -486,42 +618,52 @@ namespace RHIValidation
 #endif
 			)
 		{
-			BeginTrace = Log(Resource, SubresourceIndex, CreateTrace, TEXT("BeginTransition"), *FString::Printf(TEXT("Current: (%s) New: (%s), Flags: %s"),
-				*CurrentState.ToString(),
+			BeginTrace = Log(Resource, SubresourceIndex, CreateTrace, TEXT("BeginTransition"), *FString::Printf(TEXT("Current: (%s) New: (%s), Flags: %s, Executing Pipeline: %s"),
+				*State.Current.ToString(),
 				*TargetState.ToString(),
-				*GetResourceTransitionFlagsName(NewFlags)
+				*GetResourceTransitionFlagsName(NewFlags),
+				*GetRHIPipelineName(ExecutingPipeline)
 			));
 		}
 
 		// Check we're not already transitioning
-		ensureMsgf(!bTransitioning, TEXT("%s"), *GetReasonString_DuplicateBeginTransition(Resource, SubresourceIndex, CurrentState, TargetState, CreateTransitionBacktrace, BeginTransitionBacktrace));
+		ensureMsgf(!State.bTransitioning, TEXT("%s"), *GetReasonString_DuplicateBeginTransition(Resource, SubresourceIndex, State.Current, TargetState, State.CreateTransitionBacktrace, BeginTrace));
 
 		// Validate the explicit previous state from the RHI matches what we expect...
 		{
 			// Check for the correct pipeline
-			ensureMsgf(EnumHasAllFlags(CurrentState.Pipelines, CurrentStateFromRHI.Pipelines), TEXT("%s"), *GetReasonString_WrongPipeline(Resource, SubresourceIndex, CurrentState, TargetState));
+			ensureMsgf(EnumHasAllFlags(CurrentStateFromRHI.Pipelines, ExecutingPipeline), TEXT("%s"), *GetReasonString_WrongPipeline(Resource, SubresourceIndex, State.Current, TargetState));
 
 			// Check the current RHI state passed in matches the tracked state for the resource, or is "unknown".
-			ensureMsgf(CurrentStateFromRHI.Access == ERHIAccess::Unknown || CurrentStateFromRHI.Access == PreviousState.Access, TEXT("%s"), *GetReasonString_IncorrectPreviousState(Resource, SubresourceIndex, PreviousState, CurrentStateFromRHI));
+			ensureMsgf(CurrentStateFromRHI.Access == ERHIAccess::Unknown || (CurrentStateFromRHI.Access == State.Previous.Access && CurrentStateFromRHI.Pipelines == State.Previous.Pipelines),
+				TEXT("%s"), *GetReasonString_IncorrectPreviousState(Resource, SubresourceIndex, State.Previous, CurrentStateFromRHI));
 		}
 
 		// Check for unnecessary transitions
 		// @todo: this check is not particularly useful at the moment, as there are many unnecessary resource transitions.
 		//ensureMsgf(CurrentState != TargetState, TEXT("%s"), *GetReasonString_UnnecessaryTransition(Resource, SubresourceIndex, CurrentState));
 
-		// Update the tracked state
-		PreviousState = TargetState;
-		CurrentState = TargetState;
-		Flags = NewFlags;
-		CreateTransitionBacktrace = CreateTrace;
-		BeginTransitionBacktrace = BeginTrace;
-		bUsedWithAllUAVsOverlap = false;
-		bUsedWithExplicitUAVsOverlap = false;
-
-		bTransitioning = true;
+		// Update the tracked state once all pipes have begun.
+		State.Previous = TargetState;
+		State.Current = TargetState;
+		State.Flags = NewFlags;
+		State.CreateTransitionBacktrace = CreateTrace;
+		State.BeginTransitionBacktrace = BeginTrace;
+		State.bUsedWithAllUAVsOverlap = false;
+		State.bUsedWithExplicitUAVsOverlap = false;
+		State.bTransitioning = true;
+
+		// Replicate the state to other pipes that are not part of the begin pipe mask.
+		for (ERHIPipeline OtherPipeline : GetRHIPipelines())
+		{
+			if (!EnumHasAnyFlags(CurrentStateFromRHI.Pipelines, OtherPipeline))
+			{
+				States[OtherPipeline] = State;
+			}
+		}
 	}
 
-	void FSubresourceState::EndTransition(FResource* Resource, FSubresourceIndex const& SubresourceIndex, void* CreateTrace)
+	void FSubresourceState::EndTransition(FResource* Resource, FSubresourceIndex const& SubresourceIndex, const FState& CurrentStateFromRHI, const FState& TargetState, ERHIPipeline ExecutingPipeline, void* CreateTrace)
 	{
 		if (Resource->LoggingMode != ELoggingMode::None
 #if LOG_UNNAMED_RESOURCES
@@ -529,16 +671,31 @@ namespace RHIValidation
 #endif
 			)
 		{
-			Log(Resource, SubresourceIndex, CreateTrace, TEXT("EndTransition"), *FString::Printf(TEXT("Access: %s, Pipeline: %s"),
-				*GetRHIAccessName(CurrentState.Access),
-				*GetRHIPipelineName(CurrentState.Pipelines)
+			Log(Resource, SubresourceIndex, CreateTrace, TEXT("EndTransition"), *FString::Printf(TEXT("Access: %s, Pipeline: %s, Executing Pipeline: %s"),
+				*GetRHIAccessName(TargetState.Access),
+				*GetRHIPipelineName(TargetState.Pipelines),
+				*GetRHIPipelineName(ExecutingPipeline)
 			));
 		}
 
-		// This should not really be possible given the RHI API design. If this fires, it's more likely an RHI bug.
-		ensureMsgf(bTransitioning, TEXT("Unsolicited resource end transition call."));
-		bTransitioning = false;
-		BeginTransitionBacktrace = nullptr;
+		FPipelineState& State = States[ExecutingPipeline];
+
+		// Check that we aren't ending a transition that never began.
+		ensureMsgf(State.bTransitioning, TEXT("Unsolicited resource end transition call."));
+		State.bTransitioning = false;
+		State.BeginTransitionBacktrace = nullptr;
+
+		// Check that the end matches the begin.
+		ensureMsgf(TargetState == State.Current, TEXT("%s"), *GetReasonString_MismatchedEndTransition(Resource, SubresourceIndex, State.Current, TargetState));
+
+		// Replicate the state to other pipes that are not part of the end pipe mask.
+		for (ERHIPipeline OtherPipeline : GetRHIPipelines())
+		{
+			if (!EnumHasAnyFlags(TargetState.Pipelines, OtherPipeline))
+			{
+				States[OtherPipeline] = State;
+			}
+		}
 	}
 
 	void FSubresourceState::Assert(FResource* Resource, FSubresourceIndex const& SubresourceIndex, const FState& RequiredState, bool bAllowAllUAVsOverlap)
@@ -554,29 +711,30 @@ namespace RHIValidation
 				*GetRHIPipelineName(RequiredState.Pipelines)));
 		}
 
+		FPipelineState& State = States[RequiredState.Pipelines];
+
 		// Check we're not trying to access the resource whilst a pending resource transition is in progress.
-		ensureMsgf(!bTransitioning, TEXT("%s"), *GetReasonString_AccessDuringTransition(Resource, SubresourceIndex, CurrentState, RequiredState, CreateTransitionBacktrace, BeginTransitionBacktrace));
+		ensureMsgf(!State.bTransitioning, TEXT("%s"), *GetReasonString_AccessDuringTransition(Resource, SubresourceIndex, State.Current, RequiredState, State.CreateTransitionBacktrace, State.BeginTransitionBacktrace));
 
 		// If UAV overlaps are now disabled, ensure the resource has been transitioned if it was previously used in UAV overlap state.
-		ensureMsgf((bAllowAllUAVsOverlap || !bUsedWithAllUAVsOverlap) && (bExplicitAllowUAVOverlap || !bUsedWithExplicitUAVsOverlap), TEXT("%s"), *GetReasonString_UAVOverlap(Resource, SubresourceIndex, CurrentState, RequiredState));
+		ensureMsgf((bAllowAllUAVsOverlap || !State.bUsedWithAllUAVsOverlap) && (State.bExplicitAllowUAVOverlap || !State.bUsedWithExplicitUAVsOverlap), TEXT("%s"), *GetReasonString_UAVOverlap(Resource, SubresourceIndex, State.Current, RequiredState));
 
 		// Ensure the resource is in the required state for this operation
-		ensureMsgf(EnumHasAllFlags(CurrentState.Access, RequiredState.Access) && EnumHasAllFlags(CurrentState.Pipelines, RequiredState.Pipelines), TEXT("%s"), *GetReasonString_MissingBarrier(Resource, SubresourceIndex, CurrentState, RequiredState));
+		ensureMsgf(EnumHasAllFlags(State.Current.Access, RequiredState.Access) && EnumHasAllFlags(State.Current.Pipelines, RequiredState.Pipelines), TEXT("%s"), *GetReasonString_MissingBarrier(Resource, SubresourceIndex, State.Current, RequiredState));
 
-		PreviousState = CurrentState;
+		State.Previous = State.Current;
 
 		if (EnumHasAnyFlags(RequiredState.Access, ERHIAccess::UAVMask))
 		{
-			if (bAllowAllUAVsOverlap) { bUsedWithAllUAVsOverlap = true; }
-			if (bExplicitAllowUAVOverlap) { bUsedWithExplicitUAVsOverlap = true; }
+			if (bAllowAllUAVsOverlap) { State.bUsedWithAllUAVsOverlap = true; }
+			if (State.bExplicitAllowUAVOverlap) { State.bUsedWithExplicitUAVsOverlap = true; }
 		}
 
 		// Disable all non-compatible access types
-		CurrentState.Access = DecayResourceAccess(CurrentState.Access, RequiredState.Access, bAllowAllUAVsOverlap || bExplicitAllowUAVOverlap);
-		CurrentState.Pipelines = CurrentState.Pipelines & RequiredState.Pipelines;
+		State.Current.Access = DecayResourceAccess(State.Current.Access, RequiredState.Access, bAllowAllUAVsOverlap || State.bExplicitAllowUAVOverlap);
 	}
 
-	void FSubresourceState::SpecificUAVOverlap(FResource* Resource, FSubresourceIndex const& SubresourceIndex, bool bAllow)
+	void FSubresourceState::SpecificUAVOverlap(FResource* Resource, ERHIPipeline Pipeline, FSubresourceIndex const& SubresourceIndex, bool bAllow)
 	{
 		if (Resource->LoggingMode != ELoggingMode::None
 #if LOG_UNNAMED_RESOURCES
@@ -587,8 +745,9 @@ namespace RHIValidation
 			Log(Resource, SubresourceIndex, nullptr, TEXT("UAVOverlap"), *FString::Printf(TEXT("Allow: %s"), bAllow ? TEXT("True") : TEXT("False")));
 		}
 
-		ensureMsgf(bExplicitAllowUAVOverlap != bAllow, TEXT("%s"), *GetReasonString_MismatchedExplicitUAVOverlapCall(bAllow));
-		bExplicitAllowUAVOverlap = bAllow;
+		FPipelineState& State = States[Pipeline];
+		ensureMsgf(State.bExplicitAllowUAVOverlap != bAllow, TEXT("%s"), *GetReasonString_MismatchedExplicitUAVOverlapCall(bAllow));
+		State.bExplicitAllowUAVOverlap = bAllow;
 	}
 
 	inline void FResource::EnumerateSubresources(FSubresourceRange const& SubresourceRange, TFunctionRef<void(FSubresourceState&, FSubresourceIndex const&)> Callback, bool bBeginTransition)
@@ -637,7 +796,7 @@ namespace RHIValidation
 		}
 	}
 
-	RHI_API EReplayStatus FOperation::Replay(bool& bAllowAllUAVsOverlap) const
+	RHI_API EReplayStatus FOperation::Replay(ERHIPipeline Pipeline, bool& bAllowAllUAVsOverlap) const
 	{
 		switch (Type)
 		{
@@ -648,7 +807,7 @@ namespace RHIValidation
 			break;
 
 		case EOpType::BeginTransition:
-			Data_BeginTransition.Identity.Resource->EnumerateSubresources(Data_BeginTransition.Identity.SubresourceRange, [this](FSubresourceState& State, FSubresourceIndex const& SubresourceIndex)
+			Data_BeginTransition.Identity.Resource->EnumerateSubresources(Data_BeginTransition.Identity.SubresourceRange, [this, Pipeline](FSubresourceState& State, FSubresourceIndex const& SubresourceIndex)
 			{
 				State.BeginTransition(
 					Data_BeginTransition.Identity.Resource,
@@ -656,6 +815,7 @@ namespace RHIValidation
 					Data_BeginTransition.PreviousState,
 					Data_BeginTransition.NextState,
 					Data_BeginTransition.Flags,
+					Pipeline,
 					Data_BeginTransition.CreateBacktrace);
 
 			}, true);
@@ -663,18 +823,21 @@ namespace RHIValidation
 			break;
 
 		case EOpType::EndTransition:
-			Data_EndTransition.Identity.Resource->EnumerateSubresources(Data_EndTransition.Identity.SubresourceRange, [this](FSubresourceState& State, FSubresourceIndex const& SubresourceIndex)
+			Data_EndTransition.Identity.Resource->EnumerateSubresources(Data_EndTransition.Identity.SubresourceRange, [this, Pipeline](FSubresourceState& State, FSubresourceIndex const& SubresourceIndex)
 			{
 				State.EndTransition(
-					Data_EndTransition.Identity.Resource, 
+					Data_EndTransition.Identity.Resource,
 					SubresourceIndex,
+					Data_EndTransition.PreviousState,
+					Data_EndTransition.NextState,
+					Pipeline,
 					Data_EndTransition.CreateBacktrace);
 			});
 			Data_EndTransition.Identity.Resource->ReleaseOpRef();
 			break;
 
 		case EOpType::Assert:
-			Data_Assert.Identity.Resource->EnumerateSubresources(Data_Assert.Identity.SubresourceRange, [this, &bAllowAllUAVsOverlap](FSubresourceState& State, FSubresourceIndex const& SubresourceIndex)
+			Data_Assert.Identity.Resource->EnumerateSubresources(Data_Assert.Identity.SubresourceRange, [this, Pipeline, &bAllowAllUAVsOverlap](FSubresourceState& State, FSubresourceIndex const& SubresourceIndex)
 			{
 				State.Assert(
 					Data_Assert.Identity.Resource,
@@ -686,10 +849,20 @@ namespace RHIValidation
 			break;
 
 		case EOpType::Signal:
+			if (Data_Signal.Pipeline != Pipeline)
+			{
+				break;
+			}
+
 			Data_Signal.Fence->bSignaled = true;
 			return EReplayStatus::Signaled;
 
 		case EOpType::Wait:
+			if (Data_Wait.Pipeline != Pipeline)
+			{
+				break;
+			}
+
 			if (Data_Wait.Fence->bSignaled)
 			{
 				// The fence has been completed. Free it now.
@@ -708,10 +881,11 @@ namespace RHIValidation
 			break;
 
 		case EOpType::SpecificUAVOverlap:
-			Data_SpecificUAVOverlap.Identity.Resource->EnumerateSubresources(Data_SpecificUAVOverlap.Identity.SubresourceRange, [this](FSubresourceState& State, FSubresourceIndex const& SubresourceIndex)
+			Data_SpecificUAVOverlap.Identity.Resource->EnumerateSubresources(Data_SpecificUAVOverlap.Identity.SubresourceRange, [this, Pipeline](FSubresourceState& State, FSubresourceIndex const& SubresourceIndex)
 			{
 				State.SpecificUAVOverlap(
-					Data_SpecificUAVOverlap.Identity.Resource, 
+					Data_SpecificUAVOverlap.Identity.Resource,
+					Pipeline,
 					SubresourceIndex,
 					Data_SpecificUAVOverlap.bAllow);
 			});
@@ -735,16 +909,17 @@ namespace RHIValidation
 			Status = EReplayStatus::Normal;
 			for (int32 CurrentIndex = 0; CurrentIndex < int32(ERHIPipeline::Num); ++CurrentIndex)
 			{
+				const ERHIPipeline CurrentPipeline = ERHIPipeline(1 << CurrentIndex);
 				FOpQueueState& CurrentQueue = OpQueues[CurrentIndex];
 				if (CurrentQueue.bWaiting)
 				{
-					Status = CurrentQueue.Ops.Replay(CurrentQueue.bAllowAllUAVsOverlap);
+					Status = CurrentQueue.Ops.Replay(CurrentPipeline, CurrentQueue.bAllowAllUAVsOverlap);
 					if (!EnumHasAllFlags(Status, EReplayStatus::Waiting))
 					{
 						CurrentQueue.Ops.Reset();
 						if (CurrentIndex == DstOpQueueIndex && InOpsList.Incomplete())
 						{
-							Status |= InOpsList.Replay(CurrentQueue.bAllowAllUAVsOverlap);
+							Status |= InOpsList.Replay(CurrentPipeline, CurrentQueue.bAllowAllUAVsOverlap);
 							CurrentQueue.bWaiting = InOpsList.Incomplete();
 						}
 						else
diff --git a/Engine/Source/Runtime/RHI/Public/RHI.h b/Engine/Source/Runtime/RHI/Public/RHI.h
--- a/Engine/Source/Runtime/RHI/Public/RHI.h
+++ b/Engine/Source/Runtime/RHI/Public/RHI.h
@@ -716,11 +716,69 @@ enum class ERHIPipeline : uint8
 	Graphics = 1 << 0,
 	AsyncCompute = 1 << 1,
 
+	None = 0,
 	All = Graphics | AsyncCompute,
 	Num = 2
 };
 ENUM_CLASS_FLAGS(ERHIPipeline)
 
+inline uint32 GetRHIPipelineIndex(ERHIPipeline Pipeline)
+{
+	switch (Pipeline)
+	{
+	default:
+		checkf(false, TEXT("Attempted to get a pipeline index but the pipeline mask is not exclusively Graphics or AsyncCompute."));
+	case ERHIPipeline::Graphics:
+		return 0;
+	case ERHIPipeline::AsyncCompute:
+		return 1;
+	}
+}
+
+inline constexpr uint32 GetRHIPipelineCount()
+{
+	return uint32(ERHIPipeline::Num);
+}
+
+inline const std::initializer_list<ERHIPipeline> GetRHIPipelines()
+{
+	static const ERHIPipeline Pipelines[] = { ERHIPipeline::Graphics, ERHIPipeline::AsyncCompute };
+	return std::initializer_list<ERHIPipeline>(Pipelines, Pipelines + (uint32)ERHIPipeline::Num);
+}
+
+template <typename FunctionType>
+inline void EnumerateRHIPipelines(ERHIPipeline PipelineMask, FunctionType Function)
+{
+	for (ERHIPipeline Pipeline : GetRHIPipelines())
+	{
+		if (EnumHasAnyFlags(PipelineMask, Pipeline))
+		{
+			Function(Pipeline);
+		}
+	}
+}
+
+/** Array of pass handles by RHI pipeline, with overloads to help with enum conversion. */
+template <typename ElementType>
+class TRHIPipelineArray : public TStaticArray<ElementType, GetRHIPipelineCount()>
+{
+	using Base = TStaticArray<ElementType, GetRHIPipelineCount()>;
+public:
+	TRHIPipelineArray() = default;
+	TRHIPipelineArray(const TRHIPipelineArray&) = default;
+	TRHIPipelineArray& operator=(const TRHIPipelineArray&) = default;
+
+	FORCEINLINE ElementType& operator[](ERHIPipeline Pipeline)
+	{
+		return Base::operator[](GetRHIPipelineIndex(Pipeline));
+	}
+
+	FORCEINLINE const ElementType& operator[](ERHIPipeline Pipeline) const
+	{
+		return Base::operator[](GetRHIPipelineIndex(Pipeline));
+	}
+};
+
 enum class ERHIAccess
 {
 	// Used when the previous state of a resource is not known,
@@ -780,6 +838,9 @@ ENUM_CLASS_FLAGS(ERHIAccess)
 /** Mask of read states that can be used together for textures. */
 extern RHI_API ERHIAccess GRHITextureReadAccessMask;
 
+/** Mask of states which are allowed to be used on multiple pipelines at the same time. */
+extern RHI_API ERHIAccess GRHIMultiPipelineAccessMask;
+
 /** to customize the RHIReadSurfaceData() output */
 class FReadSurfaceDataFlags
 {
@@ -1913,7 +1974,7 @@ private:
 
 	inline void MarkBegin(ERHIPipeline Pipeline) const
 	{
-		checkf(EnumHasAllFlags(Pipeline, AllowedSrc), TEXT("Transition is being used on a source pipeline that it wasn't created for."));
+		checkf(EnumHasAllFlags(AllowedSrc, Pipeline), TEXT("Transition is being used on a source pipeline that it wasn't created for."));
 
 		int8 Mask = int8(Pipeline);
 		int8 PreviousValue = FPlatformAtomics::InterlockedAnd(&State, ~Mask);
@@ -1927,7 +1988,7 @@ private:
 
 	inline void MarkEnd(ERHIPipeline Pipeline) const
 	{
-		checkf(EnumHasAllFlags(Pipeline, AllowedDst), TEXT("Transition is being used on a destination pipeline that it wasn't created for."));
+		checkf(EnumHasAllFlags(AllowedDst, Pipeline), TEXT("Transition is being used on a destination pipeline that it wasn't created for."));
 
 		int8 Mask = int8(Pipeline) << int32(ERHIPipeline::Num);
 		int8 PreviousValue = FPlatformAtomics::InterlockedAnd(&State, ~Mask);
@@ -1954,7 +2015,6 @@ private:
 	friend class FValidationComputeContext;
 	friend class FValidationContext;
 
-	RHIValidation::FFence* Fence = nullptr;
 	RHIValidation::FOperationsList PendingOperationsBegin;
 	RHIValidation::FOperationsList PendingOperationsEnd;
 #endif
@@ -2214,7 +2274,7 @@ inline void RHIValidation::FTracker::AddOp(const RHIValidation::FOperation& Op)
 	if (GRHICommandList.Bypass() && CurrentList.Operations.Num() == 0)
 	{
 		auto& OpQueue = OpQueues[GetOpQueueIndex(Pipeline)];
-		if (!EnumHasAllFlags(Op.Replay(OpQueue.bAllowAllUAVsOverlap), EReplayStatus::Waiting))
+		if (!EnumHasAllFlags(Op.Replay(Pipeline, OpQueue.bAllowAllUAVsOverlap), EReplayStatus::Waiting))
 		{
 			return;
 		}
diff --git a/Engine/Source/Runtime/RHI/Public/RHIValidation.h b/Engine/Source/Runtime/RHI/Public/RHIValidation.h
--- a/Engine/Source/Runtime/RHI/Public/RHIValidation.h
+++ b/Engine/Source/Runtime/RHI/Public/RHIValidation.h
@@ -193,88 +193,7 @@ public:
 		return RHI->RHICreateGPUFence(Name);
 	}
 
-	virtual void RHICreateTransition(FRHITransition* Transition, ERHIPipeline SrcPipelines, ERHIPipeline DstPipelines, ERHICreateTransitionFlags CreateFlags, TArrayView<const FRHITransitionInfo> Infos)
-	{
-		using namespace RHIValidation;
-
-		RHIValidation::FFence* Fence = nullptr;
-		if (SrcPipelines != DstPipelines)
-		{
-			Fence = Transition->Fence = new RHIValidation::FFence;
-		}
-
-		TArray<FOperation> BeginOps, EndOps;
-		BeginOps.Reserve(Infos.Num() + (Fence ? 1 : 0));
-		EndOps  .Reserve(Infos.Num() + (Fence ? 1 : 0));
-
-		if (Fence)
-		{
-			EndOps.Emplace(FOperation::Wait(Fence));
-		}
-
-		bool bDoTrace = false;
-
-		for (int32 Index = 0; Index < Infos.Num(); ++Index)
-		{
-			const FRHITransitionInfo& Info = Infos[Index];
-			if (!Info.Resource)
-				continue;
-
-			checkf(Info.Type != FRHITransitionInfo::EType::Unknown, TEXT("FRHITransitionInfo::Type cannot be Unknown when creating a resource transition."));
-
-			FResourceIdentity Identity;
-
-			switch (Info.Type)
-			{
-			default: checkNoEntry(); // fall through
-			case FRHITransitionInfo::EType::Texture: 
-				Identity = Info.Texture->GetTransitionIdentity(Info); 
-				break;
-
-			case FRHITransitionInfo::EType::VertexBuffer:
-				Identity = Info.VertexBuffer->GetWholeResourceIdentity();
-				break;
-
-			case FRHITransitionInfo::EType::IndexBuffer:
-				Identity = Info.IndexBuffer->GetWholeResourceIdentity(); 
-				break;
-
-			case FRHITransitionInfo::EType::StructuredBuffer:
-				Identity = Info.StructuredBuffer->GetWholeResourceIdentity();
-				break;
-
-			case FRHITransitionInfo::EType::UAV: 
-				Identity = Info.UAV->ViewIdentity;
-				break;
-			}
-
-			// Take a backtrace of this transition creation if any of the resources it contains have logging enabled.
-			bDoTrace |= (Identity.Resource->LoggingMode != RHIValidation::ELoggingMode::None);
-
-			FState PreviousState = FState(Info.AccessBefore, SrcPipelines);
-			FState NextState = FState(Info.AccessAfter, DstPipelines);
-
-			BeginOps.Emplace(FOperation::BeginTransitionResource(Identity, PreviousState, NextState, Info.Flags, nullptr));
-			EndOps.Emplace(FOperation::EndTransitionResource(Identity, nullptr));
-		}
-
-		if (bDoTrace)
-		{
-			void* Backtrace = CaptureBacktrace();
-			for (FOperation& Op : BeginOps) { Op.Data_BeginTransition.CreateBacktrace = Backtrace; }
-			for (FOperation& Op : EndOps)   { Op.Data_EndTransition.CreateBacktrace = Backtrace; }
-		}
-
-		if (Fence)
-		{
-			BeginOps.Emplace(FOperation::Signal(Fence));
-		}
-
-		Transition->PendingOperationsBegin.Operations = MoveTemp(BeginOps);
-		Transition->PendingOperationsEnd.Operations = MoveTemp(EndOps);
-
-		return RHI->RHICreateTransition(Transition, SrcPipelines, DstPipelines, CreateFlags, Infos);
-	}
+	virtual void RHICreateTransition(FRHITransition* Transition, ERHIPipeline SrcPipelines, ERHIPipeline DstPipelines, ERHICreateTransitionFlags CreateFlags, TArrayView<const FRHITransitionInfo> Infos);
 
 	virtual void RHIReleaseTransition(FRHITransition* Transition)
 	{
diff --git a/Engine/Source/Runtime/RHI/Public/RHIValidationCommon.h b/Engine/Source/Runtime/RHI/Public/RHIValidationCommon.h
--- a/Engine/Source/Runtime/RHI/Public/RHIValidationCommon.h
+++ b/Engine/Source/Runtime/RHI/Public/RHIValidationCommon.h
@@ -139,38 +139,40 @@ namespace RHIValidation
 
 	struct FSubresourceState
 	{
-		FState PreviousState;
-		FState CurrentState;
-		EResourceTransitionFlags Flags = EResourceTransitionFlags::None;
+		struct FPipelineState
+		{
+			FPipelineState()
+			{
+				Current.Access = ERHIAccess::Unknown;
+				Current.Pipelines = ERHIPipeline::Graphics;
+				Previous = Current;
+			}
 
-		// True when a BeginTransition has been issued, and false when the transition has been ended.
-		bool bTransitioning = false;
+			FState Previous;
+			FState Current;
+			EResourceTransitionFlags Flags = EResourceTransitionFlags::None;
 
-		// True when the resource has been used within a Begin/EndUAVOverlap region.
-		bool bUsedWithAllUAVsOverlap = false;
-		
-		// True if the calling code explicitly enabled overlapping on this UAV.
-		bool bExplicitAllowUAVOverlap = false;
-		bool bUsedWithExplicitUAVsOverlap = false;
+			// True when a BeginTransition has been issued, and false when the transition has been ended.
+			bool bTransitioning = false;
 
-		// Pointer to the previous create/begin transition backtraces if logging is enabled for this resource.
-		void* CreateTransitionBacktrace = nullptr;
-		void* BeginTransitionBacktrace = nullptr;
+			// True when the resource has been used within a Begin/EndUAVOverlap region.
+			bool bUsedWithAllUAVsOverlap = false;
 
-		FSubresourceState()
-		{
-			CurrentState.Access = ERHIAccess::Unknown;
+			// True if the calling code explicitly enabled overlapping on this UAV.
+			bool bExplicitAllowUAVOverlap = false;
+			bool bUsedWithExplicitUAVsOverlap = false;
 
-			// Resource can initially be accessed on any pipe without a transition.
-			CurrentState.Pipelines = ERHIPipeline::All;
+			// Pointer to the previous create/begin transition backtraces if logging is enabled for this resource.
+			void* CreateTransitionBacktrace = nullptr;
+			void* BeginTransitionBacktrace = nullptr;
+		};
 
-			PreviousState = CurrentState;
-		}
+		TRHIPipelineArray<FPipelineState> States;
 
-		void BeginTransition   (FResource* Resource, FSubresourceIndex const& SubresourceIndex, const FState& CurrentStateFromRHI, const FState& TargetState, EResourceTransitionFlags NewFlags, void* CreateTrace);
-		void EndTransition     (FResource* Resource, FSubresourceIndex const& SubresourceIndex, void* CreateTrace);
+		void BeginTransition   (FResource* Resource, FSubresourceIndex const& SubresourceIndex, const FState& CurrentStateFromRHI, const FState& TargetState, EResourceTransitionFlags NewFlags, ERHIPipeline Pipeline, void* CreateTrace);
+		void EndTransition     (FResource* Resource, FSubresourceIndex const& SubresourceIndex, const FState& CurrentStateFromRHI, const FState& TargetState, ERHIPipeline Pipeline, void* CreateTrace);
 		void Assert            (FResource* Resource, FSubresourceIndex const& SubresourceIndex, const FState& RequiredState, bool bAllowAllUAVsOverlap);
-		void SpecificUAVOverlap(FResource* Resource, FSubresourceIndex const& SubresourceIndex, bool bAllow);
+		void SpecificUAVOverlap(FResource* Resource, ERHIPipeline Pipeline, FSubresourceIndex const& SubresourceIndex, bool bAllow);
 		void* Log              (FResource* Resource, FSubresourceIndex const& SubresourceIndex, void* CreateTrace, const TCHAR* Type, const TCHAR* LogStr);
 	};
 
@@ -287,8 +289,11 @@ namespace RHIValidation
 				NumArraySlices = InNumArraySlices;
 				NumPlanes = InNumPlanes;
 
-				WholeResourceState.CurrentState.Access = InResourceState;
-				WholeResourceState.PreviousState = WholeResourceState.CurrentState;
+				for (auto& State : WholeResourceState.States)
+				{
+					State.Current.Access = InResourceState;
+					State.Previous = State.Current;
+				}
 
 				if (InDebugName != nullptr)
 				{
@@ -537,6 +542,8 @@ namespace RHIValidation
 			struct
 			{
 				FResourceIdentity Identity;
+				FState PreviousState;
+				FState NextState;
 				void* CreateBacktrace;
 			} Data_EndTransition;
 
@@ -556,11 +563,13 @@ namespace RHIValidation
 			struct
 			{
 				FFence* Fence;
+				ERHIPipeline Pipeline;
 			} Data_Signal;
 
 			struct
 			{
 				FFence* Fence;
+				ERHIPipeline Pipeline;
 			} Data_Wait;
 
 			struct
@@ -575,11 +584,17 @@ namespace RHIValidation
 			} Data_SpecificUAVOverlap;
 		};
 
-		RHI_API EReplayStatus Replay(bool& bAllowAllUAVsOverlap) const;
+		RHI_API EReplayStatus Replay(ERHIPipeline Pipeline, bool& bAllowAllUAVsOverlap) const;
 
 		static inline FOperation BeginTransitionResource(FResourceIdentity Identity, FState PreviousState, FState NextState, EResourceTransitionFlags Flags, void* CreateBacktrace)
 		{
-			Identity.Resource->AddOpRef();
+			for (ERHIPipeline Pipeline : GetRHIPipelines())
+			{
+				if (EnumHasAnyFlags(PreviousState.Pipelines, Pipeline))
+				{
+					Identity.Resource->AddOpRef();
+				}
+			}
 
 			FOperation Op;
 			Op.Type = EOpType::BeginTransition;
@@ -591,13 +606,21 @@ namespace RHIValidation
 			return MoveTemp(Op);
 		}
 
-		static inline FOperation EndTransitionResource(FResourceIdentity Identity, void* CreateBacktrace)
+		static inline FOperation EndTransitionResource(FResourceIdentity Identity, FState PreviousState, FState NextState, void* CreateBacktrace)
 		{
-			Identity.Resource->AddOpRef();
+			for (ERHIPipeline Pipeline : GetRHIPipelines())
+			{
+				if (EnumHasAnyFlags(NextState.Pipelines, Pipeline))
+				{
+					Identity.Resource->AddOpRef();
+				}
+			}
 
 			FOperation Op;
 			Op.Type = EOpType::EndTransition;
 			Op.Data_EndTransition.Identity = Identity;
+			Op.Data_EndTransition.PreviousState = PreviousState;
+			Op.Data_EndTransition.NextState = NextState;
 			Op.Data_EndTransition.CreateBacktrace = CreateBacktrace;
 			return MoveTemp(Op);
 		}
@@ -630,19 +653,21 @@ namespace RHIValidation
 			return MoveTemp(Op);
 		}
 
-		static inline FOperation Signal(FFence* Fence)
+		static inline FOperation Signal(FFence* Fence, ERHIPipeline Pipeline)
 		{
 			FOperation Op;
 			Op.Type = EOpType::Signal;
 			Op.Data_Signal.Fence = Fence;
+			Op.Data_Signal.Pipeline = Pipeline;
 			return MoveTemp(Op);
 		}
 
-		static inline FOperation Wait(FFence* Fence)
+		static inline FOperation Wait(FFence* Fence, ERHIPipeline Pipeline)
 		{
 			FOperation Op;
 			Op.Type = EOpType::Wait;
 			Op.Data_Wait.Fence = Fence;
+			Op.Data_Wait.Pipeline = Pipeline;
 			return MoveTemp(Op);
 		}
 
@@ -671,12 +696,12 @@ namespace RHIValidation
 		TArray<FOperation> Operations;
 		int32 OperationPos = 0;
 
-		inline EReplayStatus Replay(bool& bAllowAllUAVsOverlap)
+		inline EReplayStatus Replay(ERHIPipeline Pipeline, bool& bAllowAllUAVsOverlap)
 		{
 			EReplayStatus Status = EReplayStatus::Normal;
 			for (; OperationPos < Operations.Num(); ++OperationPos)
 			{
-				Status |= Operations[OperationPos].Replay(bAllowAllUAVsOverlap);
+				Status |= Operations[OperationPos].Replay(Pipeline, bAllowAllUAVsOverlap);
 				if (EnumHasAllFlags(Status, EReplayStatus::Waiting))
 				{
 					break;
@@ -809,7 +834,7 @@ namespace RHIValidation
 			// This function exists due to the implicit transitions that RHI functions make (e.g. RHICopyToResolveTarget).
 			// It should be removed when we eventually remove all implicit transitions from the RHI.
 			AddOp(FOperation::BeginTransitionResource(Identity, PreviousState, NextState, Flags, nullptr));
-			AddOp(FOperation::EndTransitionResource(Identity, nullptr));
+			AddOp(FOperation::EndTransitionResource(Identity, PreviousState, NextState, nullptr));
 		}
 
 		inline void AllUAVsOverlap(bool bAllow)
diff --git a/Engine/Source/Runtime/RHI/Public/RHIValidationContext.h b/Engine/Source/Runtime/RHI/Public/RHIValidationContext.h
--- a/Engine/Source/Runtime/RHI/Public/RHIValidationContext.h
+++ b/Engine/Source/Runtime/RHI/Public/RHIValidationContext.h
@@ -96,7 +96,7 @@ public:
 	virtual void RHISetShaderTexture(FRHIComputeShader* Shader, uint32 TextureIndex, FRHITexture* NewTexture) override final
 	{
 		checkf(State.bComputeShaderSet, TEXT("A Compute shader has to be set to set resources into a shader!"));
-		Tracker->Assert(NewTexture->GetWholeResourceIdentity(), ERHIAccess::SRVCompute);
+		Tracker->Assert(NewTexture->GetWholeResourceIdentitySRV(), ERHIAccess::SRVCompute);
 		RHIContext->RHISetShaderTexture(Shader, TextureIndex, NewTexture);
 	}
 
