Build fix.

#rnx

diff --git a/Engine/Source/Runtime/RHI/Public/RHI.h b/Engine/Source/Runtime/RHI/Public/RHI.h
--- a/Engine/Source/Runtime/RHI/Public/RHI.h
+++ b/Engine/Source/Runtime/RHI/Public/RHI.h
@@ -740,10 +740,10 @@ inline constexpr uint32 GetRHIPipelineCount()
 	return uint32(ERHIPipeline::Num);
 }
 
-inline const std::initializer_list<ERHIPipeline> GetRHIPipelines()
+inline TArrayView<const ERHIPipeline> GetRHIPipelines()
 {
 	static const ERHIPipeline Pipelines[] = { ERHIPipeline::Graphics, ERHIPipeline::AsyncCompute };
-	return std::initializer_list<ERHIPipeline>(Pipelines, Pipelines + (uint32)ERHIPipeline::Num);
+	return Pipelines;
 }
 
 template <typename FunctionType>
