RHI validation uses a custom reporting function instead of ensure().

The custom function uses the CRC of the message string to filter repeated messages, instead of the source location. Since the string contains the name and address of the resource, this allows us to print separate reports for identical failures on different resources, something that wasn't possible with ensure (since they all originate from the same function).

Also added the resource name on the first line of the error message, so that automated tools can pick it up (some tools only capture the first line of the error message).

#rnx

diff --git a/Engine/Source/Runtime/RHI/Private/RHIValidation.cpp b/Engine/Source/Runtime/RHI/Private/RHIValidation.cpp
--- a/Engine/Source/Runtime/RHI/Private/RHIValidation.cpp
+++ b/Engine/Source/Runtime/RHI/Private/RHIValidation.cpp
@@ -35,12 +35,16 @@ static const TCHAR* GAutoLogResourceNames[] =
 	nullptr
 };
 
+TSet<uint32> FValidationRHI::SeenFailureHashes;
+FCriticalSection FValidationRHI::SeenFailureHashesMutex;
+
 FValidationRHI::FValidationRHI(FDynamicRHI* InRHI)
 	: RHI(InRHI)
 {
 	check(RHI);
 	UE_LOG(LogRHI, Warning, TEXT("FValidationRHI on, intercepting %s RHI!"), InRHI && InRHI->GetName() ? InRHI->GetName() : TEXT("<NULL>"));
 	GRHIValidationEnabled = true;
+	SeenFailureHashes.Reserve(256);
 }
 
 FValidationRHI::~FValidationRHI()
@@ -93,7 +97,7 @@ void FValidationRHI::ValidatePipeline(const FGraphicsPipelineStateInitializer& P
 		{
 			if (!bHasStencil)
 			{
-				ensureMsgf(!Initializer.bEnableFrontFaceStencil
+				RHI_VALIDATION_CHECK(!Initializer.bEnableFrontFaceStencil
 					&& Initializer.FrontFaceStencilTest == CF_Always
 					&& Initializer.FrontFaceStencilFailStencilOp == SO_Keep
 					&& Initializer.FrontFaceDepthFailStencilOp == SO_Keep
@@ -104,20 +108,20 @@ void FValidationRHI::ValidatePipeline(const FGraphicsPipelineStateInitializer& P
 					&& Initializer.BackFaceDepthFailStencilOp == SO_Keep
 					&& Initializer.BackFacePassStencilOp == SO_Keep, TEXT("No stencil render target set, yet PSO wants to use stencil operations!"));
 /*
-				ensureMsgf(PSOInitializer.StencilTargetLoadAction == ERenderTargetLoadAction::ENoAction,
+				RHI_VALIDATION_CHECK(PSOInitializer.StencilTargetLoadAction == ERenderTargetLoadAction::ENoAction,
 					TEXT("No stencil target set, yet PSO wants to load from it!"));
-				ensureMsgf(PSOInitializer.StencilTargetStoreAction == ERenderTargetStoreAction::ENoAction,
+				RHI_VALIDATION_CHECK(PSOInitializer.StencilTargetStoreAction == ERenderTargetStoreAction::ENoAction,
 					TEXT("No stencil target set, yet PSO wants to store into it!"));
 */
 			}
 		}
 		else
 		{
-			ensureMsgf(!Initializer.bEnableDepthWrite && Initializer.DepthTest == CF_Always, TEXT("No depth render target set, yet PSO wants to use depth operations!"));
-			ensureMsgf(PSOInitializer.DepthTargetLoadAction == ERenderTargetLoadAction::ENoAction
+			RHI_VALIDATION_CHECK(!Initializer.bEnableDepthWrite && Initializer.DepthTest == CF_Always, TEXT("No depth render target set, yet PSO wants to use depth operations!"));
+			RHI_VALIDATION_CHECK(PSOInitializer.DepthTargetLoadAction == ERenderTargetLoadAction::ENoAction
 				&& PSOInitializer.StencilTargetLoadAction == ERenderTargetLoadAction::ENoAction,
 				TEXT("No depth/stencil target set, yet PSO wants to load from it!"));
-			ensureMsgf(PSOInitializer.DepthTargetStoreAction == ERenderTargetStoreAction::ENoAction
+			RHI_VALIDATION_CHECK(PSOInitializer.DepthTargetStoreAction == ERenderTargetStoreAction::ENoAction
 				&& PSOInitializer.StencilTargetStoreAction == ERenderTargetStoreAction::ENoAction,
 				TEXT("No depth/stencil target set, yet PSO wants to store into it!"));
 		}
@@ -235,6 +239,36 @@ void FValidationRHI::RHICreateTransition(FRHITransition* Transition, ERHIPipelin
 	return RHI->RHICreateTransition(Transition, SrcPipelines, DstPipelines, CreateFlags, Infos);
 }
 
+void FValidationRHI::ReportValidationFailure(const TCHAR* InMessage)
+{
+	// Report failures only once per session, since many of them will happen repeatedly. This is similar to what ensure() does, but
+	// ensure() looks at the source location to determine if it's seen the error before. We want to look at the actual message, since
+	// all failures of a given kind will come from the same place, but (hopefully) the error message contains the name of the resource
+	// and a description of the state, so it should be unique for each failure.
+	uint32 Hash = FCrc::StrCrc32<TCHAR>(InMessage);
+	
+	SeenFailureHashesMutex.Lock();
+	bool bIsAlreadyInSet;
+	SeenFailureHashes.Add(Hash, &bIsAlreadyInSet);
+	SeenFailureHashesMutex.Unlock();
+
+	if (bIsAlreadyInSet)
+	{
+		return;
+	}
+
+	UE_LOG(LogRHI, Error, TEXT("%s"), InMessage);
+
+	if (FPlatformMisc::IsDebuggerPresent())
+	{
+		// Print the message again using the debug output function, because UE_LOG doesn't always reach
+		// the VS output window before the breakpoint is triggered, despite the log flush call below.
+		FPlatformMisc::LowLevelOutputDebugStringf(TEXT("%s\n"), InMessage);
+		GLog->PanicFlushThreadedLogs();
+		PLATFORM_BREAK();
+	}
+}
+
 FValidationComputeContext::FValidationComputeContext()
 	: RHIContext(nullptr)
 {
@@ -322,30 +356,46 @@ namespace RHIValidation
 		return AccessMask;
 	}
 
-#define BARRIER_TRACKER_LOG_PREFIX TEXT("\n\n")\
+// Warning: this prefix expects a string argument for the failure reason, make sure you add it.
+#define BARRIER_TRACKER_LOG_PREFIX_REASON TEXT("RHI validation failed: %s\n\n")\
+	TEXT("--------------------------------------------------------------------\n")\
+	TEXT("              RHI Resource Transition Validation Error              \n")\
+	TEXT("--------------------------------------------------------------------\n")\
+	TEXT("\n\n")
+	
+// Warning: this prefix expects a string argument for the resource name, make sure you add it.
+#define BARRIER_TRACKER_LOG_PREFIX_RESNAME TEXT("RHI validation failed for resource \"%s\":\n\n")\
 	TEXT("--------------------------------------------------------------------\n")\
 	TEXT("              RHI Resource Transition Validation Error              \n")\
 	TEXT("--------------------------------------------------------------------\n")\
 	TEXT("\n\n")
+
 #define BARRIER_TRACKER_LOG_SUFFIX TEXT("\n\n")\
 	TEXT("--------------------------------------------------------------------\n")\
 	TEXT("\n\n")
 
+	static inline const TCHAR* GetResourceDebugName(FResource* Resource)
+	{
+		const TCHAR* DebugName = Resource->GetDebugName();
+		return DebugName ? DebugName : TEXT("Unnamed");
+	}
+
 	static inline FString GetReasonString_MissingBarrier(
 		FResource* Resource, FSubresourceIndex const& SubresourceIndex, 
 		const FState& CurrentState,
 		const FState& RequiredState)
 	{
-		const TCHAR* DebugName = Resource->GetDebugName();
+		const TCHAR* DebugName = GetResourceDebugName(Resource);
 		return FString::Printf(
-			BARRIER_TRACKER_LOG_PREFIX
+			BARRIER_TRACKER_LOG_PREFIX_RESNAME
 			TEXT("Attempted to access resource \"%s\" (0x%p) (%s) from a hardware unit it is not currently accessible from. A resource transition is required.\n\n")
 			TEXT("    --- Allowed access states for this resource are: %s\n")
 			TEXT("    --- Required access states are:                  %s\n")
 			TEXT("    --- Allowed pipelines for this resource are:     %s\n")
 			TEXT("    --- Required pipelines are:                      %s\n")
 			BARRIER_TRACKER_LOG_SUFFIX,
-			DebugName ? DebugName : TEXT("Unnamed"),
+			DebugName,
+			DebugName,
 			Resource,
 			*SubresourceIndex.ToString(),
 			*GetRHIAccessName(CurrentState.Access),
@@ -377,9 +427,9 @@ namespace RHIValidation
 		const FState& AttemptedState,
 		void* CreateTrace, void* BeginTrace)
 	{
-		const TCHAR* DebugName = Resource->GetDebugName();
+		const TCHAR* DebugName = GetResourceDebugName(Resource);
 		return FString::Printf(
-			BARRIER_TRACKER_LOG_PREFIX
+			BARRIER_TRACKER_LOG_PREFIX_RESNAME
 			TEXT("Attempted to access resource \"%s\" (0x%p) (%s) whilst an asynchronous resource transition is in progress. A call to RHIEndTransitions() must be made before the resource can be accessed again.\n\n")
 			TEXT("    --- Pending access states for this resource are: %s\n")
 			TEXT("    --- Attempted access states are:                 %s\n")
@@ -387,7 +437,8 @@ namespace RHIValidation
 			TEXT("    --- Attempted pipelines are:                     %s\n")
 			TEXT("%s")
 			BARRIER_TRACKER_LOG_SUFFIX,
-			DebugName ? DebugName : TEXT("Unnamed"),
+			DebugName,
+			DebugName,
 			Resource,
 			*SubresourceIndex.ToString(),
 			*GetRHIAccessName(PendingState.Access),
@@ -403,9 +454,9 @@ namespace RHIValidation
 		const FState& TargetState,
 		void* CreateTrace, void* BeginTrace)
 	{
-		const TCHAR* DebugName = Resource->GetDebugName();
+		const TCHAR* DebugName = GetResourceDebugName(Resource);
 		return FString::Printf(
-			BARRIER_TRACKER_LOG_PREFIX
+			BARRIER_TRACKER_LOG_PREFIX_RESNAME
 			TEXT("Attempted to begin a resource transition for resource \"%s\" (0x%p) (%s) whilst a previous asynchronous resource transition is already in progress. A call to RHIEndTransitions() must be made before the resource can be transitioned again.\n\n")
 			TEXT("    --- Pending access states for this resource are:              %s\n")
 			TEXT("    --- Attempted access states for the duplicate transition are: %s\n")
@@ -413,7 +464,8 @@ namespace RHIValidation
 			TEXT("    --- Attempted pipelines for the duplicate transition are:     %s\n")
 			TEXT("%s")
 			BARRIER_TRACKER_LOG_SUFFIX,
-			DebugName ? DebugName : TEXT("Unnamed"),
+			DebugName,
+			DebugName,
 			Resource,
 			*SubresourceIndex.ToString(),
 			*GetRHIAccessName(PendingState.Access),
@@ -428,15 +480,16 @@ namespace RHIValidation
 		const FState& ActualCurrentState,
 		const FState& CurrentStateFromRHI)
 	{
-		const TCHAR* DebugName = Resource->GetDebugName();
+		const TCHAR* DebugName = GetResourceDebugName(Resource);
 		return FString::Printf(
-			BARRIER_TRACKER_LOG_PREFIX
+			BARRIER_TRACKER_LOG_PREFIX_RESNAME
 			TEXT("Attempted to begin a resource transition for resource \"%s\" (0x%p) (%s) on the wrong pipeline(s) (\"%s\"). The resource is currently accessible on the \"%s\" pipeline(s).\n\n")
 			TEXT("    --- Current access states for this resource are: %s\n")
 			TEXT("    --- Attempted access states are:                 %s\n\n")
 			TEXT("    --- Ensure that resource transitions are issued on the correct pipeline.\n")
 			BARRIER_TRACKER_LOG_SUFFIX,
-			DebugName ? DebugName : TEXT("Unnamed"),
+			DebugName,
+			DebugName,
 			Resource,
 			*SubresourceIndex.ToString(),
 			*GetRHIPipelineName(CurrentStateFromRHI.Pipelines),
@@ -450,18 +503,19 @@ namespace RHIValidation
 		const FState& CurrentState,
 		const FState& CurrentStateFromRHI)
 	{
-		const TCHAR* DebugName = Resource->GetDebugName();
+		const TCHAR* DebugName = GetResourceDebugName(Resource);
 		return FString::Printf(
-			BARRIER_TRACKER_LOG_PREFIX
+			BARRIER_TRACKER_LOG_PREFIX_RESNAME
 			TEXT("The explicit previous state \"%s\" does not match the tracked current state \"%s\" for the resource \"%s\" (0x%p) (%s).\n")
 			TEXT("    --- Allowed pipelines for this resource are:                           %s\n")
 			TEXT("    --- Previous pipelines passed as part of the resource transition were: %s\n\n")
 			TEXT("    --- The best solution is to correct the explicit previous state passed for the resource in the call to RHICreateTransition().\n")
 			TEXT("    --- Alternatively, use ERHIAccess::Unknown if the actual previous state cannot be determined. Unknown previous resource states have a performance impact so should be avoided if possible.\n")
 			BARRIER_TRACKER_LOG_SUFFIX,
+			DebugName,
 			*GetRHIAccessName(CurrentStateFromRHI.Access),
 			*GetRHIAccessName(CurrentState.Access),
-			DebugName ? DebugName : TEXT("Unnamed"),
+			DebugName,
 			Resource,
 			*SubresourceIndex.ToString(),
 			*GetRHIPipelineName(CurrentState.Pipelines),
@@ -473,17 +527,18 @@ namespace RHIValidation
 		const FState& TargetState,
 		const FState& TargetStateFromRHI)
 	{
-		const TCHAR* DebugName = Resource->GetDebugName();
+		const TCHAR* DebugName = GetResourceDebugName(Resource);
 		return FString::Printf(
-			BARRIER_TRACKER_LOG_PREFIX
+			BARRIER_TRACKER_LOG_PREFIX_RESNAME
 			TEXT("The expected target state \"%s\" on pipe \"%s\" in end transition does not match the tracked target state \"%s\" on pipe \"%s\" for the resource \"%s\" (0x%p) (%s).\n")
 			TEXT("    --- The call to EndTransition() is mismatched with the another BeginTransition() with different states.\n")
 			BARRIER_TRACKER_LOG_SUFFIX,
+			DebugName,
 			*GetRHIAccessName(TargetStateFromRHI.Access),
 			*GetRHIPipelineName(TargetState.Pipelines),
 			*GetRHIAccessName(TargetState.Access),
 			*GetRHIPipelineName(TargetStateFromRHI.Pipelines),
-			DebugName ? DebugName : TEXT("Unnamed"),
+			DebugName,
 			Resource,
 			*SubresourceIndex.ToString());
 	}
@@ -492,14 +547,15 @@ namespace RHIValidation
 		FResource* Resource, FSubresourceIndex const& SubresourceIndex,
 		const FState& CurrentState)
 	{
-		const TCHAR* DebugName = Resource->GetDebugName();
+		const TCHAR* DebugName = GetResourceDebugName(Resource);
 		return FString::Printf(
-			BARRIER_TRACKER_LOG_PREFIX
+			BARRIER_TRACKER_LOG_PREFIX_RESNAME
 			TEXT("Attempted to begin a resource transition for the resource \"%s\" (0x%p) (%s) to the \"%s\" state on the \"%s\" pipe, but the resource is already in this state. The resource transition is unnecessary.\n")
 			TEXT("    --- This is not fatal, but does have an effect on CPU and GPU performance. Consider refactoring rendering code to avoid unnecessary resource transitions.\n")
 			TEXT("    --- RenderGraph (RDG) is capable of handling resource transitions automatically.\n")
 			BARRIER_TRACKER_LOG_SUFFIX,
-			DebugName ? DebugName : TEXT("Unnamed"),
+			DebugName,
+			DebugName,
 			Resource,
 			*SubresourceIndex.ToString(),
 			*GetRHIAccessName(CurrentState.Access),
@@ -509,9 +565,10 @@ namespace RHIValidation
 	static inline FString GetReasonString_MismatchedAllUAVsOverlapCall(bool bAllow)
 	{
 		return FString::Printf(
-			BARRIER_TRACKER_LOG_PREFIX
+			BARRIER_TRACKER_LOG_PREFIX_REASON
 			TEXT("Mismatched call to %sUAVOverlap. Ensure all calls to RHICmdList.BeginUAVOverlap() are paired with a call to RHICmdList.EndUAVOverlap().")
 			BARRIER_TRACKER_LOG_SUFFIX,
+			TEXT("UAV overlap mismatch"),
 			bAllow ? TEXT("Begin") : TEXT("End")
 		);
 	}
@@ -519,9 +576,10 @@ namespace RHIValidation
 	static inline FString GetReasonString_MismatchedExplicitUAVOverlapCall(bool bAllow)
 	{
 		return FString::Printf(
-			BARRIER_TRACKER_LOG_PREFIX
+			BARRIER_TRACKER_LOG_PREFIX_REASON
 			TEXT("Mismatched call to %sUAVOverlap(FRHIUnorderedAccessView*). Ensure all calls to RHICmdList.BeginUAVOverlap() are paired with a call to RHICmdList.EndUAVOverlap().")
 			BARRIER_TRACKER_LOG_SUFFIX,
+			TEXT("UAV overlap mismatch"),
 			bAllow ? TEXT("Begin") : TEXT("End")
 		);
 	}
@@ -530,16 +588,17 @@ namespace RHIValidation
 		FResource* Resource, FSubresourceIndex const& SubresourceIndex,
 		const FState& CurrentState, const FState& RequiredState)
 	{
-		const TCHAR* DebugName = Resource->GetDebugName();
+		const TCHAR* DebugName = GetResourceDebugName(Resource);
 		return FString::Printf(
-			BARRIER_TRACKER_LOG_PREFIX
+			BARRIER_TRACKER_LOG_PREFIX_RESNAME
 			TEXT("Attempted to access resource \"%s\" (0x%p) (%s) which was previously used with overlapping UAV access, but has not been transitioned since UAV overlap was disabled. A resource transition is required.\n\n")
 			TEXT("    --- Allowed access states for this resource are: %s\n")
 			TEXT("    --- Required access states are:                  %s\n")
 			TEXT("    --- Allowed pipelines for this resource are:     %s\n")
 			TEXT("    --- Required pipelines are:                      %s\n")
 			BARRIER_TRACKER_LOG_SUFFIX,
-			DebugName ? DebugName : TEXT("Unnamed"),
+			DebugName,
+			DebugName,
 			Resource,
 			*SubresourceIndex.ToString(),
 			*GetRHIAccessName(CurrentState.Access),
@@ -627,21 +686,21 @@ namespace RHIValidation
 		}
 
 		// Check we're not already transitioning
-		ensureMsgf(!State.bTransitioning, TEXT("%s"), *GetReasonString_DuplicateBeginTransition(Resource, SubresourceIndex, State.Current, TargetState, State.CreateTransitionBacktrace, BeginTrace));
+		RHI_VALIDATION_CHECK(!State.bTransitioning, *GetReasonString_DuplicateBeginTransition(Resource, SubresourceIndex, State.Current, TargetState, State.CreateTransitionBacktrace, BeginTrace));
 
 		// Validate the explicit previous state from the RHI matches what we expect...
 		{
 			// Check for the correct pipeline
-			ensureMsgf(EnumHasAllFlags(CurrentStateFromRHI.Pipelines, ExecutingPipeline), TEXT("%s"), *GetReasonString_WrongPipeline(Resource, SubresourceIndex, State.Current, TargetState));
+			RHI_VALIDATION_CHECK(EnumHasAllFlags(CurrentStateFromRHI.Pipelines, ExecutingPipeline), *GetReasonString_WrongPipeline(Resource, SubresourceIndex, State.Current, TargetState));
 
 			// Check the current RHI state passed in matches the tracked state for the resource, or is "unknown".
-			ensureMsgf(CurrentStateFromRHI.Access == ERHIAccess::Unknown || (CurrentStateFromRHI.Access == State.Previous.Access && CurrentStateFromRHI.Pipelines == State.Previous.Pipelines),
-				TEXT("%s"), *GetReasonString_IncorrectPreviousState(Resource, SubresourceIndex, State.Previous, CurrentStateFromRHI));
+			RHI_VALIDATION_CHECK(CurrentStateFromRHI.Access == ERHIAccess::Unknown || (CurrentStateFromRHI.Access == State.Previous.Access && CurrentStateFromRHI.Pipelines == State.Previous.Pipelines),
+				*GetReasonString_IncorrectPreviousState(Resource, SubresourceIndex, State.Previous, CurrentStateFromRHI));
 		}
 
 		// Check for unnecessary transitions
 		// @todo: this check is not particularly useful at the moment, as there are many unnecessary resource transitions.
-		//ensureMsgf(CurrentState != TargetState, TEXT("%s"), *GetReasonString_UnnecessaryTransition(Resource, SubresourceIndex, CurrentState));
+		//RHI_VALIDATION_CHECK(CurrentState != TargetState, *GetReasonString_UnnecessaryTransition(Resource, SubresourceIndex, CurrentState));
 
 		// Update the tracked state once all pipes have begun.
 		State.Previous = TargetState;
@@ -681,12 +740,12 @@ namespace RHIValidation
 		FPipelineState& State = States[ExecutingPipeline];
 
 		// Check that we aren't ending a transition that never began.
-		ensureMsgf(State.bTransitioning, TEXT("Unsolicited resource end transition call."));
+		RHI_VALIDATION_CHECK(State.bTransitioning, TEXT("Unsolicited resource end transition call."));
 		State.bTransitioning = false;
 		State.BeginTransitionBacktrace = nullptr;
 
 		// Check that the end matches the begin.
-		ensureMsgf(TargetState == State.Current, TEXT("%s"), *GetReasonString_MismatchedEndTransition(Resource, SubresourceIndex, State.Current, TargetState));
+		RHI_VALIDATION_CHECK(TargetState == State.Current, *GetReasonString_MismatchedEndTransition(Resource, SubresourceIndex, State.Current, TargetState));
 
 		// Replicate the state to other pipes that are not part of the end pipe mask.
 		for (ERHIPipeline OtherPipeline : GetRHIPipelines())
@@ -714,13 +773,13 @@ namespace RHIValidation
 		FPipelineState& State = States[RequiredState.Pipelines];
 
 		// Check we're not trying to access the resource whilst a pending resource transition is in progress.
-		ensureMsgf(!State.bTransitioning, TEXT("%s"), *GetReasonString_AccessDuringTransition(Resource, SubresourceIndex, State.Current, RequiredState, State.CreateTransitionBacktrace, State.BeginTransitionBacktrace));
+		RHI_VALIDATION_CHECK(!State.bTransitioning, *GetReasonString_AccessDuringTransition(Resource, SubresourceIndex, State.Current, RequiredState, State.CreateTransitionBacktrace, State.BeginTransitionBacktrace));
 
 		// If UAV overlaps are now disabled, ensure the resource has been transitioned if it was previously used in UAV overlap state.
-		ensureMsgf((bAllowAllUAVsOverlap || !State.bUsedWithAllUAVsOverlap) && (State.bExplicitAllowUAVOverlap || !State.bUsedWithExplicitUAVsOverlap), TEXT("%s"), *GetReasonString_UAVOverlap(Resource, SubresourceIndex, State.Current, RequiredState));
+		RHI_VALIDATION_CHECK((bAllowAllUAVsOverlap || !State.bUsedWithAllUAVsOverlap) && (State.bExplicitAllowUAVOverlap || !State.bUsedWithExplicitUAVsOverlap), *GetReasonString_UAVOverlap(Resource, SubresourceIndex, State.Current, RequiredState));
 
 		// Ensure the resource is in the required state for this operation
-		ensureMsgf(EnumHasAllFlags(State.Current.Access, RequiredState.Access) && EnumHasAllFlags(State.Current.Pipelines, RequiredState.Pipelines), TEXT("%s"), *GetReasonString_MissingBarrier(Resource, SubresourceIndex, State.Current, RequiredState));
+		RHI_VALIDATION_CHECK(EnumHasAllFlags(State.Current.Access, RequiredState.Access) && EnumHasAllFlags(State.Current.Pipelines, RequiredState.Pipelines), *GetReasonString_MissingBarrier(Resource, SubresourceIndex, State.Current, RequiredState));
 
 		State.Previous = State.Current;
 
@@ -746,7 +805,7 @@ namespace RHIValidation
 		}
 
 		FPipelineState& State = States[Pipeline];
-		ensureMsgf(State.bExplicitAllowUAVOverlap != bAllow, TEXT("%s"), *GetReasonString_MismatchedExplicitUAVOverlapCall(bAllow));
+		RHI_VALIDATION_CHECK(State.bExplicitAllowUAVOverlap != bAllow, *GetReasonString_MismatchedExplicitUAVOverlapCall(bAllow));
 		State.bExplicitAllowUAVOverlap = bAllow;
 	}
 
@@ -876,7 +935,7 @@ namespace RHIValidation
 			}
 
 		case EOpType::AllUAVsOverlap:
-			ensureMsgf(bAllowAllUAVsOverlap != Data_AllUAVsOverlap.bAllow, TEXT("%s"), *GetReasonString_MismatchedAllUAVsOverlapCall(Data_AllUAVsOverlap.bAllow));
+			RHI_VALIDATION_CHECK(bAllowAllUAVsOverlap != Data_AllUAVsOverlap.bAllow, *GetReasonString_MismatchedAllUAVsOverlapCall(Data_AllUAVsOverlap.bAllow));
 			bAllowAllUAVsOverlap = Data_AllUAVsOverlap.bAllow;
 			break;
 
@@ -980,7 +1039,7 @@ namespace RHIValidation
 			{
 				ErrorMessage += FString::Printf(TEXT("\nAllocation callstack: (void**)0x%p,32"), AllocatedCallstack);
 			}
-			ensureMsgf(false, TEXT("%s"), *ErrorMessage);
+			RHI_VALIDATION_CHECK(false, *ErrorMessage);
 		}		
 	}
 
diff --git a/Engine/Source/Runtime/RHI/Public/RHIValidation.h b/Engine/Source/Runtime/RHI/Public/RHIValidation.h
--- a/Engine/Source/Runtime/RHI/Public/RHIValidation.h
+++ b/Engine/Source/Runtime/RHI/Public/RHIValidation.h
@@ -15,6 +15,16 @@
 class FValidationComputeContext;
 class FValidationContext;
 
+// This is a macro because we only want to evaluate the message expression if the checked expression is false.
+#define RHI_VALIDATION_CHECK(InExpression, InMessage) \
+	do \
+	{ \
+		if(UNLIKELY(!(InExpression))) \
+		{ \
+			FValidationRHI::ReportValidationFailure(InMessage); \
+		} \
+	}while(0)
+
 class FValidationRHI : public FDynamicRHI
 {
 public:
@@ -1835,6 +1845,8 @@ public:
 
 
 //protected:
+	static void ReportValidationFailure(const TCHAR* InMessage);
+
 	FDynamicRHI*				RHI;
 	TIndirectArray<IRHIComputeContext> OwnedContexts;
 	TMap<FRHIDepthStencilState*, FDepthStencilStateInitializerRHI> DepthStencilStates;
@@ -1844,6 +1856,8 @@ public:
 
 private:
 	FString						RHIName;
+	static TSet<uint32>			SeenFailureHashes;
+	static FCriticalSection		SeenFailureHashesMutex;
 
 	void ValidatePipeline(const FGraphicsPipelineStateInitializer& Initializer);
 };
