Buffer refactor merging Vertex, Index and Structured Buffer types into a single FRHIBuffer. No changes to APIs in this changelist other than those that naturally folded together when types merged. A f

diff --git a/Engine/Source/Runtime/OpenGLDrv/Public/OpenGLResources.h b/Engine/Source/Runtime/OpenGLDrv/Public/OpenGLResources.h
--- a/Engine/Source/Runtime/OpenGLDrv/Public/OpenGLResources.h
+++ b/Engine/Source/Runtime/OpenGLDrv/Public/OpenGLResources.h
@@ -29,15 +29,13 @@ class FOpenGLLinkedProgram;
 typedef TArray<ANSICHAR> FAnsiCharArray;
 
 
-extern void OnVertexBufferDeletion( GLuint VertexBufferResource );
-extern void OnIndexBufferDeletion( GLuint IndexBufferResource );
+extern void OnBufferDeletion( GLuint BufferResource );
 extern void OnPixelBufferDeletion( GLuint PixelBufferResource );
 extern void OnUniformBufferDeletion( GLuint UniformBufferResource, uint32 AllocatedSize, bool bStreamDraw, uint32 Offset, uint8* Pointer );
 extern void OnProgramDeletion( GLint ProgramResource );
 
-extern void CachedBindArrayBuffer( GLuint Buffer );
-extern void CachedBindElementArrayBuffer( GLuint Buffer );
-extern void CachedBindPixelUnpackBuffer( GLuint Buffer );
+extern void CachedBindBuffer( GLenum Type, GLuint Buffer );
+extern void CachedBindPixelUnpackBuffer( GLenum Type, GLuint Buffer );
 extern void CachedBindUniformBuffer( GLuint Buffer );
 extern bool IsUniformBufferBound( GLuint Buffer );
 
@@ -374,12 +372,12 @@ struct TIsGLProxyObject<TOpenGLResourceProxy<TRHIType, TOGLResourceType>>
 	enum { Value = true };
 };
 
-typedef void (*BufferBindFunction)( GLuint Buffer );
+typedef void (*BufferBindFunction)( GLenum Type, GLuint Buffer );
 
 template<typename BaseType>
 class TOpenGLTexture;
 
-template <typename BaseType, GLenum Type, BufferBindFunction BufBind>
+template <typename BaseType, BufferBindFunction BufBind>
 class TOpenGLBuffer : public BaseType
 {
 	void LoadData( uint32 InOffset, uint32 InSize, const void* InData)
@@ -416,12 +414,14 @@ class TOpenGLBuffer : public BaseType
 public:
 
 	GLuint Resource;
+	GLenum Type;
 
 	/** Needed on OS X to force a rebind of the texture buffer to the texture name to workaround radr://18379338 */
 	uint64 ModificationCount;
 
 	TOpenGLBuffer()
 		: Resource(0)
+		, Type(0)
 		, ModificationCount(0)
 		, bIsLocked(false)
 		, bIsLockReadOnly(false)
@@ -433,10 +433,11 @@ public:
 		, RealSize(0)
 	{ }
 
-	TOpenGLBuffer(uint32 InStride,uint32 InSize,uint32 InUsage,
-		const void *InData = NULL, bool bStreamedDraw = false, GLuint ResourceToUse = 0, uint32 ResourceSize = 0)
+	TOpenGLBuffer(GLenum InType, uint32 InStride, uint32 InSize, uint32 InUsage,
+		const void *InData, bool bStreamedDraw = false, GLuint ResourceToUse = 0, uint32 ResourceSize = 0)
 	: BaseType(InStride,InSize,InUsage)
 	, Resource(0)
+	, Type(InType)
 	, ModificationCount(0)
 	, bIsLocked(false)
 	, bIsLockReadOnly(false)
@@ -563,7 +564,7 @@ public:
 	{
 		VERIFY_GL_SCOPE();
 		check( (FOpenGL::SupportsVertexAttribBinding() && OpenGLConsoleVariables::bUseVAB) || ( this->GetUsage() & BUF_ZeroStride ) == 0 );
-		BufBind(Resource);
+		BufBind(Type, Resource);
 	}
 
 	uint8 *Lock(uint32 InOffset, uint32 InSize, bool bReadOnly, bool bDiscard)
@@ -891,17 +892,17 @@ private:
 	uint32 Usage;
 };
 
-class FOpenGLBaseVertexBuffer : public FRHIVertexBuffer
+class FOpenGLBaseBuffer : public FRHIBuffer
 {
 public:
-	FOpenGLBaseVertexBuffer() : ZeroStrideVertexBuffer(0)
+	FOpenGLBaseBuffer() : ZeroStrideBuffer(0)
 	{}
 
-	FOpenGLBaseVertexBuffer(uint32 InStride,uint32 InSize,uint32 InUsage): FRHIVertexBuffer(InSize,InUsage), ZeroStrideVertexBuffer(0)
+	FOpenGLBaseBuffer(uint32 InStride, uint32 InSize, uint32 InUsage): FRHIBuffer(InSize, InUsage, InStride), ZeroStrideBuffer(0)
 	{
 		if(!(FOpenGL::SupportsVertexAttribBinding() && OpenGLConsoleVariables::bUseVAB) && (InUsage & BUF_ZeroStride))
 		{
-			ZeroStrideVertexBuffer = FMemory::Malloc( InSize );
+			ZeroStrideBuffer = FMemory::Malloc( InSize );
 		}
 
 #if ENABLE_LOW_LEVEL_MEM_TRACKER
@@ -910,11 +911,11 @@ public:
 #endif
 	}
 
-	~FOpenGLBaseVertexBuffer( void )
+	~FOpenGLBaseBuffer( void )
 	{
-		if( ZeroStrideVertexBuffer )
+		if( ZeroStrideBuffer )
 		{
-			FMemory::Free( ZeroStrideVertexBuffer );
+			FMemory::Free( ZeroStrideBuffer );
 		}
 
 #if ENABLE_LOW_LEVEL_MEM_TRACKER
@@ -925,19 +926,19 @@ public:
 
 	void* GetZeroStrideBuffer( void )
 	{
-		check( ZeroStrideVertexBuffer );
-		return ZeroStrideVertexBuffer;
+		check( ZeroStrideBuffer );
+		return ZeroStrideBuffer;
 	}
 
-	void Swap(FOpenGLBaseVertexBuffer& Other)
+	void Swap(FOpenGLBaseBuffer& Other)
 	{
-		FRHIVertexBuffer::Swap(Other);
-		::Swap(ZeroStrideVertexBuffer, Other.ZeroStrideVertexBuffer);
+		FRHIBuffer::Swap(Other);
+		::Swap(ZeroStrideBuffer, Other.ZeroStrideBuffer);
 	}
 
 	static bool OnDelete(GLuint Resource,uint32 Size,bool bStreamDraw,uint32 Offset)
 	{
-		OnVertexBufferDeletion(Resource);
+		OnBufferDeletion(Resource);
 		return true;
 	}
 
@@ -954,7 +955,7 @@ public:
 	static bool IsStructuredBuffer() { return false; }
 
 private:
-	void*	ZeroStrideVertexBuffer;
+	void*	ZeroStrideBuffer;
 };
 
 struct FOpenGLEUniformBufferData : public FRefCountedObject
@@ -1015,73 +1016,8 @@ public:
 	FOpenGLAssertRHIThreadFence CopyFence;
 };
 
-
-class FOpenGLBaseIndexBuffer : public FRHIIndexBuffer
-{
-public:
-	FOpenGLBaseIndexBuffer()
-	{}
-
-	FOpenGLBaseIndexBuffer(uint32 InStride,uint32 InSize,uint32 InUsage): FRHIIndexBuffer(InStride,InSize,InUsage)
-	{
-#if ENABLE_LOW_LEVEL_MEM_TRACKER
-		LLM_SCOPED_PAUSE_TRACKING_WITH_ENUM_AND_AMOUNT(ELLMTag::Meshes, InSize, ELLMTracker::Default, ELLMAllocType::None);
-#endif
-	}
-
-	~FOpenGLBaseIndexBuffer(void)
-	{
-#if ENABLE_LOW_LEVEL_MEM_TRACKER
-		LLM_SCOPED_PAUSE_TRACKING_WITH_ENUM_AND_AMOUNT(ELLMTag::Meshes, -(int64)GetSize(), ELLMTracker::Default, ELLMAllocType::None);
-#endif
-	}
-
-	static bool OnDelete(GLuint Resource,uint32 Size,bool bStreamDraw,uint32 Offset)
-	{
-		OnIndexBufferDeletion(Resource);
-		return true;
-	}
-
-	static FORCEINLINE bool GLSupportsType()
-	{
-		return true;
-	}
-
-	static void CreateType(GLuint& Resource, const void* InData, uint32 InSize)
-	{
-		// @todo-mobile
-	}
-
-	static bool IsStructuredBuffer() { return false; }
-};
-
-class FOpenGLBaseStructuredBuffer : public FRHIStructuredBuffer
-{
-public:
-	FOpenGLBaseStructuredBuffer(uint32 InStride,uint32 InSize,uint32 InUsage): FRHIStructuredBuffer(InStride,InSize,InUsage) {}
-	static bool OnDelete(GLuint Resource,uint32 Size,bool bStreamDraw,uint32 Offset)
-	{
-		OnVertexBufferDeletion(Resource);
-		return true;
-	}
-
-	static FORCEINLINE bool GLSupportsType()
-	{
-		return FOpenGL::SupportsStructuredBuffers();
-	}
-
-	static void CreateType(GLuint& Resource, const void* InData, uint32 InSize)
-	{
-		// @todo-mobile
-	}
-
-	static bool IsStructuredBuffer() { return true; }
-};
-
-typedef TOpenGLBuffer<FOpenGLBasePixelBuffer, GL_PIXEL_UNPACK_BUFFER, CachedBindPixelUnpackBuffer> FOpenGLPixelBuffer;
-typedef TOpenGLBuffer<FOpenGLBaseVertexBuffer, GL_ARRAY_BUFFER, CachedBindArrayBuffer> FOpenGLVertexBuffer;
-typedef TOpenGLBuffer<FOpenGLBaseIndexBuffer,GL_ELEMENT_ARRAY_BUFFER,CachedBindElementArrayBuffer> FOpenGLIndexBuffer;
-typedef TOpenGLBuffer<FOpenGLBaseStructuredBuffer,GL_ARRAY_BUFFER,CachedBindArrayBuffer> FOpenGLStructuredBuffer;
+typedef TOpenGLBuffer<FOpenGLBasePixelBuffer, CachedBindPixelUnpackBuffer> FOpenGLPixelBuffer;
+typedef TOpenGLBuffer<FOpenGLBaseBuffer, CachedBindBuffer> FOpenGLBuffer;
 
 #define MAX_STREAMED_BUFFERS_IN_ARRAY 2	// must be > 1!
 #define MIN_DRAWS_IN_SINGLE_BUFFER 16
@@ -1171,9 +1107,6 @@ private:
 	uint32 MinNeededBufferSize;
 };
 
-typedef TOpenGLStreamedBufferArray<FOpenGLVertexBuffer,0> FOpenGLStreamedVertexBufferArray;
-typedef TOpenGLStreamedBufferArray<FOpenGLIndexBuffer,sizeof(uint16)> FOpenGLStreamedIndexBufferArray;
-
 struct FOpenGLVertexElement
 {
 	GLenum Type;
@@ -1967,7 +1900,7 @@ public:
 	{
 		if (IndexBuffer)
 		{
-			FOpenGLIndexBuffer* IB = (FOpenGLIndexBuffer*)IndexBuffer.GetReference();
+			FOpenGLBuffer* IB = (FOpenGLBuffer*)IndexBuffer.GetReference();
 			ModificationVersion = IB->ModificationCount;
 		}
 	}
@@ -1984,7 +1917,7 @@ public:
 	{
 		if (VertexBuffer)
 		{
-			FOpenGLVertexBuffer* VB = (FOpenGLVertexBuffer*)VertexBuffer.GetReference();
+			FOpenGLBuffer* VB = (FOpenGLBuffer*)VertexBuffer.GetReference();
 			ModificationVersion = VB->ModificationCount;
 		}
 	}
@@ -2266,19 +2199,9 @@ struct TOpenGLResourceTraits<FRHIUniformBuffer>
 	typedef FOpenGLUniformBuffer TConcreteType;
 };
 template<>
-struct TOpenGLResourceTraits<FRHIIndexBuffer>
-{
-	typedef FOpenGLIndexBuffer TConcreteType;
-};
-template<>
-struct TOpenGLResourceTraits<FRHIStructuredBuffer>
-{
-	typedef FOpenGLStructuredBuffer TConcreteType;
-};
-template<>
-struct TOpenGLResourceTraits<FRHIVertexBuffer>
+struct TOpenGLResourceTraits<FRHIBuffer>
 {
-	typedef FOpenGLVertexBuffer TConcreteType;
+	typedef FOpenGLBuffer TConcreteType;
 };
 template<>
 struct TOpenGLResourceTraits<FRHIShaderResourceView>
diff --git a/Engine/Source/Runtime/OpenGLDrv/Public/OpenGLState.h b/Engine/Source/Runtime/OpenGLDrv/Public/OpenGLState.h
--- a/Engine/Source/Runtime/OpenGLDrv/Public/OpenGLState.h
+++ b/Engine/Source/Runtime/OpenGLDrv/Public/OpenGLState.h
@@ -269,7 +269,7 @@ struct FOpenGLStream
 		, Offset(0)
 		, Divisor(0)
 	{}
-	FOpenGLVertexBuffer *VertexBuffer;
+	FOpenGLBuffer *VertexBuffer;
 	uint32 Stride;
 	uint32 Offset;
 	uint32 Divisor;
diff --git a/Engine/Source/Runtime/RHI/Private/DynamicRHI.cpp b/Engine/Source/Runtime/RHI/Private/DynamicRHI.cpp
--- a/Engine/Source/Runtime/RHI/Private/DynamicRHI.cpp
+++ b/Engine/Source/Runtime/RHI/Private/DynamicRHI.cpp
@@ -423,14 +423,9 @@ void FDynamicRHI::EnableIdealGPUCaptureOptions(bool bEnabled)
 	}	
 }
 
-void FDynamicRHI::RHITransferIndexBufferUnderlyingResource(FRHIIndexBuffer* DestIndexBuffer, FRHIIndexBuffer* SrcIndexBuffer)
+void FDynamicRHI::RHITransferBufferUnderlyingResource(FRHIBuffer* DestBuffer, FRHIBuffer* SrcBuffer)
 {
-	UE_LOG(LogRHI, Fatal, TEXT("RHITransferIndexBufferUnderlyingResource isn't implemented for the current RHI"));
-}
-
-void FDynamicRHI::RHITransferVertexBufferUnderlyingResource(FRHIVertexBuffer* DestVertexBuffer, FRHIVertexBuffer* SrcVertexBuffer)
-{
-	UE_LOG(LogRHI, Fatal, TEXT("RHITransferVertexBufferUnderlyingResource isn't implemented for the current RHI"));
+	UE_LOG(LogRHI, Fatal, TEXT("RHITransferBufferUnderlyingResource isn't implemented for the current RHI"));
 }
 
 FUnorderedAccessViewRHIRef FDynamicRHI::RHICreateUnorderedAccessView(FRHITexture* Texture, uint32 MipLevel, uint8 Format)
@@ -439,12 +434,12 @@ FUnorderedAccessViewRHIRef FDynamicRHI::RHICreateUnorderedAccessView(FRHITexture
 	return RHICreateUnorderedAccessView(Texture, MipLevel);
 }
 
-void FDynamicRHI::RHIUpdateShaderResourceView(FRHIShaderResourceView* SRV, FRHIVertexBuffer* VertexBuffer, uint32 Stride, uint8 Format)
+void FDynamicRHI::RHIUpdateShaderResourceView(FRHIShaderResourceView* SRV, FRHIBuffer* Buffer, uint32 Stride, uint8 Format)
 {
 	UE_LOG(LogRHI, Fatal, TEXT("RHIUpdateShaderResourceView isn't implemented for the current RHI"));
 }
 
-void FDynamicRHI::RHIUpdateShaderResourceView(FRHIShaderResourceView* SRV, FRHIIndexBuffer* IndexBuffer)
+void FDynamicRHI::RHIUpdateShaderResourceView(FRHIShaderResourceView* SRV, FRHIBuffer* Buffer)
 {
 	UE_LOG(LogRHI, Fatal, TEXT("RHIUpdateShaderResourceView isn't implemented for the current RHI"));
 }
@@ -535,51 +530,57 @@ void FDynamicRHI::RHICheckViewportHDRStatus(FRHIViewport* Viewport)
 }
 
 
-FShaderResourceViewInitializer::FShaderResourceViewInitializer(FRHIVertexBuffer* InVertexBuffer, EPixelFormat InFormat, uint32 InStartOffsetBytes, uint32 InNumElements)
-	: VertexBufferInitializer({ InVertexBuffer, InStartOffsetBytes, InNumElements, InFormat }), Type(EType::VertexBufferSRV)
+FShaderResourceViewInitializer::FShaderResourceViewInitializer(FRHIBuffer* InBuffer, EPixelFormat InFormat, uint32 InStartOffsetBytes, uint32 InNumElements)
+	: BufferInitializer({ InBuffer, InStartOffsetBytes, InNumElements, InFormat }), Type(EType::VertexBufferSRV)
 {
 	check(InStartOffsetBytes % RHIGetMinimumAlignmentForBufferBackedSRV(InFormat) == 0);
-	/*if (!VertexBufferInitializer.IsWholeResource())
+	/*if (!BufferInitializer.IsWholeResource())
 	{
 		const uint32 Stride = GPixelFormats[InFormat].BlockBytes;
-		check((VertexBufferInitializer.NumElements * Stride + VertexBufferInitializer.StartOffsetBytes)  <= VertexBufferInitializer.VertexBuffer->GetSize());
+		check((BufferInitializer.NumElements * Stride + BufferInitializer.StartOffsetBytes) <= BufferInitializer.Buffer->GetSize());
 	}*/
+
+	InitType();
 }
 
-FShaderResourceViewInitializer::FShaderResourceViewInitializer(FRHIVertexBuffer* InVertexBuffer, EPixelFormat InFormat)
-	: VertexBufferInitializer({ InVertexBuffer, 0, UINT32_MAX, InFormat }), Type(EType::VertexBufferSRV) 
+FShaderResourceViewInitializer::FShaderResourceViewInitializer(FRHIBuffer* InBuffer, EPixelFormat InFormat)
+	: BufferInitializer({ InBuffer, 0, UINT32_MAX, InFormat }), Type(EType::VertexBufferSRV) 
 {
+	InitType();
 }
 
-FShaderResourceViewInitializer::FShaderResourceViewInitializer(FRHIStructuredBuffer* InStructuredBuffer, uint32 InStartOffsetBytes, uint32 InNumElements)
-	: StructuredBufferInitializer(FStructuredBufferShaderResourceViewInitializer{ InStructuredBuffer, InStartOffsetBytes, InNumElements }), Type(EType::StructuredBufferSRV)
+FShaderResourceViewInitializer::FShaderResourceViewInitializer(FRHIBuffer* InBuffer, uint32 InStartOffsetBytes, uint32 InNumElements)
+	: BufferInitializer({ InBuffer, InStartOffsetBytes, InNumElements, PF_Unknown }), Type(EType::StructuredBufferSRV)
 {
-	check(InStartOffsetBytes % InStructuredBuffer->GetStride() == 0);
-	if (!StructuredBufferInitializer.IsWholeResource())
+	check(InStartOffsetBytes % InBuffer->GetStride() == 0);
+	if (!BufferInitializer.IsWholeResource())
 	{
-		const uint32 Stride = StructuredBufferInitializer.StructuredBuffer->GetStride();
-		check((StructuredBufferInitializer.NumElements * Stride + StructuredBufferInitializer.StartOffsetBytes)  <= StructuredBufferInitializer.StructuredBuffer->GetSize());
+		const uint32 Stride = BufferInitializer.Buffer->GetStride();
+		check((BufferInitializer.NumElements * Stride + BufferInitializer.StartOffsetBytes) <= BufferInitializer.Buffer->GetSize());
 	}
+
+	InitType();
 }
 
-FShaderResourceViewInitializer::FShaderResourceViewInitializer(FRHIStructuredBuffer* InStructuredBuffer)
-	: StructuredBufferInitializer(FStructuredBufferShaderResourceViewInitializer{ InStructuredBuffer, 0, UINT32_MAX }), Type(EType::StructuredBufferSRV) 
+FShaderResourceViewInitializer::FShaderResourceViewInitializer(FRHIBuffer* InBuffer)
+	: BufferInitializer({ InBuffer, 0, UINT32_MAX }), Type(EType::StructuredBufferSRV) 
 {
+	InitType();
 }
 
-FShaderResourceViewInitializer::FShaderResourceViewInitializer(FRHIIndexBuffer* InIndexBuffer, uint32 InStartOffsetBytes, uint32 InNumElements)
-	: IndexBufferInitializer(FIndexBufferShaderResourceViewInitializer{ InIndexBuffer, InStartOffsetBytes, InNumElements }), Type(EType::IndexBufferSRV)
+void FShaderResourceViewInitializer::InitType()
 {
-	check(InStartOffsetBytes % RHIGetMinimumAlignmentForBufferBackedSRV(InIndexBuffer->GetStride() == 2 ? PF_R16_UINT : PF_R32_UINT) == 0);
-	if (!IndexBufferInitializer.IsWholeResource())
+	uint32 Usage = BufferInitializer.Buffer->GetUsage();
+	if (Usage & BUF_VertexBuffer)
 	{
-		const uint32 Stride = IndexBufferInitializer.IndexBuffer->GetStride();
-		check((IndexBufferInitializer.NumElements * Stride + IndexBufferInitializer.StartOffsetBytes) <= IndexBufferInitializer.IndexBuffer->GetSize());
+		Type = EType::VertexBufferSRV;
+	}
+	else if (Usage & BUF_IndexBuffer)
+	{
+		Type = EType::IndexBufferSRV;
+	}
+	else
+	{
+		Type = EType::StructuredBufferSRV;
 	}
 }
-
-FShaderResourceViewInitializer::FShaderResourceViewInitializer(FRHIIndexBuffer* InIndexBuffer)
-	: IndexBufferInitializer(FIndexBufferShaderResourceViewInitializer{ InIndexBuffer, 0, UINT32_MAX }), Type(EType::IndexBufferSRV) 
-{
-}
-
diff --git a/Engine/Source/Runtime/RHI/Private/RHICommandList.cpp b/Engine/Source/Runtime/RHI/Private/RHICommandList.cpp
--- a/Engine/Source/Runtime/RHI/Private/RHICommandList.cpp
+++ b/Engine/Source/Runtime/RHI/Private/RHICommandList.cpp
@@ -2542,11 +2542,11 @@ FTexture3DRHIRef FDynamicRHI::RHICreateTexture3D_RenderThread(class FRHICommandL
 	return GDynamicRHI->RHICreateTexture3D(SizeX, SizeY, SizeZ, Format, NumMips, Flags, InResourceState, CreateInfo);
 }
 
-FUnorderedAccessViewRHIRef FDynamicRHI::RHICreateUnorderedAccessView_RenderThread(class FRHICommandListImmediate& RHICmdList, FRHIStructuredBuffer* StructuredBuffer, bool bUseUAVCounter, bool bAppendBuffer)
+FUnorderedAccessViewRHIRef FDynamicRHI::RHICreateUnorderedAccessView_RenderThread(class FRHICommandListImmediate& RHICmdList, FRHIBuffer* Buffer, bool bUseUAVCounter, bool bAppendBuffer)
 {
 	CSV_SCOPED_TIMING_STAT(RHITStalls, RHICreateUnorderedAccessView_RenderThread);
 	FScopedRHIThreadStaller StallRHIThread(RHICmdList);
-	return GDynamicRHI->RHICreateUnorderedAccessView(StructuredBuffer, bUseUAVCounter, bAppendBuffer);
+	return GDynamicRHI->RHICreateUnorderedAccessView(Buffer, bUseUAVCounter, bAppendBuffer);
 }
 
 FUnorderedAccessViewRHIRef FDynamicRHI::RHICreateUnorderedAccessView_RenderThread(class FRHICommandListImmediate& RHICmdList, FRHITexture* Texture, uint32 MipLevel)
@@ -2563,17 +2563,11 @@ FUnorderedAccessViewRHIRef FDynamicRHI::RHICreateUnorderedAccessView_RenderThrea
 	return GDynamicRHI->RHICreateUnorderedAccessView(Texture, MipLevel, Format);
 }
 
-FUnorderedAccessViewRHIRef FDynamicRHI::RHICreateUnorderedAccessView_RenderThread(class FRHICommandListImmediate& RHICmdList, FRHIVertexBuffer* VertexBuffer, uint8 Format)
+FUnorderedAccessViewRHIRef FDynamicRHI::RHICreateUnorderedAccessView_RenderThread(class FRHICommandListImmediate& RHICmdList, FRHIBuffer* Buffer, uint8 Format)
 {
 	CSV_SCOPED_TIMING_STAT(RHITStalls, RHICreateUnorderedAccessView_RenderThread);
 	FScopedRHIThreadStaller StallRHIThread(RHICmdList);
-	return GDynamicRHI->RHICreateUnorderedAccessView(VertexBuffer, Format);
-}
-
-FUnorderedAccessViewRHIRef FDynamicRHI::RHICreateUnorderedAccessView_RenderThread(class FRHICommandListImmediate& RHICmdList, FRHIIndexBuffer* IndexBuffer, uint8 Format)
-{
-	FScopedRHIThreadStaller StallRHIThread(RHICmdList);
-	return GDynamicRHI->RHICreateUnorderedAccessView(IndexBuffer, Format);
+	return GDynamicRHI->RHICreateUnorderedAccessView(Buffer, Format);
 }
 
 FShaderResourceViewRHIRef FDynamicRHI::RHICreateShaderResourceView_RenderThread(class FRHICommandListImmediate& RHICmdList, FRHITexture* Texture, const FRHITextureSRVCreateInfo& CreateInfo)
@@ -2583,34 +2577,27 @@ FShaderResourceViewRHIRef FDynamicRHI::RHICreateShaderResourceView_RenderThread(
 	return GDynamicRHI->RHICreateShaderResourceView(Texture, CreateInfo);
 }
 
-FShaderResourceViewRHIRef FDynamicRHI::RHICreateShaderResourceView_RenderThread(class FRHICommandListImmediate& RHICmdList, FRHIVertexBuffer* VertexBuffer, uint32 Stride, uint8 Format)
+FShaderResourceViewRHIRef FDynamicRHI::RHICreateShaderResourceView_RenderThread(class FRHICommandListImmediate& RHICmdList, FRHIBuffer* Buffer, uint32 Stride, uint8 Format)
 {
-	CSV_SCOPED_TIMING_STAT(RHITStalls, RHICreateShaderResourceView_RenderThread_VB);
+	CSV_SCOPED_TIMING_STAT(RHITStalls, RHICreateShaderResourceView_RenderThread);
 	FScopedRHIThreadStaller StallRHIThread(RHICmdList);
-	return GDynamicRHI->RHICreateShaderResourceView(VertexBuffer, Stride, Format);
+	return GDynamicRHI->RHICreateShaderResourceView(Buffer, Stride, Format);
 }
 
 FShaderResourceViewRHIRef FDynamicRHI::RHICreateShaderResourceView_RenderThread(class FRHICommandListImmediate& RHICmdList, const FShaderResourceViewInitializer& Initializer)
 {
-	CSV_SCOPED_TIMING_STAT(RHITStalls, RHICreateShaderResourceView_RenderThread_VB);
+	CSV_SCOPED_TIMING_STAT(RHITStalls, RHICreateShaderResourceView_RenderThread);
 	FScopedRHIThreadStaller StallRHIThread(RHICmdList);
 	return GDynamicRHI->RHICreateShaderResourceView(Initializer);
 }
 
-FShaderResourceViewRHIRef FDynamicRHI::RHICreateShaderResourceView_RenderThread(class FRHICommandListImmediate& RHICmdList, FRHIIndexBuffer* Buffer)
+FShaderResourceViewRHIRef FDynamicRHI::RHICreateShaderResourceView_RenderThread(class FRHICommandListImmediate& RHICmdList, FRHIBuffer* Buffer)
 {
-	CSV_SCOPED_TIMING_STAT(RHITStalls, RHICreateShaderResourceView_RenderThread_IB);
+	CSV_SCOPED_TIMING_STAT(RHITStalls, RHICreateShaderResourceView_RenderThread);
 	FScopedRHIThreadStaller StallRHIThread(RHICmdList);
 	return GDynamicRHI->RHICreateShaderResourceView(Buffer);
 }
 
-FShaderResourceViewRHIRef FDynamicRHI::RHICreateShaderResourceView_RenderThread(class FRHICommandListImmediate& RHICmdList, FRHIStructuredBuffer* StructuredBuffer)
-{
-	CSV_SCOPED_TIMING_STAT(RHITStalls, RHICreateShaderResourceView_RenderThread_SB);
-	FScopedRHIThreadStaller StallRHIThread(RHICmdList);
-	return GDynamicRHI->RHICreateShaderResourceView(StructuredBuffer);
-}
-
 FShaderResourceViewRHIRef FDynamicRHI::RHICreateShaderResourceViewWriteMask_RenderThread(class FRHICommandListImmediate& RHICmdList, FRHITexture2D* Texture2DRHI)
 {
 	CSV_SCOPED_TIMING_STAT(RHITStalls, RHICreateShaderResourceView_RenderThread_Tex2DWriteMask);
diff --git a/Engine/Source/Runtime/RHI/Public/DynamicRHI.h b/Engine/Source/Runtime/RHI/Public/DynamicRHI.h
--- a/Engine/Source/Runtime/RHI/Public/DynamicRHI.h
+++ b/Engine/Source/Runtime/RHI/Public/DynamicRHI.h
@@ -220,9 +220,9 @@ struct FRayTracingSceneInitializer
 
 struct FShaderResourceViewInitializer
 {
-	struct FVertexBufferShaderResourceViewInitializer
+	struct FBufferShaderResourceViewInitializer
 	{
-		FRHIVertexBuffer* VertexBuffer;
+		FRHIBuffer* Buffer;
 		uint32 StartOffsetBytes;
 		uint32 NumElements;
 		EPixelFormat Format;
@@ -233,55 +233,37 @@ struct FShaderResourceViewInitializer
 		}
 	};
 
-	RHI_API FShaderResourceViewInitializer(FRHIVertexBuffer* InVertexBuffer, EPixelFormat InFormat, uint32 InStartOffsetBytes, uint32 InNumElements);
-	RHI_API FShaderResourceViewInitializer(FRHIVertexBuffer* InVertexBuffer, EPixelFormat InFormat);
+	typedef FBufferShaderResourceViewInitializer FVertexBufferShaderResourceViewInitializer;
+	typedef FBufferShaderResourceViewInitializer FIndexBufferShaderResourceViewInitializer;
+	typedef FBufferShaderResourceViewInitializer FStructuredBufferShaderResourceViewInitializer;
+
+
+	RHI_API FShaderResourceViewInitializer(FRHIBuffer* InBuffer, EPixelFormat InFormat, uint32 InStartOffsetBytes, uint32 InNumElements);
+	RHI_API FShaderResourceViewInitializer(FRHIBuffer* InBuffer, EPixelFormat InFormat);
+	RHI_API FShaderResourceViewInitializer(FRHIBuffer* InBuffer, uint32 InStartOffsetBytes, uint32 InNumElements);
+	RHI_API FShaderResourceViewInitializer(FRHIBuffer* InBuffer);
 
 	const FVertexBufferShaderResourceViewInitializer& AsVertexBufferSRV() const
 	{
 		check(Type == EType::VertexBufferSRV);
-		return VertexBufferInitializer;
+		return BufferInitializer;
 	}
 
-	struct FStructuredBufferShaderResourceViewInitializer
-	{
-		FRHIStructuredBuffer* StructuredBuffer;
-		uint32 StartOffsetBytes;
-		uint32 NumElements;
-
-		inline bool IsWholeResource() const 
-		{
-			return StartOffsetBytes == 0 && NumElements == UINT32_MAX;
-		}
-	};
-
-	RHI_API FShaderResourceViewInitializer(FRHIStructuredBuffer* InStructuredBuffer, uint32 InStartOffsetBytes, uint32 InNumElements);
-	RHI_API FShaderResourceViewInitializer(FRHIStructuredBuffer* InStructuredBuffer);
-
 	const FStructuredBufferShaderResourceViewInitializer& AsStructuredBufferSRV() const
 	{
 		check(Type == EType::StructuredBufferSRV);
-		return StructuredBufferInitializer;
+		return BufferInitializer;
 	}
 
-	struct FIndexBufferShaderResourceViewInitializer
-	{
-		FRHIIndexBuffer* IndexBuffer;
-		uint32 StartOffsetBytes;
-		uint32 NumElements;
-
-		inline bool IsWholeResource() const 
-		{ 
-			return StartOffsetBytes == 0 && NumElements == UINT32_MAX;
-		}
-	};
-
-	RHI_API FShaderResourceViewInitializer(FRHIIndexBuffer* InIndexBuffer, uint32 InStartOffsetBytes, uint32 InNumElements);
-	RHI_API FShaderResourceViewInitializer(FRHIIndexBuffer* InIndexBuffer);
-
 	const FIndexBufferShaderResourceViewInitializer& AsIndexBufferSRV() const
 	{
 		check(Type == EType::IndexBufferSRV);
-		return IndexBufferInitializer;
+		return BufferInitializer;
+	}
+
+	const FBufferShaderResourceViewInitializer& AsBufferSRV() const
+	{
+		return BufferInitializer;
 	}
 
 	enum class EType : uint8
@@ -297,14 +279,11 @@ struct FShaderResourceViewInitializer
 	}
 
 private:
-	union
-	{
-		FVertexBufferShaderResourceViewInitializer VertexBufferInitializer;
-		FStructuredBufferShaderResourceViewInitializer StructuredBufferInitializer;
-		FIndexBufferShaderResourceViewInitializer IndexBufferInitializer;
-	};
+	FBufferShaderResourceViewInitializer BufferInitializer;
+
+	void InitType();
 
-	const EType Type;
+	EType Type;
 };
 
 class FDynamicRHI;
@@ -550,13 +529,6 @@ public:
 	// FlushType: Flush RHI Thread
 	virtual void RHIUnlockIndexBuffer(FRHICommandListImmediate& RHICmdList, FRHIIndexBuffer* IndexBuffer);
 
-	/**
-	* Transfer metadata and underlying resource from src to dest and release any resource owned by dest.
-	* @param DestIndexBuffer - the index buffer to update
-	* @param SrcIndexBuffer - don't use after call. If null, will release any resource owned by DestIndexBuffer
-	*/
-	virtual void RHITransferIndexBufferUnderlyingResource(FRHIIndexBuffer* DestIndexBuffer, FRHIIndexBuffer* SrcIndexBuffer);
-
 	/**
 	* @param ResourceArray - An optional pointer to a resource array containing the resource's data.
 	*/
@@ -575,10 +547,10 @@ public:
 
 	/**
 	 * Transfer metadata and underlying resource from src to dest and release any resource owned by dest.
-	 * @param DestVertexBuffer - the vertex buffer to update
-	 * @param SrcVertexBuffer - don't use after call. If null, will release any resource owned by DestVertexBuffer
+	 * @param DestBuffer - the buffer to update
+	 * @param SrcBuffer - don't use after call. If null, will release any resource owned by DestBuffer
 	 */
-	virtual void RHITransferVertexBufferUnderlyingResource(FRHIVertexBuffer* DestVertexBuffer, FRHIVertexBuffer* SrcVertexBuffer);
+	virtual void RHITransferBufferUnderlyingResource(FRHIBuffer* DestBuffer, FRHIBuffer* SrcBuffer);
 
 	/**
 	* @param ResourceArray - An optional pointer to a resource array containing the resource's data.
@@ -594,7 +566,7 @@ public:
 
 	/** Creates an unordered access view of the given structured buffer. */
 	// FlushType: Wait RHI Thread
-	virtual FUnorderedAccessViewRHIRef RHICreateUnorderedAccessView(FRHIStructuredBuffer* StructuredBuffer, bool bUseUAVCounter, bool bAppendBuffer) = 0;
+	virtual FUnorderedAccessViewRHIRef RHICreateUnorderedAccessView(FRHIBuffer* Buffer, bool bUseUAVCounter, bool bAppendBuffer) = 0;
 
 	/** Creates an unordered access view of the given texture. */
 	// FlushType: Wait RHI Thread
@@ -606,32 +578,24 @@ public:
 
 	/** Creates an unordered access view of the given vertex buffer. */
 	// FlushType: Wait RHI Thread
-	virtual FUnorderedAccessViewRHIRef RHICreateUnorderedAccessView(FRHIVertexBuffer* VertexBuffer, uint8 Format) = 0;
-
-	/** Creates an unordered access view of the given index buffer. */
-	// FlushType: Wait RHI Thread
-	virtual FUnorderedAccessViewRHIRef RHICreateUnorderedAccessView(FRHIIndexBuffer* IndexBuffer, uint8 Format) = 0;
+	virtual FUnorderedAccessViewRHIRef RHICreateUnorderedAccessView(FRHIBuffer* Buffer, uint8 Format) = 0;
 
 	/** Creates a shader resource view of the given structured buffer. */
 	// FlushType: Wait RHI Thread
-	virtual FShaderResourceViewRHIRef RHICreateShaderResourceView(FRHIStructuredBuffer* StructuredBuffer) = 0;
+	virtual FShaderResourceViewRHIRef RHICreateShaderResourceView(FRHIBuffer* Buffer) = 0;
 
 	/** Creates a shader resource view of the given vertex buffer. */
 	// FlushType: Wait RHI Thread
-	virtual FShaderResourceViewRHIRef RHICreateShaderResourceView(FRHIVertexBuffer* VertexBuffer, uint32 Stride, uint8 Format) = 0;
+	virtual FShaderResourceViewRHIRef RHICreateShaderResourceView(FRHIBuffer* Buffer, uint32 Stride, uint8 Format) = 0;
 
 	/** Creates a shader resource view **/
 	virtual FShaderResourceViewRHIRef RHICreateShaderResourceView(const FShaderResourceViewInitializer& Initializer) = 0;
 
-	/** Creates a shader resource view of the given index buffer. */
-	// FlushType: Wait RHI Thread
-	virtual FShaderResourceViewRHIRef RHICreateShaderResourceView(FRHIIndexBuffer* Buffer) = 0;
-
 	// Must be called on RHI thread timeline
 	// Make sure to call RHIThreadFence(true) afterwards so that parallel translation doesn't refer old resources
-	virtual void RHIUpdateShaderResourceView(FRHIShaderResourceView* SRV, FRHIVertexBuffer* VertexBuffer, uint32 Stride, uint8 Format);
+	virtual void RHIUpdateShaderResourceView(FRHIShaderResourceView* SRV, FRHIBuffer* Buffer, uint32 Stride, uint8 Format);
 
-	virtual void RHIUpdateShaderResourceView(FRHIShaderResourceView* SRV, FRHIIndexBuffer* IndexBuffer);
+	virtual void RHIUpdateShaderResourceView(FRHIShaderResourceView* SRV, FRHIBuffer* Buffer);
 
 	/**
 	* Computes the total size of a 2D texture with the specified parameters.
@@ -1329,16 +1293,14 @@ public:
 	virtual FTexture2DArrayRHIRef RHICreateTexture2DArray_RenderThread(class FRHICommandListImmediate& RHICmdList, uint32 SizeX, uint32 SizeY, uint32 SizeZ, uint8 Format, uint32 NumMips, uint32 NumSamples, ETextureCreateFlags Flags, ERHIAccess InResourceState, FRHIResourceCreateInfo& CreateInfo);
 	
 	virtual FTexture3DRHIRef RHICreateTexture3D_RenderThread(class FRHICommandListImmediate& RHICmdList, uint32 SizeX, uint32 SizeY, uint32 SizeZ, uint8 Format, uint32 NumMips, ETextureCreateFlags Flags, ERHIAccess InResourceState, FRHIResourceCreateInfo& CreateInfo);
-	virtual FUnorderedAccessViewRHIRef RHICreateUnorderedAccessView_RenderThread(class FRHICommandListImmediate& RHICmdList, FRHIStructuredBuffer* StructuredBuffer, bool bUseUAVCounter, bool bAppendBuffer);
+	virtual FUnorderedAccessViewRHIRef RHICreateUnorderedAccessView_RenderThread(class FRHICommandListImmediate& RHICmdList, FRHIBuffer* Buffer, bool bUseUAVCounter, bool bAppendBuffer);
 	virtual FUnorderedAccessViewRHIRef RHICreateUnorderedAccessView_RenderThread(class FRHICommandListImmediate& RHICmdList, FRHITexture* Texture, uint32 MipLevel);
 	virtual FUnorderedAccessViewRHIRef RHICreateUnorderedAccessView_RenderThread(class FRHICommandListImmediate& RHICmdList, FRHITexture* Texture, uint32 MipLevel, uint8 Format);
-	virtual FUnorderedAccessViewRHIRef RHICreateUnorderedAccessView_RenderThread(class FRHICommandListImmediate& RHICmdList, FRHIVertexBuffer* VertexBuffer, uint8 Format);
-	virtual FUnorderedAccessViewRHIRef RHICreateUnorderedAccessView_RenderThread(class FRHICommandListImmediate& RHICmdList, FRHIIndexBuffer* IndexBuffer, uint8 Format);
+	virtual FUnorderedAccessViewRHIRef RHICreateUnorderedAccessView_RenderThread(class FRHICommandListImmediate& RHICmdList, FRHIBuffer* Buffer, uint8 Format);
 	virtual FShaderResourceViewRHIRef RHICreateShaderResourceView_RenderThread(class FRHICommandListImmediate& RHICmdList, FRHITexture* Texture, const FRHITextureSRVCreateInfo& CreateInfo);
-	virtual FShaderResourceViewRHIRef RHICreateShaderResourceView_RenderThread(class FRHICommandListImmediate& RHICmdList, FRHIVertexBuffer* VertexBuffer, uint32 Stride, uint8 Format);
+	virtual FShaderResourceViewRHIRef RHICreateShaderResourceView_RenderThread(class FRHICommandListImmediate& RHICmdList, FRHIBuffer* Buffer, uint32 Stride, uint8 Format);
 	virtual FShaderResourceViewRHIRef RHICreateShaderResourceView_RenderThread(class FRHICommandListImmediate& RHICmdList, const FShaderResourceViewInitializer& Initializer);
-	virtual FShaderResourceViewRHIRef RHICreateShaderResourceView_RenderThread(class FRHICommandListImmediate& RHICmdList, FRHIIndexBuffer* Buffer);
-	virtual FShaderResourceViewRHIRef RHICreateShaderResourceView_RenderThread(class FRHICommandListImmediate& RHICmdList, FRHIStructuredBuffer* StructuredBuffer);
+	virtual FShaderResourceViewRHIRef RHICreateShaderResourceView_RenderThread(class FRHICommandListImmediate& RHICmdList, FRHIBuffer* Buffer);
 	virtual FShaderResourceViewRHIRef RHICreateShaderResourceViewWriteMask_RenderThread(class FRHICommandListImmediate& RHICmdList, FRHITexture2D* Texture2D);
 	virtual FShaderResourceViewRHIRef RHICreateShaderResourceViewFMask_RenderThread(class FRHICommandListImmediate& RHICmdList, FRHITexture2D* Texture2D);
 	virtual FTextureCubeRHIRef RHICreateTextureCube_RenderThread(class FRHICommandListImmediate& RHICmdList, uint32 Size, uint8 Format, uint32 NumMips, ETextureCreateFlags Flags, ERHIAccess InResourceState, FRHIResourceCreateInfo& CreateInfo);
@@ -1353,8 +1315,6 @@ public:
 	virtual void RHIDiscardTransientResource_RenderThread(FRHITexture* Texture) { }
 	virtual void RHIAcquireTransientResource_RenderThread(FRHIVertexBuffer* Buffer) { }
 	virtual void RHIDiscardTransientResource_RenderThread(FRHIVertexBuffer* Buffer) { }
-	virtual void RHIAcquireTransientResource_RenderThread(FRHIStructuredBuffer* Buffer) { }
-	virtual void RHIDiscardTransientResource_RenderThread(FRHIStructuredBuffer* Buffer) { }
 
 
 	virtual void RHIMapStagingSurface_RenderThread(class FRHICommandListImmediate& RHICmdList, FRHITexture* Texture, FRHIGPUFence* Fence, void*& OutData, int32& OutWidth, int32& OutHeight);
diff --git a/Engine/Source/Runtime/RHI/Public/RHI.h b/Engine/Source/Runtime/RHI/Public/RHI.h
--- a/Engine/Source/Runtime/RHI/Public/RHI.h
+++ b/Engine/Source/Runtime/RHI/Public/RHI.h
@@ -1854,9 +1854,7 @@ struct FRHITransitionInfo : public FRHISubresourceRange
 	{
 		class FRHIResource* Resource = nullptr;
 		class FRHITexture* Texture;
-		class FRHIVertexBuffer* VertexBuffer;
-		class FRHIIndexBuffer* IndexBuffer;
-		class FRHIStructuredBuffer* StructuredBuffer;
+		class FRHIBuffer* Buffer;
 		class FRHIUnorderedAccessView* UAV;
 	};
 
@@ -1864,9 +1862,7 @@ struct FRHITransitionInfo : public FRHISubresourceRange
 	{
 		Unknown,
 		Texture,
-		VertexBuffer,
-		IndexBuffer,
-		StructuredBuffer,
+		Buffer,
 		UAV
 	} Type = EType::Unknown;
 
diff --git a/Engine/Source/Runtime/RHI/Public/RHICommandList.h b/Engine/Source/Runtime/RHI/Public/RHICommandList.h
--- a/Engine/Source/Runtime/RHI/Public/RHICommandList.h
+++ b/Engine/Source/Runtime/RHI/Public/RHICommandList.h
@@ -992,14 +992,14 @@ FRHICOMMAND_MACRO(FRHICommandDrawPrimitive)
 
 FRHICOMMAND_MACRO(FRHICommandDrawIndexedPrimitive)
 {
-	FRHIIndexBuffer* IndexBuffer;
+	FRHIBuffer* IndexBuffer;
 	int32 BaseVertexIndex;
 	uint32 FirstInstance;
 	uint32 NumVertices;
 	uint32 StartIndex;
 	uint32 NumPrimitives;
 	uint32 NumInstances;
-	FORCEINLINE_DEBUGGABLE FRHICommandDrawIndexedPrimitive(FRHIIndexBuffer* InIndexBuffer, int32 InBaseVertexIndex, uint32 InFirstInstance, uint32 InNumVertices, uint32 InStartIndex, uint32 InNumPrimitives, uint32 InNumInstances)
+	FORCEINLINE_DEBUGGABLE FRHICommandDrawIndexedPrimitive(FRHIBuffer* InIndexBuffer, int32 InBaseVertexIndex, uint32 InFirstInstance, uint32 InNumVertices, uint32 InStartIndex, uint32 InNumPrimitives, uint32 InNumInstances)
 		: IndexBuffer(InIndexBuffer)
 		, BaseVertexIndex(InBaseVertexIndex)
 		, FirstInstance(InFirstInstance)
@@ -1268,9 +1268,9 @@ struct FRHICommandDispatchIndirectComputeShaderString
 
 struct FRHICommandDispatchIndirectComputeShader final : public FRHICommand<FRHICommandDispatchIndirectComputeShader, FRHICommandDispatchIndirectComputeShaderString>
 {
-	FRHIVertexBuffer* ArgumentBuffer;
+	FRHIBuffer* ArgumentBuffer;
 	uint32 ArgumentOffset;
-	FORCEINLINE_DEBUGGABLE FRHICommandDispatchIndirectComputeShader(FRHIVertexBuffer* InArgumentBuffer, uint32 InArgumentOffset)
+	FORCEINLINE_DEBUGGABLE FRHICommandDispatchIndirectComputeShader(FRHIBuffer* InArgumentBuffer, uint32 InArgumentOffset)
 		: ArgumentBuffer(InArgumentBuffer)
 		, ArgumentOffset(InArgumentOffset)
 	{
@@ -1304,9 +1304,9 @@ FRHICOMMAND_MACRO(FRHICommandEndSpecificUAVOverlap)
 
 FRHICOMMAND_MACRO(FRHICommandDrawPrimitiveIndirect)
 {
-	FRHIVertexBuffer* ArgumentBuffer;
+	FRHIBuffer* ArgumentBuffer;
 	uint32 ArgumentOffset;
-	FORCEINLINE_DEBUGGABLE FRHICommandDrawPrimitiveIndirect(FRHIVertexBuffer* InArgumentBuffer, uint32 InArgumentOffset)
+	FORCEINLINE_DEBUGGABLE FRHICommandDrawPrimitiveIndirect(FRHIBuffer* InArgumentBuffer, uint32 InArgumentOffset)
 		: ArgumentBuffer(InArgumentBuffer)
 		, ArgumentOffset(InArgumentOffset)
 	{
@@ -1316,12 +1316,12 @@ FRHICOMMAND_MACRO(FRHICommandDrawPrimitiveIndirect)
 
 FRHICOMMAND_MACRO(FRHICommandDrawIndexedIndirect)
 {
-	FRHIIndexBuffer* IndexBufferRHI;
-	FRHIStructuredBuffer* ArgumentsBufferRHI;
+	FRHIBuffer* IndexBufferRHI;
+	FRHIBuffer* ArgumentsBufferRHI;
 	uint32 DrawArgumentsIndex;
 	uint32 NumInstances;
 
-	FORCEINLINE_DEBUGGABLE FRHICommandDrawIndexedIndirect(FRHIIndexBuffer* InIndexBufferRHI, FRHIStructuredBuffer* InArgumentsBufferRHI, uint32 InDrawArgumentsIndex, uint32 InNumInstances)
+	FORCEINLINE_DEBUGGABLE FRHICommandDrawIndexedIndirect(FRHIBuffer* InIndexBufferRHI, FRHIBuffer* InArgumentsBufferRHI, uint32 InDrawArgumentsIndex, uint32 InNumInstances)
 		: IndexBufferRHI(InIndexBufferRHI)
 		, ArgumentsBufferRHI(InArgumentsBufferRHI)
 		, DrawArgumentsIndex(InDrawArgumentsIndex)
@@ -1333,11 +1333,11 @@ FRHICOMMAND_MACRO(FRHICommandDrawIndexedIndirect)
 
 FRHICOMMAND_MACRO(FRHICommandDrawIndexedPrimitiveIndirect)
 {
-	FRHIIndexBuffer* IndexBuffer;
-	FRHIVertexBuffer* ArgumentsBuffer;
+	FRHIBuffer* IndexBuffer;
+	FRHIBuffer* ArgumentsBuffer;
 	uint32 ArgumentOffset;
 
-	FORCEINLINE_DEBUGGABLE FRHICommandDrawIndexedPrimitiveIndirect(FRHIIndexBuffer* InIndexBuffer, FRHIVertexBuffer* InArgumentsBuffer, uint32 InArgumentOffset)
+	FORCEINLINE_DEBUGGABLE FRHICommandDrawIndexedPrimitiveIndirect(FRHIBuffer* InIndexBuffer, FRHIBuffer* InArgumentsBuffer, uint32 InArgumentOffset)
 		: IndexBuffer(InIndexBuffer)
 		, ArgumentsBuffer(InArgumentsBuffer)
 		, ArgumentOffset(InArgumentOffset)
@@ -1919,44 +1919,30 @@ FRHICOMMAND_MACRO(FRHICommandUpdateTextureReference)
 	RHI_API void Execute(FRHICommandListBase& CmdList);
 };
 
-struct FRHIShaderResourceViewUpdateInfo_VB
+struct FRHIShaderResourceViewUpdateInfo
 {
 	FRHIShaderResourceView* SRV;
-	FRHIVertexBuffer* VertexBuffer;
+	FRHIBuffer* Buffer;
 	uint32 Stride;
 	uint8 Format;
 };
 
-struct FRHIShaderResourceViewUpdateInfo_IB
-{
-	FRHIShaderResourceView* SRV;
-	FRHIIndexBuffer* IndexBuffer;
-};
-
-struct FRHIVertexBufferUpdateInfo
-{
-	FRHIVertexBuffer* DestBuffer;
-	FRHIVertexBuffer* SrcBuffer;
-};
-
-struct FRHIIndexBufferUpdateInfo
+struct FRHIBufferUpdateInfo
 {
-	FRHIIndexBuffer* DestBuffer;
-	FRHIIndexBuffer* SrcBuffer;
+	FRHIBuffer* DestBuffer;
+	FRHIBuffer* SrcBuffer;
 };
 
 struct FRHIResourceUpdateInfo
 {
 	enum EUpdateType
 	{
-		/** Take over underlying resource from an intermediate vertex buffer */
-		UT_VertexBuffer,
-		/** Take over underlying resource from an intermediate index buffer */
-		UT_IndexBuffer,
-		/** Update an SRV to view on a different vertex buffer */
-		UT_VertexBufferSRV,
-		/** Update an SRV to view on a different index buffer */
-		UT_IndexBufferSRV,
+		/** Take over underlying resource from an intermediate buffer */
+		UT_Buffer,
+		/** Update an SRV to view on a different buffer */
+		UT_BufferSRV,
+		/** Update an SRV to view on a different buffer with a given format */
+		UT_BufferFormatSRV,
 		/** Number of update types */
 		UT_Num
 	};
@@ -1964,10 +1950,8 @@ struct FRHIResourceUpdateInfo
 	EUpdateType Type;
 	union
 	{
-		FRHIVertexBufferUpdateInfo VertexBuffer;
-		FRHIIndexBufferUpdateInfo IndexBuffer;
-		FRHIShaderResourceViewUpdateInfo_VB VertexBufferSRV;
-		FRHIShaderResourceViewUpdateInfo_IB IndexBufferSRV;
+		FRHIBufferUpdateInfo Buffer;
+		FRHIShaderResourceViewUpdateInfo BufferSRV;
 	};
 
 	void ReleaseRefs();
@@ -2382,7 +2366,7 @@ public:
 		ALLOC_COMMAND(FRHICommandDispatchComputeShader)(ThreadGroupCountX, ThreadGroupCountY, ThreadGroupCountZ);
 	}
 
-	FORCEINLINE_DEBUGGABLE void DispatchIndirectComputeShader(FRHIVertexBuffer* ArgumentBuffer, uint32 ArgumentOffset)
+	FORCEINLINE_DEBUGGABLE void DispatchIndirectComputeShader(FRHIBuffer* ArgumentBuffer, uint32 ArgumentOffset)
 	{
 		if (Bypass())
 		{
@@ -3004,7 +2988,7 @@ public:
 		ALLOC_COMMAND(FRHICommandDrawPrimitive)(BaseVertexIndex, NumPrimitives, NumInstances);
 	}
 
-	FORCEINLINE_DEBUGGABLE void DrawIndexedPrimitive(FRHIIndexBuffer* IndexBuffer, int32 BaseVertexIndex, uint32 FirstInstance, uint32 NumVertices, uint32 StartIndex, uint32 NumPrimitives, uint32 NumInstances)
+	FORCEINLINE_DEBUGGABLE void DrawIndexedPrimitive(FRHIBuffer* IndexBuffer, int32 BaseVertexIndex, uint32 FirstInstance, uint32 NumVertices, uint32 StartIndex, uint32 NumPrimitives, uint32 NumInstances)
 	{
 		if (!IndexBuffer)
 		{
@@ -3142,7 +3126,7 @@ public:
 		ALLOC_COMMAND(FRHICommandSetGraphicsPipelineState)(GraphicsPipelineState, bApplyAdditionalState);
 	}
 
-	FORCEINLINE_DEBUGGABLE void DrawPrimitiveIndirect(FRHIVertexBuffer* ArgumentBuffer, uint32 ArgumentOffset)
+	FORCEINLINE_DEBUGGABLE void DrawPrimitiveIndirect(FRHIBuffer* ArgumentBuffer, uint32 ArgumentOffset)
 	{
 		//check(IsOutsideRenderPass());
 		if (Bypass())
@@ -3153,7 +3137,7 @@ public:
 		ALLOC_COMMAND(FRHICommandDrawPrimitiveIndirect)(ArgumentBuffer, ArgumentOffset);
 	}
 
-	FORCEINLINE_DEBUGGABLE void DrawIndexedIndirect(FRHIIndexBuffer* IndexBufferRHI, FRHIStructuredBuffer* ArgumentsBufferRHI, uint32 DrawArgumentsIndex, uint32 NumInstances)
+	FORCEINLINE_DEBUGGABLE void DrawIndexedIndirect(FRHIBuffer* IndexBufferRHI, FRHIBuffer* ArgumentsBufferRHI, uint32 DrawArgumentsIndex, uint32 NumInstances)
 	{
 		//check(IsOutsideRenderPass());
 		if (Bypass())
@@ -3164,7 +3148,7 @@ public:
 		ALLOC_COMMAND(FRHICommandDrawIndexedIndirect)(IndexBufferRHI, ArgumentsBufferRHI, DrawArgumentsIndex, NumInstances);
 	}
 
-	FORCEINLINE_DEBUGGABLE void DrawIndexedPrimitiveIndirect(FRHIIndexBuffer* IndexBuffer, FRHIVertexBuffer* ArgumentsBuffer, uint32 ArgumentOffset)
+	FORCEINLINE_DEBUGGABLE void DrawIndexedPrimitiveIndirect(FRHIBuffer* IndexBuffer, FRHIBuffer* ArgumentsBuffer, uint32 ArgumentOffset)
 	{
 		//check(IsOutsideRenderPass());
 		if (Bypass())
@@ -3956,11 +3940,11 @@ public:
 		GDynamicRHI->RHIUnlockStructuredBuffer(*this, StructuredBuffer);
 	}
 	
-	FORCEINLINE FUnorderedAccessViewRHIRef CreateUnorderedAccessView(FRHIStructuredBuffer* StructuredBuffer, bool bUseUAVCounter, bool bAppendBuffer)
+	FORCEINLINE FUnorderedAccessViewRHIRef CreateUnorderedAccessView(FRHIBuffer* Buffer, bool bUseUAVCounter, bool bAppendBuffer)
 	{
 		LLM_SCOPE(ELLMTag::RHIMisc);
-		checkf(StructuredBuffer, TEXT("Can't create a view off a null resource!"));
-		return GDynamicRHI->RHICreateUnorderedAccessView_RenderThread(*this, StructuredBuffer, bUseUAVCounter, bAppendBuffer);
+		checkf(Buffer, TEXT("Can't create a view off a null resource!"));
+		return GDynamicRHI->RHICreateUnorderedAccessView_RenderThread(*this, Buffer, bUseUAVCounter, bAppendBuffer);
 	}
 	
 	FORCEINLINE FUnorderedAccessViewRHIRef CreateUnorderedAccessView(FRHITexture* Texture, uint32 MipLevel)
@@ -3977,32 +3961,25 @@ public:
 		return GDynamicRHI->RHICreateUnorderedAccessView_RenderThread(*this, Texture, MipLevel, Format);
 	}
 
-	FORCEINLINE FUnorderedAccessViewRHIRef CreateUnorderedAccessView(FRHIVertexBuffer* VertexBuffer, uint8 Format)
-	{
-		LLM_SCOPE(ELLMTag::RHIMisc);
-		checkf(VertexBuffer, TEXT("Can't create a view off a null resource!"));
-		return GDynamicRHI->RHICreateUnorderedAccessView_RenderThread(*this, VertexBuffer, Format);
-	}
-
-	FORCEINLINE FUnorderedAccessViewRHIRef CreateUnorderedAccessView(FRHIIndexBuffer* IndexBuffer, uint8 Format)
+	FORCEINLINE FUnorderedAccessViewRHIRef CreateUnorderedAccessView(FRHIBuffer* Buffer, uint8 Format)
 	{
 		LLM_SCOPE(ELLMTag::RHIMisc);
-		checkf(IndexBuffer, TEXT("Can't create a view off a null resource!"));
-		return GDynamicRHI->RHICreateUnorderedAccessView_RenderThread(*this, IndexBuffer, Format);
+		checkf(Buffer, TEXT("Can't create a view off a null resource!"));
+		return GDynamicRHI->RHICreateUnorderedAccessView_RenderThread(*this, Buffer, Format);
 	}
 
-	FORCEINLINE FShaderResourceViewRHIRef CreateShaderResourceView(FRHIStructuredBuffer* StructuredBuffer)
+	FORCEINLINE FShaderResourceViewRHIRef CreateShaderResourceView(FRHIBuffer* Buffer)
 	{
 		LLM_SCOPE(ELLMTag::RHIMisc);
-		checkf(StructuredBuffer, TEXT("Can't create a view off a null resource!"));
-		return GDynamicRHI->RHICreateShaderResourceView_RenderThread(*this, StructuredBuffer);
+		checkf(Buffer, TEXT("Can't create a view off a null resource!"));
+		return GDynamicRHI->RHICreateShaderResourceView_RenderThread(*this, Buffer);
 	}
 	
-	FORCEINLINE FShaderResourceViewRHIRef CreateShaderResourceView(FRHIVertexBuffer* VertexBuffer, uint32 Stride, uint8 Format)
+	FORCEINLINE FShaderResourceViewRHIRef CreateShaderResourceView(FRHIBuffer* Buffer, uint32 Stride, uint8 Format)
 	{
 		LLM_SCOPE(ELLMTag::RHIMisc);
-		checkf(VertexBuffer, TEXT("Can't create a view off a null resource!"));
-		return GDynamicRHI->CreateShaderResourceView_RenderThread(*this, VertexBuffer, Stride, Format);
+		checkf(Buffer, TEXT("Can't create a view off a null resource!"));
+		return GDynamicRHI->CreateShaderResourceView_RenderThread(*this, Buffer, Stride, Format);
 	}
 	
 	FORCEINLINE FShaderResourceViewRHIRef CreateShaderResourceView(const FShaderResourceViewInitializer& Initializer)
@@ -4010,13 +3987,6 @@ public:
 		LLM_SCOPE(ELLMTag::RHIMisc);
 		return GDynamicRHI->CreateShaderResourceView_RenderThread(*this, Initializer);
 	}
-
-	FORCEINLINE FShaderResourceViewRHIRef CreateShaderResourceView(FRHIIndexBuffer* Buffer)
-	{
-		LLM_SCOPE(ELLMTag::RHIMisc);
-		checkf(Buffer, TEXT("Can't create a view off a null resource!"));
-		return GDynamicRHI->CreateShaderResourceView_RenderThread(*this, Buffer);
-	}
 	
 	FORCEINLINE uint64 CalcTexture2DPlatformSize(uint32 SizeX, uint32 SizeY, uint8 Format, uint32 NumMips, uint32 NumSamples, ETextureCreateFlags Flags, const FRHIResourceCreateInfo& CreateInfo, uint32& OutAlign)
 	{
@@ -4335,31 +4305,7 @@ public:
 		}
 	}
 
-	FORCEINLINE void AcquireTransientResource_RenderThread(FRHIVertexBuffer* Buffer)
-	{
-		if (!Buffer->IsCommitted())
-		{
-			if (GSupportsTransientResourceAliasing)
-			{
-				GDynamicRHI->RHIAcquireTransientResource_RenderThread(Buffer);
-			}
-			Buffer->SetCommitted(true);
-		}
-	}
-
-	FORCEINLINE void DiscardTransientResource_RenderThread(FRHIVertexBuffer* Buffer)
-	{
-		if (Buffer->IsCommitted())
-		{
-			if (GSupportsTransientResourceAliasing)
-			{
-				GDynamicRHI->RHIDiscardTransientResource_RenderThread(Buffer);
-			}
-			Buffer->SetCommitted(false);
-		}
-	}
-
-	FORCEINLINE void AcquireTransientResource_RenderThread(FRHIStructuredBuffer* Buffer)
+	FORCEINLINE void AcquireTransientResource_RenderThread(FRHIBuffer* Buffer)
 	{
 		if (!Buffer->IsCommitted())
 		{
@@ -4371,7 +4317,7 @@ public:
 		}
 	}
 
-	FORCEINLINE void DiscardTransientResource_RenderThread(FRHIStructuredBuffer* Buffer)
+	FORCEINLINE void DiscardTransientResource_RenderThread(FRHIBuffer* Buffer)
 	{
 		if (Buffer->IsCommitted())
 		{
@@ -5017,9 +4963,9 @@ FORCEINLINE void RHIUnlockStructuredBuffer(FRHIStructuredBuffer* StructuredBuffe
 	 FRHICommandListExecutor::GetImmediateCommandList().UnlockStructuredBuffer(StructuredBuffer);
 }
 
-FORCEINLINE FUnorderedAccessViewRHIRef RHICreateUnorderedAccessView(FRHIStructuredBuffer* StructuredBuffer, bool bUseUAVCounter, bool bAppendBuffer)
+FORCEINLINE FUnorderedAccessViewRHIRef RHICreateUnorderedAccessView(FRHIBuffer* Buffer, bool bUseUAVCounter, bool bAppendBuffer)
 {
-	return FRHICommandListExecutor::GetImmediateCommandList().CreateUnorderedAccessView(StructuredBuffer, bUseUAVCounter, bAppendBuffer);
+	return FRHICommandListExecutor::GetImmediateCommandList().CreateUnorderedAccessView(Buffer, bUseUAVCounter, bAppendBuffer);
 }
 
 FORCEINLINE FUnorderedAccessViewRHIRef RHICreateUnorderedAccessView(FRHITexture* Texture, uint32 MipLevel = 0)
@@ -5032,24 +4978,19 @@ FORCEINLINE FUnorderedAccessViewRHIRef RHICreateUnorderedAccessView(FRHITexture*
 	return FRHICommandListExecutor::GetImmediateCommandList().CreateUnorderedAccessView(Texture, MipLevel, Format);
 }
 
-FORCEINLINE FUnorderedAccessViewRHIRef RHICreateUnorderedAccessView(FRHIVertexBuffer* VertexBuffer, uint8 Format)
-{
-	return FRHICommandListExecutor::GetImmediateCommandList().CreateUnorderedAccessView(VertexBuffer, Format);
-}
-
-FORCEINLINE FUnorderedAccessViewRHIRef RHICreateUnorderedAccessView(FRHIIndexBuffer* IndexBuffer, uint8 Format)
+FORCEINLINE FUnorderedAccessViewRHIRef RHICreateUnorderedAccessView(FRHIBuffer* Buffer, uint8 Format)
 {
-	return FRHICommandListExecutor::GetImmediateCommandList().CreateUnorderedAccessView(IndexBuffer, Format);
+	return FRHICommandListExecutor::GetImmediateCommandList().CreateUnorderedAccessView(Buffer, Format);
 }
 
-FORCEINLINE FShaderResourceViewRHIRef RHICreateShaderResourceView(FRHIStructuredBuffer* StructuredBuffer)
+FORCEINLINE FShaderResourceViewRHIRef RHICreateShaderResourceView(FRHIBuffer* Buffer)
 {
-	return FRHICommandListExecutor::GetImmediateCommandList().CreateShaderResourceView(StructuredBuffer);
+	return FRHICommandListExecutor::GetImmediateCommandList().CreateShaderResourceView(Buffer);
 }
 
-FORCEINLINE FShaderResourceViewRHIRef RHICreateShaderResourceView(FRHIVertexBuffer* VertexBuffer, uint32 Stride, uint8 Format)
+FORCEINLINE FShaderResourceViewRHIRef RHICreateShaderResourceView(FRHIBuffer* Buffer, uint32 Stride, uint8 Format)
 {
-	return FRHICommandListExecutor::GetImmediateCommandList().CreateShaderResourceView(VertexBuffer, Stride, Format);
+	return FRHICommandListExecutor::GetImmediateCommandList().CreateShaderResourceView(Buffer, Stride, Format);
 }
 
 FORCEINLINE FShaderResourceViewRHIRef RHICreateShaderResourceView(const FShaderResourceViewInitializer& Initializer)
@@ -5057,11 +4998,6 @@ FORCEINLINE FShaderResourceViewRHIRef RHICreateShaderResourceView(const FShaderR
 	return FRHICommandListExecutor::GetImmediateCommandList().CreateShaderResourceView(Initializer);
 }
 
-FORCEINLINE FShaderResourceViewRHIRef RHICreateShaderResourceView(FRHIIndexBuffer* Buffer)
-{
-	return FRHICommandListExecutor::GetImmediateCommandList().CreateShaderResourceView(Buffer);
-}
-
 FORCEINLINE void RHIUpdateRHIResources(FRHIResourceUpdateInfo* UpdateInfos, int32 Num, bool bNeedReleaseRefs)
 {
 	return FRHICommandListExecutor::GetImmediateCommandList().UpdateRHIResources(UpdateInfos, Num, bNeedReleaseRefs);
@@ -5296,22 +5232,12 @@ FORCEINLINE void RHIDiscardTransientResource(FRHITexture* Resource)
 	FRHICommandListExecutor::GetImmediateCommandList().DiscardTransientResource_RenderThread(Resource);
 }
 
-FORCEINLINE void RHIAcquireTransientResource(FRHIVertexBuffer* Resource)
-{
-	FRHICommandListExecutor::GetImmediateCommandList().AcquireTransientResource_RenderThread(Resource);
-}
-
-FORCEINLINE void RHIDiscardTransientResource(FRHIVertexBuffer* Resource)
-{
-	FRHICommandListExecutor::GetImmediateCommandList().DiscardTransientResource_RenderThread(Resource);
-}
-
-FORCEINLINE void RHIAcquireTransientResource(FRHIStructuredBuffer* Resource)
+FORCEINLINE void RHIAcquireTransientResource(FRHIBuffer* Resource)
 {
 	FRHICommandListExecutor::GetImmediateCommandList().AcquireTransientResource_RenderThread(Resource);
 }
 
-FORCEINLINE void RHIDiscardTransientResource(FRHIStructuredBuffer* Resource)
+FORCEINLINE void RHIDiscardTransientResource(FRHIBuffer* Resource)
 {
 	FRHICommandListExecutor::GetImmediateCommandList().DiscardTransientResource_RenderThread(Resource);
 }
@@ -5401,51 +5327,39 @@ struct TRHIResourceUpdateBatcher
 		}
 	}
 
-	void QueueUpdateRequest(FRHIVertexBuffer* DestVertexBuffer, FRHIVertexBuffer* SrcVertexBuffer)
-	{
-		FRHIResourceUpdateInfo& UpdateInfo = GetNextUpdateInfo();
-		UpdateInfo.Type = FRHIResourceUpdateInfo::UT_VertexBuffer;
-		UpdateInfo.VertexBuffer = { DestVertexBuffer, SrcVertexBuffer };
-		DestVertexBuffer->AddRef();
-		if (SrcVertexBuffer)
-		{
-			SrcVertexBuffer->AddRef();
-		}
-	}
-
-	void QueueUpdateRequest(FRHIIndexBuffer* DestIndexBuffer, FRHIIndexBuffer* SrcIndexBuffer)
+	void QueueUpdateRequest(FRHIBuffer* DestBuffer, FRHIBuffer* SrcBuffer)
 	{
 		FRHIResourceUpdateInfo& UpdateInfo = GetNextUpdateInfo();
-		UpdateInfo.Type = FRHIResourceUpdateInfo::UT_IndexBuffer;
-		UpdateInfo.IndexBuffer = { DestIndexBuffer, SrcIndexBuffer };
-		DestIndexBuffer->AddRef();
-		if (SrcIndexBuffer)
+		UpdateInfo.Type = FRHIResourceUpdateInfo::UT_Buffer;
+		UpdateInfo.Buffer = { DestBuffer, SrcBuffer };
+		DestBuffer->AddRef();
+		if (SrcBuffer)
 		{
-			SrcIndexBuffer->AddRef();
+			SrcBuffer->AddRef();
 		}
 	}
 
-	void QueueUpdateRequest(FRHIShaderResourceView* SRV, FRHIVertexBuffer* VertexBuffer, uint32 Stride, uint8 Format)
+	void QueueUpdateRequest(FRHIShaderResourceView* SRV, FRHIBuffer* Buffer, uint32 Stride, uint8 Format)
 	{
 		FRHIResourceUpdateInfo& UpdateInfo = GetNextUpdateInfo();
-		UpdateInfo.Type = FRHIResourceUpdateInfo::UT_VertexBufferSRV;
-		UpdateInfo.VertexBufferSRV = { SRV, VertexBuffer, Stride, Format };
+		UpdateInfo.Type = FRHIResourceUpdateInfo::UT_BufferFormatSRV;
+		UpdateInfo.BufferSRV = { SRV, Buffer, Stride, Format };
 		SRV->AddRef();
-		if (VertexBuffer)
+		if (Buffer)
 		{
-			VertexBuffer->AddRef();
+			Buffer->AddRef();
 		}
 	}
 
-	void QueueUpdateRequest(FRHIShaderResourceView* SRV, FRHIIndexBuffer* IndexBuffer)
+	void QueueUpdateRequest(FRHIShaderResourceView* SRV, FRHIBuffer* Buffer)
 	{
 		FRHIResourceUpdateInfo& UpdateInfo = GetNextUpdateInfo();
-		UpdateInfo.Type = FRHIResourceUpdateInfo::UT_IndexBufferSRV;
-		UpdateInfo.IndexBufferSRV = { SRV, IndexBuffer };
+		UpdateInfo.Type = FRHIResourceUpdateInfo::UT_BufferSRV;
+		UpdateInfo.BufferSRV = { SRV, Buffer };
 		SRV->AddRef();
-		if (IndexBuffer)
+		if (Buffer)
 		{
-			IndexBuffer->AddRef();
+			Buffer->AddRef();
 		}
 	}
 
diff --git a/Engine/Source/Runtime/RHI/Public/RHICommandListCommandExecutes.inl b/Engine/Source/Runtime/RHI/Public/RHICommandListCommandExecutes.inl
--- a/Engine/Source/Runtime/RHI/Public/RHICommandListCommandExecutes.inl
+++ b/Engine/Source/Runtime/RHI/Public/RHICommandListCommandExecutes.inl
@@ -585,32 +585,18 @@ void FRHIResourceUpdateInfo::ReleaseRefs()
 {
 	switch (Type)
 	{
-	case UT_VertexBuffer:
-		VertexBuffer.DestBuffer->Release();
-		if (VertexBuffer.SrcBuffer)
+	case UT_Buffer:
+		Buffer.DestBuffer->Release();
+		if (Buffer.SrcBuffer)
 		{
-			VertexBuffer.SrcBuffer->Release();
+			Buffer.SrcBuffer->Release();
 		}
 		break;
-	case UT_IndexBuffer:
-		IndexBuffer.DestBuffer->Release();
-		if (IndexBuffer.SrcBuffer)
+	case UT_BufferSRV:
+		BufferSRV.SRV->Release();
+		if (BufferSRV.Buffer)
 		{
-			IndexBuffer.SrcBuffer->Release();
-		}
-		break;
-	case UT_VertexBufferSRV:
-		VertexBufferSRV.SRV->Release();
-		if (VertexBufferSRV.VertexBuffer)
-		{
-			VertexBufferSRV.VertexBuffer->Release();
-		}
-		break;
-	case UT_IndexBufferSRV:
-		IndexBufferSRV.SRV->Release();
-		if (IndexBufferSRV.IndexBuffer)
-		{
-			IndexBufferSRV.IndexBuffer->Release();
+			BufferSRV.Buffer->Release();
 		}
 		break;
 	default:
@@ -638,27 +624,22 @@ void FRHICommandUpdateRHIResources::Execute(FRHICommandListBase& CmdList)
 		FRHIResourceUpdateInfo& Info = UpdateInfos[Idx];
 		switch (Info.Type)
 		{
-		case FRHIResourceUpdateInfo::UT_VertexBuffer:
-			GDynamicRHI->RHITransferVertexBufferUnderlyingResource(
-				Info.VertexBuffer.DestBuffer,
-				Info.VertexBuffer.SrcBuffer);
-			break;
-		case FRHIResourceUpdateInfo::UT_IndexBuffer:
-			GDynamicRHI->RHITransferIndexBufferUnderlyingResource(
-				Info.IndexBuffer.DestBuffer,
-				Info.IndexBuffer.SrcBuffer);
+		case FRHIResourceUpdateInfo::UT_Buffer:
+			GDynamicRHI->RHITransferBufferUnderlyingResource(
+				Info.Buffer.DestBuffer,
+				Info.Buffer.SrcBuffer);
 			break;
-		case FRHIResourceUpdateInfo::UT_VertexBufferSRV:
+		case FRHIResourceUpdateInfo::UT_BufferFormatSRV:
 			GDynamicRHI->RHIUpdateShaderResourceView(
-				Info.VertexBufferSRV.SRV,
-				Info.VertexBufferSRV.VertexBuffer,
-				Info.VertexBufferSRV.Stride,
-				Info.VertexBufferSRV.Format);
+				Info.BufferSRV.SRV,
+				Info.BufferSRV.Buffer,
+				Info.BufferSRV.Stride,
+				Info.BufferSRV.Format);
 			break;
-		case FRHIResourceUpdateInfo::UT_IndexBufferSRV:
+		case FRHIResourceUpdateInfo::UT_BufferSRV:
 			GDynamicRHI->RHIUpdateShaderResourceView(
-				Info.IndexBufferSRV.SRV,
-				Info.IndexBufferSRV.IndexBuffer);
+				Info.BufferSRV.SRV,
+				Info.BufferSRV.Buffer);
 			break;
 		default:
 			// Unrecognized type, do nothing
diff --git a/Engine/Source/Runtime/RHI/Public/RHIContext.h b/Engine/Source/Runtime/RHI/Public/RHIContext.h
--- a/Engine/Source/Runtime/RHI/Public/RHIContext.h
+++ b/Engine/Source/Runtime/RHI/Public/RHIContext.h
@@ -116,7 +116,7 @@ public:
 
 	virtual void RHIDispatchComputeShader(uint32 ThreadGroupCountX, uint32 ThreadGroupCountY, uint32 ThreadGroupCountZ) = 0;
 
-	virtual void RHIDispatchIndirectComputeShader(FRHIVertexBuffer* ArgumentBuffer, uint32 ArgumentOffset) = 0;
+	virtual void RHIDispatchIndirectComputeShader(FRHIBuffer* ArgumentBuffer, uint32 ArgumentOffset) = 0;
 
 	virtual void RHISetAsyncComputeBudget(EAsyncComputeBudget Budget) {}
 
@@ -332,7 +332,7 @@ public:
 
 	virtual void RHIDispatchComputeShader(uint32 ThreadGroupCountX, uint32 ThreadGroupCountY, uint32 ThreadGroupCountZ) = 0;
 
-	virtual void RHIDispatchIndirectComputeShader(FRHIVertexBuffer* ArgumentBuffer, uint32 ArgumentOffset) = 0;
+	virtual void RHIDispatchIndirectComputeShader(FRHIBuffer* ArgumentBuffer, uint32 ArgumentOffset) = 0;
 
 	// Useful when used with geometry shader (emit polygons to different viewports), otherwise SetViewPort() is simpler
 	// @param Count >0
@@ -516,14 +516,14 @@ public:
 
 	virtual void RHIDrawPrimitive(uint32 BaseVertexIndex, uint32 NumPrimitives, uint32 NumInstances) = 0;
 
-	virtual void RHIDrawPrimitiveIndirect(FRHIVertexBuffer* ArgumentBuffer, uint32 ArgumentOffset) = 0;
+	virtual void RHIDrawPrimitiveIndirect(FRHIBuffer* ArgumentBuffer, uint32 ArgumentOffset) = 0;
 
-	virtual void RHIDrawIndexedIndirect(FRHIIndexBuffer* IndexBufferRHI, FRHIStructuredBuffer* ArgumentsBufferRHI, int32 DrawArgumentsIndex, uint32 NumInstances) = 0;
+	virtual void RHIDrawIndexedIndirect(FRHIBuffer* IndexBufferRHI, FRHIBuffer* ArgumentsBufferRHI, int32 DrawArgumentsIndex, uint32 NumInstances) = 0;
 
 	// @param NumPrimitives need to be >0 
-	virtual void RHIDrawIndexedPrimitive(FRHIIndexBuffer* IndexBuffer, int32 BaseVertexIndex, uint32 FirstInstance, uint32 NumVertices, uint32 StartIndex, uint32 NumPrimitives, uint32 NumInstances) = 0;
+	virtual void RHIDrawIndexedPrimitive(FRHIBuffer* IndexBuffer, int32 BaseVertexIndex, uint32 FirstInstance, uint32 NumVertices, uint32 StartIndex, uint32 NumPrimitives, uint32 NumInstances) = 0;
 
-	virtual void RHIDrawIndexedPrimitiveIndirect(FRHIIndexBuffer* IndexBuffer, FRHIVertexBuffer* ArgumentBuffer, uint32 ArgumentOffset) = 0;
+	virtual void RHIDrawIndexedPrimitiveIndirect(FRHIBuffer* IndexBuffer, FRHIBuffer* ArgumentBuffer, uint32 ArgumentOffset) = 0;
 
 	/**
 	* Sets Depth Bounds range with the given min/max depth.
diff --git a/Engine/Source/Runtime/RHI/Public/RHIResources.h b/Engine/Source/Runtime/RHI/Public/RHIResources.h
--- a/Engine/Source/Runtime/RHI/Public/RHIResources.h
+++ b/Engine/Source/Runtime/RHI/Public/RHIResources.h
@@ -818,129 +818,64 @@ private:
 	uint32 LayoutConstantBufferSize;
 };
 
-class FRHIIndexBuffer : public FRHIResource
+class FRHIBuffer : public FRHIResource
 #if ENABLE_RHI_VALIDATION
 	, public RHIValidation::FBufferResource
 #endif
 {
 public:
-
 	/** Initialization constructor. */
-	FRHIIndexBuffer(uint32 InStride,uint32 InSize,uint32 InUsage)
-	: Stride(InStride)
-	, Size(InSize)
-	, Usage(InUsage)
-	{}
+	FRHIBuffer(uint32 InSize, uint32 InUsage, uint32 InStride)
+		: Size(InSize)
+		, Usage(InUsage)
+		, Stride(InStride)
+	{
+	}
 
-	/** @return The stride in bytes of the index buffer; must be 2 or 4. */
+	/** @return The stride in bytes of the buffer. */
 	uint32 GetStride() const { return Stride; }
 
-	/** @return The number of bytes in the index buffer. */
+	/** @return The number of bytes in the buffer. */
 	uint32 GetSize() const { return Size; }
 
-	/** @return The usage flags used to create the index buffer. */
+	/** @return The usage flags used to create the buffer. */
 	uint32 GetUsage() const { return Usage; }
 
 protected:
-	FRHIIndexBuffer()
-		: Stride(0)
-		, Size(0)
+	FRHIBuffer()
+		: Size(0)
 		, Usage(0)
+		, Stride(0)
 	{}
 
-	void Swap(FRHIIndexBuffer& Other)
+	void Swap(FRHIBuffer& Other)
 	{
 		::Swap(Stride, Other.Stride);
 		::Swap(Size, Other.Size);
 		::Swap(Usage, Other.Usage);
 	}
 
-	void ReleaseUnderlyingResource()
+	// Used by RHI implementations that may adjust internal usage flags during object construction.
+	void SetUsage(uint32 InUsage)
 	{
-		Stride = Size = Usage = 0;
-	}
-
-private:
-	uint32 Stride;
-	uint32 Size;
-	uint32 Usage;
-};
-
-class FRHIVertexBuffer : public FRHIResource
-#if ENABLE_RHI_VALIDATION
-	, public RHIValidation::FBufferResource
-#endif
-{
-public:
-
-	/**
-	 * Initialization constructor.
-	 * @apram InUsage e.g. BUF_UnorderedAccess
-	 */
-	FRHIVertexBuffer(uint32 InSize,uint32 InUsage)
-	: Size(InSize)
-	, Usage(InUsage)
-	{}
-
-	/** @return The number of bytes in the vertex buffer. */
-	uint32 GetSize() const { return Size; }
-
-	/** @return The usage flags used to create the vertex buffer. e.g. BUF_UnorderedAccess */
-	uint32 GetUsage() const { return Usage; }
-
-protected:
-	FRHIVertexBuffer()
-		: Size(0)
-		, Usage(0)
-	{}
-
-	void Swap(FRHIVertexBuffer& Other)
-	{
-		::Swap(Size, Other.Size);
-		::Swap(Usage, Other.Usage);
+		Usage = InUsage;
 	}
 
 	void ReleaseUnderlyingResource()
 	{
-		Size = 0;
-		Usage = 0;
+		Stride = Size = Usage = 0;
 	}
 
 private:
 	uint32 Size;
-	// e.g. BUF_UnorderedAccess
 	uint32 Usage;
-};
-
-class FRHIStructuredBuffer : public FRHIResource
-#if ENABLE_RHI_VALIDATION
-	, public RHIValidation::FBufferResource
-#endif
-{
-public:
-
-	/** Initialization constructor. */
-	FRHIStructuredBuffer(uint32 InStride,uint32 InSize,uint32 InUsage)
-	: Stride(InStride)
-	, Size(InSize)
-	, Usage(InUsage)
-	{}
-
-	/** @return The stride in bytes of the structured buffer; must be 2 or 4. */
-	uint32 GetStride() const { return Stride; }
-
-	/** @return The number of bytes in the structured buffer. */
-	uint32 GetSize() const { return Size; }
-
-	/** @return The usage flags used to create the structured buffer. */
-	uint32 GetUsage() const { return Usage; }
-
-private:
 	uint32 Stride;
-	uint32 Size;
-	uint32 Usage;
 };
 
+typedef class FRHIBuffer FRHIIndexBuffer;
+typedef class FRHIBuffer FRHIVertexBuffer;
+typedef class FRHIBuffer FRHIStructuredBuffer;
+
 //
 // Textures
 //
diff --git a/Engine/Source/Runtime/RHI/Public/RHIValidation.h b/Engine/Source/Runtime/RHI/Public/RHIValidation.h
--- a/Engine/Source/Runtime/RHI/Public/RHIValidation.h
+++ b/Engine/Source/Runtime/RHI/Public/RHIValidation.h
@@ -453,42 +453,32 @@ public:
 		return UAV;
 	}
 
-	/** Creates an unordered access view of the given vertex buffer. */
+	/** Creates an unordered access view of the given buffer. */
 	// FlushType: Wait RHI Thread
-	virtual FUnorderedAccessViewRHIRef RHICreateUnorderedAccessView(FRHIVertexBuffer* VertexBuffer, uint8 Format) override final
+	virtual FUnorderedAccessViewRHIRef RHICreateUnorderedAccessView(FRHIBuffer* Buffer, uint8 Format) override final
 	{
 		check(IsFeatureLevelSupported(GMaxRHIShaderPlatform, ERHIFeatureLevel::SM5));
-		FUnorderedAccessViewRHIRef UAV = RHI->RHICreateUnorderedAccessView(VertexBuffer, Format);
-		UAV->ViewIdentity = VertexBuffer->GetWholeResourceIdentity();
+		FUnorderedAccessViewRHIRef UAV = RHI->RHICreateUnorderedAccessView(Buffer, Format);
+		UAV->ViewIdentity = Buffer->GetWholeResourceIdentity();
 		return UAV;
 	}
 
-	/** Creates an unordered access view of the given index buffer. */
-	// FlushType: Wait RHI Thread
-	virtual FUnorderedAccessViewRHIRef RHICreateUnorderedAccessView(FRHIIndexBuffer* IndexBuffer, uint8 Format) override final
-	{
-		check(IsFeatureLevelSupported(GMaxRHIShaderPlatform, ERHIFeatureLevel::SM5));
-		FUnorderedAccessViewRHIRef UAV = RHI->RHICreateUnorderedAccessView(IndexBuffer, Format);
-		UAV->ViewIdentity = IndexBuffer->GetWholeResourceIdentity();
-		return UAV;
-	}
-
-	/** Creates a shader resource view of the given structured buffer. */
+	/** Creates a shader resource view of the given buffer. */
 	// FlushType: Wait RHI Thread
-	virtual FShaderResourceViewRHIRef RHICreateShaderResourceView(FRHIStructuredBuffer* StructuredBuffer) override final
+	virtual FShaderResourceViewRHIRef RHICreateShaderResourceView(FRHIBuffer* Buffer) override final
 	{
 		check(RHISupportsComputeShaders(GMaxRHIShaderPlatform));
-		FShaderResourceViewRHIRef SRV = RHI->RHICreateShaderResourceView(StructuredBuffer);
-		SRV->ViewIdentity = StructuredBuffer->GetWholeResourceIdentity();
+		FShaderResourceViewRHIRef SRV = RHI->RHICreateShaderResourceView(Buffer);
+		SRV->ViewIdentity = Buffer->GetWholeResourceIdentity();
 		return SRV;
 	}
 
-	/** Creates a shader resource view of the given vertex buffer. */
+	/** Creates a shader resource view of the given buffer. */
 	// FlushType: Wait RHI Thread
-	virtual FShaderResourceViewRHIRef RHICreateShaderResourceView(FRHIVertexBuffer* VertexBuffer, uint32 Stride, uint8 Format) override final
+	virtual FShaderResourceViewRHIRef RHICreateShaderResourceView(FRHIBuffer* Buffer, uint32 Stride, uint8 Format) override final
 	{
-		FShaderResourceViewRHIRef SRV = RHI->RHICreateShaderResourceView(VertexBuffer, Stride, Format);
-		SRV->ViewIdentity = VertexBuffer->GetWholeResourceIdentity();
+		FShaderResourceViewRHIRef SRV = RHI->RHICreateShaderResourceView(Buffer, Stride, Format);
+		SRV->ViewIdentity = Buffer->GetWholeResourceIdentity();
 		return SRV;
 	}
 
@@ -500,15 +490,6 @@ public:
 		return RHI->RHICreateShaderResourceView(Initializer);
 	}
 
-	/** Creates a shader resource view of the given index buffer. */
-	// FlushType: Wait RHI Thread
-	virtual FShaderResourceViewRHIRef RHICreateShaderResourceView(FRHIIndexBuffer* Buffer) override final
-	{
-		FShaderResourceViewRHIRef SRV = RHI->RHICreateShaderResourceView(Buffer);
-		SRV->ViewIdentity = Buffer->GetWholeResourceIdentity();
-		return SRV;
-	}
-
 	// Must be called on RHI thread timeline
 	// Make sure to call RHIThreadFence(true) afterwards so that parallel translation doesn't refer old resources
 	virtual void RHIUpdateShaderResourceView(FRHIShaderResourceView* SRV, FRHIVertexBuffer* VertexBuffer, uint32 Stride, uint8 Format) override final
@@ -1409,22 +1390,7 @@ public:
 	virtual FShaderResourceViewRHIRef CreateShaderResourceView_RenderThread(class FRHICommandListImmediate& RHICmdList, const FShaderResourceViewInitializer& Initializer) override final
 	{
 		FShaderResourceViewRHIRef SRV = RHI->CreateShaderResourceView_RenderThread(RHICmdList, Initializer);
-
-		switch (Initializer.GetType())
-		{
-		default: checkNoEntry(); // fallthrough
-
-		case FShaderResourceViewInitializer::EType::IndexBufferSRV:
-			SRV->ViewIdentity = Initializer.AsIndexBufferSRV().IndexBuffer->GetWholeResourceIdentity();
-			break;
-
-		case FShaderResourceViewInitializer::EType::StructuredBufferSRV:
-			SRV->ViewIdentity = Initializer.AsStructuredBufferSRV().StructuredBuffer->GetWholeResourceIdentity();
-			break;
-
-		case FShaderResourceViewInitializer::EType::VertexBufferSRV:
-			SRV->ViewIdentity = Initializer.AsVertexBufferSRV().VertexBuffer->GetWholeResourceIdentity();
-		}
+		SRV->ViewIdentity = Initializer.AsBufferSRV().Buffer->GetWholeResourceIdentity();
 
 		return SRV;
 	}
@@ -1596,11 +1562,11 @@ public:
 		return Texture3D;
 	}
 
-	virtual FUnorderedAccessViewRHIRef RHICreateUnorderedAccessView_RenderThread(class FRHICommandListImmediate& RHICmdList, FRHIStructuredBuffer* StructuredBuffer, bool bUseUAVCounter, bool bAppendBuffer) override final
+	virtual FUnorderedAccessViewRHIRef RHICreateUnorderedAccessView_RenderThread(class FRHICommandListImmediate& RHICmdList, FRHIBuffer* Buffer, bool bUseUAVCounter, bool bAppendBuffer) override final
 	{
 		check(IsFeatureLevelSupported(GMaxRHIShaderPlatform, ERHIFeatureLevel::SM5));
-		FUnorderedAccessViewRHIRef UAV = RHI->RHICreateUnorderedAccessView_RenderThread(RHICmdList, StructuredBuffer, bUseUAVCounter, bAppendBuffer);
-		UAV->ViewIdentity = StructuredBuffer->GetWholeResourceIdentity();
+		FUnorderedAccessViewRHIRef UAV = RHI->RHICreateUnorderedAccessView_RenderThread(RHICmdList, Buffer, bUseUAVCounter, bAppendBuffer);
+		UAV->ViewIdentity = Buffer->GetWholeResourceIdentity();
 		return UAV;
 	}
 
@@ -1612,19 +1578,11 @@ public:
 		return UAV;
 	}
 
-	virtual FUnorderedAccessViewRHIRef RHICreateUnorderedAccessView_RenderThread(class FRHICommandListImmediate& RHICmdList, FRHIVertexBuffer* VertexBuffer, uint8 Format) override final
-	{
-		check(IsFeatureLevelSupported(GMaxRHIShaderPlatform, ERHIFeatureLevel::SM5));
-		FUnorderedAccessViewRHIRef UAV = RHI->RHICreateUnorderedAccessView_RenderThread(RHICmdList, VertexBuffer, Format);
-		UAV->ViewIdentity = VertexBuffer->GetWholeResourceIdentity();
-		return UAV;
-	}
-
-	virtual FUnorderedAccessViewRHIRef RHICreateUnorderedAccessView_RenderThread(class FRHICommandListImmediate& RHICmdList, FRHIIndexBuffer* IndexBuffer, uint8 Format) override final
+	virtual FUnorderedAccessViewRHIRef RHICreateUnorderedAccessView_RenderThread(class FRHICommandListImmediate& RHICmdList, FRHIBuffer* Buffer, uint8 Format) override final
 	{
 		check(IsFeatureLevelSupported(GMaxRHIShaderPlatform, ERHIFeatureLevel::SM5));
-		FUnorderedAccessViewRHIRef UAV = RHI->RHICreateUnorderedAccessView_RenderThread(RHICmdList, IndexBuffer, Format);
-		UAV->ViewIdentity = IndexBuffer->GetWholeResourceIdentity();
+		FUnorderedAccessViewRHIRef UAV = RHI->RHICreateUnorderedAccessView_RenderThread(RHICmdList, Buffer, Format);
+		UAV->ViewIdentity = Buffer->GetWholeResourceIdentity();
 		return UAV;
 	}
 
@@ -1635,28 +1593,20 @@ public:
 		return SRV;
 	}
 
-	virtual FShaderResourceViewRHIRef RHICreateShaderResourceView_RenderThread(class FRHICommandListImmediate& RHICmdList, FRHIVertexBuffer* VertexBuffer, uint32 Stride, uint8 Format) override final
+	virtual FShaderResourceViewRHIRef RHICreateShaderResourceView_RenderThread(class FRHICommandListImmediate& RHICmdList, FRHIBuffer* Buffer, uint32 Stride, uint8 Format) override final
 	{
-		FShaderResourceViewRHIRef SRV = RHI->RHICreateShaderResourceView_RenderThread(RHICmdList, VertexBuffer, Stride, Format);
-		SRV->ViewIdentity = VertexBuffer->GetWholeResourceIdentity();
+		FShaderResourceViewRHIRef SRV = RHI->RHICreateShaderResourceView_RenderThread(RHICmdList, Buffer, Stride, Format);
+		SRV->ViewIdentity = Buffer->GetWholeResourceIdentity();
 		return SRV;
 	}
 
-	virtual FShaderResourceViewRHIRef RHICreateShaderResourceView_RenderThread(class FRHICommandListImmediate& RHICmdList, FRHIIndexBuffer* Buffer) override final
+	virtual FShaderResourceViewRHIRef RHICreateShaderResourceView_RenderThread(class FRHICommandListImmediate& RHICmdList, FRHIBuffer* Buffer) override final
 	{
 		FShaderResourceViewRHIRef SRV = RHI->RHICreateShaderResourceView_RenderThread(RHICmdList, Buffer);
 		SRV->ViewIdentity = Buffer->GetWholeResourceIdentity();
 		return SRV;
 	}
 
-	virtual FShaderResourceViewRHIRef RHICreateShaderResourceView_RenderThread(class FRHICommandListImmediate& RHICmdList, FRHIStructuredBuffer* StructuredBuffer) override final
-	{
-		check(IsFeatureLevelSupported(GMaxRHIShaderPlatform, ERHIFeatureLevel::SM5));
-		FShaderResourceViewRHIRef SRV = RHI->RHICreateShaderResourceView_RenderThread(RHICmdList, StructuredBuffer);
-		SRV->ViewIdentity = StructuredBuffer->GetWholeResourceIdentity();
-		return SRV;
-	}
-
 	virtual FShaderResourceViewRHIRef RHICreateShaderResourceViewWriteMask_RenderThread(class FRHICommandListImmediate& RHICmdList, FRHITexture2D* Texture2D) override final
 	{
 		FShaderResourceViewRHIRef SRV = RHI->RHICreateShaderResourceViewWriteMask_RenderThread(RHICmdList, Texture2D);
@@ -1728,16 +1678,6 @@ public:
 		RHI->RHIDiscardTransientResource_RenderThread(Buffer);
 	}
 
-	virtual void RHIAcquireTransientResource_RenderThread(FRHIStructuredBuffer* Buffer) override final
-	{
-		RHI->RHIAcquireTransientResource_RenderThread(Buffer);
-	}
-
-	virtual void RHIDiscardTransientResource_RenderThread(FRHIStructuredBuffer* Buffer) override final
-	{
-		RHI->RHIDiscardTransientResource_RenderThread(Buffer);
-	}
-
 	virtual void RHIReadSurfaceFloatData_RenderThread(class FRHICommandListImmediate& RHICmdList, FRHITexture* Texture, FIntRect Rect, TArray<FFloat16Color>& OutData, ECubeFace CubeFace, int32 ArrayIndex, int32 MipIndex) override final
 	{
 		RHI->RHIReadSurfaceFloatData_RenderThread(RHICmdList, Texture, Rect, OutData, CubeFace, ArrayIndex, MipIndex);
diff --git a/Engine/Source/Runtime/RHI/Public/RHIValidationContext.h b/Engine/Source/Runtime/RHI/Public/RHIValidationContext.h
--- a/Engine/Source/Runtime/RHI/Public/RHIValidationContext.h
+++ b/Engine/Source/Runtime/RHI/Public/RHIValidationContext.h
@@ -41,7 +41,7 @@ public:
 		Tracker->Dispatch();
 	}
 
-	virtual void RHIDispatchIndirectComputeShader(FRHIVertexBuffer* ArgumentBuffer, uint32 ArgumentOffset) override final
+	virtual void RHIDispatchIndirectComputeShader(FRHIBuffer* ArgumentBuffer, uint32 ArgumentOffset) override final
 	{
 		Tracker->Assert(ArgumentBuffer->GetWholeResourceIdentity(), ERHIAccess::IndirectArgs);
 
@@ -248,7 +248,7 @@ public:
 		Tracker->Dispatch();
 	}
 
-	virtual void RHIDispatchIndirectComputeShader(FRHIVertexBuffer* ArgumentBuffer, uint32 ArgumentOffset) override final
+	virtual void RHIDispatchIndirectComputeShader(FRHIBuffer* ArgumentBuffer, uint32 ArgumentOffset) override final
 	{
 		Tracker->Assert(ArgumentBuffer->GetWholeResourceIdentity(), ERHIAccess::IndirectArgs);
 
@@ -689,7 +689,7 @@ public:
 		Tracker->Draw();
 	}
 
-	virtual void RHIDrawPrimitiveIndirect(FRHIVertexBuffer* ArgumentBuffer, uint32 ArgumentOffset) override final
+	virtual void RHIDrawPrimitiveIndirect(FRHIBuffer* ArgumentBuffer, uint32 ArgumentOffset) override final
 	{
 		checkf(State.bGfxPSOSet, TEXT("A Graphics PSO has to be set to draw!"));
 		SetupDrawing();
@@ -698,7 +698,7 @@ public:
 		Tracker->Draw();
 	}
 
-	virtual void RHIDrawIndexedIndirect(FRHIIndexBuffer* IndexBufferRHI, FRHIStructuredBuffer* ArgumentsBufferRHI, int32 DrawArgumentsIndex, uint32 NumInstances) override final
+	virtual void RHIDrawIndexedIndirect(FRHIBuffer* IndexBufferRHI, FRHIBuffer* ArgumentsBufferRHI, int32 DrawArgumentsIndex, uint32 NumInstances) override final
 	{
 		checkf(State.bGfxPSOSet, TEXT("A Graphics PSO has to be set to draw!"));
 		SetupDrawing();
@@ -709,7 +709,7 @@ public:
 	}
 
 	// @param NumPrimitives need to be >0 
-	virtual void RHIDrawIndexedPrimitive(FRHIIndexBuffer* IndexBuffer, int32 BaseVertexIndex, uint32 FirstInstance, uint32 NumVertices, uint32 StartIndex, uint32 NumPrimitives, uint32 NumInstances) override final
+	virtual void RHIDrawIndexedPrimitive(FRHIBuffer* IndexBuffer, int32 BaseVertexIndex, uint32 FirstInstance, uint32 NumVertices, uint32 StartIndex, uint32 NumPrimitives, uint32 NumInstances) override final
 	{
 		checkf(State.bGfxPSOSet, TEXT("A Graphics PSO has to be set to draw!"));
 		SetupDrawing();
@@ -718,7 +718,7 @@ public:
 		Tracker->Draw();
 	}
 
-	virtual void RHIDrawIndexedPrimitiveIndirect(FRHIIndexBuffer* IndexBuffer, FRHIVertexBuffer* ArgumentBuffer, uint32 ArgumentOffset) override final
+	virtual void RHIDrawIndexedPrimitiveIndirect(FRHIBuffer* IndexBuffer, FRHIBuffer* ArgumentBuffer, uint32 ArgumentOffset) override final
 	{
 		checkf(State.bGfxPSOSet, TEXT("A Graphics PSO has to be set to draw!"));
 		SetupDrawing();
diff --git a/Engine/Source/Runtime/RenderCore/Public/RenderGraphResources.h b/Engine/Source/Runtime/RenderCore/Public/RenderGraphResources.h
--- a/Engine/Source/Runtime/RenderCore/Public/RenderGraphResources.h
+++ b/Engine/Source/Runtime/RenderCore/Public/RenderGraphResources.h
@@ -969,9 +969,10 @@ struct FRDGBufferDesc
 		switch (Type)
 		{
 		default: checkNoEntry();
-		case EUnderlyingType::VertexBuffer: return FRHITransitionInfo::EType::VertexBuffer;
-		case EUnderlyingType::IndexBuffer: return FRHITransitionInfo::EType::IndexBuffer;
-		case EUnderlyingType::StructuredBuffer: return FRHITransitionInfo::EType::StructuredBuffer;
+		case EUnderlyingType::VertexBuffer:
+		case EUnderlyingType::IndexBuffer:
+		case EUnderlyingType::StructuredBuffer:
+			return FRHITransitionInfo::EType::Buffer;
 		}
 	}
 
diff --git a/Engine/Source/Runtime/VulkanRHI/Private/VulkanBarriers.cpp b/Engine/Source/Runtime/VulkanRHI/Private/VulkanBarriers.cpp
--- a/Engine/Source/Runtime/VulkanRHI/Private/VulkanBarriers.cpp
+++ b/Engine/Source/Runtime/VulkanRHI/Private/VulkanBarriers.cpp
@@ -143,7 +143,7 @@ static VkPipelineStageFlags GetVkStageFlagsForLayout(VkImageLayout Layout)
 //
 // Get the Vulkan stage flags, access flags and image layout (if relevant) corresponding to an ERHIAccess value from the public API.
 //
-static void GetVkStageAndAccessFlags(ERHIAccess RHIAccess, FRHITransitionInfo::EType ResourceType, bool bIsDepthStencil, VkPipelineStageFlags& StageFlags, VkAccessFlags& AccessFlags, VkImageLayout& Layout, bool bIsSourceState)
+static void GetVkStageAndAccessFlags(ERHIAccess RHIAccess, FRHITransitionInfo::EType ResourceType, uint32 UsageFlags, bool bIsDepthStencil, VkPipelineStageFlags& StageFlags, VkAccessFlags& AccessFlags, VkImageLayout& Layout, bool bIsSourceState)
 {
 	// From Vulkan's point of view, when performing a multisample resolve via a render pass attachment, resolve targets are the same as render targets .
 	// The caller signals this situation by setting both the RTV and ResolveDst flags, and we simply remove ResolveDst in that case,
@@ -232,11 +232,15 @@ static void GetVkStageAndAccessFlags(ERHIAccess RHIAccess, FRHITransitionInfo::E
 		StageFlags |= VK_PIPELINE_STAGE_VERTEX_INPUT_BIT;
 		switch (ResourceType)
 		{
-			case FRHITransitionInfo::EType::IndexBuffer:
-				AccessFlags |= VK_ACCESS_INDEX_READ_BIT;
-				break;
-			case FRHITransitionInfo::EType::VertexBuffer:
-				AccessFlags |= VK_ACCESS_VERTEX_ATTRIBUTE_READ_BIT;
+			case FRHITransitionInfo::EType::Buffer:
+				if ((UsageFlags & VK_BUFFER_USAGE_INDEX_BUFFER_BIT) != 0)
+				{
+					AccessFlags |= VK_ACCESS_INDEX_READ_BIT;
+				}
+				if ((UsageFlags & VK_BUFFER_USAGE_VERTEX_BUFFER_BIT) != 0)
+				{
+					AccessFlags |= VK_ACCESS_VERTEX_ATTRIBUTE_READ_BIT;
+				}
 				break;
 			default:
 				checkNoEntry();
@@ -569,22 +573,17 @@ void FVulkanDynamicRHI::RHICreateTransition(FRHITransition* Transition, ERHIPipe
 		FVulkanResourceMultiBuffer* Buffer = nullptr;
 		FVulkanTextureBase* Texture = nullptr;
 		FRHITransitionInfo::EType UnderlyingType = Info.Type;
+		uint32 UsageFlags = 0;
+
 		switch (Info.Type)
 		{
 		case FRHITransitionInfo::EType::Texture:
 			Texture = FVulkanTextureBase::Cast(Info.Texture);
 			break;
 
-		case FRHITransitionInfo::EType::VertexBuffer:
-			Buffer = ResourceCast(Info.VertexBuffer);
-			break;
-
-		case FRHITransitionInfo::EType::IndexBuffer:
-			Buffer = ResourceCast(Info.IndexBuffer);
-			break;
-
-		case FRHITransitionInfo::EType::StructuredBuffer:
-			Buffer = ResourceCast(Info.StructuredBuffer);
+		case FRHITransitionInfo::EType::Buffer:
+			Buffer = ResourceCast(Info.Buffer);
+			UsageFlags = Buffer->GetBufferUsageFlags();
 			break;
 
 		case FRHITransitionInfo::EType::UAV:
@@ -598,17 +597,17 @@ void FVulkanDynamicRHI::RHICreateTransition(FRHITransition* Transition, ERHIPipe
 			else if (UAV->SourceIndexBuffer)
 			{
 				Buffer = UAV->SourceIndexBuffer;
-				UnderlyingType = FRHITransitionInfo::EType::IndexBuffer;
+				UnderlyingType = FRHITransitionInfo::EType::Buffer;
 			}
 			else if (UAV->SourceVertexBuffer)
 			{
 				Buffer = UAV->SourceVertexBuffer;
-				UnderlyingType = FRHITransitionInfo::EType::VertexBuffer;
+				UnderlyingType = FRHITransitionInfo::EType::Buffer;
 			}
 			else if (UAV->SourceStructuredBuffer)
 			{
 				Buffer = UAV->SourceStructuredBuffer;
-				UnderlyingType = FRHITransitionInfo::EType::StructuredBuffer;
+				UnderlyingType = FRHITransitionInfo::EType::Buffer;
 			}
 			else
 			{
@@ -679,8 +678,8 @@ void FVulkanDynamicRHI::RHICreateTransition(FRHITransition* Transition, ERHIPipe
 		}
 		else
 		{
-			GetVkStageAndAccessFlags(Info.AccessBefore, UnderlyingType, bIsDepthStencil, SrcStageMask, SrcAccessFlags, SrcLayout, true);
-			GetVkStageAndAccessFlags(Info.AccessAfter, UnderlyingType, bIsDepthStencil, DstStageMask, DstAccessFlags, DstLayout, false);
+			GetVkStageAndAccessFlags(Info.AccessBefore, UnderlyingType, UsageFlags, bIsDepthStencil, SrcStageMask, SrcAccessFlags, SrcLayout, true);
+			GetVkStageAndAccessFlags(Info.AccessAfter, UnderlyingType, UsageFlags, bIsDepthStencil, DstStageMask, DstAccessFlags, DstLayout, false);
 		}
 
 		// In case of async compute, override the stage and access flags computed above, since only the compute shader stage is relevant.
@@ -1216,8 +1215,8 @@ void FVulkanPipelineBarrier::AddImageAccessTransition(const FVulkanSurface& Surf
 	VkPipelineStageFlags ImgSrcStage, ImgDstStage;
 	VkAccessFlags SrcAccessFlags, DstAccessFlags;
 	VkImageLayout SrcLayout, DstLayout;
-	GetVkStageAndAccessFlags(SrcAccess, FRHITransitionInfo::EType::Texture, bIsDepthStencil, ImgSrcStage, SrcAccessFlags, SrcLayout, true);
-	GetVkStageAndAccessFlags(DstAccess, FRHITransitionInfo::EType::Texture, bIsDepthStencil, ImgDstStage, DstAccessFlags, DstLayout, false);
+	GetVkStageAndAccessFlags(SrcAccess, FRHITransitionInfo::EType::Texture, 0, bIsDepthStencil, ImgSrcStage, SrcAccessFlags, SrcLayout, true);
+	GetVkStageAndAccessFlags(DstAccess, FRHITransitionInfo::EType::Texture, 0, bIsDepthStencil, ImgDstStage, DstAccessFlags, DstLayout, false);
 
 	SrcStageMask |= ImgSrcStage;
 	DstStageMask |= ImgDstStage;
diff --git a/Engine/Source/Runtime/VulkanRHI/Private/VulkanCommands.cpp b/Engine/Source/Runtime/VulkanRHI/Private/VulkanCommands.cpp
--- a/Engine/Source/Runtime/VulkanRHI/Private/VulkanCommands.cpp
+++ b/Engine/Source/Runtime/VulkanRHI/Private/VulkanCommands.cpp
@@ -186,7 +186,7 @@ static FORCEINLINE ShaderStage::EStage GetAndVerifyShaderStageAndVulkanShader(FR
 
 void FVulkanCommandListContext::RHISetStreamSource(uint32 StreamIndex, FRHIVertexBuffer* VertexBufferRHI, uint32 Offset)
 {
-	FVulkanVertexBuffer* VertexBuffer = ResourceCast(VertexBufferRHI);
+	FVulkanResourceMultiBuffer* VertexBuffer = ResourceCast(VertexBufferRHI);
 	if (VertexBuffer != nullptr)
 	{
 		PendingGfxState->SetStreamSource(StreamIndex, VertexBuffer->GetHandle(), Offset + VertexBuffer->GetOffset());
@@ -250,10 +250,10 @@ void FVulkanCommandListContext::RHIDispatchComputeShader(uint32 ThreadGroupCount
 	VulkanRHI::DebugHeavyWeightBarrier(CmdBuffer, 2);
 }
 
-void FVulkanCommandListContext::RHIDispatchIndirectComputeShader(FRHIVertexBuffer* ArgumentBufferRHI, uint32 ArgumentOffset)
+void FVulkanCommandListContext::RHIDispatchIndirectComputeShader(FRHIBuffer* ArgumentBufferRHI, uint32 ArgumentOffset)
 {
 	static_assert(sizeof(FRHIDispatchIndirectParameters) == sizeof(VkDispatchIndirectCommand), "Dispatch indirect doesn't match!");
-	FVulkanVertexBuffer* ArgumentBuffer = ResourceCast(ArgumentBufferRHI);
+	FVulkanResourceMultiBuffer* ArgumentBuffer = ResourceCast(ArgumentBufferRHI);
 
 	FVulkanCmdBuffer* Cmd = CommandBufferManager->GetActiveCmdBuffer();
 	ensure(Cmd->IsOutsideRenderPass());
@@ -654,7 +654,7 @@ void FVulkanCommandListContext::RHIDrawPrimitive(uint32 BaseVertexIndex, uint32
 	}
 }
 
-void FVulkanCommandListContext::RHIDrawPrimitiveIndirect(FRHIVertexBuffer* ArgumentBufferRHI, uint32 ArgumentOffset)
+void FVulkanCommandListContext::RHIDrawPrimitiveIndirect(FRHIBuffer* ArgumentBufferRHI, uint32 ArgumentOffset)
 {
 	static_assert(sizeof(FRHIDrawIndirectParameters) == sizeof(VkDrawIndirectCommand), "Draw indirect doesn't match!");
 
@@ -667,7 +667,7 @@ void FVulkanCommandListContext::RHIDrawPrimitiveIndirect(FRHIVertexBuffer* Argum
 	VkCommandBuffer CmdBuffer = Cmd->GetHandle();
 	PendingGfxState->PrepareForDraw(Cmd);
 
-	FVulkanVertexBuffer* ArgumentBuffer = ResourceCast(ArgumentBufferRHI);
+	FVulkanResourceMultiBuffer* ArgumentBuffer = ResourceCast(ArgumentBufferRHI);
 
 
 	VulkanRHI::vkCmdDrawIndirect(CmdBuffer, ArgumentBuffer->GetHandle(), ArgumentBuffer->GetOffset() + ArgumentOffset, 1, sizeof(VkDrawIndirectCommand));
@@ -678,7 +678,7 @@ void FVulkanCommandListContext::RHIDrawPrimitiveIndirect(FRHIVertexBuffer* Argum
 	}
 }
 
-void FVulkanCommandListContext::RHIDrawIndexedPrimitive(FRHIIndexBuffer* IndexBufferRHI, int32 BaseVertexIndex, uint32 FirstInstance,
+void FVulkanCommandListContext::RHIDrawIndexedPrimitive(FRHIBuffer* IndexBufferRHI, int32 BaseVertexIndex, uint32 FirstInstance,
 	uint32 NumVertices, uint32 StartIndex, uint32 NumPrimitives, uint32 NumInstances)
 {
 #if VULKAN_ENABLE_AGGRESSIVE_STATS
@@ -688,7 +688,7 @@ void FVulkanCommandListContext::RHIDrawIndexedPrimitive(FRHIIndexBuffer* IndexBu
 	RHI_DRAW_CALL_STATS(PendingGfxState->PrimitiveType, NumInstances*NumPrimitives);
 	checkf(GRHISupportsFirstInstance || FirstInstance == 0, TEXT("FirstInstance must be 0, see GRHISupportsFirstInstance"));
 
-	FVulkanIndexBuffer* IndexBuffer = ResourceCast(IndexBufferRHI);
+	FVulkanResourceMultiBuffer* IndexBuffer = ResourceCast(IndexBufferRHI);
 	FVulkanCmdBuffer* Cmd = CommandBufferManager->GetActiveCmdBuffer();
 	VkCommandBuffer CmdBuffer = Cmd->GetHandle();
 	PendingGfxState->PrepareForDraw(Cmd);
@@ -703,20 +703,20 @@ void FVulkanCommandListContext::RHIDrawIndexedPrimitive(FRHIIndexBuffer* IndexBu
 	}
 }
 
-void FVulkanCommandListContext::RHIDrawIndexedIndirect(FRHIIndexBuffer* IndexBufferRHI, FRHIStructuredBuffer* ArgumentsBufferRHI, int32 DrawArgumentsIndex, uint32 NumInstances)
+void FVulkanCommandListContext::RHIDrawIndexedIndirect(FRHIBuffer* IndexBufferRHI, FRHIBuffer* ArgumentsBufferRHI, int32 DrawArgumentsIndex, uint32 NumInstances)
 {
 #if VULKAN_ENABLE_AGGRESSIVE_STATS
 	SCOPE_CYCLE_COUNTER(STAT_VulkanDrawCallTime);
 #endif
 	RHI_DRAW_CALL_INC();
 
-	FVulkanIndexBuffer* IndexBuffer = ResourceCast(IndexBufferRHI);
+	FVulkanResourceMultiBuffer* IndexBuffer = ResourceCast(IndexBufferRHI);
 	FVulkanCmdBuffer* Cmd = CommandBufferManager->GetActiveCmdBuffer();
 	VkCommandBuffer CmdBuffer = Cmd->GetHandle();
 	PendingGfxState->PrepareForDraw(Cmd);
 	VulkanRHI::vkCmdBindIndexBuffer(CmdBuffer, IndexBuffer->GetHandle(), IndexBuffer->GetOffset(), IndexBuffer->GetIndexType());
 
-	FVulkanStructuredBuffer* ArgumentBuffer = ResourceCast(ArgumentsBufferRHI);
+	FVulkanResourceMultiBuffer* ArgumentBuffer = ResourceCast(ArgumentsBufferRHI);
 	VkDeviceSize ArgumentOffset = DrawArgumentsIndex * sizeof(VkDrawIndexedIndirectCommand);
 
 
@@ -728,20 +728,20 @@ void FVulkanCommandListContext::RHIDrawIndexedIndirect(FRHIIndexBuffer* IndexBuf
 	}
 }
 
-void FVulkanCommandListContext::RHIDrawIndexedPrimitiveIndirect(FRHIIndexBuffer* IndexBufferRHI, FRHIVertexBuffer* ArgumentBufferRHI,uint32 ArgumentOffset)
+void FVulkanCommandListContext::RHIDrawIndexedPrimitiveIndirect(FRHIBuffer* IndexBufferRHI, FRHIBuffer* ArgumentBufferRHI, uint32 ArgumentOffset)
 {
 #if VULKAN_ENABLE_AGGRESSIVE_STATS
 	SCOPE_CYCLE_COUNTER(STAT_VulkanDrawCallTime);
 #endif
 	RHI_DRAW_CALL_INC();
 
-	FVulkanIndexBuffer* IndexBuffer = ResourceCast(IndexBufferRHI);
+	FVulkanResourceMultiBuffer* IndexBuffer = ResourceCast(IndexBufferRHI);
 	FVulkanCmdBuffer* Cmd = CommandBufferManager->GetActiveCmdBuffer();
 	VkCommandBuffer CmdBuffer = Cmd->GetHandle();
 	PendingGfxState->PrepareForDraw(Cmd);
 	VulkanRHI::vkCmdBindIndexBuffer(CmdBuffer, IndexBuffer->GetHandle(), IndexBuffer->GetOffset(), IndexBuffer->GetIndexType());
 
-	FVulkanVertexBuffer* ArgumentBuffer = ResourceCast(ArgumentBufferRHI);
+	FVulkanResourceMultiBuffer* ArgumentBuffer = ResourceCast(ArgumentBufferRHI);
 
 
 	VulkanRHI::vkCmdDrawIndexedIndirect(CmdBuffer, ArgumentBuffer->GetHandle(), ArgumentBuffer->GetOffset() + ArgumentOffset, 1, sizeof(VkDrawIndexedIndirectCommand));
@@ -912,7 +912,7 @@ void FVulkanCommandListContext::PrepareParallelFromBase(const FVulkanCommandList
 void FVulkanCommandListContext::RHICopyToStagingBuffer(FRHIVertexBuffer* SourceBufferRHI, FRHIStagingBuffer* StagingBufferRHI, uint32 Offset, uint32 NumBytes)
 {
 	FVulkanCmdBuffer* CmdBuffer = CommandBufferManager->GetActiveCmdBuffer();
-	FVulkanVertexBuffer* VertexBuffer = ResourceCast(SourceBufferRHI);
+	FVulkanResourceMultiBuffer* VertexBuffer = ResourceCast(SourceBufferRHI);
 
 	ensure(CmdBuffer->IsOutsideRenderPass());
 	ensureMsgf((SourceBufferRHI->GetUsage() & BUF_SourceCopy) != 0, TEXT("Buffers used as copy source need to be created with BUF_SourceCopy"));
diff --git a/Engine/Source/Runtime/VulkanRHI/Private/VulkanContext.h b/Engine/Source/Runtime/VulkanRHI/Private/VulkanContext.h
--- a/Engine/Source/Runtime/VulkanRHI/Private/VulkanContext.h
+++ b/Engine/Source/Runtime/VulkanRHI/Private/VulkanContext.h
@@ -56,10 +56,10 @@ public:
 	virtual void RHISetShaderParameter(FRHIComputeShader* ComputeShader, uint32 BufferIndex, uint32 BaseIndex, uint32 NumBytes, const void* NewValue) final override;
 	virtual void RHISetStencilRef(uint32 StencilRef) final override;
 	virtual void RHIDrawPrimitive(uint32 BaseVertexIndex, uint32 NumPrimitives, uint32 NumInstances) final override;
-	virtual void RHIDrawPrimitiveIndirect(FRHIVertexBuffer* ArgumentBuffer, uint32 ArgumentOffset) final override;
-	virtual void RHIDrawIndexedIndirect(FRHIIndexBuffer* IndexBufferRHI, FRHIStructuredBuffer* ArgumentsBufferRHI, int32 DrawArgumentsIndex, uint32 NumInstances) final override;
-	virtual void RHIDrawIndexedPrimitive(FRHIIndexBuffer* IndexBuffer, int32 BaseVertexIndex, uint32 FirstInstance, uint32 NumVertices, uint32 StartIndex, uint32 NumPrimitives, uint32 NumInstances) final override;
-	virtual void RHIDrawIndexedPrimitiveIndirect(FRHIIndexBuffer* IndexBuffer, FRHIVertexBuffer* ArgumentBuffer, uint32 ArgumentOffset) final override;
+	virtual void RHIDrawPrimitiveIndirect(FRHIBuffer* ArgumentBuffer, uint32 ArgumentOffset) final override;
+	virtual void RHIDrawIndexedIndirect(FRHIBuffer* IndexBufferRHI, FRHIBuffer* ArgumentsBufferRHI, int32 DrawArgumentsIndex, uint32 NumInstances) final override;
+	virtual void RHIDrawIndexedPrimitive(FRHIBuffer* IndexBuffer, int32 BaseVertexIndex, uint32 FirstInstance, uint32 NumVertices, uint32 StartIndex, uint32 NumPrimitives, uint32 NumInstances) final override;
+	virtual void RHIDrawIndexedPrimitiveIndirect(FRHIBuffer* IndexBuffer, FRHIBuffer* ArgumentBuffer, uint32 ArgumentOffset) final override;
 	virtual void RHISetDepthBounds(float MinDepth, float MaxDepth) final override;
 	virtual void RHIPushEvent(const TCHAR* Name, FColor Color) final override;
 	virtual void RHIPopEvent() final override;
@@ -67,7 +67,7 @@ public:
 	virtual void RHISetComputeShader(FRHIComputeShader* ComputeShader) final override;
 	virtual void RHISetComputePipelineState(FRHIComputePipelineState* ComputePipelineState) final override;
 	virtual void RHIDispatchComputeShader(uint32 ThreadGroupCountX, uint32 ThreadGroupCountY, uint32 ThreadGroupCountZ) final override;
-	virtual void RHIDispatchIndirectComputeShader(FRHIVertexBuffer* ArgumentBuffer, uint32 ArgumentOffset) final override;
+	virtual void RHIDispatchIndirectComputeShader(FRHIBuffer* ArgumentBuffer, uint32 ArgumentOffset) final override;
 
 	virtual void RHISetMultipleViewports(uint32 Count, const FViewportBounds* Data) final override;
 	virtual void RHIClearUAVFloat(FRHIUnorderedAccessView* UnorderedAccessViewRHI, const FVector4& Values) final override;
