Fix for non-unity compile error.

#rnx

diff --git a/Engine/Source/Runtime/RHI/Private/RHIValidation.cpp b/Engine/Source/Runtime/RHI/Private/RHIValidation.cpp
--- a/Engine/Source/Runtime/RHI/Private/RHIValidation.cpp
+++ b/Engine/Source/Runtime/RHI/Private/RHIValidation.cpp
@@ -7,6 +7,7 @@
 #include "RHIValidation.h"
 #include "RHIValidationContext.h"
 #include "HAL/PlatformStackWalk.h"
+#include "Misc/OutputDeviceRedirector.h"
 
 #if ENABLE_RHI_VALIDATION
 
