Fixed build.

#rnx

diff --git a/Engine/Source/Runtime/RHI/Private/RHIValidation.cpp b/Engine/Source/Runtime/RHI/Private/RHIValidation.cpp
--- a/Engine/Source/Runtime/RHI/Private/RHIValidation.cpp
+++ b/Engine/Source/Runtime/RHI/Private/RHIValidation.cpp
@@ -195,16 +195,8 @@ void FValidationRHI::RHICreateTransition(FRHITransition* Transition, ERHIPipelin
 			Identity = Info.Texture->GetTransitionIdentity(Info);
 			break;
 
-		case FRHITransitionInfo::EType::VertexBuffer:
-			Identity = Info.VertexBuffer->GetWholeResourceIdentity();
-			break;
-
-		case FRHITransitionInfo::EType::IndexBuffer:
-			Identity = Info.IndexBuffer->GetWholeResourceIdentity();
-			break;
-
-		case FRHITransitionInfo::EType::StructuredBuffer:
-			Identity = Info.StructuredBuffer->GetWholeResourceIdentity();
+		case FRHITransitionInfo::EType::Buffer:
+			Identity = Info.Buffer->GetWholeResourceIdentity();
 			break;
 
 		case FRHITransitionInfo::EType::UAV:
