implemented RHI validation layer passthrough of RHIResummarizeTile.

#rnx

diff --git a/Engine/Source/Runtime/RHI/Public/RHIValidationContext.h b/Engine/Source/Runtime/RHI/Public/RHIValidationContext.h
--- a/Engine/Source/Runtime/RHI/Public/RHIValidationContext.h
+++ b/Engine/Source/Runtime/RHI/Public/RHIValidationContext.h
@@ -346,6 +346,11 @@ public:
 		}
 	}
 
+	virtual void RHIResummarizeHTile(FRHITexture2D* DepthTexture) override final
+	{
+		RHIContext->RHIResummarizeHTile(DepthTexture);
+	}
+
 	virtual void RHIBeginTransitions(TArrayView<const FRHITransition*> Transitions) override final
 	{
 		ensureMsgf(!State.bInsideBeginRenderPass, TEXT("Beginning a transition within a renderpass is not supported!"));
