Enable RHIValidation creation for XboxOne GDK

diff --git a/Engine/Source/Runtime/RHI/Public/RHIValidation.h b/Engine/Source/Runtime/RHI/Public/RHIValidation.h
--- a/Engine/Source/Runtime/RHI/Public/RHIValidation.h
+++ b/Engine/Source/Runtime/RHI/Public/RHIValidation.h
@@ -514,6 +514,25 @@ public:
 		return RHI->RHICalcTexture2DPlatformSize(SizeX, SizeY, Format, NumMips, NumSamples, Flags, CreateInfo, OutAlign);
 	}
 
+	/**
+	* Computes the total size of a virtual memory (VM) based 2D texture with the specified parameters.
+	*
+	* @param Mip0Width - width of the top mip
+	* @param Mip0Height - height of the top mip
+	* @param Format - EPixelFormat texture format
+	* @param NumMips - number of mips the texture has
+	* @param FirstMipIdx - index of the first resident mip
+	* @param NumSamples - number of MSAA samples, usually 1
+	* @param Flags - ETextureCreateFlags creation flags
+	* @param OutAlign - Alignment required for this texture. Output parameter.
+	*/
+	// FlushType: Thread safe
+	virtual uint64 RHICalcVMTexture2DPlatformSize(uint32 Mip0Width, uint32 Mip0Height, uint8 Format, uint32 NumMips, uint32 FirstMipIdx, uint32 NumSamples, ETextureCreateFlags Flags, uint32& OutAlign) override final
+	{
+		return RHI->RHICalcVMTexture2DPlatformSize(Mip0Width, Mip0Height, Format, NumMips, FirstMipIdx, NumSamples, Flags, OutAlign);
+	}
+
+
 	/**
 	* Computes the total size of a 3D texture with the specified parameters.
 	* @param SizeX - width of the texture to create
