Handle null Buffer in FShaderResourceViewInitializer::InitType.
#rnx

diff --git a/Engine/Source/Runtime/RHI/Private/DynamicRHI.cpp b/Engine/Source/Runtime/RHI/Private/DynamicRHI.cpp
--- a/Engine/Source/Runtime/RHI/Private/DynamicRHI.cpp
+++ b/Engine/Source/Runtime/RHI/Private/DynamicRHI.cpp
@@ -570,14 +570,21 @@ FShaderResourceViewInitializer::FShaderResourceViewInitializer(FRHIBuffer* InBuf
 
 void FShaderResourceViewInitializer::InitType()
 {
-	uint32 Usage = BufferInitializer.Buffer->GetUsage();
-	if (Usage & BUF_VertexBuffer)
+	if (BufferInitializer.Buffer)
 	{
-		Type = EType::VertexBufferSRV;
-	}
-	else if (Usage & BUF_IndexBuffer)
-	{
-		Type = EType::IndexBufferSRV;
+		uint32 Usage = BufferInitializer.Buffer->GetUsage();
+		if (Usage & BUF_VertexBuffer)
+		{
+			Type = EType::VertexBufferSRV;
+		}
+		else if (Usage & BUF_IndexBuffer)
+		{
+			Type = EType::IndexBufferSRV;
+		}
+		else
+		{
+			Type = EType::StructuredBufferSRV;
+		}
 	}
 	else
 	{
