Reverting Non-Unity fix as it introduces circular include RHIValidationCommon.h -> RHI.h

diff --git a/Engine/Source/Runtime/RHI/Public/RHIValidationCommon.h b/Engine/Source/Runtime/RHI/Public/RHIValidationCommon.h
--- a/Engine/Source/Runtime/RHI/Public/RHIValidationCommon.h
+++ b/Engine/Source/Runtime/RHI/Public/RHIValidationCommon.h
@@ -14,8 +14,6 @@ const bool GRHIValidationEnabled = false;
 
 #if ENABLE_RHI_VALIDATION
 
-#include "RHI.h"
-
 class FRHIUniformBuffer;
 
 // Forward declaration of function defined in RHIUtilities.h
