Fix for {Begin, End}UAVOverlap functions not using the compute context.

diff --git a/Engine/Source/Runtime/RHI/Public/RHICommandList.h b/Engine/Source/Runtime/RHI/Public/RHICommandList.h
--- a/Engine/Source/Runtime/RHI/Public/RHICommandList.h
+++ b/Engine/Source/Runtime/RHI/Public/RHICommandList.h
@@ -2525,7 +2525,7 @@ public:
 	{
 		if (Bypass())
 		{
-			GetContext().RHIBeginUAVOverlap();
+			GetComputeContext().RHIBeginUAVOverlap();
 			return;
 		}
 		ALLOC_COMMAND(FRHICommandBeginUAVOverlap)();
@@ -2557,7 +2557,7 @@ public:
 	{
 		if (Bypass())
 		{
-			GetContext().RHIBeginUAVOverlap(UAVs);
+			GetComputeContext().RHIBeginUAVOverlap(UAVs);
 			return;
 		}
 
@@ -2571,7 +2571,7 @@ public:
 	{
 		if (Bypass())
 		{
-			GetContext().RHIEndUAVOverlap(UAVs);
+			GetComputeContext().RHIEndUAVOverlap(UAVs);
 			return;
 		}
 
diff --git a/Engine/Source/Runtime/RHI/Public/RHICommandListCommandExecutes.inl b/Engine/Source/Runtime/RHI/Public/RHICommandListCommandExecutes.inl
--- a/Engine/Source/Runtime/RHI/Public/RHICommandListCommandExecutes.inl
+++ b/Engine/Source/Runtime/RHI/Public/RHICommandListCommandExecutes.inl
@@ -299,13 +299,13 @@ void FRHICommandEndUAVOverlap::Execute(FRHICommandListBase& CmdList)
 void FRHICommandBeginSpecificUAVOverlap::Execute(FRHICommandListBase& CmdList)
 {
 	RHISTAT(BeginSpecificUAVOverlap);
-	INTERNAL_DECORATOR(RHIBeginUAVOverlap)(UAVs);
+	INTERNAL_DECORATOR_COMPUTE(RHIBeginUAVOverlap)(UAVs);
 }
 
 void FRHICommandEndSpecificUAVOverlap::Execute(FRHICommandListBase& CmdList)
 {
 	RHISTAT(EndSpecificUAVOverlap);
-	INTERNAL_DECORATOR(RHIEndUAVOverlap)(UAVs);
+	INTERNAL_DECORATOR_COMPUTE(RHIEndUAVOverlap)(UAVs);
 }
 
 void FRHICommandDrawPrimitiveIndirect::Execute(FRHICommandListBase& CmdList)
