Fix a potential race and/or use-after-free of FAccelerationStructureBuildParams::Segments in BuildAccelerationStructures().

The array of build params was copied, but data pointed to by Segments array view wasn't.

diff --git a/Engine/Source/Runtime/RHI/Public/RHICommandList.h b/Engine/Source/Runtime/RHI/Public/RHICommandList.h
--- a/Engine/Source/Runtime/RHI/Public/RHICommandList.h
+++ b/Engine/Source/Runtime/RHI/Public/RHICommandList.h
@@ -2689,7 +2689,19 @@ public:
 		}
 		else
 		{
-			ALLOC_COMMAND(FRHICommandBuildAccelerationStructures)(AllocArray(Params));
+			// Copy the params themselves as well their segment lists, if there are any.
+			// AllocArray() can't be used here directly, as we have to modify the params after copy.
+			size_t DataSize = sizeof(FAccelerationStructureBuildParams) * Params.Num();
+			FAccelerationStructureBuildParams* InlineParams = (FAccelerationStructureBuildParams*) Alloc(DataSize, alignof(FAccelerationStructureBuildParams));
+			FMemory::Memcpy(InlineParams, Params.GetData(), DataSize);
+			for (int32 i=0; i<Params.Num(); ++i)
+			{
+				if (Params[i].Segments.Num())
+				{
+					InlineParams[i].Segments = AllocArray(Params[i].Segments);
+				}
+			}
+			ALLOC_COMMAND(FRHICommandBuildAccelerationStructures)(MakeArrayView(InlineParams, Params.Num()));
 		}
 	}
 
