Fixing access bits for DSV read. Fixing cases where DSV read would cause unwanted GPU behavior.

diff --git a/Engine/Source/Runtime/RHI/Public/RHIResources.h b/Engine/Source/Runtime/RHI/Public/RHIResources.h
--- a/Engine/Source/Runtime/RHI/Public/RHIResources.h
+++ b/Engine/Source/Runtime/RHI/Public/RHIResources.h
@@ -275,9 +275,7 @@ public:
 
 		// SRV access is allowed whilst a depth stencil target is "readable".
 		constexpr ERHIAccess DSVReadOnlyMask =
-			ERHIAccess::DSVRead |
-			ERHIAccess::SRVGraphics |
-			ERHIAccess::SRVCompute;
+			ERHIAccess::DSVRead;
 
 		// If write access is required, only the depth block can access the resource.
 		constexpr ERHIAccess DSVReadWriteMask =
