Add support for a bottom of pipe GPU texture cache flush to workaround an issue on specific platforms

- follow on change will enable the workaround for the needed platforms

diff --git a/Engine/Source/Runtime/Core/Public/HAL/Platform.h b/Engine/Source/Runtime/Core/Public/HAL/Platform.h
--- a/Engine/Source/Runtime/Core/Public/HAL/Platform.h
+++ b/Engine/Source/Runtime/Core/Public/HAL/Platform.h
@@ -608,6 +608,10 @@
 	#define PLATFORM_USE_PLATFORM_FILE_MANAGED_STORAGE_WRAPPER	0
 #endif
 
+#ifndef PLATFORM_REQUIRES_UAV_TO_RTV_TEXTURE_CACHE_FLUSH_WORKAROUND
+	#define PLATFORM_REQUIRES_UAV_TO_RTV_TEXTURE_CACHE_FLUSH_WORKAROUND 0
+#endif // #ifndef PLATFORM_REQUIRES_UAV_TO_RTV_TEXTURE_CACHE_FLUSH_WORKAROUND
+
 // deprecated, do not use
 #define PLATFORM_HAS_THREADSAFE_RHIGetRenderQueryResult	#
 #define PLATFORM_SUPPORTS_RHI_THREAD #
diff --git a/Engine/Source/Runtime/RHI/Public/RHICommandList.h b/Engine/Source/Runtime/RHI/Public/RHICommandList.h
--- a/Engine/Source/Runtime/RHI/Public/RHICommandList.h
+++ b/Engine/Source/Runtime/RHI/Public/RHICommandList.h
@@ -1989,6 +1989,19 @@ FRHICOMMAND_MACRO(FRHICommandBackBufferWaitTrackingBeginFrame)
 };
 #endif // #if PLATFORM_USE_BACKBUFFER_WRITE_TRANSITION_TRACKING
 
+#if PLATFORM_REQUIRES_UAV_TO_RTV_TEXTURE_CACHE_FLUSH_WORKAROUND
+FRHICOMMAND_MACRO(FRHICommandFlushTextureCacheBOP)
+{
+	FRHITexture* Texture;
+
+	FORCEINLINE_DEBUGGABLE FRHICommandFlushTextureCacheBOP(FRHITexture* TextureIn)
+		:	Texture(TextureIn)
+	{}
+	
+	RHI_API void Execute(FRHICommandListBase& CmdList);
+};
+#endif // #if PLATFORM_REQUIRES_UAV_TO_RTV_TEXTURE_CACHE_FLUSH_WORKAROUND
+
 FRHICOMMAND_MACRO(FRHICommandCopyBufferRegion)
 {
 	FRHIVertexBuffer* DestBuffer;
@@ -2672,6 +2685,18 @@ public:
 		}
 	}
 
+#if PLATFORM_REQUIRES_UAV_TO_RTV_TEXTURE_CACHE_FLUSH_WORKAROUND
+	FORCEINLINE_DEBUGGABLE void RHIFlushTextureCacheBOP(FRHITexture* Texture)
+	{
+		if (Bypass())
+		{
+			GetContext().RHIFlushTextureCacheBOP(Texture);
+			return;
+		}
+		ALLOC_COMMAND(FRHICommandFlushTextureCacheBOP)(Texture);
+	}
+#endif // #if PLATFORM_REQUIRES_UAV_TO_RTV_TEXTURE_CACHE_FLUSH_WORKAROUND
+
 #if RHI_RAYTRACING
 	FORCEINLINE_DEBUGGABLE void BuildAccelerationStructure(FRHIRayTracingGeometry* Geometry)
 	{
diff --git a/Engine/Source/Runtime/RHI/Public/RHICommandListCommandExecutes.inl b/Engine/Source/Runtime/RHI/Public/RHICommandListCommandExecutes.inl
--- a/Engine/Source/Runtime/RHI/Public/RHICommandListCommandExecutes.inl
+++ b/Engine/Source/Runtime/RHI/Public/RHICommandListCommandExecutes.inl
@@ -727,3 +727,11 @@ void FRHICommandBackBufferWaitTrackingBeginFrame::Execute(FRHICommandListBase& C
 	INTERNAL_DECORATOR(RHIBackBufferWaitTrackingBeginFrame)(FrameToken, bDeferred);
 }
 #endif // #if PLATFORM_USE_BACKBUFFER_WRITE_TRANSITION_TRACKING
+
+#if PLATFORM_REQUIRES_UAV_TO_RTV_TEXTURE_CACHE_FLUSH_WORKAROUND
+void FRHICommandFlushTextureCacheBOP::Execute(FRHICommandListBase& CmdList)
+{
+	RHISTAT(FRHICommandFlushTextureCacheBOP);
+	INTERNAL_DECORATOR(RHIFlushTextureCacheBOP)(Texture);
+}
+#endif // #if PLATFORM_REQUIRES_UAV_TO_RTV_TEXTURE_CACHE_FLUSH_WORKAROUND
diff --git a/Engine/Source/Runtime/RHI/Public/RHIContext.h b/Engine/Source/Runtime/RHI/Public/RHIContext.h
--- a/Engine/Source/Runtime/RHI/Public/RHIContext.h
+++ b/Engine/Source/Runtime/RHI/Public/RHIContext.h
@@ -682,6 +682,13 @@ public:
 	}
 #endif // #if PLATFORM_USE_BACKBUFFER_WRITE_TRANSITION_TRACKING
 
+#if PLATFORM_REQUIRES_UAV_TO_RTV_TEXTURE_CACHE_FLUSH_WORKAROUND
+	virtual void RHIFlushTextureCacheBOP(FRHITexture* Texture)
+	{
+		checkNoEntry();
+	}
+#endif // #if PLATFORM_REQUIRES_UAV_TO_RTV_TEXTURE_CACHE_FLUSH_WORKAROUND
+
 	protected:
 		FRHIRenderPassInfo RenderPassInfo;
 };
diff --git a/Engine/Source/Runtime/Renderer/Private/PostProcess/DiaphragmDOF.cpp b/Engine/Source/Runtime/Renderer/Private/PostProcess/DiaphragmDOF.cpp
--- a/Engine/Source/Runtime/Renderer/Private/PostProcess/DiaphragmDOF.cpp
+++ b/Engine/Source/Runtime/Renderer/Private/PostProcess/DiaphragmDOF.cpp
@@ -2450,6 +2450,10 @@ FRDGTextureRef DiaphragmDOF::AddPasses(
 				ERDGPassFlags::Raster,
 				[PassParameters, VertexShader, PixelShader, GatheringViewSize, DrawIndirectParametersOffset](FRHICommandList& RHICmdList)
 			{
+#if PLATFORM_REQUIRES_UAV_TO_RTV_TEXTURE_CACHE_FLUSH_WORKAROUND
+				RHICmdList.RHIFlushTextureCacheBOP(PassParameters->RenderTargets[0].GetTexture()->GetRHI());
+#endif // #if PLATFORM_REQUIRES_UAV_TO_RTV_TEXTURE_CACHE_FLUSH_WORKAROUND
+
 				RHICmdList.SetViewport(0, 0, 0.0f, GatheringViewSize.X, GatheringViewSize.Y, 1.0f);
 
 				FGraphicsPipelineStateInitializer GraphicsPSOInit;
