Merge LockBuffer functions

diff --git a/Engine/Source/Runtime/Apple/MetalRHI/Private/MetalDynamicRHI.h b/Engine/Source/Runtime/Apple/MetalRHI/Private/MetalDynamicRHI.h
--- a/Engine/Source/Runtime/Apple/MetalRHI/Private/MetalDynamicRHI.h
+++ b/Engine/Source/Runtime/Apple/MetalRHI/Private/MetalDynamicRHI.h
@@ -56,16 +56,12 @@ public:
 	virtual TRefCountPtr<FRHIComputePipelineState> RHICreateComputePipelineState(FRHIComputeShader* ComputeShader) final override;
 	virtual FUniformBufferRHIRef RHICreateUniformBuffer(const void* Contents, const FRHIUniformBufferLayout& Layout, EUniformBufferUsage Usage, EUniformBufferValidation Validation) final override;
 	virtual FIndexBufferRHIRef RHICreateIndexBuffer(uint32 Stride, uint32 Size, uint32 InUsage, ERHIAccess InResourceState, FRHIResourceCreateInfo& CreateInfo) final override;
-	virtual void* LockIndexBuffer_BottomOfPipe(FRHICommandListImmediate& RHICmdList, FRHIIndexBuffer* IndexBuffer, uint32 Offset, uint32 Size, EResourceLockMode LockMode) final override;
-	virtual void UnlockIndexBuffer_BottomOfPipe(FRHICommandListImmediate& RHICmdList, FRHIIndexBuffer* IndexBuffer) final override;
 	virtual void RHITransferBufferUnderlyingResource(FRHIBuffer* DestBuffer, FRHIBuffer* SrcBuffer) final override;
 	virtual FVertexBufferRHIRef RHICreateVertexBuffer(uint32 Size, uint32 InUsage, ERHIAccess InResourceState, FRHIResourceCreateInfo& CreateInfo) final override;
-	virtual void* LockVertexBuffer_BottomOfPipe(FRHICommandListImmediate& RHICmdList, FRHIVertexBuffer* VertexBuffer, uint32 Offset, uint32 SizeRHI, EResourceLockMode LockMode) final override;
-	virtual void UnlockVertexBuffer_BottomOfPipe(FRHICommandListImmediate& RHICmdList, FRHIVertexBuffer* VertexBuffer) final override;
 	virtual void RHICopyVertexBuffer(FRHIVertexBuffer* SourceBuffer, FRHIVertexBuffer* DestBuffer) final override;
 	virtual FStructuredBufferRHIRef RHICreateStructuredBuffer(uint32 Stride, uint32 Size, uint32 InUsage, ERHIAccess InResourceState, FRHIResourceCreateInfo& CreateInfo) final override;
-	virtual void* LockStructuredBuffer_BottomOfPipe(FRHICommandListImmediate& RHICmdList, FRHIStructuredBuffer* StructuredBuffer, uint32 Offset, uint32 SizeRHI, EResourceLockMode LockMode) final override;
-	virtual void UnlockStructuredBuffer_BottomOfPipe(FRHICommandListImmediate& RHICmdList, FRHIStructuredBuffer* StructuredBuffer) final override;
+	virtual void* LockBuffer_BottomOfPipe(FRHICommandListImmediate& RHICmdList, FRHIBuffer* Buffer, uint32 Offset, uint32 SizeRHI, EResourceLockMode LockMode) final override;
+	virtual void UnlockBuffer_BottomOfPipe(FRHICommandListImmediate& RHICmdList, FRHIBuffer* Buffer) final override;
 	virtual FUnorderedAccessViewRHIRef RHICreateUnorderedAccessView(FRHIStructuredBuffer* StructuredBuffer, bool bUseUAVCounter, bool bAppendBuffer) final override;
 	virtual FUnorderedAccessViewRHIRef RHICreateUnorderedAccessView(FRHITexture* Texture, uint32 MipLevel) final override;
 	virtual FUnorderedAccessViewRHIRef RHICreateUnorderedAccessView(FRHIBuffer* Buffer, uint8 Format) final override;
diff --git a/Engine/Source/Runtime/Apple/MetalRHI/Private/MetalIndexBuffer.cpp b/Engine/Source/Runtime/Apple/MetalRHI/Private/MetalIndexBuffer.cpp
--- a/Engine/Source/Runtime/Apple/MetalRHI/Private/MetalIndexBuffer.cpp
+++ b/Engine/Source/Runtime/Apple/MetalRHI/Private/MetalIndexBuffer.cpp
@@ -39,12 +39,12 @@ FIndexBufferRHIRef FMetalDynamicRHI::RHICreateIndexBuffer(uint32 Stride, uint32
 		check(Size == CreateInfo.ResourceArray->GetResourceDataSize());
 
 		// make a buffer usable by CPU
-		void* Buffer = ::RHILockIndexBuffer(IndexBuffer, 0, Size, RLM_WriteOnly);
+		void* Buffer = ::RHILockBuffer(IndexBuffer, 0, Size, RLM_WriteOnly);
 
 		// copy the contents of the given data into the buffer
 		FMemory::Memcpy(Buffer, CreateInfo.ResourceArray->GetResourceData(), Size);
 
-		::RHIUnlockIndexBuffer(IndexBuffer);
+		::RHIUnlockBuffer(IndexBuffer);
 
 		// Discard the resource array's contents.
 		CreateInfo.ResourceArray->Discard();
@@ -72,24 +72,6 @@ void FMetalDynamicRHI::RHITransferBufferUnderlyingResource(FRHIBuffer* DestBuffe
 	}
 }
 
-void* FMetalDynamicRHI::LockIndexBuffer_BottomOfPipe(FRHICommandListImmediate& RHICmdList, FRHIIndexBuffer* IndexBufferRHI, uint32 Offset, uint32 Size, EResourceLockMode LockMode)
-{
-	@autoreleasepool {
-	FMetalIndexBuffer* IndexBuffer = ResourceCast(IndexBufferRHI);
-
-	return (uint8*)IndexBuffer->Lock(true, LockMode, Offset, Size);
-	}
-}
-
-void FMetalDynamicRHI::UnlockIndexBuffer_BottomOfPipe(FRHICommandListImmediate& RHICmdList, FRHIIndexBuffer* IndexBufferRHI)
-{
-	@autoreleasepool {
-	FMetalIndexBuffer* IndexBuffer = ResourceCast(IndexBufferRHI);
-
-	IndexBuffer->Unlock();
-	}
-}
-
 FIndexBufferRHIRef FMetalDynamicRHI::CreateIndexBuffer_RenderThread(class FRHICommandListImmediate& RHICmdList, uint32 Stride, uint32 Size, uint32 InUsage, ERHIAccess InResourceState, FRHIResourceCreateInfo& CreateInfo)
 {
 	@autoreleasepool {
diff --git a/Engine/Source/Runtime/Apple/MetalRHI/Private/MetalStructuredBuffer.cpp b/Engine/Source/Runtime/Apple/MetalRHI/Private/MetalStructuredBuffer.cpp
--- a/Engine/Source/Runtime/Apple/MetalRHI/Private/MetalStructuredBuffer.cpp
+++ b/Engine/Source/Runtime/Apple/MetalRHI/Private/MetalStructuredBuffer.cpp
@@ -18,24 +18,6 @@ FStructuredBufferRHIRef FMetalDynamicRHI::RHICreateStructuredBuffer(uint32 Strid
 	}
 }
 
-void* FMetalDynamicRHI::LockStructuredBuffer_BottomOfPipe(FRHICommandListImmediate& RHICmdList, FRHIStructuredBuffer* StructuredBufferRHI, uint32 Offset, uint32 Size, EResourceLockMode LockMode)
-{
-	@autoreleasepool {
-	FMetalStructuredBuffer* StructuredBuffer = ResourceCast(StructuredBufferRHI);
-	
-	// just return the memory plus the offset
-	return (uint8*)StructuredBuffer->Lock(true, LockMode, Offset, Size);
-	}
-}
-
-void FMetalDynamicRHI::UnlockStructuredBuffer_BottomOfPipe(FRHICommandListImmediate& RHICmdList, FRHIStructuredBuffer* StructuredBufferRHI)
-{
-	@autoreleasepool {
-	FMetalStructuredBuffer* StructuredBuffer = ResourceCast(StructuredBufferRHI);
-	StructuredBuffer->Unlock();
-	}
-}
-
 FStructuredBufferRHIRef FMetalDynamicRHI::CreateStructuredBuffer_RenderThread(class FRHICommandListImmediate& RHICmdList, uint32 Stride, uint32 Size, uint32 InUsage, ERHIAccess InResourceState, FRHIResourceCreateInfo& CreateInfo)
 {
 	@autoreleasepool {
diff --git a/Engine/Source/Runtime/Apple/MetalRHI/Private/MetalVertexBuffer.cpp b/Engine/Source/Runtime/Apple/MetalRHI/Private/MetalVertexBuffer.cpp
--- a/Engine/Source/Runtime/Apple/MetalRHI/Private/MetalVertexBuffer.cpp
+++ b/Engine/Source/Runtime/Apple/MetalRHI/Private/MetalVertexBuffer.cpp
@@ -737,12 +737,12 @@ FVertexBufferRHIRef FMetalDynamicRHI::RHICreateVertexBuffer(uint32 Size, uint32
 		check(Size >= CreateInfo.ResourceArray->GetResourceDataSize());
 
 		// make a buffer usable by CPU
-		void* Buffer = ::RHILockVertexBuffer(VertexBuffer, 0, Size, RLM_WriteOnly);
+		void* Buffer = ::RHILockBuffer(VertexBuffer, 0, Size, RLM_WriteOnly);
 		
 		// copy the contents of the given data into the buffer
 		FMemory::Memcpy(Buffer, CreateInfo.ResourceArray->GetResourceData(), Size);
 		
-		::RHIUnlockVertexBuffer(VertexBuffer);
+		::RHIUnlockBuffer(VertexBuffer);
 
 		// Discard the resource array's contents.
 		CreateInfo.ResourceArray->Discard();
@@ -776,22 +776,22 @@ FVertexBufferRHIRef FMetalDynamicRHI::RHICreateVertexBuffer(uint32 Size, uint32
 	}
 }
 
-void* FMetalDynamicRHI::LockVertexBuffer_BottomOfPipe(FRHICommandListImmediate& RHICmdList, FRHIVertexBuffer* VertexBufferRHI, uint32 Offset, uint32 Size, EResourceLockMode LockMode)
+void* FMetalDynamicRHI::LockBuffer_BottomOfPipe(FRHICommandListImmediate& RHICmdList, FRHIBuffer* BufferRHI, uint32 Offset, uint32 Size, EResourceLockMode LockMode)
 {
 	@autoreleasepool {
-	FMetalVertexBuffer* VertexBuffer = ResourceCast(VertexBufferRHI);
+	FMetalResourceMultiBuffer* Buffer = ResourceCast(BufferRHI);
 
-	// default to vertex buffer memory
-	return (uint8*)VertexBuffer->Lock(true, LockMode, Offset, Size);
+	// default to buffer memory
+	return (uint8*)Buffer->Lock(true, LockMode, Offset, Size);
 	}
 }
 
-void FMetalDynamicRHI::UnlockVertexBuffer_BottomOfPipe(FRHICommandListImmediate& RHICmdList, FRHIVertexBuffer* VertexBufferRHI)
+void FMetalDynamicRHI::UnlockBuffer_BottomOfPipe(FRHICommandListImmediate& RHICmdList, FRHIBuffer* BufferRHI)
 {
 	@autoreleasepool {
-	FMetalVertexBuffer* VertexBuffer = ResourceCast(VertexBufferRHI);
+	FMetalResourceMultiBuffer* Buffer = ResourceCast(BufferRHI);
 
-	VertexBuffer->Unlock();
+	Buffer->Unlock();
 	}
 }
 
diff --git a/Engine/Source/Runtime/D3D12RHI/Private/D3D12Buffer.cpp b/Engine/Source/Runtime/D3D12RHI/Private/D3D12Buffer.cpp
--- a/Engine/Source/Runtime/D3D12RHI/Private/D3D12Buffer.cpp
+++ b/Engine/Source/Runtime/D3D12RHI/Private/D3D12Buffer.cpp
@@ -542,3 +542,30 @@ void FD3D12DynamicRHI::UnlockBuffer(FRHICommandListImmediate* RHICmdList, FD3D12
 
 	LockedData.Reset();
 }
+
+void* FD3D12DynamicRHI::RHILockBuffer(FRHICommandListImmediate& RHICmdList, FRHIBuffer* BufferRHI, uint32 Offset, uint32 Size, EResourceLockMode LockMode)
+{
+	FD3D12Buffer* Buffer = FD3D12DynamicRHI::ResourceCast(BufferRHI);
+	return LockBuffer(&RHICmdList, Buffer, Buffer->GetSize(), Buffer->GetUsage(), Offset, Size, LockMode);
+}
+
+void FD3D12DynamicRHI::RHIUnlockBuffer(FRHICommandListImmediate& RHICmdList, FRHIBuffer* BufferRHI)
+{
+	FD3D12Buffer* Buffer = FD3D12DynamicRHI::ResourceCast(BufferRHI);
+	UnlockBuffer(&RHICmdList, Buffer, Buffer->GetUsage());
+}
+
+void FD3D12DynamicRHI::RHITransferBufferUnderlyingResource(FRHIBuffer* DestBuffer, FRHIBuffer* SrcBuffer)
+{
+	check(DestBuffer);
+	FD3D12Buffer* Dest = ResourceCast(DestBuffer);
+	if (!SrcBuffer)
+	{
+		Dest->ReleaseUnderlyingResource();
+	}
+	else
+	{
+		FD3D12Buffer* Src = ResourceCast(SrcBuffer);
+		Dest->Swap(*Src);
+	}
+}
diff --git a/Engine/Source/Runtime/D3D12RHI/Private/D3D12IndexBuffer.cpp b/Engine/Source/Runtime/D3D12RHI/Private/D3D12IndexBuffer.cpp
--- a/Engine/Source/Runtime/D3D12RHI/Private/D3D12IndexBuffer.cpp
+++ b/Engine/Source/Runtime/D3D12RHI/Private/D3D12IndexBuffer.cpp
@@ -52,33 +52,6 @@ FIndexBufferRHIRef FD3D12DynamicRHI::RHICreateIndexBuffer(uint32 Stride, uint32
 	return Buffer;
 }
 
-void* FD3D12DynamicRHI::RHILockIndexBuffer(FRHICommandListImmediate& RHICmdList, FRHIIndexBuffer* IndexBufferRHI, uint32 Offset, uint32 Size, EResourceLockMode LockMode)
-{
-	FD3D12Buffer* Buffer = FD3D12DynamicRHI::ResourceCast(IndexBufferRHI);
-	return LockBuffer(&RHICmdList, Buffer, Buffer->GetSize(), Buffer->GetUsage(), Offset, Size, LockMode);
-}
-
-void FD3D12DynamicRHI::RHIUnlockIndexBuffer(FRHICommandListImmediate& RHICmdList, FRHIIndexBuffer* IndexBufferRHI)
-{
-	FD3D12Buffer* Buffer = FD3D12DynamicRHI::ResourceCast(IndexBufferRHI);
-	UnlockBuffer(&RHICmdList, Buffer, Buffer->GetUsage());
-}
-
-void FD3D12DynamicRHI::RHITransferBufferUnderlyingResource(FRHIBuffer* DestBuffer, FRHIBuffer* SrcBuffer)
-{
-	check(DestBuffer);
-	FD3D12Buffer* Dest = ResourceCast(DestBuffer);
-	if (!SrcBuffer)
-	{
-		Dest->ReleaseUnderlyingResource();
-	}
-	else
-	{
-		FD3D12Buffer* Src = ResourceCast(SrcBuffer);
-		Dest->Swap(*Src);
-	}
-}
-
 FIndexBufferRHIRef FD3D12DynamicRHI::CreateIndexBuffer_RenderThread(class FRHICommandListImmediate& RHICmdList, uint32 Stride, uint32 Size, uint32 InUsage, ERHIAccess InResourceState, FRHIResourceCreateInfo& CreateInfo)
 {
 	if (CreateInfo.bWithoutNativeResource)
diff --git a/Engine/Source/Runtime/D3D12RHI/Private/D3D12RHIPrivate.h b/Engine/Source/Runtime/D3D12RHI/Private/D3D12RHIPrivate.h
--- a/Engine/Source/Runtime/D3D12RHI/Private/D3D12RHIPrivate.h
+++ b/Engine/Source/Runtime/D3D12RHI/Private/D3D12RHIPrivate.h
@@ -279,16 +279,12 @@ public:
 	virtual FUniformBufferRHIRef RHICreateUniformBuffer(const void* Contents, const FRHIUniformBufferLayout& Layout, EUniformBufferUsage Usage, EUniformBufferValidation Validation) final override;
 	virtual void RHIUpdateUniformBuffer(FRHIUniformBuffer* UniformBufferRHI, const void* Contents) final override;
 	virtual FIndexBufferRHIRef RHICreateIndexBuffer(uint32 Stride, uint32 Size, uint32 InUsage, ERHIAccess InResourceState, FRHIResourceCreateInfo& CreateInfo) final override;
-	virtual void* RHILockIndexBuffer(FRHICommandListImmediate& RHICmdList, FRHIIndexBuffer* IndexBuffer, uint32 Offset, uint32 Size, EResourceLockMode LockMode) final override;
-	virtual void RHIUnlockIndexBuffer(FRHICommandListImmediate& RHICmdList, FRHIIndexBuffer* IndexBuffer) final override;
 	virtual void RHITransferBufferUnderlyingResource(FRHIBuffer* DestBuffer, FRHIBuffer* SrcBuffer) final override;
 	virtual FVertexBufferRHIRef RHICreateVertexBuffer(uint32 Size, uint32 InUsage, ERHIAccess InResourceState, FRHIResourceCreateInfo& CreateInfo) final override;
-	virtual void* RHILockVertexBuffer(FRHICommandListImmediate& RHICmdList, FRHIVertexBuffer* VertexBuffer, uint32 Offset, uint32 SizeRHI, EResourceLockMode LockMode) final override;
-	virtual void RHIUnlockVertexBuffer(FRHICommandListImmediate& RHICmdList, FRHIVertexBuffer* VertexBuffer) final override;
 	virtual void RHICopyVertexBuffer(FRHIVertexBuffer* SourceBuffer, FRHIVertexBuffer* DestBuffer) final override;
 	virtual FStructuredBufferRHIRef RHICreateStructuredBuffer(uint32 Stride, uint32 Size, uint32 InUsage, ERHIAccess InResourceState, FRHIResourceCreateInfo& CreateInfo) final override;
-	virtual void* RHILockStructuredBuffer(FRHICommandListImmediate& RHICmdList, FRHIStructuredBuffer* StructuredBuffer, uint32 Offset, uint32 SizeRHI, EResourceLockMode LockMode) final override;
-	virtual void RHIUnlockStructuredBuffer(FRHICommandListImmediate& RHICmdList, FRHIStructuredBuffer* StructuredBuffer) final override;
+	virtual void* RHILockBuffer(FRHICommandListImmediate& RHICmdList, FRHIBuffer* Buffer, uint32 Offset, uint32 Size, EResourceLockMode LockMode) final override;
+	virtual void RHIUnlockBuffer(FRHICommandListImmediate& RHICmdList, FRHIBuffer* Buffer) final override;
 	virtual FUnorderedAccessViewRHIRef RHICreateUnorderedAccessView(FRHIBuffer* StructuredBuffer, bool bUseUAVCounter, bool bAppendBuffer) final override;
 	virtual FUnorderedAccessViewRHIRef RHICreateUnorderedAccessView(FRHITexture* Texture, uint32 MipLevel) final override;
 	virtual FUnorderedAccessViewRHIRef RHICreateUnorderedAccessView(FRHIBuffer* Buffer, uint8 Format) final override;
diff --git a/Engine/Source/Runtime/D3D12RHI/Private/D3D12StructuredBuffer.cpp b/Engine/Source/Runtime/D3D12RHI/Private/D3D12StructuredBuffer.cpp
--- a/Engine/Source/Runtime/D3D12RHI/Private/D3D12StructuredBuffer.cpp
+++ b/Engine/Source/Runtime/D3D12RHI/Private/D3D12StructuredBuffer.cpp
@@ -68,15 +68,3 @@ FStructuredBufferRHIRef FD3D12DynamicRHI::RHICreateStructuredBuffer(uint32 Strid
 
 	return NewBuffer;
 }
-
-void* FD3D12DynamicRHI::RHILockStructuredBuffer(FRHICommandListImmediate& RHICmdList, FRHIStructuredBuffer* StructuredBufferRHI, uint32 Offset, uint32 Size, EResourceLockMode LockMode)
-{
-	FD3D12Buffer* Buffer = FD3D12DynamicRHI::ResourceCast(StructuredBufferRHI);
-	return LockBuffer(&RHICmdList, Buffer, Buffer->GetSize(), Buffer->GetUsage(), Offset, Size, LockMode);
-}
-
-void FD3D12DynamicRHI::RHIUnlockStructuredBuffer(FRHICommandListImmediate& RHICmdList, FRHIStructuredBuffer* StructuredBufferRHI)
-{
-	FD3D12Buffer* Buffer = FD3D12DynamicRHI::ResourceCast(StructuredBufferRHI);
-	UnlockBuffer(&RHICmdList, Buffer, Buffer->GetUsage());
-}
diff --git a/Engine/Source/Runtime/D3D12RHI/Private/D3D12VertexBuffer.cpp b/Engine/Source/Runtime/D3D12RHI/Private/D3D12VertexBuffer.cpp
--- a/Engine/Source/Runtime/D3D12RHI/Private/D3D12VertexBuffer.cpp
+++ b/Engine/Source/Runtime/D3D12RHI/Private/D3D12VertexBuffer.cpp
@@ -59,18 +59,6 @@ FVertexBufferRHIRef FD3D12DynamicRHI::RHICreateVertexBuffer(uint32 Size, uint32
 	return Buffer;
 }
 
-void* FD3D12DynamicRHI::RHILockVertexBuffer(FRHICommandListImmediate& RHICmdList, FRHIVertexBuffer* VertexBufferRHI, uint32 Offset, uint32 Size, EResourceLockMode LockMode)
-{
-	FD3D12Buffer* Buffer = FD3D12DynamicRHI::ResourceCast(VertexBufferRHI);
-	return LockBuffer(&RHICmdList, Buffer, Buffer->GetSize(), Buffer->GetUsage(), Offset, Size, LockMode);
-}
-
-void FD3D12DynamicRHI::RHIUnlockVertexBuffer(FRHICommandListImmediate& RHICmdList, FRHIVertexBuffer* VertexBufferRHI)
-{
-	FD3D12Buffer* Buffer = FD3D12DynamicRHI::ResourceCast(VertexBufferRHI);
-	UnlockBuffer(&RHICmdList, Buffer, Buffer->GetUsage());
-}
-
 FVertexBufferRHIRef FD3D12DynamicRHI::CreateVertexBuffer_RenderThread(FRHICommandListImmediate& RHICmdList, uint32 Size, uint32 InUsage, ERHIAccess InResourceState, FRHIResourceCreateInfo& CreateInfo)
 {	
 	if (CreateInfo.bWithoutNativeResource)
diff --git a/Engine/Source/Runtime/NullDrv/Public/NullRHI.h b/Engine/Source/Runtime/NullDrv/Public/NullRHI.h
--- a/Engine/Source/Runtime/NullDrv/Public/NullRHI.h
+++ b/Engine/Source/Runtime/NullDrv/Public/NullRHI.h
@@ -133,20 +133,6 @@ public:
 		return new FRHIIndexBuffer(Stride,Size,InUsage); 
 	}
 
-	virtual void* LockIndexBuffer_BottomOfPipe(FRHICommandListImmediate& RHICmdList, FRHIIndexBuffer* IndexBuffer, uint32 Offset, uint32 Size, EResourceLockMode LockMode) final override
-	{ 
-		return GetStaticBuffer(Size);
-	}
-	virtual void UnlockIndexBuffer_BottomOfPipe(FRHICommandListImmediate& RHICmdList, FRHIIndexBuffer* IndexBuffer) final override
-	{
-
-	}
-
-	virtual void RHITransferBufferUnderlyingResource(FRHIBuffer* DestBuffer, FRHIBuffer* SrcBuffer) final override
-	{
-
-	}
-
 	virtual FVertexBufferRHIRef RHICreateVertexBuffer(uint32 Size, uint32 InUsage, ERHIAccess InResourceState, FRHIResourceCreateInfo& CreateInfo) final override
 	{ 
 		if(CreateInfo.ResourceArray) 
@@ -156,16 +142,6 @@ public:
 		return new FRHIVertexBuffer(Size, InUsage, 0); 
 	}
 
-	virtual void* LockVertexBuffer_BottomOfPipe(FRHICommandListImmediate& RHICmdList, FRHIVertexBuffer* VertexBuffer, uint32 Offset, uint32 SizeRHI, EResourceLockMode LockMode) final override
-	{ 
-		return GetStaticBuffer(VertexBuffer->GetSize());
-	}
-
-	virtual void UnlockVertexBuffer_BottomOfPipe(FRHICommandListImmediate& RHICmdList, FRHIVertexBuffer* VertexBuffer) final override
-	{
-
-	}
-
 	virtual void RHICopyVertexBuffer(FRHIVertexBuffer* SourceBuffer, FRHIVertexBuffer* DestBuffer) final override
 	{
 
@@ -180,15 +156,20 @@ public:
 		return new FRHIStructuredBuffer(Stride,Size,InUsage); 
 	}
 
-	virtual void* LockStructuredBuffer_BottomOfPipe(FRHICommandListImmediate& RHICmdList, FRHIStructuredBuffer* StructuredBuffer, uint32 Offset, uint32 SizeRHI, EResourceLockMode LockMode) final override
+	virtual void* LockBuffer_BottomOfPipe(FRHICommandListImmediate& RHICmdList, FRHIBuffer* Buffer, uint32 Offset, uint32 SizeRHI, EResourceLockMode LockMode) final override
 	{ 
-		return GetStaticBuffer(StructuredBuffer->GetSize());
+		return GetStaticBuffer(Buffer->GetSize());
 	}
-	virtual void UnlockStructuredBuffer_BottomOfPipe(FRHICommandListImmediate& RHICmdList, FRHIStructuredBuffer* StructuredBuffer) final override
+
+	virtual void UnlockBuffer_BottomOfPipe(FRHICommandListImmediate& RHICmdList, FRHIBuffer* Buffer) final override
 	{
 
 	}
 
+	virtual void RHITransferBufferUnderlyingResource(FRHIBuffer* DestBuffer, FRHIBuffer* SrcBuffer) final override
+	{
+
+	}
 
 	virtual FUnorderedAccessViewRHIRef RHICreateUnorderedAccessView(FRHIBuffer* StructuredBuffer, bool bUseUAVCounter, bool bAppendBuffer) final override
 	{ 
diff --git a/Engine/Source/Runtime/OpenGLDrv/Private/OpenGLIndexBuffer.cpp b/Engine/Source/Runtime/OpenGLDrv/Private/OpenGLIndexBuffer.cpp
--- a/Engine/Source/Runtime/OpenGLDrv/Private/OpenGLIndexBuffer.cpp
+++ b/Engine/Source/Runtime/OpenGLDrv/Private/OpenGLIndexBuffer.cpp
@@ -39,24 +39,6 @@ FIndexBufferRHIRef FOpenGLDynamicRHI::CreateIndexBuffer_RenderThread(class FRHIC
 	return this->RHICreateIndexBuffer(Stride, Size, InUsage, InResourceState, CreateInfo);
 }
 
-void* FOpenGLDynamicRHI::LockIndexBuffer_BottomOfPipe(FRHICommandListImmediate& RHICmdList, FRHIIndexBuffer* IndexBufferRHI, uint32 Offset, uint32 Size, EResourceLockMode LockMode)
-{
-	RHITHREAD_GLCOMMAND_PROLOGUE();
-	VERIFY_GL_SCOPE();
-	FOpenGLBuffer* IndexBuffer = ResourceCast(IndexBufferRHI);
-	return IndexBuffer->Lock(Offset, Size, LockMode == RLM_ReadOnly, IndexBuffer->IsDynamic());
-	RHITHREAD_GLCOMMAND_EPILOGUE_RETURN(void*);
-}
-
-void FOpenGLDynamicRHI::UnlockIndexBuffer_BottomOfPipe(FRHICommandListImmediate& RHICmdList, FRHIIndexBuffer* IndexBufferRHI)
-{
-	RHITHREAD_GLCOMMAND_PROLOGUE();
-	VERIFY_GL_SCOPE();
-	FOpenGLBuffer* IndexBuffer = ResourceCast(IndexBufferRHI);
-	IndexBuffer->Unlock();
-	RHITHREAD_GLCOMMAND_EPILOGUE();
-}
-
 void FOpenGLDynamicRHI::RHITransferBufferUnderlyingResource(FRHIBuffer* DestBuffer, FRHIBuffer* SrcBuffer)
 {
 	VERIFY_GL_SCOPE();
diff --git a/Engine/Source/Runtime/OpenGLDrv/Private/OpenGLStructuredBuffer.cpp b/Engine/Source/Runtime/OpenGLDrv/Private/OpenGLStructuredBuffer.cpp
--- a/Engine/Source/Runtime/OpenGLDrv/Private/OpenGLStructuredBuffer.cpp
+++ b/Engine/Source/Runtime/OpenGLDrv/Private/OpenGLStructuredBuffer.cpp
@@ -24,17 +24,3 @@ FStructuredBufferRHIRef FOpenGLDynamicRHI::RHICreateStructuredBuffer(uint32 Stri
 	TRefCountPtr<FOpenGLBuffer> StructuredBuffer = new FOpenGLBuffer(GL_ARRAY_BUFFER, Stride, Size, InUsage | BUF_StructuredBuffer, Data);
 	return StructuredBuffer.GetReference();
 }
-
-void* FOpenGLDynamicRHI::LockStructuredBuffer_BottomOfPipe(FRHICommandListImmediate& RHICmdList, FRHIStructuredBuffer* StructuredBufferRHI, uint32 Offset, uint32 Size, EResourceLockMode LockMode)
-{
-	VERIFY_GL_SCOPE();
-	FOpenGLBuffer* StructuredBuffer = ResourceCast(StructuredBufferRHI);
-	return StructuredBuffer->Lock(Offset, Size, LockMode == RLM_ReadOnly, StructuredBuffer->IsDynamic());
-}
-
-void FOpenGLDynamicRHI::UnlockStructuredBuffer_BottomOfPipe(FRHICommandListImmediate& RHICmdList, FRHIStructuredBuffer* StructuredBufferRHI)
-{
-	VERIFY_GL_SCOPE();
-	FOpenGLBuffer* StructuredBuffer = ResourceCast(StructuredBufferRHI);
-	StructuredBuffer->Unlock();
-}
diff --git a/Engine/Source/Runtime/OpenGLDrv/Private/OpenGLVertexBuffer.cpp b/Engine/Source/Runtime/OpenGLDrv/Private/OpenGLVertexBuffer.cpp
--- a/Engine/Source/Runtime/OpenGLDrv/Private/OpenGLVertexBuffer.cpp
+++ b/Engine/Source/Runtime/OpenGLDrv/Private/OpenGLVertexBuffer.cpp
@@ -202,45 +202,45 @@ FVertexBufferRHIRef FOpenGLDynamicRHI::RHICreateVertexBuffer(uint32 Size, uint32
 	return VertexBuffer.GetReference();
 }
 
-void* FOpenGLDynamicRHI::LockVertexBuffer_BottomOfPipe(FRHICommandListImmediate& RHICmdList, FRHIVertexBuffer* VertexBufferRHI, uint32 Offset, uint32 Size, EResourceLockMode LockMode)
+void* FOpenGLDynamicRHI::LockBuffer_BottomOfPipe(FRHICommandListImmediate& RHICmdList, FRHIBuffer* BufferRHI, uint32 Offset, uint32 Size, EResourceLockMode LockMode)
 {
 	check(Size > 0);
 	RHITHREAD_GLCOMMAND_PROLOGUE();
 
 	VERIFY_GL_SCOPE();
-	FOpenGLBuffer* VertexBuffer = ResourceCast(VertexBufferRHI);
-	if( !(FOpenGL::SupportsVertexAttribBinding() && OpenGLConsoleVariables::bUseVAB) && VertexBuffer->GetUsage() & BUF_ZeroStride )
+	FOpenGLBuffer* Buffer = ResourceCast(BufferRHI);
+	if( !(FOpenGL::SupportsVertexAttribBinding() && OpenGLConsoleVariables::bUseVAB) && Buffer->GetUsage() & BUF_ZeroStride )
 	{
-		check( Offset + Size <= VertexBuffer->GetSize() );
+		check( Offset + Size <= Buffer->GetSize() );
 		// We assume we are only using the first elements from the VB, so we can later on read this memory and create an expanded version of this zero stride buffer
 		check(Offset == 0);
-		return (void*)( (uint8*)VertexBuffer->GetZeroStrideBuffer() + Offset );
+		return (void*)( (uint8*)Buffer->GetZeroStrideBuffer() + Offset );
 	}
 	else
 	{
-		if (VertexBuffer->IsDynamic() && LockMode == EResourceLockMode::RLM_WriteOnly)
+		if (Buffer->IsDynamic() && LockMode == EResourceLockMode::RLM_WriteOnly)
 		{
-			void *Staging = GetAllocation(VertexBuffer, Size, Offset);
+			void *Staging = GetAllocation(Buffer, Size, Offset);
 			if (Staging)
 			{
 				return Staging;
 			}
 		}
-		return (void*)VertexBuffer->Lock(Offset, Size, LockMode == EResourceLockMode::RLM_ReadOnly, VertexBuffer->IsDynamic());
+		return (void*)Buffer->Lock(Offset, Size, LockMode == EResourceLockMode::RLM_ReadOnly, Buffer->IsDynamic());
 	}
 	RHITHREAD_GLCOMMAND_EPILOGUE_RETURN(void*);
 }
 
-void FOpenGLDynamicRHI::UnlockVertexBuffer_BottomOfPipe(FRHICommandListImmediate& RHICmdList, FRHIVertexBuffer* VertexBufferRHI)
+void FOpenGLDynamicRHI::UnlockBuffer_BottomOfPipe(FRHICommandListImmediate& RHICmdList, FRHIBuffer* BufferRHI)
 {
 	RHITHREAD_GLCOMMAND_PROLOGUE();
 	VERIFY_GL_SCOPE();
-	FOpenGLBuffer* VertexBuffer = ResourceCast(VertexBufferRHI);
-	if( (FOpenGL::SupportsVertexAttribBinding() && OpenGLConsoleVariables::bUseVAB) || !( VertexBuffer->GetUsage() & BUF_ZeroStride ) )
+	FOpenGLBuffer* Buffer = ResourceCast(BufferRHI);
+	if( (FOpenGL::SupportsVertexAttribBinding() && OpenGLConsoleVariables::bUseVAB) || !( Buffer->GetUsage() & BUF_ZeroStride ) )
 	{
-		if (!RetireAllocation(VertexBuffer))
+		if (!RetireAllocation(Buffer))
 		{
-			VertexBuffer->Unlock();
+			Buffer->Unlock();
 		}
 	}
 	RHITHREAD_GLCOMMAND_EPILOGUE();
diff --git a/Engine/Source/Runtime/OpenGLDrv/Public/OpenGLDrv.h b/Engine/Source/Runtime/OpenGLDrv/Public/OpenGLDrv.h
--- a/Engine/Source/Runtime/OpenGLDrv/Public/OpenGLDrv.h
+++ b/Engine/Source/Runtime/OpenGLDrv/Public/OpenGLDrv.h
@@ -347,16 +347,12 @@ public:
 	virtual FUniformBufferRHIRef RHICreateUniformBuffer(const void* Contents, const FRHIUniformBufferLayout& Layout, EUniformBufferUsage Usage, EUniformBufferValidation Validation) final override;
 	virtual void RHIUpdateUniformBuffer(FRHIUniformBuffer* UniformBufferRHI, const void* Contents) final override;
 	virtual FIndexBufferRHIRef RHICreateIndexBuffer(uint32 Stride, uint32 Size, uint32 InUsage, ERHIAccess InResourceState, FRHIResourceCreateInfo& CreateInfo) final override;
-	virtual void* LockIndexBuffer_BottomOfPipe(FRHICommandListImmediate& RHICmdList, FRHIIndexBuffer* IndexBuffer, uint32 Offset, uint32 Size, EResourceLockMode LockMode) final override;
-	virtual void UnlockIndexBuffer_BottomOfPipe(FRHICommandListImmediate& RHICmdList, FRHIIndexBuffer* IndexBuffer) final override;
 	virtual FVertexBufferRHIRef RHICreateVertexBuffer(uint32 Size, uint32 InUsage, ERHIAccess InResourceState, FRHIResourceCreateInfo& CreateInfo) final override;
-	virtual void* LockVertexBuffer_BottomOfPipe(FRHICommandListImmediate& RHICmdList, FRHIVertexBuffer* VertexBuffer, uint32 Offset, uint32 SizeRHI, EResourceLockMode LockMode) final override;
-	virtual void UnlockVertexBuffer_BottomOfPipe(FRHICommandListImmediate& RHICmdList, FRHIVertexBuffer* VertexBuffer) final override;
 	virtual void RHICopyVertexBuffer(FRHIVertexBuffer* SourceBuffer, FRHIVertexBuffer* DestBuffer) final override;
 	virtual void RHITransferBufferUnderlyingResource(FRHIBuffer* DestBuffer, FRHIBuffer* SrcBuffer) final override;
 	virtual FStructuredBufferRHIRef RHICreateStructuredBuffer(uint32 Stride, uint32 Size, uint32 InUsage, ERHIAccess InResourceState, FRHIResourceCreateInfo& CreateInfo) final override;
-	virtual void* LockStructuredBuffer_BottomOfPipe(FRHICommandListImmediate& RHICmdList, FRHIStructuredBuffer* StructuredBuffer, uint32 Offset, uint32 SizeRHI, EResourceLockMode LockMode) final override;
-	virtual void UnlockStructuredBuffer_BottomOfPipe(FRHICommandListImmediate& RHICmdList, FRHIStructuredBuffer* StructuredBuffer) final override;
+	virtual void* LockBuffer_BottomOfPipe(FRHICommandListImmediate& RHICmdList, FRHIBuffer* Buffer, uint32 Offset, uint32 Size, EResourceLockMode LockMode) final override;
+	virtual void UnlockBuffer_BottomOfPipe(FRHICommandListImmediate& RHICmdList, FRHIBuffer* Buffer) final override;
 	virtual FUnorderedAccessViewRHIRef RHICreateUnorderedAccessView(FRHIBuffer* Buffer, bool bUseUAVCounter, bool bAppendBuffer) final override;
 	virtual FUnorderedAccessViewRHIRef RHICreateUnorderedAccessView(FRHITexture* Texture, uint32 MipLevel) final override;
 	virtual FUnorderedAccessViewRHIRef RHICreateUnorderedAccessView(FRHIBuffer* Buffer, uint8 Format) final override;
diff --git a/Engine/Source/Runtime/RHI/Private/RHICommandList.cpp b/Engine/Source/Runtime/RHI/Private/RHICommandList.cpp
--- a/Engine/Source/Runtime/RHI/Private/RHICommandList.cpp
+++ b/Engine/Source/Runtime/RHI/Private/RHICommandList.cpp
@@ -1993,7 +1993,7 @@ void FRHICommandListBase::operator delete(void *RawMemory)
 FVertexBufferRHIRef FDynamicRHI::CreateAndLockVertexBuffer_RenderThread(class FRHICommandListImmediate& RHICmdList, uint32 Size, uint32 InUsage, ERHIAccess InResourceState, FRHIResourceCreateInfo& CreateInfo, void*& OutDataBuffer)
 {
 	FVertexBufferRHIRef VertexBuffer = CreateVertexBuffer_RenderThread(RHICmdList, Size, InUsage, InResourceState, CreateInfo);
-	OutDataBuffer = RHILockVertexBuffer(RHICmdList, VertexBuffer, 0, Size, RLM_WriteOnly);
+	OutDataBuffer = RHILockBuffer(RHICmdList, VertexBuffer, 0, Size, RLM_WriteOnly);
 
 	return VertexBuffer;
 }
@@ -2001,7 +2001,7 @@ FVertexBufferRHIRef FDynamicRHI::CreateAndLockVertexBuffer_RenderThread(class FR
 FIndexBufferRHIRef FDynamicRHI::CreateAndLockIndexBuffer_RenderThread(class FRHICommandListImmediate& RHICmdList, uint32 Stride, uint32 Size, uint32 InUsage, ERHIAccess InResourceState, FRHIResourceCreateInfo& CreateInfo, void*& OutDataBuffer)
 {
 	FIndexBufferRHIRef IndexBuffer = CreateIndexBuffer_RenderThread(RHICmdList, Stride, Size, InUsage, InResourceState, CreateInfo);
-	OutDataBuffer = RHILockIndexBuffer(RHICmdList, IndexBuffer, 0, Size, RLM_WriteOnly);
+	OutDataBuffer = RHILockBuffer(RHICmdList, IndexBuffer, 0, Size, RLM_WriteOnly);
 	return IndexBuffer;
 }
 
@@ -2050,9 +2050,9 @@ FShaderResourceViewRHIRef FDynamicRHI::CreateShaderResourceView_RenderThread(cla
 
 static FLockTracker GLockTracker;
 
-void* FDynamicRHI::RHILockVertexBuffer(class FRHICommandListImmediate& RHICmdList, FRHIVertexBuffer* VertexBuffer, uint32 Offset, uint32 SizeRHI, EResourceLockMode LockMode)
+void* FDynamicRHI::RHILockBuffer(class FRHICommandListImmediate& RHICmdList, FRHIBuffer* Buffer, uint32 Offset, uint32 SizeRHI, EResourceLockMode LockMode)
 {
-	QUICK_SCOPE_CYCLE_COUNTER(STAT_FDynamicRHI_LockVertexBuffer);
+	QUICK_SCOPE_CYCLE_COUNTER(STAT_FDynamicRHI_LockBuffer);
 
 	void* Result;
 	if (RHICmdList.IsTopOfPipe())
@@ -2060,65 +2060,65 @@ void* FDynamicRHI::RHILockVertexBuffer(class FRHICommandListImmediate& RHICmdLis
 		bool bBuffer = CVarRHICmdBufferWriteLocks.GetValueOnRenderThread() > 0;
 		if (!bBuffer || LockMode != RLM_WriteOnly)
 		{
-			QUICK_SCOPE_CYCLE_COUNTER(STAT_RHIMETHOD_LockVertexBuffer_FlushAndLock);
-			CSV_SCOPED_TIMING_STAT(RHITFlushes, LockVertexBuffer_BottomOfPipe);
+			QUICK_SCOPE_CYCLE_COUNTER(STAT_RHIMETHOD_LockBuffer_FlushAndLock);
+			CSV_SCOPED_TIMING_STAT(RHITFlushes, LockBuffer_BottomOfPipe);
 
 			FRHICommandListScopedFlushAndExecute Flush(RHICmdList);
-			Result = GDynamicRHI->LockVertexBuffer_BottomOfPipe(RHICmdList, VertexBuffer, Offset, SizeRHI, LockMode);
+			Result = GDynamicRHI->LockBuffer_BottomOfPipe(RHICmdList, Buffer, Offset, SizeRHI, LockMode);
 		}
 		else
 		{
-			QUICK_SCOPE_CYCLE_COUNTER(STAT_RHIMETHOD_LockVertexBuffer_Malloc);
+			QUICK_SCOPE_CYCLE_COUNTER(STAT_RHIMETHOD_LockBuffer_Malloc);
 			Result = FMemory::Malloc(SizeRHI, 16);
 		}
 
 		// Only use the lock tracker at the top of the pipe. There's no need to track locks
 		// at the bottom of the pipe, and doing so would require a critical section.
-		GLockTracker.Lock(VertexBuffer, Result, Offset, SizeRHI, LockMode);
+		GLockTracker.Lock(Buffer, Result, Offset, SizeRHI, LockMode);
 	}
 	else
 	{
-		Result = GDynamicRHI->LockVertexBuffer_BottomOfPipe(RHICmdList, VertexBuffer, Offset, SizeRHI, LockMode);
+		Result = GDynamicRHI->LockBuffer_BottomOfPipe(RHICmdList, Buffer, Offset, SizeRHI, LockMode);
 	}
 
 	check(Result);
 	return Result;
 }
 
-void FDynamicRHI::RHIUnlockVertexBuffer(class FRHICommandListImmediate& RHICmdList, FRHIVertexBuffer* VertexBuffer)
+void FDynamicRHI::RHIUnlockBuffer(class FRHICommandListImmediate& RHICmdList, FRHIBuffer* Buffer)
 {
-	QUICK_SCOPE_CYCLE_COUNTER(STAT_FDynamicRHI_UnlockVertexBuffer_RenderThread);
+	QUICK_SCOPE_CYCLE_COUNTER(STAT_FDynamicRHI_UnlockBuffer_RenderThread);
 
 	if (RHICmdList.IsTopOfPipe())
 	{
-		FLockTracker::FLockParams Params = GLockTracker.Unlock(VertexBuffer);
+		FLockTracker::FLockParams Params = GLockTracker.Unlock(Buffer);
 
 		bool bBuffer = CVarRHICmdBufferWriteLocks.GetValueOnRenderThread() > 0;
 		if (!bBuffer || Params.LockMode != RLM_WriteOnly)
 		{
-			QUICK_SCOPE_CYCLE_COUNTER(STAT_RHIMETHOD_UnlockVertexBuffer_FlushAndUnlock);
-			CSV_SCOPED_TIMING_STAT(RHITFlushes, UnlockVertexBuffer_BottomOfPipe);
+			QUICK_SCOPE_CYCLE_COUNTER(STAT_RHIMETHOD_UnlockBuffer_FlushAndUnlock);
+			CSV_SCOPED_TIMING_STAT(RHITFlushes, UnlockBuffer_BottomOfPipe);
 
 			FRHICommandListScopedFlushAndExecute Flush(RHICmdList);
-			GDynamicRHI->UnlockVertexBuffer_BottomOfPipe(RHICmdList, VertexBuffer);
+			GDynamicRHI->UnlockBuffer_BottomOfPipe(RHICmdList, Buffer);
 			GLockTracker.TotalMemoryOutstanding = 0;
 		}
 		else
 		{
-			RHICmdList.EnqueueLambda([VertexBuffer, Params](FRHICommandListImmediate& InRHICmdList)
+			RHICmdList.EnqueueLambda([Buffer, Params](FRHICommandListImmediate& InRHICmdList)
 			{
-				QUICK_SCOPE_CYCLE_COUNTER(STAT_FRHICommandUpdateVertexBuffer_Execute);
-				void* Data = GDynamicRHI->LockVertexBuffer_BottomOfPipe(InRHICmdList, VertexBuffer, Params.Offset, Params.BufferSize, RLM_WriteOnly);
+				QUICK_SCOPE_CYCLE_COUNTER(STAT_FRHICommandUpdateBuffer_Execute);
+				void* Data = GDynamicRHI->LockBuffer_BottomOfPipe(InRHICmdList, Buffer, Params.Offset, Params.BufferSize, RLM_WriteOnly);
 				FMemory::Memcpy(Data, Params.Buffer, Params.BufferSize);
 				FMemory::Free(Params.Buffer);
-				GDynamicRHI->UnlockVertexBuffer_BottomOfPipe(InRHICmdList, VertexBuffer);
+				GDynamicRHI->UnlockBuffer_BottomOfPipe(InRHICmdList, Buffer);
 			});
 			RHICmdList.RHIThreadFence(true);
 		}
 
 		if (GLockTracker.TotalMemoryOutstanding > uint32(CVarRHICmdMaxOutstandingMemoryBeforeFlush.GetValueOnRenderThread()) * 1024u)
 		{
-			QUICK_SCOPE_CYCLE_COUNTER(STAT_RHIMETHOD_UnlockVertexBuffer_FlushForMem);
+			QUICK_SCOPE_CYCLE_COUNTER(STAT_RHIMETHOD_UnlockBuffer_FlushForMem);
 			// we could be loading a level or something, lets get this stuff going
 			RHICmdList.ImmediateFlush(EImmediateFlushType::DispatchToRHIThread); 
 			GLockTracker.TotalMemoryOutstanding = 0;
@@ -2126,167 +2126,7 @@ void FDynamicRHI::RHIUnlockVertexBuffer(class FRHICommandListImmediate& RHICmdLi
 	}
 	else
 	{
-		GDynamicRHI->UnlockVertexBuffer_BottomOfPipe(RHICmdList, VertexBuffer);
-	}
-}
-
-void* FDynamicRHI::RHILockIndexBuffer(class FRHICommandListImmediate& RHICmdList, FRHIIndexBuffer* IndexBuffer, uint32 Offset, uint32 SizeRHI, EResourceLockMode LockMode)
-{
-	QUICK_SCOPE_CYCLE_COUNTER(STAT_FDynamicRHI_LockIndexBuffer);
-
-	void* Result;
-	if (RHICmdList.IsTopOfPipe())
-	{
-		bool bBuffer = CVarRHICmdBufferWriteLocks.GetValueOnRenderThread() > 0;
-		if (!bBuffer || LockMode != RLM_WriteOnly)
-		{
-			QUICK_SCOPE_CYCLE_COUNTER(STAT_RHIMETHOD_LockIndexBuffer_FlushAndLock);
-			CSV_SCOPED_TIMING_STAT(RHITFlushes, LockIndexBuffer_BottomOfPipe);
-
-			FRHICommandListScopedFlushAndExecute Flush(RHICmdList);
-			Result = GDynamicRHI->LockIndexBuffer_BottomOfPipe(RHICmdList, IndexBuffer, Offset, SizeRHI, LockMode);
-		}
-		else
-		{
-			QUICK_SCOPE_CYCLE_COUNTER(STAT_RHIMETHOD_LockIndexBuffer_Malloc);
-			Result = FMemory::Malloc(SizeRHI, 16);
-		}
-
-		// Only use the lock tracker at the top of the pipe. There's no need to track locks
-		// at the bottom of the pipe, and doing so would require a critical section.
-		GLockTracker.Lock(IndexBuffer, Result, Offset, SizeRHI, LockMode);
-	}
-	else
-	{
-		Result = GDynamicRHI->LockIndexBuffer_BottomOfPipe(RHICmdList, IndexBuffer, Offset, SizeRHI, LockMode);
-	}
-
-	check(Result);
-	return Result;
-}
-
-void FDynamicRHI::RHIUnlockIndexBuffer(class FRHICommandListImmediate& RHICmdList, FRHIIndexBuffer* IndexBuffer)
-{
-	QUICK_SCOPE_CYCLE_COUNTER(STAT_FDynamicRHI_UnlockIndexBuffer_RenderThread);
-
-	if (RHICmdList.IsTopOfPipe())
-	{
-		FLockTracker::FLockParams Params = GLockTracker.Unlock(IndexBuffer);
-
-		bool bBuffer = CVarRHICmdBufferWriteLocks.GetValueOnRenderThread() > 0;
-		if (!bBuffer || Params.LockMode != RLM_WriteOnly)
-		{
-			QUICK_SCOPE_CYCLE_COUNTER(STAT_RHIMETHOD_UnlockIndexBuffer_FlushAndUnlock);
-			CSV_SCOPED_TIMING_STAT(RHITFlushes, UnlockIndexBuffer_BottomOfPipe);
-
-			FRHICommandListScopedFlushAndExecute Flush(RHICmdList);
-			GDynamicRHI->UnlockIndexBuffer_BottomOfPipe(RHICmdList, IndexBuffer);
-			GLockTracker.TotalMemoryOutstanding = 0;
-		}
-		else
-		{
-			RHICmdList.EnqueueLambda([IndexBuffer, Params](FRHICommandListImmediate& InRHICmdList)
-			{
-				QUICK_SCOPE_CYCLE_COUNTER(STAT_FRHICommandUpdateIndexBuffer_Execute);
-				void* Data = GDynamicRHI->LockIndexBuffer_BottomOfPipe(InRHICmdList, IndexBuffer, Params.Offset, Params.BufferSize, RLM_WriteOnly);
-				FMemory::Memcpy(Data, Params.Buffer, Params.BufferSize);
-				FMemory::Free(Params.Buffer);
-				GDynamicRHI->UnlockIndexBuffer_BottomOfPipe(InRHICmdList, IndexBuffer);
-			});
-			RHICmdList.RHIThreadFence(true);
-		}
-
-		if (GLockTracker.TotalMemoryOutstanding > 256 * 1024)
-		{
-			QUICK_SCOPE_CYCLE_COUNTER(STAT_RHIMETHOD_UnlockIndexBuffer_FlushForMem);
-			// we could be loading a level or something, lets get this stuff going
-			RHICmdList.ImmediateFlush(EImmediateFlushType::DispatchToRHIThread); 
-			GLockTracker.TotalMemoryOutstanding = 0;
-		}
-	}
-	else
-	{
-		GDynamicRHI->UnlockIndexBuffer_BottomOfPipe(RHICmdList, IndexBuffer);
-	}
-}
-
-void* FDynamicRHI::RHILockStructuredBuffer(class FRHICommandListImmediate& RHICmdList, FRHIStructuredBuffer* StructuredBuffer, uint32 Offset, uint32 SizeRHI, EResourceLockMode LockMode)
-{
-	QUICK_SCOPE_CYCLE_COUNTER(STAT_FDynamicRHI_LockStructuredBuffer);
-
-	void* Result;
-	if (RHICmdList.IsTopOfPipe())
-	{
-		bool bBuffer = CVarRHICmdBufferWriteLocks.GetValueOnRenderThread() > 0;
-		if (!bBuffer || LockMode != RLM_WriteOnly)
-		{
-			QUICK_SCOPE_CYCLE_COUNTER(STAT_RHIMETHOD_LockStructuredBuffer_FlushAndLock);
-			CSV_SCOPED_TIMING_STAT(RHITFlushes, LockStructuredBuffer_RenderThread);
-
-			FRHICommandListScopedFlushAndExecute Flush(RHICmdList);
-			Result = GDynamicRHI->LockStructuredBuffer_BottomOfPipe(RHICmdList, StructuredBuffer, Offset, SizeRHI, LockMode);
-		}
-		else
-		{
-			QUICK_SCOPE_CYCLE_COUNTER(STAT_RHIMETHOD_LockStructuredBuffer_Malloc);
-			Result = FMemory::Malloc(SizeRHI, 16);
-		}
-		
-		// Only use the lock tracker at the top of the pipe. There's no need to track locks
-		// at the bottom of the pipe, and doing so would require a critical section.
-		GLockTracker.Lock(StructuredBuffer, Result, Offset, SizeRHI, LockMode);
-	}
-	else
-	{
-		Result = GDynamicRHI->LockStructuredBuffer_BottomOfPipe(RHICmdList, StructuredBuffer, Offset, SizeRHI, LockMode);
-	}
-
-	check(Result);
-	return Result;
-}
-
-void FDynamicRHI::RHIUnlockStructuredBuffer(class FRHICommandListImmediate& RHICmdList, FRHIStructuredBuffer* StructuredBuffer)
-{
-	QUICK_SCOPE_CYCLE_COUNTER(STAT_FDynamicRHI_UnlockStructuredBuffer_RenderThread);
-
-	if (RHICmdList.IsTopOfPipe())
-	{
-		FLockTracker::FLockParams Params = GLockTracker.Unlock(StructuredBuffer);
-
-		bool bBuffer = CVarRHICmdBufferWriteLocks.GetValueOnRenderThread() > 0;
-		if (!bBuffer || Params.LockMode != RLM_WriteOnly)
-		{
-			QUICK_SCOPE_CYCLE_COUNTER(STAT_RHIMETHOD_UnlockStructuredBuffer_FlushAndUnlock);
-			CSV_SCOPED_TIMING_STAT(RHITFlushes, UnlockStructuredBuffer_BottomOfPipe);
-
-			FRHICommandListScopedFlushAndExecute Flush(RHICmdList);
-			GDynamicRHI->UnlockStructuredBuffer_BottomOfPipe(RHICmdList, StructuredBuffer);
-			GLockTracker.TotalMemoryOutstanding = 0;
-		}
-		else
-		{
-			RHICmdList.EnqueueLambda([StructuredBuffer, Params](FRHICommandListImmediate& InRHICmdList)
-			{
-				QUICK_SCOPE_CYCLE_COUNTER(STAT_FRHICommandUpdateStructuredBuffer_Execute);
-				void* Data = GDynamicRHI->LockStructuredBuffer_BottomOfPipe(InRHICmdList, StructuredBuffer, Params.Offset, Params.BufferSize, RLM_WriteOnly);
-				FMemory::Memcpy(Data, Params.Buffer, Params.BufferSize);
-				FMemory::Free(Params.Buffer);
-				GDynamicRHI->UnlockStructuredBuffer_BottomOfPipe(InRHICmdList, StructuredBuffer);
-			});
-			RHICmdList.RHIThreadFence(true);
-		}
-
-		if (GLockTracker.TotalMemoryOutstanding > uint32(CVarRHICmdMaxOutstandingMemoryBeforeFlush.GetValueOnRenderThread()) * 1024u)
-		{
-			QUICK_SCOPE_CYCLE_COUNTER(STAT_RHIMETHOD_UnlockStructuredBuffer_FlushForMem);
-			// we could be loading a level or something, lets get this stuff going
-			RHICmdList.ImmediateFlush(EImmediateFlushType::DispatchToRHIThread); 
-			GLockTracker.TotalMemoryOutstanding = 0;
-		}
-	}
-	else
-	{
-		GDynamicRHI->UnlockStructuredBuffer_BottomOfPipe(RHICmdList, StructuredBuffer);
+		GDynamicRHI->UnlockBuffer_BottomOfPipe(RHICmdList, Buffer);
 	}
 }
 
diff --git a/Engine/Source/Runtime/RHI/Private/RHIGPUReadback.cpp b/Engine/Source/Runtime/RHI/Private/RHIGPUReadback.cpp
--- a/Engine/Source/Runtime/RHI/Private/RHIGPUReadback.cpp
+++ b/Engine/Source/Runtime/RHI/Private/RHIGPUReadback.cpp
@@ -46,13 +46,13 @@ void* FGenericRHIStagingBuffer::Lock(uint32 InOffset, uint32 NumBytes)
 	check(ShadowBuffer);
 	check(!bIsLocked);
 	bIsLocked = true;
-	return reinterpret_cast<void*>(reinterpret_cast<uint8*>(RHILockVertexBuffer(ShadowBuffer, InOffset, NumBytes, RLM_ReadOnly)) + Offset);
+	return reinterpret_cast<void*>(reinterpret_cast<uint8*>(RHILockBuffer(ShadowBuffer, InOffset, NumBytes, RLM_ReadOnly)) + Offset);
 }
 
 void FGenericRHIStagingBuffer::Unlock()
 {
 	check(bIsLocked);
-	RHIUnlockVertexBuffer(ShadowBuffer);
+	RHIUnlockBuffer(ShadowBuffer);
 	bIsLocked = false;
 }
 
diff --git a/Engine/Source/Runtime/RHI/Public/DynamicRHI.h b/Engine/Source/Runtime/RHI/Public/DynamicRHI.h
--- a/Engine/Source/Runtime/RHI/Public/DynamicRHI.h
+++ b/Engine/Source/Runtime/RHI/Public/DynamicRHI.h
@@ -523,24 +523,12 @@ public:
 	// FlushType: Wait RHI Thread
 	virtual FIndexBufferRHIRef RHICreateIndexBuffer(uint32 Stride, uint32 Size, uint32 InUsage, ERHIAccess InResourceState, FRHIResourceCreateInfo& CreateInfo) = 0;
 
-	// FlushType: Flush RHI Thread
-	virtual void* RHILockIndexBuffer(FRHICommandListImmediate& RHICmdList, FRHIIndexBuffer* IndexBuffer, uint32 Offset, uint32 Size, EResourceLockMode LockMode);
-
-	// FlushType: Flush RHI Thread
-	virtual void RHIUnlockIndexBuffer(FRHICommandListImmediate& RHICmdList, FRHIIndexBuffer* IndexBuffer);
-
 	/**
 	* @param ResourceArray - An optional pointer to a resource array containing the resource's data.
 	*/
 	// FlushType: Wait RHI Thread
 	virtual FVertexBufferRHIRef RHICreateVertexBuffer(uint32 Size, uint32 InUsage, ERHIAccess InResourceState, FRHIResourceCreateInfo& CreateInfo) = 0;
 
-	// FlushType: Flush RHI Thread
-	virtual void* RHILockVertexBuffer(FRHICommandListImmediate& RHICmdList, FRHIVertexBuffer* VertexBuffer, uint32 Offset, uint32 SizeRHI, EResourceLockMode LockMode);
-
-	// FlushType: Flush RHI Thread
-	virtual void RHIUnlockVertexBuffer(FRHICommandListImmediate& RHICmdList, FRHIVertexBuffer* VertexBuffer);
-
 	/** Copies the contents of one vertex buffer to another vertex buffer.  They must have identical sizes. */
 	// FlushType: Flush Immediate (seems dangerous)
 	virtual void RHICopyVertexBuffer(FRHIVertexBuffer* SourceBuffer, FRHIVertexBuffer* DestBuffer) = 0;
@@ -559,10 +547,10 @@ public:
 	virtual FStructuredBufferRHIRef RHICreateStructuredBuffer(uint32 Stride, uint32 Size, uint32 InUsage, ERHIAccess InResourceState, FRHIResourceCreateInfo& CreateInfo) = 0;
 
 	// FlushType: Flush RHI Thread
-	virtual void* RHILockStructuredBuffer(FRHICommandListImmediate& RHICmdList, FRHIStructuredBuffer* StructuredBuffer, uint32 Offset, uint32 SizeRHI, EResourceLockMode LockMode);
+	virtual void* RHILockBuffer(FRHICommandListImmediate& RHICmdList, FRHIBuffer* Buffer, uint32 Offset, uint32 Size, EResourceLockMode LockMode);
 
 	// FlushType: Flush RHI Thread
-	virtual void RHIUnlockStructuredBuffer(FRHICommandListImmediate& RHICmdList, FRHIStructuredBuffer* StructuredBuffer);
+	virtual void RHIUnlockBuffer(FRHICommandListImmediate& RHICmdList, FRHIBuffer* Buffer);
 
 	/** Creates an unordered access view of the given structured buffer. */
 	// FlushType: Wait RHI Thread
@@ -1338,38 +1326,16 @@ public:
 	virtual void RHIReadSurfaceFloatData_RenderThread(class FRHICommandListImmediate& RHICmdList, FRHITexture* Texture, FIntRect Rect, TArray<FFloat16Color>& OutData, FReadSurfaceDataFlags Flags);
 
 	// Buffer Lock/Unlock
-	virtual void* LockStructuredBuffer_BottomOfPipe(class FRHICommandListImmediate& RHICmdList, FRHIStructuredBuffer* StructuredBuffer, uint32 Offset, uint32 SizeRHI, EResourceLockMode LockMode)
-	{
-		// Either this function or RHILockStructuredBuffer must be implemented by the platform RHI.
-		checkNoEntry();
-		return nullptr;
-	}
-	virtual void* LockVertexBuffer_BottomOfPipe(class FRHICommandListImmediate& RHICmdList, FRHIVertexBuffer* VertexBuffer, uint32 Offset, uint32 SizeRHI, EResourceLockMode LockMode)
-	{
-		// Either this function or RHILockVertexBuffer must be implemented by the platform RHI.
-		checkNoEntry();
-		return nullptr;
-	}
-	virtual void* LockIndexBuffer_BottomOfPipe(class FRHICommandListImmediate& RHICmdList, FRHIIndexBuffer* IndexBuffer, uint32 Offset, uint32 SizeRHI, EResourceLockMode LockMode)
+	virtual void* LockBuffer_BottomOfPipe(class FRHICommandListImmediate& RHICmdList, FRHIBuffer* Buffer, uint32 Offset, uint32 SizeRHI, EResourceLockMode LockMode)
 	{
-		// Either this function or RHILockIndexBuffer must be implemented by the platform RHI.
+		// Either this function or RHILockBuffer must be implemented by the platform RHI.
 		checkNoEntry();
 		return nullptr;
 	}
 
-	virtual void UnlockStructuredBuffer_BottomOfPipe(class FRHICommandListImmediate& RHICmdList, FRHIStructuredBuffer* StructuredBuffer)
-	{
-		// Either this function or RHIUnlockStructuredBuffer must be implemented by the platform RHI.
-		checkNoEntry();
-	}
-	virtual void UnlockVertexBuffer_BottomOfPipe(class FRHICommandListImmediate& RHICmdList, FRHIVertexBuffer* VertexBuffer)
-	{
-		// Either this function or RHIUnlockVertexBuffer must be implemented by the platform RHI.
-		checkNoEntry();
-	}
-	virtual void UnlockIndexBuffer_BottomOfPipe(class FRHICommandListImmediate& RHICmdList, FRHIIndexBuffer* IndexBuffer)
+	virtual void UnlockBuffer_BottomOfPipe(class FRHICommandListImmediate& RHICmdList, FRHIBuffer* Buffer)
 	{
-		// Either this function or RHIUnlockIndexBuffer must be implemented by the platform RHI.
+		// Either this function or RHIUnlockBuffer must be implemented by the platform RHI.
 		checkNoEntry();
 	}
 
diff --git a/Engine/Source/Runtime/RHI/Public/RHICommandList.h b/Engine/Source/Runtime/RHI/Public/RHICommandList.h
--- a/Engine/Source/Runtime/RHI/Public/RHICommandList.h
+++ b/Engine/Source/Runtime/RHI/Public/RHICommandList.h
@@ -3922,12 +3922,12 @@ public:
 	
 	FORCEINLINE void* LockIndexBuffer(FRHIIndexBuffer* IndexBuffer, uint32 Offset, uint32 SizeRHI, EResourceLockMode LockMode)
 	{
-		return GDynamicRHI->RHILockIndexBuffer(*this, IndexBuffer, Offset, SizeRHI, LockMode);
+		return GDynamicRHI->RHILockBuffer(*this, IndexBuffer, Offset, SizeRHI, LockMode);
 	}
 	
 	FORCEINLINE void UnlockIndexBuffer(FRHIIndexBuffer* IndexBuffer)
 	{
-		GDynamicRHI->RHIUnlockIndexBuffer(*this, IndexBuffer);
+		GDynamicRHI->RHIUnlockBuffer(*this, IndexBuffer);
 	}
 	
 	FORCEINLINE void* LockStagingBuffer(FRHIStagingBuffer* StagingBuffer, FRHIGPUFence* Fence, uint32 Offset, uint32 SizeRHI)
@@ -3953,12 +3953,12 @@ public:
 
 	FORCEINLINE void* LockVertexBuffer(FRHIVertexBuffer* VertexBuffer, uint32 Offset, uint32 SizeRHI, EResourceLockMode LockMode)
 	{
-		return GDynamicRHI->RHILockVertexBuffer(*this, VertexBuffer, Offset, SizeRHI, LockMode);
+		return GDynamicRHI->RHILockBuffer(*this, VertexBuffer, Offset, SizeRHI, LockMode);
 	}
 	
 	FORCEINLINE void UnlockVertexBuffer(FRHIVertexBuffer* VertexBuffer)
 	{
-		GDynamicRHI->RHIUnlockVertexBuffer(*this, VertexBuffer);
+		GDynamicRHI->RHIUnlockBuffer(*this, VertexBuffer);
 	}
 	
 	FORCEINLINE void CopyVertexBuffer(FRHIVertexBuffer* SourceBuffer, FRHIVertexBuffer* DestBuffer)
@@ -3970,12 +3970,22 @@ public:
 
 	FORCEINLINE void* LockStructuredBuffer(FRHIStructuredBuffer* StructuredBuffer, uint32 Offset, uint32 SizeRHI, EResourceLockMode LockMode)
 	{
-		return GDynamicRHI->RHILockStructuredBuffer(*this, StructuredBuffer, Offset, SizeRHI, LockMode);
+		return GDynamicRHI->RHILockBuffer(*this, StructuredBuffer, Offset, SizeRHI, LockMode);
 	}
-	
+
 	FORCEINLINE void UnlockStructuredBuffer(FRHIStructuredBuffer* StructuredBuffer)
 	{
-		GDynamicRHI->RHIUnlockStructuredBuffer(*this, StructuredBuffer);
+		GDynamicRHI->RHIUnlockBuffer(*this, StructuredBuffer);
+	}
+
+	FORCEINLINE void* LockBuffer(FRHIBuffer* Buffer, uint32 Offset, uint32 SizeRHI, EResourceLockMode LockMode)
+	{
+		return GDynamicRHI->RHILockBuffer(*this, Buffer, Offset, SizeRHI, LockMode);
+	}
+	
+	FORCEINLINE void UnlockBuffer(FRHIBuffer* Buffer)
+	{
+		GDynamicRHI->RHIUnlockBuffer(*this, Buffer);
 	}
 	
 	FORCEINLINE FUnorderedAccessViewRHIRef CreateUnorderedAccessView(FRHIBuffer* Buffer, bool bUseUAVCounter, bool bAppendBuffer)
@@ -4932,12 +4942,12 @@ FORCEINLINE FIndexBufferRHIRef RHIAsyncCreateIndexBuffer(uint32 Stride, uint32 S
 
 FORCEINLINE void* RHILockIndexBuffer(FRHIIndexBuffer* IndexBuffer, uint32 Offset, uint32 Size, EResourceLockMode LockMode)
 {
-	return FRHICommandListExecutor::GetImmediateCommandList().LockIndexBuffer(IndexBuffer, Offset, Size, LockMode);
+	return FRHICommandListExecutor::GetImmediateCommandList().LockBuffer(IndexBuffer, Offset, Size, LockMode);
 }
 
 FORCEINLINE void RHIUnlockIndexBuffer(FRHIIndexBuffer* IndexBuffer)
 {
-	 FRHICommandListExecutor::GetImmediateCommandList().UnlockIndexBuffer(IndexBuffer);
+	 FRHICommandListExecutor::GetImmediateCommandList().UnlockBuffer(IndexBuffer);
 }
 
 FORCEINLINE FVertexBufferRHIRef RHICreateAndLockVertexBuffer(uint32 Size, uint32 InUsage, FRHIResourceCreateInfo& CreateInfo, void*& OutDataBuffer)
@@ -4971,12 +4981,12 @@ FORCEINLINE FVertexBufferRHIRef RHIAsyncCreateVertexBuffer(uint32 Size, uint32 I
 
 FORCEINLINE void* RHILockVertexBuffer(FRHIVertexBuffer* VertexBuffer, uint32 Offset, uint32 SizeRHI, EResourceLockMode LockMode)
 {
-	return FRHICommandListExecutor::GetImmediateCommandList().LockVertexBuffer(VertexBuffer, Offset, SizeRHI, LockMode);
+	return FRHICommandListExecutor::GetImmediateCommandList().LockBuffer(VertexBuffer, Offset, SizeRHI, LockMode);
 }
 
 FORCEINLINE void RHIUnlockVertexBuffer(FRHIVertexBuffer* VertexBuffer)
 {
-	 FRHICommandListExecutor::GetImmediateCommandList().UnlockVertexBuffer(VertexBuffer);
+	 FRHICommandListExecutor::GetImmediateCommandList().UnlockBuffer(VertexBuffer);
 }
 
 FORCEINLINE FStructuredBufferRHIRef RHICreateStructuredBuffer(uint32 Stride, uint32 Size, uint32 InUsage, ERHIAccess InResourceState, FRHIResourceCreateInfo& CreateInfo)
@@ -4993,12 +5003,22 @@ FORCEINLINE FStructuredBufferRHIRef RHICreateStructuredBuffer(uint32 Stride, uin
 
 FORCEINLINE void* RHILockStructuredBuffer(FRHIStructuredBuffer* StructuredBuffer, uint32 Offset, uint32 SizeRHI, EResourceLockMode LockMode)
 {
-	return FRHICommandListExecutor::GetImmediateCommandList().LockStructuredBuffer(StructuredBuffer, Offset, SizeRHI, LockMode);
+	return FRHICommandListExecutor::GetImmediateCommandList().LockBuffer(StructuredBuffer, Offset, SizeRHI, LockMode);
 }
 
 FORCEINLINE void RHIUnlockStructuredBuffer(FRHIStructuredBuffer* StructuredBuffer)
 {
-	 FRHICommandListExecutor::GetImmediateCommandList().UnlockStructuredBuffer(StructuredBuffer);
+	 FRHICommandListExecutor::GetImmediateCommandList().UnlockBuffer(StructuredBuffer);
+}
+
+FORCEINLINE void* RHILockBuffer(FRHIBuffer* Buffer, uint32 Offset, uint32 SizeRHI, EResourceLockMode LockMode)
+{
+	return FRHICommandListExecutor::GetImmediateCommandList().LockBuffer(Buffer, Offset, SizeRHI, LockMode);
+}
+
+FORCEINLINE void RHIUnlockBuffer(FRHIBuffer* Buffer)
+{
+	FRHICommandListExecutor::GetImmediateCommandList().UnlockBuffer(Buffer);
 }
 
 FORCEINLINE FUnorderedAccessViewRHIRef RHICreateUnorderedAccessView(FRHIBuffer* Buffer, bool bUseUAVCounter, bool bAppendBuffer)
diff --git a/Engine/Source/Runtime/RHI/Public/RHIUtilities.h b/Engine/Source/Runtime/RHI/Public/RHIUtilities.h
--- a/Engine/Source/Runtime/RHI/Public/RHIUtilities.h
+++ b/Engine/Source/Runtime/RHI/Public/RHIUtilities.h
@@ -515,7 +515,7 @@ struct FDynamicReadBuffer : public FReadBuffer
 	{
 		check(MappedBuffer == nullptr);
 		check(IsValidRef(Buffer));
-		MappedBuffer = (uint8*)RHILockVertexBuffer(Buffer, 0, NumBytes, RLM_WriteOnly);
+		MappedBuffer = (uint8*)RHILockBuffer(Buffer, 0, NumBytes, RLM_WriteOnly);
 	}
 
 	/**
@@ -525,7 +525,7 @@ struct FDynamicReadBuffer : public FReadBuffer
 	{
 		check(MappedBuffer);
 		check(IsValidRef(Buffer));
-		RHIUnlockVertexBuffer(Buffer);
+		RHIUnlockBuffer(Buffer);
 		MappedBuffer = nullptr;
 	}
 };
diff --git a/Engine/Source/Runtime/RHI/Public/RHIValidation.h b/Engine/Source/Runtime/RHI/Public/RHIValidation.h
--- a/Engine/Source/Runtime/RHI/Public/RHIValidation.h
+++ b/Engine/Source/Runtime/RHI/Public/RHIValidation.h
@@ -365,18 +365,6 @@ public:
 		return IndexBuffer;
 	}
 
-	// FlushType: Flush RHI Thread
-	virtual void* RHILockIndexBuffer(class FRHICommandListImmediate& RHICmdList, FRHIIndexBuffer* IndexBuffer, uint32 Offset, uint32 SizeRHI, EResourceLockMode LockMode) override final
-	{
-		return RHI->RHILockIndexBuffer(RHICmdList, IndexBuffer, Offset, SizeRHI, LockMode);
-	}
-
-	// FlushType: Flush RHI Thread
-	virtual void RHIUnlockIndexBuffer(class FRHICommandListImmediate& RHICmdList, FRHIIndexBuffer* IndexBuffer) override final
-	{
-		RHI->RHIUnlockIndexBuffer(RHICmdList, IndexBuffer);
-	}
-
 	/**
 	* @param ResourceArray - An optional pointer to a resource array containing the resource's data.
 	*/
@@ -388,18 +376,6 @@ public:
 		return VertexBuffer;
 	}
 
-	// FlushType: Flush RHI Thread
-	virtual void* RHILockVertexBuffer(class FRHICommandListImmediate& RHICmdList, FRHIVertexBuffer* VertexBuffer, uint32 Offset, uint32 SizeRHI, EResourceLockMode LockMode) override final
-	{
-		return RHI->RHILockVertexBuffer(RHICmdList, VertexBuffer, Offset, SizeRHI, LockMode);
-	}
-
-	// FlushType: Flush RHI Thread
-	virtual void RHIUnlockVertexBuffer(class FRHICommandListImmediate& RHICmdList, FRHIVertexBuffer* VertexBuffer) override final
-	{
-		RHI->RHIUnlockVertexBuffer(RHICmdList, VertexBuffer);
-	}
-
 	/** Copies the contents of one vertex buffer to another vertex buffer.  They must have identical sizes. */
 	// FlushType: Flush Immediate (seems dangerous)
 	virtual void RHICopyVertexBuffer(FRHIVertexBuffer* SourceBuffer, FRHIVertexBuffer* DestBuffer) override final
@@ -419,15 +395,15 @@ public:
 	}
 
 	// FlushType: Flush RHI Thread
-	virtual void* RHILockStructuredBuffer(class FRHICommandListImmediate& RHICmdList, FRHIStructuredBuffer* StructuredBuffer, uint32 Offset, uint32 SizeRHI, EResourceLockMode LockMode) override final
+	virtual void* RHILockBuffer(class FRHICommandListImmediate& RHICmdList, FRHIBuffer* Buffer, uint32 Offset, uint32 SizeRHI, EResourceLockMode LockMode) override final
 	{
-		return RHI->RHILockStructuredBuffer(RHICmdList, StructuredBuffer, Offset, SizeRHI, LockMode);
+		return RHI->RHILockBuffer(RHICmdList, Buffer, Offset, SizeRHI, LockMode);
 	}
 
 	// FlushType: Flush RHI Thread
-	virtual void RHIUnlockStructuredBuffer(class FRHICommandListImmediate& RHICmdList, FRHIStructuredBuffer* StructuredBuffer) override final
+	virtual void RHIUnlockBuffer(class FRHICommandListImmediate& RHICmdList, FRHIBuffer* Buffer) override final
 	{
-		RHI->RHIUnlockStructuredBuffer(RHICmdList, StructuredBuffer);
+		RHI->RHIUnlockBuffer(RHICmdList, Buffer);
 	}
 
 	/** Creates an unordered access view of the given structured buffer. */
@@ -1413,16 +1389,6 @@ public:
 		return SRV;
 	}
 
-	virtual void* LockVertexBuffer_BottomOfPipe(class FRHICommandListImmediate& RHICmdList, FRHIVertexBuffer* VertexBuffer, uint32 Offset, uint32 SizeRHI, EResourceLockMode LockMode) override final
-	{
-		return RHI->LockVertexBuffer_BottomOfPipe(RHICmdList, VertexBuffer, Offset, SizeRHI, LockMode);
-	}
-
-	virtual void UnlockVertexBuffer_BottomOfPipe(class FRHICommandListImmediate& RHICmdList, FRHIVertexBuffer* VertexBuffer) override final
-	{
-		return RHI->UnlockVertexBuffer_BottomOfPipe(RHICmdList, VertexBuffer);
-	}
-
 	virtual FTexture2DRHIRef AsyncReallocateTexture2D_RenderThread(class FRHICommandListImmediate& RHICmdList, FRHITexture2D* Texture2D, int32 NewMipCount, int32 NewSizeX, int32 NewSizeY, FThreadSafeCounter* RequestStatus) override final
 	{
 		// TODO: find proper state for new texture
@@ -1450,24 +1416,14 @@ public:
 		return IndexBuffer;
 	}
 
-	virtual void* LockIndexBuffer_BottomOfPipe(class FRHICommandListImmediate& RHICmdList, FRHIIndexBuffer* IndexBuffer, uint32 Offset, uint32 SizeRHI, EResourceLockMode LockMode) override final
-	{
-		return RHI->LockIndexBuffer_BottomOfPipe(RHICmdList, IndexBuffer, Offset, SizeRHI, LockMode);
-	}
-
-	virtual void UnlockIndexBuffer_BottomOfPipe(class FRHICommandListImmediate& RHICmdList, FRHIIndexBuffer* IndexBuffer) override final
-	{
-		RHI->UnlockIndexBuffer_BottomOfPipe(RHICmdList, IndexBuffer);
-	}
-
-	virtual void* LockStructuredBuffer_BottomOfPipe(class FRHICommandListImmediate& RHICmdList, FRHIStructuredBuffer* StructuredBuffer, uint32 Offset, uint32 SizeRHI, EResourceLockMode LockMode) override final
+	virtual void* LockBuffer_BottomOfPipe(class FRHICommandListImmediate& RHICmdList, FRHIBuffer* Buffer, uint32 Offset, uint32 SizeRHI, EResourceLockMode LockMode) override final
 	{
-		return RHI->LockStructuredBuffer_BottomOfPipe(RHICmdList, StructuredBuffer, Offset, SizeRHI, LockMode);
+		return RHI->LockBuffer_BottomOfPipe(RHICmdList, Buffer, Offset, SizeRHI, LockMode);
 	}
 	
-	virtual void UnlockStructuredBuffer_BottomOfPipe(class FRHICommandListImmediate& RHICmdList, FRHIStructuredBuffer* StructuredBuffer) override final
+	virtual void UnlockBuffer_BottomOfPipe(class FRHICommandListImmediate& RHICmdList, FRHIBuffer* Buffer) override final
 	{
-		RHI->UnlockStructuredBuffer_BottomOfPipe(RHICmdList, StructuredBuffer);
+		RHI->UnlockBuffer_BottomOfPipe(RHICmdList, Buffer);
 	}
 
 	virtual FVertexShaderRHIRef CreateVertexShader_RenderThread(class FRHICommandListImmediate& RHICmdList, TArrayView<const uint8> Code, const FSHAHash& Hash) override final
diff --git a/Engine/Source/Runtime/VulkanRHI/Private/VulkanIndexBuffer.cpp b/Engine/Source/Runtime/VulkanRHI/Private/VulkanIndexBuffer.cpp
--- a/Engine/Source/Runtime/VulkanRHI/Private/VulkanIndexBuffer.cpp
+++ b/Engine/Source/Runtime/VulkanRHI/Private/VulkanIndexBuffer.cpp
@@ -415,7 +415,7 @@ void FVulkanResourceMultiBuffer::Swap(FVulkanResourceMultiBuffer& Other)
 
 FIndexBufferRHIRef FVulkanDynamicRHI::RHICreateIndexBuffer(uint32 Stride, uint32 Size, uint32 InUsage, ERHIAccess InResourceState, FRHIResourceCreateInfo& CreateInfo)
 {
-	LLM_SCOPE_VULKAN(ELLMTagVulkan::VulkanIndexBuffers);
+	LLM_SCOPE_VULKAN(ELLMTagVulkan::VulkanBuffers);
 	if (CreateInfo.bWithoutNativeResource)
 	{
 		return new FVulkanResourceMultiBuffer(nullptr, 0, 0, 0, 0, CreateInfo);
@@ -423,29 +423,29 @@ FIndexBufferRHIRef FVulkanDynamicRHI::RHICreateIndexBuffer(uint32 Stride, uint32
 	return new FVulkanResourceMultiBuffer(Device, 0, Size, InUsage | BUF_IndexBuffer, Stride, CreateInfo);
 }
 
-void* FVulkanDynamicRHI::LockIndexBuffer_BottomOfPipe(FRHICommandListImmediate& RHICmdList, FRHIIndexBuffer* IndexBufferRHI, uint32 Offset, uint32 Size, EResourceLockMode LockMode)
+void* FVulkanDynamicRHI::LockBuffer_BottomOfPipe(FRHICommandListImmediate& RHICmdList, FRHIBuffer* BufferRHI, uint32 Offset, uint32 Size, EResourceLockMode LockMode)
 {
-	FVulkanResourceMultiBuffer* IndexBuffer = ResourceCast(IndexBufferRHI);
-	return IndexBuffer->Lock(false, LockMode, Size, Offset);
+	LLM_SCOPE_VULKAN(ELLMTagVulkan::VulkanBuffers);
+	FVulkanResourceMultiBuffer* Buffer = ResourceCast(BufferRHI);
+	return Buffer->Lock(false, LockMode, Size, Offset);
 }
 
-#if VULKAN_BUFFER_LOCK_THREADSAFE
-void* FVulkanDynamicRHI::LockIndexBuffer_RenderThread(class FRHICommandListImmediate& RHICmdList, FRHIIndexBuffer* IndexBufferRHI, uint32 Offset, uint32 SizeRHI, EResourceLockMode LockMode)
+void FVulkanDynamicRHI::UnlockBuffer_BottomOfPipe(FRHICommandListImmediate& RHICmdList, FRHIBuffer* BufferRHI)
 {
-	return this->RHILockIndexBuffer(IndexBufferRHI, Offset, SizeRHI, LockMode);
+	LLM_SCOPE_VULKAN(ELLMTagVulkan::VulkanBuffers);
+	FVulkanResourceMultiBuffer* Buffer = ResourceCast(BufferRHI);
+	Buffer->Unlock(false);
 }
-#endif
 
-void FVulkanDynamicRHI::UnlockIndexBuffer_BottomOfPipe(FRHICommandListImmediate& RHICmdList, FRHIIndexBuffer* IndexBufferRHI)
+#if VULKAN_BUFFER_LOCK_THREADSAFE
+void* FVulkanDynamicRHI::LockBuffer_RenderThread(class FRHICommandListImmediate& RHICmdList, FRHIBuffer* BufferRHI, uint32 Offset, uint32 SizeRHI, EResourceLockMode LockMode)
 {
-	FVulkanResourceMultiBuffer* IndexBuffer = ResourceCast(IndexBufferRHI);
-	IndexBuffer->Unlock(false);
+	return this->RHILockBuffer(BufferRHI, Offset, SizeRHI, LockMode);
 }
 
-#if VULKAN_BUFFER_LOCK_THREADSAFE
-void FVulkanDynamicRHI::UnlockIndexBuffer_RenderThread(FRHICommandListImmediate& RHICmdList, FRHIIndexBuffer* IndexBufferRHI)
+void FVulkanDynamicRHI::UnlockBuffer_RenderThread(FRHICommandListImmediate& RHICmdList, FRHIBuffer* BufferRHI)
 {
-	this->RHIUnlockIndexBuffer(IndexBufferRHI);
+	this->RHIUnlockBuffer(BufferRHI);
 }
 #endif
 
diff --git a/Engine/Source/Runtime/VulkanRHI/Private/VulkanLLM.h b/Engine/Source/Runtime/VulkanRHI/Private/VulkanLLM.h
--- a/Engine/Source/Runtime/VulkanRHI/Private/VulkanLLM.h
+++ b/Engine/Source/Runtime/VulkanRHI/Private/VulkanLLM.h
@@ -23,8 +23,7 @@ enum class ELLMTagVulkan : LLM_TAG_TYPE
 {
 	VulkanMisc = (LLM_TAG_TYPE)ELLMTag::PlatformTagStart,
 	VulkanUniformBuffers,
-	VulkanIndexBuffers,
-	VulkanVertexBuffers,
+	VulkanBuffers,
 	VulkanTextures,
 	VulkanRenderTargets,
 	VulkanSpareMemoryGPU,
diff --git a/Engine/Source/Runtime/VulkanRHI/Private/VulkanStructuredBuffer.cpp b/Engine/Source/Runtime/VulkanRHI/Private/VulkanStructuredBuffer.cpp
--- a/Engine/Source/Runtime/VulkanRHI/Private/VulkanStructuredBuffer.cpp
+++ b/Engine/Source/Runtime/VulkanRHI/Private/VulkanStructuredBuffer.cpp
@@ -8,17 +8,3 @@ FStructuredBufferRHIRef FVulkanDynamicRHI::RHICreateStructuredBuffer(uint32 InSt
 {
 	return new FVulkanResourceMultiBuffer(Device, VK_BUFFER_USAGE_STORAGE_BUFFER_BIT, InSize, InUsage | BUF_StructuredBuffer, InStride, CreateInfo);
 }
-
-void* FVulkanDynamicRHI::LockStructuredBuffer_BottomOfPipe(FRHICommandListImmediate& RHICmdList, FRHIStructuredBuffer* StructuredBufferRHI, uint32 Offset, uint32 Size, EResourceLockMode LockMode)
-{
-	FVulkanResourceMultiBuffer* StructuredBuffer = ResourceCast(StructuredBufferRHI);
-
-	return StructuredBuffer->Lock(false, LockMode, Size, Offset);
-}
-
-void FVulkanDynamicRHI::UnlockStructuredBuffer_BottomOfPipe(FRHICommandListImmediate& RHICmdList, FRHIStructuredBuffer* StructuredBufferRHI)
-{
-	FVulkanResourceMultiBuffer* StructuredBuffer = ResourceCast(StructuredBufferRHI);
-
-	StructuredBuffer->Unlock(false);
-}
diff --git a/Engine/Source/Runtime/VulkanRHI/Private/VulkanVertexBuffer.cpp b/Engine/Source/Runtime/VulkanRHI/Private/VulkanVertexBuffer.cpp
--- a/Engine/Source/Runtime/VulkanRHI/Private/VulkanVertexBuffer.cpp
+++ b/Engine/Source/Runtime/VulkanRHI/Private/VulkanVertexBuffer.cpp
@@ -9,7 +9,7 @@
 
 FVertexBufferRHIRef FVulkanDynamicRHI::RHICreateVertexBuffer(uint32 Size, uint32 InUsage, ERHIAccess InResourceState, FRHIResourceCreateInfo& CreateInfo)
 {
-	LLM_SCOPE_VULKAN(ELLMTagVulkan::VulkanVertexBuffers);
+	LLM_SCOPE_VULKAN(ELLMTagVulkan::VulkanBuffers);
 	if (CreateInfo.bWithoutNativeResource)
 	{
 		return new FVulkanResourceMultiBuffer(nullptr, VK_BUFFER_USAGE_VERTEX_BUFFER_BIT, 0, 0, 0, CreateInfo);
@@ -18,34 +18,6 @@ FVertexBufferRHIRef FVulkanDynamicRHI::RHICreateVertexBuffer(uint32 Size, uint32
 	return VertexBuffer;
 }
 
-void* FVulkanDynamicRHI::LockVertexBuffer_BottomOfPipe(FRHICommandListImmediate& RHICmdList, FRHIVertexBuffer* VertexBufferRHI, uint32 Offset, uint32 Size, EResourceLockMode LockMode)
-{
-	LLM_SCOPE_VULKAN(ELLMTagVulkan::VulkanVertexBuffers);
-	FVulkanResourceMultiBuffer* VertexBuffer = ResourceCast(VertexBufferRHI);
-	return VertexBuffer->Lock(false, LockMode, Size, Offset);
-}
-
-#if VULKAN_BUFFER_LOCK_THREADSAFE
-void* FVulkanDynamicRHI::LockVertexBuffer_RenderThread(class FRHICommandListImmediate& RHICmdList, FRHIVertexBuffer* VertexBufferRHI, uint32 Offset, uint32 SizeRHI, EResourceLockMode LockMode)
-{
-	return this->RHILockVertexBuffer(VertexBufferRHI, Offset, SizeRHI, LockMode);
-}
-#endif
-
-void FVulkanDynamicRHI::UnlockVertexBuffer_BottomOfPipe(FRHICommandListImmediate& RHICmdList, FRHIVertexBuffer* VertexBufferRHI)
-{
-	LLM_SCOPE_VULKAN(ELLMTagVulkan::VulkanVertexBuffers);
-	FVulkanResourceMultiBuffer* VertexBuffer = ResourceCast(VertexBufferRHI);
-	VertexBuffer->Unlock(false);
-}
-
-#if VULKAN_BUFFER_LOCK_THREADSAFE
-void FVulkanDynamicRHI::UnlockVertexBuffer_RenderThread(FRHICommandListImmediate& RHICmdList, FRHIVertexBuffer* VertexBufferRHI)
-{
-	return this->RHIUnlockVertexBuffer(VertexBufferRHI);
-}
-#endif
-
 void FVulkanDynamicRHI::RHICopyVertexBuffer(FRHIVertexBuffer* SourceBufferRHI, FRHIVertexBuffer* DestBufferRHI)
 {
 	VULKAN_SIGNAL_UNIMPLEMENTED();
diff --git a/Engine/Source/Runtime/VulkanRHI/Public/VulkanDynamicRHI.h b/Engine/Source/Runtime/VulkanRHI/Public/VulkanDynamicRHI.h
--- a/Engine/Source/Runtime/VulkanRHI/Public/VulkanDynamicRHI.h
+++ b/Engine/Source/Runtime/VulkanRHI/Public/VulkanDynamicRHI.h
@@ -81,16 +81,12 @@ public:
 	virtual FUniformBufferRHIRef RHICreateUniformBuffer(const void* Contents, const FRHIUniformBufferLayout& Layout, EUniformBufferUsage Usage, EUniformBufferValidation Validation) final override;
 	virtual void RHIUpdateUniformBuffer(FRHIUniformBuffer* UniformBufferRHI, const void* Contents) final override;
 	virtual FIndexBufferRHIRef RHICreateIndexBuffer(uint32 Stride, uint32 Size, uint32 InUsage, ERHIAccess InResourceState, FRHIResourceCreateInfo& CreateInfo) final override;
-	virtual void* LockIndexBuffer_BottomOfPipe(FRHICommandListImmediate& RHICmdList, FRHIIndexBuffer* IndexBuffer, uint32 Offset, uint32 Size, EResourceLockMode LockMode) final override;
-	virtual void UnlockIndexBuffer_BottomOfPipe(FRHICommandListImmediate& RHICmdList, FRHIIndexBuffer* IndexBuffer) final override;
 	virtual FVertexBufferRHIRef RHICreateVertexBuffer(uint32 Size, uint32 InUsage, ERHIAccess InResourceState, FRHIResourceCreateInfo& CreateInfo) final override;
-	virtual void* LockVertexBuffer_BottomOfPipe(FRHICommandListImmediate& RHICmdList, FRHIVertexBuffer* VertexBuffer, uint32 Offset, uint32 SizeRHI, EResourceLockMode LockMode) final override;
-	virtual void UnlockVertexBuffer_BottomOfPipe(FRHICommandListImmediate& RHICmdList, FRHIVertexBuffer* VertexBuffer) final override;
 	virtual void RHICopyVertexBuffer(FRHIVertexBuffer* SourceBuffer, FRHIVertexBuffer* DestBuffer) final override;
 	virtual void RHITransferBufferUnderlyingResource(FRHIBuffer* DestBuffer, FRHIBuffer* SrcBuffer) final override;
 	virtual FStructuredBufferRHIRef RHICreateStructuredBuffer(uint32 Stride, uint32 Size, uint32 InUsage, ERHIAccess InResourceState, FRHIResourceCreateInfo& CreateInfo) final override;
-	virtual void* LockStructuredBuffer_BottomOfPipe(FRHICommandListImmediate& RHICmdList, FRHIStructuredBuffer* StructuredBuffer, uint32 Offset, uint32 SizeRHI, EResourceLockMode LockMode) final override;
-	virtual void UnlockStructuredBuffer_BottomOfPipe(FRHICommandListImmediate& RHICmdList, FRHIStructuredBuffer* StructuredBuffer) final override;
+	virtual void* LockBuffer_BottomOfPipe(FRHICommandListImmediate& RHICmdList, FRHIBuffer* Buffer, uint32 Offset, uint32 Size, EResourceLockMode LockMode) final override;
+	virtual void UnlockBuffer_BottomOfPipe(FRHICommandListImmediate& RHICmdList, FRHIBuffer* Buffer) final override;
 	virtual FUnorderedAccessViewRHIRef RHICreateUnorderedAccessView(FRHIBuffer* Buffer, bool bUseUAVCounter, bool bAppendBuffer) final override;
 	virtual FUnorderedAccessViewRHIRef RHICreateUnorderedAccessView(FRHITexture* Texture, uint32 MipLevel) final override;
 	virtual FUnorderedAccessViewRHIRef RHICreateUnorderedAccessView(FRHIBuffer* Buffer, uint8 Format) final override;
@@ -196,10 +192,7 @@ public:
 		// this-> is required to avoid calling the global version of this function
 		return this->RHICreateVertexBuffer(Size, InUsage, InResourceState, CreateInfo);
 	}
-#if VULKAN_BUFFER_LOCK_THREADSAFE	
-	virtual void* LockVertexBuffer_RenderThread(class FRHICommandListImmediate& RHICmdList, FRHIVertexBuffer* VertexBuffer, uint32 Offset, uint32 SizeRHI, EResourceLockMode LockMode) override final;
-	virtual void UnlockVertexBuffer_RenderThread(class FRHICommandListImmediate& RHICmdList, FRHIVertexBuffer* VertexBuffer) override final;
-#endif
+
 	virtual FTexture2DRHIRef AsyncReallocateTexture2D_RenderThread(class FRHICommandListImmediate& RHICmdList, FRHITexture2D* Texture2D, int32 NewMipCount, int32 NewSizeX, int32 NewSizeY, FThreadSafeCounter* RequestStatus) override final;
 	//virtual ETextureReallocationStatus CancelAsyncReallocateTexture2D_RenderThread(class FRHICommandListImmediate& RHICmdList, FRHITexture2D* Texture2D, bool bBlockUntilCompleted) override final;
 	
@@ -209,8 +202,8 @@ public:
 		return this->RHICreateIndexBuffer(Stride, Size, InUsage, InResourceState, CreateInfo);
 	}
 #if VULKAN_BUFFER_LOCK_THREADSAFE		
-	virtual void* LockIndexBuffer_RenderThread(class FRHICommandListImmediate& RHICmdList, FRHIIndexBuffer* IndexBuffer, uint32 Offset, uint32 SizeRHI, EResourceLockMode LockMode) override final;
-	virtual void UnlockIndexBuffer_RenderThread(class FRHICommandListImmediate& RHICmdList, FRHIIndexBuffer* IndexBuffer) override final;
+	virtual void* LockBuffer_RenderThread(class FRHICommandListImmediate& RHICmdList, FRHIBuffer* Buffer, uint32 Offset, uint32 SizeRHI, EResourceLockMode LockMode) override final;
+	virtual void UnlockBuffer_RenderThread(class FRHICommandListImmediate& RHICmdList, FRHIBuffer* Buffer) override final;
 #endif
 
 	virtual FVertexShaderRHIRef CreateVertexShader_RenderThread(class FRHICommandListImmediate& RHICmdList, TArrayView<const uint8> Code, const FSHAHash& Hash) override final
diff --git a/Engine/Source/Runtime/Windows/D3D11RHI/Private/D3D11IndexBuffer.cpp b/Engine/Source/Runtime/Windows/D3D11RHI/Private/D3D11IndexBuffer.cpp
--- a/Engine/Source/Runtime/Windows/D3D11RHI/Private/D3D11IndexBuffer.cpp
+++ b/Engine/Source/Runtime/Windows/D3D11RHI/Private/D3D11IndexBuffer.cpp
@@ -91,27 +91,30 @@ FIndexBufferRHIRef FD3D11DynamicRHI::CreateIndexBuffer_RenderThread(
 	return RHICreateIndexBuffer(Stride, Size, InUsage, InResourceState, CreateInfo);
 }
 
-void* FD3D11DynamicRHI::LockIndexBuffer_BottomOfPipe(FRHICommandListImmediate& RHICmdList, FRHIIndexBuffer* IndexBufferRHI, uint32 Offset, uint32 Size, EResourceLockMode LockMode)
+void* FD3D11DynamicRHI::LockBuffer_BottomOfPipe(FRHICommandListImmediate& RHICmdList, FRHIBuffer* BufferRHI, uint32 Offset, uint32 Size, EResourceLockMode LockMode)
 {
-	FD3D11Buffer* IndexBuffer = ResourceCast(IndexBufferRHI);
+	FD3D11Buffer* Buffer = ResourceCast(BufferRHI);
 	// If this resource is bound to the device, unbind it
-	ConditionalClearShaderResource(IndexBuffer, true);
+	ConditionalClearShaderResource(Buffer, true);
 
-	// Determine whether the index buffer is dynamic or not.
+	// Determine whether the buffer is dynamic or not.
 	D3D11_BUFFER_DESC Desc;
-	IndexBuffer->Resource->GetDesc(&Desc);
+	Buffer->Resource->GetDesc(&Desc);
 	const bool bIsDynamic = (Desc.Usage == D3D11_USAGE_DYNAMIC);
 
-	FD3D11LockedKey LockedKey(IndexBuffer->Resource);
+	FD3D11LockedKey LockedKey(Buffer->Resource);
 	FD3D11LockedData LockedData;
 
 	if(bIsDynamic)
 	{
-		check(LockMode == RLM_WriteOnly);
+		check(LockMode == RLM_WriteOnly || LockMode == RLM_WriteOnly_NoOverwrite);
 
 		// If the buffer is dynamic, map its memory for writing.
 		D3D11_MAPPED_SUBRESOURCE MappedSubresource;
-		VERIFYD3D11RESULT_EX(Direct3DDeviceIMContext->Map(IndexBuffer->Resource,0,D3D11_MAP_WRITE_DISCARD,0,&MappedSubresource), Direct3DDevice);
+
+		D3D11_MAP MapType = (LockMode == RLM_WriteOnly)? D3D11_MAP_WRITE_DISCARD : D3D11_MAP_WRITE_NO_OVERWRITE;
+		VERIFYD3D11RESULT_EX(Direct3DDeviceIMContext->Map(Buffer->Resource, 0, MapType, 0, &MappedSubresource), Direct3DDevice);
+
 		LockedData.SetData(MappedSubresource.pData);
 		LockedData.Pitch = MappedSubresource.RowPitch;
 	}
@@ -122,23 +125,29 @@ void* FD3D11DynamicRHI::LockIndexBuffer_BottomOfPipe(FRHICommandListImmediate& R
 			// If the static buffer is being locked for reading, create a staging buffer.
 			D3D11_BUFFER_DESC StagingBufferDesc;
 			ZeroMemory( &StagingBufferDesc, sizeof( D3D11_BUFFER_DESC ) );
-			StagingBufferDesc.ByteWidth = Desc.ByteWidth;
+			StagingBufferDesc.ByteWidth = Size;
 			StagingBufferDesc.Usage = D3D11_USAGE_STAGING;
 			StagingBufferDesc.BindFlags = 0;
 			StagingBufferDesc.CPUAccessFlags = D3D11_CPU_ACCESS_READ;
 			StagingBufferDesc.MiscFlags = 0;
-			TRefCountPtr<ID3D11Buffer> StagingIndexBuffer;
-			VERIFYD3D11RESULT_EX(Direct3DDevice->CreateBuffer(&StagingBufferDesc,NULL,StagingIndexBuffer.GetInitReference()), Direct3DDevice);
-			LockedData.StagingResource = StagingIndexBuffer;
-
-			// Copy the contents of the index buffer to the staging buffer.
-			Direct3DDeviceIMContext->CopyResource(StagingIndexBuffer,IndexBuffer->Resource);
+			TRefCountPtr<ID3D11Buffer> StagingBuffer;
+			VERIFYD3D11RESULT_EX(Direct3DDevice->CreateBuffer(&StagingBufferDesc, NULL, StagingBuffer.GetInitReference()), Direct3DDevice);
+			LockedData.StagingResource = StagingBuffer;
+
+			// Copy the contents of the buffer to the staging buffer.
+			D3D11_BOX SourceBox;
+			SourceBox.left = Offset;
+			SourceBox.right = Offset + Size;
+			SourceBox.top = SourceBox.front = 0;
+			SourceBox.bottom = SourceBox.back = 1;
+			Direct3DDeviceIMContext->CopySubresourceRegion(StagingBuffer, 0, 0, 0, 0, Buffer->Resource, 0, &SourceBox);
 
 			// Map the staging buffer's memory for reading.
 			D3D11_MAPPED_SUBRESOURCE MappedSubresource;
-			VERIFYD3D11RESULT_EX(Direct3DDeviceIMContext->Map(StagingIndexBuffer,0,D3D11_MAP_READ,0,&MappedSubresource), Direct3DDevice);
+			VERIFYD3D11RESULT_EX(Direct3DDeviceIMContext->Map(StagingBuffer, 0, D3D11_MAP_READ, 0, &MappedSubresource), Direct3DDevice);
 			LockedData.SetData(MappedSubresource.pData);
 			LockedData.Pitch = MappedSubresource.RowPitch;
+			Offset = 0;
 		}
 		else
 		{
@@ -155,27 +164,27 @@ void* FD3D11DynamicRHI::LockIndexBuffer_BottomOfPipe(FRHICommandListImmediate& R
 	return (void*)((uint8*)LockedData.GetData() + Offset);
 }
 
-void FD3D11DynamicRHI::UnlockIndexBuffer_BottomOfPipe(FRHICommandListImmediate& RHICmdList, FRHIIndexBuffer* IndexBufferRHI)
+void FD3D11DynamicRHI::UnlockBuffer_BottomOfPipe(FRHICommandListImmediate& RHICmdList, FRHIBuffer* BufferRHI)
 {
-	FD3D11Buffer* IndexBuffer = ResourceCast(IndexBufferRHI);
+	FD3D11Buffer* Buffer = ResourceCast(BufferRHI);
 
-	// Determine whether the index buffer is dynamic or not.
+	// Determine whether the buffer is dynamic or not.
 	D3D11_BUFFER_DESC Desc;
-	IndexBuffer->Resource->GetDesc(&Desc);
+	Buffer->Resource->GetDesc(&Desc);
 	const bool bIsDynamic = (Desc.Usage == D3D11_USAGE_DYNAMIC);
 
-	// Find the outstanding lock for this IB and remove it from the tracker.
+	// Find the outstanding lock for this buffer and remove it from the tracker.
 	FD3D11LockedData LockedData;
-	verifyf(RemoveLockedData(FD3D11LockedKey(IndexBuffer->Resource), LockedData), TEXT("Index buffer is not locked"));
+	verifyf(RemoveLockedData(FD3D11LockedKey(Buffer->Resource), LockedData), TEXT("Buffer is not locked"));
 
 	if(bIsDynamic)
 	{
-		// If the IB is dynamic, its memory was mapped directly; unmap it.
-		Direct3DDeviceIMContext->Unmap(IndexBuffer->Resource,0);
+		// If the buffer is dynamic, its memory was mapped directly; unmap it.
+		Direct3DDeviceIMContext->Unmap(Buffer->Resource, 0);
 	}
 	else
 	{
-		// If the static IB lock involved a staging resource, it was locked for reading.
+		// If the static buffer lock involved a staging resource, it was locked for reading.
 		if(LockedData.StagingResource)
 		{
 			// Unmap the staging buffer's memory.
@@ -184,8 +193,8 @@ void FD3D11DynamicRHI::UnlockIndexBuffer_BottomOfPipe(FRHICommandListImmediate&
 		}
 		else 
 		{
-			// Copy the contents of the temporary memory buffer allocated for writing into the IB.
-			Direct3DDeviceIMContext->UpdateSubresource(IndexBuffer->Resource,0,NULL,LockedData.GetData(),LockedData.Pitch,0);
+			// Copy the contents of the temporary memory buffer allocated for writing into the buffer.
+			Direct3DDeviceIMContext->UpdateSubresource(Buffer->Resource, 0, NULL, LockedData.GetData(), LockedData.Pitch, 0);
 
 			// Free the temporary memory buffer.
 			LockedData.FreeData();
diff --git a/Engine/Source/Runtime/Windows/D3D11RHI/Private/D3D11RHIPrivate.h b/Engine/Source/Runtime/Windows/D3D11RHI/Private/D3D11RHIPrivate.h
--- a/Engine/Source/Runtime/Windows/D3D11RHI/Private/D3D11RHIPrivate.h
+++ b/Engine/Source/Runtime/Windows/D3D11RHI/Private/D3D11RHIPrivate.h
@@ -437,16 +437,12 @@ public:
 	virtual FUniformBufferRHIRef RHICreateUniformBuffer(const void* Contents, const FRHIUniformBufferLayout& Layout, EUniformBufferUsage Usage, EUniformBufferValidation Validation) final override;
 	virtual void RHIUpdateUniformBuffer(FRHIUniformBuffer* UniformBufferRHI, const void* Contents) final override;
 	virtual FIndexBufferRHIRef RHICreateIndexBuffer(uint32 Stride, uint32 Size, uint32 InUsage, ERHIAccess InResourceState, FRHIResourceCreateInfo& CreateInfo) final override;
-	virtual void* LockIndexBuffer_BottomOfPipe(FRHICommandListImmediate& RHICmdList, FRHIIndexBuffer* IndexBuffer, uint32 Offset, uint32 Size, EResourceLockMode LockMode) final override;
-	virtual void UnlockIndexBuffer_BottomOfPipe(FRHICommandListImmediate& RHICmdList, FRHIIndexBuffer* IndexBuffer) final override;
 	virtual FVertexBufferRHIRef RHICreateVertexBuffer(uint32 Size, uint32 InUsage, ERHIAccess InResourceState, FRHIResourceCreateInfo& CreateInfo) final override;
-	virtual void* LockVertexBuffer_BottomOfPipe(FRHICommandListImmediate& RHICmdList, FRHIVertexBuffer* VertexBuffer, uint32 Offset, uint32 SizeRHI, EResourceLockMode LockMode) final override;
-	virtual void UnlockVertexBuffer_BottomOfPipe(FRHICommandListImmediate& RHICmdList, FRHIVertexBuffer* VertexBuffer) final override;
 	virtual void RHICopyVertexBuffer(FRHIVertexBuffer* SourceBuffer, FRHIVertexBuffer* DestBuffer) final override;
 	virtual void RHITransferBufferUnderlyingResource(FRHIBuffer* DestBuffer, FRHIBuffer* SrcBuffer) final override;
 	virtual FStructuredBufferRHIRef RHICreateStructuredBuffer(uint32 Stride, uint32 Size, uint32 InUsage, ERHIAccess InResourceState, FRHIResourceCreateInfo& CreateInfo) final override;
-	virtual void* LockStructuredBuffer_BottomOfPipe(FRHICommandListImmediate& RHICmdList, FRHIStructuredBuffer* StructuredBuffer, uint32 Offset, uint32 SizeRHI, EResourceLockMode LockMode) final override;
-	virtual void UnlockStructuredBuffer_BottomOfPipe(FRHICommandListImmediate& RHICmdList, FRHIStructuredBuffer* StructuredBuffer) final override;
+	virtual void* LockBuffer_BottomOfPipe(FRHICommandListImmediate& RHICmdList, FRHIBuffer* Buffer, uint32 Offset, uint32 Size, EResourceLockMode LockMode) final override;
+	virtual void UnlockBuffer_BottomOfPipe(FRHICommandListImmediate& RHICmdList, FRHIBuffer* Buffer) final override;
 	virtual FUnorderedAccessViewRHIRef RHICreateUnorderedAccessView(FRHIBuffer* Buffer, bool bUseUAVCounter, bool bAppendBuffer) final override;
 	virtual FUnorderedAccessViewRHIRef RHICreateUnorderedAccessView(FRHITexture* Texture, uint32 MipLevel) final override;
 	virtual FUnorderedAccessViewRHIRef RHICreateUnorderedAccessView(FRHIBuffer* Buffer, uint8 Format) final override;
diff --git a/Engine/Source/Runtime/Windows/D3D11RHI/Private/D3D11StructuredBuffer.cpp b/Engine/Source/Runtime/Windows/D3D11RHI/Private/D3D11StructuredBuffer.cpp
--- a/Engine/Source/Runtime/Windows/D3D11RHI/Private/D3D11StructuredBuffer.cpp
+++ b/Engine/Source/Runtime/Windows/D3D11RHI/Private/D3D11StructuredBuffer.cpp
@@ -99,106 +99,3 @@ FStructuredBufferRHIRef FD3D11DynamicRHI::CreateStructuredBuffer_RenderThread(
 {
 	return RHICreateStructuredBuffer(Stride, Size, InUsage, InResourceState, CreateInfo);
 }
-
-void* FD3D11DynamicRHI::LockStructuredBuffer_BottomOfPipe(FRHICommandListImmediate& RHICmdList, FRHIStructuredBuffer* StructuredBufferRHI,uint32 Offset,uint32 Size,EResourceLockMode LockMode)
-{
-	FD3D11Buffer* StructuredBuffer = ResourceCast(StructuredBufferRHI);
-	
-	// If this resource is bound to the device, unbind it
-	ConditionalClearShaderResource(StructuredBuffer, true);
-
-	// Determine whether the Structured buffer is dynamic or not.
-	D3D11_BUFFER_DESC Desc;
-	StructuredBuffer->Resource->GetDesc(&Desc);
-	const bool bIsDynamic = (Desc.Usage == D3D11_USAGE_DYNAMIC);
-
-	FD3D11LockedKey LockedKey(StructuredBuffer->Resource);
-	FD3D11LockedData LockedData;
-
-	if(bIsDynamic)
-	{
-		check(LockMode == RLM_WriteOnly);
-
-		// If the buffer is dynamic, map its memory for writing.
-		D3D11_MAPPED_SUBRESOURCE MappedSubresource;
-		VERIFYD3D11RESULT_EX(Direct3DDeviceIMContext->Map(StructuredBuffer->Resource,0,D3D11_MAP_WRITE_DISCARD,0,&MappedSubresource), Direct3DDevice);
-		LockedData.SetData(MappedSubresource.pData);
-		LockedData.Pitch = MappedSubresource.RowPitch;
-	}
-	else
-	{
-		if(LockMode == RLM_ReadOnly)
-		{
-			// If the static buffer is being locked for reading, create a staging buffer.
-			D3D11_BUFFER_DESC StagingBufferDesc;
-			ZeroMemory( &StagingBufferDesc, sizeof( D3D11_BUFFER_DESC ) );
-			StagingBufferDesc.ByteWidth = Size;
-			StagingBufferDesc.Usage = D3D11_USAGE_STAGING;
-			StagingBufferDesc.BindFlags = 0;
-			StagingBufferDesc.CPUAccessFlags = D3D11_CPU_ACCESS_READ;
-			StagingBufferDesc.MiscFlags = 0;
-			TRefCountPtr<ID3D11Buffer> StagingStructuredBuffer;
-			VERIFYD3D11RESULT_EX(Direct3DDevice->CreateBuffer(&StagingBufferDesc,NULL,StagingStructuredBuffer.GetInitReference()), Direct3DDevice);
-			LockedData.StagingResource = StagingStructuredBuffer;
-
-			// Copy the contents of the Structured buffer to the staging buffer.
-			Direct3DDeviceIMContext->CopyResource(StagingStructuredBuffer,StructuredBuffer->Resource);
-
-			// Map the staging buffer's memory for reading.
-			D3D11_MAPPED_SUBRESOURCE MappedSubresource;
-			VERIFYD3D11RESULT_EX(Direct3DDeviceIMContext->Map(StagingStructuredBuffer,0,D3D11_MAP_READ,0,&MappedSubresource), Direct3DDevice);
-			LockedData.SetData(MappedSubresource.pData);
-			LockedData.Pitch = MappedSubresource.RowPitch;
-		}
-		else
-		{
-			// If the static buffer is being locked for writing, allocate memory for the contents to be written to.
-			LockedData.AllocData(Desc.ByteWidth);
-			LockedData.Pitch = Desc.ByteWidth;
-		}
-	}
-
-	// Add the lock to the lock map.
-	AddLockedData(LockedKey, LockedData);
-
-	// Return the offset pointer
-	return (void*)((uint8*)LockedData.GetData() + Offset);
-}
-
-void FD3D11DynamicRHI::UnlockStructuredBuffer_BottomOfPipe(FRHICommandListImmediate& RHICmdList, FRHIStructuredBuffer* StructuredBufferRHI)
-{
-	FD3D11Buffer* StructuredBuffer = ResourceCast(StructuredBufferRHI);
-
-	// Determine whether the Structured buffer is dynamic or not.
-	D3D11_BUFFER_DESC Desc;
-	StructuredBuffer->Resource->GetDesc(&Desc);
-	const bool bIsDynamic = (Desc.Usage == D3D11_USAGE_DYNAMIC);
-
-	// Find the outstanding lock for this VB.
-	FD3D11LockedData LockedData;
-	verifyf(RemoveLockedData(FD3D11LockedKey(StructuredBuffer->Resource), LockedData), TEXT("Structured buffer is not locked"));
-
-	if(bIsDynamic)
-	{
-		// If the VB is dynamic, its memory was mapped directly; unmap it.
-		Direct3DDeviceIMContext->Unmap(StructuredBuffer->Resource,0);
-	}
-	else
-	{
-		// If the static VB lock involved a staging resource, it was locked for reading.
-		if(LockedData.StagingResource)
-		{
-			// Unmap the staging buffer's memory.
-			ID3D11Buffer* StagingBuffer = (ID3D11Buffer*)LockedData.StagingResource.GetReference();
-			Direct3DDeviceIMContext->Unmap(StagingBuffer,0);
-		}
-		else 
-		{
-			// Copy the contents of the temporary memory buffer allocated for writing into the VB.
-			Direct3DDeviceIMContext->UpdateSubresource(StructuredBuffer->Resource,0,NULL,LockedData.GetData(),LockedData.Pitch,0);
-
-			// Free the temporary memory buffer.
-			LockedData.FreeData();
-		}
-	}
-}
diff --git a/Engine/Source/Runtime/Windows/D3D11RHI/Private/D3D11VertexBuffer.cpp b/Engine/Source/Runtime/Windows/D3D11RHI/Private/D3D11VertexBuffer.cpp
--- a/Engine/Source/Runtime/Windows/D3D11RHI/Private/D3D11VertexBuffer.cpp
+++ b/Engine/Source/Runtime/Windows/D3D11RHI/Private/D3D11VertexBuffer.cpp
@@ -126,126 +126,6 @@ FVertexBufferRHIRef FD3D11DynamicRHI::CreateVertexBuffer_RenderThread(
 	return RHICreateVertexBuffer(Size, InUsage, InResourceState, CreateInfo);
 }
 
-void* FD3D11DynamicRHI::LockVertexBuffer_BottomOfPipe(FRHICommandListImmediate& RHICmdList, FRHIVertexBuffer* VertexBufferRHI,uint32 Offset,uint32 Size,EResourceLockMode LockMode)
-{
-	check(Size > 0);
-
-	FD3D11Buffer* VertexBuffer = ResourceCast(VertexBufferRHI);
-	
-	// If this resource is bound to the device, unbind it
-	ConditionalClearShaderResource(VertexBuffer, true);
-
-	// Determine whether the vertex buffer is dynamic or not.
-	D3D11_BUFFER_DESC Desc;
-	VertexBuffer->Resource->GetDesc(&Desc);
-	const bool bIsDynamic = (Desc.Usage == D3D11_USAGE_DYNAMIC);
-
-	FD3D11LockedKey LockedKey(VertexBuffer->Resource);
-	FD3D11LockedData LockedData;
-
-	if(bIsDynamic)
-	{
-		check(LockMode == RLM_WriteOnly || LockMode == RLM_WriteOnly_NoOverwrite);
-
-		// If the buffer is dynamic, map its memory for writing.
-		D3D11_MAPPED_SUBRESOURCE MappedSubresource;
-
-		if (LockMode == RLM_WriteOnly)
-		{
-			VERIFYD3D11RESULT_EX(Direct3DDeviceIMContext->Map(VertexBuffer->Resource,0,D3D11_MAP_WRITE_DISCARD,0,&MappedSubresource), Direct3DDevice);
-		}
-		else
-		{
-			VERIFYD3D11RESULT_EX(Direct3DDeviceIMContext->Map(VertexBuffer->Resource,0,D3D11_MAP_WRITE_NO_OVERWRITE,0,&MappedSubresource), Direct3DDevice);
-		}
-		
-		LockedData.SetData(MappedSubresource.pData);
-		LockedData.Pitch = MappedSubresource.RowPitch;
-	}
-	else
-	{
-		if(LockMode == RLM_ReadOnly)
-		{
-			// If the static buffer is being locked for reading, create a staging buffer.
-			D3D11_BUFFER_DESC StagingBufferDesc;
-			ZeroMemory( &StagingBufferDesc, sizeof( D3D11_BUFFER_DESC ) );
-			StagingBufferDesc.ByteWidth = Size;
-			StagingBufferDesc.Usage = D3D11_USAGE_STAGING;
-			StagingBufferDesc.BindFlags = 0;
-			StagingBufferDesc.CPUAccessFlags = D3D11_CPU_ACCESS_READ;
-			StagingBufferDesc.MiscFlags = 0;
-			TRefCountPtr<ID3D11Buffer> StagingVertexBuffer;
-			VERIFYD3D11RESULT_EX(Direct3DDevice->CreateBuffer(&StagingBufferDesc,NULL,StagingVertexBuffer.GetInitReference()), Direct3DDevice);
-			LockedData.StagingResource = StagingVertexBuffer;
-
-			// Copy the contents of the vertex buffer to the staging buffer.
-			D3D11_BOX SourceBox;
-			SourceBox.left = Offset;
-			SourceBox.right = Offset + Size;
-			SourceBox.top = SourceBox.front = 0;
-			SourceBox.bottom = SourceBox.back = 1;
-			Direct3DDeviceIMContext->CopySubresourceRegion(StagingVertexBuffer,0,0,0,0,VertexBuffer->Resource,0,&SourceBox);
-
-			// Map the staging buffer's memory for reading.
-			D3D11_MAPPED_SUBRESOURCE MappedSubresource;
-			VERIFYD3D11RESULT_EX(Direct3DDeviceIMContext->Map(StagingVertexBuffer,0,D3D11_MAP_READ,0,&MappedSubresource), Direct3DDevice);
-			LockedData.SetData(MappedSubresource.pData);
-			LockedData.Pitch = MappedSubresource.RowPitch;
-			Offset = 0;
-		}
-		else
-		{
-			// If the static buffer is being locked for writing, allocate memory for the contents to be written to.
-			LockedData.AllocData(Desc.ByteWidth);
-			LockedData.Pitch = Desc.ByteWidth;
-		}
-	}
-
-	// Add the lock to the lock map.
-	AddLockedData(LockedKey, LockedData);
-
-	// Return the offset pointer
-	return (void*)((uint8*)LockedData.GetData() + Offset);
-}
-
-void FD3D11DynamicRHI::UnlockVertexBuffer_BottomOfPipe(FRHICommandListImmediate& RHICmdList, FRHIVertexBuffer* VertexBufferRHI)
-{
-	FD3D11Buffer* VertexBuffer = ResourceCast(VertexBufferRHI);
-
-	// Determine whether the vertex buffer is dynamic or not.
-	D3D11_BUFFER_DESC Desc;
-	VertexBuffer->Resource->GetDesc(&Desc);
-	const bool bIsDynamic = (Desc.Usage == D3D11_USAGE_DYNAMIC);
-
-	// Find the outstanding lock for this VB.
-	FD3D11LockedData LockedData;
-	verifyf(RemoveLockedData(FD3D11LockedKey(VertexBuffer->Resource), LockedData), TEXT("Vertex buffer is not locked"));
-
-	if(bIsDynamic)
-	{
-		// If the VB is dynamic, its memory was mapped directly; unmap it.
-		Direct3DDeviceIMContext->Unmap(VertexBuffer->Resource,0);
-	}
-	else
-	{
-		// If the static VB lock involved a staging resource, it was locked for reading.
-		if(LockedData.StagingResource)
-		{
-			// Unmap the staging buffer's memory.
-			ID3D11Buffer* StagingBuffer = (ID3D11Buffer*)LockedData.StagingResource.GetReference();
-			Direct3DDeviceIMContext->Unmap(StagingBuffer,0);
-		}
-		else 
-		{
-			// Copy the contents of the temporary memory buffer allocated for writing into the VB.
-			Direct3DDeviceIMContext->UpdateSubresource(VertexBuffer->Resource,0,NULL,LockedData.GetData(),LockedData.Pitch,0);
-
-			// Free the temporary memory buffer.
-			LockedData.FreeData();
-		}
-	}
-}
-
 void FD3D11DynamicRHI::RHICopyVertexBuffer(FRHIVertexBuffer* SourceBufferRHI, FRHIVertexBuffer* DestBufferRHI)
 {
 	FD3D11Buffer* SourceBuffer = ResourceCast(SourceBufferRHI);
