Deprecate old LockBuffer functions

diff --git a/Engine/Source/Runtime/RHI/Public/RHICommandList.h b/Engine/Source/Runtime/RHI/Public/RHICommandList.h
--- a/Engine/Source/Runtime/RHI/Public/RHICommandList.h
+++ b/Engine/Source/Runtime/RHI/Public/RHICommandList.h
@@ -3920,11 +3920,13 @@ public:
 		return CreateAndLockIndexBuffer(Stride, Size, (EBufferUsageFlags) InUsage, ResourceState, CreateInfo, OutDataBuffer);
 	}
 	
+	UE_DEPRECATED(5.0, "Buffer locks have been unified. Use LockBuffer() instead.")
 	FORCEINLINE void* LockIndexBuffer(FRHIIndexBuffer* IndexBuffer, uint32 Offset, uint32 SizeRHI, EResourceLockMode LockMode)
 	{
 		return GDynamicRHI->RHILockBuffer(*this, IndexBuffer, Offset, SizeRHI, LockMode);
 	}
 	
+	UE_DEPRECATED(5.0, "Buffer locks have been unified. Use UnlockBuffer() instead.")
 	FORCEINLINE void UnlockIndexBuffer(FRHIIndexBuffer* IndexBuffer)
 	{
 		GDynamicRHI->RHIUnlockBuffer(*this, IndexBuffer);
@@ -3951,11 +3953,13 @@ public:
 		return CreateAndLockVertexBuffer(Size, (EBufferUsageFlags) InUsage, ResourceState, CreateInfo, OutDataBuffer);
 	}
 
+	UE_DEPRECATED(5.0, "Buffer locks have been unified. Use LockBuffer() instead.")
 	FORCEINLINE void* LockVertexBuffer(FRHIVertexBuffer* VertexBuffer, uint32 Offset, uint32 SizeRHI, EResourceLockMode LockMode)
 	{
 		return GDynamicRHI->RHILockBuffer(*this, VertexBuffer, Offset, SizeRHI, LockMode);
 	}
 	
+	UE_DEPRECATED(5.0, "Buffer locks have been unified. Use UnlockBuffer() instead.")
 	FORCEINLINE void UnlockVertexBuffer(FRHIVertexBuffer* VertexBuffer)
 	{
 		GDynamicRHI->RHIUnlockBuffer(*this, VertexBuffer);
@@ -3968,11 +3972,13 @@ public:
 		GDynamicRHI->RHICopyVertexBuffer(SourceBuffer,DestBuffer);
 	}
 
+	UE_DEPRECATED(5.0, "Buffer locks have been unified. Use LockBuffer() instead.")
 	FORCEINLINE void* LockStructuredBuffer(FRHIStructuredBuffer* StructuredBuffer, uint32 Offset, uint32 SizeRHI, EResourceLockMode LockMode)
 	{
 		return GDynamicRHI->RHILockBuffer(*this, StructuredBuffer, Offset, SizeRHI, LockMode);
 	}
 
+	UE_DEPRECATED(5.0, "Buffer locks have been unified. Use UnlockBuffer() instead.")
 	FORCEINLINE void UnlockStructuredBuffer(FRHIStructuredBuffer* StructuredBuffer)
 	{
 		GDynamicRHI->RHIUnlockBuffer(*this, StructuredBuffer);
@@ -4940,11 +4946,13 @@ FORCEINLINE FIndexBufferRHIRef RHIAsyncCreateIndexBuffer(uint32 Stride, uint32 S
 	return RHIAsyncCreateIndexBuffer(Stride, Size, InUsage, ResourceState, CreateInfo);
 }
 
+UE_DEPRECATED(5.0, "Buffer locks have been unified. Use RHILockBuffer() instead.")
 FORCEINLINE void* RHILockIndexBuffer(FRHIIndexBuffer* IndexBuffer, uint32 Offset, uint32 Size, EResourceLockMode LockMode)
 {
 	return FRHICommandListExecutor::GetImmediateCommandList().LockBuffer(IndexBuffer, Offset, Size, LockMode);
 }
 
+UE_DEPRECATED(5.0, "Buffer locks have been unified. Use RHIUnlockBuffer() instead.")
 FORCEINLINE void RHIUnlockIndexBuffer(FRHIIndexBuffer* IndexBuffer)
 {
 	 FRHICommandListExecutor::GetImmediateCommandList().UnlockBuffer(IndexBuffer);
@@ -4979,11 +4987,13 @@ FORCEINLINE FVertexBufferRHIRef RHIAsyncCreateVertexBuffer(uint32 Size, uint32 I
 	return RHIAsyncCreateVertexBuffer(Size, InUsage, ResourceState, CreateInfo);
 }
 
+UE_DEPRECATED(5.0, "Buffer locks have been unified. Use RHILockBuffer() instead.")
 FORCEINLINE void* RHILockVertexBuffer(FRHIVertexBuffer* VertexBuffer, uint32 Offset, uint32 SizeRHI, EResourceLockMode LockMode)
 {
 	return FRHICommandListExecutor::GetImmediateCommandList().LockBuffer(VertexBuffer, Offset, SizeRHI, LockMode);
 }
 
+UE_DEPRECATED(5.0, "Buffer locks have been unified. Use RHIUnlockBuffer() instead.")
 FORCEINLINE void RHIUnlockVertexBuffer(FRHIVertexBuffer* VertexBuffer)
 {
 	 FRHICommandListExecutor::GetImmediateCommandList().UnlockBuffer(VertexBuffer);
@@ -5001,11 +5011,13 @@ FORCEINLINE FStructuredBufferRHIRef RHICreateStructuredBuffer(uint32 Stride, uin
 	return RHICreateStructuredBuffer(Stride, Size, InUsage, ResourceState, CreateInfo);
 }
 
+UE_DEPRECATED(5.0, "Buffer locks have been unified. Use RHILockBuffer() instead.")
 FORCEINLINE void* RHILockStructuredBuffer(FRHIStructuredBuffer* StructuredBuffer, uint32 Offset, uint32 SizeRHI, EResourceLockMode LockMode)
 {
 	return FRHICommandListExecutor::GetImmediateCommandList().LockBuffer(StructuredBuffer, Offset, SizeRHI, LockMode);
 }
 
+UE_DEPRECATED(5.0, "Buffer locks have been unified. Use RHIUnlockBuffer() instead.")
 FORCEINLINE void RHIUnlockStructuredBuffer(FRHIStructuredBuffer* StructuredBuffer)
 {
 	 FRHICommandListExecutor::GetImmediateCommandList().UnlockBuffer(StructuredBuffer);
