Improved RHI validation logging / messages.
 - Modified GetDebugName() to return the full description in the form "ResourceName" (0xaddress) (subresource index), and included this in the assert log line as well as assert details.
 - Renamed "BeginTrace" to "EndTrace" in EndTransition barrier log string.

diff --git a/Engine/Source/Runtime/RHI/Private/RHIValidation.cpp b/Engine/Source/Runtime/RHI/Private/RHIValidation.cpp
--- a/Engine/Source/Runtime/RHI/Private/RHIValidation.cpp
+++ b/Engine/Source/Runtime/RHI/Private/RHIValidation.cpp
@@ -393,15 +393,14 @@ namespace RHIValidation
 		return AccessMask;
 	}
 
-// Warning: this prefix expects a string argument for the failure reason, make sure you add it.
-#define BARRIER_TRACKER_LOG_PREFIX_REASON TEXT("RHI validation failed: %s\n\n")\
+#define BARRIER_TRACKER_LOG_PREFIX_REASON(ReasonString) TEXT("RHI validation failed: " ReasonString ":\n\n")\
 	TEXT("--------------------------------------------------------------------\n")\
 	TEXT("              RHI Resource Transition Validation Error              \n")\
 	TEXT("--------------------------------------------------------------------\n")\
 	TEXT("\n\n")
 	
 // Warning: this prefix expects a string argument for the resource name, make sure you add it.
-#define BARRIER_TRACKER_LOG_PREFIX_RESNAME TEXT("RHI validation failed for resource \"%s\":\n\n")\
+#define BARRIER_TRACKER_LOG_PREFIX_RESNAME TEXT("RHI validation failed for resource: %s:\n\n")\
 	TEXT("--------------------------------------------------------------------\n")\
 	TEXT("              RHI Resource Transition Validation Error              \n")\
 	TEXT("--------------------------------------------------------------------\n")\
@@ -411,10 +410,31 @@ namespace RHIValidation
 	TEXT("--------------------------------------------------------------------\n")\
 	TEXT("\n\n")
 
-	static inline const TCHAR* GetResourceDebugName(FResource* Resource)
+	static inline FString GetResourceDebugName(FResource const* Resource, FSubresourceIndex const& SubresourceIndex)
 	{
 		const TCHAR* DebugName = Resource->GetDebugName();
-		return DebugName ? DebugName : TEXT("Unnamed");
+		if (!DebugName)
+		{
+			DebugName = TEXT("Unnamed");
+		}
+
+		if (SubresourceIndex.IsWholeResource())
+		{
+			return FString::Printf(
+				TEXT("\"%s\" (0x%p) (Whole Resource)"),
+				DebugName,
+				Resource);
+		}
+		else
+		{
+			return FString::Printf(
+				TEXT("\"%s\" (0x%p) (Mip %d, Slice %d, Plane %d)"),
+				DebugName,
+				Resource, 
+				SubresourceIndex.MipIndex,
+				SubresourceIndex.ArraySlice,
+				SubresourceIndex.PlaneIndex);
+		}
 	}
 
 	static inline FString GetReasonString_MissingBarrier(
@@ -422,19 +442,17 @@ namespace RHIValidation
 		const FState& CurrentState,
 		const FState& RequiredState)
 	{
-		const TCHAR* DebugName = GetResourceDebugName(Resource);
+		FString DebugName = GetResourceDebugName(Resource, SubresourceIndex);
 		return FString::Printf(
 			BARRIER_TRACKER_LOG_PREFIX_RESNAME
-			TEXT("Attempted to access resource \"%s\" (0x%p) (%s) from a hardware unit it is not currently accessible from. A resource transition is required.\n\n")
+			TEXT("Attempted to access resource %s from a hardware unit it is not currently accessible from. A resource transition is required.\n\n")
 			TEXT("    --- Allowed access states for this resource are: %s\n")
 			TEXT("    --- Required access states are:                  %s\n")
 			TEXT("    --- Allowed pipelines for this resource are:     %s\n")
 			TEXT("    --- Required pipelines are:                      %s\n")
 			BARRIER_TRACKER_LOG_SUFFIX,
-			DebugName,
-			DebugName,
-			Resource,
-			*SubresourceIndex.ToString(),
+			*DebugName,
+			*DebugName,
 			*GetRHIAccessName(CurrentState.Access),
 			*GetRHIAccessName(RequiredState.Access),
 			*GetRHIPipelineName(CurrentState.Pipelines),
@@ -464,20 +482,18 @@ namespace RHIValidation
 		const FState& AttemptedState,
 		void* CreateTrace, void* BeginTrace)
 	{
-		const TCHAR* DebugName = GetResourceDebugName(Resource);
+		FString DebugName = GetResourceDebugName(Resource, SubresourceIndex);
 		return FString::Printf(
 			BARRIER_TRACKER_LOG_PREFIX_RESNAME
-			TEXT("Attempted to access resource \"%s\" (0x%p) (%s) whilst an asynchronous resource transition is in progress. A call to RHIEndTransitions() must be made before the resource can be accessed again.\n\n")
+			TEXT("Attempted to access resource %s whilst an asynchronous resource transition is in progress. A call to RHIEndTransitions() must be made before the resource can be accessed again.\n\n")
 			TEXT("    --- Pending access states for this resource are: %s\n")
 			TEXT("    --- Attempted access states are:                 %s\n")
 			TEXT("    --- Pending pipelines for this resource are:     %s\n")
 			TEXT("    --- Attempted pipelines are:                     %s\n")
 			TEXT("%s")
 			BARRIER_TRACKER_LOG_SUFFIX,
-			DebugName,
-			DebugName,
-			Resource,
-			*SubresourceIndex.ToString(),
+			*DebugName,
+			*DebugName,
 			*GetRHIAccessName(PendingState.Access),
 			*GetRHIAccessName(AttemptedState.Access),
 			*GetRHIPipelineName(PendingState.Pipelines),
@@ -491,20 +507,18 @@ namespace RHIValidation
 		const FState& TargetState,
 		void* CreateTrace, void* BeginTrace)
 	{
-		const TCHAR* DebugName = GetResourceDebugName(Resource);
+		FString DebugName = GetResourceDebugName(Resource, SubresourceIndex);
 		return FString::Printf(
 			BARRIER_TRACKER_LOG_PREFIX_RESNAME
-			TEXT("Attempted to begin a resource transition for resource \"%s\" (0x%p) (%s) whilst a previous asynchronous resource transition is already in progress. A call to RHIEndTransitions() must be made before the resource can be transitioned again.\n\n")
+			TEXT("Attempted to begin a resource transition for resource %s whilst a previous asynchronous resource transition is already in progress. A call to RHIEndTransitions() must be made before the resource can be transitioned again.\n\n")
 			TEXT("    --- Pending access states for this resource are:              %s\n")
 			TEXT("    --- Attempted access states for the duplicate transition are: %s\n")
 			TEXT("    --- Pending pipelines for this resource are:                  %s\n")
 			TEXT("    --- Attempted pipelines for the duplicate transition are:     %s\n")
 			TEXT("%s")
 			BARRIER_TRACKER_LOG_SUFFIX,
-			DebugName,
-			DebugName,
-			Resource,
-			*SubresourceIndex.ToString(),
+			*DebugName,
+			*DebugName,
 			*GetRHIAccessName(PendingState.Access),
 			*GetRHIAccessName(TargetState.Access),
 			*GetRHIPipelineName(PendingState.Pipelines),
@@ -517,18 +531,16 @@ namespace RHIValidation
 		const FState& ActualCurrentState,
 		const FState& CurrentStateFromRHI)
 	{
-		const TCHAR* DebugName = GetResourceDebugName(Resource);
+		FString DebugName = GetResourceDebugName(Resource, SubresourceIndex);
 		return FString::Printf(
 			BARRIER_TRACKER_LOG_PREFIX_RESNAME
-			TEXT("Attempted to begin a resource transition for resource \"%s\" (0x%p) (%s) on the wrong pipeline(s) (\"%s\"). The resource is currently accessible on the \"%s\" pipeline(s).\n\n")
+			TEXT("Attempted to begin a resource transition for resource %s on the wrong pipeline(s) (\"%s\"). The resource is currently accessible on the \"%s\" pipeline(s).\n\n")
 			TEXT("    --- Current access states for this resource are: %s\n")
 			TEXT("    --- Attempted access states are:                 %s\n\n")
 			TEXT("    --- Ensure that resource transitions are issued on the correct pipeline.\n")
 			BARRIER_TRACKER_LOG_SUFFIX,
-			DebugName,
-			DebugName,
-			Resource,
-			*SubresourceIndex.ToString(),
+			*DebugName,
+			*DebugName,
 			*GetRHIPipelineName(CurrentStateFromRHI.Pipelines),
 			*GetRHIPipelineName(ActualCurrentState.Pipelines),
 			*GetRHIAccessName(ActualCurrentState.Access),
@@ -540,21 +552,19 @@ namespace RHIValidation
 		const FState& CurrentState,
 		const FState& CurrentStateFromRHI)
 	{
-		const TCHAR* DebugName = GetResourceDebugName(Resource);
+		FString DebugName = GetResourceDebugName(Resource, SubresourceIndex);
 		return FString::Printf(
 			BARRIER_TRACKER_LOG_PREFIX_RESNAME
-			TEXT("The explicit previous state \"%s\" does not match the tracked current state \"%s\" for the resource \"%s\" (0x%p) (%s).\n")
+			TEXT("The explicit previous state \"%s\" does not match the tracked current state \"%s\" for the resource %s.\n")
 			TEXT("    --- Allowed pipelines for this resource are:                           %s\n")
 			TEXT("    --- Previous pipelines passed as part of the resource transition were: %s\n\n")
 			TEXT("    --- The best solution is to correct the explicit previous state passed for the resource in the call to RHICreateTransition().\n")
 			TEXT("    --- Alternatively, use ERHIAccess::Unknown if the actual previous state cannot be determined. Unknown previous resource states have a performance impact so should be avoided if possible.\n")
 			BARRIER_TRACKER_LOG_SUFFIX,
-			DebugName,
+			*DebugName,
 			*GetRHIAccessName(CurrentStateFromRHI.Access),
 			*GetRHIAccessName(CurrentState.Access),
-			DebugName,
-			Resource,
-			*SubresourceIndex.ToString(),
+			*DebugName,
 			*GetRHIPipelineName(CurrentState.Pipelines),
 			*GetRHIPipelineName(CurrentStateFromRHI.Pipelines));
 	}
@@ -564,37 +574,33 @@ namespace RHIValidation
 		const FState& TargetState,
 		const FState& TargetStateFromRHI)
 	{
-		const TCHAR* DebugName = GetResourceDebugName(Resource);
+		FString DebugName = GetResourceDebugName(Resource, SubresourceIndex);
 		return FString::Printf(
 			BARRIER_TRACKER_LOG_PREFIX_RESNAME
-			TEXT("The expected target state \"%s\" on pipe \"%s\" in end transition does not match the tracked target state \"%s\" on pipe \"%s\" for the resource \"%s\" (0x%p) (%s).\n")
+			TEXT("The expected target state \"%s\" on pipe \"%s\" in end transition does not match the tracked target state \"%s\" on pipe \"%s\" for the resource %s.\n")
 			TEXT("    --- The call to EndTransition() is mismatched with the another BeginTransition() with different states.\n")
 			BARRIER_TRACKER_LOG_SUFFIX,
-			DebugName,
+			*DebugName,
 			*GetRHIAccessName(TargetStateFromRHI.Access),
 			*GetRHIPipelineName(TargetState.Pipelines),
 			*GetRHIAccessName(TargetState.Access),
 			*GetRHIPipelineName(TargetStateFromRHI.Pipelines),
-			DebugName,
-			Resource,
-			*SubresourceIndex.ToString());
+			*DebugName);
 	}
 
 	static inline FString GetReasonString_UnnecessaryTransition(
 		FResource* Resource, FSubresourceIndex const& SubresourceIndex,
 		const FState& CurrentState)
 	{
-		const TCHAR* DebugName = GetResourceDebugName(Resource);
+		FString DebugName = GetResourceDebugName(Resource, SubresourceIndex);
 		return FString::Printf(
 			BARRIER_TRACKER_LOG_PREFIX_RESNAME
-			TEXT("Attempted to begin a resource transition for the resource \"%s\" (0x%p) (%s) to the \"%s\" state on the \"%s\" pipe, but the resource is already in this state. The resource transition is unnecessary.\n")
+			TEXT("Attempted to begin a resource transition for the resource %s to the \"%s\" state on the \"%s\" pipe, but the resource is already in this state. The resource transition is unnecessary.\n")
 			TEXT("    --- This is not fatal, but does have an effect on CPU and GPU performance. Consider refactoring rendering code to avoid unnecessary resource transitions.\n")
 			TEXT("    --- RenderGraph (RDG) is capable of handling resource transitions automatically.\n")
 			BARRIER_TRACKER_LOG_SUFFIX,
-			DebugName,
-			DebugName,
-			Resource,
-			*SubresourceIndex.ToString(),
+			*DebugName,
+			*DebugName,
 			*GetRHIAccessName(CurrentState.Access),
 			*GetRHIPipelineName(CurrentState.Pipelines));
 	}
@@ -602,22 +608,26 @@ namespace RHIValidation
 	static inline FString GetReasonString_MismatchedAllUAVsOverlapCall(bool bAllow)
 	{
 		return FString::Printf(
-			BARRIER_TRACKER_LOG_PREFIX_REASON
-			TEXT("Mismatched call to %sUAVOverlap. Ensure all calls to RHICmdList.BeginUAVOverlap() are paired with a call to RHICmdList.EndUAVOverlap().")
+			BARRIER_TRACKER_LOG_PREFIX_REASON("UAV overlap mismatch")
+			TEXT("Mismatched call to %sUAVOverlap.\n\n")
+			TEXT("    --- Ensure all calls to RHICmdList.BeginUAVOverlap() are paired with a call to RHICmdList.EndUAVOverlap().\n")
 			BARRIER_TRACKER_LOG_SUFFIX,
-			TEXT("UAV overlap mismatch"),
 			bAllow ? TEXT("Begin") : TEXT("End")
 		);
 	}
 
-	static inline FString GetReasonString_MismatchedExplicitUAVOverlapCall(bool bAllow)
+	static inline FString GetReasonString_MismatchedExplicitUAVOverlapCall(
+		FResource* Resource, FSubresourceIndex const& SubresourceIndex,
+		bool bAllow)
 	{
+		FString DebugName = GetResourceDebugName(Resource, SubresourceIndex);
 		return FString::Printf(
-			BARRIER_TRACKER_LOG_PREFIX_REASON
-			TEXT("Mismatched call to %sUAVOverlap(FRHIUnorderedAccessView*). Ensure all calls to RHICmdList.BeginUAVOverlap() are paired with a call to RHICmdList.EndUAVOverlap().")
+			BARRIER_TRACKER_LOG_PREFIX_REASON("UAV overlap mismatch")
+			TEXT("Mismatched call to %sUAVOverlap(FRHIUnorderedAccessView*) for the resource %s.\n\n")
+			TEXT("    --- Ensure all calls to RHICmdList.BeginUAVOverlap() are paired with a call to RHICmdList.EndUAVOverlap().\n")
 			BARRIER_TRACKER_LOG_SUFFIX,
-			TEXT("UAV overlap mismatch"),
-			bAllow ? TEXT("Begin") : TEXT("End")
+			bAllow ? TEXT("Begin") : TEXT("End"),
+			*DebugName
 		);
 	}
 
@@ -625,19 +635,17 @@ namespace RHIValidation
 		FResource* Resource, FSubresourceIndex const& SubresourceIndex,
 		const FState& CurrentState, const FState& RequiredState)
 	{
-		const TCHAR* DebugName = GetResourceDebugName(Resource);
+		FString DebugName = GetResourceDebugName(Resource, SubresourceIndex);
 		return FString::Printf(
 			BARRIER_TRACKER_LOG_PREFIX_RESNAME
-			TEXT("Attempted to access resource \"%s\" (0x%p) (%s) which was previously used with overlapping UAV access, but has not been transitioned since UAV overlap was disabled. A resource transition is required.\n\n")
+			TEXT("Attempted to access resource %s which was previously used with overlapping UAV access, but has not been transitioned since UAV overlap was disabled. A resource transition is required.\n\n")
 			TEXT("    --- Allowed access states for this resource are: %s\n")
 			TEXT("    --- Required access states are:                  %s\n")
 			TEXT("    --- Allowed pipelines for this resource are:     %s\n")
 			TEXT("    --- Required pipelines are:                      %s\n")
 			BARRIER_TRACKER_LOG_SUFFIX,
-			DebugName,
-			DebugName,
-			Resource,
-			*SubresourceIndex.ToString(),
+			*DebugName,
+			*DebugName,
 			*GetRHIAccessName(CurrentState.Access),
 			*GetRHIAccessName(RequiredState.Access),
 			*GetRHIPipelineName(CurrentState.Pipelines),
@@ -674,27 +682,24 @@ namespace RHIValidation
 		}
 	}
 
-	void* FSubresourceState::Log(FResource* Resource, FSubresourceIndex const& SubresourceIndex, void* CreateTrace, const TCHAR* Type, const TCHAR* LogStr)
+	void* FSubresourceState::Log(FResource* Resource, FSubresourceIndex const& SubresourceIndex, void* CreateTrace, const TCHAR* TracePrefix, const TCHAR* Type, const TCHAR* LogStr)
 	{
 		void* Trace = CaptureBacktrace();
 
 		if (CreateTrace)
 		{
-			FPlatformMisc::LowLevelOutputDebugStringf(TEXT("\n%s: (0x%p) (%s), Type: %s, %s, CreateTrace: 0x%p, BeginTrace: 0x%p\n"),
-				*Resource->DebugName,
-				Resource,
-				*SubresourceIndex.ToString(),
+			FPlatformMisc::LowLevelOutputDebugStringf(TEXT("\n%s: Type: %s, %s, CreateTrace: 0x%p, %sTrace: 0x%p\n"),
+				*GetResourceDebugName(Resource, SubresourceIndex),
 				Type,
 				LogStr,
 				CreateTrace,
+				TracePrefix,
 				Trace);
 		}
 		else
 		{
-			FPlatformMisc::LowLevelOutputDebugStringf(TEXT("\n%s: (0x%p) (%s), Type: %s, %s, Trace: 0x%p\n"),
-				*Resource->DebugName,
-				Resource,
-				*SubresourceIndex.ToString(),
+			FPlatformMisc::LowLevelOutputDebugStringf(TEXT("\n%s: Type: %s, %s, Trace: 0x%p\n"),
+				*GetResourceDebugName(Resource, SubresourceIndex),
 				Type,
 				LogStr,
 				Trace);
@@ -714,7 +719,7 @@ namespace RHIValidation
 #endif
 			)
 		{
-			BeginTrace = Log(Resource, SubresourceIndex, CreateTrace, TEXT("BeginTransition"), *FString::Printf(TEXT("Current: (%s) New: (%s), Flags: %s, Executing Pipeline: %s"),
+			BeginTrace = Log(Resource, SubresourceIndex, CreateTrace, TEXT("Begin"), TEXT("BeginTransition"), *FString::Printf(TEXT("Current: (%s) New: (%s), Flags: %s, Executing Pipeline: %s"),
 				*State.Current.ToString(),
 				*TargetState.ToString(),
 				*GetResourceTransitionFlagsName(NewFlags),
@@ -767,7 +772,7 @@ namespace RHIValidation
 #endif
 			)
 		{
-			Log(Resource, SubresourceIndex, CreateTrace, TEXT("EndTransition"), *FString::Printf(TEXT("Access: %s, Pipeline: %s, Executing Pipeline: %s"),
+			Log(Resource, SubresourceIndex, CreateTrace, TEXT("End"), TEXT("EndTransition"), *FString::Printf(TEXT("Access: %s, Pipeline: %s, Executing Pipeline: %s"),
 				*GetRHIAccessName(TargetState.Access),
 				*GetRHIPipelineName(TargetState.Pipelines),
 				*GetRHIPipelineName(ExecutingPipeline)
@@ -802,7 +807,7 @@ namespace RHIValidation
 #endif
 			)
 		{
-			Log(Resource, SubresourceIndex, nullptr, TEXT("Assert"), *FString::Printf(TEXT("Access: %s, Pipeline: %s"), 
+			Log(Resource, SubresourceIndex, nullptr, nullptr, TEXT("Assert"), *FString::Printf(TEXT("Access: %s, Pipeline: %s"), 
 				*GetRHIAccessName(RequiredState.Access),
 				*GetRHIPipelineName(RequiredState.Pipelines)));
 		}
@@ -838,11 +843,11 @@ namespace RHIValidation
 #endif
 			)
 		{
-			Log(Resource, SubresourceIndex, nullptr, TEXT("UAVOverlap"), *FString::Printf(TEXT("Allow: %s"), bAllow ? TEXT("True") : TEXT("False")));
+			Log(Resource, SubresourceIndex, nullptr, nullptr, TEXT("UAVOverlap"), *FString::Printf(TEXT("Allow: %s"), bAllow ? TEXT("True") : TEXT("False")));
 		}
 
 		FPipelineState& State = States[Pipeline];
-		RHI_VALIDATION_CHECK(State.bExplicitAllowUAVOverlap != bAllow, *GetReasonString_MismatchedExplicitUAVOverlapCall(bAllow));
+		RHI_VALIDATION_CHECK(State.bExplicitAllowUAVOverlap != bAllow, *GetReasonString_MismatchedExplicitUAVOverlapCall(Resource, SubresourceIndex, bAllow));
 		State.bExplicitAllowUAVOverlap = bAllow;
 	}
 
diff --git a/Engine/Source/Runtime/RHI/Public/RHIValidationCommon.h b/Engine/Source/Runtime/RHI/Public/RHIValidationCommon.h
--- a/Engine/Source/Runtime/RHI/Public/RHIValidationCommon.h
+++ b/Engine/Source/Runtime/RHI/Public/RHIValidationCommon.h
@@ -106,13 +106,6 @@ namespace RHIValidation
 				&& ArraySlice == kWholeResource
 				&& PlaneIndex == kWholeResource;
 		}
-
-		inline FString ToString() const
-		{
-			return IsWholeResource()
-				? FString(TEXT("Whole Resource"))
-				: FString::Printf(TEXT("Mip %d, Slice %d, Plane %d"), MipIndex, ArraySlice, PlaneIndex);
-		}
 	};
 
 	struct FState
@@ -182,7 +175,7 @@ namespace RHIValidation
 		void EndTransition     (FResource* Resource, FSubresourceIndex const& SubresourceIndex, const FState& CurrentStateFromRHI, const FState& TargetState, ERHIPipeline Pipeline, void* CreateTrace);
 		void Assert            (FResource* Resource, FSubresourceIndex const& SubresourceIndex, const FState& RequiredState, bool bAllowAllUAVsOverlap);
 		void SpecificUAVOverlap(FResource* Resource, ERHIPipeline Pipeline, FSubresourceIndex const& SubresourceIndex, bool bAllow);
-		void* Log              (FResource* Resource, FSubresourceIndex const& SubresourceIndex, void* CreateTrace, const TCHAR* Type, const TCHAR* LogStr);
+		void* Log              (FResource* Resource, FSubresourceIndex const& SubresourceIndex, void* CreateTrace, const TCHAR* TracePrefix, const TCHAR* Type, const TCHAR* LogStr);
 	};
 
 	struct FSubresourceRange
