Renamed SetGlobalUniformBuffers to SetStaticUniformBuffers for consistency.

diff --git a/Engine/Plugins/FX/Niagara/Source/Niagara/Private/NiagaraEmitterInstanceBatcher.cpp b/Engine/Plugins/FX/Niagara/Source/Niagara/Private/NiagaraEmitterInstanceBatcher.cpp
--- a/Engine/Plugins/FX/Niagara/Source/Niagara/Private/NiagaraEmitterInstanceBatcher.cpp
+++ b/Engine/Plugins/FX/Niagara/Source/Niagara/Private/NiagaraEmitterInstanceBatcher.cpp
@@ -1255,12 +1255,12 @@ void NiagaraEmitterInstanceBatcher::ExecuteAll(FRHICommandList& RHICmdList, FRHI
 
 	SCOPED_DRAW_EVENTF(RHICmdList, NiagaraEmitterInstanceBatcher_ExecuteAll, TEXT("NiagaraEmitterInstanceBatcher_ExecuteAll - TickStage(%d)"), TickStage);
 
-	FUniformBufferStaticBindings GlobalUniformBuffers;
+	FUniformBufferStaticBindings StaticUniformBuffers;
 	if (FRHIUniformBuffer* SceneTexturesUniformBuffer = GNiagaraViewDataManager.GetSceneTextureUniformParameters())
 	{
-		GlobalUniformBuffers.AddUniformBuffer(SceneTexturesUniformBuffer);
+		StaticUniformBuffers.AddUniformBuffer(SceneTexturesUniformBuffer);
 	}
-	SCOPED_UNIFORM_BUFFER_GLOBAL_BINDINGS(RHICmdList, GlobalUniformBuffers);
+	SCOPED_UNIFORM_BUFFER_STATIC_BINDINGS(RHICmdList, StaticUniformBuffers);
 
 	FMemMark Mark(FMemStack::Get());
 	TArray<FOverlappableTicks, TMemStackAllocator<> > SimPasses;
diff --git a/Engine/Source/Runtime/Apple/MetalRHI/Private/MetalCommands.cpp b/Engine/Source/Runtime/Apple/MetalRHI/Private/MetalCommands.cpp
--- a/Engine/Source/Runtime/Apple/MetalRHI/Private/MetalCommands.cpp
+++ b/Engine/Source/Runtime/Apple/MetalRHI/Private/MetalCommands.cpp
@@ -277,7 +277,7 @@ void FMetalRHICommandContext::RHISetGraphicsPipelineState(FRHIGraphicsPipelineSt
 	}
 }
 
-void FMetalRHICommandContext::RHISetGlobalUniformBuffers(const FUniformBufferStaticBindings& InUniformBuffers)
+void FMetalRHICommandContext::RHISetStaticUniformBuffers(const FUniformBufferStaticBindings& InUniformBuffers)
 {
 	FMemory::Memzero(GlobalUniformBuffers.GetData(), GlobalUniformBuffers.Num() * sizeof(FRHIUniformBuffer*));
 
diff --git a/Engine/Source/Runtime/Apple/MetalRHI/Public/MetalRHIContext.h b/Engine/Source/Runtime/Apple/MetalRHI/Public/MetalRHIContext.h
--- a/Engine/Source/Runtime/Apple/MetalRHI/Public/MetalRHIContext.h
+++ b/Engine/Source/Runtime/Apple/MetalRHI/Public/MetalRHIContext.h
@@ -124,7 +124,7 @@ public:
 	
 	virtual void RHISetGraphicsPipelineState(FRHIGraphicsPipelineState* GraphicsState, bool bApplyAdditionalState) final override;
 	
-	virtual void RHISetGlobalUniformBuffers(const FUniformBufferStaticBindings& InUniformBuffers) final override;
+	virtual void RHISetStaticUniformBuffers(const FUniformBufferStaticBindings& InUniformBuffers) final override;
 
 	/** Set the shader resource view of a surface. */
 	virtual void RHISetShaderTexture(FRHIGraphicsShader* Shader, uint32 TextureIndex, FRHITexture* NewTexture) final override;
diff --git a/Engine/Source/Runtime/D3D12RHI/Private/D3D12CommandContext.h b/Engine/Source/Runtime/D3D12RHI/Private/D3D12CommandContext.h
--- a/Engine/Source/Runtime/D3D12RHI/Private/D3D12CommandContext.h
+++ b/Engine/Source/Runtime/D3D12RHI/Private/D3D12CommandContext.h
@@ -253,7 +253,7 @@ public:
 	virtual void RHISetComputePipelineState(FRHIComputePipelineState* ComputePipelineState) final override;
 	virtual void RHIDispatchComputeShader(uint32 ThreadGroupCountX, uint32 ThreadGroupCountY, uint32 ThreadGroupCountZ) final override;
 	virtual void RHIDispatchIndirectComputeShader(FRHIBuffer* ArgumentBuffer, uint32 ArgumentOffset) final override;
-	virtual void RHISetGlobalUniformBuffers(const FUniformBufferStaticBindings& InUniformBuffers) final override;
+	virtual void RHISetStaticUniformBuffers(const FUniformBufferStaticBindings& InUniformBuffers) final override;
 	virtual void RHISetShaderTexture(FRHIComputeShader* PixelShader, uint32 TextureIndex, FRHITexture* NewTexture) final override;
 	virtual void RHISetShaderSampler(FRHIComputeShader* ComputeShader, uint32 SamplerIndex, FRHISamplerState* NewState) final override;
 	virtual void RHISetUAVParameter(FRHIPixelShader* PixelShaderRHI, uint32 UAVIndex, FRHIUnorderedAccessView* UAVRHI) final override;
@@ -619,9 +619,9 @@ public:
 	{
 		ContextRedirect(RHISetShaderResourceViewParameter(Shader, SamplerIndex, SRV));
 	}
-	FORCEINLINE virtual void RHISetGlobalUniformBuffers(const FUniformBufferStaticBindings& InUniformBuffers) final override
+	FORCEINLINE virtual void RHISetStaticUniformBuffers(const FUniformBufferStaticBindings& InUniformBuffers) final override
 	{
-		ContextRedirect(RHISetGlobalUniformBuffers(InUniformBuffers));
+		ContextRedirect(RHISetStaticUniformBuffers(InUniformBuffers));
 	}
 	FORCEINLINE virtual void RHISetShaderUniformBuffer(FRHIGraphicsShader* Shader, uint32 BufferIndex, FRHIUniformBuffer* Buffer) final override
 	{
diff --git a/Engine/Source/Runtime/D3D12RHI/Private/D3D12Commands.cpp b/Engine/Source/Runtime/D3D12RHI/Private/D3D12Commands.cpp
--- a/Engine/Source/Runtime/D3D12RHI/Private/D3D12Commands.cpp
+++ b/Engine/Source/Runtime/D3D12RHI/Private/D3D12Commands.cpp
@@ -428,7 +428,7 @@ void FD3D12CommandContext::RHIEndTransitions(TArrayView<const FRHITransition*> T
 	}
 }
 
-void FD3D12CommandContext::RHISetGlobalUniformBuffers(const FUniformBufferStaticBindings& InUniformBuffers)
+void FD3D12CommandContext::RHISetStaticUniformBuffers(const FUniformBufferStaticBindings& InUniformBuffers)
 {
 	FMemory::Memzero(GlobalUniformBuffers.GetData(), GlobalUniformBuffers.Num() * sizeof(FRHIUniformBuffer*));
 
diff --git a/Engine/Source/Runtime/OpenGLDrv/Private/OpenGLCommands.cpp b/Engine/Source/Runtime/OpenGLDrv/Private/OpenGLCommands.cpp
--- a/Engine/Source/Runtime/OpenGLDrv/Private/OpenGLCommands.cpp
+++ b/Engine/Source/Runtime/OpenGLDrv/Private/OpenGLCommands.cpp
@@ -1103,7 +1103,7 @@ void FOpenGLDynamicRHI::RHISetShaderTexture(FRHIComputeShader* ComputeShaderRHI,
 	}
 }
 
-void FOpenGLDynamicRHI::RHISetGlobalUniformBuffers(const FUniformBufferStaticBindings& InUniformBuffers)
+void FOpenGLDynamicRHI::RHISetStaticUniformBuffers(const FUniformBufferStaticBindings& InUniformBuffers)
 {
 	FMemory::Memzero(GlobalUniformBuffers.GetData(), GlobalUniformBuffers.Num() * sizeof(FRHIUniformBuffer*));
 
diff --git a/Engine/Source/Runtime/OpenGLDrv/Public/OpenGLDrv.h b/Engine/Source/Runtime/OpenGLDrv/Public/OpenGLDrv.h
--- a/Engine/Source/Runtime/OpenGLDrv/Public/OpenGLDrv.h
+++ b/Engine/Source/Runtime/OpenGLDrv/Public/OpenGLDrv.h
@@ -468,7 +468,7 @@ public:
 	virtual void RHISetUAVParameter(FRHIComputeShader* ComputeShader, uint32 UAVIndex, FRHIUnorderedAccessView* UAV, uint32 InitialCount) final override;
 	virtual void RHISetShaderResourceViewParameter(FRHIGraphicsShader* Shader, uint32 SamplerIndex, FRHIShaderResourceView* SRV) final override;
 	virtual void RHISetShaderResourceViewParameter(FRHIComputeShader* ComputeShader, uint32 SamplerIndex, FRHIShaderResourceView* SRV) final override;
-	virtual void RHISetGlobalUniformBuffers(const FUniformBufferStaticBindings& InUniformBuffers) final override;
+	virtual void RHISetStaticUniformBuffers(const FUniformBufferStaticBindings& InUniformBuffers) final override;
 	virtual void RHISetShaderUniformBuffer(FRHIGraphicsShader* Shader, uint32 BufferIndex, FRHIUniformBuffer* Buffer) final override;
 	virtual void RHISetShaderUniformBuffer(FRHIComputeShader* ComputeShader, uint32 BufferIndex, FRHIUniformBuffer* Buffer) final override;
 	virtual void RHISetShaderParameter(FRHIGraphicsShader* Shader, uint32 BufferIndex, uint32 BaseIndex, uint32 NumBytes, const void* NewValue) final override;
diff --git a/Engine/Source/Runtime/RHI/Private/RHICommandList.cpp b/Engine/Source/Runtime/RHI/Private/RHICommandList.cpp
--- a/Engine/Source/Runtime/RHI/Private/RHICommandList.cpp
+++ b/Engine/Source/Runtime/RHI/Private/RHICommandList.cpp
@@ -28,8 +28,8 @@ DECLARE_DWORD_COUNTER_STAT(TEXT("Immed. Command count"), STAT_ImmedCmdListCount,
 
 UE_TRACE_CHANNEL_DEFINE(RHICommandsChannel);
 
-#if VALIDATE_UNIFORM_BUFFER_GLOBAL_BINDINGS
-bool FScopedUniformBufferGlobalBindings::bRecursionGuard = false;
+#if VALIDATE_UNIFORM_BUFFER_STATIC_BINDINGS
+bool FScopedUniformBufferStaticBindings::bRecursionGuard = false;
 #endif
 
 #if !PLATFORM_USES_FIXED_RHI_CLASS
diff --git a/Engine/Source/Runtime/RHI/Public/RHICommandList.h b/Engine/Source/Runtime/RHI/Public/RHICommandList.h
--- a/Engine/Source/Runtime/RHI/Public/RHICommandList.h
+++ b/Engine/Source/Runtime/RHI/Public/RHICommandList.h
@@ -1648,11 +1648,11 @@ FRHICOMMAND_MACRO(FRHICommandClearColorTextures)
 	RHI_API void Execute(FRHICommandListBase& CmdList);
 };
 
-FRHICOMMAND_MACRO(FRHICommandSetGlobalUniformBuffers)
+FRHICOMMAND_MACRO(FRHICommandSetStaticUniformBuffers)
 {
 	FUniformBufferStaticBindings UniformBuffers;
 
-	FORCEINLINE_DEBUGGABLE FRHICommandSetGlobalUniformBuffers(const FUniformBufferStaticBindings& InUniformBuffers)
+	FORCEINLINE_DEBUGGABLE FRHICommandSetStaticUniformBuffers(const FUniformBufferStaticBindings& InUniformBuffers)
 		: UniformBuffers(InUniformBuffers)
 	{}
 	RHI_API void Execute(FRHICommandListBase& CmdList);
@@ -2304,14 +2304,20 @@ public:
 
 	inline FRHIComputeShader* GetBoundComputeShader() const { return BoundComputeShaderRHI; }
 
+	UE_DEPRECATED(5.0, "Please rename to SetStaticUniformBuffers")
 	FORCEINLINE_DEBUGGABLE void SetGlobalUniformBuffers(const FUniformBufferStaticBindings& UniformBuffers)
+	{
+		SetStaticUniformBuffers(UniformBuffers);
+	}
+
+	FORCEINLINE_DEBUGGABLE void SetStaticUniformBuffers(const FUniformBufferStaticBindings& UniformBuffers)
 	{
 		if (Bypass())
 		{
-			GetComputeContext().RHISetGlobalUniformBuffers(UniformBuffers);
+			GetComputeContext().RHISetStaticUniformBuffers(UniformBuffers);
 			return;
 		}
-		ALLOC_COMMAND(FRHICommandSetGlobalUniformBuffers)(UniformBuffers);
+		ALLOC_COMMAND(FRHICommandSetStaticUniformBuffers)(UniformBuffers);
 	}
 
 	FORCEINLINE_DEBUGGABLE void SetShaderUniformBuffer(FRHIComputeShader* Shader, uint32 BaseIndex, FRHIUniformBuffer* UniformBuffer)
@@ -4776,40 +4782,46 @@ struct FScopedGPUMask
 	#define SCOPED_GPU_MASK(RHICmdList, GPUMask)
 #endif // WITH_MGPU
 
-struct RHI_API FScopedUniformBufferGlobalBindings
+struct RHI_API FScopedUniformBufferStaticBindings
 {
-	FScopedUniformBufferGlobalBindings(FRHIComputeCommandList& InRHICmdList, FUniformBufferStaticBindings UniformBuffers)
+	FScopedUniformBufferStaticBindings(FRHIComputeCommandList& InRHICmdList, FUniformBufferStaticBindings UniformBuffers)
 		: RHICmdList(InRHICmdList)
 	{
-#if VALIDATE_UNIFORM_BUFFER_GLOBAL_BINDINGS
+#if VALIDATE_UNIFORM_BUFFER_STATIC_BINDINGS
 		checkf(!bRecursionGuard, TEXT("Uniform buffer global binding scope has been called recursively!"));
 		bRecursionGuard = true;
 #endif
 
-		RHICmdList.SetGlobalUniformBuffers(UniformBuffers);
+		RHICmdList.SetStaticUniformBuffers(UniformBuffers);
 	}
 
 	template <typename... TArgs>
-	FScopedUniformBufferGlobalBindings(FRHIComputeCommandList& InRHICmdList, TArgs... Args)
-		: FScopedUniformBufferGlobalBindings(InRHICmdList, FUniformBufferStaticBindings{ Args... })
+	FScopedUniformBufferStaticBindings(FRHIComputeCommandList& InRHICmdList, TArgs... Args)
+		: FScopedUniformBufferStaticBindings(InRHICmdList, FUniformBufferStaticBindings{ Args... })
 	{}
 
-	~FScopedUniformBufferGlobalBindings()
+	~FScopedUniformBufferStaticBindings()
 	{
-		RHICmdList.SetGlobalUniformBuffers(FUniformBufferStaticBindings());
+		RHICmdList.SetStaticUniformBuffers(FUniformBufferStaticBindings());
 
-#if VALIDATE_UNIFORM_BUFFER_GLOBAL_BINDINGS
+#if VALIDATE_UNIFORM_BUFFER_STATIC_BINDINGS
 		bRecursionGuard = false;
 #endif
 	}
 
 	FRHIComputeCommandList& RHICmdList;
 
-#if VALIDATE_UNIFORM_BUFFER_GLOBAL_BINDINGS
+#if VALIDATE_UNIFORM_BUFFER_STATIC_BINDINGS
 	static bool bRecursionGuard;
 #endif
 };
 
+#define SCOPED_UNIFORM_BUFFER_STATIC_BINDINGS(RHICmdList, UniformBuffers) FScopedUniformBufferStaticBindings PREPROCESSOR_JOIN(UniformBuffers, __LINE__){ RHICmdList, UniformBuffers }
+
+UE_DEPRECATED(5.0, "Please rename to FScopedUniformBufferStaticBindings, or use the SCOPED_UNIFORM_BUFFER_STATIC_BINDINGS instead.")
+typedef FScopedUniformBufferStaticBindings FScopedUniformBufferGlobalBindings;
+
+// UE_DEPRECATED(5.0)
 #define SCOPED_UNIFORM_BUFFER_GLOBAL_BINDINGS(RHICmdList, UniformBuffers) FScopedUniformBufferGlobalBindings PREPROCESSOR_JOIN(UniformBuffers, __LINE__){ RHICmdList, UniformBuffers }
 
 // Single commandlist for async compute generation.  In the future we may expand this to allow async compute command generation
diff --git a/Engine/Source/Runtime/RHI/Public/RHICommandListCommandExecutes.inl b/Engine/Source/Runtime/RHI/Public/RHICommandListCommandExecutes.inl
--- a/Engine/Source/Runtime/RHI/Public/RHICommandListCommandExecutes.inl
+++ b/Engine/Source/Runtime/RHI/Public/RHICommandListCommandExecutes.inl
@@ -430,10 +430,10 @@ void FRHICommandWriteGPUFence::Execute(FRHICommandListBase& CmdList)
 	}
 }
 
-void FRHICommandSetGlobalUniformBuffers::Execute(FRHICommandListBase& CmdList)
+void FRHICommandSetStaticUniformBuffers::Execute(FRHICommandListBase& CmdList)
 {
-	RHISTAT(SetGlobalUniformBuffers);
-	INTERNAL_DECORATOR_COMPUTE(RHISetGlobalUniformBuffers)(UniformBuffers);
+	RHISTAT(SetStaticUniformBuffers);
+	INTERNAL_DECORATOR_COMPUTE(RHISetStaticUniformBuffers)(UniformBuffers);
 }
 
 void FRHICommandBuildLocalUniformBuffer::Execute(FRHICommandListBase& CmdList)
diff --git a/Engine/Source/Runtime/RHI/Public/RHIContext.h b/Engine/Source/Runtime/RHI/Public/RHIContext.h
--- a/Engine/Source/Runtime/RHI/Public/RHIContext.h
+++ b/Engine/Source/Runtime/RHI/Public/RHIContext.h
@@ -30,9 +30,9 @@ struct FRayTracingLocalShaderBindings;
 struct FRHIBreadcrumb;
 enum class EAsyncComputeBudget;
 
-#define VALIDATE_UNIFORM_BUFFER_GLOBAL_BINDINGS (!UE_BUILD_SHIPPING && !UE_BUILD_TEST)
+#define VALIDATE_UNIFORM_BUFFER_STATIC_BINDINGS (!UE_BUILD_SHIPPING && !UE_BUILD_TEST)
 
-/** A list of global uniform buffer bindings. */
+/** A list of static uniform buffer bindings. */
 class FUniformBufferStaticBindings
 {
 public:
@@ -56,7 +56,7 @@ public:
 		const FUniformBufferStaticSlot Slot = Layout.StaticSlot;
 		checkf(IsUniformBufferStaticSlotValid(Slot), TEXT("Attempted to set a global uniform buffer %s with an invalid slot."), *Layout.GetDebugName());
 
-#if VALIDATE_UNIFORM_BUFFER_GLOBAL_BINDINGS
+#if VALIDATE_UNIFORM_BUFFER_STATIC_BINDINGS
 		ensureMsgf(INDEX_NONE == Slots.Find(Slot), TEXT("Uniform Buffer %s was added twice to the binding array."), *Layout.GetDebugName());
 #endif
 
@@ -111,6 +111,9 @@ private:
 	int32 SlotCount = 0;
 };
 
+UE_DEPRECATED(5.0, "Please rename to FUniformBufferStaticBindings")
+typedef FUniformBufferStaticBindings FUniformBufferGlobalBindings;
+
 /** Context that is capable of doing Compute work.  Can be async or compute on the gfx pipe. */
 class IRHIComputeContext
 {
@@ -206,11 +209,14 @@ public:
 
 	virtual void RHISetShaderParameter(FRHIComputeShader* ComputeShader, uint32 BufferIndex, uint32 BaseIndex, uint32 NumBytes, const void* NewValue) = 0;
 
-	virtual void RHISetGlobalUniformBuffers(const FUniformBufferStaticBindings& InUniformBuffers)
+	virtual void RHISetStaticUniformBuffers(const FUniformBufferStaticBindings& InUniformBuffers)
 	{
 		/** empty default implementation. */
 	}
 
+	UE_DEPRECATED(5.0, "Please rename to RHISetStaticUniformBuffers.")
+	virtual void RHISetGlobalUniformBuffers(const FUniformBufferStaticBindings& InUniformBuffers) {}
+
 	virtual void RHIPushEvent(const TCHAR* Name, FColor Color) = 0;
 
 	virtual void RHIPopEvent() = 0;
diff --git a/Engine/Source/Runtime/RHI/Public/RHIDefinitions.h b/Engine/Source/Runtime/RHI/Public/RHIDefinitions.h
--- a/Engine/Source/Runtime/RHI/Public/RHIDefinitions.h
+++ b/Engine/Source/Runtime/RHI/Public/RHIDefinitions.h
@@ -1052,7 +1052,7 @@ enum class EUniformBufferBindingFlags : uint8
 	/** If set, the uniform buffer can be bound as an RHI shader parameter on an RHI shader (i.e. RHISetShaderUniformBuffer). */
 	Shader = 1 << 0,
 
-	/** If set, the uniform buffer can be bound globally through a static slot (i.e. RHISetGlobalUniformBuffers). */
+	/** If set, the uniform buffer can be bound globally through a static slot (i.e. RHISetStaticUniformBuffers). */
 	Static = 1 << 1,
 
 	/** If set, the uniform buffer can be bound globally or per-shader, depending on the use case. Only one binding model should be
diff --git a/Engine/Source/Runtime/RHI/Public/RHIValidationContext.h b/Engine/Source/Runtime/RHI/Public/RHIValidationContext.h
--- a/Engine/Source/Runtime/RHI/Public/RHIValidationContext.h
+++ b/Engine/Source/Runtime/RHI/Public/RHIValidationContext.h
@@ -164,10 +164,10 @@ public:
 		RHIContext->RHISetShaderParameter(Shader, BufferIndex, BaseIndex, NumBytes, NewValue);
 	}
 
-	virtual void RHISetGlobalUniformBuffers(const FUniformBufferStaticBindings& InUniformBuffers) override final
+	virtual void RHISetStaticUniformBuffers(const FUniformBufferStaticBindings& InUniformBuffers) override final
 	{
 		InUniformBuffers.Bind(State.GlobalUniformBuffers.Bindings);
-		RHIContext->RHISetGlobalUniformBuffers(InUniformBuffers);
+		RHIContext->RHISetStaticUniformBuffers(InUniformBuffers);
 	}
 
 	virtual void RHIPushEvent(const TCHAR* Name, FColor Color) override final
@@ -671,10 +671,10 @@ public:
 		RHIContext->RHISetShaderParameter(Shader, BufferIndex, BaseIndex, NumBytes, NewValue);
 	}
 
-	virtual void RHISetGlobalUniformBuffers(const FUniformBufferStaticBindings& InUniformBuffers) override final
+	virtual void RHISetStaticUniformBuffers(const FUniformBufferStaticBindings& InUniformBuffers) override final
 	{
 		InUniformBuffers.Bind(State.GlobalUniformBuffers.Bindings);
-		RHIContext->RHISetGlobalUniformBuffers(InUniformBuffers);
+		RHIContext->RHISetStaticUniformBuffers(InUniformBuffers);
 	}
 
 	virtual void RHISetStencilRef(uint32 StencilRef) override final
diff --git a/Engine/Source/Runtime/RHICore/Private/RHICoreShader.cpp b/Engine/Source/Runtime/RHICore/Private/RHICoreShader.cpp
--- a/Engine/Source/Runtime/RHICore/Private/RHICoreShader.cpp
+++ b/Engine/Source/Runtime/RHICore/Private/RHICoreShader.cpp
@@ -35,7 +35,7 @@ void ValidateStaticUniformBuffer(FRHIUniformBuffer* UniformBuffer, FUniformBuffe
 		{
 			UE_LOG(LogRHICore, Fatal,
 				TEXT("Shader requested a global uniform buffer of type '%s' at static slot '%s', but it was null. The uniform buffer should ")
-				TEXT("be bound using RHICmdList.SetGlobalUniformBuffers() or passed into an RDG pass using SHADER_PARAMETER_STRUCT_REF() or ")
+				TEXT("be bound using RHICmdList.SetStaticUniformBuffers() or passed into an RDG pass using SHADER_PARAMETER_STRUCT_REF() or ")
 				TEXT("SHADER_PARAMETER_RDG_UNIFORM_BUFFER()."),
 				ExpectedStructMetadata->GetShaderVariableName(), *SlotRegistry.GetDebugDescription(Slot));
 		}
diff --git a/Engine/Source/Runtime/RenderCore/Private/RenderGraphBuilder.cpp b/Engine/Source/Runtime/RenderCore/Private/RenderGraphBuilder.cpp
--- a/Engine/Source/Runtime/RenderCore/Private/RenderGraphBuilder.cpp
+++ b/Engine/Source/Runtime/RenderCore/Private/RenderGraphBuilder.cpp
@@ -1331,7 +1331,7 @@ void FRDGBuilder::Execute()
 		ExecutePass(EpiloguePass);
 	}
 
-	RHICmdList.SetGlobalUniformBuffers({});
+	RHICmdList.SetStaticUniformBuffers({});
 
 #if WITH_MGPU
 	if (NameForTemporalEffect != NAME_None)
diff --git a/Engine/Source/Runtime/RenderCore/Private/RenderGraphPass.cpp b/Engine/Source/Runtime/RenderCore/Private/RenderGraphPass.cpp
--- a/Engine/Source/Runtime/RenderCore/Private/RenderGraphPass.cpp
+++ b/Engine/Source/Runtime/RenderCore/Private/RenderGraphPass.cpp
@@ -12,7 +12,7 @@ static FRDGPassHandlesByPipeline GetPassesByPipeline(const FRDGPass* Pass)
 	return Passes;
 }
 
-FUniformBufferStaticBindings FRDGParameterStruct::GetGlobalUniformBuffers() const
+FUniformBufferStaticBindings FRDGParameterStruct::GetStaticUniformBuffers() const
 {
 	FUniformBufferStaticBindings GlobalUniformBuffers;
 
@@ -337,6 +337,6 @@ const TCHAR* FRDGPass::GetName() const
 void FRDGPass::Execute(FRHIComputeCommandList& RHICmdList)
 {
 	QUICK_SCOPE_CYCLE_COUNTER(STAT_FRDGPass_Execute);
-	RHICmdList.SetGlobalUniformBuffers(ParameterStruct.GetGlobalUniformBuffers());
+	RHICmdList.SetStaticUniformBuffers(ParameterStruct.GetStaticUniformBuffers());
 	ExecuteImpl(RHICmdList);
 }
\ No newline at end of file
diff --git a/Engine/Source/Runtime/RenderCore/Public/RenderGraphParameter.h b/Engine/Source/Runtime/RenderCore/Public/RenderGraphParameter.h
--- a/Engine/Source/Runtime/RenderCore/Public/RenderGraphParameter.h
+++ b/Engine/Source/Runtime/RenderCore/Public/RenderGraphParameter.h
@@ -242,8 +242,8 @@ public:
 	template <typename FunctionType>
 	void EnumerateUniformBuffers(FunctionType Function) const;
 
-	/** Returns a set of static global uniform buffer bindings for the parameter struct. */
-	FUniformBufferStaticBindings GetGlobalUniformBuffers() const;
+	/** Returns a set of static uniform buffer bindings for the parameter struct. */
+	FUniformBufferStaticBindings GetStaticUniformBuffers() const;
 
 	/** Returns the render pass info generated from the render target binding slots. */
 	FRHIRenderPassInfo GetRenderPassInfo() const;
@@ -297,7 +297,7 @@ FORCEINLINE static FRHIRenderPassInfo GetRenderPassInfo(TParameterStruct* Parame
 
 /** Helper function to get RHI global uniform buffers out of a pass parameters struct. */
 template <typename TParameterStruct>
-FORCEINLINE static FUniformBufferStaticBindings GetGlobalUniformBuffers(TParameterStruct* Parameters)
+FORCEINLINE static FUniformBufferStaticBindings GetStaticUniformBuffers(TParameterStruct* Parameters)
 {
-	return FRDGParameterStruct(Parameters).GetGlobalUniformBuffers();
+	return FRDGParameterStruct(Parameters).GetStaticUniformBuffers();
 }
\ No newline at end of file
diff --git a/Engine/Source/Runtime/Renderer/Private/SceneRendering.h b/Engine/Source/Runtime/Renderer/Private/SceneRendering.h
--- a/Engine/Source/Runtime/Renderer/Private/SceneRendering.h
+++ b/Engine/Source/Runtime/Renderer/Private/SceneRendering.h
@@ -458,17 +458,17 @@ public:
 	template <typename ParameterStructType>
 	FParallelCommandListBindings(ParameterStructType* ParameterStruct)
 		: RenderPassInfo(GetRenderPassInfo(ParameterStruct))
-		, GlobalUniformBuffers(GetGlobalUniformBuffers(ParameterStruct))
+		, StaticUniformBuffers(GetStaticUniformBuffers(ParameterStruct))
 	{}
 
 	inline void SetOnCommandList(FRHICommandList& RHICmdList) const
 	{
 		RHICmdList.BeginRenderPass(RenderPassInfo, TEXT("Parallel"));
-		RHICmdList.SetGlobalUniformBuffers(GlobalUniformBuffers);
+		RHICmdList.SetStaticUniformBuffers(StaticUniformBuffers);
 	}
 
 	FRHIRenderPassInfo RenderPassInfo;
-	FUniformBufferStaticBindings GlobalUniformBuffers;
+	FUniformBufferStaticBindings StaticUniformBuffers;
 };
 
 class FParallelCommandListSet
diff --git a/Engine/Source/Runtime/SlateRHIRenderer/Private/SlateRHIRenderingPolicy.cpp b/Engine/Source/Runtime/SlateRHIRenderer/Private/SlateRHIRenderingPolicy.cpp
--- a/Engine/Source/Runtime/SlateRHIRenderer/Private/SlateRHIRenderingPolicy.cpp
+++ b/Engine/Source/Runtime/SlateRHIRenderer/Private/SlateRHIRenderingPolicy.cpp
@@ -1052,7 +1052,7 @@ void FSlateRHIRenderingPolicy::DrawElements(
 					if (VertexShader.IsValid() && PixelShader.IsValid() && SceneTextureUniformBuffer)
 					{
 						const FUniformBufferStaticBindings StaticUniformBuffers(SceneTextureUniformBuffer);
-						SCOPED_UNIFORM_BUFFER_GLOBAL_BINDINGS(RHICmdList, StaticUniformBuffers);
+						SCOPED_UNIFORM_BUFFER_STATIC_BINDINGS(RHICmdList, StaticUniformBuffers);
 
 #if WITH_SLATE_VISUALIZERS
 						if (CVarShowSlateBatching.GetValueOnRenderThread() != 0)
diff --git a/Engine/Source/Runtime/VulkanRHI/Private/VulkanCommands.cpp b/Engine/Source/Runtime/VulkanRHI/Private/VulkanCommands.cpp
--- a/Engine/Source/Runtime/VulkanRHI/Private/VulkanCommands.cpp
+++ b/Engine/Source/Runtime/VulkanRHI/Private/VulkanCommands.cpp
@@ -549,7 +549,7 @@ inline void FVulkanCommandListContext::SetShaderUniformBuffer(ShaderStage::EStag
 	}
 }
 
-void FVulkanCommandListContext::RHISetGlobalUniformBuffers(const FUniformBufferStaticBindings& InUniformBuffers)
+void FVulkanCommandListContext::RHISetStaticUniformBuffers(const FUniformBufferStaticBindings& InUniformBuffers)
 {
 	FMemory::Memzero(GlobalUniformBuffers.GetData(), GlobalUniformBuffers.Num() * sizeof(FRHIUniformBuffer*));
 
diff --git a/Engine/Source/Runtime/VulkanRHI/Private/VulkanContext.h b/Engine/Source/Runtime/VulkanRHI/Private/VulkanContext.h
--- a/Engine/Source/Runtime/VulkanRHI/Private/VulkanContext.h
+++ b/Engine/Source/Runtime/VulkanRHI/Private/VulkanContext.h
@@ -50,7 +50,7 @@ public:
 	virtual void RHISetUAVParameter(FRHIComputeShader* ComputeShader, uint32 UAVIndex, FRHIUnorderedAccessView* UAV, uint32 InitialCount) final override;
 	virtual void RHISetShaderResourceViewParameter(FRHIGraphicsShader* Shader, uint32 SamplerIndex, FRHIShaderResourceView* SRV) final override;
 	virtual void RHISetShaderResourceViewParameter(FRHIComputeShader* ComputeShader, uint32 SamplerIndex, FRHIShaderResourceView* SRV) final override;
-	virtual void RHISetGlobalUniformBuffers(const FUniformBufferStaticBindings& InUniformBuffers) final override;
+	virtual void RHISetStaticUniformBuffers(const FUniformBufferStaticBindings& InUniformBuffers) final override;
 	virtual void RHISetShaderUniformBuffer(FRHIGraphicsShader* Shader, uint32 BufferIndex, FRHIUniformBuffer* Buffer) final override;
 	virtual void RHISetShaderUniformBuffer(FRHIComputeShader* ComputeShader, uint32 BufferIndex, FRHIUniformBuffer* Buffer) final override;
 	virtual void RHISetShaderParameter(FRHIGraphicsShader* Shader, uint32 BufferIndex, uint32 BaseIndex, uint32 NumBytes, const void* NewValue) final override;
diff --git a/Engine/Source/Runtime/Windows/D3D11RHI/Private/D3D11Commands.cpp b/Engine/Source/Runtime/Windows/D3D11RHI/Private/D3D11Commands.cpp
--- a/Engine/Source/Runtime/Windows/D3D11RHI/Private/D3D11Commands.cpp
+++ b/Engine/Source/Runtime/Windows/D3D11RHI/Private/D3D11Commands.cpp
@@ -614,7 +614,7 @@ void FD3D11DynamicRHI::RHISetShaderSampler(FRHIComputeShader* ComputeShaderRHI,u
 	StateCache.SetSamplerState<SF_Compute>(StateResource, SamplerIndex);
 }
 
-void FD3D11DynamicRHI::RHISetGlobalUniformBuffers(const FUniformBufferStaticBindings& InUniformBuffers)
+void FD3D11DynamicRHI::RHISetStaticUniformBuffers(const FUniformBufferStaticBindings& InUniformBuffers)
 {
 	FMemory::Memzero(StaticUniformBuffers.GetData(), StaticUniformBuffers.Num() * sizeof(FRHIUniformBuffer*));
 
diff --git a/Engine/Source/Runtime/Windows/D3D11RHI/Private/D3D11RHIPrivate.h b/Engine/Source/Runtime/Windows/D3D11RHI/Private/D3D11RHIPrivate.h
--- a/Engine/Source/Runtime/Windows/D3D11RHI/Private/D3D11RHIPrivate.h
+++ b/Engine/Source/Runtime/Windows/D3D11RHI/Private/D3D11RHIPrivate.h
@@ -564,7 +564,7 @@ public:
 	virtual void RHISetUAVParameter(FRHIComputeShader* ComputeShader, uint32 UAVIndex, FRHIUnorderedAccessView* UAV, uint32 InitialCount) final override;
 	virtual void RHISetShaderResourceViewParameter(FRHIGraphicsShader* Shader, uint32 SamplerIndex, FRHIShaderResourceView* SRV) final override;
 	virtual void RHISetShaderResourceViewParameter(FRHIComputeShader* ComputeShader, uint32 SamplerIndex, FRHIShaderResourceView* SRV) final override;
-	virtual void RHISetGlobalUniformBuffers(const FUniformBufferStaticBindings& InUniformBuffers) final override;
+	virtual void RHISetStaticUniformBuffers(const FUniformBufferStaticBindings& InUniformBuffers) final override;
 	virtual void RHISetShaderUniformBuffer(FRHIGraphicsShader* Shader, uint32 BufferIndex, FRHIUniformBuffer* Buffer) final override;
 	virtual void RHISetShaderUniformBuffer(FRHIComputeShader* ComputeShader, uint32 BufferIndex, FRHIUniformBuffer* Buffer) final override;
 	virtual void RHISetShaderParameter(FRHIGraphicsShader* Shader, uint32 BufferIndex, uint32 BaseIndex, uint32 NumBytes, const void* NewValue) final override;
