Added Buffer constructor to FRHITransitionInfo

diff --git a/Engine/Source/Runtime/RHI/Public/RHI.h b/Engine/Source/Runtime/RHI/Public/RHI.h
--- a/Engine/Source/Runtime/RHI/Public/RHI.h
+++ b/Engine/Source/Runtime/RHI/Public/RHI.h
@@ -1907,6 +1907,14 @@ struct FRHITransitionInfo : public FRHISubresourceRange
 		, Flags(InFlags)
 	{}
 
+	FRHITransitionInfo(class FRHIBuffer* InRHIBuffer, ERHIAccess InPreviousState, ERHIAccess InNewState, EResourceTransitionFlags InFlags = EResourceTransitionFlags::None)
+		: Buffer(InRHIBuffer)
+		, Type(EType::Buffer)
+		, AccessBefore(InPreviousState)
+		, AccessAfter(InNewState)
+		, Flags(InFlags)
+	{}
+
 	inline bool operator == (FRHITransitionInfo const& RHS) const
 	{
 		return Resource == RHS.Resource
