Add DSVWrite state assert in the RHI validation layer's implementation of RHIResummarizeHtile
 - Htile resummarization happens on the depth hardware, so resources must be in a writable depth state.

diff --git a/Engine/Source/Runtime/RHI/Public/RHIValidationContext.h b/Engine/Source/Runtime/RHI/Public/RHIValidationContext.h
--- a/Engine/Source/Runtime/RHI/Public/RHIValidationContext.h
+++ b/Engine/Source/Runtime/RHI/Public/RHIValidationContext.h
@@ -348,6 +348,7 @@ public:
 
 	virtual void RHIResummarizeHTile(FRHITexture2D* DepthTexture) override final
 	{
+		Tracker->Assert(DepthTexture->GetWholeResourceIdentity(), ERHIAccess::DSVWrite);
 		RHIContext->RHIResummarizeHTile(DepthTexture);
 	}
 
