Replace FVertexBufferRHIRef, FIndexBufferRHIRef and FStructuredBufferRHIRef with FBufferRHIRef.

diff --git a/Engine/Plugins/Enterprise/LidarPointCloud/Source/LidarPointCloudRuntime/Private/Rendering/LidarPointCloudRenderBuffers.h b/Engine/Plugins/Enterprise/LidarPointCloud/Source/LidarPointCloudRuntime/Private/Rendering/LidarPointCloudRenderBuffers.h
--- a/Engine/Plugins/Enterprise/LidarPointCloud/Source/LidarPointCloudRuntime/Private/Rendering/LidarPointCloudRenderBuffers.h
+++ b/Engine/Plugins/Enterprise/LidarPointCloud/Source/LidarPointCloudRuntime/Private/Rendering/LidarPointCloudRenderBuffers.h
@@ -47,7 +47,7 @@ public:
 class FLidarPointCloudRenderBuffer : public FLidarPointCloudBuffer, public FRenderResource
 {
 public:
-	FVertexBufferRHIRef Buffer;
+	FBufferRHIRef Buffer;
 	FShaderResourceViewRHIRef SRV;
 
 	int32 PointCount;
diff --git a/Engine/Plugins/Experimental/GPULightmass/Source/GPULightmass/Private/IrradianceCaching.h b/Engine/Plugins/Experimental/GPULightmass/Source/GPULightmass/Private/IrradianceCaching.h
--- a/Engine/Plugins/Experimental/GPULightmass/Source/GPULightmass/Private/IrradianceCaching.h
+++ b/Engine/Plugins/Experimental/GPULightmass/Source/GPULightmass/Private/IrradianceCaching.h
@@ -32,7 +32,7 @@ struct FIrradianceCache
 
 	const int32 IrradianceCacheMaxSize = 1048576;
 
-	FStructuredBufferRHIRef IrradianceCacheRecords;
+	FBufferRHIRef IrradianceCacheRecords;
 	FUnorderedAccessViewRHIRef IrradianceCacheRecordsUAV;
 
 	TUniformBufferRef<FIrradianceCachingParameters> IrradianceCachingParametersUniformBuffer;
diff --git a/Engine/Plugins/Experimental/GPULightmass/Source/GPULightmass/Private/LightmapRenderer.cpp b/Engine/Plugins/Experimental/GPULightmass/Source/GPULightmass/Private/LightmapRenderer.cpp
--- a/Engine/Plugins/Experimental/GPULightmass/Source/GPULightmass/Private/LightmapRenderer.cpp
+++ b/Engine/Plugins/Experimental/GPULightmass/Source/GPULightmass/Private/LightmapRenderer.cpp
@@ -159,7 +159,7 @@ struct FGPUTileDescription
 
 struct FGPUBatchedTileRequests
 {
-	FStructuredBufferRHIRef BatchedTilesBuffer;
+	FBufferRHIRef BatchedTilesBuffer;
 	FShaderResourceViewRHIRef BatchedTilesSRV;
 	TResourceArray<FGPUTileDescription> BatchedTilesDesc;
 };
@@ -1199,9 +1199,9 @@ void FLightmapRenderer::Finalize(FRDGBuilder& GraphBuilder)
 			FRDGTextureRef StagingHQLayer1 = GraphBuilder.RegisterExternalTexture(UploadTilePoolGPU->PooledRenderTargets[1], TEXT("StagingHQLayer1"));
 			FRDGTextureRef StagingShadowMask = GraphBuilder.RegisterExternalTexture(UploadTilePoolGPU->PooledRenderTargets[2], TEXT("StagingShadowMask"));
 
-			FStructuredBufferRHIRef SrcTilePositionsBuffer;
+			FBufferRHIRef SrcTilePositionsBuffer;
 			FShaderResourceViewRHIRef SrcTilePositionsSRV;
-			FStructuredBufferRHIRef DstTilePositionsBuffer;
+			FBufferRHIRef DstTilePositionsBuffer;
 			FShaderResourceViewRHIRef DstTilePositionsSRV;
 
 			if (OutputRenderTargets[0] != nullptr)
@@ -1983,19 +1983,19 @@ void FLightmapRenderer::Finalize(FRDGBuilder& GraphBuilder)
 					FRDGTextureRef ShadowMaskSampleCount = GraphBuilder.RegisterExternalTexture(LightmapTilePoolGPU.PooledRenderTargets[3], TEXT("ShadowMaskSampleCount"));
 
 					TResourceArray<int32> LightTypeArray;
-					FVertexBufferRHIRef LightTypeBuffer;
+					FBufferRHIRef LightTypeBuffer;
 					FShaderResourceViewRHIRef LightTypeSRV;
 
 					TResourceArray<int32> ChannelIndexArray;
-					FVertexBufferRHIRef ChannelIndexBuffer;
+					FBufferRHIRef ChannelIndexBuffer;
 					FShaderResourceViewRHIRef ChannelIndexSRV;
 
 					auto& LightSampleIndexArray = *GraphBuilder.AllocObject<TResourceArray<int32>>();
-					FVertexBufferRHIRef LightSampleIndexBuffer;
+					FBufferRHIRef LightSampleIndexBuffer;
 					FShaderResourceViewRHIRef LightSampleIndexSRV;
 
 					TResourceArray<FLightShaderConstants> LightShaderParameterArray;
-					FStructuredBufferRHIRef LightShaderParameterBuffer;
+					FBufferRHIRef LightShaderParameterBuffer;
 					FShaderResourceViewRHIRef LightShaderParameterSRV;
 
 					for (const FLightmapTileRequest& Tile : PendingShadowTileRequests)
diff --git a/Engine/Plugins/Experimental/GPULightmass/Source/GPULightmass/Private/Scene/InstancedStaticMesh.h b/Engine/Plugins/Experimental/GPULightmass/Source/GPULightmass/Private/Scene/InstancedStaticMesh.h
--- a/Engine/Plugins/Experimental/GPULightmass/Source/GPULightmass/Private/Scene/InstancedStaticMesh.h
+++ b/Engine/Plugins/Experimental/GPULightmass/Source/GPULightmass/Private/Scene/InstancedStaticMesh.h
@@ -15,9 +15,9 @@ public:
 	UInstancedStaticMeshComponent* ComponentUObject;
 	FStaticMeshRenderData* RenderData;
 	TUniquePtr<FInstancedStaticMeshRenderData> InstancedRenderData;
-	FVertexBufferRHIRef InstanceOriginBuffer;
-	FVertexBufferRHIRef InstanceTransformBuffer;
-	FVertexBufferRHIRef InstanceLightmapBuffer;
+	FBufferRHIRef InstanceOriginBuffer;
+	FBufferRHIRef InstanceTransformBuffer;
+	FBufferRHIRef InstanceLightmapBuffer;
 	TUniformBufferRef<FPrimitiveUniformShaderParameters> UniformBuffer;
 
 	TArray<FIntPoint> LODPerInstanceLightmapSize;
diff --git a/Engine/Plugins/Experimental/GPULightmass/Source/GPULightmass/Private/Scene/Scene.h b/Engine/Plugins/Experimental/GPULightmass/Source/GPULightmass/Private/Scene/Scene.h
--- a/Engine/Plugins/Experimental/GPULightmass/Source/GPULightmass/Private/Scene/Scene.h
+++ b/Engine/Plugins/Experimental/GPULightmass/Source/GPULightmass/Private/Scene/Scene.h
@@ -60,10 +60,10 @@ struct FCachedRayTracingSceneData
 	TArray<FVisibleRayTracingMeshCommand> VisibleRayTracingMeshCommands;
 	TChunkedArray<FRayTracingMeshCommand> MeshCommandStorage;
 
-	FStructuredBufferRHIRef PrimitiveSceneDataBufferRHI;
+	FBufferRHIRef PrimitiveSceneDataBufferRHI;
 	FShaderResourceViewRHIRef PrimitiveSceneDataBufferSRV;
 
-	FStructuredBufferRHIRef LightmapSceneDataBufferRHI;
+	FBufferRHIRef LightmapSceneDataBufferRHI;
 	FShaderResourceViewRHIRef LightmapSceneDataBufferSRV;
 
 	TArray<FRayTracingGeometryInstance> RayTracingGeometryInstances;
diff --git a/Engine/Plugins/Experimental/RemoteSession/Source/RemoteSession/Private/Channels/RemoteSessionARCameraChannel.cpp b/Engine/Plugins/Experimental/RemoteSession/Source/RemoteSession/Private/Channels/RemoteSessionARCameraChannel.cpp
--- a/Engine/Plugins/Experimental/RemoteSession/Source/RemoteSession/Private/Channels/RemoteSessionARCameraChannel.cpp
+++ b/Engine/Plugins/Experimental/RemoteSession/Source/RemoteSession/Private/Channels/RemoteSessionARCameraChannel.cpp
@@ -160,9 +160,9 @@ private:
 	/** Note the session channel is responsible for preventing GC */
 	UMaterialInterface* PPMaterial;
 	/** Index buffer for drawing the quad */
-	FIndexBufferRHIRef IndexBufferRHI;
+	FBufferRHIRef IndexBufferRHI;
 	/** Vertex buffer for drawing the quad */
-	FVertexBufferRHIRef VertexBufferRHI;
+	FBufferRHIRef VertexBufferRHI;
 };
 
 FARCameraSceneViewExtension::FARCameraSceneViewExtension(const FAutoRegister& AutoRegister, FRemoteSessionARCameraChannel& InChannel) :
diff --git a/Engine/Plugins/Experimental/SimpleHMD/Source/SimpleHMD/Private/SimpleHMD.cpp b/Engine/Plugins/Experimental/SimpleHMD/Source/SimpleHMD/Private/SimpleHMD.cpp
--- a/Engine/Plugins/Experimental/SimpleHMD/Source/SimpleHMD/Private/SimpleHMD.cpp
+++ b/Engine/Plugins/Experimental/SimpleHMD/Source/SimpleHMD/Private/SimpleHMD.cpp
@@ -205,14 +205,14 @@ void FSimpleHMD::DrawDistortionMesh_RenderThread(struct FHeadMountedDisplayPassC
 	};
 
 	FRHIResourceCreateInfo CreateInfo;
-	FVertexBufferRHIRef VertexBufferRHI = RHICreateVertexBuffer(sizeof(FDistortionVertex) * NumVerts, BUF_Volatile, CreateInfo);
+	FBufferRHIRef VertexBufferRHI = RHICreateVertexBuffer(sizeof(FDistortionVertex) * NumVerts, BUF_Volatile, CreateInfo);
 	void* VoidPtr = RHILockBuffer(VertexBufferRHI, 0, sizeof(FDistortionVertex) * NumVerts, RLM_WriteOnly);
 	FPlatformMemory::Memcpy(VoidPtr, View.StereoPass == eSSP_RIGHT_EYE ? RightVerts : LeftVerts, sizeof(FDistortionVertex) * NumVerts);
 	RHIUnlockBuffer(VertexBufferRHI);
 
 	static const uint16 Indices[] = { 0, 1, 2, 0, 2, 3 };
 
-	FIndexBufferRHIRef IndexBufferRHI = RHICreateIndexBuffer(sizeof(uint16), sizeof(uint16) * 6, BUF_Volatile, CreateInfo);
+	FBufferRHIRef IndexBufferRHI = RHICreateIndexBuffer(sizeof(uint16), sizeof(uint16) * 6, BUF_Volatile, CreateInfo);
 	void* VoidPtr2 = RHILockBuffer(IndexBufferRHI, 0, sizeof(uint16) * 6, RLM_WriteOnly);
 	FPlatformMemory::Memcpy(VoidPtr2, Indices, sizeof(uint16) * 6);
 	RHIUnlockBuffer(IndexBufferRHI);
diff --git a/Engine/Plugins/Experimental/VirtualHeightfieldMesh/Source/VirtualHeightfieldMesh/Private/VirtualHeightfieldMeshSceneProxy.cpp b/Engine/Plugins/Experimental/VirtualHeightfieldMesh/Source/VirtualHeightfieldMesh/Private/VirtualHeightfieldMeshSceneProxy.cpp
--- a/Engine/Plugins/Experimental/VirtualHeightfieldMesh/Source/VirtualHeightfieldMesh/Private/VirtualHeightfieldMeshSceneProxy.cpp
+++ b/Engine/Plugins/Experimental/VirtualHeightfieldMesh/Source/VirtualHeightfieldMesh/Private/VirtualHeightfieldMeshSceneProxy.cpp
@@ -84,12 +84,12 @@ namespace VirtualHeightfieldMesh
 	struct FDrawInstanceBuffers
 	{
 		/* Culled instance buffer. */
-		FStructuredBufferRHIRef InstanceBuffer;
+		FBufferRHIRef InstanceBuffer;
 		FUnorderedAccessViewRHIRef InstanceBufferUAV;
 		FShaderResourceViewRHIRef InstanceBufferSRV;
 
 		/* IndirectArgs buffer for final DrawInstancedIndirect. */
-		FVertexBufferRHIRef IndirectArgsBuffer;
+		FBufferRHIRef IndirectArgsBuffer;
 		FUnorderedAccessViewRHIRef IndirectArgsBufferUAV;
 	};
 
diff --git a/Engine/Plugins/Experimental/VirtualHeightfieldMesh/Source/VirtualHeightfieldMesh/Private/VirtualHeightfieldMeshVertexFactory.cpp b/Engine/Plugins/Experimental/VirtualHeightfieldMesh/Source/VirtualHeightfieldMesh/Private/VirtualHeightfieldMeshVertexFactory.cpp
--- a/Engine/Plugins/Experimental/VirtualHeightfieldMesh/Source/VirtualHeightfieldMesh/Private/VirtualHeightfieldMeshVertexFactory.cpp
+++ b/Engine/Plugins/Experimental/VirtualHeightfieldMesh/Source/VirtualHeightfieldMesh/Private/VirtualHeightfieldMeshVertexFactory.cpp
@@ -15,7 +15,7 @@ IMPLEMENT_GLOBAL_SHADER_PARAMETER_STRUCT(FVirtualHeightfieldMeshVertexFactoryPar
 namespace
 {
 	template< typename T >
-	FIndexBufferRHIRef CreateIndexBuffer(uint32 NumQuadsPerSide)
+	FBufferRHIRef CreateIndexBuffer(uint32 NumQuadsPerSide)
 	{
 		TResourceArray<T, INDEXBUFFER_ALIGNMENT> Indices;
 
diff --git a/Engine/Plugins/Experimental/Water/Source/Runtime/Public/GerstnerWaterWaveViewExtension.h b/Engine/Plugins/Experimental/Water/Source/Runtime/Public/GerstnerWaterWaveViewExtension.h
--- a/Engine/Plugins/Experimental/Water/Source/Runtime/Public/GerstnerWaterWaveViewExtension.h
+++ b/Engine/Plugins/Experimental/Water/Source/Runtime/Public/GerstnerWaterWaveViewExtension.h
@@ -29,10 +29,10 @@ public:
 
 	bool bRebuildGPUData = false;
 
-	FStructuredBufferRHIRef DataBuffer;
+	FBufferRHIRef DataBuffer;
 	FShaderResourceViewRHIRef DataSRV;
 
-	FStructuredBufferRHIRef IndirectionBuffer;
+	FBufferRHIRef IndirectionBuffer;
 	FShaderResourceViewRHIRef IndirectionSRV;
 
 };
\ No newline at end of file
diff --git a/Engine/Plugins/Experimental/Water/Source/Runtime/Public/WaterInstanceDataBuffer.h b/Engine/Plugins/Experimental/Water/Source/Runtime/Public/WaterInstanceDataBuffer.h
--- a/Engine/Plugins/Experimental/Water/Source/Runtime/Public/WaterInstanceDataBuffer.h
+++ b/Engine/Plugins/Experimental/Water/Source/Runtime/Public/WaterInstanceDataBuffer.h
@@ -54,7 +54,7 @@ public:
 		}
 	}
 
-	FVertexBufferRHIRef GetBuffer(int32 InBufferID) const
+	FBufferRHIRef GetBuffer(int32 InBufferID) const
 	{
 		return Buffer[InBufferID];
 	}
@@ -92,6 +92,6 @@ private:
 		RHIUnlockBuffer(Buffer[InBufferID]);
 	}
 
-	FVertexBufferRHIRef Buffer[NumBuffers];
+	FBufferRHIRef Buffer[NumBuffers];
 	FVector4* BufferMemory[NumBuffers];
 };
diff --git a/Engine/Plugins/Experimental/Water/Source/Runtime/Public/WaterVertexFactory.h b/Engine/Plugins/Experimental/Water/Source/Runtime/Public/WaterVertexFactory.h
--- a/Engine/Plugins/Experimental/Water/Source/Runtime/Public/WaterVertexFactory.h
+++ b/Engine/Plugins/Experimental/Water/Source/Runtime/Public/WaterVertexFactory.h
@@ -46,7 +46,7 @@ public:
 
 private:
 	template <typename IndexType>
-	FIndexBufferRHIRef CreateIndexBuffer()
+	FBufferRHIRef CreateIndexBuffer()
 	{
 		TResourceArray<IndexType, INDEXBUFFER_ALIGNMENT> Indices;
 
diff --git a/Engine/Plugins/FX/Niagara/Source/Niagara/Classes/NiagaraDataInterfaceSkeletalMesh.h b/Engine/Plugins/FX/Niagara/Source/Niagara/Classes/NiagaraDataInterfaceSkeletalMesh.h
--- a/Engine/Plugins/FX/Niagara/Source/Niagara/Classes/NiagaraDataInterfaceSkeletalMesh.h
+++ b/Engine/Plugins/FX/Niagara/Source/Niagara/Classes/NiagaraDataInterfaceSkeletalMesh.h
@@ -374,9 +374,9 @@ public:
 	int32 GetFilteredSocketBoneOffset() const { return FilteredSocketBoneOffset; }
 
 protected:
-	FVertexBufferRHIRef BufferTriangleUniformSamplerProbAliasRHI = nullptr;
+	FBufferRHIRef BufferTriangleUniformSamplerProbAliasRHI = nullptr;
 	FShaderResourceViewRHIRef BufferTriangleUniformSamplerProbAliasSRV = nullptr;
-	FVertexBufferRHIRef BufferTriangleMatricesOffsetRHI = nullptr;
+	FBufferRHIRef BufferTriangleMatricesOffsetRHI = nullptr;
 	FShaderResourceViewRHIRef BufferTriangleMatricesOffsetSRV = nullptr;
 
 	bool bSamplingRegionsAllAreaWeighted = false;
@@ -386,18 +386,18 @@ protected:
 	TResourceArray<int32> SampleRegionsTriangleIndicies;
 	TResourceArray<int32> SampleRegionsVertices;
 
-	FVertexBufferRHIRef SampleRegionsProbAliasBuffer;
+	FBufferRHIRef SampleRegionsProbAliasBuffer;
 	FShaderResourceViewRHIRef SampleRegionsProbAliasSRV;
-	FVertexBufferRHIRef SampleRegionsTriangleIndicesBuffer;
+	FBufferRHIRef SampleRegionsTriangleIndicesBuffer;
 	FShaderResourceViewRHIRef SampleRegionsTriangleIndicesSRV;
-	FVertexBufferRHIRef SampleRegionsVerticesBuffer;
+	FBufferRHIRef SampleRegionsVerticesBuffer;
 	FShaderResourceViewRHIRef SampleRegionsVerticesSRV;
 
 	int32 NumFilteredBones = 0;
 	int32 NumUnfilteredBones = 0;
 	int32 ExcludedBoneIndex = INDEX_NONE;
 	TResourceArray<uint16> FilteredAndUnfilteredBonesArray;
-	FVertexBufferRHIRef FilteredAndUnfilteredBonesBuffer;
+	FBufferRHIRef FilteredAndUnfilteredBonesBuffer;
 	FShaderResourceViewRHIRef FilteredAndUnfilteredBonesSRV;
 
 	int32 NumFilteredSockets = 0;
@@ -452,10 +452,10 @@ public:
 	/** Encapsulates a GPU read / CPU write buffer for bone data */
 	struct FSkeletalBuffer
 	{
-		FVertexBufferRHIRef SectionBuffer;
+		FBufferRHIRef SectionBuffer;
 		FShaderResourceViewRHIRef SectionSRV;
 
-		FVertexBufferRHIRef SamplingBuffer;
+		FBufferRHIRef SamplingBuffer;
 		FShaderResourceViewRHIRef SamplingSRV;
 	};
 
diff --git a/Engine/Plugins/FX/Niagara/Source/Niagara/Classes/NiagaraDataInterfaceStaticMesh.h b/Engine/Plugins/FX/Niagara/Source/Niagara/Classes/NiagaraDataInterfaceStaticMesh.h
--- a/Engine/Plugins/FX/Niagara/Source/Niagara/Classes/NiagaraDataInterfaceStaticMesh.h
+++ b/Engine/Plugins/FX/Niagara/Source/Niagara/Classes/NiagaraDataInterfaceStaticMesh.h
@@ -116,7 +116,7 @@ protected:
 
 	TArray<SectionInfo> ValidSections;					// Only the section we want to spawn from
 
-	FVertexBufferRHIRef BufferSectionRHI = nullptr;
+	FBufferRHIRef BufferSectionRHI = nullptr;
 	FShaderResourceViewRHIRef BufferSectionSRV = nullptr;
 
 	FShaderResourceViewRHIRef BufferUniformTriangleSamplingSRV = nullptr;
@@ -130,11 +130,11 @@ protected:
 	FShaderResourceViewRHIRef MeshColorBufferSRV;
 
 	TResourceArray<FVector4> SocketTransformsResourceArray;
-	FVertexBufferRHIRef SocketTransformsBuffer;
+	FBufferRHIRef SocketTransformsBuffer;
 	FShaderResourceViewRHIRef SocketTransformsSRV;
 
 	TResourceArray<uint16> FilteredAndUnfilteredSocketsResourceArray;
-	FVertexBufferRHIRef FilteredAndUnfilteredSocketsBuffer;
+	FBufferRHIRef FilteredAndUnfilteredSocketsBuffer;
 	FShaderResourceViewRHIRef FilteredAndUnfilteredSocketsSRV;
 
 	uint32 NumSockets = 0;
diff --git a/Engine/Plugins/FX/Niagara/Source/Niagara/Classes/NiagaraEmitterInstanceBatcher.h b/Engine/Plugins/FX/Niagara/Source/Niagara/Classes/NiagaraEmitterInstanceBatcher.h
--- a/Engine/Plugins/FX/Niagara/Source/Niagara/Classes/NiagaraEmitterInstanceBatcher.h
+++ b/Engine/Plugins/FX/Niagara/Source/Niagara/Classes/NiagaraEmitterInstanceBatcher.h
@@ -277,7 +277,7 @@ private:
 
 	struct DummyUAV
 	{
-		FVertexBufferRHIRef Buffer;
+		FBufferRHIRef Buffer;
 		FTexture2DRHIRef Texture;
 		FUnorderedAccessViewRHIRef UAV;
 
diff --git a/Engine/Plugins/FX/Niagara/Source/Niagara/Private/NiagaraDataInterfaceMeshRendererInfo.cpp b/Engine/Plugins/FX/Niagara/Source/Niagara/Private/NiagaraDataInterfaceMeshRendererInfo.cpp
--- a/Engine/Plugins/FX/Niagara/Source/Niagara/Private/NiagaraDataInterfaceMeshRendererInfo.cpp
+++ b/Engine/Plugins/FX/Niagara/Source/Niagara/Private/NiagaraDataInterfaceMeshRendererInfo.cpp
@@ -74,7 +74,7 @@ public:
 	{}
 
 	uint32 GetNumMeshes() { return MeshData.Num(); }
-	FVertexBufferRHIRef GetMeshDataBufferRHI() const { return BufferMeshDataRHI; }
+	FBufferRHIRef GetMeshDataBufferRHI() const { return BufferMeshDataRHI; }
 	FShaderResourceViewRHIRef GetMeshDataBufferSRV() const { return BufferMeshDataSRV; }
 
 	virtual void InitRHI() override final
@@ -115,7 +115,7 @@ public:
 	}
 
 private:
-	FVertexBufferRHIRef BufferMeshDataRHI = nullptr;
+	FBufferRHIRef BufferMeshDataRHI = nullptr;
 	FShaderResourceViewRHIRef BufferMeshDataSRV = nullptr;
 	const FNDIMeshRendererInfo::FMeshDataArray& MeshData; // cached from FNDIMeshRendererInfo, which is guaranteed to live longer than we are
 #if STATS
diff --git a/Engine/Plugins/FX/Niagara/Source/Niagara/Private/NiagaraRenderer.cpp b/Engine/Plugins/FX/Niagara/Source/Niagara/Private/NiagaraRenderer.cpp
--- a/Engine/Plugins/FX/Niagara/Source/Niagara/Private/NiagaraRenderer.cpp
+++ b/Engine/Plugins/FX/Niagara/Source/Niagara/Private/NiagaraRenderer.cpp
@@ -42,7 +42,7 @@ public:
 	FNiagaraEmptyBufferSRV(EPixelFormat InPixelFormat, const FString& InDebugName, uint32 InDefaultValue = 0) : PixelFormat(InPixelFormat), DebugName(InDebugName), DefaultValue(InDefaultValue) {}
 	EPixelFormat PixelFormat;
 	FString DebugName;
-	FVertexBufferRHIRef Buffer;
+	FBufferRHIRef Buffer;
 	FShaderResourceViewRHIRef SRV;
 	uint32 DefaultValue = 0;
 
diff --git a/Engine/Plugins/FX/Niagara/Source/NiagaraShader/Private/NiagaraDebugShaders.cpp b/Engine/Plugins/FX/Niagara/Source/NiagaraShader/Private/NiagaraDebugShaders.cpp
--- a/Engine/Plugins/FX/Niagara/Source/NiagaraShader/Private/NiagaraDebugShaders.cpp
+++ b/Engine/Plugins/FX/Niagara/Source/NiagaraShader/Private/NiagaraDebugShaders.cpp
@@ -222,7 +222,7 @@ void NiagaraDebugShaders::DrawDebugLines(
 
 void NiagaraDebugShaders::DrawDebugLines(
 	class FRDGBuilder& GraphBuilder, const class FViewInfo& View, FRDGTextureRef SceneColor, FRDGTextureRef SceneDepth,
-	const FVertexBufferRHIRef& ArgsBuffer, const FShaderResourceViewRHIRef& LineBuffer
+	const FBufferRHIRef& ArgsBuffer, const FShaderResourceViewRHIRef& LineBuffer
 )
 {
 	TShaderMapRef<FNiagaraDebugDrawLineVS> VertexShader(View.ShaderMap);
diff --git a/Engine/Plugins/FX/Niagara/Source/NiagaraShader/Public/NiagaraDebugShaders.h b/Engine/Plugins/FX/Niagara/Source/NiagaraShader/Public/NiagaraDebugShaders.h
--- a/Engine/Plugins/FX/Niagara/Source/NiagaraShader/Public/NiagaraDebugShaders.h
+++ b/Engine/Plugins/FX/Niagara/Source/NiagaraShader/Public/NiagaraDebugShaders.h
@@ -18,7 +18,7 @@ namespace NiagaraDebugShaders
 
 	NIAGARASHADER_API void DrawDebugLines(
 		class FRDGBuilder& GraphBuilder, const class FViewInfo& View, FRDGTextureRef SceneColor, FRDGTextureRef SceneDepth,
-		const FVertexBufferRHIRef& ArgsBuffer, const FShaderResourceViewRHIRef& LineBuffer
+		const FBufferRHIRef& ArgsBuffer, const FShaderResourceViewRHIRef& LineBuffer
 	);
 
 	NIAGARASHADER_API void VisualizeTexture(
diff --git a/Engine/Plugins/FX/Niagara/Source/NiagaraVertexFactories/Public/NiagaraRibbonVertexFactory.h b/Engine/Plugins/FX/Niagara/Source/NiagaraVertexFactories/Public/NiagaraRibbonVertexFactory.h
--- a/Engine/Plugins/FX/Niagara/Source/NiagaraVertexFactories/Public/NiagaraRibbonVertexFactory.h
+++ b/Engine/Plugins/FX/Niagara/Source/NiagaraVertexFactories/Public/NiagaraRibbonVertexFactory.h
@@ -135,27 +135,27 @@ public:
 	*/
 	void SetDynamicParameterBuffer(const FVertexBuffer* InDynamicParameterBuffer, int32 ParameterIndex, uint32 StreamOffset, uint32 Stride);
 
-	void SetSortedIndices(const FVertexBufferRHIRef& InSortedIndicesBuffer, const FShaderResourceViewRHIRef& InSortedIndicesSRV, uint32 InSortedIndicesOffset)
+	void SetSortedIndices(const FBufferRHIRef& InSortedIndicesBuffer, const FShaderResourceViewRHIRef& InSortedIndicesSRV, uint32 InSortedIndicesOffset)
 	{
 		SortedIndicesBuffer = InSortedIndicesBuffer;
 		SortedIndicesSRV = InSortedIndicesSRV;
 		SortedIndicesOffset = InSortedIndicesOffset;
 	}
 
-	void SetTangentAndDistances(const FVertexBufferRHIRef& InTangentAndDistancesBuffer, const FShaderResourceViewRHIRef& InTangentAndDistancesSRV)
+	void SetTangentAndDistances(const FBufferRHIRef& InTangentAndDistancesBuffer, const FShaderResourceViewRHIRef& InTangentAndDistancesSRV)
 	{
 		TangentAndDistancesBuffer = InTangentAndDistancesBuffer;
 		TangentAndDistancesSRV = InTangentAndDistancesSRV;
 	}
 
 
-	void SetMultiRibbonIndicesSRV(const FVertexBufferRHIRef& InMultiRibbonIndicesBuffer, const FShaderResourceViewRHIRef& InMultiRibbonIndicesSRV)
+	void SetMultiRibbonIndicesSRV(const FBufferRHIRef& InMultiRibbonIndicesBuffer, const FShaderResourceViewRHIRef& InMultiRibbonIndicesSRV)
 	{
 		MultiRibbonIndicesBuffer = InMultiRibbonIndicesBuffer;
 		MultiRibbonIndicesSRV = InMultiRibbonIndicesSRV;
 	}
 
-	void SetPackedPerRibbonDataByIndexSRV(const FVertexBufferRHIRef& InPackedPerRibbonDataByIndexBuffer, const FShaderResourceViewRHIRef& InPackedPerRibbonDataByIndexSRV)
+	void SetPackedPerRibbonDataByIndexSRV(const FBufferRHIRef& InPackedPerRibbonDataByIndexBuffer, const FShaderResourceViewRHIRef& InPackedPerRibbonDataByIndexSRV)
 	{
 		PackedPerRibbonDataByIndexBuffer = InPackedPerRibbonDataByIndexBuffer;
 		PackedPerRibbonDataByIndexSRV = InPackedPerRibbonDataByIndexSRV;
@@ -224,10 +224,10 @@ private:
 
 	const FNiagaraDataSet *DataSet;
 
-	FVertexBufferRHIRef SortedIndicesBuffer;
-	FVertexBufferRHIRef TangentAndDistancesBuffer;
-	FVertexBufferRHIRef MultiRibbonIndicesBuffer;
-	FVertexBufferRHIRef PackedPerRibbonDataByIndexBuffer;
+	FBufferRHIRef SortedIndicesBuffer;
+	FBufferRHIRef TangentAndDistancesBuffer;
+	FBufferRHIRef MultiRibbonIndicesBuffer;
+	FBufferRHIRef PackedPerRibbonDataByIndexBuffer;
 
 	FShaderResourceViewRHIRef SortedIndicesSRV;
 	FShaderResourceViewRHIRef TangentAndDistancesSRV;
diff --git a/Engine/Plugins/Media/AvfMedia/Source/AvfMedia/Private/Player/AvfMediaVideoSampler.cpp b/Engine/Plugins/Media/AvfMedia/Source/AvfMedia/Private/Player/AvfMediaVideoSampler.cpp
--- a/Engine/Plugins/Media/AvfMedia/Source/AvfMedia/Private/Player/AvfMediaVideoSampler.cpp
+++ b/Engine/Plugins/Media/AvfMedia/Source/AvfMedia/Private/Player/AvfMediaVideoSampler.cpp
@@ -313,7 +313,7 @@ void FAvfMediaVideoSampler::ProcessFrame(CVPixelBufferRef Frame, FTimespan Sampl
 
 					PixelShader->SetParameters(RHICmdList, YTex, UVTex, *ColorTransform, MediaShaders::YUVOffset8bits, true);
 
-					FVertexBufferRHIRef VertexBuffer = CreateTempMediaVertexBuffer();
+					FBufferRHIRef VertexBuffer = CreateTempMediaVertexBuffer();
 					RHICmdList.SetStreamSource(0, VertexBuffer, 0);
 					RHICmdList.SetViewport(0, 0, 0.0f, YWidth, YHeight, 1.0f);
 
diff --git a/Engine/Plugins/Media/ElectraUtil/Source/ElectraSamples/Private/ElectraTextureSampleAVF.cpp b/Engine/Plugins/Media/ElectraUtil/Source/ElectraSamples/Private/ElectraTextureSampleAVF.cpp
--- a/Engine/Plugins/Media/ElectraUtil/Source/ElectraSamples/Private/ElectraTextureSampleAVF.cpp
+++ b/Engine/Plugins/Media/ElectraUtil/Source/ElectraSamples/Private/ElectraTextureSampleAVF.cpp
@@ -248,7 +248,7 @@ void FElectraMediaTexConvApple::ConvertTexture(FTexture2DRHIRef & InDstTexture,
 
 					PixelShader->SetParameters(RHICmdList, YTex, UVTex, *ColorTransform, MediaShaders::YUVOffset8bits, true);
 
-					FVertexBufferRHIRef VertexBuffer = CreateTempMediaVertexBuffer();
+					FBufferRHIRef VertexBuffer = CreateTempMediaVertexBuffer();
 					RHICmdList.SetStreamSource(0, VertexBuffer, 0);
 					RHICmdList.SetViewport(0, 0, 0.0f, YWidth, YHeight, 1.0f);
 
diff --git a/Engine/Plugins/Media/ImgMedia/Source/ImgMedia/Private/Readers/ExrImgMediaReaderGpu.h b/Engine/Plugins/Media/ImgMedia/Source/ImgMedia/Private/Readers/ExrImgMediaReaderGpu.h
--- a/Engine/Plugins/Media/ImgMedia/Source/ImgMedia/Private/Readers/ExrImgMediaReaderGpu.h
+++ b/Engine/Plugins/Media/ImgMedia/Source/ImgMedia/Private/Readers/ExrImgMediaReaderGpu.h
@@ -15,7 +15,7 @@ struct FStructuredBufferPoolItem
 	/**
 	* This is the actual buffer reference that we need to keep after it is locked and until it is unlocked.
 	*/
-	FStructuredBufferRHIRef BufferRef;
+	FBufferRHIRef BufferRef;
 
 	/** 
 	* A pointer to mapped GPU memory.
diff --git a/Engine/Plugins/Media/MediaIOFramework/Source/MediaIOCore/Private/MediaCapture.cpp b/Engine/Plugins/Media/MediaIOFramework/Source/MediaIOCore/Private/MediaCapture.cpp
--- a/Engine/Plugins/Media/MediaIOFramework/Source/MediaIOCore/Private/MediaCapture.cpp
+++ b/Engine/Plugins/Media/MediaIOFramework/Source/MediaIOCore/Private/MediaCapture.cpp
@@ -883,7 +883,7 @@ void UMediaCapture::Capture_RenderThread(FRHICommandListImmediate& RHICmdList,
 				}
 
 				// draw full size quad into render target
-				FVertexBufferRHIRef VertexBuffer = CreateTempMediaVertexBuffer(ULeft, URight, VTop, VBottom);
+				FBufferRHIRef VertexBuffer = CreateTempMediaVertexBuffer(ULeft, URight, VTop, VBottom);
 				RHICmdList.SetStreamSource(0, VertexBuffer, 0);
 
 				// set viewport to RT size
diff --git a/Engine/Plugins/Media/WebMMedia/Source/WebMMedia/Private/Player/WebMVideoDecoder.cpp b/Engine/Plugins/Media/WebMMedia/Source/WebMMedia/Private/Player/WebMVideoDecoder.cpp
--- a/Engine/Plugins/Media/WebMMedia/Source/WebMMedia/Private/Player/WebMVideoDecoder.cpp
+++ b/Engine/Plugins/Media/WebMMedia/Source/WebMMedia/Private/Player/WebMVideoDecoder.cpp
@@ -24,7 +24,7 @@ namespace
 	public:
 
 		FVertexDeclarationRHIRef VertexDeclarationRHI;
-		FVertexBufferRHIRef VertexBufferRHI;
+		FBufferRHIRef VertexBufferRHI;
 
 		virtual ~FMoviePlaybackResources() { }
 
diff --git a/Engine/Plugins/Runtime/AR/AppleAR/AppleARKit/Source/AppleARKit/Private/AppleARKitVideoOverlay.h b/Engine/Plugins/Runtime/AR/AppleAR/AppleARKit/Source/AppleARKit/Private/AppleARKitVideoOverlay.h
--- a/Engine/Plugins/Runtime/AR/AppleAR/AppleARKit/Source/AppleARKit/Private/AppleARKitVideoOverlay.h
+++ b/Engine/Plugins/Runtime/AR/AppleAR/AppleARKit/Source/AppleARKit/Private/AppleARKitVideoOverlay.h
@@ -102,8 +102,8 @@ private:
 	void UpdateDebugOverlay();
 
 	// 0 for landscape, 1 for portrait
-	FVertexBufferRHIRef OverlayVertexBufferRHI;
-	FIndexBufferRHIRef IndexBufferRHI;
+	FBufferRHIRef OverlayVertexBufferRHI;
+	FBufferRHIRef IndexBufferRHI;
 	
 	EARKitOcclusionType OcclusionType = EARKitOcclusionType::None;
 	
diff --git a/Engine/Plugins/Runtime/AR/Google/GoogleARCore/Source/GoogleARCoreRendering/Public/GoogleARCorePassthroughCameraRenderer.h b/Engine/Plugins/Runtime/AR/Google/GoogleARCore/Source/GoogleARCoreRendering/Public/GoogleARCorePassthroughCameraRenderer.h
--- a/Engine/Plugins/Runtime/AR/Google/GoogleARCore/Source/GoogleARCoreRendering/Public/GoogleARCorePassthroughCameraRenderer.h
+++ b/Engine/Plugins/Runtime/AR/Google/GoogleARCore/Source/GoogleARCoreRendering/Public/GoogleARCorePassthroughCameraRenderer.h
@@ -78,8 +78,8 @@ private:
 	void RenderVideoOverlayWithMaterial(FRHICommandListImmediate& RHICmdList, FSceneView& InView, UMaterialInstanceDynamic* OverlayMaterialToUse, bool bRenderingOcclusion);
 
 private:
-	FIndexBufferRHIRef OverlayIndexBufferRHI;
-	FVertexBufferRHIRef OverlayVertexBufferRHI;
+	FBufferRHIRef OverlayIndexBufferRHI;
+	FBufferRHIRef OverlayVertexBufferRHI;
 	
 	UMaterialInstanceDynamic* RegularOverlayMaterial = nullptr;
 	UMaterialInstanceDynamic* DebugOverlayMaterial = nullptr;
diff --git a/Engine/Plugins/Runtime/GoogleVR/GoogleVRHMD/Source/GoogleVRHMD/Private/GoogleVRHMD.h b/Engine/Plugins/Runtime/GoogleVR/GoogleVRHMD/Source/GoogleVRHMD/Private/GoogleVRHMD.h
--- a/Engine/Plugins/Runtime/GoogleVR/GoogleVRHMD/Source/GoogleVRHMD/Private/GoogleVRHMD.h
+++ b/Engine/Plugins/Runtime/GoogleVR/GoogleVRHMD/Source/GoogleVRHMD/Private/GoogleVRHMD.h
@@ -388,9 +388,9 @@ private:
 	FIntPoint GVRRenderTargetSize;
 	IRendererModule* RendererModule;
 
-	FIndexBufferRHIRef DistortionMeshIndices;
-	FVertexBufferRHIRef DistortionMeshVerticesLeftEye;
-	FVertexBufferRHIRef DistortionMeshVerticesRightEye;
+	FBufferRHIRef DistortionMeshIndices;
+	FBufferRHIRef DistortionMeshVerticesLeftEye;
+	FBufferRHIRef DistortionMeshVerticesRightEye;
 
 	// Cached data that should only be updated once per frame
 	mutable uint32		LastUpdatedCacheFrame;
diff --git a/Engine/Plugins/Runtime/GoogleVR/GoogleVRHMD/Source/GoogleVRHMD/Private/GoogleVRHMDRender.cpp b/Engine/Plugins/Runtime/GoogleVR/GoogleVRHMD/Source/GoogleVRHMD/Private/GoogleVRHMDRender.cpp
--- a/Engine/Plugins/Runtime/GoogleVR/GoogleVRHMD/Source/GoogleVRHMD/Private/GoogleVRHMDRender.cpp
+++ b/Engine/Plugins/Runtime/GoogleVR/GoogleVRHMD/Source/GoogleVRHMD/Private/GoogleVRHMDRender.cpp
@@ -43,7 +43,7 @@ void FGoogleVRHMD::GenerateDistortionCorrectionIndexBuffer()
 
 void FGoogleVRHMD::GenerateDistortionCorrectionVertexBuffer(EStereoscopicPass Eye)
 {
-	FVertexBufferRHIRef& DistortionMeshVertices = (Eye == eSSP_LEFT_EYE) ? DistortionMeshVerticesLeftEye : DistortionMeshVerticesRightEye;
+	FBufferRHIRef& DistortionMeshVertices = (Eye == eSSP_LEFT_EYE) ? DistortionMeshVerticesLeftEye : DistortionMeshVerticesRightEye;
 	FRHIResourceCreateInfo CreateInfo;
 	DistortionMeshVertices = RHICreateVertexBuffer(sizeof(FDistortionVertex) * NumVerts, BUF_Static, CreateInfo);
 	void* VoidPtr = RHILockBuffer(DistortionMeshVertices, 0, sizeof(FDistortionVertex) * NumVerts, RLM_WriteOnly);
@@ -138,7 +138,7 @@ void FGoogleVRHMD::DrawDistortionMesh_RenderThread(struct FHeadMountedDisplayPas
 		};
 
 		FRHIResourceCreateInfo CreateInfo;
-		FVertexBufferRHIRef VertexBufferRHI = RHICreateVertexBuffer(sizeof(FDistortionVertex) * 4, BUF_Volatile, CreateInfo);
+		FBufferRHIRef VertexBufferRHI = RHICreateVertexBuffer(sizeof(FDistortionVertex) * 4, BUF_Volatile, CreateInfo);
 		void* VoidPtr = RHILockBuffer(VertexBufferRHI, 0, sizeof(FDistortionVertex) * 4, RLM_WriteOnly);
 		FPlatformMemory::Memcpy(VoidPtr, Verts, sizeof(FDistortionVertex) * 4);
 		RHIUnlockBuffer(VertexBufferRHI);
diff --git a/Engine/Plugins/Runtime/HairStrands/Source/HairStrandsCore/Private/GroomTextureBuilder.cpp b/Engine/Plugins/Runtime/HairStrands/Source/HairStrandsCore/Private/GroomTextureBuilder.cpp
--- a/Engine/Plugins/Runtime/HairStrands/Source/HairStrandsCore/Private/GroomTextureBuilder.cpp
+++ b/Engine/Plugins/Runtime/HairStrands/Source/HairStrandsCore/Private/GroomTextureBuilder.cpp
@@ -869,7 +869,7 @@ static void InternalBuildStrandsTextures_GPU(
 			FRHIShaderResourceView* PositionBuffer = nullptr;
 			FRHIShaderResourceView* UVsBuffer = nullptr;
 			FRHIShaderResourceView* TangentBuffer = nullptr;
-			FIndexBufferRHIRef IndexBuffer = nullptr;
+			FBufferRHIRef IndexBuffer = nullptr;
 			uint32 TotalVertexCount = 0;
 			uint32 TotalIndexCount = 0;
 			uint32 UVsChannelIndex = InInfo.UVChannelIndex;
diff --git a/Engine/Plugins/Runtime/HairStrands/Source/HairStrandsCore/Private/HairStrandsRendering.cpp b/Engine/Plugins/Runtime/HairStrands/Source/HairStrandsCore/Private/HairStrandsRendering.cpp
--- a/Engine/Plugins/Runtime/HairStrands/Source/HairStrandsCore/Private/HairStrandsRendering.cpp
+++ b/Engine/Plugins/Runtime/HairStrands/Source/HairStrandsCore/Private/HairStrandsRendering.cpp
@@ -1105,7 +1105,7 @@ static void UpdateHairAccelerationStructure(FRHICommandList& RHICmdList, FRayTra
 	RHICmdList.BuildAccelerationStructures(MakeArrayView(&Params, 1));
 }
 
-static void BuildHairAccelerationStructure_Strands(FRHICommandList& RHICmdList, uint32 RaytracingVertexCount, FVertexBufferRHIRef& PositionBuffer, FRayTracingGeometry* OutRayTracingGeometry)
+static void BuildHairAccelerationStructure_Strands(FRHICommandList& RHICmdList, uint32 RaytracingVertexCount, FBufferRHIRef& PositionBuffer, FRayTracingGeometry* OutRayTracingGeometry)
 {
 	FRayTracingGeometryInitializer Initializer;
 	static const FName DebugName("HairStrands");
@@ -1145,7 +1145,7 @@ static void BuildHairAccelerationStructure_Cards(FRHICommandList& RHICmdList,
 	Initializer.bFastBuild = true;
 	Initializer.bAllowUpdate = true;
 
-	FVertexBufferRHIRef PositionBuffer(DeformedResource->GetBuffer(FHairCardsDeformedResource::Current).Buffer->GetRHI()); // This will likely flicker result in half speed motion as everyother frame will use the wrong buffer
+	FBufferRHIRef PositionBuffer(DeformedResource->GetBuffer(FHairCardsDeformedResource::Current).Buffer->GetRHI()); // This will likely flicker result in half speed motion as everyother frame will use the wrong buffer
 
 	FRayTracingGeometrySegment Segment;
 	Segment.VertexBuffer = PositionBuffer;
@@ -1174,7 +1174,7 @@ static void BuildHairAccelerationStructure_Meshes(FRHICommandList& RHICmdList,
 	Initializer.bFastBuild = true;
 	Initializer.bAllowUpdate = true;
 
-	FVertexBufferRHIRef PositionBuffer(DeformedResource->GetBuffer(FHairMeshesDeformedResource::Current).Buffer->GetRHI()); // This will likely flicker result in half speed motion as everyother frame will use the wrong buffer
+	FBufferRHIRef PositionBuffer(DeformedResource->GetBuffer(FHairMeshesDeformedResource::Current).Buffer->GetRHI()); // This will likely flicker result in half speed motion as everyother frame will use the wrong buffer
 
 	FRayTracingGeometrySegment Segment;
 	Segment.VertexBuffer = PositionBuffer;
@@ -1503,7 +1503,7 @@ void ComputeHairStrandsInterpolation(
 				const bool bNeedFullBuild = !Instance->Strands.RenRaytracingResource->bIsRTGeometryInitialized;
 				if (bNeedFullBuild)
 				{
-					FVertexBufferRHIRef PositionBuffer(Instance->Strands.RenRaytracingResource->PositionBuffer.Buffer->GetRHI());
+					FBufferRHIRef PositionBuffer(Instance->Strands.RenRaytracingResource->PositionBuffer.Buffer->GetRHI());
 					BuildHairAccelerationStructure_Strands(RHICmdList, Instance->Strands.RenRaytracingResource->VertexCount, PositionBuffer, &Instance->Strands.RenRaytracingResource->RayTracingGeometry);
 				}
 				else
diff --git a/Engine/Plugins/Runtime/Oculus/OculusVR/Source/OculusHMD/Private/OculusHMD_StressTester.cpp b/Engine/Plugins/Runtime/Oculus/OculusVR/Source/OculusHMD/Private/OculusHMD_StressTester.cpp
--- a/Engine/Plugins/Runtime/Oculus/OculusVR/Source/OculusHMD/Private/OculusHMD_StressTester.cpp
+++ b/Engine/Plugins/Runtime/Oculus/OculusVR/Source/OculusHMD/Private/OculusHMD_StressTester.cpp
@@ -53,10 +53,10 @@ struct FTextureVertex
 	FVector2D	UV;
 };
 
-inline FVertexBufferRHIRef CreateTempOcculusVertexBuffer()
+inline FBufferRHIRef CreateTempOcculusVertexBuffer()
 {
 	FRHIResourceCreateInfo CreateInfo;
-	FVertexBufferRHIRef VertexBufferRHI = RHICreateVertexBuffer(sizeof(FTextureVertex) * 4, BUF_Volatile, CreateInfo);
+	FBufferRHIRef VertexBufferRHI = RHICreateVertexBuffer(sizeof(FTextureVertex) * 4, BUF_Volatile, CreateInfo);
 	void* VoidPtr = RHILockBuffer(VertexBufferRHI, 0, sizeof(FTextureVertex) * 4, RLM_WriteOnly);
 
 	FTextureVertex* Vertices = (FTextureVertex*)VoidPtr;
diff --git a/Engine/Plugins/Runtime/WarpUtils/Source/PFMExporter/Private/PFMExporterShader.cpp b/Engine/Plugins/Runtime/WarpUtils/Source/PFMExporter/Private/PFMExporterShader.cpp
--- a/Engine/Plugins/Runtime/WarpUtils/Source/PFMExporter/Private/PFMExporterShader.cpp
+++ b/Engine/Plugins/Runtime/WarpUtils/Source/PFMExporter/Private/PFMExporterShader.cpp
@@ -108,7 +108,7 @@ bool FPFMExporterShader::ApplyPFMExporter_RenderThread(
 	DstPfmMesh.BeginExport_RenderThread(RHICmdList);
 
 	FRHIResourceCreateInfo CreateInfo;
-	FVertexBufferRHIRef VertexBufferRHI = RHICreateVertexBuffer(sizeof(FFilterVertex) * NumVertices, BUF_Dynamic, CreateInfo);
+	FBufferRHIRef VertexBufferRHI = RHICreateVertexBuffer(sizeof(FFilterVertex) * NumVertices, BUF_Dynamic, CreateInfo);
 	{//Fill buffer with vertex+selected UV channel:
 		void* VoidPtr = RHILockBuffer(VertexBufferRHI, 0, sizeof(FFilterVertex) * NumVertices, RLM_WriteOnly);
 		FFilterVertex* pVertices = reinterpret_cast<FFilterVertex*>(VoidPtr);
diff --git a/Engine/Plugins/Runtime/WindowsMoviePlayer/Source/WindowsMoviePlayer/Private/WindowsMovieStreamer.cpp b/Engine/Plugins/Runtime/WindowsMoviePlayer/Source/WindowsMoviePlayer/Private/WindowsMovieStreamer.cpp
--- a/Engine/Plugins/Runtime/WindowsMoviePlayer/Source/WindowsMoviePlayer/Private/WindowsMovieStreamer.cpp
+++ b/Engine/Plugins/Runtime/WindowsMoviePlayer/Source/WindowsMoviePlayer/Private/WindowsMovieStreamer.cpp
@@ -165,7 +165,7 @@ void FMediaFoundationMovieStreamer::ConvertSample()
 		}
 
 		// draw full size quad into render target
-		FVertexBufferRHIRef VertexBuffer = CreateTempMediaVertexBuffer();
+		FBufferRHIRef VertexBuffer = CreateTempMediaVertexBuffer();
 		CommandList.SetStreamSource(0, VertexBuffer, 0);
 		// set viewport to RT size
 		CommandList.SetViewport(0, 0, 0.0f, OutputDim.X, OutputDim.Y, 1.0f);
diff --git a/Engine/Plugins/Runtime/nDisplay/Source/MPCDI/Private/MPCDIWarpMesh.h b/Engine/Plugins/Runtime/nDisplay/Source/MPCDI/Private/MPCDIWarpMesh.h
--- a/Engine/Plugins/Runtime/nDisplay/Source/MPCDI/Private/MPCDIWarpMesh.h
+++ b/Engine/Plugins/Runtime/nDisplay/Source/MPCDI/Private/MPCDIWarpMesh.h
@@ -98,8 +98,8 @@ private:
 
 	FTransform MeshToOrigin;
 
-	FVertexBufferRHIRef VertexBufferRHI;
-	FIndexBufferRHIRef  IndexBufferRHI;
+	FBufferRHIRef VertexBufferRHI;
+	FBufferRHIRef IndexBufferRHI;
 
 	uint32 NumTriangles;
 	uint32 NumVertices;
diff --git a/Engine/Plugins/Runtime/nDisplay/Source/OutputRemap/Private/OutputRemapMesh.h b/Engine/Plugins/Runtime/nDisplay/Source/OutputRemap/Private/OutputRemapMesh.h
--- a/Engine/Plugins/Runtime/nDisplay/Source/OutputRemap/Private/OutputRemapMesh.h
+++ b/Engine/Plugins/Runtime/nDisplay/Source/OutputRemap/Private/OutputRemapMesh.h
@@ -69,8 +69,8 @@ private:
 	void ReloadMesh();
 
 public:
-	FVertexBufferRHIRef VertexBufferRHI;
-	FIndexBufferRHIRef  IndexBufferRHI;
+	FBufferRHIRef VertexBufferRHI;
+	FBufferRHIRef IndexBufferRHI;
 
 	uint32  NumVertices;
 	uint32  NumIndices;
diff --git a/Engine/Plugins/Tests/RHITests/Source/RHITests/Public/RHIBufferTests.h b/Engine/Plugins/Tests/RHITests/Source/RHITests/Public/RHIBufferTests.h
--- a/Engine/Plugins/Tests/RHITests/Source/RHITests/Public/RHIBufferTests.h
+++ b/Engine/Plugins/Tests/RHITests/Source/RHITests/Public/RHIBufferTests.h
@@ -75,7 +75,7 @@ class FRHIBufferTests
 		}
 
 		FRHIResourceCreateInfo Info;
-		FVertexBufferRHIRef VertexBuffer = RHICreateVertexBuffer(BufferSize, BUF_ShaderResource | BUF_UnorderedAccess | BUF_SourceCopy, Info);
+		FBufferRHIRef VertexBuffer = RHICreateVertexBuffer(BufferSize, BUF_ShaderResource | BUF_UnorderedAccess | BUF_SourceCopy, Info);
 		FUnorderedAccessViewRHIRef UAV = RHICreateUnorderedAccessView(VertexBuffer, Format);
 		FShaderResourceViewRHIRef SRV = RHICreateShaderResourceView(VertexBuffer, GPixelFormats[Format].BlockBytes, Format);
 		bool bResult = RunTest_UAVClear_Buffer(RHICmdList, TestName, VertexBuffer.GetReference(), UAV, BufferSize, ClearValue, ClearPtr, TestValue);
@@ -97,7 +97,7 @@ class FRHIBufferTests
 		}
 
 		FRHIResourceCreateInfo Info;
-		FVertexBufferRHIRef VertexBuffer = RHICreateVertexBuffer(BufferSize, BUF_ShaderResource | BUF_UnorderedAccess | BUF_SourceCopy | BUF_ByteAddressBuffer, Info);
+		FBufferRHIRef VertexBuffer = RHICreateVertexBuffer(BufferSize, BUF_ShaderResource | BUF_UnorderedAccess | BUF_SourceCopy | BUF_ByteAddressBuffer, Info);
 		FUnorderedAccessViewRHIRef UAV = RHICreateUnorderedAccessView(VertexBuffer, Format);
 		FShaderResourceViewRHIRef SRV = RHICreateShaderResourceView(VertexBuffer, GPixelFormats[Format].BlockBytes, Format);
 		bool bResult = RunTest_UAVClear_Buffer(RHICmdList, TestName, VertexBuffer.GetReference(), UAV, BufferSize, ClearValue, ClearPtr, TestValue);
@@ -113,7 +113,7 @@ class FRHIBufferTests
 		check(NumTestBytes == Stride);
 
 		FRHIResourceCreateInfo Info;
-		FStructuredBufferRHIRef StructuredBuffer = RHICreateStructuredBuffer(Stride, BufferSize, BUF_ShaderResource | BUF_UnorderedAccess | BUF_SourceCopy | InExtraUsage, Info);
+		FBufferRHIRef StructuredBuffer = RHICreateStructuredBuffer(Stride, BufferSize, BUF_ShaderResource | BUF_UnorderedAccess | BUF_SourceCopy | InExtraUsage, Info);
 		FUnorderedAccessViewRHIRef UAV = RHICreateUnorderedAccessView(StructuredBuffer, false, false);
 		FShaderResourceViewRHIRef SRV = RHICreateShaderResourceView(StructuredBuffer);
 		bool bResult = RunTest_UAVClear_Buffer(RHICmdList, *FString::Printf(TEXT("RunTest_UAVClear_StructuredBuffer, Stride: %d, Size: %d, ByteAddress: %d"), Stride, BufferSize, (InExtraUsage & BUF_ByteAddressBuffer) ? 1 : 0), StructuredBuffer.GetReference(), UAV, BufferSize, ClearValue, ClearPtr, TestValue);
diff --git a/Engine/Source/Developer/MaterialBaking/Private/MaterialBakingModule.cpp b/Engine/Source/Developer/MaterialBaking/Private/MaterialBakingModule.cpp
--- a/Engine/Source/Developer/MaterialBaking/Private/MaterialBakingModule.cpp
+++ b/Engine/Source/Developer/MaterialBaking/Private/MaterialBakingModule.cpp
@@ -60,8 +60,8 @@ namespace FMaterialBakingModuleImpl
 		// memory is first initialized.
 		const uint32 SmallestPooledBufferSize = 256*1024;
 
-		TArray<FIndexBufferRHIRef>  IndexBuffers;
-		TArray<FVertexBufferRHIRef> VertexBuffers;
+		TArray<FBufferRHIRef>  IndexBuffers;
+		TArray<FBufferRHIRef> VertexBuffers;
 
 		template <typename RefType>
 		RefType GetSmallestFit(uint32 SizeInBytes, TArray<RefType>& Array)
@@ -89,12 +89,12 @@ namespace FMaterialBakingModuleImpl
 			return Ref;
 		}
 
-		virtual FIndexBufferRHIRef AllocIndexBuffer(uint32 NumElements) override
+		virtual FBufferRHIRef AllocIndexBuffer(uint32 NumElements) override
 		{
 			uint32 BufferSize = GetIndexBufferSize(NumElements);
 			if (BufferSize > SmallestPooledBufferSize)
 			{
-				FIndexBufferRHIRef Ref = GetSmallestFit(GetIndexBufferSize(NumElements), IndexBuffers);
+				FBufferRHIRef Ref = GetSmallestFit(GetIndexBufferSize(NumElements), IndexBuffers);
 				if (Ref.IsValid())
 				{
 					return Ref;
@@ -104,7 +104,7 @@ namespace FMaterialBakingModuleImpl
 			return FDynamicMeshBufferAllocator::AllocIndexBuffer(NumElements);
 		}
 
-		virtual void ReleaseIndexBuffer(FIndexBufferRHIRef& IndexBufferRHI) override
+		virtual void ReleaseIndexBuffer(FBufferRHIRef& IndexBufferRHI) override
 		{
 			if (IndexBufferRHI->GetSize() > SmallestPooledBufferSize)
 			{
@@ -114,12 +114,12 @@ namespace FMaterialBakingModuleImpl
 			IndexBufferRHI = nullptr;
 		}
 
-		virtual FVertexBufferRHIRef AllocVertexBuffer(uint32 Stride, uint32 NumElements) override
+		virtual FBufferRHIRef AllocVertexBuffer(uint32 Stride, uint32 NumElements) override
 		{
 			uint32 BufferSize = GetVertexBufferSize(Stride, NumElements);
 			if (BufferSize > SmallestPooledBufferSize)
 			{
-				FVertexBufferRHIRef Ref = GetSmallestFit(BufferSize, VertexBuffers);
+				FBufferRHIRef Ref = GetSmallestFit(BufferSize, VertexBuffers);
 				if (Ref.IsValid())
 				{
 					return Ref;
@@ -129,7 +129,7 @@ namespace FMaterialBakingModuleImpl
 			return FDynamicMeshBufferAllocator::AllocVertexBuffer(Stride, NumElements);
 		}
 
-		virtual void ReleaseVertexBuffer(FVertexBufferRHIRef& VertexBufferRHI) override
+		virtual void ReleaseVertexBuffer(FBufferRHIRef& VertexBufferRHI) override
 		{
 			if (VertexBufferRHI->GetSize() > SmallestPooledBufferSize)
 			{
diff --git a/Engine/Source/Runtime/D3D12RHI/Private/D3D12RayTracing.cpp b/Engine/Source/Runtime/D3D12RHI/Private/D3D12RayTracing.cpp
--- a/Engine/Source/Runtime/D3D12RHI/Private/D3D12RayTracing.cpp
+++ b/Engine/Source/Runtime/D3D12RHI/Private/D3D12RayTracing.cpp
@@ -2371,7 +2371,7 @@ FRayTracingSceneRHIRef FD3D12DynamicRHI::RHICreateRayTracingScene(const FRayTrac
 	return new FD3D12RayTracingScene(&Adapter, Initializer);
 }
 
-FVertexBufferRHIRef FD3D12RayTracingGeometry::NullTransformBuffer;
+FBufferRHIRef FD3D12RayTracingGeometry::NullTransformBuffer;
 
 FD3D12RayTracingGeometry::FD3D12RayTracingGeometry(const FRayTracingGeometryInitializer& Initializer)
 {
@@ -2603,7 +2603,7 @@ void FD3D12RayTracingGeometry::TransitionBuffers(FD3D12CommandContext& CommandCo
 
 	for (const FRayTracingGeometrySegment& Segment : Segments)
 	{
-		const FVertexBufferRHIRef& RHIVertexBuffer = Segment.VertexBuffer;
+		const FBufferRHIRef& RHIVertexBuffer = Segment.VertexBuffer;
 		FD3D12Buffer* VertexBuffer = CommandContext.RetrieveObject<FD3D12Buffer>(RHIVertexBuffer.GetReference());
 		if (VertexBuffer->GetResource()->RequiresResourceStateTracking())
 		{
@@ -4044,7 +4044,7 @@ static void DispatchRays(FD3D12CommandContext& CommandContext,
 
 	FD3D12Adapter* Adapter = CommandContext.GetParentDevice()->GetParentAdapter();
 
-	FVertexBufferRHIRef DispatchRaysDescBufferRHI;
+	FBufferRHIRef DispatchRaysDescBufferRHI;
 	FD3D12Buffer* DispatchRaysDescBuffer = nullptr;
 
 	if (ArgumentBuffer)
diff --git a/Engine/Source/Runtime/D3D12RHI/Private/D3D12RayTracing.h b/Engine/Source/Runtime/D3D12RHI/Private/D3D12RayTracing.h
--- a/Engine/Source/Runtime/D3D12RHI/Private/D3D12RayTracing.h
+++ b/Engine/Source/Runtime/D3D12RHI/Private/D3D12RayTracing.h
@@ -59,8 +59,8 @@ public:
 	TArray<FRayTracingGeometrySegment> Segments; // Defines addressable parts of the mesh that can be used for material assignment (one segment = one SBT record)
 	D3D12_RAYTRACING_ACCELERATION_STRUCTURE_BUILD_FLAGS BuildFlags;
 
-	FIndexBufferRHIRef  RHIIndexBuffer;
-	static FVertexBufferRHIRef NullTransformBuffer; // Null transform for hidden sections
+	FBufferRHIRef  RHIIndexBuffer;
+	static FBufferRHIRef NullTransformBuffer; // Null transform for hidden sections
 
 	TRefCountPtr<FD3D12Buffer> AccelerationStructureBuffers[MAX_NUM_GPUS];
 	TRefCountPtr<FD3D12Buffer> ScratchBuffers[MAX_NUM_GPUS];
diff --git a/Engine/Source/Runtime/EmptyRHI/Private/EmptyIndexBuffer.cpp b/Engine/Source/Runtime/EmptyRHI/Private/EmptyIndexBuffer.cpp
--- a/Engine/Source/Runtime/EmptyRHI/Private/EmptyIndexBuffer.cpp
+++ b/Engine/Source/Runtime/EmptyRHI/Private/EmptyIndexBuffer.cpp
@@ -23,7 +23,7 @@ void FEmptyIndexBuffer::Unlock()
 
 }
 
-FIndexBufferRHIRef FEmptyDynamicRHI::RHICreateIndexBuffer(uint32 Stride, uint32 Size, uint32 InUsage, FRHIResourceCreateInfo& CreateInfo)
+FBufferRHIRef FEmptyDynamicRHI::RHICreateIndexBuffer(uint32 Stride, uint32 Size, uint32 InUsage, FRHIResourceCreateInfo& CreateInfo)
 {
 	if (CreateInfo.bCreateRHIObjectOnly)
 	{
diff --git a/Engine/Source/Runtime/EmptyRHI/Private/EmptyStructuredBuffer.cpp b/Engine/Source/Runtime/EmptyRHI/Private/EmptyStructuredBuffer.cpp
--- a/Engine/Source/Runtime/EmptyRHI/Private/EmptyStructuredBuffer.cpp
+++ b/Engine/Source/Runtime/EmptyRHI/Private/EmptyStructuredBuffer.cpp
@@ -25,7 +25,7 @@ FEmptyStructuredBuffer::~FEmptyStructuredBuffer()
 
 
 
-FStructuredBufferRHIRef FEmptyDynamicRHI::RHICreateStructuredBuffer(uint32 Stride, uint32 Size, uint32 InUsage, FRHIResourceCreateInfo& CreateInfo)
+FBufferRHIRef FEmptyDynamicRHI::RHICreateStructuredBuffer(uint32 Stride, uint32 Size, uint32 InUsage, FRHIResourceCreateInfo& CreateInfo)
 {
 	return nullptr;//new FEmptyStructuredBuffer(Stride, Size, ResourceArray, InUsage);
 }
diff --git a/Engine/Source/Runtime/EmptyRHI/Private/EmptyVertexBuffer.cpp b/Engine/Source/Runtime/EmptyRHI/Private/EmptyVertexBuffer.cpp
--- a/Engine/Source/Runtime/EmptyRHI/Private/EmptyVertexBuffer.cpp
+++ b/Engine/Source/Runtime/EmptyRHI/Private/EmptyVertexBuffer.cpp
@@ -22,7 +22,7 @@ void FEmptyVertexBuffer::Unlock()
 	
 }
 
-FVertexBufferRHIRef FEmptyDynamicRHI::RHICreateVertexBuffer(uint32 Size, uint32 InUsage, FRHIResourceCreateInfo& CreateInfo)
+FBufferRHIRef FEmptyDynamicRHI::RHICreateVertexBuffer(uint32 Size, uint32 InUsage, FRHIResourceCreateInfo& CreateInfo)
 {
 	if (CreateInfo.bCreateRHIObjectOnly)
 	{
diff --git a/Engine/Source/Runtime/EmptyRHI/Public/EmptyRHI.h b/Engine/Source/Runtime/EmptyRHI/Public/EmptyRHI.h
--- a/Engine/Source/Runtime/EmptyRHI/Public/EmptyRHI.h
+++ b/Engine/Source/Runtime/EmptyRHI/Public/EmptyRHI.h
@@ -49,16 +49,16 @@ public:
 	virtual FComputeShaderRHIRef RHICreateComputeShader(TArrayView<const uint8> Code, const FSHAHash& Hash) final override;
 	virtual FBoundShaderStateRHIRef RHICreateBoundShaderState(FRHIVertexDeclaration* VertexDeclaration, FRHIVertexShader* VertexShader, FRHIHullShader* HullShader, FRHIDomainShaderFRHIDomainShader*, FRHIPixelShader* PixelShader, FRHIGeometryShader* GeometryShader) final override;
 	virtual FUniformBufferRHIRef RHICreateUniformBuffer(const void* Contents, const FRHIUniformBufferLayout& Layout, EUniformBufferUsage Usage, EUniformBufferValidation Validation) final override;
-	virtual FIndexBufferRHIRef RHICreateIndexBuffer(uint32 Stride, uint32 Size, uint32 InUsage, FRHIResourceCreateInfo& CreateInfo) final override;
+	virtual FBufferRHIRef RHICreateIndexBuffer(uint32 Stride, uint32 Size, uint32 InUsage, FRHIResourceCreateInfo& CreateInfo) final override;
 	virtual void* LockIndexBuffer_BottomOfPipe(FRHICommandListImmediate& RHICmdList, FRHIBuffer* IndexBuffer, uint32 Offset, uint32 Size, EResourceLockMode LockMode) final override;
 	virtual void UnlockIndexBuffer_BottomOfPipe(FRHICommandListImmediate& RHICmdList, FRHIBuffer* IndexBuffer) final override;
 	virtual void RHITransferIndexBufferUnderlyingResource(FRHIBuffer* DestIndexBuffer, FRHIBuffer* SrcIndexBuffer) final override;
-	virtual FVertexBufferRHIRef RHICreateVertexBuffer(uint32 Size, uint32 InUsage, FRHIResourceCreateInfo& CreateInfo) final override;
+	virtual FBufferRHIRef RHICreateVertexBuffer(uint32 Size, uint32 InUsage, FRHIResourceCreateInfo& CreateInfo) final override;
 	virtual void* LockVertexBuffer_BottomOfPipe(FRHICommandListImmediate& RHICmdList, FRHIBuffer* VertexBuffer, uint32 Offset, uint32 SizeRHI, EResourceLockMode LockMode) final override;
 	virtual void UnlockVertexBuffer_BottomOfPipe(FRHICommandListImmediate& RHICmdList, FRHIBuffer* VertexBuffer) final override;
 	virtual void RHICopyBuffer(FRHIBuffer* SourceBuffer, FRHIBuffer* DestBuffer) final override;
 	virtual void RHITransferBufferUnderlyingResource(FRHIBuffer* DestBuffer, FRHIBuffer* SrcBuffer) final override;
-	virtual FStructuredBufferRHIRef RHICreateStructuredBuffer(uint32 Stride, uint32 Size, uint32 InUsage, FRHIResourceCreateInfo& CreateInfo) final override;
+	virtual FBufferRHIRef RHICreateStructuredBuffer(uint32 Stride, uint32 Size, uint32 InUsage, FRHIResourceCreateInfo& CreateInfo) final override;
 	virtual void* LockStructuredBuffer_BottomOfPipe(FRHICommandListImmediate& RHICmdList, FRHIBuffer* StructuredBuffer, uint32 Offset, uint32 SizeRHI, EResourceLockMode LockMode) final override;
 	virtual void UnlockStructuredBuffer_BottomOfPipe(FRHICommandListImmediate& RHICmdList, FRHIBuffer* StructuredBuffer) final override;
 	virtual FUnorderedAccessViewRHIRef RHICreateUnorderedAccessView(FRHIBuffer* StructuredBuffer, bool bUseUAVCounter, bool bAppendBuffer) final override;
diff --git a/Engine/Source/Runtime/Engine/Classes/Engine/InstancedStaticMesh.h b/Engine/Source/Runtime/Engine/Classes/Engine/InstancedStaticMesh.h
--- a/Engine/Source/Runtime/Engine/Classes/Engine/InstancedStaticMesh.h
+++ b/Engine/Source/Runtime/Engine/Classes/Engine/InstancedStaticMesh.h
@@ -120,17 +120,17 @@ public:
 	/** Keep CPU copy of instance data*/
 	bool RequireCPUAccess;
 
-	FVertexBufferRHIRef GetInstanceOriginBuffer()
+	FBufferRHIRef GetInstanceOriginBuffer()
 	{
 		return InstanceOriginBuffer.VertexBufferRHI;
 	}
 
-	FVertexBufferRHIRef GetInstanceTransformBuffer()
+	FBufferRHIRef GetInstanceTransformBuffer()
 	{
 		return InstanceTransformBuffer.VertexBufferRHI;
 	}
 
-	FVertexBufferRHIRef GetInstanceLightmapBuffer()
+	FBufferRHIRef GetInstanceLightmapBuffer()
 	{
 		return InstanceLightmapBuffer.VertexBufferRHI;
 	}
@@ -163,7 +163,7 @@ private:
 	/** Delete existing resources */
 	void CleanUp();
 
-	void CreateVertexBuffer(FResourceArrayInterface* InResourceArray, uint32 InUsage, uint32 InStride, uint8 InFormat, FVertexBufferRHIRef& OutVertexBufferRHI, FShaderResourceViewRHIRef& OutInstanceSRV);
+	void CreateVertexBuffer(FResourceArrayInterface* InResourceArray, uint32 InUsage, uint32 InStride, uint8 InFormat, FBufferRHIRef& OutVertexBufferRHI, FShaderResourceViewRHIRef& OutInstanceSRV);
 	
 	/**  */
 	void UpdateFromCommandBuffer_RenderThread(FInstanceUpdateCmdBuffer& CmdBuffer);
diff --git a/Engine/Source/Runtime/Engine/Private/BatchedElements.cpp b/Engine/Source/Runtime/Engine/Private/BatchedElements.cpp
--- a/Engine/Source/Runtime/Engine/Private/BatchedElements.cpp
+++ b/Engine/Source/Runtime/Engine/Private/BatchedElements.cpp
@@ -720,7 +720,7 @@ void FBatchedElements::DrawPointElements(FRHICommandList& RHICmdList, const FMat
 		const int32 NumVertices = NumTris * 3;
 
 		FRHIResourceCreateInfo CreateInfo;
-		FVertexBufferRHIRef VertexBufferRHI = RHICreateVertexBuffer(sizeof(FSimpleElementVertex) * NumVertices, BUF_Volatile, CreateInfo);
+		FBufferRHIRef VertexBufferRHI = RHICreateVertexBuffer(sizeof(FSimpleElementVertex) * NumVertices, BUF_Volatile, CreateInfo);
 		void* VerticesPtr = RHILockBuffer(VertexBufferRHI, 0, sizeof(FSimpleElementVertex) * NumVertices, RLM_WriteOnly);
 
 		FSimpleElementVertex* PointVertices = (FSimpleElementVertex*)VerticesPtr;
@@ -827,7 +827,7 @@ bool FBatchedElements::Draw(FRHICommandList& RHICmdList, const FMeshPassProcesso
 				RHICmdList.SetStencilRef(StencilRef);
 
 				FRHIResourceCreateInfo CreateInfo;
-				FVertexBufferRHIRef VertexBufferRHI = RHICreateVertexBuffer(sizeof(FSimpleElementVertex) * LineVertices.Num(), BUF_Volatile, CreateInfo);
+				FBufferRHIRef VertexBufferRHI = RHICreateVertexBuffer(sizeof(FSimpleElementVertex) * LineVertices.Num(), BUF_Volatile, CreateInfo);
 				void* VoidPtr = RHILockBuffer(VertexBufferRHI, 0, sizeof(FSimpleElementVertex) * LineVertices.Num(), RLM_WriteOnly);
 
 				FMemory::Memcpy(VoidPtr, LineVertices.GetData(), sizeof(FSimpleElementVertex) * LineVertices.Num());
@@ -898,7 +898,7 @@ bool FBatchedElements::Draw(FRHICommandList& RHICmdList, const FMeshPassProcesso
 					RHICmdList.SetStencilRef(StencilRef);
 
 					FRHIResourceCreateInfo CreateInfo;
-					FVertexBufferRHIRef VertexBufferRHI = RHICreateVertexBuffer(sizeof(FSimpleElementVertex) * 8 * 3 * NumLinesThisBatch, BUF_Volatile, CreateInfo);
+					FBufferRHIRef VertexBufferRHI = RHICreateVertexBuffer(sizeof(FSimpleElementVertex) * 8 * 3 * NumLinesThisBatch, BUF_Volatile, CreateInfo);
 					void* ThickVertexData = RHILockBuffer(VertexBufferRHI, 0, sizeof(FSimpleElementVertex) * 8 * 3 * NumLinesThisBatch, RLM_WriteOnly);
 					FSimpleElementVertex* ThickVertices = (FSimpleElementVertex*)ThickVertexData;
 					check(ThickVertices);
@@ -984,7 +984,7 @@ bool FBatchedElements::Draw(FRHICommandList& RHICmdList, const FMeshPassProcesso
 				check(WireTriVerts.Num() == WireTris.Num() * 3);
 
 				FRHIResourceCreateInfo CreateInfo;
-				FVertexBufferRHIRef VertexBufferRHI = RHICreateVertexBuffer(sizeof(FSimpleElementVertex) * WireTriVerts.Num(), BUF_Volatile, CreateInfo);
+				FBufferRHIRef VertexBufferRHI = RHICreateVertexBuffer(sizeof(FSimpleElementVertex) * WireTriVerts.Num(), BUF_Volatile, CreateInfo);
 				void* VoidPtr = RHILockBuffer(VertexBufferRHI, 0, sizeof(FSimpleElementVertex) * WireTriVerts.Num(), RLM_WriteOnly);
 				FMemory::Memcpy(VoidPtr, WireTriVerts.GetData(), sizeof(FSimpleElementVertex) * WireTriVerts.Num());
 				RHIUnlockBuffer(VertexBufferRHI);
@@ -1072,7 +1072,7 @@ bool FBatchedElements::Draw(FRHICommandList& RHICmdList, const FMeshPassProcesso
 			if (ValidSpriteCount > 0)
 			{
 				FRHIResourceCreateInfo CreateInfo;
-				FVertexBufferRHIRef VertexBufferRHI = RHICreateVertexBuffer(sizeof(FSimpleElementVertex) * ValidSpriteCount * 6, BUF_Volatile, CreateInfo);
+				FBufferRHIRef VertexBufferRHI = RHICreateVertexBuffer(sizeof(FSimpleElementVertex) * ValidSpriteCount * 6, BUF_Volatile, CreateInfo);
 				void* VoidPtr = RHILockBuffer(VertexBufferRHI, 0, sizeof(FSimpleElementVertex) * ValidSpriteCount * 6, RLM_WriteOnly);
 				FSimpleElementVertex* SpriteList = reinterpret_cast<FSimpleElementVertex*>(VoidPtr);
 
@@ -1140,7 +1140,7 @@ bool FBatchedElements::Draw(FRHICommandList& RHICmdList, const FMeshPassProcesso
 		if( MeshElements.Num() > 0)
 		{
 			FRHIResourceCreateInfo CreateInfo;
-			FVertexBufferRHIRef VertexBufferRHI = RHICreateVertexBuffer(sizeof(FSimpleElementVertex) * MeshVertices.Num(), BUF_Volatile, CreateInfo);
+			FBufferRHIRef VertexBufferRHI = RHICreateVertexBuffer(sizeof(FSimpleElementVertex) * MeshVertices.Num(), BUF_Volatile, CreateInfo);
 			void* VoidPtr = RHILockBuffer(VertexBufferRHI, 0, sizeof(FSimpleElementVertex) * MeshVertices.Num(), RLM_WriteOnly);
 			FPlatformMemory::Memcpy(VoidPtr, MeshVertices.GetData(), sizeof(FSimpleElementVertex) * MeshVertices.Num());
 			RHIUnlockBuffer(VertexBufferRHI);
@@ -1166,7 +1166,7 @@ bool FBatchedElements::Draw(FRHICommandList& RHICmdList, const FMeshPassProcesso
 					PrepareShaders(RHICmdList, GraphicsPSOInit, FeatureLevel, MeshElement.BlendMode, Transform, bNeedToSwitchVerticalAxis, MeshElement.BatchedElementParameters, MeshElement.Texture, bHitTesting, Gamma, &MeshElement.GlowInfo, &View);
 					RHICmdList.SetStencilRef(StencilRef);
 
-					FIndexBufferRHIRef IndexBufferRHI = RHICreateIndexBuffer(sizeof(uint16), sizeof(uint16) * MeshElement.Indices.Num(), BUF_Volatile, CreateInfo);
+					FBufferRHIRef IndexBufferRHI = RHICreateIndexBuffer(sizeof(uint16), sizeof(uint16) * MeshElement.Indices.Num(), BUF_Volatile, CreateInfo);
 					void* VoidPtr2 = RHILockBuffer(IndexBufferRHI, 0, sizeof(uint16) * MeshElement.Indices.Num(), RLM_WriteOnly);
 					FPlatformMemory::Memcpy(VoidPtr2, MeshElement.Indices.GetData(), sizeof(uint16) * MeshElement.Indices.Num());
 					RHIUnlockBuffer(IndexBufferRHI);
diff --git a/Engine/Source/Runtime/Engine/Private/DynamicMeshBuilder.cpp b/Engine/Source/Runtime/Engine/Private/DynamicMeshBuilder.cpp
--- a/Engine/Source/Runtime/Engine/Private/DynamicMeshBuilder.cpp
+++ b/Engine/Source/Runtime/Engine/Private/DynamicMeshBuilder.cpp
@@ -98,13 +98,13 @@ public:
 	 * @param Args The buffer size in bytes.
 	 * @returns A suitably sized buffer or NULL on failure.
 	 */
-	FIndexBufferRHIRef CreateResource(FGlobalDynamicMeshPoolPolicy::CreationArguments Args)
+	FBufferRHIRef CreateResource(FGlobalDynamicMeshPoolPolicy::CreationArguments Args)
 	{
 		FGlobalDynamicMeshPoolPolicy::CreationArguments BufferSize = GetPoolBucketSize(GetPoolBucketIndex(Args));
 		// The use of BUF_Static is deliberate - on OS X the buffer backing-store orphaning & reallocation will dominate execution time
 		// so to avoid this we don't reuse a buffer for several frames, thereby avoiding the pipeline stall and the reallocation cost.
 		FRHIResourceCreateInfo CreateInfo;
-		FIndexBufferRHIRef VertexBuffer = RHICreateIndexBuffer(sizeof(DynamicMeshIndexType), BufferSize, BUF_Static, CreateInfo);
+		FBufferRHIRef VertexBuffer = RHICreateIndexBuffer(sizeof(DynamicMeshIndexType), BufferSize, BUF_Static, CreateInfo);
 		return VertexBuffer;
 	}
 	
@@ -112,7 +112,7 @@ public:
 	 * @param Resource The buffer to get data for.
 	 * @returns The arguments used to create the buffer.
 	 */
-	FGlobalDynamicMeshPoolPolicy::CreationArguments GetCreationArguments(FIndexBufferRHIRef Resource)
+	FGlobalDynamicMeshPoolPolicy::CreationArguments GetCreationArguments(FBufferRHIRef Resource)
 	{
 		return (Resource->GetSize());
 	}
@@ -120,12 +120,12 @@ public:
 	/** Frees the resource
 	 * @param Resource The buffer to prepare for release from the pool permanently.
 	 */
-	void FreeResource(FIndexBufferRHIRef Resource)
+	void FreeResource(FBufferRHIRef Resource)
 	{
 	}
 };
 
-class FGlobalDynamicMeshIndexPool : public TRenderResourcePool<FIndexBufferRHIRef, FGlobalDynamicMeshIndexPolicy, FGlobalDynamicMeshPoolPolicy::CreationArguments>
+class FGlobalDynamicMeshIndexPool : public TRenderResourcePool<FBufferRHIRef, FGlobalDynamicMeshIndexPolicy, FGlobalDynamicMeshPoolPolicy::CreationArguments>
 {
 public:
 	/** Destructor */
@@ -156,11 +156,11 @@ public:
 	 * @param Args The buffer size in bytes.
 	 * @returns A suitably sized buffer or NULL on failure.
 	 */
-	FVertexBufferRHIRef CreateResource(FGlobalDynamicMeshPoolPolicy::CreationArguments Args)
+	FBufferRHIRef CreateResource(FGlobalDynamicMeshPoolPolicy::CreationArguments Args)
 	{
 		FGlobalDynamicMeshPoolPolicy::CreationArguments BufferSize = GetPoolBucketSize(GetPoolBucketIndex(Args));
 		FRHIResourceCreateInfo CreateInfo;
-		FVertexBufferRHIRef VertexBuffer = RHICreateVertexBuffer(BufferSize, BUF_Volatile | BUF_ShaderResource, CreateInfo);
+		FBufferRHIRef VertexBuffer = RHICreateVertexBuffer(BufferSize, BUF_Volatile | BUF_ShaderResource, CreateInfo);
 		return VertexBuffer;
 	}
 	
@@ -168,7 +168,7 @@ public:
 	 * @param Resource The buffer to get data for.
 	 * @returns The arguments used to create the buffer.
 	 */
-	FGlobalDynamicMeshPoolPolicy::CreationArguments GetCreationArguments(FVertexBufferRHIRef Resource)
+	FGlobalDynamicMeshPoolPolicy::CreationArguments GetCreationArguments(FBufferRHIRef Resource)
 	{
 		return (Resource->GetSize());
 	}
@@ -176,12 +176,12 @@ public:
 	/** Frees the resource
 	 * @param Resource The buffer to prepare for release from the pool permanently.
 	 */
-	void FreeResource(FVertexBufferRHIRef Resource)
+	void FreeResource(FBufferRHIRef Resource)
 	{
 	}
 };
 
-class FGlobalDynamicMeshVertexPool : public TRenderResourcePool<FVertexBufferRHIRef, FGlobalDynamicMeshVertexPolicy, FGlobalDynamicMeshPoolPolicy::CreationArguments>
+class FGlobalDynamicMeshVertexPool : public TRenderResourcePool<FBufferRHIRef, FGlobalDynamicMeshVertexPolicy, FGlobalDynamicMeshPoolPolicy::CreationArguments>
 {
 public:
 	/** Destructor */
@@ -244,7 +244,7 @@ int32 FDynamicMeshBufferAllocator::GetVertexBufferSize(uint32 Stride, uint32 Num
 	return NumElements * Stride;
 }
 
-FIndexBufferRHIRef FDynamicMeshBufferAllocator::AllocIndexBuffer(uint32 NumElements)
+FBufferRHIRef FDynamicMeshBufferAllocator::AllocIndexBuffer(uint32 NumElements)
 {
 	uint32 SizeInBytes = GetIndexBufferSize(NumElements);
 
@@ -252,12 +252,12 @@ FIndexBufferRHIRef FDynamicMeshBufferAllocator::AllocIndexBuffer(uint32 NumEleme
 	return RHICreateIndexBuffer(sizeof(DynamicMeshIndexType), SizeInBytes, BUF_Volatile, CreateInfo);
 }
 
-void FDynamicMeshBufferAllocator::ReleaseIndexBuffer(FIndexBufferRHIRef& IndexBufferRHI)
+void FDynamicMeshBufferAllocator::ReleaseIndexBuffer(FBufferRHIRef& IndexBufferRHI)
 {
 	IndexBufferRHI = nullptr;
 }
 
-FVertexBufferRHIRef FDynamicMeshBufferAllocator::AllocVertexBuffer(uint32 Stride, uint32 NumElements)
+FBufferRHIRef FDynamicMeshBufferAllocator::AllocVertexBuffer(uint32 Stride, uint32 NumElements)
 {
 	uint32 SizeInBytes = GetVertexBufferSize(Stride, NumElements);
 
@@ -265,7 +265,7 @@ FVertexBufferRHIRef FDynamicMeshBufferAllocator::AllocVertexBuffer(uint32 Stride
 	return RHICreateVertexBuffer(SizeInBytes, BUF_Volatile | BUF_ShaderResource, CreateInfo);
 }
 
-void FDynamicMeshBufferAllocator::ReleaseVertexBuffer(FVertexBufferRHIRef& VertexBufferRHI)
+void FDynamicMeshBufferAllocator::ReleaseVertexBuffer(FBufferRHIRef& VertexBufferRHI)
 {
 	VertexBufferRHI = nullptr;
 }
@@ -273,7 +273,7 @@ void FDynamicMeshBufferAllocator::ReleaseVertexBuffer(FVertexBufferRHIRef& Verte
 /** This is our default implementation using GDynamicMeshIndexPool. */
 class FPooledDynamicMeshBufferAllocator : public FDynamicMeshBufferAllocator
 {
-	virtual FIndexBufferRHIRef AllocIndexBuffer(uint32 NumElements) override
+	virtual FBufferRHIRef AllocIndexBuffer(uint32 NumElements) override
 	{
 		uint32 SizeInBytes = NumElements * sizeof(DynamicMeshIndexType);
 		if (SizeInBytes <= FGlobalDynamicMeshIndexPolicy().GetPoolBucketSize(FGlobalDynamicMeshIndexPolicy::NumPoolBuckets - 1))
@@ -284,7 +284,7 @@ class FPooledDynamicMeshBufferAllocator : public FDynamicMeshBufferAllocator
 		return FDynamicMeshBufferAllocator::AllocIndexBuffer(NumElements);
 	}
 
-	virtual void ReleaseIndexBuffer(FIndexBufferRHIRef& IndexBufferRHI)
+	virtual void ReleaseIndexBuffer(FBufferRHIRef& IndexBufferRHI)
 	{
 		if (IsValidRef(IndexBufferRHI))
 		{
@@ -297,7 +297,7 @@ class FPooledDynamicMeshBufferAllocator : public FDynamicMeshBufferAllocator
 		}
 	}
 
-	virtual FVertexBufferRHIRef AllocVertexBuffer(uint32 Stride, uint32 NumElements)
+	virtual FBufferRHIRef AllocVertexBuffer(uint32 Stride, uint32 NumElements)
 	{
 		uint32 SizeInBytes = NumElements * Stride;
 		if (SizeInBytes <= FGlobalDynamicMeshIndexPolicy().GetPoolBucketSize(FGlobalDynamicMeshIndexPolicy::NumPoolBuckets - 1))
@@ -308,7 +308,7 @@ class FPooledDynamicMeshBufferAllocator : public FDynamicMeshBufferAllocator
 		return FDynamicMeshBufferAllocator::AllocVertexBuffer(Stride, NumElements);
 	}
 
-	virtual void ReleaseVertexBuffer(FVertexBufferRHIRef& VertexBufferRHI)
+	virtual void ReleaseVertexBuffer(FBufferRHIRef& VertexBufferRHI)
 	{
 		if (IsValidRef(VertexBufferRHI))
 		{
diff --git a/Engine/Source/Runtime/Engine/Private/GPUSkinCache.cpp b/Engine/Source/Runtime/Engine/Private/GPUSkinCache.cpp
--- a/Engine/Source/Runtime/Engine/Private/GPUSkinCache.cpp
+++ b/Engine/Source/Runtime/Engine/Private/GPUSkinCache.cpp
@@ -1486,7 +1486,7 @@ void FGPUSkinCache::ProcessRayTracingGeometryToUpdate(
 		{
 			uint32 MemoryEstimation = 0;
 
-			FIndexBufferRHIRef IndexBufferRHI = LODModel.MultiSizeIndexContainer.GetIndexBuffer()->IndexBufferRHI;
+			FBufferRHIRef IndexBufferRHI = LODModel.MultiSizeIndexContainer.GetIndexBuffer()->IndexBufferRHI;
 			MemoryEstimation += IndexBufferRHI->GetSize();
 			uint32 VertexBufferStride = LODModel.StaticVertexBuffers.PositionVertexBuffer.GetStride();
 			MemoryEstimation += LODModel.StaticVertexBuffers.PositionVertexBuffer.VertexBufferRHI->GetSize();
diff --git a/Engine/Source/Runtime/Engine/Private/GPUSort.cpp b/Engine/Source/Runtime/Engine/Private/GPUSort.cpp
--- a/Engine/Source/Runtime/Engine/Private/GPUSort.cpp
+++ b/Engine/Source/Runtime/Engine/Private/GPUSort.cpp
@@ -75,7 +75,7 @@ class FSortOffsetBuffers : public FRenderResource
 public:
 
 	/** Vertex buffer storage for the actual offsets. */
-	FVertexBufferRHIRef Buffers[2];
+	FBufferRHIRef Buffers[2];
 	/** Shader resource views for offset buffers. */
 	FShaderResourceViewRHIRef BufferSRVs[2];
 	/** Unordered access views for offset buffers. */
@@ -178,7 +178,7 @@ class FRadixSortParametersBuffer : public FRenderResource
 public:
 
 	/** The vertex buffer used for storage. */
-	FVertexBufferRHIRef SortParametersBufferRHI;
+	FBufferRHIRef SortParametersBufferRHI;
 	/** Shader resource view in to the vertex buffer. */
 	FShaderResourceViewRHIRef SortParametersBufferSRV;
 	
@@ -868,7 +868,7 @@ static bool RunGPUSortTest(FRHICommandListImmediate& RHICmdList, int32 TestSize,
 	TArray<uint32> RefSortedKeys;
 	TArray<uint32> SortedKeys;
 	TArray<uint32> SortedValues;
-	FVertexBufferRHIRef KeysBufferRHI[2], ValuesBufferRHI[2];
+	FBufferRHIRef KeysBufferRHI[2], ValuesBufferRHI[2];
 	FShaderResourceViewRHIRef KeysBufferSRV[2], ValuesBufferSRV[2];
 	FUnorderedAccessViewRHIRef KeysBufferUAV[2], ValuesBufferUAV[2];
 	int32 ResultBufferIndex;
diff --git a/Engine/Source/Runtime/Engine/Private/InstancedStaticMesh.cpp b/Engine/Source/Runtime/Engine/Private/InstancedStaticMesh.cpp
--- a/Engine/Source/Runtime/Engine/Private/InstancedStaticMesh.cpp
+++ b/Engine/Source/Runtime/Engine/Private/InstancedStaticMesh.cpp
@@ -450,7 +450,7 @@ SIZE_T FStaticMeshInstanceBuffer::GetResourceSize() const
 	return 0;
 }
 
-void FStaticMeshInstanceBuffer::CreateVertexBuffer(FResourceArrayInterface* InResourceArray, uint32 InUsage, uint32 InStride, uint8 InFormat, FVertexBufferRHIRef& OutVertexBufferRHI, FShaderResourceViewRHIRef& OutInstanceSRV)
+void FStaticMeshInstanceBuffer::CreateVertexBuffer(FResourceArrayInterface* InResourceArray, uint32 InUsage, uint32 InStride, uint8 InFormat, FBufferRHIRef& OutVertexBufferRHI, FShaderResourceViewRHIRef& OutInstanceSRV)
 {
 	check(InResourceArray);
 	check(InResourceArray->GetResourceDataSize() > 0);
diff --git a/Engine/Source/Runtime/Engine/Private/MeshParticleVertexFactory.cpp b/Engine/Source/Runtime/Engine/Private/MeshParticleVertexFactory.cpp
--- a/Engine/Source/Runtime/Engine/Private/MeshParticleVertexFactory.cpp
+++ b/Engine/Source/Runtime/Engine/Private/MeshParticleVertexFactory.cpp
@@ -79,7 +79,7 @@ public:
 	}
 
 private:
-	FVertexBufferRHIRef VB;
+	FBufferRHIRef VB;
 	FShaderResourceViewRHIRef SRV;
 };
 
diff --git a/Engine/Source/Runtime/Engine/Private/Particles/ParticleCurveTexture.cpp b/Engine/Source/Runtime/Engine/Private/Particles/ParticleCurveTexture.cpp
--- a/Engine/Source/Runtime/Engine/Private/Particles/ParticleCurveTexture.cpp
+++ b/Engine/Source/Runtime/Engine/Private/Particles/ParticleCurveTexture.cpp
@@ -229,7 +229,7 @@ static void InjectCurves(
 		//
 		FRHIResourceCreateInfo CreateInfo;
 		void* ScratchData = nullptr;
-		FVertexBufferRHIRef ScratchVertexBufferRHI = RHICreateAndLockVertexBuffer(TotalSamples * sizeof(FColor), BUF_Volatile , CreateInfo, ScratchData);
+		FBufferRHIRef ScratchVertexBufferRHI = RHICreateAndLockVertexBuffer(TotalSamples * sizeof(FColor), BUF_Volatile , CreateInfo, ScratchData);
 		FColor* RESTRICT DestSamples = (FColor*)ScratchData;
 
 		int32 CurrOffset = 0;
diff --git a/Engine/Source/Runtime/Engine/Private/Particles/ParticleGpuSimulation.cpp b/Engine/Source/Runtime/Engine/Private/Particles/ParticleGpuSimulation.cpp
--- a/Engine/Source/Runtime/Engine/Private/Particles/ParticleGpuSimulation.cpp
+++ b/Engine/Source/Runtime/Engine/Private/Particles/ParticleGpuSimulation.cpp
@@ -2238,7 +2238,7 @@ static FBox ComputeParticleBounds(
 		// Create a buffer for storing bounds.
 		const int32 BufferSize = GroupCount * 2 * sizeof(FVector4);
 		FRHIResourceCreateInfo CreateInfo;
-		FVertexBufferRHIRef BoundsVertexBufferRHI = RHICreateVertexBuffer(
+		FBufferRHIRef BoundsVertexBufferRHI = RHICreateVertexBuffer(
 			BufferSize,
 			BUF_Static | BUF_UnorderedAccess | BUF_KeepCPUAccessible,
 			CreateInfo);
diff --git a/Engine/Source/Runtime/Engine/Private/PositionVertexBuffer.cpp b/Engine/Source/Runtime/Engine/Private/PositionVertexBuffer.cpp
--- a/Engine/Source/Runtime/Engine/Private/PositionVertexBuffer.cpp
+++ b/Engine/Source/Runtime/Engine/Private/PositionVertexBuffer.cpp
@@ -189,7 +189,7 @@ void FPositionVertexBuffer::operator=(const FPositionVertexBuffer &Other)
 }
 
 template <bool bRenderThread>
-FVertexBufferRHIRef FPositionVertexBuffer::CreateRHIBuffer_Internal()
+FBufferRHIRef FPositionVertexBuffer::CreateRHIBuffer_Internal()
 {
 	if (NumVertices)
 	{
@@ -209,12 +209,12 @@ FVertexBufferRHIRef FPositionVertexBuffer::CreateRHIBuffer_Internal()
 	return nullptr;
 }
 
-FVertexBufferRHIRef FPositionVertexBuffer::CreateRHIBuffer_RenderThread()
+FBufferRHIRef FPositionVertexBuffer::CreateRHIBuffer_RenderThread()
 {
 	return CreateRHIBuffer_Internal<true>();
 }
 
-FVertexBufferRHIRef FPositionVertexBuffer::CreateRHIBuffer_Async()
+FBufferRHIRef FPositionVertexBuffer::CreateRHIBuffer_Async()
 {
 	return CreateRHIBuffer_Internal<false>();
 }
diff --git a/Engine/Source/Runtime/Engine/Private/RawIndexBuffer.cpp b/Engine/Source/Runtime/Engine/Private/RawIndexBuffer.cpp
--- a/Engine/Source/Runtime/Engine/Private/RawIndexBuffer.cpp
+++ b/Engine/Source/Runtime/Engine/Private/RawIndexBuffer.cpp
@@ -315,7 +315,7 @@ FIndexArrayView FRawStaticIndexBuffer::GetArrayView() const
 }
 
 template <bool bRenderThread>
-FIndexBufferRHIRef FRawStaticIndexBuffer::CreateRHIBuffer_Internal()
+FBufferRHIRef FRawStaticIndexBuffer::CreateRHIBuffer_Internal()
 {
 	const uint32 IndexStride = b32Bit ? sizeof(uint32) : sizeof(uint16);
 	const uint32 SizeInBytes = IndexStorage.Num();
@@ -343,12 +343,12 @@ FIndexBufferRHIRef FRawStaticIndexBuffer::CreateRHIBuffer_Internal()
 	return nullptr;
 }
 
-FIndexBufferRHIRef FRawStaticIndexBuffer::CreateRHIBuffer_RenderThread()
+FBufferRHIRef FRawStaticIndexBuffer::CreateRHIBuffer_RenderThread()
 {
 	return CreateRHIBuffer_Internal<true>();
 }
 
-FIndexBufferRHIRef FRawStaticIndexBuffer::CreateRHIBuffer_Async()
+FBufferRHIRef FRawStaticIndexBuffer::CreateRHIBuffer_Async()
 {
 	return CreateRHIBuffer_Internal<false>();
 }
diff --git a/Engine/Source/Runtime/Engine/Private/Rendering/ColorVertexBuffer.cpp b/Engine/Source/Runtime/Engine/Private/Rendering/ColorVertexBuffer.cpp
--- a/Engine/Source/Runtime/Engine/Private/Rendering/ColorVertexBuffer.cpp
+++ b/Engine/Source/Runtime/Engine/Private/Rendering/ColorVertexBuffer.cpp
@@ -387,7 +387,7 @@ uint32 FColorVertexBuffer::GetAllocatedSize() const
 }
 
 template <bool bRenderThread>
-FVertexBufferRHIRef FColorVertexBuffer::CreateRHIBuffer_Internal()
+FBufferRHIRef FColorVertexBuffer::CreateRHIBuffer_Internal()
 {
 	if (NumVertices)
 	{
@@ -407,12 +407,12 @@ FVertexBufferRHIRef FColorVertexBuffer::CreateRHIBuffer_Internal()
 	return nullptr;
 }
 
-FVertexBufferRHIRef FColorVertexBuffer::CreateRHIBuffer_RenderThread()
+FBufferRHIRef FColorVertexBuffer::CreateRHIBuffer_RenderThread()
 {
 	return CreateRHIBuffer_Internal<true>();
 }
 
-FVertexBufferRHIRef FColorVertexBuffer::CreateRHIBuffer_Async()
+FBufferRHIRef FColorVertexBuffer::CreateRHIBuffer_Async()
 {
 	return CreateRHIBuffer_Internal<false>();
 }
diff --git a/Engine/Source/Runtime/Engine/Private/Rendering/MultiSizeIndexContainer.cpp b/Engine/Source/Runtime/Engine/Private/Rendering/MultiSizeIndexContainer.cpp
--- a/Engine/Source/Runtime/Engine/Private/Rendering/MultiSizeIndexContainer.cpp
+++ b/Engine/Source/Runtime/Engine/Private/Rendering/MultiSizeIndexContainer.cpp
@@ -182,7 +182,7 @@ void FMultiSizeIndexContainer::SerializeMetaData(FArchive& Ar, bool bNeedsCPUAcc
 	IndexBuffer->SerializeMetaData(Ar);
 }
 
-FIndexBufferRHIRef FMultiSizeIndexContainer::CreateRHIBuffer_RenderThread()
+FBufferRHIRef FMultiSizeIndexContainer::CreateRHIBuffer_RenderThread()
 {
 	if (IndexBuffer)
 	{
@@ -198,7 +198,7 @@ FIndexBufferRHIRef FMultiSizeIndexContainer::CreateRHIBuffer_RenderThread()
 	return nullptr;
 }
 
-FIndexBufferRHIRef FMultiSizeIndexContainer::CreateRHIBuffer_Async()
+FBufferRHIRef FMultiSizeIndexContainer::CreateRHIBuffer_Async()
 {
 	if (IndexBuffer)
 	{
diff --git a/Engine/Source/Runtime/Engine/Private/Rendering/SkeletalMeshVertexClothBuffer.cpp b/Engine/Source/Runtime/Engine/Private/Rendering/SkeletalMeshVertexClothBuffer.cpp
--- a/Engine/Source/Runtime/Engine/Private/Rendering/SkeletalMeshVertexClothBuffer.cpp
+++ b/Engine/Source/Runtime/Engine/Private/Rendering/SkeletalMeshVertexClothBuffer.cpp
@@ -73,7 +73,7 @@ void FSkeletalMeshVertexClothBuffer::ClearMetaData()
 }
 
 template <bool bRenderThread>
-FVertexBufferRHIRef FSkeletalMeshVertexClothBuffer::CreateRHIBuffer_Internal()
+FBufferRHIRef FSkeletalMeshVertexClothBuffer::CreateRHIBuffer_Internal()
 {
 	if (NumVertices)
 	{
@@ -95,12 +95,12 @@ FVertexBufferRHIRef FSkeletalMeshVertexClothBuffer::CreateRHIBuffer_Internal()
 	return nullptr;
 }
 
-FVertexBufferRHIRef FSkeletalMeshVertexClothBuffer::CreateRHIBuffer_RenderThread()
+FBufferRHIRef FSkeletalMeshVertexClothBuffer::CreateRHIBuffer_RenderThread()
 {
 	return CreateRHIBuffer_Internal<true>();
 }
 
-FVertexBufferRHIRef FSkeletalMeshVertexClothBuffer::CreateRHIBuffer_Async()
+FBufferRHIRef FSkeletalMeshVertexClothBuffer::CreateRHIBuffer_Async()
 {
 	return CreateRHIBuffer_Internal<false>();
 }
diff --git a/Engine/Source/Runtime/Engine/Private/Rendering/SkinWeightVertexBuffer.cpp b/Engine/Source/Runtime/Engine/Private/Rendering/SkinWeightVertexBuffer.cpp
--- a/Engine/Source/Runtime/Engine/Private/Rendering/SkinWeightVertexBuffer.cpp
+++ b/Engine/Source/Runtime/Engine/Private/Rendering/SkinWeightVertexBuffer.cpp
@@ -151,12 +151,12 @@ void FSkinWeightLookupVertexBuffer::ReleaseRHI()
 	FVertexBuffer::ReleaseRHI();
 }
 
-FVertexBufferRHIRef FSkinWeightLookupVertexBuffer::CreateRHIBuffer_RenderThread()
+FBufferRHIRef FSkinWeightLookupVertexBuffer::CreateRHIBuffer_RenderThread()
 {
 	return CreateRHIBuffer_Internal<true>();
 }
 
-FVertexBufferRHIRef FSkinWeightLookupVertexBuffer::CreateRHIBuffer_Async()
+FBufferRHIRef FSkinWeightLookupVertexBuffer::CreateRHIBuffer_Async()
 {
 	return CreateRHIBuffer_Internal<false>();
 }
@@ -177,7 +177,7 @@ void FSkinWeightLookupVertexBuffer::SetWeightOffsetAndInfluenceCount(uint32 Vert
 }
 
 template <bool bRenderThread>
-FVertexBufferRHIRef FSkinWeightLookupVertexBuffer::CreateRHIBuffer_Internal()
+FBufferRHIRef FSkinWeightLookupVertexBuffer::CreateRHIBuffer_Internal()
 {
 	if (NumVertices)
 	{
@@ -387,7 +387,7 @@ void FSkinWeightDataVertexBuffer::CopyMetaData(const FSkinWeightDataVertexBuffer
 }
 
 template <bool bRenderThread>
-FVertexBufferRHIRef FSkinWeightDataVertexBuffer::CreateRHIBuffer_Internal()
+FBufferRHIRef FSkinWeightDataVertexBuffer::CreateRHIBuffer_Internal()
 {
 	if (NumBones)
 	{
@@ -411,12 +411,12 @@ FVertexBufferRHIRef FSkinWeightDataVertexBuffer::CreateRHIBuffer_Internal()
 	return nullptr;
 }
 
-FVertexBufferRHIRef FSkinWeightDataVertexBuffer::CreateRHIBuffer_RenderThread()
+FBufferRHIRef FSkinWeightDataVertexBuffer::CreateRHIBuffer_RenderThread()
 {
 	return CreateRHIBuffer_Internal<true>();
 }
 
-FVertexBufferRHIRef FSkinWeightDataVertexBuffer::CreateRHIBuffer_Async()
+FBufferRHIRef FSkinWeightDataVertexBuffer::CreateRHIBuffer_Async()
 {
 	return CreateRHIBuffer_Internal<false>();
 }
diff --git a/Engine/Source/Runtime/Engine/Private/Rendering/StaticMeshVertexBuffer.cpp b/Engine/Source/Runtime/Engine/Private/Rendering/StaticMeshVertexBuffer.cpp
--- a/Engine/Source/Runtime/Engine/Private/Rendering/StaticMeshVertexBuffer.cpp
+++ b/Engine/Source/Runtime/Engine/Private/Rendering/StaticMeshVertexBuffer.cpp
@@ -267,7 +267,7 @@ void FStaticMeshVertexBuffer::operator=(const FStaticMeshVertexBuffer &Other)
 }
 
 template <bool bRenderThread>
-FVertexBufferRHIRef FStaticMeshVertexBuffer::CreateTangentsRHIBuffer_Internal()
+FBufferRHIRef FStaticMeshVertexBuffer::CreateTangentsRHIBuffer_Internal()
 {
 	if (GetNumVertices())
 	{
@@ -287,18 +287,18 @@ FVertexBufferRHIRef FStaticMeshVertexBuffer::CreateTangentsRHIBuffer_Internal()
 	return nullptr;
 }
 
-FVertexBufferRHIRef FStaticMeshVertexBuffer::CreateTangentsRHIBuffer_RenderThread()
+FBufferRHIRef FStaticMeshVertexBuffer::CreateTangentsRHIBuffer_RenderThread()
 {
 	return CreateTangentsRHIBuffer_Internal<true>();
 }
 
-FVertexBufferRHIRef FStaticMeshVertexBuffer::CreateTangentsRHIBuffer_Async()
+FBufferRHIRef FStaticMeshVertexBuffer::CreateTangentsRHIBuffer_Async()
 {
 	return CreateTangentsRHIBuffer_Internal<false>();
 }
 
 template <bool bRenderThread>
-FVertexBufferRHIRef FStaticMeshVertexBuffer::CreateTexCoordRHIBuffer_Internal()
+FBufferRHIRef FStaticMeshVertexBuffer::CreateTexCoordRHIBuffer_Internal()
 {
 	if (GetNumTexCoords())
 	{
@@ -318,12 +318,12 @@ FVertexBufferRHIRef FStaticMeshVertexBuffer::CreateTexCoordRHIBuffer_Internal()
 	return nullptr;
 }
 
-FVertexBufferRHIRef FStaticMeshVertexBuffer::CreateTexCoordRHIBuffer_RenderThread()
+FBufferRHIRef FStaticMeshVertexBuffer::CreateTexCoordRHIBuffer_RenderThread()
 {
 	return CreateTexCoordRHIBuffer_Internal<true>();
 }
 
-FVertexBufferRHIRef FStaticMeshVertexBuffer::CreateTexCoordRHIBuffer_Async()
+FBufferRHIRef FStaticMeshVertexBuffer::CreateTexCoordRHIBuffer_Async()
 {
 	return CreateTexCoordRHIBuffer_Internal<false>();
 }
diff --git a/Engine/Source/Runtime/Engine/Private/SkeletalRenderCPUSkin.cpp b/Engine/Source/Runtime/Engine/Private/SkeletalRenderCPUSkin.cpp
--- a/Engine/Source/Runtime/Engine/Private/SkeletalRenderCPUSkin.cpp
+++ b/Engine/Source/Runtime/Engine/Private/SkeletalRenderCPUSkin.cpp
@@ -371,8 +371,8 @@ void FSkeletalMeshObjectCPUSkin::FSkeletalMeshObjectLOD::InitResources(FSkelMesh
 		check(SkelMeshRenderData->LODRenderData.IsValidIndex(LODIndex));
 		FSkeletalMeshLODRenderData& LODModel = SkelMeshRenderData->LODRenderData[LODIndex];
 		FResourceArrayInterface* OfflineData = LODModel.RayTracingData.Num() ? &LODModel.RayTracingData : nullptr;
-		FVertexBufferRHIRef VertexBufferRHI = LODModel.StaticVertexBuffers.PositionVertexBuffer.VertexBufferRHI;
-		FIndexBufferRHIRef IndexBufferRHI = LODModel.MultiSizeIndexContainer.GetIndexBuffer()->IndexBufferRHI;
+		FBufferRHIRef VertexBufferRHI = LODModel.StaticVertexBuffers.PositionVertexBuffer.VertexBufferRHI;
+		FBufferRHIRef IndexBufferRHI = LODModel.MultiSizeIndexContainer.GetIndexBuffer()->IndexBufferRHI;
 		uint32 VertexBufferStride = LODModel.StaticVertexBuffers.PositionVertexBuffer.GetStride();
 
 		uint32 TrianglesCount = 0;
diff --git a/Engine/Source/Runtime/Engine/Private/SkeletalRenderGPUSkin.h b/Engine/Source/Runtime/Engine/Private/SkeletalRenderGPUSkin.h
--- a/Engine/Source/Runtime/Engine/Private/SkeletalRenderGPUSkin.h
+++ b/Engine/Source/Runtime/Engine/Private/SkeletalRenderGPUSkin.h
@@ -267,7 +267,7 @@ private:
 	FSkeletalMeshRenderData* SkelMeshRenderData;
 
 	//temporary buffer used for normalization
-	FVertexBufferRHIRef NormalizationBufferRHI;
+	FBufferRHIRef NormalizationBufferRHI;
 
 	// guaranteed only to be valid if the Normalization Buffer is valid
 	FUnorderedAccessViewRHIRef NormalizationBufferUAV;
diff --git a/Engine/Source/Runtime/Engine/Private/SkeletalRenderStatic.cpp b/Engine/Source/Runtime/Engine/Private/SkeletalRenderStatic.cpp
--- a/Engine/Source/Runtime/Engine/Private/SkeletalRenderStatic.cpp
+++ b/Engine/Source/Runtime/Engine/Private/SkeletalRenderStatic.cpp
@@ -53,8 +53,8 @@ void FSkeletalMeshObjectStatic::InitResources(USkinnedMeshComponent* InMeshCompo
 					check(SkelLOD.SkelMeshRenderData->LODRenderData.IsValidIndex(LODIndex));
 
 					FSkeletalMeshLODRenderData& LODModel = SkelLOD.SkelMeshRenderData->LODRenderData[LODIndex];
-					FVertexBufferRHIRef VertexBufferRHI = LODModel.StaticVertexBuffers.PositionVertexBuffer.VertexBufferRHI;
-					FIndexBufferRHIRef IndexBufferRHI = LODModel.MultiSizeIndexContainer.GetIndexBuffer()->IndexBufferRHI;
+					FBufferRHIRef VertexBufferRHI = LODModel.StaticVertexBuffers.PositionVertexBuffer.VertexBufferRHI;
+					FBufferRHIRef IndexBufferRHI = LODModel.MultiSizeIndexContainer.GetIndexBuffer()->IndexBufferRHI;
 					uint32 VertexBufferStride = LODModel.StaticVertexBuffers.PositionVertexBuffer.GetStride();
 
 					uint32 TrianglesCount = 0;
diff --git a/Engine/Source/Runtime/Engine/Private/Streaming/LandscapeMeshMobileUpdate.cpp b/Engine/Source/Runtime/Engine/Private/Streaming/LandscapeMeshMobileUpdate.cpp
--- a/Engine/Source/Runtime/Engine/Private/Streaming/LandscapeMeshMobileUpdate.cpp
+++ b/Engine/Source/Runtime/Engine/Private/Streaming/LandscapeMeshMobileUpdate.cpp
@@ -165,7 +165,7 @@ void FLandscapeMeshMobileStreamOut::ShrinkResources(const FContext& Context)
 			FRHIResourceCreateInfo CreateInfo;
 			FRHIBuffer* LandsacpeVBRHI = RenderData->VertexBuffer->VertexBufferRHI;
 			const int32 NewSize = (int32)LandsacpeVBRHI->GetSize() + SizeDelta;
-			FVertexBufferRHIRef IntermediateVertexBuffer = RHICreateVertexBuffer(NewSize, BUF_Static, CreateInfo);
+			FBufferRHIRef IntermediateVertexBuffer = RHICreateVertexBuffer(NewSize, BUF_Static, CreateInfo);
 			FRHICommandListExecutor::GetImmediateCommandList().CopyBufferRegion(IntermediateVertexBuffer, 0, LandsacpeVBRHI, 0, NewSize);
 
 			TRHIResourceUpdateBatcher<1> Batcher;
diff --git a/Engine/Source/Runtime/Engine/Private/Streaming/LandscapeMeshMobileUpdate.h b/Engine/Source/Runtime/Engine/Private/Streaming/LandscapeMeshMobileUpdate.h
--- a/Engine/Source/Runtime/Engine/Private/Streaming/LandscapeMeshMobileUpdate.h
+++ b/Engine/Source/Runtime/Engine/Private/Streaming/LandscapeMeshMobileUpdate.h
@@ -64,7 +64,7 @@ protected:
 	void DoFinishUpdate(const FContext& Context);
 	void DoCancel(const FContext& Context);
 
-	FVertexBufferRHIRef IntermediateVertexBuffer;
+	FBufferRHIRef IntermediateVertexBuffer;
 	void* StagingLODDataArray[MAX_MESH_LOD_COUNT];
 	int64 StagingLODDataSizes[MAX_MESH_LOD_COUNT];
 };
diff --git a/Engine/Source/Runtime/Engine/Private/Streaming/SkeletalMeshUpdate.h b/Engine/Source/Runtime/Engine/Private/Streaming/SkeletalMeshUpdate.h
--- a/Engine/Source/Runtime/Engine/Private/Streaming/SkeletalMeshUpdate.h
+++ b/Engine/Source/Runtime/Engine/Private/Streaming/SkeletalMeshUpdate.h
@@ -70,14 +70,14 @@ protected:
 	/** Correspond to the buffers in FSkeletalMeshLODRenderData */
 	struct FIntermediateBuffers
 	{
-		FVertexBufferRHIRef TangentsVertexBuffer;
-		FVertexBufferRHIRef TexCoordVertexBuffer;
-		FVertexBufferRHIRef PositionVertexBuffer;
-		FVertexBufferRHIRef ColorVertexBuffer;
+		FBufferRHIRef TangentsVertexBuffer;
+		FBufferRHIRef TexCoordVertexBuffer;
+		FBufferRHIRef PositionVertexBuffer;
+		FBufferRHIRef ColorVertexBuffer;
 		FSkinWeightRHIInfo SkinWeightVertexBuffer;
-		FVertexBufferRHIRef ClothVertexBuffer;
-		FIndexBufferRHIRef IndexBuffer;
-		FIndexBufferRHIRef AdjacencyIndexBuffer;
+		FBufferRHIRef ClothVertexBuffer;
+		FBufferRHIRef IndexBuffer;
+		FBufferRHIRef AdjacencyIndexBuffer;
 		TArray<TPair<FName, FSkinWeightRHIInfo>> AltSkinWeightVertexBuffers;
 
 		void CreateFromCPUData_RenderThread(FSkeletalMeshLODRenderData& LODResource);
diff --git a/Engine/Source/Runtime/Engine/Private/Streaming/StaticMeshUpdate.h b/Engine/Source/Runtime/Engine/Private/Streaming/StaticMeshUpdate.h
--- a/Engine/Source/Runtime/Engine/Private/Streaming/StaticMeshUpdate.h
+++ b/Engine/Source/Runtime/Engine/Private/Streaming/StaticMeshUpdate.h
@@ -73,16 +73,16 @@ protected:
 	/** Correspond to the buffers in FStaticMeshLODResources */
 	struct FIntermediateBuffers
 	{
-		FVertexBufferRHIRef TangentsVertexBuffer;
-		FVertexBufferRHIRef TexCoordVertexBuffer;
-		FVertexBufferRHIRef PositionVertexBuffer;
-		FVertexBufferRHIRef ColorVertexBuffer;
-		FIndexBufferRHIRef IndexBuffer;
-		FIndexBufferRHIRef ReversedIndexBuffer;
-		FIndexBufferRHIRef DepthOnlyIndexBuffer;
-		FIndexBufferRHIRef ReversedDepthOnlyIndexBuffer;
-		FIndexBufferRHIRef WireframeIndexBuffer;
-		FIndexBufferRHIRef AdjacencyIndexBuffer;
+		FBufferRHIRef TangentsVertexBuffer;
+		FBufferRHIRef TexCoordVertexBuffer;
+		FBufferRHIRef PositionVertexBuffer;
+		FBufferRHIRef ColorVertexBuffer;
+		FBufferRHIRef IndexBuffer;
+		FBufferRHIRef ReversedIndexBuffer;
+		FBufferRHIRef DepthOnlyIndexBuffer;
+		FBufferRHIRef ReversedDepthOnlyIndexBuffer;
+		FBufferRHIRef WireframeIndexBuffer;
+		FBufferRHIRef AdjacencyIndexBuffer;
 
 		void CreateFromCPUData_RenderThread(FStaticMeshLODResources& LODResource);
 		void CreateFromCPUData_Async(FStaticMeshLODResources& LODResource);
diff --git a/Engine/Source/Runtime/Engine/Public/DynamicMeshBuilder.h b/Engine/Source/Runtime/Engine/Public/DynamicMeshBuilder.h
--- a/Engine/Source/Runtime/Engine/Public/DynamicMeshBuilder.h
+++ b/Engine/Source/Runtime/Engine/Public/DynamicMeshBuilder.h
@@ -146,10 +146,10 @@ public:
 	int32 GetVertexBufferSize(uint32 Stride, uint32 NumElements) const;
 	virtual ~FDynamicMeshBufferAllocator();
 
-	virtual FIndexBufferRHIRef AllocIndexBuffer(uint32 NumElements);
-	virtual void ReleaseIndexBuffer(FIndexBufferRHIRef& IndexBufferRHI);
-	virtual FVertexBufferRHIRef AllocVertexBuffer(uint32 Stride, uint32 NumElements);
-	virtual void ReleaseVertexBuffer(FVertexBufferRHIRef& VertexBufferRHI);
+	virtual FBufferRHIRef AllocIndexBuffer(uint32 NumElements);
+	virtual void ReleaseIndexBuffer(FBufferRHIRef& IndexBufferRHI);
+	virtual FBufferRHIRef AllocVertexBuffer(uint32 Stride, uint32 NumElements);
+	virtual void ReleaseVertexBuffer(FBufferRHIRef& VertexBufferRHI);
 };
 
 /**
diff --git a/Engine/Source/Runtime/Engine/Public/GPUSkinVertexFactory.h b/Engine/Source/Runtime/Engine/Public/GPUSkinVertexFactory.h
--- a/Engine/Source/Runtime/Engine/Public/GPUSkinVertexFactory.h
+++ b/Engine/Source/Runtime/Engine/Public/GPUSkinVertexFactory.h
@@ -75,7 +75,7 @@ struct FVertexBufferAndSRV
 		VertexBufferSRV.SafeRelease();
 	}
 
-	FVertexBufferRHIRef VertexBufferRHI;
+	FBufferRHIRef VertexBufferRHI;
 	FShaderResourceViewRHIRef VertexBufferSRV;
 };
 
diff --git a/Engine/Source/Runtime/Engine/Public/GPUSortManager.h b/Engine/Source/Runtime/Engine/Public/GPUSortManager.h
--- a/Engine/Source/Runtime/Engine/Public/GPUSortManager.h
+++ b/Engine/Source/Runtime/Engine/Public/GPUSortManager.h
@@ -74,14 +74,14 @@ public:
 private:
 
 	/** Vertex buffer storage for particle sort keys. */
-	FVertexBufferRHIRef KeyBuffers[2];
+	FBufferRHIRef KeyBuffers[2];
 	/** Shader resource view for particle sort keys. */
 	FShaderResourceViewRHIRef KeyBufferSRVs[2];
 	/** Unordered access view for particle sort keys. */
 	FUnorderedAccessViewRHIRef KeyBufferUAVs[2];
 
 	/** Vertex buffer containing sorted particle vertices. */
-	FVertexBufferRHIRef VertexBuffers[2];
+	FBufferRHIRef VertexBuffers[2];
 	/** Shader resource view for sorting particle vertices. */
 	FShaderResourceViewRHIRef VertexBufferSortSRVs[2];
 	/** Unordered access view for sorting particle vertices. */
diff --git a/Engine/Source/Runtime/Engine/Public/PrecomputedVolumetricLightmap.h b/Engine/Source/Runtime/Engine/Public/PrecomputedVolumetricLightmap.h
--- a/Engine/Source/Runtime/Engine/Public/PrecomputedVolumetricLightmap.h
+++ b/Engine/Source/Runtime/Engine/Public/PrecomputedVolumetricLightmap.h
@@ -199,10 +199,10 @@ public:
 	/**
 	 * Runtime data for sub level streaming
 	 */
-	FVertexBufferRHIRef SubLevelBrickPositionsBuffer;
+	FBufferRHIRef SubLevelBrickPositionsBuffer;
 	FShaderResourceViewRHIRef SubLevelBrickPositionsSRV;
 
-	FVertexBufferRHIRef IndirectionTextureOriginalValuesBuffer;
+	FBufferRHIRef IndirectionTextureOriginalValuesBuffer;
 	FShaderResourceViewRHIRef IndirectionTextureOriginalValuesSRV;
 
 	int32 BrickDataBaseOffsetInAtlas;
diff --git a/Engine/Source/Runtime/Engine/Public/PrimitiveUniformShaderParameters.h b/Engine/Source/Runtime/Engine/Public/PrimitiveUniformShaderParameters.h
--- a/Engine/Source/Runtime/Engine/Public/PrimitiveUniformShaderParameters.h
+++ b/Engine/Source/Runtime/Engine/Public/PrimitiveUniformShaderParameters.h
@@ -333,26 +333,26 @@ public:
 	FInstanceSceneShaderData InstanceSceneData;
 	FLightmapSceneShaderData LightmapSceneData;
 
-	FStructuredBufferRHIRef PrimitiveSceneDataBufferRHI;
+	FBufferRHIRef PrimitiveSceneDataBufferRHI;
 	FShaderResourceViewRHIRef PrimitiveSceneDataBufferSRV;
 
-	FStructuredBufferRHIRef SkyIrradianceEnvironmentMapRHI;
+	FBufferRHIRef SkyIrradianceEnvironmentMapRHI;
 	FShaderResourceViewRHIRef SkyIrradianceEnvironmentMapSRV;
 
-	FStructuredBufferRHIRef InstanceSceneDataBufferRHI;
+	FBufferRHIRef InstanceSceneDataBufferRHI;
 	FShaderResourceViewRHIRef InstanceSceneDataBufferSRV;
 
 	FTexture2DRHIRef PrimitiveSceneDataTextureRHI;
 	FShaderResourceViewRHIRef PrimitiveSceneDataTextureSRV;
 
-	FStructuredBufferRHIRef LightmapSceneDataBufferRHI;
+	FBufferRHIRef LightmapSceneDataBufferRHI;
 	FShaderResourceViewRHIRef LightmapSceneDataBufferSRV;
 
 //#if WITH_EDITOR
-	FVertexBufferRHIRef EditorVisualizeLevelInstanceDataBufferRHI;
+	FBufferRHIRef EditorVisualizeLevelInstanceDataBufferRHI;
 	FShaderResourceViewRHIRef EditorVisualizeLevelInstanceDataBufferSRV;
 
-	FVertexBufferRHIRef EditorSelectedDataBufferRHI;
+	FBufferRHIRef EditorSelectedDataBufferRHI;
 	FShaderResourceViewRHIRef EditorSelectedDataBufferSRV;
 //#endif
 };
diff --git a/Engine/Source/Runtime/Engine/Public/RawIndexBuffer.h b/Engine/Source/Runtime/Engine/Public/RawIndexBuffer.h
--- a/Engine/Source/Runtime/Engine/Public/RawIndexBuffer.h
+++ b/Engine/Source/Runtime/Engine/Public/RawIndexBuffer.h
@@ -251,8 +251,8 @@ public:
 	int32 GetIndexDataSize() const { return IndexStorage.Num(); }
 
 	/** Create an RHI index buffer with CPU data. CPU data may be discarded after creation (see TResourceArray::Discard) */
-	FIndexBufferRHIRef CreateRHIBuffer_RenderThread();
-	FIndexBufferRHIRef CreateRHIBuffer_Async();
+	FBufferRHIRef CreateRHIBuffer_RenderThread();
+	FBufferRHIRef CreateRHIBuffer_Async();
 
 	/** Copy everything, keeping reference to the same RHI resources. */
 	void CopyRHIForStreaming(const FRawStaticIndexBuffer& Other, bool InAllowCPUAccess);
@@ -300,7 +300,7 @@ public:
 
 private:
 	template <bool bRenderThread>
-	FIndexBufferRHIRef CreateRHIBuffer_Internal();
+	FBufferRHIRef CreateRHIBuffer_Internal();
 
 	/** Storage for indices. */
 	TResourceArray<uint8, INDEXBUFFER_ALIGNMENT> IndexStorage;
@@ -468,8 +468,8 @@ public:
 	}
 
 	/** Create an RHI index buffer with CPU data. CPU data may be discarded after creation (see TResourceArray::Discard) */
-	FIndexBufferRHIRef CreateRHIBuffer_RenderThread() { return CreateRHIBuffer_Internal<true>(); }
-	FIndexBufferRHIRef CreateRHIBuffer_Async() { return CreateRHIBuffer_Internal<false>(); }
+	FBufferRHIRef CreateRHIBuffer_RenderThread() { return CreateRHIBuffer_Internal<true>(); }
+	FBufferRHIRef CreateRHIBuffer_Async() { return CreateRHIBuffer_Internal<false>(); }
 
 	/** Similar to Init/ReleaseRHI but only update existing SRV so references to the SRV stays valid */
 	template <uint32 MaxNumUpdates>
@@ -515,7 +515,7 @@ private:
 	}
 
 	template <bool bRenderThread>
-	FIndexBufferRHIRef CreateRHIBuffer_Internal()
+	FBufferRHIRef CreateRHIBuffer_Internal()
 	{
 		if (CachedNumIndices)
 		{
@@ -529,7 +529,7 @@ private:
 				Flags = (EBufferUsageFlags)(Flags | BUF_ShaderResource);
 			}
 
-			FIndexBufferRHIRef Ret;
+			FBufferRHIRef Ret;
 			const uint32 Size = Indices.Num() * sizeof(INDEX_TYPE);
 			CreateInfo.bWithoutNativeResource = !Size;
 			if (bRenderThread)
diff --git a/Engine/Source/Runtime/Engine/Public/Rendering/ColorVertexBuffer.h b/Engine/Source/Runtime/Engine/Public/Rendering/ColorVertexBuffer.h
--- a/Engine/Source/Runtime/Engine/Public/Rendering/ColorVertexBuffer.h
+++ b/Engine/Source/Runtime/Engine/Public/Rendering/ColorVertexBuffer.h
@@ -138,8 +138,8 @@ public:
 	}
 
 	/** Create an RHI vertex buffer with CPU data. CPU data may be discarded after creation (see TResourceArray::Discard) */
-	FVertexBufferRHIRef CreateRHIBuffer_RenderThread();
-	FVertexBufferRHIRef CreateRHIBuffer_Async();
+	FBufferRHIRef CreateRHIBuffer_RenderThread();
+	FBufferRHIRef CreateRHIBuffer_Async();
 
 	/** Copy everything, keeping reference to the same RHI resources. */
 	void CopyRHIForStreaming(const FColorVertexBuffer& Other, bool InAllowCPUAccess);
@@ -211,7 +211,7 @@ private:
 	void AllocateData(bool bNeedsCPUAccess = true);
 
 	template <bool bRenderThread>
-	FVertexBufferRHIRef CreateRHIBuffer_Internal();
+	FBufferRHIRef CreateRHIBuffer_Internal();
 
 	/** Purposely hidden */
 	ENGINE_API FColorVertexBuffer(const FColorVertexBuffer &rhs);
diff --git a/Engine/Source/Runtime/Engine/Public/Rendering/MorphTargetVertexInfoBuffers.h b/Engine/Source/Runtime/Engine/Public/Rendering/MorphTargetVertexInfoBuffers.h
--- a/Engine/Source/Runtime/Engine/Public/Rendering/MorphTargetVertexInfoBuffers.h
+++ b/Engine/Source/Runtime/Engine/Public/Rendering/MorphTargetVertexInfoBuffers.h
@@ -57,10 +57,10 @@ public:
 		return NumSplitsPerMorph[Index];
 	}
 
-	FVertexBufferRHIRef VertexIndicesVB;
+	FBufferRHIRef VertexIndicesVB;
 	FShaderResourceViewRHIRef VertexIndicesSRV;
 
-	FVertexBufferRHIRef MorphDeltasVB;
+	FBufferRHIRef MorphDeltasVB;
 	FShaderResourceViewRHIRef MorphDeltasSRV;
 
 	// Changes to this struct must be reflected in MorphTargets.usf
diff --git a/Engine/Source/Runtime/Engine/Public/Rendering/MultiSizeIndexContainer.h b/Engine/Source/Runtime/Engine/Public/Rendering/MultiSizeIndexContainer.h
--- a/Engine/Source/Runtime/Engine/Public/Rendering/MultiSizeIndexContainer.h
+++ b/Engine/Source/Runtime/Engine/Public/Rendering/MultiSizeIndexContainer.h
@@ -94,8 +94,8 @@ public:
 
 	friend FArchive& operator<<(FArchive& Ar, FMultiSizeIndexContainer& Buffer);
 
-	FIndexBufferRHIRef CreateRHIBuffer_RenderThread();
-	FIndexBufferRHIRef CreateRHIBuffer_Async();
+	FBufferRHIRef CreateRHIBuffer_RenderThread();
+	FBufferRHIRef CreateRHIBuffer_Async();
 
 	template <uint32 MaxNumUpdates>
 	void InitRHIForStreaming(FRHIBuffer* IntermediateBuffer, TRHIResourceUpdateBatcher<MaxNumUpdates>& Batcher)
diff --git a/Engine/Source/Runtime/Engine/Public/Rendering/PositionVertexBuffer.h b/Engine/Source/Runtime/Engine/Public/Rendering/PositionVertexBuffer.h
--- a/Engine/Source/Runtime/Engine/Public/Rendering/PositionVertexBuffer.h
+++ b/Engine/Source/Runtime/Engine/Public/Rendering/PositionVertexBuffer.h
@@ -95,8 +95,8 @@ public:
 	}
 
 	/** Create an RHI vertex buffer with CPU data. CPU data may be discarded after creation (see TResourceArray::Discard) */
-	FVertexBufferRHIRef CreateRHIBuffer_RenderThread();
-	FVertexBufferRHIRef CreateRHIBuffer_Async();
+	FBufferRHIRef CreateRHIBuffer_RenderThread();
+	FBufferRHIRef CreateRHIBuffer_Async();
 
 	/** Copy everything, keeping reference to the same RHI resources. */
 	void CopyRHIForStreaming(const FPositionVertexBuffer& Other, bool InAllowCPUAccess);
@@ -163,6 +163,6 @@ private:
 	void AllocateData(bool bInNeedsCPUAccess = true);
 
 	template <bool bRenderThread>
-	FVertexBufferRHIRef CreateRHIBuffer_Internal();
+	FBufferRHIRef CreateRHIBuffer_Internal();
 };
 
diff --git a/Engine/Source/Runtime/Engine/Public/Rendering/SkeletalMeshVertexClothBuffer.h b/Engine/Source/Runtime/Engine/Public/Rendering/SkeletalMeshVertexClothBuffer.h
--- a/Engine/Source/Runtime/Engine/Public/Rendering/SkeletalMeshVertexClothBuffer.h
+++ b/Engine/Source/Runtime/Engine/Public/Rendering/SkeletalMeshVertexClothBuffer.h
@@ -124,8 +124,8 @@ public:
 	}
 
 	/** Create an RHI vertex buffer with CPU data. CPU data may be discarded after creation (see TResourceArray::Discard) */
-	FVertexBufferRHIRef CreateRHIBuffer_RenderThread();
-	FVertexBufferRHIRef CreateRHIBuffer_Async();
+	FBufferRHIRef CreateRHIBuffer_RenderThread();
+	FBufferRHIRef CreateRHIBuffer_Async();
 
 	/** Similar to Init/ReleaseRHI but only update existing SRV so references to the SRV stays valid */
 	template <uint32 MaxNumUpdates>
@@ -173,5 +173,5 @@ private:
 	void AllocateData();
 
 	template <bool bRenderThread>
-	FVertexBufferRHIRef CreateRHIBuffer_Internal();
+	FBufferRHIRef CreateRHIBuffer_Internal();
 };
\ No newline at end of file
diff --git a/Engine/Source/Runtime/Engine/Public/Rendering/SkinWeightVertexBuffer.h b/Engine/Source/Runtime/Engine/Public/Rendering/SkinWeightVertexBuffer.h
--- a/Engine/Source/Runtime/Engine/Public/Rendering/SkinWeightVertexBuffer.h
+++ b/Engine/Source/Runtime/Engine/Public/Rendering/SkinWeightVertexBuffer.h
@@ -64,8 +64,8 @@ struct FSkinWeightInfo
 
 struct FSkinWeightRHIInfo
 {
-	FVertexBufferRHIRef DataVertexBufferRHI;
-	FVertexBufferRHIRef LookupVertexBufferRHI;
+	FBufferRHIRef DataVertexBufferRHI;
+	FBufferRHIRef LookupVertexBufferRHI;
 };
 
 /** A lookup vertex buffer storing skin weight stream offset / influence count. Only used for unlimited bone influences. */
@@ -128,8 +128,8 @@ public:
 	{ return bNeedsCPUAccess; }
 
 	/** Create an RHI vertex buffer with CPU data. CPU data may be discarded after creation (see TResourceArray::Discard) */
-	FVertexBufferRHIRef CreateRHIBuffer_RenderThread();
-	FVertexBufferRHIRef CreateRHIBuffer_Async();
+	FBufferRHIRef CreateRHIBuffer_RenderThread();
+	FBufferRHIRef CreateRHIBuffer_Async();
 
 	/** Similar to Init/ReleaseRHI but only update existing SRV so references to the SRV stays valid */
 	template <uint32 MaxNumUpdates>
@@ -170,7 +170,7 @@ private:
 	ENGINE_API void AllocateData();
 
 	template <bool bRenderThread>
-	FVertexBufferRHIRef CreateRHIBuffer_Internal();
+	FBufferRHIRef CreateRHIBuffer_Internal();
 
 	/** true if this vertex buffer will be used with CPU skinning. Resource arrays are set to cpu accessible if this is true */
 	bool bNeedsCPUAccess;
@@ -293,8 +293,8 @@ public:
 	ENGINE_API void CopyDataFromBuffer(const TArrayView<const FSkinWeightInfo>& SkinWeightData);
 
 	/** Create an RHI vertex buffer with CPU data. CPU data may be discarded after creation (see TResourceArray::Discard) */
-	FVertexBufferRHIRef CreateRHIBuffer_RenderThread();
-	FVertexBufferRHIRef CreateRHIBuffer_Async();
+	FBufferRHIRef CreateRHIBuffer_RenderThread();
+	FBufferRHIRef CreateRHIBuffer_Async();
 
 	/** Similar to Init/ReleaseRHI but only update existing SRV so references to the SRV stays valid */
 	template <uint32 MaxNumUpdates>
@@ -368,7 +368,7 @@ private:
 	ENGINE_API void AllocateData();
 
 	template <bool bRenderThread>
-	FVertexBufferRHIRef CreateRHIBuffer_Internal();
+	FBufferRHIRef CreateRHIBuffer_Internal();
 };
 
 /** A container for skin weights data vertex buffer and lookup vertex buffer. */
diff --git a/Engine/Source/Runtime/Engine/Public/Rendering/StaticMeshVertexBuffer.h b/Engine/Source/Runtime/Engine/Public/Rendering/StaticMeshVertexBuffer.h
--- a/Engine/Source/Runtime/Engine/Public/Rendering/StaticMeshVertexBuffer.h
+++ b/Engine/Source/Runtime/Engine/Public/Rendering/StaticMeshVertexBuffer.h
@@ -389,10 +389,10 @@ public:
 	}
 
 	/** Create an RHI vertex buffer with CPU data. CPU data may be discarded after creation (see TResourceArray::Discard) */
-	FVertexBufferRHIRef CreateTangentsRHIBuffer_RenderThread();
-	FVertexBufferRHIRef CreateTangentsRHIBuffer_Async();
-	FVertexBufferRHIRef CreateTexCoordRHIBuffer_RenderThread();
-	FVertexBufferRHIRef CreateTexCoordRHIBuffer_Async();
+	FBufferRHIRef CreateTangentsRHIBuffer_RenderThread();
+	FBufferRHIRef CreateTangentsRHIBuffer_Async();
+	FBufferRHIRef CreateTexCoordRHIBuffer_RenderThread();
+	FBufferRHIRef CreateTexCoordRHIBuffer_Async();
 
 	/** Copy everything, keeping reference to the same RHI resources. */
 	void CopyRHIForStreaming(const FStaticMeshVertexBuffer& Other, bool InAllowCPUAccess);
@@ -540,10 +540,10 @@ private:
 	void ConvertHalfTexcoordsToFloat(const uint8* InData);
 
 	template <bool bRenderThread>
-	FVertexBufferRHIRef CreateTangentsRHIBuffer_Internal();
+	FBufferRHIRef CreateTangentsRHIBuffer_Internal();
 
 	template<bool bRenderThread>
-	FVertexBufferRHIRef CreateTexCoordRHIBuffer_Internal();
+	FBufferRHIRef CreateTexCoordRHIBuffer_Internal();
 
 	void InitTangentAndTexCoordStrides();
 };
\ No newline at end of file
diff --git a/Engine/Source/Runtime/Engine/Public/StaticMeshResources.h b/Engine/Source/Runtime/Engine/Public/StaticMeshResources.h
--- a/Engine/Source/Runtime/Engine/Public/StaticMeshResources.h
+++ b/Engine/Source/Runtime/Engine/Public/StaticMeshResources.h
@@ -286,7 +286,7 @@ private:
 		uint32 pad1;
 	};
 
-	FVertexBufferRHIRef BufferSectionTriangleRHI = nullptr;
+	FBufferRHIRef BufferSectionTriangleRHI = nullptr;
 	FShaderResourceViewRHIRef BufferSectionTriangleSRV = nullptr;
 
 	FStaticMeshSectionAreaWeightedTriangleSamplerArray* Samplers = nullptr;
diff --git a/Engine/Source/Runtime/HeadMountedDisplay/Public/HeadMountedDisplayTypes.h b/Engine/Source/Runtime/HeadMountedDisplay/Public/HeadMountedDisplayTypes.h
--- a/Engine/Source/Runtime/HeadMountedDisplay/Public/HeadMountedDisplayTypes.h
+++ b/Engine/Source/Runtime/HeadMountedDisplay/Public/HeadMountedDisplayTypes.h
@@ -31,8 +31,8 @@ public:
 
 	void BuildMesh(const FVector2D Positions[], uint32 VertexCount, EHMDMeshType MeshType);
 
-	FVertexBufferRHIRef VertexBufferRHI;
-	FIndexBufferRHIRef IndexBufferRHI;
+	FBufferRHIRef VertexBufferRHI;
+	FBufferRHIRef IndexBufferRHI;
 
 	unsigned  NumVertices;
 	unsigned  NumIndices;
diff --git a/Engine/Source/Runtime/Landscape/Private/LandscapeEditLayers.cpp b/Engine/Source/Runtime/Landscape/Private/LandscapeEditLayers.cpp
--- a/Engine/Source/Runtime/Landscape/Private/LandscapeEditLayers.cpp
+++ b/Engine/Source/Runtime/Landscape/Private/LandscapeEditLayers.cpp
@@ -771,7 +771,7 @@ public:
 private:
 	friend class FLandscapeLayerWeightmapExtractMaterialLayersCS;
 
-	FStructuredBufferRHIRef ComponentsData;
+	FBufferRHIRef ComponentsData;
 	FShaderResourceViewRHIRef ComponentsDataSRV;
 	TArray<FLandscapeLayerWeightmapExtractMaterialLayersComponentData> OriginalComponentsData;
 	int32 ComponentsDataCount;
@@ -945,17 +945,17 @@ public:
 private:
 	friend class FLandscapeLayerWeightmapPackMaterialLayersCS;
 
-	FStructuredBufferRHIRef ComponentsData;
+	FBufferRHIRef ComponentsData;
 	FShaderResourceViewRHIRef ComponentsDataSRV;
 	TArray<FLandscapeLayerWeightmapPackMaterialLayersComponentData> OriginalComponentsData;
 	int32 ComponentsDataCount;
 
 	TArray<float> OriginalWeightmapWeightBlendModeData;
-	FVertexBufferRHIRef WeightmapWeightBlendMode;
+	FBufferRHIRef WeightmapWeightBlendMode;
 	FShaderResourceViewRHIRef WeightmapWeightBlendModeSRV;
 
 	TArray<FVector2D> OriginalTextureOutputOffset;
-	FVertexBufferRHIRef WeightmapTextureOutputOffset;
+	FBufferRHIRef WeightmapTextureOutputOffset;
 	FShaderResourceViewRHIRef WeightmapTextureOutputOffsetSRV;
 };
 
diff --git a/Engine/Source/Runtime/Landscape/Public/LandscapeRender.h b/Engine/Source/Runtime/Landscape/Public/LandscapeRender.h
--- a/Engine/Source/Runtime/Landscape/Public/LandscapeRender.h
+++ b/Engine/Source/Runtime/Landscape/Public/LandscapeRender.h
@@ -458,7 +458,7 @@ class FNullLandscapeRenderSystemResources : public FRenderResource
 {
 public:
 
-	FVertexBufferRHIRef SectionLODBuffer;
+	FBufferRHIRef SectionLODBuffer;
 	FShaderResourceViewRHIRef SectionLODSRV;
 	TUniformBufferRef<FLandscapeSectionLODUniformParameters> UniformBuffer;
 
@@ -549,13 +549,13 @@ struct FLandscapeRenderSystem
 	TArray<FLandscapeComponentSceneProxy*> SceneProxies;
 	TArray<uint8> SectionCurrentFirstLODIndices;
 
-	FVertexBufferRHIRef SectionLODBuffer;
+	FBufferRHIRef SectionLODBuffer;
 	FShaderResourceViewRHIRef SectionLODSRV;
-	FVertexBufferRHIRef SectionLODBiasBuffer;
+	FBufferRHIRef SectionLODBiasBuffer;
 	FShaderResourceViewRHIRef SectionLODBiasSRV;
-	FVertexBufferRHIRef SectionTessellationFalloffCBuffer;
+	FBufferRHIRef SectionTessellationFalloffCBuffer;
 	FShaderResourceViewRHIRef SectionTessellationFalloffCSRV;
-	FVertexBufferRHIRef SectionTessellationFalloffKBuffer;
+	FBufferRHIRef SectionTessellationFalloffKBuffer;
 	FShaderResourceViewRHIRef SectionTessellationFalloffKSRV;
 
 	TUniformBufferRef<FLandscapeSectionLODUniformParameters> UniformBuffer;
diff --git a/Engine/Source/Runtime/MediaAssets/Private/Misc/MediaTextureResource.cpp b/Engine/Source/Runtime/MediaAssets/Private/Misc/MediaTextureResource.cpp
--- a/Engine/Source/Runtime/MediaAssets/Private/Misc/MediaTextureResource.cpp
+++ b/Engine/Source/Runtime/MediaAssets/Private/Misc/MediaTextureResource.cpp
@@ -865,7 +865,7 @@ void FMediaTextureResource::ConvertSample(const TSharedPtr<IMediaTextureSample,
 			}
 
 			// draw full size quad into render target
-			FVertexBufferRHIRef VertexBuffer = CreateTempMediaVertexBuffer();
+			FBufferRHIRef VertexBuffer = CreateTempMediaVertexBuffer();
 			CommandList.SetStreamSource(0, VertexBuffer, 0);
 			// set viewport to RT size
 			CommandList.SetViewport(0, 0, 0.0f, OutputDim.X, OutputDim.Y, 1.0f);
@@ -1001,7 +1001,7 @@ void FMediaTextureResource::CopyFromExternalTexture(const TSharedPtr <IMediaText
 			CopyShader->SetParameters(CommandList, SampleTexture, SamplerState, ScaleRotation, Offset);
 
 			// draw full size quad into render target
-			FVertexBufferRHIRef VertexBuffer = CreateTempMediaVertexBuffer();
+			FBufferRHIRef VertexBuffer = CreateTempMediaVertexBuffer();
 			CommandList.SetStreamSource(0, VertexBuffer, 0);
 			// set viewport to RT size
 			CommandList.SetViewport(0, 0, 0.0f, OutputDim.X, OutputDim.Y, 1.0f);
diff --git a/Engine/Source/Runtime/OpenGLDrv/Public/OpenGLResources.h b/Engine/Source/Runtime/OpenGLDrv/Public/OpenGLResources.h
--- a/Engine/Source/Runtime/OpenGLDrv/Public/OpenGLResources.h
+++ b/Engine/Source/Runtime/OpenGLDrv/Public/OpenGLResources.h
@@ -2016,7 +2016,7 @@ public:
 
 	virtual ~FOpenGLVertexBufferUnorderedAccessView();
 
-	FVertexBufferRHIRef VertexBufferRHI; // to keep the vertex buffer alive
+	FBufferRHIRef VertexBufferRHI; // to keep the vertex buffer alive
 
 	FOpenGLDynamicRHI* OpenGLRHI;
 
@@ -2032,7 +2032,7 @@ public:
 
 	virtual ~FOpenGLStructuredBufferUnorderedAccessView();
 
-	FStructuredBufferRHIRef StructuredBufferRHI; // to keep the stuctured buffer alive
+	FBufferRHIRef StructuredBufferRHI; // to keep the stuctured buffer alive
 
 	FOpenGLDynamicRHI* OpenGLRHI;
 
@@ -2056,8 +2056,8 @@ public:
 	int32 LimitMip;
 
 	/** Needed on OS X to force a rebind of the texture buffer to the texture name to workaround radr://18379338 */
-	FVertexBufferRHIRef VertexBuffer;
-	FIndexBufferRHIRef IndexBuffer;
+	FBufferRHIRef VertexBuffer;
+	FBufferRHIRef IndexBuffer;
 	uint64 ModificationVersion;
 	uint8 Format;
 
diff --git a/Engine/Source/Runtime/RHI/Public/DynamicRHI.h b/Engine/Source/Runtime/RHI/Public/DynamicRHI.h
--- a/Engine/Source/Runtime/RHI/Public/DynamicRHI.h
+++ b/Engine/Source/Runtime/RHI/Public/DynamicRHI.h
@@ -129,7 +129,7 @@ struct FRayTracingGeometrySegment
 {
 	DECLARE_TYPE_LAYOUT(FRayTracingGeometrySegment, NonVirtual);
 public:
-	LAYOUT_FIELD_INITIALIZED(FVertexBufferRHIRef, VertexBuffer, nullptr);
+	LAYOUT_FIELD_INITIALIZED(FBufferRHIRef, VertexBuffer, nullptr);
 	LAYOUT_FIELD_INITIALIZED(EVertexElementType, VertexBufferElementType, VET_Float3);
 
 	// Offset in bytes from the base address of the vertex buffer.
@@ -159,7 +159,7 @@ struct FRayTracingGeometryInitializer
 {
 	DECLARE_EXPORTED_TYPE_LAYOUT(FRayTracingGeometryInitializer, RHI_API, NonVirtual);
 public:
-	LAYOUT_FIELD_INITIALIZED(FIndexBufferRHIRef, IndexBuffer, nullptr);
+	LAYOUT_FIELD_INITIALIZED(FBufferRHIRef, IndexBuffer, nullptr);
 
 	// Offset in bytes from the base address of the index buffer.
 	LAYOUT_FIELD_INITIALIZED(uint32, IndexBufferOffset, 0);
diff --git a/Engine/Source/Runtime/RHI/Public/RHICommandList.h b/Engine/Source/Runtime/RHI/Public/RHICommandList.h
--- a/Engine/Source/Runtime/RHI/Public/RHICommandList.h
+++ b/Engine/Source/Runtime/RHI/Public/RHICommandList.h
@@ -3890,9 +3890,9 @@ public:
 
 
 // Forward declare RHI creation function so they can still be called from the deprecated immediate command list resource creation functions
-FIndexBufferRHIRef RHICreateIndexBuffer(uint32 Stride, uint32 Size, uint32 InUsage, FRHIResourceCreateInfo& CreateInfo);
-FVertexBufferRHIRef RHICreateVertexBuffer(uint32 Size, uint32 InUsage, FRHIResourceCreateInfo& CreateInfo);
-FStructuredBufferRHIRef RHICreateStructuredBuffer(uint32 Stride, uint32 Size, uint32 InUsage, FRHIResourceCreateInfo& CreateInfo);
+FBufferRHIRef RHICreateIndexBuffer(uint32 Stride, uint32 Size, uint32 InUsage, FRHIResourceCreateInfo& CreateInfo);
+FBufferRHIRef RHICreateVertexBuffer(uint32 Size, uint32 InUsage, FRHIResourceCreateInfo& CreateInfo);
+FBufferRHIRef RHICreateStructuredBuffer(uint32 Stride, uint32 Size, uint32 InUsage, FRHIResourceCreateInfo& CreateInfo);
 FTexture2DRHIRef RHICreateTexture2D(uint32 SizeX, uint32 SizeY, uint8 Format, uint32 NumMips, uint32 NumSamples, ETextureCreateFlags Flags, FRHIResourceCreateInfo& CreateInfo);
 FTexture2DRHIRef RHICreateTextureExternal2D(uint32 SizeX, uint32 SizeY, uint8 Format, uint32 NumMips, uint32 NumSamples, ETextureCreateFlags Flags, FRHIResourceCreateInfo& CreateInfo);
 FTexture2DRHIRef RHIAsyncCreateTexture2D(uint32 SizeX, uint32 SizeY, uint8 Format, uint32 NumMips, ETextureCreateFlags Flags, void** InitialMipData, uint32 NumInitialMips);
@@ -4073,12 +4073,12 @@ public:
 		return RHICreateUniformBuffer(Contents, Layout, Usage);
 	}
 	
-	FORCEINLINE FIndexBufferRHIRef CreateAndLockIndexBuffer(uint32 Stride, uint32 Size, EBufferUsageFlags InUsage, ERHIAccess InResourceState, FRHIResourceCreateInfo& CreateInfo, void*& OutDataBuffer)
+	FORCEINLINE FBufferRHIRef CreateAndLockIndexBuffer(uint32 Stride, uint32 Size, EBufferUsageFlags InUsage, ERHIAccess InResourceState, FRHIResourceCreateInfo& CreateInfo, void*& OutDataBuffer)
 	{
 		return GDynamicRHI->CreateAndLockIndexBuffer_RenderThread(*this, Stride, Size, InUsage, InResourceState, CreateInfo, OutDataBuffer);
 	}
 
-	FORCEINLINE FIndexBufferRHIRef CreateAndLockIndexBuffer(uint32 Stride, uint32 Size, uint32 InUsage, FRHIResourceCreateInfo& CreateInfo, void*& OutDataBuffer)
+	FORCEINLINE FBufferRHIRef CreateAndLockIndexBuffer(uint32 Stride, uint32 Size, uint32 InUsage, FRHIResourceCreateInfo& CreateInfo, void*& OutDataBuffer)
 	{
 		EBufferUsageFlags Usage = (EBufferUsageFlags) InUsage | BUF_IndexBuffer;
 		ERHIAccess ResourceState = RHIGetDefaultResourceState(Usage, true);
@@ -4107,12 +4107,12 @@ public:
 		GDynamicRHI->UnlockStagingBuffer_RenderThread(*this, StagingBuffer);
 	}
 	
-	FORCEINLINE FVertexBufferRHIRef CreateAndLockVertexBuffer(uint32 Size, EBufferUsageFlags InUsage, ERHIAccess InResourceState, FRHIResourceCreateInfo& CreateInfo, void*& OutDataBuffer)
+	FORCEINLINE FBufferRHIRef CreateAndLockVertexBuffer(uint32 Size, EBufferUsageFlags InUsage, ERHIAccess InResourceState, FRHIResourceCreateInfo& CreateInfo, void*& OutDataBuffer)
 	{
 		return GDynamicRHI->CreateAndLockVertexBuffer_RenderThread(*this, Size, InUsage, InResourceState, CreateInfo, OutDataBuffer);
 	}
 
-	FORCEINLINE FVertexBufferRHIRef CreateAndLockVertexBuffer(uint32 Size, uint32 InUsage, FRHIResourceCreateInfo& CreateInfo, void*& OutDataBuffer)
+	FORCEINLINE FBufferRHIRef CreateAndLockVertexBuffer(uint32 Size, uint32 InUsage, FRHIResourceCreateInfo& CreateInfo, void*& OutDataBuffer)
 	{
 		EBufferUsageFlags Usage = (EBufferUsageFlags) InUsage | BUF_VertexBuffer;
 		ERHIAccess ResourceState = RHIGetDefaultResourceState(Usage, true);
@@ -5095,22 +5095,22 @@ FORCEINLINE FStagingBufferRHIRef RHICreateStagingBuffer()
 	return FRHICommandListExecutor::GetImmediateCommandList().CreateStagingBuffer();
 }
 
-FORCEINLINE FIndexBufferRHIRef RHICreateAndLockIndexBuffer(uint32 Stride, uint32 Size, uint32 InUsage, FRHIResourceCreateInfo& CreateInfo, void*& OutDataBuffer)
+FORCEINLINE FBufferRHIRef RHICreateAndLockIndexBuffer(uint32 Stride, uint32 Size, uint32 InUsage, FRHIResourceCreateInfo& CreateInfo, void*& OutDataBuffer)
 {
 	return FRHICommandListExecutor::GetImmediateCommandList().CreateAndLockIndexBuffer(Stride, Size, InUsage, CreateInfo, OutDataBuffer);
 }
 
-FORCEINLINE FIndexBufferRHIRef RHICreateIndexBuffer(uint32 Stride, uint32 Size, uint32 InUsage, ERHIAccess InResourceState, FRHIResourceCreateInfo& CreateInfo)
+FORCEINLINE FBufferRHIRef RHICreateIndexBuffer(uint32 Stride, uint32 Size, uint32 InUsage, ERHIAccess InResourceState, FRHIResourceCreateInfo& CreateInfo)
 {
 	return GDynamicRHI->CreateBuffer_RenderThread(FRHICommandListExecutor::GetImmediateCommandList(), Size, (EBufferUsageFlags) InUsage | BUF_IndexBuffer, Stride, InResourceState, CreateInfo);
 }
 
-FORCEINLINE FIndexBufferRHIRef RHIAsyncCreateIndexBuffer(uint32 Stride, uint32 Size, uint32 InUsage, ERHIAccess InResourceState, FRHIResourceCreateInfo& CreateInfo)
+FORCEINLINE FBufferRHIRef RHIAsyncCreateIndexBuffer(uint32 Stride, uint32 Size, uint32 InUsage, ERHIAccess InResourceState, FRHIResourceCreateInfo& CreateInfo)
 {
 	return GDynamicRHI->RHICreateBuffer(Size, (EBufferUsageFlags)InUsage, Stride, InResourceState, CreateInfo);
 }
 
-FORCEINLINE FIndexBufferRHIRef RHICreateIndexBuffer(uint32 Stride, uint32 Size, uint32 InUsage, FRHIResourceCreateInfo& CreateInfo)
+FORCEINLINE FBufferRHIRef RHICreateIndexBuffer(uint32 Stride, uint32 Size, uint32 InUsage, FRHIResourceCreateInfo& CreateInfo)
 {
 	bool bHasInitialData = CreateInfo.BulkData != nullptr;
 	EBufferUsageFlags Usage = (EBufferUsageFlags) InUsage | BUF_IndexBuffer;
@@ -5118,7 +5118,7 @@ FORCEINLINE FIndexBufferRHIRef RHICreateIndexBuffer(uint32 Stride, uint32 Size,
 	return RHICreateIndexBuffer(Stride, Size, Usage, ResourceState, CreateInfo);
 }
 
-FORCEINLINE FIndexBufferRHIRef RHIAsyncCreateIndexBuffer(uint32 Stride, uint32 Size, uint32 InUsage, FRHIResourceCreateInfo& CreateInfo)
+FORCEINLINE FBufferRHIRef RHIAsyncCreateIndexBuffer(uint32 Stride, uint32 Size, uint32 InUsage, FRHIResourceCreateInfo& CreateInfo)
 {
 	bool bHasInitialData = CreateInfo.BulkData != nullptr;
 	EBufferUsageFlags Usage = (EBufferUsageFlags) InUsage | BUF_IndexBuffer;
@@ -5138,22 +5138,22 @@ FORCEINLINE void RHIUnlockIndexBuffer(FRHIBuffer* IndexBuffer)
 	 FRHICommandListExecutor::GetImmediateCommandList().UnlockBuffer(IndexBuffer);
 }
 
-FORCEINLINE FVertexBufferRHIRef RHICreateAndLockVertexBuffer(uint32 Size, uint32 InUsage, FRHIResourceCreateInfo& CreateInfo, void*& OutDataBuffer)
+FORCEINLINE FBufferRHIRef RHICreateAndLockVertexBuffer(uint32 Size, uint32 InUsage, FRHIResourceCreateInfo& CreateInfo, void*& OutDataBuffer)
 {
 	return FRHICommandListExecutor::GetImmediateCommandList().CreateAndLockVertexBuffer(Size, InUsage, CreateInfo, OutDataBuffer);
 }
 
-FORCEINLINE FVertexBufferRHIRef RHICreateVertexBuffer(uint32 Size, uint32 InUsage, ERHIAccess InResourceState, FRHIResourceCreateInfo& CreateInfo)
+FORCEINLINE FBufferRHIRef RHICreateVertexBuffer(uint32 Size, uint32 InUsage, ERHIAccess InResourceState, FRHIResourceCreateInfo& CreateInfo)
 {
 	return GDynamicRHI->CreateBuffer_RenderThread(FRHICommandListExecutor::GetImmediateCommandList(), Size, (EBufferUsageFlags) InUsage | BUF_VertexBuffer, 0, InResourceState, CreateInfo);
 }
 
-FORCEINLINE FVertexBufferRHIRef RHIAsyncCreateVertexBuffer(uint32 Size, uint32 InUsage, ERHIAccess InResourceState, FRHIResourceCreateInfo& CreateInfo)
+FORCEINLINE FBufferRHIRef RHIAsyncCreateVertexBuffer(uint32 Size, uint32 InUsage, ERHIAccess InResourceState, FRHIResourceCreateInfo& CreateInfo)
 {
 	return GDynamicRHI->RHICreateBuffer(Size, (EBufferUsageFlags)InUsage, 0, InResourceState, CreateInfo);
 }
 
-FORCEINLINE FVertexBufferRHIRef RHICreateVertexBuffer(uint32 Size, uint32 InUsage, FRHIResourceCreateInfo& CreateInfo)
+FORCEINLINE FBufferRHIRef RHICreateVertexBuffer(uint32 Size, uint32 InUsage, FRHIResourceCreateInfo& CreateInfo)
 {
 	bool bHasInitialData = CreateInfo.BulkData != nullptr;
 	EBufferUsageFlags Usage = (EBufferUsageFlags) InUsage | BUF_VertexBuffer;
@@ -5161,7 +5161,7 @@ FORCEINLINE FVertexBufferRHIRef RHICreateVertexBuffer(uint32 Size, uint32 InUsag
 	return RHICreateVertexBuffer(Size, Usage, ResourceState, CreateInfo);
 }
 
-FORCEINLINE FVertexBufferRHIRef RHIAsyncCreateVertexBuffer(uint32 Size, uint32 InUsage, FRHIResourceCreateInfo& CreateInfo)
+FORCEINLINE FBufferRHIRef RHIAsyncCreateVertexBuffer(uint32 Size, uint32 InUsage, FRHIResourceCreateInfo& CreateInfo)
 {
 	bool bHasInitialData = CreateInfo.BulkData != nullptr;
 	EBufferUsageFlags Usage = (EBufferUsageFlags) InUsage | BUF_VertexBuffer;
@@ -5181,12 +5181,12 @@ FORCEINLINE void RHIUnlockVertexBuffer(FRHIBuffer* VertexBuffer)
 	 FRHICommandListExecutor::GetImmediateCommandList().UnlockBuffer(VertexBuffer);
 }
 
-FORCEINLINE FStructuredBufferRHIRef RHICreateStructuredBuffer(uint32 Stride, uint32 Size, uint32 InUsage, ERHIAccess InResourceState, FRHIResourceCreateInfo& CreateInfo)
+FORCEINLINE FBufferRHIRef RHICreateStructuredBuffer(uint32 Stride, uint32 Size, uint32 InUsage, ERHIAccess InResourceState, FRHIResourceCreateInfo& CreateInfo)
 {
 	return GDynamicRHI->CreateBuffer_RenderThread(FRHICommandListExecutor::GetImmediateCommandList(), Size, (EBufferUsageFlags) InUsage | BUF_StructuredBuffer, Stride, InResourceState, CreateInfo);
 }
 
-FORCEINLINE FStructuredBufferRHIRef RHICreateStructuredBuffer(uint32 Stride, uint32 Size, uint32 InUsage, FRHIResourceCreateInfo& CreateInfo)
+FORCEINLINE FBufferRHIRef RHICreateStructuredBuffer(uint32 Stride, uint32 Size, uint32 InUsage, FRHIResourceCreateInfo& CreateInfo)
 {
 	bool bHasInitialData = CreateInfo.BulkData != nullptr;
 	EBufferUsageFlags Usage = (EBufferUsageFlags) InUsage | BUF_StructuredBuffer;
diff --git a/Engine/Source/Runtime/RHI/Public/RHIResources.h b/Engine/Source/Runtime/RHI/Public/RHIResources.h
--- a/Engine/Source/Runtime/RHI/Public/RHIResources.h
+++ b/Engine/Source/Runtime/RHI/Public/RHIResources.h
@@ -1551,7 +1551,7 @@ public:
 
 	virtual void* Lock(uint32 Offset, uint32 NumBytes) final override;
 	virtual void Unlock() final override;
-	FVertexBufferRHIRef ShadowBuffer;
+	FBufferRHIRef ShadowBuffer;
 	uint32 Offset;
 };
 
diff --git a/Engine/Source/Runtime/RHI/Public/RHIUtilities.h b/Engine/Source/Runtime/RHI/Public/RHIUtilities.h
--- a/Engine/Source/Runtime/RHI/Public/RHIUtilities.h
+++ b/Engine/Source/Runtime/RHI/Public/RHIUtilities.h
@@ -190,7 +190,7 @@ struct FTextureRWBuffer3D
 /** Encapsulates a GPU read/write buffer with its UAV and SRV. */
 struct FRWBuffer
 {
-	FVertexBufferRHIRef Buffer;
+	FBufferRHIRef Buffer;
 	FUnorderedAccessViewRHIRef UAV;
 	FShaderResourceViewRHIRef SRV;
 	uint32 NumBytes;
@@ -358,7 +358,7 @@ struct FTextureReadBuffer2D
 /** Encapsulates a GPU read buffer with its SRV. */
 struct FReadBuffer
 {
-	FVertexBufferRHIRef Buffer;
+	FBufferRHIRef Buffer;
 	FShaderResourceViewRHIRef SRV;
 	uint32 NumBytes;
 
@@ -386,7 +386,7 @@ struct FReadBuffer
 /** Encapsulates a GPU read/write structured buffer with its UAV and SRV. */
 struct FRWBufferStructured
 {
-	FStructuredBufferRHIRef Buffer;
+	FBufferRHIRef Buffer;
 	FUnorderedAccessViewRHIRef UAV;
 	FShaderResourceViewRHIRef SRV;
 	uint32 NumBytes;
@@ -439,7 +439,7 @@ struct FRWBufferStructured
 
 struct FByteAddressBuffer
 {
-	FStructuredBufferRHIRef Buffer;
+	FBufferRHIRef Buffer;
 	FShaderResourceViewRHIRef SRV;
 	uint32 NumBytes;
 
diff --git a/Engine/Source/Runtime/RenderCore/Private/ClearQuad.cpp b/Engine/Source/Runtime/RenderCore/Private/ClearQuad.cpp
--- a/Engine/Source/Runtime/RenderCore/Private/ClearQuad.cpp
+++ b/Engine/Source/Runtime/RenderCore/Private/ClearQuad.cpp
@@ -138,7 +138,7 @@ void DrawClearQuadMRT(FRHICommandList& RHICmdList, bool bClearColor, int32 NumCl
 		InnerVertices[3].Set(FMath::Lerp(-1.0f, 1.0f, FractionRect.X), FMath::Lerp(1.0f, -1.0f, FractionRect.W), Depth, 1.0f);
 
 		FRHIResourceCreateInfo CreateInfo;
-		FVertexBufferRHIRef VertexBufferRHI = RHICreateVertexBuffer(sizeof(FVector4) * 10, BUF_Volatile, CreateInfo);
+		FBufferRHIRef VertexBufferRHI = RHICreateVertexBuffer(sizeof(FVector4) * 10, BUF_Volatile, CreateInfo);
 		void* VoidPtr = RHILockBuffer(VertexBufferRHI, 0, sizeof(FVector4) * 10, RLM_WriteOnly);
 		
 		FVector4* Vertices = reinterpret_cast<FVector4*>(VoidPtr);
diff --git a/Engine/Source/Runtime/RenderCore/Private/RenderUtils.cpp b/Engine/Source/Runtime/RenderCore/Private/RenderUtils.cpp
--- a/Engine/Source/Runtime/RenderCore/Private/RenderUtils.cpp
+++ b/Engine/Source/Runtime/RenderCore/Private/RenderUtils.cpp
@@ -1368,12 +1368,12 @@ public:
 static TGlobalResource<FUnitCubeVertexBuffer> GUnitCubeVertexBuffer;
 static TGlobalResource<FUnitCubeIndexBuffer> GUnitCubeIndexBuffer;
 
-RENDERCORE_API FVertexBufferRHIRef& GetUnitCubeVertexBuffer()
+RENDERCORE_API FBufferRHIRef& GetUnitCubeVertexBuffer()
 {
 	return GUnitCubeVertexBuffer.VertexBufferRHI;
 }
 
-RENDERCORE_API FIndexBufferRHIRef& GetUnitCubeIndexBuffer()
+RENDERCORE_API FBufferRHIRef& GetUnitCubeIndexBuffer()
 {
 	return GUnitCubeIndexBuffer.IndexBufferRHI;
 }
diff --git a/Engine/Source/Runtime/RenderCore/Public/MediaShaders.h b/Engine/Source/Runtime/RenderCore/Public/MediaShaders.h
--- a/Engine/Source/Runtime/RenderCore/Public/MediaShaders.h
+++ b/Engine/Source/Runtime/RenderCore/Public/MediaShaders.h
@@ -63,10 +63,10 @@ struct FMediaElementVertex
 	{ }
 };
 
-inline FVertexBufferRHIRef CreateTempMediaVertexBuffer(float ULeft = 0.0f, float URight = 1.0f, float VTop = 0.0f, float VBottom = 1.0f)
+inline FBufferRHIRef CreateTempMediaVertexBuffer(float ULeft = 0.0f, float URight = 1.0f, float VTop = 0.0f, float VBottom = 1.0f)
 {
 	FRHIResourceCreateInfo CreateInfo;
-	FVertexBufferRHIRef VertexBufferRHI = RHICreateVertexBuffer(sizeof(FMediaElementVertex) * 4, BUF_Volatile, CreateInfo);
+	FBufferRHIRef VertexBufferRHI = RHICreateVertexBuffer(sizeof(FMediaElementVertex) * 4, BUF_Volatile, CreateInfo);
 	void* VoidPtr = RHILockBuffer(VertexBufferRHI, 0, sizeof(FMediaElementVertex) * 4, RLM_WriteOnly);
 
 	FMediaElementVertex* Vertices = (FMediaElementVertex*)VoidPtr;
diff --git a/Engine/Source/Runtime/RenderCore/Public/RenderResource.h b/Engine/Source/Runtime/RenderCore/Public/RenderResource.h
--- a/Engine/Source/Runtime/RenderCore/Public/RenderResource.h
+++ b/Engine/Source/Runtime/RenderCore/Public/RenderResource.h
@@ -498,7 +498,7 @@ public:
 class RENDERCORE_API FVertexBuffer : public FRenderResource
 {
 public:
-	FVertexBufferRHIRef VertexBufferRHI;
+	FBufferRHIRef VertexBufferRHI;
 
 	/** Destructor. */
 	virtual ~FVertexBuffer() {}
@@ -606,7 +606,7 @@ extern RENDERCORE_API TGlobalResource<FNullVertexBuffer> GNullVertexBuffer;
 class FIndexBuffer : public FRenderResource
 {
 public:
-	FIndexBufferRHIRef IndexBufferRHI;
+	FBufferRHIRef IndexBufferRHI;
 
 	/** Destructor. */
 	virtual ~FIndexBuffer() {}
diff --git a/Engine/Source/Runtime/RenderCore/Public/RenderUtils.h b/Engine/Source/Runtime/RenderCore/Public/RenderUtils.h
--- a/Engine/Source/Runtime/RenderCore/Public/RenderUtils.h
+++ b/Engine/Source/Runtime/RenderCore/Public/RenderUtils.h
@@ -585,10 +585,10 @@ inline bool VelocityEncodeDepth(const FStaticShaderPlatform Platform)
 }
 
 /** Unit cube vertex buffer (VertexDeclarationFVector4) */
-RENDERCORE_API FVertexBufferRHIRef& GetUnitCubeVertexBuffer();
+RENDERCORE_API FBufferRHIRef& GetUnitCubeVertexBuffer();
 
 /** Unit cube index buffer */
-RENDERCORE_API FIndexBufferRHIRef& GetUnitCubeIndexBuffer();
+RENDERCORE_API FBufferRHIRef& GetUnitCubeIndexBuffer();
 
 /**
 * Takes the requested buffer size and quantizes it to an appropriate size for the rest of the
diff --git a/Engine/Source/Runtime/Renderer/Private/DistanceFieldLightingShared.h b/Engine/Source/Runtime/Renderer/Private/DistanceFieldLightingShared.h
--- a/Engine/Source/Runtime/Renderer/Private/DistanceFieldLightingShared.h
+++ b/Engine/Source/Runtime/Renderer/Private/DistanceFieldLightingShared.h
@@ -462,7 +462,7 @@ public:
 	// Volatile must be written every frame before use.  Supports multiple writes per frame on PS4, unlike Dynamic.
 	bool bVolatile;
 
-	FVertexBufferRHIRef Buffer;
+	FBufferRHIRef Buffer;
 	FShaderResourceViewRHIRef BufferSRV;
 
 	FCPUUpdatedBuffer()
diff --git a/Engine/Source/Runtime/Renderer/Private/GPUBenchmark.cpp b/Engine/Source/Runtime/Renderer/Private/GPUBenchmark.cpp
--- a/Engine/Source/Runtime/Renderer/Private/GPUBenchmark.cpp
+++ b/Engine/Source/Runtime/Renderer/Private/GPUBenchmark.cpp
@@ -418,7 +418,7 @@ void RendererGPUBenchmark(FRHICommandListImmediate& RHICmdList, FSynthBenchmarkR
 	}
 
 	FRHIResourceCreateInfo CreateInfo(&Vertices);
-	FVertexBufferRHIRef VertexBuffer = RHICreateVertexBuffer(GBenchmarkVertices * sizeof(FBenchmarkVertex), BUF_Static, CreateInfo);
+	FBufferRHIRef VertexBuffer = RHICreateVertexBuffer(GBenchmarkVertices * sizeof(FBenchmarkVertex), BUF_Static, CreateInfo);
 
 	// two RT to ping pong so we force the GPU to flush it's pipeline
 	TRefCountPtr<IPooledRenderTarget> RTItems[3];
diff --git a/Engine/Source/Runtime/Renderer/Private/Halton.h b/Engine/Source/Runtime/Renderer/Private/Halton.h
--- a/Engine/Source/Runtime/Renderer/Private/Halton.h
+++ b/Engine/Source/Runtime/Renderer/Private/Halton.h
@@ -67,7 +67,7 @@ public:
 		return Iteration;
 	}
 
-	FStructuredBufferRHIRef SequenceIteration;
+	FBufferRHIRef SequenceIteration;
 
 private:
 	void InitializeSequence();
@@ -122,7 +122,7 @@ public:
 		return DimensionCount;
 	}
 
-	FStructuredBufferRHIRef PrimesBuffer;
+	FBufferRHIRef PrimesBuffer;
 
 private:
 	TArray<int> Primes;
diff --git a/Engine/Source/Runtime/Renderer/Private/LightSceneInfo.h b/Engine/Source/Runtime/Renderer/Private/LightSceneInfo.h
--- a/Engine/Source/Runtime/Renderer/Private/LightSceneInfo.h
+++ b/Engine/Source/Runtime/Renderer/Private/LightSceneInfo.h
@@ -198,7 +198,7 @@ public:
 	/** The identifier for the primitive in Scene->PrimitiveOctree. */
 	FOctreeElementId2 OctreeId;
 
-	mutable FVertexBufferRHIRef ShadowCapsuleShapesVertexBuffer;
+	mutable FBufferRHIRef ShadowCapsuleShapesVertexBuffer;
 	mutable FShaderResourceViewRHIRef ShadowCapsuleShapesSRV;
 
 	/** Persistent shadow state used for HZB occlusion culling. */
diff --git a/Engine/Source/Runtime/Renderer/Private/MeshDrawCommands.h b/Engine/Source/Runtime/Renderer/Private/MeshDrawCommands.h
--- a/Engine/Source/Runtime/Renderer/Private/MeshDrawCommands.h
+++ b/Engine/Source/Runtime/Renderer/Private/MeshDrawCommands.h
@@ -24,7 +24,7 @@ struct FPrimitiveIdVertexBufferPoolEntry
 {
 	int32 BufferSize = 0;
 	uint32 LastDiscardId = 0;
-	FVertexBufferRHIRef BufferRHI;
+	FBufferRHIRef BufferRHI;
 };
 
 class FPrimitiveIdVertexBufferPool : public FRenderResource
diff --git a/Engine/Source/Runtime/Renderer/Private/PostProcess/SceneFilterRendering.cpp b/Engine/Source/Runtime/Renderer/Private/PostProcess/SceneFilterRendering.cpp
--- a/Engine/Source/Runtime/Renderer/Private/PostProcess/SceneFilterRendering.cpp
+++ b/Engine/Source/Runtime/Renderer/Private/PostProcess/SceneFilterRendering.cpp
@@ -204,7 +204,7 @@ void DrawTransformedRectangle(
 	// we don't do the triangle optimization as this case is rare for the DrawTransformedRectangle case
 
 	FRHIResourceCreateInfo CreateInfo;
-	FVertexBufferRHIRef VertexBufferRHI = RHICreateVertexBuffer(sizeof(FFilterVertex) * 4, BUF_Volatile, CreateInfo);
+	FBufferRHIRef VertexBufferRHI = RHICreateVertexBuffer(sizeof(FFilterVertex) * 4, BUF_Volatile, CreateInfo);
 	void* VoidPtr = RHILockBuffer(VertexBufferRHI, 0, sizeof(FFilterVertex) * 4, RLM_WriteOnly);
 
 	FFilterVertex* Vertices = reinterpret_cast<FFilterVertex*>(VoidPtr);
diff --git a/Engine/Source/Runtime/Renderer/Private/RayTracing/RayTracingLighting.cpp b/Engine/Source/Runtime/Renderer/Private/RayTracing/RayTracingLighting.cpp
--- a/Engine/Source/Runtime/Renderer/Private/RayTracing/RayTracingLighting.cpp
+++ b/Engine/Source/Runtime/Renderer/Private/RayTracing/RayTracingLighting.cpp
@@ -135,7 +135,7 @@ static void CreateRaytracingLightCullingStructure(
 	FRHIResourceCreateInfo CreateInfo;
 	CreateInfo.ResourceArray = &RankedLights;
 
-	FStructuredBufferRHIRef RayTracingCullLights = RHICreateStructuredBuffer(sizeof(RankedLights[0]),
+	FBufferRHIRef RayTracingCullLights = RHICreateStructuredBuffer(sizeof(RankedLights[0]),
 		RankedLights.GetResourceDataSize(),
 		BUF_Static | BUF_ShaderResource,
 		CreateInfo);
diff --git a/Engine/Source/Runtime/Renderer/Private/SceneOcclusion.cpp b/Engine/Source/Runtime/Renderer/Private/SceneOcclusion.cpp
--- a/Engine/Source/Runtime/Renderer/Private/SceneOcclusion.cpp
+++ b/Engine/Source/Runtime/Renderer/Private/SceneOcclusion.cpp
@@ -1265,7 +1265,7 @@ static void BeginOcclusionTests(
 		{
 			uint32 BaseVertexOffset = 0;
 			FRHIResourceCreateInfo CreateInfo;
-			FVertexBufferRHIRef VertexBufferRHI = RHICreateVertexBuffer(sizeof(FVector) * NumVertices, BUF_Volatile, CreateInfo);
+			FBufferRHIRef VertexBufferRHI = RHICreateVertexBuffer(sizeof(FVector) * NumVertices, BUF_Volatile, CreateInfo);
 			void* VoidPtr = RHILockBuffer(VertexBufferRHI, 0, sizeof(FVector) * NumVertices, RLM_WriteOnly);
 
 			{
diff --git a/Engine/Source/Runtime/Renderer/Private/ScenePrivate.h b/Engine/Source/Runtime/Renderer/Private/ScenePrivate.h
--- a/Engine/Source/Runtime/Renderer/Private/ScenePrivate.h
+++ b/Engine/Source/Runtime/Renderer/Private/ScenePrivate.h
@@ -694,7 +694,7 @@ struct FHLODSceneNodeVisibilityState
 
 struct FExposureBufferData
 {
-	FVertexBufferRHIRef Buffer;
+	FBufferRHIRef Buffer;
 	FShaderResourceViewRHIRef SRV;
 	FUnorderedAccessViewRHIRef UAV;
 
@@ -1089,11 +1089,11 @@ public:
 	FRadianceCacheState RadianceCacheState;
 	FRadianceCacheState RadiosityRadianceCacheState;
 
-	FVertexBufferRHIRef IndirectShadowCapsuleShapesVertexBuffer;
+	FBufferRHIRef IndirectShadowCapsuleShapesVertexBuffer;
 	FShaderResourceViewRHIRef IndirectShadowCapsuleShapesSRV;
-	FVertexBufferRHIRef IndirectShadowMeshDistanceFieldCasterIndicesVertexBuffer;
+	FBufferRHIRef IndirectShadowMeshDistanceFieldCasterIndicesVertexBuffer;
 	FShaderResourceViewRHIRef IndirectShadowMeshDistanceFieldCasterIndicesSRV;
-	FVertexBufferRHIRef IndirectShadowLightDirectionVertexBuffer;
+	FBufferRHIRef IndirectShadowLightDirectionVertexBuffer;
 	FShaderResourceViewRHIRef IndirectShadowLightDirectionSRV;
 	FRWBuffer IndirectShadowVolumetricLightmapDerivedLightDirection;
 	FRWBuffer CapsuleTileIntersectionCountsBuffer;
diff --git a/Engine/Source/Runtime/Renderer/Private/SceneRendering.h b/Engine/Source/Runtime/Renderer/Private/SceneRendering.h
--- a/Engine/Source/Runtime/Renderer/Private/SceneRendering.h
+++ b/Engine/Source/Runtime/Renderer/Private/SceneRendering.h
@@ -1024,14 +1024,14 @@ struct FRayTracingLightData
 	TUniformBufferRef<FRaytracingLightDataPacked>	UniformBuffer;
 
 	/** Structured buffer containing all light data */
-	FStructuredBufferRHIRef							LightBuffer;
+	FBufferRHIRef									LightBuffer;
 	FShaderResourceViewRHIRef						LightBufferSRV;
 
 	/** Buffer of light indices reference by the culling volume */
 	FRWBuffer										LightIndices;
 
 	/** Camera-centered volume used to cull lights to cells */
-	FStructuredBufferRHIRef							LightCullVolume;
+	FBufferRHIRef									LightCullVolume;
 	FShaderResourceViewRHIRef						LightCullVolumeSRV;
 };
 #endif
diff --git a/Engine/Source/Runtime/Renderer/Private/ShadowRendering.cpp b/Engine/Source/Runtime/Renderer/Private/ShadowRendering.cpp
--- a/Engine/Source/Runtime/Renderer/Private/ShadowRendering.cpp
+++ b/Engine/Source/Runtime/Renderer/Private/ShadowRendering.cpp
@@ -865,7 +865,7 @@ void FProjectedShadowInfo::SetupProjectionStencilMask(
 		float StencilFar = Far.Z / Far.W;
 
 		FRHIResourceCreateInfo CreateInfo;
-		FVertexBufferRHIRef VertexBufferRHI = RHICreateVertexBuffer(sizeof(FVector4) * 12, BUF_Volatile, CreateInfo);
+		FBufferRHIRef VertexBufferRHI = RHICreateVertexBuffer(sizeof(FVector4) * 12, BUF_Volatile, CreateInfo);
 		void* VoidPtr = RHILockBuffer(VertexBufferRHI, 0, sizeof(FVector4) * 12, RLM_WriteOnly);
 
 		// Generate the vertices used
@@ -940,7 +940,7 @@ void FProjectedShadowInfo::SetupProjectionStencilMask(
 		VertexShader->SetParameters(RHICmdList, *View, this);
 
 		FRHIResourceCreateInfo CreateInfo(TEXT("FProjectedShadowInfoStencilFrustum"));
-		FVertexBufferRHIRef VertexBufferRHI = RHICreateVertexBuffer(sizeof(FVector4) * FrustumVertices.Num(), BUF_Volatile, CreateInfo);
+		FBufferRHIRef VertexBufferRHI = RHICreateVertexBuffer(sizeof(FVector4) * FrustumVertices.Num(), BUF_Volatile, CreateInfo);
 		void* VoidPtr = RHILockBuffer(VertexBufferRHI, 0, sizeof(FVector4) * FrustumVertices.Num(), RLM_WriteOnly);
 		FPlatformMemory::Memcpy(VoidPtr, FrustumVertices.GetData(), sizeof(FVector4) * FrustumVertices.Num());
 		RHIUnlockBuffer(VertexBufferRHI);
@@ -1149,7 +1149,7 @@ void FProjectedShadowInfo::RenderProjection(
 		else
 		{
 			FRHIResourceCreateInfo CreateInfo(TEXT("FProjectedShadowInfoFrustum"));
-			FVertexBufferRHIRef VertexBufferRHI = RHICreateVertexBuffer(sizeof(FVector4) * FrustumVertices.Num(), BUF_Volatile, CreateInfo);
+			FBufferRHIRef VertexBufferRHI = RHICreateVertexBuffer(sizeof(FVector4) * FrustumVertices.Num(), BUF_Volatile, CreateInfo);
 			void* VoidPtr = RHILockBuffer(VertexBufferRHI, 0, sizeof(FVector4) * FrustumVertices.Num(), RLM_WriteOnly);
 			FPlatformMemory::Memcpy(VoidPtr, FrustumVertices.GetData(), sizeof(FVector4) * FrustumVertices.Num());
 			RHIUnlockBuffer(VertexBufferRHI);
diff --git a/Engine/Source/Runtime/Renderer/Private/Tests/RayTracingTestbed.cpp b/Engine/Source/Runtime/Renderer/Private/Tests/RayTracingTestbed.cpp
--- a/Engine/Source/Runtime/Renderer/Private/Tests/RayTracingTestbed.cpp
+++ b/Engine/Source/Runtime/Renderer/Private/Tests/RayTracingTestbed.cpp
@@ -25,7 +25,7 @@ bool RunRayTracingTestbed_RenderThread(const FString& Parameters)
 		return true;
 	}
 
-	FVertexBufferRHIRef VertexBuffer;
+	FBufferRHIRef VertexBuffer;
 
 	{
 		TResourceArray<FVector> PositionData;
@@ -40,7 +40,7 @@ bool RunRayTracingTestbed_RenderThread(const FString& Parameters)
 		VertexBuffer = RHICreateVertexBuffer(PositionData.GetResourceDataSize(), BUF_Static, CreateInfo);
 	}
 
-	FIndexBufferRHIRef IndexBuffer;
+	FBufferRHIRef IndexBuffer;
 
 	{
 		TResourceArray<uint16> IndexData;
@@ -57,7 +57,7 @@ bool RunRayTracingTestbed_RenderThread(const FString& Parameters)
 
 	static constexpr uint32 NumRays = 4;
 
-	FStructuredBufferRHIRef RayBuffer;
+	FBufferRHIRef RayBuffer;
 	FShaderResourceViewRHIRef RayBufferView;
 
 	{
@@ -75,7 +75,7 @@ bool RunRayTracingTestbed_RenderThread(const FString& Parameters)
 		RayBufferView = RHICreateShaderResourceView(RayBuffer);
 	}
 
-	FVertexBufferRHIRef OcclusionResultBuffer;
+	FBufferRHIRef OcclusionResultBuffer;
 	FUnorderedAccessViewRHIRef OcclusionResultBufferView;
 
 	{
@@ -84,7 +84,7 @@ bool RunRayTracingTestbed_RenderThread(const FString& Parameters)
 		OcclusionResultBufferView = RHICreateUnorderedAccessView(OcclusionResultBuffer, PF_R32_UINT);
 	}
 
-	FVertexBufferRHIRef IntersectionResultBuffer;
+	FBufferRHIRef IntersectionResultBuffer;
 	FUnorderedAccessViewRHIRef IntersectionResultBufferView;
 
 	{
diff --git a/Engine/Source/Runtime/Renderer/Private/VT/VirtualTextureFeedback.cpp b/Engine/Source/Runtime/Renderer/Private/VT/VirtualTextureFeedback.cpp
--- a/Engine/Source/Runtime/Renderer/Private/VT/VirtualTextureFeedback.cpp
+++ b/Engine/Source/Runtime/Renderer/Private/VT/VirtualTextureFeedback.cpp
@@ -180,7 +180,7 @@ void FVirtualTextureFeedback::ReleaseRHI()
 	Fences->ReleaseRHI();
 }
 
-void FVirtualTextureFeedback::TransferGPUToCPU(FRHICommandListImmediate& RHICmdList, FVertexBufferRHIRef const& Buffer, FVirtualTextureFeedbackBufferDesc const& Desc)
+void FVirtualTextureFeedback::TransferGPUToCPU(FRHICommandListImmediate& RHICmdList, FBufferRHIRef const& Buffer, FVirtualTextureFeedbackBufferDesc const& Desc)
 {
 	if (NumPending >= MaxTransfers)
 	{
diff --git a/Engine/Source/Runtime/Renderer/Private/VT/VirtualTextureFeedback.h b/Engine/Source/Runtime/Renderer/Private/VT/VirtualTextureFeedback.h
--- a/Engine/Source/Runtime/Renderer/Private/VT/VirtualTextureFeedback.h
+++ b/Engine/Source/Runtime/Renderer/Private/VT/VirtualTextureFeedback.h
@@ -19,7 +19,7 @@ public:
 	~FVirtualTextureFeedback();
 
 	/** Commit a RHIBuffer to be transferred for later CPU analysis. */
-	void TransferGPUToCPU(FRHICommandListImmediate& RHICmdList, FVertexBufferRHIRef const& Buffer, FVirtualTextureFeedbackBufferDesc const& Desc);
+	void TransferGPUToCPU(FRHICommandListImmediate& RHICmdList, FBufferRHIRef const& Buffer, FVirtualTextureFeedbackBufferDesc const& Desc);
 	void TransferGPUToCPU(FRDGBuilder& GraphBuilder, class FRDGBuffer* Buffer, FVirtualTextureFeedbackBufferDesc const& Desc);
 
 	/** Returns true if there are any pending transfer results that are ready so that we can call Map(). */
diff --git a/Engine/Source/Runtime/Renderer/Private/VT/VirtualTextureFeedbackBuffer.cpp b/Engine/Source/Runtime/Renderer/Private/VT/VirtualTextureFeedbackBuffer.cpp
--- a/Engine/Source/Runtime/Renderer/Private/VT/VirtualTextureFeedbackBuffer.cpp
+++ b/Engine/Source/Runtime/Renderer/Private/VT/VirtualTextureFeedbackBuffer.cpp
@@ -119,7 +119,7 @@ void FVirtualTextureFeedbackBuffer::ReleaseRHI()
 	UAV = nullptr;
 }
 
-void SubmitVirtualTextureFeedbackBuffer(FRHICommandListImmediate& RHICmdList, FVertexBufferRHIRef const& InBuffer, FVirtualTextureFeedbackBufferDesc const& InDesc)
+void SubmitVirtualTextureFeedbackBuffer(FRHICommandListImmediate& RHICmdList, FBufferRHIRef const& InBuffer, FVirtualTextureFeedbackBufferDesc const& InDesc)
 {
 	GVirtualTextureFeedback.TransferGPUToCPU(RHICmdList, InBuffer, InDesc);
 }
diff --git a/Engine/Source/Runtime/Renderer/Private/VT/VirtualTextureSpace.h b/Engine/Source/Runtime/Renderer/Private/VT/VirtualTextureSpace.h
--- a/Engine/Source/Runtime/Renderer/Private/VT/VirtualTextureSpace.h
+++ b/Engine/Source/Runtime/Renderer/Private/VT/VirtualTextureSpace.h
@@ -117,7 +117,7 @@ private:
 
 	TArray<FPageTableUpdate> PageTableUpdates[VIRTUALTEXTURE_SPACE_MAXLAYERS];
 
-	FVertexBufferRHIRef UpdateBuffer;
+	FBufferRHIRef UpdateBuffer;
 	FShaderResourceViewRHIRef UpdateBufferSRV;
 
 	uint32 PageTableSize;
diff --git a/Engine/Source/Runtime/Renderer/Public/VT/VirtualTextureFeedbackBuffer.h b/Engine/Source/Runtime/Renderer/Public/VT/VirtualTextureFeedbackBuffer.h
--- a/Engine/Source/Runtime/Renderer/Public/VT/VirtualTextureFeedbackBuffer.h
+++ b/Engine/Source/Runtime/Renderer/Public/VT/VirtualTextureFeedbackBuffer.h
@@ -72,5 +72,5 @@ private:
  * Multiple buffers can be transferred per frame using this function.
  * The function can be called from the render thread only.
 */
-RENDERER_API void SubmitVirtualTextureFeedbackBuffer(class FRHICommandListImmediate& RHICmdList, FVertexBufferRHIRef const& InBuffer, FVirtualTextureFeedbackBufferDesc const& InDesc);
+RENDERER_API void SubmitVirtualTextureFeedbackBuffer(class FRHICommandListImmediate& RHICmdList, FBufferRHIRef const& InBuffer, FVirtualTextureFeedbackBufferDesc const& InDesc);
 RENDERER_API void SubmitVirtualTextureFeedbackBuffer(class FRDGBuilder& GraphBuilder, class FRDGBuffer* InBuffer, FVirtualTextureFeedbackBufferDesc const& InDesc);
\ No newline at end of file
