Fix assert in FVulkanDevice::FVulkanDevice when running debug linux editor

Hit when using amdgpu driver which also exposes software llvmpipe device.

from vulkaninfo:

        Devices: count = 2
                GPU id = 0 (llvmpipe (LLVM 11.0.1, 256 bits))
                Layer-Device Extensions: count = 3
                        VK_EXT_debug_marker     : extension revision 4
                        VK_EXT_tooling_info     : extension revision 1
                        VK_EXT_validation_cache : extension revision 1

                GPU id = 1 (AMD RADV VEGA10 (ACO))
                Layer-Device Extensions: count = 3
                        VK_EXT_debug_marker     : extension revision 4
                        VK_EXT_tooling_info     : extension revision 1
                        VK_EXT_validation_cache : extension revision 1

GPU0:
VkPhysicalDeviceProperties:
---------------------------
        apiVersion     = [REDACTED_PHONE])
        driverVersion  = 1 (0x0001)
        vendorID       = 0x10005
        deviceID       = 0x0000
        deviceType     = PHYSICAL_DEVICE_TYPE_CPU
        deviceName     = llvmpipe (LLVM 11.0.1, 256 bits)

GPU1:
-----
VkPhysicalDeviceProperties:
---------------------------
        apiVersion     = [REDACTED_PHONE])
        driverVersion  = 83898370 (0x5003002)
        vendorID       = 0x1002
        deviceID       = 0x687f
        deviceType     = PHYSICAL_DEVICE_TYPE_DISCRETE_GPU
        deviceName     = AMD RADV VEGA10 (ACO)

diff --git a/Engine/Source/Runtime/RHI/Public/RHIDefinitions.h b/Engine/Source/Runtime/RHI/Public/RHIDefinitions.h
--- a/Engine/Source/Runtime/RHI/Public/RHIDefinitions.h
+++ b/Engine/Source/Runtime/RHI/Public/RHIDefinitions.h
@@ -1343,6 +1343,7 @@ enum class EGpuVendorId
 	Arm			= 0x13B5, 
 	Qualcomm	= 0x5143,
 	Intel		= 0x8086,
+	Mesa		= 0x10005,
 };
 
 /** An enumeration of the different RHI reference types. */
@@ -1931,6 +1932,7 @@ inline EGpuVendorId RHIConvertToGpuVendorId(uint32 VendorId)
 		return EGpuVendorId::NotQueried;
 
 	case EGpuVendorId::Amd:
+	case EGpuVendorId::Mesa:
 	case EGpuVendorId::ImgTec:
 	case EGpuVendorId::Nvidia:
 	case EGpuVendorId::Arm:
