Fix shader platform DDPI data being broken on platforms with lower-cased filenames
 - Shader platform name is grabbed from the filename of the config file, which doesn't match the case if the platform always has lower case filenames.
 - Use case insensitive comparison instead.

diff --git a/Engine/Source/Runtime/RHI/Private/RHI.cpp b/Engine/Source/Runtime/RHI/Private/RHI.cpp
--- a/Engine/Source/Runtime/RHI/Private/RHI.cpp
+++ b/Engine/Source/Runtime/RHI/Private/RHI.cpp
@@ -1517,7 +1517,7 @@ void LexFromString(EShaderPlatform& Value, const TCHAR* String)
 
 	for (uint8 i = 0; i < (uint8)EShaderPlatform::SP_NumPlatforms; ++i)
 	{
-		if (LexToString((EShaderPlatform)i, false).Equals(String))
+		if (LexToString((EShaderPlatform)i, false).Equals(String, ESearchCase::IgnoreCase))
 		{
 			Value = (EShaderPlatform)i;
 			return;
