Deprecate old buffer types

diff --git a/Engine/Source/Runtime/RHI/Public/RHIResources.h b/Engine/Source/Runtime/RHI/Public/RHIResources.h
--- a/Engine/Source/Runtime/RHI/Public/RHIResources.h
+++ b/Engine/Source/Runtime/RHI/Public/RHIResources.h
@@ -877,9 +877,9 @@ private:
 	uint32 Stride;
 };
 
-typedef class FRHIBuffer FRHIIndexBuffer;
-typedef class FRHIBuffer FRHIVertexBuffer;
-typedef class FRHIBuffer FRHIStructuredBuffer;
+UE_DEPRECATED(5.0, "FRHIIndexBuffer is deprecated, please use FRHIBuffer.")      typedef class FRHIBuffer FRHIIndexBuffer;
+UE_DEPRECATED(5.0, "FRHIVertexBuffer is deprecated, please use FRHIBuffer.")     typedef class FRHIBuffer FRHIVertexBuffer;
+UE_DEPRECATED(5.0, "FRHIStructuredBuffer is deprecated, please use FRHIBuffer.") typedef class FRHIBuffer FRHIStructuredBuffer;
 
 //
 // Textures
@@ -1481,9 +1481,9 @@ typedef TRefCountPtr<FRHIComputeFence>	FComputeFenceRHIRef;
 typedef TRefCountPtr<FRHIBoundShaderState> FBoundShaderStateRHIRef;
 typedef TRefCountPtr<FRHIUniformBuffer> FUniformBufferRHIRef;
 typedef TRefCountPtr<FRHIBuffer> FBufferRHIRef;
-typedef FBufferRHIRef FIndexBufferRHIRef;
-typedef FBufferRHIRef FVertexBufferRHIRef;
-typedef FBufferRHIRef FStructuredBufferRHIRef;
+UE_DEPRECATED(5.0, "FIndexBufferRHIRef is deprecated, please use FBufferRHIRef.")      typedef FBufferRHIRef FIndexBufferRHIRef;
+UE_DEPRECATED(5.0, "FVertexBufferRHIRef is deprecated, please use FBufferRHIRef.")     typedef FBufferRHIRef FVertexBufferRHIRef;
+UE_DEPRECATED(5.0, "FStructuredBufferRHIRef is deprecated, please use FBufferRHIRef.") typedef FBufferRHIRef FStructuredBufferRHIRef;
 typedef TRefCountPtr<FRHITexture> FTextureRHIRef;
 typedef TRefCountPtr<FRHITexture2D> FTexture2DRHIRef;
 typedef TRefCountPtr<FRHITexture2DArray> FTexture2DArrayRHIRef;
