Increased MaxVertexElementCount from 16 to 17, and MaxVertexElementCount_NumBits from 4 to 5 to account for vertex factories (such as GPU skinning) that already exhaust the 16 limit and primitive ID stream pushes it over that when enabled.

diff --git a/Engine/Source/Runtime/RHI/Public/RHIDefinitions.h b/Engine/Source/Runtime/RHI/Public/RHIDefinitions.h
--- a/Engine/Source/Runtime/RHI/Public/RHIDefinitions.h
+++ b/Engine/Source/Runtime/RHI/Public/RHIDefinitions.h
@@ -694,8 +694,8 @@ enum
 /** The maximum number of vertex elements which can be used by a vertex declaration. */
 enum
 {
-	MaxVertexElementCount = 16,
-	MaxVertexElementCount_NumBits = 4,
+	MaxVertexElementCount = 17,
+	MaxVertexElementCount_NumBits = 5,
 };
 static_assert(MaxVertexElementCount <= (1 << MaxVertexElementCount_NumBits), "MaxVertexElementCount will not fit on MaxVertexElementCount_NumBits");
 
