Get rid of CreateAndLock functions

diff --git a/Engine/Source/Runtime/Engine/Private/RawIndexBuffer.cpp b/Engine/Source/Runtime/Engine/Private/RawIndexBuffer.cpp
--- a/Engine/Source/Runtime/Engine/Private/RawIndexBuffer.cpp
+++ b/Engine/Source/Runtime/Engine/Private/RawIndexBuffer.cpp
@@ -41,11 +41,11 @@ void FRawIndexBuffer::InitRHI()
 	{
 		// Create the index buffer.
 		FRHIResourceCreateInfo CreateInfo;
-		void* Buffer = nullptr;
-		IndexBufferRHI = RHICreateAndLockIndexBuffer(sizeof(uint16),Size,BUF_Static,CreateInfo, Buffer);
+		IndexBufferRHI = RHICreateBuffer(Size, BUF_Static | BUF_IndexBuffer, sizeof(uint16), ERHIAccess::VertexOrIndexBuffer, CreateInfo);
 
-		// Initialize the buffer.		
-		FMemory::Memcpy(Buffer,Indices.GetData(),Size);
+		// Initialize the buffer.
+		void* Buffer = RHILockBuffer(IndexBufferRHI, 0, Size, RLM_WriteOnly);
+		FMemory::Memcpy(Buffer, Indices.GetData(), Size);
 		RHIUnlockBuffer(IndexBufferRHI);
 	}
 }
@@ -101,10 +101,11 @@ void FRawIndexBuffer16or32::InitRHI()
 	{
 		// Create the index buffer.
 		FRHIResourceCreateInfo CreateInfo;
-		void* Buffer = nullptr;
-		IndexBufferRHI = RHICreateAndLockIndexBuffer(IndexStride,Size,BUF_Static,CreateInfo, Buffer);
-		
+		IndexBufferRHI = RHICreateBuffer(Size, BUF_Static | BUF_IndexBuffer, IndexStride, ERHIAccess::VertexOrIndexBuffer, CreateInfo);
+
 		// Initialize the buffer.		
+		void* Buffer = RHILockBuffer(IndexBufferRHI, 0, Size, RLM_WriteOnly);
+
 		if (b32Bit)
 		{
 			FMemory::Memcpy(Buffer, Indices.GetData(), Size);
diff --git a/Engine/Source/Runtime/RHI/Private/RHICommandList.cpp b/Engine/Source/Runtime/RHI/Private/RHICommandList.cpp
--- a/Engine/Source/Runtime/RHI/Private/RHICommandList.cpp
+++ b/Engine/Source/Runtime/RHI/Private/RHICommandList.cpp
@@ -1987,23 +1987,6 @@ void FRHICommandListBase::operator delete(void *RawMemory)
 	FMemory::Free(RawMemory);
 }	
 
-///////// Pass through functions that allow RHIs to optimize certain calls.
-
-FBufferRHIRef FDynamicRHI::CreateAndLockVertexBuffer_RenderThread(class FRHICommandListImmediate& RHICmdList, uint32 Size, uint32 Usage, ERHIAccess ResourceState, FRHIResourceCreateInfo& CreateInfo, void*& OutDataBuffer)
-{
-	FBufferRHIRef VertexBuffer = CreateBuffer_RenderThread(RHICmdList, Size, (EBufferUsageFlags)Usage | BUF_VertexBuffer, 0, ResourceState, CreateInfo);
-	OutDataBuffer = RHILockBuffer(RHICmdList, VertexBuffer, 0, Size, RLM_WriteOnly);
-
-	return VertexBuffer;
-}
-
-FBufferRHIRef FDynamicRHI::CreateAndLockIndexBuffer_RenderThread(class FRHICommandListImmediate& RHICmdList, uint32 Stride, uint32 Size, uint32 Usage, ERHIAccess ResourceState, FRHIResourceCreateInfo& CreateInfo, void*& OutDataBuffer)
-{
-	FBufferRHIRef IndexBuffer = CreateBuffer_RenderThread(RHICmdList, Size, (EBufferUsageFlags)Usage | BUF_IndexBuffer, Stride, ResourceState, CreateInfo);
-	OutDataBuffer = RHILockBuffer(RHICmdList, IndexBuffer, 0, Size, RLM_WriteOnly);
-	return IndexBuffer;
-}
-
 
 FBufferRHIRef FDynamicRHI::CreateBuffer_RenderThread(class FRHICommandListImmediate& RHICmdList, uint32 Size, EBufferUsageFlags Usage, uint32 Stride, ERHIAccess ResourceState, FRHIResourceCreateInfo& CreateInfo)
 {
diff --git a/Engine/Source/Runtime/RHI/Public/DynamicRHI.h b/Engine/Source/Runtime/RHI/Public/DynamicRHI.h
--- a/Engine/Source/Runtime/RHI/Public/DynamicRHI.h
+++ b/Engine/Source/Runtime/RHI/Public/DynamicRHI.h
@@ -1247,10 +1247,6 @@ public:
 	FORCEINLINE IRHICommandContextContainer* RHIGetCommandContextContainer(int32 Index, int32 Num, FRHIGPUMask GPUMask) { return RHIGetCommandContextContainer(Index, Num); }
 #endif
 
-	// Legacy helper functions
-	FBufferRHIRef CreateAndLockVertexBuffer_RenderThread(class FRHICommandListImmediate& RHICmdList, uint32 Size, uint32 Usage, ERHIAccess ResourceState, FRHIResourceCreateInfo& CreateInfo, void*& OutDataBuffer);
-	FBufferRHIRef CreateAndLockIndexBuffer_RenderThread(class FRHICommandListImmediate& RHICmdList, uint32 Stride, uint32 Size, uint32 Usage, ERHIAccess ResourceState, FRHIResourceCreateInfo& CreateInfo, void*& OutDataBuffer);
-
 	virtual FBufferRHIRef CreateBuffer_RenderThread(class FRHICommandListImmediate& RHICmdList, uint32 Size, EBufferUsageFlags Usage, uint32 Stride, ERHIAccess ResourceState, FRHIResourceCreateInfo& CreateInfo);
 	virtual FShaderResourceViewRHIRef CreateShaderResourceView_RenderThread(class FRHICommandListImmediate& RHICmdList, FRHIBuffer* Buffer, uint32 Stride, uint8 Format);
 	virtual FShaderResourceViewRHIRef CreateShaderResourceView_RenderThread(class FRHICommandListImmediate& RHICmdList, const FShaderResourceViewInitializer& Initializer);
diff --git a/Engine/Source/Runtime/RHI/Public/RHICommandList.h b/Engine/Source/Runtime/RHI/Public/RHICommandList.h
--- a/Engine/Source/Runtime/RHI/Public/RHICommandList.h
+++ b/Engine/Source/Runtime/RHI/Public/RHICommandList.h
@@ -3966,7 +3966,9 @@ public:
 	
 	FORCEINLINE FBufferRHIRef CreateAndLockIndexBuffer(uint32 Stride, uint32 Size, EBufferUsageFlags InUsage, ERHIAccess InResourceState, FRHIResourceCreateInfo& CreateInfo, void*& OutDataBuffer)
 	{
-		return GDynamicRHI->CreateAndLockIndexBuffer_RenderThread(*this, Stride, Size, InUsage, InResourceState, CreateInfo, OutDataBuffer);
+		FBufferRHIRef IndexBuffer = GDynamicRHI->CreateBuffer_RenderThread(*this, Size, (EBufferUsageFlags)InUsage | BUF_IndexBuffer, Stride, InResourceState, CreateInfo);
+		OutDataBuffer = GDynamicRHI->RHILockBuffer(*this, IndexBuffer, 0, Size, RLM_WriteOnly);
+		return IndexBuffer;
 	}
 
 	FORCEINLINE FBufferRHIRef CreateAndLockIndexBuffer(uint32 Stride, uint32 Size, uint32 InUsage, FRHIResourceCreateInfo& CreateInfo, void*& OutDataBuffer)
@@ -4000,7 +4002,9 @@ public:
 	
 	FORCEINLINE FBufferRHIRef CreateAndLockVertexBuffer(uint32 Size, EBufferUsageFlags InUsage, ERHIAccess InResourceState, FRHIResourceCreateInfo& CreateInfo, void*& OutDataBuffer)
 	{
-		return GDynamicRHI->CreateAndLockVertexBuffer_RenderThread(*this, Size, InUsage, InResourceState, CreateInfo, OutDataBuffer);
+		FBufferRHIRef VertexBuffer = GDynamicRHI->CreateBuffer_RenderThread(*this, Size, (EBufferUsageFlags)InUsage | BUF_VertexBuffer, 0, InResourceState, CreateInfo);
+		OutDataBuffer = GDynamicRHI->RHILockBuffer(*this, VertexBuffer, 0, Size, RLM_WriteOnly);
+		return VertexBuffer;
 	}
 
 	FORCEINLINE FBufferRHIRef CreateAndLockVertexBuffer(uint32 Size, uint32 InUsage, FRHIResourceCreateInfo& CreateInfo, void*& OutDataBuffer)
@@ -4991,6 +4995,11 @@ FORCEINLINE FBufferRHIRef RHICreateAndLockIndexBuffer(uint32 Stride, uint32 Size
 	return FRHICommandListExecutor::GetImmediateCommandList().CreateAndLockIndexBuffer(Stride, Size, InUsage, CreateInfo, OutDataBuffer);
 }
 
+FORCEINLINE FBufferRHIRef RHICreateBuffer(uint32 Size, EBufferUsageFlags Usage, uint32 Stride, ERHIAccess ResourceState, FRHIResourceCreateInfo& CreateInfo)
+{
+	return GDynamicRHI->CreateBuffer_RenderThread(FRHICommandListExecutor::GetImmediateCommandList(), Size, Usage, Stride, ResourceState, CreateInfo);
+}
+
 FORCEINLINE FBufferRHIRef RHICreateIndexBuffer(uint32 Stride, uint32 Size, uint32 InUsage, ERHIAccess InResourceState, FRHIResourceCreateInfo& CreateInfo)
 {
 	return GDynamicRHI->CreateBuffer_RenderThread(FRHICommandListExecutor::GetImmediateCommandList(), Size, (EBufferUsageFlags) InUsage | BUF_IndexBuffer, Stride, InResourceState, CreateInfo);
