Fix RHI Validation complaining about aliased copies to block compressed formats.
There are probably other copy rules to consider here, but this fixes the false positives coming from the GPU texture compression.

diff --git a/Engine/Source/Runtime/RHI/Public/RHIValidationUtils.h b/Engine/Source/Runtime/RHI/Public/RHIValidationUtils.h
--- a/Engine/Source/Runtime/RHI/Public/RHIValidationUtils.h
+++ b/Engine/Source/Runtime/RHI/Public/RHIValidationUtils.h
@@ -16,7 +16,16 @@ public:
 	{
 		check(SourceTexture);
 		check(DestTexture);
-		checkf(SourceTexture->GetFormat() == DestTexture->GetFormat(), TEXT("Some RHIs do not allow format conversion by the GPU for transfer operations!"));
+
+		const EPixelFormat SrcFormat = SourceTexture->GetFormat();
+		const EPixelFormat DstFormat = DestTexture->GetFormat();
+		const bool bIsSrcBlockCompressed = GPixelFormats[SrcFormat].BlockSizeX > 1;
+		const bool bIsDstBlockCompressed = GPixelFormats[DstFormat].BlockSizeX > 1;
+		const int32 SrcBlockBytes = GPixelFormats[SrcFormat].BlockBytes;
+		const int32 DstBlockBytes = GPixelFormats[DstFormat].BlockBytes;
+		const bool bValidCopyFormats = (SrcFormat == DstFormat) || (!bIsSrcBlockCompressed && bIsDstBlockCompressed && SrcBlockBytes == DstBlockBytes);
+
+		checkf(bValidCopyFormats, TEXT("Some RHIs do not support this format conversion by the GPU for transfer operations!"));
 
 		FIntVector SrcSize = SourceTexture->GetSizeXYZ();
 		FIntVector DestSize = DestTexture->GetSizeXYZ();
