Added enum to name conversion functions for buffer and texture flags.

#rnx

diff --git a/Engine/Source/Runtime/RHI/Private/RHI.cpp b/Engine/Source/Runtime/RHI/Private/RHI.cpp
--- a/Engine/Source/Runtime/RHI/Private/RHI.cpp
+++ b/Engine/Source/Runtime/RHI/Private/RHI.cpp
@@ -165,6 +165,79 @@ FString GetRHIPipelineName(ERHIPipeline Pipeline)
 	});
 }
 
+FString GetBufferUsageFlagsName(EBufferUsageFlags BufferUsage)
+{
+	return BuildEnumNameBitList<EBufferUsageFlags>(BufferUsage, [](EBufferUsageFlags BufferUsage)
+	{
+		switch (BufferUsage)
+		{
+		default: checkNoEntry(); // fall through
+		case BUF_None:					return TEXT("BUF_None");
+		case BUF_Static:				return TEXT("BUF_Static");
+		case BUF_Dynamic:				return TEXT("BUF_Dynamic");
+		case BUF_Volatile:				return TEXT("BUF_Volitile");
+		case BUF_UnorderedAccess:		return TEXT("BUF_UnorderedAccess");
+		case BUF_ByteAddressBuffer:		return TEXT("BUF_ByteAddressBuffer");
+		case BUF_SourceCopy:			return TEXT("BUF_SourceCopy");
+		case BUF_StreamOutput:			return TEXT("BUF_StreamOutput");
+		case BUF_DrawIndirect:			return TEXT("BUF_DrawIndirect");
+		case BUF_ShaderResource:		return TEXT("BUF_ShaderResource");
+		case BUF_KeepCPUAccessible:		return TEXT("BUF_KeepCPUAccessible");
+		case BUF_ZeroStride:			return TEXT("BUF_ZeroStride");
+		case BUF_FastVRAM:				return TEXT("BUF_FastVRAM");
+		case BUF_Transient:				return TEXT("BUF_Transient");
+		case BUF_Shared:				return TEXT("BUF_Shared");
+		case BUF_AccelerationStructure:	return TEXT("BUF_AccelerationStructure");
+		case BUF_VertexBuffer:			return TEXT("BUF_VertexBuffer");
+		case BUF_IndexBuffer:			return TEXT("BUF_IndexBuffer");
+		case BUF_StructuredBuffer:		return TEXT("BUF_StructuredBuffer");
+		}
+	});
+}
+
+FString GetTextureCreateFlagsName(ETextureCreateFlags TextureCreateFlags)
+{
+	return BuildEnumNameBitList<ETextureCreateFlags>(TextureCreateFlags, [](ETextureCreateFlags TextureCreateFlags)
+	{
+		switch (TextureCreateFlags)
+		{
+		default: checkNoEntry(); // fall through
+		case TexCreate_None:							return TEXT("TexCreate_None");
+		case TexCreate_RenderTargetable:				return TEXT("TexCreate_RenderTargetable");
+		case TexCreate_ResolveTargetable:				return TEXT("TexCreate_ResolveTargetable");
+		case TexCreate_DepthStencilTargetable:			return TEXT("TexCreate_DepthStencilTargetable");
+		case TexCreate_ShaderResource:					return TEXT("TexCreate_ShaderResource");
+		case TexCreate_SRGB:							return TEXT("TexCreate_SRGB");
+		case TexCreate_CPUWritable:						return TEXT("TexCreate_CPUWritable");
+		case TexCreate_NoTiling:						return TEXT("TexCreate_NoTiling");
+		case TexCreate_VideoDecode:						return TEXT("TexCreate_VideoDecode");
+		case TexCreate_Dynamic:							return TEXT("TexCreate_Dynamic");
+		case TexCreate_InputAttachmentRead:				return TEXT("TexCreate_InputAttachmentRead");
+		case TexCreate_Foveation:						return TEXT("(TexCreate_Foveation | TexCreate_3DTiling");
+		case TexCreate_Memoryless:						return TEXT("TexCreate_Memoryless");
+		case TexCreate_GenerateMipCapable:				return TEXT("TexCreate_GenerateMipCapable");
+		case TexCreate_FastVRAMPartialAlloc:			return TEXT("TexCreate_FastVRAMPartialAlloc");
+		case TexCreate_DisableSRVCreation:				return TEXT("TexCreate_DisableSRVCreation");
+		case TexCreate_DisableDCC:						return TEXT("TexCreate_DisableDCC");
+		case TexCreate_UAV:								return TEXT("TexCreate_UAV");
+		case TexCreate_Presentable:						return TEXT("TexCreate_Presentable");
+		case TexCreate_CPUReadback:						return TEXT("TexCreate_CPUReadback");
+		case TexCreate_OfflineProcessed:				return TEXT("TexCreate_OfflineProcessed");
+		case TexCreate_FastVRAM:						return TEXT("TexCreate_FastVRAM");
+		case TexCreate_HideInVisualizeTexture:			return TEXT("TexCreate_HideInVisualizeTexture");
+		case TexCreate_Virtual:							return TEXT("TexCreate_Virtual");
+		case TexCreate_TargetArraySlicesIndependently:	return TEXT("TexCreate_TargetArraySlicesIndependently");
+		case TexCreate_Shared:							return TEXT("TexCreate_Shared");
+		case TexCreate_NoFastClear:						return TEXT("TexCreate_NoFastClear");
+		case TexCreate_DepthStencilResolveTarget:		return TEXT("TexCreate_DepthStencilResolveTarget");
+		case TexCreate_Streamable:						return TEXT("TexCreate_Streamable");
+		case TexCreate_NoFastClearFinalize:				return TEXT("TexCreate_NoFastClearFinalize");
+		case TexCreate_AFRManual:						return TEXT("TexCreate_AFRManual");
+		case TexCreate_ReduceMemoryWithTilingMode:		return TEXT("TexCreate_ReduceMemoryWithTilingMode");
+		case TexCreate_Transient:						return TEXT("TexCreate_Transient");
+		}
+	});
+}
 
 #if STATS
 #include "Stats/StatsData.h"
diff --git a/Engine/Source/Runtime/RHI/Public/RHI.h b/Engine/Source/Runtime/RHI/Public/RHI.h
--- a/Engine/Source/Runtime/RHI/Public/RHI.h
+++ b/Engine/Source/Runtime/RHI/Public/RHI.h
@@ -725,6 +725,11 @@ extern RHI_API void GetShadingPathName(ERHIShadingPath::Type InShadingPath, FStr
 /** Creates an FName for the given shading path. */
 extern RHI_API void GetShadingPathName(ERHIShadingPath::Type InShadingPath, FName& OutName);
 
+/** Returns a string of friendly name bits for the buffer usage flags enum. */
+extern RHI_API FString GetBufferUsageFlagsName(EBufferUsageFlags BufferUsage);
+
+/** Returns a string of friendly name bits for the texture create flags enum. */
+extern RHI_API FString GetTextureCreateFlagsName(ETextureCreateFlags TextureCreateFlags);
 
 enum class ERHIPipeline : uint8
 {
