Mark scope as ParallelRendering, needed for FShaderMapResource::CreateShader().

diff --git a/Engine/Source/Runtime/RHI/Private/PipelineStateCache.cpp b/Engine/Source/Runtime/RHI/Private/PipelineStateCache.cpp
--- a/Engine/Source/Runtime/RHI/Private/PipelineStateCache.cpp
+++ b/Engine/Source/Runtime/RHI/Private/PipelineStateCache.cpp
@@ -884,6 +884,7 @@ public:
 	void DoTask(ENamedThreads::Type CurrentThread, const FGraphEventRef& MyCompletionGraphEvent)
 	{
 		LLM_SCOPE(ELLMTag::PSO);
+		FTaskTagScope Scope(ETaskTag::EParallelRenderingThread);
 
 		if (Pipeline->IsCompute())
 		{
@@ -1149,6 +1150,7 @@ public:
 
 	void DoTask(ENamedThreads::Type CurrentThread, const FGraphEventRef& MyCompletionGraphEvent)
 	{
+		FTaskTagScope Scope(ETaskTag::EParallelRenderingThread);
 		FRayTracingPipelineState* RayTracingPipeline = static_cast<FRayTracingPipelineState*>(Pipeline);
 		check(!RayTracingPipeline->RHIPipeline.IsValid());
 		RayTracingPipeline->RHIPipeline = RHICreateRayTracingPipelineState(Initializer);
