Deprecate CreateAndLock functions

diff --git a/Engine/Source/Runtime/RHI/Public/RHICommandList.h b/Engine/Source/Runtime/RHI/Public/RHICommandList.h
--- a/Engine/Source/Runtime/RHI/Public/RHICommandList.h
+++ b/Engine/Source/Runtime/RHI/Public/RHICommandList.h
@@ -4978,6 +4978,7 @@ FORCEINLINE FStagingBufferRHIRef RHICreateStagingBuffer()
 	return FRHICommandListExecutor::GetImmediateCommandList().CreateStagingBuffer();
 }
 
+UE_DEPRECATED(5.0, "Use RHICreateBuffer() and RHILockBuffer() instead.")
 FORCEINLINE FBufferRHIRef RHICreateAndLockIndexBuffer(uint32 Stride, uint32 Size, uint32 InUsage, FRHIResourceCreateInfo& CreateInfo, void*& OutDataBuffer)
 {
 	return FRHICommandListExecutor::GetImmediateCommandList().CreateAndLockIndexBuffer(Stride, Size, InUsage, CreateInfo, OutDataBuffer);
@@ -5026,6 +5027,7 @@ FORCEINLINE void RHIUnlockIndexBuffer(FRHIBuffer* IndexBuffer)
 	 FRHICommandListExecutor::GetImmediateCommandList().UnlockBuffer(IndexBuffer);
 }
 
+UE_DEPRECATED(5.0, "Use RHICreateBuffer() and RHILockBuffer() instead.")
 FORCEINLINE FBufferRHIRef RHICreateAndLockVertexBuffer(uint32 Size, uint32 InUsage, FRHIResourceCreateInfo& CreateInfo, void*& OutDataBuffer)
 {
 	return FRHICommandListExecutor::GetImmediateCommandList().CreateAndLockVertexBuffer(Size, InUsage, CreateInfo, OutDataBuffer);
