Always use feature level ES31 when running editor with -opengl commandline

diff --git a/Engine/Source/Runtime/OpenGLDrv/Public/OpenGL3.h b/Engine/Source/Runtime/OpenGLDrv/Public/OpenGL3.h
--- a/Engine/Source/Runtime/OpenGLDrv/Public/OpenGL3.h
+++ b/Engine/Source/Runtime/OpenGLDrv/Public/OpenGL3.h
@@ -754,10 +754,6 @@ struct FOpenGL3 : public FOpenGLBase
 
 	static FORCEINLINE EShaderPlatform GetShaderPlatform()
 	{
-		ERHIFeatureLevel::Type PreviewFeatureLevel;
-		bool bUsePreviewFeatureLevel = RHIGetPreviewFeatureLevel(PreviewFeatureLevel);
-		check(bUsePreviewFeatureLevel);
-		check(PreviewFeatureLevel == ERHIFeatureLevel::ES3_1);
 		return bAndroidGLESCompatibilityMode ? SP_OPENGL_ES3_1_ANDROID : SP_OPENGL_PCES3_1;
 	}
 
diff --git a/Engine/Source/Runtime/RHI/Private/Windows/WindowsDynamicRHI.cpp b/Engine/Source/Runtime/RHI/Private/Windows/WindowsDynamicRHI.cpp
--- a/Engine/Source/Runtime/RHI/Private/Windows/WindowsDynamicRHI.cpp
+++ b/Engine/Source/Runtime/RHI/Private/Windows/WindowsDynamicRHI.cpp
@@ -94,19 +94,8 @@ static IDynamicRHIModule* LoadDynamicRHIModule(ERHIFeatureLevel::Type& DesiredFe
 
 	bool bPreferD3D12 = ShouldPreferD3D12();
 	
-	// command line overrides
+	// OpenGL can only be used for mobile preview.
 	bool bForceOpenGL = FParse::Param(FCommandLine::Get(), TEXT("opengl"));
-	if (bForceOpenGL)
-	{
-		// OpenGL can only be used for mobile preview.
-		ERHIFeatureLevel::Type PreviewFeatureLevel;
-		bool bUsePreviewFeatureLevel = RHIGetPreviewFeatureLevel(PreviewFeatureLevel);
-		if (!bUsePreviewFeatureLevel)
-		{
-			FMessageDialog::Open(EAppMsgType::Ok, NSLOCTEXT("WindowsDynamicRHI", "OpenGLRemoved", "Warning: OpenGL is no longer supported for desktop platforms. The default RHI will be used."));
-			bForceOpenGL = false;
-		}
-	}
 
 	bool bForceSM5 = FParse::Param(FCommandLine::Get(), TEXT("sm5"));
 	bool bForceSM6 = FParse::Param(FCommandLine::Get(), TEXT("sm6"));
@@ -115,6 +104,7 @@ static IDynamicRHIModule* LoadDynamicRHIModule(ERHIFeatureLevel::Type& DesiredFe
 	bool bForceVulkan = FParse::Param(FCommandLine::Get(), TEXT("vulkan"));
 	bool bForceD3D11 = FParse::Param(FCommandLine::Get(), TEXT("d3d11")) || FParse::Param(FCommandLine::Get(), TEXT("dx11")) || ((bForceSM5 || (bPreferES31 && !bAllowD3D12FeatureLevelES31)) && !bForceVulkan && !bForceOpenGL);
 	bool bForceD3D12 = (FParse::Param(FCommandLine::Get(), TEXT("d3d12")) || FParse::Param(FCommandLine::Get(), TEXT("dx12"))) && (!bPreferES31 || bAllowD3D12FeatureLevelES31);
+
 	DesiredFeatureLevel = ERHIFeatureLevel::Num;
 	
 	if(!(bForceVulkan||bForceOpenGL||bForceD3D11||bForceD3D12))
@@ -219,6 +209,9 @@ static IDynamicRHIModule* LoadDynamicRHIModule(ERHIFeatureLevel::Type& DesiredFe
 			DynamicRHIModule = NULL;
 		}
 
+		// OpenGL can only be used for mobile preview
+		DesiredFeatureLevel = ERHIFeatureLevel::ES3_1;
+
 		LoadedRHIModuleName = OpenGLRHIModuleName;
 	}
 	else if (bForceVulkan)
