Partial merge of CL 15453755
 - Fix incorrect assert in RHI validation layer's UAV tracking.

diff --git a/Engine/Source/Runtime/RHI/Public/RHIValidationCommon.h b/Engine/Source/Runtime/RHI/Public/RHIValidationCommon.h
--- a/Engine/Source/Runtime/RHI/Public/RHIValidationCommon.h
+++ b/Engine/Source/Runtime/RHI/Public/RHIValidationCommon.h
@@ -827,7 +827,7 @@ namespace RHIValidation
 
 		inline void AssertUAV(FUnorderedAccessView* UAV, ERHIAccess Access, int32 Slot)
 		{
-			checkSlow(!(Access == ~ERHIAccess::UAVMask));
+			checkSlow(!(Access & ~ERHIAccess::UAVMask));
 			AssertUAV(UAV, Access == ERHIAccess::UAVGraphics ? EUAVMode::Graphics : EUAVMode::Compute, Slot);
 		}
 
