Attempt to make blacklisted driver message less confusing without using the hammer of blacklisting all latest drivers.

diff --git a/Engine/Source/Runtime/RHI/Private/DynamicRHI.cpp b/Engine/Source/Runtime/RHI/Private/DynamicRHI.cpp
--- a/Engine/Source/Runtime/RHI/Private/DynamicRHI.cpp
+++ b/Engine/Source/Runtime/RHI/Private/DynamicRHI.cpp
@@ -174,7 +174,7 @@ static void RHIDetectAndWarnOfBadDrivers(bool bHasEditorToken)
 			}
 			else
 			{
-				LocalizedMsg = FText::Format(NSLOCTEXT("MessageDialog", "VideoCardDriverIssueReport","The installed version of the {Vendor} graphics driver has known issues.\nPlease update to the latest driver version.\n\n{AdapterName}\nInstalled: {InstalledVer}\nRecommended: {RecommendedVer}"),Args);
+				LocalizedMsg = FText::Format(NSLOCTEXT("MessageDialog", "VideoCardDriverIssueReport","The installed version of the {Vendor} graphics driver has known issues.\nPlease install either the latest or the recommended driver version.\n\n{AdapterName}\nInstalled: {InstalledVer}\nRecommended: {RecommendedVer}"),Args);
 			}
 
 			FPlatformMisc::MessageBoxExt(EAppMsgType::Ok,
