Fix bugs in RHI validation layer
 - Stencil SRVs didn't have their plane index set appropriately, causing SRVGraphics/SRVCompute access to be logged against the depth plane instead, resulting in spurious validation errors.
 - Filter for unique resources in RDG's handling of Begin/EndUAVOverlap. The tracking for overlaps is kept on the underlying resource, not the view, meaning if two views to the same resource are passed to Begin/EndUAVOverlap, it will trigger the Begin/End mismatch validation.

Add support for enabling RHI validation logging on the command line
 - Use -RHIValidationLog=" ... resource names ... ", where the list of names are comma separated.
 - Replaces the hardcoded GAutoLogResourceNames array in RHIValidation.cpp.

Minor change: switch ordering of args in SpecificUAVOverlap to match the other state validation functions.

diff --git a/Engine/Source/Runtime/RHI/Private/RHIValidation.cpp b/Engine/Source/Runtime/RHI/Private/RHIValidation.cpp
--- a/Engine/Source/Runtime/RHI/Private/RHIValidation.cpp
+++ b/Engine/Source/Runtime/RHI/Private/RHIValidation.cpp
@@ -11,6 +11,18 @@
 
 #if ENABLE_RHI_VALIDATION
 
+bool GRHIValidationEnabled = false;
+
+// When set to 1, callstack for each uniform buffer allocation will be tracked 
+// (slow and leaks memory, but can be handy to find the location where an invalid
+// allocation has been made)
+#define CAPTURE_UNIFORMBUFFER_ALLOCATION_BACKTRACES 0
+
+// When set to 1, logs resource transitions on all unnamed resources, useful for
+// tracking down missing barriers when "-RHIValidationLog" cannot be used.
+// Don't leave this enabled. Log backtraces are leaked.
+#define LOG_UNNAMED_RESOURCES 0
+
 namespace RHIValidation
 {
 	int32 GBreakOnTransitionError = 1;
@@ -21,32 +33,38 @@ namespace RHIValidation
 		TEXT(" 0: disabled;\n")
 		TEXT(" 1: break in the debugger if a validation error is encountered."),
 		ECVF_RenderThreadSafe);
-}
-
-bool GRHIValidationEnabled = false;
 
-// When set to 1, callstack for each uniform buffer allocation will be tracked 
-// (slow and leaks memory, but can be handy to find the location where an invalid
-// allocation has been made)
-#define CAPTURE_UNIFORMBUFFER_ALLOCATION_BACKTRACES 0
+	// Returns an array of resource names parsed from the "-RHIValidationLog" command line switch.
+	// RHI validation logging is automatically enabled for resources whose debug names match those in this list.
+	// Multiple values are comma separated, e.g. -RHIValidationLog="SceneDepthZ,GBufferA"
+	static TArray<FString> const& GetAutoLogResourceNames()
+	{
+		struct FInit
+		{
+			TArray<FString> Strings;
 
-// When set to 1, logs resource transitions on all unnamed resources, useful for
-// tracking down missing barriers when GAutoLogResourceNames cannot be used.
-// Don't leave this enabled. Log backtraces are leaked.
-#define LOG_UNNAMED_RESOURCES 0
+			FInit()
+			{
+				FString ResourceNames;
+				if (FParse::Value(FCommandLine::Get(), TEXT("-RHIValidationLog="), ResourceNames, false))
+				{
+					FString Left, Right;
+					while (ResourceNames.Split(TEXT(","), &Left, &Right))
+					{
+						Left.TrimStartAndEndInline();
+						Strings.Add(Left);
+						ResourceNames = Right;
+					}
 
-static const TCHAR* GAutoLogResourceNames[] =
-{
-	//
-	// Add resource names here to automatically enable barrier logging for those resources.
-	// e.g.
-	//  TEXT("SceneDepthZ"),
-	//
+					ResourceNames.TrimStartAndEndInline();
+					Strings.Add(ResourceNames);
+				}
+			}
+		} static Init;
 
-	// ----------------------
-	// Add names above this line
-	nullptr
-};
+		return Init.Strings;
+	}
+}
 
 TSet<uint32> FValidationRHI::SeenFailureHashes;
 FCriticalSection FValidationRHI::SeenFailureHashesMutex;
@@ -677,13 +695,9 @@ namespace RHIValidation
 			// does/doesn't match one in the AutoLogResourceNames array.
 			if (Name)
 			{
-				for (int32 NameIndex = 0; NameIndex < UE_ARRAY_COUNT(GAutoLogResourceNames); ++NameIndex)
+				for (FString const& Str : GetAutoLogResourceNames())
 				{
-					const TCHAR* Str = GAutoLogResourceNames[NameIndex];
-					if (!Str)
-						break;
-
-					if (FCString::Strcmp(Name, Str) == 0)
+					if (FCString::Stricmp(Name, *Str) == 0)
 					{
 						LoggingMode = ELoggingMode::Automatic;
 						return;
@@ -848,7 +862,7 @@ namespace RHIValidation
 		State.Current.Access = DecayResourceAccess(State.Current.Access, RequiredState.Access, bAllowAllUAVsOverlap || State.bExplicitAllowUAVOverlap);
 	}
 
-	void FSubresourceState::SpecificUAVOverlap(FResource* Resource, ERHIPipeline Pipeline, FSubresourceIndex const& SubresourceIndex, bool bAllow)
+	void FSubresourceState::SpecificUAVOverlap(FResource* Resource, FSubresourceIndex const& SubresourceIndex, ERHIPipeline Pipeline, bool bAllow)
 	{
 		if (Resource->LoggingMode != ELoggingMode::None
 #if LOG_UNNAMED_RESOURCES
@@ -999,8 +1013,8 @@ namespace RHIValidation
 			{
 				State.SpecificUAVOverlap(
 					Data_SpecificUAVOverlap.Identity.Resource,
-					Pipeline,
 					SubresourceIndex,
+					Pipeline,
 					Data_SpecificUAVOverlap.bAllow);
 			});
 			Data_SpecificUAVOverlap.Identity.Resource->ReleaseOpRef();
diff --git a/Engine/Source/Runtime/RHI/Public/RHIValidation.h b/Engine/Source/Runtime/RHI/Public/RHIValidation.h
--- a/Engine/Source/Runtime/RHI/Public/RHIValidation.h
+++ b/Engine/Source/Runtime/RHI/Public/RHIValidation.h
@@ -1516,7 +1516,11 @@ public:
 	virtual FShaderResourceViewRHIRef RHICreateShaderResourceView_RenderThread(class FRHICommandListImmediate& RHICmdList, FRHITexture* Texture, const FRHITextureSRVCreateInfo& CreateInfo) override final
 	{
 		FShaderResourceViewRHIRef SRV = RHI->RHICreateShaderResourceView_RenderThread(RHICmdList, Texture, CreateInfo);
-		SRV->ViewIdentity = Texture->GetViewIdentity(CreateInfo.MipLevel, CreateInfo.NumMipLevels, CreateInfo.FirstArraySlice, CreateInfo.NumArraySlices, uint32(RHIValidation::EResourcePlane::Common), 1);
+		RHIValidation::EResourcePlane Plane = CreateInfo.Format == PF_X24_G8
+			? RHIValidation::EResourcePlane::Stencil
+			: RHIValidation::EResourcePlane::Common;
+
+		SRV->ViewIdentity = Texture->GetViewIdentity(CreateInfo.MipLevel, CreateInfo.NumMipLevels, CreateInfo.FirstArraySlice, CreateInfo.NumArraySlices, uint32(Plane), 1);
 		return SRV;
 	}
 
diff --git a/Engine/Source/Runtime/RHI/Public/RHIValidationCommon.h b/Engine/Source/Runtime/RHI/Public/RHIValidationCommon.h
--- a/Engine/Source/Runtime/RHI/Public/RHIValidationCommon.h
+++ b/Engine/Source/Runtime/RHI/Public/RHIValidationCommon.h
@@ -174,7 +174,7 @@ namespace RHIValidation
 		void BeginTransition   (FResource* Resource, FSubresourceIndex const& SubresourceIndex, const FState& CurrentStateFromRHI, const FState& TargetState, EResourceTransitionFlags NewFlags, ERHIPipeline Pipeline, void* CreateTrace);
 		void EndTransition     (FResource* Resource, FSubresourceIndex const& SubresourceIndex, const FState& CurrentStateFromRHI, const FState& TargetState, ERHIPipeline Pipeline, void* CreateTrace);
 		void Assert            (FResource* Resource, FSubresourceIndex const& SubresourceIndex, const FState& RequiredState, bool bAllowAllUAVsOverlap);
-		void SpecificUAVOverlap(FResource* Resource, ERHIPipeline Pipeline, FSubresourceIndex const& SubresourceIndex, bool bAllow);
+		void SpecificUAVOverlap(FResource* Resource, FSubresourceIndex const& SubresourceIndex, ERHIPipeline Pipeline, bool bAllow);
 		void* Log              (FResource* Resource, FSubresourceIndex const& SubresourceIndex, void* CreateTrace, const TCHAR* TracePrefix, const TCHAR* Type, const TCHAR* LogStr);
 	};
 
diff --git a/Engine/Source/Runtime/RenderCore/Private/RenderGraphBuilder.cpp b/Engine/Source/Runtime/RenderCore/Private/RenderGraphBuilder.cpp
--- a/Engine/Source/Runtime/RenderCore/Private/RenderGraphBuilder.cpp
+++ b/Engine/Source/Runtime/RenderCore/Private/RenderGraphBuilder.cpp
@@ -8,42 +8,64 @@
 #include "VisualizeTexture.h"
 #include "ProfilingDebugging/CsvProfiler.h"
 
-template <typename FunctionType>
-inline void EnumeratePassUAVs(const FRDGPass* Pass, FunctionType Function)
+#if ENABLE_RHI_VALIDATION
+
+inline void GatherPassUAVsForOverlapValidation(const FRDGPass* Pass, TArray<FRHIUnorderedAccessView*, TInlineAllocator<MaxSimultaneousUAVs, FRDGArrayAllocator>>& OutUAVs)
 {
+	// RHI validation tracking of Begin/EndUAVOverlaps happens on the underlying resource, so we need to be careful about not
+	// passing multiple UAVs that refer to the same resource, otherwise we get double-Begin and double-End validation errors.
+	// Filter UAVs to only those with unique parent resources.
+	TArray<FRDGParentResource*, TInlineAllocator<MaxSimultaneousUAVs, FRDGArrayAllocator>> UniqueParents;
 	Pass->GetParameters().Enumerate([&](FRDGParameter Parameter)
 	{
 		if (Parameter.IsUAV())
 		{
 			if (FRDGUnorderedAccessViewRef UAV = Parameter.GetAsUAV())
 			{
-				Function(UAV->GetRHI());
+				FRDGParentResource* Parent = UAV->GetParent();
+
+				// Check if we've already seen this parent.
+				bool bFound = false;
+				for (int32 Index = 0; !bFound && Index < UniqueParents.Num(); ++Index)
+				{
+					bFound = UniqueParents[Index] == Parent;
+				}
+
+				if (!bFound)
+				{
+					UniqueParents.Add(Parent);
+					OutUAVs.Add(UAV->GetRHI());
+				}
 			}
 		}
 	});
 }
 
+#endif
+
 inline void BeginUAVOverlap(const FRDGPass* Pass, FRHIComputeCommandList& RHICmdList)
 {
 #if ENABLE_RHI_VALIDATION
-	TArray<FRHIUnorderedAccessView*, TInlineAllocator<8, FRDGArrayAllocator>> UAVs;
-	EnumeratePassUAVs(Pass, [&](FRHIUnorderedAccessView* UAV)
+	TArray<FRHIUnorderedAccessView*, TInlineAllocator<MaxSimultaneousUAVs, FRDGArrayAllocator>> UAVs;
+	GatherPassUAVsForOverlapValidation(Pass, UAVs);
+
+	if (UAVs.Num())
 	{
-		UAVs.Add(UAV);
-	});
-	RHICmdList.BeginUAVOverlap(UAVs);
+		RHICmdList.BeginUAVOverlap(UAVs);
+	}
 #endif
 }
 
 inline void EndUAVOverlap(const FRDGPass* Pass, FRHIComputeCommandList& RHICmdList)
 {
 #if ENABLE_RHI_VALIDATION
-	TArray<FRHIUnorderedAccessView*, TInlineAllocator<8, FRDGArrayAllocator>> UAVs;
-	EnumeratePassUAVs(Pass, [&](FRHIUnorderedAccessView* UAV)
+	TArray<FRHIUnorderedAccessView*, TInlineAllocator<MaxSimultaneousUAVs, FRDGArrayAllocator>> UAVs;
+	GatherPassUAVsForOverlapValidation(Pass, UAVs);
+
+	if (UAVs.Num())
 	{
-		UAVs.Add(UAV);
-	});
-	RHICmdList.EndUAVOverlap(UAVs);
+		RHICmdList.EndUAVOverlap(UAVs);
+	}
 #endif
 }
 
