Implements GRHIDefaultMSAASampleOffsets

diff --git a/Engine/Source/Runtime/RHI/Private/RHI.cpp b/Engine/Source/Runtime/RHI/Private/RHI.cpp
--- a/Engine/Source/Runtime/RHI/Private/RHI.cpp
+++ b/Engine/Source/Runtime/RHI/Private/RHI.cpp
@@ -1739,6 +1739,62 @@ void FGenericDataDrivenShaderPlatformInfo::Initialize()
 	}
 }
 
+//
+//	MSAA sample offsets.
+//
+// https://docs.microsoft.com/en-us/windows/win32/api/d3d11/ne-d3d11-d3d11_standard_multisample_quality_levels
+FVector2D GRHIDefaultMSAASampleOffsets[1 + 2 + 4 + 8 + 16] = {
+	// MSAA x1
+	FVector2D(+0.0f / 8.0f, +0.0f / 8.0f),
+
+	// MSAA x2
+	FVector2D(+4.0f / 8.0f, +4.0f / 8.0f),
+	FVector2D(-4.0f / 8.0f, -4.0f / 8.0f),
+
+	// MSAA x4
+	FVector2D(-2.0f / 8.0f, -6.0f / 8.0f),
+	FVector2D(+6.0f / 8.0f, -2.0f / 8.0f),
+	FVector2D(-6.0f / 8.0f, +2.0f / 8.0f),
+	FVector2D(+2.0f / 8.0f, +6.0f / 8.0f),
+
+	// MSAA x8
+	FVector2D(+1.0f / 8.0f, -3.0f / 8.0f),
+	FVector2D(-1.0f / 8.0f, +3.0f / 8.0f),
+	FVector2D(+5.0f / 8.0f, +1.0f / 8.0f),
+	FVector2D(-3.0f / 8.0f, -5.0f / 8.0f),
+	FVector2D(-5.0f / 8.0f, +5.0f / 8.0f),
+	FVector2D(-7.0f / 8.0f, -1.0f / 8.0f),
+	FVector2D(+3.0f / 8.0f, +7.0f / 8.0f),
+	FVector2D(+7.0f / 8.0f, -7.0f / 8.0f),
+
+	// MSAA x16
+	FVector2D(+1.0f / 8.0f, +1.0f / 8.0f),
+	FVector2D(-1.0f / 8.0f, -3.0f / 8.0f),
+	FVector2D(-3.0f / 8.0f, +2.0f / 8.0f),
+	FVector2D(+4.0f / 8.0f, -1.0f / 8.0f),
+	FVector2D(-5.0f / 8.0f, -2.0f / 8.0f),
+	FVector2D(+2.0f / 8.0f, +5.0f / 8.0f),
+	FVector2D(+5.0f / 8.0f, +3.0f / 8.0f),
+	FVector2D(+3.0f / 8.0f, -5.0f / 8.0f),
+	FVector2D(-2.0f / 8.0f, +6.0f / 8.0f),
+	FVector2D(+0.0f / 8.0f, -7.0f / 8.0f),
+	FVector2D(-4.0f / 8.0f, -6.0f / 8.0f),
+	FVector2D(-6.0f / 8.0f, +4.0f / 8.0f),
+	FVector2D(-8.0f / 8.0f, +0.0f / 8.0f),
+	FVector2D(+7.0f / 8.0f, -4.0f / 8.0f),
+	FVector2D(+6.0f / 8.0f, +7.0f / 8.0f),
+	FVector2D(-7.0f / 8.0f, -8.0f / 8.0f),
+};
+
+int32 CalculateMSAASampleArrayIndex(int32 NumSamples, int32 SampleIndex)
+{
+	check(NumSamples > 0 && NumSamples <= 16);
+	check(FMath::IsPowerOfTwo(NumSamples));
+	check(SampleIndex < NumSamples);
+
+	return NumSamples - 1 + SampleIndex;
+}
+
 //
 //	Pixel format information.
 //
diff --git a/Engine/Source/Runtime/RHI/Public/RHI.h b/Engine/Source/Runtime/RHI/Public/RHI.h
--- a/Engine/Source/Runtime/RHI/Public/RHI.h
+++ b/Engine/Source/Runtime/RHI/Public/RHI.h
@@ -645,6 +645,18 @@ extern RHI_API bool GRHISupportsStencilRefFromPixelShader;
 /** Whether current RHI supports overestimated conservative rasterization. */
 extern RHI_API bool GRHISupportsConservativeRasterization;
 
+/** Tables of all MSAA sample offset for all MSAA supported. Use GetMSAASampleOffsets() to read it. */
+extern RHI_API FVector2D GRHIDefaultMSAASampleOffsets[1 + 2 + 4 + 8 + 16];
+
+// Calculate the index of the sample in GRHIDefaultMSAASampleOffsets
+extern RHI_API int32 CalculateMSAASampleArrayIndex(int32 NumSamples, int32 SampleIndex);
+
+// Gets the MSAA sample's offset from the center of the pixel coordinate.
+inline FVector2D GetMSAASampleOffsets(int32 NumSamples, int32 SampleIndex)
+{
+	return GRHIDefaultMSAASampleOffsets[CalculateMSAASampleArrayIndex(NumSamples, SampleIndex)];
+}
+
 /** Information about a pixel format. */
 struct FPixelFormatInfo
 {
