UE5EA - Update missing desktop Vulkan to not using DXC
- Removed a couple of unused/deprecated functions

diff --git a/Engine/Source/Runtime/RHI/Public/RHIDefinitions.h b/Engine/Source/Runtime/RHI/Public/RHIDefinitions.h
--- a/Engine/Source/Runtime/RHI/Public/RHIDefinitions.h
+++ b/Engine/Source/Runtime/RHI/Public/RHIDefinitions.h
@@ -1660,15 +1660,6 @@ inline bool IsSwitchPlatform(const FStaticShaderPlatform Platform)
 		|| FDataDrivenShaderPlatformInfo::GetIsLanguageNintendo(Platform);
 }
 
-UE_DEPRECATED(4.27, "IsPS4Platform() is deprecated; please use DataDrivenShaderPlatformInfo instead.") 
-inline bool IsPS4Platform(const FStaticShaderPlatform Platform)
-{
-	PRAGMA_DISABLE_DEPRECATION_WARNINGS
-	return Platform == SP_PS4_REMOVED
-		|| FDataDrivenShaderPlatformInfo::GetIsLanguageSony(Platform);
-	PRAGMA_ENABLE_DEPRECATION_WARNINGS
-}
-
 inline bool IsVulkanPlatform(const FStaticShaderPlatform Platform)
 {
 	return Platform == SP_VULKAN_SM5 || Platform == SP_VULKAN_SM5_LUMIN || Platform == SP_VULKAN_PCES3_1 
@@ -1698,25 +1689,6 @@ inline bool IsVulkanMobilePlatform(const FStaticShaderPlatform Platform)
 		|| (FDataDrivenShaderPlatformInfo::GetIsLanguageVulkan(Platform) && FDataDrivenShaderPlatformInfo::GetIsMobile(Platform));
 }
 
-UE_DEPRECATED(4.27, "IsD3DPlatform(bIncludeXboxOne) is deprecated; please use IsD3DPlatform() and DataDrivenShaderPlatformInfo instead.") 
-inline bool IsD3DPlatform(const FStaticShaderPlatform Platform, bool bIncludeXboxOne)
-{
-	switch (Platform)
-	{
-	case SP_PCD3D_SM5:
-	case SP_PCD3D_ES3_1:
-		return true;
-PRAGMA_DISABLE_DEPRECATION_WARNINGS
-	case SP_XBOXONE_D3D12_REMOVED:
-		return bIncludeXboxOne;
-PRAGMA_ENABLE_DEPRECATION_WARNINGS
-	default:
-		return FDataDrivenShaderPlatformInfo::GetIsLanguageD3D(Platform);
-	}
-
-	return false;
-}
-
 inline bool IsD3DPlatform(const FStaticShaderPlatform Platform)
 {
 	switch (Platform)
@@ -1737,18 +1709,7 @@ PRAGMA_ENABLE_DEPRECATION_WARNINGS
 
 inline bool IsHlslccShaderPlatform(const FStaticShaderPlatform Platform)
 {
-	return IsMetalPlatform(Platform) || IsVulkanPlatform(Platform) || IsSwitchPlatform(Platform) || IsOpenGLPlatform(Platform) || FDataDrivenShaderPlatformInfo::GetIsHlslcc(Platform);
-}
-
-UE_DEPRECATED(4.27, "Removed; please don't use.") 
-inline bool IsDeprecatedShaderPlatform(const FStaticShaderPlatform ShaderPlatform)
-{
-	PRAGMA_DISABLE_DEPRECATION_WARNINGS
-	return ShaderPlatform == SP_OPENGL_SM5_REMOVED || ShaderPlatform == SP_PCD3D_SM4_REMOVED || ShaderPlatform == SP_OPENGL_ES2_IOS_REMOVED ||
-		ShaderPlatform == SP_PCD3D_ES2_REMOVED || ShaderPlatform == SP_METAL_MACES2_REMOVED || ShaderPlatform == SP_OPENGL_PCES2_REMOVED ||
-		ShaderPlatform == SP_OPENGL_ES2_ANDROID_REMOVED || ShaderPlatform == SP_OPENGL_ES2_WEBGL_REMOVED ||
-		ShaderPlatform == SP_VULKAN_SM4_REMOVED || ShaderPlatform == SP_OPENGL_SM4_REMOVED || ShaderPlatform == SP_OPENGL_ES31_EXT_REMOVED;
-	PRAGMA_ENABLE_DEPRECATION_WARNINGS
+	return IsMetalPlatform(Platform) || IsVulkanMobilePlatform(Platform) || IsSwitchPlatform(Platform) || IsOpenGLPlatform(Platform) || FDataDrivenShaderPlatformInfo::GetIsHlslcc(Platform);
 }
 
 inline FStaticFeatureLevel GetMaxSupportedFeatureLevel(const FStaticShaderPlatform InShaderPlatform)
