Remove FRHITextureReference from platform RHI implementations.
 - FRHITextureReference is now implemented entirely within the base RHI module, and platform specific code from each RHI has been removed.
 - This also unifies how each platform deals with nullptr texture references. Some used GWhiteTexture, others used GBlackTexture, others probably just crash. Now, nullptr references will refer to GBlackTexture.
 - New implementation also handles RHI validation properly. Old code would crash when a nullptr reference is used in a draw.

diff --git a/Engine/Source/Runtime/Apple/MetalRHI/Private/MetalDynamicRHI.h b/Engine/Source/Runtime/Apple/MetalRHI/Private/MetalDynamicRHI.h
--- a/Engine/Source/Runtime/Apple/MetalRHI/Private/MetalDynamicRHI.h
+++ b/Engine/Source/Runtime/Apple/MetalRHI/Private/MetalDynamicRHI.h
@@ -75,7 +75,6 @@ public:
 	virtual uint64 RHIGetMinimumAlignmentForBufferBackedSRV(EPixelFormat Format) final override;
 	virtual void RHIGetTextureMemoryStats(FTextureMemoryStats& OutStats) final override;
 	virtual bool RHIGetTextureMemoryVisualizeData(FColor* TextureData,int32 SizeX,int32 SizeY,int32 Pitch,int32 PixelSize) final override;
-	virtual FTextureReferenceRHIRef RHICreateTextureReference(FLastRenderTimeContainer* LastRenderTime) final override;
 	virtual FTexture2DRHIRef RHICreateTexture2D(uint32 SizeX, uint32 SizeY, uint8 Format, uint32 NumMips, uint32 NumSamples, ETextureCreateFlags Flags, ERHIAccess InResourceState, FRHIResourceCreateInfo& CreateInfo) final override;
 	virtual FTexture2DRHIRef RHICreateTextureExternal2D(uint32 SizeX, uint32 SizeY, uint8 Format, uint32 NumMips, uint32 NumSamples, ETextureCreateFlags Flags, ERHIAccess InResourceState, FRHIResourceCreateInfo& Info) final override;
 	virtual FTexture2DRHIRef RHIAsyncCreateTexture2D(uint32 SizeX, uint32 SizeY, uint8 Format, uint32 NumMips, ETextureCreateFlags Flags, ERHIAccess InResourceState, void** InitialMipData, uint32 NumInitialMips) final override;
diff --git a/Engine/Source/Runtime/Apple/MetalRHI/Private/MetalTexture.cpp b/Engine/Source/Runtime/Apple/MetalRHI/Private/MetalTexture.cpp
--- a/Engine/Source/Runtime/Apple/MetalRHI/Private/MetalTexture.cpp
+++ b/Engine/Source/Runtime/Apple/MetalRHI/Private/MetalTexture.cpp
@@ -29,40 +29,6 @@ FAutoConsoleVariableRef CVarMetalForceIOSTexturesShared(
 														TEXT("If true, forces all textures to be Shared on iOS"),
 														ECVF_RenderThreadSafe);
 
-
-/** Texture reference class. */
-class FMetalTextureReference : public FRHITextureReference
-{
-public:
-	explicit FMetalTextureReference(FLastRenderTimeContainer* InLastRenderTime)
-	: FRHITextureReference(InLastRenderTime)
-	{}
-	
-	// IRefCountedObject interface.
-	virtual uint32 AddRef() const
-	{
-		return FRHIResource::AddRef();
-	}
-	virtual uint32 Release() const
-	{
-		return FRHIResource::Release();
-	}
-	virtual uint32 GetRefCount() const
-	{
-		return FRHIResource::GetRefCount();
-	}
-	
-	void SetReferencedTexture(FRHITexture* InTexture)
-	{
-		FRHITextureReference::SetReferencedTexture(InTexture);
-	}
-	
-	virtual void* GetTextureBaseRHI() override final
-	{
-		return GetMetalSurfaceFromRHITexture(GetReferencedTexture());
-	}
-};
-
 /** Given a pointer to a RHI texture that was created by the Metal RHI, returns a pointer to the FMetalTextureBase it encapsulates. */
 FMetalSurface* GetMetalSurfaceFromRHITexture(FRHITexture* Texture)
 {
@@ -2419,27 +2385,6 @@ FTextureCubeRHIRef FMetalDynamicRHI::RHICreateTextureCubeArray_RenderThread(clas
 }
 
 
-
-
-FTextureReferenceRHIRef FMetalDynamicRHI::RHICreateTextureReference(FLastRenderTimeContainer* InLastRenderTime)
-{
-	@autoreleasepool {
-		return new FMetalTextureReference(InLastRenderTime);
-	}
-}
-
-void FMetalRHICommandContext::RHIUpdateTextureReference(FRHITextureReference* TextureRefRHI, FRHITexture* NewTextureRHI)
-{
-	@autoreleasepool {
-		FMetalTextureReference* TextureRef = (FMetalTextureReference*)TextureRefRHI;
-		if (TextureRef)
-		{
-			TextureRef->SetReferencedTexture(NewTextureRHI);
-		}
-	}
-}
-
-
 void FMetalDynamicRHI::RHIBindDebugLabelName(FRHITexture* TextureRHI, const TCHAR* Name)
 {
 	@autoreleasepool {
diff --git a/Engine/Source/Runtime/Apple/MetalRHI/Public/MetalRHIContext.h b/Engine/Source/Runtime/Apple/MetalRHI/Public/MetalRHIContext.h
--- a/Engine/Source/Runtime/Apple/MetalRHI/Public/MetalRHIContext.h
+++ b/Engine/Source/Runtime/Apple/MetalRHI/Public/MetalRHIContext.h
@@ -200,8 +200,6 @@ public:
 	
 	virtual void RHIPopEvent() final override;
 	
-	virtual void RHIUpdateTextureReference(FRHITextureReference* TextureRef, FRHITexture* NewTexture) final override;
-	
 	virtual void RHICopyToStagingBuffer(FRHIBuffer* SourceBufferRHI, FRHIStagingBuffer* DestinationStagingBufferRHI, uint32 Offset, uint32 NumBytes) final override;
 	virtual void RHIWriteGPUFence(FRHIGPUFence* FenceRHI) final override;
 
diff --git a/Engine/Source/Runtime/D3D12RHI/Private/D3D12CommandContext.h b/Engine/Source/Runtime/D3D12RHI/Private/D3D12CommandContext.h
--- a/Engine/Source/Runtime/D3D12RHI/Private/D3D12CommandContext.h
+++ b/Engine/Source/Runtime/D3D12RHI/Private/D3D12CommandContext.h
@@ -307,7 +307,6 @@ public:
 	PRAGMA_DISABLE_DEPRECATION_WARNINGS
     virtual void RHISetShadingRateImage(FRHITexture* RateImageTexture, EVRSRateCombiner Combiner) final override;
 	PRAGMA_ENABLE_DEPRECATION_WARNINGS
-	virtual void RHIUpdateTextureReference(FRHITextureReference* TextureRef, FRHITexture* NewTexture) final override;
 #if PLATFORM_USE_BACKBUFFER_WRITE_TRANSITION_TRACKING
 	virtual void RHIBackBufferWaitTrackingBeginFrame(uint64 FrameToken, bool bDeferred) final override;
 #endif // #if PLATFORM_USE_BACKBUFFER_WRITE_TRANSITION_TRACKING
@@ -690,11 +689,6 @@ public:
 	}
 	PRAGMA_ENABLE_DEPRECATION_WARNINGS
 
-	FORCEINLINE virtual void RHIUpdateTextureReference(FRHITextureReference* TextureRef, FRHITexture* NewTexture) final override
-	{
-		ContextRedirect(RHIUpdateTextureReference(TextureRef, NewTexture));
-	}
-
 	FORCEINLINE virtual void RHIWaitForTemporalEffect(const FName& InEffectName) final AFR_API_OVERRIDE
 	{
 		ContextRedirect(RHIWaitForTemporalEffect(InEffectName));
diff --git a/Engine/Source/Runtime/D3D12RHI/Private/D3D12RHIPrivate.h b/Engine/Source/Runtime/D3D12RHI/Private/D3D12RHIPrivate.h
--- a/Engine/Source/Runtime/D3D12RHI/Private/D3D12RHIPrivate.h
+++ b/Engine/Source/Runtime/D3D12RHI/Private/D3D12RHIPrivate.h
@@ -304,7 +304,6 @@ public:
 	virtual uint64 RHIGetMinimumAlignmentForBufferBackedSRV(EPixelFormat Format) final override;
 	virtual void RHIGetTextureMemoryStats(FTextureMemoryStats& OutStats) final override;
 	virtual bool RHIGetTextureMemoryVisualizeData(FColor* TextureData, int32 SizeX, int32 SizeY, int32 Pitch, int32 PixelSize) final override;
-	virtual FTextureReferenceRHIRef RHICreateTextureReference(FLastRenderTimeContainer* LastRenderTime) final override;
 	virtual FTexture2DRHIRef RHICreateTexture2D(uint32 SizeX, uint32 SizeY, uint8 Format, uint32 NumMips, uint32 NumSamples, ETextureCreateFlags Flags, ERHIAccess InResourceState, FRHIResourceCreateInfo& CreateInfo) override;
 	virtual FTexture2DRHIRef RHIAsyncCreateTexture2D(uint32 SizeX, uint32 SizeY, uint8 Format, uint32 NumMips, ETextureCreateFlags Flags, ERHIAccess InResourceState, void** InitialMipData, uint32 NumInitialMips) final override;
 	virtual void RHICopySharedMips(FRHITexture2D* DestTexture2D, FRHITexture2D* SrcTexture2D) final override;
diff --git a/Engine/Source/Runtime/D3D12RHI/Private/D3D12Texture.cpp b/Engine/Source/Runtime/D3D12RHI/Private/D3D12Texture.cpp
--- a/Engine/Source/Runtime/D3D12RHI/Private/D3D12Texture.cpp
+++ b/Engine/Source/Runtime/D3D12RHI/Private/D3D12Texture.cpp
@@ -2688,40 +2688,6 @@ void FD3D12DynamicRHI::RHIVirtualTextureSetFirstMipVisible(FRHITexture2D* Textur
 {
 }
 
-FTextureReferenceRHIRef FD3D12DynamicRHI::RHICreateTextureReference(FLastRenderTimeContainer* LastRenderTime)
-{
-	return GetAdapter().CreateLinkedObject<FD3D12TextureReference>(
-		FRHIGPUMask::All(), 
-		[&](FD3D12Device* Device)
-	{
-		return new FD3D12TextureReference(Device, LastRenderTime); 
-	});
-}
-
-void FD3D12CommandContext::RHIUpdateTextureReference(FRHITextureReference* TextureRefRHI, FRHITexture* NewTextureRHI)
-{
-#if 0
-	// Updating texture references is disallowed while the RHI could be caching them in referenced resource tables.
-	check(ResourceTableFrameCounter == INDEX_NONE);
-#endif
-
-	FD3D12TextureReference* TextureRef = RetrieveObject<FD3D12TextureReference>(TextureRefRHI);
-	if (TextureRef)
-	{
-		FD3D12TextureBase* NewTexture = NULL;
-		FD3D12ShaderResourceView* NewSRV = NULL;
-		if (NewTextureRHI)
-		{
-			NewTexture = RetrieveTextureBase(NewTextureRHI);
-			if (NewTexture)
-			{
-				NewSRV = NewTexture->GetShaderResourceView();
-			}
-		}
-		TextureRef->SetReferencedTexture(NewTextureRHI, NewTexture, NewSRV);
-	}
-}
-
 ID3D12CommandQueue* FD3D12DynamicRHI::RHIGetD3DCommandQueue()
 {
 	// Multi-GPU support : any code using this function needs validation.
diff --git a/Engine/Source/Runtime/D3D12RHI/Private/D3D12Texture.h b/Engine/Source/Runtime/D3D12RHI/Private/D3D12Texture.h
--- a/Engine/Source/Runtime/D3D12RHI/Private/D3D12Texture.h
+++ b/Engine/Source/Runtime/D3D12RHI/Private/D3D12Texture.h
@@ -478,44 +478,6 @@ typedef TD3D12Texture2D<FD3D12BaseTexture2D>      FD3D12Texture2D;
 typedef TD3D12Texture2D<FD3D12BaseTexture2DArray> FD3D12Texture2DArray;
 typedef TD3D12Texture2D<FD3D12BaseTextureCube>    FD3D12TextureCube;
 
-/** Texture reference class. */
-class FD3D12TextureReference : public FRHITextureReference, public FD3D12TextureBase
-{
-public:
-	FD3D12TextureReference(class FD3D12Device* InParent, FLastRenderTimeContainer* LastRenderTime)
-		: FRHITextureReference(LastRenderTime)
-		, FD3D12TextureBase(InParent)
-	{
-		BaseShaderResource = NULL;
-	}
-
-	void SetReferencedTexture(FRHITexture* InTexture, FD3D12BaseShaderResource* InBaseShaderResource, FD3D12ShaderResourceView* InSRV)
-	{
-		ShaderResourceView = InSRV;
-		BaseShaderResource = InBaseShaderResource;
-		FRHITextureReference::SetReferencedTexture(InTexture);
-	}
-
-	virtual void* GetTextureBaseRHI() override final
-	{
-		return static_cast<FD3D12TextureBase*>(this);
-	}
-
-	// IRefCountedObject interface.
-	virtual uint32 AddRef() const
-	{
-		return FRHIResource::AddRef();
-	}
-	virtual uint32 Release() const
-	{
-		return FRHIResource::Release();
-	}
-	virtual uint32 GetRefCount() const
-	{
-		return FRHIResource::GetRefCount();
-	}
-};
-
 class FD3D12Viewport;
 
 class FD3D12BackBufferReferenceTexture2D : public FD3D12Texture2D
@@ -632,12 +594,6 @@ struct TD3D12ResourceTraits<FRHITextureCube>
 	typedef FD3D12TextureCube TConcreteType;
 };
 
-template<>
-struct TD3D12ResourceTraits<FRHITextureReference>
-{
-	typedef FD3D12TextureReference TConcreteType;
-};
-
 struct FRHICommandD3D12AsyncReallocateTexture2D final : public FRHICommand<FRHICommandD3D12AsyncReallocateTexture2D>
 {
 	FD3D12Texture2D* OldTexture;
diff --git a/Engine/Source/Runtime/EmptyRHI/Private/EmptyTexture.cpp b/Engine/Source/Runtime/EmptyRHI/Private/EmptyTexture.cpp
--- a/Engine/Source/Runtime/EmptyRHI/Private/EmptyTexture.cpp
+++ b/Engine/Source/Runtime/EmptyRHI/Private/EmptyTexture.cpp
@@ -256,12 +256,3 @@ uint64 FEmptyDynamicRHI::RHICalcTextureCubePlatformSize(uint32 Size, uint8 Forma
 {
 	return 0;
 }
-
-FTextureReferenceRHIRef FEmptyDynamicRHI::RHICreateTextureReference(FLastRenderTimeContainer* LastRenderTime)
-{
-	return nullptr;
-}
-
-void FEmptyDynamicRHI::RHIUpdateTextureReference(FRHITextureReference* TextureRef, FRHITexture* NewTexture)
-{
-}
diff --git a/Engine/Source/Runtime/EmptyRHI/Public/EmptyRHI.h b/Engine/Source/Runtime/EmptyRHI/Public/EmptyRHI.h
--- a/Engine/Source/Runtime/EmptyRHI/Public/EmptyRHI.h
+++ b/Engine/Source/Runtime/EmptyRHI/Public/EmptyRHI.h
@@ -74,8 +74,6 @@ public:
 	virtual uint64 RHICalcTextureCubePlatformSize(uint32 Size, uint8 Format, uint32 NumMips, ETextureCreateFlags Flags, const FRHIResourceCreateInfo& CreateInfo, uint32& OutAlign) final override;
 	virtual void RHIGetTextureMemoryStats(FTextureMemoryStats& OutStats) final override;
 	virtual bool RHIGetTextureMemoryVisualizeData(FColor* TextureData,int32 SizeX,int32 SizeY,int32 Pitch,int32 PixelSize) final override;
-	virtual FTextureReferenceRHIRef RHICreateTextureReference(FLastRenderTimeContainer* LastRenderTime) final override;
-	virtual void RHIUpdateTextureReference(FRHITextureReference* TextureRef, FRHITexture* NewTexture) final override;
 	virtual FTexture2DRHIRef RHICreateTexture2D(uint32 SizeX, uint32 SizeY, uint8 Format, uint32 NumMips, uint32 NumSamples, ETextureCreateFlags Flags, FRHIResourceCreateInfo& CreateInfo) final override;
 	virtual FTexture2DRHIRef RHIAsyncCreateTexture2D(uint32 SizeX, uint32 SizeY, uint8 Format, uint32 NumMips, ETextureCreateFlags Flags, void** InitialMipData, uint32 NumInitialMips) final override;
 	virtual void RHICopySharedMips(FRHITexture2D* DestTexture2D, FRHITexture2D* SrcTexture2D) final override;
diff --git a/Engine/Source/Runtime/NullDrv/Public/NullRHI.h b/Engine/Source/Runtime/NullDrv/Public/NullRHI.h
--- a/Engine/Source/Runtime/NullDrv/Public/NullRHI.h
+++ b/Engine/Source/Runtime/NullDrv/Public/NullRHI.h
@@ -225,19 +225,6 @@ public:
 		return false; 
 	}
 
-	virtual FTextureReferenceRHIRef RHICreateTextureReference(FLastRenderTimeContainer* LastRenderTime) final override
-	{ 
-		return new FRHITextureReferenceNullImpl(); 
-	}
-
-	virtual void RHIUpdateTextureReference(FRHITextureReference* TextureRef, FRHITexture* NewTexture) final override
-	{ 
-		if(TextureRef) 
-		{ 
-			((FRHITextureReferenceNullImpl*)TextureRef)->SetReferencedTexture(NewTexture); 
-		} 
-	}
-
 
 	virtual FTexture2DRHIRef RHICreateTexture2D(uint32 SizeX, uint32 SizeY, uint8 Format, uint32 NumMips, uint32 NumSamples, ETextureCreateFlags Flags, ERHIAccess InResourceState, FRHIResourceCreateInfo& CreateInfo) final override
 	{ 
diff --git a/Engine/Source/Runtime/OpenGLDrv/Private/OpenGLTexture.cpp b/Engine/Source/Runtime/OpenGLDrv/Private/OpenGLTexture.cpp
--- a/Engine/Source/Runtime/OpenGLDrv/Private/OpenGLTexture.cpp
+++ b/Engine/Source/Runtime/OpenGLDrv/Private/OpenGLTexture.cpp
@@ -2523,7 +2523,6 @@ void FOpenGLDynamicRHI::RHIBindDebugLabelName(FRHITexture* TextureRHI, const TCH
 #endif
 }
 
-
 void FOpenGLDynamicRHI::RHIVirtualTextureSetFirstMipInMemory(FRHITexture2D* TextureRHI, uint32 FirstMip)
 {
 }
@@ -2532,26 +2531,6 @@ void FOpenGLDynamicRHI::RHIVirtualTextureSetFirstMipVisible(FRHITexture2D* Textu
 {
 }
 
-FTextureReferenceRHIRef FOpenGLDynamicRHI::RHICreateTextureReference(FLastRenderTimeContainer* InLastRenderTime)
-{
-	return new FOpenGLTextureReference(InLastRenderTime);
-}
-
-void FOpenGLTextureReference::SetReferencedTexture(FRHITexture* InTexture)
-{
-	FRHITextureReference::SetReferencedTexture(InTexture);
-	TexturePtr = GetOpenGLTextureFromRHITexture(InTexture);
-}
-
-void FOpenGLDynamicRHI::RHIUpdateTextureReference(FRHITextureReference* TextureRefRHI, FRHITexture* NewTextureRHI)
-{
-	auto* TextureRef = (FOpenGLTextureReference*)TextureRefRHI;
-	if (TextureRef)
-	{
-		TextureRef->SetReferencedTexture(NewTextureRHI);
-	}
-}
-
 void FOpenGLDynamicRHI::RHICopySubTextureRegion(FRHITexture2D* SourceTextureRHI, FRHITexture2D* DestinationTextureRHI, FBox2D SourceBox, FBox2D DestinationBox)
 {
 	VERIFY_GL_SCOPE();
diff --git a/Engine/Source/Runtime/OpenGLDrv/Public/OpenGLDrv.h b/Engine/Source/Runtime/OpenGLDrv/Public/OpenGLDrv.h
--- a/Engine/Source/Runtime/OpenGLDrv/Public/OpenGLDrv.h
+++ b/Engine/Source/Runtime/OpenGLDrv/Public/OpenGLDrv.h
@@ -366,8 +366,6 @@ public:
 	virtual uint64 RHICalcTextureCubePlatformSize(uint32 Size, uint8 Format, uint32 NumMips, ETextureCreateFlags Flags, const FRHIResourceCreateInfo& CreateInfo, uint32& OutAlign) final override;
 	virtual void RHIGetTextureMemoryStats(FTextureMemoryStats& OutStats) final override;
 	virtual bool RHIGetTextureMemoryVisualizeData(FColor* TextureData,int32 SizeX,int32 SizeY,int32 Pitch,int32 PixelSize) final override;
-	virtual FTextureReferenceRHIRef RHICreateTextureReference(FLastRenderTimeContainer* LastRenderTime) final override;
-	virtual void RHIUpdateTextureReference(FRHITextureReference* TextureRef, FRHITexture* NewTexture) final override;
 	virtual FTexture2DRHIRef RHICreateTexture2D(uint32 SizeX, uint32 SizeY, uint8 Format, uint32 NumMips, uint32 NumSamples, ETextureCreateFlags Flags, ERHIAccess InResourceState, FRHIResourceCreateInfo& CreateInfo) final override;
 	virtual FTexture2DRHIRef RHICreateTextureExternal2D(uint32 SizeX, uint32 SizeY, uint8 Format, uint32 NumMips, uint32 NumSamples, ETextureCreateFlags Flags, ERHIAccess InResourceState, FRHIResourceCreateInfo& CreateInfo) final override;
 	virtual FTexture2DRHIRef RHIAsyncCreateTexture2D(uint32 SizeX, uint32 SizeY, uint8 Format, uint32 NumMips, ETextureCreateFlags Flags, ERHIAccess InResourceState, void** InitialMipData, uint32 NumInitialMips) final override;
@@ -788,12 +786,6 @@ public:
 		return this->RHICreateShaderResourceView(Buffer);
 	}
 
-	virtual FTextureReferenceRHIRef RHICreateTextureReference_RenderThread(class FRHICommandListImmediate& RHICmdList, FLastRenderTimeContainer* LastRenderTime) final override
-	{
-		// threadsafe, doesn't really do anything
-		return this->RHICreateTextureReference(LastRenderTime);
-	}
-
 
 	virtual FTexture2DRHIRef RHICreateTexture2D_RenderThread(class FRHICommandListImmediate& RHICmdList, uint32 SizeX, uint32 SizeY, uint8 Format, uint32 NumMips, uint32 NumSamples, ETextureCreateFlags Flags, ERHIAccess InResourceState, FRHIResourceCreateInfo& CreateInfo) final override
 	{
diff --git a/Engine/Source/Runtime/OpenGLDrv/Public/OpenGLResources.h b/Engine/Source/Runtime/OpenGLDrv/Public/OpenGLResources.h
--- a/Engine/Source/Runtime/OpenGLDrv/Public/OpenGLResources.h
+++ b/Engine/Source/Runtime/OpenGLDrv/Public/OpenGLResources.h
@@ -1792,25 +1792,6 @@ typedef TOpenGLTexture<FOpenGLBaseTexture2DArray>		FOpenGLTexture2DArray;
 typedef TOpenGLTexture<FOpenGLBaseTexture3D>			FOpenGLTexture3D;
 typedef TOpenGLTexture<FOpenGLBaseTextureCube>			FOpenGLTextureCube;
 
-class FOpenGLTextureReference : public FRHITextureReference
-{
-	FOpenGLTextureBase* TexturePtr;
-
-public:
-	explicit FOpenGLTextureReference(FLastRenderTimeContainer* InLastRenderTime)
-		: FRHITextureReference(InLastRenderTime)
-		, TexturePtr(NULL)
-	{}
-
-	void SetReferencedTexture(FRHITexture* InTexture);
-	FOpenGLTextureBase* GetTexturePtr() const { return TexturePtr; }
-
-	virtual void* GetTextureBaseRHI() override final
-	{
-		return TexturePtr;
-	}
-};
-
 /** Given a pointer to a RHI texture that was created by the OpenGL RHI, returns a pointer to the FOpenGLTextureBase it encapsulates. */
 inline FOpenGLTextureBase* GetOpenGLTextureFromRHITexture(FRHITexture* Texture)
 {
diff --git a/Engine/Source/Runtime/RHI/Private/RHI.cpp b/Engine/Source/Runtime/RHI/Private/RHI.cpp
--- a/Engine/Source/Runtime/RHI/Private/RHI.cpp
+++ b/Engine/Source/Runtime/RHI/Private/RHI.cpp
@@ -292,6 +292,8 @@ RHI_API FDrawCallCategoryName* FDrawCallCategoryName::Array[FDrawCallCategoryNam
 RHI_API int32 FDrawCallCategoryName::DisplayCounts[FDrawCallCategoryName::MAX_DRAWCALL_CATEGORY];
 RHI_API int32 FDrawCallCategoryName::NumCategory = 0;
 
+TRefCountPtr<FRHITexture> FRHITextureReference::DefaultTexture;
+
 FString FVertexElement::ToString() const
 {
 	return FString::Printf(TEXT("<%u %u %u %u %u %u>")
diff --git a/Engine/Source/Runtime/RHI/Private/RHICommandList.cpp b/Engine/Source/Runtime/RHI/Private/RHICommandList.cpp
--- a/Engine/Source/Runtime/RHI/Private/RHICommandList.cpp
+++ b/Engine/Source/Runtime/RHI/Private/RHICommandList.cpp
@@ -70,11 +70,6 @@ TAutoConsoleVariable<int32> CVarRHICmdFlushRenderThreadTasks(
 	0,
 	TEXT("If true, then we flush the render thread tasks every pass. For issue diagnosis. This is a master switch for more granular cvars."));
 
-TAutoConsoleVariable<int32> CVarRHICmdFlushUpdateTextureReference(
-	TEXT("r.RHICmdFlushUpdateTextureReference"),
-	0,
-	TEXT("If true, then we flush the rhi thread when we do RHIUpdateTextureReference, otherwise this is deferred. For issue diagnosis."));
-
 static TAutoConsoleVariable<int32> CVarRHICmdFlushOnQueueParallelSubmit(
 	TEXT("r.RHICmdFlushOnQueueParallelSubmit"),
 	0,
@@ -2314,12 +2309,6 @@ FRHIShaderLibraryRef FDynamicRHI::RHICreateShaderLibrary_RenderThread(class FRHI
 	return GDynamicRHI->RHICreateShaderLibrary(Platform, FilePath, Name);
 }
 
-FTextureReferenceRHIRef FDynamicRHI::RHICreateTextureReference_RenderThread(class FRHICommandListImmediate& RHICmdList, FLastRenderTimeContainer* LastRenderTime)
-{
-	// Notice no StallRHIThread needed! This function is safe on all RHIs
-	return GDynamicRHI->RHICreateTextureReference(LastRenderTime);
-}
-
 FTexture2DRHIRef FDynamicRHI::RHICreateTexture2D_RenderThread(class FRHICommandListImmediate& RHICmdList, uint32 SizeX, uint32 SizeY, uint8 Format, uint32 NumMips, uint32 NumSamples, ETextureCreateFlags Flags, ERHIAccess InResourceState, FRHIResourceCreateInfo& CreateInfo)
 {
 	CSV_SCOPED_TIMING_STAT(RHITStalls, RHICreateTexture2D_RenderThread);
@@ -2508,17 +2497,10 @@ void FDynamicRHI::RHIReadSurfaceFloatData_RenderThread(class FRHICommandListImme
 
 void FRHICommandListImmediate::UpdateTextureReference(FRHITextureReference* TextureRef, FRHITexture* NewTexture)
 {
-	if (Bypass() || !IsRunningRHIInSeparateThread() || CVarRHICmdFlushUpdateTextureReference.GetValueOnRenderThread() > 0)
+	EnqueueLambda([TextureRef, NewTexture](auto&)
 	{
-		{
-			QUICK_SCOPE_CYCLE_COUNTER(STAT_RHIMETHOD_UpdateTextureReference_FlushRHI);
-			CSV_SCOPED_TIMING_STAT(RHITFlushes, UpdateTextureReference);
-			ImmediateFlush(EImmediateFlushType::FlushRHIThread);
-		}
-		GetContext().RHIUpdateTextureReference(TextureRef, NewTexture);
-		return;
-	}
-	ALLOC_COMMAND(FRHICommandUpdateTextureReference)(TextureRef, NewTexture);
+		TextureRef->SetReferencedTexture(NewTexture);
+	});
 	RHIThreadFence(true);
 	if (GetUsedMemory() > 256 * 1024)
 	{
diff --git a/Engine/Source/Runtime/RHI/Public/DynamicRHI.h b/Engine/Source/Runtime/RHI/Public/DynamicRHI.h
--- a/Engine/Source/Runtime/RHI/Public/DynamicRHI.h
+++ b/Engine/Source/Runtime/RHI/Public/DynamicRHI.h
@@ -691,9 +691,6 @@ public:
 	// FlushType: Flush Immediate
 	virtual bool RHIGetTextureMemoryVisualizeData(FColor* TextureData, int32 SizeX, int32 SizeY, int32 Pitch, int32 PixelSize) = 0;
 
-	// FlushType: Wait RHI Thread
-	virtual FTextureReferenceRHIRef RHICreateTextureReference(FLastRenderTimeContainer* LastRenderTime) = 0;
-
 	/**
 	* Creates a 2D RHI texture resource
 	* @param SizeX - width of the texture to create
@@ -1293,7 +1290,6 @@ public:
 	virtual void UpdateTexture3D_RenderThread(class FRHICommandListImmediate& RHICmdList, FRHITexture3D* Texture, uint32 MipIndex, const struct FUpdateTextureRegion3D& UpdateRegion, uint32 SourceRowPitch, uint32 SourceDepthPitch, const uint8* SourceData);
 
 	virtual FRHIShaderLibraryRef RHICreateShaderLibrary_RenderThread(class FRHICommandListImmediate& RHICmdList, EShaderPlatform Platform, FString FilePath, FString Name);
-	virtual FTextureReferenceRHIRef RHICreateTextureReference_RenderThread(class FRHICommandListImmediate& RHICmdList, FLastRenderTimeContainer* LastRenderTime);
 	virtual FTexture2DRHIRef RHICreateTexture2D_RenderThread(class FRHICommandListImmediate& RHICmdList, uint32 SizeX, uint32 SizeY, uint8 Format, uint32 NumMips, uint32 NumSamples, ETextureCreateFlags Flags, ERHIAccess InResourceState, FRHIResourceCreateInfo& CreateInfo);
 	virtual FTexture2DRHIRef RHICreateTextureExternal2D_RenderThread(class FRHICommandListImmediate& RHICmdList, uint32 SizeX, uint32 SizeY, uint8 Format, uint32 NumMips, uint32 NumSamples, ETextureCreateFlags Flags, ERHIAccess InResourceState, FRHIResourceCreateInfo& CreateInfo);
 
diff --git a/Engine/Source/Runtime/RHI/Public/RHICommandList.h b/Engine/Source/Runtime/RHI/Public/RHICommandList.h
--- a/Engine/Source/Runtime/RHI/Public/RHICommandList.h
+++ b/Engine/Source/Runtime/RHI/Public/RHICommandList.h
@@ -1911,18 +1911,6 @@ FRHICOMMAND_MACRO(FRHICommandDebugBreak)
 	}
 };
 
-FRHICOMMAND_MACRO(FRHICommandUpdateTextureReference)
-{
-	FRHITextureReference* TextureRef;
-	FRHITexture* NewTexture;
-	FORCEINLINE_DEBUGGABLE FRHICommandUpdateTextureReference(FRHITextureReference* InTextureRef, FRHITexture* InNewTexture)
-		: TextureRef(InTextureRef)
-		, NewTexture(InNewTexture)
-	{
-	}
-	RHI_API void Execute(FRHICommandListBase& CmdList);
-};
-
 struct FRHIShaderResourceViewUpdateInfo
 {
 	FRHIShaderResourceView* SRV;
@@ -4118,12 +4106,6 @@ public:
 		return GDynamicRHI->RHIGetTextureMemoryVisualizeData(TextureData,SizeX,SizeY,Pitch,PixelSize);
 	}
 	
-	FORCEINLINE FTextureReferenceRHIRef CreateTextureReference(FLastRenderTimeContainer* LastRenderTime)
-	{
-		LLM_SCOPE(ELLMTag::Textures);
-		return GDynamicRHI->RHICreateTextureReference_RenderThread(*this, LastRenderTime);
-	}
-	
 	FORCEINLINE void CopySharedMips(FRHITexture2D* DestTexture2D, FRHITexture2D* SrcTexture2D)
 	{
 		QUICK_SCOPE_CYCLE_COUNTER(STAT_RHIMETHOD_CopySharedMips_Flush);
@@ -5142,12 +5124,13 @@ FORCEINLINE void RHIUpdateRHIResources(FRHIResourceUpdateInfo* UpdateInfos, int3
 
 FORCEINLINE FTextureReferenceRHIRef RHICreateTextureReference(FLastRenderTimeContainer* LastRenderTime)
 {
-	return FRHICommandListExecutor::GetImmediateCommandList().CreateTextureReference(LastRenderTime);
+	return new FRHITextureReference(LastRenderTime);
 }
 
 FORCEINLINE void RHIUpdateTextureReference(FRHITextureReference* TextureRef, FRHITexture* NewTexture)
 {
-	 FRHICommandListExecutor::GetImmediateCommandList().UpdateTextureReference(TextureRef, NewTexture);
+	check(IsInRenderingThread());
+	FRHICommandListExecutor::GetImmediateCommandList().UpdateTextureReference(TextureRef, NewTexture);
 }
 
 FORCEINLINE FTexture2DRHIRef RHICreateTexture2D(uint32 SizeX, uint32 SizeY, uint8 Format, uint32 NumMips, uint32 NumSamples, ETextureCreateFlags Flags, ERHIAccess InResourceState, FRHIResourceCreateInfo& CreateInfo)
diff --git a/Engine/Source/Runtime/RHI/Public/RHICommandListCommandExecutes.inl b/Engine/Source/Runtime/RHI/Public/RHICommandListCommandExecutes.inl
--- a/Engine/Source/Runtime/RHI/Public/RHICommandListCommandExecutes.inl
+++ b/Engine/Source/Runtime/RHI/Public/RHICommandListCommandExecutes.inl
@@ -54,7 +54,6 @@ struct FRHICommandTransitionTextures;
 struct FRHICommandTransitionTexturesPipeline;
 struct FRHICommandTransitionTexturesDepth;
 struct FRHICommandTransitionTexturesArray;
-struct FRHICommandUpdateTextureReference;
 struct FRHICommandClearRayTracingBindings;
 struct FRHICommandRayTraceOcclusion;
 struct FRHICommandRayTraceIntersection;
@@ -582,12 +581,6 @@ void FRHICommandSetRayTracingBindings::Execute(FRHICommandListBase& CmdList)
 
 #endif // RHI_RAYTRACING
 
-void FRHICommandUpdateTextureReference::Execute(FRHICommandListBase& CmdList)
-{
-	RHISTAT(UpdateTextureReference);
-	INTERNAL_DECORATOR(RHIUpdateTextureReference)(TextureRef, NewTexture);
-}
-
 void FRHIResourceUpdateInfo::ReleaseRefs()
 {
 	switch (Type)
diff --git a/Engine/Source/Runtime/RHI/Public/RHIContext.h b/Engine/Source/Runtime/RHI/Public/RHIContext.h
--- a/Engine/Source/Runtime/RHI/Public/RHIContext.h
+++ b/Engine/Source/Runtime/RHI/Public/RHIContext.h
@@ -567,8 +567,6 @@ public:
 		/* empty default implementation */
 	}
 
-	virtual void RHIUpdateTextureReference(FRHITextureReference* TextureRef, FRHITexture* NewTexture) = 0;
-
 	virtual void RHIBeginRenderPass(const FRHIRenderPassInfo& InInfo, const TCHAR* InName) = 0;
 
 	virtual void RHIEndRenderPass() = 0;
diff --git a/Engine/Source/Runtime/RHI/Public/RHIResources.h b/Engine/Source/Runtime/RHI/Public/RHIResources.h
--- a/Engine/Source/Runtime/RHI/Public/RHIResources.h
+++ b/Engine/Source/Runtime/RHI/Public/RHIResources.h
@@ -1195,28 +1195,49 @@ private:
 	uint32 Size;
 };
 
-class RHI_API FRHITextureReference : public FRHITexture
+class RHI_API FRHITextureReference final : public FRHITexture
 {
 public:
 	explicit FRHITextureReference(FLastRenderTimeContainer* InLastRenderTime)
-		: FRHITexture(0,0,PF_Unknown,TexCreate_None,InLastRenderTime, FClearValueBinding())
-	{}
+		: FRHITexture(0, 0, PF_Unknown, TexCreate_None, InLastRenderTime, FClearValueBinding())
+	{
+		check(DefaultTexture);
+		ReferencedTexture = DefaultTexture;
+	}
 
-	virtual FRHITextureReference* GetTextureReference() override { return this; }
-	inline FRHITexture* GetReferencedTexture() const { return ReferencedTexture.GetReference(); }
+	virtual class FRHITextureReference* GetTextureReference() override
+	{
+		return this;
+	}
 
-	void SetReferencedTexture(FRHITexture* InTexture)
+	virtual FIntVector GetSizeXYZ() const override 
 	{
-		ReferencedTexture = InTexture;
+		check(ReferencedTexture);
+		return ReferencedTexture->GetSizeXYZ();
 	}
 
-	virtual FIntVector GetSizeXYZ() const final override
+	virtual void* GetNativeResource() const override 
 	{
-		if (ReferencedTexture)
-		{
-			return ReferencedTexture->GetSizeXYZ();
-		}
-		return FIntVector(0, 0, 0);
+		check(ReferencedTexture);
+		return ReferencedTexture->GetNativeResource();
+	}
+
+	virtual void* GetNativeShaderResourceView() const override
+	{
+		check(ReferencedTexture);
+		return ReferencedTexture->GetNativeShaderResourceView();
+	}
+
+	virtual void* GetTextureBaseRHI() override 
+	{
+		check(ReferencedTexture);
+		return ReferencedTexture->GetTextureBaseRHI();
+	}
+
+	virtual void GetWriteMaskProperties(void*& OutData, uint32& OutSize) override
+	{
+		check(ReferencedTexture);
+		return ReferencedTexture->GetWriteMaskProperties(OutData, OutSize);
 	}
 
 #if ENABLE_RHI_VALIDATION
@@ -1227,21 +1248,28 @@ public:
 	}
 #endif
 
-private:
-	TRefCountPtr<FRHITexture> ReferencedTexture;
-};
-
-class RHI_API FRHITextureReferenceNullImpl : public FRHITextureReference
-{
-public:
-	FRHITextureReferenceNullImpl()
-		: FRHITextureReference(NULL)
-	{}
+	inline FRHITexture* GetReferencedTexture() const
+	{
+		return ReferencedTexture.GetReference();
+	}
 
+private:
+	friend class FRHICommandListImmediate;
+	// Called only from FRHICommandListImmediate::UpdateTextureReference()
 	void SetReferencedTexture(FRHITexture* InTexture)
 	{
-		FRHITextureReference::SetReferencedTexture(InTexture);
+		ReferencedTexture = InTexture
+			? InTexture
+			: DefaultTexture.GetReference();
 	}
+
+	TRefCountPtr<FRHITexture> ReferencedTexture;
+
+	// This pointer is set by the InitRHI() function on the FBlackTextureWithSRV global resource,
+	// to allow FRHITextureReference to use the global black texture when the reference is nullptr.
+	// A pointer is required since FBlackTextureWithSRV is defined in RenderCore.
+	friend class FBlackTextureWithSRV;
+	static TRefCountPtr<FRHITexture> DefaultTexture;
 };
 
 //
diff --git a/Engine/Source/Runtime/RHI/Public/RHIValidation.h b/Engine/Source/Runtime/RHI/Public/RHIValidation.h
--- a/Engine/Source/Runtime/RHI/Public/RHIValidation.h
+++ b/Engine/Source/Runtime/RHI/Public/RHIValidation.h
@@ -400,6 +400,8 @@ public:
 	// FlushType: Wait RHI Thread
 	virtual FUnorderedAccessViewRHIRef RHICreateUnorderedAccessView(FRHITexture* Texture, uint32 MipLevel) override final
 	{
+		RHI_VALIDATION_CHECK(Texture->GetTextureReference() == nullptr, TEXT("Creating an unordered access view of an FRHITextureReference is not supported."));
+
 		FUnorderedAccessViewRHIRef UAV = RHI->RHICreateUnorderedAccessView(Texture, MipLevel);
 		UAV->ViewIdentity = Texture->GetViewIdentity(MipLevel, 1, 0, 0, uint32(RHIValidation::EResourcePlane::Common), 1);
 		return UAV;
@@ -542,12 +544,6 @@ public:
 		return RHI->RHIGetTextureMemoryVisualizeData(TextureData, SizeX, SizeY, Pitch, PixelSize);
 	}
 
-	// FlushType: Wait RHI Thread
-	virtual FTextureReferenceRHIRef RHICreateTextureReference(FLastRenderTimeContainer* LastRenderTime) override final
-	{
-		return RHI->RHICreateTextureReference(LastRenderTime);
-	}
-
 	/**
 	* Creates a 2D RHI texture resource
 	* @param SizeX - width of the texture to create
@@ -694,6 +690,8 @@ public:
 	// FlushType: Wait RHI Thread
 	virtual FShaderResourceViewRHIRef RHICreateShaderResourceView(FRHITexture* TextureRHI, const FRHITextureSRVCreateInfo& CreateInfo) override final
 	{
+		RHI_VALIDATION_CHECK(TextureRHI->GetTextureReference() == nullptr, TEXT("Creating a shader resource view of an FRHITextureReference is not supported."));
+
 		FShaderResourceViewRHIRef SRV = RHI->RHICreateShaderResourceView(TextureRHI, CreateInfo);
 		SRV->ViewIdentity = TextureRHI->GetViewIdentity(CreateInfo.MipLevel, CreateInfo.NumMipLevels, CreateInfo.FirstArraySlice, CreateInfo.NumArraySlices, uint32(RHIValidation::EResourcePlane::Common), 1);
 		return SRV;
@@ -1463,11 +1461,6 @@ public:
 		return RHI->RHICreateShaderLibrary_RenderThread(RHICmdList, Platform, FilePath, Name);
 	}
 
-	virtual FTextureReferenceRHIRef RHICreateTextureReference_RenderThread(class FRHICommandListImmediate& RHICmdList, FLastRenderTimeContainer* LastRenderTime) override final
-	{
-		return RHI->RHICreateTextureReference_RenderThread(RHICmdList, LastRenderTime);
-	}
-
 	virtual FTexture2DRHIRef RHICreateTexture2D_RenderThread(class FRHICommandListImmediate& RHICmdList, uint32 SizeX, uint32 SizeY, uint8 Format, uint32 NumMips, uint32 NumSamples, ETextureCreateFlags Flags, ERHIAccess InResourceState, FRHIResourceCreateInfo& CreateInfo) override final
 	{
 		ensure(FMath::Max(SizeX, SizeY) >= (1u << (FMath::Max(1u, NumMips) - 1)));
@@ -1501,6 +1494,8 @@ public:
 
 	virtual FUnorderedAccessViewRHIRef RHICreateUnorderedAccessView_RenderThread(class FRHICommandListImmediate& RHICmdList, FRHITexture* Texture, uint32 MipLevel) override final
 	{
+		RHI_VALIDATION_CHECK(Texture->GetTextureReference() == nullptr, TEXT("Creating an unordered access view of an FRHITextureReference is not supported."));
+
 		FUnorderedAccessViewRHIRef UAV = RHI->RHICreateUnorderedAccessView_RenderThread(RHICmdList, Texture, MipLevel);
 		UAV->ViewIdentity = Texture->GetViewIdentity(MipLevel, 1, 0, 0, uint32(RHIValidation::EResourcePlane::Common), 1);
 		return UAV;
@@ -1515,6 +1510,8 @@ public:
 
 	virtual FShaderResourceViewRHIRef RHICreateShaderResourceView_RenderThread(class FRHICommandListImmediate& RHICmdList, FRHITexture* Texture, const FRHITextureSRVCreateInfo& CreateInfo) override final
 	{
+		RHI_VALIDATION_CHECK(Texture->GetTextureReference() == nullptr, TEXT("Creating a shader resource view of an FRHITextureReference is not supported."));
+
 		FShaderResourceViewRHIRef SRV = RHI->RHICreateShaderResourceView_RenderThread(RHICmdList, Texture, CreateInfo);
 		RHIValidation::EResourcePlane Plane = CreateInfo.Format == PF_X24_G8
 			? RHIValidation::EResourcePlane::Stencil
diff --git a/Engine/Source/Runtime/RHI/Public/RHIValidationContext.h b/Engine/Source/Runtime/RHI/Public/RHIValidationContext.h
--- a/Engine/Source/Runtime/RHI/Public/RHIValidationContext.h
+++ b/Engine/Source/Runtime/RHI/Public/RHIValidationContext.h
@@ -770,11 +770,6 @@ public:
 		RHIContext->RHIPopEvent();
 	}
 
-	virtual void RHIUpdateTextureReference(FRHITextureReference* TextureRef, FRHITexture* NewTexture) override final
-	{
-		RHIContext->RHIUpdateTextureReference(TextureRef, NewTexture);
-	}
-
 	virtual void RHIBeginRenderPass(const FRHIRenderPassInfo& InInfo, const TCHAR* InName) override final
 	{
 		checkf(!State.bInsideBeginRenderPass, TEXT("Trying to begin RenderPass '%s', but already inside '%s'!"), *State.RenderPassName, InName);
diff --git a/Engine/Source/Runtime/RenderCore/Private/RenderUtils.cpp b/Engine/Source/Runtime/RenderCore/Private/RenderUtils.cpp
--- a/Engine/Source/Runtime/RenderCore/Private/RenderUtils.cpp
+++ b/Engine/Source/Runtime/RenderCore/Private/RenderUtils.cpp
@@ -211,8 +211,23 @@ public:
 	}
 };
 
+class FBlackTextureWithSRV : public FColoredTexture<0, 0, 0, 255>
+{
+	virtual void InitRHI() override
+	{
+		FColoredTexture::InitRHI();
+		FRHITextureReference::DefaultTexture = TextureRHI;
+	}
+
+	virtual void ReleaseRHI() override
+	{
+		FRHITextureReference::DefaultTexture.SafeRelease();
+		FColoredTexture::ReleaseRHI();
+	}
+};
+
 FTextureWithSRV* GWhiteTextureWithSRV = new TGlobalResource<FColoredTexture<255,255,255,255> >;
-FTextureWithSRV* GBlackTextureWithSRV = new TGlobalResource<FColoredTexture<0,0,0,255> >;
+FTextureWithSRV* GBlackTextureWithSRV = new TGlobalResource<FBlackTextureWithSRV>();
 FTextureWithSRV* GTransparentBlackTextureWithSRV = new TGlobalResource<FColoredTexture<0,0,0,0> >;
 FTexture* GWhiteTexture = GWhiteTextureWithSRV;
 FTexture* GBlackTexture = GBlackTextureWithSRV;
diff --git a/Engine/Source/Runtime/VulkanRHI/Private/VulkanCommands.cpp b/Engine/Source/Runtime/VulkanRHI/Private/VulkanCommands.cpp
--- a/Engine/Source/Runtime/VulkanRHI/Private/VulkanCommands.cpp
+++ b/Engine/Source/Runtime/VulkanRHI/Private/VulkanCommands.cpp
@@ -48,46 +48,6 @@ static_assert(STRUCT_OFFSET(FRHIDispatchIndirectParameters, ThreadGroupCountX) =
 static_assert(STRUCT_OFFSET(FRHIDispatchIndirectParameters, ThreadGroupCountY) == STRUCT_OFFSET(VkDispatchIndirectCommand, y), "FRHIDispatchIndirectParameters Y dimension is wrong.");
 static_assert(STRUCT_OFFSET(FRHIDispatchIndirectParameters, ThreadGroupCountZ) == STRUCT_OFFSET(VkDispatchIndirectCommand, z), "FRHIDispatchIndirectParameters Z dimension is wrong.");
 
-/** Given a pointer to a RHI texture that was created by the Vulkan RHI, returns a pointer to the FVulkanTextureBase it encapsulates. */
-inline FVulkanTextureBase* GetVulkanTextureFromRHITexture(FRHITexture* Texture)
-{
-	if (!Texture)
-	{
-		return NULL;
-	}
-	if(FRHITextureReference* TexRef = Texture->GetTextureReference())
-	{
-		Texture = TexRef->GetReferencedTexture();
-		if (!Texture)
-		{
-			return NULL;
-		}
-	}
-
-
-	if (FRHITexture2D* Tex2D = Texture->GetTexture2D())
-	{
-		return static_cast<FVulkanTexture2D*>(Tex2D);
-	}
-	else if (FRHITexture2DArray* Tex2DArray = Texture->GetTexture2DArray())
-	{
-		return static_cast<FVulkanTexture2DArray*>(Tex2DArray);
-	}
-	else if (FRHITexture3D* Tex3D = Texture->GetTexture3D())
-	{
-		return static_cast<FVulkanTexture3D*>(Tex3D);
-	}
-	else if (FRHITextureCube* TexCube = Texture->GetTextureCube())
-	{
-		return static_cast<FVulkanTextureCube*>(TexCube);
-	}
-	else
-	{
-		UE_LOG(LogVulkanRHI, Fatal, TEXT("Unknown Vulkan RHI texture type"));
-		return nullptr;
-	}
-}
-
 static FORCEINLINE ShaderStage::EStage GetAndVerifyShaderStage(FRHIGraphicsShader* ShaderRHI, FVulkanPendingGfxState* PendingGfxState)
 {
 	switch (ShaderRHI->GetFrequency())
@@ -301,7 +261,7 @@ void FVulkanCommandListContext::RHISetUAVParameter(FRHIComputeShader* ComputeSha
 
 void FVulkanCommandListContext::RHISetShaderTexture(FRHIGraphicsShader* ShaderRHI, uint32 TextureIndex, FRHITexture* NewTextureRHI)
 {
-	FVulkanTextureBase* Texture = GetVulkanTextureFromRHITexture(NewTextureRHI);
+	FVulkanTextureBase* Texture = FVulkanTextureBase::Cast(NewTextureRHI);
 	VkImageLayout Layout = LayoutManager.FindLayoutChecked(Texture->Surface.Image);
 
 	ShaderStage::EStage Stage = GetAndVerifyShaderStage(ShaderRHI, PendingGfxState);
@@ -314,7 +274,7 @@ void FVulkanCommandListContext::RHISetShaderTexture(FRHIComputeShader* ComputeSh
 	FVulkanComputeShader* ComputeShader = ResourceCast(ComputeShaderRHI);
 	check(PendingComputeState->GetCurrentShader() == ComputeShader);
 
-	FVulkanTextureBase* VulkanTexture = GetVulkanTextureFromRHITexture(NewTextureRHI);
+	FVulkanTextureBase* VulkanTexture = FVulkanTextureBase::Cast(NewTextureRHI);
 	VkImageLayout Layout = LayoutManager.FindLayoutChecked(VulkanTexture->Surface.Image);
 	PendingComputeState->SetTextureForStage(TextureIndex, VulkanTexture, Layout);
 	NewTextureRHI->SetLastRenderTime((float)FPlatformTime::Seconds());
diff --git a/Engine/Source/Runtime/VulkanRHI/Private/VulkanContext.h b/Engine/Source/Runtime/VulkanRHI/Private/VulkanContext.h
--- a/Engine/Source/Runtime/VulkanRHI/Private/VulkanContext.h
+++ b/Engine/Source/Runtime/VulkanRHI/Private/VulkanContext.h
@@ -85,8 +85,6 @@ public:
 	virtual void RHIBeginRenderQuery(FRHIRenderQuery* RenderQuery) final override;
 	virtual void RHIEndRenderQuery(FRHIRenderQuery* RenderQuery) final override;
 
-	virtual void RHIUpdateTextureReference(FRHITextureReference* TextureRef, FRHITexture* NewTexture) final override;
-
 	virtual void RHISubmitCommandsHint() final override;
 
 	virtual void RHIBeginDrawingViewport(FRHIViewport* Viewport, FRHITexture* RenderTargetRHI) final override;
diff --git a/Engine/Source/Runtime/VulkanRHI/Private/VulkanResources.h b/Engine/Source/Runtime/VulkanRHI/Private/VulkanResources.h
--- a/Engine/Source/Runtime/VulkanRHI/Private/VulkanResources.h
+++ b/Engine/Source/Runtime/VulkanRHI/Private/VulkanResources.h
@@ -704,50 +704,6 @@ public:
 	}
 };
 
-class FVulkanTextureReference : public FRHITextureReference, public FVulkanTextureBase
-{
-	FName GetResourceFName() { return GetName(); }
-public:
-	explicit FVulkanTextureReference(FVulkanDevice& Device, FLastRenderTimeContainer* InLastRenderTime)
-	:	FRHITextureReference(InLastRenderTime)
-	,	FVulkanTextureBase(Device, VK_IMAGE_VIEW_TYPE_MAX_ENUM, PF_Unknown, 0, 0, 0, 1, 1, 1, VK_NULL_HANDLE, VK_NULL_HANDLE, TexCreate_None, FRHIResourceCreateInfo(TEXT("")))
-	{}
-
-	virtual FRHITexture* GetRHITexture()
-	{
-		return this;
-	};
-
-
-	// IRefCountedObject interface.
-	virtual uint32 AddRef() const override final
-	{
-		return FRHIResource::AddRef();
-	}
-
-	virtual uint32 Release() const override final
-	{
-		return FRHIResource::Release();
-	}
-
-	virtual uint32 GetRefCount() const override final
-	{
-		return FRHIResource::GetRefCount();
-	}
-
-	virtual void* GetTextureBaseRHI() override final
-	{
-		return GetReferencedTexture() ? GetReferencedTexture()->GetTextureBaseRHI() : nullptr;
-	}
-
-	virtual void* GetNativeResource() const
-	{
-		return (void*)Surface.Image;
-	}
-
-	void SetReferencedTexture(FRHITexture* InTexture);
-};
-
 class FVulkanQueryPool : public VulkanRHI::FDeviceChild
 {
 public:
diff --git a/Engine/Source/Runtime/VulkanRHI/Private/VulkanTexture.cpp b/Engine/Source/Runtime/VulkanRHI/Private/VulkanTexture.cpp
--- a/Engine/Source/Runtime/VulkanRHI/Private/VulkanTexture.cpp
+++ b/Engine/Source/Runtime/VulkanRHI/Private/VulkanTexture.cpp
@@ -2569,11 +2569,6 @@ FVulkanTexture2DArray::FVulkanTexture2DArray(FTextureRHIRef& SrcTextureRHI, cons
 {
 }
 
-void FVulkanTextureReference::SetReferencedTexture(FRHITexture* InTexture)
-{
-	FRHITextureReference::SetReferencedTexture(InTexture);
-}
-
 
 FVulkanTextureCube::FVulkanTextureCube(FVulkanDevice& Device, EPixelFormat Format, uint32 Size, bool bArray, uint32 ArraySize, uint32 NumMips, ETextureCreateFlags Flags, ERHIAccess InResourceState, const FRHIResourceCreateInfo& CreateInfo)
 :	 FRHITextureCube(Size, NumMips, Format, Flags, CreateInfo.ClearValueBinding)
@@ -2869,21 +2864,6 @@ uint64 FVulkanDynamicRHI::RHICalcTextureCubePlatformSize(uint32 Size, uint8 Form
 	return MemReq.size;
 }
 
-FTextureReferenceRHIRef FVulkanDynamicRHI::RHICreateTextureReference(FLastRenderTimeContainer* LastRenderTime)
-{
-	return new FVulkanTextureReference(*Device, LastRenderTime);
-}
-
-void FVulkanCommandListContext::RHIUpdateTextureReference(FRHITextureReference* TextureRef, FRHITexture* NewTexture)
-{
-	//#todo-rco: Implementation needs to be verified
-	FVulkanTextureReference* VulkanTextureRef = (FVulkanTextureReference*)TextureRef;
-	if (VulkanTextureRef)
-	{
-		VulkanTextureRef->SetReferencedTexture(NewTexture);
-	}
-}
-
 void FVulkanCommandListContext::RHICopyTexture(FRHITexture* SourceTexture, FRHITexture* DestTexture, const FRHICopyTextureInfo& CopyInfo)
 {
 	LLM_SCOPE_VULKAN(ELLMTagVulkan::VulkanTextures);
diff --git a/Engine/Source/Runtime/VulkanRHI/Public/VulkanDynamicRHI.h b/Engine/Source/Runtime/VulkanRHI/Public/VulkanDynamicRHI.h
--- a/Engine/Source/Runtime/VulkanRHI/Public/VulkanDynamicRHI.h
+++ b/Engine/Source/Runtime/VulkanRHI/Public/VulkanDynamicRHI.h
@@ -99,7 +99,6 @@ public:
 	virtual uint64 RHICalcTextureCubePlatformSize(uint32 Size, uint8 Format, uint32 NumMips, ETextureCreateFlags Flags, const FRHIResourceCreateInfo& CreateInfo, uint32& OutAlign) final override;
 	virtual void RHIGetTextureMemoryStats(FTextureMemoryStats& OutStats) final override;
 	virtual bool RHIGetTextureMemoryVisualizeData(FColor* TextureData, int32 SizeX, int32 SizeY, int32 Pitch, int32 PixelSize) final override;
-	virtual FTextureReferenceRHIRef RHICreateTextureReference(FLastRenderTimeContainer* LastRenderTime) final override;
 	virtual FTexture2DRHIRef RHICreateTexture2D(uint32 SizeX, uint32 SizeY, uint8 Format, uint32 NumMips, uint32 NumSamples, ETextureCreateFlags Flags, ERHIAccess InResourceState, FRHIResourceCreateInfo& CreateInfo) final override;
 	virtual FTexture2DRHIRef RHIAsyncCreateTexture2D(uint32 SizeX, uint32 SizeY, uint8 Format, uint32 NumMips, ETextureCreateFlags Flags, ERHIAccess InResourceState, void** InitialMipData, uint32 NumInitialMips) final override;
 	virtual void RHICopySharedMips(FRHITexture2D* DestTexture2D, FRHITexture2D* SrcTexture2D) final override;
@@ -238,12 +237,6 @@ public:
 	virtual FUpdateTexture3DData BeginUpdateTexture3D_RenderThread(class FRHICommandListImmediate& RHICmdList, FRHITexture3D* Texture, uint32 MipIndex, const struct FUpdateTextureRegion3D& UpdateRegion) override final;
 	virtual void EndUpdateTexture3D_RenderThread(class FRHICommandListImmediate& RHICmdList, FUpdateTexture3DData& UpdateData) override final;
 
-	virtual FTextureReferenceRHIRef RHICreateTextureReference_RenderThread(class FRHICommandListImmediate& RHICmdList, FLastRenderTimeContainer* LastRenderTime) final override
-	{
-		// this-> is required to avoid calling the global version of this function
-		return this->RHICreateTextureReference(LastRenderTime);
-	}
-
 	virtual FTexture2DRHIRef RHICreateTexture2D_RenderThread(class FRHICommandListImmediate& RHICmdList, uint32 SizeX, uint32 SizeY, uint8 Format, uint32 NumMips, uint32 NumSamples, ETextureCreateFlags Flags, ERHIAccess InResourceState, FRHIResourceCreateInfo& CreateInfo) override final
 	{
 		return RHICreateTexture2D(SizeX, SizeY, Format, NumMips, NumSamples, Flags, InResourceState, CreateInfo);
diff --git a/Engine/Source/Runtime/Windows/D3D11RHI/Private/D3D11RHIPrivate.h b/Engine/Source/Runtime/Windows/D3D11RHI/Private/D3D11RHIPrivate.h
--- a/Engine/Source/Runtime/Windows/D3D11RHI/Private/D3D11RHIPrivate.h
+++ b/Engine/Source/Runtime/Windows/D3D11RHI/Private/D3D11RHIPrivate.h
@@ -457,8 +457,6 @@ public:
 	virtual uint64 RHIGetMinimumAlignmentForBufferBackedSRV(EPixelFormat Format) final override;
 	virtual void RHIGetTextureMemoryStats(FTextureMemoryStats& OutStats) final override;
 	virtual bool RHIGetTextureMemoryVisualizeData(FColor* TextureData,int32 SizeX,int32 SizeY,int32 Pitch,int32 PixelSize) final override;
-	virtual FTextureReferenceRHIRef RHICreateTextureReference(FLastRenderTimeContainer* LastRenderTime) final override;
-	virtual void RHIUpdateTextureReference(FRHITextureReference* TextureRef, FRHITexture* NewTexture) final override;
 	virtual FTexture2DRHIRef RHICreateTexture2D(uint32 SizeX, uint32 SizeY, uint8 Format, uint32 NumMips, uint32 NumSamples, ETextureCreateFlags Flags, ERHIAccess InResourceState, FRHIResourceCreateInfo& CreateInfo) final override;
 	virtual FTexture2DRHIRef RHIAsyncCreateTexture2D(uint32 SizeX, uint32 SizeY, uint8 Format, uint32 NumMips, ETextureCreateFlags Flags, ERHIAccess InResourceState, void** InitialMipData, uint32 NumInitialMips) final override;
 	virtual void RHICopySharedMips(FRHITexture2D* DestTexture2D, FRHITexture2D* SrcTexture2D) final override;
@@ -638,7 +636,6 @@ public:
 	virtual FTextureCubeRHIRef RHICreateTextureCube_RenderThread(class FRHICommandListImmediate& RHICmdList, uint32 Size, uint8 Format, uint32 NumMips, ETextureCreateFlags Flags, ERHIAccess InResourceState, FRHIResourceCreateInfo& CreateInfo) final override;
 	virtual FTextureCubeRHIRef RHICreateTextureCubeArray_RenderThread(class FRHICommandListImmediate& RHICmdList, uint32 Size, uint32 ArraySize, uint8 Format, uint32 NumMips, ETextureCreateFlags Flags, ERHIAccess InResourceState, FRHIResourceCreateInfo& CreateInfo) final override;
 	virtual FRenderQueryRHIRef RHICreateRenderQuery_RenderThread(class FRHICommandListImmediate& RHICmdList, ERenderQueryType QueryType) final override;
-	virtual FTextureReferenceRHIRef RHICreateTextureReference_RenderThread(class FRHICommandListImmediate& RHICmdList, FLastRenderTimeContainer* LastRenderTime) final override;
 	virtual void* RHILockTextureCubeFace_RenderThread(class FRHICommandListImmediate& RHICmdList, FRHITextureCube* Texture, uint32 FaceIndex, uint32 ArrayIndex, uint32 MipIndex, EResourceLockMode LockMode, uint32& DestStride, bool bLockWithinMiptail) final override;
 	virtual void RHIUnlockTextureCubeFace_RenderThread(class FRHICommandListImmediate& RHICmdList, FRHITextureCube* Texture, uint32 FaceIndex, uint32 ArrayIndex, uint32 MipIndex, bool bLockWithinMiptail) final override;
 	virtual void RHICopySubTextureRegion_RenderThread(class FRHICommandListImmediate& RHICmdList, FRHITexture2D* SourceTexture, FRHITexture2D* DestinationTexture, FBox2D SourceBox, FBox2D DestinationBox) final override;
diff --git a/Engine/Source/Runtime/Windows/D3D11RHI/Private/D3D11Texture.cpp b/Engine/Source/Runtime/Windows/D3D11RHI/Private/D3D11Texture.cpp
--- a/Engine/Source/Runtime/Windows/D3D11RHI/Private/D3D11Texture.cpp
+++ b/Engine/Source/Runtime/Windows/D3D11RHI/Private/D3D11Texture.cpp
@@ -2141,18 +2141,6 @@ void FD3D11DynamicRHI::RHIVirtualTextureSetFirstMipVisible(FRHITexture2D* Textur
 {
 }
 
-FTextureReferenceRHIRef FD3D11DynamicRHI::RHICreateTextureReference(FLastRenderTimeContainer* LastRenderTime)
-{
-	return new FD3D11TextureReference(this,LastRenderTime);
-}
-
-FTextureReferenceRHIRef FD3D11DynamicRHI::RHICreateTextureReference_RenderThread(
-	class FRHICommandListImmediate& RHICmdList,
-	FLastRenderTimeContainer* LastRenderTime)
-{
-	return RHICreateTextureReference(LastRenderTime);
-}
-
 void FD3D11DynamicRHI::RHICopySubTextureRegion(FRHITexture2D* SourceTextureRHI, FRHITexture2D* DestinationTextureRHI, FBox2D SourceBox, FBox2D DestinationBox)
 {
 	FD3D11Texture2D* SourceTexture = ResourceCast(SourceTextureRHI);
@@ -2254,29 +2242,6 @@ void FD3D11DynamicRHI::RHICopySubTextureRegion_RenderThread(
 	}
 }
 
-void FD3D11DynamicRHI::RHIUpdateTextureReference(FRHITextureReference* TextureRefRHI, FRHITexture* NewTextureRHI)
-{	
-	// Updating texture references is disallowed while the RHI could be caching them in referenced resource tables.
-	check(ResourceTableFrameCounter == INDEX_NONE);
-
-	FD3D11TextureReference* TextureRef = (FD3D11TextureReference*)TextureRefRHI;
-	if (TextureRef)
-	{
-		FD3D11TextureBase* NewTexture = NULL;
-		ID3D11ShaderResourceView* NewSRV = NULL;
-		if (NewTextureRHI)
-		{
-			NewTexture = GetD3D11TextureFromRHITexture(NewTextureRHI);
-			if (NewTexture)
-			{
-				NewSRV = NewTexture->GetShaderResourceView();
-			}
-		}
-		TextureRef->SetReferencedTexture(NewTextureRHI,NewTexture,NewSRV);
-	}
-}
-
-
 template<typename BaseResourceType>
 TD3D11Texture2D<BaseResourceType>* FD3D11DynamicRHI::CreateTextureFromResource(bool bTextureArray, bool bCubeTexture, EPixelFormat Format, ETextureCreateFlags TexCreateFlags, const FClearValueBinding& ClearValueBinding, ID3D11Texture2D* TextureResource)
 {
diff --git a/Engine/Source/Runtime/Windows/D3D11RHI/Public/D3D11Resources.h b/Engine/Source/Runtime/Windows/D3D11RHI/Public/D3D11Resources.h
--- a/Engine/Source/Runtime/Windows/D3D11RHI/Public/D3D11Resources.h
+++ b/Engine/Source/Runtime/Windows/D3D11RHI/Public/D3D11Resources.h
@@ -539,43 +539,6 @@ typedef TD3D11Texture2D<FD3D11BaseTexture2D>      FD3D11Texture2D;
 typedef TD3D11Texture2D<FD3D11BaseTexture2DArray> FD3D11Texture2DArray;
 typedef TD3D11Texture2D<FD3D11BaseTextureCube>    FD3D11TextureCube;
 
-/** Texture reference class. */
-class FD3D11TextureReference : public FRHITextureReference, public FD3D11TextureBase
-{
-public:
-	FD3D11TextureReference(class FD3D11DynamicRHI* InD3DRHI, FLastRenderTimeContainer* LastRenderTime)
-		: FRHITextureReference(LastRenderTime)
-		, FD3D11TextureBase(InD3DRHI,NULL,NULL, 0, false,TArray<TRefCountPtr<ID3D11RenderTargetView> >(),NULL)
-	{
-		BaseShaderResource = NULL;
-	}
-
-	void SetReferencedTexture(FRHITexture* InTexture, FD3D11BaseShaderResource* InBaseShaderResource, ID3D11ShaderResourceView* InSRV)
-	{
-		ShaderResourceView = InSRV;
-		BaseShaderResource = InBaseShaderResource;
-		FRHITextureReference::SetReferencedTexture(InTexture);
-	}
-
-	virtual void* GetTextureBaseRHI() override final
-	{
-		return static_cast<FD3D11TextureBase*>(this);
-	}
-	// IRefCountedObject interface.
-	virtual uint32 AddRef() const
-	{
-		return FRHIResource::AddRef();
-	}
-	virtual uint32 Release() const
-	{
-		return FRHIResource::Release();
-	}
-	virtual uint32 GetRefCount() const
-	{
-		return FRHIResource::GetRefCount();
-	}
-};
-
 /** Given a pointer to a RHI texture that was created by the D3D11 RHI, returns a pointer to the FD3D11TextureBase it encapsulates. */
 FORCEINLINE FD3D11TextureBase* GetD3D11TextureFromRHITexture(FRHITexture* Texture)
 {
