There's no need to print a warning when RHI validation is enabled, a regular log message is sufficient.

CIS was complaining that tests finish with warnings just because validation is enabled.

#rnx

diff --git a/Engine/Source/Runtime/RHI/Private/RHIValidation.cpp b/Engine/Source/Runtime/RHI/Private/RHIValidation.cpp
--- a/Engine/Source/Runtime/RHI/Private/RHIValidation.cpp
+++ b/Engine/Source/Runtime/RHI/Private/RHIValidation.cpp
@@ -74,7 +74,7 @@ FValidationRHI::FValidationRHI(FDynamicRHI* InRHI)
 	: RHI(InRHI)
 {
 	check(RHI);
-	UE_LOG(LogRHI, Warning, TEXT("FValidationRHI on, intercepting %s RHI!"), InRHI && InRHI->GetName() ? InRHI->GetName() : TEXT("<NULL>"));
+	UE_LOG(LogRHI, Log, TEXT("FValidationRHI on, intercepting %s RHI!"), InRHI && InRHI->GetName() ? InRHI->GetName() : TEXT("<NULL>"));
 	GRHIValidationEnabled = true;
 	SeenFailureHashes.Reserve(256);
 }
