Fixed incorrect ordering of RHI validation ops when using parallel command lists.

#rnx

diff --git a/Engine/Source/Runtime/RHI/Private/RHICommandList.cpp b/Engine/Source/Runtime/RHI/Private/RHICommandList.cpp
--- a/Engine/Source/Runtime/RHI/Private/RHICommandList.cpp
+++ b/Engine/Source/Runtime/RHI/Private/RHICommandList.cpp
@@ -705,6 +705,9 @@ void FRHICommandListExecutor::ExecuteList(FRHICommandListImmediate& CmdList)
 		}
 	}
 #endif
+
+	CmdList.FlushValidationOps();
+
 	{
 		SCOPE_CYCLE_COUNTER(STAT_ImmedCmdListExecuteTime);
 		ExecuteInner(CmdList);
diff --git a/Engine/Source/Runtime/RHI/Public/RHICommandList.h b/Engine/Source/Runtime/RHI/Public/RHICommandList.h
--- a/Engine/Source/Runtime/RHI/Public/RHICommandList.h
+++ b/Engine/Source/Runtime/RHI/Public/RHICommandList.h
@@ -2239,6 +2239,15 @@ FRHICOMMAND_MACRO(FRHICommandSetRayTracingBindings)
 };
 #endif // RHI_RAYTRACING
 
+#if ENABLE_RHI_VALIDATION
+
+FRHICOMMAND_MACRO(FRHICommandFlushValidationOps)
+{
+	RHI_API void Execute(FRHICommandListBase& CmdList);
+};
+
+#endif
+
 // Using variadic macro because some types are fancy template<A,B> stuff, which gets broken off at the comma and interpreted as multiple arguments. 
 #define ALLOC_COMMAND(...) new ( AllocCommand(sizeof(__VA_ARGS__), alignof(__VA_ARGS__)) ) __VA_ARGS__
 #define ALLOC_COMMAND_CL(RHICmdList, ...) new ( (RHICmdList).AllocCommand(sizeof(__VA_ARGS__), alignof(__VA_ARGS__)) ) __VA_ARGS__
@@ -4490,7 +4499,19 @@ public:
 		 
 		return GDynamicRHI->RHIFlushResources();
 	}
-	
+
+	FORCEINLINE_DEBUGGABLE void FlushValidationOps()
+	{
+#if ENABLE_RHI_VALIDATION
+		if (Bypass())
+		{
+			GetContext().RHIFlushValidationOps();
+			return;
+		}
+		ALLOC_COMMAND(FRHICommandFlushValidationOps)();
+#endif
+	}
+
 	FORCEINLINE uint32 GetGPUFrameCycles()
 	{
 		return RHIGetGPUFrameCycles(GetGPUMask().ToIndex());
diff --git a/Engine/Source/Runtime/RHI/Public/RHICommandListCommandExecutes.inl b/Engine/Source/Runtime/RHI/Public/RHICommandListCommandExecutes.inl
--- a/Engine/Source/Runtime/RHI/Public/RHICommandListCommandExecutes.inl
+++ b/Engine/Source/Runtime/RHI/Public/RHICommandListCommandExecutes.inl
@@ -594,6 +594,15 @@ void FRHICommandSetRayTracingBindings::Execute(FRHICommandListBase& CmdList)
 
 #endif // RHI_RAYTRACING
 
+#if ENABLE_RHI_VALIDATION
+
+void FRHICommandFlushValidationOps::Execute(FRHICommandListBase& CmdList)
+{
+	INTERNAL_DECORATOR(RHIFlushValidationOps)();
+}
+
+#endif
+
 void FRHIResourceUpdateInfo::ReleaseRefs()
 {
 	switch (Type)
diff --git a/Engine/Source/Runtime/RHI/Public/RHIContext.h b/Engine/Source/Runtime/RHI/Public/RHIContext.h
--- a/Engine/Source/Runtime/RHI/Public/RHIContext.h
+++ b/Engine/Source/Runtime/RHI/Public/RHIContext.h
@@ -751,6 +751,8 @@ public:
 	}
 #endif // #if PLATFORM_REQUIRES_UAV_TO_RTV_TEXTURE_CACHE_FLUSH_WORKAROUND
 
+	virtual void RHIFlushValidationOps() {}
+
 	protected:
 		FRHIRenderPassInfo RenderPassInfo;
 };
diff --git a/Engine/Source/Runtime/RHI/Public/RHIValidationContext.h b/Engine/Source/Runtime/RHI/Public/RHIValidationContext.h
--- a/Engine/Source/Runtime/RHI/Public/RHIValidationContext.h
+++ b/Engine/Source/Runtime/RHI/Public/RHIValidationContext.h
@@ -968,6 +968,11 @@ public:
 		RHIContext->RHISetRayTracingMissShader(Scene, ShaderSlotInScene, Pipeline, ShaderIndexInPipeline, NumUniformBuffers, UniformBuffers, UserData);
 	}
 
+	virtual void RHIFlushValidationOps() override final
+	{
+		RHIValidation::FTracker::ReplayOpQueue(ERHIPipeline::Graphics, Tracker->Finalize());
+	}
+
 	void SetupDrawing()
 	{
 		// nothing to validate right now
