Backed out duplicate fix for RHI validation ops and hand merged prior fix in staging.

#rnx

diff --git a/Engine/Source/Runtime/RHI/Private/RHICommandList.cpp b/Engine/Source/Runtime/RHI/Private/RHICommandList.cpp
--- a/Engine/Source/Runtime/RHI/Private/RHICommandList.cpp
+++ b/Engine/Source/Runtime/RHI/Private/RHICommandList.cpp
@@ -705,9 +705,6 @@ void FRHICommandListExecutor::ExecuteList(FRHICommandListImmediate& CmdList)
 		}
 	}
 #endif
-
-	CmdList.FlushValidationOps();
-
 	{
 		SCOPE_CYCLE_COUNTER(STAT_ImmedCmdListExecuteTime);
 		ExecuteInner(CmdList);
diff --git a/Engine/Source/Runtime/RHI/Public/RHICommandList.h b/Engine/Source/Runtime/RHI/Public/RHICommandList.h
--- a/Engine/Source/Runtime/RHI/Public/RHICommandList.h
+++ b/Engine/Source/Runtime/RHI/Public/RHICommandList.h
@@ -2239,15 +2239,6 @@ FRHICOMMAND_MACRO(FRHICommandSetRayTracingBindings)
 };
 #endif // RHI_RAYTRACING
 
-#if ENABLE_RHI_VALIDATION
-
-FRHICOMMAND_MACRO(FRHICommandFlushValidationOps)
-{
-	RHI_API void Execute(FRHICommandListBase& CmdList);
-};
-
-#endif
-
 // Using variadic macro because some types are fancy template<A,B> stuff, which gets broken off at the comma and interpreted as multiple arguments. 
 #define ALLOC_COMMAND(...) new ( AllocCommand(sizeof(__VA_ARGS__), alignof(__VA_ARGS__)) ) __VA_ARGS__
 #define ALLOC_COMMAND_CL(RHICmdList, ...) new ( (RHICmdList).AllocCommand(sizeof(__VA_ARGS__), alignof(__VA_ARGS__)) ) __VA_ARGS__
@@ -4499,19 +4490,7 @@ public:
 		 
 		return GDynamicRHI->RHIFlushResources();
 	}
-
-	FORCEINLINE_DEBUGGABLE void FlushValidationOps()
-	{
-#if ENABLE_RHI_VALIDATION
-		if (Bypass())
-		{
-			GetContext().RHIFlushValidationOps();
-			return;
-		}
-		ALLOC_COMMAND(FRHICommandFlushValidationOps)();
-#endif
-	}
-
+	
 	FORCEINLINE uint32 GetGPUFrameCycles()
 	{
 		return RHIGetGPUFrameCycles(GetGPUMask().ToIndex());
diff --git a/Engine/Source/Runtime/RHI/Public/RHICommandListCommandExecutes.inl b/Engine/Source/Runtime/RHI/Public/RHICommandListCommandExecutes.inl
--- a/Engine/Source/Runtime/RHI/Public/RHICommandListCommandExecutes.inl
+++ b/Engine/Source/Runtime/RHI/Public/RHICommandListCommandExecutes.inl
@@ -594,15 +594,6 @@ void FRHICommandSetRayTracingBindings::Execute(FRHICommandListBase& CmdList)
 
 #endif // RHI_RAYTRACING
 
-#if ENABLE_RHI_VALIDATION
-
-void FRHICommandFlushValidationOps::Execute(FRHICommandListBase& CmdList)
-{
-	INTERNAL_DECORATOR(RHIFlushValidationOps)();
-}
-
-#endif
-
 void FRHIResourceUpdateInfo::ReleaseRefs()
 {
 	switch (Type)
diff --git a/Engine/Source/Runtime/RHI/Public/RHIContext.h b/Engine/Source/Runtime/RHI/Public/RHIContext.h
--- a/Engine/Source/Runtime/RHI/Public/RHIContext.h
+++ b/Engine/Source/Runtime/RHI/Public/RHIContext.h
@@ -751,8 +751,6 @@ public:
 	}
 #endif // #if PLATFORM_REQUIRES_UAV_TO_RTV_TEXTURE_CACHE_FLUSH_WORKAROUND
 
-	virtual void RHIFlushValidationOps() {}
-
 	protected:
 		FRHIRenderPassInfo RenderPassInfo;
 };
diff --git a/Engine/Source/Runtime/RHI/Public/RHIValidationContext.h b/Engine/Source/Runtime/RHI/Public/RHIValidationContext.h
--- a/Engine/Source/Runtime/RHI/Public/RHIValidationContext.h
+++ b/Engine/Source/Runtime/RHI/Public/RHIValidationContext.h
@@ -968,11 +968,6 @@ public:
 		RHIContext->RHISetRayTracingMissShader(Scene, ShaderSlotInScene, Pipeline, ShaderIndexInPipeline, NumUniformBuffers, UniformBuffers, UserData);
 	}
 
-	virtual void RHIFlushValidationOps() override final
-	{
-		RHIValidation::FTracker::ReplayOpQueue(ERHIPipeline::Graphics, Tracker->Finalize());
-	}
-
 	void SetupDrawing()
 	{
 		// nothing to validate right now
@@ -1056,6 +1051,14 @@ public:
 		using namespace RHIValidation;
 		check(CurrentContext == nullptr);
 
+		if (Index == 0)
+		{
+			// The RHIs are expected to close the main command list before calling the parallel command lists. For the validator,
+			// this means replaying the operations now, before executing the commands from the parallel lists.
+			FValidationContext* DefaultContext = (FValidationContext*)RHIGetDefaultContext();
+			FTracker::ReplayOpQueue(ERHIPipeline::Graphics, DefaultContext->Tracker->Finalize());
+		}
+
 		InnerContainer->SubmitAndFreeContextContainer(Index, Num);
 
 		for (FOperationsList& OpList : CompletedOpLists)
