Added explicit ERHIAccess::BVHRead/Write states for acceleration structures. Added validation for top level acceleration structure barriers.

This only affects the validation layer at this point. Individual RHIs still perform internal transitions for now.

diff --git a/Engine/Source/Runtime/D3D12RHI/Private/D3D12RayTracing.cpp b/Engine/Source/Runtime/D3D12RHI/Private/D3D12RayTracing.cpp
--- a/Engine/Source/Runtime/D3D12RHI/Private/D3D12RayTracing.cpp
+++ b/Engine/Source/Runtime/D3D12RHI/Private/D3D12RayTracing.cpp
@@ -2408,7 +2408,7 @@ static TRefCountPtr<FD3D12Buffer> CreateRayTracingBuffer(FD3D12Adapter* Adapter,
 		Result = Adapter->CreateRHIBuffer(
 			nullptr, BufferDesc, D3D12_RAYTRACING_ACCELERATION_STRUCTURE_BYTE_ALIGNMENT,
 			0, BufferDesc.Width, BUF_AccelerationStructure,
-			ED3D12ResourceStateMode::SingleState, ERHIAccess::RTAccelerationStructure,
+			ED3D12ResourceStateMode::SingleState, ERHIAccess::BVHWrite,
 			CreateInfo, ResourceAllocator);
 	}
 	else if (Type == ERayTracingBufferType::Scratch)
diff --git a/Engine/Source/Runtime/RHI/Private/RHI.cpp b/Engine/Source/Runtime/RHI/Private/RHI.cpp
--- a/Engine/Source/Runtime/RHI/Private/RHI.cpp
+++ b/Engine/Source/Runtime/RHI/Private/RHI.cpp
@@ -133,6 +133,8 @@ FString GetRHIAccessName(ERHIAccess Access)
 			case ERHIAccess::CopyDest:            return TEXT("CopyDest");
 			case ERHIAccess::ResolveDst:          return TEXT("ResolveDst");
 			case ERHIAccess::DSVWrite:            return TEXT("DSVWrite");
+			case ERHIAccess::BVHRead:             return TEXT("BVHRead");
+			case ERHIAccess::BVHWrite:            return TEXT("BVHWrite");
 			case ERHIAccess::Discard:             return TEXT("Discard");
 			}
 		});
diff --git a/Engine/Source/Runtime/RHI/Private/RHIValidation.cpp b/Engine/Source/Runtime/RHI/Private/RHIValidation.cpp
--- a/Engine/Source/Runtime/RHI/Private/RHIValidation.cpp
+++ b/Engine/Source/Runtime/RHI/Private/RHIValidation.cpp
@@ -297,6 +297,10 @@ void FValidationRHI::RHICreateTransition(FRHITransition* Transition, const FRHIT
 		case FRHITransitionInfo::EType::UAV:
 			Identity = Info.UAV->ViewIdentity;
 			break;
+
+		case FRHITransitionInfo::EType::BVH:
+			Identity = Info.BVH->GetWholeResourceIdentity();
+			break;
 		}
 
 		bDoTrace |= (Identity.Resource->LoggingMode != RHIValidation::ELoggingMode::None);
@@ -444,7 +448,7 @@ namespace RHIValidation
 		using T = __underlying_type(ERHIAccess);
 		checkf((T(RequiredAccess) & (T(RequiredAccess) - 1)) == 0, TEXT("Only one required access bit may be set at once."));
 		
-		if (EnumHasAnyFlags(RequiredAccess, ERHIAccess::UAVMask))
+		if (EnumHasAnyFlags(RequiredAccess, ERHIAccess::UAVMask | ERHIAccess::BVHWrite))
 		{
 			// UAV writes decay to no allowed resource access when overlaps are disabled. A barrier is always required after the dispatch/draw.
 			// Otherwise keep the same accessmask and don't touch or decay the state
@@ -1062,7 +1066,7 @@ namespace RHIValidation
 
 		State.Previous = State.Current;
 
-		if (EnumHasAnyFlags(RequiredState.Access, ERHIAccess::UAVMask))
+		if (EnumHasAnyFlags(RequiredState.Access, ERHIAccess::UAVMask | ERHIAccess::BVHWrite))
 		{
 			if (bAllowAllUAVsOverlap) { State.bUsedWithAllUAVsOverlap = true; }
 			if (State.bExplicitAllowUAVOverlap) { State.bUsedWithExplicitUAVsOverlap = true; }
diff --git a/Engine/Source/Runtime/RHI/Public/RHI.h b/Engine/Source/Runtime/RHI/Public/RHI.h
--- a/Engine/Source/Runtime/RHI/Public/RHI.h
+++ b/Engine/Source/Runtime/RHI/Public/RHI.h
@@ -857,10 +857,16 @@ enum class ERHIAccess
 	CopyDest            	= 1 << 12,
 	ResolveDst          	= 1 << 13,
 	DSVWrite            	= 1 << 14,
-	RTAccelerationStructure = 1 << 15,
+
+	// Ray tracing acceleration structure states.
+	// Buffer that contains an AS must always be in either of these states.
+	// BVHRead -- required for AS inputs to build/update/copy/trace commands.
+	// BVHWrite -- required for AS outputs of build/update/copy commands.
+	BVHRead                  = 1 << 15,
+	BVHWrite                 = 1 << 16,
 
 	// Invalid released state (transient resources)
-	Discard					= 1 << 16,
+	Discard					= 1 << 17,
 
 	Last = Discard,
 	None = Unknown,
@@ -888,7 +894,7 @@ enum class ERHIAccess
 	WriteOnlyMask = WriteOnlyExclusiveMask | DSVWrite,
 
 	// A mask of all bits representing writable states which may also include readable states.
-	WritableMask = WriteOnlyMask | UAVMask | RTAccelerationStructure
+	WritableMask = WriteOnlyMask | UAVMask | BVHWrite
 };
 ENUM_CLASS_FLAGS(ERHIAccess)
 
@@ -1911,6 +1917,7 @@ struct FRHITransitionInfo : public FRHISubresourceRange
 		class FRHITexture* Texture;
 		class FRHIBuffer* Buffer;
 		class FRHIUnorderedAccessView* UAV;
+		class FRHIRayTracingAccelerationStructure* BVH;
 	};
 
 	enum class EType : uint8
@@ -1918,7 +1925,8 @@ struct FRHITransitionInfo : public FRHISubresourceRange
 		Unknown,
 		Texture,
 		Buffer,
-		UAV
+		UAV,
+		BVH,
 	} Type = EType::Unknown;
 
 	ERHIAccess AccessBefore = ERHIAccess::Unknown;
@@ -1959,6 +1967,14 @@ struct FRHITransitionInfo : public FRHISubresourceRange
 		, Flags(InFlags)
 	{}
 
+	FRHITransitionInfo(class FRHIRayTracingAccelerationStructure* InBVH, ERHIAccess InPreviousState, ERHIAccess InNewState, EResourceTransitionFlags InFlags = EResourceTransitionFlags::None)
+		: BVH(InBVH)
+		, Type(EType::BVH)
+		, AccessBefore(InPreviousState)
+		, AccessAfter(InNewState)
+		, Flags(InFlags)
+	{}
+
 	inline bool operator == (FRHITransitionInfo const& RHS) const
 	{
 		return Resource == RHS.Resource
diff --git a/Engine/Source/Runtime/RHI/Public/RHIResources.h b/Engine/Source/Runtime/RHI/Public/RHIResources.h
--- a/Engine/Source/Runtime/RHI/Public/RHIResources.h
+++ b/Engine/Source/Runtime/RHI/Public/RHIResources.h
@@ -1533,13 +1533,21 @@ typedef TRefCountPtr<FRHIRayTracingPipelineState> FRayTracingPipelineStateRHIRef
 // Ray tracing resources
 //
 
+class FRHIRayTracingAccelerationStructure
+	: public FRHIResource
+#if ENABLE_RHI_VALIDATION
+	, public RHIValidation::FAccelerationStructureResource
+#endif
+{
+};
+
 /** Bottom level ray tracing acceleration structure (contains triangles). */
-class FRHIRayTracingGeometry : public FRHIResource {};
+class FRHIRayTracingGeometry : public FRHIRayTracingAccelerationStructure {};
 
 typedef TRefCountPtr<FRHIRayTracingGeometry>     FRayTracingGeometryRHIRef;
 
 /** Top level ray tracing acceleration structure (contains instances of meshes). */
-class FRHIRayTracingScene : public FRHIResource
+class FRHIRayTracingScene : public FRHIRayTracingAccelerationStructure
 {
 public:
 	FRHIShaderResourceView* GetShaderResourceView() { return ShaderResourceView; }
diff --git a/Engine/Source/Runtime/RHI/Public/RHIValidation.h b/Engine/Source/Runtime/RHI/Public/RHIValidation.h
--- a/Engine/Source/Runtime/RHI/Public/RHIValidation.h
+++ b/Engine/Source/Runtime/RHI/Public/RHIValidation.h
@@ -1676,12 +1676,16 @@ public:
 #if RHI_RAYTRACING
 	virtual FRayTracingGeometryRHIRef RHICreateRayTracingGeometry(const FRayTracingGeometryInitializer& Initializer) override final
 	{
-		return RHI->RHICreateRayTracingGeometry(Initializer);
+		FRayTracingGeometryRHIRef Result = RHI->RHICreateRayTracingGeometry(Initializer);
+		Result->InitBarrierTracking(ERHIAccess::BVHWrite, *Initializer.DebugName.ToString()); // BVHs are always created in BVHWrite state
+		return Result;
 	}
 
 	virtual FRayTracingSceneRHIRef RHICreateRayTracingScene(const FRayTracingSceneInitializer& Initializer) override final
 	{
-		return RHI->RHICreateRayTracingScene(Initializer);
+		FRayTracingSceneRHIRef Result = RHI->RHICreateRayTracingScene(Initializer);
+		Result->InitBarrierTracking(ERHIAccess::BVHWrite, *Initializer.DebugName.ToString()); // BVHs are always created in BVHWrite state
+		return Result;
 	}
 
 	virtual FRayTracingShaderRHIRef RHICreateRayTracingShader(TArrayView<const uint8> Code, const FSHAHash& Hash, EShaderFrequency ShaderFrequency) override final
diff --git a/Engine/Source/Runtime/RHI/Public/RHIValidationCommon.h b/Engine/Source/Runtime/RHI/Public/RHIValidationCommon.h
--- a/Engine/Source/Runtime/RHI/Public/RHIValidationCommon.h
+++ b/Engine/Source/Runtime/RHI/Public/RHIValidationCommon.h
@@ -369,6 +369,11 @@ namespace RHIValidation
 		}
 	};
 
+	class FAccelerationStructureResource : public FBufferResource
+	{
+	public:
+	};
+
 	class FTextureResource
 	{
 	private:
diff --git a/Engine/Source/Runtime/RHI/Public/RHIValidationContext.h b/Engine/Source/Runtime/RHI/Public/RHIValidationContext.h
--- a/Engine/Source/Runtime/RHI/Public/RHIValidationContext.h
+++ b/Engine/Source/Runtime/RHI/Public/RHIValidationContext.h
@@ -932,11 +932,14 @@ public:
 
 	virtual void RHIBuildAccelerationStructures(const TArrayView<const FAccelerationStructureBuildParams> Params) override final
 	{
+		// #yuriy_todo: explicit transitions and state validation for BLAS
 		RHIContext->RHIBuildAccelerationStructures(Params);
 	}
 
 	virtual void RHIBuildAccelerationStructure(FRHIRayTracingScene* Scene) override final
 	{
+		// #yuriy_todo: validate all referenced BLAS states
+		Tracker->Assert(Scene->GetWholeResourceIdentity(), ERHIAccess::BVHWrite);
 		RHIContext->RHIBuildAccelerationStructure(Scene);
 	}
 
@@ -945,6 +948,8 @@ public:
 		FRHIUnorderedAccessView* Output,
 		uint32 NumRays) override final
 	{
+		// #yuriy_todo: validate all referenced BLAS states
+		Tracker->Assert(Scene->GetWholeResourceIdentity(), ERHIAccess::BVHRead);
 		RHIContext->RHIRayTraceOcclusion(Scene, Rays, Output, NumRays);
 	}
 
@@ -953,6 +958,8 @@ public:
 		FRHIUnorderedAccessView* Output,
 		uint32 NumRays) override final
 	{
+		// #yuriy_todo: validate all referenced BLAS states
+		Tracker->Assert(Scene->GetWholeResourceIdentity(), ERHIAccess::BVHRead);
 		RHIContext->RHIRayTraceIntersection(Scene, Rays, Output, NumRays);
 	}
 
@@ -961,6 +968,8 @@ public:
 		const FRayTracingShaderBindings& GlobalResourceBindings,
 		uint32 Width, uint32 Height) override final
 	{
+		// #yuriy_todo: validate all referenced BLAS states
+		Tracker->Assert(Scene->GetWholeResourceIdentity(), ERHIAccess::BVHRead);
 		RHIContext->RHIRayTraceDispatch(RayTracingPipelineState, RayGenShader, Scene, GlobalResourceBindings, Width, Height);
 	}
 
@@ -969,8 +978,10 @@ public:
 		const FRayTracingShaderBindings& GlobalResourceBindings,
 		FRHIBuffer* ArgumentBuffer, uint32 ArgumentOffset) override final
 	{
+		// #yuriy_todo: validate all referenced BLAS states
 		Tracker->Assert(ArgumentBuffer->GetWholeResourceIdentity(), ERHIAccess::IndirectArgs);
 		Tracker->Assert(ArgumentBuffer->GetWholeResourceIdentity(), ERHIAccess::SRVCompute);
+		Tracker->Assert(Scene->GetWholeResourceIdentity(), ERHIAccess::BVHRead);
 
 		RHIContext->RHIRayTraceDispatchIndirect(RayTracingPipelineState, RayGenShader, Scene, GlobalResourceBindings, ArgumentBuffer, ArgumentOffset);
 	}
diff --git a/Engine/Source/Runtime/Renderer/Private/DeferredShadingRenderer.cpp b/Engine/Source/Runtime/Renderer/Private/DeferredShadingRenderer.cpp
--- a/Engine/Source/Runtime/Renderer/Private/DeferredShadingRenderer.cpp
+++ b/Engine/Source/Runtime/Renderer/Private/DeferredShadingRenderer.cpp
@@ -1176,6 +1176,7 @@ bool FDeferredShadingSceneRenderer::DispatchRayTracingWorldUpdates(FRDGBuilder&
 #endif // DO_CHECK
 
 		FRayTracingSceneInitializer SceneInitializer;
+		SceneInitializer.DebugName = FName(TEXT("FDeferredShadingSceneRenderer.Views[].RayTracingScene"));
 		SceneInitializer.Instances = View.RayTracingGeometryInstances;
 		SceneInitializer.ShaderSlotsPerGeometrySegment = RAY_TRACING_NUM_SHADER_SLOTS;
 		SceneInitializer.NumMissShaderSlots = RAY_TRACING_NUM_MISS_SHADER_SLOTS;
@@ -1431,6 +1432,12 @@ void FDeferredShadingSceneRenderer::WaitForRayTracingScene(FRDGBuilder& GraphBui
 			RHICmdList.EndTransition(RayTracingDynamicGeometryUpdateEndTransition);
 			RayTracingDynamicGeometryUpdateEndTransition = nullptr;
 		}
+
+		for (int32 ViewIndex = 0; ViewIndex < Views.Num(); ++ViewIndex)
+		{
+			FRHIRayTracingScene* RayTracingScene = Views[ViewIndex].RayTracingScene.RayTracingSceneRHI.GetReference();
+			RHICmdList.Transition(FRHITransitionInfo(RayTracingScene, ERHIAccess::BVHWrite, ERHIAccess::BVHRead));
+		}
 	});
 }
 
diff --git a/Engine/Source/Runtime/Renderer/Private/Tests/RayTracingTestbed.cpp b/Engine/Source/Runtime/Renderer/Private/Tests/RayTracingTestbed.cpp
--- a/Engine/Source/Runtime/Renderer/Private/Tests/RayTracingTestbed.cpp
+++ b/Engine/Source/Runtime/Renderer/Private/Tests/RayTracingTestbed.cpp
@@ -120,8 +120,14 @@ bool RunRayTracingTestbed_RenderThread(const FString& Parameters)
 	FRHICommandListImmediate& RHICmdList = FRHICommandListExecutor::GetImmediateCommandList();
 
 	RHICmdList.BuildAccelerationStructure(Geometry);
+
+	// #yuriy_todo: explicit transitions and state validation for BLAS
+	// RHICmdList.Transition(FRHITransitionInfo(Geometry.GetReference(), ERHIAccess::BVHWrite, ERHIAccess::BVHRead));
+
 	RHICmdList.BuildAccelerationStructure(Scene);
 
+	RHICmdList.Transition(FRHITransitionInfo(Scene.GetReference(), ERHIAccess::BVHWrite, ERHIAccess::BVHRead));
+
 	RHICmdList.RayTraceOcclusion(Scene, RayBufferView, OcclusionResultBufferView, NumRays);
 	RHICmdList.RayTraceIntersection(Scene, RayBufferView, IntersectionResultBufferView, NumRays);
 
