Highlight the fact that we're having fresh memory buffer initialization performance problem in Insights

diff --git a/Engine/Source/Runtime/RHI/Private/RHICommandList.cpp b/Engine/Source/Runtime/RHI/Private/RHICommandList.cpp
--- a/Engine/Source/Runtime/RHI/Private/RHICommandList.cpp
+++ b/Engine/Source/Runtime/RHI/Private/RHICommandList.cpp
@@ -2070,7 +2070,12 @@ void FDynamicRHI::RHIUnlockBuffer(class FRHICommandListImmediate& RHICmdList, FR
 			{
 				QUICK_SCOPE_CYCLE_COUNTER(STAT_FRHICommandUpdateBuffer_Execute);
 				void* Data = GDynamicRHI->LockBuffer_BottomOfPipe(InRHICmdList, Buffer, Params.Offset, Params.BufferSize, RLM_WriteOnly);
-				FMemory::Memcpy(Data, Params.Buffer, Params.BufferSize);
+				{
+					// If we spend a long time doing this memcpy, it means we got freshly allocated memory from the OS that has never been
+					// initialized and is causing pagefault to bring zeroed pages into our process.
+					TRACE_CPUPROFILER_EVENT_SCOPE(RHIUnlockBuffer_Memcpy);
+					FMemory::Memcpy(Data, Params.Buffer, Params.BufferSize);
+				}
 				FMemory::Free(Params.Buffer);
 				GDynamicRHI->UnlockBuffer_BottomOfPipe(InRHICmdList, Buffer);
 			});
