[Insights] Changed trace event for second GPU to use same "GpuProfiler" logger (instead of a new "GpuProfiler2" logger).

diff --git a/Engine/Source/Developer/TraceServices/Private/Analyzers/GpuProfilerTraceAnalysis.cpp b/Engine/Source/Developer/TraceServices/Private/Analyzers/GpuProfilerTraceAnalysis.cpp
--- a/Engine/Source/Developer/TraceServices/Private/Analyzers/GpuProfilerTraceAnalysis.cpp
+++ b/Engine/Source/Developer/TraceServices/Private/Analyzers/GpuProfilerTraceAnalysis.cpp
@@ -21,7 +21,7 @@ void FGpuProfilerAnalyzer::OnAnalysisBegin(const FOnAnalysisContext& Context)
 
 	Builder.RouteEvent(RouteId_EventSpec, "GpuProfiler", "EventSpec");
 	Builder.RouteEvent(RouteId_Frame, "GpuProfiler", "Frame");
-	Builder.RouteEvent(RouteId_Frame2, "GpuProfiler2", "Frame2");
+	Builder.RouteEvent(RouteId_Frame2, "GpuProfiler", "Frame2");
 }
 
 bool FGpuProfilerAnalyzer::OnEvent(uint16 RouteId, EStyle Style, const FOnEventContext& Context)
diff --git a/Engine/Source/Developer/TraceServices/Private/Modules/TimingProfilerModule.cpp b/Engine/Source/Developer/TraceServices/Private/Modules/TimingProfilerModule.cpp
--- a/Engine/Source/Developer/TraceServices/Private/Modules/TimingProfilerModule.cpp
+++ b/Engine/Source/Developer/TraceServices/Private/Modules/TimingProfilerModule.cpp
@@ -36,7 +36,6 @@ void FTimingProfilerModule::GetLoggers(TArray<const TCHAR *>& OutLoggers)
 {
 	OutLoggers.Add(TEXT("CpuProfiler"));
 	OutLoggers.Add(TEXT("GpuProfiler"));
-	OutLoggers.Add(TEXT("GpuProfiler2"));
 }
 
 const ITimingProfilerProvider* ReadTimingProfilerProvider(const IAnalysisSession& Session)
diff --git a/Engine/Source/Runtime/RHI/Private/GpuProfilerTrace.cpp b/Engine/Source/Runtime/RHI/Private/GpuProfilerTrace.cpp
--- a/Engine/Source/Runtime/RHI/Private/GpuProfilerTrace.cpp
+++ b/Engine/Source/Runtime/RHI/Private/GpuProfilerTrace.cpp
@@ -47,7 +47,7 @@ UE_TRACE_EVENT_BEGIN(GpuProfiler, Frame)
 UE_TRACE_EVENT_END()
 
 // GPU Index 1
-UE_TRACE_EVENT_BEGIN(GpuProfiler2, Frame2)
+UE_TRACE_EVENT_BEGIN(GpuProfiler, Frame2)
 	UE_TRACE_EVENT_FIELD(uint64, CalibrationBias)
 	UE_TRACE_EVENT_FIELD(uint64, TimestampBase)
 	UE_TRACE_EVENT_FIELD(uint32, RenderingFrameNumber)
@@ -177,7 +177,7 @@ void FGpuProfilerTrace::EndFrame(uint32 GPUIndex)
 		}
 		else if (GPUIndex == 1)
 		{
-			UE_TRACE_LOG(GpuProfiler2, Frame2, GpuChannel)
+			UE_TRACE_LOG(GpuProfiler, Frame2, GpuChannel)
 				<< Frame2.CalibrationBias(Bias)
 				<< Frame2.TimestampBase(GCurrentFrame.TimestampBase)
 				<< Frame2.RenderingFrameNumber(GCurrentFrame.RenderingFrameNumber)
