Fix for race in RHI validation due to transient allocator re-initializing barrier tracking on the render thread.

diff --git a/Engine/Source/Runtime/D3D12RHI/Private/D3D12TransientResourceAllocator.cpp b/Engine/Source/Runtime/D3D12RHI/Private/D3D12TransientResourceAllocator.cpp
--- a/Engine/Source/Runtime/D3D12RHI/Private/D3D12TransientResourceAllocator.cpp
+++ b/Engine/Source/Runtime/D3D12RHI/Private/D3D12TransientResourceAllocator.cpp
@@ -731,7 +731,7 @@ void FD3D12TransientResourceAllocator::DeallocateMemory(FD3D12BaseShaderResource
 }
 
 
-void FD3D12TransientResourceAllocator::Freeze()
+void FD3D12TransientResourceAllocator::Freeze(FRHICommandListImmediate&)
 {
 	// all memory should have been freed again - so release back to the manager
 	// can already be reused by the next transient allocator
diff --git a/Engine/Source/Runtime/D3D12RHI/Private/D3D12TransientResourceAllocator.h b/Engine/Source/Runtime/D3D12RHI/Private/D3D12TransientResourceAllocator.h
--- a/Engine/Source/Runtime/D3D12RHI/Private/D3D12TransientResourceAllocator.h
+++ b/Engine/Source/Runtime/D3D12RHI/Private/D3D12TransientResourceAllocator.h
@@ -190,7 +190,7 @@ public:
 	virtual FRHIBuffer* CreateBuffer(const FRHIBufferCreateInfo& InCreateInfo, const TCHAR* InDebugName) override;
 	virtual void DeallocateMemory(FRHITexture* InTexture) override;
 	virtual void DeallocateMemory(FRHIBuffer* InBuffer) override;
-	virtual void Freeze() override;
+	virtual void Freeze(FRHICommandListImmediate&) override;
 
 	// Override deallocation function - memory possibly already freed during DeallocateMemory
 	virtual void DeallocateResource(FD3D12ResourceLocation& ResourceLocation) override;
diff --git a/Engine/Source/Runtime/RHI/Private/RHIValidation.cpp b/Engine/Source/Runtime/RHI/Private/RHIValidation.cpp
--- a/Engine/Source/Runtime/RHI/Private/RHIValidation.cpp
+++ b/Engine/Source/Runtime/RHI/Private/RHIValidation.cpp
@@ -1406,7 +1406,7 @@ FValidationTransientResourceAllocator::~FValidationTransientResourceAllocator()
 				const AllocatedResourceData& ResourceData = It->Value;
 				if (RHIResource->GetRefCount() > 1)
 				{
-					ErrorMessage += FString::Printf(TEXT("         %s (%s) - Refcount: %d\n"), *ResourceData.DebugName, ResourceData.ResourceType == AllocatedResourceData::EType::Texture ? TEXT("Texture") : TEXT("DataBuffer"), RHIResource->GetRefCount());
+					ErrorMessage += FString::Printf(TEXT("         %s (%s) - Refcount: %d\n"), ResourceData.DebugName, ResourceData.ResourceType == AllocatedResourceData::EType::Texture ? TEXT("Texture") : TEXT("DataBuffer"), RHIResource->GetRefCount());
 				}
 			}
 			ErrorMessage += FString::Printf(TRANSIENT_RESOURCE_LOG_SUFFIX);
@@ -1425,14 +1425,27 @@ FRHITexture* FValidationTransientResourceAllocator::CreateTexture(const FRHIText
 
 	FRHITexture* RHITexture = RHIAllocator->CreateTexture(InCreateInfo, InDebugName);
 
-	// Setup for resource tracking
-	RHITexture->InitBarrierTracking(InCreateInfo.NumMips, InCreateInfo.ArraySize, InCreateInfo.Format, InCreateInfo.Flags, ERHIAccess::Discard, InDebugName);
+	bool bReinitializeTrackerResource = true;
+
+	if (RHIValidation::FResource* Resource = RHITexture->GetTrackerResource())
+	{
+		if (!Resource->IsBarrierTrackingInitialized())
+		{
+			RHITexture->InitBarrierTracking(InCreateInfo.NumMips, InCreateInfo.ArraySize, InCreateInfo.Format, InCreateInfo.Flags, ERHIAccess::Discard, InDebugName);
+			bReinitializeTrackerResource = false;
+		}
+	}
 
 	// Store allocation data
 	AllocatedResourceData ResourceData;
 	ResourceData.DebugName = InDebugName;
 	ResourceData.ResourceType = AllocatedResourceData::EType::Texture;
 	ResourceData.bMemoryAllocated = true;
+	ResourceData.bReinitializeBarrierTracking = bReinitializeTrackerResource;
+	ResourceData.Texture.Flags = InCreateInfo.Flags;
+	ResourceData.Texture.Format = InCreateInfo.Format;
+	ResourceData.Texture.ArraySize = InCreateInfo.ArraySize;
+	ResourceData.Texture.NumMips = InCreateInfo.NumMips;
 	AllocatedResourceMap.Add(RHITexture, ResourceData);
 
 	return RHITexture;
@@ -1442,17 +1455,23 @@ FRHIBuffer* FValidationTransientResourceAllocator::CreateBuffer(const FRHIBuffer
 {
 	check(!bFrozen);
 	checkf(!EnumHasAnyFlags(InCreateInfo.Usage, BUF_Transient), TEXT("Attempted to create transient buffer %s using legacy BUF_Transient flag."), InDebugName);
-
+	
 	FRHIBuffer* RHIBuffer = RHIAllocator->CreateBuffer(InCreateInfo, InDebugName);
 
-	// Setup for resource tracking
-	RHIBuffer->InitBarrierTracking(ERHIAccess::Discard, InDebugName);
-	
+	bool bReinitializeTrackerResource = true;
+
+	if (!RHIBuffer->IsBarrierTrackingInitialized())
+	{
+		RHIBuffer->InitBarrierTracking(ERHIAccess::Discard, InDebugName);
+		bReinitializeTrackerResource = false;
+	}
+
 	// Store allocation data
 	AllocatedResourceData ResourceData;
 	ResourceData.DebugName = InDebugName;
 	ResourceData.ResourceType = AllocatedResourceData::EType::Buffer;
 	ResourceData.bMemoryAllocated = true;
+	ResourceData.bReinitializeBarrierTracking = bReinitializeTrackerResource;
 	AllocatedResourceMap.Add(RHIBuffer, ResourceData);
 
 	return RHIBuffer;
@@ -1482,7 +1501,7 @@ void FValidationTransientResourceAllocator::DeallocateMemory(FRHIBuffer* InBuffe
 	ResourceData->bMemoryAllocated = false;
 }
 
-void FValidationTransientResourceAllocator::Freeze()
+void FValidationTransientResourceAllocator::Freeze(FRHICommandListImmediate& RHICmdList)
 {
 	check(!bFrozen);
 
@@ -1506,16 +1525,56 @@ void FValidationTransientResourceAllocator::Freeze()
 			const AllocatedResourceData& ResourceData = It->Value;
 			if (ResourceData.bMemoryAllocated)
 			{
-				ErrorMessage += FString::Printf(TEXT("         %s (%s)\n"), *ResourceData.DebugName, ResourceData.ResourceType == AllocatedResourceData::EType::Texture ? TEXT("Texture") : TEXT("DataBuffer"));
+				ErrorMessage += FString::Printf(TEXT("         %s (%s)\n"), ResourceData.DebugName, ResourceData.ResourceType == AllocatedResourceData::EType::Texture ? TEXT("Texture") : TEXT("DataBuffer"));
 			}
 		}
 		ErrorMessage += FString::Printf(TRANSIENT_RESOURCE_LOG_SUFFIX);
 		FValidationRHI::ReportValidationFailure(*ErrorMessage);
 	}
 
+	RHICmdList.EnqueueLambda([this](FRHICommandListImmediate& InRHICmdList)
+	{
+		// Tracking will be re-initialized, so we need to flush any remaining references.
+		static_cast<FValidationContext&>(InRHICmdList.GetContext()).FlushValidationOps();
+
+		InitBarrierTracking();
+	});
+
 	bFrozen = true;
 
-	RHIAllocator->Freeze();
+	RHIAllocator->Freeze(RHICmdList);
+}
+
+void FValidationTransientResourceAllocator::InitBarrierTracking()
+{
+	// Barrier tracking initialization has to happen on the RHI thread, because RHI resources are pooled and reused.
+
+	for (const auto& Entry : AllocatedResourceMap)
+	{
+		FRHIResource* Resource = Entry.Key;
+		const AllocatedResourceData& ResourceData = Entry.Value;
+
+		if (!ResourceData.bReinitializeBarrierTracking)
+		{
+			continue;
+		}
+
+		switch (ResourceData.ResourceType)
+		{
+		case AllocatedResourceData::EType::Texture:
+		{
+			FRHITexture* Texture = static_cast<FRHITexture*>(Resource);
+			Texture->InitBarrierTracking(ResourceData.Texture.NumMips, ResourceData.Texture.ArraySize, ResourceData.Texture.Format, ResourceData.Texture.Flags, ERHIAccess::Discard, ResourceData.DebugName);
+		}
+		break;
+		case AllocatedResourceData::EType::Buffer:
+		{
+			FRHIBuffer* Buffer = static_cast<FRHIBuffer*>(Resource);
+			Buffer->InitBarrierTracking(ERHIAccess::Discard, ResourceData.DebugName);
+		}
+		break;
+		}
+	}
 }
 
 #endif	// ENABLE_RHI_VALIDATION
diff --git a/Engine/Source/Runtime/RHI/Public/RHICommandList.h b/Engine/Source/Runtime/RHI/Public/RHICommandList.h
--- a/Engine/Source/Runtime/RHI/Public/RHICommandList.h
+++ b/Engine/Source/Runtime/RHI/Public/RHICommandList.h
@@ -2494,16 +2494,6 @@ public:
 		ALLOC_COMMAND(FRHICommandClearUAVUint)(UnorderedAccessViewRHI, Values);
 	}
 
-	FORCEINLINE_DEBUGGABLE void ReleaseTransientResourceAllocator(IRHITransientResourceAllocator* InAllocator)
-	{
-		if (Bypass())
-		{
-			GetComputeContext().RHIReleaseTransientResourceAllocator(InAllocator);
-			return;
-		}
-		ALLOC_COMMAND(FRHICommandReleaseTransientResourceAllocator)(InAllocator);
-	}
-
 	FORCEINLINE_DEBUGGABLE void BeginTransitions(TArrayView<const FRHITransition*> Transitions)
 	{
 		if (Bypass())
@@ -4490,6 +4480,16 @@ public:
 		return GDynamicRHI->RHICreateRenderQuery_RenderThread(*this, QueryType);
 	}
 
+	FORCEINLINE_DEBUGGABLE void ReleaseTransientResourceAllocator(IRHITransientResourceAllocator* InAllocator)
+	{
+		check(InAllocator);
+		if (Bypass())
+		{
+			GetContext().RHIReleaseTransientResourceAllocator(InAllocator);
+			return;
+		}
+		ALLOC_COMMAND(FRHICommandReleaseTransientResourceAllocator)(InAllocator);
+	}
 
 	FORCEINLINE void AcquireTransientResource_RenderThread(FRHITexture* Texture)
 	{
diff --git a/Engine/Source/Runtime/RHI/Public/RHICommandListCommandExecutes.inl b/Engine/Source/Runtime/RHI/Public/RHICommandListCommandExecutes.inl
--- a/Engine/Source/Runtime/RHI/Public/RHICommandListCommandExecutes.inl
+++ b/Engine/Source/Runtime/RHI/Public/RHICommandListCommandExecutes.inl
@@ -377,7 +377,7 @@ void FRHICommandClearUAVUint::Execute(FRHICommandListBase& CmdList)
 void FRHICommandReleaseTransientResourceAllocator::Execute(FRHICommandListBase& CmdList)
 {
 	RHISTAT(ReleaseTransientResourceAllocator);
-	INTERNAL_DECORATOR_COMPUTE(RHIReleaseTransientResourceAllocator)(Allocator);
+	INTERNAL_DECORATOR(RHIReleaseTransientResourceAllocator)(Allocator);
 }
 
 void FRHICommandCopyToResolveTarget::Execute(FRHICommandListBase& CmdList)
diff --git a/Engine/Source/Runtime/RHI/Public/RHIResources.h b/Engine/Source/Runtime/RHI/Public/RHIResources.h
--- a/Engine/Source/Runtime/RHI/Public/RHIResources.h
+++ b/Engine/Source/Runtime/RHI/Public/RHIResources.h
@@ -3121,8 +3121,9 @@ struct FRHIBufferCreateInfo
 	   The allocator manages the lifetime of all the created resources and will take care of destroying
 	   all the resources when the allocator itself is destroyed.	   
 **/
-struct RHI_API IRHITransientResourceAllocator
+class RHI_API IRHITransientResourceAllocator
 {
+public:
 	// Virtual destructor
 	virtual ~IRHITransientResourceAllocator() {}
 
@@ -3135,5 +3136,5 @@ struct RHI_API IRHITransientResourceAllocator
 	virtual void DeallocateMemory(FRHIBuffer* InBuffer) = 0;
 
 	// Freeze allocator for now creation and validates that all resources have their memory deallocated
-	virtual void Freeze() = 0;
+	virtual void Freeze(class FRHICommandListImmediate& RHICmdList) = 0;
 };
\ No newline at end of file
diff --git a/Engine/Source/Runtime/RHI/Public/RHIValidationContext.h b/Engine/Source/Runtime/RHI/Public/RHIValidationContext.h
--- a/Engine/Source/Runtime/RHI/Public/RHIValidationContext.h
+++ b/Engine/Source/Runtime/RHI/Public/RHIValidationContext.h
@@ -1050,6 +1050,11 @@ public:
 		PlatformContext->Tracker = &State.TrackerInstance;
 	}
 
+	inline void FlushValidationOps()
+	{
+		RHIValidation::FTracker::ReplayOpQueue(ERHIPipeline::Graphics, Tracker->Finalize());
+	}
+
 protected:
 	struct FState
 	{
diff --git a/Engine/Source/Runtime/RHI/Public/RHIValidationTransientResourceAllocator.h b/Engine/Source/Runtime/RHI/Public/RHIValidationTransientResourceAllocator.h
--- a/Engine/Source/Runtime/RHI/Public/RHIValidationTransientResourceAllocator.h
+++ b/Engine/Source/Runtime/RHI/Public/RHIValidationTransientResourceAllocator.h
@@ -21,10 +21,13 @@ public:
 	virtual FRHIBuffer* CreateBuffer(const FRHIBufferCreateInfo& InCreateInfo, const TCHAR* InDebugName) override final;
 	virtual void DeallocateMemory(FRHITexture* InTexture) override final;
 	virtual void DeallocateMemory(FRHIBuffer* InBuffer) override final;
-	virtual void Freeze() override final;
+	virtual void Freeze(FRHICommandListImmediate&) override final;
 
 private:
 
+	friend class FValidationContext;
+	void InitBarrierTracking();
+
 	// Actual RHI transient allocator which will get all functions forwarded
 	IRHITransientResourceAllocator* RHIAllocator = nullptr;
 
@@ -40,9 +43,18 @@ private:
 			Buffer,
 		};
 
-		FString DebugName;
-		EType ResourceType;
-		bool bMemoryAllocated;
+		const TCHAR* DebugName = nullptr;
+		EType ResourceType = EType::Texture;
+		bool bMemoryAllocated = false;
+		bool bReinitializeBarrierTracking = false;
+
+		struct FTexture
+		{
+			ETextureCreateFlags Flags = TexCreate_None;
+			EPixelFormat Format = PF_Unknown;
+			uint16 ArraySize = 0;
+			uint8 NumMips = 0;
+		} Texture;
 	};
 	TMap<FRHIResource*, AllocatedResourceData> AllocatedResourceMap;
 };
diff --git a/Engine/Source/Runtime/RenderCore/Private/RenderGraphBuilder.cpp b/Engine/Source/Runtime/RenderCore/Private/RenderGraphBuilder.cpp
--- a/Engine/Source/Runtime/RenderCore/Private/RenderGraphBuilder.cpp
+++ b/Engine/Source/Runtime/RenderCore/Private/RenderGraphBuilder.cpp
@@ -1391,7 +1391,7 @@ void FRDGBuilder::Execute()
 
 			if (TransientResourceAllocator)
 			{
-				TransientResourceAllocator->Freeze();
+				TransientResourceAllocator->Freeze(RHICmdList);
 			}
 		}
 
