Allow FRWBufferStructured to supply a default access level

diff --git a/Engine/Source/Runtime/RHI/Public/RHIUtilities.h b/Engine/Source/Runtime/RHI/Public/RHIUtilities.h
--- a/Engine/Source/Runtime/RHI/Public/RHIUtilities.h
+++ b/Engine/Source/Runtime/RHI/Public/RHIUtilities.h
@@ -388,7 +388,7 @@ struct FRWBufferStructured
 		Release();
 	}
 
-	void Initialize(const TCHAR* InDebugName, uint32 BytesPerElement, uint32 NumElements, uint32 AdditionalUsage = 0, bool bUseUavCounter = false, bool bAppendBuffer = false)
+	void Initialize(const TCHAR* InDebugName, uint32 BytesPerElement, uint32 NumElements, uint32 AdditionalUsage = 0, bool bUseUavCounter = false, bool bAppendBuffer = false, ERHIAccess InResourceState = ERHIAccess::UAVMask)
 	{
 		check(GMaxRHIFeatureLevel >= ERHIFeatureLevel::SM5 || GMaxRHIFeatureLevel == ERHIFeatureLevel::ES3_1);
 		// Provide a debug name if using Fast VRAM so the allocators diagnostics will work
@@ -396,7 +396,7 @@ struct FRWBufferStructured
 
 		NumBytes = BytesPerElement * NumElements;
 		FRHIResourceCreateInfo CreateInfo(InDebugName);
-		Buffer = RHICreateStructuredBuffer(BytesPerElement, NumBytes, BUF_UnorderedAccess | BUF_ShaderResource | AdditionalUsage, ERHIAccess::UAVMask, CreateInfo);
+		Buffer = RHICreateStructuredBuffer(BytesPerElement, NumBytes, BUF_UnorderedAccess | BUF_ShaderResource | AdditionalUsage, InResourceState, CreateInfo);
 		UAV = RHICreateUnorderedAccessView(Buffer, bUseUavCounter, bAppendBuffer);
 		SRV = RHICreateShaderResourceView(Buffer);
 	}
