Renamed FAccelerationStructureBuildParams -> FRayTracingGeometryBuildParams, added FRayTracingSceneBuildParams for advanced scene BVH build options.

diff --git a/Engine/Plugins/Runtime/HairStrands/Source/HairStrandsCore/Private/HairStrandsRendering.cpp b/Engine/Plugins/Runtime/HairStrands/Source/HairStrandsCore/Private/HairStrandsRendering.cpp
--- a/Engine/Plugins/Runtime/HairStrands/Source/HairStrandsCore/Private/HairStrandsRendering.cpp
+++ b/Engine/Plugins/Runtime/HairStrands/Source/HairStrandsCore/Private/HairStrandsRendering.cpp
@@ -1103,7 +1103,7 @@ static void UpdateHairAccelerationStructure(FRHICommandList& RHICmdList, FRayTra
 {
 	SCOPED_DRAW_EVENT(RHICmdList, CommitHairRayTracingGeometryUpdates);
 
-	FAccelerationStructureBuildParams Params;
+	FRayTracingGeometryBuildParams Params;
 	Params.BuildMode = EAccelerationStructureBuildMode::Update;
 	Params.Geometry = RayTracingGeometry->RayTracingGeometryRHI;
 	Params.Segments = RayTracingGeometry->Initializer.Segments;
diff --git a/Engine/Source/Runtime/D3D12RHI/Private/D3D12CommandContext.h b/Engine/Source/Runtime/D3D12RHI/Private/D3D12CommandContext.h
--- a/Engine/Source/Runtime/D3D12RHI/Private/D3D12CommandContext.h
+++ b/Engine/Source/Runtime/D3D12RHI/Private/D3D12CommandContext.h
@@ -378,8 +378,8 @@ public:
 	virtual void RHICopyBufferRegion(FRHIBuffer* DestBuffer, uint64 DstOffset, FRHIBuffer* SourceBuffer, uint64 SrcOffset, uint64 NumBytes) final override;
 	virtual void RHICopyBufferRegions(const TArrayView<const FCopyBufferRegionParams> Params) final override;
 	virtual void RHIBindAccelerationStructureMemory(FRHIRayTracingScene* Scene, FRHIBuffer* Buffer, uint32 BufferOffset) final override;
-	virtual void RHIBuildAccelerationStructures(const TArrayView<const FAccelerationStructureBuildParams> Params) final override;
-	virtual void RHIBuildAccelerationStructure(FRHIRayTracingScene* Scene) final override;
+	virtual void RHIBuildAccelerationStructures(const TArrayView<const FRayTracingGeometryBuildParams> Params) final override;
+	virtual void RHIBuildAccelerationStructure(const FRayTracingSceneBuildParams& SceneBuildParams) final override;
 	virtual void RHIClearRayTracingBindings(FRHIRayTracingScene* Scene) final override;
 	virtual void RHIRayTraceOcclusion(FRHIRayTracingScene* Scene,
 		FRHIShaderResourceView* Rays,
@@ -729,14 +729,14 @@ public:
 		ContextRedirect(RHIEndRenderPass());
 	}
 
-	virtual void RHIBuildAccelerationStructures(const TArrayView<const FAccelerationStructureBuildParams> Params) final override
+	virtual void RHIBuildAccelerationStructures(const TArrayView<const FRayTracingGeometryBuildParams> Params) final override
 	{
 		ContextRedirect(RHIBuildAccelerationStructures(Params));
 	}
 
-	virtual void RHIBuildAccelerationStructure(FRHIRayTracingScene* Scene) final override
+	virtual void RHIBuildAccelerationStructure(const FRayTracingSceneBuildParams& SceneBuildParams) final override
 	{
-		ContextRedirect(RHIBuildAccelerationStructure(Scene));
+		ContextRedirect(RHIBuildAccelerationStructure(SceneBuildParams));
 	}
 
 	virtual void RHIRayTraceOcclusion(FRHIRayTracingScene* Scene,
diff --git a/Engine/Source/Runtime/D3D12RHI/Private/D3D12RayTracing.cpp b/Engine/Source/Runtime/D3D12RHI/Private/D3D12RayTracing.cpp
--- a/Engine/Source/Runtime/D3D12RHI/Private/D3D12RayTracing.cpp
+++ b/Engine/Source/Runtime/D3D12RHI/Private/D3D12RayTracing.cpp
@@ -3593,11 +3593,11 @@ FD3D12RayTracingShaderTable* FD3D12RayTracingScene::FindOrCreateShaderTable(cons
 	return CreatedShaderTable;
 }
 
-void FD3D12CommandContext::RHIBuildAccelerationStructures(const TArrayView<const FAccelerationStructureBuildParams> Params)
+void FD3D12CommandContext::RHIBuildAccelerationStructures(const TArrayView<const FRayTracingGeometryBuildParams> Params)
 {
 	// Update geometry vertex buffers
 
-	for (const FAccelerationStructureBuildParams& P : Params)
+	for (const FRayTracingGeometryBuildParams& P : Params)
 	{
 		FD3D12RayTracingGeometry* Geometry = FD3D12DynamicRHI::ResourceCast(P.Geometry.GetReference());
 		Geometry->UnregisterAsRenameListener(GetGPUIndex());
@@ -3624,7 +3624,7 @@ void FD3D12CommandContext::RHIBuildAccelerationStructures(const TArrayView<const
 
 	// Transition all VBs and IBs to readable state
 
-	for (const FAccelerationStructureBuildParams& P : Params)
+	for (const FRayTracingGeometryBuildParams& P : Params)
 	{
 		FD3D12RayTracingGeometry* Geometry = FD3D12DynamicRHI::ResourceCast(P.Geometry.GetReference());
 		Geometry->TransitionBuffers(*this);
@@ -3634,7 +3634,7 @@ void FD3D12CommandContext::RHIBuildAccelerationStructures(const TArrayView<const
 
 	// Then do all work
 
-	for (const FAccelerationStructureBuildParams& P : Params)
+	for (const FRayTracingGeometryBuildParams& P : Params)
 	{
 		FD3D12RayTracingGeometry* Geometry = FD3D12DynamicRHI::ResourceCast(P.Geometry.GetReference());
 		Geometry->SetDirty(GetGPUMask(), true);
@@ -3650,9 +3650,9 @@ void FD3D12CommandContext::RHIBuildAccelerationStructures(const TArrayView<const
 	CommandListHandle.AddUAVBarrier();
 }
 
-void FD3D12CommandContext::RHIBuildAccelerationStructure(FRHIRayTracingScene* InScene)
+void FD3D12CommandContext::RHIBuildAccelerationStructure(const FRayTracingSceneBuildParams& SceneBuildParams)
 {
-	FD3D12RayTracingScene* Scene = FD3D12DynamicRHI::ResourceCast(InScene);
+	FD3D12RayTracingScene* Scene = FD3D12DynamicRHI::ResourceCast(SceneBuildParams.Scene);
 	Scene->BuildAccelerationStructure(*this);
 }
 
diff --git a/Engine/Source/Runtime/Engine/Private/GPUSkinCache.cpp b/Engine/Source/Runtime/Engine/Private/GPUSkinCache.cpp
--- a/Engine/Source/Runtime/Engine/Private/GPUSkinCache.cpp
+++ b/Engine/Source/Runtime/Engine/Private/GPUSkinCache.cpp
@@ -821,14 +821,14 @@ void FGPUSkinCache::CommitRayTracingGeometryUpdates(FRHICommandListImmediate& RH
 				
 		// Track the amount of primitives which need to be build/updated in a single batch
 		uint64 PrimitivesToUpdates = 0;
-		TArray<FAccelerationStructureBuildParams> BatchedBuildParams;
+		TArray<FRayTracingGeometryBuildParams> BatchedBuildParams;
 		BatchedBuildParams.Reserve(RayTracingGeometriesToUpdate.Num());
 
 		// Iterate all the geometries which need an update
 		for (TMap<FRayTracingGeometry*, EAccelerationStructureBuildMode>::TRangedForIterator Iter = RayTracingGeometriesToUpdate.begin(); Iter != RayTracingGeometriesToUpdate.end(); ++Iter)
 		{
 			FRayTracingGeometry* RayTracingGeometry = Iter.Key();
-			FAccelerationStructureBuildParams BuildParams;
+			FRayTracingGeometryBuildParams BuildParams;
 			BuildParams.Geometry	= RayTracingGeometry->RayTracingGeometryRHI;
 			BuildParams.BuildMode   = Iter.Value();
 			BuildParams.Segments	= RayTracingGeometry->Initializer.Segments;
diff --git a/Engine/Source/Runtime/RHI/Public/RHICommandList.h b/Engine/Source/Runtime/RHI/Public/RHICommandList.h
--- a/Engine/Source/Runtime/RHI/Public/RHICommandList.h
+++ b/Engine/Source/Runtime/RHI/Public/RHICommandList.h
@@ -2100,10 +2100,10 @@ struct FRHICommandBindAccelerationStructureMemory final : public FRHICommand<FRH
 
 struct FRHICommandBuildAccelerationStructure final : public FRHICommand<FRHICommandBuildAccelerationStructure>
 {
-	FRHIRayTracingScene* Scene;
+	FRayTracingSceneBuildParams SceneBuildParams;
 
-	explicit FRHICommandBuildAccelerationStructure(FRHIRayTracingScene* InScene)
-		: Scene(InScene)
+	explicit FRHICommandBuildAccelerationStructure(const FRayTracingSceneBuildParams& InSceneBuildParams)
+		: SceneBuildParams(InSceneBuildParams)
 	{}
 
 	RHI_API void Execute(FRHICommandListBase& CmdList);
@@ -2122,9 +2122,9 @@ FRHICOMMAND_MACRO(FRHICommandClearRayTracingBindings)
 
 struct FRHICommandBuildAccelerationStructures final : public FRHICommand<FRHICommandBuildAccelerationStructures>
 {
-	const TArrayView<const FAccelerationStructureBuildParams> Params;
+	const TArrayView<const FRayTracingGeometryBuildParams> Params;
 
-	explicit FRHICommandBuildAccelerationStructures(const TArrayView<const FAccelerationStructureBuildParams> InParams)
+	explicit FRHICommandBuildAccelerationStructures(const TArrayView<const FRayTracingGeometryBuildParams> InParams)
 		: Params(InParams)
 	{}
 
@@ -2799,13 +2799,13 @@ public:
 #if RHI_RAYTRACING
 	FORCEINLINE_DEBUGGABLE void BuildAccelerationStructure(FRHIRayTracingGeometry* Geometry)
 	{
-		FAccelerationStructureBuildParams Params;
+		FRayTracingGeometryBuildParams Params;
 		Params.Geometry = Geometry;
 		Params.BuildMode = EAccelerationStructureBuildMode::Build;
 		BuildAccelerationStructures(MakeArrayView(&Params, 1));
 	}
 
-	FORCEINLINE_DEBUGGABLE void BuildAccelerationStructures(const TArrayView<const FAccelerationStructureBuildParams> Params)
+	FORCEINLINE_DEBUGGABLE void BuildAccelerationStructures(const TArrayView<const FRayTracingGeometryBuildParams> Params)
 	{
 		if (Bypass())
 		{
@@ -2815,8 +2815,8 @@ public:
 		{
 			// Copy the params themselves as well their segment lists, if there are any.
 			// AllocArray() can't be used here directly, as we have to modify the params after copy.
-			size_t DataSize = sizeof(FAccelerationStructureBuildParams) * Params.Num();
-			FAccelerationStructureBuildParams* InlineParams = (FAccelerationStructureBuildParams*) Alloc(DataSize, alignof(FAccelerationStructureBuildParams));
+			size_t DataSize = sizeof(FRayTracingGeometryBuildParams) * Params.Num();
+			FRayTracingGeometryBuildParams* InlineParams = (FRayTracingGeometryBuildParams*) Alloc(DataSize, alignof(FRayTracingGeometryBuildParams));
 			FMemory::Memcpy(InlineParams, Params.GetData(), DataSize);
 			for (int32 i=0; i<Params.Num(); ++i)
 			{
@@ -2829,18 +2829,25 @@ public:
 		}
 	}
 
-	FORCEINLINE_DEBUGGABLE void BuildAccelerationStructure(FRHIRayTracingScene* Scene)
+	FORCEINLINE_DEBUGGABLE void BuildAccelerationStructure(const FRayTracingSceneBuildParams& SceneBuildParams)
 	{
 		if (Bypass())
 		{
-			GetComputeContext().RHIBuildAccelerationStructure(Scene);
+			GetComputeContext().RHIBuildAccelerationStructure(SceneBuildParams);
 		}
 		else
 		{
-			ALLOC_COMMAND(FRHICommandBuildAccelerationStructure)(Scene);
+			ALLOC_COMMAND(FRHICommandBuildAccelerationStructure)(SceneBuildParams);
 		}
 	}
 
+	FORCEINLINE_DEBUGGABLE void BuildAccelerationStructure(FRHIRayTracingScene* Scene)
+	{
+		FRayTracingSceneBuildParams Params;
+		Params.Scene = Scene;
+		BuildAccelerationStructure(Params);
+	}
+
 	FORCEINLINE_DEBUGGABLE void BindAccelerationStructureMemory(FRHIRayTracingScene* Scene, FRHIBuffer* Buffer, uint32 BufferOffset)
 	{
 		if (Bypass())
diff --git a/Engine/Source/Runtime/RHI/Public/RHICommandListCommandExecutes.inl b/Engine/Source/Runtime/RHI/Public/RHICommandListCommandExecutes.inl
--- a/Engine/Source/Runtime/RHI/Public/RHICommandListCommandExecutes.inl
+++ b/Engine/Source/Runtime/RHI/Public/RHICommandListCommandExecutes.inl
@@ -543,7 +543,7 @@ void FRHICommandBindAccelerationStructureMemory::Execute(FRHICommandListBase& Cm
 void FRHICommandBuildAccelerationStructure::Execute(FRHICommandListBase& CmdList)
 {
 	RHISTAT(BuildAccelerationStructure);
-	INTERNAL_DECORATOR_COMPUTE(RHIBuildAccelerationStructure)(Scene);
+	INTERNAL_DECORATOR_COMPUTE(RHIBuildAccelerationStructure)(SceneBuildParams);
 }
 
 void FRHICommandClearRayTracingBindings::Execute(FRHICommandListBase& CmdList)
diff --git a/Engine/Source/Runtime/RHI/Public/RHIContext.h b/Engine/Source/Runtime/RHI/Public/RHIContext.h
--- a/Engine/Source/Runtime/RHI/Public/RHIContext.h
+++ b/Engine/Source/Runtime/RHI/Public/RHIContext.h
@@ -25,7 +25,8 @@ struct FViewportBounds;
 struct FRayTracingGeometryInstance;
 struct FRayTracingShaderBindings;
 struct FRayTracingGeometrySegment;
-struct FAccelerationStructureBuildParams;
+struct FRayTracingGeometryBuildParams;
+struct FRayTracingSceneBuildParams;
 struct FRayTracingLocalShaderBindings;
 enum class EAsyncComputeBudget;
 
@@ -304,18 +305,12 @@ public:
 		/* empty default implementation */
 	}
 
-
-	virtual void RHIBuildAccelerationStructure(FRHIRayTracingGeometry* Geometry)
-	{
-		checkNoEntry();
-	}
-
-	virtual void RHIBuildAccelerationStructures(const TArrayView<const FAccelerationStructureBuildParams> Params)
+	virtual void RHIBuildAccelerationStructures(const TArrayView<const FRayTracingGeometryBuildParams> Params)
 	{
 		checkNoEntry();
 	}
 
-	virtual void RHIBuildAccelerationStructure(FRHIRayTracingScene* Scene)
+	virtual void RHIBuildAccelerationStructure(const FRayTracingSceneBuildParams& SceneBuildParams)
 	{
 		checkNoEntry();
 	}
@@ -359,7 +354,7 @@ enum class EAccelerationStructureBuildMode
 	Update,
 };
 
-struct FAccelerationStructureBuildParams
+struct FRayTracingGeometryBuildParams
 {
 	FRayTracingGeometryRHIRef Geometry;
 	EAccelerationStructureBuildMode BuildMode = EAccelerationStructureBuildMode::Build;
@@ -369,6 +364,29 @@ struct FAccelerationStructureBuildParams
 	TArrayView<const FRayTracingGeometrySegment> Segments;
 };
 
+struct FRayTracingSceneBuildParams
+{
+	// Scene to be built. May be null if explicit instance buffer is provided.
+	FRHIRayTracingScene* Scene = nullptr;
+
+	// Acceleration structure will be written to this buffer. The buffer must be in BVHWrite state.
+	FRHIBuffer* ResultBuffer = nullptr;
+	uint32 ResultBufferOffset = 0;
+
+	// Optional explicit build scratch buffer. Must be in UAV state.
+	// Scratch buffer will be created automatically if this is not provided.
+	FRHIBuffer* ScratchBuffer = nullptr;
+	uint32 ScratchBufferOffset = 0;
+
+	// Optional explicit buffer of native ray tracing instance descriptors. Must be in SRV state.
+	// Instance list will be taken from the scene if this is not provided.
+	FRHIBuffer* InstanceBuffer = nullptr;
+	uint32 InstanceBufferOffset = 0;
+
+	// Number of elements in the InstanceBuffer. Only used if InstanceBuffer is provided.
+	uint32 NumInstances = 0;
+};
+
 struct FCopyBufferRegionParams
 {
 	FRHIBuffer* DestBuffer;
@@ -669,25 +687,32 @@ public:
 		checkNoEntry();
 	}
 
-	virtual void RHIBuildAccelerationStructures(const TArrayView<const FAccelerationStructureBuildParams> Params)
+	virtual void RHIBuildAccelerationStructures(const TArrayView<const FRayTracingGeometryBuildParams> Params)
 	{
 		checkNoEntry();
 	}
 
-	virtual void RHIBuildAccelerationStructure(FRHIRayTracingGeometry* Geometry) final override
+	void RHIBuildAccelerationStructure(FRHIRayTracingGeometry* Geometry)
 	{
-		FAccelerationStructureBuildParams Params;
+		FRayTracingGeometryBuildParams Params;
 		Params.Geometry = Geometry;
 		Params.BuildMode = EAccelerationStructureBuildMode::Build;
 
 		RHIBuildAccelerationStructures(MakeArrayView(&Params, 1));
 	}
 
-	virtual void RHIBuildAccelerationStructure(FRHIRayTracingScene* Scene)
+	virtual void RHIBuildAccelerationStructure(const FRayTracingSceneBuildParams& SceneBuildParams)
 	{
 		checkNoEntry();
 	}
 
+	void RHIBuildAccelerationStructure(FRHIRayTracingScene* Scene)
+	{
+		FRayTracingSceneBuildParams Params;
+		Params.Scene = Scene;
+		RHIBuildAccelerationStructure(Params);
+	}
+
 	virtual void RHIRayTraceOcclusion(FRHIRayTracingScene* Scene,
 		FRHIShaderResourceView* Rays,
 		FRHIUnorderedAccessView* Output,
diff --git a/Engine/Source/Runtime/RHI/Public/RHIValidationContext.h b/Engine/Source/Runtime/RHI/Public/RHIValidationContext.h
--- a/Engine/Source/Runtime/RHI/Public/RHIValidationContext.h
+++ b/Engine/Source/Runtime/RHI/Public/RHIValidationContext.h
@@ -947,17 +947,20 @@ public:
 		RHIContext->RHIClearRayTracingBindings(Scene);
 	}
 
-	virtual void RHIBuildAccelerationStructures(const TArrayView<const FAccelerationStructureBuildParams> Params) override final
+	virtual void RHIBuildAccelerationStructures(const TArrayView<const FRayTracingGeometryBuildParams> Params) override final
 	{
 		// #yuriy_todo: explicit transitions and state validation for BLAS
 		RHIContext->RHIBuildAccelerationStructures(Params);
 	}
 
-	virtual void RHIBuildAccelerationStructure(FRHIRayTracingScene* Scene) override final
+	virtual void RHIBuildAccelerationStructure(const FRayTracingSceneBuildParams& SceneBuildParams) override final
 	{
 		// #yuriy_todo: validate all referenced BLAS states
-		Tracker->Assert(Scene->GetWholeResourceIdentity(), ERHIAccess::BVHWrite);
-		RHIContext->RHIBuildAccelerationStructure(Scene);
+		if (SceneBuildParams.Scene)
+		{
+			Tracker->Assert(SceneBuildParams.Scene->GetWholeResourceIdentity(), ERHIAccess::BVHWrite);
+		}
+		RHIContext->RHIBuildAccelerationStructure(SceneBuildParams);
 	}
 
 	virtual void RHIRayTraceOcclusion(FRHIRayTracingScene* Scene,
diff --git a/Engine/Source/Runtime/RenderCore/Private/RayTracingGeometryManager.cpp b/Engine/Source/Runtime/RenderCore/Private/RayTracingGeometryManager.cpp
--- a/Engine/Source/Runtime/RenderCore/Private/RayTracingGeometryManager.cpp
+++ b/Engine/Source/Runtime/RenderCore/Private/RayTracingGeometryManager.cpp
@@ -165,10 +165,10 @@ void FRayTracingGeometryManager::ProcessBuildRequests(FRHIComputeCommandList& In
 	InCmdList.BuildAccelerationStructures(BuildParams);
 }
 
-void FRayTracingGeometryManager::SetupBuildParams(const BuildRequest& InBuildRequest, TArray<FAccelerationStructureBuildParams>& InBuildParams)
+void FRayTracingGeometryManager::SetupBuildParams(const BuildRequest& InBuildRequest, TArray<FRayTracingGeometryBuildParams>& InBuildParams)
 {
 	// Setup the actual build params
-	FAccelerationStructureBuildParams BuildParam;
+	FRayTracingGeometryBuildParams BuildParam;
 	BuildParam.Geometry = InBuildRequest.Owner->RayTracingGeometryRHI;
 	BuildParam.BuildMode = InBuildRequest.BuildMode;
 	InBuildParams.Add(BuildParam);
diff --git a/Engine/Source/Runtime/RenderCore/Public/RayTracingGeometryManager.h b/Engine/Source/Runtime/RenderCore/Public/RayTracingGeometryManager.h
--- a/Engine/Source/Runtime/RenderCore/Public/RayTracingGeometryManager.h
+++ b/Engine/Source/Runtime/RenderCore/Public/RayTracingGeometryManager.h
@@ -46,7 +46,7 @@ private:
 		EAccelerationStructureBuildMode BuildMode;
 	};
 
-	void SetupBuildParams(const BuildRequest& InBuildRequest, TArray<FAccelerationStructureBuildParams>& InBuildParams);
+	void SetupBuildParams(const BuildRequest& InBuildRequest, TArray<FRayTracingGeometryBuildParams>& InBuildParams);
 
 	FCriticalSection RequestCS;
 
@@ -54,7 +54,7 @@ private:
 
 	// Working array with all active build build params in the RHI
 	TArray<BuildRequest> SortedRequests;
-	TArray<FAccelerationStructureBuildParams> BuildParams;	
+	TArray<FRayTracingGeometryBuildParams> BuildParams;
 };
 
 extern RENDERCORE_API FRayTracingGeometryManager GRayTracingGeometryManager;
diff --git a/Engine/Source/Runtime/Renderer/Private/RayTracing/RayTracingDynamicGeometry.cpp b/Engine/Source/Runtime/Renderer/Private/RayTracing/RayTracingDynamicGeometry.cpp
--- a/Engine/Source/Runtime/Renderer/Private/RayTracing/RayTracingDynamicGeometry.cpp
+++ b/Engine/Source/Runtime/Renderer/Private/RayTracing/RayTracingDynamicGeometry.cpp
@@ -303,7 +303,7 @@ void FRayTracingDynamicGeometryCollection::AddDynamicMeshBatchForGeometryUpdate(
 		Geometry.RayTracingGeometryRHI = RHICreateRayTracingGeometry(Geometry.Initializer);
 	}
 
-	FAccelerationStructureBuildParams Params;
+	FRayTracingGeometryBuildParams Params;
 	Params.Geometry = Geometry.RayTracingGeometryRHI;
 	Params.BuildMode = bRefit
 		? EAccelerationStructureBuildMode::Update
@@ -362,7 +362,7 @@ void FRayTracingDynamicGeometryCollection::DispatchUpdates(FRHIComputeCommandLis
 
 				// Setup the array views on final allocated segments array
 				FRayTracingGeometrySegment* SegmentData = Segments.GetData();
-				for (FAccelerationStructureBuildParams& Param : BuildParams)
+				for (FRayTracingGeometryBuildParams& Param : BuildParams)
 				{
 					uint32 SegmentCount = Param.Segments.Num();
 					if (SegmentCount > 0)
diff --git a/Engine/Source/Runtime/Renderer/Public/RayTracingDynamicGeometryCollection.h b/Engine/Source/Runtime/Renderer/Public/RayTracingDynamicGeometryCollection.h
--- a/Engine/Source/Runtime/Renderer/Public/RayTracingDynamicGeometryCollection.h
+++ b/Engine/Source/Runtime/Renderer/Public/RayTracingDynamicGeometryCollection.h
@@ -25,7 +25,7 @@ public:
 
 private:
 	TArray<struct FMeshComputeDispatchCommand> DispatchCommands;
-	TArray<FAccelerationStructureBuildParams> BuildParams;
+	TArray<FRayTracingGeometryBuildParams> BuildParams;
 	TArray<FRayTracingGeometrySegment> Segments;
 
 	struct FVertexPositionBuffer
diff --git a/Engine/Source/Runtime/VulkanRHI/Private/VulkanContext.h b/Engine/Source/Runtime/VulkanRHI/Private/VulkanContext.h
--- a/Engine/Source/Runtime/VulkanRHI/Private/VulkanContext.h
+++ b/Engine/Source/Runtime/VulkanRHI/Private/VulkanContext.h
@@ -101,8 +101,8 @@ public:
 	virtual void RHINextSubpass() final override;
 
 #if VULKAN_RHI_RAYTRACING
-	virtual void RHIBuildAccelerationStructures(const TArrayView<const FAccelerationStructureBuildParams> Params) final override;
-	virtual void RHIBuildAccelerationStructure(FRHIRayTracingScene* Scene) final override;
+	virtual void RHIBuildAccelerationStructures(const TArrayView<const FRayTracingGeometryBuildParams> Params) final override;
+	virtual void RHIBuildAccelerationStructure(const FRayTracingSceneBuildParams& SceneBuildParams) final override;
 #endif
 
 	inline FVulkanCommandBufferManager* GetCommandBufferManager()
diff --git a/Engine/Source/Runtime/VulkanRHI/Private/VulkanRayTracing.cpp b/Engine/Source/Runtime/VulkanRHI/Private/VulkanRayTracing.cpp
--- a/Engine/Source/Runtime/VulkanRHI/Private/VulkanRayTracing.cpp
+++ b/Engine/Source/Runtime/VulkanRHI/Private/VulkanRayTracing.cpp
@@ -497,9 +497,9 @@ FRayTracingGeometryRHIRef FVulkanDynamicRHI::RHICreateRayTracingGeometry(const F
 }
 
 // Todo: High level rhi call should have transitioned and verified vb and ib to read for each segment
-void FVulkanCommandListContext::RHIBuildAccelerationStructures(const TArrayView<const FAccelerationStructureBuildParams> Params)
+void FVulkanCommandListContext::RHIBuildAccelerationStructures(const TArrayView<const FRayTracingGeometryBuildParams> Params)
 {
-	for (const FAccelerationStructureBuildParams& P : Params)
+	for (const FRayTracingGeometryBuildParams& P : Params)
 	{
 		FVulkanRayTracingGeometry* const Geometry = ResourceCast(P.Geometry.GetReference());
 
@@ -511,9 +511,9 @@ void FVulkanCommandListContext::RHIBuildAccelerationStructures(const TArrayView<
 	}
 }
 
-void FVulkanCommandListContext::RHIBuildAccelerationStructure(FRHIRayTracingScene* InScene)
+void FVulkanCommandListContext::RHIBuildAccelerationStructure(const FRayTracingSceneBuildParams& SceneBuildParams)
 {
-	FVulkanRayTracingScene* const Scene = ResourceCast(InScene);
+	FVulkanRayTracingScene* const Scene = ResourceCast(SceneBuildParams.Scene);
 	Scene->BuildAccelerationStructure(*this);
 }
 
