Added missing default implementations of several RHI validation layer ray tracing methods: RHICalcRayTracingSceneSize, RHICalcRayTracingGeometrySize, RHIBindAccelerationStructureMemory.

diff --git a/Engine/Source/Runtime/RHI/Public/RHIValidation.h b/Engine/Source/Runtime/RHI/Public/RHIValidation.h
--- a/Engine/Source/Runtime/RHI/Public/RHIValidation.h
+++ b/Engine/Source/Runtime/RHI/Public/RHIValidation.h
@@ -1697,6 +1697,16 @@ public:
 	{
 		return RHI->RHICreateRayTracingPipelineState(Initializer);
 	}
+
+	virtual FRayTracingAccelerationStructureSize RHICalcRayTracingSceneSize(uint32 MaxInstances, ERayTracingAccelerationStructureFlags Flags) override final
+	{
+		return RHI->RHICalcRayTracingSceneSize(MaxInstances, Flags);
+	}
+
+	virtual FRayTracingAccelerationStructureSize RHICalcRayTracingGeometrySize(const FRayTracingGeometryInitializer& Initializer) override final
+	{
+		return RHI->RHICalcRayTracingGeometrySize(Initializer);
+	}
 #endif // RHI_RAYTRACING
 
 
diff --git a/Engine/Source/Runtime/RHI/Public/RHIValidationContext.h b/Engine/Source/Runtime/RHI/Public/RHIValidationContext.h
--- a/Engine/Source/Runtime/RHI/Public/RHIValidationContext.h
+++ b/Engine/Source/Runtime/RHI/Public/RHIValidationContext.h
@@ -963,6 +963,11 @@ public:
 		RHIContext->RHIBuildAccelerationStructure(SceneBuildParams);
 	}
 
+	virtual void RHIBindAccelerationStructureMemory(FRHIRayTracingScene* Scene, FRHIBuffer* Buffer, uint32 BufferOffset) override final
+	{
+		RHIContext->RHIBindAccelerationStructureMemory(Scene, Buffer, BufferOffset);
+	}
+
 	virtual void RHIRayTraceOcclusion(FRHIRayTracingScene* Scene,
 		FRHIShaderResourceView* Rays,
 		FRHIUnorderedAccessView* Output,
