Fix for RHI validation when transient allocator is null.

diff --git a/Engine/Source/Runtime/RHI/Private/RHIValidation.cpp b/Engine/Source/Runtime/RHI/Private/RHIValidation.cpp
--- a/Engine/Source/Runtime/RHI/Private/RHIValidation.cpp
+++ b/Engine/Source/Runtime/RHI/Private/RHIValidation.cpp
@@ -93,8 +93,14 @@ FValidationRHI::~FValidationRHI()
 IRHITransientResourceAllocator* FValidationRHI::RHICreateTransientResourceAllocator()
 {
 	// Wrap around validation allocator
-	IRHITransientResourceAllocator* RHIAllocator = RHI->RHICreateTransientResourceAllocator();
-	return new FValidationTransientResourceAllocator(RHIAllocator);
+	if (IRHITransientResourceAllocator* RHIAllocator = RHI->RHICreateTransientResourceAllocator())
+	{
+		return new FValidationTransientResourceAllocator(RHIAllocator);
+	}
+	else
+	{
+		return nullptr;
+	}
 }
 
 IRHICommandContext* FValidationRHI::RHIGetDefaultContext()
