Fix: the RHI validation layer assumed that UAV indices must be in the range 0-7.

This is true for D3D11, since SRVs and UAVs have different index spaces, but in Vulkan all the resources share the same space, so UAV indices can be larger than 7. The validation code overflowed the statically allocated UAV tracking buffer for shaders which happened to have UAVs in slots above 7. The solution is to dynamically size the tracking array, and reset it when a new shader or PSO is bound.

#rnx

diff --git a/Engine/Source/Runtime/RHI/Public/RHIValidationCommon.h b/Engine/Source/Runtime/RHI/Public/RHIValidationCommon.h
--- a/Engine/Source/Runtime/RHI/Public/RHIValidationCommon.h
+++ b/Engine/Source/Runtime/RHI/Public/RHIValidationCommon.h
@@ -818,10 +818,27 @@ namespace RHIValidation
 		struct FUAVTracker
 		{
 		private:
-			FUnorderedAccessView* UAVs[MaxSimultaneousUAVs] = {};
+			TArray<FUnorderedAccessView*> UAVs;
 
 		public:
-			inline FUnorderedAccessView*& operator[](int32 Slot) { return UAVs[Slot]; }
+			FUAVTracker()
+			{
+				UAVs.Reserve(MaxSimultaneousUAVs);
+			}
+
+			inline FUnorderedAccessView*& operator[](int32 Slot)
+			{
+				if (Slot >= UAVs.Num())
+				{
+					UAVs.SetNumZeroed(Slot + 1);
+				}
+				return UAVs[Slot];
+			}
+
+			inline void Reset()
+			{
+				UAVs.SetNum(0, false);
+			}
 
 			inline void DrawOrDispatch(FTracker* BarrierTracker, const FState& RequiredState)
 			{
@@ -831,7 +848,7 @@ namespace RHIValidation
 				uint32 NumUniqueIdentities = 0;
 				FResourceIdentity UniqueIdentities[MaxSimultaneousUAVs];
 
-				for (int32 UAVIndex = 0; UAVIndex < MaxSimultaneousUAVs; ++UAVIndex)
+				for (int32 UAVIndex = 0; UAVIndex < UAVs.Num(); ++UAVIndex)
 				{
 					if (UAVs[UAVIndex])
 					{
@@ -846,6 +863,7 @@ namespace RHIValidation
 
 						if (!bFound)
 						{
+							check(NumUniqueIdentities < UE_ARRAY_COUNT(UniqueIdentities));
 							UniqueIdentities[NumUniqueIdentities++] = Identity;
 
 							// Assert unique resources have the required state.
@@ -933,7 +951,7 @@ namespace RHIValidation
 
 		inline void ResetUAVState(EUAVMode Mode)
 		{
-			UAVTrackers[int32(Mode)] = {};
+			UAVTrackers[int32(Mode)].Reset();
 		}
 
 		inline void ResetAllUAVState()
