Fix broken shadows on Vulkan.

Vulkan was recording buffer upload commands to separate command buffers which were submitted immediately. This meant that a sequence of copy / dispatch, copy / dispatch commands from the high level code would actually execute as copy, copy, dispatch, dispatch. If the two copies reused the same buffer, this led to the first dispatch using data meant for the second dispatch. The fix is to simply use the active command buffer for the copy commands, but we also need to ensure that the high level code doesn't try to do this during a render pass, since that's not allowed in VK. Added validation for this.

#rnx

diff --git a/Engine/Source/Runtime/RHI/Private/RHIValidation.cpp b/Engine/Source/Runtime/RHI/Private/RHIValidation.cpp
--- a/Engine/Source/Runtime/RHI/Private/RHIValidation.cpp
+++ b/Engine/Source/Runtime/RHI/Private/RHIValidation.cpp
@@ -342,6 +342,37 @@ void FValidationRHI::RHICreateTransition(FRHITransition* Transition, const FRHIT
 	return RHI->RHICreateTransition(Transition, CreateInfo);
 }
 
+namespace RHIValidation
+{
+	static inline FString GetReasonString_LockBufferInsideRenderPass(FResource* Buffer)
+	{
+		const TCHAR* DebugName = Buffer->GetDebugName();
+		return FString::Printf(TEXT("Locking non-volatile buffers for writing inside a render pass is not allowed. Resource: \"%s\" (0x%p)."), DebugName ? DebugName : TEXT("Unnamed"), Buffer);
+	}
+}
+
+void* FValidationRHI::RHILockBuffer(class FRHICommandListImmediate& RHICmdList, FRHIBuffer* Buffer, uint32 Offset, uint32 SizeRHI, EResourceLockMode LockMode)
+{
+	using namespace RHIValidation;
+
+	if ((Buffer->GetUsage() & BUF_Volatile) == 0 && LockMode == RLM_WriteOnly)
+	{
+		bool bIsInsideRenderPass;
+		if (RHICmdList.IsTopOfPipe())
+		{
+			bIsInsideRenderPass = RHICmdList.IsInsideRenderPass();
+		}
+		else
+		{
+			FValidationContext& Ctx = static_cast<FValidationContext&>(RHICmdList.GetContext());
+			bIsInsideRenderPass = Ctx.State.bInsideBeginRenderPass;
+		}
+		RHI_VALIDATION_CHECK(!bIsInsideRenderPass, *GetReasonString_LockBufferInsideRenderPass(Buffer));
+	}
+
+	return RHI->RHILockBuffer(RHICmdList, Buffer, Offset, SizeRHI, LockMode);
+}
+
 void FValidationRHI::ReportValidationFailure(const TCHAR* InMessage)
 {
 	// Report failures only once per session, since many of them will happen repeatedly. This is similar to what ensure() does, but
diff --git a/Engine/Source/Runtime/RHI/Public/RHIValidation.h b/Engine/Source/Runtime/RHI/Public/RHIValidation.h
--- a/Engine/Source/Runtime/RHI/Public/RHIValidation.h
+++ b/Engine/Source/Runtime/RHI/Public/RHIValidation.h
@@ -378,10 +378,7 @@ public:
 	}
 
 	// FlushType: Flush RHI Thread
-	virtual void* RHILockBuffer(class FRHICommandListImmediate& RHICmdList, FRHIBuffer* Buffer, uint32 Offset, uint32 SizeRHI, EResourceLockMode LockMode) override final
-	{
-		return RHI->RHILockBuffer(RHICmdList, Buffer, Offset, SizeRHI, LockMode);
-	}
+	virtual void* RHILockBuffer(class FRHICommandListImmediate& RHICmdList, FRHIBuffer* Buffer, uint32 Offset, uint32 SizeRHI, EResourceLockMode LockMode) override final;
 
 	// FlushType: Flush RHI Thread
 	virtual void RHIUnlockBuffer(class FRHICommandListImmediate& RHICmdList, FRHIBuffer* Buffer) override final
@@ -1308,6 +1305,12 @@ public:
 
 	virtual FBufferRHIRef CreateBuffer_RenderThread(class FRHICommandListImmediate& RHICmdList, uint32 Size, EBufferUsageFlags Usage, uint32 Stride, ERHIAccess ResourceState, FRHIResourceCreateInfo& CreateInfo) override final
 	{
+		if (CreateInfo.ResourceArray && RHICmdList.IsInsideRenderPass())
+		{
+			FString Msg = FString::Printf(TEXT("Creating buffers with initial data during a render pass is not supported, buffer name: \"%s\""), CreateInfo.DebugName);
+			RHI_VALIDATION_CHECK(false, *Msg);
+		}
+
 		FBufferRHIRef Buffer = RHI->CreateBuffer_RenderThread(RHICmdList, Size, Usage, Stride, ResourceState, CreateInfo);
 		Buffer->InitBarrierTracking(ResourceState, CreateInfo.DebugName);
 		return Buffer;
diff --git a/Engine/Source/Runtime/VulkanRHI/Private/VulkanBuffer.cpp b/Engine/Source/Runtime/VulkanRHI/Private/VulkanBuffer.cpp
--- a/Engine/Source/Runtime/VulkanRHI/Private/VulkanBuffer.cpp
+++ b/Engine/Source/Runtime/VulkanRHI/Private/VulkanBuffer.cpp
@@ -312,12 +312,10 @@ inline void FVulkanResourceMultiBuffer::InternalUnlock(FVulkanCommandListContext
 	VulkanRHI::FStagingBuffer* StagingBuffer = PendingLock.StagingBuffer;
 	PendingLock.StagingBuffer = nullptr;
 
-	FVulkanCmdBuffer* Cmd = Context.GetCommandBufferManager()->GetUploadCmdBuffer();
-	if (!Cmd->HasBegun())
-	{
-		Cmd->Begin();
-	}
-	ensure(Cmd->IsOutsideRenderPass());
+	// We need to do this on the active command buffer instead of using an upload command buffer. The high level code sometimes reuses the same
+	// buffer in sequences of upload / dispatch, upload / dispatch, so we need to order the copy commands correctly with respect to the dispatches.
+	FVulkanCmdBuffer* Cmd = Context.GetCommandBufferManager()->GetActiveCmdBuffer();
+	check(Cmd && Cmd->IsOutsideRenderPass());
 	VkCommandBuffer CmdBuffer = Cmd->GetHandle();
 
 	VulkanRHI::DebugHeavyWeightBarrier(CmdBuffer, 16);
