Enabled Manual Vertex Fetch on Metal MRT for mobile

diff --git a/Engine/Source/Runtime/RHI/Public/RHI.h b/Engine/Source/Runtime/RHI/Public/RHI.h
--- a/Engine/Source/Runtime/RHI/Public/RHI.h
+++ b/Engine/Source/Runtime/RHI/Public/RHI.h
@@ -198,7 +198,7 @@ inline bool RHISupports4ComponentUAVReadWrite(const FStaticShaderPlatform Platfo
 	Shader Platform must not use the mobile renderer, and for Metal, the shader language must be at least 2. */
 inline bool RHISupportsManualVertexFetch(const FStaticShaderPlatform InShaderPlatform)
 {
-	bool bIsMetalMobilePlatform = IsMetalPlatform(InShaderPlatform) && !IsPCPlatform(InShaderPlatform);
+	bool bIsMetalMobilePlatform = IsMetalPlatform(InShaderPlatform) && !IsMetalSM5Platform(InShaderPlatform);
 	bool bIsUnsupportedGL = IsOpenGLPlatform(InShaderPlatform) && !IsSwitchPlatform(InShaderPlatform);
 
 	return !bIsUnsupportedGL && !IsMobilePlatform(InShaderPlatform) && !bIsMetalMobilePlatform;
