Tessellation Removal: Removing RHI Hull/Domain shader types and enums.

diff --git a/Engine/Source/Developer/Apple/MetalShaderFormat/Private/MetalShaderCompiler.cpp b/Engine/Source/Developer/Apple/MetalShaderFormat/Private/MetalShaderCompiler.cpp
--- a/Engine/Source/Developer/Apple/MetalShaderFormat/Private/MetalShaderCompiler.cpp
+++ b/Engine/Source/Developer/Apple/MetalShaderFormat/Private/MetalShaderCompiler.cpp
@@ -812,8 +812,6 @@ static const EHlslShaderFrequency FrequencyTable[] =
 	HSF_VertexShader,
 	HSF_InvalidFrequency,
 	HSF_InvalidFrequency,
-	HSF_HullShader,
-	HSF_DomainShader,
 	HSF_PixelShader,
 	HSF_InvalidFrequency,
 	HSF_ComputeShader
diff --git a/Engine/Source/Developer/ShaderCompilerCommon/Private/ShaderCompilerCommon.cpp b/Engine/Source/Developer/ShaderCompilerCommon/Private/ShaderCompilerCommon.cpp
--- a/Engine/Source/Developer/ShaderCompilerCommon/Private/ShaderCompilerCommon.cpp
+++ b/Engine/Source/Developer/ShaderCompilerCommon/Private/ShaderCompilerCommon.cpp
@@ -1228,8 +1228,6 @@ FString CreateShaderCompilerWorkerDirectCommandLine(const FShaderCompilerInput&
 	case SF_Vertex:			Text += TEXT(" -vs"); break;
 	case SF_Mesh:			Text += TEXT(" -ms"); break;
 	case SF_Amplification:	Text += TEXT(" -as"); break;
-	case SF_Hull:			Text += TEXT(" -hs"); break;
-	case SF_Domain:			Text += TEXT(" -ds"); break;
 	case SF_Geometry:		Text += TEXT(" -gs"); break;
 	case SF_Pixel:			Text += TEXT(" -ps"); break;
 	case SF_Compute:		Text += TEXT(" -cs"); break;
@@ -1289,8 +1287,6 @@ static FString CreateShaderConductorCommandLine(const FShaderCompilerInput& Inpu
 	case SF_Vertex:			Stage = TEXT("vs"); break;
 	case SF_Pixel:			Stage = TEXT("ps"); break;
 	case SF_Geometry:		Stage = TEXT("gs"); break;
-	case SF_Hull:			Stage = TEXT("hs"); break;
-	case SF_Domain:			Stage = TEXT("ds"); break;
 	case SF_Compute:		Stage = TEXT("cs"); break;
 	default:				return FString();
 	}
@@ -1435,14 +1431,6 @@ void CompileOfflineMali(const FShaderCompilerInput& Input, FShaderCompilerOutput
 				GLSLSourceFile += bVulkanSpirV ? TEXT(".spv") : TEXT(".geom");
 				CompilerCommand += TEXT(" -g");
 			break;
-			case SF_Hull:
-				GLSLSourceFile += bVulkanSpirV ? TEXT(".spv") : TEXT(".tesc");
-				CompilerCommand += TEXT(" -t");
-			break;
-			case SF_Domain:
-				GLSLSourceFile += bVulkanSpirV ? TEXT(".spv") : TEXT(".tese");
-				CompilerCommand += TEXT(" -e");
-			break;
 			case SF_Compute:
 				GLSLSourceFile += bVulkanSpirV ? TEXT(".spv") : TEXT(".comp");
 				CompilerCommand += TEXT(" -C");
@@ -1777,8 +1765,6 @@ namespace CrossCompiler
 		TEXT("Vertex"),
 		TEXT("Mesh"),
 		TEXT("Amplification"),
-		TEXT("Hull"),
-		TEXT("Domain"),
 		TEXT("Pixel"),
 		TEXT("Geometry"),
 		TEXT("Compute"),
diff --git a/Engine/Source/Developer/ShaderFormatOpenGL/Private/OpenGLShaderCompiler.cpp b/Engine/Source/Developer/ShaderFormatOpenGL/Private/OpenGLShaderCompiler.cpp
--- a/Engine/Source/Developer/ShaderFormatOpenGL/Private/OpenGLShaderCompiler.cpp
+++ b/Engine/Source/Developer/ShaderFormatOpenGL/Private/OpenGLShaderCompiler.cpp
@@ -474,8 +474,6 @@ GLenum GLFrequencyTable[] =
 	GL_VERTEX_SHADER,	// SF_Vertex
 	GLenum(0), // SF_Mesh
 	GLenum(0), // SF_Amplification
-	GL_TESS_CONTROL_SHADER,	 // SF_Hull
-	GL_TESS_EVALUATION_SHADER, // SF_Domain
 	GL_FRAGMENT_SHADER, // SF_Pixel
 	GL_GEOMETRY_SHADER,	// SF_Geometry
 	GL_COMPUTE_SHADER,  // SF_Compute
@@ -564,8 +562,6 @@ static ANSICHAR TranslateFrequencyToCrossCompilerPrefix(int32 Frequency)
 	{
 		case SF_Vertex: return 'v';
 		case SF_Pixel: return 'p';
-		case SF_Hull: return 'h';
-		case SF_Domain: return 'd';
 		case SF_Geometry: return 'g';
 		case SF_Compute: return 'c';
 	}
@@ -634,12 +630,6 @@ void FOpenGLFrontend::BuildShaderOutput(
 	case SF_Geometry:
 		Header.FrequencyMarker = 0x4753;
 		break;
-	case SF_Hull:
-		Header.FrequencyMarker = 0x4853;
-		break;
-	case SF_Domain:
-		Header.FrequencyMarker = 0x4453;
-		break;
 	case SF_Compute:
 		Header.FrequencyMarker = 0x4353;
 		break;
@@ -2213,8 +2203,6 @@ void FOpenGLFrontend::CompileShader(const FShaderCompilerInput& Input, FShaderCo
 		HSF_VertexShader,
 		HSF_InvalidFrequency,
 		HSF_InvalidFrequency,
-		bIsSM5 ? HSF_HullShader : HSF_InvalidFrequency,
-		bIsSM5 ? HSF_DomainShader : HSF_InvalidFrequency,
 		HSF_PixelShader,
 		HSF_GeometryShader,
 		RHISupportsComputeShaders(Input.Target.GetPlatform()) ? HSF_ComputeShader : HSF_InvalidFrequency
diff --git a/Engine/Source/Developer/VulkanShaderFormat/Private/VulkanBackend.cpp b/Engine/Source/Developer/VulkanShaderFormat/Private/VulkanBackend.cpp
--- a/Engine/Source/Developer/VulkanShaderFormat/Private/VulkanBackend.cpp
+++ b/Engine/Source/Developer/VulkanShaderFormat/Private/VulkanBackend.cpp
@@ -710,8 +710,6 @@ static inline ShaderStage::EStage GetDescriptorSetForStage(_mesa_glsl_parser_tar
 	case fragment_shader:					return ShaderStage::GetStageForFrequency(SF_Pixel);
 	case compute_shader:					return ShaderStage::GetStageForFrequency(SF_Compute);
 	case geometry_shader:					return ShaderStage::GetStageForFrequency(SF_Geometry);
-	case tessellation_evaluation_shader:	return ShaderStage::GetStageForFrequency(SF_Domain);
-	case tessellation_control_shader:		return ShaderStage::GetStageForFrequency(SF_Hull);
 	default: check(0);	break;	// NOT IMPLEMENTED!
 	}
 
diff --git a/Engine/Source/Developer/VulkanShaderFormat/Private/VulkanShaderCompiler.cpp b/Engine/Source/Developer/VulkanShaderFormat/Private/VulkanShaderCompiler.cpp
--- a/Engine/Source/Developer/VulkanShaderFormat/Private/VulkanShaderCompiler.cpp
+++ b/Engine/Source/Developer/VulkanShaderFormat/Private/VulkanShaderCompiler.cpp
@@ -2251,8 +2251,7 @@ static bool CompileWithShaderConductor(
 {
 	const FShaderCompilerInput& Input = CompilerInfo.Input;
 
-	const bool bUsingTessellation = Input.IsUsingTessellation();
-	const bool bRewriteHlslSource = !bUsingTessellation;
+	const bool bRewriteHlslSource = true;
 	const bool bDebugDump = CompilerInfo.bDebugDump;
 
 	CrossCompiler::FShaderConductorContext CompilerContext;
@@ -2338,8 +2337,6 @@ void DoCompileVulkanShader(const FShaderCompilerInput& Input, FShaderCompilerOut
 		HSF_VertexShader,
 		HSF_InvalidFrequency,
 		HSF_InvalidFrequency,
-		bIsSM5 ? HSF_HullShader : HSF_InvalidFrequency,
-		bIsSM5 ? HSF_DomainShader : HSF_InvalidFrequency,
 		HSF_PixelShader,
 		bIsSM5 ? HSF_GeometryShader : HSF_InvalidFrequency,
 		HSF_ComputeShader
diff --git a/Engine/Source/Developer/Windows/ShaderFormatD3D/Private/D3DShaderCompiler.cpp b/Engine/Source/Developer/Windows/ShaderFormatD3D/Private/D3DShaderCompiler.cpp
--- a/Engine/Source/Developer/Windows/ShaderFormatD3D/Private/D3DShaderCompiler.cpp
+++ b/Engine/Source/Developer/Windows/ShaderFormatD3D/Private/D3DShaderCompiler.cpp
@@ -99,10 +99,6 @@ static const TCHAR* GetShaderProfileName(FShaderTarget Target, bool bForceSM6)
 			return USE_SHADER_MODEL_6_6 ? TEXT("ms_6_6") : TEXT("ms_6_5");
 		case SF_Amplification:
 			return USE_SHADER_MODEL_6_6 ? TEXT("as_6_6") : TEXT("as_6_5");
-		case SF_Hull:
-			return USE_SHADER_MODEL_6_6 ? TEXT("hs_6_6") : TEXT("hs_6_5");
-		case SF_Domain:
-			return USE_SHADER_MODEL_6_6 ? TEXT("ds_6_6") : TEXT("ds_6_5");
 		case SF_Geometry:
 			return USE_SHADER_MODEL_6_6 ? TEXT("gs_6_6") : TEXT("gs_6_5");
 		case SF_Compute:
@@ -127,10 +123,6 @@ static const TCHAR* GetShaderProfileName(FShaderTarget Target, bool bForceSM6)
 			return bForceSM6 ? TEXT("ps_6_0") : TEXT("ps_5_0");
 		case SF_Vertex:
 			return bForceSM6 ? TEXT("vs_6_0") : TEXT("vs_5_0");
-		case SF_Hull:
-			return bForceSM6 ? TEXT("hs_6_0") : TEXT("hs_5_0");
-		case SF_Domain:
-			return bForceSM6 ? TEXT("ds_6_0") : TEXT("ds_5_0");
 		case SF_Geometry:
 			return bForceSM6 ? TEXT("gs_6_0") : TEXT("gs_5_0");
 		case SF_Compute:
diff --git a/Engine/Source/Developer/Windows/ShaderFormatD3D/Private/D3DShaderCompilerDXC.cpp b/Engine/Source/Developer/Windows/ShaderFormatD3D/Private/D3DShaderCompilerDXC.cpp
--- a/Engine/Source/Developer/Windows/ShaderFormatD3D/Private/D3DShaderCompilerDXC.cpp
+++ b/Engine/Source/Developer/Windows/ShaderFormatD3D/Private/D3DShaderCompilerDXC.cpp
@@ -468,8 +468,6 @@ static ShaderConductor::ShaderStage ToDXCShaderStage(EShaderFrequency Frequency)
 	case SF_Vertex:		return ShaderConductor::ShaderStage::VertexShader;
 	case SF_Pixel:		return ShaderConductor::ShaderStage::PixelShader;
 	case SF_Geometry:	return ShaderConductor::ShaderStage::GeometryShader;
-	case SF_Hull:		return ShaderConductor::ShaderStage::HullShader;
-	case SF_Domain:		return ShaderConductor::ShaderStage::DomainShader;
 	case SF_Compute:	return ShaderConductor::ShaderStage::ComputeShader;
 	default:			return ShaderConductor::ShaderStage::NumShaderStages;
 	}
@@ -514,8 +512,7 @@ static bool RewriteUsingSC(FString& PreprocessedShaderSource, const FShaderCompi
 	bool bDumpDebugInfo, ELanguage Language, FShaderCompilerOutput& Output)
 {
 	bool bResult = true;
-	const bool bUsingTessellation = Input.IsUsingTessellation();
-	if (bUsingTessellation || bIsRayTracingShader)
+	if (bIsRayTracingShader)
 	{
 		bResult = false;
 	}
diff --git a/Engine/Source/Editor/UnrealEd/Private/Commandlets/ShaderPipelineCacheToolsCommandlet.cpp b/Engine/Source/Editor/UnrealEd/Private/Commandlets/ShaderPipelineCacheToolsCommandlet.cpp
--- a/Engine/Source/Editor/UnrealEd/Private/Commandlets/ShaderPipelineCacheToolsCommandlet.cpp
+++ b/Engine/Source/Editor/UnrealEd/Private/Commandlets/ShaderPipelineCacheToolsCommandlet.cpp
@@ -748,10 +748,6 @@ int32 ExpandPSOSC(const TArray<FString>& Tokens)
 				PrintShaders(InverseMap, Item.GraphicsDesc.FragmentShader);
 				UE_LOG(LogShaderPipelineCacheTools, Verbose, TEXT("GeometryShader"));
 				PrintShaders(InverseMap, Item.GraphicsDesc.GeometryShader);
-				UE_LOG(LogShaderPipelineCacheTools, Verbose, TEXT("HullShader"));
-				PrintShaders(InverseMap, Item.GraphicsDesc.HullShader);
-				UE_LOG(LogShaderPipelineCacheTools, Verbose, TEXT("DomainShader"));
-				PrintShaders(InverseMap, Item.GraphicsDesc.DomainShader);
 				break;
 			case FPipelineCacheFileFormatPSO::DescriptorType::RayTracing:
 				UE_LOG(LogShaderPipelineCacheTools, Verbose, TEXT("RayTracingShader"));
@@ -797,7 +793,7 @@ int32 ExpandPSOSC(const TArray<FString>& Tokens)
 	{ 
 		NumExamined++;
 		
-		static_assert(SF_Vertex == 0 && SF_Compute == 7, "Shader Frequencies have changed, please update");
+		static_assert(SF_Vertex == 0 && SF_Compute == 5, "Shader Frequencies have changed, please update");
 		TArray<int32> StableShadersPerSlot[SF_NumFrequencies];
 		bool ActivePerSlot[SF_NumFrequencies] = { false };
 
@@ -814,8 +810,6 @@ int32 ExpandPSOSC(const TArray<FString>& Tokens)
 			ActivePerSlot[SF_Amplification] = GetStableShaders(InverseMap, StableShaderKeyIndexTable, Item.GraphicsDesc.AmplificationShader, StableShadersPerSlot[SF_Amplification], OutAnyActiveButMissing);
 			ActivePerSlot[SF_Pixel] = GetStableShaders(InverseMap, StableShaderKeyIndexTable, Item.GraphicsDesc.FragmentShader, StableShadersPerSlot[SF_Pixel], OutAnyActiveButMissing);
 			ActivePerSlot[SF_Geometry] = GetStableShaders(InverseMap, StableShaderKeyIndexTable, Item.GraphicsDesc.GeometryShader, StableShadersPerSlot[SF_Geometry], OutAnyActiveButMissing);
-			ActivePerSlot[SF_Hull] = GetStableShaders(InverseMap, StableShaderKeyIndexTable, Item.GraphicsDesc.HullShader, StableShadersPerSlot[SF_Hull], OutAnyActiveButMissing);
-			ActivePerSlot[SF_Domain] = GetStableShaders(InverseMap, StableShaderKeyIndexTable, Item.GraphicsDesc.DomainShader, StableShadersPerSlot[SF_Domain], OutAnyActiveButMissing);
 		}
 		else if (Item.Type == FPipelineCacheFileFormatPSO::DescriptorType::RayTracing)
 		{
@@ -842,8 +836,6 @@ int32 ExpandPSOSC(const TArray<FString>& Tokens)
 				PrintShaders(InverseMap, StableShaderKeyIndexTable, Item.GraphicsDesc.AmplificationShader, TEXT("AmplificationShader"));
 				PrintShaders(InverseMap, StableShaderKeyIndexTable, Item.GraphicsDesc.FragmentShader, TEXT("FragmentShader"));
 				PrintShaders(InverseMap, StableShaderKeyIndexTable, Item.GraphicsDesc.GeometryShader, TEXT("GeometryShader"));
-				PrintShaders(InverseMap, StableShaderKeyIndexTable, Item.GraphicsDesc.HullShader, TEXT("HullShader"));
-				PrintShaders(InverseMap, StableShaderKeyIndexTable, Item.GraphicsDesc.DomainShader, TEXT("DomainShader"));
 			}
 			else if (Item.Type == FPipelineCacheFileFormatPSO::DescriptorType::RayTracing)
 			{
@@ -921,8 +913,6 @@ int32 ExpandPSOSC(const TArray<FString>& Tokens)
 				PrintShaders(InverseMap, StableShaderKeyIndexTable, Item.GraphicsDesc.AmplificationShader, TEXT("AmplificationShader"));
 				PrintShaders(InverseMap, StableShaderKeyIndexTable, Item.GraphicsDesc.FragmentShader, TEXT("FragmentShader"));
 				PrintShaders(InverseMap, StableShaderKeyIndexTable, Item.GraphicsDesc.GeometryShader, TEXT("GeometryShader"));
-				PrintShaders(InverseMap, StableShaderKeyIndexTable, Item.GraphicsDesc.HullShader, TEXT("HullShader"));
-				PrintShaders(InverseMap, StableShaderKeyIndexTable, Item.GraphicsDesc.DomainShader, TEXT("DomainShader"));
 
 				continue;
 			}
@@ -1018,7 +1008,7 @@ int32 ExpandPSOSC(const TArray<FString>& Tokens)
 		}
 		for (const FPermutation& Perm : Item.Permutations)
 		{
-			// because it is a CSV, and for backward compat, compute shaders will just be a zeroed graphics desc with the shader in the hull shader slot.
+			// because it is a CSV, and for backward compat, compute shaders will just be a zeroed graphics desc with the shader in the mesh shader slot.
 			FString PSOLine = Item.PSO->CommonToString();
 			PSOLine += TEXT(",");
 			if (Item.PSO->Type == FPipelineCacheFileFormatPSO::DescriptorType::Compute)
@@ -1029,7 +1019,7 @@ int32 ExpandPSOSC(const TArray<FString>& Tokens)
 				for (int32 SlotIndex = 0; SlotIndex < SF_Compute; SlotIndex++)  // SF_Compute here because the stablepc.csv file format does not have a compute slot
 				{
 					check(!Item.ActivePerSlot[SlotIndex]); // none of these should be active for a compute shader
-					if (SlotIndex == SF_Hull)
+					if (SlotIndex == SF_Mesh)
 					{
 						FStableShaderKeyAndValue ShaderKeyAndValue = StableShaderKeyIndexTable[Perm.Slots[SF_Compute]];
 						ShaderKeyAndValue.OutputHash = FSHAHash(); // Saved output hash needs to be zeroed so that BuildPSOSC can use this entry even if shaders code changes in future builds
@@ -1214,12 +1204,10 @@ static TSet<FPipelineCacheFileFormatPSO> ParseStableCSV(const FString& FileName,
 			return;
 		}
 
-		// For backward compatibility, compute shaders are stored as a zeroed graphics desc with the shader in the hull shader slot.
+		// For backward compatibility, compute shaders are stored as a zeroed graphics desc with the shader in the mesh shader slot.
 		static FName NAME_SF_Vertex("SF_Vertex");
 		static FName NAME_SF_Mesh("SF_Mesh");
 		static FName NAME_SF_Amplification("SF_Amplification");
-		static FName NAME_SF_Hull("SF_Hull");
-		static FName NAME_SF_Domain("SF_Domain");
 		static FName NAME_SF_Pixel("SF_Pixel");
 		static FName NAME_SF_Geometry("SF_Geometry");
 		static FName NAME_SF_Compute("SF_Compute");
@@ -1251,14 +1239,6 @@ static TSet<FPipelineCacheFileFormatPSO> ParseStableCSV(const FString& FileName,
 			{
 				AdjustedSlotIndex = SF_Amplification;
 			}
-			else if (Shader.TargetFrequency == NAME_SF_Hull)
-			{
-				AdjustedSlotIndex = SF_Hull;
-			}
-			else if (Shader.TargetFrequency == NAME_SF_Domain)
-			{
-				AdjustedSlotIndex = SF_Domain;
-			}
 			else if (Shader.TargetFrequency == NAME_SF_Pixel)
 			{
 				AdjustedSlotIndex = SF_Pixel;
@@ -1344,12 +1324,6 @@ static TSet<FPipelineCacheFileFormatPSO> ParseStableCSV(const FString& FileName,
 			case SF_Geometry:
 				PSO.GraphicsDesc.GeometryShader = Match;
 				break;
-			case SF_Hull:
-				PSO.GraphicsDesc.HullShader = Match;
-				break;
-			case SF_Domain:
-				PSO.GraphicsDesc.DomainShader = Match;
-				break;
 			case SF_Compute:
 				PSO.ComputeDesc.ComputeShader = Match;
 				break;
@@ -1375,9 +1349,7 @@ static TSet<FPipelineCacheFileFormatPSO> ParseStableCSV(const FString& FileName,
 				PSO.GraphicsDesc.MeshShader == FSHAHash() &&
 				PSO.GraphicsDesc.AmplificationShader == FSHAHash() &&
 				PSO.GraphicsDesc.FragmentShader == FSHAHash() &&
-				PSO.GraphicsDesc.GeometryShader == FSHAHash() &&
-				PSO.GraphicsDesc.HullShader == FSHAHash() &&
-				PSO.GraphicsDesc.DomainShader == FSHAHash());
+				PSO.GraphicsDesc.GeometryShader == FSHAHash());
 		}
 		else if (PSO.Type == FPipelineCacheFileFormatPSO::DescriptorType::Graphics)
 		{
@@ -2066,9 +2038,7 @@ int32 BuildPSOSC(const TArray<FString>& Tokens)
 							TestItem.GraphicsDesc.MeshShader == Item.GraphicsDesc.MeshShader &&
 							TestItem.GraphicsDesc.AmplificationShader == Item.GraphicsDesc.AmplificationShader &&
 							TestItem.GraphicsDesc.FragmentShader == Item.GraphicsDesc.FragmentShader &&
-							TestItem.GraphicsDesc.GeometryShader == Item.GraphicsDesc.GeometryShader &&
-							TestItem.GraphicsDesc.HullShader == Item.GraphicsDesc.HullShader &&
-							TestItem.GraphicsDesc.DomainShader == Item.GraphicsDesc.DomainShader
+							TestItem.GraphicsDesc.GeometryShader == Item.GraphicsDesc.GeometryShader
 							)
 						{
 							bMatchedKept = true;
diff --git a/Engine/Source/Programs/ShaderCompileWorker/Private/ShaderCompileWorker.cpp b/Engine/Source/Programs/ShaderCompileWorker/Private/ShaderCompileWorker.cpp
--- a/Engine/Source/Programs/ShaderCompileWorker/Private/ShaderCompileWorker.cpp
+++ b/Engine/Source/Programs/ShaderCompileWorker/Private/ShaderCompileWorker.cpp
@@ -894,14 +894,6 @@ static void DirectCompile(const TArray<const class IShaderFormat*>& ShaderFormat
 			{
 				Frequency = SF_Amplification;
 			}
-			else if (!FCString::Strcmp(*Token, TEXT("hs")))
-			{
-				Frequency = SF_Hull;
-			}
-			else if (!FCString::Strcmp(*Token, TEXT("ds")))
-			{
-				Frequency = SF_Domain;
-			}
 			else if (!FCString::Strcmp(*Token, TEXT("gs")))
 			{
 				Frequency = SF_Geometry;
diff --git a/Engine/Source/Runtime/D3D12RHI/Private/D3D12StateCache.cpp b/Engine/Source/Runtime/D3D12RHI/Private/D3D12StateCache.cpp
--- a/Engine/Source/Runtime/D3D12RHI/Private/D3D12StateCache.cpp
+++ b/Engine/Source/Runtime/D3D12RHI/Private/D3D12StateCache.cpp
@@ -353,7 +353,7 @@ void FD3D12StateCacheBase::SetScissorRects(uint32 Count, const D3D12_RECT* const
 
 inline bool ShouldSkipStage(uint32 Stage)
 {
-	return ((Stage == SF_Mesh || Stage == SF_Amplification) && !GRHISupportsMeshShaders) || Stage == SF_Hull || Stage == SF_Domain;
+	return ((Stage == SF_Mesh || Stage == SF_Amplification) && !GRHISupportsMeshShaders);
 }
 
 template <ED3D12PipelineType PipelineType>
diff --git a/Engine/Source/Runtime/Engine/Private/Materials/HLSLMaterialTranslator.cpp b/Engine/Source/Runtime/Engine/Private/Materials/HLSLMaterialTranslator.cpp
--- a/Engine/Source/Runtime/Engine/Private/Materials/HLSLMaterialTranslator.cpp
+++ b/Engine/Source/Runtime/Engine/Private/Materials/HLSLMaterialTranslator.cpp
@@ -7552,7 +7552,7 @@ int32 FHLSLMaterialTranslator::TransformBase(EMaterialCommonBasis SourceCoordBas
 	}
 		
 	{ // validation
-		if (ShaderFrequency != SF_Pixel && ShaderFrequency != SF_Compute && ShaderFrequency != SF_Domain && ShaderFrequency != SF_Vertex)
+		if (ShaderFrequency != SF_Pixel && ShaderFrequency != SF_Compute && ShaderFrequency != SF_Vertex)
 		{
 			return NonPixelShaderExpressionError();
 		}
@@ -7598,15 +7598,7 @@ int32 FHLSLMaterialTranslator::TransformBase(EMaterialCommonBasis SourceCoordBas
 			check(AWComponent == 0);
 			if (DestCoordBasis == MCB_World)
 			{
-				if (ShaderFrequency == SF_Domain)
-				{
-					// domain shader uses a prescale value to preserve scaling factor on WorldTransform	when sampling a displacement map
-					CodeStr = FString(TEXT("TransformTangent<TO>World_PreScaled(Parameters, <A>.xyz)"));
-				}
-				else
-				{
-					CodeStr = TEXT("mul(<A>, <MATRIX>(Parameters.TangentToWorld))");
-				}
+				CodeStr = TEXT("mul(<A>, <MATRIX>(Parameters.TangentToWorld))");
 			}
 			// else use MCB_World as intermediary basis
 			break;
@@ -8115,11 +8107,6 @@ int32 FHLSLMaterialTranslator::AntialiasedTextureMask(int32 Tex, int32 UV, float
 
 int32 FHLSLMaterialTranslator::DepthOfFieldFunction(int32 Depth, int32 FunctionValueIndex)
 {
-	if (ShaderFrequency == SF_Hull)
-	{
-		return Errorf(TEXT("Invalid node DepthOfFieldFunction used in hull shader input!"));
-	}
-
 	if (Depth == INDEX_NONE)
 	{
 		return INDEX_NONE;
@@ -9223,7 +9210,7 @@ int32 FHLSLMaterialTranslator::CustomExpression( class UMaterialExpressionCustom
 		}
 		Code.ReplaceInline(TEXT("\n"), TEXT("\r\n"), ESearchCase::CaseSensitive);
 
-		FString ParametersType = ShaderFrequency == SF_Vertex ? TEXT("Vertex") : (ShaderFrequency == SF_Domain ? TEXT("Tessellation") : TEXT("Pixel"));
+		FString ParametersType = ShaderFrequency == SF_Vertex ? TEXT("Vertex") : TEXT("Pixel");
 
 		FMaterialCustomExpressionEntry& Entry = CustomExpressions.AddDefaulted_GetRef();
 		CustomEntry = &Entry;
diff --git a/Engine/Source/Runtime/Engine/Private/Materials/MaterialExpressions.cpp b/Engine/Source/Runtime/Engine/Private/Materials/MaterialExpressions.cpp
--- a/Engine/Source/Runtime/Engine/Private/Materials/MaterialExpressions.cpp
+++ b/Engine/Source/Runtime/Engine/Private/Materials/MaterialExpressions.cpp
@@ -6563,8 +6563,6 @@ int32 UMaterialExpressionBlendMaterialAttributes::Compile(class FMaterialCompile
 	switch (AttributeFrequency)
 	{
 	case SF_Vertex:	BlendType = VertexAttributeBlendType;	break;
-	case SF_Hull:	BlendType = VertexAttributeBlendType;	break;
-	case SF_Domain:	BlendType = VertexAttributeBlendType;	break;
 	case SF_Pixel:	BlendType = PixelAttributeBlendType;	break;
 	default:
 		return Compiler->Errorf(TEXT("Attribute blending for shader frequency %i not implemented."), AttributeFrequency);
diff --git a/Engine/Source/Runtime/Engine/Private/ShaderCompiler/ShaderCompiler.cpp b/Engine/Source/Runtime/Engine/Private/ShaderCompiler/ShaderCompiler.cpp
--- a/Engine/Source/Runtime/Engine/Private/ShaderCompiler/ShaderCompiler.cpp
+++ b/Engine/Source/Runtime/Engine/Private/ShaderCompiler/ShaderCompiler.cpp
@@ -4634,8 +4634,6 @@ void GlobalBeginCompileShader(
 	// Add the appropriate definitions for the shader frequency.
 	{
 		Input.Environment.SetDefine(TEXT("PIXELSHADER"), Target.Frequency == SF_Pixel);
-		Input.Environment.SetDefine(TEXT("DOMAINSHADER"), Target.Frequency == SF_Domain);
-		Input.Environment.SetDefine(TEXT("HULLSHADER"), Target.Frequency == SF_Hull);
 		Input.Environment.SetDefine(TEXT("VERTEXSHADER"), Target.Frequency == SF_Vertex);
 		Input.Environment.SetDefine(TEXT("MESHSHADER"), Target.Frequency == SF_Mesh);
 		Input.Environment.SetDefine(TEXT("AMPLIFICATIONSHADER"), Target.Frequency == SF_Amplification);
diff --git a/Engine/Source/Runtime/Engine/Private/ShaderDerivedDataVersion.h b/Engine/Source/Runtime/Engine/Private/ShaderDerivedDataVersion.h
--- a/Engine/Source/Runtime/Engine/Private/ShaderDerivedDataVersion.h
+++ b/Engine/Source/Runtime/Engine/Private/ShaderDerivedDataVersion.h
@@ -10,5 +10,5 @@
 
 // In case of merge conflicts with DDC versions, you MUST generate a new GUID and set this new
 // guid as version
-#define GLOBALSHADERMAP_DERIVEDDATA_VER			TEXT("31E4CCE9AC91470E8CAB008CABE852D4")
-#define MATERIALSHADERMAP_DERIVEDDATA_VER		TEXT("FCCBAA5CC66144528A3A2E4B4A29881B}")
\ No newline at end of file
+#define GLOBALSHADERMAP_DERIVEDDATA_VER			TEXT("0357918925494B80A54D1CC105798A81")
+#define MATERIALSHADERMAP_DERIVEDDATA_VER		TEXT("BCC06CA191E448B99C80CA33353AE671")
\ No newline at end of file
diff --git a/Engine/Source/Runtime/OpenGLDrv/Private/OpenGLCommands.cpp b/Engine/Source/Runtime/OpenGLDrv/Private/OpenGLCommands.cpp
--- a/Engine/Source/Runtime/OpenGLDrv/Private/OpenGLCommands.cpp
+++ b/Engine/Source/Runtime/OpenGLDrv/Private/OpenGLCommands.cpp
@@ -26,8 +26,6 @@
 DECLARE_ISBOUNDSHADER(VertexShader)
 DECLARE_ISBOUNDSHADER(PixelShader)
 DECLARE_ISBOUNDSHADER(GeometryShader)
-DECLARE_ISBOUNDSHADER(HullShader)
-DECLARE_ISBOUNDSHADER(DomainShader)
 */
 
 #if 0//DO_CHECK
@@ -2643,8 +2641,6 @@ void FOpenGLDynamicRHI::CommitGraphicsResourceTablesInner()
 	PendingState.bAnyDirtyGraphicsUniformBuffers = false;
 	PendingState.DirtyUniformBuffers[SF_Vertex] = 0;
 	PendingState.DirtyUniformBuffers[SF_Pixel] = 0;
-	PendingState.DirtyUniformBuffers[SF_Hull] = 0;
-	PendingState.DirtyUniformBuffers[SF_Domain] = 0;
 	PendingState.DirtyUniformBuffers[SF_Geometry] = 0;
 }
 
diff --git a/Engine/Source/Runtime/OpenGLDrv/Private/OpenGLShaders.cpp b/Engine/Source/Runtime/OpenGLDrv/Private/OpenGLShaders.cpp
--- a/Engine/Source/Runtime/OpenGLDrv/Private/OpenGLShaders.cpp
+++ b/Engine/Source/Runtime/OpenGLDrv/Private/OpenGLShaders.cpp
@@ -3381,7 +3381,7 @@ void FOpenGLDynamicRHI::BindPendingShaderState( FOpenGLContextState& ContextStat
 	{
 		int32 NextUniformBufferIndex = OGL_FIRST_UNIFORM_BUFFER;
 
-		static_assert(SF_NumGraphicsFrequencies == 7 && SF_NumFrequencies == 12, "Unexpected SF_ ordering");
+		static_assert(SF_NumGraphicsFrequencies == 5 && SF_NumFrequencies == 10, "Unexpected SF_ ordering");
 		static_assert(SF_RayGen > SF_NumGraphicsFrequencies, "SF_NumGraphicsFrequencies be the number of frequencies supported in OpenGL");
 
 		int32 NumUniformBuffers[SF_NumGraphicsFrequencies];
diff --git a/Engine/Source/Runtime/RHI/Private/PipelineFileCache.cpp b/Engine/Source/Runtime/RHI/Private/PipelineFileCache.cpp
--- a/Engine/Source/Runtime/RHI/Private/PipelineFileCache.cpp
+++ b/Engine/Source/Runtime/RHI/Private/PipelineFileCache.cpp
@@ -65,12 +65,13 @@ enum class EPipelineCacheFileFormatVersions : uint32
 	PatchSizeReduction_NoDuplicatedGuid = 18,
 	AlphaToCoverage = 19,
 	AddingMeshShaders = 20,
+	RemovingTessellationShaders = 21,
 };
 
 const uint64 FPipelineCacheFileFormatMagic = 0x5049504543414348; // PIPECACH
 const uint64 FPipelineCacheTOCFileFormatMagic = 0x544F435354415232; // TOCSTAR2
 const uint64 FPipelineCacheEOFFileFormatMagic = 0x454F462D4D41524B; // EOF-MARK
-const uint32 FPipelineCacheFileFormatCurrentVersion = (uint32)EPipelineCacheFileFormatVersions::AddingMeshShaders;
+const uint32 FPipelineCacheFileFormatCurrentVersion = (uint32)EPipelineCacheFileFormatVersions::RemovingTessellationShaders;
 const int32  FPipelineCacheGraphicsDescPartsNum = 63; // parser will expect this number of parts in a description string
 
 /**
@@ -328,12 +329,10 @@ FString FPipelineCacheFileFormatPSO::GraphicsDescriptor::ShadersToString() const
 {
 	FString Result;
 
-	Result += FString::Printf(TEXT("%s,%s,%s,%s,%s,%s,%s")
+	Result += FString::Printf(TEXT("%s,%s,%s,%s,%s")
 		, *VertexShader.ToString()
 		, *FragmentShader.ToString()
 		, *GeometryShader.ToString()
-		, *HullShader.ToString()
-		, *DomainShader.ToString()
 		, *MeshShader.ToString()
 		, *AmplificationShader.ToString()
 	);
@@ -354,8 +353,6 @@ void FPipelineCacheFileFormatPSO::GraphicsDescriptor::ShadersFromString(const FS
 	VertexShader.FromString(*PartIt++);
 	FragmentShader.FromString(*PartIt++);
 	GeometryShader.FromString(*PartIt++);
-	HullShader.FromString(*PartIt++);
-	DomainShader.FromString(*PartIt++);
 	MeshShader.FromString(*PartIt++);
 	AmplificationShader.FromString(*PartIt++);
 
@@ -364,7 +361,7 @@ void FPipelineCacheFileFormatPSO::GraphicsDescriptor::ShadersFromString(const FS
 
 FString FPipelineCacheFileFormatPSO::GraphicsDescriptor::ShaderHeaderLine()
 {
-	return FString(TEXT("VertexShader,FragmentShader,GeometryShader,HullShader,DomainShader,MeshShader,AmplificationShader"));
+	return FString(TEXT("VertexShader,FragmentShader,GeometryShader,MeshShader,AmplificationShader"));
 }
 
 FString FPipelineCacheFileFormatPSO::GraphicsDescriptor::StateToString() const
@@ -680,31 +677,11 @@ bool FPipelineCacheFileFormatPSO::Verify() const
 		}
 #endif
 		
-#if PLATFORM_SUPPORTS_TESSELLATION_SHADERS
-		if( GraphicsDesc.HullShader == FSHAHash() && GraphicsDesc.DomainShader == FSHAHash() && GraphicsDesc.PrimitiveType >= PT_1_ControlPointPatchList && GraphicsDesc.PrimitiveType <= PT_32_ControlPointPatchList)
-		{
-			// Not using tessellation - we shouldn't try to draw patches
-			return false;
-		}
-		else if( (GraphicsDesc.HullShader != FSHAHash() && GraphicsDesc.DomainShader == FSHAHash()) ||
-				 (GraphicsDesc.HullShader == FSHAHash() && GraphicsDesc.DomainShader != FSHAHash()) )
-		{
-			// Hull without Domain or vice-versa
-			return false;
-		}
-#else
-		if(GraphicsDesc.HullShader != FSHAHash() || GraphicsDesc.DomainShader != FSHAHash())
-		{
-			// Define says we don't support tessellation - why have we got tessellation shaders - not a valid PSO for target platform
-			return false;
-		}
-		
 		if(GraphicsDesc.PrimitiveType >= PT_1_ControlPointPatchList && GraphicsDesc.PrimitiveType <= PT_32_ControlPointPatchList)
 		{
 			// Define says we don't support tessellation - can't draw patches - not a valid PSO for target platform
 			return false;
 		}
-#endif
 
 #if PLATFORM_SUPPORTS_GEOMETRY_SHADERS
 		// Is there anything to actually test here?
@@ -810,8 +787,6 @@ bool FPipelineCacheFileFormatPSO::Verify() const
 				KeyHash = FCrc::MemCrc32(&Key.GraphicsDesc.VertexShader.Hash, sizeof(Key.GraphicsDesc.VertexShader.Hash), KeyHash);
 				KeyHash = FCrc::MemCrc32(&Key.GraphicsDesc.FragmentShader.Hash, sizeof(Key.GraphicsDesc.FragmentShader.Hash), KeyHash);
 				KeyHash = FCrc::MemCrc32(&Key.GraphicsDesc.GeometryShader.Hash, sizeof(Key.GraphicsDesc.GeometryShader.Hash), KeyHash);
-				KeyHash = FCrc::MemCrc32(&Key.GraphicsDesc.HullShader.Hash, sizeof(Key.GraphicsDesc.HullShader.Hash), KeyHash);
-				KeyHash = FCrc::MemCrc32(&Key.GraphicsDesc.DomainShader.Hash, sizeof(Key.GraphicsDesc.DomainShader.Hash), KeyHash);
 				KeyHash = FCrc::MemCrc32(&Key.GraphicsDesc.MeshShader.Hash, sizeof(Key.GraphicsDesc.MeshShader.Hash), KeyHash);
 				KeyHash = FCrc::MemCrc32(&Key.GraphicsDesc.AmplificationShader.Hash, sizeof(Key.GraphicsDesc.AmplificationShader.Hash), KeyHash);
 
@@ -910,8 +885,14 @@ bool FPipelineCacheFileFormatPSO::Verify() const
 			Ar << Info.GraphicsDesc.VertexShader;
 			Ar << Info.GraphicsDesc.FragmentShader;
 			Ar << Info.GraphicsDesc.GeometryShader;
-			Ar << Info.GraphicsDesc.HullShader;
-			Ar << Info.GraphicsDesc.DomainShader;
+
+			if (Ar.GameNetVer() < (uint32)EPipelineCacheFileFormatVersions::RemovingTessellationShaders)
+			{
+				FSHAHash HullShader;
+				Ar << HullShader;
+				FSHAHash DomainShader;
+				Ar << DomainShader;
+			}
 			if (Ar.GameNetVer() >= (uint32)EPipelineCacheFileFormatVersions::AddingMeshShaders)
 			{
 				Ar << Info.GraphicsDesc.MeshShader;
@@ -1130,17 +1111,6 @@ FPipelineCacheFileFormatPSO::FPipelineCacheFileFormatPSO()
 	}
 #endif
 	
-#if PLATFORM_SUPPORTS_TESSELLATION_SHADERS
-	if (Init.BoundShaderState.HullShaderRHI)
-	{
-		PSO.GraphicsDesc.HullShader = Init.BoundShaderState.HullShaderRHI->GetHash();
-	}
-	
-	if (Init.BoundShaderState.DomainShaderRHI)
-	{
-		PSO.GraphicsDesc.DomainShader = Init.BoundShaderState.DomainShaderRHI->GetHash();
-	}
-#endif
 	if (Init.BoundShaderState.PixelShaderRHI)
 	{
 		PSO.GraphicsDesc.FragmentShader = Init.BoundShaderState.PixelShaderRHI->GetHash();
@@ -1242,8 +1212,6 @@ bool FPipelineCacheFileFormatPSO::operator==(const FPipelineCacheFileFormatPSO&
 						GraphicsDesc.VertexShader == Other.GraphicsDesc.VertexShader &&
 						GraphicsDesc.FragmentShader == Other.GraphicsDesc.FragmentShader &&
 						GraphicsDesc.GeometryShader == Other.GraphicsDesc.GeometryShader &&
-						GraphicsDesc.HullShader == Other.GraphicsDesc.HullShader &&
-						GraphicsDesc.DomainShader == Other.GraphicsDesc.DomainShader &&
 						GraphicsDesc.MeshShader == Other.GraphicsDesc.MeshShader &&
 						GraphicsDesc.AmplificationShader == Other.GraphicsDesc.MeshShader &&
 						GraphicsDesc.RenderTargetsActive == Other.GraphicsDesc.RenderTargetsActive &&
@@ -1909,13 +1877,7 @@ public:
 											
 											if (NewEntry.GraphicsDesc.FragmentShader != FSHAHash())
 												Meta.Shaders.Add(NewEntry.GraphicsDesc.FragmentShader);
-											
-											if (NewEntry.GraphicsDesc.HullShader != FSHAHash())
-												Meta.Shaders.Add(NewEntry.GraphicsDesc.HullShader);
-											
-											if (NewEntry.GraphicsDesc.DomainShader != FSHAHash())
-												Meta.Shaders.Add(NewEntry.GraphicsDesc.DomainShader);
-											
+
 											if (NewEntry.GraphicsDesc.GeometryShader != FSHAHash())
 												Meta.Shaders.Add(NewEntry.GraphicsDesc.GeometryShader);
 
@@ -2020,13 +1982,7 @@ public:
                                         
                                         if (Entry.GraphicsDesc.FragmentShader != FSHAHash())
                                             Meta.Shaders.Add(Entry.GraphicsDesc.FragmentShader);
-                                        
-                                        if (Entry.GraphicsDesc.HullShader != FSHAHash())
-                                            Meta.Shaders.Add(Entry.GraphicsDesc.HullShader);
-                                        
-                                        if (Entry.GraphicsDesc.DomainShader != FSHAHash())
-                                            Meta.Shaders.Add(Entry.GraphicsDesc.DomainShader);
-                                        
+
                                         if (Entry.GraphicsDesc.GeometryShader != FSHAHash())
 											Meta.Shaders.Add(Entry.GraphicsDesc.GeometryShader);
 
@@ -2213,13 +2169,7 @@ public:
                                         
                                         if (Entry.GraphicsDesc.FragmentShader != FSHAHash())
                                             Meta.Shaders.Add(Entry.GraphicsDesc.FragmentShader);
-                                        
-                                        if (Entry.GraphicsDesc.HullShader != FSHAHash())
-                                            Meta.Shaders.Add(Entry.GraphicsDesc.HullShader);
-                                        
-                                        if (Entry.GraphicsDesc.DomainShader != FSHAHash())
-                                            Meta.Shaders.Add(Entry.GraphicsDesc.DomainShader);
-                                        
+
                                         if (Entry.GraphicsDesc.GeometryShader != FSHAHash())
 											Meta.Shaders.Add(Entry.GraphicsDesc.GeometryShader);
 
@@ -2460,14 +2410,6 @@ public:
 			{
 				TempShaders.Add(NewEntry.GraphicsDesc.GeometryShader);
 			}
-			if (NewEntry.GraphicsDesc.HullShader != FSHAHash())
-			{
-				TempShaders.Add(NewEntry.GraphicsDesc.HullShader);
-			}
-			if (NewEntry.GraphicsDesc.DomainShader != FSHAHash())
-			{
-				TempShaders.Add(NewEntry.GraphicsDesc.DomainShader);
-			}
 			if (NewEntry.GraphicsDesc.MeshShader != FSHAHash())
 			{
 				TempShaders.Add(NewEntry.GraphicsDesc.MeshShader);
@@ -3399,16 +3341,6 @@ struct FPipelineCacheFileData
 										Entry.Value.Shaders.Add(PSO.GraphicsDesc.GeometryShader);
 									}
 									
-									if (PSO.GraphicsDesc.HullShader != FSHAHash())
-									{
-										Entry.Value.Shaders.Add(PSO.GraphicsDesc.HullShader);
-									}
-									
-									if (PSO.GraphicsDesc.DomainShader != FSHAHash())
-									{
-										Entry.Value.Shaders.Add(PSO.GraphicsDesc.DomainShader);
-									}
-
 									if (PSO.GraphicsDesc.MeshShader != FSHAHash())
 									{
 										Entry.Value.Shaders.Add(PSO.GraphicsDesc.MeshShader);
@@ -3555,12 +3487,6 @@ bool FPipelineFileCache::SavePipelineFileCacheFrom(uint32 GameVersion, EShaderPl
 				if (Item.GraphicsDesc.FragmentShader != FSHAHash())
 					Meta.Shaders.Add(Item.GraphicsDesc.FragmentShader);
 
-				if (Item.GraphicsDesc.HullShader != FSHAHash())
-					Meta.Shaders.Add(Item.GraphicsDesc.HullShader);
-
-				if (Item.GraphicsDesc.DomainShader != FSHAHash())
-					Meta.Shaders.Add(Item.GraphicsDesc.DomainShader);
-
 				if (Item.GraphicsDesc.GeometryShader != FSHAHash())
 					Meta.Shaders.Add(Item.GraphicsDesc.GeometryShader);
 
diff --git a/Engine/Source/Runtime/RHI/Private/PipelineStateCache.cpp b/Engine/Source/Runtime/RHI/Private/PipelineStateCache.cpp
--- a/Engine/Source/Runtime/RHI/Private/PipelineStateCache.cpp
+++ b/Engine/Source/Runtime/RHI/Private/PipelineStateCache.cpp
@@ -33,10 +33,6 @@ static inline uint32 GetTypeHash(const FBoundShaderStateInput& Input)
 		^ GetTypeHash(Input.MeshShaderRHI)
 		^ GetTypeHash(Input.AmplificationShaderRHI)
 #endif
-#if PLATFORM_SUPPORTS_TESSELLATION_SHADERS
-		^ GetTypeHash(Input.HullShaderRHI)
-		^ GetTypeHash(Input.DomainShaderRHI)
-#endif
 #if PLATFORM_SUPPORTS_GEOMETRY_SHADERS
 		^ GetTypeHash(Input.GeometryShaderRHI)
 #endif
@@ -111,16 +107,6 @@ static void HandlePipelineCreationFailure(const FGraphicsPipelineStateInitialize
 		UE_LOG(LogRHI, Error, TEXT("Amplification: %s"), *Init.BoundShaderState.AmplificationShaderRHI->ShaderName);
 	}
 #endif
-#if PLATFORM_SUPPORTS_TESSELLATION_SHADERS
-	if(Init.BoundShaderState.HullShaderRHI)
-	{
-		UE_LOG(LogRHI, Error, TEXT("Hull: %s"), *Init.BoundShaderState.HullShaderRHI->ShaderName);
-	}
-	if(Init.BoundShaderState.DomainShaderRHI)
-	{
-		UE_LOG(LogRHI, Error, TEXT("Domain: %s"), *Init.BoundShaderState.DomainShaderRHI->ShaderName);
-	}
-#endif
 #if PLATFORM_SUPPORTS_GEOMETRY_SHADERS
 	if(Init.BoundShaderState.GeometryShaderRHI)
 	{
diff --git a/Engine/Source/Runtime/RHI/Private/RHICommandList.cpp b/Engine/Source/Runtime/RHI/Private/RHICommandList.cpp
--- a/Engine/Source/Runtime/RHI/Private/RHICommandList.cpp
+++ b/Engine/Source/Runtime/RHI/Private/RHICommandList.cpp
@@ -2210,20 +2210,6 @@ FComputeShaderRHIRef FDynamicRHI::CreateComputeShader_RenderThread(class FRHICom
 	return GDynamicRHI->RHICreateComputeShader(Code, Hash);
 }
 
-FHullShaderRHIRef FDynamicRHI::CreateHullShader_RenderThread(class FRHICommandListImmediate& RHICmdList, TArrayView<const uint8> Code, const FSHAHash& Hash)
-{
-	CSV_SCOPED_TIMING_STAT(RHITStalls, CreateHullShader_RenderThread);
-	FScopedRHIThreadStaller StallRHIThread(RHICmdList);
-	return GDynamicRHI->RHICreateHullShader(Code, Hash);
-}
-
-FDomainShaderRHIRef FDynamicRHI::CreateDomainShader_RenderThread(class FRHICommandListImmediate& RHICmdList, TArrayView<const uint8> Code, const FSHAHash& Hash)
-{
-	CSV_SCOPED_TIMING_STAT(RHITStalls, CreateDomainShader_RenderThread);
-	FScopedRHIThreadStaller StallRHIThread(RHICmdList);
-	return GDynamicRHI->RHICreateDomainShader(Code, Hash);
-}
-
 void FDynamicRHI::UpdateTexture2D_RenderThread(class FRHICommandListImmediate& RHICmdList, FRHITexture2D* Texture, uint32 MipIndex, const struct FUpdateTextureRegion2D& UpdateRegion, uint32 SourcePitch, const uint8* SourceData)
 {
 	CSV_SCOPED_TIMING_STAT(RHITStalls, UpdateTexture2D_RenderThread);
diff --git a/Engine/Source/Runtime/RHI/Private/RHIValidation.cpp b/Engine/Source/Runtime/RHI/Private/RHIValidation.cpp
--- a/Engine/Source/Runtime/RHI/Private/RHIValidation.cpp
+++ b/Engine/Source/Runtime/RHI/Private/RHIValidation.cpp
@@ -79,12 +79,6 @@ FValidationRHI::FValidationRHI(FDynamicRHI* InRHI)
 	SeenFailureHashes.Reserve(256);
 }
 
-static inline bool IsTessellationPrimitive(EPrimitiveType Type)
-{
-	return (Type >= PT_1_ControlPointPatchList && Type <= PT_32_ControlPointPatchList);
-}
-
-
 FValidationRHI::~FValidationRHI()
 {
 	GRHIValidationEnabled = false;
@@ -177,13 +171,6 @@ void FValidationRHI::ValidatePipeline(const FGraphicsPipelineStateInitializer& P
 				TEXT("No depth/stencil target set, yet PSO wants to store into it!"));
 		}
 	}
-
-	if (IsTessellationPrimitive(PSOInitializer.PrimitiveType ))
-	{
-		ensureMsgf(RHISupportsTessellation(GMaxRHIShaderPlatform), TEXT("Tried to create a tessellation PSO but RHI doesn't support it!"));
-		ensureMsgf(PSOInitializer.BoundShaderState.HullShaderRHI && PSOInitializer.BoundShaderState.DomainShaderRHI, 
-			TEXT("Tried to create a tessellation PSO but no Hull or Domain shader set!"));
-	}
 }
 
 void FValidationRHI::RHICreateTransition(FRHITransition* Transition, const FRHITransitionCreateInfo& CreateInfo)
diff --git a/Engine/Source/Runtime/RHI/Public/BoundShaderStateCache.h b/Engine/Source/Runtime/RHI/Public/BoundShaderStateCache.h
--- a/Engine/Source/Runtime/RHI/Public/BoundShaderStateCache.h
+++ b/Engine/Source/Runtime/RHI/Public/BoundShaderStateCache.h
@@ -47,8 +47,6 @@ public:
 	FORCEINLINE FRHIMeshShader*     GetMeshShader() const   { return MeshShader; }
 	FORCEINLINE FRHIAmplificationShader*   GetAmplificationShader() const   { return AmplificationShader; }
 	FORCEINLINE FRHIPixelShader*    GetPixelShader() const    { return PixelShader; }
-	FORCEINLINE FRHIHullShader*     GetHullShader() const     { return nullptr; }
-	FORCEINLINE FRHIDomainShader*   GetDomainShader() const   { return nullptr; }
 	FORCEINLINE FRHIGeometryShader* GetGeometryShader() const { return GeometryShader; }
 
 	/**
@@ -214,8 +212,6 @@ public:
 	FORCEINLINE FRHIMeshShader*     GetMeshShader() const     { return Key.GetMeshShader(); }
 	FORCEINLINE FRHIAmplificationShader* GetAmplificationShader() const { return Key.GetAmplificationShader(); }
 	FORCEINLINE FRHIPixelShader*    GetPixelShader() const    { return Key.GetPixelShader(); }
-	FORCEINLINE FRHIHullShader*     GetHullShader() const     { return Key.GetHullShader(); }
-	FORCEINLINE FRHIDomainShader*   GetDomainShader() const   { return Key.GetDomainShader(); }
 	FORCEINLINE FRHIGeometryShader* GetGeometryShader() const { return Key.GetGeometryShader(); }
 
 	/**
diff --git a/Engine/Source/Runtime/RHI/Public/DynamicRHI.h b/Engine/Source/Runtime/RHI/Public/DynamicRHI.h
--- a/Engine/Source/Runtime/RHI/Public/DynamicRHI.h
+++ b/Engine/Source/Runtime/RHI/Public/DynamicRHI.h
@@ -397,18 +397,6 @@ public:
 	// FlushType: Wait RHI Thread
 	virtual FVertexShaderRHIRef RHICreateVertexShader(TArrayView<const uint8> Code, const FSHAHash& Hash) = 0;
 
-	// FlushType: Wait RHI Thread
-	virtual FHullShaderRHIRef RHICreateHullShader(TArrayView<const uint8> Code, const FSHAHash& Hash)
-	{
-		return FHullShaderRHIRef();
-	}
-
-	// FlushType: Wait RHI Thread
-	virtual FDomainShaderRHIRef RHICreateDomainShader(TArrayView<const uint8> Code, const FSHAHash& Hash)
-	{
-		return FDomainShaderRHIRef();
-	}
-
 	// FlushType: Wait RHI Thread
 	virtual FGeometryShaderRHIRef RHICreateGeometryShader(TArrayView<const uint8> Code, const FSHAHash& Hash) = 0;
 
@@ -530,13 +518,11 @@ public:
 	virtual void UnlockStagingBuffer_RenderThread(class FRHICommandListImmediate& RHICmdList, FRHIStagingBuffer* StagingBuffer);
 
 	/**
-	* Creates a bound shader state instance which encapsulates a decl, vertex shader, hull shader, domain shader and pixel shader
+	* Creates a bound shader state instance which encapsulates a decl, vertex shader and pixel shader
 	* CAUTION: Even though this is marked as threadsafe, it is only valid to call from the render thread or the RHI thread. It need not be threadsafe unless the RHI support parallel translation.
 	* CAUTION: Platforms that support RHIThread but don't actually have a threadsafe implementation must flush internally with FScopedRHIThreadStaller StallRHIThread(FRHICommandListExecutor::GetImmediateCommandList()); when the call is from the render thread
 	* @param VertexDeclaration - existing vertex decl
 	* @param VertexShader - existing vertex shader
-	* @param HullShader - existing hull shader
-	* @param DomainShader - existing domain shader
 	* @param GeometryShader - existing geometry shader
 	* @param PixelShader - existing pixel shader
 	*/
@@ -1309,8 +1295,6 @@ public:
 	virtual FPixelShaderRHIRef CreatePixelShader_RenderThread(class FRHICommandListImmediate& RHICmdList, TArrayView<const uint8> Code, const FSHAHash& Hash);
 	virtual FGeometryShaderRHIRef CreateGeometryShader_RenderThread(class FRHICommandListImmediate& RHICmdList, TArrayView<const uint8> Code, const FSHAHash& Hash);
 	virtual FComputeShaderRHIRef CreateComputeShader_RenderThread(class FRHICommandListImmediate& RHICmdList, TArrayView<const uint8> Code, const FSHAHash& Hash);
-	virtual FHullShaderRHIRef CreateHullShader_RenderThread(class FRHICommandListImmediate& RHICmdList, TArrayView<const uint8> Code, const FSHAHash& Hash);
-	virtual FDomainShaderRHIRef CreateDomainShader_RenderThread(class FRHICommandListImmediate& RHICmdList, TArrayView<const uint8> Code, const FSHAHash& Hash);
 	virtual void* LockTexture2D_RenderThread(class FRHICommandListImmediate& RHICmdList, FRHITexture2D* Texture, uint32 MipIndex, EResourceLockMode LockMode, uint32& DestStride, bool bLockWithinMiptail, bool bNeedsDefaultRHIFlush = true);
 	virtual void UnlockTexture2D_RenderThread(class FRHICommandListImmediate& RHICmdList, FRHITexture2D* Texture, uint32 MipIndex, bool bLockWithinMiptail, bool bNeedsDefaultRHIFlush = true);
 	virtual void UpdateTexture2D_RenderThread(class FRHICommandListImmediate& RHICmdList, FRHITexture2D* Texture, uint32 MipIndex, const struct FUpdateTextureRegion2D& UpdateRegion, uint32 SourcePitch, const uint8* SourceData);
diff --git a/Engine/Source/Runtime/RHI/Public/PipelineFileCache.h b/Engine/Source/Runtime/RHI/Public/PipelineFileCache.h
--- a/Engine/Source/Runtime/RHI/Public/PipelineFileCache.h
+++ b/Engine/Source/Runtime/RHI/Public/PipelineFileCache.h
@@ -154,8 +154,6 @@ struct RHI_API FPipelineCacheFileFormatPSO
 		FSHAHash VertexShader;
 		FSHAHash FragmentShader;
 		FSHAHash GeometryShader;
-		FSHAHash HullShader;
-		FSHAHash DomainShader;
 		FSHAHash MeshShader;
 		FSHAHash AmplificationShader;
 
diff --git a/Engine/Source/Runtime/RHI/Public/RHICommandList.h b/Engine/Source/Runtime/RHI/Public/RHICommandList.h
--- a/Engine/Source/Runtime/RHI/Public/RHICommandList.h
+++ b/Engine/Source/Runtime/RHI/Public/RHICommandList.h
@@ -621,8 +621,6 @@ protected:
 	FORCEINLINE void ValidateBoundShader(FRHIVertexShader* ShaderRHI) { checkSlow(BoundShaderInput.VertexShaderRHI == ShaderRHI); }
 	FORCEINLINE void ValidateBoundShader(FRHIPixelShader* ShaderRHI) { checkSlow(BoundShaderInput.PixelShaderRHI == ShaderRHI); }
 	FORCEINLINE void ValidateBoundShader(FRHIGeometryShader* ShaderRHI) { checkSlow(BoundShaderInput.GeometryShaderRHI == ShaderRHI); }
-	FORCEINLINE void ValidateBoundShader(FRHIHullShader* ShaderRHI) { checkSlow(BoundShaderInput.HullShaderRHI == ShaderRHI); }
-	FORCEINLINE void ValidateBoundShader(FRHIDomainShader* ShaderRHI) { checkSlow(BoundShaderInput.DomainShaderRHI == ShaderRHI); }
 	FORCEINLINE void ValidateBoundShader(FRHIComputeShader* ShaderRHI) { checkSlow(BoundComputeShaderRHI == ShaderRHI); }
 	FORCEINLINE void ValidateBoundShader(FRHIMeshShader* ShaderRHI) { checkSlow(BoundShaderInput.MeshShaderRHI == ShaderRHI); }
 	FORCEINLINE void ValidateBoundShader(FRHIAmplificationShader* ShaderRHI) { checkSlow(BoundShaderInput.AmplificationShaderRHI == ShaderRHI); }
@@ -635,8 +633,6 @@ protected:
 		case SF_Vertex: checkSlow(BoundShaderInput.VertexShaderRHI == ShaderRHI); break;
 		case SF_Mesh: checkSlow(BoundShaderInput.MeshShaderRHI == ShaderRHI); break;
 		case SF_Amplification: checkSlow(BoundShaderInput.AmplificationShaderRHI == ShaderRHI); break;
-		case SF_Hull: checkSlow(BoundShaderInput.HullShaderRHI == ShaderRHI); break;
-		case SF_Domain: checkSlow(BoundShaderInput.DomainShaderRHI == ShaderRHI); break;
 		case SF_Pixel: checkSlow(BoundShaderInput.PixelShaderRHI == ShaderRHI); break;
 		case SF_Geometry: checkSlow(BoundShaderInput.GeometryShaderRHI == ShaderRHI); break;
 		default: checkfSlow(false, TEXT("Unexpected graphics shader type %d"), ShaderRHI->GetFrequency());
@@ -2879,8 +2875,6 @@ public:
 	inline FRHIVertexShader* GetBoundVertexShader() const { return BoundShaderInput.VertexShaderRHI; }
 	inline FRHIMeshShader* GetBoundMeshShader() const { return BoundShaderInput.MeshShaderRHI; }
 	inline FRHIAmplificationShader* GetBoundAmplificationShader() const { return BoundShaderInput.AmplificationShaderRHI; }
-	inline FRHIHullShader* GetBoundHullShader() const { return BoundShaderInput.HullShaderRHI; }
-	inline FRHIDomainShader* GetBoundDomainShader() const { return BoundShaderInput.DomainShaderRHI; }
 	inline FRHIPixelShader* GetBoundPixelShader() const { return BoundShaderInput.PixelShaderRHI; }
 	inline FRHIGeometryShader* GetBoundGeometryShader() const { return BoundShaderInput.GeometryShaderRHI; }
 
@@ -4010,18 +4004,6 @@ public:
 		return GDynamicRHI->CreateAmplificationShader_RenderThread(*this, Code, Hash);
 	}
 	
-	FORCEINLINE FHullShaderRHIRef CreateHullShader(TArrayView<const uint8> Code, const FSHAHash& Hash)
-	{
-		LLM_SCOPE(ELLMTag::Shaders);
-		return GDynamicRHI->CreateHullShader_RenderThread(*this, Code, Hash);
-	}
-	
-	FORCEINLINE FDomainShaderRHIRef CreateDomainShader(TArrayView<const uint8> Code, const FSHAHash& Hash)
-	{
-		LLM_SCOPE(ELLMTag::Shaders);
-		return GDynamicRHI->CreateDomainShader_RenderThread(*this, Code, Hash);
-	}
-	
 	FORCEINLINE FGeometryShaderRHIRef CreateGeometryShader(TArrayView<const uint8> Code, const FSHAHash& Hash)
 	{
 		LLM_SCOPE(ELLMTag::Shaders);
@@ -5056,16 +5038,6 @@ FORCEINLINE FAmplificationShaderRHIRef RHICreateAmplificationShader(TArrayView<c
 	return FRHICommandListExecutor::GetImmediateCommandList().CreateAmplificationShader(Code, Hash);
 }
 
-FORCEINLINE FHullShaderRHIRef RHICreateHullShader(TArrayView<const uint8> Code, const FSHAHash& Hash)
-{
-	return FRHICommandListExecutor::GetImmediateCommandList().CreateHullShader(Code, Hash);
-}
-
-FORCEINLINE FDomainShaderRHIRef RHICreateDomainShader(TArrayView<const uint8> Code, const FSHAHash& Hash)
-{
-	return FRHICommandListExecutor::GetImmediateCommandList().CreateDomainShader(Code, Hash);
-}
-
 FORCEINLINE FGeometryShaderRHIRef RHICreateGeometryShader(TArrayView<const uint8> Code, const FSHAHash& Hash)
 {
 	return FRHICommandListExecutor::GetImmediateCommandList().CreateGeometryShader(Code, Hash);
diff --git a/Engine/Source/Runtime/RHI/Public/RHICommandListCommandExecutes.inl b/Engine/Source/Runtime/RHI/Public/RHICommandListCommandExecutes.inl
--- a/Engine/Source/Runtime/RHI/Public/RHICommandListCommandExecutes.inl
+++ b/Engine/Source/Runtime/RHI/Public/RHICommandListCommandExecutes.inl
@@ -487,8 +487,6 @@ void FRHICommandSetLocalUniformBuffer<TRHIShader>::Execute(FRHICommandListBase&
 	}
 }
 template struct FRHICommandSetLocalUniformBuffer<FRHIVertexShader>;
-template struct FRHICommandSetLocalUniformBuffer<FRHIHullShader>;
-template struct FRHICommandSetLocalUniformBuffer<FRHIDomainShader>;
 template struct FRHICommandSetLocalUniformBuffer<FRHIGeometryShader>;
 template struct FRHICommandSetLocalUniformBuffer<FRHIPixelShader>;
 template struct FRHICommandSetLocalUniformBuffer<FRHIComputeShader>;
diff --git a/Engine/Source/Runtime/RHI/Public/RHIDefinitions.h b/Engine/Source/Runtime/RHI/Public/RHIDefinitions.h
--- a/Engine/Source/Runtime/RHI/Public/RHIDefinitions.h
+++ b/Engine/Source/Runtime/RHI/Public/RHIDefinitions.h
@@ -25,23 +25,21 @@ enum EShaderFrequency : uint8
 	SF_Vertex			= 0,
 	SF_Mesh				= 1,
 	SF_Amplification	= 2,
-	SF_Hull				= 3,
-	SF_Domain			= 4,
-	SF_Pixel			= 5,
-	SF_Geometry			= 6,
-	SF_Compute			= 7,
-	SF_RayGen			= 8,
-	SF_RayMiss			= 9,
-	SF_RayHitGroup		= 10,
-	SF_RayCallable		= 11,
-
-	SF_NumFrequencies	= 12,
+	SF_Pixel			= 3,
+	SF_Geometry			= 4,
+	SF_Compute			= 5,
+	SF_RayGen			= 6,
+	SF_RayMiss			= 7,
+	SF_RayHitGroup		= 8,
+	SF_RayCallable		= 9,
+
+	SF_NumFrequencies	= 10,
 
 	// Number of standard shader frequencies for graphics pipeline (excluding compute)
-	SF_NumGraphicsFrequencies = 7,
+	SF_NumGraphicsFrequencies = 5,
 
 	// Number of standard shader frequencies (including compute)
-	SF_NumStandardFrequencies = 8,
+	SF_NumStandardFrequencies = 6,
 
 	SF_NumBits			= 4,
 };
@@ -1410,8 +1408,6 @@ enum ERHIResourceType
 	RRT_BlendState,
 	RRT_VertexDeclaration,
 	RRT_VertexShader,
-	RRT_HullShader,
-	RRT_DomainShader,
 	RRT_PixelShader,
 	RRT_GeometryShader,
 	RRT_ComputeShader,
@@ -1999,8 +1995,6 @@ inline const TCHAR* GetShaderFrequencyString(EShaderFrequency Frequency, bool bI
 	case SF_Vertex:			String = TEXT("SF_Vertex"); break;
 	case SF_Mesh:			String = TEXT("SF_Mesh"); break;
 	case SF_Amplification:	String = TEXT("SF_Amplification"); break;
-	case SF_Hull:			String = TEXT("SF_Hull"); break;
-	case SF_Domain:			String = TEXT("SF_Domain"); break;
 	case SF_Geometry:		String = TEXT("SF_Geometry"); break;
 	case SF_Pixel:			String = TEXT("SF_Pixel"); break;
 	case SF_Compute:		String = TEXT("SF_Compute"); break;
@@ -2044,9 +2038,3 @@ inline bool IsRayTracingShaderFrequency(EShaderFrequency Frequency)
 #else
 	#define GEOMETRY_SHADER(GeometryShader)	nullptr
 #endif
-
-#if PLATFORM_SUPPORTS_TESSELLATION_SHADERS
-	#define TESSELLATION_SHADER(HullOrDomainShader)	(HullOrDomainShader)
-#else
-	#define TESSELLATION_SHADER(HullOrDomainShader)	nullptr
-#endif
diff --git a/Engine/Source/Runtime/RHI/Public/RHIResources.h b/Engine/Source/Runtime/RHI/Public/RHIResources.h
--- a/Engine/Source/Runtime/RHI/Public/RHIResources.h
+++ b/Engine/Source/Runtime/RHI/Public/RHIResources.h
@@ -503,18 +503,6 @@ public:
 	FRHIAmplificationShader() : FRHIGraphicsShader(SF_Amplification) {}
 };
 
-class FRHIHullShader : public FRHIGraphicsShader
-{
-public:
-	FRHIHullShader() : FRHIGraphicsShader(SF_Hull) {}
-};
-
-class FRHIDomainShader : public FRHIGraphicsShader
-{
-public:
-	FRHIDomainShader() : FRHIGraphicsShader(SF_Domain) {}
-};
-
 class FRHIPixelShader : public FRHIGraphicsShader
 {
 public:
@@ -1513,8 +1501,6 @@ typedef TRefCountPtr<FRHIVertexDeclaration> FVertexDeclarationRHIRef;
 typedef TRefCountPtr<FRHIVertexShader> FVertexShaderRHIRef;
 typedef TRefCountPtr<FRHIMeshShader> FMeshShaderRHIRef;
 typedef TRefCountPtr<FRHIAmplificationShader> FAmplificationShaderRHIRef;
-typedef TRefCountPtr<FRHIHullShader> FHullShaderRHIRef;
-typedef TRefCountPtr<FRHIDomainShader> FDomainShaderRHIRef;
 typedef TRefCountPtr<FRHIPixelShader> FPixelShaderRHIRef;
 typedef TRefCountPtr<FRHIGeometryShader> FGeometryShaderRHIRef;
 typedef TRefCountPtr<FRHIComputeShader> FComputeShaderRHIRef;
@@ -1928,24 +1914,18 @@ template<typename TRHIShader> struct TRHIShaderToEnum {};
 template<> struct TRHIShaderToEnum<FRHIVertexShader>		{ enum { ShaderFrequency = SF_Vertex }; };
 template<> struct TRHIShaderToEnum<FRHIMeshShader>			{ enum { ShaderFrequency = SF_Mesh}; };
 template<> struct TRHIShaderToEnum<FRHIAmplificationShader>	{ enum { ShaderFrequency = SF_Amplification}; };
-template<> struct TRHIShaderToEnum<FRHIHullShader>			{ enum { ShaderFrequency = SF_Hull }; };
-template<> struct TRHIShaderToEnum<FRHIDomainShader>		{ enum { ShaderFrequency = SF_Domain }; };
 template<> struct TRHIShaderToEnum<FRHIPixelShader>			{ enum { ShaderFrequency = SF_Pixel }; };
 template<> struct TRHIShaderToEnum<FRHIGeometryShader>		{ enum { ShaderFrequency = SF_Geometry }; };
 template<> struct TRHIShaderToEnum<FRHIComputeShader>		{ enum { ShaderFrequency = SF_Compute }; };
 template<> struct TRHIShaderToEnum<FRHIVertexShader*>		{ enum { ShaderFrequency = SF_Vertex }; };
 template<> struct TRHIShaderToEnum<FRHIMeshShader*>			{ enum { ShaderFrequency = SF_Mesh }; };
 template<> struct TRHIShaderToEnum<FRHIAmplificationShader*>{ enum { ShaderFrequency = SF_Amplification }; };
-template<> struct TRHIShaderToEnum<FRHIHullShader*>			{ enum { ShaderFrequency = SF_Hull }; };
-template<> struct TRHIShaderToEnum<FRHIDomainShader*>		{ enum { ShaderFrequency = SF_Domain }; };
 template<> struct TRHIShaderToEnum<FRHIPixelShader*>		{ enum { ShaderFrequency = SF_Pixel }; };
 template<> struct TRHIShaderToEnum<FRHIGeometryShader*>		{ enum { ShaderFrequency = SF_Geometry }; };
 template<> struct TRHIShaderToEnum<FRHIComputeShader*>		{ enum { ShaderFrequency = SF_Compute }; };
 template<> struct TRHIShaderToEnum<FVertexShaderRHIRef>		{ enum { ShaderFrequency = SF_Vertex }; };
 template<> struct TRHIShaderToEnum<FMeshShaderRHIRef>		{ enum { ShaderFrequency = SF_Mesh }; };
 template<> struct TRHIShaderToEnum<FAmplificationShaderRHIRef>{ enum { ShaderFrequency = SF_Amplification }; };
-template<> struct TRHIShaderToEnum<FHullShaderRHIRef>		{ enum { ShaderFrequency = SF_Hull }; };
-template<> struct TRHIShaderToEnum<FDomainShaderRHIRef>		{ enum { ShaderFrequency = SF_Domain }; };
 template<> struct TRHIShaderToEnum<FPixelShaderRHIRef>		{ enum { ShaderFrequency = SF_Pixel }; };
 template<> struct TRHIShaderToEnum<FGeometryShaderRHIRef>	{ enum { ShaderFrequency = SF_Geometry }; };
 template<> struct TRHIShaderToEnum<FComputeShaderRHIRef>	{ enum { ShaderFrequency = SF_Compute }; };
@@ -2008,16 +1988,6 @@ struct FBoundShaderStateInput
 			VertexShaderRHI->AddRef();
 		}
 
-		if (HullShaderRHI)
-		{
-			HullShaderRHI->AddRef();
-		}
-
-		if (DomainShaderRHI)
-		{
-			DomainShaderRHI->AddRef();
-		}
-
 		if (PixelShaderRHI)
 		{
 			PixelShaderRHI->AddRef();
@@ -2053,16 +2023,6 @@ struct FBoundShaderStateInput
 			VertexShaderRHI->Release();
 		}
 
-		if (HullShaderRHI)
-		{
-			HullShaderRHI->Release();
-		}
-
-		if (DomainShaderRHI)
-		{
-			DomainShaderRHI->Release();
-		}
-
 		if (PixelShaderRHI)
 		{
 			PixelShaderRHI->Release();
@@ -2078,8 +2038,6 @@ struct FBoundShaderStateInput
 	FRHIVertexShader* VertexShaderRHI = nullptr;
 	FRHIMeshShader* MeshShaderRHI = nullptr;
 	FRHIAmplificationShader* AmplificationShaderRHI = nullptr;
-	FRHIHullShader* HullShaderRHI = nullptr;
-	FRHIDomainShader* DomainShaderRHI = nullptr;
 	FRHIPixelShader* PixelShaderRHI = nullptr;
 	FRHIGeometryShader* GeometryShaderRHI = nullptr;
 };
@@ -2234,10 +2192,6 @@ public:
 #endif
 #if PLATFORM_SUPPORTS_GEOMETRY_SHADERS
 			BoundShaderState.GeometryShaderRHI != rhs.BoundShaderState.GeometryShaderRHI ||
-#endif
-#if PLATFORM_SUPPORTS_TESSELLATION_SHADERS
-			BoundShaderState.DomainShaderRHI != rhs.BoundShaderState.DomainShaderRHI ||
-			BoundShaderState.HullShaderRHI != rhs.BoundShaderState.HullShaderRHI ||
 #endif
 			BlendState != rhs.BlendState ||
 			RasterizerState != rhs.RasterizerState ||
diff --git a/Engine/Source/Runtime/RHI/Public/RHIValidation.h b/Engine/Source/Runtime/RHI/Public/RHIValidation.h
--- a/Engine/Source/Runtime/RHI/Public/RHIValidation.h
+++ b/Engine/Source/Runtime/RHI/Public/RHIValidation.h
@@ -151,20 +151,6 @@ public:
 		return RHI->RHICreateVertexShader(Code, Hash);
 	}
 
-	// FlushType: Wait RHI Thread
-	virtual FHullShaderRHIRef RHICreateHullShader(TArrayView<const uint8> Code, const FSHAHash& Hash) override final
-	{
-		check(RHISupportsTessellation(GMaxRHIShaderPlatform));
-		return RHI->RHICreateHullShader(Code, Hash);
-	}
-
-	// FlushType: Wait RHI Thread
-	virtual FDomainShaderRHIRef RHICreateDomainShader(TArrayView<const uint8> Code, const FSHAHash& Hash) override final
-	{
-		check(RHISupportsTessellation(GMaxRHIShaderPlatform));
-		return RHI->RHICreateDomainShader(Code, Hash);
-	}
-
 	// FlushType: Wait RHI Thread
 	virtual FGeometryShaderRHIRef RHICreateGeometryShader(TArrayView<const uint8> Code, const FSHAHash& Hash) override final
 	{
@@ -295,13 +281,11 @@ public:
 	}
 
 	/**
-	* Creates a bound shader state instance which encapsulates a decl, vertex shader, hull shader, domain shader and pixel shader
+	* Creates a bound shader state instance which encapsulates a decl, vertex shader and pixel shader
 	* CAUTION: Even though this is marked as threadsafe, it is only valid to call from the render thread or the RHI thread. It need not be threadsafe unless the RHI support parallel translation.
 	* CAUTION: Platforms that support RHIThread but don't actually have a threadsafe implementation must flush internally with FScopedRHIThreadStaller StallRHIThread(FRHICommandListExecutor::GetImmediateCommandList()); when the call is from the render thread
 	* @param VertexDeclaration - existing vertex decl
 	* @param VertexShader - existing vertex shader
-	* @param HullShader - existing hull shader
-	* @param DomainShader - existing domain shader
 	* @param GeometryShader - existing geometry shader
 	* @param PixelShader - existing pixel shader
 	*/
@@ -1417,18 +1401,6 @@ public:
 		return RHI->CreateComputeShader_RenderThread(RHICmdList, Code, Hash);
 	}
 
-	virtual FHullShaderRHIRef CreateHullShader_RenderThread(class FRHICommandListImmediate& RHICmdList, TArrayView<const uint8> Code, const FSHAHash& Hash) override final
-	{
-		check(RHISupportsTessellation(GMaxRHIShaderPlatform));
-		return RHI->CreateHullShader_RenderThread(RHICmdList, Code, Hash);
-	}
-
-	virtual FDomainShaderRHIRef CreateDomainShader_RenderThread(class FRHICommandListImmediate& RHICmdList, TArrayView<const uint8> Code, const FSHAHash& Hash) override final
-	{
-		check(RHISupportsTessellation(GMaxRHIShaderPlatform));
-		return RHI->CreateDomainShader_RenderThread(RHICmdList, Code, Hash);
-	}
-
 	virtual void* LockTexture2D_RenderThread(class FRHICommandListImmediate& RHICmdList, FRHITexture2D* Texture, uint32 MipIndex, EResourceLockMode LockMode, uint32& DestStride, bool bLockWithinMiptail, bool bNeedsDefaultRHIFlush = true) override final
 	{
 		return RHI->LockTexture2D_RenderThread(RHICmdList, Texture, MipIndex, LockMode, DestStride, bLockWithinMiptail, bNeedsDefaultRHIFlush);
diff --git a/Engine/Source/Runtime/RenderCore/Private/Shader.cpp b/Engine/Source/Runtime/RenderCore/Private/Shader.cpp
--- a/Engine/Source/Runtime/RenderCore/Private/Shader.cpp
+++ b/Engine/Source/Runtime/RenderCore/Private/Shader.cpp
@@ -167,12 +167,8 @@ void FShaderParameterMap::VerifyBindingsAreComplete(const TCHAR* ShaderTypeName,
 		{
 			FString ErrorMessage = FString(TEXT("Found unbound parameters being used in shadertype ")) + ShaderTypeName + TEXT(" (VertexFactory: ") + VertexFactoryName + TEXT(")\n") + UnBoundParameters;
 
-			// There will be unbound parameters for Metal's "Hull" shader stage as it is merely a placeholder to provide binding indices to the RHI
-			if(!IsMetalPlatform((EShaderPlatform)Target.Platform) || Target.Frequency != SF_Hull)
-			{
-				// We use a non-Slate message box to avoid problem where we haven't compiled the shaders for Slate.
-				FPlatformMisc::MessageBoxExt(EAppMsgType::Ok, *ErrorMessage, TEXT("Error"));
-			}
+			// We use a non-Slate message box to avoid problem where we haven't compiled the shaders for Slate.
+			FPlatformMisc::MessageBoxExt(EAppMsgType::Ok, *ErrorMessage, TEXT("Error"));
 		}
 	}
 #endif // WITH_EDITORONLY_DATA
diff --git a/Engine/Source/Runtime/RenderCore/Private/ShaderCodeArchive.cpp b/Engine/Source/Runtime/RenderCore/Private/ShaderCodeArchive.cpp
--- a/Engine/Source/Runtime/RenderCore/Private/ShaderCodeArchive.cpp
+++ b/Engine/Source/Runtime/RenderCore/Private/ShaderCodeArchive.cpp
@@ -968,8 +968,6 @@ TRefCountPtr<FRHIShader> FShaderCodeArchive::CreateShader(int32 Index)
 	case SF_Amplification: Shader = RHICreateAmplificationShader(ShaderCodeView, ShaderHash); CheckShaderCreation(Shader, Index); break;
 	case SF_Pixel: Shader = RHICreatePixelShader(ShaderCodeView, ShaderHash); CheckShaderCreation(Shader, Index); break;
 	case SF_Geometry: Shader = RHICreateGeometryShader(ShaderCodeView, ShaderHash); CheckShaderCreation(Shader, Index); break;
-	case SF_Hull: Shader = RHICreateHullShader(ShaderCodeView, ShaderHash); CheckShaderCreation(Shader, Index); break;
-	case SF_Domain: Shader = RHICreateDomainShader(ShaderCodeView, ShaderHash); CheckShaderCreation(Shader, Index); break;
 	case SF_Compute: Shader = RHICreateComputeShader(ShaderCodeView, ShaderHash); CheckShaderCreation(Shader, Index); break;
 	case SF_RayGen: case SF_RayMiss: case SF_RayHitGroup: case SF_RayCallable:
 #if RHI_RAYTRACING
diff --git a/Engine/Source/Runtime/RenderCore/Private/ShaderCodeLibrary.cpp b/Engine/Source/Runtime/RenderCore/Private/ShaderCodeLibrary.cpp
--- a/Engine/Source/Runtime/RenderCore/Private/ShaderCodeLibrary.cpp
+++ b/Engine/Source/Runtime/RenderCore/Private/ShaderCodeLibrary.cpp
@@ -2620,24 +2620,6 @@ FPixelShaderRHIRef FShaderCodeLibrary::CreatePixelShader(EShaderPlatform Platfor
 	return nullptr;
 }
 
-FHullShaderRHIRef FShaderCodeLibrary::CreateHullShader(EShaderPlatform Platform, const FSHAHash& Hash)
-{
-	if (FShaderLibrariesCollection::Impl)
-	{
-		return FHullShaderRHIRef(FShaderLibrariesCollection::Impl->CreateShader(SF_Hull, Hash));
-	}
-	return nullptr;
-}
-
-FDomainShaderRHIRef FShaderCodeLibrary::CreateDomainShader(EShaderPlatform Platform, const FSHAHash& Hash)
-{
-	if (FShaderLibrariesCollection::Impl)
-	{
-		return FDomainShaderRHIRef(FShaderLibrariesCollection::Impl->CreateShader(SF_Domain, Hash));
-	}
-	return nullptr;
-}
-
 FGeometryShaderRHIRef FShaderCodeLibrary::CreateGeometryShader(EShaderPlatform Platform, const FSHAHash& Hash)
 {
 	if (FShaderLibrariesCollection::Impl)
diff --git a/Engine/Source/Runtime/RenderCore/Private/ShaderParameterStruct.cpp b/Engine/Source/Runtime/RenderCore/Private/ShaderParameterStruct.cpp
--- a/Engine/Source/Runtime/RenderCore/Private/ShaderParameterStruct.cpp
+++ b/Engine/Source/Runtime/RenderCore/Private/ShaderParameterStruct.cpp
@@ -229,8 +229,6 @@ void FShaderParameterBindings::BindForLegacyShaderParameters(const FShader* Shad
 	switch (Type->GetFrequency())
 	{
 	case SF_Vertex:
-	case SF_Hull:
-	case SF_Domain:
 	case SF_Pixel:
 	case SF_Geometry:
 	case SF_Compute:
diff --git a/Engine/Source/Runtime/RenderCore/Private/ShaderPipelineCache.cpp b/Engine/Source/Runtime/RenderCore/Private/ShaderPipelineCache.cpp
--- a/Engine/Source/Runtime/RenderCore/Private/ShaderPipelineCache.cpp
+++ b/Engine/Source/Runtime/RenderCore/Private/ShaderPipelineCache.cpp
@@ -641,22 +641,6 @@ bool FShaderPipelineCache::Precompile(FRHICommandListImmediate& RHICmdList, ESha
 					GraphicsInitializer.BoundShaderState.VertexShaderRHI = VertexShader;
 				}
 
-#if PLATFORM_SUPPORTS_TESSELLATION_SHADERS
-				FHullShaderRHIRef HullShader;
-				if (PSO.GraphicsDesc.HullShader != FSHAHash())
-				{
-					HullShader = FShaderCodeLibrary::CreateHullShader(Platform, PSO.GraphicsDesc.HullShader);
-					GraphicsInitializer.BoundShaderState.HullShaderRHI = HullShader;
-				}
-
-				FDomainShaderRHIRef DomainShader;
-				if (PSO.GraphicsDesc.DomainShader != FSHAHash())
-				{
-					DomainShader = FShaderCodeLibrary::CreateDomainShader(Platform, PSO.GraphicsDesc.DomainShader);
-					GraphicsInitializer.BoundShaderState.DomainShaderRHI = DomainShader;
-				}
-#endif
-
 #if PLATFORM_SUPPORTS_GEOMETRY_SHADERS
 				FGeometryShaderRHIRef GeometryShader;
 				if (PSO.GraphicsDesc.GeometryShader != FSHAHash())
@@ -842,19 +826,7 @@ void FShaderPipelineCache::PreparePipelineBatch(TDoubleLinkedList<FPipelineCache
                     RequiredShaders.Add(PSO.GraphicsDesc.VertexShader);
                     bOK &= FShaderCodeLibrary::ContainsShaderCode(PSO.GraphicsDesc.VertexShader);
                     UE_CLOG(!bOK, LogRHI, Verbose, TEXT("Failed to find VertexShader shader: %s"), *(PSO.GraphicsDesc.VertexShader.ToString()));
-					
-					if (PSO.GraphicsDesc.HullShader != EmptySHA)
-					{
-						RequiredShaders.Add(PSO.GraphicsDesc.HullShader);
-						bOK &= FShaderCodeLibrary::ContainsShaderCode(PSO.GraphicsDesc.HullShader);
-						UE_CLOG(!bOK, LogRHI, Verbose, TEXT("Failed to find HullShader shader: %s"), *(PSO.GraphicsDesc.HullShader.ToString()));
-					}
-					if (PSO.GraphicsDesc.DomainShader != EmptySHA)
-					{
-						RequiredShaders.Add(PSO.GraphicsDesc.DomainShader);
-						bOK &= FShaderCodeLibrary::ContainsShaderCode(PSO.GraphicsDesc.DomainShader);
-						UE_CLOG(!bOK, LogRHI, Verbose, TEXT("Failed to find DomainShader shader: %s"), *(PSO.GraphicsDesc.DomainShader.ToString()));
-					}
+
 					if (PSO.GraphicsDesc.GeometryShader != EmptySHA)
 					{
 						RequiredShaders.Add(PSO.GraphicsDesc.GeometryShader);
@@ -897,19 +869,7 @@ void FShaderPipelineCache::PreparePipelineBatch(TDoubleLinkedList<FPipelineCache
 					bOK &= FShaderCodeLibrary::PreloadShader(PSO.GraphicsDesc.AmplificationShader, AsyncJob.ReadRequests);
 					UE_CLOG(!bOK, LogRHI, Verbose, TEXT("Failed to read AmplificationShader shader: %s"), *(PSO.GraphicsDesc.AmplificationShader.ToString()));
 				}
-		
-				if (bOK && PSO.GraphicsDesc.HullShader != EmptySHA)
-				{
-					bOK &= FShaderCodeLibrary::PreloadShader(PSO.GraphicsDesc.HullShader, AsyncJob.ReadRequests);
-                    UE_CLOG(!bOK, LogRHI, Verbose, TEXT("Failed to read HullShader shader: %s"), *(PSO.GraphicsDesc.HullShader.ToString()));
-				}
-		
-				if (bOK && PSO.GraphicsDesc.DomainShader != EmptySHA)
-				{
-					bOK &= FShaderCodeLibrary::PreloadShader(PSO.GraphicsDesc.DomainShader, AsyncJob.ReadRequests);
-                    UE_CLOG(!bOK, LogRHI, Verbose, TEXT("Failed to read DomainShader shader: %s"), *(PSO.GraphicsDesc.DomainShader.ToString()));
-				}
-		
+
 				if (bOK && PSO.GraphicsDesc.FragmentShader != EmptySHA)
 				{
 					bOK &= FShaderCodeLibrary::PreloadShader(PSO.GraphicsDesc.FragmentShader, AsyncJob.ReadRequests);
diff --git a/Engine/Source/Runtime/RenderCore/Private/ShaderResource.cpp b/Engine/Source/Runtime/RenderCore/Private/ShaderResource.cpp
--- a/Engine/Source/Runtime/RenderCore/Private/ShaderResource.cpp
+++ b/Engine/Source/Runtime/RenderCore/Private/ShaderResource.cpp
@@ -459,8 +459,6 @@ TRefCountPtr<FRHIShader> FShaderMapResource_InlineCode::CreateRHIShader(int32 Sh
 	case SF_Mesh: RHIShader = RHICreateMeshShader(ShaderCodeView, ShaderHash); break;
 	case SF_Amplification: RHIShader = RHICreateAmplificationShader(ShaderCodeView, ShaderHash); break;
 	case SF_Pixel: RHIShader = RHICreatePixelShader(ShaderCodeView, ShaderHash); break;
-	case SF_Hull: RHIShader = RHICreateHullShader(ShaderCodeView, ShaderHash); break;
-	case SF_Domain: RHIShader = RHICreateDomainShader(ShaderCodeView, ShaderHash); break;
 	case SF_Geometry: RHIShader = RHICreateGeometryShader(ShaderCodeView, ShaderHash); break;
 	case SF_Compute: RHIShader = RHICreateComputeShader(ShaderCodeView, ShaderHash); break;
 	case SF_RayGen: case SF_RayMiss: case SF_RayHitGroup: case SF_RayCallable:
diff --git a/Engine/Source/Runtime/RenderCore/Public/Shader.h b/Engine/Source/Runtime/RenderCore/Public/Shader.h
--- a/Engine/Source/Runtime/RenderCore/Public/Shader.h
+++ b/Engine/Source/Runtime/RenderCore/Public/Shader.h
@@ -923,16 +923,6 @@ public:
 	{
 		return static_cast<FRHIPixelShader*>(GetRHIShaderBase(SF_Pixel));
 	}
-	/** @return the shader's hull shader */
-	inline FRHIHullShader* GetHullShader() const
-	{
-		return static_cast<FRHIHullShader*>(GetRHIShaderBase(SF_Hull));
-	}
-	/** @return the shader's domain shader */
-	inline FRHIDomainShader* GetDomainShader() const
-	{
-		return static_cast<FRHIDomainShader*>(GetRHIShaderBase(SF_Domain));
-	}
 	/** @return the shader's geometry shader */
 	inline FRHIGeometryShader* GetGeometryShader() const
 	{
diff --git a/Engine/Source/Runtime/RenderCore/Public/ShaderCodeLibrary.h b/Engine/Source/Runtime/RenderCore/Public/ShaderCodeLibrary.h
--- a/Engine/Source/Runtime/RenderCore/Public/ShaderCodeLibrary.h
+++ b/Engine/Source/Runtime/RenderCore/Public/ShaderCodeLibrary.h
@@ -190,8 +190,6 @@ struct RENDERCORE_API FShaderCodeLibrary
 
 	static FVertexShaderRHIRef CreateVertexShader(EShaderPlatform Platform, const FSHAHash& Hash);
 	static FPixelShaderRHIRef CreatePixelShader(EShaderPlatform Platform, const FSHAHash& Hash);
-	static FHullShaderRHIRef CreateHullShader(EShaderPlatform Platform, const FSHAHash& Hash);
-	static FDomainShaderRHIRef CreateDomainShader(EShaderPlatform Platform, const FSHAHash& Hash);
 	static FGeometryShaderRHIRef CreateGeometryShader(EShaderPlatform Platform, const FSHAHash& Hash);
 	static FComputeShaderRHIRef CreateComputeShader(EShaderPlatform Platform, const FSHAHash& Hash);
 	static FMeshShaderRHIRef CreateMeshShader(EShaderPlatform Platform, const FSHAHash& Hash);
diff --git a/Engine/Source/Runtime/RenderCore/Public/ShaderCompilerCore.h b/Engine/Source/Runtime/RenderCore/Public/ShaderCompilerCore.h
--- a/Engine/Source/Runtime/RenderCore/Public/ShaderCompilerCore.h
+++ b/Engine/Source/Runtime/RenderCore/Public/ShaderCompilerCore.h
@@ -307,12 +307,6 @@ struct FShaderCompilerInput
 
 	friend RENDERCORE_API FArchive& operator<<(FArchive& Ar, FShaderCompilerInput& Input);
 
-	bool IsUsingTessellation() const
-	{
-		// no more tessellation in UE5
-		return false;
-	}
-
 	bool IsRayTracingShader() const
 	{
 		return IsRayTracingShaderFrequency(Target.GetFrequency());
diff --git a/Engine/Source/Runtime/RenderCore/Public/ShaderCore.h b/Engine/Source/Runtime/RenderCore/Public/ShaderCore.h
--- a/Engine/Source/Runtime/RenderCore/Public/ShaderCore.h
+++ b/Engine/Source/Runtime/RenderCore/Public/ShaderCore.h
@@ -67,7 +67,7 @@ DECLARE_DWORD_ACCUMULATOR_STAT_EXTERN(TEXT("Num Shaders Duplicated"), STAT_Shade
 
 inline TStatId GetMemoryStatType(EShaderFrequency ShaderFrequency)
 {
-	static_assert(12 == SF_NumFrequencies, "EShaderFrequency has a bad size.");
+	static_assert(10 == SF_NumFrequencies, "EShaderFrequency has a bad size.");
 
 	switch(ShaderFrequency)
 	{
diff --git a/Engine/Source/Runtime/RenderCore/Public/ShaderParameterUtils.h b/Engine/Source/Runtime/RenderCore/Public/ShaderParameterUtils.h
--- a/Engine/Source/Runtime/RenderCore/Public/ShaderParameterUtils.h
+++ b/Engine/Source/Runtime/RenderCore/Public/ShaderParameterUtils.h
@@ -422,18 +422,6 @@ inline bool SetUAVParameterIfCS(TRHICmdList& RHICmdList, FRHIPixelShader* Shader
 	return UAVParameter.IsBound();
 }
 
-template<typename TRHICmdList>
-inline bool SetUAVParameterIfCS(TRHICmdList& RHICmdList, FRHIHullShader* Shader, const FShaderResourceParameter& UAVParameter, FRHIUnorderedAccessView* UAV)
-{
-	return false;
-}
-
-template<typename TRHICmdList>
-inline bool SetUAVParameterIfCS(TRHICmdList& RHICmdList, FRHIDomainShader* Shader, const FShaderResourceParameter& UAVParameter, FRHIUnorderedAccessView* UAV)
-{
-	return false;
-}
-
 template<typename TRHICmdList>
 inline bool SetUAVParameterIfCS(TRHICmdList& RHICmdList, FRHIGeometryShader* Shader, const FShaderResourceParameter& UAVParameter, FRHIUnorderedAccessView* UAV)
 {
diff --git a/Engine/Source/Runtime/RenderCore/Public/VertexFactory.h b/Engine/Source/Runtime/RenderCore/Public/VertexFactory.h
--- a/Engine/Source/Runtime/RenderCore/Public/VertexFactory.h
+++ b/Engine/Source/Runtime/RenderCore/Public/VertexFactory.h
@@ -244,8 +244,6 @@ static const FTypeLayoutDesc* GetVertexFactoryParametersLayout(EShaderFrequency
 	case SF_Vertex: return TVertexFactoryParameterTraits<SF_Vertex, VertexFactoryType>::GetLayout();
 	case SF_Mesh: return TVertexFactoryParameterTraits<SF_Mesh, VertexFactoryType>::GetLayout();
 	case SF_Amplification: return TVertexFactoryParameterTraits<SF_Amplification, VertexFactoryType>::GetLayout();
-	case SF_Hull: return TVertexFactoryParameterTraits<SF_Hull, VertexFactoryType>::GetLayout();
-	case SF_Domain: return TVertexFactoryParameterTraits<SF_Domain, VertexFactoryType>::GetLayout();
 	case SF_Pixel: return TVertexFactoryParameterTraits<SF_Pixel, VertexFactoryType>::GetLayout();
 	case SF_Geometry: return TVertexFactoryParameterTraits<SF_Geometry, VertexFactoryType>::GetLayout();
 	case SF_Compute: return TVertexFactoryParameterTraits<SF_Compute, VertexFactoryType>::GetLayout();
@@ -265,8 +263,6 @@ static FVertexFactoryShaderParameters* ConstructVertexFactoryParameters(EShaderF
 	case SF_Vertex: return TVertexFactoryParameterTraits<SF_Vertex, VertexFactoryType>::Create(ParameterMap);
 	case SF_Mesh: return TVertexFactoryParameterTraits<SF_Mesh, VertexFactoryType>::Create(ParameterMap);
 	case SF_Amplification: return TVertexFactoryParameterTraits<SF_Amplification, VertexFactoryType>::Create(ParameterMap);
-	case SF_Hull: return TVertexFactoryParameterTraits<SF_Hull, VertexFactoryType>::Create(ParameterMap);
-	case SF_Domain: return TVertexFactoryParameterTraits<SF_Domain, VertexFactoryType>::Create(ParameterMap);
 	case SF_Pixel: return TVertexFactoryParameterTraits<SF_Pixel, VertexFactoryType>::Create(ParameterMap);
 	case SF_Geometry: return TVertexFactoryParameterTraits<SF_Geometry, VertexFactoryType>::Create(ParameterMap);
 	case SF_Compute: return TVertexFactoryParameterTraits<SF_Compute, VertexFactoryType>::Create(ParameterMap);
@@ -296,8 +292,6 @@ static void GetVertexFactoryParametersElementShaderBindings(EShaderFrequency Sha
 	case SF_Vertex: TVertexFactoryParameterTraits<SF_Vertex, VertexFactoryType>::GetElementShaderBindings(Parameters, Scene, View, Shader, InputStreamType, FeatureLevel, VertexFactory, BatchElement, ShaderBindings, VertexStreams); break;
 	case SF_Mesh: TVertexFactoryParameterTraits<SF_Mesh, VertexFactoryType>::GetElementShaderBindings(Parameters, Scene, View, Shader, InputStreamType, FeatureLevel, VertexFactory, BatchElement, ShaderBindings, VertexStreams); break;
 	case SF_Amplification: TVertexFactoryParameterTraits<SF_Amplification, VertexFactoryType>::GetElementShaderBindings(Parameters, Scene, View, Shader, InputStreamType, FeatureLevel, VertexFactory, BatchElement, ShaderBindings, VertexStreams); break;
-	case SF_Hull: TVertexFactoryParameterTraits<SF_Hull, VertexFactoryType>::GetElementShaderBindings(Parameters, Scene, View, Shader, InputStreamType, FeatureLevel, VertexFactory, BatchElement, ShaderBindings, VertexStreams); break;
-	case SF_Domain: TVertexFactoryParameterTraits<SF_Domain, VertexFactoryType>::GetElementShaderBindings(Parameters, Scene, View, Shader, InputStreamType, FeatureLevel, VertexFactory, BatchElement, ShaderBindings, VertexStreams); break;
 	case SF_Pixel: TVertexFactoryParameterTraits<SF_Pixel, VertexFactoryType>::GetElementShaderBindings(Parameters, Scene, View, Shader, InputStreamType, FeatureLevel, VertexFactory, BatchElement, ShaderBindings, VertexStreams); break;
 	case SF_Geometry: TVertexFactoryParameterTraits<SF_Geometry, VertexFactoryType>::GetElementShaderBindings(Parameters, Scene, View, Shader, InputStreamType, FeatureLevel, VertexFactory, BatchElement, ShaderBindings, VertexStreams); break;
 	case SF_Compute: TVertexFactoryParameterTraits<SF_Compute, VertexFactoryType>::GetElementShaderBindings(Parameters, Scene, View, Shader, InputStreamType, FeatureLevel, VertexFactory, BatchElement, ShaderBindings, VertexStreams); break;
diff --git a/Engine/Source/Runtime/Renderer/Private/AtmosphereRendering.cpp b/Engine/Source/Runtime/Renderer/Private/AtmosphereRendering.cpp
--- a/Engine/Source/Runtime/Renderer/Private/AtmosphereRendering.cpp
+++ b/Engine/Source/Runtime/Renderer/Private/AtmosphereRendering.cpp
@@ -80,8 +80,6 @@ void FAtmosphereShaderTextureParameters::Bind(const FShaderParameterMap& Paramet
 	template void FAtmosphereShaderTextureParameters::Set< TRHIShader >( FRHICommandList& RHICmdList, TRHIShader* ShaderRHI, const FSceneView& View ) const;
 
 IMPLEMENT_ATMOSPHERE_TEXTURE_PARAM_SET(FRHIVertexShader);
-IMPLEMENT_ATMOSPHERE_TEXTURE_PARAM_SET(FRHIHullShader);
-IMPLEMENT_ATMOSPHERE_TEXTURE_PARAM_SET(FRHIDomainShader);
 IMPLEMENT_ATMOSPHERE_TEXTURE_PARAM_SET(FRHIGeometryShader);
 IMPLEMENT_ATMOSPHERE_TEXTURE_PARAM_SET(FRHIPixelShader);
 IMPLEMENT_ATMOSPHERE_TEXTURE_PARAM_SET(FRHIComputeShader);
diff --git a/Engine/Source/Runtime/Renderer/Private/ShaderBaseClasses.cpp b/Engine/Source/Runtime/Renderer/Private/ShaderBaseClasses.cpp
--- a/Engine/Source/Runtime/Renderer/Private/ShaderBaseClasses.cpp
+++ b/Engine/Source/Runtime/Renderer/Private/ShaderBaseClasses.cpp
@@ -332,8 +332,6 @@ void FMaterialShader::SetParameters(
 	);
 
 IMPLEMENT_MATERIAL_SHADER_SetParameters(FRHIVertexShader);
-IMPLEMENT_MATERIAL_SHADER_SetParameters(FRHIHullShader);
-IMPLEMENT_MATERIAL_SHADER_SetParameters(FRHIDomainShader);
 IMPLEMENT_MATERIAL_SHADER_SetParameters(FRHIGeometryShader);
 IMPLEMENT_MATERIAL_SHADER_SetParameters(FRHIPixelShader);
 IMPLEMENT_MATERIAL_SHADER_SetParameters(FRHIComputeShader);
diff --git a/Engine/Source/Runtime/VulkanRHI/Private/VulkanPipeline.cpp b/Engine/Source/Runtime/VulkanRHI/Private/VulkanPipeline.cpp
--- a/Engine/Source/Runtime/VulkanRHI/Private/VulkanPipeline.cpp
+++ b/Engine/Source/Runtime/VulkanRHI/Private/VulkanPipeline.cpp
@@ -1115,7 +1115,6 @@ bool FVulkanPipelineStateCacheManager::CreateGfxPipelineFromEntry(FVulkanRHIGrap
 	PipelineInfo.pStages = ShaderStages;
 	// main_00000000_00000000
 	ANSICHAR EntryPoints[ShaderStage::NumStages][24];
-	bool bHasTessellation = false;
 	for (int32 ShaderStage = 0; ShaderStage < ShaderStage::NumStages; ++ShaderStage)
 	{
 		if (!PSO->ShaderModules[ShaderStage])
@@ -1127,7 +1126,6 @@ bool FVulkanPipelineStateCacheManager::CreateGfxPipelineFromEntry(FVulkanRHIGrap
 		ShaderStages[PipelineInfo.stageCount].sType = VK_STRUCTURE_TYPE_PIPELINE_SHADER_STAGE_CREATE_INFO;
 		VkShaderStageFlagBits Stage = UEFrequencyToVKStageBit(ShaderStage::GetFrequencyForGfxStage(CurrStage));
 		ShaderStages[PipelineInfo.stageCount].stage = Stage;
-		bHasTessellation = bHasTessellation || ((Stage & (VK_SHADER_STAGE_TESSELLATION_CONTROL_BIT | VK_SHADER_STAGE_TESSELLATION_EVALUATION_BIT)) != 0);
 		ShaderStages[PipelineInfo.stageCount].module = PSO->ShaderModules[CurrStage];
 		Shaders[ShaderStage]->GetEntryPoint(EntryPoints[PipelineInfo.stageCount], 24);
 		ShaderStages[PipelineInfo.stageCount].pName = EntryPoints[PipelineInfo.stageCount];
@@ -1193,15 +1191,6 @@ bool FVulkanPipelineStateCacheManager::CreateGfxPipelineFromEntry(FVulkanRHIGrap
 
 	PipelineInfo.pDynamicState = &DynamicState;
 
-	VkPipelineTessellationStateCreateInfo TessState;
-	if (bHasTessellation)
-	{
-		ZeroVulkanStruct(TessState, VK_STRUCTURE_TYPE_PIPELINE_TESSELLATION_STATE_CREATE_INFO);
-		PipelineInfo.pTessellationState = &TessState;
-		check(InputAssembly.topology == VK_PRIMITIVE_TOPOLOGY_PATCH_LIST);
-		TessState.patchControlPoints = GfxEntry->ControlPoints;
-	}
-
 	VkResult Result = VK_ERROR_INITIALIZATION_FAILED;
 	double BeginTime = FPlatformTime::Seconds();
 	if(bUseLRU)
@@ -1363,110 +1352,55 @@ FVulkanLayout* FVulkanPipelineStateCacheManager::FindOrAddLayout(const FVulkanDe
 	return Layout;
 }
 
-static inline VkPrimitiveTopology UEToVulkanTopologyType(const FVulkanDevice* InDevice, EPrimitiveType PrimitiveType, bool bHasTessellation, uint16& OutControlPoints)
-{
-	if (bHasTessellation)
-	{
-		switch (PrimitiveType)
-		{
-		case PT_TriangleList:
-			// This is the case for tessellation without AEN or other buffers, so just flip to 3 CPs
-			OutControlPoints = 3;
-			return VK_PRIMITIVE_TOPOLOGY_PATCH_LIST;
-		case PT_1_ControlPointPatchList:
-		case PT_2_ControlPointPatchList:
-		case PT_3_ControlPointPatchList:
-		case PT_4_ControlPointPatchList:
-		case PT_5_ControlPointPatchList:
-		case PT_6_ControlPointPatchList:
-		case PT_7_ControlPointPatchList:
-		case PT_8_ControlPointPatchList:
-		case PT_9_ControlPointPatchList:
-		case PT_10_ControlPointPatchList:
-		case PT_12_ControlPointPatchList:
-		case PT_13_ControlPointPatchList:
-		case PT_14_ControlPointPatchList:
-		case PT_15_ControlPointPatchList:
-		case PT_16_ControlPointPatchList:
-		case PT_17_ControlPointPatchList:
-		case PT_18_ControlPointPatchList:
-		case PT_19_ControlPointPatchList:
-		case PT_20_ControlPointPatchList:
-		case PT_22_ControlPointPatchList:
-		case PT_23_ControlPointPatchList:
-		case PT_24_ControlPointPatchList:
-		case PT_25_ControlPointPatchList:
-		case PT_26_ControlPointPatchList:
-		case PT_27_ControlPointPatchList:
-		case PT_28_ControlPointPatchList:
-		case PT_29_ControlPointPatchList:
-		case PT_30_ControlPointPatchList:
-		case PT_31_ControlPointPatchList:
-		case PT_32_ControlPointPatchList:
-			OutControlPoints = (PrimitiveType - PT_1_ControlPointPatchList + 1);
-			checkf(
-				OutControlPoints <= InDevice->GetLimits().maxTessellationPatchSize,
-				TEXT("OutControlPoints (%d) exceeded limit of maximal patch size (%d)"),
-				OutControlPoints,
-				InDevice->GetLimits().maxTessellationPatchSize
-			);
-			return VK_PRIMITIVE_TOPOLOGY_PATCH_LIST;
-		default:
-			checkf(false, TEXT("Unsupported tessellation EPrimitiveType %d; probably missing a case in FStaticMeshSceneProxy::GetMeshElement()!"), (uint32)PrimitiveType);
-			break;
-		}
-		OutControlPoints = 0;
-	}
-	else
-	{
-		OutControlPoints = 0;
-		switch (PrimitiveType)
-		{
-		case PT_PointList:
-			return VK_PRIMITIVE_TOPOLOGY_POINT_LIST;
-		case PT_LineList:
-			return VK_PRIMITIVE_TOPOLOGY_LINE_LIST;
-		case PT_TriangleList:
-			return VK_PRIMITIVE_TOPOLOGY_TRIANGLE_LIST;
-		case PT_TriangleStrip:
-			return VK_PRIMITIVE_TOPOLOGY_TRIANGLE_STRIP;
-		case PT_1_ControlPointPatchList:
-		case PT_2_ControlPointPatchList:
-		case PT_3_ControlPointPatchList:
-		case PT_4_ControlPointPatchList:
-		case PT_5_ControlPointPatchList:
-		case PT_6_ControlPointPatchList:
-		case PT_7_ControlPointPatchList:
-		case PT_8_ControlPointPatchList:
-		case PT_9_ControlPointPatchList:
-		case PT_10_ControlPointPatchList:
-		case PT_12_ControlPointPatchList:
-		case PT_13_ControlPointPatchList:
-		case PT_14_ControlPointPatchList:
-		case PT_15_ControlPointPatchList:
-		case PT_16_ControlPointPatchList:
-		case PT_17_ControlPointPatchList:
-		case PT_18_ControlPointPatchList:
-		case PT_19_ControlPointPatchList:
-		case PT_20_ControlPointPatchList:
-		case PT_22_ControlPointPatchList:
-		case PT_23_ControlPointPatchList:
-		case PT_24_ControlPointPatchList:
-		case PT_25_ControlPointPatchList:
-		case PT_26_ControlPointPatchList:
-		case PT_27_ControlPointPatchList:
-		case PT_28_ControlPointPatchList:
-		case PT_29_ControlPointPatchList:
-		case PT_30_ControlPointPatchList:
-		case PT_31_ControlPointPatchList:
-		case PT_32_ControlPointPatchList:
-			OutControlPoints = (PrimitiveType - PT_1_ControlPointPatchList + 1);
-			checkf(false, TEXT("Missing tessellation shaders, however tried to use EPrimitiveType %d (%d control points)"), (uint32)PrimitiveType, OutControlPoints);
-			break;
-		default:
-			checkf(false, TEXT("Unsupported EPrimitiveType %d"), (uint32)PrimitiveType);
-			break;
-		}
+static inline VkPrimitiveTopology UEToVulkanTopologyType(const FVulkanDevice* InDevice, EPrimitiveType PrimitiveType, uint16& OutControlPoints)
+{
+	OutControlPoints = 0;
+	switch (PrimitiveType)
+	{
+	case PT_PointList:
+		return VK_PRIMITIVE_TOPOLOGY_POINT_LIST;
+	case PT_LineList:
+		return VK_PRIMITIVE_TOPOLOGY_LINE_LIST;
+	case PT_TriangleList:
+		return VK_PRIMITIVE_TOPOLOGY_TRIANGLE_LIST;
+	case PT_TriangleStrip:
+		return VK_PRIMITIVE_TOPOLOGY_TRIANGLE_STRIP;
+	case PT_1_ControlPointPatchList:
+	case PT_2_ControlPointPatchList:
+	case PT_3_ControlPointPatchList:
+	case PT_4_ControlPointPatchList:
+	case PT_5_ControlPointPatchList:
+	case PT_6_ControlPointPatchList:
+	case PT_7_ControlPointPatchList:
+	case PT_8_ControlPointPatchList:
+	case PT_9_ControlPointPatchList:
+	case PT_10_ControlPointPatchList:
+	case PT_12_ControlPointPatchList:
+	case PT_13_ControlPointPatchList:
+	case PT_14_ControlPointPatchList:
+	case PT_15_ControlPointPatchList:
+	case PT_16_ControlPointPatchList:
+	case PT_17_ControlPointPatchList:
+	case PT_18_ControlPointPatchList:
+	case PT_19_ControlPointPatchList:
+	case PT_20_ControlPointPatchList:
+	case PT_22_ControlPointPatchList:
+	case PT_23_ControlPointPatchList:
+	case PT_24_ControlPointPatchList:
+	case PT_25_ControlPointPatchList:
+	case PT_26_ControlPointPatchList:
+	case PT_27_ControlPointPatchList:
+	case PT_28_ControlPointPatchList:
+	case PT_29_ControlPointPatchList:
+	case PT_30_ControlPointPatchList:
+	case PT_31_ControlPointPatchList:
+	case PT_32_ControlPointPatchList:
+		OutControlPoints = (PrimitiveType - PT_1_ControlPointPatchList + 1);
+		checkf(false, TEXT("Missing tessellation shaders, however tried to use EPrimitiveType %d (%d control points)"), (uint32)PrimitiveType, OutControlPoints);
+		break;
+	default:
+		checkf(false, TEXT("Unsupported EPrimitiveType %d"), (uint32)PrimitiveType);
+		break;
 	}
 
 	return VK_PRIMITIVE_TOPOLOGY_MAX_ENUM;
@@ -1526,10 +1460,8 @@ void FVulkanPipelineStateCacheManager::CreateGfxEntry(const FGraphicsPipelineSta
 
 	OutGfxEntry->UseAlphaToCoverage = PSOInitializer.NumSamples > 1 && BlendState->Initializer.bUseAlphaToCoverage ? 1 : 0;
 
-	const bool bHasTessellation = (PSOInitializer.BoundShaderState.DomainShaderRHI != nullptr);
-
 	OutGfxEntry->RasterizationSamples = PSOInitializer.NumSamples;
-	OutGfxEntry->Topology = (uint32)UEToVulkanTopologyType(Device, PSOInitializer.PrimitiveType, bHasTessellation, OutGfxEntry->ControlPoints);
+	OutGfxEntry->Topology = (uint32)UEToVulkanTopologyType(Device, PSOInitializer.PrimitiveType, OutGfxEntry->ControlPoints);
 	uint32 NumRenderTargets = PSOInitializer.ComputeNumValidRenderTargets();
 	
 	if (PSOInitializer.SubpassHint == ESubpassHint::DeferredShadingSubpass && PSOInitializer.SubpassIndex >= 2)
diff --git a/Engine/Source/Runtime/Windows/D3D11RHI/Private/D3D11Commands.cpp b/Engine/Source/Runtime/Windows/D3D11RHI/Private/D3D11Commands.cpp
--- a/Engine/Source/Runtime/Windows/D3D11RHI/Private/D3D11Commands.cpp
+++ b/Engine/Source/Runtime/Windows/D3D11RHI/Private/D3D11Commands.cpp
@@ -332,8 +332,6 @@ void FD3D11DynamicRHI::RHISetBoundShaderState(FRHIBoundShaderState* BoundShaderS
 	// Shader changed so all resource tables are dirty
 	DirtyUniformBuffers[SF_Vertex] = 0xffff;
 	DirtyUniformBuffers[SF_Pixel] = 0xffff;
-	DirtyUniformBuffers[SF_Hull] = 0xffff;
-	DirtyUniformBuffers[SF_Domain] = 0xffff;
 	DirtyUniformBuffers[SF_Geometry] = 0xffff;
 
 	// Shader changed.  All UB's must be reset by high level code to match other platforms anway.
