Removal of incorrect or unnecessary ESPMode::Fast usage.

diff --git a/Engine/Source/Editor/UnrealEd/Private/PlayLevel.cpp b/Engine/Source/Editor/UnrealEd/Private/PlayLevel.cpp
--- a/Engine/Source/Editor/UnrealEd/Private/PlayLevel.cpp
+++ b/Engine/Source/Editor/UnrealEd/Private/PlayLevel.cpp
@@ -3106,7 +3106,7 @@ TSharedRef<SPIEViewport> UEditorEngine::GeneratePIEViewportWindow(const FRequest
 #if PLATFORM_MAC
 		FSlateApplication::Get().AddWindow(PieWindow.ToSharedRef());
 #else
-		TSharedRef<SWindow, ESPMode::Fast> MainWindow = FModuleManager::LoadModuleChecked<IMainFrameModule>(TEXT("MainFrame")).GetParentWindow().ToSharedRef();
+		TSharedRef<SWindow> MainWindow = FModuleManager::LoadModuleChecked<IMainFrameModule>(TEXT("MainFrame")).GetParentWindow().ToSharedRef();
 		if (InSessionParams.EditorPlaySettings->PIEAlwaysOnTop)
 		{
 			FSlateApplication::Get().AddWindowAsNativeChild(PieWindow.ToSharedRef(), MainWindow, true);
diff --git a/Engine/Source/Editor/UnrealEd/Public/SourceCodeNavigation.h b/Engine/Source/Editor/UnrealEd/Public/SourceCodeNavigation.h
--- a/Engine/Source/Editor/UnrealEd/Public/SourceCodeNavigation.h
+++ b/Engine/Source/Editor/UnrealEd/Public/SourceCodeNavigation.h
@@ -52,7 +52,7 @@ private:
 * Should be registered with FSourceCodeNavigation::AddNavigationHandler and unregistered
 * using FSourceCodeNavigation::Remove.NavigationHandler
 */
-class ISourceCodeNavigationHandler : public TSharedFromThis<ISourceCodeNavigationHandler, ESPMode::Fast>
+class ISourceCodeNavigationHandler : public TSharedFromThis<ISourceCodeNavigationHandler>
 {
 public:
 	virtual ~ISourceCodeNavigationHandler() {};
diff --git a/Engine/Source/Programs/SlateViewer/Private/IOS/IOSSlateViewerMain.cpp b/Engine/Source/Programs/SlateViewer/Private/IOS/IOSSlateViewerMain.cpp
--- a/Engine/Source/Programs/SlateViewer/Private/IOS/IOSSlateViewerMain.cpp
+++ b/Engine/Source/Programs/SlateViewer/Private/IOS/IOSSlateViewerMain.cpp
@@ -70,7 +70,7 @@ void FAppEntry::Init()
 	}
 
 #if WITH_SHARED_POINTER_TESTS
-	SharedPointerTesting::TestSharedPointer< ESPMode::Fast >();
+	SharedPointerTesting::TestSharedPointer< ESPMode::NotThreadSafe >();
 	SharedPointerTesting::TestSharedPointer< ESPMode::ThreadSafe >();
 #endif
 	
diff --git a/Engine/Source/Programs/SlateViewer/Private/SlateViewerApp.cpp b/Engine/Source/Programs/SlateViewer/Private/SlateViewerApp.cpp
--- a/Engine/Source/Programs/SlateViewer/Private/SlateViewerApp.cpp
+++ b/Engine/Source/Programs/SlateViewer/Private/SlateViewerApp.cpp
@@ -83,7 +83,7 @@ int RunSlateViewer( const TCHAR* CommandLine )
 
 
 #if WITH_SHARED_POINTER_TESTS
-	SharedPointerTesting::TestSharedPointer<ESPMode::Fast>();
+	SharedPointerTesting::TestSharedPointer<ESPMode::NotThreadSafe>();
 	SharedPointerTesting::TestSharedPointer<ESPMode::ThreadSafe>();
 #endif
 
diff --git a/Engine/Source/Programs/UnrealInsights/Private/IOS/IOSUnrealInsightsMain.cpp b/Engine/Source/Programs/UnrealInsights/Private/IOS/IOSUnrealInsightsMain.cpp
--- a/Engine/Source/Programs/UnrealInsights/Private/IOS/IOSUnrealInsightsMain.cpp
+++ b/Engine/Source/Programs/UnrealInsights/Private/IOS/IOSUnrealInsightsMain.cpp
@@ -83,7 +83,7 @@ void FAppEntry::Init()
 	}
 
 #if WITH_SHARED_POINTER_TESTS
-	SharedPointerTesting::TestSharedPointer<ESPMode::Fast>();
+	SharedPointerTesting::TestSharedPointer<ESPMode::NotThreadSafe>();
 	SharedPointerTesting::TestSharedPointer<ESPMode::ThreadSafe>();
 #endif
 
diff --git a/Engine/Source/Programs/UnrealInsights/Private/UserInterfaceCommand.cpp b/Engine/Source/Programs/UnrealInsights/Private/UserInterfaceCommand.cpp
--- a/Engine/Source/Programs/UnrealInsights/Private/UserInterfaceCommand.cpp
+++ b/Engine/Source/Programs/UnrealInsights/Private/UserInterfaceCommand.cpp
@@ -195,7 +195,7 @@ void FUserInterfaceCommand::Run()
 #endif
 
 #if WITH_SHARED_POINTER_TESTS
-	SharedPointerTesting::TestSharedPointer<ESPMode::Fast>();
+	SharedPointerTesting::TestSharedPointer<ESPMode::NotThreadSafe>();
 	SharedPointerTesting::TestSharedPointer<ESPMode::ThreadSafe>();
 #endif
 
diff --git a/Engine/Source/Runtime/AIModule/Classes/BrainComponent.h b/Engine/Source/Runtime/AIModule/Classes/BrainComponent.h
--- a/Engine/Source/Runtime/AIModule/Classes/BrainComponent.h
+++ b/Engine/Source/Runtime/AIModule/Classes/BrainComponent.h
@@ -65,7 +65,7 @@ struct AIMODULE_API FAIMessage
 	static void Broadcast(UObject* WorldContextObject, const FAIMessage& Message);
 };
 
-typedef TSharedPtr<struct FAIMessageObserver, ESPMode::Fast> FAIMessageObserverHandle;
+typedef TSharedPtr<struct FAIMessageObserver> FAIMessageObserverHandle;
 
 struct AIMODULE_API FAIMessageObserver : public TSharedFromThis<FAIMessageObserver>
 {
diff --git a/Engine/Source/Runtime/Engine/Classes/Engine/StreamableManager.h b/Engine/Source/Runtime/Engine/Classes/Engine/StreamableManager.h
--- a/Engine/Source/Runtime/Engine/Classes/Engine/StreamableManager.h
+++ b/Engine/Source/Runtime/Engine/Classes/Engine/StreamableManager.h
@@ -14,7 +14,7 @@ DECLARE_DELEGATE(FStreamableDelegate);
 DECLARE_DELEGATE_OneParam(FStreamableUpdateDelegate, TSharedRef<struct FStreamableHandle>);
 
 /** A handle to a synchronous or async load. As long as the handle is Active, loaded assets will stay in memory */
-struct ENGINE_API FStreamableHandle : public TSharedFromThis<FStreamableHandle, ESPMode::Fast>
+struct ENGINE_API FStreamableHandle : public TSharedFromThis<FStreamableHandle>
 {
 	/** 
 	 * If this request has finished loading, meaning all available assets were loaded
diff --git a/Engine/Source/Runtime/Net/Core/Public/Net/Core/Analytics/NetAnalytics.h b/Engine/Source/Runtime/Net/Core/Public/Net/Core/Analytics/NetAnalytics.h
--- a/Engine/Source/Runtime/Net/Core/Public/Net/Core/Analytics/NetAnalytics.h
+++ b/Engine/Source/Runtime/Net/Core/Public/Net/Core/Analytics/NetAnalytics.h
@@ -48,10 +48,10 @@ class FNetAnalyticsAggregator;
 #define NET_ANALYTICS_MULTITHREADING 0
 
 
-#if NET_ANALYTICS_MULTITHREADING
+#if NET_ANALYTICS_MULTITHREADING || FORCE_THREADSAFE_SHAREDPTRS
 	constexpr const ESPMode NetAnalyticsThreadSafety = ESPMode::ThreadSafe;
 #else
-	constexpr const ESPMode NetAnalyticsThreadSafety = ESPMode::Fast;
+	constexpr const ESPMode NetAnalyticsThreadSafety = ESPMode::NotThreadSafe;
 #endif
 
 
diff --git a/Engine/Source/Runtime/RHI/Private/PipelineStateCache.cpp b/Engine/Source/Runtime/RHI/Private/PipelineStateCache.cpp
--- a/Engine/Source/Runtime/RHI/Private/PipelineStateCache.cpp
+++ b/Engine/Source/Runtime/RHI/Private/PipelineStateCache.cpp
@@ -1049,12 +1049,12 @@ FComputePipelineState* PipelineStateCache::GetAndOrCreateComputePipelineState(FR
 	bool DoAsyncCompile = IsAsyncCompilationAllowed(RHICmdList);
 
 
-	TSharedPtr<FComputePipelineState, ESPMode::Fast> OutCachedState;
+	TSharedPtr<FComputePipelineState> OutCachedState;
 
 	// Find or add an entry for this initializer
 	bool WasFound = GComputePipelineCache.FindOrAdd(ComputeShader, OutCachedState, [&RHICmdList, &ComputeShader, &DoAsyncCompile] {
 			// create new graphics state
-			TSharedPtr<FComputePipelineState, ESPMode::Fast> PipelineState(new FComputePipelineState(ComputeShader));
+			TSharedPtr<FComputePipelineState> PipelineState(new FComputePipelineState(ComputeShader));
 			PipelineState->Stats = FPipelineFileCache::RegisterPSOStats(GetTypeHash(ComputeShader));
 
 			// create a compilation task, or just do it now...
@@ -1488,7 +1488,7 @@ void DumpPipelineCacheStats()
 
 	for (auto GraphicsPipeLine : GGraphicsPipelineCache.Current())
 	{
-		TSharedPtr<FGraphicsPipelineState, ESPMode::Fast> State = GraphicsPipeLine.Value;
+		TSharedPtr<FGraphicsPipelineState> State = GraphicsPipeLine.Value;
 
 		// calc timestats
 		double SinceUse = FPlatformTime::Seconds() - State->FirstUsedTime;
