Added FRHI{Texture, Buffer}ViewCache to clean up FRDGPooledX / FRHITransientX duplicate code. Added support for UAVs with a format.

diff --git a/Engine/Source/Runtime/RHI/Private/RHI.cpp b/Engine/Source/Runtime/RHI/Private/RHI.cpp
--- a/Engine/Source/Runtime/RHI/Private/RHI.cpp
+++ b/Engine/Source/Runtime/RHI/Private/RHI.cpp
@@ -5,6 +5,7 @@
 =============================================================================*/
 
 #include "RHI.h"
+#include "RHITransientResourceAllocator.h"
 #include "Modules/ModuleManager.h"
 #include "Misc/ConfigCacheIni.h"
 #include "Misc/MessageDialog.h"
@@ -2010,3 +2011,192 @@ SIZE_T CalculateImageBytes(uint32 SizeX,uint32 SizeY,uint32 SizeZ,uint8 Format)
 		return static_cast<SIZE_T>(SizeX / GPixelFormats[Format].BlockSizeX) * (SizeY / GPixelFormats[Format].BlockSizeY) * GPixelFormats[Format].BlockBytes;
 	}
 }
+
+FRHIShaderResourceView* FRHITextureViewCache::GetOrCreateSRV(FRHITexture* Texture, const FRHITextureSRVCreateInfo& SRVCreateInfo)
+{
+	for (const auto& KeyValue : SRVs)
+	{
+		if (KeyValue.Key == SRVCreateInfo)
+		{
+			return KeyValue.Value.GetReference();
+		}
+	}
+
+	FShaderResourceViewRHIRef RHIShaderResourceView;
+
+	if (SRVCreateInfo.MetaData == ERHITextureMetaDataAccess::None)
+	{
+		RHIShaderResourceView = RHICreateShaderResourceView(Texture, SRVCreateInfo);
+	}
+	else
+	{
+		FRHITexture2D* Texture2D = Texture->GetTexture2D();
+		check(Texture2D);
+
+		switch (SRVCreateInfo.MetaData)
+		{
+		case ERHITextureMetaDataAccess::HTile:
+			check(GRHISupportsExplicitHTile);
+			RHIShaderResourceView = RHICreateShaderResourceViewHTile(Texture2D);
+			break;
+
+		case ERHITextureMetaDataAccess::FMask:
+			RHIShaderResourceView = RHICreateShaderResourceViewFMask(Texture2D);
+			break;
+
+		case ERHITextureMetaDataAccess::CMask:
+			RHIShaderResourceView = RHICreateShaderResourceViewWriteMask(Texture2D);
+			break;
+
+		default:
+			checkf(false, TEXT("Invalid texture metadata access used when creating SRV."));
+			return nullptr;
+		}
+	}
+
+	check(RHIShaderResourceView);
+	FRHIShaderResourceView* View = RHIShaderResourceView.GetReference();
+	SRVs.Emplace(SRVCreateInfo, MoveTemp(RHIShaderResourceView));
+	return View;
+}
+
+FRHIUnorderedAccessView* FRHITextureViewCache::GetOrCreateUAV(FRHITexture* Texture, const FRHITextureUAVCreateInfo& UAVCreateInfo)
+{
+	for (const auto& KeyValue : UAVs)
+	{
+		if (KeyValue.Key == UAVCreateInfo)
+		{
+			return KeyValue.Value.GetReference();
+		}
+	}
+
+	FUnorderedAccessViewRHIRef RHIUnorderedAccessView;
+
+	if (UAVCreateInfo.MetaData == ERHITextureMetaDataAccess::None)
+	{
+		if (UAVCreateInfo.Format != PF_Unknown)
+		{
+			RHIUnorderedAccessView = RHICreateUnorderedAccessView(Texture, UAVCreateInfo.MipLevel, UAVCreateInfo.Format);
+		}
+		else
+		{
+			RHIUnorderedAccessView = RHICreateUnorderedAccessView(Texture, UAVCreateInfo.MipLevel);
+		}
+	}
+	else
+	{
+		FRHITexture2D* Texture2D = Texture->GetTexture2D();
+		check(Texture2D);
+
+		switch (UAVCreateInfo.MetaData)
+		{
+		case ERHITextureMetaDataAccess::HTile:
+			check(GRHISupportsExplicitHTile);
+			RHIUnorderedAccessView = RHICreateUnorderedAccessViewHTile(Texture2D);
+			break;
+
+		case ERHITextureMetaDataAccess::Stencil:
+			RHIUnorderedAccessView = RHICreateUnorderedAccessViewStencil(Texture2D, UAVCreateInfo.MipLevel);
+			break;
+
+		default:
+			checkf(false, TEXT("Invalid texture metadata access used when creating UAV."));
+			return nullptr;
+		}
+	}
+
+	check(RHIUnorderedAccessView);
+	FRHIUnorderedAccessView* View = RHIUnorderedAccessView.GetReference();
+	UAVs.Emplace(UAVCreateInfo, MoveTemp(RHIUnorderedAccessView));
+	return View;
+}
+
+FRHIShaderResourceView* FRHIBufferViewCache::GetOrCreateSRV(FRHIBuffer* Buffer, const FRHIBufferSRVCreateInfo& SRVCreateInfo)
+{
+	for (const auto& KeyValue : SRVs)
+	{
+		if (KeyValue.Key == SRVCreateInfo)
+		{
+			return KeyValue.Value.GetReference();
+		}
+	}
+
+	FShaderResourceViewRHIRef RHIShaderResourceView;
+
+	if (SRVCreateInfo.Format != PF_Unknown)
+	{
+		RHIShaderResourceView = RHICreateShaderResourceView(Buffer, SRVCreateInfo.BytesPerElement, SRVCreateInfo.Format);
+	}
+	else
+	{
+		RHIShaderResourceView = RHICreateShaderResourceView(Buffer);
+	}
+
+	FRHIShaderResourceView* View = RHIShaderResourceView.GetReference();
+	SRVs.Emplace(SRVCreateInfo, MoveTemp(RHIShaderResourceView));
+	return View;
+}
+
+FRHIUnorderedAccessView* FRHIBufferViewCache::GetOrCreateUAV(FRHIBuffer* Buffer, const FRHIBufferUAVCreateInfo& UAVCreateInfo)
+{
+	for (const auto& KeyValue : UAVs)
+	{
+		if (KeyValue.Key == UAVCreateInfo)
+		{
+			return KeyValue.Value.GetReference();
+		}
+	}
+
+	FUnorderedAccessViewRHIRef RHIUnorderedAccessView;
+
+	if (UAVCreateInfo.Format != PF_Unknown)
+	{
+		RHIUnorderedAccessView = RHICreateUnorderedAccessView(Buffer, UAVCreateInfo.Format);
+	}
+	else
+	{
+		RHIUnorderedAccessView = RHICreateUnorderedAccessView(Buffer, UAVCreateInfo.bSupportsAtomicCounter, UAVCreateInfo.bSupportsAppendBuffer);
+	}
+
+	FRHIUnorderedAccessView* View = RHIUnorderedAccessView.GetReference();
+	UAVs.Emplace(UAVCreateInfo, MoveTemp(RHIUnorderedAccessView));
+	return View;
+}
+
+#if !(UE_BUILD_SHIPPING || UE_BUILD_TEST)
+
+void FRHITextureViewCache::SetDebugName(const TCHAR* DebugName)
+{
+	for (const auto& KeyValue : UAVs)
+	{
+		RHIBindDebugLabelName(KeyValue.Value, DebugName);
+	}
+}
+
+void FRHIBufferViewCache::SetDebugName(const TCHAR* DebugName)
+{
+	for (const auto& KeyValue : UAVs)
+	{
+		RHIBindDebugLabelName(KeyValue.Value, DebugName);
+	}
+}
+
+#endif
+
+void FRHITransientTexture::Init(const TCHAR* InName, uint32 InAllocationIndex)
+{
+	FRHITransientResource::Init(InName, InAllocationIndex);
+	ViewCache.SetDebugName(InName);
+
+#if !(UE_BUILD_SHIPPING || UE_BUILD_TEST)
+	RHIBindDebugLabelName(GetRHI(), InName);
+#endif
+}
+
+void FRHITransientBuffer::Init(const TCHAR* InName, uint32 InAllocationIndex)
+{
+	FRHITransientResource::Init(InName, InAllocationIndex);
+	ViewCache.SetDebugName(InName);
+
+	// TODO: Add method to rename a buffer.
+}
\ No newline at end of file
diff --git a/Engine/Source/Runtime/RHI/Private/RHITransientResourceAllocator.cpp b/Engine/Source/Runtime/RHI/Private/RHITransientResourceAllocator.cpp
deleted file mode 100644
diff --git a/Engine/Source/Runtime/RHI/Public/RHIResources.h b/Engine/Source/Runtime/RHI/Public/RHIResources.h
--- a/Engine/Source/Runtime/RHI/Public/RHIResources.h
+++ b/Engine/Source/Runtime/RHI/Public/RHIResources.h
@@ -3135,14 +3135,20 @@ FORCEINLINE uint32 GetTypeHash(const FRHITextureSRVCreateInfo& Var)
 struct FRHITextureUAVCreateInfo
 {
 public:
-	explicit FRHITextureUAVCreateInfo(uint8 InMipLevel = 0u, ERHITextureMetaDataAccess InMetaData = ERHITextureMetaDataAccess::None)
-		: MipLevel(InMipLevel)
-		, MetaData(InMetaData)
+	FRHITextureUAVCreateInfo() = default;
+
+	explicit FRHITextureUAVCreateInfo(uint8 InMipLevel, EPixelFormat InFormat = PF_Unknown)
+		: Format(InFormat)
+		, MipLevel(InMipLevel)
+	{}
+
+	explicit FRHITextureUAVCreateInfo(ERHITextureMetaDataAccess InMetaData)
+		: MetaData(InMetaData)
 	{}
 
 	FORCEINLINE bool operator==(const FRHITextureUAVCreateInfo& Other)const
 	{
-		return MipLevel == Other.MipLevel && MetaData == Other.MetaData;
+		return Format == Other.Format && MipLevel == Other.MipLevel && MetaData == Other.MetaData;
 	}
 
 	FORCEINLINE bool operator!=(const FRHITextureUAVCreateInfo& Other)const
@@ -3150,8 +3156,9 @@ public:
 		return !(*this == Other);
 	}
 
-	uint8 MipLevel;
-	ERHITextureMetaDataAccess MetaData;
+	EPixelFormat Format = PF_Unknown;
+	uint8 MipLevel = 0;
+	ERHITextureMetaDataAccess MetaData = ERHITextureMetaDataAccess::None;
 };
 
 /** Descriptor used to create a buffer resource */
@@ -3231,4 +3238,46 @@ struct FRHIBufferUAVCreateInfo
 	/** Whether the uav supports atomic counter or append buffer operations (used for structured buffers) */
 	bool bSupportsAtomicCounter = false;
 	bool bSupportsAppendBuffer = false;
+};
+
+class RHI_API FRHITextureViewCache
+{
+public:
+	// Finds a UAV matching the descriptor in the cache or creates a new one and updates the cache.
+	FRHIUnorderedAccessView* GetOrCreateUAV(FRHITexture* Texture, const FRHITextureUAVCreateInfo& CreateInfo);
+
+	// Finds a SRV matching the descriptor in the cache or creates a new one and updates the cache.
+	FRHIShaderResourceView* GetOrCreateSRV(FRHITexture* Texture, const FRHITextureSRVCreateInfo& CreateInfo);
+
+	// Sets the debug name of the RHI view resources.
+#if !(UE_BUILD_SHIPPING || UE_BUILD_TEST)
+	void SetDebugName(const TCHAR* DebugName);
+#else
+	void SetDebugName(const TCHAR* DebugName) {}
+#endif
+
+private:
+	TArray<TPair<FRHITextureUAVCreateInfo, FUnorderedAccessViewRHIRef>, TInlineAllocator<1>> UAVs;
+	TArray<TPair<FRHITextureSRVCreateInfo, FShaderResourceViewRHIRef>, TInlineAllocator<1>> SRVs;
+};
+
+class RHI_API FRHIBufferViewCache
+{
+public:
+	// Finds a UAV matching the descriptor in the cache or creates a new one and updates the cache.
+	FRHIUnorderedAccessView* GetOrCreateUAV(FRHIBuffer* Buffer, const FRHIBufferUAVCreateInfo& CreateInfo);
+
+	// Finds a SRV matching the descriptor in the cache or creates a new one and updates the cache.
+	FRHIShaderResourceView* GetOrCreateSRV(FRHIBuffer* Buffer, const FRHIBufferSRVCreateInfo& CreateInfo);
+
+	// Sets the debug name of the RHI view resources.
+#if !(UE_BUILD_SHIPPING || UE_BUILD_TEST)
+	void SetDebugName(const TCHAR* DebugName);
+#else
+	void SetDebugName(const TCHAR* DebugName) {}
+#endif
+
+private:
+	TArray<TPair<FRHIBufferUAVCreateInfo, FUnorderedAccessViewRHIRef>, TInlineAllocator<1>> UAVs;
+	TArray<TPair<FRHIBufferSRVCreateInfo, FShaderResourceViewRHIRef>, TInlineAllocator<1>> SRVs;
 };
\ No newline at end of file
diff --git a/Engine/Source/Runtime/RHI/Public/RHITransientResourceAllocator.h b/Engine/Source/Runtime/RHI/Public/RHITransientResourceAllocator.h
--- a/Engine/Source/Runtime/RHI/Public/RHITransientResourceAllocator.h
+++ b/Engine/Source/Runtime/RHI/Public/RHITransientResourceAllocator.h
@@ -25,7 +25,7 @@ public:
 	virtual ~FRHITransientResource() = default;
 
 	// (Internal) Initializes the transient resource with a new allocation / name.
-	void Init(const TCHAR* InName, uint32 InAllocationIndex)
+	virtual void Init(const TCHAR* InName, uint32 InAllocationIndex)
 	{
 		Name = InName;
 		AllocationIndex = InAllocationIndex;
@@ -85,6 +85,8 @@ public:
 		, CreateInfo(InCreateInfo)
 	{}
 
+	void Init(const TCHAR* InName, uint32 InAllocationIndex) override;
+
 	// Returns the underlying RHI texture.
 	FRHITexture* GetRHI() const { return static_cast<FRHITexture*>(FRHITransientResource::GetRHI()); }
 
@@ -92,15 +94,16 @@ public:
 	FORCEINLINE const FRHITextureCreateInfo& GetCreateInfo() const { return CreateInfo; }
 
 	// Finds a UAV matching the descriptor in the cache or creates a new one and updates the cache.
-	FRHIUnorderedAccessView* GetOrCreateUAV(const FRHITextureUAVCreateInfo& CreateInfo);
+	FORCEINLINE FRHIUnorderedAccessView* GetOrCreateUAV(const FRHITextureUAVCreateInfo& InCreateInfo) { return ViewCache.GetOrCreateUAV(GetRHI(), InCreateInfo); }
 
 	// Finds a SRV matching the descriptor in the cache or creates a new one and updates the cache.
-	FRHIShaderResourceView* GetOrCreateSRV(const FRHITextureSRVCreateInfo& CreateInfo);
+	FORCEINLINE FRHIShaderResourceView* GetOrCreateSRV(const FRHITextureSRVCreateInfo& InCreateInfo) { return ViewCache.GetOrCreateSRV(GetRHI(), InCreateInfo); }
 
-private:
-	FRHITextureCreateInfo CreateInfo;
-	TArray<TPair<FRHITextureUAVCreateInfo, FUnorderedAccessViewRHIRef>, TInlineAllocator<1>> UAVs;
-	TArray<TPair<FRHITextureSRVCreateInfo, FShaderResourceViewRHIRef>, TInlineAllocator<1>> SRVs;
+	// The create info describing the texture.
+	const FRHITextureCreateInfo CreateInfo;
+
+	// The persistent view cache containing all views created for this texture.
+	FRHITextureViewCache ViewCache;
 };
 
 class RHI_API FRHITransientBuffer final : public FRHITransientResource
@@ -111,6 +114,8 @@ public:
 		, CreateInfo(InCreateInfo)
 	{}
 
+	void Init(const TCHAR* InName, uint32 InAllocationIndex) override;
+
 	// Returns the underlying RHI buffer.
 	FORCEINLINE FRHIBuffer* GetRHI() const { return static_cast<FRHIBuffer*>(FRHITransientResource::GetRHI()); }
 
@@ -118,18 +123,19 @@ public:
 	FORCEINLINE const FRHIBufferCreateInfo& GetCreateInfo() const { return CreateInfo; }
 
 	// Finds a UAV matching the descriptor in the cache or creates a new one and updates the cache.
-	FRHIUnorderedAccessView* GetOrCreateUAV(const FRHIBufferUAVCreateInfo& CreateInfo);
+	FORCEINLINE FRHIUnorderedAccessView* GetOrCreateUAV(const FRHIBufferUAVCreateInfo& InCreateInfo) { return ViewCache.GetOrCreateUAV(GetRHI(), InCreateInfo); }
 
 	// Finds a SRV matching the descriptor in the cache or creates a new one and updates the cache.
-	FRHIShaderResourceView* GetOrCreateSRV(const FRHIBufferSRVCreateInfo& CreateInfo);
+	FORCEINLINE FRHIShaderResourceView* GetOrCreateSRV(const FRHIBufferSRVCreateInfo& InCreateInfo) { return ViewCache.GetOrCreateSRV(GetRHI(), InCreateInfo); }
 
-private:
-	FRHIBufferCreateInfo CreateInfo;
-	TArray<TPair<FRHIBufferUAVCreateInfo, FUnorderedAccessViewRHIRef>, TInlineAllocator<1>> UAVs;
-	TArray<TPair<FRHIBufferSRVCreateInfo, FShaderResourceViewRHIRef>, TInlineAllocator<1>> SRVs;
+	// The create info describing the texture.
+	const FRHIBufferCreateInfo CreateInfo;
+
+	// The persistent view cache containing all views created for this buffer.
+	FRHIBufferViewCache ViewCache;
 };
 
-class RHI_API IRHITransientResourceAllocator
+class IRHITransientResourceAllocator
 {
 public:
 	virtual ~IRHITransientResourceAllocator() = default;
@@ -146,5 +152,8 @@ public:
 	virtual void Freeze(FRHICommandListImmediate& RHICmdList) = 0;
 
 	// Releases the transient allocator and deletes the instance. Any FRHITransientResource* access after this call is not allowed.
-	virtual void Release(FRHICommandListImmediate& RHICmdList);
+	virtual void Release(FRHICommandListImmediate& RHICmdList)
+	{
+		delete this;
+	}
 };
\ No newline at end of file
diff --git a/Engine/Source/Runtime/RenderCore/Private/RenderGraphBuilder.cpp b/Engine/Source/Runtime/RenderCore/Private/RenderGraphBuilder.cpp
--- a/Engine/Source/Runtime/RenderCore/Private/RenderGraphBuilder.cpp
+++ b/Engine/Source/Runtime/RenderCore/Private/RenderGraphBuilder.cpp
@@ -2514,7 +2514,7 @@ void FRDGBuilder::BeginResourceRHI(FTransientResourceAllocator* TransientResourc
 
 				AddToPrologueBarriers(PassHandle, [&](FRDGBarrierBatchBegin& Barriers)
 				{
-					Barriers.AddAlias(Texture, FRHITransientAliasingInfo::Acquire(Texture->GetRHIUnchecked(), Texture->TransientTexture->GetAliasingOverlaps()));
+					Barriers.AddAlias(Texture, FRHITransientAliasingInfo::Acquire(Texture->GetRHIUnchecked(), TransientTexture->GetAliasingOverlaps()));
 				});
 
 			#if STATS
@@ -2531,7 +2531,7 @@ void FRDGBuilder::BeginResourceRHI(FTransientResourceAllocator* TransientResourc
 
 	Texture->FirstPass = PassHandle;
 
-	check(Texture->PooledTexture);
+	check(Texture->GetRHIUnchecked());
 }
 
 void FRDGBuilder::BeginResourceRHI(FTransientResourceAllocator*, FRDGPassHandle PassHandle, FRDGTextureSRVRef SRV)
@@ -2544,16 +2544,10 @@ void FRDGBuilder::BeginResourceRHI(FTransientResourceAllocator*, FRDGPassHandle
 	}
 
 	FRDGTextureRef Texture = SRV->Desc.Texture;
-	check(Texture->GetRHIUnchecked());
+	FRHITexture* TextureRHI = Texture->GetRHIUnchecked();
+	check(TextureRHI);
 
-	if (Texture->bTransient)
-	{
-		SRV->ResourceRHI = Texture->TransientTexture->GetOrCreateSRV(SRV->Desc);
-	}
-	else
-	{
-		SRV->ResourceRHI = Texture->PooledTexture->GetOrCreateSRV(SRV->Desc);
-	}
+	SRV->ResourceRHI = Texture->ViewCache->GetOrCreateSRV(TextureRHI, SRV->Desc);
 }
 
 void FRDGBuilder::BeginResourceRHI(FTransientResourceAllocator* TransientResourceAllocator, FRDGPassHandle PassHandle, FRDGTextureUAVRef UAV)
@@ -2568,15 +2562,10 @@ void FRDGBuilder::BeginResourceRHI(FTransientResourceAllocator* TransientResourc
 	BeginResourceRHI(TransientResourceAllocator, PassHandle, UAV->Desc.Texture);
 
 	FRDGTextureRef Texture = UAV->Desc.Texture;
+	FRHITexture* TextureRHI = Texture->GetRHIUnchecked();
+	check(TextureRHI);
 
-	if (Texture->bTransient)
-	{
-		UAV->ResourceRHI = Texture->TransientTexture->GetOrCreateUAV(UAV->Desc);
-	}
-	else
-	{
-		UAV->ResourceRHI = Texture->PooledTexture->GetOrCreateUAV(UAV->Desc);
-	}
+	UAV->ResourceRHI = Texture->ViewCache->GetOrCreateUAV(TextureRHI, UAV->Desc);
 }
 
 void FRDGBuilder::BeginResourceRHI(FTransientResourceAllocator* TransientResourceAllocator, FRDGPassHandle PassHandle, FRDGBufferRef Buffer)
@@ -2610,7 +2599,7 @@ void FRDGBuilder::BeginResourceRHI(FTransientResourceAllocator* TransientResourc
 
 				AddToPrologueBarriers(PassHandle, [&](FRDGBarrierBatchBegin& Barriers)
 				{
-					Barriers.AddAlias(Buffer, FRHITransientAliasingInfo::Acquire(Buffer->GetRHIUnchecked(), Buffer->TransientBuffer->GetAliasingOverlaps()));
+					Barriers.AddAlias(Buffer, FRHITransientAliasingInfo::Acquire(Buffer->GetRHIUnchecked(), TransientBuffer->GetAliasingOverlaps()));
 				});
 
 			#if STATS
@@ -2640,7 +2629,8 @@ void FRDGBuilder::BeginResourceRHI(FTransientResourceAllocator*, FRDGPassHandle
 	}
 
 	FRDGBufferRef Buffer = SRV->Desc.Buffer;
-	checkf(Buffer->GetRHIUnchecked(), TEXT("Pass parameters contained an SRV of RDG buffer %s, before that buffer was referenced as a UAV, which isn't supported.  Make sure to call ClearUnusedGraphResources to remove unbound parameters which can also cause this."), Buffer->Name);
+	FRHIBuffer* BufferRHI = Buffer->GetRHIUnchecked();
+	check(BufferRHI);
 
 	FRHIBufferSRVCreateInfo SRVCreateInfo = SRV->Desc;
 
@@ -2652,14 +2642,7 @@ void FRDGBuilder::BeginResourceRHI(FTransientResourceAllocator*, FRDGPassHandle
 		SRVCreateInfo.Format = PF_Unknown;
 	}
 
-	if (Buffer->bTransient)
-	{
-		SRV->ResourceRHI = Buffer->TransientBuffer->GetOrCreateSRV(SRVCreateInfo);
-	}
-	else
-	{
-		SRV->ResourceRHI = Buffer->PooledBuffer->GetOrCreateSRV(SRVCreateInfo);
-	}
+	SRV->ResourceRHI = Buffer->ViewCache->GetOrCreateSRV(BufferRHI, SRVCreateInfo);
 }
 
 void FRDGBuilder::BeginResourceRHI(FTransientResourceAllocator* TransientResourceAllocator, FRDGPassHandle PassHandle, FRDGBufferUAV* UAV)
@@ -2684,14 +2667,7 @@ void FRDGBuilder::BeginResourceRHI(FTransientResourceAllocator* TransientResourc
 		UAVCreateInfo.Format = PF_Unknown;
 	}
 
-	if (Buffer->bTransient)
-	{
-		UAV->ResourceRHI = Buffer->TransientBuffer->GetOrCreateUAV(UAVCreateInfo);
-	}
-	else
-	{
-		UAV->ResourceRHI = Buffer->PooledBuffer->GetOrCreateUAV(UAVCreateInfo);
-	}
+	UAV->ResourceRHI = Buffer->ViewCache->GetOrCreateUAV(Buffer->GetRHIUnchecked(), UAVCreateInfo);
 }
 
 void FRDGBuilder::EndResourceRHI(FTransientResourceAllocator& TransientResourceAllocator, FRDGPassHandle PassHandle, FRDGTextureRef Texture, uint32 ReferenceCount)
diff --git a/Engine/Source/Runtime/RenderCore/Private/RenderGraphResourcePool.cpp b/Engine/Source/Runtime/RenderCore/Private/RenderGraphResourcePool.cpp
--- a/Engine/Source/Runtime/RenderCore/Private/RenderGraphResourcePool.cpp
+++ b/Engine/Source/Runtime/RenderCore/Private/RenderGraphResourcePool.cpp
@@ -38,13 +38,8 @@ TRefCountPtr<FRDGPooledBuffer> FRenderGraphResourcePool::FindFreeBufferInternal(
 		if (PooledBuffer->Desc == Desc)
 		{
 			PooledBuffer->LastUsedFrame = FrameCounter;
+			PooledBuffer->ViewCache.SetDebugName(InDebugName);
 			PooledBuffer->Name = InDebugName;
-			
-			//@todo - rename other resources too
-			for (const auto& Pair : PooledBuffer->UAVs)
-			{
-				RHIBindDebugLabelName(Pair.Value, InDebugName);
-			}
 
 			return PooledBuffer;
 		}
diff --git a/Engine/Source/Runtime/RenderCore/Private/RenderGraphResources.cpp b/Engine/Source/Runtime/RenderCore/Private/RenderGraphResources.cpp
--- a/Engine/Source/Runtime/RenderCore/Private/RenderGraphResources.cpp
+++ b/Engine/Source/Runtime/RenderCore/Private/RenderGraphResources.cpp
@@ -164,35 +164,6 @@ bool FRDGSubresourceState::IsTransitionRequired(const FRDGSubresourceState& Prev
 	return false;
 }
 
-void FRDGPooledTexture::InitViews(const FUnorderedAccessViewRHIRef& FirstMipUAV)
-{
-	const ETextureCreateFlags Flags = Texture->GetFlags();
-	const uint32 NumMips = Texture->GetNumMips();
-
-	if (EnumHasAnyFlags(Flags, TexCreate_ShaderResource))
-	{
-		SRVs.Empty();
-	}
-
-	if (EnumHasAnyFlags(Flags, TexCreate_UAV))
-	{
-		MipUAVs.Empty(NumMips);
-
-		uint32 MipLevel = 0;
-
-		if (FirstMipUAV)
-		{
-			MipUAVs.Add(FirstMipUAV);
-			MipLevel++;
-		}
-
-		for (; MipLevel < NumMips; MipLevel++)
-		{
-			MipUAVs.Add(RHICreateUnorderedAccessView(Texture, MipLevel));
-		}
-	}
-}
-
 void FRDGPooledTexture::Finalize()
 {
 	for (FRDGSubresourceState& SubresourceState : State)
@@ -208,75 +179,6 @@ void FRDGPooledTexture::Reset()
 	Owner = nullptr;
 }
 
-FRHIShaderResourceView* FRDGPooledTexture::GetOrCreateSRV(const FRHITextureSRVCreateInfo& Desc)
-{
-	if (Desc.MetaData == ERDGTextureMetaDataAccess::HTile)
-	{
-		check(GRHISupportsExplicitHTile);
-		if (!HTileSRV)
-		{
-			HTileSRV = RHICreateShaderResourceViewHTile((FRHITexture2D*)GetRHI());
-		}
-		return HTileSRV;
-	}
-
-	if (Desc.MetaData == ERDGTextureMetaDataAccess::FMask)
-	{
-		if (!FMaskSRV)
-		{
-			FMaskSRV = RHICreateShaderResourceViewFMask((FRHITexture2D*)GetRHI());
-		}
-		return FMaskSRV;
-	}
-
-	if (Desc.MetaData == ERDGTextureMetaDataAccess::CMask)
-	{
-		if (!CMaskSRV)
-		{
-			CMaskSRV = RHICreateShaderResourceViewWriteMask((FRHITexture2D*)GetRHI());
-		}
-		return CMaskSRV;
-	}
-
-	for (const auto& SRVPair : SRVs)
-	{
-		if (SRVPair.Key == Desc)
-		{
-			return SRVPair.Value;
-		}
-	}
-
-	FShaderResourceViewRHIRef RHIShaderResourceView = RHICreateShaderResourceView(GetRHI(), Desc);
-
-	FRHIShaderResourceView* View = RHIShaderResourceView;
-	SRVs.Emplace(Desc, MoveTemp(RHIShaderResourceView));
-	return View;
-}
-
-FRHIUnorderedAccessView* FRDGPooledTexture::GetOrCreateUAV(const FRHITextureUAVCreateInfo& UAVDesc)
-{
-	if (UAVDesc.MetaData == ERDGTextureMetaDataAccess::HTile)
-	{
-		check(GRHISupportsExplicitHTile);
-		if (!HTileUAV)
-		{
-			HTileUAV = RHICreateUnorderedAccessViewHTile((FRHITexture2D*)GetRHI());
-		}
-		return HTileUAV;
-	}
-
-	if (UAVDesc.MetaData == ERDGTextureMetaDataAccess::Stencil)
-	{
-		if (!StencilUAV)
-		{
-			StencilUAV = RHICreateUnorderedAccessViewStencil((FRHITexture2D*)GetRHI(), 0);
-		}
-		return StencilUAV;
-	}
-
-	return MipUAVs[UAVDesc.MipLevel];
-}
-
 FRDGTextureSubresourceRange FRDGTexture::GetSubresourceRangeSRV() const
 {
 	FRDGTextureSubresourceRange Range = GetSubresourceRange();
@@ -306,6 +208,7 @@ void FRDGTexture::SetRHI(FRDGPooledTexture* InPooledTexture)
 {
 	PooledTexture = InPooledTexture;
 	State = &PooledTexture->State;
+	ViewCache = &PooledTexture->ViewCache;
 
 	// Return the previous owner and assign this texture as the new one.
 	FRDGTextureRef PreviousOwner = PooledTexture->Owner;
@@ -325,7 +228,8 @@ void FRDGTexture::SetRHI(FRHITransientTexture* InTransientTexture, FRDGAllocator
 {
 	TransientTexture = InTransientTexture;
 	State = Allocator.AllocNoDestruct<FRDGTextureSubresourceState>();
-	ResourceRHI = TransientTexture->GetRHI();
+	ViewCache = &InTransientTexture->ViewCache;
+	ResourceRHI = InTransientTexture->GetRHI();
 
 	bTransient = true;
 }
@@ -380,6 +284,7 @@ void FRDGBuffer::SetRHI(FRDGPooledBuffer* InPooledBuffer)
 	PooledBuffer = InPooledBuffer;
 	Allocation = InPooledBuffer;
 	State = &PooledBuffer->State;
+	ViewCache = &PooledBuffer->ViewCache;
 	ResourceRHI = InPooledBuffer->GetRHI();
 }
 
@@ -387,7 +292,8 @@ void FRDGBuffer::SetRHI(FRHITransientBuffer* InTransientBuffer, FRDGAllocator& A
 {
 	TransientBuffer = InTransientBuffer;
 	State = Allocator.AllocNoDestruct<FRDGSubresourceState>();
-	ResourceRHI = TransientBuffer->GetRHI();
+	ViewCache = &InTransientBuffer->ViewCache;
+	ResourceRHI = InTransientBuffer->GetRHI();
 
 	bTransient = true;;
 }
@@ -419,56 +325,4 @@ void FRDGBuffer::Finalize()
 			Allocation = PooledBuffer;
 		}
 	}
-}
-
-FRHIShaderResourceView* FRDGPooledBuffer::GetOrCreateSRV(const FRHIBufferSRVCreateInfo& SRVDesc)
-{
-	for (const auto& KeyValue : SRVs)
-	{
-		if (KeyValue.Key == SRVDesc)
-		{
-			return KeyValue.Value.GetReference();
-		}
-	}
-
-	FShaderResourceViewRHIRef RHIShaderResourceView;
-
-	if (SRVDesc.Format != PF_Unknown)
-	{
-		RHIShaderResourceView = RHICreateShaderResourceView(GetRHI(), SRVDesc.BytesPerElement, SRVDesc.Format);
-	}
-	else
-	{
-		RHIShaderResourceView = RHICreateShaderResourceView(GetRHI());
-	}
-
-	FRHIShaderResourceView* View = RHIShaderResourceView.GetReference();
-	SRVs.Emplace(SRVDesc, MoveTemp(RHIShaderResourceView));
-	return View;
-}
-
-FRHIUnorderedAccessView* FRDGPooledBuffer::GetOrCreateUAV(const FRHIBufferUAVCreateInfo& UAVDesc)
-{
-	for (const auto& KeyValue : UAVs)
-	{
-		if (KeyValue.Key == UAVDesc)
-		{
-			return KeyValue.Value.GetReference();
-		}
-	}
-
-	FUnorderedAccessViewRHIRef RHIUnorderedAccessView;
-
-	if (UAVDesc.Format != PF_Unknown)
-	{
-		RHIUnorderedAccessView = RHICreateUnorderedAccessView(GetRHI(), UAVDesc.Format);
-	}
-	else
-	{
-		RHIUnorderedAccessView = RHICreateUnorderedAccessView(GetRHI(), UAVDesc.bSupportsAtomicCounter, UAVDesc.bSupportsAppendBuffer);
-	}
-
-	FRHIUnorderedAccessView* View = RHIUnorderedAccessView.GetReference();
-	UAVs.Emplace(UAVDesc, MoveTemp(RHIUnorderedAccessView));
-	return View;
 }
\ No newline at end of file
diff --git a/Engine/Source/Runtime/RenderCore/Private/RenderTargetPool.cpp b/Engine/Source/Runtime/RenderCore/Private/RenderTargetPool.cpp
--- a/Engine/Source/Runtime/RenderCore/Private/RenderTargetPool.cpp
+++ b/Engine/Source/Runtime/RenderCore/Private/RenderTargetPool.cpp
@@ -1006,29 +1006,6 @@ void FRenderTargetPool::TickPoolElements()
 		{
 			AllocationLevelInKB -= ComputeSizeInKB(*PooledRenderTargets[OldestElementIndex]);
 
-			// NOTE: Comment out because RHI resource shouldn't be deleted on the RenderThread
-			//       But command to request delete should be added to fix this
-			/*
-			// Don't defer delete because we know it hasn't been used for at least 2 frames 
-			// and defer deletion queue might not get immediatly flushed causing VRAM
-			// memory spikes
-			if (TRefCountPtr<FPooledRenderTarget>& PooledTarget = PooledRenderTargets[OldestElementIndex])
-			{
-				FRHITexture* TargetableRHI = PooledTarget->GetTargetableRHI();
-				FRHITexture* ShaderResourceRHI = PooledTarget->GetShaderResourceRHI();
-
-				if (TargetableRHI)
-				{
-					TargetableRHI->DoNoDeferDelete();
-				}
-
-				if (ShaderResourceRHI && ShaderResourceRHI != TargetableRHI)
-				{
-					ShaderResourceRHI->DoNoDeferDelete();
-				}
-			}
-			*/
-
 			// we assume because of reference counting the resource gets released when not needed any more
 			// we don't use Remove() to not shuffle around the elements for better transparency on RenderTargetPoolEvents
 			FreeElementAtIndex(OldestElementIndex);
@@ -1224,12 +1201,12 @@ void FPooledRenderTarget::InitRDG()
 
 	if (RenderTargetItem.TargetableTexture)
 	{
-		TargetableTexture = new FRDGPooledTexture(RenderTargetItem.TargetableTexture, Translate(Desc, ERenderTargetTexture::Targetable), RenderTargetItem.UAV);
+		TargetableTexture = new FRDGPooledTexture(RenderTargetItem.TargetableTexture, Translate(Desc, ERenderTargetTexture::Targetable));
 	}
 
 	if (RenderTargetItem.ShaderResourceTexture != RenderTargetItem.TargetableTexture)
 	{
-		ShaderResourceTexture = new FRDGPooledTexture(RenderTargetItem.ShaderResourceTexture, Translate(Desc, ERenderTargetTexture::ShaderResource), nullptr);
+		ShaderResourceTexture = new FRDGPooledTexture(RenderTargetItem.ShaderResourceTexture, Translate(Desc, ERenderTargetTexture::ShaderResource));
 	}
 	else
 	{
diff --git a/Engine/Source/Runtime/RenderCore/Public/RenderGraphResources.h b/Engine/Source/Runtime/RenderCore/Public/RenderGraphResources.h
--- a/Engine/Source/Runtime/RenderCore/Public/RenderGraphResources.h
+++ b/Engine/Source/Runtime/RenderCore/Public/RenderGraphResources.h
@@ -343,34 +343,24 @@ class RENDERCORE_API FRDGPooledTexture final
 	: public FRefCountedObject
 {
 public:
-	FRDGPooledTexture(FRHITexture* InTexture, const FRDGTextureSubresourceLayout& InLayout, const FUnorderedAccessViewRHIRef& FirstMipUAV)
+	FRDGPooledTexture(FRHITexture* InTexture, const FRDGTextureSubresourceLayout& InLayout)
 		: Texture(InTexture)
 		, Layout(InLayout)
 	{
-		InitViews(FirstMipUAV);
 		Reset();
 	}
 
 	/** Finds a UAV matching the descriptor in the cache or creates a new one and updates the cache. */
-	FRHIUnorderedAccessView* GetOrCreateUAV(const FRHITextureUAVCreateInfo& UAVDesc);
+	FORCEINLINE FRHIUnorderedAccessView* GetOrCreateUAV(const FRHITextureUAVCreateInfo& UAVDesc) { return ViewCache.GetOrCreateUAV(Texture, UAVDesc); }
 
 	/** Finds a SRV matching the descriptor in the cache or creates a new one and updates the cache. */
-	FRHIShaderResourceView* GetOrCreateSRV(const FRHITextureSRVCreateInfo& SRVDesc);
+	FORCEINLINE FRHIShaderResourceView* GetOrCreateSRV(const FRHITextureSRVCreateInfo& SRVDesc) { return ViewCache.GetOrCreateSRV(Texture, SRVDesc); }
 
-	FRHITexture* GetRHI() const
-	{
-		return Texture;
-	}
+	FORCEINLINE FRHITexture* GetRHI() const { return Texture; }
 
-	FRDGTexture* GetOwner() const
-	{
-		return Owner;
-	}
+	FORCEINLINE FRDGTexture* GetOwner() const { return Owner; }
 
 private:
-	/** Initializes cached views. Safe to call multiple times; each call will recreate. */
-	void InitViews(const FUnorderedAccessViewRHIRef& FirstMipUAV);
-
 	/** Prepares the pooled texture state for re-use across RDG builder instances. */
 	void Finalize();
 
@@ -381,15 +371,7 @@ private:
 	FRDGTexture* Owner = nullptr;
 	FRDGTextureSubresourceLayout Layout;
 	FRDGTextureSubresourceState State;
-
-	/** Cached views created for the RHI texture. */
-	TArray<FUnorderedAccessViewRHIRef, TInlineAllocator<1>> MipUAVs;
-	TArray<TPair<FRHITextureSRVCreateInfo, FShaderResourceViewRHIRef>, TInlineAllocator<1>> SRVs;
-	FUnorderedAccessViewRHIRef HTileUAV;
-	FShaderResourceViewRHIRef  HTileSRV;
-	FUnorderedAccessViewRHIRef StencilUAV;
-	FShaderResourceViewRHIRef  FMaskSRV;
-	FShaderResourceViewRHIRef  CMaskSRV;
+	FRHITextureViewCache ViewCache;
 
 	friend FRDGTexture;
 	friend FRDGBuilder;
@@ -508,6 +490,9 @@ private:
 		FRHITransientTexture* TransientTexture;
 	};
 
+	/** The assigned view cache for this texture (sourced from transient / pooled texture). Never reset. */
+	FRHITextureViewCache* ViewCache = nullptr;
+
 	/** Cached state pointer from the pooled texture. */
 	FRDGTextureSubresourceState* State = nullptr;
 
@@ -661,8 +646,8 @@ class FRDGTextureUAVDesc final
 public:
 	FRDGTextureUAVDesc() = default;
 
-	FRDGTextureUAVDesc(FRDGTextureRef InTexture, uint8 InMipLevel = 0, ERHITextureMetaDataAccess InMetaData = ERHITextureMetaDataAccess::None)
-		: FRHITextureUAVCreateInfo(InMipLevel, InMetaData)
+	FRDGTextureUAVDesc(FRDGTextureRef InTexture, uint8 InMipLevel = 0, EPixelFormat InFormat = PF_Unknown)
+		: FRHITextureUAVCreateInfo(InMipLevel, InFormat)
 		, Texture(InTexture)
 	{}
 
@@ -862,33 +847,23 @@ public:
 	const FRDGBufferDesc Desc;
 
 	/** Finds a UAV matching the descriptor in the cache or creates a new one and updates the cache. */
-	FRHIUnorderedAccessView* GetOrCreateUAV(const FRHIBufferUAVCreateInfo& UAVDesc);
+	FORCEINLINE FRHIUnorderedAccessView* GetOrCreateUAV(const FRHIBufferUAVCreateInfo& UAVDesc) { return ViewCache.GetOrCreateUAV(Buffer, UAVDesc); }
 
 	/** Finds a SRV matching the descriptor in the cache or creates a new one and updates the cache. */
-	FRHIShaderResourceView* GetOrCreateSRV(const FRHIBufferSRVCreateInfo& SRVDesc);
+	FORCEINLINE FRHIShaderResourceView* GetOrCreateSRV(const FRHIBufferSRVCreateInfo& SRVDesc) { return ViewCache.GetOrCreateSRV(Buffer, SRVDesc); }
 
 	/** Returns the RHI buffer. */
-	FRHIBuffer* GetRHI() const
-	{
-		return Buffer;
-	}
+	FORCEINLINE FRHIBuffer* GetRHI() const { return Buffer; }
 
 	UE_DEPRECATED(5.0, "Buffers types have been consolidated; use GetRHI() instead.")
-	FRHIBuffer* GetVertexBufferRHI() const
-	{
-		return Buffer;
-	}
+	FORCEINLINE FRHIBuffer* GetVertexBufferRHI() const { return Buffer; }
 
 	UE_DEPRECATED(5.0, "Buffers types have been consolidated; use GetRHI() instead.")
-	FRHIBuffer* GetStructuredBufferRHI() const
-	{
-		return Buffer;
-	}
+	FORCEINLINE FRHIBuffer* GetStructuredBufferRHI() const { return Buffer; }
 
 private:
 	TRefCountPtr<FRHIBuffer> Buffer;
-	TArray<TPair<FRHIBufferUAVCreateInfo, FUnorderedAccessViewRHIRef>, TInlineAllocator<1>> UAVs;
-	TArray<TPair<FRHIBufferSRVCreateInfo, FShaderResourceViewRHIRef>, TInlineAllocator<1>> SRVs;
+	FRHIBufferViewCache ViewCache;
 
 	void Reset()
 	{
@@ -1005,6 +980,9 @@ private:
 		FRHITransientBuffer* TransientBuffer;
 	};
 
+	/** The assigned buffer view cache (sourced from the pooled / transient buffer. Never reset. */
+	FRHIBufferViewCache* ViewCache = nullptr;
+
 	/** Cached state pointer from the pooled / transient buffer. */
 	FRDGSubresourceState* State = nullptr;
 
