Fixing fallback cases when creating cached UAVs/SRVs.

diff --git a/Engine/Source/Runtime/RHI/Private/RHI.cpp b/Engine/Source/Runtime/RHI/Private/RHI.cpp
--- a/Engine/Source/Runtime/RHI/Private/RHI.cpp
+++ b/Engine/Source/Runtime/RHI/Private/RHI.cpp
@@ -2024,11 +2024,7 @@ FRHIShaderResourceView* FRHITextureViewCache::GetOrCreateSRV(FRHITexture* Textur
 
 	FShaderResourceViewRHIRef RHIShaderResourceView;
 
-	if (SRVCreateInfo.MetaData == ERHITextureMetaDataAccess::None)
-	{
-		RHIShaderResourceView = RHICreateShaderResourceView(Texture, SRVCreateInfo);
-	}
-	else
+	if (SRVCreateInfo.MetaData != ERHITextureMetaDataAccess::None)
 	{
 		FRHITexture2D* Texture2D = Texture->GetTexture2D();
 		check(Texture2D);
@@ -2047,13 +2043,14 @@ FRHIShaderResourceView* FRHITextureViewCache::GetOrCreateSRV(FRHITexture* Textur
 		case ERHITextureMetaDataAccess::CMask:
 			RHIShaderResourceView = RHICreateShaderResourceViewWriteMask(Texture2D);
 			break;
-
-		default:
-			checkf(false, TEXT("Invalid texture metadata access used when creating SRV."));
-			return nullptr;
 		}
 	}
 
+	if (!RHIShaderResourceView)
+	{
+		RHIShaderResourceView = RHICreateShaderResourceView(Texture, SRVCreateInfo);
+	}
+
 	check(RHIShaderResourceView);
 	FRHIShaderResourceView* View = RHIShaderResourceView.GetReference();
 	SRVs.Emplace(SRVCreateInfo, MoveTemp(RHIShaderResourceView));
@@ -2072,18 +2069,7 @@ FRHIUnorderedAccessView* FRHITextureViewCache::GetOrCreateUAV(FRHITexture* Textu
 
 	FUnorderedAccessViewRHIRef RHIUnorderedAccessView;
 
-	if (UAVCreateInfo.MetaData == ERHITextureMetaDataAccess::None)
-	{
-		if (UAVCreateInfo.Format != PF_Unknown)
-		{
-			RHIUnorderedAccessView = RHICreateUnorderedAccessView(Texture, UAVCreateInfo.MipLevel, UAVCreateInfo.Format);
-		}
-		else
-		{
-			RHIUnorderedAccessView = RHICreateUnorderedAccessView(Texture, UAVCreateInfo.MipLevel);
-		}
-	}
-	else
+	if (UAVCreateInfo.MetaData != ERHITextureMetaDataAccess::None)
 	{
 		FRHITexture2D* Texture2D = Texture->GetTexture2D();
 		check(Texture2D);
@@ -2098,10 +2084,18 @@ FRHIUnorderedAccessView* FRHITextureViewCache::GetOrCreateUAV(FRHITexture* Textu
 		case ERHITextureMetaDataAccess::Stencil:
 			RHIUnorderedAccessView = RHICreateUnorderedAccessViewStencil(Texture2D, UAVCreateInfo.MipLevel);
 			break;
+		}
+	}
 
-		default:
-			checkf(false, TEXT("Invalid texture metadata access used when creating UAV."));
-			return nullptr;
+	if (!RHIUnorderedAccessView)
+	{
+		if (UAVCreateInfo.Format != PF_Unknown)
+		{
+			RHIUnorderedAccessView = RHICreateUnorderedAccessView(Texture, UAVCreateInfo.MipLevel, UAVCreateInfo.Format);
+		}
+		else
+		{
+			RHIUnorderedAccessView = RHICreateUnorderedAccessView(Texture, UAVCreateInfo.MipLevel);
 		}
 	}
 
