RTV/DSV size validation
 - Binding differently-sized textures as RTV / DSV is not allowed.

diff --git a/Engine/Source/Runtime/RHI/Public/RHIResources.h b/Engine/Source/Runtime/RHI/Public/RHIResources.h
--- a/Engine/Source/Runtime/RHI/Public/RHIResources.h
+++ b/Engine/Source/Runtime/RHI/Public/RHIResources.h
@@ -966,6 +966,21 @@ public:
 		return nullptr;
 	}
 
+	/**
+	 * Returns the dimensions (i.e. the actual number of texels in each dimension) of the specified mip. ArraySize is ignored.
+	 * The Z component will always be 1 for 2D/cube resources and will contain depth for volume textures.
+	 * This differs from GetSizeXYZ() which returns ArraySize in Z for 2D arrays.
+	 */
+	virtual FIntVector GetMipDimensions(uint8 MipIndex) const
+	{
+		FIntVector Size = GetSizeXYZ();
+		return FIntVector(
+			FMath::Max(Size.X >> MipIndex, 1),
+			FMath::Max(Size.Y >> MipIndex, 1),
+			FMath::Max(Size.Z >> MipIndex, 1)
+		);
+	}
+
 	/** @return The number of mip-maps in the texture. */
 	uint32 GetNumMips() const { return NumMips; }
 
@@ -1128,6 +1143,17 @@ public:
 		return FIntVector(GetSizeX(), GetSizeY(), SizeZ);
 	}
 
+	// Because GetSizeXYZ() returns ArraySize in Z, we need to override this function to return 1 instead.
+	// @todo: Unify the meaning of "Z" in all texture resources
+	virtual FIntVector GetMipDimensions(uint8 MipIndex) const final override
+	{
+		return FIntVector(
+			FMath::Max(GetSizeX() >> MipIndex, 1u),
+			FMath::Max(GetSizeY() >> MipIndex, 1u),
+			1
+		);
+	}
+
 private:
 
 	uint32 SizeZ;
diff --git a/Engine/Source/Runtime/RHI/Public/RHIValidationContext.h b/Engine/Source/Runtime/RHI/Public/RHIValidationContext.h
--- a/Engine/Source/Runtime/RHI/Public/RHIValidationContext.h
+++ b/Engine/Source/Runtime/RHI/Public/RHIValidationContext.h
@@ -829,6 +829,8 @@ public:
 		State.RenderPassName = InName;
 		State.bGfxPSOSet = false;
 
+		FIntVector ViewDimensions = FIntVector(0);
+
 		// assert that render targets are writable
 		for (int32 RTVIndex = 0; RTVIndex < MaxSimultaneousRenderTargets; ++RTVIndex)
 		{
@@ -838,6 +840,11 @@ public:
 				continue;
 			}
 
+			// Check all bound textures have the same dimensions
+			FIntVector MipDimensions = RTV.RenderTarget->GetMipDimensions(RTV.MipIndex);
+			checkf(ViewDimensions.IsZero() || ViewDimensions == MipDimensions, TEXT("Render target size mismatch. All render and depth target views must have the same effective dimensions."));
+			ViewDimensions = MipDimensions;
+
 			uint32 ArraySlice = RTV.ArraySlice;
 			uint32 NumArraySlices = 1;
 			if (RTV.ArraySlice < 0)
@@ -851,6 +858,16 @@ public:
 
 		FRHIRenderPassInfo::FDepthStencilEntry& DSV = State.RenderPassInfo.DepthStencilRenderTarget;
 
+		if (DSV.DepthStencilTarget)
+		{
+			// Check all bound textures have the same dimensions
+			FIntVector MipDimensions = DSV.DepthStencilTarget->GetMipDimensions(0);
+			checkf(ViewDimensions.IsZero() || ViewDimensions == MipDimensions, TEXT("Depth target size mismatch. All render and depth target views must have the same effective dimensions."));
+			ViewDimensions = MipDimensions;
+		}
+
+		// @todo: additional checks for matching array slice counts on RTVs/DSVs
+
 		// assert depth is in the correct mode
 		if (DSV.ExclusiveDepthStencil.IsUsingDepth())
 		{
