Adding commands to dump RHI resource info.
rhi.DumpResourceMemory dumps the memory usage of each RHI resource if the type is instrumented. Currently only fully implemented for Textures and Buffers in D3D11 and D3D12.
rhi.DumpResourceCounts dumps the counts of each resource type. Should work for all RHIs.
Repurposed FRHIResourceInfo to use for generic resource info. Removed older uses of it that didn't actually provide data.
Made ERHIResourceType required to construct an FRHIResource.
Adding Name to FRHIBuffer.
Removed the unused FVulkanBuffer.
Removed the unused FOpenGLTexture and FOpenGLBaseTexture.

diff --git a/Engine/Plugins/FX/Niagara/Source/Niagara/Classes/NiagaraScriptExecutionContext.h b/Engine/Plugins/FX/Niagara/Source/Niagara/Classes/NiagaraScriptExecutionContext.h
--- a/Engine/Plugins/FX/Niagara/Source/Niagara/Classes/NiagaraScriptExecutionContext.h
+++ b/Engine/Plugins/FX/Niagara/Source/Niagara/Classes/NiagaraScriptExecutionContext.h
@@ -300,7 +300,7 @@ struct FNiagaraGpuSpawnInfo
 class FNiagaraRHIUniformBufferLayout : public FRHIResource
 {
 public:
-	explicit FNiagaraRHIUniformBufferLayout(const TCHAR* LayoutName) : UBLayout(LayoutName) { }
+	explicit FNiagaraRHIUniformBufferLayout(const TCHAR* LayoutName) : FRHIResource(RRT_None), UBLayout(LayoutName) { }
 
 	FRHIUniformBufferLayout UBLayout;
 };
diff --git a/Engine/Source/Runtime/D3D12RHI/Private/D3D12Buffer.cpp b/Engine/Source/Runtime/D3D12RHI/Private/D3D12Buffer.cpp
--- a/Engine/Source/Runtime/D3D12RHI/Private/D3D12Buffer.cpp
+++ b/Engine/Source/Runtime/D3D12RHI/Private/D3D12Buffer.cpp
@@ -355,6 +355,12 @@ FD3D12Buffer* FD3D12Adapter::CreateRHIBuffer(
 			FD3D12Buffer* NewBuffer = new FD3D12Buffer(Device, Size, InUsage, Stride);
 			NewBuffer->BufferAlignment = Alignment;
 
+			if (InDebugName)
+			{
+				TRACE_CPUPROFILER_EVENT_SCOPE(D3D12RHI::SetDebugName);
+				NewBuffer->SetName(InDebugName);
+			}
+
 			if (Device->GetGPUIndex() == FirstGPUIndex)
 			{
 				AllocateBuffer(Device, InDesc, Size, InUsage, InResourceStateMode, InCreateState, Alignment, NewBuffer, NewBuffer->ResourceLocation, ResourceAllocator, InDebugName);
@@ -375,6 +381,13 @@ FD3D12Buffer* FD3D12Adapter::CreateRHIBuffer(
 		{
 			FD3D12Buffer* NewBuffer = new FD3D12Buffer(Device, Size, InUsage, Stride);
 			NewBuffer->BufferAlignment = Alignment;
+
+			if (InDebugName)
+			{
+				TRACE_CPUPROFILER_EVENT_SCOPE(D3D12RHI::SetDebugName);
+				NewBuffer->SetName(InDebugName);
+			}
+
 			AllocateBuffer(Device, InDesc, Size, InUsage, InResourceStateMode, InCreateState, Alignment, NewBuffer, NewBuffer->ResourceLocation, ResourceAllocator, InDebugName);
 			
 			// Unlock immediately if there is no initial data
diff --git a/Engine/Source/Runtime/D3D12RHI/Private/D3D12Texture.cpp b/Engine/Source/Runtime/D3D12RHI/Private/D3D12Texture.cpp
--- a/Engine/Source/Runtime/D3D12RHI/Private/D3D12Texture.cpp
+++ b/Engine/Source/Runtime/D3D12RHI/Private/D3D12Texture.cpp
@@ -1628,10 +1628,6 @@ FTexture3DRHIRef FD3D12DynamicRHI::RHICreateTexture3D(uint32 SizeX, uint32 SizeY
 
 void FD3D12DynamicRHI::RHIGetResourceInfo(FRHITexture* Ref, FRHIResourceInfo& OutInfo)
 {
-	if (Ref)
-	{
-		OutInfo = Ref->ResourceInfo;
-	}
 }
 
 /**
diff --git a/Engine/Source/Runtime/D3D12RHI/Private/D3D12Texture.h b/Engine/Source/Runtime/D3D12RHI/Private/D3D12Texture.h
--- a/Engine/Source/Runtime/D3D12RHI/Private/D3D12Texture.h
+++ b/Engine/Source/Runtime/D3D12RHI/Private/D3D12Texture.h
@@ -235,6 +235,8 @@ public:
 	/** Flags used when the texture was created */
 	ETextureCreateFlags Flags;
 
+	TD3D12Texture2D() = delete;
+
 	/** Initialization constructor. */
 	TD3D12Texture2D(
 		class FD3D12Device* InParent,
@@ -286,6 +288,18 @@ public:
 
 	virtual ~TD3D12Texture2D();
 
+	// FRHIResource overrides
+#if RHI_WANT_RESOURCE_INFO
+	bool GetResourceInfo(FRHIResourceInfo& OutResourceInfo) const override
+	{
+		OutResourceInfo = FRHIResourceInfo{};
+		OutResourceInfo.Name = this->GetName().ToString();
+		OutResourceInfo.Type = this->GetType();
+		OutResourceInfo.VRamAllocation.AllocationSize = GetMemorySize();
+		return true;
+	}
+#endif
+
 	/**
 	* Locks one of the texture's mip-maps.
 	* @return A pointer to the specified texture data.
@@ -396,7 +410,7 @@ private:
 class FD3D12Texture3D : public FRHITexture3D, public FD3D12TextureBase
 {
 public:
-
+	FD3D12Texture3D() = delete;
 	/** Initialization constructor. */
 	FD3D12Texture3D(
 	class FD3D12Device* InParent,
@@ -416,6 +430,18 @@ public:
 
 	virtual ~FD3D12Texture3D();
 
+	// FRHIResource overrides
+#if RHI_WANT_RESOURCE_INFO
+	bool GetResourceInfo(FRHIResourceInfo& OutResourceInfo) const override
+	{
+		OutResourceInfo = FRHIResourceInfo{};
+		OutResourceInfo.Name = GetName().ToString();
+		OutResourceInfo.Type = GetType();
+		OutResourceInfo.VRamAllocation.AllocationSize = ResourceLocation.GetSize();
+		return true;
+	}
+#endif
+
 	/** FRHITexture override.  See FRHITexture::GetNativeResource() */
 	virtual void* GetNativeResource() const override final
 	{
diff --git a/Engine/Source/Runtime/D3D12RHI/Private/D3D12View.h b/Engine/Source/Runtime/D3D12RHI/Private/D3D12View.h
--- a/Engine/Source/Runtime/D3D12RHI/Private/D3D12View.h
+++ b/Engine/Source/Runtime/D3D12RHI/Private/D3D12View.h
@@ -1059,12 +1059,14 @@ class FD3D12RenderTargetView : public FD3D12View<D3D12_RENDER_TARGET_VIEW_DESC>,
 public:
 	FD3D12RenderTargetView(FD3D12Device* InParent, const D3D12_RENDER_TARGET_VIEW_DESC& InRTVDesc, FD3D12BaseShaderResource* InBaseShaderResource)
 		: FD3D12View(InParent, ViewSubresourceSubsetFlags_None)
+		, FRHIResource(RRT_None)
 	{
 		CreateView(InRTVDesc, InBaseShaderResource, InBaseShaderResource->ResourceLocation);
 	}
 
 	FD3D12RenderTargetView(FD3D12Device* InParent, const D3D12_RENDER_TARGET_VIEW_DESC& InRTVDesc, FD3D12BaseShaderResource* InBaseShaderResource, FD3D12ResourceLocation& InResourceLocation)
 		: FD3D12View(InParent, ViewSubresourceSubsetFlags_None)
+		, FRHIResource(RRT_None)
 	{
 		CreateView(InRTVDesc, InBaseShaderResource, InResourceLocation);
 	}
@@ -1086,6 +1088,7 @@ class FD3D12DepthStencilView : public FD3D12View<D3D12_DEPTH_STENCIL_VIEW_DESC>,
 public:
 	FD3D12DepthStencilView(FD3D12Device* InParent, const D3D12_DEPTH_STENCIL_VIEW_DESC& InDSVDesc, FD3D12BaseShaderResource* InBaseShaderResource, bool InHasStencil)
 		: FD3D12View(InParent, ViewSubresourceSubsetFlags_DepthAndStencilDsv)
+		, FRHIResource(RRT_None)
 		, bHasDepth(true)				// Assume all DSVs have depth bits in their format
 		, bHasStencil(InHasStencil)		// Only some DSVs have stencil bits in their format
 	{
diff --git a/Engine/Source/Runtime/D3D12RHI/Public/D3D12Resources.h b/Engine/Source/Runtime/D3D12RHI/Public/D3D12Resources.h
--- a/Engine/Source/Runtime/D3D12RHI/Public/D3D12Resources.h
+++ b/Engine/Source/Runtime/D3D12RHI/Public/D3D12Resources.h
@@ -925,6 +925,18 @@ public:
 	void UploadResourceData(class FRHICommandListImmediate* InRHICmdList, FResourceArrayInterface* InResourceArray, D3D12_RESOURCE_STATES InDestinationState);
 	FD3D12SyncPoint UploadResourceDataViaCopyQueue(FResourceArrayInterface* InResourceArray);
 
+	// FRHIResource overrides
+#if RHI_WANT_RESOURCE_INFO
+	bool GetResourceInfo(FRHIResourceInfo& OutResourceInfo) const override
+	{
+		OutResourceInfo = FRHIResourceInfo{};
+		OutResourceInfo.Name = GetName().ToString();
+		OutResourceInfo.Type = GetType();
+		OutResourceInfo.VRamAllocation.AllocationSize = ResourceLocation.GetSize();
+		return true;
+	}
+#endif
+
 	void Rename(FD3D12ResourceLocation& NewLocation);
 	void RenameLDAChain(FD3D12ResourceLocation& NewLocation);
 
diff --git a/Engine/Source/Runtime/OpenGLDrv/Private/OpenGLTexture.cpp b/Engine/Source/Runtime/OpenGLDrv/Private/OpenGLTexture.cpp
--- a/Engine/Source/Runtime/OpenGLDrv/Private/OpenGLTexture.cpp
+++ b/Engine/Source/Runtime/OpenGLDrv/Private/OpenGLTexture.cpp
@@ -1576,8 +1576,6 @@ void TOpenGLTexture<RHIResourceType>::CloneViaPBO( TOpenGLTexture<RHIResourceTyp
 	CachedBindPixelUnpackBuffer(GL_PIXEL_UNPACK_BUFFER, 0);
 }
 
-template class TOpenGLTexture<FOpenGLBaseTexture>;
-
 /*-----------------------------------------------------------------------------
 	2D texture support.
 -----------------------------------------------------------------------------*/
diff --git a/Engine/Source/Runtime/OpenGLDrv/Private/OpenGLUAV.cpp b/Engine/Source/Runtime/OpenGLDrv/Private/OpenGLUAV.cpp
--- a/Engine/Source/Runtime/OpenGLDrv/Private/OpenGLUAV.cpp
+++ b/Engine/Source/Runtime/OpenGLDrv/Private/OpenGLUAV.cpp
@@ -176,8 +176,7 @@ FOpenGLShaderResourceView::~FOpenGLShaderResourceView()
 
 FUnorderedAccessViewRHIRef FOpenGLDynamicRHI::RHICreateUnorderedAccessView(FRHITexture* TextureRHI, uint32 MipLevel)
 {
-	FOpenGLTexture* Texture = ResourceCast(TextureRHI);
-	check(Texture->GetFlags() & TexCreate_UAV);
+	check(TextureRHI->GetFlags() & TexCreate_UAV);
 	return new FOpenGLTextureUnorderedAccessView(TextureRHI);
 }
 
diff --git a/Engine/Source/Runtime/OpenGLDrv/Public/OpenGLResources.h b/Engine/Source/Runtime/OpenGLDrv/Public/OpenGLResources.h
--- a/Engine/Source/Runtime/OpenGLDrv/Public/OpenGLResources.h
+++ b/Engine/Source/Runtime/OpenGLDrv/Public/OpenGLResources.h
@@ -1763,19 +1763,6 @@ public:
 	}
 };
 
-class OPENGLDRV_API FOpenGLBaseTexture : public FRHITexture
-{
-public:
-	FOpenGLBaseTexture(uint32 InSizeX, uint32 InSizeY, uint32 InSizeZ, uint32 InNumMips, uint32 InNumSamples, uint32 InNumSamplesTileMem, uint32 InArraySize, EPixelFormat InFormat, ETextureCreateFlags InFlags, const FClearValueBinding& InClearValue)
-		: FRHITexture(InNumMips, InNumSamples, InFormat, InFlags, NULL, InClearValue)
-	{}
-
-	uint32 GetSizeX() const { return 0; }
-	uint32 GetSizeY() const { return 0; }
-	uint32 GetSizeZ() const { return 0; }
-};
-
-typedef TOpenGLTexture<FOpenGLBaseTexture>				FOpenGLTexture;
 typedef TOpenGLTexture<FOpenGLBaseTexture2D>			FOpenGLTexture2D;
 typedef TOpenGLTexture<FOpenGLBaseTexture2DArray>		FOpenGLTexture2DArray;
 typedef TOpenGLTexture<FOpenGLBaseTexture3D>			FOpenGLTexture3D;
@@ -2305,11 +2292,7 @@ struct TOpenGLResourceTraits<FRHITexture3D>
 {
 	typedef FOpenGLTexture3D TConcreteType;
 };
-template<>
-struct TOpenGLResourceTraits<FRHITexture>
-{
-	typedef FOpenGLTexture TConcreteType;
-};
+
 template<>
 struct TOpenGLResourceTraits<FRHITexture2D>
 {
diff --git a/Engine/Source/Runtime/RHI/Private/RHI.cpp b/Engine/Source/Runtime/RHI/Private/RHI.cpp
--- a/Engine/Source/Runtime/RHI/Private/RHI.cpp
+++ b/Engine/Source/Runtime/RHI/Private/RHI.cpp
@@ -13,6 +13,7 @@
 #include "ProfilingDebugging/CsvProfiler.h"
 #include "String/LexFromString.h"
 #include "String/ParseTokens.h"
+#include "Misc/BufferedOutputDevice.h"
 
 IMPLEMENT_MODULE(FDefaultModuleImpl, RHI);
 
@@ -629,6 +630,327 @@ bool operator==(const FBlendStateInitializerRHI::FRenderTarget& A, const FBlendS
 	return bSame;
 }
 
+#if RHI_WANT_RESOURCE_INFO
+
+static FCriticalSection GRHIResourceTrackingCriticalSection;
+static TSet<FRHIResource*> GRHITrackedResources;
+
+void FRHIResource::BeginTrackingResource(FRHIResource* InResource)
+{
+	FScopeLock Lock(&GRHIResourceTrackingCriticalSection);
+
+	InResource->bBeingTracked = true;
+
+	GRHITrackedResources.Add(InResource);
+}
+
+void FRHIResource::EndTrackingResource(FRHIResource* InResource)
+{
+	if (InResource->bBeingTracked)
+	{
+		FScopeLock Lock(&GRHIResourceTrackingCriticalSection);
+		GRHITrackedResources.Remove(InResource);
+		InResource->bBeingTracked = false;
+	}
+}
+
+struct FRHIResourceTypeName
+{
+	ERHIResourceType Type;
+	const TCHAR* Name;
+};
+static const FRHIResourceTypeName GRHIResourceTypeNames[] =
+{
+#define RHI_RESOURCE_TYPE_DEF(Name) { RRT_##Name, TEXT(#Name) }
+
+	RHI_RESOURCE_TYPE_DEF(None),
+	RHI_RESOURCE_TYPE_DEF(SamplerState),
+	RHI_RESOURCE_TYPE_DEF(RasterizerState),
+	RHI_RESOURCE_TYPE_DEF(DepthStencilState),
+	RHI_RESOURCE_TYPE_DEF(BlendState),
+	RHI_RESOURCE_TYPE_DEF(VertexDeclaration),
+	RHI_RESOURCE_TYPE_DEF(VertexShader),
+	RHI_RESOURCE_TYPE_DEF(MeshShader),
+	RHI_RESOURCE_TYPE_DEF(AmplificationShader),
+	RHI_RESOURCE_TYPE_DEF(PixelShader),
+	RHI_RESOURCE_TYPE_DEF(GeometryShader),
+	RHI_RESOURCE_TYPE_DEF(RayTracingShader),
+	RHI_RESOURCE_TYPE_DEF(ComputeShader),
+	RHI_RESOURCE_TYPE_DEF(GraphicsPipelineState),
+	RHI_RESOURCE_TYPE_DEF(ComputePipelineState),
+	RHI_RESOURCE_TYPE_DEF(RayTracingPipelineState),
+	RHI_RESOURCE_TYPE_DEF(BoundShaderState),
+	RHI_RESOURCE_TYPE_DEF(UniformBuffer),
+	RHI_RESOURCE_TYPE_DEF(Buffer),
+	RHI_RESOURCE_TYPE_DEF(Texture),
+	RHI_RESOURCE_TYPE_DEF(Texture2D),
+	RHI_RESOURCE_TYPE_DEF(Texture2DArray),
+	RHI_RESOURCE_TYPE_DEF(Texture3D),
+	RHI_RESOURCE_TYPE_DEF(TextureCube),
+	RHI_RESOURCE_TYPE_DEF(TextureReference),
+	RHI_RESOURCE_TYPE_DEF(TimestampCalibrationQuery),
+	RHI_RESOURCE_TYPE_DEF(GPUFence),
+	RHI_RESOURCE_TYPE_DEF(RenderQuery),
+	RHI_RESOURCE_TYPE_DEF(RenderQueryPool),
+	RHI_RESOURCE_TYPE_DEF(ComputeFence),
+	RHI_RESOURCE_TYPE_DEF(Viewport),
+	RHI_RESOURCE_TYPE_DEF(UnorderedAccessView),
+	RHI_RESOURCE_TYPE_DEF(ShaderResourceView),
+	RHI_RESOURCE_TYPE_DEF(RayTracingAccelerationStructure),
+	RHI_RESOURCE_TYPE_DEF(StagingBuffer),
+	RHI_RESOURCE_TYPE_DEF(CustomPresent),
+	RHI_RESOURCE_TYPE_DEF(ShaderLibrary),
+	RHI_RESOURCE_TYPE_DEF(PipelineBinaryLibrary),
+};
+
+static ERHIResourceType RHIResourceTypeFromString(const FString& InString)
+{
+	for (const auto& TypeName : GRHIResourceTypeNames)
+	{
+		if (InString.Equals(TypeName.Name, ESearchCase::IgnoreCase))
+		{
+			return TypeName.Type;
+		}
+	}
+	return RRT_None;
+}
+
+static const TCHAR* StringFromRHIResourceType(ERHIResourceType ResourceType)
+{
+	for (const auto& TypeName : GRHIResourceTypeNames)
+	{
+		if (TypeName.Type == ResourceType)
+		{
+			return TypeName.Name;
+		}
+	}
+	return TEXT("<unknown>");
+}
+
+static FAutoConsoleCommandWithWorldArgsAndOutputDevice GDumpRHIResourceCountsCmd(
+	TEXT("rhi.DumpResourceCounts"),
+	TEXT("Dumps RHI resource counts to the log"),
+	FConsoleCommandWithWorldArgsAndOutputDeviceDelegate::CreateStatic([](const TArray<FString>& Args, UWorld*, FOutputDevice& OutputDevice)
+{
+	int32 ResourceCounts[RRT_Num]{};
+	int32 TotalResources{};
+
+	FRHIResourceInfo ResourceInfo;
+
+	{
+		FScopeLock Lock(&GRHIResourceTrackingCriticalSection);
+
+		TotalResources = GRHITrackedResources.Num();
+
+		for (const FRHIResource* Resource : GRHITrackedResources)
+		{
+			if (Resource)
+			{
+				ERHIResourceType ResourceType = Resource->GetType();
+				if (ResourceType > 0 && ResourceType < RRT_Num)
+				{
+					ResourceCounts[ResourceType]++;
+				}
+			}
+		}
+	}
+
+	FBufferedOutputDevice BufferedOutput;
+	FName CategoryName(TEXT("RHIResources"));
+
+	BufferedOutput.CategorizedLogf(CategoryName, ELogVerbosity::Log, TEXT("RHIResource Counts"));
+
+	for (int32 Index = 0; Index < RRT_Num; Index++)
+	{
+		const int32 CurrentCount = ResourceCounts[Index];
+		if (CurrentCount > 0)
+		{
+			BufferedOutput.CategorizedLogf(CategoryName, ELogVerbosity::Log, TEXT("%s: %d"),
+				StringFromRHIResourceType((ERHIResourceType)Index),
+				CurrentCount);
+		}
+	}
+	BufferedOutput.CategorizedLogf(CategoryName, ELogVerbosity::Log, TEXT("Total: %d"), TotalResources);
+
+	BufferedOutput.RedirectTo(OutputDevice);
+}));
+
+static FAutoConsoleCommandWithWorldArgsAndOutputDevice GDumpRHIResourceMemoryCmd(
+	TEXT("rhi.DumpResourceMemory"),
+	TEXT("Dumps RHI resource memory stats to the log"),
+	FConsoleCommandWithWorldArgsAndOutputDeviceDelegate::CreateStatic([](const TArray<FString>& Args, UWorld*, FOutputDevice& OutputDevice)
+{
+	FString NameFilter;
+	ERHIResourceType TypeFilter = RRT_None;
+	int32 NumberOfResourcesToShow = 50;
+	bool bUseCSVOutput = false;
+
+	for (const FString& Argument : Args)
+	{
+		if (Argument.Equals(TEXT("all"), ESearchCase::IgnoreCase))
+		{
+			NumberOfResourcesToShow = -1;
+		}
+		else if (Argument.Equals(TEXT("csv"), ESearchCase::IgnoreCase))
+		{
+			bUseCSVOutput = true;
+		}
+		else if (Argument.StartsWith(TEXT("Name="), ESearchCase::IgnoreCase))
+		{
+			NameFilter = Argument.RightChop(5);
+		}
+		else if (Argument.StartsWith(TEXT("Type="), ESearchCase::IgnoreCase))
+		{
+			TypeFilter = RHIResourceTypeFromString(Argument.RightChop(5));
+		}
+		else if (FCString::IsNumeric(*Argument))
+		{
+			LexFromString(NumberOfResourcesToShow, *Argument);
+		}
+		else
+		{
+			NameFilter = Argument;
+		}
+	}
+
+	auto ShouldIncludeResource = [&](const FRHIResource* Resource, const FRHIResourceInfo& ResourceInfo) -> bool
+	{
+		if (!NameFilter.IsEmpty() && ResourceInfo.Name.Find(NameFilter, ESearchCase::IgnoreCase) == INDEX_NONE)
+		{
+			return false;
+		}
+		if (TypeFilter != RRT_None)
+		{
+			if (TypeFilter == RRT_Texture)
+			{
+				if (ResourceInfo.Type != RRT_Texture2D &&
+					ResourceInfo.Type != RRT_Texture2DArray &&
+					ResourceInfo.Type != RRT_Texture3D &&
+					ResourceInfo.Type != RRT_TextureCube)
+				{
+					return false;
+				}
+			}
+			else if (TypeFilter != ResourceInfo.Type)
+			{
+				return false;
+			}
+		}
+		return true;
+	};
+
+
+	struct FLocalResourceEntry
+	{
+		const FRHIResource* Resource;
+		FRHIResourceInfo ResourceInfo;
+	};
+	TArray<FLocalResourceEntry> Resources;
+
+	FScopeLock Lock(&GRHIResourceTrackingCriticalSection);
+
+	int32 TotalResourcesWithInfo = 0;
+	int32 TotalTrackedResources = 0;
+	int64 TotalTrackedResourceSize = 0;
+
+	{
+		FRHIResourceInfo ResourceInfo;
+		TotalTrackedResources = GRHITrackedResources.Num();
+
+		for (const FRHIResource* Resource : GRHITrackedResources)
+		{
+			if (Resource && Resource->GetResourceInfo(ResourceInfo))
+			{
+				if (ShouldIncludeResource(Resource, ResourceInfo))
+				{
+					Resources.Emplace(FLocalResourceEntry{ Resource, ResourceInfo });
+				}
+
+				TotalResourcesWithInfo++;
+				TotalTrackedResourceSize += ResourceInfo.VRamAllocation.AllocationSize;
+			}
+		}
+	}
+
+	if (NumberOfResourcesToShow < 0 || NumberOfResourcesToShow > Resources.Num())
+	{
+		NumberOfResourcesToShow = Resources.Num();
+	}
+
+	Resources.Sort([](const FLocalResourceEntry& EntryA, const FLocalResourceEntry& EntryB)
+	{
+		return EntryA.ResourceInfo.VRamAllocation.AllocationSize > EntryB.ResourceInfo.VRamAllocation.AllocationSize;
+	});
+
+	FBufferedOutputDevice BufferedOutput;
+	FName CategoryName(TEXT("RHIResources"));
+
+	if (bUseCSVOutput)
+	{
+		BufferedOutput.Logf(ELogVerbosity::Log, TEXT("\"Name\",\"Type\",\"Size\""));
+	}
+	else
+	{
+		BufferedOutput.CategorizedLogf(CategoryName, ELogVerbosity::Log, TEXT("Tracked RHIResources"));
+
+		if (NumberOfResourcesToShow != Resources.Num())
+		{
+			BufferedOutput.CategorizedLogf(CategoryName, ELogVerbosity::Log, TEXT("Showing %d out of %d with info and %d tracked"), NumberOfResourcesToShow, TotalResourcesWithInfo, TotalTrackedResources);
+		}
+		else
+		{
+			BufferedOutput.CategorizedLogf(CategoryName, ELogVerbosity::Log, TEXT("Showing %d with info and %d tracked"), TotalResourcesWithInfo, TotalTrackedResources);
+		}
+	}
+
+	int64 TotalShownResourceSize = 0;
+
+	for (int32 Index = 0; Index < Resources.Num(); Index++)
+	{
+		if (Index < NumberOfResourcesToShow)
+		{
+			const FRHIResourceInfo& ResourceInfo = Resources[Index].ResourceInfo;
+			const TCHAR* ResourceType = StringFromRHIResourceType(ResourceInfo.Type);
+			const int64 SizeInBytes = ResourceInfo.VRamAllocation.AllocationSize;
+
+			if (bUseCSVOutput)
+			{
+				BufferedOutput.Logf(ELogVerbosity::Log, TEXT("\"%s\",\"%s\",\"%.3f\""),
+					*ResourceInfo.Name,
+					ResourceType,
+					SizeInBytes / double(1 << 20));
+			}
+			else
+			{
+				BufferedOutput.CategorizedLogf(CategoryName, ELogVerbosity::Log, TEXT("Name: %s - Type: %s - Size: %.3f MB"),
+					*ResourceInfo.Name,
+					ResourceType,
+					SizeInBytes / double(1 << 20));
+			}
+
+			TotalShownResourceSize += SizeInBytes;
+		}
+	}
+
+	const double TotalSizeF = TotalTrackedResourceSize / double(1 << 20);
+	const double ShownSizeF = TotalShownResourceSize / double(1 << 20);
+
+	if (!bUseCSVOutput)
+	{
+		if (NumberOfResourcesToShow != Resources.Num())
+		{
+			BufferedOutput.CategorizedLogf(CategoryName, ELogVerbosity::Log, TEXT("Shown %d entries. Size: %.3f MB (%.2f%% of total)"),
+				NumberOfResourcesToShow, TotalSizeF, 100.0 * ShownSizeF / TotalSizeF);
+		}
+
+		BufferedOutput.CategorizedLogf(CategoryName, ELogVerbosity::Log, TEXT("Total tracked resource size: %.3f MB"), TotalSizeF);
+	}
+
+	BufferedOutput.RedirectTo(OutputDevice);
+}));
+
+#endif // RHI_WANT_RESOURCE_INFO
+
 bool FRHIResource::Bypass()
 {
 	return GRHICommandList.Bypass();
diff --git a/Engine/Source/Runtime/RHI/Public/RHI.h b/Engine/Source/Runtime/RHI/Public/RHI.h
--- a/Engine/Source/Runtime/RHI/Public/RHI.h
+++ b/Engine/Source/Runtime/RHI/Public/RHI.h
@@ -1416,7 +1416,8 @@ typedef TArray<FScreenResolutionRHI>	FScreenResolutionArray;
 
 struct FVRamAllocation
 {
-	FVRamAllocation(uint32 InAllocationStart = 0, uint32 InAllocationSize = 0)
+	FVRamAllocation() = default;
+	FVRamAllocation(uint64 InAllocationStart, uint64 InAllocationSize)
 		: AllocationStart(InAllocationStart)
 		, AllocationSize(InAllocationSize)
 	{
@@ -1425,13 +1426,15 @@ struct FVRamAllocation
 	bool IsValid() const { return AllocationSize > 0; }
 
 	// in bytes
-	uint32 AllocationStart;
+	uint64 AllocationStart{};
 	// in bytes
-	uint32 AllocationSize;
+	uint64 AllocationSize{};
 };
 
 struct FRHIResourceInfo
 {
+	FString Name;
+	ERHIResourceType Type{ RRT_None };
 	FVRamAllocation VRamAllocation;
 };
 
diff --git a/Engine/Source/Runtime/RHI/Public/RHIDefinitions.h b/Engine/Source/Runtime/RHI/Public/RHIDefinitions.h
--- a/Engine/Source/Runtime/RHI/Public/RHIDefinitions.h
+++ b/Engine/Source/Runtime/RHI/Public/RHIDefinitions.h
@@ -1482,9 +1482,15 @@ enum ERHIResourceType
 	RRT_BlendState,
 	RRT_VertexDeclaration,
 	RRT_VertexShader,
+	RRT_MeshShader,
+	RRT_AmplificationShader,
 	RRT_PixelShader,
 	RRT_GeometryShader,
+	RRT_RayTracingShader,
 	RRT_ComputeShader,
+	RRT_GraphicsPipelineState,
+	RRT_ComputePipelineState,
+	RRT_RayTracingPipelineState,
 	RRT_BoundShaderState,
 	RRT_UniformBuffer,
 	RRT_Buffer,
@@ -1494,10 +1500,19 @@ enum ERHIResourceType
 	RRT_Texture3D,
 	RRT_TextureCube,
 	RRT_TextureReference,
+	RRT_TimestampCalibrationQuery,
+	RRT_GPUFence,
 	RRT_RenderQuery,
+	RRT_RenderQueryPool,
+	RRT_ComputeFence,
 	RRT_Viewport,
 	RRT_UnorderedAccessView,
 	RRT_ShaderResourceView,
+	RRT_RayTracingAccelerationStructure,
+	RRT_StagingBuffer,
+	RRT_CustomPresent,
+	RRT_ShaderLibrary,
+	RRT_PipelineBinaryLibrary,
 
 	RRT_Num
 };
diff --git a/Engine/Source/Runtime/RHI/Public/RHIResources.h b/Engine/Source/Runtime/RHI/Public/RHIResources.h
--- a/Engine/Source/Runtime/RHI/Public/RHIResources.h
+++ b/Engine/Source/Runtime/RHI/Public/RHIResources.h
@@ -23,6 +23,10 @@
 
 #define DISABLE_RHI_DEFFERED_DELETE 0
 
+#ifndef RHI_WANT_RESOURCE_INFO
+#define RHI_WANT_RESOURCE_INFO 0
+#endif
+
 struct FClearValueBinding;
 struct FRHIResourceInfo;
 struct FGenerateMipsStruct;
@@ -32,16 +36,29 @@ enum class EClearBinding;
 class RHI_API FRHIResource
 {
 public:
-	FRHIResource(bool InbDoNotDeferDelete = false)
-		: MarkedForDelete(0)
-		, bDoNotDeferDelete(InbDoNotDeferDelete)
-		, bCommitted(true)
+	UE_DEPRECATED(5.1, "FRHIResource(bool) is deprecated, please use FRHIResource(ERHIResourceType)")
+	FRHIResource(bool InbDoNotDeferDelete=false)
+		: ResourceType(RRT_None)
+	{
+	}
+
+	FRHIResource(ERHIResourceType InResourceType)
+		: ResourceType(InResourceType)
 	{
+#if RHI_WANT_RESOURCE_INFO
+		BeginTrackingResource(this);
+#endif
 	}
-	virtual ~FRHIResource() 
+
+	virtual ~FRHIResource()
 	{
 		check(PlatformNeedsExtraDeletionLatency() || (NumRefs.GetValue() == 0 && (CurrentlyDeleting == this || bDoNotDeferDelete || Bypass()))); // this should not have any outstanding refs
+
+#if RHI_WANT_RESOURCE_INFO
+		EndTrackingResource(this);
+#endif
 	}
+
 	FORCEINLINE_DEBUGGABLE uint32 AddRef() const
 	{
 		int32 NewValue = NumRefs.Increment();
@@ -109,11 +126,32 @@ public:
 		return !MarkedForDelete && NumRefs.GetValue() > 0;
 	}
 
+	inline ERHIResourceType GetType() const { return ResourceType; }
+
+#if RHI_WANT_RESOURCE_INFO
+	// Get resource info if available.
+	// Should return true if the ResourceInfo was filled with data.
+	virtual bool GetResourceInfo(FRHIResourceInfo& OutResourceInfo) const
+	{
+		OutResourceInfo = FRHIResourceInfo{};
+		return false;
+	}
+
+	static void BeginTrackingResource(FRHIResource* InResource);
+	static void EndTrackingResource(FRHIResource* InResource);
+#endif
+
 private:
 	mutable FThreadSafeCounter NumRefs;
-	mutable int32 MarkedForDelete;
-	bool bDoNotDeferDelete;
-	bool bCommitted;
+	mutable int32 MarkedForDelete{ 0 };
+
+	const ERHIResourceType ResourceType;
+
+	bool bDoNotDeferDelete{ false };
+	bool bCommitted{ true };
+#if RHI_WANT_RESOURCE_INFO
+	bool bBeingTracked{ false };
+#endif
 
 	static TLockFreePointerListUnordered<FRHIResource, PLATFORM_CACHE_LINE_SIZE> PendingDeletes;
 	static FRHIResource* CurrentlyDeleting;
@@ -133,14 +171,8 @@ private:
 	// to ensure the GPU has completely finished with them. This avoids expensive fences, etc.
 	struct ResourcesToDelete
 	{
-		ResourcesToDelete(uint32 InFrameDeleted = 0)
-			: FrameDeleted(InFrameDeleted)
-		{
-
-		}
-
 		TArray<FRHIResource*>	Resources;
-		uint32					FrameDeleted;
+		uint32					FrameDeleted{};
 	};
 
 	static TArray<ResourcesToDelete> DeferredDeletionQueue;
@@ -411,25 +443,31 @@ private:
 class FRHISamplerState : public FRHIResource 
 {
 public:
+	FRHISamplerState() : FRHIResource(RRT_SamplerState) {}
 	virtual bool IsImmutable() const { return false; }
 };
 
 class FRHIRasterizerState : public FRHIResource
 {
 public:
+	FRHIRasterizerState() : FRHIResource(RRT_RasterizerState) {}
 	virtual bool GetInitializer(struct FRasterizerStateInitializerRHI& Init) { return false; }
 };
+
 class FRHIDepthStencilState : public FRHIResource
 {
 public:
+	FRHIDepthStencilState() : FRHIResource(RRT_DepthStencilState) {}
 #if ENABLE_RHI_VALIDATION
 	FExclusiveDepthStencil ActualDSMode;
 #endif
 	virtual bool GetInitializer(struct FDepthStencilStateInitializerRHI& Init) { return false; }
 };
+
 class FRHIBlendState : public FRHIResource
 {
 public:
+	FRHIBlendState() : FRHIResource(RRT_BlendState) {}
 	virtual bool GetInitializer(class FBlendStateInitializerRHI& Init) { return false; }
 };
 
@@ -441,10 +479,15 @@ typedef TArray<struct FVertexElement,TFixedAllocator<MaxVertexElementCount> > FV
 class FRHIVertexDeclaration : public FRHIResource
 {
 public:
+	FRHIVertexDeclaration() : FRHIResource(RRT_VertexDeclaration) {}
 	virtual bool GetInitializer(FVertexDeclarationElementList& Init) { return false; }
 };
 
-class FRHIBoundShaderState : public FRHIResource {};
+class FRHIBoundShaderState : public FRHIResource
+{
+public:
+	FRHIBoundShaderState() : FRHIResource(RRT_BoundShaderState) {}
+};
 
 //
 // Shaders
@@ -464,8 +507,10 @@ public:
 	FORCEINLINE const TCHAR* GetShaderName() const { return TEXT(""); }
 #endif
 
-	explicit FRHIShader(EShaderFrequency InFrequency)
-		: Frequency(InFrequency)
+	FRHIShader() = delete;
+	FRHIShader(ERHIResourceType InResourceType, EShaderFrequency InFrequency)
+		: FRHIResource(InResourceType)
+		, Frequency(InFrequency)
 	{
 	}
 
@@ -482,49 +527,50 @@ private:
 class FRHIGraphicsShader : public FRHIShader
 {
 public:
-	explicit FRHIGraphicsShader(EShaderFrequency InFrequency) : FRHIShader(InFrequency) {}
+	explicit FRHIGraphicsShader(ERHIResourceType InResourceType, EShaderFrequency InFrequency)
+		: FRHIShader(InResourceType, InFrequency) {}
 };
 
 class FRHIVertexShader : public FRHIGraphicsShader
 {
 public:
-	FRHIVertexShader() : FRHIGraphicsShader(SF_Vertex) {}
+	FRHIVertexShader() : FRHIGraphicsShader(RRT_VertexShader, SF_Vertex) {}
 };
 
 class FRHIMeshShader : public FRHIGraphicsShader
 {
 public:
-	FRHIMeshShader() : FRHIGraphicsShader(SF_Mesh) {}
+	FRHIMeshShader() : FRHIGraphicsShader(RRT_MeshShader, SF_Mesh) {}
 };
 
 class FRHIAmplificationShader : public FRHIGraphicsShader
 {
 public:
-	FRHIAmplificationShader() : FRHIGraphicsShader(SF_Amplification) {}
+	FRHIAmplificationShader() : FRHIGraphicsShader(RRT_AmplificationShader, SF_Amplification) {}
 };
 
 class FRHIPixelShader : public FRHIGraphicsShader
 {
 public:
-	FRHIPixelShader() : FRHIGraphicsShader(SF_Pixel) {}
+	FRHIPixelShader() : FRHIGraphicsShader(RRT_PixelShader, SF_Pixel) {}
 };
 
 class FRHIGeometryShader : public FRHIGraphicsShader
 {
 public:
-	FRHIGeometryShader() : FRHIGraphicsShader(SF_Geometry) {}
+	FRHIGeometryShader() : FRHIGraphicsShader(RRT_GeometryShader, SF_Geometry) {}
 };
 
 class FRHIRayTracingShader : public FRHIShader
 {
 public:
-	explicit FRHIRayTracingShader(EShaderFrequency InFrequency) : FRHIShader(InFrequency) {}
+	explicit FRHIRayTracingShader(EShaderFrequency InFrequency) : FRHIShader(RRT_RayTracingShader, InFrequency) {}
 };
 
 class RHI_API FRHIComputeShader : public FRHIShader
 {
 public:
-	FRHIComputeShader() : FRHIShader(SF_Compute), Stats(nullptr) {}
+	FRHIComputeShader() : FRHIShader(RRT_ComputeShader, SF_Compute), Stats(nullptr) {}
 	
 	inline void SetStats(struct FPipelineStateStats* Ptr) { Stats = Ptr; }
 	void UpdateStats();
@@ -540,6 +586,8 @@ private:
 class FRHIGraphicsPipelineState : public FRHIResource 
 {
 public:
+	FRHIGraphicsPipelineState() : FRHIResource(RRT_GraphicsPipelineState) {}
+
 	inline void SetSortKey(uint64 InSortKey) { SortKey = InSortKey; }
 	inline uint64 GetSortKey() const { return SortKey; }
 
@@ -552,8 +600,16 @@ private:
 	FExclusiveDepthStencil DSMode;
 #endif
 };
-class FRHIComputePipelineState : public FRHIResource {};
-class FRHIRayTracingPipelineState : public FRHIResource {};
+class FRHIComputePipelineState : public FRHIResource
+{
+public:
+	FRHIComputePipelineState() : FRHIResource(RRT_ComputePipelineState) {}
+};
+class FRHIRayTracingPipelineState : public FRHIResource
+{
+public:
+	FRHIRayTracingPipelineState() : FRHIResource(RRT_RayTracingPipelineState) {}
+};
 
 //
 // Buffers
@@ -758,9 +814,12 @@ class FRHIUniformBuffer : public FRHIResource
 {
 public:
 
+	FRHIUniformBuffer() = delete;
+
 	/** Initialization constructor. */
 	FRHIUniformBuffer(const FRHIUniformBufferLayout& InLayout)
-	: Layout(&InLayout)
+	: FRHIResource(RRT_UniformBuffer)
+	, Layout(&InLayout)
 	, LayoutConstantBufferSize(InLayout.ConstantBufferSize)
 	{}
 
@@ -829,9 +888,12 @@ class FRHIBuffer : public FRHIResource
 #endif
 {
 public:
+	FRHIBuffer() : FRHIResource(RRT_Buffer) {}
+
 	/** Initialization constructor. */
 	FRHIBuffer(uint32 InSize, uint32 InUsage, uint32 InStride)
-		: Size(InSize)
+		: FRHIResource(RRT_Buffer)
+		, Size(InSize)
 		, Usage(InUsage)
 		, Stride(InStride)
 	{
@@ -846,13 +908,11 @@ public:
 	/** @return The usage flags used to create the buffer. */
 	uint32 GetUsage() const { return Usage; }
 
-protected:
-	FRHIBuffer()
-		: Size(0)
-		, Usage(0)
-		, Stride(0)
-	{}
+	void SetName(const FName& InName) { BufferName = InName; }
 
+	FName GetName() const { return BufferName; }
+
+protected:
 	void Swap(FRHIBuffer& Other)
 	{
 		::Swap(Stride, Other.Stride);
@@ -872,9 +932,10 @@ protected:
 	}
 
 private:
-	uint32 Size;
-	uint32 Usage;
-	uint32 Stride;
+	uint32 Size{};
+	uint32 Usage{};
+	uint32 Stride{};
+	FName BufferName;
 };
 
 UE_DEPRECATED(5.0, "FRHIIndexBuffer is deprecated, please use FRHIBuffer.")      typedef class FRHIBuffer FRHIIndexBuffer;
@@ -913,8 +974,9 @@ class RHI_API FRHITexture : public FRHIResource
 public:
 	
 	/** Initialization constructor. */
-	FRHITexture(uint32 InNumMips, uint32 InNumSamples, EPixelFormat InFormat, ETextureCreateFlags InFlags, FLastRenderTimeContainer* InLastRenderTime, const FClearValueBinding& InClearValue)
-		: ClearValue(InClearValue)
+	FRHITexture(ERHIResourceType InResourceType, uint32 InNumMips, uint32 InNumSamples, EPixelFormat InFormat, ETextureCreateFlags InFlags, FLastRenderTimeContainer* InLastRenderTime, const FClearValueBinding& InClearValue)
+		: FRHIResource(InResourceType)
+		, ClearValue(InClearValue)
 		, NumMips(InNumMips)
 		, NumSamples(InNumSamples)
 		, Format(InFormat)
@@ -996,8 +1058,6 @@ public:
 	/** @return Whether the texture is multi sampled. */
 	bool IsMultisampled() const { return NumSamples > 1; }		
 
-	FRHIResourceInfo ResourceInfo;
-
 	/** sets the last time this texture was cached in a resource table. */
 	FORCEINLINE_DEBUGGABLE void SetLastRenderTime(float InLastRenderTime)
 	{
@@ -1087,8 +1147,8 @@ class RHI_API FRHITexture2D : public FRHITexture
 public:
 	
 	/** Initialization constructor. */
-	FRHITexture2D(uint32 InSizeX,uint32 InSizeY,uint32 InNumMips,uint32 InNumSamples,EPixelFormat InFormat,ETextureCreateFlags InFlags, const FClearValueBinding& InClearValue)
-	: FRHITexture(InNumMips, InNumSamples, InFormat, InFlags, NULL, InClearValue)
+	FRHITexture2D(uint32 InSizeX,uint32 InSizeY,uint32 InNumMips,uint32 InNumSamples,EPixelFormat InFormat,ETextureCreateFlags InFlags, const FClearValueBinding& InClearValue, ERHIResourceType InResourceTypeOverride = RRT_None)
+	: FRHITexture(InResourceTypeOverride != RRT_None ? InResourceTypeOverride : RRT_Texture2D, InNumMips, InNumSamples, InFormat, InFlags, NULL, InClearValue)
 	, SizeX(InSizeX)
 	, SizeY(InSizeY)
 	{}
@@ -1124,7 +1184,7 @@ public:
 	
 	/** Initialization constructor. */
 	FRHITexture2DArray(uint32 InSizeX,uint32 InSizeY,uint32 InSizeZ,uint32 InNumMips,uint32 NumSamples, EPixelFormat InFormat,ETextureCreateFlags InFlags, const FClearValueBinding& InClearValue)
-	: FRHITexture2D(InSizeX, InSizeY, InNumMips,NumSamples,InFormat,InFlags, InClearValue)
+	: FRHITexture2D(InSizeX, InSizeY, InNumMips,NumSamples,InFormat,InFlags, InClearValue, RRT_Texture2DArray)
 	, SizeZ(InSizeZ)
 	{
 		check(InSizeZ != 0);
@@ -1165,7 +1225,7 @@ public:
 	
 	/** Initialization constructor. */
 	FRHITexture3D(uint32 InSizeX,uint32 InSizeY,uint32 InSizeZ,uint32 InNumMips,EPixelFormat InFormat,ETextureCreateFlags InFlags, const FClearValueBinding& InClearValue)
-	: FRHITexture(InNumMips,1,InFormat,InFlags,NULL, InClearValue)
+	: FRHITexture(RRT_Texture3D, InNumMips,1,InFormat,InFlags,NULL, InClearValue)
 	, SizeX(InSizeX)
 	, SizeY(InSizeY)
 	, SizeZ(InSizeZ)
@@ -1201,7 +1261,7 @@ public:
 	
 	/** Initialization constructor. */
 	FRHITextureCube(uint32 InSize,uint32 InNumMips,EPixelFormat InFormat,ETextureCreateFlags InFlags, const FClearValueBinding& InClearValue)
-	: FRHITexture(InNumMips,1,InFormat,InFlags,NULL, InClearValue)
+	: FRHITexture(RRT_TextureCube, InNumMips,1,InFormat,InFlags,NULL, InClearValue)
 	, Size(InSize)
 	{}
 	
@@ -1225,7 +1285,7 @@ class RHI_API FRHITextureReference final : public FRHITexture
 {
 public:
 	explicit FRHITextureReference(FLastRenderTimeContainer* InLastRenderTime)
-		: FRHITexture(0, 0, PF_Unknown, TexCreate_None, InLastRenderTime, FClearValueBinding())
+		: FRHITexture(RRT_TextureReference, 0, 0, PF_Unknown, TexCreate_None, InLastRenderTime, FClearValueBinding())
 	{
 		check(DefaultTexture);
 		ReferencedTexture = DefaultTexture;
@@ -1305,6 +1365,7 @@ private:
 class RHI_API FRHITimestampCalibrationQuery : public FRHIResource
 {
 public:
+	FRHITimestampCalibrationQuery() : FRHIResource(RRT_TimestampCalibrationQuery) {}
 	uint64 GPUMicroseconds[MAX_NUM_GPUS] = {};
 	uint64 CPUMicroseconds[MAX_NUM_GPUS] = {};
 };
@@ -1319,7 +1380,7 @@ public:
 class RHI_API FRHIGPUFence : public FRHIResource
 {
 public:
-	FRHIGPUFence(FName InName) : FenceName(InName) {}
+	FRHIGPUFence(FName InName) : FRHIResource(RRT_GPUFence), FenceName(InName) {}
 	virtual ~FRHIGPUFence() {}
 
 	virtual void Clear() = 0;
@@ -1360,7 +1421,11 @@ private:
 	uint32 InsertedFrameNumber;
 };
 
-class FRHIRenderQuery : public FRHIResource {};
+class FRHIRenderQuery : public FRHIResource
+{
+public:
+	FRHIRenderQuery() : FRHIResource(RRT_RenderQuery) {}
+};
 
 class FRHIRenderQueryPool;
 class RHI_API FRHIPooledRenderQuery
@@ -1394,6 +1459,7 @@ public:
 class RHI_API FRHIRenderQueryPool : public FRHIResource
 {
 public:
+	FRHIRenderQueryPool() : FRHIResource(RRT_RenderQueryPool) {}
 	virtual ~FRHIRenderQueryPool() {};
 	virtual FRHIPooledRenderQuery AllocateQuery() = 0;
 
@@ -1430,7 +1496,8 @@ class FRHIComputeFence final : public FRHIResource
 public:
 
 	FRHIComputeFence(FName InName)
-		: Name(InName)
+		: FRHIResource(RRT_ComputeFence)
+		, Name(InName)
 	{}
 
 	FORCEINLINE FName GetName() const
@@ -1454,6 +1521,8 @@ public:
 class FRHIViewport : public FRHIResource 
 {
 public:
+	FRHIViewport() : FRHIResource(RRT_Viewport) {}
+
 	/**
 	 * Returns access to the platform-specific native resource pointer.  This is designed to be used to provide plugins with access
 	 * to the underlying resource and should be used very carefully or not at all.
@@ -1514,13 +1583,19 @@ class FRHIUnorderedAccessView : public FRHIResource
 #if ENABLE_RHI_VALIDATION
 	, public RHIValidation::FUnorderedAccessView
 #endif
-{};
+{
+public:
+	FRHIUnorderedAccessView() : FRHIResource(RRT_UnorderedAccessView) {}
+};
 
 class FRHIShaderResourceView : public FRHIResource 
 #if ENABLE_RHI_VALIDATION
 	, public RHIValidation::FShaderResourceView
 #endif
-{};
+{
+public:
+	FRHIShaderResourceView() : FRHIResource(RRT_ShaderResourceView) {}
+};
 
 
 typedef TRefCountPtr<FRHISamplerState> FSamplerStateRHIRef;
@@ -1569,6 +1644,8 @@ class FRHIRayTracingAccelerationStructure
 	, public RHIValidation::FAccelerationStructureResource
 #endif
 {
+public:
+	FRHIRayTracingAccelerationStructure() : FRHIResource(RRT_RayTracingAccelerationStructure) {}
 };
 
 /** Bottom level ray tracing acceleration structure (contains triangles). */
@@ -1589,7 +1666,8 @@ class RHI_API FRHIStagingBuffer : public FRHIResource
 {
 public:
 	FRHIStagingBuffer()
-		: bIsLocked(false)
+		: FRHIResource(RRT_StagingBuffer)
+		, bIsLocked(false)
 	{}
 
 	virtual ~FRHIStagingBuffer() {}
@@ -1901,7 +1979,7 @@ public:
 class FRHICustomPresent : public FRHIResource
 {
 public:
-	FRHICustomPresent() {}
+	FRHICustomPresent() : FRHIResource(RRT_CustomPresent) {}
 	
 	virtual ~FRHICustomPresent() {} // should release any references to D3D resources.
 	
@@ -2469,7 +2547,7 @@ protected:
 class FRHIShaderLibrary : public FRHIResource
 {
 public:
-	FRHIShaderLibrary(EShaderPlatform InPlatform, FString const& InName) : Platform(InPlatform), LibraryName(InName), LibraryId(GetTypeHash(InName)) {}
+	FRHIShaderLibrary(EShaderPlatform InPlatform, FString const& InName) : FRHIResource(RRT_ShaderLibrary), Platform(InPlatform), LibraryName(InName), LibraryId(GetTypeHash(InName)) {}
 	virtual ~FRHIShaderLibrary() {}
 	
 	FORCEINLINE EShaderPlatform GetPlatform(void) const { return Platform; }
@@ -2501,7 +2579,7 @@ typedef TRefCountPtr<FRHIShaderLibrary>	FRHIShaderLibraryRef;
 class FRHIPipelineBinaryLibrary : public FRHIResource
 {
 public:
-	FRHIPipelineBinaryLibrary(EShaderPlatform InPlatform, FString const& FilePath) : Platform(InPlatform) {}
+	FRHIPipelineBinaryLibrary(EShaderPlatform InPlatform, FString const& FilePath) : FRHIResource(RRT_PipelineBinaryLibrary), Platform(InPlatform) {}
 	virtual ~FRHIPipelineBinaryLibrary() {}
 	
 	FORCEINLINE EShaderPlatform GetPlatform(void) const { return Platform; }
diff --git a/Engine/Source/Runtime/RHI/RHI.Build.cs b/Engine/Source/Runtime/RHI/RHI.Build.cs
--- a/Engine/Source/Runtime/RHI/RHI.Build.cs
+++ b/Engine/Source/Runtime/RHI/RHI.Build.cs
@@ -23,6 +23,11 @@ public class RHI : ModuleRules
 					PublicDefinitions.Add("RHI_WANT_BREADCRUMB_EVENTS=1");
 				}
 
+				if (Target.Configuration != UnrealTargetConfiguration.Shipping)
+                {
+					PublicDefinitions.Add("RHI_WANT_RESOURCE_INFO=1");
+                }
+
 				// UEBuildAndroid.cs adds VulkanRHI for Android builds if it is enabled
 				if (Target.Platform == UnrealTargetPlatform.Win64)
 				{
diff --git a/Engine/Source/Runtime/VulkanRHI/Private/VulkanRHI.cpp b/Engine/Source/Runtime/VulkanRHI/Private/VulkanRHI.cpp
--- a/Engine/Source/Runtime/VulkanRHI/Private/VulkanRHI.cpp
+++ b/Engine/Source/Runtime/VulkanRHI/Private/VulkanRHI.cpp
@@ -1229,84 +1229,6 @@ void FVulkanDynamicRHI::RHICopySubTextureRegion(FRHITexture2D* SourceTexture, FR
 	RHIGetDefaultContext()->RHICopyTexture(SourceTexture, DestinationTexture, CopyInfo);
 }
 
-
-FVulkanBuffer::FVulkanBuffer(FVulkanDevice& InDevice, uint32 InSize, VkFlags InUsage, VkMemoryPropertyFlags InMemPropertyFlags, bool bInAllowMultiLock, const char* File, int32 Line) :
-	Device(InDevice),
-	Buf(VK_NULL_HANDLE),
-	Allocation(nullptr),
-	Size(InSize),
-	Usage(InUsage),
-	BufferPtr(nullptr),
-	bAllowMultiLock(bInAllowMultiLock),
-	LockStack(0)
-{
-	VkBufferCreateInfo BufInfo;
-	ZeroVulkanStruct(BufInfo, VK_STRUCTURE_TYPE_BUFFER_CREATE_INFO);
-	BufInfo.size = Size;
-	BufInfo.usage = Usage;
-	VERIFYVULKANRESULT_EXPANDED(VulkanRHI::vkCreateBuffer(Device.GetInstanceHandle(), &BufInfo, VULKAN_CPU_ALLOCATOR, &Buf));
-
-	VkMemoryRequirements MemoryRequirements;
-	VulkanRHI::vkGetBufferMemoryRequirements(Device.GetInstanceHandle(), Buf, &MemoryRequirements);
-
-	Allocation = InDevice.GetDeviceMemoryManager().Alloc(false, MemoryRequirements.size, MemoryRequirements.memoryTypeBits, InMemPropertyFlags, nullptr, VULKAN_MEMORY_MEDIUM_PRIORITY, File ? File : __FILE__, Line ? Line : __LINE__);
-	check(Allocation);
-	VERIFYVULKANRESULT_EXPANDED(VulkanRHI::vkBindBufferMemory(Device.GetInstanceHandle(), Buf, Allocation->GetHandle(), 0));
-}
-
-FVulkanBuffer::~FVulkanBuffer()
-{
-	// The buffer should be unmapped
-	check(BufferPtr == nullptr);
-
-	Device.GetDeferredDeletionQueue().EnqueueResource(FDeferredDeletionQueue2::EType::Buffer, Buf);
-	Buf = VK_NULL_HANDLE;
-
-	Device.GetDeviceMemoryManager().Free(Allocation);
-	Allocation = nullptr;
-}
-
-void* FVulkanBuffer::Lock(uint32 InSize, uint32 InOffset)
-{
-	check(InSize + InOffset <= Size);
-
-	uint32 BufferPtrOffset = 0;
-	if (bAllowMultiLock)
-	{
-		if (LockStack == 0)
-		{
-			// lock the whole range
-			BufferPtr = Allocation->Map(GetSize(), 0);
-		}
-		// offset the whole range by the requested offset
-		BufferPtrOffset = InOffset;
-		LockStack++;
-	}
-	else
-	{
-		check(BufferPtr == nullptr);
-		BufferPtr = Allocation->Map(InSize, InOffset);
-	}
-
-	return (uint8*)BufferPtr + BufferPtrOffset;
-}
-
-void FVulkanBuffer::Unlock()
-{
-	// The buffer should be mapped, before it can be unmapped
-	check(BufferPtr != nullptr);
-
-	// for multi-lock, if not down to 0, do nothing
-	if (bAllowMultiLock && --LockStack > 0)
-	{
-		return;
-	}
-
-	Allocation->Unmap();
-	BufferPtr = nullptr;
-}
-
-
 FVulkanDescriptorSetsLayout::FVulkanDescriptorSetsLayout(FVulkanDevice* InDevice) :
 	Device(InDevice)
 {
@@ -1523,33 +1445,6 @@ void FVulkanDescriptorSetsLayout::Compile(FVulkanDescriptorSetLayoutMap& DSetLay
 }
 
 
-void FVulkanBufferView::Create(FVulkanBuffer& Buffer, EPixelFormat Format, uint32 InOffset, uint32 InSize)
-{
-	Offset = InOffset;
-	Size = InSize;
-	check(Format != PF_Unknown);
-	VkFormat BufferFormat = GVulkanBufferFormat[Format];
-	check(BufferFormat != VK_FORMAT_UNDEFINED);
-
-	VkBufferViewCreateInfo ViewInfo;
-	ZeroVulkanStruct(ViewInfo, VK_STRUCTURE_TYPE_BUFFER_VIEW_CREATE_INFO);
-	ViewInfo.buffer = Buffer.GetBufferHandle();
-	ViewInfo.format = BufferFormat;
-	ViewInfo.offset = Offset;
-	ViewInfo.range = Size;
-	Flags = Buffer.GetFlags() & VK_BUFFER_USAGE_UNIFORM_TEXEL_BUFFER_BIT;
-	check(Flags);
-
-	VERIFYVULKANRESULT(VulkanRHI::vkCreateBufferView(GetParent()->GetInstanceHandle(), &ViewInfo, VULKAN_CPU_ALLOCATOR, &View));
-	
-	if (UseVulkanDescriptorCache())
-	{
-		ViewId = ++GVulkanBufferViewHandleIdCounter;
-	}
-
-	INC_DWORD_STAT(STAT_VulkanNumBufferViews);
-}
-
 void FVulkanBufferView::Create(FVulkanResourceMultiBuffer* Buffer, EPixelFormat Format, uint32 InOffset, uint32 InSize)
 {
 	check(Format != PF_Unknown);
diff --git a/Engine/Source/Runtime/VulkanRHI/Private/VulkanResources.h b/Engine/Source/Runtime/VulkanRHI/Private/VulkanResources.h
--- a/Engine/Source/Runtime/VulkanRHI/Private/VulkanResources.h
+++ b/Engine/Source/Runtime/VulkanRHI/Private/VulkanResources.h
@@ -18,7 +18,6 @@
 class FVulkanDevice;
 class FVulkanQueue;
 class FVulkanCmdBuffer;
-class FVulkanBuffer;
 class FVulkanBufferCPU;
 struct FVulkanTextureBase;
 class FVulkanTexture2D;
@@ -333,20 +332,6 @@ public:
 	void EvictSurface(FVulkanDevice& InDevice);
 
 
-#if 0
-	/**
-	 * Locks one of the texture's mip-maps.
-	 * @param ArrayIndex Index of the texture array/face in the form Index*6+Face
-	 * @return A pointer to the specified texture data.
-	 */
-	void* Lock(uint32 MipIndex, uint32 ArrayIndex, EResourceLockMode LockMode, uint32& DestStride);
-
-	/** Unlocks a previously locked mip-map.
-	 * @param ArrayIndex Index of the texture array/face in the form Index*6+Face
-	 */
-	void Unlock(uint32 MipIndex, uint32 ArrayIndex);
-#endif
-
 	/**
 	 * Returns how much memory is used by the surface
 	 */
@@ -911,7 +896,8 @@ public:
 struct FVulkanBufferView : public FRHIResource, public VulkanRHI::FDeviceChild
 {
 	FVulkanBufferView(FVulkanDevice* InDevice)
-		: VulkanRHI::FDeviceChild(InDevice)
+		: FRHIResource(RRT_None)
+		, VulkanRHI::FDeviceChild(InDevice)
 		, View(VK_NULL_HANDLE)
 		, ViewId(0)
 		, Flags(0)
@@ -925,7 +911,6 @@ struct FVulkanBufferView : public FRHIResource, public VulkanRHI::FDeviceChild
 		Destroy();
 	}
 
-	void Create(FVulkanBuffer& Buffer, EPixelFormat Format, uint32 InOffset, uint32 InSize);
 	void Create(FVulkanResourceMultiBuffer* Buffer, EPixelFormat Format, uint32 InOffset, uint32 InSize);
 	void Create(VkFormat Format, FVulkanResourceMultiBuffer* Buffer, uint32 InOffset, uint32 InSize);
 	void Destroy();
@@ -937,36 +922,6 @@ struct FVulkanBufferView : public FRHIResource, public VulkanRHI::FDeviceChild
 	uint32 Size;
 };
 
-class FVulkanBuffer : public FRHIResource
-{
-public:
-	FVulkanBuffer(FVulkanDevice& Device, uint32 InSize, VkFlags InUsage, VkMemoryPropertyFlags InMemPropertyFlags, bool bAllowMultiLock, const char* File, int32 Line);
-	virtual ~FVulkanBuffer();
-
-	inline VkBuffer GetBufferHandle() const { return Buf; }
-
-	inline uint32 GetSize() const { return Size; }
-
-	void* Lock(uint32 InSize, uint32 InOffset = 0);
-
-	void Unlock();
-
-	inline VkFlags GetFlags() const { return Usage; }
-
-private:
-	FVulkanDevice& Device;
-	VkBuffer Buf;
-	VulkanRHI::FDeviceMemoryAllocation* Allocation;
-	uint32 Size;
-	VkFlags Usage;
-
-	void* BufferPtr;	
-	VkMappedMemoryRange MappedRange;
-
-	bool bAllowMultiLock;
-	int32 LockStack;
-};
-
 struct FVulkanRingBuffer : public FVulkanEvictable, public VulkanRHI::FDeviceChild
 {
 	virtual void Evict(FVulkanDevice& Device);
diff --git a/Engine/Source/Runtime/VulkanRHI/Private/VulkanTexture.cpp b/Engine/Source/Runtime/VulkanRHI/Private/VulkanTexture.cpp
--- a/Engine/Source/Runtime/VulkanRHI/Private/VulkanTexture.cpp
+++ b/Engine/Source/Runtime/VulkanRHI/Private/VulkanTexture.cpp
@@ -1195,121 +1195,6 @@ uint64 FVulkanSurface::GetAllocationOffset() const
 	}
 }
 
-
-#if 0
-void* FVulkanSurface::Lock(uint32 MipIndex, uint32 ArrayIndex, EResourceLockMode LockMode, uint32& DestStride)
-{
-	DestStride = 0;
-
-	check((MemProps & VK_MEMORY_PROPERTY_HOST_VISIBLE_BIT) == 0 ? Tiling == VK_IMAGE_TILING_OPTIMAL : Tiling == VK_IMAGE_TILING_LINEAR);
-
-	// Verify all buffers are unmapped
-	auto& Data = MipMapMapping.FindOrAdd(MipIndex);
-	checkf(Data == nullptr, TEXT("The buffer needs to be unmapped, before it can be mapped"));
-
-	// Get the layout of the subresource
-	VkImageSubresource ImageSubResource;
-	FMemory::Memzero(ImageSubResource);
-
-	ImageSubResource.aspectMask = GetAspectMask();
-	ImageSubResource.mipLevel = MipIndex;
-	ImageSubResource.arrayLayer = ArrayIndex;
-
-	// Get buffer size
-	// Pitch can be only retrieved from linear textures.
-	VkSubresourceLayout SubResourceLayout;
-	VulkanRHI::vkGetImageSubresourceLayout(Device->GetInstanceHandle(), Image, &ImageSubResource, &SubResourceLayout);
-
-	// Set linear row-pitch
-	GetMipStride(MipIndex, DestStride);
-
-	if(Tiling == VK_IMAGE_TILING_LINEAR)
-	{
-		// Verify pitch if linear
-		check(DestStride == SubResourceLayout.rowPitch);
-
-		// Map buffer to a pointer
-		Data = Allocation->Map(SubResourceLayout.size, SubResourceLayout.offset);
-		return Data;
-	}
-
-	// From here on, the code is dedicated to optimal textures
-
-	// Verify all buffers are unmapped
-	TRefCountPtr<FVulkanBuffer>& LinearBuffer = MipMapBuffer.FindOrAdd(MipIndex);
-	checkf(LinearBuffer == nullptr, TEXT("The buffer needs to be unmapped, before it can be mapped"));
-
-	// Create intermediate buffer which is going to be used to perform buffer to image copy
-	// The copy buffer is always one face and single mip level.
-	const uint32 Layers = 1;
-	const uint32 Bytes = SubResourceLayout.size * Layers;
-
-	VkBufferUsageFlags Usage = 0;
-	Usage |= VK_BUFFER_USAGE_TRANSFER_SRC_BIT;
-
-	VkMemoryPropertyFlags Flags = 0;
-	Flags |= VK_MEMORY_PROPERTY_HOST_VISIBLE_BIT | VK_MEMORY_PROPERTY_HOST_COHERENT;
-
-	LinearBuffer = new FVulkanBuffer(*Device, Bytes, Usage, Flags, false, __FILE__, __LINE__);
-
-	void* DataPtr = LinearBuffer->Lock(Bytes);
-	check(DataPtr);
-
-	return DataPtr;
-}
-
-void FVulkanSurface::Unlock(uint32 MipIndex, uint32 ArrayIndex)
-{
-	check((MemProps & VK_MEMORY_PROPERTY_HOST_VISIBLE_BIT) == 0 ? Tiling == VK_IMAGE_TILING_OPTIMAL : Tiling == VK_IMAGE_TILING_LINEAR);
-
-	if(Tiling == VK_IMAGE_TILING_LINEAR)
-	{
-		void*& Data = MipMapMapping.FindOrAdd(MipIndex);
-		checkf(Data != nullptr, TEXT("The buffer needs to be mapped, before it can be unmapped"));
-
-		Allocation->Unmap();
-		Data = nullptr;
-		return;
-	}
-
-	TRefCountPtr<FVulkanBuffer>& LinearBuffer = MipMapBuffer.FindOrAdd(MipIndex);
-	checkf(LinearBuffer != nullptr, TEXT("The buffer needs to be mapped, before it can be unmapped"));
-	LinearBuffer->Unlock();
-
-	VkImageSubresource ImageSubResource;
-	FMemory::Memzero(ImageSubResource);
-	ImageSubResource.aspectMask = GetAspectMask();
-	ImageSubResource.mipLevel = MipIndex;
-	ImageSubResource.arrayLayer = ArrayIndex;
-
-	VkSubresourceLayout SubResourceLayout;
-
-	VulkanRHI::vkGetImageSubresourceLayout(Device->GetInstanceHandle(), Image, &ImageSubResource, &SubResourceLayout);
-
-	VkBufferImageCopy Region;
-	FMemory::Memzero(Region);
-	Region.bufferOffset = 0;
-	Region.bufferRowLength = (Width >> MipIndex);
-	Region.bufferImageHeight = (Height >> MipIndex);
-
-	// The data/image is always parsed per one face.
-	// Meaning that a cubemap will have atleast 6 locks/unlocks
-	Region.imageSubresource.baseArrayLayer = ArrayIndex;	// Layer/face copy destination
-	Region.imageSubresource.layerCount = 1;	// Indicates number of arrays in the buffer, this is also the amount of "faces/layers" to be copied
-	Region.imageSubresource.aspectMask = GetAspectMask();
-	Region.imageSubresource.mipLevel = MipIndex;
-
-	Region.imageExtent.width = Region.bufferRowLength;
-	Region.imageExtent.height = Region.bufferImageHeight;
-	Region.imageExtent.depth = 1;
-
-	LinearBuffer->CopyTo(*this, Region, nullptr);
-
-	// Release buffer
-	LinearBuffer = nullptr;
-}
-#endif
-
 void FVulkanSurface::GetMipStride(uint32 MipIndex, uint32& Stride)
 {
 	// Calculate the width of the MipMap.
diff --git a/Engine/Source/Runtime/Windows/D3D11RHI/Private/D3D11Buffer.cpp b/Engine/Source/Runtime/Windows/D3D11RHI/Private/D3D11Buffer.cpp
--- a/Engine/Source/Runtime/Windows/D3D11RHI/Private/D3D11Buffer.cpp
+++ b/Engine/Source/Runtime/Windows/D3D11RHI/Private/D3D11Buffer.cpp
@@ -112,7 +112,12 @@ FBufferRHIRef FD3D11DynamicRHI::RHICreateBuffer(uint32 Size, EBufferUsageFlags U
 		CreateInfo.ResourceArray->Discard();
 	}
 
-	return new FD3D11Buffer(BufferResource, Size, Usage, Stride);
+	FD3D11Buffer* NewBuffer = new FD3D11Buffer(BufferResource, Size, Usage, Stride);
+	if (CreateInfo.DebugName)
+	{
+		NewBuffer->SetName(CreateInfo.DebugName);
+	}
+	return NewBuffer;
 }
 
 FBufferRHIRef FD3D11DynamicRHI::CreateBuffer_RenderThread(class FRHICommandListImmediate& RHICmdList, uint32 Size, EBufferUsageFlags Usage, uint32 Stride, ERHIAccess ResourceState, FRHIResourceCreateInfo& CreateInfo)
diff --git a/Engine/Source/Runtime/Windows/D3D11RHI/Private/D3D11Texture.cpp b/Engine/Source/Runtime/Windows/D3D11RHI/Private/D3D11Texture.cpp
--- a/Engine/Source/Runtime/Windows/D3D11RHI/Private/D3D11Texture.cpp
+++ b/Engine/Source/Runtime/Windows/D3D11RHI/Private/D3D11Texture.cpp
@@ -1051,6 +1051,11 @@ TD3D11Texture2D<BaseResourceType>* FD3D11DynamicRHI::CreateD3D11Texture2D(uint32
 #endif
 		);
 
+	if (CreateInfo.DebugName)
+	{
+		Texture2D->SetName(CreateInfo.DebugName);
+	}
+
 	D3D11TextureAllocated(*Texture2D);
 	
 #if !PLATFORM_HOLOLENS
@@ -1190,6 +1195,11 @@ FD3D11Texture3D* FD3D11DynamicRHI::CreateD3D11Texture3D(uint32 SizeX,uint32 Size
 	RenderTargetViews.Add(RenderTargetView);
 	FD3D11Texture3D* Texture3D = new FD3D11Texture3D(this,TextureResource,ShaderResourceView,RenderTargetViews,SizeX,SizeY,SizeZ,NumMips,(EPixelFormat)Format,Flags, CreateInfo.ClearValueBinding);
 
+	if (CreateInfo.DebugName)
+	{
+		Texture3D->SetName(CreateInfo.DebugName);
+	}
+
 	D3D11TextureAllocated(*Texture3D);
 #if !PLATFORM_HOLOLENS
 	if (IsRHIDeviceNVIDIA() && (Flags & TexCreate_AFRManual))
@@ -1407,10 +1417,6 @@ FTexture3DRHIRef FD3D11DynamicRHI::RHICreateTexture3D_RenderThread(
 
 void FD3D11DynamicRHI::RHIGetResourceInfo(FRHITexture* Ref, FRHIResourceInfo& OutInfo)
 {
-	if(Ref)
-	{
-		OutInfo = Ref->ResourceInfo;
-	}
 }
 
 FShaderResourceViewRHIRef FD3D11DynamicRHI::RHICreateShaderResourceView(FRHITexture* TextureRHI, const FRHITextureSRVCreateInfo& CreateInfo)
diff --git a/Engine/Source/Runtime/Windows/D3D11RHI/Public/D3D11Resources.h b/Engine/Source/Runtime/Windows/D3D11RHI/Public/D3D11Resources.h
--- a/Engine/Source/Runtime/Windows/D3D11RHI/Public/D3D11Resources.h
+++ b/Engine/Source/Runtime/Windows/D3D11RHI/Public/D3D11Resources.h
@@ -361,6 +361,18 @@ public:
 
 	virtual ~TD3D11Texture2D();
 
+	// FRHIResource overrides
+#if RHI_WANT_RESOURCE_INFO
+	bool GetResourceInfo(FRHIResourceInfo& OutResourceInfo) const override
+	{
+		OutResourceInfo = FRHIResourceInfo{};
+		OutResourceInfo.Name = this->GetName().ToString();
+		OutResourceInfo.Type = this->GetType();
+		OutResourceInfo.VRamAllocation.AllocationSize = GetMemorySize();
+		return true;
+	}
+#endif
+
 	/**
 	 * Locks one of the texture's mip-maps.
 	 * @return A pointer to the specified texture data.
@@ -453,6 +465,18 @@ public:
 	}
 
 	virtual ~FD3D11Texture3D();
+
+	// FRHIResource overrides
+#if RHI_WANT_RESOURCE_INFO
+	bool GetResourceInfo(FRHIResourceInfo& OutResourceInfo) const override
+	{
+		OutResourceInfo = FRHIResourceInfo{};
+		OutResourceInfo.Name = GetName().ToString();
+		OutResourceInfo.Type = GetType();
+		OutResourceInfo.VRamAllocation.AllocationSize = GetMemorySize();
+		return true;
+	}
+#endif
 	
 	// Accessors.
 	ID3D11Texture3D* GetResource() const { return (ID3D11Texture3D*)FD3D11TextureBase::GetResource(); }
@@ -612,7 +636,20 @@ public:
 	FD3D11Buffer(ID3D11Buffer* InResource, uint32 InSize, uint32 InUsage, uint32 InStride)
 	: FRHIBuffer(InSize, InUsage, InStride)
 	, Resource(InResource)
-	{}
+	{
+	}
+
+	// FRHIResource overrides
+#if RHI_WANT_RESOURCE_INFO
+	bool GetResourceInfo(FRHIResourceInfo& OutResourceInfo) const override
+	{
+		OutResourceInfo = FRHIResourceInfo{};
+		OutResourceInfo.Name = GetName().ToString();
+		OutResourceInfo.Type = GetType();
+		OutResourceInfo.VRamAllocation.AllocationSize = GetSize();
+		return true;
+	}
+#endif
 
 	virtual ~FD3D11Buffer()
 	{
