Switching FRHIResourceInfo::Name to be FName instead of FString to prevent excessive allocations.
Adding usage text to rhi.DumpResourceMemory.
Tweaking rhi.DumpResourceMemory output to be more explicit with the number of items shown.

diff --git a/Engine/Source/Runtime/D3D12RHI/Private/D3D12Texture.h b/Engine/Source/Runtime/D3D12RHI/Private/D3D12Texture.h
--- a/Engine/Source/Runtime/D3D12RHI/Private/D3D12Texture.h
+++ b/Engine/Source/Runtime/D3D12RHI/Private/D3D12Texture.h
@@ -293,7 +293,7 @@ public:
 	bool GetResourceInfo(FRHIResourceInfo& OutResourceInfo) const override
 	{
 		OutResourceInfo = FRHIResourceInfo{};
-		OutResourceInfo.Name = this->GetName().ToString();
+		OutResourceInfo.Name = this->GetName();
 		OutResourceInfo.Type = this->GetType();
 		OutResourceInfo.VRamAllocation.AllocationSize = GetMemorySize();
 		return true;
@@ -435,7 +435,7 @@ public:
 	bool GetResourceInfo(FRHIResourceInfo& OutResourceInfo) const override
 	{
 		OutResourceInfo = FRHIResourceInfo{};
-		OutResourceInfo.Name = GetName().ToString();
+		OutResourceInfo.Name = GetName();
 		OutResourceInfo.Type = GetType();
 		OutResourceInfo.VRamAllocation.AllocationSize = ResourceLocation.GetSize();
 		return true;
diff --git a/Engine/Source/Runtime/D3D12RHI/Public/D3D12Resources.h b/Engine/Source/Runtime/D3D12RHI/Public/D3D12Resources.h
--- a/Engine/Source/Runtime/D3D12RHI/Public/D3D12Resources.h
+++ b/Engine/Source/Runtime/D3D12RHI/Public/D3D12Resources.h
@@ -930,7 +930,7 @@ public:
 	bool GetResourceInfo(FRHIResourceInfo& OutResourceInfo) const override
 	{
 		OutResourceInfo = FRHIResourceInfo{};
-		OutResourceInfo.Name = GetName().ToString();
+		OutResourceInfo.Name = GetName();
 		OutResourceInfo.Type = GetType();
 		OutResourceInfo.VRamAllocation.AllocationSize = ResourceLocation.GetSize();
 		return true;
diff --git a/Engine/Source/Runtime/RHI/Private/RHI.cpp b/Engine/Source/Runtime/RHI/Private/RHI.cpp
--- a/Engine/Source/Runtime/RHI/Private/RHI.cpp
+++ b/Engine/Source/Runtime/RHI/Private/RHI.cpp
@@ -11,6 +11,7 @@
 #include "Misc/MessageDialog.h"
 #include "RHIShaderFormatDefinitions.inl"
 #include "ProfilingDebugging/CsvProfiler.h"
+#include "String/Find.h"
 #include "String/LexFromString.h"
 #include "String/ParseTokens.h"
 #include "Misc/BufferedOutputDevice.h"
@@ -777,7 +778,8 @@ static FAutoConsoleCommandWithWorldArgsAndOutputDevice GDumpRHIResourceCountsCmd
 
 static FAutoConsoleCommandWithWorldArgsAndOutputDevice GDumpRHIResourceMemoryCmd(
 	TEXT("rhi.DumpResourceMemory"),
-	TEXT("Dumps RHI resource memory stats to the log"),
+	TEXT("Dumps RHI resource memory stats to the log\n")
+	TEXT("Usage: rhi.DumpResourceMemory [<Number To Show>] [all] [Name=<Filter Text>] [Type=<RHI Resource Type>] [csv]"),
 	FConsoleCommandWithWorldArgsAndOutputDeviceDelegate::CreateStatic([](const TArray<FString>& Args, UWorld*, FOutputDevice& OutputDevice)
 {
 	FString NameFilter;
@@ -813,11 +815,16 @@ static FAutoConsoleCommandWithWorldArgsAndOutputDevice GDumpRHIResourceMemoryCmd
 		}
 	}
 
+	TCHAR ResourceNameBuffer[FName::StringBufferSize];
+
 	auto ShouldIncludeResource = [&](const FRHIResource* Resource, const FRHIResourceInfo& ResourceInfo) -> bool
 	{
-		if (!NameFilter.IsEmpty() && ResourceInfo.Name.Find(NameFilter, ESearchCase::IgnoreCase) == INDEX_NONE)
+		if (!NameFilter.IsEmpty())
 		{
-			return false;
+			if (ResourceInfo.Name.ToString(ResourceNameBuffer) == 0 || UE::String::FindFirst(ResourceNameBuffer, *NameFilter, ESearchCase::IgnoreCase) == INDEX_NONE)
+			{
+				return false;
+			}
 		}
 		if (TypeFilter != RRT_None)
 		{
@@ -872,6 +879,7 @@ static FAutoConsoleCommandWithWorldArgsAndOutputDevice GDumpRHIResourceMemoryCmd
 		}
 	}
 
+	const int32 NumberOfResourcesBeforeNumberFilter = Resources.Num();
 	if (NumberOfResourcesToShow < 0 || NumberOfResourcesToShow > Resources.Num())
 	{
 		NumberOfResourcesToShow = Resources.Num();
@@ -891,15 +899,11 @@ static FAutoConsoleCommandWithWorldArgsAndOutputDevice GDumpRHIResourceMemoryCmd
 	}
 	else
 	{
-		BufferedOutput.CategorizedLogf(CategoryName, ELogVerbosity::Log, TEXT("Tracked RHIResources"));
+		BufferedOutput.CategorizedLogf(CategoryName, ELogVerbosity::Log, TEXT("Tracked RHIResources (%d total with info, %d total tracked)"), TotalResourcesWithInfo, TotalTrackedResources);
 
-		if (NumberOfResourcesToShow != Resources.Num())
+		if (NumberOfResourcesToShow != NumberOfResourcesBeforeNumberFilter)
 		{
-			BufferedOutput.CategorizedLogf(CategoryName, ELogVerbosity::Log, TEXT("Showing %d out of %d with info and %d tracked"), NumberOfResourcesToShow, TotalResourcesWithInfo, TotalTrackedResources);
-		}
-		else
-		{
-			BufferedOutput.CategorizedLogf(CategoryName, ELogVerbosity::Log, TEXT("Showing %d with info and %d tracked"), TotalResourcesWithInfo, TotalTrackedResources);
+			BufferedOutput.CategorizedLogf(CategoryName, ELogVerbosity::Log, TEXT("Showing %d of %d matched resources"), NumberOfResourcesToShow, NumberOfResourcesBeforeNumberFilter);
 		}
 	}
 
@@ -910,20 +914,22 @@ static FAutoConsoleCommandWithWorldArgsAndOutputDevice GDumpRHIResourceMemoryCmd
 		if (Index < NumberOfResourcesToShow)
 		{
 			const FRHIResourceInfo& ResourceInfo = Resources[Index].ResourceInfo;
+
+			ResourceInfo.Name.ToString(ResourceNameBuffer);
 			const TCHAR* ResourceType = StringFromRHIResourceType(ResourceInfo.Type);
 			const int64 SizeInBytes = ResourceInfo.VRamAllocation.AllocationSize;
 
 			if (bUseCSVOutput)
 			{
 				BufferedOutput.Logf(ELogVerbosity::Log, TEXT("\"%s\",\"%s\",\"%.3f\""),
-					*ResourceInfo.Name,
+					ResourceNameBuffer,
 					ResourceType,
 					SizeInBytes / double(1 << 20));
 			}
 			else
 			{
 				BufferedOutput.CategorizedLogf(CategoryName, ELogVerbosity::Log, TEXT("Name: %s - Type: %s - Size: %.3f MB"),
-					*ResourceInfo.Name,
+					ResourceNameBuffer,
 					ResourceType,
 					SizeInBytes / double(1 << 20));
 			}
@@ -937,10 +943,10 @@ static FAutoConsoleCommandWithWorldArgsAndOutputDevice GDumpRHIResourceMemoryCmd
 
 	if (!bUseCSVOutput)
 	{
-		if (NumberOfResourcesToShow != Resources.Num())
+		if (NumberOfResourcesToShow != TotalResourcesWithInfo)
 		{
 			BufferedOutput.CategorizedLogf(CategoryName, ELogVerbosity::Log, TEXT("Shown %d entries. Size: %.3f MB (%.2f%% of total)"),
-				NumberOfResourcesToShow, TotalSizeF, 100.0 * ShownSizeF / TotalSizeF);
+				NumberOfResourcesToShow, ShownSizeF, 100.0 * ShownSizeF / TotalSizeF);
 		}
 
 		BufferedOutput.CategorizedLogf(CategoryName, ELogVerbosity::Log, TEXT("Total tracked resource size: %.3f MB"), TotalSizeF);
diff --git a/Engine/Source/Runtime/RHI/Public/RHI.h b/Engine/Source/Runtime/RHI/Public/RHI.h
--- a/Engine/Source/Runtime/RHI/Public/RHI.h
+++ b/Engine/Source/Runtime/RHI/Public/RHI.h
@@ -1433,7 +1433,7 @@ struct FVRamAllocation
 
 struct FRHIResourceInfo
 {
-	FString Name;
+	FName Name;
 	ERHIResourceType Type{ RRT_None };
 	FVRamAllocation VRamAllocation;
 };
diff --git a/Engine/Source/Runtime/Windows/D3D11RHI/Public/D3D11Resources.h b/Engine/Source/Runtime/Windows/D3D11RHI/Public/D3D11Resources.h
--- a/Engine/Source/Runtime/Windows/D3D11RHI/Public/D3D11Resources.h
+++ b/Engine/Source/Runtime/Windows/D3D11RHI/Public/D3D11Resources.h
@@ -366,7 +366,7 @@ public:
 	bool GetResourceInfo(FRHIResourceInfo& OutResourceInfo) const override
 	{
 		OutResourceInfo = FRHIResourceInfo{};
-		OutResourceInfo.Name = this->GetName().ToString();
+		OutResourceInfo.Name = this->GetName();
 		OutResourceInfo.Type = this->GetType();
 		OutResourceInfo.VRamAllocation.AllocationSize = GetMemorySize();
 		return true;
@@ -471,7 +471,7 @@ public:
 	bool GetResourceInfo(FRHIResourceInfo& OutResourceInfo) const override
 	{
 		OutResourceInfo = FRHIResourceInfo{};
-		OutResourceInfo.Name = GetName().ToString();
+		OutResourceInfo.Name = GetName();
 		OutResourceInfo.Type = GetType();
 		OutResourceInfo.VRamAllocation.AllocationSize = GetMemorySize();
 		return true;
@@ -644,7 +644,7 @@ public:
 	bool GetResourceInfo(FRHIResourceInfo& OutResourceInfo) const override
 	{
 		OutResourceInfo = FRHIResourceInfo{};
-		OutResourceInfo.Name = GetName().ToString();
+		OutResourceInfo.Name = GetName();
 		OutResourceInfo.Type = GetType();
 		OutResourceInfo.VRamAllocation.AllocationSize = GetSize();
 		return true;
