* Fixed buffer overrun crash with FGpuProfiler
* Added CVar (r.GpuProfilerMaxEventBufferSizeKB) to control the GPU profiler scratch buffer's size (dynamically tweakable)

diff --git a/Engine/Source/Runtime/RHI/Private/GpuProfilerTrace.cpp b/Engine/Source/Runtime/RHI/Private/GpuProfilerTrace.cpp
--- a/Engine/Source/Runtime/RHI/Private/GpuProfilerTrace.cpp
+++ b/Engine/Source/Runtime/RHI/Private/GpuProfilerTrace.cpp
@@ -14,7 +14,12 @@
 namespace GpuProfilerTrace
 {
 
-static const uint32 MaxEventBufferSize = 16 << 10;
+static TAutoConsoleVariable<int32> CVarGpuProfilerMaxEventBufferSizeKB(
+	TEXT("r.GpuProfilerMaxEventBufferSizeKB"),
+	16,
+	TEXT("Size of the scratch buffer in kB."),
+	ECVF_Default);
+
 
 struct
 {
@@ -23,9 +28,10 @@ struct
 	uint64							TimestampBase;
 	uint64							LastTimestamp;
 	uint32							RenderingFrameNumber;
-	uint16							EventBufferSize;
+	uint32							EventBufferSize;
 	bool							bActive;
-	uint8							EventBuffer[MaxEventBufferSize];
+	uint8*							EventBuffer = nullptr;
+	uint32							MaxEventBufferSize = 0;
 } GCurrentFrame;
 
 static TSet<uint32> GEventNames;
@@ -70,6 +76,14 @@ void FGpuProfilerTrace::BeginFrame(FGPUTimingCalibrationTimestamp& Calibration)
 	GCurrentFrame.TimestampBase = 0;
 	GCurrentFrame.EventBufferSize = 0;
 	GCurrentFrame.bActive = true;
+
+	int32 NeededSize = CVarGpuProfilerMaxEventBufferSizeKB.GetValueOnRenderThread() * 1024;
+	if ((GCurrentFrame.MaxEventBufferSize != NeededSize) && (NeededSize > 0))
+	{
+		FMemory::Free(GCurrentFrame.EventBuffer);
+		GCurrentFrame.EventBuffer = (uint8*)FMemory::Malloc(NeededSize);
+		GCurrentFrame.MaxEventBufferSize = NeededSize;
+	}
 }
 
 void FGpuProfilerTrace::SpecifyEventByName(const FName& Name)
@@ -109,8 +123,9 @@ void FGpuProfilerTrace::BeginEventByName(const FName& Name, uint32 FrameNumber,
 	}
 
 	// Prevent buffer overrun
-	if (GCurrentFrame.EventBufferSize + sizeof(uint64) + sizeof(uint32) >= MaxEventBufferSize)
+	if ((GCurrentFrame.EventBufferSize + 10 + sizeof(uint32)) >= GCurrentFrame.MaxEventBufferSize) // 10 is the max size that FTraceUtils::Encode7bit might use + some space for the FName index (uint32)
 	{
+		UE_LOG(LogRHI, Error, TEXT("GpuProfiler's scratch buffer is out of space for this frame (current size : %d kB). Dropping this frame. The size can be increased dynamically with the console variable r.GpuProfilerMaxEventBufferSizeKB"), GCurrentFrame.MaxEventBufferSize / 1024);
 		// Deactivate for the current frame to avoid errors while decoding an incomplete trace
 		GCurrentFrame.bActive = false;
 		return;
@@ -142,10 +157,12 @@ void FGpuProfilerTrace::EndEvent(uint64 TimestampMicroseconds)
 	{
 		return;
 	}
-	
+
 	// Prevent buffer overrun
-	if (GCurrentFrame.EventBufferSize + sizeof(uint64) >= MaxEventBufferSize)
+	if (GCurrentFrame.EventBufferSize + 10 >= GCurrentFrame.MaxEventBufferSize) // 10 is the max size that FTraceUtils::Encode7bit might use + 4 for the FName (index == uint32)
 	{
+		UE_LOG(LogRHI, Error, TEXT("GpuProfiler's scratch buffer is out of space for this frame (current size : %d kB). Dropping this frame. The size can be increased dynamically with the console variable r.GpuProfilerMaxEventBufferSizeKB"), GCurrentFrame.MaxEventBufferSize / 1024);
+
 		// Deactivate for the current frame to avoid errors while decoding an incomplete trace
 		GCurrentFrame.bActive = false;
 		return;
@@ -169,11 +186,11 @@ void FGpuProfilerTrace::EndFrame(uint32 GPUIndex)
 
 		if (GPUIndex == 0)
 		{
-			UE_TRACE_LOG(GpuProfiler, Frame, GpuChannel)
-				<< Frame.CalibrationBias(Bias)
-				<< Frame.TimestampBase(GCurrentFrame.TimestampBase)
-				<< Frame.RenderingFrameNumber(GCurrentFrame.RenderingFrameNumber)
-				<< Frame.Data(GCurrentFrame.EventBuffer, GCurrentFrame.EventBufferSize);
+		UE_TRACE_LOG(GpuProfiler, Frame, GpuChannel)
+			<< Frame.CalibrationBias(Bias)
+			<< Frame.TimestampBase(GCurrentFrame.TimestampBase)
+			<< Frame.RenderingFrameNumber(GCurrentFrame.RenderingFrameNumber)
+			<< Frame.Data(GCurrentFrame.EventBuffer, GCurrentFrame.EventBufferSize);
 		}
 		else if (GPUIndex == 1)
 		{
@@ -189,4 +206,13 @@ void FGpuProfilerTrace::EndFrame(uint32 GPUIndex)
 	GCurrentFrame.bActive = false;
 }
 
+void FGpuProfilerTrace::Deinitialize()
+{
+	using namespace GpuProfilerTrace;
+
+	FMemory::Free(GCurrentFrame.EventBuffer);
+	GCurrentFrame.EventBuffer = nullptr;
+	GCurrentFrame.MaxEventBufferSize = 0;
+}
+
 #endif
diff --git a/Engine/Source/Runtime/RHI/Public/GpuProfilerTrace.h b/Engine/Source/Runtime/RHI/Public/GpuProfilerTrace.h
--- a/Engine/Source/Runtime/RHI/Public/GpuProfilerTrace.h
+++ b/Engine/Source/Runtime/RHI/Public/GpuProfilerTrace.h
@@ -24,6 +24,7 @@ struct FGpuProfilerTrace
 	RHI_API static void BeginEventByName(const FName& Name, uint32 FrameNumber, uint64 TimestampMicroseconds);
 	RHI_API static void EndEvent(uint64 TimestampMicroseconds);
 	RHI_API static void EndFrame(uint32 GPUIndex);
+	RHI_API static void Deinitialize();
 };
 
 #define TRACE_GPUPROFILER_DEFINE_EVENT_TYPE(Name) \
@@ -47,6 +48,9 @@ struct FGpuProfilerTrace
 #define TRACE_GPUPROFILER_END_FRAME() \
 	FGpuProfilerTrace::EndFrame();
 
+#define TRACE_GPUPROFILER_DEINITIALIZE() \
+	FGpuProfilerTrace::Deinitialize();
+
 #else
 
 struct FGpuProfilerTrace
@@ -64,5 +68,6 @@ struct FGpuProfilerTrace
 #define TRACE_GPUPROFILER_BEGIN_EVENT(...)
 #define TRACE_GPUPROFILER_END_EVENT(...)
 #define TRACE_GPUPROFILER_END_FRAME(...)
+#define TRACE_GPUPROFILER_DEINITIALIZE(...)
 
 #endif
diff --git a/Engine/Source/Runtime/RenderCore/Private/ProfilingDebugging/RealtimeGPUProfiler.cpp b/Engine/Source/Runtime/RenderCore/Private/ProfilingDebugging/RealtimeGPUProfiler.cpp
--- a/Engine/Source/Runtime/RenderCore/Private/ProfilingDebugging/RealtimeGPUProfiler.cpp
+++ b/Engine/Source/Runtime/RenderCore/Private/ProfilingDebugging/RealtimeGPUProfiler.cpp
@@ -772,6 +772,8 @@ void FRealtimeGPUProfiler::Cleanup()
 	}
 	Frames.Empty();
 	RenderQueryPool.SafeRelease();
+	
+	TRACE_GPUPROFILER_DEINITIALIZE();
 }
 
 #if UE_TRACE_ENABLED
