Initial import of AGXRHI

diff --git a/Engine/Source/Runtime/Apple/AGXRHI/Public/AGXState.h b/Engine/Source/Runtime/Apple/AGXRHI/Public/AGXState.h
new file mode 100644
--- /dev/null
+++ b/Engine/Source/Runtime/Apple/AGXRHI/Public/AGXState.h
@@ -0,0 +1,152 @@
+// Copyright Epic Games, Inc. All Rights Reserved.
+
+/*=============================================================================
+	AGXState.h: AGX RHI state definitions.
+=============================================================================*/
+
+#pragma once
+
+#include "CoreMinimal.h"
+#include "RHI.h"
+THIRD_PARTY_INCLUDES_START
+#include "mtlpp.hpp"
+THIRD_PARTY_INCLUDES_END
+
+class FAGXSampler : public mtlpp::SamplerState
+{
+public:
+	FAGXSampler(ns::Ownership retain = ns::Ownership::Retain) : mtlpp::SamplerState(nullptr, nullptr, retain) { }
+	FAGXSampler(ns::Protocol<id<MTLSamplerState>>::type handle, ns::Ownership retain = ns::Ownership::Retain)
+	: mtlpp::SamplerState(handle, nullptr, retain) {}
+	
+	FAGXSampler(mtlpp::SamplerState&& rhs)
+	: mtlpp::SamplerState((mtlpp::SamplerState&&)rhs)
+	{
+		
+	}
+	
+	FAGXSampler(const FAGXSampler& rhs)
+	: mtlpp::SamplerState(rhs)
+	{
+		
+	}
+	
+	FAGXSampler(const SamplerState& rhs)
+	: mtlpp::SamplerState(rhs)
+	{
+		
+	}
+	
+	FAGXSampler(FAGXSampler&& rhs)
+	: mtlpp::SamplerState((mtlpp::SamplerState&&)rhs)
+	{
+		
+	}
+	
+	FAGXSampler& operator=(const FAGXSampler& rhs)
+	{
+		if (this != &rhs)
+		{
+			mtlpp::SamplerState::operator=(rhs);
+		}
+		return *this;
+	}
+	
+	FAGXSampler& operator=(FAGXSampler&& rhs)
+	{
+		mtlpp::SamplerState::operator=((mtlpp::SamplerState&&)rhs);
+		return *this;
+	}
+	
+	inline bool operator==(FAGXSampler const& rhs) const
+	{
+		return GetPtr() == rhs.GetPtr();
+	}
+	
+	inline bool operator!=(FAGXSampler const& rhs) const
+	{
+		return GetPtr() != rhs.GetPtr();
+	}
+	
+	friend uint32 GetTypeHash(FAGXSampler const& Hash)
+	{
+		return GetTypeHash(Hash.GetPtr());
+	}
+};
+
+class FAGXSamplerState : public FRHISamplerState
+{
+public:
+	
+	/** 
+	 * Constructor/destructor
+	 */
+	FAGXSamplerState(mtlpp::Device Device, const FSamplerStateInitializerRHI& Initializer);
+	~FAGXSamplerState();
+
+	FAGXSampler State;
+#if !PLATFORM_MAC
+	FAGXSampler NoAnisoState;
+#endif
+};
+
+class FAGXRasterizerState : public FRHIRasterizerState
+{
+public:
+
+	/**
+	 * Constructor/destructor
+	 */
+	FAGXRasterizerState(const FRasterizerStateInitializerRHI& Initializer);
+	~FAGXRasterizerState();
+	
+	virtual bool GetInitializer(FRasterizerStateInitializerRHI& Init) override final;
+	
+	FRasterizerStateInitializerRHI State;
+};
+
+class FAGXDepthStencilState : public FRHIDepthStencilState
+{
+public:
+
+	/**
+	 * Constructor/destructor
+	 */
+	FAGXDepthStencilState(mtlpp::Device Device, const FDepthStencilStateInitializerRHI& Initializer);
+	~FAGXDepthStencilState();
+	
+	virtual bool GetInitializer(FDepthStencilStateInitializerRHI& Init) override final;
+	
+	FDepthStencilStateInitializerRHI Initializer;
+	mtlpp::DepthStencilState State;
+	bool bIsDepthWriteEnabled;
+	bool bIsStencilWriteEnabled;
+};
+
+class FAGXBlendState : public FRHIBlendState
+{
+public:
+
+	/**
+	 * Constructor/destructor
+	 */
+	FAGXBlendState(const FBlendStateInitializerRHI& Initializer);
+	~FAGXBlendState();
+	
+	virtual bool GetInitializer(FBlendStateInitializerRHI& Init) override final;
+
+	struct FBlendPerMRT
+	{
+		mtlpp::RenderPipelineColorAttachmentDescriptor BlendState;
+		uint8 BlendStateKey;
+	};
+	FBlendPerMRT RenderTargetStates[MaxSimultaneousRenderTargets];
+	bool bUseIndependentRenderTargetBlendStates;
+	bool bUseAlphaToCoverage;
+
+private:
+	// this tracks blend settings (in a bit flag) into a unique key that uses few bits, for PipelineState MRT setup
+	static TMap<uint32, uint8> BlendSettingsToUniqueKeyMap;
+	static uint8 NextKey;
+	static FCriticalSection Mutex;
+};
diff --git a/Engine/Source/Runtime/Apple/AGXRHI/Public/AGXUniformBuffer.h b/Engine/Source/Runtime/Apple/AGXRHI/Public/AGXUniformBuffer.h
new file mode 100644
--- /dev/null
+++ b/Engine/Source/Runtime/Apple/AGXRHI/Public/AGXUniformBuffer.h
@@ -0,0 +1,68 @@
+// Copyright Epic Games, Inc. All Rights Reserved.
+
+#pragma once
+
+//
+//  Implements handles to linearly allocated per-frame constant buffers for shared memory systems.
+//
+
+#import <Metal/Metal.h>
+
+#include "RHIResources.h"
+
+#define METAL_UNIFORM_BUFFER_VALIDATION !UE_BUILD_SHIPPING
+
+class FAGXStateCache;
+
+class FAGXSuballocatedUniformBuffer : public FRHIUniformBuffer
+{
+    friend class FAGXStateCache;
+public:
+    // The last render thread frame this uniform buffer updated or pushed contents to the GPU backing
+    int32 LastFrameUpdated;
+    // Offset within the GPU backing this uniform buffer owns
+    uint32 Offset;
+    // The GPU backing buffer for this uniform buffer. Many FAGXMobileUniformBuffers can own regions.
+    // This UB does not own a reference to the backing buffer.
+    // This Backing is recycled at the end of every frame so you MUST update it if LastFrameUpdated != this frame or contents are undefined.
+    id <MTLBuffer> Backing;
+    // CPU side shadow memory to hold updates for single-draw or multi-frame buffers.
+    // This allows you to upload on a frame but actually use this UB later on
+    void* Shadow;
+private:
+    TArray<TRefCountPtr<FRHIResource> > ResourceTable;
+#if METAL_UNIFORM_BUFFER_VALIDATION
+    EUniformBufferValidation Validation;
+#endif
+    
+public:
+    // Creates a uniform buffer.
+    // If Usage is SingleDraw or MultiFrame we will keep a copy of the data
+    FAGXSuballocatedUniformBuffer(const FRHIUniformBufferLayout& Layout, EUniformBufferUsage Usage, EUniformBufferValidation Validation);
+    ~FAGXSuballocatedUniformBuffer();
+    
+	void Update(const void* Contents, TArray<TRefCountPtr<FRHIResource> > const& InResourceTable);
+    
+    // Prepares this uniform buffer for binding.
+    // You MUST call this before binding it.
+    // If this is a UB that was created before the current frame the data on the GPU will
+    // not be correct until this function returns.
+    void PrepareToBind();
+
+	// Copies the RDG resources to a resource table for a deferred update on the RHI thread.
+	void CopyResourceTable_RenderThread(const void* Contents, TArray<TRefCountPtr<FRHIResource> >& OutResourceTable);
+
+private:
+    
+    // Pushes the data in Contents to the gpu.
+    // Updates the frame counter to FrameNumber.
+    // (this is to support the case where we create buffer and reuse it many frames later).
+    // This acquires a region in the current frame's transient uniform buffer
+    // and copies Contents into that backing.
+    // The amount of data is determined by the Layout
+    void PushToGPUBacking(const void* Contents);
+    
+    bool HasShadow();
+};
+
+typedef FAGXSuballocatedUniformBuffer FAGXUniformBuffer;
diff --git a/Engine/Source/Runtime/Apple/AGXRHI/Public/AGXViewport.h b/Engine/Source/Runtime/Apple/AGXRHI/Public/AGXViewport.h
new file mode 100644
--- /dev/null
+++ b/Engine/Source/Runtime/Apple/AGXRHI/Public/AGXViewport.h
@@ -0,0 +1,89 @@
+// Copyright Epic Games, Inc. All Rights Reserved.
+
+/*=============================================================================
+	AGXViewport.h: AGX RHI viewport definitions.
+=============================================================================*/
+
+#pragma once
+
+#if PLATFORM_MAC
+#include "Mac/CocoaTextView.h"
+@interface FAGXView : FCocoaTextView
+@end
+#endif
+#include "HAL/PlatformFramePacer.h"
+THIRD_PARTY_INCLUDES_START
+#include "mtlpp.hpp"
+THIRD_PARTY_INCLUDES_END
+
+enum EAGXViewportAccessFlag
+{
+	EAGXViewportAccessRHI,
+	EAGXViewportAccessRenderer,
+	EAGXViewportAccessGame,
+	EAGXViewportAccessDisplayLink
+};
+
+class FAGXCommandQueue;
+
+typedef void (^FAGXViewportPresentHandler)(uint32 CGDirectDisplayID, double OutputSeconds, double OutputDuration);
+
+class FAGXViewport : public FRHIViewport
+{
+public:
+	FAGXViewport(void* WindowHandle, uint32 InSizeX, uint32 InSizeY, bool bInIsFullscreen,EPixelFormat Format);
+	~FAGXViewport();
+
+	void Resize(uint32 InSizeX, uint32 InSizeY, bool bInIsFullscreen,EPixelFormat Format);
+	
+	TRefCountPtr<FAGXTexture2D> GetBackBuffer(EAGXViewportAccessFlag Accessor) const;
+	id<CAMetalDrawable> GetDrawable(EAGXViewportAccessFlag Accessor);
+	FAGXTexture GetDrawableTexture(EAGXViewportAccessFlag Accessor);
+	ns::AutoReleased<FAGXTexture> GetCurrentTexture(EAGXViewportAccessFlag Accessor);
+	void ReleaseDrawable(void);
+
+	// supports pulling the raw MTLTexture
+	virtual void* GetNativeBackBufferTexture() const override { return GetBackBuffer(EAGXViewportAccessRenderer).GetReference(); }
+	virtual void* GetNativeBackBufferRT() const override { return (const_cast<FAGXViewport *>(this))->GetDrawableTexture(EAGXViewportAccessRenderer); }
+	
+#if PLATFORM_MAC
+	NSWindow* GetWindow() const;
+	
+	virtual void SetCustomPresent(FRHICustomPresent* InCustomPresent) override
+	{
+		CustomPresent = InCustomPresent;
+	}
+
+	virtual FRHICustomPresent* GetCustomPresent() const override { return CustomPresent; }
+#endif
+	
+	void Present(FAGXCommandQueue& CommandQueue, bool bLockToVsync);
+	void Swap();
+	
+private:
+	uint32 GetViewportIndex(EAGXViewportAccessFlag Accessor) const;
+
+private:
+	id<CAMetalDrawable> Drawable;
+	TRefCountPtr<FAGXTexture2D> BackBuffer[2];
+	mutable FCriticalSection Mutex;
+	
+	ns::AutoReleased<FAGXTexture> DrawableTextures[2];
+	
+	uint32 DisplayID;
+	FAGXViewportPresentHandler Block;
+	volatile int32 FrameAvailable;
+	TRefCountPtr<FAGXTexture2D> LastCompleteFrame;
+	bool bIsFullScreen;
+
+#if PLATFORM_MAC
+	FAGXView* View;
+	FRHICustomPresent* CustomPresent;
+#endif
+};
+
+template<>
+struct TAGXResourceTraits<FRHIViewport>
+{
+	typedef FAGXViewport TConcreteType;
+};
diff --git a/Engine/Source/Runtime/Apple/Common/Public/MetalShaderResources.h b/Engine/Source/Runtime/Apple/Common/Public/MetalShaderResources.h
new file mode 100644
--- /dev/null
+++ b/Engine/Source/Runtime/Apple/Common/Public/MetalShaderResources.h
@@ -0,0 +1,319 @@
+// Copyright Epic Games, Inc. All Rights Reserved.
+
+/*=============================================================================
+	MetalShaderResources.h: Metal shader resource RHI definitions.
+=============================================================================*/
+
+#pragma once
+
+#include "CrossCompilerCommon.h"
+
+/**
+* Shader related constants.
+*/
+enum
+{
+	METAL_MAX_UNIFORM_BUFFER_BINDINGS = 12,	// @todo-mobile: Remove me
+	METAL_FIRST_UNIFORM_BUFFER = 0,			// @todo-mobile: Remove me
+	METAL_MAX_COMPUTE_STAGE_UAV_UNITS = 8,	// @todo-mobile: Remove me
+	METAL_UAV_NOT_SUPPORTED_FOR_GRAPHICS_UNIT = -1, // for now, only CS supports UAVs/ images
+	METAL_MAX_BUFFERS = 31,
+};
+
+/**
+* Buffer data-types for MetalRHI & MetalSL
+*/
+enum class EMetalBufferFormat : uint8
+{
+	Unknown					=0,
+	
+	R8Sint					=1,
+	R8Uint					=2,
+	R8Snorm					=3,
+	R8Unorm					=4,
+	
+	R16Sint					=5,
+	R16Uint					=6,
+	R16Snorm				=7,
+	R16Unorm				=8,
+	R16Half					=9,
+	
+	R32Sint					=10,
+	R32Uint					=11,
+	R32Float				=12,
+	
+	RG8Sint					=13,
+	RG8Uint					=14,
+	RG8Snorm				=15,
+	RG8Unorm				=16,
+	
+	RG16Sint				=17,
+	RG16Uint				=18,
+	RG16Snorm				=19,
+	RG16Unorm				=20,
+	RG16Half				=21,
+	
+	RG32Sint				=22,
+	RG32Uint				=23,
+	RG32Float				=24,
+	
+	RGB8Sint				=25,
+	RGB8Uint				=26,
+	RGB8Snorm				=27,
+	RGB8Unorm				=28,
+	
+	RGB16Sint				=29,
+	RGB16Uint				=30,
+	RGB16Snorm				=31,
+	RGB16Unorm				=32,
+	RGB16Half				=33,
+	
+	RGB32Sint				=34,
+	RGB32Uint				=35,
+	RGB32Float				=36,
+	
+	RGBA8Sint				=37,
+	RGBA8Uint				=38,
+	RGBA8Snorm				=39,
+	RGBA8Unorm				=40,
+	
+	BGRA8Unorm				=41,
+	
+	RGBA16Sint				=42,
+	RGBA16Uint				=43,
+	RGBA16Snorm				=44,
+	RGBA16Unorm				=45,
+	RGBA16Half				=46,
+	
+	RGBA32Sint				=47,
+	RGBA32Uint				=48,
+	RGBA32Float				=49,
+	
+	RGB10A2Unorm			=50,
+	
+	RG11B10Half 			=51,
+	
+	R5G6B5Unorm         	=52,
+	
+	Max						=53
+};
+
+struct FMetalShaderResourceTable : public FBaseShaderResourceTable
+{
+	/** Mapping of bound Textures to their location in resource tables. */
+	TArray<uint32> TextureMap;
+	friend bool operator==(const FMetalShaderResourceTable &A, const FMetalShaderResourceTable& B)
+	{
+		if (!(((FBaseShaderResourceTable&)A) == ((FBaseShaderResourceTable&)B)))
+		{
+			return false;
+		}
+		if (A.TextureMap.Num() != B.TextureMap.Num())
+		{
+			return false;
+		}
+		if (FMemory::Memcmp(A.TextureMap.GetData(), B.TextureMap.GetData(), A.TextureMap.GetTypeSize()*A.TextureMap.Num()) != 0)
+		{
+			return false;
+		}
+		return true;
+	}
+};
+
+inline FArchive& operator<<(FArchive& Ar, FMetalShaderResourceTable& SRT)
+{
+	Ar << ((FBaseShaderResourceTable&)SRT);
+	Ar << SRT.TextureMap;
+	return Ar;
+}
+
+struct FMetalShaderBindings
+{
+	TArray<TArray<CrossCompiler::FPackedArrayInfo>>	PackedUniformBuffers;
+	TArray<CrossCompiler::FPackedArrayInfo>			PackedGlobalArrays;
+	FMetalShaderResourceTable						ShaderResourceTable;
+	TMap<uint8, TArray<uint8>>						ArgumentBufferMasks;
+	CrossCompiler::FShaderBindingInOutMask			InOutMask;
+
+    uint32  LinearBuffer;
+	uint32	TypedBuffers;
+	uint32 	ConstantBuffers;
+	uint32  ArgumentBuffers;
+	uint8	NumSamplers;
+	uint8	NumUniformBuffers;
+	uint8	NumUAVs;
+	bool	bDiscards;
+
+	FMetalShaderBindings() :
+		LinearBuffer(0),
+        TypedBuffers(0),
+		ConstantBuffers(0),
+		ArgumentBuffers(0),
+		NumSamplers(0),
+		NumUniformBuffers(0),
+		NumUAVs(0),
+		bDiscards(false)
+	{
+	}
+};
+
+inline FArchive& operator<<(FArchive& Ar, FMetalShaderBindings& Bindings)
+{
+	Ar << Bindings.PackedUniformBuffers;
+	Ar << Bindings.PackedGlobalArrays;
+	Ar << Bindings.ShaderResourceTable;
+	Ar << Bindings.ArgumentBufferMasks;
+    Ar << Bindings.LinearBuffer;
+    Ar << Bindings.TypedBuffers;
+	Ar << Bindings.ConstantBuffers;
+	Ar << Bindings.ArgumentBuffers;
+	Ar << Bindings.InOutMask;
+	Ar << Bindings.NumSamplers;
+	Ar << Bindings.NumUniformBuffers;
+	Ar << Bindings.NumUAVs;
+	Ar << Bindings.bDiscards;
+	return Ar;
+}
+
+enum class EMetalOutputWindingMode : uint8
+{
+	Clockwise = 0,
+	CounterClockwise = 1,
+};
+
+enum class EMetalPartitionMode : uint8
+{
+	Pow2 = 0,
+	Integer = 1,
+	FractionalOdd = 2,
+	FractionalEven = 3,
+};
+
+enum class EMetalComponentType : uint8
+{
+	Uint = 0,
+	Int,
+	Half,
+	Float,
+	Bool,
+	Max
+};
+
+struct FMetalAttribute
+{
+	uint32 Index;
+	uint32 Components;
+	uint32 Offset;
+	EMetalComponentType Type;
+	uint32 Semantic;
+	
+	FMetalAttribute()
+	: Index(0)
+	, Components(0)
+	, Offset(0)
+	, Type(EMetalComponentType::Uint)
+	, Semantic(0)
+	{
+		
+	}
+	
+	friend FArchive& operator<<(FArchive& Ar, FMetalAttribute& Attr)
+	{
+		Ar << Attr.Index;
+		Ar << Attr.Type;
+		Ar << Attr.Components;
+		Ar << Attr.Offset;
+		Ar << Attr.Semantic;
+		return Ar;
+	}
+};
+
+struct FMetalCodeHeader
+{
+	FMetalShaderBindings Bindings;
+	TArray<CrossCompiler::FUniformBufferCopyInfo> UniformBuffersCopyInfo;
+	FString ShaderName;
+
+	uint64 CompilerBuild;
+	uint32 CompilerVersion;
+	uint32 SourceLen;
+	uint32 SourceCRC;
+	uint32 NumThreadsX;
+	uint32 NumThreadsY;
+	uint32 NumThreadsZ;
+	uint32 CompileFlags;
+	uint8 Frequency;
+	uint8 Version;
+	int8 SideTable;
+	bool bDeviceFunctionConstants;
+	
+	FMetalCodeHeader()
+	: CompilerBuild(0)
+	, CompilerVersion(0)
+	, SourceLen(0)
+	, SourceCRC(0)
+	, NumThreadsX(0)
+	, NumThreadsY(0)
+	, NumThreadsZ(0)
+	, CompileFlags(0)
+	, Frequency(0)
+	, Version(0)
+	, SideTable(-1)
+	, bDeviceFunctionConstants(false)
+	{
+	}
+};
+
+inline FArchive& operator<<(FArchive& Ar, FMetalCodeHeader& Header)
+{
+	Ar << Header.Bindings;
+
+	int32 NumInfos = Header.UniformBuffersCopyInfo.Num();
+	Ar << NumInfos;
+	if (Ar.IsSaving())
+	{
+		for (int32 Index = 0; Index < NumInfos; ++Index)
+		{
+			Ar << Header.UniformBuffersCopyInfo[Index];
+		}
+	}
+	else if (Ar.IsLoading())
+	{
+		Header.UniformBuffersCopyInfo.Empty(NumInfos);
+		for (int32 Index = 0; Index < NumInfos; ++Index)
+		{
+			CrossCompiler::FUniformBufferCopyInfo Info;
+			Ar << Info;
+			Header.UniformBuffersCopyInfo.Add(Info);
+		}
+	}
+	
+	Ar << Header.ShaderName;
+
+	Ar << Header.CompilerBuild;
+	Ar << Header.CompilerVersion;
+	Ar << Header.SourceLen;
+	Ar << Header.SourceCRC;
+	Ar << Header.NumThreadsX;
+	Ar << Header.NumThreadsY;
+	Ar << Header.NumThreadsZ;
+	Ar << Header.CompileFlags;
+	Ar << Header.Frequency;
+	Ar << Header.Version;
+	Ar << Header.SideTable;
+	Ar << Header.bDeviceFunctionConstants;
+
+    return Ar;
+}
+
+struct FMetalShaderLibraryHeader
+{
+	FString Format;
+	uint32 NumLibraries;
+	uint32 NumShadersPerLibrary;
+	
+	friend FArchive& operator<<(FArchive& Ar, FMetalShaderLibraryHeader& Header)
+	{
+		return Ar << Header.Format << Header.NumLibraries << Header.NumShadersPerLibrary;
+	}
+};
diff --git a/Engine/Source/Runtime/Apple/MetalRHI/MetalRHI.Build.cs b/Engine/Source/Runtime/Apple/MetalRHI/MetalRHI.Build.cs
--- a/Engine/Source/Runtime/Apple/MetalRHI/MetalRHI.Build.cs
+++ b/Engine/Source/Runtime/Apple/MetalRHI/MetalRHI.Build.cs
@@ -7,6 +7,8 @@ public class MetalRHI : ModuleRules
 {	
 	public MetalRHI(ReadOnlyTargetRules Target) : base(Target)
 	{
+		PublicIncludePaths.Add("Runtime/Apple/Common/Public");
+
 		PrivateDependencyModuleNames.AddRange(
 			new string[] {
 				"Core",
diff --git a/Engine/Source/Runtime/Core/Public/Stats/Stats2.h b/Engine/Source/Runtime/Core/Public/Stats/Stats2.h
--- a/Engine/Source/Runtime/Core/Public/Stats/Stats2.h
+++ b/Engine/Source/Runtime/Core/Public/Stats/Stats2.h
@@ -2241,6 +2241,7 @@ DECLARE_STATS_GROUP(TEXT("Memory StaticMesh"),STATGROUP_MemoryStaticMesh, STATCA
 DECLARE_STATS_GROUP(TEXT("Memory"),STATGROUP_Memory, STATCAT_Advanced);
 DECLARE_STATS_GROUP(TEXT("Mesh Particles"),STATGROUP_MeshParticles, STATCAT_Advanced);
 DECLARE_STATS_GROUP(TEXT("Metal"),STATGROUP_MetalRHI, STATCAT_Advanced);
+DECLARE_STATS_GROUP(TEXT("AGX"),STATGROUP_AGXRHI, STATCAT_Advanced);
 DECLARE_STATS_GROUP(TEXT("Morph"),STATGROUP_MorphTarget, STATCAT_Advanced);
 DECLARE_STATS_GROUP(TEXT("Navigation"),STATGROUP_Navigation, STATCAT_Advanced);
 DECLARE_STATS_GROUP(TEXT("Net"),STATGROUP_Net, STATCAT_Advanced);
diff --git a/Engine/Source/Runtime/RHI/Private/Apple/AppleDynamicRHI.cpp b/Engine/Source/Runtime/RHI/Private/Apple/AppleDynamicRHI.cpp
--- a/Engine/Source/Runtime/RHI/Private/Apple/AppleDynamicRHI.cpp
+++ b/Engine/Source/Runtime/RHI/Private/Apple/AppleDynamicRHI.cpp
@@ -1,120 +1,194 @@
 // Copyright Epic Games, Inc. All Rights Reserved.
 
+
 #include "RHI.h"
-#include "Misc/ConfigCacheIni.h"
-#include "Misc/Parse.h"
+#include "Modules/ModuleManager.h"
 #include "Misc/CommandLine.h"
+#include "Misc/ConfigCacheIni.h"
 #include "Misc/MessageDialog.h"
-#include "Modules/ModuleManager.h"
+#include COMPILED_PLATFORM_HEADER_WITH_PREFIX(Apple/Platform, PlatformDynamicRHI.h)
 
-int32 GAppleMetalEnabled = 1;
-static FAutoConsoleVariableRef CVarMacMetalEnabled(
-	TEXT("rhi.Apple.UseMetal"),
-	GAppleMetalEnabled,
-	TEXT("If set to true uses Metal when available rather than OpenGL as the graphics API. (Default: True)"));
 
-int32 GAppleOpenGLDisabled = 0;
-static FAutoConsoleVariableRef CVarMacOpenGLDisabled(
-	TEXT("rhi.Apple.OpenGLDisabled"),
-	GAppleOpenGLDisabled,
-	TEXT("If set, OpenGL RHI will not be used if Metal is not available. Instead, a dialog box explaining that the hardware requirements are not met will appear. (Default: False)"));
+//------------------------------------------------------------------------------
+// MARK: - FAppleDynamicRHIOptions Union
+//
 
-FDynamicRHI* PlatformCreateDynamicRHI()
+union FAppleDynamicRHIOptions
 {
-	SCOPED_AUTORELEASE_POOL;
+	struct
+	{
+		uint16 PreferAGX                : 1;
+		uint16 ForceSM5                 : 1;
+		uint16 ForceSM6                 : 1;
+		uint16 PreferES31               : 1;
+		uint16 AllowAGXFeatureLevelES31 : 1;
+		uint16 ForceMTL                 : 1;
+		uint16 ForceAGX                 : 1;
+		uint16 UnusedReservedBits       : 9;
+	};
+
+	uint16 All;
+};
 
-	FDynamicRHI* DynamicRHI = NULL;
-	IDynamicRHIModule* DynamicRHIModule = NULL;
 
-	bool const bIsMetalSupported = FPlatformMisc::HasPlatformFeature(TEXT("Metal"));
-	
-	// Must be Metal!
-	if(!bIsMetalSupported)
+//------------------------------------------------------------------------------
+// MARK: - Apple Dynamic RHI Support Routines
+//
+
+static bool ShouldPreferAGX()
+{
+	// TODO: Add logic here to detect Apple GPUs and set this preference.
+	return false;
+}
+
+static bool ShouldAllowAGXFeatureLevelES31()
+{
+	return true;
+}
+
+static inline bool ValidateAppleDynamicRHIOptions(FAppleDynamicRHIOptions* Options)
+{
+	if (0 != (Options->ForceMTL & Options->ForceAGX))
+	{
+		UE_LOG(LogRHI, Fatal, TEXT("-mtl and -agx are mutually exclusive options but more than one was specified on the command line."));
+		return false;
+	}
+	if (0 != (Options->ForceSM5 & Options->ForceSM6))
 	{
-		FText Title = NSLOCTEXT("AppleDynamicRHI", "OpenGLNotSupportedTitle","Metal Not Supported");
-#if PLATFORM_MAC
-		FMessageDialog::Open(EAppMsgType::Ok, NSLOCTEXT("MacPlatformCreateDynamicRHI", "OpenGLNotSupported.", "You must have a Metal compatible graphics card and be running Mac OS X 10.11.6 or later to launch this process."), &Title);
-#else
-		FMessageDialog::Open(EAppMsgType::Ok, NSLOCTEXT("AppleDynamicRHI", "OpenGLNotSupported.", "You must have a Metal compatible iOS or tvOS device with iOS 8 or later to launch this app."), &Title);
-#endif
-		FPlatformMisc::RequestExit(true);
+		UE_LOG(LogRHI, Fatal, TEXT("-sm5 and -sm6 are mutually exclusive options but more than one was specified on the command line."));
+		return false;
 	}
-	
-	if (FParse::Param(FCommandLine::Get(),TEXT("opengl")))
+	if (0 != (Options->ForceMTL & Options->ForceSM6))
 	{
-		UE_LOG(LogRHI, Log, TEXT("OpenGL command line option ignored; Apple platforms only support Metal."));
+		UE_LOG(LogRHI, Warning, TEXT("-mtl and -sm6 are incompatible options, using MetalRHI with SM5."));
+		Options->ForceSM5 = 1;
+		Options->ForceSM6 = 0;
+		Options->ForceMTL = 1;
+		Options->ForceAGX = 0;
 	}
+	if (0 != Options->ForceSM6)
+	{
+		Options->ForceMTL = 0;
+		Options->ForceAGX = 1;
+	}
+	return true;
+}
 
-	ERHIFeatureLevel::Type RequestedFeatureLevel = ERHIFeatureLevel::Num;
+static bool InitAppleDynamicRHIOptions(FAppleDynamicRHIOptions* Options)
+{
+	Options->PreferAGX                = ShouldPreferAGX();
+	Options->ForceSM5                 = FParse::Param(FCommandLine::Get(), TEXT("sm5"));
+	Options->ForceSM6                 = FParse::Param(FCommandLine::Get(), TEXT("sm6"));
+	Options->PreferES31               = FPlatformDynamicRHI::ShouldPreferFeatureLevelES31() && !(Options->ForceSM5 || Options->ForceSM6);
+	Options->AllowAGXFeatureLevelES31 = ShouldAllowAGXFeatureLevelES31();
+	Options->ForceMTL                 = FParse::Param(FCommandLine::Get(), TEXT("mtl"));
+	Options->ForceAGX                 = FParse::Param(FCommandLine::Get(), TEXT("agx")) && (!Options->PreferES31 || Options->AllowAGXFeatureLevelES31);
+
+	return ValidateAppleDynamicRHIOptions(Options);
+}
+
+static inline bool ShouldUseShaderModelPreference(const FAppleDynamicRHIOptions& Options)
+{
+	return (0 != (Options.ForceSM5 | Options.ForceSM6 | Options.PreferES31));
+}
+
+static void ComputeRequestedFeatureLevel(const FAppleDynamicRHIOptions& Options, ERHIFeatureLevel::Type& RequestedFeatureLevel)
+{
+	if (!ShouldUseShaderModelPreference(Options))
 	{
-		// Check the list of targeted shader platforms and decide an RHI based off them
 		TArray<FString> TargetedShaderFormats;
-#if PLATFORM_MAC
-		GConfig->GetArray(TEXT("/Script/MacTargetPlatform.MacTargetSettings"), TEXT("TargetedRHIs"), TargetedShaderFormats, GEngineIni);
-#else
-		bool bSupportsMetalMRT = false;
-		GConfig->GetBool(TEXT("/Script/IOSRuntimeSettings.IOSRuntimeSettings"), TEXT("bSupportsMetalMRT"), bSupportsMetalMRT, GEngineIni);
-		if (bSupportsMetalMRT)
+
+		FPlatformDynamicRHI::AddTargetedShaderFormats(TargetedShaderFormats);
+
+		if (TargetedShaderFormats.Num() > 0)
 		{
-#if PLATFORM_TVOS
-			TargetedShaderFormats.Add(LegacyShaderPlatformToShaderFormat(SP_METAL_MRT_TVOS).ToString());
-#else
-			TargetedShaderFormats.Add(LegacyShaderPlatformToShaderFormat(SP_METAL_MRT).ToString());
-#endif
+			// Pick the first one
+			FName ShaderFormatName(*TargetedShaderFormats[0]);
+			EShaderPlatform TargetedPlatform = ShaderFormatToLegacyShaderPlatform(ShaderFormatName);
+			RequestedFeatureLevel = GetMaxSupportedFeatureLevel(TargetedPlatform);
 		}
-		
-#if PLATFORM_TVOS
-		TargetedShaderFormats.Add(LegacyShaderPlatformToShaderFormat(SP_METAL_TVOS).ToString());
-#else
-		TargetedShaderFormats.Add(LegacyShaderPlatformToShaderFormat(SP_METAL).ToString());
-#endif
-		
-#endif // else branch of PLATFORM_MAC
-		
-		// Metal is not always available, so don't assume that we can use the first platform
-		for (FString Name : TargetedShaderFormats)
+	}
+	else
+	{
+		if (Options.ForceSM6)
 		{
-			FName ShaderFormatName(*Name);
-			EShaderPlatform TargetedPlatform = ShaderFormatToLegacyShaderPlatform(ShaderFormatName);
-			
-			// Instead use the first platform that *could* work
-			if (IsMetalPlatform(TargetedPlatform))
-			{
-				RequestedFeatureLevel = GetMaxSupportedFeatureLevel(TargetedPlatform);
-				break;
-			}
+			RequestedFeatureLevel = ERHIFeatureLevel::SM6;
+		}
+		else if (Options.ForceSM5)
+		{
+			RequestedFeatureLevel = ERHIFeatureLevel::SM5;
+		}
+		else
+		{
+			check(Options.PreferES31 != 0);
+			RequestedFeatureLevel = ERHIFeatureLevel::ES3_1;
 		}
 	}
 
-	// Load the dynamic RHI module.
+	check(RequestedFeatureLevel != ERHIFeatureLevel::Num);
+}
+
+static inline bool ShouldUseAGX(const FAppleDynamicRHIOptions& Options)
+{
+	return (0 != (Options.ForceAGX | (Options.PreferAGX & ~Options.ForceMTL)));
+}
+
+static IDynamicRHIModule* LoadDynamicRHIModule(ERHIFeatureLevel::Type& RequestedFeatureLevel)
+{
+	IDynamicRHIModule*      DynamicRHIModule = nullptr;
+	FAppleDynamicRHIOptions Options          = { .All = 0 };
+
+	if (!InitAppleDynamicRHIOptions(&Options))
 	{
-		DynamicRHIModule = &FModuleManager::LoadModuleChecked<IDynamicRHIModule>(TEXT("MetalRHI"));
-		
+		return nullptr;
+	}
+
+	ComputeRequestedFeatureLevel(Options, RequestedFeatureLevel);
+
+	if (ShouldUseAGX(Options))
+	{
+		FApp::SetGraphicsRHI(TEXT("AGX"));
+		const TCHAR* AGXRHIModuleName = TEXT("AGXRHI");
+		DynamicRHIModule = &FModuleManager::LoadModuleChecked<IDynamicRHIModule>(AGXRHIModuleName);
+		if (!DynamicRHIModule || !DynamicRHIModule->IsSupported())
 		{
-#if PLATFORM_MAC
-			if (FParse::Param(FCommandLine::Get(),TEXT("metal")))
-			{
-				RequestedFeatureLevel = ERHIFeatureLevel::SM5;
-			}
-			else if (FParse::Param(FCommandLine::Get(),TEXT("metalsm5")) || FParse::Param(FCommandLine::Get(),TEXT("metalmrt")))
+			if (Options.ForceAGX)
 			{
-				RequestedFeatureLevel = ERHIFeatureLevel::SM5;
+				FMessageDialog::Open(EAppMsgType::Ok, NSLOCTEXT("AppleDynamicRHI", "RequiredAGX", "The AGX RHI is not supported on your system. Try running without the -agx command line argument."));
+				FPlatformMisc::RequestExit(1);
 			}
-#else
-			if (FParse::Param(FCommandLine::Get(),TEXT("metal")))
+			if (DynamicRHIModule)
 			{
-				RequestedFeatureLevel = ERHIFeatureLevel::ES3_1;
+				FModuleManager::Get().UnloadModule(AGXRHIModuleName);
 			}
-			else if (FParse::Param(FCommandLine::Get(),TEXT("metalmrt")))
-			{
-				RequestedFeatureLevel = ERHIFeatureLevel::SM5;
-			}
-#endif
+			DynamicRHIModule = nullptr;
 		}
+	}
+
+	if (DynamicRHIModule == nullptr)
+	{
 		FApp::SetGraphicsRHI(TEXT("Metal"));
+		DynamicRHIModule = &FModuleManager::LoadModuleChecked<IDynamicRHIModule>(TEXT("MetalRHI"));
+	}
+
+	return DynamicRHIModule;
+}
+
+
+//------------------------------------------------------------------------------
+// MARK: - Dynamic RHI API
+//
+
+FDynamicRHI* PlatformCreateDynamicRHI()
+{
+	FDynamicRHI*           DynamicRHI            = nullptr;
+	IDynamicRHIModule*     DynamicRHIModule      = nullptr;
+	ERHIFeatureLevel::Type RequestedFeatureLevel = ERHIFeatureLevel::Num;
+
+	if (nullptr != (DynamicRHIModule = LoadDynamicRHIModule(RequestedFeatureLevel)))
+	{
+		DynamicRHI = DynamicRHIModule->CreateRHI(RequestedFeatureLevel);
 	}
-	
-	// Create the dynamic RHI.
-	DynamicRHI = DynamicRHIModule->CreateRHI(RequestedFeatureLevel);
+
 	return DynamicRHI;
 }
diff --git a/Engine/Source/Runtime/RHI/Private/Apple/Platform/IOS/IOSPlatformDynamicRHI.h b/Engine/Source/Runtime/RHI/Private/Apple/Platform/IOS/IOSPlatformDynamicRHI.h
new file mode 100644
--- /dev/null
+++ b/Engine/Source/Runtime/RHI/Private/Apple/Platform/IOS/IOSPlatformDynamicRHI.h
@@ -0,0 +1,62 @@
+// Copyright Epic Games, Inc. All Rights Reserved.
+
+#pragma once
+
+
+//------------------------------------------------------------------------------
+// MARK: - iOS/tvOS Platform Dynamic RHI Defines
+//
+
+
+//------------------------------------------------------------------------------
+// NOTE: Unreal Engine uses the same headers for iOS and tvOS; consequenty, we
+//       use PLATFORM defines in this _rare_ instance to differentiate shader
+//       platform types.
+//------------------------------------------------------------------------------
+
+
+#if PLATFORM_TVOS
+#define SHADER_PLATFORM_METAL_MRT	SP_METAL_MRT_TVOS
+#define SHADER_PLATFORM_METAL		SP_METAL_TVOS
+
+#else // PLATFORM_TVOS
+#define SHADER_PLATFORM_METAL_MRT	SP_METAL_MRT
+#define SHADER_PLATFORM_METAL		SP_METAL
+
+#endif // PLATFORM_TVOS
+
+
+//------------------------------------------------------------------------------
+// MARK: - iOS/tvOS Platform Dynamic RHI Routines
+//
+
+namespace UE
+{
+namespace FIOSPlatformDynamicRHI
+{
+
+bool ShouldPreferFeatureLevelES31()
+{
+	return FParse::Param(FCommandLine::Get(), TEXT("metal"));
+}
+
+FORCEINLINE bool ShouldSupportMetalMRT()
+{
+	bool bSupportsMetalMRT = false;
+	GConfig->GetBool(TEXT("/Script/IOSRuntimeSettings.IOSRuntimeSettings"), TEXT("bSupportsMetalMRT"), bSupportsMetalMRT, GEngineIni);
+	return (bSupportsMetalMRT || FParse::Param(FCommandLine::Get(), TEXT("metalmrt"))) && !ShouldPreferFeatureLevelES31();
+}
+
+void AddTargetedShaderFormats(TArray<FString>& TargetedShaderFormats)
+{
+	if (ShouldSupportMetalMRT())
+	{
+		TargetedShaderFormats.Add(LegacyShaderPlatformToShaderFormat(SHADER_PLATFORM_METAL_MRT).ToString());
+	}
+	TargetedShaderFormats.Add(LegacyShaderPlatformToShaderFormat(SHADER_PLATFORM_METAL).ToString());
+}
+
+} // namespace FIOSPlatformDynamicRHI
+} // namespace UE
+
+namespace FPlatformDynamicRHI = UE::FIOSPlatformDynamicRHI;
