U5 - Delete Resources on the RHI Thread to fix layering violation

diff --git a/Engine/Plugins/Experimental/GPULightmass/Source/GPULightmass/Private/VolumetricLightmap.cpp b/Engine/Plugins/Experimental/GPULightmass/Source/GPULightmass/Private/VolumetricLightmap.cpp
--- a/Engine/Plugins/Experimental/GPULightmass/Source/GPULightmass/Private/VolumetricLightmap.cpp
+++ b/Engine/Plugins/Experimental/GPULightmass/Source/GPULightmass/Private/VolumetricLightmap.cpp
@@ -106,7 +106,7 @@ void FVolumetricLightmapRenderer::VoxelizeScene()
 	ReleaseBrickData(VolumetricLightmapData.BrickData);
 	ReleaseBrickData( AccumulationBrickData);
 
-	FRHICommandListExecutor::GetImmediateCommandList().ImmediateFlush(EImmediateFlushType::FlushRHIThreadFlushResourcesFlushDeferredDeletes);
+	FRHICommandListExecutor::GetImmediateCommandList().ImmediateFlush(EImmediateFlushType::FlushRHIThreadFlushResources);
 
 	VolumeMin = CombinedImportanceVolume.Min;
 
diff --git a/Engine/Plugins/Tests/RHITests/Source/RHITests/Private/RHIBufferTests.cpp b/Engine/Plugins/Tests/RHITests/Source/RHITests/Private/RHIBufferTests.cpp
--- a/Engine/Plugins/Tests/RHITests/Source/RHITests/Private/RHIBufferTests.cpp
+++ b/Engine/Plugins/Tests/RHITests/Source/RHITests/Private/RHIBufferTests.cpp
@@ -22,7 +22,7 @@ bool FRHIBufferTests::VerifyBufferContents(const TCHAR* TestName, FRHICommandLis
 	}
 
 	// Immediate flush to clean up the staging buffer / other resources
-	RHICmdList.ImmediateFlush(EImmediateFlushType::FlushRHIThreadFlushResourcesFlushDeferredDeletes);
+	RHICmdList.ImmediateFlush(EImmediateFlushType::FlushRHIThreadFlushResources);
 
 	if (!Result)
 	{
diff --git a/Engine/Plugins/Tests/RHITests/Source/RHITests/Private/RHITestsCommon.cpp b/Engine/Plugins/Tests/RHITests/Source/RHITests/Private/RHITestsCommon.cpp
--- a/Engine/Plugins/Tests/RHITests/Source/RHITests/Private/RHITestsCommon.cpp
+++ b/Engine/Plugins/Tests/RHITests/Source/RHITests/Private/RHITestsCommon.cpp
@@ -28,7 +28,7 @@ bool RunOnRenderThreadSynchronous(TFunctionRef<bool(FRHICommandListImmediate&)>
 	});
 
 	// Flush to wait for the above rendering command to complete.
-	FlushRenderingCommands(true);
+	FlushRenderingCommands();
 
 	return bResult;
 }
diff --git a/Engine/Plugins/Tests/RHITests/Source/RHITests/Private/RHITextureTests.cpp b/Engine/Plugins/Tests/RHITests/Source/RHITests/Private/RHITextureTests.cpp
--- a/Engine/Plugins/Tests/RHITests/Source/RHITests/Private/RHITextureTests.cpp
+++ b/Engine/Plugins/Tests/RHITests/Source/RHITests/Private/RHITextureTests.cpp
@@ -57,11 +57,11 @@ bool FRHITextureTests::VerifyTextureContents(const TCHAR* TestName, FRHICommandL
 
 					RHICmdList.UnmapStagingSurface(StagingTexture);
 				}
-				RHICmdList.ImmediateFlush(EImmediateFlushType::FlushRHIThreadFlushResourcesFlushDeferredDeletes);
+				RHICmdList.ImmediateFlush(EImmediateFlushType::FlushRHIThreadFlushResources);
 			}
 		}
 	}
-	RHICmdList.ImmediateFlush(EImmediateFlushType::FlushRHIThreadFlushResourcesFlushDeferredDeletes);
+	RHICmdList.ImmediateFlush(EImmediateFlushType::FlushRHIThreadFlushResources);
 
 	if (bResult)
 	{
diff --git a/Engine/Plugins/Tests/RHITests/Source/RHITests/Public/RHIBufferTests.h b/Engine/Plugins/Tests/RHITests/Source/RHITests/Public/RHIBufferTests.h
--- a/Engine/Plugins/Tests/RHITests/Source/RHITests/Public/RHIBufferTests.h
+++ b/Engine/Plugins/Tests/RHITests/Source/RHITests/Public/RHIBufferTests.h
@@ -80,7 +80,7 @@ class FRHIBufferTests
 		FShaderResourceViewRHIRef SRV = RHICreateShaderResourceView(VertexBuffer, GPixelFormats[Format].BlockBytes, Format);
 		bool bResult = RunTest_UAVClear_Buffer(RHICmdList, TestName, VertexBuffer.GetReference(), UAV, BufferSize, ClearValue, ClearPtr, TestValue);
 
-		RHICmdList.ImmediateFlush(EImmediateFlushType::FlushRHIThreadFlushResourcesFlushDeferredDeletes);
+		RHICmdList.ImmediateFlush(EImmediateFlushType::FlushRHIThreadFlushResources);
 
 		return bResult;
 	}
@@ -102,7 +102,7 @@ class FRHIBufferTests
 		FShaderResourceViewRHIRef SRV = RHICreateShaderResourceView(VertexBuffer, GPixelFormats[Format].BlockBytes, Format);
 		bool bResult = RunTest_UAVClear_Buffer(RHICmdList, TestName, VertexBuffer.GetReference(), UAV, BufferSize, ClearValue, ClearPtr, TestValue);
 
-		RHICmdList.ImmediateFlush(EImmediateFlushType::FlushRHIThreadFlushResourcesFlushDeferredDeletes);
+		RHICmdList.ImmediateFlush(EImmediateFlushType::FlushRHIThreadFlushResources);
 
 		return bResult;
 	}
@@ -120,7 +120,7 @@ class FRHIBufferTests
 		FShaderResourceViewRHIRef SRV = RHICreateShaderResourceView(StructuredBuffer);
 		bool bResult = RunTest_UAVClear_Buffer(RHICmdList, TestName, StructuredBuffer.GetReference(), UAV, BufferSize, ClearValue, ClearPtr, TestValue);
 
-		RHICmdList.ImmediateFlush(EImmediateFlushType::FlushRHIThreadFlushResourcesFlushDeferredDeletes);
+		RHICmdList.ImmediateFlush(EImmediateFlushType::FlushRHIThreadFlushResources);
 
 		return bResult;
 	}
diff --git a/Engine/Plugins/Tests/RHITests/Source/RHITests/Public/RHITextureTests.h b/Engine/Plugins/Tests/RHITests/Source/RHITests/Public/RHITextureTests.h
--- a/Engine/Plugins/Tests/RHITests/Source/RHITests/Public/RHITextureTests.h
+++ b/Engine/Plugins/Tests/RHITests/Source/RHITests/Public/RHITextureTests.h
@@ -78,7 +78,7 @@ public:
 			});
 		}
 
-		RHICmdList.ImmediateFlush(EImmediateFlushType::FlushRHIThreadFlushResourcesFlushDeferredDeletes);
+		RHICmdList.ImmediateFlush(EImmediateFlushType::FlushRHIThreadFlushResources);
 
 		return bResult0 && bResult1;
 	}
@@ -106,7 +106,7 @@ public:
 				RUN_TEST(RunTest_UAVClear_Texture(RHICmdList, *TestName, Texture.GetReference(), Mip, ClearValue, ClearPtr, TestValue));
 			}
 		}
-		RHICmdList.ImmediateFlush(EImmediateFlushType::FlushRHIThreadFlushResourcesFlushDeferredDeletes);
+		RHICmdList.ImmediateFlush(EImmediateFlushType::FlushRHIThreadFlushResources);
 
 		return bResult;
 	}
@@ -179,7 +179,7 @@ public:
 				RUN_TEST(RunTest_UAVClear_Texture(RHICmdList, *TestName, Texture.GetReference(), Mip, ClearValue, ClearPtr, TestValue));
 			}
 		}
-		RHICmdList.ImmediateFlush(EImmediateFlushType::FlushRHIThreadFlushResourcesFlushDeferredDeletes);
+		RHICmdList.ImmediateFlush(EImmediateFlushType::FlushRHIThreadFlushResources);
 
 		return bResult;
 	}
diff --git a/Engine/Source/Runtime/Apple/MetalRHI/Private/MetalDynamicRHI.cpp b/Engine/Source/Runtime/Apple/MetalRHI/Private/MetalDynamicRHI.cpp
--- a/Engine/Source/Runtime/Apple/MetalRHI/Private/MetalDynamicRHI.cpp
+++ b/Engine/Source/Runtime/Apple/MetalRHI/Private/MetalDynamicRHI.cpp
@@ -51,8 +51,7 @@ FGraphicsPipelineStateRHIRef FMetalDynamicRHI::RHICreateGraphicsPipelineState(co
 		if(!State->Compile())
 		{
 			// Compilation failures are propagated up to the caller.
-			State->DoNoDeferDelete();
-			delete State;
+			State->Delete();
 			return nullptr;
 		}
 
diff --git a/Engine/Source/Runtime/Core/Public/Experimental/Containers/HazardPointer.h b/Engine/Source/Runtime/Core/Public/Experimental/Containers/HazardPointer.h
--- a/Engine/Source/Runtime/Core/Public/Experimental/Containers/HazardPointer.h
+++ b/Engine/Source/Runtime/Core/Public/Experimental/Containers/HazardPointer.h
@@ -295,4 +295,10 @@ public:
 	{
 		return Hazard != nullptr && Record != nullptr;
 	}
-};
\ No newline at end of file
+};
+
+template<typename H>
+THazardPointer<H, false> MakeHazardPointer(std::atomic<H*>& InHazard, FHazardPointerCollection& Collection)
+{
+	return {InHazard, Collection};
+}
\ No newline at end of file
diff --git a/Engine/Source/Runtime/D3D12RHI/Private/D3D12UniformBuffer.cpp b/Engine/Source/Runtime/D3D12RHI/Private/D3D12UniformBuffer.cpp
--- a/Engine/Source/Runtime/D3D12RHI/Private/D3D12UniformBuffer.cpp
+++ b/Engine/Source/Runtime/D3D12RHI/Private/D3D12UniformBuffer.cpp
@@ -210,7 +210,7 @@ void FD3D12DynamicRHI::RHIUpdateUniformBuffer(FRHIUniformBuffer* UniformBufferRH
 
 FD3D12UniformBuffer::~FD3D12UniformBuffer()
 {
-	check(!GRHISupportsRHIThread || IsInRenderingThread());
+	check(IsInRHIThread() || IsInRenderingThread());
 
 	int64 BufferSize = ResourceLocation.GetSize();
 	UpdateBufferStats(GET_STATFNAME(STAT_UniformBufferMemory), -BufferSize);
diff --git a/Engine/Source/Runtime/D3D12RHI/Private/D3D12Viewport.cpp b/Engine/Source/Runtime/D3D12RHI/Private/D3D12Viewport.cpp
--- a/Engine/Source/Runtime/D3D12RHI/Private/D3D12Viewport.cpp
+++ b/Engine/Source/Runtime/D3D12RHI/Private/D3D12Viewport.cpp
@@ -334,13 +334,6 @@ FD3D12Texture2D* GetSwapChainSurface(FD3D12Device* Parent, EPixelFormat PixelFor
 		FD3D12ShaderResourceView* WrappedShaderResourceView = new FD3D12ShaderResourceView(Device, SRVDesc, NewTexture);
 		NewTexture->SetShaderResourceView(WrappedShaderResourceView);
 
-		if (Device->GetGPUIndex() == Parent->GetGPUIndex())
-		{
-			NewTexture->DoNoDeferDelete();
-			BackBufferRenderTargetView->DoNoDeferDelete();
-			WrappedShaderResourceView->DoNoDeferDelete();
-		}
-
 		return NewTexture;
 	});
 
@@ -362,7 +355,6 @@ FD3D12Texture2D* FD3D12Viewport::CreateDummyBackBufferTextures(FD3D12Adapter* In
 	FD3D12Texture2D* Result = InAdapter->CreateLinkedObject<FD3D12Texture2D>(FRHIGPUMask::All(), [&](FD3D12Device* Device)
 	{
 		FD3D12Texture2D* NewTexture = new FD3D12BackBufferReferenceTexture2D(this, bInIsSDR, Device, InSizeX, InSizeY, PixelFormat);
-		NewTexture->DoNoDeferDelete();
 		return NewTexture;
 	});
 	return Result;
@@ -370,7 +362,7 @@ FD3D12Texture2D* FD3D12Viewport::CreateDummyBackBufferTextures(FD3D12Adapter* In
 
 FD3D12Viewport::~FD3D12Viewport()
 {
-	check(IsInRenderingThread());
+	check(IsInRHIThread() || IsInRenderingThread());
 
 #if !PLATFORM_HOLOLENS && D3D12_VIEWPORT_EXPOSES_SWAP_CHAIN
 	// If the swap chain was in fullscreen mode, switch back to windowed before releasing the swap chain.
@@ -486,7 +478,6 @@ void FD3D12Viewport::Resize(uint32 InSizeX, uint32 InSizeY, bool bInIsFullscreen
 
 			for (FD3D12TextureBase& Tex : *BackBuffers[i])
 			{
-				static_cast<FD3D12Texture2D&>(Tex).DoNoDeferDelete();
 				Tex.GetResource()->DoNotDeferDelete();
 			}
 		}
@@ -504,7 +495,6 @@ void FD3D12Viewport::Resize(uint32 InSizeX, uint32 InSizeY, bool bInIsFullscreen
 
 			for (FD3D12TextureBase& Tex : *SDRBackBuffers[i])
 			{
-				static_cast<FD3D12Texture2D&>(Tex).DoNoDeferDelete();
 				Tex.GetResource()->DoNotDeferDelete();
 			}
 		}
@@ -513,6 +503,11 @@ void FD3D12Viewport::Resize(uint32 InSizeX, uint32 InSizeY, bool bInIsFullscreen
 		check(SDRBackBuffers[i] == nullptr);
 	}
 
+	// Flush the outstanding GPU work and wait for it to complete.
+	FlushRenderingCommands();
+	FRHICommandListExecutor::CheckNoOutstandingCmdLists();
+	Adapter->BlockUntilIdle();
+
 	// Keep the current pixel format if one wasn't specified.
 	if (PreferredPixelFormat == PF_Unknown)
 	{
diff --git a/Engine/Source/Runtime/Engine/Private/Rendering/StreamableTextureResource.cpp b/Engine/Source/Runtime/Engine/Private/Rendering/StreamableTextureResource.cpp
--- a/Engine/Source/Runtime/Engine/Private/Rendering/StreamableTextureResource.cpp
+++ b/Engine/Source/Runtime/Engine/Private/Rendering/StreamableTextureResource.cpp
@@ -225,11 +225,6 @@ void FStreamableTextureResource::FinalizeStreaming(FRHITexture* InTextureRHI)
 		STAT(IncrementTextureStats());
 	}
 
-	if (GRHIForceNoDeletionLatencyForStreamingTextures && !GRHIValidationEnabled)
-	{
-		TextureRHI->DoNoDeferDelete();
-	}
-
 	TextureRHI = InTextureRHI;
 	if (ensure(TextureReferenceRHI.IsValid()))
 	{
diff --git a/Engine/Source/Runtime/Engine/Private/ShaderCompiler/ShaderCompiler.cpp b/Engine/Source/Runtime/Engine/Private/ShaderCompiler/ShaderCompiler.cpp
--- a/Engine/Source/Runtime/Engine/Private/ShaderCompiler/ShaderCompiler.cpp
+++ b/Engine/Source/Runtime/Engine/Private/ShaderCompiler/ShaderCompiler.cpp
@@ -2474,7 +2474,7 @@ void FShaderCompilerStats::WriteStats(FOutputDevice* Ar)
 #if ALLOW_DEBUG_FILES
 	static TCHAR DebugText[] = TEXT("Wrote shader compile stats to file '%s'.");
 	{
-		FlushRenderingCommands(true);
+		FlushRenderingCommands();
 
 		FString FileName = FPaths::Combine(*FPaths::ProjectSavedDir(), FString::Printf(TEXT("MaterialStats/Stats-%s.csv"), *FDateTime::Now().ToString()));
 		auto DebugWriter = IFileManager::Get().CreateFileWriter(*FileName);
diff --git a/Engine/Source/Runtime/RHI/Private/DynamicRHI.cpp b/Engine/Source/Runtime/RHI/Private/DynamicRHI.cpp
--- a/Engine/Source/Runtime/RHI/Private/DynamicRHI.cpp
+++ b/Engine/Source/Runtime/RHI/Private/DynamicRHI.cpp
@@ -335,7 +335,7 @@ void RHIExit()
 	else if (GUsingNullRHI)
 	{
 		// If we are using NullRHI flush the command list here in case somethings has been added to the command list during exit calls
-		FRHICommandListExecutor::GetImmediateCommandList().ImmediateFlush(EImmediateFlushType::FlushRHIThreadFlushResourcesFlushDeferredDeletes);
+		FRHICommandListExecutor::GetImmediateCommandList().ImmediateFlush(EImmediateFlushType::FlushRHIThreadFlushResources);
 	}
 }
 
@@ -488,7 +488,7 @@ FDefaultRHIRenderQueryPool::FDefaultRHIRenderQueryPool(ERenderQueryType InQueryT
 
 FDefaultRHIRenderQueryPool::~FDefaultRHIRenderQueryPool()
 {
-	check(IsInRenderingThread());
+	check(IsInRHIThread() || IsInRenderingThread());
 	checkf(AllocatedQueries == Queries.Num(), TEXT("Querypool deleted before all Queries have been released"));
 }
 
diff --git a/Engine/Source/Runtime/RHI/Private/RHI.cpp b/Engine/Source/Runtime/RHI/Private/RHI.cpp
--- a/Engine/Source/Runtime/RHI/Private/RHI.cpp
+++ b/Engine/Source/Runtime/RHI/Private/RHI.cpp
@@ -293,10 +293,10 @@ const FClearValueBinding FClearValueBinding::Green(FLinearColor(0.0f, 1.0f, 0.0f
 // Note: this is used as the default normal for DBuffer decals.  It must decode to a value of 0 in DecodeDBufferData.
 const FClearValueBinding FClearValueBinding::DefaultNormal8Bit(FLinearColor(128.0f / 255.0f, 128.0f / 255.0f, 128.0f / 255.0f, 1.0f));
 
-TLockFreePointerListUnordered<FRHIResource, PLATFORM_CACHE_LINE_SIZE> FRHIResource::PendingDeletes;
+std::atomic<TClosableMpscQueue<FRHIResource*>*> FRHIResource::PendingDeletes { new TClosableMpscQueue<FRHIResource*>() };
+FHazardPointerCollection FRHIResource::PendingDeletesHPC;
 FRHIResource* FRHIResource::CurrentlyDeleting = nullptr;
-TArray<FRHIResource::ResourcesToDelete> FRHIResource::DeferredDeletionQueue;
-uint32 FRHIResource::CurrentFrame = 0;
+
 RHI_API FDrawCallCategoryName* FDrawCallCategoryName::Array[FDrawCallCategoryName::MAX_DRAWCALL_CATEGORY];
 RHI_API int32 FDrawCallCategoryName::DisplayCounts[FDrawCallCategoryName::MAX_DRAWCALL_CATEGORY][MAX_NUM_GPUS];
 RHI_API int32 FDrawCallCategoryName::NumCategory = 0;
@@ -964,7 +964,7 @@ bool FRHIResource::Bypass()
 
 DECLARE_CYCLE_STAT(TEXT("Delete Resources"), STAT_DeleteResources, STATGROUP_RHICMDLIST);
 
-void FRHIResource::FlushPendingDeletes(bool bFlushDeferredDeletes)
+void FRHIResource::FlushPendingDeletes()
 {
 	SCOPE_CYCLE_COUNTER(STAT_DeleteResources);
 
@@ -979,97 +979,36 @@ void FRHIResource::FlushPendingDeletes(bool bFlushDeferredDeletes)
 		RHICmdList.SubmitCommandsHint();
 	}
 #endif
-	RHICmdList.ImmediateFlush(EImmediateFlushType::FlushRHIThread);
-	FRHICommandListExecutor::CheckNoOutstandingCmdLists();
-	if (GDynamicRHI)
-	{
-		GDynamicRHI->RHIPerFrameRHIFlushComplete();
-	}
 
-	auto Delete = [](TArray<FRHIResource*>& ToDelete)
-	{
-		for (int32 Index = 0; Index < ToDelete.Num(); Index++)
-		{
-			FRHIResource* Ref = ToDelete[Index];
-			check(Ref->MarkedForDelete == 1);
-			if (Ref->GetRefCount() == 0) // caches can bring dead objects back to life
-			{
-				CurrentlyDeleting = Ref;
-				delete Ref;
-				CurrentlyDeleting = nullptr;
-			}
-			else
-			{
-				Ref->MarkedForDelete = 0;
-				FPlatformMisc::MemoryBarrier();
-			}
-		}
-	};
+	TArray<FRHIResource*> DeletedResources;
+	TClosableMpscQueue<FRHIResource*>* PendingDeletesPtr = PendingDeletes.exchange(new TClosableMpscQueue<FRHIResource*>());
+	PendingDeletesPtr->Close([&DeletedResources](FRHIResource* Resource)
+	{	
+		DeletedResources.Push(Resource);		
+	});
+	PendingDeletesHPC.Delete(PendingDeletesPtr);
 
-	while (1)
+	RHICmdList.EnqueueLambda([DeletedResources = MoveTemp(DeletedResources)](FRHICommandListImmediate& RHICmdList) mutable
 	{
-		if (PendingDeletes.IsEmpty())
-		{
-			break;
-		}
-		if (PlatformNeedsExtraDeletionLatency())
-		{
-			const int32 Index = DeferredDeletionQueue.AddDefaulted();
-			ResourcesToDelete& ResourceBatch = DeferredDeletionQueue[Index];
-			ResourceBatch.FrameDeleted = CurrentFrame;
-			PendingDeletes.PopAll(ResourceBatch.Resources);
-			check(ResourceBatch.Resources.Num());
-		}
-		else
+		if (GDynamicRHI)
 		{
-			TArray<FRHIResource*> ToDelete;
-			PendingDeletes.PopAll(ToDelete);
-			check(ToDelete.Num());
-			Delete(ToDelete);
+			GDynamicRHI->RHIPerFrameRHIFlushComplete();
 		}
-	}
-
-	const uint32 NumFramesToExpire = RHIRESOURCE_NUM_FRAMES_TO_EXPIRE;
-
-	if (DeferredDeletionQueue.Num())
-	{
-		if (bFlushDeferredDeletes)
-		{
-			FRHICommandListExecutor::GetImmediateCommandList().BlockUntilGPUIdle();
 
-			for (int32 Idx = 0; Idx < DeferredDeletionQueue.Num(); ++Idx)
-			{
-				ResourcesToDelete& ResourceBatch = DeferredDeletionQueue[Idx];
-				Delete(ResourceBatch.Resources);
-			}
-
-			DeferredDeletionQueue.Empty();
-		}
-		else
+		for (FRHIResource* Resource : DeletedResources)
 		{
-			int32 DeletedBatchCount = 0;
-			while (DeletedBatchCount < DeferredDeletionQueue.Num())
+			check(Resource->MarkedForDelete.load(std::memory_order_relaxed) == 1);
+			if (Resource->NumRefs.load(std::memory_order_acquire) == 0) // caches can bring dead objects back to life
 			{
-				ResourcesToDelete& ResourceBatch = DeferredDeletionQueue[DeletedBatchCount];
-				if (((ResourceBatch.FrameDeleted + NumFramesToExpire) < CurrentFrame) || !GIsRHIInitialized)
-				{
-					Delete(ResourceBatch.Resources);
-					++DeletedBatchCount;
-				}
-				else
-				{
-					break;
-				}
+				FRHIResource::CurrentlyDeleting = Resource;
+				delete Resource;
 			}
-
-			if (DeletedBatchCount)
+			else
 			{
-				DeferredDeletionQueue.RemoveAt(0, DeletedBatchCount);
-			}
+				verify(Resource->MarkedForDelete.exchange(0, std::memory_order_release) == 1);	
+			}	
 		}
-
-		++CurrentFrame;
-	}
+	});
 }
 
 static_assert(ERHIZBuffer::FarPlane != ERHIZBuffer::NearPlane, "Near and Far planes must be different!");
diff --git a/Engine/Source/Runtime/RHI/Public/RHICommandList.h b/Engine/Source/Runtime/RHI/Public/RHICommandList.h
--- a/Engine/Source/Runtime/RHI/Public/RHICommandList.h
+++ b/Engine/Source/Runtime/RHI/Public/RHICommandList.h
@@ -3945,8 +3945,7 @@ namespace EImmediateFlushType
 		DispatchToRHIThread, 
 		WaitForDispatchToRHIThread,
 		FlushRHIThread,
-		FlushRHIThreadFlushResources,
-		FlushRHIThreadFlushResourcesFlushDeferredDeletes
+		FlushRHIThreadFlushResources
 	};
 };
 
diff --git a/Engine/Source/Runtime/RHI/Public/RHICommandList.inl b/Engine/Source/Runtime/RHI/Public/RHICommandList.inl
--- a/Engine/Source/Runtime/RHI/Public/RHICommandList.inl
+++ b/Engine/Source/Runtime/RHI/Public/RHICommandList.inl
@@ -90,6 +90,14 @@ FORCEINLINE_DEBUGGABLE void FRHICommandListImmediate::ImmediateFlush(EImmediateF
 			WaitForDispatch();
 		}
 		break;
+
+	case EImmediateFlushType::FlushRHIThreadFlushResources:
+		{
+			CSV_SCOPED_TIMING_STAT(RHITFlushes, FlushRHIThreadFlushResourcesTotal);
+			PipelineStateCache::FlushResources();
+			FRHIResource::FlushPendingDeletes();
+		}
+	//this fall-through is intentional
 	case EImmediateFlushType::FlushRHIThread:
 		{
 			CSV_SCOPED_TIMING_STAT(RHITFlushes, FlushRHIThreadTotal);
@@ -105,21 +113,6 @@ FORCEINLINE_DEBUGGABLE void FRHICommandListImmediate::ImmediateFlush(EImmediateF
 			WaitForTasks(true); // these are already done, but this resets the outstanding array
 		}
 		break;
-	case EImmediateFlushType::FlushRHIThreadFlushResources:
-	case EImmediateFlushType::FlushRHIThreadFlushResourcesFlushDeferredDeletes:
-		{
-			CSV_SCOPED_TIMING_STAT(RHITFlushes, FlushRHIThreadFlushResourcesTotal);
-			if (HasCommands())
-			{
-				GRHICommandList.ExecuteList(*this);
-			}
-			WaitForDispatch();
-			WaitForRHIThreadTasks();
-			WaitForTasks(true); // these are already done, but this resets the outstanding array
-			PipelineStateCache::FlushResources();
-			FRHIResource::FlushPendingDeletes(FlushType == EImmediateFlushType::FlushRHIThreadFlushResourcesFlushDeferredDeletes);
-		}
-		break;
 	default:
 		check(0);
 	}
diff --git a/Engine/Source/Runtime/RHI/Public/RHIResources.h b/Engine/Source/Runtime/RHI/Public/RHIResources.h
--- a/Engine/Source/Runtime/RHI/Public/RHIResources.h
+++ b/Engine/Source/Runtime/RHI/Public/RHIResources.h
@@ -2,6 +2,8 @@
 
 #pragma once
 
+#include <atomic>
+
 #include "CoreTypes.h"
 #include "Misc/AssertionMacros.h"
 #include "HAL/UnrealMemory.h"
@@ -20,8 +22,8 @@
 #include "Hash/CityHash.h"
 #include "Async/TaskGraphInterfaces.h"
 #include "Serialization/MemoryImage.h"
-
-#define DISABLE_RHI_DEFFERED_DELETE 0
+#include "Experimental/Containers/HazardPointer.h"
+#include "Containers/ClosableMpscQueue.h"
 
 #ifndef RHI_WANT_RESOURCE_INFO
 #define RHI_WANT_RESOURCE_INFO 0
@@ -52,7 +54,9 @@ public:
 
 	virtual ~FRHIResource()
 	{
-		check(PlatformNeedsExtraDeletionLatency() || (NumRefs.GetValue() == 0 && (CurrentlyDeleting == this || bDoNotDeferDelete || Bypass()))); // this should not have any outstanding refs
+		check(IsEngineExitRequested() || CurrentlyDeleting == this);
+		check(NumRefs.load(std::memory_order_relaxed) == 0); // this should not have any outstanding refs
+		CurrentlyDeleting = nullptr;
 
 #if RHI_WANT_RESOURCE_INFO
 		EndTrackingResource(this);
@@ -61,24 +65,27 @@ public:
 
 	FORCEINLINE_DEBUGGABLE uint32 AddRef() const
 	{
-		int32 NewValue = NumRefs.Increment();
+		int32 NewValue = NumRefs.fetch_add(1, std::memory_order_acquire) + 1;
 		checkSlow(NewValue > 0); 
 		return uint32(NewValue);
 	}
 	FORCEINLINE_DEBUGGABLE uint32 Release() const
 	{
-		int32 NewValue = NumRefs.Decrement();
+		int32 NewValue = NumRefs.fetch_sub(1, std::memory_order_release) - 1;
+		check(NewValue >= 0);
+
 		if (NewValue == 0)
 		{
-			if (!DeferDelete())
-			{ 
-				delete this;
-			}
-			else
+			if (MarkedForDelete.exchange(1, std::memory_order_release) == 0)
 			{
-				if (FPlatformAtomics::InterlockedCompareExchange(&MarkedForDelete, 1, 0) == 0)
+				while(true)
 				{
-					PendingDeletes.Push(const_cast<FRHIResource*>(this));
+					auto HP = MakeHazardPointer(PendingDeletes, PendingDeletesHPC);
+					TClosableMpscQueue<FRHIResource*>* PendingDeletesPtr = HP.Get();
+					if(PendingDeletesPtr->Enqueue(const_cast<FRHIResource*>(this)))
+					{
+						break;
+					}
 				}
 			}
 		}
@@ -87,24 +94,12 @@ public:
 	}
 	FORCEINLINE_DEBUGGABLE uint32 GetRefCount() const
 	{
-		int32 CurrentValue = NumRefs.GetValue();
+		int32 CurrentValue = NumRefs.load(std::memory_order_relaxed);
 		checkSlow(CurrentValue >= 0); 
 		return uint32(CurrentValue);
 	}
-	void DoNoDeferDelete()
-	{
-		check(!MarkedForDelete);
-		bDoNotDeferDelete = true;
-		FPlatformMisc::MemoryBarrier();
-		check(!MarkedForDelete);
-	}
 
-	static void FlushPendingDeletes(bool bFlushDeferredDeletes = false);
-
-	FORCEINLINE static bool PlatformNeedsExtraDeletionLatency()
-	{
-		return GRHINeedsExtraDeletionLatency && GIsRHIInitialized;
-	}
+	static void FlushPendingDeletes();
 
 	static bool Bypass();
 
@@ -123,7 +118,14 @@ public:
 
 	bool IsValid() const
 	{
-		return !MarkedForDelete && NumRefs.GetValue() > 0;
+		return !MarkedForDelete.load(std::memory_order_relaxed) && NumRefs.load(std::memory_order_relaxed) > 0;
+	}
+
+	void Delete()
+	{
+		verify(MarkedForDelete.fetch_add(1, std::memory_order_acquire) == 0);
+		CurrentlyDeleting = this;
+		delete this;
 	}
 
 	inline ERHIResourceType GetType() const { return ResourceType; }
@@ -142,31 +144,18 @@ public:
 #endif
 
 private:
-	mutable FThreadSafeCounter NumRefs;
-	mutable int32 MarkedForDelete{ 0 };
-
+	mutable std::atomic_int NumRefs { 0 };
+	mutable std::atomic_int MarkedForDelete { 0 };
 	const ERHIResourceType ResourceType;
-
-	bool bDoNotDeferDelete{ false };
-	bool bCommitted{ true };
+	bool bCommitted { true };
 #if RHI_WANT_RESOURCE_INFO
-	bool bBeingTracked{ false };
+	bool bBeingTracked { false };
 #endif
 
-	static TLockFreePointerListUnordered<FRHIResource, PLATFORM_CACHE_LINE_SIZE> PendingDeletes;
+	static std::atomic<TClosableMpscQueue<FRHIResource*>*> PendingDeletes;
+	static FHazardPointerCollection PendingDeletesHPC;
 	static FRHIResource* CurrentlyDeleting;
 
-	FORCEINLINE bool DeferDelete() const
-	{
-#if DISABLE_RHI_DEFFERED_DELETE
-		checkf(!GRHIValidationEnabled, TEXT("RHI validation is not supported when DISABLE_RHI_DEFERRED_DELETE flag is set."));
-		return false;
-#else
-		// Defer if GRHINeedsExtraDeletionLatency or we are doing threaded rendering (unless otherwise requested).
-		return !bDoNotDeferDelete && (GRHINeedsExtraDeletionLatency || !Bypass());
-#endif
-	}
-
 	// Some APIs don't do internal reference counting, so we have to wait an extra couple of frames before deleting resources
 	// to ensure the GPU has completely finished with them. This avoids expensive fences, etc.
 	struct ResourcesToDelete
@@ -174,9 +163,6 @@ private:
 		TArray<FRHIResource*>	Resources;
 		uint32					FrameDeleted{};
 	};
-
-	static TArray<ResourcesToDelete> DeferredDeletionQueue;
-	static uint32 CurrentFrame;
 };
 
 class FExclusiveDepthStencil
diff --git a/Engine/Source/Runtime/RenderCore/Private/RenderingThread.cpp b/Engine/Source/Runtime/RenderCore/Private/RenderingThread.cpp
--- a/Engine/Source/Runtime/RenderCore/Private/RenderingThread.cpp
+++ b/Engine/Source/Runtime/RenderCore/Private/RenderingThread.cpp
@@ -1253,7 +1253,7 @@ void FRenderCommandFence::Wait(bool bProcessGameThreadTasks) const
 /**
  * Waits for the rendering thread to finish executing all pending rendering commands.  Should only be used from the game thread.
  */
-void FlushRenderingCommands(bool bFlushDeferredDeletes)
+void FlushRenderingCommands()
 {
 	if (!GIsRHIInitialized)
 	{
@@ -1273,13 +1273,9 @@ void FlushRenderingCommands(bool bFlushDeferredDeletes)
 		FTaskGraphInterface::Get().ProcessThreadUntilIdle(ENamedThreads::GameThread_Local);
 	}
 
-	ENQUEUE_RENDER_COMMAND(FlushPendingDeleteRHIResourcesCmd)(
-		[bFlushDeferredDeletes](FRHICommandListImmediate& RHICmdList)
+	ENQUEUE_RENDER_COMMAND(FlushPendingDeleteRHIResourcesCmd)([](FRHICommandListImmediate& RHICmdList)
 	{
-		RHICmdList.ImmediateFlush(
-			bFlushDeferredDeletes ?
-			EImmediateFlushType::FlushRHIThreadFlushResourcesFlushDeferredDeletes :
-			EImmediateFlushType::FlushRHIThreadFlushResources);
+		RHICmdList.ImmediateFlush(EImmediateFlushType::FlushRHIThreadFlushResources);
 	});
 
 	// Find the objects which may be cleaned up once the rendering thread command queue has been flushed.
diff --git a/Engine/Source/Runtime/RenderCore/Public/RenderingThread.h b/Engine/Source/Runtime/RenderCore/Public/RenderingThread.h
--- a/Engine/Source/Runtime/RenderCore/Public/RenderingThread.h
+++ b/Engine/Source/Runtime/RenderCore/Public/RenderingThread.h
@@ -67,7 +67,7 @@ extern RENDERCORE_API void AdvanceRenderingThreadStatsGT( bool bDiscardCallstack
 /**
  * Waits for the rendering thread to finish executing all pending rendering commands.  Should only be used from the game thread.
  */
-extern RENDERCORE_API void FlushRenderingCommands(bool bFlushDeferredDeletes = false);
+extern RENDERCORE_API void FlushRenderingCommands();
 
 extern RENDERCORE_API void FlushPendingDeleteRHIResources_GameThread();
 extern RENDERCORE_API void FlushPendingDeleteRHIResources_RenderThread();
diff --git a/Engine/Source/Runtime/SlateRHIRenderer/Private/SlateRHIRenderer.cpp b/Engine/Source/Runtime/SlateRHIRenderer/Private/SlateRHIRenderer.cpp
--- a/Engine/Source/Runtime/SlateRHIRenderer/Private/SlateRHIRenderer.cpp
+++ b/Engine/Source/Runtime/SlateRHIRenderer/Private/SlateRHIRenderer.cpp
@@ -515,7 +515,7 @@ void FSlateRHIRenderer::OnWindowDestroyed(const TSharedRef<SWindow>& InWindow)
 
 		// Need to flush rendering commands as the viewport may be in use by the render thread
 		// and the rendering resources must be released on the render thread before the viewport can be deleted
-		FlushRenderingCommands(true /* bFlushDeferredDeletes */);
+		FlushRenderingCommands();
 
 		delete *ViewportInfoPtr;
 	}
diff --git a/Engine/Source/Runtime/VulkanRHI/Private/VulkanPipeline.cpp b/Engine/Source/Runtime/VulkanRHI/Private/VulkanPipeline.cpp
--- a/Engine/Source/Runtime/VulkanRHI/Private/VulkanPipeline.cpp
+++ b/Engine/Source/Runtime/VulkanRHI/Private/VulkanPipeline.cpp
@@ -1768,7 +1768,7 @@ FVulkanRHIGraphicsPipelineState* FVulkanPipelineStateCacheManager::RHICreateGrap
 			FVulkanRHIGraphicsPipelineState** MapPSO = GraphicsPSOLockedMap.Find(Key);
 			if(MapPSO)//another thread could end up creating it.
 			{
-				delete NewPSO;
+				NewPSO->Delete();
 				NewPSO = *MapPSO;
 			}
 			else
diff --git a/Engine/Source/Runtime/Windows/D3D11RHI/Private/D3D11Viewport.cpp b/Engine/Source/Runtime/Windows/D3D11RHI/Private/D3D11Viewport.cpp
--- a/Engine/Source/Runtime/Windows/D3D11RHI/Private/D3D11Viewport.cpp
+++ b/Engine/Source/Runtime/Windows/D3D11RHI/Private/D3D11Viewport.cpp
@@ -204,8 +204,6 @@ FD3D11Texture2D* FD3D11Viewport::GetSwapChainSurface(FD3D11DynamicRHI* D3DRHI, E
 
 	D3D11TextureAllocated2D(*NewTexture);
 
-	NewTexture->DoNoDeferDelete();
-
 	return NewTexture;
 }
 
@@ -268,6 +266,10 @@ void FD3D11Viewport::Resize(uint32 InSizeX, uint32 InSizeY, bool bInIsFullscreen
 	}
 	BackBuffer.SafeRelease();
 
+	// Flush the outstanding GPU work and wait for it to complete.
+	FlushRenderingCommands();
+	FRHICommandListExecutor::CheckNoOutstandingCmdLists();
+
 	// Make sure we use a format the current device supports.
 	PreferredPixelFormat = D3DRHI->GetDisplayFormat(PreferredPixelFormat);
 
