U5 - defer delete PSO cache entries on the RHI CmdList timeline in conjunction with the recent change of moving RHIResources on the same timeline.

diff --git a/Engine/Source/Runtime/RHI/Private/PipelineStateCache.cpp b/Engine/Source/Runtime/RHI/Private/PipelineStateCache.cpp
--- a/Engine/Source/Runtime/RHI/Private/PipelineStateCache.cpp
+++ b/Engine/Source/Runtime/RHI/Private/PipelineStateCache.cpp
@@ -535,35 +535,41 @@ public:
 	{
 		check(IsInRenderingThread());
 
-		for (TMyValue& OldPipelineState : DeleteArray)
+		FRHICommandListImmediate& RHICmdList = FRHICommandListExecutor::GetImmediateCommandList();
+		RHICmdList.EnqueueLambda([DeleteArray = MoveTemp(DeleteArray)](FRHICommandListImmediate& RHICmdList) mutable
 		{
-			//once in the delayed list this object should not be findable anymore, so the 0 should remain, making this safe
+			for (TMyValue& OldPipelineState : DeleteArray)
+			{
+				//once in the delayed list this object should not be findable anymore, so the 0 should remain, making this safe
 #if PIPELINESTATECACHE_VERIFYTHREADSAFE
-			check(OldPipelineState->InUseCount.GetValue() == 0);
+				check(OldPipelineState->InUseCount.GetValue() == 0);
 #endif
-			delete OldPipelineState;
-		}
-		DeleteArray.Empty();
+				delete OldPipelineState;
+			}
+			DeleteArray.Empty();
+		});
 	}
 
 
 	int32 DiscardAndSwap()
 	{
+		check(IsInRenderingThread());
+
 		// the consolidate should always be run before the DiscardAndSwap.
 		// there should be no inuse pipeline states in the backfill map (because they should have been moved into the CurrentMap).
 		int32 Discarded = BackfillMap->Num();
 
-
-		for ( const auto& DiscardIterator : *BackfillMap )
+		FRHICommandListImmediate& RHICmdList = FRHICommandListExecutor::GetImmediateCommandList();
+		RHICmdList.EnqueueLambda([BackfillMap = MoveTemp(*BackfillMap)](FRHICommandListImmediate& RHICmdList) mutable
 		{
+			for ( const auto& DiscardIterator :  BackfillMap)
+			{
 #if PIPELINESTATECACHE_VERIFYTHREADSAFE
-			check( DiscardIterator.Value->InUseCount.GetValue() == 0);
+				check( DiscardIterator.Value->InUseCount.GetValue() == 0);
 #endif
-			delete DiscardIterator.Value;
-		}
-
-		BackfillMap->Empty();
-
+				delete DiscardIterator.Value;
+			}
+		});
 
 		if ( CurrentMap == &Map1 )
 		{
