Fix for breadcrumbs with async compute submission.

diff --git a/Engine/Source/Runtime/RHI/Private/RHICommandList.cpp b/Engine/Source/Runtime/RHI/Private/RHICommandList.cpp
--- a/Engine/Source/Runtime/RHI/Private/RHICommandList.cpp
+++ b/Engine/Source/Runtime/RHI/Private/RHICommandList.cpp
@@ -340,6 +340,16 @@ void FRHIAsyncComputeCommandListImmediate::ImmediateDispatch(FRHIAsyncComputeCom
 			SwapCmdList->ExchangeCmdList(RHIComputeCmdList);
 			RHIComputeCmdList.CopyContext(*SwapCmdList);
 			RHIComputeCmdList.GPUMask = SwapCmdList->GPUMask;
+		#if RHI_WANT_BREADCRUMB_EVENTS
+			FRHIBreadcrumbState BreadcrumbState;
+
+			// Once executed, the memory containing the breadcrumbs will be freed, so any open markers are popped and stored into BreadcrumbState
+			SwapCmdList->ExportBreadcrumbState(BreadcrumbState);
+
+			// And then pushed into the newly opened list.
+			RHIComputeCmdList.ImportBreadcrumbState(BreadcrumbState);
+		#endif // RHI_WANT_BREADCRUMB_EVENTS
+
 			// NB: InitialGPUMask set to GPUMask since exchanging the list
 			// is equivalent to a Reset.
 			RHIComputeCmdList.InitialGPUMask = SwapCmdList->GPUMask;
