rhi.DumpResourceMemory csv will now dump to a file in the Profiling directory.

- Extend the size precision to .9f to account for very small allocations.

diff --git a/Engine/Source/Runtime/RHI/Private/RHI.cpp b/Engine/Source/Runtime/RHI/Private/RHI.cpp
--- a/Engine/Source/Runtime/RHI/Private/RHI.cpp
+++ b/Engine/Source/Runtime/RHI/Private/RHI.cpp
@@ -15,6 +15,7 @@
 #include "String/LexFromString.h"
 #include "String/ParseTokens.h"
 #include "Misc/BufferedOutputDevice.h"
+#include "Misc/OutputDeviceFile.h"
 
 IMPLEMENT_MODULE(FDefaultModuleImpl, RHI);
 
@@ -786,6 +787,7 @@ static FAutoConsoleCommandWithWorldArgsAndOutputDevice GDumpRHIResourceMemoryCmd
 	ERHIResourceType TypeFilter = RRT_None;
 	int32 NumberOfResourcesToShow = 50;
 	bool bUseCSVOutput = false;
+	FArchive* CSVFile{ nullptr };
 
 	for (const FString& Argument : Args)
 	{
@@ -796,6 +798,9 @@ static FAutoConsoleCommandWithWorldArgsAndOutputDevice GDumpRHIResourceMemoryCmd
 		else if (Argument.Equals(TEXT("csv"), ESearchCase::IgnoreCase))
 		{
 			bUseCSVOutput = true;
+
+			const FString Filename = FString::Printf(TEXT("%srhiDumpResourceMemory-%s.csv"), *FPaths::ProfilingDir(), *FDateTime::Now().ToString());
+			CSVFile = IFileManager::Get().CreateFileWriter(*Filename, FILEWRITE_AllowRead);
 		}
 		else if (Argument.StartsWith(TEXT("Name="), ESearchCase::IgnoreCase))
 		{
@@ -895,7 +900,8 @@ static FAutoConsoleCommandWithWorldArgsAndOutputDevice GDumpRHIResourceMemoryCmd
 
 	if (bUseCSVOutput)
 	{
-		BufferedOutput.Logf(ELogVerbosity::Log, TEXT("\"Name\",\"Type\",\"Size\""));
+		const TCHAR* Header = TEXT("Name,Type,Size\n");
+		CSVFile->Serialize(TCHAR_TO_ANSI(Header), FPlatformString::Strlen(Header));
 	}
 	else
 	{
@@ -921,14 +927,16 @@ static FAutoConsoleCommandWithWorldArgsAndOutputDevice GDumpRHIResourceMemoryCmd
 
 			if (bUseCSVOutput)
 			{
-				BufferedOutput.Logf(ELogVerbosity::Log, TEXT("\"%s\",\"%s\",\"%.3f\""),
+				const FString Row = FString::Printf(TEXT("%s,%s,%.9f\n"),
 					ResourceNameBuffer,
 					ResourceType,
 					SizeInBytes / double(1 << 20));
+
+				CSVFile->Serialize(TCHAR_TO_ANSI(*Row), Row.Len());
 			}
 			else
 			{
-				BufferedOutput.CategorizedLogf(CategoryName, ELogVerbosity::Log, TEXT("Name: %s - Type: %s - Size: %.3f MB"),
+				BufferedOutput.CategorizedLogf(CategoryName, ELogVerbosity::Log, TEXT("Name: %s - Type: %s - Size: %.9f MB"),
 					ResourceNameBuffer,
 					ResourceType,
 					SizeInBytes / double(1 << 20));
@@ -941,15 +949,20 @@ static FAutoConsoleCommandWithWorldArgsAndOutputDevice GDumpRHIResourceMemoryCmd
 	const double TotalSizeF = TotalTrackedResourceSize / double(1 << 20);
 	const double ShownSizeF = TotalShownResourceSize / double(1 << 20);
 
-	if (!bUseCSVOutput)
+	if (bUseCSVOutput)
+	{
+		delete CSVFile;
+		CSVFile = nullptr;
+	}
+	else
 	{
 		if (NumberOfResourcesToShow != TotalResourcesWithInfo)
 		{
-			BufferedOutput.CategorizedLogf(CategoryName, ELogVerbosity::Log, TEXT("Shown %d entries. Size: %.3f MB (%.2f%% of total)"),
+			BufferedOutput.CategorizedLogf(CategoryName, ELogVerbosity::Log, TEXT("Shown %d entries. Size: %.9f MB (%.2f%% of total)"),
 				NumberOfResourcesToShow, ShownSizeF, 100.0 * ShownSizeF / TotalSizeF);
 		}
 
-		BufferedOutput.CategorizedLogf(CategoryName, ELogVerbosity::Log, TEXT("Total tracked resource size: %.3f MB"), TotalSizeF);
+		BufferedOutput.CategorizedLogf(CategoryName, ELogVerbosity::Log, TEXT("Total tracked resource size: %.9f MB"), TotalSizeF);
 	}
 
 	BufferedOutput.RedirectTo(OutputDevice);
