Make sure we stop tracking all RHI resources before shutdown to avoid using the CS during static destruction after it's been destroyed.
Making sure we clean up a few more system textures before shutdown.

diff --git a/Engine/Source/Runtime/RHI/Private/DynamicRHI.cpp b/Engine/Source/Runtime/RHI/Private/DynamicRHI.cpp
--- a/Engine/Source/Runtime/RHI/Private/DynamicRHI.cpp
+++ b/Engine/Source/Runtime/RHI/Private/DynamicRHI.cpp
@@ -331,6 +331,10 @@ void RHIExit()
 		GDynamicRHI->Shutdown();
 		delete GDynamicRHI;
 		GDynamicRHI = NULL;
+
+#if RHI_WANT_RESOURCE_INFO
+		FRHIResource::StopTrackingAllResources();
+#endif
 	}
 	else if (GUsingNullRHI)
 	{
diff --git a/Engine/Source/Runtime/RHI/Private/RHI.cpp b/Engine/Source/Runtime/RHI/Private/RHI.cpp
--- a/Engine/Source/Runtime/RHI/Private/RHI.cpp
+++ b/Engine/Source/Runtime/RHI/Private/RHI.cpp
@@ -636,14 +636,18 @@ bool operator==(const FBlendStateInitializerRHI::FRenderTarget& A, const FBlendS
 
 static FCriticalSection GRHIResourceTrackingCriticalSection;
 static TSet<FRHIResource*> GRHITrackedResources;
+static bool GRHITrackingResources = true;
 
 void FRHIResource::BeginTrackingResource(FRHIResource* InResource)
 {
-	FScopeLock Lock(&GRHIResourceTrackingCriticalSection);
+	if (GRHITrackingResources)
+	{
+		FScopeLock Lock(&GRHIResourceTrackingCriticalSection);
 
-	InResource->bBeingTracked = true;
+		InResource->bBeingTracked = true;
 
-	GRHITrackedResources.Add(InResource);
+		GRHITrackedResources.Add(InResource);
+	}
 }
 
 void FRHIResource::EndTrackingResource(FRHIResource* InResource)
@@ -656,6 +660,20 @@ void FRHIResource::EndTrackingResource(FRHIResource* InResource)
 	}
 }
 
+void FRHIResource::StopTrackingAllResources()
+{
+	FScopeLock Lock(&GRHIResourceTrackingCriticalSection);
+	for (FRHIResource* Resource : GRHITrackedResources)
+	{
+		if (Resource)
+		{
+			Resource->bBeingTracked = false;
+		}
+	}
+	GRHITrackedResources.Empty();
+	GRHITrackingResources = false;
+}
+
 struct FRHIResourceTypeName
 {
 	ERHIResourceType Type;
diff --git a/Engine/Source/Runtime/RHI/Public/RHIResources.h b/Engine/Source/Runtime/RHI/Public/RHIResources.h
--- a/Engine/Source/Runtime/RHI/Public/RHIResources.h
+++ b/Engine/Source/Runtime/RHI/Public/RHIResources.h
@@ -141,6 +141,7 @@ public:
 
 	static void BeginTrackingResource(FRHIResource* InResource);
 	static void EndTrackingResource(FRHIResource* InResource);
+	static void StopTrackingAllResources();
 #endif
 
 private:
diff --git a/Engine/Source/Runtime/Renderer/Private/SystemTextures.cpp b/Engine/Source/Runtime/Renderer/Private/SystemTextures.cpp
--- a/Engine/Source/Runtime/Renderer/Private/SystemTextures.cpp
+++ b/Engine/Source/Runtime/Renderer/Private/SystemTextures.cpp
@@ -825,6 +825,11 @@ void FSystemTextures::ReleaseDynamicRHI()
 	StencilDummySRV.SafeRelease();
 	GTAOPreIntegrated.SafeRelease();
 
+	DefaultTextures.Empty();
+	DefaultBuffers.Empty();
+	HashDefaultTextures.Clear();
+	HashDefaultBuffers.Clear();
+
 	GRenderTargetPool.FreeUnusedResources();
 
 	// Indicate that textures will need to be reinitialized.
