Fix for potential deadlock on shutdown in FRHIFrameOffsetThread

  - to ensure FRHIFrameOffsetThread::Run() can exit we need to set bRun to false and then signal one final RHISignalFlipEvent() to ensure its not stuck in RHIWaitForFlip()
    - FRHIFrameOffsetThread::Shutdown() does this in the incorrect order
      - RHISignalFlipEvent() is called and then the thread is killed which calls Stop() and sets bRun to false
    - move the RHISignalFlipEvent() call from Shutdown() to Stop() and call it after setting bRun to false to resolve this

  - FRHIFrameOffsetThread::Shutdown() recreates WaitEvent after its been deleted to trigger it one final time
    - it isn't necessary to trigger it a final time because FRHIFrameFlipTrackingRunnable::Stop() has already done this and FRHIFrameFlipTrackingRunnable::Run() is the only consumer
      - so just remove the trigger from FRHIFrameOffsetThread::Shutdown()

#okforghithub public

diff --git a/Engine/Source/Runtime/RHI/Private/RHIUtilities.cpp b/Engine/Source/Runtime/RHI/Private/RHIUtilities.cpp
--- a/Engine/Source/Runtime/RHI/Private/RHIUtilities.cpp
+++ b/Engine/Source/Runtime/RHI/Private/RHIUtilities.cpp
@@ -220,6 +220,7 @@ struct FRHIFrameOffsetThread : public FRunnable
 	virtual void Stop() override
 	{
 		bRun = false;
+		GDynamicRHI->RHISignalFlipEvent();
 	}
 
 public:
@@ -268,7 +269,6 @@ public:
 			return;
 		}
 		bInitialized = false;
-		GDynamicRHI->RHISignalFlipEvent();
 
 		if (WaitEvent)
 		{
@@ -282,8 +282,6 @@ public:
 			delete Thread;
 			Thread = nullptr;
 		}
-
-		Singleton.GetOrInitializeWaitEvent()->Trigger();
 	}
 
 	static void SetFrameDebugInfo(uint64 PresentIndex, uint64 FrameIndex, uint64 InputTime)
