Adding ability to format breadcrumbs if you want to use them manually.
Using manual breadcrumbs to track mesh draw command sources.
Fixing issues caused by recursive Command List execution where the newer Command Lists would end up hiding the higher command lists in the breadcrumb dumps.
Cleaned up a few aspects of breadcrumbs, including moving some functionality to the proper classes.

diff --git a/Engine/Source/Runtime/RHI/Private/RHICommandList.cpp b/Engine/Source/Runtime/RHI/Private/RHICommandList.cpp
--- a/Engine/Source/Runtime/RHI/Private/RHICommandList.cpp
+++ b/Engine/Source/Runtime/RHI/Private/RHICommandList.cpp
@@ -213,7 +213,7 @@ struct FRHIBreadcrumbState
 	TArray<int32, TInlineAllocator<8>> NameOffsets;
 };
 
-void FRHIComputeCommandList::ExportBreadcrumbState(FRHIBreadcrumbState& State)
+void FRHICommandListBase::ExportBreadcrumbState(FRHIBreadcrumbState& State)
 {
 	FRHIBreadcrumb* Breadcrumb = BreadcrumbStackTop;
 	while (Breadcrumb)
@@ -231,7 +231,7 @@ void FRHIComputeCommandList::ExportBreadcrumbState(FRHIBreadcrumbState& State)
 	check(BreadcrumbStackTop == nullptr);
 }
 
-void FRHIComputeCommandList::ImportBreadcrumbState(const FRHIBreadcrumbState& State)
+void FRHICommandListBase::ImportBreadcrumbState(const FRHIBreadcrumbState& State)
 {
 	const int32 NameCount = State.NameOffsets.Num();
 	for (int32 i = 0; i < NameCount; ++i)
@@ -373,47 +373,63 @@ DECLARE_CYCLE_STAT(TEXT("SmallList"), STAT_SmallList, STATGROUP_RHICMDLIST);
 DECLARE_CYCLE_STAT(TEXT("PTrans"), STAT_PTrans, STATGROUP_RHICMDLIST);
 
 #if WITH_ADDITIONAL_CRASH_CONTEXTS && RHI_WANT_BREADCRUMB_EVENTS
-static void WriteRenderBreadcrumbs(FCrashContextExtendedWriter& Writer, FRHIBreadcrumb const* const* BreadcrumbPtr, const TCHAR* ThreadName)
+static void WriteRenderBreadcrumbs(FCrashContextExtendedWriter& Writer, const FRHIBreadcrumb** BreadcrumbStack, uint32 BreadcrumbStackIndex, const TCHAR* ThreadName)
 {
 	enum
 	{
-		MAX_BREADCRUMBS = 128,
-		MAX_BREADCRUMB_STRING = 2048,
+		MAX_BREADCRUMBS = 64,
+		MAX_BREADCRUMB_STRING = 4096,
 		MAX_BREADCRUMB_NAME_STRING = 128,
 	};
 	static int BreadcrumbId = 0;
-	const TCHAR* BreadcrumbNames[MAX_BREADCRUMBS];
 
 	TCHAR StaticBreadcrumbStackString[MAX_BREADCRUMB_STRING];
-	TCHAR StaticBreadcrumbName[MAX_BREADCRUMB_NAME_STRING];
-	TCHAR* BreadcrumbStackString = &StaticBreadcrumbStackString[0];
 	size_t BreadcrumbStackStringSize = 0;
-	BreadcrumbStackStringSize += FCString::Snprintf(&StaticBreadcrumbStackString[BreadcrumbStackStringSize], MAX_BREADCRUMB_STRING, TEXT("Breadcrumbs '%s'\n"), ThreadName);
-	if (BreadcrumbPtr && *BreadcrumbPtr)
+
+	auto WriteBreadcrumbLine = [&](const TCHAR* InFormat, ...)
 	{
-		FRHIBreadcrumb const* Breadcrumb = *BreadcrumbPtr;
-		int32 Index = 0;
-		while (Breadcrumb && Index < MAX_BREADCRUMBS)
-		{
-			BreadcrumbNames[Index] = Breadcrumb->Name;
-			Breadcrumb = Breadcrumb->Parent;
-			Index++;
-		}
-		uint32 StackPos = 0;
 		if (BreadcrumbStackStringSize < MAX_BREADCRUMB_STRING)
 		{
-			BreadcrumbStackStringSize += FCString::Snprintf(&StaticBreadcrumbStackString[BreadcrumbStackStringSize], MAX_BREADCRUMB_STRING - BreadcrumbStackStringSize, TEXT("RHI\n"));
+			int32 WrittenLength = 0;
+			GET_VARARGS_RESULT(
+				&StaticBreadcrumbStackString[BreadcrumbStackStringSize],
+				MAX_BREADCRUMB_STRING - BreadcrumbStackStringSize,
+				MAX_BREADCRUMB_STRING - BreadcrumbStackStringSize - 1,
+				InFormat, InFormat, WrittenLength);
+
+			BreadcrumbStackStringSize += WrittenLength;
 		}
-		for (int32 i = Index - 1; i >= 0; --i)
+	};
+
+	WriteBreadcrumbLine(TEXT("Breadcrumbs '%s'\n"), ThreadName);
+
+	const uint32 NumBreadcrumbStacks = BreadcrumbStackIndex + 1;
+
+	for (uint32 BreadcrumbIndex = 0; BreadcrumbIndex < NumBreadcrumbStacks; BreadcrumbIndex++)
+	{
+		if (const FRHIBreadcrumb* CurrentBreadcrumb = BreadcrumbStack[BreadcrumbStackIndex - BreadcrumbIndex])
 		{
-			if (BreadcrumbStackStringSize < MAX_BREADCRUMB_STRING)
+			const TCHAR* BreadcrumbNames[MAX_BREADCRUMBS];
+
+			const FRHIBreadcrumb* Breadcrumb = CurrentBreadcrumb;
+			int32 NameIndex = 0;
+			while (Breadcrumb && NameIndex < MAX_BREADCRUMBS)
+			{
+				BreadcrumbNames[NameIndex++] = Breadcrumb->Name;
+				Breadcrumb = Breadcrumb->Parent;
+			}
+
+			WriteBreadcrumbLine(TEXT("Context %d/%d\n"), BreadcrumbIndex + 1, NumBreadcrumbStacks);
+
+			uint32 StackPos = 0;
+			for (int32 i = NameIndex - 1; i >= 0; --i)
 			{
-				BreadcrumbStackStringSize += FCString::Snprintf(&StaticBreadcrumbStackString[BreadcrumbStackStringSize], MAX_BREADCRUMB_STRING - BreadcrumbStackStringSize, TEXT("\t%02d %s\n"), StackPos, BreadcrumbNames[i]);
+				WriteBreadcrumbLine(TEXT("\t%02d %s\n"), StackPos++, BreadcrumbNames[i]);
 			}
-			StackPos++;
 		}
 	}
 
+	TCHAR StaticBreadcrumbName[MAX_BREADCRUMB_NAME_STRING];
 	FCString::Snprintf(StaticBreadcrumbName, MAX_BREADCRUMB_NAME_STRING, TEXT("Breadcrumbs_%s_%d"), ThreadName, BreadcrumbId++);
 	Writer.AddString(StaticBreadcrumbName, StaticBreadcrumbStackString);
 	UE_LOG(LogRHI, Error, StaticBreadcrumbStackString);
@@ -430,16 +446,16 @@ void FRHICommandListExecutor::ExecuteInner_DoExecute(FRHICommandListBase& CmdLis
 	FMemMark Mark(FMemStack::Get());
 
 #if WITH_ADDITIONAL_CRASH_CONTEXTS && RHI_WANT_BREADCRUMB_EVENTS
-	const FRHIBreadcrumb* OriginalTop = CmdList.Context ? *CmdList.Context->GetBreadcrumbStackTopRef() : nullptr;
+	IRHIComputeContext* LocalContext = CmdList.Context ? CmdList.Context : CmdList.ComputeContext;
 
 	// Need a struct because Crash context scope only allows 1 argument
 	struct
 	{
 		const TCHAR* ThreadName;
-		FRHIBreadcrumb const* const* BreadcrumbPointer;
-	} CrashState;
+		const FRHIBreadcrumb** BreadcrumbStack;
+		uint32 BreadcrumbStackIndex;
+	} CrashState{};
 	CrashState.ThreadName = TEXT("Parallel");
-	CrashState.BreadcrumbPointer = CmdList.Context ? CmdList.Context->GetBreadcrumbStackTopRef() : nullptr;
 	if (IsInRenderingThread())
 	{
 		CrashState.ThreadName = TEXT("RenderingThread");
@@ -449,16 +465,36 @@ void FRHICommandListExecutor::ExecuteInner_DoExecute(FRHICommandListBase& CmdLis
 		CrashState.ThreadName = TEXT("RHIThread");
 	}
 
-	UE_ADD_CRASH_CONTEXT_SCOPE([&CrashState](FCrashContextExtendedWriter& Writer) {  WriteRenderBreadcrumbs(Writer, CrashState.BreadcrumbPointer, CrashState.ThreadName); });
+	bool PopBreadcrumbStack = false;
 
-	if (CmdList.Context)
+	if (LocalContext)
 	{
-		CmdList.Context->RHISetBreadcrumbStackTop(CmdList.PopFirstUnsubmittedBreadcrumb());
+		if (LocalContext->BreadcrumbStackIndex < IRHIComputeContext::MaxBreadcrumbStacks - 1)
+		{
+			LocalContext->BreadcrumbStackIndex++;
+			PopBreadcrumbStack = true;
+		}
+
+		// if we can't fit a next stack in, we have to stomp the top one, the show must go on.
+
+		LocalContext->RHISetBreadcrumbStackTop(CmdList.PopFirstUnsubmittedBreadcrumb());
+
+		CrashState.BreadcrumbStack = &LocalContext->BreadcrumbStackTop[0];
+		CrashState.BreadcrumbStackIndex = LocalContext->BreadcrumbStackIndex;
 	}
-	else if (CmdList.ComputeContext)
+
+	UE_ADD_CRASH_CONTEXT_SCOPE([&CrashState](FCrashContextExtendedWriter& Writer)
 	{
-		CmdList.ComputeContext->RHISetBreadcrumbStackTop(CmdList.PopFirstUnsubmittedBreadcrumb());
-	}
+		WriteRenderBreadcrumbs(Writer, CrashState.BreadcrumbStack, CrashState.BreadcrumbStackIndex, CrashState.ThreadName);
+	});
+
+	ON_SCOPE_EXIT
+	{
+		if (PopBreadcrumbStack)
+		{
+			LocalContext->BreadcrumbStackIndex--;
+		}
+	};
 #endif
 
 #if WITH_MGPU
diff --git a/Engine/Source/Runtime/RHI/Public/RHICommandList.h b/Engine/Source/Runtime/RHI/Public/RHICommandList.h
--- a/Engine/Source/Runtime/RHI/Public/RHICommandList.h
+++ b/Engine/Source/Runtime/RHI/Public/RHICommandList.h
@@ -568,15 +568,20 @@ public:
 
 	FORCEINLINE const FRHIGPUMask& GetGPUMask() const { return GPUMask; }
 
+#if RHI_WANT_BREADCRUMB_EVENTS
+	void ExportBreadcrumbState(FRHIBreadcrumbState& State);
+	void ImportBreadcrumbState(const FRHIBreadcrumbState& State);
+#endif
+
 protected:
 #if RHI_WANT_BREADCRUMB_EVENTS
-	FRHIBreadcrumb* PushBreadcrumb_Internal(const TCHAR* InName)
+	FRHIBreadcrumb* PushBreadcrumb_Internal(const TCHAR* InText, int32 InLen = 0)
 	{
 		FRHIBreadcrumb* NewBreadcrumb = (FRHIBreadcrumb*)MemManager.Alloc(sizeof(FRHIBreadcrumb), alignof(FRHIBreadcrumb));
 
-		int32 Len = FCString::Strlen(InName) + 1;
+		int32 Len = (InLen > 0 ? InLen : FCString::Strlen(InText)) + 1;
 		TCHAR* NewName = (TCHAR*)MemManager.Alloc(Len * sizeof(TCHAR), alignof(TCHAR));
-		FCString::Strcpy(NewName, Len, InName);
+		FCString::Strcpy(NewName, Len, InText);
 
 		NewBreadcrumb->Parent = BreadcrumbStackTop;
 		NewBreadcrumb->Name = NewName;
@@ -589,6 +594,16 @@ protected:
 		return NewBreadcrumb;
 	}
 
+	FRHIBreadcrumb* PushBreadcrumbPrintf_Internal(const TCHAR* InFormat, ...)
+	{
+		TCHAR BreadcrumbString[1024];
+		int32 WrittenLength = 0;
+
+		GET_VARARGS_RESULT(BreadcrumbString, UE_ARRAY_COUNT(BreadcrumbString), UE_ARRAY_COUNT(BreadcrumbString) - 1, InFormat, InFormat, WrittenLength);
+
+		return PushBreadcrumb_Internal(BreadcrumbString, WrittenLength);
+	}
+
 	FRHIBreadcrumb* PopBreadcrumb_Internal()
 	{
 		check(BreadcrumbStackTop != nullptr); // popping more than pushing
@@ -596,14 +611,18 @@ protected:
 		return BreadcrumbStackTop;
 	}
 
-	FRHIBreadcrumb* GetBreadcrumbStackTop() const { return BreadcrumbStackTop; }
+	FRHIBreadcrumb* GetBreadcrumbStackTop() const
+	{
+		return BreadcrumbStackTop;
+	}
+
 	FRHIBreadcrumb* PopFirstUnsubmittedBreadcrumb()
 	{
 		FRHIBreadcrumb* Breadcrumb = FirstUnsubmittedBreadcrumb;
 		FirstUnsubmittedBreadcrumb = nullptr;
 		return Breadcrumb;
 	}
-#endif
+#endif // RHI_WANT_BREADCRUMB_EVENTS
 
 private:
 	FRHICommandBase* Root;
@@ -2742,10 +2761,24 @@ public:
 		ALLOC_COMMAND(FRHICommandPopEvent)();
 	}
 
-	FORCEINLINE_DEBUGGABLE void PushBreadcrumb(const TCHAR* Name)
+	FORCEINLINE_DEBUGGABLE void PushBreadcrumb(const TCHAR* InText)
+	{
+#if RHI_WANT_BREADCRUMB_EVENTS
+		FRHIBreadcrumb* Breadcrumb = PushBreadcrumb_Internal(InText);
+		if (Bypass())
+		{
+			GetComputeContext().RHISetBreadcrumbStackTop(Breadcrumb);
+			return;
+		}
+		ALLOC_COMMAND(FRHICommandSetBreadcrumbStackTop)(Breadcrumb);
+#endif
+	}
+
+	template<typename... Types>
+	FORCEINLINE_DEBUGGABLE void PushBreadcrumbPrintf(const TCHAR* Format, Types... Arguments)
 	{
 #if RHI_WANT_BREADCRUMB_EVENTS
-		FRHIBreadcrumb* Breadcrumb = PushBreadcrumb_Internal(Name);
+		FRHIBreadcrumb* Breadcrumb = PushBreadcrumbPrintf_Internal(Format, Arguments...);
 		if (Bypass())
 		{
 			GetComputeContext().RHISetBreadcrumbStackTop(Breadcrumb);
@@ -2930,8 +2963,6 @@ public:
 #endif
 
 #if RHI_WANT_BREADCRUMB_EVENTS
-	void ExportBreadcrumbState(FRHIBreadcrumbState& State);
-	void ImportBreadcrumbState(const FRHIBreadcrumbState& State);
 	void InheritBreadcrumbs(const FRHIComputeCommandList& Parent);
 #endif
 
diff --git a/Engine/Source/Runtime/RHI/Public/RHIContext.h b/Engine/Source/Runtime/RHI/Public/RHIContext.h
--- a/Engine/Source/Runtime/RHI/Public/RHIContext.h
+++ b/Engine/Source/Runtime/RHI/Public/RHIContext.h
@@ -335,18 +335,25 @@ public:
 	}
 
 #if RHI_WANT_BREADCRUMB_EVENTS
-	inline void RHISetBreadcrumbStackTop(FRHIBreadcrumb* Breadcrumb)
+	inline void RHISetBreadcrumbStackTop(const FRHIBreadcrumb* InBreadcrumbStackTop)
 	{
-		BreadcrumbStackTop = Breadcrumb;
+		if (ensure(BreadcrumbStackIndex >= 0))
+		{
+			BreadcrumbStackTop[BreadcrumbStackIndex] = InBreadcrumbStackTop;
+		}
 	}
-
-	const FRHIBreadcrumb* const * GetBreadcrumbStackTopRef() const
+	
+	const FRHIBreadcrumb* GetBreadcrumbStackTop() const
 	{
-		return &BreadcrumbStackTop;
+		return BreadcrumbStackIndex >= 0 ? BreadcrumbStackTop[BreadcrumbStackIndex] : nullptr;
 	}
 
+	enum { MaxBreadcrumbStacks = 4 };
+
 	// Top of the breadcrumb stack on the RHI thread.
-	FRHIBreadcrumb* BreadcrumbStackTop = nullptr;
+	const FRHIBreadcrumb* BreadcrumbStackTop[MaxBreadcrumbStacks]{};
+	// Index into the breadcrumbs, incremented for each command list submit and decremented when complete.
+	int32 BreadcrumbStackIndex{ -1 };
 #endif
 
 #if ENABLE_RHI_VALIDATION
diff --git a/Engine/Source/Runtime/RenderCore/Public/ProfilingDebugging/RealtimeGPUProfiler.h b/Engine/Source/Runtime/RenderCore/Public/ProfilingDebugging/RealtimeGPUProfiler.h
--- a/Engine/Source/Runtime/RenderCore/Public/ProfilingDebugging/RealtimeGPUProfiler.h
+++ b/Engine/Source/Runtime/RenderCore/Public/ProfilingDebugging/RealtimeGPUProfiler.h
@@ -30,12 +30,22 @@ struct RENDERCORE_API FBreadcrumbEvent
 {
 	FRHIComputeCommandList* RHICmdList{};
 
-	FORCEINLINE FBreadcrumbEvent(FRHIComputeCommandList& InRHICmdList, const TCHAR* String)
+	FORCEINLINE FBreadcrumbEvent(FRHIComputeCommandList& InRHICmdList, const TCHAR* InText)
 		: RHICmdList(&InRHICmdList)
 	{
 		if (RHICmdList)
 		{
-			RHICmdList->PushBreadcrumb(String);
+			RHICmdList->PushBreadcrumb(InText);
+		}
+	}
+
+	template<typename... Types>
+	FORCEINLINE FBreadcrumbEvent(FRHIComputeCommandList& InRHICmdList, const TCHAR* Format, Types... Arguments)
+		: RHICmdList(&InRHICmdList)
+	{
+		if (RHICmdList)
+		{
+			RHICmdList->PushBreadcrumbPrintf(Format, Arguments...);
 		}
 	}
 
@@ -50,8 +60,10 @@ struct RENDERCORE_API FBreadcrumbEvent
 };
 
 	#define BREADCRUMB_EVENT(RHICmdList, Name) FBreadcrumbEvent PREPROCESSOR_JOIN(BreadcrumbEvent_##Name,__LINE__)(RHICmdList, TEXT(#Name));
+	#define BREADCRUMB_EVENTF(RHICmdList, Name, Format, ...) FBreadcrumbEvent PREPROCESSOR_JOIN(BreadcrumbEvent_##Name,__LINE__)(RHICmdList, Format, ##__VA_ARGS__);
 #else
 	#define BREADCRUMB_EVENT(RHICmdList, Name) do { } while(0)
+	#define BREADCRUMB_EVENTF(RHICmdList, Name, Format, ...) do { } while(0)
 #endif
 
 #if WANTS_DRAW_MESH_EVENTS
diff --git a/Engine/Source/Runtime/Renderer/Private/MeshPassProcessor.cpp b/Engine/Source/Runtime/Renderer/Private/MeshPassProcessor.cpp
--- a/Engine/Source/Runtime/Renderer/Private/MeshPassProcessor.cpp
+++ b/Engine/Source/Runtime/Renderer/Private/MeshPassProcessor.cpp
@@ -1117,6 +1117,11 @@ void FMeshDrawCommand::SubmitDraw(
 	FRHIBuffer* IndirectArgsOverrideBuffer,
 	uint32 IndirectArgsOverrideByteOffset)
 {
+#if MESH_DRAW_COMMAND_DEBUG_DATA
+	BREADCRUMB_EVENTF(RHICmdList, MeshDrawCommand, TEXT("%s %s"),
+		*MeshDrawCommand.DebugData.MaterialName,
+		MeshDrawCommand.DebugData.ResourceName.IsValid() ? *MeshDrawCommand.DebugData.ResourceName.ToString() : TEXT(""));
+#endif
 #if WANTS_DRAW_MESH_EVENTS
 	FMeshDrawEvent MeshEvent(MeshDrawCommand, InstanceFactor, RHICmdList);
 #endif
