Quick fix for rhibypass. Since there's no execute pass, we can't increment BreadcrumbStackIndex before executing commands.

diff --git a/Engine/Source/Runtime/RHI/Public/RHIContext.h b/Engine/Source/Runtime/RHI/Public/RHIContext.h
--- a/Engine/Source/Runtime/RHI/Public/RHIContext.h
+++ b/Engine/Source/Runtime/RHI/Public/RHIContext.h
@@ -353,7 +353,7 @@ public:
 	// Top of the breadcrumb stack on the RHI thread.
 	const FRHIBreadcrumb* BreadcrumbStackTop[MaxBreadcrumbStacks]{};
 	// Index into the breadcrumbs, incremented for each command list submit and decremented when complete.
-	int32 BreadcrumbStackIndex{ -1 };
+	int32 BreadcrumbStackIndex{ 0 };
 #endif
 
 #if ENABLE_RHI_VALIDATION
