Before shutting down and deleting the Dynamic RHI, flush any potential queued commands.
PipelineStateCache specifically can now enqueue commands but this should also catch any other uses before here.
Adding explicit Command List argument to FRHIResource::FlushPendingDeletes to a) emphasize that it enqueues commands and b) allow the caller the option to flush the command list afterwards.

diff --git a/Engine/Source/Runtime/D3D12RHI/Private/D3D12Adapter.cpp b/Engine/Source/Runtime/D3D12RHI/Private/D3D12Adapter.cpp
--- a/Engine/Source/Runtime/D3D12RHI/Private/D3D12Adapter.cpp
+++ b/Engine/Source/Runtime/D3D12RHI/Private/D3D12Adapter.cpp
@@ -1106,7 +1106,11 @@ void FD3D12Adapter::Cleanup()
 	// Ask all initialized FRenderResources to release their RHI resources.
 	FRenderResource::ReleaseRHIForAllResources();
 
-	FRHIResource::FlushPendingDeletes();
+	{
+		FRHICommandListImmediate& RHICmdList = FRHICommandListExecutor::GetImmediateCommandList();
+		FRHIResource::FlushPendingDeletes(RHICmdList);
+		RHICmdList.ImmediateFlush(EImmediateFlushType::FlushRHIThread);
+	}
 
 	// Cleanup resources
 	DeferredDeletionQueue.ReleaseResources(true, true);
diff --git a/Engine/Source/Runtime/D3D12RHI/Private/D3D12Device.cpp b/Engine/Source/Runtime/D3D12RHI/Private/D3D12Device.cpp
--- a/Engine/Source/Runtime/D3D12RHI/Private/D3D12Device.cpp
+++ b/Engine/Source/Runtime/D3D12RHI/Private/D3D12Device.cpp
@@ -397,7 +397,8 @@ void FD3D12Device::Cleanup()
 	ReleasePooledUniformBuffers();
 
 	// Flush all pending deletes before destroying the device or any command contexts.
-	FRHIResource::FlushPendingDeletes();
+	FRHIResource::FlushPendingDeletes(RHICmdList);
+	RHICmdList.ImmediateFlush(EImmediateFlushType::FlushRHIThread);
 
 	// Delete array index 0 (the default context) last
 	for (int32 i = CommandContextArray.Num() - 1; i >= 0; i--)
diff --git a/Engine/Source/Runtime/RHI/Private/DynamicRHI.cpp b/Engine/Source/Runtime/RHI/Private/DynamicRHI.cpp
--- a/Engine/Source/Runtime/RHI/Private/DynamicRHI.cpp
+++ b/Engine/Source/Runtime/RHI/Private/DynamicRHI.cpp
@@ -327,6 +327,9 @@ void RHIExit()
 		// Clean up all cached pipelines
 		PipelineStateCache::Shutdown();
 
+		// Flush any potential commands queued before we shut things down.
+		FRHICommandListExecutor::GetImmediateCommandList().ImmediateFlush(EImmediateFlushType::FlushRHIThread);
+
 		// Destruct the dynamic RHI.
 		GDynamicRHI->Shutdown();
 		delete GDynamicRHI;
diff --git a/Engine/Source/Runtime/RHI/Private/RHI.cpp b/Engine/Source/Runtime/RHI/Private/RHI.cpp
--- a/Engine/Source/Runtime/RHI/Private/RHI.cpp
+++ b/Engine/Source/Runtime/RHI/Private/RHI.cpp
@@ -991,13 +991,12 @@ bool FRHIResource::Bypass()
 
 DECLARE_CYCLE_STAT(TEXT("Delete Resources"), STAT_DeleteResources, STATGROUP_RHICMDLIST);
 
-void FRHIResource::FlushPendingDeletes()
+void FRHIResource::FlushPendingDeletes(FRHICommandListImmediate& RHICmdList)
 {
 	SCOPE_CYCLE_COUNTER(STAT_DeleteResources);
 
 	check(IsInRenderingThread());
 
-	FRHICommandListImmediate& RHICmdList = FRHICommandListExecutor::GetImmediateCommandList();
 #if ENABLE_RHI_VALIDATION
 	if (GDynamicRHI)
 	{
@@ -1011,7 +1010,7 @@ void FRHIResource::FlushPendingDeletes()
 	TClosableMpscQueue<FRHIResource*>* PendingDeletesPtr = PendingDeletes.exchange(new TClosableMpscQueue<FRHIResource*>());
 	PendingDeletesPtr->Close([&DeletedResources](FRHIResource* Resource)
 	{	
-		DeletedResources.Push(Resource);		
+		DeletedResources.Push(Resource);
 	});
 	PendingDeletesHPC.Delete(PendingDeletesPtr);
 
diff --git a/Engine/Source/Runtime/RHI/Private/RHICommandList.cpp b/Engine/Source/Runtime/RHI/Private/RHICommandList.cpp
--- a/Engine/Source/Runtime/RHI/Private/RHICommandList.cpp
+++ b/Engine/Source/Runtime/RHI/Private/RHICommandList.cpp
@@ -947,7 +947,7 @@ void FRHICommandListExecutor::LatchBypass()
 
 		if (NewBypass && !bLatchedBypass)
 		{
-			FRHIResource::FlushPendingDeletes();
+			FRHIResource::FlushPendingDeletes(FRHICommandListExecutor::GetImmediateCommandList());
 		}
 		bLatchedBypass = NewBypass;
 	}
diff --git a/Engine/Source/Runtime/RHI/Public/RHICommandList.inl b/Engine/Source/Runtime/RHI/Public/RHICommandList.inl
--- a/Engine/Source/Runtime/RHI/Public/RHICommandList.inl
+++ b/Engine/Source/Runtime/RHI/Public/RHICommandList.inl
@@ -95,7 +95,7 @@ FORCEINLINE_DEBUGGABLE void FRHICommandListImmediate::ImmediateFlush(EImmediateF
 		{
 			CSV_SCOPED_TIMING_STAT(RHITFlushes, FlushRHIThreadFlushResourcesTotal);
 			PipelineStateCache::FlushResources();
-			FRHIResource::FlushPendingDeletes();
+			FRHIResource::FlushPendingDeletes(FRHICommandListExecutor::GetImmediateCommandList());
 		}
 	//this fall-through is intentional
 	case EImmediateFlushType::FlushRHIThread:
diff --git a/Engine/Source/Runtime/RHI/Public/RHIResources.h b/Engine/Source/Runtime/RHI/Public/RHIResources.h
--- a/Engine/Source/Runtime/RHI/Public/RHIResources.h
+++ b/Engine/Source/Runtime/RHI/Public/RHIResources.h
@@ -29,6 +29,7 @@
 #define RHI_WANT_RESOURCE_INFO 0
 #endif
 
+class FRHICommandListImmediate;
 struct FClearValueBinding;
 struct FRHIResourceInfo;
 struct FGenerateMipsStruct;
@@ -99,7 +100,7 @@ public:
 		return uint32(CurrentValue);
 	}
 
-	static void FlushPendingDeletes();
+	static void FlushPendingDeletes(FRHICommandListImmediate& RHICmdList);
 
 	static bool Bypass();
 
diff --git a/Engine/Source/Runtime/RenderCore/Private/RenderingThread.cpp b/Engine/Source/Runtime/RenderCore/Private/RenderingThread.cpp
--- a/Engine/Source/Runtime/RenderCore/Private/RenderingThread.cpp
+++ b/Engine/Source/Runtime/RenderCore/Private/RenderingThread.cpp
@@ -1305,7 +1305,7 @@ void FlushPendingDeleteRHIResources_RenderThread()
 {
 	if (!IsRunningRHIInSeparateThread())
 	{
-		FRHIResource::FlushPendingDeletes();
+		FRHIResource::FlushPendingDeletes(FRHICommandListExecutor::GetImmediateCommandList());
 	}
 }
 
diff --git a/Engine/Source/Runtime/VulkanRHI/Private/VulkanDevice.cpp b/Engine/Source/Runtime/VulkanRHI/Private/VulkanDevice.cpp
--- a/Engine/Source/Runtime/VulkanRHI/Private/VulkanDevice.cpp
+++ b/Engine/Source/Runtime/VulkanRHI/Private/VulkanDevice.cpp
@@ -1238,7 +1238,11 @@ void FVulkanDevice::Destroy()
 #endif
 
 	// Flush all pending deletes before destroying the device and any Vulkan context objects.
-	FRHIResource::FlushPendingDeletes();
+	{
+		FRHICommandListImmediate& RHICmdList = FRHICommandListExecutor::GetImmediateCommandList();
+		FRHIResource::FlushPendingDeletes(RHICmdList);
+		RHICmdList.ImmediateFlush(EImmediateFlushType::FlushRHIThread);
+	}
 
 	VulkanRHI::vkDestroyImageView(GetInstanceHandle(), DefaultTextureView.View, VULKAN_CPU_ALLOCATOR);
 	DefaultTextureView = {};
diff --git a/Engine/Source/Runtime/VulkanRHI/Private/VulkanRHI.cpp b/Engine/Source/Runtime/VulkanRHI/Private/VulkanRHI.cpp
--- a/Engine/Source/Runtime/VulkanRHI/Private/VulkanRHI.cpp
+++ b/Engine/Source/Runtime/VulkanRHI/Private/VulkanRHI.cpp
@@ -335,11 +335,15 @@ void FVulkanDynamicRHI::Shutdown()
 		Device->CleanUpRayTracing();
 #endif // VULKAN_RHI_RAYTRACING
 
+		FRHICommandListImmediate& RHICmdList = FRHICommandListExecutor::GetImmediateCommandList();
+
 		// Flush all pending deletes before destroying the device.
-		FRHIResource::FlushPendingDeletes();
+		FRHIResource::FlushPendingDeletes(RHICmdList);
+		RHICmdList.ImmediateFlush(EImmediateFlushType::FlushRHIThread);
 
 		// And again since some might get on a pending queue
-		FRHIResource::FlushPendingDeletes();
+		FRHIResource::FlushPendingDeletes(RHICmdList);
+		RHICmdList.ImmediateFlush(EImmediateFlushType::FlushRHIThread);
 	}
 
 	Device->Destroy();
diff --git a/Engine/Source/Runtime/Windows/D3D11RHI/Private/D3D11Device.cpp b/Engine/Source/Runtime/Windows/D3D11RHI/Private/D3D11Device.cpp
--- a/Engine/Source/Runtime/Windows/D3D11RHI/Private/D3D11Device.cpp
+++ b/Engine/Source/Runtime/Windows/D3D11RHI/Private/D3D11Device.cpp
@@ -528,7 +528,11 @@ void FD3D11DynamicRHI::CleanupD3DDevice()
 		StateCache.SetContext(nullptr);
 
 		// Flush all pending deletes before destroying the device.
-		FRHIResource::FlushPendingDeletes();
+		{
+			FRHICommandListImmediate& RHICmdList = FRHICommandListExecutor::GetImmediateCommandList();
+			FRHIResource::FlushPendingDeletes(RHICmdList);
+			RHICmdList.ImmediateFlush(EImmediateFlushType::FlushRHIThread);
+		}
 
 		ReleasePooledUniformBuffers();
 		ReleasePooledTextures();
