Adding PixelFormat Capability flags.

diff --git a/Engine/Source/Runtime/D3D12RHI/Private/D3D12Device.cpp b/Engine/Source/Runtime/D3D12RHI/Private/D3D12Device.cpp
--- a/Engine/Source/Runtime/D3D12RHI/Private/D3D12Device.cpp
+++ b/Engine/Source/Runtime/D3D12RHI/Private/D3D12Device.cpp
@@ -171,6 +171,71 @@ void FD3D12Device::SetupAfterDeviceCreation()
 {
 	ID3D12Device* Direct3DDevice = GetParentAdapter()->GetD3DDevice();
 
+	for (uint32 FormatIndex = PF_Unknown; FormatIndex < PF_MAX; FormatIndex++)
+	{
+		FPixelFormatInfo& PixelFormatInfo = GPixelFormats[FormatIndex];
+		const DXGI_FORMAT PlatformFormat = static_cast<DXGI_FORMAT>(PixelFormatInfo.PlatformFormat);
+
+		EPixelFormatCapabilities Capabilities = EPixelFormatCapabilities::None;
+
+		if (PlatformFormat != DXGI_FORMAT_UNKNOWN)
+		{
+			D3D12_FEATURE_DATA_FORMAT_SUPPORT FormatSupport{};
+			FormatSupport.Format = PlatformFormat;
+
+			auto ConvertCap1 = [&Capabilities, &FormatSupport](EPixelFormatCapabilities UnrealCap, D3D12_FORMAT_SUPPORT1 InFlags)
+			{
+				if (EnumHasAnyFlags(FormatSupport.Support1, InFlags))
+				{
+					EnumAddFlags(Capabilities, UnrealCap);
+				}
+			};
+			auto ConvertCap2 = [&Capabilities, &FormatSupport](EPixelFormatCapabilities UnrealCap, D3D12_FORMAT_SUPPORT2 InFlags)
+			{
+				if (EnumHasAnyFlags(FormatSupport.Support2, InFlags))
+				{
+					EnumAddFlags(Capabilities, UnrealCap);
+				}
+			};
+
+			HRESULT hr = Direct3DDevice->CheckFeatureSupport(D3D12_FEATURE_FORMAT_SUPPORT, &FormatSupport, sizeof(FormatSupport));
+			if (SUCCEEDED(hr))
+			{
+				ConvertCap1(EPixelFormatCapabilities::Texture1D,			D3D12_FORMAT_SUPPORT1_TEXTURE1D);
+				ConvertCap1(EPixelFormatCapabilities::Texture2D,			D3D12_FORMAT_SUPPORT1_TEXTURE2D);
+				ConvertCap1(EPixelFormatCapabilities::Texture3D,			D3D12_FORMAT_SUPPORT1_TEXTURE3D);
+				ConvertCap1(EPixelFormatCapabilities::TextureCube,			D3D12_FORMAT_SUPPORT1_TEXTURECUBE);
+				ConvertCap1(EPixelFormatCapabilities::Buffer,               D3D12_FORMAT_SUPPORT1_BUFFER);
+				ConvertCap1(EPixelFormatCapabilities::VertexBuffer,         D3D12_FORMAT_SUPPORT1_IA_VERTEX_BUFFER);
+				ConvertCap1(EPixelFormatCapabilities::IndexBuffer,          D3D12_FORMAT_SUPPORT1_IA_INDEX_BUFFER);
+
+				if (EnumHasAnyFlags(Capabilities, EPixelFormatCapabilities::AnyTexture))
+				{
+					ConvertCap1(EPixelFormatCapabilities::RenderTarget,     D3D12_FORMAT_SUPPORT1_RENDER_TARGET);
+					ConvertCap1(EPixelFormatCapabilities::DepthStencil,     D3D12_FORMAT_SUPPORT1_DEPTH_STENCIL);
+					ConvertCap1(EPixelFormatCapabilities::TextureMipmaps,   D3D12_FORMAT_SUPPORT1_MIP);
+					ConvertCap1(EPixelFormatCapabilities::TextureLoad,      D3D12_FORMAT_SUPPORT1_SHADER_LOAD);
+					ConvertCap1(EPixelFormatCapabilities::TextureSample,    D3D12_FORMAT_SUPPORT1_SHADER_SAMPLE);
+					ConvertCap1(EPixelFormatCapabilities::TextureGather,    D3D12_FORMAT_SUPPORT1_SHADER_GATHER);
+					ConvertCap2(EPixelFormatCapabilities::TextureAtomics,   D3D12_FORMAT_SUPPORT2_UAV_ATOMIC_EXCHANGE);
+					ConvertCap1(EPixelFormatCapabilities::TextureBlendable, D3D12_FORMAT_SUPPORT1_BLENDABLE);
+				}
+
+				if (EnumHasAnyFlags(Capabilities, EPixelFormatCapabilities::Buffer))
+				{
+					ConvertCap1(EPixelFormatCapabilities::BufferLoad,       D3D12_FORMAT_SUPPORT1_SHADER_LOAD);
+					ConvertCap2(EPixelFormatCapabilities::BufferStore,      D3D12_FORMAT_SUPPORT2_UAV_TYPED_STORE);
+					ConvertCap2(EPixelFormatCapabilities::BufferAtomics,    D3D12_FORMAT_SUPPORT2_UAV_ATOMIC_EXCHANGE);
+				}
+
+				ConvertCap2(EPixelFormatCapabilities::TypedUAVLoad,			D3D12_FORMAT_SUPPORT2_UAV_TYPED_LOAD);
+				ConvertCap2(EPixelFormatCapabilities::TypedUAVStore,		D3D12_FORMAT_SUPPORT2_UAV_TYPED_STORE);
+			}
+		}
+
+		PixelFormatInfo.Capabilities = Capabilities;
+	}
+
 	GRHISupportsArrayIndexFromAnyShader = true;
 	GRHISupportsStencilRefFromPixelShader = false; // TODO: Sort out DXC shader database SM6.0 usage. DX12 supports this feature, but need to improve DXC support.
 
diff --git a/Engine/Source/Runtime/RHI/Private/RHI.cpp b/Engine/Source/Runtime/RHI/Private/RHI.cpp
--- a/Engine/Source/Runtime/RHI/Private/RHI.cpp
+++ b/Engine/Source/Runtime/RHI/Private/RHI.cpp
@@ -2213,95 +2213,115 @@ int32 CalculateMSAASampleArrayIndex(int32 NumSamples, int32 SampleIndex)
 //	Pixel format information.
 //
 
-FPixelFormatInfo	GPixelFormats[PF_MAX] =
-{
-	// Name						BlockSizeX	BlockSizeY	BlockSizeZ	BlockBytes	NumComponents	PlatformFormat	Supported		UnrealFormat
-
-	{ TEXT("unknown"),			0,			0,			0,			0,			0,				0,				0,				PF_Unknown			},
-	{ TEXT("A32B32G32R32F"),	1,			1,			1,			16,			4,				0,				1,				PF_A32B32G32R32F	},
-	{ TEXT("B8G8R8A8"),			1,			1,			1,			4,			4,				0,				1,				PF_B8G8R8A8			},
-	{ TEXT("G8"),				1,			1,			1,			1,			1,				0,				1,				PF_G8				},
-	{ TEXT("G16"),				1,			1,			1,			2,			1,				0,				1,				PF_G16				},
-	{ TEXT("DXT1"),				4,			4,			1,			8,			3,				0,				1,				PF_DXT1				},
-	{ TEXT("DXT3"),				4,			4,			1,			16,			4,				0,				1,				PF_DXT3				},
-	{ TEXT("DXT5"),				4,			4,			1,			16,			4,				0,				1,				PF_DXT5				},
-	{ TEXT("UYVY"),				2,			1,			1,			4,			4,				0,				0,				PF_UYVY				},
-	{ TEXT("FloatRGB"),			1,			1,			1,			4,			3,				0,				1,				PF_FloatRGB			},
-	{ TEXT("FloatRGBA"),		1,			1,			1,			8,			4,				0,				1,				PF_FloatRGBA		},
-	{ TEXT("DepthStencil"),		1,			1,			1,			4,			1,				0,				0,				PF_DepthStencil		},
-	{ TEXT("ShadowDepth"),		1,			1,			1,			4,			1,				0,				0,				PF_ShadowDepth		},
-	{ TEXT("R32_FLOAT"),		1,			1,			1,			4,			1,				0,				1,				PF_R32_FLOAT		},
-	{ TEXT("G16R16"),			1,			1,			1,			4,			2,				0,				1,				PF_G16R16			},
-	{ TEXT("G16R16F"),			1,			1,			1,			4,			2,				0,				1,				PF_G16R16F			},
-	{ TEXT("G16R16F_FILTER"),	1,			1,			1,			4,			2,				0,				1,				PF_G16R16F_FILTER	},
-	{ TEXT("G32R32F"),			1,			1,			1,			8,			2,				0,				1,				PF_G32R32F			},
-	{ TEXT("A2B10G10R10"),      1,          1,          1,          4,          4,              0,              1,				PF_A2B10G10R10		},
-	{ TEXT("A16B16G16R16"),		1,			1,			1,			8,			4,				0,				1,				PF_A16B16G16R16		},
-	{ TEXT("D24"),				1,			1,			1,			4,			1,				0,				1,				PF_D24				},
-	{ TEXT("PF_R16F"),			1,			1,			1,			2,			1,				0,				1,				PF_R16F				},
-	{ TEXT("PF_R16F_FILTER"),	1,			1,			1,			2,			1,				0,				1,				PF_R16F_FILTER		},
-	{ TEXT("BC5"),				4,			4,			1,			16,			2,				0,				1,				PF_BC5				},
-	{ TEXT("V8U8"),				1,			1,			1,			2,			2,				0,				1,				PF_V8U8				},
-	{ TEXT("A1"),				1,			1,			1,			1,			1,				0,				0,				PF_A1				},
-	{ TEXT("FloatR11G11B10"),	1,			1,			1,			4,			3,				0,				0,				PF_FloatR11G11B10	},
-	{ TEXT("A8"),				1,			1,			1,			1,			1,				0,				1,				PF_A8				},	
-	{ TEXT("R32_UINT"),			1,			1,			1,			4,			1,				0,				1,				PF_R32_UINT			},
-	{ TEXT("R32_SINT"),			1,			1,			1,			4,			1,				0,				1,				PF_R32_SINT			},
+FPixelFormatInfo::FPixelFormatInfo(
+	EPixelFormat InUnrealFormat,
+	const TCHAR* InName,
+	int32 InBlockSizeX,
+	int32 InBlockSizeY,
+	int32 InBlockSizeZ,
+	int32 InBlockBytes,
+	int32 InNumComponents,
+	bool  InSupported)
+	: Name(InName)
+	, UnrealFormat(InUnrealFormat)
+	, BlockSizeX(InBlockSizeX)
+	, BlockSizeY(InBlockSizeY)
+	, BlockSizeZ(InBlockSizeZ)
+	, BlockBytes(InBlockBytes)
+	, NumComponents(InNumComponents)
+	, Supported(InSupported)
+{
+}
+
+FPixelFormatInfo    GPixelFormats[PF_MAX] =
+{
+	//               UnrealFormat                 Name                 BlockSizeX  BlockSizeY  BlockSizeZ  BlockBytes  NumComponents    Supported
+
+	FPixelFormatInfo(PF_Unknown,            TEXT("unknown"),               0,          0,          0,          0,          0,              0),
+	FPixelFormatInfo(PF_A32B32G32R32F,      TEXT("A32B32G32R32F"),         1,          1,          1,          16,         4,              1),
+	FPixelFormatInfo(PF_B8G8R8A8,           TEXT("B8G8R8A8"),              1,          1,          1,          4,          4,              1),
+	FPixelFormatInfo(PF_G8,                 TEXT("G8"),                    1,          1,          1,          1,          1,              1),
+	FPixelFormatInfo(PF_G16,                TEXT("G16"),                   1,          1,          1,          2,          1,              1),
+	FPixelFormatInfo(PF_DXT1,               TEXT("DXT1"),                  4,          4,          1,          8,          3,              1),
+	FPixelFormatInfo(PF_DXT3,               TEXT("DXT3"),                  4,          4,          1,          16,         4,              1),
+	FPixelFormatInfo(PF_DXT5,               TEXT("DXT5"),                  4,          4,          1,          16,         4,              1),
+	FPixelFormatInfo(PF_UYVY,               TEXT("UYVY"),                  2,          1,          1,          4,          4,              0),
+	FPixelFormatInfo(PF_FloatRGB,           TEXT("FloatRGB"),              1,          1,          1,          4,          3,              1),
+	FPixelFormatInfo(PF_FloatRGBA,          TEXT("FloatRGBA"),             1,          1,          1,          8,          4,              1),
+	FPixelFormatInfo(PF_DepthStencil,       TEXT("DepthStencil"),          1,          1,          1,          4,          1,              0),
+	FPixelFormatInfo(PF_ShadowDepth,        TEXT("ShadowDepth"),           1,          1,          1,          4,          1,              0),
+	FPixelFormatInfo(PF_R32_FLOAT,          TEXT("R32_FLOAT"),             1,          1,          1,          4,          1,              1),
+	FPixelFormatInfo(PF_G16R16,             TEXT("G16R16"),                1,          1,          1,          4,          2,              1),
+	FPixelFormatInfo(PF_G16R16F,            TEXT("G16R16F"),               1,          1,          1,          4,          2,              1),
+	FPixelFormatInfo(PF_G16R16F_FILTER,     TEXT("G16R16F_FILTER"),        1,          1,          1,          4,          2,              1),
+	FPixelFormatInfo(PF_G32R32F,            TEXT("G32R32F"),               1,          1,          1,          8,          2,              1),
+	FPixelFormatInfo(PF_A2B10G10R10,        TEXT("A2B10G10R10"),           1,          1,          1,          4,          4,              1),
+	FPixelFormatInfo(PF_A16B16G16R16,       TEXT("A16B16G16R16"),          1,          1,          1,          8,          4,              1),
+	FPixelFormatInfo(PF_D24,                TEXT("D24"),                   1,          1,          1,          4,          1,              1),
+	FPixelFormatInfo(PF_R16F,               TEXT("PF_R16F"),               1,          1,          1,          2,          1,              1),
+	FPixelFormatInfo(PF_R16F_FILTER,        TEXT("PF_R16F_FILTER"),        1,          1,          1,          2,          1,              1),
+	FPixelFormatInfo(PF_BC5,                TEXT("BC5"),                   4,          4,          1,          16,         2,              1),
+	FPixelFormatInfo(PF_V8U8,               TEXT("V8U8"),                  1,          1,          1,          2,          2,              1),
+	FPixelFormatInfo(PF_A1,                 TEXT("A1"),                    1,          1,          1,          1,          1,              0),
+	FPixelFormatInfo(PF_FloatR11G11B10,     TEXT("FloatR11G11B10"),        1,          1,          1,          4,          3,              0),
+	FPixelFormatInfo(PF_A8,                 TEXT("A8"),                    1,          1,          1,          1,          1,              1),
+	FPixelFormatInfo(PF_R32_UINT,           TEXT("R32_UINT"),              1,          1,          1,          4,          1,              1),
+	FPixelFormatInfo(PF_R32_SINT,           TEXT("R32_SINT"),              1,          1,          1,          4,          1,              1),
 
 	// IOS Support
-	{ TEXT("PVRTC2"),			8,			4,			1,			8,			4,				0,				0,				PF_PVRTC2			},
-	{ TEXT("PVRTC4"),			4,			4,			1,			8,			4,				0,				0,				PF_PVRTC4			},
-
-	{ TEXT("R16_UINT"),			1,			1,			1,			2,			1,				0,				1,				PF_R16_UINT			},
-	{ TEXT("R16_SINT"),			1,			1,			1,			2,			1,				0,				1,				PF_R16_SINT			},
-	{ TEXT("R16G16B16A16_UINT"),1,			1,			1,			8,			4,				0,				1,				PF_R16G16B16A16_UINT},
-	{ TEXT("R16G16B16A16_SINT"),1,			1,			1,			8,			4,				0,				1,				PF_R16G16B16A16_SINT},
-	{ TEXT("R5G6B5_UNORM"),     1,          1,          1,          2,          3,              0,              1,              PF_R5G6B5_UNORM		},
-	{ TEXT("R8G8B8A8"),			1,			1,			1,			4,			4,				0,				1,				PF_R8G8B8A8			},
-	{ TEXT("A8R8G8B8"),			1,			1,			1,			4,			4,				0,				1,				PF_A8R8G8B8			},
-	{ TEXT("BC4"),				4,			4,			1,			8,			1,				0,				1,				PF_BC4				},
-	{ TEXT("R8G8"),				1,			1,			1,			2,			2,				0,				1,				PF_R8G8				},
-
-	{ TEXT("ATC_RGB"),			4,			4,			1,			8,			3,				0,				0,				PF_ATC_RGB			},
-	{ TEXT("ATC_RGBA_E"),		4,			4,			1,			16,			4,				0,				0,				PF_ATC_RGBA_E		},
-	{ TEXT("ATC_RGBA_I"),		4,			4,			1,			16,			4,				0,				0,				PF_ATC_RGBA_I		},
-	{ TEXT("X24_G8"),			1,			1,			1,			1,			1,				0,				0,				PF_X24_G8			},
-	{ TEXT("ETC1"),				4,			4,			1,			8,			3,				0,				0,				PF_ETC1				},
-	{ TEXT("ETC2_RGB"),			4,			4,			1,			8,			3,				0,				0,				PF_ETC2_RGB			},
-	{ TEXT("ETC2_RGBA"),		4,			4,			1,			16,			4,				0,				0,				PF_ETC2_RGBA		},
-	{ TEXT("PF_R32G32B32A32_UINT"),1,		1,			1,			16,			4,				0,				1,				PF_R32G32B32A32_UINT},
-	{ TEXT("PF_R16G16_UINT"),	1,			1,			1,			4,			4,				0,				1,				PF_R16G16_UINT},
+	FPixelFormatInfo(PF_PVRTC2,             TEXT("PVRTC2"),                8,          4,          1,          8,          4,              0),
+	FPixelFormatInfo(PF_PVRTC4,             TEXT("PVRTC4"),                4,          4,          1,          8,          4,              0),
+
+	FPixelFormatInfo(PF_R16_UINT,           TEXT("R16_UINT"),              1,          1,          1,          2,          1,              1),
+	FPixelFormatInfo(PF_R16_SINT,           TEXT("R16_SINT"),              1,          1,          1,          2,          1,              1),
+	FPixelFormatInfo(PF_R16G16B16A16_UINT,  TEXT("R16G16B16A16_UINT"),     1,          1,          1,          8,          4,              1),
+	FPixelFormatInfo(PF_R16G16B16A16_SINT,  TEXT("R16G16B16A16_SINT"),     1,          1,          1,          8,          4,              1),
+	FPixelFormatInfo(PF_R5G6B5_UNORM,       TEXT("R5G6B5_UNORM"),          1,          1,          1,          2,          3,              1),
+	FPixelFormatInfo(PF_R8G8B8A8,           TEXT("R8G8B8A8"),              1,          1,          1,          4,          4,              1),
+	FPixelFormatInfo(PF_A8R8G8B8,           TEXT("A8R8G8B8"),              1,          1,          1,          4,          4,              1),
+	FPixelFormatInfo(PF_BC4,                TEXT("BC4"),                   4,          4,          1,          8,          1,              1),
+	FPixelFormatInfo(PF_R8G8,               TEXT("R8G8"),                  1,          1,          1,          2,          2,              1),
+
+	FPixelFormatInfo(PF_ATC_RGB,            TEXT("ATC_RGB"),               4,          4,          1,          8,          3,              0),
+	FPixelFormatInfo(PF_ATC_RGBA_E,         TEXT("ATC_RGBA_E"),            4,          4,          1,          16,         4,              0),
+	FPixelFormatInfo(PF_ATC_RGBA_I,         TEXT("ATC_RGBA_I"),            4,          4,          1,          16,         4,              0),
+	FPixelFormatInfo(PF_X24_G8,             TEXT("X24_G8"),                1,          1,          1,          1,          1,              0),
+	FPixelFormatInfo(PF_ETC1,               TEXT("ETC1"),                  4,          4,          1,          8,          3,              0),
+	FPixelFormatInfo(PF_ETC2_RGB,           TEXT("ETC2_RGB"),              4,          4,          1,          8,          3,              0),
+	FPixelFormatInfo(PF_ETC2_RGBA,          TEXT("ETC2_RGBA"),             4,          4,          1,          16,         4,              0),
+	FPixelFormatInfo(PF_R32G32B32A32_UINT,  TEXT("PF_R32G32B32A32_UINT"),  1,          1,          1,          16,         4,              1),
+	FPixelFormatInfo(PF_R16G16_UINT,        TEXT("PF_R16G16_UINT"),        1,          1,          1,          4,          4,              1),
 
 	// ASTC support
-	{ TEXT("ASTC_4x4"),			4,			4,			1,			16,			4,				0,				0,				PF_ASTC_4x4			},
-	{ TEXT("ASTC_6x6"),			6,			6,			1,			16,			4,				0,				0,				PF_ASTC_6x6			},
-	{ TEXT("ASTC_8x8"),			8,			8,			1,			16,			4,				0,				0,				PF_ASTC_8x8			},
-	{ TEXT("ASTC_10x10"),		10,			10,			1,			16,			4,				0,				0,				PF_ASTC_10x10		},
-	{ TEXT("ASTC_12x12"),		12,			12,			1,			16,			4,				0,				0,				PF_ASTC_12x12		},
-
-	{ TEXT("BC6H"),				4,			4,			1,			16,			3,				0,				1,				PF_BC6H				},
-	{ TEXT("BC7"),				4,			4,			1,			16,			4,				0,				1,				PF_BC7				},
-	{ TEXT("R8_UINT"),			1,			1,			1,			1,			1,				0,				1,				PF_R8_UINT			},
-	{ TEXT("L8"),				1,			1,			1,			1,			1,				0,				0,				PF_L8				},
-	{ TEXT("XGXR8"),			1,			1,			1,			4,			4,				0,				1,				PF_XGXR8 			},
-	{ TEXT("R8G8B8A8_UINT"),	1,			1,			1,			4,			4,				0,				1,				PF_R8G8B8A8_UINT	},
-	{ TEXT("R8G8B8A8_SNORM"),	1,			1,			1,			4,			4,				0,				1,				PF_R8G8B8A8_SNORM	},
-
-	{ TEXT("R16G16B16A16_UINT"),1,			1,			1,			8,			4,				0,				1,				PF_R16G16B16A16_UNORM },
-	{ TEXT("R16G16B16A16_SINT"),1,			1,			1,			8,			4,				0,				1,				PF_R16G16B16A16_SNORM },
-	{ TEXT("PLATFORM_HDR_0"),	0,			0,			0,			0,			0,				0,				0,				PF_PLATFORM_HDR_0   },
-	{ TEXT("PLATFORM_HDR_1"),	0,			0,			0,			0,			0,				0,				0,				PF_PLATFORM_HDR_1   },
-	{ TEXT("PLATFORM_HDR_2"),	0,			0,			0,			0,			0,				0,				0,				PF_PLATFORM_HDR_2   },
+	FPixelFormatInfo(PF_ASTC_4x4,           TEXT("ASTC_4x4"),              4,          4,          1,          16,         4,              0),
+	FPixelFormatInfo(PF_ASTC_6x6,           TEXT("ASTC_6x6"),              6,          6,          1,          16,         4,              0),
+	FPixelFormatInfo(PF_ASTC_8x8,           TEXT("ASTC_8x8"),              8,          8,          1,          16,         4,              0),
+	FPixelFormatInfo(PF_ASTC_10x10,         TEXT("ASTC_10x10"),            10,         10,         1,          16,         4,              0),
+	FPixelFormatInfo(PF_ASTC_12x12,         TEXT("ASTC_12x12"),            12,         12,         1,          16,         4,              0),
+
+	FPixelFormatInfo(PF_BC6H,               TEXT("BC6H"),                  4,          4,          1,          16,         3,              1),
+	FPixelFormatInfo(PF_BC7,                TEXT("BC7"),                   4,          4,          1,          16,         4,              1),
+	FPixelFormatInfo(PF_R8_UINT,            TEXT("R8_UINT"),               1,          1,          1,          1,          1,              1),
+	FPixelFormatInfo(PF_L8,                 TEXT("L8"),                    1,          1,          1,          1,          1,              0),
+	FPixelFormatInfo(PF_XGXR8,              TEXT("XGXR8"),                 1,          1,          1,          4,          4,              1),
+	FPixelFormatInfo(PF_R8G8B8A8_UINT,      TEXT("R8G8B8A8_UINT"),         1,          1,          1,          4,          4,              1),
+	FPixelFormatInfo(PF_R8G8B8A8_SNORM,     TEXT("R8G8B8A8_SNORM"),        1,          1,          1,          4,          4,              1),
+
+	FPixelFormatInfo(PF_R16G16B16A16_UNORM, TEXT("R16G16B16A16_UINT"),     1,          1,          1,          8,          4,              1),
+	FPixelFormatInfo(PF_R16G16B16A16_SNORM, TEXT("R16G16B16A16_SINT"),     1,          1,          1,          8,          4,              1),
+	FPixelFormatInfo(PF_PLATFORM_HDR_0,     TEXT("PLATFORM_HDR_0"),        0,          0,          0,          0,          0,              0),
+	FPixelFormatInfo(PF_PLATFORM_HDR_1,     TEXT("PLATFORM_HDR_1"),        0,          0,          0,          0,          0,              0),
+	FPixelFormatInfo(PF_PLATFORM_HDR_2,     TEXT("PLATFORM_HDR_2"),        0,          0,          0,          0,          0,              0),
 
 	// NV12 contains 2 textures: R8 luminance plane followed by R8G8 1/4 size chrominance plane.
 	// BlockSize/BlockBytes/NumComponents values don't make much sense for this format, so set them all to one.
-	{ TEXT("NV12"),				1,			1,			1,			1,			1,				0,				0,				PF_NV12             },
+	FPixelFormatInfo(PF_NV12,               TEXT("NV12"),                  1,          1,          1,          1,          1,              0),
 
-	{ TEXT("PF_R32G32_UINT"),   1,   		1,			1,			8,			2,				0,				1,				PF_R32G32_UINT      },
+	FPixelFormatInfo(PF_R32G32_UINT,        TEXT("PF_R32G32_UINT"),        1,          1,          1,          8,          2,              1),
 
-	{ TEXT("PF_ETC2_R11_EAC"),  4,   		4,			1,			8,			1,				0,				0,				PF_ETC2_R11_EAC     },
-	{ TEXT("PF_ETC2_RG11_EAC"), 4,   		4,			1,			16,			2,				0,				0,				PF_ETC2_RG11_EAC    },
-	{ TEXT("R8"),				1,			1,			1,			1,			1,				0,				1,				PF_R8				},
+	FPixelFormatInfo(PF_ETC2_R11_EAC,       TEXT("PF_ETC2_R11_EAC"),       4,          4,          1,          8,          1,              0),
+	FPixelFormatInfo(PF_ETC2_RG11_EAC,      TEXT("PF_ETC2_RG11_EAC"),      4,          4,          1,          16,         2,              0),
+	FPixelFormatInfo(PF_R8,                 TEXT("R8"),                    1,          1,          1,          1,          1,              1),
 };
 
 static struct FValidatePixelFormats
diff --git a/Engine/Source/Runtime/RHI/Public/RHI.h b/Engine/Source/Runtime/RHI/Public/RHI.h
--- a/Engine/Source/Runtime/RHI/Public/RHI.h
+++ b/Engine/Source/Runtime/RHI/Public/RHI.h
@@ -708,20 +708,66 @@ inline FVector2D GetMSAASampleOffsets(int32 NumSamples, int32 SampleIndex)
 	return GRHIDefaultMSAASampleOffsets[CalculateMSAASampleArrayIndex(NumSamples, SampleIndex)];
 }
 
+enum class EPixelFormatCapabilities : uint32
+{
+    None             = 0,
+    Texture1D        = 1ull << 1,
+    Texture2D        = 1ull << 2,
+    Texture3D        = 1ull << 3,
+    TextureCube      = 1ull << 4,
+    RenderTarget     = 1ull << 5,
+    DepthStencil     = 1ull << 6,
+	TextureMipmaps   = 1ull << 7,
+	TextureLoad      = 1ull << 8,
+	TextureSample    = 1ull << 9,
+	TextureGather    = 1ull << 10,
+	TextureAtomics   = 1ull << 11,
+	TextureBlendable = 1ull << 12,
+
+	Buffer           = 1ull << 13,
+    VertexBuffer     = 1ull << 14,
+    IndexBuffer      = 1ull << 15,
+	BufferLoad       = 1ull << 16,
+    BufferStore      = 1ull << 17,
+    BufferAtomics    = 1ull << 18,
+
+    TypedUAVLoad     = 1ull << 19,
+    TypedUAVStore    = 1ull << 20,
+
+	AnyTexture = Texture1D | Texture2D | Texture3D | TextureCube,
+};
+ENUM_CLASS_FLAGS(EPixelFormatCapabilities);
+
 /** Information about a pixel format. */
 struct FPixelFormatInfo
 {
-	const TCHAR*	Name;
-	int32			BlockSizeX,
-					BlockSizeY,
-					BlockSizeZ,
-					BlockBytes,
-					NumComponents;
-	/** Platform specific token, e.g. D3DFORMAT with D3DDrv										*/
-	uint32			PlatformFormat;
+	FPixelFormatInfo() = delete;
+	FPixelFormatInfo(
+		EPixelFormat InUnrealFormat,
+		const TCHAR* InName,
+		int32 InBlockSizeX,
+		int32 InBlockSizeY,
+		int32 InBlockSizeZ,
+		int32 InBlockBytes,
+		int32 InNumComponents,
+		bool  InSupported);
+
+	const TCHAR*				Name;
+	EPixelFormat				UnrealFormat;
+	int32						BlockSizeX;
+	int32						BlockSizeY;
+	int32						BlockSizeZ;
+	int32						BlockBytes;
+	int32						NumComponents;
+
+	/** Per platform cabilities for the format */
+	EPixelFormatCapabilities	Capabilities{ EPixelFormatCapabilities::None };
+
+	/** Platform specific converted format */
+	uint32						PlatformFormat{ 0 };
+
 	/** Whether the texture format is supported on the current platform/ rendering combination	*/
-	bool			Supported;
-	EPixelFormat	UnrealFormat;
+	bool						Supported;
 };
 
 extern RHI_API FPixelFormatInfo GPixelFormats[PF_MAX];		// Maps members of EPixelFormat to a FPixelFormatInfo describing the format.
diff --git a/Engine/Source/Runtime/Windows/D3D11RHI/Private/D3D11Device.cpp b/Engine/Source/Runtime/Windows/D3D11RHI/Private/D3D11Device.cpp
--- a/Engine/Source/Runtime/Windows/D3D11RHI/Private/D3D11Device.cpp
+++ b/Engine/Source/Runtime/Windows/D3D11RHI/Private/D3D11Device.cpp
@@ -390,10 +390,77 @@ uint32 FD3D11DynamicRHI::GetMaxMSAAQuality(uint32 SampleCount)
 	return 0xffffffff;
 }
 
-static bool GFormatSupportsTypedUAVLoad[PF_MAX];
-
 void FD3D11DynamicRHI::SetupAfterDeviceCreation()
 {
+	for (uint32 FormatIndex = PF_Unknown; FormatIndex < PF_MAX; FormatIndex++)
+	{
+		FPixelFormatInfo& PixelFormatInfo = GPixelFormats[FormatIndex];
+		const DXGI_FORMAT PlatformFormat = static_cast<DXGI_FORMAT>(PixelFormatInfo.PlatformFormat);
+
+		EPixelFormatCapabilities Capabilities = EPixelFormatCapabilities::None;
+
+		if (PlatformFormat != DXGI_FORMAT_UNKNOWN)
+		{
+			D3D11_FEATURE_DATA_FORMAT_SUPPORT FormatSupport{};
+			FormatSupport.InFormat = PlatformFormat;
+
+			D3D11_FEATURE_DATA_FORMAT_SUPPORT2 FormatSupport2{};
+			FormatSupport2.InFormat = PlatformFormat;
+
+			auto ConvertCap1 = [&Capabilities, &FormatSupport](EPixelFormatCapabilities UnrealCap, D3D11_FORMAT_SUPPORT InFlag)
+			{
+				if (EnumHasAllFlags((D3D11_FORMAT_SUPPORT)FormatSupport.OutFormatSupport, InFlag))
+				{
+					EnumAddFlags(Capabilities, UnrealCap);
+				}
+			};
+			auto ConvertCap2 = [&Capabilities, &FormatSupport2](EPixelFormatCapabilities UnrealCap, D3D11_FORMAT_SUPPORT2 InFlag)
+			{
+				if (EnumHasAllFlags((D3D11_FORMAT_SUPPORT2)FormatSupport2.OutFormatSupport2, InFlag))
+				{
+					EnumAddFlags(Capabilities, UnrealCap);
+				}
+			};
+
+			HRESULT SupportHR = Direct3DDevice->CheckFeatureSupport(D3D11_FEATURE_FORMAT_SUPPORT, &FormatSupport, sizeof(FormatSupport));
+			HRESULT Support2HR = Direct3DDevice->CheckFeatureSupport(D3D11_FEATURE_FORMAT_SUPPORT2, &FormatSupport2, sizeof(FormatSupport2));
+			if (SUCCEEDED(SupportHR) || SUCCEEDED(Support2HR))
+			{
+				ConvertCap1(EPixelFormatCapabilities::Texture1D,			D3D11_FORMAT_SUPPORT_TEXTURE1D);
+				ConvertCap1(EPixelFormatCapabilities::Texture2D,			D3D11_FORMAT_SUPPORT_TEXTURE2D);
+				ConvertCap1(EPixelFormatCapabilities::Texture3D,			D3D11_FORMAT_SUPPORT_TEXTURE3D);
+				ConvertCap1(EPixelFormatCapabilities::TextureCube,			D3D11_FORMAT_SUPPORT_TEXTURECUBE);
+				ConvertCap1(EPixelFormatCapabilities::Buffer,               D3D11_FORMAT_SUPPORT_BUFFER);
+				ConvertCap1(EPixelFormatCapabilities::VertexBuffer,         D3D11_FORMAT_SUPPORT_IA_VERTEX_BUFFER);
+				ConvertCap1(EPixelFormatCapabilities::IndexBuffer,          D3D11_FORMAT_SUPPORT_IA_INDEX_BUFFER);
+
+				if (EnumHasAnyFlags(Capabilities, EPixelFormatCapabilities::AnyTexture))
+				{
+					ConvertCap1(EPixelFormatCapabilities::RenderTarget,     D3D11_FORMAT_SUPPORT_RENDER_TARGET);
+					ConvertCap1(EPixelFormatCapabilities::DepthStencil,     D3D11_FORMAT_SUPPORT_DEPTH_STENCIL);
+					ConvertCap1(EPixelFormatCapabilities::TextureMipmaps,   D3D11_FORMAT_SUPPORT_DEPTH_STENCIL);
+					ConvertCap1(EPixelFormatCapabilities::TextureLoad,      D3D11_FORMAT_SUPPORT_DEPTH_STENCIL);
+					ConvertCap1(EPixelFormatCapabilities::TextureSample,    D3D11_FORMAT_SUPPORT_DEPTH_STENCIL);
+					ConvertCap1(EPixelFormatCapabilities::TextureGather,    D3D11_FORMAT_SUPPORT_DEPTH_STENCIL);
+					ConvertCap1(EPixelFormatCapabilities::TextureAtomics,   D3D11_FORMAT_SUPPORT_DEPTH_STENCIL);
+					ConvertCap1(EPixelFormatCapabilities::TextureBlendable, D3D11_FORMAT_SUPPORT_DEPTH_STENCIL);
+				}
+
+				if (EnumHasAnyFlags(Capabilities, EPixelFormatCapabilities::Buffer))
+				{
+					ConvertCap1(EPixelFormatCapabilities::BufferLoad,       D3D11_FORMAT_SUPPORT_SHADER_LOAD);
+					ConvertCap2(EPixelFormatCapabilities::BufferStore,      D3D11_FORMAT_SUPPORT2_UAV_TYPED_STORE);
+					ConvertCap2(EPixelFormatCapabilities::BufferAtomics,    D3D11_FORMAT_SUPPORT2_UAV_ATOMIC_EXCHANGE);
+				}
+
+				ConvertCap2(EPixelFormatCapabilities::TypedUAVLoad,         D3D11_FORMAT_SUPPORT2_UAV_TYPED_LOAD);
+				ConvertCap2(EPixelFormatCapabilities::TypedUAVStore,        D3D11_FORMAT_SUPPORT2_UAV_TYPED_STORE);
+			}
+		}
+
+		PixelFormatInfo.Capabilities = Capabilities;
+	}
+
 	// without that the first RHIClear would get a scissor rect of (0,0)-(0,0) which means we get a draw call clear 
 	RHISetScissorRect(false, 0, 0, 0, 0);
 
@@ -432,34 +499,11 @@ void FD3D11DynamicRHI::SetupAfterDeviceCreation()
 			UE_LOG(LogD3D11RHI, Log, TEXT("Array index from any shader is supported"));
 		}
 	}
-
-	// Check for typed UAV load support
-	for (uint32 PF = 0; PF < PF_MAX; ++PF)
-	{
-		if (GPixelFormats[PF].PlatformFormat != 0)
-		{
-			D3D11_FEATURE_DATA_FORMAT_SUPPORT2 Data;
-			Data.InFormat = static_cast<DXGI_FORMAT>(GPixelFormats[PF].PlatformFormat);
-			Data.OutFormatSupport2 = 0;
-			HRESULT Result = Direct3DDevice->CheckFeatureSupport(D3D11_FEATURE_FORMAT_SUPPORT2, &Data, sizeof(Data));
-			GFormatSupportsTypedUAVLoad[PF] =
-				SUCCEEDED(Result) &&
-				(Data.OutFormatSupport2 & D3D11_FORMAT_SUPPORT2_UAV_TYPED_LOAD) == D3D11_FORMAT_SUPPORT2_UAV_TYPED_LOAD;
-		}
-		else
-		{
-			GFormatSupportsTypedUAVLoad[PF] = false;
-		}
-	}
-
-	GFormatSupportsTypedUAVLoad[PF_R32_FLOAT] = true;
-	GFormatSupportsTypedUAVLoad[PF_R32_UINT] = true;
-	GFormatSupportsTypedUAVLoad[PF_R32_SINT] = true;
 }
 
 bool FD3D11DynamicRHI::RHIIsTypedUAVLoadSupported(EPixelFormat PixelFormat)
 {
-	return GFormatSupportsTypedUAVLoad[PixelFormat];
+	return EnumHasAllFlags(GPixelFormats[PixelFormat].Capabilities, EPixelFormatCapabilities::TypedUAVLoad);
 }
 
 void FD3D11DynamicRHI::UpdateMSAASettings()
