Fix a memory leak

diff --git a/Engine/Source/Runtime/RHI/Public/RHICommandListCommandExecutes.inl b/Engine/Source/Runtime/RHI/Public/RHICommandListCommandExecutes.inl
--- a/Engine/Source/Runtime/RHI/Public/RHICommandListCommandExecutes.inl
+++ b/Engine/Source/Runtime/RHI/Public/RHICommandListCommandExecutes.inl
@@ -626,6 +626,7 @@ void FRHIResourceUpdateInfo::ReleaseRefs()
 		}
 		break;
 	case UT_BufferSRV:
+	case UT_BufferFormatSRV:
 		BufferSRV.SRV->Release();
 		if (BufferSRV.Buffer)
 		{
