Do not override buffer type for a null buffers inside FShaderResourceViewInitializer as they already initilized to a correct type

diff --git a/Engine/Source/Runtime/RHI/Private/DynamicRHI.cpp b/Engine/Source/Runtime/RHI/Private/DynamicRHI.cpp
--- a/Engine/Source/Runtime/RHI/Private/DynamicRHI.cpp
+++ b/Engine/Source/Runtime/RHI/Private/DynamicRHI.cpp
@@ -607,7 +607,7 @@ void FShaderResourceViewInitializer::InitType()
 			Type = EType::VertexBufferSRV;
 		}
 		else if (Usage & BUF_IndexBuffer)
-	{
+		{
 			Type = EType::IndexBufferSRV;
 		}
 		else if (Usage & BUF_AccelerationStructure)
@@ -619,8 +619,4 @@ void FShaderResourceViewInitializer::InitType()
 			Type = EType::StructuredBufferSRV;
 		}
 	}
-	else
-	{
-		Type = EType::StructuredBufferSRV;
-	}
 }
