Removed GSupportsResourceView as all platforms support it

diff --git a/Engine/Plugins/FX/Niagara/Source/Niagara/Private/NiagaraDataInterfaceCollisionQuery.cpp b/Engine/Plugins/FX/Niagara/Source/Niagara/Private/NiagaraDataInterfaceCollisionQuery.cpp
--- a/Engine/Plugins/FX/Niagara/Source/Niagara/Private/NiagaraDataInterfaceCollisionQuery.cpp
+++ b/Engine/Plugins/FX/Niagara/Source/Niagara/Private/NiagaraDataInterfaceCollisionQuery.cpp
@@ -954,11 +954,6 @@ bool UNiagaraDataInterfaceCollisionQuery::CopyToInternal(UNiagaraDataInterface*
 
 void UNiagaraDataInterfaceCollisionQuery::PushToRenderThreadImpl()
 {
-	if (!GSupportsResourceView)
-	{
-		return;
-	}
-
 	FNiagaraDataIntefaceProxyCollisionQuery* RT_Proxy = GetProxyAs<FNiagaraDataIntefaceProxyCollisionQuery>();
 
 	// Push Updates to Proxy, first release any resources
diff --git a/Engine/Plugins/FX/Niagara/Source/Niagara/Private/NiagaraDataInterfaceCurveBase.cpp b/Engine/Plugins/FX/Niagara/Source/Niagara/Private/NiagaraDataInterfaceCurveBase.cpp
--- a/Engine/Plugins/FX/Niagara/Source/Niagara/Private/NiagaraDataInterfaceCurveBase.cpp
+++ b/Engine/Plugins/FX/Niagara/Source/Niagara/Private/NiagaraDataInterfaceCurveBase.cpp
@@ -389,11 +389,6 @@ bool UNiagaraDataInterfaceCurveBase::GetExposedVariableValue(const FNiagaraVaria
 
 void UNiagaraDataInterfaceCurveBase::PushToRenderThreadImpl()
 {
-	if (!GSupportsResourceView)
-	{
-		return;
-	}
-
 	FNiagaraDataInterfaceProxyCurveBase* RT_Proxy = GetProxyAs<FNiagaraDataInterfaceProxyCurveBase>();
 
 	const float rtLUTMinTime = this->LUTMinTime;
diff --git a/Engine/Plugins/FX/Niagara/Source/Niagara/Private/NiagaraRendererMeshes.cpp b/Engine/Plugins/FX/Niagara/Source/Niagara/Private/NiagaraRendererMeshes.cpp
--- a/Engine/Plugins/FX/Niagara/Source/Niagara/Private/NiagaraRendererMeshes.cpp
+++ b/Engine/Plugins/FX/Niagara/Source/Niagara/Private/NiagaraRendererMeshes.cpp
@@ -911,9 +911,7 @@ void FNiagaraRendererMeshes::GetDynamicMeshElements(const TArray<const FSceneVie
 	if (SourceParticleData == nullptr ||
 		(SourceMode == ENiagaraRendererSourceDataMode::Particles && SourceParticleData->GetNumInstances() == 0) ||
 		Meshes.Num() == 0 ||
-		GbEnableNiagaraMeshRendering == 0 ||
-		!GSupportsResourceView  // Current shader requires SRV to draw properly in all cases.
-		)
+		GbEnableNiagaraMeshRendering == 0)
 	{
 		return;
 	}
@@ -1093,9 +1091,7 @@ void FNiagaraRendererMeshes::GetDynamicRayTracingInstances(FRayTracingMaterialGa
 		(SourceMode == ENiagaraRendererSourceDataMode::Particles && SourceParticleData->GetNumInstancesAllocated() == 0) ||
 		(SourceMode == ENiagaraRendererSourceDataMode::Particles && SourceParticleData->GetNumInstances() == 0) ||
 		Meshes.Num() == 0 ||
-		GbEnableNiagaraMeshRendering == 0 ||
-		!GSupportsResourceView  // Current shader requires SRV to draw properly in all cases.
-		)
+		GbEnableNiagaraMeshRendering == 0)
 	{
 		return;
 	}
diff --git a/Engine/Plugins/FX/Niagara/Source/Niagara/Private/NiagaraRendererRibbons.cpp b/Engine/Plugins/FX/Niagara/Source/Niagara/Private/NiagaraRendererRibbons.cpp
--- a/Engine/Plugins/FX/Niagara/Source/Niagara/Private/NiagaraRendererRibbons.cpp
+++ b/Engine/Plugins/FX/Niagara/Source/Niagara/Private/NiagaraRendererRibbons.cpp
@@ -426,9 +426,7 @@ void FNiagaraRendererRibbons::GetDynamicMeshElements(const TArray<const FSceneVi
 	if (!SourceParticleData ||
 		SourceParticleData->GetNumInstances() < 2 ||
 		DynamicDataRibbon->SegmentData.Num() == 0 ||
-		GbEnableNiagaraRibbonRendering == 0 ||
-		!GSupportsResourceView // Current shader requires SRV to draw properly in all cases.
-		)
+		GbEnableNiagaraRibbonRendering == 0)
 	{
 		return;
 	}
@@ -1419,9 +1417,7 @@ void FNiagaraRendererRibbons::GetDynamicRayTracingInstances(FRayTracingMaterialG
 	if (SourceParticleData == nullptr ||
 		SourceParticleData->GetNumInstancesAllocated() == 0 ||
 		SourceParticleData->GetNumInstances() == 0 ||
-		GbEnableNiagaraRibbonRendering == 0 ||
-		!GSupportsResourceView // Current shader requires SRV to draw properly in all cases.
-		)
+		GbEnableNiagaraRibbonRendering == 0)
 	{
 		return;
 	}
diff --git a/Engine/Plugins/FX/Niagara/Source/Niagara/Private/NiagaraRendererSprites.cpp b/Engine/Plugins/FX/Niagara/Source/Niagara/Private/NiagaraRendererSprites.cpp
--- a/Engine/Plugins/FX/Niagara/Source/Niagara/Private/NiagaraRendererSprites.cpp
+++ b/Engine/Plugins/FX/Niagara/Source/Niagara/Private/NiagaraRendererSprites.cpp
@@ -852,9 +852,7 @@ void FNiagaraRendererSprites::GetDynamicMeshElements(const TArray<const FSceneVi
 	FNiagaraDataBuffer* SourceParticleData = DynamicDataSprites->GetParticleDataToRender();
 	if (SourceParticleData == nullptr	||
 		(SourceMode == ENiagaraRendererSourceDataMode::Particles && SourceParticleData->GetNumInstances() == 0) ||
-		GbEnableNiagaraSpriteRendering == 0 ||
-		!GSupportsResourceView // Current shader requires SRV to draw properly in all cases.
-		)
+		GbEnableNiagaraSpriteRendering == 0)
 	{
 		return;
 	}
@@ -967,9 +965,7 @@ void FNiagaraRendererSprites::GetDynamicRayTracingInstances(FRayTracingMaterialG
 	if (SourceParticleData == nullptr	||
 		(SourceMode == ENiagaraRendererSourceDataMode::Particles && SourceParticleData->GetNumInstancesAllocated() == 0) ||
 		(SourceMode == ENiagaraRendererSourceDataMode::Particles && SourceParticleData->GetNumInstances() == 0) ||
-		GbEnableNiagaraSpriteRendering == 0 ||
-		!GSupportsResourceView // Current shader requires SRV to draw properly in all cases.
-		)
+		GbEnableNiagaraSpriteRendering == 0)
 	{
 		return;
 	}
diff --git a/Engine/Source/Runtime/Engine/Private/Particles/ParticleResources.cpp b/Engine/Source/Runtime/Engine/Private/Particles/ParticleResources.cpp
--- a/Engine/Source/Runtime/Engine/Private/Particles/ParticleResources.cpp
+++ b/Engine/Source/Runtime/Engine/Private/Particles/ParticleResources.cpp
@@ -137,17 +137,11 @@ TGlobalResource<FSixTriangleParticleIndexBuffer> GSixTriangleParticleIndexBuffer
 void FParticleScratchVertexBuffer::InitRHI()
 {
 	// Create a scratch vertex buffer for injecting particles and rendering tiles.
-	uint32 Flags = BUF_Volatile;
-	if (GSupportsResourceView)
-	{
-		Flags |= BUF_ShaderResource;
-	}
+	uint32 Flags = BUF_Volatile | BUF_ShaderResource;
+
 	FRHIResourceCreateInfo CreateInfo(TEXT("FParticleScratchVertexBuffer"));
 	VertexBufferRHI = RHICreateVertexBuffer(GParticleScratchVertexBufferSize, Flags, CreateInfo);
-	if (GSupportsResourceView)
-	{
-		VertexBufferSRV_G32R32F = RHICreateShaderResourceView(VertexBufferRHI, /*Stride=*/ sizeof(FVector2D), PF_G32R32F);
-	}
+	VertexBufferSRV_G32R32F = RHICreateShaderResourceView(VertexBufferRHI, /*Stride=*/ sizeof(FVector2D), PF_G32R32F);
 }
 
 FParticleShaderParamRef FParticleScratchVertexBuffer::GetShaderParam()
diff --git a/Engine/Source/Runtime/Engine/Private/Particles/ParticleVertexFactory.cpp b/Engine/Source/Runtime/Engine/Private/Particles/ParticleVertexFactory.cpp
--- a/Engine/Source/Runtime/Engine/Private/Particles/ParticleVertexFactory.cpp
+++ b/Engine/Source/Runtime/Engine/Private/Particles/ParticleVertexFactory.cpp
@@ -31,10 +31,7 @@ public:
 		FMemory::Memzero(BufferData, Size);
 		RHIUnlockBuffer(VertexBufferRHI);
 
-		if (GSupportsResourceView)
-		{
-			VertexBufferSRV = RHICreateShaderResourceView(VertexBufferRHI, sizeof(FVector2D), PF_G32R32F);
-		}
+		VertexBufferSRV = RHICreateShaderResourceView(VertexBufferRHI, sizeof(FVector2D), PF_G32R32F);
 	}
 	
 	virtual void ReleaseRHI() override
diff --git a/Engine/Source/Runtime/Engine/Private/Rendering/SkinWeightVertexBuffer.cpp b/Engine/Source/Runtime/Engine/Private/Rendering/SkinWeightVertexBuffer.cpp
--- a/Engine/Source/Runtime/Engine/Private/Rendering/SkinWeightVertexBuffer.cpp
+++ b/Engine/Source/Runtime/Engine/Private/Rendering/SkinWeightVertexBuffer.cpp
@@ -131,7 +131,7 @@ void FSkinWeightLookupVertexBuffer::InitRHI()
 
 	if (VertexBufferRHI)
 	{
-		bool bSRV = GSupportsResourceView && GPixelFormats[PixelFormat].Supported;
+		bool bSRV = GPixelFormats[PixelFormat].Supported;
 		// When bAllowCPUAccess is true, the meshes is likely going to be used for Niagara to spawn particles on mesh surface.
 		// And it can be the case for CPU *and* GPU access: no differenciation today. That is why we create a SRV in this case.
 		// This also avoid setting lots of states on all the members of all the different buffers used by meshes. Follow up: https://jira.it.epicgames.net/browse/UE-69376.
@@ -433,7 +433,7 @@ void FSkinWeightDataVertexBuffer::InitRHI()
 	// BUF_ShaderResource is needed for support of the SkinCache (we could make is dependent on GEnableGPUSkinCacheShaders or are there other users?)
 	VertexBufferRHI = CreateRHIBuffer_RenderThread();
 	
-	bool bSRV = VertexBufferRHI && GSupportsResourceView && GPixelFormats[GetPixelFormat()].Supported;
+	bool bSRV = VertexBufferRHI && GPixelFormats[GetPixelFormat()].Supported;
 	// When bAllowCPUAccess is true, the meshes is likely going to be used for Niagara to spawn particles on mesh surface.
 	// And it can be the case for CPU *and* GPU access: no differenciation today. That is why we create a SRV in this case.
 	// This also avoid setting lots of states on all the members of all the different buffers used by meshes. Follow up: https://jira.it.epicgames.net/browse/UE-69376.
diff --git a/Engine/Source/Runtime/Engine/Public/FXSystem.h b/Engine/Source/Runtime/Engine/Public/FXSystem.h
--- a/Engine/Source/Runtime/Engine/Public/FXSystem.h
+++ b/Engine/Source/Runtime/Engine/Public/FXSystem.h
@@ -106,8 +106,7 @@ inline bool RHISupportsGPUParticles()
 	return FXConsoleVariables::bAllowGPUParticles
 		&& GSupportsWideMRT
 		&& GPixelFormats[PF_G32R32F].Supported 
-		&& GSupportsTexture3D 
-		&& GSupportsResourceView;
+		&& GSupportsTexture3D;
 }
 
 class FFXSystemInterface;
diff --git a/Engine/Source/Runtime/OpenGLDrv/Private/Lumin/LuminGL4.h b/Engine/Source/Runtime/OpenGLDrv/Private/Lumin/LuminGL4.h
--- a/Engine/Source/Runtime/OpenGLDrv/Private/Lumin/LuminGL4.h
+++ b/Engine/Source/Runtime/OpenGLDrv/Private/Lumin/LuminGL4.h
@@ -564,8 +564,6 @@ struct FLuminOpenGL4 : public FOpenGL4
 	// whether NV_bindless_texture is supported
 	static bool bSupportsBindlessTexture;
 
-	static FORCEINLINE bool SupportsResourceView() { return true; }
-
 	static FORCEINLINE void TexBuffer(GLenum Target, GLenum InternalFormat, GLuint Buffer)
 	{
 		glTexBuffer(Target, InternalFormat, Buffer);
diff --git a/Engine/Source/Runtime/OpenGLDrv/Private/OpenGLDevice.cpp b/Engine/Source/Runtime/OpenGLDrv/Private/OpenGLDevice.cpp
--- a/Engine/Source/Runtime/OpenGLDrv/Private/OpenGLDevice.cpp
+++ b/Engine/Source/Runtime/OpenGLDrv/Private/OpenGLDevice.cpp
@@ -847,7 +847,6 @@ static void InitRHICapabilitiesForGL()
 	GSupportsTexture3D = FOpenGL::SupportsTexture3D();
 	GSupportsMobileMultiView = FOpenGL::SupportsMobileMultiView();
 	GSupportsImageExternal = FOpenGL::SupportsImageExternal();
-	GSupportsResourceView = FOpenGL::SupportsResourceView();
 
 	GSupportsShaderFramebufferFetch = FOpenGL::SupportsShaderFramebufferFetch();
 	GSupportsShaderMRTFramebufferFetch = FOpenGL::SupportsShaderMRTFramebufferFetch();
diff --git a/Engine/Source/Runtime/OpenGLDrv/Public/OpenGL.h b/Engine/Source/Runtime/OpenGLDrv/Public/OpenGL.h
--- a/Engine/Source/Runtime/OpenGLDrv/Public/OpenGL.h
+++ b/Engine/Source/Runtime/OpenGLDrv/Public/OpenGL.h
@@ -102,7 +102,6 @@ public:
 	static FORCEINLINE bool SupportsTextureLODBias()					{ return true; }
 	static FORCEINLINE bool SupportsTextureCompare()					{ return true; }
 	static FORCEINLINE bool SupportsDrawIndexOffset()					{ return true; }
-	static FORCEINLINE bool SupportsResourceView()						{ return true; }
 	static FORCEINLINE bool SupportsDiscardFrameBuffer()				{ return false; }
 	static FORCEINLINE bool SupportsIndexedExtensions()					{ return true; }
 	static FORCEINLINE bool SupportsColorBufferFloat()					{ return true; }
diff --git a/Engine/Source/Runtime/OpenGLDrv/Public/OpenGLES.h b/Engine/Source/Runtime/OpenGLDrv/Public/OpenGLES.h
--- a/Engine/Source/Runtime/OpenGLDrv/Public/OpenGLES.h
+++ b/Engine/Source/Runtime/OpenGLDrv/Public/OpenGLES.h
@@ -196,7 +196,6 @@ struct FOpenGLES : public FOpenGLBase
 	static FORCEINLINE bool SupportsTextureLODBias() { return false; }
 	static FORCEINLINE bool SupportsTextureCompare() { return false; }
 	static FORCEINLINE bool SupportsDrawIndexOffset() { return false; }
-	static FORCEINLINE bool SupportsResourceView() { return glTexBufferEXT != nullptr; }
 	static FORCEINLINE bool SupportsDiscardFrameBuffer() { return true; }
 	static FORCEINLINE bool SupportsIndexedExtensions() { return false; }
 	static FORCEINLINE bool SupportsColorBufferFloat() { return bSupportsColorBufferFloat; }
diff --git a/Engine/Source/Runtime/RHI/Private/RHI.cpp b/Engine/Source/Runtime/RHI/Private/RHI.cpp
--- a/Engine/Source/Runtime/RHI/Private/RHI.cpp
+++ b/Engine/Source/Runtime/RHI/Private/RHI.cpp
@@ -1193,7 +1193,6 @@ bool GRHINeedsUnatlasedCSMDepthsWorkaround = false;
 bool GSupportsTexture3D = true;
 bool GSupportsMobileMultiView = false;
 bool GSupportsImageExternal = false;
-bool GSupportsResourceView = true;
 bool GRHISupportsDrawIndirect = true;
 bool GRHISupportsMultithreading = false;
 bool GRHISupportsUpdateFromBufferTexture = false;
diff --git a/Engine/Source/Runtime/RHI/Public/RHI.h b/Engine/Source/Runtime/RHI/Public/RHI.h
--- a/Engine/Source/Runtime/RHI/Public/RHI.h
+++ b/Engine/Source/Runtime/RHI/Public/RHI.h
@@ -80,9 +80,6 @@ extern RHI_API bool GSupportsRenderDepthTargetableShaderResources;
 extern RHI_API ERHIFeatureLevel::Type GMaxRHIFeatureLevel;
 extern RHI_API EShaderPlatform GMaxRHIShaderPlatform;
 
-/** true if the RHI supports SRVs */
-extern RHI_API bool GSupportsResourceView;
-
 /** true if the RHI supports Draw Indirect */
 extern RHI_API bool GRHISupportsDrawIndirect;
 
diff --git a/Engine/Source/Runtime/RHI/Public/RHIUtilities.h b/Engine/Source/Runtime/RHI/Public/RHIUtilities.h
--- a/Engine/Source/Runtime/RHI/Public/RHIUtilities.h
+++ b/Engine/Source/Runtime/RHI/Public/RHIUtilities.h
@@ -87,14 +87,6 @@ static bool IsFloatFormat(EPixelFormat Format)
 	return false;
 }
 
-inline bool RHISupportsTextureBuffers(const FStaticShaderPlatform Platform)
-{
-	return IsFeatureLevelSupported(Platform, ERHIFeatureLevel::SM5)
-		|| IsVulkanPlatform(Platform)
-		|| IsMetalPlatform(Platform)
-		|| (GetMaxSupportedFeatureLevel(Platform) == ERHIFeatureLevel::ES3_1 && GSupportsResourceView);
-}
-
 /** Get the best default resource state for the given texture creation flags */
 extern RHI_API ERHIAccess RHIGetDefaultResourceState(ETextureCreateFlags InUsage, bool bInHasInitialData);
 
@@ -123,8 +115,6 @@ struct FTextureRWBuffer2D
 	static constexpr ETextureCreateFlags DefaultTextureInitFlag = TexCreate_ShaderResource | TexCreate_UAV;
 	void Initialize(const TCHAR* InDebugName, const uint32 BytesPerElement, const uint32 SizeX, const uint32 SizeY, const EPixelFormat Format, ETextureCreateFlags Flags = DefaultTextureInitFlag)
 	{
-		check(RHISupportsTextureBuffers(GMaxRHIShaderPlatform));
-
 		NumBytes = SizeX * SizeY * BytesPerElement;
 
 		FRHIResourceCreateInfo CreateInfo(InDebugName);
@@ -185,8 +175,6 @@ struct FTextureRWBuffer3D
 	// @param AdditionalUsage passed down to RHICreateVertexBuffer(), get combined with "BUF_UnorderedAccess | BUF_ShaderResource" e.g. BUF_Static
 	void Initialize(const TCHAR* InDebugName, uint32 BytesPerElement, uint32 SizeX, uint32 SizeY, uint32 SizeZ, EPixelFormat Format)
 	{
-		check(RHISupportsTextureBuffers(GMaxRHIShaderPlatform));
-
 		NumBytes = SizeX * SizeY * SizeZ * BytesPerElement;
 
 		FRHIResourceCreateInfo CreateInfo(InDebugName);
@@ -284,8 +272,6 @@ struct FRWBuffer
 	// @param AdditionalUsage passed down to RHICreateVertexBuffer(), get combined with "BUF_UnorderedAccess | BUF_ShaderResource" e.g. BUF_Static
 	void Initialize(const TCHAR* InDebugName, uint32 BytesPerElement, uint32 NumElements, EPixelFormat Format, ERHIAccess InResourceState, uint32 AdditionalUsage = 0, FResourceArrayInterface *InResourceArray = nullptr)	
 	{
-		check(RHISupportsTextureBuffers(GMaxRHIShaderPlatform));
-
 		// Provide a debug name if using Fast VRAM so the allocators diagnostics will work
 		ensure(!((AdditionalUsage & BUF_FastVRAM) && !InDebugName));
 		NumBytes = BytesPerElement * NumElements;
@@ -346,8 +332,6 @@ struct FTextureReadBuffer2D
 	const static ETextureCreateFlags DefaultTextureInitFlag = TexCreate_ShaderResource;
 	void Initialize(const TCHAR* InDebugName, const uint32 BytesPerElement, const uint32 SizeX, const uint32 SizeY, const EPixelFormat Format, ETextureCreateFlags Flags = DefaultTextureInitFlag, FResourceBulkDataInterface* InBulkData = nullptr)
 	{
-		check(RHISupportsTextureBuffers(GMaxRHIShaderPlatform));
-
 		NumBytes = SizeX * SizeY * BytesPerElement;
 
 		FRHIResourceCreateInfo CreateInfo(InDebugName);
@@ -398,7 +382,6 @@ struct FReadBuffer
 
 	void Initialize(const TCHAR* InDebugName, uint32 BytesPerElement, uint32 NumElements, EPixelFormat Format, uint32 AdditionalUsage = 0, FResourceArrayInterface* InResourceArray = nullptr)
 	{
-		check(GSupportsResourceView);
 		NumBytes = BytesPerElement * NumElements;
 		FRHIResourceCreateInfo CreateInfo(InDebugName);
 		CreateInfo.ResourceArray = InResourceArray;
