U5 - moving the flush from before FlushRHIThreadFlushResources to after FlushRHIThreadFlushResources resulted in the "removal" of the quiescent state of the PipelineStateCache as now parallel RT commands were still able to run while we handle the PSO cache deletions. Unfortunately we now have to double flush (before and after) for some GT related flush cases, but as we try to remove flushes in general this might be acceptable until we find a better solution.

diff --git a/Engine/Source/Runtime/RHI/Public/RHICommandList.inl b/Engine/Source/Runtime/RHI/Public/RHICommandList.inl
--- a/Engine/Source/Runtime/RHI/Public/RHICommandList.inl
+++ b/Engine/Source/Runtime/RHI/Public/RHICommandList.inl
@@ -73,6 +73,7 @@ FORCEINLINE_DEBUGGABLE void FRHICommandListImmediate::ImmediateFlush(EImmediateF
 			WaitForTasks();
 		}
 		break;
+
 	case EImmediateFlushType::DispatchToRHIThread:
 		{
 			if (HasCommands())
@@ -81,6 +82,7 @@ FORCEINLINE_DEBUGGABLE void FRHICommandListImmediate::ImmediateFlush(EImmediateF
 			}
 		}
 		break;
+
 	case EImmediateFlushType::WaitForDispatchToRHIThread:
 		{
 			if (HasCommands())
@@ -91,13 +93,6 @@ FORCEINLINE_DEBUGGABLE void FRHICommandListImmediate::ImmediateFlush(EImmediateF
 		}
 		break;
 
-	case EImmediateFlushType::FlushRHIThreadFlushResources:
-		{
-			CSV_SCOPED_TIMING_STAT(RHITFlushes, FlushRHIThreadFlushResourcesTotal);
-			PipelineStateCache::FlushResources();
-			FRHIResource::FlushPendingDeletes(FRHICommandListExecutor::GetImmediateCommandList());
-		}
-	//this fall-through is intentional
 	case EImmediateFlushType::FlushRHIThread:
 		{
 			CSV_SCOPED_TIMING_STAT(RHITFlushes, FlushRHIThreadTotal);
@@ -113,6 +108,23 @@ FORCEINLINE_DEBUGGABLE void FRHICommandListImmediate::ImmediateFlush(EImmediateF
 			WaitForTasks(true); // these are already done, but this resets the outstanding array
 		}
 		break;
+
+	case EImmediateFlushType::FlushRHIThreadFlushResources:
+		{
+			CSV_SCOPED_TIMING_STAT(RHITFlushes, FlushRHIThreadFlushResourcesTotal);
+			if (HasCommands())
+			{
+				GRHICommandList.ExecuteList(*this);
+			}
+			WaitForDispatch();
+			WaitForRHIThreadTasks();
+			WaitForTasks(true); // these are already done, but this resets the outstanding array
+
+			PipelineStateCache::FlushResources();
+			FRHIResource::FlushPendingDeletes(FRHICommandListExecutor::GetImmediateCommandList());
+		}
+		break;
+
 	default:
 		check(0);
 	}
diff --git a/Engine/Source/Runtime/RenderCore/Private/RenderingThread.cpp b/Engine/Source/Runtime/RenderCore/Private/RenderingThread.cpp
--- a/Engine/Source/Runtime/RenderCore/Private/RenderingThread.cpp
+++ b/Engine/Source/Runtime/RenderCore/Private/RenderingThread.cpp
@@ -1292,6 +1292,8 @@ void FlushRenderingCommands()
 	ENQUEUE_RENDER_COMMAND(FlushPendingDeleteRHIResourcesCmd)([](FRHICommandListImmediate& RHICmdList)
 	{
 		RHICmdList.ImmediateFlush(EImmediateFlushType::FlushRHIThreadFlushResources);
+		//double flush to flush out the deferred deletions queued into the ImmediateCmdList
+		RHICmdList.ImmediateFlush(EImmediateFlushType::FlushRHIThread);
 	});
 
 	// Find the objects which may be cleaned up once the rendering thread command queue has been flushed.
