U5 - adding RHI resource validation that can effectively detect the annoying resource resurrection mechanics. Which immediately triggered an issue where resurrection happened from within the GetCachedBoundShaderState_Threadsafe while the resource was already processed for deletion. This would affect Dx12 and Vulkan API alike.

diff --git a/Engine/Source/Runtime/RHI/Private/BoundShaderStateCache.cpp b/Engine/Source/Runtime/RHI/Private/BoundShaderStateCache.cpp
--- a/Engine/Source/Runtime/RHI/Private/BoundShaderStateCache.cpp
+++ b/Engine/Source/Runtime/RHI/Private/BoundShaderStateCache.cpp
@@ -151,7 +151,7 @@ FBoundShaderStateRHIRef GetCachedBoundShaderState_Threadsafe(
 		);
 	}
 	
-	if(CachedBoundShaderStateLink)
+	if(CachedBoundShaderStateLink && CachedBoundShaderStateLink->BoundShaderState->IsValid())
 	{
 		// If we've already created a bound shader state with these parameters, reuse it.
 		return CachedBoundShaderStateLink->BoundShaderState;
diff --git a/Engine/Source/Runtime/RHI/Private/RHI.cpp b/Engine/Source/Runtime/RHI/Private/RHI.cpp
--- a/Engine/Source/Runtime/RHI/Private/RHI.cpp
+++ b/Engine/Source/Runtime/RHI/Private/RHI.cpp
@@ -1022,16 +1022,11 @@ void FRHIResource::FlushPendingDeletes(FRHICommandListImmediate& RHICmdList)
 
 		for (FRHIResource* Resource : DeletedResources)
 		{
-			check(Resource->MarkedForDelete.load(std::memory_order_relaxed) == 1);
-			if (Resource->NumRefs.load(std::memory_order_acquire) == 0) // caches can bring dead objects back to life
+			if (Resource->AtomicFlags.Deleteing())
 			{
 				FRHIResource::CurrentlyDeleting = Resource;
 				delete Resource;
 			}
-			else
-			{
-				verify(Resource->MarkedForDelete.exchange(0, std::memory_order_release) == 1);	
-			}	
 		}
 	});
 }
diff --git a/Engine/Source/Runtime/RHI/Public/RHIResources.h b/Engine/Source/Runtime/RHI/Public/RHIResources.h
--- a/Engine/Source/Runtime/RHI/Public/RHIResources.h
+++ b/Engine/Source/Runtime/RHI/Public/RHIResources.h
@@ -29,6 +29,11 @@
 #define RHI_WANT_RESOURCE_INFO 0
 #endif
 
+#ifndef RHI_RESOURCE_LIFETIME_VALIDATION
+#define RHI_RESOURCE_LIFETIME_VALIDATION 0
+#endif
+
+
 class FRHICommandListImmediate;
 struct FClearValueBinding;
 struct FRHIResourceInfo;
@@ -56,7 +61,7 @@ public:
 	virtual ~FRHIResource()
 	{
 		check(IsEngineExitRequested() || CurrentlyDeleting == this);
-		check(NumRefs.load(std::memory_order_relaxed) == 0); // this should not have any outstanding refs
+		check(AtomicFlags.GetNumRefs(std::memory_order_relaxed) == 0); // this should not have any outstanding refs
 		CurrentlyDeleting = nullptr;
 
 #if RHI_WANT_RESOURCE_INFO
@@ -66,18 +71,18 @@ public:
 
 	FORCEINLINE_DEBUGGABLE uint32 AddRef() const
 	{
-		int32 NewValue = NumRefs.fetch_add(1, std::memory_order_acquire) + 1;
+		int32 NewValue = AtomicFlags.AddRef(std::memory_order_acquire);
 		checkSlow(NewValue > 0); 
 		return uint32(NewValue);
 	}
 	FORCEINLINE_DEBUGGABLE uint32 Release() const
 	{
-		int32 NewValue = NumRefs.fetch_sub(1, std::memory_order_release) - 1;
+		int32 NewValue = AtomicFlags.Release(std::memory_order_release);
 		check(NewValue >= 0);
 
 		if (NewValue == 0)
 		{
-			if (MarkedForDelete.exchange(1, std::memory_order_release) == 0)
+			if (!AtomicFlags.MarkForDelete(std::memory_order_release))
 			{
 				while(true)
 				{
@@ -95,7 +100,7 @@ public:
 	}
 	FORCEINLINE_DEBUGGABLE uint32 GetRefCount() const
 	{
-		int32 CurrentValue = NumRefs.load(std::memory_order_relaxed);
+		int32 CurrentValue = AtomicFlags.GetNumRefs(std::memory_order_relaxed);
 		checkSlow(CurrentValue >= 0); 
 		return uint32(CurrentValue);
 	}
@@ -119,12 +124,12 @@ public:
 
 	bool IsValid() const
 	{
-		return !MarkedForDelete.load(std::memory_order_relaxed) && NumRefs.load(std::memory_order_relaxed) > 0;
+		return AtomicFlags.IsValid(std::memory_order_relaxed);
 	}
 
 	void Delete()
 	{
-		verify(MarkedForDelete.fetch_add(1, std::memory_order_acquire) == 0);
+		verify(!AtomicFlags.MarkForDelete(std::memory_order_acquire));
 		CurrentlyDeleting = this;
 		delete this;
 	}
@@ -146,8 +151,160 @@ public:
 #endif
 
 private:
-	mutable std::atomic_int NumRefs { 0 };
-	mutable std::atomic_int MarkedForDelete { 0 };
+	class FAtomicFlags
+	{
+#if RHI_RESOURCE_LIFETIME_VALIDATION
+		struct FPacked
+		{
+			FPacked(uint64 InNumRefs, bool InMarkedForDelete, bool InDeleting) : NumRefs(InNumRefs), MarkedForDelete(InMarkedForDelete), Deleting(InDeleting)
+			{
+			}
+
+			FPacked() : NumRefs(0), MarkedForDelete(0), Deleting(0)
+			{
+			}
+
+			uint32 NumRefs			: 30;
+			uint32 MarkedForDelete	: 1;
+			uint32 Deleting			: 1;
+		};
+		std::atomic<FPacked> Packed = {};
+#else
+		std::atomic_int NumRefs = { 0 };
+		std::atomic_bool MarkedForDelete = { 0 };
+#endif
+
+	public:
+		int32 AddRef(std::memory_order MemoryOrder)
+		{
+#if RHI_RESOURCE_LIFETIME_VALIDATION
+			FPacked Old = Packed.load(std::memory_order_relaxed);
+			while(true)
+			{
+				check(Old.Deleting == false);
+				if(Packed.compare_exchange_weak(Old, FPacked(Old.NumRefs + 1, Old.MarkedForDelete, false), MemoryOrder, std::memory_order_relaxed))
+				{
+					return Old.NumRefs + 1;
+				}
+			}
+#else
+			return NumRefs.fetch_add(1, MemoryOrder) + 1;
+#endif
+		}
+
+		int32 Release(std::memory_order MemoryOrder)
+		{
+#if RHI_RESOURCE_LIFETIME_VALIDATION
+			FPacked Old = Packed.load(std::memory_order_relaxed);
+			while(true)
+			{
+				check(Old.Deleting == false);
+				if(Packed.compare_exchange_weak(Old, FPacked(Old.NumRefs - 1, Old.MarkedForDelete, false), MemoryOrder, std::memory_order_relaxed))
+				{
+					return Old.NumRefs - 1;
+				}
+			}
+#else
+			return NumRefs.fetch_sub(1, MemoryOrder) - 1;
+#endif
+		}
+
+		bool MarkForDelete(std::memory_order MemoryOrder)
+		{
+#if RHI_RESOURCE_LIFETIME_VALIDATION
+			FPacked Old = Packed.load(std::memory_order_relaxed);
+			while(true)
+			{
+				check(Old.Deleting == false);
+				if(Packed.compare_exchange_weak(Old, FPacked(Old.NumRefs, true, false), MemoryOrder, std::memory_order_relaxed))
+				{
+					return Old.MarkedForDelete;
+				}
+			}
+#else
+			return MarkedForDelete.exchange(true, MemoryOrder);
+#endif
+		}
+
+		bool UnmarkForDelete(std::memory_order MemoryOrder)
+		{
+#if RHI_RESOURCE_LIFETIME_VALIDATION
+			FPacked Old = Packed.load(std::memory_order_relaxed);
+			while(true)
+			{
+				check(Old.Deleting == false); 
+				check(Old.MarkedForDelete == true);
+				if(Packed.compare_exchange_weak(Old, FPacked(Old.NumRefs, false, false), MemoryOrder, std::memory_order_relaxed))
+				{
+					return Old.MarkedForDelete;
+				}
+			}
+#else
+			bool Old = MarkedForDelete.exchange(false, MemoryOrder);
+			check(Old == true);
+			return Old;
+#endif
+		}
+
+		bool Deleteing()
+		{
+#if RHI_RESOURCE_LIFETIME_VALIDATION
+			FPacked Old = Packed.load(std::memory_order_relaxed);
+			while(true)
+			{
+				check(Old.Deleting == false); 
+				check(Old.MarkedForDelete == true);
+				if(Old.NumRefs == 0)
+				{
+					if(Packed.compare_exchange_weak(Old, FPacked(0, true, true), std::memory_order_acquire, std::memory_order_relaxed))
+					{
+						return true;
+					}
+				}
+				else
+				{
+					if(Packed.compare_exchange_weak(Old, FPacked(Old.NumRefs, false, false), std::memory_order_release, std::memory_order_relaxed))
+					{
+						return false;
+					}
+				}
+			}
+#else
+			check(MarkedForDelete.load(std::memory_order_relaxed) == 1);
+			if (NumRefs.load(std::memory_order_acquire) == 0) // caches can bring dead objects back to life
+			{
+				return true;
+			}
+			else
+			{
+				verify(MarkedForDelete.exchange(false, std::memory_order_release));	
+				return false;
+			}
+#endif
+		}
+
+		bool IsValid(std::memory_order MemoryOrder)
+		{
+#if RHI_RESOURCE_LIFETIME_VALIDATION
+			FPacked Old = Packed.load(MemoryOrder);
+			return !Old.MarkedForDelete && Old.NumRefs > 0;
+#else
+			return !MarkedForDelete.load(MemoryOrder) && NumRefs.load(MemoryOrder) > 0;
+#endif
+		}
+
+		int32 GetNumRefs(std::memory_order MemoryOrder)
+		{
+#if RHI_RESOURCE_LIFETIME_VALIDATION
+			FPacked Old = Packed.load(MemoryOrder);
+			return Old.NumRefs;
+#else
+			return NumRefs.load(MemoryOrder);
+#endif
+		}
+	};
+	mutable FAtomicFlags AtomicFlags;
+
 	const ERHIResourceType ResourceType;
 	bool bCommitted { true };
 #if RHI_WANT_RESOURCE_INFO
