UE5 - RHI Validation error if nesting RHIBegin/EndFrame
#rnx

diff --git a/Engine/Source/Runtime/RHI/Private/RHIValidation.cpp b/Engine/Source/Runtime/RHI/Private/RHIValidation.cpp
--- a/Engine/Source/Runtime/RHI/Private/RHIValidation.cpp
+++ b/Engine/Source/Runtime/RHI/Private/RHIValidation.cpp
@@ -446,6 +446,20 @@ FValidationContext::FValidationContext()
 	Tracker = &State.TrackerInstance;
 }
 
+void FValidationContext::RHIBeginFrame()
+{
+	++State.BeginEndFrameCounter;
+	if (State.BeginEndFrameCounter != 1)
+	{
+		ensureMsgf(0, TEXT("RHIBeginFrame called twice in a row! Previous callstack: (void**)0x%p,32"), State.PreviousBeginFrame);
+	}
+	delete [] State.PreviousBeginFrame;
+	State.PreviousBeginFrame = RHIValidation::CaptureBacktrace();
+
+	State.Reset();
+	RHIContext->RHIBeginFrame();
+}
+
 void FValidationContext::RHIEndFrame()
 {
 	RHIContext->RHIEndFrame();
@@ -454,6 +468,15 @@ void FValidationContext::RHIEndFrame()
 	// The RenderThread FrameID is update in RHIAdvanceFrameFence which is called on the RenderThread
 	FValidationRHI* ValidateRHI = (FValidationRHI*)GDynamicRHI;
 	ValidateRHI->RHIThreadFrameID++;
+
+	--State.BeginEndFrameCounter;
+	if (State.BeginEndFrameCounter != 0)
+	{
+		ensureMsgf(0, TEXT("RHIEndFrame called twice in a row! Previous callstack: (void**)0x%p,32"), State.PreviousEndFrame);
+		State.BeginEndFrameCounter = 0;
+	}
+	delete [] State.PreviousEndFrame;
+	State.PreviousEndFrame = RHIValidation::CaptureBacktrace();
 }
 
 void FValidationContext::FState::Reset()
diff --git a/Engine/Source/Runtime/RHI/Public/RHIValidationContext.h b/Engine/Source/Runtime/RHI/Public/RHIValidationContext.h
--- a/Engine/Source/Runtime/RHI/Public/RHIValidationContext.h
+++ b/Engine/Source/Runtime/RHI/Public/RHIValidationContext.h
@@ -463,11 +463,7 @@ public:
 	}
 
 	// This method is queued with an RHIThread, otherwise it will flush after it is queued; without an RHI thread there is no benefit to queuing this frame advance commands
-	virtual void RHIBeginFrame() override final
-	{
-		State.Reset();
-		RHIContext->RHIBeginFrame();
-	}
+	virtual void RHIBeginFrame() override final;
 
 	// This method is queued with an RHIThread, otherwise it will flush after it is queued; without an RHI thread there is no benefit to queuing this frame advance commands
 	virtual void RHIEndFrame() override final;
@@ -1086,6 +1082,10 @@ protected:
 		RHIValidation::FTracker TrackerInstance{ ERHIPipeline::Graphics };
 		RHIValidation::FGlobalUniformBuffers GlobalUniformBuffers;
 
+		void* PreviousBeginFrame = nullptr;
+		void* PreviousEndFrame = nullptr;
+		int32 BeginEndFrameCounter = 0;
+
 		FRHIRenderPassInfo RenderPassInfo;
 		FString RenderPassName;
 		FString PreviousRenderPassName;
