Fix for clang compile error

diff --git a/Engine/Source/Runtime/RHI/Private/RHIValidation.cpp b/Engine/Source/Runtime/RHI/Private/RHIValidation.cpp
--- a/Engine/Source/Runtime/RHI/Private/RHIValidation.cpp
+++ b/Engine/Source/Runtime/RHI/Private/RHIValidation.cpp
@@ -453,7 +453,7 @@ void FValidationContext::RHIBeginFrame()
 	{
 		ensureMsgf(0, TEXT("RHIBeginFrame called twice in a row! Previous callstack: (void**)0x%p,32"), State.PreviousBeginFrame);
 	}
-	delete [] State.PreviousBeginFrame;
+	delete [] (uint64*)State.PreviousBeginFrame;
 	State.PreviousBeginFrame = RHIValidation::CaptureBacktrace();
 
 	State.Reset();
@@ -475,7 +475,7 @@ void FValidationContext::RHIEndFrame()
 		ensureMsgf(0, TEXT("RHIEndFrame called twice in a row! Previous callstack: (void**)0x%p,32"), State.PreviousEndFrame);
 		State.BeginEndFrameCounter = 0;
 	}
-	delete [] State.PreviousEndFrame;
+	delete [] (uint64*)State.PreviousEndFrame;
 	State.PreviousEndFrame = RHIValidation::CaptureBacktrace();
 }
 
