Make default RHI to be D3D12 on Windows.

diff --git a/Engine/Source/Runtime/RHI/Private/Windows/WindowsDynamicRHI.cpp b/Engine/Source/Runtime/RHI/Private/Windows/WindowsDynamicRHI.cpp
--- a/Engine/Source/Runtime/RHI/Private/Windows/WindowsDynamicRHI.cpp
+++ b/Engine/Source/Runtime/RHI/Private/Windows/WindowsDynamicRHI.cpp
@@ -27,7 +27,7 @@ static bool ShouldPreferD3D12()
 		}
 	}
 	
-	return false;
+	return true;
 }
 
 static bool ShouldForceFeatureLevelES31()
