Re-enable async compilation for the (yet-to-be-cached) compute PSOs.

diff --git a/Engine/Source/Runtime/RHI/Private/PipelineStateCache.cpp b/Engine/Source/Runtime/RHI/Private/PipelineStateCache.cpp
--- a/Engine/Source/Runtime/RHI/Private/PipelineStateCache.cpp
+++ b/Engine/Source/Runtime/RHI/Private/PipelineStateCache.cpp
@@ -1005,8 +1005,7 @@ uint64 PipelineStateCache::RetrieveGraphicsPipelineStateSortKey(FGraphicsPipelin
 
 FComputePipelineState* PipelineStateCache::GetAndOrCreateComputePipelineState(FRHIComputeCommandList& RHICmdList, FRHIComputeShader* ComputeShader)
 {	
-	// TODO: add async compute PSO compilation support
-	bool DoAsyncCompile = false;//IsAsyncCompilationAllowed(RHICmdList);
+	bool DoAsyncCompile = IsAsyncCompilationAllowed(RHICmdList);
 
 	FComputePipelineState* OutCachedState = nullptr;
 
