Re-enable async compilation for the (yet-to-be-cached) compute PSOs.

diff --git a/Engine/Source/Runtime/RHI/Private/PipelineStateCache.cpp b/Engine/Source/Runtime/RHI/Private/PipelineStateCache.cpp
--- a/Engine/Source/Runtime/RHI/Private/PipelineStateCache.cpp
+++ b/Engine/Source/Runtime/RHI/Private/PipelineStateCache.cpp
@@ -861,14 +861,23 @@ public:
 		: Pipeline(InPipeline)
 		, Initializer(InInitializer)
 	{
-		Initializer.BoundShaderState.AddRefResources();
+		if (!Pipeline->IsCompute())
+		{
+			Initializer.BoundShaderState.AddRefResources();
 
-		if (Initializer.BlendState)
-			Initializer.BlendState->AddRef();
-		if (Initializer.RasterizerState)
-			Initializer.RasterizerState->AddRef();
-		if (Initializer.DepthStencilState)
-			Initializer.DepthStencilState->AddRef();
+			if (Initializer.BlendState)
+			{
+				Initializer.BlendState->AddRef();
+			}
+			if (Initializer.RasterizerState)
+			{
+				Initializer.RasterizerState->AddRef();
+			}
+			if (Initializer.DepthStencilState)
+			{
+				Initializer.DepthStencilState->AddRef();
+			}
+		}
 	}
 
 	static ESubsequentsMode::Type GetSubsequentsMode() { return ESubsequentsMode::TrackSubsequents; }
@@ -902,7 +911,7 @@ public:
 
 			FGraphicsPipelineState* GfxPipeline = static_cast<FGraphicsPipelineState*>(Pipeline);
 			GfxPipeline->RHIPipeline = RHICreateGraphicsPipelineState(Initializer);
-			
+
 			if (GfxPipeline->RHIPipeline)
 			{
 				GfxPipeline->SortKey = GfxPipeline->RHIPipeline->GetSortKey();
@@ -913,13 +922,19 @@ public:
 			}
 
 			Initializer.BoundShaderState.ReleaseResources();
-			
+
 			if (Initializer.BlendState)
+			{
 				Initializer.BlendState->Release();
+			}
 			if (Initializer.RasterizerState)
+			{
 				Initializer.RasterizerState->Release();
+			}
 			if (Initializer.DepthStencilState)
+			{
 				Initializer.DepthStencilState->Release();
+			}
 		}
 	}
 
@@ -1005,8 +1020,7 @@ uint64 PipelineStateCache::RetrieveGraphicsPipelineStateSortKey(FGraphicsPipelin
 
 FComputePipelineState* PipelineStateCache::GetAndOrCreateComputePipelineState(FRHIComputeCommandList& RHICmdList, FRHIComputeShader* ComputeShader)
 {	
-	// TODO: add async compute PSO compilation support
-	bool DoAsyncCompile = false;//IsAsyncCompilationAllowed(RHICmdList);
+	bool DoAsyncCompile = IsAsyncCompilationAllowed(RHICmdList);
 
 	FComputePipelineState* OutCachedState = nullptr;
 
