Adding Transient resource information to resource dumps.

diff --git a/Engine/Source/Runtime/D3D12RHI/Private/D3D12Texture.h b/Engine/Source/Runtime/D3D12RHI/Private/D3D12Texture.h
--- a/Engine/Source/Runtime/D3D12RHI/Private/D3D12Texture.h
+++ b/Engine/Source/Runtime/D3D12RHI/Private/D3D12Texture.h
@@ -321,6 +321,7 @@ public:
 		OutResourceInfo.Name = this->GetName();
 		OutResourceInfo.Type = this->GetType();
 		OutResourceInfo.VRamAllocation.AllocationSize = GetMemorySize();
+		OutResourceInfo.IsTransient = this->ResourceLocation.IsTransient();
 		return true;
 	}
 #endif
@@ -463,6 +464,7 @@ public:
 		OutResourceInfo.Name = GetName();
 		OutResourceInfo.Type = GetType();
 		OutResourceInfo.VRamAllocation.AllocationSize = ResourceLocation.GetSize();
+		OutResourceInfo.IsTransient = ResourceLocation.IsTransient();
 		return true;
 	}
 #endif
diff --git a/Engine/Source/Runtime/D3D12RHI/Public/D3D12Resources.h b/Engine/Source/Runtime/D3D12RHI/Public/D3D12Resources.h
--- a/Engine/Source/Runtime/D3D12RHI/Public/D3D12Resources.h
+++ b/Engine/Source/Runtime/D3D12RHI/Public/D3D12Resources.h
@@ -969,6 +969,7 @@ public:
 		OutResourceInfo.Name = GetName();
 		OutResourceInfo.Type = GetType();
 		OutResourceInfo.VRamAllocation.AllocationSize = ResourceLocation.GetSize();
+		OutResourceInfo.IsTransient = this->ResourceLocation.IsTransient();
 		return true;
 	}
 #endif
diff --git a/Engine/Source/Runtime/RHI/Private/RHI.cpp b/Engine/Source/Runtime/RHI/Private/RHI.cpp
--- a/Engine/Source/Runtime/RHI/Private/RHI.cpp
+++ b/Engine/Source/Runtime/RHI/Private/RHI.cpp
@@ -742,6 +742,30 @@ static const TCHAR* StringFromRHIResourceType(ERHIResourceType ResourceType)
 	return TEXT("<unknown>");
 }
 
+enum class EBooleanFilter
+{
+	No,
+	Yes,
+	All
+};
+
+static EBooleanFilter ParseBooleanFilter(const FString& InText)
+{
+	if (InText.Equals(TEXT("No"), ESearchCase::IgnoreCase))
+	{
+		return EBooleanFilter::No;
+	}
+	if (InText.Equals(TEXT("Yes"), ESearchCase::IgnoreCase))
+	{
+		return EBooleanFilter::Yes;
+	}
+	if (InText.Equals(TEXT("All"), ESearchCase::IgnoreCase))
+	{
+		return EBooleanFilter::All;
+	}
+	return EBooleanFilter::No;
+}
+
 static FAutoConsoleCommandWithWorldArgsAndOutputDevice GDumpRHIResourceCountsCmd(
 	TEXT("rhi.DumpResourceCounts"),
 	TEXT("Dumps RHI resource counts to the log"),
@@ -793,11 +817,12 @@ static FAutoConsoleCommandWithWorldArgsAndOutputDevice GDumpRHIResourceCountsCmd
 static FAutoConsoleCommandWithWorldArgsAndOutputDevice GDumpRHIResourceMemoryCmd(
 	TEXT("rhi.DumpResourceMemory"),
 	TEXT("Dumps RHI resource memory stats to the log\n")
-	TEXT("Usage: rhi.DumpResourceMemory [<Number To Show>] [all] [Name=<Filter Text>] [Type=<RHI Resource Type>] [csv]"),
+	TEXT("Usage: rhi.DumpResourceMemory [<Number To Show>] [all] [Name=<Filter Text>] [Type=<RHI Resource Type>] [Transient=<no, yes, or all> [csv]"),
 	FConsoleCommandWithWorldArgsAndOutputDeviceDelegate::CreateStatic([](const TArray<FString>& Args, UWorld*, FOutputDevice& OutputDevice)
 {
 	FString NameFilter;
 	ERHIResourceType TypeFilter = RRT_None;
+	EBooleanFilter TransientFilter = EBooleanFilter::No;
 	int32 NumberOfResourcesToShow = 50;
 	bool bUseCSVOutput = false;
 	FArchive* CSVFile{ nullptr };
@@ -823,6 +848,10 @@ static FAutoConsoleCommandWithWorldArgsAndOutputDevice GDumpRHIResourceMemoryCmd
 		{
 			TypeFilter = RHIResourceTypeFromString(Argument.RightChop(5));
 		}
+		else if (Argument.StartsWith(TEXT("Transient="), ESearchCase::IgnoreCase))
+		{
+			TransientFilter = ParseBooleanFilter(Argument.RightChop(10));
+		}
 		else if (FCString::IsNumeric(*Argument))
 		{
 			LexFromString(NumberOfResourcesToShow, *Argument);
@@ -861,10 +890,17 @@ static FAutoConsoleCommandWithWorldArgsAndOutputDevice GDumpRHIResourceMemoryCmd
 				return false;
 			}
 		}
+		if (TransientFilter != EBooleanFilter::All)
+		{
+			const bool bAllowedFlag = TransientFilter == EBooleanFilter::Yes ? true : false;
+			if (ResourceInfo.IsTransient != bAllowedFlag)
+			{
+				return false;
+			}
+		}
 		return true;
 	};
 
-
 	struct FLocalResourceEntry
 	{
 		const FRHIResource* Resource;
@@ -877,6 +913,7 @@ static FAutoConsoleCommandWithWorldArgsAndOutputDevice GDumpRHIResourceMemoryCmd
 	int32 TotalResourcesWithInfo = 0;
 	int32 TotalTrackedResources = 0;
 	int64 TotalTrackedResourceSize = 0;
+	int64 TotalTrackedTransientResourceSize = 0;
 
 	{
 		FRHIResourceInfo ResourceInfo;
@@ -892,7 +929,14 @@ static FAutoConsoleCommandWithWorldArgsAndOutputDevice GDumpRHIResourceMemoryCmd
 				}
 
 				TotalResourcesWithInfo++;
-				TotalTrackedResourceSize += ResourceInfo.VRamAllocation.AllocationSize;
+				if (ResourceInfo.IsTransient)
+				{
+					TotalTrackedTransientResourceSize += ResourceInfo.VRamAllocation.AllocationSize;
+				}
+				else
+				{
+					TotalTrackedResourceSize += ResourceInfo.VRamAllocation.AllocationSize;
+				}
 			}
 		}
 	}
@@ -913,7 +957,7 @@ static FAutoConsoleCommandWithWorldArgsAndOutputDevice GDumpRHIResourceMemoryCmd
 
 	if (bUseCSVOutput)
 	{
-		const TCHAR* Header = TEXT("Name,Type,Size\n");
+		const TCHAR* Header = TEXT("Name,Type,Size,Transient\n");
 		CSVFile->Serialize(TCHAR_TO_ANSI(Header), FPlatformString::Strlen(Header));
 	}
 	else
@@ -940,28 +984,29 @@ static FAutoConsoleCommandWithWorldArgsAndOutputDevice GDumpRHIResourceMemoryCmd
 
 			if (bUseCSVOutput)
 			{
-				const FString Row = FString::Printf(TEXT("%s,%s,%.9f\n"),
+				const FString Row = FString::Printf(TEXT("%s,%s,%s,%.9f\n"),
 					ResourceNameBuffer,
 					ResourceType,
-					SizeInBytes / double(1 << 20));
+					ResourceInfo.IsTransient ? TEXT("Yes") : TEXT("No"),
+					SizeInBytes / double(1 << 20)
+				);
 
 				CSVFile->Serialize(TCHAR_TO_ANSI(*Row), Row.Len());
 			}
 			else
 			{
-				BufferedOutput.CategorizedLogf(CategoryName, ELogVerbosity::Log, TEXT("Name: %s - Type: %s - Size: %.9f MB"),
+				BufferedOutput.CategorizedLogf(CategoryName, ELogVerbosity::Log, TEXT("Name: %s - Type: %s - Transient: %s - Size: %.9f MB"),
 					ResourceNameBuffer,
 					ResourceType,
-					SizeInBytes / double(1 << 20));
+					ResourceInfo.IsTransient ? TEXT("Yes") : TEXT("No"),
+					SizeInBytes / double(1 << 20)
+				);
 			}
 
 			TotalShownResourceSize += SizeInBytes;
 		}
 	}
 
-	const double TotalSizeF = TotalTrackedResourceSize / double(1 << 20);
-	const double ShownSizeF = TotalShownResourceSize / double(1 << 20);
-
 	if (bUseCSVOutput)
 	{
 		delete CSVFile;
@@ -969,16 +1014,49 @@ static FAutoConsoleCommandWithWorldArgsAndOutputDevice GDumpRHIResourceMemoryCmd
 	}
 	else
 	{
+		const double TotalNonTransientSizeF = TotalTrackedResourceSize / double(1 << 20);
+		const double TotalTransientSizeF = TotalTrackedTransientResourceSize / double(1 << 20);
+		const double TotalSizeF = (TotalTrackedResourceSize + TotalTrackedTransientResourceSize) / double(1 << 20);
+		const double ShownSizeF = TotalShownResourceSize / double(1 << 20);
+
 		if (NumberOfResourcesToShow != TotalResourcesWithInfo)
 		{
-			BufferedOutput.CategorizedLogf(CategoryName, ELogVerbosity::Log, TEXT("Shown %d entries. Size: %.9f MB (%.2f%% of total)"),
-				NumberOfResourcesToShow, ShownSizeF, 100.0 * ShownSizeF / TotalSizeF);
+			double TotalSizeToUse = 0.0;
+			const TCHAR* ExtraText = TEXT("");
+
+			if (TransientFilter == EBooleanFilter::No)
+			{
+				TotalSizeToUse = TotalNonTransientSizeF;
+				ExtraText = TEXT(" non-transient");
+			}
+			else if (TransientFilter == EBooleanFilter::Yes)
+			{
+				TotalSizeToUse = TotalTransientSizeF;
+				ExtraText = TEXT(" transient");
+			}
+			else
+			{
+				TotalSizeToUse = TotalSizeF;
+				ExtraText = TEXT("");
+			}
+
+			if (TotalSizeToUse == 0.0)
+			{
+				TotalSizeToUse = TotalSizeF;
+			}
+
+			BufferedOutput.CategorizedLogf(CategoryName, ELogVerbosity::Log, TEXT("Shown %d entries. Size: %.9f MB (%.2f%% of total%s)"),
+				NumberOfResourcesToShow, ShownSizeF, 100.0 * ShownSizeF / TotalSizeToUse, ExtraText);
 		}
 
 		BufferedOutput.CategorizedLogf(CategoryName, ELogVerbosity::Log, TEXT("Total tracked resource size: %.9f MB"), TotalSizeF);
+		if (TotalTrackedTransientResourceSize > 0)
+		{
+			BufferedOutput.CategorizedLogf(CategoryName, ELogVerbosity::Log, TEXT("    Non-Transient: %.9f MB"), TotalNonTransientSizeF);
+			BufferedOutput.CategorizedLogf(CategoryName, ELogVerbosity::Log, TEXT("    Transient: %.9f MB"), TotalTransientSizeF);
+		}
+		BufferedOutput.RedirectTo(OutputDevice);
 	}
-
-	BufferedOutput.RedirectTo(OutputDevice);
 }));
 
 #endif // RHI_WANT_RESOURCE_INFO
diff --git a/Engine/Source/Runtime/RHI/Public/RHI.h b/Engine/Source/Runtime/RHI/Public/RHI.h
--- a/Engine/Source/Runtime/RHI/Public/RHI.h
+++ b/Engine/Source/Runtime/RHI/Public/RHI.h
@@ -1488,6 +1488,7 @@ struct FRHIResourceInfo
 	FName Name;
 	ERHIResourceType Type{ RRT_None };
 	FVRamAllocation VRamAllocation;
+	bool IsTransient{ false };
 };
 
 enum class EClearBinding
