Also dump Streaming, RT, UAV & RTAS information during RHI resource memory dump

diff --git a/Engine/Source/Runtime/RHI/Private/RHI.cpp b/Engine/Source/Runtime/RHI/Private/RHI.cpp
--- a/Engine/Source/Runtime/RHI/Private/RHI.cpp
+++ b/Engine/Source/Runtime/RHI/Private/RHI.cpp
@@ -957,7 +957,7 @@ static FAutoConsoleCommandWithWorldArgsAndOutputDevice GDumpRHIResourceMemoryCmd
 
 	if (bUseCSVOutput)
 	{
-		const TCHAR* Header = TEXT("Name,Type,Size,Transient\n");
+		const TCHAR* Header = TEXT("Name,Type,Size,Transient,Streaming,RenderTarget,UAV,\"Raytracing Acceleration Structure\"\n");
 		CSVFile->Serialize(TCHAR_TO_ANSI(Header), FPlatformString::Strlen(Header));
 	}
 	else
@@ -982,25 +982,87 @@ static FAutoConsoleCommandWithWorldArgsAndOutputDevice GDumpRHIResourceMemoryCmd
 			const TCHAR* ResourceType = StringFromRHIResourceType(ResourceInfo.Type);
 			const int64 SizeInBytes = ResourceInfo.VRamAllocation.AllocationSize;
 
+			bool bTransient = ResourceInfo.IsTransient;
+			bool bStreaming = false;
+			bool bRT = false;
+			bool bDS = false;
+			bool bUAV = false;
+			bool bRTAS = false;
+
+			bool bIsTexture = ResourceInfo.Type == RRT_Texture2D ||
+				ResourceInfo.Type == RRT_Texture2DArray ||
+				ResourceInfo.Type == RRT_Texture3D ||
+				ResourceInfo.Type == RRT_TextureCube;
+			if (bIsTexture)
+			{
+				FRHITexture* Texture = (FRHITexture*)Resources[Index].Resource;
+				bRT = EnumHasAnyFlags(Texture->GetFlags(), TexCreate_RenderTargetable);
+				bDS = EnumHasAnyFlags(Texture->GetFlags(), TexCreate_DepthStencilTargetable);
+				bUAV = EnumHasAnyFlags(Texture->GetFlags(), TexCreate_UAV);
+				bStreaming = EnumHasAnyFlags(Texture->GetFlags(), TexCreate_Streamable);
+			}
+			else if (ResourceInfo.Type == RRT_Buffer)
+			{
+				FRHIBuffer* Buffer = (FRHIBuffer*)Resources[Index].Resource;
+				bUAV = EnumHasAnyFlags((EBufferUsageFlags)Buffer->GetUsage(), BUF_UnorderedAccess);
+				bRTAS = EnumHasAnyFlags((EBufferUsageFlags)Buffer->GetUsage(), BUF_AccelerationStructure);
+			}
+
 			if (bUseCSVOutput)
 			{
-				const FString Row = FString::Printf(TEXT("%s,%s,%s,%.9f\n"),
+				const FString Row = FString::Printf(TEXT("%s,%s,%.9f,%s,%s,%s,%s,%s\n"),
 					ResourceNameBuffer,
 					ResourceType,
-					ResourceInfo.IsTransient ? TEXT("Yes") : TEXT("No"),
-					SizeInBytes / double(1 << 20)
-				);
+					SizeInBytes / double(1 << 20),
+					bTransient ? TEXT("Yes") : TEXT("No"),
+					bStreaming ? TEXT("Yes") : TEXT("No"),
+					(bRT || bDS) ? TEXT("Yes") : TEXT("No"),
+					bUAV ? TEXT("Yes") : TEXT("No"),
+					bRTAS ? TEXT("Yes") : TEXT("No"));
 
 				CSVFile->Serialize(TCHAR_TO_ANSI(*Row), Row.Len());
 			}
 			else
 			{
-				BufferedOutput.CategorizedLogf(CategoryName, ELogVerbosity::Log, TEXT("Name: %s - Type: %s - Transient: %s - Size: %.9f MB"),
+				FString ResoureFlags;
+				bool bHasFlag = false;
+				if (bTransient)
+				{
+					ResoureFlags += "Transient";
+					bHasFlag = true;
+				}
+				if (bStreaming)
+				{
+					ResoureFlags += bHasFlag ? "|Streaming" : "Streaming";
+					bHasFlag = true;
+				}
+				if (bRT)
+				{
+					ResoureFlags += bHasFlag ? "|RT" : "RT";
+					bHasFlag = true;
+				}
+				else if (bDS)
+				{
+					ResoureFlags += bHasFlag ? "|DS" : "DS";
+					bHasFlag = true;
+				}
+				if (bUAV)
+				{
+					ResoureFlags += bHasFlag ? "|UAV" : "UAV";
+					bHasFlag = true;
+				}
+				if (bRTAS)
+				{
+					ResoureFlags += bHasFlag ? "|RTAS" : "RTAS";
+					bHasFlag = true;
+
+				}
+
+				BufferedOutput.CategorizedLogf(CategoryName, ELogVerbosity::Log, TEXT("Name: %s - Type: %s - Size: %.9f MB - Flags: %s"),
 					ResourceNameBuffer,
 					ResourceType,
-					ResourceInfo.IsTransient ? TEXT("Yes") : TEXT("No"),
-					SizeInBytes / double(1 << 20)
-				);
+					SizeInBytes / double(1 << 20),
+					bHasFlag ? *ResoureFlags : TEXT("None"));
 			}
 
 			TotalShownResourceSize += SizeInBytes;
