Fix a bug of alpha to coverage on mobile

diff --git a/Engine/Source/Runtime/D3D12RHI/Private/D3D12State.cpp b/Engine/Source/Runtime/D3D12RHI/Private/D3D12State.cpp
--- a/Engine/Source/Runtime/D3D12RHI/Private/D3D12State.cpp
+++ b/Engine/Source/Runtime/D3D12RHI/Private/D3D12State.cpp
@@ -454,6 +454,7 @@ bool FD3D12BlendState::GetInitializer(class FBlendStateInitializerRHI& Init)
 			| ((Src.RenderTargetWriteMask & D3D12_COLOR_WRITE_ENABLE_ALPHA) ? CW_ALPHA : 0));
 	}
 	Init.bUseIndependentRenderTargetBlendStates = !!Desc.IndependentBlendEnable;
+	Init.bUseAlphaToCoverage = !!Desc.AlphaToCoverageEnable;
 	return true;
 }
 
diff --git a/Engine/Source/Runtime/RHI/Private/PipelineFileCache.cpp b/Engine/Source/Runtime/RHI/Private/PipelineFileCache.cpp
--- a/Engine/Source/Runtime/RHI/Private/PipelineFileCache.cpp
+++ b/Engine/Source/Runtime/RHI/Private/PipelineFileCache.cpp
@@ -882,6 +882,7 @@ bool FPipelineCacheFileFormatPSO::Verify() const
 				KeyHash = FCrc::MemCrc32(&Key.GraphicsDesc.StencilStore, sizeof(Key.GraphicsDesc.StencilStore), KeyHash);
 
 				KeyHash = FCrc::MemCrc32(&Key.GraphicsDesc.BlendState.bUseIndependentRenderTargetBlendStates, sizeof(Key.GraphicsDesc.BlendState.bUseIndependentRenderTargetBlendStates), KeyHash);
+				KeyHash = FCrc::MemCrc32(&Key.GraphicsDesc.BlendState.bUseAlphaToCoverage, sizeof(Key.GraphicsDesc.BlendState.bUseAlphaToCoverage), KeyHash);
 				for( uint32 i = 0; i < MaxSimultaneousRenderTargets; i++ )
 				{
 					KeyHash = FCrc::MemCrc32(&Key.GraphicsDesc.BlendState.RenderTargets[i].ColorBlendOp, sizeof(Key.GraphicsDesc.BlendState.RenderTargets[i].ColorBlendOp), KeyHash);
diff --git a/Engine/Source/Runtime/RHI/Private/RHI.cpp b/Engine/Source/Runtime/RHI/Private/RHI.cpp
--- a/Engine/Source/Runtime/RHI/Private/RHI.cpp
+++ b/Engine/Source/Runtime/RHI/Private/RHI.cpp
@@ -545,6 +545,7 @@ void FBlendStateInitializerRHI::FromString(const FStringView& InSrc)
 uint32 GetTypeHash(const FBlendStateInitializerRHI& Initializer)
 {
 	uint32 Hash = GetTypeHash(Initializer.bUseIndependentRenderTargetBlendStates);
+	Hash = HashCombine(Hash, Initializer.bUseAlphaToCoverage);
 	for (int32 i = 0; i < MaxSimultaneousRenderTargets; ++i)
 	{
 		Hash = HashCombine(Hash, GetTypeHash(Initializer.RenderTargets[i]));
@@ -556,6 +557,7 @@ uint32 GetTypeHash(const FBlendStateInitializerRHI& Initializer)
 bool operator== (const FBlendStateInitializerRHI& A, const FBlendStateInitializerRHI& B)
 {
 	bool bSame = A.bUseIndependentRenderTargetBlendStates == B.bUseIndependentRenderTargetBlendStates;
+	bSame = bSame && A.bUseAlphaToCoverage == B.bUseAlphaToCoverage;
 	for (int32 i = 0; i < MaxSimultaneousRenderTargets && bSame; ++i)
 	{
 		bSame = bSame && A.RenderTargets[i] == B.RenderTargets[i];
