Fixed PSO cache parsing to account for tesselation shaders removal and changing MaxVertexElementCount from 16 to 17

diff --git a/Engine/Source/Runtime/RHI/Private/PipelineFileCache.cpp b/Engine/Source/Runtime/RHI/Private/PipelineFileCache.cpp
--- a/Engine/Source/Runtime/RHI/Private/PipelineFileCache.cpp
+++ b/Engine/Source/Runtime/RHI/Private/PipelineFileCache.cpp
@@ -81,7 +81,7 @@ const uint64 FPipelineCacheFileFormatMagic = 0x5049504543414348; // PIPECACH
 const uint64 FPipelineCacheTOCFileFormatMagic = 0x544F435354415232; // TOCSTAR2
 const uint64 FPipelineCacheEOFFileFormatMagic = 0x454F462D4D41524B; // EOF-MARK
 const uint32 FPipelineCacheFileFormatCurrentVersion = (uint32)EPipelineCacheFileFormatVersions::MoreRenderTargetFlags;
-const int32  FPipelineCacheGraphicsDescPartsNum = 63; // parser will expect this number of parts in a description string
+const int32  FPipelineCacheGraphicsDescPartsNum = 64; // parser will expect this number of parts in a description string
 
 /**
   * PipelineFileCache API access
@@ -428,7 +428,7 @@ FString FPipelineCacheFileFormatPSO::GraphicsDescriptor::ShadersToString() const
 
 void FPipelineCacheFileFormatPSO::GraphicsDescriptor::ShadersFromString(const FStringView& Src)
 {
-	constexpr int32 PartCount = 7;
+	constexpr int32 PartCount = 5;
 
 	TArray<FStringView, TInlineAllocator<PartCount>> Parts;
 	UE::String::ParseTokens(Src.TrimStartAndEnd(), TEXT(','), [&Parts](FStringView Part) { Parts.Add(Part); });
