Replace Warning with Display for "Clearing the OS Cache" message.

#rnx

diff --git a/Engine/Source/Runtime/RHI/Private/PipelineFileCache.cpp b/Engine/Source/Runtime/RHI/Private/PipelineFileCache.cpp
--- a/Engine/Source/Runtime/RHI/Private/PipelineFileCache.cpp
+++ b/Engine/Source/Runtime/RHI/Private/PipelineFileCache.cpp
@@ -2848,7 +2848,7 @@ void FPipelineFileCache::PreCompileComplete()
 
 void FPipelineFileCache::ClearOSPipelineCache()
 {
-	UE_LOG(LogTemp, Warning, TEXT("Clearing the OS Cache"));
+	UE_LOG(LogTemp, Display, TEXT("Clearing the OS Cache"));
 	
 	bool bCmdLineSkip = FParse::Param(FCommandLine::Get(), TEXT("skippsoclear"));
 	if (CVarClearOSPSOFileCache.GetValueOnAnyThread() > 0 && !bCmdLineSkip)
