Shows all draw events by default with ProfileGPU

diff --git a/Engine/Source/Runtime/RHI/Private/GPUProfiler.cpp b/Engine/Source/Runtime/RHI/Private/GPUProfiler.cpp
--- a/Engine/Source/Runtime/RHI/Private/GPUProfiler.cpp
+++ b/Engine/Source/Runtime/RHI/Private/GPUProfiler.cpp
@@ -33,7 +33,7 @@ static TAutoConsoleVariable<FString> GProfileGPURootCVar(
 
 static TAutoConsoleVariable<float> GProfileThresholdPercent(
 	TEXT("r.ProfileGPU.ThresholdPercent"),
-	0.05f,
+	0.0f,
 	TEXT("Percent of the total execution duration the event needs to be larger than to be printed."),
 	ECVF_Default);
 
@@ -45,7 +45,7 @@ static TAutoConsoleVariable<int32> GProfileShowEventHistogram(
 
 static TAutoConsoleVariable<int32> GProfileGPUShowEvents(
 	TEXT("r.ProfileGPU.ShowLeafEvents"),
-	0,
+	1,
 	TEXT("Allows profileGPU to display event-only leaf nodes with no draws associated."),
 	ECVF_Default);
 
