Allow to start an editor in a game mobile preview with any supported RHI, eg: -game -featureleveles31 -vulkan

diff --git a/Engine/Source/Runtime/RHI/Private/Windows/WindowsDynamicRHI.cpp b/Engine/Source/Runtime/RHI/Private/Windows/WindowsDynamicRHI.cpp
--- a/Engine/Source/Runtime/RHI/Private/Windows/WindowsDynamicRHI.cpp
+++ b/Engine/Source/Runtime/RHI/Private/Windows/WindowsDynamicRHI.cpp
@@ -241,6 +241,7 @@ static TOptional<WindowsRHI> ChooseForcedRHI()
 		UE_LOG(LogRHI, Fatal, TEXT("-d3d12, -d3d11, -vulkan, and -opengl are mutually exclusive options, but more than one was specified on the command-line."));
 	}
 
+#if	!WITH_EDITOR
 	// FeatureLevelES31 is also a command line override, so it will determine the underlying RHI unless one is specified
 	if (FParse::Param(FCommandLine::Get(), TEXT("FeatureLevelES31")) || FParse::Param(FCommandLine::Get(), TEXT("FeatureLevelES3_1")))
 	{
@@ -267,6 +268,7 @@ static TOptional<WindowsRHI> ChooseForcedRHI()
 			ForcedRHI = WindowsRHI::D3D11;
 		}
 	}
+#endif// !WITH_EDITOR
 
 	return ForcedRHI;
 }
