Don't crash addref'ing empty bound states (UE-118955).

diff --git a/Engine/Source/Runtime/RHI/Private/PipelineStateCache.cpp b/Engine/Source/Runtime/RHI/Private/PipelineStateCache.cpp
--- a/Engine/Source/Runtime/RHI/Private/PipelineStateCache.cpp
+++ b/Engine/Source/Runtime/RHI/Private/PipelineStateCache.cpp
@@ -863,7 +863,46 @@ public:
 	{
 		if (!Pipeline->IsCompute())
 		{
-			Initializer.BoundShaderState.AddRefResources();
+#if PLATFORM_SUPPORTS_MESH_SHADERS
+			if (Initializer.BoundShaderState.MeshShaderRHI)
+			{
+				Initializer.BoundShaderState.MeshShaderRHI->AddRef();
+			}
+			if (Initializer.BoundShaderState.AmplificationShaderRHI)
+			{
+				Initializer.BoundShaderState.AmplificationShaderRHI->AddRef();
+			}
+#endif
+			if (Initializer.BoundShaderState.VertexDeclarationRHI)
+			{
+				Initializer.BoundShaderState.VertexDeclarationRHI->AddRef();
+			}
+			if (Initializer.BoundShaderState.VertexShaderRHI)
+			{
+				Initializer.BoundShaderState.VertexShaderRHI->AddRef();
+			}
+			if (Initializer.BoundShaderState.PixelShaderRHI)
+			{
+				Initializer.BoundShaderState.PixelShaderRHI->AddRef();
+			}
+#if PLATFORM_SUPPORTS_GEOMETRY_SHADERS
+			if (Initializer.BoundShaderState.GeometryShaderRHI)
+			{
+				Initializer.BoundShaderState.GeometryShaderRHI->AddRef();
+			}
+#endif
+			if (Initializer.BlendState)
+			{
+				Initializer.BlendState->AddRef();
+			}
+			if (Initializer.RasterizerState)
+			{
+				Initializer.RasterizerState->AddRef();
+			}
+			if (Initializer.DepthStencilState)
+			{
+				Initializer.DepthStencilState->AddRef();
+			}
 
 			if (Initializer.BlendState)
 			{
@@ -921,7 +960,46 @@ public:
 				HandlePipelineCreationFailure(Initializer);
 			}
 
-			Initializer.BoundShaderState.ReleaseResources();
+#if PLATFORM_SUPPORTS_MESH_SHADERS
+			if (Initializer.BoundShaderState.MeshShaderRHI)
+			{
+				Initializer.BoundShaderState.MeshShaderRHI->Release();
+			}
+			if (Initializer.BoundShaderState.AmplificationShaderRHI)
+			{
+				Initializer.BoundShaderState.AmplificationShaderRHI->Release();
+			}
+#endif
+			if (Initializer.BoundShaderState.VertexDeclarationRHI)
+			{
+				Initializer.BoundShaderState.VertexDeclarationRHI->Release();
+			}
+			if (Initializer.BoundShaderState.VertexShaderRHI)
+			{
+				Initializer.BoundShaderState.VertexShaderRHI->Release();
+			}
+			if (Initializer.BoundShaderState.PixelShaderRHI)
+			{
+				Initializer.BoundShaderState.PixelShaderRHI->Release();
+			}
+#if PLATFORM_SUPPORTS_GEOMETRY_SHADERS
+			if (Initializer.BoundShaderState.GeometryShaderRHI)
+			{
+				Initializer.BoundShaderState.GeometryShaderRHI->Release();
+			}
+#endif
+			if (Initializer.BlendState)
+			{
+				Initializer.BlendState->Release();
+			}
+			if (Initializer.RasterizerState)
+			{
+				Initializer.RasterizerState->Release();
+			}
+			if (Initializer.DepthStencilState)
+			{
+				Initializer.DepthStencilState->Release();
+			}
 
 			if (Initializer.BlendState)
 			{
