Dispatch to RHI thread after submitting an RDG command list.

diff --git a/Engine/Source/Runtime/RHI/Private/RHICommandList.cpp b/Engine/Source/Runtime/RHI/Private/RHICommandList.cpp
--- a/Engine/Source/Runtime/RHI/Private/RHICommandList.cpp
+++ b/Engine/Source/Runtime/RHI/Private/RHICommandList.cpp
@@ -1664,6 +1664,7 @@ FRHICOMMAND_MACRO(FRHICommandWaitForAndSubmitRTSubList)
 void FRHICommandListBase::QueueRenderThreadCommandListSubmit(FGraphEventRef& RenderThreadCompletionEvent, class FRHICommandList* CmdList)
 {
 	check(IsRunningRHIInSeparateThread() && IsInActualRenderingThread());
+
 	ALLOC_COMMAND(FRHICommandWaitForAndSubmitRTSubList)(RenderThreadCompletionEvent, CmdList);
 
 #if WITH_MGPU
@@ -1671,6 +1672,11 @@ void FRHICommandListBase::QueueRenderThreadCommandListSubmit(FGraphEventRef& Ren
 	// before the sub-list executed.
 	ALLOC_COMMAND(FRHICommandSetGPUMask)(GPUMask);
 #endif
+
+	if (IsRunningRHIInSeparateThread())
+	{
+		FRHICommandListExecutor::GetImmediateCommandList().ImmediateFlush(EImmediateFlushType::DispatchToRHIThread);
+	}
 }
 
 void FRHICommandListBase::AddDispatchPrerequisite(const FGraphEventRef& Prereq)
