Allow GeForceNowWrapper to not only be unused by the RHI module, but also unlinked and avoid its dependencies (on Slate and CoreUObject) when building standalone programs and servers.

diff --git a/Engine/Source/Runtime/RHI/Private/DynamicRHI.cpp b/Engine/Source/Runtime/RHI/Private/DynamicRHI.cpp
--- a/Engine/Source/Runtime/RHI/Private/DynamicRHI.cpp
+++ b/Engine/Source/Runtime/RHI/Private/DynamicRHI.cpp
@@ -15,7 +15,7 @@
 #include "GenericPlatform/GenericPlatformCrashContext.h"
 #include "PipelineStateCache.h"
 
-#if NV_GEFORCENOW
+#if defined(NV_GEFORCENOW) && NV_GEFORCENOW
 #include "GeForceNOWWrapper.h"
 #endif
 
@@ -284,7 +284,7 @@ void RHIInit(bool bHasEditorToken)
 	}
 
 #if PLATFORM_WINDOWS || PLATFORM_MAC
-#if NV_GEFORCENOW
+#if defined(NV_GEFORCENOW) && NV_GEFORCENOW
 	bool bDetectAndWarnBadDrivers = true;
 	if (IsRHIDeviceNVIDIA() && !!CVarDisableDriverWarningPopupIfGFN.GetValueOnAnyThread())
 	{
diff --git a/Engine/Source/Runtime/RHI/Private/Windows/WindowsDynamicRHI.cpp b/Engine/Source/Runtime/RHI/Private/Windows/WindowsDynamicRHI.cpp
--- a/Engine/Source/Runtime/RHI/Private/Windows/WindowsDynamicRHI.cpp
+++ b/Engine/Source/Runtime/RHI/Private/Windows/WindowsDynamicRHI.cpp
@@ -10,7 +10,7 @@
 
 #include "Windows/WindowsPlatformApplicationMisc.h"
 
-#if NV_GEFORCENOW
+#if defined(NV_GEFORCENOW) && NV_GEFORCENOW
 #include "GeForceNOWWrapper.h"
 #endif
 
@@ -90,7 +90,7 @@ static bool PreferFeatureLevelES31()
 	if (!GIsEditor)
 	{
 		bool bIsRunningInGFN = false;
-#if NV_GEFORCENOW
+#if defined(NV_GEFORCENOW) && NV_GEFORCENOW
 		//Prevent ES31 from being forced since we have other ways of setting scalability issues on GFN.
 		GeForceNOWWrapper::Get().Initialize();
 		bIsRunningInGFN = GeForceNOWWrapper::Get().IsRunningInGFN();
diff --git a/Engine/Source/Runtime/RHI/RHI.Build.cs b/Engine/Source/Runtime/RHI/RHI.Build.cs
--- a/Engine/Source/Runtime/RHI/RHI.Build.cs
+++ b/Engine/Source/Runtime/RHI/RHI.Build.cs
@@ -10,7 +10,15 @@ public class RHI : ModuleRules
 		PrivateDependencyModuleNames.Add("Core");
 		PrivateDependencyModuleNames.Add("TraceLog");
 		PrivateDependencyModuleNames.Add("ApplicationCore");
-		PrivateDependencyModuleNames.Add("GeForceNOWWrapper");
+
+		if (Target.Type != TargetRules.TargetType.Server
+		&& Target.Type != TargetRules.TargetType.Program
+		&& Target.Configuration != UnrealTargetConfiguration.Unknown
+		&& Target.Configuration != UnrealTargetConfiguration.Debug
+		&& Target.Platform == UnrealTargetPlatform.Win64)
+		{
+			PrivateDependencyModuleNames.Add("GeForceNOWWrapper");
+		}
 
 		if (Target.bCompileAgainstEngine)
 		{
