Fix for GpuTransformsSrv getting destroyed while in use

diff --git a/Engine/Source/Runtime/RHI/Public/DynamicRHI.h b/Engine/Source/Runtime/RHI/Public/DynamicRHI.h
--- a/Engine/Source/Runtime/RHI/Public/DynamicRHI.h
+++ b/Engine/Source/Runtime/RHI/Public/DynamicRHI.h
@@ -103,7 +103,7 @@ struct FRayTracingGeometryInstance
 	TArrayView<const FMatrix> Transforms;
 
 	// Optional buffer that stores GPU transforms. Used instead of CPU-side transform data.
-	FRHIShaderResourceView* GPUTransformsSRV = nullptr;
+	FShaderResourceViewRHIRef GPUTransformsSRV = nullptr;
 
 	// Conservative number of instances. Some of the actual instances may be made inactive if GPU transforms are used.
 	// Must be less or equal to number of entries in Transforms view if CPU transform data is used.
