Fix RHI command list render thread submission threading check.

diff --git a/Engine/Source/Runtime/RHI/Private/RHICommandList.cpp b/Engine/Source/Runtime/RHI/Private/RHICommandList.cpp
--- a/Engine/Source/Runtime/RHI/Private/RHICommandList.cpp
+++ b/Engine/Source/Runtime/RHI/Private/RHICommandList.cpp
@@ -1663,7 +1663,7 @@ FRHICOMMAND_MACRO(FRHICommandWaitForAndSubmitRTSubList)
 
 void FRHICommandListBase::QueueRenderThreadCommandListSubmit(FGraphEventRef& RenderThreadCompletionEvent, class FRHICommandList* CmdList)
 {
-	check(IsRunningRHIInSeparateThread() && IsInActualRenderingThread());
+	check(IsInRenderingThread() && IsImmediate());
 
 	ALLOC_COMMAND(FRHICommandWaitForAndSubmitRTSubList)(RenderThreadCompletionEvent, CmdList);
 
