Fix transient RHI validation to use correct array size for cubemaps.

diff --git a/Engine/Source/Runtime/RHI/Private/RHIValidation.cpp b/Engine/Source/Runtime/RHI/Private/RHIValidation.cpp
--- a/Engine/Source/Runtime/RHI/Private/RHIValidation.cpp
+++ b/Engine/Source/Runtime/RHI/Private/RHIValidation.cpp
@@ -1688,7 +1688,15 @@ void FValidationTransientResourceAllocator::InitBarrierTracking(const FAllocated
 		case FAllocatedResourceData::EType::Texture:
 		{
 			FRHITexture* Texture = static_cast<FRHITexture*>(Resource);
-			Texture->InitBarrierTracking(ResourceData.Texture.NumMips, ResourceData.Texture.ArraySize, ResourceData.Texture.Format, ResourceData.Texture.Flags, ERHIAccess::Discard, ResourceData.DebugName);
+
+			int32 ArraySize = ResourceData.Texture.ArraySize;
+
+			if (Texture->GetTextureCube() != nullptr)
+			{
+				ArraySize *= 6;
+			}
+
+			Texture->InitBarrierTracking(ResourceData.Texture.NumMips, ArraySize, ResourceData.Texture.Format, ResourceData.Texture.Flags, ERHIAccess::Discard, ResourceData.DebugName);
 		}
 		break;
 		case FAllocatedResourceData::EType::Buffer:
