Fixed additional RHI validation cubemap error.

diff --git a/Engine/Source/Runtime/RHI/Private/RHIValidation.cpp b/Engine/Source/Runtime/RHI/Private/RHIValidation.cpp
--- a/Engine/Source/Runtime/RHI/Private/RHIValidation.cpp
+++ b/Engine/Source/Runtime/RHI/Private/RHIValidation.cpp
@@ -1538,7 +1538,7 @@ FRHITransientTexture* FValidationTransientResourceAllocator::CreateTexture(const
 	{
 		if (!Resource->IsBarrierTrackingInitialized())
 		{
-			RHITexture->InitBarrierTracking(InCreateInfo.NumMips, InCreateInfo.ArraySize, InCreateInfo.Format, InCreateInfo.Flags, ERHIAccess::Discard, InDebugName);
+			RHITexture->InitBarrierTracking(InCreateInfo.NumMips, InCreateInfo.ArraySize * (InCreateInfo.IsTextureCube() ? 6 : 1), InCreateInfo.Format, InCreateInfo.Flags, ERHIAccess::Discard, InDebugName);
 		}
 		else
 		{
