Adding `summary` argument to rhi.dumpresourcememory.  This skips listing each individual resource and provides a simplified output that just shows total memory sizes.

- This is useful if you want to understand how much memory a rendering system consumes.  For example running `rhi.dumpresourcememory summary name=Lumen` will output just the following:

    Shown 62 entries with name Lumen. Size: 582.78/6373.53 MB (9.14% of total non-transient)
- Show the TotalSizeToUse in the output.
- If there is a name filter display that in the output.

diff --git a/Engine/Source/Runtime/Engine/Private/UnrealEngine.cpp b/Engine/Source/Runtime/Engine/Private/UnrealEngine.cpp
--- a/Engine/Source/Runtime/Engine/Private/UnrealEngine.cpp
+++ b/Engine/Source/Runtime/Engine/Private/UnrealEngine.cpp
@@ -6619,6 +6619,7 @@ bool UEngine::HandleMemReportDeferredCommand( const TCHAR* Cmd, FOutputDevice& A
 	FSlowHeartBeatScope HeartBeatSuspend;
 
 	const bool bPerformSlowCommands = FParse::Param( Cmd, TEXT("FULL") );
+	const bool bPerformGPUCommands = FParse::Param( Cmd, TEXT("GPU") );
 	const bool bLogOutputToFile = !FParse::Param( Cmd, TEXT("LOG") );
 	const bool bCSV = FParse::Param(Cmd, TEXT("CSV"));
 	FString InFileName;
diff --git a/Engine/Source/Runtime/RHI/Private/RHI.cpp b/Engine/Source/Runtime/RHI/Private/RHI.cpp
--- a/Engine/Source/Runtime/RHI/Private/RHI.cpp
+++ b/Engine/Source/Runtime/RHI/Private/RHI.cpp
@@ -819,7 +819,7 @@ static FAutoConsoleCommandWithWorldArgsAndOutputDevice GDumpRHIResourceCountsCmd
 static FAutoConsoleCommandWithWorldArgsAndOutputDevice GDumpRHIResourceMemoryCmd(
 	TEXT("rhi.DumpResourceMemory"),
 	TEXT("Dumps RHI resource memory stats to the log\n")
-	TEXT("Usage: rhi.DumpResourceMemory [<Number To Show>] [all] [Name=<Filter Text>] [Type=<RHI Resource Type>] [Transient=<no, yes, or all> [csv]"),
+	TEXT("Usage: rhi.DumpResourceMemory [<Number To Show>] [all] [summary] [Name=<Filter Text>] [Type=<RHI Resource Type>] [Transient=<no, yes, or all> [csv]"),
 	FConsoleCommandWithWorldArgsAndOutputDeviceDelegate::CreateStatic([](const TArray<FString>& Args, UWorld*, FOutputDevice& OutputDevice)
 {
 	FString NameFilter;
@@ -827,6 +827,7 @@ static FAutoConsoleCommandWithWorldArgsAndOutputDevice GDumpRHIResourceMemoryCmd
 	EBooleanFilter TransientFilter = EBooleanFilter::No;
 	int32 NumberOfResourcesToShow = 50;
 	bool bUseCSVOutput = false;
+	bool bSummaryOutput = false;
 	FArchive* CSVFile{ nullptr };
 
 	for (const FString& Argument : Args)
@@ -858,6 +859,13 @@ static FAutoConsoleCommandWithWorldArgsAndOutputDevice GDumpRHIResourceMemoryCmd
 		{
 			LexFromString(NumberOfResourcesToShow, *Argument);
 		}
+		else if (Argument.Equals(TEXT("summary"), ESearchCase::IgnoreCase))
+		{
+			// Respects name, type and transient filters but only reports total sizes.
+			// Does not report a list of individual resources.
+			bSummaryOutput = true;
+			NumberOfResourcesToShow = -1;
+		}
 		else
 		{
 			NameFilter = Argument;
@@ -964,7 +972,10 @@ static FAutoConsoleCommandWithWorldArgsAndOutputDevice GDumpRHIResourceMemoryCmd
 	}
 	else
 	{
-		BufferedOutput.CategorizedLogf(CategoryName, ELogVerbosity::Log, TEXT("Tracked RHIResources (%d total with info, %d total tracked)"), TotalResourcesWithInfo, TotalTrackedResources);
+		if (bSummaryOutput == false)
+		{
+			BufferedOutput.CategorizedLogf(CategoryName, ELogVerbosity::Log, TEXT("Tracked RHIResources (%d total with info, %d total tracked)"), TotalResourcesWithInfo, TotalTrackedResources);
+		}
 
 		if (NumberOfResourcesToShow != NumberOfResourcesBeforeNumberFilter)
 		{
@@ -1010,61 +1021,64 @@ static FAutoConsoleCommandWithWorldArgsAndOutputDevice GDumpRHIResourceMemoryCmd
 				bRTAS = EnumHasAnyFlags((EBufferUsageFlags)Buffer->GetUsage(), BUF_AccelerationStructure);
 			}
 
-			if (bUseCSVOutput)
+			if (bSummaryOutput == false)
 			{
-				const FString Row = FString::Printf(TEXT("%s,%s,%.9f,%s,%s,%s,%s,%s\n"),
-					ResourceNameBuffer,
-					ResourceType,
-					SizeInBytes / double(1 << 20),
-					bTransient ? TEXT("Yes") : TEXT("No"),
-					bStreaming ? TEXT("Yes") : TEXT("No"),
-					(bRT || bDS) ? TEXT("Yes") : TEXT("No"),
-					bUAV ? TEXT("Yes") : TEXT("No"),
-					bRTAS ? TEXT("Yes") : TEXT("No"));
-
-				CSVFile->Serialize(TCHAR_TO_ANSI(*Row), Row.Len());
-			}
-			else
-			{
-				FString ResoureFlags;
-				bool bHasFlag = false;
-				if (bTransient)
-				{
-					ResoureFlags += "Transient";
-					bHasFlag = true;
-				}
-				if (bStreaming)
-				{
-					ResoureFlags += bHasFlag ? "|Streaming" : "Streaming";
-					bHasFlag = true;
-				}
-				if (bRT)
-				{
-					ResoureFlags += bHasFlag ? "|RT" : "RT";
-					bHasFlag = true;
-				}
-				else if (bDS)
+				if (bUseCSVOutput)
 				{
-					ResoureFlags += bHasFlag ? "|DS" : "DS";
-					bHasFlag = true;
+					const FString Row = FString::Printf(TEXT("%s,%s,%.9f,%s,%s,%s,%s,%s\n"),
+						ResourceNameBuffer,
+						ResourceType,
+						SizeInBytes / double(1 << 20),
+						bTransient ? TEXT("Yes") : TEXT("No"),
+						bStreaming ? TEXT("Yes") : TEXT("No"),
+						(bRT || bDS) ? TEXT("Yes") : TEXT("No"),
+						bUAV ? TEXT("Yes") : TEXT("No"),
+						bRTAS ? TEXT("Yes") : TEXT("No"));
+
+					CSVFile->Serialize(TCHAR_TO_ANSI(*Row), Row.Len());
 				}
-				if (bUAV)
+				else
 				{
-					ResoureFlags += bHasFlag ? "|UAV" : "UAV";
-					bHasFlag = true;
+					FString ResoureFlags;
+					bool bHasFlag = false;
+					if (bTransient)
+					{
+						ResoureFlags += "Transient";
+						bHasFlag = true;
+					}
+					if (bStreaming)
+					{
+						ResoureFlags += bHasFlag ? "|Streaming" : "Streaming";
+						bHasFlag = true;
+					}
+					if (bRT)
+					{
+						ResoureFlags += bHasFlag ? "|RT" : "RT";
+						bHasFlag = true;
+					}
+					else if (bDS)
+					{
+						ResoureFlags += bHasFlag ? "|DS" : "DS";
+						bHasFlag = true;
+					}
+					if (bUAV)
+					{
+						ResoureFlags += bHasFlag ? "|UAV" : "UAV";
+						bHasFlag = true;
+					}
+					if (bRTAS)
+					{
+						ResoureFlags += bHasFlag ? "|RTAS" : "RTAS";
+						bHasFlag = true;
+
+					}
+
+					BufferedOutput.CategorizedLogf(CategoryName, ELogVerbosity::Log, TEXT("Name: %s - Type: %s - Size: %.9f MB - Flags: %s"),
+						ResourceNameBuffer,
+						ResourceType,
+						SizeInBytes / double(1 << 20),
+						bHasFlag ? *ResoureFlags : TEXT("None"));
 				}
-				if (bRTAS)
-				{
-					ResoureFlags += bHasFlag ? "|RTAS" : "RTAS";
-					bHasFlag = true;
-
-				}
-
-				BufferedOutput.CategorizedLogf(CategoryName, ELogVerbosity::Log, TEXT("Name: %s - Type: %s - Size: %.9f MB - Flags: %s"),
-					ResourceNameBuffer,
-					ResourceType,
-					SizeInBytes / double(1 << 20),
-					bHasFlag ? *ResoureFlags : TEXT("None"));
 			}
 
 			TotalShownResourceSize += SizeInBytes;
@@ -1109,16 +1123,20 @@ static FAutoConsoleCommandWithWorldArgsAndOutputDevice GDumpRHIResourceMemoryCmd
 				TotalSizeToUse = TotalSizeF;
 			}
 
-			BufferedOutput.CategorizedLogf(CategoryName, ELogVerbosity::Log, TEXT("Shown %d entries. Size: %.9f MB (%.2f%% of total%s)"),
-				NumberOfResourcesToShow, ShownSizeF, 100.0 * ShownSizeF / TotalSizeToUse, ExtraText);
+			BufferedOutput.CategorizedLogf(CategoryName, ELogVerbosity::Log, TEXT("Shown %d entries%s. Size: %.2f/%.2f MB (%.2f%% of total%s)"),
+				NumberOfResourcesToShow, !NameFilter.IsEmpty() ? *FString::Printf(TEXT(" with name %s"), *NameFilter) : TEXT(""), ShownSizeF, TotalSizeToUse, 100.0 * ShownSizeF / TotalSizeToUse, ExtraText);
 		}
 
-		BufferedOutput.CategorizedLogf(CategoryName, ELogVerbosity::Log, TEXT("Total tracked resource size: %.9f MB"), TotalSizeF);
-		if (TotalTrackedTransientResourceSize > 0)
+		if (bSummaryOutput == false)
 		{
-			BufferedOutput.CategorizedLogf(CategoryName, ELogVerbosity::Log, TEXT("    Non-Transient: %.9f MB"), TotalNonTransientSizeF);
-			BufferedOutput.CategorizedLogf(CategoryName, ELogVerbosity::Log, TEXT("    Transient: %.9f MB"), TotalTransientSizeF);
+			BufferedOutput.CategorizedLogf(CategoryName, ELogVerbosity::Log, TEXT("Total tracked resource size: %.9f MB"), TotalSizeF);
+			if (TotalTrackedTransientResourceSize > 0)
+			{
+				BufferedOutput.CategorizedLogf(CategoryName, ELogVerbosity::Log, TEXT("    Non-Transient: %.9f MB"), TotalNonTransientSizeF);
+				BufferedOutput.CategorizedLogf(CategoryName, ELogVerbosity::Log, TEXT("    Transient: %.9f MB"), TotalTransientSizeF);
+			}
 		}
+
 		BufferedOutput.RedirectTo(OutputDevice);
 	}
 }));
