Fix D3D12 preference in Windows RHI settings (UE-118810).

- It should not override user's explicit choice of DefaultGraphicsRHI in project settings in the editor.
- Either in the editor or in game, if the project isn't using D3D at all, it shouldn't apply.

diff --git a/Engine/Source/Runtime/RHI/Private/Windows/WindowsDynamicRHI.cpp b/Engine/Source/Runtime/RHI/Private/Windows/WindowsDynamicRHI.cpp
--- a/Engine/Source/Runtime/RHI/Private/Windows/WindowsDynamicRHI.cpp
+++ b/Engine/Source/Runtime/RHI/Private/Windows/WindowsDynamicRHI.cpp
@@ -16,20 +16,6 @@
 
 static const TCHAR* GLoadedRHIModuleName;
 
-static bool PreferD3D12()
-{
-	if (!GIsEditor)
-	{
-		bool bPreferD3D12 = false;
-		if (GConfig->GetBool(TEXT("D3DRHIPreference"), TEXT("bUseD3D12InGame"), bPreferD3D12, GGameUserSettingsIni))
-		{
-			return bPreferD3D12;
-		}
-	}
-	
-	return true;
-}
-
 // Default to Performance Mode on low-end machines
 static bool DefaultFeatureLevelES31()
 {
@@ -143,32 +129,9 @@ namespace
 // Choose the default from DefaultGraphicsRHI or TargetedRHIs. DefaultGraphicsRHI has precedence.
 static WindowsRHI ChooseDefaultRHI()
 {
-	WindowsRHI DefaultRHI = WindowsRHI::D3D11;
-
-	// Check the list of targeted shader platforms and decide an RHI based off them
-	TArray<FString> TargetedShaderFormats;
-	GConfig->GetArray(TEXT("/Script/WindowsTargetPlatform.WindowsTargetSettings"), TEXT("TargetedRHIs"), TargetedShaderFormats, GEngineIni);
-	if (TargetedShaderFormats.Num() > 0)
-	{
-		// Pick the first one
-		FName ShaderFormatName(*TargetedShaderFormats[0]);
-		EShaderPlatform TargetedPlatform = ShaderFormatToLegacyShaderPlatform(ShaderFormatName);
-
-		if (IsD3DPlatform(TargetedPlatform))
-		{
-			DefaultRHI = WindowsRHI::D3D11;
-		}
-		else if (IsVulkanPlatform(TargetedPlatform))
-		{
-			DefaultRHI = WindowsRHI::Vulkan;
-		}
-		else if (IsOpenGLPlatform(TargetedPlatform))
-		{
-			DefaultRHI = WindowsRHI::OpenGL;
-		}
-	}
+	WindowsRHI DefaultRHI = WindowsRHI::D3D12;
 
-	//Default graphics RHI is only used if no command line option is specified
+	// Default graphics RHI is the main project setting that governs the choice, so it takes the priority
 	FConfigFile EngineSettings;
 	FString PlatformNameString = FPlatformProperties::PlatformName();
 	const TCHAR* PlatformName = *PlatformNameString;
@@ -191,21 +154,45 @@ static WindowsRHI ChooseDefaultRHI()
 		{
 			DefaultRHI = WindowsRHI::Vulkan;
 		}
+		else if (DefaultGraphicsRHI != TEXT("DefaultGraphicsRHI_Default"))
+		{
+			UE_LOG(LogRHI, Error, TEXT("Unrecognized setting '%s' for DefaultGraphicsRHI"), *DefaultGraphicsRHI);
+		}
 	}
+	else 
+	{
+		// If we don't have DefaultGraphicsRHI set, try to deduce it from the list of targeted shader platforms
+		TArray<FString> TargetedShaderFormats;
+		GConfig->GetArray(TEXT("/Script/WindowsTargetPlatform.WindowsTargetSettings"), TEXT("TargetedRHIs"), TargetedShaderFormats, GEngineIni);
+		if (TargetedShaderFormats.Num() > 0)
+		{
+			// Pick the first one
+			FName ShaderFormatName(*TargetedShaderFormats[0]);
+			EShaderPlatform TargetedPlatform = ShaderFormatToLegacyShaderPlatform(ShaderFormatName);
 
-	return DefaultRHI;
-}
-
-static TOptional<WindowsRHI> ChoosePreferredRHI()
-{
-	TOptional<WindowsRHI> PreferredRHI = {};
+			// not checking D3D as DefaultRHI begins initialized as D3D12
+			if (IsVulkanPlatform(TargetedPlatform))
+			{
+				DefaultRHI = WindowsRHI::Vulkan;
+			}
+			else if (IsOpenGLPlatform(TargetedPlatform))
+			{
+				DefaultRHI = WindowsRHI::OpenGL;
+			}
+		}
+	}
 
-	if (PreferD3D12())
+	// If we are in game, there is a separate setting that can make it prefer D3D12 over D3D11 (but not over other RHIs).
+	if (!GIsEditor && (DefaultRHI == WindowsRHI::D3D11 || DefaultRHI == WindowsRHI::D3D12))
 	{
-		PreferredRHI = WindowsRHI::D3D12;
+		bool bUseD3D12InGame = false;
+		if (GConfig->GetBool(TEXT("D3DRHIPreference"), TEXT("bUseD3D12InGame"), bUseD3D12InGame, GGameUserSettingsIni))
+		{
+			DefaultRHI = bUseD3D12InGame ? WindowsRHI::D3D12 : WindowsRHI::D3D11;
+		}
 	}
 
-	return PreferredRHI;
+	return DefaultRHI;
 }
 
 static TOptional<WindowsRHI> ChooseForcedRHI()
@@ -312,20 +299,14 @@ static IDynamicRHIModule* LoadDynamicRHIModule(ERHIFeatureLevel::Type& DesiredFe
 		*GPUCrashDebuggingCVar = bUseGPUCrashDebugging;
 	}
 
-	// RHI selection priority: blacklist -> force -> prefer -> default
-	// The IsSupported function blacklists an RHI, causing the default to be selected or none if the default is blacklisted.
-	// Command line parameters force the RHI over preferences chosen in settings.
-	// Finally, the default is chosen based on machine capabilities.
+	// RHI is chosen by the project settings (first DefaultGraphicsRHI, then TargetedRHIs are consulted, "Default" maps to D3D12). 
+	// After this, a separate game-only setting (does not affect editor) bPreferD3D12InGame selects between D3D12 or D3D11 (but will not have any effect if Vulkan or OpenGL are chosen).
+	// Commandline switches apply after this and can force an arbitrary RHIs. If RHI isn't supported, the game will refuse to start.
 
 	WindowsRHI DefaultRHI = ChooseDefaultRHI();
-	TOptional<WindowsRHI> PreferredRHI = ChoosePreferredRHI();
 	TOptional<WindowsRHI> ForcedRHI = ChooseForcedRHI();
 
 	WindowsRHI ChosenRHI = DefaultRHI;
-	if (PreferredRHI)
-	{
-		ChosenRHI = PreferredRHI.GetValue();
-	}
 	if (ForcedRHI)
 	{
 		ChosenRHI = ForcedRHI.GetValue();
