Add STAT_RTAccelerationStructureMemory to track buffer sizes.

- Fixes ensure message that was triggering.

https://horde.devtools.epicgames.com/job/60fb063cd3905c0001fac5d3

diff --git a/Engine/Source/Runtime/D3D12RHI/Public/D3D12Resources.h b/Engine/Source/Runtime/D3D12RHI/Public/D3D12Resources.h
--- a/Engine/Source/Runtime/D3D12RHI/Public/D3D12Resources.h
+++ b/Engine/Source/Runtime/D3D12RHI/Public/D3D12Resources.h
@@ -903,6 +903,10 @@ inline FName GetBufferStats(uint32 Usage)
 	{
 		return GET_STATFNAME(STAT_IndexBufferMemory);
 	}
+	else if (Usage & BUF_AccelerationStructure)
+	{
+		return GET_STATFNAME(STAT_RTAccelerationStructureMemory);
+	}
 	else
 	{
 		return GET_STATFNAME(STAT_StructuredBufferMemory);
diff --git a/Engine/Source/Runtime/RHI/Private/RHI.cpp b/Engine/Source/Runtime/RHI/Private/RHI.cpp
--- a/Engine/Source/Runtime/RHI/Private/RHI.cpp
+++ b/Engine/Source/Runtime/RHI/Private/RHI.cpp
@@ -48,6 +48,7 @@ DEFINE_STAT(STAT_TextureMemoryCube);
 DEFINE_STAT(STAT_UniformBufferMemory);
 DEFINE_STAT(STAT_IndexBufferMemory);
 DEFINE_STAT(STAT_VertexBufferMemory);
+DEFINE_STAT(STAT_RTAccelerationStructureMemory);
 DEFINE_STAT(STAT_StructuredBufferMemory);
 DEFINE_STAT(STAT_PixelBufferMemory);
 
diff --git a/Engine/Source/Runtime/RHI/Public/RHI.h b/Engine/Source/Runtime/RHI/Public/RHI.h
--- a/Engine/Source/Runtime/RHI/Public/RHI.h
+++ b/Engine/Source/Runtime/RHI/Public/RHI.h
@@ -2510,6 +2510,7 @@ DECLARE_MEMORY_STAT_POOL_EXTERN(TEXT("Texture memory Cube"),STAT_TextureMemoryCu
 DECLARE_MEMORY_STAT_POOL_EXTERN(TEXT("Uniform buffer memory"),STAT_UniformBufferMemory,STATGROUP_RHI,FPlatformMemory::MCR_GPU,RHI_API);
 DECLARE_MEMORY_STAT_POOL_EXTERN(TEXT("Index buffer memory"),STAT_IndexBufferMemory,STATGROUP_RHI,FPlatformMemory::MCR_GPU,RHI_API);
 DECLARE_MEMORY_STAT_POOL_EXTERN(TEXT("Vertex buffer memory"),STAT_VertexBufferMemory,STATGROUP_RHI,FPlatformMemory::MCR_GPU,RHI_API);
+DECLARE_MEMORY_STAT_POOL_EXTERN(TEXT("Ray Tracing Acceleration Structure memory"), STAT_RTAccelerationStructureMemory, STATGROUP_RHI, FPlatformMemory::MCR_GPU, RHI_API);
 DECLARE_MEMORY_STAT_POOL_EXTERN(TEXT("Structured buffer memory"),STAT_StructuredBufferMemory,STATGROUP_RHI,FPlatformMemory::MCR_GPU,RHI_API);
 DECLARE_MEMORY_STAT_POOL_EXTERN(TEXT("Pixel buffer memory"),STAT_PixelBufferMemory,STATGROUP_RHI,FPlatformMemory::MCR_GPU,RHI_API);
 
