Fixed assert due to parallel render thread access to RHI query pool.

diff --git a/Engine/Source/Runtime/RHI/Public/RHIResources.h b/Engine/Source/Runtime/RHI/Public/RHIResources.h
--- a/Engine/Source/Runtime/RHI/Public/RHIResources.h
+++ b/Engine/Source/Runtime/RHI/Public/RHIResources.h
@@ -1657,7 +1657,7 @@ inline void FRHIPooledRenderQuery::ReleaseQuery()
 
 inline FRHIPooledRenderQuery::~FRHIPooledRenderQuery()
 {
-	check(IsInRenderingThread());
+	check(IsInParallelRenderingThread());
 	ReleaseQuery();
 }
 
