Fix for RHI's which do not support async compute.

diff --git a/Engine/Source/Runtime/RHI/Private/RHICommandList.cpp b/Engine/Source/Runtime/RHI/Private/RHICommandList.cpp
--- a/Engine/Source/Runtime/RHI/Private/RHICommandList.cpp
+++ b/Engine/Source/Runtime/RHI/Private/RHICommandList.cpp
@@ -1105,7 +1105,7 @@ void FRHICommandListExecutor::Transition(TArrayView<const FRHITransitionInfo> In
 
 	FRHIAsyncComputeCommandListImmediate& RHICmdListAsyncCompute = GetImmediateAsyncComputeCommandList();
 
-	if (!RHICmdListAsyncCompute.IsAsyncCompute())
+	if (!GSupportsEfficientAsyncCompute)
 	{
 		checkf(SrcPipelines != ERHIPipeline::AsyncCompute, TEXT("Async compute is disabled. Cannot transition from it."));
 		checkf(DstPipelines != ERHIPipeline::AsyncCompute, TEXT("Async compute is disabled. Cannot transition to it."));
