Make D3D11 to be the default again for projects without an RHI preference.

diff --git a/Engine/Source/Runtime/RHI/Private/Windows/WindowsDynamicRHI.cpp b/Engine/Source/Runtime/RHI/Private/Windows/WindowsDynamicRHI.cpp
--- a/Engine/Source/Runtime/RHI/Private/Windows/WindowsDynamicRHI.cpp
+++ b/Engine/Source/Runtime/RHI/Private/Windows/WindowsDynamicRHI.cpp
@@ -129,7 +129,7 @@ namespace
 // Choose the default from DefaultGraphicsRHI or TargetedRHIs. DefaultGraphicsRHI has precedence.
 static WindowsRHI ChooseDefaultRHI()
 {
-	WindowsRHI DefaultRHI = WindowsRHI::D3D12;
+	WindowsRHI DefaultRHI = WindowsRHI::D3D11;
 
 	// Default graphics RHI is the main project setting that governs the choice, so it takes the priority
 	FConfigFile EngineSettings;
