Set default values for GRHIMinimum/MaximumWaveSize to 0 and define it as an "unknown" value (i.e. when running on RHIs that do not support wave ops).

diff --git a/Engine/Source/Runtime/RHI/Private/RHI.cpp b/Engine/Source/Runtime/RHI/Private/RHI.cpp
--- a/Engine/Source/Runtime/RHI/Private/RHI.cpp
+++ b/Engine/Source/Runtime/RHI/Private/RHI.cpp
@@ -1394,8 +1394,8 @@ uint32 GRHIRayTracingAccelerationStructureAlignment = 0;
 uint32 GRHIRayTracingShaderTableAlignment = 0;
 uint32 GRHIRayTracingInstanceDescriptorSize = 0;
 bool GRHISupportsWaveOperations = false;
-int32 GRHIMinimumWaveSize = 4; // Minimum supported value in SM 6.0
-int32 GRHIMaximumWaveSize = 128; // Maximum supported value in SM 6.0
+int32 GRHIMinimumWaveSize = 0;
+int32 GRHIMaximumWaveSize = 0;
 bool GRHISupportsRHIThread = false;
 bool GRHISupportsRHIOnTaskThread = false;
 bool GRHISupportsParallelRHIExecute = false;
diff --git a/Engine/Source/Runtime/RHI/Public/RHI.h b/Engine/Source/Runtime/RHI/Public/RHI.h
--- a/Engine/Source/Runtime/RHI/Public/RHI.h
+++ b/Engine/Source/Runtime/RHI/Public/RHI.h
@@ -602,7 +602,11 @@ extern RHI_API uint32 GRHIRayTracingInstanceDescriptorSize;
 /** Whether or not the RHI supports shader wave operations (shader model 6.0). */
 extern RHI_API bool GRHISupportsWaveOperations;
 
-/** Specifies the minimum and maximum number of lanes in the SIMD wave that this GPU can support. I.e. 32 on NVIDIA, 64 on AMD. Values are in range [4..128]. */
+/** 
+* Specifies the minimum and maximum number of lanes in the SIMD wave that this GPU can support. I.e. 32 on NVIDIA, 64 on AMD.
+* Valid values are in range [4..128] (as per SM 6.0 specification) or 0 if unknown.
+* Rendering code must always check GRHISupportsWaveOperations in addition to wave min/max size.
+*/
 extern RHI_API int32 GRHIMinimumWaveSize;
 extern RHI_API int32 GRHIMaximumWaveSize;
 
