[LLM] - RHI layer scopes for LLM tracking (cross platform)

diff --git a/Engine/Source/Runtime/RHI/Private/RHICommandList.cpp b/Engine/Source/Runtime/RHI/Private/RHICommandList.cpp
--- a/Engine/Source/Runtime/RHI/Private/RHICommandList.cpp
+++ b/Engine/Source/Runtime/RHI/Private/RHICommandList.cpp
@@ -865,7 +865,7 @@ static FORCEINLINE bool IsInRenderingOrRHIThread()
 
 void FRHICommandListExecutor::ExecuteList(FRHICommandListBase& CmdList)
 {
-	LLM_SCOPE(ELLMTag::RHIMisc);
+	LLM_SCOPE_BYNAME(TEXT("FRHICommandListExecutorExecuteList"));
 
 	check(&CmdList != &GetImmediateCommandList() && (GRHISupportsParallelRHIExecute || IsInRenderingOrRHIThread()));
 
diff --git a/Engine/Source/Runtime/RHI/Public/RHICommandList.h b/Engine/Source/Runtime/RHI/Public/RHICommandList.h
--- a/Engine/Source/Runtime/RHI/Public/RHICommandList.h
+++ b/Engine/Source/Runtime/RHI/Public/RHICommandList.h
@@ -801,6 +801,7 @@ struct FRHICommand : public FRHICommandBase
 
 	void ExecuteAndDestruct(FRHICommandListBase& CmdList, FRHICommandListDebugContext& Context) override final
 	{
+		LLM_SCOPE_BYNAME(TEXT("RHIMisc/CommandList/ExecuteAndDestruct"));
 		TRACE_CPUPROFILER_EVENT_SCOPE_ON_CHANNEL_STR(NameType::TStr(), RHICommandsChannel);
 		
 		TCmd *ThisCmd = static_cast<TCmd*>(this);
@@ -4072,25 +4073,25 @@ public:
 
 	FORCEINLINE FSamplerStateRHIRef CreateSamplerState(const FSamplerStateInitializerRHI& Initializer)
 	{
-		LLM_SCOPE(ELLMTag::RHIMisc);
+		LLM_SCOPE_BYNAME(TEXT("RHIMisc/CreateSamplerState"));
 		return RHICreateSamplerState(Initializer);
 	}
 	
 	FORCEINLINE FRasterizerStateRHIRef CreateRasterizerState(const FRasterizerStateInitializerRHI& Initializer)
 	{
-		LLM_SCOPE(ELLMTag::RHIMisc);
+		LLM_SCOPE_BYNAME(TEXT("RHIMisc/CreateRasterizerState"));
 		return RHICreateRasterizerState(Initializer);
 	}
 	
 	FORCEINLINE FDepthStencilStateRHIRef CreateDepthStencilState(const FDepthStencilStateInitializerRHI& Initializer)
 	{
-		LLM_SCOPE(ELLMTag::RHIMisc);
+		LLM_SCOPE_BYNAME(TEXT("RHIMisc/CreateDepthStencilState"));
 		return RHICreateDepthStencilState(Initializer);
 	}
 	
 	FORCEINLINE FBlendStateRHIRef CreateBlendState(const FBlendStateInitializerRHI& Initializer)
 	{
-		LLM_SCOPE(ELLMTag::RHIMisc);
+		LLM_SCOPE_BYNAME(TEXT("RHIMisc/CreateBlendState"));
 		return RHICreateBlendState(Initializer);
 	}
 	
@@ -4267,49 +4268,49 @@ public:
 	
 	FORCEINLINE FUnorderedAccessViewRHIRef CreateUnorderedAccessView(FRHIBuffer* Buffer, bool bUseUAVCounter, bool bAppendBuffer)
 	{
-		LLM_SCOPE(ELLMTag::RHIMisc);
+		LLM_SCOPE_BYNAME(TEXT("RHIMisc/CreateUnorderedAccessView"));
 		checkf(Buffer, TEXT("Can't create a view off a null resource!"));
 		return GDynamicRHI->RHICreateUnorderedAccessView_RenderThread(*this, Buffer, bUseUAVCounter, bAppendBuffer);
 	}
 	
 	FORCEINLINE FUnorderedAccessViewRHIRef CreateUnorderedAccessView(FRHITexture* Texture, uint32 MipLevel)
 	{
-		LLM_SCOPE(ELLMTag::RHIMisc);
+		LLM_SCOPE_BYNAME(TEXT("RHIMisc/CreateUnorderedAccessView"));
 		checkf(Texture, TEXT("Can't create a view off a null resource!"));
 		return GDynamicRHI->RHICreateUnorderedAccessView_RenderThread(*this, Texture, MipLevel);
 	}
 	
 	FORCEINLINE FUnorderedAccessViewRHIRef CreateUnorderedAccessView(FRHITexture* Texture, uint32 MipLevel, uint8 Format)
 	{
-		LLM_SCOPE(ELLMTag::RHIMisc);
+		LLM_SCOPE_BYNAME(TEXT("RHIMisc/CreateUnorderedAccessView"));
 		checkf(Texture, TEXT("Can't create a view off a null resource!"));
 		return GDynamicRHI->RHICreateUnorderedAccessView_RenderThread(*this, Texture, MipLevel, Format);
 	}
 
 	FORCEINLINE FUnorderedAccessViewRHIRef CreateUnorderedAccessView(FRHIBuffer* Buffer, uint8 Format)
 	{
-		LLM_SCOPE(ELLMTag::RHIMisc);
+		LLM_SCOPE_BYNAME(TEXT("RHIMisc/CreateUnorderedAccessView"));
 		checkf(Buffer, TEXT("Can't create a view off a null resource!"));
 		return GDynamicRHI->RHICreateUnorderedAccessView_RenderThread(*this, Buffer, Format);
 	}
 
 	FORCEINLINE FShaderResourceViewRHIRef CreateShaderResourceView(FRHIBuffer* Buffer)
 	{
-		LLM_SCOPE(ELLMTag::RHIMisc);
+		LLM_SCOPE_BYNAME(TEXT("RHIMisc/CreateShaderResourceView"));
 		checkf(Buffer, TEXT("Can't create a view off a null resource!"));
 		return GDynamicRHI->RHICreateShaderResourceView_RenderThread(*this, Buffer);
 	}
 	
 	FORCEINLINE FShaderResourceViewRHIRef CreateShaderResourceView(FRHIBuffer* Buffer, uint32 Stride, uint8 Format)
 	{
-		LLM_SCOPE(ELLMTag::RHIMisc);
+		LLM_SCOPE_BYNAME(TEXT("RHIMisc/CreateShaderResourceView"));
 		checkf(Buffer, TEXT("Can't create a view off a null resource!"));
 		return GDynamicRHI->CreateShaderResourceView_RenderThread(*this, Buffer, Stride, Format);
 	}
 	
 	FORCEINLINE FShaderResourceViewRHIRef CreateShaderResourceView(const FShaderResourceViewInitializer& Initializer)
 	{
-		LLM_SCOPE(ELLMTag::RHIMisc);
+		LLM_SCOPE_BYNAME(TEXT("RHIMisc/CreateShaderResourceView"));
 		return GDynamicRHI->CreateShaderResourceView_RenderThread(*this, Initializer);
 	}
 	
@@ -4361,14 +4362,14 @@ public:
 	
 	FORCEINLINE FShaderResourceViewRHIRef CreateShaderResourceView(FRHITexture* Texture, const FRHITextureSRVCreateInfo& CreateInfo)
 	{
-		LLM_SCOPE(ELLMTag::RHIMisc);
+		LLM_SCOPE_BYNAME(TEXT("RHIMisc/CreateShaderResourceView"));
 		checkf(Texture, TEXT("Can't create a view off a null resource!"));
 		return GDynamicRHI->RHICreateShaderResourceView_RenderThread(*this, Texture, CreateInfo);
 	}
 
 	FORCEINLINE FShaderResourceViewRHIRef CreateShaderResourceView(FRHITexture* Texture, uint8 MipLevel)
 	{
-		LLM_SCOPE(ELLMTag::RHIMisc);
+		LLM_SCOPE_BYNAME(TEXT("RHIMisc/CreateShaderResourceView"));
 		checkf(Texture, TEXT("Can't create a view off a null resource!"));
 		const FRHITextureSRVCreateInfo CreateInfo(MipLevel, 1, Texture->GetFormat());
 		return GDynamicRHI->RHICreateShaderResourceView_RenderThread(*this, Texture, CreateInfo);
@@ -4376,7 +4377,7 @@ public:
 	
 	FORCEINLINE FShaderResourceViewRHIRef CreateShaderResourceView(FRHITexture* Texture, uint8 MipLevel, uint8 NumMipLevels, EPixelFormat Format)
 	{
-		LLM_SCOPE(ELLMTag::RHIMisc);
+		LLM_SCOPE_BYNAME(TEXT("RHIMisc/CreateShaderResourceView"));
 		checkf(Texture, TEXT("Can't create a view off a null resource!"));
 		const FRHITextureSRVCreateInfo CreateInfo(MipLevel, NumMipLevels, Format);
 		return GDynamicRHI->RHICreateShaderResourceView_RenderThread(*this, Texture, CreateInfo);
@@ -4384,14 +4385,14 @@ public:
 
 	FORCEINLINE FShaderResourceViewRHIRef CreateShaderResourceViewWriteMask(FRHITexture2D* Texture2DRHI)
 	{
-		LLM_SCOPE(ELLMTag::RHIMisc);
+		LLM_SCOPE_BYNAME(TEXT("RHIMisc/CreateShaderResourceViewWriteMask"));
 		checkf(Texture2DRHI, TEXT("Can't create a view off a null resource!"));
 		return GDynamicRHI->RHICreateShaderResourceViewWriteMask_RenderThread(*this, Texture2DRHI);
 	}
 
 	FORCEINLINE FShaderResourceViewRHIRef CreateShaderResourceViewFMask(FRHITexture2D* Texture2DRHI)
 	{
-		LLM_SCOPE(ELLMTag::RHIMisc);
+		LLM_SCOPE_BYNAME(TEXT("RHIMisc/CreateShaderResourceViewFMask"));
 		checkf(Texture2DRHI, TEXT("Can't create a view off a null resource!"));
 		return GDynamicRHI->RHICreateShaderResourceViewFMask_RenderThread(*this, Texture2DRHI);
 	}
@@ -4694,7 +4695,7 @@ public:
 	
 	FORCEINLINE void Tick(float DeltaTime)
 	{
-		LLM_SCOPE(ELLMTag::RHIMisc);
+		LLM_SCOPE_BYNAME(TEXT("RHIMisc/RHITick"));
 		RHITick(DeltaTime);
 	}
 	
