Don't crash if UpdateTextureReference() is called with a null texture reference.

This can happen when a UTextureRenderTarget2D object is created as part of a class default object, because the UTexture constructor skips initializing the TextureReference member for template objects, but there's other code which doesn't check the same condition before calling InitCustomFormat() or other methods. This isn't particularly nice, but UE4 allowed it (because the old RHIUpdateTextureReference implementations checked for null) so we'll make it resilient here as well.

#rnx

diff --git a/Engine/Source/Runtime/RHI/Private/RHICommandList.cpp b/Engine/Source/Runtime/RHI/Private/RHICommandList.cpp
--- a/Engine/Source/Runtime/RHI/Private/RHICommandList.cpp
+++ b/Engine/Source/Runtime/RHI/Private/RHICommandList.cpp
@@ -2713,6 +2713,11 @@ void FDynamicRHI::RHIReadSurfaceFloatData_RenderThread(class FRHICommandListImme
 
 void FRHICommandListImmediate::UpdateTextureReference(FRHITextureReference* TextureRef, FRHITexture* NewTexture)
 {
+	if (TextureRef == nullptr)
+	{
+		return;
+	}
+
 	EnqueueLambda([TextureRef, NewTexture](auto&)
 	{
 		TextureRef->SetReferencedTexture(NewTexture);
