Adding RHIUniformBufferLayoutInitializer and RHIUniformBufferLayout is a proper FRHIResource. This brings proper RHI lifetime support to RHIUniformBufferLayout so we don't delete the layouts while RHI commands that use them are still in flight.

This should fully address UniformBufferLayout/UniformBuffer crashes that randomly happen during shader compiling in the Editor.

diff --git a/Engine/Plugins/FX/Niagara/Source/Niagara/Classes/NiagaraComputeExecutionContext.h b/Engine/Plugins/FX/Niagara/Source/Niagara/Classes/NiagaraComputeExecutionContext.h
--- a/Engine/Plugins/FX/Niagara/Source/Niagara/Classes/NiagaraComputeExecutionContext.h
+++ b/Engine/Plugins/FX/Niagara/Source/Niagara/Classes/NiagaraComputeExecutionContext.h
@@ -11,12 +11,13 @@ class FNiagaraGPUSystemTick;
 struct FNiagaraComputeInstanceData;
 struct FNiagaraComputeExecutionContext;
 
-class FNiagaraRHIUniformBufferLayout : public FRHIResource
+class FNiagaraRHIUniformBufferLayout : public FRHIUniformBufferLayout
 {
 public:
-	explicit FNiagaraRHIUniformBufferLayout(const TCHAR* LayoutName) : FRHIResource(RRT_None), UBLayout(LayoutName) { }
-
-	FRHIUniformBufferLayout UBLayout;
+	explicit FNiagaraRHIUniformBufferLayout(const TCHAR* LayoutName, uint32 ConstantBufferSize)
+		: FRHIUniformBufferLayout(FRHIUniformBufferLayoutInitializer(LayoutName, ConstantBufferSize))
+	{
+	}
 };
 
 struct FNiagaraSimStageData
diff --git a/Engine/Plugins/FX/Niagara/Source/Niagara/Private/NiagaraComputeExecutionContext.cpp b/Engine/Plugins/FX/Niagara/Source/Niagara/Private/NiagaraComputeExecutionContext.cpp
--- a/Engine/Plugins/FX/Niagara/Source/Niagara/Private/NiagaraComputeExecutionContext.cpp
+++ b/Engine/Plugins/FX/Niagara/Source/Niagara/Private/NiagaraComputeExecutionContext.cpp
@@ -18,7 +18,7 @@ FNiagaraComputeExecutionContext::FNiagaraComputeExecutionContext()
 	, GPUScript(nullptr)
 	, GPUScript_RT(nullptr)
 {
-	ExternalCBufferLayout = new FNiagaraRHIUniformBufferLayout(TEXT("Niagara GPU External CBuffer"));
+	ExternalCBufferLayout = new FNiagaraRHIUniformBufferLayout(TEXT("Niagara GPU External CBuffer"), 0);
 }
 
 FNiagaraComputeExecutionContext::~FNiagaraComputeExecutionContext()
diff --git a/Engine/Plugins/FX/Niagara/Source/Niagara/Private/NiagaraEmitterInstance.cpp b/Engine/Plugins/FX/Niagara/Source/Niagara/Private/NiagaraEmitterInstance.cpp
--- a/Engine/Plugins/FX/Niagara/Source/Niagara/Private/NiagaraEmitterInstance.cpp
+++ b/Engine/Plugins/FX/Niagara/Source/Niagara/Private/NiagaraEmitterInstance.cpp
@@ -1522,8 +1522,11 @@ void FNiagaraEmitterInstance::Tick(float DeltaSeconds)
 		// Because each context is only ran once each frame, the CBuffer layout stays constant for the lifetime duration of the CBuffer (one frame).
 
 		// @todo-threadsafety do this once during init. Should not change during runtime...
-		GPUExecContext->ExternalCBufferLayout->UBLayout.ConstantBufferSize = ParmSize / (GPUExecContext->HasInterpolationParameters ? 2 : 1);
-		GPUExecContext->ExternalCBufferLayout->UBLayout.ComputeHash();
+		const uint32 ConstantBufferSize = ParmSize / (GPUExecContext->HasInterpolationParameters ? 2 : 1);
+		if (GPUExecContext->ExternalCBufferLayout->ConstantBufferSize != ConstantBufferSize)
+		{
+			GPUExecContext->ExternalCBufferLayout = new FNiagaraRHIUniformBufferLayout(TEXT("Niagara GPU External CBuffer"), ConstantBufferSize);
+		}
 
 		// Need to call post-tick, which calls the copy to previous for interpolated spawning
 		SpawnExecContext.PostTick();
diff --git a/Engine/Plugins/FX/Niagara/Source/Niagara/Private/NiagaraEmitterInstanceBatcher.cpp b/Engine/Plugins/FX/Niagara/Source/Niagara/Private/NiagaraEmitterInstanceBatcher.cpp
--- a/Engine/Plugins/FX/Niagara/Source/Niagara/Private/NiagaraEmitterInstanceBatcher.cpp
+++ b/Engine/Plugins/FX/Niagara/Source/Niagara/Private/NiagaraEmitterInstanceBatcher.cpp
@@ -227,21 +227,10 @@ NiagaraEmitterInstanceBatcher::NiagaraEmitterInstanceBatcher(ERHIFeatureLevel::T
 		}
 	}
 
-	GlobalCBufferLayout = new FNiagaraRHIUniformBufferLayout(TEXT("Niagara GPU Global CBuffer"));
-	GlobalCBufferLayout->UBLayout.ConstantBufferSize = sizeof(FNiagaraGlobalParameters);
-	GlobalCBufferLayout->UBLayout.ComputeHash();
-
-	SystemCBufferLayout = new FNiagaraRHIUniformBufferLayout(TEXT("Niagara GPU System CBuffer"));
-	SystemCBufferLayout->UBLayout.ConstantBufferSize = sizeof(FNiagaraSystemParameters);
-	SystemCBufferLayout->UBLayout.ComputeHash();
-
-	OwnerCBufferLayout = new FNiagaraRHIUniformBufferLayout(TEXT("Niagara GPU Owner CBuffer"));
-	OwnerCBufferLayout->UBLayout.ConstantBufferSize = sizeof(FNiagaraOwnerParameters);
-	OwnerCBufferLayout->UBLayout.ComputeHash();
-
-	EmitterCBufferLayout = new FNiagaraRHIUniformBufferLayout(TEXT("Niagara GPU Emitter CBuffer"));
-	EmitterCBufferLayout->UBLayout.ConstantBufferSize = sizeof(FNiagaraEmitterParameters);
-	EmitterCBufferLayout->UBLayout.ComputeHash();
+	GlobalCBufferLayout  = new FNiagaraRHIUniformBufferLayout(TEXT("Niagara GPU Global CBuffer"),  sizeof(FNiagaraGlobalParameters));
+	SystemCBufferLayout  = new FNiagaraRHIUniformBufferLayout(TEXT("Niagara GPU System CBuffer"),  sizeof(FNiagaraSystemParameters));
+	OwnerCBufferLayout   = new FNiagaraRHIUniformBufferLayout(TEXT("Niagara GPU Owner CBuffer"),   sizeof(FNiagaraOwnerParameters));
+	EmitterCBufferLayout = new FNiagaraRHIUniformBufferLayout(TEXT("Niagara GPU Emitter CBuffer"), sizeof(FNiagaraEmitterParameters));
 
 #if RHI_RAYTRACING
 	RayTracingHelper.Reset(new FNiagaraRayTracingHelper(InShaderPlatform));
@@ -365,9 +354,9 @@ void NiagaraEmitterInstanceBatcher::BuildConstantBuffers(FNiagaraGPUSystemTick&
 
 	const FRHIUniformBufferLayout* SystemLayouts[FNiagaraGPUSystemTick::UBT_NumSystemTypes] =
 	{
-		&GlobalCBufferLayout->UBLayout,
-		&SystemCBufferLayout->UBLayout,
-		&OwnerCBufferLayout->UBLayout
+		GlobalCBufferLayout,
+		SystemCBufferLayout,
+		OwnerCBufferLayout
 	};
 
 	for (int32 InterpIt = 0; InterpIt < InterpScale; ++InterpIt)
@@ -380,7 +369,7 @@ void NiagaraEmitterInstanceBatcher::BuildConstantBuffers(FNiagaraGPUSystemTick&
 			{
 				BufferRef = RHICreateUniformBuffer(
 					Tick.GetUniformBufferSource((FNiagaraGPUSystemTick::EUniformBufferType) SystemTypeIt, nullptr, !InterpIt),
-					*SystemLayouts[SystemTypeIt],
+					SystemLayouts[SystemTypeIt],
 					((BoundParameterCounts[SystemTypeIt][InterpIt] > 1) || HasMultipleStages)
 						? EUniformBufferUsage::UniformBuffer_SingleFrame
 						: EUniformBufferUsage::UniformBuffer_SingleDraw);
@@ -397,8 +386,8 @@ void NiagaraEmitterInstanceBatcher::BuildConstantBuffers(FNiagaraGPUSystemTick&
 
 				if (BoundParameterCounts[InstanceTypeIt][InterpIt])
 				{
-					FRHIUniformBufferLayout& Layout = InstanceTypeIt == FNiagaraGPUSystemTick::UBT_Emitter ? EmitterCBufferLayout->UBLayout : Instance.Context->ExternalCBufferLayout->UBLayout;
-					if (Layout.Resources.Num() > 0 || Layout.ConstantBufferSize > 0)
+					FRHIUniformBufferLayout* Layout = InstanceTypeIt == FNiagaraGPUSystemTick::UBT_Emitter ? EmitterCBufferLayout : Instance.Context->ExternalCBufferLayout;
+					if (Layout->Resources.Num() > 0 || Layout->ConstantBufferSize > 0)
 					{
 						BufferRef = RHICreateUniformBuffer(
 							Tick.GetUniformBufferSource((FNiagaraGPUSystemTick::EUniformBufferType) InstanceTypeIt, &Instance, !InterpIt),
diff --git a/Engine/Shaders/Public/ShaderVersion.ush b/Engine/Shaders/Public/ShaderVersion.ush
--- a/Engine/Shaders/Public/ShaderVersion.ush
+++ b/Engine/Shaders/Public/ShaderVersion.ush
@@ -3,4 +3,4 @@
 // in Platform.ush (which should be included in any shader) it allows to invalidate the shader DDC.
 // 
 // If you are merging streams and there is a conflict with this GUID you should make a new GUID rather than taking one or the other.
-// GUID = 04937C71BC6C4F878D4CA02CC386BF6E
+// GUID = 5F9FBB1B163644EA9F41E0B80737816F
diff --git a/Engine/Source/Runtime/Apple/AGXRHI/Private/AGXDynamicRHI.h b/Engine/Source/Runtime/Apple/AGXRHI/Private/AGXDynamicRHI.h
--- a/Engine/Source/Runtime/Apple/AGXRHI/Private/AGXDynamicRHI.h
+++ b/Engine/Source/Runtime/Apple/AGXRHI/Private/AGXDynamicRHI.h
@@ -51,7 +51,7 @@ public:
 	virtual FBoundShaderStateRHIRef RHICreateBoundShaderState(FRHIVertexDeclaration* VertexDeclaration, FRHIVertexShader* VertexShader, FRHIPixelShader* PixelShader, FRHIGeometryShader* GeometryShader) final override;
 	virtual FGraphicsPipelineStateRHIRef RHICreateGraphicsPipelineState(const FGraphicsPipelineStateInitializer& Initializer) final override;
 	virtual TRefCountPtr<FRHIComputePipelineState> RHICreateComputePipelineState(FRHIComputeShader* ComputeShader) final override;
-	virtual FUniformBufferRHIRef RHICreateUniformBuffer(const void* Contents, const FRHIUniformBufferLayout& Layout, EUniformBufferUsage Usage, EUniformBufferValidation Validation) final override;
+	virtual FUniformBufferRHIRef RHICreateUniformBuffer(const void* Contents, const FRHIUniformBufferLayout* Layout, EUniformBufferUsage Usage, EUniformBufferValidation Validation) final override;
 	virtual FBufferRHIRef RHICreateBuffer(uint32 Size, EBufferUsageFlags Usage, uint32 Stride, ERHIAccess InResourceState, FRHIResourceCreateInfo& CreateInfo) final override;
 	virtual void RHITransferBufferUnderlyingResource(FRHIBuffer* DestBuffer, FRHIBuffer* SrcBuffer) final override;
 	virtual void RHICopyBuffer(FRHIBuffer* SourceBuffer, FRHIBuffer* DestBuffer) final override;
diff --git a/Engine/Source/Runtime/Apple/AGXRHI/Private/AGXUniformBuffer.cpp b/Engine/Source/Runtime/Apple/AGXRHI/Private/AGXUniformBuffer.cpp
--- a/Engine/Source/Runtime/Apple/AGXRHI/Private/AGXUniformBuffer.cpp
+++ b/Engine/Source/Runtime/Apple/AGXRHI/Private/AGXUniformBuffer.cpp
@@ -11,7 +11,7 @@
 
 #pragma mark Suballocated Uniform Buffer Implementation
 
-FAGXSuballocatedUniformBuffer::FAGXSuballocatedUniformBuffer(const FRHIUniformBufferLayout& Layout, EUniformBufferUsage Usage, EUniformBufferValidation InValidation)
+FAGXSuballocatedUniformBuffer::FAGXSuballocatedUniformBuffer(const FRHIUniformBufferLayout* Layout, EUniformBufferUsage Usage, EUniformBufferValidation InValidation)
     : FRHIUniformBuffer(Layout)
     , LastFrameUpdated(0)
     , Offset(0)
diff --git a/Engine/Source/Runtime/Apple/AGXRHI/Private/AGXUniformBufferCommands.cpp b/Engine/Source/Runtime/Apple/AGXRHI/Private/AGXUniformBufferCommands.cpp
--- a/Engine/Source/Runtime/Apple/AGXRHI/Private/AGXUniformBufferCommands.cpp
+++ b/Engine/Source/Runtime/Apple/AGXRHI/Private/AGXUniformBufferCommands.cpp
@@ -35,7 +35,7 @@ static void DoUpdateUniformBuffer(FAGXUniformBuffer* UB, const void* Contents)
     }
 }
 
-FUniformBufferRHIRef FAGXDynamicRHI::RHICreateUniformBuffer(const void* Contents, const FRHIUniformBufferLayout& Layout, EUniformBufferUsage Usage, EUniformBufferValidation Validation)
+FUniformBufferRHIRef FAGXDynamicRHI::RHICreateUniformBuffer(const void* Contents, const FRHIUniformBufferLayout* Layout, EUniformBufferUsage Usage, EUniformBufferValidation Validation)
 {
     FAGXDeviceContext& DeviceContext = (FAGXDeviceContext&)GetAGXDeviceContext();
     FAGXFrameAllocator* UniformAllocator = DeviceContext.GetUniformAllocator();
diff --git a/Engine/Source/Runtime/Apple/AGXRHI/Public/AGXUniformBuffer.h b/Engine/Source/Runtime/Apple/AGXRHI/Public/AGXUniformBuffer.h
--- a/Engine/Source/Runtime/Apple/AGXRHI/Public/AGXUniformBuffer.h
+++ b/Engine/Source/Runtime/Apple/AGXRHI/Public/AGXUniformBuffer.h
@@ -38,7 +38,7 @@ private:
 public:
     // Creates a uniform buffer.
     // If Usage is SingleDraw or MultiFrame we will keep a copy of the data
-    FAGXSuballocatedUniformBuffer(const FRHIUniformBufferLayout& Layout, EUniformBufferUsage Usage, EUniformBufferValidation Validation);
+    FAGXSuballocatedUniformBuffer(const FRHIUniformBufferLayout* Layout, EUniformBufferUsage Usage, EUniformBufferValidation Validation);
     ~FAGXSuballocatedUniformBuffer();
     
 	void Update(const void* Contents, TArray<TRefCountPtr<FRHIResource> > const& InResourceTable);
diff --git a/Engine/Source/Runtime/Apple/MetalRHI/Private/MetalDynamicRHI.h b/Engine/Source/Runtime/Apple/MetalRHI/Private/MetalDynamicRHI.h
--- a/Engine/Source/Runtime/Apple/MetalRHI/Private/MetalDynamicRHI.h
+++ b/Engine/Source/Runtime/Apple/MetalRHI/Private/MetalDynamicRHI.h
@@ -52,7 +52,7 @@ public:
 	virtual FBoundShaderStateRHIRef RHICreateBoundShaderState(FRHIVertexDeclaration* VertexDeclaration, FRHIVertexShader* VertexShader, FRHIPixelShader* PixelShader, FRHIGeometryShader* GeometryShader) final override;
 	virtual FGraphicsPipelineStateRHIRef RHICreateGraphicsPipelineState(const FGraphicsPipelineStateInitializer& Initializer) final override;
 	virtual TRefCountPtr<FRHIComputePipelineState> RHICreateComputePipelineState(FRHIComputeShader* ComputeShader) final override;
-	virtual FUniformBufferRHIRef RHICreateUniformBuffer(const void* Contents, const FRHIUniformBufferLayout& Layout, EUniformBufferUsage Usage, EUniformBufferValidation Validation) final override;
+	virtual FUniformBufferRHIRef RHICreateUniformBuffer(const void* Contents, const FRHIUniformBufferLayout* Layout, EUniformBufferUsage Usage, EUniformBufferValidation Validation) final override;
 	virtual FBufferRHIRef RHICreateBuffer(uint32 Size, EBufferUsageFlags Usage, uint32 Stride, ERHIAccess InResourceState, FRHIResourceCreateInfo& CreateInfo) final override;
 	virtual void RHITransferBufferUnderlyingResource(FRHIBuffer* DestBuffer, FRHIBuffer* SrcBuffer) final override;
 	virtual void RHICopyBuffer(FRHIBuffer* SourceBuffer, FRHIBuffer* DestBuffer) final override;
diff --git a/Engine/Source/Runtime/Apple/MetalRHI/Private/MetalUniformBuffer.cpp b/Engine/Source/Runtime/Apple/MetalRHI/Private/MetalUniformBuffer.cpp
--- a/Engine/Source/Runtime/Apple/MetalRHI/Private/MetalUniformBuffer.cpp
+++ b/Engine/Source/Runtime/Apple/MetalRHI/Private/MetalUniformBuffer.cpp
@@ -11,7 +11,7 @@
 
 #pragma mark Suballocated Uniform Buffer Implementation
 
-FMetalSuballocatedUniformBuffer::FMetalSuballocatedUniformBuffer(const FRHIUniformBufferLayout& Layout, EUniformBufferUsage Usage, EUniformBufferValidation InValidation)
+FMetalSuballocatedUniformBuffer::FMetalSuballocatedUniformBuffer(const FRHIUniformBufferLayout* Layout, EUniformBufferUsage Usage, EUniformBufferValidation InValidation)
     : FRHIUniformBuffer(Layout)
     , LastFrameUpdated(0)
     , Offset(0)
diff --git a/Engine/Source/Runtime/Apple/MetalRHI/Private/MetalUniformBufferCommands.cpp b/Engine/Source/Runtime/Apple/MetalRHI/Private/MetalUniformBufferCommands.cpp
--- a/Engine/Source/Runtime/Apple/MetalRHI/Private/MetalUniformBufferCommands.cpp
+++ b/Engine/Source/Runtime/Apple/MetalRHI/Private/MetalUniformBufferCommands.cpp
@@ -35,7 +35,7 @@ static void DoUpdateUniformBuffer(FMetalUniformBuffer* UB, const void* Contents)
     }
 }
 
-FUniformBufferRHIRef FMetalDynamicRHI::RHICreateUniformBuffer(const void* Contents, const FRHIUniformBufferLayout& Layout, EUniformBufferUsage Usage, EUniformBufferValidation Validation)
+FUniformBufferRHIRef FMetalDynamicRHI::RHICreateUniformBuffer(const void* Contents, const FRHIUniformBufferLayout* Layout, EUniformBufferUsage Usage, EUniformBufferValidation Validation)
 {
     FMetalDeviceContext& DeviceContext = (FMetalDeviceContext&)GetMetalDeviceContext();
     FMetalFrameAllocator* UniformAllocator = DeviceContext.GetUniformAllocator();
diff --git a/Engine/Source/Runtime/Apple/MetalRHI/Public/MetalUniformBuffer.h b/Engine/Source/Runtime/Apple/MetalRHI/Public/MetalUniformBuffer.h
--- a/Engine/Source/Runtime/Apple/MetalRHI/Public/MetalUniformBuffer.h
+++ b/Engine/Source/Runtime/Apple/MetalRHI/Public/MetalUniformBuffer.h
@@ -38,7 +38,7 @@ private:
 public:
     // Creates a uniform buffer.
     // If Usage is SingleDraw or MultiFrame we will keep a copy of the data
-    FMetalSuballocatedUniformBuffer(const FRHIUniformBufferLayout& Layout, EUniformBufferUsage Usage, EUniformBufferValidation Validation);
+    FMetalSuballocatedUniformBuffer(const FRHIUniformBufferLayout* Layout, EUniformBufferUsage Usage, EUniformBufferValidation Validation);
     ~FMetalSuballocatedUniformBuffer();
     
 	void Update(const void* Contents, TArray<TRefCountPtr<FRHIResource> > const& InResourceTable);
diff --git a/Engine/Source/Runtime/D3D12RHI/Private/D3D12RHIPrivate.h b/Engine/Source/Runtime/D3D12RHI/Private/D3D12RHIPrivate.h
--- a/Engine/Source/Runtime/D3D12RHI/Private/D3D12RHIPrivate.h
+++ b/Engine/Source/Runtime/D3D12RHI/Private/D3D12RHIPrivate.h
@@ -285,7 +285,7 @@ public:
 	virtual TRefCountPtr<FRHIComputePipelineState> RHICreateComputePipelineState(FRHIComputeShader* ComputeShader) final override;
 	virtual void RHICreateTransition(FRHITransition* Transition, const FRHITransitionCreateInfo& CreateInfo) final override;
 	virtual void RHIReleaseTransition(FRHITransition* Transition) final override;
-	virtual FUniformBufferRHIRef RHICreateUniformBuffer(const void* Contents, const FRHIUniformBufferLayout& Layout, EUniformBufferUsage Usage, EUniformBufferValidation Validation) final override;
+	virtual FUniformBufferRHIRef RHICreateUniformBuffer(const void* Contents, const FRHIUniformBufferLayout* Layout, EUniformBufferUsage Usage, EUniformBufferValidation Validation) final override;
 	virtual void RHIUpdateUniformBuffer(FRHIUniformBuffer* UniformBufferRHI, const void* Contents) final override;
 	virtual FBufferRHIRef RHICreateBuffer(uint32 Size, EBufferUsageFlags Usage, uint32 Stride, ERHIAccess ResourceState, FRHIResourceCreateInfo& CreateInfo) final override;
 	virtual void RHITransferBufferUnderlyingResource(FRHIBuffer* DestBuffer, FRHIBuffer* SrcBuffer) final override;
diff --git a/Engine/Source/Runtime/D3D12RHI/Private/D3D12UniformBuffer.cpp b/Engine/Source/Runtime/D3D12RHI/Private/D3D12UniformBuffer.cpp
--- a/Engine/Source/Runtime/D3D12RHI/Private/D3D12UniformBuffer.cpp
+++ b/Engine/Source/Runtime/D3D12RHI/Private/D3D12UniformBuffer.cpp
@@ -8,13 +8,13 @@
 #include "UniformBuffer.h"
 #include "ShaderParameterStruct.h"
 
-FUniformBufferRHIRef FD3D12DynamicRHI::RHICreateUniformBuffer(const void* Contents, const FRHIUniformBufferLayout& Layout, EUniformBufferUsage Usage, EUniformBufferValidation Validation)
+FUniformBufferRHIRef FD3D12DynamicRHI::RHICreateUniformBuffer(const void* Contents, const FRHIUniformBufferLayout* Layout, EUniformBufferUsage Usage, EUniformBufferValidation Validation)
 {
 	SCOPE_CYCLE_COUNTER(STAT_D3D12UpdateUniformBufferTime);
 
 	if (Validation == EUniformBufferValidation::ValidateResources)
 	{
-		ValidateShaderParameterResourcesRHI(Contents, Layout);
+		ValidateShaderParameterResourcesRHI(Contents, *Layout);
 	}
 
 	//Note: This is not overly efficient in the mGPU case (we create two+ upload locations) but the CPU savings of having no extra indirection to the resource are worth
@@ -26,7 +26,7 @@ FUniformBufferRHIRef FD3D12DynamicRHI::RHICreateUniformBuffer(const void* Conten
 		FD3D12UniformBuffer* NewUniformBuffer = new FD3D12UniformBuffer(Device, Layout, Usage);
 		check(nullptr != NewUniformBuffer);
 
-		const uint32 NumBytesActualData = Layout.ConstantBufferSize;
+		const uint32 NumBytesActualData = Layout->ConstantBufferSize;
 		if (NumBytesActualData > 0)
 		{
 			// Is this check really needed?
@@ -73,9 +73,9 @@ FUniformBufferRHIRef FD3D12DynamicRHI::RHICreateUniformBuffer(const void* Conten
 
 	check(UniformBufferOut);
 
-	if (Layout.Resources.Num())
+	if (Layout->Resources.Num())
 	{
-		const int32 NumResources = Layout.Resources.Num();
+		const int32 NumResources = Layout->Resources.Num();
 
 		for (FD3D12UniformBuffer& CurrentBuffer : *UniformBufferOut)
 		{
@@ -83,7 +83,7 @@ FUniformBufferRHIRef FD3D12DynamicRHI::RHICreateUniformBuffer(const void* Conten
 			CurrentBuffer.ResourceTable.AddZeroed(NumResources);
 			for (int32 Index = 0; Index < NumResources; ++Index)
 			{
-				CurrentBuffer.ResourceTable[Index] = GetShaderParameterResourceRHI(Contents, Layout.Resources[Index].MemberOffset, Layout.Resources[Index].MemberType);
+				CurrentBuffer.ResourceTable[Index] = GetShaderParameterResourceRHI(Contents, Layout->Resources[Index].MemberOffset, Layout->Resources[Index].MemberType);
 			}
 		}
 	}
diff --git a/Engine/Source/Runtime/D3D12RHI/Public/D3D12Resources.h b/Engine/Source/Runtime/D3D12RHI/Public/D3D12Resources.h
--- a/Engine/Source/Runtime/D3D12RHI/Public/D3D12Resources.h
+++ b/Engine/Source/Runtime/D3D12RHI/Public/D3D12Resources.h
@@ -924,7 +924,7 @@ public:
 	const EUniformBufferUsage UniformBufferUsage;
 
 	/** Initialization constructor. */
-	FD3D12UniformBuffer(class FD3D12Device* InParent, const FRHIUniformBufferLayout& InLayout, EUniformBufferUsage InUniformBufferUsage)
+	FD3D12UniformBuffer(class FD3D12Device* InParent, const FRHIUniformBufferLayout* InLayout, EUniformBufferUsage InUniformBufferUsage)
 		: FRHIUniformBuffer(InLayout)
 		, FD3D12DeviceChild(InParent)
 #if USE_STATIC_ROOT_SIGNATURE
diff --git a/Engine/Source/Runtime/Engine/Private/GPUSkinVertexFactory.cpp b/Engine/Source/Runtime/Engine/Private/GPUSkinVertexFactory.cpp
--- a/Engine/Source/Runtime/Engine/Private/GPUSkinVertexFactory.cpp
+++ b/Engine/Source/Runtime/Engine/Private/GPUSkinVertexFactory.cpp
@@ -352,7 +352,7 @@ bool FGPUBaseSkinVertexFactory::FShaderDataType::UpdateBoneData(FRHICommandListI
 	}
 	else
 	{
-		UniformBuffer = RHICreateUniformBuffer(&GBoneUniformStruct, FBoneMatricesUniformShaderParameters::StaticStructMetadata.GetLayout(), UniformBuffer_MultiFrame);
+		UniformBuffer = RHICreateUniformBuffer(&GBoneUniformStruct, FBoneMatricesUniformShaderParameters::StaticStructMetadata.GetLayoutPtr(), UniformBuffer_MultiFrame);
 	}
 	return false;
 }
diff --git a/Engine/Source/Runtime/Engine/Private/Materials/MaterialShader.cpp b/Engine/Source/Runtime/Engine/Private/Materials/MaterialShader.cpp
--- a/Engine/Source/Runtime/Engine/Private/Materials/MaterialShader.cpp
+++ b/Engine/Source/Runtime/Engine/Private/Materials/MaterialShader.cpp
@@ -1804,13 +1804,6 @@ void FMaterialShaderMap::Compile(
 	ShaderMapId = InShaderMapId;
 	bIsPersistent = Material->IsPersistent();
 
-	// Setup the material compilation environment.
-	/*{
-		FShaderParametersMetadata* UniformBufferStruct = NewContent->MaterialCompilationOutput.UniformExpressionSet.CreateBufferStruct();
-		Material->SetupMaterialEnvironment(InPlatform, *UniformBufferStruct, InMaterialCompilationOutput.UniformExpressionSet, *MaterialEnvironment);
-		delete UniformBufferStruct;
-	}*/
-
 #if ALLOW_SHADERMAP_DEBUG_DATA && WITH_EDITOR
 	// Store the material name for debugging purposes.
 	// Note: Material instances with static parameters will have the same FriendlyName for their shader maps!
@@ -2902,6 +2895,15 @@ void FMaterialShaderMap::InitalizeForODSC(EShaderPlatform TargetShaderPlatform,
 }
 #endif // WITH_EDITOR
 
+void FMaterialShaderMap::PostFinalizeContent()
+{
+	UniformBufferLayout.SafeRelease();
+	if (GetContent())
+	{
+		UniformBufferLayout = RHICreateUniformBufferLayout(GetUniformExpressionSet().GetUniformBufferLayoutInitializer());
+	}
+}
+
 #if WITH_EDITOR
 void FMaterialShaderMap::GetOutdatedTypes(TArray<const FShaderType*>& OutdatedShaderTypes, TArray<const FShaderPipelineType*>& OutdatedShaderPipelineTypes, TArray<const FVertexFactoryType*>& OutdatedFactoryTypes) const
 {
diff --git a/Engine/Source/Runtime/Engine/Private/Materials/MaterialShared.cpp b/Engine/Source/Runtime/Engine/Private/Materials/MaterialShared.cpp
--- a/Engine/Source/Runtime/Engine/Private/Materials/MaterialShared.cpp
+++ b/Engine/Source/Runtime/Engine/Private/Materials/MaterialShared.cpp
@@ -3126,16 +3126,16 @@ void FMaterialRenderProxy::EvaluateUniformExpressions(FUniformExpressionCache& O
 		OutUniformExpressionCache.AllocatedVTs.Add(AllocatedVT);
 	}
 
-	const FRHIUniformBufferLayout& UniformBufferLayout = UniformExpressionSet.GetUniformBufferLayout();
+	const FRHIUniformBufferLayout* UniformBufferLayout = ShaderMap->GetUniformBufferLayout();
 	FMemMark Mark(FMemStack::Get());
-	uint8* TempBuffer = FMemStack::Get().PushBytes(UniformBufferLayout.ConstantBufferSize, SHADER_PARAMETER_STRUCT_ALIGNMENT);
+	uint8* TempBuffer = FMemStack::Get().PushBytes(UniformBufferLayout->ConstantBufferSize, SHADER_PARAMETER_STRUCT_ALIGNMENT);
 
 	check(TempBuffer != nullptr);
-	UniformExpressionSet.FillUniformBuffer(Context, OutUniformExpressionCache, TempBuffer, UniformBufferLayout.ConstantBufferSize);
+	UniformExpressionSet.FillUniformBuffer(Context, OutUniformExpressionCache, UniformBufferLayout, TempBuffer, UniformBufferLayout->ConstantBufferSize);
 
 	if (CommandListIfLocalMode)
 	{
-		OutUniformExpressionCache.LocalUniformBuffer = CommandListIfLocalMode->BuildLocalUniformBuffer(TempBuffer, UniformBufferLayout.ConstantBufferSize, UniformBufferLayout);
+		OutUniformExpressionCache.LocalUniformBuffer = CommandListIfLocalMode->BuildLocalUniformBuffer(TempBuffer, UniformBufferLayout->ConstantBufferSize, UniformBufferLayout);
 		check(OutUniformExpressionCache.LocalUniformBuffer.IsValid());
 	}
 	else
@@ -3147,7 +3147,7 @@ void FMaterialRenderProxy::EvaluateUniformExpressions(FUniformExpressionCache& O
 
 		if (IsValidRef(OutUniformExpressionCache.UniformBuffer))
 		{
-			check(OutUniformExpressionCache.UniformBuffer->GetLayout() == UniformBufferLayout);
+			check(OutUniformExpressionCache.UniformBuffer->GetLayoutPtr() == UniformBufferLayout);
 			RHIUpdateUniformBuffer(OutUniformExpressionCache.UniformBuffer, TempBuffer);
 		}
 		else
diff --git a/Engine/Source/Runtime/Engine/Private/Materials/MaterialUniformExpressions.cpp b/Engine/Source/Runtime/Engine/Private/Materials/MaterialUniformExpressions.cpp
--- a/Engine/Source/Runtime/Engine/Private/Materials/MaterialUniformExpressions.cpp
+++ b/Engine/Source/Runtime/Engine/Private/Materials/MaterialUniformExpressions.cpp
@@ -531,6 +531,9 @@ FShaderParametersMetadata* FUniformExpressionSet::CreateBufferStruct()
 	NextMemberOffset += SHADER_PARAMETER_POINTER_ALIGNMENT;
 
 	const uint32 StructSize = Align(NextMemberOffset, SHADER_PARAMETER_STRUCT_ALIGNMENT);
+
+	UniformBufferLayoutInitializer = FRHIUniformBufferLayoutInitializer(TEXT("Material"));
+
 	FShaderParametersMetadata* UniformBufferStruct = new FShaderParametersMetadata(
 		FShaderParametersMetadata::EUseCase::DataDrivenUniformBuffer,
 		EUniformBufferBindingFlags::Shader,
@@ -541,9 +544,10 @@ FShaderParametersMetadata* FUniformExpressionSet::CreateBufferStruct()
 		UE_LOG_SOURCE_FILE(__FILE__),
 		__LINE__,
 		StructSize,
-		Members);
+		Members,
+		false,
+		&UniformBufferLayoutInitializer);
 
-	UniformBufferLayout = UniformBufferStruct->GetLayout();
 	return UniformBufferStruct;
 }
 
@@ -694,11 +698,11 @@ void FUniformExpressionSet::GetTextureValue(int32 Index, const FMaterialRenderCo
 	}
 }
 
-void FUniformExpressionSet::FillUniformBuffer(const FMaterialRenderContext& MaterialRenderContext, const FUniformExpressionCache& UniformExpressionCache, uint8* TempBuffer, int TempBufferSize) const
+void FUniformExpressionSet::FillUniformBuffer(const FMaterialRenderContext& MaterialRenderContext, const FUniformExpressionCache& UniformExpressionCache, const FRHIUniformBufferLayout* UniformBufferLayout, uint8* TempBuffer, int TempBufferSize) const
 {
 	check(IsInParallelRenderingThread());
 
-	if (UniformBufferLayout.ConstantBufferSize > 0)
+	if (UniformBufferLayout->ConstantBufferSize > 0)
 	{
 		// stat disabled by default due to low-value/high-frequency
 		//QUICK_SCOPE_CYCLE_COUNTER(STAT_FUniformExpressionSet_FillUniformBuffer);
@@ -813,7 +817,7 @@ void FUniformExpressionSet::FillUniformBuffer(const FMaterialRenderContext& Mate
 				NumPageTableIndirectionTextures++;
 			}
 	
-			check(UniformBufferLayout.Resources.Num() == 
+			check(UniformBufferLayout->Resources.Num() == 
 				UniformTextureParameters[(uint32)EMaterialTextureParameterType::Standard2D].Num() * 2
 				+ UniformTextureParameters[(uint32)EMaterialTextureParameterType::Cube].Num() * 2
 				+ UniformTextureParameters[(uint32)EMaterialTextureParameterType::Array2D].Num() * 2
diff --git a/Engine/Source/Runtime/Engine/Private/Materials/ParameterCollection.cpp b/Engine/Source/Runtime/Engine/Private/Materials/ParameterCollection.cpp
--- a/Engine/Source/Runtime/Engine/Private/Materials/ParameterCollection.cpp
+++ b/Engine/Source/Runtime/Engine/Private/Materials/ParameterCollection.cpp
@@ -708,10 +708,7 @@ void FMaterialParameterCollectionInstanceResource::GameThread_Destroy()
 	);
 }
 
-FMaterialParameterCollectionInstanceResource::FMaterialParameterCollectionInstanceResource() :
-	UniformBufferLayout(TEXT("MaterialParameterCollectionInstanceResource"))
-{
-}
+FMaterialParameterCollectionInstanceResource::FMaterialParameterCollectionInstanceResource() = default;
 
 FMaterialParameterCollectionInstanceResource::~FMaterialParameterCollectionInstanceResource()
 {
@@ -726,18 +723,22 @@ void FMaterialParameterCollectionInstanceResource::UpdateContents(const FGuid& I
 	if (InId != FGuid() && Data.Num() > 0)
 	{
 		const uint32 NewSize = Data.GetTypeSize() * Data.Num();
-		check(UniformBufferLayout.Resources.Num() == 0);
+		check(UniformBufferLayout == nullptr || UniformBufferLayout->Resources.Num() == 0);
 
 		if (!bRecreateUniformBuffer && IsValidRef(UniformBuffer))
 		{
-			check(NewSize == UniformBufferLayout.ConstantBufferSize);
-			check(UniformBuffer->GetLayout() == UniformBufferLayout);
+			check(NewSize == UniformBufferLayout->ConstantBufferSize);
+			check(UniformBuffer->GetLayoutPtr() == UniformBufferLayout);
 			RHIUpdateUniformBuffer(UniformBuffer, Data.GetData());
 		}
 		else
 		{
-			UniformBufferLayout.ConstantBufferSize = NewSize;
-			UniformBufferLayout.ComputeHash();
+			FRHIUniformBufferLayoutInitializer UniformBufferLayoutInitializer(TEXT("MaterialParameterCollectionInstanceResource"));
+			UniformBufferLayoutInitializer.ConstantBufferSize = NewSize;
+			UniformBufferLayoutInitializer.ComputeHash();
+
+			UniformBufferLayout = RHICreateUniformBufferLayout(UniformBufferLayoutInitializer);
+
 			UniformBuffer = RHICreateUniformBuffer(Data.GetData(), UniformBufferLayout, UniformBuffer_MultiFrame);
 		}
 	}
diff --git a/Engine/Source/Runtime/Engine/Public/MaterialShared.h b/Engine/Source/Runtime/Engine/Public/MaterialShared.h
--- a/Engine/Source/Runtime/Engine/Public/MaterialShared.h
+++ b/Engine/Source/Runtime/Engine/Public/MaterialShared.h
@@ -465,7 +465,7 @@ public:
 
 	void SetParameterCollections(const TArray<class UMaterialParameterCollection*>& Collections);
 
-	ENGINE_API void FillUniformBuffer(const FMaterialRenderContext& MaterialRenderContext, const FUniformExpressionCache& UniformExpressionCache, uint8* TempBuffer, int TempBufferSize) const;
+	ENGINE_API void FillUniformBuffer(const FMaterialRenderContext& MaterialRenderContext, const FUniformExpressionCache& UniformExpressionCache, const FRHIUniformBufferLayout* UniformBufferLayout, uint8* TempBuffer, int TempBufferSize) const;
 
 	// Get a combined hash of all referenced Texture2D's underlying RHI textures, going through TextureReferences. Can be used to tell if any texture has gone through texture streaming mip changes recently.
 	ENGINE_API uint32 GetReferencedTexture2DRHIHash(const FMaterialRenderContext& MaterialRenderContext) const;
@@ -475,9 +475,9 @@ public:
 		return UniformExternalTextureParameters.Num() > 0;
 	}
 
-	const FRHIUniformBufferLayout& GetUniformBufferLayout() const
+	const FRHIUniformBufferLayoutInitializer& GetUniformBufferLayoutInitializer() const
 	{
-		return UniformBufferLayout;
+		return UniformBufferLayoutInitializer;
 	}
 
 	inline const FMaterialVectorParameterInfo& GetVectorParameter(uint32 Index) const { return UniformVectorParameters[Index]; }
@@ -536,7 +536,7 @@ protected:
 	/** Ids of parameter collections referenced by the material that was translated. */
 	LAYOUT_FIELD(TMemoryImageArray<FGuid>, ParameterCollections);
 
-	LAYOUT_FIELD(FRHIUniformBufferLayout, UniformBufferLayout);
+	LAYOUT_FIELD(FRHIUniformBufferLayoutInitializer, UniformBufferLayoutInitializer);
 };
 
 /** Stores outputs from the material compile that need to be saved. */
@@ -1255,6 +1255,8 @@ public:
 	}
 
 	const FUniformExpressionSet& GetUniformExpressionSet() const { return GetContent()->MaterialCompilationOutput.UniformExpressionSet; }
+	const FRHIUniformBufferLayout* GetUniformBufferLayout() const { return UniformBufferLayout; }
+
 	int32 GetNumRefs() const { return NumRefs; }
 	int32 GetRefCount() const { return NumRefs; }
 
@@ -1278,6 +1280,9 @@ public:
 	void InitalizeForODSC(EShaderPlatform TargetShaderPlatform, const FMaterialCompilationOutput& NewCompilationOutput);
 #endif
 
+protected:
+	void PostFinalizeContent() override;
+
 private:
 	/** 
 	 * A global map from a material's static parameter set to any shader map cached for that material. 
@@ -1302,6 +1307,8 @@ private:
 	TRefCountPtr<FSharedShaderCompilerEnvironment> PendingCompilerEnvironment;
 	TArray<TRefCountPtr<FMaterial>> CompilingMaterialDependencies;
 
+	FUniformBufferLayoutRHIRef UniformBufferLayout;
+
 #if ALLOW_SHADERMAP_DEBUG_DATA
 	float CompileTime;
 #endif
diff --git a/Engine/Source/Runtime/Engine/Public/ParameterCollection.h b/Engine/Source/Runtime/Engine/Public/ParameterCollection.h
--- a/Engine/Source/Runtime/Engine/Public/ParameterCollection.h
+++ b/Engine/Source/Runtime/Engine/Public/ParameterCollection.h
@@ -63,7 +63,7 @@ private:
 	/** Uniform buffer containing the UMaterialParameterCollection default parameter values and UMaterialParameterCollectionInstance instance overrides. */
 	FUniformBufferRHIRef UniformBuffer;
 
-	FRHIUniformBufferLayout UniformBufferLayout;
+	FUniformBufferLayoutRHIRef UniformBufferLayout;
 
 	void UpdateContents(const FGuid& InId, const TArray<FVector4>& Data, const FName& InOwnerName, bool bRecreateUniformBuffer);
 };
diff --git a/Engine/Source/Runtime/NullDrv/Public/NullRHI.h b/Engine/Source/Runtime/NullDrv/Public/NullRHI.h
--- a/Engine/Source/Runtime/NullDrv/Public/NullRHI.h
+++ b/Engine/Source/Runtime/NullDrv/Public/NullRHI.h
@@ -104,7 +104,7 @@ public:
 
 	}
 
-	virtual FUniformBufferRHIRef RHICreateUniformBuffer(const void* Contents, const FRHIUniformBufferLayout& Layout, EUniformBufferUsage Usage, EUniformBufferValidation Validation) final override
+	virtual FUniformBufferRHIRef RHICreateUniformBuffer(const void* Contents, const FRHIUniformBufferLayout* Layout, EUniformBufferUsage Usage, EUniformBufferValidation Validation) final override
 	{ 
 		return new FRHIUniformBuffer(Layout); 
 	}
diff --git a/Engine/Source/Runtime/OpenGLDrv/Private/OpenGLUniformBuffer.cpp b/Engine/Source/Runtime/OpenGLDrv/Private/OpenGLUniformBuffer.cpp
--- a/Engine/Source/Runtime/OpenGLDrv/Private/OpenGLUniformBuffer.cpp
+++ b/Engine/Source/Runtime/OpenGLDrv/Private/OpenGLUniformBuffer.cpp
@@ -375,7 +375,7 @@ static uint32 UniqueUniformBufferID()
 	return ++GUniqueUniformBufferID;
 }
 
-FOpenGLUniformBuffer::FOpenGLUniformBuffer(const FRHIUniformBufferLayout& InLayout)
+FOpenGLUniformBuffer::FOpenGLUniformBuffer(const FRHIUniformBufferLayout* InLayout)
 	: FRHIUniformBuffer(InLayout)
 	, Resource(0)
 	, Offset(0)
@@ -439,17 +439,17 @@ FOpenGLUniformBuffer::~FOpenGLUniformBuffer()
 }
 
 
-static void SetLayoutTable(FOpenGLUniformBuffer* NewUniformBuffer, const void* Contents,const FRHIUniformBufferLayout &Layout, EUniformBufferValidation Validation)
+static void SetLayoutTable(FOpenGLUniformBuffer* NewUniformBuffer, const void* Contents, const FRHIUniformBufferLayout* Layout, EUniformBufferValidation Validation)
 {
-	if (Layout.Resources.Num())
+	if (Layout->Resources.Num())
 	{
-		int32 NumResources = Layout.Resources.Num();
+		int32 NumResources = Layout->Resources.Num();
 		NewUniformBuffer->ResourceTable.Empty(NumResources);
 		NewUniformBuffer->ResourceTable.AddZeroed(NumResources);
 
 		for (int32 Index = 0; Index < NumResources; ++Index)
 		{
-			NewUniformBuffer->ResourceTable[Index] = GetShaderParameterResourceRHI(Contents, Layout.Resources[Index].MemberOffset, Layout.Resources[Index].MemberType);
+			NewUniformBuffer->ResourceTable[Index] = GetShaderParameterResourceRHI(Contents, Layout->Resources[Index].MemberOffset, Layout->Resources[Index].MemberType);
 		}
 	}
 }
@@ -497,21 +497,21 @@ void CopyDataToUniformBuffer(const bool bCanRunOnThisThread, FOpenGLUniformBuffe
 	}
 }
 
-static FUniformBufferRHIRef CreateUniformBuffer(const void* Contents, const FRHIUniformBufferLayout& Layout, EUniformBufferUsage Usage, EUniformBufferValidation Validation)
+static FUniformBufferRHIRef CreateUniformBuffer(const void* Contents, const FRHIUniformBufferLayout* Layout, EUniformBufferUsage Usage, EUniformBufferValidation Validation)
 {
 	// This should really be synchronized, if there's a chance it'll be used from more than one buffer. Luckily, uniform buffers
 	// are only used for drawing/shader usage, not for loading resources or framebuffer blitting, so no synchronization primitives for now.
 
 	// Explicitly check that the size is nonzero before allowing CreateBuffer to opaquely fail.
-	check(Layout.Resources.Num() > 0 || Layout.ConstantBufferSize > 0);
+	check(Layout->Resources.Num() > 0 || Layout->ConstantBufferSize > 0);
 
 	FOpenGLUniformBuffer* NewUniformBuffer = new FOpenGLUniformBuffer(Layout);
 
 	TFunction<void(void)> GLCreationFunc;
 
-	const uint32 BucketIndex = GetPoolBucketIndex(Layout.ConstantBufferSize);
+	const uint32 BucketIndex = GetPoolBucketIndex(Layout->ConstantBufferSize);
 	const uint32 SizeOfBufferToAllocate = UniformBufferSizeBuckets[BucketIndex];
-	const uint32 AllocatedSize = (SizeOfBufferToAllocate > 0) ? SizeOfBufferToAllocate : Layout.ConstantBufferSize;;
+	const uint32 AllocatedSize = (SizeOfBufferToAllocate > 0) ? SizeOfBufferToAllocate : Layout->ConstantBufferSize;;
 	
 	// EmulatedUniformDataRef will not be initialized on RHI thread. safe to use on RT thread.
 	FOpenGLEUniformBufferDataRef EmulatedUniformDataRef;
@@ -585,7 +585,7 @@ static FUniformBufferRHIRef CreateUniformBuffer(const void* Contents, const FRHI
 
 	check(!GUseEmulatedUniformBuffers || (IsValidRef(EmulatedUniformDataRef) && (EmulatedUniformDataRef->Data.Num() * EmulatedUniformDataRef->Data.GetTypeSize() == AllocatedSize)));
 
-	CopyDataToUniformBuffer(bCanCreateOnThisThread, NewUniformBuffer,Contents, Layout.ConstantBufferSize);
+	CopyDataToUniformBuffer(bCanCreateOnThisThread, NewUniformBuffer,Contents, Layout->ConstantBufferSize);
 
 	// Initialize the resource table for this uniform buffer.
 	SetLayoutTable(NewUniformBuffer, Contents, Layout, Validation);
@@ -593,18 +593,18 @@ static FUniformBufferRHIRef CreateUniformBuffer(const void* Contents, const FRHI
 	return NewUniformBuffer;
 }	
 
-FUniformBufferRHIRef FOpenGLDynamicRHI::RHICreateUniformBuffer(const void* Contents, const FRHIUniformBufferLayout& Layout, EUniformBufferUsage Usage, EUniformBufferValidation Validation)
+FUniformBufferRHIRef FOpenGLDynamicRHI::RHICreateUniformBuffer(const void* Contents, const FRHIUniformBufferLayout* Layout, EUniformBufferUsage Usage, EUniformBufferValidation Validation)
 {
 	FRHICommandListImmediate& RHICmdList = FRHICommandListExecutor::GetImmediateCommandList();
 	// This should really be synchronized, if there's a chance it'll be used from more than one buffer. Luckily, uniform buffers
 	// are only used for drawing/shader usage, not for loading resources or framebuffer blitting, so no synchronization primitives for now.
 
 	// Explicitly check that the size is nonzero before allowing CreateBuffer to opaquely fail.
-	check(Layout.Resources.Num() > 0 || Layout.ConstantBufferSize > 0);
+	check(Layout->Resources.Num() > 0 || Layout->ConstantBufferSize > 0);
 
 	if (Validation == EUniformBufferValidation::ValidateResources)
 	{
-		ValidateShaderParameterResourcesRHI(Contents, Layout);
+		ValidateShaderParameterResourcesRHI(Contents, *Layout);
 	}
 
 	bool bStreamDraw = (Usage == UniformBuffer_SingleDraw || Usage == UniformBuffer_SingleFrame);
@@ -617,13 +617,13 @@ FUniformBufferRHIRef FOpenGLDynamicRHI::RHICreateUniformBuffer(const void* Conte
 	const bool bCanCreateOnThisThread = RHICmdList.Bypass() || (!IsRunningRHIInSeparateThread() && IsInRenderingThread()) || IsInRHIThread();
 
 	// If the uniform buffer contains constants, allocate a uniform buffer resource from GL.
-	if (Layout.ConstantBufferSize > 0)
+	if (Layout->ConstantBufferSize > 0)
 	{
 		uint32 SizeOfBufferToAllocate = 0;
 		if (IsPoolingEnabled())
 		{
 			// Find the appropriate bucket based on size
-			const uint32 BucketIndex = GetPoolBucketIndex(Layout.ConstantBufferSize);
+			const uint32 BucketIndex = GetPoolBucketIndex(Layout->ConstantBufferSize);
 			int StreamedIndex = bStreamDraw ? 1 : 0;
 
 			FPooledGLUniformBuffer FreeBufferEntry;
@@ -686,7 +686,7 @@ FUniformBufferRHIRef FOpenGLDynamicRHI::RHICreateUniformBuffer(const void* Conte
 	NewUniformBuffer->SetGLUniformBufferParams(AllocatedResource, OffsetInBuffer, PersistentlyMappedBuffer, AllocatedSize, EmulatedUniformDataRef, bStreamDraw);
 
 	check(!GUseEmulatedUniformBuffers || (IsValidRef(EmulatedUniformDataRef) && (EmulatedUniformDataRef->Data.Num() * EmulatedUniformDataRef->Data.GetTypeSize() == AllocatedSize)));
-	CopyDataToUniformBuffer(bCanCreateOnThisThread, NewUniformBuffer, Contents, Layout.ConstantBufferSize);
+	CopyDataToUniformBuffer(bCanCreateOnThisThread, NewUniformBuffer, Contents, Layout->ConstantBufferSize);
 
 	// Initialize the resource table for this uniform buffer.
 	SetLayoutTable(NewUniformBuffer, Contents, Layout, Validation);
diff --git a/Engine/Source/Runtime/OpenGLDrv/Public/OpenGLDrv.h b/Engine/Source/Runtime/OpenGLDrv/Public/OpenGLDrv.h
--- a/Engine/Source/Runtime/OpenGLDrv/Public/OpenGLDrv.h
+++ b/Engine/Source/Runtime/OpenGLDrv/Public/OpenGLDrv.h
@@ -343,7 +343,7 @@ public:
 	virtual FGeometryShaderRHIRef RHICreateGeometryShader(TArrayView<const uint8> Code, const FSHAHash& Hash) final override;
 	virtual FComputeShaderRHIRef RHICreateComputeShader(TArrayView<const uint8> Code, const FSHAHash& Hash) final override;
 
-	virtual FUniformBufferRHIRef RHICreateUniformBuffer(const void* Contents, const FRHIUniformBufferLayout& Layout, EUniformBufferUsage Usage, EUniformBufferValidation Validation) final override;
+	virtual FUniformBufferRHIRef RHICreateUniformBuffer(const void* Contents, const FRHIUniformBufferLayout* Layout, EUniformBufferUsage Usage, EUniformBufferValidation Validation) final override;
 	virtual void RHIUpdateUniformBuffer(FRHIUniformBuffer* UniformBufferRHI, const void* Contents) final override;
 	virtual FBufferRHIRef RHICreateBuffer(uint32 Size, EBufferUsageFlags Usage, uint32 Stride, ERHIAccess ResourceState, FRHIResourceCreateInfo& CreateInfo) final override;
 	virtual void RHICopyBuffer(FRHIBuffer* SourceBuffer, FRHIBuffer* DestBuffer) final override;
diff --git a/Engine/Source/Runtime/OpenGLDrv/Public/OpenGLResources.h b/Engine/Source/Runtime/OpenGLDrv/Public/OpenGLResources.h
--- a/Engine/Source/Runtime/OpenGLDrv/Public/OpenGLResources.h
+++ b/Engine/Source/Runtime/OpenGLDrv/Public/OpenGLResources.h
@@ -933,7 +933,7 @@ public:
 	bool bStreamDraw;
 
 	/** Initialization constructor. */
-	FOpenGLUniformBuffer(const FRHIUniformBufferLayout& InLayout);
+	FOpenGLUniformBuffer(const FRHIUniformBufferLayout* InLayout);
 
 	void SetGLUniformBufferParams(GLuint InResource, uint32 InOffset, uint8* InPersistentlyMappedBuffer, uint32 InAllocatedSize, FOpenGLEUniformBufferDataRef InEmulatedBuffer, bool bInStreamDraw);
 
diff --git a/Engine/Source/Runtime/RHI/Private/DynamicRHI.cpp b/Engine/Source/Runtime/RHI/Private/DynamicRHI.cpp
--- a/Engine/Source/Runtime/RHI/Private/DynamicRHI.cpp
+++ b/Engine/Source/Runtime/RHI/Private/DynamicRHI.cpp
@@ -219,6 +219,10 @@ void RHIInit(bool bHasEditorToken)
 {
 	if (!GDynamicRHI)
 	{
+#if RHI_WANT_RESOURCE_INFO
+		FRHIResource::StartTrackingAllResources();
+#endif
+
 		// read in any data driven shader platform info structures we can find
 		FGenericDataDrivenShaderPlatformInfo::Initialize();
 
diff --git a/Engine/Source/Runtime/RHI/Private/RHI.cpp b/Engine/Source/Runtime/RHI/Private/RHI.cpp
--- a/Engine/Source/Runtime/RHI/Private/RHI.cpp
+++ b/Engine/Source/Runtime/RHI/Private/RHI.cpp
@@ -52,8 +52,8 @@ DEFINE_STAT(STAT_RTAccelerationStructureMemory);
 DEFINE_STAT(STAT_StructuredBufferMemory);
 DEFINE_STAT(STAT_PixelBufferMemory);
 
-IMPLEMENT_TYPE_LAYOUT(FRHIUniformBufferLayout);
-IMPLEMENT_TYPE_LAYOUT(FRHIUniformBufferLayout::FResourceParameter);
+IMPLEMENT_TYPE_LAYOUT(FRHIUniformBufferLayoutInitializer);
+IMPLEMENT_TYPE_LAYOUT(FRHIUniformBufferResource);
 
 #if !defined(RHIRESOURCE_NUM_FRAMES_TO_EXPIRE)
 	#define RHIRESOURCE_NUM_FRAMES_TO_EXPIRE 3
@@ -635,7 +635,7 @@ bool operator==(const FBlendStateInitializerRHI::FRenderTarget& A, const FBlendS
 
 static FCriticalSection GRHIResourceTrackingCriticalSection;
 static TSet<FRHIResource*> GRHITrackedResources;
-static bool GRHITrackingResources = true;
+static bool GRHITrackingResources = false;
 
 void FRHIResource::BeginTrackingResource(FRHIResource* InResource)
 {
@@ -659,6 +659,11 @@ void FRHIResource::EndTrackingResource(FRHIResource* InResource)
 	}
 }
 
+void FRHIResource::StartTrackingAllResources()
+{
+	GRHITrackingResources = true;
+}
+
 void FRHIResource::StopTrackingAllResources()
 {
 	FScopeLock Lock(&GRHIResourceTrackingCriticalSection);
@@ -699,6 +704,7 @@ static const FRHIResourceTypeName GRHIResourceTypeNames[] =
 	RHI_RESOURCE_TYPE_DEF(ComputePipelineState),
 	RHI_RESOURCE_TYPE_DEF(RayTracingPipelineState),
 	RHI_RESOURCE_TYPE_DEF(BoundShaderState),
+	RHI_RESOURCE_TYPE_DEF(UniformBufferLayout),
 	RHI_RESOURCE_TYPE_DEF(UniformBuffer),
 	RHI_RESOURCE_TYPE_DEF(Buffer),
 	RHI_RESOURCE_TYPE_DEF(Texture),
diff --git a/Engine/Source/Runtime/RHI/Public/DynamicRHI.h b/Engine/Source/Runtime/RHI/Public/DynamicRHI.h
--- a/Engine/Source/Runtime/RHI/Public/DynamicRHI.h
+++ b/Engine/Source/Runtime/RHI/Public/DynamicRHI.h
@@ -567,7 +567,13 @@ public:
 	* @return The new uniform buffer.
 	*/
 	// FlushType: Thread safe, but varies depending on the RHI
-	virtual FUniformBufferRHIRef RHICreateUniformBuffer(const void* Contents, const FRHIUniformBufferLayout& Layout, EUniformBufferUsage Usage, EUniformBufferValidation Validation) = 0;
+	virtual FUniformBufferRHIRef RHICreateUniformBuffer(const void* Contents, const FRHIUniformBufferLayout* Layout, EUniformBufferUsage Usage, EUniformBufferValidation Validation) = 0;
+
+	UE_DEPRECATED(5.0, "Use Layout pointers instead")
+	virtual FUniformBufferRHIRef RHICreateUniformBuffer(const void* Contents, const FRHIUniformBufferLayout& Layout, EUniformBufferUsage Usage, EUniformBufferValidation Validation)
+	{
+		return RHICreateUniformBuffer(Contents, &Layout, Usage, Validation);
+	}
 
 	virtual void RHIUpdateUniformBuffer(FRHIUniformBuffer* UniformBufferRHI, const void* Contents) = 0;
 
@@ -1516,11 +1522,23 @@ FORCEINLINE TRefCountPtr<FRHIRayTracingPipelineState> RHICreateRayTracingPipelin
 }
 #endif //RHI_RAYTRACING
 
-FORCEINLINE FUniformBufferRHIRef RHICreateUniformBuffer(const void* Contents, const FRHIUniformBufferLayout& Layout, EUniformBufferUsage Usage, EUniformBufferValidation Validation = EUniformBufferValidation::ValidateResources)
+FORCEINLINE FUniformBufferLayoutRHIRef RHICreateUniformBufferLayout(const FRHIUniformBufferLayoutInitializer& Initializer)
+{
+	LLM_SCOPE(ELLMTag::RHIMisc);
+	return new FRHIUniformBufferLayout(Initializer);
+}
+
+FORCEINLINE FUniformBufferRHIRef RHICreateUniformBuffer(const void* Contents, const FRHIUniformBufferLayout* Layout, EUniformBufferUsage Usage, EUniformBufferValidation Validation = EUniformBufferValidation::ValidateResources)
 {
 	return GDynamicRHI->RHICreateUniformBuffer(Contents, Layout, Usage, Validation);
 }
 
+UE_DEPRECATED(5.0, "Use Layout pointers instead")
+FORCEINLINE FUniformBufferRHIRef RHICreateUniformBuffer(const void* Contents, const FRHIUniformBufferLayout& Layout, EUniformBufferUsage Usage, EUniformBufferValidation Validation = EUniformBufferValidation::ValidateResources)
+{
+	return GDynamicRHI->RHICreateUniformBuffer(Contents, &Layout, Usage, Validation);
+}
+
 FORCEINLINE void RHIUpdateUniformBuffer(FRHIUniformBuffer* UniformBufferRHI, const void* Contents)
 {
 	return GDynamicRHI->RHIUpdateUniformBuffer(UniformBufferRHI, Contents);
diff --git a/Engine/Source/Runtime/RHI/Public/RHICommandList.h b/Engine/Source/Runtime/RHI/Public/RHICommandList.h
--- a/Engine/Source/Runtime/RHI/Public/RHICommandList.h
+++ b/Engine/Source/Runtime/RHI/Public/RHICommandList.h
@@ -1766,7 +1766,7 @@ struct FComputedUniformBuffer
 struct FLocalUniformBufferWorkArea
 {
 	void* Contents;
-	const FRHIUniformBufferLayout* Layout;
+	FUniformBufferLayoutRHIRef Layout;
 	FComputedUniformBuffer* ComputedUniformBuffer;
 #if DO_CHECK || USING_CODE_ANALYSIS // the below variables are used in check(), which can be enabled in Shipping builds (see Build.h)
 	FRHICommandListBase* CheckCmdList;
@@ -1813,9 +1813,9 @@ FRHICOMMAND_MACRO(FRHICommandBuildLocalUniformBuffer)
 		FRHICommandListBase* CheckCmdList,
 		const void* Contents,
 		uint32 ContentsSize,
-		const FRHIUniformBufferLayout& Layout
+		const FRHIUniformBufferLayout* Layout
 		)
-		: WorkArea(CheckCmdList, Contents, ContentsSize, &Layout)
+		: WorkArea(CheckCmdList, Contents, ContentsSize, Layout)
 
 	{
 	}
@@ -2421,7 +2421,7 @@ public:
 		SetShaderUniformBuffer(Shader.GetReference(), BaseIndex, UniformBuffer);
 	}
 
-	FORCEINLINE_DEBUGGABLE FLocalUniformBuffer BuildLocalUniformBuffer(const void* Contents, uint32 ContentsSize, const FRHIUniformBufferLayout& Layout)
+	FORCEINLINE_DEBUGGABLE FLocalUniformBuffer BuildLocalUniformBuffer(const void* Contents, uint32 ContentsSize, const FRHIUniformBufferLayout* Layout)
 	{
 		FLocalUniformBuffer Result;
 		if (Bypass())
@@ -2437,6 +2437,12 @@ public:
 		return Result;
 	}
 
+	UE_DEPRECATED(5.0, "Use Layout pointers instead")
+	FORCEINLINE_DEBUGGABLE FLocalUniformBuffer BuildLocalUniformBuffer(const void* Contents, uint32 ContentsSize, const FRHIUniformBufferLayout& Layout)
+	{
+		return BuildLocalUniformBuffer(Contents, ContentsSize, &Layout);
+	}
+
 	template <typename TRHIShader>
 	FORCEINLINE_DEBUGGABLE void SetLocalShaderUniformBuffer(TRHIShader* Shader, uint32 BaseIndex, const FLocalUniformBuffer& UniformBuffer)
 	{
@@ -4164,10 +4170,16 @@ public:
 		return RHICreateComputePipelineState(ComputeShader);
 	}
 
-	FORCEINLINE FUniformBufferRHIRef CreateUniformBuffer(const void* Contents, const FRHIUniformBufferLayout& Layout, EUniformBufferUsage Usage)
+	FORCEINLINE FUniformBufferRHIRef CreateUniformBuffer(const void* Contents, const FRHIUniformBufferLayout* Layout, EUniformBufferUsage Usage)
 	{
 		return RHICreateUniformBuffer(Contents, Layout, Usage);
 	}
+
+	UE_DEPRECATED(5.0, "Use Layout pointers instead")
+	FORCEINLINE FUniformBufferRHIRef CreateUniformBuffer(const void* Contents, const FRHIUniformBufferLayout& Layout, EUniformBufferUsage Usage)
+	{
+		return RHICreateUniformBuffer(Contents, &Layout, Usage);
+	}
 	
 	FORCEINLINE FBufferRHIRef CreateAndLockIndexBuffer(uint32 Stride, uint32 Size, EBufferUsageFlags InUsage, ERHIAccess InResourceState, FRHIResourceCreateInfo& CreateInfo, void*& OutDataBuffer)
 	{
diff --git a/Engine/Source/Runtime/RHI/Public/RHICommandListCommandExecutes.inl b/Engine/Source/Runtime/RHI/Public/RHICommandListCommandExecutes.inl
--- a/Engine/Source/Runtime/RHI/Public/RHICommandListCommandExecutes.inl
+++ b/Engine/Source/Runtime/RHI/Public/RHICommandListCommandExecutes.inl
@@ -462,7 +462,7 @@ void FRHICommandBuildLocalUniformBuffer::Execute(FRHICommandListBase& CmdList)
 	check(WorkArea.Contents); 
 	if (WorkArea.ComputedUniformBuffer->UseCount)
 	{
-		WorkArea.ComputedUniformBuffer->UniformBuffer = GDynamicRHI->RHICreateUniformBuffer(WorkArea.Contents, *WorkArea.Layout, UniformBuffer_SingleFrame, EUniformBufferValidation::ValidateResources);
+		WorkArea.ComputedUniformBuffer->UniformBuffer = GDynamicRHI->RHICreateUniformBuffer(WorkArea.Contents, WorkArea.Layout, UniformBuffer_SingleFrame, EUniformBufferValidation::ValidateResources);
 	}
 	WorkArea.Layout = nullptr;
 	WorkArea.Contents = nullptr;
diff --git a/Engine/Source/Runtime/RHI/Public/RHIDefinitions.h b/Engine/Source/Runtime/RHI/Public/RHIDefinitions.h
--- a/Engine/Source/Runtime/RHI/Public/RHIDefinitions.h
+++ b/Engine/Source/Runtime/RHI/Public/RHIDefinitions.h
@@ -1520,6 +1520,7 @@ enum ERHIResourceType
 	RRT_ComputePipelineState,
 	RRT_RayTracingPipelineState,
 	RRT_BoundShaderState,
+	RRT_UniformBufferLayout,
 	RRT_UniformBuffer,
 	RRT_Buffer,
 	RRT_Texture,
diff --git a/Engine/Source/Runtime/RHI/Public/RHIResources.h b/Engine/Source/Runtime/RHI/Public/RHIResources.h
--- a/Engine/Source/Runtime/RHI/Public/RHIResources.h
+++ b/Engine/Source/Runtime/RHI/Public/RHIResources.h
@@ -148,6 +148,7 @@ public:
 
 	static void BeginTrackingResource(FRHIResource* InResource);
 	static void EndTrackingResource(FRHIResource* InResource);
+	static void StartTrackingAllResources();
 	static void StopTrackingAllResources();
 #endif
 
@@ -792,34 +793,52 @@ public:
 // Enabling this requires -norhithread to work correctly since FRHIResource lifetime is managed by both the RT and RHIThread
 #define VALIDATE_UNIFORM_BUFFER_LIFETIME 0
 
-/** The layout of a uniform buffer in memory. */
-struct FRHIUniformBufferLayout
+/** Data structure to store information about resource parameter in a shader parameter structure. */
+struct FRHIUniformBufferResource
 {
-public:
-	static const uint16 kInvalidOffset = TNumericLimits<uint16>::Max();
+	DECLARE_EXPORTED_TYPE_LAYOUT(FRHIUniformBufferResource, RHI_API, NonVirtual);
 
-	/** Data structure to store information about resource parameter in a shader parameter structure. */
-	struct FResourceParameter
-	{
-		DECLARE_EXPORTED_TYPE_LAYOUT(FResourceParameter, RHI_API, NonVirtual);
-	public:
-		friend inline FArchive& operator<<(FArchive& Ar, FResourceParameter& Ref)
-		{
-			uint8 Type = (uint8)Ref.MemberType;
-			Ar << Ref.MemberOffset;
-			Ar << Type;
-			Ref.MemberType = (EUniformBufferBaseType)Type;
-			return Ar;
-		}
+	/** Byte offset to each resource in the uniform buffer memory. */
+	LAYOUT_FIELD(uint16, MemberOffset);
 
-		/** Byte offset to each resource in the uniform buffer memory. */
-		LAYOUT_FIELD(uint16, MemberOffset);
-		/** Type of the member that allow (). */
-		LAYOUT_FIELD(EUniformBufferBaseType, MemberType);
-	};
+	/** Type of the member that allow (). */
+	LAYOUT_FIELD(EUniformBufferBaseType, MemberType);
+};
 
-	DECLARE_EXPORTED_TYPE_LAYOUT(FRHIUniformBufferLayout, RHI_API, NonVirtual);
-public:
+inline FArchive& operator<<(FArchive& Ar, FRHIUniformBufferResource& Ref)
+{
+	uint8 Type = (uint8)Ref.MemberType;
+	Ar << Ref.MemberOffset;
+	Ar << Type;
+	Ref.MemberType = (EUniformBufferBaseType)Type;
+	return Ar;
+}
+
+/** Compare two uniform buffer layout resources. */
+inline bool operator==(const FRHIUniformBufferResource& A, const FRHIUniformBufferResource& B)
+{
+	return A.MemberOffset == B.MemberOffset
+		&& A.MemberType == B.MemberType;
+}
+
+static constexpr uint16 kUniformBufferInvalidOffset = TNumericLimits<uint16>::Max();
+
+/** Initializer for the layout of a uniform buffer in memory. */
+struct FRHIUniformBufferLayoutInitializer
+{
+	DECLARE_EXPORTED_TYPE_LAYOUT(FRHIUniformBufferLayoutInitializer, RHI_API, NonVirtual);
+
+	FRHIUniformBufferLayoutInitializer() = default;
+
+	explicit FRHIUniformBufferLayoutInitializer(const TCHAR * InName)
+		: Name(InName)
+	{}
+	explicit FRHIUniformBufferLayoutInitializer(const TCHAR * InName, uint32 InConstantBufferSize)
+		: Name(InName)
+		, ConstantBufferSize(InConstantBufferSize)
+	{
+		ComputeHash();
+	}
 
 	inline uint32 GetHash() const
 	{
@@ -859,20 +878,7 @@ public:
 		Hash = TmpHash;
 	}
 
-	FRHIUniformBufferLayout() = default;
-
-	explicit FRHIUniformBufferLayout(const TCHAR* InName)
-		: Name(InName)
-	{}
-
-#if VALIDATE_UNIFORM_BUFFER_LAYOUT_LIFETIME
-	~FRHIUniformBufferLayout()
-	{
-		check(NumUsesForDebugging == 0 || IsEngineExitRequested());
-	}
-#endif
-
-	void CopyFrom(const FRHIUniformBufferLayout& Source)
+	void CopyFrom(const FRHIUniformBufferLayoutInitializer& Source)
 	{
 		ConstantBufferSize = Source.ConstantBufferSize;
 		StaticSlot = Source.StaticSlot;
@@ -889,7 +895,7 @@ public:
 
 	bool HasRenderTargets() const
 	{
-		return RenderTargetsOffset != kInvalidOffset;
+		return RenderTargetsOffset != kUniformBufferInvalidOffset;
 	}
 
 	bool HasExternalOutputs() const
@@ -902,7 +908,7 @@ public:
 		return IsUniformBufferStaticSlotValid(StaticSlot);
 	}
 
-	friend FArchive& operator<<(FArchive& Ar, FRHIUniformBufferLayout& Ref)
+	friend FArchive& operator<<(FArchive& Ar, FRHIUniformBufferLayoutInitializer& Ref)
 	{
 		Ar << Ref.ConstantBufferSize;
 		Ar << Ref.StaticSlot;
@@ -920,62 +926,143 @@ public:
 		return Ar;
 	}
 
+private:
+	// for debugging / error message
+	LAYOUT_FIELD(FMemoryImageString, Name);
+
+public:
+	/** The list of all resource inlined into the shader parameter structure. */
+	LAYOUT_FIELD(TMemoryImageArray<FRHIUniformBufferResource>, Resources);
+
+	/** The list of all RDG resource references inlined into the shader parameter structure. */
+	LAYOUT_FIELD(TMemoryImageArray<FRHIUniformBufferResource>, GraphResources);
+
+	/** The list of all RDG texture references inlined into the shader parameter structure. */
+	LAYOUT_FIELD(TMemoryImageArray<FRHIUniformBufferResource>, GraphTextures);
+
+	/** The list of all RDG buffer references inlined into the shader parameter structure. */
+	LAYOUT_FIELD(TMemoryImageArray<FRHIUniformBufferResource>, GraphBuffers);
+
+	/** The list of all RDG uniform buffer references inlined into the shader parameter structure. */
+	LAYOUT_FIELD(TMemoryImageArray<FRHIUniformBufferResource>, GraphUniformBuffers);
+
+	/** The list of all non-RDG uniform buffer references inlined into the shader parameter structure. */
+	LAYOUT_FIELD(TMemoryImageArray<FRHIUniformBufferResource>, UniformBuffers);
+
+private:
+	LAYOUT_FIELD_INITIALIZED(uint32, Hash, 0);
+
+public:
 	/** The size of the constant buffer in bytes. */
 	LAYOUT_FIELD_INITIALIZED(uint32, ConstantBufferSize, 0);
 
+	/** The render target binding slots offset, if it exists. */
+	LAYOUT_FIELD_INITIALIZED(uint16, RenderTargetsOffset, kUniformBufferInvalidOffset);
+
 	/** The static slot (if applicable). */
 	LAYOUT_FIELD_INITIALIZED(FUniformBufferStaticSlot, StaticSlot, MAX_UNIFORM_BUFFER_STATIC_SLOTS);
 
-	/** The render target binding slots offset, if it exists. */
-	LAYOUT_FIELD_INITIALIZED(uint16, RenderTargetsOffset, kInvalidOffset);
+	/** The binding flags describing how this resource can be bound to the RHI. */
+	LAYOUT_FIELD_INITIALIZED(EUniformBufferBindingFlags, BindingFlags, EUniformBufferBindingFlags::Shader);
 
 	/** Whether this layout may contain non-render-graph outputs (e.g. RHI UAVs). */
 	LAYOUT_FIELD_INITIALIZED(bool, bHasNonGraphOutputs, false);
+};
 
-	/** The binding flags describing how this resource can be bound to the RHI. */
-	LAYOUT_FIELD_INITIALIZED(EUniformBufferBindingFlags, BindingFlags, EUniformBufferBindingFlags::Shader);
+/** Compare two uniform buffer layout initializers. */
+inline bool operator==(const FRHIUniformBufferLayoutInitializer& A, const FRHIUniformBufferLayoutInitializer& B)
+{
+	return A.ConstantBufferSize == B.ConstantBufferSize
+		&& A.StaticSlot == B.StaticSlot
+		&& A.BindingFlags == B.BindingFlags
+		&& A.Resources == B.Resources;
+}
+
+/** The layout of a uniform buffer in memory. */
+struct FRHIUniformBufferLayout : public FRHIResource
+{
+	FRHIUniformBufferLayout() = delete;
+
+	explicit FRHIUniformBufferLayout(const FRHIUniformBufferLayoutInitializer& Initializer)
+		: FRHIResource(RRT_UniformBufferLayout)
+		, Name(Initializer.GetDebugName())
+		, Resources(Initializer.Resources)
+		, GraphResources(Initializer.GraphResources)
+		, GraphTextures(Initializer.GraphTextures)
+		, GraphBuffers(Initializer.GraphBuffers)
+		, GraphUniformBuffers(Initializer.GraphUniformBuffers)
+		, UniformBuffers(Initializer.UniformBuffers)
+		, Hash(Initializer.GetHash())
+		, ConstantBufferSize(Initializer.ConstantBufferSize)
+		, RenderTargetsOffset(Initializer.RenderTargetsOffset)
+		, StaticSlot(Initializer.StaticSlot)
+		, BindingFlags(Initializer.BindingFlags)
+		, bHasNonGraphOutputs(Initializer.bHasNonGraphOutputs)
+	{}
+
+	inline const FString& GetDebugName() const
+	{
+		return Name;
+	}
+
+	inline uint32 GetHash() const
+	{
+		checkSlow(Hash != 0);
+		return Hash;
+	}
+
+	inline bool HasRenderTargets() const
+	{
+		return RenderTargetsOffset != kUniformBufferInvalidOffset;
+	}
+
+	inline bool HasExternalOutputs() const
+	{
+		return bHasNonGraphOutputs;
+	}
+
+	inline bool HasStaticSlot() const
+	{
+		return IsUniformBufferStaticSlotValid(StaticSlot);
+	}
+
+	const FString Name;
 
 	/** The list of all resource inlined into the shader parameter structure. */
-	LAYOUT_FIELD(TMemoryImageArray<FResourceParameter>, Resources);
+	const TArray<FRHIUniformBufferResource> Resources;
 
 	/** The list of all RDG resource references inlined into the shader parameter structure. */
-	LAYOUT_FIELD(TMemoryImageArray<FResourceParameter>, GraphResources);
+	const TArray<FRHIUniformBufferResource> GraphResources;
 
 	/** The list of all RDG texture references inlined into the shader parameter structure. */
-	LAYOUT_FIELD(TMemoryImageArray<FResourceParameter>, GraphTextures);
+	const TArray<FRHIUniformBufferResource> GraphTextures;
 
 	/** The list of all RDG buffer references inlined into the shader parameter structure. */
-	LAYOUT_FIELD(TMemoryImageArray<FResourceParameter>, GraphBuffers);
+	const TArray<FRHIUniformBufferResource> GraphBuffers;
 
 	/** The list of all RDG uniform buffer references inlined into the shader parameter structure. */
-	LAYOUT_FIELD(TMemoryImageArray<FResourceParameter>, GraphUniformBuffers);
+	const TArray<FRHIUniformBufferResource> GraphUniformBuffers;
 
 	/** The list of all non-RDG uniform buffer references inlined into the shader parameter structure. */
-	LAYOUT_FIELD(TMemoryImageArray<FResourceParameter>, UniformBuffers);
+	const TArray<FRHIUniformBufferResource> UniformBuffers;
 
-	LAYOUT_MUTABLE_FIELD_INITIALIZED(int32, NumUsesForDebugging, 0);
+	const uint32 Hash;
 
-private:
-	// for debugging / error message
-	LAYOUT_FIELD(FMemoryImageString, Name);
-	LAYOUT_FIELD_INITIALIZED(uint32, Hash, 0);
-};
+	/** The size of the constant buffer in bytes. */
+	const uint32 ConstantBufferSize;
 
-/** Compare two uniform buffer layouts. */
-inline bool operator==(const FRHIUniformBufferLayout::FResourceParameter& A, const FRHIUniformBufferLayout::FResourceParameter& B)
-{
-	return A.MemberOffset == B.MemberOffset
-		&& A.MemberType == B.MemberType;
-}
+	/** The render target binding slots offset, if it exists. */
+	const uint16 RenderTargetsOffset;
 
-/** Compare two uniform buffer layouts. */
-inline bool operator==(const FRHIUniformBufferLayout& A, const FRHIUniformBufferLayout& B)
-{
-	return A.ConstantBufferSize == B.ConstantBufferSize
-		&& A.StaticSlot == B.StaticSlot
-		&& A.BindingFlags == B.BindingFlags
-		&& A.Resources == B.Resources;
-}
+	/** The static slot (if applicable). */
+	const FUniformBufferStaticSlot StaticSlot;
+
+	/** The binding flags describing how this resource can be bound to the RHI. */
+	const EUniformBufferBindingFlags BindingFlags;
+
+	/** Whether this layout may contain non-render-graph outputs (e.g. RHI UAVs). */
+	const bool bHasNonGraphOutputs;
+};
 
 class FRHIUniformBuffer : public FRHIResource
 #if ENABLE_RHI_VALIDATION
@@ -983,27 +1070,15 @@ class FRHIUniformBuffer : public FRHIResource
 #endif
 {
 public:
-
 	FRHIUniformBuffer() = delete;
 
 	/** Initialization constructor. */
-	FRHIUniformBuffer(const FRHIUniformBufferLayout& InLayout)
+	FRHIUniformBuffer(const FRHIUniformBufferLayout* InLayout)
 	: FRHIResource(RRT_UniformBuffer)
-	, Layout(&InLayout)
-	, LayoutConstantBufferSize(InLayout.ConstantBufferSize)
+	, Layout(InLayout)
+	, LayoutConstantBufferSize(InLayout->ConstantBufferSize)
 	{}
 
-	FORCEINLINE_DEBUGGABLE uint32 AddRef() const
-	{
-#if VALIDATE_UNIFORM_BUFFER_LAYOUT_LIFETIME
-		if (GetRefCount() == 0)
-		{
-			Layout->NumUsesForDebugging++;
-		}
-#endif
-		return FRHIResource::AddRef();
-	}
-
 	FORCEINLINE_DEBUGGABLE uint32 Release() const
 	{
 		const FRHIUniformBufferLayout* LocalLayout = Layout;
@@ -1016,10 +1091,6 @@ public:
 
 		if (NewRefCount == 0)
 		{
-#if VALIDATE_UNIFORM_BUFFER_LAYOUT_LIFETIME
-			LocalLayout->NumUsesForDebugging--;
-			check(LocalLayout->NumUsesForDebugging >= 0);
-#endif
 #if VALIDATE_UNIFORM_BUFFER_LIFETIME
 			check(LocalNumMeshCommandReferencesForDebugging == 0 || IsEngineExitRequested());
 #endif
@@ -1035,11 +1106,7 @@ public:
 		return LayoutConstantBufferSize;
 	}
 	const FRHIUniformBufferLayout& GetLayout() const { return *Layout; }
-
-	bool IsGlobal() const
-	{
-		return IsUniformBufferStaticSlotValid(Layout->StaticSlot);
-	}
+	const FRHIUniformBufferLayout* GetLayoutPtr() const { return Layout; }
 
 #if VALIDATE_UNIFORM_BUFFER_LIFETIME
 	mutable int32 NumMeshCommandReferencesForDebugging = 0;
@@ -1047,7 +1114,7 @@ public:
 
 private:
 	/** Layout of the uniform buffer. */
-	const FRHIUniformBufferLayout* Layout;
+	TRefCountPtr<const FRHIUniformBufferLayout> Layout;
 
 	uint32 LayoutConstantBufferSize;
 };
@@ -1783,6 +1850,7 @@ typedef TRefCountPtr<FRHIComputeShader> FComputeShaderRHIRef;
 typedef TRefCountPtr<FRHIRayTracingShader>          FRayTracingShaderRHIRef;
 typedef TRefCountPtr<FRHIComputeFence>	FComputeFenceRHIRef;
 typedef TRefCountPtr<FRHIBoundShaderState> FBoundShaderStateRHIRef;
+typedef TRefCountPtr<const FRHIUniformBufferLayout> FUniformBufferLayoutRHIRef;
 typedef TRefCountPtr<FRHIUniformBuffer> FUniformBufferRHIRef;
 typedef TRefCountPtr<FRHIBuffer> FBufferRHIRef;
 UE_DEPRECATED(5.0, "FIndexBufferRHIRef is deprecated, please use FBufferRHIRef.")      typedef FBufferRHIRef FIndexBufferRHIRef;
diff --git a/Engine/Source/Runtime/RHI/Public/RHIValidation.h b/Engine/Source/Runtime/RHI/Public/RHIValidation.h
--- a/Engine/Source/Runtime/RHI/Public/RHIValidation.h
+++ b/Engine/Source/Runtime/RHI/Public/RHIValidation.h
@@ -341,10 +341,10 @@ public:
 	* @return The new uniform buffer.
 	*/
 	// FlushType: Thread safe, but varies depending on the RHI override final
-	virtual FUniformBufferRHIRef RHICreateUniformBuffer(const void* Contents, const FRHIUniformBufferLayout& Layout, EUniformBufferUsage Usage, EUniformBufferValidation Validation) override final
+	virtual FUniformBufferRHIRef RHICreateUniformBuffer(const void* Contents, const FRHIUniformBufferLayout* Layout, EUniformBufferUsage Usage, EUniformBufferValidation Validation) override final
 	{
 		// RHICreateUniformBuffer can be called from the RenderThread (normal use case) or the RHI Thread (FLocalUniformBuffer)
-		check(Layout.Resources.Num() > 0 || Layout.ConstantBufferSize > 0);
+		check(Layout->Resources.Num() > 0 || Layout->ConstantBufferSize > 0);
 		FUniformBufferRHIRef UniformBuffer = RHI->RHICreateUniformBuffer(Contents, Layout, Usage, Validation);
 		UniformBuffer->InitLifetimeTracking(IsInRenderingThread() ? RenderThreadFrameID : RHIThreadFrameID, Usage);
 		return UniformBuffer;
diff --git a/Engine/Source/Runtime/RenderCore/Private/RenderGraphResources.cpp b/Engine/Source/Runtime/RenderCore/Private/RenderGraphResources.cpp
--- a/Engine/Source/Runtime/RenderCore/Private/RenderGraphResources.cpp
+++ b/Engine/Source/Runtime/RenderCore/Private/RenderGraphResources.cpp
@@ -168,8 +168,8 @@ void FRDGUniformBuffer::InitRHI()
 		EUniformBufferValidation::None;
 #endif
 
-	const FRDGParameterStruct PassParameters = GetParameters();
-	UniformBufferRHI = RHICreateUniformBuffer(PassParameters.GetContents(), PassParameters.GetLayout(), UniformBuffer_SingleFrame, Validation);
+	const FRDGParameterStruct& PassParameters = GetParameters();
+	UniformBufferRHI = RHICreateUniformBuffer(PassParameters.GetContents(), PassParameters.GetLayoutPtr(), UniformBuffer_SingleFrame, Validation);
 	ResourceRHI = UniformBufferRHI;
 }
 
diff --git a/Engine/Source/Runtime/RenderCore/Private/ShaderMap.cpp b/Engine/Source/Runtime/RenderCore/Private/ShaderMap.cpp
--- a/Engine/Source/Runtime/RenderCore/Private/ShaderMap.cpp
+++ b/Engine/Source/Runtime/RenderCore/Private/ShaderMap.cpp
@@ -60,6 +60,8 @@ void FShaderMapBase::AssignContent(TMemoryImageObject<FShaderMapContent> InConte
 
 	Content = InContent;
 	PointerTable = CreatePointerTable();
+
+	PostFinalizeContent();
 }
 
 void FShaderMapBase::AssignCopy(const FShaderMapBase& Source)
@@ -99,6 +101,7 @@ void FShaderMapBase::InitResource()
 		Resource = new FShaderMapResource_InlineCode(GetShaderPlatform(), Code);
 		BeginInitResource(Resource);
 	}
+	PostFinalizeContent();
 }
 
 void FShaderMapBase::FinalizeContent()
@@ -224,6 +227,7 @@ bool FShaderMapBase::Serialize(FArchive& Ar, bool bInlineShaderResources, bool b
 			{
 				Content.Freeze(PointerTable);
 			}
+			PostFinalizeContent();
 
 			NumFrozenShaders = Content.Object->GetNumShaders();
 			INC_DWORD_STAT_BY(STAT_Shaders_ShaderMemory, Content.FrozenSize);
diff --git a/Engine/Source/Runtime/RenderCore/Private/ShaderParameterMetadata.cpp b/Engine/Source/Runtime/RenderCore/Private/ShaderParameterMetadata.cpp
--- a/Engine/Source/Runtime/RenderCore/Private/ShaderParameterMetadata.cpp
+++ b/Engine/Source/Runtime/RenderCore/Private/ShaderParameterMetadata.cpp
@@ -292,8 +292,10 @@ FShaderParametersMetadata::FShaderParametersMetadata(
 	const int32 InFileLine,
 	uint32 InSize,
 	const TArray<FMember>& InMembers,
-	bool bForceCompleteInitialization)
-	: StructTypeName(InStructTypeName)
+	bool bForceCompleteInitialization,
+	FRHIUniformBufferLayoutInitializer* OutLayoutInitializer)
+	: LayoutName(InLayoutName)
+	, StructTypeName(InStructTypeName)
 	, ShaderVariableName(InShaderVariableName)
 	, StaticSlotName(InStaticSlotName)
 	, ShaderVariableHashedName(InShaderVariableName)
@@ -302,10 +304,8 @@ FShaderParametersMetadata::FShaderParametersMetadata(
 	, Size(InSize)
 	, UseCase(InUseCase)
 	, BindingFlags(InBindingFlags)
-	, Layout(InLayoutName)
 	, Members(InMembers)
 	, GlobalListLink(this)
-	, bLayoutInitialized(false)
 {
 	checkf(UseCase == EUseCase::UniformBuffer || !EnumHasAnyFlags(BindingFlags, EUniformBufferBindingFlags::Static), TEXT("Only uniform buffers can utilize the global binding flag."));
 
@@ -365,7 +365,7 @@ FShaderParametersMetadata::FShaderParametersMetadata(
 		// We cannot initialize the layout during global initialization, since we have to walk nested struct members.
 		// Structs created during global initialization will have bRegisterForAutoBinding==false, and are initialized during startup.
 		// Structs created at runtime with bRegisterForAutoBinding==true can be initialized now.
-		InitializeLayout();
+		InitializeLayout(OutLayoutInitializer);
 	}
 }
 
@@ -386,7 +386,7 @@ FShaderParametersMetadata::~FShaderParametersMetadata()
 		GetGlobalShaderVariableToStructMap().Remove(FName(ShaderVariableName, FNAME_Find));
 #endif
 
-		if (bLayoutInitialized)
+		if (IsLayoutInitialized())
 		{
 			GetLayoutHashStructMap().Remove(GetLayout().GetHash());
 		}
@@ -398,17 +398,20 @@ void FShaderParametersMetadata::InitializeAllUniformBufferStructs()
 {
 	for (TLinkedList<FShaderParametersMetadata*>::TIterator StructIt(FShaderParametersMetadata::GetStructList()); StructIt; StructIt.Next())
 	{
-		if (!StructIt->bLayoutInitialized)
+		if (!StructIt->IsLayoutInitialized())
 		{
 			StructIt->InitializeLayout();
 		}
 	}
 }
 
-void FShaderParametersMetadata::InitializeLayout()
+void FShaderParametersMetadata::InitializeLayout(FRHIUniformBufferLayoutInitializer* OutLayoutInitializer)
 {
-	check(!bLayoutInitialized);
-	Layout.ConstantBufferSize = Size;
+	check(!IsLayoutInitialized());
+
+	FRHIUniformBufferLayoutInitializer LocalLayoutInitializer(LayoutName);
+	FRHIUniformBufferLayoutInitializer& LayoutInitializer = OutLayoutInitializer ? *OutLayoutInitializer : LocalLayoutInitializer;
+	LayoutInitializer.ConstantBufferSize = Size;
 
 	if (StaticSlotName)
 	{
@@ -425,8 +428,8 @@ void FShaderParametersMetadata::InitializeLayout()
 			TEXT("Uniform buffer of type '%s' and shader name '%s' attempted to reference static slot '%s', but the slot could not be found in the registry."),
 			StructTypeName, ShaderVariableName, StaticSlotName);
 
-		Layout.StaticSlot = StaticSlot;
-		Layout.BindingFlags = BindingFlags;
+		LayoutInitializer.StaticSlot = StaticSlot;
+		LayoutInitializer.BindingFlags = BindingFlags;
 	}
 	else
 	{
@@ -619,42 +622,42 @@ void FShaderParametersMetadata::InitializeLayout()
 			for (uint32 ArrayElementId = 0; ArrayElementId < (bIsArray ? ArraySize : 1u); ArrayElementId++)
 			{
 				const uint32 AbsoluteMemberOffset = CurrentMember.GetOffset() + MemberStack[i].StructOffset + ArrayElementId * SHADER_PARAMETER_POINTER_ALIGNMENT;
-				check(AbsoluteMemberOffset < (1u << (sizeof(FRHIUniformBufferLayout::FResourceParameter::MemberOffset) * 8)));
-				const FRHIUniformBufferLayout::FResourceParameter ResourceParameter{ uint16(AbsoluteMemberOffset), BaseType };
+				check(AbsoluteMemberOffset < (1u << (sizeof(FRHIUniformBufferResource::MemberOffset) * 8)));
+				const FRHIUniformBufferResource ResourceParameter{ uint16(AbsoluteMemberOffset), BaseType };
 
-				Layout.Resources.Add(ResourceParameter);
+				LayoutInitializer.Resources.Add(ResourceParameter);
 
 				if (IsRDGTextureReferenceShaderParameterType(BaseType) || BaseType == UBMT_RENDER_TARGET_BINDING_SLOTS)
 				{
-					Layout.GraphResources.Add(ResourceParameter);
-					Layout.GraphTextures.Add(ResourceParameter);
+					LayoutInitializer.GraphResources.Add(ResourceParameter);
+					LayoutInitializer.GraphTextures.Add(ResourceParameter);
 
 					if (BaseType == UBMT_RENDER_TARGET_BINDING_SLOTS)
 					{
-						checkf(!Layout.HasRenderTargets(), TEXT("Shader parameter struct %s has multiple render target binding slots."), GetStructTypeName());
-						Layout.RenderTargetsOffset = ResourceParameter.MemberOffset;
+						checkf(!LayoutInitializer.HasRenderTargets(), TEXT("Shader parameter struct %s has multiple render target binding slots."), GetStructTypeName());
+						LayoutInitializer.RenderTargetsOffset = ResourceParameter.MemberOffset;
 					}
 				}
 				else if (IsRDGBufferReferenceShaderParameterType(BaseType))
 				{
-					Layout.GraphResources.Add(ResourceParameter);
-					Layout.GraphBuffers.Add(ResourceParameter);
+					LayoutInitializer.GraphResources.Add(ResourceParameter);
+					LayoutInitializer.GraphBuffers.Add(ResourceParameter);
 				}
 				else if (BaseType == UBMT_RDG_UNIFORM_BUFFER)
 				{
-					Layout.GraphResources.Add(ResourceParameter);
-					Layout.GraphUniformBuffers.Add(ResourceParameter);
+					LayoutInitializer.GraphResources.Add(ResourceParameter);
+					LayoutInitializer.GraphUniformBuffers.Add(ResourceParameter);
 				}
 				else if (BaseType == UBMT_REFERENCED_STRUCT)
 				{
-					Layout.UniformBuffers.Add(ResourceParameter);
+					LayoutInitializer.UniformBuffers.Add(ResourceParameter);
 				}
 			}
 		}
 
 		if (BaseType == UBMT_UAV)
 		{
-			Layout.bHasNonGraphOutputs = true;
+			LayoutInitializer.bHasNonGraphOutputs = true;
 		}
 		else if (BaseType == UBMT_REFERENCED_STRUCT || BaseType == UBMT_RDG_UNIFORM_BUFFER)
 		{
@@ -664,7 +667,7 @@ void FShaderParametersMetadata::InitializeLayout()
 				{
 					if (Member.GetBaseType() == UBMT_UAV)
 					{
-						Layout.bHasNonGraphOutputs = true;
+						LayoutInitializer.bHasNonGraphOutputs = true;
 					}
 				}
 			}
@@ -686,15 +689,15 @@ void FShaderParametersMetadata::InitializeLayout()
 	}
 
 	const auto ByMemberOffset = [](
-		const FRHIUniformBufferLayout::FResourceParameter& A,
-		const FRHIUniformBufferLayout::FResourceParameter& B)
+		const FRHIUniformBufferResource& A,
+		const FRHIUniformBufferResource& B)
 	{
 		return A.MemberOffset < B.MemberOffset;
 	};
 
 	const auto ByTypeThenMemberOffset = [](
-		const FRHIUniformBufferLayout::FResourceParameter& A,
-		const FRHIUniformBufferLayout::FResourceParameter& B)
+		const FRHIUniformBufferResource& A,
+		const FRHIUniformBufferResource& B)
 	{
 		if (A.MemberType == B.MemberType)
 		{
@@ -703,15 +706,15 @@ void FShaderParametersMetadata::InitializeLayout()
 		return A.MemberType < B.MemberType;
 	};
 
-	Layout.Resources.Sort(ByMemberOffset);
-	Layout.GraphResources.Sort(ByMemberOffset);
-	Layout.GraphTextures.Sort(ByTypeThenMemberOffset);
-	Layout.GraphBuffers.Sort(ByTypeThenMemberOffset);
-	Layout.GraphUniformBuffers.Sort(ByMemberOffset);
-	Layout.UniformBuffers.Sort(ByMemberOffset);
+	LayoutInitializer.Resources.Sort(ByMemberOffset);
+	LayoutInitializer.GraphResources.Sort(ByMemberOffset);
+	LayoutInitializer.GraphTextures.Sort(ByTypeThenMemberOffset);
+	LayoutInitializer.GraphBuffers.Sort(ByTypeThenMemberOffset);
+	LayoutInitializer.GraphUniformBuffers.Sort(ByMemberOffset);
+	LayoutInitializer.UniformBuffers.Sort(ByMemberOffset);
 
 	// Compute the hash of the RHI layout.
-	Layout.ComputeHash();
+	LayoutInitializer.ComputeHash();
 	
 	// Compute the hash about the entire layout of the structure.
 	{
@@ -745,7 +748,7 @@ void FShaderParametersMetadata::InitializeLayout()
 			}
 			else if (BaseType == UBMT_INCLUDED_STRUCT || BaseType == UBMT_NESTED_STRUCT)
 			{
-				if (!ChildStruct->bLayoutInitialized)
+				if (!ChildStruct->IsLayoutInitialized())
 				{
 					const_cast<FShaderParametersMetadata*>(ChildStruct)->InitializeLayout();
 				}
@@ -765,10 +768,10 @@ void FShaderParametersMetadata::InitializeLayout()
 
 	if (UseCase == EUseCase::UniformBuffer)
 	{
-		GetLayoutHashStructMap().Emplace(Layout.GetHash(), this);
+		GetLayoutHashStructMap().Emplace(LayoutInitializer.GetHash(), this);
 	}
 
-	bLayoutInitialized = true;
+	Layout = new FRHIUniformBufferLayout(LayoutInitializer);
 }
 
 void FShaderParametersMetadata::GetNestedStructs(TArray<const FShaderParametersMetadata*>& OutNestedStructs) const
@@ -795,7 +798,7 @@ void FShaderParametersMetadata::AddResourceTableEntries(TMap<FString, FResourceT
 	
 	FUniformBufferEntry UniformBufferEntry;
 	UniformBufferEntry.StaticSlotName = StaticSlotName;
-	UniformBufferEntry.LayoutHash = Layout.GetHash();
+	UniformBufferEntry.LayoutHash = GetLayout().GetHash();
 	UniformBufferEntry.BindingFlags = BindingFlags;
 	UniformBufferMap.Add(ShaderVariableName, UniformBufferEntry);
 }
diff --git a/Engine/Source/Runtime/RenderCore/Public/RenderGraphParameter.h b/Engine/Source/Runtime/RenderCore/Public/RenderGraphParameter.h
--- a/Engine/Source/Runtime/RenderCore/Public/RenderGraphParameter.h
+++ b/Engine/Source/Runtime/RenderCore/Public/RenderGraphParameter.h
@@ -214,7 +214,7 @@ class RENDERCORE_API FRDGParameterStruct
 public:
 	template <typename ParameterStructType>
 	explicit FRDGParameterStruct(const ParameterStructType* Parameters)
-		: FRDGParameterStruct(Parameters, &ParameterStructType::FTypeInfo::GetStructMetadata()->GetLayout())
+		: FRDGParameterStruct(Parameters, ParameterStructType::FTypeInfo::GetStructMetadata()->GetLayoutPtr())
 	{}
 
 	explicit FRDGParameterStruct(const void* InContents, const FRHIUniformBufferLayout* InLayout)
@@ -229,6 +229,7 @@ public:
 
 	/** Returns the layout associated with this struct. */
 	const FRHIUniformBufferLayout& GetLayout() const { return *Layout; }
+	const FRHIUniformBufferLayout* GetLayoutPtr() const { return Layout; }
 
 	/** Helpful forwards from the layout. */
 	FORCEINLINE bool HasRenderTargets() const   { return Layout->HasRenderTargets(); }
@@ -277,7 +278,7 @@ public:
 	FRHIRenderPassInfo GetRenderPassInfo() const;
 
 private:
-	FRDGParameter GetParameterInternal(TArrayView<const FRHIUniformBufferLayout::FResourceParameter> Parameters, uint32 ParameterIndex) const
+	FRDGParameter GetParameterInternal(TArrayView<const FRHIUniformBufferResource> Parameters, uint32 ParameterIndex) const
 	{
 		checkf(ParameterIndex < static_cast<uint32>(Parameters.Num()), TEXT("Attempted to access RDG pass parameter outside of index for Layout '%s'"), *Layout->GetDebugName());
 		const EUniformBufferBaseType MemberType = Parameters[ParameterIndex].MemberType;
@@ -286,7 +287,7 @@ private:
 	}
 
 	const uint8* Contents;
-	const FRHIUniformBufferLayout* Layout;
+	FUniformBufferLayoutRHIRef Layout;
 
 	friend class FRDGPass;
 };
diff --git a/Engine/Source/Runtime/RenderCore/Public/RenderGraphPass.h b/Engine/Source/Runtime/RenderCore/Public/RenderGraphPass.h
--- a/Engine/Source/Runtime/RenderCore/Public/RenderGraphPass.h
+++ b/Engine/Source/Runtime/RenderCore/Public/RenderGraphPass.h
@@ -539,7 +539,7 @@ public:
 		const ParameterStructType* InParameterStruct,
 		ERDGPassFlags InPassFlags,
 		ExecuteLambdaType&& InExecuteLambda)
-		: FRDGPass(MoveTemp(InName), FRDGParameterStruct(InParameterStruct, &InParameterMetadata->GetLayout()), InPassFlags)
+		: FRDGPass(MoveTemp(InName), FRDGParameterStruct(InParameterStruct, InParameterMetadata->GetLayoutPtr()), InPassFlags)
 		, ExecuteLambda(MoveTemp(InExecuteLambda))
 #if RDG_ENABLE_DEBUG
 		, DebugParameterStruct(InParameterStruct)
diff --git a/Engine/Source/Runtime/RenderCore/Public/Shader.h b/Engine/Source/Runtime/RenderCore/Public/Shader.h
--- a/Engine/Source/Runtime/RenderCore/Public/Shader.h
+++ b/Engine/Source/Runtime/RenderCore/Public/Shader.h
@@ -2059,6 +2059,7 @@ protected:
 protected:
 	virtual const FTypeLayoutDesc& GetContentTypeDesc() const = 0;
 	virtual FShaderMapPointerTable* CreatePointerTable() const = 0;
+	virtual void PostFinalizeContent() { }
 
 private:
 #if WITH_EDITOR
diff --git a/Engine/Source/Runtime/RenderCore/Public/ShaderParameterMacros.h b/Engine/Source/Runtime/RenderCore/Public/ShaderParameterMacros.h
--- a/Engine/Source/Runtime/RenderCore/Public/ShaderParameterMacros.h
+++ b/Engine/Source/Runtime/RenderCore/Public/ShaderParameterMacros.h
@@ -113,13 +113,13 @@ public:
 	static TUniformBufferRef<TBufferStruct> CreateUniformBufferImmediate(const TBufferStruct& Value, EUniformBufferUsage Usage, EUniformBufferValidation Validation = EUniformBufferValidation::ValidateResources)
 	{
 		check(IsInRenderingThread() || IsInRHIThread());
-		return TUniformBufferRef<TBufferStruct>(RHICreateUniformBuffer(&Value, TBufferStruct::FTypeInfo::GetStructMetadata()->GetLayout(), Usage, Validation));
+		return TUniformBufferRef<TBufferStruct>(RHICreateUniformBuffer(&Value, TBufferStruct::FTypeInfo::GetStructMetadata()->GetLayoutPtr(), Usage, Validation));
 	}
 
 	/** Creates a uniform buffer with the given value, and returns a structured reference to it. */
 	static FLocalUniformBuffer CreateLocalUniformBuffer(FRHICommandList& RHICmdList, const TBufferStruct& Value, EUniformBufferUsage Usage)
 	{
-		return RHICmdList.BuildLocalUniformBuffer(&Value, sizeof(TBufferStruct), TBufferStruct::FTypeInfo::GetStructMetadata()->GetLayout());
+		return RHICmdList.BuildLocalUniformBuffer(&Value, sizeof(TBufferStruct), TBufferStruct::FTypeInfo::GetStructMetadata()->GetLayoutPtr());
 	}
 
 	void UpdateUniformBufferImmediate(const TBufferStruct& Value)
@@ -1198,7 +1198,7 @@ struct TShaderParameterStructTypeInfo<StructType[InNumElements]>
 
 #define INTERNAL_SHADER_PARAMETER_STRUCT_CREATE_UNIFORM_BUFFER return nullptr;
 
-#define INTERNAL_UNIFORM_BUFFER_STRUCT_CREATE_UNIFORM_BUFFER return RHICreateUniformBuffer(&InContents, StaticStructMetadata.GetLayout(), InUsage);
+#define INTERNAL_UNIFORM_BUFFER_STRUCT_CREATE_UNIFORM_BUFFER return RHICreateUniformBuffer(&InContents, StaticStructMetadata.GetLayoutPtr(), InUsage);
 
 /** Begins a uniform buffer struct declaration. */
 #define INTERNAL_SHADER_PARAMETER_STRUCT_BEGIN(StructTypeName,PrefixKeywords,ConstructorSuffix,GetStructMetadataScope,CreateUniformBufferImpl) \
diff --git a/Engine/Source/Runtime/RenderCore/Public/ShaderParameterMetadata.h b/Engine/Source/Runtime/RenderCore/Public/ShaderParameterMetadata.h
--- a/Engine/Source/Runtime/RenderCore/Public/ShaderParameterMetadata.h
+++ b/Engine/Source/Runtime/RenderCore/Public/ShaderParameterMetadata.h
@@ -239,7 +239,8 @@ public:
 		const int32 InFileLine,
 		uint32 InSize,
 		const TArray<FMember>& InMembers,
-		bool bForceCompleteInitialization = false);
+		bool bForceCompleteInitialization = false,
+		FRHIUniformBufferLayoutInitializer* OutLayoutInitializer = nullptr);
 
 	virtual ~FShaderParametersMetadata();
 
@@ -272,10 +273,16 @@ public:
 
 	uint32 GetSize() const { return Size; }
 	EUseCase GetUseCase() const { return UseCase; }
-	const FRHIUniformBufferLayout& GetLayout() const 
-	{ 
-		check(bLayoutInitialized);
-		return Layout; 
+	inline bool IsLayoutInitialized() const { return Layout != nullptr; }
+	const FRHIUniformBufferLayout& GetLayout() const
+	{
+		check(IsLayoutInitialized());
+		return *Layout;
+	}
+	const FRHIUniformBufferLayout* GetLayoutPtr() const
+	{
+		check(IsLayoutInitialized());
+		return Layout;
 	}
 	const TArray<FMember>& GetMembers() const { return Members; }
 
@@ -300,7 +307,7 @@ public:
 	uint32 GetLayoutHash() const
 	{
 		check(UseCase == EUseCase::ShaderParameterStruct || UseCase == EUseCase::UniformBuffer);
-		check(bLayoutInitialized);
+		check(IsLayoutInitialized());
 		return LayoutHash;	
 	}
 
@@ -337,6 +344,8 @@ public:
 	}
 
 private:
+	const TCHAR* const LayoutName;
+
 	/** Name of the structure type in C++ and shader code. */
 	const TCHAR* const StructTypeName;
 
@@ -364,7 +373,7 @@ private:
 	const EUniformBufferBindingFlags BindingFlags;
 
 	/** Layout of all the resources in the shader parameter struct. */
-	FRHIUniformBufferLayout Layout;
+	FUniformBufferLayoutRHIRef Layout{};
 	
 	/** List of all members. */
 	TArray<FMember> Members;
@@ -372,14 +381,11 @@ private:
 	/** Shackle elements in global link list of globally named shader parameters. */
 	TLinkedList<FShaderParametersMetadata*> GlobalListLink;
 
-	/** Whether the layout is actually initialized yet or not. */
-	uint32 bLayoutInitialized : 1;
-
 	/** Hash about the entire memory layout of the structure. */
 	uint32 LayoutHash = 0;
 
 
-	void InitializeLayout();
+	void InitializeLayout(FRHIUniformBufferLayoutInitializer* OutLayoutInitializer = nullptr);
 
 	void AddResourceTableEntriesRecursive(const TCHAR* UniformBufferName, const TCHAR* Prefix, uint16& ResourceIndex, TMap<FString, FResourceTableEntry>& ResourceTableMap) const;
 
diff --git a/Engine/Source/Runtime/RenderCore/Public/ShaderParameterUtils.h b/Engine/Source/Runtime/RenderCore/Public/ShaderParameterUtils.h
--- a/Engine/Source/Runtime/RenderCore/Public/ShaderParameterUtils.h
+++ b/Engine/Source/Runtime/RenderCore/Public/ShaderParameterUtils.h
@@ -604,8 +604,8 @@ inline void SetUniformBufferParameterImmediate(
 	checkSlow(Parameter.IsInitialized());
 	if(Parameter.IsBound())
 	{
-		const FRHIUniformBufferLayout& UniformBufferLayout = TBufferStruct::StaticStructMetadata.GetLayout();
-		FLocalUniformBuffer UniformBuffer = RHICmdList.BuildLocalUniformBuffer(&UniformBufferValue, UniformBufferLayout.ConstantBufferSize, UniformBufferLayout);
+		const FRHIUniformBufferLayout* UniformBufferLayout = TBufferStruct::StaticStructMetadata.GetLayoutPtr();
+		FLocalUniformBuffer UniformBuffer = RHICmdList.BuildLocalUniformBuffer(&UniformBufferValue, UniformBufferLayout->ConstantBufferSize, UniformBufferLayout);
 
 		RHICmdList.SetLocalShaderUniformBuffer(
 			Shader,
@@ -627,8 +627,8 @@ inline void SetUniformBufferParameterImmediate(
 	checkSlow(Parameter.IsInitialized());
 	if(Parameter.IsBound())
 	{
-		const FRHIUniformBufferLayout& UniformBufferLayout = TBufferStruct::StaticStructMetadata.GetLayout();
-		FLocalUniformBuffer UniformBuffer = RHICmdList.BuildLocalUniformBuffer(&UniformBufferValue, UniformBufferLayout.ConstantBufferSize, UniformBufferLayout);
+		const FRHIUniformBufferLayout* UniformBufferLayout = TBufferStruct::StaticStructMetadata.GetLayoutPtr();
+		FLocalUniformBuffer UniformBuffer = RHICmdList.BuildLocalUniformBuffer(&UniformBufferValue, UniformBufferLayout->ConstantBufferSize, UniformBufferLayout);
 
 		RHICmdList.SetLocalShaderUniformBuffer(
 			Shader,
diff --git a/Engine/Source/Runtime/Renderer/Private/ShaderBaseClasses.cpp b/Engine/Source/Runtime/Renderer/Private/ShaderBaseClasses.cpp
--- a/Engine/Source/Runtime/Renderer/Private/ShaderBaseClasses.cpp
+++ b/Engine/Source/Runtime/Renderer/Private/ShaderBaseClasses.cpp
@@ -72,7 +72,7 @@ FMaterialShader::FMaterialShader(const FMaterialShaderType::CompiledShaderInitia
 {
 #if WITH_EDITORONLY_DATA
 	check(!DebugDescription.IsEmpty());
-	DebugUniformExpressionUBLayout.CopyFrom(Initializer.UniformExpressionSet.GetUniformBufferLayout());
+	DebugUniformExpressionUBLayout.CopyFrom(Initializer.UniformExpressionSet.GetUniformBufferLayoutInitializer());
 #endif // WITH_EDITORONLY_DATA
 
 	// Bind the material uniform buffer parameter.
@@ -104,6 +104,18 @@ FRHIUniformBuffer* FMaterialShader::GetParameterCollectionBuffer(const FGuid& Id
 }
 
 #if !(UE_BUILD_TEST || UE_BUILD_SHIPPING || !WITH_EDITOR)
+template<typename TLayout>
+static void DumpUB(const TLayout& Layout)
+{
+	UE_LOG(LogShaders, Warning, TEXT("Layout %s, Hash %08x"), *Layout.GetDebugName(), Layout.GetHash());
+	FString ResourcesString;
+	for (int32 Index = 0; Index < Layout.Resources.Num(); ++Index)
+	{
+		ResourcesString += FString::Printf(TEXT("%d "), (uint8)Layout.Resources[Index].MemberType);
+	}
+	UE_LOG(LogShaders, Warning, TEXT("Layout CB Size %d %d Resources: %s"), Layout.ConstantBufferSize, Layout.Resources.Num(), *ResourcesString);
+};
+
 void FMaterialShader::VerifyExpressionAndShaderMaps(const FMaterialRenderProxy* MaterialRenderProxy, const FMaterial& Material, const FUniformExpressionCache* UniformExpressionCache) const
 {
 	// Validate that the shader is being used for a material that matches the uniform expression set the shader was compiled for.
@@ -119,16 +131,6 @@ void FMaterialShader::VerifyExpressionAndShaderMaps(const FMaterialRenderProxy*
 		|| UniformExpressionCache->CachedUniformExpressionShaderMap != ShaderMap;
 	if (!bUniformExpressionSetMismatch)
 	{
-		auto DumpUB = [](const FRHIUniformBufferLayout& Layout)
-		{
-			UE_LOG(LogShaders, Warning, TEXT("Layout %s, Hash %08x"), *Layout.GetDebugName(), Layout.GetHash());
-			FString ResourcesString;
-			for (int32 Index = 0; Index < Layout.Resources.Num(); ++Index)
-			{
-				ResourcesString += FString::Printf(TEXT("%d "), (uint8)Layout.Resources[Index].MemberType);
-			}
-			UE_LOG(LogShaders, Warning, TEXT("Layout CB Size %d %d Resources: %s"), Layout.ConstantBufferSize, Layout.Resources.Num(), *ResourcesString);
-		};
 		if (UniformExpressionCache->LocalUniformBuffer.IsValid())
 		{
 			if (UniformExpressionCache->LocalUniformBuffer.BypassUniform)
@@ -441,7 +443,7 @@ bool FMaterialShader::Serialize(FArchive& Ar)
 		DebugUniformExpressionUBLayout.Resources.Reserve(ResourceOffsets.Num());
 		for (int32 i = 0; i < ResourceOffsets.Num(); i++)
 		{
-			DebugUniformExpressionUBLayout.Resources.Emplace(FRHIUniformBufferLayout::FResourceParameter{ ResourceOffsets[i], EUniformBufferBaseType(ResourceTypes[i]) });
+			DebugUniformExpressionUBLayout.Resources.Emplace(FRHIUniformBufferResource{ ResourceOffsets[i], EUniformBufferBaseType(ResourceTypes[i]) });
 		}
 		DebugUniformExpressionUBLayout.ComputeHash();
 #endif
diff --git a/Engine/Source/Runtime/Renderer/Public/MaterialShader.h b/Engine/Source/Runtime/Renderer/Public/MaterialShader.h
--- a/Engine/Source/Runtime/Renderer/Public/MaterialShader.h
+++ b/Engine/Source/Runtime/Renderer/Public/MaterialShader.h
@@ -148,6 +148,6 @@ private:
 
 protected:
 	LAYOUT_FIELD_EDITORONLY(FDebugUniformExpressionSet, DebugUniformExpressionSet);
-	LAYOUT_FIELD_EDITORONLY(FRHIUniformBufferLayout, DebugUniformExpressionUBLayout);
+	LAYOUT_FIELD_EDITORONLY(FRHIUniformBufferLayoutInitializer, DebugUniformExpressionUBLayout);
 	LAYOUT_FIELD_EDITORONLY(FMemoryImageString, DebugDescription);
 };
diff --git a/Engine/Source/Runtime/VulkanRHI/Private/VulkanResources.h b/Engine/Source/Runtime/VulkanRHI/Private/VulkanResources.h
--- a/Engine/Source/Runtime/VulkanRHI/Private/VulkanResources.h
+++ b/Engine/Source/Runtime/VulkanRHI/Private/VulkanResources.h
@@ -1168,7 +1168,7 @@ protected:
 class FVulkanUniformBuffer : public FRHIUniformBuffer
 {
 public:
-	FVulkanUniformBuffer(FVulkanDevice& Device, const FRHIUniformBufferLayout& InLayout, const void* Contents, EUniformBufferUsage InUsage, EUniformBufferValidation Validation);
+	FVulkanUniformBuffer(FVulkanDevice& Device, const FRHIUniformBufferLayout* InLayout, const void* Contents, EUniformBufferUsage InUsage, EUniformBufferValidation Validation);
 	virtual ~FVulkanUniformBuffer();
 
 	const TArray<TRefCountPtr<FRHIResource>>& GetResourceTable() const { return ResourceTable; }
diff --git a/Engine/Source/Runtime/VulkanRHI/Private/VulkanUniformBuffer.cpp b/Engine/Source/Runtime/VulkanRHI/Private/VulkanUniformBuffer.cpp
--- a/Engine/Source/Runtime/VulkanRHI/Private/VulkanUniformBuffer.cpp
+++ b/Engine/Source/Runtime/VulkanRHI/Private/VulkanUniformBuffer.cpp
@@ -92,7 +92,7 @@ static void UpdateUniformBufferHelper(FVulkanCommandListContext& Context, FVulka
 	}
 };
 
-FVulkanUniformBuffer::FVulkanUniformBuffer(FVulkanDevice& Device, const FRHIUniformBufferLayout& InLayout, const void* Contents, EUniformBufferUsage InUsage, EUniformBufferValidation Validation)
+FVulkanUniformBuffer::FVulkanUniformBuffer(FVulkanDevice& Device, const FRHIUniformBufferLayout* InLayout, const void* Contents, EUniformBufferUsage InUsage, EUniformBufferValidation Validation)
 	: FRHIUniformBuffer(InLayout)
 	, Device(&Device) 
 	, Usage(InUsage)
@@ -104,10 +104,10 @@ FVulkanUniformBuffer::FVulkanUniformBuffer(FVulkanDevice& Device, const FRHIUnif
 	// Verify the correctness of our thought pattern how the resources are delivered
 	//	- If we have at least one resource, we also expect ResourceOffset to have an offset
 	//	- Meaning, there is always a uniform buffer with a size specified larged than 0 bytes
-	check(InLayout.Resources.Num() > 0 || InLayout.ConstantBufferSize > 0);
+	check(InLayout->Resources.Num() > 0 || InLayout->ConstantBufferSize > 0);
 
 	// Setup resource table
-	const uint32 NumResources = InLayout.Resources.Num();
+	const uint32 NumResources = InLayout->Resources.Num();
 	if (NumResources > 0)
 	{
 		// Transfer the resource table to an internal resource-array
@@ -115,17 +115,17 @@ FVulkanUniformBuffer::FVulkanUniformBuffer(FVulkanDevice& Device, const FRHIUnif
 		ResourceTable.AddZeroed(NumResources);
 		for (uint32 Index = 0; Index < NumResources; ++Index)
 		{
-			ResourceTable[Index] = GetShaderParameterResourceRHI(Contents, InLayout.Resources[Index].MemberOffset, InLayout.Resources[Index].MemberType);
+			ResourceTable[Index] = GetShaderParameterResourceRHI(Contents, InLayout->Resources[Index].MemberOffset, InLayout->Resources[Index].MemberType);
 		}
 	}
 
-	if (InLayout.ConstantBufferSize > 0)
+	if (InLayout->ConstantBufferSize > 0)
 	{
 		if (UseRingBuffer(InUsage))
 		{
 			FRHICommandListImmediate& RHICmdList = FRHICommandListExecutor::GetImmediateCommandList();
 			FVulkanUniformBuffer* UniformBuffer = this;
-			int32 DataSize = InLayout.ConstantBufferSize;
+			int32 DataSize = InLayout->ConstantBufferSize;
 			
 			// make sure we allocate from RingBuffer on RHIT
 			const bool bCanAllocOnThisThread = RHICmdList.Bypass() || (!IsRunningRHIInSeparateThread() && IsInRenderingThread()) || IsInRHIThread();
@@ -152,7 +152,7 @@ FVulkanUniformBuffer::FVulkanUniformBuffer(FVulkanDevice& Device, const FRHIUnif
 		{
 			VulkanRHI::FMemoryManager& ResourceMgr = Device.GetMemoryManager();
 			// Set it directly as there is no previous one
-			ResourceMgr.AllocUniformBuffer(Allocation, InLayout.ConstantBufferSize, Contents);
+			ResourceMgr.AllocUniformBuffer(Allocation, InLayout->ConstantBufferSize, Contents);
 		}
 	}
 }
@@ -184,7 +184,7 @@ void FVulkanUniformBuffer::UpdateResourceTable(FRHIResource** Resources, int32 R
 }
 
 
-FUniformBufferRHIRef FVulkanDynamicRHI::RHICreateUniformBuffer(const void* Contents, const FRHIUniformBufferLayout& Layout, EUniformBufferUsage Usage, EUniformBufferValidation Validation)
+FUniformBufferRHIRef FVulkanDynamicRHI::RHICreateUniformBuffer(const void* Contents, const FRHIUniformBufferLayout* Layout, EUniformBufferUsage Usage, EUniformBufferValidation Validation)
 {
 	LLM_SCOPE_VULKAN(ELLMTagVulkan::VulkanUniformBuffers);
 
diff --git a/Engine/Source/Runtime/VulkanRHI/Public/VulkanDynamicRHI.h b/Engine/Source/Runtime/VulkanRHI/Public/VulkanDynamicRHI.h
--- a/Engine/Source/Runtime/VulkanRHI/Public/VulkanDynamicRHI.h
+++ b/Engine/Source/Runtime/VulkanRHI/Public/VulkanDynamicRHI.h
@@ -77,7 +77,7 @@ public:
 	virtual FBoundShaderStateRHIRef RHICreateBoundShaderState(FRHIVertexDeclaration* VertexDeclaration, FRHIVertexShader* VertexShader, FRHIPixelShader* PixelShader, FRHIGeometryShader* GeometryShader) final override;
 	virtual FGraphicsPipelineStateRHIRef RHICreateGraphicsPipelineState(const FGraphicsPipelineStateInitializer& Initializer) final override;
 	virtual FComputePipelineStateRHIRef RHICreateComputePipelineState(FRHIComputeShader* ComputeShader) final override;
-	virtual FUniformBufferRHIRef RHICreateUniformBuffer(const void* Contents, const FRHIUniformBufferLayout& Layout, EUniformBufferUsage Usage, EUniformBufferValidation Validation) final override;
+	virtual FUniformBufferRHIRef RHICreateUniformBuffer(const void* Contents, const FRHIUniformBufferLayout* Layout, EUniformBufferUsage Usage, EUniformBufferValidation Validation) final override;
 	virtual void RHIUpdateUniformBuffer(FRHIUniformBuffer* UniformBufferRHI, const void* Contents) final override;
 	virtual FBufferRHIRef RHICreateBuffer(uint32 Size, EBufferUsageFlags Usage, uint32 Stride, ERHIAccess ResourceState, FRHIResourceCreateInfo& CreateInfo) final override;
 	virtual void RHICopyBuffer(FRHIBuffer* SourceBuffer, FRHIBuffer* DestBuffer) final override;
diff --git a/Engine/Source/Runtime/Windows/D3D11RHI/Private/D3D11RHIPrivate.h b/Engine/Source/Runtime/Windows/D3D11RHI/Private/D3D11RHIPrivate.h
--- a/Engine/Source/Runtime/Windows/D3D11RHI/Private/D3D11RHIPrivate.h
+++ b/Engine/Source/Runtime/Windows/D3D11RHI/Private/D3D11RHIPrivate.h
@@ -472,7 +472,7 @@ public:
 	virtual void* RHILockStagingBuffer(FRHIStagingBuffer* StagingBuffer, FRHIGPUFence* Fence, uint32 Offset, uint32 SizeRHI) final override;
     virtual void RHIUnlockStagingBuffer(FRHIStagingBuffer* StagingBuffer) final override;
 	virtual FBoundShaderStateRHIRef RHICreateBoundShaderState(FRHIVertexDeclaration* VertexDeclaration, FRHIVertexShader* VertexShader, FRHIPixelShader* PixelShader, FRHIGeometryShader* GeometryShader) final override;
-	virtual FUniformBufferRHIRef RHICreateUniformBuffer(const void* Contents, const FRHIUniformBufferLayout& Layout, EUniformBufferUsage Usage, EUniformBufferValidation Validation) final override;
+	virtual FUniformBufferRHIRef RHICreateUniformBuffer(const void* Contents, const FRHIUniformBufferLayout* Layout, EUniformBufferUsage Usage, EUniformBufferValidation Validation) final override;
 	virtual void RHIUpdateUniformBuffer(FRHIUniformBuffer* UniformBufferRHI, const void* Contents) final override;
 	virtual FBufferRHIRef RHICreateBuffer(uint32 Size, EBufferUsageFlags Usage, uint32 Stride, ERHIAccess ResourceState, FRHIResourceCreateInfo& CreateInfo) final override;
 	virtual void RHICopyBuffer(FRHIBuffer* SourceBuffer, FRHIBuffer* DestBuffer) final override;
diff --git a/Engine/Source/Runtime/Windows/D3D11RHI/Private/D3D11UniformBuffer.cpp b/Engine/Source/Runtime/Windows/D3D11RHI/Private/D3D11UniformBuffer.cpp
--- a/Engine/Source/Runtime/Windows/D3D11RHI/Private/D3D11UniformBuffer.cpp
+++ b/Engine/Source/Runtime/Windows/D3D11RHI/Private/D3D11UniformBuffer.cpp
@@ -164,17 +164,17 @@ static TRefCountPtr<ID3D11Buffer> CreateAndUpdatePooledUniformBuffer(
 	return UniformBufferResource;
 }
 
-FUniformBufferRHIRef FD3D11DynamicRHI::RHICreateUniformBuffer(const void* Contents, const FRHIUniformBufferLayout& Layout, EUniformBufferUsage Usage, EUniformBufferValidation Validation)
+FUniformBufferRHIRef FD3D11DynamicRHI::RHICreateUniformBuffer(const void* Contents, const FRHIUniformBufferLayout* Layout, EUniformBufferUsage Usage, EUniformBufferValidation Validation)
 {
 	check(IsInRenderingThread() || IsInRHIThread());
 
 	if (Validation == EUniformBufferValidation::ValidateResources)
 	{
-		ValidateShaderParameterResourcesRHI(Contents, Layout);
+		ValidateShaderParameterResourcesRHI(Contents, *Layout);
 	}
 
 	FD3D11UniformBuffer* NewUniformBuffer = nullptr;
-	const uint32 NumBytes = Layout.ConstantBufferSize;
+	const uint32 NumBytes = Layout->ConstantBufferSize;
 	if (NumBytes > 0)
 	{
 		// Constant buffers must also be 16-byte aligned.
@@ -245,16 +245,15 @@ FUniformBufferRHIRef FD3D11DynamicRHI::RHICreateUniformBuffer(const void* Conten
 		NewUniformBuffer = new FD3D11UniformBuffer(this, Layout, nullptr, FRingAllocation());
 	}
 
-	const TCHAR* LayoutName = *Layout.GetDebugName();
-	if (Layout.Resources.Num())
+	if (Layout->Resources.Num())
 	{
-		const int32 ResourceCount = Layout.Resources.Num();
+		const int32 ResourceCount = Layout->Resources.Num();
 		NewUniformBuffer->ResourceTable.Empty(ResourceCount);
 		NewUniformBuffer->ResourceTable.AddZeroed(ResourceCount);
 
 		for (int32 Index = 0; Index < ResourceCount; ++Index)
 		{
-			const auto ResourceParameter = Layout.Resources[Index];
+			const auto ResourceParameter = Layout->Resources[Index];
 			NewUniformBuffer->ResourceTable[Index] = GetShaderParameterResourceRHI(Contents, ResourceParameter.MemberOffset, ResourceParameter.MemberType);
 		}
 	}
diff --git a/Engine/Source/Runtime/Windows/D3D11RHI/Public/D3D11Resources.h b/Engine/Source/Runtime/Windows/D3D11RHI/Public/D3D11Resources.h
--- a/Engine/Source/Runtime/Windows/D3D11RHI/Public/D3D11Resources.h
+++ b/Engine/Source/Runtime/Windows/D3D11RHI/Public/D3D11Resources.h
@@ -611,7 +611,7 @@ public:
 	TArray<TRefCountPtr<FRHIResource> > ResourceTable;
 
 	/** Initialization constructor. */
-	FD3D11UniformBuffer(class FD3D11DynamicRHI* InD3D11RHI, const FRHIUniformBufferLayout& InLayout, ID3D11Buffer* InResource,const FRingAllocation& InRingAllocation)
+	FD3D11UniformBuffer(class FD3D11DynamicRHI* InD3D11RHI, const FRHIUniformBufferLayout* InLayout, ID3D11Buffer* InResource,const FRingAllocation& InRingAllocation)
 	: FRHIUniformBuffer(InLayout)
 	, Resource(InResource)
 	, RingAllocation(InRingAllocation)
