rhi.dumpresourcememory respects the -csv arg so it works with memreport.

- If we are dumping a summary with rhi.dumpresourcememory that takes precedence and ignores the csv arg.
- Remove outputting csv to file only.  This way this command plays nicely with memreport.

diff --git a/Engine/Source/Runtime/RHI/Private/RHI.cpp b/Engine/Source/Runtime/RHI/Private/RHI.cpp
--- a/Engine/Source/Runtime/RHI/Private/RHI.cpp
+++ b/Engine/Source/Runtime/RHI/Private/RHI.cpp
@@ -836,7 +836,6 @@ static FAutoConsoleCommandWithWorldArgsAndOutputDevice GDumpRHIResourceMemoryCmd
 	int32 NumberOfResourcesToShow = 50;
 	bool bUseCSVOutput = false;
 	bool bSummaryOutput = false;
-	FArchive* CSVFile{ nullptr };
 
 	for (const FString& Argument : Args)
 	{
@@ -844,12 +843,9 @@ static FAutoConsoleCommandWithWorldArgsAndOutputDevice GDumpRHIResourceMemoryCmd
 		{
 			NumberOfResourcesToShow = -1;
 		}
-		else if (Argument.Equals(TEXT("csv"), ESearchCase::IgnoreCase))
+		else if (Argument.Equals(TEXT("-csv"), ESearchCase::IgnoreCase))
 		{
 			bUseCSVOutput = true;
-
-			const FString Filename = FString::Printf(TEXT("%srhiDumpResourceMemory-%s.csv"), *FPaths::ProfilingDir(), *FDateTime::Now().ToString());
-			CSVFile = IFileManager::Get().CreateFileWriter(*Filename, FILEWRITE_AllowRead);
 		}
 		else if (Argument.StartsWith(TEXT("Name="), ESearchCase::IgnoreCase))
 		{
@@ -880,6 +876,12 @@ static FAutoConsoleCommandWithWorldArgsAndOutputDevice GDumpRHIResourceMemoryCmd
 		}
 	}
 
+	// Summary output and csv are mutually exclusive.  So if both are on, summary takes precedence
+	if (bSummaryOutput && bUseCSVOutput)
+	{
+		bUseCSVOutput = false;
+	}
+
 	TCHAR ResourceNameBuffer[FName::StringBufferSize];
 
 	auto ShouldIncludeResource = [&](const FRHIResource* Resource, const FRHIResourceInfo& ResourceInfo) -> bool
@@ -975,8 +977,7 @@ static FAutoConsoleCommandWithWorldArgsAndOutputDevice GDumpRHIResourceMemoryCmd
 
 	if (bUseCSVOutput)
 	{
-		const TCHAR* Header = TEXT("Name,Type,Size,Transient,Streaming,RenderTarget,UAV,\"Raytracing Acceleration Structure\"\n");
-		CSVFile->Serialize(TCHAR_TO_ANSI(Header), FPlatformString::Strlen(Header));
+		BufferedOutput.CategorizedLogf(CategoryName, ELogVerbosity::Log, TEXT("Name,Type,Size,Transient,Streaming,RenderTarget,UAV,\"Raytracing Acceleration Structure\""));
 	}
 	else
 	{
@@ -1033,7 +1034,7 @@ static FAutoConsoleCommandWithWorldArgsAndOutputDevice GDumpRHIResourceMemoryCmd
 			{
 				if (bUseCSVOutput)
 				{
-					const FString Row = FString::Printf(TEXT("%s,%s,%.9f,%s,%s,%s,%s,%s\n"),
+					BufferedOutput.CategorizedLogf(CategoryName, ELogVerbosity::Log, TEXT("%s,%s,%.9f,%s,%s,%s,%s,%s"),
 						ResourceNameBuffer,
 						ResourceType,
 						SizeInBytes / double(1 << 20),
@@ -1042,8 +1043,6 @@ static FAutoConsoleCommandWithWorldArgsAndOutputDevice GDumpRHIResourceMemoryCmd
 						(bRT || bDS) ? TEXT("Yes") : TEXT("No"),
 						bUAV ? TEXT("Yes") : TEXT("No"),
 						bRTAS ? TEXT("Yes") : TEXT("No"));
-
-					CSVFile->Serialize(TCHAR_TO_ANSI(*Row), Row.Len());
 				}
 				else
 				{
@@ -1093,12 +1092,7 @@ static FAutoConsoleCommandWithWorldArgsAndOutputDevice GDumpRHIResourceMemoryCmd
 		}
 	}
 
-	if (bUseCSVOutput)
-	{
-		delete CSVFile;
-		CSVFile = nullptr;
-	}
-	else
+	if (!bUseCSVOutput)
 	{
 		const double TotalNonTransientSizeF = TotalTrackedResourceSize / double(1 << 20);
 		const double TotalTransientSizeF = TotalTrackedTransientResourceSize / double(1 << 20);
@@ -1144,9 +1138,9 @@ static FAutoConsoleCommandWithWorldArgsAndOutputDevice GDumpRHIResourceMemoryCmd
 				BufferedOutput.CategorizedLogf(CategoryName, ELogVerbosity::Log, TEXT("    Transient: %.9f MB"), TotalTransientSizeF);
 			}
 		}
-
-		BufferedOutput.RedirectTo(OutputDevice);
 	}
+
+	BufferedOutput.RedirectTo(OutputDevice);
 }));
 
 #endif // RHI_WANT_RESOURCE_INFO
@@ -2710,4 +2704,6 @@ void FRHITransientBuffer::Init(const TCHAR* InName, uint32 InAllocationIndex, ui
 	ViewCache.SetDebugName(InName);
 
 	// TODO: Add method to rename a buffer.
-}
\ No newline at end of file
+}
+
+PRAGMA_ENABLE_OPTIMIZATION
\ No newline at end of file
