Remove PRAGMA_ENABLE_OPTIMIZATION I didn't intend to check in.

diff --git a/Engine/Source/Runtime/RHI/Private/RHI.cpp b/Engine/Source/Runtime/RHI/Private/RHI.cpp
--- a/Engine/Source/Runtime/RHI/Private/RHI.cpp
+++ b/Engine/Source/Runtime/RHI/Private/RHI.cpp
@@ -2699,5 +2699,3 @@ void FRHITransientBuffer::Init(const TCHAR* InName, uint32 InAllocationIndex, ui
 
 	// TODO: Add method to rename a buffer.
 }
-
-PRAGMA_ENABLE_OPTIMIZATION
\ No newline at end of file
