PR #8343: [UE5] Add SM6 into RHI LexToString, fix a typo in SceneView's LexToString (Contributed by [REDACTED_PERSON])

diff --git a/Engine/Source/Runtime/Engine/Private/SceneView.cpp b/Engine/Source/Runtime/Engine/Private/SceneView.cpp
--- a/Engine/Source/Runtime/Engine/Private/SceneView.cpp
+++ b/Engine/Source/Runtime/Engine/Private/SceneView.cpp
@@ -412,7 +412,7 @@ FString LexToString(EMaterialQualityLevel::Type QualityLevel)
 		default:
 			break;
 	}
-	return TEXT("UnknownFeatureLevel");
+	return TEXT("UnknownMaterialQualityLevel");
 }
 
 /**
diff --git a/Engine/Source/Runtime/RHI/Private/RHI.cpp b/Engine/Source/Runtime/RHI/Private/RHI.cpp
--- a/Engine/Source/Runtime/RHI/Private/RHI.cpp
+++ b/Engine/Source/Runtime/RHI/Private/RHI.cpp
@@ -2128,6 +2128,8 @@ FString LexToString(ERHIFeatureLevel::Type Level)
 			return TEXT("SM4_REMOVED");
 		case ERHIFeatureLevel::SM5:
 			return TEXT("SM5");
+		case ERHIFeatureLevel::SM6:
+			return TEXT("SM6");
 		default:
 			break;
 	}
