Improving some RHI memory tracking
- RHIUniformBufferLayout's are now tracked under "RHIMisc/CreateUniformBufferLayout"
- RHIResource tracking overhead is now tracked under "RHIMisc/ResourceTracking" instead of contributing to other tags.

diff --git a/Engine/Source/Runtime/RHI/Private/RHI.cpp b/Engine/Source/Runtime/RHI/Private/RHI.cpp
--- a/Engine/Source/Runtime/RHI/Private/RHI.cpp
+++ b/Engine/Source/Runtime/RHI/Private/RHI.cpp
@@ -641,6 +641,8 @@ void FRHIResource::BeginTrackingResource(FRHIResource* InResource)
 {
 	if (GRHITrackingResources)
 	{
+		LLM_SCOPE_BYNAME(TEXT("RHIMisc/ResourceTracking"));
+
 		FScopeLock Lock(&GRHIResourceTrackingCriticalSection);
 
 		InResource->bBeingTracked = true;
diff --git a/Engine/Source/Runtime/RHI/Public/DynamicRHI.h b/Engine/Source/Runtime/RHI/Public/DynamicRHI.h
--- a/Engine/Source/Runtime/RHI/Public/DynamicRHI.h
+++ b/Engine/Source/Runtime/RHI/Public/DynamicRHI.h
@@ -1524,7 +1524,7 @@ FORCEINLINE TRefCountPtr<FRHIRayTracingPipelineState> RHICreateRayTracingPipelin
 
 FORCEINLINE FUniformBufferLayoutRHIRef RHICreateUniformBufferLayout(const FRHIUniformBufferLayoutInitializer& Initializer)
 {
-	LLM_SCOPE(ELLMTag::RHIMisc);
+	LLM_SCOPE_BYNAME(TEXT("RHIMisc/CreateUniformBufferLayout"));
 	return new FRHIUniformBufferLayout(Initializer);
 }
 
