Creating a 'better' define name for enabling RHI resource info. Adding the ability to force disable resource info via global definitions.

diff --git a/Engine/Source/Runtime/D3D12RHI/Private/D3D12Texture.h b/Engine/Source/Runtime/D3D12RHI/Private/D3D12Texture.h
--- a/Engine/Source/Runtime/D3D12RHI/Private/D3D12Texture.h
+++ b/Engine/Source/Runtime/D3D12RHI/Private/D3D12Texture.h
@@ -292,7 +292,7 @@ public:
 	virtual ~TD3D12Texture2D();
 
 	// FRHIResource overrides
-#if RHI_WANT_RESOURCE_INFO
+#if RHI_ENABLE_RESOURCE_INFO
 	bool GetResourceInfo(FRHIResourceInfo& OutResourceInfo) const override
 	{
 		OutResourceInfo = FRHIResourceInfo{};
@@ -435,7 +435,7 @@ public:
 	virtual ~FD3D12Texture3D();
 
 	// FRHIResource overrides
-#if RHI_WANT_RESOURCE_INFO
+#if RHI_ENABLE_RESOURCE_INFO
 	bool GetResourceInfo(FRHIResourceInfo& OutResourceInfo) const override
 	{
 		OutResourceInfo = FRHIResourceInfo{};
diff --git a/Engine/Source/Runtime/D3D12RHI/Public/D3D12Resources.h b/Engine/Source/Runtime/D3D12RHI/Public/D3D12Resources.h
--- a/Engine/Source/Runtime/D3D12RHI/Public/D3D12Resources.h
+++ b/Engine/Source/Runtime/D3D12RHI/Public/D3D12Resources.h
@@ -963,7 +963,7 @@ public:
 	FD3D12SyncPoint UploadResourceDataViaCopyQueue(FResourceArrayInterface* InResourceArray);
 
 	// FRHIResource overrides
-#if RHI_WANT_RESOURCE_INFO
+#if RHI_ENABLE_RESOURCE_INFO
 	bool GetResourceInfo(FRHIResourceInfo& OutResourceInfo) const override
 	{
 		OutResourceInfo = FRHIResourceInfo{};
diff --git a/Engine/Source/Runtime/RHI/Private/DynamicRHI.cpp b/Engine/Source/Runtime/RHI/Private/DynamicRHI.cpp
--- a/Engine/Source/Runtime/RHI/Private/DynamicRHI.cpp
+++ b/Engine/Source/Runtime/RHI/Private/DynamicRHI.cpp
@@ -219,7 +219,7 @@ void RHIInit(bool bHasEditorToken)
 {
 	if (!GDynamicRHI)
 	{
-#if RHI_WANT_RESOURCE_INFO
+#if RHI_ENABLE_RESOURCE_INFO
 		FRHIResource::StartTrackingAllResources();
 #endif
 
@@ -345,7 +345,7 @@ void RHIExit()
 		delete GDynamicRHI;
 		GDynamicRHI = NULL;
 
-#if RHI_WANT_RESOURCE_INFO
+#if RHI_ENABLE_RESOURCE_INFO
 		FRHIResource::StopTrackingAllResources();
 #endif
 	}
diff --git a/Engine/Source/Runtime/RHI/Private/RHI.cpp b/Engine/Source/Runtime/RHI/Private/RHI.cpp
--- a/Engine/Source/Runtime/RHI/Private/RHI.cpp
+++ b/Engine/Source/Runtime/RHI/Private/RHI.cpp
@@ -17,7 +17,7 @@
 #include "Misc/BufferedOutputDevice.h"
 #include "Misc/OutputDeviceFile.h"
 
-#if RHI_WANT_RESOURCE_INFO
+#if RHI_ENABLE_RESOURCE_INFO
 #include "HAL/FileManager.h"
 #endif
 
@@ -631,7 +631,7 @@ bool operator==(const FBlendStateInitializerRHI::FRenderTarget& A, const FBlendS
 	return bSame;
 }
 
-#if RHI_WANT_RESOURCE_INFO
+#if RHI_ENABLE_RESOURCE_INFO
 
 static FCriticalSection GRHIResourceTrackingCriticalSection;
 static TSet<FRHIResource*> GRHITrackedResources;
@@ -1145,7 +1145,7 @@ static FAutoConsoleCommandWithWorldArgsAndOutputDevice GDumpRHIResourceMemoryCmd
 	BufferedOutput.RedirectTo(OutputDevice);
 }));
 
-#endif // RHI_WANT_RESOURCE_INFO
+#endif // RHI_ENABLE_RESOURCE_INFO
 
 bool FRHIResource::Bypass()
 {
diff --git a/Engine/Source/Runtime/RHI/Public/RHIResources.h b/Engine/Source/Runtime/RHI/Public/RHIResources.h
--- a/Engine/Source/Runtime/RHI/Public/RHIResources.h
+++ b/Engine/Source/Runtime/RHI/Public/RHIResources.h
@@ -26,10 +26,18 @@
 #include "Containers/ClosableMpscQueue.h"
 #include "Misc/CoreDelegates.h"
 
+// RHI_WANT_RESOURCE_INFO should be controlled by the RHI module.
 #ifndef RHI_WANT_RESOURCE_INFO
 #define RHI_WANT_RESOURCE_INFO 0
 #endif
 
+// RHI_FORCE_DISABLE_RESOURCE_INFO can be defined anywhere else, like in GlobalDefinitions.
+#ifndef RHI_FORCE_DISABLE_RESOURCE_INFO
+#define RHI_FORCE_DISABLE_RESOURCE_INFO 0
+#endif
+
+#define RHI_ENABLE_RESOURCE_INFO (RHI_WANT_RESOURCE_INFO && !RHI_FORCE_DISABLE_RESOURCE_INFO)
+
 #ifndef RHI_RESOURCE_LIFETIME_VALIDATION
 #define RHI_RESOURCE_LIFETIME_VALIDATION 0
 #endif
@@ -54,7 +62,7 @@ public:
 	FRHIResource(ERHIResourceType InResourceType)
 		: ResourceType(InResourceType)
 	{
-#if RHI_WANT_RESOURCE_INFO
+#if RHI_ENABLE_RESOURCE_INFO
 		BeginTrackingResource(this);
 #endif
 	}
@@ -65,7 +73,7 @@ public:
 		check(AtomicFlags.GetNumRefs(std::memory_order_relaxed) == 0); // this should not have any outstanding refs
 		CurrentlyDeleting = nullptr;
 
-#if RHI_WANT_RESOURCE_INFO
+#if RHI_ENABLE_RESOURCE_INFO
 		EndTrackingResource(this);
 #endif
 	}
@@ -137,7 +145,7 @@ public:
 
 	inline ERHIResourceType GetType() const { return ResourceType; }
 
-#if RHI_WANT_RESOURCE_INFO
+#if RHI_ENABLE_RESOURCE_INFO
 	// Get resource info if available.
 	// Should return true if the ResourceInfo was filled with data.
 	virtual bool GetResourceInfo(FRHIResourceInfo& OutResourceInfo) const
@@ -309,7 +317,7 @@ private:
 
 	const ERHIResourceType ResourceType;
 	bool bCommitted { true };
-#if RHI_WANT_RESOURCE_INFO
+#if RHI_ENABLE_RESOURCE_INFO
 	bool bBeingTracked { false };
 #endif
 
diff --git a/Engine/Source/Runtime/Windows/D3D11RHI/Public/D3D11Resources.h b/Engine/Source/Runtime/Windows/D3D11RHI/Public/D3D11Resources.h
--- a/Engine/Source/Runtime/Windows/D3D11RHI/Public/D3D11Resources.h
+++ b/Engine/Source/Runtime/Windows/D3D11RHI/Public/D3D11Resources.h
@@ -339,7 +339,7 @@ public:
 	virtual ~TD3D11Texture2D();
 
 	// FRHIResource overrides
-#if RHI_WANT_RESOURCE_INFO
+#if RHI_ENABLE_RESOURCE_INFO
 	bool GetResourceInfo(FRHIResourceInfo& OutResourceInfo) const override
 	{
 		OutResourceInfo = FRHIResourceInfo{};
@@ -444,7 +444,7 @@ public:
 	virtual ~FD3D11Texture3D();
 
 	// FRHIResource overrides
-#if RHI_WANT_RESOURCE_INFO
+#if RHI_ENABLE_RESOURCE_INFO
 	bool GetResourceInfo(FRHIResourceInfo& OutResourceInfo) const override
 	{
 		OutResourceInfo = FRHIResourceInfo{};
@@ -617,7 +617,7 @@ public:
 	}
 
 	// FRHIResource overrides
-#if RHI_WANT_RESOURCE_INFO
+#if RHI_ENABLE_RESOURCE_INFO
 	bool GetResourceInfo(FRHIResourceInfo& OutResourceInfo) const override
 	{
 		OutResourceInfo = FRHIResourceInfo{};
