Filling in some gaps on RHIValidation context overloads.

diff --git a/Engine/Source/Runtime/RHI/Public/RHIValidationContext.h b/Engine/Source/Runtime/RHI/Public/RHIValidationContext.h
--- a/Engine/Source/Runtime/RHI/Public/RHIValidationContext.h
+++ b/Engine/Source/Runtime/RHI/Public/RHIValidationContext.h
@@ -1057,6 +1057,20 @@ public:
 		RHIContext->RHISetRayTracingMissShader(Scene, ShaderSlotInScene, Pipeline, ShaderIndexInPipeline, NumUniformBuffers, UniformBuffers, UserData);
 	}
 
+#if PLATFORM_USE_BACKBUFFER_WRITE_TRANSITION_TRACKING
+	virtual void RHIBackBufferWaitTrackingBeginFrame(uint64 FrameToken, bool bDeferred) override final
+	{
+		RHIContext->RHIBackBufferWaitTrackingBeginFrame(FrameToken, bDeferred);
+	}
+#endif // #if PLATFORM_USE_BACKBUFFER_WRITE_TRANSITION_TRACKING
+
+#if PLATFORM_REQUIRES_UAV_TO_RTV_TEXTURE_CACHE_FLUSH_WORKAROUND
+	virtual void RHIFlushTextureCacheBOP(FRHITexture* Texture) override final
+	{
+		RHIContext->RHIFlushTextureCacheBOP(Texture);
+	}
+#endif // PLATFORM_REQUIRES_UAV_TO_RTV_TEXTURE_CACHE_FLUSH_WORKAROUND
+
 	void SetupDrawing()
 	{
 		// nothing to validate right now
