GMaxBufferDimensions is supposed to default to the D3D11 limit, which is 2 to the power of 27, not 2 shifted by 27.

The constant in the D3D spec is D3D11_REQ_BUFFER_RESOURCE_TEXEL_COUNT_2_TO_EXP , which has a value of 27, so the value is 1 << 27, not 2 << 27.

#rnx

diff --git a/Engine/Source/Runtime/RHI/Private/RHI.cpp b/Engine/Source/Runtime/RHI/Private/RHI.cpp
--- a/Engine/Source/Runtime/RHI/Private/RHI.cpp
+++ b/Engine/Source/Runtime/RHI/Private/RHI.cpp
@@ -1367,7 +1367,7 @@ bool GRHISupportsLazyShaderCodeLoading = false;
 TRHIGlobal<int32> GMaxShadowDepthBufferSizeX(2048);
 TRHIGlobal<int32> GMaxShadowDepthBufferSizeY(2048);
 TRHIGlobal<int32> GMaxTextureDimensions(2048);
-TRHIGlobal<int64> GMaxBufferDimensions(2<<27);
+TRHIGlobal<int64> GMaxBufferDimensions(1<<27);
 TRHIGlobal<int64> GMaxComputeSharedMemory(1<<15);
 TRHIGlobal<int32> GMaxVolumeTextureDimensions(2048);
 TRHIGlobal<int32> GMaxCubeTextureDimensions(2048);
