Type truncation warning fixes for some RenderCore and RHI headers

diff --git a/Engine/Source/Runtime/Core/Public/PixelFormat.h b/Engine/Source/Runtime/Core/Public/PixelFormat.h
--- a/Engine/Source/Runtime/Core/Public/PixelFormat.h
+++ b/Engine/Source/Runtime/Core/Public/PixelFormat.h
@@ -196,4 +196,9 @@ enum class EPixelFormatChannelFlags : uint8
 
 	None = 0,
 };
-ENUM_CLASS_FLAGS(EPixelFormatChannelFlags);
\ No newline at end of file
+ENUM_CLASS_FLAGS(EPixelFormatChannelFlags);
+
+// EPixelFormat is currently used interchangably with uint8, and most call sites taking a uint8
+// should be updated to take an EPixelFormat instead, but in the interim this allows fixing
+// type conversion warnings
+#define UE_PIXELFORMAT_TO_UINT8(argument) static_cast<uint8>(argument)
diff --git a/Engine/Source/Runtime/RHI/Public/RHICommandList.h b/Engine/Source/Runtime/RHI/Public/RHICommandList.h
--- a/Engine/Source/Runtime/RHI/Public/RHICommandList.h
+++ b/Engine/Source/Runtime/RHI/Public/RHICommandList.h
@@ -3379,7 +3379,7 @@ public:
 		{
 			if (PSOContext.CachedRenderTargets[i].Texture)
 			{
-				GraphicsPSOInit.RenderTargetFormats[i] = PSOContext.CachedRenderTargets[i].Texture->GetFormat();
+				GraphicsPSOInit.RenderTargetFormats[i] = UE_PIXELFORMAT_TO_UINT8(PSOContext.CachedRenderTargets[i].Texture->GetFormat());
 				GraphicsPSOInit.RenderTargetFlags[i] = PSOContext.CachedRenderTargets[i].Texture->GetFlags();
 				const FRHITexture2DArray* TextureArray = PSOContext.CachedRenderTargets[i].Texture->GetTexture2DArray();
 			}
diff --git a/Engine/Source/Runtime/RHI/Public/RHIResources.h b/Engine/Source/Runtime/RHI/Public/RHIResources.h
--- a/Engine/Source/Runtime/RHI/Public/RHIResources.h
+++ b/Engine/Source/Runtime/RHI/Public/RHIResources.h
@@ -2464,7 +2464,7 @@ public:
 		, RasterizerState(nullptr)
 		, DepthStencilState(nullptr)
 		, RenderTargetsEnabled(0)
-		, RenderTargetFormats(InPlace, PF_Unknown)
+		, RenderTargetFormats(InPlace, UE_PIXELFORMAT_TO_UINT8(PF_Unknown))
 		, RenderTargetFlags(InPlace, TexCreate_None)
 		, DepthStencilTargetFormat(PF_Unknown)
 		, DepthStencilTargetFlag(TexCreate_None)
@@ -2505,7 +2505,7 @@ public:
 		ERenderTargetLoadAction		InStencilTargetLoadAction,
 		ERenderTargetStoreAction	InStencilTargetStoreAction,
 		FExclusiveDepthStencil		InDepthStencilAccess,
-		uint32						InNumSamples,
+		uint16						InNumSamples,
 		ESubpassHint				InSubpassHint,
 		uint8						InSubpassIndex,
 		EConservativeRasterization	InConservativeRasterization,
@@ -2969,7 +2969,7 @@ struct FRHIRenderPassInfo
 
 
 	// Color, no depth, optional resolve, optional mip, optional array slice
-	explicit FRHIRenderPassInfo(FRHITexture* ColorRT, ERenderTargetActions ColorAction, FRHITexture* ResolveRT = nullptr, uint32 InMipIndex = 0, int32 InArraySlice = -1)
+	explicit FRHIRenderPassInfo(FRHITexture* ColorRT, ERenderTargetActions ColorAction, FRHITexture* ResolveRT = nullptr, uint8 InMipIndex = 0, int32 InArraySlice = -1)
 	{
 		check(ColorRT);
 		ColorRenderTargets[0].RenderTarget = ColorRT;
@@ -3242,7 +3242,7 @@ struct RHI_API FRHITextureCreateInfo
 		EPixelFormat InFormat,
 		FClearValueBinding InClearValue,
 		ETextureCreateFlags InFlags,
-		uint32 InArraySize,
+		uint16 InArraySize,
 		uint8 InNumMips = 1,
 		uint8 InNumSamples = 1)
 	{
@@ -3275,7 +3275,7 @@ struct RHI_API FRHITextureCreateInfo
 		EPixelFormat InFormat,
 		FClearValueBinding InClearValue,
 		ETextureCreateFlags InFlags,
-		uint32 InArraySize,
+		uint16 InArraySize,
 		uint8 InNumMips = 1,
 		uint8 InNumSamples = 1)
 	{
diff --git a/Engine/Source/Runtime/RHI/Public/RHIUtilities.h b/Engine/Source/Runtime/RHI/Public/RHIUtilities.h
--- a/Engine/Source/Runtime/RHI/Public/RHIUtilities.h
+++ b/Engine/Source/Runtime/RHI/Public/RHIUtilities.h
@@ -137,7 +137,7 @@ struct FTextureRWBuffer2D
 
 		FRHIResourceCreateInfo CreateInfo(InDebugName);
 		Buffer = RHICreateTexture2D(
-			SizeX, SizeY, Format, //PF_R32_FLOAT,
+			SizeX, SizeY, UE_PIXELFORMAT_TO_UINT8(Format), //PF_R32_FLOAT,
 			/*NumMips=*/ 1,
 			1,
 			Flags,
@@ -196,7 +196,7 @@ struct FTextureRWBuffer3D
 
 		FRHIResourceCreateInfo CreateInfo(InDebugName);
 		Buffer = RHICreateTexture3D(
-			SizeX, SizeY, SizeZ, Format,
+			SizeX, SizeY, SizeZ, UE_PIXELFORMAT_TO_UINT8(Format),
 			/*NumMips=*/ 1,
 			/*Flags=*/ TexCreate_ShaderResource | TexCreate_UAV,
 			/*BulkData=*/ CreateInfo);
@@ -295,8 +295,8 @@ struct FRWBuffer
 		FRHIResourceCreateInfo CreateInfo(InDebugName);
 		CreateInfo.ResourceArray = InResourceArray;
 		Buffer = RHICreateVertexBuffer(NumBytes, BUF_UnorderedAccess | BUF_ShaderResource | AdditionalUsage, InResourceState, CreateInfo);
-		UAV = RHICreateUnorderedAccessView(Buffer, Format);
-		SRV = RHICreateShaderResourceView(Buffer, BytesPerElement, Format);
+		UAV = RHICreateUnorderedAccessView(Buffer, UE_PIXELFORMAT_TO_UINT8(Format));
+		SRV = RHICreateShaderResourceView(Buffer, BytesPerElement, UE_PIXELFORMAT_TO_UINT8(Format));
 	}
 
 	void Initialize(const TCHAR* InDebugName, uint32 BytesPerElement, uint32 NumElements, EPixelFormat Format, EBufferUsageFlags AdditionalUsage = BUF_None, FResourceArrayInterface* InResourceArray = nullptr)
@@ -353,7 +353,7 @@ struct FTextureReadBuffer2D
 		FRHIResourceCreateInfo CreateInfo(InDebugName);
 		CreateInfo.BulkData = InBulkData;
 		Buffer = RHICreateTexture2D(
-			SizeX, SizeY, Format, //PF_R32_FLOAT,
+			SizeX, SizeY, UE_PIXELFORMAT_TO_UINT8(Format), //PF_R32_FLOAT,
 			/*NumMips=*/ 1,
 			1,
 			Flags,
@@ -402,7 +402,7 @@ struct FReadBuffer
 		FRHIResourceCreateInfo CreateInfo(InDebugName);
 		CreateInfo.ResourceArray = InResourceArray;
 		Buffer = RHICreateVertexBuffer(NumBytes, BUF_ShaderResource | AdditionalUsage, ERHIAccess::SRVMask, CreateInfo);
-		SRV = RHICreateShaderResourceView(Buffer, BytesPerElement, Format);
+		SRV = RHICreateShaderResourceView(Buffer, BytesPerElement, UE_PIXELFORMAT_TO_UINT8(Format));
 	}
 
 	void Release()
diff --git a/Engine/Source/Runtime/RenderCore/Public/RenderGraphDefinitions.h b/Engine/Source/Runtime/RenderCore/Public/RenderGraphDefinitions.h
--- a/Engine/Source/Runtime/RenderCore/Public/RenderGraphDefinitions.h
+++ b/Engine/Source/Runtime/RenderCore/Public/RenderGraphDefinitions.h
@@ -239,7 +239,7 @@ public:
 	explicit inline TRDGHandle(int32 InIndex)
 	{
 		check(InIndex >= 0 && InIndex <= kNullIndex);
-		Index = InIndex;
+		Index = (IndexType)InIndex;
 	}
 
 	FORCEINLINE IndexType GetIndex() const { check(IsValid()); return Index; }
diff --git a/Engine/Source/Runtime/RenderCore/Public/RenderUtils.h b/Engine/Source/Runtime/RenderCore/Public/RenderUtils.h
--- a/Engine/Source/Runtime/RenderCore/Public/RenderUtils.h
+++ b/Engine/Source/Runtime/RenderCore/Public/RenderUtils.h
@@ -247,7 +247,7 @@ public:
 	virtual void InitRHI()
 	{
 		FVertexDeclarationElementList Elements;
-		uint32 Stride = sizeof(FVector2D);
+		uint16 Stride = sizeof(FVector2D);
 		Elements.Add(FVertexElement(0, 0, VET_Float2, 0, Stride, false));
 		VertexDeclarationRHI = RHICreateVertexDeclaration(Elements);
 	}
@@ -263,7 +263,7 @@ extern RENDERCORE_API TGlobalResource<FTileVertexDeclaration> GTileVertexDeclara
 /**
  * Maps from an X,Y,Z cube vertex coordinate to the corresponding vertex index.
  */
-inline uint16 GetCubeVertexIndex(uint32 X,uint32 Y,uint32 Z) { return X * 4 + Y * 2 + Z; }
+inline uint16 GetCubeVertexIndex(uint32 X,uint32 Y,uint32 Z) { return (uint16)(X * 4 + Y * 2 + Z); }
 
 /**
 * A 3x1 of xyz(11:11:10) format.
diff --git a/Engine/Source/Runtime/RenderCore/Public/RendererInterface.h b/Engine/Source/Runtime/RenderCore/Public/RendererInterface.h
--- a/Engine/Source/Runtime/RenderCore/Public/RendererInterface.h
+++ b/Engine/Source/Runtime/RenderCore/Public/RendererInterface.h
@@ -197,7 +197,7 @@ public:
 		ETextureCreateFlags InFlags,
 		ETextureCreateFlags InTargetableFlags,
 		bool bInForceSeparateTargetAndShaderResource,
-		uint32 InArraySize = 1,
+		uint16 InArraySize = 1,
 		uint16 InNumMips = 1,
 		bool InAutowritable = true)
 	{
@@ -235,7 +235,7 @@ public:
 		ETextureCreateFlags InFlags,
 		ETextureCreateFlags InTargetableFlags,
 		bool bInForceSeparateTargetAndShaderResource,
-		uint32 InArraySize,
+		uint16 InArraySize,
 		uint16 InNumMips = 1,
 		bool InAutowritable = true)
 	{
diff --git a/Engine/Source/Runtime/RenderCore/Public/Shader.h b/Engine/Source/Runtime/RenderCore/Public/Shader.h
--- a/Engine/Source/Runtime/RenderCore/Public/Shader.h
+++ b/Engine/Source/Runtime/RenderCore/Public/Shader.h
@@ -275,9 +275,9 @@ protected:
 	explicit FShaderMapResource(EShaderPlatform InPlatform, int32 NumShaders);
 	virtual ~FShaderMapResource();
 
-	uint32 GetAllocatedSize() const
+	SIZE_T GetAllocatedSize() const
 	{
-		uint32 Size = NumRHIShaders * sizeof(std::atomic<FRHIShader*>);
+		SIZE_T Size = NumRHIShaders * sizeof(std::atomic<FRHIShader*>);
 #if RHI_RAYTRACING
 		Size += RayTracingMaterialLibraryIndices.GetAllocatedSize();
 #endif
@@ -437,7 +437,7 @@ public:
 			TokenBits.AddZeroed();
 		}
 
-		uint8 Token = InValue;
+		uint8 Token;
 
 		// Anything that does not fit in 4 bits needs to go into FullLengths, with a special token value of 0
 		// InValue == 0 also should go into FullLengths, because its Token value is also 0
@@ -446,6 +446,10 @@ public:
 			Token = 0;
 			FullLengths.Add(InValue);
 		}
+		else
+		{
+			Token = (uint8)InValue;
+		}
 
 		const uint32 Shift = (NumTokens % 8) * 4;
 		// Add the new token bits into the existing uint32
diff --git a/Engine/Source/Runtime/RenderCore/Public/ShaderCore.h b/Engine/Source/Runtime/RenderCore/Public/ShaderCore.h
--- a/Engine/Source/Runtime/RenderCore/Public/ShaderCore.h
+++ b/Engine/Source/Runtime/RenderCore/Public/ShaderCore.h
@@ -460,7 +460,7 @@ struct FShaderCompilerEnvironment
 
 	void SetRenderTargetOutputFormat(uint32 RenderTargetIndex, EPixelFormat PixelFormat)
 	{
-		RenderTargetOutputFormatsMap.Add(RenderTargetIndex, PixelFormat);
+		RenderTargetOutputFormatsMap.Add(RenderTargetIndex, UE_PIXELFORMAT_TO_UINT8(PixelFormat));
 	}
 
 	friend FArchive& operator<<(FArchive& Ar,FShaderCompilerEnvironment& Environment)
diff --git a/Engine/Source/Runtime/RenderCore/Public/VertexFactory.h b/Engine/Source/Runtime/RenderCore/Public/VertexFactory.h
--- a/Engine/Source/Runtime/RenderCore/Public/VertexFactory.h
+++ b/Engine/Source/Runtime/RenderCore/Public/VertexFactory.h
@@ -139,8 +139,8 @@ struct FVertexStreamComponent
 	FVertexStreamComponent(const FVertexBuffer* InVertexBuffer, uint32 InOffset, uint32 InStride, EVertexElementType InType, EVertexStreamUsage Usage = EVertexStreamUsage::Default) :
 		VertexBuffer(InVertexBuffer),
 		StreamOffset(0),
-		Offset(InOffset),
-		Stride(InStride),
+		Offset((uint8)InOffset),
+		Stride((uint8)InStride),
 		Type(InType),
 		VertexStreamUsage(Usage)
 	{
@@ -151,8 +151,8 @@ struct FVertexStreamComponent
 	FVertexStreamComponent(const FVertexBuffer* InVertexBuffer, uint32 InStreamOffset, uint32 InOffset, uint32 InStride, EVertexElementType InType, EVertexStreamUsage Usage = EVertexStreamUsage::Default) :
 		VertexBuffer(InVertexBuffer),
 		StreamOffset(InStreamOffset),
-		Offset(InOffset),
-		Stride(InStride),
+		Offset((uint8)InOffset),
+		Stride((uint8)InStride),
 		Type(InType),
 		VertexStreamUsage(Usage)
 	{
