Make it easier to export to Raspberry Pi 4 (UE-106435, partially).

- Merging parts of PR #7685 by jackoalan
- Adding a new vendor id.
- Splitting ASTC/ETC2 support as ASTC may not be supported.

diff --git a/Engine/Source/Runtime/RHI/Private/RHI.cpp b/Engine/Source/Runtime/RHI/Private/RHI.cpp
--- a/Engine/Source/Runtime/RHI/Private/RHI.cpp
+++ b/Engine/Source/Runtime/RHI/Private/RHI.cpp
@@ -1740,6 +1740,7 @@ RHI_API const TCHAR* RHIVendorIdToString()
 	case 0x1010: return TEXT("ImgTec");
 	case 0x10DE: return TEXT("NVIDIA");
 	case 0x13B5: return TEXT("ARM");
+	case 0x14E4: return TEXT("Broadcom");
 	case 0x5143: return TEXT("Qualcomm");
 	case 0x8086: return TEXT("Intel");
 	default: return TEXT("Unknown");
@@ -1754,6 +1755,7 @@ RHI_API const TCHAR* RHIVendorIdToString(EGpuVendorId VendorId)
 	case EGpuVendorId::ImgTec: return TEXT("ImgTec");
 	case EGpuVendorId::Nvidia: return TEXT("NVIDIA");
 	case EGpuVendorId::Arm: return TEXT("ARM");
+	case EGpuVendorId::Broadcom: return TEXT("Broadcom");
 	case EGpuVendorId::Qualcomm: return TEXT("Qualcomm");
 	case EGpuVendorId::Intel: return TEXT("Intel");
 	case EGpuVendorId::NotQueried: return TEXT("Not Queried");
diff --git a/Engine/Source/Runtime/RHI/Public/RHIDefinitions.h b/Engine/Source/Runtime/RHI/Public/RHIDefinitions.h
--- a/Engine/Source/Runtime/RHI/Public/RHIDefinitions.h
+++ b/Engine/Source/Runtime/RHI/Public/RHIDefinitions.h
@@ -1500,6 +1500,7 @@ enum class EGpuVendorId
 	ImgTec		= 0x1010,
 	Nvidia		= 0x10DE, 
 	Arm			= 0x13B5, 
+	Broadcom	= 0x14E4,
 	Qualcomm	= 0x5143,
 	Intel		= 0x8086,
 	Mesa		= 0x10005,
@@ -2126,6 +2127,7 @@ inline EGpuVendorId RHIConvertToGpuVendorId(uint32 VendorId)
 	case EGpuVendorId::ImgTec:
 	case EGpuVendorId::Nvidia:
 	case EGpuVendorId::Arm:
+	case EGpuVendorId::Broadcom:
 	case EGpuVendorId::Qualcomm:
 	case EGpuVendorId::Intel:
 		return (EGpuVendorId)VendorId;
diff --git a/Engine/Source/Runtime/VulkanRHI/Private/Android/VulkanAndroidPlatform.h b/Engine/Source/Runtime/VulkanRHI/Private/Android/VulkanAndroidPlatform.h
--- a/Engine/Source/Runtime/VulkanRHI/Private/Android/VulkanAndroidPlatform.h
+++ b/Engine/Source/Runtime/VulkanRHI/Private/Android/VulkanAndroidPlatform.h
@@ -69,6 +69,7 @@ public:
 
 	static bool SupportsBCTextureFormats() { return false; }
 	static bool SupportsASTCTextureFormats() { return true; }
+	static bool SupportsETC2TextureFormats() { return true; }
 	static bool SupportsQuerySurfaceProperties() { return false; }
 
 	static void SetupFeatureLevels()
diff --git a/Engine/Source/Runtime/VulkanRHI/Private/VulkanDevice.cpp b/Engine/Source/Runtime/VulkanRHI/Private/VulkanDevice.cpp
--- a/Engine/Source/Runtime/VulkanRHI/Private/VulkanDevice.cpp
+++ b/Engine/Source/Runtime/VulkanRHI/Private/VulkanDevice.cpp
@@ -799,7 +799,10 @@ void FVulkanDevice::SetupFormats()
 		{
 			SetComponentMapping(PF_ASTC_12x12, VK_COMPONENT_SWIZZLE_R, VK_COMPONENT_SWIZZLE_G, VK_COMPONENT_SWIZZLE_B, VK_COMPONENT_SWIZZLE_A);
 		}
+	}
 
+	if (FVulkanPlatform::SupportsETC2TextureFormats())
+	{
 		MapFormatSupport(PF_ETC2_RGB, VK_FORMAT_ETC2_R8G8B8_UNORM_BLOCK);
 		if (GPixelFormats[PF_ETC2_RGB].Supported)
 		{
diff --git a/Engine/Source/Runtime/VulkanRHI/Private/VulkanGenericPlatform.h b/Engine/Source/Runtime/VulkanRHI/Private/VulkanGenericPlatform.h
--- a/Engine/Source/Runtime/VulkanRHI/Private/VulkanGenericPlatform.h
+++ b/Engine/Source/Runtime/VulkanRHI/Private/VulkanGenericPlatform.h
@@ -37,9 +37,10 @@ public:
 	// create the platform-specific surface object - required
 	static void CreateSurface(VkSurfaceKHR* OutSurface);
 
-	// most platforms support BC* but not ASTC*
+	// most platforms support BC* but not ASTC* or ETC2*
 	static bool SupportsBCTextureFormats() { return true; }
 	static bool SupportsASTCTextureFormats() { return false; }
+	static bool SupportsETC2TextureFormats() { return false; }
 
 	// most platforms can query the surface for the present mode, and size, etc
 	static bool SupportsQuerySurfaceProperties() { return true; }
