Specify new Max Shader Target for iOS 15

diff --git a/Engine/Config/BaseEngine.ini b/Engine/Config/BaseEngine.ini
--- a/Engine/Config/BaseEngine.ini
+++ b/Engine/Config/BaseEngine.ini
@@ -2477,7 +2477,7 @@ bAutomaticSigning=false
 UseFastIntrinsics=False
 ForceFloats=False
 EnableMathOptimisations=True
-MaxShaderLanguageVersion=5
+MaxShaderLanguageVersion=7
 bDisableMotionData=False
 bEnableAdvertisingIdentifier=True
 
diff --git a/Engine/Source/Runtime/RHI/Private/RHI.cpp b/Engine/Source/Runtime/RHI/Private/RHI.cpp
--- a/Engine/Source/Runtime/RHI/Private/RHI.cpp
+++ b/Engine/Source/Runtime/RHI/Private/RHI.cpp
@@ -1798,19 +1798,13 @@ RHI_API uint32 RHIGetShaderLanguageVersion(const FStaticShaderPlatform Platform)
 			static int32 MaxShaderVersion = -1;
 			if (MaxShaderVersion < 0)
 			{
-				MaxShaderVersion = 6;
+				MaxShaderVersion = 7;
 				int32 MinShaderVersion = 5;
 				if(!GConfig->GetInt(TEXT("/Script/IOSRuntimeSettings.IOSRuntimeSettings"), TEXT("MaxShaderLanguageVersion"), MaxShaderVersion, GEngineIni))
 				{
 					MaxShaderVersion = 0;
 				}
                 
-                // If we are using Mobile desktop rendering, we need a minimum of Metal 2.1
-                if(IsMetalSM5Platform(Platform))
-                {
-                    MinShaderVersion = 4;
-                }
-                
 				MaxShaderVersion = FMath::Max(MinShaderVersion, MaxShaderVersion);
 			}
 			Version = (uint32)MaxShaderVersion;
