Removing deprecated EShaderPlatform entries.

diff --git a/Engine/Source/Runtime/RHI/Public/RHIDefinitions.h b/Engine/Source/Runtime/RHI/Public/RHIDefinitions.h
--- a/Engine/Source/Runtime/RHI/Public/RHIDefinitions.h
+++ b/Engine/Source/Runtime/RHI/Public/RHIDefinitions.h
@@ -53,32 +53,17 @@ static_assert(SF_NumFrequencies <= (1 << SF_NumBits), "SF_NumFrequencies will no
 enum EShaderPlatform
 {
 	SP_PCD3D_SM5					= 0,
-	SP_OPENGL_SM4_REMOVED			UE_DEPRECATED(4.27, "ShaderPlatform is removed; please don't use.") = 1,
-	SP_PS4_REMOVED					UE_DEPRECATED(4.27, "ShaderPlatform is removed; please don't use.") = 2,
-	SP_OPENGL_PCES2_REMOVED			UE_DEPRECATED(4.27, "ShaderPlatform is removed; please don't use.") = 3,
-	SP_XBOXONE_D3D12_REMOVED		UE_DEPRECATED(4.27, "ShaderPlatform is removed; please don't use.") = 4,
-	SP_PCD3D_SM4_REMOVED			UE_DEPRECATED(4.27, "ShaderPlatform is removed; please don't use.") = 5,
-	SP_OPENGL_SM5_REMOVED			UE_DEPRECATED(4.27, "ShaderPlatform is removed; please don't use.") = 6,
-	SP_PCD3D_ES2_REMOVED			UE_DEPRECATED(4.27, "ShaderPlatform is removed; please don't use.") = 7,
-	SP_OPENGL_ES2_ANDROID_REMOVED	UE_DEPRECATED(4.27, "ShaderPlatform is removed; please don't use.") = 8,
-	SP_OPENGL_ES2_WEBGL_REMOVED		UE_DEPRECATED(4.27, "ShaderPlatform is removed; please don't use.") = 9, 
-	SP_OPENGL_ES2_IOS_REMOVED		UE_DEPRECATED(4.27, "ShaderPlatform is removed; please don't use.") = 10,
 	SP_METAL						= 11,
 	SP_METAL_MRT					= 12,
-	SP_OPENGL_ES31_EXT_REMOVED		UE_DEPRECATED(4.27, "ShaderPlatform is removed; please don't use.") = 13,
 	SP_PCD3D_ES3_1					= 14,
 	SP_OPENGL_PCES3_1				= 15,
 	SP_METAL_SM5					= 16,
 	SP_VULKAN_PCES3_1				= 17,
 	SP_METAL_SM5_NOTESS_REMOVED		= 18,
-	SP_VULKAN_SM4_REMOVED			UE_DEPRECATED(4.27, "ShaderPlatform is removed; please don't use.") = 19,
 	SP_VULKAN_SM5					= 20,
 	SP_VULKAN_ES3_1_ANDROID			= 21,
 	SP_METAL_MACES3_1 				= 22,
-	SP_METAL_MACES2_REMOVED			UE_DEPRECATED(4.27, "ShaderPlatform is removed; please don't use.") = 23,
 	SP_OPENGL_ES3_1_ANDROID			= 24,
-	SP_SWITCH_REMOVED				UE_DEPRECATED(4.27, "ShaderPlatform is removed; please don't use.") = 25,
-	SP_SWITCH_FORWARD_REMOVED		UE_DEPRECATED(4.27, "ShaderPlatform is removed; please don't use.") = 26,
 	SP_METAL_MRT_MAC				= 27,
 	SP_VULKAN_SM5_LUMIN_REMOVED		UE_DEPRECATED(5.0, "ShaderPlatform is removed; please don't use.") = 28,
 	SP_VULKAN_ES3_1_LUMIN_REMOVED	UE_DEPRECATED(5.0, "ShaderPlatform is removed; please don't use.") = 29,
@@ -1833,15 +1818,6 @@ inline bool IsConsolePlatform(const FStaticShaderPlatform Platform)
 	return FDataDrivenShaderPlatformInfo::GetIsConsole(Platform);
 }
 
-UE_DEPRECATED(4.27, "IsSwitchPlatform() is deprecated; please use DataDrivenShaderPlatformInfo instead.")
-inline bool IsSwitchPlatform(const FStaticShaderPlatform Platform)
-{
-	PRAGMA_DISABLE_DEPRECATION_WARNINGS
-	return Platform == SP_SWITCH_REMOVED || Platform == SP_SWITCH_FORWARD_REMOVED || 
-		FDataDrivenShaderPlatformInfo::GetIsLanguageNintendo(Platform);	
-	PRAGMA_ENABLE_DEPRECATION_WARNINGS
-}
-
 inline bool IsVulkanPlatform(const FStaticShaderPlatform Platform)
 {
 	return Platform == SP_VULKAN_SM5 || Platform == SP_VULKAN_PCES3_1 
@@ -1881,9 +1857,6 @@ inline bool IsD3DPlatform(const FStaticShaderPlatform Platform)
 	switch (Platform)
 	{
 	case SP_PCD3D_SM5:
-PRAGMA_DISABLE_DEPRECATION_WARNINGS
-	case SP_XBOXONE_D3D12_REMOVED:
-PRAGMA_ENABLE_DEPRECATION_WARNINGS
 		return true;
 	default:
 		return FDataDrivenShaderPlatformInfo::GetIsLanguageD3D(Platform);
