Make GMinClipZ and GProjectionSignY constant to always match D3D clip-space.

#rnx

diff --git a/Engine/Source/Runtime/OpenGLDrv/Private/OpenGLDevice.cpp b/Engine/Source/Runtime/OpenGLDrv/Private/OpenGLDevice.cpp
--- a/Engine/Source/Runtime/OpenGLDrv/Private/OpenGLDevice.cpp
+++ b/Engine/Source/Runtime/OpenGLDrv/Private/OpenGLDevice.cpp
@@ -876,11 +876,6 @@ static void InitRHICapabilitiesForGL()
 	GShaderPlatformForFeatureLevel[ERHIFeatureLevel::SM4_REMOVED] = SP_NumPlatforms;
 	GShaderPlatformForFeatureLevel[ERHIFeatureLevel::SM5] = SP_NumPlatforms;
 
-	// Set to same values as in DX11, as for the time being clip space adjustment are done entirely
-	// in HLSLCC-generated shader code and OpenGLDrv.
-	GMinClipZ = 0.0f;
-	GProjectionSignY = 1.0f;
-
 	// Disable texture streaming on devices with ES3.1 or lower, unless we have the GL_APPLE_copy_texture_levels extension or support for glCopyImageSubData
 	if (GMaxRHIFeatureLevel <= ERHIFeatureLevel::ES3_1)
 	{
diff --git a/Engine/Source/Runtime/RHI/Private/RHI.cpp b/Engine/Source/Runtime/RHI/Private/RHI.cpp
--- a/Engine/Source/Runtime/RHI/Private/RHI.cpp
+++ b/Engine/Source/Runtime/RHI/Private/RHI.cpp
@@ -1363,8 +1363,6 @@ bool GRHISupportsDrawIndirect = true;
 bool GRHISupportsMultithreading = false;
 bool GRHISupportsUpdateFromBufferTexture = false;
 bool GSupportsWideMRT = true;
-float GMinClipZ = 0.0f;
-float GProjectionSignY = 1.0f;
 bool GRHINeedsExtraDeletionLatency = false;
 bool GRHIForceNoDeletionLatencyForStreamingTextures = false;
 TRHIGlobal<int32> GMaxComputeDispatchDimension((1 << 16) - 1);
diff --git a/Engine/Source/Runtime/RHI/Public/RHI.h b/Engine/Source/Runtime/RHI/Public/RHI.h
--- a/Engine/Source/Runtime/RHI/Public/RHI.h
+++ b/Engine/Source/Runtime/RHI/Public/RHI.h
@@ -427,11 +427,11 @@ extern RHI_API bool GRHISupportsDirectGPUMemoryLock;
 /** true if the multi-threaded shader creation is supported by (or desirable for) the RHI. */
 extern RHI_API bool GRHISupportsMultithreadedShaderCreation;
 
-/** The minimum Z value in clip space for the RHI. */
-extern RHI_API float GMinClipZ;
+/** The minimum Z value in clip space for the RHI. This is a constant value to always match D3D clip-space. */
+constexpr float GMinClipZ = 0.0f;
 
-/** The sign to apply to the Y axis of projection matrices. */
-extern RHI_API float GProjectionSignY;
+/** The sign to apply to the Y axis of projection matrices. This is a constant value to always match D3D clip-space. */
+constexpr float GProjectionSignY = 1.0f;
 
 /** Does this RHI need to wait for deletion of resources due to ref counting. */
 extern RHI_API bool GRHINeedsExtraDeletionLatency;
diff --git a/Engine/Source/Runtime/VulkanRHI/Private/VulkanRHI.cpp b/Engine/Source/Runtime/VulkanRHI/Private/VulkanRHI.cpp
--- a/Engine/Source/Runtime/VulkanRHI/Private/VulkanRHI.cpp
+++ b/Engine/Source/Runtime/VulkanRHI/Private/VulkanRHI.cpp
@@ -830,8 +830,6 @@ void FVulkanDynamicRHI::InitInstance()
 
 		FHardwareInfo::RegisterHardwareInfo(NAME_RHI, TEXT("Vulkan"));
 
-		GProjectionSignY = 1.0f;
-
 		SavePipelineCacheCmd = IConsoleManager::Get().RegisterConsoleCommand(
 			TEXT("r.Vulkan.SavePipelineCache"),
 			TEXT("Save pipeline cache."),
