Add download link to blacklist driver popup

diff --git a/Engine/Source/Runtime/RHI/Private/DynamicRHI.cpp b/Engine/Source/Runtime/RHI/Private/DynamicRHI.cpp
--- a/Engine/Source/Runtime/RHI/Private/DynamicRHI.cpp
+++ b/Engine/Source/Runtime/RHI/Private/DynamicRHI.cpp
@@ -152,23 +152,28 @@ static void RHIDetectAndWarnOfBadDrivers(bool bHasEditorToken)
 
 			// Note: we don't localize the vendor's name.
 			FString VendorString = DriverInfo.ProviderName;
+			FText HyperlinkText;
 			if (DriverInfo.IsNVIDIA())
 			{
 				VendorString = TEXT("NVIDIA");
+				HyperlinkText = NSLOCTEXT("MessageDialog", "DriverDownloadLinkNVIDIA", "https://www.nvidia.com/en-us/geforce/drivers/");
 			}
 			else if (DriverInfo.IsAMD())
 			{
 				VendorString = TEXT("AMD");
+				HyperlinkText = NSLOCTEXT("MessageDialog", "DriverDownloadLinkAMD", "https://www.amd.com/en/support");
 			}
 			else if (DriverInfo.IsIntel())
 			{
 				VendorString = TEXT("Intel");
+				HyperlinkText = NSLOCTEXT("MessageDialog", "DriverDownloadLinkIntel", "https://downloadcenter.intel.com/product/80939/Graphics");
 			}
 
 			// format message box UI
 			FFormatNamedArguments Args;
 			Args.Add(TEXT("AdapterName"), FText::FromString(DriverInfo.DeviceDescription));
 			Args.Add(TEXT("Vendor"), FText::FromString(VendorString));
+			Args.Add(TEXT("Hyperlink"), HyperlinkText);
 			Args.Add(TEXT("RecommendedVer"), FText::FromString(DetectedGPUHardware.GetSuggestedDriverVersion(DriverInfo.RHIName)));
 			Args.Add(TEXT("InstalledVer"), FText::FromString(DriverInfo.UserDriverVersion));
 
@@ -176,11 +181,11 @@ static void RHIDetectAndWarnOfBadDrivers(bool bHasEditorToken)
 			FText LocalizedMsg;
 			if (bLatestBlacklisted)
 			{
-				LocalizedMsg = FText::Format(NSLOCTEXT("MessageDialog", "LatestVideoCardDriverIssueReport","The latest version of the {Vendor} graphics driver has known issues.\nPlease install the recommended driver version.\n\n{AdapterName}\nInstalled: {InstalledVer}\nRecommended: {RecommendedVer}"),Args);
+				LocalizedMsg = FText::Format(NSLOCTEXT("MessageDialog", "LatestVideoCardDriverIssueReport","The latest version of the {Vendor} graphics driver has known issues.\nPlease install the recommended driver version.\n\n{Hyperlink}\n\n{AdapterName}\nInstalled: {InstalledVer}\nRecommended: {RecommendedVer}"),Args);
 			}
 			else
 			{
-				LocalizedMsg = FText::Format(NSLOCTEXT("MessageDialog", "VideoCardDriverIssueReport","The installed version of the {Vendor} graphics driver has known issues.\nPlease install either the latest or the recommended driver version.\n\n{AdapterName}\nInstalled: {InstalledVer}\nRecommended: {RecommendedVer}"),Args);
+				LocalizedMsg = FText::Format(NSLOCTEXT("MessageDialog", "VideoCardDriverIssueReport","The installed version of the {Vendor} graphics driver has known issues.\nPlease install either the latest or the recommended driver version.\n\n{Hyperlink}\n\n{AdapterName}\nInstalled: {InstalledVer}\nRecommended: {RecommendedVer}"),Args);
 			}
 
 			FPlatformMisc::MessageBoxExt(EAppMsgType::Ok,
@@ -595,7 +600,7 @@ void FShaderResourceViewInitializer::InitType()
 			Type = EType::VertexBufferSRV;
 		}
 		else if (Usage & BUF_IndexBuffer)
-		{
+	{
 			Type = EType::IndexBufferSRV;
 		}
 		else if (Usage & BUF_AccelerationStructure)
