Enable the user pipeline cache on D3D12 and evict elements out of it that haven't been used for a set period

diff --git a/Engine/Source/Runtime/RHI/Private/PipelineFileCache.cpp b/Engine/Source/Runtime/RHI/Private/PipelineFileCache.cpp
--- a/Engine/Source/Runtime/RHI/Private/PipelineFileCache.cpp
+++ b/Engine/Source/Runtime/RHI/Private/PipelineFileCache.cpp
@@ -21,6 +21,7 @@ PipelineFileCache.cpp: Pipeline state cache implementation.
 #include "ProfilingDebugging/CsvProfiler.h"
 #include "String/LexFromString.h"
 #include "String/ParseTokens.h"
+#include "Misc/ScopeExit.h"
 
 static FString JOURNAL_FILE_EXTENSION(TEXT(".jnl"));
 
@@ -66,12 +67,13 @@ enum class EPipelineCacheFileFormatVersions : uint32
 	AlphaToCoverage = 19,
 	AddingMeshShaders = 20,
 	RemovingTessellationShaders = 21,
+	LastUsedTime = 22,
 };
 
 const uint64 FPipelineCacheFileFormatMagic = 0x5049504543414348; // PIPECACH
 const uint64 FPipelineCacheTOCFileFormatMagic = 0x544F435354415232; // TOCSTAR2
 const uint64 FPipelineCacheEOFFileFormatMagic = 0x454F462D4D41524B; // EOF-MARK
-const uint32 FPipelineCacheFileFormatCurrentVersion = (uint32)EPipelineCacheFileFormatVersions::RemovingTessellationShaders;
+const uint32 FPipelineCacheFileFormatCurrentVersion = (uint32)EPipelineCacheFileFormatVersions::LastUsedTime;
 const int32  FPipelineCacheGraphicsDescPartsNum = 63; // parser will expect this number of parts in a description string
 
 /**
@@ -114,6 +116,20 @@ static TAutoConsoleVariable<int32> CVarPSOFileCacheSaveUserCache(
                                                             ECVF_Default | ECVF_RenderThreadSafe
                                                             );
 
+static TAutoConsoleVariable<int32> CVarPSOFileCacheUserCacheUnusedElementRetainDays(
+                                                            TEXT("r.ShaderPipelineCache.UserCacheUnusedElementRetainDays"),
+															30,
+                                                            TEXT("The amount of time in days to keep unused PSO entries in the cache."),
+                                                            ECVF_Default
+                                                            );
+
+static TAutoConsoleVariable<int32> CVarPSOFileCacheUserCacheUnusedElementCheckPeriod(
+                                                            TEXT("r.ShaderPipelineCache.UserCacheUnusedElementCheckPeriod"),
+															-1,
+                                                            TEXT("The amount of time in days between running the garbage collection on unused PSOs in the user cache. Use a negative value to disable."),
+                                                            ECVF_Default
+                                                            );
+
 static TAutoConsoleVariable<int32> CVarLazyLoadShadersWhenPSOCacheIsPresent(
 	TEXT("r.ShaderPipelineCache.LazyLoadShadersWhenPSOCacheIsPresent"),
 	0,
@@ -148,6 +164,11 @@ bool FPipelineFileCache::FileCacheEnabled = false;
 FPipelineFileCache::FPipelineStateLoggedEvent FPipelineFileCache::PSOLoggedEvent;
 uint64 FPipelineFileCache::GameUsageMask = 0;
 
+static int64 GetCurrentUnixTime()
+{
+	return FDateTime::UtcNow().ToUnixTimestamp();
+}
+
 bool DefaultPSOMaskComparisonFunction(uint64 ReferenceMask, uint64 PSOMask)
 {
 	return (ReferenceMask & PSOMask) == ReferenceMask;
@@ -182,6 +203,7 @@ struct FPipelineCacheFileFormatHeader
 	TEnumAsByte<EShaderPlatform> Platform; // The shader platform for all referenced PSOs.
 	FGuid Guid;				// Guid to identify the file uniquely
 	uint64 TableOffset;		// absolute file offset to TOC
+	int64 LastGCUnixTime;   // Last time that the cache was scanned to remove out of date elements.
 	
 	friend FArchive& operator<<(FArchive& Ar, FPipelineCacheFileFormatHeader& Info)
 	{
@@ -191,7 +213,12 @@ struct FPipelineCacheFileFormatHeader
 		Ar << Info.Platform;
 		Ar << Info.Guid;
 		Ar << Info.TableOffset;
-		
+
+		if (Info.Version >= (uint32)EPipelineCacheFileFormatVersions::LastUsedTime)
+		{
+			Ar << Info.LastGCUnixTime;
+		}
+
 		return Ar;
 	}
 };
@@ -217,6 +244,7 @@ struct FPipelineCacheFileFormatPSOMetaData
 	FPipelineCacheFileFormatPSOMetaData()
 	: FileOffset(0)
 	, UsageMask(0)
+	, LastUsedUnixTime(0)
 	, EngineFlags(0)
 	{
 	}
@@ -231,7 +259,53 @@ struct FPipelineCacheFileFormatPSOMetaData
 	FPipelineStateStats Stats;
 	TSet<FSHAHash> Shaders;
 	uint64 UsageMask;
+	int64  LastUsedUnixTime;
 	uint16 EngineFlags;
+
+	void AddShaders(const FPipelineCacheFileFormatPSO& NewEntry)
+	{
+		switch (NewEntry.Type)
+		{
+			case FPipelineCacheFileFormatPSO::DescriptorType::Compute:
+			{
+				INC_DWORD_STAT(STAT_SerializedComputePipelineStateCount);
+				Shaders.Add(NewEntry.ComputeDesc.ComputeShader);
+				break;
+			}
+			case FPipelineCacheFileFormatPSO::DescriptorType::Graphics:
+			{
+				INC_DWORD_STAT(STAT_SerializedGraphicsPipelineStateCount);
+
+				if (NewEntry.GraphicsDesc.VertexShader != FSHAHash())
+					Shaders.Add(NewEntry.GraphicsDesc.VertexShader);
+
+				if (NewEntry.GraphicsDesc.FragmentShader != FSHAHash())
+					Shaders.Add(NewEntry.GraphicsDesc.FragmentShader);
+
+				if (NewEntry.GraphicsDesc.GeometryShader != FSHAHash())
+					Shaders.Add(NewEntry.GraphicsDesc.GeometryShader);
+
+				if (NewEntry.GraphicsDesc.MeshShader != FSHAHash())
+					Shaders.Add(NewEntry.GraphicsDesc.MeshShader);
+
+				if (NewEntry.GraphicsDesc.AmplificationShader != FSHAHash())
+					Shaders.Add(NewEntry.GraphicsDesc.AmplificationShader);
+
+				break;
+			}
+			case FPipelineCacheFileFormatPSO::DescriptorType::RayTracing:
+			{
+				INC_DWORD_STAT(STAT_SerializedRayTracingPipelineStateCount);
+				Shaders.Add(NewEntry.RayTracingDesc.ShaderHash);
+				break;
+			}
+			default:
+			{
+				check(false);
+				break;
+			}
+		}
+	}
 	
 	friend FArchive& operator<<(FArchive& Ar, FPipelineCacheFileFormatPSOMetaData& Info)
 	{
@@ -271,7 +345,12 @@ struct FPipelineCacheFileFormatPSOMetaData
 		{
 			Ar << Info.EngineFlags;
 		}
-		
+
+		if (Ar.GameNetVer() >= (uint32)EPipelineCacheFileFormatVersions::LastUsedTime)
+		{
+			Ar << Info.LastUsedUnixTime;
+		}
+
 		return Ar;
 	}
 };
@@ -1446,9 +1525,9 @@ public:
 		FArchive* FileReader = IFileManager::Get().CreateFileReader(*FilePath);
 		if (FileReader)
 		{
-			FPipelineCacheFileFormatHeader Header;
+			FileReader->SetGameNetVer(FPipelineCacheFileFormatCurrentVersion);
+			FPipelineCacheFileFormatHeader Header = {};
 			*FileReader << Header;
-            FileReader->SetGameNetVer(FPipelineCacheFileFormatCurrentVersion);
 			if (Header.Magic == FPipelineCacheFileFormatMagic && Header.Version == FPipelineCacheFileFormatCurrentVersion && Header.GameVersion == GameVersion && Header.Platform == ShaderPlatform)
 			{
 				check(Header.TableOffset > 0);
@@ -1494,7 +1573,7 @@ public:
 				}
 				else
 				{
-					UE_LOG(LogRHI, Log, TEXT("Filed to create async read file handle to FPipelineCacheFile: %s (GUID: %s)"), *FilePath, *Header.Guid.ToString());
+					UE_LOG(LogRHI, Log, TEXT("Failed to create async read file handle to FPipelineCacheFile: %s (GUID: %s)"), *FilePath, *Header.Guid.ToString());
 					bSuccess = false;
 				}
 			}
@@ -1506,6 +1585,184 @@ public:
 		
 		return bSuccess;
 	}
+
+	void GarbageCollectUserCache(FString const& UserCacheFilePath)
+	{
+		UE_LOG(LogRHI, Log, TEXT("FPipelineCacheFile: GarbageCollectUserCache() Begin"));
+		ON_SCOPE_EXIT{ UE_LOG(LogRHI, Log, TEXT("FPipelineCacheFile: GarbageCollectUserCache() End")); };
+
+		int32 GCPeriodInDays = CVarPSOFileCacheUserCacheUnusedElementCheckPeriod.GetValueOnAnyThread();
+		if (GCPeriodInDays < 0)
+		{
+			UE_LOG(LogRHI, Log, TEXT("User cache GC is disabled"));
+			return;
+		}
+
+		FArchive* FileReader = IFileManager::Get().CreateFileReader(*UserCacheFilePath);
+		if (!FileReader)
+		{
+			UE_LOG(LogRHI, Log, TEXT("No user cache file found"));
+			return;
+		}
+
+		ON_SCOPE_EXIT
+		{
+			if (FileReader)
+			{
+				FileReader->Close();
+				delete FileReader;
+			}
+		};
+
+		FileReader->SetGameNetVer(FPipelineCacheFileFormatCurrentVersion);
+		FPipelineCacheFileFormatHeader Header;
+		*FileReader << Header;
+
+		if (!(Header.Magic == FPipelineCacheFileFormatMagic && Header.Version == FPipelineCacheFileFormatCurrentVersion && Header.GameVersion == GameVersion && Header.Platform == ShaderPlatform))
+		{
+			UE_LOG(LogRHI, Error, TEXT("File has invalid or out of date header"));
+			return;
+		}
+
+		FTimespan GCPeriod = FTimespan::FromDays(GCPeriodInDays);
+		int64 NextGCTime = Header.LastGCUnixTime + GCPeriod.GetTotalSeconds();
+		const int64 UnixTime = GetCurrentUnixTime();
+		if (UnixTime < NextGCTime)
+		{
+			const FTimespan TimespanToNextGC = FTimespan::FromSeconds(NextGCTime - UnixTime);
+			const double DaysToNextGC = TimespanToNextGC.GetTotalDays();
+			UE_LOG(LogRHI, Log, TEXT("Next GC on user cache is in %0.3f days."), DaysToNextGC);
+			return;
+		}
+
+		FPipelineCacheFileFormatTOC Content;
+		if (Header.TableOffset < (uint64)FileReader->TotalSize())
+		{
+			FileReader->Seek(Header.TableOffset);
+			*FileReader << Content;
+
+			// FPipelineCacheFileFormatTOC archive read can set the FArchive to error on failure
+			if (FileReader->IsError())
+			{
+				UE_LOG(LogRHI, Log, TEXT("Failed to read TOC"));
+				return;
+			}
+		}
+
+		int64 StaleDays = CVarPSOFileCacheUserCacheUnusedElementRetainDays.GetValueOnAnyThread();
+		FTimespan StaleTimespan = FTimespan::FromDays(StaleDays);
+		int64 EvictionTime = UnixTime - (int64)StaleTimespan.GetTotalSeconds();
+
+		auto EntryShouldBeRemovedFromUserCache = [&Header, GameFileGuid=this->GameFileGuid, EvictionTime](const FPipelineCacheFileFormatPSOMetaData& MetaData)
+		{
+			// Remove the element if it is in the user cache and the time has elapsed, or if it was in a cache that no longer exists.
+			if (MetaData.FileGuid == Header.Guid)
+			{
+				return EvictionTime >= MetaData.LastUsedUnixTime;
+			}
+			else
+			{
+				return MetaData.FileGuid != GameFileGuid;
+			}
+		};
+
+		int32 NumOutOfDateEntries = 0;
+		for (auto const& Entry : Content.MetaData)
+		{
+			if (EntryShouldBeRemovedFromUserCache(Entry.Value))
+			{
+				NumOutOfDateEntries++;
+			}
+		}
+
+		if (NumOutOfDateEntries == 0)
+		{
+			UE_LOG(LogRHI, Log, TEXT("No out of date entries."));
+			return;
+		}
+
+		if (NumOutOfDateEntries == Content.MetaData.Num())
+		{
+			FileReader->Close();
+			delete FileReader;
+			FileReader = nullptr;
+			if (IFileManager::Get().FileExists(*UserCacheFilePath))
+			{
+				IFileManager::Get().Delete(*UserCacheFilePath);
+			}
+			UE_LOG(LogRHI, Log, TEXT("All entries are out of date, recreating cache."));
+			return;
+		}
+
+		UE_LOG(LogRHI, Log, TEXT("%d/%d elements are out of date, performing GC."), NumOutOfDateEntries, Content.MetaData.Num());
+		
+		TArray<uint8> Buffer;
+		FMemoryWriter MemoryWriter(Buffer);
+		MemoryWriter.SetGameNetVer(FPipelineCacheFileFormatCurrentVersion);
+
+		FPipelineCacheFileFormatHeader NewHeader = Header;
+		check(NewHeader.Magic == FPipelineCacheFileFormatMagic);
+		check(NewHeader.Platform == ShaderPlatform);
+		NewHeader.Version = FPipelineCacheFileFormatCurrentVersion;
+		NewHeader.LastGCUnixTime = UnixTime;
+		NewHeader.TableOffset = 0; // Will overwrite with the correct offset after building the TOC
+		MemoryWriter << NewHeader;
+
+		FPipelineCacheFileFormatTOC NewTOC;
+		NewTOC.SortedOrder = Content.SortedOrder; // Removal maintains sort order of existing cache
+
+		for (auto const& Entry : Content.MetaData)
+		{
+			if (!EntryShouldBeRemovedFromUserCache(Entry.Value))
+			{
+				// Copy the meta data, and the FPipelineCacheFileFormatPSO if it exists in the user cache (the meta data can point into the game cache).
+				FPipelineCacheFileFormatPSOMetaData NewEntry = Entry.Value;
+				if (Entry.Value.FileGuid == Header.Guid && Entry.Value.FileSize > 0)
+				{
+					NewEntry.FileSize = Entry.Value.FileSize;
+					NewEntry.FileOffset = MemoryWriter.Tell();
+
+					// Copy from file to new memory writer
+					FPipelineCacheFileFormatPSO ExistingPSO;
+
+					FileReader->Seek(Entry.Value.FileOffset);
+					*FileReader << ExistingPSO;
+					MemoryWriter << ExistingPSO;
+				}
+
+				NewTOC.MetaData.Add(Entry.Key, NewEntry);
+			}
+		}
+
+		NewHeader.TableOffset = MemoryWriter.Tell();
+		MemoryWriter << NewTOC;
+		MemoryWriter.Seek(0);
+		MemoryWriter << NewHeader;
+
+		UE_LOG(LogRHI, Log, TEXT("Deleting existing cache file"));
+
+		int64 OriginalSize = FileReader->TotalSize();
+		FileReader->Close();
+		delete FileReader;
+		FileReader = nullptr;
+		if (IFileManager::Get().FileExists(*UserCacheFilePath))
+		{
+			IFileManager::Get().Delete(*UserCacheFilePath);
+		}
+
+		FArchive* FileWriter = IFileManager::Get().CreateFileWriter(*UserCacheFilePath);
+		if (!FileWriter)
+		{
+			UE_LOG(LogRHI, Log, TEXT("Unable to open new cache file for writing"));
+			return;
+		}
+
+		int64 NewSize = MemoryWriter.TotalSize();
+		FileWriter->Serialize(Buffer.GetData(), MemoryWriter.TotalSize());
+		FileWriter->Close();
+		delete FileWriter;
+		UE_LOG(LogRHI, Log, TEXT("Rewrote cache file. Old Size %lld, new size %lld (%lld byte reduction)"), OriginalSize, NewSize, OriginalSize - NewSize);
+	}
 	
     bool ShouldDeleteExistingUserCache()
     {
@@ -1524,7 +1781,7 @@ public:
     {
 		return FPipelineFileCache::LogPSOtoFileCache() && (CVarPSOFileCacheSaveUserCache.GetValueOnAnyThread() > 0);
     }
-    
+
 	bool OpenPipelineFileCache(FString const& FileName, EShaderPlatform Platform, FGuid& OutGameFileGuid)
 	{
 		SET_DWORD_STAT(STAT_TotalGraphicsPipelineStateCount, 0);
@@ -1599,6 +1856,8 @@ public:
         
         if (ShouldLoadUserCache())
         {
+			GarbageCollectUserCache(FilePath);
+
 			FPipelineCacheFileFormatTOC UserTOC;
             bUserFileOk = OpenPipelineFileCache(FilePath, UserFileGuid, UserAsyncFileHandle, UserTOC);
             if (!bUserFileOk)
@@ -1666,7 +1925,7 @@ public:
 		return bGameFileOk || bUserFileOk;
 	}
 	
-	void MergePSOUsageToMetaData(TMap<uint32, FPSOUsageData>& NewPSOUsage, TMap<uint32, FPipelineCacheFileFormatPSOMetaData>& MetaData, bool bRemoveUpdatedentries = false)
+	void MergePSOUsageToMetaData(TMap<uint32, FPSOUsageData>& NewPSOUsage, TMap<uint32, FPipelineCacheFileFormatPSOMetaData>& MetaData, int64 CurrentUnixTime, bool bRemoveUpdatedentries = false)
 	{
 		for(auto It = NewPSOUsage.CreateIterator(); It; ++It)
 		{
@@ -1678,6 +1937,7 @@ public:
 			{
 				PSOMetaData->UsageMask |= MaskEntry.Value.UsageMask;
 				PSOMetaData->EngineFlags |= MaskEntry.Value.EngineFlags;
+				PSOMetaData->LastUsedUnixTime = CurrentUnixTime;
 				
 				if(bRemoveUpdatedentries)
 				{
@@ -1710,6 +1970,7 @@ public:
 		{
             uint32 NumNewEntries = 0;
             
+			int64 UnixTime = GetCurrentUnixTime();
 			FString JournalPath;
 			if (Mode != FPipelineFileCache::SaveMode::BoundPSOsOnly)
 			{
@@ -1727,6 +1988,7 @@ public:
 					Header.Platform = ShaderPlatform;
 					Header.Guid = UserFileGuid;
 					Header.TableOffset = 0;
+					Header.LastGCUnixTime = UnixTime;
 
 					*JournalWriter << Header;
 				}
@@ -1824,6 +2086,7 @@ public:
 						Header.Platform = ShaderPlatform;
 						Header.Guid = UserFileGuid;
 						Header.TableOffset = 0;
+						Header.LastGCUnixTime = UnixTime;
 		 
 						*FileWriter << Header;
 		 
@@ -1851,7 +2114,7 @@ public:
                                 check(!IsPSOEntryCached(NewEntry));
                                 
                                 uint32 PSOHash = GetTypeHash(NewEntry);
-								
+
 								FPipelineStateStats const* Stat = Stats.FindRef(PSOHash);
                                 if (Stat && Stat->TotalBindCount > 0)
 								{
@@ -1859,48 +2122,7 @@ public:
 									Meta.Stats.PSOHash = PSOHash;
 									Meta.FileOffset = PSOOffset;
 									Meta.FileGuid = UserFileGuid;
-									
-									switch(NewEntry.Type)
-									{
-										case FPipelineCacheFileFormatPSO::DescriptorType::Compute:
-										{
-											INC_DWORD_STAT(STAT_SerializedComputePipelineStateCount);
-											Meta.Shaders.Add(NewEntry.ComputeDesc.ComputeShader);
-											break;
-										}
-										case FPipelineCacheFileFormatPSO::DescriptorType::Graphics:
-										{
-											INC_DWORD_STAT(STAT_SerializedGraphicsPipelineStateCount);
-											
-											if (NewEntry.GraphicsDesc.VertexShader != FSHAHash())
-												Meta.Shaders.Add(NewEntry.GraphicsDesc.VertexShader);
-											
-											if (NewEntry.GraphicsDesc.FragmentShader != FSHAHash())
-												Meta.Shaders.Add(NewEntry.GraphicsDesc.FragmentShader);
-
-											if (NewEntry.GraphicsDesc.GeometryShader != FSHAHash())
-												Meta.Shaders.Add(NewEntry.GraphicsDesc.GeometryShader);
-
-											if (NewEntry.GraphicsDesc.MeshShader != FSHAHash())
-												Meta.Shaders.Add(NewEntry.GraphicsDesc.MeshShader);
-
-											if (NewEntry.GraphicsDesc.AmplificationShader != FSHAHash())
-												Meta.Shaders.Add(NewEntry.GraphicsDesc.AmplificationShader);
-
-											break;
-										}
-										case FPipelineCacheFileFormatPSO::DescriptorType::RayTracing:
-										{
-											INC_DWORD_STAT(STAT_SerializedRayTracingPipelineStateCount);
-											Meta.Shaders.Add(NewEntry.RayTracingDesc.ShaderHash);
-											break;
-										}
-										default:
-										{
-											check(false);
-											break;
-										}
-									}
+									Meta.AddShaders(NewEntry);
 									
 									TArray<uint8> Bytes;
 									FMemoryWriter Wr(Bytes);
@@ -1934,7 +2156,7 @@ public:
                             }
 							
 							// Update TOC Metadata usage and clear relevant entries in NewPSOUsage as we are saving this file cache TOC
-							MergePSOUsageToMetaData(NewPSOUsage, TOC.MetaData, true);
+							MergePSOUsageToMetaData(NewPSOUsage, TOC.MetaData, UnixTime, true);
 							
                             Header.TableOffset = PSOOffset;
                             TOCOffset = PSOOffset;
@@ -1964,55 +2186,14 @@ public:
                                 Meta.FileOffset = 0;
                                 Meta.FileSize = 0;
                                 Meta.FileGuid = Header.Guid;
-								
-                                switch(Entry.Type)
-                                {
-                                    case FPipelineCacheFileFormatPSO::DescriptorType::Compute:
-                                    {
-                                        INC_DWORD_STAT(STAT_SerializedComputePipelineStateCount);
-                                        Meta.Shaders.Add(Entry.ComputeDesc.ComputeShader);
-                                        break;
-                                    }
-                                    case FPipelineCacheFileFormatPSO::DescriptorType::Graphics:
-                                    {
-                                        INC_DWORD_STAT(STAT_SerializedGraphicsPipelineStateCount);
-                                        
-                                        if (Entry.GraphicsDesc.VertexShader != FSHAHash())
-                                            Meta.Shaders.Add(Entry.GraphicsDesc.VertexShader);
-                                        
-                                        if (Entry.GraphicsDesc.FragmentShader != FSHAHash())
-                                            Meta.Shaders.Add(Entry.GraphicsDesc.FragmentShader);
-
-                                        if (Entry.GraphicsDesc.GeometryShader != FSHAHash())
-											Meta.Shaders.Add(Entry.GraphicsDesc.GeometryShader);
-
-										if (Entry.GraphicsDesc.MeshShader != FSHAHash())
-											Meta.Shaders.Add(Entry.GraphicsDesc.MeshShader);
-
-										if (Entry.GraphicsDesc.AmplificationShader != FSHAHash())
-											Meta.Shaders.Add(Entry.GraphicsDesc.AmplificationShader);
-                                        
-                                        break;
-                                    }
-									case FPipelineCacheFileFormatPSO::DescriptorType::RayTracing:
-									{
-										INC_DWORD_STAT(STAT_SerializedRayTracingPipelineStateCount);
-										Meta.Shaders.Add(Entry.RayTracingDesc.ShaderHash);
-										break;
-									}
-                                    default:
-                                    {
-                                        check(false);
-                                        break;
-                                    }
-                                }
+								Meta.AddShaders(Entry);
 								
                                 TempTOC.MetaData.Add(Meta.Stats.PSOHash, Meta);
                                 PSOs.Add(Meta.Stats.PSOHash, Entry);
                             }
 							
 							// Update TOC Metadata usage masks - don't clear NewPSOUsage as we are using a TempTOC
-							MergePSOUsageToMetaData(NewPSOUsage, TempTOC.MetaData);
+							MergePSOUsageToMetaData(NewPSOUsage, TempTOC.MetaData, UnixTime);
                             
                             for (auto& Pair : Stats)
                             {
@@ -2151,55 +2332,14 @@ public:
                                 Meta.FileOffset = 0;
                                 Meta.FileSize = 0;
                                 Meta.FileGuid = UserFileGuid;
-								
-                                switch(Entry.Type)
-                                {
-                                    case FPipelineCacheFileFormatPSO::DescriptorType::Compute:
-                                    {
-                                        INC_DWORD_STAT(STAT_SerializedComputePipelineStateCount);
-                                        Meta.Shaders.Add(Entry.ComputeDesc.ComputeShader);
-                                        break;
-                                    }
-                                    case FPipelineCacheFileFormatPSO::DescriptorType::Graphics:
-                                    {
-                                        INC_DWORD_STAT(STAT_SerializedGraphicsPipelineStateCount);
-                                        
-                                        if (Entry.GraphicsDesc.VertexShader != FSHAHash())
-                                            Meta.Shaders.Add(Entry.GraphicsDesc.VertexShader);
-                                        
-                                        if (Entry.GraphicsDesc.FragmentShader != FSHAHash())
-                                            Meta.Shaders.Add(Entry.GraphicsDesc.FragmentShader);
-
-                                        if (Entry.GraphicsDesc.GeometryShader != FSHAHash())
-											Meta.Shaders.Add(Entry.GraphicsDesc.GeometryShader);
-
-										if (Entry.GraphicsDesc.MeshShader != FSHAHash())
-											Meta.Shaders.Add(Entry.GraphicsDesc.MeshShader);
-
-										if (Entry.GraphicsDesc.AmplificationShader != FSHAHash())
-											Meta.Shaders.Add(Entry.GraphicsDesc.AmplificationShader);
-                                        
-                                        break;
-                                    }
-									case FPipelineCacheFileFormatPSO::DescriptorType::RayTracing:
-									{
-										INC_DWORD_STAT(STAT_SerializedRayTracingPipelineStateCount);
-										Meta.Shaders.Add(Entry.RayTracingDesc.ShaderHash);
-										break;
-									}
-                                    default:
-                                    {
-                                        check(false);
-                                        break;
-                                    }
-                                }
+								Meta.AddShaders(Entry);
                                 
                                 TOC.MetaData.Add(Meta.Stats.PSOHash, Meta);
                                 PSOs.Add(Meta.Stats.PSOHash, Entry);
                             }
 							
 							// Update TOC Metadata usage and clear updated entries in NewPSOUsage as file TOC is getting updated
-							MergePSOUsageToMetaData(NewPSOUsage, TOC.MetaData, true);
+							MergePSOUsageToMetaData(NewPSOUsage, TOC.MetaData, UnixTime, true);
 							
 							FPipelineCacheFileFormatTOC TempTOC = TOC;
 							// Update PSO usage stats for new and old
diff --git a/Engine/Source/Runtime/RHI/Private/Windows/WindowsDynamicRHI.cpp b/Engine/Source/Runtime/RHI/Private/Windows/WindowsDynamicRHI.cpp
--- a/Engine/Source/Runtime/RHI/Private/Windows/WindowsDynamicRHI.cpp
+++ b/Engine/Source/Runtime/RHI/Private/Windows/WindowsDynamicRHI.cpp
@@ -315,6 +315,9 @@ static IDynamicRHIModule* LoadDynamicRHIModule(ERHIFeatureLevel::Type& DesiredFe
 
 			auto PSOFileCacheUserCacheCVar = IConsoleManager::Get().FindTConsoleVariableDataInt(TEXT("r.ShaderPipelineCache.SaveUserCache"));
 			*PSOFileCacheUserCacheCVar = UE_BUILD_SHIPPING;
+
+			auto PSOFileCacheLogCVar = IConsoleManager::Get().FindTConsoleVariableDataInt(TEXT("r.ShaderPipelineCache.LogPSO"));
+			*PSOFileCacheLogCVar = UE_BUILD_SHIPPING;
 #endif
 
 			if (FPlatformProcess::IsApplicationRunning(TEXT("fraps.exe")))
