Fixed build warning related to using FRHIVertexBuffer. Now using FRHIBuffer.

diff --git a/Engine/Plugins/FX/Niagara/Source/Niagara/Classes/NiagaraEmitterInstanceBatcher.h b/Engine/Plugins/FX/Niagara/Source/Niagara/Classes/NiagaraEmitterInstanceBatcher.h
--- a/Engine/Plugins/FX/Niagara/Source/Niagara/Classes/NiagaraEmitterInstanceBatcher.h
+++ b/Engine/Plugins/FX/Niagara/Source/Niagara/Classes/NiagaraEmitterInstanceBatcher.h
@@ -291,7 +291,7 @@ private:
 
 #if WITH_MGPU
 	static const FName TemporalEffectName;
-	TArray<FRHIVertexBuffer*> TemporalEffectBuffers;
+	TArray<FRHIBuffer*> TemporalEffectBuffers;
 	ENiagaraGpuComputeTickStage::Type StageToWaitForTemporalEffect = ENiagaraGpuComputeTickStage::First;
 	ENiagaraGpuComputeTickStage::Type StageToBroadcastTemporalEffect = ENiagaraGpuComputeTickStage::First;
 
diff --git a/Engine/Source/Runtime/D3D12RHI/Private/D3D12Commands.cpp b/Engine/Source/Runtime/D3D12RHI/Private/D3D12Commands.cpp
--- a/Engine/Source/Runtime/D3D12RHI/Private/D3D12Commands.cpp
+++ b/Engine/Source/Runtime/D3D12RHI/Private/D3D12Commands.cpp
@@ -2412,7 +2412,7 @@ void FD3D12CommandContext::RHIBroadcastTemporalEffect(const FName& InEffectName,
 #endif // WITH_MGPU
 }
 
-void FD3D12CommandContext::RHIBroadcastTemporalEffect(const FName& InEffectName, const TArrayView<FRHIVertexBuffer*> InBuffers)
+void FD3D12CommandContext::RHIBroadcastTemporalEffect(const FName& InEffectName, const TArrayView<FRHIBuffer*> InBuffers)
 {
 #if WITH_MGPU
 	FMemMark Mark(FMemStack::Get());
diff --git a/Engine/Source/Runtime/RHI/Public/RHICommandList.h b/Engine/Source/Runtime/RHI/Public/RHICommandList.h
--- a/Engine/Source/Runtime/RHI/Public/RHICommandList.h
+++ b/Engine/Source/Runtime/RHI/Public/RHICommandList.h
@@ -3068,7 +3068,7 @@ public:
 		ALLOC_COMMAND(FRHICommandBroadcastTemporalEffect<FRHITexture>)(EffectName, AllocArray(Textures));
 	}
 
-	FORCEINLINE_DEBUGGABLE void BroadcastTemporalEffect(const FName& EffectName, const TArrayView<FRHIVertexBuffer*> Buffers)
+	FORCEINLINE_DEBUGGABLE void BroadcastTemporalEffect(const FName& EffectName, const TArrayView<FRHIBuffer*> Buffers)
 	{
 		//check(IsOutsideRenderPass());
 		if (Bypass())
@@ -3077,7 +3077,7 @@ public:
 			return;
 		}
 
-		ALLOC_COMMAND(FRHICommandBroadcastTemporalEffect<FRHIVertexBuffer>)(EffectName, AllocArray(Buffers));
+		ALLOC_COMMAND(FRHICommandBroadcastTemporalEffect<FRHIBuffer>)(EffectName, AllocArray(Buffers));
 	}
 #endif // WITH_MGPU
 
diff --git a/Engine/Source/Runtime/RHI/Public/RHICommandListCommandExecutes.inl b/Engine/Source/Runtime/RHI/Public/RHICommandListCommandExecutes.inl
--- a/Engine/Source/Runtime/RHI/Public/RHICommandListCommandExecutes.inl
+++ b/Engine/Source/Runtime/RHI/Public/RHICommandListCommandExecutes.inl
@@ -103,7 +103,7 @@ template <> void FRHICommandBroadcastTemporalEffect<FRHITexture>::Execute(FRHICo
 	RHISTAT(BroadcastTemporalEffect);
 	INTERNAL_DECORATOR(RHIBroadcastTemporalEffect)(EffectName, Resources);
 }
-template <> void FRHICommandBroadcastTemporalEffect<FRHIVertexBuffer>::Execute(FRHICommandListBase& CmdList)
+template <> void FRHICommandBroadcastTemporalEffect<FRHIBuffer>::Execute(FRHICommandListBase& CmdList)
 {
 	RHISTAT(BroadcastTemporalEffect);
 	INTERNAL_DECORATOR(RHIBroadcastTemporalEffect)(EffectName, Resources);
diff --git a/Engine/Source/Runtime/RHI/Public/RHIContext.h b/Engine/Source/Runtime/RHI/Public/RHIContext.h
--- a/Engine/Source/Runtime/RHI/Public/RHIContext.h
+++ b/Engine/Source/Runtime/RHI/Public/RHIContext.h
@@ -309,7 +309,7 @@ public:
 		/* empty default implementation */
 	}
 
-	virtual void RHIBroadcastTemporalEffect(const FName& InEffectName, const TArrayView<FRHIVertexBuffer*> InBuffers)
+	virtual void RHIBroadcastTemporalEffect(const FName& InEffectName, const TArrayView<FRHIBuffer*> InBuffers)
 	{
 		/* empty default implementation */
 	}
