SNORM/UNORM types aren't included in IsFloatFormat.  Create new utilities and use them to fix Vulkan validation error (uint vs float in clear pass).

diff --git a/Engine/Plugins/Runtime/HairStrands/Source/HairStrandsCore/Private/GroomResources.cpp b/Engine/Plugins/Runtime/HairStrands/Source/HairStrandsCore/Private/GroomResources.cpp
--- a/Engine/Plugins/Runtime/HairStrands/Source/HairStrandsCore/Private/GroomResources.cpp
+++ b/Engine/Plugins/Runtime/HairStrands/Source/HairStrandsCore/Private/GroomResources.cpp
@@ -324,7 +324,7 @@ void InternalCreateVertexBufferRDG(FRDGBuilder& GraphBuilder, uint32 InVertexCou
 	FRDGBufferDesc Desc = FRDGBufferDesc::CreateBufferDesc(FormatType::SizeInByte, InVertexCount);
 	FRDGBufferRef Buffer = GraphBuilder.CreateBuffer(Desc, DebugName, ERDGBufferFlags::MultiFrame);
 
-	if (IsFloatFormat(FormatType::Format))
+	if (IsFloatFormat(FormatType::Format) || IsUnormFormat(FormatType::Format) || IsSnormFormat(FormatType::Format))
 	{
 		AddClearUAVFloatPass(GraphBuilder, GraphBuilder.CreateUAV(Buffer, FormatType::Format), 0.0f);
 	}
diff --git a/Engine/Source/Runtime/RHI/Public/RHIUtilities.h b/Engine/Source/Runtime/RHI/Public/RHIUtilities.h
--- a/Engine/Source/Runtime/RHI/Public/RHIUtilities.h
+++ b/Engine/Source/Runtime/RHI/Public/RHIUtilities.h
@@ -106,6 +106,36 @@ static bool IsFloatFormat(EPixelFormat Format)
 	return false;
 }
 
+static bool IsUnormFormat(EPixelFormat Format)
+{
+	switch (Format)
+	{
+	case PF_R5G6B5_UNORM:
+	case PF_R16G16B16A16_UNORM:
+	case PF_B5G5R5A1_UNORM:
+		return true;
+
+	default:
+		break;
+	}
+	return false;
+}
+
+static bool IsSnormFormat(EPixelFormat Format)
+{
+	switch (Format)
+	{
+	case PF_R8G8B8A8_SNORM:
+	case PF_R16G16B16A16_SNORM:
+	case PF_G16R16_SNORM:
+		return true;
+
+	default:
+		break;
+	}
+	return false;
+}
+
 /** Get the best default resource state for the given texture creation flags */
 extern RHI_API ERHIAccess RHIGetDefaultResourceState(ETextureCreateFlags InUsage, bool bInHasInitialData);
 
