Always allow buffers to be the source of a copy in Vulkan.  Remove ensure() for BUF_SourceCopy on all RHIs, moved it to RHIValidation.  All this together should allow r.DumpGPU to run cleanly in DX11/DX12/VK in regular circumstances.

diff --git a/Engine/Source/Runtime/D3D12RHI/Private/D3D12Commands.cpp b/Engine/Source/Runtime/D3D12RHI/Private/D3D12Commands.cpp
--- a/Engine/Source/Runtime/D3D12RHI/Private/D3D12Commands.cpp
+++ b/Engine/Source/Runtime/D3D12RHI/Private/D3D12Commands.cpp
@@ -645,9 +645,6 @@ void FD3D12CommandContext::RHICopyToStagingBuffer(FRHIBuffer* SourceBufferRHI, F
 	FD3D12Buffer* VertexBuffer = RetrieveObject<FD3D12Buffer>(SourceBufferRHI);
 	check(VertexBuffer);
 
-	ensureMsgf(EnumHasAnyFlags(SourceBufferRHI->GetUsage(), BUF_SourceCopy), TEXT("Buffers used as copy source need to be created with BUF_SourceCopy"));
-
-
 	// Ensure our shadow buffer is large enough to hold the readback.
 	if (!StagingBuffer->ResourceLocation.IsValid() || StagingBuffer->ShadowBufferSize < NumBytes)
 	{
diff --git a/Engine/Source/Runtime/RHI/Private/RHIValidation.cpp b/Engine/Source/Runtime/RHI/Private/RHIValidation.cpp
--- a/Engine/Source/Runtime/RHI/Private/RHIValidation.cpp
+++ b/Engine/Source/Runtime/RHI/Private/RHIValidation.cpp
@@ -477,6 +477,31 @@ void FValidationContext::RHIEndFrame()
 	State.PreviousEndFrame = RHIValidation::CaptureBacktrace();
 }
 
+namespace RHIValidation
+{
+	static inline FString GetReasonString_SourceCopyFlagMissing(FRHIBuffer* Buffer)
+	{
+		return FString::Printf(TEXT("Buffers used as copy source need to be created with BUF_SourceCopy! Resource: \"%s\" (0x%p)."), 
+			(Buffer->GetName().GetStringLength() > 0) ? *Buffer->GetName().ToString() : TEXT("Unnamed"), Buffer);
+	}
+}
+
+void FValidationContext::RHICopyToStagingBuffer(FRHIBuffer* SourceBufferRHI, FRHIStagingBuffer* DestinationStagingBufferRHI, uint32 InOffset, uint32 InNumBytes)
+{
+	using namespace RHIValidation;
+	Tracker->Assert(SourceBufferRHI->GetWholeResourceIdentity(), ERHIAccess::CopySrc);
+	RHI_VALIDATION_CHECK(EnumHasAnyFlags(SourceBufferRHI->GetUsage(), BUF_SourceCopy), *GetReasonString_SourceCopyFlagMissing(SourceBufferRHI));
+	RHIContext->RHICopyToStagingBuffer(SourceBufferRHI, DestinationStagingBufferRHI, InOffset, InNumBytes);
+}
+
+void FValidationComputeContext::RHICopyToStagingBuffer(FRHIBuffer* SourceBufferRHI, FRHIStagingBuffer* DestinationStagingBufferRHI, uint32 InOffset, uint32 InNumBytes)
+{
+	using namespace RHIValidation;
+	Tracker->Assert(SourceBufferRHI->GetWholeResourceIdentity(), ERHIAccess::CopySrc);
+	RHI_VALIDATION_CHECK(EnumHasAnyFlags(SourceBufferRHI->GetUsage(), BUF_SourceCopy), *GetReasonString_SourceCopyFlagMissing(SourceBufferRHI));
+	RHIContext->RHICopyToStagingBuffer(SourceBufferRHI, DestinationStagingBufferRHI, InOffset, InNumBytes);
+}
+
 void FValidationContext::FState::Reset()
 {
 	bInsideBeginRenderPass = false;
diff --git a/Engine/Source/Runtime/RHI/Public/RHIValidationContext.h b/Engine/Source/Runtime/RHI/Public/RHIValidationContext.h
--- a/Engine/Source/Runtime/RHI/Public/RHIValidationContext.h
+++ b/Engine/Source/Runtime/RHI/Public/RHIValidationContext.h
@@ -205,11 +205,7 @@ public:
 		RHIContext->RHIWriteGPUFence(FenceRHI);
 	}
 
-	virtual void RHICopyToStagingBuffer(FRHIBuffer* SourceBufferRHI, FRHIStagingBuffer* DestinationStagingBufferRHI, uint32 InOffset, uint32 InNumBytes) override final
-	{
-		Tracker->Assert(SourceBufferRHI->GetWholeResourceIdentity(), ERHIAccess::CopySrc); // @todo: is CopySrc correct?
-		RHIContext->RHICopyToStagingBuffer(SourceBufferRHI, DestinationStagingBufferRHI, InOffset, InNumBytes);
-	}
+	virtual void RHICopyToStagingBuffer(FRHIBuffer* SourceBufferRHI, FRHIStagingBuffer* DestinationStagingBufferRHI, uint32 InOffset, uint32 InNumBytes) override final;
 
 	inline void LinkToContext(IRHIComputeContext* PlatformContext)
 	{
@@ -921,11 +917,7 @@ public:
 		RHIContext->RHIWriteGPUFence(FenceRHI);
 	}
 
-	virtual void RHICopyToStagingBuffer(FRHIBuffer* SourceBufferRHI, FRHIStagingBuffer* DestinationStagingBufferRHI, uint32 InOffset, uint32 InNumBytes) override final
-	{
-		Tracker->Assert(SourceBufferRHI->GetWholeResourceIdentity(), ERHIAccess::CopySrc); // @todo: is CopySrc correct?
-		RHIContext->RHICopyToStagingBuffer(SourceBufferRHI, DestinationStagingBufferRHI, InOffset, InNumBytes);
-	}
+	virtual void RHICopyToStagingBuffer(FRHIBuffer* SourceBufferRHI, FRHIStagingBuffer* DestinationStagingBufferRHI, uint32 InOffset, uint32 InNumBytes) override final;
 
 	virtual void RHICopyTexture(FRHITexture* SourceTexture, FRHITexture* DestTexture, const FRHICopyTextureInfo& CopyInfo) override final
 	{
diff --git a/Engine/Source/Runtime/VulkanRHI/Private/VulkanBuffer.cpp b/Engine/Source/Runtime/VulkanRHI/Private/VulkanBuffer.cpp
--- a/Engine/Source/Runtime/VulkanRHI/Private/VulkanBuffer.cpp
+++ b/Engine/Source/Runtime/VulkanRHI/Private/VulkanBuffer.cpp
@@ -94,6 +94,9 @@ FVulkanResourceMultiBuffer::FVulkanResourceMultiBuffer(FVulkanDevice* InDevice,
 {
 	VULKAN_TRACK_OBJECT_CREATE(FVulkanResourceMultiBuffer, this);
 
+	// Always include TRANSFER_SRC since hardware vendors confirmed it wouldn't have any performance cost and we need it for some debug functionalities.
+	BufferUsageFlags |= VK_BUFFER_USAGE_TRANSFER_SRC_BIT;
+
 	if (EnumHasAnyFlags(InUEUsage, BUF_VertexBuffer))
 	{
 		BufferUsageFlags |= VK_BUFFER_USAGE_VERTEX_BUFFER_BIT;
@@ -117,14 +120,12 @@ FVulkanResourceMultiBuffer::FVulkanResourceMultiBuffer(FVulkanDevice* InDevice,
 		const bool bUAV = EnumHasAnyFlags(InUEUsage, BUF_UnorderedAccess);
 		const bool bIndirect = EnumHasAllFlags(InUEUsage, BUF_DrawIndirect);
 		const bool bCPUReadable = EnumHasAnyFlags(InUEUsage, BUF_KeepCPUAccessible);
-		const bool bCopySource = EnumHasAnyFlags(InUEUsage, BUF_SourceCopy);
 
 		BufferUsageFlags |= bVolatile ? 0 : VK_BUFFER_USAGE_TRANSFER_DST_BIT;
 		BufferUsageFlags |= (bShaderResource && !bIsUniformBuffer) ? VK_BUFFER_USAGE_UNIFORM_TEXEL_BUFFER_BIT : 0;
 		BufferUsageFlags |= bUAV ? VK_BUFFER_USAGE_STORAGE_TEXEL_BUFFER_BIT : 0;
 		BufferUsageFlags |= bIndirect ? VK_BUFFER_USAGE_INDIRECT_BUFFER_BIT : 0;
 		BufferUsageFlags |= bCPUReadable ? (VK_BUFFER_USAGE_TRANSFER_SRC_BIT | VK_BUFFER_USAGE_TRANSFER_DST_BIT) : 0;
-		BufferUsageFlags |= bCopySource ? VK_BUFFER_USAGE_TRANSFER_SRC_BIT : 0;
 
 #if VULKAN_RHI_RAYTRACING
 		BufferUsageFlags |= VK_BUFFER_USAGE_SHADER_DEVICE_ADDRESS_BIT;
diff --git a/Engine/Source/Runtime/VulkanRHI/Private/VulkanCommands.cpp b/Engine/Source/Runtime/VulkanRHI/Private/VulkanCommands.cpp
--- a/Engine/Source/Runtime/VulkanRHI/Private/VulkanCommands.cpp
+++ b/Engine/Source/Runtime/VulkanRHI/Private/VulkanCommands.cpp
@@ -825,7 +825,6 @@ void FVulkanCommandListContext::RHICopyToStagingBuffer(FRHIBuffer* SourceBufferR
 	FVulkanResourceMultiBuffer* VertexBuffer = ResourceCast(SourceBufferRHI);
 
 	ensure(CmdBuffer->IsOutsideRenderPass());
-	ensureMsgf(EnumHasAnyFlags(SourceBufferRHI->GetUsage(), BUF_SourceCopy), TEXT("Buffers used as copy source need to be created with BUF_SourceCopy"));
 
 	FVulkanStagingBuffer* StagingBuffer = ResourceCast(StagingBufferRHI);
 	if (!StagingBuffer->StagingBuffer || StagingBuffer->StagingBuffer->GetSize() < NumBytes) //-V1051
diff --git a/Engine/Source/Runtime/Windows/D3D11RHI/Private/D3D11Commands.cpp b/Engine/Source/Runtime/Windows/D3D11RHI/Private/D3D11Commands.cpp
--- a/Engine/Source/Runtime/Windows/D3D11RHI/Private/D3D11Commands.cpp
+++ b/Engine/Source/Runtime/Windows/D3D11RHI/Private/D3D11Commands.cpp
@@ -1994,8 +1994,6 @@ void FD3D11DynamicRHI::RHICopyToStagingBuffer(FRHIBuffer* SourceBufferRHI, FRHIS
 		ensureMsgf(!StagingBuffer->bIsLocked, TEXT("Attempting to Copy to a locked staging buffer. This may have undefined behavior"));
 		if (SourceBuffer)
 		{
-			ensureMsgf(EnumHasAnyFlags(SourceBufferRHI->GetUsage(), BUF_SourceCopy), TEXT("Buffers used as copy source need to be created with BUF_SourceCopy"));
-
 			if (!StagingBuffer->StagedRead || StagingBuffer->ShadowBufferSize < NumBytes)
 			{
 				// Free previously allocated buffer.
