Cherrypick of RHI related changes from 17953988:
* Add SupportsInlineRayTracing to DataDrivenPlatformInfo and add RHISupportsInlineRayTracing to RHI

diff --git a/Engine/Source/Runtime/RHI/Private/RHI.cpp b/Engine/Source/Runtime/RHI/Private/RHI.cpp
--- a/Engine/Source/Runtime/RHI/Private/RHI.cpp
+++ b/Engine/Source/Runtime/RHI/Private/RHI.cpp
@@ -1397,6 +1397,7 @@ bool GRHISupportsRayTracingPSOAdditions = false;
 bool GRHISupportsRayTracingDispatchIndirect = false;
 bool GRHISupportsRayTracingAsyncBuildAccelerationStructure = false;
 bool GRHISupportsRayTracingAMDHitToken = false;
+bool GRHISupportsInlineRayTracing = false;
 uint32 GRHIRayTracingAccelerationStructureAlignment = 0;
 uint32 GRHIRayTracingScratchBufferAlignment = 0;
 uint32 GRHIRayTracingShaderTableAlignment = 0;
@@ -2229,6 +2230,7 @@ void FGenericDataDrivenShaderPlatformInfo::ParseDataDrivenShaderInfo(const FConf
 	GET_SECTION_BOOL_HELPER(bSupports4ComponentUAVReadWrite);
 	GET_SECTION_BOOL_HELPER(bSupportsRenderTargetWriteMask);
 	GET_SECTION_BOOL_HELPER(bSupportsRayTracing);
+	GET_SECTION_BOOL_HELPER(bSupportsInlineRayTracing);
 	GET_SECTION_BOOL_HELPER(bSupportsRayTracingProceduralPrimitive);
 	GET_SECTION_BOOL_HELPER(bSupportsRayTracingIndirectInstanceData);
 	GET_SECTION_BOOL_HELPER(bSupportsPathTracing);
diff --git a/Engine/Source/Runtime/RHI/Public/RHI.h b/Engine/Source/Runtime/RHI/Public/RHI.h
--- a/Engine/Source/Runtime/RHI/Public/RHI.h
+++ b/Engine/Source/Runtime/RHI/Public/RHI.h
@@ -592,6 +592,9 @@ extern RHI_API bool GRHISupportsRayTracingAsyncBuildAccelerationStructure;
 /** Whether or not the RHI supports the AMD Hit Token extension. */
 extern RHI_API bool GRHISupportsRayTracingAMDHitToken;
 
+/** Whether or not the RHI supports inline ray tracing. */
+extern RHI_API bool GRHISupportsInlineRayTracing;
+
 /** Required alignment for ray tracing acceleration structures. */
 extern RHI_API uint32 GRHIRayTracingAccelerationStructureAlignment;
 
diff --git a/Engine/Source/Runtime/RHI/Public/RHIDefinitions.h b/Engine/Source/Runtime/RHI/Public/RHIDefinitions.h
--- a/Engine/Source/Runtime/RHI/Public/RHIDefinitions.h
+++ b/Engine/Source/Runtime/RHI/Public/RHIDefinitions.h
@@ -378,6 +378,7 @@ class RHI_API FGenericDataDrivenShaderPlatformInfo
 	uint32 bOverrideFMaterial_NeedsGBufferEnabled : 1;
 	uint32 bSupportsMobileDistanceField : 1;
 	uint32 bSupportsFFTBloom : 1;
+	uint32 bSupportsInlineRayTracing : 1;
 	uint32 bSupportsVertexShaderLayer : 1;
 	uint32 bSupportsBindless : 1;
 
@@ -570,6 +571,12 @@ public:
 		return Infos[Platform].bSupportsRayTracing;
 	}
 
+	static FORCEINLINE_DEBUGGABLE const bool GetSupportsInlineRayTracing(const FStaticShaderPlatform Platform)
+	{
+		check(IsValid(Platform));
+		return Infos[Platform].bSupportsInlineRayTracing;
+	}
+
 	static FORCEINLINE_DEBUGGABLE const bool GetSupportsRayTracingProceduralPrimitive(const FStaticShaderPlatform Platform)
 	{
 		check(IsValid(Platform));
