Add separate RHI capability properties for full ray tracing shaders and inline ray tracing. While most RHIs may support both features, some may only support one.

diff --git a/Engine/Config/VulkanPC/DataDrivenPlatformInfo.ini b/Engine/Config/VulkanPC/DataDrivenPlatformInfo.ini
--- a/Engine/Config/VulkanPC/DataDrivenPlatformInfo.ini
+++ b/Engine/Config/VulkanPC/DataDrivenPlatformInfo.ini
@@ -50,6 +50,8 @@ bSupports4ComponentUAVReadWrite = false
 bSupportsRenderTargetWriteMask = false
 
 bSupportsRayTracing = false
+bSupportsRayTracingShaders = false
+bSupportsInlineRayTracing = false
 bSupportsRayTracingIndirectInstanceData = false
 bSupportsGPUSkinCache = true
 bSupportsByteBufferComputeShaders = true
diff --git a/Engine/Config/Windows/DataDrivenPlatformInfo.ini b/Engine/Config/Windows/DataDrivenPlatformInfo.ini
--- a/Engine/Config/Windows/DataDrivenPlatformInfo.ini
+++ b/Engine/Config/Windows/DataDrivenPlatformInfo.ini
@@ -60,6 +60,8 @@ MaxFeatureLevel=SM5
 bIsPC=true
 bSupportsDebugViewShaders=true
 bSupportsRayTracing=true
+bSupportsRayTracingShaders=true
+bSupportsInlineRayTracing=false
 bSupportsRayTracingProceduralPrimitive=true
 bSupportsRayTracingIndirectInstanceData=true
 bSupportsPathTracing=true
@@ -101,6 +103,8 @@ bIsPC=true
 MaxMeshShaderThreadGroupSize=128
 bSupportsDebugViewShaders=true
 bSupportsRayTracing=true
+bSupportsRayTracingShaders=true
+bSupportsInlineRayTracing=false
 bSupportsRayTracingProceduralPrimitive=true
 bSupportsRayTracingIndirectInstanceData=true
 bSupportsPathTracing=true
diff --git a/Engine/Source/Editor/DetailCustomizations/Private/DebugCameraControllerSettingsCustomization.cpp b/Engine/Source/Editor/DetailCustomizations/Private/DebugCameraControllerSettingsCustomization.cpp
--- a/Engine/Source/Editor/DetailCustomizations/Private/DebugCameraControllerSettingsCustomization.cpp
+++ b/Engine/Source/Editor/DetailCustomizations/Private/DebugCameraControllerSettingsCustomization.cpp
@@ -51,7 +51,7 @@ void FDebugCameraControllerSettingsViewModeIndexCustomization::CustomizeHeader(T
 	EnumRestriction->AddHiddenValue(ViewModeIndexEnum->GetNameStringByValue((uint8)EViewModeIndex::VMI_VisualizeNanite));
 	EnumRestriction->AddHiddenValue(ViewModeIndexEnum->GetNameStringByValue((uint8)EViewModeIndex::VMI_StationaryLightOverlap));
 #if RHI_RAYTRACING
-	if (!GRHISupportsRayTracing)
+	if (!GRHISupportsRayTracing || !GRHISupportsRayTracingShaders)
 	{
 		EnumRestriction->AddHiddenValue(ViewModeIndexEnum->GetNameStringByValue((uint8)EViewModeIndex::VMI_PathTracing));
 		EnumRestriction->AddHiddenValue(ViewModeIndexEnum->GetNameStringByValue((uint8)EViewModeIndex::VMI_RayTracingDebug));
diff --git a/Engine/Source/Runtime/D3D12RHI/Private/D3D12Adapter.cpp b/Engine/Source/Runtime/D3D12RHI/Private/D3D12Adapter.cpp
--- a/Engine/Source/Runtime/D3D12RHI/Private/D3D12Adapter.cpp
+++ b/Engine/Source/Runtime/D3D12RHI/Private/D3D12Adapter.cpp
@@ -994,7 +994,8 @@ void FD3D12Adapter::InitializeDevices()
 				{
 					UE_LOG(LogD3D12RHI, Log, TEXT("D3D12 ray tracing tier 1.0 is supported."));
 
-					GRHISupportsRayTracing = true;
+					GRHISupportsRayTracing = RHISupportsRayTracing(GMaxRHIShaderPlatform);
+					GRHISupportsRayTracingShaders = GRHISupportsRayTracing && RHISupportsRayTracingShaders(GMaxRHIShaderPlatform);
 
 					if (D3D12Caps5.RaytracingTier >= D3D12_RAYTRACING_TIER_1_1
 						&& RootDevice7)
@@ -1002,6 +1003,7 @@ void FD3D12Adapter::InitializeDevices()
 						UE_LOG(LogD3D12RHI, Log, TEXT("D3D12 ray tracing tier 1.1 is supported."));
 
 						GRHISupportsRayTracingPSOAdditions = true;
+						GRHISupportsInlineRayTracing = GRHISupportsRayTracing && RHISupportsInlineRayTracing(GMaxRHIShaderPlatform);
 					}
 				}
 				else if (D3D12Caps5.RaytracingTier == D3D12_RAYTRACING_TIER_NOT_SUPPORTED 
diff --git a/Engine/Source/Runtime/D3D12RHI/Private/D3D12Shaders.cpp b/Engine/Source/Runtime/D3D12RHI/Private/D3D12Shaders.cpp
--- a/Engine/Source/Runtime/D3D12RHI/Private/D3D12Shaders.cpp
+++ b/Engine/Source/Runtime/D3D12RHI/Private/D3D12Shaders.cpp
@@ -174,7 +174,7 @@ FComputeShaderRHIRef FD3D12DynamicRHI::RHICreateComputeShader(TArrayView<const u
 
 FRayTracingShaderRHIRef FD3D12DynamicRHI::RHICreateRayTracingShader(TArrayView<const uint8> Code, const FSHAHash& Hash, EShaderFrequency ShaderFrequency)
 {
-	checkf(GRHISupportsRayTracing, TEXT("Tried to create RayTracing shader but RHI doesn't support it!"));
+	checkf(GRHISupportsRayTracing && GRHISupportsRayTracingShaders, TEXT("Tried to create RayTracing shader but RHI doesn't support it!"));
 
 	FShaderCodeReader ShaderCode(Code);
 	FD3D12RayTracingShader* Shader = new FD3D12RayTracingShader(ShaderFrequency);
diff --git a/Engine/Source/Runtime/D3D12RHI/Private/Windows/WindowsD3D12Device.cpp b/Engine/Source/Runtime/D3D12RHI/Private/Windows/WindowsD3D12Device.cpp
--- a/Engine/Source/Runtime/D3D12RHI/Private/Windows/WindowsD3D12Device.cpp
+++ b/Engine/Source/Runtime/D3D12RHI/Private/Windows/WindowsD3D12Device.cpp
@@ -821,6 +821,17 @@ static bool IsRayTracingEmulated(uint32 DeviceId)
 	return false;
 }
 
+static void DisableRayTracingSupport()
+{
+	GRHISupportsRayTracing = false;
+	GRHISupportsRayTracingPSOAdditions = false;
+	GRHISupportsRayTracingDispatchIndirect = false;
+	GRHISupportsRayTracingAsyncBuildAccelerationStructure = false;
+	GRHISupportsRayTracingAMDHitToken = false;
+	GRHISupportsInlineRayTracing = false;
+	GRHISupportsRayTracingShaders = false;
+}
+
 void FD3D12DynamicRHI::Init()
 {
 	for (TSharedPtr<FD3D12Adapter>& Adapter : ChosenAdapters)
@@ -908,8 +919,7 @@ void FD3D12DynamicRHI::Init()
 		&& !FPlatformMisc::VerifyWindowsVersion(10, 0, GMinimumWindowsBuildVersionForRayTracing)
 		&& !bIsRunningNvidiaGFN)
 	{
-		GRHISupportsRayTracing = false;
-
+		DisableRayTracingSupport();
 		UE_LOG(LogD3D12RHI, Warning, TEXT("Ray tracing is disabled because it requires Windows 10 version %u"), (uint32)GMinimumWindowsBuildVersionForRayTracing);
 	}
 #endif
@@ -947,7 +957,7 @@ void FD3D12DynamicRHI::Init()
 			&& GMinimumDriverVersionForRayTracingNVIDIA > 0
 			&& DriverVersion < (uint32)GMinimumDriverVersionForRayTracingNVIDIA)
 		{
-			GRHISupportsRayTracing = false;
+			DisableRayTracingSupport();
 			UE_LOG(LogD3D12RHI, Warning, TEXT("Ray tracing is disabled because the driver is too old"));
 		}
 
@@ -967,11 +977,11 @@ void FD3D12DynamicRHI::Init()
 #if DXR_ALLOW_EMULATED_RAYTRACING
 		if (!GAllowEmulatedRayTracing)
 		{
-			GRHISupportsRayTracing = false;
+			DisableRayTracingSupport();
 			UE_LOG(LogD3D12RHI, Warning, TEXT("Ray tracing is disabled for NVIDIA cards with the Pascal architecture. This can be overridden with the following CVar: r.D3D12.DXR.AllowEmulatedRayTracing=1"));
 		}
 #else
-		GRHISupportsRayTracing = false;
+		DisableRayTracingSupport();
 		UE_LOG(LogD3D12RHI, Warning, TEXT("Ray tracing is disabled for NVIDIA cards with the Pascal architecture."));
 #endif // DXR_ALLOW_EMULATED_RAYTRACING
 		}
@@ -987,8 +997,7 @@ void FD3D12DynamicRHI::Init()
 		if (GMinimumDriverVersionForRayTracingAMD > 0 
 			&& agsCheckDriverVersion(AmdAgsGpuInfo.radeonSoftwareVersion, GMinimumDriverVersionForRayTracingAMD) == AGS_SOFTWAREVERSIONCHECK_OLDER)
 		{
-			GRHISupportsRayTracing = false;
-
+			DisableRayTracingSupport();
 			UE_LOG(LogD3D12RHI, Warning, TEXT("Ray tracing is disabled because the driver is too old"));
 		}
 
@@ -1131,10 +1140,6 @@ void FD3D12DynamicRHI::Init()
 	// - Standalones are added to the deferred deletion queue of its parent FD3D12Adapter
 	GRHIForceNoDeletionLatencyForStreamingTextures = !!PLATFORM_WINDOWS;
 
-#if 0//D3D12_RHI_RAYTRACING
-	GRHISupportsRayTracing = GetAdapter().GetD3DRayTracingDevice() != nullptr;
-#endif
-
 	D3D12_FEATURE_DATA_D3D12_OPTIONS6 options = {};
 	HRESULT hr = GetAdapter().GetD3DDevice()->CheckFeatureSupport(D3D12_FEATURE_D3D12_OPTIONS6, &options, sizeof(options));
 	if(hr == S_OK && options.VariableShadingRateTier != D3D12_VARIABLE_SHADING_RATE_TIER_NOT_SUPPORTED)
diff --git a/Engine/Source/Runtime/Engine/Private/GameViewportClient.cpp b/Engine/Source/Runtime/Engine/Private/GameViewportClient.cpp
--- a/Engine/Source/Runtime/Engine/Private/GameViewportClient.cpp
+++ b/Engine/Source/Runtime/Engine/Private/GameViewportClient.cpp
@@ -3428,7 +3428,7 @@ bool UGameViewportClient::HandleViewModeCommand( const TCHAR* Cmd, FOutputDevice
 	}
 
 #if RHI_RAYTRACING
-	if (!GRHISupportsRayTracing)
+	if (!GRHISupportsRayTracing || !GRHISupportsRayTracingShaders)
 	{
 		if (ViewModeIndex == VMI_PathTracing)
 		{
diff --git a/Engine/Source/Runtime/RHI/Private/RHI.cpp b/Engine/Source/Runtime/RHI/Private/RHI.cpp
--- a/Engine/Source/Runtime/RHI/Private/RHI.cpp
+++ b/Engine/Source/Runtime/RHI/Private/RHI.cpp
@@ -1398,6 +1398,7 @@ bool GRHISupportsRayTracingDispatchIndirect = false;
 bool GRHISupportsRayTracingAsyncBuildAccelerationStructure = false;
 bool GRHISupportsRayTracingAMDHitToken = false;
 bool GRHISupportsInlineRayTracing = false;
+bool GRHISupportsRayTracingShaders = false;
 uint32 GRHIRayTracingAccelerationStructureAlignment = 0;
 uint32 GRHIRayTracingScratchBufferAlignment = 0;
 uint32 GRHIRayTracingShaderTableAlignment = 0;
@@ -2230,6 +2231,7 @@ void FGenericDataDrivenShaderPlatformInfo::ParseDataDrivenShaderInfo(const FConf
 	GET_SECTION_BOOL_HELPER(bSupports4ComponentUAVReadWrite);
 	GET_SECTION_BOOL_HELPER(bSupportsRenderTargetWriteMask);
 	GET_SECTION_BOOL_HELPER(bSupportsRayTracing);
+	GET_SECTION_BOOL_HELPER(bSupportsRayTracingShaders);
 	GET_SECTION_BOOL_HELPER(bSupportsInlineRayTracing);
 	GET_SECTION_BOOL_HELPER(bSupportsRayTracingProceduralPrimitive);
 	GET_SECTION_BOOL_HELPER(bSupportsRayTracingIndirectInstanceData);
diff --git a/Engine/Source/Runtime/RHI/Public/RHI.h b/Engine/Source/Runtime/RHI/Public/RHI.h
--- a/Engine/Source/Runtime/RHI/Public/RHI.h
+++ b/Engine/Source/Runtime/RHI/Public/RHI.h
@@ -201,14 +201,32 @@ inline bool RHISupportsAbsoluteVertexID(const EShaderPlatform InShaderPlatform)
 	return IsVulkanPlatform(InShaderPlatform) || IsVulkanMobilePlatform(InShaderPlatform);
 }
 
-/** Can this platform compile ray tracing shaders (regardless of project settings).
+/** Whether this platform can build acceleration structures and use full ray tracing pipelines or inline ray tracing (ray queries).
  *  To use at runtime, also check GRHISupportsRayTracing and r.RayTracing CVar (see IsRayTracingEnabled() helper).
+ *  Check GRHISupportsRayTracingShaders before using full ray tracing pipeline state objects.
+ *  Check GRHISupportsInlineRayTracing before using inline ray tracing features in compute and other shaders.
  **/
-inline RHI_API bool RHISupportsRayTracingShaders(const FStaticShaderPlatform Platform)
+inline RHI_API bool RHISupportsRayTracing(const FStaticShaderPlatform Platform)
 {
 	return FDataDrivenShaderPlatformInfo::GetSupportsRayTracing(Platform);
 }
 
+/** Whether this platform can compile ray tracing shaders (regardless of project settings).
+ *  To use at runtime, also check GRHISupportsRayTracing and r.RayTracing CVar (see IsRayTracingEnabled() helper).
+ **/
+inline RHI_API bool RHISupportsRayTracingShaders(const FStaticShaderPlatform Platform)
+{
+	return FDataDrivenShaderPlatformInfo::GetSupportsRayTracingShaders(Platform);
+}
+
+/** Whether this platform can compile shaders with inline ray tracing features.
+ *  To use at runtime, also check GRHISupportsRayTracing and r.RayTracing CVar (see IsRayTracingEnabled() helper).
+ **/
+inline RHI_API bool RHISupportsInlineRayTracing(const FStaticShaderPlatform Platform)
+{
+	return FDataDrivenShaderPlatformInfo::GetSupportsInlineRayTracing(Platform);
+}
+
 /** Can this platform compile mesh shaders with tier0 capability.
  *  To use at runtime, also check GRHISupportsMeshShadersTier0.
  **/
@@ -577,9 +595,18 @@ extern RHI_API bool GRHISupportsFirstInstance;
 /** Whether or not the RHI can handle dynamic resolution or not. */
 extern RHI_API bool GRHISupportsDynamicResolution;
 
-/** Whether or not the RHI supports ray tracing on current hardware (acceleration structure building and new ray tracing-specific shader types). */
+/**
+* Whether or not the RHI supports ray tracing on current hardware (acceleration structure building and new ray tracing-specific shader types). 
+* GRHISupportsRayTracingShaders and GRHISupportsInlineRayTracing must also be checked before dispatching ray tracing workloads.
+*/
 extern RHI_API bool GRHISupportsRayTracing;
 
+/**
+* Whether or not the RHI supports ray tracing raygen, miss and hit shaders (i.e. full ray tracing pipeline). 
+* The RHI may support inline ray tracing from compute shaders, but not the full pipeline.
+*/
+extern RHI_API bool GRHISupportsRayTracingShaders;
+
 /** Whether or not the RHI supports adding new shaders to an existing RT PSO. */
 extern RHI_API bool GRHISupportsRayTracingPSOAdditions;
 
@@ -592,7 +619,7 @@ extern RHI_API bool GRHISupportsRayTracingAsyncBuildAccelerationStructure;
 /** Whether or not the RHI supports the AMD Hit Token extension. */
 extern RHI_API bool GRHISupportsRayTracingAMDHitToken;
 
-/** Whether or not the RHI supports inline ray tracing. */
+/** Whether or not the RHI supports inline ray tracing in compute shaders, without a full ray tracing pipeline. */
 extern RHI_API bool GRHISupportsInlineRayTracing;
 
 /** Required alignment for ray tracing acceleration structures. */
diff --git a/Engine/Source/Runtime/RHI/Public/RHIDefinitions.h b/Engine/Source/Runtime/RHI/Public/RHIDefinitions.h
--- a/Engine/Source/Runtime/RHI/Public/RHIDefinitions.h
+++ b/Engine/Source/Runtime/RHI/Public/RHIDefinitions.h
@@ -379,6 +379,7 @@ class RHI_API FGenericDataDrivenShaderPlatformInfo
 	uint32 bSupportsMobileDistanceField : 1;
 	uint32 bSupportsFFTBloom : 1;
 	uint32 bSupportsInlineRayTracing : 1;
+	uint32 bSupportsRayTracingShaders : 1;
 	uint32 bSupportsVertexShaderLayer : 1;
 	uint32 bSupportsBindless : 1;
 
@@ -571,34 +572,40 @@ public:
 		return Infos[Platform].bSupportsRayTracing;
 	}
 
+	static FORCEINLINE_DEBUGGABLE const bool GetSupportsRayTracingShaders(const FStaticShaderPlatform Platform)
+	{
+		check(IsValid(Platform));
+		return Infos[Platform].bSupportsRayTracing && Infos[Platform].bSupportsRayTracingShaders;
+	}
+
 	static FORCEINLINE_DEBUGGABLE const bool GetSupportsInlineRayTracing(const FStaticShaderPlatform Platform)
 	{
 		check(IsValid(Platform));
-		return Infos[Platform].bSupportsInlineRayTracing;
+		return Infos[Platform].bSupportsRayTracing && Infos[Platform].bSupportsInlineRayTracing;
 	}
 
 	static FORCEINLINE_DEBUGGABLE const bool GetSupportsRayTracingProceduralPrimitive(const FStaticShaderPlatform Platform)
 	{
 		check(IsValid(Platform));
-		return Infos[Platform].bSupportsRayTracingProceduralPrimitive;
+		return Infos[Platform].bSupportsRayTracing && Infos[Platform].bSupportsRayTracingProceduralPrimitive;
 	}
 
 	static FORCEINLINE_DEBUGGABLE const bool GetSupportsRayTracingIndirectInstanceData(const FStaticShaderPlatform Platform)
 	{
 		check(IsValid(Platform));
-		return Infos[Platform].bSupportsRayTracingIndirectInstanceData;
+		return Infos[Platform].bSupportsRayTracing && Infos[Platform].bSupportsRayTracingIndirectInstanceData;
 	}
 
 	static FORCEINLINE_DEBUGGABLE const bool GetSupportsPathTracing(const FStaticShaderPlatform Platform)
 	{
 		check(IsValid(Platform));
-		return Infos[Platform].bSupportsPathTracing;
+		return Infos[Platform].bSupportsRayTracing && Infos[Platform].bSupportsPathTracing;
 	}
 
 	static FORCEINLINE_DEBUGGABLE const bool GetSupportsHighEndRayTracingReflections(const FStaticShaderPlatform Platform)
 	{
 		check(IsValid(Platform));
-		return Infos[Platform].bSupportsHighEndRayTracingReflections;
+		return Infos[Platform].bSupportsRayTracing && Infos[Platform].bSupportsHighEndRayTracingReflections;
 	}
 
 	static FORCEINLINE_DEBUGGABLE const bool GetSupportsGPUSkinCache(const FStaticShaderPlatform Platform)
diff --git a/Engine/Source/Runtime/RenderCore/Private/ShaderCodeArchive.cpp b/Engine/Source/Runtime/RenderCore/Private/ShaderCodeArchive.cpp
--- a/Engine/Source/Runtime/RenderCore/Private/ShaderCodeArchive.cpp
+++ b/Engine/Source/Runtime/RenderCore/Private/ShaderCodeArchive.cpp
@@ -1002,7 +1002,7 @@ TRefCountPtr<FRHIShader> FShaderCodeArchive::CreateShader(int32 Index)
 	case SF_Compute: Shader = RHICreateComputeShader(ShaderCodeView, ShaderHash); CheckShaderCreation(Shader, Index); break;
 	case SF_RayGen: case SF_RayMiss: case SF_RayHitGroup: case SF_RayCallable:
 #if RHI_RAYTRACING
-		if (GRHISupportsRayTracing)
+		if (GRHISupportsRayTracing && GRHISupportsRayTracingShaders)
 		{
 			Shader = RHICreateRayTracingShader(ShaderCodeView, ShaderHash, ShaderEntry.GetFrequency());
 			CheckShaderCreation(Shader, Index);
@@ -1375,7 +1375,7 @@ TRefCountPtr<FRHIShader> FIoStoreShaderCodeArchive::CreateShader(int32 Index)
 	case SF_Compute: Shader = RHICreateComputeShader(ShaderCodeView, ShaderHash); break;
 	case SF_RayGen: case SF_RayMiss: case SF_RayHitGroup: case SF_RayCallable:
 #if RHI_RAYTRACING
-		if (GRHISupportsRayTracing)
+		if (GRHISupportsRayTracing && GRHISupportsRayTracingShaders)
 		{
 			Shader = RHICreateRayTracingShader(ShaderCodeView, ShaderHash, ShaderEntry.GetFrequency());
 		}
diff --git a/Engine/Source/Runtime/RenderCore/Private/ShaderResource.cpp b/Engine/Source/Runtime/RenderCore/Private/ShaderResource.cpp
--- a/Engine/Source/Runtime/RenderCore/Private/ShaderResource.cpp
+++ b/Engine/Source/Runtime/RenderCore/Private/ShaderResource.cpp
@@ -342,7 +342,7 @@ FShaderMapResource::FShaderMapResource(EShaderPlatform InPlatform, int32 NumShad
 {
 	RHIShaders = MakeUnique<std::atomic<FRHIShader*>[]>(NumRHIShaders); // this MakeUnique() zero-initializes the array
 #if RHI_RAYTRACING
-	if (GRHISupportsRayTracing)
+	if (GRHISupportsRayTracing && GRHISupportsRayTracingShaders)
 	{
 		RayTracingMaterialLibraryIndices.AddUninitialized(NumShaders);
 		FMemory::Memset(RayTracingMaterialLibraryIndices.GetData(), 0xff, NumShaders * RayTracingMaterialLibraryIndices.GetTypeSize());
@@ -425,7 +425,7 @@ FRHIShader* FShaderMapResource::CreateShader(int32 ShaderIndex)
 
 	TRefCountPtr<FRHIShader> RHIShader = CreateRHIShader(ShaderIndex);
 #if RHI_RAYTRACING
-	if (GRHISupportsRayTracing && RHIShader.IsValid() && RHIShader->GetFrequency() == SF_RayHitGroup)
+	if (GRHISupportsRayTracing && GRHISupportsRayTracingShaders && RHIShader.IsValid() && RHIShader->GetFrequency() == SF_RayHitGroup)
 	{
 		RayTracingMaterialLibraryIndices[ShaderIndex] = AddToRayTracingLibrary(static_cast<FRHIRayTracingShader*>(RHIShader.GetReference()));
 	}
@@ -493,7 +493,7 @@ TRefCountPtr<FRHIShader> FShaderMapResource_InlineCode::CreateRHIShader(int32 Sh
 	case SF_Compute: RHIShader = RHICreateComputeShader(ShaderCodeView, ShaderHash); break;
 	case SF_RayGen: case SF_RayMiss: case SF_RayHitGroup: case SF_RayCallable:
 #if RHI_RAYTRACING
-		if (GRHISupportsRayTracing)
+		if (GRHISupportsRayTracing && GRHISupportsRayTracingShaders)
 		{
 			RHIShader = RHICreateRayTracingShader(ShaderCodeView, ShaderHash, Frequency);
 		}
diff --git a/Engine/Source/Runtime/Renderer/Private/Tests/RayTracingTestbed.cpp b/Engine/Source/Runtime/Renderer/Private/Tests/RayTracingTestbed.cpp
--- a/Engine/Source/Runtime/Renderer/Private/Tests/RayTracingTestbed.cpp
+++ b/Engine/Source/Runtime/Renderer/Private/Tests/RayTracingTestbed.cpp
@@ -20,7 +20,8 @@ bool RunRayTracingTestbed_RenderThread(const FString& Parameters)
 {
 	check(IsInRenderingThread());
 
-	if (!GRHISupportsRayTracing)
+	// The ray tracing testbed currently rquires full ray tracing pipeline support.
+	if (!GRHISupportsRayTracing || !GRHISupportsRayTracingShaders)
 	{
 		//Return true so the test passes in DX11, until the testing framework allows to skip tests depending on defined preconditions
 		return true;
