Make a few more functions use FStaticShaderPlatform.

diff --git a/Engine/Source/Runtime/RHI/Public/RHI.h b/Engine/Source/Runtime/RHI/Public/RHI.h
--- a/Engine/Source/Runtime/RHI/Public/RHI.h
+++ b/Engine/Source/Runtime/RHI/Public/RHI.h
@@ -196,7 +196,7 @@ inline bool RHISupportsManualVertexFetch(const FStaticShaderPlatform InShaderPla
 /** 
  * Returns true if SV_VertexID contains BaseVertexIndex passed to the draw call, false if shaders must manually construct an absolute VertexID.
  */
-inline bool RHISupportsAbsoluteVertexID(const EShaderPlatform InShaderPlatform)
+inline bool RHISupportsAbsoluteVertexID(const FStaticShaderPlatform InShaderPlatform)
 {
 	return IsVulkanPlatform(InShaderPlatform) || IsVulkanMobilePlatform(InShaderPlatform);
 }
@@ -263,7 +263,7 @@ inline bool RHISupportsRenderTargetWriteMask(const FStaticShaderPlatform Platfor
 }
 
 /** True if the given shader platform supports overestimated conservative rasterization */
-inline RHI_API bool RHISupportsConservativeRasterization(EShaderPlatform Platform)
+inline RHI_API bool RHISupportsConservativeRasterization(const FStaticShaderPlatform Platform)
 {
 	return FDataDrivenShaderPlatformInfo::GetSupportsConservativeRasterization(Platform);
 }
