Add multiview and FDM/FSR attachments to the PSO caching system (Contributed by: Remi Palandri)
#8656

diff --git a/Engine/Source/Runtime/RHI/Private/PipelineFileCache.cpp b/Engine/Source/Runtime/RHI/Private/PipelineFileCache.cpp
--- a/Engine/Source/Runtime/RHI/Private/PipelineFileCache.cpp
+++ b/Engine/Source/Runtime/RHI/Private/PipelineFileCache.cpp
@@ -75,13 +75,14 @@ enum class EPipelineCacheFileFormatVersions : uint32
 	RemovingTessellationShaders = 21,
 	LastUsedTime = 22,
 	MoreRenderTargetFlags = 23,
+	FragmentDensityAttachment = 24,
 };
 
 const uint64 FPipelineCacheFileFormatMagic = 0x5049504543414348; // PIPECACH
 const uint64 FPipelineCacheTOCFileFormatMagic = 0x544F435354415232; // TOCSTAR2
 const uint64 FPipelineCacheEOFFileFormatMagic = 0x454F462D4D41524B; // EOF-MARK
-const RHI_API uint32 FPipelineCacheFileFormatCurrentVersion = (uint32)EPipelineCacheFileFormatVersions::MoreRenderTargetFlags;
-const int32  FPipelineCacheGraphicsDescPartsNum = 64; // parser will expect this number of parts in a description string
+const RHI_API uint32 FPipelineCacheFileFormatCurrentVersion = (uint32)EPipelineCacheFileFormatVersions::FragmentDensityAttachment;
+const int32  FPipelineCacheGraphicsDescPartsNum = 66; // parser will expect this number of parts in a description string
 
 /**
   * PipelineFileCache API access
@@ -514,6 +515,11 @@ FString FPipelineCacheFileFormatPSO::GraphicsDescriptor::StateToString() const
 		, uint32(SubpassHint)
 		, uint32(SubpassIndex)
 	);
+
+	Result += FString::Printf(TEXT("%d,%d,")
+		, uint32(MultiViewCount)
+		, uint32(bHasFragmentDensityAttachment)
+	);
 	
 	FVertexElement NullVE;
 	FMemory::Memzero(NullVE);
@@ -593,6 +599,12 @@ void FPipelineCacheFileFormatPSO::GraphicsDescriptor::AddStateToReadableString(T
 	OutBuilder << uint32(SubpassIndex);
 	OutBuilder << TEXT("\n");
 
+	OutBuilder << TEXT(" MVC:");
+	OutBuilder << MultiViewCount;
+	OutBuilder << TEXT(" HasFDM:");
+	OutBuilder << bHasFragmentDensityAttachment;
+	OutBuilder << TEXT("\n");
+
 	OutBuilder << TEXT(" NumVE ");
 	OutBuilder << VertexDescriptor.Num();
 	OutBuilder << TEXT("\n");
@@ -667,6 +679,17 @@ bool FPipelineCacheFileFormatPSO::GraphicsDescriptor::StateFromString(const FStr
 		SubpassIndex = LocalSubpassIndex;
 	}
 
+	// parse multiview and FDM information
+	{
+		uint32 LocalMultiViewCount = 0;
+		uint32 LocalHasFDM = 0;
+		check(PartEnd - PartIt >= 2);
+		LexFromString(LocalMultiViewCount, *PartIt++);
+		LexFromString(LocalHasFDM, *PartIt++);
+		MultiViewCount = (uint8)LocalMultiViewCount;
+		bHasFragmentDensityAttachment = (bool)LocalHasFDM;
+	}
+
 	check(PartEnd - PartIt >= 1); //not a very robust parser
 	int32 VertDescNum = 0;
 	LexFromString(VertDescNum, *PartIt++);
@@ -761,6 +784,11 @@ FString FPipelineCacheFileFormatPSO::GraphicsDescriptor::StateHeaderLine()
 		, TEXT("SubpassHint")
 		, TEXT("SubpassIndex")
 	);
+
+	Result += FString::Printf(TEXT("%s,%s,")
+		, TEXT("MultiViewCount")
+		, TEXT("bHasFDMAttachment")
+	);
 	
 	Result += FString::Printf(TEXT("%s,")
 		, TEXT("VertexDescriptorNum")
@@ -1047,6 +1075,9 @@ bool FPipelineCacheFileFormatPSO::Verify() const
 				KeyHash = FCrc::MemCrc32(&Key.GraphicsDesc.SubpassHint, sizeof(Key.GraphicsDesc.SubpassHint), KeyHash);
 				KeyHash = FCrc::MemCrc32(&Key.GraphicsDesc.SubpassIndex, sizeof(Key.GraphicsDesc.SubpassIndex), KeyHash);
 				
+				KeyHash = FCrc::MemCrc32(&Key.GraphicsDesc.MultiViewCount, sizeof(Key.GraphicsDesc.MultiViewCount), KeyHash);
+				KeyHash = FCrc::MemCrc32(&Key.GraphicsDesc.bHasFragmentDensityAttachment, sizeof(Key.GraphicsDesc.bHasFragmentDensityAttachment), KeyHash);
+
 				for(auto const& Element : Key.GraphicsDesc.VertexDescriptor)
 				{
 					KeyHash = FCrc::MemCrc32(&Element, sizeof(FVertexElement), KeyHash);
@@ -1242,6 +1273,20 @@ bool FPipelineCacheFileFormatPSO::Verify() const
 			Ar << Info.GraphicsDesc.SubpassHint;
 			Ar << Info.GraphicsDesc.SubpassIndex;
 
+			if (Ar.GameNetVer() < (uint32)EPipelineCacheFileFormatVersions::FragmentDensityAttachment)
+			{
+				uint8 MultiViewCount = 0;
+				Ar << MultiViewCount;
+				
+				bool bHasFragmentDensityAttachment = false;
+				Ar << bHasFragmentDensityAttachment;
+			}
+			else
+			{
+				Ar << Info.GraphicsDesc.MultiViewCount;
+				Ar << Info.GraphicsDesc.bHasFragmentDensityAttachment;
+			}
+
 			break;
 		}
 		case FPipelineCacheFileFormatPSO::DescriptorType::RayTracing:
@@ -1415,6 +1460,9 @@ FPipelineCacheFileFormatPSO::FPipelineCacheFileFormatPSO()
 	PSO.GraphicsDesc.SubpassHint = (uint8)Init.SubpassHint;
 	PSO.GraphicsDesc.SubpassIndex = Init.SubpassIndex;
 	
+	PSO.GraphicsDesc.MultiViewCount = (uint8)Init.MultiViewCount;
+	PSO.GraphicsDesc.bHasFragmentDensityAttachment = Init.bHasFragmentDensityAttachment;
+
 #if !UE_BUILD_SHIPPING
 	bOK = bOK && PSO.Verify();
 #endif
@@ -1469,6 +1517,7 @@ bool FPipelineCacheFileFormatPSO::operator==(const FPipelineCacheFileFormatPSO&
 						GraphicsDesc.DepthStencilFlags == Other.GraphicsDesc.DepthStencilFlags && GraphicsDesc.DepthLoad == Other.GraphicsDesc.DepthLoad &&
 						GraphicsDesc.DepthStore == Other.GraphicsDesc.DepthStore && GraphicsDesc.StencilLoad == Other.GraphicsDesc.StencilLoad && GraphicsDesc.StencilStore == Other.GraphicsDesc.StencilStore &&
 						GraphicsDesc.SubpassHint == Other.GraphicsDesc.SubpassHint && GraphicsDesc.SubpassIndex == Other.GraphicsDesc.SubpassIndex &&
+						GraphicsDesc.MultiViewCount == Other.GraphicsDesc.MultiViewCount && GraphicsDesc.bHasFragmentDensityAttachment == Other.GraphicsDesc.bHasFragmentDensityAttachment &&
 					FMemory::Memcmp(&GraphicsDesc.BlendState, &Other.GraphicsDesc.BlendState, sizeof(FBlendStateInitializerRHI)) == 0 &&
 					FMemory::Memcmp(&GraphicsDesc.RasterizerState, &Other.GraphicsDesc.RasterizerState, sizeof(FPipelineFileCacheRasterizerState)) == 0 &&
 					FMemory::Memcmp(&GraphicsDesc.DepthStencilState, &Other.GraphicsDesc.DepthStencilState, sizeof(FDepthStencilStateInitializerRHI)) == 0 &&
diff --git a/Engine/Source/Runtime/RHI/Public/PipelineFileCache.h b/Engine/Source/Runtime/RHI/Public/PipelineFileCache.h
--- a/Engine/Source/Runtime/RHI/Public/PipelineFileCache.h
+++ b/Engine/Source/Runtime/RHI/Public/PipelineFileCache.h
@@ -150,6 +150,9 @@ struct RHI_API FPipelineCacheFileFormatPSO
 		
 		uint8 SubpassHint;	
 		uint8 SubpassIndex;
+
+		uint8	MultiViewCount;
+		bool	bHasFragmentDensityAttachment;
 		
 		FString ToString() const;
 		void AddToReadableString(TReadableStringBuilder& OutBuilder) const;
diff --git a/Engine/Source/Runtime/RenderCore/Private/PipelineCacheUtilities.cpp b/Engine/Source/Runtime/RenderCore/Private/PipelineCacheUtilities.cpp
--- a/Engine/Source/Runtime/RenderCore/Private/PipelineCacheUtilities.cpp
+++ b/Engine/Source/Runtime/RenderCore/Private/PipelineCacheUtilities.cpp
@@ -55,7 +55,7 @@ namespace Private
 
 		enum class EVersion : int32
 		{
-			Current = 2
+			Current = 3
 		};
 
 		/** Magic to reject other files */
diff --git a/Engine/Source/Runtime/RenderCore/Private/ShaderPipelineCache.cpp b/Engine/Source/Runtime/RenderCore/Private/ShaderPipelineCache.cpp
--- a/Engine/Source/Runtime/RenderCore/Private/ShaderPipelineCache.cpp
+++ b/Engine/Source/Runtime/RenderCore/Private/ShaderPipelineCache.cpp
@@ -682,7 +682,10 @@ bool FShaderPipelineCache::Precompile(FRHICommandListImmediate& RHICmdList, ESha
 
 			GraphicsInitializer.SubpassHint = (ESubpassHint)PSO.GraphicsDesc.SubpassHint;
 			GraphicsInitializer.SubpassIndex = PSO.GraphicsDesc.SubpassIndex;
-			
+
+			GraphicsInitializer.MultiViewCount = PSO.GraphicsDesc.MultiViewCount;
+			GraphicsInitializer.bHasFragmentDensityAttachment = PSO.GraphicsDesc.bHasFragmentDensityAttachment;
+
 			GraphicsInitializer.DepthStencilTargetFormat = PSO.GraphicsDesc.DepthStencilFormat;
 			GraphicsInitializer.DepthStencilTargetFlag = PSO.GraphicsDesc.DepthStencilFlags;
 			GraphicsInitializer.DepthTargetLoadAction = PSO.GraphicsDesc.DepthLoad;
