Backpedalling changes to data drive more settings. Turns out bIsMobile is more nuanced than I initially thought. Switching back to a few hardcoded shader platform uses until we can wrangle this better.

diff --git a/Engine/Config/Android/DataDrivenPlatformInfo.ini b/Engine/Config/Android/DataDrivenPlatformInfo.ini
--- a/Engine/Config/Android/DataDrivenPlatformInfo.ini
+++ b/Engine/Config/Android/DataDrivenPlatformInfo.ini
@@ -76,7 +76,6 @@ bSupportsManualVertexFetch=false
 [ShaderPlatform VULKAN_SM5_ANDROID]
 Language=Vulkan
 MaxFeatureLevel=SM5
-bIsMobile=true
 bTargetsTiledGPU=true
 bNeedsOfflineCompiler=true
 bSupportsIndexBufferUAVs=true
diff --git a/Engine/Config/IOS/DataDrivenPlatformInfo.ini b/Engine/Config/IOS/DataDrivenPlatformInfo.ini
--- a/Engine/Config/IOS/DataDrivenPlatformInfo.ini
+++ b/Engine/Config/IOS/DataDrivenPlatformInfo.ini
@@ -56,7 +56,6 @@ bSupportsManualVertexFetch=false
 [ShaderPlatform METAL_MRT]
 Language=Metal
 MaxFeatureLevel=SM5
-bIsMobile=true
 bIsMetalMRT=true
 bTargetsTiledGPU=true
 bSupportsIndexBufferUAVs=true
diff --git a/Engine/Config/TVOS/DataDrivenPlatformInfo.ini b/Engine/Config/TVOS/DataDrivenPlatformInfo.ini
--- a/Engine/Config/TVOS/DataDrivenPlatformInfo.ini
+++ b/Engine/Config/TVOS/DataDrivenPlatformInfo.ini
@@ -19,7 +19,6 @@ PlatformGroupName=Mobile
 [ShaderPlatform METAL_MRT_TVOS]
 Language=Metal
 MaxFeatureLevel=SM5
-bIsMobile=true
 bIsMetalMRT=true
 bTargetsTiledGPU=true
 bSupportsIndexBufferUAVs=true
diff --git a/Engine/Source/Runtime/RHI/Public/RHIDefinitions.h b/Engine/Source/Runtime/RHI/Public/RHIDefinitions.h
--- a/Engine/Source/Runtime/RHI/Public/RHIDefinitions.h
+++ b/Engine/Source/Runtime/RHI/Public/RHIDefinitions.h
@@ -1979,18 +1979,22 @@ inline bool IsVulkanSM5Platform(const FStaticShaderPlatform Platform)
 		&& FDataDrivenShaderPlatformInfo::GetMaxFeatureLevel(Platform) == ERHIFeatureLevel::SM5;
 }
 
+// @todo: data drive uses of this function
 inline bool IsVulkanMobileSM5Platform(const FStaticShaderPlatform Platform)
 {
-	return FDataDrivenShaderPlatformInfo::GetIsLanguageVulkan(Platform)
-		&& FDataDrivenShaderPlatformInfo::GetMaxFeatureLevel(Platform) == ERHIFeatureLevel::SM5
-		&& FDataDrivenShaderPlatformInfo::GetIsMobile(Platform);
+	return Platform == SP_VULKAN_SM5_ANDROID;
+// 	return FDataDrivenShaderPlatformInfo::GetIsLanguageVulkan(Platform)
+// 		&& FDataDrivenShaderPlatformInfo::GetMaxFeatureLevel(Platform) == ERHIFeatureLevel::SM5
+// 		&& FDataDrivenShaderPlatformInfo::GetIsMobile(Platform);
 }
 
+// @todo: data drive uses of this function
 inline bool IsMetalMobileSM5Platform(const FStaticShaderPlatform Platform)
 {
-	return FDataDrivenShaderPlatformInfo::GetIsLanguageMetal(Platform)
-		&& FDataDrivenShaderPlatformInfo::GetMaxFeatureLevel(Platform) == ERHIFeatureLevel::SM5
-		&& FDataDrivenShaderPlatformInfo::GetIsMobile(Platform);
+	return Platform == SP_METAL_MRT;
+// 	return FDataDrivenShaderPlatformInfo::GetIsLanguageMetal(Platform)
+// 		&& FDataDrivenShaderPlatformInfo::GetMaxFeatureLevel(Platform) == ERHIFeatureLevel::SM5
+// 		&& FDataDrivenShaderPlatformInfo::GetIsMobile(Platform);
 }
 
 inline bool IsAndroidOpenGLESPlatform(const FStaticShaderPlatform Platform)
@@ -2001,7 +2005,7 @@ inline bool IsAndroidOpenGLESPlatform(const FStaticShaderPlatform Platform)
 inline bool IsVulkanMobilePlatform(const FStaticShaderPlatform Platform)
 {
 	return FDataDrivenShaderPlatformInfo::GetIsLanguageVulkan(Platform)
-		&& FDataDrivenShaderPlatformInfo::GetIsMobile(Platform)
+		//&& FDataDrivenShaderPlatformInfo::GetIsMobile(Platform)
 		// This was limited to the ES3_1 platforms when hard coded
 		&& FDataDrivenShaderPlatformInfo::GetMaxFeatureLevel(Platform) == ERHIFeatureLevel::ES3_1;
 }
