RHI TextureProfiler.  Gathers data for all textures and render targets, and provides various methods for capturing that data.  Currently only implemented for PS5.

diff --git a/Engine/Source/Runtime/RHI/Private/TextureProfiler.cpp b/Engine/Source/Runtime/RHI/Private/TextureProfiler.cpp
new file mode 100644
--- /dev/null
+++ b/Engine/Source/Runtime/RHI/Private/TextureProfiler.cpp
@@ -0,0 +1,353 @@
+// Copyright Epic Games, Inc. All Rights Reserved.
+
+#include "TextureProfiler.h"
+
+#include "UObject/NameTypes.h"
+
+#if TEXTURE_PROFILER_ENABLED
+#include "ProfilingDebugging/CsvProfiler.h"
+
+static int32 ToMB(uint64 Value)
+{
+	const size_t MB = 1024 * 1024;
+	uint64 Value64 = (Value + (MB - 1)) / MB;
+	int32 Value32 = static_cast<int32>(Value64);
+	check(Value64 == Value32);
+	return Value32;
+}
+
+static int32 ToKB(uint64 Value)
+{
+	const size_t KB = 1024;
+	uint64 Value64 = (Value + (KB - 1)) / KB;
+	int32 Value32 = static_cast<int32>(Value64);
+	check(Value64 == Value32);
+	return Value32;
+}
+
+CSV_DEFINE_CATEGORY(RenderTargetProfiler, true);
+CSV_DEFINE_CATEGORY(RenderTargetWasteProfiler, true);
+CSV_DEFINE_CATEGORY(TextureProfiler, true);
+CSV_DEFINE_CATEGORY(TextureWasteProfiler, true);
+
+static TAutoConsoleVariable<int32> CVarTextureProfilerMinTextureSizeMB(
+	TEXT("r.TextureProfiler.MinTextureSizeMB"),
+	16,
+	TEXT("The minimum size for any texture to be reported.  All textures below this threshold will be reported as Other."));
+
+static TAutoConsoleVariable<int32> CVarTextureProfilerMinRenderTargetSizeMB(
+	TEXT("r.TextureProfiler.MinRenderTargetSizeMB"),
+	0,
+	TEXT("The minimum combined size for render targets to be reported.  All combined sizes less than this threshold will be reported as Other."));
+
+static TAutoConsoleVariable<bool> CVarTextureProfilerEnableTextureCSV(
+	TEXT("r.TextureProfiler.EnableTextureCSV"),
+	true,
+	TEXT("True to enable csv profiler output for all textures.  Does not include render targets."));
+
+static TAutoConsoleVariable<bool> CVarTextureProfilerEnableRenderTargetCSV(
+	TEXT("r.TextureProfiler.EnableRenderTargetCSV"),
+	true,
+	TEXT("True to enable csv profiler output for all Render Targets."));
+
+static FAutoConsoleCommand CmdTextureProfilerDumpRenderTargets(
+	TEXT("r.TextureProfiler.DumpRenderTargets"),
+	TEXT("Dumps all render targets allocated by the RHI.\n")
+	TEXT("Arguments:\n")
+	TEXT("-CombineTextureNames    Combines all textures of the same name into a single line of output\n")
+	TEXT("-CSV                    Produces CSV ready output"),
+
+	FConsoleCommandWithWorldArgsAndOutputDeviceDelegate::CreateLambda([](const TArray< FString >& Args, UWorld* InWorld, FOutputDevice& Ar)
+		{
+			bool AsCSV = Args.Contains(TEXT("-CSV"));
+			bool CombineTextureNames = Args.Contains(TEXT("-CombineTextureNames"));
+			FTextureProfiler::Get()->DumpTextures(true, CombineTextureNames, AsCSV, Ar);
+		})
+	);
+
+static FAutoConsoleCommand CmdTextureProfilerDumpTextures(
+	TEXT("r.TextureProfiler.DumpTextures"),
+	TEXT("Dumps all textures allocated by the RHI.  Does not include render targets\n")
+	TEXT("Arguments:\n")
+	TEXT("-CombineTextureNames    Combines all textures of the same name into a single line of output\n")
+	TEXT("-CSV                    Produces CSV ready output"),
+
+	FConsoleCommandWithWorldArgsAndOutputDeviceDelegate::CreateLambda([](const TArray< FString >& Args, UWorld* InWorld, FOutputDevice& Ar)
+		{
+			bool AsCSV = Args.Contains(TEXT("-CSV"));
+			bool CombineTextureNames = Args.Contains(TEXT("-CombineTextureNames"));
+			FTextureProfiler::Get()->DumpTextures(false, CombineTextureNames, AsCSV, Ar);
+		})
+	);
+
+FTextureProfiler* FTextureProfiler::Instance = nullptr;
+
+FTextureProfiler::FTexureDetails::FTexureDetails(FName InTextureName, size_t InSize, uint32 InAlign, size_t InAllocationWaste, bool InIsRenderTarget)
+	: TextureName(InTextureName)
+	, Size(InSize)
+	, Align(InAlign)
+	, AllocationWaste(InAllocationWaste)
+	, IsRenderTarget(InIsRenderTarget)
+{
+	InTextureName.GetPlainANSIString(TextureNameString);
+}
+
+void FTextureProfiler::FTexureDetails::SetName(FName InTextureName)
+{
+	TextureName = InTextureName;
+	InTextureName.GetPlainANSIString(TextureNameString);
+}
+
+void FTextureProfiler::FTexureDetails::ResetPeakSize()
+{
+	PeakSize = Size;
+}
+
+FTextureProfiler::FTexureDetails& FTextureProfiler::FTexureDetails::operator+=(const FTextureProfiler::FTexureDetails& Other)
+{
+	Size += Other.Size;
+	PeakSize = FMath::Max(PeakSize, Size);
+	AllocationWaste += Other.AllocationWaste;
+	Align = 0;
+	Count += Other.Count;
+	return *this;
+}
+
+FTextureProfiler::FTexureDetails& FTextureProfiler::FTexureDetails::operator-=(const FTextureProfiler::FTexureDetails& Other)
+{
+	Size -= Other.Size;
+	AllocationWaste -= Other.AllocationWaste;
+	Align = 0;
+	Count -= Other.Count;
+	return *this;
+}
+
+void FTextureProfiler::Init()
+{
+	FCoreDelegates::OnEndFrameRT.AddRaw(this, &FTextureProfiler::Update);
+}
+
+FTextureProfiler* FTextureProfiler::Get()
+{
+	if (Instance == nullptr)
+	{
+		Instance = new FTextureProfiler;
+	}
+	return Instance;
+}
+
+void FTextureProfiler::AddTextureAllocationInternal(void* UniqueTexturePtr, FName TextureName, size_t Size, uint32 Alignment, size_t AllocationWaste, bool IsRenderTarget)
+{
+	FScopeLock Lock(&TextureMapCS);
+
+	FTexureDetails AddedDetails(
+		TextureName,
+		Size,
+		Alignment,
+		AllocationWaste,
+		IsRenderTarget
+	);
+
+	AddedDetails.Count = 1;
+
+	FTexureDetails& TotalValue = IsRenderTarget ? TotalRenderTargetSize : TotalTextureSize;
+	TotalValue += AddedDetails;
+
+	check(ToMB(TotalValue.Size) >= 0);
+
+	FTexureDetails& NamedValue = IsRenderTarget ? CombinedRenderTargetSizes.FindOrAdd(TextureName) : CombinedTextureSizes.FindOrAdd(TextureName);
+	NamedValue += AddedDetails;
+
+	if (NamedValue.TextureName.IsNone() || !NamedValue.TextureName.IsValid())
+	{
+		NamedValue.SetName(TextureName);
+	}
+
+	TexturesMap.Add(UniqueTexturePtr, AddedDetails);
+}
+
+void FTextureProfiler::RemoveTextureAllocationInternal(void* UniqueTexturePtr, bool IsRenderTarget)
+{
+	FScopeLock Lock(&TextureMapCS);
+	FTexureDetails Details = TexturesMap.FindAndRemoveChecked(UniqueTexturePtr);
+	
+	FTexureDetails& TotalValue = IsRenderTarget ? TotalRenderTargetSize : TotalTextureSize;
+	TotalValue -= Details;
+
+	FTexureDetails& NamedValue = IsRenderTarget ? CombinedRenderTargetSizes[Details.TextureName] : CombinedTextureSizes[Details.TextureName];
+	NamedValue -= Details;
+}
+
+void FTextureProfiler::ChangeTextureNameInternal(void* UniqueTexturePtr, FName NewName, bool IsRenderTarget)
+{
+	FScopeLock Lock(&TextureMapCS);
+	FTexureDetails& Details = TexturesMap[UniqueTexturePtr];
+	FName OldName = Details.TextureName;
+	Details.SetName(NewName);
+
+	FTexureDetails& OldCombinedValue = IsRenderTarget ? CombinedRenderTargetSizes[OldName] : CombinedTextureSizes[OldName];
+	FTexureDetails& NewCombinedValue = IsRenderTarget ? CombinedRenderTargetSizes.FindOrAdd(NewName) : CombinedTextureSizes.FindOrAdd(NewName);
+
+	if (NewCombinedValue.TextureName.IsNone() || !NewCombinedValue.TextureName.IsValid())
+	{
+		NewCombinedValue.SetName(NewName);
+	}
+
+	OldCombinedValue -= Details;
+	NewCombinedValue += Details;
+}
+
+#if CSV_PROFILER
+
+static void ReportTextureStat(const TAutoConsoleVariable<bool>& EnableVar, const char* StatName, uint32 CategoryIndex, int32 Size, ECsvCustomStatOp Op)
+{
+	FCsvProfilerTrace::OutputInlineStat(StatName, CategoryIndex);
+	FCsvProfiler::RecordCustomStat(StatName, CategoryIndex, Size, Op);
+}
+#endif //CSV_PROFILER
+
+void FTextureProfiler::Update()
+{
+#if CSV_PROFILER
+
+	check(IsInRenderingThread());
+
+	if (FCsvProfiler::Get()->IsCapturing_Renderthread() && !CVarTextureProfilerEnableRenderTargetCSV.GetValueOnAnyThread() && !CVarTextureProfilerEnableTextureCSV.GetValueOnAnyThread())
+	{
+		return;
+	}
+
+	FScopeLock Lock(&TextureMapCS);
+	
+	size_t MinTextureSize = (size_t)CVarTextureProfilerMinTextureSizeMB.GetValueOnAnyThread() * 1024LL * 1024LL;
+	size_t MinRenderTargetSize = (size_t)CVarTextureProfilerMinRenderTargetSizeMB.GetValueOnAnyThread() * 1024LL * 1024LL;
+
+	// Keep track of all stats that have been added this round.  the CVS trace profiler wants to know what stats have been added, but only once per frame
+	// Keep one set per category
+	struct ECategoryIndex
+	{
+		enum
+		{
+			RenderTarget = 0,
+			RenderTargetWaste,
+			Texture,
+			TextureWaste,
+			Count
+		};
+	};
+	
+	FTexureDetails OtherRenderTargetSizes;
+	for (auto& Pair : CombinedRenderTargetSizes)
+	{
+		// Don't report empty sizes or sizes less than the minimum
+		if (Pair.Value.PeakSize == 0 || Pair.Value.PeakSize < MinRenderTargetSize)
+		{
+			OtherRenderTargetSizes += Pair.Value;
+			continue;
+		}
+
+		ReportTextureStat(CVarTextureProfilerEnableRenderTargetCSV, Pair.Value.TextureNameString, CSV_CATEGORY_INDEX(RenderTargetProfiler), ToMB(Pair.Value.PeakSize), ECsvCustomStatOp::Set);
+		ReportTextureStat(CVarTextureProfilerEnableRenderTargetCSV, Pair.Value.TextureNameString, CSV_CATEGORY_INDEX(RenderTargetWasteProfiler), ToKB(Pair.Value.AllocationWaste), ECsvCustomStatOp::Set);
+
+		Pair.Value.ResetPeakSize();
+	}
+
+	FTexureDetails OtherTextureSizes;
+	for (auto& Pair : CombinedTextureSizes)
+	{
+		// Don't report empty sizes or sizes less than the minimum
+		if (Pair.Value.PeakSize == 0 || Pair.Value.PeakSize < MinTextureSize)
+		{
+			OtherTextureSizes += Pair.Value;
+			continue;
+		}
+
+		ReportTextureStat(CVarTextureProfilerEnableTextureCSV, Pair.Value.TextureNameString, CSV_CATEGORY_INDEX(TextureProfiler), ToMB(Pair.Value.PeakSize), ECsvCustomStatOp::Set);
+		ReportTextureStat(CVarTextureProfilerEnableTextureCSV, Pair.Value.TextureNameString, CSV_CATEGORY_INDEX(TextureWasteProfiler), ToKB(Pair.Value.AllocationWaste), ECsvCustomStatOp::Set);
+
+		Pair.Value.ResetPeakSize();
+	}
+
+	ReportTextureStat(CVarTextureProfilerEnableRenderTargetCSV, "Total", CSV_CATEGORY_INDEX(RenderTargetProfiler), ToMB(TotalRenderTargetSize.PeakSize), ECsvCustomStatOp::Set);
+	ReportTextureStat(CVarTextureProfilerEnableRenderTargetCSV, "Total", CSV_CATEGORY_INDEX(RenderTargetWasteProfiler), ToKB(TotalRenderTargetSize.AllocationWaste), ECsvCustomStatOp::Set);
+	ReportTextureStat(CVarTextureProfilerEnableRenderTargetCSV, "Other", CSV_CATEGORY_INDEX(RenderTargetProfiler), ToMB(OtherRenderTargetSizes.PeakSize), ECsvCustomStatOp::Set);
+	ReportTextureStat(CVarTextureProfilerEnableRenderTargetCSV, "Other", CSV_CATEGORY_INDEX(RenderTargetWasteProfiler), ToKB(OtherRenderTargetSizes.AllocationWaste), ECsvCustomStatOp::Set);
+
+	ReportTextureStat(CVarTextureProfilerEnableTextureCSV, "Total", CSV_CATEGORY_INDEX(TextureProfiler), ToMB(TotalTextureSize.PeakSize), ECsvCustomStatOp::Set);
+	ReportTextureStat(CVarTextureProfilerEnableTextureCSV, "Total", CSV_CATEGORY_INDEX(TextureWasteProfiler), ToKB(TotalTextureSize.AllocationWaste), ECsvCustomStatOp::Set);
+	ReportTextureStat(CVarTextureProfilerEnableTextureCSV, "Other", CSV_CATEGORY_INDEX(TextureProfiler), ToMB(OtherTextureSizes.PeakSize), ECsvCustomStatOp::Set);
+	ReportTextureStat(CVarTextureProfilerEnableTextureCSV, "Other", CSV_CATEGORY_INDEX(TextureWasteProfiler), ToKB(OtherTextureSizes.AllocationWaste), ECsvCustomStatOp::Set);
+
+	TotalRenderTargetSize.ResetPeakSize();
+	TotalTextureSize.ResetPeakSize();
+#endif //CSV_PROFILER
+}
+
+void FTextureProfiler::DumpTextures(bool RenderTargets, bool CombineTextureNames, bool AsCSV, FOutputDevice& OutputDevice)
+{
+	FScopeLock Lock(&TextureMapCS);
+	const TCHAR Sep = AsCSV ? TEXT(',') : TEXT('\t');
+
+	if (!AsCSV)
+	{
+		if (RenderTargets)
+		{
+			OutputDevice.Log(TEXT("Dumping all allocated RHI render targets"));
+		}
+		else
+		{
+			OutputDevice.Log(TEXT("Dumping all allocated RHI textures"));
+		}
+	}
+
+	FTexureDetails& TotalValue = RenderTargets ? TotalRenderTargetSize : TotalTextureSize;
+
+	if (CombineTextureNames)
+	{
+		OutputDevice.Logf(TEXT("%s%c%s%c%s%c%s"), TEXT("Texture"), Sep, TEXT("Combined Size(KB)"), Sep, TEXT("Count"), Sep, TEXT("Combined Wasted(KB)"));
+
+		size_t MinTextureSize = CVarTextureProfilerMinTextureSizeMB.GetValueOnAnyThread() * 1024 * 1024;
+		size_t MinRenderTargetSize = CVarTextureProfilerMinRenderTargetSizeMB.GetValueOnAnyThread() * 1024 * 1024;
+
+		size_t MinSize = RenderTargets ? MinRenderTargetSize : MinTextureSize;
+
+		TMap<FName, FTexureDetails>& CombinedSizes = RenderTargets ? CombinedRenderTargetSizes : CombinedTextureSizes;
+		FTexureDetails OtherTextureSizes;
+		for (const auto& Pair : CombinedSizes)
+		{
+			// Don't report empty sizes or sizes less than the minimum
+			if (Pair.Value.Size == 0 || Pair.Value.Size < MinSize)
+			{
+				OtherTextureSizes += Pair.Value;
+				continue;
+			}
+
+			FString NameString = Pair.Value.TextureName.ToString();
+			OutputDevice.Logf(TEXT("%s%c%d%c%d%c%d"), *NameString, Sep, ToMB(Pair.Value.Size), Sep, Pair.Value.Count, Sep, ToMB(Pair.Value.AllocationWaste));
+		}
+
+		if (OtherTextureSizes.Size > 0)
+		{
+			OutputDevice.Logf(TEXT("%s%c%d%c%d%c%d"), TEXT("Other"), Sep, ToMB(OtherTextureSizes.Size), Sep, OtherTextureSizes.Count, Sep, ToMB(OtherTextureSizes.AllocationWaste));
+		}
+		OutputDevice.Logf(TEXT("%s%c%d%c%d%c%d"), TEXT("Total"), Sep, ToMB(TotalValue.Size), Sep, TotalValue.Count, Sep, ToMB(TotalValue.AllocationWaste));
+	}
+	else
+	{
+		OutputDevice.Logf(TEXT("%s%c%s%c%s%c%s"), TEXT("Texture"), Sep, TEXT("Size(KB)"), Sep, TEXT("Wasted(KB)"), Sep, TEXT("Alignment"));
+		for (const auto& Pair : TexturesMap)
+		{
+			if (Pair.Value.IsRenderTarget != RenderTargets)
+			{
+				continue;
+			}
+
+			FString NameString = Pair.Value.TextureName.ToString();
+			OutputDevice.Logf(TEXT("%s%c%d%c%d%c%p"), *NameString, Sep, ToMB(Pair.Value.Size), Sep, ToMB(Pair.Value.AllocationWaste), Sep, Pair.Value.Align);
+		}
+
+		OutputDevice.Logf(TEXT("%s%c%d%c%d%c%p"), TEXT("Total"), Sep, ToMB(TotalValue.Size), Sep, ToMB(TotalValue.AllocationWaste), Sep, TotalValue.Align);
+	}
+}
+
+#endif //TEXTURE_PROFILER_ENABLED
diff --git a/Engine/Source/Runtime/RHI/Public/TextureProfiler.h b/Engine/Source/Runtime/RHI/Public/TextureProfiler.h
new file mode 100644
--- /dev/null
+++ b/Engine/Source/Runtime/RHI/Public/TextureProfiler.h
@@ -0,0 +1,86 @@
+//Copyright Epic Games, Inc. All Rights Reserved.
+
+#pragma once
+
+#include "Misc/Build.h"
+#include "Templates/Function.h"
+#include "Misc/CoreMisc.h"
+#include "Containers/Map.h"
+#include "UObject/NameTypes.h"
+
+#define TEXTURE_PROFILER_ENABLED WITH_ENGINE && !UE_BUILD_SHIPPING
+
+#if TEXTURE_PROFILER_ENABLED
+/**
+* FTextureProfiler class. This manages recording and reporting texture allocations in the RHI
+*/
+class FTextureProfiler
+{
+	static FTextureProfiler* Instance;
+public:
+	// Singleton interface
+	static FTextureProfiler* Get();
+
+	void Init();
+
+	void DumpTextures(bool RenderTargets, bool CombineTextureNames, bool AsCSV, FOutputDevice& OutputDevice);
+
+	template <typename RHITextureType>
+	void AddTextureAllocation(RHITextureType* UniqueTexturePtr, FName TextureName, size_t Size, uint32 Alignment, size_t AllocationWaste, bool IsRenderTarget)
+	{
+		AddTextureAllocationInternal(reinterpret_cast<void*>(UniqueTexturePtr), TextureName, Size, Alignment, AllocationWaste, IsRenderTarget);
+	}
+
+	template <typename RHITextureType>
+	void RemoveTextureAllocation(RHITextureType* UniqueTexturePtr, bool IsRenderTarget)
+	{
+		RemoveTextureAllocationInternal(reinterpret_cast<void*>(UniqueTexturePtr), IsRenderTarget);
+	}
+
+	template <typename RHITextureType>
+	void ChangeTextureName(RHITextureType* UniqueTexturePtr, FName NewName, bool IsRenderTarget)
+	{
+		ChangeTextureNameInternal(reinterpret_cast<void*>(UniqueTexturePtr), NewName, IsRenderTarget);
+	}
+private:
+
+	FTextureProfiler() = default;
+	FTextureProfiler(const FTextureProfiler&) = delete;
+	FTextureProfiler(const FTextureProfiler&&) = delete;
+
+	void AddTextureAllocationInternal(void* UniqueTexturePtr, FName TextureName, size_t Size, uint32 Alignment, size_t AllocationWaste, bool IsRenderTarget);
+	void RemoveTextureAllocationInternal(void* UniqueTexturePtr, bool IsRenderTarget);
+	void ChangeTextureNameInternal(void* UniqueTexturePtr, FName NewName, bool IsRenderTarget);
+	void Update();
+
+	FCriticalSection TextureMapCS;
+	struct FTexureDetails
+	{
+		FTexureDetails() = default;
+		FTexureDetails(FName InTextureName, size_t InSize, uint32 InAlign, size_t InAllocationWaste, bool InIsRenderTarget);
+
+		void SetName(FName InTextureName);
+		void ResetPeakSize();
+
+		FTexureDetails& operator+=(const FTexureDetails& Other);
+		FTexureDetails& operator-=(const FTexureDetails& Other);
+
+		FName TextureName;
+		char TextureNameString[NAME_SIZE];
+		size_t Size = 0;
+		size_t PeakSize = 0;
+		uint32 Align = 0;
+		size_t AllocationWaste = 0;
+		int Count = 0;
+		bool IsRenderTarget = false;
+	};
+
+	TMap<void*, FTexureDetails> TexturesMap;
+
+	// Keep track of the totals separately to reduce the cost of rounding error for sizes
+	FTexureDetails TotalTextureSize;
+	FTexureDetails TotalRenderTargetSize;
+	TMap<FName, FTexureDetails> CombinedTextureSizes;
+	TMap<FName, FTexureDetails> CombinedRenderTargetSizes;
+};
+#endif //TEXTURE_PROFILER_ENABLED
