[RHI] - removing superfluous and crashing line after discussing with Krzysztof

diff --git a/Engine/Source/Runtime/RHI/Public/RHICommandList.h b/Engine/Source/Runtime/RHI/Public/RHICommandList.h
--- a/Engine/Source/Runtime/RHI/Public/RHICommandList.h
+++ b/Engine/Source/Runtime/RHI/Public/RHICommandList.h
@@ -3347,7 +3347,6 @@ public:
 			{
 				GraphicsPSOInit.RenderTargetFormats[i] = UE_PIXELFORMAT_TO_UINT8(PSOContext.CachedRenderTargets[i].Texture->GetFormat());
 				GraphicsPSOInit.RenderTargetFlags[i] = PSOContext.CachedRenderTargets[i].Texture->GetFlags();
-				const FRHITexture2DArray* TextureArray = PSOContext.CachedRenderTargets[i].Texture->GetTexture2DArray();
 			}
 			else
 			{
