Texture profiler working on D3D12.
Texture profiler enabled and interface cleaned up for both platforms

diff --git a/Engine/Source/Runtime/D3D12RHI/Private/D3D12Texture.cpp b/Engine/Source/Runtime/D3D12RHI/Private/D3D12Texture.cpp
--- a/Engine/Source/Runtime/D3D12RHI/Private/D3D12Texture.cpp
+++ b/Engine/Source/Runtime/D3D12RHI/Private/D3D12Texture.cpp
@@ -6,6 +6,7 @@
 
 #include "D3D12RHIPrivate.h"
 #include "D3D12RHIBridge.h"
+#include "TextureProfiler.h"
 
 int64 FD3D12GlobalStats::GDedicatedVideoMemory = 0;
 int64 FD3D12GlobalStats::GDedicatedSystemMemory = 0;
@@ -330,8 +331,23 @@ TStatId FD3D12TextureStats::GetD3D12StatEnum(D3D12_RESOURCE_FLAGS MiscFlags)
 // Note: This function can be called from many different threads
 // @param TextureSize >0 to allocate, <0 to deallocate
 // @param b3D true:3D, false:2D or cube map
-void FD3D12TextureStats::UpdateD3D12TextureStats(const D3D12_RESOURCE_DESC& Desc, int64 TextureSize, bool b3D, bool bCubeMap, bool bStreamable)
+template<typename TD3D12Texture>
+void FD3D12TextureStats::UpdateD3D12TextureStats(TD3D12Texture& Texture, const D3D12_RESOURCE_DESC& Desc, int64 TextureSize, bool b3D, bool bCubeMap, bool bStreamable, bool bNewTexture)
 {
+
+#if TEXTURE_PROFILER_ENABLED
+
+	if (!bNewTexture && 
+		!Texture.ResourceLocation.IsTransient() 
+		&& !EnumHasAnyFlags(Texture.GetFlags(), TexCreate_Virtual)
+		&& Texture.ResourceLocation.GetType() != FD3D12ResourceLocation::ResourceLocationType::eAliased
+		&& Texture.ResourceLocation.GetType() != FD3D12ResourceLocation::ResourceLocationType::eHeapAliased)
+	{
+		uint64 SafeSize = (uint64)(TextureSize >= 0 ? TextureSize : 0);
+		FTextureProfiler::Get()->UpdateTextureAllocation(&Texture, SafeSize, Desc.Alignment, 0);
+	}
+#endif
+
 	if (TextureSize == 0)
 	{
 		return;
@@ -387,12 +403,24 @@ void FD3D12TextureStats::D3D12TextureAllocated(TD3D12Texture2D<BaseResourceType>
 
 			Texture.SetMemorySize(TextureSize);
 
-			UpdateD3D12TextureStats(*Desc, TextureSize, false, Texture.IsCubemap(), Texture.IsStreamable());
+			UpdateD3D12TextureStats(Texture, *Desc, TextureSize, false, Texture.IsCubemap(), Texture.IsStreamable(), true);
 		}
 		else
 		{
 			Texture.SetMemorySize(Texture.ResourceLocation.GetSize());
 		}
+
+#if TEXTURE_PROFILER_ENABLED
+		if (!EnumHasAnyFlags(Texture.GetFlags(), TexCreate_Virtual)
+			&& !Texture.ResourceLocation.IsTransient()
+			&& Texture.ResourceLocation.GetType() != FD3D12ResourceLocation::ResourceLocationType::eAliased
+			&& Texture.ResourceLocation.GetType() != FD3D12ResourceLocation::ResourceLocationType::eHeapAliased)
+		{
+			size_t Size = Texture.GetMemorySize();
+			uint32 Alignment = Desc->Alignment;
+			FTextureProfiler::Get()->AddTextureAllocation(&Texture, Size, Alignment, 0);
+		}
+#endif
 	}
 }
 
@@ -410,7 +438,17 @@ void FD3D12TextureStats::D3D12TextureDeleted(TD3D12Texture2D<BaseResourceType>&
 			const int64 TextureSize = Texture.GetMemorySize();
 			ensure(TextureSize > 0 || EnumHasAnyFlags(Texture.Flags, TexCreate_Virtual) || Texture.GetAliasingSourceTexture() != nullptr);
 
-			UpdateD3D12TextureStats(Desc, -TextureSize, false, Texture.IsCubemap(), Texture.IsStreamable());
+			UpdateD3D12TextureStats(Texture, Desc, -TextureSize, false, Texture.IsCubemap(), Texture.IsStreamable(), false);
+
+#if TEXTURE_PROFILER_ENABLED
+			if (!EnumHasAnyFlags(Texture.GetFlags(), TexCreate_Virtual)
+				&& !Texture.ResourceLocation.IsTransient()
+				&& Texture.ResourceLocation.GetType() != FD3D12ResourceLocation::ResourceLocationType::eAliased
+				&& Texture.ResourceLocation.GetType() != FD3D12ResourceLocation::ResourceLocationType::eHeapAliased)
+			{
+				FTextureProfiler::Get()->RemoveTextureAllocation(&Texture);
+			}
+#endif
 		}
 	}
 }
@@ -426,21 +464,33 @@ void FD3D12TextureStats::D3D12TextureAllocated(FD3D12Texture3D& Texture)
 
 	if (D3D12Texture3D)
 	{
+		const D3D12_RESOURCE_DESC& Desc = D3D12Texture3D->GetDesc();
 		// Don't update state for virtual or transient textures	
 		if (!EnumHasAnyFlags(Texture.GetFlags(), TexCreate_Virtual) && !Texture.ResourceLocation.IsTransient())
 		{
-			const D3D12_RESOURCE_DESC& Desc = D3D12Texture3D->GetDesc();
 			const D3D12_RESOURCE_ALLOCATION_INFO AllocationInfo = Texture.GetParentDevice()->GetDevice()->GetResourceAllocationInfo(0, 1, &Desc);
 			const int64 TextureSize = AllocationInfo.SizeInBytes;
 
 			Texture.SetMemorySize(TextureSize);
 
-			UpdateD3D12TextureStats(Desc, TextureSize, true, false, Texture.IsStreamable());
+			UpdateD3D12TextureStats(Texture, Desc, TextureSize, true, false, Texture.IsStreamable(), true);
 		}
 		else
 		{
 			Texture.SetMemorySize(Texture.ResourceLocation.GetSize());
 		}
+
+#if TEXTURE_PROFILER_ENABLED
+		if (!EnumHasAnyFlags(Texture.GetFlags(), TexCreate_Virtual)
+			&& !Texture.ResourceLocation.IsTransient()
+			&& Texture.ResourceLocation.GetType() != FD3D12ResourceLocation::ResourceLocationType::eAliased
+			&& Texture.ResourceLocation.GetType() != FD3D12ResourceLocation::ResourceLocationType::eHeapAliased)
+		{
+			size_t Size = Texture.GetMemorySize();
+			uint32 Alignment = Desc.Alignment;
+			FTextureProfiler::Get()->AddTextureAllocation(&Texture, Size, Alignment, 0);
+		}
+#endif
 	}
 }
 
@@ -457,8 +507,18 @@ void FD3D12TextureStats::D3D12TextureDeleted(FD3D12Texture3D& Texture)
 			const int64 TextureSize = Texture.GetMemorySize();
 			if (TextureSize > 0)
 			{
-				UpdateD3D12TextureStats(Desc, -TextureSize, true, false, Texture.IsStreamable());
+				UpdateD3D12TextureStats(Texture, Desc, -TextureSize, true, false, Texture.IsStreamable(), false);
 			}
+
+#if TEXTURE_PROFILER_ENABLED
+			if (!EnumHasAnyFlags(Texture.GetFlags(), TexCreate_Virtual)
+				&& !Texture.ResourceLocation.IsTransient()
+				&& Texture.ResourceLocation.GetType() != FD3D12ResourceLocation::ResourceLocationType::eAliased
+				&& Texture.ResourceLocation.GetType() != FD3D12ResourceLocation::ResourceLocationType::eHeapAliased)
+			{
+				FTextureProfiler::Get()->RemoveTextureAllocation(&Texture);
+			}
+#endif
 		}
 	}
 }
@@ -2950,6 +3010,20 @@ void FD3D12DynamicRHI::RHIBindDebugLabelName(FRHITexture* TextureRHI, const TCHA
 
 	// Also set on RHI object
 	TextureRHI->SetName(Name);
+
+#if TEXTURE_PROFILER_ENABLED
+
+	FD3D12TextureBase* D3D12Texture = GetD3D12TextureFromRHITexture(TextureRHI);
+	
+	if (!EnumHasAnyFlags(TextureRHI->GetFlags(), TexCreate_Virtual)
+		&& !D3D12Texture->ResourceLocation.IsTransient()
+		&& D3D12Texture->ResourceLocation.GetType() != FD3D12ResourceLocation::ResourceLocationType::eAliased
+		&& D3D12Texture->ResourceLocation.GetType() != FD3D12ResourceLocation::ResourceLocationType::eHeapAliased)
+	{
+		FTextureProfiler::Get()->UpdateTextureName(TextureRHI);
+	}
+	
+#endif
 }
 
 void FD3D12DynamicRHI::RHIVirtualTextureSetFirstMipInMemory(FRHITexture2D* TextureRHI, uint32 FirstMip)
diff --git a/Engine/Source/Runtime/D3D12RHI/Private/D3D12Texture.h b/Engine/Source/Runtime/D3D12RHI/Private/D3D12Texture.h
--- a/Engine/Source/Runtime/D3D12RHI/Private/D3D12Texture.h
+++ b/Engine/Source/Runtime/D3D12RHI/Private/D3D12Texture.h
@@ -612,7 +612,8 @@ public:
 	// @param TextureSize >0 to allocate, <0 to deallocate
 	// @param b3D true:3D, false:2D or cube map
 	// @param bStreamable true:Streamable, false:not streamable
-	static void UpdateD3D12TextureStats(const D3D12_RESOURCE_DESC& Desc, int64 TextureSize, bool b3D, bool bCubeMap, bool bStreamable);
+	template<typename TD3D12Texture>
+	static void UpdateD3D12TextureStats(TD3D12Texture& Texture, const D3D12_RESOURCE_DESC& Desc, int64 TextureSize, bool b3D, bool bCubeMap, bool bStreamable, bool bNewTexture = false);
 
 	template<typename BaseResourceType>
 	static void D3D12TextureAllocated(TD3D12Texture2D<BaseResourceType>& Texture, const D3D12_RESOURCE_DESC *Desc = nullptr);
diff --git a/Engine/Source/Runtime/RHI/Private/DynamicRHI.cpp b/Engine/Source/Runtime/RHI/Private/DynamicRHI.cpp
--- a/Engine/Source/Runtime/RHI/Private/DynamicRHI.cpp
+++ b/Engine/Source/Runtime/RHI/Private/DynamicRHI.cpp
@@ -14,6 +14,7 @@
 #include "GenericPlatform/GenericPlatformDriver.h"
 #include "GenericPlatform/GenericPlatformCrashContext.h"
 #include "PipelineStateCache.h"
+#include "TextureProfiler.h"
 
 #if defined(NV_GEFORCENOW) && NV_GEFORCENOW
 #include "GeForceNOWWrapper.h"
@@ -289,6 +290,10 @@ void RHIInit(bool bHasEditorToken)
 				FGenericCrashContext::SetEngineData(TEXT("RHI.DriverDate"), GRHIAdapterDriverDate);
 				FGenericCrashContext::SetEngineData(TEXT("RHI.FeatureLevel"), FeatureLevelString);
 				FGenericCrashContext::SetEngineData(TEXT("RHI.GPUVendor"), RHIVendorIdToString());
+
+#if TEXTURE_PROFILER_ENABLED
+				FTextureProfiler::Get()->Init();
+#endif
 			}
 #if PLATFORM_ALLOW_NULL_RHI
 			else
diff --git a/Engine/Source/Runtime/RHI/Private/TextureProfiler.cpp b/Engine/Source/Runtime/RHI/Private/TextureProfiler.cpp
--- a/Engine/Source/Runtime/RHI/Private/TextureProfiler.cpp
+++ b/Engine/Source/Runtime/RHI/Private/TextureProfiler.cpp
@@ -2,11 +2,12 @@
 
 #include "TextureProfiler.h"
 
+#include "ProfilingDebugging/CsvProfiler.h"
+#include "RHI.h"
+#include "RHIResources.h"
 #include "UObject/NameTypes.h"
 
 #if TEXTURE_PROFILER_ENABLED
-#include "ProfilingDebugging/CsvProfiler.h"
-
 static int32 ToMB(uint64 Value)
 {
 	const size_t MB = 1024 * 1024;
@@ -16,15 +17,6 @@ static int32 ToMB(uint64 Value)
 	return Value32;
 }
 
-static int32 ToKB(uint64 Value)
-{
-	const size_t KB = 1024;
-	uint64 Value64 = (Value + (KB - 1)) / KB;
-	int32 Value32 = static_cast<int32>(Value64);
-	check(Value64 == Value32);
-	return Value32;
-}
-
 CSV_DEFINE_CATEGORY(RenderTargetProfiler, true);
 CSV_DEFINE_CATEGORY(RenderTargetWasteProfiler, true);
 CSV_DEFINE_CATEGORY(TextureProfiler, true);
@@ -82,14 +74,15 @@ static FAutoConsoleCommand CmdTextureProfilerDumpTextures(
 
 FTextureProfiler* FTextureProfiler::Instance = nullptr;
 
-FTextureProfiler::FTexureDetails::FTexureDetails(FName InTextureName, size_t InSize, uint32 InAlign, size_t InAllocationWaste, bool InIsRenderTarget)
-	: TextureName(InTextureName)
+
+FTextureProfiler::FTexureDetails::FTexureDetails(FRHITexture* Texture, size_t InSize, uint32 InAlign, size_t InAllocationWaste)
+	: TextureName(Texture->GetName().IsValid() ? Texture->GetName() : FName("None"))
 	, Size(InSize)
 	, Align(InAlign)
 	, AllocationWaste(InAllocationWaste)
-	, IsRenderTarget(InIsRenderTarget)
+	, IsRenderTarget(EnumHasAnyFlags(Texture->GetFlags(), TexCreate_RenderTargetable | TexCreate_DepthStencilTargetable))
 {
-	InTextureName.GetPlainANSIString(TextureNameString);
+	TextureName.GetPlainANSIString(TextureNameString);
 }
 
 void FTextureProfiler::FTexureDetails::SetName(FName InTextureName)
@@ -136,57 +129,82 @@ FTextureProfiler* FTextureProfiler::Get()
 	return Instance;
 }
 
-void FTextureProfiler::AddTextureAllocationInternal(void* UniqueTexturePtr, FName TextureName, size_t Size, uint32 Alignment, size_t AllocationWaste, bool IsRenderTarget)
+void FTextureProfiler::AddTextureAllocation(FRHITexture* UniqueTexturePtr, size_t Size, uint32 Alignment, size_t AllocationWaste)
 {
 	FScopeLock Lock(&TextureMapCS);
 
 	FTexureDetails AddedDetails(
-		TextureName,
+		UniqueTexturePtr,
 		Size,
 		Alignment,
-		AllocationWaste,
-		IsRenderTarget
-	);
+		AllocationWaste);
 
 	AddedDetails.Count = 1;
 
-	FTexureDetails& TotalValue = IsRenderTarget ? TotalRenderTargetSize : TotalTextureSize;
+	FTexureDetails& TotalValue = AddedDetails.IsRenderTarget ? TotalRenderTargetSize : TotalTextureSize;
 	TotalValue += AddedDetails;
 
+	// Make sure we have a sane value
 	check(ToMB(TotalValue.Size) >= 0);
 
-	FTexureDetails& NamedValue = IsRenderTarget ? CombinedRenderTargetSizes.FindOrAdd(TextureName) : CombinedTextureSizes.FindOrAdd(TextureName);
+	FTexureDetails& NamedValue = AddedDetails.IsRenderTarget ? CombinedRenderTargetSizes.FindOrAdd(AddedDetails.TextureName) : CombinedTextureSizes.FindOrAdd(AddedDetails.TextureName);
 	NamedValue += AddedDetails;
 
 	if (NamedValue.TextureName.IsNone() || !NamedValue.TextureName.IsValid())
 	{
-		NamedValue.SetName(TextureName);
+		NamedValue.SetName(AddedDetails.TextureName);
 	}
 
 	TexturesMap.Add(UniqueTexturePtr, AddedDetails);
 }
 
-void FTextureProfiler::RemoveTextureAllocationInternal(void* UniqueTexturePtr, bool IsRenderTarget)
+void FTextureProfiler::UpdateTextureAllocation(FRHITexture* UniqueTexturePtr, size_t Size, uint32 Alignment, size_t AllocationWaste)
+{
+	FScopeLock Lock(&TextureMapCS);
+
+	FTexureDetails UpdatedDetails(
+		UniqueTexturePtr,
+		Size,
+		Alignment,
+		AllocationWaste);
+
+	FTexureDetails& ExistingTexture = TexturesMap[UniqueTexturePtr];
+	FTexureDetails& TotalValue = UpdatedDetails.IsRenderTarget ? TotalRenderTargetSize : TotalTextureSize;
+
+	TotalValue -= ExistingTexture;
+	TotalValue += UpdatedDetails;
+
+	check(ToMB(TotalValue.Size) >= 0);
+
+	FTexureDetails& NamedValue = UpdatedDetails.IsRenderTarget ? CombinedRenderTargetSizes.FindOrAdd(UpdatedDetails.TextureName) : CombinedTextureSizes.FindOrAdd(UpdatedDetails.TextureName);
+	NamedValue -= ExistingTexture;
+	NamedValue += UpdatedDetails;
+
+	ExistingTexture = UpdatedDetails;
+}
+
+void FTextureProfiler::RemoveTextureAllocation(FRHITexture* UniqueTexturePtr)
 {
 	FScopeLock Lock(&TextureMapCS);
 	FTexureDetails Details = TexturesMap.FindAndRemoveChecked(UniqueTexturePtr);
 	
-	FTexureDetails& TotalValue = IsRenderTarget ? TotalRenderTargetSize : TotalTextureSize;
+	FTexureDetails& TotalValue = Details.IsRenderTarget ? TotalRenderTargetSize : TotalTextureSize;
 	TotalValue -= Details;
 
-	FTexureDetails& NamedValue = IsRenderTarget ? CombinedRenderTargetSizes[Details.TextureName] : CombinedTextureSizes[Details.TextureName];
+	FTexureDetails& NamedValue = Details.IsRenderTarget ? CombinedRenderTargetSizes[Details.TextureName] : CombinedTextureSizes[Details.TextureName];
 	NamedValue -= Details;
 }
 
-void FTextureProfiler::ChangeTextureNameInternal(void* UniqueTexturePtr, FName NewName, bool IsRenderTarget)
+void FTextureProfiler::UpdateTextureName(FRHITexture* UniqueTexturePtr)
 {
 	FScopeLock Lock(&TextureMapCS);
 	FTexureDetails& Details = TexturesMap[UniqueTexturePtr];
 	FName OldName = Details.TextureName;
+	FName NewName = UniqueTexturePtr->GetName();
 	Details.SetName(NewName);
 
-	FTexureDetails& OldCombinedValue = IsRenderTarget ? CombinedRenderTargetSizes[OldName] : CombinedTextureSizes[OldName];
-	FTexureDetails& NewCombinedValue = IsRenderTarget ? CombinedRenderTargetSizes.FindOrAdd(NewName) : CombinedTextureSizes.FindOrAdd(NewName);
+	FTexureDetails& OldCombinedValue = Details.IsRenderTarget ? CombinedRenderTargetSizes[OldName] : CombinedTextureSizes[OldName];
+	FTexureDetails& NewCombinedValue = Details.IsRenderTarget ? CombinedRenderTargetSizes.FindOrAdd(NewName) : CombinedTextureSizes.FindOrAdd(NewName);
 
 	if (NewCombinedValue.TextureName.IsNone() || !NewCombinedValue.TextureName.IsValid())
 	{
@@ -247,7 +265,7 @@ void FTextureProfiler::Update()
 		}
 
 		ReportTextureStat(CVarTextureProfilerEnableRenderTargetCSV, Pair.Value.TextureNameString, CSV_CATEGORY_INDEX(RenderTargetProfiler), ToMB(Pair.Value.PeakSize), ECsvCustomStatOp::Set);
-		ReportTextureStat(CVarTextureProfilerEnableRenderTargetCSV, Pair.Value.TextureNameString, CSV_CATEGORY_INDEX(RenderTargetWasteProfiler), ToKB(Pair.Value.AllocationWaste), ECsvCustomStatOp::Set);
+		ReportTextureStat(CVarTextureProfilerEnableRenderTargetCSV, Pair.Value.TextureNameString, CSV_CATEGORY_INDEX(RenderTargetWasteProfiler), ToMB(Pair.Value.AllocationWaste), ECsvCustomStatOp::Set);
 
 		Pair.Value.ResetPeakSize();
 	}
@@ -263,20 +281,20 @@ void FTextureProfiler::Update()
 		}
 
 		ReportTextureStat(CVarTextureProfilerEnableTextureCSV, Pair.Value.TextureNameString, CSV_CATEGORY_INDEX(TextureProfiler), ToMB(Pair.Value.PeakSize), ECsvCustomStatOp::Set);
-		ReportTextureStat(CVarTextureProfilerEnableTextureCSV, Pair.Value.TextureNameString, CSV_CATEGORY_INDEX(TextureWasteProfiler), ToKB(Pair.Value.AllocationWaste), ECsvCustomStatOp::Set);
+		ReportTextureStat(CVarTextureProfilerEnableTextureCSV, Pair.Value.TextureNameString, CSV_CATEGORY_INDEX(TextureWasteProfiler), ToMB(Pair.Value.AllocationWaste), ECsvCustomStatOp::Set);
 
 		Pair.Value.ResetPeakSize();
 	}
 
 	ReportTextureStat(CVarTextureProfilerEnableRenderTargetCSV, "Total", CSV_CATEGORY_INDEX(RenderTargetProfiler), ToMB(TotalRenderTargetSize.PeakSize), ECsvCustomStatOp::Set);
-	ReportTextureStat(CVarTextureProfilerEnableRenderTargetCSV, "Total", CSV_CATEGORY_INDEX(RenderTargetWasteProfiler), ToKB(TotalRenderTargetSize.AllocationWaste), ECsvCustomStatOp::Set);
+	ReportTextureStat(CVarTextureProfilerEnableRenderTargetCSV, "Total", CSV_CATEGORY_INDEX(RenderTargetWasteProfiler), ToMB(TotalRenderTargetSize.AllocationWaste), ECsvCustomStatOp::Set);
 	ReportTextureStat(CVarTextureProfilerEnableRenderTargetCSV, "Other", CSV_CATEGORY_INDEX(RenderTargetProfiler), ToMB(OtherRenderTargetSizes.PeakSize), ECsvCustomStatOp::Set);
-	ReportTextureStat(CVarTextureProfilerEnableRenderTargetCSV, "Other", CSV_CATEGORY_INDEX(RenderTargetWasteProfiler), ToKB(OtherRenderTargetSizes.AllocationWaste), ECsvCustomStatOp::Set);
+	ReportTextureStat(CVarTextureProfilerEnableRenderTargetCSV, "Other", CSV_CATEGORY_INDEX(RenderTargetWasteProfiler), ToMB(OtherRenderTargetSizes.AllocationWaste), ECsvCustomStatOp::Set);
 
 	ReportTextureStat(CVarTextureProfilerEnableTextureCSV, "Total", CSV_CATEGORY_INDEX(TextureProfiler), ToMB(TotalTextureSize.PeakSize), ECsvCustomStatOp::Set);
-	ReportTextureStat(CVarTextureProfilerEnableTextureCSV, "Total", CSV_CATEGORY_INDEX(TextureWasteProfiler), ToKB(TotalTextureSize.AllocationWaste), ECsvCustomStatOp::Set);
+	ReportTextureStat(CVarTextureProfilerEnableTextureCSV, "Total", CSV_CATEGORY_INDEX(TextureWasteProfiler), ToMB(TotalTextureSize.AllocationWaste), ECsvCustomStatOp::Set);
 	ReportTextureStat(CVarTextureProfilerEnableTextureCSV, "Other", CSV_CATEGORY_INDEX(TextureProfiler), ToMB(OtherTextureSizes.PeakSize), ECsvCustomStatOp::Set);
-	ReportTextureStat(CVarTextureProfilerEnableTextureCSV, "Other", CSV_CATEGORY_INDEX(TextureWasteProfiler), ToKB(OtherTextureSizes.AllocationWaste), ECsvCustomStatOp::Set);
+	ReportTextureStat(CVarTextureProfilerEnableTextureCSV, "Other", CSV_CATEGORY_INDEX(TextureWasteProfiler), ToMB(OtherTextureSizes.AllocationWaste), ECsvCustomStatOp::Set);
 
 	TotalRenderTargetSize.ResetPeakSize();
 	TotalTextureSize.ResetPeakSize();
diff --git a/Engine/Source/Runtime/RHI/Public/TextureProfiler.h b/Engine/Source/Runtime/RHI/Public/TextureProfiler.h
--- a/Engine/Source/Runtime/RHI/Public/TextureProfiler.h
+++ b/Engine/Source/Runtime/RHI/Public/TextureProfiler.h
@@ -11,10 +11,13 @@
 #define TEXTURE_PROFILER_ENABLED WITH_ENGINE && !UE_BUILD_SHIPPING
 
 #if TEXTURE_PROFILER_ENABLED
+
+class FRHITexture;
+
 /**
 * FTextureProfiler class. This manages recording and reporting texture allocations in the RHI
 */
-class FTextureProfiler
+class RHI_API FTextureProfiler
 {
 	static FTextureProfiler* Instance;
 public:
@@ -25,39 +28,23 @@ public:
 
 	void DumpTextures(bool RenderTargets, bool CombineTextureNames, bool AsCSV, FOutputDevice& OutputDevice);
 
-	template <typename RHITextureType>
-	void AddTextureAllocation(RHITextureType* UniqueTexturePtr, FName TextureName, size_t Size, uint32 Alignment, size_t AllocationWaste, bool IsRenderTarget)
-	{
-		AddTextureAllocationInternal(reinterpret_cast<void*>(UniqueTexturePtr), TextureName, Size, Alignment, AllocationWaste, IsRenderTarget);
-	}
-
-	template <typename RHITextureType>
-	void RemoveTextureAllocation(RHITextureType* UniqueTexturePtr, bool IsRenderTarget)
-	{
-		RemoveTextureAllocationInternal(reinterpret_cast<void*>(UniqueTexturePtr), IsRenderTarget);
-	}
-
-	template <typename RHITextureType>
-	void ChangeTextureName(RHITextureType* UniqueTexturePtr, FName NewName, bool IsRenderTarget)
-	{
-		ChangeTextureNameInternal(reinterpret_cast<void*>(UniqueTexturePtr), NewName, IsRenderTarget);
-	}
+	void AddTextureAllocation(FRHITexture* UniqueTexturePtr, size_t Size, uint32 Alignment, size_t AllocationWaste);
+	void UpdateTextureAllocation(FRHITexture* UniqueTexturePtr, size_t Size, uint32 Alignment, size_t AllocationWaste);
+	void RemoveTextureAllocation(FRHITexture* UniqueTexturePtr);
+	void UpdateTextureName(FRHITexture* UniqueTexturePtr);
 private:
 
 	FTextureProfiler() = default;
 	FTextureProfiler(const FTextureProfiler&) = delete;
 	FTextureProfiler(const FTextureProfiler&&) = delete;
 
-	void AddTextureAllocationInternal(void* UniqueTexturePtr, FName TextureName, size_t Size, uint32 Alignment, size_t AllocationWaste, bool IsRenderTarget);
-	void RemoveTextureAllocationInternal(void* UniqueTexturePtr, bool IsRenderTarget);
-	void ChangeTextureNameInternal(void* UniqueTexturePtr, FName NewName, bool IsRenderTarget);
 	void Update();
 
 	FCriticalSection TextureMapCS;
 	struct FTexureDetails
 	{
 		FTexureDetails() = default;
-		FTexureDetails(FName InTextureName, size_t InSize, uint32 InAlign, size_t InAllocationWaste, bool InIsRenderTarget);
+		FTexureDetails(FRHITexture* Texture, size_t InSize, uint32 InAlign, size_t InAllocationWaste);
 
 		void SetName(FName InTextureName);
 		void ResetPeakSize();
