Fix missing function overrides and crashes in the RHI validation layer

diff --git a/Engine/Source/Runtime/RHI/Public/RHIValidation.h b/Engine/Source/Runtime/RHI/Public/RHIValidation.h
--- a/Engine/Source/Runtime/RHI/Public/RHIValidation.h
+++ b/Engine/Source/Runtime/RHI/Public/RHIValidation.h
@@ -456,6 +456,13 @@ public:
 	virtual void RHIUpdateShaderResourceView(FRHIShaderResourceView* SRV, FRHIBuffer* Buffer, uint32 Stride, uint8 Format) override final
 	{
 		RHI->RHIUpdateShaderResourceView(SRV, Buffer, Stride, Format);
+		SRV->ViewIdentity = Buffer->GetWholeResourceIdentity();
+	}
+
+	virtual void RHIUpdateShaderResourceView(FRHIShaderResourceView* SRV, FRHIBuffer* Buffer) override final
+	{
+		RHI->RHIUpdateShaderResourceView(SRV, Buffer);
+		SRV->ViewIdentity = Buffer->GetWholeResourceIdentity();
 	}
 
 	/**
@@ -729,6 +736,11 @@ public:
 		return RHI->RHIGenerateMips(Texture);
 	}
 
+	virtual bool RHIRequiresComputeGenerateMips() const final override
+	{
+		return RHI->RHIRequiresComputeGenerateMips();
+	}
+
 	/**
 	* Computes the size in memory required by a given texture.
 	*
diff --git a/Engine/Source/Runtime/RHI/Public/RHIValidationContext.h b/Engine/Source/Runtime/RHI/Public/RHIValidationContext.h
--- a/Engine/Source/Runtime/RHI/Public/RHIValidationContext.h
+++ b/Engine/Source/Runtime/RHI/Public/RHIValidationContext.h
@@ -207,6 +207,27 @@ public:
 
 	virtual void RHICopyToStagingBuffer(FRHIBuffer* SourceBufferRHI, FRHIStagingBuffer* DestinationStagingBufferRHI, uint32 InOffset, uint32 InNumBytes) override final;
 
+	virtual void RHIBuildAccelerationStructures(const TArrayView<const FRayTracingGeometryBuildParams> Params) override final
+	{
+		// #yuriy_todo: explicit transitions and state validation for BLAS
+		RHIContext->RHIBuildAccelerationStructures(Params);
+	}
+
+	virtual void RHIBuildAccelerationStructure(const FRayTracingSceneBuildParams& SceneBuildParams) override final
+	{
+		// #yuriy_todo: validate all referenced BLAS states
+		if (SceneBuildParams.Scene)
+		{
+			Tracker->Assert(SceneBuildParams.Scene->GetWholeResourceIdentity(), ERHIAccess::BVHWrite);
+		}
+		RHIContext->RHIBuildAccelerationStructure(SceneBuildParams);
+	}
+
+	virtual void RHIBindAccelerationStructureMemory(FRHIRayTracingScene* Scene, FRHIBuffer* Buffer, uint32 BufferOffset) override final
+	{
+		RHIContext->RHIBindAccelerationStructureMemory(Scene, Buffer, BufferOffset);
+	}
+
 	inline void LinkToContext(IRHIComputeContext* PlatformContext)
 	{
 		RHIContext = PlatformContext;
