Fix for texture profiler destroying a string that the csv profiler was holding onto, causing a crash when writing csv frames.

diff --git a/Engine/Source/Runtime/RHI/Private/TextureProfiler.cpp b/Engine/Source/Runtime/RHI/Private/TextureProfiler.cpp
--- a/Engine/Source/Runtime/RHI/Private/TextureProfiler.cpp
+++ b/Engine/Source/Runtime/RHI/Private/TextureProfiler.cpp
@@ -24,7 +24,7 @@ CSV_DEFINE_CATEGORY(TextureWasteProfiler, true);
 
 static TAutoConsoleVariable<int32> CVarTextureProfilerMinTextureSizeMB(
 	TEXT("r.TextureProfiler.MinTextureSizeMB"),
-	16,
+	0,
 	TEXT("The minimum size for any texture to be reported.  All textures below this threshold will be reported as Other."));
 
 static TAutoConsoleVariable<int32> CVarTextureProfilerMinRenderTargetSizeMB(
@@ -74,21 +74,26 @@ static FAutoConsoleCommand CmdTextureProfilerDumpTextures(
 
 FTextureProfiler* FTextureProfiler::Instance = nullptr;
 
-
 FTextureProfiler::FTexureDetails::FTexureDetails(FRHITexture* Texture, size_t InSize, uint32 InAlign, size_t InAllocationWaste)
-	: TextureName(Texture->GetName().IsValid() ? Texture->GetName() : FName("None"))
-	, Size(InSize)
+	: Size(InSize)
+	, PeakSize(0)
 	, Align(InAlign)
 	, AllocationWaste(InAllocationWaste)
 	, IsRenderTarget(EnumHasAnyFlags(Texture->GetFlags(), TexCreate_RenderTargetable | TexCreate_DepthStencilTargetable))
+	, TextureName(Texture->GetName().IsValid() ? Texture->GetName() : FName("None"))
 {
-	TextureName.GetPlainANSIString(TextureNameString);
+	SetName(TextureName);
 }
 
 void FTextureProfiler::FTexureDetails::SetName(FName InTextureName)
 {
 	TextureName = InTextureName;
-	InTextureName.GetPlainANSIString(TextureNameString);
+	TextureName.SetNumber(0);
+
+	// Do not delete the texture name string, it is not owned by this object
+	//delete TextureNameString;
+
+	TextureNameString = FTextureProfiler::Get()->GetTextureNameString(InTextureName);
 }
 
 void FTextureProfiler::FTexureDetails::ResetPeakSize()
@@ -96,6 +101,12 @@ void FTextureProfiler::FTexureDetails::ResetPeakSize()
 	PeakSize = Size;
 }
 
+FTextureProfiler::FTexureDetails::~FTexureDetails()
+{
+	// Do not delete the texture name string, it is not owned by this object
+	//delete TextureNameString;
+}
+
 FTextureProfiler::FTexureDetails& FTextureProfiler::FTexureDetails::operator+=(const FTextureProfiler::FTexureDetails& Other)
 {
 	Size += Other.Size;
@@ -147,12 +158,12 @@ void FTextureProfiler::AddTextureAllocation(FRHITexture* UniqueTexturePtr, size_
 	// Make sure we have a sane value
 	check(ToMB(TotalValue.Size) >= 0);
 
-	FTexureDetails& NamedValue = AddedDetails.IsRenderTarget ? CombinedRenderTargetSizes.FindOrAdd(AddedDetails.TextureName) : CombinedTextureSizes.FindOrAdd(AddedDetails.TextureName);
+	FTexureDetails& NamedValue = AddedDetails.IsRenderTarget ? CombinedRenderTargetSizes.FindOrAdd(AddedDetails.GetTextureName()) : CombinedTextureSizes.FindOrAdd(AddedDetails.GetTextureName());
 	NamedValue += AddedDetails;
 
-	if (NamedValue.TextureName.IsNone() || !NamedValue.TextureName.IsValid())
+	if (NamedValue.GetTextureName().IsNone() || !NamedValue.GetTextureName().IsValid())
 	{
-		NamedValue.SetName(AddedDetails.TextureName);
+		NamedValue.SetName(AddedDetails.GetTextureName());
 	}
 
 	TexturesMap.Add(UniqueTexturePtr, AddedDetails);
@@ -176,7 +187,7 @@ void FTextureProfiler::UpdateTextureAllocation(FRHITexture* UniqueTexturePtr, si
 
 	check(ToMB(TotalValue.Size) >= 0);
 
-	FTexureDetails& NamedValue = UpdatedDetails.IsRenderTarget ? CombinedRenderTargetSizes.FindOrAdd(UpdatedDetails.TextureName) : CombinedTextureSizes.FindOrAdd(UpdatedDetails.TextureName);
+	FTexureDetails& NamedValue = UpdatedDetails.IsRenderTarget ? CombinedRenderTargetSizes.FindOrAdd(UpdatedDetails.GetTextureName()) : CombinedTextureSizes.FindOrAdd(UpdatedDetails.GetTextureName());
 	NamedValue -= ExistingTexture;
 	NamedValue += UpdatedDetails;
 
@@ -191,7 +202,7 @@ void FTextureProfiler::RemoveTextureAllocation(FRHITexture* UniqueTexturePtr)
 	FTexureDetails& TotalValue = Details.IsRenderTarget ? TotalRenderTargetSize : TotalTextureSize;
 	TotalValue -= Details;
 
-	FTexureDetails& NamedValue = Details.IsRenderTarget ? CombinedRenderTargetSizes[Details.TextureName] : CombinedTextureSizes[Details.TextureName];
+	FTexureDetails& NamedValue = Details.IsRenderTarget ? CombinedRenderTargetSizes[Details.GetTextureName()] : CombinedTextureSizes[Details.GetTextureName()];
 	NamedValue -= Details;
 }
 
@@ -199,22 +210,37 @@ void FTextureProfiler::UpdateTextureName(FRHITexture* UniqueTexturePtr)
 {
 	FScopeLock Lock(&TextureMapCS);
 	FTexureDetails& Details = TexturesMap[UniqueTexturePtr];
-	FName OldName = Details.TextureName;
-	FName NewName = UniqueTexturePtr->GetName();
-	Details.SetName(NewName);
+	FName OldName = Details.GetTextureName();
+	Details.SetName(UniqueTexturePtr->GetName());
 
-	FTexureDetails& NewCombinedValue = Details.IsRenderTarget ? CombinedRenderTargetSizes.FindOrAdd(NewName) : CombinedTextureSizes.FindOrAdd(NewName);
+	FTexureDetails& NewCombinedValue = Details.IsRenderTarget ? CombinedRenderTargetSizes.FindOrAdd(Details.GetTextureName()) : CombinedTextureSizes.FindOrAdd(Details.GetTextureName());
 	FTexureDetails& OldCombinedValue = Details.IsRenderTarget ? CombinedRenderTargetSizes[OldName] : CombinedTextureSizes[OldName];
 
-	if (NewCombinedValue.TextureName.IsNone() || !NewCombinedValue.TextureName.IsValid())
+	if (NewCombinedValue.GetTextureName().IsNone() || !NewCombinedValue.GetTextureName().IsValid())
 	{
-		NewCombinedValue.SetName(NewName);
+		NewCombinedValue.SetName(Details.GetTextureName());
 	}
 
 	OldCombinedValue -= Details;
 	NewCombinedValue += Details;
 }
 
+const char* FTextureProfiler::GetTextureNameString(FName TextureName)
+{
+	char*& TextureNameString = TextureNameStrings.FindOrAdd(TextureName, nullptr);
+
+	if (TextureNameString == nullptr)
+	{
+		char NewTextureNameString[NAME_SIZE];
+		TextureName.GetPlainANSIString(NewTextureNameString);
+		int Length = FCStringAnsi::Strlen(NewTextureNameString);
+		TextureNameString = new char[Length + 1];
+		FCStringAnsi::Strcpy(TextureNameString, Length + 1, NewTextureNameString);
+	}
+
+	return TextureNameString;
+}
+
 #if CSV_PROFILER
 
 static void ReportTextureStat(const TAutoConsoleVariable<bool>& EnableVar, const char* StatName, uint32 CategoryIndex, int32 Size, ECsvCustomStatOp Op)
@@ -264,8 +290,8 @@ void FTextureProfiler::Update()
 			continue;
 		}
 
-		ReportTextureStat(CVarTextureProfilerEnableRenderTargetCSV, Pair.Value.TextureNameString, CSV_CATEGORY_INDEX(RenderTargetProfiler), ToMB(Pair.Value.PeakSize), ECsvCustomStatOp::Set);
-		ReportTextureStat(CVarTextureProfilerEnableRenderTargetCSV, Pair.Value.TextureNameString, CSV_CATEGORY_INDEX(RenderTargetWasteProfiler), ToMB(Pair.Value.AllocationWaste), ECsvCustomStatOp::Set);
+		ReportTextureStat(CVarTextureProfilerEnableRenderTargetCSV, Pair.Value.GetTextureNameString(), CSV_CATEGORY_INDEX(RenderTargetProfiler), ToMB(Pair.Value.PeakSize), ECsvCustomStatOp::Set);
+		ReportTextureStat(CVarTextureProfilerEnableRenderTargetCSV, Pair.Value.GetTextureNameString(), CSV_CATEGORY_INDEX(RenderTargetWasteProfiler), ToMB(Pair.Value.AllocationWaste), ECsvCustomStatOp::Set);
 
 		Pair.Value.ResetPeakSize();
 	}
@@ -280,8 +306,8 @@ void FTextureProfiler::Update()
 			continue;
 		}
 
-		ReportTextureStat(CVarTextureProfilerEnableTextureCSV, Pair.Value.TextureNameString, CSV_CATEGORY_INDEX(TextureProfiler), ToMB(Pair.Value.PeakSize), ECsvCustomStatOp::Set);
-		ReportTextureStat(CVarTextureProfilerEnableTextureCSV, Pair.Value.TextureNameString, CSV_CATEGORY_INDEX(TextureWasteProfiler), ToMB(Pair.Value.AllocationWaste), ECsvCustomStatOp::Set);
+		ReportTextureStat(CVarTextureProfilerEnableTextureCSV, Pair.Value.GetTextureNameString(), CSV_CATEGORY_INDEX(TextureProfiler), ToMB(Pair.Value.PeakSize), ECsvCustomStatOp::Set);
+		ReportTextureStat(CVarTextureProfilerEnableTextureCSV, Pair.Value.GetTextureNameString(), CSV_CATEGORY_INDEX(TextureWasteProfiler), ToMB(Pair.Value.AllocationWaste), ECsvCustomStatOp::Set);
 
 		Pair.Value.ResetPeakSize();
 	}
@@ -340,7 +366,7 @@ void FTextureProfiler::DumpTextures(bool RenderTargets, bool CombineTextureNames
 				continue;
 			}
 
-			FString NameString = Pair.Value.TextureName.ToString();
+			FString NameString = Pair.Value.GetTextureName().ToString();
 			OutputDevice.Logf(TEXT("%s%c%d%c%d%c%d"), *NameString, Sep, ToMB(Pair.Value.Size), Sep, Pair.Value.Count, Sep, ToMB(Pair.Value.AllocationWaste));
 		}
 
@@ -360,7 +386,7 @@ void FTextureProfiler::DumpTextures(bool RenderTargets, bool CombineTextureNames
 				continue;
 			}
 
-			FString NameString = Pair.Value.TextureName.ToString();
+			FString NameString = Pair.Value.GetTextureName().ToString();
 			OutputDevice.Logf(TEXT("%s%c%d%c%d%c%p"), *NameString, Sep, ToMB(Pair.Value.Size), Sep, ToMB(Pair.Value.AllocationWaste), Sep, Pair.Value.Align);
 		}
 
diff --git a/Engine/Source/Runtime/RHI/Public/TextureProfiler.h b/Engine/Source/Runtime/RHI/Public/TextureProfiler.h
--- a/Engine/Source/Runtime/RHI/Public/TextureProfiler.h
+++ b/Engine/Source/Runtime/RHI/Public/TextureProfiler.h
@@ -32,36 +32,52 @@ public:
 	void UpdateTextureAllocation(FRHITexture* UniqueTexturePtr, size_t Size, uint32 Alignment, size_t AllocationWaste);
 	void RemoveTextureAllocation(FRHITexture* UniqueTexturePtr);
 	void UpdateTextureName(FRHITexture* UniqueTexturePtr);
+
+	const char* GetTextureNameString(FName TextureName);
 private:
 
 	FTextureProfiler() = default;
 	FTextureProfiler(const FTextureProfiler&) = delete;
-	FTextureProfiler(const FTextureProfiler&&) = delete;
+	FTextureProfiler(FTextureProfiler&&) = delete;
 
 	void Update();
 
 	FCriticalSection TextureMapCS;
-	struct FTexureDetails
+
+	class FTexureDetails
 	{
+	public:
 		FTexureDetails() = default;
 		FTexureDetails(FRHITexture* Texture, size_t InSize, uint32 InAlign, size_t InAllocationWaste);
-
+		~FTexureDetails();
+		
 		void SetName(FName InTextureName);
+		const char* GetNameString() const { return TextureNameString; }
+
 		void ResetPeakSize();
+		void SetValues(const FTexureDetails& Values);
+		FName GetTextureName() const { return TextureName; }
+		const char* GetTextureNameString() const { return TextureNameString; }
 
 		FTexureDetails& operator+=(const FTexureDetails& Other);
 		FTexureDetails& operator-=(const FTexureDetails& Other);
+		FTexureDetails& operator=(const FTexureDetails& Other) = default;
 
-		FName TextureName;
-		char TextureNameString[NAME_SIZE];
 		size_t Size = 0;
 		size_t PeakSize = 0;
 		uint32 Align = 0;
 		size_t AllocationWaste = 0;
 		int Count = 0;
 		bool IsRenderTarget = false;
+
+	private:
+		FName TextureName;
+
+		// This memory is not owned, do not delete
+		const char* TextureNameString = nullptr;
 	};
 
+	TMap<FName, char*> TextureNameStrings;
 	TMap<void*, FTexureDetails> TexturesMap;
 
 	// Keep track of the totals separately to reduce the cost of rounding error for sizes
