Temporary workaround for a rare streaming crash: make FRayTracingGeometryInstance.GeometryRHI ref counted so it won't get streamed out.

diff --git a/Engine/Source/Runtime/RHI/Public/DynamicRHI.h b/Engine/Source/Runtime/RHI/Public/DynamicRHI.h
--- a/Engine/Source/Runtime/RHI/Public/DynamicRHI.h
+++ b/Engine/Source/Runtime/RHI/Public/DynamicRHI.h
@@ -96,7 +96,8 @@ ENUM_CLASS_FLAGS(ERayTracingInstanceFlags);
 */
 struct FRayTracingGeometryInstance
 {
-	FRHIRayTracingGeometry* GeometryRHI = nullptr;
+	// TODO: UE-130819 Ref counting is a temporary workaround for a very rare streaming crash.
+	FRayTracingGeometryRHIRef GeometryRHI = nullptr;
 
 	// A single physical mesh may be duplicated many times in the scene with different transforms and user data.
 	// All copies share the same shader binding table entries and therefore will have the same material and shader resources.
