Texture profiler producing far too many csv stats.  The r.TextureProfiler.MinTextureSizeMB variable was mistakenly set to 0 instead of 16.

diff --git a/Engine/Source/Runtime/RHI/Private/TextureProfiler.cpp b/Engine/Source/Runtime/RHI/Private/TextureProfiler.cpp
--- a/Engine/Source/Runtime/RHI/Private/TextureProfiler.cpp
+++ b/Engine/Source/Runtime/RHI/Private/TextureProfiler.cpp
@@ -24,7 +24,7 @@ CSV_DEFINE_CATEGORY(TextureWasteProfiler, true);
 
 static TAutoConsoleVariable<int32> CVarTextureProfilerMinTextureSizeMB(
 	TEXT("r.TextureProfiler.MinTextureSizeMB"),
-	0,
+	16,
 	TEXT("The minimum size for any texture to be reported.  All textures below this threshold will be reported as Other."));
 
 static TAutoConsoleVariable<int32> CVarTextureProfilerMinRenderTargetSizeMB(
