Back out CL: 17859155 - adding 200MB of extra memory because it keeps BLAS longer around than planned. Investigation going on

diff --git a/Engine/Source/Runtime/RHI/Public/DynamicRHI.h b/Engine/Source/Runtime/RHI/Public/DynamicRHI.h
--- a/Engine/Source/Runtime/RHI/Public/DynamicRHI.h
+++ b/Engine/Source/Runtime/RHI/Public/DynamicRHI.h
@@ -96,8 +96,7 @@ ENUM_CLASS_FLAGS(ERayTracingInstanceFlags);
 */
 struct FRayTracingGeometryInstance
 {
-	// TODO: UE-130819 Ref counting is a temporary workaround for a very rare streaming crash.
-	FRayTracingGeometryRHIRef GeometryRHI = nullptr;
+	FRHIRayTracingGeometry* GeometryRHI = nullptr;
 
 	// A single physical mesh may be duplicated many times in the scene with different transforms and user data.
 	// All copies share the same shader binding table entries and therefore will have the same material and shader resources.
