Disable RHI resource tracking in Test builds

diff --git a/Engine/Source/Runtime/RHI/RHI.Build.cs b/Engine/Source/Runtime/RHI/RHI.Build.cs
--- a/Engine/Source/Runtime/RHI/RHI.Build.cs
+++ b/Engine/Source/Runtime/RHI/RHI.Build.cs
@@ -31,7 +31,7 @@ public class RHI : ModuleRules
 					PublicDefinitions.Add("RHI_WANT_BREADCRUMB_EVENTS=1");
 				}
 
-				if (Target.Configuration != UnrealTargetConfiguration.Shipping)
+				if (Target.Configuration != UnrealTargetConfiguration.Shipping && Target.Configuration != UnrealTargetConfiguration.Test)
                 {
 					PublicDefinitions.Add("RHI_WANT_RESOURCE_INFO=1");
                 }
