Fix for Linux Server build

diff --git a/Engine/Source/Runtime/RHI/Private/TextureProfiler.cpp b/Engine/Source/Runtime/RHI/Private/TextureProfiler.cpp
--- a/Engine/Source/Runtime/RHI/Private/TextureProfiler.cpp
+++ b/Engine/Source/Runtime/RHI/Private/TextureProfiler.cpp
@@ -241,7 +241,7 @@ const char* FTextureProfiler::GetTextureNameString(FName TextureName)
 	return TextureNameString;
 }
 
-#if CSV_PROFILER
+#if CSVPROFILERTRACE_ENABLED
 
 static void ReportTextureStat(const TAutoConsoleVariable<bool>& EnableVar, const char* StatName, uint32 CategoryIndex, int32 Size, ECsvCustomStatOp Op)
 {
@@ -252,7 +252,7 @@ static void ReportTextureStat(const TAutoConsoleVariable<bool>& EnableVar, const
 
 void FTextureProfiler::Update()
 {
-#if CSV_PROFILER
+#if CSVPROFILERTRACE_ENABLED
 
 	check(IsInRenderingThread());
 
