Removed extra bSupportsInlineRayTracing that appeared during merge.

#rnx

diff --git a/Engine/Source/Runtime/RHI/Public/RHIDefinitions.h b/Engine/Source/Runtime/RHI/Public/RHIDefinitions.h
--- a/Engine/Source/Runtime/RHI/Public/RHIDefinitions.h
+++ b/Engine/Source/Runtime/RHI/Public/RHIDefinitions.h
@@ -378,7 +378,6 @@ class RHI_API FGenericDataDrivenShaderPlatformInfo
 	uint32 bOverrideFMaterial_NeedsGBufferEnabled : 1;
 	uint32 bSupportsMobileDistanceField : 1;
 	uint32 bSupportsFFTBloom : 1;
-	uint32 bSupportsVertexShaderLayer : 1;
 	uint32 bSupportsInlineRayTracing : 1;
 	uint32 bSupportsRayTracingShaders : 1;
 	uint32 bSupportsVertexShaderLayer : 1;
