Moving RHIIsTypedUAVLoadSupported and RHIIsTypedUAVStoreSupported out of FDynamicRHI. New versions now use TextureFormat Capability flags.

diff --git a/Engine/Plugins/FX/Niagara/Source/Niagara/Private/NiagaraDataInterfaceRenderTarget2D.cpp b/Engine/Plugins/FX/Niagara/Source/Niagara/Private/NiagaraDataInterfaceRenderTarget2D.cpp
--- a/Engine/Plugins/FX/Niagara/Source/Niagara/Private/NiagaraDataInterfaceRenderTarget2D.cpp
+++ b/Engine/Plugins/FX/Niagara/Source/Niagara/Private/NiagaraDataInterfaceRenderTarget2D.cpp
@@ -112,7 +112,7 @@ bool GetRenderTargetFormat(bool bOverrideFormat, ETextureRenderTargetFormat Over
 	}
 
 	// If the format does not support typed store we need to find one that will
-	while ( GDynamicRHI->RHIIsTypedUAVStoreSupported(GetPixelFormatFromRenderTargetFormat(OutRenderTargetFormat)) == false )
+	while ( RHIIsTypedUAVStoreSupported(GetPixelFormatFromRenderTargetFormat(OutRenderTargetFormat)) == false )
 	{
 		static const TPair<ETextureRenderTargetFormat, ETextureRenderTargetFormat> FormatRemapTable[] =
 		{
diff --git a/Engine/Source/Runtime/Apple/AGXRHI/Private/AGXDynamicRHI.h b/Engine/Source/Runtime/Apple/AGXRHI/Private/AGXDynamicRHI.h
--- a/Engine/Source/Runtime/Apple/AGXRHI/Private/AGXDynamicRHI.h
+++ b/Engine/Source/Runtime/Apple/AGXRHI/Private/AGXDynamicRHI.h
@@ -185,8 +185,6 @@ public:
 
 	virtual void RHIUpdateUniformBuffer(FRHIUniformBuffer* UniformBufferRHI, const void* Contents) final override;
 
-	virtual bool RHIIsTypedUAVLoadSupported(EPixelFormat PixelFormat) override;
-
 	virtual uint16 RHIGetPlatformTextureMaxSampleCount() override;
 
 private:
diff --git a/Engine/Source/Runtime/Apple/AGXRHI/Private/AGXRHI.cpp b/Engine/Source/Runtime/Apple/AGXRHI/Private/AGXRHI.cpp
--- a/Engine/Source/Runtime/Apple/AGXRHI/Private/AGXRHI.cpp
+++ b/Engine/Source/Runtime/Apple/AGXRHI/Private/AGXRHI.cpp
@@ -1240,11 +1240,6 @@ void* FAGXDynamicRHI::RHIGetNativeInstance()
 	return nullptr;
 }
 
-bool FAGXDynamicRHI::RHIIsTypedUAVLoadSupported(EPixelFormat PixelFormat)
-{
-	return EnumHasAnyFlags(GPixelFormats[PixelFormat].Capabilities, EPixelFormatCapabilities::TypedUAVLoad);
-}
-
 uint16 FAGXDynamicRHI::RHIGetPlatformTextureMaxSampleCount()
 {
 	TArray<ECompositingSampleCount::Type> SamplesArray{ ECompositingSampleCount::Type::One, ECompositingSampleCount::Type::Two, ECompositingSampleCount::Type::Four, ECompositingSampleCount::Type::Eight };
diff --git a/Engine/Source/Runtime/Apple/MetalRHI/Private/MetalDynamicRHI.h b/Engine/Source/Runtime/Apple/MetalRHI/Private/MetalDynamicRHI.h
--- a/Engine/Source/Runtime/Apple/MetalRHI/Private/MetalDynamicRHI.h
+++ b/Engine/Source/Runtime/Apple/MetalRHI/Private/MetalDynamicRHI.h
@@ -188,8 +188,6 @@ public:
 
 	virtual void RHIUpdateUniformBuffer(FRHIUniformBuffer* UniformBufferRHI, const void* Contents) final override;
 
-	virtual bool RHIIsTypedUAVLoadSupported(EPixelFormat PixelFormat) override;
-
 	virtual uint16 RHIGetPlatformTextureMaxSampleCount() override;
 
 private:
diff --git a/Engine/Source/Runtime/Apple/MetalRHI/Private/MetalRHI.cpp b/Engine/Source/Runtime/Apple/MetalRHI/Private/MetalRHI.cpp
--- a/Engine/Source/Runtime/Apple/MetalRHI/Private/MetalRHI.cpp
+++ b/Engine/Source/Runtime/Apple/MetalRHI/Private/MetalRHI.cpp
@@ -1293,11 +1293,6 @@ void* FMetalDynamicRHI::RHIGetNativeInstance()
 	return nullptr;
 }
 
-bool FMetalDynamicRHI::RHIIsTypedUAVLoadSupported(EPixelFormat PixelFormat)
-{
-	return EnumHasAnyFlags(GPixelFormats[PixelFormat].Capabilities, EPixelFormatCapabilities::TypedUAVLoad);
-}
-
 uint16 FMetalDynamicRHI::RHIGetPlatformTextureMaxSampleCount()
 {
 	TArray<ECompositingSampleCount::Type> SamplesArray{ ECompositingSampleCount::Type::One, ECompositingSampleCount::Type::Two, ECompositingSampleCount::Type::Four, ECompositingSampleCount::Type::Eight };
diff --git a/Engine/Source/Runtime/D3D12RHI/Private/D3D12Device.cpp b/Engine/Source/Runtime/D3D12RHI/Private/D3D12Device.cpp
--- a/Engine/Source/Runtime/D3D12RHI/Private/D3D12Device.cpp
+++ b/Engine/Source/Runtime/D3D12RHI/Private/D3D12Device.cpp
@@ -592,12 +592,3 @@ D3D12_RESOURCE_ALLOCATION_INFO FD3D12Device::GetResourceAllocationInfo(const D3D
 	}
 }
 
-bool FD3D12DynamicRHI::RHIIsTypedUAVLoadSupported(EPixelFormat PixelFormat)
-{
-	return EnumHasAnyFlags(GPixelFormats[PixelFormat].Capabilities, EPixelFormatCapabilities::TypedUAVLoad);
-}
-
-bool FD3D12DynamicRHI::RHIIsTypedUAVStoreSupported(EPixelFormat PixelFormat)
-{
-	return EnumHasAnyFlags(GPixelFormats[PixelFormat].Capabilities, EPixelFormatCapabilities::TypedUAVStore);
-}
diff --git a/Engine/Source/Runtime/D3D12RHI/Private/D3D12RHIPrivate.h b/Engine/Source/Runtime/D3D12RHI/Private/D3D12RHIPrivate.h
--- a/Engine/Source/Runtime/D3D12RHI/Private/D3D12RHIPrivate.h
+++ b/Engine/Source/Runtime/D3D12RHI/Private/D3D12RHIPrivate.h
@@ -399,9 +399,6 @@ public:
 
 	virtual IRHITransientResourceAllocator* RHICreateTransientResourceAllocator() override;
 
-	virtual bool RHIIsTypedUAVLoadSupported(EPixelFormat PixelFormat) override;
-	virtual bool RHIIsTypedUAVStoreSupported(EPixelFormat PixelFormat) override;
-
 #if WITH_MGPU
 	virtual IRHICommandContextContainer* RHIGetCommandContextContainer(int32 Index, int32 Num, FRHIGPUMask GPUMask)final override;
 #endif
diff --git a/Engine/Source/Runtime/MediaAssets/Private/Misc/MediaTextureResource.cpp b/Engine/Source/Runtime/MediaAssets/Private/Misc/MediaTextureResource.cpp
--- a/Engine/Source/Runtime/MediaAssets/Private/Misc/MediaTextureResource.cpp
+++ b/Engine/Source/Runtime/MediaAssets/Private/Misc/MediaTextureResource.cpp
@@ -200,7 +200,7 @@ namespace MediaTextureResourceHelpers
 
 	bool SupportsComputeMipGen(EPixelFormat InFormat)
 	{
-		return RHIRequiresComputeGenerateMips() && GDynamicRHI->RHIIsTypedUAVLoadSupported(InFormat);
+		return RHIRequiresComputeGenerateMips() && RHIIsTypedUAVLoadSupported(InFormat);
 	}
 
 	EPixelFormat GetConvertedPixelFormat(const TSharedPtr<IMediaTextureSample, ESPMode::ThreadSafe>& Sample)
diff --git a/Engine/Source/Runtime/RHI/Public/DynamicRHI.h b/Engine/Source/Runtime/RHI/Public/DynamicRHI.h
--- a/Engine/Source/Runtime/RHI/Public/DynamicRHI.h
+++ b/Engine/Source/Runtime/RHI/Public/DynamicRHI.h
@@ -1474,8 +1474,8 @@ public:
 	virtual void RHICalibrateTimers() {}
 	virtual void RHIPollRenderQueryResults() {}
 
-	virtual bool RHIIsTypedUAVLoadSupported(EPixelFormat PixelFormat) { return true; }
-	virtual bool RHIIsTypedUAVStoreSupported(EPixelFormat PixelFormat) { return true; }
+	UE_DEPRECATED(5.0, "The global version of RHIIsTypedUAVLoadSupported should be used")
+	virtual bool RHIIsTypedUAVLoadSupported(EPixelFormat PixelFormat) { return ::RHIIsTypedUAVLoadSupported(PixelFormat); }
 
 	virtual uint16 RHIGetPlatformTextureMaxSampleCount() { return 8; };
 
diff --git a/Engine/Source/Runtime/RHI/Public/RHI.h b/Engine/Source/Runtime/RHI/Public/RHI.h
--- a/Engine/Source/Runtime/RHI/Public/RHI.h
+++ b/Engine/Source/Runtime/RHI/Public/RHI.h
@@ -842,11 +842,21 @@ extern RHI_API FPixelFormatInfo GPixelFormats[PF_MAX];		// Maps members of EPixe
 /** Initialize the 'best guess' pixel format capabilities. Platform formats and support must be filled out before calling this. */
 extern RHI_API void RHIInitDefaultPixelFormatCapabilities();
 
-inline bool PixelFormatHasCapabilities(EPixelFormat InFormat, EPixelFormatCapabilities InCapabilities)
+inline bool RHIPixelFormatHasCapabilities(EPixelFormat InFormat, EPixelFormatCapabilities InCapabilities)
 {
 	return EnumHasAllFlags(GPixelFormats[InFormat].Capabilities, InCapabilities);
 }
 
+inline bool RHIIsTypedUAVLoadSupported(EPixelFormat InFormat)
+{
+	return RHIPixelFormatHasCapabilities(InFormat, EPixelFormatCapabilities::TypedUAVLoad);
+}
+
+inline bool RHIIsTypedUAVStoreSupported(EPixelFormat InFormat)
+{
+	return RHIPixelFormatHasCapabilities(InFormat, EPixelFormatCapabilities::TypedUAVStore);
+}
+
 //
 //	CalculateImageBytes
 //
diff --git a/Engine/Source/Runtime/RHI/Public/RHIValidation.h b/Engine/Source/Runtime/RHI/Public/RHIValidation.h
--- a/Engine/Source/Runtime/RHI/Public/RHIValidation.h
+++ b/Engine/Source/Runtime/RHI/Public/RHIValidation.h
@@ -1727,11 +1727,6 @@ public:
 		RHI->RHIPollRenderQueryResults();
 	}
 
-	virtual bool RHIIsTypedUAVLoadSupported(EPixelFormat PixelFormat) override final
-	{
-		return RHI->RHIIsTypedUAVLoadSupported(PixelFormat);
-	}
-
 	virtual uint16 RHIGetPlatformTextureMaxSampleCount() override final
 	{
 		return RHI->RHIGetPlatformTextureMaxSampleCount();
diff --git a/Engine/Source/Runtime/RenderCore/Private/GenerateMips.cpp b/Engine/Source/Runtime/RenderCore/Private/GenerateMips.cpp
--- a/Engine/Source/Runtime/RenderCore/Private/GenerateMips.cpp
+++ b/Engine/Source/Runtime/RenderCore/Private/GenerateMips.cpp
@@ -304,7 +304,7 @@ END_SHADER_PARAMETER_STRUCT()
 
 bool FGenerateMips::WillFormatSupportCompute(EPixelFormat InPixelFormat)
 {
-	return RHIRequiresComputeGenerateMips() && GDynamicRHI->RHIIsTypedUAVStoreSupported(InPixelFormat);
+	return RHIRequiresComputeGenerateMips() && RHIIsTypedUAVStoreSupported(InPixelFormat);
 }
 
 void FGenerateMips::Execute(FRDGBuilder& GraphBuilder, FRDGTextureRef Texture, FGenerateMipsParams Params, EGenerateMipsPass Pass)
diff --git a/Engine/Source/Runtime/Renderer/Private/PhysicsFieldRendering.cpp b/Engine/Source/Runtime/Renderer/Private/PhysicsFieldRendering.cpp
--- a/Engine/Source/Runtime/Renderer/Private/PhysicsFieldRendering.cpp
+++ b/Engine/Source/Runtime/Renderer/Private/PhysicsFieldRendering.cpp
@@ -183,7 +183,7 @@ void RenderPhysicsField(
 	{
 		FPhysicsFieldResource* PhysicsFieldResource = (GPhysicsFieldSystemType == 0) ? PhysicsFieldProxy->DebugResource : PhysicsFieldProxy->FieldResource;
 
-		if (Views.Num() > 0 && PhysicsFieldResource && ShaderDrawDebug::IsEnabled() && ShaderPrint::IsEnabled(Views[0]) && GDynamicRHI->RHIIsTypedUAVLoadSupported(PF_FloatRGBA))
+		if (Views.Num() > 0 && PhysicsFieldResource && ShaderDrawDebug::IsEnabled() && ShaderPrint::IsEnabled(Views[0]) && RHIIsTypedUAVLoadSupported(PF_FloatRGBA))
 		{
 			AddPhysicsFieldRayMarchingPass(GraphBuilder, Views[0], PhysicsFieldResource, SceneColorTexture);
 		}
diff --git a/Engine/Source/Runtime/Renderer/Private/PostProcess/PostProcessSubsurface.cpp b/Engine/Source/Runtime/Renderer/Private/PostProcess/PostProcessSubsurface.cpp
--- a/Engine/Source/Runtime/Renderer/Private/PostProcess/PostProcessSubsurface.cpp
+++ b/Engine/Source/Runtime/Renderer/Private/PostProcess/PostProcessSubsurface.cpp
@@ -283,8 +283,7 @@ bool IsSubsurfaceRequiredForView(const FViewInfo& View)
 
 bool IsProfileIdCacheEnabled()
 {
-	const bool bIsR8UintSupported = GDynamicRHI && GDynamicRHI->RHIIsTypedUAVLoadSupported(PF_R8_UINT);
-	return bIsR8UintSupported && CVarSSSBurleyEnableProfileIdCache.GetValueOnRenderThread() != 0;
+	return RHIIsTypedUAVLoadSupported(PF_R8_UINT) && CVarSSSBurleyEnableProfileIdCache.GetValueOnRenderThread() != 0;
 }
 
 uint32 GetSubsurfaceRequiredViewMask(TArrayView<const FViewInfo> Views)
diff --git a/Engine/Source/Runtime/Renderer/Private/VolumetricCloudRendering.cpp b/Engine/Source/Runtime/Renderer/Private/VolumetricCloudRendering.cpp
--- a/Engine/Source/Runtime/Renderer/Private/VolumetricCloudRendering.cpp
+++ b/Engine/Source/Runtime/Renderer/Private/VolumetricCloudRendering.cpp
@@ -231,7 +231,7 @@ static bool ShouldUseComputeForCloudTracing(const FStaticFeatureLevel InFeatureL
 	// So in this case, when capturing such a sky light, we must disabled the compute path.
 	const bool bMSAAEnabled = GetDefaultMSAACount(InFeatureLevel) > 1;
 
-	return !CVarVolumetricCloudDisableCompute.GetValueOnRenderThread() && GDynamicRHI->RHIIsTypedUAVLoadSupported(PF_FloatRGBA) && GDynamicRHI->RHIIsTypedUAVLoadSupported(PF_G16R16F) && !bMSAAEnabled;
+	return !CVarVolumetricCloudDisableCompute.GetValueOnRenderThread() && RHIIsTypedUAVLoadSupported(PF_FloatRGBA) && RHIIsTypedUAVLoadSupported(PF_G16R16F) && !bMSAAEnabled;
 }
 
 bool ShouldRenderVolumetricCloud(const FScene* Scene, const FEngineShowFlags& EngineShowFlags)
diff --git a/Engine/Source/Runtime/VulkanRHI/Private/VulkanRHI.cpp b/Engine/Source/Runtime/VulkanRHI/Private/VulkanRHI.cpp
--- a/Engine/Source/Runtime/VulkanRHI/Private/VulkanRHI.cpp
+++ b/Engine/Source/Runtime/VulkanRHI/Private/VulkanRHI.cpp
@@ -1788,16 +1788,4 @@ void FVulkanDynamicRHI::VulkanSetImageLayout( VkCommandBuffer CmdBuffer, VkImage
 	::VulkanSetImageLayout( CmdBuffer, Image, OldLayout, NewLayout, SubresourceRange );
 }
 
-bool FVulkanDynamicRHI::RHIIsTypedUAVLoadSupported(EPixelFormat PixelFormat)
-{
-	if (Device)
-	{
-		const FPixelFormatInfo& FormatInfo = GPixelFormats[PixelFormat];
-		const VkFormat Format = (VkFormat)FormatInfo.PlatformFormat;
-		const VkFormatFeatureFlags FormatFlags = Device->GetFormatProperties()[Format].optimalTilingFeatures;
-		return (FormatFlags & VK_FORMAT_FEATURE_STORAGE_IMAGE_BIT) != 0;
-	}
-	return false;
-}
-
 #undef LOCTEXT_NAMESPACE
diff --git a/Engine/Source/Runtime/VulkanRHI/Public/VulkanDynamicRHI.h b/Engine/Source/Runtime/VulkanRHI/Public/VulkanDynamicRHI.h
--- a/Engine/Source/Runtime/VulkanRHI/Public/VulkanDynamicRHI.h
+++ b/Engine/Source/Runtime/VulkanRHI/Public/VulkanDynamicRHI.h
@@ -375,8 +375,6 @@ public:
 
 	bool RHIRequiresComputeGenerateMips() const override { return true; };
 
-	virtual bool RHIIsTypedUAVLoadSupported(EPixelFormat PixelFormat) final override;
-
 	inline TArray<FVulkanViewport*>& GetViewports()
 	{
 		return Viewports;
diff --git a/Engine/Source/Runtime/Windows/D3D11RHI/Private/D3D11Device.cpp b/Engine/Source/Runtime/Windows/D3D11RHI/Private/D3D11Device.cpp
--- a/Engine/Source/Runtime/Windows/D3D11RHI/Private/D3D11Device.cpp
+++ b/Engine/Source/Runtime/Windows/D3D11RHI/Private/D3D11Device.cpp
@@ -538,16 +538,6 @@ void FD3D11DynamicRHI::SetupAfterDeviceCreation()
 	}
 }
 
-bool FD3D11DynamicRHI::RHIIsTypedUAVLoadSupported(EPixelFormat PixelFormat)
-{
-	return EnumHasAllFlags(GPixelFormats[PixelFormat].Capabilities, EPixelFormatCapabilities::TypedUAVLoad);
-}
-
-bool FD3D11DynamicRHI::RHIIsTypedUAVStoreSupported(EPixelFormat PixelFormat)
-{
-	return EnumHasAllFlags(GPixelFormats[PixelFormat].Capabilities, EPixelFormatCapabilities::TypedUAVStore);
-}
-
 void FD3D11DynamicRHI::UpdateMSAASettings()
 {	
 	check(DX_MAX_MSAA_COUNT == 8);
diff --git a/Engine/Source/Runtime/Windows/D3D11RHI/Private/D3D11RHIPrivate.h b/Engine/Source/Runtime/Windows/D3D11RHI/Private/D3D11RHIPrivate.h
--- a/Engine/Source/Runtime/Windows/D3D11RHI/Private/D3D11RHIPrivate.h
+++ b/Engine/Source/Runtime/Windows/D3D11RHI/Private/D3D11RHIPrivate.h
@@ -682,9 +682,6 @@ public:
 
 	virtual void RHICalibrateTimers() override;
 
-	virtual bool RHIIsTypedUAVLoadSupported(EPixelFormat PixelFormat) override;
-	virtual bool RHIIsTypedUAVStoreSupported(EPixelFormat PixelFormat) override;
-
 	// Accessors.
 	ID3D11Device* GetDevice() const
 	{
