[RHI] Fix for TextureProfiler to properly process textures renaming

diff --git a/Engine/Source/Runtime/RHI/Private/TextureProfiler.cpp b/Engine/Source/Runtime/RHI/Private/TextureProfiler.cpp
--- a/Engine/Source/Runtime/RHI/Private/TextureProfiler.cpp
+++ b/Engine/Source/Runtime/RHI/Private/TextureProfiler.cpp
@@ -209,7 +209,14 @@ void FTextureProfiler::RemoveTextureAllocation(FRHITexture* UniqueTexturePtr)
 void FTextureProfiler::UpdateTextureName(FRHITexture* UniqueTexturePtr)
 {
 	FScopeLock Lock(&TextureMapCS);
-	FTexureDetails& Details = TexturesMap[UniqueTexturePtr];
+
+	FTexureDetails* DetailsPtr = TexturesMap.Find(UniqueTexturePtr);
+	if (DetailsPtr == nullptr)
+	{
+		return;
+	}
+
+	FTexureDetails& Details = *DetailsPtr;
 	FName OldName = Details.GetTextureName();
 	Details.SetName(UniqueTexturePtr->GetName());
 
diff --git a/Engine/Source/Runtime/RHI/Public/RHIResources.h b/Engine/Source/Runtime/RHI/Public/RHIResources.h
--- a/Engine/Source/Runtime/RHI/Public/RHIResources.h
+++ b/Engine/Source/Runtime/RHI/Public/RHIResources.h
@@ -17,6 +17,7 @@
 #include "RHIDefinitions.h"
 #include "Templates/RefCounting.h"
 #include "PixelFormat.h"
+#include "TextureProfiler.h"
 #include "Containers/LockFreeList.h"
 #include "Misc/SecureHash.h"
 #include "Hash/CityHash.h"
@@ -1358,6 +1359,10 @@ public:
 	void SetName(const FName& InName)
 	{
 		TextureName = InName;
+
+#if TEXTURE_PROFILER_ENABLED
+		FTextureProfiler::Get()->UpdateTextureName(this);
+#endif
 	}
 
 	FName GetName() const
