Fix a race condition that can lead to double-freeing FRHIUniformBuffers. Double-freeing happens when a FRHIResource is revived right after passing the Deleting check in FRHIResource::FlushPendingDeletes but before it is actually deleted.

#tests soaked in the crash-reproing creative hub for hours multiple times without crashing

diff --git a/Engine/Source/Runtime/RHI/Public/RHICommandList.h b/Engine/Source/Runtime/RHI/Public/RHICommandList.h
--- a/Engine/Source/Runtime/RHI/Public/RHICommandList.h
+++ b/Engine/Source/Runtime/RHI/Public/RHICommandList.h
@@ -930,7 +930,7 @@ struct FRHICommandSetShaderUniformBuffer final : public FRHICommand<FRHICommandS
 {
 	TRHIShader* Shader;
 	uint32 BaseIndex;
-	FRHIUniformBuffer* UniformBuffer;
+	TRefCountPtr<FRHIUniformBuffer> UniformBuffer;
 	FORCEINLINE_DEBUGGABLE FRHICommandSetShaderUniformBuffer(TRHIShader* InShader, uint32 InBaseIndex, FRHIUniformBuffer* InUniformBuffer)
 		: Shader(InShader)
 		, BaseIndex(InBaseIndex)
@@ -1714,7 +1714,29 @@ FRHICOMMAND_MACRO(FRHICommandSetStaticUniformBuffers)
 
 	FORCEINLINE_DEBUGGABLE FRHICommandSetStaticUniformBuffers(const FUniformBufferStaticBindings& InUniformBuffers)
 		: UniformBuffers(InUniformBuffers)
-	{}
+	{
+		const int32 NumUniformBuffers = UniformBuffers.GetUniformBufferCount();
+		for (int32 Index = 0; Index < NumUniformBuffers; ++Index)
+		{
+			FRHIUniformBuffer* UniformBuffer = UniformBuffers.GetUniformBuffer(Index);
+			if (UniformBuffer)
+			{
+				UniformBuffer->AddRef();
+			}
+		}
+	}
+	~FRHICommandSetStaticUniformBuffers()
+	{
+		const int32 NumUniformBuffers = UniformBuffers.GetUniformBufferCount();
+		for (int32 Index = 0; Index < NumUniformBuffers; ++Index)
+		{
+			FRHIUniformBuffer* UniformBuffer = UniformBuffers.GetUniformBuffer(Index);
+			if (UniformBuffer)
+			{
+				UniformBuffer->Release();
+			}
+		}
+	}
 	RHI_API void Execute(FRHICommandListBase& CmdList);
 };
 
