Add Pixel Format: PF_R64_UINT.  Needed for (compliant) Vulkan's 64bit image atomics in upcoming changelist.  bSupported defaults to FALSE (so no additional code for most APIs).

diff --git a/Engine/Source/Runtime/Apple/MetalRHI/Private/MetalRHI.cpp b/Engine/Source/Runtime/Apple/MetalRHI/Private/MetalRHI.cpp
--- a/Engine/Source/Runtime/Apple/MetalRHI/Private/MetalRHI.cpp
+++ b/Engine/Source/Runtime/Apple/MetalRHI/Private/MetalRHI.cpp
@@ -707,6 +707,7 @@ FMetalDynamicRHI::FMetalDynamicRHI(ERHIFeatureLevel::Type RequestedFeatureLevel)
 	GMetalBufferFormats[PF_R32G32B32_SINT		] = { mtlpp::PixelFormat::Invalid, (uint8)EMetalBufferFormat::Unknown };
 	GMetalBufferFormats[PF_R32G32B32F			] = { mtlpp::PixelFormat::Invalid, (uint8)EMetalBufferFormat::Unknown };
 	GMetalBufferFormats[PF_R8_SINT				] = { mtlpp::PixelFormat::R8Sint, (uint8)EMetalBufferFormat::R8Sint };
+	GMetalBufferFormats[PF_R64_UINT				] = { mtlpp::PixelFormat::Invalid, (uint8)EMetalBufferFormat::Unknown };
 
 	// Initialize the platform pixel format map.
 	GPixelFormats[PF_Unknown			].PlatformFormat	= (uint32)mtlpp::PixelFormat::Invalid;
@@ -932,6 +933,8 @@ FMetalDynamicRHI::FMetalDynamicRHI(ERHIFeatureLevel::Type RequestedFeatureLevel)
 	GPixelFormats[PF_R32G32B32F			].PlatformFormat	= (uint32)mtlpp::PixelFormat::Invalid;
 	GPixelFormats[PF_R32G32B32F			].Supported			= false;
 	GPixelFormats[PF_R8_SINT			].PlatformFormat	= (uint32)mtlpp::PixelFormat::R8Sint;
+	GPixelFormats[PF_R64_UINT			].PlatformFormat	= (uint32)mtlpp::PixelFormat::Invalid;
+	GPixelFormats[PF_R64_UINT			].Supported			= false;
 
 #if METAL_DEBUG_OPTIONS
 	for (uint32 i = 0; i < PF_MAX; i++)
diff --git a/Engine/Source/Runtime/Core/Public/PixelFormat.h b/Engine/Source/Runtime/Core/Public/PixelFormat.h
--- a/Engine/Source/Runtime/Core/Public/PixelFormat.h
+++ b/Engine/Source/Runtime/Core/Public/PixelFormat.h
@@ -96,7 +96,8 @@ enum EPixelFormat
 	PF_R32G32B32_SINT		=81,
 	PF_R32G32B32F			=82,
 	PF_R8_SINT				=83,	
-	PF_MAX					=84,
+	PF_R64_UINT				=84,
+	PF_MAX					=85,
 };
 #define FOREACH_ENUM_EPIXELFORMAT(op) \
 	op(PF_Unknown) \
@@ -182,7 +183,9 @@ enum EPixelFormat
 	op(PF_R32G32B32_UINT) \
 	op(PF_R32G32B32_SINT) \
 	op(PF_R32G32B32F) \
-	op(PF_R8_SINT)
+	op(PF_R8_SINT) \
+	op(PF_R64_UINT)
+
 // Defines which channel is valid for each pixel format
 enum class EPixelFormatChannelFlags : uint8
 {
diff --git a/Engine/Source/Runtime/CoreUObject/Public/UObject/NoExportTypes.h b/Engine/Source/Runtime/CoreUObject/Public/UObject/NoExportTypes.h
--- a/Engine/Source/Runtime/CoreUObject/Public/UObject/NoExportTypes.h
+++ b/Engine/Source/Runtime/CoreUObject/Public/UObject/NoExportTypes.h
@@ -280,6 +280,7 @@ enum EPixelFormat
 	PF_R32G32B32_SINT,
 	PF_R32G32B32F,
 	PF_R8_SINT,
+	PF_R64_UINT,
 	PF_MAX,
 };
 
diff --git a/Engine/Source/Runtime/RHI/Private/RHI.cpp b/Engine/Source/Runtime/RHI/Private/RHI.cpp
--- a/Engine/Source/Runtime/RHI/Private/RHI.cpp
+++ b/Engine/Source/Runtime/RHI/Private/RHI.cpp
@@ -2520,12 +2520,13 @@ FPixelFormatInfo    GPixelFormats[PF_MAX] =
 	FPixelFormatInfo(PF_ASTC_10x10_HDR,         TEXT("ASTC_10x10_HDR"),    10,         10,         1,          16,         4,              0),
 	FPixelFormatInfo(PF_ASTC_12x12_HDR,         TEXT("ASTC_12x12_HDR"),    12,         12,         1,          16,         4,              0),
 
-	FPixelFormatInfo(PF_G16R16_SNORM,       TEXT("G16R16_SNORM"),       1,          1,          1,          4,          2,              1),
-	FPixelFormatInfo(PF_R8G8_UINT,          TEXT("R8G8_UINT"),          1,          1,          1,          2,          2,              1),
-	FPixelFormatInfo(PF_R32G32B32_UINT,     TEXT("R32G32B32_UINT"),     1,          1,          1,          12,         3,              1),
-	FPixelFormatInfo(PF_R32G32B32_SINT,     TEXT("R32G32B32_SINT"),     1,          1,          1,          12,         3,              1),
-	FPixelFormatInfo(PF_R32G32B32F,         TEXT("R32G32B32F"),         1,          1,          1,          12,         3,              1),
-	FPixelFormatInfo(PF_R8_SINT,            TEXT("R8_SINT"),            1,          1,          1,          1,          1,              1),
+	FPixelFormatInfo(PF_G16R16_SNORM,       TEXT("G16R16_SNORM"),          1,          1,          1,          4,          2,              1),
+	FPixelFormatInfo(PF_R8G8_UINT,          TEXT("R8G8_UINT"),             1,          1,          1,          2,          2,              1),
+	FPixelFormatInfo(PF_R32G32B32_UINT,     TEXT("R32G32B32_UINT"),        1,          1,          1,          12,         3,              1),
+	FPixelFormatInfo(PF_R32G32B32_SINT,     TEXT("R32G32B32_SINT"),        1,          1,          1,          12,         3,              1),
+	FPixelFormatInfo(PF_R32G32B32F,         TEXT("R32G32B32F"),            1,          1,          1,          12,         3,              1),
+	FPixelFormatInfo(PF_R8_SINT,            TEXT("R8_SINT"),               1,          1,          1,          1,          1,              1),
+	FPixelFormatInfo(PF_R64_UINT,			TEXT("R64_UINT"),              1,          1,          1,          8,          1,              0),
 };
 
 void RHIInitDefaultPixelFormatCapabilities()
diff --git a/Engine/Source/Runtime/RenderCore/Private/RenderUtils.cpp b/Engine/Source/Runtime/RenderCore/Private/RenderUtils.cpp
--- a/Engine/Source/Runtime/RenderCore/Private/RenderUtils.cpp
+++ b/Engine/Source/Runtime/RenderCore/Private/RenderUtils.cpp
@@ -960,6 +960,7 @@ EPixelFormatChannelFlags GetPixelFormatValidChannels(EPixelFormat InPixelFormat)
 		EPixelFormatChannelFlags::RGB,		// PF_R32G32B32_SINT
 		EPixelFormatChannelFlags::RGB,		// PF_R32G32B32F
 		EPixelFormatChannelFlags::R,		// PF_R8_SINT
+		EPixelFormatChannelFlags::R,		// PF_R64_UINT
 	};
 	static_assert(UE_ARRAY_COUNT(PixelFormatToChannelFlags) == (uint8)PF_MAX, "Missing pixel format");
 	return (InPixelFormat < PF_MAX) ? PixelFormatToChannelFlags[(uint8)InPixelFormat] : EPixelFormatChannelFlags::None;
diff --git a/Engine/Source/Runtime/Renderer/Private/SystemTextures.cpp b/Engine/Source/Runtime/Renderer/Private/SystemTextures.cpp
--- a/Engine/Source/Runtime/Renderer/Private/SystemTextures.cpp
+++ b/Engine/Source/Runtime/Renderer/Private/SystemTextures.cpp
@@ -968,6 +968,8 @@ TOutType ConvertInputFormat(const TInType& In)
 {
 	return TOutType(In);
 }
+template<> uint64 ConvertInputFormat<float, uint64, EDefaultInputType::UNorm>(const float& In) { return FMath::Clamp(In,  0.f, 1.f) * float(MAX_uint64); }
+template<>  int64 ConvertInputFormat<float,  int64, EDefaultInputType::SNorm>(const float& In) { return FMath::Clamp(In, -1.f, 1.f) * float(MAX_int64); }
 template<> uint32 ConvertInputFormat<float, uint32, EDefaultInputType::UNorm>(const float& In) { return FMath::Clamp(In,  0.f, 1.f) * float(MAX_uint32); }
 template<>  int32 ConvertInputFormat<float,  int32, EDefaultInputType::SNorm>(const float& In) { return FMath::Clamp(In, -1.f, 1.f) * float(MAX_int32); }
 template<> uint16 ConvertInputFormat<float, uint16, EDefaultInputType::UNorm>(const float& In) { return FMath::Clamp(In,  0.f, 1.f) * MAX_uint16; }
@@ -1062,10 +1064,13 @@ template<typename TInType>
 void InitializeData(const TInType& InData, EPixelFormat InFormat, uint8* OutData, uint32& OutByteCount)
 {
 	// If a new format is added insure that it is either supported here, or at least flagged as not supported
-	static_assert(PF_MAX == 84);
+	static_assert(PF_MAX == 85);
 
 	switch (InFormat)
 	{
+		// 64bits
+		case PF_R64_UINT:				{ FormatData<EDefaultInputType::Typed, TInType, uint64>					(InData, OutData, OutByteCount); } break;
+
 		// 32bits
 		case PF_R32G32B32A32_UINT:		{ FormatData<EDefaultInputType::Typed, TInType,	uint32,  0, 1, 2, 3>	(InData, OutData, OutByteCount); } break;
 		case PF_A32B32G32R32F:			{ FormatData<EDefaultInputType::Typed, TInType,	float,   3, 2, 1, 0>	(InData, OutData, OutByteCount); } break;
diff --git a/Engine/Source/Runtime/VulkanRHI/Private/VulkanDevice.cpp b/Engine/Source/Runtime/VulkanRHI/Private/VulkanDevice.cpp
--- a/Engine/Source/Runtime/VulkanRHI/Private/VulkanDevice.cpp
+++ b/Engine/Source/Runtime/VulkanRHI/Private/VulkanDevice.cpp
@@ -640,6 +640,10 @@ void FVulkanDevice::SetupFormats()
 	MapFormatSupport(PF_R32G32B32F, { VK_FORMAT_R32G32B32_SFLOAT }, ComponentMappingRGB1);
 	MapFormatSupport(PF_R8_SINT, { VK_FORMAT_R8_SINT }, ComponentMappingR001);
 
+	// This will be the format used for 64bit image atomics
+	MapFormatSupport(PF_R64_UINT, { VK_FORMAT_R64_UINT }, ComponentMappingR001, EPixelFormatCapabilities::TextureAtomics);
+	checkf(!GRHISupportsAtomicUInt64 || EnumHasAnyFlags(GPixelFormats[PF_R64_UINT].Capabilities, EPixelFormatCapabilities::TextureSample | EPixelFormatCapabilities::UAV), TEXT("Image atomics were enabled, but the R64 format does not have TextureSample or UAV capabilities!"));
+
 	if (CVarVulkanUseD24.GetValueOnAnyThread() != 0)
 	{
 		// prefer VK_FORMAT_D24_UNORM_S8_UINT
diff --git a/Engine/Source/Runtime/VulkanRHI/Private/VulkanRHIPrivate.h b/Engine/Source/Runtime/VulkanRHI/Private/VulkanRHIPrivate.h
--- a/Engine/Source/Runtime/VulkanRHI/Private/VulkanRHIPrivate.h
+++ b/Engine/Source/Runtime/VulkanRHI/Private/VulkanRHIPrivate.h
@@ -513,6 +513,8 @@ namespace VulkanRHI
 		case VK_FORMAT_R16G16B16A16_SNORM:
 		case VK_FORMAT_R16G16B16A16_UINT:
 		case VK_FORMAT_R16G16B16A16_SINT:
+		case VK_FORMAT_R64_UINT:
+		case VK_FORMAT_R64_SINT:
 			return 64;
 		case VK_FORMAT_R32G32B32A32_SFLOAT:
 		case VK_FORMAT_R32G32B32A32_UINT:
