Implements FRHITextureCreateInfo::Validate() to also validate maximum dimension sizes

diff --git a/Engine/Source/Runtime/RHI/Private/RHI.cpp b/Engine/Source/Runtime/RHI/Private/RHI.cpp
--- a/Engine/Source/Runtime/RHI/Private/RHI.cpp
+++ b/Engine/Source/Runtime/RHI/Private/RHI.cpp
@@ -2022,38 +2022,88 @@ void FRHIRenderPassInfo::Validate() const
 }
 #endif
 
-bool FRHITextureCreateInfo::IsValid() const
+#define ValidateTextureDesc(expr, format, ...) \
+	if (bFatal) \
+	{ \
+		checkf(expr, format, ##__VA_ARGS__); \
+	} \
+	else if (!(expr)) \
+	{ \
+		return false;\
+	} \
+
+// static
+bool FRHITextureCreateInfo::Validate(const FRHITextureCreateInfo& Desc, const TCHAR* Name, bool bFatal)
 {
-	if (Extent.X <= 0 || Extent.Y <= 0 || Depth == 0 || ArraySize == 0 || NumMips == 0 || NumSamples < 1 || NumSamples > 8)
+	// Validate texture's pixel format.
 	{
-		return false;
+		ValidateTextureDesc(Desc.Format != PF_Unknown, TEXT("Illegal to create texture %s with an invalid pixel format."), Name);
+		ValidateTextureDesc(Desc.Format < PF_MAX, TEXT("Illegal to create texture %s with an invalid pixel format."), Name);
+		ValidateTextureDesc(GPixelFormats[Desc.Format].Supported,
+			TEXT("Failed to create texture %s with pixel format %s because it is not supported."),
+			Name,
+			GPixelFormats[Desc.Format].Name);
 	}
 
-	if (NumSamples > 1 && !(Dimension == ETextureDimension::Texture2D || Dimension == ETextureDimension::Texture2DArray))
+	// Validate texture's extent.
 	{
-		return false;
+		int32 MaxDimension = (Desc.Dimension == ETextureDimension::TextureCube || Desc.Dimension == ETextureDimension::TextureCubeArray) ? GMaxCubeTextureDimensions : GMaxTextureDimensions;
+
+		ValidateTextureDesc(Desc.Extent.X > 0, TEXT("Texture %s's Extent.X=%d is invalid."), Name, Desc.Extent.X);
+		ValidateTextureDesc(Desc.Extent.X <= MaxDimension, TEXT("Texture %s's Extent.X=%d is too large."), Name, Desc.Extent.X);
+
+		ValidateTextureDesc(Desc.Extent.Y > 0, TEXT("Texture %s's Extent.Y=%d is invalid."), Name, Desc.Extent.Y);
+		ValidateTextureDesc(Desc.Extent.Y <= MaxDimension, TEXT("Texture %s's Extent.Y=%d is too large."), Name, Desc.Extent.Y);
 	}
 
-	if (Dimension == ETextureDimension::Texture3D)
+	// Validate texture's depth
+	if (Desc.Dimension == ETextureDimension::Texture3D)
 	{
-		if (ArraySize > 1)
-		{
-			return false;
-		}
+		ValidateTextureDesc(Desc.Depth > 0, TEXT("Texture %s's Depth=%d is invalid."), Name, int32(Desc.Depth));
+		ValidateTextureDesc(Desc.Depth <= GMaxTextureDimensions, TEXT("Texture %s's Extent.Depth=%d is too large."), Name, Desc.Depth);
 	}
-	else if (Depth > 1)
+	else
 	{
-		return false;
+		ValidateTextureDesc(Desc.Depth == 1, TEXT("Texture %s's Depth=%d is invalid for Dimension=%s."), Name, int32(Desc.Depth), GetTextureDimensionString(Desc.Dimension));
 	}
 
-	if (Format == PF_Unknown)
+	// Validate texture's array size
+	if (Desc.Dimension == ETextureDimension::Texture2DArray || Desc.Dimension == ETextureDimension::TextureCubeArray)
 	{
-		return false;
+		ValidateTextureDesc(Desc.ArraySize > 0, TEXT("Texture %s's ArraySize=%d is invalid."), Name, Desc.ArraySize);
+		ValidateTextureDesc(Desc.ArraySize <= GMaxTextureArrayLayers, TEXT("Texture %s's Extent.ArraySize=%d is too large."), Name, int32(Desc.ArraySize));
+	}
+	else
+	{
+		ValidateTextureDesc(Desc.ArraySize == 1, TEXT("Texture %s's ArraySize=%d is invalid for Dimension=%s."), Name, Desc.ArraySize,  GetTextureDimensionString(Desc.Dimension));
+	}
+
+	// Validate texture's samples count.
+	if (Desc.Dimension == ETextureDimension::Texture2D || Desc.Dimension == ETextureDimension::Texture2DArray)
+	{
+		ValidateTextureDesc(Desc.NumSamples > 0, TEXT("Texture %s's NumSamples=%d is invalid."), Name, Desc.NumSamples);
+	}
+	else
+	{
+		ValidateTextureDesc(Desc.NumSamples == 1, TEXT("Texture %s's NumSamples=%d is invalid for Dimension=%s."), Name, Desc.NumSamples, GetTextureDimensionString(Desc.Dimension));
+	}
+
+	// Validate texture's mips.
+	if (Desc.IsMultisample())
+	{
+		ValidateTextureDesc(Desc.NumMips == 1, TEXT("MSAA Texture %s's can only have one mip."), Name);
+	}
+	else
+	{
+		ValidateTextureDesc(Desc.NumMips > 0, TEXT("Texture %s's NumMips=%d is invalid."), Name, Desc.NumMips);
+		ValidateTextureDesc(Desc.NumMips <= GMaxTextureMipCount, TEXT("Texture %s's NumMips=%d is too large."), Name, Desc.NumMips);
 	}
 
 	return true;
 }
 
+#undef ValidateTextureDesc
+
 static FRHIPanicEvent RHIPanicEvent;
 FRHIPanicEvent& RHIGetPanicDelegate()
 {
diff --git a/Engine/Source/Runtime/RHI/Private/RHIValidation.cpp b/Engine/Source/Runtime/RHI/Private/RHIValidation.cpp
--- a/Engine/Source/Runtime/RHI/Private/RHIValidation.cpp
+++ b/Engine/Source/Runtime/RHI/Private/RHIValidation.cpp
@@ -1543,6 +1543,7 @@ FValidationTransientResourceAllocator::~FValidationTransientResourceAllocator()
 FRHITransientTexture* FValidationTransientResourceAllocator::CreateTexture(const FRHITextureCreateInfo& InCreateInfo, const TCHAR* InDebugName, uint32 InPassIndex)
 {
 	check(!bFrozen);
+	check(FRHITextureCreateInfo::CheckValidity(InCreateInfo, InDebugName));
 	checkf(!EnumHasAnyFlags(InCreateInfo.Flags, TexCreate_Transient), TEXT("Attempted to create transient texture %s using legacy TexCreate_Transient flag."), InDebugName);
 
 	FRHITransientTexture* TransientTexture = RHIAllocator->CreateTexture(InCreateInfo, InDebugName, InPassIndex);
diff --git a/Engine/Source/Runtime/RHI/Public/RHIResources.h b/Engine/Source/Runtime/RHI/Public/RHIResources.h
--- a/Engine/Source/Runtime/RHI/Public/RHIResources.h
+++ b/Engine/Source/Runtime/RHI/Public/RHIResources.h
@@ -3515,7 +3515,10 @@ struct RHI_API FRHITextureCreateInfo
 	}
 
 	/** Returns whether this descriptor conforms to requirements. */
-	bool IsValid() const;
+	inline bool IsValid() const
+	{
+		return FRHITextureCreateInfo::Validate(*this, /* Name = */ TEXT(""), /* bFatal = */ false);
+	}
 
 	/** Clear value to use when fast-clearing the texture. */
 	FClearValueBinding ClearValue;
@@ -3546,6 +3549,15 @@ struct RHI_API FRHITextureCreateInfo
 
 	/** Number of samples in the texture. >1 for MSAA. */
 	uint8 NumSamples = 1;
+
+	/** Check the validity. */
+	static bool CheckValidity(const FRHITextureCreateInfo& Desc, const TCHAR* Name)
+	{
+		return FRHITextureCreateInfo::Validate(Desc, Name, /* bFatal = */ true);
+	}
+
+private:
+	static bool Validate(const FRHITextureCreateInfo& Desc, const TCHAR* Name, bool bFatal);
 };
 
 /** Used to specify a texture metadata plane when creating a view. */
diff --git a/Engine/Source/Runtime/RenderCore/Private/RenderGraphValidation.cpp b/Engine/Source/Runtime/RenderCore/Private/RenderGraphValidation.cpp
--- a/Engine/Source/Runtime/RenderCore/Private/RenderGraphValidation.cpp
+++ b/Engine/Source/Runtime/RenderCore/Private/RenderGraphValidation.cpp
@@ -266,12 +266,8 @@ void FRDGUserValidation::ValidateCreateTexture(const FRDGTextureDesc& Desc, cons
 	checkf(Name, TEXT("Creating a texture requires a valid debug name."));
 	ExecuteGuard(TEXT("CreateTexture"), Name);
 
-	// Validate the pixel format.
-	checkf(Desc.Format != PF_Unknown, TEXT("Illegal to create texture %s with an invalid pixel format."), Name);
-	checkf(Desc.Format < PF_MAX, TEXT("Illegal to create texture %s with invalid FPooledRenderTargetDesc::Format."), Name);
-	checkf(GPixelFormats[Desc.Format].Supported,
-		TEXT("Failed to create texture %s with pixel format %s because it is not supported."), Name, GPixelFormats[Desc.Format].Name);
-	checkf(Desc.IsValid(), TEXT("Texture %s was created with an invalid descriptor."), Name);
+	// Make sure the descriptor is supported by the RHI.
+	check(FRDGTextureDesc::CheckValidity(Desc, Name));
 
 	// Can't create back buffer textures
 	checkf(!EnumHasAnyFlags(Desc.Flags, ETextureCreateFlags::Presentable), TEXT("Illegal to create texture %s with presentable flag."), Name);
