Adding UAV view slice start index and size parameters to UAV creation.
RHIs interface update.
Not implemented on Switch, Mac and OpenGL yet.

Tested with Strata on pc dx11/dx12, xb1 base, ps4 and ps5.
Tested wihtout Strata on pc and ps4.
ShooterGame test on PC
Mac preflight https://horde.devtools.epicgames.com/job/61b268709947f1ef980c602e

diff --git a/Engine/Source/Runtime/Apple/AGXRHI/Private/AGXDynamicRHI.h b/Engine/Source/Runtime/Apple/AGXRHI/Private/AGXDynamicRHI.h
--- a/Engine/Source/Runtime/Apple/AGXRHI/Private/AGXDynamicRHI.h
+++ b/Engine/Source/Runtime/Apple/AGXRHI/Private/AGXDynamicRHI.h
@@ -58,7 +58,7 @@ public:
 	virtual void* LockBuffer_BottomOfPipe(FRHICommandListImmediate& RHICmdList, FRHIBuffer* Buffer, uint32 Offset, uint32 SizeRHI, EResourceLockMode LockMode) final override;
 	virtual void UnlockBuffer_BottomOfPipe(FRHICommandListImmediate& RHICmdList, FRHIBuffer* Buffer) final override;
 	virtual FUnorderedAccessViewRHIRef RHICreateUnorderedAccessView(FRHIBuffer* Buffer, bool bUseUAVCounter, bool bAppendBuffer) final override;
-	virtual FUnorderedAccessViewRHIRef RHICreateUnorderedAccessView(FRHITexture* Texture, uint32 MipLevel) final override;
+	virtual FUnorderedAccessViewRHIRef RHICreateUnorderedAccessView(FRHITexture* Texture, uint32 MipLevel, uint16 FirstArraySlice, uint16 NumArraySlices) final override;
 	virtual FUnorderedAccessViewRHIRef RHICreateUnorderedAccessView(FRHIBuffer* Buffer, uint8 Format) final override;
 	virtual FShaderResourceViewRHIRef RHICreateShaderResourceView(FRHIBuffer* Buffer) final override;
 	virtual FShaderResourceViewRHIRef RHICreateShaderResourceView(FRHITexture* Texture2DRHI, const FRHITextureSRVCreateInfo& CreateInfo) final override;
@@ -148,7 +148,7 @@ public:
 	virtual FRenderQueryRHIRef RHICreateRenderQuery_RenderThread(class FRHICommandListImmediate& RHICmdList, ERenderQueryType QueryType) final override;
 	
 	virtual FUnorderedAccessViewRHIRef RHICreateUnorderedAccessView_RenderThread(class FRHICommandListImmediate& RHICmdList, FRHIBuffer* Buffer, bool bUseUAVCounter, bool bAppendBuffer) final override;
-	virtual FUnorderedAccessViewRHIRef RHICreateUnorderedAccessView_RenderThread(class FRHICommandListImmediate& RHICmdList, FRHITexture* Texture, uint32 MipLevel) final override;
+	virtual FUnorderedAccessViewRHIRef RHICreateUnorderedAccessView_RenderThread(class FRHICommandListImmediate& RHICmdList, FRHITexture* Texture, uint32 MipLevel, uint16 FirstArraySlice, uint16 NumArraySlices) final override;
 	virtual FUnorderedAccessViewRHIRef RHICreateUnorderedAccessView_RenderThread(class FRHICommandListImmediate& RHICmdList, FRHIBuffer* Buffer, uint8 Format) final override;
 
 	virtual FShaderResourceViewRHIRef RHICreateShaderResourceView_RenderThread(class FRHICommandListImmediate& RHICmdList, FRHIBuffer* Buffer, uint32 Stride, uint8 Format) final override;
diff --git a/Engine/Source/Runtime/Apple/AGXRHI/Private/AGXUAV.cpp b/Engine/Source/Runtime/Apple/AGXRHI/Private/AGXUAV.cpp
--- a/Engine/Source/Runtime/Apple/AGXRHI/Private/AGXUAV.cpp
+++ b/Engine/Source/Runtime/Apple/AGXRHI/Private/AGXUAV.cpp
@@ -74,18 +74,18 @@ FUnorderedAccessViewRHIRef FAGXDynamicRHI::RHICreateUnorderedAccessView_RenderTh
 	return GDynamicRHI->RHICreateUnorderedAccessView(Buffer, bUseUAVCounter, bAppendBuffer);
 }
 
-FUnorderedAccessViewRHIRef FAGXDynamicRHI::RHICreateUnorderedAccessView_RenderThread(class FRHICommandListImmediate& RHICmdList, FRHITexture* Texture, uint32 MipLevel)
+FUnorderedAccessViewRHIRef FAGXDynamicRHI::RHICreateUnorderedAccessView_RenderThread(class FRHICommandListImmediate& RHICmdList, FRHITexture* Texture, uint32 MipLevel, uint16 FirstArraySlice, uint16 NumArraySlices)
 {
 	FAGXSurface* Surface = (FAGXSurface*)Texture->GetTextureBaseRHI();
 	FAGXTexture Tex = Surface->Texture;
 	if (!(Tex.GetUsage() & mtlpp::TextureUsage::PixelFormatView))
 	{
 		FScopedRHIThreadStaller StallRHIThread(RHICmdList);
-		return GDynamicRHI->RHICreateUnorderedAccessView(Texture, MipLevel);
+		return GDynamicRHI->RHICreateUnorderedAccessView(Texture, MipLevel, FirstArraySlice, NumArraySlices);
 	}
 	else
 	{
-		return GDynamicRHI->RHICreateUnorderedAccessView(Texture, MipLevel);
+		return GDynamicRHI->RHICreateUnorderedAccessView(Texture, MipLevel, FirstArraySlice, NumArraySlices);
 	}
 }
 
@@ -114,8 +114,10 @@ FUnorderedAccessViewRHIRef FAGXDynamicRHI::RHICreateUnorderedAccessView(FRHIBuff
 	}
 }
 
-FUnorderedAccessViewRHIRef FAGXDynamicRHI::RHICreateUnorderedAccessView(FRHITexture* TextureRHI, uint32 MipLevel)
+FUnorderedAccessViewRHIRef FAGXDynamicRHI::RHICreateUnorderedAccessView(FRHITexture* TextureRHI, uint32 MipLevel, uint16 FirstArraySlice, uint16 NumArraySlices)
 {
+	// Slice selection of a texture array still need to be implemented on AGX
+	check(FirstArraySlice == 0 && NumArraySlices == 0);
 	@autoreleasepool {
 	FAGXShaderResourceView* SRV = new FAGXShaderResourceView;
 	SRV->SourceTexture = TextureRHI;
diff --git a/Engine/Source/Runtime/Apple/MetalRHI/Private/MetalDynamicRHI.h b/Engine/Source/Runtime/Apple/MetalRHI/Private/MetalDynamicRHI.h
--- a/Engine/Source/Runtime/Apple/MetalRHI/Private/MetalDynamicRHI.h
+++ b/Engine/Source/Runtime/Apple/MetalRHI/Private/MetalDynamicRHI.h
@@ -59,7 +59,7 @@ public:
 	virtual void* LockBuffer_BottomOfPipe(FRHICommandListImmediate& RHICmdList, FRHIBuffer* Buffer, uint32 Offset, uint32 SizeRHI, EResourceLockMode LockMode) final override;
 	virtual void UnlockBuffer_BottomOfPipe(FRHICommandListImmediate& RHICmdList, FRHIBuffer* Buffer) final override;
 	virtual FUnorderedAccessViewRHIRef RHICreateUnorderedAccessView(FRHIBuffer* Buffer, bool bUseUAVCounter, bool bAppendBuffer) final override;
-	virtual FUnorderedAccessViewRHIRef RHICreateUnorderedAccessView(FRHITexture* Texture, uint32 MipLevel) final override;
+	virtual FUnorderedAccessViewRHIRef RHICreateUnorderedAccessView(FRHITexture* Texture, uint32 MipLevel, uint16 FirstArraySlice, uint16 NumArraySlices) final override;
 	virtual FUnorderedAccessViewRHIRef RHICreateUnorderedAccessView(FRHIBuffer* Buffer, uint8 Format) final override;
 	virtual FShaderResourceViewRHIRef RHICreateShaderResourceView(FRHIBuffer* Buffer) final override;
 	virtual FShaderResourceViewRHIRef RHICreateShaderResourceView(FRHITexture* Texture2DRHI, const FRHITextureSRVCreateInfo& CreateInfo) final override;
@@ -151,7 +151,7 @@ public:
 	virtual FRenderQueryRHIRef RHICreateRenderQuery_RenderThread(class FRHICommandListImmediate& RHICmdList, ERenderQueryType QueryType) final override;
 	
 	virtual FUnorderedAccessViewRHIRef RHICreateUnorderedAccessView_RenderThread(class FRHICommandListImmediate& RHICmdList, FRHIBuffer* Buffer, bool bUseUAVCounter, bool bAppendBuffer) final override;
-	virtual FUnorderedAccessViewRHIRef RHICreateUnorderedAccessView_RenderThread(class FRHICommandListImmediate& RHICmdList, FRHITexture* Texture, uint32 MipLevel) final override;
+	virtual FUnorderedAccessViewRHIRef RHICreateUnorderedAccessView_RenderThread(class FRHICommandListImmediate& RHICmdList, FRHITexture* Texture, uint32 MipLevel, uint16 FirstArraySlice, uint16 NumArraySlices) final override;
 	virtual FUnorderedAccessViewRHIRef RHICreateUnorderedAccessView_RenderThread(class FRHICommandListImmediate& RHICmdList, FRHIBuffer* Buffer, uint8 Format) final override;
 
 	virtual FShaderResourceViewRHIRef RHICreateShaderResourceView_RenderThread(class FRHICommandListImmediate& RHICmdList, FRHIBuffer* Buffer, uint32 Stride, uint8 Format) final override;
diff --git a/Engine/Source/Runtime/Apple/MetalRHI/Private/MetalUAV.cpp b/Engine/Source/Runtime/Apple/MetalRHI/Private/MetalUAV.cpp
--- a/Engine/Source/Runtime/Apple/MetalRHI/Private/MetalUAV.cpp
+++ b/Engine/Source/Runtime/Apple/MetalRHI/Private/MetalUAV.cpp
@@ -74,18 +74,18 @@ FUnorderedAccessViewRHIRef FMetalDynamicRHI::RHICreateUnorderedAccessView_Render
 	return GDynamicRHI->RHICreateUnorderedAccessView(Buffer, bUseUAVCounter, bAppendBuffer);
 }
 
-FUnorderedAccessViewRHIRef FMetalDynamicRHI::RHICreateUnorderedAccessView_RenderThread(class FRHICommandListImmediate& RHICmdList, FRHITexture* Texture, uint32 MipLevel)
+FUnorderedAccessViewRHIRef FMetalDynamicRHI::RHICreateUnorderedAccessView_RenderThread(class FRHICommandListImmediate& RHICmdList, FRHITexture* Texture, uint32 MipLevel, uint16 FirstArraySlice, uint16 NumArraySlices)
 {
 	FMetalSurface* Surface = (FMetalSurface*)Texture->GetTextureBaseRHI();
 	FMetalTexture Tex = Surface->Texture;
 	if (!(Tex.GetUsage() & mtlpp::TextureUsage::PixelFormatView))
 	{
 		FScopedRHIThreadStaller StallRHIThread(RHICmdList);
-		return GDynamicRHI->RHICreateUnorderedAccessView(Texture, MipLevel);
+		return GDynamicRHI->RHICreateUnorderedAccessView(Texture, MipLevel, FirstArraySlice, NumArraySlices);
 	}
 	else
 	{
-		return GDynamicRHI->RHICreateUnorderedAccessView(Texture, MipLevel);
+		return GDynamicRHI->RHICreateUnorderedAccessView(Texture, MipLevel, FirstArraySlice, NumArraySlices);
 	}
 }
 
@@ -114,8 +114,10 @@ FUnorderedAccessViewRHIRef FMetalDynamicRHI::RHICreateUnorderedAccessView(FRHIBu
 	}
 }
 
-FUnorderedAccessViewRHIRef FMetalDynamicRHI::RHICreateUnorderedAccessView(FRHITexture* TextureRHI, uint32 MipLevel)
+FUnorderedAccessViewRHIRef FMetalDynamicRHI::RHICreateUnorderedAccessView(FRHITexture* TextureRHI, uint32 MipLevel, uint16 FirstArraySlice, uint16 NumArraySlices)
 {
+	// Slice selection of a texture array still need to be implemented on Metal
+	check(FirstArraySlice == 0 && NumArraySlices == 0);
 	@autoreleasepool {
 	FMetalShaderResourceView* SRV = new FMetalShaderResourceView;
 	SRV->SourceTexture = TextureRHI;
diff --git a/Engine/Source/Runtime/D3D12RHI/Private/D3D12RHIPrivate.h b/Engine/Source/Runtime/D3D12RHI/Private/D3D12RHIPrivate.h
--- a/Engine/Source/Runtime/D3D12RHI/Private/D3D12RHIPrivate.h
+++ b/Engine/Source/Runtime/D3D12RHI/Private/D3D12RHIPrivate.h
@@ -313,8 +313,8 @@ public:
 	virtual void RHIUnlockBuffer(FRHICommandListImmediate& RHICmdList, FRHIBuffer* Buffer) final override;
 	virtual void RHIUnlockBufferMGPU(FRHICommandListImmediate& RHICmdList, FRHIBuffer* Buffer, uint32 GPUIndex) final override;
 	virtual FUnorderedAccessViewRHIRef RHICreateUnorderedAccessView(FRHIBuffer* Buffer, bool bUseUAVCounter, bool bAppendBuffer) final override;
-	virtual FUnorderedAccessViewRHIRef RHICreateUnorderedAccessView(FRHITexture* Texture, uint32 MipLevel) final override;
-	virtual FUnorderedAccessViewRHIRef RHICreateUnorderedAccessView(FRHITexture* Texture, uint32 MipLevel, uint8 Format) final override;
+	virtual FUnorderedAccessViewRHIRef RHICreateUnorderedAccessView(FRHITexture* Texture, uint32 MipLevel, uint16 FirstArraySlice, uint16 NumArraySlices) final override;
+	virtual FUnorderedAccessViewRHIRef RHICreateUnorderedAccessView(FRHITexture* Texture, uint32 MipLevel, uint8 Format, uint16 FirstArraySlice, uint16 NumArraySlices) final override;
 	virtual FUnorderedAccessViewRHIRef RHICreateUnorderedAccessView(FRHIBuffer* Buffer, uint8 Format) final override;
 	virtual FShaderResourceViewRHIRef RHICreateShaderResourceView(FRHIBuffer* Buffer) final override;
 	virtual FShaderResourceViewRHIRef RHICreateShaderResourceView(FRHIBuffer* Buffer, uint32 Stride, uint8 Format) final override;
@@ -470,8 +470,8 @@ public:
 	virtual FTexture2DArrayRHIRef RHICreateTexture2DArray_RenderThread(class FRHICommandListImmediate& RHICmdList, uint32 SizeX, uint32 SizeY, uint32 SizeZ, uint8 Format, uint32 NumMips, uint32 NumSamples, ETextureCreateFlags Flags, ERHIAccess InResourceState, FRHIResourceCreateInfo& CreateInfo) override;
 	virtual FTexture3DRHIRef RHICreateTexture3D_RenderThread(class FRHICommandListImmediate& RHICmdList, uint32 SizeX, uint32 SizeY, uint32 SizeZ, uint8 Format, uint32 NumMips, ETextureCreateFlags Flags, ERHIAccess InResourceState, FRHIResourceCreateInfo& CreateInfo) override;
 	virtual FUnorderedAccessViewRHIRef RHICreateUnorderedAccessView_RenderThread(class FRHICommandListImmediate& RHICmdList, FRHIBuffer* Buffer, bool bUseUAVCounter, bool bAppendBuffer) override final;
-	virtual FUnorderedAccessViewRHIRef RHICreateUnorderedAccessView_RenderThread(class FRHICommandListImmediate& RHICmdList, FRHITexture* Texture, uint32 MipLevel) override final;
-	virtual FUnorderedAccessViewRHIRef RHICreateUnorderedAccessView_RenderThread(class FRHICommandListImmediate& RHICmdList, FRHITexture* Texture, uint32 MipLevel, uint8 Format) override final;
+	virtual FUnorderedAccessViewRHIRef RHICreateUnorderedAccessView_RenderThread(class FRHICommandListImmediate& RHICmdList, FRHITexture* Texture, uint32 MipLevel, uint16 FirstArraySlice, uint16 NumArraySlices) override final;
+	virtual FUnorderedAccessViewRHIRef RHICreateUnorderedAccessView_RenderThread(class FRHICommandListImmediate& RHICmdList, FRHITexture* Texture, uint32 MipLevel, uint8 Format, uint16 FirstArraySlice, uint16 NumArraySlices) override final;
 	virtual FUnorderedAccessViewRHIRef RHICreateUnorderedAccessView_RenderThread(class FRHICommandListImmediate& RHICmdList, FRHIBuffer* Buffer, uint8 Format) override final;
 	virtual FShaderResourceViewRHIRef RHICreateShaderResourceView_RenderThread(class FRHICommandListImmediate& RHICmdList, FRHITexture* Texture, const FRHITextureSRVCreateInfo& CreateInfo) override final;
 	virtual FShaderResourceViewRHIRef RHICreateShaderResourceView_RenderThread(class FRHICommandListImmediate& RHICmdList, FRHIBuffer* Buffer, uint32 Stride, uint8 Format) override final;
diff --git a/Engine/Source/Runtime/D3D12RHI/Private/D3D12UAV.cpp b/Engine/Source/Runtime/D3D12RHI/Private/D3D12UAV.cpp
--- a/Engine/Source/Runtime/D3D12RHI/Private/D3D12UAV.cpp
+++ b/Engine/Source/Runtime/D3D12RHI/Private/D3D12UAV.cpp
@@ -67,7 +67,7 @@ FUnorderedAccessViewRHIRef FD3D12DynamicRHI::RHICreateUnorderedAccessView(FRHIBu
 	return CreateUAV(UAVDesc, Buffer, bNeedsCounterResource);
 }
 
-FUnorderedAccessViewRHIRef FD3D12DynamicRHI::RHICreateUnorderedAccessView(FRHITexture* TextureRHI, uint32 MipLevel, uint8 Format)
+FUnorderedAccessViewRHIRef FD3D12DynamicRHI::RHICreateUnorderedAccessView(FRHITexture* TextureRHI, uint32 MipLevel, uint8 Format, uint16 FirstArraySlice, uint16 NumArraySlices)
 {
 	FD3D12TextureBase* Texture = GetD3D12TextureFromRHITexture(TextureRHI);
 
@@ -91,8 +91,8 @@ FUnorderedAccessViewRHIRef FD3D12DynamicRHI::RHICreateUnorderedAccessView(FRHITe
 		FD3D12Texture2DArray* Texture2DArray = (FD3D12Texture2DArray*)Texture;
 		UAVDesc.ViewDimension = D3D12_UAV_DIMENSION_TEXTURE2DARRAY;
 		UAVDesc.Texture2DArray.MipSlice = MipLevel;
-		UAVDesc.Texture2DArray.FirstArraySlice = 0;
-		UAVDesc.Texture2DArray.ArraySize = Texture2DArray->GetSizeZ();
+		UAVDesc.Texture2DArray.FirstArraySlice = NumArraySlices == 0 ? 0 : FirstArraySlice;
+		UAVDesc.Texture2DArray.ArraySize = NumArraySlices == 0 ? Texture2DArray->GetSizeZ() : NumArraySlices;
 		UAVDesc.Texture2DArray.PlaneSlice = GetPlaneSliceFromViewFormat(PlatformResourceFormat, UAVDesc.Format);
 
 		return CreateUAV(UAVDesc, Texture2DArray, false);
@@ -119,9 +119,9 @@ FUnorderedAccessViewRHIRef FD3D12DynamicRHI::RHICreateUnorderedAccessView(FRHITe
 	}
 }
 
-FUnorderedAccessViewRHIRef FD3D12DynamicRHI::RHICreateUnorderedAccessView(FRHITexture* TextureRHI, uint32 MipLevel)
+FUnorderedAccessViewRHIRef FD3D12DynamicRHI::RHICreateUnorderedAccessView(FRHITexture* TextureRHI, uint32 MipLevel, uint16 FirstArraySlice, uint16 NumArraySlices)
 {
-	return RHICreateUnorderedAccessView(TextureRHI, MipLevel, TextureRHI->GetFormat());
+	return RHICreateUnorderedAccessView(TextureRHI, MipLevel, TextureRHI->GetFormat(), FirstArraySlice, NumArraySlices);
 }
 
 FUnorderedAccessViewRHIRef FD3D12DynamicRHI::RHICreateUnorderedAccessView(FRHIBuffer* BufferRHI, uint8 Format)
@@ -317,14 +317,14 @@ FUnorderedAccessViewRHIRef FD3D12DynamicRHI::RHICreateUnorderedAccessView_Render
 	return RHICreateUnorderedAccessView(BufferRHI, bUseUAVCounter, bAppendBuffer);
 }
 
-FUnorderedAccessViewRHIRef FD3D12DynamicRHI::RHICreateUnorderedAccessView_RenderThread(FRHICommandListImmediate& RHICmdList, FRHITexture* Texture, uint32 MipLevel)
+FUnorderedAccessViewRHIRef FD3D12DynamicRHI::RHICreateUnorderedAccessView_RenderThread(FRHICommandListImmediate& RHICmdList, FRHITexture* Texture, uint32 MipLevel, uint16 FirstArraySlice, uint16 NumArraySlices)
 {
-	return FD3D12DynamicRHI::RHICreateUnorderedAccessView(Texture, MipLevel);
+	return FD3D12DynamicRHI::RHICreateUnorderedAccessView(Texture, MipLevel, FirstArraySlice, NumArraySlices);
 }
 
-FUnorderedAccessViewRHIRef FD3D12DynamicRHI::RHICreateUnorderedAccessView_RenderThread(FRHICommandListImmediate& RHICmdList, FRHITexture* Texture, uint32 MipLevel, uint8 Format)
+FUnorderedAccessViewRHIRef FD3D12DynamicRHI::RHICreateUnorderedAccessView_RenderThread(FRHICommandListImmediate& RHICmdList, FRHITexture* Texture, uint32 MipLevel, uint8 Format, uint16 FirstArraySlice, uint16 NumArraySlices)
 {
-	return FD3D12DynamicRHI::RHICreateUnorderedAccessView(Texture, MipLevel, Format);
+	return FD3D12DynamicRHI::RHICreateUnorderedAccessView(Texture, MipLevel, Format, FirstArraySlice, NumArraySlices);
 }
 
 FUnorderedAccessViewRHIRef FD3D12DynamicRHI::RHICreateUnorderedAccessView_RenderThread(FRHICommandListImmediate& RHICmdList, FRHIBuffer* BufferRHI, uint8 Format)
diff --git a/Engine/Source/Runtime/NullDrv/Public/NullRHI.h b/Engine/Source/Runtime/NullDrv/Public/NullRHI.h
--- a/Engine/Source/Runtime/NullDrv/Public/NullRHI.h
+++ b/Engine/Source/Runtime/NullDrv/Public/NullRHI.h
@@ -149,7 +149,7 @@ public:
 	}
 
 
-	virtual FUnorderedAccessViewRHIRef RHICreateUnorderedAccessView(FRHITexture* Texture, uint32 MipLevel) final override
+	virtual FUnorderedAccessViewRHIRef RHICreateUnorderedAccessView(FRHITexture* Texture, uint32 MipLevel, uint16 FirstArraySlice, uint16 NumArraySlices) final override
 	{ 
 		return new FRHIUnorderedAccessView(); 
 	}
diff --git a/Engine/Source/Runtime/OpenGLDrv/Private/OpenGLUAV.cpp b/Engine/Source/Runtime/OpenGLDrv/Private/OpenGLUAV.cpp
--- a/Engine/Source/Runtime/OpenGLDrv/Private/OpenGLUAV.cpp
+++ b/Engine/Source/Runtime/OpenGLDrv/Private/OpenGLUAV.cpp
@@ -174,9 +174,11 @@ FOpenGLShaderResourceView::~FOpenGLShaderResourceView()
 	}
 }
 
-FUnorderedAccessViewRHIRef FOpenGLDynamicRHI::RHICreateUnorderedAccessView(FRHITexture* TextureRHI, uint32 MipLevel)
+FUnorderedAccessViewRHIRef FOpenGLDynamicRHI::RHICreateUnorderedAccessView(FRHITexture* TextureRHI, uint32 MipLevel, uint16 FirstArraySlice, uint16 NumArraySlices)
 {
 	check(TextureRHI->GetFlags() & TexCreate_UAV);
+	// Slice selection of a texture array still need to be implemented on OpenGL
+	check(FirstArraySlice == 0 && NumArraySlices == 0);
 	return new FOpenGLTextureUnorderedAccessView(TextureRHI);
 }
 
diff --git a/Engine/Source/Runtime/OpenGLDrv/Public/OpenGLDrv.h b/Engine/Source/Runtime/OpenGLDrv/Public/OpenGLDrv.h
--- a/Engine/Source/Runtime/OpenGLDrv/Public/OpenGLDrv.h
+++ b/Engine/Source/Runtime/OpenGLDrv/Public/OpenGLDrv.h
@@ -347,7 +347,7 @@ public:
 	virtual void* LockBuffer_BottomOfPipe(FRHICommandListImmediate& RHICmdList, FRHIBuffer* Buffer, uint32 Offset, uint32 Size, EResourceLockMode LockMode) final override;
 	virtual void UnlockBuffer_BottomOfPipe(FRHICommandListImmediate& RHICmdList, FRHIBuffer* Buffer) final override;
 	virtual FUnorderedAccessViewRHIRef RHICreateUnorderedAccessView(FRHIBuffer* Buffer, bool bUseUAVCounter, bool bAppendBuffer) final override;
-	virtual FUnorderedAccessViewRHIRef RHICreateUnorderedAccessView(FRHITexture* Texture, uint32 MipLevel) final override;
+	virtual FUnorderedAccessViewRHIRef RHICreateUnorderedAccessView(FRHITexture* Texture, uint32 MipLevel, uint16 FirstArraySlice, uint16 NumArraySlices) final override;
 	virtual FUnorderedAccessViewRHIRef RHICreateUnorderedAccessView(FRHIBuffer* Buffer, uint8 Format) final override;
 	virtual FShaderResourceViewRHIRef RHICreateShaderResourceView(FRHIBuffer* Buffer) final override;
 	virtual FShaderResourceViewRHIRef RHICreateShaderResourceView(FRHIBuffer* Buffer, uint32 Stride, uint8 Format) final override;
@@ -827,10 +827,10 @@ public:
 		RHITHREAD_GLCOMMAND_EPILOGUE_RETURN(FUnorderedAccessViewRHIRef);
 	}
 
-	virtual FUnorderedAccessViewRHIRef RHICreateUnorderedAccessView_RenderThread(class FRHICommandListImmediate& RHICmdList, FRHITexture* Texture, uint32 MipLevel) final override
+	virtual FUnorderedAccessViewRHIRef RHICreateUnorderedAccessView_RenderThread(class FRHICommandListImmediate& RHICmdList, FRHITexture* Texture, uint32 MipLevel, uint16 FirstArraySlice, uint16 NumArraySlices) final override
 	{
 		RHITHREAD_GLCOMMAND_PROLOGUE();
-		return this->RHICreateUnorderedAccessView(Texture, MipLevel);
+		return this->RHICreateUnorderedAccessView(Texture, MipLevel, FirstArraySlice, NumArraySlices);
 		RHITHREAD_GLCOMMAND_EPILOGUE_RETURN(FUnorderedAccessViewRHIRef);
 	}
 
diff --git a/Engine/Source/Runtime/RHI/Private/DynamicRHI.cpp b/Engine/Source/Runtime/RHI/Private/DynamicRHI.cpp
--- a/Engine/Source/Runtime/RHI/Private/DynamicRHI.cpp
+++ b/Engine/Source/Runtime/RHI/Private/DynamicRHI.cpp
@@ -472,10 +472,20 @@ void FDynamicRHI::RHITransferBufferUnderlyingResource(FRHIBuffer* DestBuffer, FR
 	UE_LOG(LogRHI, Fatal, TEXT("RHITransferBufferUnderlyingResource isn't implemented for the current RHI"));
 }
 
+FUnorderedAccessViewRHIRef FDynamicRHI::RHICreateUnorderedAccessView(FRHITexture* Texture, uint32 MipLevel)
+{
+	return RHICreateUnorderedAccessView(Texture, MipLevel, 0, 0);
+}
+
 FUnorderedAccessViewRHIRef FDynamicRHI::RHICreateUnorderedAccessView(FRHITexture* Texture, uint32 MipLevel, uint8 Format)
+{
+	return RHICreateUnorderedAccessView(Texture, MipLevel, Format, 0, 0);
+}
+
+FUnorderedAccessViewRHIRef FDynamicRHI::RHICreateUnorderedAccessView(FRHITexture* Texture, uint32 MipLevel, uint8 Format, uint16 FirstArraySlice, uint16 NumArraySlices)
 {
 	UE_LOG(LogRHI, Fatal, TEXT("RHICreateUnorderedAccessView with Format parameter isn't implemented for the current RHI"));
-	return RHICreateUnorderedAccessView(Texture, MipLevel);
+	return RHICreateUnorderedAccessView(Texture, MipLevel, FirstArraySlice, NumArraySlices);
 }
 
 void FDynamicRHI::RHIUpdateShaderResourceView(FRHIShaderResourceView* SRV, FRHIBuffer* Buffer, uint32 Stride, uint8 Format)
diff --git a/Engine/Source/Runtime/RHI/Private/RHI.cpp b/Engine/Source/Runtime/RHI/Private/RHI.cpp
--- a/Engine/Source/Runtime/RHI/Private/RHI.cpp
+++ b/Engine/Source/Runtime/RHI/Private/RHI.cpp
@@ -2717,11 +2717,11 @@ FRHIUnorderedAccessView* FRHITextureViewCache::GetOrCreateUAV(FRHITexture* Textu
 	{
 		if (UAVCreateInfo.Format != PF_Unknown)
 		{
-			RHIUnorderedAccessView = RHICreateUnorderedAccessView(Texture, UAVCreateInfo.MipLevel, UAVCreateInfo.Format);
+			RHIUnorderedAccessView = RHICreateUnorderedAccessView(Texture, UAVCreateInfo.MipLevel, UAVCreateInfo.Format, UAVCreateInfo.FirstArraySlice, UAVCreateInfo.NumArraySlices);
 		}
 		else
 		{
-			RHIUnorderedAccessView = RHICreateUnorderedAccessView(Texture, UAVCreateInfo.MipLevel);
+			RHIUnorderedAccessView = RHICreateUnorderedAccessView(Texture, UAVCreateInfo.MipLevel, UAVCreateInfo.FirstArraySlice, UAVCreateInfo.NumArraySlices);
 		}
 	}
 
diff --git a/Engine/Source/Runtime/RHI/Private/RHICommandList.cpp b/Engine/Source/Runtime/RHI/Private/RHICommandList.cpp
--- a/Engine/Source/Runtime/RHI/Private/RHICommandList.cpp
+++ b/Engine/Source/Runtime/RHI/Private/RHICommandList.cpp
@@ -2531,18 +2531,18 @@ FUnorderedAccessViewRHIRef FDynamicRHI::RHICreateUnorderedAccessView_RenderThrea
 	return GDynamicRHI->RHICreateUnorderedAccessView(Buffer, bUseUAVCounter, bAppendBuffer);
 }
 
-FUnorderedAccessViewRHIRef FDynamicRHI::RHICreateUnorderedAccessView_RenderThread(class FRHICommandListImmediate& RHICmdList, FRHITexture* Texture, uint32 MipLevel)
+FUnorderedAccessViewRHIRef FDynamicRHI::RHICreateUnorderedAccessView_RenderThread(class FRHICommandListImmediate& RHICmdList, FRHITexture* Texture, uint32 MipLevel, uint16 FirstArraySlice, uint16 NumArraySlices)
 {
 	CSV_SCOPED_TIMING_STAT(RHITStalls, RHICreateUnorderedAccessView_RenderThread);
 	FScopedRHIThreadStaller StallRHIThread(RHICmdList);
-	return GDynamicRHI->RHICreateUnorderedAccessView(Texture, MipLevel);
+	return GDynamicRHI->RHICreateUnorderedAccessView(Texture, MipLevel, FirstArraySlice, NumArraySlices);
 }
 
-FUnorderedAccessViewRHIRef FDynamicRHI::RHICreateUnorderedAccessView_RenderThread(class FRHICommandListImmediate& RHICmdList, FRHITexture* Texture, uint32 MipLevel, uint8 Format)
+FUnorderedAccessViewRHIRef FDynamicRHI::RHICreateUnorderedAccessView_RenderThread(class FRHICommandListImmediate& RHICmdList, FRHITexture* Texture, uint32 MipLevel, uint8 Format, uint16 FirstArraySlice, uint16 NumArraySlices)
 {
 	CSV_SCOPED_TIMING_STAT(RHITStalls, RHICreateUnorderedAccessView_RenderThread);
 	FScopedRHIThreadStaller StallRHIThread(RHICmdList);
-	return GDynamicRHI->RHICreateUnorderedAccessView(Texture, MipLevel, Format);
+	return GDynamicRHI->RHICreateUnorderedAccessView(Texture, MipLevel, Format, FirstArraySlice, NumArraySlices);
 }
 
 FUnorderedAccessViewRHIRef FDynamicRHI::RHICreateUnorderedAccessView_RenderThread(class FRHICommandListImmediate& RHICmdList, FRHIBuffer* Buffer, uint8 Format)
diff --git a/Engine/Source/Runtime/RHI/Public/DynamicRHI.h b/Engine/Source/Runtime/RHI/Public/DynamicRHI.h
--- a/Engine/Source/Runtime/RHI/Public/DynamicRHI.h
+++ b/Engine/Source/Runtime/RHI/Public/DynamicRHI.h
@@ -647,12 +647,22 @@ public:
 
 	/** Creates an unordered access view of the given texture. */
 	// FlushType: Wait RHI Thread
-	virtual FUnorderedAccessViewRHIRef RHICreateUnorderedAccessView(FRHITexture* Texture, uint32 MipLevel) = 0;
+	UE_DEPRECATED(5.0, "RHICreateUnorderedAccessView now requires FirstArraySlice and NumArraySlices parameters.")
+	virtual FUnorderedAccessViewRHIRef RHICreateUnorderedAccessView(FRHITexture* Texture, uint32 MipLevel);
 
 	/** Creates an unordered access view of the given texture. */
 	// FlushType: Wait RHI Thread
+	virtual FUnorderedAccessViewRHIRef RHICreateUnorderedAccessView(FRHITexture* Texture, uint32 MipLevel, uint16 FirstArraySlice, uint16 NumArraySlices) = 0;
+
+	/** Creates an unordered access view of the given texture. */
+	// FlushType: Wait RHI Thread
+	UE_DEPRECATED(5.0, "RHICreateUnorderedAccessView now requires FirstArraySlice and NumArraySlices parameters.")
 	virtual FUnorderedAccessViewRHIRef RHICreateUnorderedAccessView(FRHITexture* Texture, uint32 MipLevel, uint8 Format);
 
+	/** Creates an unordered access view of the given texture. */
+	// FlushType: Wait RHI Thread
+	virtual FUnorderedAccessViewRHIRef RHICreateUnorderedAccessView(FRHITexture* Texture, uint32 MipLevel, uint8 Format, uint16 FirstArraySlice, uint16 NumArraySlices);
+
 	/** Creates an unordered access view of the given vertex buffer. */
 	// FlushType: Wait RHI Thread
 	virtual FUnorderedAccessViewRHIRef RHICreateUnorderedAccessView(FRHIBuffer* Buffer, uint8 Format) = 0;
@@ -1401,8 +1411,8 @@ public:
 	
 	virtual FTexture3DRHIRef RHICreateTexture3D_RenderThread(class FRHICommandListImmediate& RHICmdList, uint32 SizeX, uint32 SizeY, uint32 SizeZ, uint8 Format, uint32 NumMips, ETextureCreateFlags Flags, ERHIAccess InResourceState, FRHIResourceCreateInfo& CreateInfo);
 	virtual FUnorderedAccessViewRHIRef RHICreateUnorderedAccessView_RenderThread(class FRHICommandListImmediate& RHICmdList, FRHIBuffer* Buffer, bool bUseUAVCounter, bool bAppendBuffer);
-	virtual FUnorderedAccessViewRHIRef RHICreateUnorderedAccessView_RenderThread(class FRHICommandListImmediate& RHICmdList, FRHITexture* Texture, uint32 MipLevel);
-	virtual FUnorderedAccessViewRHIRef RHICreateUnorderedAccessView_RenderThread(class FRHICommandListImmediate& RHICmdList, FRHITexture* Texture, uint32 MipLevel, uint8 Format);
+	virtual FUnorderedAccessViewRHIRef RHICreateUnorderedAccessView_RenderThread(class FRHICommandListImmediate& RHICmdList, FRHITexture* Texture, uint32 MipLevel, uint16 FirstArraySlice, uint16 NumArraySlices);
+	virtual FUnorderedAccessViewRHIRef RHICreateUnorderedAccessView_RenderThread(class FRHICommandListImmediate& RHICmdList, FRHITexture* Texture, uint32 MipLevel, uint8 Format, uint16 FirstArraySlice, uint16 NumArraySlices);
 	virtual FUnorderedAccessViewRHIRef RHICreateUnorderedAccessView_RenderThread(class FRHICommandListImmediate& RHICmdList, FRHIBuffer* Buffer, uint8 Format);
 	virtual FShaderResourceViewRHIRef RHICreateShaderResourceView_RenderThread(class FRHICommandListImmediate& RHICmdList, FRHITexture* Texture, const FRHITextureSRVCreateInfo& CreateInfo);
 	virtual FShaderResourceViewRHIRef RHICreateShaderResourceView_RenderThread(class FRHICommandListImmediate& RHICmdList, FRHIBuffer* Buffer, uint32 Stride, uint8 Format);
diff --git a/Engine/Source/Runtime/RHI/Public/RHICommandList.h b/Engine/Source/Runtime/RHI/Public/RHICommandList.h
--- a/Engine/Source/Runtime/RHI/Public/RHICommandList.h
+++ b/Engine/Source/Runtime/RHI/Public/RHICommandList.h
@@ -4315,18 +4315,18 @@ public:
 		return GDynamicRHI->RHICreateUnorderedAccessView_RenderThread(*this, Buffer, bUseUAVCounter, bAppendBuffer);
 	}
 	
-	FORCEINLINE FUnorderedAccessViewRHIRef CreateUnorderedAccessView(FRHITexture* Texture, uint32 MipLevel)
+	FORCEINLINE FUnorderedAccessViewRHIRef CreateUnorderedAccessView(FRHITexture* Texture, uint32 MipLevel, uint16 FirstArraySlice, uint16 NumArraySlices)
 	{
 		LLM_SCOPE_BYNAME(TEXT("RHIMisc/CreateUnorderedAccessView"));
 		checkf(Texture, TEXT("Can't create a view off a null resource!"));
-		return GDynamicRHI->RHICreateUnorderedAccessView_RenderThread(*this, Texture, MipLevel);
+		return GDynamicRHI->RHICreateUnorderedAccessView_RenderThread(*this, Texture, MipLevel, FirstArraySlice, NumArraySlices);
 	}
 	
-	FORCEINLINE FUnorderedAccessViewRHIRef CreateUnorderedAccessView(FRHITexture* Texture, uint32 MipLevel, uint8 Format)
+	FORCEINLINE FUnorderedAccessViewRHIRef CreateUnorderedAccessView(FRHITexture* Texture, uint32 MipLevel, uint8 Format, uint16 FirstArraySlice, uint16 NumArraySlices)
 	{
 		LLM_SCOPE_BYNAME(TEXT("RHIMisc/CreateUnorderedAccessView"));
 		checkf(Texture, TEXT("Can't create a view off a null resource!"));
-		return GDynamicRHI->RHICreateUnorderedAccessView_RenderThread(*this, Texture, MipLevel, Format);
+		return GDynamicRHI->RHICreateUnorderedAccessView_RenderThread(*this, Texture, MipLevel, Format, FirstArraySlice, NumArraySlices);
 	}
 
 	FORCEINLINE FUnorderedAccessViewRHIRef CreateUnorderedAccessView(FRHIBuffer* Buffer, uint8 Format)
@@ -5380,14 +5380,14 @@ FORCEINLINE FUnorderedAccessViewRHIRef RHICreateUnorderedAccessView(FRHIBuffer*
 	return FRHICommandListExecutor::GetImmediateCommandList().CreateUnorderedAccessView(Buffer, bUseUAVCounter, bAppendBuffer);
 }
 
-FORCEINLINE FUnorderedAccessViewRHIRef RHICreateUnorderedAccessView(FRHITexture* Texture, uint32 MipLevel = 0)
+FORCEINLINE FUnorderedAccessViewRHIRef RHICreateUnorderedAccessView(FRHITexture* Texture, uint32 MipLevel = 0, uint16 FirstArraySlice = 0, uint16 NumArraySlices = 0)
 {
-	return FRHICommandListExecutor::GetImmediateCommandList().CreateUnorderedAccessView(Texture, MipLevel);
+	return FRHICommandListExecutor::GetImmediateCommandList().CreateUnorderedAccessView(Texture, MipLevel, FirstArraySlice, NumArraySlices);
 }
 
-FORCEINLINE FUnorderedAccessViewRHIRef RHICreateUnorderedAccessView(FRHITexture* Texture, uint32 MipLevel, uint8 Format)
+FORCEINLINE FUnorderedAccessViewRHIRef RHICreateUnorderedAccessView(FRHITexture* Texture, uint32 MipLevel, uint8 Format, uint16 FirstArraySlice = 0, uint16 NumArraySlices = 0)
 {
-	return FRHICommandListExecutor::GetImmediateCommandList().CreateUnorderedAccessView(Texture, MipLevel, Format);
+	return FRHICommandListExecutor::GetImmediateCommandList().CreateUnorderedAccessView(Texture, MipLevel, Format, FirstArraySlice, NumArraySlices);
 }
 
 FORCEINLINE FUnorderedAccessViewRHIRef RHICreateUnorderedAccessView(FRHIBuffer* Buffer, uint8 Format)
diff --git a/Engine/Source/Runtime/RHI/Public/RHIResources.h b/Engine/Source/Runtime/RHI/Public/RHIResources.h
--- a/Engine/Source/Runtime/RHI/Public/RHIResources.h
+++ b/Engine/Source/Runtime/RHI/Public/RHIResources.h
@@ -3661,9 +3661,11 @@ struct FRHITextureUAVCreateInfo
 public:
 	FRHITextureUAVCreateInfo() = default;
 
-	explicit FRHITextureUAVCreateInfo(uint8 InMipLevel, EPixelFormat InFormat = PF_Unknown)
+	explicit FRHITextureUAVCreateInfo(uint8 InMipLevel, EPixelFormat InFormat = PF_Unknown, uint16 InFirstArraySlice = 0, uint16 InNumArraySlices = 0)
 		: Format(InFormat)
 		, MipLevel(InMipLevel)
+		, FirstArraySlice(InFirstArraySlice)
+		, NumArraySlices(InNumArraySlices)
 	{}
 
 	explicit FRHITextureUAVCreateInfo(ERHITextureMetaDataAccess InMetaData)
@@ -3672,7 +3674,7 @@ public:
 
 	FORCEINLINE bool operator==(const FRHITextureUAVCreateInfo& Other)const
 	{
-		return Format == Other.Format && MipLevel == Other.MipLevel && MetaData == Other.MetaData;
+		return Format == Other.Format && MipLevel == Other.MipLevel && MetaData == Other.MetaData && FirstArraySlice == Other.FirstArraySlice && NumArraySlices == Other.NumArraySlices;
 	}
 
 	FORCEINLINE bool operator!=(const FRHITextureUAVCreateInfo& Other)const
@@ -3682,6 +3684,8 @@ public:
 
 	EPixelFormat Format = PF_Unknown;
 	uint8 MipLevel = 0;
+	uint16 FirstArraySlice = 0;
+	uint16 NumArraySlices = 0;	// When 0, the default behavior will be used, e.g. all slices mapped.
 	ERHITextureMetaDataAccess MetaData = ERHITextureMetaDataAccess::None;
 };
 
diff --git a/Engine/Source/Runtime/RHI/Public/RHIValidation.h b/Engine/Source/Runtime/RHI/Public/RHIValidation.h
--- a/Engine/Source/Runtime/RHI/Public/RHIValidation.h
+++ b/Engine/Source/Runtime/RHI/Public/RHIValidation.h
@@ -399,23 +399,25 @@ public:
 
 	/** Creates an unordered access view of the given texture. */
 	// FlushType: Wait RHI Thread
-	virtual FUnorderedAccessViewRHIRef RHICreateUnorderedAccessView(FRHITexture* Texture, uint32 MipLevel) override final
+	virtual FUnorderedAccessViewRHIRef RHICreateUnorderedAccessView(FRHITexture* Texture, uint32 MipLevel, uint16 FirstArraySlice, uint16 NumArraySlices) override final
 	{
 		RHI_VALIDATION_CHECK(Texture->GetTextureReference() == nullptr, TEXT("Creating an unordered access view of an FRHITextureReference is not supported."));
+		RHI_VALIDATION_CHECK(Texture->GetTexture2DArray() != nullptr || (Texture->GetTexture2DArray() == nullptr && (FirstArraySlice == 0 && NumArraySlices == 0)), TEXT("ArraySlice controls are only supported with Texture2dArray (it could be supported with Texture1dArray and Cubemap)."));
+		RHI_VALIDATION_CHECK(Texture->GetTexture2DArray() == nullptr || (Texture->GetTexture2DArray() != nullptr && (NumArraySlices == 0 || (FirstArraySlice + NumArraySlices) <= Texture->GetSizeXYZ().Z)), TEXT("ArraySlice controls missused on Texture2dArray."));
 
-		FUnorderedAccessViewRHIRef UAV = RHI->RHICreateUnorderedAccessView(Texture, MipLevel);
-		UAV->ViewIdentity = Texture->GetViewIdentity(MipLevel, 1, 0, 0, uint32(RHIValidation::EResourcePlane::Common), 1);
+		FUnorderedAccessViewRHIRef UAV = RHI->RHICreateUnorderedAccessView(Texture, MipLevel, FirstArraySlice, NumArraySlices);
+		UAV->ViewIdentity = Texture->GetViewIdentity(MipLevel, 1, FirstArraySlice, NumArraySlices, uint32(RHIValidation::EResourcePlane::Common), 1);
 		return UAV;
 	}
 
 	/** Creates an unordered access view of the given texture. */
 	// FlushType: Wait RHI Thread
-	virtual FUnorderedAccessViewRHIRef RHICreateUnorderedAccessView(FRHITexture* Texture, uint32 MipLevel, uint8 Format) override final
+	virtual FUnorderedAccessViewRHIRef RHICreateUnorderedAccessView(FRHITexture* Texture, uint32 MipLevel, uint8 Format, uint16 FirstArraySlice, uint16 NumArraySlices) override final
 	{
 		RHI_VALIDATION_CHECK(Texture->GetTextureReference() == nullptr, TEXT("Creating an unordered access view of an FRHITextureReference is not supported."));
 
-		FUnorderedAccessViewRHIRef UAV = RHI->RHICreateUnorderedAccessView(Texture, MipLevel, Format);
-		UAV->ViewIdentity = Texture->GetViewIdentity(MipLevel, 1, 0, 0, uint32(RHIValidation::EResourcePlane::Common), 1);
+		FUnorderedAccessViewRHIRef UAV = RHI->RHICreateUnorderedAccessView(Texture, MipLevel, Format, FirstArraySlice, NumArraySlices);
+		UAV->ViewIdentity = Texture->GetViewIdentity(MipLevel, 1, FirstArraySlice, NumArraySlices, uint32(RHIValidation::EResourcePlane::Common), 1);
 		return UAV;
 	}
 
@@ -1541,12 +1543,12 @@ public:
 		return UAV;
 	}
 
-	virtual FUnorderedAccessViewRHIRef RHICreateUnorderedAccessView_RenderThread(class FRHICommandListImmediate& RHICmdList, FRHITexture* Texture, uint32 MipLevel) override final
+	virtual FUnorderedAccessViewRHIRef RHICreateUnorderedAccessView_RenderThread(class FRHICommandListImmediate& RHICmdList, FRHITexture* Texture, uint32 MipLevel, uint16 FirstArraySlice, uint16 NumArraySlices) override final
 	{
 		RHI_VALIDATION_CHECK(Texture->GetTextureReference() == nullptr, TEXT("Creating an unordered access view of an FRHITextureReference is not supported."));
 
-		FUnorderedAccessViewRHIRef UAV = RHI->RHICreateUnorderedAccessView_RenderThread(RHICmdList, Texture, MipLevel);
-		UAV->ViewIdentity = Texture->GetViewIdentity(MipLevel, 1, 0, 0, uint32(RHIValidation::EResourcePlane::Common), 1);
+		FUnorderedAccessViewRHIRef UAV = RHI->RHICreateUnorderedAccessView_RenderThread(RHICmdList, Texture, MipLevel, FirstArraySlice, NumArraySlices);
+		UAV->ViewIdentity = Texture->GetViewIdentity(MipLevel, 1, FirstArraySlice, NumArraySlices, uint32(RHIValidation::EResourcePlane::Common), 1);
 		return UAV;
 	}
 
diff --git a/Engine/Source/Runtime/RenderCore/Private/RenderTargetPool.cpp b/Engine/Source/Runtime/RenderCore/Private/RenderTargetPool.cpp
--- a/Engine/Source/Runtime/RenderCore/Private/RenderTargetPool.cpp
+++ b/Engine/Source/Runtime/RenderCore/Private/RenderTargetPool.cpp
@@ -530,7 +530,7 @@ Done:
 					? Desc.UAVFormat
 					: Desc.Format;
 
-				Found->RenderTargetItem.UAV = RHICreateUnorderedAccessView(Found->RenderTargetItem.TargetableTexture, 0, AliasFormat);
+				Found->RenderTargetItem.UAV = RHICreateUnorderedAccessView(Found->RenderTargetItem.TargetableTexture, 0, AliasFormat, 0, 0);
 			}
 			else
 			{
diff --git a/Engine/Source/Runtime/RenderCore/Public/RenderGraphResources.h b/Engine/Source/Runtime/RenderCore/Public/RenderGraphResources.h
--- a/Engine/Source/Runtime/RenderCore/Public/RenderGraphResources.h
+++ b/Engine/Source/Runtime/RenderCore/Public/RenderGraphResources.h
@@ -690,8 +690,8 @@ class FRDGTextureUAVDesc final
 public:
 	FRDGTextureUAVDesc() = default;
 
-	FRDGTextureUAVDesc(FRDGTextureRef InTexture, uint8 InMipLevel = 0, EPixelFormat InFormat = PF_Unknown)
-		: FRHITextureUAVCreateInfo(InMipLevel, InFormat != PF_Unknown ? InFormat : InTexture->Desc.UAVFormat)
+	FRDGTextureUAVDesc(FRDGTextureRef InTexture, uint8 InMipLevel = 0, EPixelFormat InFormat = PF_Unknown, uint16 InFirstArraySlice = 0, uint16 InNumArraySlices = 0)
+		: FRHITextureUAVCreateInfo(InMipLevel, InFormat != PF_Unknown ? InFormat : InTexture->Desc.UAVFormat, InFirstArraySlice, InNumArraySlices)
 		, Texture(InTexture)
 	{}
 
diff --git a/Engine/Source/Runtime/VulkanRHI/Private/VulkanResources.h b/Engine/Source/Runtime/VulkanRHI/Private/VulkanResources.h
--- a/Engine/Source/Runtime/VulkanRHI/Private/VulkanResources.h
+++ b/Engine/Source/Runtime/VulkanRHI/Private/VulkanResources.h
@@ -1228,7 +1228,7 @@ class FVulkanUnorderedAccessView : public FRHIUnorderedAccessView, public Vulkan
 public:
 
 	FVulkanUnorderedAccessView(FVulkanDevice* Device, FVulkanResourceMultiBuffer* Buffer, bool bUseUAVCounter, bool bAppendBuffer);
-	FVulkanUnorderedAccessView(FVulkanDevice* Device, FRHITexture* TextureRHI, uint32 MipLevel);
+	FVulkanUnorderedAccessView(FVulkanDevice* Device, FRHITexture* TextureRHI, uint32 MipLevel, uint16 FirstArraySlice, uint16 NumArraySlices);
 	FVulkanUnorderedAccessView(FVulkanDevice* Device, FVulkanResourceMultiBuffer* Buffer, EPixelFormat Format);
 
 	~FVulkanUnorderedAccessView();
@@ -1242,6 +1242,8 @@ protected:
 	TRefCountPtr<FRHITexture> SourceTexture;
 	FVulkanTextureView TextureView;
 	uint32 MipLevel;
+	uint16 FirstArraySlice;
+	uint16 NumArraySlices;
 
 	// The buffer this UAV comes from (can be null)
 	TRefCountPtr<FVulkanResourceMultiBuffer> SourceBuffer;
diff --git a/Engine/Source/Runtime/VulkanRHI/Private/VulkanUAV.cpp b/Engine/Source/Runtime/VulkanRHI/Private/VulkanUAV.cpp
--- a/Engine/Source/Runtime/VulkanRHI/Private/VulkanUAV.cpp
+++ b/Engine/Source/Runtime/VulkanRHI/Private/VulkanUAV.cpp
@@ -241,16 +241,20 @@ void FVulkanShaderResourceView::UpdateView()
 FVulkanUnorderedAccessView::FVulkanUnorderedAccessView(FVulkanDevice* Device, FVulkanResourceMultiBuffer* Buffer, bool bUseUAVCounter, bool bAppendBuffer)
 	: VulkanRHI::FVulkanViewBase(Device)
 	, MipLevel(0)
+	, FirstArraySlice(0)
+	, NumArraySlices(0)
 	, SourceBuffer(Buffer)
 	, BufferViewFormat(PF_Unknown)
 	, VolatileLockCounter(MAX_uint32)
 {
 }
 
-FVulkanUnorderedAccessView::FVulkanUnorderedAccessView(FVulkanDevice* Device, FRHITexture* TextureRHI, uint32 MipLevel)
+FVulkanUnorderedAccessView::FVulkanUnorderedAccessView(FVulkanDevice* Device, FRHITexture* TextureRHI, uint32 MipLevel, uint16 InFirstArraySlice, uint16 InNumArraySlices)
 	: VulkanRHI::FVulkanViewBase(Device)
 	, SourceTexture(TextureRHI)
 	, MipLevel(MipLevel)
+	, FirstArraySlice(InFirstArraySlice)
+	, NumArraySlices(InNumArraySlices)
 	, BufferViewFormat(PF_Unknown)
 	, VolatileLockCounter(MAX_uint32)
 {
@@ -262,6 +266,8 @@ FVulkanUnorderedAccessView::FVulkanUnorderedAccessView(FVulkanDevice* Device, FR
 FVulkanUnorderedAccessView::FVulkanUnorderedAccessView(FVulkanDevice* Device, FVulkanResourceMultiBuffer* Buffer, EPixelFormat Format)
 	: VulkanRHI::FVulkanViewBase(Device)
 	, MipLevel(0)
+	, FirstArraySlice(0)
+	, NumArraySlices(0)
 	, BufferViewFormat(Format)
 	, VolatileLockCounter(MAX_uint32)
 {
@@ -336,7 +342,8 @@ void FVulkanUnorderedAccessView::UpdateView()
 		else if (FRHITexture2DArray* Tex2DArray = SourceTexture->GetTexture2DArray())
 		{
 			FVulkanTexture2DArray* VTex2DArray = ResourceCast(Tex2DArray);
-			TextureView.Create(*Device, VTex2DArray->Surface.Image, VK_IMAGE_VIEW_TYPE_2D_ARRAY, VTex2DArray->Surface.GetPartialAspectMask(), Format, UEToVkTextureFormat(Format, false), MipLevel, 1, 0, VTex2DArray->GetSizeZ(), true);
+			TextureView.Create(*Device, VTex2DArray->Surface.Image, VK_IMAGE_VIEW_TYPE_2D_ARRAY, VTex2DArray->Surface.GetPartialAspectMask(), Format, UEToVkTextureFormat(Format, false), MipLevel, 1, 
+				NumArraySlices == 0 ? 0 : FirstArraySlice, NumArraySlices == 0 ? VTex2DArray->GetSizeZ() : NumArraySlices, true);
 		}
 		else
 		{
@@ -353,9 +360,9 @@ FUnorderedAccessViewRHIRef FVulkanDynamicRHI::RHICreateUnorderedAccessView(FRHIB
 	return UAV;
 }
 
-FUnorderedAccessViewRHIRef FVulkanDynamicRHI::RHICreateUnorderedAccessView(FRHITexture* TextureRHI, uint32 MipLevel)
+FUnorderedAccessViewRHIRef FVulkanDynamicRHI::RHICreateUnorderedAccessView(FRHITexture* TextureRHI, uint32 MipLevel, uint16 FirstArraySlice, uint16 NumArraySlices)
 {
-	FVulkanUnorderedAccessView* UAV = new FVulkanUnorderedAccessView(Device, TextureRHI, MipLevel);
+	FVulkanUnorderedAccessView* UAV = new FVulkanUnorderedAccessView(Device, TextureRHI, MipLevel, FirstArraySlice, NumArraySlices);
 	return UAV;
 }
 
diff --git a/Engine/Source/Runtime/VulkanRHI/Public/VulkanDynamicRHI.h b/Engine/Source/Runtime/VulkanRHI/Public/VulkanDynamicRHI.h
--- a/Engine/Source/Runtime/VulkanRHI/Public/VulkanDynamicRHI.h
+++ b/Engine/Source/Runtime/VulkanRHI/Public/VulkanDynamicRHI.h
@@ -85,7 +85,7 @@ public:
 	virtual void* LockBuffer_BottomOfPipe(FRHICommandListImmediate& RHICmdList, FRHIBuffer* Buffer, uint32 Offset, uint32 Size, EResourceLockMode LockMode) final override;
 	virtual void UnlockBuffer_BottomOfPipe(FRHICommandListImmediate& RHICmdList, FRHIBuffer* Buffer) final override;
 	virtual FUnorderedAccessViewRHIRef RHICreateUnorderedAccessView(FRHIBuffer* Buffer, bool bUseUAVCounter, bool bAppendBuffer) final override;
-	virtual FUnorderedAccessViewRHIRef RHICreateUnorderedAccessView(FRHITexture* Texture, uint32 MipLevel) final override;
+	virtual FUnorderedAccessViewRHIRef RHICreateUnorderedAccessView(FRHITexture* Texture, uint32 MipLevel, uint16 FirstArraySlice, uint16 NumArraySlices) final override;
 	virtual FUnorderedAccessViewRHIRef RHICreateUnorderedAccessView(FRHIBuffer* Buffer, uint8 Format) final override;
 	virtual FShaderResourceViewRHIRef RHICreateShaderResourceView(FRHIBuffer* Buffer) final override;
 	virtual FShaderResourceViewRHIRef RHICreateShaderResourceView(FRHIBuffer* Buffer, uint32 Stride, uint8 Format) final override;
@@ -253,10 +253,10 @@ public:
 		return this->RHICreateUnorderedAccessView(Buffer, bUseUAVCounter, bAppendBuffer);
 	}
 
-	virtual FUnorderedAccessViewRHIRef RHICreateUnorderedAccessView_RenderThread(class FRHICommandListImmediate& RHICmdList, FRHITexture* Texture, uint32 MipLevel) final override
+	virtual FUnorderedAccessViewRHIRef RHICreateUnorderedAccessView_RenderThread(class FRHICommandListImmediate& RHICmdList, FRHITexture* Texture, uint32 MipLevel, uint16 FirstArraySlice, uint16 NumArraySlices) final override
 	{
 		// this-> is required to avoid calling the global version of this function
-		return this->RHICreateUnorderedAccessView(Texture, MipLevel);
+		return this->RHICreateUnorderedAccessView(Texture, MipLevel, FirstArraySlice, NumArraySlices);
 	}
 
 	virtual FUnorderedAccessViewRHIRef RHICreateUnorderedAccessView_RenderThread(class FRHICommandListImmediate& RHICmdList, FRHIBuffer* Buffer, uint8 Format) final override
diff --git a/Engine/Source/Runtime/Windows/D3D11RHI/Private/D3D11RHIPrivate.h b/Engine/Source/Runtime/Windows/D3D11RHI/Private/D3D11RHIPrivate.h
--- a/Engine/Source/Runtime/Windows/D3D11RHI/Private/D3D11RHIPrivate.h
+++ b/Engine/Source/Runtime/Windows/D3D11RHI/Private/D3D11RHIPrivate.h
@@ -480,7 +480,7 @@ public:
 	virtual void* LockBuffer_BottomOfPipe(FRHICommandListImmediate& RHICmdList, FRHIBuffer* Buffer, uint32 Offset, uint32 Size, EResourceLockMode LockMode) final override;
 	virtual void UnlockBuffer_BottomOfPipe(FRHICommandListImmediate& RHICmdList, FRHIBuffer* Buffer) final override;
 	virtual FUnorderedAccessViewRHIRef RHICreateUnorderedAccessView(FRHIBuffer* Buffer, bool bUseUAVCounter, bool bAppendBuffer) final override;
-	virtual FUnorderedAccessViewRHIRef RHICreateUnorderedAccessView(FRHITexture* Texture, uint32 MipLevel) final override;
+	virtual FUnorderedAccessViewRHIRef RHICreateUnorderedAccessView(FRHITexture* Texture, uint32 MipLevel, uint16 FirstArraySlice, uint16 NumArraySlices) final override;
 	virtual FUnorderedAccessViewRHIRef RHICreateUnorderedAccessView(FRHIBuffer* Buffer, uint8 Format) final override;
 	virtual FShaderResourceViewRHIRef RHICreateShaderResourceView(FRHIBuffer* Buffer) final override;
 	virtual FShaderResourceViewRHIRef RHICreateShaderResourceView(FRHIBuffer* Buffer, uint32 Stride, uint8 Format) final override;
@@ -665,7 +665,7 @@ public:
 	virtual FTexture2DArrayRHIRef RHICreateTexture2DArray_RenderThread(class FRHICommandListImmediate& RHICmdList, uint32 SizeX, uint32 SizeY, uint32 SizeZ, uint8 Format, uint32 NumMips, uint32 NumSamples, ETextureCreateFlags Flags, ERHIAccess InResourceState, FRHIResourceCreateInfo& CreateInfo) final override;
 	virtual FTexture3DRHIRef RHICreateTexture3D_RenderThread(class FRHICommandListImmediate& RHICmdList, uint32 SizeX, uint32 SizeY, uint32 SizeZ, uint8 Format, uint32 NumMips, ETextureCreateFlags Flags, ERHIAccess InResourceState, FRHIResourceCreateInfo& CreateInfo) final override;
 	virtual FUnorderedAccessViewRHIRef RHICreateUnorderedAccessView_RenderThread(class FRHICommandListImmediate& RHICmdList, FRHIBuffer* Buffer, bool bUseUAVCounter, bool bAppendBuffer) final override;
-	virtual FUnorderedAccessViewRHIRef RHICreateUnorderedAccessView_RenderThread(class FRHICommandListImmediate& RHICmdList, FRHITexture* Texture, uint32 MipLevel) final override;
+	virtual FUnorderedAccessViewRHIRef RHICreateUnorderedAccessView_RenderThread(class FRHICommandListImmediate& RHICmdList, FRHITexture* Texture, uint32 MipLevel, uint16 FirstArraySlice, uint16 NumArraySlices) final override;
 	virtual FUnorderedAccessViewRHIRef RHICreateUnorderedAccessView_RenderThread(class FRHICommandListImmediate& RHICmdList, FRHIBuffer* Buffer, uint8 Format) final override;
 	virtual FShaderResourceViewRHIRef RHICreateShaderResourceView_RenderThread(class FRHICommandListImmediate& RHICmdList, const FShaderResourceViewInitializer& Initializer) final override;
 	virtual FShaderResourceViewRHIRef RHICreateShaderResourceView_RenderThread(class FRHICommandListImmediate& RHICmdList, FRHITexture* Texture, const FRHITextureSRVCreateInfo& CreateInfo) final override;
diff --git a/Engine/Source/Runtime/Windows/D3D11RHI/Private/D3D11UAV.cpp b/Engine/Source/Runtime/Windows/D3D11RHI/Private/D3D11UAV.cpp
--- a/Engine/Source/Runtime/Windows/D3D11RHI/Private/D3D11UAV.cpp
+++ b/Engine/Source/Runtime/Windows/D3D11RHI/Private/D3D11UAV.cpp
@@ -74,7 +74,7 @@ FUnorderedAccessViewRHIRef FD3D11DynamicRHI::RHICreateUnorderedAccessView_Render
 	return RHICreateUnorderedAccessView(Buffer, bUseUAVCounter, bAppendBuffer);
 }
 
-FUnorderedAccessViewRHIRef FD3D11DynamicRHI::RHICreateUnorderedAccessView(FRHITexture* TextureRHI, uint32 MipLevel)
+FUnorderedAccessViewRHIRef FD3D11DynamicRHI::RHICreateUnorderedAccessView(FRHITexture* TextureRHI, uint32 MipLevel, uint16 FirstArraySlice, uint16 NumArraySlices)
 {
 	FD3D11TextureBase* Texture = GetD3D11TextureFromRHITexture(TextureRHI);
 	
@@ -93,8 +93,8 @@ FUnorderedAccessViewRHIRef FD3D11DynamicRHI::RHICreateUnorderedAccessView(FRHITe
 		FD3D11Texture2DArray* Texture2DArray = (FD3D11Texture2DArray*)Texture;
 		UAVDesc.ViewDimension = D3D11_UAV_DIMENSION_TEXTURE2DARRAY;
 		UAVDesc.Texture2DArray.MipSlice = MipLevel;
-		UAVDesc.Texture2DArray.FirstArraySlice = 0;
-		UAVDesc.Texture2DArray.ArraySize = Texture2DArray->GetSizeZ();
+		UAVDesc.Texture2DArray.FirstArraySlice = NumArraySlices == 0 ? 0 : FirstArraySlice;
+		UAVDesc.Texture2DArray.ArraySize = NumArraySlices == 0 ? Texture2DArray->GetSizeZ() : NumArraySlices;
 	}
 	else if (TextureRHI->GetTextureCube() != NULL)
 	{
@@ -121,9 +121,11 @@ FUnorderedAccessViewRHIRef FD3D11DynamicRHI::RHICreateUnorderedAccessView(FRHITe
 FUnorderedAccessViewRHIRef FD3D11DynamicRHI::RHICreateUnorderedAccessView_RenderThread(
 	class FRHICommandListImmediate& RHICmdList,
 	FRHITexture* Texture,
-	uint32 MipLevel)
+	uint32 MipLevel, 
+	uint16 FirstArraySlice, 
+	uint16 NumArraySlices)
 {
-	return RHICreateUnorderedAccessView(Texture, MipLevel);
+	return RHICreateUnorderedAccessView(Texture, MipLevel, FirstArraySlice, NumArraySlices);
 }
 
 FUnorderedAccessViewRHIRef FD3D11DynamicRHI::RHICreateUnorderedAccessView(FRHIBuffer* BufferRHI, uint8 Format)
