UnrealBuildTool: Suppress new clang 13 warnings
RHI: fix unresolved exports

Clang 13 needs RHI_API on the class template definitions to match the declarations. Unsure if this is a regression for Clang 13 or due to some other reason, but adding them doesn't negatively affect any other compile targets.

#rnx

diff --git a/Engine/Source/Programs/UnrealBuildTool/Platform/Android/AndroidToolChain.cs b/Engine/Source/Programs/UnrealBuildTool/Platform/Android/AndroidToolChain.cs
--- a/Engine/Source/Programs/UnrealBuildTool/Platform/Android/AndroidToolChain.cs
+++ b/Engine/Source/Programs/UnrealBuildTool/Platform/Android/AndroidToolChain.cs
@@ -609,6 +609,7 @@ namespace UnrealBuildTool
 			{
 				Result += " -Wno-unused-but-set-variable";
 				Result += " -Wno-unused-but-set-parameter";
+				Result += " -Wno-ordered-compare-function-pointers";
 			}
 			// this will hide the warnings about static functions in headers that aren't used in every single .cpp file
 			Result += " -Wno-unused-function";
diff --git a/Engine/Source/Programs/UnrealBuildTool/Platform/Linux/LinuxToolChain.cs b/Engine/Source/Programs/UnrealBuildTool/Platform/Linux/LinuxToolChain.cs
--- a/Engine/Source/Programs/UnrealBuildTool/Platform/Linux/LinuxToolChain.cs
+++ b/Engine/Source/Programs/UnrealBuildTool/Platform/Linux/LinuxToolChain.cs
@@ -686,6 +686,7 @@ namespace UnrealBuildTool
 			{
 				Result += " -Wno-unused-but-set-variable";
 				Result += " -Wno-unused-but-set-parameter";
+				Result += " -Wno-ordered-compare-function-pointers";
 			}
 			// this will hide the warnings about static functions in headers that aren't used in every single .cpp file
 			Result += " -Wno-unused-function";
diff --git a/Engine/Source/Programs/UnrealBuildTool/Platform/Windows/VCToolChain.cs b/Engine/Source/Programs/UnrealBuildTool/Platform/Windows/VCToolChain.cs
--- a/Engine/Source/Programs/UnrealBuildTool/Platform/Windows/VCToolChain.cs
+++ b/Engine/Source/Programs/UnrealBuildTool/Platform/Windows/VCToolChain.cs
@@ -852,6 +852,11 @@ namespace UnrealBuildTool
 				// @todo clang: Kind of a shame to turn these off.  We'd like to catch unused variables, but it is tricky with how our assertion macros work.
 				Arguments.Add("-Wno-inconsistent-missing-override");
 				Arguments.Add("-Wno-unused-variable");
+				if (EnvVars.CompilerVersion >= new VersionNumber(13))
+				{
+					Arguments.Add("-Wno-unused-but-set-variable");
+					Arguments.Add("-Wno-unused-but-set-parameter");
+				}
 				Arguments.Add("-Wno-unused-local-typedefs");
 				Arguments.Add("-Wno-unused-function");
 				Arguments.Add("-Wno-unused-private-field");
@@ -891,6 +896,10 @@ namespace UnrealBuildTool
 				Arguments.Add("-Wno-invalid-token-paste");
 				Arguments.Add("-Wno-null-pointer-arithmetic");
 				Arguments.Add("-Wno-constant-logical-operand"); // Triggered by || of two template-derived values inside a static_assert
+				if (EnvVars.CompilerVersion >= new VersionNumber(13))
+				{
+					Arguments.Add("-Wno-ordered-compare-function-pointers");
+				}
 
 			}
 		}
diff --git a/Engine/Source/Runtime/RHI/Public/RHICommandList.h b/Engine/Source/Runtime/RHI/Public/RHICommandList.h
--- a/Engine/Source/Runtime/RHI/Public/RHICommandList.h
+++ b/Engine/Source/Runtime/RHI/Public/RHICommandList.h
@@ -2355,11 +2355,11 @@ FRHICOMMAND_MACRO(FRHICommandSetRayTracingBindings)
 #define ALLOC_COMMAND(...) new ( AllocCommand(sizeof(__VA_ARGS__), alignof(__VA_ARGS__)) ) __VA_ARGS__
 #define ALLOC_COMMAND_CL(RHICmdList, ...) new ( (RHICmdList).AllocCommand(sizeof(__VA_ARGS__), alignof(__VA_ARGS__)) ) __VA_ARGS__
 
-template<> void FRHICommandSetShaderParameter<FRHIComputeShader>::Execute(FRHICommandListBase& CmdList);
-template<> void FRHICommandSetShaderUniformBuffer<FRHIComputeShader>::Execute(FRHICommandListBase& CmdList);
-template<> void FRHICommandSetShaderTexture<FRHIComputeShader>::Execute(FRHICommandListBase& CmdList);
-template<> void FRHICommandSetShaderResourceViewParameter<FRHIComputeShader>::Execute(FRHICommandListBase& CmdList);
-template<> void FRHICommandSetShaderSampler<FRHIComputeShader>::Execute(FRHICommandListBase& CmdList);
+template<> RHI_API void FRHICommandSetShaderParameter<FRHIComputeShader>::Execute(FRHICommandListBase& CmdList);
+template<> RHI_API void FRHICommandSetShaderUniformBuffer<FRHIComputeShader>::Execute(FRHICommandListBase& CmdList);
+template<> RHI_API void FRHICommandSetShaderTexture<FRHIComputeShader>::Execute(FRHICommandListBase& CmdList);
+template<> RHI_API void FRHICommandSetShaderResourceViewParameter<FRHIComputeShader>::Execute(FRHICommandListBase& CmdList);
+template<> RHI_API void FRHICommandSetShaderSampler<FRHIComputeShader>::Execute(FRHICommandListBase& CmdList);
 
 
 class RHI_API FRHIComputeCommandList : public FRHICommandListBase
diff --git a/Engine/Source/Runtime/RHI/Public/RHICommandListCommandExecutes.inl b/Engine/Source/Runtime/RHI/Public/RHICommandListCommandExecutes.inl
--- a/Engine/Source/Runtime/RHI/Public/RHICommandListCommandExecutes.inl
+++ b/Engine/Source/Runtime/RHI/Public/RHICommandListCommandExecutes.inl
@@ -98,12 +98,12 @@ void FRHICommandWaitForTemporalEffect::Execute(FRHICommandListBase& CmdList)
 	INTERNAL_DECORATOR(RHIWaitForTemporalEffect)(EffectName);
 }
 
-template <> void FRHICommandBroadcastTemporalEffect<FRHITexture>::Execute(FRHICommandListBase& CmdList)
+template<> RHI_API void FRHICommandBroadcastTemporalEffect<FRHITexture>::Execute(FRHICommandListBase& CmdList)
 {
 	RHISTAT(BroadcastTemporalEffect);
 	INTERNAL_DECORATOR(RHIBroadcastTemporalEffect)(EffectName, Resources);
 }
-template <> void FRHICommandBroadcastTemporalEffect<FRHIBuffer>::Execute(FRHICommandListBase& CmdList)
+template<> RHI_API void FRHICommandBroadcastTemporalEffect<FRHIBuffer>::Execute(FRHICommandListBase& CmdList)
 {
 	RHISTAT(BroadcastTemporalEffect);
 	INTERNAL_DECORATOR(RHIBroadcastTemporalEffect)(EffectName, Resources);
@@ -123,61 +123,61 @@ void FRHICommandSetStencilRef::Execute(FRHICommandListBase& CmdList)
 	INTERNAL_DECORATOR(RHISetStencilRef)(StencilRef);
 }
 
-template<> void FRHICommandSetShaderParameter<FRHIComputeShader>::Execute(FRHICommandListBase& CmdList)
+template<> RHI_API void FRHICommandSetShaderParameter<FRHIComputeShader>::Execute(FRHICommandListBase& CmdList)
 {
 	RHISTAT(SetShaderParameter);
 	INTERNAL_DECORATOR_COMPUTE(RHISetShaderParameter)(Shader, BufferIndex, BaseIndex, NumBytes, NewValue);
 }
 
-template <> void FRHICommandSetShaderParameter<FRHIGraphicsShader>::Execute(FRHICommandListBase& CmdList)
+template<> RHI_API void FRHICommandSetShaderParameter<FRHIGraphicsShader>::Execute(FRHICommandListBase& CmdList)
 {
 	RHISTAT(SetShaderParameter);
 	INTERNAL_DECORATOR(RHISetShaderParameter)(Shader, BufferIndex, BaseIndex, NumBytes, NewValue);
 }
 
-template<> void FRHICommandSetShaderUniformBuffer<FRHIComputeShader>::Execute(FRHICommandListBase& CmdList)
+template<> RHI_API void FRHICommandSetShaderUniformBuffer<FRHIComputeShader>::Execute(FRHICommandListBase& CmdList)
 {
 	RHISTAT(SetShaderUniformBuffer);
 	INTERNAL_DECORATOR_COMPUTE(RHISetShaderUniformBuffer)(Shader, BaseIndex, UniformBuffer);
 }
 
-template<> void FRHICommandSetShaderUniformBuffer<FRHIGraphicsShader>::Execute(FRHICommandListBase& CmdList)
+template<> RHI_API void FRHICommandSetShaderUniformBuffer<FRHIGraphicsShader>::Execute(FRHICommandListBase& CmdList)
 {
 	RHISTAT(SetShaderUniformBuffer);
 	INTERNAL_DECORATOR(RHISetShaderUniformBuffer)(Shader, BaseIndex, UniformBuffer);
 }
 
-template<> void FRHICommandSetShaderTexture<FRHIComputeShader>::Execute(FRHICommandListBase& CmdList)
+template<> RHI_API void FRHICommandSetShaderTexture<FRHIComputeShader>::Execute(FRHICommandListBase& CmdList)
 {
 	RHISTAT(SetShaderTexture);
 	INTERNAL_DECORATOR_COMPUTE(RHISetShaderTexture)(Shader, TextureIndex, Texture);
 }
 
-template<> void FRHICommandSetShaderTexture<FRHIGraphicsShader>::Execute(FRHICommandListBase& CmdList)
+template<> RHI_API void FRHICommandSetShaderTexture<FRHIGraphicsShader>::Execute(FRHICommandListBase& CmdList)
 {
 	RHISTAT(SetShaderTexture);
 	INTERNAL_DECORATOR(RHISetShaderTexture)(Shader, TextureIndex, Texture);
 }
 
-template<> void FRHICommandSetShaderResourceViewParameter<FRHIComputeShader>::Execute(FRHICommandListBase& CmdList)
+template<> RHI_API void FRHICommandSetShaderResourceViewParameter<FRHIComputeShader>::Execute(FRHICommandListBase& CmdList)
 {
 	RHISTAT(SetShaderResourceViewParameter);
 	INTERNAL_DECORATOR_COMPUTE(RHISetShaderResourceViewParameter)(Shader, SamplerIndex, SRV);
 }
 
-template<> void FRHICommandSetShaderResourceViewParameter<FRHIGraphicsShader>::Execute(FRHICommandListBase& CmdList)
+template<> RHI_API void FRHICommandSetShaderResourceViewParameter<FRHIGraphicsShader>::Execute(FRHICommandListBase& CmdList)
 {
 	RHISTAT(SetShaderResourceViewParameter);
 	INTERNAL_DECORATOR(RHISetShaderResourceViewParameter)(Shader, SamplerIndex, SRV);
 }
 
-template <> void FRHICommandSetUAVParameter<FRHIComputeShader>::Execute(FRHICommandListBase& CmdList)
+template<> RHI_API void FRHICommandSetUAVParameter<FRHIComputeShader>::Execute(FRHICommandListBase& CmdList)
 {
 	RHISTAT(SetUAVParameter);
 	INTERNAL_DECORATOR_COMPUTE(RHISetUAVParameter)(Shader, UAVIndex, UAV);
 }
 
-template <> void FRHICommandSetUAVParameter<FRHIPixelShader>::Execute(FRHICommandListBase& CmdList)
+template<> RHI_API void FRHICommandSetUAVParameter<FRHIPixelShader>::Execute(FRHICommandListBase& CmdList)
 {
 	RHISTAT(SetUAVParameter);
 	INTERNAL_DECORATOR(RHISetUAVParameter)(Shader, UAVIndex, UAV);
@@ -189,13 +189,13 @@ void FRHICommandSetUAVParameter_InitialCount::Execute(FRHICommandListBase& CmdLi
 	INTERNAL_DECORATOR_COMPUTE(RHISetUAVParameter)(Shader, UAVIndex, UAV, InitialCount);
 }
 
-template<> void FRHICommandSetShaderSampler<FRHIComputeShader>::Execute(FRHICommandListBase& CmdList)
+template<> RHI_API void FRHICommandSetShaderSampler<FRHIComputeShader>::Execute(FRHICommandListBase& CmdList)
 {
 	RHISTAT(SetShaderSampler);
 	INTERNAL_DECORATOR_COMPUTE(RHISetShaderSampler)(Shader, SamplerIndex, Sampler);
 }
 
-template<> void FRHICommandSetShaderSampler<FRHIGraphicsShader>::Execute(FRHICommandListBase& CmdList)
+template<> RHI_API void FRHICommandSetShaderSampler<FRHIGraphicsShader>::Execute(FRHICommandListBase& CmdList)
 {
 	RHISTAT(SetShaderSampler);
 	INTERNAL_DECORATOR(RHISetShaderSampler)(Shader, SamplerIndex, Sampler);
