Fixed RHI validation error with CopySrc being incorrectly dropped from depth stencil read state.

diff --git a/Engine/Source/Runtime/RHI/Private/RHIValidation.cpp b/Engine/Source/Runtime/RHI/Private/RHIValidation.cpp
--- a/Engine/Source/Runtime/RHI/Private/RHIValidation.cpp
+++ b/Engine/Source/Runtime/RHI/Private/RHIValidation.cpp
@@ -582,7 +582,8 @@ namespace RHIValidation
 				ERHIAccess::DSVRead |
 				ERHIAccess::DSVWrite |
 				ERHIAccess::SRVGraphics |
-				ERHIAccess::SRVCompute;
+				ERHIAccess::SRVCompute |
+				ERHIAccess::CopySrc;
 
 			return AccessMask & CompatibleStates;
 		}
