Fix for RHI validation assert on shutdown.

diff --git a/Engine/Source/Runtime/RHI/Private/RHIValidation.cpp b/Engine/Source/Runtime/RHI/Private/RHIValidation.cpp
--- a/Engine/Source/Runtime/RHI/Private/RHIValidation.cpp
+++ b/Engine/Source/Runtime/RHI/Private/RHIValidation.cpp
@@ -1537,7 +1537,7 @@ TLockFreePointerListUnordered<FValidationContext, PLATFORM_CACHE_LINE_SIZE> FVal
 
 FValidationTransientResourceAllocator::~FValidationTransientResourceAllocator()
 {
-	delete RHIAllocator;
+	checkf(!RHIAllocator, TEXT("Release was not called on FRHITransientResourceAllocator."));
 }
 
 FRHITransientTexture* FValidationTransientResourceAllocator::CreateTexture(const FRHITextureCreateInfo& InCreateInfo, const TCHAR* InDebugName, uint32 InPassIndex)
@@ -1663,6 +1663,7 @@ void FValidationTransientResourceAllocator::Release(FRHICommandListImmediate& RH
 	}
 
 	RHIAllocator->Release(RHICmdList);
+	RHIAllocator = nullptr;
 	delete this;
 }
 
