Use 10:10:10:2 backbuffer

diff --git a/Engine/Source/Runtime/RHI/Public/RHISurfaceDataConversion.h b/Engine/Source/Runtime/RHI/Public/RHISurfaceDataConversion.h
--- a/Engine/Source/Runtime/RHI/Public/RHISurfaceDataConversion.h
+++ b/Engine/Source/Runtime/RHI/Public/RHISurfaceDataConversion.h
@@ -134,6 +134,26 @@ static inline void ConvertRawR10G10B10A2DataToFColor(uint32 Width, uint32 Height
 	}
 }
 
+static inline void ConvertRawB10G10R10A2DataToFColor(uint32 Width, uint32 Height, uint8* In, uint32 SrcPitch, FColor* Out)
+{
+	for (uint32 Y = 0; Y < Height; Y++)
+	{
+		FRHIR10G10B10A2* SrcPtr = (FRHIR10G10B10A2*)(In + Y * SrcPitch);
+		FColor* DestPtr = Out + Y * Width;
+		for (uint32 X = 0; X < Width; X++)
+		{
+			*DestPtr = FColor::MakeRequantizeFrom1010102(
+				SrcPtr->B,
+				SrcPtr->G,
+				SrcPtr->R,
+				SrcPtr->A
+			);
+			++SrcPtr;
+			++DestPtr;
+		}
+	}
+}
+
 static inline void ConvertRawR16G16B16A16FDataToFColor(uint32 Width, uint32 Height, uint8 *In, uint32 SrcPitch, FColor* Out, bool LinearToGamma)
 {
 	FPlane	MinValue(0.0f, 0.0f, 0.0f, 0.0f),
