Fix crash when CopyToResolveTarget or CopyTexture were called with null textures.

#rnx

diff --git a/Engine/Source/Runtime/RHI/Public/RHICommandList.h b/Engine/Source/Runtime/RHI/Public/RHICommandList.h
--- a/Engine/Source/Runtime/RHI/Public/RHICommandList.h
+++ b/Engine/Source/Runtime/RHI/Public/RHICommandList.h
@@ -1515,10 +1515,11 @@ FRHICOMMAND_MACRO(FRHICommandCopyToResolveTarget)
 		, SourceTexture(InSourceTexture)
 		, DestTexture(InDestTexture)
 	{
-		ensure(SourceTexture);
-		ensure(DestTexture);
-		ensure(SourceTexture->GetTexture2D() || SourceTexture->GetTexture3D() || SourceTexture->GetTextureCube() || SourceTexture->GetTexture2DArray());
-		ensure(DestTexture->GetTexture2D() || DestTexture->GetTexture3D() || DestTexture->GetTextureCube() || DestTexture->GetTexture2DArray());
+		if (ensure(SourceTexture) && ensure(DestTexture))
+		{
+			ensure(SourceTexture->GetTexture2D() || SourceTexture->GetTexture3D() || SourceTexture->GetTextureCube() || SourceTexture->GetTexture2DArray());
+			ensure(DestTexture->GetTexture2D() || DestTexture->GetTexture3D() || DestTexture->GetTextureCube() || DestTexture->GetTexture2DArray());
+		}
 	}
 	RHI_API void Execute(FRHICommandListBase& CmdList);
 };
@@ -1534,10 +1535,11 @@ FRHICOMMAND_MACRO(FRHICommandCopyTexture)
 		, SourceTexture(InSourceTexture)
 		, DestTexture(InDestTexture)
 	{
-		ensure(SourceTexture);
-		ensure(DestTexture);
-		ensure(SourceTexture->GetTexture2D() || SourceTexture->GetTexture2DArray() || SourceTexture->GetTexture3D() || SourceTexture->GetTextureCube());
-		ensure(DestTexture->GetTexture2D() || DestTexture->GetTexture2DArray() || DestTexture->GetTexture3D() || DestTexture->GetTextureCube());
+		if (ensure(SourceTexture) && ensure(DestTexture))
+		{
+			ensure(SourceTexture->GetTexture2D() || SourceTexture->GetTexture2DArray() || SourceTexture->GetTexture3D() || SourceTexture->GetTextureCube());
+			ensure(DestTexture->GetTexture2D() || DestTexture->GetTexture2DArray() || DestTexture->GetTexture3D() || DestTexture->GetTextureCube());
+		}
 	}
 	RHI_API void Execute(FRHICommandListBase& CmdList);
 };
