Re-enable the incremental user PSO cache by having it save at shutdown.

Enable PSO cache CVars by default on Windows. This has no effect on D3D11 since the cache is disabled if GRHISupportsPipelineFileCache is false.

diff --git a/Engine/Source/Runtime/RHI/Private/HoloLens/HoloLensDynamicRHI.cpp b/Engine/Source/Runtime/RHI/Private/HoloLens/HoloLensDynamicRHI.cpp
--- a/Engine/Source/Runtime/RHI/Private/HoloLens/HoloLensDynamicRHI.cpp
+++ b/Engine/Source/Runtime/RHI/Private/HoloLens/HoloLensDynamicRHI.cpp
@@ -46,18 +46,6 @@ static IDynamicRHIModule* LoadDynamicRHIModule()
 	{
 		FApp::SetGraphicsRHI(TEXT("DirectX 12"));
 		DynamicRHIModule = FModuleManager::LoadModulePtr<IDynamicRHIModule>(TEXT("D3D12RHI"));
-
-#if !WITH_EDITOR
-		// Enable -psocache by default on DX12. Since RHI is selected at runtime we can't set this at compile time with PIPELINE_CACHE_DEFAULT_ENABLED.
-		auto PSOFileCacheEnabledCVar = IConsoleManager::Get().FindTConsoleVariableDataInt(TEXT("r.ShaderPipelineCache.Enabled"));
-		*PSOFileCacheEnabledCVar = 1;
-
-		auto PSOFileCacheReportCVar = IConsoleManager::Get().FindTConsoleVariableDataInt(TEXT("r.ShaderPipelineCache.ReportPSO"));
-		*PSOFileCacheReportCVar = 1;
-
-		auto PSOFileCacheUserCacheCVar = IConsoleManager::Get().FindTConsoleVariableDataInt(TEXT("r.ShaderPipelineCache.SaveUserCache"));
-		*PSOFileCacheUserCacheCVar = UE_BUILD_SHIPPING;
-#endif
 	}
 	else
 	{
diff --git a/Engine/Source/Runtime/RHI/Private/PipelineFileCache.cpp b/Engine/Source/Runtime/RHI/Private/PipelineFileCache.cpp
--- a/Engine/Source/Runtime/RHI/Private/PipelineFileCache.cpp
+++ b/Engine/Source/Runtime/RHI/Private/PipelineFileCache.cpp
@@ -2975,7 +2975,7 @@ bool FPipelineFileCache::IsPipelineFileCacheEnabled()
 		bCmdLineForce = FParse::Param(FCommandLine::Get(), TEXT("psocache"));
 		UE_CLOG(bCmdLineForce, LogRHI, Warning, TEXT("****************************** Forcing PSO cache from command line"));
 	}
-	return GRHISupportsPipelineFileCache && FileCacheEnabled && (bCmdLineForce || CVarPSOFileCacheEnabled.GetValueOnAnyThread() == 1);
+	return FileCacheEnabled && (bCmdLineForce || CVarPSOFileCacheEnabled.GetValueOnAnyThread() == 1);
 }
 
 bool FPipelineFileCache::LogPSOtoFileCache()
@@ -3001,7 +3001,7 @@ bool FPipelineFileCache::ReportNewPSOs()
         bCmdLineForce = FParse::Param(FCommandLine::Get(), TEXT("reportpso"));
         UE_CLOG(bCmdLineForce, LogRHI, Warning, TEXT("****************************** Forcing reporting of new PSOs from command line"));
     }
-    return (bCmdLineForce || CVarPSOFileCacheReportPSO.GetValueOnAnyThread() == 1);
+	return (bCmdLineForce || CVarPSOFileCacheReportPSO.GetValueOnAnyThread() == 1);
 }
 
 bool FPipelineFileCache::LogPSODetails()
@@ -3050,7 +3050,7 @@ bool FPipelineFileCache::ShouldEnableFileCache()
 		}
 	}
 #endif
-	return true;
+	return GRHISupportsPipelineFileCache;
 }
 
 void FPipelineFileCache::PreCompileComplete()
@@ -3456,7 +3456,6 @@ void FPipelineFileCache::CacheRayTracingPSO(const FRayTracingPipelineStateInitia
 	if (!IsPipelineFileCacheEnabled() || !(LogPSOtoFileCache() || ReportNewPSOs()))
 	{
 		return;
-
 	}
 
 	TArrayView<FRHIRayTracingShader*> ShaderTables[] =
diff --git a/Engine/Source/Runtime/RHI/Private/Windows/WindowsDynamicRHI.cpp b/Engine/Source/Runtime/RHI/Private/Windows/WindowsDynamicRHI.cpp
--- a/Engine/Source/Runtime/RHI/Private/Windows/WindowsDynamicRHI.cpp
+++ b/Engine/Source/Runtime/RHI/Private/Windows/WindowsDynamicRHI.cpp
@@ -375,23 +375,6 @@ static IDynamicRHIModule* LoadDynamicRHIModule(ERHIFeatureLevel::Type& DesiredFe
 			DynamicRHIModule = NULL;
 			LoadedRHIModuleName = nullptr;
 		}
-		else
-		{
-#if !WITH_EDITOR
-			// Enable -psocache by default on DX12. Since RHI is selected at runtime we can't set this at compile time with PIPELINE_CACHE_DEFAULT_ENABLED.
-			auto PSOFileCacheEnabledCVar = IConsoleManager::Get().FindTConsoleVariableDataInt(TEXT("r.ShaderPipelineCache.Enabled"));
-			*PSOFileCacheEnabledCVar = 1;
-
-			auto PSOFileCacheReportCVar = IConsoleManager::Get().FindTConsoleVariableDataInt(TEXT("r.ShaderPipelineCache.ReportPSO"));
-			*PSOFileCacheReportCVar = 1;
-
-			auto PSOFileCacheUserCacheCVar = IConsoleManager::Get().FindTConsoleVariableDataInt(TEXT("r.ShaderPipelineCache.SaveUserCache"));
-			*PSOFileCacheUserCacheCVar = UE_BUILD_SHIPPING;
-
-			auto PSOFileCacheLogCVar = IConsoleManager::Get().FindTConsoleVariableDataInt(TEXT("r.ShaderPipelineCache.LogPSO"));
-			*PSOFileCacheLogCVar = UE_BUILD_SHIPPING;
-#endif
-		}
 	}
 
 	// Fallback to D3D11RHI if nothing is selected
diff --git a/Engine/Source/Runtime/RHI/Public/PipelineFileCache.h b/Engine/Source/Runtime/RHI/Public/PipelineFileCache.h
--- a/Engine/Source/Runtime/RHI/Public/PipelineFileCache.h
+++ b/Engine/Source/Runtime/RHI/Public/PipelineFileCache.h
@@ -13,7 +13,7 @@ DECLARE_DWORD_ACCUMULATOR_STAT_EXTERN(TEXT("Total Graphics Pipeline State Count"
 DECLARE_DWORD_ACCUMULATOR_STAT_EXTERN(TEXT("Total Compute Pipeline State Count"), STAT_TotalComputePipelineStateCount, STATGROUP_PipelineStateCache, RHI_API);
 DECLARE_DWORD_ACCUMULATOR_STAT_EXTERN(TEXT("Total RayTracing Pipeline State Count"), STAT_TotalRayTracingPipelineStateCount, STATGROUP_PipelineStateCache, RHI_API);
 
-#define PIPELINE_CACHE_DEFAULT_ENABLED (!WITH_EDITOR && PLATFORM_MAC)
+#define PIPELINE_CACHE_DEFAULT_ENABLED (!WITH_EDITOR)
 
 /**
  * PSO_COOKONLY_DATA
diff --git a/Engine/Source/Runtime/RenderCore/Private/ShaderPipelineCache.cpp b/Engine/Source/Runtime/RenderCore/Private/ShaderPipelineCache.cpp
--- a/Engine/Source/Runtime/RenderCore/Private/ShaderPipelineCache.cpp
+++ b/Engine/Source/Runtime/RenderCore/Private/ShaderPipelineCache.cpp
@@ -1240,14 +1240,7 @@ FShaderPipelineCache::FShaderPipelineCache(EShaderPlatform Platform)
 
 FShaderPipelineCache::~FShaderPipelineCache()
 {
-	// Only save PSO Record / Logging at shutdown
-	if (GetShaderPipelineCacheSaveBoundPSOLog())
-	{
-		FShaderPipelineCache::SavePipelineFileCache(FPipelineFileCache::SaveMode::BoundPSOsOnly);
-	}
-	
-	// Close with shutdown flag
-	Close(true);
+	Close();
 	
 	// The render thread tick should be dead now and we are safe to destroy everything that needs to wait or manual destruction
 
@@ -1598,7 +1591,7 @@ bool FShaderPipelineCache::Save(FPipelineFileCache::SaveMode Mode)
 	return bOK;
 }
 
-void FShaderPipelineCache::Close(bool bShuttingDown)
+void FShaderPipelineCache::Close()
 {
 	FScopeLock Lock(&Mutex);
 		
@@ -1614,8 +1607,7 @@ void FShaderPipelineCache::Close(bool bShuttingDown)
 		Save(FPipelineFileCache::SaveMode::BoundPSOsOnly);
 	}
 	
-	// Force a fast save, just in case - but not at shutdown / destruction
-	if (GetPSOFileCacheSaveUserCache() && !bShuttingDown)
+	if (GetPSOFileCacheSaveUserCache())
 	{
 		Save(FPipelineFileCache::SaveMode::Incremental);
 	}
diff --git a/Engine/Source/Runtime/RenderCore/Public/ShaderPipelineCache.h b/Engine/Source/Runtime/RenderCore/Public/ShaderPipelineCache.h
--- a/Engine/Source/Runtime/RenderCore/Public/ShaderPipelineCache.h
+++ b/Engine/Source/Runtime/RenderCore/Public/ShaderPipelineCache.h
@@ -210,7 +210,7 @@ public:
 private:
 	bool Open(FString const& Name, EShaderPlatform Platform);
 	bool Save(FPipelineFileCache::SaveMode Mode);
-	void Close(bool bShuttingDown = false);
+	void Close();
 
 	bool Precompile(FRHICommandListImmediate& RHICmdList, EShaderPlatform Platform, FPipelineCacheFileFormatPSO const& PSO);
 	void PreparePipelineBatch(TDoubleLinkedList<FPipelineCacheFileFormatPSORead*>& PipelineBatch);
