Adding depth clip control to the rasterizer state.

diff --git a/Engine/Source/Runtime/Apple/AGXRHI/Private/AGXCommandEncoder.cpp b/Engine/Source/Runtime/Apple/AGXRHI/Private/AGXCommandEncoder.cpp
--- a/Engine/Source/Runtime/Apple/AGXRHI/Private/AGXCommandEncoder.cpp
+++ b/Engine/Source/Runtime/Apple/AGXRHI/Private/AGXCommandEncoder.cpp
@@ -825,6 +825,14 @@ void FAGXCommandEncoder::SetTriangleFillMode(mtlpp::TriangleFillMode const InFil
 	}
 }
 
+void FAGXCommandEncoder::SetDepthClipMode(mtlpp::DepthClipMode const InDepthClipMode)
+{
+	check(RenderCommandEncoder);
+	{
+		RenderCommandEncoder.SetDepthClipMode(InDepthClipMode);
+	}
+}
+
 void FAGXCommandEncoder::SetBlendColor(float const Red, float const Green, float const Blue, float const Alpha)
 {
 	check(RenderCommandEncoder);
diff --git a/Engine/Source/Runtime/Apple/AGXRHI/Private/AGXCommandEncoder.h b/Engine/Source/Runtime/Apple/AGXRHI/Private/AGXCommandEncoder.h
--- a/Engine/Source/Runtime/Apple/AGXRHI/Private/AGXCommandEncoder.h
+++ b/Engine/Source/Runtime/Apple/AGXRHI/Private/AGXCommandEncoder.h
@@ -256,6 +256,12 @@ public:
 	 * @param FillMode The fill mode.
 	 */
 	void SetTriangleFillMode(mtlpp::TriangleFillMode const FillMode);
+
+	/*
+	 * Set wether to clip or clamp triangles based on depth.
+	 * @param FillMode The fill mode.
+	 */
+	void SetDepthClipMode(mtlpp::DepthClipMode DepthClipMode);
 	
 	/*
 	 * Set the constant blend color used across all blending on all render targets
diff --git a/Engine/Source/Runtime/Apple/AGXRHI/Private/AGXStateCache.cpp b/Engine/Source/Runtime/Apple/AGXRHI/Private/AGXStateCache.cpp
--- a/Engine/Source/Runtime/Apple/AGXRHI/Private/AGXStateCache.cpp
+++ b/Engine/Source/Runtime/Apple/AGXRHI/Private/AGXStateCache.cpp
@@ -37,6 +37,16 @@ static mtlpp::CullMode TranslateCullMode(ERasterizerCullMode CullMode)
 	}
 }
 
+static mtlpp::DepthClipMode TranslateDepthClipMode(ERasterizerDepthClipMode DepthClipMode)
+{
+	switch (DepthClipMode)
+	{
+	case ERasterizerDepthClipMode::DepthClip:	return mtlpp::DepthClipMode::Clip;
+	case ERasterizerDepthClipMode::DepthClamp:	return mtlpp::DepthClipMode::Clamp;
+	default:									return mtlpp::DepthClipMode::Clip;
+	}
+}
+
 FORCEINLINE mtlpp::StoreAction GetMetalRTStoreAction(ERenderTargetStoreAction StoreAction)
 {
 	switch(StoreAction)
@@ -383,7 +393,7 @@ void FAGXStateCache::SetRasterizerState(FAGXRasterizerState* InRasterizerState)
 	if(RasterizerState != InRasterizerState)
 	{
 		RasterizerState = InRasterizerState;
-		RasterBits |= EAGXRenderFlagFrontFacingWinding|EAGXRenderFlagCullMode|EAGXRenderFlagDepthBias|EAGXRenderFlagTriangleFillMode;
+		RasterBits |= EAGXRenderFlagFrontFacingWinding|EAGXRenderFlagCullMode|EAGXRenderFlagDepthBias|EAGXRenderFlagTriangleFillMode|EAGXRenderFlagDepthClipMode;
 	}
 }
 
@@ -1945,6 +1955,11 @@ void FAGXStateCache::SetRenderState(FAGXCommandEncoder& CommandEncoder, FAGXComm
             	VisibilityWritten = VisibilityOffset + FAGXQueryBufferPool::EQueryResultMaxSize;
 			}
         }
+		if (RasterBits & EAGXRenderFlagDepthClipMode)
+		{
+			check(IsValidRef(RasterizerState));
+			CommandEncoder.SetDepthClipMode(TranslateDepthClipMode(RasterizerState->State.DepthClipMode));
+		}
 		RasterBits = 0;
 	}
 }
diff --git a/Engine/Source/Runtime/Apple/AGXRHI/Private/AGXStateCache.h b/Engine/Source/Runtime/Apple/AGXRHI/Private/AGXStateCache.h
--- a/Engine/Source/Runtime/Apple/AGXRHI/Private/AGXStateCache.h
+++ b/Engine/Source/Runtime/Apple/AGXRHI/Private/AGXStateCache.h
@@ -32,7 +32,8 @@ enum EAGXRenderFlags
     EAGXRenderFlagDepthStencilState = 1 << 7,
     EAGXRenderFlagStencilReferenceValue = 1 << 8,
     EAGXRenderFlagVisibilityResultMode = 1 << 9,
-    EAGXRenderFlagMask = 0x1FF
+	EAGXRenderFlagDepthClipMode = 1 << 10,
+    EAGXRenderFlagMask = 0x3FF
 };
 
 class FAGXStateCache
diff --git a/Engine/Source/Runtime/Apple/MetalRHI/Private/MetalCommandEncoder.cpp b/Engine/Source/Runtime/Apple/MetalRHI/Private/MetalCommandEncoder.cpp
--- a/Engine/Source/Runtime/Apple/MetalRHI/Private/MetalCommandEncoder.cpp
+++ b/Engine/Source/Runtime/Apple/MetalRHI/Private/MetalCommandEncoder.cpp
@@ -1363,6 +1363,14 @@ void FMetalCommandEncoder::SetTriangleFillMode(mtlpp::TriangleFillMode const InF
 	}
 }
 
+void FMetalCommandEncoder::SetDepthClipMode(mtlpp::DepthClipMode const InDepthClipMode)
+{
+	check(RenderCommandEncoder);
+	{
+		RenderCommandEncoder.SetDepthClipMode(InDepthClipMode);
+	}
+}
+
 void FMetalCommandEncoder::SetBlendColor(float const Red, float const Green, float const Blue, float const Alpha)
 {
 	check(RenderCommandEncoder);
diff --git a/Engine/Source/Runtime/Apple/MetalRHI/Private/MetalCommandEncoder.h b/Engine/Source/Runtime/Apple/MetalRHI/Private/MetalCommandEncoder.h
--- a/Engine/Source/Runtime/Apple/MetalRHI/Private/MetalCommandEncoder.h
+++ b/Engine/Source/Runtime/Apple/MetalRHI/Private/MetalCommandEncoder.h
@@ -357,6 +357,12 @@ public:
 	 * @param FillMode The fill mode.
 	 */
 	void SetTriangleFillMode(mtlpp::TriangleFillMode const FillMode);
+
+	/*
+	 * Set wether to clip or clamp triangles based on depth.
+	 * @param FillMode The fill mode.
+	 */
+	void SetDepthClipMode(mtlpp::DepthClipMode DepthClipMode);
 	
 	/*
 	 * Set the constant blend color used across all blending on all render targets
diff --git a/Engine/Source/Runtime/Apple/MetalRHI/Private/MetalStateCache.cpp b/Engine/Source/Runtime/Apple/MetalRHI/Private/MetalStateCache.cpp
--- a/Engine/Source/Runtime/Apple/MetalRHI/Private/MetalStateCache.cpp
+++ b/Engine/Source/Runtime/Apple/MetalRHI/Private/MetalStateCache.cpp
@@ -38,6 +38,16 @@ static mtlpp::CullMode TranslateCullMode(ERasterizerCullMode CullMode)
 	}
 }
 
+static mtlpp::DepthClipMode TranslateDepthClipMode(ERasterizerDepthClipMode DepthClipMode)
+{
+	switch (DepthClipMode)
+	{
+	case ERasterizerDepthClipMode::DepthClip:	return mtlpp::DepthClipMode::Clip;
+	case ERasterizerDepthClipMode::DepthClamp:	return mtlpp::DepthClipMode::Clamp;
+	default:									return mtlpp::DepthClipMode::Clip;
+	}
+}
+
 FORCEINLINE mtlpp::StoreAction GetMetalRTStoreAction(ERenderTargetStoreAction StoreAction)
 {
 	switch(StoreAction)
@@ -384,7 +394,7 @@ void FMetalStateCache::SetRasterizerState(FMetalRasterizerState* InRasterizerSta
 	if(RasterizerState != InRasterizerState)
 	{
 		RasterizerState = InRasterizerState;
-		RasterBits |= EMetalRenderFlagFrontFacingWinding|EMetalRenderFlagCullMode|EMetalRenderFlagDepthBias|EMetalRenderFlagTriangleFillMode;
+		RasterBits |= EMetalRenderFlagFrontFacingWinding|EMetalRenderFlagCullMode|EMetalRenderFlagDepthBias|EMetalRenderFlagTriangleFillMode|EMetalRenderFlagDepthClipMode;
 	}
 }
 
@@ -1970,7 +1980,12 @@ void FMetalStateCache::SetRenderState(FMetalCommandEncoder& CommandEncoder, FMet
 			{
             	VisibilityWritten = VisibilityOffset + FMetalQueryBufferPool::EQueryResultMaxSize;
 			}
-        }
+		}
+		if (RasterBits & EMetalRenderFlagDepthClipMode)
+		{
+			check(IsValidRef(RasterizerState));
+			CommandEncoder.SetDepthClipMode(TranslateDepthClipMode(RasterizerState->State.DepthClipMode));
+		}
 		RasterBits = 0;
 	}
 }
diff --git a/Engine/Source/Runtime/Apple/MetalRHI/Private/MetalStateCache.h b/Engine/Source/Runtime/Apple/MetalRHI/Private/MetalStateCache.h
--- a/Engine/Source/Runtime/Apple/MetalRHI/Private/MetalStateCache.h
+++ b/Engine/Source/Runtime/Apple/MetalRHI/Private/MetalStateCache.h
@@ -32,7 +32,8 @@ enum EMetalRenderFlags
     EMetalRenderFlagDepthStencilState = 1 << 7,
     EMetalRenderFlagStencilReferenceValue = 1 << 8,
     EMetalRenderFlagVisibilityResultMode = 1 << 9,
-    EMetalRenderFlagMask = 0x1FF
+	EMetalRenderFlagDepthClipMode = 1 << 10,
+    EMetalRenderFlagMask = 0x3FF
 };
 
 class FMetalStateCache
diff --git a/Engine/Source/Runtime/D3D12RHI/Private/D3D12State.cpp b/Engine/Source/Runtime/D3D12RHI/Private/D3D12State.cpp
--- a/Engine/Source/Runtime/D3D12RHI/Private/D3D12State.cpp
+++ b/Engine/Source/Runtime/D3D12RHI/Private/D3D12State.cpp
@@ -309,7 +309,7 @@ FRasterizerStateRHIRef FD3D12DynamicRHI::RHICreateRasterizerState(const FRasteri
 	RasterizerDesc.SlopeScaledDepthBias = Initializer.SlopeScaleDepthBias;
 	RasterizerDesc.FrontCounterClockwise = true;
 	RasterizerDesc.DepthBias = FMath::FloorToInt(Initializer.DepthBias * (float)(1 << 24));
-	RasterizerDesc.DepthClipEnable = true;
+	RasterizerDesc.DepthClipEnable = Initializer.DepthClipMode == ERasterizerDepthClipMode::DepthClip;
 	RasterizerDesc.MultisampleEnable = Initializer.bAllowMSAA;
 
 	return RasterizerState;
diff --git a/Engine/Source/Runtime/Engine/Private/BatchedElements.cpp b/Engine/Source/Runtime/Engine/Private/BatchedElements.cpp
--- a/Engine/Source/Runtime/Engine/Private/BatchedElements.cpp
+++ b/Engine/Source/Runtime/Engine/Private/BatchedElements.cpp
@@ -875,7 +875,7 @@ bool FBatchedElements::Draw(FRHICommandList& RHICmdList, const FMeshPassProcesso
 
 					const bool bEnableMSAA = true;
 					const bool bEnableLineAA = false;
-					FRasterizerStateInitializerRHI Initializer = { FM_Solid, CM_None, 0, DepthBiasThisBatch, bEnableMSAA, bEnableLineAA };
+					FRasterizerStateInitializerRHI Initializer(FM_Solid, CM_None, 0, DepthBiasThisBatch, ERasterizerDepthClipMode::DepthClip, bEnableMSAA, bEnableLineAA);
 					auto RasterState = RHICreateRasterizerState(Initializer);
 					GraphicsPSOInit.RasterizerState = RasterState.GetReference();
 					PrepareShaders(RHICmdList, GraphicsPSOInit, StencilRef, FeatureLevel, SE_BLEND_AlphaBlend, RelativeMatrices, bNeedToSwitchVerticalAxis, BatchedElementParameters, GWhiteTexture, bHitTesting, Gamma, NULL, &View);
@@ -973,7 +973,7 @@ bool FBatchedElements::Draw(FRHICommandList& RHICmdList, const FMeshPassProcesso
 
 				const bool bEnableMSAA = true;
 				const bool bEnableLineAA = false;
-				FRasterizerStateInitializerRHI Initializer = { FM_Wireframe, CM_None, 0, 0, bEnableMSAA, bEnableLineAA };
+				FRasterizerStateInitializerRHI Initializer(FM_Wireframe, CM_None, bEnableMSAA, bEnableLineAA);
 
 				int32 MaxVerticesAllowed = ((GDrawUPVertexCheckCount / sizeof(FSimpleElementVertex)) / 3) * 3;
 				/*
diff --git a/Engine/Source/Runtime/OpenGLDrv/Private/OpenGLCommands.cpp b/Engine/Source/Runtime/OpenGLDrv/Private/OpenGLCommands.cpp
--- a/Engine/Source/Runtime/OpenGLDrv/Private/OpenGLCommands.cpp
+++ b/Engine/Source/Runtime/OpenGLDrv/Private/OpenGLCommands.cpp
@@ -271,6 +271,19 @@ void FOpenGLDynamicRHI::UpdateRasterizerStateInOpenGLContext( FOpenGLContextStat
 		ContextState.RasterizerState.CullMode = PendingState.RasterizerState.CullMode;
 	}
 
+	if (FOpenGL::SupportsDepthClamp() && ContextState.RasterizerState.DepthClipMode != PendingState.RasterizerState.DepthClipMode)
+	{
+		if (PendingState.RasterizerState.DepthClipMode == ERasterizerDepthClipMode::DepthClamp)
+		{
+			glEnable(GL_DEPTH_CLAMP);
+		}
+		else
+		{
+			glDisable(GL_DEPTH_CLAMP);
+		}
+		ContextState.RasterizerState.DepthClipMode = PendingState.RasterizerState.DepthClipMode;
+	}
+
 	// Convert our platform independent depth bias into an OpenGL depth bias.
 	const float BiasScale = float((1<<24)-1);	// Warning: this assumes depth bits == 24, and won't be correct with 32.
 	float DepthBias = PendingState.RasterizerState.DepthBias * BiasScale;
diff --git a/Engine/Source/Runtime/OpenGLDrv/Private/OpenGLES.cpp b/Engine/Source/Runtime/OpenGLDrv/Private/OpenGLES.cpp
--- a/Engine/Source/Runtime/OpenGLDrv/Private/OpenGLES.cpp
+++ b/Engine/Source/Runtime/OpenGLDrv/Private/OpenGLES.cpp
@@ -127,6 +127,9 @@ bool FOpenGLES::bRequiresARMShaderFramebufferFetchDepthStencilUndef = false;
 /** GL_EXT_buffer_storage */
 bool FOpenGLES::bSupportsBufferStorage = false;
 
+/** GL_EXT_depth_clamp */
+bool FOpenGLES::bSupportsDepthClamp = false;
+
 bool FOpenGLES::bHasHardwareHiddenSurfaceRemoval = false;
 bool FOpenGLES::bSupportsMobileMultiView = false;
 GLint FOpenGLES::MaxMSAASamplesTileMem = 1;
@@ -222,6 +225,7 @@ void FOpenGLES::ProcessExtensions(const FString& ExtensionsString)
 	bSupportsDXT = ExtensionsString.Contains(TEXT("GL_NV_texture_compression_s3tc")) || ExtensionsString.Contains(TEXT("GL_EXT_texture_compression_s3tc"));
 	bSupportsNVFrameBufferBlit = ExtensionsString.Contains(TEXT("GL_NV_framebuffer_blit"));
 	bSupportsBufferStorage = ExtensionsString.Contains(TEXT("GL_EXT_buffer_storage"));
+	bSupportsDepthClamp = ExtensionsString.Contains(TEXT("GL_EXT_depth_clamp"));
 	bSupportsASTCDecodeMode = ExtensionsString.Contains(TEXT("GL_EXT_texture_compression_astc_decode_mode"));
 
 	// Report shader precision
diff --git a/Engine/Source/Runtime/OpenGLDrv/Private/OpenGLState.cpp b/Engine/Source/Runtime/OpenGLDrv/Private/OpenGLState.cpp
--- a/Engine/Source/Runtime/OpenGLDrv/Private/OpenGLState.cpp
+++ b/Engine/Source/Runtime/OpenGLDrv/Private/OpenGLState.cpp
@@ -348,6 +348,7 @@ FRasterizerStateRHIRef FOpenGLDynamicRHI::RHICreateRasterizerState(const FRaster
 	RasterizerState->Data.FillMode = TranslateFillMode(Initializer.FillMode);
 	RasterizerState->Data.DepthBias = Initializer.DepthBias;
 	RasterizerState->Data.SlopeScaleDepthBias = Initializer.SlopeScaleDepthBias;
+	RasterizerState->Data.DepthClipMode = Initializer.DepthClipMode;
 	
 	return RasterizerState;
 }
@@ -359,6 +360,7 @@ bool FOpenGLRasterizerState::GetInitializer(FRasterizerStateInitializerRHI& Init
 	Init.FillMode = TranslateFillMode(Data.FillMode);
 	Init.DepthBias = Data.DepthBias;
 	Init.SlopeScaleDepthBias = Data.SlopeScaleDepthBias;
+	Init.DepthClipMode = Data.DepthClipMode;
 	return true;
 }
 
diff --git a/Engine/Source/Runtime/OpenGLDrv/Public/OpenGL.h b/Engine/Source/Runtime/OpenGLDrv/Public/OpenGL.h
--- a/Engine/Source/Runtime/OpenGLDrv/Public/OpenGL.h
+++ b/Engine/Source/Runtime/OpenGLDrv/Public/OpenGL.h
@@ -135,6 +135,7 @@ public:
 	static FORCEINLINE bool AmdWorkaround()								{ return false; }
 	static FORCEINLINE bool SupportsSeparateShaderObjects()				{ return false; }
 	static FORCEINLINE bool SupportsProgramBinary()						{ return false; }
+	static FORCEINLINE bool SupportsDepthClamp()						{ return true; }
 	
 	static FORCEINLINE bool SupportsASTCDecodeMode()					{ return false; }
 
diff --git a/Engine/Source/Runtime/OpenGLDrv/Public/OpenGLES.h b/Engine/Source/Runtime/OpenGLDrv/Public/OpenGLES.h
--- a/Engine/Source/Runtime/OpenGLDrv/Public/OpenGLES.h
+++ b/Engine/Source/Runtime/OpenGLDrv/Public/OpenGLES.h
@@ -97,6 +97,9 @@ typedef GLfloat GLdouble;
 #ifndef GL_MAX_TEXTURE_BUFFER_SIZE
 #define GL_MAX_TEXTURE_BUFFER_SIZE 0x8C2B
 #endif
+#ifndef GL_DEPTH_CLAMP
+#define GL_DEPTH_CLAMP 0x864F
+#endif
 
 /** For the shader stage bits that don't exist just use 0 */
 #define GL_GEOMETRY_SHADER_BIT				0x00000000
@@ -217,6 +220,7 @@ struct FOpenGLES : public FOpenGLBase
 	static FORCEINLINE bool SupportsRGB10A2() { return bSupportsRGB10A2; }
 	static FORCEINLINE bool SupportsDrawIndirect() { return true; }
 	static FORCEINLINE bool SupportsBufferStorage() { return bSupportsBufferStorage; }
+	static FORCEINLINE bool SupportsDepthClamp() { return bSupportsDepthClamp; }
 	
 
 	static FORCEINLINE bool HasBinaryProgramRetrievalFailed() { return bBinaryProgramRetrievalFailed; }
@@ -835,6 +839,9 @@ public:
 	/** GL_EXT_buffer_storage */
 	static bool bSupportsBufferStorage;
 
+	/** GL_EXT_depth_clamp */
+	static bool bSupportsDepthClamp;
+
 	enum class EFeatureLevelSupport : uint8
 	{
 		Invalid,	// no feature level has yet been determined
diff --git a/Engine/Source/Runtime/OpenGLDrv/Public/OpenGLState.h b/Engine/Source/Runtime/OpenGLDrv/Public/OpenGLState.h
--- a/Engine/Source/Runtime/OpenGLDrv/Public/OpenGLState.h
+++ b/Engine/Source/Runtime/OpenGLDrv/Public/OpenGLState.h
@@ -75,18 +75,11 @@ public:
 
 struct FOpenGLRasterizerStateData
 {
-	GLenum FillMode;
-	GLenum CullMode;
-	float DepthBias;
-	float SlopeScaleDepthBias;
-
-	FOpenGLRasterizerStateData()
-		: FillMode(GL_FILL)
-		, CullMode(GL_NONE)
-		, DepthBias(0.0f)
-		, SlopeScaleDepthBias(0.0f)
-	{
-	}
+	GLenum FillMode = GL_FILL;
+	GLenum CullMode = GL_NONE;
+	float DepthBias = 0.0f;
+	float SlopeScaleDepthBias = 0.0f;
+	ERasterizerDepthClipMode DepthClipMode = ERasterizerDepthClipMode::DepthClip;
 };
 
 class FOpenGLRasterizerState : public FRHIRasterizerState
diff --git a/Engine/Source/Runtime/RHI/Private/PipelineFileCache.cpp b/Engine/Source/Runtime/RHI/Private/PipelineFileCache.cpp
--- a/Engine/Source/Runtime/RHI/Private/PipelineFileCache.cpp
+++ b/Engine/Source/Runtime/RHI/Private/PipelineFileCache.cpp
@@ -76,12 +76,13 @@ enum class EPipelineCacheFileFormatVersions : uint32
 	LastUsedTime = 22,
 	MoreRenderTargetFlags = 23,
 	FragmentDensityAttachment = 24,
+	AddingDepthClipMode = 25,
 };
 
 const uint64 FPipelineCacheFileFormatMagic = 0x5049504543414348; // PIPECACH
 const uint64 FPipelineCacheTOCFileFormatMagic = 0x544F435354415232; // TOCSTAR2
 const uint64 FPipelineCacheEOFFileFormatMagic = 0x454F462D4D41524B; // EOF-MARK
-const RHI_API uint32 FPipelineCacheFileFormatCurrentVersion = (uint32)EPipelineCacheFileFormatVersions::FragmentDensityAttachment;
+const RHI_API uint32 FPipelineCacheFileFormatCurrentVersion = (uint32)EPipelineCacheFileFormatVersions::AddingDepthClipMode;
 const int32  FPipelineCacheGraphicsDescPartsNum = 66; // parser will expect this number of parts in a description string
 
 /**
@@ -371,6 +372,7 @@ FString FPipelineFileCacheRasterizerState::ToString() const
 		, SlopeScaleDepthBias
 		, uint32(FillMode)
 		, uint32(CullMode)
+		, uint32(DepthClipMode)
 		, uint32(!!bAllowMSAA)
 		, uint32(!!bEnableLineAA)
 	);
@@ -391,6 +393,7 @@ void FPipelineFileCacheRasterizerState::FromString(const FStringView& Src)
 	LexFromString(SlopeScaleDepthBias, *PartIt++);
 	LexFromString((uint8&)FillMode, *PartIt++);
 	LexFromString((uint8&)CullMode, *PartIt++);
+	LexFromString((uint8&)DepthClipMode, *PartIt++);
 	LexFromString((uint8&)bAllowMSAA, *PartIt++);
 	LexFromString((uint8&)bEnableLineAA, *PartIt++);
 
diff --git a/Engine/Source/Runtime/RHI/Private/RHI.cpp b/Engine/Source/Runtime/RHI/Private/RHI.cpp
--- a/Engine/Source/Runtime/RHI/Private/RHI.cpp
+++ b/Engine/Source/Runtime/RHI/Private/RHI.cpp
@@ -375,6 +375,7 @@ uint32 GetTypeHash(const FRasterizerStateInitializerRHI& Initializer)
 	Hash = HashCombine(Hash, GetTypeHash(Initializer.CullMode));
 	Hash = HashCombine(Hash, GetTypeHash(Initializer.DepthBias));
 	Hash = HashCombine(Hash, GetTypeHash(Initializer.SlopeScaleDepthBias));
+	Hash = HashCombine(Hash, GetTypeHash(Initializer.DepthClipMode));
 	Hash = HashCombine(Hash, GetTypeHash(Initializer.bAllowMSAA));
 	Hash = HashCombine(Hash, GetTypeHash(Initializer.bEnableLineAA));
 	return Hash;
@@ -386,7 +387,8 @@ bool operator== (const FRasterizerStateInitializerRHI& A, const FRasterizerState
 		A.FillMode == B.FillMode && 
 		A.CullMode == B.CullMode && 
 		A.DepthBias == B.DepthBias && 
-		A.SlopeScaleDepthBias == B.SlopeScaleDepthBias && 
+		A.SlopeScaleDepthBias == B.SlopeScaleDepthBias &&
+		A.DepthClipMode == B.DepthClipMode &&
 		A.bAllowMSAA == B.bAllowMSAA && 
 		A.bEnableLineAA == B.bEnableLineAA;
 	return bSame;
diff --git a/Engine/Source/Runtime/RHI/Public/PipelineFileCache.h b/Engine/Source/Runtime/RHI/Public/PipelineFileCache.h
--- a/Engine/Source/Runtime/RHI/Public/PipelineFileCache.h
+++ b/Engine/Source/Runtime/RHI/Public/PipelineFileCache.h
@@ -32,6 +32,7 @@ struct RHI_API FPipelineFileCacheRasterizerState
 	float SlopeScaleDepthBias;
 	TEnumAsByte<ERasterizerFillMode> FillMode;
 	TEnumAsByte<ERasterizerCullMode> CullMode;
+	ERasterizerDepthClipMode DepthClipMode;
 	bool bAllowMSAA;
 	bool bEnableLineAA;
 	
@@ -41,14 +42,15 @@ struct RHI_API FPipelineFileCacheRasterizerState
 		SlopeScaleDepthBias = Other.SlopeScaleDepthBias;
 		FillMode = Other.FillMode;
 		CullMode = Other.CullMode;
+		DepthClipMode = Other.DepthClipMode;
 		bAllowMSAA = Other.bAllowMSAA;
 		bEnableLineAA = Other.bEnableLineAA;
 		return *this;
 	}
 	
-	operator FRasterizerStateInitializerRHI () const
+	operator FRasterizerStateInitializerRHI() const
 	{
-		FRasterizerStateInitializerRHI Initializer = {FillMode, CullMode, DepthBias, SlopeScaleDepthBias, bAllowMSAA, bEnableLineAA};
+		FRasterizerStateInitializerRHI Initializer(FillMode, CullMode, DepthBias, SlopeScaleDepthBias, DepthClipMode, bAllowMSAA, bEnableLineAA);
 		return Initializer;
 	}
 	
@@ -58,18 +60,20 @@ struct RHI_API FPipelineFileCacheRasterizerState
 		Ar << RasterizerStateInitializer.SlopeScaleDepthBias;
 		Ar << RasterizerStateInitializer.FillMode;
 		Ar << RasterizerStateInitializer.CullMode;
+		Ar << RasterizerStateInitializer.DepthClipMode;
 		Ar << RasterizerStateInitializer.bAllowMSAA;
 		Ar << RasterizerStateInitializer.bEnableLineAA;
 		return Ar;
 	}
-	
+
 	friend RHI_API uint32 GetTypeHash(const FPipelineFileCacheRasterizerState &Key)
 	{
 		uint32 KeyHash = (*((uint32*)&Key.DepthBias) ^ *((uint32*)&Key.SlopeScaleDepthBias));
 		KeyHash ^= (Key.FillMode << 8);
 		KeyHash ^= Key.CullMode;
-		KeyHash ^= Key.bAllowMSAA ? 2 : 0;
-		KeyHash ^= Key.bEnableLineAA ? 1 : 0;
+		KeyHash ^= Key.DepthClipMode == ERasterizerDepthClipMode::DepthClamp ? 0x951f4c3b : 0; // crc32 "DepthClamp"
+		KeyHash ^= Key.bAllowMSAA ? 0x694ea601 : 0; // crc32 "bAllowMSAA"
+		KeyHash ^= Key.bEnableLineAA ? 0x48271d01 : 0; // crc32 "bEnableLineAA"
 		return KeyHash;
 	}
 	FString ToString() const;
diff --git a/Engine/Source/Runtime/RHI/Public/RHI.h b/Engine/Source/Runtime/RHI/Public/RHI.h
--- a/Engine/Source/Runtime/RHI/Public/RHI.h
+++ b/Engine/Source/Runtime/RHI/Public/RHI.h
@@ -1327,19 +1327,42 @@ struct FSamplerStateInitializerRHI
 
 struct FRasterizerStateInitializerRHI
 {
-	TEnumAsByte<ERasterizerFillMode> FillMode;
-	TEnumAsByte<ERasterizerCullMode> CullMode;
-	float DepthBias;
-	float SlopeScaleDepthBias;
-	bool bAllowMSAA;
-	bool bEnableLineAA;
-	
+	TEnumAsByte<ERasterizerFillMode> FillMode = FM_Point;
+	TEnumAsByte<ERasterizerCullMode> CullMode = CM_None;
+	float DepthBias = 0.0f;
+	float SlopeScaleDepthBias = 0.0f;
+	ERasterizerDepthClipMode DepthClipMode = ERasterizerDepthClipMode::DepthClip;
+	bool bAllowMSAA = false;
+	bool bEnableLineAA = false;
+
+	FRasterizerStateInitializerRHI() = default;
+
+	FRasterizerStateInitializerRHI(ERasterizerFillMode InFillMode, ERasterizerCullMode InCullMode, bool bInAllowMSAA, bool bInEnableLineAA)
+		: FillMode(InFillMode)
+		, CullMode(InCullMode)
+		, bAllowMSAA(bInAllowMSAA)
+		, bEnableLineAA(bInEnableLineAA)
+	{
+	}
+
+	FRasterizerStateInitializerRHI(ERasterizerFillMode InFillMode, ERasterizerCullMode InCullMode, float InDepthBias, float InSlopeScaleDepthBias, ERasterizerDepthClipMode InDepthClipMode, bool bInAllowMSAA, bool bInEnableLineAA)
+		: FillMode(InFillMode)
+		, CullMode(InCullMode)
+		, DepthBias(InDepthBias)
+		, SlopeScaleDepthBias(InSlopeScaleDepthBias)
+		, DepthClipMode(InDepthClipMode)
+		, bAllowMSAA(bInAllowMSAA)
+		, bEnableLineAA(bInEnableLineAA)
+	{
+	}
+
 	friend FArchive& operator<<(FArchive& Ar,FRasterizerStateInitializerRHI& RasterizerStateInitializer)
 	{
 		Ar << RasterizerStateInitializer.FillMode;
 		Ar << RasterizerStateInitializer.CullMode;
 		Ar << RasterizerStateInitializer.DepthBias;
 		Ar << RasterizerStateInitializer.SlopeScaleDepthBias;
+		Ar << RasterizerStateInitializer.DepthClipMode;
 		Ar << RasterizerStateInitializer.bAllowMSAA;
 		Ar << RasterizerStateInitializer.bEnableLineAA;
 		return Ar;
@@ -1917,6 +1940,18 @@ inline constexpr bool IsValidAccess(ERHIAccess Access)
 	return !IsInvalidAccess(Access);
 }
 
+enum class ERHITransitionAccessMode
+{
+	// All upcoming rhi transitions can still have Unkown before state - rhi might need to do internal tracking if needed
+	AllowUnknown = 0,
+
+	// All upcoming rhi transitions are fully defined - no internal RHI tracking required anymore
+	DisallowUnknown,
+
+	Default = AllowUnknown
+};
+ENUM_CLASS_FLAGS(ERHITransitionAccessMode);
+
 enum class ERHITransitionCreateFlags
 {
 	None = 0,
diff --git a/Engine/Source/Runtime/RHI/Public/RHIDefinitions.h b/Engine/Source/Runtime/RHI/Public/RHIDefinitions.h
--- a/Engine/Source/Runtime/RHI/Public/RHIDefinitions.h
+++ b/Engine/Source/Runtime/RHI/Public/RHIDefinitions.h
@@ -1054,6 +1054,16 @@ enum ERasterizerCullMode
 };
 static_assert(ERasterizerCullMode_Num <= (1 << ERasterizerCullMode_NumBits), "ERasterizerCullMode_Num will not fit on ERasterizerCullMode_NumBits");
 
+enum class ERasterizerDepthClipMode : uint8
+{
+	DepthClip,
+	DepthClamp,
+
+	Num,
+	NumBits = 1,
+};
+static_assert(uint32(ERasterizerDepthClipMode::Num) <= (1U << uint32(ERasterizerDepthClipMode::NumBits)), "ERasterizerDepthClipMode::Num will not fit on ERasterizerDepthClipMode::NumBits");
+
 enum EColorWriteMask
 {
 	CW_RED   = 0x01,
diff --git a/Engine/Source/Runtime/RHI/Public/RHIStaticStates.h b/Engine/Source/Runtime/RHI/Public/RHIStaticStates.h
--- a/Engine/Source/Runtime/RHI/Public/RHIStaticStates.h
+++ b/Engine/Source/Runtime/RHI/Public/RHIStaticStates.h
@@ -142,13 +142,13 @@ public:
  * settings.
  * Should only be used from the rendering thread.
  */
-template<ERasterizerFillMode FillMode=FM_Solid,ERasterizerCullMode CullMode=CM_None,bool bEnableLineAA=false,bool bEnableMSAA=true>
+template<ERasterizerFillMode FillMode=FM_Solid,ERasterizerCullMode CullMode=CM_None,bool bEnableLineAA=false,bool bEnableMSAA=true, ERasterizerDepthClipMode DepthClipMode=ERasterizerDepthClipMode::DepthClip>
 class TStaticRasterizerState : public TStaticStateRHI<TStaticRasterizerState<FillMode,CullMode,bEnableLineAA>,FRasterizerStateRHIRef, FRHIRasterizerState*>
 {
 public:
 	FORCEINLINE_DEBUGGABLE static FRasterizerStateRHIRef CreateRHI()
 	{
-		FRasterizerStateInitializerRHI Initializer = { FillMode, CullMode, 0, 0, bEnableMSAA, bEnableLineAA };
+		const FRasterizerStateInitializerRHI Initializer(FillMode, CullMode, 0.0f, 0.0f, DepthClipMode, bEnableMSAA, bEnableLineAA);
 		return RHICreateRasterizerState(Initializer);
 	}
 };
diff --git a/Engine/Source/Runtime/RenderCore/Private/PipelineCacheUtilities.cpp b/Engine/Source/Runtime/RenderCore/Private/PipelineCacheUtilities.cpp
--- a/Engine/Source/Runtime/RenderCore/Private/PipelineCacheUtilities.cpp
+++ b/Engine/Source/Runtime/RenderCore/Private/PipelineCacheUtilities.cpp
@@ -55,7 +55,7 @@ namespace Private
 
 		enum class EVersion : int32
 		{
-			Current = 3
+			Current = 4
 		};
 
 		/** Magic to reject other files */
diff --git a/Engine/Source/Runtime/VulkanRHI/Private/VulkanState.cpp b/Engine/Source/Runtime/VulkanRHI/Private/VulkanState.cpp
--- a/Engine/Source/Runtime/VulkanRHI/Private/VulkanState.cpp
+++ b/Engine/Source/Runtime/VulkanRHI/Private/VulkanState.cpp
@@ -232,8 +232,8 @@ FVulkanSamplerState::FVulkanSamplerState(const VkSamplerCreateInfo& InInfo, FVul
 }
 
 FVulkanRasterizerState::FVulkanRasterizerState(const FRasterizerStateInitializerRHI& InInitializer)
+	: Initializer(InInitializer)
 {
-	Initializer = InInitializer;
 	FVulkanRasterizerState::ResetCreateInfo(RasterizerState);
 
 	// @todo vulkan: I'm assuming that Solid and Wireframe wouldn't ever be mixed within the same BoundShaderState, so we are ignoring the fill mode as a unique identifier
@@ -242,7 +242,7 @@ FVulkanRasterizerState::FVulkanRasterizerState(const FRasterizerStateInitializer
 	RasterizerState.polygonMode = RasterizerFillModeToVulkan(Initializer.FillMode);
 	RasterizerState.cullMode = RasterizerCullModeToVulkan(Initializer.CullMode);
 
-	//RasterizerState.depthClampEnable = VK_FALSE;
+	RasterizerState.depthClampEnable = InInitializer.DepthClipMode == ERasterizerDepthClipMode::DepthClamp ? VK_TRUE : VK_FALSE;
 	RasterizerState.depthBiasEnable = Initializer.DepthBias != 0.0f ? VK_TRUE : VK_FALSE;
 	//RasterizerState.rasterizerDiscardEnable = VK_FALSE;
 
diff --git a/Engine/Source/Runtime/Windows/D3D11RHI/Private/D3D11State.cpp b/Engine/Source/Runtime/Windows/D3D11RHI/Private/D3D11State.cpp
--- a/Engine/Source/Runtime/Windows/D3D11RHI/Private/D3D11State.cpp
+++ b/Engine/Source/Runtime/Windows/D3D11RHI/Private/D3D11State.cpp
@@ -204,7 +204,7 @@ FRasterizerStateRHIRef FD3D11DynamicRHI::RHICreateRasterizerState(const FRasteri
 	RasterizerDesc.SlopeScaledDepthBias = Initializer.SlopeScaleDepthBias;
 	RasterizerDesc.FrontCounterClockwise = true;
 	RasterizerDesc.DepthBias = FMath::FloorToInt(Initializer.DepthBias * (float)(1 << 24));
-	RasterizerDesc.DepthClipEnable = true;
+	RasterizerDesc.DepthClipEnable = Initializer.DepthClipMode == ERasterizerDepthClipMode::DepthClip;
 	RasterizerDesc.MultisampleEnable = Initializer.bAllowMSAA;
 	RasterizerDesc.ScissorEnable = true;
 
