[RHI] - Fixing FGraphicsPipelineRenderTargetsInfo not being initialized.

diff --git a/Engine/Source/Runtime/RHI/Public/RHIResources.h b/Engine/Source/Runtime/RHI/Public/RHIResources.h
--- a/Engine/Source/Runtime/RHI/Public/RHIResources.h
+++ b/Engine/Source/Runtime/RHI/Public/RHIResources.h
@@ -2760,21 +2760,29 @@ enum class EConservativeRasterization : uint8
 
 struct FGraphicsPipelineRenderTargetsInfo
 {
-	uint32															RenderTargetsEnabled;
+	FGraphicsPipelineRenderTargetsInfo()
+		: RenderTargetFormats(InPlace, UE_PIXELFORMAT_TO_UINT8(PF_Unknown))
+		, RenderTargetFlags(InPlace, TexCreate_None)
+		, DepthStencilAccess(FExclusiveDepthStencil::DepthNop_StencilNop)
+	{
+	}
+
+	uint32															RenderTargetsEnabled = 0;
 	TStaticArray<uint8, MaxSimultaneousRenderTargets>				RenderTargetFormats;
 	TStaticArray<ETextureCreateFlags, MaxSimultaneousRenderTargets>	RenderTargetFlags;
-	EPixelFormat													DepthStencilTargetFormat;
-	ETextureCreateFlags												DepthStencilTargetFlag;
-	ERenderTargetLoadAction											DepthTargetLoadAction;
-	ERenderTargetStoreAction										DepthTargetStoreAction;
-	ERenderTargetLoadAction											StencilTargetLoadAction;
-	ERenderTargetStoreAction										StencilTargetStoreAction;
+	EPixelFormat													DepthStencilTargetFormat = PF_Unknown;
+	ETextureCreateFlags												DepthStencilTargetFlag = ETextureCreateFlags::None;
+	ERenderTargetLoadAction											DepthTargetLoadAction = ERenderTargetLoadAction::ENoAction;
+	ERenderTargetStoreAction										DepthTargetStoreAction = ERenderTargetStoreAction::ENoAction;
+	ERenderTargetLoadAction											StencilTargetLoadAction = ERenderTargetLoadAction::ENoAction;
+	ERenderTargetStoreAction										StencilTargetStoreAction = ERenderTargetStoreAction::ENoAction;
 	FExclusiveDepthStencil											DepthStencilAccess;
-	uint16															NumSamples;
-	uint8															MultiViewCount;
-	bool															bHasFragmentDensityAttachment;
+	uint16															NumSamples = 0;
+	uint8															MultiViewCount = 0;
+	bool															bHasFragmentDensityAttachment = false;
 };
 
+
 class FGraphicsPipelineStateInitializer
 {
 public:
