Windows RHI selection fixes

diff --git a/Engine/Source/Runtime/RHI/Private/Windows/WindowsDynamicRHI.cpp b/Engine/Source/Runtime/RHI/Private/Windows/WindowsDynamicRHI.cpp
--- a/Engine/Source/Runtime/RHI/Private/Windows/WindowsDynamicRHI.cpp
+++ b/Engine/Source/Runtime/RHI/Private/Windows/WindowsDynamicRHI.cpp
@@ -240,7 +240,7 @@ static TOptional<WindowsRHI> ChooseForcedRHI()
 	}
 
 	// FeatureLevelES31 is also a command line override, so it will determine the underlying RHI unless one is specified
-	if (bES31DXOnly && FParse::Param(FCommandLine::Get(), TEXT("FeatureLevelES31")) || FParse::Param(FCommandLine::Get(), TEXT("FeatureLevelES3_1")))
+	if (bES31DXOnly && (FParse::Param(FCommandLine::Get(), TEXT("FeatureLevelES31")) || FParse::Param(FCommandLine::Get(), TEXT("FeatureLevelES3_1"))))
 	{
 		if (ForcedRHI == WindowsRHI::OpenGL)
 		{
@@ -297,6 +297,20 @@ static ERHIFeatureLevel::Type ChooseFeatureLevel(TOptional<WindowsRHI> ChosenRHI
 		FeatureLevel = ERHIFeatureLevel::ES3_1;
 	}
 
+	bool bES31DXOnly = false;
+#if !WITH_EDITOR
+	if (!GIsEditor)
+	{
+		GConfig->GetBool(TEXT("PerformanceMode"), TEXT("bES31DXOnly"), bES31DXOnly, GEngineIni);
+	}
+#endif
+
+	// Allow Vulkan to force ES31 unless we blocked it in config
+	if (!bES31DXOnly && ChosenRHI == WindowsRHI::Vulkan && bForceES31)
+	{
+		FeatureLevel = ERHIFeatureLevel::ES3_1;
+	}
+
 	return FeatureLevel;
 }
 
