Finishing move to ComputePipelineState objects
- Deprecating IRHIComputeContext::RHISetComputeShader.
- Moving the fallback version of  RHISetComputePipelineState to IRHICommandContextPSOFallback.
- Making IRHIComputeContext::RHISetComputePipelineState pure virtual.
- Moving the fallback versions of RHICreateGraphicsPipelineState and RHICreateComputePipelineState into the new FDynamicRHIPSOFallback for platforms that were already using IRHICommandContextPSOFallback.
- Making FDynamicRHI::RHICreateGraphicsPipelineState and FDynamicRHI::RHICreateComputePipelineState pure virtual.

diff --git a/Engine/Source/Runtime/Apple/AGXRHI/Private/AGXCommands.cpp b/Engine/Source/Runtime/Apple/AGXRHI/Private/AGXCommands.cpp
--- a/Engine/Source/Runtime/Apple/AGXRHI/Private/AGXCommands.cpp
+++ b/Engine/Source/Runtime/Apple/AGXRHI/Private/AGXCommands.cpp
@@ -120,19 +120,6 @@ void FAGXRHICommandContext::RHISetStreamSource(uint32 StreamIndex, FRHIBuffer* V
 	}
 }
 
-void FAGXRHICommandContext::RHISetComputeShader(FRHIComputeShader* ComputeShaderRHI)
-{
-	@autoreleasepool {
-	FAGXComputeShader* ComputeShader = ResourceCast(ComputeShaderRHI);
-
-	// cache this for Dispatch
-	// sets this compute shader pipeline as the current (this resets all state, so we need to set all resources after calling this)
-	Context->GetCurrentState().SetComputeShader(ComputeShader);
-
-	ApplyStaticUniformBuffers(ComputeShader);
-	}
-}
-
 void FAGXRHICommandContext::RHISetComputePipelineState(FRHIComputePipelineState* ComputePipelineState)
 {
 	@autoreleasepool {
diff --git a/Engine/Source/Runtime/Apple/AGXRHI/Private/AGXRHIContext.cpp b/Engine/Source/Runtime/Apple/AGXRHI/Private/AGXRHIContext.cpp
--- a/Engine/Source/Runtime/Apple/AGXRHI/Private/AGXRHIContext.cpp
+++ b/Engine/Source/Runtime/Apple/AGXRHI/Private/AGXRHIContext.cpp
@@ -88,15 +88,6 @@ void FAGXRHIComputeContext::RHISetAsyncComputeBudget(EAsyncComputeBudget Budget)
 	FAGXRHICommandContext::RHISetAsyncComputeBudget(Budget);
 }
 
-void FAGXRHIComputeContext::RHISetComputeShader(FRHIComputeShader* ComputeShader)
-{
-	if (!Context->GetCurrentCommandBuffer())
-	{
-		Context->InitFrame(false, 0, 0);
-	}
-	FAGXRHICommandContext::RHISetComputeShader(ComputeShader);
-}
-
 void FAGXRHIComputeContext::RHISetComputePipelineState(FRHIComputePipelineState* ComputePipelineState)
 {
 	if (!Context->GetCurrentCommandBuffer())
diff --git a/Engine/Source/Runtime/Apple/AGXRHI/Public/AGXRHIContext.h b/Engine/Source/Runtime/Apple/AGXRHI/Public/AGXRHIContext.h
--- a/Engine/Source/Runtime/Apple/AGXRHI/Public/AGXRHIContext.h
+++ b/Engine/Source/Runtime/Apple/AGXRHI/Public/AGXRHIContext.h
@@ -34,12 +34,6 @@ public:
 	/** Get the profiler pointer */
 	FORCEINLINE class FAGXProfiler* GetProfiler() const { return Profiler; }
 
-	/**
-	 *Sets the current compute shader.  Mostly for compliance with platforms
-	 *that require shader setting before resource binding.
-	 */
-	virtual void RHISetComputeShader(FRHIComputeShader* ComputeShader) override;
-	
 	virtual void RHISetComputePipelineState(FRHIComputePipelineState* ComputePipelineState) override;
 	
 	virtual void RHIDispatchComputeShader(uint32 ThreadGroupCountX, uint32 ThreadGroupCountY, uint32 ThreadGroupCountZ) final override;
@@ -260,7 +254,6 @@ public:
 	virtual ~FAGXRHIComputeContext();
 	
 	virtual void RHISetAsyncComputeBudget(EAsyncComputeBudget Budget) final override;
-	virtual void RHISetComputeShader(FRHIComputeShader* ComputeShader) final override;
 	virtual void RHISetComputePipelineState(FRHIComputePipelineState* ComputePipelineState) final override;
 	virtual void RHISubmitCommandsHint() final override;
 };
diff --git a/Engine/Source/Runtime/Apple/MetalRHI/Private/MetalCommands.cpp b/Engine/Source/Runtime/Apple/MetalRHI/Private/MetalCommands.cpp
--- a/Engine/Source/Runtime/Apple/MetalRHI/Private/MetalCommands.cpp
+++ b/Engine/Source/Runtime/Apple/MetalRHI/Private/MetalCommands.cpp
@@ -121,19 +121,6 @@ void FMetalRHICommandContext::RHISetStreamSource(uint32 StreamIndex, FRHIBuffer*
 	}
 }
 
-void FMetalRHICommandContext::RHISetComputeShader(FRHIComputeShader* ComputeShaderRHI)
-{
-	@autoreleasepool {
-	FMetalComputeShader* ComputeShader = ResourceCast(ComputeShaderRHI);
-
-	// cache this for Dispatch
-	// sets this compute shader pipeline as the current (this resets all state, so we need to set all resources after calling this)
-	Context->GetCurrentState().SetComputeShader(ComputeShader);
-
-	ApplyStaticUniformBuffers(ComputeShader);
-	}
-}
-
 void FMetalRHICommandContext::RHISetComputePipelineState(FRHIComputePipelineState* ComputePipelineState)
 {
 	@autoreleasepool {
diff --git a/Engine/Source/Runtime/Apple/MetalRHI/Private/MetalRHIContext.cpp b/Engine/Source/Runtime/Apple/MetalRHI/Private/MetalRHIContext.cpp
--- a/Engine/Source/Runtime/Apple/MetalRHI/Private/MetalRHIContext.cpp
+++ b/Engine/Source/Runtime/Apple/MetalRHI/Private/MetalRHIContext.cpp
@@ -105,15 +105,6 @@ void FMetalRHIComputeContext::RHISetAsyncComputeBudget(EAsyncComputeBudget Budge
 	FMetalRHICommandContext::RHISetAsyncComputeBudget(Budget);
 }
 
-void FMetalRHIComputeContext::RHISetComputeShader(FRHIComputeShader* ComputeShader)
-{
-	if (!Context->GetCurrentCommandBuffer())
-	{
-		Context->InitFrame(false, 0, 0);
-	}
-	FMetalRHICommandContext::RHISetComputeShader(ComputeShader);
-}
-
 void FMetalRHIComputeContext::RHISetComputePipelineState(FRHIComputePipelineState* ComputePipelineState)
 {
 	if (!Context->GetCurrentCommandBuffer())
diff --git a/Engine/Source/Runtime/Apple/MetalRHI/Public/MetalRHIContext.h b/Engine/Source/Runtime/Apple/MetalRHI/Public/MetalRHIContext.h
--- a/Engine/Source/Runtime/Apple/MetalRHI/Public/MetalRHIContext.h
+++ b/Engine/Source/Runtime/Apple/MetalRHI/Public/MetalRHIContext.h
@@ -35,12 +35,6 @@ public:
 	/** Get the profiler pointer */
 	FORCEINLINE class FMetalProfiler* GetProfiler() const { return Profiler; }
 
-	/**
-	 *Sets the current compute shader.  Mostly for compliance with platforms
-	 *that require shader setting before resource binding.
-	 */
-	virtual void RHISetComputeShader(FRHIComputeShader* ComputeShader) override;
-	
 	virtual void RHISetComputePipelineState(FRHIComputePipelineState* ComputePipelineState) override;
 	
 	virtual void RHIDispatchComputeShader(uint32 ThreadGroupCountX, uint32 ThreadGroupCountY, uint32 ThreadGroupCountZ) final override;
@@ -261,7 +255,6 @@ public:
 	virtual ~FMetalRHIComputeContext();
 	
 	virtual void RHISetAsyncComputeBudget(EAsyncComputeBudget Budget) final override;
-	virtual void RHISetComputeShader(FRHIComputeShader* ComputeShader) final override;
 	virtual void RHISetComputePipelineState(FRHIComputePipelineState* ComputePipelineState) final override;
 	virtual void RHISubmitCommandsHint() final override;
 };
diff --git a/Engine/Source/Runtime/D3D12RHI/Private/D3D12CommandContext.h b/Engine/Source/Runtime/D3D12RHI/Private/D3D12CommandContext.h
--- a/Engine/Source/Runtime/D3D12RHI/Private/D3D12CommandContext.h
+++ b/Engine/Source/Runtime/D3D12RHI/Private/D3D12CommandContext.h
@@ -267,7 +267,6 @@ public:
 	virtual void RHIEndTransitions(TArrayView<const FRHITransition*> Transitions) final override;
 
 	// IRHIComputeContext interface
-	virtual void RHISetComputeShader(FRHIComputeShader* ComputeShader) final override;
 	virtual void RHISetComputePipelineState(FRHIComputePipelineState* ComputePipelineState) final override;
 	virtual void RHIDispatchComputeShader(uint32 ThreadGroupCountX, uint32 ThreadGroupCountY, uint32 ThreadGroupCountZ) final override;
 	virtual void RHIDispatchIndirectComputeShader(FRHIBuffer* ArgumentBuffer, uint32 ArgumentOffset) final override;
@@ -508,10 +507,6 @@ public:
 #define ContextRedirect(Call) { for (uint32 GPUIndex : GPUMask) PhysicalContexts[GPUIndex]->##Call; }
 #define ContextGPU0(Call) { PhysicalContexts[0]->##Call; }
 
-	FORCEINLINE virtual void RHISetComputeShader(FRHIComputeShader* ComputeShader) final override
-	{
-		ContextRedirect(RHISetComputeShader(ComputeShader));
-	}
 	FORCEINLINE virtual void RHISetComputePipelineState(FRHIComputePipelineState* ComputePipelineState) final override
 	{
 		ContextRedirect(RHISetComputePipelineState(ComputePipelineState));
diff --git a/Engine/Source/Runtime/D3D12RHI/Private/D3D12Commands.cpp b/Engine/Source/Runtime/D3D12RHI/Private/D3D12Commands.cpp
--- a/Engine/Source/Runtime/D3D12RHI/Private/D3D12Commands.cpp
+++ b/Engine/Source/Runtime/D3D12RHI/Private/D3D12Commands.cpp
@@ -112,14 +112,6 @@ void FD3D12CommandContext::RHISetStreamSource(uint32 StreamIndex, FRHIBuffer* Ve
 	StateCache.SetStreamSource(VertexBuffer ? &VertexBuffer->ResourceLocation : nullptr, StreamIndex, Offset);
 }
 
-void FD3D12CommandContext::RHISetComputeShader(FRHIComputeShader* ComputeShaderRHI)
-{
-	// TODO: Eventually the high-level should just use RHISetComputePipelineState() directly similar to how graphics PSOs are handled.
-	TRefCountPtr<FRHIComputePipelineState> ComputePipelineStateRHI = RHICreateComputePipelineState(ComputeShaderRHI);
-	FD3D12ComputePipelineState* const ComputePipelineState = FD3D12DynamicRHI::ResourceCast(ComputePipelineStateRHI.GetReference());
-	RHISetComputePipelineState(ComputePipelineState);
-}
-
 void FD3D12CommandContext::RHIDispatchComputeShader(uint32 ThreadGroupCountX, uint32 ThreadGroupCountY, uint32 ThreadGroupCountZ)
 {
 	FD3D12ComputeShader* ComputeShader = nullptr;
diff --git a/Engine/Source/Runtime/NullDrv/Public/NullRHI.h b/Engine/Source/Runtime/NullDrv/Public/NullRHI.h
--- a/Engine/Source/Runtime/NullDrv/Public/NullRHI.h
+++ b/Engine/Source/Runtime/NullDrv/Public/NullRHI.h
@@ -10,7 +10,7 @@
 struct Rect;
 
 /** A null implementation of the dynamically bound RHI. */
-class FNullDynamicRHI : public FDynamicRHI , public IRHICommandContextPSOFallback
+class FNullDynamicRHI : public FDynamicRHIPSOFallback, public IRHICommandContextPSOFallback
 {
 public:
 
diff --git a/Engine/Source/Runtime/OpenGLDrv/Public/OpenGLDrv.h b/Engine/Source/Runtime/OpenGLDrv/Public/OpenGLDrv.h
--- a/Engine/Source/Runtime/OpenGLDrv/Public/OpenGLDrv.h
+++ b/Engine/Source/Runtime/OpenGLDrv/Public/OpenGLDrv.h
@@ -288,7 +288,7 @@ struct FOpenGLGPUProfiler : public FGPUProfiler
 PRAGMA_DISABLE_DEPRECATION_WARNINGS
 
 /** The interface which is implemented by the dynamically bound RHI. */
-class OPENGLDRV_API FOpenGLDynamicRHI  final : public FDynamicRHI, public IRHICommandContextPSOFallback
+class OPENGLDRV_API FOpenGLDynamicRHI  final : public FDynamicRHIPSOFallback, public IRHICommandContextPSOFallback
 {
 public:
 
diff --git a/Engine/Source/Runtime/RHI/Private/RHICommandList.cpp b/Engine/Source/Runtime/RHI/Private/RHICommandList.cpp
--- a/Engine/Source/Runtime/RHI/Private/RHICommandList.cpp
+++ b/Engine/Source/Runtime/RHI/Private/RHICommandList.cpp
@@ -2754,3 +2754,9 @@ void FRHICommandListImmediate::UpdateRHIResources(FRHIResourceUpdateInfo* Update
 		}
 	}
 }
+
+RHI_API void RHISetComputeShaderBackwardsCompatible(IRHIComputeContext* InContext, FRHIComputeShader* InShader)
+{
+	TRefCountPtr<FRHIComputePipelineState> ComputePipelineState = RHICreateComputePipelineState(InShader);
+	InContext->RHISetComputePipelineState(ComputePipelineState);
+}
diff --git a/Engine/Source/Runtime/RHI/Private/RHIValidation.cpp b/Engine/Source/Runtime/RHI/Private/RHIValidation.cpp
--- a/Engine/Source/Runtime/RHI/Private/RHIValidation.cpp
+++ b/Engine/Source/Runtime/RHI/Private/RHIValidation.cpp
@@ -434,7 +434,7 @@ FValidationComputeContext::FValidationComputeContext()
 void FValidationComputeContext::FState::Reset()
 {
 	ComputePassName.Reset();
-	bComputeShaderSet = false;
+	bComputePSOSet = false;
 	TrackerInstance.ResetAllUAVState();
 	GlobalUniformBuffers.Reset();
 }
@@ -517,7 +517,7 @@ void FValidationContext::FState::Reset()
 	RenderPassName.Reset();
 	PreviousRenderPassName.Reset();
 	ComputePassName.Reset();
-	bComputeShaderSet = false;
+	bComputePSOSet = false;
 	TrackerInstance.ResetAllUAVState();
 	GlobalUniformBuffers.Reset();
 }
diff --git a/Engine/Source/Runtime/RHI/Public/DynamicRHI.h b/Engine/Source/Runtime/RHI/Public/DynamicRHI.h
--- a/Engine/Source/Runtime/RHI/Public/DynamicRHI.h
+++ b/Engine/Source/Runtime/RHI/Public/DynamicRHI.h
@@ -358,23 +358,17 @@ public:
 	* CAUTION: Platforms that support RHIThread but don't actually have a threadsafe implementation must flush internally with FScopedRHIThreadStaller StallRHIThread(FRHICommandListExecutor::GetImmediateCommandList()); when the call is from the render thread
 	*/
 	// FlushType: Thread safe
-	// TODO: [PSO API] Make pure virtual
-	virtual FGraphicsPipelineStateRHIRef RHICreateGraphicsPipelineState(const FGraphicsPipelineStateInitializer& Initializer)
-	{
-		return new FRHIGraphicsPipelineStateFallBack(Initializer);
-	}
+	virtual FGraphicsPipelineStateRHIRef RHICreateGraphicsPipelineState(const FGraphicsPipelineStateInitializer& Initializer) = 0;
 
-	virtual TRefCountPtr<FRHIComputePipelineState> RHICreateComputePipelineState(FRHIComputeShader* ComputeShader)
-	{
-		return new FRHIComputePipelineStateFallback(ComputeShader);
-	}
+	// FlushType: Thread safe
+	virtual FComputePipelineStateRHIRef RHICreateComputePipelineState(FRHIComputeShader * ComputeShader) = 0;
 
 	virtual FGraphicsPipelineStateRHIRef RHICreateGraphicsPipelineState(const FGraphicsPipelineStateInitializer& Initializer, FRHIPipelineBinaryLibrary* PipelineBinary)
 	{
 		return RHICreateGraphicsPipelineState(Initializer);
 	}
 
-	virtual TRefCountPtr<FRHIComputePipelineState> RHICreateComputePipelineState(FRHIComputeShader* ComputeShader, FRHIPipelineBinaryLibrary* PipelineBinary)
+	virtual FComputePipelineStateRHIRef RHICreateComputePipelineState(FRHIComputeShader* ComputeShader, FRHIPipelineBinaryLibrary* PipelineBinary)
 	{
 		return RHICreateComputePipelineState(ComputeShader);
 	}
@@ -1322,6 +1316,21 @@ protected:
 /** A global pointer to the dynamically bound RHI implementation. */
 extern RHI_API FDynamicRHI* GDynamicRHI;
 
+// Dynamic RHI for RHIs that do not support real Graphics/Compute Pipelines.
+class FDynamicRHIPSOFallback : public FDynamicRHI
+{
+public:
+	virtual FGraphicsPipelineStateRHIRef RHICreateGraphicsPipelineState(const FGraphicsPipelineStateInitializer& Initializer) override
+	{
+		return new FRHIGraphicsPipelineStateFallBack(Initializer);
+	}
+
+	virtual FComputePipelineStateRHIRef RHICreateComputePipelineState(FRHIComputeShader* ComputeShader) override
+	{
+		return new FRHIComputePipelineStateFallback(ComputeShader);
+	}
+};
+
 FORCEINLINE FSamplerStateRHIRef RHICreateSamplerState(const FSamplerStateInitializerRHI& Initializer)
 {
 	return GDynamicRHI->RHICreateSamplerState(Initializer);
@@ -1359,7 +1368,7 @@ FORCEINLINE FVertexDeclarationRHIRef RHICreateVertexDeclaration(const FVertexDec
 	return GDynamicRHI->RHICreateVertexDeclaration(Elements);
 }
 
-FORCEINLINE TRefCountPtr<FRHIComputePipelineState> RHICreateComputePipelineState(FRHIComputeShader* ComputeShader)
+FORCEINLINE FComputePipelineStateRHIRef RHICreateComputePipelineState(FRHIComputeShader* ComputeShader)
 {
 	return GDynamicRHI->RHICreateComputePipelineState(ComputeShader);
 }
diff --git a/Engine/Source/Runtime/RHI/Public/RHICommandList.h b/Engine/Source/Runtime/RHI/Public/RHICommandList.h
--- a/Engine/Source/Runtime/RHI/Public/RHICommandList.h
+++ b/Engine/Source/Runtime/RHI/Public/RHICommandList.h
@@ -2546,7 +2546,9 @@ public:
 		ComputeShader->UpdateStats();
 		if (Bypass())
 		{
+PRAGMA_DISABLE_DEPRECATION_WARNINGS
 			GetComputeContext().RHISetComputeShader(ComputeShader);
+PRAGMA_ENABLE_DEPRECATION_WARNINGS
 			return;
 		}
 		ALLOC_COMMAND(FRHICommandSetComputeShader)(ComputeShader);
diff --git a/Engine/Source/Runtime/RHI/Public/RHICommandListCommandExecutes.inl b/Engine/Source/Runtime/RHI/Public/RHICommandListCommandExecutes.inl
--- a/Engine/Source/Runtime/RHI/Public/RHICommandListCommandExecutes.inl
+++ b/Engine/Source/Runtime/RHI/Public/RHICommandListCommandExecutes.inl
@@ -264,7 +264,9 @@ void FRHICommandNextSubpass::Execute(FRHICommandListBase& CmdList)
 void FRHICommandSetComputeShader::Execute(FRHICommandListBase& CmdList)
 {
 	RHISTAT(SetComputeShader);
+PRAGMA_DISABLE_DEPRECATION_WARNINGS
 	INTERNAL_DECORATOR_COMPUTE(RHISetComputeShader)(ComputeShader);
+PRAGMA_ENABLE_DEPRECATION_WARNINGS
 }
 
 void FRHICommandSetComputePipelineState::Execute(FRHICommandListBase& CmdList)
diff --git a/Engine/Source/Runtime/RHI/Public/RHIContext.h b/Engine/Source/Runtime/RHI/Public/RHIContext.h
--- a/Engine/Source/Runtime/RHI/Public/RHIContext.h
+++ b/Engine/Source/Runtime/RHI/Public/RHIContext.h
@@ -165,20 +165,15 @@ public:
 	{
 	}
 
-	/**
-	*Sets the current compute shader.
-	*/
-	virtual void RHISetComputeShader(FRHIComputeShader* ComputeShader) = 0;
-
-	virtual void RHISetComputePipelineState(FRHIComputePipelineState* ComputePipelineState)
+	UE_DEPRECATED(5.1, "ComputePipelineStates should be used instead of direct ComputeShaders.")
+	virtual void RHISetComputeShader(FRHIComputeShader* ComputeShader)
 	{
-		if (ComputePipelineState)
-		{
-			FRHIComputePipelineStateFallback* FallbackState = static_cast<FRHIComputePipelineStateFallback*>(ComputePipelineState);
-			RHISetComputeShader(FallbackState->GetComputeShader());
-		}
+		RHI_API void RHISetComputeShaderBackwardsCompatible(IRHIComputeContext*, FRHIComputeShader*);
+		RHISetComputeShaderBackwardsCompatible(this, ComputeShader);
 	}
 
+	virtual void RHISetComputePipelineState(FRHIComputePipelineState* ComputePipelineState) = 0;
+
 	virtual void RHIDispatchComputeShader(uint32 ThreadGroupCountX, uint32 ThreadGroupCountY, uint32 ThreadGroupCountZ) = 0;
 
 	virtual void RHIDispatchIndirectComputeShader(FRHIBuffer* ArgumentBuffer, uint32 ArgumentOffset) = 0;
@@ -841,23 +836,17 @@ FORCEINLINE FBoundShaderStateRHIRef RHICreateBoundShaderState(
 );
 
 
-// Command Context for RHIs that do not support real Graphics Pipelines.
+// Command Context for RHIs that do not support real Graphics/Compute Pipelines.
 class IRHICommandContextPSOFallback : public IRHICommandContext
 {
 public:
-	/**
-	* Set bound shader state. This will set the vertex decl/shader, and pixel shader
-	* @param BoundShaderState - state resource
-	*/
 	virtual void RHISetBoundShaderState(FRHIBoundShaderState* BoundShaderState) = 0;
-
 	virtual void RHISetDepthStencilState(FRHIDepthStencilState* NewState, uint32 StencilRef) = 0;
-
 	virtual void RHISetRasterizerState(FRHIRasterizerState* NewState) = 0;
-
 	virtual void RHISetBlendState(FRHIBlendState* NewState, const FLinearColor& BlendFactor) = 0;
-
 	virtual void RHIEnableDepthBoundsTest(bool bEnable) = 0;
+	// TODO: uncomment when removed from IRHIComputeContext
+	//virtual void RHISetComputeShader(FRHIComputeShader* ComputeShader) = 0;
 
 	/**
 	* This will set most relevant pipeline state. Legacy APIs are expected to set corresponding disjoint state as well.
@@ -876,6 +865,16 @@ public:
 	}
 #endif
 
+	virtual void RHISetComputePipelineState(FRHIComputePipelineState* ComputePipelineState)
+	{
+		if (FRHIComputePipelineStateFallback* FallbackState = static_cast<FRHIComputePipelineStateFallback*>(ComputePipelineState))
+		{
+PRAGMA_DISABLE_DEPRECATION_WARNINGS
+			RHISetComputeShader(FallbackState->GetComputeShader());
+PRAGMA_ENABLE_DEPRECATION_WARNINGS
+		}
+	}
+
 private:
 	void SetGraphicsPipelineStateFromInitializer(const FGraphicsPipelineStateInitializer& PsoInit, uint32 StencilRef, bool bApplyAdditionalState)
 	{
diff --git a/Engine/Source/Runtime/RHI/Public/RHIValidationContext.h b/Engine/Source/Runtime/RHI/Public/RHIValidationContext.h
--- a/Engine/Source/Runtime/RHI/Public/RHIValidationContext.h
+++ b/Engine/Source/Runtime/RHI/Public/RHIValidationContext.h
@@ -21,19 +21,15 @@ public:
 	{
 	}
 
-	/**
-	*Sets the current compute shader.  Mostly for compliance with platforms
-	*that require shader setting before resource binding.
-	*/
-	virtual void RHISetComputeShader(FRHIComputeShader* Shader) override final
+	virtual void RHISetComputePipelineState(FRHIComputePipelineState* ComputePipelineState) override final
 	{
-		State.bComputeShaderSet = true;
+		State.bComputePSOSet = true;
 
 		// Reset the compute UAV tracker since the renderer must re-bind all resources after changing a shader.
 		Tracker->ResetUAVState(RHIValidation::EUAVMode::Compute);
 
 		State.GlobalUniformBuffers.bInSetPipelineStateCall = true;
-		RHIContext->RHISetComputeShader(Shader);
+		RHIContext->RHISetComputePipelineState(ComputePipelineState);
 		State.GlobalUniformBuffers.bInSetPipelineStateCall = false;
 	}
 
@@ -117,7 +113,7 @@ public:
 	/** Set the shader resource view of a surface.  This is used for binding TextureMS parameter types that need a multi sampled view. */
 	virtual void RHISetShaderTexture(FRHIComputeShader* Shader, uint32 TextureIndex, FRHITexture* NewTexture) override final
 	{
-		checkf(State.bComputeShaderSet, TEXT("A Compute shader has to be set to set resources into a shader!"));
+		checkf(State.bComputePSOSet, TEXT("A Compute PSO has to be set to set resources into a shader!"));
 		Tracker->Assert(NewTexture->GetWholeResourceIdentitySRV(), ERHIAccess::SRVCompute);
 		RHIContext->RHISetShaderTexture(Shader, TextureIndex, NewTexture);
 	}
@@ -130,7 +126,7 @@ public:
 	*/
 	virtual void RHISetShaderSampler(FRHIComputeShader* Shader, uint32 SamplerIndex, FRHISamplerState* NewState) override final
 	{
-		checkf(State.bComputeShaderSet, TEXT("A Compute shader has to be set to set resources into a shader!"));
+		checkf(State.bComputePSOSet, TEXT("A Compute PSO has to be set to set resources into a shader!"));
 		RHIContext->RHISetShaderSampler(Shader, SamplerIndex, NewState);
 	}
 
@@ -142,7 +138,7 @@ public:
 	*/
 	virtual void RHISetUAVParameter(FRHIComputeShader* Shader, uint32 UAVIndex, FRHIUnorderedAccessView* UAV) override final
 	{
-		checkf(State.bComputeShaderSet, TEXT("A Compute shader has to be set to set resources into a shader!"));
+		checkf(State.bComputePSOSet, TEXT("A Compute PSO has to be set to set resources into a shader!"));
 		Tracker->AssertUAV(UAV, RHIValidation::EUAVMode::Compute, UAVIndex);
 		RHIContext->RHISetUAVParameter(Shader, UAVIndex, UAV);
 	}
@@ -156,14 +152,14 @@ public:
 	*/
 	virtual void RHISetUAVParameter(FRHIComputeShader* Shader, uint32 UAVIndex, FRHIUnorderedAccessView* UAV, uint32 InitialCount) override final
 	{
-		checkf(State.bComputeShaderSet, TEXT("A Compute shader has to be set to set resources into a shader!"));
+		checkf(State.bComputePSOSet, TEXT("A Compute PSO has to be set to set resources into a shader!"));
 		Tracker->AssertUAV(UAV, RHIValidation::EUAVMode::Compute, UAVIndex);
 		RHIContext->RHISetUAVParameter(Shader, UAVIndex, UAV, InitialCount);
 	}
 
 	virtual void RHISetShaderResourceViewParameter(FRHIComputeShader* Shader, uint32 SamplerIndex, FRHIShaderResourceView* SRV) override final
 	{
-		checkf(State.bComputeShaderSet, TEXT("A Compute shader has to be set to set resources into a shader!"));
+		checkf(State.bComputePSOSet, TEXT("A Compute PSO has to be set to set resources into a shader!"));
 		if (SRV) 
 		{
 			Tracker->Assert(SRV->ViewIdentity, ERHIAccess::SRVCompute);
@@ -173,14 +169,14 @@ public:
 
 	virtual void RHISetShaderUniformBuffer(FRHIComputeShader* Shader, uint32 BufferIndex, FRHIUniformBuffer* Buffer) override final
 	{
-		checkf(State.bComputeShaderSet, TEXT("A Compute shader has to be set to set resources into a shader!"));
+		checkf(State.bComputePSOSet, TEXT("A Compute PSO has to be set to set resources into a shader!"));
 		State.GlobalUniformBuffers.ValidateSetShaderUniformBuffer(Buffer);
 		RHIContext->RHISetShaderUniformBuffer(Shader, BufferIndex, Buffer);
 	}
 
 	virtual void RHISetShaderParameter(FRHIComputeShader* Shader, uint32 BufferIndex, uint32 BaseIndex, uint32 NumBytes, const void* NewValue) override final
 	{
-		checkf(State.bComputeShaderSet, TEXT("A Compute shader has to be set to set resources into a shader!"));
+		checkf(State.bComputePSOSet, TEXT("A Compute PSO has to be set to set resources into a shader!"));
 		RHIContext->RHISetShaderParameter(Shader, BufferIndex, BaseIndex, NumBytes, NewValue);
 	}
 
@@ -244,7 +240,7 @@ protected:
 		RHIValidation::FGlobalUniformBuffers GlobalUniformBuffers;
 
 		FString ComputePassName;
-		bool bComputeShaderSet{};
+		bool bComputePSOSet{};
 
 		void Reset();
 	};
@@ -264,19 +260,16 @@ public:
 		return *RHIContext;
 	}
 
-	/**
-	*Sets the current compute shader.  Mostly for compliance with platforms
-	*that require shader setting before resource binding.
-	*/
-	virtual void RHISetComputeShader(FRHIComputeShader* Shader) override final
+	virtual void RHISetComputePipelineState(FRHIComputePipelineState* ComputePipelineState) override final
 	{
-		State.bComputeShaderSet = true;
-		State.bGfxPSOSet = false;
+		State.bComputePSOSet = true;
 
 		// Reset the compute UAV tracker since the renderer must re-bind all resources after changing a shader.
 		Tracker->ResetUAVState(RHIValidation::EUAVMode::Compute);
 
-		RHIContext->RHISetComputeShader(Shader);
+		State.GlobalUniformBuffers.bInSetPipelineStateCall = true;
+		RHIContext->RHISetComputePipelineState(ComputePipelineState);
+		State.GlobalUniformBuffers.bInSetPipelineStateCall = false;
 	}
 
 	virtual void RHIDispatchComputeShader(uint32 ThreadGroupCountX, uint32 ThreadGroupCountY, uint32 ThreadGroupCountZ) override final
@@ -570,7 +563,7 @@ public:
 	{
 		checkf(State.bInsideBeginRenderPass, TEXT("Graphics PSOs can only be set inside a RenderPass!"));
 		State.bGfxPSOSet = true;
-		State.bComputeShaderSet = false;
+		State.bComputePSOSet = false;
 
 		ValidateDepthStencilForSetGraphicsPipelineState(GraphicsState->DSMode);
 
@@ -587,7 +580,7 @@ public:
 	{
 		checkf(State.bInsideBeginRenderPass, TEXT("Graphics PSOs can only be set inside a RenderPass!"));
 		State.bGfxPSOSet = true;
-		State.bComputeShaderSet = false;
+		State.bComputePSOSet = false;
 
 		ValidateDepthStencilForSetGraphicsPipelineState(PsoInit.DepthStencilState->ActualDSMode);
 
@@ -611,7 +604,7 @@ public:
 	/** Set the shader resource view of a surface.  This is used for binding TextureMS parameter types that need a multi sampled view. */
 	virtual void RHISetShaderTexture(FRHIComputeShader* Shader, uint32 TextureIndex, FRHITexture* NewTexture) override final
 	{
-		checkf(State.bComputeShaderSet, TEXT("A Compute shader has to be set to set resources into a shader!"));
+		checkf(State.bComputePSOSet, TEXT("A Compute PSO has to be set to set resources into a shader!"));
 		Tracker->Assert(NewTexture->GetWholeResourceIdentitySRV(), ERHIAccess::SRVCompute);
 		RHIContext->RHISetShaderTexture(Shader, TextureIndex, NewTexture);
 	}
@@ -624,7 +617,7 @@ public:
 	*/
 	virtual void RHISetShaderSampler(FRHIComputeShader* Shader, uint32 SamplerIndex, FRHISamplerState* NewState) override final
 	{
-		checkf(State.bComputeShaderSet, TEXT("A Compute shader has to be set to set resources into a shader!"));
+		checkf(State.bComputePSOSet, TEXT("A Compute PSO has to be set to set resources into a shader!"));
 		RHIContext->RHISetShaderSampler(Shader, SamplerIndex, NewState);
 	}
 
@@ -648,7 +641,7 @@ public:
 	*/
 	virtual void RHISetUAVParameter(FRHIPixelShader* Shader, uint32 UAVIndex, FRHIUnorderedAccessView* UAV) override final
 	{
-		checkf(State.bGfxPSOSet, TEXT("A Compute shader has to be set to set resources into a shader!"));
+		checkf(State.bGfxPSOSet, TEXT("A Compute PSO has to be set to set resources into a shader!"));
 		Tracker->AssertUAV(UAV, RHIValidation::EUAVMode::Graphics, UAVIndex);
 		RHIContext->RHISetUAVParameter(Shader, UAVIndex, UAV);
 	}
@@ -664,7 +657,7 @@ public:
 	*/
 	virtual void RHISetUAVParameter(FRHIComputeShader* Shader, uint32 UAVIndex, FRHIUnorderedAccessView* UAV) override final
 	{
-		checkf(State.bComputeShaderSet, TEXT("A Compute shader has to be set to set resources into a shader!"));
+		checkf(State.bComputePSOSet, TEXT("A Compute PSO has to be set to set resources into a shader!"));
 		Tracker->AssertUAV(UAV, RHIValidation::EUAVMode::Compute, UAVIndex);
 		RHIContext->RHISetUAVParameter(Shader, UAVIndex, UAV);
 	}
@@ -678,7 +671,7 @@ public:
 	*/
 	virtual void RHISetUAVParameter(FRHIComputeShader* Shader, uint32 UAVIndex, FRHIUnorderedAccessView* UAV, uint32 InitialCount) override final
 	{
-		checkf(State.bComputeShaderSet, TEXT("A Compute shader has to be set to set resources into a shader!"));
+		checkf(State.bComputePSOSet, TEXT("A Compute PSO has to be set to set resources into a shader!"));
 		Tracker->AssertUAV(UAV, RHIValidation::EUAVMode::Compute, UAVIndex);
 		RHIContext->RHISetUAVParameter(Shader, UAVIndex, UAV, InitialCount);
 	}
@@ -695,7 +688,7 @@ public:
 
 	virtual void RHISetShaderResourceViewParameter(FRHIComputeShader* Shader, uint32 SamplerIndex, FRHIShaderResourceView* SRV) override final
 	{
-		checkf(State.bComputeShaderSet, TEXT("A Compute shader has to be set to set resources into a shader!"));
+		checkf(State.bComputePSOSet, TEXT("A Compute PSO has to be set to set resources into a shader!"));
 		if (SRV)
 		{
 			Tracker->Assert(SRV->ViewIdentity, ERHIAccess::SRVCompute);
@@ -712,7 +705,7 @@ public:
 
 	virtual void RHISetShaderUniformBuffer(FRHIComputeShader* Shader, uint32 BufferIndex, FRHIUniformBuffer* Buffer) override final
 	{
-		checkf(State.bComputeShaderSet, TEXT("A Compute shader has to be set to set resources into a shader!"));
+		checkf(State.bComputePSOSet, TEXT("A Compute PSO has to be set to set resources into a shader!"));
 		State.GlobalUniformBuffers.ValidateSetShaderUniformBuffer(Buffer);
 		RHIContext->RHISetShaderUniformBuffer(Shader, BufferIndex, Buffer);
 	}
@@ -725,7 +718,7 @@ public:
 
 	virtual void RHISetShaderParameter(FRHIComputeShader* Shader, uint32 BufferIndex, uint32 BaseIndex, uint32 NumBytes, const void* NewValue) override final
 	{
-		checkf(State.bComputeShaderSet, TEXT("A Compute shader has to be set to set resources into a shader!"));
+		checkf(State.bComputePSOSet, TEXT("A Compute PSO has to be set to set resources into a shader!"));
 		RHIContext->RHISetShaderParameter(Shader, BufferIndex, BaseIndex, NumBytes, NewValue);
 	}
 
@@ -1112,7 +1105,7 @@ protected:
 		FString PreviousRenderPassName;
 		FString ComputePassName;
 		bool bGfxPSOSet{};
-		bool bComputeShaderSet{};
+		bool bComputePSOSet{};
 		bool bInsideBeginRenderPass{};
 
 		void Reset();
diff --git a/Engine/Source/Runtime/VulkanRHI/Private/VulkanCommands.cpp b/Engine/Source/Runtime/VulkanRHI/Private/VulkanCommands.cpp
--- a/Engine/Source/Runtime/VulkanRHI/Private/VulkanCommands.cpp
+++ b/Engine/Source/Runtime/VulkanRHI/Private/VulkanCommands.cpp
@@ -115,13 +115,6 @@ void FVulkanCommandListContext::RHISetStreamSource(uint32 StreamIndex, FRHIBuffe
 	}
 }
 
-void FVulkanCommandListContext::RHISetComputeShader(FRHIComputeShader* ComputeShaderRHI)
-{
-	FVulkanComputeShader* ComputeShader = ResourceCast(ComputeShaderRHI);
-	FVulkanComputePipeline* ComputePipeline = Device->GetPipelineStateCache()->GetOrCreateComputePipeline(ComputeShader);
-	RHISetComputePipelineState(ComputePipeline);
-}
-
 void FVulkanCommandListContext::RHISetComputePipelineState(FRHIComputePipelineState* ComputePipelineState)
 {
 	FVulkanCmdBuffer* CmdBuffer = CommandBufferManager->GetActiveCmdBuffer();
diff --git a/Engine/Source/Runtime/VulkanRHI/Private/VulkanContext.h b/Engine/Source/Runtime/VulkanRHI/Private/VulkanContext.h
--- a/Engine/Source/Runtime/VulkanRHI/Private/VulkanContext.h
+++ b/Engine/Source/Runtime/VulkanRHI/Private/VulkanContext.h
@@ -65,7 +65,6 @@ public:
 	virtual void RHIPushEvent(const TCHAR* Name, FColor Color) final override;
 	virtual void RHIPopEvent() final override;
 
-	virtual void RHISetComputeShader(FRHIComputeShader* ComputeShader) final override;
 	virtual void RHISetComputePipelineState(FRHIComputePipelineState* ComputePipelineState) final override;
 	virtual void RHIDispatchComputeShader(uint32 ThreadGroupCountX, uint32 ThreadGroupCountY, uint32 ThreadGroupCountZ) final override;
 	virtual void RHIDispatchIndirectComputeShader(FRHIBuffer* ArgumentBuffer, uint32 ArgumentOffset) final override;
diff --git a/Engine/Source/Runtime/Windows/D3D11RHI/Private/D3D11RHIPrivate.h b/Engine/Source/Runtime/Windows/D3D11RHI/Private/D3D11RHIPrivate.h
--- a/Engine/Source/Runtime/Windows/D3D11RHI/Private/D3D11RHIPrivate.h
+++ b/Engine/Source/Runtime/Windows/D3D11RHI/Private/D3D11RHIPrivate.h
@@ -399,7 +399,7 @@ struct FD3D11Adapter
 PRAGMA_DISABLE_DEPRECATION_WARNINGS
 
 /** The interface which is implemented by the dynamically bound RHI. */
-class D3D11RHI_API FD3D11DynamicRHI : public FDynamicRHI, public IRHICommandContextPSOFallback
+class D3D11RHI_API FD3D11DynamicRHI : public FDynamicRHIPSOFallback, public IRHICommandContextPSOFallback
 {
 public:
 	typedef TMap<FD3D11LockedKey, FD3D11LockedData> FD3D11LockTracker;
