Adding ERHIInterfaceType so the public RHIs can identify themselves.
Adding templated GetDynamicRHI that makes sure we don't cast the Validation DynamicRHI by accident.

diff --git a/Engine/Source/Runtime/Apple/AGXRHI/Private/AGXDynamicRHI.h b/Engine/Source/Runtime/Apple/AGXRHI/Private/AGXDynamicRHI.h
--- a/Engine/Source/Runtime/Apple/AGXRHI/Private/AGXDynamicRHI.h
+++ b/Engine/Source/Runtime/Apple/AGXRHI/Private/AGXDynamicRHI.h
@@ -37,6 +37,7 @@ public:
 	virtual void Init();
 	virtual void Shutdown() {}
 	virtual const TCHAR* GetName() override { return TEXT("AGX"); }
+	virtual ERHIInterfaceType GetInterfaceType() const override { return ERHIInterfaceType::Agx; }
 
 	virtual FRHIShaderLibraryRef RHICreateShaderLibrary(EShaderPlatform Platform, FString const& FilePath, FString const& Name) final override;
 	virtual FSamplerStateRHIRef RHICreateSamplerState(const FSamplerStateInitializerRHI& Initializer) final override;
diff --git a/Engine/Source/Runtime/Apple/MetalRHI/Private/MetalDynamicRHI.h b/Engine/Source/Runtime/Apple/MetalRHI/Private/MetalDynamicRHI.h
--- a/Engine/Source/Runtime/Apple/MetalRHI/Private/MetalDynamicRHI.h
+++ b/Engine/Source/Runtime/Apple/MetalRHI/Private/MetalDynamicRHI.h
@@ -38,6 +38,7 @@ public:
 	virtual void Init();
 	virtual void Shutdown() {}
 	virtual const TCHAR* GetName() override { return TEXT("Metal"); }
+	virtual ERHIInterfaceType GetInterfaceType() const override { return ERHIInterfaceType::Metal; }
 
 	virtual FRHIShaderLibraryRef RHICreateShaderLibrary(EShaderPlatform Platform, FString const& FilePath, FString const& Name) final override;
 	virtual FSamplerStateRHIRef RHICreateSamplerState(const FSamplerStateInitializerRHI& Initializer) final override;
diff --git a/Engine/Source/Runtime/D3D12RHI/Private/D3D12RHIPrivate.h b/Engine/Source/Runtime/D3D12RHI/Private/D3D12RHIPrivate.h
--- a/Engine/Source/Runtime/D3D12RHI/Private/D3D12RHIPrivate.h
+++ b/Engine/Source/Runtime/D3D12RHI/Private/D3D12RHIPrivate.h
@@ -293,6 +293,7 @@ public:
 	virtual void PostInit() override;
 	virtual void Shutdown() override;
 	virtual const TCHAR* GetName() override { return TEXT("D3D12"); }
+	virtual ERHIInterfaceType GetInterfaceType() const override final { return ERHIInterfaceType::D3D12; }
 
 	template<typename TRHIType>
 	static FORCEINLINE typename TD3D12ResourceTraits<TRHIType>::TConcreteType* ResourceCast(TRHIType* Resource)
diff --git a/Engine/Source/Runtime/NullDrv/Public/NullRHI.h b/Engine/Source/Runtime/NullDrv/Public/NullRHI.h
--- a/Engine/Source/Runtime/NullDrv/Public/NullRHI.h
+++ b/Engine/Source/Runtime/NullDrv/Public/NullRHI.h
@@ -20,6 +20,7 @@ public:
 	virtual void Init();
 	virtual void Shutdown();
 	virtual const TCHAR* GetName() override { return TEXT("Null"); }
+	virtual ERHIInterfaceType GetInterfaceType() const override { return ERHIInterfaceType::Null; }
 
 	virtual FSamplerStateRHIRef RHICreateSamplerState(const FSamplerStateInitializerRHI& Initializer) final override
 	{ 
diff --git a/Engine/Source/Runtime/OpenGLDrv/Public/OpenGLDrv.h b/Engine/Source/Runtime/OpenGLDrv/Public/OpenGLDrv.h
--- a/Engine/Source/Runtime/OpenGLDrv/Public/OpenGLDrv.h
+++ b/Engine/Source/Runtime/OpenGLDrv/Public/OpenGLDrv.h
@@ -306,6 +306,7 @@ public:
 
 	virtual void Shutdown();
 	virtual const TCHAR* GetName() override { return TEXT("OpenGL"); }
+	virtual ERHIInterfaceType GetInterfaceType() const override final { return ERHIInterfaceType::OpenGL; }
 
 	// If using a Proxy object return the contained GL object rather than the proxy itself.
 	template<typename TRHIType>
diff --git a/Engine/Source/Runtime/RHI/Public/DynamicRHI.h b/Engine/Source/Runtime/RHI/Public/DynamicRHI.h
--- a/Engine/Source/Runtime/RHI/Public/DynamicRHI.h
+++ b/Engine/Source/Runtime/RHI/Public/DynamicRHI.h
@@ -187,6 +187,9 @@ public:
 
 	virtual const TCHAR* GetName() = 0;
 
+	virtual ERHIInterfaceType GetInterfaceType() const { return ERHIInterfaceType::Hidden; }
+	virtual FDynamicRHI* GetNonValidationRHI() { return this; }
+
 	/** Called after PostInit to initialize the pixel format info, which is needed for some commands default implementations */
 	void InitPixelFormatInfo(const TArray<uint32>& PixelFormatBlockBytesIn)
 	{
@@ -1331,6 +1334,23 @@ public:
 	}
 };
 
+FORCEINLINE ERHIInterfaceType RHIGetInterfaceType()
+{
+	return GDynamicRHI->GetInterfaceType();
+}
+
+template<typename TRHI>
+FORCEINLINE TRHI* CastDynamicRHI(FDynamicRHI* InDynamicRHI)
+{
+	return static_cast<TRHI*>(InDynamicRHI->GetNonValidationRHI());
+}
+
+template<typename TRHI>
+FORCEINLINE TRHI* GetDynamicRHI()
+{
+	return CastDynamicRHI<TRHI>(GDynamicRHI);
+}
+
 FORCEINLINE FSamplerStateRHIRef RHICreateSamplerState(const FSamplerStateInitializerRHI& Initializer)
 {
 	return GDynamicRHI->RHICreateSamplerState(Initializer);
diff --git a/Engine/Source/Runtime/RHI/Public/RHIDefinitions.h b/Engine/Source/Runtime/RHI/Public/RHIDefinitions.h
--- a/Engine/Source/Runtime/RHI/Public/RHIDefinitions.h
+++ b/Engine/Source/Runtime/RHI/Public/RHIDefinitions.h
@@ -24,6 +24,18 @@
 #define RHI_RAYTRACING 0
 #endif
 
+enum class ERHIInterfaceType
+{
+	Hidden,
+	Null,
+	D3D11,
+	D3D12,
+	Vulkan,
+	Metal,
+	Agx,
+	OpenGL,
+};
+
 enum EShaderFrequency : uint8
 {
 	SF_Vertex			= 0,
diff --git a/Engine/Source/Runtime/RHI/Public/RHIValidation.h b/Engine/Source/Runtime/RHI/Public/RHIValidation.h
--- a/Engine/Source/Runtime/RHI/Public/RHIValidation.h
+++ b/Engine/Source/Runtime/RHI/Public/RHIValidation.h
@@ -64,6 +64,16 @@ public:
 		return *RHIName;
 	}
 
+	virtual ERHIInterfaceType GetInterfaceType() const override final
+	{
+		return RHI->GetInterfaceType();
+	}
+
+	virtual FDynamicRHI* GetNonValidationRHI() override final
+	{
+		return RHI;
+	}
+
 	/////// RHI Methods
 
 	// FlushType: Thread safe
diff --git a/Engine/Source/Runtime/VulkanRHI/Public/VulkanDynamicRHI.h b/Engine/Source/Runtime/VulkanRHI/Public/VulkanDynamicRHI.h
--- a/Engine/Source/Runtime/VulkanRHI/Public/VulkanDynamicRHI.h
+++ b/Engine/Source/Runtime/VulkanRHI/Public/VulkanDynamicRHI.h
@@ -56,6 +56,7 @@ public:
 	virtual void PostInit() final override;
 	virtual void Shutdown() final override;;
 	virtual const TCHAR* GetName() final override { return TEXT("Vulkan"); }
+	virtual ERHIInterfaceType GetInterfaceType() const override { return ERHIInterfaceType::Vulkan; }
 
 	void InitInstance();
 
diff --git a/Engine/Source/Runtime/Windows/D3D11RHI/Private/D3D11RHIPrivate.h b/Engine/Source/Runtime/Windows/D3D11RHI/Private/D3D11RHIPrivate.h
--- a/Engine/Source/Runtime/Windows/D3D11RHI/Private/D3D11RHIPrivate.h
+++ b/Engine/Source/Runtime/Windows/D3D11RHI/Private/D3D11RHIPrivate.h
@@ -420,7 +420,8 @@ public:
 	virtual void PostInit() override;
 	virtual void Shutdown() override;
 	virtual const TCHAR* GetName() override { return TEXT("D3D11"); }
-	
+	virtual ERHIInterfaceType GetInterfaceType() const override final { return ERHIInterfaceType::D3D11; }
+
 #if PLATFORM_HOLOLENS
 	virtual void RHISuspendRendering() override;
 	virtual void RHIResumeRendering() override;
