Fix non-unity compile issues with FDebugName.

diff --git a/Engine/Source/Runtime/RHI/Private/RHI.cpp b/Engine/Source/Runtime/RHI/Private/RHI.cpp
--- a/Engine/Source/Runtime/RHI/Private/RHI.cpp
+++ b/Engine/Source/Runtime/RHI/Private/RHI.cpp
@@ -2824,3 +2824,54 @@ void FRHITransientBuffer::Acquire(const TCHAR* InName, uint32 InPassIndex, uint6
 
 	// TODO: Add method to rename a buffer.
 }
+
+FDebugName::FDebugName()
+	: Name()
+	, Number(NAME_NO_NUMBER_INTERNAL)
+{
+}
+
+FDebugName::FDebugName(FName InName)
+	: Name(InName)
+	, Number(NAME_NO_NUMBER_INTERNAL)
+{
+}
+
+FDebugName::FDebugName(FName InName, int32 InNumber)
+	: Name(InName)
+	, Number(InNumber)
+{
+}
+
+FDebugName::FDebugName(FMemoryImageName InName, int32 InNumber)
+	: Name(InName)
+	, Number(InNumber)
+{
+}
+
+FDebugName& FDebugName::operator=(FName Other)
+{
+	Name = FMemoryImageName(Other);
+	Number = NAME_NO_NUMBER_INTERNAL;
+	return *this;
+}
+
+FString FDebugName::ToString() const
+{
+	FString Out;
+	FName(Name).AppendString(Out);
+	if (Number != NAME_NO_NUMBER_INTERNAL)
+	{
+		Out.Appendf(TEXT("_%u"), Number);
+	}
+	return Out;
+}
+
+void FDebugName::AppendString(FStringBuilderBase& Builder) const
+{
+	FName(Name).AppendString(Builder);
+	if (Number != NAME_NO_NUMBER_INTERNAL)
+	{
+		Builder << '_' << Number;
+	}
+}
diff --git a/Engine/Source/Runtime/RHI/Private/RHIResources.cpp b/Engine/Source/Runtime/RHI/Private/RHIResources.cpp
deleted file mode 100644
