MGPU:  Fix for assert when RHI validation is enabled in the Lyra project, due to default virtual implementation of RHILockBufferMGPU being called.  Added pass through implementations for RHILockBufferMGPU / RHIUnlockBufferMGPU to FValidationRHI.  Moved common validation logic to a shared function to avoid cut-and-paste.

#rnx

diff --git a/Engine/Source/Runtime/RHI/Private/RHIValidation.cpp b/Engine/Source/Runtime/RHI/Private/RHIValidation.cpp
--- a/Engine/Source/Runtime/RHI/Private/RHIValidation.cpp
+++ b/Engine/Source/Runtime/RHI/Private/RHIValidation.cpp
@@ -372,7 +372,7 @@ namespace RHIValidation
 	}
 }
 
-void* FValidationRHI::RHILockBuffer(class FRHICommandListImmediate& RHICmdList, FRHIBuffer* Buffer, uint32 Offset, uint32 SizeRHI, EResourceLockMode LockMode)
+void FValidationRHI::LockBufferValidate(class FRHICommandListImmediate& RHICmdList, FRHIBuffer* Buffer, EResourceLockMode LockMode)
 {
 	using namespace RHIValidation;
 
@@ -390,10 +390,22 @@ void* FValidationRHI::RHILockBuffer(class FRHICommandListImmediate& RHICmdList,
 		}
 		RHI_VALIDATION_CHECK(!bIsInsideRenderPass, *GetReasonString_LockBufferInsideRenderPass(Buffer));
 	}
+}
+
+void* FValidationRHI::RHILockBuffer(class FRHICommandListImmediate& RHICmdList, FRHIBuffer* Buffer, uint32 Offset, uint32 SizeRHI, EResourceLockMode LockMode)
+{
+	LockBufferValidate(RHICmdList, Buffer, LockMode);
 
 	return RHI->RHILockBuffer(RHICmdList, Buffer, Offset, SizeRHI, LockMode);
 }
 
+void* FValidationRHI::RHILockBufferMGPU(class FRHICommandListImmediate& RHICmdList, FRHIBuffer* Buffer, uint32 GPUIndex, uint32 Offset, uint32 SizeRHI, EResourceLockMode LockMode)
+{
+	LockBufferValidate(RHICmdList, Buffer, LockMode);
+
+	return RHI->RHILockBufferMGPU(RHICmdList, Buffer, GPUIndex, Offset, SizeRHI, LockMode);
+}
+
 void FValidationRHI::ReportValidationFailure(const TCHAR* InMessage)
 {
 	// Report failures only once per session, since many of them will happen repeatedly. This is similar to what ensure() does, but
diff --git a/Engine/Source/Runtime/RHI/Public/RHIValidation.h b/Engine/Source/Runtime/RHI/Public/RHIValidation.h
--- a/Engine/Source/Runtime/RHI/Public/RHIValidation.h
+++ b/Engine/Source/Runtime/RHI/Public/RHIValidation.h
@@ -391,12 +391,17 @@ public:
 
 	// FlushType: Flush RHI Thread
 	virtual void* RHILockBuffer(class FRHICommandListImmediate& RHICmdList, FRHIBuffer* Buffer, uint32 Offset, uint32 SizeRHI, EResourceLockMode LockMode) override final;
+	virtual void* RHILockBufferMGPU(class FRHICommandListImmediate& RHICmdList, FRHIBuffer* Buffer, uint32 GPUIndex, uint32 Offset, uint32 SizeRHI, EResourceLockMode LockMode) override final;
 
 	// FlushType: Flush RHI Thread
 	virtual void RHIUnlockBuffer(class FRHICommandListImmediate& RHICmdList, FRHIBuffer* Buffer) override final
 	{
 		RHI->RHIUnlockBuffer(RHICmdList, Buffer);
 	}
+	virtual void RHIUnlockBufferMGPU(class FRHICommandListImmediate& RHICmdList, FRHIBuffer* Buffer, uint32 GPUIndex) override final
+	{
+		RHI->RHIUnlockBufferMGPU(RHICmdList, Buffer, GPUIndex);
+	}
 
 	/** Creates an unordered access view of the given structured buffer. */
 	// FlushType: Wait RHI Thread
@@ -1750,6 +1755,9 @@ private:
 	static FCriticalSection		SeenFailureHashesMutex;
 
 	void ValidatePipeline(const FGraphicsPipelineStateInitializer& Initializer);
+
+	// Shared validation logic, called from RHILockBuffer / RHILockBufferMGPU
+	void LockBufferValidate(class FRHICommandListImmediate& RHICmdList, FRHIBuffer* Buffer, EResourceLockMode LockMode);
 };
 
 #endif	// ENABLE_RHI_VALIDATION
