Merging DynamicRHI.h/cpp to fix up bad prior merge that will leave duplicate symbols in future.

diff --git a/Engine/Source/Runtime/RHI/Private/DynamicRHI.cpp b/Engine/Source/Runtime/RHI/Private/DynamicRHI.cpp
--- a/Engine/Source/Runtime/RHI/Private/DynamicRHI.cpp
+++ b/Engine/Source/Runtime/RHI/Private/DynamicRHI.cpp
@@ -668,3 +668,54 @@ void FShaderResourceViewInitializer::InitType()
 		}
 	}
 }
+
+FDebugName::FDebugName()
+	: Name()
+	, Number(NAME_NO_NUMBER_INTERNAL)
+{
+}
+
+FDebugName::FDebugName(FName InName)
+	: Name(InName)
+	, Number(NAME_NO_NUMBER_INTERNAL)
+{
+}
+
+FDebugName::FDebugName(FName InName, int32 InNumber)
+	: Name(InName)
+	, Number(InNumber)
+{
+}
+
+FDebugName::FDebugName(FMemoryImageName InName, int32 InNumber)
+	: Name(InName)
+	, Number(InNumber)
+{
+}
+
+FDebugName& FDebugName::operator=(FName Other)
+{
+	Name = FMemoryImageName(Other);
+	Number = NAME_NO_NUMBER_INTERNAL;
+	return *this;
+}
+
+FString FDebugName::ToString() const
+{
+	FString Out;
+	FName(Name).AppendString(Out);
+	if (Number != NAME_NO_NUMBER_INTERNAL)
+	{
+		Out.Appendf(TEXT("_%u"), Number);
+	}
+	return Out;
+}
+
+void FDebugName::AppendString(FStringBuilderBase& Builder) const
+{
+	FName(Name).AppendString(Builder);
+	if (Number != NAME_NO_NUMBER_INTERNAL)
+	{
+		Builder << '_' << Number;
+	}
+}
diff --git a/Engine/Source/Runtime/RHI/Public/DynamicRHI.h b/Engine/Source/Runtime/RHI/Public/DynamicRHI.h
--- a/Engine/Source/Runtime/RHI/Public/DynamicRHI.h
+++ b/Engine/Source/Runtime/RHI/Public/DynamicRHI.h
@@ -80,6 +80,42 @@ struct FRHIFlipDetails
 	{}
 };
 
+/**
+ * A type used only for printing a string for debugging/profiling.
+ * Adds Number as a suffix to the printed string even if the base name includes a number, so may prints a string like: Base_1_1
+ * This type will always store a numeric suffix explicitly inside itself and never in the name table so it will always be at least 12 bytes
+ * regardless of the value of UE_FNAME_OUTLINE_NUMBER.
+ * It is not comparable or convertible to other name types to encourage its use only for debugging and avoid using more storage than necessary
+ * for the primary use cases of FName (names of objects, assets etc which are widely used and therefor deduped in the name table).
+ */
+class RHI_API FDebugName
+{
+	DECLARE_INLINE_TYPE_LAYOUT(FDebugName, NonVirtual);
+
+public:
+	FDebugName();
+	FDebugName(FName InName);
+	FDebugName(FName InName, int32 InNumber);
+	FDebugName(FMemoryImageName InName, int32 InNumber);
+
+	FDebugName& operator=(FName Other);
+
+	FString ToString() const;
+	bool IsNone() const { return Name.IsNone() && Number == NAME_NO_NUMBER_INTERNAL; }
+	void AppendString(FStringBuilderBase& Builder) const;
+
+private:
+	LAYOUT_FIELD(FMemoryImageName, Name);
+	LAYOUT_FIELD(uint32, Number);
+};
+
+inline FStringBuilderBase& operator<<(FStringBuilderBase& Builder, const FDebugName& Name)
+{
+	Name.AppendString(Builder);
+	return Builder;
+}
+
+
 struct FShaderResourceViewInitializer
 {
 	struct FBufferShaderResourceViewInitializer
