Test for wide name strings and handle conversion to an ascii string if needed.

#rnx

diff --git a/Engine/Source/Runtime/RHI/Private/TextureProfiler.cpp b/Engine/Source/Runtime/RHI/Private/TextureProfiler.cpp
--- a/Engine/Source/Runtime/RHI/Private/TextureProfiler.cpp
+++ b/Engine/Source/Runtime/RHI/Private/TextureProfiler.cpp
@@ -6,6 +6,7 @@
 #include "RHI.h"
 #include "RHIResources.h"
 #include "UObject/NameTypes.h"
+#include "Containers/StringConv.h"
 
 #if TEXTURE_PROFILER_ENABLED
 static int32 ToMB(uint64 Value)
@@ -234,15 +235,26 @@ void FTextureProfiler::UpdateTextureName(FRHITexture* UniqueTexturePtr)
 
 const char* FTextureProfiler::GetTextureNameString(FName TextureName)
 {
-	char*& TextureNameString = TextureNameStrings.FindOrAdd(TextureName, nullptr);
+	const char*& TextureNameString = TextureNameStrings.FindOrAdd(TextureName, nullptr);
 
 	if (TextureNameString == nullptr)
 	{
-		char NewTextureNameString[NAME_SIZE];
-		TextureName.GetPlainANSIString(NewTextureNameString);
-		int Length = FCStringAnsi::Strlen(NewTextureNameString);
-		TextureNameString = new char[Length + 1];
-		FCStringAnsi::Strcpy(TextureNameString, Length + 1, NewTextureNameString);
+		if(TextureName.GetDisplayNameEntry()->IsWide())
+		{
+			WIDECHAR NewTextureNameString[NAME_SIZE];
+			TextureName.GetPlainWIDEString(NewTextureNameString);
+			int Length = StringCast<ANSICHAR>(NewTextureNameString).Length();
+			TextureNameString = new char[Length + 1];
+			FCStringAnsi::Strcpy(const_cast<char*>(TextureNameString), Length + 1, StringCast<ANSICHAR>(NewTextureNameString).Get());
+		}
+		else
+		{
+			char NewTextureNameString[NAME_SIZE];
+			TextureName.GetPlainANSIString(NewTextureNameString);
+			int Length = FCStringAnsi::Strlen(NewTextureNameString);
+			TextureNameString = new char[Length + 1];
+			FCStringAnsi::Strcpy(const_cast<char*>(TextureNameString), Length + 1, NewTextureNameString);
+		}
 	}
 
 	return TextureNameString;
diff --git a/Engine/Source/Runtime/RHI/Public/TextureProfiler.h b/Engine/Source/Runtime/RHI/Public/TextureProfiler.h
--- a/Engine/Source/Runtime/RHI/Public/TextureProfiler.h
+++ b/Engine/Source/Runtime/RHI/Public/TextureProfiler.h
@@ -33,7 +33,6 @@ public:
 	void RemoveTextureAllocation(FRHITexture* UniqueTexturePtr);
 	void UpdateTextureName(FRHITexture* UniqueTexturePtr);
 
-	const char* GetTextureNameString(FName TextureName);
 private:
 
 	FTextureProfiler() = default;
@@ -42,6 +41,9 @@ private:
 
 	void Update();
 
+	// Should only be used by FTextureDetails
+	const char* GetTextureNameString(FName TextureName);
+
 	FCriticalSection TextureMapCS;
 
 	class FTexureDetails
@@ -52,12 +54,11 @@ private:
 		~FTexureDetails();
 		
 		void SetName(FName InTextureName);
-		const char* GetNameString() const { return TextureNameString; }
 
 		void ResetPeakSize();
 		void SetValues(const FTexureDetails& Values);
 		FName GetTextureName() const { return TextureName; }
-		const char* GetTextureNameString() const { return TextureNameString; }
+		const char* GetTextureNameString() const { check(TextureNameString != nullptr);  return TextureNameString; }
 
 		FTexureDetails& operator+=(const FTexureDetails& Other);
 		FTexureDetails& operator-=(const FTexureDetails& Other);
@@ -77,7 +78,8 @@ private:
 		const char* TextureNameString = nullptr;
 	};
 
-	TMap<FName, char*> TextureNameStrings;
+	// Allocated so that a reference to the string can be stored in FTextureDetails
+	TMap<FName, const char*> TextureNameStrings;
 	TMap<void*, FTexureDetails> TexturesMap;
 
 	// Keep track of the totals separately to reduce the cost of rounding error for sizes
@@ -85,5 +87,6 @@ private:
 	FTexureDetails TotalRenderTargetSize;
 	TMap<FName, FTexureDetails> CombinedTextureSizes;
 	TMap<FName, FTexureDetails> CombinedRenderTargetSizes;
+	friend class FTexureDetails;
 };
 #endif //TEXTURE_PROFILER_ENABLED
