RHI function (and DataDrivenPlatform parameter) for whether shader atomics are supported for volume textures

diff --git a/Engine/Config/IOS/DataDrivenPlatformInfo.ini b/Engine/Config/IOS/DataDrivenPlatformInfo.ini
--- a/Engine/Config/IOS/DataDrivenPlatformInfo.ini
+++ b/Engine/Config/IOS/DataDrivenPlatformInfo.ini
@@ -52,6 +52,7 @@ bIsMobile=true
 bTargetsTiledGPU=true
 bSupports4ComponentUAVReadWrite=true
 bSupportsManualVertexFetch=false
+bSupportsVolumeTextureAtomics=false
 
 [ShaderPlatform METAL_MRT]
 Language=Metal
@@ -61,3 +62,4 @@ bTargetsTiledGPU=true
 bSupportsIndexBufferUAVs=true
 bSupports4ComponentUAVReadWrite=true
 bSupportsManualVertexFetch=false
+bSupportsVolumeTextureAtomics=false
diff --git a/Engine/Config/Mac/DataDrivenPlatformInfo.ini b/Engine/Config/Mac/DataDrivenPlatformInfo.ini
--- a/Engine/Config/Mac/DataDrivenPlatformInfo.ini
+++ b/Engine/Config/Mac/DataDrivenPlatformInfo.ini
@@ -46,6 +46,7 @@ bSupportsVolumetricFog=true
 bSupportsPercentageCloserShadows=true
 bSupportsCapsuleShadows=true
 bSupports4ComponentUAVReadWrite=true
+bSupportsVolumeTextureAtomics=false
 
 [ShaderPlatform METAL_MRT_MAC]
 Language=Metal
@@ -55,6 +56,7 @@ bIsMetalMRT=true
 bSupportsIndexBufferUAVs=true
 bSupports4ComponentUAVReadWrite=true
 bSupportsManualVertexFetch=true
+bSupportsVolumeTextureAtomics=false
 
 [ShaderPlatform METAL_MACES3_1]
 Language=Metal
@@ -63,3 +65,4 @@ bIsPC=true
 bIsMobile=true
 bSupports4ComponentUAVReadWrite=true
 bSupportsManualVertexFetch=false
+bSupportsVolumeTextureAtomics=false
diff --git a/Engine/Source/Runtime/RHI/Private/RHI.cpp b/Engine/Source/Runtime/RHI/Private/RHI.cpp
--- a/Engine/Source/Runtime/RHI/Private/RHI.cpp
+++ b/Engine/Source/Runtime/RHI/Private/RHI.cpp
@@ -2258,6 +2258,7 @@ void FGenericDataDrivenShaderPlatformInfo::SetDefaultValues()
 	bSupportsGPUSkinCache = true;
 	bSupportsManualVertexFetch = true;
 	bMobileSupportFetchBindedCustomStencilBuffer = false;
+	bSupportsVolumeTextureAtomics = true;
 }
 
 void FGenericDataDrivenShaderPlatformInfo::ParseDataDrivenShaderInfo(const FConfigSection& Section, FGenericDataDrivenShaderPlatformInfo& Info)
@@ -2347,6 +2348,7 @@ void FGenericDataDrivenShaderPlatformInfo::ParseDataDrivenShaderInfo(const FConf
 	GET_SECTION_BOOL_HELPER(bSupportsVertexShaderLayer);
 	GET_SECTION_BOOL_HELPER(bMobileSupportFetchBindedCustomStencilBuffer);
 	GET_SECTION_BOOL_HELPER(bSupportsBindless);
+	GET_SECTION_BOOL_HELPER(bSupportsVolumeTextureAtomics);
 #undef GET_SECTION_BOOL_HELPER
 #undef GET_SECTION_INT_HELPER
 
diff --git a/Engine/Source/Runtime/RHI/Public/RHI.h b/Engine/Source/Runtime/RHI/Public/RHI.h
--- a/Engine/Source/Runtime/RHI/Public/RHI.h
+++ b/Engine/Source/Runtime/RHI/Public/RHI.h
@@ -274,6 +274,11 @@ inline bool RHISupportsBindless(EShaderPlatform Platform)
 	return FDataDrivenShaderPlatformInfo::GetSupportsBindless(Platform);
 }
 
+inline bool RHISupportsVolumeTextureAtomics(EShaderPlatform Platform)
+{
+	return FDataDrivenShaderPlatformInfo::GetSupportsVolumeTextureAtomics(Platform);
+}
+
 // Wrapper for GRHI## global variables, allows values to be overridden for mobile preview modes.
 template <typename TValueType>
 class TRHIGlobal
diff --git a/Engine/Source/Runtime/RHI/Public/RHIDefinitions.h b/Engine/Source/Runtime/RHI/Public/RHIDefinitions.h
--- a/Engine/Source/Runtime/RHI/Public/RHIDefinitions.h
+++ b/Engine/Source/Runtime/RHI/Public/RHIDefinitions.h
@@ -396,6 +396,7 @@ class RHI_API FGenericDataDrivenShaderPlatformInfo
 	uint32 bSupportsVertexShaderLayer : 1;
 	uint32 bMobileSupportFetchBindedCustomStencilBuffer : 1;
 	uint32 bSupportsBindless : 1;
+	uint32 bSupportsVolumeTextureAtomics : 1;
 
 		
 #if WITH_EDITOR
@@ -915,6 +916,11 @@ public:
 		return Infos[Platform].bSupportsBindless;
 	}
 
+	static FORCEINLINE_DEBUGGABLE const bool GetSupportsVolumeTextureAtomics(const FStaticShaderPlatform Platform)
+	{
+		return Infos[Platform].bSupportsVolumeTextureAtomics;
+	}
+
 #if WITH_EDITOR
 	static FORCEINLINE_DEBUGGABLE FText GetFriendlyName(const FStaticShaderPlatform Platform)
 	{
