Fix a memory stomp caused by writing to a released FLastRenderTimeContainer

#localization none
#tests ran ReplayRun tests on multiple platforms

diff --git a/Engine/Source/Runtime/RHI/Public/RHICommandList.h b/Engine/Source/Runtime/RHI/Public/RHICommandList.h
--- a/Engine/Source/Runtime/RHI/Public/RHICommandList.h
+++ b/Engine/Source/Runtime/RHI/Public/RHICommandList.h
@@ -5344,9 +5344,9 @@ FORCEINLINE void RHIUpdateRHIResources(FRHIResourceUpdateInfo* UpdateInfos, int3
 	return FRHICommandListExecutor::GetImmediateCommandList().UpdateRHIResources(UpdateInfos, Num, bNeedReleaseRefs);
 }
 
-FORCEINLINE FTextureReferenceRHIRef RHICreateTextureReference(FLastRenderTimeContainer* LastRenderTime)
+FORCEINLINE FTextureReferenceRHIRef RHICreateTextureReference()
 {
-	return new FRHITextureReference(LastRenderTime);
+	return new FRHITextureReference();
 }
 
 FORCEINLINE void RHIUpdateTextureReference(FRHITextureReference* TextureRef, FRHITexture* NewTexture)
diff --git a/Engine/Source/Runtime/RHI/Public/RHIResources.h b/Engine/Source/Runtime/RHI/Public/RHIResources.h
--- a/Engine/Source/Runtime/RHI/Public/RHIResources.h
+++ b/Engine/Source/Runtime/RHI/Public/RHIResources.h
@@ -1236,14 +1236,13 @@ class RHI_API FRHITexture : public FRHIResource
 public:
 	
 	/** Initialization constructor. */
-	FRHITexture(ERHIResourceType InResourceType, uint32 InNumMips, uint32 InNumSamples, EPixelFormat InFormat, ETextureCreateFlags InFlags, FLastRenderTimeContainer* InLastRenderTime, const FClearValueBinding& InClearValue)
+	FRHITexture(ERHIResourceType InResourceType, uint32 InNumMips, uint32 InNumSamples, EPixelFormat InFormat, ETextureCreateFlags InFlags, const FClearValueBinding& InClearValue)
 		: FRHIResource(InResourceType)
 		, ClearValue(InClearValue)
 		, NumMips(InNumMips)
 		, NumSamples(InNumSamples)
 		, Format(InFormat)
 		, Flags(InFlags)
-	, LastRenderTime(InLastRenderTime ? *InLastRenderTime : DefaultLastRenderTime)	
 	{}
 
 	// Dynamic cast methods.
@@ -1328,19 +1327,9 @@ public:
 		LastRenderTime.SetLastRenderTime(InLastRenderTime);
 	}
 
-	/** Returns the last render time container, or NULL if none were specified at creation. */
-	FLastRenderTimeContainer* GetLastRenderTimeContainer()
+	double GetLastRenderTime() const
 	{
-		if (&LastRenderTime == &DefaultLastRenderTime)
-		{
-			return NULL;
-		}
-		return &LastRenderTime;
-	}
-
-	FORCEINLINE_DEBUGGABLE void SetDefaultLastRenderTimeContainer()
-	{
-		LastRenderTime = DefaultLastRenderTime;
+		return LastRenderTime.GetLastRenderTime();
 	}
 
 	void SetName(const FName& InName)
@@ -1405,8 +1394,7 @@ private:
 	uint32 NumSamples;
 	EPixelFormat Format;
 	ETextureCreateFlags Flags;
-	FLastRenderTimeContainer& LastRenderTime;
-	FLastRenderTimeContainer DefaultLastRenderTime;	
+	FLastRenderTimeContainer LastRenderTime;
 	FName TextureName;
 };
 
@@ -1416,7 +1404,7 @@ public:
 	
 	/** Initialization constructor. */
 	FRHITexture2D(uint32 InSizeX,uint32 InSizeY,uint32 InNumMips,uint32 InNumSamples,EPixelFormat InFormat,ETextureCreateFlags InFlags, const FClearValueBinding& InClearValue, ERHIResourceType InResourceTypeOverride = RRT_None)
-	: FRHITexture(InResourceTypeOverride != RRT_None ? InResourceTypeOverride : RRT_Texture2D, InNumMips, InNumSamples, InFormat, InFlags, NULL, InClearValue)
+	: FRHITexture(InResourceTypeOverride != RRT_None ? InResourceTypeOverride : RRT_Texture2D, InNumMips, InNumSamples, InFormat, InFlags, InClearValue)
 	, SizeX(InSizeX)
 	, SizeY(InSizeY)
 	{}
@@ -1493,7 +1481,7 @@ public:
 	
 	/** Initialization constructor. */
 	FRHITexture3D(uint32 InSizeX,uint32 InSizeY,uint32 InSizeZ,uint32 InNumMips,EPixelFormat InFormat,ETextureCreateFlags InFlags, const FClearValueBinding& InClearValue)
-	: FRHITexture(RRT_Texture3D, InNumMips,1,InFormat,InFlags,NULL, InClearValue)
+	: FRHITexture(RRT_Texture3D, InNumMips,1,InFormat,InFlags, InClearValue)
 	, SizeX(InSizeX)
 	, SizeY(InSizeY)
 	, SizeZ(InSizeZ)
@@ -1529,7 +1517,7 @@ public:
 	
 	/** Initialization constructor. */
 	FRHITextureCube(uint32 InSize,uint32 InNumMips,EPixelFormat InFormat,ETextureCreateFlags InFlags, const FClearValueBinding& InClearValue)
-	: FRHITexture(RRT_TextureCube, InNumMips,1,InFormat,InFlags,NULL, InClearValue)
+	: FRHITexture(RRT_TextureCube, InNumMips,1,InFormat,InFlags, InClearValue)
 	, Size(InSize)
 	{}
 	
@@ -1552,8 +1540,8 @@ private:
 class RHI_API FRHITextureReference final : public FRHITexture
 {
 public:
-	explicit FRHITextureReference(FLastRenderTimeContainer* InLastRenderTime)
-		: FRHITexture(RRT_TextureReference, 0, 0, PF_Unknown, TexCreate_None, InLastRenderTime, FClearValueBinding())
+	explicit FRHITextureReference()
+		: FRHITexture(RRT_TextureReference, 0, 0, PF_Unknown, TexCreate_None, FClearValueBinding())
 	{
 		check(DefaultTexture);
 		ReferencedTexture = DefaultTexture;
diff --git a/Engine/Source/Runtime/RenderCore/Private/RenderResource.cpp b/Engine/Source/Runtime/RenderCore/Private/RenderResource.cpp
--- a/Engine/Source/Runtime/RenderCore/Private/RenderResource.cpp
+++ b/Engine/Source/Runtime/RenderCore/Private/RenderResource.cpp
@@ -309,15 +309,28 @@ void FTextureReference::BeginRelease_GameThread()
 	bInitialized_GameThread = false;
 }
 
+double FTextureReference::GetLastRenderTime() const
+{
+	if (bInitialized_GameThread && TextureReferenceRHI)
+	{
+		return TextureReferenceRHI->GetLastRenderTime();
+	}
+
+	return FLastRenderTimeContainer().GetLastRenderTime();
+}
+
 void FTextureReference::InvalidateLastRenderTime()
 {
-	LastRenderTimeRHI.SetLastRenderTime(-FLT_MAX);
+	if (bInitialized_GameThread && TextureReferenceRHI)
+	{
+		TextureReferenceRHI->SetLastRenderTime(-FLT_MAX);
+	}
 }
 
 void FTextureReference::InitRHI()
 {
 	SCOPED_LOADTIMER(FTextureReference_InitRHI);
-	TextureReferenceRHI = RHICreateTextureReference(&LastRenderTimeRHI);
+	TextureReferenceRHI = RHICreateTextureReference();
 }
 	
 int32 GTextureReferenceRevertsLastRenderContainer = 1;
@@ -328,44 +341,6 @@ FAutoConsoleVariableRef CVarTextureReferenceRevertsLastRenderContainer(
 
 void FTextureReference::ReleaseRHI()
 {
-#if PLATFORM_ANDROID 
-	if (TextureReferenceRHI.GetReference())
-	{
-		bool bTextureReferenceRevertsLastRenderContainer = GTextureReferenceRevertsLastRenderContainer != 0;
-		
-		// Check Android's config rules system so we can HF this at startup if needed.
-		{			
-			static bool bConfigRulesChecked = false;
-			static TOptional<bool> bConfigRulesRevertsLastRenderContainer;
-			if (!bConfigRulesChecked)
-			{
-				const FString* ConfigRulesStr = FAndroidMisc::GetConfigRulesVariable(TEXT("TextureReferenceRevertsLastRenderContainer"));
-				if (ConfigRulesStr)
-				{
-					bConfigRulesRevertsLastRenderContainer = ConfigRulesStr->Equals("true", ESearchCase::IgnoreCase);
-					UE_LOG(LogRHI, Log, TEXT("TextureReferenceRevertsLastRenderContainer, set by config rules: %d"), (int)bConfigRulesRevertsLastRenderContainer.GetValue());
-				}
-				else
-				{
-					UE_LOG(LogRHI, Log, TEXT("TextureReferenceRevertsLastRenderContainer, no config rule set: %d"), (int)bTextureReferenceRevertsLastRenderContainer);
-				}
-				bConfigRulesChecked = true;
-			}
-
-			if (bConfigRulesRevertsLastRenderContainer.IsSet())
-			{
-				bTextureReferenceRevertsLastRenderContainer = bConfigRulesRevertsLastRenderContainer.GetValue();
-			}
-		} 
-
-		if (bTextureReferenceRevertsLastRenderContainer && TextureReferenceRHI->GetLastRenderTimeContainer() == &LastRenderTimeRHI)
-		{
-			// we're going away, TextureReferenceRHI must swap out its (soon to be) dangling ref.
-			TextureReferenceRHI->SetDefaultLastRenderTimeContainer();
-		}
-	}
-#endif
-
 	TextureReferenceRHI.SafeRelease();
 }
 
diff --git a/Engine/Source/Runtime/RenderCore/Public/RenderResource.h b/Engine/Source/Runtime/RenderCore/Public/RenderResource.h
--- a/Engine/Source/Runtime/RenderCore/Public/RenderResource.h
+++ b/Engine/Source/Runtime/RenderCore/Public/RenderResource.h
@@ -507,9 +507,6 @@ public:
 
 
 private:
-	/** The last time the texture has been rendered via this reference. */
-	FLastRenderTimeContainer LastRenderTimeRHI;
-
 	/** True if the texture reference has been initialized from the game thread. */
 	bool bInitialized_GameThread;
 
@@ -521,7 +518,7 @@ public:
 	virtual ~FTextureReference();
 
 	/** Returns the last time the texture has been rendered via this reference. */
-	double GetLastRenderTime() const { return LastRenderTimeRHI.GetLastRenderTime(); }
+	double GetLastRenderTime() const;
 
 	/** Invalidates the last render time. */
 	void InvalidateLastRenderTime();
diff --git a/Engine/Source/Runtime/Renderer/Private/VT/VirtualTextureSpace.cpp b/Engine/Source/Runtime/Renderer/Private/VT/VirtualTextureSpace.cpp
--- a/Engine/Source/Runtime/Renderer/Private/VT/VirtualTextureSpace.cpp
+++ b/Engine/Source/Runtime/Renderer/Private/VT/VirtualTextureSpace.cpp
@@ -106,9 +106,9 @@ void FVirtualTextureSpace::InitRHI()
 	for (uint32 TextureIndex = 0u; TextureIndex < GetNumPageTableTextures(); ++TextureIndex)
 	{
 		FTextureEntry& TextureEntry = PageTable[TextureIndex];
-		TextureEntry.TextureReferenceRHI = RHICreateTextureReference(nullptr);
+		TextureEntry.TextureReferenceRHI = RHICreateTextureReference();
 	}
-	PageTableIndirection.TextureReferenceRHI = RHICreateTextureReference(nullptr);
+	PageTableIndirection.TextureReferenceRHI = RHICreateTextureReference();
 	RHIUpdateTextureReference(PageTableIndirection.TextureReferenceRHI, GBlackUintTexture->TextureRHI);
 }
 
