Fix PSO cache not being picked up.

- PSO cache was not being picked up due to a "race condition" when mounting pak files. What seems to have been happening is that we would open the shader library at the time when one of the chunks (other than chunk 0, which holds the PSO cache) got mounted, then proceed to open the PSO cache, and failed to do so (because chunk 0 isn't yet available).  However, since we have the user cache on PC, we didn't completely fail to open it. Instead, the call treated it as if we didn't have a bundled one, and created only the user part to write to.  That made the call to open the PSO cache non-reenterable, as subsequent attempts to open the bundled cache (when chunk0 was finally mounted) were ignored.

- Also set name to stable.upipelinecache by default as this is the name produced during the cook by default.

diff --git a/Engine/Source/Runtime/RHI/Private/PipelineFileCache.cpp b/Engine/Source/Runtime/RHI/Private/PipelineFileCache.cpp
--- a/Engine/Source/Runtime/RHI/Private/PipelineFileCache.cpp
+++ b/Engine/Source/Runtime/RHI/Private/PipelineFileCache.cpp
@@ -2043,7 +2043,7 @@ public:
 		FString GamePath = FPaths::ProjectContentDir() / TEXT("PipelineCaches") / ANSI_TO_TCHAR(FPlatformProperties::IniPlatformName()) / FString::Printf(TEXT("%s_%s.upipelinecache"), *FileName, *PlatformName.ToString());
 
 		static bool bCommandLineNotStable = FParse::Param(FCommandLine::Get(), TEXT("nostablepipelinecache"));
-		if (!bCommandLineNotStable && IFileManager::Get().FileExists(*GamePathStable))
+		if (!bCommandLineNotStable)
 		{
 			GamePath = GamePathStable;
 		}
@@ -2239,7 +2239,7 @@ public:
                 FString GamePathStable = FPaths::ProjectContentDir() / TEXT("PipelineCaches") / ANSI_TO_TCHAR(FPlatformProperties::IniPlatformName()) / FString::Printf(TEXT("%s_%s.stable.upipelinecache"), *Name, *PlatformName.ToString());
                 FString GamePath = FPaths::ProjectContentDir() / TEXT("PipelineCaches") / ANSI_TO_TCHAR(FPlatformProperties::IniPlatformName()) / FString::Printf(TEXT("%s_%s.upipelinecache"), *Name, *PlatformName.ToString());
 				static bool bCommandLineNotStable = FParse::Param(FCommandLine::Get(), TEXT("nostablepipelinecache"));
-                if (!bCommandLineNotStable && IFileManager::Get().FileExists(*GamePathStable))
+                if (!bCommandLineNotStable)
                 {
                     GamePath = GamePathStable;
                 }
