Reinstate some public interfaces but mark them as deprecated

diff --git a/Engine/Source/Runtime/RHI/Public/RHICommandList.h b/Engine/Source/Runtime/RHI/Public/RHICommandList.h
--- a/Engine/Source/Runtime/RHI/Public/RHICommandList.h
+++ b/Engine/Source/Runtime/RHI/Public/RHICommandList.h
@@ -5349,6 +5349,12 @@ FORCEINLINE FTextureReferenceRHIRef RHICreateTextureReference()
 	return new FRHITextureReference();
 }
 
+UE_DEPRECATED(5.0, "The LastRenderTime parameter will be removed in the future")
+FORCEINLINE FTextureReferenceRHIRef RHICreateTextureReference(FLastRenderTimeContainer* LastRenderTime)
+{
+	return RHICreateTextureReference();
+}
+
 FORCEINLINE void RHIUpdateTextureReference(FRHITextureReference* TextureRef, FRHITexture* NewTexture)
 {
 	check(IsInRenderingThread());
diff --git a/Engine/Source/Runtime/RHI/Public/RHIResources.h b/Engine/Source/Runtime/RHI/Public/RHIResources.h
--- a/Engine/Source/Runtime/RHI/Public/RHIResources.h
+++ b/Engine/Source/Runtime/RHI/Public/RHIResources.h
@@ -1245,6 +1245,11 @@ public:
 		, Flags(InFlags)
 	{}
 
+	UE_DEPRECATED(5.0, "The InLastRenderTime parameter will be removed in the future")
+	FRHITexture(ERHIResourceType InResourceType, uint32 InNumMips, uint32 InNumSamples, EPixelFormat InFormat, ETextureCreateFlags InFlags, FLastRenderTimeContainer* InLastRenderTime, const FClearValueBinding& InClearValue)
+		: FRHITexture(InResourceType, InNumMips, InNumSamples, InFormat, InFlags, InClearValue)
+	{}
+
 	// Dynamic cast methods.
 	virtual class FRHITexture2D* GetTexture2D() { return NULL; }
 	virtual class FRHITexture2DArray* GetTexture2DArray() { return NULL; }
@@ -1332,6 +1337,18 @@ public:
 		return LastRenderTime.GetLastRenderTime();
 	}
 
+	/** Returns the last render time container, or NULL if none were specified at creation. */
+	UE_DEPRECATED(5.0, "GetLastRenderTimeContainer is deprecated and will be removed in the future")
+	FLastRenderTimeContainer* GetLastRenderTimeContainer()
+	{
+		return nullptr;
+	}
+
+	UE_DEPRECATED(5.0, "SetDefaultLastRenderTimeContainer is deprecated and will be removed in the future")
+	FORCEINLINE_DEBUGGABLE void SetDefaultLastRenderTimeContainer()
+	{
+	}
+
 	void SetName(const FName& InName)
 	{
 		TextureName = InName;
@@ -1547,6 +1564,11 @@ public:
 		ReferencedTexture = DefaultTexture;
 	}
 
+	UE_DEPRECATED(5.0, "The InLastRenderTime parameter will be removed in the future")
+	explicit FRHITextureReference(FLastRenderTimeContainer* InLastRenderTime)
+		: FRHITextureReference()
+	{}
+
 	virtual class FRHITextureReference* GetTextureReference() override
 	{
 		return this;
