Enhancing Windows RHI selection to choose SM6 when it's enabled in the project's TargetedRHIs list. D3D12 can now fall back to SM5 if SM6 is not supported.
Adding Feature Level to the D3D12 RHI name, useful for the editor tooltip.
Fixing WindowsDynamicRHI config loads to use GConfig instead of FConfigCacheIni::LoadLocalIniFile so we don't miss changes made during startup.

diff --git a/Engine/Source/Runtime/D3D12RHI/Private/D3D12RHIPrivate.h b/Engine/Source/Runtime/D3D12RHI/Private/D3D12RHIPrivate.h
--- a/Engine/Source/Runtime/D3D12RHI/Private/D3D12RHIPrivate.h
+++ b/Engine/Source/Runtime/D3D12RHI/Private/D3D12RHIPrivate.h
@@ -1284,7 +1284,8 @@ public:
 	virtual void ShutdownModule() override;
 
 	// IDynamicRHIModule
-	virtual bool IsSupported() override;
+	virtual bool IsSupported() override { return IsSupported(ERHIFeatureLevel::SM5); }
+	virtual bool IsSupported(ERHIFeatureLevel::Type RequestedFeatureLevel) override;
 	virtual FDynamicRHI* CreateRHI(ERHIFeatureLevel::Type RequestedFeatureLevel = ERHIFeatureLevel::Num) override;
 
 private:
diff --git a/Engine/Source/Runtime/D3D12RHI/Private/Windows/WindowsD3D12Device.cpp b/Engine/Source/Runtime/D3D12RHI/Private/Windows/WindowsD3D12Device.cpp
--- a/Engine/Source/Runtime/D3D12RHI/Private/Windows/WindowsD3D12Device.cpp
+++ b/Engine/Source/Runtime/D3D12RHI/Private/Windows/WindowsD3D12Device.cpp
@@ -372,19 +372,58 @@ static bool SupportsHDROutput(FD3D12DynamicRHI* D3DRHI)
 	return bSupportsHDROutput;
 }
 
-static bool ShouldEnableExperimentalShaderModels()
+static bool IsAdapterBlocked(FD3D12Adapter* InAdapter)
 {
 #if !UE_BUILD_SHIPPING
-	if (CVarExperimentalShaderModels.GetValueOnAnyThread() == 1)
+	if (InAdapter)
 	{
-		return true;
+		FString BlockedIHVString;
+		if (GConfig->GetString(TEXT("SystemSettings"), TEXT("RHI.BlockIHVD3D12"), BlockedIHVString, GEngineIni))
+		{
+			TArray<FString> BlockedIHVs;
+			BlockedIHVString.ParseIntoArray(BlockedIHVs, TEXT(","));
+
+			const TCHAR* VendorId = RHIVendorIdToString(EGpuVendorId(InAdapter->GetD3DAdapterDesc().VendorId));
+			for (const FString& BlockedVendor : BlockedIHVs)
+			{
+				if (BlockedVendor.Equals(VendorId, ESearchCase::IgnoreCase))
+				{
+					return true;
+				}
+			}
+		}
+	}
+#endif
+
+	return false;
+}
+
+static bool IsAdapterSupported(FD3D12Adapter* InAdapter, ERHIFeatureLevel::Type InRequestedFeatureLevel)
+{
+	if (InAdapter)
+	{
+		if (const FD3D12AdapterDesc& Desc = InAdapter->GetDesc(); Desc.IsValid())
+		{
+			ERHIFeatureLevel::Type MaxSupportedFeatureLevel = ERHIFeatureLevel::Num;
+
+			// TODO: D3D_FEATURE_LEVEL_12_2?
+			if (Desc.MaxSupportedFeatureLevel >= D3D_FEATURE_LEVEL_12_1 && Desc.MaxSupportedShaderModel >= D3D_SHADER_MODEL_6_6)
+			{
+				MaxSupportedFeatureLevel = ERHIFeatureLevel::SM6;
+			}
+			else if (Desc.MaxSupportedFeatureLevel >= D3D_FEATURE_LEVEL_11_0)
+			{
+				MaxSupportedFeatureLevel = ERHIFeatureLevel::SM5;
+			}
+
+			return MaxSupportedFeatureLevel != ERHIFeatureLevel::Num && MaxSupportedFeatureLevel >= InRequestedFeatureLevel;
+		}
 	}
-#endif // !UE_BUILD_SHIPPING
 
 	return false;
 }
 
-bool FD3D12DynamicRHIModule::IsSupported()
+bool FD3D12DynamicRHIModule::IsSupported(ERHIFeatureLevel::Type RequestedFeatureLevel)
 {
 #if !PLATFORM_HOLOLENS
 	if (!FPlatformMisc::VerifyWindowsVersion(10, 0))
@@ -393,52 +432,16 @@ bool FD3D12DynamicRHIModule::IsSupported()
 	}
 #endif
 
-	if (ShouldEnableExperimentalShaderModels())
-	{
-		// Experimental features must be enabled before doing anything else with D3D.
-
-		UUID ExperimentalFeatures[] =
-		{
-			D3D12ExperimentalShaderModels
-		};
-		HRESULT hr = D3D12EnableExperimentalFeatures(UE_ARRAY_COUNT(ExperimentalFeatures), ExperimentalFeatures, nullptr, nullptr);
-		if (SUCCEEDED(hr))
-		{
-			UE_LOG(LogD3D12RHI, Log, TEXT("D3D12 experimental shader models enabled"));
-		}
-	}
-
 	// If not computed yet
 	if (ChosenAdapters.Num() == 0)
 	{
 		FindAdapter();
 	}
 
-	bool bBlockD3D12 = false;
-#if UE_BUILD_SHIPPING
-	FString BlockedIHVString;
-	if (ChosenAdapters.Num() > 0 && ChosenAdapters[0].IsValid() && GConfig->GetString(TEXT("SystemSettings"), TEXT("RHI.BlockIHVD3D12"), BlockedIHVString, GEngineIni))
-	{
-		TArray<FString> BlockedIHVs;
-		BlockedIHVString.ParseIntoArray(BlockedIHVs, TEXT(","));
-
-		const TCHAR* VendorId = RHIVendorIdToString(EGpuVendorId(ChosenAdapters[0]->GetD3DAdapterDesc().VendorId));
-		for (const FString& BlockedVendor : BlockedIHVs)
-		{
-			if (BlockedVendor.Equals(VendorId, ESearchCase::IgnoreCase))
-			{
-				bBlockD3D12 = true;
-				break;
-			}
-		}
-	}
-#endif
-
 	// The hardware must support at least 11.0.
 	return ChosenAdapters.Num() > 0
-		&& ChosenAdapters[0]->GetDesc().IsValid()
-		&& ChosenAdapters[0]->GetDesc().MaxSupportedFeatureLevel >= D3D_FEATURE_LEVEL_11_0
-		&& !bBlockD3D12;
+		&& !IsAdapterBlocked(ChosenAdapters[0].Get())
+		&& IsAdapterSupported(ChosenAdapters[0].Get(), RequestedFeatureLevel);
 }
 
 namespace D3D12RHI
@@ -456,6 +459,9 @@ namespace D3D12RHI
 		case D3D_FEATURE_LEVEL_11_1:	return TEXT("11_1");
 		case D3D_FEATURE_LEVEL_12_0:	return TEXT("12_0");
 		case D3D_FEATURE_LEVEL_12_1:	return TEXT("12_1");
+#if !PLATFORM_HOLOLENS
+		case D3D_FEATURE_LEVEL_12_2:	return TEXT("12_2");
+#endif
 		}
 		return TEXT("X_X");
 	}
@@ -483,6 +489,23 @@ void FD3D12DynamicRHIModule::FindAdapter()
 	// Once we've chosen one we don't need to do it again.
 	check(ChosenAdapters.Num() == 0);
 
+#if !UE_BUILD_SHIPPING
+	if (CVarExperimentalShaderModels.GetValueOnAnyThread() == 1)
+	{
+		// Experimental features must be enabled before doing anything else with D3D.
+
+		UUID ExperimentalFeatures[] =
+		{
+			D3D12ExperimentalShaderModels
+		};
+		HRESULT hr = D3D12EnableExperimentalFeatures(UE_ARRAY_COUNT(ExperimentalFeatures), ExperimentalFeatures, nullptr, nullptr);
+		if (SUCCEEDED(hr))
+		{
+			UE_LOG(LogD3D12RHI, Log, TEXT("D3D12 experimental shader models enabled"));
+		}
+	}
+#endif
+
 	// Try to create the DXGIFactory.  This will fail if we're not running Vista.
 	TRefCountPtr<IDXGIFactory4> DXGIFactory4;
 	SafeCreateDXGIFactory(DXGIFactory4.GetInitReference());
@@ -679,8 +702,7 @@ static bool DoesAnyAdapterSupportSM6(const TArray<TSharedPtr<FD3D12Adapter>>& Ad
 {
 	for (const TSharedPtr<FD3D12Adapter>& Adapter : Adapters)
 	{
-		// TODO: D3D_FEATURE_LEVEL_12_2
-		if (Adapter->GetFeatureLevel() >= D3D_FEATURE_LEVEL_12_1 && Adapter->GetHighestShaderModel() >= D3D_SHADER_MODEL_6_6)
+		if (IsAdapterSupported(Adapter.Get(), ERHIFeatureLevel::SM6))
 		{
 			return true;
 		}
diff --git a/Engine/Source/Runtime/RHI/Private/Windows/WindowsDynamicRHI.cpp b/Engine/Source/Runtime/RHI/Private/Windows/WindowsDynamicRHI.cpp
--- a/Engine/Source/Runtime/RHI/Private/Windows/WindowsDynamicRHI.cpp
+++ b/Engine/Source/Runtime/RHI/Private/Windows/WindowsDynamicRHI.cpp
@@ -16,6 +16,22 @@
 
 static const TCHAR* GLoadedRHIModuleName;
 
+static TArray<EShaderPlatform> GetTargetedShaderPlatforms()
+{
+	TArray<FString> TargetedShaderFormats;
+	GConfig->GetArray(TEXT("/Script/WindowsTargetPlatform.WindowsTargetSettings"), TEXT("TargetedRHIs"), TargetedShaderFormats, GEngineIni);
+
+	TArray<EShaderPlatform> ShaderPlatforms;
+	ShaderPlatforms.Reserve(TargetedShaderFormats.Num());
+
+	for (const FString& ShaderFormat : TargetedShaderFormats)
+	{
+		ShaderPlatforms.Add(ShaderFormatToLegacyShaderPlatform(FName(*ShaderFormat)));
+	}
+
+	return ShaderPlatforms;
+}
+
 // Default to Performance Mode on low-end machines
 static bool DefaultFeatureLevelES31()
 {
@@ -25,13 +41,9 @@ static bool DefaultFeatureLevelES31()
 		return ForceES31.GetValue();
 	}
 
-	FConfigFile EngineSettings;
-	FString PlatformNameString = FPlatformProperties::IniPlatformName();
-	FConfigCacheIni::LoadLocalIniFile(EngineSettings, TEXT("Engine"), true, *PlatformNameString);
-
 	// Force Performance mode for machines with too few cores including hyperthreads
 	int MinCoreCount = 0;
-	if (EngineSettings.GetInt(TEXT("PerformanceMode"), TEXT("MinCoreCount"), MinCoreCount) && FPlatformMisc::NumberOfCoresIncludingHyperthreads() < MinCoreCount)
+	if (GConfig->GetInt(TEXT("PerformanceMode"), TEXT("MinCoreCount"), MinCoreCount, GEngineIni) && FPlatformMisc::NumberOfCoresIncludingHyperthreads() < MinCoreCount)
 	{
 		ForceES31 = true;
 		return true;
@@ -39,7 +51,7 @@ static bool DefaultFeatureLevelES31()
 
 	FString MinMemorySizeBucketString;
 	FString MinIntegratedMemorySizeBucketString;
-	if (EngineSettings.GetString(TEXT("PerformanceMode"), TEXT("MinMemorySizeBucket"), MinMemorySizeBucketString) && EngineSettings.GetString(TEXT("PerformanceMode"), TEXT("MinIntegratedMemorySizeBucket"), MinIntegratedMemorySizeBucketString))
+	if (GConfig->GetString(TEXT("PerformanceMode"), TEXT("MinMemorySizeBucket"), MinMemorySizeBucketString, GEngineIni) && GConfig->GetString(TEXT("PerformanceMode"), TEXT("MinIntegratedMemorySizeBucket"), MinIntegratedMemorySizeBucketString, GEngineIni))
 	{
 		for (int EnumIndex = int(EPlatformMemorySizeBucket::Largest); EnumIndex <= int(EPlatformMemorySizeBucket::Tiniest); EnumIndex++)
 		{
@@ -104,6 +116,23 @@ static bool PreferFeatureLevelES31()
 	return false;
 }
 
+static bool IsES31D3DOnly()
+{
+	bool bES31DXOnly = false;
+#if !WITH_EDITOR
+	if (!GIsEditor)
+	{
+		GConfig->GetBool(TEXT("PerformanceMode"), TEXT("bES31DXOnly"), bES31DXOnly, GEngineIni);
+	}
+#endif
+	return bES31DXOnly;
+}
+
+static bool AllowD3D11FeatureLevelES31()
+{
+	return true;
+}
+
 static bool AllowD3D12FeatureLevelES31()
 {
 	if (!GIsEditor)
@@ -115,19 +144,21 @@ static bool AllowD3D12FeatureLevelES31()
 	return true;
 }
 
-namespace
+static bool AllowVulkanFeatureLevelES31()
 {
-	enum class WindowsRHI
-	{
-		D3D11,
-		D3D12,
-		Vulkan,
-		OpenGL,
-	};
+	return !IsES31D3DOnly();
 }
 
+enum class WindowsRHI
+{
+	D3D11,
+	D3D12,
+	Vulkan,
+	OpenGL,
+};
+
 // Choose the default from DefaultGraphicsRHI or TargetedRHIs. DefaultGraphicsRHI has precedence.
-static WindowsRHI ChooseDefaultRHI()
+static WindowsRHI ChooseDefaultRHI(const TArray<EShaderPlatform>& TargetedShaderPlatforms)
 {
 	// Make sure the DDSPI is initialized before we try and read from it
 	FGenericDataDrivenShaderPlatformInfo::Initialize();
@@ -135,16 +166,13 @@ static WindowsRHI ChooseDefaultRHI()
 	WindowsRHI DefaultRHI = WindowsRHI::D3D11;
 
 	// Default graphics RHI is the main project setting that governs the choice, so it takes the priority
-	FConfigFile EngineSettings;
-	FString PlatformNameString = FPlatformProperties::PlatformName();
-	const TCHAR* PlatformName = *PlatformNameString;
-	FConfigCacheIni::LoadLocalIniFile(EngineSettings, TEXT("Engine"), true, PlatformName);
 	FString DefaultGraphicsRHI;
-	if (EngineSettings.GetString(TEXT("/Script/WindowsTargetPlatform.WindowsTargetSettings"), TEXT("DefaultGraphicsRHI"), DefaultGraphicsRHI))
+	if (GConfig->GetString(TEXT("/Script/WindowsTargetPlatform.WindowsTargetSettings"), TEXT("DefaultGraphicsRHI"), DefaultGraphicsRHI, GEngineIni))
 	{
-		FString NAME_DX11(TEXT("DefaultGraphicsRHI_DX11"));
-		FString NAME_DX12(TEXT("DefaultGraphicsRHI_DX12"));
-		FString NAME_VULKAN(TEXT("DefaultGraphicsRHI_Vulkan"));
+		const FString NAME_DX11(TEXT("DefaultGraphicsRHI_DX11"));
+		const FString NAME_DX12(TEXT("DefaultGraphicsRHI_DX12"));
+		const FString NAME_VULKAN(TEXT("DefaultGraphicsRHI_Vulkan"));
+
 		if (DefaultGraphicsRHI == NAME_DX11)
 		{
 			DefaultRHI = WindowsRHI::D3D11;
@@ -162,26 +190,24 @@ static WindowsRHI ChooseDefaultRHI()
 			UE_LOG(LogRHI, Error, TEXT("Unrecognized setting '%s' for DefaultGraphicsRHI"), *DefaultGraphicsRHI);
 		}
 	}
-	else 
+	else if (TargetedShaderPlatforms.Num() > 0)
 	{
 		// If we don't have DefaultGraphicsRHI set, try to deduce it from the list of targeted shader platforms
-		TArray<FString> TargetedShaderFormats;
-		GConfig->GetArray(TEXT("/Script/WindowsTargetPlatform.WindowsTargetSettings"), TEXT("TargetedRHIs"), TargetedShaderFormats, GEngineIni);
-		if (TargetedShaderFormats.Num() > 0)
-		{
-			// Pick the first one
-			FName ShaderFormatName(*TargetedShaderFormats[0]);
-			EShaderPlatform TargetedPlatform = ShaderFormatToLegacyShaderPlatform(ShaderFormatName);
+		// Pick the first one
+		const EShaderPlatform TargetedPlatform = TargetedShaderPlatforms[0];
 
-			// not checking D3D as DefaultRHI begins initialized as D3D12
-			if (IsVulkanPlatform(TargetedPlatform))
-			{
-				DefaultRHI = WindowsRHI::Vulkan;
-			}
-			else if (IsOpenGLPlatform(TargetedPlatform))
-			{
-				DefaultRHI = WindowsRHI::OpenGL;
-			}
+		// not checking D3D as DefaultRHI begins initialized as D3D12
+		if (IsVulkanPlatform(TargetedPlatform))
+		{
+			DefaultRHI = WindowsRHI::Vulkan;
+		}
+		else if (IsOpenGLPlatform(TargetedPlatform))
+		{
+			DefaultRHI = WindowsRHI::OpenGL;
+		}
+		else
+		{
+			DefaultRHI = WindowsRHI::D3D11;
 		}
 	}
 
@@ -233,14 +259,9 @@ static TOptional<WindowsRHI> ChooseForcedRHI()
 
 #if	!WITH_EDITOR && UE_BUILD_SHIPPING
 	// In Shipping builds we can limit ES31 on Windows to only DX11. All RHIs are allowed by default.
-	bool bES31DXOnly = false;
-	if (!GIsEditor)
-	{
-		GConfig->GetBool(TEXT("PerformanceMode"), TEXT("bES31DXOnly"), bES31DXOnly, GEngineIni);
-	}
 
 	// FeatureLevelES31 is also a command line override, so it will determine the underlying RHI unless one is specified
-	if (bES31DXOnly && (FParse::Param(FCommandLine::Get(), TEXT("FeatureLevelES31")) || FParse::Param(FCommandLine::Get(), TEXT("FeatureLevelES3_1"))))
+	if (IsES31D3DOnly() && (FParse::Param(FCommandLine::Get(), TEXT("FeatureLevelES31")) || FParse::Param(FCommandLine::Get(), TEXT("FeatureLevelES3_1"))))
 	{
 		if (ForcedRHI == WindowsRHI::OpenGL)
 		{
@@ -270,48 +291,78 @@ static TOptional<WindowsRHI> ChooseForcedRHI()
 	return ForcedRHI;
 }
 
-static ERHIFeatureLevel::Type ChooseFeatureLevel(TOptional<WindowsRHI> ChosenRHI, TOptional<WindowsRHI> ForcedRHI)
+static TOptional<ERHIFeatureLevel::Type> ChooseForcedFeatureLevel(TOptional<WindowsRHI> ChosenRHI, TOptional<WindowsRHI> ForcedRHI)
 {
-	ERHIFeatureLevel::Type FeatureLevel = ERHIFeatureLevel::SM5;
-	if (FParse::Param(FCommandLine::Get(), TEXT("sm6")))
+	TOptional<ERHIFeatureLevel::Type> ForcedFeatureLevel{};
+
+	if (FParse::Param(FCommandLine::Get(), TEXT("es31")) || FParse::Param(FCommandLine::Get(), TEXT("FeatureLevelES31")) || FParse::Param(FCommandLine::Get(), TEXT("FeatureLevelES3_1")))
 	{
-		FeatureLevel = ERHIFeatureLevel::SM6;
+		ForcedFeatureLevel = ERHIFeatureLevel::ES3_1;
 	}
 
-	bool bAllowD3D12FeatureLevelES31 = AllowD3D12FeatureLevelES31();
-
-	if (!ForcedRHI && (ChosenRHI == WindowsRHI::D3D11 || (ChosenRHI == WindowsRHI::D3D12 && bAllowD3D12FeatureLevelES31)) && PreferFeatureLevelES31())
+	if (FParse::Param(FCommandLine::Get(), TEXT("sm5")))
 	{
-		FeatureLevel = ERHIFeatureLevel::ES3_1;
+		ForcedFeatureLevel = ERHIFeatureLevel::SM5;
 	}
 
-	bool bForceES31 = FParse::Param(FCommandLine::Get(), TEXT("FeatureLevelES31")) || FParse::Param(FCommandLine::Get(), TEXT("FeatureLevelES3_1"));
-	if (bForceES31 && (ChosenRHI == WindowsRHI::D3D11 || (ChosenRHI == WindowsRHI::D3D12 && bAllowD3D12FeatureLevelES31)))
+	if (FParse::Param(FCommandLine::Get(), TEXT("sm6")))
 	{
-		FeatureLevel = ERHIFeatureLevel::ES3_1;
+		ForcedFeatureLevel = ERHIFeatureLevel::SM6;
 	}
 
-	if (ChosenRHI == WindowsRHI::OpenGL)
+	return ForcedFeatureLevel;
+}
+
+static bool IsD3DSM6PlatformTargeted(const TArray<EShaderPlatform>& TargetedShaderPlatforms)
+{
+	for (const EShaderPlatform ShaderPlatform : TargetedShaderPlatforms)
 	{
-		// OpenGL can only be used for mobile preview
-		FeatureLevel = ERHIFeatureLevel::ES3_1;
+		if (IsD3DPlatform(ShaderPlatform) && IsFeatureLevelSupported(ShaderPlatform, ERHIFeatureLevel::SM6))
+		{
+			return true;
+		}
 	}
 
-	bool bES31DXOnly = false;
-#if !WITH_EDITOR
-	if (!GIsEditor)
+	return false;
+}
+
+static ERHIFeatureLevel::Type ChooseFeatureLevel(TOptional<WindowsRHI> ChosenRHI, TOptional<ERHIFeatureLevel::Type> ForcedFeatureLevel, const TArray<EShaderPlatform>& TargetedShaderPlatforms)
+{
+	TOptional<ERHIFeatureLevel::Type> FeatureLevel{};
+
+	if (ForcedFeatureLevel)
 	{
-		GConfig->GetBool(TEXT("PerformanceMode"), TEXT("bES31DXOnly"), bES31DXOnly, GEngineIni);
+		if (ForcedFeatureLevel == ERHIFeatureLevel::ES3_1 &&
+			(ChosenRHI == WindowsRHI::D3D11 && AllowD3D11FeatureLevelES31()) ||
+			(ChosenRHI == WindowsRHI::D3D12 && AllowD3D12FeatureLevelES31()) ||
+			(ChosenRHI == WindowsRHI::Vulkan && AllowVulkanFeatureLevelES31()))
+		{
+			FeatureLevel = ForcedFeatureLevel;
+		}
+		else if (ForcedFeatureLevel == ERHIFeatureLevel::SM6 && ChosenRHI == WindowsRHI::D3D12)
+		{
+			FeatureLevel = ForcedFeatureLevel;
+		}
 	}
-#endif
 
-	// Allow Vulkan to force ES31 unless we blocked it in config
-	if (!bES31DXOnly && ChosenRHI == WindowsRHI::Vulkan && bForceES31)
+	if (!FeatureLevel)
 	{
-		FeatureLevel = ERHIFeatureLevel::ES3_1;
+		if (ChosenRHI == WindowsRHI::OpenGL)
+		{
+			// OpenGL can only be used for mobile preview
+			FeatureLevel = ERHIFeatureLevel::ES3_1;
+		}
+		else if (ChosenRHI == WindowsRHI::D3D12 && IsD3DSM6PlatformTargeted(TargetedShaderPlatforms))
+		{
+			FeatureLevel = ERHIFeatureLevel::SM6;
+		}
+		else
+		{
+			FeatureLevel = ERHIFeatureLevel::SM5;
+		}
 	}
 
-	return FeatureLevel;
+	return FeatureLevel.GetValue();
 }
 
 static IDynamicRHIModule* LoadDynamicRHIModule(ERHIFeatureLevel::Type& DesiredFeatureLevel, const TCHAR*& LoadedRHIModuleName)
@@ -323,11 +374,13 @@ static IDynamicRHIModule* LoadDynamicRHIModule(ERHIFeatureLevel::Type& DesiredFe
 		*GPUCrashDebuggingCVar = bUseGPUCrashDebugging;
 	}
 
+	const TArray<EShaderPlatform> TargetedShaderPlatforms = GetTargetedShaderPlatforms();
+
 	// RHI is chosen by the project settings (first DefaultGraphicsRHI, then TargetedRHIs are consulted, "Default" maps to D3D12). 
 	// After this, a separate game-only setting (does not affect editor) bPreferD3D12InGame selects between D3D12 or D3D11 (but will not have any effect if Vulkan or OpenGL are chosen).
 	// Commandline switches apply after this and can force an arbitrary RHIs. If RHI isn't supported, the game will refuse to start.
 
-	WindowsRHI DefaultRHI = ChooseDefaultRHI();
+	WindowsRHI DefaultRHI = ChooseDefaultRHI(TargetedShaderPlatforms);
 	TOptional<WindowsRHI> ForcedRHI = ChooseForcedRHI();
 
 	WindowsRHI ChosenRHI = DefaultRHI;
@@ -336,10 +389,11 @@ static IDynamicRHIModule* LoadDynamicRHIModule(ERHIFeatureLevel::Type& DesiredFe
 		ChosenRHI = ForcedRHI.GetValue();
 	}
 
-	DesiredFeatureLevel = ChooseFeatureLevel(ChosenRHI, ForcedRHI);
+	TOptional<ERHIFeatureLevel::Type> ForcedFeatureLevel = ChooseForcedFeatureLevel(ChosenRHI, ForcedRHI);
+	DesiredFeatureLevel = ChooseFeatureLevel(ChosenRHI, ForcedFeatureLevel, TargetedShaderPlatforms);
 
 	// Load the dynamic RHI module.
-	IDynamicRHIModule* DynamicRHIModule = NULL;
+	IDynamicRHIModule* DynamicRHIModule = nullptr;
 
 	if (ChosenRHI == WindowsRHI::OpenGL)
 	{
@@ -347,11 +401,11 @@ static IDynamicRHIModule* LoadDynamicRHIModule(ERHIFeatureLevel::Type& DesiredFe
 		const TCHAR* OpenGLRHIModuleName = TEXT("OpenGLDrv");
 		DynamicRHIModule = &FModuleManager::LoadModuleChecked<IDynamicRHIModule>(OpenGLRHIModuleName);
 
-		if (!DynamicRHIModule->IsSupported())
+		if (!DynamicRHIModule->IsSupported(DesiredFeatureLevel))
 		{
 			FMessageDialog::Open(EAppMsgType::Ok, NSLOCTEXT("WindowsDynamicRHI", "RequiredOpenGL", "OpenGL 4.3 is required to run the engine."));
 			FPlatformMisc::RequestExit(1);
-			DynamicRHIModule = NULL;
+			DynamicRHIModule = nullptr;
 		}
 
 		LoadedRHIModuleName = OpenGLRHIModuleName;
@@ -361,21 +415,44 @@ static IDynamicRHIModule* LoadDynamicRHIModule(ERHIFeatureLevel::Type& DesiredFe
 		FApp::SetGraphicsRHI(TEXT("Vulkan"));
 		const TCHAR* VulkanRHIModuleName = TEXT("VulkanRHI");
 		DynamicRHIModule = &FModuleManager::LoadModuleChecked<IDynamicRHIModule>(VulkanRHIModuleName);
-		if (!DynamicRHIModule->IsSupported())
+
+		if (!DynamicRHIModule->IsSupported(DesiredFeatureLevel))
 		{
 			FMessageDialog::Open(EAppMsgType::Ok, NSLOCTEXT("WindowsDynamicRHI", "RequiredVulkan", "Vulkan Driver is required to run the engine."));
 			FPlatformMisc::RequestExit(1);
-			DynamicRHIModule = NULL;
+			DynamicRHIModule = nullptr;
 		}
 		LoadedRHIModuleName = VulkanRHIModuleName;
 	}
 	else if (ChosenRHI == WindowsRHI::D3D12)
 	{
-		FApp::SetGraphicsRHI(TEXT("DirectX 12"));
+		FString FeatureLevelName;
+		GetFeatureLevelName(DesiredFeatureLevel, FeatureLevelName);
+
+		const FString RHIDisplayName = FString::Printf(TEXT("DirectX 12 (%s)"), *FeatureLevelName);
+		FApp::SetGraphicsRHI(RHIDisplayName);
+
 		LoadedRHIModuleName = TEXT("D3D12RHI");
 		DynamicRHIModule = FModuleManager::LoadModulePtr<IDynamicRHIModule>(LoadedRHIModuleName);
 
-		if (!DynamicRHIModule || !DynamicRHIModule->IsSupported())
+		bool bD3D12Supported = DynamicRHIModule && DynamicRHIModule->IsSupported(DesiredFeatureLevel);
+
+		// Fallback to SM5 if SM6 is not supported
+		if (!bD3D12Supported && DynamicRHIModule && DesiredFeatureLevel == ERHIFeatureLevel::SM6)
+		{
+			if (ForcedFeatureLevel)
+			{
+				FMessageDialog::Open(EAppMsgType::Ok, NSLOCTEXT("WindowsDynamicRHI", "RequiredDX12SM6", "DX12 SM6 is not supported on your system. Try running without the -sm6 command line argument."));
+				FPlatformMisc::RequestExit(1);
+			}
+
+			UE_LOG(LogRHI, Log, TEXT("D3D12 SM6 is not supported, trying SM5"));
+
+			DesiredFeatureLevel = ERHIFeatureLevel::SM5;
+			bD3D12Supported = DynamicRHIModule->IsSupported(DesiredFeatureLevel);
+		}
+
+		if (!bD3D12Supported)
 		{
 			if (ForcedRHI == WindowsRHI::D3D12)
 			{
@@ -386,7 +463,7 @@ static IDynamicRHIModule* LoadDynamicRHIModule(ERHIFeatureLevel::Type& DesiredFe
 			{
 				FModuleManager::Get().UnloadModule(LoadedRHIModuleName);
 			}
-			DynamicRHIModule = NULL;
+			DynamicRHIModule = nullptr;
 			LoadedRHIModuleName = nullptr;
 		}
 	}
@@ -398,11 +475,11 @@ static IDynamicRHIModule* LoadDynamicRHIModule(ERHIFeatureLevel::Type& DesiredFe
 		const TCHAR* D3D11RHIModuleName = TEXT("D3D11RHI");
 		DynamicRHIModule = &FModuleManager::LoadModuleChecked<IDynamicRHIModule>(D3D11RHIModuleName);
 
-		if (!DynamicRHIModule->IsSupported())
+		if (!DynamicRHIModule->IsSupported(DesiredFeatureLevel))
 		{
 			FMessageDialog::Open(EAppMsgType::Ok, NSLOCTEXT("WindowsDynamicRHI", "RequiredDX11Feature_11_SM5", "A D3D11-compatible GPU (Feature Level 11.0, Shader Model 5.0) is required to run the engine."));
 			FPlatformMisc::RequestExit(1);
-			DynamicRHIModule = NULL;
+			DynamicRHIModule = nullptr;
 		}
 
 		LoadedRHIModuleName = D3D11RHIModuleName;
diff --git a/Engine/Source/Runtime/RHI/Public/DynamicRHI.h b/Engine/Source/Runtime/RHI/Public/DynamicRHI.h
--- a/Engine/Source/Runtime/RHI/Public/DynamicRHI.h
+++ b/Engine/Source/Runtime/RHI/Public/DynamicRHI.h
@@ -1682,6 +1682,8 @@ public:
 	/** Checks whether the RHI is supported by the current system. */
 	virtual bool IsSupported() = 0;
 
+	virtual bool IsSupported(ERHIFeatureLevel::Type RequestedFeatureLevel) { return IsSupported(); }
+
 	/** Creates a new instance of the dynamic RHI implemented by the module. */
 	virtual FDynamicRHI* CreateRHI(ERHIFeatureLevel::Type RequestedFeatureLevel = ERHIFeatureLevel::Num) = 0;
 };
