Copying //UE5/Dev-ParallelRendering to //UE5/Main

diff --git a/Engine/Source/Runtime/RHI/Public/RHICommandList.h b/Engine/Source/Runtime/RHI/Public/RHICommandList.h
--- a/Engine/Source/Runtime/RHI/Public/RHICommandList.h
+++ b/Engine/Source/Runtime/RHI/Public/RHICommandList.h
@@ -1515,11 +1515,8 @@ FRHICOMMAND_MACRO(FRHICommandCopyToResolveTarget)
 		, SourceTexture(InSourceTexture)
 		, DestTexture(InDestTexture)
 	{
-		if (ensure(SourceTexture) && ensure(DestTexture))
-		{
-			ensure(SourceTexture->GetTexture2D() || SourceTexture->GetTexture3D() || SourceTexture->GetTextureCube() || SourceTexture->GetTexture2DArray());
-			ensure(DestTexture->GetTexture2D() || DestTexture->GetTexture3D() || DestTexture->GetTextureCube() || DestTexture->GetTexture2DArray());
-		}
+		ensure(SourceTexture);
+		ensure(DestTexture);
 	}
 	RHI_API void Execute(FRHICommandListBase& CmdList);
 };
@@ -1535,11 +1532,8 @@ FRHICOMMAND_MACRO(FRHICommandCopyTexture)
 		, SourceTexture(InSourceTexture)
 		, DestTexture(InDestTexture)
 	{
-		if (ensure(SourceTexture) && ensure(DestTexture))
-		{
-			ensure(SourceTexture->GetTexture2D() || SourceTexture->GetTexture2DArray() || SourceTexture->GetTexture3D() || SourceTexture->GetTextureCube());
-			ensure(DestTexture->GetTexture2D() || DestTexture->GetTexture2DArray() || DestTexture->GetTexture3D() || DestTexture->GetTextureCube());
-		}
+		ensure(SourceTexture);
+		ensure(DestTexture);
 	}
 	RHI_API void Execute(FRHICommandListBase& CmdList);
 };
@@ -4013,13 +4007,6 @@ public:
 FBufferRHIRef RHICreateIndexBuffer(uint32 Stride, uint32 Size, uint32 InUsage, FRHIResourceCreateInfo& CreateInfo);
 FBufferRHIRef RHICreateVertexBuffer(uint32 Size, uint32 InUsage, FRHIResourceCreateInfo& CreateInfo);
 FBufferRHIRef RHICreateStructuredBuffer(uint32 Stride, uint32 Size, uint32 InUsage, FRHIResourceCreateInfo& CreateInfo);
-FTexture2DRHIRef RHICreateTexture2D(uint32 SizeX, uint32 SizeY, uint8 Format, uint32 NumMips, uint32 NumSamples, ETextureCreateFlags Flags, FRHIResourceCreateInfo& CreateInfo);
-FTexture2DRHIRef RHICreateTextureExternal2D(uint32 SizeX, uint32 SizeY, uint8 Format, uint32 NumMips, uint32 NumSamples, ETextureCreateFlags Flags, FRHIResourceCreateInfo& CreateInfo);
-FTexture2DRHIRef RHIAsyncCreateTexture2D(uint32 SizeX, uint32 SizeY, uint8 Format, uint32 NumMips, ETextureCreateFlags Flags, void** InitialMipData, uint32 NumInitialMips);
-FTexture2DArrayRHIRef RHICreateTexture2DArray(uint32 SizeX, uint32 SizeY, uint32 SizeZ, uint8 Format, uint32 NumMips, uint32 NumSamples, ETextureCreateFlags Flags, FRHIResourceCreateInfo& CreateInfo);
-FTexture3DRHIRef RHICreateTexture3D(uint32 SizeX, uint32 SizeY, uint32 SizeZ, uint8 Format, uint32 NumMips, ETextureCreateFlags Flags, FRHIResourceCreateInfo& CreateInfo);
-FTextureCubeRHIRef RHICreateTextureCube(uint32 Size, uint8 Format, uint32 NumMips, ETextureCreateFlags Flags, FRHIResourceCreateInfo& CreateInfo);
-FTextureCubeRHIRef RHICreateTextureCubeArray(uint32 Size, uint32 ArraySize, uint8 Format, uint32 NumMips, ETextureCreateFlags Flags, FRHIResourceCreateInfo& CreateInfo);
 
 extern RHI_API ERHIAccess RHIGetDefaultResourceState(ETextureCreateFlags InUsage, bool bInHasInitialData);
 extern RHI_API ERHIAccess RHIGetDefaultResourceState(EBufferUsageFlags InUsage, bool bInHasInitialData);
@@ -4334,19 +4321,28 @@ public:
 		return GDynamicRHI->CreateShaderResourceView_RenderThread(*this, Initializer);
 	}
 	
+	UE_DEPRECATED(5.1, "The CalcTexture... functions on the immediate RHI command list are deprecated. Use the global scope RHICalcTexturePlatformSize instead.")
 	FORCEINLINE uint64 CalcTexture2DPlatformSize(uint32 SizeX, uint32 SizeY, uint8 Format, uint32 NumMips, uint32 NumSamples, ETextureCreateFlags Flags, const FRHIResourceCreateInfo& CreateInfo, uint32& OutAlign)
 	{
+		PRAGMA_DISABLE_DEPRECATION_WARNINGS
 		return RHICalcTexture2DPlatformSize(SizeX, SizeY, Format, NumMips, NumSamples, Flags, CreateInfo, OutAlign);
+		PRAGMA_ENABLE_DEPRECATION_WARNINGS
 	}
 	
+	UE_DEPRECATED(5.1, "The CalcTexture... functions on the immediate RHI command list are deprecated. Use the global scope RHICalcTexturePlatformSize instead.")
 	FORCEINLINE uint64 CalcTexture3DPlatformSize(uint32 SizeX, uint32 SizeY, uint32 SizeZ, uint8 Format, uint32 NumMips, ETextureCreateFlags Flags, const FRHIResourceCreateInfo& CreateInfo, uint32& OutAlign)
 	{
+		PRAGMA_DISABLE_DEPRECATION_WARNINGS
 		return RHICalcTexture3DPlatformSize(SizeX, SizeY, SizeZ, Format, NumMips, Flags, CreateInfo, OutAlign);
+		PRAGMA_ENABLE_DEPRECATION_WARNINGS
 	}
 	
+	UE_DEPRECATED(5.1, "The CalcTexture... functions on the immediate RHI command list are deprecated. Use the global scope RHICalcTexturePlatformSize instead.")
 	FORCEINLINE uint64 CalcTextureCubePlatformSize(uint32 Size, uint8 Format, uint32 NumMips, ETextureCreateFlags Flags, const FRHIResourceCreateInfo& CreateInfo, uint32& OutAlign)
 	{
+		PRAGMA_DISABLE_DEPRECATION_WARNINGS
 		return RHICalcTextureCubePlatformSize(Size, Format, NumMips, Flags, CreateInfo, OutAlign);
+		PRAGMA_ENABLE_DEPRECATION_WARNINGS
 	}
 	
 	FORCEINLINE void GetTextureMemoryStats(FTextureMemoryStats& OutStats)
@@ -4361,17 +4357,18 @@ public:
 		return GDynamicRHI->RHIGetTextureMemoryVisualizeData(TextureData,SizeX,SizeY,Pitch,PixelSize);
 	}
 	
-	FORCEINLINE void CopySharedMips(FRHITexture2D* DestTexture2D, FRHITexture2D* SrcTexture2D)
+	FORCEINLINE void CopySharedMips(FRHITexture* DestTexture, FRHITexture* SrcTexture)
 	{
 		QUICK_SCOPE_CYCLE_COUNTER(STAT_RHIMETHOD_CopySharedMips_Flush);
-		DestTexture2D->AddRef();
-		SrcTexture2D->AddRef();
-		EnqueueLambda([DestTexture2D, SrcTexture2D](FRHICommandList&)
+		DestTexture->AddRef();
+		SrcTexture->AddRef();
+		EnqueueLambda([DestTexture, SrcTexture](FRHICommandList&)
 		{
 			LLM_SCOPE(ELLMTag::Textures);
-			GDynamicRHI->RHICopySharedMips(DestTexture2D, SrcTexture2D);
-			DestTexture2D->Release();
-			SrcTexture2D->Release();
+			//@todo move this function onto the IRHIComputeCommandContext
+			GDynamicRHI->RHICopySharedMips(DestTexture, SrcTexture);
+			DestTexture->Release();
+			SrcTexture->Release();
 		});
 	}
 
@@ -5361,38 +5358,199 @@ FORCEINLINE void RHIUpdateTextureReference(FRHITextureReference* TextureRef, FRH
 	FRHICommandListExecutor::GetImmediateCommandList().UpdateTextureReference(TextureRef, NewTexture);
 }
 
-FORCEINLINE FTexture2DRHIRef RHICreateTexture2D(uint32 SizeX, uint32 SizeY, uint8 Format, uint32 NumMips, uint32 NumSamples, ETextureCreateFlags Flags, ERHIAccess InResourceState, FRHIResourceCreateInfo& CreateInfo)
+FORCEINLINE FTextureRHIRef RHICreateTexture(const FRHITextureCreateDesc& CreateDesc)
 {
-	LLM_SCOPE(EnumHasAnyFlags(Flags, TexCreate_RenderTargetable | TexCreate_DepthStencilTargetable) ? ELLMTag::RenderTargets : ELLMTag::Textures);
-	return GDynamicRHI->RHICreateTexture2D_RenderThread(FRHICommandListExecutor::GetImmediateCommandList(), SizeX, SizeY, Format, NumMips, NumSamples, Flags, InResourceState, CreateInfo);
+	//check(IsInRenderingThread()); // @todo: texture type unification Some passes call this function on parallel rendering threads (e.g. FRHIGPUTextureReadback::EnqueueCopyInternal)
+	FRHICommandListImmediate& RHICmdList = FRHICommandListExecutor::GetImmediateCommandList();
+
+	LLM_SCOPE(EnumHasAnyFlags(CreateDesc.Flags, TexCreate_RenderTargetable | TexCreate_DepthStencilTargetable) ? ELLMTag::RenderTargets : ELLMTag::Textures);
+	return GDynamicRHI->RHICreateTexture_RenderThread(RHICmdList, CreateDesc);
 }
 
-FORCEINLINE FTexture2DRHIRef RHICreateTextureExternal2D(uint32 SizeX, uint32 SizeY, uint8 Format, uint32 NumMips, uint32 NumSamples, ETextureCreateFlags Flags, ERHIAccess InResourceState, FRHIResourceCreateInfo& CreateInfo)
-{
-	LLM_SCOPE(EnumHasAnyFlags(Flags, TexCreate_RenderTargetable | TexCreate_DepthStencilTargetable) ? ELLMTag::RenderTargets : ELLMTag::Textures);
-	return GDynamicRHI->RHICreateTextureExternal2D_RenderThread(FRHICommandListExecutor::GetImmediateCommandList(), SizeX, SizeY, Format, NumMips, NumSamples, Flags, InResourceState, CreateInfo);
+//UE_DEPRECATED(5.1, "FRHITexture2D is deprecated, please use RHICreateTexture(const FRHITextureCreateDesc&).")
+FORCEINLINE FTexture2DRHIRef RHICreateTexture2D(uint32 SizeX, uint32 SizeY, uint8 Format, uint32 NumMips, uint32 NumSamples, ETextureCreateFlags Flags, ERHIAccess ResourceState, const FRHIResourceCreateInfo& CreateInfo)
+{
+	const uint32 Depth = 1;
+	const uint32 ArraySize = 1;
+
+	return RHICreateTexture(
+		FRHITextureCreateDesc(
+			  ETextureDimension::Texture2D
+			, Flags
+			, (EPixelFormat)Format
+			, CreateInfo.ClearValueBinding
+			, { (int32)SizeX, (int32)SizeY }
+			, Depth
+			, ArraySize
+			, NumMips
+			, NumSamples
+			, CreateInfo.ExtData
+			, ResourceState
+			, CreateInfo.DebugName
+			, CreateInfo.BulkData
+			, CreateInfo.GPUMask
+		)
+	);
+}
+
+//UE_DEPRECATED(5.1, "FRHITexture2D is deprecated, please use RHICreateTexture(const FRHITextureCreateDesc&).")
+FORCEINLINE FTexture2DRHIRef RHICreateTextureExternal2D(uint32 SizeX, uint32 SizeY, uint8 Format, uint32 NumMips, uint32 NumSamples, ETextureCreateFlags Flags, ERHIAccess ResourceState, const  FRHIResourceCreateInfo& CreateInfo)
+{
+	const uint32 Depth = 1;
+	const uint32 ArraySize = 1;
+
+	ERHIAccess InitialState = ResourceState == ERHIAccess::Unknown
+		? RHIGetDefaultResourceState(Flags, false)
+		: ResourceState;
+
+	return RHICreateTexture(
+		FRHITextureCreateDesc(
+			  ETextureDimension::Texture2D
+			, Flags | ETextureCreateFlags::External
+			, (EPixelFormat)Format
+			, CreateInfo.ClearValueBinding
+			, { (int32)SizeX, (int32)SizeY }
+			, Depth
+			, ArraySize
+			, NumMips
+			, NumSamples
+			, CreateInfo.ExtData
+			, InitialState
+			, CreateInfo.DebugName
+			, CreateInfo.BulkData
+			, CreateInfo.GPUMask
+		)
+	);
 }
 
-FORCEINLINE FTexture2DRHIRef RHIAsyncCreateTexture2D(uint32 SizeX, uint32 SizeY, uint8 Format, uint32 NumMips, ETextureCreateFlags Flags, ERHIAccess InResourceState, void** InitialMipData, uint32 NumInitialMips)
+//UE_DEPRECATED(5.1, "FRHITexture2D is deprecated, please use RHICreateTexture(const FRHITextureCreateDesc&).")
+FORCEINLINE FTexture2DRHIRef RHIAsyncCreateTexture2D(uint32 SizeX, uint32 SizeY, uint8 Format, uint32 NumMips, ETextureCreateFlags Flags, ERHIAccess ResourceState, void** InitialMipData, uint32 NumInitialMips)
 {
 	LLM_SCOPE(EnumHasAnyFlags(Flags, TexCreate_RenderTargetable | TexCreate_DepthStencilTargetable) ? ELLMTag::RenderTargets : ELLMTag::Textures);
-	return GDynamicRHI->RHIAsyncCreateTexture2D(SizeX, SizeY, Format, NumMips, Flags, InResourceState, InitialMipData, NumInitialMips);
+	return GDynamicRHI->RHIAsyncCreateTexture2D(SizeX, SizeY, Format, NumMips, Flags, ResourceState, InitialMipData, NumInitialMips);
+}
+
+//UE_DEPRECATED(5.1, "FRHITexture2DArray is deprecated, please use RHICreateTexture(const FRHITextureCreateDesc&).")
+FORCEINLINE FTexture2DArrayRHIRef RHICreateTexture2DArray(uint32 SizeX, uint32 SizeY, uint32 ArraySize, uint8 Format, uint32 NumMips, uint32 NumSamples, ETextureCreateFlags Flags, ERHIAccess ResourceState, const FRHIResourceCreateInfo& CreateInfo)
+{
+	const uint32 Depth = 1;
+
+	return RHICreateTexture(
+		FRHITextureCreateDesc(
+			  ETextureDimension::Texture2DArray
+			, Flags
+			, (EPixelFormat)Format
+			, CreateInfo.ClearValueBinding
+			, { (int32)SizeX, (int32)SizeY }
+			, Depth
+			, ArraySize
+			, NumMips
+			, NumSamples
+			, CreateInfo.ExtData
+			, ResourceState
+			, CreateInfo.DebugName
+			, CreateInfo.BulkData
+			, CreateInfo.GPUMask
+		)
+	);
 }
 
-FORCEINLINE FTexture2DRHIRef RHICreateTexture2D(uint32 SizeX, uint32 SizeY, uint8 Format, uint32 NumMips, uint32 NumSamples, ETextureCreateFlags Flags, FRHIResourceCreateInfo& CreateInfo)
+//UE_DEPRECATED(5.1, "FRHITexture3D is deprecated, please use RHICreateTexture(const FRHITextureCreateDesc&).")
+FORCEINLINE FTexture3DRHIRef RHICreateTexture3D(uint32 SizeX, uint32 SizeY, uint32 SizeZ, uint8 Format, uint32 NumMips, ETextureCreateFlags Flags, ERHIAccess ResourceState, const FRHIResourceCreateInfo& CreateInfo)
+{
+	const uint32 ArraySize = 1;
+	const uint32 NumSamples = 1;
+
+	return RHICreateTexture(
+		FRHITextureCreateDesc(
+			  ETextureDimension::Texture3D
+			, Flags
+			, (EPixelFormat)Format
+			, CreateInfo.ClearValueBinding
+			, { (int32)SizeX, (int32)SizeY }
+			, SizeZ
+			, ArraySize
+			, NumMips
+			, NumSamples
+			, CreateInfo.ExtData
+			, ResourceState
+			, CreateInfo.DebugName
+			, CreateInfo.BulkData
+			, CreateInfo.GPUMask
+		)
+	);
+}
+
+//UE_DEPRECATED(5.1, "FRHITextureCube is deprecated, please use RHICreateTexture(const FRHITextureCreateDesc&).")
+FORCEINLINE FTextureCubeRHIRef RHICreateTextureCube(uint32 Size, uint8 Format, uint32 NumMips, ETextureCreateFlags Flags, ERHIAccess ResourceState, const FRHIResourceCreateInfo& CreateInfo)
+{
+	const uint32 Depth = 1;
+	const uint32 ArraySize = 1;
+	const uint32 NumSamples = 1;
+
+	return RHICreateTexture(
+		FRHITextureCreateDesc(
+			  ETextureDimension::TextureCube
+			, Flags
+			, (EPixelFormat)Format
+			, CreateInfo.ClearValueBinding
+			, { (int32)Size, (int32)Size }
+			, Depth
+			, ArraySize
+			, NumMips
+			, NumSamples
+			, CreateInfo.ExtData
+			, ResourceState
+			, CreateInfo.DebugName
+			, CreateInfo.BulkData
+			, CreateInfo.GPUMask
+		)
+	);
+}
+
+//UE_DEPRECATED(5.1, "FRHITextureCube is deprecated, please use RHICreateTexture(const FRHITextureCreateDesc&).")
+FORCEINLINE FTextureCubeRHIRef RHICreateTextureCubeArray(uint32 Size, uint32 ArraySize, uint8 Format, uint32 NumMips, ETextureCreateFlags Flags, ERHIAccess ResourceState, const FRHIResourceCreateInfo& CreateInfo)
+{
+	const uint32 Depth = 1;
+	const uint32 NumSamples = 1;
+
+	return RHICreateTexture(
+		FRHITextureCreateDesc(
+			  ETextureDimension::TextureCubeArray
+			, Flags
+			, (EPixelFormat)Format
+			, CreateInfo.ClearValueBinding
+			, { (int32)Size, (int32)Size }
+			, Depth
+			, ArraySize
+			, NumMips
+			, NumSamples
+			, CreateInfo.ExtData
+			, ResourceState
+			, CreateInfo.DebugName
+			, CreateInfo.BulkData
+			, CreateInfo.GPUMask
+		)
+	);
+}
+
+//UE_DEPRECATED(5.1, "FRHITexture2D is deprecated, please use RHICreateTexture(const FRHITextureCreateDesc&).")
+FORCEINLINE FTexture2DRHIRef RHICreateTexture2D(uint32 SizeX, uint32 SizeY, uint8 Format, uint32 NumMips, uint32 NumSamples, ETextureCreateFlags Flags, const FRHIResourceCreateInfo& CreateInfo)
 {
 	bool bHasInitialData = CreateInfo.BulkData != nullptr;
 	ERHIAccess ResourceState = RHIGetDefaultResourceState((ETextureCreateFlags)Flags, bHasInitialData);
 	return RHICreateTexture2D(SizeX, SizeY, Format, NumMips, NumSamples, Flags, ResourceState, CreateInfo);
 }
 
-FORCEINLINE FTexture2DRHIRef RHICreateTextureExternal2D(uint32 SizeX, uint32 SizeY, uint8 Format, uint32 NumMips, uint32 NumSamples, ETextureCreateFlags Flags, FRHIResourceCreateInfo& CreateInfo)
+//UE_DEPRECATED(5.1, "FRHITexture2D is deprecated, please use RHICreateTexture(const FRHITextureCreateDesc&).")
+FORCEINLINE FTexture2DRHIRef RHICreateTextureExternal2D(uint32 SizeX, uint32 SizeY, uint8 Format, uint32 NumMips, uint32 NumSamples, ETextureCreateFlags Flags, const FRHIResourceCreateInfo& CreateInfo)
 {
 	bool bHasInitialData = CreateInfo.BulkData != nullptr;
 	ERHIAccess ResourceState = RHIGetDefaultResourceState((ETextureCreateFlags)Flags, bHasInitialData);
 	return RHICreateTextureExternal2D(SizeX, SizeY, Format, NumMips, NumSamples, Flags, ResourceState, CreateInfo);
 }
 
+//UE_DEPRECATED(5.1, "FRHITexture2D is deprecated, please use RHICreateTexture(const FRHITextureCreateDesc&).")
 FORCEINLINE FTexture2DRHIRef RHIAsyncCreateTexture2D(uint32 SizeX, uint32 SizeY, uint8 Format, uint32 NumMips, ETextureCreateFlags Flags, void** InitialMipData, uint32 NumInitialMips)
 {
 	bool bHasInitialData = InitialMipData != nullptr;
@@ -5400,35 +5558,41 @@ FORCEINLINE FTexture2DRHIRef RHIAsyncCreateTexture2D(uint32 SizeX, uint32 SizeY,
 	return RHIAsyncCreateTexture2D(SizeX, SizeY, Format, NumMips, Flags, ResourceState, InitialMipData, NumInitialMips);
 }
 
-FORCEINLINE void RHICopySharedMips(FRHITexture2D* DestTexture2D, FRHITexture2D* SrcTexture2D)
+//UE_DEPRECATED(5.1, "FRHITexture2DArray is deprecated, please use RHICreateTexture(const FRHITextureCreateDesc&).")
+FORCEINLINE FTexture2DArrayRHIRef RHICreateTexture2DArray(uint32 SizeX, uint32 SizeY, uint32 SizeZ, uint8 Format, uint32 NumMips, uint32 NumSamples, ETextureCreateFlags Flags, const FRHIResourceCreateInfo& CreateInfo)
 {
-	return FRHICommandListExecutor::GetImmediateCommandList().CopySharedMips(DestTexture2D, SrcTexture2D);
+	bool bHasInitialData = CreateInfo.BulkData != nullptr;
+	ERHIAccess ResourceState = RHIGetDefaultResourceState((ETextureCreateFlags)Flags, bHasInitialData);
+	return RHICreateTexture2DArray(SizeX, SizeY, SizeZ, Format, NumMips, NumSamples, Flags, ResourceState, CreateInfo);
 }
 
-FORCEINLINE FTexture2DArrayRHIRef RHICreateTexture2DArray(uint32 SizeX, uint32 SizeY, uint32 SizeZ, uint8 Format, uint32 NumMips, uint32 NumSamples, ETextureCreateFlags Flags, ERHIAccess InResourceState, FRHIResourceCreateInfo& CreateInfo)
+//UE_DEPRECATED(5.1, "FRHITexture3D is deprecated, please use RHICreateTexture(const FRHITextureCreateDesc&).")
+FORCEINLINE FTexture3DRHIRef RHICreateTexture3D(uint32 SizeX, uint32 SizeY, uint32 SizeZ, uint8 Format, uint32 NumMips, ETextureCreateFlags Flags, const FRHIResourceCreateInfo& CreateInfo)
 {
-	LLM_SCOPE(EnumHasAnyFlags(Flags, TexCreate_RenderTargetable | TexCreate_DepthStencilTargetable) ? ELLMTag::RenderTargets : ELLMTag::Textures);
-	return GDynamicRHI->RHICreateTexture2DArray_RenderThread(FRHICommandListExecutor::GetImmediateCommandList(), SizeX, SizeY, SizeZ, Format, NumMips, NumSamples, Flags, InResourceState, CreateInfo);
+	bool bHasInitialData = CreateInfo.BulkData != nullptr;
+	ERHIAccess ResourceState = RHIGetDefaultResourceState((ETextureCreateFlags)Flags, bHasInitialData);
+	return RHICreateTexture3D(SizeX, SizeY, SizeZ, Format, NumMips, Flags, ResourceState, CreateInfo);
 }
 
-FORCEINLINE FTexture2DArrayRHIRef RHICreateTexture2DArray(uint32 SizeX, uint32 SizeY, uint32 SizeZ, uint8 Format, uint32 NumMips, uint32 NumSamples, ETextureCreateFlags Flags, FRHIResourceCreateInfo& CreateInfo)
+//UE_DEPRECATED(5.1, "FRHITextureCube is deprecated, please use RHICreateTexture(const FRHITextureCreateDesc&).")
+FORCEINLINE FTextureCubeRHIRef RHICreateTextureCube(uint32 Size, uint8 Format, uint32 NumMips, ETextureCreateFlags Flags, const FRHIResourceCreateInfo& CreateInfo)
 {
 	bool bHasInitialData = CreateInfo.BulkData != nullptr;
 	ERHIAccess ResourceState = RHIGetDefaultResourceState((ETextureCreateFlags)Flags, bHasInitialData);
-	return RHICreateTexture2DArray(SizeX, SizeY, SizeZ, Format, NumMips, NumSamples, Flags, ResourceState, CreateInfo);
+	return RHICreateTextureCube(Size, Format, NumMips, Flags, ResourceState, CreateInfo);
 }
 
-FORCEINLINE FTexture3DRHIRef RHICreateTexture3D(uint32 SizeX, uint32 SizeY, uint32 SizeZ, uint8 Format, uint32 NumMips, ETextureCreateFlags Flags, ERHIAccess ResourceState, FRHIResourceCreateInfo& CreateInfo)
+//UE_DEPRECATED(5.1, "FRHITextureCube is deprecated, please use RHICreateTexture(const FRHITextureCreateDesc&).")
+FORCEINLINE FTextureCubeRHIRef RHICreateTextureCubeArray(uint32 Size, uint32 ArraySize, uint8 Format, uint32 NumMips, ETextureCreateFlags Flags, const FRHIResourceCreateInfo& CreateInfo)
 {
-	LLM_SCOPE(EnumHasAnyFlags(Flags, TexCreate_RenderTargetable | TexCreate_DepthStencilTargetable) ? ELLMTag::RenderTargets : ELLMTag::Textures);
-	return GDynamicRHI->RHICreateTexture3D_RenderThread(FRHICommandListExecutor::GetImmediateCommandList(), SizeX, SizeY, SizeZ, Format, NumMips, Flags, ResourceState, CreateInfo);
+	bool bHasInitialData = CreateInfo.BulkData != nullptr;
+	ERHIAccess ResourceState = RHIGetDefaultResourceState((ETextureCreateFlags)Flags, bHasInitialData);
+	return RHICreateTextureCubeArray(Size, ArraySize, Format, NumMips, Flags, ResourceState, CreateInfo);
 }
 
-FORCEINLINE FTexture3DRHIRef RHICreateTexture3D(uint32 SizeX, uint32 SizeY, uint32 SizeZ, uint8 Format, uint32 NumMips, ETextureCreateFlags Flags, FRHIResourceCreateInfo& CreateInfo)
+FORCEINLINE void RHICopySharedMips(FRHITexture* DestTexture, FRHITexture* SrcTexture)
 {
-	bool bHasInitialData = CreateInfo.BulkData != nullptr;
-	ERHIAccess ResourceState = RHIGetDefaultResourceState((ETextureCreateFlags)Flags, bHasInitialData);
-	return RHICreateTexture3D(SizeX, SizeY, SizeZ, Format, NumMips, Flags, ResourceState, CreateInfo);
+	return FRHICommandListExecutor::GetImmediateCommandList().CopySharedMips(DestTexture, SrcTexture);
 }
 
 FORCEINLINE FShaderResourceViewRHIRef RHICreateShaderResourceView(FRHITexture* Texture, uint8 MipLevel)
@@ -5516,32 +5680,6 @@ FORCEINLINE void RHIUpdateTexture3D(FRHITexture3D* Texture, uint32 MipIndex, con
 	 FRHICommandListExecutor::GetImmediateCommandList().UpdateTexture3D(Texture, MipIndex, UpdateRegion, SourceRowPitch, SourceDepthPitch, SourceData);
 }
 
-FORCEINLINE FTextureCubeRHIRef RHICreateTextureCube(uint32 Size, uint8 Format, uint32 NumMips, ETextureCreateFlags Flags, ERHIAccess InResourceState, FRHIResourceCreateInfo& CreateInfo)
-{
-	LLM_SCOPE(EnumHasAnyFlags(Flags, TexCreate_RenderTargetable | TexCreate_DepthStencilTargetable) ? ELLMTag::RenderTargets : ELLMTag::Textures);
-	return GDynamicRHI->RHICreateTextureCube_RenderThread(FRHICommandListExecutor::GetImmediateCommandList(), Size, Format, NumMips, Flags, InResourceState, CreateInfo);
-}
-
-FORCEINLINE FTextureCubeRHIRef RHICreateTextureCubeArray(uint32 Size, uint32 ArraySize, uint8 Format, uint32 NumMips, ETextureCreateFlags Flags, ERHIAccess InResourceState, FRHIResourceCreateInfo& CreateInfo)
-{
-	LLM_SCOPE(EnumHasAnyFlags(Flags, TexCreate_RenderTargetable | TexCreate_DepthStencilTargetable) ? ELLMTag::RenderTargets : ELLMTag::Textures);
-	return GDynamicRHI->RHICreateTextureCubeArray_RenderThread(FRHICommandListExecutor::GetImmediateCommandList(), Size, ArraySize, Format, NumMips, Flags, InResourceState, CreateInfo);
-}
-
-FORCEINLINE FTextureCubeRHIRef RHICreateTextureCube(uint32 Size, uint8 Format, uint32 NumMips, ETextureCreateFlags Flags, FRHIResourceCreateInfo& CreateInfo)
-{
-	bool bHasInitialData = CreateInfo.BulkData != nullptr;
-	ERHIAccess ResourceState = RHIGetDefaultResourceState((ETextureCreateFlags)Flags, bHasInitialData);
-	return RHICreateTextureCube(Size, Format, NumMips, Flags, ResourceState, CreateInfo);
-}
-
-FORCEINLINE FTextureCubeRHIRef RHICreateTextureCubeArray(uint32 Size, uint32 ArraySize, uint8 Format, uint32 NumMips, ETextureCreateFlags Flags, FRHIResourceCreateInfo& CreateInfo)
-{
-	bool bHasInitialData = CreateInfo.BulkData != nullptr;
-	ERHIAccess ResourceState = RHIGetDefaultResourceState((ETextureCreateFlags)Flags, bHasInitialData);
-	return RHICreateTextureCubeArray(Size, ArraySize, Format, NumMips, Flags, ResourceState, CreateInfo);
-}
-
 FORCEINLINE void* RHILockTextureCubeFace(FRHITextureCube* Texture, uint32 FaceIndex, uint32 ArrayIndex, uint32 MipIndex, EResourceLockMode LockMode, uint32& DestStride, bool bLockWithinMiptail)
 {
 	return FRHICommandListExecutor::GetImmediateCommandList().LockTextureCubeFace(Texture, FaceIndex, ArrayIndex, MipIndex, LockMode, DestStride, bLockWithinMiptail);
diff --git a/Engine/Source/Runtime/RHI/Public/RHIDefinitions.h b/Engine/Source/Runtime/RHI/Public/RHIDefinitions.h
--- a/Engine/Source/Runtime/RHI/Public/RHIDefinitions.h
+++ b/Engine/Source/Runtime/RHI/Public/RHIDefinitions.h
@@ -1682,7 +1682,7 @@ enum class EGpuVendorId
 };
 
 /** An enumeration of the different RHI reference types. */
-enum ERHIResourceType
+enum ERHIResourceType : uint8
 {
 	RRT_None,
 
@@ -1706,10 +1706,12 @@ enum ERHIResourceType
 	RRT_UniformBuffer,
 	RRT_Buffer,
 	RRT_Texture,
+	// @todo: texture type unification - remove these
 	RRT_Texture2D,
 	RRT_Texture2DArray,
 	RRT_Texture3D,
 	RRT_TextureCube,
+	// @todo: texture type unification - remove these
 	RRT_TextureReference,
 	RRT_TimestampCalibrationQuery,
 	RRT_GPUFence,
@@ -1729,7 +1731,7 @@ enum ERHIResourceType
 };
 
 /** Describes the dimension of a texture. */
-enum class ETextureDimension
+enum class ETextureDimension : uint8
 {
 	Texture2D,
 	Texture2DArray,
diff --git a/Engine/Source/Runtime/RHI/Public/RHIResources.h b/Engine/Source/Runtime/RHI/Public/RHIResources.h
--- a/Engine/Source/Runtime/RHI/Public/RHIResources.h
+++ b/Engine/Source/Runtime/RHI/Public/RHIResources.h
@@ -39,11 +39,6 @@
 
 #define RHI_ENABLE_RESOURCE_INFO (RHI_WANT_RESOURCE_INFO && !RHI_FORCE_DISABLE_RESOURCE_INFO)
 
-#ifndef RHI_RESOURCE_LIFETIME_VALIDATION
-#define RHI_RESOURCE_LIFETIME_VALIDATION 0
-#endif
-
-
 class FRHICommandListImmediate;
 struct FClearValueBinding;
 struct FRHIResourceInfo;
@@ -57,11 +52,19 @@ public:
 	UE_DEPRECATED(5.0, "FRHIResource(bool) is deprecated, please use FRHIResource(ERHIResourceType)")
 	FRHIResource(bool InbDoNotDeferDelete=false)
 		: ResourceType(RRT_None)
+		, bCommitted(true)
+#if RHI_ENABLE_RESOURCE_INFO
+		, bBeingTracked(false)
+#endif
 	{
 	}
 
 	FRHIResource(ERHIResourceType InResourceType)
 		: ResourceType(InResourceType)
+		, bCommitted(true)
+#if RHI_ENABLE_RESOURCE_INFO
+		, bBeingTracked(false)
+#endif
 	{
 #if RHI_ENABLE_RESOURCE_INFO
 		BeginTrackingResource(this);
@@ -161,162 +164,85 @@ public:
 private:
 	class FAtomicFlags
 	{
-#if RHI_RESOURCE_LIFETIME_VALIDATION
-		struct FPacked
-		{
-			FPacked(uint64 InNumRefs, bool InMarkedForDelete, bool InDeleting) : NumRefs(InNumRefs), MarkedForDelete(InMarkedForDelete), Deleting(InDeleting)
-			{
-			}
+		static constexpr uint32 MarkedForDeleteBit    = 1 << 30;
+		static constexpr uint32 DeletingBit           = 1 << 31;
+		static constexpr uint32 NumRefsMask           = ~(MarkedForDeleteBit | DeletingBit);
 
-			FPacked() : NumRefs(0), MarkedForDelete(0), Deleting(0)
-			{
-			}
-
-			uint32 NumRefs			: 30;
-			uint32 MarkedForDelete	: 1;
-			uint32 Deleting			: 1;
-		};
-		std::atomic<FPacked> Packed = {};
-#else
-		std::atomic_int NumRefs = { 0 };
-		std::atomic_bool MarkedForDelete = { 0 };
-#endif
+		std::atomic_uint Packed = { 0 };
 
 	public:
 		int32 AddRef(std::memory_order MemoryOrder)
 		{
-#if RHI_RESOURCE_LIFETIME_VALIDATION
-			FPacked Old = Packed.load(std::memory_order_relaxed);
-			while(true)
-			{
-				check(Old.Deleting == false);
-				if(Packed.compare_exchange_weak(Old, FPacked(Old.NumRefs + 1, Old.MarkedForDelete, false), MemoryOrder, std::memory_order_relaxed))
-				{
-					return Old.NumRefs + 1;
-				}
-			}
-#else
-			return NumRefs.fetch_add(1, MemoryOrder) + 1;
-#endif
+			uint32 OldPacked = Packed.fetch_add(1, MemoryOrder);
+			checkf((OldPacked & DeletingBit) == 0, TEXT("Resource is being deleted."));
+			int32  NumRefs = (OldPacked & NumRefsMask) + 1;
+			checkf(NumRefs < NumRefsMask, TEXT("Reference count has overflowed."));
+			return NumRefs;
 		}
 
 		int32 Release(std::memory_order MemoryOrder)
 		{
-#if RHI_RESOURCE_LIFETIME_VALIDATION
-			FPacked Old = Packed.load(std::memory_order_relaxed);
-			while(true)
-			{
-				check(Old.Deleting == false);
-				if(Packed.compare_exchange_weak(Old, FPacked(Old.NumRefs - 1, Old.MarkedForDelete, false), MemoryOrder, std::memory_order_relaxed))
-				{
-					return Old.NumRefs - 1;
-				}
-			}
-#else
-			return NumRefs.fetch_sub(1, MemoryOrder) - 1;
-#endif
+			uint32 OldPacked = Packed.fetch_sub(1, MemoryOrder);
+			checkf((OldPacked & DeletingBit) == 0, TEXT("Resource is being deleted."));
+			int32  NumRefs = (OldPacked & NumRefsMask) - 1;
+			checkf(NumRefs >= 0, TEXT("Reference count has underflowed."));
+			return NumRefs;
 		}
 
 		bool MarkForDelete(std::memory_order MemoryOrder)
 		{
-#if RHI_RESOURCE_LIFETIME_VALIDATION
-			FPacked Old = Packed.load(std::memory_order_relaxed);
-			while(true)
-			{
-				check(Old.Deleting == false);
-				if(Packed.compare_exchange_weak(Old, FPacked(Old.NumRefs, true, false), MemoryOrder, std::memory_order_relaxed))
-				{
-					return Old.MarkedForDelete;
-				}
-			}
-#else
-			return MarkedForDelete.exchange(true, MemoryOrder);
-#endif
+			uint32 OldPacked = Packed.fetch_or(MarkedForDeleteBit, MemoryOrder);
+			check((OldPacked & DeletingBit) == 0);
+			return (OldPacked & MarkedForDeleteBit) != 0;
 		}
 
 		bool UnmarkForDelete(std::memory_order MemoryOrder)
 		{
-#if RHI_RESOURCE_LIFETIME_VALIDATION
-			FPacked Old = Packed.load(std::memory_order_relaxed);
-			while(true)
-			{
-				check(Old.Deleting == false); 
-				check(Old.MarkedForDelete == true);
-				if(Packed.compare_exchange_weak(Old, FPacked(Old.NumRefs, false, false), MemoryOrder, std::memory_order_relaxed))
-				{
-					return Old.MarkedForDelete;
-				}
-			}
-#else
-			bool Old = MarkedForDelete.exchange(false, MemoryOrder);
-			check(Old == true);
-			return Old;
-#endif
+			uint32 OldPacked = Packed.fetch_xor(MarkedForDeleteBit, MemoryOrder);
+			check((OldPacked & DeletingBit) == 0);
+			bool  OldMarkedForDelete = (OldPacked & MarkedForDeleteBit) != 0;
+			check(OldMarkedForDelete == true);
+			return OldMarkedForDelete;
 		}
 
 		bool Deleteing()
 		{
-#if RHI_RESOURCE_LIFETIME_VALIDATION
-			FPacked Old = Packed.load(std::memory_order_relaxed);
-			while(true)
-			{
-				check(Old.Deleting == false); 
-				check(Old.MarkedForDelete == true);
-				if(Old.NumRefs == 0)
-				{
-					if(Packed.compare_exchange_weak(Old, FPacked(0, true, true), std::memory_order_acquire, std::memory_order_relaxed))
-					{
-						return true;
-					}
-				}
-				else
-				{
-					if(Packed.compare_exchange_weak(Old, FPacked(Old.NumRefs, false, false), std::memory_order_release, std::memory_order_relaxed))
-					{
-						return false;
-					}
-				}
-			}
-#else
-			check(MarkedForDelete.load(std::memory_order_relaxed) == 1);
-			if (NumRefs.load(std::memory_order_acquire) == 0) // caches can bring dead objects back to life
+			uint32 LocalPacked = Packed.load(std::memory_order_acquire);
+			check((LocalPacked & MarkedForDeleteBit) != 0);
+			check((LocalPacked & DeletingBit) == 0);
+			uint32 NumRefs = LocalPacked & NumRefsMask;
+
+			if (NumRefs == 0) // caches can bring dead objects back to life
 			{
+#if DO_CHECK
+				Packed.fetch_or(DeletingBit, std::memory_order_acquire);
+#endif
 				return true;
 			}
 			else
 			{
-				verify(MarkedForDelete.exchange(false, std::memory_order_release));	
+				UnmarkForDelete(std::memory_order_release);
 				return false;
 			}
-#endif
 		}
 
 		bool IsValid(std::memory_order MemoryOrder)
 		{
-#if RHI_RESOURCE_LIFETIME_VALIDATION
-			FPacked Old = Packed.load(MemoryOrder);
-			return !Old.MarkedForDelete && Old.NumRefs > 0;
-#else
-			return !MarkedForDelete.load(MemoryOrder) && NumRefs.load(MemoryOrder) > 0;
-#endif
+			uint32 LocalPacked = Packed.load(MemoryOrder);
+			return (LocalPacked & MarkedForDeleteBit) == 0 && (LocalPacked & NumRefsMask) != 0;
 		}
 
 		int32 GetNumRefs(std::memory_order MemoryOrder)
 		{
-#if RHI_RESOURCE_LIFETIME_VALIDATION
-			FPacked Old = Packed.load(MemoryOrder);
-			return Old.NumRefs;
-#else
-			return NumRefs.load(MemoryOrder);
-#endif
+			return Packed.load(MemoryOrder) & NumRefsMask;
 		}
 	};
 	mutable FAtomicFlags AtomicFlags;
 
 	const ERHIResourceType ResourceType;
-	bool bCommitted { true };
+	uint8 bCommitted : 1;
 #if RHI_ENABLE_RESOURCE_INFO
-	bool bBeingTracked { false };
+	uint8 bBeingTracked : 1;
 #endif
 
 	static std::atomic<TClosableMpscQueue<FRHIResource*>*> PendingDeletes;
@@ -1228,39 +1154,516 @@ private:
 	double LastRenderTime;
 };
 
-class RHI_API FRHITexture : public FRHIResource
-#if ENABLE_RHI_VALIDATION
-	, public RHIValidation::FTextureResource
-#endif
+
+/** Descriptor used to create a texture resource */
+struct RHI_API FRHITextureDesc
 {
-public:
-	
-	/** Initialization constructor. */
-	FRHITexture(ERHIResourceType InResourceType, uint32 InNumMips, uint32 InNumSamples, EPixelFormat InFormat, ETextureCreateFlags InFlags, const FClearValueBinding& InClearValue)
-		: FRHIResource(InResourceType)
+	static FRHITextureDesc Create2D(
+		  FIntPoint           Size
+		, EPixelFormat        Format
+		, FClearValueBinding  ClearValue
+		, ETextureCreateFlags Flags
+		, uint8               NumMips    = 1
+		, uint8               NumSamples = 1
+		, uint32              ExtData    = 0
+		)
+	{
+		const uint32 Depth     = 1;
+		const uint32 ArraySize = 1;
+		return FRHITextureDesc(ETextureDimension::Texture2D, Flags, Format, ClearValue, { Size.X, Size.Y }, Depth, ArraySize, NumMips, NumSamples, ExtData);
+	}
+
+	static FRHITextureDesc Create2DArray(
+		  FIntPoint           Size
+		, EPixelFormat        Format
+		, FClearValueBinding  ClearValue
+		, ETextureCreateFlags Flags
+		, uint16              ArraySize
+		, uint8               NumMips    = 1
+		, uint8               NumSamples = 1
+		, uint32              ExtData    = 0
+		)
+	{
+		const uint32 Depth   = 1;
+		return FRHITextureDesc(ETextureDimension::Texture2DArray, Flags, Format, ClearValue, { Size.X, Size.Y }, Depth, ArraySize, NumMips, NumSamples, ExtData);
+	}
+
+	static FRHITextureDesc Create3D(
+		  FIntVector          Size
+		, EPixelFormat        Format
+		, FClearValueBinding  ClearValue
+		, ETextureCreateFlags Flags
+		, uint8               NumMips    = 1
+		, uint32              ExtData    = 0
+		)
+	{
+		const uint32 ArraySize  = 1;
+		const uint32 LocalNumSamples = 1;
+
+		checkf(Size.Z <= TNumericLimits<decltype(FRHITextureDesc::Depth)>::Max(), TEXT("Depth parameter (Size.Z) exceeds valid range"));
+
+		return FRHITextureDesc(ETextureDimension::Texture3D, Flags, Format, ClearValue, { Size.X, Size.Y }, Size.Z, ArraySize, NumMips, LocalNumSamples, ExtData);
+	}
+
+	static FRHITextureDesc CreateCube(
+		  uint32              Size
+		, EPixelFormat        Format
+		, FClearValueBinding  ClearValue
+		, ETextureCreateFlags Flags
+		, uint8               NumMips    = 1
+		, uint8               NumSamples = 1
+		, uint32              ExtData    = 0
+		)
+	{
+		checkf(Size <= (uint32)TNumericLimits<int32>::Max(), TEXT("Size parameter exceeds valid range"));
+
+		const uint32 Depth     = 1;
+		const uint32 ArraySize = 1;
+		return FRHITextureDesc(ETextureDimension::TextureCube, Flags, Format, ClearValue, { (int32)Size, (int32)Size }, Depth, ArraySize, NumMips, NumSamples, ExtData);
+	}
+
+	static FRHITextureDesc CreateCubeArray(
+		  uint32              Size
+		, EPixelFormat        Format
+		, FClearValueBinding  ClearValue
+		, ETextureCreateFlags Flags
+		, uint16              ArraySize
+		, uint8               NumMips    = 1
+		, uint8               NumSamples = 1
+		, uint32              ExtData    = 0
+		)
+	{
+		checkf(Size <= (uint32)TNumericLimits<int32>::Max(), TEXT("Size parameter exceeds valid range"));
+
+		const uint32 Depth   = 1;
+		return FRHITextureDesc(ETextureDimension::TextureCubeArray, Flags, Format, ClearValue, { (int32)Size, (int32)Size }, Depth, ArraySize, NumMips, NumSamples, ExtData);
+	}
+
+	FRHITextureDesc()
+		: NumSamples(1)
+		, Dimension(ETextureDimension::Texture2D)
+	{}
+
+	FRHITextureDesc(
+		  ETextureDimension   InDimension
+		, ETextureCreateFlags InFlags
+		, EPixelFormat        InFormat
+		, FClearValueBinding  InClearValue
+		, FIntPoint           InExtent
+		, uint16              InDepth
+		, uint16              InArraySize
+		, uint8               InNumMips
+		, uint8               InNumSamples
+		, uint32              InExtData
+		)
+		: Flags     (InFlags     )
 		, ClearValue(InClearValue)
-		, NumMips(InNumMips)
+		, ExtData   (InExtData   )
+		, Extent    (InExtent    )
+		, Depth     (InDepth     )
+		, ArraySize (InArraySize )
+		, NumMips   (InNumMips   )
 		, NumSamples(InNumSamples)
-		, Format(InFormat)
-		, Flags(InFlags)
+		, Dimension (InDimension )
+		, Format    (InFormat    )
 	{}
 
-	UE_DEPRECATED(5.0, "The InLastRenderTime parameter will be removed in the future")
-	FRHITexture(ERHIResourceType InResourceType, uint32 InNumMips, uint32 InNumSamples, EPixelFormat InFormat, ETextureCreateFlags InFlags, FLastRenderTimeContainer* InLastRenderTime, const FClearValueBinding& InClearValue)
-		: FRHITexture(InResourceType, InNumMips, InNumSamples, InFormat, InFlags, InClearValue)
+	UE_DEPRECATED(5.1, "Prefer using one of the FRHITextureDesc::Create...() functions rather than this FRHITextureDesc constructor. Otherwise use the FRHITextureDesc constructor that does not take optional arguments.")
+	FRHITextureDesc(
+		  ETextureDimension   InDimension
+		, ETextureCreateFlags InFlags
+		, EPixelFormat        InFormat
+		, FIntPoint           InExtent
+		, FClearValueBinding  InClearValue
+		, uint16              InDepth      = 1
+		, uint16              InArraySize  = 1
+		, uint8               InNumMips    = 1
+		, uint8               InNumSamples = 1
+		, uint32              InExtData    = 0
+		)
+		: FRHITextureDesc(
+		  InDimension
+		, InFlags
+		, InFormat
+		, InClearValue
+		, InExtent
+		, InDepth
+		, InArraySize
+		, InNumMips
+		, InNumSamples
+		, InExtData
+		)
 	{}
 
-	// Dynamic cast methods.
-	virtual class FRHITexture2D* GetTexture2D() { return NULL; }
-	virtual class FRHITexture2DArray* GetTexture2DArray() { return NULL; }
-	virtual class FRHITexture3D* GetTexture3D() { return NULL; }
-	virtual class FRHITextureCube* GetTextureCube() { return NULL; }
-	virtual class FRHITextureReference* GetTextureReference() { return NULL; }
+	friend uint32 GetTypeHash(const FRHITextureDesc& Desc)
+	{
+		uint32 Hash = GetTypeHash(Desc.Dimension);
+		Hash = HashCombine(Hash, GetTypeHash(Desc.Flags		));
+		Hash = HashCombine(Hash, GetTypeHash(Desc.Format	));
+		Hash = HashCombine(Hash, GetTypeHash(Desc.UAVFormat	));
+		Hash = HashCombine(Hash, GetTypeHash(Desc.Extent	));
+		Hash = HashCombine(Hash, GetTypeHash(Desc.Depth		));
+		Hash = HashCombine(Hash, GetTypeHash(Desc.ArraySize	));
+		Hash = HashCombine(Hash, GetTypeHash(Desc.NumMips	));
+		Hash = HashCombine(Hash, GetTypeHash(Desc.NumSamples));
+		Hash = HashCombine(Hash, GetTypeHash(Desc.ClearValue));
+		Hash = HashCombine(Hash, GetTypeHash(Desc.ExtData   ));
+		return Hash;
+	}
+
+	bool operator == (const FRHITextureDesc& Other) const
+	{
+		return Dimension  == Other.Dimension
+			&& Flags      == Other.Flags
+			&& Format     == Other.Format
+			&& UAVFormat  == Other.UAVFormat
+			&& Extent     == Other.Extent
+			&& Depth      == Other.Depth
+			&& ArraySize  == Other.ArraySize
+			&& NumMips    == Other.NumMips
+			&& NumSamples == Other.NumSamples
+			&& ClearValue == Other.ClearValue
+			&& ExtData    == Other.ExtData;
+	}
+
+	bool operator != (const FRHITextureDesc& Other) const
+	{
+		return !(*this == Other);
+	}
+
+	bool IsTexture2D() const
+	{
+		return Dimension == ETextureDimension::Texture2D || Dimension == ETextureDimension::Texture2DArray;
+	}
+
+	bool IsTexture3D() const
+	{
+		return Dimension == ETextureDimension::Texture3D;
+	}
+
+	bool IsTextureCube() const
+	{
+		return Dimension == ETextureDimension::TextureCube || Dimension == ETextureDimension::TextureCubeArray;
+	}
+
+	bool IsTextureArray() const
+	{
+		return Dimension == ETextureDimension::Texture2DArray || Dimension == ETextureDimension::TextureCubeArray;
+	}
+
+	bool IsMipChain() const
+	{
+		return NumMips > 1;
+	}
+
+	bool IsMultisample() const
+	{
+		return NumSamples > 1;
+	}
+
+	FIntVector GetSize() const
+	{
+		return FIntVector(Extent.X, Extent.Y, Depth);
+	}
+
+	void Reset()
+	{
+		// Usually we don't want to propagate MSAA samples.
+		NumSamples = 1;
+
+		// Remove UAV flag for textures that don't need it (some formats are incompatible).
+		Flags |= TexCreate_RenderTargetable;
+		Flags &= ~(TexCreate_UAV | TexCreate_ResolveTargetable | TexCreate_DepthStencilResolveTarget);
+	}
+
+	/** Returns whether this descriptor conforms to requirements. */
+	bool IsValid() const
+	{
+		return FRHITextureDesc::Validate(*this, /* Name = */ TEXT(""), /* bFatal = */ false);
+	}
+
+	/** Texture flags passed on to RHI texture. */
+	ETextureCreateFlags Flags = TexCreate_None;
+
+	/** Clear value to use when fast-clearing the texture. */
+	FClearValueBinding ClearValue;
+
+	/** Platform-specific additional data. Used for offline processed textures on some platforms. */
+	uint32 ExtData = 0;
+
+	/** Extent of the texture in x and y. */
+	FIntPoint Extent = FIntPoint(1, 1);
+
+	/** Depth of the texture if the dimension is 3D. */
+	uint16 Depth = 1;
+
+	/** The number of array elements in the texture. (Keep at 1 if dimension is 3D). */
+	uint16 ArraySize = 1;
+
+	/** Number of mips in the texture mip-map chain. */
+	uint8 NumMips = 1;
+
+	/** Number of samples in the texture. >1 for MSAA. */
+	uint8 NumSamples : 5;
+
+	/** Texture dimension to use when creating the RHI texture. */
+	ETextureDimension Dimension : 3;
+
+	/** Pixel format used to create RHI texture. */
+	EPixelFormat Format = PF_Unknown;
+
+	/** Texture format used when creating the UAV. PF_Unknown means to use the default one (same as Format). */
+	EPixelFormat UAVFormat = PF_Unknown;
+
+	/** Check the validity. */
+	static bool CheckValidity(const FRHITextureDesc& Desc, const TCHAR* Name)
+	{
+		return FRHITextureDesc::Validate(Desc, Name, /* bFatal = */ true);
+	}
+
+	/**
+	 * Returns an estimated total memory size the described texture will occupy in GPU memory.
+	 * This is an estimate because it only considers the dimensions / format etc of the texture, 
+	 * not any specifics about platform texture layout.
+	 * 
+	 * To get a true measure of a texture resource for the current running platform RHI, use RHICalcTexturePlatformSize().
+	 * 
+	 * @param FirstMipIndex - the index of the most detailed mip to consider in the memory size calculation. Must be < NumMips and <= LastMipIndex.
+	 * @param LastMipIndex  - the index of the least detailed mip to consider in the memory size calculation. Must be < NumMips and >= FirstMipIndex.
+	 */
+	uint64 CalcMemorySizeEstimate(uint32 FirstMipIndex, uint32 LastMipIndex) const;
+	uint64 CalcMemorySizeEstimate(uint32 FirstMipIndex = 0) const
+	{
+		return CalcMemorySizeEstimate(FirstMipIndex, NumMips - 1);
+	}
+
+private:
+	static bool Validate(const FRHITextureDesc& Desc, const TCHAR* Name, bool bFatal);
+};
+
+// @todo deprecate
+using FRHITextureCreateInfo = FRHITextureDesc;
+
+extern RHI_API ERHIAccess RHIGetDefaultResourceState(ETextureCreateFlags InUsage, bool bInHasInitialData);
+
+struct FRHITextureCreateDesc : public FRHITextureDesc
+{
+	static FRHITextureCreateDesc Create2D(
+		  TCHAR const*                DebugName
+		, FIntPoint                   Size
+		, EPixelFormat                Format
+		, FClearValueBinding          ClearValue
+		, ETextureCreateFlags         Flags
+		, uint8                       NumMips      = 1
+		, uint8                       NumSamples   = 1
+		, uint32		              ExtData      = 0
+		, ERHIAccess                  InitialState = ERHIAccess::Unknown
+		, FResourceBulkDataInterface* BulkData     = nullptr
+		, FRHIGPUMask                 GPUMask      = FRHIGPUMask()
+		)
+	{
+		return FRHITextureCreateDesc(
+			  FRHITextureDesc::Create2D(Size, Format, ClearValue, Flags, NumMips, NumSamples, ExtData)
+			, (InitialState == ERHIAccess::Unknown) ? RHIGetDefaultResourceState(Flags, BulkData != nullptr) : InitialState
+			, DebugName
+			, BulkData
+			, GPUMask
+		);
+	}
+
+	static FRHITextureCreateDesc Create2DArray(
+		  TCHAR const*                DebugName
+		, FIntPoint                   Size
+		, EPixelFormat                Format
+		, FClearValueBinding          ClearValue
+		, ETextureCreateFlags         Flags
+		, uint16                      ArraySize
+		, uint8                       NumMips      = 1
+		, uint8                       NumSamples   = 1
+		, uint32		              ExtData      = 0
+		, ERHIAccess                  InitialState = ERHIAccess::Unknown
+		, FResourceBulkDataInterface* BulkData     = nullptr
+		, FRHIGPUMask                 GPUMask      = FRHIGPUMask()
+		)
+	{
+		return FRHITextureCreateDesc(
+			  FRHITextureDesc::Create2DArray(Size, Format, ClearValue, Flags, ArraySize, NumMips, NumSamples, ExtData)
+			, (InitialState == ERHIAccess::Unknown) ? RHIGetDefaultResourceState(Flags, BulkData != nullptr) : InitialState
+			, DebugName
+			, BulkData
+			, GPUMask
+		);
+	}
+
+	static FRHITextureCreateDesc Create3D(
+		  TCHAR const*                DebugName
+		, FIntVector                  Size
+		, EPixelFormat                Format
+		, FClearValueBinding          ClearValue
+		, ETextureCreateFlags         Flags
+		, uint8                       NumMips      = 1
+		, uint32		              ExtData      = 0
+		, ERHIAccess                  InitialState = ERHIAccess::Unknown
+		, FResourceBulkDataInterface* BulkData     = nullptr
+		, FRHIGPUMask                 GPUMask      = FRHIGPUMask()
+		)
+	{
+		return FRHITextureCreateDesc(
+			  FRHITextureDesc::Create3D(Size, Format, ClearValue, Flags, NumMips, ExtData)
+			, (InitialState == ERHIAccess::Unknown) ? RHIGetDefaultResourceState(Flags, BulkData != nullptr) : InitialState
+			, DebugName
+			, BulkData
+			, GPUMask
+		);
+	}
+
+	static FRHITextureCreateDesc CreateCube(
+		  TCHAR const*                DebugName
+		, uint32                      Size
+		, EPixelFormat                Format
+		, FClearValueBinding          ClearValue
+		, ETextureCreateFlags         Flags
+		, uint8                       NumMips      = 1
+		, uint8                       NumSamples   = 1
+		, uint32		              ExtData      = 0
+		, ERHIAccess                  InitialState = ERHIAccess::Unknown
+		, FResourceBulkDataInterface* BulkData     = nullptr
+		, FRHIGPUMask                 GPUMask      = FRHIGPUMask()
+		)
+	{
+		return FRHITextureCreateDesc(
+			  FRHITextureDesc::CreateCube(Size, Format, ClearValue, Flags, NumMips, NumSamples, ExtData)
+			, (InitialState == ERHIAccess::Unknown) ? RHIGetDefaultResourceState(Flags, BulkData != nullptr) : InitialState
+			, DebugName
+			, BulkData
+			, GPUMask
+		);
+	}
+
+	static FRHITextureCreateDesc CreateCubeArray(
+		  TCHAR const*                DebugName
+		, uint32                      Size
+		, EPixelFormat                Format
+		, FClearValueBinding          ClearValue
+		, ETextureCreateFlags         Flags
+		, uint16                      ArraySize
+		, uint8                       NumMips      = 1
+		, uint8                       NumSamples   = 1
+		, uint32		              ExtData      = 0
+		, ERHIAccess                  InitialState = ERHIAccess::Unknown
+		, FResourceBulkDataInterface* BulkData     = nullptr
+		, FRHIGPUMask                 GPUMask      = FRHIGPUMask()
+		)
+	{
+		return FRHITextureCreateDesc(
+			  FRHITextureDesc::CreateCubeArray(Size, Format, ClearValue, Flags, ArraySize, NumMips, NumSamples, ExtData)
+			, (InitialState == ERHIAccess::Unknown) ? RHIGetDefaultResourceState(Flags, BulkData != nullptr) : InitialState
+			, DebugName
+			, BulkData
+			, GPUMask
+		);
+	}
+
+	FRHITextureCreateDesc() = default;
+
+	// Constructor will full argument set.
+	FRHITextureCreateDesc(
+		  ETextureDimension           InDimension
+		, ETextureCreateFlags         InFlags
+		, EPixelFormat                InFormat
+		, FClearValueBinding          InClearValue
+		, FIntPoint                   InExtent
+		, uint16                      InDepth
+		, uint16                      InArraySize
+		, uint8                       InNumMips
+		, uint8                       InNumSamples
+		, uint32                      InExtData
+		, ERHIAccess                  InInitialState
+		, TCHAR const*                InDebugName
+		, FResourceBulkDataInterface* InBulkData    = nullptr 
+		, FRHIGPUMask                 InGPUMask     = FRHIGPUMask()
+		)
+		: FRHITextureDesc(InDimension, InFlags, InFormat, InClearValue, InExtent, InDepth, InArraySize, InNumMips, InNumSamples, InExtData)
+		, InitialState   (InInitialState)
+		, DebugName      (InDebugName)
+		, GPUMask        (InGPUMask)
+		, BulkData       (InBulkData)
+	{}
+
+	// Constructor for when you already have an FRHITextureDesc
+	FRHITextureCreateDesc(
+		  FRHITextureDesc const&      InDesc
+		, ERHIAccess                  InInitialState
+		, TCHAR const*                InDebugName
+		, FResourceBulkDataInterface* InBulkData     = nullptr
+		, FRHIGPUMask                 InGPUMask      = FRHIGPUMask()
+		)
+		: FRHITextureDesc(InDesc)
+		, InitialState   (InInitialState)
+		, DebugName      (InDebugName)
+		, GPUMask        (InGPUMask)
+		, BulkData       (InBulkData)
+	{}
+
+	void CheckValidity() const
+	{
+		FRHITextureDesc::CheckValidity(*this, DebugName);
+
+		ensureMsgf(InitialState != ERHIAccess::Unknown, TEXT("Resource %s cannot be created in an unknown state."), DebugName);
+	}
+
+	/* The RHI access state that the resource will be created in. */
+	ERHIAccess InitialState = ERHIAccess::Unknown;
+
+	/* A friendly name for the resource. */
+	const TCHAR* DebugName = nullptr;
+
+	/* A mask representing which GPUs to create the resource on, in a multi-GPU system. */
+	FRHIGPUMask GPUMask;
+
+	/* Optional initial data to fill the resource with. */
+	FResourceBulkDataInterface* BulkData = nullptr;
+};
+
+class FRHITexture;
+/*UE_DEPRECATED(5.1, "FRHITexture2D is deprecated, please use FRHITexture.")      */ typedef class FRHITexture FRHITexture2D;
+/*UE_DEPRECATED(5.1, "FRHITexture2DArray is deprecated, please use FRHITexture.") */ typedef class FRHITexture FRHITexture2DArray;
+/*UE_DEPRECATED(5.1, "FRHITexture3D is deprecated, please use FRHITexture.")      */ typedef class FRHITexture FRHITexture3D;
+/*UE_DEPRECATED(5.1, "FRHITextureCube is deprecated, please use FRHITexture.")    */ typedef class FRHITexture FRHITextureCube;
 
-	virtual FRHIDescriptorHandle GetDefaultBindlessHandle() const { return FRHIDescriptorHandle(); }
 
-	// Slower method to get Size X, Y & Z information. Prefer sub-classes' GetSizeX(), etc
-	virtual FIntVector GetSizeXYZ() const = 0;
+class RHI_API FRHITexture : public FRHIResource
+#if ENABLE_RHI_VALIDATION
+	, public RHIValidation::FTextureResource
+#endif
+{
+protected:
+	/** Initialization constructor. Should only be called by platform RHI implementations. */
+	FRHITexture(const FRHITextureCreateDesc& InDesc)
+		: FRHIResource(RRT_Texture)
+#if ENABLE_RHI_VALIDATION
+		, RHIValidation::FTextureResource(InDesc)
+#endif
+		, TextureDesc(InDesc)
+	{
+		SetName(InDesc.DebugName);
+	}
+
+public:
+	/**
+	 * Get the texture description used to create the texture
+	 * Still virtual because FRHITextureReference can override this function - remove virtual when FRHITextureReference is deprecated
+	 *
+	 * @return TextureDesc used to create the texture
+	 */
+	virtual const FRHITextureDesc& GetDesc() const { return TextureDesc; }
+	
+	///
+	/// Virtual functions implemented per RHI
+	/// 
+	
+	virtual class FRHITextureReference* GetTextureReference() { return NULL; }
+	virtual FRHIDescriptorHandle GetDefaultBindlessHandle() const { return FRHIDescriptorHandle(); }
 
 	/**
 	 * Returns access to the platform-specific native resource pointer.  This is designed to be used to provide plugins with access
@@ -1285,7 +1688,7 @@ public:
 		// Override this in derived classes to expose access to the native texture resource
 		return nullptr;
 	}
-	
+
 	/**
 	 * Returns access to the platform-specific RHI texture baseclass.  This is designed to provide the RHI with fast access to its base classes in the face of multiple inheritance.
 	 * @return	The pointer to the platform-specific RHI texture baseclass or NULL if it not initialized or not supported for this RHI
@@ -1296,269 +1699,194 @@ public:
 		return nullptr;
 	}
 
+	virtual void GetWriteMaskProperties(void*& OutData, uint32& OutSize)
+	{
+		OutData = nullptr;
+		OutSize = 0;
+	}
+
+	///
+	/// Helper getter functions - non virtual
+	/// 
+
+	/**
+	 * Returns the x, y & z dimensions if the texture
+	 * The Z component will always be 1 for 2D/cube resources and will contain depth for volume textures & array size for array textures
+	 */
+	FIntVector GetSizeXYZ() const
+	{
+		const FRHITextureDesc& Desc = GetDesc();
+		switch (Desc.Dimension)
+		{
+		case ETextureDimension::Texture2D:		  return FIntVector(Desc.Extent.X, Desc.Extent.Y, 1);
+		case ETextureDimension::Texture2DArray:	  return FIntVector(Desc.Extent.X, Desc.Extent.Y, Desc.ArraySize);
+		case ETextureDimension::Texture3D:		  return FIntVector(Desc.Extent.X, Desc.Extent.Y, Desc.Depth);
+		case ETextureDimension::TextureCube:	  return FIntVector(Desc.Extent.X, Desc.Extent.Y, 1);
+		case ETextureDimension::TextureCubeArray: return FIntVector(Desc.Extent.X, Desc.Extent.Y, Desc.ArraySize);
+		}
+		return FIntVector(0, 0, 0);
+	}
+
 	/**
 	 * Returns the dimensions (i.e. the actual number of texels in each dimension) of the specified mip. ArraySize is ignored.
 	 * The Z component will always be 1 for 2D/cube resources and will contain depth for volume textures.
 	 * This differs from GetSizeXYZ() which returns ArraySize in Z for 2D arrays.
 	 */
-	virtual FIntVector GetMipDimensions(uint8 MipIndex) const
+	FIntVector GetMipDimensions(uint8 MipIndex) const
 	{
-		FIntVector Size = GetSizeXYZ();
+		const FRHITextureDesc& Desc = GetDesc();
 		return FIntVector(
-			FMath::Max(Size.X >> MipIndex, 1),
-			FMath::Max(Size.Y >> MipIndex, 1),
-			FMath::Max(Size.Z >> MipIndex, 1)
+			FMath::Max<int32>(Desc.Extent.X >> MipIndex, 1),
+			FMath::Max<int32>(Desc.Extent.Y >> MipIndex, 1),
+			FMath::Max<int32>(Desc.Depth    >> MipIndex, 1)
 		);
 	}
 
-	/** @return The number of mip-maps in the texture. */
-	uint32 GetNumMips() const { return NumMips; }
-
-	/** @return The format of the pixels in the texture. */
-	EPixelFormat GetFormat() const { return Format; }
-
-	/** @return The flags used to create the texture. */
-	ETextureCreateFlags GetFlags() const { return Flags; }
-
-	/* @return the number of samples for multi-sampling. */
-	uint32 GetNumSamples() const { return NumSamples; }
-
 	/** @return Whether the texture is multi sampled. */
-	bool IsMultisampled() const { return NumSamples > 1; }		
-
-	/** sets the last time this texture was cached in a resource table. */
-	FORCEINLINE_DEBUGGABLE void SetLastRenderTime(float InLastRenderTime)
-	{
-		LastRenderTime.SetLastRenderTime(InLastRenderTime);
-	}
-
-	double GetLastRenderTime() const
-	{
-		return LastRenderTime.GetLastRenderTime();
-	}
-
-	/** Returns the last render time container, or NULL if none were specified at creation. */
-	UE_DEPRECATED(5.0, "GetLastRenderTimeContainer is deprecated and will be removed in the future")
-	FLastRenderTimeContainer* GetLastRenderTimeContainer()
-	{
-		return nullptr;
-	}
-
-	UE_DEPRECATED(5.0, "SetDefaultLastRenderTimeContainer is deprecated and will be removed in the future")
-	FORCEINLINE_DEBUGGABLE void SetDefaultLastRenderTimeContainer()
-	{
-	}
-
-	void SetName(const FName& InName)
-	{
-		TextureName = InName;
-
-#if TEXTURE_PROFILER_ENABLED
-		FTextureProfiler::Get()->UpdateTextureName(this);
-#endif
-	}
-
-	FName GetName() const
-	{
-		return TextureName;
-	}
+	bool IsMultisampled() const { return GetDesc().NumSamples > 1; }
 
+	/** @return Whether the texture has a clear color defined */
 	bool HasClearValue() const
 	{
-		return ClearValue.ColorBinding != EClearBinding::ENoneBound;
+		return GetDesc().ClearValue.ColorBinding != EClearBinding::ENoneBound;
 	}
 
+	/** @return the clear color value if set */
 	FLinearColor GetClearColor() const
 	{
-		return ClearValue.GetClearColor();
+		return GetDesc().ClearValue.GetClearColor();
 	}
 
+	/** @return the depth & stencil clear value if set */
 	void GetDepthStencilClearValue(float& OutDepth, uint32& OutStencil) const
 	{
-		return ClearValue.GetDepthStencil(OutDepth, OutStencil);
+		return GetDesc().ClearValue.GetDepthStencil(OutDepth, OutStencil);
 	}
 
+	/** @return the depth clear value if set */
 	float GetDepthClearValue() const
 	{
 		float Depth;
 		uint32 Stencil;
-		ClearValue.GetDepthStencil(Depth, Stencil);
+		GetDesc().ClearValue.GetDepthStencil(Depth, Stencil);
 		return Depth;
 	}
 
+	/** @return the stencil clear value if set */
 	uint32 GetStencilClearValue() const
 	{
 		float Depth;
 		uint32 Stencil;
-		ClearValue.GetDepthStencil(Depth, Stencil);
+		GetDesc().ClearValue.GetDepthStencil(Depth, Stencil);
 		return Stencil;
 	}
 
-	const FClearValueBinding GetClearBinding() const
+	///
+	/// RenderTime & Name functions - non virtual
+	/// 
+
+	/** sets the last time this texture was cached in a resource table. */
+	FORCEINLINE_DEBUGGABLE void SetLastRenderTime(float InLastRenderTime)
 	{
-		return ClearValue;
+		LastRenderTime.SetLastRenderTime(InLastRenderTime);
 	}
 
-	virtual void GetWriteMaskProperties(void*& OutData, uint32& OutSize)
+	double GetLastRenderTime() const
 	{
-		OutData = nullptr;
-		OutSize = 0;
+		return LastRenderTime.GetLastRenderTime();
 	}
 
-private:
-	FClearValueBinding ClearValue;
-	uint32 NumMips;
-	uint32 NumSamples;
-	EPixelFormat Format;
-	ETextureCreateFlags Flags;
-	FLastRenderTimeContainer LastRenderTime;
-	FName TextureName;
-};
-
-class RHI_API FRHITexture2D : public FRHITexture
-{
-public:
-	
-	/** Initialization constructor. */
-	FRHITexture2D(uint32 InSizeX,uint32 InSizeY,uint32 InNumMips,uint32 InNumSamples,EPixelFormat InFormat,ETextureCreateFlags InFlags, const FClearValueBinding& InClearValue, ERHIResourceType InResourceTypeOverride = RRT_None)
-	: FRHITexture(InResourceTypeOverride != RRT_None ? InResourceTypeOverride : RRT_Texture2D, InNumMips, InNumSamples, InFormat, InFlags, InClearValue)
-	, SizeX(InSizeX)
-	, SizeY(InSizeY)
-	{}
-	
-	// Dynamic cast methods.
-	virtual FRHITexture2D* GetTexture2D() { return this; }
-
-	/** @return The width of the texture. */
-	uint32 GetSizeX() const { return SizeX; }
-	
-	/** @return The height of the texture. */
-	uint32 GetSizeY() const { return SizeY; }
-
-	inline FIntPoint GetSizeXY() const
+	/** Returns the last render time container, or NULL if none were specified at creation. */
+	UE_DEPRECATED(5.0, "GetLastRenderTimeContainer is deprecated and will be removed in the future")
+	FLastRenderTimeContainer* GetLastRenderTimeContainer()
 	{
-		return FIntPoint(SizeX, SizeY);
+		return nullptr;
 	}
 
-	virtual FIntVector GetSizeXYZ() const override
+	UE_DEPRECATED(5.0, "SetDefaultLastRenderTimeContainer is deprecated and will be removed in the future")
+	FORCEINLINE_DEBUGGABLE void SetDefaultLastRenderTimeContainer()
 	{
-		return FIntVector(SizeX, SizeY, 1);
 	}
 
-private:
+	void SetName(const FName& InName)
+	{
+		TextureName = InName;
 
-	uint32 SizeX;
-	uint32 SizeY;
-};
+#if TEXTURE_PROFILER_ENABLED
+		FTextureProfiler::Get()->UpdateTextureName(this);
+#endif
+	}
 
-class RHI_API FRHITexture2DArray : public FRHITexture2D
-{
-public:
-	
-	/** Initialization constructor. */
-	FRHITexture2DArray(uint32 InSizeX,uint32 InSizeY,uint32 InSizeZ,uint32 InNumMips,uint32 NumSamples, EPixelFormat InFormat,ETextureCreateFlags InFlags, const FClearValueBinding& InClearValue)
-	: FRHITexture2D(InSizeX, InSizeY, InNumMips,NumSamples,InFormat,InFlags, InClearValue, RRT_Texture2DArray)
-	, SizeZ(InSizeZ)
+	FName GetName() const
 	{
-		check(InSizeZ != 0);
+		return TextureName;
 	}
 	
-	// Dynamic cast methods.
-	virtual FRHITexture2DArray* GetTexture2DArray() { return this; }
+	///
+	/// Deprecated functions
+	/// 
 
-	virtual FRHITexture2D* GetTexture2D() { return NULL; }
+	//UE_DEPRECATED(5.1, "FRHITexture2D is deprecated, please use FRHITexture directly")
+	inline FRHITexture2D* GetTexture2D() { return TextureDesc.Dimension == ETextureDimension::Texture2D ? this : nullptr; }
+	//UE_DEPRECATED(5.1, "FRHITexture2DArray is deprecated, please use FRHITexture directly")
+	inline FRHITexture2DArray* GetTexture2DArray() { return TextureDesc.Dimension == ETextureDimension::Texture2DArray ? this : nullptr; }
+	//UE_DEPRECATED(5.1, "FRHITexture3D is deprecated, please use FRHITexture directly")
+	inline FRHITexture3D* GetTexture3D() { return TextureDesc.Dimension == ETextureDimension::Texture3D ? this : nullptr; }
+	//UE_DEPRECATED(5.1, "FRHITextureCube is deprecated, please use FRHITexture directly")
+	inline FRHITextureCube* GetTextureCube() { return TextureDesc.Dimension == ETextureDimension::TextureCube ? this : nullptr; }
 
-	/** @return The number of textures in the array. */
-	uint32 GetSizeZ() const { return SizeZ; }
+	//UE_DEPRECATED(5.1, "GetSizeX() is deprecated, please use GetDesc().Extent.X instead")
+	uint32 GetSizeX() const { return GetDesc().Extent.X; }
 
-	virtual FIntVector GetSizeXYZ() const final override
-	{
-		return FIntVector(GetSizeX(), GetSizeY(), SizeZ);
-	}
+	//UE_DEPRECATED(5.1, "GetSizeY() is deprecated, please use GetDesc().Extent.Y instead")
+	uint32 GetSizeY() const { return GetDesc().Extent.Y; }
 
-	// Because GetSizeXYZ() returns ArraySize in Z, we need to override this function to return 1 instead.
-	// @todo: Unify the meaning of "Z" in all texture resources
-	virtual FIntVector GetMipDimensions(uint8 MipIndex) const final override
-	{
-		return FIntVector(
-			FMath::Max(GetSizeX() >> MipIndex, 1u),
-			FMath::Max(GetSizeY() >> MipIndex, 1u),
-			1
-		);
-	}
+	//UE_DEPRECATED(5.1, "GetSizeXY() is deprecated, please use GetDesc().Extent.X or GetDesc().Extent.Y instead")
+	FIntPoint GetSizeXY() const { return FIntPoint(GetDesc().Extent.X, GetDesc().Extent.Y); }
 
-private:
+	//UE_DEPRECATED(5.1, "GetSizeZ() is deprecated, please use GetDesc().ArraySize instead for TextureArrays and GetDesc().Depth for 3D textures")
+	uint32 GetSizeZ() const { return GetSizeXYZ().Z; }
 
-	uint32 SizeZ;
-};
+	//UE_DEPRECATED(5.1, "GetNumMips() is deprecated, please use GetDesc().NumMips instead")
+	uint32 GetNumMips() const { return GetDesc().NumMips; }
 
-class RHI_API FRHITexture3D : public FRHITexture
-{
-public:
-	
-	/** Initialization constructor. */
-	FRHITexture3D(uint32 InSizeX,uint32 InSizeY,uint32 InSizeZ,uint32 InNumMips,EPixelFormat InFormat,ETextureCreateFlags InFlags, const FClearValueBinding& InClearValue)
-	: FRHITexture(RRT_Texture3D, InNumMips,1,InFormat,InFlags, InClearValue)
-	, SizeX(InSizeX)
-	, SizeY(InSizeY)
-	, SizeZ(InSizeZ)
-	{}
-	
-	// Dynamic cast methods.
-	virtual FRHITexture3D* GetTexture3D() { return this; }
-	
-	/** @return The width of the texture. */
-	uint32 GetSizeX() const { return SizeX; }
-	
-	/** @return The height of the texture. */
-	uint32 GetSizeY() const { return SizeY; }
+	//UE_DEPRECATED(5.1, "GetFormat() is deprecated, please use GetDesc().Format instead")
+	EPixelFormat GetFormat() const { return GetDesc().Format; }
 
-	/** @return The depth of the texture. */
-	uint32 GetSizeZ() const { return SizeZ; }
+	//UE_DEPRECATED(5.1, "GetFlags() is deprecated, please use GetDesc().Flags instead")
+	ETextureCreateFlags GetFlags() const { return GetDesc().Flags; }
 
-	virtual FIntVector GetSizeXYZ() const final override
-	{
-		return FIntVector(SizeX, SizeY, SizeZ);
-	}
+	//UE_DEPRECATED(5.1, "GetNumSamples() is deprecated, please use GetDesc().NumSamples instead")
+	uint32 GetNumSamples() const { return GetDesc().NumSamples; }
 
-private:
+	//UE_DEPRECATED(5.1, "GetClearBinding() is deprecated, please use GetDesc().ClearValue instead")
+	const FClearValueBinding GetClearBinding() const { return GetDesc().ClearValue; }
 
-	uint32 SizeX;
-	uint32 SizeY;
-	uint32 SizeZ;
-};
+	//UE_DEPRECATED(5.1, "GetSize() is deprecated, please use GetDesc().Extent.X instead")
+	uint32 GetSize() const { check(GetDesc().IsTextureCube()); return GetDesc().Extent.X; }
 
-class RHI_API FRHITextureCube : public FRHITexture
-{
-public:
-	
-	/** Initialization constructor. */
-	FRHITextureCube(uint32 InSize,uint32 InNumMips,EPixelFormat InFormat,ETextureCreateFlags InFlags, const FClearValueBinding& InClearValue)
-	: FRHITexture(RRT_TextureCube, InNumMips,1,InFormat,InFlags, InClearValue)
-	, Size(InSize)
-	{}
-	
-	// Dynamic cast methods.
-	virtual FRHITextureCube* GetTextureCube() { return this; }
-	
-	/** @return The width and height of each face of the cubemap. */
-	uint32 GetSize() const { return Size; }
+private:
 
-	virtual FIntVector GetSizeXYZ() const final override
+	friend class FRHITextureReference;
+	/** Constructor for texture references */
+	FRHITexture(ERHIResourceType InResourceType)
+		: FRHIResource(InResourceType)
 	{
-		return FIntVector(Size, Size, 1);
+		check(InResourceType == RRT_TextureReference);
 	}
 
-private:
+	FRHITextureDesc TextureDesc;
 
-	uint32 Size;
+	FLastRenderTimeContainer LastRenderTime;
+
+	FName TextureName;
 };
 
 class RHI_API FRHITextureReference final : public FRHITexture
 {
 public:
 	explicit FRHITextureReference()
-		: FRHITexture(RRT_TextureReference, 0, 0, PF_Unknown, TexCreate_None, FClearValueBinding())
+		: FRHITexture(RRT_TextureReference)
 	{
 		check(DefaultTexture);
 		ReferencedTexture = DefaultTexture;
@@ -1574,12 +1902,6 @@ public:
 		return this;
 	}
 
-	virtual FIntVector GetSizeXYZ() const override 
-	{
-		check(ReferencedTexture);
-		return ReferencedTexture->GetSizeXYZ();
-	}
-
 	virtual void* GetNativeResource() const override 
 	{
 		check(ReferencedTexture);
@@ -1617,6 +1939,12 @@ public:
 		return ReferencedTexture.GetReference();
 	}
 
+	virtual const FRHITextureDesc& GetDesc() const override
+	{
+		check(ReferencedTexture);
+		return ReferencedTexture->GetDesc();
+	}
+
 private:
 	friend class FRHICommandListImmediate;
 	// Called only from FRHICommandListImmediate::UpdateTextureReference()
@@ -1899,10 +2227,10 @@ UE_DEPRECATED(5.0, "FIndexBufferRHIRef is deprecated, please use FBufferRHIRef."
 UE_DEPRECATED(5.0, "FVertexBufferRHIRef is deprecated, please use FBufferRHIRef.")     typedef FBufferRHIRef FVertexBufferRHIRef;
 UE_DEPRECATED(5.0, "FStructuredBufferRHIRef is deprecated, please use FBufferRHIRef.") typedef FBufferRHIRef FStructuredBufferRHIRef;
 typedef TRefCountPtr<FRHITexture> FTextureRHIRef;
-typedef TRefCountPtr<FRHITexture2D> FTexture2DRHIRef;
-typedef TRefCountPtr<FRHITexture2DArray> FTexture2DArrayRHIRef;
-typedef TRefCountPtr<FRHITexture3D> FTexture3DRHIRef;
-typedef TRefCountPtr<FRHITextureCube> FTextureCubeRHIRef;
+/*UE_DEPRECATED(5.1, "FTexture2DRHIRef is deprecated, please use FTextureRHIRef.")      */ typedef FTextureRHIRef FTexture2DRHIRef;
+/*UE_DEPRECATED(5.1, "FTexture2DArrayRHIRef is deprecated, please use FTextureRHIRef.") */ typedef FTextureRHIRef FTexture2DArrayRHIRef;
+/*UE_DEPRECATED(5.1, "FTexture3DRHIRef is deprecated, please use FTextureRHIRef.")      */ typedef FTextureRHIRef FTexture3DRHIRef;
+/*UE_DEPRECATED(5.1, "FTextureCubeRHIRef is deprecated, please use FTextureRHIRef.")    */ typedef FTextureRHIRef FTextureCubeRHIRef;
 typedef TRefCountPtr<FRHITextureReference> FTextureReferenceRHIRef;
 typedef TRefCountPtr<FRHIRenderQuery> FRenderQueryRHIRef;
 typedef TRefCountPtr<FRHIRenderQueryPool> FRenderQueryPoolRHIRef;
@@ -3646,198 +3974,6 @@ private:
 	RHI_API void OnVerifyNumUAVsFailed(int32 InNumUAVs);
 };
 
-/** Descriptor used to create a texture resource */
-struct RHI_API FRHITextureCreateInfo
-{
-	static FRHITextureCreateInfo Create2D(
-		FIntPoint InExtent,
-		EPixelFormat InFormat,
-		FClearValueBinding InClearValue,
-		ETextureCreateFlags InFlags,
-		uint8 InNumMips = 1,
-		uint8 InNumSamples = 1)
-	{
-		return FRHITextureCreateInfo(ETextureDimension::Texture2D, InFlags, InFormat, InExtent, InClearValue, 1, 1, InNumMips, InNumSamples);
-	}
-
-	static FRHITextureCreateInfo Create2DArray(
-		FIntPoint InExtent,
-		EPixelFormat InFormat,
-		FClearValueBinding InClearValue,
-		ETextureCreateFlags InFlags,
-		uint16 InArraySize,
-		uint8 InNumMips = 1,
-		uint8 InNumSamples = 1)
-	{
-		return FRHITextureCreateInfo(ETextureDimension::Texture2DArray, InFlags, InFormat, InExtent, InClearValue, 1, InArraySize, InNumMips, InNumSamples);
-	}
-
-	static FRHITextureCreateInfo Create3D(
-		FIntVector InSize,
-		EPixelFormat InFormat,
-		FClearValueBinding InClearValue,
-		ETextureCreateFlags InFlags,
-		uint8 InNumMips = 1)
-	{
-		checkf(InSize.Z >= 0 && InSize.Z <= TNumericLimits<uint16>::Max(), TEXT("Depth parameter (InSize.Z) exceeds valid range"));
-		return FRHITextureCreateInfo(ETextureDimension::Texture3D, InFlags, InFormat, FIntPoint(InSize.X, InSize.Y), InClearValue, (uint16)InSize.Z, 1, InNumMips);
-	}
-
-	static FRHITextureCreateInfo CreateCube(
-		uint32 InSizeInPixels,
-		EPixelFormat InFormat,
-		FClearValueBinding InClearValue,
-		ETextureCreateFlags InFlags,
-		uint8 InNumMips = 1,
-		uint8 InNumSamples = 1)
-	{
-		return FRHITextureCreateInfo(ETextureDimension::TextureCube, InFlags, InFormat, FIntPoint(InSizeInPixels, InSizeInPixels), InClearValue, 1, 1, InNumMips, InNumSamples);
-	}
-
-	static FRHITextureCreateInfo CreateCubeArray(
-		uint32 InSizeInPixels,
-		EPixelFormat InFormat,
-		FClearValueBinding InClearValue,
-		ETextureCreateFlags InFlags,
-		uint16 InArraySize,
-		uint8 InNumMips = 1,
-		uint8 InNumSamples = 1)
-	{
-		return FRHITextureCreateInfo(ETextureDimension::TextureCubeArray, InFlags, InFormat, FIntPoint(InSizeInPixels, InSizeInPixels), InClearValue, 1, InArraySize, InNumMips, InNumSamples);
-	}
-
-	FRHITextureCreateInfo() = default;
-	FRHITextureCreateInfo(
-		ETextureDimension InDimension,
-		ETextureCreateFlags InFlags,
-		EPixelFormat InFormat,
-		FIntPoint InExtent,
-		FClearValueBinding InClearValue,
-		uint16 InDepth = 1,
-		uint16 InArraySize = 1,
-		uint8 InNumMips = 1,
-		uint8 InNumSamples = 1)
-		: ClearValue(InClearValue)
-		, Dimension(InDimension)
-		, Flags(InFlags)
-		, Format(InFormat)
-		, Extent(InExtent)
-		, Depth(InDepth)
-		, ArraySize(InArraySize)
-		, NumMips(InNumMips)
-		, NumSamples(InNumSamples)
-	{}
-
-	bool operator == (const FRHITextureCreateInfo& Other) const
-	{
-		return Dimension == Other.Dimension
-			&& Flags == Other.Flags
-			&& Format == Other.Format
-			&& UAVFormat == Other.UAVFormat
-			&& Extent == Other.Extent
-			&& Depth == Other.Depth
-			&& ArraySize == Other.ArraySize
-			&& NumMips == Other.NumMips
-			&& NumSamples == Other.NumSamples
-			&& ClearValue == Other.ClearValue;
-	}
-
-	bool operator != (const FRHITextureCreateInfo& Other) const
-	{
-		return !(*this == Other);
-	}
-
-	bool IsTexture2D() const
-	{
-		return Dimension == ETextureDimension::Texture2D || Dimension == ETextureDimension::Texture2DArray;
-	}
-
-	bool IsTexture3D() const
-	{
-		return Dimension == ETextureDimension::Texture3D;
-	}
-
-	bool IsTextureCube() const
-	{
-		return Dimension == ETextureDimension::TextureCube || Dimension == ETextureDimension::TextureCubeArray;
-	}
-
-	bool IsTextureArray() const
-	{
-		return Dimension == ETextureDimension::Texture2DArray || Dimension == ETextureDimension::TextureCubeArray;
-	}
-
-	bool IsMipChain() const
-	{
-		return NumMips > 1;
-	}
-
-	bool IsMultisample() const
-	{
-		return NumSamples > 1;
-	}
-
-	FIntVector GetSize() const
-	{
-		return FIntVector(Extent.X, Extent.Y, Depth);
-	}
-
-	void Reset()
-	{
-		// Usually we don't want to propagate MSAA samples.
-		NumSamples = 1;
-
-		// Remove UAV flag for textures that don't need it (some formats are incompatible).
-		Flags |= TexCreate_RenderTargetable;
-		Flags &= ~(TexCreate_UAV | TexCreate_ResolveTargetable | TexCreate_DepthStencilResolveTarget);
-	}
-
-	/** Returns whether this descriptor conforms to requirements. */
-	inline bool IsValid() const
-	{
-		return FRHITextureCreateInfo::Validate(*this, /* Name = */ TEXT(""), /* bFatal = */ false);
-	}
-
-	/** Clear value to use when fast-clearing the texture. */
-	FClearValueBinding ClearValue;
-
-	/** Texture dimension to use when creating the RHI texture. */
-	ETextureDimension Dimension = ETextureDimension::Texture2D;
-
-	/** Texture flags passed on to RHI texture. */
-	ETextureCreateFlags Flags = TexCreate_None;
-
-	/** Pixel format used to create RHI texture. */
-	EPixelFormat Format = PF_Unknown;
-
-	/** Texture format used when creating the UAV. PF_Unknown means to use the default one (same as Format). */
-	EPixelFormat UAVFormat = PF_Unknown;
-
-	/** Extent of the texture in x and y. */
-	FIntPoint Extent = FIntPoint(1, 1);
-
-	/** Depth of the texture if the dimension is 3D. */
-	uint16 Depth = 1;
-
-	/** The number of array elements in the texture. (Keep at 1 if dimension is 3D). */
-	uint16 ArraySize = 1;
-
-	/** Number of mips in the texture mip-map chain. */
-	uint8 NumMips = 1;
-
-	/** Number of samples in the texture. >1 for MSAA. */
-	uint8 NumSamples = 1;
-
-	/** Check the validity. */
-	static bool CheckValidity(const FRHITextureCreateInfo& Desc, const TCHAR* Name)
-	{
-		return FRHITextureCreateInfo::Validate(Desc, Name, /* bFatal = */ true);
-	}
-
-private:
-	static bool Validate(const FRHITextureCreateInfo& Desc, const TCHAR* Name, bool bFatal);
-};
-
 /** Used to specify a texture metadata plane when creating a view. */
 enum class ERHITextureMetaDataAccess : uint8
 {
diff --git a/Engine/Source/Runtime/RHI/Public/RHIValidation.h b/Engine/Source/Runtime/RHI/Public/RHIValidation.h
--- a/Engine/Source/Runtime/RHI/Public/RHIValidation.h
+++ b/Engine/Source/Runtime/RHI/Public/RHIValidation.h
@@ -486,71 +486,12 @@ public:
 		SRV->ViewIdentity = Buffer->GetWholeResourceIdentity();
 	}
 
-	/**
-	* Computes the total size of a 2D texture with the specified parameters.
-	*
-	* @param SizeX - width of the texture to compute
-	* @param SizeY - height of the texture to compute
-	* @param Format - EPixelFormat texture format
-	* @param NumMips - number of mips to compute or 0 for full mip pyramid
-	* @param NumSamples - number of MSAA samples, usually 1
-	* @param Flags - ETextureCreateFlags creation flags
-	* @param OutAlign - Alignment required for this texture.  Output parameter.
-	*/
-	// FlushType: Thread safe
-	virtual uint64 RHICalcTexture2DPlatformSize(uint32 SizeX, uint32 SizeY, uint8 Format, uint32 NumMips, uint32 NumSamples, ETextureCreateFlags Flags, const FRHIResourceCreateInfo& CreateInfo, uint32& OutAlign) override final
-	{
-		return RHI->RHICalcTexture2DPlatformSize(SizeX, SizeY, Format, NumMips, NumSamples, Flags, CreateInfo, OutAlign);
-	}
-
-	/**
-	* Computes the total size of a virtual memory (VM) based 2D texture with the specified parameters.
-	*
-	* @param Mip0Width - width of the top mip
-	* @param Mip0Height - height of the top mip
-	* @param Format - EPixelFormat texture format
-	* @param NumMips - number of mips the texture has
-	* @param FirstMipIdx - index of the first resident mip
-	* @param NumSamples - number of MSAA samples, usually 1
-	* @param Flags - ETextureCreateFlags creation flags
-	* @param OutAlign - Alignment required for this texture. Output parameter.
-	*/
-	// FlushType: Thread safe
-	virtual uint64 RHICalcVMTexture2DPlatformSize(uint32 Mip0Width, uint32 Mip0Height, uint8 Format, uint32 NumMips, uint32 FirstMipIdx, uint32 NumSamples, ETextureCreateFlags Flags, uint32& OutAlign) override final
-	{
-		return RHI->RHICalcVMTexture2DPlatformSize(Mip0Width, Mip0Height, Format, NumMips, FirstMipIdx, NumSamples, Flags, OutAlign);
-	}
-
-
-	/**
-	* Computes the total size of a 3D texture with the specified parameters.
-	* @param SizeX - width of the texture to create
-	* @param SizeY - height of the texture to create
-	* @param SizeZ - depth of the texture to create
-	* @param Format - EPixelFormat texture format
-	* @param NumMips - number of mips to generate or 0 for full mip pyramid
-	* @param Flags - ETextureCreateFlags creation flags
-	* @param OutAlign - Alignment required for this texture.  Output parameter.
-	*/
-	// FlushType: Thread safe
-	virtual uint64 RHICalcTexture3DPlatformSize(uint32 SizeX, uint32 SizeY, uint32 SizeZ, uint8 Format, uint32 NumMips, ETextureCreateFlags Flags, const FRHIResourceCreateInfo& CreateInfo, uint32& OutAlign) override final
+	virtual FRHICalcTextureSizeResult RHICalcTexturePlatformSize(FRHITextureDesc const& Desc, uint32 FirstMipIndex) override final
 	{
-		return RHI->RHICalcTexture3DPlatformSize(SizeX, SizeY, SizeZ, Format, NumMips, Flags, CreateInfo, OutAlign);
-	}
+		ensure(Desc.IsValid());
+		ensure(FirstMipIndex < Desc.NumMips);
 
-	/**
-	* Computes the total size of a cube texture with the specified parameters.
-	* @param Size - width/height of the texture to create
-	* @param Format - EPixelFormat texture format
-	* @param NumMips - number of mips to generate or 0 for full mip pyramid
-	* @param Flags - ETextureCreateFlags creation flags
-	* @param OutAlign - Alignment required for this texture.  Output parameter.
-	*/
-	// FlushType: Thread safe
-	virtual uint64 RHICalcTextureCubePlatformSize(uint32 Size, uint8 Format, uint32 NumMips, ETextureCreateFlags Flags, const FRHIResourceCreateInfo& CreateInfo, uint32& OutAlign) override final
-	{
-		ensure(Size >= (1u << (FMath::Max(1u, NumMips) - 1)));
-		return RHI->RHICalcTextureCubePlatformSize(Size, Format, NumMips, Flags, CreateInfo, OutAlign);
+		return RHI->RHICalcTexturePlatformSize(Desc, FirstMipIndex);
 	}
 
 	/**
@@ -581,38 +522,13 @@ public:
 	}
 
 	/**
-	* Creates a 2D RHI texture resource
-	* @param SizeX - width of the texture to create
-	* @param SizeY - height of the texture to create
-	* @param Format - EPixelFormat texture format
-	* @param NumMips - number of mips to generate or 0 for full mip pyramid
-	* @param NumSamples - number of MSAA samples, usually 1
-	* @param Flags - ETextureCreateFlags creation flags
-	*/
-	// FlushType: Wait RHI Thread
-	virtual FTexture2DRHIRef RHICreateTexture2D(uint32 SizeX, uint32 SizeY, uint8 Format, uint32 NumMips, uint32 NumSamples, ETextureCreateFlags Flags, ERHIAccess InResourceState, FRHIResourceCreateInfo& CreateInfo) override final
-	{
-		ensure(FMath::Max(SizeX, SizeY) >= (1u << (FMath::Max(1u, NumMips) - 1)));
-		FTexture2DRHIRef Texture = RHI->RHICreateTexture2D(SizeX, SizeY, Format, NumMips, NumSamples, Flags, InResourceState, CreateInfo);
-		Texture->InitBarrierTracking(NumMips, 1, EPixelFormat(Format), Flags, InResourceState, CreateInfo.DebugName);
-		return Texture;
-	}
-
-	/**
-	* Creates a 2D RHI texture external resource
-	* @param SizeX - width of the texture to create
-	* @param SizeY - height of the texture to create
-	* @param Format - EPixelFormat texture format
-	* @param NumMips - number of mips to generate or 0 for full mip pyramid
-	* @param NumSamples - number of MSAA samples, usually 1
-	* @param Flags - ETextureCreateFlags creation flags
+	* Creates an RHI texture resource.
 	*/
-	// FlushType: Wait RHI Thread
-	virtual FTexture2DRHIRef RHICreateTextureExternal2D(uint32 SizeX, uint32 SizeY, uint8 Format, uint32 NumMips, uint32 NumSamples, ETextureCreateFlags Flags, ERHIAccess InResourceState, FRHIResourceCreateInfo& CreateInfo) override final
+	virtual FTextureRHIRef RHICreateTexture(const FRHITextureCreateDesc& CreateDesc) override final
 	{
-		ensure(FMath::Max(SizeX, SizeY) >= (1u << (FMath::Max(1u, NumMips) - 1)));
-		FTexture2DRHIRef Texture = RHI->RHICreateTextureExternal2D(SizeX, SizeY, Format, NumMips, NumSamples, Flags, InResourceState, CreateInfo);
-		Texture->InitBarrierTracking(NumMips, 1, EPixelFormat(Format), Flags, InResourceState, CreateInfo.DebugName);
+		CreateDesc.CheckValidity();
+		FTextureRHIRef Texture = RHI->RHICreateTexture(CreateDesc);
+		ensure(Texture->IsBarrierTrackingInitialized());
 		return Texture;
 	}
 
@@ -630,12 +546,12 @@ public:
 	* @returns a reference to a 2D texture resource
 	*/
 	// FlushType: Thread safe
-	virtual FTexture2DRHIRef RHIAsyncCreateTexture2D(uint32 SizeX, uint32 SizeY, uint8 Format, uint32 NumMips, ETextureCreateFlags Flags, ERHIAccess InResourceState, void** InitialMipData, uint32 NumInitialMips) override final
+	virtual FTextureRHIRef RHIAsyncCreateTexture2D(uint32 SizeX, uint32 SizeY, uint8 Format, uint32 NumMips, ETextureCreateFlags Flags, ERHIAccess InResourceState, void** InitialMipData, uint32 NumInitialMips) override final
 	{
 		check(GRHISupportsAsyncTextureCreation);
 		ensure(FMath::Max(SizeX, SizeY) >= (1u << (FMath::Max(1u, NumMips) - 1)));
-		FTexture2DRHIRef Texture = RHI->RHIAsyncCreateTexture2D(SizeX, SizeY, Format, NumMips, Flags, InResourceState, InitialMipData, NumInitialMips);
-		Texture->InitBarrierTracking(NumMips, 1, EPixelFormat(Format), Flags, InResourceState, TEXT("UnknownTexture_RHIAsyncCreateTexture2D"));
+		FTextureRHIRef Texture = RHI->RHIAsyncCreateTexture2D(SizeX, SizeY, Format, NumMips, Flags, InResourceState, InitialMipData, NumInitialMips);
+		ensure(Texture->IsBarrierTrackingInitialized());
 		return Texture;
 	}
 
@@ -657,42 +573,6 @@ public:
 		RHI->RHITransferBufferUnderlyingResource(DestBuffer, SrcBuffer);
 	}
 
-	/**
-	* Creates a Array RHI texture resource
-	* @param SizeX - width of the texture to create
-	* @param SizeY - height of the texture to create
-	* @param SizeZ - depth of the texture to create
-	* @param Format - EPixelFormat texture format
-	* @param NumMips - number of mips to generate or 0 for full mip pyramid
-	* @param Flags - ETextureCreateFlags creation flags
-	*/
-	// FlushType: Wait RHI Thread
-	virtual FTexture2DArrayRHIRef RHICreateTexture2DArray(uint32 SizeX, uint32 SizeY, uint32 SizeZ, uint8 Format, uint32 NumMips, uint32 NumSamples, ETextureCreateFlags Flags, ERHIAccess InResourceState, FRHIResourceCreateInfo& CreateInfo)
-	{
-		ensure(FMath::Max(SizeX, SizeY) >= (1u << (FMath::Max(1u, NumMips) - 1)));
-		FTexture2DArrayRHIRef TextureArray = RHI->RHICreateTexture2DArray(SizeX, SizeY, SizeZ, Format, NumMips, NumSamples, Flags, InResourceState, CreateInfo);
-		TextureArray->InitBarrierTracking(NumMips, SizeZ, EPixelFormat(Format), Flags, InResourceState, CreateInfo.DebugName);
-		return TextureArray;
-	}
-
-	/**
-	* Creates a 3d RHI texture resource
-	* @param SizeX - width of the texture to create
-	* @param SizeY - height of the texture to create
-	* @param SizeZ - depth of the texture to create
-	* @param Format - EPixelFormat texture format
-	* @param NumMips - number of mips to generate or 0 for full mip pyramid
-	* @param Flags - ETextureCreateFlags creation flags
-	*/
-	// FlushType: Wait RHI Thread
-	virtual FTexture3DRHIRef RHICreateTexture3D(uint32 SizeX, uint32 SizeY, uint32 SizeZ, uint8 Format, uint32 NumMips, ETextureCreateFlags Flags, ERHIAccess InResourceState, FRHIResourceCreateInfo& CreateInfo) override final
-	{
-		ensure(FMath::Max(SizeX, FMath::Max(SizeY, SizeZ)) >= (1u << (FMath::Max(1u, NumMips) - 1)));
-		FTexture3DRHIRef Texture = RHI->RHICreateTexture3D(SizeX, SizeY, SizeZ, Format, NumMips, Flags, InResourceState, CreateInfo);
-		Texture->InitBarrierTracking(NumMips, 1, EPixelFormat(Format), Flags, InResourceState, CreateInfo.DebugName);
-		return Texture;
-	}
-
 	/**
 	* Creates a shader resource view for a texture
 	*/
@@ -906,39 +786,6 @@ public:
 		RHI->RHIUpdateTexture3D(Texture, MipIndex, UpdateRegion, SourceRowPitch, SourceDepthPitch, SourceData);
 	}
 
-	/**
-	* Creates a Cube RHI texture resource
-	* @param Size - width/height of the texture to create
-	* @param Format - EPixelFormat texture format
-	* @param NumMips - number of mips to generate or 0 for full mip pyramid
-	* @param Flags - ETextureCreateFlags creation flags
-	*/
-	// FlushType: Wait RHI Thread
-	virtual FTextureCubeRHIRef RHICreateTextureCube(uint32 Size, uint8 Format, uint32 NumMips, ETextureCreateFlags Flags, ERHIAccess InResourceState, FRHIResourceCreateInfo& CreateInfo) override final
-	{
-		ensure(Size >= (1u << (FMath::Max(1u, NumMips) - 1)));
-		FTextureCubeRHIRef TextureCube = RHI->RHICreateTextureCube(Size, Format, NumMips, Flags, InResourceState, CreateInfo);
-		TextureCube->InitBarrierTracking(NumMips, 6, EPixelFormat(Format), Flags, InResourceState, CreateInfo.DebugName);
-		return TextureCube;
-	}
-
-	/**
-	* Creates a Cube Array RHI texture resource
-	* @param Size - width/height of the texture to create
-	* @param ArraySize - number of array elements of the texture to create
-	* @param Format - EPixelFormat texture format
-	* @param NumMips - number of mips to generate or 0 for full mip pyramid
-	* @param Flags - ETextureCreateFlags creation flags
-	*/
-	// FlushType: Wait RHI Thread
-	virtual FTextureCubeRHIRef RHICreateTextureCubeArray(uint32 Size, uint32 ArraySize, uint8 Format, uint32 NumMips, ETextureCreateFlags Flags, ERHIAccess InResourceState, FRHIResourceCreateInfo& CreateInfo) override final
-	{
-		ensure(Size >= (1u << (FMath::Max(1u, NumMips) - 1)));
-		FTextureCubeRHIRef TextureCubeArray = RHI->RHICreateTextureCubeArray(Size, ArraySize, Format, NumMips, Flags, InResourceState, CreateInfo);
-		TextureCubeArray->InitBarrierTracking(NumMips, ArraySize * 6, EPixelFormat(Format), Flags, InResourceState, CreateInfo.DebugName);
-		return TextureCubeArray;
-	}
-
 	/**
 	* Locks an RHI texture's mip-map for read/write operations on the CPU
 	* @param Texture - the RHI texture resource to lock
@@ -1518,28 +1365,12 @@ public:
 		return RHI->RHICreateShaderLibrary_RenderThread(RHICmdList, Platform, FilePath, Name);
 	}
 
-	virtual FTexture2DRHIRef RHICreateTexture2D_RenderThread(class FRHICommandListImmediate& RHICmdList, uint32 SizeX, uint32 SizeY, uint8 Format, uint32 NumMips, uint32 NumSamples, ETextureCreateFlags Flags, ERHIAccess InResourceState, FRHIResourceCreateInfo& CreateInfo) override final
-	{
-		ensure(FMath::Max(SizeX, SizeY) >= (1u << (FMath::Max(1u, NumMips) - 1)));
-		FTexture2DRHIRef Texture2D = RHI->RHICreateTexture2D_RenderThread(RHICmdList, SizeX, SizeY, Format, NumMips, NumSamples, Flags, InResourceState, CreateInfo);
-		Texture2D->InitBarrierTracking(NumMips, 1, EPixelFormat(Format), Flags, InResourceState, CreateInfo.DebugName);
-		return Texture2D;
-	}
-
-	virtual FTexture2DRHIRef RHICreateTextureExternal2D_RenderThread(class FRHICommandListImmediate& RHICmdList, uint32 SizeX, uint32 SizeY, uint8 Format, uint32 NumMips, uint32 NumSamples, ETextureCreateFlags Flags, ERHIAccess InResourceState, FRHIResourceCreateInfo& CreateInfo) override final
-	{
-		ensure(FMath::Max(SizeX, SizeY) >= (1u << (FMath::Max(1u, NumMips) - 1)));
-		FTexture2DRHIRef Texture2D = RHI->RHICreateTextureExternal2D_RenderThread(RHICmdList, SizeX, SizeY, Format, NumMips, NumSamples, Flags, InResourceState, CreateInfo);
-		Texture2D->InitBarrierTracking(NumMips, 1, EPixelFormat(Format), Flags, InResourceState, CreateInfo.DebugName);
-		return Texture2D;
-	}
-
-	virtual FTexture3DRHIRef RHICreateTexture3D_RenderThread(class FRHICommandListImmediate& RHICmdList, uint32 SizeX, uint32 SizeY, uint32 SizeZ, uint8 Format, uint32 NumMips, ETextureCreateFlags Flags, ERHIAccess InResourceState, FRHIResourceCreateInfo& CreateInfo) override final
+	virtual FTextureRHIRef RHICreateTexture_RenderThread(class FRHICommandListImmediate& RHICmdList, const FRHITextureCreateDesc& CreateDesc) override final
 	{
-		ensure(FMath::Max(SizeX, FMath::Max(SizeY, SizeZ)) >= (1u << (FMath::Max(1u, NumMips) - 1)));
-		FTexture3DRHIRef Texture3D = RHI->RHICreateTexture3D_RenderThread(RHICmdList, SizeX, SizeY, SizeZ, Format, NumMips, Flags, InResourceState, CreateInfo);
-		Texture3D->InitBarrierTracking(NumMips, 1, EPixelFormat(Format), Flags, InResourceState, CreateInfo.DebugName);
-		return Texture3D;
+		CreateDesc.CheckValidity();
+		FTextureRHIRef Texture = RHI->RHICreateTexture_RenderThread(RHICmdList, CreateDesc);
+		ensure(Texture->IsBarrierTrackingInitialized());
+		return Texture;
 	}
 
 	virtual FUnorderedAccessViewRHIRef RHICreateUnorderedAccessView_RenderThread(class FRHICommandListImmediate& RHICmdList, FRHIBuffer* Buffer, bool bUseUAVCounter, bool bAppendBuffer) override final
@@ -1611,22 +1442,6 @@ public:
 		}
 		return SRV;
 	}
-
-	virtual FTextureCubeRHIRef RHICreateTextureCube_RenderThread(class FRHICommandListImmediate& RHICmdList, uint32 Size, uint8 Format, uint32 NumMips, ETextureCreateFlags Flags, ERHIAccess InResourceState, FRHIResourceCreateInfo& CreateInfo) override final
-	{
-		ensure(Size >= (1u << (FMath::Max(1u, NumMips) - 1)));
-		FTextureCubeRHIRef TextureCube = RHI->RHICreateTextureCube_RenderThread(RHICmdList, Size, Format, NumMips, Flags, InResourceState, CreateInfo);
-		TextureCube->InitBarrierTracking(NumMips, 6, EPixelFormat(Format), Flags, InResourceState, CreateInfo.DebugName);
-		return TextureCube;
-	}
-
-	virtual FTextureCubeRHIRef RHICreateTextureCubeArray_RenderThread(class FRHICommandListImmediate& RHICmdList, uint32 Size, uint32 ArraySize, uint8 Format, uint32 NumMips, ETextureCreateFlags Flags, ERHIAccess InResourceState, FRHIResourceCreateInfo& CreateInfo) override final
-	{
-		ensure(Size >= (1u << (FMath::Max(1u, NumMips) - 1)));
-		FTextureCubeRHIRef TextureCubeArray = RHI->RHICreateTextureCubeArray_RenderThread(RHICmdList, Size, ArraySize, Format, NumMips, Flags, InResourceState, CreateInfo);
-		TextureCubeArray->InitBarrierTracking(NumMips, ArraySize * 6, EPixelFormat(Format), Flags, InResourceState, CreateInfo.DebugName);
-		return TextureCubeArray;
-	}
 	
 	virtual FRenderQueryRHIRef RHICreateRenderQuery_RenderThread(class FRHICommandListImmediate& RHICmdList, ERenderQueryType QueryType) override final
 	{
diff --git a/Engine/Source/Runtime/RHI/Public/RHIValidationCommon.h b/Engine/Source/Runtime/RHI/Public/RHIValidationCommon.h
--- a/Engine/Source/Runtime/RHI/Public/RHIValidationCommon.h
+++ b/Engine/Source/Runtime/RHI/Public/RHIValidationCommon.h
@@ -16,6 +16,7 @@ const bool GRHIValidationEnabled = false;
 
 class FRHIUniformBuffer;
 class FValidationRHI;
+struct FRHITextureCreateDesc;
 
 // Forward declaration of function defined in RHIUtilities.h
 static inline bool IsStencilFormat(EPixelFormat Format);
@@ -383,10 +384,22 @@ namespace RHIValidation
 		FResource PRIVATE_TrackerResource;
 
 	public:
+		FTextureResource() = default;
+		RHI_API FTextureResource(FRHITextureCreateDesc const& CreateDesc);
+
 		virtual ~FTextureResource() {}
 
 		virtual FResource* GetTrackerResource() { return &PRIVATE_TrackerResource; }
 
+		void InitBarrierTracking(FRHITextureCreateDesc const& CreateDesc);
+
+		inline bool IsBarrierTrackingInitialized() const
+		{
+			// @todo: clean up const_cast once FRHITextureReference is removed and
+			// we don't need to keep a separate PRIVATE_TrackerResource object.
+			return const_cast<FTextureResource*>(this)->GetTrackerResource()->IsBarrierTrackingInitialized();
+		}
+
 		inline void InitBarrierTracking(int32 InNumMips, int32 InNumArraySlices, EPixelFormat PixelFormat, ETextureCreateFlags Flags, ERHIAccess InResourceState, const TCHAR* InDebugName)
 		{
 			FResource* Resource = GetTrackerResource();
@@ -454,7 +467,7 @@ namespace RHIValidation
 			}
 			else
 			{
-				checkSlow(Info.MipIndex < uint32(Resource->NumMips));
+				check(Info.MipIndex < uint32(Resource->NumMips));
 				Identity.SubresourceRange.MipIndex = Info.MipIndex;
 				Identity.SubresourceRange.NumMips = 1;
 			}
@@ -466,7 +479,7 @@ namespace RHIValidation
 			}
 			else
 			{
-				checkSlow(Info.ArraySlice < uint32(Resource->NumArraySlices));
+				check(Info.ArraySlice < uint32(Resource->NumArraySlices));
 				Identity.SubresourceRange.ArraySlice = Info.ArraySlice;
 				Identity.SubresourceRange.NumArraySlices = 1;
 			}
@@ -478,7 +491,7 @@ namespace RHIValidation
 			}
 			else
 			{
-				checkSlow(Info.PlaneSlice < uint32(Resource->NumPlanes));
+				check(Info.PlaneSlice < uint32(Resource->NumPlanes));
 				Identity.SubresourceRange.PlaneIndex = Info.PlaneSlice;
 				Identity.SubresourceRange.NumPlanes = 1;
 			}
diff --git a/Engine/Source/Runtime/RenderCore/Private/RenderTargetPool.cpp b/Engine/Source/Runtime/RenderCore/Private/RenderTargetPool.cpp
--- a/Engine/Source/Runtime/RenderCore/Private/RenderTargetPool.cpp
+++ b/Engine/Source/Runtime/RenderCore/Private/RenderTargetPool.cpp
@@ -17,35 +17,13 @@ TRefCountPtr<IPooledRenderTarget> CreateRenderTarget(FRHITexture* Texture, const
 {
 	check(Texture);
 
-	const FIntVector Size = Texture->GetSizeXYZ();
-
-	FPooledRenderTargetDesc Desc;
-	Desc.Extent = FIntPoint(Size.X, Size.Y);
-	Desc.ClearValue = Texture->GetClearBinding();
-	Desc.Format = Texture->GetFormat();
-	Desc.NumMips = Texture->GetNumMips();
-	Desc.NumSamples = Texture->GetNumSamples();
-	Desc.Flags = Texture->GetFlags();
-	Desc.DebugName = Name;
-
-	if (FRHITextureCube* TextureCube = Texture->GetTextureCube())
-	{
-		Desc.bIsCubemap = true;
-	}
-	else if (FRHITexture3D* Texture3D = Texture->GetTexture3D())
-	{
-		Desc.Depth = Size.Z;
-	}
-	else if (FRHITexture2DArray* TextureArray = Texture->GetTexture2DArray())
-	{
-		Desc.bIsArray = true;
-		Desc.ArraySize = Size.Z;
-	}
-
 	FSceneRenderTargetItem Item;
 	Item.TargetableTexture = Texture;
 	Item.ShaderResourceTexture = Texture;
 
+	FPooledRenderTargetDesc Desc = Translate(Texture->GetDesc());
+	Desc.DebugName = Name;
+
 	TRefCountPtr<IPooledRenderTarget> PooledRenderTarget;
 	GRenderTargetPool.CreateUntrackedElement(Desc, PooledRenderTarget, Item);
 	return MoveTemp(PooledRenderTarget);
@@ -61,30 +39,6 @@ bool CacheRenderTarget(FRHITexture* Texture, const TCHAR* Name, TRefCountPtr<IPo
 	return false;
 }
 
-static uint64 GetTypeHash(FClearValueBinding Binding)
-{
-	uint64 Hash = 0;
-	switch (Binding.ColorBinding)
-	{
-	case EClearBinding::EColorBound:
-		Hash = CityHash64((const char*)Binding.Value.Color, sizeof(Binding.Value.Color));
-		break;
-	case EClearBinding::EDepthStencilBound:
-		Hash = uint64(GetTypeHash(Binding.Value.DSValue.Depth)) << 32 | uint64(Binding.Value.DSValue.Stencil);
-		break;
-	}
-	return Hash ^ uint64(Binding.ColorBinding);
-}
-
-inline uint64 ComputeHash(const FRHITextureCreateInfo& InCreateInfo)
-{
-	// Make sure all padding is removed.
-	FRHITextureCreateInfo NewInfo;
-	FPlatformMemory::Memzero(&NewInfo, sizeof(FRHITextureCreateInfo));
-	NewInfo = InCreateInfo;
-	return CityHash64((const char*)&NewInfo, sizeof(FRHITextureCreateInfo));
-}
-
 RENDERCORE_API void DumpRenderTargetPoolMemory(FOutputDevice& OutputDevice)
 {
 	GRenderTargetPool.DumpMemoryUsage(OutputDevice);
@@ -109,7 +63,10 @@ TRefCountPtr<IPooledRenderTarget> FRenderTargetPool::FindFreeElementInternal(FRH
 	// FastVRAM is no longer supported by the render target pool.
 	EnumRemoveFlags(Desc.Flags, ETextureCreateFlags::FastVRAM | ETextureCreateFlags::FastVRAMPartialAlloc);
 
-	const uint64 DescHash = ComputeHash(Desc);
+	// We always want SRV access
+	Desc.Flags |= TexCreate_ShaderResource;
+
+	const uint32 DescHash = GetTypeHash(Desc);
 
 	for (uint32 Index = 0, Num = (uint32)PooledRenderTargets.Num(); Index < Num; ++Index)
 	{
@@ -139,84 +96,14 @@ TRefCountPtr<IPooledRenderTarget> FRenderTargetPool::FindFreeElementInternal(FRH
 	{
 		TRACE_CPUPROFILER_EVENT_SCOPE(FRenderTargetPool::CreateTexture);
 
-		FRHIResourceCreateInfo CreateInfo(Name, Desc.ClearValue);
-
-		FTextureRHIRef ResultTexture;
 		const ERHIAccess AccessInitial = ERHIAccess::SRVMask;
-		const ETextureCreateFlags TextureFlags = Desc.Flags | TexCreate_ShaderResource;
-
-		// Only create resources if we're not asked to defer creation.
-		if (Desc.IsTexture2D())
-		{
-			if (!Desc.IsTextureArray())
-			{
-				ResultTexture = RHICreateTexture2D(
-					Desc.Extent.X,
-					Desc.Extent.Y,
-					(uint8)Desc.Format,
-					Desc.NumMips,
-					Desc.NumSamples,
-					TextureFlags,
-					AccessInitial,
-					CreateInfo
-				);
-			}
-			else
-			{
-				ResultTexture = RHICreateTexture2DArray(
-					Desc.Extent.X,
-					Desc.Extent.Y,
-					Desc.ArraySize,
-					(uint8)Desc.Format,
-					Desc.NumMips,
-					Desc.NumSamples,
-					TextureFlags,
-					AccessInitial,
-					CreateInfo
-				);
-			}
-		}
-		else if (Desc.IsTexture3D())
-		{
-			ResultTexture = RHICreateTexture3D(
-				Desc.Extent.X,
-				Desc.Extent.Y,
-				Desc.Depth,
-				(uint8)Desc.Format,
-				Desc.NumMips,
-				TextureFlags,
-				AccessInitial,
-				CreateInfo);
-		}
-		else
-		{
-			check(Desc.IsTextureCube());
-			if (Desc.IsTextureArray())
-			{
-				ResultTexture = RHICreateTextureCubeArray(
-					Desc.Extent.X,
-					Desc.ArraySize,
-					(uint8)Desc.Format,
-					Desc.NumMips,
-					TextureFlags,
-					AccessInitial,
-					CreateInfo
-				);
-			}
-			else
-			{
-				ResultTexture = RHICreateTextureCube(
-					Desc.Extent.X,
-					(uint8)Desc.Format,
-					Desc.NumMips,
-					TextureFlags,
-					AccessInitial,
-					CreateInfo
-				);
-			}
-		}
+		FRHITextureCreateDesc CreateDesc(Desc, AccessInitial, Name);
 
-		Found = new FPooledRenderTarget(ResultTexture, AccessInitial, Translate(Desc), this);
+		Found = new FPooledRenderTarget(
+			RHICreateTexture(CreateDesc),
+			AccessInitial,
+			Translate(CreateDesc),
+			this);
 
 		PooledRenderTargets.Add(Found);
 		PooledRenderTargetHashes.Add(DescHash);
diff --git a/Engine/Source/Runtime/RenderCore/Public/RenderGraphResources.inl b/Engine/Source/Runtime/RenderCore/Public/RenderGraphResources.inl
--- a/Engine/Source/Runtime/RenderCore/Public/RenderGraphResources.inl
+++ b/Engine/Source/Runtime/RenderCore/Public/RenderGraphResources.inl
@@ -97,7 +97,7 @@ inline FRHIBufferCreateInfo Translate(const FRDGBufferDesc& InDesc)
 	return CreateInfo;
 }
 
-FRDGTextureDesc Translate(const FPooledRenderTargetDesc& InDesc)
+inline FRDGTextureDesc Translate(const FPooledRenderTargetDesc& InDesc)
 {
 	check(InDesc.IsValid());
 
diff --git a/Engine/Source/Runtime/RenderCore/Public/RenderTargetPool.h b/Engine/Source/Runtime/RenderCore/Public/RenderTargetPool.h
--- a/Engine/Source/Runtime/RenderCore/Public/RenderTargetPool.h
+++ b/Engine/Source/Runtime/RenderCore/Public/RenderTargetPool.h
@@ -152,7 +152,7 @@ private:
 	void FreeElementAtIndex(int32 Index);
 
 	/** Elements can be 0, we compact the buffer later. */
-	TArray<uint64> PooledRenderTargetHashes;
+	TArray<uint32> PooledRenderTargetHashes;
 	TArray< TRefCountPtr<FPooledRenderTarget> > PooledRenderTargets;
 	TArray< TRefCountPtr<FPooledRenderTarget> > DeferredDeleteArray;
 
diff --git a/Engine/Source/Runtime/Renderer/Private/PostProcess/PostProcessMaterial.cpp b/Engine/Source/Runtime/Renderer/Private/PostProcess/PostProcessMaterial.cpp
--- a/Engine/Source/Runtime/Renderer/Private/PostProcess/PostProcessMaterial.cpp
+++ b/Engine/Source/Runtime/Renderer/Private/PostProcess/PostProcessMaterial.cpp
@@ -474,8 +474,7 @@ FScreenPassTexture AddPostProcessMaterialPass(
 
 	// Inputs.OverrideOutput is used to force drawing directly to the backbuffer. OpenGL doesn't support using the backbuffer color target with a custom depth/stencil
 	// buffer, so in that case we must draw to an intermediate target and copy to the backbuffer at the end. Ideally, we would test if Inputs.OverrideOutput.Texture
-	// is actually the backbuffer (as returned by AndroidEGL::GetOnScreenColorRenderBuffer() and such), but it's not worth doing all the plumbing and increasing the
-	// RHI surface area just for this hack.
+	// is actually the backbuffer, but it's not worth doing all the plumbing and increasing the RHI surface area just for this hack.
 	const bool bBackbufferWithDepthStencil = (DepthStencilTexture != nullptr && !GRHISupportsBackBufferWithCustomDepthStencil && Inputs.OverrideOutput.IsValid());
 
 	// The other case when we must render to an intermediate target is when we have to flip the image vertically because we're the last postprocess pass on mobile OpenGL.
