Adding Min/Max FeatureLevel configs for Windows D3D11, D3D12 and Vulkan.
This unifies a few Performance Mode settings together and now allows a project to restrict D3D12 to SM6 only.

diff --git a/Engine/Config/BaseEngine.ini b/Engine/Config/BaseEngine.ini
--- a/Engine/Config/BaseEngine.ini
+++ b/Engine/Config/BaseEngine.ini
@@ -2716,6 +2716,12 @@ bEnable2DPhysics=false
 
 [/Script/WindowsTargetPlatform.WindowsTargetSettings]
 +TargetedRHIs=PCD3D_SM5
+D3D12MinimumFeatureLevel=ES3_1
+D3D12MaximumFeatureLevel=SM6
+D3D11MinimumFeatureLevel=ES3_1
+D3D11MaximumFeatureLevel=SM5
+VulkanMinimumFeatureLevel=ES3_1
+VulkanMaximumFeatureLevel=SM5
 bEnableDistanceFields=true
 bEnableRayTracing=true
 bTarget32Bit=false
diff --git a/Engine/Source/Developer/Windows/WindowsTargetPlatform/Classes/WindowsTargetSettings.h b/Engine/Source/Developer/Windows/WindowsTargetPlatform/Classes/WindowsTargetSettings.h
--- a/Engine/Source/Developer/Windows/WindowsTargetPlatform/Classes/WindowsTargetSettings.h
+++ b/Engine/Source/Developer/Windows/WindowsTargetPlatform/Classes/WindowsTargetSettings.h
@@ -20,6 +20,7 @@ enum class ECompilerVersion : uint8
 	VisualStudio2019 = 3 UMETA(DisplayName = "Visual Studio 2019"),
 	VisualStudio2022 = 4 UMETA(DisplayName = "Visual Studio 2022"),
 };
+
 UENUM()
 enum class EDefaultGraphicsRHI : uint8
 {
@@ -29,7 +30,13 @@ enum class EDefaultGraphicsRHI : uint8
 	DefaultGraphicsRHI_Vulkan = 3 UMETA(DisplayName = "Vulkan"),
 };
 
-
+UENUM()
+enum class EWindowsRHIFeatureLevel : uint8
+{
+	ES3_1 = 0 UMETA(DisplayName = "ES3.1"),
+	SM5   = 1 UMETA(DisplayName = "SM5"),
+	SM6   = 2 UMETA(DisplayName = "SM6"),
+};
 
 /**
  * Implements the settings for the Windows target platform. The first instance of this class is initialized in
@@ -61,6 +68,24 @@ public:
 	UPROPERTY(EditAnywhere, config, Category="Targeted RHIs", Meta = (DisplayName = "Default RHI", ConfigRestartRequired = true))
 	EDefaultGraphicsRHI DefaultGraphicsRHI;
 
+	UPROPERTY(EditAnywhere, config, Category = "Targeted RHIs", AdvancedDisplay, Meta = (DisplayName = "DirectX 12 Minimum FeatureLevel", ConfigRestartRequired = true))
+	EWindowsRHIFeatureLevel D3D12MinimumFeatureLevel;
+
+	UPROPERTY(EditAnywhere, config, Category = "Targeted RHIs", AdvancedDisplay, Meta = (DisplayName = "DirectX 12 Maximum FeatureLevel", ConfigRestartRequired = true))
+	EWindowsRHIFeatureLevel D3D12MaximumFeatureLevel;
+
+	UPROPERTY(EditAnywhere, config, Category = "Targeted RHIs", AdvancedDisplay, Meta = (DisplayName = "DirectX 11 Minimum FeatureLevel", ConfigRestartRequired = true))
+	EWindowsRHIFeatureLevel D3D11MinimumFeatureLevel;
+
+	UPROPERTY(EditAnywhere, config, Category = "Targeted RHIs", AdvancedDisplay, Meta = (DisplayName = "DirectX 11 Maximum FeatureLevel", ConfigRestartRequired = true))
+	EWindowsRHIFeatureLevel D3D11MaximumFeatureLevel;
+
+	UPROPERTY(EditAnywhere, config, Category = "Targeted RHIs", AdvancedDisplay, Meta = (DisplayName = "Vulkan Minimum FeatureLevel", ConfigRestartRequired = true))
+	EWindowsRHIFeatureLevel VulkanMinimumFeatureLevel;
+
+	UPROPERTY(EditAnywhere, config, Category = "Targeted RHIs", AdvancedDisplay, Meta = (DisplayName = "Vulkan Maximum FeatureLevel", ConfigRestartRequired = true))
+	EWindowsRHIFeatureLevel VulkanMaximumFeatureLevel;
+
 	/** Sample rate to run the audio mixer with. */
 	UPROPERTY(config, EditAnywhere, Category = "Audio", Meta = (DisplayName = "Audio Mixer Sample Rate"))
 	int32 AudioSampleRate;
diff --git a/Engine/Source/Runtime/RHI/Private/Windows/WindowsDynamicRHI.cpp b/Engine/Source/Runtime/RHI/Private/Windows/WindowsDynamicRHI.cpp
--- a/Engine/Source/Runtime/RHI/Private/Windows/WindowsDynamicRHI.cpp
+++ b/Engine/Source/Runtime/RHI/Private/Windows/WindowsDynamicRHI.cpp
@@ -32,6 +32,43 @@ static TArray<EShaderPlatform> GetTargetedShaderPlatforms()
 	return ShaderPlatforms;
 }
 
+static bool IsD3DSM6Platform(EShaderPlatform InShaderPlatform)
+{
+	return IsD3DPlatform(InShaderPlatform) && IsFeatureLevelSupported(InShaderPlatform, ERHIFeatureLevel::SM6);
+}
+
+static bool IsD3DSM6PlatformTargeted(const TArray<EShaderPlatform>& TargetedShaderPlatforms)
+{
+	for (const EShaderPlatform ShaderPlatform : TargetedShaderPlatforms)
+	{
+		if (IsD3DSM6Platform(ShaderPlatform))
+		{
+			return true;
+		}
+	}
+
+	return false;
+}
+
+static TOptional<ERHIFeatureLevel::Type> ParseFeatureLevelFromSetting(const TCHAR* InName)
+{
+	TOptional<ERHIFeatureLevel::Type> ResultFeatureLevel{};
+
+	FString ConfigFeatureLevel;
+	if (GConfig->GetString(TEXT("/Script/WindowsTargetPlatform.WindowsTargetSettings"), InName, ConfigFeatureLevel, GEngineIni))
+	{
+		const FName FeatureLevelName(*ConfigFeatureLevel);
+
+		ERHIFeatureLevel::Type ParsedFeatureLevel = ERHIFeatureLevel::Num;
+		if (GetFeatureLevelFromName(FeatureLevelName, ParsedFeatureLevel))
+		{
+			ResultFeatureLevel = ParsedFeatureLevel;
+		}
+	}
+
+	return ResultFeatureLevel;
+}
+
 // Default to Performance Mode on low-end machines
 static bool DefaultFeatureLevelES31()
 {
@@ -128,27 +165,6 @@ static bool IsES31D3DOnly()
 	return bES31DXOnly;
 }
 
-static bool AllowD3D11FeatureLevelES31()
-{
-	return true;
-}
-
-static bool AllowD3D12FeatureLevelES31()
-{
-	if (!GIsEditor)
-	{
-		bool bAllowD3D12FeatureLevelES31 = true;
-		GConfig->GetBool(TEXT("SystemSettings"), TEXT("bAllowD3D12FeatureLevelES31"), bAllowD3D12FeatureLevelES31, GEngineIni);
-		return bAllowD3D12FeatureLevelES31;
-	}
-	return true;
-}
-
-static bool AllowVulkanFeatureLevelES31()
-{
-	return !IsES31D3DOnly();
-}
-
 enum class WindowsRHI
 {
 	D3D11,
@@ -205,23 +221,34 @@ static WindowsRHI ChooseDefaultRHI(const TArray<EShaderPlatform>& TargetedShader
 		{
 			DefaultRHI = WindowsRHI::OpenGL;
 		}
+		else if (IsD3DSM6Platform(TargetedPlatform))
+		{
+			DefaultRHI = WindowsRHI::D3D12;
+		}
 		else
 		{
 			DefaultRHI = WindowsRHI::D3D11;
 		}
 	}
 
+	return DefaultRHI;
+}
+
+static TOptional<WindowsRHI> ChoosePreferredRHI(WindowsRHI InDefaultRHI)
+{
+	TOptional<WindowsRHI> RHIPreference{};
+
 	// If we are in game, there is a separate setting that can make it prefer D3D12 over D3D11 (but not over other RHIs).
-	if (!GIsEditor && (DefaultRHI == WindowsRHI::D3D11 || DefaultRHI == WindowsRHI::D3D12))
+	if (!GIsEditor && (InDefaultRHI == WindowsRHI::D3D11 || InDefaultRHI == WindowsRHI::D3D12))
 	{
 		bool bUseD3D12InGame = false;
-		if (GConfig->GetBool(TEXT("D3DRHIPreference"), TEXT("bUseD3D12InGame"), bUseD3D12InGame, GGameUserSettingsIni))
+		if (GConfig->GetBool(TEXT("D3DRHIPreference"), TEXT("bUseD3D12InGame"), bUseD3D12InGame, GGameUserSettingsIni) && bUseD3D12InGame)
 		{
-			DefaultRHI = bUseD3D12InGame ? WindowsRHI::D3D12 : WindowsRHI::D3D11;
+			RHIPreference = WindowsRHI::D3D12;
 		}
 	}
 
-	return DefaultRHI;
+	return RHIPreference;
 }
 
 static TOptional<WindowsRHI> ChooseForcedRHI()
@@ -313,56 +340,94 @@ static TOptional<ERHIFeatureLevel::Type> ChooseForcedFeatureLevel(TOptional<Wind
 	return ForcedFeatureLevel;
 }
 
-static bool IsD3DSM6PlatformTargeted(const TArray<EShaderPlatform>& TargetedShaderPlatforms)
+static ERHIFeatureLevel::Type FilterFeatureLevel(ERHIFeatureLevel::Type InFeatureLevel, const TCHAR* MinSetting, const TCHAR* MaxSetting)
 {
-	for (const EShaderPlatform ShaderPlatform : TargetedShaderPlatforms)
+	const TOptional<ERHIFeatureLevel::Type> MinFeatureLevel = ParseFeatureLevelFromSetting(MinSetting);
+	const TOptional<ERHIFeatureLevel::Type> MaxFeatureLevel = ParseFeatureLevelFromSetting(MaxSetting);
+
+	if (MinFeatureLevel && InFeatureLevel < MinFeatureLevel.GetValue())
 	{
-		if (IsD3DPlatform(ShaderPlatform) && IsFeatureLevelSupported(ShaderPlatform, ERHIFeatureLevel::SM6))
-		{
-			return true;
-		}
+		return MinFeatureLevel.GetValue();
+	}
+	else if (MaxFeatureLevel && InFeatureLevel > MaxFeatureLevel.GetValue())
+	{
+		return MaxFeatureLevel.GetValue();
 	}
 
-	return false;
+	return InFeatureLevel;
 }
 
-static ERHIFeatureLevel::Type ChooseFeatureLevel(TOptional<WindowsRHI> ChosenRHI, TOptional<ERHIFeatureLevel::Type> ForcedFeatureLevel, const TArray<EShaderPlatform>& TargetedShaderPlatforms)
+static ERHIFeatureLevel::Type FilterFeatureLevel(ERHIFeatureLevel::Type FeatureLevel, WindowsRHI ChosenRHI)
 {
-	TOptional<ERHIFeatureLevel::Type> FeatureLevel{};
+	if (ChosenRHI == WindowsRHI::OpenGL)
+	{
+		// Seriously locked down.
+		return ERHIFeatureLevel::ES3_1;
+	}
+
+	if (ChosenRHI == WindowsRHI::D3D12)
+	{
+		return FilterFeatureLevel(FeatureLevel, TEXT("D3D12MinimumFeatureLevel"), TEXT("D3D12MaximumFeatureLevel"));
+	}
 
+	if (ChosenRHI == WindowsRHI::D3D11)
+	{
+		return FilterFeatureLevel(FeatureLevel, TEXT("D3D11MinimumFeatureLevel"), TEXT("D3D11MaximumFeatureLevel"));
+	}
+
+	if (ChosenRHI == WindowsRHI::Vulkan)
+	{
+		return FilterFeatureLevel(FeatureLevel, TEXT("VulkanMinimumFeatureLevel"), TEXT("VulkanMaximumFeatureLevel"));
+	}
+
+	return FeatureLevel;
+}
+
+static ERHIFeatureLevel::Type ChooseFeatureLevel(TOptional<WindowsRHI> ChosenRHI, TOptional<ERHIFeatureLevel::Type> ForcedFeatureLevel, const TArray<EShaderPlatform>& TargetedShaderPlatforms)
+{
 	if (ForcedFeatureLevel)
 	{
-		if (ForcedFeatureLevel == ERHIFeatureLevel::ES3_1 &&
-			(ChosenRHI == WindowsRHI::D3D11 && AllowD3D11FeatureLevelES31()) ||
-			(ChosenRHI == WindowsRHI::D3D12 && AllowD3D12FeatureLevelES31()) ||
-			(ChosenRHI == WindowsRHI::Vulkan && AllowVulkanFeatureLevelES31()))
+		// Allow the forced feature level if we're in a position to compile its shaders
+		if (!FPlatformProperties::RequiresCookedData())
 		{
-			FeatureLevel = ForcedFeatureLevel;
+			return ForcedFeatureLevel.GetValue();
 		}
-		else if (ForcedFeatureLevel == ERHIFeatureLevel::SM6 && ChosenRHI == WindowsRHI::D3D12)
+
+		// Make sure the feature level is supported by the runtime, otherwise fall back to the default
+		const ERHIFeatureLevel::Type FilteredForcedFeatureLevel = FilterFeatureLevel(ForcedFeatureLevel.GetValue(), ChosenRHI.GetValue());
+		if (FilteredForcedFeatureLevel == ForcedFeatureLevel)
 		{
-			FeatureLevel = ForcedFeatureLevel;
+			return FilteredForcedFeatureLevel;
 		}
 	}
 
-	if (!FeatureLevel)
+	ERHIFeatureLevel::Type FeatureLevel;
+
+	if (ChosenRHI == WindowsRHI::OpenGL)
+	{
+		// OpenGL can only be used for mobile preview
+		FeatureLevel = ERHIFeatureLevel::ES3_1;
+	}
+	else if (ChosenRHI == WindowsRHI::D3D12)
+	{
+		FeatureLevel = IsD3DSM6PlatformTargeted(TargetedShaderPlatforms) ? ERHIFeatureLevel::SM6 : ERHIFeatureLevel::SM5;
+	}
+	else
+	{
+		FeatureLevel = ERHIFeatureLevel::SM5;
+	}
+
+	ERHIFeatureLevel::Type FilteredFeatureLevel = FilterFeatureLevel(FeatureLevel, ChosenRHI.GetValue());
+
+	// If the user wanted to force a feature level and we couldn't set it, log out why and what we're actually running with
+	if (ForcedFeatureLevel)
 	{
-		if (ChosenRHI == WindowsRHI::OpenGL)
-		{
-			// OpenGL can only be used for mobile preview
-			FeatureLevel = ERHIFeatureLevel::ES3_1;
-		}
-		else if (ChosenRHI == WindowsRHI::D3D12 && IsD3DSM6PlatformTargeted(TargetedShaderPlatforms))
-		{
-			FeatureLevel = ERHIFeatureLevel::SM6;
-		}
-		else
-		{
-			FeatureLevel = ERHIFeatureLevel::SM5;
-		}
+		const FString ForcedName = LexToString(ForcedFeatureLevel.GetValue());
+		const FString UsedName = LexToString(FilteredFeatureLevel);
+		UE_LOG(LogRHI, Warning, TEXT("User requested Feature Level '%s' but that is not supported by this project. Falling back to Feature Level '%s'."), *ForcedName, *UsedName);
 	}
 
-	return FeatureLevel.GetValue();
+	return FilteredFeatureLevel;
 }
 
 static IDynamicRHIModule* LoadDynamicRHIModule(ERHIFeatureLevel::Type& DesiredFeatureLevel, const TCHAR*& LoadedRHIModuleName)
@@ -381,13 +446,18 @@ static IDynamicRHIModule* LoadDynamicRHIModule(ERHIFeatureLevel::Type& DesiredFe
 	// Commandline switches apply after this and can force an arbitrary RHIs. If RHI isn't supported, the game will refuse to start.
 
 	WindowsRHI DefaultRHI = ChooseDefaultRHI(TargetedShaderPlatforms);
-	TOptional<WindowsRHI> ForcedRHI = ChooseForcedRHI();
+	const TOptional<WindowsRHI> PreferredRHI = ChoosePreferredRHI(DefaultRHI);
+	const TOptional<WindowsRHI> ForcedRHI = ChooseForcedRHI();
 
 	WindowsRHI ChosenRHI = DefaultRHI;
 	if (ForcedRHI)
 	{
 		ChosenRHI = ForcedRHI.GetValue();
 	}
+	else if (PreferredRHI)
+	{
+		ChosenRHI = PreferredRHI.GetValue();
+	}
 
 	TOptional<ERHIFeatureLevel::Type> ForcedFeatureLevel = ChooseForcedFeatureLevel(ChosenRHI, ForcedRHI);
 	DesiredFeatureLevel = ChooseFeatureLevel(ChosenRHI, ForcedFeatureLevel, TargetedShaderPlatforms);
@@ -446,10 +516,18 @@ static IDynamicRHIModule* LoadDynamicRHIModule(ERHIFeatureLevel::Type& DesiredFe
 				FPlatformMisc::RequestExit(1);
 			}
 
-			UE_LOG(LogRHI, Log, TEXT("D3D12 SM6 is not supported, trying SM5"));
+			if (FilterFeatureLevel(ERHIFeatureLevel::SM5, WindowsRHI::D3D12) == ERHIFeatureLevel::SM5)
+			{
+				UE_LOG(LogRHI, Log, TEXT("D3D12 SM6 is not supported, trying SM5"));
 
-			DesiredFeatureLevel = ERHIFeatureLevel::SM5;
-			bD3D12Supported = DynamicRHIModule->IsSupported(DesiredFeatureLevel);
+				DesiredFeatureLevel = ERHIFeatureLevel::SM5;
+				bD3D12Supported = DynamicRHIModule->IsSupported(DesiredFeatureLevel);
+			}
+			
+			if (!bD3D12Supported)
+			{
+				UE_LOG(LogRHI, Log, TEXT("D3D12 SM5 is not supported, trying D3D11"));
+			}
 		}
 
 		if (!bD3D12Supported)
@@ -471,6 +549,9 @@ static IDynamicRHIModule* LoadDynamicRHIModule(ERHIFeatureLevel::Type& DesiredFe
 	// Fallback to D3D11RHI if nothing is selected
 	if (!DynamicRHIModule)
 	{
+		// We need to filter one more time to handle the D3D12 fallback
+		DesiredFeatureLevel = FilterFeatureLevel(DesiredFeatureLevel, WindowsRHI::D3D11);
+
 		FApp::SetGraphicsRHI(TEXT("DirectX 11"));
 		const TCHAR* D3D11RHIModuleName = TEXT("D3D11RHI");
 		DynamicRHIModule = &FModuleManager::LoadModuleChecked<IDynamicRHIModule>(D3D11RHIModuleName);
