Restoring the accidentally removed AllowD3D12FeatureLevelES31 and making it use the new project settings.

diff --git a/Engine/Source/Runtime/RHI/Private/Windows/WindowsDynamicRHI.cpp b/Engine/Source/Runtime/RHI/Private/Windows/WindowsDynamicRHI.cpp
--- a/Engine/Source/Runtime/RHI/Private/Windows/WindowsDynamicRHI.cpp
+++ b/Engine/Source/Runtime/RHI/Private/Windows/WindowsDynamicRHI.cpp
@@ -16,6 +16,14 @@
 
 static const TCHAR* GLoadedRHIModuleName;
 
+enum class WindowsRHI
+{
+	D3D11,
+	D3D12,
+	Vulkan,
+	OpenGL,
+};
+
 static TArray<EShaderPlatform> GetTargetedShaderPlatforms()
 {
 	TArray<FString> TargetedShaderFormats;
@@ -69,6 +77,49 @@
 	return ResultFeatureLevel;
 }
 
+static ERHIFeatureLevel::Type FilterFeatureLevel(ERHIFeatureLevel::Type InFeatureLevel, const TCHAR* MinSetting, const TCHAR* MaxSetting)
+{
+	const TOptional<ERHIFeatureLevel::Type> MinFeatureLevel = ParseFeatureLevelFromSetting(MinSetting);
+	const TOptional<ERHIFeatureLevel::Type> MaxFeatureLevel = ParseFeatureLevelFromSetting(MaxSetting);
+
+	if (MinFeatureLevel && InFeatureLevel < MinFeatureLevel.GetValue())
+	{
+		return MinFeatureLevel.GetValue();
+	}
+	else if (MaxFeatureLevel && InFeatureLevel > MaxFeatureLevel.GetValue())
+	{
+		return MaxFeatureLevel.GetValue();
+	}
+
+	return InFeatureLevel;
+}
+
+static ERHIFeatureLevel::Type FilterFeatureLevel(ERHIFeatureLevel::Type FeatureLevel, WindowsRHI ChosenRHI)
+{
+	if (ChosenRHI == WindowsRHI::OpenGL)
+	{
+		// Seriously locked down.
+		return ERHIFeatureLevel::ES3_1;
+	}
+
+	if (ChosenRHI == WindowsRHI::D3D12)
+	{
+		return FilterFeatureLevel(FeatureLevel, TEXT("D3D12MinimumFeatureLevel"), TEXT("D3D12MaximumFeatureLevel"));
+	}
+
+	if (ChosenRHI == WindowsRHI::D3D11)
+	{
+		return FilterFeatureLevel(FeatureLevel, TEXT("D3D11MinimumFeatureLevel"), TEXT("D3D11MaximumFeatureLevel"));
+	}
+
+	if (ChosenRHI == WindowsRHI::Vulkan)
+	{
+		return FilterFeatureLevel(FeatureLevel, TEXT("VulkanMinimumFeatureLevel"), TEXT("VulkanMaximumFeatureLevel"));
+	}
+
+	return FeatureLevel;
+}
+
 // Default to Performance Mode on low-end machines
 static bool DefaultFeatureLevelES31()
 {
@@ -165,13 +216,14 @@
 	return bES31DXOnly;
 }
 
-enum class WindowsRHI
+static bool AllowD3D12FeatureLevelES31()
 {
-	D3D11,
-	D3D12,
-	Vulkan,
-	OpenGL,
-};
+	if (!GIsEditor)
+	{
+		return FilterFeatureLevel(ERHIFeatureLevel::ES3_1, WindowsRHI::D3D12) == ERHIFeatureLevel::ES3_1;
+	}
+	return true;
+}
 
 // Choose the default from DefaultGraphicsRHI or TargetedRHIs. DefaultGraphicsRHI has precedence.
 static WindowsRHI ChooseDefaultRHI(const TArray<EShaderPlatform>& TargetedShaderPlatforms)
@@ -340,49 +392,6 @@
 	return ForcedFeatureLevel;
 }
 
-static ERHIFeatureLevel::Type FilterFeatureLevel(ERHIFeatureLevel::Type InFeatureLevel, const TCHAR* MinSetting, const TCHAR* MaxSetting)
-{
-	const TOptional<ERHIFeatureLevel::Type> MinFeatureLevel = ParseFeatureLevelFromSetting(MinSetting);
-	const TOptional<ERHIFeatureLevel::Type> MaxFeatureLevel = ParseFeatureLevelFromSetting(MaxSetting);
-
-	if (MinFeatureLevel && InFeatureLevel < MinFeatureLevel.GetValue())
-	{
-		return MinFeatureLevel.GetValue();
-	}
-	else if (MaxFeatureLevel && InFeatureLevel > MaxFeatureLevel.GetValue())
-	{
-		return MaxFeatureLevel.GetValue();
-	}
-
-	return InFeatureLevel;
-}
-
-static ERHIFeatureLevel::Type FilterFeatureLevel(ERHIFeatureLevel::Type FeatureLevel, WindowsRHI ChosenRHI)
-{
-	if (ChosenRHI == WindowsRHI::OpenGL)
-	{
-		// Seriously locked down.
-		return ERHIFeatureLevel::ES3_1;
-	}
-
-	if (ChosenRHI == WindowsRHI::D3D12)
-	{
-		return FilterFeatureLevel(FeatureLevel, TEXT("D3D12MinimumFeatureLevel"), TEXT("D3D12MaximumFeatureLevel"));
-	}
-
-	if (ChosenRHI == WindowsRHI::D3D11)
-	{
-		return FilterFeatureLevel(FeatureLevel, TEXT("D3D11MinimumFeatureLevel"), TEXT("D3D11MaximumFeatureLevel"));
-	}
-
-	if (ChosenRHI == WindowsRHI::Vulkan)
-	{
-		return FilterFeatureLevel(FeatureLevel, TEXT("VulkanMinimumFeatureLevel"), TEXT("VulkanMaximumFeatureLevel"));
-	}
-
-	return FeatureLevel;
-}
-
 static ERHIFeatureLevel::Type ChooseFeatureLevel(TOptional<WindowsRHI> ChosenRHI, TOptional<ERHIFeatureLevel::Type> ForcedFeatureLevel, const TArray<EShaderPlatform>& TargetedShaderPlatforms)
 {
 	if (ForcedFeatureLevel)
