RHIValidation: Fixed CitySample being rendered with artifacts (mostly in black and white) by implementing some missing members.

#rnx

diff --git a/Engine/Source/Runtime/RHI/Public/RHIValidation.h b/Engine/Source/Runtime/RHI/Public/RHIValidation.h
--- a/Engine/Source/Runtime/RHI/Public/RHIValidation.h
+++ b/Engine/Source/Runtime/RHI/Public/RHIValidation.h
@@ -212,6 +212,11 @@ public:
 		return RHI->RHICreateShaderLibrary(Platform, FilePath, Name);
 	}
 
+	virtual FRenderQueryPoolRHIRef RHICreateRenderQueryPool(ERenderQueryType QueryType, uint32 NumQueries = UINT32_MAX) override final
+	{
+		return RHI->RHICreateRenderQueryPool(QueryType, NumQueries);
+	}
+
 	virtual FGPUFenceRHIRef RHICreateGPUFence(const FName &Name) override final
 	{
 		return RHI->RHICreateGPUFence(Name);
@@ -771,6 +776,11 @@ public:
 		RHI->RHIUpdateTexture2D(Texture, MipIndex, UpdateRegion, SourcePitch, SourceData);
 	}
 
+	virtual void RHIUpdateFromBufferTexture2D(FRHITexture2D* Texture, uint32 MipIndex, const struct FUpdateTextureRegion2D& UpdateRegion, uint32 SourcePitch, FRHIBuffer* Buffer, uint32 BufferOffset) override final
+	{
+		RHI->RHIUpdateFromBufferTexture2D(Texture, MipIndex, UpdateRegion, SourcePitch, Buffer, BufferOffset);
+	}
+
 	/**
 	* Updates a region of a 3D texture from system memory
 	* @param Texture - the RHI texture resource to update
@@ -827,6 +837,19 @@ public:
 		RHI->RHIBindDebugLabelName(Texture, Name);
 	}
 
+	virtual void RHIBindDebugLabelName(FRHIBuffer* Buffer, const TCHAR* Name) override final
+	{
+		check(IsInRenderingThread());
+
+		FString NameCopyRT = Name;
+		FRHICommandListExecutor::GetImmediateCommandList().EnqueueLambda([Buffer, NameCopyRHIT = MoveTemp(NameCopyRT)](FRHICommandListImmediate& RHICmdList)
+		{
+			((FValidationContext&)RHICmdList.GetContext()).Tracker->Rename(Buffer, *NameCopyRHIT);
+		});
+
+		RHI->RHIBindDebugLabelName(Buffer, Name);
+	}
+
 	virtual void RHIBindDebugLabelName(FRHIUnorderedAccessView* UnorderedAccessViewRHI, const TCHAR* Name) override final
 	{
 		RHIValidation::FResource* Resource = UnorderedAccessViewRHI->ViewIdentity.Resource;
@@ -881,6 +904,11 @@ public:
 		RHI->RHIRead3DSurfaceFloatData(Texture, Rect, ZMinMax, OutData);
 	}
 
+	virtual void RHIRead3DSurfaceFloatData(FRHITexture* Texture, FIntRect Rect, FIntPoint ZMinMax, TArray<FFloat16Color>& OutData, FReadSurfaceDataFlags InFlags) override final
+	{
+		RHI->RHIRead3DSurfaceFloatData(Texture, Rect, ZMinMax, OutData, InFlags);
+	}
+
 	// FlushType: Wait RHI Thread
 	virtual FRenderQueryRHIRef RHICreateRenderQuery(ERenderQueryType QueryType) override final
 	{
@@ -1024,6 +1052,21 @@ public:
 		RHI->RHIResizeViewport(Viewport, SizeX, SizeY, bIsFullscreen, PreferredPixelFormat);
 	}
 
+	virtual EColorSpaceAndEOTF RHIGetColorSpace(FRHIViewport* Viewport) override final
+	{
+		return RHI->RHIGetColorSpace(Viewport);
+	}
+
+	virtual EPixelFormat RHIPreferredPixelFormatHint(EPixelFormat PreferredPixelFormat) override final
+	{
+		return RHI->RHIPreferredPixelFormatHint(PreferredPixelFormat);
+	}
+
+	virtual void RHICheckViewportHDRStatus(FRHIViewport* Viewport) override final
+	{
+		RHI->RHICheckViewportHDRStatus(Viewport);
+	}
+
 	//  must be called from the main thread.
 	// FlushType: Thread safe
 	virtual void RHITick(float DeltaTime) override final
@@ -1340,6 +1383,21 @@ public:
 		RHI->UpdateTexture2D_RenderThread(RHICmdList, Texture, MipIndex, UpdateRegion, SourcePitch, SourceData);
 	}
 
+	virtual void UpdateFromBufferTexture2D_RenderThread(class FRHICommandListImmediate& RHICmdList, FRHITexture2D* Texture, uint32 MipIndex, const struct FUpdateTextureRegion2D& UpdateRegion, uint32 SourcePitch, FRHIBuffer* Buffer, uint32 BufferOffset) override final
+	{
+		RHI->UpdateFromBufferTexture2D_RenderThread(RHICmdList, Texture, MipIndex, UpdateRegion, SourcePitch, Buffer, BufferOffset);
+	}
+
+	virtual void* LockTexture2DArray_RenderThread(class FRHICommandListImmediate& RHICmdList, FRHITexture2DArray* Texture, uint32 ArrayIndex, uint32 MipIndex, EResourceLockMode LockMode, uint32& DestStride, bool bLockWithinMiptail) override final
+	{
+		return RHI->LockTexture2DArray_RenderThread(RHICmdList, Texture, ArrayIndex, MipIndex, LockMode, DestStride, bLockWithinMiptail);
+	}
+
+	virtual void UnlockTexture2DArray_RenderThread(class FRHICommandListImmediate& RHICmdList, FRHITexture2DArray* Texture, uint32 ArrayIndex, uint32 MipIndex, bool bLockWithinMiptail) override final
+	{
+		RHI->UnlockTexture2DArray_RenderThread(RHICmdList, Texture, ArrayIndex, MipIndex, bLockWithinMiptail);
+	}
+
 	virtual FUpdateTexture3DData BeginUpdateTexture3D_RenderThread(class FRHICommandListImmediate& RHICmdList, FRHITexture3D* Texture, uint32 MipIndex, const struct FUpdateTextureRegion3D& UpdateRegion) override final
 	{
 		return RHI->BeginUpdateTexture3D_RenderThread(RHICmdList, Texture, MipIndex, UpdateRegion);
@@ -1389,6 +1447,15 @@ public:
 		return UAV;
 	}
 
+	virtual FUnorderedAccessViewRHIRef RHICreateUnorderedAccessView_RenderThread(class FRHICommandListImmediate& RHICmdList, FRHITexture* Texture, uint32 MipLevel, uint8 Format, uint16 FirstArraySlice, uint16 NumArraySlices) override final
+	{
+		RHI_VALIDATION_CHECK(Texture->GetTextureReference() == nullptr, TEXT("Creating an unordered access view of an FRHITextureReference is not supported."));
+
+		FUnorderedAccessViewRHIRef UAV = RHI->RHICreateUnorderedAccessView_RenderThread(RHICmdList, Texture, MipLevel, Format, FirstArraySlice, NumArraySlices);
+		UAV->ViewIdentity = Texture->GetViewIdentity(MipLevel, 1, FirstArraySlice, NumArraySlices, uint32(RHIValidation::EResourcePlane::Common), 1);
+		return UAV;
+	}
+
 	virtual FUnorderedAccessViewRHIRef RHICreateUnorderedAccessView_RenderThread(class FRHICommandListImmediate& RHICmdList, FRHIBuffer* Buffer, uint8 Format) override final
 	{
 		FUnorderedAccessViewRHIRef UAV = RHI->RHICreateUnorderedAccessView_RenderThread(RHICmdList, Buffer, Format);
