Removed a few high-frequency/low-value RHI stats to improve profiler readability and stability

diff --git a/Engine/Source/Runtime/Core/Private/Async/TaskGraph.cpp b/Engine/Source/Runtime/Core/Private/Async/TaskGraph.cpp
--- a/Engine/Source/Runtime/Core/Private/Async/TaskGraph.cpp
+++ b/Engine/Source/Runtime/Core/Private/Async/TaskGraph.cpp
@@ -1586,7 +1586,6 @@ PRAGMA_ENABLE_DEPRECATION_WARNINGS
 
 	virtual uint64 ProcessThreadUntilIdle(ENamedThreads::Type CurrentThread) final override
 	{
-		SCOPED_NAMED_EVENT(ProcessThreadUntilIdle, FColor::Red);
 		int32 QueueIndex = ENamedThreads::GetQueueIndex(CurrentThread);
 		CurrentThread = ENamedThreads::GetThreadIndex(CurrentThread);
 		check(CurrentThread >= 0 && CurrentThread < NumNamedThreads);
@@ -2107,7 +2106,6 @@ private:
 
 	uint64 ProcessThreadUntilIdle(ENamedThreads::Type CurrentThread) final override
 	{
-		SCOPED_NAMED_EVENT(ProcessThreadUntilIdle, FColor::Red);
 		int32 QueueIndex = ENamedThreads::GetQueueIndex(CurrentThread);
 		CurrentThread = ENamedThreads::GetThreadIndex(CurrentThread);
 		check(CurrentThread >= 0 && CurrentThread < NumNamedThreads);
diff --git a/Engine/Source/Runtime/RHI/Private/RHICommandList.cpp b/Engine/Source/Runtime/RHI/Private/RHICommandList.cpp
--- a/Engine/Source/Runtime/RHI/Private/RHICommandList.cpp
+++ b/Engine/Source/Runtime/RHI/Private/RHICommandList.cpp
@@ -640,7 +640,6 @@ void FRHICommandListExecutor::ExecuteInner(FRHICommandListBase& CmdList)
 		{
 			if (!bIsInGameThread && !FTaskGraphInterface::Get().IsThreadProcessingTasks(RenderThread_Local))
 			{
-				QUICK_SCOPE_CYCLE_COUNTER(STAT_FRHICommandListExecutor_ExecuteInner_DoTasksBeforeDispatch);
 				// move anything down the pipe that needs to go
 				FTaskGraphInterface::Get().ProcessThreadUntilIdle(RenderThread_Local);
 			}
@@ -671,7 +670,6 @@ void FRHICommandListExecutor::ExecuteInner(FRHICommandListBase& CmdList)
 			FGraphEventArray Prereq;
 			Exchange(Prereq, CmdList.RTTasks); 
 			{
-				QUICK_SCOPE_CYCLE_COUNTER(STAT_FRHICommandListExecutor_SwapCmdLists);
 				SwapCmdList = new FRHICommandList(CmdList.GetGPUMask());
 
 				// Super scary stuff here, but we just want the swap command list to inherit everything and leave the immediate command list wiped.
@@ -687,7 +685,6 @@ void FRHICommandListExecutor::ExecuteInner(FRHICommandListBase& CmdList)
 				CmdList.Data.bInsideRenderPass = SwapCmdList->Data.bInsideRenderPass;
 				CmdList.Data.bInsideComputePass = SwapCmdList->Data.bInsideComputePass;
 			}
-			QUICK_SCOPE_CYCLE_COUNTER(STAT_FRHICommandListExecutor_SubmitTasks);
 
 			//if we use a FDispatchRHIThreadTask, we must have it pass an event along to the FExecuteRHIThreadTask it will spawn so that fences can know which event to wait on for execution completion
 			//before the dispatch completes.
