Add EPixelFormatCapabilities::TextureStore to complete list of possible capabilities (BufferLoad, BufferStore, TextureLoad, TextureStore), since some platforms expose formats with differing capabilities.
Interpret D3D*_FORMAT_SUPPORT_BUFFER as EPixelFormatCapabilities::BufferLoad, rather than D3D*_FORMAT_SUPPORT_SHADER_LOAD, which is for texture loads.

diff --git a/Engine/Source/Runtime/D3D12RHI/Private/D3D12Device.cpp b/Engine/Source/Runtime/D3D12RHI/Private/D3D12Device.cpp
--- a/Engine/Source/Runtime/D3D12RHI/Private/D3D12Device.cpp
+++ b/Engine/Source/Runtime/D3D12RHI/Private/D3D12Device.cpp
@@ -231,11 +231,12 @@ void FD3D12Device::SetupAfterDeviceCreation()
 				ConvertCap1(SRVFormatSupport, EPixelFormatCapabilities::TextureGather,    D3D12_FORMAT_SUPPORT1_SHADER_GATHER);
 				ConvertCap2(UAVFormatSupport, EPixelFormatCapabilities::TextureAtomics,   D3D12_FORMAT_SUPPORT2_UAV_ATOMIC_EXCHANGE);
 				ConvertCap1(SRVFormatSupport, EPixelFormatCapabilities::TextureBlendable, D3D12_FORMAT_SUPPORT1_BLENDABLE);
+				ConvertCap2(UAVFormatSupport, EPixelFormatCapabilities::TextureStore,     D3D12_FORMAT_SUPPORT2_UAV_TYPED_STORE);
 			}
 
 			if (EnumHasAnyFlags(Capabilities, EPixelFormatCapabilities::Buffer))
 			{
-				ConvertCap1(SRVFormatSupport, EPixelFormatCapabilities::BufferLoad,       D3D12_FORMAT_SUPPORT1_SHADER_LOAD);
+				ConvertCap1(SRVFormatSupport, EPixelFormatCapabilities::BufferLoad,       D3D12_FORMAT_SUPPORT1_BUFFER);
 				ConvertCap2(UAVFormatSupport, EPixelFormatCapabilities::BufferStore,      D3D12_FORMAT_SUPPORT2_UAV_TYPED_STORE);
 				ConvertCap2(UAVFormatSupport, EPixelFormatCapabilities::BufferAtomics,    D3D12_FORMAT_SUPPORT2_UAV_ATOMIC_EXCHANGE);
 			}
diff --git a/Engine/Source/Runtime/RHI/Public/RHI.h b/Engine/Source/Runtime/RHI/Public/RHI.h
--- a/Engine/Source/Runtime/RHI/Public/RHI.h
+++ b/Engine/Source/Runtime/RHI/Public/RHI.h
@@ -795,21 +795,22 @@ enum class EPixelFormatCapabilities : uint32
 	TextureGather    = 1ull << 10,
 	TextureAtomics   = 1ull << 11,
 	TextureBlendable = 1ull << 12,
+	TextureStore     = 1ull << 13,
 
-	Buffer           = 1ull << 13,
-    VertexBuffer     = 1ull << 14,
-    IndexBuffer      = 1ull << 15,
-	BufferLoad       = 1ull << 16,
-    BufferStore      = 1ull << 17,
-    BufferAtomics    = 1ull << 18,
+	Buffer           = 1ull << 14,
+    VertexBuffer     = 1ull << 15,
+    IndexBuffer      = 1ull << 16,
+	BufferLoad       = 1ull << 17,
+    BufferStore      = 1ull << 18,
+    BufferAtomics    = 1ull << 19,
 
-	UAV              = 1ull << 19,
-    TypedUAVLoad     = 1ull << 20,
-	TypedUAVStore    = 1ull << 21,
+	UAV              = 1ull << 20,
+    TypedUAVLoad     = 1ull << 21,
+	TypedUAVStore    = 1ull << 22,
 
 	AnyTexture       = Texture1D | Texture2D | Texture3D | TextureCube,
 
-	AllTextureFlags  = AnyTexture | RenderTarget | DepthStencil | TextureMipmaps | TextureLoad | TextureSample | TextureGather | TextureAtomics | TextureBlendable,
+	AllTextureFlags  = AnyTexture | RenderTarget | DepthStencil | TextureMipmaps | TextureLoad | TextureSample | TextureGather | TextureAtomics | TextureBlendable | TextureStore,
 	AllBufferFlags   = Buffer | VertexBuffer | IndexBuffer | BufferLoad | BufferStore | BufferAtomics,
 	AllUAVFlags      = UAV | TypedUAVLoad | TypedUAVStore,
 
diff --git a/Engine/Source/Runtime/Windows/D3D11RHI/Private/D3D11Device.cpp b/Engine/Source/Runtime/Windows/D3D11RHI/Private/D3D11Device.cpp
--- a/Engine/Source/Runtime/Windows/D3D11RHI/Private/D3D11Device.cpp
+++ b/Engine/Source/Runtime/Windows/D3D11RHI/Private/D3D11Device.cpp
@@ -470,11 +470,12 @@ void FD3D11DynamicRHI::SetupAfterDeviceCreation()
                 ConvertCap1(SRVFormatSupport, EPixelFormatCapabilities::TextureGather,    D3D11_FORMAT_SUPPORT_SHADER_GATHER);
                 ConvertCap2(UAVFormatSupport, EPixelFormatCapabilities::TextureAtomics,   D3D11_FORMAT_SUPPORT2_UAV_ATOMIC_EXCHANGE);
                 ConvertCap1(RTVFormatSupport, EPixelFormatCapabilities::TextureBlendable, D3D11_FORMAT_SUPPORT_BLENDABLE);
+                ConvertCap2(UAVFormatSupport, EPixelFormatCapabilities::TextureStore,     D3D11_FORMAT_SUPPORT2_UAV_TYPED_STORE);
             }
 
             if (EnumHasAnyFlags(Capabilities, EPixelFormatCapabilities::Buffer))
             {
-                ConvertCap1(SRVFormatSupport, EPixelFormatCapabilities::BufferLoad,       D3D11_FORMAT_SUPPORT_SHADER_LOAD);
+                ConvertCap1(SRVFormatSupport, EPixelFormatCapabilities::BufferLoad,       D3D11_FORMAT_SUPPORT_BUFFER);
                 ConvertCap2(UAVFormatSupport, EPixelFormatCapabilities::BufferStore,      D3D11_FORMAT_SUPPORT2_UAV_TYPED_STORE);
                 ConvertCap2(UAVFormatSupport, EPixelFormatCapabilities::BufferAtomics,    D3D11_FORMAT_SUPPORT2_UAV_ATOMIC_EXCHANGE);
             }
