- Initialize member variables

#rnx

diff --git a/Engine/Source/Runtime/RHI/Public/RHIValidationCommon.h b/Engine/Source/Runtime/RHI/Public/RHIValidationCommon.h
--- a/Engine/Source/Runtime/RHI/Public/RHIValidationCommon.h
+++ b/Engine/Source/Runtime/RHI/Public/RHIValidationCommon.h
@@ -530,6 +530,12 @@ namespace RHIValidation
 
 	struct FView
 	{
+		FView()
+		{
+			ViewIdentity.Resource = nullptr;
+			ViewIdentity.SubresourceRange = FSubresourceRange(0, 0, 0, 0, 0, 0);
+		}
+
 		FResourceIdentity ViewIdentity;
 	};
 
