MGPU:  Bug fix for various crashes when running with -vmgpu or multiple GPUs (startup crash in Virtual Textures, and others if you disable that system to get past the first crash).  By default, we want texture resources to be created on all GPUs in multi-GPU scenarios.  The RHI texture unification work introduced this bug.

I do have plans to provide ways of flagging resources as view specific, with validation that runs even when not using MGPU features.  My current idea is a new type of RDG scope that marks a block of rendering as view specific, throws errors if a view specific resource is used outside of such a scope, and will automatically migrate external resources cross GPU if the user changes the view's assignment between frames.  That would allow us to start extensively using single GPU resources (for better efficiency in GPU memory usage and perf) without worrying about crashes that only occur when MGPU is enabled.

#rnx

diff --git a/Engine/Source/Runtime/RHI/Public/RHIResources.h b/Engine/Source/Runtime/RHI/Public/RHIResources.h
--- a/Engine/Source/Runtime/RHI/Public/RHIResources.h
+++ b/Engine/Source/Runtime/RHI/Public/RHIResources.h
@@ -1503,7 +1503,7 @@ struct FRHITextureCreateDesc : public FRHITextureDesc
 		, uint32		              ExtData      = 0
 		, ERHIAccess                  InitialState = ERHIAccess::Unknown
 		, FResourceBulkDataInterface* BulkData     = nullptr
-		, FRHIGPUMask                 GPUMask      = FRHIGPUMask()
+		, FRHIGPUMask                 GPUMask      = FRHIGPUMask::All()
 		)
 	{
 		return FRHITextureCreateDesc(
@@ -1527,7 +1527,7 @@ struct FRHITextureCreateDesc : public FRHITextureDesc
 		, uint32		              ExtData      = 0
 		, ERHIAccess                  InitialState = ERHIAccess::Unknown
 		, FResourceBulkDataInterface* BulkData     = nullptr
-		, FRHIGPUMask                 GPUMask      = FRHIGPUMask()
+		, FRHIGPUMask                 GPUMask      = FRHIGPUMask::All()
 		)
 	{
 		return FRHITextureCreateDesc(
@@ -1549,7 +1549,7 @@ struct FRHITextureCreateDesc : public FRHITextureDesc
 		, uint32		              ExtData      = 0
 		, ERHIAccess                  InitialState = ERHIAccess::Unknown
 		, FResourceBulkDataInterface* BulkData     = nullptr
-		, FRHIGPUMask                 GPUMask      = FRHIGPUMask()
+		, FRHIGPUMask                 GPUMask      = FRHIGPUMask::All()
 		)
 	{
 		return FRHITextureCreateDesc(
@@ -1572,7 +1572,7 @@ struct FRHITextureCreateDesc : public FRHITextureDesc
 		, uint32		              ExtData      = 0
 		, ERHIAccess                  InitialState = ERHIAccess::Unknown
 		, FResourceBulkDataInterface* BulkData     = nullptr
-		, FRHIGPUMask                 GPUMask      = FRHIGPUMask()
+		, FRHIGPUMask                 GPUMask      = FRHIGPUMask::All()
 		)
 	{
 		return FRHITextureCreateDesc(
@@ -1596,7 +1596,7 @@ struct FRHITextureCreateDesc : public FRHITextureDesc
 		, uint32		              ExtData      = 0
 		, ERHIAccess                  InitialState = ERHIAccess::Unknown
 		, FResourceBulkDataInterface* BulkData     = nullptr
-		, FRHIGPUMask                 GPUMask      = FRHIGPUMask()
+		, FRHIGPUMask                 GPUMask      = FRHIGPUMask::All()
 		)
 	{
 		return FRHITextureCreateDesc(
@@ -1625,7 +1625,7 @@ struct FRHITextureCreateDesc : public FRHITextureDesc
 		, ERHIAccess                  InInitialState
 		, TCHAR const*                InDebugName
 		, FResourceBulkDataInterface* InBulkData    = nullptr 
-		, FRHIGPUMask                 InGPUMask     = FRHIGPUMask()
+		, FRHIGPUMask                 InGPUMask     = FRHIGPUMask::All()
 		)
 		: FRHITextureDesc(InDimension, InFlags, InFormat, InClearValue, InExtent, InDepth, InArraySize, InNumMips, InNumSamples, InExtData)
 		, InitialState   (InInitialState)
@@ -1640,7 +1640,7 @@ struct FRHITextureCreateDesc : public FRHITextureDesc
 		, ERHIAccess                  InInitialState
 		, TCHAR const*                InDebugName
 		, FResourceBulkDataInterface* InBulkData     = nullptr
-		, FRHIGPUMask                 InGPUMask      = FRHIGPUMask()
+		, FRHIGPUMask                 InGPUMask      = FRHIGPUMask::All()
 		)
 		: FRHITextureDesc(InDesc)
 		, InitialState   (InInitialState)
