fix RHICreateTextureCubeArray

diff --git a/Engine/Source/Runtime/RHI/Public/RHICommandList.h b/Engine/Source/Runtime/RHI/Public/RHICommandList.h
--- a/Engine/Source/Runtime/RHI/Public/RHICommandList.h
+++ b/Engine/Source/Runtime/RHI/Public/RHICommandList.h
@@ -5533,7 +5533,7 @@ FORCEINLINE FTextureCubeRHIRef RHICreateTextureCube(uint32 Size, uint8 Format, u
 FORCEINLINE FTextureCubeRHIRef RHICreateTextureCubeArray(uint32 Size, uint32 ArraySize, uint8 Format, uint32 NumMips, ETextureCreateFlags Flags, ERHIAccess ResourceState, const FRHIResourceCreateInfo& CreateInfo)
 {
 	return RHICreateTexture(
-		FRHITextureCreateDesc::CreateCube(CreateInfo.DebugName)
+		FRHITextureCreateDesc::CreateCubeArray(CreateInfo.DebugName)
 			.SetExtent(Size)
 			.SetArraySize(ArraySize)
 			.SetFormat((EPixelFormat)Format)
