Texture RHI size query prevent check on textures that fail to build
don't call the underlying device RHICalcTexturePlatformSize if dimensions are zero
happens when texture build returns false

diff --git a/Engine/Source/Runtime/RHI/Public/DynamicRHI.h b/Engine/Source/Runtime/RHI/Public/DynamicRHI.h
--- a/Engine/Source/Runtime/RHI/Public/DynamicRHI.h
+++ b/Engine/Source/Runtime/RHI/Public/DynamicRHI.h
@@ -1303,6 +1303,13 @@ FORCEINLINE void RHIUpdateUniformBuffer(FRHIUniformBuffer* UniformBufferRHI, con
 
 FORCEINLINE FDynamicRHI::FRHICalcTextureSizeResult RHICalcTexturePlatformSize(FRHITextureDesc const& Desc, uint32 FirstMipIndex = 0)
 {
+	if ( ! Desc.IsValid() )
+	{
+		// Invalid texture desc; return zero to indicate failure
+		FDynamicRHI::FRHICalcTextureSizeResult ZeroResult = { 0 };
+		return ZeroResult;
+	}
+
 	return GDynamicRHI->RHICalcTexturePlatformSize(Desc, FirstMipIndex);
 }
 
