Converting RHICreateTargetableShaderResourceXXX to use the new API.
RHICreateTargetableShaderResource2D changes will come in a later CL.

diff --git a/Engine/Plugins/Runtime/WindowsMixedReality/Source/WindowsMixedRealityHMD/Private/WindowsMixedRealityHMD.cpp b/Engine/Plugins/Runtime/WindowsMixedReality/Source/WindowsMixedRealityHMD/Private/WindowsMixedRealityHMD.cpp
--- a/Engine/Plugins/Runtime/WindowsMixedReality/Source/WindowsMixedRealityHMD/Private/WindowsMixedRealityHMD.cpp
+++ b/Engine/Plugins/Runtime/WindowsMixedReality/Source/WindowsMixedRealityHMD/Private/WindowsMixedRealityHMD.cpp
@@ -1621,44 +1621,33 @@ namespace WindowsMixedReality
 			return false;
 		}
 
+		bool bMultiView = false;
+#if PLATFORM_HOLOLENS
+		bMultiView = (bIsMobileMultiViewEnabled && !HMD->IsThirdCameraActive());
+#endif
+
 		// Since our textures must be BGRA, this plugin did require a change to WindowsD3D11Device.cpp
 		// to add the D3D11_CREATE_DEVICE_BGRA_SUPPORT flag to the graphics device.
 		FRHIResourceCreateInfo CreateInfo(TEXT("FWindowsMixedRealityHMD"));
 
-#if PLATFORM_HOLOLENS
-		if (bIsMobileMultiViewEnabled && !HMD->IsThirdCameraActive())
-		{
-			FTexture2DArrayRHIRef texture, resource;
-			RHICreateTargetableShaderResource2DArray(
-				sizeX,
-				sizeY,
-				2,
-				PF_B8G8R8A8, // must be BGRA
-				numMips,
-				flags,
-				targetableTextureFlags,
-				CreateInfo,
-				texture,
-				resource);
-			outTargetableTexture = texture;
-			outShaderResourceTexture = resource;
+		FRHITextureCreateDesc Desc =
+			FRHITextureCreateDesc::Create2D(TEXT("FWindowsMixedRealityHMD"))
+			.SetExtent(sizeX, sizeY)
+			.SetFormat(PF_B8G8R8A8) // must be BGRA
+			.SetNumMips(numMips)
+			.SetFlags(flags);
 
-			CurrentBackBuffer = outTargetableTexture;
+		if (bMultiView)
+		{
+			Desc.SetDimension(ETextureDimension::Texture2DArray)
+				.SetArraySize(2);
 		}
-		else
-#endif
+
+		RHICreateTargetableShaderResource(Desc, targetableTextureFlags, outTargetableTexture, outShaderResourceTexture);
+
+		if (bMultiView)
 		{
-			RHICreateTargetableShaderResource2D(
-				sizeX,
-				sizeY,
-				PF_B8G8R8A8, // must be BGRA
-				numMips,
-				flags,
-				targetableTextureFlags,
-				false,
-				CreateInfo,
-				outTargetableTexture,
-				outShaderResourceTexture);
+			CurrentBackBuffer = outTargetableTexture;
 		}
 
 		bNeedReallocateDepthTexture = true;
@@ -1687,54 +1676,37 @@ namespace WindowsMixedReality
 			OutShaderResourceTexture = CurrentDepthBuffer;
 			return true;
 		}
-		
-		FRHIResourceCreateInfo CreateInfo(TEXT("FWindowsMixedRealityHMD"));
-		// This binding is necessary - without it there will be a runtime error.
-		// Current shader assumes far depth since scene depth uses far depth.
-		CreateInfo.ClearValueBinding = FClearValueBinding::DepthFar;
 
+		bool bMultiView = false;
 #if PLATFORM_HOLOLENS
-		if (bIsMobileMultiViewEnabled && !HMD->IsThirdCameraActive())
+		bMultiView = (bIsMobileMultiViewEnabled && !HMD->IsThirdCameraActive());
+#endif
+
+		FRHITextureCreateDesc Desc =
+			FRHITextureCreateDesc::Create2D(TEXT("FWindowsMixedRealityHMD"))
+			// Do not use input format - this will resolve to X32_TYPELESS_G8X24_UINT which cannot be used for a depthstencil buffer.
+			// DepthStencil will resolve to R32G8X24_TYPELESS which is usable for a depthstencil buffer!
+			.SetFormat(PF_DepthStencil) // must be BGRA
+			// Do not use input mips, this will resolve to 0 which will throw creating the texture.
+			.SetNumMips(1)
+			.SetFlags(InTexFlags)
+			// This binding is necessary - without it there will be a runtime error.
+			// Current shader assumes far depth since scene depth uses far depth.
+			.SetClearValue(FClearValueBinding::DepthFar);
+
+		if (bMultiView)
 		{
-			FIntPoint size = GetIdealRenderTargetSize();
-
-			FTexture2DArrayRHIRef texture, resource;
-			RHICreateTargetableShaderResource2DArray(
-				size.X,
-				size.Y,
-				2,
-				// Do not use input format - this will resolve to X32_TYPELESS_G8X24_UINT which cannot be used for a depthstencil buffer.
-				// DepthStencil will resolve to R32G8X24_TYPELESS which is usable for a depthstencil buffer!
-				PF_DepthStencil,
-				// Do not use input mips, this will resolve to 0 which will throw creating the texture.
-				1,
-				InTexFlags,
-				TargetableTextureFlags,
-				CreateInfo,
-				texture,
-				resource);
-			OutTargetableTexture = texture;
-			OutShaderResourceTexture = resource;
+			Desc.SetDimension(ETextureDimension::Texture2DArray)
+				.SetExtent(GetIdealRenderTargetSize())
+				.SetArraySize(2);
 		}
 		else
-#endif
 		{
-			RHICreateTargetableShaderResource2D(
-				SizeX,
-				SizeY,
-				// Do not use input format - this will resolve to X32_TYPELESS_G8X24_UINT which cannot be used for a depthstencil buffer.
-				// DepthStencil will resolve to R32G8X24_TYPELESS which is usable for a depthstencil buffer!
-				PF_DepthStencil,
-				// Do not use input mips, this will resolve to 0 which will throw creating the texture.
-				1,
-				InTexFlags,
-				TargetableTextureFlags,
-				false,
-				CreateInfo,
-				OutTargetableTexture,
-				OutShaderResourceTexture);
+			Desc.SetExtent(SizeX, SizeY);
 		}
 
+		RHICreateTargetableShaderResource(Desc, TargetableTextureFlags, OutTargetableTexture, OutShaderResourceTexture);
+
 		CurrentDepthBuffer = OutTargetableTexture;
 		bNeedReallocateDepthTexture = false;
 
diff --git a/Engine/Source/Runtime/Engine/Private/TextureRenderTarget2DArray.cpp b/Engine/Source/Runtime/Engine/Private/TextureRenderTarget2DArray.cpp
--- a/Engine/Source/Runtime/Engine/Private/TextureRenderTarget2DArray.cpp
+++ b/Engine/Source/Runtime/Engine/Private/TextureRenderTarget2DArray.cpp
@@ -233,29 +233,26 @@ void FTextureRenderTarget2DArrayResource::InitDynamicRHI()
 		}
 
 		// Create the RHI texture. Only one mip is used and the texture is targetable for resolve.
-		ETextureCreateFlags TexCreateFlags = bIsSRGB ? TexCreate_SRGB : TexCreate_None;
+		ETextureCreateFlags TexCreateFlags = bIsSRGB ? ETextureCreateFlags::SRGB : ETextureCreateFlags::None;
 		if (Owner->bCanCreateUAV)
 		{
-			TexCreateFlags |= TexCreate_UAV;
+			TexCreateFlags |= ETextureCreateFlags::UAV;
 		}
 
 		{
-			FRHIResourceCreateInfo CreateInfo(TEXT("FTextureRenderTarget2DArrayResource"), FClearValueBinding(Owner->ClearColor));
-			RHICreateTargetableShaderResource2DArray(
-				Owner->SizeX,
-				Owner->SizeY,
-				Owner->Slices,
-				Owner->GetFormat(),
-				Owner->GetNumMips(),
-				TexCreateFlags,
-				TexCreate_RenderTargetable,
-				CreateInfo,
-				RenderTarget2DArrayRHI,
-				Texture2DArrayRHI
-			);
+			const FRHITextureCreateDesc Desc =
+				FRHITextureCreateDesc::Create2DArray(TEXT("FTextureRenderTarget2DArrayResource"))
+				.SetExtent(Owner->SizeX, Owner->SizeY)
+				.SetArraySize(Owner->Slices)
+				.SetFormat(Owner->GetFormat())
+				.SetNumMips(Owner->GetNumMips())
+				.SetFlags(TexCreateFlags)
+				.SetClearValue(FClearValueBinding(Owner->ClearColor));
+
+			RHICreateTargetableShaderResource(Desc, ETextureCreateFlags::RenderTargetable, RenderTarget2DArrayRHI, Texture2DArrayRHI);
 		}
 
-		if (EnumHasAnyFlags(TexCreateFlags, TexCreate_UAV))
+		if (EnumHasAnyFlags(TexCreateFlags, ETextureCreateFlags::UAV))
 		{
 			UnorderedAccessViewRHI = RHICreateUnorderedAccessView(RenderTarget2DArrayRHI);
 		}
diff --git a/Engine/Source/Runtime/Engine/Private/TextureRenderTargetCube.cpp b/Engine/Source/Runtime/Engine/Private/TextureRenderTargetCube.cpp
--- a/Engine/Source/Runtime/Engine/Private/TextureRenderTargetCube.cpp
+++ b/Engine/Source/Runtime/Engine/Private/TextureRenderTargetCube.cpp
@@ -235,27 +235,25 @@ void FTextureRenderTargetCubeResource::InitDynamicRHI()
 		}
 
 		// Create the RHI texture. Only one mip is used and the texture is targetable for resolve.
-		ETextureCreateFlags TexCreateFlags = bIsSRGB ? TexCreate_SRGB : TexCreate_None;
+		ETextureCreateFlags TexCreateFlags = bIsSRGB ? ETextureCreateFlags::SRGB : ETextureCreateFlags::None;
 		if (Owner->bCanCreateUAV)
 		{
-			TexCreateFlags |= TexCreate_UAV;
+			TexCreateFlags |= ETextureCreateFlags::UAV;
 		}
 
 		{
-			FRHIResourceCreateInfo CreateInfo(TEXT("FTextureRenderTargetCubeResource"), FClearValueBinding(Owner->ClearColor));
-			RHICreateTargetableShaderResourceCube(
-				Owner->SizeX,
-				Owner->GetFormat(), 
-				Owner->GetNumMips(),
-				TexCreateFlags, 
-				TexCreate_RenderTargetable,
-				false,
-				CreateInfo,
-				RenderTargetCubeRHI,
-				TextureCubeRHI );
+			const FRHITextureCreateDesc Desc =
+				FRHITextureCreateDesc::CreateCube(TEXT("FTextureRenderTargetCubeResource"))
+				.SetExtent(Owner->SizeX)
+				.SetFormat(Owner->GetFormat())
+				.SetNumMips(Owner->GetNumMips())
+				.SetFlags(TexCreateFlags)
+				.SetClearValue(FClearValueBinding(Owner->ClearColor));
+
+			RHICreateTargetableShaderResource(Desc, ETextureCreateFlags::RenderTargetable, RenderTargetCubeRHI, TextureCubeRHI);
 		}
 
-		if (EnumHasAnyFlags(TexCreateFlags, TexCreate_UAV))
+		if (EnumHasAnyFlags(TexCreateFlags, ETextureCreateFlags::UAV))
 		{
 			UnorderedAccessViewRHI = RHICreateUnorderedAccessView(RenderTargetCubeRHI);
 		}
@@ -272,7 +270,7 @@ void FTextureRenderTargetCubeResource::InitDynamicRHI()
 				Owner->GetFormat(),
 				Owner->GetNumMips(), 
 				/* NumSamples =*/ 1,
-				TexCreate_RenderTargetable|TexCreateFlags,
+				ETextureCreateFlags::RenderTargetable|TexCreateFlags,
 				CreateInfo
 				);
 			SetGPUMask(CreateInfo.GPUMask);
diff --git a/Engine/Source/Runtime/Engine/Private/TextureRenderTargetVolume.cpp b/Engine/Source/Runtime/Engine/Private/TextureRenderTargetVolume.cpp
--- a/Engine/Source/Runtime/Engine/Private/TextureRenderTargetVolume.cpp
+++ b/Engine/Source/Runtime/Engine/Private/TextureRenderTargetVolume.cpp
@@ -247,20 +247,16 @@ void FTextureRenderTargetVolumeResource::InitDynamicRHI()
 		}
 
 		{
-			FRHIResourceCreateInfo CreateInfo(TEXT("FTextureRenderTargetVolumeResource"), FClearValueBinding(Owner->ClearColor));
-			RHICreateTargetableShaderResource3D(
-				Owner->SizeX,
-				Owner->SizeY,
-				Owner->SizeZ,
-				Owner->GetFormat(),
-				Owner->GetNumMips(),
-				TexCreateFlags, 
-				TexCreate_RenderTargetable,
-				false,
-				CreateInfo,
-				RenderTargetVolumeRHI,
-				TextureVolumeRHI
-			);
+			const FRHITextureCreateDesc Desc =
+				FRHITextureCreateDesc::Create3D(TEXT("FTextureRenderTargetVolumeResource"))
+				.SetExtent(Owner->SizeX, Owner->SizeY)
+				.SetDepth(Owner->SizeZ)
+				.SetFormat(Owner->GetFormat())
+				.SetNumMips(Owner->GetNumMips())
+				.SetFlags(TexCreateFlags)
+				.SetClearValue(FClearValueBinding(Owner->ClearColor));
+
+			RHICreateTargetableShaderResource(Desc, ETextureCreateFlags::RenderTargetable, RenderTargetVolumeRHI, TextureVolumeRHI);
 		}
 
 		if (EnumHasAnyFlags(TexCreateFlags, TexCreate_UAV))
diff --git a/Engine/Source/Runtime/Engine/Private/UnrealClient.cpp b/Engine/Source/Runtime/Engine/Private/UnrealClient.cpp
--- a/Engine/Source/Runtime/Engine/Private/UnrealClient.cpp
+++ b/Engine/Source/Runtime/Engine/Private/UnrealClient.cpp
@@ -268,7 +268,7 @@ float FRenderTarget::GetDisplayGamma() const
 * Accessor for the surface RHI when setting this render target
 * @return render target surface RHI resource
 */
-const FTexture2DRHIRef& FRenderTarget::GetRenderTargetTexture() const
+const FTextureRHIRef& FRenderTarget::GetRenderTargetTexture() const
 {
 	return RenderTargetTextureRHI;
 }
@@ -2097,12 +2097,22 @@ void FViewport::FHitProxyMap::Init(uint32 NewSizeX,uint32 NewSizeY)
 
 	// Create a render target to store the hit proxy map.
 	{
-		FRHIResourceCreateInfo CreateInfo(TEXT("HitProxyTexture"), FClearValueBinding::White);
-		RHICreateTargetableShaderResource2D(SizeX,SizeY,PF_B8G8R8A8,1,TexCreate_None,TexCreate_RenderTargetable,false,CreateInfo,RenderTargetTextureRHI,HitProxyTexture);
+		const FRHITextureCreateDesc Desc =
+			FRHITextureCreateDesc::Create2D(TEXT("HitProxyTexture"))
+			.SetExtent(SizeX, SizeY)
+			.SetFormat(PF_B8G8R8A8)
+			.SetClearValue(FClearValueBinding::White);
+
+		RHICreateTargetableShaderResource(Desc, ETextureCreateFlags::RenderTargetable, RenderTargetTextureRHI, HitProxyTexture);
 	}
 	{
-		FRHIResourceCreateInfo CreateInfo(TEXT("HitProxyCPUTexture"));
-		HitProxyCPUTexture = RHICreateTexture2D(SizeX, SizeY, PF_B8G8R8A8,1,1,TexCreate_CPUReadback,CreateInfo);
+		const FRHITextureCreateDesc Desc =
+			FRHITextureCreateDesc::Create2D(TEXT("HitProxyCPUTexture"))
+			.SetExtent(SizeX, SizeY)
+			.SetFormat(PF_B8G8R8A8)
+			.SetFlags(ETextureCreateFlags::CPUReadback);
+
+		HitProxyCPUTexture = RHICreateTexture(Desc);
 	}
 }
 
@@ -2371,3 +2381,13 @@ FDummyViewport::~FDummyViewport()
 		DebugCanvas = NULL;
 	}
 }
+
+void FDummyViewport::InitDynamicRHI()
+{
+	const FRHITextureCreateDesc Desc =
+		FRHITextureCreateDesc::Create2D(TEXT("FDummyViewport"))
+		.SetExtent(SizeX, SizeY)
+		.SetFormat(PF_A2B10G10R10);
+
+	RHICreateTargetableShaderResource(Desc, ETextureCreateFlags::RenderTargetable, RenderTargetTextureRHI);
+}
diff --git a/Engine/Source/Runtime/Engine/Public/PlanarReflectionSceneProxy.h b/Engine/Source/Runtime/Engine/Public/PlanarReflectionSceneProxy.h
--- a/Engine/Source/Runtime/Engine/Public/PlanarReflectionSceneProxy.h
+++ b/Engine/Source/Runtime/Engine/Public/PlanarReflectionSceneProxy.h
@@ -43,22 +43,13 @@ public:
 		);
 		SamplerStateRHI = GetOrCreateSamplerState(SamplerStateInitializer);
 
-		FTexture2DRHIRef Texture2DRHI;
-		FRHIResourceCreateInfo CreateInfo = { TEXT("FPlanarReflectionRenderTarget"), FClearValueBinding(FLinearColor::Black) };
-
-		RHICreateTargetableShaderResource2D(
-			GetSizeX(), 
-			GetSizeY(), 
-			PF_FloatRGBA,
-			1,
-			TexCreate_None,
-			TexCreate_RenderTargetable,
-			false,
-			CreateInfo,
-			RenderTargetTextureRHI,
-			Texture2DRHI
-			);
-		TextureRHI = (FTextureRHIRef&)Texture2DRHI;
+		const FRHITextureCreateDesc Desc =
+			FRHITextureCreateDesc::Create2D(TEXT("FPlanarReflectionRenderTarget"))
+			.SetExtent(GetSizeXY())
+			.SetFormat(PF_FloatRGBA)
+			.SetClearValue(FClearValueBinding::Black);
+
+		RHICreateTargetableShaderResource(Desc, ETextureCreateFlags::RenderTargetable, RenderTargetTextureRHI, TextureRHI);
 	}
 
 	virtual FIntPoint GetSizeXY() const { return Size; }
diff --git a/Engine/Source/Runtime/Engine/Public/UnrealClient.h b/Engine/Source/Runtime/Engine/Public/UnrealClient.h
--- a/Engine/Source/Runtime/Engine/Public/UnrealClient.h
+++ b/Engine/Source/Runtime/Engine/Public/UnrealClient.h
@@ -44,7 +44,7 @@ public:
 	* Accessor for the surface RHI when setting this render target
 	* @return render target surface RHI resource
 	*/
-	ENGINE_API virtual const FTexture2DRHIRef& GetRenderTargetTexture() const;
+	ENGINE_API virtual const FTextureRHIRef& GetRenderTargetTexture() const;
 	ENGINE_API virtual FUnorderedAccessViewRHIRef GetRenderTargetUAV() const;
 
 	/**
@@ -111,7 +111,7 @@ public:
 
 protected:
 
-	FTexture2DRHIRef RenderTargetTextureRHI;
+	FTextureRHIRef RenderTargetTextureRHI;
 
 	/**
 	 * Reads the viewport's displayed pixels into a preallocated color buffer.
@@ -1215,13 +1215,7 @@ public:
 	//~ End FViewport Interface
 
 	//~ Begin FRenderResource Interface
-	virtual void InitDynamicRHI() override
-	{
-		FTexture2DRHIRef ShaderResourceTextureRHI;
-
-		FRHIResourceCreateInfo CreateInfo(TEXT("FDummyViewport"));
-		RHICreateTargetableShaderResource2D( SizeX, SizeY, PF_A2B10G10R10, 1, TexCreate_None, TexCreate_RenderTargetable, false, CreateInfo, RenderTargetTextureRHI, ShaderResourceTextureRHI );
-	}
+	virtual void InitDynamicRHI() override;
 
 	// @todo DLL: Without these functions we get unresolved linker errors with FRenderResource
 	virtual void InitRHI() override{}
diff --git a/Engine/Source/Runtime/RHI/Private/RHIUtilities.cpp b/Engine/Source/Runtime/RHI/Private/RHIUtilities.cpp
--- a/Engine/Source/Runtime/RHI/Private/RHIUtilities.cpp
+++ b/Engine/Source/Runtime/RHI/Private/RHIUtilities.cpp
@@ -604,3 +604,260 @@ RHI_API ERHIAccess RHIGetDefaultResourceState(EBufferUsageFlags InUsage, bool bI
 	return ResourceState;
 }
 
+void RHICreateTargetableShaderResource(
+	const FRHITextureCreateDesc& BaseDesc,
+	ETextureCreateFlags TargetableTextureFlags,
+	FTextureRHIRef& OutTargetableTexture)
+{
+	// Ensure none of the usage flags are passed in.
+	check(!EnumHasAnyFlags(BaseDesc.Flags, ETextureCreateFlags::RenderTargetable | ETextureCreateFlags::ResolveTargetable));
+
+	// Ensure that the targetable texture is either render or depth-stencil targetable.
+	check(EnumHasAnyFlags(TargetableTextureFlags, ETextureCreateFlags::RenderTargetable | ETextureCreateFlags::DepthStencilTargetable | ETextureCreateFlags::UAV));
+
+	FRHITextureCreateDesc Desc =
+		FRHITextureCreateDesc(BaseDesc)
+		.SetFlags(TargetableTextureFlags | ETextureCreateFlags::ShaderResource)
+		.SetInitialState(ERHIAccess::SRVMask);
+
+	OutTargetableTexture = RHICreateTexture(Desc);
+}
+
+void RHICreateTargetableShaderResource(
+	const FRHITextureCreateDesc& BaseDesc,
+	ETextureCreateFlags TargetableTextureFlags,
+	bool bForceSeparateTargetAndShaderResource,
+	bool bForceSharedTargetAndShaderResource,
+	FTextureRHIRef& OutTargetableTexture,
+	FTextureRHIRef& OutShaderResourceTexture
+	)
+{
+	// Ensure none of the usage flags are passed in.
+	check(!EnumHasAnyFlags(BaseDesc.Flags, ETextureCreateFlags::RenderTargetable | ETextureCreateFlags::ResolveTargetable));
+
+	// Ensure we aren't forcing separate and shared textures at the same time.
+	check(!(bForceSeparateTargetAndShaderResource && bForceSharedTargetAndShaderResource));
+
+	// Ensure that the targetable texture is either render or depth-stencil targetable.
+	check(EnumHasAnyFlags(TargetableTextureFlags, ETextureCreateFlags::RenderTargetable | ETextureCreateFlags::DepthStencilTargetable | ETextureCreateFlags::UAV));
+
+	if (BaseDesc.NumSamples > 1 && !bForceSharedTargetAndShaderResource)
+	{
+		bForceSeparateTargetAndShaderResource = RHISupportsSeparateMSAAAndResolveTextures(GMaxRHIShaderPlatform);
+	}
+
+	if (!bForceSeparateTargetAndShaderResource)
+	{
+		FRHITextureCreateDesc Desc =
+			FRHITextureCreateDesc(BaseDesc)
+			.AddFlags(TargetableTextureFlags | ETextureCreateFlags::ShaderResource)
+			.SetInitialState(ERHIAccess::SRVMask);
+
+		// Create a single texture that has both TargetableTextureFlags and ETextureCreateFlags::ShaderResource set.
+		OutTargetableTexture = OutShaderResourceTexture = RHICreateTexture(Desc);
+	}
+	else
+	{
+		ETextureCreateFlags ResolveTargetableTextureFlags = ETextureCreateFlags::ResolveTargetable;
+		if (EnumHasAnyFlags(TargetableTextureFlags, ETextureCreateFlags::DepthStencilTargetable))
+		{
+			ResolveTargetableTextureFlags |= ETextureCreateFlags::DepthStencilResolveTarget;
+		}
+
+		FRHITextureCreateDesc TargetableDesc =
+			FRHITextureCreateDesc(BaseDesc)
+			.AddFlags(TargetableTextureFlags)
+			.SetInitialState(ERHIAccess::SRVMask);
+
+		FRHITextureCreateDesc ResourceDesc =
+			FRHITextureCreateDesc(BaseDesc)
+			.AddFlags(ResolveTargetableTextureFlags | ETextureCreateFlags::ShaderResource)
+			.SetInitialState(ERHIAccess::SRVMask)
+			.SetNumSamples(1);
+
+		// Create a texture that has TargetableTextureFlags set, and a second texture that has ETextureCreateFlags::ResolveTargetable and ETextureCreateFlags::ShaderResource set.
+		OutTargetableTexture = RHICreateTexture(TargetableDesc);
+		OutShaderResourceTexture = RHICreateTexture(ResourceDesc);
+	}
+}
+
+void RHICreateTargetableShaderResource2D(
+	uint32 SizeX,
+	uint32 SizeY,
+	uint8 Format,
+	uint32 NumMips,
+	ETextureCreateFlags Flags,
+	ETextureCreateFlags TargetableTextureFlags,
+	bool bForceSeparateTargetAndShaderResource,
+	bool bForceSharedTargetAndShaderResource,
+	const FRHIResourceCreateInfo& CreateInfo,
+	FTextureRHIRef& OutTargetableTexture,
+	FTextureRHIRef& OutShaderResourceTexture,
+	uint32 NumSamples
+)
+{
+	FRHITextureCreateDesc Desc =
+		FRHITextureCreateDesc::Create2D(CreateInfo.DebugName)
+		.SetExtent(SizeX, SizeY)
+		.SetFormat((EPixelFormat)Format)
+		.SetNumMips(NumMips)
+		.SetNumSamples(NumSamples)
+		.SetFlags(Flags)
+		.SetBulkData(CreateInfo.BulkData)
+		.SetClearValue(CreateInfo.ClearValueBinding)
+		.SetExtData(CreateInfo.ExtData)
+		.SetGPUMask(CreateInfo.GPUMask);
+
+	RHICreateTargetableShaderResource(Desc, TargetableTextureFlags, bForceSeparateTargetAndShaderResource, bForceSharedTargetAndShaderResource, OutTargetableTexture, OutShaderResourceTexture);
+}
+
+void RHICreateTargetableShaderResource2D(
+	uint32 SizeX,
+	uint32 SizeY,
+	uint8 Format,
+	uint32 NumMips,
+	ETextureCreateFlags Flags,
+	ETextureCreateFlags TargetableTextureFlags,
+	bool bForceSeparateTargetAndShaderResource,
+	const FRHIResourceCreateInfo& CreateInfo,
+	FTexture2DRHIRef& OutTargetableTexture,
+	FTexture2DRHIRef& OutShaderResourceTexture,
+	uint32 NumSamples)
+{
+PRAGMA_DISABLE_DEPRECATION_WARNINGS
+	RHICreateTargetableShaderResource2D(SizeX, SizeY, Format, NumMips, Flags, TargetableTextureFlags, bForceSeparateTargetAndShaderResource, false, CreateInfo, OutTargetableTexture, OutShaderResourceTexture, NumSamples);
+PRAGMA_ENABLE_DEPRECATION_WARNINGS
+}
+
+void RHICreateTargetableShaderResource2DArray(
+	uint32 SizeX,
+	uint32 SizeY,
+	uint32 SizeZ,
+	uint8 Format,
+	uint32 NumMips,
+	ETextureCreateFlags Flags,
+	ETextureCreateFlags TargetableTextureFlags,
+	bool bForceSeparateTargetAndShaderResource,
+	bool bForceSharedTargetAndShaderResource,
+	const FRHIResourceCreateInfo& CreateInfo,
+	FTextureRHIRef& OutTargetableTexture,
+	FTextureRHIRef& OutShaderResourceTexture,
+	uint32 NumSamples
+)
+{
+	FRHITextureCreateDesc Desc =
+		FRHITextureCreateDesc::Create2DArray(CreateInfo.DebugName)
+		.SetExtent(SizeX, SizeY)
+		.SetArraySize(SizeZ)
+		.SetFormat((EPixelFormat)Format)
+		.SetNumMips(NumMips)
+		.SetNumSamples(NumSamples)
+		.SetFlags(Flags)
+		.SetBulkData(CreateInfo.BulkData)
+		.SetClearValue(CreateInfo.ClearValueBinding)
+		.SetExtData(CreateInfo.ExtData)
+		.SetGPUMask(CreateInfo.GPUMask);
+
+	RHICreateTargetableShaderResource(Desc, TargetableTextureFlags, bForceSeparateTargetAndShaderResource, bForceSharedTargetAndShaderResource, OutTargetableTexture, OutShaderResourceTexture);
+}
+
+void RHICreateTargetableShaderResource2DArray(
+	uint32 SizeX,
+	uint32 SizeY,
+	uint32 SizeZ,
+	uint8 Format,
+	uint32 NumMips,
+	ETextureCreateFlags Flags,
+	ETextureCreateFlags TargetableTextureFlags,
+	const FRHIResourceCreateInfo& CreateInfo,
+	FTextureRHIRef& OutTargetableTexture,
+	FTextureRHIRef& OutShaderResourceTexture,
+	uint32 NumSamples)
+{
+PRAGMA_DISABLE_DEPRECATION_WARNINGS
+	RHICreateTargetableShaderResource2DArray(SizeX, SizeY, SizeZ, Format, NumMips, Flags, TargetableTextureFlags, false, false, CreateInfo, OutTargetableTexture, OutShaderResourceTexture, NumSamples);
+PRAGMA_ENABLE_DEPRECATION_WARNINGS
+}
+
+void RHICreateTargetableShaderResourceCube(
+	uint32 LinearSize,
+	uint8 Format,
+	uint32 NumMips,
+	ETextureCreateFlags Flags,
+	ETextureCreateFlags TargetableTextureFlags,
+	bool bForceSeparateTargetAndShaderResource,
+	const FRHIResourceCreateInfo& CreateInfo,
+	FTextureRHIRef& OutTargetableTexture,
+	FTextureRHIRef& OutShaderResourceTexture
+)
+{
+	FRHITextureCreateDesc Desc =
+		FRHITextureCreateDesc::CreateCube(CreateInfo.DebugName)
+		.SetExtent(LinearSize)
+		.SetFormat((EPixelFormat)Format)
+		.SetNumMips(NumMips)
+		.SetFlags(Flags)
+		.SetBulkData(CreateInfo.BulkData)
+		.SetClearValue(CreateInfo.ClearValueBinding)
+		.SetExtData(CreateInfo.ExtData)
+		.SetGPUMask(CreateInfo.GPUMask);
+
+	RHICreateTargetableShaderResource(Desc, TargetableTextureFlags, bForceSeparateTargetAndShaderResource, false, OutTargetableTexture, OutShaderResourceTexture);
+}
+
+void RHICreateTargetableShaderResourceCubeArray(
+	uint32 LinearSize,
+	uint32 ArraySize,
+	uint8 Format,
+	uint32 NumMips,
+	ETextureCreateFlags Flags,
+	ETextureCreateFlags TargetableTextureFlags,
+	bool bForceSeparateTargetAndShaderResource,
+	const FRHIResourceCreateInfo& CreateInfo,
+	FTextureRHIRef& OutTargetableTexture,
+	FTextureRHIRef& OutShaderResourceTexture
+)
+{
+	FRHITextureCreateDesc Desc =
+		FRHITextureCreateDesc::CreateCubeArray(CreateInfo.DebugName)
+		.SetExtent(LinearSize)
+		.SetArraySize(ArraySize)
+		.SetFormat((EPixelFormat)Format)
+		.SetNumMips(NumMips)
+		.SetFlags(Flags)
+		.SetBulkData(CreateInfo.BulkData)
+		.SetClearValue(CreateInfo.ClearValueBinding)
+		.SetExtData(CreateInfo.ExtData)
+		.SetGPUMask(CreateInfo.GPUMask);
+
+	RHICreateTargetableShaderResource(Desc, TargetableTextureFlags, bForceSeparateTargetAndShaderResource, false, OutTargetableTexture, OutShaderResourceTexture);
+}
+
+void RHICreateTargetableShaderResource3D(
+	uint32 SizeX,
+	uint32 SizeY,
+	uint32 SizeZ,
+	uint8 Format,
+	uint32 NumMips,
+	ETextureCreateFlags Flags,
+	ETextureCreateFlags TargetableTextureFlags,
+	bool bForceSeparateTargetAndShaderResource,
+	const FRHIResourceCreateInfo& CreateInfo,
+	FTextureRHIRef& OutTargetableTexture,
+	FTextureRHIRef& OutShaderResourceTexture
+)
+{
+	FRHITextureCreateDesc Desc =
+		FRHITextureCreateDesc::Create3D(CreateInfo.DebugName)
+		.SetExtent(SizeX, SizeY)
+		.SetDepth(SizeZ)
+		.SetFormat((EPixelFormat)Format)
+		.SetNumMips(NumMips)
+		.SetFlags(Flags)
+		.SetBulkData(CreateInfo.BulkData)
+		.SetClearValue(CreateInfo.ClearValueBinding)
+		.SetExtData(CreateInfo.ExtData)
+		.SetGPUMask(CreateInfo.GPUMask);
+
+	RHICreateTargetableShaderResource(Desc, TargetableTextureFlags, bForceSeparateTargetAndShaderResource, false, OutTargetableTexture, OutShaderResourceTexture);
+}
diff --git a/Engine/Source/Runtime/RHI/Public/RHIUtilities.h b/Engine/Source/Runtime/RHI/Public/RHIUtilities.h
--- a/Engine/Source/Runtime/RHI/Public/RHIUtilities.h
+++ b/Engine/Source/Runtime/RHI/Public/RHIUtilities.h
@@ -674,6 +674,46 @@ inline void TransitionRenderPassTargets(FRHICommandList& RHICmdList, const FRHIR
 	RHICmdList.Transition(MakeArrayView(Transitions, TransitionIndex));
 }
 
+RHI_API void RHICreateTargetableShaderResource(
+	const FRHITextureCreateDesc& BaseDesc,
+	ETextureCreateFlags TargetableTextureFlags,
+	FTextureRHIRef& OutTargetableTexture);
+
+/**
+ * Creates 1 or 2 textures with the same dimensions/format.
+ * If the RHI supports textures that can be used as both shader resources and render targets,
+ * and bForceSeparateTargetAndShaderResource=false, then a single texture is created.
+ * Otherwise two textures are created, one of them usable as a shader resource and resolve target, and one of them usable as a render target.
+ * Two texture references are always returned, but they may reference the same texture.
+ * If two different textures are returned, the render-target texture must be manually copied to the shader-resource texture.
+ */
+RHI_API void RHICreateTargetableShaderResource(
+	const FRHITextureCreateDesc& Desc,
+	ETextureCreateFlags TargetableTextureFlags,
+	bool bForceSeparateTargetAndShaderResource,
+	bool bForceSharedTargetAndShaderResource,
+	FTextureRHIRef& OutTargetableTexture,
+	FTextureRHIRef& OutShaderResourceTexture);
+
+inline void RHICreateTargetableShaderResource(
+	const FRHITextureCreateDesc& Desc,
+	ETextureCreateFlags TargetableTextureFlags,
+	bool bForceSeparateTargetAndShaderResource,
+	FTextureRHIRef& OutTargetableTexture,
+	FTextureRHIRef& OutShaderResourceTexture)
+{
+	RHICreateTargetableShaderResource(Desc, TargetableTextureFlags, bForceSeparateTargetAndShaderResource, false, OutTargetableTexture, OutShaderResourceTexture);
+}
+
+inline void RHICreateTargetableShaderResource(
+	const FRHITextureCreateDesc& Desc,
+	ETextureCreateFlags TargetableTextureFlags,
+	FTextureRHIRef& OutTargetableTexture,
+	FTextureRHIRef& OutShaderResourceTexture)
+{
+	RHICreateTargetableShaderResource(Desc, TargetableTextureFlags, false, false, OutTargetableTexture, OutShaderResourceTexture);
+}
+
 /**
  * Creates 1 or 2 textures with the same dimensions/format.
  * If the RHI supports textures that can be used as both shader resources and render targets,
@@ -682,55 +722,23 @@ inline void TransitionRenderPassTargets(FRHICommandList& RHICmdList, const FRHIR
  * Two texture references are always returned, but they may reference the same texture.
  * If two different textures are returned, the render-target texture must be manually copied to the shader-resource texture.
  */
-inline void RHICreateTargetableShaderResource2D(
+//UE_DEPRECATED(5.1, "RHICreateTargetableShaderResource2D is deprecated, RHICreateTargetableShaderResource with a FRHITextureCreateDesc should be used instead")
+RHI_API void RHICreateTargetableShaderResource2D(
 	uint32 SizeX,
 	uint32 SizeY,
-	uint8 Format,	
+	uint8 Format,
 	uint32 NumMips,
 	ETextureCreateFlags Flags,
 	ETextureCreateFlags TargetableTextureFlags,
 	bool bForceSeparateTargetAndShaderResource,
 	bool bForceSharedTargetAndShaderResource,
-	FRHIResourceCreateInfo& CreateInfo,
-	FTexture2DRHIRef& OutTargetableTexture,
-	FTexture2DRHIRef& OutShaderResourceTexture,
-	uint32 NumSamples = 1
-)
-{
-	// Ensure none of the usage flags are passed in.
-	check(!(Flags & TexCreate_RenderTargetable));
-	check(!(Flags & TexCreate_ResolveTargetable));
-
-	// Ensure we aren't forcing separate and shared textures at the same time.
-	check(!(bForceSeparateTargetAndShaderResource && bForceSharedTargetAndShaderResource));
-
-	// Ensure that the targetable texture is either render or depth-stencil targetable.
-	check(EnumHasAnyFlags(TargetableTextureFlags, TexCreate_RenderTargetable | TexCreate_DepthStencilTargetable | TexCreate_UAV));
+	const FRHIResourceCreateInfo& CreateInfo,
+	FTextureRHIRef& OutTargetableTexture,
+	FTextureRHIRef& OutShaderResourceTexture,
+	uint32 NumSamples = 1);
 
-	if (NumSamples > 1 && !bForceSharedTargetAndShaderResource)
-	{
-		bForceSeparateTargetAndShaderResource = RHISupportsSeparateMSAAAndResolveTextures(GMaxRHIShaderPlatform);
-	}
-
-	if (!bForceSeparateTargetAndShaderResource)
-	{
-		// Create a single texture that has both TargetableTextureFlags and TexCreate_ShaderResource set.
-		OutTargetableTexture = OutShaderResourceTexture = RHICreateTexture2D(SizeX, SizeY, Format, NumMips, NumSamples, Flags | TargetableTextureFlags | TexCreate_ShaderResource, ERHIAccess::SRVMask, CreateInfo);
-	}
-	else
-	{
-		ETextureCreateFlags ResolveTargetableTextureFlags = TexCreate_ResolveTargetable;
-		if (EnumHasAnyFlags(TargetableTextureFlags, TexCreate_DepthStencilTargetable))
-		{
-			ResolveTargetableTextureFlags |= TexCreate_DepthStencilResolveTarget;
-		}
-		// Create a texture that has TargetableTextureFlags set, and a second texture that has TexCreate_ResolveTargetable and TexCreate_ShaderResource set.
-		OutTargetableTexture = RHICreateTexture2D(SizeX, SizeY, Format, NumMips, NumSamples, Flags | TargetableTextureFlags, ERHIAccess::SRVMask, CreateInfo);
-		OutShaderResourceTexture = RHICreateTexture2D(SizeX, SizeY, Format, NumMips, 1, Flags | ResolveTargetableTextureFlags | TexCreate_ShaderResource, ERHIAccess::SRVMask, CreateInfo);
-	}
-}
-
-inline void RHICreateTargetableShaderResource2D(
+//UE_DEPRECATED(5.1, "RHICreateTargetableShaderResource2D is deprecated, RHICreateTargetableShaderResource with a FRHITextureCreateDesc should be used instead")
+RHI_API void RHICreateTargetableShaderResource2D(
 	uint32 SizeX,
 	uint32 SizeY,
 	uint8 Format,
@@ -738,15 +746,13 @@ inline void RHICreateTargetableShaderResource2D(
 	ETextureCreateFlags Flags,
 	ETextureCreateFlags TargetableTextureFlags,
 	bool bForceSeparateTargetAndShaderResource,
-	FRHIResourceCreateInfo& CreateInfo,
-	FTexture2DRHIRef& OutTargetableTexture,
-	FTexture2DRHIRef& OutShaderResourceTexture,
-	uint32 NumSamples = 1)
-{
-	RHICreateTargetableShaderResource2D(SizeX, SizeY, Format, NumMips, Flags, TargetableTextureFlags, bForceSeparateTargetAndShaderResource, false, CreateInfo, OutTargetableTexture, OutShaderResourceTexture, NumSamples);
-}
+	const FRHIResourceCreateInfo& CreateInfo,
+	FTextureRHIRef& OutTargetableTexture,
+	FTextureRHIRef& OutShaderResourceTexture,
+	uint32 NumSamples = 1);
 
-inline void RHICreateTargetableShaderResource2DArray(
+UE_DEPRECATED(5.1, "RHICreateTargetableShaderResource2DArray is deprecated, RHICreateTargetableShaderResource with a FRHITextureCreateDesc should be used instead")
+RHI_API void RHICreateTargetableShaderResource2DArray(
 	uint32 SizeX,
 	uint32 SizeY,
 	uint32 SizeZ,
@@ -756,46 +762,13 @@ inline void RHICreateTargetableShaderResource2DArray(
 	ETextureCreateFlags TargetableTextureFlags,
 	bool bForceSeparateTargetAndShaderResource,
 	bool bForceSharedTargetAndShaderResource,
-	FRHIResourceCreateInfo& CreateInfo,
-	FTexture2DArrayRHIRef& OutTargetableTexture,
-	FTexture2DArrayRHIRef& OutShaderResourceTexture,
-	uint32 NumSamples = 1
-)
-{
-	// Ensure none of the usage flags are passed in.
-	check(!(Flags & TexCreate_RenderTargetable));
-	check(!(Flags & TexCreate_ResolveTargetable));
-
-	// Ensure we aren't forcing separate and shared textures at the same time.
-	check(!(bForceSeparateTargetAndShaderResource && bForceSharedTargetAndShaderResource));
+	const FRHIResourceCreateInfo& CreateInfo,
+	FTextureRHIRef& OutTargetableTexture,
+	FTextureRHIRef& OutShaderResourceTexture,
+	uint32 NumSamples = 1);
 
-	// Ensure that the targetable texture is either render or depth-stencil targetable.
-	check(EnumHasAnyFlags(TargetableTextureFlags, TexCreate_RenderTargetable | TexCreate_DepthStencilTargetable | TexCreate_UAV));
-
-	if (NumSamples > 1 && !bForceSharedTargetAndShaderResource)
-	{
-		bForceSeparateTargetAndShaderResource = RHISupportsSeparateMSAAAndResolveTextures(GMaxRHIShaderPlatform);
-	}
-
-	if (!bForceSeparateTargetAndShaderResource)
-	{
-		// Create a single texture that has both TargetableTextureFlags and TexCreate_ShaderResource set.
-		OutTargetableTexture = OutShaderResourceTexture = RHICreateTexture2DArray(SizeX, SizeY, SizeZ, Format, NumMips, NumSamples, Flags | TargetableTextureFlags | TexCreate_ShaderResource, CreateInfo);
-	}
-	else
-	{
-		ETextureCreateFlags ResolveTargetableTextureFlags = TexCreate_ResolveTargetable;
-		if (EnumHasAnyFlags(TargetableTextureFlags, TexCreate_DepthStencilTargetable))
-		{
-			ResolveTargetableTextureFlags |= TexCreate_DepthStencilResolveTarget;
-		}
-		// Create a texture that has TargetableTextureFlags set, and a second texture that has TexCreate_ResolveTargetable and TexCreate_ShaderResource set.
-		OutTargetableTexture = RHICreateTexture2DArray(SizeX, SizeY, SizeZ, Format, NumMips, NumSamples, Flags | TargetableTextureFlags, CreateInfo);
-		OutShaderResourceTexture = RHICreateTexture2DArray(SizeX, SizeY, SizeZ, Format, NumMips, 1,  Flags | ResolveTargetableTextureFlags | TexCreate_ShaderResource, CreateInfo);
-	}
-}
-
-inline void RHICreateTargetableShaderResource2DArray(
+UE_DEPRECATED(5.1, "RHICreateTargetableShaderResource2DArray is deprecated, RHICreateTargetableShaderResource with a FRHITextureCreateDesc should be used instead")
+RHI_API void RHICreateTargetableShaderResource2DArray(
 	uint32 SizeX,
 	uint32 SizeY,
 	uint32 SizeZ,
@@ -803,13 +776,10 @@ inline void RHICreateTargetableShaderResource2DArray(
 	uint32 NumMips,
 	ETextureCreateFlags Flags,
 	ETextureCreateFlags TargetableTextureFlags,
-	FRHIResourceCreateInfo& CreateInfo,
-	FTexture2DArrayRHIRef& OutTargetableTexture,
-	FTexture2DArrayRHIRef& OutShaderResourceTexture,
-	uint32 NumSamples = 1)
-{
-	RHICreateTargetableShaderResource2DArray(SizeX, SizeY, SizeZ, Format, NumMips, Flags, TargetableTextureFlags, false, false, CreateInfo, OutTargetableTexture, OutShaderResourceTexture, NumSamples);
-}
+	const FRHIResourceCreateInfo& CreateInfo,
+	FTextureRHIRef& OutTargetableTexture,
+	FTextureRHIRef& OutShaderResourceTexture,
+	uint32 NumSamples = 1);
 
 /**
  * Creates 1 or 2 textures with the same dimensions/format.
@@ -819,37 +789,17 @@ inline void RHICreateTargetableShaderResource2DArray(
  * Two texture references are always returned, but they may reference the same texture.
  * If two different textures are returned, the render-target texture must be manually copied to the shader-resource texture.
  */
-inline void RHICreateTargetableShaderResourceCube(
+UE_DEPRECATED(5.1, "RHICreateTargetableShaderResourceCube is deprecated, RHICreateTargetableShaderResource with a FRHITextureCreateDesc should be used instead")
+RHI_API void RHICreateTargetableShaderResourceCube(
 	uint32 LinearSize,
 	uint8 Format,
 	uint32 NumMips,
 	ETextureCreateFlags Flags,
 	ETextureCreateFlags TargetableTextureFlags,
 	bool bForceSeparateTargetAndShaderResource,
-	FRHIResourceCreateInfo& CreateInfo,
-	FTextureCubeRHIRef& OutTargetableTexture,
-	FTextureCubeRHIRef& OutShaderResourceTexture
-	)
-{
-	// Ensure none of the usage flags are passed in.
-	check(!(Flags & TexCreate_RenderTargetable));
-	check(!(Flags & TexCreate_ResolveTargetable));
-
-	// Ensure that the targetable texture is either render or depth-stencil targetable.
-	check(EnumHasAnyFlags(TargetableTextureFlags, TexCreate_RenderTargetable | TexCreate_DepthStencilTargetable));
-
-	if(!bForceSeparateTargetAndShaderResource/* && GSupportsRenderDepthTargetableShaderResources*/)
-	{
-		// Create a single texture that has both TargetableTextureFlags and TexCreate_ShaderResource set.
-		OutTargetableTexture = OutShaderResourceTexture = RHICreateTextureCube(LinearSize, Format, NumMips, Flags | TargetableTextureFlags | TexCreate_ShaderResource, CreateInfo);
-	}
-	else
-	{
-		// Create a texture that has TargetableTextureFlags set, and a second texture that has TexCreate_ResolveTargetable and TexCreate_ShaderResource set.
-		OutTargetableTexture = RHICreateTextureCube(LinearSize, Format, NumMips, Flags | TargetableTextureFlags, CreateInfo);
-		OutShaderResourceTexture = RHICreateTextureCube(LinearSize, Format, NumMips, Flags | TexCreate_ResolveTargetable | TexCreate_ShaderResource, CreateInfo);
-	}
-}
+	const FRHIResourceCreateInfo& CreateInfo,
+	FTextureRHIRef& OutTargetableTexture,
+	FTextureRHIRef& OutShaderResourceTexture);
 
 /**
  * Creates 1 or 2 textures with the same dimensions/format.
@@ -859,7 +809,8 @@ inline void RHICreateTargetableShaderResourceCube(
  * Two texture references are always returned, but they may reference the same texture.
  * If two different textures are returned, the render-target texture must be manually copied to the shader-resource texture.
  */
-inline void RHICreateTargetableShaderResourceCubeArray(
+UE_DEPRECATED(5.1, "RHICreateTargetableShaderResourceCubeArray is deprecated, RHICreateTargetableShaderResource with a FRHITextureCreateDesc should be used instead")
+RHI_API void RHICreateTargetableShaderResourceCubeArray(
 	uint32 LinearSize,
 	uint32 ArraySize,
 	uint8 Format,
@@ -867,30 +818,9 @@ inline void RHICreateTargetableShaderResourceCubeArray(
 	ETextureCreateFlags Flags,
 	ETextureCreateFlags TargetableTextureFlags,
 	bool bForceSeparateTargetAndShaderResource,
-	FRHIResourceCreateInfo& CreateInfo,
-	FTextureCubeRHIRef& OutTargetableTexture,
-	FTextureCubeRHIRef& OutShaderResourceTexture
-	)
-{
-	// Ensure none of the usage flags are passed in.
-	check(!(Flags & TexCreate_RenderTargetable));
-	check(!(Flags & TexCreate_ResolveTargetable));
-
-	// Ensure that the targetable texture is either render or depth-stencil targetable.
-	check(EnumHasAnyFlags(TargetableTextureFlags, TexCreate_RenderTargetable | TexCreate_DepthStencilTargetable));
-
-	if(!bForceSeparateTargetAndShaderResource/* && GSupportsRenderDepthTargetableShaderResources*/)
-	{
-		// Create a single texture that has both TargetableTextureFlags and TexCreate_ShaderResource set.
-		OutTargetableTexture = OutShaderResourceTexture = RHICreateTextureCubeArray(LinearSize, ArraySize, Format, NumMips, Flags | TargetableTextureFlags | TexCreate_ShaderResource, CreateInfo);
-	}
-	else
-	{
-		// Create a texture that has TargetableTextureFlags set, and a second texture that has TexCreate_ResolveTargetable and TexCreate_ShaderResource set.
-		OutTargetableTexture = RHICreateTextureCubeArray(LinearSize, ArraySize, Format, NumMips, Flags | TargetableTextureFlags, CreateInfo);
-		OutShaderResourceTexture = RHICreateTextureCubeArray(LinearSize, ArraySize, Format, NumMips, Flags | TexCreate_ResolveTargetable | TexCreate_ShaderResource, CreateInfo);
-	}
-}
+	const FRHIResourceCreateInfo& CreateInfo,
+	FTextureRHIRef& OutTargetableTexture,
+	FTextureRHIRef& OutShaderResourceTexture);
 
 /**
  * Creates 1 or 2 textures with the same dimensions/format.
@@ -900,44 +830,19 @@ inline void RHICreateTargetableShaderResourceCubeArray(
  * Two texture references are always returned, but they may reference the same texture.
  * If two different textures are returned, the render-target texture must be manually copied to the shader-resource texture.
  */
-inline void RHICreateTargetableShaderResource3D(
+UE_DEPRECATED(5.1, "RHICreateTargetableShaderResource3D is deprecated, RHICreateTargetableShaderResource with a FRHITextureCreateDesc should be used instead")
+RHI_API void RHICreateTargetableShaderResource3D(
 	uint32 SizeX,
 	uint32 SizeY,
 	uint32 SizeZ,
-	uint8 Format,	
+	uint8 Format,
 	uint32 NumMips,
 	ETextureCreateFlags Flags,
 	ETextureCreateFlags TargetableTextureFlags,
 	bool bForceSeparateTargetAndShaderResource,
-	FRHIResourceCreateInfo& CreateInfo,
-	FTexture3DRHIRef& OutTargetableTexture,
-	FTexture3DRHIRef& OutShaderResourceTexture
-	)
-{
-	// Ensure none of the usage flags are passed in.
-	check(!(Flags & TexCreate_RenderTargetable));
-	check(!(Flags & TexCreate_ResolveTargetable));
-
-	// Ensure that the targetable texture is either render or depth-stencil targetable.
-	check(EnumHasAnyFlags(TargetableTextureFlags, TexCreate_RenderTargetable | TexCreate_DepthStencilTargetable | TexCreate_UAV));
-
-	if (!bForceSeparateTargetAndShaderResource/* && GSupportsRenderDepthTargetableShaderResources*/)
-	{
-		// Create a single texture that has both TargetableTextureFlags and TexCreate_ShaderResource set.
-		OutTargetableTexture = OutShaderResourceTexture = RHICreateTexture3D(SizeX, SizeY, SizeZ, Format, NumMips, Flags | TargetableTextureFlags | TexCreate_ShaderResource, CreateInfo);
-	}
-	else
-	{
-		ETextureCreateFlags ResolveTargetableTextureFlags = TexCreate_ResolveTargetable;
-		if (EnumHasAnyFlags(TargetableTextureFlags, TexCreate_DepthStencilTargetable))
-		{
-			ResolveTargetableTextureFlags |= TexCreate_DepthStencilResolveTarget;
-		}
-		// Create a texture that has TargetableTextureFlags set, and a second texture that has TexCreate_ResolveTargetable and TexCreate_ShaderResource set.
-		OutTargetableTexture = RHICreateTexture3D(SizeX, SizeY, SizeZ, Format, NumMips, Flags | TargetableTextureFlags, CreateInfo);
-		OutShaderResourceTexture = RHICreateTexture3D(SizeX, SizeY, SizeZ, Format, NumMips, Flags | ResolveTargetableTextureFlags | TexCreate_ShaderResource, CreateInfo);
-	}
-}
+	const FRHIResourceCreateInfo& CreateInfo,
+	FTextureRHIRef& OutTargetableTexture,
+	FTextureRHIRef& OutShaderResourceTexture);
 
 /**
  * Computes the vertex count for a given number of primitives of the specified type.
diff --git a/Engine/Source/Runtime/SlateRHIRenderer/Private/SlateRHIRenderer.cpp b/Engine/Source/Runtime/SlateRHIRenderer/Private/SlateRHIRenderer.cpp
--- a/Engine/Source/Runtime/SlateRHIRenderer/Private/SlateRHIRenderer.cpp
+++ b/Engine/Source/Runtime/SlateRHIRenderer/Private/SlateRHIRenderer.cpp
@@ -148,17 +148,21 @@ void FViewportInfo::RecreateDepthBuffer_RenderThread()
 	DepthStencil.SafeRelease();
 	if (bRequiresStencilTest)
 	{
-		FTexture2DRHIRef ShaderResourceUnused;
-		FRHIResourceCreateInfo CreateInfo(TEXT("SlateViewportDepthStencil"), FClearValueBinding::DepthZero);
-
 		ETextureCreateFlags TargetableTextureFlags = TexCreate_DepthStencilTargetable;
 		if (CVarMemorylessDepthStencil.GetValueOnAnyThread() != 0)
 		{
 			// Use Memoryless target, expecting that DepthStencil content is intermediate and can't be preserved between renderpasses
-			TargetableTextureFlags|= TexCreate_Memoryless;
+			TargetableTextureFlags |= TexCreate_Memoryless;
 		}
+
+		const FRHITextureCreateDesc Desc =
+			FRHITextureCreateDesc::Create2D(TEXT("SlateViewportDepthStencil"))
+			.SetExtent(Width, Height)
+			.SetFormat(PF_DepthStencil)
+			.SetClearValue(FClearValueBinding::DepthZero);
+
+		RHICreateTargetableShaderResource(Desc, TargetableTextureFlags, DepthStencil);
 		
-		RHICreateTargetableShaderResource2D(Width, Height, PF_DepthStencil, 1, TexCreate_None, TargetableTextureFlags, false, CreateInfo, DepthStencil, ShaderResourceUnused);
 		check(IsValidRef(DepthStencil));
 	}
 }
@@ -787,29 +791,38 @@ void FSlateRHIRenderer::DrawWindow_RenderThread(FRHICommandListImmediate& RHICmd
 				|| !ViewportInfo.HDRSourceRT || ViewportInfo.HDRSourceRT->GetRenderTargetItem().TargetableTexture->GetFormat() != BackBuffer->GetFormat()))
 			{
 				// Composition buffers
-				ETextureCreateFlags BaseFlags = RHISupportsRenderTargetWriteMask(GMaxRHIShaderPlatform) ? TexCreate_NoFastClearFinalize | TexCreate_DisableDCC : TexCreate_None;
-
-				FPooledRenderTargetDesc Desc(FPooledRenderTargetDesc::Create2DDesc(FIntPoint(ViewportWidth, ViewportHeight),
-					GetSlateRecommendedColorFormat(),
-					FClearValueBinding::Transparent,
-					BaseFlags,
-					TexCreate_ShaderResource | TexCreate_RenderTargetable,
-					false,
-					1,
-					true,
-					true));
+				{
+					ETextureCreateFlags BaseFlags = RHISupportsRenderTargetWriteMask(GMaxRHIShaderPlatform) ? TexCreate_NoFastClearFinalize | TexCreate_DisableDCC : TexCreate_None;
+					FPooledRenderTargetDesc Desc(FPooledRenderTargetDesc::Create2DDesc(FIntPoint(ViewportWidth, ViewportHeight),
+						GetSlateRecommendedColorFormat(),
+						FClearValueBinding::Transparent,
+						BaseFlags,
+						TexCreate_ShaderResource | TexCreate_RenderTargetable,
+						false,
+						1,
+						true,
+						true));
+
+					GRenderTargetPool.FindFreeElement(RHICmdList, Desc, ViewportInfo.UITargetRT, TEXT("UITargetRT"));
+
+					Desc.Format = BackBuffer->GetFormat();
+					GRenderTargetPool.FindFreeElement(RHICmdList, Desc, ViewportInfo.HDRSourceRT, TEXT("HDRSourceRT"));
+				}
 
-				GRenderTargetPool.FindFreeElement(RHICmdList, Desc, ViewportInfo.UITargetRT, TEXT("UITargetRT"));
+				// LUT
+				{
+					ViewportInfo.ColorSpaceLUTRT.SafeRelease();
+					ViewportInfo.ColorSpaceLUTSRV.SafeRelease();
 
-				Desc.Format = BackBuffer->GetFormat();
-				GRenderTargetPool.FindFreeElement(RHICmdList, Desc, ViewportInfo.HDRSourceRT, TEXT("HDRSourceRT"));
+					const FRHITextureCreateDesc Desc =
+						FRHITextureCreateDesc::Create3D(TEXT("ColorSpaceLUT"))
+						.SetExtent(CompositionLUTSize)
+						.SetDepth(CompositionLUTSize)
+						.SetFormat(PF_A2B10G10R10);
 
-				// LUT
-				ViewportInfo.ColorSpaceLUTRT.SafeRelease();
-				ViewportInfo.ColorSpaceLUTSRV.SafeRelease();
+					RHICreateTargetableShaderResource(Desc, ETextureCreateFlags::RenderTargetable, ViewportInfo.ColorSpaceLUTRT, ViewportInfo.ColorSpaceLUTSRV);
+				}
 
-				FRHIResourceCreateInfo CreateInfo(TEXT("ColorSpaceLUT"));
-				RHICreateTargetableShaderResource3D(CompositionLUTSize, CompositionLUTSize, CompositionLUTSize, PF_A2B10G10R10, 1, TexCreate_None, TexCreate_RenderTargetable, false, CreateInfo, ViewportInfo.ColorSpaceLUTRT, ViewportInfo.ColorSpaceLUTSRV);
 				bLUTStale = true;
 			}
 
