Restoring CopyToResolveTarget to do internal transitions.

diff --git a/Engine/Source/Runtime/D3D12RHI/Private/D3D12RenderTarget.cpp b/Engine/Source/Runtime/D3D12RHI/Private/D3D12RenderTarget.cpp
--- a/Engine/Source/Runtime/D3D12RHI/Private/D3D12RenderTarget.cpp
+++ b/Engine/Source/Runtime/D3D12RHI/Private/D3D12RenderTarget.cpp
@@ -462,6 +462,16 @@ void FD3D12CommandContext::RHICopyToResolveTarget(FRHITexture* SourceTextureRHI,
 		check(SourceTexture == DestTexture);
 	}
 
+	// Force transtion to readable since RHI user expects both source and dest to be in readable state after calling RHICopyToResolveTarget
+	if (SourceTextureRHI && ResolveParams.SourceAccessFinal != ERHIAccess::Unknown)
+	{
+		RHICmdList.TransitionResource(ResolveParams.SourceAccessFinal, SourceTextureRHI);
+	}
+	if (DestTextureRHI && SourceTextureRHI != DestTextureRHI && ResolveParams.DestAccessFinal != ERHIAccess::Unknown)
+	{
+		RHICmdList.TransitionResource(ResolveParams.DestAccessFinal, DestTextureRHI);
+	}
+
 	ConditionalFlushCommandList();
 
 	DEBUG_EXECUTE_COMMAND_LIST(this);
diff --git a/Engine/Source/Runtime/RHI/Public/RHICommandList.h b/Engine/Source/Runtime/RHI/Public/RHICommandList.h
--- a/Engine/Source/Runtime/RHI/Public/RHICommandList.h
+++ b/Engine/Source/Runtime/RHI/Public/RHICommandList.h
@@ -448,19 +448,6 @@ public:
 		return TArrayView<T>((T*) NewArray, InArray.Num());
 	}
 
-	template <typename T>
-	FORCEINLINE_DEBUGGABLE const TArrayView<T> AllocArrayIfDeferred(const TArrayView<T> InArray)
-	{
-		if (Bypass())
-		{
-			return InArray;
-		}
-		else
-		{
-			return AllocArray(InArray);
-		}
-	}
-
 	FORCEINLINE_DEBUGGABLE TCHAR* AllocString(const TCHAR* Name)
 	{
 		int32 Len = FCString::Strlen(Name) + 1;
@@ -2390,7 +2377,7 @@ public:
 
 	/** Custom new/delete with recycling */
 	void* operator new(size_t Size);
-	void operator delete(void* RawMemory);
+	void operator delete(void *RawMemory);
 
 	template <typename LAMBDA>
 	FORCEINLINE_DEBUGGABLE void EnqueueLambda(LAMBDA&& Lambda)
@@ -2408,7 +2395,7 @@ public:
 	inline FRHIComputeShader* GetBoundComputeShader() const { return BoundComputeShaderRHI; }
 
 	UE_DEPRECATED(5.0, "Please rename to SetStaticUniformBuffers")
-		FORCEINLINE_DEBUGGABLE void SetGlobalUniformBuffers(const FUniformBufferStaticBindings& UniformBuffers)
+	FORCEINLINE_DEBUGGABLE void SetGlobalUniformBuffers(const FUniformBufferStaticBindings& UniformBuffers)
 	{
 		SetStaticUniformBuffers(UniformBuffers);
 	}
@@ -2456,7 +2443,7 @@ public:
 	}
 
 	UE_DEPRECATED(5.0, "Use Layout pointers instead")
-		FORCEINLINE_DEBUGGABLE FLocalUniformBuffer BuildLocalUniformBuffer(const void* Contents, uint32 ContentsSize, const FRHIUniformBufferLayout& Layout)
+	FORCEINLINE_DEBUGGABLE FLocalUniformBuffer BuildLocalUniformBuffer(const void* Contents, uint32 ContentsSize, const FRHIUniformBufferLayout& Layout)
 	{
 		return BuildLocalUniformBuffer(Contents, ContentsSize, &Layout);
 	}
@@ -2559,16 +2546,16 @@ public:
 	}
 
 	UE_DEPRECATED(5.1, "ComputePipelineStates should be used instead of direct ComputeShaders. You can use SetComputePipelineState(RHICmdList, ComputeShader).")
-		FORCEINLINE_DEBUGGABLE void SetComputeShader(FRHIComputeShader* ComputeShader)
+	FORCEINLINE_DEBUGGABLE void SetComputeShader(FRHIComputeShader* ComputeShader)
 	{
 		BoundComputeShaderRHI = ComputeShader;
 		ComputeShader->UpdateStats();
 		if (Bypass())
 		{
-			PRAGMA_DISABLE_DEPRECATION_WARNINGS
-				GetComputeContext().RHISetComputeShader(ComputeShader);
-			PRAGMA_ENABLE_DEPRECATION_WARNINGS
-				return;
+PRAGMA_DISABLE_DEPRECATION_WARNINGS
+			GetComputeContext().RHISetComputeShader(ComputeShader);
+PRAGMA_ENABLE_DEPRECATION_WARNINGS
+			return;
 		}
 		ALLOC_COMMAND(FRHICommandSetComputeShader)(ComputeShader);
 	}
@@ -2578,7 +2565,7 @@ public:
 		BoundComputeShaderRHI = ComputeShader;
 		if (Bypass())
 		{
-			extern RHI_API FRHIComputePipelineState* ExecuteSetComputePipelineState(FComputePipelineState * ComputePipelineState);
+			extern RHI_API FRHIComputePipelineState* ExecuteSetComputePipelineState(FComputePipelineState* ComputePipelineState);
 			FRHIComputePipelineState* RHIComputePipelineState = ExecuteSetComputePipelineState(ComputePipelineState);
 			GetComputeContext().RHISetComputePipelineState(RHIComputePipelineState);
 			return;
@@ -3550,28 +3537,9 @@ public:
 		if (Bypass())
 		{
 			GetContext().RHICopyToResolveTarget(SourceTextureRHI, DestTextureRHI, ResolveParams);
+			return;
 		}
-		else
-		{
-			ALLOC_COMMAND(FRHICommandCopyToResolveTarget)(SourceTextureRHI, DestTextureRHI, ResolveParams);
-		}
-
-		TArray<FRHITransitionInfo, TInlineAllocator<2>> Infos;
-
-		if (ResolveParams.SourceAccessFinal != ERHIAccess::Unknown && SourceTextureRHI != DestTextureRHI)
-		{
-			Infos.Emplace(SourceTextureRHI, ERHIAccess::Unknown, ResolveParams.SourceAccessFinal);
-		}
-
-		if (ResolveParams.DestAccessFinal != ERHIAccess::Unknown)
-		{
-			Infos.Emplace(DestTextureRHI, ERHIAccess::Unknown, ResolveParams.DestAccessFinal);
-		}
-
-		if (!Infos.IsEmpty())
-		{
-			Transition(Infos);
-		}
+		ALLOC_COMMAND(FRHICommandCopyToResolveTarget)(SourceTextureRHI, DestTextureRHI, ResolveParams);
 	}
 
 	FORCEINLINE_DEBUGGABLE void CopyTexture(FRHITexture* SourceTextureRHI, FRHITexture* DestTextureRHI, const FRHICopyTextureInfo& CopyInfo)
