remove check as parent resource is allowed to be null

diff --git a/Engine/Source/Runtime/RHI/Public/RHIResources.h b/Engine/Source/Runtime/RHI/Public/RHIResources.h
--- a/Engine/Source/Runtime/RHI/Public/RHIResources.h
+++ b/Engine/Source/Runtime/RHI/Public/RHIResources.h
@@ -2394,7 +2394,6 @@ public:
 protected:
 	void SetParentResource(FRHIViewableResource* InParentResource)
 	{
-		check(InParentResource);
 		ParentResource = InParentResource;
 	}
 
