Background PSO compililation changes,
Add an option to use a thread pool instead of taskgraph for GFX PSO precompile tasks.
Disabled by default. Use r.pso.PrecompileThreadPoolSize=N to enable.
threadpool PSO precompile tasks are not added as dependencies of RHIcmd list and should not block RHIT.

diff --git a/Engine/Source/Runtime/RHI/Private/PipelineStateCache.cpp b/Engine/Source/Runtime/RHI/Private/PipelineStateCache.cpp
--- a/Engine/Source/Runtime/RHI/Private/PipelineStateCache.cpp
+++ b/Engine/Source/Runtime/RHI/Private/PipelineStateCache.cpp
@@ -11,6 +11,7 @@ PipelineStateCache.cpp: Pipeline state cache implementation.
 #include "CoreGlobals.h"
 #include "Misc/TimeGuard.h"
 #include "Containers/DiscardableKeyValueCache.h"
+#include "Async/Async.h"
 
 // perform cache eviction each frame, used to stress the system and flush out bugs
 #define PSO_DO_CACHE_EVICT_EACH_FRAME 0
@@ -92,6 +93,63 @@ void SetComputePipelineState(FRHIComputeCommandList& RHICmdList, FRHIComputeShad
 	RHICmdList.SetComputePipelineState(PipelineStateCache::GetAndOrCreateComputePipelineState(RHICmdList, ComputeShader, false), ComputeShader);
 }
 
+static int32 GPSOPrecompileThreadPoolSize = 0;
+static FAutoConsoleVariableRef GPSOPrecompileThreadPoolSizeVar(
+	TEXT("r.pso.PrecompileThreadPoolSize"),
+	GPSOPrecompileThreadPoolSize,
+	TEXT("The maximum number of threads available for concurrent PSO Precompiling.\n")
+	TEXT("0 to disable threadpool usage when precompiling PSOs. (default)")
+	,
+	ECVF_RenderThreadSafe | ECVF_ReadOnly
+);
+
+class FPSOThreadPool
+{
+public:
+	~FPSOThreadPool()
+	{
+		ShutdownThreadPool();
+	}
+
+	FQueuedThreadPool& Get()
+	{
+		if (PSOPrecompileCompileThreadPool.load() == nullptr)
+		{
+			FScopeLock lock(&LockCS);
+			if (PSOPrecompileCompileThreadPool.load() == nullptr)
+			{
+				FQueuedThreadPool* PSOPrecompileCompileThreadPoolLocal = FQueuedThreadPool::Allocate();
+				PSOPrecompileCompileThreadPoolLocal->Create(GPSOPrecompileThreadPoolSize, 512 * 1024, EThreadPriority::TPri_BelowNormal, TEXT("PSOPrecompilePool"));
+				PSOPrecompileCompileThreadPool = PSOPrecompileCompileThreadPoolLocal;
+			}
+		}
+		return *PSOPrecompileCompileThreadPool;
+	}
+
+	void ShutdownThreadPool()
+	{
+		FQueuedThreadPool* LocalPool = PSOPrecompileCompileThreadPool.exchange(nullptr);
+		if (LocalPool != nullptr)
+		{
+
+			LocalPool->Destroy();
+			delete LocalPool;
+		}
+	}
+
+private:
+	FCriticalSection LockCS;
+	std::atomic<FQueuedThreadPool*> PSOPrecompileCompileThreadPool {};
+};
+
+static FPSOThreadPool GPSOThreadPool;
+
+void PipelineStateCache::PreCompileComplete()
+{
+	// free up our threads when the precompile completes.
+	GPSOThreadPool.ShutdownThreadPool();
+}
+
 extern RHI_API FRHIComputePipelineState* ExecuteSetComputePipelineState(FComputePipelineState* ComputePipelineState);
 extern RHI_API FRHIGraphicsPipelineState* ExecuteSetGraphicsPipelineState(FGraphicsPipelineState* GraphicsPipelineState);
 
@@ -145,6 +203,55 @@ static void HandlePipelineCreationFailure(const FGraphicsPipelineStateInitialize
 	}
 }
 
+// FAsyncTask used by the PSO threadpool that will reschedule once only.
+template<typename TTask>
+class FAsyncTaskLimitedReschedule : public FAsyncTask<TTask>
+{
+public:
+	/** Forwarding constructor. */
+	template<typename...T>
+	explicit FAsyncTaskLimitedReschedule(T&&... Args) : FAsyncTask<TTask>(Forward<T>(Args)...)	{ }
+
+	bool Reschedule(FQueuedThreadPool* InQueuedPool, EQueuedWorkPriority InQueuedWorkPriority)
+	{
+		if (bCanReschedule.exchange(false))
+		{
+			return FAsyncTask<TTask>::Reschedule(InQueuedPool, InQueuedWorkPriority);
+		}
+		return false;
+	}
+private:
+	std::atomic_bool bCanReschedule = true;
+};
+
+class FPSOPrecompileTask : public FNonAbandonableTask
+{
+	friend class FAsyncTask<FPSOPrecompileTask>;
+private:
+	TUniqueFunction<void()> Func;
+	TUniqueFunction<void()> OnComplete;
+public:
+	FPSOPrecompileTask(TUniqueFunction<void()> InFunc, TUniqueFunction<void()> InOnComplete) : Func(MoveTemp(InFunc)), OnComplete(MoveTemp(InOnComplete)) {}
+
+protected:
+	void DoWork() 
+	{ 
+		QUICK_SCOPE_CYCLE_COUNTER(STAT_FPSOPrecompileTask_Work);
+
+		Func(); 
+		Func = nullptr;
+		OnComplete(); 
+		OnComplete = nullptr;
+	}
+
+	FORCEINLINE TStatId GetStatId() const
+	{
+		RETURN_QUICK_DECLARE_CYCLE_STAT(FPSOPrecompileTask, STATGROUP_ThreadPoolAsyncTasks);
+	}
+};
+
+
+
 /**
  * Base class to hold pipeline state (and optionally stats)
  */
@@ -160,20 +267,43 @@ public:
 
 	virtual ~FPipelineState() 
 	{
+		check(IsComplete());
+		check(!PrecompileTask.IsValid());
 	}
 
 	virtual bool IsCompute() const = 0;
 
 	FGraphEventRef CompletionEvent;
-	
-	void WaitCompletion()
+	TUniquePtr<FAsyncTaskLimitedReschedule<FPSOPrecompileTask>> PrecompileTask;
+
+	bool IsComplete()
 	{
+		return (!CompletionEvent.IsValid() || CompletionEvent->IsComplete()) && (!PrecompileTask.IsValid() || PrecompileTask->IsDone());
+	}
+
+	// return true if we actually waited on the task
+	bool WaitCompletion()
+	{
+		bool bNeedsToWait = false;
 		if(CompletionEvent.IsValid() && !CompletionEvent->IsComplete())
 		{
+			QUICK_SCOPE_CYCLE_COUNTER(STAT_FPipelineState_WaitCompletion);
+#if PSO_TRACK_CACHE_STATS
 			UE_LOG(LogRHI, Log, TEXT("FTaskGraphInterface Waiting on FPipelineState completionEvent"));
+#endif
+			bNeedsToWait = true;
 			FTaskGraphInterface::Get().WaitUntilTaskCompletes( CompletionEvent );
-			CompletionEvent = nullptr;
 		}
+		CompletionEvent = nullptr;
+
+		if (PrecompileTask.IsValid())
+		{
+			bNeedsToWait = bNeedsToWait || !PrecompileTask->IsDone();
+			PrecompileTask->EnsureCompletion();
+			PrecompileTask = nullptr;
+		}
+
+		return bNeedsToWait;
 	}
 
 	inline void AddUse()
@@ -478,6 +608,8 @@ public:
 		
 		
 	}
+	// The list of tasks that are still in progress.
+	TArray<TTuple<TMyKey,TMyValue>> Uncompleted;
 
 	bool Add(const TMyKey& InKey, const TMyValue& InValue)
 	{
@@ -512,11 +644,7 @@ public:
 				const TMyKey& ThreadKey = PipelineStateCacheIterator->Key;
 				const TMyValue& ThreadValue = PipelineStateCacheIterator->Value;
 
-				// All events should be complete because we are running this code after the RHI Flush
-				if(!ThreadValue->CompletionEvent.IsValid() || ThreadValue->CompletionEvent->IsComplete())
 				{
-					ThreadValue->CompletionEvent = nullptr;
-
 					BackfillMap->Remove(ThreadKey);
 
 					TMyValue* CurrentValue = CurrentMap->Find(ThreadKey);
@@ -532,12 +660,23 @@ public:
 					else
 					{
 						CurrentMap->Add(ThreadKey, ThreadValue);
+						Uncompleted.Add(TTuple<TMyKey, TMyValue>(ThreadKey, ThreadValue));
 					}
 					PipelineStateCacheIterator.RemoveCurrent();
 				}
 			}
 		}
 
+		// tick and complete any uncompleted PSO tasks (we free up precompile tasks here).
+		for (int32 i = Uncompleted.Num() - 1; i >= 0; i--)
+		{
+			checkSlow(CurrentMap->Find(Uncompleted[i].Key));
+			if (Uncompleted[i].Value->IsComplete())
+			{
+				Uncompleted[i].Value->WaitCompletion();
+				Uncompleted.RemoveAtSwap(i, 1, /*bAllowShrinking=*/ false);
+			}
+		}
 	}
 
 	void ProcessDelayedCleanup()
@@ -553,13 +692,16 @@ public:
 #if PIPELINESTATECACHE_VERIFYTHREADSAFE
 				check(OldPipelineState->InUseCount.GetValue() == 0);
 #endif
+				// Duplicate entries must wait for in progress compiles to complete.
+				// inprogress tasks could also remain in this container and deferred for the next tick.
+				bool bWaited = OldPipelineState->WaitCompletion();
+				UE_CLOG(bWaited, LogRHI, Log, TEXT("Waited on a pipeline compile task while discarding duplicate."));
 				delete OldPipelineState;
 			}
 			DeleteArray.Empty();
 		});
 	}
 
-
 	int32 DiscardAndSwap()
 	{
 		check(IsInRenderingThread());
@@ -576,6 +718,9 @@ public:
 #if PIPELINESTATECACHE_VERIFYTHREADSAFE
 				check( DiscardIterator.Value->InUseCount.GetValue() == 0);
 #endif
+				// Incomplete tasks should be put back to the current map. There should be no incomplete tasks encountered here.
+				bool bWaited = DiscardIterator.Value->WaitCompletion();
+				UE_CLOG(bWaited, LogRHI, Error, TEXT("Waited on a pipeline compile task while discarding retired PSOs."));
 				delete DiscardIterator.Value;
 			}
 		});
@@ -590,6 +735,18 @@ public:
 			CurrentMap = &Map1;
 			BackfillMap = &Map2;
 		}
+
+		// keep alive incomplete tasks by moving them back to the current map.
+		for (int32 i = Uncompleted.Num() - 1; i >= 0; i--)
+		{
+			int32 Removed = BackfillMap->Remove(Uncompleted[i].Key);
+			checkSlow(Removed);
+			if (Removed)
+			{
+				CurrentMap->Add(Uncompleted[i].Key, Uncompleted[i].Value);
+			}
+		}
+
 		return Discarded;
 	}
 	
@@ -635,6 +792,9 @@ private:
 	TArray<FPipelineStateCacheType*> AllThreadsPipelineStateCache;
 
 	uint32 DuplicateStateGenerated;
+#if PSO_TRACK_CACHE_STATS
+	friend void DumpPipelineCacheStats();
+#endif
 
 #if PIPELINESTATECACHE_VERIFYTHREADSAFE
 	volatile int32 VerifyMutex;
@@ -859,6 +1019,13 @@ FRayTracingPipelineCache GRayTracingPipelineCache;
 /**
  *  Compile task
  */
+static std::atomic<int32> GPipelinePrecompileTasksInFlight = { 0 };
+
+int32 PipelineStateCache::GetNumActivePipelinePrecompileTasks()
+{
+	return GPipelinePrecompileTasksInFlight.load();
+}
+
 class FCompilePipelineStateTask
 {
 public:
@@ -870,6 +1037,11 @@ public:
 		: Pipeline(InPipeline)
 		, Initializer(InInitializer)
 	{
+		if(Initializer.bFromPSOFileCache)
+		{
+			GPipelinePrecompileTasksInFlight++;
+		}
+
 		if (!Pipeline->IsCompute())
 		{
 			if (Initializer.BoundShaderState.GetMeshShader())
@@ -924,9 +1096,22 @@ public:
 		}
 	}
 
+	~FCompilePipelineStateTask()
+	{
+		if (Initializer.bFromPSOFileCache)
+		{
+			GPipelinePrecompileTasksInFlight--;
+		}
+	}
+
 	static ESubsequentsMode::Type GetSubsequentsMode() { return ESubsequentsMode::TrackSubsequents; }
 
 	void DoTask(ENamedThreads::Type CurrentThread, const FGraphEventRef& MyCompletionGraphEvent)
+	{
+		CompilePSO();
+	}
+
+	void CompilePSO()
 	{
 		LLM_SCOPE(ELLMTag::PSO);
 		FTaskTagScope Scope(ETaskTag::EParallelRenderingThread);
@@ -1178,10 +1363,9 @@ FComputePipelineState* PipelineStateCache::GetAndOrCreateComputePipelineState(FR
 	{
 		if (DoAsyncCompile)
 		{
-			FGraphEventRef& CompletionEvent = OutCachedState->CompletionEvent;
-			if ( CompletionEvent.IsValid() && !CompletionEvent->IsComplete() )
+			if(!OutCachedState->IsComplete())
 			{
-				RHICmdList.AddDispatchPrerequisite(CompletionEvent);
+				RHICmdList.AddDispatchPrerequisite(OutCachedState->CompletionEvent);
 			}
 		}
 
@@ -1470,6 +1654,25 @@ FRHIComputePipelineState* ExecuteSetComputePipelineState(FComputePipelineState*
 	return ComputePipelineState->RHIPipeline;
 }
 
+#if PSO_TRACK_CACHE_STATS
+
+static std::atomic_uint64_t TotalPrecompileWaitTime;
+static std::atomic_int64_t TotalNumPrecompileJobs;
+static std::atomic_int64_t TotalNumPrecompileJobsCompleted;
+
+void StatsStartPrecompile()
+{
+	TotalNumPrecompileJobs++;
+}
+
+void StatsEndPrecompile(uint64 TimeToComplete)
+{
+	TotalPrecompileWaitTime += TimeToComplete;
+	TotalNumPrecompileJobsCompleted++;
+}
+
+#endif
+
 FGraphicsPipelineState* PipelineStateCache::GetAndOrCreateGraphicsPipelineState(FRHICommandList& RHICmdList, const FGraphicsPipelineStateInitializer& Initializer, EApplyRendertargetOption ApplyFlags)
 {
 	LLM_SCOPE(ELLMTag::PSO);
@@ -1544,8 +1747,39 @@ FGraphicsPipelineState* PipelineStateCache::GetAndOrCreateGraphicsPipelineState(
 		// create a compilation task, or just do it now...
 		if (DoAsyncCompile)
 		{
-			OutCachedState->CompletionEvent = TGraphTask<FCompilePipelineStateTask>::CreateTask().ConstructAndDispatchWhenReady(OutCachedState, Initializer);
-			RHICmdList.AddDispatchPrerequisite(OutCachedState->CompletionEvent);
+			// Use normal task graph for non-precompile jobs.
+			if (!Initializer.bFromPSOFileCache || GPSOPrecompileThreadPoolSize <=0)
+			{
+				OutCachedState->CompletionEvent = TGraphTask<FCompilePipelineStateTask>::CreateTask().ConstructAndDispatchWhenReady(OutCachedState, Initializer);
+				RHICmdList.AddDispatchPrerequisite(OutCachedState->CompletionEvent);
+			}
+			else
+			{
+				// Here, PSO precompiles use a separate thread pool.
+				// Note that we do not add precompile tasks as cmdlist prerequisites.
+				FGraphEventRef CompletedEventGraph = FGraphEvent::CreateGraphEvent();
+				TUniquePtr<FCompilePipelineStateTask> ThreadPoolTask = MakeUnique<FCompilePipelineStateTask>(OutCachedState, Initializer);
+				OutCachedState->CompletionEvent = CompletedEventGraph;
+				uint64 StartTime = FPlatformTime::Cycles64();
+#if	PSO_TRACK_CACHE_STATS
+				StatsStartPrecompile();
+#endif
+				OutCachedState->PrecompileTask = MakeUnique<FAsyncTaskLimitedReschedule<FPSOPrecompileTask>>(
+					[ThreadPoolTask = MoveTemp(ThreadPoolTask)]()
+					{
+						ThreadPoolTask->CompilePSO();
+				}
+				,
+					[CompletedEventGraph, StartTime]()
+					{
+						CompletedEventGraph->DispatchSubsequents();
+	#if PSO_TRACK_CACHE_STATS
+						StatsEndPrecompile(FPlatformTime::Cycles64() - StartTime);
+	#endif
+					}
+				 );
+				OutCachedState->PrecompileTask->StartBackgroundTask(&GPSOThreadPool.Get(), EQueuedWorkPriority::Normal);
+			}
 		}
 		else
 		{
@@ -1560,18 +1794,22 @@ FGraphicsPipelineState* PipelineStateCache::GetAndOrCreateGraphicsPipelineState(
 			}
 		}
 
-		// GGraphicsPipelineCache.Add(*Initializer, OutCachedState, LockFlags);
 		GGraphicsPipelineCache.Add(Initializer, OutCachedState);
 	}
 	else
 	{
-		if (DoAsyncCompile)
+		if (DoAsyncCompile && !Initializer.bFromPSOFileCache && !OutCachedState->IsComplete())
 		{
-			FGraphEventRef& CompletionEvent = OutCachedState->CompletionEvent;
-			if ( CompletionEvent.IsValid() && !CompletionEvent->IsComplete() )
+			if (OutCachedState->PrecompileTask)
 			{
-				RHICmdList.AddDispatchPrerequisite(CompletionEvent);
+				// if this is an in-progress threadpool precompile task then it could be seconds away in the queue.
+				// Reissue this task so that it jumps the precompile queue.
+				OutCachedState->PrecompileTask->Reschedule(&GPSOThreadPool.Get(), EQueuedWorkPriority::Highest);
+#if PSO_TRACK_CACHE_STATS
+		UE_LOG(LogRHI, Log, TEXT("An incomplete precompile task was required for rendering!"));
+#endif
 			}
+			RHICmdList.AddDispatchPrerequisite(OutCachedState->CompletionEvent);
 		}
 
 #if PSO_TRACK_CACHE_STATS
@@ -1613,16 +1851,16 @@ void DumpPipelineCacheStats()
 	int NumHitsAcrossFrames = 0;
 	int NumItemsMultipleFrameHits = 0;
 
-	int NumCachedItems = GGraphicsPipelineCache.Current().Num();
+	int NumCachedItems = GGraphicsPipelineCache.CurrentMap->Num();
 
 	if (NumCachedItems == 0)
 	{
 		return;
 	}
 
-	for (auto GraphicsPipeLine : GGraphicsPipelineCache.Current())
+	for (auto GraphicsPipeLine : *GGraphicsPipelineCache.CurrentMap)
 	{
-		TSharedPtr<FGraphicsPipelineState> State = GraphicsPipeLine.Value;
+		FGraphicsPipelineState* State = GraphicsPipeLine.Value;
 
 		// calc timestats
 		double SinceUse = FPlatformTime::Seconds() - State->FirstUsedTime;
@@ -1653,11 +1891,14 @@ void DumpPipelineCacheStats()
 	}
 
 	UE_LOG(LogRHI, Log, TEXT("Have %d GraphicsPipeline entries"), NumCachedItems);
+	UE_LOG(LogRHI, Log, TEXT("Threadpool precompile: %d GraphicsPipeline in flight, %d Jobs started, %d completed"), GPipelinePrecompileTasksInFlight.load(), TotalNumPrecompileJobs.load(), TotalNumPrecompileJobsCompleted.load());
+	UE_LOG(LogRHI, Log, TEXT("Threadpool precompile: %f s avg precompile time"), FPlatformTime::GetSecondsPerCycle64() * (TotalPrecompileWaitTime.load()/FMath::Max((int64_t)1, TotalNumPrecompileJobsCompleted.load())));
+
 	UE_LOG(LogRHI, Log, TEXT("Secs Used: Min=%.02f, Max=%.02f, Avg=%.02f. %d used in last 30 secs"), MinTime, MaxTime, TotalTime / NumCachedItems, NumUsedLastMin);
 	UE_LOG(LogRHI, Log, TEXT("Frames Used: Min=%d, Max=%d, Avg=%d"), MinFrames, MaxFrames, TotalFrames / NumCachedItems);
 	UE_LOG(LogRHI, Log, TEXT("Hits: Avg=%d, Items with hits across frames=%d, Avg Hits across Frames=%d"), NumHits / NumCachedItems, NumItemsMultipleFrameHits, NumHitsAcrossFrames / NumCachedItems);
 
-	size_t TrackingMem = sizeof(FGraphicsPipelineStateInitializer) * GGraphicsPipelineCache.Num();
+	size_t TrackingMem = sizeof(FGraphicsPipelineStateInitializer) * GGraphicsPipelineCache.CurrentMap->Num();
 	UE_LOG(LogRHI, Log, TEXT("Tracking Mem: %d kb"), TrackingMem / 1024);
 #else
 	UE_LOG(LogRHI, Error, TEXT("DEfine PSO_TRACK_CACHE_STATS for state and stats!"));
diff --git a/Engine/Source/Runtime/RHI/Public/PipelineStateCache.h b/Engine/Source/Runtime/RHI/Public/PipelineStateCache.h
--- a/Engine/Source/Runtime/RHI/Public/PipelineStateCache.h
+++ b/Engine/Source/Runtime/RHI/Public/PipelineStateCache.h
@@ -75,6 +75,12 @@ namespace PipelineStateCache
 
 	/* Clears all pipeline cached state. Called on shutdown, calling GetAndOrCreate after this will recreate state */
 	extern RHI_API void Shutdown();
+
+	/* Called when PSO precompile has completed */
+	extern RHI_API void PreCompileComplete();
+
+	/* Returns the number of PSO precompiles currently in progress */
+	extern RHI_API int32 GetNumActivePipelinePrecompileTasks();
 }
 
 // Returns the hit group index within the ray tracing pipeline or INDEX_NONE if given shader does not exist.
diff --git a/Engine/Source/Runtime/RenderCore/Private/ShaderPipelineCache.cpp b/Engine/Source/Runtime/RenderCore/Private/ShaderPipelineCache.cpp
--- a/Engine/Source/Runtime/RenderCore/Private/ShaderPipelineCache.cpp
+++ b/Engine/Source/Runtime/RenderCore/Private/ShaderPipelineCache.cpp
@@ -621,86 +621,94 @@ bool FShaderPipelineCache::Precompile(FRHICommandListImmediate& RHICmdList, ESha
 			TRACE_CPUPROFILER_EVENT_SCOPE(FShaderPipelineCache::PrecompileGraphics);
 
 			FGraphicsPipelineStateInitializer GraphicsInitializer;
-
-			if (PSO.GraphicsDesc.MeshShader != FSHAHash())
 			{
-				FMeshShaderRHIRef MeshShader = FShaderCodeLibrary::CreateMeshShader(Platform, PSO.GraphicsDesc.MeshShader);
-				GraphicsInitializer.BoundShaderState.SetMeshShader(MeshShader);
+				QUICK_SCOPE_CYCLE_COUNTER(STAT_FShaderPipelineCache_PrepareInitializer);
 
-				if (PSO.GraphicsDesc.AmplificationShader != FSHAHash())
+				if (PSO.GraphicsDesc.MeshShader != FSHAHash())
 				{
-					FAmplificationShaderRHIRef AmplificationShader = FShaderCodeLibrary::CreateAmplificationShader(Platform, PSO.GraphicsDesc.AmplificationShader);
-					GraphicsInitializer.BoundShaderState.SetAmplificationShader(AmplificationShader);
+					FMeshShaderRHIRef MeshShader = FShaderCodeLibrary::CreateMeshShader(Platform, PSO.GraphicsDesc.MeshShader);
+					GraphicsInitializer.BoundShaderState.SetMeshShader(MeshShader);
+
+					if (PSO.GraphicsDesc.AmplificationShader != FSHAHash())
+					{
+						FAmplificationShaderRHIRef AmplificationShader = FShaderCodeLibrary::CreateAmplificationShader(Platform, PSO.GraphicsDesc.AmplificationShader);
+						GraphicsInitializer.BoundShaderState.SetAmplificationShader(AmplificationShader);
+					}
 				}
-			}
-			else
-			{
-				FRHIVertexDeclaration* VertexDesc = PipelineStateCache::GetOrCreateVertexDeclaration(PSO.GraphicsDesc.VertexDescriptor);
-				GraphicsInitializer.BoundShaderState.VertexDeclarationRHI = VertexDesc;
-			
-				FVertexShaderRHIRef VertexShader;
-				if (PSO.GraphicsDesc.VertexShader != FSHAHash())
+				else
 				{
-					VertexShader = FShaderCodeLibrary::CreateVertexShader(Platform, PSO.GraphicsDesc.VertexShader);
-					GraphicsInitializer.BoundShaderState.VertexShaderRHI = VertexShader;
-				}
+					FRHIVertexDeclaration* VertexDesc = PipelineStateCache::GetOrCreateVertexDeclaration(PSO.GraphicsDesc.VertexDescriptor);
+					GraphicsInitializer.BoundShaderState.VertexDeclarationRHI = VertexDesc;
+
+					FVertexShaderRHIRef VertexShader;
+					if (PSO.GraphicsDesc.VertexShader != FSHAHash())
+					{
+						VertexShader = FShaderCodeLibrary::CreateVertexShader(Platform, PSO.GraphicsDesc.VertexShader);
+						GraphicsInitializer.BoundShaderState.VertexShaderRHI = VertexShader;
+					}
 
 #if PLATFORM_SUPPORTS_GEOMETRY_SHADERS
-				FGeometryShaderRHIRef GeometryShader;
-				if (PSO.GraphicsDesc.GeometryShader != FSHAHash())
+					FGeometryShaderRHIRef GeometryShader;
+					if (PSO.GraphicsDesc.GeometryShader != FSHAHash())
+					{
+						GeometryShader = FShaderCodeLibrary::CreateGeometryShader(Platform, PSO.GraphicsDesc.GeometryShader);
+						GraphicsInitializer.BoundShaderState.SetGeometryShader(GeometryShader);
+					}
+#endif
+				}
+
+				FPixelShaderRHIRef FragmentShader;
+				if (PSO.GraphicsDesc.FragmentShader != FSHAHash())
 				{
-					GeometryShader = FShaderCodeLibrary::CreateGeometryShader(Platform, PSO.GraphicsDesc.GeometryShader);
-					GraphicsInitializer.BoundShaderState.SetGeometryShader(GeometryShader);
+					FragmentShader = FShaderCodeLibrary::CreatePixelShader(Platform, PSO.GraphicsDesc.FragmentShader);
+					GraphicsInitializer.BoundShaderState.PixelShaderRHI = FragmentShader;
 				}
-#endif
-			}
 
-			FPixelShaderRHIRef FragmentShader;
-			if (PSO.GraphicsDesc.FragmentShader != FSHAHash())
-			{
-				FragmentShader = FShaderCodeLibrary::CreatePixelShader(Platform, PSO.GraphicsDesc.FragmentShader);
-				GraphicsInitializer.BoundShaderState.PixelShaderRHI = FragmentShader;
-			}
+				auto BlendState = GetOrCreateBlendState(PSO.GraphicsDesc.BlendState);
+				GraphicsInitializer.BlendState = BlendState;
 
-			auto BlendState = GetOrCreateBlendState(PSO.GraphicsDesc.BlendState);
-			GraphicsInitializer.BlendState = BlendState;
-			
-			auto RasterState = GetOrCreateRasterizerState(PSO.GraphicsDesc.RasterizerState);
-			GraphicsInitializer.RasterizerState = RasterState;
-			
-			auto DepthState = GetOrCreateDepthStencilState(PSO.GraphicsDesc.DepthStencilState);
-			GraphicsInitializer.DepthStencilState = DepthState;
+				auto RasterState = GetOrCreateRasterizerState(PSO.GraphicsDesc.RasterizerState);
+				GraphicsInitializer.RasterizerState = RasterState;
 
-			for (uint32 i = 0; i < MaxSimultaneousRenderTargets; ++i)
-			{
-				GraphicsInitializer.RenderTargetFormats[i] = PSO.GraphicsDesc.RenderTargetFormats[i];
-				GraphicsInitializer.RenderTargetFlags[i] = PSO.GraphicsDesc.RenderTargetFlags[i];
-			}
-			
-			GraphicsInitializer.RenderTargetsEnabled = PSO.GraphicsDesc.RenderTargetsActive;
-			GraphicsInitializer.NumSamples = PSO.GraphicsDesc.MSAASamples;
+				auto DepthState = GetOrCreateDepthStencilState(PSO.GraphicsDesc.DepthStencilState);
+				GraphicsInitializer.DepthStencilState = DepthState;
 
-			GraphicsInitializer.SubpassHint = (ESubpassHint)PSO.GraphicsDesc.SubpassHint;
-			GraphicsInitializer.SubpassIndex = PSO.GraphicsDesc.SubpassIndex;
+				for (uint32 i = 0; i < MaxSimultaneousRenderTargets; ++i)
+				{
+					GraphicsInitializer.RenderTargetFormats[i] = PSO.GraphicsDesc.RenderTargetFormats[i];
+					GraphicsInitializer.RenderTargetFlags[i] = PSO.GraphicsDesc.RenderTargetFlags[i];
+				}
 
-			GraphicsInitializer.MultiViewCount = PSO.GraphicsDesc.MultiViewCount;
-			GraphicsInitializer.bHasFragmentDensityAttachment = PSO.GraphicsDesc.bHasFragmentDensityAttachment;
+				GraphicsInitializer.RenderTargetsEnabled = PSO.GraphicsDesc.RenderTargetsActive;
+				GraphicsInitializer.NumSamples = PSO.GraphicsDesc.MSAASamples;
 
-			GraphicsInitializer.DepthStencilTargetFormat = PSO.GraphicsDesc.DepthStencilFormat;
-			GraphicsInitializer.DepthStencilTargetFlag = PSO.GraphicsDesc.DepthStencilFlags;
-			GraphicsInitializer.DepthTargetLoadAction = PSO.GraphicsDesc.DepthLoad;
-			GraphicsInitializer.StencilTargetLoadAction = PSO.GraphicsDesc.StencilLoad;
-			GraphicsInitializer.DepthTargetStoreAction = PSO.GraphicsDesc.DepthStore;
-			GraphicsInitializer.StencilTargetStoreAction = PSO.GraphicsDesc.StencilStore;
-			
-			GraphicsInitializer.PrimitiveType = PSO.GraphicsDesc.PrimitiveType;
+				GraphicsInitializer.SubpassHint = (ESubpassHint)PSO.GraphicsDesc.SubpassHint;
+				GraphicsInitializer.SubpassIndex = PSO.GraphicsDesc.SubpassIndex;
+
+				GraphicsInitializer.MultiViewCount = PSO.GraphicsDesc.MultiViewCount;
+				GraphicsInitializer.bHasFragmentDensityAttachment = PSO.GraphicsDesc.bHasFragmentDensityAttachment;
+
+				GraphicsInitializer.DepthStencilTargetFormat = PSO.GraphicsDesc.DepthStencilFormat;
+				GraphicsInitializer.DepthStencilTargetFlag = PSO.GraphicsDesc.DepthStencilFlags;
+				GraphicsInitializer.DepthTargetLoadAction = PSO.GraphicsDesc.DepthLoad;
+				GraphicsInitializer.StencilTargetLoadAction = PSO.GraphicsDesc.StencilLoad;
+				GraphicsInitializer.DepthTargetStoreAction = PSO.GraphicsDesc.DepthStore;
+				GraphicsInitializer.StencilTargetStoreAction = PSO.GraphicsDesc.StencilStore;
 			
-			// This indicates we do not want a fatal error if this compilation fails
-			// (ie, if this entry in the file cache is bad)
-			GraphicsInitializer.bFromPSOFileCache = true;
+				GraphicsInitializer.PrimitiveType = PSO.GraphicsDesc.PrimitiveType;
 			
+				// This indicates we do not want a fatal error if this compilation fails
+				// (ie, if this entry in the file cache is bad)
+				GraphicsInitializer.bFromPSOFileCache = true;
+			}
 			// Use SetGraphicsPipelineState to call down into PipelineStateCache and also handle the fallback case used by OpenGL.
-			SetGraphicsPipelineState(RHICmdList, GraphicsInitializer, 0, EApplyRendertargetOption::DoNothing, false);
+			// we lose track of our PSO precompile jobs through here. they may not be RHI dependencies and could persist beyond the next tick.
+			// GetNumActivePipelinePrecompileTasks() is tracking non-rhi dependent precompiles.
+			// It's being used for rate limiting and ensuring tasks are done before calling OnPrecompilationComplete delegate.
+			{
+				QUICK_SCOPE_CYCLE_COUNTER(STAT_FShaderPipelineCache_SetGraphicsPipelineState);
+				SetGraphicsPipelineState(RHICmdList, GraphicsInitializer, 0, EApplyRendertargetOption::DoNothing, false);
+			}
 			bOk = true;
 		}
 		else if(FPipelineCacheFileFormatPSO::DescriptorType::Compute == PSO.Type)
@@ -985,6 +993,20 @@ void FShaderPipelineCache::PreparePipelineBatch(TDoubleLinkedList<FPipelineCache
 	}
 }
 
+
+static int32 GMinPipelinePrecompileTasksInFlightThreshold = 10;
+
+static FAutoConsoleVariableRef GPipelinePrecompileTasksInFlightVar(
+	TEXT("shaderpipeline.MinPrecompileTasksInFlight"),
+	GMinPipelinePrecompileTasksInFlightThreshold,
+	TEXT("Note: Only used when threadpool PSO precompiling is active.\n")
+	TEXT("The number of active PSO precompile jobs in flight before we will submit another batch of jobs. \n")
+	TEXT("i.e. when the number of inflight precompile tasks drops below this threshold we can add the next batch of precompile tasks. \n")
+	TEXT("This is to prevent bubbles between precompile batches but also to avoid saturating dispatch.")
+	,
+	ECVF_RenderThreadSafe
+);
+
 bool FShaderPipelineCache::ReadyForPrecompile()
 {
 	for(uint32 i = 0; i < (uint32)ReadTasks.Num();/*NOP*/)
@@ -1005,7 +1027,9 @@ bool FShaderPipelineCache::ReadyForPrecompile()
 		LastPrecompileRHIFence = nullptr;
 	}
 
-	return (CompileTasks.Num() != 0) && !LastPrecompileRHIFence.GetReference();
+	return (CompileTasks.Num() != 0) 
+		&& PipelineStateCache::GetNumActivePipelinePrecompileTasks() < GMinPipelinePrecompileTasksInFlightThreshold
+		&& !LastPrecompileRHIFence.GetReference();
 }
 
 void FShaderPipelineCache::PrecompilePipelineBatch()
@@ -1271,6 +1295,7 @@ bool FShaderPipelineCache::IsTickable() const
 		(FPlatformAtomics::AtomicRead(&TotalActiveTasks) != 0 ||
 			FPlatformAtomics::AtomicRead(&TotalWaitingTasks) != 0 ||
 			FPlatformAtomics::AtomicRead(&TotalCompleteTasks) != 0 ||
+			PipelineStateCache::GetNumActivePipelinePrecompileTasks()||
 			ReadyForAutoSave()||
 			GetShaderPipelineCacheSaveBoundPSOLog());
 }
@@ -1290,7 +1315,7 @@ void FShaderPipelineCache::Tick( float DeltaTime )
 	}
 
 	// HACK: note for the time being until the code further upstream properly uses the OnPrecompilationComplete delegate, we need to provide some delays to make sure it gets the messaging properly (the -10.0f and the -20.0f below).
-	if (((FMath::Max(0ll, FPlatformAtomics::AtomicRead(&TotalWaitingTasks)) == 0 && FMath::Max(0ll, FPlatformAtomics::AtomicRead(&TotalActiveTasks)) == 0 && FPlatformAtomics::AtomicRead(&TotalCompleteTasks) != 0) || (CVarPSOFileCacheMaxPrecompileTime.GetValueOnAnyThread() > 0.0 && TotalPrecompileWallTime - 10.0f > CVarPSOFileCacheMaxPrecompileTime.GetValueOnAnyThread() && TotalPrecompileTasks > 0)) && !LastPrecompileRHIFence.GetReference())
+	if (((PipelineStateCache::GetNumActivePipelinePrecompileTasks() == 0 && FMath::Max(0ll, FPlatformAtomics::AtomicRead(&TotalWaitingTasks)) == 0 && FMath::Max(0ll, FPlatformAtomics::AtomicRead(&TotalActiveTasks)) == 0 && FPlatformAtomics::AtomicRead(&TotalCompleteTasks) != 0) || (CVarPSOFileCacheMaxPrecompileTime.GetValueOnAnyThread() > 0.0 && TotalPrecompileWallTime - 10.0f > CVarPSOFileCacheMaxPrecompileTime.GetValueOnAnyThread() && TotalPrecompileTasks > 0)) && !LastPrecompileRHIFence.GetReference())
     {
 		UE_LOG(LogRHI, Warning, TEXT("FShaderPipelineCache completed %u tasks in %.2fs (%.2fs wall time since intial open)."), (uint32)TotalCompleteTasks, FPlatformTime::ToSeconds64(TotalPrecompileTime), TotalPrecompileWallTime);
         if (OnPrecompilationComplete.IsBound())
@@ -1306,7 +1331,8 @@ void FShaderPipelineCache::Tick( float DeltaTime )
 		else
 		{
 			FPipelineFileCache::PreCompileComplete();
-		
+			PipelineStateCache::PreCompileComplete();
+
         	FPlatformAtomics::InterlockedExchange(&TotalCompleteTasks, 0);
         	FPlatformAtomics::InterlockedExchange(&TotalPrecompileTime, 0);
 			bPreOptimizing = false;
