Fixing issue where RHICreateTargetableShaderResource would overwrite all the flags from the passed in texture desc.

diff --git a/Engine/Source/Runtime/RHI/Private/RHIUtilities.cpp b/Engine/Source/Runtime/RHI/Private/RHIUtilities.cpp
--- a/Engine/Source/Runtime/RHI/Private/RHIUtilities.cpp
+++ b/Engine/Source/Runtime/RHI/Private/RHIUtilities.cpp
@@ -613,7 +613,7 @@ void RHICreateTargetableShaderResource(
 
 	FRHITextureCreateDesc Desc =
 		FRHITextureCreateDesc(BaseDesc)
-		.SetFlags(TargetableTextureFlags | ETextureCreateFlags::ShaderResource)
+		.AddFlags(TargetableTextureFlags | ETextureCreateFlags::ShaderResource)
 		.SetInitialState(ERHIAccess::SRVMask);
 
 	OutTargetableTexture = RHICreateTexture(Desc);
