- Changed driver update dialog box to Yes/No where Yes will open the URL

diff --git a/Engine/Source/Runtime/RHI/Private/DynamicRHI.cpp b/Engine/Source/Runtime/RHI/Private/DynamicRHI.cpp
--- a/Engine/Source/Runtime/RHI/Private/DynamicRHI.cpp
+++ b/Engine/Source/Runtime/RHI/Private/DynamicRHI.cpp
@@ -187,27 +187,31 @@ static void RHIDetectAndWarnOfBadDrivers(bool bHasEditorToken)
 			{
 				if (!DenyListEntry.RHIName.IsEmpty())
 				{
-					LocalizedMsg = FText::Format(NSLOCTEXT("MessageDialog", "LatestVideoCardDriverRHIIssueReport", "The latest version of the {Vendor} graphics driver has known issues in {RHI}.\nPlease install the recommended driver version or switch to a different rendering API.\n\n{Hyperlink}\n\n{AdapterName}\nInstalled: {InstalledVer}\nRecommended: {RecommendedVer}"), Args);
+					LocalizedMsg = FText::Format(NSLOCTEXT("MessageDialog", "LatestVideoCardDriverRHIIssueReport", "The latest version of the {Vendor} graphics driver has known issues in {RHI}.\nPlease install the recommended driver version or switch to a different rendering API.\n\nWould you like to visit the following URL to download the driver?\n\n{Hyperlink}\n\n{AdapterName}\nInstalled: {InstalledVer}\nRecommended: {RecommendedVer}"), Args);
 				}
 				else
 				{
-					LocalizedMsg = FText::Format(NSLOCTEXT("MessageDialog", "LatestVideoCardDriverIssueReport", "The latest version of the {Vendor} graphics driver has known issues.\nPlease install the recommended driver version.\n\n{Hyperlink}\n\n{AdapterName}\nInstalled: {InstalledVer}\nRecommended: {RecommendedVer}"), Args);
+					LocalizedMsg = FText::Format(NSLOCTEXT("MessageDialog", "LatestVideoCardDriverIssueReport", "The latest version of the {Vendor} graphics driver has known issues.\nPlease install the recommended driver version.\n\nWould you like to visit the following URL to download the driver?\n\n{Hyperlink}\n\n{AdapterName}\nInstalled: {InstalledVer}\nRecommended: {RecommendedVer}"), Args);
 				}
 			}
 			else
 			{
 				if (!DenyListEntry.RHIName.IsEmpty())
 				{
-					LocalizedMsg = FText::Format(NSLOCTEXT("MessageDialog", "VideoCardDriverRHIIssueReport", "The installed version of the {Vendor} graphics driver has known issues in {RHI}.\nPlease install either the latest or the recommended driver version or switch to a different rendering API.\n\n{Hyperlink}\n\n{AdapterName}\nInstalled: {InstalledVer}\nRecommended: {RecommendedVer}"), Args);
+					LocalizedMsg = FText::Format(NSLOCTEXT("MessageDialog", "VideoCardDriverRHIIssueReport", "The installed version of the {Vendor} graphics driver has known issues in {RHI}.\nPlease install either the latest or the recommended driver version or switch to a different rendering API.\n\nWould you like to visit the following URL to download the driver?\n\n{Hyperlink}\n\n{AdapterName}\nInstalled: {InstalledVer}\nRecommended: {RecommendedVer}"), Args);
 				}
 				else
 				{
-					LocalizedMsg = FText::Format(NSLOCTEXT("MessageDialog", "VideoCardDriverIssueReport", "The installed version of the {Vendor} graphics driver has known issues.\nPlease install either the latest or the recommended driver version.\n\n{Hyperlink}\n\n{AdapterName}\nInstalled: {InstalledVer}\nRecommended: {RecommendedVer}"), Args);
+					LocalizedMsg = FText::Format(NSLOCTEXT("MessageDialog", "VideoCardDriverIssueReport", "The installed version of the {Vendor} graphics driver has known issues.\nPlease install either the latest or the recommended driver version.\n\nWould you like to visit the following URL to download the driver?\n\n{Hyperlink}\n\n{AdapterName}\nInstalled: {InstalledVer}\nRecommended: {RecommendedVer}"), Args);
 				}
 			}
 
 			FText Title = NSLOCTEXT("MessageDialog", "TitleVideoCardDriverIssue", "WARNING: Known issues with graphics driver");
-			FMessageDialog::Open(EAppMsgType::Ok, LocalizedMsg, &Title);
+			EAppReturnType::Type Response = FMessageDialog::Open(EAppMsgType::YesNo, LocalizedMsg, &Title);
+			if (Response == EAppReturnType::Yes)
+			{
+				FPlatformProcess::LaunchURL(*HyperlinkText.ToString(), nullptr, nullptr);
+			}
 		}
 	}
 }
