Fixed texture 2D array resource creation

diff --git a/Engine/Source/Runtime/RHI/Public/RHIResources.h b/Engine/Source/Runtime/RHI/Public/RHIResources.h
--- a/Engine/Source/Runtime/RHI/Public/RHIResources.h
+++ b/Engine/Source/Runtime/RHI/Public/RHIResources.h
@@ -1585,7 +1585,7 @@ struct FRHITextureCreateDesc : public FRHITextureDesc
 	static FRHITextureCreateDesc Create2DArray(const TCHAR* DebugName, int32 SizeX, int32 SizeY, int32 ArraySize, EPixelFormat Format)
 	{
 		return Create2DArray(DebugName)
-			.SetExtent(SizeY, SizeY)
+			.SetExtent(SizeX, SizeY)
 			.SetFormat(Format)
 			.SetArraySize(ArraySize);
 	}
