Adds config to set default RHI by GPU name. DeviceID would be better but it requires a refactor since RHI Init gets the DeviceID but we need it prior to choosing the RHI.

diff --git a/Engine/Source/Runtime/RHI/Private/Windows/WindowsDynamicRHI.cpp b/Engine/Source/Runtime/RHI/Private/Windows/WindowsDynamicRHI.cpp
--- a/Engine/Source/Runtime/RHI/Private/Windows/WindowsDynamicRHI.cpp
+++ b/Engine/Source/Runtime/RHI/Private/Windows/WindowsDynamicRHI.cpp
@@ -167,6 +167,30 @@ static bool DefaultFeatureLevelES31()
 		}
 	}
 
+	TArray<FString> DeviceDefaultRHIList;
+	GConfig->GetArray(TEXT("Devices"), TEXT("DeviceDefaultRHIList"), DeviceDefaultRHIList, GHardwareIni);
+
+	FString GPUBrand = FPlatformMisc::GetPrimaryGPUBrand();
+	for (const FString& DeviceDefaultRHIString : DeviceDefaultRHIList)
+	{
+		const TCHAR* Line = *DeviceDefaultRHIString;
+
+		ensure(Line[0] == TCHAR('('));
+
+		FString RHIName;
+		FParse::Value(Line+1, TEXT("RHI="), RHIName);
+
+		FString DeviceName;
+		FParse::Value(Line+1, TEXT("DeviceName="), DeviceName);
+
+		if (RHIName.Compare("D3D11_ES31", ESearchCase::IgnoreCase) == 0 && GPUBrand.Compare(DeviceName, ESearchCase::IgnoreCase) == 0)
+		{
+			ForceES31 = true;
+
+			return true;
+		}
+	}
+
 	ForceES31 = false;
 	return false;
 }
