Implements FGenericDataDrivenShaderPlatformInfo::GetSupportsRealTypes

diff --git a/Engine/Config/Windows/DataDrivenPlatformInfo.ini b/Engine/Config/Windows/DataDrivenPlatformInfo.ini
--- a/Engine/Config/Windows/DataDrivenPlatformInfo.ini
+++ b/Engine/Config/Windows/DataDrivenPlatformInfo.ini
@@ -139,6 +139,7 @@ bSupportsIndexBufferUAVs=true
 bSupportsBindless=true
 bSupportsROV=true
 bSupportsOIT=true
+bSupportsRealTypes=RuntimeDependent
 
 [PreviewPlatform PCD3D_SM6]
 PlatformName=PCD3D_SM6
diff --git a/Engine/Shaders/Private/TemporalSuperResolution/TSRCommon.ush b/Engine/Shaders/Private/TemporalSuperResolution/TSRCommon.ush
--- a/Engine/Shaders/Private/TemporalSuperResolution/TSRCommon.ush
+++ b/Engine/Shaders/Private/TemporalSuperResolution/TSRCommon.ush
@@ -11,7 +11,7 @@
 
 
 //------------------------------------------------------- RECOMPILE HASH
-// 7C48165B-DF1A-4FA1-88CF-F64C6B7DF2D7
+// 88217342-DCE7-451F-8670-0A14BC8B6109
 
 
 //------------------------------------------------------- COMPILER CONFIG
diff --git a/Engine/Source/Runtime/RHI/Private/RHI.cpp b/Engine/Source/Runtime/RHI/Private/RHI.cpp
--- a/Engine/Source/Runtime/RHI/Private/RHI.cpp
+++ b/Engine/Source/Runtime/RHI/Private/RHI.cpp
@@ -2280,6 +2280,34 @@ static inline uint32 GetSectionUint(const FConfigSection& Section, FName Key, ui
 	}
 }
 
+// Gets an integer from a section.  It returns the original value if the setting does not exist
+static inline uint32 GetSectionFeatureSupport(const FConfigSection& Section, FName Key, uint32 OriginalValue)
+{
+	const FConfigValue* ConfigValue = Section.Find(Key);
+	if (ConfigValue != nullptr)
+	{
+		FString Value = ConfigValue->GetValue();
+		if (Value == TEXT("Unsupported"))
+		{
+			return uint32(ERHIFeatureSupport::Unsupported);
+		}
+		if (Value == TEXT("RuntimeDependent"))
+		{
+			return uint32(ERHIFeatureSupport::RuntimeDependent);
+		}
+		else if (Value == TEXT("RuntimeGuaranteed"))
+		{
+			return uint32(ERHIFeatureSupport::RuntimeGuaranteed);
+		}
+		else
+		{
+			checkf(false, TEXT("Unknown ERHIFeatureSupport value \"%s\" for %s"), *Value, *Key.ToString());
+		}
+	}
+
+	return OriginalValue;
+}
+
 FRHIViewableResource* GetViewableResource(const FRHITransitionInfo& Info)
 {
 	switch (Info.Type)
@@ -2320,6 +2348,8 @@ void FGenericDataDrivenShaderPlatformInfo::ParseDataDrivenShaderInfo(const FConf
 	Info.SettingName = GetSectionBool(Section, #SettingName, Info.SettingName)
 #define GET_SECTION_INT_HELPER(SettingName)	\
 	Info.SettingName = GetSectionUint(Section, #SettingName, Info.SettingName)
+#define GET_SECTION_SUPPORT_HELPER(SettingName)	\
+	Info.SettingName = GetSectionFeatureSupport(Section, #SettingName, Info.SettingName)
 
 	GET_SECTION_BOOL_HELPER(bIsMobile);
 	GET_SECTION_BOOL_HELPER(bIsMetalMRT);
@@ -2401,8 +2431,10 @@ void FGenericDataDrivenShaderPlatformInfo::ParseDataDrivenShaderInfo(const FConf
 	GET_SECTION_BOOL_HELPER(bSupportsVolumeTextureAtomics);
 	GET_SECTION_BOOL_HELPER(bSupportsROV);
 	GET_SECTION_BOOL_HELPER(bSupportsOIT);
+	GET_SECTION_SUPPORT_HELPER(bSupportsRealTypes);
 #undef GET_SECTION_BOOL_HELPER
 #undef GET_SECTION_INT_HELPER
+#undef GET_SECTION_SUPPORT_HELPER
 
 #if WITH_EDITOR
 	FTextStringHelper::ReadFromBuffer(*GetSectionString(Section, FName("FriendlyName")), Info.FriendlyName);
diff --git a/Engine/Source/Runtime/RHI/Public/RHIDefinitions.h b/Engine/Source/Runtime/RHI/Public/RHIDefinitions.h
--- a/Engine/Source/Runtime/RHI/Public/RHIDefinitions.h
+++ b/Engine/Source/Runtime/RHI/Public/RHIDefinitions.h
@@ -36,6 +36,21 @@ enum class ERHIInterfaceType
 	OpenGL,
 };
 
+enum class ERHIFeatureSupport : uint8
+{
+	// The RHI feature is completely unavailable at runtime
+	Unsupported,
+
+	// The RHI feature can be available at runtime based on hardware or driver
+	RuntimeDependent,
+
+	// The RHI feature is guaranteed to be available at runtime.
+	RuntimeGuaranteed,
+
+	Num,
+	NumBits = 2,
+};
+
 enum EShaderFrequency : uint8
 {
 	SF_Vertex			= 0,
@@ -399,6 +414,7 @@ class RHI_API FGenericDataDrivenShaderPlatformInfo
 	uint32 bSupportsVolumeTextureAtomics : 1;
 	uint32 bSupportsROV : 1;
 	uint32 bSupportsOIT : 1;
+	uint32 bSupportsRealTypes : int32(ERHIFeatureSupport::NumBits);
 		
 #if WITH_EDITOR
 	FText FriendlyName;
@@ -932,6 +948,11 @@ public:
 		return Infos[Platform].bSupportsOIT;
 	}
 
+	static FORCEINLINE_DEBUGGABLE const ERHIFeatureSupport GetSupportsRealTypes(const FStaticShaderPlatform Platform)
+	{
+		return ERHIFeatureSupport(Infos[Platform].bSupportsRealTypes);
+	}
+
 #if WITH_EDITOR
 	static FORCEINLINE_DEBUGGABLE FText GetFriendlyName(const FStaticShaderPlatform Platform)
 	{
diff --git a/Engine/Source/Runtime/RenderCore/Public/ShaderCompilerCore.h b/Engine/Source/Runtime/RenderCore/Public/ShaderCompilerCore.h
--- a/Engine/Source/Runtime/RenderCore/Public/ShaderCompilerCore.h
+++ b/Engine/Source/Runtime/RenderCore/Public/ShaderCompilerCore.h
@@ -110,7 +110,7 @@ enum ECompilerFlags
 	CFLAG_InlineRayTracing,
 	// Force using the SC rewrite functionality before calling DXC on D3D12
 	CFLAG_D3D12ForceShaderConductorRewrite,
-	// Enable support of C-style data types for platforms that can. Check for PLATFORM_SUPPORTS_REAL_TYPES.
+	// Enable support of C-style data types for platforms that can. Check for PLATFORM_SUPPORTS_REAL_TYPES and FDataDrivenShaderPlatformInfo::GetSupportsRealTypes()
 	CFLAG_AllowRealTypes,
 	// Precompile HLSL to optimized HLSL, then forward to FXC. Speeds up some shaders that take longer with FXC and works around crashes in FXC.
 	CFLAG_PrecompileWithDXC,
diff --git a/Engine/Source/Runtime/Renderer/Private/PostProcess/TemporalSuperResolution.cpp b/Engine/Source/Runtime/Renderer/Private/PostProcess/TemporalSuperResolution.cpp
--- a/Engine/Source/Runtime/Renderer/Private/PostProcess/TemporalSuperResolution.cpp
+++ b/Engine/Source/Runtime/Renderer/Private/PostProcess/TemporalSuperResolution.cpp
@@ -176,7 +176,10 @@ public:
 
 	static void ModifyCompilationEnvironment(const FGlobalShaderPermutationParameters& Parameters, FShaderCompilerEnvironment& OutEnvironment)
 	{
-		OutEnvironment.CompilerFlags.Add(CFLAG_AllowRealTypes);
+		if (FDataDrivenShaderPlatformInfo::GetSupportsRealTypes(Parameters.Platform) == ERHIFeatureSupport::RuntimeGuaranteed)
+		{
+			OutEnvironment.CompilerFlags.Add(CFLAG_AllowRealTypes);
+		}
 		OutEnvironment.CompilerFlags.Add(CFLAG_WarningsAsErrors);
 		OutEnvironment.CompilerFlags.Add(CFLAG_HLSL2021);
 	}
