remove RHIEnqueueDecompress from public API

diff --git a/Engine/Source/Runtime/NullDrv/Public/NullRHI.h b/Engine/Source/Runtime/NullDrv/Public/NullRHI.h
--- a/Engine/Source/Runtime/NullDrv/Public/NullRHI.h
+++ b/Engine/Source/Runtime/NullDrv/Public/NullRHI.h
@@ -596,10 +596,6 @@ public:
 	virtual void RHIBlockUntilGPUIdle() final override
 	{
 	}
-	virtual bool RHIEnqueueDecompress(uint8_t*, uint8_t*, int, void*) final override
-	{
-		return false;
-	}
 	virtual bool RHIGetAvailableResolutions(FScreenResolutionArray& Resolutions, bool bIgnoreRefreshRate) final override
 	{ 
 		return false; 
diff --git a/Engine/Source/Runtime/RHI/Public/DynamicRHI.h b/Engine/Source/Runtime/RHI/Public/DynamicRHI.h
--- a/Engine/Source/Runtime/RHI/Public/DynamicRHI.h
+++ b/Engine/Source/Runtime/RHI/Public/DynamicRHI.h
@@ -925,9 +925,6 @@ public:
 	// FlushType: Flush Immediate
 	virtual bool RHIIsRenderingSuspended() { return false; };
 
-	// FlushType: Flush Immediate
-	virtual bool RHIEnqueueDecompress(uint8_t* SrcBuffer, uint8_t* DestBuffer, int CompressedSize, void* ErrorCodeBuffer) { return false; }
-
 	/**
 	*	Retrieve available screen resolutions.
 	*
diff --git a/Engine/Source/Runtime/RHI/Public/RHICommandList.h b/Engine/Source/Runtime/RHI/Public/RHICommandList.h
--- a/Engine/Source/Runtime/RHI/Public/RHICommandList.h
+++ b/Engine/Source/Runtime/RHI/Public/RHICommandList.h
@@ -4720,12 +4720,7 @@ PRAGMA_ENABLE_DEPRECATION_WARNINGS
 		ImmediateFlush(EImmediateFlushType::FlushRHIThread); 
 		return GDynamicRHI->RHIIsRenderingSuspended();
 	}
-
-	FORCEINLINE bool EnqueueDecompress(uint8_t* SrcBuffer, uint8_t* DestBuffer, int CompressedSize, void* ErrorCodeBuffer)
-	{
-		return GDynamicRHI->RHIEnqueueDecompress(SrcBuffer, DestBuffer, CompressedSize, ErrorCodeBuffer);
-	}
-		
+	
 	FORCEINLINE bool GetAvailableResolutions(FScreenResolutionArray& Resolutions, bool bIgnoreRefreshRate)
 	{
 		return RHIGetAvailableResolutions(Resolutions, bIgnoreRefreshRate);
diff --git a/Engine/Source/Runtime/RHI/Public/RHIValidation.h b/Engine/Source/Runtime/RHI/Public/RHIValidation.h
--- a/Engine/Source/Runtime/RHI/Public/RHIValidation.h
+++ b/Engine/Source/Runtime/RHI/Public/RHIValidation.h
@@ -1113,12 +1113,6 @@ public:
 		return RHI->RHIIsRenderingSuspended();
 	}
 
-	// FlushType: Flush Immediate
-	virtual bool RHIEnqueueDecompress(uint8_t* SrcBuffer, uint8_t* DestBuffer, int CompressedSize, void* ErrorCodeBuffer) override final
-	{
-		return RHI->RHIEnqueueDecompress(SrcBuffer, DestBuffer, CompressedSize, ErrorCodeBuffer);
-	}
-
 	/**
 	*	Retrieve available screen resolutions.
 	*
