Variable Refresh Rate Support for platforms that support it - Enable with rhi.SyncAllowVariable

diff --git a/Engine/Source/Runtime/RHI/Private/RHIUtilities.cpp b/Engine/Source/Runtime/RHI/Private/RHIUtilities.cpp
--- a/Engine/Source/Runtime/RHI/Private/RHIUtilities.cpp
+++ b/Engine/Source/Runtime/RHI/Private/RHIUtilities.cpp
@@ -116,6 +116,13 @@ TAutoConsoleVariable<float> CVarRHISyncSlackMS(
 );
 #endif
 
+TAutoConsoleVariable<int32> CVarRHISyncAllowVariable(
+	TEXT("rhi.SyncAllowVariable"),
+	1,
+	TEXT("When 1, allows the RHI to use variable refresh rate, if supported by the output hardware."),
+	ECVF_Default
+);
+
 class FRHIFrameFlipTrackingRunnable : public FRunnable
 {
 	static FRunnableThread* Thread;
@@ -490,6 +497,11 @@ RHI_API float RHIGetSyncSlackMS()
 	return SyncSlackMS;
 }
 
+RHI_API bool RHIGetSyncAllowVariable()
+{
+	return CVarRHISyncAllowVariable.GetValueOnAnyThread() != 0;
+}
+
 RHI_API void RHIGetPresentThresholds(float& OutTopPercent, float& OutBottomPercent)
 {
 	OutTopPercent = FMath::Clamp(CVarRHIPresentThresholdTop.GetValueOnAnyThread(), 0.0f, 1.0f);
diff --git a/Engine/Source/Runtime/RHI/Public/RHIUtilities.h b/Engine/Source/Runtime/RHI/Public/RHIUtilities.h
--- a/Engine/Source/Runtime/RHI/Public/RHIUtilities.h
+++ b/Engine/Source/Runtime/RHI/Public/RHIUtilities.h
@@ -890,6 +890,9 @@ extern RHI_API float RHIGetSyncSlackMS();
 /** Returns the top and bottom vsync present thresholds (the values of rhi.PresentThreshold.Top and rhi.PresentThreshold.Bottom) */
 extern RHI_API void RHIGetPresentThresholds(float& OutTopPercent, float& OutBottomPercent);
 
+/** Returns the value of the rhi.SyncAllowVariable CVar. */
+extern RHI_API bool RHIGetSyncAllowVariable();
+
 /** Signals the completion of the specified task graph event when the given frame has flipped. */
 extern RHI_API void RHICompleteGraphEventOnFlip(uint64 PresentIndex, FGraphEventRef Event);
 
