Mark RHIEnqueueDecompress as deprecated instead of removing it

diff --git a/Engine/Source/Runtime/RHI/Public/DynamicRHI.h b/Engine/Source/Runtime/RHI/Public/DynamicRHI.h
--- a/Engine/Source/Runtime/RHI/Public/DynamicRHI.h
+++ b/Engine/Source/Runtime/RHI/Public/DynamicRHI.h
@@ -925,6 +925,9 @@ public:
 	// FlushType: Flush Immediate
 	virtual bool RHIIsRenderingSuspended() { return false; };
 
+	UE_DEPRECATED(5.1, "No longer used: FCompression::UncompressMemory should be used instead")
+	virtual bool RHIEnqueueDecompress(uint8_t* SrcBuffer, uint8_t* DestBuffer, int CompressedSize, void* ErrorCodeBuffer) { return false; }
+
 	/**
 	*	Retrieve available screen resolutions.
 	*
diff --git a/Engine/Source/Runtime/RHI/Public/RHICommandList.h b/Engine/Source/Runtime/RHI/Public/RHICommandList.h
--- a/Engine/Source/Runtime/RHI/Public/RHICommandList.h
+++ b/Engine/Source/Runtime/RHI/Public/RHICommandList.h
@@ -4721,6 +4721,12 @@ PRAGMA_ENABLE_DEPRECATION_WARNINGS
 		return GDynamicRHI->RHIIsRenderingSuspended();
 	}
 	
+	UE_DEPRECATED(5.1, "No longer used: FCompression::UncompressMemory should be used instead")
+	FORCEINLINE bool EnqueueDecompress(uint8_t* SrcBuffer, uint8_t* DestBuffer, int CompressedSize, void* ErrorCodeBuffer)
+	{
+		return false;
+	}
+
 	FORCEINLINE bool GetAvailableResolutions(FScreenResolutionArray& Resolutions, bool bIgnoreRefreshRate)
 	{
 		return RHIGetAvailableResolutions(Resolutions, bIgnoreRefreshRate);
