Fix crash on startup in packaged builds when Nanite mesh shader code path is used.

FPipelineCacheFileFormatPSO::Init() used to assume that a vertex declaration is present in the pipeline initializer. However, this may not be the case when mesh shaders are used.

diff --git a/Engine/Source/Runtime/RHI/Private/PipelineFileCache.cpp b/Engine/Source/Runtime/RHI/Private/PipelineFileCache.cpp
--- a/Engine/Source/Runtime/RHI/Private/PipelineFileCache.cpp
+++ b/Engine/Source/Runtime/RHI/Private/PipelineFileCache.cpp
@@ -1360,45 +1360,50 @@ FPipelineCacheFileFormatPSO::FPipelineCacheFileFormatPSO()
 	
 	// Because of the cheat in the copy constructor - lets play this safe
 	FMemory::Memset(&PSO.GraphicsDesc, 0, sizeof(GraphicsDescriptor));
-	
-	check (Init.BoundShaderState.VertexDeclarationRHI);
-	check (Init.BoundShaderState.VertexDeclarationRHI->IsValid());
+
+#if PLATFORM_SUPPORTS_MESH_SHADERS
+	checkf(Init.BoundShaderState.GetVertexShader() || Init.BoundShaderState.GetMeshShader(), TEXT("A graphics pipeline must always have either a vertex or a mesh shader"));
+	if (Init.BoundShaderState.GetVertexShader())
+#else
+	checkf(Init.BoundShaderState.GetVertexShader(), TEXT("A graphics pipeline must always have a vertex shader"));
+#endif
 	{
-		bOK &= Init.BoundShaderState.VertexDeclarationRHI->GetInitializer(PSO.GraphicsDesc.VertexDescriptor);
-		check(bOK);
-		
-		PSO.GraphicsDesc.VertexDescriptor.Sort([](FVertexElement const& A, FVertexElement const& B)
+		check (Init.BoundShaderState.VertexDeclarationRHI);
+		check (Init.BoundShaderState.VertexDeclarationRHI->IsValid());
 		{
-			if (A.StreamIndex < B.StreamIndex)
-			{
-				return true;
-			}
-			if (A.StreamIndex > B.StreamIndex)
-			{
-				return false;
-			}
-			if (A.Offset < B.Offset)
-			{
-				return true;
-			}
-			if (A.Offset > B.Offset)
-			{
-				return false;
-			}
-			if (A.AttributeIndex < B.AttributeIndex)
-			{
-				return true;
-			}
-			if (A.AttributeIndex > B.AttributeIndex)
+			bOK &= Init.BoundShaderState.VertexDeclarationRHI->GetInitializer(PSO.GraphicsDesc.VertexDescriptor);
+			check(bOK);
+		
+			PSO.GraphicsDesc.VertexDescriptor.Sort([](FVertexElement const& A, FVertexElement const& B)
 			{
+				if (A.StreamIndex < B.StreamIndex)
+				{
+					return true;
+				}
+				if (A.StreamIndex > B.StreamIndex)
+				{
+					return false;
+				}
+				if (A.Offset < B.Offset)
+				{
+					return true;
+				}
+				if (A.Offset > B.Offset)
+				{
+					return false;
+				}
+				if (A.AttributeIndex < B.AttributeIndex)
+				{
+					return true;
+				}
+				if (A.AttributeIndex > B.AttributeIndex)
+				{
+					return false;
+				}
 				return false;
-			}
-			return false;
-		});
-	}
-	
-	if (Init.BoundShaderState.VertexShaderRHI)
-	{
+			});
+		}
+
 		PSO.GraphicsDesc.VertexShader = Init.BoundShaderState.VertexShaderRHI->GetHash();
 	}
 
