only warn about unknown data-driven shader platforms if the owning platform is enabled.

#rnx

diff --git a/Engine/Source/Runtime/RHI/Private/RHI.cpp b/Engine/Source/Runtime/RHI/Private/RHI.cpp
--- a/Engine/Source/Runtime/RHI/Private/RHI.cpp
+++ b/Engine/Source/Runtime/RHI/Private/RHI.cpp
@@ -2470,7 +2470,12 @@ void FGenericDataDrivenShaderPlatformInfo::Initialize()
 				LexFromString(ShaderPlatform, *SectionName.Mid(15));
 				if (ShaderPlatform == EShaderPlatform::SP_NumPlatforms)
 				{
-					UE_LOG(LogRHI, Warning, TEXT("Found an unknown shader platform %s in a DataDriven ini file"), *SectionName.Mid(15));
+#if DDPI_HAS_EXTENDED_PLATFORMINFO_DATA
+					const bool bIsEnabled = FDataDrivenPlatformInfoRegistry::GetPlatformInfo(PlatformName).bEnabledForUse;
+#else
+					const bool bIsEnabled = true;
+#endif
+					UE_CLOG(bIsEnabled, LogRHI, Warning, TEXT("Found an unknown shader platform %s in a DataDriven ini file"), *SectionName.Mid(15));
 					continue;
 				}
 				
