Added GetTypeHash implementations for RHI view descriptors and packed FRHITextureSRVCreateInfo more effectively to reduce space.

diff --git a/Engine/Source/Runtime/RHI/Public/RHIResources.h b/Engine/Source/Runtime/RHI/Public/RHIResources.h
--- a/Engine/Source/Runtime/RHI/Public/RHIResources.h
+++ b/Engine/Source/Runtime/RHI/Public/RHIResources.h
@@ -4227,7 +4227,7 @@ enum class ERHITextureMetaDataAccess : uint8
 	CMask
 };
 
-enum ERHITextureSRVOverrideSRGBType
+enum ERHITextureSRVOverrideSRGBType : uint8
 {
 	SRGBO_Default,
 	SRGBO_ForceDisable,
@@ -4244,7 +4244,7 @@ struct FRHITextureSRVCreateInfo
 		, NumArraySlices(0)
 	{}
 
-	explicit FRHITextureSRVCreateInfo(uint8 InMipLevel, uint8 InNumMipLevels, uint32 InFirstArraySlice, uint32 InNumArraySlices, EPixelFormat InFormat = PF_Unknown)
+	explicit FRHITextureSRVCreateInfo(uint8 InMipLevel, uint8 InNumMipLevels, uint16 InFirstArraySlice, uint16 InNumArraySlices, EPixelFormat InFormat = PF_Unknown)
 		: Format(InFormat)
 		, MipLevel(InMipLevel)
 		, NumMipLevels(InNumMipLevels)
@@ -4266,10 +4266,10 @@ struct FRHITextureSRVCreateInfo
 	ERHITextureSRVOverrideSRGBType SRGBOverride;
 
 	/** Specify first array slice index. By default 0. */
-	uint32 FirstArraySlice;
+	uint16 FirstArraySlice;
 
 	/** Specify number of array slices. If FirstArraySlice and NumArraySlices are both zero, the SRV is created for all array slices. By default 0. */
-	uint32 NumArraySlices;
+	uint16 NumArraySlices;
 
 	/** Specify the metadata plane to use when creating a view. */
 	ERHITextureMetaDataAccess MetaData = ERHITextureMetaDataAccess::None;
@@ -4290,13 +4290,15 @@ struct FRHITextureSRVCreateInfo
 	{
 		return !(*this == Other);
 	}
-};
 
-FORCEINLINE uint32 GetTypeHash(const FRHITextureSRVCreateInfo& Var)
-{
-	uint32 Hash0 = uint32(Var.Format) | uint32(Var.MipLevel) << 8 | uint32(Var.NumMipLevels) << 16 | uint32(Var.SRGBOverride) << 24;
-	return HashCombine(HashCombine(GetTypeHash(Hash0), GetTypeHash(Var.FirstArraySlice)), GetTypeHash(Var.NumArraySlices));
-}
+	friend uint32 GetTypeHash(const FRHITextureSRVCreateInfo& Info)
+	{
+		uint32 Hash = uint32(Info.Format) | uint32(Info.MipLevel) << 8 | uint32(Info.NumMipLevels) << 16 | uint32(Info.SRGBOverride) << 24;
+		Hash = HashCombine(Hash, uint32(Info.FirstArraySlice) | uint32(Info.NumArraySlices) << 16);
+		Hash = HashCombine(Hash, uint32(Info.MetaData));
+		return Hash;
+	}
+};
 
 struct FRHITextureUAVCreateInfo
 {
@@ -4324,6 +4326,13 @@ public:
 		return !(*this == Other);
 	}
 
+	friend uint32 GetTypeHash(const FRHITextureUAVCreateInfo& Info)
+	{
+		uint32 Hash = uint32(Info.Format) | uint32(Info.MipLevel) << 8 | uint32(Info.FirstArraySlice) << 16;
+		Hash = HashCombine(Hash, uint32(Info.NumArraySlices) | uint32(Info.MetaData) << 16);
+		return Hash;
+	}
+
 	EPixelFormat Format = PF_Unknown;
 	uint8 MipLevel = 0;
 	uint16 FirstArraySlice = 0;
@@ -4380,6 +4389,11 @@ struct FRHIBufferSRVCreateInfo
 		return !(*this == Other);
 	}
 
+	friend uint32 GetTypeHash(const FRHIBufferSRVCreateInfo& Desc)
+	{
+		return HashCombine(uint32(Desc.BytesPerElement), uint32(Desc.Format));
+	}
+
 	/** Number of bytes per element. */
 	uint32 BytesPerElement = 1;
 
@@ -4405,6 +4419,11 @@ struct FRHIBufferUAVCreateInfo
 		return !(*this == Other);
 	}
 
+	friend uint32 GetTypeHash(const FRHIBufferUAVCreateInfo& Info)
+	{
+		return uint32(Info.Format) | uint32(Info.bSupportsAtomicCounter) << 8 | uint32(Info.bSupportsAppendBuffer) << 16;
+	}
+
 	/** Number of bytes per element (used for typed buffers). */
 	EPixelFormat Format = PF_Unknown;
 
