Add support to default RHI based on device id and a list of devices that are capable of updating to avoid prompting users who can't update their drivers

diff --git a/Engine/Source/Runtime/ApplicationCore/Private/Windows/WindowsPlatformApplicationMisc.cpp b/Engine/Source/Runtime/ApplicationCore/Private/Windows/WindowsPlatformApplicationMisc.cpp
--- a/Engine/Source/Runtime/ApplicationCore/Private/Windows/WindowsPlatformApplicationMisc.cpp
+++ b/Engine/Source/Runtime/ApplicationCore/Private/Windows/WindowsPlatformApplicationMisc.cpp
@@ -311,15 +311,16 @@ int32 FWindowsPlatformApplicationMisc::GetMonitorDPI(const FMonitorInfo& Monitor
 	return DisplayDPI;
 }
 
-// Looks for an adapter with >= 512 MB of dedicated video memory and assumes we use it.
-bool FWindowsPlatformApplicationMisc::ProbablyHasIntegratedGPU()
+// Looks for an adapter with the most dedicated video memory
+FWindowsPlatformApplicationMisc::FGPUInfo FWindowsPlatformApplicationMisc::GetBestGPUInfo()
 { 
 	TRefCountPtr<IDXGIFactory1> DXGIFactory1;
 	if (CreateDXGIFactory1(__uuidof(IDXGIFactory1), (void**)&DXGIFactory1) != S_OK || !DXGIFactory1)
 	{
-		return false;
+		return {};
 	}
 
+	DXGI_ADAPTER_DESC BestDesc = {};
 	TRefCountPtr<IDXGIAdapter> TempAdapter;
 	for (uint32 AdapterIndex = 0; DXGIFactory1->EnumAdapters(AdapterIndex, TempAdapter.GetInitReference()) != DXGI_ERROR_NOT_FOUND; ++AdapterIndex)
 	{
@@ -328,15 +329,14 @@ bool FWindowsPlatformApplicationMisc::ProbablyHasIntegratedGPU()
 			DXGI_ADAPTER_DESC Desc;
 			TempAdapter->GetDesc(&Desc);
 
-			const int MIN_GPU_MEMORY = 512 * 1024 * 1024;
-			if (Desc.DedicatedVideoMemory >= MIN_GPU_MEMORY)
+			if (Desc.DedicatedVideoMemory > BestDesc.DedicatedVideoMemory || AdapterIndex == 0)
 			{
-				return false;
+				BestDesc = Desc;
 			}
 		}
 	}
 
-	return true;
+	return FGPUInfo{ BestDesc.VendorId, BestDesc.DeviceId, BestDesc.DedicatedVideoMemory };
 }
 
 float FWindowsPlatformApplicationMisc::GetDPIScaleFactorAtPoint(float X, float Y)
diff --git a/Engine/Source/Runtime/ApplicationCore/Public/Windows/WindowsPlatformApplicationMisc.h b/Engine/Source/Runtime/ApplicationCore/Public/Windows/WindowsPlatformApplicationMisc.h
--- a/Engine/Source/Runtime/ApplicationCore/Public/Windows/WindowsPlatformApplicationMisc.h
+++ b/Engine/Source/Runtime/ApplicationCore/Public/Windows/WindowsPlatformApplicationMisc.h
@@ -27,7 +27,15 @@ struct APPLICATIONCORE_API FWindowsPlatformApplicationMisc : public FGenericPlat
 	/** Windows platform only */
 	/** Function should retrieve the DPI value for the provided monitor information structure */
 	static int32 GetMonitorDPI(const FMonitorInfo& MonitorInfo);
-	static bool ProbablyHasIntegratedGPU();
+
+	struct FGPUInfo
+	{
+		uint32 VendorId = 0;
+		uint32 DeviceId = 0;
+		uint64 DedicatedVideoMemory = 0;
+	};
+
+	static FGPUInfo GetBestGPUInfo();
 	/** End Windows platform only */
 };
 
diff --git a/Engine/Source/Runtime/RHI/Private/DynamicRHI.cpp b/Engine/Source/Runtime/RHI/Private/DynamicRHI.cpp
--- a/Engine/Source/Runtime/RHI/Private/DynamicRHI.cpp
+++ b/Engine/Source/Runtime/RHI/Private/DynamicRHI.cpp
@@ -146,10 +146,40 @@ static void RHIDetectAndWarnOfBadDrivers(bool bHasEditorToken)
 	{
 		FDriverDenyListEntry DenyListEntry = DetectedGPUHardware.FindDriverDenyListEntry();
 
+		TArray<FString> DeviceCanUpdateDriverList;
+		GConfig->GetArray(TEXT("Devices"), TEXT("DeviceCanUpdateDriverList"), DeviceCanUpdateDriverList, GHardwareIni);
+
+		bool bVendorHasEntries = false;
+		bool bDeviceCanUpdateDriver = false;
+		for (const FString& DeviceCanUpdateDriverString : DeviceCanUpdateDriverList)
+		{
+			const TCHAR* Line = *DeviceCanUpdateDriverString;
+
+			FString VendorId;
+			FParse::Value(Line + 1, TEXT("VendorId="), VendorId);
+			uint32 VendorIdInt = FParse::HexNumber(*VendorId);
+
+			FString DeviceId;
+			FParse::Value(Line + 1, TEXT("DeviceId="), DeviceId);
+			uint32 DeviceIdInt = FParse::HexNumber(*DeviceId);
+
+			bVendorHasEntries |= DriverInfo.VendorId && DriverInfo.VendorId == VendorIdInt;
+
+			if (DriverInfo.VendorId && GRHIDeviceId &&
+				DriverInfo.VendorId == VendorIdInt && GRHIDeviceId == DeviceIdInt)
+			{
+				bDeviceCanUpdateDriver = true;
+				break;
+			}
+		}
+
 		GRHIAdapterDriverOnDenyList = DenyListEntry.IsValid();
 		FGenericCrashContext::SetEngineData(TEXT("RHI.DriverBlacklisted"), DenyListEntry.IsValid() ? TEXT("true") : TEXT("false"));
 
-		if (DenyListEntry.IsValid())
+		// Only alert users who are capable of updating their driver. Assume vendors with an empty list can always update.
+		bool bShowPrompt = bDeviceCanUpdateDriver || !bVendorHasEntries;
+
+		if (DenyListEntry.IsValid() && bShowPrompt)
 		{
 			bool bLatestDenied = DetectedGPUHardware.IsLatestDenied();
 
@@ -295,6 +325,7 @@ void RHIInit(bool bHasEditorToken)
 				FGenericCrashContext::SetEngineData(TEXT("RHI.DriverDate"), GRHIAdapterDriverDate);
 				FGenericCrashContext::SetEngineData(TEXT("RHI.FeatureLevel"), FeatureLevelString);
 				FGenericCrashContext::SetEngineData(TEXT("RHI.GPUVendor"), RHIVendorIdToString());
+				FGenericCrashContext::SetEngineData(TEXT("RHI.DeviceId"), FString::Printf(TEXT("%04X"), GRHIDeviceId));
 
 #if TEXTURE_PROFILER_ENABLED
 				FTextureProfiler::Get()->Init();
diff --git a/Engine/Source/Runtime/RHI/Private/Windows/WindowsDynamicRHI.cpp b/Engine/Source/Runtime/RHI/Private/Windows/WindowsDynamicRHI.cpp
--- a/Engine/Source/Runtime/RHI/Private/Windows/WindowsDynamicRHI.cpp
+++ b/Engine/Source/Runtime/RHI/Private/Windows/WindowsDynamicRHI.cpp
@@ -137,6 +137,8 @@ static bool DefaultFeatureLevelES31()
 		return true;
 	}
 
+	FWindowsPlatformApplicationMisc::FGPUInfo BestGPUInfo = FWindowsPlatformApplicationMisc::GetBestGPUInfo();
+
 	FString MinMemorySizeBucketString;
 	FString MinIntegratedMemorySizeBucketString;
 	if (GConfig->GetString(TEXT("PerformanceMode"), TEXT("MinMemorySizeBucket"), MinMemorySizeBucketString, GEngineIni) && GConfig->GetString(TEXT("PerformanceMode"), TEXT("MinIntegratedMemorySizeBucket"), MinIntegratedMemorySizeBucketString, GEngineIni))
@@ -157,7 +159,8 @@ static bool DefaultFeatureLevelES31()
 			// Force Performance mode for machines with too little memory when shared with the GPU
 			if (MinIntegratedMemorySizeBucketString == BucketString)
 			{
-				if (FPlatformMemory::GetMemorySizeBucket() >= EPlatformMemorySizeBucket(EnumIndex) && FWindowsPlatformApplicationMisc::ProbablyHasIntegratedGPU())
+				const int MIN_GPU_MEMORY = 512 * 1024 * 1024;
+				if (FPlatformMemory::GetMemorySizeBucket() >= EPlatformMemorySizeBucket(EnumIndex) && BestGPUInfo.DedicatedVideoMemory < MIN_GPU_MEMORY)
 				{
 					ForceES31 = true;
 
@@ -189,6 +192,23 @@ static bool DefaultFeatureLevelES31()
 
 			return true;
 		}
+
+		FString VendorId;
+		FParse::Value(Line + 1, TEXT("VendorId="), VendorId);
+		uint32 VendorIdInt = FParse::HexNumber(*VendorId);
+
+		FString DeviceId;
+		FParse::Value(Line + 1, TEXT("DeviceId="), DeviceId);
+		uint32 DeviceIdInt = FParse::HexNumber(*DeviceId);
+
+		if (BestGPUInfo.VendorId && BestGPUInfo.DeviceId &&
+			BestGPUInfo.VendorId == VendorIdInt && BestGPUInfo.DeviceId == DeviceIdInt &&
+			RHIName.Compare("D3D11_ES31", ESearchCase::IgnoreCase) == 0)
+		{
+			ForceES31 = true;
+
+			return true;
+		}
 	}
 
 	ForceES31 = false;
