Remove support for PLATFORM_REQUIRES_UAV_TO_RTV_TEXTURE_CACHE_FLUSH_WORKAROUND

- this was working around a system level issue but we no longer support the older SDK versions which had the issue on the impacted platform

- testing
  - ran QAGame TM-DepthOfField on impacted platform as DoF was one of the passes which previously incurred the issue

diff --git a/Engine/Source/Runtime/Core/Public/HAL/Platform.h b/Engine/Source/Runtime/Core/Public/HAL/Platform.h
--- a/Engine/Source/Runtime/Core/Public/HAL/Platform.h
+++ b/Engine/Source/Runtime/Core/Public/HAL/Platform.h
@@ -618,10 +618,6 @@
 	#define PLATFORM_USE_PLATFORM_FILE_MANAGED_STORAGE_WRAPPER	0
 #endif
 
-#ifndef PLATFORM_REQUIRES_UAV_TO_RTV_TEXTURE_CACHE_FLUSH_WORKAROUND
-	#define PLATFORM_REQUIRES_UAV_TO_RTV_TEXTURE_CACHE_FLUSH_WORKAROUND 0
-#endif
-
 #ifndef PLATFORM_NEEDS_GPU_UAV_RESOURCE_INIT_WORKAROUND
 	#define PLATFORM_NEEDS_GPU_UAV_RESOURCE_INIT_WORKAROUND 0
 #endif
diff --git a/Engine/Source/Runtime/RHI/Public/RHICommandList.h b/Engine/Source/Runtime/RHI/Public/RHICommandList.h
--- a/Engine/Source/Runtime/RHI/Public/RHICommandList.h
+++ b/Engine/Source/Runtime/RHI/Public/RHICommandList.h
@@ -2120,19 +2120,6 @@ FRHICOMMAND_MACRO(FRHICommandBackBufferWaitTrackingBeginFrame)
 };
 #endif // #if PLATFORM_USE_BACKBUFFER_WRITE_TRANSITION_TRACKING
 
-#if PLATFORM_REQUIRES_UAV_TO_RTV_TEXTURE_CACHE_FLUSH_WORKAROUND
-FRHICOMMAND_MACRO(FRHICommandFlushTextureCacheBOP)
-{
-	FRHITexture* Texture;
-
-	FORCEINLINE_DEBUGGABLE FRHICommandFlushTextureCacheBOP(FRHITexture* TextureIn)
-		:	Texture(TextureIn)
-	{}
-	
-	RHI_API void Execute(FRHICommandListBase& CmdList);
-};
-#endif // #if PLATFORM_REQUIRES_UAV_TO_RTV_TEXTURE_CACHE_FLUSH_WORKAROUND
-
 FRHICOMMAND_MACRO(FRHICommandCopyBufferRegion)
 {
 	FRHIBuffer* DestBuffer;
@@ -2996,18 +2983,6 @@ PRAGMA_ENABLE_DEPRECATION_WARNINGS
 	}
 #endif // WITH_MGPU
 
-#if PLATFORM_REQUIRES_UAV_TO_RTV_TEXTURE_CACHE_FLUSH_WORKAROUND
-	FORCEINLINE_DEBUGGABLE void RHIFlushTextureCacheBOP(FRHITexture* Texture)
-	{
-		if (Bypass())
-		{
-			GetContext().RHIFlushTextureCacheBOP(Texture);
-			return;
-		}
-		ALLOC_COMMAND(FRHICommandFlushTextureCacheBOP)(Texture);
-	}
-#endif // #if PLATFORM_REQUIRES_UAV_TO_RTV_TEXTURE_CACHE_FLUSH_WORKAROUND
-
 #if RHI_RAYTRACING
 	void BuildAccelerationStructure(FRHIRayTracingGeometry* Geometry);
 	void BuildAccelerationStructures(const TArrayView<const FRayTracingGeometryBuildParams> Params);
diff --git a/Engine/Source/Runtime/RHI/Public/RHICommandListCommandExecutes.inl b/Engine/Source/Runtime/RHI/Public/RHICommandListCommandExecutes.inl
--- a/Engine/Source/Runtime/RHI/Public/RHICommandListCommandExecutes.inl
+++ b/Engine/Source/Runtime/RHI/Public/RHICommandListCommandExecutes.inl
@@ -780,11 +780,3 @@ void FRHICommandBackBufferWaitTrackingBeginFrame::Execute(FRHICommandListBase& C
 	INTERNAL_DECORATOR(RHIBackBufferWaitTrackingBeginFrame)(FrameToken, bDeferred);
 }
 #endif // #if PLATFORM_USE_BACKBUFFER_WRITE_TRANSITION_TRACKING
-
-#if PLATFORM_REQUIRES_UAV_TO_RTV_TEXTURE_CACHE_FLUSH_WORKAROUND
-void FRHICommandFlushTextureCacheBOP::Execute(FRHICommandListBase& CmdList)
-{
-	RHISTAT(FRHICommandFlushTextureCacheBOP);
-	INTERNAL_DECORATOR(RHIFlushTextureCacheBOP)(Texture);
-}
-#endif // #if PLATFORM_REQUIRES_UAV_TO_RTV_TEXTURE_CACHE_FLUSH_WORKAROUND
diff --git a/Engine/Source/Runtime/RHI/Public/RHIContext.h b/Engine/Source/Runtime/RHI/Public/RHIContext.h
--- a/Engine/Source/Runtime/RHI/Public/RHIContext.h
+++ b/Engine/Source/Runtime/RHI/Public/RHIContext.h
@@ -832,13 +832,6 @@ public:
 	}
 #endif // #if PLATFORM_USE_BACKBUFFER_WRITE_TRANSITION_TRACKING
 
-#if PLATFORM_REQUIRES_UAV_TO_RTV_TEXTURE_CACHE_FLUSH_WORKAROUND
-	virtual void RHIFlushTextureCacheBOP(FRHITexture* Texture)
-	{
-		checkNoEntry();
-	}
-#endif // #if PLATFORM_REQUIRES_UAV_TO_RTV_TEXTURE_CACHE_FLUSH_WORKAROUND
-
 	protected:
 		FRHIRenderPassInfo RenderPassInfo;
 };
diff --git a/Engine/Source/Runtime/RHI/Public/RHIValidationContext.h b/Engine/Source/Runtime/RHI/Public/RHIValidationContext.h
--- a/Engine/Source/Runtime/RHI/Public/RHIValidationContext.h
+++ b/Engine/Source/Runtime/RHI/Public/RHIValidationContext.h
@@ -1066,13 +1066,6 @@ public:
 	}
 #endif // #if PLATFORM_USE_BACKBUFFER_WRITE_TRANSITION_TRACKING
 
-#if PLATFORM_REQUIRES_UAV_TO_RTV_TEXTURE_CACHE_FLUSH_WORKAROUND
-	virtual void RHIFlushTextureCacheBOP(FRHITexture* Texture) override final
-	{
-		RHIContext->RHIFlushTextureCacheBOP(Texture);
-	}
-#endif // PLATFORM_REQUIRES_UAV_TO_RTV_TEXTURE_CACHE_FLUSH_WORKAROUND
-
 	void SetupDrawing()
 	{
 		// nothing to validate right now
diff --git a/Engine/Source/Runtime/Renderer/Private/PostProcess/DiaphragmDOF.cpp b/Engine/Source/Runtime/Renderer/Private/PostProcess/DiaphragmDOF.cpp
--- a/Engine/Source/Runtime/Renderer/Private/PostProcess/DiaphragmDOF.cpp
+++ b/Engine/Source/Runtime/Renderer/Private/PostProcess/DiaphragmDOF.cpp
@@ -2450,10 +2450,6 @@ FRDGTextureRef DiaphragmDOF::AddPasses(
 				ERDGPassFlags::Raster,
 				[PassParameters, VertexShader, PixelShader, GatheringViewSize, DrawIndirectParametersOffset](FRHICommandList& RHICmdList)
 			{
-#if PLATFORM_REQUIRES_UAV_TO_RTV_TEXTURE_CACHE_FLUSH_WORKAROUND
-				RHICmdList.RHIFlushTextureCacheBOP(PassParameters->RenderTargets[0].GetTexture()->GetRHI());
-#endif // #if PLATFORM_REQUIRES_UAV_TO_RTV_TEXTURE_CACHE_FLUSH_WORKAROUND
-
 				RHICmdList.SetViewport(0, 0, 0.0f, GatheringViewSize.X, GatheringViewSize.Y, 1.0f);
 
 				FGraphicsPipelineStateInitializer GraphicsPSOInit;
