Improve 120fps Support

#rnx

diff --git a/Engine/Source/Runtime/Core/Public/GenericPlatform/GenericPlatformMisc.h b/Engine/Source/Runtime/Core/Public/GenericPlatform/GenericPlatformMisc.h
--- a/Engine/Source/Runtime/Core/Public/GenericPlatform/GenericPlatformMisc.h
+++ b/Engine/Source/Runtime/Core/Public/GenericPlatform/GenericPlatformMisc.h
@@ -1692,13 +1692,21 @@ public:
 	}
 
 	/**
-	 * retrieves the maximum refresh rate supported by the platform
+	 * retrieves the current maximum refresh rate supported by the platform
 	 */
 	static inline int32 GetMaxRefreshRate()
 	{
 		return 60;
 	}
 
+	/**
+	 * retrieves the maximum refresh rate supported by the platform
+	 */
+	static inline int32 GetMaxSupportedRefreshRate()
+	{
+		return GetMaxRefreshRate();
+	}
+
 	/**
 	 * Returns the platform's maximum allowed value for rhi.SyncInterval
 	 */
diff --git a/Engine/Source/Runtime/RHI/Private/RHIUtilities.cpp b/Engine/Source/Runtime/RHI/Private/RHIUtilities.cpp
--- a/Engine/Source/Runtime/RHI/Private/RHIUtilities.cpp
+++ b/Engine/Source/Runtime/RHI/Private/RHIUtilities.cpp
@@ -11,6 +11,7 @@ RHIUtilities.cpp:
 #include "HAL/Runnable.h"
 #include "HAL/RunnableThread.h"
 #include "Misc/ScopeLock.h"
+#include "HAL/PlatformFramePacer.h"
 
 #define USE_FRAME_OFFSET_THREAD 1
 
@@ -123,6 +124,28 @@ TAutoConsoleVariable<int32> CVarRHISyncAllowVariable(
 	ECVF_Default
 );
 
+int32 GEnableConsole120Fps = 0;
+int32 InternalEnableConsole120Fps = 0;
+static void OnEnableConsole120FpsCVarRHIChanged(IConsoleVariable* Variable)
+{
+	InternalEnableConsole120Fps &= (FPlatformMisc::GetMaxSupportedRefreshRate() >= 120);
+	if (InternalEnableConsole120Fps != GEnableConsole120Fps)
+	{
+		GEnableConsole120Fps = InternalEnableConsole120Fps;
+		// needs to update the FramePace since it updates the SyncInterval based on the RefreshRate.
+		FPlatformRHIFramePacer::SetFramePace(GEnableConsole120Fps ? 120 : 60);
+		UE_LOG(LogRHI, Log, TEXT("Console 120Fps = %d"), GEnableConsole120Fps);
+	}
+}
+
+static FAutoConsoleVariableRef CVarRHInableConsole120Fps(
+	TEXT("rhi.EnableConsole120Fps"),
+	InternalEnableConsole120Fps,
+	TEXT("Enable Console 120fps if Monitor supports it and Console is properly setup"),
+	FConsoleVariableDelegate::CreateStatic(&OnEnableConsole120FpsCVarRHIChanged),
+	ECVF_Default
+);
+
 class FRHIFrameFlipTrackingRunnable : public FRunnable
 {
 	static FRunnableThread* Thread;
