UnrealBuildTool: Compile PCH with -fpch-instantiate-templates for Switch. Improves compile time by almost 15%

RHI fix for specialization after instantiation errors when thsi flag is enabled (from chris waters)

diff --git a/Engine/Source/Runtime/RHI/Public/RHICommandList.h b/Engine/Source/Runtime/RHI/Public/RHICommandList.h
--- a/Engine/Source/Runtime/RHI/Public/RHICommandList.h
+++ b/Engine/Source/Runtime/RHI/Public/RHICommandList.h
@@ -2355,7 +2355,7 @@ template<> RHI_API void FRHICommandSetShaderUniformBuffer<FRHIComputeShader>::Ex
 template<> RHI_API void FRHICommandSetShaderTexture<FRHIComputeShader>::Execute(FRHICommandListBase& CmdList);
 template<> RHI_API void FRHICommandSetShaderResourceViewParameter<FRHIComputeShader>::Execute(FRHICommandListBase& CmdList);
 template<> RHI_API void FRHICommandSetShaderSampler<FRHIComputeShader>::Execute(FRHICommandListBase& CmdList);
-
+template<> RHI_API void FRHICommandSetUAVParameter<FRHIComputeShader>::Execute(FRHICommandListBase& CmdList);
 
 class RHI_API FRHIComputeCommandList : public FRHICommandListBase
 {
@@ -3051,6 +3051,13 @@ PRAGMA_ENABLE_DEPRECATION_WARNINGS
 	}
 };
 
+template<> RHI_API void FRHICommandSetShaderParameter<FRHIGraphicsShader>::Execute(FRHICommandListBase& CmdList);
+template<> RHI_API void FRHICommandSetShaderUniformBuffer<FRHIGraphicsShader>::Execute(FRHICommandListBase& CmdList);
+template<> RHI_API void FRHICommandSetShaderTexture<FRHIGraphicsShader>::Execute(FRHICommandListBase& CmdList);
+template<> RHI_API void FRHICommandSetShaderResourceViewParameter<FRHIGraphicsShader>::Execute(FRHICommandListBase& CmdList);
+template<> RHI_API void FRHICommandSetShaderSampler<FRHIGraphicsShader>::Execute(FRHICommandListBase& CmdList);
+template<> RHI_API void FRHICommandSetUAVParameter<FRHIPixelShader>::Execute(FRHICommandListBase& CmdList);
+
 class RHI_API FRHICommandList : public FRHIComputeCommandList
 {
 public:
