RHI validation fixes for CitySample.
 - Fixed RHIUpdateShaderResourceView crash when buffer is null.
 - Allow validation to pass any pipe access for resources that haven't been transitioned yet. This is necessary for read-only resources that are referenced by async compute. The resource can safely be transitioned from any initial pipe, at which point it becomes tied to the pipe it's being transitioned to.

diff --git a/Engine/Source/Runtime/RHI/Public/RHIValidation.h b/Engine/Source/Runtime/RHI/Public/RHIValidation.h
--- a/Engine/Source/Runtime/RHI/Public/RHIValidation.h
+++ b/Engine/Source/Runtime/RHI/Public/RHIValidation.h
@@ -482,13 +482,29 @@ public:
 	virtual void RHIUpdateShaderResourceView(FRHIShaderResourceView* SRV, FRHIBuffer* Buffer, uint32 Stride, uint8 Format) override final
 	{
 		RHI->RHIUpdateShaderResourceView(SRV, Buffer, Stride, Format);
-		SRV->ViewIdentity = Buffer->GetWholeResourceIdentity();
+
+		if (Buffer)
+		{
+			SRV->ViewIdentity = Buffer->GetWholeResourceIdentity();
+		}
+		else
+		{
+			SRV->ResetViewIdentity();
+		}
 	}
 
 	virtual void RHIUpdateShaderResourceView(FRHIShaderResourceView* SRV, FRHIBuffer* Buffer) override final
 	{
 		RHI->RHIUpdateShaderResourceView(SRV, Buffer);
-		SRV->ViewIdentity = Buffer->GetWholeResourceIdentity();
+
+		if (Buffer)
+		{
+			SRV->ViewIdentity = Buffer->GetWholeResourceIdentity();
+		}
+		else
+		{
+			SRV->ResetViewIdentity();
+		}
 	}
 
 	virtual FRHICalcTextureSizeResult RHICalcTexturePlatformSize(FRHITextureDesc const& Desc, uint32 FirstMipIndex) override final
diff --git a/Engine/Source/Runtime/RHI/Public/RHIValidationCommon.h b/Engine/Source/Runtime/RHI/Public/RHIValidationCommon.h
--- a/Engine/Source/Runtime/RHI/Public/RHIValidationCommon.h
+++ b/Engine/Source/Runtime/RHI/Public/RHIValidationCommon.h
@@ -354,9 +354,12 @@ namespace RHIValidation
 			TransientState = FTransientState(InResourceState);
 			TrackedAccess = InResourceState;
 
-			for (auto& State : WholeResourceState.States)
+			for (ERHIPipeline Pipeline : GetRHIPipelines())
 			{
+				auto& State = WholeResourceState.States[Pipeline];
+
 				State.Current.Access = InResourceState;
+				State.Current.Pipelines = Pipeline;
 				State.Previous = State.Current;
 			}
 
@@ -534,6 +537,11 @@ namespace RHIValidation
 	struct FView
 	{
 		FView()
+		{
+			ResetViewIdentity();
+		}
+
+		void ResetViewIdentity()
 		{
 			ViewIdentity.Resource = nullptr;
 			ViewIdentity.SubresourceRange = FSubresourceRange(0, 0, 0, 0, 0, 0);
