Added support for per-slice texture array UAV+SRV on some platforms

diff --git a/Engine/Source/Runtime/Apple/MetalRHI/Private/MetalUAV.cpp b/Engine/Source/Runtime/Apple/MetalRHI/Private/MetalUAV.cpp
--- a/Engine/Source/Runtime/Apple/MetalRHI/Private/MetalUAV.cpp
+++ b/Engine/Source/Runtime/Apple/MetalRHI/Private/MetalUAV.cpp
@@ -89,9 +89,6 @@ FMetalResourceViewBase::FMetalResourceViewBase(
 	, Stride(0)
 	, Offset(0)
 {
-	// Slice selection of a texture array still need to be implemented on Metal
-	check(InFirstArraySlice == 0 && InNumArraySlices == 0);
-
 	if (SourceTexture)
 	{
 		SourceTexture->AddRef();
@@ -146,7 +143,9 @@ FMetalResourceViewBase::FMetalResourceViewBase(
 				MipLevel == 0
 				&& NumMips == SourceTexture->Texture.GetMipmapLevelCount()
 				&& MetalFormat == SourceTexture->Texture.GetPixelFormat()
-				&& !(bInUAV && SourceTexture->GetDesc().IsTextureCube());		// @todo: Remove this once Cube UAV supported for all Metal Devices
+				&& !(bInUAV && SourceTexture->GetDesc().IsTextureCube())		// @todo: Remove this once Cube UAV supported for all Metal Devices
+				&& InFirstArraySlice == 0
+				&& InNumArraySlices == 0;
 
 			if (bUseSourceTex)
 			{
@@ -157,19 +156,46 @@ FMetalResourceViewBase::FMetalResourceViewBase(
 			{
 				// Recreate the texture to enable MTLTextureUsagePixelFormatView which must be off unless we definitely use this feature or we are throwing ~4% performance vs. Windows on the floor.
 				// @todo recreating resources like this will likely prevent us from making view creation multi-threaded.
+				// @todo RW: Flag usage creation logic has changed: This should be a check(PixelFormatView) now: SourceTexture should have been created with correct flags
 				if (!(SourceTexture->Texture.GetUsage() & mtlpp::TextureUsage::PixelFormatView))
+				{
 					SourceTexture->PrepareTextureView();
+				}
+
+				const uint32 TextureSliceCount = SourceTexture->Texture.GetArrayLength();
+				const uint32 CubeSliceMultiplier = SourceTexture->GetDesc().IsTextureCube() ? 6 : 1;
+				const uint32 NumArraySlices = (InNumArraySlices > 0 ? InNumArraySlices : TextureSliceCount) * CubeSliceMultiplier;
 
-				ns::Range Slices(0, SourceTexture->Texture.GetArrayLength() * (SourceTexture->GetDesc().IsTextureCube() ? 6 : 1));
-				
 				// @todo: Remove this type swizzle once Cube UAV supported for all Metal Devices - SRV seem to want to stay as cube but UAV are expected to be 2DArray
 				mtlpp::TextureType TextureType = bInUAV && SourceTexture->GetDesc().IsTextureCube() ? mtlpp::TextureType::Texture2DArray : SourceTexture->Texture.GetTextureType();
 
+				// Assume a texture view of 1 slice into a multislice texture wants to be the non-array texture type
+				// This doesn't really matter to Metal but will be very important when this texture is bound in the shader
+				if (InNumArraySlices == 1)
+				{
+					switch (TextureType)
+					{
+					case mtlpp::TextureType::Texture2DArray:
+						TextureType = mtlpp::TextureType::Texture2D;
+						break;
+					case mtlpp::TextureType::TextureCubeArray:
+						TextureType = mtlpp::TextureType::TextureCube;
+						break;
+					default:
+						// NOP
+						break;
+					}
+				}
+
 				TextureView = SourceTexture->Texture.NewTextureView(
 					MetalFormat,
 					TextureType,
 					ns::Range(MipLevel, NumMips),
-					Slices);
+					ns::Range(InFirstArraySlice, NumArraySlices));
+
+#if METAL_DEBUG_OPTIONS
+				TextureView.SetLabel([SourceTexture->Texture.GetLabel() stringByAppendingString:@"_TextureView"]);
+#endif
 			}
 		}
 	}
diff --git a/Engine/Source/Runtime/RHI/Private/RHI.cpp b/Engine/Source/Runtime/RHI/Private/RHI.cpp
--- a/Engine/Source/Runtime/RHI/Private/RHI.cpp
+++ b/Engine/Source/Runtime/RHI/Private/RHI.cpp
@@ -2029,7 +2029,7 @@ void FRHIRenderPassInfo::Validate() const
 }
 #endif
 
-#define ValidateTextureDesc(expr, format, ...) \
+#define ValidateResourceDesc(expr, format, ...) \
 	if (bFatal) \
 	{ \
 		checkf(expr, format, ##__VA_ARGS__); \
@@ -2044,9 +2044,9 @@ bool FRHITextureDesc::Validate(const FRHITextureCreateInfo& Desc, const TCHAR* N
 {
 	// Validate texture's pixel format.
 	{
-		ValidateTextureDesc(Desc.Format != PF_Unknown, TEXT("Illegal to create texture %s with an invalid pixel format."), Name);
-		ValidateTextureDesc(Desc.Format < PF_MAX, TEXT("Illegal to create texture %s with an invalid pixel format."), Name);
-		ValidateTextureDesc(GPixelFormats[Desc.Format].Supported,
+		ValidateResourceDesc(Desc.Format != PF_Unknown, TEXT("Illegal to create texture %s with an invalid pixel format."), Name);
+		ValidateResourceDesc(Desc.Format < PF_MAX, TEXT("Illegal to create texture %s with an invalid pixel format."), Name);
+		ValidateResourceDesc(GPixelFormats[Desc.Format].Supported,
 			TEXT("Failed to create texture %s with pixel format %s because it is not supported."),
 			Name,
 			GPixelFormats[Desc.Format].Name);
@@ -2056,60 +2056,134 @@ bool FRHITextureDesc::Validate(const FRHITextureCreateInfo& Desc, const TCHAR* N
 	{
 		int32 MaxDimension = (Desc.Dimension == ETextureDimension::TextureCube || Desc.Dimension == ETextureDimension::TextureCubeArray) ? GMaxCubeTextureDimensions : GMaxTextureDimensions;
 
-		ValidateTextureDesc(Desc.Extent.X > 0, TEXT("Texture %s's Extent.X=%d is invalid."), Name, Desc.Extent.X);
-		ValidateTextureDesc(Desc.Extent.X <= MaxDimension, TEXT("Texture %s's Extent.X=%d is too large."), Name, Desc.Extent.X);
+		ValidateResourceDesc(Desc.Extent.X > 0, TEXT("Texture %s's Extent.X=%d is invalid."), Name, Desc.Extent.X);
+		ValidateResourceDesc(Desc.Extent.X <= MaxDimension, TEXT("Texture %s's Extent.X=%d is too large."), Name, Desc.Extent.X);
 
-		ValidateTextureDesc(Desc.Extent.Y > 0, TEXT("Texture %s's Extent.Y=%d is invalid."), Name, Desc.Extent.Y);
-		ValidateTextureDesc(Desc.Extent.Y <= MaxDimension, TEXT("Texture %s's Extent.Y=%d is too large."), Name, Desc.Extent.Y);
+		ValidateResourceDesc(Desc.Extent.Y > 0, TEXT("Texture %s's Extent.Y=%d is invalid."), Name, Desc.Extent.Y);
+		ValidateResourceDesc(Desc.Extent.Y <= MaxDimension, TEXT("Texture %s's Extent.Y=%d is too large."), Name, Desc.Extent.Y);
 	}
 
 	// Validate texture's depth
 	if (Desc.Dimension == ETextureDimension::Texture3D)
 	{
-		ValidateTextureDesc(Desc.Depth > 0, TEXT("Texture %s's Depth=%d is invalid."), Name, int32(Desc.Depth));
-		ValidateTextureDesc(Desc.Depth <= GMaxTextureDimensions, TEXT("Texture %s's Extent.Depth=%d is too large."), Name, Desc.Depth);
+		ValidateResourceDesc(Desc.Depth > 0, TEXT("Texture %s's Depth=%d is invalid."), Name, int32(Desc.Depth));
+		ValidateResourceDesc(Desc.Depth <= GMaxTextureDimensions, TEXT("Texture %s's Extent.Depth=%d is too large."), Name, Desc.Depth);
 	}
 	else
 	{
-		ValidateTextureDesc(Desc.Depth == 1, TEXT("Texture %s's Depth=%d is invalid for Dimension=%s."), Name, int32(Desc.Depth), GetTextureDimensionString(Desc.Dimension));
+		ValidateResourceDesc(Desc.Depth == 1, TEXT("Texture %s's Depth=%d is invalid for Dimension=%s."), Name, int32(Desc.Depth), GetTextureDimensionString(Desc.Dimension));
 	}
 
 	// Validate texture's array size
 	if (Desc.Dimension == ETextureDimension::Texture2DArray || Desc.Dimension == ETextureDimension::TextureCubeArray)
 	{
-		ValidateTextureDesc(Desc.ArraySize > 0, TEXT("Texture %s's ArraySize=%d is invalid."), Name, Desc.ArraySize);
-		ValidateTextureDesc(Desc.ArraySize <= GMaxTextureArrayLayers, TEXT("Texture %s's Extent.ArraySize=%d is too large."), Name, int32(Desc.ArraySize));
+		ValidateResourceDesc(Desc.ArraySize > 0, TEXT("Texture %s's ArraySize=%d is invalid."), Name, Desc.ArraySize);
+		ValidateResourceDesc(Desc.ArraySize <= GMaxTextureArrayLayers, TEXT("Texture %s's Extent.ArraySize=%d is too large."), Name, int32(Desc.ArraySize));
 	}
 	else
 	{
-		ValidateTextureDesc(Desc.ArraySize == 1, TEXT("Texture %s's ArraySize=%d is invalid for Dimension=%s."), Name, Desc.ArraySize,  GetTextureDimensionString(Desc.Dimension));
+		ValidateResourceDesc(Desc.ArraySize == 1, TEXT("Texture %s's ArraySize=%d is invalid for Dimension=%s."), Name, Desc.ArraySize,  GetTextureDimensionString(Desc.Dimension));
 	}
 
 	// Validate texture's samples count.
 	if (Desc.Dimension == ETextureDimension::Texture2D || Desc.Dimension == ETextureDimension::Texture2DArray)
 	{
-		ValidateTextureDesc(Desc.NumSamples > 0, TEXT("Texture %s's NumSamples=%d is invalid."), Name, Desc.NumSamples);
+		ValidateResourceDesc(Desc.NumSamples > 0, TEXT("Texture %s's NumSamples=%d is invalid."), Name, Desc.NumSamples);
 	}
 	else
 	{
-		ValidateTextureDesc(Desc.NumSamples == 1, TEXT("Texture %s's NumSamples=%d is invalid for Dimension=%s."), Name, Desc.NumSamples, GetTextureDimensionString(Desc.Dimension));
+		ValidateResourceDesc(Desc.NumSamples == 1, TEXT("Texture %s's NumSamples=%d is invalid for Dimension=%s."), Name, Desc.NumSamples, GetTextureDimensionString(Desc.Dimension));
 	}
 
 	// Validate texture's mips.
 	if (Desc.IsMultisample())
 	{
-		ValidateTextureDesc(Desc.NumMips == 1, TEXT("MSAA Texture %s's can only have one mip."), Name);
+		ValidateResourceDesc(Desc.NumMips == 1, TEXT("MSAA Texture %s's can only have one mip."), Name);
 	}
 	else
 	{
-		ValidateTextureDesc(Desc.NumMips > 0, TEXT("Texture %s's NumMips=%d is invalid."), Name, Desc.NumMips);
-		ValidateTextureDesc(Desc.NumMips <= GMaxTextureMipCount, TEXT("Texture %s's NumMips=%d is too large."), Name, Desc.NumMips);
+		ValidateResourceDesc(Desc.NumMips > 0, TEXT("Texture %s's NumMips=%d is invalid."), Name, Desc.NumMips);
+		ValidateResourceDesc(Desc.NumMips <= GMaxTextureMipCount, TEXT("Texture %s's NumMips=%d is too large."), Name, Desc.NumMips);
 	}
 
 	return true;
 }
 
-#undef ValidateTextureDesc
+// static
+bool FRHITextureSRVCreateInfo::Validate(const FRHITextureDesc& TextureDesc, const FRHITextureSRVCreateInfo& TextureSRVDesc, const TCHAR* TextureName, bool bFatal)
+{
+	if (TextureName == nullptr)
+	{
+		TextureName = TEXT("UnnamedTexture");
+	}
+
+	ValidateResourceDesc(TextureDesc.Flags & TexCreate_ShaderResource,
+		TEXT("Attempted to create SRV from texture %s which was not created with TexCreate_ShaderResource"),
+		TextureName);
+
+	// Validate the pixel format if overridden by the SRV's descriptor.
+	if (TextureSRVDesc.Format == PF_X24_G8)
+	{
+		// PF_X24_G8 is a bit of mess in the RHI, used to read the stencil, but have varying BlockBytes.
+		ValidateResourceDesc(TextureDesc.Format == PF_DepthStencil,
+			TEXT("PF_X24_G8 is only to read stencil from a PF_DepthStencil texture"));
+	}
+	else if (TextureSRVDesc.Format != PF_Unknown)
+	{
+		ValidateResourceDesc(TextureSRVDesc.Format < PF_MAX,
+			TEXT("Illegal to create SRV for texture %s with invalid FPooledRenderTargetDesc::Format."),
+			TextureName);
+		ValidateResourceDesc(GPixelFormats[TextureSRVDesc.Format].Supported,
+			TEXT("Failed to create SRV for texture %s with pixel format %s because it is not supported."),
+			TextureName, GPixelFormats[TextureSRVDesc.Format].Name);
+
+		EPixelFormat ResourcePixelFormat = TextureDesc.Format;
+
+		ValidateResourceDesc(
+			GPixelFormats[TextureSRVDesc.Format].BlockBytes == GPixelFormats[ResourcePixelFormat].BlockBytes &&
+			GPixelFormats[TextureSRVDesc.Format].BlockSizeX == GPixelFormats[ResourcePixelFormat].BlockSizeX &&
+			GPixelFormats[TextureSRVDesc.Format].BlockSizeY == GPixelFormats[ResourcePixelFormat].BlockSizeY &&
+			GPixelFormats[TextureSRVDesc.Format].BlockSizeZ == GPixelFormats[ResourcePixelFormat].BlockSizeZ,
+			TEXT("Failed to create SRV for texture %s with pixel format %s because it does not match the byte size of the texture's pixel format %s."),
+			TextureName, GPixelFormats[TextureSRVDesc.Format].Name, GPixelFormats[ResourcePixelFormat].Name);
+	}
+
+	ValidateResourceDesc((TextureSRVDesc.MipLevel + TextureSRVDesc.NumMipLevels) <= TextureDesc.NumMips,
+		TEXT("Failed to create SRV at mips %d-%d: the texture %s has only %d mip levels."),
+		TextureSRVDesc.MipLevel, (TextureSRVDesc.MipLevel + TextureSRVDesc.NumMipLevels), TextureName, TextureDesc.NumMips);
+
+	// Validate the array sloces
+	if (TextureDesc.IsTextureArray())
+	{
+		ValidateResourceDesc((TextureSRVDesc.FirstArraySlice + TextureSRVDesc.NumArraySlices) <= TextureDesc.ArraySize,
+			TEXT("Failed to create SRV at array slices %d-%d: the texture array %s has only %d slices."),
+			TextureSRVDesc.FirstArraySlice,
+			(TextureSRVDesc.FirstArraySlice + TextureSRVDesc.NumArraySlices),
+			TextureName,
+			TextureDesc.ArraySize);
+	}
+	else
+	{
+		ValidateResourceDesc(TextureSRVDesc.FirstArraySlice == 0,
+			TEXT("Failed to create SRV with FirstArraySlice=%d: the texture %s is not a texture array."),
+			TextureSRVDesc.FirstArraySlice, TextureName);
+		ValidateResourceDesc(TextureSRVDesc.NumArraySlices == 0,
+			TEXT("Failed to create SRV with NumArraySlices=%d: the texture %s is not a texture array."),
+			TextureSRVDesc.NumArraySlices, TextureName);
+	}
+
+	ValidateResourceDesc(TextureSRVDesc.MetaData != ERHITextureMetaDataAccess::FMask || GRHISupportsExplicitFMask,
+		TEXT("Failed to create FMask SRV for texture %s because the current RHI doesn't support it. Be sure to gate the call with GRHISupportsExplicitFMask."),
+		TextureName);
+
+	ValidateResourceDesc(TextureSRVDesc.MetaData != ERHITextureMetaDataAccess::HTile || GRHISupportsExplicitHTile,
+		TEXT("Failed to create HTile SRV for texture %s because the current RHI doesn't support it. Be sure to gate the call with GRHISupportsExplicitHTile."),
+		TextureName);
+
+	return true;
+}
+
+#undef ValidateResourceDesc
 
 uint64 FRHITextureDesc::CalcMemorySizeEstimate(uint32 FirstMipIndex, uint32 LastMipIndex) const
 {
diff --git a/Engine/Source/Runtime/RHI/Public/RHIResources.h b/Engine/Source/Runtime/RHI/Public/RHIResources.h
--- a/Engine/Source/Runtime/RHI/Public/RHIResources.h
+++ b/Engine/Source/Runtime/RHI/Public/RHIResources.h
@@ -4233,7 +4233,7 @@ enum ERHITextureSRVOverrideSRGBType : uint8
 	SRGBO_ForceDisable,
 };
 
-struct FRHITextureSRVCreateInfo
+struct RHI_API FRHITextureSRVCreateInfo
 {
 	explicit FRHITextureSRVCreateInfo(uint8 InMipLevel = 0u, uint8 InNumMipLevels = 1u, EPixelFormat InFormat = PF_Unknown)
 		: Format(InFormat)
@@ -4298,6 +4298,15 @@ struct FRHITextureSRVCreateInfo
 		Hash = HashCombine(Hash, uint32(Info.MetaData));
 		return Hash;
 	}
+
+	/** Check the validity. */
+	static bool CheckValidity(const FRHITextureDesc& TextureDesc, const FRHITextureSRVCreateInfo& TextureSRVDesc, const TCHAR* TextureName)
+	{
+		return FRHITextureSRVCreateInfo::Validate(TextureDesc, TextureSRVDesc, TextureName, /* bFatal = */ true);
+	}
+
+protected:
+	static bool Validate(const FRHITextureDesc& TextureDesc, const FRHITextureSRVCreateInfo& TextureSRVDesc, const TCHAR* TextureName, bool bFatal);
 };
 
 struct FRHITextureUAVCreateInfo
diff --git a/Engine/Source/Runtime/RHI/Public/RHIValidation.h b/Engine/Source/Runtime/RHI/Public/RHIValidation.h
--- a/Engine/Source/Runtime/RHI/Public/RHIValidation.h
+++ b/Engine/Source/Runtime/RHI/Public/RHIValidation.h
@@ -601,6 +601,7 @@ public:
 	virtual FShaderResourceViewRHIRef RHICreateShaderResourceView(FRHITexture* TextureRHI, const FRHITextureSRVCreateInfo& CreateInfo) override final
 	{
 		RHI_VALIDATION_CHECK(TextureRHI->GetTextureReference() == nullptr, TEXT("Creating a shader resource view of an FRHITextureReference is not supported."));
+		FRHITextureSRVCreateInfo::CheckValidity(TextureRHI->GetDesc(), CreateInfo, TextureRHI->GetTrackerResource()->GetDebugName());
 
 		FShaderResourceViewRHIRef SRV = RHI->RHICreateShaderResourceView(TextureRHI, CreateInfo);
 
@@ -1476,6 +1477,7 @@ public:
 	virtual FShaderResourceViewRHIRef RHICreateShaderResourceView_RenderThread(class FRHICommandListImmediate& RHICmdList, FRHITexture* Texture, const FRHITextureSRVCreateInfo& CreateInfo) override final
 	{
 		RHI_VALIDATION_CHECK(Texture->GetTextureReference() == nullptr, TEXT("Creating a shader resource view of an FRHITextureReference is not supported."));
+		FRHITextureSRVCreateInfo::CheckValidity(Texture->GetDesc(), CreateInfo, Texture->GetTrackerResource()->GetDebugName());
 
 		FShaderResourceViewRHIRef SRV = RHI->RHICreateShaderResourceView_RenderThread(RHICmdList, Texture, CreateInfo);
 		RHIValidation::EResourcePlane Plane = CreateInfo.Format == PF_X24_G8
diff --git a/Engine/Source/Runtime/RenderCore/Private/RenderGraphValidation.cpp b/Engine/Source/Runtime/RenderCore/Private/RenderGraphValidation.cpp
--- a/Engine/Source/Runtime/RenderCore/Private/RenderGraphValidation.cpp
+++ b/Engine/Source/Runtime/RenderCore/Private/RenderGraphValidation.cpp
@@ -299,41 +299,7 @@ void FRDGUserValidation::ValidateCreateSRV(const FRDGTextureSRVDesc& Desc)
 	FRDGTextureRef Texture = Desc.Texture;
 	checkf(Texture, TEXT("Texture SRV created with a null texture."));
 	ExecuteGuard(TEXT("CreateSRV"), Texture->Name);
-	checkf(Texture->Desc.Flags & TexCreate_ShaderResource, TEXT("Attempted to create SRV from texture %s which was not created with TexCreate_ShaderResource"), Desc.Texture->Name);
-
-	// Validate the pixel format if overridden by the SRV's descriptor.
-	if (Desc.Format == PF_X24_G8)
-	{
-		// PF_X24_G8 is a bit of mess in the RHI, used to read the stencil, but have varying BlockBytes.
-		checkf(Texture->Desc.Format == PF_DepthStencil, TEXT("PF_X24_G8 is only to read stencil from a PF_DepthStencil texture"));
-	}
-	else if (Desc.Format != PF_Unknown)
-	{
-		checkf(Desc.Format < PF_MAX, TEXT("Illegal to create SRV for texture %s with invalid FPooledRenderTargetDesc::Format."), Texture->Name);
-		checkf(GPixelFormats[Desc.Format].Supported, TEXT("Failed to create SRV for texture %s with pixel format %s because it is not supported."),
-			Texture->Name, GPixelFormats[Desc.Format].Name);
-
-		EPixelFormat ResourcePixelFormat = Texture->Desc.Format;
-
-		checkf(
-			GPixelFormats[Desc.Format].BlockBytes == GPixelFormats[ResourcePixelFormat].BlockBytes &&
-			GPixelFormats[Desc.Format].BlockSizeX == GPixelFormats[ResourcePixelFormat].BlockSizeX &&
-			GPixelFormats[Desc.Format].BlockSizeY == GPixelFormats[ResourcePixelFormat].BlockSizeY &&
-			GPixelFormats[Desc.Format].BlockSizeZ == GPixelFormats[ResourcePixelFormat].BlockSizeZ,
-			TEXT("Failed to create SRV for texture %s with pixel format %s because it does not match the byte size of the texture's pixel format %s."),
-			Texture->Name, GPixelFormats[Desc.Format].Name, GPixelFormats[ResourcePixelFormat].Name);
-	}
-
-	checkf((Desc.MipLevel + Desc.NumMipLevels) <= Texture->Desc.NumMips, TEXT("Failed to create SRV at mips %d-%d: the texture %s has only %d mip levels."),
-		Desc.MipLevel, (Desc.MipLevel + Desc.NumMipLevels), Texture->Name, Texture->Desc.NumMips);
-
-	checkf(Desc.MetaData != ERDGTextureMetaDataAccess::FMask || GRHISupportsExplicitFMask,
-		TEXT("Failed to create FMask SRV for texture %s because the current RHI doesn't support it. Be sure to gate the call with GRHISupportsExplicitFMask."),
-		Texture->Name);
-
-	checkf(Desc.MetaData != ERDGTextureMetaDataAccess::HTile || GRHISupportsExplicitHTile,
-		TEXT("Failed to create HTile SRV for texture %s because the current RHI doesn't support it. Be sure to gate the call with GRHISupportsExplicitHTile."),
-		Texture->Name);
+	check(FRDGTextureSRVDesc::CheckValidity(Texture->Desc, Desc, Texture->Name));
 }
 
 void FRDGUserValidation::ValidateCreateSRV(const FRDGBufferSRVDesc& Desc)
diff --git a/Engine/Source/Runtime/RenderCore/Public/RenderGraphResources.h b/Engine/Source/Runtime/RenderCore/Public/RenderGraphResources.h
--- a/Engine/Source/Runtime/RenderCore/Public/RenderGraphResources.h
+++ b/Engine/Source/Runtime/RenderCore/Public/RenderGraphResources.h
@@ -665,6 +665,10 @@ public:
 	{
 		Texture = InTexture;
 		NumMipLevels = InTexture->Desc.NumMips;
+		if (InTexture->Desc.IsTextureArray())
+		{
+			NumArraySlices = InTexture->Desc.ArraySize;
+		}
 	}
 
 	/** Create SRV that access all sub resources of texture. */
@@ -676,15 +680,26 @@ public:
 	/** Create SRV that access one specific mip level. */
 	static FRDGTextureSRVDesc CreateForMipLevel(FRDGTextureRef Texture, int32 MipLevel)
 	{
-		FRDGTextureSRVDesc Desc;
-		Desc.Texture = Texture;
-		check(MipLevel >= -1 && MipLevel <= TNumericLimits<int8>::Max()); 
+		check(MipLevel >= -1 && MipLevel <= TNumericLimits<int8>::Max());
+		FRDGTextureSRVDesc Desc = FRDGTextureSRVDesc::Create(Texture);
 		Desc.MipLevel = (int8)MipLevel;
 		Desc.NumMipLevels = 1;
 		return Desc;
 	}
 
-	/** Create SRV that access one specific mip level. */
+	/** Create SRV that access one specific slice. */
+	static FRDGTextureSRVDesc CreateForSlice(FRDGTextureRef Texture, int32 SliceIndex)
+	{
+		check(Texture);
+		check(Texture->Desc.IsTextureArray());
+		check(SliceIndex >= 0 && SliceIndex < Texture->Desc.ArraySize);
+		FRDGTextureSRVDesc Desc = FRDGTextureSRVDesc::Create(Texture);
+		Desc.FirstArraySlice = SliceIndex;
+		Desc.NumArraySlices = 1;
+		return Desc;
+	}
+
+	/** Create SRV that access all sub resources of texture with a specific pixel format. */
 	static FRDGTextureSRVDesc CreateWithPixelFormat(FRDGTextureRef Texture, EPixelFormat PixelFormat)
 	{
 		FRDGTextureSRVDesc Desc = FRDGTextureSRVDesc::Create(Texture);
@@ -715,6 +730,16 @@ public:
 		return HashCombine(GetTypeHash(static_cast<const FRHITextureSRVCreateInfo&>(Desc)), GetTypeHash(Desc.Texture));
 	}
 
+	/** Returns whether this descriptor conforms to requirements. */
+	bool IsValid() const
+	{
+		if (!Texture)
+		{
+			return false;
+		}
+		return FRHITextureSRVCreateInfo::Validate(Texture->Desc, *this, Texture->Name, /* bFatal = */ false);
+	}
+
 	FRDGTextureRef Texture = nullptr;
 };
 
diff --git a/Engine/Source/Runtime/RenderCore/Public/RenderGraphResources.inl b/Engine/Source/Runtime/RenderCore/Public/RenderGraphResources.inl
--- a/Engine/Source/Runtime/RenderCore/Public/RenderGraphResources.inl
+++ b/Engine/Source/Runtime/RenderCore/Public/RenderGraphResources.inl
@@ -166,9 +166,15 @@ inline FRDGTextureSubresourceRange FRDGTextureUAV::GetSubresourceRange() const
 {
 	FRDGTextureSubresourceRange Range = GetParent()->GetSubresourceRange();
 	Range.MipIndex = Desc.MipLevel;
+	Range.ArraySlice = Desc.FirstArraySlice;
 	Range.NumMips = 1;
 	Range.PlaneSlice = GetResourceTransitionPlaneForMetadataAccess(Desc.MetaData);
 
+	if (Desc.NumArraySlices != 0)
+	{
+		Range.NumArraySlices = Desc.NumArraySlices;
+	}
+
 	if (Desc.MetaData != ERDGTextureMetaDataAccess::None)
 	{
 		Range.NumPlaneSlices = 1;
