Use DSVRead | DSVWrite for TransitionRenderPassTargets to fix assert in SetTrackedAccess.

diff --git a/Engine/Source/Runtime/RHI/Public/RHICommandList.h b/Engine/Source/Runtime/RHI/Public/RHICommandList.h
--- a/Engine/Source/Runtime/RHI/Public/RHICommandList.h
+++ b/Engine/Source/Runtime/RHI/Public/RHICommandList.h
@@ -3613,6 +3613,7 @@ public:
 	using FRHIComputeCommandList::TransitionResource;
 	using FRHIComputeCommandList::TransitionResources;
 
+	UE_DEPRECATED(5.1, "TransitionResource has been deprecated. Use Transition instead.")
 	FORCEINLINE_DEBUGGABLE void TransitionResource(FExclusiveDepthStencil DepthStencilMode, FRHITexture* DepthTexture)
 	{
 		check(DepthStencilMode.IsUsingDepth() || DepthStencilMode.IsUsingStencil());
diff --git a/Engine/Source/Runtime/RHI/Public/RHIUtilities.h b/Engine/Source/Runtime/RHI/Public/RHIUtilities.h
--- a/Engine/Source/Runtime/RHI/Public/RHIUtilities.h
+++ b/Engine/Source/Runtime/RHI/Public/RHIUtilities.h
@@ -654,7 +654,7 @@ inline void TransitionRenderPassTargets(FRHICommandList& RHICmdList, const FRHIR
 	const FRHIRenderPassInfo::FDepthStencilEntry& DepthStencilTarget = RPInfo.DepthStencilRenderTarget;
 	if (DepthStencilTarget.DepthStencilTarget != nullptr && (RPInfo.DepthStencilRenderTarget.ExclusiveDepthStencil.IsAnyWrite()))
 	{
-		RHICmdList.TransitionResource(RPInfo.DepthStencilRenderTarget.ExclusiveDepthStencil, DepthStencilTarget.DepthStencilTarget);
+		RHICmdList.Transition(FRHITransitionInfo(DepthStencilTarget.DepthStencilTarget, ERHIAccess::Unknown, ERHIAccess::DSVRead | ERHIAccess::DSVWrite));
 	}
 
 	RHICmdList.Transition(MakeArrayView(Transitions, TransitionIndex));
