Add an RHIFrameTime to mesure when we actually kick the gamethread from the flip tracking thread. Give a more stable and representative FPS when using gtsynctype==2, which is used for the low latency frame syncing.

diff --git a/Engine/Source/Runtime/Engine/Private/UnrealEngine.cpp b/Engine/Source/Runtime/Engine/Private/UnrealEngine.cpp
--- a/Engine/Source/Runtime/Engine/Private/UnrealEngine.cpp
+++ b/Engine/Source/Runtime/Engine/Private/UnrealEngine.cpp
@@ -605,6 +605,30 @@ ENGINE_API float GAverageFPS = 0.0f;
 ENGINE_API float GAverageMS = 0.0f;
 ENGINE_API double GLastMemoryWarningTime = 0.f;
 
+
+void CalculateFPSTimings()
+{
+	// Calculate the average frame time via continued averaging.
+	static double LastTime = 0.0;
+	double CurrentTime = FPlatformTime::Seconds();
+	float FrameTimeMS = (float)((CurrentTime - LastTime) * 1000.0);
+
+	static auto CVarGTSyncType = IConsoleManager::Get().FindConsoleVariable(TEXT("r.GTSyncType"));
+	static auto CVarVsync = IConsoleManager::Get().FindConsoleVariable(TEXT("r.VSync"));
+	if (CVarGTSyncType->GetInt() == 2 && CVarVsync->GetInt() != 0)
+	{
+		FrameTimeMS = RHIGetFrameTime();
+	}
+	
+	// A 3/4, 1/4 split gets close to a simple 10 frame moving average
+	GAverageMS = GAverageMS * 0.75f + FrameTimeMS * 0.25f;
+
+	LastTime = CurrentTime;
+
+	// Calculate average framerate.
+	GAverageFPS = 1000.f / GAverageMS;
+}
+
 static FCachedSystemScalabilityCVars GCachedScalabilityCVars;
 
 const FCachedSystemScalabilityCVars& GetCachedScalabilityCVars()
diff --git a/Engine/Source/Runtime/Engine/Public/UnrealEngine.h b/Engine/Source/Runtime/Engine/Public/UnrealEngine.h
--- a/Engine/Source/Runtime/Engine/Public/UnrealEngine.h
+++ b/Engine/Source/Runtime/Engine/Public/UnrealEngine.h
@@ -357,20 +357,8 @@ public:
 ENGINE_API FString appGetStartupMap(const TCHAR* CommandLine);
 
 // Calculate the average frame time by using the stats system.
-inline void CalculateFPSTimings()
-{
-	extern ENGINE_API float GAverageFPS;
-	extern ENGINE_API float GAverageMS;
-	// Calculate the average frame time via continued averaging.
-	static double LastTime = 0.0;
-	double CurrentTime = FPlatformTime::Seconds();
-	const float FrameTimeMS = (float)((CurrentTime - LastTime) * 1000.0);
-	// A 3/4, 1/4 split gets close to a simple 10 frame moving average
-	GAverageMS = GAverageMS * 0.75f + FrameTimeMS * 0.25f;
-	LastTime = CurrentTime;
-	// Calculate average framerate.
-	GAverageFPS = 1000.f / GAverageMS;
-}
+ENGINE_API void CalculateFPSTimings();
+
 
 /** @return The font to use for rendering stats display. */
 extern ENGINE_API UFont* GetStatsFont();
diff --git a/Engine/Source/Runtime/RHI/Private/RHIUtilities.cpp b/Engine/Source/Runtime/RHI/Private/RHIUtilities.cpp
--- a/Engine/Source/Runtime/RHI/Private/RHIUtilities.cpp
+++ b/Engine/Source/Runtime/RHI/Private/RHIUtilities.cpp
@@ -124,6 +124,9 @@ TAutoConsoleVariable<int32> CVarRHISyncAllowVariable(
 	ECVF_Default
 );
 
+float GRHIFrameTimeMS = 0.0f;
+double GLastRHITimeInSeconds = 0.0;
+
 int32 GEnableConsole120Fps = 0;
 int32 InternalEnableConsole120Fps = 0;
 static void OnEnableConsole120FpsCVarRHIChanged(IConsoleVariable* Variable)
@@ -393,6 +396,7 @@ uint32 FRHIFrameFlipTrackingRunnable::Run()
 			SyncTime = CurrentTimeInSeconds;
 		}
 
+		bool bUpdateRHIFrameTime = false;
 		// Complete any relevant task graph events.
 		FScopeLock Lock(&CS);
 		for (int32 PairIndex = FramePairs.Num() - 1; PairIndex >= 0; --PairIndex)
@@ -404,8 +408,14 @@ uint32 FRHIFrameFlipTrackingRunnable::Run()
 				Pair.Event->DispatchSubsequents();
 
 				FramePairs.RemoveAtSwap(PairIndex);
+				bUpdateRHIFrameTime = true;
 			}
 		}
+
+		if(bUpdateRHIFrameTime)
+		{
+			RHICalculateFrameTime();
+		}	
 	}
 
 	return 0;
@@ -551,6 +561,18 @@ RHI_API void RHIInitializeFlipTracking()
 	FRHIFrameFlipTrackingRunnable::Initialize();
 }
 
+RHI_API void RHICalculateFrameTime()
+{
+	double CurrentTimeInSeconds = FPlatformTime::Seconds();
+	GRHIFrameTimeMS = (float)((CurrentTimeInSeconds - GLastRHITimeInSeconds) * 1000.0);
+	GLastRHITimeInSeconds = CurrentTimeInSeconds;
+}
+
+RHI_API float RHIGetFrameTime()
+{
+	return GRHIFrameTimeMS;
+}
+
 RHI_API void RHIShutdownFlipTracking()
 {
 	FRHIFrameFlipTrackingRunnable::Shutdown();
diff --git a/Engine/Source/Runtime/RHI/Public/RHIUtilities.h b/Engine/Source/Runtime/RHI/Public/RHIUtilities.h
--- a/Engine/Source/Runtime/RHI/Public/RHIUtilities.h
+++ b/Engine/Source/Runtime/RHI/Public/RHIUtilities.h
@@ -902,6 +902,11 @@ extern RHI_API void RHISetFrameDebugInfo(uint64 PresentIndex, uint64 FrameIndex,
 extern RHI_API void RHIInitializeFlipTracking();
 extern RHI_API void RHIShutdownFlipTracking();
 
+/** Sets the FrameIndex and InputTime for the current frame. */
+extern RHI_API float RHIGetFrameTime();
+
+extern RHI_API void RHICalculateFrameTime();
+
 struct FRHILockTracker
 {
 	struct FLockParams
