Deprecated RDG variant of texture readback.

diff --git a/Engine/Source/Runtime/RHI/Public/RHIGPUReadback.h b/Engine/Source/Runtime/RHI/Public/RHIGPUReadback.h
--- a/Engine/Source/Runtime/RHI/Public/RHIGPUReadback.h
+++ b/Engine/Source/Runtime/RHI/Public/RHIGPUReadback.h
@@ -110,9 +110,10 @@ class RHI_API FRHIGPUTextureReadback final : public FRHIGPUMemoryReadback
 public:
 	FRHIGPUTextureReadback(FName RequestName);
 
+	UE_DEPRECATED(5.1, "EnqueueCopyRDG is deprecated. Use EnqueueCopy instead.")
 	void EnqueueCopyRDG(FRHICommandList& RHICmdList, FRHITexture* SourceTexture, FResolveRect Rect = FResolveRect());
-	void EnqueueCopy(FRHICommandList& RHICmdList, FRHITexture* SourceTexture, FResolveRect Rect = FResolveRect()) override;
 
+	void EnqueueCopy(FRHICommandList& RHICmdList, FRHITexture* SourceTexture, FResolveRect Rect = FResolveRect()) override;
 
 	UE_DEPRECATED(5.0, "Use FRHIGPUTextureReadback::Lock( int32& OutRowPitchInPixels) instead.")
 	void* Lock(uint32 NumBytes) override;
diff --git a/Engine/Source/Runtime/RenderCore/Private/RenderGraphUtils.cpp b/Engine/Source/Runtime/RenderCore/Private/RenderGraphUtils.cpp
--- a/Engine/Source/Runtime/RenderCore/Private/RenderGraphUtils.cpp
+++ b/Engine/Source/Runtime/RenderCore/Private/RenderGraphUtils.cpp
@@ -818,7 +818,7 @@ void AddEnqueueCopyPass(FRDGBuilder& GraphBuilder, FRHIGPUTextureReadback* Readb
 		ERDGPassFlags::Readback,
 		[Readback, SourceTexture, Rect](FRHICommandList& RHICmdList)
 	{
-		Readback->EnqueueCopyRDG(RHICmdList, SourceTexture->GetRHI(), Rect);
+		Readback->EnqueueCopy(RHICmdList, SourceTexture->GetRHI(), Rect);
 	});
 }
 
