Fixed transition validation error for GPU readback staging texture.

diff --git a/Engine/Source/Runtime/RHI/Private/RHIGPUReadback.cpp b/Engine/Source/Runtime/RHI/Private/RHIGPUReadback.cpp
--- a/Engine/Source/Runtime/RHI/Private/RHIGPUReadback.cpp
+++ b/Engine/Source/Runtime/RHI/Private/RHIGPUReadback.cpp
@@ -145,7 +145,7 @@ void FRHIGPUTextureReadback::EnqueueCopyInternal(FRHICommandList& RHICmdList, FR
 			const FRHITextureCreateDesc Desc =
 				FRHITextureCreateDesc::Create2D(*FenceName, TextureSize.X, TextureSize.Y, SourceTexture->GetFormat())
 				.SetFlags(ETextureCreateFlags::CPUReadback | ETextureCreateFlags::HideInVisualizeTexture)
-				.SetInitialState(ERHIAccess::CPURead);
+				.SetInitialState(ERHIAccess::CopyDest);
 
 			DestinationStagingTexture = RHICreateTexture(Desc);
 		}
@@ -168,6 +168,7 @@ void* FRHIGPUTextureReadback::Lock(uint32 NumBytes)
 		void* ResultsBuffer = nullptr;
 		int32 BufferWidth = 0, BufferHeight = 0;
 		FRHICommandListImmediate& RHICmdList = FRHICommandListExecutor::GetImmediateCommandList();
+		RHICmdList.Transition(FRHITransitionInfo(DestinationStagingTexture, ERHIAccess::CopyDest, ERHIAccess::CPURead));
 		RHICmdList.MapStagingSurface(DestinationStagingTexture, Fence.GetReference(), ResultsBuffer, BufferWidth, BufferHeight, LastLockGPUIndex);
 		return ResultsBuffer;
 	}
@@ -186,6 +187,7 @@ void* FRHIGPUTextureReadback::Lock(int32& OutRowPitchInPixels, int32 *OutBufferH
 		void* ResultsBuffer = nullptr;
 		int32 BufferWidth = 0, BufferHeight = 0;
 		FRHICommandListImmediate& RHICmdList = FRHICommandListExecutor::GetImmediateCommandList();
+		RHICmdList.Transition(FRHITransitionInfo(DestinationStagingTexture, ERHIAccess::CopyDest, ERHIAccess::CPURead));
 		RHICmdList.MapStagingSurface(DestinationStagingTexture, Fence.GetReference(), ResultsBuffer, BufferWidth, BufferHeight, LastLockGPUIndex);
 		if (OutBufferHeight)
 		{
@@ -214,4 +216,5 @@ void FRHIGPUTextureReadback::Unlock()
 
 	FRHICommandListImmediate& RHICmdList = FRHICommandListExecutor::GetImmediateCommandList();
 	RHICmdList.UnmapStagingSurface(DestinationStagingTexture, LastLockGPUIndex);
+	RHICmdList.Transition(FRHITransitionInfo(DestinationStagingTexture, ERHIAccess::CPURead, ERHIAccess::CopyDest));
 }
