cleaning up global RHICommandList graph events on exit to avoid undefined order of destruction

diff --git a/Engine/Source/Runtime/RHI/Private/DynamicRHI.cpp b/Engine/Source/Runtime/RHI/Private/DynamicRHI.cpp
--- a/Engine/Source/Runtime/RHI/Private/DynamicRHI.cpp
+++ b/Engine/Source/Runtime/RHI/Private/DynamicRHI.cpp
@@ -410,6 +410,9 @@ void RHIExit()
 		FRHICommandListExecutor::GetImmediateCommandList().ImmediateFlush(EImmediateFlushType::FlushRHIThreadFlushResources);
 		FRHICommandListExecutor::GetImmediateCommandList().ImmediateFlush(EImmediateFlushType::FlushRHIThread);
 	}
+
+	void CleanupRHICommandListGraphEvents();
+	CleanupRHICommandListGraphEvents();
 }
 
 
diff --git a/Engine/Source/Runtime/RHI/Private/RHICommandList.cpp b/Engine/Source/Runtime/RHI/Private/RHICommandList.cpp
--- a/Engine/Source/Runtime/RHI/Private/RHICommandList.cpp
+++ b/Engine/Source/Runtime/RHI/Private/RHICommandList.cpp
@@ -161,6 +161,22 @@ static FGraphEventRef RHIThreadBufferLockFence;
 static FGraphEventRef GRHIThreadEndDrawingViewportFences[2];
 static uint32 GRHIThreadEndDrawingViewportFenceIndex = 0;
 
+// global graph events must be destroyed explicitly to avoid undefined order of static destruction, as they can be destroyed after their allocator
+void CleanupRHICommandListGraphEvents()
+{
+	AllOutstandingTasks.Reset();
+	WaitOutstandingTasks.Reset();
+	RHIThreadTask.SafeRelease();
+	PrevRHIThreadTask.SafeRelease();
+	RenderThreadSublistDispatchTask.SafeRelease();
+	RHIThreadBufferLockFence.SafeRelease();
+
+	for (FGraphEventRef& GraphEvent : GRHIThreadEndDrawingViewportFences)
+	{
+		GraphEvent.SafeRelease();
+	}
+}
+
 // Used by AsyncCompute
 RHI_API FRHICommandListFenceAllocator GRHIFenceAllocator;
 
