Removing the long-form create functions and constructors from FRHITextureCreateDesc in favor of the shorter versions.

diff --git a/Engine/Plugins/Media/PixelStreaming/Source/PixelStreaming/Private/TextureSourceComputeI420.cpp b/Engine/Plugins/Media/PixelStreaming/Source/PixelStreaming/Private/TextureSourceComputeI420.cpp
--- a/Engine/Plugins/Media/PixelStreaming/Source/PixelStreaming/Private/TextureSourceComputeI420.cpp
+++ b/Engine/Plugins/Media/PixelStreaming/Source/PixelStreaming/Private/TextureSourceComputeI420.cpp
@@ -61,53 +61,31 @@ namespace
 	{
 		Buffer = webrtc::I420Buffer::Create(PlaneYDimensions.X, PlaneYDimensions.Y);
 
-		FRHITextureCreateDesc TextureDescY = FRHITextureCreateDesc::Create2D(
-			TEXT("Compute YUV Target"),
-			{ PlaneYDimensions.X, PlaneYDimensions.Y },
-			EPixelFormat::PF_R8,
-			FClearValueBinding::None,
-			TexCreate_UAV,
-			1,
-			1,
-			0,
-			ERHIAccess::UAVCompute);
-
-		FRHITextureCreateDesc TextureDescUV = FRHITextureCreateDesc::Create2D(
-			TEXT("Compute YUV Target"),
-			{ PlaneUVDimensions.X, PlaneUVDimensions.Y },
-			EPixelFormat::PF_R8,
-			FClearValueBinding::None,
-			TexCreate_UAV,
-			1,
-			1,
-			0,
-			ERHIAccess::UAVCompute);
+		const FRHITextureCreateDesc TextureDescY =
+			FRHITextureCreateDesc::Create2D(TEXT("Compute YUV Target"), PlaneYDimensions, PF_R8)
+			.SetClearValue(FClearValueBinding::None)
+			.SetFlags(ETextureCreateFlags::UAV)
+			.SetInitialState(ERHIAccess::UAVCompute);
+
+		const FRHITextureCreateDesc TextureDescUV =
+			FRHITextureCreateDesc::Create2D(TEXT("Compute YUV Target"), PlaneUVDimensions, PF_R8)
+			.SetClearValue(FClearValueBinding::None)
+			.SetFlags(ETextureCreateFlags::UAV)
+			.SetInitialState(ERHIAccess::UAVCompute);
 
 		TextureY = GDynamicRHI->RHICreateTexture(TextureDescY);
 		TextureU = GDynamicRHI->RHICreateTexture(TextureDescUV);
 		TextureV = GDynamicRHI->RHICreateTexture(TextureDescUV);
 
-		FRHITextureCreateDesc StagingDescY = FRHITextureCreateDesc::Create2D(
-			TEXT("YUV Output CPU Texture"),
-			{ PlaneYDimensions.X, PlaneYDimensions.Y },
-			EPixelFormat::PF_R8,
-			FClearValueBinding::None,
-			TexCreate_CPUReadback,
-			1,
-			1,
-			0,
-			ERHIAccess::Unknown);
-
-		FRHITextureCreateDesc StagingDescUV = FRHITextureCreateDesc::Create2D(
-			TEXT("YUV Output CPU Texture"),
-			{ PlaneUVDimensions.X, PlaneUVDimensions.Y },
-			EPixelFormat::PF_R8,
-			FClearValueBinding::None,
-			TexCreate_CPUReadback,
-			1,
-			1,
-			0,
-			ERHIAccess::Unknown);
+		const FRHITextureCreateDesc StagingDescY =
+			FRHITextureCreateDesc::Create2D(TEXT("YUV Output CPU Texture"), PlaneYDimensions, PF_R8)
+			.SetClearValue(FClearValueBinding::None)
+			.SetFlags(ETextureCreateFlags::CPUReadback);
+
+		const FRHITextureCreateDesc StagingDescUV =
+			FRHITextureCreateDesc::Create2D(TEXT("YUV Output CPU Texture"), PlaneUVDimensions, PF_R8)
+			.SetClearValue(FClearValueBinding::None)
+			.SetFlags(ETextureCreateFlags::CPUReadback);
 
 		StagingTextureY = GDynamicRHI->RHICreateTexture(StagingDescY);
 		StagingTextureU = GDynamicRHI->RHICreateTexture(StagingDescUV);
@@ -147,16 +125,10 @@ namespace
 
 		if (!TempFBCopy)
 		{
-			FRHITextureCreateDesc Desc = FRHITextureCreateDesc::Create2D(
-				TEXT("PixelStreamingBlankTexture"),
-				{ (int32)SourceTexture->GetDesc().Extent.X, (int32)SourceTexture->GetDesc().Extent.Y },
-				SourceTexture->GetDesc().Format,
-				FClearValueBinding::None,
-				TexCreate_ShaderResource,
-				1,
-				1,
-				0,
-				ERHIAccess::Unknown);
+			const FRHITextureCreateDesc Desc =
+				FRHITextureCreateDesc::Create2D(TEXT("PixelStreamingBlankTexture"), SourceTexture->GetDesc().Extent, SourceTexture->GetDesc().Format)
+				.SetClearValue(FClearValueBinding::None)
+				.SetFlags(ETextureCreateFlags::ShaderResource);
 
 			TempFBCopy = GDynamicRHI->RHICreateTexture(Desc);
 		}
diff --git a/Engine/Plugins/Media/PixelStreaming/Source/PixelStreaming/Private/Utils.h b/Engine/Plugins/Media/PixelStreaming/Source/PixelStreaming/Private/Utils.h
--- a/Engine/Plugins/Media/PixelStreaming/Source/PixelStreaming/Private/Utils.h
+++ b/Engine/Plugins/Media/PixelStreaming/Source/PixelStreaming/Private/Utils.h
@@ -76,64 +76,42 @@ namespace UE::PixelStreaming
 	inline FTextureRHIRef CreateRHITexture(uint32 Width, uint32 Height)
 	{
 		// Create empty texture
-		ETextureCreateFlags Flags;
-		ERHIAccess InitialState;
+		FRHITextureCreateDesc Desc =
+			FRHITextureCreateDesc::Create2D(TEXT("PixelStreamingBlankTexture"), Width, Height, PF_B8G8R8A8)
+			.SetClearValue(FClearValueBinding::None)
+			.SetFlags(ETextureCreateFlags::RenderTargetable);
 
 		if (RHIGetInterfaceType() == ERHIInterfaceType::Vulkan)
 		{
-			Flags = TexCreate_RenderTargetable | TexCreate_External;
-			InitialState = ERHIAccess::Present;
+			Desc.AddFlags(ETextureCreateFlags::External)
+				.SetInitialState(ERHIAccess::Present);
 		}
 		else
 		{
-			Flags = TexCreate_RenderTargetable | TexCreate_Shared;
-			InitialState = ERHIAccess::CopyDest;
+			Desc.AddFlags(ETextureCreateFlags::Shared)
+				.SetInitialState(ERHIAccess::CopyDest);
 		}
 
-		FRHITextureCreateDesc Desc = FRHITextureCreateDesc::Create2D(
-			TEXT("PixelStreamingBlankTexture"),
-			{ (int32)Width, (int32)Height },
-			EPixelFormat::PF_B8G8R8A8,
-			FClearValueBinding::None,
-			Flags,
-			1,
-			1,
-			0,
-			InitialState
-		);
-
 		return GDynamicRHI->RHICreateTexture(Desc);
 	}
 
 	inline FTextureRHIRef CreateCPUReadbackTexture(uint32 Width, uint32 Height)
 	{
 		// Create empty texture
-		ETextureCreateFlags Flags;
-		ERHIAccess InitialState;
+		FRHITextureCreateDesc Desc =
+			FRHITextureCreateDesc::Create2D(TEXT("PixelStreamingBlankTexture"), Width, Height, PF_B8G8R8A8)
+			.SetClearValue(FClearValueBinding::None)
+			.SetFlags(ETextureCreateFlags::HideInVisualizeTexture | ETextureCreateFlags::CPUReadback);
 
 		if (RHIGetInterfaceType() == ERHIInterfaceType::Vulkan)
 		{
-			Flags = TexCreate_HideInVisualizeTexture | TexCreate_CPUReadback;
-			InitialState = ERHIAccess::Present;
+			Desc.SetInitialState(ERHIAccess::Present);
 		}
 		else
 		{
-			Flags = TexCreate_HideInVisualizeTexture | TexCreate_CPUReadback;
-			InitialState = ERHIAccess::CopyDest;
+			Desc.SetInitialState(ERHIAccess::CopyDest);
 		}
 
-		FRHITextureCreateDesc Desc = FRHITextureCreateDesc::Create2D(
-			TEXT("PixelStreamingBlankTexture"),
-			{ (int32)Width, (int32)Height },
-			EPixelFormat::PF_B8G8R8A8,
-			FClearValueBinding::None,
-			Flags,
-			1,
-			1,
-			0,
-			InitialState
-		);
-
 		return GDynamicRHI->RHICreateTexture(Desc);
 	}
 
diff --git a/Engine/Source/Runtime/NullDrv/Public/NullRHI.h b/Engine/Source/Runtime/NullDrv/Public/NullRHI.h
--- a/Engine/Source/Runtime/NullDrv/Public/NullRHI.h
+++ b/Engine/Source/Runtime/NullDrv/Public/NullRHI.h
@@ -217,22 +217,14 @@ public:
 
 	virtual FTextureRHIRef RHIAsyncCreateTexture2D(uint32 SizeX, uint32 SizeY, uint8 Format, uint32 NumMips, ETextureCreateFlags Flags, ERHIAccess InResourceState, void** InitialMipData, uint32 NumInitialMips) final override
 	{ 
-		const uint32 NumSamples = 1;
-		const uint32 ExtData = 0;
-
-		return this->RHICreateTexture(
-			FRHITextureCreateDesc::Create2D(
-				  TEXT("FNullDynamicRHI::RHIAsyncCreateTexture2D")
-				, { (int32)SizeX, (int32)SizeY}
-				, (EPixelFormat)Format
-				, FClearValueBinding::None
-				, Flags
-				, NumMips
-				, NumSamples
-				, ExtData
-				, InResourceState
-			)
-		);
+		const FRHITextureCreateDesc Desc =
+			FRHITextureCreateDesc::Create2D(TEXT("FNullDynamicRHI::RHIAsyncCreateTexture2D"), SizeX, SizeY, (EPixelFormat)Format)
+			.SetClearValue(FClearValueBinding::None)
+			.SetFlags(Flags)
+			.SetNumMips(NumMips)
+			.SetInitialState(InResourceState);
+
+		return this->RHICreateTexture(Desc);
 	}
 
 	virtual void RHICopySharedMips(FRHITexture2D* DestTexture2D, FRHITexture2D* SrcTexture2D) final override
@@ -253,18 +245,15 @@ public:
 		return 0; 
 	}
 	virtual FTexture2DRHIRef RHIAsyncReallocateTexture2D(FRHITexture2D* Texture2D, int32 NewMipCount, int32 NewSizeX, int32 NewSizeY, FThreadSafeCounter* RequestStatus) final override
-	{ 
-		return this->RHICreateTexture(
-			FRHITextureCreateDesc::Create2D(
-				  TEXT("FNullDynamicRHI::RHIAsyncReallocateTexture2D")
-				, { NewSizeX, NewSizeY }
-				, Texture2D->GetFormat()
-				, Texture2D->GetClearBinding()
-				, Texture2D->GetFlags()
-				, NewMipCount
-				, Texture2D->GetNumSamples()
-			)
-		);
+	{
+		const FRHITextureCreateDesc Desc =
+			FRHITextureCreateDesc::Create2D(TEXT("FNullDynamicRHI::RHIAsyncReallocateTexture2D"), NewSizeX, NewSizeY, Texture2D->GetFormat())
+			.SetClearValue(Texture2D->GetClearBinding())
+			.SetFlags(Texture2D->GetFlags())
+			.SetNumMips(NewMipCount)
+			.SetNumSamples(Texture2D->GetNumSamples());
+
+		return this->RHICreateTexture(Desc);
 	}
 	virtual ETextureReallocationStatus RHIFinalizeAsyncReallocateTexture2D(FRHITexture2D* Texture2D, bool bBlockUntilCompleted) final override
 	{ 
@@ -388,15 +377,11 @@ public:
 
 	virtual FTexture2DRHIRef RHIGetViewportBackBuffer(FRHIViewport* Viewport) final override
 	{
-		return this->RHICreateTexture(
-			FRHITextureCreateDesc::Create2D(
-				  TEXT("FNullDynamicRHI::RHIGetViewportBackBuffer")
-				, { 1, 1 }
-				, PF_B8G8R8A8
-				, FClearValueBinding::None
-				, ETextureCreateFlags::RenderTargetable
-			)
-		);
+		const FRHITextureCreateDesc Desc =
+			FRHITextureCreateDesc::Create2D(TEXT("FNullDynamicRHI::RHIGetViewportBackBuffer"), 1, 1, PF_B8G8R8A8)
+			.SetFlags(ETextureCreateFlags::RenderTargetable);
+
+		return this->RHICreateTexture(Desc);
 	}
 
 	virtual void RHIBeginFrame() final override
diff --git a/Engine/Source/Runtime/OpenGLDrv/Private/Android/AndroidOpenGL.cpp b/Engine/Source/Runtime/OpenGLDrv/Private/Android/AndroidOpenGL.cpp
--- a/Engine/Source/Runtime/OpenGLDrv/Private/Android/AndroidOpenGL.cpp
+++ b/Engine/Source/Runtime/OpenGLDrv/Private/Android/AndroidOpenGL.cpp
@@ -451,13 +451,9 @@ FOpenGLTexture* PlatformCreateBuiltinBackBuffer(FOpenGLDynamicRHI* OpenGLRHI, ui
 	// Otherwise return null and we will create an off-screen surface afterward.
 	if (!FPlatformMisc::SupportsBackbufferSampling())
 	{
-		FOpenGLTextureCreateDesc CreateDesc = FOpenGLTextureCreateDesc::Create2D(
-			TEXT("PlatformCreateBuiltinBackBuffer"),
-			{ (int32)SizeX, (int32)SizeY },
-			PF_B8G8R8A8,
-			FClearValueBinding::Transparent,
-			TexCreate_RenderTargetable | TexCreate_Presentable | TexCreate_ResolveTargetable
-		);
+		const FOpenGLTextureCreateDesc CreateDesc =
+			FRHITextureCreateDesc::Create2D(TEXT("PlatformCreateBuiltinBackBuffer"), SizeX, SizeY, PF_B8G8R8A8)
+			.SetFlags(ETextureCreateFlags::RenderTargetable | ETextureCreateFlags::Presentable | ETextureCreateFlags::ResolveTargetable);
 
 		return new FOpenGLTexture(CreateDesc);
 	}
diff --git a/Engine/Source/Runtime/OpenGLDrv/Private/OpenGLTexture.cpp b/Engine/Source/Runtime/OpenGLDrv/Private/OpenGLTexture.cpp
--- a/Engine/Source/Runtime/OpenGLDrv/Private/OpenGLTexture.cpp
+++ b/Engine/Source/Runtime/OpenGLDrv/Private/OpenGLTexture.cpp
@@ -2227,58 +2227,41 @@ void FOpenGLDynamicRHI::RHICopyTexture(FRHITexture* SourceTextureRHI, FRHITextur
 
 FTexture2DRHIRef FOpenGLDynamicRHI::RHICreateTexture2DFromResource(EPixelFormat Format, uint32 SizeX, uint32 SizeY, uint32 NumMips, uint32 NumSamples, uint32 NumSamplesTileMem, const FClearValueBinding& ClearValueBinding, GLuint Resource, ETextureCreateFlags TexCreateFlags)
 {
-	return new FOpenGLTexture(FRHITextureCreateDesc::Create2D(
-		TEXT("RHICreateTexture2DFromResource"),
-		{ (int32)SizeX, (int32)SizeY },
-		Format,
-		ClearValueBinding,
-		TexCreateFlags,
-		NumMips,
-		NumSamples
-	), Resource);
+	const FRHITextureCreateDesc Desc =
+		FRHITextureCreateDesc::Create2D(TEXT("RHICreateTexture2DFromResource"), SizeX, SizeY, Format)
+		.SetClearValue(ClearValueBinding)
+		.SetFlags(TexCreateFlags)
+		.SetNumMips(NumMips)
+		.SetNumSamples(NumSamples);
+
+	return new FOpenGLTexture(Desc, Resource);
 }
 
 FTexture2DRHIRef FOpenGLDynamicRHI::RHICreateTexture2DArrayFromResource(EPixelFormat Format, uint32 SizeX, uint32 SizeY, uint32 ArraySize, uint32 NumMips, uint32 NumSamples, uint32 NumSamplesTileMem, const FClearValueBinding& ClearValueBinding, GLuint Resource, ETextureCreateFlags TexCreateFlags)
 {
-	return new FOpenGLTexture(FRHITextureCreateDesc::Create2DArray(
-		TEXT("RHICreateTexture2DArrayFromResource"),
-		{ (int32)SizeX, (int32)SizeY },
-		Format,
-		ClearValueBinding,
-		TexCreateFlags,
-		ArraySize,
-		NumMips,
-		NumSamples
-	), Resource);
+	const FRHITextureCreateDesc Desc =
+		FRHITextureCreateDesc::Create2DArray(TEXT("RHICreateTexture2DArrayFromResource"), SizeX, SizeY, ArraySize, Format)
+		.SetClearValue(ClearValueBinding)
+		.SetFlags(TexCreateFlags)
+		.SetNumMips(NumMips)
+		.SetNumSamples(NumSamples);
+
+	return new FOpenGLTexture(Desc, Resource);
 }
 
 FTextureCubeRHIRef FOpenGLDynamicRHI::RHICreateTextureCubeFromResource(EPixelFormat Format, uint32 Size, bool bArray, uint32 ArraySize, uint32 NumMips, uint32 NumSamples, uint32 NumSamplesTileMem, const FClearValueBinding& ClearValueBinding, GLuint Resource, ETextureCreateFlags TexCreateFlags)
 {
-	if (bArray)
-	{
-		return new FOpenGLTexture(FRHITextureCreateDesc::CreateCubeArray(
-			TEXT("RHICreateTextureCubeFromResource"),
-			Size,
-			Format,
-			ClearValueBinding,
-			TexCreateFlags,
-			ArraySize,
-			NumMips,
-			NumSamples
-		), Resource);
-	}
-	else
-	{
-		return new FOpenGLTexture(FRHITextureCreateDesc::CreateCube(
-			TEXT("RHICreateTextureCubeFromResource"),
-			Size,
-			Format,
-			ClearValueBinding,
-			TexCreateFlags,
-			NumMips,
-			NumSamples
-		), Resource);
-	}
+	const FRHITextureCreateDesc Desc =
+		FRHITextureCreateDesc::Create(TEXT("RHICreateTextureCubeFromResource"), bArray ? ETextureDimension::TextureCube : ETextureDimension::TextureCubeArray)
+		.SetExtent(Size)
+		.SetArraySize(bArray ? ArraySize : 1)
+		.SetFormat(Format)
+		.SetClearValue(ClearValueBinding)
+		.SetFlags(TexCreateFlags)
+		.SetNumMips(NumMips)
+		.SetNumSamples(NumSamples);
+
+	return new FOpenGLTexture(Desc, Resource);
 }
 
 void FOpenGLDynamicRHI::RHIAliasTextureResources(FTextureRHIRef& DestRHITexture, FTextureRHIRef& SrcRHITexture)
diff --git a/Engine/Source/Runtime/OpenGLDrv/Private/OpenGLViewport.cpp b/Engine/Source/Runtime/OpenGLDrv/Private/OpenGLViewport.cpp
--- a/Engine/Source/Runtime/OpenGLDrv/Private/OpenGLViewport.cpp
+++ b/Engine/Source/Runtime/OpenGLDrv/Private/OpenGLViewport.cpp
@@ -315,13 +315,12 @@ void FOpenGLViewport::Resize(uint32 InSizeX,uint32 InSizeY,bool bInIsFullscreen)
 		BackBuffer = PlatformCreateBuiltinBackBuffer(OpenGLRHI, InSizeX, InSizeY);
 		if (!BackBuffer)
 		{
-			BackBuffer = new FOpenGLTexture(FRHITextureCreateDesc::Create2D(
-				TEXT("FOpenGLViewport"),
-				{ (int32)InSizeX, (int32)InSizeY },
-				PixelFormat,
-				FClearValueBinding::Transparent,
-				TexCreate_RenderTargetable | TexCreate_ResolveTargetable
-			));
+			const FRHITextureCreateDesc Desc =
+				FRHITextureCreateDesc::Create2D(TEXT("FOpenGLViewport"), InSizeX, InSizeY, PixelFormat)
+				.SetClearValue(FClearValueBinding::Transparent)
+				.SetFlags(ETextureCreateFlags::RenderTargetable | ETextureCreateFlags::ResolveTargetable);
+
+			BackBuffer = new FOpenGLTexture(Desc);
 		}
 
 		RHICmdList.EnqueueLambda([=](FRHICommandListImmediate&)
diff --git a/Engine/Source/Runtime/OpenGLDrv/Private/Windows/OpenGLWindows.cpp b/Engine/Source/Runtime/OpenGLDrv/Private/Windows/OpenGLWindows.cpp
--- a/Engine/Source/Runtime/OpenGLDrv/Private/Windows/OpenGLWindows.cpp
+++ b/Engine/Source/Runtime/OpenGLDrv/Private/Windows/OpenGLWindows.cpp
@@ -1003,13 +1003,12 @@ FOpenGLTexture* PlatformCreateBuiltinBackBuffer(FOpenGLDynamicRHI* OpenGLRHI, ui
 {
 	if (FOpenGL::IsAndroidGLESCompatibilityModeEnabled())
 	{
-		return new FOpenGLTexture(FRHITextureCreateDesc::Create2D(
-			TEXT("PlatformCreateBuiltinBackBuffer"),
-			{ (int32)SizeX, (int32)SizeY },
-			PF_B8G8R8A8,
-			FClearValueBinding::Transparent,
-			TexCreate_RenderTargetable | TexCreate_Presentable | TexCreate_ResolveTargetable
-		));
+		const FRHITextureCreateDesc Desc =
+			FRHITextureCreateDesc::Create2D(TEXT("PlatformCreateBuiltinBackBuffer"), SizeX, SizeY, PF_B8G8R8A8)
+			.SetClearValue(FClearValueBinding::Transparent)
+			.SetFlags(ETextureCreateFlags::RenderTargetable | ETextureCreateFlags::Presentable | ETextureCreateFlags::ResolveTargetable);
+
+		return new FOpenGLTexture(Desc);
 	}
 
 	return nullptr;
diff --git a/Engine/Source/Runtime/RHI/Public/RHIResources.h b/Engine/Source/Runtime/RHI/Public/RHIResources.h
--- a/Engine/Source/Runtime/RHI/Public/RHIResources.h
+++ b/Engine/Source/Runtime/RHI/Public/RHIResources.h
@@ -1555,33 +1555,6 @@ struct FRHITextureCreateDesc : public FRHITextureDesc
 			.SetFormat(Format);
 	}
 
-	static FRHITextureCreateDesc Create2D(
-		  TCHAR const*                DebugName
-		, FIntPoint                   Size
-		, EPixelFormat                Format
-		, FClearValueBinding          ClearValue
-		, ETextureCreateFlags         Flags
-		, uint8                       NumMips      = 1
-		, uint8                       NumSamples   = 1
-		, uint32		              ExtData      = 0
-		, ERHIAccess                  InitialState = ERHIAccess::Unknown
-		, FResourceBulkDataInterface* BulkData     = nullptr
-		, FRHIGPUMask                 GPUMask      = FRHIGPUMask::All()
-		)
-	{
-		return Create2D(DebugName)
-			.SetExtent(Size)
-			.SetFormat(Format)
-			.SetClearValue(ClearValue)
-			.SetFlags(Flags)
-			.SetNumMips(NumMips)
-			.SetNumSamples(NumSamples)
-			.SetExtData(ExtData)
-			.SetBulkData(BulkData)
-			.SetGPUMask(GPUMask)
-			.SetInitialState(InitialState == ERHIAccess::Unknown ? RHIGetDefaultResourceState(Flags, BulkData != nullptr) : InitialState);
-	}
-
 	static FRHITextureCreateDesc Create2DArray(const TCHAR* DebugName, FIntPoint Size, uint16 ArraySize, EPixelFormat Format)
 	{
 		return Create2DArray(DebugName)
@@ -1598,35 +1571,6 @@ struct FRHITextureCreateDesc : public FRHITextureDesc
 			.SetArraySize(ArraySize);
 	}
 
-	static FRHITextureCreateDesc Create2DArray(
-		  TCHAR const*                DebugName
-		, FIntPoint                   Size
-		, EPixelFormat                Format
-		, FClearValueBinding          ClearValue
-		, ETextureCreateFlags         Flags
-		, uint16                      ArraySize
-		, uint8                       NumMips      = 1
-		, uint8                       NumSamples   = 1
-		, uint32		              ExtData      = 0
-		, ERHIAccess                  InitialState = ERHIAccess::Unknown
-		, FResourceBulkDataInterface* BulkData     = nullptr
-		, FRHIGPUMask                 GPUMask      = FRHIGPUMask::All()
-		)
-	{
-		return Create2DArray(DebugName)
-			.SetExtent(Size)
-			.SetFormat(Format)
-			.SetClearValue(ClearValue)
-			.SetFlags(Flags)
-			.SetArraySize(ArraySize)
-			.SetNumMips(NumMips)
-			.SetNumSamples(NumSamples)
-			.SetExtData(ExtData)
-			.SetBulkData(BulkData)
-			.SetGPUMask(GPUMask)
-			.SetInitialState(InitialState == ERHIAccess::Unknown ? RHIGetDefaultResourceState(Flags, BulkData != nullptr) : InitialState);
-	}
-
 	static FRHITextureCreateDesc Create3D(const TCHAR* DebugName, FIntVector Size, EPixelFormat Format)
 	{
 		return Create3D(DebugName)
@@ -1643,32 +1587,6 @@ struct FRHITextureCreateDesc : public FRHITextureDesc
 			.SetFormat(Format);
 	}
 
-	static FRHITextureCreateDesc Create3D(
-		  TCHAR const*                DebugName
-		, FIntVector                  Size
-		, EPixelFormat                Format
-		, FClearValueBinding          ClearValue
-		, ETextureCreateFlags         Flags
-		, uint8                       NumMips      = 1
-		, uint32		              ExtData      = 0
-		, ERHIAccess                  InitialState = ERHIAccess::Unknown
-		, FResourceBulkDataInterface* BulkData     = nullptr
-		, FRHIGPUMask                 GPUMask      = FRHIGPUMask::All()
-		)
-	{
-		return Create3D(DebugName)
-			.SetExtent(FIntPoint(Size.X, Size.Y))
-			.SetDepth(Size.Z)
-			.SetFormat(Format)
-			.SetClearValue(ClearValue)
-			.SetFlags(Flags)
-			.SetNumMips(NumMips)
-			.SetExtData(ExtData)
-			.SetBulkData(BulkData)
-			.SetGPUMask(GPUMask)
-			.SetInitialState(InitialState == ERHIAccess::Unknown ? RHIGetDefaultResourceState(Flags, BulkData != nullptr) : InitialState);
-	}
-
 	static FRHITextureCreateDesc CreateCube(const TCHAR* DebugName, uint32 Size, EPixelFormat Format)
 	{
 		return CreateCube(DebugName)
@@ -1676,33 +1594,6 @@ struct FRHITextureCreateDesc : public FRHITextureDesc
 			.SetFormat(Format);
 	}
 
-	static FRHITextureCreateDesc CreateCube(
-		  TCHAR const*                DebugName
-		, uint32                      Size
-		, EPixelFormat                Format
-		, FClearValueBinding          ClearValue
-		, ETextureCreateFlags         Flags
-		, uint8                       NumMips      = 1
-		, uint8                       NumSamples   = 1
-		, uint32		              ExtData      = 0
-		, ERHIAccess                  InitialState = ERHIAccess::Unknown
-		, FResourceBulkDataInterface* BulkData     = nullptr
-		, FRHIGPUMask                 GPUMask      = FRHIGPUMask::All()
-		)
-	{
-		return CreateCube(DebugName)
-			.SetExtent(Size)
-			.SetFormat(Format)
-			.SetClearValue(ClearValue)
-			.SetFlags(Flags)
-			.SetNumMips(NumMips)
-			.SetNumSamples(NumSamples)
-			.SetExtData(ExtData)
-			.SetBulkData(BulkData)
-			.SetGPUMask(GPUMask)
-			.SetInitialState(InitialState == ERHIAccess::Unknown ? RHIGetDefaultResourceState(Flags, BulkData != nullptr) : InitialState);
-	}
-
 	static FRHITextureCreateDesc CreateCubeArray(const TCHAR* DebugName, uint32 Size, uint16 ArraySize, EPixelFormat Format)
 	{
 		return CreateCubeArray(DebugName)
@@ -1711,35 +1602,6 @@ struct FRHITextureCreateDesc : public FRHITextureDesc
 			.SetArraySize(ArraySize);
 	}
 
-	static FRHITextureCreateDesc CreateCubeArray(
-		  TCHAR const*                DebugName
-		, uint32                      Size
-		, EPixelFormat                Format
-		, FClearValueBinding          ClearValue
-		, ETextureCreateFlags         Flags
-		, uint16                      ArraySize
-		, uint8                       NumMips      = 1
-		, uint8                       NumSamples   = 1
-		, uint32		              ExtData      = 0
-		, ERHIAccess                  InitialState = ERHIAccess::Unknown
-		, FResourceBulkDataInterface* BulkData     = nullptr
-		, FRHIGPUMask                 GPUMask      = FRHIGPUMask::All()
-		)
-	{
-		return CreateCubeArray(DebugName)
-			.SetExtent(Size)
-			.SetFormat(Format)
-			.SetClearValue(ClearValue)
-			.SetFlags(Flags)
-			.SetArraySize(ArraySize)
-			.SetNumMips(NumMips)
-			.SetNumSamples(NumSamples)
-			.SetExtData(ExtData)
-			.SetBulkData(BulkData)
-			.SetGPUMask(GPUMask)
-			.SetInitialState(InitialState == ERHIAccess::Unknown ? RHIGetDefaultResourceState(Flags, BulkData != nullptr) : InitialState);
-	}
-
 	FRHITextureCreateDesc() = default;
 
 	// Constructor with minimal argument set. Name and dimension are always required.
@@ -1749,30 +1611,6 @@ struct FRHITextureCreateDesc : public FRHITextureDesc
 	{
 	}
 
-	// Constructor with full argument set.
-	FRHITextureCreateDesc(
-		  ETextureDimension           InDimension
-		, ETextureCreateFlags         InFlags
-		, EPixelFormat                InFormat
-		, FClearValueBinding          InClearValue
-		, FIntPoint                   InExtent
-		, uint16                      InDepth
-		, uint16                      InArraySize
-		, uint8                       InNumMips
-		, uint8                       InNumSamples
-		, uint32                      InExtData
-		, ERHIAccess                  InInitialState
-		, TCHAR const*                InDebugName
-		, FResourceBulkDataInterface* InBulkData    = nullptr 
-		, FRHIGPUMask                 InGPUMask     = FRHIGPUMask::All()
-		)
-		: FRHITextureDesc(InDimension, InFlags, InFormat, InClearValue, InExtent, InDepth, InArraySize, InNumMips, InNumSamples, InExtData)
-		, InitialState   (InInitialState)
-		, DebugName      (InDebugName)
-		, GPUMask        (InGPUMask)
-		, BulkData       (InBulkData)
-	{}
-
 	// Constructor for when you already have an FRHITextureDesc
 	FRHITextureCreateDesc(
 		  FRHITextureDesc const&      InDesc
diff --git a/Engine/Source/Runtime/VulkanRHI/Private/VulkanDevice.cpp b/Engine/Source/Runtime/VulkanRHI/Private/VulkanDevice.cpp
--- a/Engine/Source/Runtime/VulkanRHI/Private/VulkanDevice.cpp
+++ b/Engine/Source/Runtime/VulkanRHI/Private/VulkanDevice.cpp
@@ -1115,17 +1115,11 @@ void FVulkanDevice::InitGPU(int32 DeviceIndex)
 		FSamplerStateInitializerRHI Default(SF_Point);
 		DefaultSampler = ResourceCast(RHICreateSamplerState(Default).GetReference());
 
-		FRHITextureCreateDesc Desc = FRHITextureCreateDesc::Create2D(
-			TEXT("FVulkanDevice_DefaultImage"),
-			{ 1, 1 },
-			PF_B8G8R8A8,
-			FClearValueBinding::None,
-			TexCreate_RenderTargetable | TexCreate_ShaderResource,
-			1,
-			1,
-			0,
-			ERHIAccess::SRVMask
-		);
+		const FRHITextureCreateDesc Desc =
+			FRHITextureCreateDesc::Create2D(TEXT("FVulkanDevice_DefaultImage"), 1, 1, PF_B8G8R8A8)
+			.SetClearValue(FClearValueBinding::None)
+			.SetFlags(ETextureCreateFlags::RenderTargetable | ETextureCreateFlags::ShaderResource)
+			.SetInitialState(ERHIAccess::SRVMask);
 
 		DefaultTexture = new FVulkanTexture(*this, Desc, nullptr);
 	}
diff --git a/Engine/Source/Runtime/VulkanRHI/Private/VulkanRHI.cpp b/Engine/Source/Runtime/VulkanRHI/Private/VulkanRHI.cpp
--- a/Engine/Source/Runtime/VulkanRHI/Private/VulkanRHI.cpp
+++ b/Engine/Source/Runtime/VulkanRHI/Private/VulkanRHI.cpp
@@ -1427,64 +1427,38 @@ void FVulkanDynamicRHI::RHISubmitCommandsAndFlushGPU()
 
 FTexture2DRHIRef FVulkanDynamicRHI::RHICreateTexture2DFromResource(EPixelFormat Format, uint32 SizeX, uint32 SizeY, uint32 NumMips, uint32 NumSamples, VkImage Resource, ETextureCreateFlags Flags)
 {
-	FRHITextureCreateDesc Desc = FRHITextureCreateDesc::Create2D(
-		TEXT("VulkanTexture2DFromResource"),
-		{ (int32)SizeX, (int32)SizeY },
-		Format,
-		EnumHasAnyFlags(Flags, TexCreate_DepthStencilTargetable) ? FClearValueBinding::DepthZero : FClearValueBinding::Transparent,
-		Flags,
-		NumMips,
-		NumSamples
-	);
+	const FRHITextureCreateDesc Desc =
+		FRHITextureCreateDesc::Create2D(TEXT("VulkanTexture2DFromResource"), SizeX, SizeY, Format)
+		.SetClearValue(EnumHasAnyFlags(Flags, ETextureCreateFlags::DepthStencilTargetable) ? FClearValueBinding::DepthZero : FClearValueBinding::Transparent)
+		.SetFlags(Flags)
+		.SetNumMips(NumMips)
+		.SetNumSamples(NumSamples);
 
 	return new FVulkanTexture(*Device, Desc, Resource, false);
 }
 
 FTexture2DArrayRHIRef FVulkanDynamicRHI::RHICreateTexture2DArrayFromResource(EPixelFormat Format, uint32 SizeX, uint32 SizeY, uint32 ArraySize, uint32 NumMips, uint32 NumSamples, VkImage Resource, ETextureCreateFlags Flags)
 {
-	FRHITextureCreateDesc Desc = FRHITextureCreateDesc::Create2DArray(
-		TEXT("VulkanTextureArrayFromResource"),
-		{ (int32)SizeX, (int32)SizeY },
-		Format,
-		EnumHasAnyFlags(Flags, TexCreate_DepthStencilTargetable) ? FClearValueBinding::DepthZero : FClearValueBinding::Transparent,
-		Flags,
-		ArraySize,
-		NumMips,
-		NumSamples
-	);
+	const FRHITextureCreateDesc Desc =
+		FRHITextureCreateDesc::Create2DArray(TEXT("VulkanTextureArrayFromResource"), SizeX, SizeY, ArraySize, Format)
+		.SetClearValue(EnumHasAnyFlags(Flags, ETextureCreateFlags::DepthStencilTargetable) ? FClearValueBinding::DepthZero : FClearValueBinding::Transparent)
+		.SetFlags(Flags)
+		.SetNumMips(NumMips)
+		.SetNumSamples(NumSamples);
 
 	return new FVulkanTexture(*Device, Desc, Resource, false);
 }
 
 FTextureCubeRHIRef FVulkanDynamicRHI::RHICreateTextureCubeFromResource(EPixelFormat Format, uint32 Size, bool bArray, uint32 ArraySize, uint32 NumMips, VkImage Resource, ETextureCreateFlags Flags)
 {
-	FRHITextureCreateDesc Desc;
-
-	if (ArraySize > 1)
-	{
-		Desc = FRHITextureCreateDesc::CreateCubeArray(
-			TEXT("VulkanTextureCubeFromResource"),
-			Size,
-			Format,
-			EnumHasAnyFlags(Flags, TexCreate_DepthStencilTargetable) ? FClearValueBinding::DepthZero : FClearValueBinding::Transparent,
-			Flags,
-			ArraySize,
-			NumMips,
-			1
-		);
-	}
-	else
-	{
-		Desc = FRHITextureCreateDesc::CreateCube(
-			TEXT("VulkanTextureCubeFromResource"),
-			Size,
-			Format,
-			EnumHasAnyFlags(Flags, TexCreate_DepthStencilTargetable) ? FClearValueBinding::DepthZero : FClearValueBinding::Transparent,
-			Flags,
-			NumMips,
-			1
-		);
-	}
+	const FRHITextureCreateDesc Desc =
+		FRHITextureCreateDesc::Create(TEXT("VulkanTextureCubeFromResource"), ArraySize > 1 ? ETextureDimension::TextureCubeArray : ETextureDimension::TextureCube)
+		.SetExtent(Size)
+		.SetArraySize(ArraySize)
+		.SetFormat(Format)
+		.SetClearValue(EnumHasAnyFlags(Flags, ETextureCreateFlags::DepthStencilTargetable) ? FClearValueBinding::DepthZero : FClearValueBinding::Transparent)
+		.SetFlags(Flags)
+		.SetNumMips(NumMips);
 
 	return new FVulkanTexture(*Device, Desc, Resource, false);
 }
diff --git a/Engine/Source/Runtime/VulkanRHI/Private/VulkanSwapChain.cpp b/Engine/Source/Runtime/VulkanRHI/Private/VulkanSwapChain.cpp
--- a/Engine/Source/Runtime/VulkanRHI/Private/VulkanSwapChain.cpp
+++ b/Engine/Source/Runtime/VulkanRHI/Private/VulkanSwapChain.cpp
@@ -913,15 +913,12 @@ void FVulkanSwapChain::CreateQCOMDepthStencil(const FVulkanTexture& InSurface) c
 	const ETextureCreateFlags UEFlags = Desc.Flags;
 	check(UEFlags & TexCreate_DepthStencilTargetable);
 
-	FRHITextureCreateDesc CreateDesc = FRHITextureCreateDesc::Create2D(
-		TEXT("FVulkanSwapChainQCOM"),
-		{ (int32)Desc.Extent.Y, (int32)Desc.Extent.X }, // Yes, these are intentionally swapped.
-		Desc.Format,
-		FClearValueBinding::None,
-		UEFlags,
-		Desc.NumMips,
-		Desc.NumSamples
-	);
+	const FRHITextureCreateDesc CreateDesc =
+		FRHITextureCreateDesc::Create2D(TEXT("FVulkanSwapChainQCOM"), Desc.Extent.Y, Desc.Extent.X, Desc.Format) // Desc.Extent.X and Desc.Extent.Y are intentionally swapped.
+		.SetClearValue(FClearValueBinding::None)
+		.SetFlags(UEFlags)
+		.SetNumMips(Desc.NumMips)
+		.SetNumSamples(Desc.NumSamples);
 
 	QCOMDepthStencilSurface = new FVulkanTexture(Device, CreateDesc, nullptr);
 
diff --git a/Engine/Source/Runtime/VulkanRHI/Private/VulkanTexture.cpp b/Engine/Source/Runtime/VulkanRHI/Private/VulkanTexture.cpp
--- a/Engine/Source/Runtime/VulkanRHI/Private/VulkanTexture.cpp
+++ b/Engine/Source/Runtime/VulkanRHI/Private/VulkanTexture.cpp
@@ -1056,16 +1056,14 @@ FTexture2DRHIRef FVulkanDynamicRHI::AsyncReallocateTexture2D_RenderThread(FRHICo
 	}
 
 	FVulkanTexture* OldTexture = FVulkanTexture::Cast(OldTextureRHI);
+	const FRHITextureDesc& OldDesc = OldTexture->GetDesc();
 
-	FRHITextureCreateDesc Desc = FRHITextureCreateDesc::Create2D(
-		TEXT("AsyncReallocateTexture2D_RenderThread"),
-		{ NewSizeX, NewSizeY },
-		OldTexture->GetFormat(),
-		OldTexture->GetClearBinding(),
-		OldTexture->GetFlags(),
-		NewMipCount,
-		OldTexture->GetNumSamples()
-	);
+	const FRHITextureCreateDesc Desc =
+		FRHITextureCreateDesc::Create2D(TEXT("AsyncReallocateTexture2D_RenderThread"), NewSizeX, NewSizeY, OldDesc.Format)
+		.SetClearValue(OldDesc.ClearValue)
+		.SetFlags(OldDesc.Flags)
+		.SetNumMips(NewMipCount)
+		.SetNumSamples(OldDesc.NumSamples);
 
 	FVulkanTexture* NewTexture = new FVulkanTexture(*Device, Desc, nullptr);
 	FVulkanCommandListContext& Context = (FVulkanCommandListContext&)RHICmdList.GetContext().GetLowestLevelContext();
@@ -1078,16 +1076,15 @@ FTexture2DRHIRef FVulkanDynamicRHI::RHIAsyncReallocateTexture2D(FRHITexture2D* O
 {
 	LLM_SCOPE_VULKAN(ELLMTagVulkan::VulkanTextures);
 	FVulkanTexture* OldTexture = FVulkanTexture::Cast(OldTextureRHI);
+	const FRHITextureDesc& OldDesc = OldTexture->GetDesc();
+
+	const FRHITextureCreateDesc Desc =
+		FRHITextureCreateDesc::Create2D(TEXT("AsyncReallocateTexture2D_RenderThread"), NewSizeX, NewSizeY, OldDesc.Format)
+		.SetClearValue(OldDesc.ClearValue)
+		.SetFlags(OldDesc.Flags)
+		.SetNumMips(NewMipCount)
+		.SetNumSamples(OldDesc.NumSamples);
 
-	FRHITextureCreateDesc Desc = FRHITextureCreateDesc::Create2D(
-		TEXT("AsyncReallocateTexture2D_RenderThread"),
-		{ NewSizeX, NewSizeY },
-		OldTexture->GetFormat(),
-		OldTexture->GetClearBinding(),
-		OldTexture->GetFlags(),
-		NewMipCount,
-		OldTexture->GetNumSamples()
-	);
 	FVulkanTexture* NewTexture = new FVulkanTexture(*Device, Desc, nullptr);
 
 	DoAsyncReallocateTexture2D(Device->GetImmediateContext(), OldTexture, NewTexture, NewMipCount, NewSizeX, NewSizeY, RequestStatus);
diff --git a/Engine/Source/Runtime/VulkanRHI/Private/VulkanViewport.cpp b/Engine/Source/Runtime/VulkanRHI/Private/VulkanViewport.cpp
--- a/Engine/Source/Runtime/VulkanRHI/Private/VulkanViewport.cpp
+++ b/Engine/Source/Runtime/VulkanRHI/Private/VulkanViewport.cpp
@@ -15,7 +15,7 @@
 
 
 FVulkanBackBuffer::FVulkanBackBuffer(FVulkanDevice& Device, FVulkanViewport* InViewport, EPixelFormat Format, uint32 SizeX, uint32 SizeY, ETextureCreateFlags UEFlags)
-	: FVulkanTexture(Device, FRHITextureCreateDesc::Create2D(TEXT("FVulkanBackBuffer"), { (int32)SizeX, (int32)SizeY }, Format, FClearValueBinding::None, UEFlags), VK_NULL_HANDLE, false)
+	: FVulkanTexture(Device, FRHITextureCreateDesc::Create2D(TEXT("FVulkanBackBuffer"), SizeX, SizeY, Format).SetFlags(UEFlags), VK_NULL_HANDLE, false)
 	, Viewport(InViewport)
 {
 }
@@ -692,17 +692,11 @@ void FVulkanViewport::CreateSwapchain(FVulkanSwapChainRecreateInfo* RecreateInfo
 		uint32 BackBufferSizeX = RequiresRenderingBackBuffer() ? SizeX : 1;
 		uint32 BackBufferSizeY = RequiresRenderingBackBuffer() ? SizeY : 1;
 
-		FRHITextureCreateDesc Desc = FRHITextureCreateDesc::Create2D(
-			TEXT("RenderingBackBuffer"),
-			{ (int32)BackBufferSizeX, (int32)BackBufferSizeY },
-			PixelFormat,
-			FClearValueBinding::None,
-			TexCreate_RenderTargetable | TexCreate_ShaderResource | TexCreate_ResolveTargetable,
-			1,
-			1,
-			0,
-			ERHIAccess::Present
-		);
+		const FRHITextureCreateDesc Desc =
+			FRHITextureCreateDesc::Create2D(TEXT("RenderingBackBuffer"), BackBufferSizeX, BackBufferSizeY, PixelFormat)
+			.SetClearValue(FClearValueBinding::None)
+			.SetFlags(ETextureCreateFlags::RenderTargetable | ETextureCreateFlags::ShaderResource | ETextureCreateFlags::ResolveTargetable)
+			.SetInitialState(ERHIAccess::Present);
 
 		RenderingBackBuffer = new FVulkanTexture(*Device, Desc, nullptr);
 #if VULKAN_ENABLE_DRAW_MARKERS
diff --git a/Engine/Source/Runtime/Windows/D3D11RHI/Private/D3D11Texture.cpp b/Engine/Source/Runtime/Windows/D3D11RHI/Private/D3D11Texture.cpp
--- a/Engine/Source/Runtime/Windows/D3D11RHI/Private/D3D11Texture.cpp
+++ b/Engine/Source/Runtime/Windows/D3D11RHI/Private/D3D11Texture.cpp
@@ -997,17 +997,13 @@ FTextureRHIRef FD3D11DynamicRHI::RHIAsyncCreateTexture2D(uint32 SizeX, uint32 Si
 		SubresourceData[MipIndex].SysMemSlicePitch = MipSize;
 	}
 
-	FTextureRHIRef Texture = CreateD3D11Texture2D(
-		FRHITextureCreateDesc::Create2D(
-			  TEXT("RHIAsyncCreateTexture2D")
-			, { (int32)SizeX, (int32)SizeY }
-			, (EPixelFormat)Format
-			, FClearValueBinding::None
-			, Flags
-			, NumMips
-		)
-		, SubresourceData
-	);
+	const FRHITextureCreateDesc Desc =
+		FRHITextureCreateDesc::Create2D(TEXT("RHIAsyncCreateTexture2D"), SizeX, SizeY, (EPixelFormat)Format)
+		.SetClearValue(FClearValueBinding::None)
+		.SetFlags(Flags)
+		.SetNumMips(NumMips);
+
+	FTextureRHIRef Texture = CreateD3D11Texture2D(Desc, SubresourceData);
 
 	if (TempBufferSize != ZeroBufferSize)
 	{
@@ -1944,24 +1940,20 @@ FD3D11Texture* FD3D11DynamicRHI::CreateTextureFromResource(bool bTextureArray, b
 		check(IsValidRef(ShaderResourceView));
 	}
 
-	ETextureDimension Dimension = bTextureArray
+	const ETextureDimension Dimension = bTextureArray
 		? (bCubeTexture ? ETextureDimension::TextureCubeArray : ETextureDimension::Texture2DArray)
 		: (bCubeTexture ? ETextureDimension::TextureCube      : ETextureDimension::Texture2D     );
 
-	FRHITextureCreateDesc RHITextureDesc = FRHITextureCreateDesc(
-		Dimension,
-		TexCreateFlags,
-		(EPixelFormat)Format,
-		ClearValueBinding,
-		{ (int32)TextureDesc.Width, (int32)TextureDesc.Height },
-		1, // Depth
-		TextureDesc.ArraySize,
-		TextureDesc.MipLevels,
-		TextureDesc.SampleDesc.Count,
-		0, // ExtData
-		RHIGetDefaultResourceState(TexCreateFlags, false),
-		TEXT("FD3D11DynamicRHI::CreateTextureFromResource")
-	);
+	const FRHITextureCreateDesc RHITextureDesc =
+		FRHITextureCreateDesc::Create(TEXT("FD3D11DynamicRHI::CreateTextureFromResource"), Dimension)
+		.SetExtent(TextureDesc.Width, TextureDesc.Height)
+		.SetFormat((EPixelFormat)Format)
+		.SetClearValue(ClearValueBinding)
+		.SetArraySize(TextureDesc.ArraySize)
+		.SetFlags(TexCreateFlags)
+		.SetNumMips(TextureDesc.MipLevels)
+		.SetNumSamples(TextureDesc.SampleDesc.Count)
+		.SetInitialState(RHIGetDefaultResourceState(TexCreateFlags, false));
 
 	FD3D11Texture* Texture2D = new FD3D11Texture(
 		RHITextureDesc,
diff --git a/Engine/Source/Runtime/Windows/D3D11RHI/Private/D3D11Viewport.cpp b/Engine/Source/Runtime/Windows/D3D11RHI/Private/D3D11Viewport.cpp
--- a/Engine/Source/Runtime/Windows/D3D11RHI/Private/D3D11Viewport.cpp
+++ b/Engine/Source/Runtime/Windows/D3D11RHI/Private/D3D11Viewport.cpp
@@ -180,20 +180,10 @@ FD3D11Texture* FD3D11Viewport::GetSwapChainSurface(FD3D11DynamicRHI* D3DRHI, EPi
 	SRVDesc.Texture2D.MipLevels = 1;
 	VERIFYD3D11RESULT_EX(D3DRHI->GetDevice()->CreateShaderResourceView(BackBufferResource,&SRVDesc,BackBufferShaderResourceView.GetInitReference()), D3DRHI->GetDevice());
 
-	FRHITextureCreateDesc CreateDesc(
-		ETextureDimension::Texture2D,
-		ETextureCreateFlags::RenderTargetable,
-		PixelFormat,
-		FClearValueBinding(),
-		{ (int32)TextureDesc.Width, (int32)TextureDesc.Height },
-		1, // Depth
-		1, // ArraySize,
-		1, // NumMips,
-		1, // NumSamples
-		0, // ExtData
-		RHIGetDefaultResourceState(ETextureCreateFlags::RenderTargetable, false),
-		TEXT("FD3D11Viewport::GetSwapChainSurface")
-	);
+	const FRHITextureCreateDesc CreateDesc =
+		FRHITextureCreateDesc::Create2D(TEXT("FD3D11Viewport::GetSwapChainSurface"), TextureDesc.Width, TextureDesc.Height, PixelFormat)
+		.SetFlags(ETextureCreateFlags::RenderTargetable)
+		.SetInitialState(RHIGetDefaultResourceState(ETextureCreateFlags::RenderTargetable, false));
 
 	FD3D11Texture* NewTexture = new FD3D11Texture(
 		CreateDesc,
