Fix RHICopyTexture validation when copying to/from non-zero mips.

#rnx

diff --git a/Engine/Source/Runtime/RHI/Public/RHIValidationContext.h b/Engine/Source/Runtime/RHI/Public/RHIValidationContext.h
--- a/Engine/Source/Runtime/RHI/Public/RHIValidationContext.h
+++ b/Engine/Source/Runtime/RHI/Public/RHIValidationContext.h
@@ -944,7 +944,7 @@ public:
 		Tracker->Assert(SourceTexture->GetWholeResourceIdentity(), ERHIAccess::CopySrc);
 		Tracker->Assert(DestTexture->GetWholeResourceIdentity(), ERHIAccess::CopyDest);
 
-		FValidationRHIUtils::ValidateCopyTexture(SourceTexture, DestTexture, CopyInfo.Size, CopyInfo.SourcePosition, CopyInfo.DestPosition);
+		FValidationRHIUtils::ValidateCopyTexture(SourceTexture, DestTexture, CopyInfo);
 		RHIContext->RHICopyTexture(SourceTexture, DestTexture, CopyInfo);
 	}
 
diff --git a/Engine/Source/Runtime/RHI/Public/RHIValidationUtils.h b/Engine/Source/Runtime/RHI/Public/RHIValidationUtils.h
--- a/Engine/Source/Runtime/RHI/Public/RHIValidationUtils.h
+++ b/Engine/Source/Runtime/RHI/Public/RHIValidationUtils.h
@@ -7,12 +7,7 @@
 class FValidationRHIUtils
 {
 public:
-	static void ValidateCopyTexture(
-		FRHITexture*	SourceTexture,
-		FRHITexture*	DestTexture,
-		FIntVector			CopySize = FIntVector::ZeroValue,
-		const FIntVector&	SourcePosition = FIntVector::ZeroValue,
-		const FIntVector&	DestPosition = FIntVector::ZeroValue)
+	static void ValidateCopyTexture(FRHITexture* SourceTexture, FRHITexture* DestTexture, const FRHICopyTextureInfo& CopyInfo)
 	{
 		check(SourceTexture);
 		check(DestTexture);
@@ -27,8 +22,9 @@ public:
 
 		checkf(bValidCopyFormats, TEXT("Some RHIs do not support this format conversion by the GPU for transfer operations!"));
 
-		FIntVector SrcSize = SourceTexture->GetSizeXYZ();
-		FIntVector DestSize = DestTexture->GetSizeXYZ();
+		FIntVector SrcSize = SourceTexture->GetMipDimensions(CopyInfo.SourceMipIndex);
+		FIntVector DestSize = DestTexture->GetMipDimensions(CopyInfo.DestMipIndex);
+		FIntVector CopySize = CopyInfo.Size;
 		if (CopySize == FIntVector::ZeroValue)
 		{
 			CopySize = SrcSize;
@@ -36,16 +32,16 @@ public:
 
 		checkf(CopySize.X <= DestSize.X && CopySize.Y <= DestSize.Y, TEXT("Some RHIs can't perform scaling operations [%dx%d to %dx%d] during copies!"), SrcSize.X, SrcSize.Y, DestSize.X, DestSize.Y);
 
-		check(SourcePosition.X >= 0 && SourcePosition.Y >= 0 && SourcePosition.Z >= 0);
-		check(SourcePosition.X + CopySize.X <= SrcSize.X && SourcePosition.Y + CopySize.Y <= SrcSize.Y);
+		check(CopyInfo.SourcePosition.X >= 0 && CopyInfo.SourcePosition.Y >= 0 && CopyInfo.SourcePosition.Z >= 0);
+		check(CopyInfo.SourcePosition.X + CopySize.X <= SrcSize.X && CopyInfo.SourcePosition.Y + CopySize.Y <= SrcSize.Y);
 
-		check(DestPosition.X >= 0 && DestPosition.Y >= 0 && DestPosition.Z >= 0);
-		check(DestPosition.X + CopySize.X <= DestSize.X && DestPosition.Y + CopySize.Y <= DestSize.Y);
+		check(CopyInfo.DestPosition.X >= 0 && CopyInfo.DestPosition.Y >= 0 && CopyInfo.DestPosition.Z >= 0);
+		check(CopyInfo.DestPosition.X + CopySize.X <= DestSize.X && CopyInfo.DestPosition.Y + CopySize.Y <= DestSize.Y);
 
 		if (SourceTexture->GetTexture3D() && DestTexture->GetTexture3D())
 		{
-			check(SourcePosition.Z + CopySize.Z <= SrcSize.Z);
-			check(DestPosition.Z + CopySize.Z <= DestSize.Z);
+			check(CopyInfo.SourcePosition.Z + CopySize.Z <= SrcSize.Z);
+			check(CopyInfo.DestPosition.Z + CopySize.Z <= DestSize.Z);
 		}
 	}
 };
