Fix RHI validation errors caused by texture references not using the correct tracker info.

Texture references need to use the tracker info of the underlying texture object, not their own tracker info, as the local info is never filled in.

#rnx

diff --git a/Engine/Source/Runtime/RHI/Public/RHIResources.h b/Engine/Source/Runtime/RHI/Public/RHIResources.h
--- a/Engine/Source/Runtime/RHI/Public/RHIResources.h
+++ b/Engine/Source/Runtime/RHI/Public/RHIResources.h
@@ -1091,7 +1091,7 @@ public:
 	}
 
 #if ENABLE_RHI_VALIDATION
-	virtual RHIValidation::FResource* GetTrackerResource() = 0;
+	virtual RHIValidation::FResource* GetValidationTrackerResource() = 0;
 #endif
 
 protected:
@@ -1145,7 +1145,7 @@ public:
 	virtual uint32 GetParentGPUIndex() const { return 0; }
 
 #if ENABLE_RHI_VALIDATION
-	virtual RHIValidation::FResource* GetTrackerResource() final override
+	virtual RHIValidation::FResource* GetValidationTrackerResource() final override
 	{
 		return this;
 	}
@@ -1902,9 +1902,11 @@ public:
 	uint32 GetSize() const { check(GetDesc().IsTextureCube()); return GetDesc().Extent.X; }
 
 #if ENABLE_RHI_VALIDATION
-	virtual RHIValidation::FResource* GetTrackerResource() override
+	virtual RHIValidation::FResource* GetValidationTrackerResource() override
 	{
-		return RHIValidation::FTextureResource::GetTrackerResource();
+		// Use the method inherited from RHIValidation::FTextureResource, as that's already a virtual overridden
+		// by subclasses such as FRHITextureReference to return the correct storage for the tracker information.
+		return GetTrackerResource();
 	}
 #endif
 
@@ -1968,6 +1970,8 @@ public:
 	}
 
 #if ENABLE_RHI_VALIDATION
+	// Implement RHIValidation::FTextureResource::GetTrackerResource to use the tracker info
+	// for the referenced texture.
 	virtual RHIValidation::FResource* GetTrackerResource() final override
 	{
 		check(ReferencedTexture);
diff --git a/Engine/Source/Runtime/RHI/Public/RHIValidationCommon.h b/Engine/Source/Runtime/RHI/Public/RHIValidationCommon.h
--- a/Engine/Source/Runtime/RHI/Public/RHIValidationCommon.h
+++ b/Engine/Source/Runtime/RHI/Public/RHIValidationCommon.h
@@ -407,7 +407,7 @@ namespace RHIValidation
 
 		virtual ~FTextureResource() {}
 
-		FResource* GetTrackerResource() { return &PRIVATE_TrackerResource; }
+		virtual FResource* GetTrackerResource() { return &PRIVATE_TrackerResource; }
 
 		void InitBarrierTracking(FRHITextureCreateDesc const& CreateDesc);
 
diff --git a/Engine/Source/Runtime/RHI/Public/RHIValidationContext.h b/Engine/Source/Runtime/RHI/Public/RHIValidationContext.h
--- a/Engine/Source/Runtime/RHI/Public/RHIValidationContext.h
+++ b/Engine/Source/Runtime/RHI/Public/RHIValidationContext.h
@@ -423,7 +423,7 @@ public:
 			*Info.Resource->GetName().ToString(),
 			*GetRHIAccessName(Info.Access));
 
-		Tracker->SetTrackedAccess(Info.Resource->GetTrackerResource(), Info.Access);
+		Tracker->SetTrackedAccess(Info.Resource->GetValidationTrackerResource(), Info.Access);
 
 		RHIContext->SetTrackedAccess(Info);
 	}
