Removed default RHICopyTexture implementation that forwarded to RHICopyToResolveTarget.

diff --git a/Engine/Source/Runtime/RHI/Public/RHIContext.h b/Engine/Source/Runtime/RHI/Public/RHIContext.h
--- a/Engine/Source/Runtime/RHI/Public/RHIContext.h
+++ b/Engine/Source/Runtime/RHI/Public/RHIContext.h
@@ -737,34 +737,7 @@ public:
 	{
 	}
 
-	virtual void RHICopyTexture(FRHITexture* SourceTexture, FRHITexture* DestTexture, const FRHICopyTextureInfo& CopyInfo)
-	{
-		const bool bIsCube = SourceTexture->GetTextureCube() != nullptr;
-		const bool bAllCubeFaces = bIsCube && (CopyInfo.NumSlices % 6) == 0;
-		const int32 NumArraySlices = bAllCubeFaces ? CopyInfo.NumSlices / 6 : CopyInfo.NumSlices;
-		const int32 NumFaces = bAllCubeFaces ? 6 : 1;
-		for (int32 ArrayIndex = 0; ArrayIndex < NumArraySlices; ++ArrayIndex)
-		{
-			int32 SourceArrayIndex = CopyInfo.SourceSliceIndex + ArrayIndex;
-			int32 DestArrayIndex = CopyInfo.DestSliceIndex + ArrayIndex;
-			for (int32 FaceIndex = 0; FaceIndex < NumFaces; ++FaceIndex)
-			{
-				FResolveParams ResolveParams(FResolveRect(0, 0, 0, 0),
-					bIsCube ? (ECubeFace)FaceIndex : CubeFace_PosX,
-					CopyInfo.SourceMipIndex,
-					SourceArrayIndex,
-					DestArrayIndex,
-					FResolveRect(0, 0, 0, 0)
-				);
-				if (CopyInfo.Size != FIntVector::ZeroValue)
-				{
-					ResolveParams.Rect = FResolveRect(CopyInfo.SourcePosition.X, CopyInfo.SourcePosition.Y, CopyInfo.SourcePosition.X + CopyInfo.Size.X, CopyInfo.SourcePosition.Y + CopyInfo.Size.Y);
-					ResolveParams.DestRect = FResolveRect(CopyInfo.DestPosition.X, CopyInfo.DestPosition.Y, CopyInfo.DestPosition.X + CopyInfo.Size.X, CopyInfo.DestPosition.Y + CopyInfo.Size.Y);
-				}
-				RHICopyToResolveTarget(SourceTexture, DestTexture, ResolveParams);
-			}
-		}
-	}
+	virtual void RHICopyTexture(FRHITexture* SourceTexture, FRHITexture* DestTexture, const FRHICopyTextureInfo& CopyInfo) = 0;
 
 	virtual void RHICopyBufferRegion(FRHIBuffer* DestBuffer, uint64 DstOffset, FRHIBuffer* SourceBuffer, uint64 SrcOffset, uint64 NumBytes)
 	{
