Fix for GPU readback copy not using resolve rect.

diff --git a/Engine/Source/Runtime/RHI/Private/RHIGPUReadback.cpp b/Engine/Source/Runtime/RHI/Private/RHIGPUReadback.cpp
--- a/Engine/Source/Runtime/RHI/Private/RHIGPUReadback.cpp
+++ b/Engine/Source/Runtime/RHI/Private/RHIGPUReadback.cpp
@@ -119,15 +119,10 @@ FRHIGPUTextureReadback::FRHIGPUTextureReadback(FName RequestName) : FRHIGPUMemor
 
 void FRHIGPUTextureReadback::EnqueueCopyRDG(FRHICommandList& RHICmdList, FRHITexture* SourceTexture, FResolveRect Rect)
 {
-	EnqueueCopyInternal(RHICmdList, SourceTexture, FResolveParams(Rect));
+	EnqueueCopy(RHICmdList, SourceTexture, Rect);
 }
 
 void FRHIGPUTextureReadback::EnqueueCopy(FRHICommandList& RHICmdList, FRHITexture* SourceTexture, FResolveRect Rect)
-{
-	EnqueueCopyInternal(RHICmdList, SourceTexture, FResolveParams(Rect));
-}
-
-void FRHIGPUTextureReadback::EnqueueCopyInternal(FRHICommandList& RHICmdList, FRHITexture* SourceTexture, FResolveParams ResolveParams)
 {
 	Fence->Clear();
 
@@ -150,8 +145,16 @@ void FRHIGPUTextureReadback::EnqueueCopyInternal(FRHICommandList& RHICmdList, FR
 			DestinationStagingTexture = RHICreateTexture(Desc);
 		}
 
-		// Transfer memory GPU -> CPU
-		RHICmdList.CopyTexture(SourceTexture, DestinationStagingTexture, {});
+		FRHICopyTextureInfo CopyInfo;
+
+		if (Rect.IsValid())
+		{
+			CopyInfo.SourcePosition = FIntVector(Rect.X1, Rect.Y1, 1);
+			CopyInfo.DestPosition = CopyInfo.SourcePosition;
+			CopyInfo.Size = FIntVector(Rect.X2 - Rect.X1, Rect.Y2 - Rect.Y1, 1);
+		}
+
+		RHICmdList.CopyTexture(SourceTexture, DestinationStagingTexture, CopyInfo);
 
 		RHICmdList.WriteGPUFence(Fence);
 
diff --git a/Engine/Source/Runtime/RHI/Public/RHIGPUReadback.h b/Engine/Source/Runtime/RHI/Public/RHIGPUReadback.h
--- a/Engine/Source/Runtime/RHI/Public/RHIGPUReadback.h
+++ b/Engine/Source/Runtime/RHI/Public/RHIGPUReadback.h
@@ -124,8 +124,5 @@ public:
 	UE_DEPRECATED(5.0, "Use FRHIGPUTextureReadback::Lock( int32& OutRowPitchInPixels) instead.")
 	void LockTexture(FRHICommandListImmediate& RHICmdList, void*& OutBufferPtr, int32& OutRowPitchInPixels);
 
-private:
-	void EnqueueCopyInternal(FRHICommandList& RHICmdList, FRHITexture* SourceTexture, FResolveParams Params);
-
 	FTextureRHIRef DestinationStagingTexture;
 };
