Fix double transition when CopyTextureWithTransitions is called with the same source and destination.

#rnx

diff --git a/Engine/Source/Runtime/RHI/Public/RHIUtilities.h b/Engine/Source/Runtime/RHI/Public/RHIUtilities.h
--- a/Engine/Source/Runtime/RHI/Public/RHIUtilities.h
+++ b/Engine/Source/Runtime/RHI/Public/RHIUtilities.h
@@ -807,8 +807,7 @@ inline void CopyTextureWithTransitions(FRHICommandListImmediate& RHICmdList, FRH
 	if (SrcTexture == DstTexture)
 	{
 		RHICmdList.Transition({
-			FRHITransitionInfo(SrcTexture, ERHIAccess::Unknown, ERHIAccess::SRVMask),
-			FRHITransitionInfo(DstTexture, ERHIAccess::Unknown, ERHIAccess::SRVMask)
+			FRHITransitionInfo(SrcTexture, ERHIAccess::Unknown, ERHIAccess::SRVMask)
 		});
 		return;
 	}
