Replaced cases where render target resolve path was being used as a simple copy with an explicit copy. Added checks to render target pass info to catch future cases in the constructor.

diff --git a/Engine/Plugins/Compositing/OpenCVLensDistortion/Source/OpenCVLensDistortion/Private/LensDistortionDisplacementMapRendering.cpp b/Engine/Plugins/Compositing/OpenCVLensDistortion/Source/OpenCVLensDistortion/Private/LensDistortionDisplacementMapRendering.cpp
--- a/Engine/Plugins/Compositing/OpenCVLensDistortion/Source/OpenCVLensDistortion/Private/LensDistortionDisplacementMapRendering.cpp
+++ b/Engine/Plugins/Compositing/OpenCVLensDistortion/Source/OpenCVLensDistortion/Private/LensDistortionDisplacementMapRendering.cpp
@@ -119,7 +119,7 @@ static void DrawUVDisplacementToRenderTargetFromPreComputedDisplacementMap_Rende
 #endif
 
 	// Set render target.
-	FRHIRenderPassInfo RPInfo(OutTextureRenderTargetResource->GetRenderTargetTexture(), ERenderTargetActions::Clear_Store, OutTextureRenderTargetResource->TextureRHI);
+	FRHIRenderPassInfo RPInfo(OutTextureRenderTargetResource->GetRenderTargetTexture(), ERenderTargetActions::Clear_Store);
 	RHICmdList.BeginRenderPass(RPInfo, TEXT("DrawUVDisplacementFromPrecomputedDisplacementMap"));
 	{
 
@@ -153,6 +153,7 @@ static void DrawUVDisplacementToRenderTargetFromPreComputedDisplacementMap_Rende
 		RHICmdList.DrawPrimitive(0, PrimitiveCount, 1);
 	}
 	RHICmdList.EndRenderPass();
+	CopyTextureWithTransitions(RHICmdList, OutTextureRenderTargetResource->GetRenderTargetTexture(), OutTextureRenderTargetResource->TextureRHI, {});
 #endif
 }
 
diff --git a/Engine/Source/Runtime/MovieSceneCapture/Private/FrameGrabber.cpp b/Engine/Source/Runtime/MovieSceneCapture/Private/FrameGrabber.cpp
--- a/Engine/Source/Runtime/MovieSceneCapture/Private/FrameGrabber.cpp
+++ b/Engine/Source/Runtime/MovieSceneCapture/Private/FrameGrabber.cpp
@@ -120,7 +120,7 @@ void FViewportSurfaceReader::ResolveRenderTarget(FViewportSurfaceReader* RenderT
 
 		FRHITexture* DestRenderTarget = ResampleTexturePooledRenderTarget->GetRHI();
 
-		FRHIRenderPassInfo RPInfo(DestRenderTarget, ERenderTargetActions::Load_Store, ReadbackTexture);
+		FRHIRenderPassInfo RPInfo(DestRenderTarget, ERenderTargetActions::Load_Store);
 		RHICmdList.BeginRenderPass(RPInfo, TEXT("FrameGrabberResolveRenderTarget"));
 		{
 			RHICmdList.SetViewport(0, 0, 0.0f, TargetSize.X, TargetSize.Y, 1.0f);
@@ -174,6 +174,7 @@ void FViewportSurfaceReader::ResolveRenderTarget(FViewportSurfaceReader* RenderT
 				EDRF_Default);
 		}
 		RHICmdList.EndRenderPass();
+		CopyTextureWithTransitions(RHICmdList, DestRenderTarget, ReadbackTexture, {});
 
 		if (RenderToReadback)
 		{
diff --git a/Engine/Source/Runtime/RHI/Public/RHIResources.h b/Engine/Source/Runtime/RHI/Public/RHIResources.h
--- a/Engine/Source/Runtime/RHI/Public/RHIResources.h
+++ b/Engine/Source/Runtime/RHI/Public/RHIResources.h
@@ -3758,6 +3758,7 @@ PRAGMA_ENABLE_DEPRECATION_WARNINGS
 	// Color, no depth, optional resolve, optional mip, optional array slice
 	explicit FRHIRenderPassInfo(FRHITexture* ColorRT, ERenderTargetActions ColorAction, FRHITexture* ResolveRT = nullptr, uint8 InMipIndex = 0, int32 InArraySlice = -1)
 	{
+		check(!ResolveRT || ResolveRT->IsMultisampled());
 		check(ColorRT);
 		ColorRenderTargets[0].RenderTarget = ColorRT;
 		ColorRenderTargets[0].ResolveTarget = ResolveRT;
@@ -3828,6 +3829,7 @@ PRAGMA_ENABLE_DEPRECATION_WARNINGS
 		check(NumColorRTs > 0);
 		for (int32 Index = 0; Index < NumColorRTs; ++Index)
 		{
+			check(!ResolveRTs[Index] || ResolveRTs[Index]->IsMultisampled());
 			check(ColorRTs[Index]);
 			ColorRenderTargets[Index].RenderTarget = ColorRTs[Index];
 			ColorRenderTargets[Index].ResolveTarget = ResolveRTs[Index];
@@ -3835,6 +3837,7 @@ PRAGMA_ENABLE_DEPRECATION_WARNINGS
 			ColorRenderTargets[Index].MipIndex = 0;
 			ColorRenderTargets[Index].Action = ColorAction;
 		}
+		check(!ResolveDepthRT || ResolveDepthRT->IsMultisampled());
 		check(DepthRT);
 		DepthStencilRenderTarget.DepthStencilTarget = DepthRT;
 		DepthStencilRenderTarget.ResolveTarget = ResolveDepthRT;
@@ -3845,6 +3848,7 @@ PRAGMA_ENABLE_DEPRECATION_WARNINGS
 	// Depth, no color
 	explicit FRHIRenderPassInfo(FRHITexture* DepthRT, EDepthStencilTargetActions DepthActions, FRHITexture* ResolveDepthRT = nullptr, FExclusiveDepthStencil InEDS = FExclusiveDepthStencil::DepthWrite_StencilWrite)
 	{
+		check(!ResolveDepthRT || ResolveDepthRT->IsMultisampled());
 		check(DepthRT);
 		DepthStencilRenderTarget.DepthStencilTarget = DepthRT;
 		DepthStencilRenderTarget.ResolveTarget = ResolveDepthRT;
@@ -3856,6 +3860,7 @@ PRAGMA_ENABLE_DEPRECATION_WARNINGS
 	explicit FRHIRenderPassInfo(FRHITexture* DepthRT, uint32 InNumOcclusionQueries, EDepthStencilTargetActions DepthActions, FRHITexture* ResolveDepthRT = nullptr, FExclusiveDepthStencil InEDS = FExclusiveDepthStencil::DepthWrite_StencilWrite)
 		: NumOcclusionQueries(InNumOcclusionQueries)
 	{
+		check(!ResolveDepthRT || ResolveDepthRT->IsMultisampled());
 		check(DepthRT);
 		DepthStencilRenderTarget.DepthStencilTarget = DepthRT;
 		DepthStencilRenderTarget.ResolveTarget = ResolveDepthRT;
@@ -3884,6 +3889,8 @@ PRAGMA_ENABLE_DEPRECATION_WARNINGS
 	explicit FRHIRenderPassInfo(FRHITexture* ColorRT, ERenderTargetActions ColorAction, FRHITexture* ResolveColorRT,
 		FRHITexture* DepthRT, EDepthStencilTargetActions DepthActions, FRHITexture* ResolveDepthRT, FExclusiveDepthStencil InEDS = FExclusiveDepthStencil::DepthWrite_StencilWrite)
 	{
+		check(!ResolveColorRT || ResolveColorRT->IsMultisampled());
+		check(!ResolveDepthRT || ResolveDepthRT->IsMultisampled());
 		check(ColorRT);
 		ColorRenderTargets[0].RenderTarget = ColorRT;
 		ColorRenderTargets[0].ResolveTarget = ResolveColorRT;
@@ -3904,6 +3911,8 @@ PRAGMA_ENABLE_DEPRECATION_WARNINGS
 		FRHITexture* InShadingRateTexture, EVRSRateCombiner InShadingRateTextureCombiner,
 		FExclusiveDepthStencil InEDS = FExclusiveDepthStencil::DepthWrite_StencilWrite)
 	{
+		check(!ResolveColorRT || ResolveColorRT->IsMultisampled());
+		check(!ResolveDepthRT || ResolveDepthRT->IsMultisampled());
 		check(ColorRT);
 		ColorRenderTargets[0].RenderTarget = ColorRT;
 		ColorRenderTargets[0].ResolveTarget = ResolveColorRT;
diff --git a/Engine/Source/Runtime/RHI/Public/RHIValidationContext.h b/Engine/Source/Runtime/RHI/Public/RHIValidationContext.h
--- a/Engine/Source/Runtime/RHI/Public/RHIValidationContext.h
+++ b/Engine/Source/Runtime/RHI/Public/RHIValidationContext.h
@@ -849,6 +849,7 @@ public:
 			FRHIRenderPassInfo::FColorEntry& RTV = State.RenderPassInfo.ColorRenderTargets[RTVIndex];
 			if (RTV.RenderTarget == nullptr)
 			{
+				checkf(RTV.ResolveTarget == nullptr, TEXT("Render target is null, but resolve target is not."));
 				continue;
 			}
 
@@ -866,6 +867,17 @@ public:
 			}
 
 			Tracker->Assert(RTV.RenderTarget->GetViewIdentity(RTV.MipIndex, 1, ArraySlice, NumArraySlices, 0, 0), ERHIAccess::RTV);
+
+			if (RTV.ResolveTarget)
+			{
+				const FRHITextureDesc& RenderTargetDesc = RTV.RenderTarget->GetDesc();
+				const FRHITextureDesc& ResolveTargetDesc = RTV.ResolveTarget->GetDesc();
+
+				checkf(RenderTargetDesc.Extent == ResolveTargetDesc.Extent, TEXT("Render target extent must match resolve target extent."));
+				checkf(RenderTargetDesc.Format == ResolveTargetDesc.Format, TEXT("Render target format must match resolve target format."));
+
+				Tracker->Assert(RTV.ResolveTarget->GetViewIdentity(RTV.MipIndex, 1, ArraySlice, NumArraySlices, 0, 0), ERHIAccess::ResolveSrc);
+			}
 		}
 
 		FRHIRenderPassInfo::FDepthStencilEntry& DSV = State.RenderPassInfo.DepthStencilRenderTarget;
@@ -878,6 +890,15 @@ public:
 			ViewDimensions = MipDimensions;
 		}
 
+		if (DSV.ResolveTarget)
+		{
+			const FRHITextureDesc& DepthStencilTargetDesc = DSV.DepthStencilTarget->GetDesc();
+			const FRHITextureDesc& ResolveTargetDesc = DSV.ResolveTarget->GetDesc();
+
+			checkf(DepthStencilTargetDesc.Extent == ResolveTargetDesc.Extent, TEXT("Depth stencil target extent must match resolve target extent."));
+			checkf(DepthStencilTargetDesc.IsTexture2D() && ResolveTargetDesc.IsTexture2D(), TEXT("Only 2D depth stencil resolves are supported."));
+		}
+
 		// @todo: additional checks for matching array slice counts on RTVs/DSVs
 
 		// assert depth is in the correct mode
@@ -889,6 +910,11 @@ public:
 
 			checkf(DSV.DepthStencilTarget, TEXT("Depth read/write is enabled but no depth stencil texture is bound."));
 			Tracker->Assert(DSV.DepthStencilTarget->GetViewIdentity(0, 0, 0, 0, uint32(RHIValidation::EResourcePlane::Common), 1), DepthAccess);
+
+			if (DSV.ResolveTarget)
+			{
+				Tracker->Assert(DSV.ResolveTarget->GetViewIdentity(0, 0, 0, 0, uint32(RHIValidation::EResourcePlane::Common), 1), ERHIAccess::ResolveSrc);
+			}
 		}
 
 		// assert stencil is in the correct mode
@@ -905,6 +931,11 @@ public:
 			if (bIsStencilFormat)
 			{
 				Tracker->Assert(DSV.DepthStencilTarget->GetViewIdentity(0, 0, 0, 0, uint32(RHIValidation::EResourcePlane::Stencil), 1), StencilAccess);
+
+				if (DSV.ResolveTarget)
+				{
+					Tracker->Assert(DSV.ResolveTarget->GetViewIdentity(0, 0, 0, 0, uint32(RHIValidation::EResourcePlane::Stencil), 1), ERHIAccess::ResolveSrc);
+				}
 			}
 		}
 
