Fixed dereference of deleted string pointer.

diff --git a/Engine/Source/Runtime/RHI/Private/RHIValidation.cpp b/Engine/Source/Runtime/RHI/Private/RHIValidation.cpp
--- a/Engine/Source/Runtime/RHI/Private/RHIValidation.cpp
+++ b/Engine/Source/Runtime/RHI/Private/RHIValidation.cpp
@@ -1767,7 +1767,7 @@ void FValidationTransientResourceAllocator::Release(FRHICommandListImmediate& RH
 			{
 				const FAllocatedResourceData& ResourceData = KeyValue.Value;
 
-				ErrorMessage += FString::Printf(TEXT("         %s (%s)\n"), ResourceData.DebugName, ResourceData.ResourceType == FAllocatedResourceData::EType::Texture ? TEXT("Texture") : TEXT("Buffer"));
+				ErrorMessage += FString::Printf(TEXT("         %s (%s)\n"), *ResourceData.DebugName, ResourceData.ResourceType == FAllocatedResourceData::EType::Texture ? TEXT("Texture") : TEXT("Buffer"));
 			}
 			ErrorMessage += FString::Printf(TRANSIENT_RESOURCE_LOG_SUFFIX);
 			FValidationRHI::ReportValidationFailure(*ErrorMessage);
@@ -1801,13 +1801,13 @@ void FValidationTransientResourceAllocator::InitBarrierTracking(const FAllocated
 				ArraySize *= 6;
 			}
 
-			Texture->InitBarrierTracking(ResourceData.Texture.NumMips, ArraySize, ResourceData.Texture.Format, ResourceData.Texture.Flags, ERHIAccess::Discard, ResourceData.DebugName);
+			Texture->InitBarrierTracking(ResourceData.Texture.NumMips, ArraySize, ResourceData.Texture.Format, ResourceData.Texture.Flags, ERHIAccess::Discard, *ResourceData.DebugName);
 		}
 		break;
 		case FAllocatedResourceData::EType::Buffer:
 		{
 			FRHIBuffer* Buffer = static_cast<FRHIBuffer*>(Resource);
-			Buffer->InitBarrierTracking(ERHIAccess::Discard, ResourceData.DebugName);
+			Buffer->InitBarrierTracking(ERHIAccess::Discard, *ResourceData.DebugName);
 		}
 		break;
 		}
diff --git a/Engine/Source/Runtime/RHI/Public/RHIValidationTransientResourceAllocator.h b/Engine/Source/Runtime/RHI/Public/RHIValidationTransientResourceAllocator.h
--- a/Engine/Source/Runtime/RHI/Public/RHIValidationTransientResourceAllocator.h
+++ b/Engine/Source/Runtime/RHI/Public/RHIValidationTransientResourceAllocator.h
@@ -36,7 +36,7 @@ private:
 			Buffer,
 		};
 
-		const TCHAR* DebugName = nullptr;
+		FString DebugName;
 		EType ResourceType = EType::Texture;
 
 		struct FTexture
