Fixed CIS build error.

diff --git a/Engine/Source/Runtime/RHI/Private/RHIValidation.cpp b/Engine/Source/Runtime/RHI/Private/RHIValidation.cpp
--- a/Engine/Source/Runtime/RHI/Private/RHIValidation.cpp
+++ b/Engine/Source/Runtime/RHI/Private/RHIValidation.cpp
@@ -1763,7 +1763,7 @@ void FValidationTransientResourceAllocator::Release(FRHICommandListImmediate& RH
 				TEXT("Resources with Allocated Memory:\n"),
 				AllocatedResourceMap.Num());
 
-			for (const auto KeyValue : AllocatedResourceMap)
+			for (const auto& KeyValue : AllocatedResourceMap)
 			{
 				const FAllocatedResourceData& ResourceData = KeyValue.Value;
 
