Reworking shader platform settings on Windows to be per-RHI. This removes the ill-fated "min/max feature level" configs that were just causing confusion.
On Windows, "TargetedRHIs" is now split into "D3D11TargetedShaderFormats", "D3D12TargetedShaderFormats" and "VulkanTargetedShaderFormats". "TargetedRHIs" is still parsed for backwards compatibility.
Using this, projects can now be more easily configured for D3D12-only or even Vulkan-only.
Updated FShaderFormatsPropertyDetails to use FName instead of FString for shader platforms. Also added a filtering method for mixed RHI platforms like Windows.

diff --git a/Engine/Config/BaseEngine.ini b/Engine/Config/BaseEngine.ini
--- a/Engine/Config/BaseEngine.ini
+++ b/Engine/Config/BaseEngine.ini
@@ -2748,13 +2748,8 @@ DefaultGravityZ=-980.0
 bEnable2DPhysics=false
 
 [/Script/WindowsTargetPlatform.WindowsTargetSettings]
-+TargetedRHIs=PCD3D_SM5
-D3D12MinimumFeatureLevel=ES3_1
-D3D12MaximumFeatureLevel=SM6
-D3D11MinimumFeatureLevel=ES3_1
-D3D11MaximumFeatureLevel=SM5
-VulkanMinimumFeatureLevel=ES3_1
-VulkanMaximumFeatureLevel=SM5
++D3D12TargetedShaderFormats=PCD3D_SM5
++D3D11TargetedShaderFormats=PCD3D_SM5
 bEnableDistanceFields=true
 bEnableRayTracing=true
 bTarget32Bit=false
diff --git a/Engine/Source/Developer/LauncherServices/Private/Launcher/LauncherWorker.cpp b/Engine/Source/Developer/LauncherServices/Private/Launcher/LauncherWorker.cpp
--- a/Engine/Source/Developer/LauncherServices/Private/Launcher/LauncherWorker.cpp
+++ b/Engine/Source/Developer/LauncherServices/Private/Launcher/LauncherWorker.cpp
@@ -297,7 +297,15 @@ static void AddDeviceToLaunchCommand(const FString& DeviceId, TSharedPtr<ITarget
 			FConfigCacheIni::LoadLocalIniFile(WindowsEngineSettings, TEXT("Engine"), true, TEXT("Windows"));
 
 			bCheckTargetedRHIs = true;
-			WindowsEngineSettings.GetArray(TEXT("/Script/WindowsTargetPlatform.WindowsTargetSettings"), TEXT("TargetedRHIs"), TargetedShaderFormats);
+			WindowsEngineSettings.GetArray(TEXT("/Script/WindowsTargetPlatform.WindowsTargetSettings"), TEXT("VulkanTargetedShaderFormats"), TargetedShaderFormats);
+
+			TArray<FString> OldConfigShaderFormats;
+			WindowsEngineSettings.GetArray(TEXT("/Script/WindowsTargetPlatform.WindowsTargetSettings"), TEXT("TargetedRHIs"), OldConfigShaderFormats);
+
+			for (const FString& OldConfigShaderFormat : OldConfigShaderFormats)
+			{
+				TargetedShaderFormats.AddUnique(OldConfigShaderFormat);
+			}
 		}
 		else if (Platform.StartsWith(TEXT("Linux")))
 		{
diff --git a/Engine/Source/Developer/SharedSettingsWidgets/Private/ShaderFormatsPropertyDetails.cpp b/Engine/Source/Developer/SharedSettingsWidgets/Private/ShaderFormatsPropertyDetails.cpp
--- a/Engine/Source/Developer/SharedSettingsWidgets/Private/ShaderFormatsPropertyDetails.cpp
+++ b/Engine/Source/Developer/SharedSettingsWidgets/Private/ShaderFormatsPropertyDetails.cpp
@@ -21,66 +21,7 @@
 
 #define LOCTEXT_NAMESPACE "ShaderFormatsPropertyDetails"
 
-FText FShaderFormatsPropertyDetails::GetFriendlyNameFromRHINameMac(const FString& InRHIName)
-{
-	FText FriendlyRHIName = FText::FromString(InRHIName);
-	
-	FName RHIName(*InRHIName, FNAME_Find);
-	EShaderPlatform Platform = ShaderFormatToLegacyShaderPlatform(RHIName);
-	switch(Platform)
-	{
-		case SP_PCD3D_SM5:
-			FriendlyRHIName = LOCTEXT("D3DSM5", "Direct3D 11+ (SM5)");
-			break;
-		case SP_PCD3D_ES3_1:
-			FriendlyRHIName = LOCTEXT("D3DES31", "Direct3D (ES3.1, Mobile Preview)");
-			break;
-		case SP_OPENGL_PCES3_1:
-			FriendlyRHIName = LOCTEXT("OpenGLES31PC", "OpenGL (ES3.1, Mobile Preview)");
-			break;
-		case SP_OPENGL_ES3_1_ANDROID:
-			FriendlyRHIName = LOCTEXT("OpenGLES31", "OpenGLES 3.1 (Mobile)");
-			break;
-		case SP_METAL:
-			FriendlyRHIName = LOCTEXT("Metal", "iOS Metal Mobile Renderer (ES3.1, Metal 1.1+, iOS 9.0 or later)");
-			break;
-		case SP_METAL_MRT:
-			FriendlyRHIName = LOCTEXT("MetalMRT", "iOS Metal Desktop Renderer (SM5, Metal 1.2+, iOS 10.0 or later)");
-			break;
-		case SP_METAL_TVOS:
-			FriendlyRHIName = LOCTEXT("MetalTV", "tvOS Metal Mobile Renderer (ES3.1, Metal 1.1+, tvOS 9.0 or later)");
-			break;
-		case SP_METAL_MRT_TVOS:
-			FriendlyRHIName = LOCTEXT("MetalMRTTV", "tvOS Metal Desktop Renderer (SM5, Metal 1.2+, tvOS 10.0 or later)");
-			break;
-		case SP_METAL_SM5:
-			FriendlyRHIName = LOCTEXT("MetalSM5", "Mac Metal Desktop Renderer (SM5, Metal 2.1+, macOS Catalina 10.15.7 or later)");
-			break;
-		case SP_METAL_MACES3_1:
-			FriendlyRHIName = LOCTEXT("MetalES3.1", "Mac Metal High-End Mobile Preview (ES3.1)");
-			break;
-		case SP_METAL_MRT_MAC:
-			FriendlyRHIName = LOCTEXT("MetalMRTMac", "Mac Metal iOS/tvOS Desktop Renderer Preview (SM5)");
-			break;
-		case SP_VULKAN_SM5:
-		case SP_VULKAN_SM5_ANDROID:
-			FriendlyRHIName = LOCTEXT("VulkanSM5", "Vulkan (SM5)");
-			break;
-		case SP_VULKAN_PCES3_1:
-		case SP_VULKAN_ES3_1_ANDROID:
-			FriendlyRHIName = LOCTEXT("VulkanES31", "Vulkan (ES 3.1)");
-			break;
-		case SP_D3D_ES3_1_HOLOLENS:
-			FriendlyRHIName = LOCTEXT("D3DES31HL", "Direct3D (ES3.1, Hololens)");
-			break;
-		default:
-			break;
-	}
-	
-	return FriendlyRHIName;
-}
-
-FShaderFormatsPropertyDetails::FShaderFormatsPropertyDetails(IDetailLayoutBuilder* InDetailBuilder, FString InProperty, FString InTitle)
+FShaderFormatsPropertyDetails::FShaderFormatsPropertyDetails(IDetailLayoutBuilder* InDetailBuilder, const TCHAR* InProperty, const TCHAR* InTitle)
 : DetailBuilder(InDetailBuilder)
 , Property(InProperty)
 , Title(InTitle)
@@ -94,7 +35,11 @@ void FShaderFormatsPropertyDetails::SetOnUpdateShaderWarning(FSimpleDelegate con
 	ShaderFormatsPropertyHandle->SetOnPropertyValueChanged(Delegate);
 }
 
-void FShaderFormatsPropertyDetails::CreateTargetShaderFormatsPropertyView(ITargetPlatform* TargetPlatform, GetFriendlyNameFromRHINameFnc FriendlyNameFnc)
+void FShaderFormatsPropertyDetails::CreateTargetShaderFormatsPropertyView(
+	ITargetPlatform* TargetPlatform,
+	GetFriendlyNameFromRHINameFnc FriendlyNameFnc,
+	FilterShaderPlatformFnc* FilterShaderPlatformFunc,
+	ECategoryPriority::Type Priority)
 {
 	check(TargetPlatform);
 	DetailBuilder->HideProperty(ShaderFormatsPropertyHandle);
@@ -103,12 +48,17 @@ void FShaderFormatsPropertyDetails::CreateTargetShaderFormatsPropertyView(ITarge
 	TArray<FName> ShaderFormats;
 	TargetPlatform->GetAllPossibleShaderFormats(ShaderFormats);
 	
-	IDetailCategoryBuilder& TargetedRHICategoryBuilder = DetailBuilder->EditCategory(*Title);
+	IDetailCategoryBuilder& TargetedRHICategoryBuilder = DetailBuilder->EditCategory(Title, FText::GetEmpty(), Priority);
 	
 	int32 ShaderCounter = 0;
 	for (const FName& ShaderFormat : ShaderFormats)
 	{
-		FText FriendlyShaderFormatName = FriendlyNameFnc(ShaderFormat.ToString());
+		if (FilterShaderPlatformFunc && !FilterShaderPlatformFunc(ShaderFormat))
+		{
+			continue;
+		}
+
+		const FText FriendlyShaderFormatName = FriendlyNameFnc(ShaderFormat);
 		if (!FriendlyShaderFormatName.IsEmpty())
 		{
 			ShaderFormatOrder.Add(ShaderFormat, ShaderCounter++);
diff --git a/Engine/Source/Developer/SharedSettingsWidgets/Public/ShaderFormatsPropertyDetails.h b/Engine/Source/Developer/SharedSettingsWidgets/Public/ShaderFormatsPropertyDetails.h
--- a/Engine/Source/Developer/SharedSettingsWidgets/Public/ShaderFormatsPropertyDetails.h
+++ b/Engine/Source/Developer/SharedSettingsWidgets/Public/ShaderFormatsPropertyDetails.h
@@ -3,7 +3,7 @@
 #pragma once
 
 #include "CoreMinimal.h"
-#include "IDetailCustomization.h"
+#include "DetailLayoutBuilder.h"
 
 class IDetailLayoutBuilder;
 class IPropertyHandle;
@@ -11,19 +11,17 @@ class ITargetPlatform;
 
 enum class ECheckBoxState : uint8;
 
-/* FShaderFormatsPropertyDetails
- *****************************************************************************/
-
 /**
  * Helper which implements details panel customizations for a device profiles parent property
  */
 class SHAREDSETTINGSWIDGETS_API FShaderFormatsPropertyDetails
 : public TSharedFromThis<FShaderFormatsPropertyDetails>
 {
-	
 public:
-	typedef FText (*GetFriendlyNameFromRHINameFnc)(const FString&);
-	static FText GetFriendlyNameFromRHINameMac(const FString& InRHIName);
+	typedef FText(Deprecated_GetFriendlyNameFromRHINameFnc)(const FString&);
+
+	typedef FText (GetFriendlyNameFromRHINameFnc)(FName);
+	typedef bool (FilterShaderPlatformFnc)(FName);
 	
 	/**
 	 * Constructor for the parent property details view
@@ -32,16 +30,18 @@ public:
 	 * @param InProperty - The category name to override
 	 * @param InTitle - Title for display
 	 */
-	FShaderFormatsPropertyDetails(IDetailLayoutBuilder* InDetailBuilder, FString InProperty = TEXT("TargetedRHIs"), FString InTitle = TEXT("Targeted RHIs"));
-	
+	FShaderFormatsPropertyDetails(IDetailLayoutBuilder* InDetailBuilder, const TCHAR* InProperty = TEXT("TargetedRHIs"), const TCHAR* InTitle = TEXT("Targeted RHIs"));
+
 	/** Simple delegate for updating shader version warning */
-	void SetOnUpdateShaderWarning(FSimpleDelegate const& Delegate);
-	
-	/**
-	 * Create the UI to select which windows shader formats we are targetting
-	 */
-	void CreateTargetShaderFormatsPropertyView(ITargetPlatform* TargetPlatform, GetFriendlyNameFromRHINameFnc FriendlyNameFnc);
+	void SetOnUpdateShaderWarning(const FSimpleDelegate& Delegate);
 	
+	/** Create the UI to select which windows shader formats we are targeting */
+	void CreateTargetShaderFormatsPropertyView(ITargetPlatform* TargetPlatform, GetFriendlyNameFromRHINameFnc* FriendlyNameFnc, FilterShaderPlatformFnc* FilterShaderPlatformFunc = nullptr, ECategoryPriority::Type InPriority = ECategoryPriority::Default);
+
+	UE_DEPRECATED(5.1, "CreateTargetShaderFormatsPropertyView now gets RHI names via FName instead of FString. Please change your callback function as this is a breaking change.")
+	void CreateTargetShaderFormatsPropertyView(ITargetPlatform* TargetPlatform, Deprecated_GetFriendlyNameFromRHINameFnc* FriendlyNameFnc)
+	{
+	}
 	
 	/**
 	 * @param InRHIName - The input RHI to check
@@ -65,8 +65,9 @@ private:
 	/** The category name to override */
 	FString Property;
 
+
 	/** Title for display */
-	FString Title;
+	FName Title;
 
 	/** Preserve shader format order when writing to property */
 	TMap<FName, int> ShaderFormatOrder;
diff --git a/Engine/Source/Developer/Windows/WindowsTargetPlatform/Classes/WindowsTargetSettings.h b/Engine/Source/Developer/Windows/WindowsTargetPlatform/Classes/WindowsTargetSettings.h
--- a/Engine/Source/Developer/Windows/WindowsTargetPlatform/Classes/WindowsTargetSettings.h
+++ b/Engine/Source/Developer/Windows/WindowsTargetPlatform/Classes/WindowsTargetSettings.h
@@ -30,14 +30,6 @@ enum class EDefaultGraphicsRHI : uint8
 	DefaultGraphicsRHI_Vulkan = 3 UMETA(DisplayName = "Vulkan"),
 };
 
-UENUM()
-enum class EWindowsRHIFeatureLevel : uint8
-{
-	ES3_1 = 0 UMETA(DisplayName = "ES3.1"),
-	SM5   = 1 UMETA(DisplayName = "SM5"),
-	SM6   = 2 UMETA(DisplayName = "SM6"),
-};
-
 /**
  * Implements the settings for the Windows target platform. The first instance of this class is initialized in
  * WindowsTargetPlatform, really early during the startup sequence before the CDO has been constructed, so its config 
@@ -48,43 +40,29 @@ class WINDOWSTARGETPLATFORM_API UWindowsTargetSettings
 	: public UObject
 {
 public:
-
 	GENERATED_UCLASS_BODY()
 
-	/** The compiler version to use for this project. May be different to the chosen IDE. */
-	UPROPERTY(EditAnywhere, config, Category="Toolchain", Meta=(DisplayName="Compiler Version"))
-	ECompilerVersion Compiler;
-
-	/** 
-	 * The collection of RHI's we want to support on this platform.
-	 * This is not always the full list of RHI we can support.
-	 */
-	UPROPERTY(EditAnywhere, config, Category=Rendering)
-	TArray<FString> TargetedRHIs;
+	virtual void PostInitProperties() override;
 
-	/** 
-	 * Select which RHI to use. Make sure its also selected as a Targeted RHI. Requires Editor restart.
-	 */
+	/** Select which RHI to use. Make sure its also selected as a Targeted RHI. Requires Editor restart. */
 	UPROPERTY(EditAnywhere, config, Category="Targeted RHIs", Meta = (DisplayName = "Default RHI", ConfigRestartRequired = true))
 	EDefaultGraphicsRHI DefaultGraphicsRHI;
 
-	UPROPERTY(EditAnywhere, config, Category = "Targeted RHIs", AdvancedDisplay, Meta = (DisplayName = "DirectX 12 Minimum FeatureLevel", ConfigRestartRequired = true))
-	EWindowsRHIFeatureLevel D3D12MinimumFeatureLevel;
-
-	UPROPERTY(EditAnywhere, config, Category = "Targeted RHIs", AdvancedDisplay, Meta = (DisplayName = "DirectX 12 Maximum FeatureLevel", ConfigRestartRequired = true))
-	EWindowsRHIFeatureLevel D3D12MaximumFeatureLevel;
+	UPROPERTY(config, meta = (DeprecatedProperty, DeprecationMessage = "Use one of the RHI specific lists."))
+	TArray<FString> TargetedRHIs_DEPRECATED;
 
-	UPROPERTY(EditAnywhere, config, Category = "Targeted RHIs", AdvancedDisplay, Meta = (DisplayName = "DirectX 11 Minimum FeatureLevel", ConfigRestartRequired = true))
-	EWindowsRHIFeatureLevel D3D11MinimumFeatureLevel;
+	UPROPERTY(EditAnywhere, config, Category = "Rendering")
+	TArray<FString> D3D12TargetedShaderFormats;
 
-	UPROPERTY(EditAnywhere, config, Category = "Targeted RHIs", AdvancedDisplay, Meta = (DisplayName = "DirectX 11 Maximum FeatureLevel", ConfigRestartRequired = true))
-	EWindowsRHIFeatureLevel D3D11MaximumFeatureLevel;
+	UPROPERTY(EditAnywhere, config, Category = "Rendering")
+	TArray<FString> D3D11TargetedShaderFormats;
 
-	UPROPERTY(EditAnywhere, config, Category = "Targeted RHIs", AdvancedDisplay, Meta = (DisplayName = "Vulkan Minimum FeatureLevel", ConfigRestartRequired = true))
-	EWindowsRHIFeatureLevel VulkanMinimumFeatureLevel;
+	UPROPERTY(EditAnywhere, config, Category = "Rendering")
+	TArray<FString> VulkanTargetedShaderFormats;
 
-	UPROPERTY(EditAnywhere, config, Category = "Targeted RHIs", AdvancedDisplay, Meta = (DisplayName = "Vulkan Maximum FeatureLevel", ConfigRestartRequired = true))
-	EWindowsRHIFeatureLevel VulkanMaximumFeatureLevel;
+	/** The compiler version to use for this project. May be different to the chosen IDE. */
+	UPROPERTY(EditAnywhere, config, Category = "Toolchain", Meta = (DisplayName = "Compiler Version"))
+	ECompilerVersion Compiler;
 
 	/** Sample rate to run the audio mixer with. */
 	UPROPERTY(config, EditAnywhere, Category = "Audio", Meta = (DisplayName = "Audio Mixer Sample Rate"))
diff --git a/Engine/Source/Developer/Windows/WindowsTargetPlatform/Private/GenericWindowsTargetPlatform.h b/Engine/Source/Developer/Windows/WindowsTargetPlatform/Private/GenericWindowsTargetPlatform.h
--- a/Engine/Source/Developer/Windows/WindowsTargetPlatform/Private/GenericWindowsTargetPlatform.h
+++ b/Engine/Source/Developer/Windows/WindowsTargetPlatform/Private/GenericWindowsTargetPlatform.h
@@ -258,8 +258,23 @@ public:
 	virtual void GetAllTargetedShaderFormats( TArray<FName>& OutFormats ) const override 
 	{
 		// Get the Target RHIs for this platform, we do not always want all those that are supported. (reload in case user changed in the editor)
-		TArray<FString>TargetedShaderFormats;
-		GConfig->GetArray(TEXT("/Script/WindowsTargetPlatform.WindowsTargetSettings"), TEXT("TargetedRHIs"), TargetedShaderFormats, GEngineIni);
+		TArray<FString> TargetedShaderFormats;
+
+		auto AddWindowsShaderFormats = [&TargetedShaderFormats](const TCHAR* InName)
+		{
+			TArray<FString> NewTargetedShaderFormats;
+			GConfig->GetArray(TEXT("/Script/WindowsTargetPlatform.WindowsTargetSettings"), InName, NewTargetedShaderFormats, GEngineIni);
+
+			for (const FString& NewShaderFormat : NewTargetedShaderFormats)
+			{
+				TargetedShaderFormats.AddUnique(NewShaderFormat);
+			}
+		};
+
+		AddWindowsShaderFormats(TEXT("TargetedRHIs"));
+		AddWindowsShaderFormats(TEXT("D3D12TargetedShaderFormats"));
+		AddWindowsShaderFormats(TEXT("D3D11TargetedShaderFormats"));
+		AddWindowsShaderFormats(TEXT("VulkanTargetedShaderFormats"));
 
 		// Gather the list of Target RHIs and filter out any that may be invalid.
 		TArray<FName> PossibleShaderFormats;
diff --git a/Engine/Source/Developer/Windows/WindowsTargetPlatform/Private/WindowsTargetPlatformClasses.cpp b/Engine/Source/Developer/Windows/WindowsTargetPlatform/Private/WindowsTargetPlatformClasses.cpp
--- a/Engine/Source/Developer/Windows/WindowsTargetPlatform/Private/WindowsTargetPlatformClasses.cpp
+++ b/Engine/Source/Developer/Windows/WindowsTargetPlatform/Private/WindowsTargetPlatformClasses.cpp
@@ -21,3 +21,44 @@ UWindowsTargetSettings::UWindowsTargetSettings( const FObjectInitializer& Object
 	AudioNumBuffersToEnqueue = 1;
 	AudioNumSourceWorkers = 4;
 }
+
+static bool FilterShaderPlatform_D3D12(const FString& InShaderPlatform)
+{
+	return InShaderPlatform == TEXT("PCD3D_SM6") || InShaderPlatform == TEXT("PCD3D_SM5") || InShaderPlatform == TEXT("PCD3D_ES31");
+}
+
+static bool FilterShaderPlatform_D3D11(const FString& InShaderPlatform)
+{
+	return InShaderPlatform == TEXT("PCD3D_SM5") || InShaderPlatform == TEXT("PCD3D_ES31");
+}
+
+static bool FilterShaderPlatform_Vulkan(const FString& InShaderPlatform)
+{
+	return InShaderPlatform == TEXT("SF_VULKAN_SM5");
+}
+
+template<typename TFilterFunc>
+static void AddToShaderFormatList(TArray<FString>& Dest, const TArray<FString>& Source, TFilterFunc FilterFunc)
+{
+	for (const FString& ShaderFormat : Source)
+	{
+		if (FilterFunc(ShaderFormat))
+		{
+			Dest.AddUnique(ShaderFormat);
+		}
+	}
+}
+
+void UWindowsTargetSettings::PostInitProperties()
+{
+	Super::PostInitProperties();
+
+	if (!TargetedRHIs_DEPRECATED.IsEmpty())
+	{
+		AddToShaderFormatList(D3D12TargetedShaderFormats, TargetedRHIs_DEPRECATED, &FilterShaderPlatform_D3D12);
+		AddToShaderFormatList(D3D11TargetedShaderFormats, TargetedRHIs_DEPRECATED, &FilterShaderPlatform_D3D11);
+		AddToShaderFormatList(VulkanTargetedShaderFormats, TargetedRHIs_DEPRECATED, &FilterShaderPlatform_Vulkan);
+
+		TargetedRHIs_DEPRECATED.Empty();
+	}
+}
diff --git a/Engine/Source/Editor/DetailCustomizations/Private/LinuxTargetSettingsDetails.cpp b/Engine/Source/Editor/DetailCustomizations/Private/LinuxTargetSettingsDetails.cpp
--- a/Engine/Source/Editor/DetailCustomizations/Private/LinuxTargetSettingsDetails.cpp
+++ b/Engine/Source/Editor/DetailCustomizations/Private/LinuxTargetSettingsDetails.cpp
@@ -27,6 +27,7 @@
 #include "Interfaces/ITargetPlatform.h"
 #include "Interfaces/ITargetPlatformModule.h"
 #include "SExternalImageReference.h"
+#include "RHIShaderFormatDefinitions.inl"
 
 #if WITH_ENGINE
 #include "AudioDevice.h"
@@ -46,34 +47,31 @@ namespace LinuxTargetSettingsDetailsConstants
 	const FText DisabledTip = LOCTEXT("GitHubSourceRequiredToolTip", "This requires GitHub source.");
 }
 
-static FText GetFriendlyNameFromLinuxRHIName(const FString& InRHIName)
+static FText GetFriendlyNameFromLinuxShaderFormat(const FName InShaderFormat)
 {
-	FText FriendlyRHIName = LOCTEXT("UnknownRHI", "UnknownRHI");
-	if (InRHIName == TEXT("GLSL_150"))
-	{
-		FriendlyRHIName = LOCTEXT("OpenGL3", "OpenGL 3 (SM4)");
-	}
-	else if (InRHIName == TEXT("GLSL_150_ES2"))
-	{
-		FriendlyRHIName = LOCTEXT("OpenGL3ES2", "OpenGL 3 (ES2)");
-	}
-	else if (InRHIName == TEXT("GLSL_150_ES31"))
+	FText FriendlyRHIName;
+
+	if (InShaderFormat == NAME_GLSL_150_ES31)
 	{
 		FriendlyRHIName = LOCTEXT("OpenGL3ES31", "OpenGL 3 (ES3.1, Experimental)");
 	}
-	else if (InRHIName == TEXT("SF_VULKAN_ES31_ANDROID") || InRHIName == TEXT("SF_VULKAN_ES31"))
+	else if (InShaderFormat == NAME_VULKAN_ES3_1_ANDROID || InShaderFormat == NAME_VULKAN_ES3_1)
 	{
 		FriendlyRHIName = LOCTEXT("Vulkan ES31", "Vulkan Mobile (ES3.1)");
 	}
-	else if (InRHIName == TEXT("SF_VULKAN_SM5"))
+	else if (InShaderFormat == NAME_VULKAN_SM5)
 	{
 		FriendlyRHIName = LOCTEXT("VulkanSM5", "Vulkan Desktop (SM5)");
 	}
-	else if (InRHIName == TEXT("GLSL_430"))
+	else if (InShaderFormat == TEXT("GLSL_430"))
 	{
 		// Explicitly remove these formats as they are obsolete/not quite supported; users can still target them by adding them as +TargetedRHIs in the TargetPlatform ini.
 		FriendlyRHIName = FText::GetEmpty();
 	}
+	else
+	{
+		FriendlyRHIName = LOCTEXT("UnknownRHI", "UnknownRHI");
+	}
 
 	return FriendlyRHIName;
 }
@@ -151,7 +149,7 @@ void FLinuxTargetSettingsDetails::CustomizeDetails( IDetailLayoutBuilder& Detail
 	// Setup the supported/targeted RHI property view
 	ITargetPlatform* TargetPlatform = FModuleManager::GetModuleChecked<ITargetPlatformModule>("LinuxTargetPlatform").GetTargetPlatforms()[0];
 	TargetShaderFormatsDetails = MakeShareable(new FShaderFormatsPropertyDetails(&DetailBuilder));
-	TargetShaderFormatsDetails->CreateTargetShaderFormatsPropertyView(TargetPlatform, GetFriendlyNameFromLinuxRHIName);
+	TargetShaderFormatsDetails->CreateTargetShaderFormatsPropertyView(TargetPlatform, &GetFriendlyNameFromLinuxShaderFormat);
 
 	// Next add the splash image customization
 	const FText EditorSplashDesc(LOCTEXT("EditorSplashLabel", "Editor Splash"));
diff --git a/Engine/Source/Editor/DetailCustomizations/Private/MacTargetSettingsDetails.cpp b/Engine/Source/Editor/DetailCustomizations/Private/MacTargetSettingsDetails.cpp
--- a/Engine/Source/Editor/DetailCustomizations/Private/MacTargetSettingsDetails.cpp
+++ b/Engine/Source/Editor/DetailCustomizations/Private/MacTargetSettingsDetails.cpp
@@ -100,6 +100,65 @@ static FString GetIconFilename(EMacImageScope::Type Scope)
 	}
 }
 
+static FText GetFriendlyNameFromRHINameMac(FName InRHIName)
+{
+	FText FriendlyRHIName;
+
+	const EShaderPlatform Platform = ShaderFormatToLegacyShaderPlatform(InRHIName);
+	switch (Platform)
+	{
+	case SP_PCD3D_SM5:
+		FriendlyRHIName = LOCTEXT("D3DSM5", "Direct3D 11+ (SM5)");
+		break;
+	case SP_PCD3D_ES3_1:
+		FriendlyRHIName = LOCTEXT("D3DES31", "Direct3D (ES3.1, Mobile Preview)");
+		break;
+	case SP_OPENGL_PCES3_1:
+		FriendlyRHIName = LOCTEXT("OpenGLES31PC", "OpenGL (ES3.1, Mobile Preview)");
+		break;
+	case SP_OPENGL_ES3_1_ANDROID:
+		FriendlyRHIName = LOCTEXT("OpenGLES31", "OpenGLES 3.1 (Mobile)");
+		break;
+	case SP_METAL:
+		FriendlyRHIName = LOCTEXT("Metal", "iOS Metal Mobile Renderer (ES3.1, Metal 1.1+, iOS 9.0 or later)");
+		break;
+	case SP_METAL_MRT:
+		FriendlyRHIName = LOCTEXT("MetalMRT", "iOS Metal Desktop Renderer (SM5, Metal 1.2+, iOS 10.0 or later)");
+		break;
+	case SP_METAL_TVOS:
+		FriendlyRHIName = LOCTEXT("MetalTV", "tvOS Metal Mobile Renderer (ES3.1, Metal 1.1+, tvOS 9.0 or later)");
+		break;
+	case SP_METAL_MRT_TVOS:
+		FriendlyRHIName = LOCTEXT("MetalMRTTV", "tvOS Metal Desktop Renderer (SM5, Metal 1.2+, tvOS 10.0 or later)");
+		break;
+	case SP_METAL_SM5:
+		FriendlyRHIName = LOCTEXT("MetalSM5", "Mac Metal Desktop Renderer (SM5, Metal 2.1+, macOS Catalina 10.15.7 or later)");
+		break;
+	case SP_METAL_MACES3_1:
+		FriendlyRHIName = LOCTEXT("MetalES3.1", "Mac Metal High-End Mobile Preview (ES3.1)");
+		break;
+	case SP_METAL_MRT_MAC:
+		FriendlyRHIName = LOCTEXT("MetalMRTMac", "Mac Metal iOS/tvOS Desktop Renderer Preview (SM5)");
+		break;
+	case SP_VULKAN_SM5:
+	case SP_VULKAN_SM5_ANDROID:
+		FriendlyRHIName = LOCTEXT("VulkanSM5", "Vulkan (SM5)");
+		break;
+	case SP_VULKAN_PCES3_1:
+	case SP_VULKAN_ES3_1_ANDROID:
+		FriendlyRHIName = LOCTEXT("VulkanES31", "Vulkan (ES 3.1)");
+		break;
+	case SP_D3D_ES3_1_HOLOLENS:
+		FriendlyRHIName = LOCTEXT("D3DES31HL", "Direct3D (ES3.1, Hololens)");
+		break;
+	default:
+		FriendlyRHIName = FText::FromString(InRHIName.ToString());
+		break;
+	}
+
+	return FriendlyRHIName;
+}
+
 void FMacTargetSettingsDetails::CustomizeDetails( IDetailLayoutBuilder& DetailBuilder )
 {
 	FSimpleDelegate OnUpdateShaderStandardWarning = FSimpleDelegate::CreateSP(this, &FMacTargetSettingsDetails::UpdateShaderStandardWarning);
@@ -109,7 +168,7 @@ void FMacTargetSettingsDetails::CustomizeDetails( IDetailLayoutBuilder& DetailBu
 	// Setup the supported/targeted RHI property view
 	TargetShaderFormatsDetails = MakeShareable(new FShaderFormatsPropertyDetails(&DetailBuilder, TEXT("TargetedRHIs"), TEXT("Targeted RHIs")));
 	TargetShaderFormatsDetails->SetOnUpdateShaderWarning(OnUpdateShaderStandardWarning);
-	TargetShaderFormatsDetails->CreateTargetShaderFormatsPropertyView(TargetPlatform, FShaderFormatsPropertyDetails::GetFriendlyNameFromRHINameMac);
+	TargetShaderFormatsDetails->CreateTargetShaderFormatsPropertyView(TargetPlatform, &GetFriendlyNameFromRHINameMac);
 	
 	// Setup the shader version property view
     // Handle max. shader version a little specially.
diff --git a/Engine/Source/Editor/DetailCustomizations/Private/WindowsTargetSettingsDetails.cpp b/Engine/Source/Editor/DetailCustomizations/Private/WindowsTargetSettingsDetails.cpp
--- a/Engine/Source/Editor/DetailCustomizations/Private/WindowsTargetSettingsDetails.cpp
+++ b/Engine/Source/Editor/DetailCustomizations/Private/WindowsTargetSettingsDetails.cpp
@@ -50,9 +50,9 @@ namespace WindowsTargetSettingsDetailsConstants
 	const FText DisabledTip = LOCTEXT("GitHubSourceRequiredToolTip", "This requires GitHub source.");
 }
 
-static FText GetFriendlyNameFromShaderPlatform(const FString& InShaderPlatformName)
+static FText GetFriendlyNameFromWindowsShaderPlatform(FName InShaderPlatformName)
 {
-	const EShaderPlatform ShaderPlatform = ShaderFormatNameToShaderPlatform(FName(*InShaderPlatformName));
+	const EShaderPlatform ShaderPlatform = ShaderFormatNameToShaderPlatform(InShaderPlatformName);
 
 	FText FriendlyRHIName;
 	switch (ShaderPlatform)
@@ -80,7 +80,7 @@ static FText GetFriendlyNameFromShaderPlatform(const FString& InShaderPlatformNa
 		break;
 
 	default:
-		UE_LOG(LogEngine, Warning, TEXT("Unknown Windows target RHI %s"), *InShaderPlatformName);
+		UE_LOG(LogEngine, Warning, TEXT("Unknown Windows target RHI %s"), *InShaderPlatformName.ToString());
 		FriendlyRHIName = LOCTEXT("UnknownRHI", "UnknownRHI");
 		break;
 	}
@@ -88,6 +88,48 @@ static FText GetFriendlyNameFromShaderPlatform(const FString& InShaderPlatformNa
 	return FriendlyRHIName;
 }
 
+static FText GetFriendlyNameForWindowsShaderPlatformCheckbox(FName InShaderPlatformName)
+{
+	const EShaderPlatform ShaderPlatform = ShaderFormatNameToShaderPlatform(InShaderPlatformName);
+
+	FText FriendlyName;
+
+	switch (ShaderPlatform)
+	{
+	case SP_PCD3D_SM6:
+		FriendlyName = LOCTEXT("SM6", "SM6");
+		break;
+	case SP_PCD3D_SM5:
+		FriendlyName = LOCTEXT("SM5", "SM5");
+		break;
+	case SP_PCD3D_ES3_1:
+		FriendlyName = LOCTEXT("ES31", "ES3.1");
+		break;
+	case SP_VULKAN_SM5:
+		FriendlyName = LOCTEXT("SM5", "SM5");
+		break;
+	default:
+		break;
+	}
+
+	return FriendlyName;
+}
+
+static bool FilterShaderPlatform_D3D12(FName InShaderPlatform)
+{
+	return InShaderPlatform == NAME_PCD3D_SM6 || InShaderPlatform == NAME_PCD3D_SM5 || InShaderPlatform == NAME_PCD3D_ES3_1;
+}
+
+static bool FilterShaderPlatform_D3D11(FName InShaderPlatform)
+{
+	return InShaderPlatform == NAME_PCD3D_SM5 || InShaderPlatform == NAME_PCD3D_ES3_1;
+}
+
+static bool FilterShaderPlatform_Vulkan(FName InShaderPlatform)
+{
+	return InShaderPlatform == NAME_VULKAN_SM5;
+}
+
 TSharedRef<IDetailCustomization> FWindowsTargetSettingsDetails::MakeInstance()
 {
 	return MakeShareable(new FWindowsTargetSettingsDetails);
@@ -159,8 +201,15 @@ void FWindowsTargetSettingsDetails::CustomizeDetails( IDetailLayoutBuilder& Deta
 {
 	// Setup the supported/targeted RHI property view
 	ITargetPlatform* TargetPlatform = FModuleManager::GetModuleChecked<ITargetPlatformModule>("WindowsTargetPlatform").GetTargetPlatforms()[0];
-	TargetShaderFormatsDetails = MakeShareable(new FShaderFormatsPropertyDetails(&DetailBuilder));
-	TargetShaderFormatsDetails->CreateTargetShaderFormatsPropertyView(TargetPlatform, GetFriendlyNameFromShaderPlatform);
+
+	D3D12TargetShaderFormatsDetails = MakeShareable(new FShaderFormatsPropertyDetails(&DetailBuilder, TEXT("D3D12TargetedShaderFormats"), TEXT("D3D12 Targeted Shader Formats")));
+	D3D12TargetShaderFormatsDetails->CreateTargetShaderFormatsPropertyView(TargetPlatform, &GetFriendlyNameForWindowsShaderPlatformCheckbox, &FilterShaderPlatform_D3D12, ECategoryPriority::Important);
+
+	D3D11TargetShaderFormatsDetails = MakeShareable(new FShaderFormatsPropertyDetails(&DetailBuilder, TEXT("D3D11TargetedShaderFormats"), TEXT("D3D11 Targeted Shader Formats")));
+	D3D11TargetShaderFormatsDetails->CreateTargetShaderFormatsPropertyView(TargetPlatform, &GetFriendlyNameForWindowsShaderPlatformCheckbox, &FilterShaderPlatform_D3D11, ECategoryPriority::Important);
+
+	VulkanTargetShaderFormatsDetails = MakeShareable(new FShaderFormatsPropertyDetails(&DetailBuilder, TEXT("VulkanTargetedShaderFormats"), TEXT("Vulkan Targeted Shader Formats")));
+	VulkanTargetShaderFormatsDetails->CreateTargetShaderFormatsPropertyView(TargetPlatform, &GetFriendlyNameForWindowsShaderPlatformCheckbox, &FilterShaderPlatform_Vulkan, ECategoryPriority::Important);
 
 	// Next add the splash image customization
 	const FText EditorSplashDesc(LOCTEXT("EditorSplashLabel", "Editor Splash"));
diff --git a/Engine/Source/Editor/DetailCustomizations/Private/WindowsTargetSettingsDetails.h b/Engine/Source/Editor/DetailCustomizations/Private/WindowsTargetSettingsDetails.h
--- a/Engine/Source/Editor/DetailCustomizations/Private/WindowsTargetSettingsDetails.h
+++ b/Engine/Source/Editor/DetailCustomizations/Private/WindowsTargetSettingsDetails.h
@@ -41,12 +41,11 @@ private:
 	/** This gets the current value of the audio stream caching bool property. */
 	ECheckBoxState GetAudioStreamCachingToggled(TSharedPtr<IPropertyHandle> PropertyHandle) const;
 
-protected:
-
-
-
 private:
 	/** Reference to the target shader formats property view */
-	TSharedPtr<FShaderFormatsPropertyDetails> TargetShaderFormatsDetails;
+	TSharedPtr<FShaderFormatsPropertyDetails> D3D12TargetShaderFormatsDetails;
+	TSharedPtr<FShaderFormatsPropertyDetails> D3D11TargetShaderFormatsDetails;
+	TSharedPtr<FShaderFormatsPropertyDetails> VulkanTargetShaderFormatsDetails;
+
 	FAudioPluginWidgetManager AudioPluginWidgetManager;
 };
diff --git a/Engine/Source/Runtime/RHI/Private/Windows/WindowsDynamicRHI.cpp b/Engine/Source/Runtime/RHI/Private/Windows/WindowsDynamicRHI.cpp
--- a/Engine/Source/Runtime/RHI/Private/Windows/WindowsDynamicRHI.cpp
+++ b/Engine/Source/Runtime/RHI/Private/Windows/WindowsDynamicRHI.cpp
@@ -5,6 +5,7 @@
 #include "Misc/CommandLine.h"
 #include "Misc/ConfigCacheIni.h"
 #include "Misc/MessageDialog.h"
+#include "Containers/StaticArray.h"
 
 #if WINDOWS_USE_FEATURE_DYNAMIC_RHI
 
@@ -16,108 +17,265 @@
 
 static const TCHAR* GLoadedRHIModuleName;
 
-enum class WindowsRHI
+enum class EWindowsRHI
 {
 	D3D11,
 	D3D12,
 	Vulkan,
 	OpenGL,
+	count
 };
+static constexpr int32 EWindowsRHICount = static_cast<int32>(EWindowsRHI::count);
 
-static TArray<EShaderPlatform> GetTargetedShaderPlatforms()
+static const TCHAR* ModuleNameFromWindowsRHI(EWindowsRHI InWindowsRHI)
 {
-	TArray<FString> TargetedShaderFormats;
-	GConfig->GetArray(TEXT("/Script/WindowsTargetPlatform.WindowsTargetSettings"), TEXT("TargetedRHIs"), TargetedShaderFormats, GEngineIni);
-
-	TArray<EShaderPlatform> ShaderPlatforms;
-	ShaderPlatforms.Reserve(TargetedShaderFormats.Num());
-
-	for (const FString& ShaderFormat : TargetedShaderFormats)
+	switch (InWindowsRHI)
 	{
-		ShaderPlatforms.Add(ShaderFormatToLegacyShaderPlatform(FName(*ShaderFormat)));
+	default: check(false);
+	case EWindowsRHI::D3D11:  return TEXT("D3D11RHI");
+	case EWindowsRHI::D3D12:  return TEXT("D3D12RHI");
+	case EWindowsRHI::Vulkan: return TEXT("VulkanRHI");
+	case EWindowsRHI::OpenGL: return TEXT("OpenGLDrv");
 	}
+}
 
-	return ShaderPlatforms;
+static FString GetRHINameFromWindowsRHI(EWindowsRHI InWindowsRHI, ERHIFeatureLevel::Type InFeatureLevel)
+{
+	switch (InWindowsRHI)
+	{
+	default: check(false);
+	case EWindowsRHI::D3D11:  return TEXT("DirectX 11");
+	case EWindowsRHI::D3D12:
+	{
+		FString FeatureLevelName;
+		GetFeatureLevelName(InFeatureLevel, FeatureLevelName);
+		return FString::Printf(TEXT("DirectX 12 (%s)"), *FeatureLevelName);
+	}
+	case EWindowsRHI::Vulkan: return TEXT("Vulkan");
+	case EWindowsRHI::OpenGL: return TEXT("OpenGL");
+	}
 }
 
-static bool IsD3DSM6Platform(EShaderPlatform InShaderPlatform)
+static ERHIFeatureLevel::Type GetDefaultFeatureLevelForRHI(EWindowsRHI InWindowsRHI)
 {
-	return IsD3DPlatform(InShaderPlatform) && IsFeatureLevelSupported(InShaderPlatform, ERHIFeatureLevel::SM6);
+	switch (InWindowsRHI)
+	{
+	case EWindowsRHI::D3D11:  return ERHIFeatureLevel::SM5;
+	case EWindowsRHI::D3D12:  return ERHIFeatureLevel::SM5;
+	case EWindowsRHI::Vulkan: return ERHIFeatureLevel::SM5;
+	case EWindowsRHI::OpenGL: return ERHIFeatureLevel::ES3_1;
+	default: check(false);    return ERHIFeatureLevel::SM5;
+	}
 }
 
-static bool IsD3DSM6PlatformTargeted(const TArray<EShaderPlatform>& TargetedShaderPlatforms)
+struct FWindowsRHIConfig
+{
+	TArray<EShaderPlatform>        ShaderPlatforms;
+	TArray<ERHIFeatureLevel::Type> FeatureLevels;
+};
+
+struct FParsedWindowsDynamicRHIConfig
+{
+	TOptional<EWindowsRHI> DefaultRHI{};
+
+	TStaticArray<FWindowsRHIConfig, EWindowsRHICount> RHIConfigs;
+
+	bool IsEmpty() const;
+
+	bool IsRHISupported(EWindowsRHI InWindowsRHI) const;
+	TOptional<ERHIFeatureLevel::Type> GetHighestSupportedFeatureLevel(EWindowsRHI InWindowsRHI) const;
+	TOptional<ERHIFeatureLevel::Type> GetNextHighestTargetedFeatureLevel(EWindowsRHI InWindowsRHI, ERHIFeatureLevel::Type InFeatureLevel) const;
+
+	bool IsFeatureLevelTargeted(EWindowsRHI InWindowsRHI, ERHIFeatureLevel::Type InFeatureLevel) const;
+
+	void MergeConfig(EWindowsRHI InWindowsRHI, const FWindowsRHIConfig& Other);
+};
+
+bool FParsedWindowsDynamicRHIConfig::IsEmpty() const
 {
-	for (const EShaderPlatform ShaderPlatform : TargetedShaderPlatforms)
+	for (const FWindowsRHIConfig& Config : RHIConfigs)
 	{
-		if (IsD3DSM6Platform(ShaderPlatform))
+		if (!Config.ShaderPlatforms.IsEmpty())
 		{
-			return true;
+			return false;
 		}
 	}
 
-	return false;
+	return true;
+}
+
+bool FParsedWindowsDynamicRHIConfig::IsRHISupported(EWindowsRHI InWindowsRHI) const
+{
+	// If we don't require cooked data, then we should be able to support this RHI at any feature level.
+	if (!FPlatformProperties::RequiresCookedData())
+	{
+		return true;
+	}
+
+	return !RHIConfigs[(int32)InWindowsRHI].ShaderPlatforms.IsEmpty();
 }
 
-static TOptional<ERHIFeatureLevel::Type> ParseFeatureLevelFromSetting(const TCHAR* InName)
+TOptional<ERHIFeatureLevel::Type> FParsedWindowsDynamicRHIConfig::GetHighestSupportedFeatureLevel(EWindowsRHI InWindowsRHI) const
 {
-	TOptional<ERHIFeatureLevel::Type> ResultFeatureLevel{};
+	const TArray<ERHIFeatureLevel::Type>& FeatureLevels = RHIConfigs[(int32)InWindowsRHI].FeatureLevels;
+	if (FeatureLevels.Num() == 0)
+	{
+		return TOptional<ERHIFeatureLevel::Type>();
+	}
 
-	FString ConfigFeatureLevel;
-	if (GConfig->GetString(TEXT("/Script/WindowsTargetPlatform.WindowsTargetSettings"), InName, ConfigFeatureLevel, GEngineIni))
+	ERHIFeatureLevel::Type MaxFeatureLevel = (ERHIFeatureLevel::Type)0;
+	for (ERHIFeatureLevel::Type SupportedFeatureLevel : FeatureLevels)
 	{
-		const FName FeatureLevelName(*ConfigFeatureLevel);
+		MaxFeatureLevel = std::max(MaxFeatureLevel, SupportedFeatureLevel);
+	}
+	return MaxFeatureLevel;
+}
+
+TOptional<ERHIFeatureLevel::Type> FParsedWindowsDynamicRHIConfig::GetNextHighestTargetedFeatureLevel(EWindowsRHI InWindowsRHI, ERHIFeatureLevel::Type InFeatureLevel) const
+{
+	TArray<ERHIFeatureLevel::Type> LowerFeatureLevels(RHIConfigs[(int32)InWindowsRHI].FeatureLevels);
+	LowerFeatureLevels.RemoveAll([InFeatureLevel](ERHIFeatureLevel::Type OtherFeatureLevel) { return OtherFeatureLevel >= InFeatureLevel; });
 
-		ERHIFeatureLevel::Type ParsedFeatureLevel = ERHIFeatureLevel::Num;
-		if (GetFeatureLevelFromName(FeatureLevelName, ParsedFeatureLevel))
+	if (LowerFeatureLevels.Num())
+	{
+		ERHIFeatureLevel::Type MaxFeatureLevel = (ERHIFeatureLevel::Type)0;
+		for (ERHIFeatureLevel::Type SupportedFeatureLevel : LowerFeatureLevels)
 		{
-			ResultFeatureLevel = ParsedFeatureLevel;
+			MaxFeatureLevel = std::max(MaxFeatureLevel, SupportedFeatureLevel);
 		}
+		return MaxFeatureLevel;
 	}
 
-	return ResultFeatureLevel;
+	return TOptional<ERHIFeatureLevel::Type>();
+}
+
+bool FParsedWindowsDynamicRHIConfig::IsFeatureLevelTargeted(EWindowsRHI InWindowsRHI, ERHIFeatureLevel::Type InFeatureLevel) const
+{
+	for (ERHIFeatureLevel::Type SupportedFeatureLevel : RHIConfigs[(int32)InWindowsRHI].FeatureLevels)
+	{
+		if (SupportedFeatureLevel == InFeatureLevel)
+		{
+			return true;
+		}
+	}
+	return false;
 }
 
-static ERHIFeatureLevel::Type FilterFeatureLevel(ERHIFeatureLevel::Type InFeatureLevel, const TCHAR* MinSetting, const TCHAR* MaxSetting)
+void FParsedWindowsDynamicRHIConfig::MergeConfig(EWindowsRHI InWindowsRHI, const FWindowsRHIConfig& Other)
 {
-	const TOptional<ERHIFeatureLevel::Type> MinFeatureLevel = ParseFeatureLevelFromSetting(MinSetting);
-	const TOptional<ERHIFeatureLevel::Type> MaxFeatureLevel = ParseFeatureLevelFromSetting(MaxSetting);
+	FWindowsRHIConfig& ExistingConfig = RHIConfigs[(int32)InWindowsRHI];
 
-	if (MinFeatureLevel && InFeatureLevel < MinFeatureLevel.GetValue())
+	for (EShaderPlatform ShaderPlatform : Other.ShaderPlatforms)
 	{
-		return MinFeatureLevel.GetValue();
+		ExistingConfig.ShaderPlatforms.AddUnique(ShaderPlatform);
 	}
-	else if (MaxFeatureLevel && InFeatureLevel > MaxFeatureLevel.GetValue())
+
+	for (ERHIFeatureLevel::Type FeatureLevel : Other.FeatureLevels)
+	{
+		ExistingConfig.FeatureLevels.AddUnique(FeatureLevel);
+	}
+}
+
+TOptional<EWindowsRHI> ParseDefaultWindowsRHI()
+{
+	TOptional<EWindowsRHI> DefaultRHI{};
+
+	FString DefaultGraphicsRHI;
+	if (GConfig->GetString(TEXT("/Script/WindowsTargetPlatform.WindowsTargetSettings"), TEXT("DefaultGraphicsRHI"), DefaultGraphicsRHI, GEngineIni))
 	{
-		return MaxFeatureLevel.GetValue();
+		const FString NAME_DX11(TEXT("DefaultGraphicsRHI_DX11"));
+		const FString NAME_DX12(TEXT("DefaultGraphicsRHI_DX12"));
+		const FString NAME_VULKAN(TEXT("DefaultGraphicsRHI_Vulkan"));
+
+		DefaultRHI = EWindowsRHI::D3D11;
+
+		if (DefaultGraphicsRHI == NAME_DX11)
+		{
+			DefaultRHI = EWindowsRHI::D3D11;
+		}
+		else if (DefaultGraphicsRHI == NAME_DX12)
+		{
+			DefaultRHI = EWindowsRHI::D3D12;
+		}
+		else if (DefaultGraphicsRHI == NAME_VULKAN)
+		{
+			DefaultRHI = EWindowsRHI::Vulkan;
+		}
+		else if (DefaultGraphicsRHI != TEXT("DefaultGraphicsRHI_Default"))
+		{
+			UE_LOG(LogRHI, Error, TEXT("Unrecognized setting '%s' for DefaultGraphicsRHI"), *DefaultGraphicsRHI);
+		}
 	}
 
-	return InFeatureLevel;
+	return DefaultRHI;
 }
 
-static ERHIFeatureLevel::Type FilterFeatureLevel(ERHIFeatureLevel::Type FeatureLevel, WindowsRHI ChosenRHI)
+static TArray<EShaderPlatform> ParseShaderPlatformsConfig(const TCHAR* InSettingName)
 {
-	if (ChosenRHI == WindowsRHI::OpenGL)
+	TArray<FString> TargetedShaderFormats;
+	GConfig->GetArray(TEXT("/Script/WindowsTargetPlatform.WindowsTargetSettings"), InSettingName, TargetedShaderFormats, GEngineIni);
+
+	TArray<EShaderPlatform> ShaderPlatforms;
+	ShaderPlatforms.Reserve(TargetedShaderFormats.Num());
+
+	for (const FString& ShaderFormat : TargetedShaderFormats)
 	{
-		// Seriously locked down.
-		return ERHIFeatureLevel::ES3_1;
+		ShaderPlatforms.AddUnique(ShaderFormatToLegacyShaderPlatform(FName(*ShaderFormat)));
 	}
 
-	if (ChosenRHI == WindowsRHI::D3D12)
+	return ShaderPlatforms;
+}
+
+static TArray<ERHIFeatureLevel::Type> FeatureLevelsFromShaderPlatforms(const TArray<EShaderPlatform>& InShaderPlatforms)
+{
+	TArray<ERHIFeatureLevel::Type> FeatureLevels;
+	FeatureLevels.Reserve(InShaderPlatforms.Num());
+
+	for (EShaderPlatform ShaderPlatform : InShaderPlatforms)
 	{
-		return FilterFeatureLevel(FeatureLevel, TEXT("D3D12MinimumFeatureLevel"), TEXT("D3D12MaximumFeatureLevel"));
+		FeatureLevels.AddUnique(FDataDrivenShaderPlatformInfo::GetMaxFeatureLevel(ShaderPlatform));
 	}
 
-	if (ChosenRHI == WindowsRHI::D3D11)
+	return FeatureLevels;
+}
+
+static FWindowsRHIConfig ParseWindowsRHIConfig(const TCHAR* ShaderFormatsSettingName)
+{
+	FWindowsRHIConfig Config;
+
+	Config.ShaderPlatforms = ParseShaderPlatformsConfig(ShaderFormatsSettingName);
+	Config.FeatureLevels = FeatureLevelsFromShaderPlatforms(Config.ShaderPlatforms);
+
+	return Config;
+}
+
+FParsedWindowsDynamicRHIConfig ParseWindowsDynamicRHIConfig()
+{
+	FParsedWindowsDynamicRHIConfig Config;
+
+	Config.DefaultRHI = ParseDefaultWindowsRHI();
+
+	Config.RHIConfigs[(int32)EWindowsRHI::D3D11]  = ParseWindowsRHIConfig(TEXT("D3D11TargetedShaderFormats"));
+	Config.RHIConfigs[(int32)EWindowsRHI::D3D12]  = ParseWindowsRHIConfig(TEXT("D3D12TargetedShaderFormats"));
+	Config.RHIConfigs[(int32)EWindowsRHI::Vulkan] = ParseWindowsRHIConfig(TEXT("VulkanTargetedShaderFormats"));
+
+	// Only add OpenGL support to non-client programs.
+	if (!FPlatformProperties::RequiresCookedData())
 	{
-		return FilterFeatureLevel(FeatureLevel, TEXT("D3D11MinimumFeatureLevel"), TEXT("D3D11MaximumFeatureLevel"));
+		Config.RHIConfigs[(int32)EWindowsRHI::OpenGL].ShaderPlatforms.Add(SP_OPENGL_PCES3_1);
+		Config.RHIConfigs[(int32)EWindowsRHI::OpenGL].FeatureLevels.Add(ERHIFeatureLevel::ES3_1);
 	}
 
-	if (ChosenRHI == WindowsRHI::Vulkan)
+	if (FWindowsRHIConfig DeprecatedConfig = ParseWindowsRHIConfig(TEXT("TargetedRHIs")); !DeprecatedConfig.ShaderPlatforms.IsEmpty())
 	{
-		return FilterFeatureLevel(FeatureLevel, TEXT("VulkanMinimumFeatureLevel"), TEXT("VulkanMaximumFeatureLevel"));
+		// Since we don't have context here, we have to add this old config setting to every potential RHI.
+		Config.MergeConfig(EWindowsRHI::D3D11,  DeprecatedConfig);
+		Config.MergeConfig(EWindowsRHI::D3D12,  DeprecatedConfig);
+		Config.MergeConfig(EWindowsRHI::Vulkan, DeprecatedConfig);
+		Config.MergeConfig(EWindowsRHI::OpenGL, DeprecatedConfig);
 	}
 
-	return FeatureLevel;
+	return Config;
 }
 
 // Default to Performance Mode on low-end machines
@@ -260,117 +418,84 @@ static bool IsES31D3DOnly()
 	return bES31DXOnly;
 }
 
-static bool AllowD3D12FeatureLevelES31()
+static bool AllowD3D12FeatureLevelES31(const FParsedWindowsDynamicRHIConfig& Config)
 {
 	if (!GIsEditor)
 	{
-		return FilterFeatureLevel(ERHIFeatureLevel::ES3_1, WindowsRHI::D3D12) == ERHIFeatureLevel::ES3_1;
+		return Config.IsFeatureLevelTargeted(EWindowsRHI::D3D12, ERHIFeatureLevel::ES3_1);
 	}
 	return true;
 }
 
 // Choose the default from DefaultGraphicsRHI or TargetedRHIs. DefaultGraphicsRHI has precedence.
-static WindowsRHI ChooseDefaultRHI(const TArray<EShaderPlatform>& TargetedShaderPlatforms)
+static EWindowsRHI ChooseDefaultRHI(const FParsedWindowsDynamicRHIConfig& Config)
 {
-	// Make sure the DDSPI is initialized before we try and read from it
-	FGenericDataDrivenShaderPlatformInfo::Initialize();
-
-	WindowsRHI DefaultRHI = WindowsRHI::D3D11;
-
 	// Default graphics RHI is the main project setting that governs the choice, so it takes the priority
-	FString DefaultGraphicsRHI;
-	if (GConfig->GetString(TEXT("/Script/WindowsTargetPlatform.WindowsTargetSettings"), TEXT("DefaultGraphicsRHI"), DefaultGraphicsRHI, GEngineIni))
+	if (TOptional<EWindowsRHI> ConfigDefault = Config.DefaultRHI)
 	{
-		const FString NAME_DX11(TEXT("DefaultGraphicsRHI_DX11"));
-		const FString NAME_DX12(TEXT("DefaultGraphicsRHI_DX12"));
-		const FString NAME_VULKAN(TEXT("DefaultGraphicsRHI_Vulkan"));
-
-		if (DefaultGraphicsRHI == NAME_DX11)
-		{
-			DefaultRHI = WindowsRHI::D3D11;
-		}
-		else if (DefaultGraphicsRHI == NAME_DX12)
-		{
-			DefaultRHI = WindowsRHI::D3D12;
-		}
-		else if (DefaultGraphicsRHI == NAME_VULKAN)
-		{
-			DefaultRHI = WindowsRHI::Vulkan;
-		}
-		else if (DefaultGraphicsRHI != TEXT("DefaultGraphicsRHI_Default"))
-		{
-			UE_LOG(LogRHI, Error, TEXT("Unrecognized setting '%s' for DefaultGraphicsRHI"), *DefaultGraphicsRHI);
-		}
+		return ConfigDefault.GetValue();
 	}
-	else if (TargetedShaderPlatforms.Num() > 0)
+
+	const EWindowsRHI DefaultRHIOrder[] =
 	{
-		// If we don't have DefaultGraphicsRHI set, try to deduce it from the list of targeted shader platforms
-		// Pick the first one
-		const EShaderPlatform TargetedPlatform = TargetedShaderPlatforms[0];
+		EWindowsRHI::D3D12,
+		EWindowsRHI::D3D11,
+		EWindowsRHI::Vulkan,
+	};
 
-		// not checking D3D as DefaultRHI begins initialized as D3D12
-		if (IsVulkanPlatform(TargetedPlatform))
-		{
-			DefaultRHI = WindowsRHI::Vulkan;
-		}
-		else if (IsOpenGLPlatform(TargetedPlatform))
-		{
-			DefaultRHI = WindowsRHI::OpenGL;
-		}
-		else if (IsD3DSM6Platform(TargetedPlatform))
-		{
-			DefaultRHI = WindowsRHI::D3D12;
-		}
-		else
+	// Find the first RHI with configured support based on the order above
+	for (EWindowsRHI DefaultRHI : DefaultRHIOrder)
+	{
+		if (TOptional<ERHIFeatureLevel::Type> HighestFL = Config.GetHighestSupportedFeatureLevel(DefaultRHI))
 		{
-			DefaultRHI = WindowsRHI::D3D11;
+			return DefaultRHI;
 		}
 	}
 
-	return DefaultRHI;
+	return EWindowsRHI::D3D11;
 }
 
-static TOptional<WindowsRHI> ChoosePreferredRHI(WindowsRHI InDefaultRHI)
+static TOptional<EWindowsRHI> ChoosePreferredRHI(EWindowsRHI InDefaultRHI)
 {
-	TOptional<WindowsRHI> RHIPreference{};
+	TOptional<EWindowsRHI> RHIPreference{};
 
 	// If we are in game, there is a separate setting that can make it prefer D3D12 over D3D11 (but not over other RHIs).
-	if (!GIsEditor && (InDefaultRHI == WindowsRHI::D3D11 || InDefaultRHI == WindowsRHI::D3D12))
+	if (!GIsEditor && (InDefaultRHI == EWindowsRHI::D3D11 || InDefaultRHI == EWindowsRHI::D3D12))
 	{
 		bool bUseD3D12InGame = false;
 		if (GConfig->GetBool(TEXT("D3DRHIPreference"), TEXT("bUseD3D12InGame"), bUseD3D12InGame, GGameUserSettingsIni) && bUseD3D12InGame)
 		{
-			RHIPreference = WindowsRHI::D3D12;
+			RHIPreference = EWindowsRHI::D3D12;
 		}
 	}
 
 	return RHIPreference;
 }
 
-static TOptional<WindowsRHI> ChooseForcedRHI()
+static TOptional<EWindowsRHI> ChooseForcedRHI(const FParsedWindowsDynamicRHIConfig& Config)
 {
-	TOptional<WindowsRHI> ForcedRHI = {};
+	TOptional<EWindowsRHI> ForcedRHI = {};
 
 	// Command line overrides
 	uint32 Sum = 0;
 	if (FParse::Param(FCommandLine::Get(), TEXT("vulkan")))
 	{
-		ForcedRHI = WindowsRHI::Vulkan;
+		ForcedRHI = EWindowsRHI::Vulkan;
 		Sum++;
 	}
 	if (FParse::Param(FCommandLine::Get(), TEXT("opengl")))
 	{
-		ForcedRHI = WindowsRHI::OpenGL;
+		ForcedRHI = EWindowsRHI::OpenGL;
 		Sum++;
 	}
 	if (FParse::Param(FCommandLine::Get(), TEXT("d3d11")) || FParse::Param(FCommandLine::Get(), TEXT("dx11")))
 	{
-		ForcedRHI = WindowsRHI::D3D11;
+		ForcedRHI = EWindowsRHI::D3D11;
 		Sum++;
 	}
 	if (FParse::Param(FCommandLine::Get(), TEXT("d3d12")) || FParse::Param(FCommandLine::Get(), TEXT("dx12")))
 	{
-		ForcedRHI = WindowsRHI::D3D12;
+		ForcedRHI = EWindowsRHI::D3D12;
 		Sum++;
 	}
 
@@ -386,19 +511,19 @@ static TOptional<WindowsRHI> ChooseForcedRHI()
 	// FeatureLevelES31 is also a command line override, so it will determine the underlying RHI unless one is specified
 	if (IsES31D3DOnly() && (FParse::Param(FCommandLine::Get(), TEXT("FeatureLevelES31")) || FParse::Param(FCommandLine::Get(), TEXT("FeatureLevelES3_1"))))
 	{
-		if (ForcedRHI == WindowsRHI::OpenGL)
+		if (ForcedRHI == EWindowsRHI::OpenGL)
 		{
 			FMessageDialog::Open(EAppMsgType::Ok, NSLOCTEXT("WindowsDynamicRHI", "RHIPerformanceOpenGL", "OpenGL is not supported for Performance Mode."));
 			UE_LOG(LogRHI, Fatal, TEXT("OpenGL is not supported for Performance Mode."));
 		}
-		else if (ForcedRHI == WindowsRHI::Vulkan)
+		else if (ForcedRHI == EWindowsRHI::Vulkan)
 		{
 			FMessageDialog::Open(EAppMsgType::Ok, NSLOCTEXT("WindowsDynamicRHI", "RHIPerformanceVulkan", "Vulkan is not supported for Performance Mode."));
 			UE_LOG(LogRHI, Fatal, TEXT("Vulkan is not supported for Performance Mode."));
 		}
-		else if (ForcedRHI == WindowsRHI::D3D12)
+		else if (ForcedRHI == EWindowsRHI::D3D12)
 		{
-			if (!AllowD3D12FeatureLevelES31())
+			if (!AllowD3D12FeatureLevelES31(Config))
 			{
 				FMessageDialog::Open(EAppMsgType::Ok, NSLOCTEXT("WindowsDynamicRHI", "RHIPerformanceDX12", "DirectX 12 is not supported for Performance Mode."));
 				UE_LOG(LogRHI, Fatal, TEXT("DirectX 12 is not supported for Performance Mode."));
@@ -406,7 +531,7 @@ static TOptional<WindowsRHI> ChooseForcedRHI()
 		}
 		else
 		{
-			ForcedRHI = WindowsRHI::D3D11;
+			ForcedRHI = EWindowsRHI::D3D11;
 		}
 	}
 #endif //!WITH_EDITOR && UE_BUILD_SHIPPING
@@ -414,7 +539,7 @@ static TOptional<WindowsRHI> ChooseForcedRHI()
 	return ForcedRHI;
 }
 
-static TOptional<ERHIFeatureLevel::Type> ChooseForcedFeatureLevel(TOptional<WindowsRHI> ChosenRHI, TOptional<WindowsRHI> ForcedRHI)
+static TOptional<ERHIFeatureLevel::Type> GetForcedFeatureLevel()
 {
 	TOptional<ERHIFeatureLevel::Type> ForcedFeatureLevel{};
 
@@ -436,7 +561,7 @@ static TOptional<ERHIFeatureLevel::Type> ChooseForcedFeatureLevel(TOptional<Wind
 	return ForcedFeatureLevel;
 }
 
-static ERHIFeatureLevel::Type ChooseFeatureLevel(TOptional<WindowsRHI> ChosenRHI, TOptional<ERHIFeatureLevel::Type> ForcedFeatureLevel, const TArray<EShaderPlatform>& TargetedShaderPlatforms)
+static ERHIFeatureLevel::Type ChooseFeatureLevel(EWindowsRHI ChosenRHI, const TOptional<EWindowsRHI> ForcedRHI, const TOptional<ERHIFeatureLevel::Type> ForcedFeatureLevel, const FParsedWindowsDynamicRHIConfig& Config)
 {
 	if (ForcedFeatureLevel)
 	{
@@ -447,44 +572,123 @@ static ERHIFeatureLevel::Type ChooseFeatureLevel(TOptional<WindowsRHI> ChosenRHI
 		}
 
 		// Make sure the feature level is supported by the runtime, otherwise fall back to the default
-		const ERHIFeatureLevel::Type FilteredForcedFeatureLevel = FilterFeatureLevel(ForcedFeatureLevel.GetValue(), ChosenRHI.GetValue());
-		if (FilteredForcedFeatureLevel == ForcedFeatureLevel)
+		if (Config.IsFeatureLevelTargeted(ChosenRHI, ForcedFeatureLevel.GetValue()))
 		{
-			return FilteredForcedFeatureLevel;
+			return ForcedFeatureLevel.GetValue();
 		}
 	}
 
-	ERHIFeatureLevel::Type FeatureLevel;
+	TOptional<ERHIFeatureLevel::Type> FeatureLevel{};
 
-	if (ChosenRHI == WindowsRHI::OpenGL)
-	{
-		// OpenGL can only be used for mobile preview
-		FeatureLevel = ERHIFeatureLevel::ES3_1;
-	}
-	else if (ChosenRHI == WindowsRHI::D3D12)
+	if ((ChosenRHI == EWindowsRHI::D3D11 || ChosenRHI == EWindowsRHI::D3D12) && Config.IsFeatureLevelTargeted(ChosenRHI, ERHIFeatureLevel::ES3_1) && PreferFeatureLevelES31())
 	{
-		FeatureLevel = IsD3DSM6PlatformTargeted(TargetedShaderPlatforms) ? ERHIFeatureLevel::SM6 : ERHIFeatureLevel::SM5;
+		FeatureLevel = TOptional<ERHIFeatureLevel::Type>(ERHIFeatureLevel::ES3_1);
 	}
 	else
 	{
-		FeatureLevel = ERHIFeatureLevel::SM5;
+		FeatureLevel = Config.GetHighestSupportedFeatureLevel(ChosenRHI);
 	}
 
-	ERHIFeatureLevel::Type FilteredFeatureLevel = FilterFeatureLevel(FeatureLevel, ChosenRHI.GetValue());
+	// If we were forced to a specific RHI while not forced to a specific feature level and the project isn't configured for it, find the default Feature Level for that RHI
+	if (!FeatureLevel && ForcedRHI)
+	{
+		FeatureLevel = GetDefaultFeatureLevelForRHI(ChosenRHI);
+
+		if (FPlatformProperties::RequiresCookedData())
+		{
+			const TCHAR* RHIName = ModuleNameFromWindowsRHI(ForcedRHI.GetValue());
+			const FString FeatureLevelName = LexToString(FeatureLevel.GetValue());
+			UE_LOG(LogRHI, Warning, TEXT("User requested RHI '%s' but that is not supported by this project's data. Defaulting to Feature Level '%s'."), RHIName, *FeatureLevelName);
+		}
+	}
 
 	// If the user wanted to force a feature level and we couldn't set it, log out why and what we're actually running with
 	if (ForcedFeatureLevel)
 	{
 		const FString ForcedName = LexToString(ForcedFeatureLevel.GetValue());
-		const FString UsedName = LexToString(FilteredFeatureLevel);
+		const FString UsedName = LexToString(FeatureLevel.GetValue());
 		UE_LOG(LogRHI, Warning, TEXT("User requested Feature Level '%s' but that is not supported by this project. Falling back to Feature Level '%s'."), *ForcedName, *UsedName);
 	}
 
-	return FilteredFeatureLevel;
+	return FeatureLevel.GetValue();
+}
+
+static bool HandleUnsupportedFeatureLevel(EWindowsRHI& WindowsRHI, ERHIFeatureLevel::Type& FeatureLevel, TOptional<ERHIFeatureLevel::Type> ForcedFeatureLevel, const FParsedWindowsDynamicRHIConfig& Config)
+{
+	if (ForcedFeatureLevel)
+	{
+		if (WindowsRHI == EWindowsRHI::D3D12 && FeatureLevel == ERHIFeatureLevel::SM6)
+		{
+			FMessageDialog::Open(EAppMsgType::Ok, NSLOCTEXT("WindowsDynamicRHI", "RequiredDX12SM6", "DX12 SM6 is not supported on your system. Try running without the -sm6 command line argument."));
+			FPlatformMisc::RequestExit(1);
+		}
+
+		return false;
+	}
+
+	if (TOptional<ERHIFeatureLevel::Type> FallbackFeatureLevel = Config.GetNextHighestTargetedFeatureLevel(WindowsRHI, FeatureLevel))
+	{
+		UE_LOG(LogRHI, Log, TEXT("RHI %s with Feature Level %s not supported, falling back to Feature Level %s"), ModuleNameFromWindowsRHI(WindowsRHI), *LexToString(FeatureLevel), *LexToString(FallbackFeatureLevel.GetValue()));
+
+		FeatureLevel = FallbackFeatureLevel.GetValue();
+		return true;
+	}
+
+	return false;
+}
+
+static bool HandleUnsupportedRHI(EWindowsRHI& WindowsRHI, ERHIFeatureLevel::Type& FeatureLevel, TOptional<EWindowsRHI> ForcedRHI, const FParsedWindowsDynamicRHIConfig& Config)
+{
+	if (ForcedRHI)
+	{
+		if (ForcedRHI == EWindowsRHI::D3D12)
+		{
+			FMessageDialog::Open(EAppMsgType::Ok, NSLOCTEXT("WindowsDynamicRHI", "RequiredDX12", "DX12 is not supported on your system. Try running without the -dx12 or -d3d12 command line argument."));
+			FPlatformMisc::RequestExit(1);
+		}
+	}
+
+	if (WindowsRHI == EWindowsRHI::D3D12)
+	{
+		if (TOptional<ERHIFeatureLevel::Type> D3D11FeatureLevel = Config.GetHighestSupportedFeatureLevel(EWindowsRHI::D3D11))
+		{
+			UE_LOG(LogRHI, Log, TEXT("D3D12 is not supported, falling back to D3D11 with Feature Level %s"), *LexToString(D3D11FeatureLevel.GetValue()));
+
+			WindowsRHI = EWindowsRHI::D3D11;
+			FeatureLevel = D3D11FeatureLevel.GetValue();
+			return true;
+		}
+
+		FMessageDialog::Open(EAppMsgType::Ok, NSLOCTEXT("WindowsDynamicRHI", "RequiredDX12", "DX12 is not supported on your system. Try running without the -dx12 or -d3d12 command line argument."));
+		FPlatformMisc::RequestExit(1);
+	}
+
+	if (WindowsRHI == EWindowsRHI::D3D11)
+	{
+		FMessageDialog::Open(EAppMsgType::Ok, NSLOCTEXT("WindowsDynamicRHI", "RequiredDX11Feature_11_SM5", "A D3D11-compatible GPU (Feature Level 11.0, Shader Model 5.0) is required to run the engine."));
+		FPlatformMisc::RequestExit(1);
+	}
+
+	if (WindowsRHI == EWindowsRHI::Vulkan)
+	{
+		FMessageDialog::Open(EAppMsgType::Ok, NSLOCTEXT("WindowsDynamicRHI", "RequiredVulkan", "Vulkan Driver is required to run the engine."));
+		FPlatformMisc::RequestExit(1);
+	}
+
+	if (WindowsRHI == EWindowsRHI::OpenGL)
+	{
+		FMessageDialog::Open(EAppMsgType::Ok, NSLOCTEXT("WindowsDynamicRHI", "RequiredOpenGL", "OpenGL 4.3 is required to run the engine."));
+		FPlatformMisc::RequestExit(1);
+	}
+
+	return false;
 }
 
 static IDynamicRHIModule* LoadDynamicRHIModule(ERHIFeatureLevel::Type& DesiredFeatureLevel, const TCHAR*& LoadedRHIModuleName)
 {
+	// Make sure the DDSPI is initialized before we try and read from it
+	FGenericDataDrivenShaderPlatformInfo::Initialize();
+
 	bool bUseGPUCrashDebugging = false;
 	if (!GIsEditor && GConfig->GetBool(TEXT("D3DRHIPreference"), TEXT("bUseGPUCrashDebugging"), bUseGPUCrashDebugging, GGameUserSettingsIni))
 	{
@@ -492,17 +696,17 @@ static IDynamicRHIModule* LoadDynamicRHIModule(ERHIFeatureLevel::Type& DesiredFe
 		*GPUCrashDebuggingCVar = bUseGPUCrashDebugging;
 	}
 
-	const TArray<EShaderPlatform> TargetedShaderPlatforms = GetTargetedShaderPlatforms();
+	const FParsedWindowsDynamicRHIConfig Config = ParseWindowsDynamicRHIConfig();
 
 	// RHI is chosen by the project settings (first DefaultGraphicsRHI, then TargetedRHIs are consulted, "Default" maps to D3D12). 
 	// After this, a separate game-only setting (does not affect editor) bPreferD3D12InGame selects between D3D12 or D3D11 (but will not have any effect if Vulkan or OpenGL are chosen).
 	// Commandline switches apply after this and can force an arbitrary RHIs. If RHI isn't supported, the game will refuse to start.
 
-	WindowsRHI DefaultRHI = ChooseDefaultRHI(TargetedShaderPlatforms);
-	const TOptional<WindowsRHI> PreferredRHI = ChoosePreferredRHI(DefaultRHI);
-	const TOptional<WindowsRHI> ForcedRHI = ChooseForcedRHI();
+	EWindowsRHI DefaultRHI = ChooseDefaultRHI(Config);
+	const TOptional<EWindowsRHI> PreferredRHI = ChoosePreferredRHI(DefaultRHI);
+	const TOptional<EWindowsRHI> ForcedRHI = ChooseForcedRHI(Config);
 
-	WindowsRHI ChosenRHI = DefaultRHI;
+	EWindowsRHI ChosenRHI = DefaultRHI;
 	if (ForcedRHI)
 	{
 		ChosenRHI = ForcedRHI.GetValue();
@@ -512,113 +716,37 @@ static IDynamicRHIModule* LoadDynamicRHIModule(ERHIFeatureLevel::Type& DesiredFe
 		ChosenRHI = PreferredRHI.GetValue();
 	}
 
-	TOptional<ERHIFeatureLevel::Type> ForcedFeatureLevel = ChooseForcedFeatureLevel(ChosenRHI, ForcedRHI);
-	DesiredFeatureLevel = ChooseFeatureLevel(ChosenRHI, ForcedFeatureLevel, TargetedShaderPlatforms);
+	const TOptional<ERHIFeatureLevel::Type> ForcedFeatureLevel = GetForcedFeatureLevel();
+
+	DesiredFeatureLevel = ChooseFeatureLevel(ChosenRHI, ForcedRHI, ForcedFeatureLevel, Config);
 
 	// Load the dynamic RHI module.
-	IDynamicRHIModule* DynamicRHIModule = nullptr;
 
-	if (ChosenRHI == WindowsRHI::OpenGL)
+	bool bTryWithNewConfig = false;
+	do
 	{
-		FApp::SetGraphicsRHI(TEXT("OpenGL"));
-		const TCHAR* OpenGLRHIModuleName = TEXT("OpenGLDrv");
-		DynamicRHIModule = &FModuleManager::LoadModuleChecked<IDynamicRHIModule>(OpenGLRHIModuleName);
+		const FString RHIName = GetRHINameFromWindowsRHI(ChosenRHI, DesiredFeatureLevel);
+		FApp::SetGraphicsRHI(RHIName);
 
-		if (!DynamicRHIModule->IsSupported(DesiredFeatureLevel))
-		{
-			FMessageDialog::Open(EAppMsgType::Ok, NSLOCTEXT("WindowsDynamicRHI", "RequiredOpenGL", "OpenGL 4.3 is required to run the engine."));
-			FPlatformMisc::RequestExit(1);
-			DynamicRHIModule = nullptr;
-		}
+		const TCHAR* ModuleName = ModuleNameFromWindowsRHI(ChosenRHI);
+		IDynamicRHIModule* DynamicRHIModule = FModuleManager::LoadModulePtr<IDynamicRHIModule>(ModuleName);
 
-		LoadedRHIModuleName = OpenGLRHIModuleName;
-	}
-	else if (ChosenRHI == WindowsRHI::Vulkan)
-	{
-		FApp::SetGraphicsRHI(TEXT("Vulkan"));
-		const TCHAR* VulkanRHIModuleName = TEXT("VulkanRHI");
-		DynamicRHIModule = &FModuleManager::LoadModuleChecked<IDynamicRHIModule>(VulkanRHIModuleName);
-
-		if (!DynamicRHIModule->IsSupported(DesiredFeatureLevel))
+		if (DynamicRHIModule && DynamicRHIModule->IsSupported(DesiredFeatureLevel))
 		{
-			FMessageDialog::Open(EAppMsgType::Ok, NSLOCTEXT("WindowsDynamicRHI", "RequiredVulkan", "Vulkan Driver is required to run the engine."));
-			FPlatformMisc::RequestExit(1);
-			DynamicRHIModule = nullptr;
+			LoadedRHIModuleName = ModuleName;
+			return DynamicRHIModule;
 		}
-		LoadedRHIModuleName = VulkanRHIModuleName;
-	}
-	else if (ChosenRHI == WindowsRHI::D3D12)
-	{
-		FString FeatureLevelName;
-		GetFeatureLevelName(DesiredFeatureLevel, FeatureLevelName);
-
-		const FString RHIDisplayName = FString::Printf(TEXT("DirectX 12 (%s)"), *FeatureLevelName);
-		FApp::SetGraphicsRHI(RHIDisplayName);
-
-		LoadedRHIModuleName = TEXT("D3D12RHI");
-		DynamicRHIModule = FModuleManager::LoadModulePtr<IDynamicRHIModule>(LoadedRHIModuleName);
-
-		bool bD3D12Supported = DynamicRHIModule && DynamicRHIModule->IsSupported(DesiredFeatureLevel);
-
-		// Fallback to SM5 if SM6 is not supported
-		if (!bD3D12Supported && DynamicRHIModule && DesiredFeatureLevel == ERHIFeatureLevel::SM6)
-		{
-			if (ForcedFeatureLevel)
-			{
-				FMessageDialog::Open(EAppMsgType::Ok, NSLOCTEXT("WindowsDynamicRHI", "RequiredDX12SM6", "DX12 SM6 is not supported on your system. Try running without the -sm6 command line argument."));
-				FPlatformMisc::RequestExit(1);
-			}
 
-			if (FilterFeatureLevel(ERHIFeatureLevel::SM5, WindowsRHI::D3D12) == ERHIFeatureLevel::SM5)
-			{
-				UE_LOG(LogRHI, Log, TEXT("D3D12 SM6 is not supported, trying SM5"));
-
-				DesiredFeatureLevel = ERHIFeatureLevel::SM5;
-				bD3D12Supported = DynamicRHIModule->IsSupported(DesiredFeatureLevel);
-			}
-			
-			if (!bD3D12Supported)
-			{
-				UE_LOG(LogRHI, Log, TEXT("D3D12 SM5 is not supported, trying D3D11"));
-			}
-		}
+		bTryWithNewConfig = HandleUnsupportedFeatureLevel(ChosenRHI, DesiredFeatureLevel, ForcedFeatureLevel, Config);
 
-		if (!bD3D12Supported)
+		if (!bTryWithNewConfig)
 		{
-			if (ForcedRHI == WindowsRHI::D3D12)
-			{
-				FMessageDialog::Open(EAppMsgType::Ok, NSLOCTEXT("WindowsDynamicRHI", "RequiredDX12", "DX12 is not supported on your system. Try running without the -dx12 or -d3d12 command line argument."));
-				FPlatformMisc::RequestExit(1);
-			}
-			if (DynamicRHIModule)
-			{
-				FModuleManager::Get().UnloadModule(LoadedRHIModuleName);
-			}
-			DynamicRHIModule = nullptr;
-			LoadedRHIModuleName = nullptr;
+			bTryWithNewConfig = HandleUnsupportedRHI(ChosenRHI, DesiredFeatureLevel, ForcedRHI, Config);
 		}
-	}
 
-	// Fallback to D3D11RHI if nothing is selected
-	if (!DynamicRHIModule)
-	{
-		// We need to filter one more time to handle the D3D12 fallback
-		DesiredFeatureLevel = FilterFeatureLevel(DesiredFeatureLevel, WindowsRHI::D3D11);
-
-		FApp::SetGraphicsRHI(TEXT("DirectX 11"));
-		const TCHAR* D3D11RHIModuleName = TEXT("D3D11RHI");
-		DynamicRHIModule = &FModuleManager::LoadModuleChecked<IDynamicRHIModule>(D3D11RHIModuleName);
+	} while (bTryWithNewConfig);
 
-		if (!DynamicRHIModule->IsSupported(DesiredFeatureLevel))
-		{
-			FMessageDialog::Open(EAppMsgType::Ok, NSLOCTEXT("WindowsDynamicRHI", "RequiredDX11Feature_11_SM5", "A D3D11-compatible GPU (Feature Level 11.0, Shader Model 5.0) is required to run the engine."));
-			FPlatformMisc::RequestExit(1);
-			DynamicRHIModule = nullptr;
-		}
-
-		LoadedRHIModuleName = D3D11RHIModuleName;
-	}
-	return DynamicRHIModule;
+	return nullptr;
 }
 
 FDynamicRHI* PlatformCreateDynamicRHI()
