Implemented ShouldCompileRayTracingCallableShadersForProject(ShaderPlatform)
- added bSupportsRayTracingCallableShaders to DDPI.
- only Windows D3D12 RHI supports callable shaders at the moment.

diff --git a/Engine/Config/Windows/DataDrivenPlatformInfo.ini b/Engine/Config/Windows/DataDrivenPlatformInfo.ini
--- a/Engine/Config/Windows/DataDrivenPlatformInfo.ini
+++ b/Engine/Config/Windows/DataDrivenPlatformInfo.ini
@@ -62,6 +62,7 @@ bSupportsDebugViewShaders=true
 bSupportsRayTracing=true
 bSupportsRayTracingShaders=true
 bSupportsInlineRayTracing=false
+bSupportsRayTracingCallableShaders=true
 bSupportsRayTracingProceduralPrimitive=true
 bSupportsRayTracingIndirectInstanceData=true
 bSupportsPathTracing=true
@@ -107,6 +108,7 @@ bSupportsDebugViewShaders=true
 bSupportsRayTracing=true
 bSupportsRayTracingShaders=true
 bSupportsInlineRayTracing=false
+bSupportsRayTracingCallableShaders=true
 bSupportsRayTracingProceduralPrimitive=true
 bSupportsRayTracingIndirectInstanceData=true
 bSupportsPathTracing=true
diff --git a/Engine/Source/Runtime/RHI/Private/RHI.cpp b/Engine/Source/Runtime/RHI/Private/RHI.cpp
--- a/Engine/Source/Runtime/RHI/Private/RHI.cpp
+++ b/Engine/Source/Runtime/RHI/Private/RHI.cpp
@@ -2438,6 +2438,7 @@ void FGenericDataDrivenShaderPlatformInfo::ParseDataDrivenShaderInfo(const FConf
 	GET_SECTION_BOOL_HELPER(bSupportsRayTracing);
 	GET_SECTION_BOOL_HELPER(bSupportsRayTracingShaders);
 	GET_SECTION_BOOL_HELPER(bSupportsInlineRayTracing);
+	GET_SECTION_BOOL_HELPER(bSupportsRayTracingCallableShaders);
 	GET_SECTION_BOOL_HELPER(bSupportsRayTracingProceduralPrimitive);
 	GET_SECTION_BOOL_HELPER(bSupportsRayTracingTraversalStatistics);
 	GET_SECTION_BOOL_HELPER(bSupportsRayTracingIndirectInstanceData);
diff --git a/Engine/Source/Runtime/RHI/Public/RHI.h b/Engine/Source/Runtime/RHI/Public/RHI.h
--- a/Engine/Source/Runtime/RHI/Public/RHI.h
+++ b/Engine/Source/Runtime/RHI/Public/RHI.h
@@ -228,6 +228,14 @@ inline RHI_API bool RHISupportsInlineRayTracing(const FStaticShaderPlatform Plat
 	return FDataDrivenShaderPlatformInfo::GetSupportsInlineRayTracing(Platform);
 }
 
+/** Whether this platform can compile ray tracing callable shaders.
+ *  To use at runtime, also check GRHISupportsRayTracing and r.RayTracing CVar (see IsRayTracingEnabled() helper).
+ **/
+inline RHI_API bool RHISupportsRayTracingCallableShaders(const FStaticShaderPlatform Platform)
+{
+	return FDataDrivenShaderPlatformInfo::GetSupportsRayTracingCallableShaders(Platform);
+}
+
 /** Can this platform compile mesh shaders with tier0 capability.
  *  To use at runtime, also check GRHISupportsMeshShadersTier0.
  **/
diff --git a/Engine/Source/Runtime/RHI/Public/RHIDefinitions.h b/Engine/Source/Runtime/RHI/Public/RHIDefinitions.h
--- a/Engine/Source/Runtime/RHI/Public/RHIDefinitions.h
+++ b/Engine/Source/Runtime/RHI/Public/RHIDefinitions.h
@@ -356,6 +356,7 @@ class RHI_API FGenericDataDrivenShaderPlatformInfo
 	uint32 bSupports4ComponentUAVReadWrite: 1;
 	uint32 bSupportsRenderTargetWriteMask: 1;
 	uint32 bSupportsRayTracing: 1;
+	uint32 bSupportsRayTracingCallableShaders : 1;
 	uint32 bSupportsRayTracingProceduralPrimitive : 1;
 	uint32 bSupportsRayTracingTraversalStatistics : 1;
 	uint32 bSupportsRayTracingIndirectInstanceData : 1; // Whether instance transforms can be copied from the GPU to the TLAS instances buffer
@@ -616,6 +617,12 @@ public:
 		return Infos[Platform].bSupportsRayTracing && Infos[Platform].bSupportsInlineRayTracing;
 	}
 
+	static FORCEINLINE_DEBUGGABLE const bool GetSupportsRayTracingCallableShaders(const FStaticShaderPlatform Platform)
+	{
+		check(IsValid(Platform));
+		return Infos[Platform].bSupportsRayTracing && Infos[Platform].bSupportsRayTracingCallableShaders;
+	}
+
 	static FORCEINLINE_DEBUGGABLE const bool GetSupportsRayTracingProceduralPrimitive(const FStaticShaderPlatform Platform)
 	{
 		check(IsValid(Platform));
diff --git a/Engine/Source/Runtime/RenderCore/Public/RenderResource.h b/Engine/Source/Runtime/RenderCore/Public/RenderResource.h
--- a/Engine/Source/Runtime/RenderCore/Public/RenderResource.h
+++ b/Engine/Source/Runtime/RenderCore/Public/RenderResource.h
@@ -687,6 +687,11 @@ FORCEINLINE bool ShouldCompileRayTracingShadersForProject(EShaderPlatform Shader
 	}
 }
 
+FORCEINLINE bool ShouldCompileRayTracingCallableShadersForProject(EShaderPlatform ShaderPlatform)
+{
+	return RHISupportsRayTracingCallableShaders(ShaderPlatform) && ShouldCompileRayTracingShadersForProject(ShaderPlatform);
+}
+
 // Returns `true` when running on RT-capable machine, RT support is enabled for the project and by game graphics options.
 // This function may only be called at runtime, never during cooking.
 extern RENDERCORE_API bool IsRayTracingEnabled();
