removed `check(!IsInRHIThread())` as with "task retraction" globally enabled, the task that executes `FRHICommandListBase::HandleRTThreadTaskCompletion` is still "any thread" task but as RHIT waits for its completion, if the task hasn't started execution yet, RHIT can "retract" it and execute locally

diff --git a/Engine/Source/Runtime/RHI/Private/RHICommandList.cpp b/Engine/Source/Runtime/RHI/Private/RHICommandList.cpp
--- a/Engine/Source/Runtime/RHI/Private/RHICommandList.cpp
+++ b/Engine/Source/Runtime/RHI/Private/RHICommandList.cpp
@@ -2144,7 +2144,7 @@
 DECLARE_CYCLE_STAT(TEXT("RTTask completion join"), STAT_HandleRTThreadTaskCompletion_Join, STATGROUP_RHICMDLIST);
 void FRHICommandListBase::HandleRTThreadTaskCompletion(const FGraphEventRef& MyCompletionGraphEvent)
 {
-	check(!IsImmediate() && !IsInRHIThread());
+	check(!IsImmediate());
 	for (int32 Index = 0; Index < RTTasks.Num(); Index++)
 	{
 		if (RTTasks[Index].IsValid() && !RTTasks[Index]->IsComplete())
