Increase the size of the GPU profiler event buffer because in some cases Slate pushes a lot of events due to lots of windows being rendered.

#rnx

diff --git a/Engine/Source/Runtime/RHI/Private/GpuProfilerTrace.cpp b/Engine/Source/Runtime/RHI/Private/GpuProfilerTrace.cpp
--- a/Engine/Source/Runtime/RHI/Private/GpuProfilerTrace.cpp
+++ b/Engine/Source/Runtime/RHI/Private/GpuProfilerTrace.cpp
@@ -16,7 +16,7 @@ namespace GpuProfilerTrace
 
 static TAutoConsoleVariable<int32> CVarGpuProfilerMaxEventBufferSizeKB(
 	TEXT("r.GpuProfilerMaxEventBufferSizeKB"),
-	16,
+	32,
 	TEXT("Size of the scratch buffer in kB."),
 	ECVF_Default);
 
