Added option to force bad driver warnings to trigger a fatal error by using a CVar on non-shipping builds.

diff --git a/Engine/Source/Runtime/RHI/Private/DynamicRHI.cpp b/Engine/Source/Runtime/RHI/Private/DynamicRHI.cpp
--- a/Engine/Source/Runtime/RHI/Private/DynamicRHI.cpp
+++ b/Engine/Source/Runtime/RHI/Private/DynamicRHI.cpp
@@ -47,6 +47,15 @@ static TAutoConsoleVariable<int32> CVarWarnOfBadDrivers(
 	ECVF_RenderThreadSafe
 	);
 
+static TAutoConsoleVariable<int32> CVarBadDriverWarningIsFatal(
+	TEXT("r.BadDriverWarningIsFatal"),
+	0,
+	TEXT("If non-zero, trigger a fatal error when warning of bad drivers.\n")
+	TEXT("For the fatal error to occur, r.WarnOfBadDrivers must be non-zero.\n")
+	TEXT(" 0: off (default)\n")
+	TEXT(" 1: a fatal error occurs after the out of date driver message is dismissed\n"),
+	ECVF_RenderThreadSafe);
+
 static TAutoConsoleVariable<int32> CVarDisableDriverWarningPopupIfGFN(
 	TEXT("r.DisableDriverWarningPopupIfGFN"),
 	1,
@@ -242,6 +251,14 @@ static void RHIDetectAndWarnOfBadDrivers(bool bHasEditorToken)
 			{
 				FPlatformProcess::LaunchURL(*HyperlinkText.ToString(), nullptr, nullptr);
 			}
+#if !UE_BUILD_SHIPPING
+			if (CVarBadDriverWarningIsFatal.GetValueOnGameThread())
+			{
+				// Force a fatal error depending on CVar
+				UE_LOG(LogRHI, Fatal, TEXT("Fatal crash requested when graphics drivers are out of date.\n")
+					TEXT("To prevent this crash, please update drivers."));
+			}
+#endif
 		}
 	}
 }
@@ -261,6 +278,15 @@ static void RHIDetectAndWarnOfBadDrivers(bool bHasEditorToken)
 		FPlatformMisc::MessageBoxExt(EAppMsgType::Ok,
 									 *NSLOCTEXT("MessageDialog", "UpdateMacOSX_Body", "Please update to the latest version of macOS for best performance and stability.").ToString(),
 									 *NSLOCTEXT("MessageDialog", "UpdateMacOSX_Title", "Update macOS").ToString());
+		
+#if !UE_BUILD_SHIPPING
+		if (CVarBadDriverWarningIsFatal.GetValueOnGameThread())
+		{
+			// Force a fatal error depending on CVar
+			UE_LOG(LogRHI, Fatal, TEXT("Fatal crash requested when graphics drivers are out of date.\n")
+				TEXT("To prevent this crash, please update macOS."));
+		}
+#endif
 	}
 }
 #endif // PLATFORM_WINDOWS
