[Backout] - CL20586117
Original CL Desc
-----------------------------------------------------------------
Added option to force bad driver warnings to trigger a fatal error by using a CVar on non-shipping builds.

diff --git a/Engine/Source/Runtime/RHI/Private/DynamicRHI.cpp b/Engine/Source/Runtime/RHI/Private/DynamicRHI.cpp
--- a/Engine/Source/Runtime/RHI/Private/DynamicRHI.cpp
+++ b/Engine/Source/Runtime/RHI/Private/DynamicRHI.cpp
@@ -14,7 +14,6 @@
 #include "GenericPlatform/GenericPlatformDriver.h"
 #include "GenericPlatform/GenericPlatformCrashContext.h"
 #include "PipelineStateCache.h"
-#include "TextureProfiler.h"
 
 #if defined(NV_GEFORCENOW) && NV_GEFORCENOW
 #include "GeForceNOWWrapper.h"
@@ -23,8 +22,9 @@
 IMPLEMENT_TYPE_LAYOUT(FRayTracingGeometryInitializer);
 IMPLEMENT_TYPE_LAYOUT(FRayTracingGeometrySegment);
 
-static_assert(sizeof(FRayTracingGeometryInstance) <= 104,
-	"Ray tracing instance descriptor is expected to be no more than 104 bytes, as there may be a very large number of them.");
+static_assert(sizeof(FRayTracingGeometryInstance) <= 96,
+	"Ray tracing instance descriptor is expected to be no more than 96 bytes, "
+	"as there may be a very large number of them.");
 
 #ifndef PLATFORM_ALLOW_NULL_RHI
 	#define PLATFORM_ALLOW_NULL_RHI		0
@@ -87,15 +87,13 @@ void InitNullRHI()
 	FGenericCrashContext::SetEngineData(TEXT("RHI.RHIName"), TEXT("NullRHI"));
 }
 
-#if PLATFORM_WINDOWS || PLATFORM_UNIX
+#if PLATFORM_WINDOWS
 static void RHIDetectAndWarnOfBadDrivers(bool bHasEditorToken)
 {
 	int32 CVarValue = CVarWarnOfBadDrivers.GetValueOnGameThread();
 
-	if(!GIsRHIInitialized || !CVarValue || (GRHIVendorId == 0) || FApp::IsUnattended())
+	if(!GIsRHIInitialized || !CVarValue || GRHIVendorId == 0)
 	{
-		UE_LOG(LogRHI, Log, TEXT("Skipping Driver Check: RHI%s initialized, WarnOfBadDrivers=%d, VendorId=0x%x, is%s unattended"), 
-			GIsRHIInitialized ? TEXT("") : TEXT(" NOT"), CVarValue, GRHIVendorId, FApp::IsUnattended() ? TEXT("") : TEXT(" NOT"));
 		return;
 	}
 
@@ -226,31 +224,27 @@ static void RHIDetectAndWarnOfBadDrivers(bool bHasEditorToken)
 			{
 				if (!DenyListEntry.RHIName.IsEmpty())
 				{
-					LocalizedMsg = FText::Format(NSLOCTEXT("MessageDialog", "LatestVideoCardDriverRHIIssueReport", "The latest version of the {Vendor} graphics driver has known issues in {RHI}.\nPlease install the recommended driver version or switch to a different rendering API.\n\nWould you like to visit the following URL to download the driver?\n\n{Hyperlink}\n\n{AdapterName}\nInstalled: {InstalledVer}\nRecommended: {RecommendedVer}"), Args);
+					LocalizedMsg = FText::Format(NSLOCTEXT("MessageDialog", "LatestVideoCardDriverRHIIssueReport", "The latest version of the {Vendor} graphics driver has known issues in {RHI}.\nPlease install the recommended driver version or switch to a different rendering API.\n\n{Hyperlink}\n\n{AdapterName}\nInstalled: {InstalledVer}\nRecommended: {RecommendedVer}"), Args);
 				}
 				else
 				{
-					LocalizedMsg = FText::Format(NSLOCTEXT("MessageDialog", "LatestVideoCardDriverIssueReport", "The latest version of the {Vendor} graphics driver has known issues.\nPlease install the recommended driver version.\n\nWould you like to visit the following URL to download the driver?\n\n{Hyperlink}\n\n{AdapterName}\nInstalled: {InstalledVer}\nRecommended: {RecommendedVer}"), Args);
+					LocalizedMsg = FText::Format(NSLOCTEXT("MessageDialog", "LatestVideoCardDriverIssueReport", "The latest version of the {Vendor} graphics driver has known issues.\nPlease install the recommended driver version.\n\n{Hyperlink}\n\n{AdapterName}\nInstalled: {InstalledVer}\nRecommended: {RecommendedVer}"), Args);
 				}
 			}
 			else
 			{
 				if (!DenyListEntry.RHIName.IsEmpty())
 				{
-					LocalizedMsg = FText::Format(NSLOCTEXT("MessageDialog", "VideoCardDriverRHIIssueReport", "The installed version of the {Vendor} graphics driver has known issues in {RHI}.\nPlease install either the latest or the recommended driver version or switch to a different rendering API.\n\nWould you like to visit the following URL to download the driver?\n\n{Hyperlink}\n\n{AdapterName}\nInstalled: {InstalledVer}\nRecommended: {RecommendedVer}"), Args);
+					LocalizedMsg = FText::Format(NSLOCTEXT("MessageDialog", "VideoCardDriverRHIIssueReport", "The installed version of the {Vendor} graphics driver has known issues in {RHI}.\nPlease install either the latest or the recommended driver version or switch to a different rendering API.\n\n{Hyperlink}\n\n{AdapterName}\nInstalled: {InstalledVer}\nRecommended: {RecommendedVer}"), Args);
 				}
 				else
 				{
-					LocalizedMsg = FText::Format(NSLOCTEXT("MessageDialog", "VideoCardDriverIssueReport", "The installed version of the {Vendor} graphics driver has known issues.\nPlease install either the latest or the recommended driver version.\n\nWould you like to visit the following URL to download the driver?\n\n{Hyperlink}\n\n{AdapterName}\nInstalled: {InstalledVer}\nRecommended: {RecommendedVer}"), Args);
+					LocalizedMsg = FText::Format(NSLOCTEXT("MessageDialog", "VideoCardDriverIssueReport", "The installed version of the {Vendor} graphics driver has known issues.\nPlease install either the latest or the recommended driver version.\n\n{Hyperlink}\n\n{AdapterName}\nInstalled: {InstalledVer}\nRecommended: {RecommendedVer}"), Args);
 				}
 			}
 
 			FText Title = NSLOCTEXT("MessageDialog", "TitleVideoCardDriverIssue", "WARNING: Known issues with graphics driver");
-			EAppReturnType::Type Response = FMessageDialog::Open(EAppMsgType::YesNo, LocalizedMsg, &Title);
-			if (Response == EAppReturnType::Yes)
-			{
-				FPlatformProcess::LaunchURL(*HyperlinkText.ToString(), nullptr, nullptr);
-			}
+			FMessageDialog::Open(EAppMsgType::Ok, LocalizedMsg, &Title);
 #if !UE_BUILD_SHIPPING
 			if (CVarBadDriverWarningIsFatal.GetValueOnGameThread())
 			{
@@ -352,10 +346,6 @@ void RHIInit(bool bHasEditorToken)
 				FGenericCrashContext::SetEngineData(TEXT("RHI.FeatureLevel"), FeatureLevelString);
 				FGenericCrashContext::SetEngineData(TEXT("RHI.GPUVendor"), RHIVendorIdToString());
 				FGenericCrashContext::SetEngineData(TEXT("RHI.DeviceId"), FString::Printf(TEXT("%04X"), GRHIDeviceId));
-
-#if TEXTURE_PROFILER_ENABLED
-				FTextureProfiler::Get()->Init();
-#endif
 			}
 #if PLATFORM_ALLOW_NULL_RHI
 			else
@@ -369,7 +359,7 @@ void RHIInit(bool bHasEditorToken)
 		check(GDynamicRHI);
 	}
 
-#if PLATFORM_WINDOWS || PLATFORM_MAC || PLATFORM_UNIX
+#if PLATFORM_WINDOWS || PLATFORM_MAC
 #if defined(NV_GEFORCENOW) && NV_GEFORCENOW
 	bool bDetectAndWarnBadDrivers = true;
 	if (IsRHIDeviceNVIDIA() && !!CVarDisableDriverWarningPopupIfGFN.GetValueOnAnyThread())
@@ -436,9 +426,6 @@ void RHIExit()
 		FRHICommandListExecutor::GetImmediateCommandList().ImmediateFlush(EImmediateFlushType::FlushRHIThreadFlushResources);
 		FRHICommandListExecutor::GetImmediateCommandList().ImmediateFlush(EImmediateFlushType::FlushRHIThread);
 	}
-
-	void CleanupRHICommandListGraphEvents();
-	CleanupRHICommandListGraphEvents();
 }
 
 
@@ -538,20 +525,10 @@ void FDynamicRHI::RHITransferBufferUnderlyingResource(FRHIBuffer* DestBuffer, FR
 	UE_LOG(LogRHI, Fatal, TEXT("RHITransferBufferUnderlyingResource isn't implemented for the current RHI"));
 }
 
-FUnorderedAccessViewRHIRef FDynamicRHI::RHICreateUnorderedAccessView(FRHITexture* Texture, uint32 MipLevel)
-{
-	return RHICreateUnorderedAccessView(Texture, MipLevel, 0, 0);
-}
-
 FUnorderedAccessViewRHIRef FDynamicRHI::RHICreateUnorderedAccessView(FRHITexture* Texture, uint32 MipLevel, uint8 Format)
-{
-	return RHICreateUnorderedAccessView(Texture, MipLevel, Format, 0, 0);
-}
-
-FUnorderedAccessViewRHIRef FDynamicRHI::RHICreateUnorderedAccessView(FRHITexture* Texture, uint32 MipLevel, uint8 Format, uint16 FirstArraySlice, uint16 NumArraySlices)
 {
 	UE_LOG(LogRHI, Fatal, TEXT("RHICreateUnorderedAccessView with Format parameter isn't implemented for the current RHI"));
-	return RHICreateUnorderedAccessView(Texture, MipLevel, FirstArraySlice, NumArraySlices);
+	return RHICreateUnorderedAccessView(Texture, MipLevel);
 }
 
 void FDynamicRHI::RHIUpdateShaderResourceView(FRHIShaderResourceView* SRV, FRHIBuffer* Buffer, uint32 Stride, uint8 Format)
@@ -569,6 +546,18 @@ uint64 FDynamicRHI::RHIGetMinimumAlignmentForBufferBackedSRV(EPixelFormat Format
 	return 1;
 }
 
+uint64 FDynamicRHI::RHICalcTexture2DArrayPlatformSize(uint32 SizeX, uint32 SizeY, uint32 ArraySize, uint8 Format, uint32 NumMips, uint32 NumSamples, ETextureCreateFlags Flags, const FRHIResourceCreateInfo& CreateInfo, uint32& OutAlign)
+{
+	// ensureMsgf(false, TEXT("RHICalcTexture2DArrayPlatformSize isn't implemented for the current RHI"));
+	return ArraySize * RHICalcTexture2DPlatformSize(SizeX, SizeY, Format, NumMips, NumSamples, Flags, CreateInfo, OutAlign);
+}
+
+uint64 FDynamicRHI::RHICalcVMTexture2DPlatformSize(uint32 Mip0Width, uint32 Mip0Height, uint8 Format, uint32 NumMips, uint32 FirstMipIdx, uint32 NumSamples, ETextureCreateFlags Flags, uint32& OutAlign)
+{
+	UE_LOG(LogRHI, Fatal, TEXT("RHICalcVMTexture2DPlatformSize isn't implemented for the current RHI"));
+	return -1;
+}
+
 FDefaultRHIRenderQueryPool::FDefaultRHIRenderQueryPool(ERenderQueryType InQueryType, FDynamicRHI* InDynamicRHI, uint32 InNumQueries)
 	: DynamicRHI(InDynamicRHI)
 	, QueryType(InQueryType)
@@ -643,19 +632,6 @@ void FDynamicRHI::RHICheckViewportHDRStatus(FRHIViewport* Viewport)
 {
 }
 
-void* FDynamicRHI::RHILockBufferMGPU(FRHICommandListImmediate& RHICmdList, FRHIBuffer* Buffer, uint32 GPUIndex, uint32 Offset, uint32 Size, EResourceLockMode LockMode)
-{
-	// Fall through to single GPU case
-	check(GPUIndex == 0);
-	return RHILockBuffer(RHICmdList, Buffer, Offset, Size, LockMode);
-}
-
-void FDynamicRHI::RHIUnlockBufferMGPU(FRHICommandListImmediate& RHICmdList, FRHIBuffer* Buffer, uint32 GPUIndex)
-{
-	// Fall through to single GPU case
-	check(GPUIndex == 0);
-	RHIUnlockBuffer(RHICmdList, Buffer);
-}
 
 FShaderResourceViewInitializer::FShaderResourceViewInitializer(FRHIBuffer* InBuffer, EPixelFormat InFormat, uint32 InStartOffsetBytes, uint32 InNumElements)
 	: BufferInitializer({ InBuffer, InStartOffsetBytes, InNumElements, InFormat }), Type(EType::VertexBufferSRV)
@@ -698,19 +674,6 @@ FShaderResourceViewInitializer::FShaderResourceViewInitializer(FRHIBuffer* InBuf
 	InitType();
 }
 
-FRawBufferShaderResourceViewInitializer::FRawBufferShaderResourceViewInitializer(FRHIBuffer* InBuffer)
-	: FShaderResourceViewInitializer(nullptr)
-{
-	check(GRHISupportsRawViewsForAnyBuffer);
-
-	Type = EType::RawBufferSRV;
-
-	BufferInitializer.Buffer = InBuffer;
-	BufferInitializer.Format = PF_Unknown;
-	BufferInitializer.StartOffsetBytes = 0;
-	BufferInitializer.NumElements = UINT32_MAX; // Whole resource
-}
-
 void FShaderResourceViewInitializer::InitType()
 {
 	if (BufferInitializer.Buffer)
