Fixed crash in RHI validation layer due to compiled out RHICopyBufferRegion.

diff --git a/Engine/Source/Runtime/RHI/Public/RHIValidationContext.h b/Engine/Source/Runtime/RHI/Public/RHIValidationContext.h
--- a/Engine/Source/Runtime/RHI/Public/RHIValidationContext.h
+++ b/Engine/Source/Runtime/RHI/Public/RHIValidationContext.h
@@ -979,14 +979,12 @@ public:
 		RHIContext->RHICopyTexture(SourceTexture, DestTexture, CopyInfo);
 	}
 
-#if RHI_RAYTRACING
 	virtual void RHICopyBufferRegion(FRHIBuffer* DestBuffer, uint64 DstOffset, FRHIBuffer* SourceBuffer, uint64 SrcOffset, uint64 NumBytes)
 	{
 		Tracker->Assert(SourceBuffer->GetWholeResourceIdentity(), ERHIAccess::CopySrc);
 		Tracker->Assert(DestBuffer->GetWholeResourceIdentity(), ERHIAccess::CopyDest);
 		RHIContext->RHICopyBufferRegion(DestBuffer, DstOffset, SourceBuffer, SrcOffset, NumBytes);
 	}
-#endif
 
 	void RHIClearRayTracingBindings(FRHIRayTracingScene* Scene)
 	{
