Implemented support for parallel buffer creation and lock / unlock.
 - Refactored relevant FDynamicRHI APIs to take an FRHICommandListBase reference. Platforms which require operations deferred to the RHI thread will enqueue commands into the command list.
 - Hardened platform implementations to use locks where needed.
 - Added unit tests to stress test creating and locking many buffers in parallel.

diff --git a/Engine/Plugins/Tests/RHITests/Source/RHITests/Private/RHIBufferTests.cpp b/Engine/Plugins/Tests/RHITests/Source/RHITests/Private/RHIBufferTests.cpp
--- a/Engine/Plugins/Tests/RHITests/Source/RHITests/Private/RHIBufferTests.cpp
+++ b/Engine/Plugins/Tests/RHITests/Source/RHITests/Private/RHIBufferTests.cpp
@@ -2,6 +2,48 @@
 
 #include "RHIBufferTests.h"
 
+bool FRHIBufferTests::VerifyBufferContents(const TCHAR* TestName, FRHICommandListImmediate& RHICmdList, TArrayView<FRHIBuffer*> Buffers, TFunctionRef<bool(int32 BufferIndex, void* Ptr, uint32 NumBytes)> VerifyCallback)
+{
+	TArray<FStagingBufferRHIRef, TInlineAllocator<1>> StagingBuffers;
+
+	for (FRHIBuffer* Buffer : Buffers)
+	{
+		FStagingBufferRHIRef StagingBuffer = RHICreateStagingBuffer();
+		RHICmdList.CopyToStagingBuffer(Buffer, StagingBuffer, 0, Buffer->GetSize());
+		StagingBuffers.Emplace(StagingBuffer);
+	}
+
+	// @todo - readback API is inconsistent across RHIs
+	RHICmdList.BlockUntilGPUIdle();
+	RHICmdList.FlushResources();
+
+	ON_SCOPE_EXIT
+	{
+		// Immediate flush to clean up the staging buffer / other resources
+		RHICmdList.ImmediateFlush(EImmediateFlushType::FlushRHIThreadFlushResources);
+	};
+
+	for (int32 Index = 0; Index < Buffers.Num(); ++Index)
+	{
+		FRHIBuffer* Buffer = Buffers[Index];
+		FRHIStagingBuffer* StagingBuffer = StagingBuffers[Index];
+		uint32 NumBytes = Buffer->GetSize();
+
+		void* Memory = RHILockStagingBuffer(StagingBuffer, 0, NumBytes);
+		bool Result = VerifyCallback(Index, Memory, NumBytes);
+		RHIUnlockStagingBuffer(StagingBuffer);
+
+		if (!Result)
+		{
+			UE_LOG(LogRHIUnitTestCommandlet, Error, TEXT("Test failed (buffer index %d) \"%s\""), Index, TestName);
+			return false;
+		}
+	}
+
+	UE_LOG(LogRHIUnitTestCommandlet, Display, TEXT("Test passed. \"%s\""), TestName);
+	return true;
+}
+
 // Copies data in the specified vertex buffer back to the CPU, and passes a pointer to that data to the provided verification lambda.
 bool FRHIBufferTests::VerifyBufferContents(const TCHAR* TestName, FRHICommandListImmediate& RHICmdList, FRHIBuffer* Buffer, TFunctionRef<bool(void* Ptr, uint32 NumBytes)> VerifyCallback)
 {
diff --git a/Engine/Plugins/Tests/RHITests/Source/RHITests/Private/RHITestsModule.cpp b/Engine/Plugins/Tests/RHITests/Source/RHITests/Private/RHITestsModule.cpp
--- a/Engine/Plugins/Tests/RHITests/Source/RHITests/Private/RHITestsModule.cpp
+++ b/Engine/Plugins/Tests/RHITests/Source/RHITests/Private/RHITestsModule.cpp
@@ -25,6 +25,8 @@ static bool RunTests_RenderThread(FRHICommandListImmediate& RHICmdList)
 
 		RUN_TEST(FRHIBufferTests::Test_RHIClearUAVUint_StructuredBuffer(RHICmdList));
 		RUN_TEST(FRHIBufferTests::Test_RHIClearUAVFloat_StructuredBuffer(RHICmdList));
+
+		RUN_TEST(FRHIBufferTests::Test_RHICreateBuffer_Parallel(RHICmdList));
 	}
 
 	// Texture2D/3D
diff --git a/Engine/Plugins/Tests/RHITests/Source/RHITests/Public/RHIBufferTests.h b/Engine/Plugins/Tests/RHITests/Source/RHITests/Public/RHIBufferTests.h
--- a/Engine/Plugins/Tests/RHITests/Source/RHITests/Public/RHIBufferTests.h
+++ b/Engine/Plugins/Tests/RHITests/Source/RHITests/Public/RHIBufferTests.h
@@ -3,12 +3,79 @@
 #pragma once
 
 #include "RHITestsCommon.h"
+#include "Async/ParallelFor.h"
+
+PRAGMA_DISABLE_OPTIMIZATION
 
 class FRHIBufferTests
 {
 	// Copies data in the specified vertex buffer back to the CPU, and passes a pointer to that data to the provided verification lambda.
+	static bool VerifyBufferContents(const TCHAR* TestName, FRHICommandListImmediate& RHICmdList, TArrayView<FRHIBuffer*> Buffers, TFunctionRef<bool(int32 BufferIndex, void* Ptr, uint32 NumBytes)> VerifyCallback);
 	static bool VerifyBufferContents(const TCHAR* TestName, FRHICommandListImmediate& RHICmdList, FRHIBuffer* Buffer, TFunctionRef<bool(void* Ptr, uint32 NumBytes)> VerifyCallback);
 
+	template <typename TestLambdaType>
+	static void ParallelDispatchCommands(FRHICommandListImmediate& RHICmdList, int32 NumTests, TestLambdaType TestLambda)
+	{
+		const int32 NumTasks = 32;
+
+		int32 NumTestsPerTask = FMath::Max(NumTests / NumTasks, 1);
+		int32 NumTestsLaunched = 0;
+
+		struct FTaskData
+		{
+			FTaskData() = default;
+			FTaskData(FRHICommandList* InRHICmdList, const FGraphEventRef& InEvent)
+				: RHICmdList(InRHICmdList)
+				, Event(InEvent)
+			{}
+
+			FRHICommandList* RHICmdList = nullptr;
+			FGraphEventRef Event;
+		};
+
+		TArray<FTaskData> TaskDatas;
+
+		while (NumTestsLaunched < NumTests)
+		{
+			FRHICommandList* RHICmdListUpload = new FRHICommandList(FRHIGPUMask::All());
+
+			const int32 NumTestsInTask = FMath::Min(NumTestsPerTask, NumTests - NumTestsLaunched);
+
+#if 1
+			RHICmdListUpload->DisallowBypass();
+
+			FGraphEventRef Event = FFunctionGraphTask::CreateAndDispatchWhenReady(
+				[TestLambda, RHICmdListUpload, NumTestsLaunched, NumTestsInTask](ENamedThreads::Type, const FGraphEventRef&)
+			{
+				FTaskTagScope Scope(ETaskTag::EParallelRenderingThread);
+
+				for (int32 Index = 0; Index < NumTestsInTask; ++Index)
+				{
+					TestLambda(*RHICmdListUpload, Index + NumTestsLaunched, NumTestsLaunched);
+				}
+
+			}, TStatId(), nullptr, ENamedThreads::AnyHiPriThreadHiPriTask);
+
+			TaskDatas.Emplace(RHICmdListUpload, MoveTemp(Event));
+#else
+			for (int32 Index = 0; Index < NumTestsInTask; ++Index)
+			{
+				TestLambda(*RHICmdListUpload, Index + NumTestsLaunched, NumTestsLaunched);
+			}
+			delete RHICmdListUpload;
+#endif
+
+			NumTestsLaunched += NumTestsInTask;
+		}
+
+		for (FTaskData TaskData : TaskDatas)
+		{
+			RHICmdList.QueueAsyncCommandListSubmit(TaskData.Event, TaskData.RHICmdList);
+		}
+
+		RHICmdList.WaitForTasks();
+	}
+
 	template <typename BufferType, typename ValueType, uint32 NumTestBytes>
 	static bool RunTest_UAVClear_Buffer(FRHICommandListImmediate& RHICmdList, const FString& TestName, BufferType* BufferRHI, FRHIUnorderedAccessView* UAV, uint32 BufferSize, const ValueType& ClearValue, void(FRHIComputeCommandList::* ClearPtr)(FRHIUnorderedAccessView*, ValueType const&), const uint8(&TestValue)[NumTestBytes])
 	{
@@ -300,4 +367,96 @@ public:
 
 		return bResult;
 	}
-};
\ No newline at end of file
+
+	static bool Test_RHICreateBuffer_Parallel(FRHICommandListImmediate& RHICmdList)
+	{
+		if (!GRHISupportsMultithreadedResources)
+		{
+			return true;
+		}
+
+		SCOPED_NAMED_EVENT_TEXT("Test_RHICreateBuffer_Parallel", FColor::Magenta);
+
+		const int32 NumBuffersToCreate = 256;
+
+		FGenericPlatformMath::RandInit(1);
+
+		TArray<TRefCountPtr<FRHIBuffer>> Buffers;
+		Buffers.SetNum(NumBuffersToCreate);
+
+		TArray<int32> RandomNumberPerBuffer;
+		RandomNumberPerBuffer.SetNum(NumBuffersToCreate);
+
+		for (int32 Index = 0; Index < NumBuffersToCreate; ++Index)
+		{
+			RandomNumberPerBuffer[Index] = FGenericPlatformMath::Rand();
+		}
+
+		ParallelDispatchCommands(RHICmdList, NumBuffersToCreate, [&](FRHICommandList& InRHICmdList, int32 Index, int32 Num)
+		{
+			SCOPED_NAMED_EVENT_TEXT("TestCreateBuffer", FColor::Magenta);
+
+			const int32 Random = RandomNumberPerBuffer[Index];
+			const int32 BufferSize = Align(Random % 65536, 16);
+			const int32 BufferStride = 4;
+
+			EBufferUsageFlags Usage = EBufferUsageFlags::VertexBuffer;
+
+			switch (Random % 3)
+			{
+			case 0:
+				Usage |= EBufferUsageFlags::Static;
+				break;
+			case 1:
+				Usage |= EBufferUsageFlags::Dynamic;
+				break;
+			case 2:
+				Usage |= EBufferUsageFlags::Volatile;
+				break;
+			}
+
+			FRHIResourceCreateInfo CreateInfo(TEXT("Buffer"));
+
+			TRefCountPtr<FRHIBuffer> Buffer = InRHICmdList.CreateBuffer(BufferSize, Usage, 0, ERHIAccess::VertexOrIndexBuffer, CreateInfo);
+
+			uint32* Data = (uint32*)InRHICmdList.LockBuffer(Buffer, 0, Buffer->GetSize(), RLM_WriteOnly);
+
+			uint32 NumDWORDs = Buffer->GetSize() >> 2;
+
+			for (uint32 DataIndex = 0; DataIndex < NumDWORDs; ++DataIndex)
+			{
+				Data[DataIndex] = DataIndex;
+			}
+
+			InRHICmdList.UnlockBuffer(Buffer);
+
+			Buffers[Index] = Buffer;
+		});
+
+		TArray<FRHIBuffer*> BufferPtrs;
+		BufferPtrs.Reserve(Buffers.Num());
+
+		for (int32 Index = 0; Index < NumBuffersToCreate; ++Index)
+		{
+			BufferPtrs.Emplace(Buffers[Index]);
+		}
+
+		return VerifyBufferContents(TEXT("Test_RHICreateBuffer_Parallel"), RHICmdList, BufferPtrs, [&](int32, void* Ptr, uint32 NumBytes)
+		{
+			uint32* Data = (uint32*)Ptr;
+			uint32 NumDWORDs = NumBytes >> 2;
+
+			for (uint32 Index = 0; Index < NumDWORDs; Index++)
+			{
+				if (Data[Index] != Index)
+				{
+					return false;
+				}
+			}
+
+			return true;
+		});
+	}
+};
+
+PRAGMA_ENABLE_OPTIMIZATION
\ No newline at end of file
diff --git a/Engine/Source/Runtime/Apple/AGXRHI/Private/AGXDynamicRHI.h b/Engine/Source/Runtime/Apple/AGXRHI/Private/AGXDynamicRHI.h
--- a/Engine/Source/Runtime/Apple/AGXRHI/Private/AGXDynamicRHI.h
+++ b/Engine/Source/Runtime/Apple/AGXRHI/Private/AGXDynamicRHI.h
@@ -55,11 +55,11 @@ public:
 	virtual FGraphicsPipelineStateRHIRef RHICreateGraphicsPipelineState(const FGraphicsPipelineStateInitializer& Initializer) final override;
 	virtual TRefCountPtr<FRHIComputePipelineState> RHICreateComputePipelineState(FRHIComputeShader* ComputeShader) final override;
 	virtual FUniformBufferRHIRef RHICreateUniformBuffer(const void* Contents, const FRHIUniformBufferLayout* Layout, EUniformBufferUsage Usage, EUniformBufferValidation Validation) final override;
-	virtual FBufferRHIRef RHICreateBuffer(uint32 Size, EBufferUsageFlags Usage, uint32 Stride, ERHIAccess InResourceState, FRHIResourceCreateInfo& CreateInfo) final override;
+	virtual FBufferRHIRef RHICreateBuffer(FRHICommandListBase& RHICmdList, uint32 Size, EBufferUsageFlags Usage, uint32 Stride, ERHIAccess InResourceState, FRHIResourceCreateInfo& CreateInfo) final override;
 	virtual void RHITransferBufferUnderlyingResource(FRHIBuffer* DestBuffer, FRHIBuffer* SrcBuffer) final override;
 	virtual void RHICopyBuffer(FRHIBuffer* SourceBuffer, FRHIBuffer* DestBuffer) final override;
-	virtual void* LockBuffer_BottomOfPipe(FRHICommandListImmediate& RHICmdList, FRHIBuffer* Buffer, uint32 Offset, uint32 SizeRHI, EResourceLockMode LockMode) final override;
-	virtual void UnlockBuffer_BottomOfPipe(FRHICommandListImmediate& RHICmdList, FRHIBuffer* Buffer) final override;
+	virtual void* LockBuffer_BottomOfPipe(FRHICommandListBase& RHICmdList, FRHIBuffer* Buffer, uint32 Offset, uint32 SizeRHI, EResourceLockMode LockMode) final override;
+	virtual void UnlockBuffer_BottomOfPipe(FRHICommandListBase& RHICmdList, FRHIBuffer* Buffer) final override;
 	virtual FUnorderedAccessViewRHIRef RHICreateUnorderedAccessView(FRHIBuffer* Buffer, bool bUseUAVCounter, bool bAppendBuffer) final override;
 	virtual FUnorderedAccessViewRHIRef RHICreateUnorderedAccessView(FRHITexture* Texture, uint32 MipLevel, uint16 FirstArraySlice, uint16 NumArraySlices) final override;
 	virtual FUnorderedAccessViewRHIRef RHICreateUnorderedAccessView(FRHIBuffer* Buffer, uint8 Format) final override;
@@ -119,9 +119,6 @@ public:
 	virtual class IRHICommandContext* AGXRHI_API RHIGetDefaultContext() final override;
 	virtual IRHIComputeContext* RHIGetDefaultAsyncComputeContext() final override;
 	virtual class IRHICommandContextContainer* RHIGetCommandContextContainer(int32 Index, int32 Num) final override;
-	
-	virtual FBufferRHIRef CreateBuffer_RenderThread(class FRHICommandListImmediate& RHICmdList, uint32 Size, EBufferUsageFlags Usage, uint32 Stride, ERHIAccess ResourceState, FRHIResourceCreateInfo& CreateInfo) final override;
-	
 	virtual FTexture2DRHIRef AsyncReallocateTexture2D_RenderThread(class FRHICommandListImmediate& RHICmdList, FRHITexture2D* Texture2D, int32 NewMipCount, int32 NewSizeX, int32 NewSizeY, FThreadSafeCounter* RequestStatus) final override;
 	virtual ETextureReallocationStatus FinalizeAsyncReallocateTexture2D_RenderThread(class FRHICommandListImmediate& RHICmdList, FRHITexture2D* Texture2D, bool bBlockUntilCompleted) final override;
 	virtual ETextureReallocationStatus CancelAsyncReallocateTexture2D_RenderThread(class FRHICommandListImmediate& RHICmdList, FRHITexture2D* Texture2D, bool bBlockUntilCompleted) final override;
diff --git a/Engine/Source/Runtime/Apple/AGXRHI/Private/AGXRHI.cpp b/Engine/Source/Runtime/Apple/AGXRHI/Private/AGXRHI.cpp
--- a/Engine/Source/Runtime/Apple/AGXRHI/Private/AGXRHI.cpp
+++ b/Engine/Source/Runtime/Apple/AGXRHI/Private/AGXRHI.cpp
@@ -183,6 +183,7 @@ FAGXDynamicRHI::FAGXDynamicRHI(ERHIFeatureLevel::Type RequestedFeatureLevel)
 #endif
 	
 	GRHISupportsMultithreading = true;
+	GRHISupportsMultithreadedResources = true;
 	
 	// we cannot render to a volume texture without geometry shader or vertex-shader-layer support, so initialise to false and enable based on platform feature availability
 	GSupportsVolumeTextureRendering = false;
diff --git a/Engine/Source/Runtime/Apple/AGXRHI/Private/AGXVertexBuffer.cpp b/Engine/Source/Runtime/Apple/AGXRHI/Private/AGXVertexBuffer.cpp
--- a/Engine/Source/Runtime/Apple/AGXRHI/Private/AGXVertexBuffer.cpp
+++ b/Engine/Source/Runtime/Apple/AGXRHI/Private/AGXVertexBuffer.cpp
@@ -604,24 +604,11 @@ void FAGXRHIBuffer::Unlock()
 	LastLockFrame = GetAGXDeviceContext().GetFrameNumberRHIThread();
 }
 
-FAGXResourceMultiBuffer::FAGXResourceMultiBuffer(uint32 InSize, EBufferUsageFlags InUsage, EAGXBufferUsage InAgxUsage, uint32 InStride, FResourceArrayInterface* ResourceArray, ERHIResourceType Type)
+FAGXResourceMultiBuffer::FAGXResourceMultiBuffer(uint32 InSize, EBufferUsageFlags InUsage, EAGXBufferUsage InAgxUsage, uint32 InStride, ERHIResourceType Type)
 	: FRHIBuffer(InSize, InUsage, InStride)
 	, FAGXRHIBuffer(InSize, InUsage, InAgxUsage, Type)
 	, IndexType((InStride == 2) ? MTLIndexTypeUInt16 : MTLIndexTypeUInt32)
 {
-	if (EnumHasAnyFlags(InUsage, BUF_StructuredBuffer))
-	{
-		check((InSize % InStride) == 0);
-
-		if (ResourceArray)
-		{
-			// copy any resources to the CPU address
-			void* LockedMemory = RHILockBuffer(this, 0, InSize, RLM_WriteOnly);
-			FMemory::Memcpy(LockedMemory, ResourceArray->GetResourceData(), InSize);
-			ResourceArray->Discard();
-			RHIUnlockBuffer(this);
-		}
-	}
 }
 
 FAGXResourceMultiBuffer::~FAGXResourceMultiBuffer()
@@ -637,63 +624,22 @@ void FAGXResourceMultiBuffer::Swap(FAGXResourceMultiBuffer& Other)
 	}
 }
 
-FBufferRHIRef FAGXDynamicRHI::RHICreateBuffer(uint32 Size, EBufferUsageFlags Usage, uint32 Stride, ERHIAccess InResourceState, FRHIResourceCreateInfo& CreateInfo)
+FBufferRHIRef FAGXDynamicRHI::RHICreateBuffer(FRHICommandListBase& RHICmdList, uint32 Size, EBufferUsageFlags Usage, uint32 Stride, ERHIAccess InResourceState, FRHIResourceCreateInfo& CreateInfo)
 {
-	check(0);
-	@autoreleasepool {
+	@autoreleasepool{
 	if (CreateInfo.bWithoutNativeResource)
 	{
-		return new FAGXResourceMultiBuffer(0, BUF_None, AGXBufferUsage(BUF_None), 0, nullptr, RRT_Buffer);
-	}
-	
-	// make the RHI object, which will allocate memory
-	FAGXResourceMultiBuffer* Buffer = new FAGXResourceMultiBuffer(Size, Usage, AGXBufferUsage(Usage), 0, nullptr, RRT_Buffer);
-
-	if (CreateInfo.ResourceArray)
-	{
-		check(Size >= CreateInfo.ResourceArray->GetResourceDataSize());
-
-		// make a buffer usable by CPU
-		void* BufferData = ::RHILockBuffer(Buffer, 0, Size, RLM_WriteOnly);
-		
-		// copy the contents of the given data into the buffer
-		FMemory::Memcpy(BufferData, CreateInfo.ResourceArray->GetResourceData(), Size);
-		
-		::RHIUnlockBuffer(Buffer);
-
-		// Discard the resource array's contents.
-		CreateInfo.ResourceArray->Discard();
-	}
-	else if (Buffer->StorageMode == MTLStorageModePrivate)
-	{
-		check (!Buffer->TransferBuffer);
-
-		if (GAGXBufferZeroFill)
-		{
-			for(FAGXRHIBuffer::FAGXBufferAndViews& Backing : Buffer->BufferPool)
-			{
-				FAGXBuffer& TheBuffer = Backing.Buffer;
-				GetAGXDeviceContext().FillBuffer(TheBuffer, ns::Range(0, TheBuffer.GetLength()), 0);
-			}
-		}
-	}
-#if PLATFORM_MAC
-	else if (GAGXBufferZeroFill && Buffer->StorageMode == MTLStorageModeManaged)
-	{
-		for(FAGXRHIBuffer::FAGXBufferAndViews& Backing : Buffer->BufferPool)
-		{
-			FAGXBuffer& TheBuffer = Backing.Buffer;
-			GetAGXDeviceContext().FillBuffer(TheBuffer, ns::Range(0, TheBuffer.GetLength()), 0);
-			MTLPP_VALIDATE(mtlpp::Buffer, TheBuffer, AGXSafeGetRuntimeDebuggingLevel() >= EAGXDebugLevelValidation, DidModify(ns::Range(0, TheBuffer)));
-		}
+		return new FAGXResourceMultiBuffer(0, BUF_None, AGXBufferUsage(BUF_None), 0, RRT_Buffer);
 	}
-#endif
 
-	return Buffer;
+	// make the RHI object, which will allocate memory
+	TRefCountPtr<FAGXResourceMultiBuffer> Buffer = new FAGXResourceMultiBuffer(Size, Usage, AGXBufferUsage(Usage), Stride, RRT_Buffer);
+	Buffer->Init(RHICmdList, Size, Usage, CreateInfo, Buffer);
+	return Buffer.GetReference();
 	}
 }
 
-void* FAGXDynamicRHI::LockBuffer_BottomOfPipe(FRHICommandListImmediate& RHICmdList, FRHIBuffer* BufferRHI, uint32 Offset, uint32 Size, EResourceLockMode LockMode)
+void* FAGXDynamicRHI::LockBuffer_BottomOfPipe(FRHICommandListBase& RHICmdList, FRHIBuffer* BufferRHI, uint32 Offset, uint32 Size, EResourceLockMode LockMode)
 {
 	@autoreleasepool {
 	FAGXResourceMultiBuffer* Buffer = ResourceCast(BufferRHI);
@@ -703,7 +649,7 @@ void* FAGXDynamicRHI::LockBuffer_BottomOfPipe(FRHICommandListImmediate& RHICmdLi
 	}
 }
 
-void FAGXDynamicRHI::UnlockBuffer_BottomOfPipe(FRHICommandListImmediate& RHICmdList, FRHIBuffer* BufferRHI)
+void FAGXDynamicRHI::UnlockBuffer_BottomOfPipe(FRHICommandListBase& RHICmdList, FRHIBuffer* BufferRHI)
 {
 	@autoreleasepool {
 	FAGXResourceMultiBuffer* Buffer = ResourceCast(BufferRHI);
@@ -751,7 +697,7 @@ void FAGXDynamicRHI::RHITransferBufferUnderlyingResource(FRHIBuffer* DestBuffer,
 		FAGXResourceMultiBuffer* Dest = ResourceCast(DestBuffer);
 		if (!SrcBuffer)
 		{
-			TRefCountPtr<FAGXResourceMultiBuffer> DeletionProxy = new FAGXResourceMultiBuffer(0, Dest->GetUsage(), Dest->GetAgxUsage(), Dest->GetStride(), nullptr, Dest->Type);
+			TRefCountPtr<FAGXResourceMultiBuffer> DeletionProxy = new FAGXResourceMultiBuffer(0, Dest->GetUsage(), Dest->GetAgxUsage(), Dest->GetStride(), Dest->Type);
 			Dest->Swap(*DeletionProxy);
 		}
 		else
@@ -762,7 +708,7 @@ void FAGXDynamicRHI::RHITransferBufferUnderlyingResource(FRHIBuffer* DestBuffer,
 	}
 }
 
-void FAGXRHIBuffer::Init_RenderThread(FRHICommandListImmediate& RHICmdList, uint32 InSize, EBufferUsageFlags InUsage, FRHIResourceCreateInfo& CreateInfo, FRHIResource* Resource)
+void FAGXRHIBuffer::Init(FRHICommandListBase& RHICmdList, uint32 InSize, EBufferUsageFlags InUsage, FRHIResourceCreateInfo& CreateInfo, FRHIResource* Resource)
 {
 	if (CreateInfo.ResourceArray)
 	{
@@ -788,7 +734,7 @@ void FAGXRHIBuffer::Init_RenderThread(FRHICommandListImmediate& RHICmdList, uint
 					FMemory::Memcpy(Result, CreateInfo.ResourceArray->GetResourceData(), InSize);
 					
 					RHICmdList.EnqueueLambda(
-						[this, Result, InSize](FRHICommandList& RHICmdList)
+						[this, Result, InSize](FRHICommandListBase& RHICmdList)
 						{
 							void* Backing = this->Lock(true, RLM_WriteOnly, 0, InSize);
 							FMemory::Memcpy(Backing, Result, InSize);
@@ -813,21 +759,4 @@ void FAGXRHIBuffer::Init_RenderThread(FRHICommandListImmediate& RHICmdList, uint
 		// Discard the resource array's contents.
 		CreateInfo.ResourceArray->Discard();
 	}
-}
-
-FBufferRHIRef FAGXDynamicRHI::CreateBuffer_RenderThread(class FRHICommandListImmediate& RHICmdList, uint32 Size, EBufferUsageFlags Usage, uint32 Stride, ERHIAccess ResourceState, FRHIResourceCreateInfo& CreateInfo)
-{
-	@autoreleasepool {
-		if (CreateInfo.bWithoutNativeResource)
-		{
-			return new FAGXResourceMultiBuffer(0, BUF_None, AGXBufferUsage(BUF_None), 0, nullptr, RRT_Buffer);
-		}
-		
-		// make the RHI object, which will allocate memory
-		TRefCountPtr<FAGXResourceMultiBuffer> Buffer = new FAGXResourceMultiBuffer(Size, Usage, AGXBufferUsage(Usage), Stride, nullptr, RRT_Buffer);
-		
-		Buffer->Init_RenderThread(RHICmdList, Size, Usage, CreateInfo, Buffer);
-		
-		return Buffer.GetReference();
-	}
-}
+}
\ No newline at end of file
diff --git a/Engine/Source/Runtime/Apple/AGXRHI/Public/AGXResources.h b/Engine/Source/Runtime/Apple/AGXRHI/Public/AGXResources.h
--- a/Engine/Source/Runtime/Apple/AGXRHI/Public/AGXResources.h
+++ b/Engine/Source/Runtime/Apple/AGXRHI/Public/AGXResources.h
@@ -322,7 +322,7 @@ public:
 	/**
 	 * Initialize the buffer contents from the render-thread.
 	 */
-	void Init_RenderThread(class FRHICommandListImmediate& RHICmdList, uint32 Size, EBufferUsageFlags InUsage, FRHIResourceCreateInfo& CreateInfo, FRHIResource* Resource);
+	void Init(class FRHICommandListBase& RHICmdList, uint32 Size, EBufferUsageFlags InUsage, FRHIResourceCreateInfo& CreateInfo, FRHIResource* Resource);
 	
 	/**
 	 * Get a linear texture for given format.
@@ -448,7 +448,7 @@ private:
 class FAGXResourceMultiBuffer : public FRHIBuffer, public FAGXRHIBuffer
 {
 public:
-	FAGXResourceMultiBuffer(uint32 InSize, EBufferUsageFlags InUsage, EAGXBufferUsage InAgxUsage, uint32 InStride, FResourceArrayInterface* ResourceArray, ERHIResourceType ResourceType);
+	FAGXResourceMultiBuffer(uint32 InSize, EBufferUsageFlags InUsage, EAGXBufferUsage InAgxUsage, uint32 InStride, ERHIResourceType ResourceType);
 	virtual ~FAGXResourceMultiBuffer();
 
 	void Swap(FAGXResourceMultiBuffer& Other);
diff --git a/Engine/Source/Runtime/Apple/MetalRHI/Private/MetalDynamicRHI.h b/Engine/Source/Runtime/Apple/MetalRHI/Private/MetalDynamicRHI.h
--- a/Engine/Source/Runtime/Apple/MetalRHI/Private/MetalDynamicRHI.h
+++ b/Engine/Source/Runtime/Apple/MetalRHI/Private/MetalDynamicRHI.h
@@ -56,11 +56,11 @@ public:
 	virtual FGraphicsPipelineStateRHIRef RHICreateGraphicsPipelineState(const FGraphicsPipelineStateInitializer& Initializer) final override;
 	virtual TRefCountPtr<FRHIComputePipelineState> RHICreateComputePipelineState(FRHIComputeShader* ComputeShader) final override;
 	virtual FUniformBufferRHIRef RHICreateUniformBuffer(const void* Contents, const FRHIUniformBufferLayout* Layout, EUniformBufferUsage Usage, EUniformBufferValidation Validation) final override;
-	virtual FBufferRHIRef RHICreateBuffer(uint32 Size, EBufferUsageFlags Usage, uint32 Stride, ERHIAccess InResourceState, FRHIResourceCreateInfo& CreateInfo) final override;
+	virtual FBufferRHIRef RHICreateBuffer(FRHICommandListBase& RHICmdList, uint32 Size, EBufferUsageFlags Usage, uint32 Stride, ERHIAccess InResourceState, FRHIResourceCreateInfo& CreateInfo) final override;
 	virtual void RHITransferBufferUnderlyingResource(FRHIBuffer* DestBuffer, FRHIBuffer* SrcBuffer) final override;
 	virtual void RHICopyBuffer(FRHIBuffer* SourceBuffer, FRHIBuffer* DestBuffer) final override;
-	virtual void* LockBuffer_BottomOfPipe(FRHICommandListImmediate& RHICmdList, FRHIBuffer* Buffer, uint32 Offset, uint32 SizeRHI, EResourceLockMode LockMode) final override;
-	virtual void UnlockBuffer_BottomOfPipe(FRHICommandListImmediate& RHICmdList, FRHIBuffer* Buffer) final override;
+	virtual void* LockBuffer_BottomOfPipe(FRHICommandListBase& RHICmdList, FRHIBuffer* Buffer, uint32 Offset, uint32 SizeRHI, EResourceLockMode LockMode) final override;
+	virtual void UnlockBuffer_BottomOfPipe(FRHICommandListBase& RHICmdList, FRHIBuffer* Buffer) final override;
 	virtual FUnorderedAccessViewRHIRef RHICreateUnorderedAccessView(FRHIBuffer* Buffer, bool bUseUAVCounter, bool bAppendBuffer) final override;
 	virtual FUnorderedAccessViewRHIRef RHICreateUnorderedAccessView(FRHITexture* Texture, uint32 MipLevel, uint16 FirstArraySlice, uint16 NumArraySlices) final override;
 	virtual FUnorderedAccessViewRHIRef RHICreateUnorderedAccessView(FRHIBuffer* Buffer, uint8 Format) final override;
@@ -123,8 +123,6 @@ public:
 	virtual IRHIComputeContext* RHIGetDefaultAsyncComputeContext() final override;
 	virtual class IRHICommandContextContainer* RHIGetCommandContextContainer(int32 Index, int32 Num) final override;
 	
-	virtual FBufferRHIRef CreateBuffer_RenderThread(class FRHICommandListImmediate& RHICmdList, uint32 Size, EBufferUsageFlags Usage, uint32 Stride, ERHIAccess ResourceState, FRHIResourceCreateInfo& CreateInfo) final override;
-	
 	virtual FTexture2DRHIRef AsyncReallocateTexture2D_RenderThread(class FRHICommandListImmediate& RHICmdList, FRHITexture2D* Texture2D, int32 NewMipCount, int32 NewSizeX, int32 NewSizeY, FThreadSafeCounter* RequestStatus) final override;
 	virtual ETextureReallocationStatus FinalizeAsyncReallocateTexture2D_RenderThread(class FRHICommandListImmediate& RHICmdList, FRHITexture2D* Texture2D, bool bBlockUntilCompleted) final override;
 	virtual ETextureReallocationStatus CancelAsyncReallocateTexture2D_RenderThread(class FRHICommandListImmediate& RHICmdList, FRHITexture2D* Texture2D, bool bBlockUntilCompleted) final override;
diff --git a/Engine/Source/Runtime/Apple/MetalRHI/Private/MetalRHI.cpp b/Engine/Source/Runtime/Apple/MetalRHI/Private/MetalRHI.cpp
--- a/Engine/Source/Runtime/Apple/MetalRHI/Private/MetalRHI.cpp
+++ b/Engine/Source/Runtime/Apple/MetalRHI/Private/MetalRHI.cpp
@@ -167,6 +167,7 @@ FMetalDynamicRHI::FMetalDynamicRHI(ERHIFeatureLevel::Type RequestedFeatureLevel)
 #endif
 	
 	GRHISupportsMultithreading = true;
+	GRHISupportsMultithreadedResources = true;
 	
 	// we cannot render to a volume texture without geometry shader or vertex-shader-layer support, so initialise to false and enable based on platform feature availability
 	GSupportsVolumeTextureRendering = false;
diff --git a/Engine/Source/Runtime/Apple/MetalRHI/Private/MetalVertexBuffer.cpp b/Engine/Source/Runtime/Apple/MetalRHI/Private/MetalVertexBuffer.cpp
--- a/Engine/Source/Runtime/Apple/MetalRHI/Private/MetalVertexBuffer.cpp
+++ b/Engine/Source/Runtime/Apple/MetalRHI/Private/MetalVertexBuffer.cpp
@@ -712,24 +712,11 @@ void FMetalRHIBuffer::Unlock()
 	LastLockFrame = GetMetalDeviceContext().GetFrameNumberRHIThread();
 }
 
-FMetalResourceMultiBuffer::FMetalResourceMultiBuffer(uint32 InSize, EBufferUsageFlags InUsage, EMetalBufferUsage InMetalUsage, uint32 InStride, FResourceArrayInterface* ResourceArray, ERHIResourceType Type)
+FMetalResourceMultiBuffer::FMetalResourceMultiBuffer(uint32 InSize, EBufferUsageFlags InUsage, EMetalBufferUsage InMetalUsage, uint32 InStride, ERHIResourceType Type)
 	: FRHIBuffer(InSize, InUsage, InStride)
 	, FMetalRHIBuffer(InSize, InUsage, InMetalUsage, Type)
 	, IndexType((InStride == 2) ? mtlpp::IndexType::UInt16 : mtlpp::IndexType::UInt32)
 {
-	if (EnumHasAnyFlags(InUsage, BUF_StructuredBuffer))
-	{
-		check((InSize % InStride) == 0);
-
-		if (ResourceArray)
-		{
-			// copy any resources to the CPU address
-			void* LockedMemory = RHILockBuffer(this, 0, InSize, RLM_WriteOnly);
-			FMemory::Memcpy(LockedMemory, ResourceArray->GetResourceData(), InSize);
-			ResourceArray->Discard();
-			RHIUnlockBuffer(this);
-		}
-	}
 }
 
 FMetalResourceMultiBuffer::~FMetalResourceMultiBuffer()
@@ -745,63 +732,21 @@ void FMetalResourceMultiBuffer::Swap(FMetalResourceMultiBuffer& Other)
 	}
 }
 
-FBufferRHIRef FMetalDynamicRHI::RHICreateBuffer(uint32 Size, EBufferUsageFlags Usage, uint32 Stride, ERHIAccess InResourceState, FRHIResourceCreateInfo& CreateInfo)
+FBufferRHIRef FMetalDynamicRHI::RHICreateBuffer(FRHICommandListBase& RHICmdList, uint32 Size, EBufferUsageFlags Usage, uint32 Stride, ERHIAccess InResourceState, FRHIResourceCreateInfo& CreateInfo)
 {
-	check(0);
-	@autoreleasepool {
+	@autoreleasepool{
 	if (CreateInfo.bWithoutNativeResource)
 	{
-		return new FMetalResourceMultiBuffer(0, BUF_None, MetalBufferUsage(BUF_None), 0, nullptr, RRT_Buffer);
-	}
-	
-	// make the RHI object, which will allocate memory
-	FMetalResourceMultiBuffer* Buffer = new FMetalResourceMultiBuffer(Size, Usage, MetalBufferUsage(Usage), 0, nullptr, RRT_Buffer);
-
-	if (CreateInfo.ResourceArray)
-	{
-		check(Size >= CreateInfo.ResourceArray->GetResourceDataSize());
-
-		// make a buffer usable by CPU
-		void* BufferData = ::RHILockBuffer(Buffer, 0, Size, RLM_WriteOnly);
-		
-		// copy the contents of the given data into the buffer
-		FMemory::Memcpy(BufferData, CreateInfo.ResourceArray->GetResourceData(), Size);
-		
-		::RHIUnlockBuffer(Buffer);
-
-		// Discard the resource array's contents.
-		CreateInfo.ResourceArray->Discard();
-	}
-	else if (Buffer->Mode == mtlpp::StorageMode::Private)
-	{
-		check (!Buffer->TransferBuffer);
-
-		if (GMetalBufferZeroFill && !FMetalCommandQueue::SupportsFeature(EMetalFeaturesFences))
-		{
-			for(FMetalRHIBuffer::FMetalBufferAndViews& Backing : Buffer->BufferPool)
-			{
-				FMetalBuffer& TheBuffer = Backing.Buffer;
-				GetMetalDeviceContext().FillBuffer(TheBuffer, ns::Range(0, TheBuffer.GetLength()), 0);
-			}
-		}
-	}
-#if PLATFORM_MAC
-	else if (GMetalBufferZeroFill && Buffer->Mode == mtlpp::StorageMode::Managed)
-	{
-		for(FMetalRHIBuffer::FMetalBufferAndViews& Backing : Buffer->BufferPool)
-		{
-			FMetalBuffer& TheBuffer = Backing.Buffer;
-			GetMetalDeviceContext().FillBuffer(TheBuffer, ns::Range(0, TheBuffer.GetLength()), 0);
-			MTLPP_VALIDATE(mtlpp::Buffer, TheBuffer, SafeGetRuntimeDebuggingLevel() >= EMetalDebugLevelValidation, DidModify(ns::Range(0, TheBuffer)));
-		}
+		return new FMetalResourceMultiBuffer(0, BUF_None, MetalBufferUsage(BUF_None), 0, RRT_Buffer);
 	}
-#endif
 
-	return Buffer;
+	TRefCountPtr<FMetalResourceMultiBuffer> Buffer = new FMetalResourceMultiBuffer(Size, Usage, MetalBufferUsage(Usage), Stride, RRT_Buffer);
+	Buffer->Init(RHICmdList, Size, Usage, CreateInfo, Buffer);
+	return Buffer.GetReference();
 	}
 }
 
-void* FMetalDynamicRHI::LockBuffer_BottomOfPipe(FRHICommandListImmediate& RHICmdList, FRHIBuffer* BufferRHI, uint32 Offset, uint32 Size, EResourceLockMode LockMode)
+void* FMetalDynamicRHI::LockBuffer_BottomOfPipe(FRHICommandListBase& RHICmdList, FRHIBuffer* BufferRHI, uint32 Offset, uint32 Size, EResourceLockMode LockMode)
 {
 	@autoreleasepool {
 	FMetalResourceMultiBuffer* Buffer = ResourceCast(BufferRHI);
@@ -811,7 +756,7 @@ void* FMetalDynamicRHI::LockBuffer_BottomOfPipe(FRHICommandListImmediate& RHICmd
 	}
 }
 
-void FMetalDynamicRHI::UnlockBuffer_BottomOfPipe(FRHICommandListImmediate& RHICmdList, FRHIBuffer* BufferRHI)
+void FMetalDynamicRHI::UnlockBuffer_BottomOfPipe(FRHICommandListBase& RHICmdList, FRHIBuffer* BufferRHI)
 {
 	@autoreleasepool {
 	FMetalResourceMultiBuffer* Buffer = ResourceCast(BufferRHI);
@@ -859,7 +804,7 @@ void FMetalDynamicRHI::RHITransferBufferUnderlyingResource(FRHIBuffer* DestBuffe
 		FMetalResourceMultiBuffer* Dest = ResourceCast(DestBuffer);
 		if (!SrcBuffer)
 		{
-			TRefCountPtr<FMetalResourceMultiBuffer> DeletionProxy = new FMetalResourceMultiBuffer(0, Dest->GetUsage(), Dest->GetMetalUsage(), Dest->GetStride(), nullptr, Dest->Type);
+			TRefCountPtr<FMetalResourceMultiBuffer> DeletionProxy = new FMetalResourceMultiBuffer(0, Dest->GetUsage(), Dest->GetMetalUsage(), Dest->GetStride(), Dest->Type);
 			Dest->Swap(*DeletionProxy);
 		}
 		else
@@ -870,7 +815,7 @@ void FMetalDynamicRHI::RHITransferBufferUnderlyingResource(FRHIBuffer* DestBuffe
 	}
 }
 
-void FMetalRHIBuffer::Init_RenderThread(FRHICommandListImmediate& RHICmdList, uint32 InSize, EBufferUsageFlags InUsage, FRHIResourceCreateInfo& CreateInfo, FRHIResource* Resource)
+void FMetalRHIBuffer::Init(FRHICommandListBase& RHICmdList, uint32 InSize, EBufferUsageFlags InUsage, FRHIResourceCreateInfo& CreateInfo, FRHIResource* Resource)
 {
 	if (CreateInfo.ResourceArray)
 	{
@@ -896,7 +841,7 @@ void FMetalRHIBuffer::Init_RenderThread(FRHICommandListImmediate& RHICmdList, ui
 					FMemory::Memcpy(Result, CreateInfo.ResourceArray->GetResourceData(), InSize);
 					
 					RHICmdList.EnqueueLambda(
-						[this, Result, InSize](FRHICommandList& RHICmdList)
+						[this, Result, InSize](FRHICommandListBase& RHICmdList)
 						{
 							void* Backing = this->Lock(true, RLM_WriteOnly, 0, InSize);
 							FMemory::Memcpy(Backing, Result, InSize);
@@ -921,21 +866,4 @@ void FMetalRHIBuffer::Init_RenderThread(FRHICommandListImmediate& RHICmdList, ui
 		// Discard the resource array's contents.
 		CreateInfo.ResourceArray->Discard();
 	}
-}
-
-FBufferRHIRef FMetalDynamicRHI::CreateBuffer_RenderThread(class FRHICommandListImmediate& RHICmdList, uint32 Size, EBufferUsageFlags Usage, uint32 Stride, ERHIAccess ResourceState, FRHIResourceCreateInfo& CreateInfo)
-{
-	@autoreleasepool {
-		if (CreateInfo.bWithoutNativeResource)
-		{
-			return new FMetalResourceMultiBuffer(0, BUF_None, MetalBufferUsage(BUF_None), 0, nullptr, RRT_Buffer);
-		}
-		
-		// make the RHI object, which will allocate memory
-		TRefCountPtr<FMetalResourceMultiBuffer> Buffer = new FMetalResourceMultiBuffer(Size, Usage, MetalBufferUsage(Usage), Stride, nullptr, RRT_Buffer);
-		
-		Buffer->Init_RenderThread(RHICmdList, Size, Usage, CreateInfo, Buffer);
-		
-		return Buffer.GetReference();
-	}
-}
+}
\ No newline at end of file
diff --git a/Engine/Source/Runtime/Apple/MetalRHI/Public/MetalResources.h b/Engine/Source/Runtime/Apple/MetalRHI/Public/MetalResources.h
--- a/Engine/Source/Runtime/Apple/MetalRHI/Public/MetalResources.h
+++ b/Engine/Source/Runtime/Apple/MetalRHI/Public/MetalResources.h
@@ -345,7 +345,7 @@ public:
 	/**
 	 * Initialize the buffer contents from the render-thread.
 	 */
-	void Init_RenderThread(class FRHICommandListImmediate& RHICmdList, uint32 Size, EBufferUsageFlags InUsage, FRHIResourceCreateInfo& CreateInfo, FRHIResource* Resource);
+	void Init(class FRHICommandListBase& RHICmdList, uint32 Size, EBufferUsageFlags InUsage, FRHIResourceCreateInfo& CreateInfo, FRHIResource* Resource);
 	
 	/**
 	 * Get a linear texture for given format.
@@ -471,7 +471,7 @@ private:
 class FMetalResourceMultiBuffer : public FRHIBuffer, public FMetalRHIBuffer
 {
 public:
-	FMetalResourceMultiBuffer(uint32 InSize, EBufferUsageFlags InUsage, EMetalBufferUsage InMetalUsage, uint32 InStride, FResourceArrayInterface* ResourceArray, ERHIResourceType ResourceType);
+	FMetalResourceMultiBuffer(uint32 InSize, EBufferUsageFlags InUsage, EMetalBufferUsage InMetalUsage, uint32 InStride, ERHIResourceType ResourceType);
 	virtual ~FMetalResourceMultiBuffer();
 
 	void Swap(FMetalResourceMultiBuffer& Other);
diff --git a/Engine/Source/Runtime/D3D12RHI/Private/D3D12Allocation.cpp b/Engine/Source/Runtime/D3D12RHI/Private/D3D12Allocation.cpp
--- a/Engine/Source/Runtime/D3D12RHI/Private/D3D12Allocation.cpp
+++ b/Engine/Source/Runtime/D3D12RHI/Private/D3D12Allocation.cpp
@@ -1449,6 +1449,8 @@ void FD3D12DefaultBufferAllocator::AllocDefaultResource(D3D12_HEAP_TYPE InHeapTy
 		return;
 	}
 
+	FScopeLock Lock(&CS);
+
 	// Patch out deny shader resource because it doesn't add anything for buffers and allows more pool sharing
 	// TODO: check if this is different on Xbox?
 	D3D12_RESOURCE_DESC ResourceDesc = InResourceDesc;
@@ -1478,6 +1480,8 @@ void FD3D12DefaultBufferAllocator::AllocDefaultResource(D3D12_HEAP_TYPE InHeapTy
 
 void FD3D12DefaultBufferAllocator::FreeDefaultBufferPools()
 {
+	FScopeLock Lock(&CS);
+
 	for (FD3D12BufferPool*& DefaultBufferPool : DefaultBufferPools)
 	{
 		if (DefaultBufferPool)
@@ -1495,6 +1499,7 @@ void FD3D12DefaultBufferAllocator::FreeDefaultBufferPools()
 void FD3D12DefaultBufferAllocator::BeginFrame()
 {
 #if USE_BUFFER_POOL_ALLOCATOR
+	FScopeLock Lock(&CS);
 
 	if (GD3D12VRAMBufferPoolDefrag > 0 && GD3D12VRAMBufferPoolDefragMaxCopySizePerFrame > 0)
 	{
@@ -1537,6 +1542,8 @@ void FD3D12DefaultBufferAllocator::BeginFrame()
 
 void FD3D12DefaultBufferAllocator::CleanupFreeBlocks(uint64 InFrameLag)
 {
+	FScopeLock Lock(&CS);
+
 	for (FD3D12BufferPool* DefaultBufferPool : DefaultBufferPools)
 	{
 		if (DefaultBufferPool)
@@ -1548,6 +1555,8 @@ void FD3D12DefaultBufferAllocator::CleanupFreeBlocks(uint64 InFrameLag)
 
 void FD3D12DefaultBufferAllocator::UpdateMemoryStats()
 {
+	FScopeLock Lock(&CS);
+
 	uint32 MemoryAllocated = 0;
 	uint32 MemoryUsed = 0;
 	uint32 FreeMemory = 0;
diff --git a/Engine/Source/Runtime/D3D12RHI/Private/D3D12Allocation.h b/Engine/Source/Runtime/D3D12RHI/Private/D3D12Allocation.h
--- a/Engine/Source/Runtime/D3D12RHI/Private/D3D12Allocation.h
+++ b/Engine/Source/Runtime/D3D12RHI/Private/D3D12Allocation.h
@@ -437,6 +437,7 @@ private:
 	FD3D12BufferPool* CreateBufferPool(D3D12_HEAP_TYPE InHeapType, D3D12_RESOURCE_FLAGS InResourceFlags, EBufferUsageFlags InBufferUsage, ED3D12ResourceStateMode InResourceStateMode, uint32 Alignment);
 
 	TArray<FD3D12BufferPool*> DefaultBufferPools;
+	FCriticalSection CS;
 };
 
 //-----------------------------------------------------------------------------
diff --git a/Engine/Source/Runtime/D3D12RHI/Private/D3D12Buffer.cpp b/Engine/Source/Runtime/D3D12RHI/Private/D3D12Buffer.cpp
--- a/Engine/Source/Runtime/D3D12RHI/Private/D3D12Buffer.cpp
+++ b/Engine/Source/Runtime/D3D12RHI/Private/D3D12Buffer.cpp
@@ -181,7 +181,7 @@ struct FD3D12RHICommandInitializeBuffer final : public FRHICommand<FD3D12RHIComm
 };
 
 
-void FD3D12Buffer::UploadResourceData(class FRHICommandListImmediate* RHICmdList, FResourceArrayInterface* InResourceArray, D3D12_RESOURCE_STATES InDestinationState)
+void FD3D12Buffer::UploadResourceData(FRHICommandListBase& RHICmdList, FResourceArrayInterface* InResourceArray, D3D12_RESOURCE_STATES InDestinationState)
 {
 	check(InResourceArray);
 	check(ResourceLocation.IsValid());
@@ -197,7 +197,7 @@ void FD3D12Buffer::UploadResourceData(class FRHICommandListImmediate* RHICmdList
 	}
 	else
 	{
-		const bool bOnAsyncThread = !IsInRHIThread() && !IsInRenderingThread();
+		const bool bOnAsyncThread = !IsInRHIThread() && !IsInParallelRenderingThread();
 
 		// Get an upload heap and initialize data
 		FD3D12ResourceLocation SrcResourceLoc(GetParentDevice());
@@ -214,53 +214,16 @@ void FD3D12Buffer::UploadResourceData(class FRHICommandListImmediate* RHICmdList
 		check(pData);
 		FMemory::Memcpy(pData, InResourceArray->GetResourceData(), BufferSize);
 
-		if (bOnAsyncThread)
-		{
-			// Need to update buffer content on RHI thread (immediate context) because the buffer can be a
-			// sub-allocation and its backing resource may be in a state incompatible with the copy queue.
-			// TODO:
-			// Create static buffers in COMMON state, rely on state promotion/decay to avoid transition barriers,
-			// and initialize them asynchronously on the copy queue. D3D12 buffers always allow simultaneous acess
-			// so it is legal to write to a region on the copy queue while other non-overlapping regions are
-			// being read on the graphics/compute queue. Currently, d3ddebug throws error for such usage.
-			// Once Microsoft (via Windows update) fix the debug layer, async static buffer initialization should
-			// be done on the copy queue.
-			FD3D12ResourceLocation* SrcResourceLoc_Heap = new FD3D12ResourceLocation(SrcResourceLoc.GetParentDevice());
-			FD3D12ResourceLocation::TransferOwnership(*SrcResourceLoc_Heap, SrcResourceLoc);
-			TRefCountPtr<FD3D12Buffer> ThisRef = this;
-			ENQUEUE_RENDER_COMMAND(CmdD3D12InitializeBuffer)(
-				[DestinationBuffer = MoveTemp(ThisRef), SrcResourceLoc_Heap, BufferSize, InDestinationState](FRHICommandListImmediate& RHICmdList) mutable
-				{
-					if (RHICmdList.Bypass())
-					{
-						FD3D12RHICommandInitializeBuffer Command(MoveTemp(DestinationBuffer), *SrcResourceLoc_Heap, BufferSize, InDestinationState);
-						Command.Execute(RHICmdList);
-					}
-					else
-					{
-						new (RHICmdList.AllocCommand<FD3D12RHICommandInitializeBuffer>()) FD3D12RHICommandInitializeBuffer(MoveTemp(DestinationBuffer), *SrcResourceLoc_Heap, BufferSize, InDestinationState);
-					}
-					delete SrcResourceLoc_Heap;
-				});
-		}
-		else if (!RHICmdList || RHICmdList->Bypass())
+		if (RHICmdList.IsBottomOfPipe())
 		{
 			// On RHIT or RT (when bypassing), we can access immediate context directly
 			FD3D12RHICommandInitializeBuffer Command(this, SrcResourceLoc, BufferSize, InDestinationState);
-			if (RHICmdList)
-			{
-				Command.Execute(*RHICmdList);
-			}
-			else
-			{
-				FD3D12CommandContext& CommandContext = GetParentDevice()->GetDefaultCommandContext();
-				Command.ExecuteOnCommandContext(CommandContext);
-			}
+			FD3D12CommandContext& CommandContext = GetParentDevice()->GetDefaultCommandContext();
+			Command.ExecuteOnCommandContext(CommandContext);
 		}
 		else
 		{
-			// On RT but not bypassing
-			new (RHICmdList->AllocCommand<FD3D12RHICommandInitializeBuffer>()) FD3D12RHICommandInitializeBuffer(this, SrcResourceLoc, BufferSize, InDestinationState);
+			new (RHICmdList.AllocCommand<FD3D12RHICommandInitializeBuffer>()) FD3D12RHICommandInitializeBuffer(this, SrcResourceLoc, BufferSize, InDestinationState);
 		}
 	}
 
@@ -546,17 +509,12 @@ void FD3D12Buffer::GetResourceDescAndAlignment(uint64 InSize, uint32 InStride, E
 	Alignment = (InStride > 0) && (EnumHasAnyFlags(InUsage, BUF_StructuredBuffer) || !EnumHasAnyFlags(InUsage, BUF_ByteAddressBuffer | BUF_DrawIndirect)) ? InStride : 4;
 }
 
-FBufferRHIRef FD3D12DynamicRHI::RHICreateBuffer(uint32 Size, EBufferUsageFlags Usage, uint32 Stride, ERHIAccess InResourceState, FRHIResourceCreateInfo& CreateInfo)
-{	
-	return CreateBuffer(nullptr, Size, Usage, Stride, InResourceState, CreateInfo);
-}
-
-FBufferRHIRef FD3D12DynamicRHI::CreateBuffer_RenderThread(class FRHICommandListImmediate& RHICmdList, uint32 Size, EBufferUsageFlags Usage, uint32 Stride, ERHIAccess ResourceState, FRHIResourceCreateInfo& CreateInfo)
+FBufferRHIRef FD3D12DynamicRHI::RHICreateBuffer(FRHICommandListBase& RHICmdList, uint32 Size, EBufferUsageFlags Usage, uint32 Stride, ERHIAccess InResourceState, FRHIResourceCreateInfo& CreateInfo)
 {
-	return CreateBuffer(&RHICmdList, Size, Usage, Stride, ResourceState, CreateInfo);
+	return CreateBuffer(RHICmdList, Size, Usage, Stride, InResourceState, CreateInfo);
 }
 
-FBufferRHIRef FD3D12DynamicRHI::CreateBuffer(FRHICommandListImmediate* RHICmdList, uint32 Size, EBufferUsageFlags Usage, uint32 Stride, ERHIAccess InResourceState, FRHIResourceCreateInfo& CreateInfo)
+FBufferRHIRef FD3D12DynamicRHI::CreateBuffer(FRHICommandListBase& RHICmdList, uint32 Size, EBufferUsageFlags Usage, uint32 Stride, ERHIAccess InResourceState, FRHIResourceCreateInfo& CreateInfo)
 {
 	if (CreateInfo.bWithoutNativeResource)
 	{
@@ -566,10 +524,10 @@ FBufferRHIRef FD3D12DynamicRHI::CreateBuffer(FRHICommandListImmediate* RHICmdLis
 			});
 	}
 
-	return CreateD3D12Buffer(RHICmdList, Size, Usage, Stride, InResourceState, CreateInfo);
+	return CreateD3D12Buffer(&RHICmdList, Size, Usage, Stride, InResourceState, CreateInfo);
 }
 
-FD3D12Buffer* FD3D12DynamicRHI::CreateD3D12Buffer(class FRHICommandListImmediate* RHICmdList, uint32 Size, EBufferUsageFlags Usage, uint32 Stride, ERHIAccess InResourceState, FRHIResourceCreateInfo& CreateInfo, ID3D12ResourceAllocator* ResourceAllocator)
+FD3D12Buffer* FD3D12DynamicRHI::CreateD3D12Buffer(class FRHICommandListBase* RHICmdList, uint32 Size, EBufferUsageFlags Usage, uint32 Stride, ERHIAccess InResourceState, FRHIResourceCreateInfo& CreateInfo, ID3D12ResourceAllocator* ResourceAllocator)
 {
 	D3D12_RESOURCE_DESC Desc;
 	uint32 Alignment;
@@ -594,14 +552,14 @@ FD3D12Buffer* FD3D12DynamicRHI::CreateD3D12Buffer(class FRHICommandListImmediate
 	D3D12_RESOURCE_STATES CreateState = (CreateInfo.ResourceArray && bSupportResourceStateTracking) ? D3D12_RESOURCE_STATE_COPY_DEST : DesiredState;
 	bool bHasInitialData = CreateInfo.ResourceArray != nullptr;
 
-	FD3D12SyncPoint CopyQueueSyncPoint;
 	FD3D12Buffer* Buffer = GetAdapter().CreateRHIBuffer(Desc, Alignment, Stride, Size, Usage, StateMode, CreateState, bHasInitialData, CreateInfo.GPUMask, ResourceAllocator, CreateInfo.DebugName);
 	check(Buffer->ResourceLocation.IsValid());
 
 	// Copy the resource data if available 
 	if (bHasInitialData)
 	{
-		Buffer->UploadResourceData(RHICmdList, CreateInfo.ResourceArray, DesiredState);
+		check(RHICmdList);
+		Buffer->UploadResourceData(*RHICmdList, CreateInfo.ResourceArray, DesiredState);
 	}
 
 	return Buffer;
@@ -613,7 +571,7 @@ FRHIBuffer* FD3D12DynamicRHI::CreateBuffer(const FRHIBufferCreateInfo& CreateInf
 	return CreateD3D12Buffer(nullptr, CreateInfo.Size, CreateInfo.Usage, CreateInfo.Stride, InitialState, ResourceCreateInfo, ResourceAllocator);
 }
 
-void* FD3D12DynamicRHI::LockBuffer(FRHICommandListImmediate* RHICmdList, FD3D12Buffer* Buffer, uint32 BufferSize, EBufferUsageFlags BufferUsage, uint32 Offset, uint32 Size, EResourceLockMode LockMode)
+void* FD3D12DynamicRHI::LockBuffer(FRHICommandListBase& RHICmdList, FD3D12Buffer* Buffer, uint32 BufferSize, EBufferUsageFlags BufferUsage, uint32 Offset, uint32 Size, EResourceLockMode LockMode)
 {
 	SCOPE_CYCLE_COUNTER(STAT_D3D12LockBufferTime);
 
@@ -641,18 +599,17 @@ void* FD3D12DynamicRHI::LockBuffer(FRHICommandListImmediate* RHICmdList, FD3D12B
 			FD3D12Device* Device = Buffer->GetParentDevice();
 
 			// If on the RenderThread, queue up a command on the RHIThread to rename this buffer at the correct time
-			if (ShouldDeferBufferLockOperation(RHICmdList) && LockMode == RLM_WriteOnly)
+			if (RHICmdList.IsTopOfPipe() && LockMode == RLM_WriteOnly)
 			{
-				FRHICommandRenameUploadBuffer* Command = ALLOC_COMMAND_CL(*RHICmdList, FRHICommandRenameUploadBuffer)(Buffer, Device);
-
+				FRHICommandRenameUploadBuffer* Command = ALLOC_COMMAND_CL(RHICmdList, FRHICommandRenameUploadBuffer)(Buffer, Device);
 				Data = Adapter.GetUploadHeapAllocator(Device->GetGPUIndex()).AllocUploadResource(BufferSize, Buffer->BufferAlignment, Command->NewLocation);
-				RHICmdList->RHIThreadFence(true);
+				RHICmdList.RHIThreadFence(true);
 			}
 			else
 			{
 				FRHICommandRenameUploadBuffer Command(Buffer, Device);
 				Data = Adapter.GetUploadHeapAllocator(Device->GetGPUIndex()).AllocUploadResource(BufferSize, Buffer->BufferAlignment, Command.NewLocation);
-				Command.Execute(*RHICmdList);
+				Command.Execute(RHICmdList);
 			}
 		}
 	}
@@ -665,6 +622,8 @@ void* FD3D12DynamicRHI::LockBuffer(FRHICommandListImmediate* RHICmdList, FD3D12B
 		// Locking for read must occur immediately so we can't queue up the operations later.
 		if (LockMode == RLM_ReadOnly)
 		{
+			FRHICommandListImmediate& RHICmdListImmediate = RHICmdList.GetAsImmediate();
+
 			LockedData.bLockedForReadOnly = true;
 			// If the static buffer is being locked for reading, create a staging buffer.
 			FD3D12Resource* StagingBuffer = nullptr;
@@ -698,12 +657,12 @@ void* FD3D12DynamicRHI::LockBuffer(FRHICommandListImmediate* RHICmdList, FD3D12B
 					DefaultContext.FlushCommands(true);
 				};
 
-				if (ShouldDeferBufferLockOperation(RHICmdList))
+				if (RHICmdListImmediate.IsTopOfPipe())
 				{
 					// Sync when in the render thread implementation
 					check(IsInRHIThread() == false);
 
-					RHICmdList->ImmediateFlush(EImmediateFlushType::FlushRHIThread);
+					RHICmdListImmediate.ImmediateFlush(EImmediateFlushType::FlushRHIThread);
 					pfnCopyContents();
 				}
 				else
@@ -733,7 +692,7 @@ void* FD3D12DynamicRHI::LockBuffer(FRHICommandListImmediate* RHICmdList, FD3D12B
 	return Data;
 }
 
-void FD3D12DynamicRHI::UnlockBuffer(FRHICommandListImmediate* RHICmdList, FD3D12Buffer* Buffer, EBufferUsageFlags BufferUsage)
+void FD3D12DynamicRHI::UnlockBuffer(FRHICommandListBase& RHICmdList, FD3D12Buffer* Buffer, EBufferUsageFlags BufferUsage)
 {
 	SCOPE_CYCLE_COUNTER(STAT_D3D12UnlockBufferTime);
 
@@ -760,12 +719,12 @@ void FD3D12DynamicRHI::UnlockBuffer(FRHICommandListImmediate* RHICmdList, FD3D12
 			for (FD3D12Buffer::FLinkedObjectIterator CurrentBuffer(Buffer); CurrentBuffer; ++CurrentBuffer)
 			{
 				// If we are on the render thread, queue up the copy on the RHIThread so it happens at the correct time.
-				if (ShouldDeferBufferLockOperation(RHICmdList))
+				if (RHICmdList.IsTopOfPipe())
 				{
 					if (CurrentBuffer.Get() == LastBuffer)
 					{
 						// Command associated with last buffer (will be only buffer if single GPU) receives ownership of locked data
-						ALLOC_COMMAND_CL(*RHICmdList, FRHICommandUpdateBuffer)(&CurrentBuffer->ResourceLocation, LockedData.ResourceLocation, LockedData.LockedOffset, LockedData.LockedPitch);
+						ALLOC_COMMAND_CL(RHICmdList, FRHICommandUpdateBuffer)(&CurrentBuffer->ResourceLocation, LockedData.ResourceLocation, LockedData.LockedOffset, LockedData.LockedPitch);
 					}
 					else
 					{
@@ -773,7 +732,7 @@ void FD3D12DynamicRHI::UnlockBuffer(FRHICommandListImmediate* RHICmdList, FD3D12
 						// last command handling clean up the locked data after it has been propagated to all GPUs.
 						FD3D12ResourceLocation NodeResourceLocation(LockedData.ResourceLocation.GetParentDevice());
 						FD3D12ResourceLocation::ReferenceNode(NodeResourceLocation.GetParentDevice(), NodeResourceLocation, LockedData.ResourceLocation);
-						ALLOC_COMMAND_CL(*RHICmdList, FRHICommandUpdateBuffer)(&CurrentBuffer->ResourceLocation, NodeResourceLocation, LockedData.LockedOffset, LockedData.LockedPitch);
+						ALLOC_COMMAND_CL(RHICmdList, FRHICommandUpdateBuffer)(&CurrentBuffer->ResourceLocation, NodeResourceLocation, LockedData.LockedOffset, LockedData.LockedPitch);
 					}
 				}
 				else
@@ -787,7 +746,7 @@ void FD3D12DynamicRHI::UnlockBuffer(FRHICommandListImmediate* RHICmdList, FD3D12
 	LockedData.Reset();
 }
 
-void* FD3D12DynamicRHI::RHILockBuffer(FRHICommandListImmediate& RHICmdList, FRHIBuffer* BufferRHI, uint32 Offset, uint32 Size, EResourceLockMode LockMode)
+void* FD3D12DynamicRHI::RHILockBuffer(FRHICommandListBase& RHICmdList, FRHIBuffer* BufferRHI, uint32 Offset, uint32 Size, EResourceLockMode LockMode)
 {
 	// If you hit this assert, you should be using LockBufferMGPU and iterating over FRHIGPUMask::All() to initialize the resource separately for each GPU.
 	// "MultiGPUAllocate" only makes sense if a buffer must vary per GPU, for example if it's a buffer that includes GPU specific virtual addresses for ray
@@ -795,10 +754,10 @@ void* FD3D12DynamicRHI::RHILockBuffer(FRHICommandListImmediate& RHICmdList, FRHI
 	check(!EnumHasAnyFlags(BufferRHI->GetUsage(), BUF_MultiGPUAllocate));
 
 	FD3D12Buffer* Buffer = FD3D12DynamicRHI::ResourceCast(BufferRHI);
-	return LockBuffer(&RHICmdList, Buffer, Buffer->GetSize(), Buffer->GetUsage(), Offset, Size, LockMode);
+	return LockBuffer(RHICmdList, Buffer, Buffer->GetSize(), Buffer->GetUsage(), Offset, Size, LockMode);
 }
 
-void* FD3D12DynamicRHI::RHILockBufferMGPU(FRHICommandListImmediate& RHICmdList, FRHIBuffer* BufferRHI, uint32 GPUIndex, uint32 Offset, uint32 Size, EResourceLockMode LockMode)
+void* FD3D12DynamicRHI::RHILockBufferMGPU(FRHICommandListBase& RHICmdList, FRHIBuffer* BufferRHI, uint32 GPUIndex, uint32 Offset, uint32 Size, EResourceLockMode LockMode)
 {
 	// If you hit this assert, you should be using LockBuffer to initialize the resource, rather than this function.  The MGPU version is only for resources
 	// with the MultiGPUAllocate flag, where it's necessary for the caller to initialize the buffer for each GPU.  The other LockBuffer call initializes the
@@ -806,23 +765,23 @@ void* FD3D12DynamicRHI::RHILockBufferMGPU(FRHICommandListImmediate& RHICmdList,
 	check(EnumHasAnyFlags(BufferRHI->GetUsage(), BUF_MultiGPUAllocate));
 
 	FD3D12Buffer* Buffer = FD3D12DynamicRHI::ResourceCast(BufferRHI, GPUIndex);
-	return LockBuffer(&RHICmdList, Buffer, Buffer->GetSize(), Buffer->GetUsage(), Offset, Size, LockMode);
+	return LockBuffer(RHICmdList, Buffer, Buffer->GetSize(), Buffer->GetUsage(), Offset, Size, LockMode);
 }
 
-void FD3D12DynamicRHI::RHIUnlockBuffer(FRHICommandListImmediate& RHICmdList, FRHIBuffer* BufferRHI)
+void FD3D12DynamicRHI::RHIUnlockBuffer(FRHICommandListBase& RHICmdList, FRHIBuffer* BufferRHI)
 {
 	check(!EnumHasAnyFlags(BufferRHI->GetUsage(), BUF_MultiGPUAllocate));
 
 	FD3D12Buffer* Buffer = FD3D12DynamicRHI::ResourceCast(BufferRHI);
-	UnlockBuffer(&RHICmdList, Buffer, Buffer->GetUsage());
+	UnlockBuffer(RHICmdList, Buffer, Buffer->GetUsage());
 }
 
-void FD3D12DynamicRHI::RHIUnlockBufferMGPU(FRHICommandListImmediate& RHICmdList, FRHIBuffer* BufferRHI, uint32 GPUIndex)
+void FD3D12DynamicRHI::RHIUnlockBufferMGPU(FRHICommandListBase& RHICmdList, FRHIBuffer* BufferRHI, uint32 GPUIndex)
 {
 	check(EnumHasAnyFlags(BufferRHI->GetUsage(), BUF_MultiGPUAllocate));
 
 	FD3D12Buffer* Buffer = FD3D12DynamicRHI::ResourceCast(BufferRHI, GPUIndex);
-	UnlockBuffer(&RHICmdList, Buffer, Buffer->GetUsage());
+	UnlockBuffer(RHICmdList, Buffer, Buffer->GetUsage());
 }
 
 void FD3D12DynamicRHI::RHITransferBufferUnderlyingResource(FRHIBuffer* DestBuffer, FRHIBuffer* SrcBuffer)
diff --git a/Engine/Source/Runtime/D3D12RHI/Private/D3D12RHI.cpp b/Engine/Source/Runtime/D3D12RHI/Private/D3D12RHI.cpp
--- a/Engine/Source/Runtime/D3D12RHI/Private/D3D12RHI.cpp
+++ b/Engine/Source/Runtime/D3D12RHI/Private/D3D12RHI.cpp
@@ -82,6 +82,7 @@ FD3D12DynamicRHI::FD3D12DynamicRHI(const TArray<TSharedPtr<FD3D12Adapter>>& Chos
 #endif // PLATFORM_WINDOWS
 
 	GRHISupportsMultithreading = true;
+	GRHISupportsMultithreadedResources = true;
 	GRHIMultiPipelineMergeableAccessMask = GRHIMergeableAccessMask;
 	EnumRemoveFlags(GRHIMultiPipelineMergeableAccessMask, ERHIAccess::UAVMask);
 
diff --git a/Engine/Source/Runtime/D3D12RHI/Private/D3D12RHIPrivate.h b/Engine/Source/Runtime/D3D12RHI/Private/D3D12RHIPrivate.h
--- a/Engine/Source/Runtime/D3D12RHI/Private/D3D12RHIPrivate.h
+++ b/Engine/Source/Runtime/D3D12RHI/Private/D3D12RHIPrivate.h
@@ -339,13 +339,13 @@ public:
 	virtual void RHIReleaseTransition(FRHITransition* Transition) final override;
 	virtual FUniformBufferRHIRef RHICreateUniformBuffer(const void* Contents, const FRHIUniformBufferLayout* Layout, EUniformBufferUsage Usage, EUniformBufferValidation Validation) final override;
 	virtual void RHIUpdateUniformBuffer(FRHIUniformBuffer* UniformBufferRHI, const void* Contents) final override;
-	virtual FBufferRHIRef RHICreateBuffer(uint32 Size, EBufferUsageFlags Usage, uint32 Stride, ERHIAccess ResourceState, FRHIResourceCreateInfo& CreateInfo) final override;
+	virtual FBufferRHIRef RHICreateBuffer(FRHICommandListBase& RHICmdList, uint32 Size, EBufferUsageFlags Usage, uint32 Stride, ERHIAccess ResourceState, FRHIResourceCreateInfo& CreateInfo) final override;
 	virtual void RHITransferBufferUnderlyingResource(FRHIBuffer* DestBuffer, FRHIBuffer* SrcBuffer) final override;
 	virtual void RHICopyBuffer(FRHIBuffer* SourceBuffer, FRHIBuffer* DestBuffer) final override;
-	virtual void* RHILockBuffer(FRHICommandListImmediate& RHICmdList, FRHIBuffer* Buffer, uint32 Offset, uint32 Size, EResourceLockMode LockMode) final override;
-	virtual void* RHILockBufferMGPU(FRHICommandListImmediate& RHICmdList, FRHIBuffer* Buffer, uint32 GPUIndex, uint32 Offset, uint32 Size, EResourceLockMode LockMode) final override;
-	virtual void RHIUnlockBuffer(FRHICommandListImmediate& RHICmdList, FRHIBuffer* Buffer) final override;
-	virtual void RHIUnlockBufferMGPU(FRHICommandListImmediate& RHICmdList, FRHIBuffer* Buffer, uint32 GPUIndex) final override;
+	virtual void* RHILockBuffer(FRHICommandListBase& RHICmdList, FRHIBuffer* Buffer, uint32 Offset, uint32 Size, EResourceLockMode LockMode) final override;
+	virtual void* RHILockBufferMGPU(FRHICommandListBase& RHICmdList, FRHIBuffer* Buffer, uint32 GPUIndex, uint32 Offset, uint32 Size, EResourceLockMode LockMode) final override;
+	virtual void RHIUnlockBuffer(FRHICommandListBase& RHICmdList, FRHIBuffer* Buffer) final override;
+	virtual void RHIUnlockBufferMGPU(FRHICommandListBase& RHICmdList, FRHIBuffer* Buffer, uint32 GPUIndex) final override;
 	virtual FUnorderedAccessViewRHIRef RHICreateUnorderedAccessView(FRHIBuffer* Buffer, bool bUseUAVCounter, bool bAppendBuffer) final override;
 	virtual FUnorderedAccessViewRHIRef RHICreateUnorderedAccessView(FRHITexture* Texture, uint32 MipLevel, uint16 FirstArraySlice, uint16 NumArraySlices) final override;
 	virtual FUnorderedAccessViewRHIRef RHICreateUnorderedAccessView(FRHITexture* Texture, uint32 MipLevel, uint8 Format, uint16 FirstArraySlice, uint16 NumArraySlices) final override;
@@ -484,10 +484,6 @@ public:
 		return RHICreateComputeShader(Code, Hash);
 	}
 
-	virtual FBufferRHIRef CreateBuffer_RenderThread(class FRHICommandListImmediate& RHICmdList, uint32 Size, EBufferUsageFlags Usage, uint32 Stride, ERHIAccess ResourceState, FRHIResourceCreateInfo& CreateInfo) override final;
-	// virtual FTexture2DRHIRef AsyncReallocateTexture2D_RenderThread(class FRHICommandListImmediate& RHICmdList, FRHITexture2D* Texture2D, int32 NewMipCount, int32 NewSizeX, int32 NewSizeY, FThreadSafeCounter* RequestStatus) override final;
-	// virtual ETextureReallocationStatus FinalizeAsyncReallocateTexture2D_RenderThread(class FRHICommandListImmediate& RHICmdList, FRHITexture2D* Texture2D, bool bBlockUntilCompleted) override final;
-	// virtual ETextureReallocationStatus CancelAsyncReallocateTexture2D_RenderThread(class FRHICommandListImmediate& RHICmdList, FRHITexture2D* Texture2D, bool bBlockUntilCompleted) override final;
 	virtual void UpdateTexture2D_RenderThread(class FRHICommandListImmediate& RHICmdList, FRHITexture2D* Texture, uint32 MipIndex, const struct FUpdateTextureRegion2D& UpdateRegion, uint32 SourcePitch, const uint8* SourceData) override final;
 	virtual void* LockTexture2D_RenderThread(class FRHICommandListImmediate& RHICmdList, FRHITexture2D* Texture, uint32 MipIndex, EResourceLockMode LockMode, uint32& DestStride, bool bLockWithinMiptail, bool bNeedsDefaultRHIFlush = true) override final;
 	virtual void UnlockTexture2D_RenderThread(class FRHICommandListImmediate& RHICmdList, FRHITexture2D* Texture, uint32 MipIndex, bool bLockWithinMiptail, bool bNeedsDefaultRHIFlush = true) override final;
@@ -549,11 +545,11 @@ public:
 
 	static int32 GetResourceBarrierBatchSizeLimit();
 
-	FBufferRHIRef CreateBuffer(FRHICommandListImmediate* RHICmdList, uint32 Size, EBufferUsageFlags Usage, uint32 Stride, ERHIAccess ResourceState, FRHIResourceCreateInfo& CreateInfo);
-	void* LockBuffer(FRHICommandListImmediate* RHICmdList, FD3D12Buffer* Buffer, uint32 BufferSize, EBufferUsageFlags BufferUsage, uint32 Offset, uint32 Size, EResourceLockMode LockMode);
-	void UnlockBuffer(FRHICommandListImmediate* RHICmdList, FD3D12Buffer* Buffer, EBufferUsageFlags BufferUsage);
+	FBufferRHIRef CreateBuffer(FRHICommandListBase& RHICmdList, uint32 Size, EBufferUsageFlags Usage, uint32 Stride, ERHIAccess ResourceState, FRHIResourceCreateInfo& CreateInfo);
+	void* LockBuffer(FRHICommandListBase& RHICmdList, FD3D12Buffer* Buffer, uint32 BufferSize, EBufferUsageFlags BufferUsage, uint32 Offset, uint32 Size, EResourceLockMode LockMode);
+	void UnlockBuffer(FRHICommandListBase& RHICmdList, FD3D12Buffer* Buffer, EBufferUsageFlags BufferUsage);
 
-	static inline bool ShouldDeferBufferLockOperation(FRHICommandListImmediate* RHICmdList)
+	static inline bool ShouldDeferBufferLockOperation(FRHICommandListBase* RHICmdList)
 	{
 		if (RHICmdList == nullptr)
 		{
@@ -1197,7 +1193,7 @@ public:
 	virtual FD3D12ResourceDesc GetResourceDesc(const FRHITextureDesc& CreateInfo) const;
 
 	virtual FD3D12Texture* CreateD3D12Texture(const FRHITextureCreateDesc& CreateDesc, class FRHICommandListImmediate* RHICmdList, ID3D12ResourceAllocator* ResourceAllocator = nullptr);
-	FD3D12Buffer* CreateD3D12Buffer(class FRHICommandListImmediate* RHICmdList, uint32 Size, EBufferUsageFlags Usage, uint32 Stride, ERHIAccess ResourceState, FRHIResourceCreateInfo& CreateInfo, ID3D12ResourceAllocator* ResourceAllocator = nullptr);
+	FD3D12Buffer* CreateD3D12Buffer(class FRHICommandListBase* RHICmdList, uint32 Size, EBufferUsageFlags Usage, uint32 Stride, ERHIAccess ResourceState, FRHIResourceCreateInfo& CreateInfo, ID3D12ResourceAllocator* ResourceAllocator = nullptr);
 	virtual FD3D12Texture* CreateNewD3D12Texture(const FRHITextureCreateDesc& CreateDesc, class FD3D12Device* Device);
 
 	FRHIBuffer* CreateBuffer(const FRHIBufferCreateInfo& CreateInfo, const TCHAR* DebugName, ERHIAccess InitialState, ID3D12ResourceAllocator* ResourceAllocator);
diff --git a/Engine/Source/Runtime/D3D12RHI/Public/D3D12Resources.h b/Engine/Source/Runtime/D3D12RHI/Public/D3D12Resources.h
--- a/Engine/Source/Runtime/D3D12RHI/Public/D3D12Resources.h
+++ b/Engine/Source/Runtime/D3D12RHI/Public/D3D12Resources.h
@@ -918,7 +918,7 @@ public:
 
 	virtual uint32 GetParentGPUIndex() const override;
 
-	void UploadResourceData(class FRHICommandListImmediate* InRHICmdList, FResourceArrayInterface* InResourceArray, D3D12_RESOURCE_STATES InDestinationState);
+	void UploadResourceData(FRHICommandListBase& InRHICmdList, FResourceArrayInterface* InResourceArray, D3D12_RESOURCE_STATES InDestinationState);
 	FD3D12SyncPoint UploadResourceDataViaCopyQueue(FResourceArrayInterface* InResourceArray);
 
 	// FRHIResource overrides
diff --git a/Engine/Source/Runtime/Engine/Private/RawIndexBuffer.cpp b/Engine/Source/Runtime/Engine/Private/RawIndexBuffer.cpp
--- a/Engine/Source/Runtime/Engine/Private/RawIndexBuffer.cpp
+++ b/Engine/Source/Runtime/Engine/Private/RawIndexBuffer.cpp
@@ -331,7 +331,7 @@ FBufferRHIRef FRawStaticIndexBuffer::CreateRHIBuffer_Internal()
 		// This also avoid setting lots of states on all the members of all the different buffers used by meshes. Follow up: https://jira.it.epicgames.net/browse/UE-69376.
 		bSRV |= IndexStorage.GetAllowCPUAccess();
 
-		const EBufferUsageFlags BufferFlags = BUF_Static | (bSRV ? BUF_ShaderResource : BUF_None);
+		const EBufferUsageFlags BufferFlags = EBufferUsageFlags::Static | (bSRV ? EBufferUsageFlags::ShaderResource : EBufferUsageFlags::None);
 
 		// Create the index buffer.
 		FRHIResourceCreateInfo CreateInfo(Is32Bit() ? TEXT("FRawStaticIndexBuffer32") : TEXT("FRawStaticIndexBuffer16"), &IndexStorage);
@@ -342,7 +342,8 @@ FBufferRHIRef FRawStaticIndexBuffer::CreateRHIBuffer_Internal()
 		}
 		else
 		{
-			return RHIAsyncCreateIndexBuffer(IndexStride, SizeInBytes, BufferFlags, CreateInfo);
+			FRHIAsyncCommandList CommandList;
+			return CommandList->CreateBuffer(SizeInBytes, BufferFlags | EBufferUsageFlags::IndexBuffer, IndexStride, ERHIAccess::SRVMask, CreateInfo);
 		}
 	}
 	return nullptr;
diff --git a/Engine/Source/Runtime/Engine/Public/RawIndexBuffer.h b/Engine/Source/Runtime/Engine/Public/RawIndexBuffer.h
--- a/Engine/Source/Runtime/Engine/Public/RawIndexBuffer.h
+++ b/Engine/Source/Runtime/Engine/Public/RawIndexBuffer.h
@@ -548,12 +548,12 @@ private:
 		{
 			// Create the index buffer.
 			FRHIResourceCreateInfo CreateInfo(sizeof(INDEX_TYPE) == 4 ? TEXT("FRawStaticIndexBuffer32") : TEXT("FRawStaticIndexBuffer16"), &Indices);
-			EBufferUsageFlags Flags = BUF_Static;
+			EBufferUsageFlags Flags = EBufferUsageFlags::Static;
 
 			if (IsSRVNeeded())
 			{
 				// BUF_ShaderResource is needed for SkinCache RecomputeSkinTangents
-				Flags = (EBufferUsageFlags)(Flags | BUF_ShaderResource);
+				Flags |= EBufferUsageFlags::ShaderResource;
 			}
 
 			// Need to cache number of indices from the source array *before* RHICreateIndexBuffer is called
@@ -569,7 +569,8 @@ private:
 			}
 			else
 			{
-				Ret = RHIAsyncCreateIndexBuffer(sizeof(INDEX_TYPE), Size, Flags, CreateInfo);
+				FRHIAsyncCommandList CommandList;
+				return CommandList->CreateBuffer(Size, Flags | EBufferUsageFlags::IndexBuffer, sizeof(INDEX_TYPE), ERHIAccess::SRVMask, CreateInfo);
 			}
 
 			return Ret;
diff --git a/Engine/Source/Runtime/NullDrv/Public/NullRHI.h b/Engine/Source/Runtime/NullDrv/Public/NullRHI.h
--- a/Engine/Source/Runtime/NullDrv/Public/NullRHI.h
+++ b/Engine/Source/Runtime/NullDrv/Public/NullRHI.h
@@ -115,7 +115,7 @@ public:
 
 	}
 
-	virtual FBufferRHIRef RHICreateBuffer(uint32 Size, EBufferUsageFlags Usage, uint32 Stride, ERHIAccess ResourceState , FRHIResourceCreateInfo& CreateInfo) final override
+	virtual FBufferRHIRef RHICreateBuffer(FRHICommandListBase& RHICmdList, uint32 Size, EBufferUsageFlags Usage, uint32 Stride, ERHIAccess ResourceState , FRHIResourceCreateInfo& CreateInfo) final override
 	{ 
 		if(CreateInfo.ResourceArray) 
 		{ 
@@ -129,12 +129,12 @@ public:
 
 	}
 
-	virtual void* LockBuffer_BottomOfPipe(FRHICommandListImmediate& RHICmdList, FRHIBuffer* Buffer, uint32 Offset, uint32 SizeRHI, EResourceLockMode LockMode) final override
-	{ 
+	virtual void* LockBuffer_BottomOfPipe(FRHICommandListBase& RHICmdList, FRHIBuffer* Buffer, uint32 Offset, uint32 SizeRHI, EResourceLockMode LockMode) final override
+	{
 		return GetStaticBuffer(Buffer->GetSize());
 	}
 
-	virtual void UnlockBuffer_BottomOfPipe(FRHICommandListImmediate& RHICmdList, FRHIBuffer* Buffer) final override
+	virtual void UnlockBuffer_BottomOfPipe(FRHICommandListBase& RHICmdList, FRHIBuffer* Buffer) final override
 	{
 
 	}
diff --git a/Engine/Source/Runtime/OpenGLDrv/Private/OpenGLBuffer.cpp b/Engine/Source/Runtime/OpenGLDrv/Private/OpenGLBuffer.cpp
--- a/Engine/Source/Runtime/OpenGLDrv/Private/OpenGLBuffer.cpp
+++ b/Engine/Source/Runtime/OpenGLDrv/Private/OpenGLBuffer.cpp
@@ -176,7 +176,7 @@ void BeginFrame_VertexBufferCleanup()
 	FrameBytes = 0;
 }
 
-FBufferRHIRef FOpenGLDynamicRHI::RHICreateBuffer(uint32 Size, EBufferUsageFlags Usage, uint32 Stride, ERHIAccess InResourceState, FRHIResourceCreateInfo& CreateInfo)
+FBufferRHIRef FOpenGLDynamicRHI::RHICreateBuffer(FRHICommandListBase& RHICmdList, uint32 Size, EBufferUsageFlags Usage, uint32 Stride, ERHIAccess InResourceState, FRHIResourceCreateInfo& CreateInfo)
 {
 	if (CreateInfo.bWithoutNativeResource)
 	{
@@ -202,7 +202,7 @@ FBufferRHIRef FOpenGLDynamicRHI::RHICreateBuffer(uint32 Size, EBufferUsageFlags
 		BufferType = GL_ELEMENT_ARRAY_BUFFER;
 	}
 
-	TRefCountPtr<FOpenGLBuffer> Buffer = new FOpenGLBuffer(BufferType, Stride, Size, Usage, Data);
+	TRefCountPtr<FOpenGLBuffer> Buffer = new FOpenGLBuffer(&RHICmdList, BufferType, Stride, Size, Usage, Data);
 	
 	if (CreateInfo.ResourceArray)
 	{
@@ -212,7 +212,7 @@ FBufferRHIRef FOpenGLDynamicRHI::RHICreateBuffer(uint32 Size, EBufferUsageFlags
 	return Buffer.GetReference();
 }
 
-void* FOpenGLDynamicRHI::LockBuffer_BottomOfPipe(FRHICommandListImmediate& RHICmdList, FRHIBuffer* BufferRHI, uint32 Offset, uint32 Size, EResourceLockMode LockMode)
+void* FOpenGLDynamicRHI::LockBuffer_BottomOfPipe(FRHICommandListBase& RHICmdList, FRHIBuffer* BufferRHI, uint32 Offset, uint32 Size, EResourceLockMode LockMode)
 {
 	check(Size > 0);
 	RHITHREAD_GLCOMMAND_PROLOGUE();
@@ -234,7 +234,7 @@ void* FOpenGLDynamicRHI::LockBuffer_BottomOfPipe(FRHICommandListImmediate& RHICm
 	RHITHREAD_GLCOMMAND_EPILOGUE_RETURN(void*);
 }
 
-void FOpenGLDynamicRHI::UnlockBuffer_BottomOfPipe(FRHICommandListImmediate& RHICmdList, FRHIBuffer* BufferRHI)
+void FOpenGLDynamicRHI::UnlockBuffer_BottomOfPipe(FRHICommandListBase& RHICmdList, FRHIBuffer* BufferRHI)
 {
 	RHITHREAD_GLCOMMAND_PROLOGUE();
 	VERIFY_GL_SCOPE();
diff --git a/Engine/Source/Runtime/OpenGLDrv/Private/OpenGLCommands.cpp b/Engine/Source/Runtime/OpenGLDrv/Private/OpenGLCommands.cpp
--- a/Engine/Source/Runtime/OpenGLDrv/Private/OpenGLCommands.cpp
+++ b/Engine/Source/Runtime/OpenGLDrv/Private/OpenGLCommands.cpp
@@ -902,7 +902,7 @@ void FOpenGLDynamicRHI::UpdateSRV(FOpenGLShaderResourceView* SRV)
 		uint32 SizeY = Texture->GetSizeY();
 		
 		uint32 MipBytes = SizeX * SizeY;
-		TRefCountPtr<FOpenGLPixelBuffer> PixelBuffer = new FOpenGLPixelBuffer(GL_PIXEL_UNPACK_BUFFER, 0, MipBytes, BUF_Dynamic, nullptr);
+		TRefCountPtr<FOpenGLPixelBuffer> PixelBuffer = new FOpenGLPixelBuffer(nullptr, GL_PIXEL_UNPACK_BUFFER, 0, MipBytes, BUF_Dynamic, nullptr);
 		
 		glBindBuffer( GL_PIXEL_PACK_BUFFER, 0 );
 		glBindBuffer( GL_PIXEL_PACK_BUFFER, PixelBuffer->Resource );
diff --git a/Engine/Source/Runtime/OpenGLDrv/Private/OpenGLDevice.cpp b/Engine/Source/Runtime/OpenGLDrv/Private/OpenGLDevice.cpp
--- a/Engine/Source/Runtime/OpenGLDrv/Private/OpenGLDevice.cpp
+++ b/Engine/Source/Runtime/OpenGLDrv/Private/OpenGLDevice.cpp
@@ -814,6 +814,8 @@ static void InitRHICapabilitiesForGL()
 		&& !IsPCPlatform(GMaxRHIShaderPlatform)
 #endif
 		;
+
+	GRHISupportsMultithreadedResources = GRHISupportsRHIThread;
 	
 	// By default use emulated UBs on mobile
 	GUseEmulatedUniformBuffers = IsUsingEmulatedUniformBuffers(GMaxRHIShaderPlatform);
diff --git a/Engine/Source/Runtime/OpenGLDrv/Private/OpenGLTexture.cpp b/Engine/Source/Runtime/OpenGLDrv/Private/OpenGLTexture.cpp
--- a/Engine/Source/Runtime/OpenGLDrv/Private/OpenGLTexture.cpp
+++ b/Engine/Source/Runtime/OpenGLDrv/Private/OpenGLTexture.cpp
@@ -883,7 +883,7 @@ void FOpenGLTexture::Resolve(uint32 MipIndex,uint32 ArrayIndex)
 	// Standard path with a PBO mirroring ever slice of a texture to allow multiple simulataneous maps
 	if (!IsValidRef(PixelBuffers[BufferIndex]))
 	{
-		PixelBuffers[BufferIndex] = new FOpenGLPixelBuffer(GL_PIXEL_UNPACK_BUFFER, 0, MipBytes, BUF_Dynamic, nullptr);
+		PixelBuffers[BufferIndex] = new FOpenGLPixelBuffer(nullptr, GL_PIXEL_UNPACK_BUFFER, 0, MipBytes, BUF_Dynamic, nullptr);
 	}
 	
 	TRefCountPtr<FOpenGLPixelBuffer> PixelBuffer = PixelBuffers[BufferIndex];
@@ -1043,7 +1043,7 @@ void* FOpenGLTexture::Lock(uint32 InMipIndex,uint32 ArrayIndex,EResourceLockMode
 		if (!IsValidRef(PixelBuffers[BufferIndex]))
 		{
 			bBufferExists = false;
-			PixelBuffers[BufferIndex] = new FOpenGLPixelBuffer(GL_PIXEL_UNPACK_BUFFER, 0, MipBytes, BUF_Dynamic, nullptr);
+			PixelBuffers[BufferIndex] = new FOpenGLPixelBuffer(nullptr, GL_PIXEL_UNPACK_BUFFER, 0, MipBytes, BUF_Dynamic, nullptr);
 		}
 
 		TRefCountPtr<FOpenGLPixelBuffer> PixelBuffer = PixelBuffers[BufferIndex];
@@ -1501,7 +1501,7 @@ void FOpenGLTexture::CloneViaPBO(FOpenGLTexture* Src, uint32 InNumMips, int32 Sr
 			// Standard path with a PBO mirroring ever slice of a texture to allow multiple simulataneous maps
 			if (!IsValidRef(PixelBuffers[BufferIndex]))
 			{
-				PixelBuffers[BufferIndex] = new FOpenGLPixelBuffer(GL_PIXEL_UNPACK_BUFFER, 0, MipBytes, BUF_Dynamic, nullptr);
+				PixelBuffers[BufferIndex] = new FOpenGLPixelBuffer(nullptr, GL_PIXEL_UNPACK_BUFFER, 0, MipBytes, BUF_Dynamic, nullptr);
 			}
 			
 			TRefCountPtr<FOpenGLPixelBuffer> PixelBuffer = PixelBuffers[BufferIndex];
diff --git a/Engine/Source/Runtime/OpenGLDrv/Public/OpenGLDrv.h b/Engine/Source/Runtime/OpenGLDrv/Public/OpenGLDrv.h
--- a/Engine/Source/Runtime/OpenGLDrv/Public/OpenGLDrv.h
+++ b/Engine/Source/Runtime/OpenGLDrv/Public/OpenGLDrv.h
@@ -371,11 +371,11 @@ public:
 
 	virtual FUniformBufferRHIRef RHICreateUniformBuffer(const void* Contents, const FRHIUniformBufferLayout* Layout, EUniformBufferUsage Usage, EUniformBufferValidation Validation) final override;
 	virtual void RHIUpdateUniformBuffer(FRHIUniformBuffer* UniformBufferRHI, const void* Contents) final override;
-	virtual FBufferRHIRef RHICreateBuffer(uint32 Size, EBufferUsageFlags Usage, uint32 Stride, ERHIAccess ResourceState, FRHIResourceCreateInfo& CreateInfo) final override;
+	virtual FBufferRHIRef RHICreateBuffer(FRHICommandListBase& RHICmdList, uint32 Size, EBufferUsageFlags Usage, uint32 Stride, ERHIAccess ResourceState, FRHIResourceCreateInfo& CreateInfo) final override;
 	virtual void RHICopyBuffer(FRHIBuffer* SourceBuffer, FRHIBuffer* DestBuffer) final override;
 	virtual void RHITransferBufferUnderlyingResource(FRHIBuffer* DestBuffer, FRHIBuffer* SrcBuffer) final override;
-	virtual void* LockBuffer_BottomOfPipe(FRHICommandListImmediate& RHICmdList, FRHIBuffer* Buffer, uint32 Offset, uint32 Size, EResourceLockMode LockMode) final override;
-	virtual void UnlockBuffer_BottomOfPipe(FRHICommandListImmediate& RHICmdList, FRHIBuffer* Buffer) final override;
+	virtual void* LockBuffer_BottomOfPipe(FRHICommandListBase& RHICmdList, FRHIBuffer* Buffer, uint32 Offset, uint32 Size, EResourceLockMode LockMode) final override;
+	virtual void UnlockBuffer_BottomOfPipe(FRHICommandListBase& RHICmdList, FRHIBuffer* Buffer) final override;
 	virtual FUnorderedAccessViewRHIRef RHICreateUnorderedAccessView(FRHIBuffer* Buffer, bool bUseUAVCounter, bool bAppendBuffer) final override;
 	virtual FUnorderedAccessViewRHIRef RHICreateUnorderedAccessView(FRHITexture* Texture, uint32 MipLevel, uint16 FirstArraySlice, uint16 NumArraySlices) final override;
 	virtual FUnorderedAccessViewRHIRef RHICreateUnorderedAccessView(FRHIBuffer* Buffer, uint8 Format) final override;
@@ -652,7 +652,7 @@ public:
 			x ReturnValue = (x)0;\
 			ALLOC_COMMAND_CL(RHICmdList, FRHICommandGLCommand)([&ReturnValue, GLCommand = MoveTemp(GLCommand)]() { ReturnValue = GLCommand(); }); \
 			RHITHREAD_GLTRACE_BLOCKING;\
-			RHICmdList.ImmediateFlush(EImmediateFlushType::FlushRHIThread);\
+			RHICmdList.GetAsImmediate().ImmediateFlush(EImmediateFlushType::FlushRHIThread);\
 			return ReturnValue;\
 		}\
 
@@ -666,7 +666,7 @@ public:
 		{\
 			ALLOC_COMMAND_CL(RHICmdList, FRHICommandGLCommand)([&ReturnValue, GLCommand = MoveTemp(GLCommand)]() { ReturnValue = GLCommand(); }); \
 			RHITHREAD_GLTRACE_BLOCKING;\
-			RHICmdList.ImmediateFlush(EImmediateFlushType::FlushRHIThread);\
+			RHICmdList.GetAsImmediate().ImmediateFlush(EImmediateFlushType::FlushRHIThread);\
 		}\
 
 
@@ -679,7 +679,7 @@ public:
 		{\
 			ALLOC_COMMAND_CL(RHICmdList, FRHICommandGLCommand)( MoveTemp(GLCommand) ); \
 			RHITHREAD_GLTRACE_BLOCKING;\
-			RHICmdList.ImmediateFlush(EImmediateFlushType::FlushRHIThread);\
+			RHICmdList.GetAsImmediate().ImmediateFlush(EImmediateFlushType::FlushRHIThread);\
 		}\
 
 #define RHITHREAD_GLCOMMAND_EPILOGUE_NORETURN() };\
@@ -691,7 +691,7 @@ public:
 		{\
 			ALLOC_COMMAND_CL(RHICmdList, FRHICommandGLCommand)(  MoveTemp(GLCommand) ); \
 			RHITHREAD_GLTRACE_BLOCKING;\
-			RHICmdList.ImmediateFlush(EImmediateFlushType::FlushRHIThread);\
+			RHICmdList.GetAsImmediate().ImmediateFlush(EImmediateFlushType::FlushRHIThread);\
 		}\
 
 	struct FTextureLockTracker
@@ -762,11 +762,6 @@ public:
 	virtual void* LockTexture2DArray_RenderThread(class FRHICommandListImmediate& RHICmdList, FRHITexture2DArray* Texture, uint32 ArrayIndex, uint32 MipIndex, EResourceLockMode LockMode, uint32& DestStride, bool bLockWithinMiptail) final override;
 	virtual void UnlockTexture2DArray_RenderThread(class FRHICommandListImmediate& RHICmdList, FRHITexture2DArray* Texture, uint32 ArrayIndex, uint32 MipIndex, bool bLockWithinMiptail) final override;
 
-	virtual FBufferRHIRef CreateBuffer_RenderThread(class FRHICommandListImmediate& RHICmdList, uint32 Size, EBufferUsageFlags Usage, uint32 Stride, ERHIAccess ResourceState, FRHIResourceCreateInfo& CreateInfo) final override
-	{
-		return this->RHICreateBuffer(Size, Usage, Stride, ResourceState, CreateInfo);
-	}
-
 	virtual FShaderResourceViewRHIRef CreateShaderResourceView_RenderThread(class FRHICommandListImmediate& RHICmdList, FRHIBuffer* Buffer, uint32 Stride, uint8 Format) final override
 	{
 		return this->RHICreateShaderResourceView(Buffer, Stride, Format);
diff --git a/Engine/Source/Runtime/OpenGLDrv/Public/OpenGLResources.h b/Engine/Source/Runtime/OpenGLDrv/Public/OpenGLResources.h
--- a/Engine/Source/Runtime/OpenGLDrv/Public/OpenGLResources.h
+++ b/Engine/Source/Runtime/OpenGLDrv/Public/OpenGLResources.h
@@ -426,7 +426,7 @@ public:
 		, RealSize(0)
 	{ }
 
-	TOpenGLBuffer(GLenum InType, uint32 InStride, uint32 InSize, EBufferUsageFlags InUsage,
+	TOpenGLBuffer(FRHICommandListBase* RHICmdList, GLenum InType, uint32 InStride, uint32 InSize, EBufferUsageFlags InUsage,
 		const void *InData, bool bStreamedDraw = false, GLuint ResourceToUse = 0, uint32 ResourceSize = 0)
 	: BaseType(InStride,InSize,InUsage)
 	, Resource(0)
@@ -440,12 +440,9 @@ public:
 	, LockBuffer(NULL)
 	, RealSize(InSize)
 	{
-
 		RealSize = ResourceSize ? ResourceSize : InSize;
 
-		FRHICommandListImmediate& RHICmdList = FRHICommandListExecutor::GetImmediateCommandList();
-
-		if (ShouldRunGLRenderContextOpOnThisThread(RHICmdList))
+		if (!RHICmdList || RHICmdList->IsBottomOfPipe())
 		{
 			CreateGLBuffer(InData, ResourceToUse, ResourceSize);
 		}
@@ -454,17 +451,16 @@ public:
 			void* BuffData = nullptr;
 			if (InData)
 			{
-				BuffData = RHICmdList.Alloc(RealSize, 16);
+				BuffData = RHICmdList->Alloc(RealSize, 16);
 				FMemory::Memcpy(BuffData, InData, RealSize);
 			}
 			TransitionFence.Reset();
-			ALLOC_COMMAND_CL(RHICmdList, FRHICommandGLCommand)([=]() 
+			ALLOC_COMMAND_CL(*RHICmdList, FRHICommandGLCommand)([=]() 
 			{
 				CreateGLBuffer(BuffData, ResourceToUse, ResourceSize); 
 				TransitionFence.WriteAssertFence();
 			});
 			TransitionFence.SetRHIThreadFence();
-
 		}
 	}
 
diff --git a/Engine/Source/Runtime/OpenGLDrv/Public/OpenGLUtil.h b/Engine/Source/Runtime/OpenGLDrv/Public/OpenGLUtil.h
--- a/Engine/Source/Runtime/OpenGLDrv/Public/OpenGLUtil.h
+++ b/Engine/Source/Runtime/OpenGLDrv/Public/OpenGLUtil.h
@@ -156,7 +156,7 @@ struct FRHICommandGLCommand final : public FRHICommand<FRHICommandGLCommand, FRH
 };
 
 void RunOnGLRenderContextThread(TUniqueFunction<void(void)> GLFunc, bool bWaitForCompletion = false);
-inline bool ShouldRunGLRenderContextOpOnThisThread(FRHICommandListImmediate& RHICmdList)
+inline bool ShouldRunGLRenderContextOpOnThisThread(FRHICommandListBase& RHICmdList)
 {
 	return (RHICmdList.Bypass() || !IsRunningRHIInSeparateThread() || IsInRHIThread());
 }
diff --git a/Engine/Source/Runtime/RHI/Private/DynamicRHI.cpp b/Engine/Source/Runtime/RHI/Private/DynamicRHI.cpp
--- a/Engine/Source/Runtime/RHI/Private/DynamicRHI.cpp
+++ b/Engine/Source/Runtime/RHI/Private/DynamicRHI.cpp
@@ -643,14 +643,14 @@ void FDynamicRHI::RHICheckViewportHDRStatus(FRHIViewport* Viewport)
 {
 }
 
-void* FDynamicRHI::RHILockBufferMGPU(FRHICommandListImmediate& RHICmdList, FRHIBuffer* Buffer, uint32 GPUIndex, uint32 Offset, uint32 Size, EResourceLockMode LockMode)
+void* FDynamicRHI::RHILockBufferMGPU(FRHICommandListBase& RHICmdList, FRHIBuffer* Buffer, uint32 GPUIndex, uint32 Offset, uint32 Size, EResourceLockMode LockMode)
 {
 	// Fall through to single GPU case
 	check(GPUIndex == 0);
 	return RHILockBuffer(RHICmdList, Buffer, Offset, Size, LockMode);
 }
 
-void FDynamicRHI::RHIUnlockBufferMGPU(FRHICommandListImmediate& RHICmdList, FRHIBuffer* Buffer, uint32 GPUIndex)
+void FDynamicRHI::RHIUnlockBufferMGPU(FRHICommandListBase& RHICmdList, FRHIBuffer* Buffer, uint32 GPUIndex)
 {
 	// Fall through to single GPU case
 	check(GPUIndex == 0);
diff --git a/Engine/Source/Runtime/RHI/Private/RHI.cpp b/Engine/Source/Runtime/RHI/Private/RHI.cpp
--- a/Engine/Source/Runtime/RHI/Private/RHI.cpp
+++ b/Engine/Source/Runtime/RHI/Private/RHI.cpp
@@ -1461,6 +1461,7 @@ bool GRHIRequiresRenderTargetForPixelShaderUAVs = false;
 bool GRHISupportsUAVFormatAliasing = false;
 bool GRHISupportsDirectGPUMemoryLock = false;
 bool GRHISupportsMultithreadedShaderCreation = true;
+bool GRHISupportsMultithreadedResources = false;
 
 bool GRHISupportsMSAADepthSampleAccess = false;
 
diff --git a/Engine/Source/Runtime/RHI/Private/RHICommandList.cpp b/Engine/Source/Runtime/RHI/Private/RHICommandList.cpp
--- a/Engine/Source/Runtime/RHI/Private/RHICommandList.cpp
+++ b/Engine/Source/Runtime/RHI/Private/RHICommandList.cpp
@@ -156,7 +156,6 @@
 static FGraphEventRef RHIThreadTask;
 static FGraphEventRef PrevRHIThreadTask;
 static FGraphEventRef RenderThreadSublistDispatchTask;
-static FGraphEventRef RHIThreadBufferLockFence;
 
 static FGraphEventRef GRHIThreadEndDrawingViewportFences[2];
 static uint32 GRHIThreadEndDrawingViewportFenceIndex = 0;
@@ -169,7 +168,6 @@
 	RHIThreadTask.SafeRelease();
 	PrevRHIThreadTask.SafeRelease();
 	RenderThreadSublistDispatchTask.SafeRelease();
-	RHIThreadBufferLockFence.SafeRelease();
 
 	for (FGraphEventRef& GraphEvent : GRHIThreadEndDrawingViewportFences)
 	{
@@ -698,6 +696,7 @@
                 // is equivalent to a Reset.
 				CmdList.InitialGPUMask = SwapCmdList->GPUMask;
 				CmdList.PSOContext = SwapCmdList->PSOContext;
+				CmdList.RHIThreadBufferLockFence = SwapCmdList->RHIThreadBufferLockFence;
 				CmdList.Data.bInsideRenderPass = SwapCmdList->Data.bInsideRenderPass;
 				CmdList.Data.bInsideComputePass = SwapCmdList->Data.bInsideComputePass;
 			}
@@ -993,11 +992,8 @@
 	}
 };
 
-
-FGraphEventRef FRHICommandListImmediate::RHIThreadFence(bool bSetLockFence)
+FGraphEventRef FRHICommandListBase::RHIThreadFence(bool bSetLockFence)
 {
-	check(IsInRenderingThread());
-
 	if (IsRunningRHIInSeparateThread())
 	{
 		FRHICommandRHIThreadFence* Cmd = ALLOC_COMMAND(FRHICommandRHIThreadFence)();
@@ -1628,6 +1624,14 @@
 void FRHICommandListBase::QueueAsyncCommandListSubmit(FGraphEventRef& AnyThreadCompletionEvent, FRHICommandList* CmdList)
 {
 	check(IsInRenderingThread() && IsImmediate());
+	check(AnyThreadCompletionEvent);
+	check(CmdList);
+
+	if (Bypass())
+	{
+		FRHICommandWaitForAndSubmitSubList(AnyThreadCompletionEvent, CmdList).Execute(*this);
+		return;
+	}
 
 	if (IsRunningRHIInSeparateThread())
 	{
@@ -1644,6 +1648,11 @@
 	ALLOC_COMMAND(FRHICommandWaitForAndSubmitSubList)(AnyThreadCompletionEvent, CmdList);
 	if (IsRunningRHIInSeparateThread())
 	{
+		// The async command list is being recorded in parallel, we have to take a conservative RHI thread fence
+		// for all the work in case a lock was used. The RHI thread waits on the work, so the render thread can't
+		// access the pending RHI thread fence.
+		RHIThreadFence(true);
+
 		FRHICommandListExecutor::GetImmediateCommandList().ImmediateFlush(EImmediateFlushType::DispatchToRHIThread); // we don't want stuff after the async cmd list to be bundled with it
 	}
 }
@@ -1669,9 +1678,22 @@
 void FRHICommandListBase::QueueRenderThreadCommandListSubmit(FGraphEventRef& RenderThreadCompletionEvent, FRHICommandList* CmdList)
 {
 	check(IsInRenderingThread() && IsImmediate());
+	check(RenderThreadCompletionEvent);
+	check(CmdList);
 
+	if (Bypass())
+	{
+		FRHICommandWaitForAndSubmitRTSubList(RenderThreadCompletionEvent, CmdList).Execute(*this);
+		return;
+	}
+
 	ALLOC_COMMAND(FRHICommandWaitForAndSubmitRTSubList)(RenderThreadCompletionEvent, CmdList);
 
+	// The async command list is being recorded in parallel, we have to take a conservative RHI thread fence
+	// for all the work in case a lock was used. The RHI thread waits on the work, so the render thread can't
+	// access the pending RHI thread fence.
+	RHIThreadFence(true);
+
 #if WITH_MGPU
 	// This will restore the context GPU masks to whatever they were set to
 	// before the sub-list executed.
@@ -1707,8 +1729,21 @@
 
 void FRHICommandListBase::QueueCommandListSubmit(FRHICommandList* CmdList)
 {
+	check(CmdList);
+
+	if (Bypass())
+	{
+		FRHICommandSubmitSubList(CmdList).Execute(*this);
+		return;
+	}
+
 	ALLOC_COMMAND(FRHICommandSubmitSubList)(CmdList);
 
+	if (CmdList->RHIThreadBufferLockFence)
+	{
+		RHIThreadBufferLockFence = CmdList->RHIThreadBufferLockFence;
+	}
+
 #if WITH_MGPU
 	// This will restore the context GPU masks to whatever they were set to
 	// before the sub-list executed.
@@ -2246,23 +2281,11 @@
 }
 #endif
 
-void* FRHICommandListBase::operator new(size_t Size)
+FBufferRHIRef FDynamicRHI::CreateBuffer_RenderThread(class FRHICommandListBase& RHICmdList, uint32 Size, EBufferUsageFlags Usage, uint32 Stride, ERHIAccess ResourceState, FRHIResourceCreateInfo& CreateInfo)
 {
-	check(0); // you shouldn't be creating these
-	return FMemory::Malloc(Size);
-}
-
-void FRHICommandListBase::operator delete(void *RawMemory)
-{
-	FMemory::Free(RawMemory);
-}
-
-
-FBufferRHIRef FDynamicRHI::CreateBuffer_RenderThread(class FRHICommandListImmediate& RHICmdList, uint32 Size, EBufferUsageFlags Usage, uint32 Stride, ERHIAccess ResourceState, FRHIResourceCreateInfo& CreateInfo)
-{
 	CSV_SCOPED_TIMING_STAT(RHITStalls, CreateBuffer_RenderThread);
-	FScopedRHIThreadStaller StallRHIThread(RHICmdList);
-	FBufferRHIRef Buffer = GDynamicRHI->RHICreateBuffer(Size, Usage, Stride, ResourceState, CreateInfo);
+	FScopedRHIThreadStaller StallRHIThread(RHICmdList.GetAsImmediate());
+	FBufferRHIRef Buffer = GDynamicRHI->RHICreateBuffer(RHICmdList, Size, Usage, Stride, ResourceState, CreateInfo);
 	return Buffer;
 }
 
@@ -2289,7 +2312,7 @@
 
 static FLockTracker GLockTracker;
 
-void* FDynamicRHI::RHILockBuffer(class FRHICommandListImmediate& RHICmdList, FRHIBuffer* Buffer, uint32 Offset, uint32 SizeRHI, EResourceLockMode LockMode)
+void* FDynamicRHI::RHILockBuffer(class FRHICommandListBase& RHICmdList, FRHIBuffer* Buffer, uint32 Offset, uint32 SizeRHI, EResourceLockMode LockMode)
 {
 	QUICK_SCOPE_CYCLE_COUNTER(STAT_FDynamicRHI_LockBuffer);
 
@@ -2302,7 +2325,7 @@
 			QUICK_SCOPE_CYCLE_COUNTER(STAT_RHIMETHOD_LockBuffer_FlushAndLock);
 			CSV_SCOPED_TIMING_STAT(RHITFlushes, LockBuffer_BottomOfPipe);
 
-			FRHICommandListScopedFlushAndExecute Flush(RHICmdList);
+			FRHICommandListScopedFlushAndExecute Flush(RHICmdList.GetAsImmediate());
 			Result = GDynamicRHI->LockBuffer_BottomOfPipe(RHICmdList, Buffer, Offset, SizeRHI, LockMode);
 		}
 		else
@@ -2324,7 +2347,7 @@
 	return Result;
 }
 
-void FDynamicRHI::RHIUnlockBuffer(class FRHICommandListImmediate& RHICmdList, FRHIBuffer* Buffer)
+void FDynamicRHI::RHIUnlockBuffer(class FRHICommandListBase& RHICmdList, FRHIBuffer* Buffer)
 {
 	QUICK_SCOPE_CYCLE_COUNTER(STAT_FDynamicRHI_UnlockBuffer_RenderThread);
 
@@ -2338,13 +2361,13 @@
 			QUICK_SCOPE_CYCLE_COUNTER(STAT_RHIMETHOD_UnlockBuffer_FlushAndUnlock);
 			CSV_SCOPED_TIMING_STAT(RHITFlushes, UnlockBuffer_BottomOfPipe);
 
-			FRHICommandListScopedFlushAndExecute Flush(RHICmdList);
+			FRHICommandListScopedFlushAndExecute Flush(RHICmdList.GetAsImmediate());
 			GDynamicRHI->UnlockBuffer_BottomOfPipe(RHICmdList, Buffer);
 			GLockTracker.TotalMemoryOutstanding = 0;
 		}
 		else
 		{
-			RHICmdList.EnqueueLambda([Buffer, Params](FRHICommandListImmediate& InRHICmdList)
+			RHICmdList.EnqueueLambda([Buffer, Params](FRHICommandListBase& InRHICmdList)
 			{
 				QUICK_SCOPE_CYCLE_COUNTER(STAT_FRHICommandUpdateBuffer_Execute);
 				void* Data = GDynamicRHI->LockBuffer_BottomOfPipe(InRHICmdList, Buffer, Params.Offset, Params.BufferSize, RLM_WriteOnly);
@@ -2360,11 +2383,11 @@
 			RHICmdList.RHIThreadFence(true);
 		}
 
-		if (GLockTracker.TotalMemoryOutstanding > uint32(CVarRHICmdMaxOutstandingMemoryBeforeFlush.GetValueOnRenderThread()) * 1024u)
+		if (RHICmdList.IsImmediate() && GLockTracker.TotalMemoryOutstanding > uint32(CVarRHICmdMaxOutstandingMemoryBeforeFlush.GetValueOnRenderThread()) * 1024u)
 		{
 			QUICK_SCOPE_CYCLE_COUNTER(STAT_RHIMETHOD_UnlockBuffer_FlushForMem);
 			// we could be loading a level or something, lets get this stuff going
-			RHICmdList.ImmediateFlush(EImmediateFlushType::DispatchToRHIThread); 
+			RHICmdList.GetAsImmediate().ImmediateFlush(EImmediateFlushType::DispatchToRHIThread); 
 			GLockTracker.TotalMemoryOutstanding = 0;
 		}
 	}
diff --git a/Engine/Source/Runtime/RHI/Private/RHIValidation.cpp b/Engine/Source/Runtime/RHI/Private/RHIValidation.cpp
--- a/Engine/Source/Runtime/RHI/Private/RHIValidation.cpp
+++ b/Engine/Source/Runtime/RHI/Private/RHIValidation.cpp
@@ -390,10 +390,12 @@ namespace RHIValidation
 	}
 }
 
-void FValidationRHI::LockBufferValidate(class FRHICommandListImmediate& RHICmdList, FRHIBuffer* Buffer, EResourceLockMode LockMode)
+void FValidationRHI::LockBufferValidate(class FRHICommandListBase& RHICmdList, FRHIBuffer* Buffer, EResourceLockMode LockMode)
 {
 	using namespace RHIValidation;
 
+	check(GRHISupportsMultithreadedResources || RHICmdList.IsImmediate());
+
 	if (!EnumHasAnyFlags(Buffer->GetUsage(), BUF_Volatile) && LockMode == RLM_WriteOnly)
 	{
 		bool bIsInsideRenderPass;
@@ -410,14 +412,14 @@ void FValidationRHI::LockBufferValidate(class FRHICommandListImmediate& RHICmdLi
 	}
 }
 
-void* FValidationRHI::RHILockBuffer(class FRHICommandListImmediate& RHICmdList, FRHIBuffer* Buffer, uint32 Offset, uint32 SizeRHI, EResourceLockMode LockMode)
+void* FValidationRHI::RHILockBuffer(class FRHICommandListBase& RHICmdList, FRHIBuffer* Buffer, uint32 Offset, uint32 SizeRHI, EResourceLockMode LockMode)
 {
 	LockBufferValidate(RHICmdList, Buffer, LockMode);
 
 	return RHI->RHILockBuffer(RHICmdList, Buffer, Offset, SizeRHI, LockMode);
 }
 
-void* FValidationRHI::RHILockBufferMGPU(class FRHICommandListImmediate& RHICmdList, FRHIBuffer* Buffer, uint32 GPUIndex, uint32 Offset, uint32 SizeRHI, EResourceLockMode LockMode)
+void* FValidationRHI::RHILockBufferMGPU(class FRHICommandListBase& RHICmdList, FRHIBuffer* Buffer, uint32 GPUIndex, uint32 Offset, uint32 SizeRHI, EResourceLockMode LockMode)
 {
 	LockBufferValidate(RHICmdList, Buffer, LockMode);
 
diff --git a/Engine/Source/Runtime/RHI/Public/DynamicRHI.h b/Engine/Source/Runtime/RHI/Public/DynamicRHI.h
--- a/Engine/Source/Runtime/RHI/Public/DynamicRHI.h
+++ b/Engine/Source/Runtime/RHI/Public/DynamicRHI.h
@@ -18,6 +18,7 @@ class FGraphicsPipelineStateInitializer;
 class FLastRenderTimeContainer;
 class FReadSurfaceDataFlags;
 class FRHICommandList;
+class FRHICommandListBase;
 class FRHIComputeFence;
 class FRayTracingPipelineState;
 class IRHITransientResourceAllocator;
@@ -423,15 +424,17 @@ public:
 	* @param ResourceArray - An optional pointer to a resource array containing the resource's data.
 	*/
 	// FlushType: Wait RHI Thread
-	virtual FBufferRHIRef RHICreateBuffer(uint32 Size, EBufferUsageFlags Usage, uint32 Stride, ERHIAccess InResourceState, FRHIResourceCreateInfo& CreateInfo) = 0;
+	UE_DEPRECATED(5.1, "Use the command list version of RHICreateBuffer instead.")
+	virtual FBufferRHIRef RHICreateBuffer(uint32 Size, EBufferUsageFlags Usage, uint32 Stride, ERHIAccess InResourceState, FRHIResourceCreateInfo& CreateInfo) { checkNoEntry();  return nullptr; }
 
-	// FlushType: Flush RHI Thread
-	virtual void* RHILockBuffer(FRHICommandListImmediate& RHICmdList, FRHIBuffer* Buffer, uint32 Offset, uint32 Size, EResourceLockMode LockMode);
-	virtual void* RHILockBufferMGPU(FRHICommandListImmediate& RHICmdList, FRHIBuffer* Buffer, uint32 GPUIndex, uint32 Offset, uint32 Size, EResourceLockMode LockMode);
+	virtual FBufferRHIRef RHICreateBuffer(FRHICommandListBase& RHICmdList, uint32 Size, EBufferUsageFlags Usage, uint32 Stride, ERHIAccess InResourceState, FRHIResourceCreateInfo& CreateInfo) = 0;
+
+	virtual void* RHILockBuffer(FRHICommandListBase& RHICmdList, FRHIBuffer* Buffer, uint32 Offset, uint32 Size, EResourceLockMode LockMode);
+	virtual void* RHILockBufferMGPU(FRHICommandListBase& RHICmdList, FRHIBuffer* Buffer, uint32 GPUIndex, uint32 Offset, uint32 Size, EResourceLockMode LockMode);
 
 	// FlushType: Flush RHI Thread
-	virtual void RHIUnlockBuffer(FRHICommandListImmediate& RHICmdList, FRHIBuffer* Buffer);
-	virtual void RHIUnlockBufferMGPU(FRHICommandListImmediate& RHICmdList, FRHIBuffer* Buffer, uint32 GPUIndex);
+	virtual void RHIUnlockBuffer(FRHICommandListBase& RHICmdList, FRHIBuffer* Buffer);
+	virtual void RHIUnlockBufferMGPU(FRHICommandListBase& RHICmdList, FRHIBuffer* Buffer, uint32 GPUIndex);
 
 	/** Creates an unordered access view of the given buffer. */
 	// FlushType: Wait RHI Thread
@@ -1052,7 +1055,8 @@ public:
 	FORCEINLINE IRHICommandContextContainer* RHIGetCommandContextContainer(int32 Index, int32 Num, FRHIGPUMask GPUMask) { return RHIGetCommandContextContainer(Index, Num); }
 #endif
 
-	virtual FBufferRHIRef CreateBuffer_RenderThread(class FRHICommandListImmediate& RHICmdList, uint32 Size, EBufferUsageFlags Usage, uint32 Stride, ERHIAccess ResourceState, FRHIResourceCreateInfo& CreateInfo);
+	UE_DEPRECATED(5.1, "CreateBuffer_RenderThread is deprecated. Use RHICreateBuffer instead.")
+	virtual FBufferRHIRef CreateBuffer_RenderThread(class FRHICommandListBase& RHICmdList, uint32 Size, EBufferUsageFlags Usage, uint32 Stride, ERHIAccess ResourceState, FRHIResourceCreateInfo& CreateInfo);
 	virtual FShaderResourceViewRHIRef CreateShaderResourceView_RenderThread(class FRHICommandListImmediate& RHICmdList, FRHIBuffer* Buffer, uint32 Stride, uint8 Format);
 	virtual FShaderResourceViewRHIRef CreateShaderResourceView_RenderThread(class FRHICommandListImmediate& RHICmdList, const FShaderResourceViewInitializer& Initializer);
 	virtual FShaderResourceViewRHIRef CreateShaderResourceView_RenderThread(class FRHICommandListImmediate& RHICmdList, FRHIBuffer* Buffer);
@@ -1129,14 +1133,14 @@ public:
 	}
 
 	// Buffer Lock/Unlock
-	virtual void* LockBuffer_BottomOfPipe(class FRHICommandListImmediate& RHICmdList, FRHIBuffer* Buffer, uint32 Offset, uint32 SizeRHI, EResourceLockMode LockMode)
+	virtual void* LockBuffer_BottomOfPipe(class FRHICommandListBase& RHICmdList, FRHIBuffer* Buffer, uint32 Offset, uint32 SizeRHI, EResourceLockMode LockMode)
 	{
 		// Either this function or RHILockBuffer must be implemented by the platform RHI.
 		checkNoEntry();
 		return nullptr;
 	}
 
-	virtual void UnlockBuffer_BottomOfPipe(class FRHICommandListImmediate& RHICmdList, FRHIBuffer* Buffer)
+	virtual void UnlockBuffer_BottomOfPipe(class FRHICommandListBase& RHICmdList, FRHIBuffer* Buffer)
 	{
 		// Either this function or RHIUnlockBuffer must be implemented by the platform RHI.
 		checkNoEntry();
diff --git a/Engine/Source/Runtime/RHI/Public/RHI.h b/Engine/Source/Runtime/RHI/Public/RHI.h
--- a/Engine/Source/Runtime/RHI/Public/RHI.h
+++ b/Engine/Source/Runtime/RHI/Public/RHI.h
@@ -462,6 +462,9 @@ extern RHI_API bool GRHISupportsDirectGPUMemoryLock;
 /** true if the multi-threaded shader creation is supported by (or desirable for) the RHI. */
 extern RHI_API bool GRHISupportsMultithreadedShaderCreation;
 
+/** Does the RHI support parallel resource commands (i.e. create / lock / unlock) on non-immediate command list APIs, recorded off the render thread. */
+extern RHI_API bool GRHISupportsMultithreadedResources;
+
 /** The minimum Z value in clip space for the RHI. This is a constant value to always match D3D clip-space. */
 constexpr float GMinClipZ = 0.0f;
 
diff --git a/Engine/Source/Runtime/RHI/Public/RHICommandList.h b/Engine/Source/Runtime/RHI/Public/RHICommandList.h
--- a/Engine/Source/Runtime/RHI/Public/RHICommandList.h
+++ b/Engine/Source/Runtime/RHI/Public/RHICommandList.h
@@ -188,6 +188,8 @@ struct RHI_API FLockTracker
 		{
 		}
 	};
+
+	FCriticalSection CriticalSection;
 	TArray<FLockParams, TInlineAllocator<16> > OutstandingLocks;
 	uint32 TotalMemoryOutstanding;
 
@@ -198,6 +200,7 @@ struct RHI_API FLockTracker
 
 	FORCEINLINE_DEBUGGABLE void Lock(void* RHIBuffer, void* Buffer, uint32 Offset, uint32 SizeRHI, EResourceLockMode LockMode)
 	{
+		FScopeLock Lock(&CriticalSection);
 #if DO_CHECK
 		for (auto& Parms : OutstandingLocks)
 		{
@@ -209,6 +212,7 @@ struct RHI_API FLockTracker
 	}
 	FORCEINLINE_DEBUGGABLE FLockParams Unlock(void* RHIBuffer)
 	{
+		FScopeLock Lock(&CriticalSection);
 		for (int32 Index = 0; Index < OutstandingLocks.Num(); Index++)
 		{
 			if (OutstandingLocks[Index].RHIBuffer == RHIBuffer)
@@ -407,6 +411,33 @@ private:
 	uint32 FenceFrameNumber[MAX_FENCE_INDICES];
 };
 
+// Using variadic macro because some types are fancy template<A,B> stuff, which gets broken off at the comma and interpreted as multiple arguments. 
+#define ALLOC_COMMAND(...) new ( AllocCommand(sizeof(__VA_ARGS__), alignof(__VA_ARGS__)) ) __VA_ARGS__
+#define ALLOC_COMMAND_CL(RHICmdList, ...) new ( (RHICmdList).AllocCommand(sizeof(__VA_ARGS__), alignof(__VA_ARGS__)) ) __VA_ARGS__
+
+// This controls if the cmd list bypass can be toggled at runtime. It is quite expensive to have these branches in there.
+#define CAN_TOGGLE_COMMAND_LIST_BYPASS (!UE_BUILD_SHIPPING && !UE_BUILD_TEST)
+
+// Controls whether a command list can locally disallow bypass mode.
+#define CAN_DISALLOW_COMMAND_LIST_BYPASS CAN_TOGGLE_COMMAND_LIST_BYPASS || PLATFORM_RHITHREAD_DEFAULT_BYPASS != 0
+
+template <typename RHICmdListType, typename LAMBDA>
+struct TRHILambdaCommand final : public FRHICommandBase
+{
+	LAMBDA Lambda;
+
+	TRHILambdaCommand(LAMBDA&& InLambda)
+		: Lambda(Forward<LAMBDA>(InLambda))
+	{}
+
+	void ExecuteAndDestruct(FRHICommandListBase& CmdList, FRHICommandListDebugContext&) override final
+	{
+		TRACE_CPUPROFILER_EVENT_SCOPE_ON_CHANNEL(TRHILambdaCommand, RHICommandsChannel);
+		Lambda(*static_cast<RHICmdListType*>(&CmdList));
+		Lambda.~LAMBDA();
+	}
+};
+
 extern RHI_API FRHICommandListFenceAllocator GRHIFenceAllocator;
 
 class RHI_API FRHICommandListBase : public FNoncopyable
@@ -414,13 +445,10 @@ class RHI_API FRHICommandListBase : public FNoncopyable
 public:
 	~FRHICommandListBase();
 
-	/** Custom new/delete with recycling */
-	void* operator new(size_t Size);
-	void operator delete(void *RawMemory);
-
 	inline void Flush();
 	inline bool IsImmediate() const;
 	inline bool IsImmediateAsyncCompute() const;
+	inline FRHICommandListImmediate& GetAsImmediate();
 
 	const int32 GetUsedMemory() const;
 	void QueueAsyncCommandListSubmit(FGraphEventRef& AnyThreadCompletionEvent, class FRHICommandList* CmdList);
@@ -589,16 +617,59 @@ public:
 	const FRHIBreadcrumbStack& GetBreadcrumbStack() const { return BreadcrumbStack; }
 #endif // RHI_WANT_BREADCRUMB_EVENTS
 
+	template <typename LAMBDA>
+	FORCEINLINE_DEBUGGABLE void EnqueueLambda(LAMBDA&& Lambda)
+	{
+		if (IsBottomOfPipe())
+		{
+			Lambda(*this);
+		}
+		else
+		{
+			ALLOC_COMMAND(TRHILambdaCommand<FRHICommandListBase, LAMBDA>)(Forward<LAMBDA>(Lambda));
+		}
+	}
+
+	FGraphEventRef RHIThreadFence(bool bSetLockFence = false);
+
+	FORCEINLINE void* LockBuffer(FRHIBuffer* Buffer, uint32 Offset, uint32 SizeRHI, EResourceLockMode LockMode)
+	{
+		return GDynamicRHI->RHILockBuffer(*this, Buffer, Offset, SizeRHI, LockMode);
+	}
+
+	FORCEINLINE void UnlockBuffer(FRHIBuffer* Buffer)
+	{
+		GDynamicRHI->RHIUnlockBuffer(*this, Buffer);
+	}
+
+	FORCEINLINE FBufferRHIRef CreateBuffer(uint32 Size, EBufferUsageFlags Usage, uint32 Stride, ERHIAccess ResourceState, FRHIResourceCreateInfo& CreateInfo)
+	{
+		FBufferRHIRef Buffer = GDynamicRHI->RHICreateBuffer(*this, Size, Usage, Stride, ResourceState, CreateInfo);
+		Buffer->SetTrackedAccess_Unsafe(ResourceState);
+		return Buffer;
+	}
+
+	void DisallowBypass()
+	{
+#if CAN_DISALLOW_COMMAND_LIST_BYPASS
+		bAllowBypass = false;
+#endif
+	}
+
 private:
 	FRHICommandBase* Root;
 	FRHICommandBase** CommandLink;
 	bool bExecuting;
+#if CAN_DISALLOW_COMMAND_LIST_BYPASS
+	bool bAllowBypass = true;
+#endif
 	uint32 NumCommands;
 	uint32 UID;
 	IRHICommandContext* Context;
 	IRHIComputeContext* ComputeContext;
 	FMemStackBase MemManager; 
 	FGraphEventArray RTTasks;
+	FGraphEventRef RHIThreadBufferLockFence;
 
 	friend class FRHICommandListExecutor;
 	friend class FRHICommandListIterator;
@@ -779,23 +850,6 @@ struct FRHICommand : public FRHICommandBase
 	virtual void StoreDebugInfo(FRHICommandListDebugContext& Context) {};
 };
 
-template <typename RHICmdListType, typename LAMBDA>
-struct TRHILambdaCommand final : public FRHICommandBase
-{
-	LAMBDA Lambda;
-
-	TRHILambdaCommand(LAMBDA&& InLambda)
-		: Lambda(Forward<LAMBDA>(InLambda))
-	{}
-
-	void ExecuteAndDestruct(FRHICommandListBase& CmdList, FRHICommandListDebugContext&) override final
-	{
-		TRACE_CPUPROFILER_EVENT_SCOPE_ON_CHANNEL(TRHILambdaCommand, RHICommandsChannel);
-		Lambda(*static_cast<RHICmdListType*>(&CmdList));
-		Lambda.~LAMBDA();
-	}
-};
-
 #define FRHICOMMAND_MACRO(CommandName)								\
 struct PREPROCESSOR_JOIN(CommandName##String, __LINE__)				\
 {																	\
@@ -2359,10 +2413,6 @@ FRHICOMMAND_MACRO(FRHICommandSetRayTracingBindings)
 };
 #endif // RHI_RAYTRACING
 
-// Using variadic macro because some types are fancy template<A,B> stuff, which gets broken off at the comma and interpreted as multiple arguments. 
-#define ALLOC_COMMAND(...) new ( AllocCommand(sizeof(__VA_ARGS__), alignof(__VA_ARGS__)) ) __VA_ARGS__
-#define ALLOC_COMMAND_CL(RHICmdList, ...) new ( (RHICmdList).AllocCommand(sizeof(__VA_ARGS__), alignof(__VA_ARGS__)) ) __VA_ARGS__
-
 template<> RHI_API void FRHICommandSetShaderParameter<FRHIComputeShader>::Execute(FRHICommandListBase& CmdList);
 template<> RHI_API void FRHICommandSetShaderUniformBuffer<FRHIComputeShader>::Execute(FRHICommandListBase& CmdList);
 template<> RHI_API void FRHICommandSetShaderTexture<FRHIComputeShader>::Execute(FRHICommandListBase& CmdList);
@@ -4066,7 +4116,6 @@ public:
 	static FGraphEventArray& GetRenderThreadTaskArray();
 	static void WaitOnRenderThreadTaskFence(FGraphEventRef& Fence);
 	static bool AnyRenderThreadTasksOutstanding();
-	FGraphEventRef RHIThreadFence(bool bSetLockFence = false);
 
 	//Queue the given async compute commandlists in order with the current immediate commandlist
 	void QueueAsyncCompute(FRHIComputeCommandList& RHIComputeCmdList);
@@ -4187,20 +4236,23 @@ public:
 	{
 		return RHICreateUniformBuffer(Contents, &Layout, Usage);
 	}
-	
+
+	UE_DEPRECATED(5.1, "Use CreateBuffer and LockBuffer separately")
 	FORCEINLINE FBufferRHIRef CreateAndLockIndexBuffer(uint32 Stride, uint32 Size, EBufferUsageFlags InUsage, ERHIAccess InResourceState, FRHIResourceCreateInfo& CreateInfo, void*& OutDataBuffer)
 	{
-		FBufferRHIRef IndexBuffer = GDynamicRHI->CreateBuffer_RenderThread(*this, Size, InUsage | BUF_IndexBuffer, Stride, InResourceState, CreateInfo);
-		IndexBuffer->SetTrackedAccess_Unsafe(InResourceState);
-		OutDataBuffer = GDynamicRHI->RHILockBuffer(*this, IndexBuffer, 0, Size, RLM_WriteOnly);
+		FBufferRHIRef IndexBuffer = CreateBuffer(Size, InUsage | BUF_IndexBuffer, Stride, InResourceState, CreateInfo);
+		OutDataBuffer = LockBuffer(IndexBuffer, 0, Size, RLM_WriteOnly);
 		return IndexBuffer;
 	}
 
+	UE_DEPRECATED(5.1, "Use CreateBuffer and LockBuffer separately")
 	FORCEINLINE FBufferRHIRef CreateAndLockIndexBuffer(uint32 Stride, uint32 Size, EBufferUsageFlags InUsage, FRHIResourceCreateInfo& CreateInfo, void*& OutDataBuffer)
 	{
 		EBufferUsageFlags Usage = InUsage | BUF_IndexBuffer;
 		ERHIAccess ResourceState = RHIGetDefaultResourceState(Usage, true);
+	PRAGMA_DISABLE_DEPRECATION_WARNINGS
 		return CreateAndLockIndexBuffer(Stride, Size, Usage, ResourceState, CreateInfo, OutDataBuffer);
+	PRAGMA_DISABLE_DEPRECATION_WARNINGS
 	}
 	
 	UE_DEPRECATED(5.0, "Buffer locks have been unified. Use LockBuffer() instead.")
@@ -4224,20 +4276,23 @@ public:
 	{
 		GDynamicRHI->UnlockStagingBuffer_RenderThread(*this, StagingBuffer);
 	}
-	
+
+	UE_DEPRECATED(5.1, "Use CreateBuffer and LockBuffer separately")
 	FORCEINLINE FBufferRHIRef CreateAndLockVertexBuffer(uint32 Size, EBufferUsageFlags InUsage, ERHIAccess InResourceState, FRHIResourceCreateInfo& CreateInfo, void*& OutDataBuffer)
 	{
-		FBufferRHIRef VertexBuffer = GDynamicRHI->CreateBuffer_RenderThread(*this, Size, InUsage | BUF_VertexBuffer, 0, InResourceState, CreateInfo);
-		VertexBuffer->SetTrackedAccess_Unsafe(InResourceState);
-		OutDataBuffer = GDynamicRHI->RHILockBuffer(*this, VertexBuffer, 0, Size, RLM_WriteOnly);
+		FBufferRHIRef VertexBuffer = CreateBuffer(Size, InUsage | BUF_VertexBuffer, 0, InResourceState, CreateInfo);
+		OutDataBuffer = LockBuffer(VertexBuffer, 0, Size, RLM_WriteOnly);
 		return VertexBuffer;
 	}
 
+	UE_DEPRECATED(5.1, "Use CreateBuffer and LockBuffer separately")
 	FORCEINLINE FBufferRHIRef CreateAndLockVertexBuffer(uint32 Size, EBufferUsageFlags InUsage, FRHIResourceCreateInfo& CreateInfo, void*& OutDataBuffer)
 	{
 		EBufferUsageFlags Usage = InUsage | BUF_VertexBuffer;
 		ERHIAccess ResourceState = RHIGetDefaultResourceState(Usage, true);
+	PRAGMA_DISABLE_DEPRECATION_WARNINGS
 		return CreateAndLockVertexBuffer(Size, Usage, ResourceState, CreateInfo, OutDataBuffer);
+	PRAGMA_DISABLE_DEPRECATION_WARNINGS
 	}
 
 	UE_DEPRECATED(5.0, "Buffer locks have been unified. Use LockBuffer() instead.")
@@ -4277,16 +4332,6 @@ public:
 		GDynamicRHI->RHIUnlockBuffer(*this, StructuredBuffer);
 	}
 
-	FORCEINLINE void* LockBuffer(FRHIBuffer* Buffer, uint32 Offset, uint32 SizeRHI, EResourceLockMode LockMode)
-	{
-		return GDynamicRHI->RHILockBuffer(*this, Buffer, Offset, SizeRHI, LockMode);
-	}
-	
-	FORCEINLINE void UnlockBuffer(FRHIBuffer* Buffer)
-	{
-		GDynamicRHI->RHIUnlockBuffer(*this, Buffer);
-	}
-
 	// LockBufferMGPU / UnlockBufferMGPU may ONLY be called for buffers with the EBufferUsageFlags::MultiGPUAllocate flag set!
 	// And buffers with that flag set may not call the regular (single GPU) LockBuffer / UnlockBuffer.  The single GPU version
 	// of LockBuffer uses driver mirroring to propagate the updated buffer to other GPUs, while the MGPU / MultiGPUAllocate
@@ -5057,9 +5102,6 @@ public:
 	}
 };
 
-// This controls if the cmd list bypass can be toggled at runtime. It is quite expensive to have these branches in there.
-#define CAN_TOGGLE_COMMAND_LIST_BYPASS (!UE_BUILD_SHIPPING && !UE_BUILD_TEST)
-
 class RHI_API FRHICommandListExecutor
 {
 public:
@@ -5211,111 +5253,121 @@ FORCEINLINE FStagingBufferRHIRef RHICreateStagingBuffer()
 UE_DEPRECATED(5.0, "Use RHICreateBuffer() and RHILockBuffer() instead.")
 FORCEINLINE FBufferRHIRef RHICreateAndLockIndexBuffer(uint32 Stride, uint32 Size, EBufferUsageFlags InUsage, FRHIResourceCreateInfo& CreateInfo, void*& OutDataBuffer)
 {
+	check(IsInRenderingThread());
+PRAGMA_DISABLE_DEPRECATION_WARNINGS
 	return FRHICommandListExecutor::GetImmediateCommandList().CreateAndLockIndexBuffer(Stride, Size, InUsage, CreateInfo, OutDataBuffer);
+PRAGMA_ENABLE_DEPRECATION_WARNINGS
 }
 
 FORCEINLINE FBufferRHIRef RHICreateBuffer(uint32 Size, EBufferUsageFlags Usage, uint32 Stride, ERHIAccess ResourceState, FRHIResourceCreateInfo& CreateInfo)
 {
-	FBufferRHIRef Buffer = GDynamicRHI->CreateBuffer_RenderThread(FRHICommandListExecutor::GetImmediateCommandList(), Size, Usage, Stride, ResourceState, CreateInfo);
-	Buffer->SetTrackedAccess_Unsafe(ResourceState);
-	return Buffer;
+	check(IsInRenderingThread());
+	return FRHICommandListExecutor::GetImmediateCommandList().CreateBuffer(Size, Usage, Stride, ResourceState, CreateInfo);
 }
 
-FORCEINLINE FBufferRHIRef RHICreateIndexBuffer(uint32 Stride, uint32 Size, EBufferUsageFlags InUsage, ERHIAccess InResourceState, FRHIResourceCreateInfo& CreateInfo)
+FORCEINLINE FBufferRHIRef RHICreateIndexBuffer(uint32 Stride, uint32 Size, EBufferUsageFlags Usage, ERHIAccess ResourceState, FRHIResourceCreateInfo& CreateInfo)
 {
-	FBufferRHIRef Buffer = GDynamicRHI->CreateBuffer_RenderThread(FRHICommandListExecutor::GetImmediateCommandList(), Size, InUsage | BUF_IndexBuffer, Stride, InResourceState, CreateInfo);
-	Buffer->SetTrackedAccess_Unsafe(InResourceState);
-	return Buffer;
+	check(IsInRenderingThread());
+	return FRHICommandListExecutor::GetImmediateCommandList().CreateBuffer(Size, Usage | EBufferUsageFlags::IndexBuffer, Stride, ResourceState, CreateInfo);
 }
 
-FORCEINLINE FBufferRHIRef RHIAsyncCreateIndexBuffer(uint32 Stride, uint32 Size, EBufferUsageFlags InUsage, ERHIAccess InResourceState, FRHIResourceCreateInfo& CreateInfo)
+UE_DEPRECATED(5.1, "RHIAsyncCreateIndexBuffer is deprecated. Use FRHICommandList::CreateBuffer instead.")
+FORCEINLINE FBufferRHIRef RHIAsyncCreateIndexBuffer(uint32 Stride, uint32 Size, EBufferUsageFlags Usage, ERHIAccess ResourceState, FRHIResourceCreateInfo& CreateInfo)
 {
-	FBufferRHIRef Buffer = GDynamicRHI->RHICreateBuffer(Size, InUsage, Stride, InResourceState, CreateInfo);
-	Buffer->SetTrackedAccess_Unsafe(InResourceState);
-	return Buffer;
+	check(IsInRenderingThread());
+	return FRHICommandListExecutor::GetImmediateCommandList().CreateBuffer(Size, Usage, Stride, ResourceState, CreateInfo);
 }
 
 FORCEINLINE FBufferRHIRef RHICreateIndexBuffer(uint32 Stride, uint32 Size, EBufferUsageFlags InUsage, FRHIResourceCreateInfo& CreateInfo)
 {
 	bool bHasInitialData = CreateInfo.BulkData != nullptr;
-	EBufferUsageFlags Usage = InUsage | BUF_IndexBuffer;
+	EBufferUsageFlags Usage = InUsage | EBufferUsageFlags::IndexBuffer;
 	ERHIAccess ResourceState = RHIGetDefaultResourceState(Usage, bHasInitialData);
 	return RHICreateIndexBuffer(Stride, Size, Usage, ResourceState, CreateInfo);
 }
 
+UE_DEPRECATED(5.1, "RHIAsyncCreateIndexBuffer is deprecated. Use FRHICommandList::CreateBuffer instead.")
 FORCEINLINE FBufferRHIRef RHIAsyncCreateIndexBuffer(uint32 Stride, uint32 Size, EBufferUsageFlags InUsage, FRHIResourceCreateInfo& CreateInfo)
 {
 	bool bHasInitialData = CreateInfo.BulkData != nullptr;
-	EBufferUsageFlags Usage = InUsage | BUF_IndexBuffer;
+	EBufferUsageFlags Usage = InUsage | EBufferUsageFlags::IndexBuffer;
 	ERHIAccess ResourceState = RHIGetDefaultResourceState(Usage, bHasInitialData);
+PRAGMA_DISABLE_DEPRECATION_WARNINGS
 	return RHIAsyncCreateIndexBuffer(Stride, Size, Usage, ResourceState, CreateInfo);
+PRAGMA_ENABLE_DEPRECATION_WARNINGS
 }
 
 UE_DEPRECATED(5.0, "Buffer locks have been unified. Use RHILockBuffer() instead.")
 FORCEINLINE void* RHILockIndexBuffer(FRHIBuffer* IndexBuffer, uint32 Offset, uint32 Size, EResourceLockMode LockMode)
 {
+	check(IsInRenderingThread());
 	return FRHICommandListExecutor::GetImmediateCommandList().LockBuffer(IndexBuffer, Offset, Size, LockMode);
 }
 
 UE_DEPRECATED(5.0, "Buffer locks have been unified. Use RHIUnlockBuffer() instead.")
 FORCEINLINE void RHIUnlockIndexBuffer(FRHIBuffer* IndexBuffer)
 {
-	 FRHICommandListExecutor::GetImmediateCommandList().UnlockBuffer(IndexBuffer);
+	check(IsInRenderingThread());
+	FRHICommandListExecutor::GetImmediateCommandList().UnlockBuffer(IndexBuffer);
 }
 
 UE_DEPRECATED(5.0, "Use RHICreateBuffer() and RHILockBuffer() instead.")
 FORCEINLINE FBufferRHIRef RHICreateAndLockVertexBuffer(uint32 Size, EBufferUsageFlags InUsage, FRHIResourceCreateInfo& CreateInfo, void*& OutDataBuffer)
 {
+	check(IsInRenderingThread());
 	return FRHICommandListExecutor::GetImmediateCommandList().CreateAndLockVertexBuffer(Size, InUsage, CreateInfo, OutDataBuffer);
 }
 
-FORCEINLINE FBufferRHIRef RHICreateVertexBuffer(uint32 Size, EBufferUsageFlags InUsage, ERHIAccess InResourceState, FRHIResourceCreateInfo& CreateInfo)
+FORCEINLINE FBufferRHIRef RHICreateVertexBuffer(uint32 Size, EBufferUsageFlags Usage, ERHIAccess ResourceState, FRHIResourceCreateInfo& CreateInfo)
 {
-	FBufferRHIRef Buffer = GDynamicRHI->CreateBuffer_RenderThread(FRHICommandListExecutor::GetImmediateCommandList(), Size, InUsage | BUF_VertexBuffer, 0, InResourceState, CreateInfo);
-	Buffer->SetTrackedAccess_Unsafe(InResourceState);
-	return Buffer;
+	check(IsInRenderingThread());
+	return FRHICommandListExecutor::GetImmediateCommandList().CreateBuffer(Size, Usage | EBufferUsageFlags::VertexBuffer, 0, ResourceState, CreateInfo);
 }
 
-FORCEINLINE FBufferRHIRef RHIAsyncCreateVertexBuffer(uint32 Size, EBufferUsageFlags InUsage, ERHIAccess InResourceState, FRHIResourceCreateInfo& CreateInfo)
+UE_DEPRECATED(5.1, "RHIAsyncCreateVertexBuffer is deprecated. Use FRHICommandList::CreateBuffer instead.")
+FORCEINLINE FBufferRHIRef RHIAsyncCreateVertexBuffer(uint32 Size, EBufferUsageFlags Usage, ERHIAccess ResourceState, FRHIResourceCreateInfo& CreateInfo)
 {
-	FBufferRHIRef Buffer = GDynamicRHI->RHICreateBuffer(Size, InUsage, 0, InResourceState, CreateInfo);
-	Buffer->SetTrackedAccess_Unsafe(InResourceState);
-	return Buffer;
+	check(IsInRenderingThread());
+	return FRHICommandListExecutor::GetImmediateCommandList().CreateBuffer(Size, Usage, 0, ResourceState, CreateInfo);
 }
 
 FORCEINLINE FBufferRHIRef RHICreateVertexBuffer(uint32 Size, EBufferUsageFlags InUsage, FRHIResourceCreateInfo& CreateInfo)
 {
 	bool bHasInitialData = CreateInfo.BulkData != nullptr;
-	EBufferUsageFlags Usage = InUsage | BUF_VertexBuffer;
+	EBufferUsageFlags Usage = InUsage | EBufferUsageFlags::VertexBuffer;
 	ERHIAccess ResourceState = RHIGetDefaultResourceState(Usage, bHasInitialData);
 	return RHICreateVertexBuffer(Size, Usage, ResourceState, CreateInfo);
 }
 
+UE_DEPRECATED(5.1, "RHIAsyncCreateVertexBuffer is deprecated. Use FRHICommandList::CreateBuffer instead.")
 FORCEINLINE FBufferRHIRef RHIAsyncCreateVertexBuffer(uint32 Size, EBufferUsageFlags InUsage, FRHIResourceCreateInfo& CreateInfo)
 {
 	bool bHasInitialData = CreateInfo.BulkData != nullptr;
-	EBufferUsageFlags Usage = InUsage | BUF_VertexBuffer;
+	EBufferUsageFlags Usage = InUsage | EBufferUsageFlags::VertexBuffer;
 	ERHIAccess ResourceState = RHIGetDefaultResourceState(Usage, bHasInitialData);
+PRAGMA_DISABLE_DEPRECATION_WARNINGS
 	return RHIAsyncCreateVertexBuffer(Size, Usage, ResourceState, CreateInfo);
+PRAGMA_ENABLE_DEPRECATION_WARNINGS
 }
 
 UE_DEPRECATED(5.0, "Buffer locks have been unified. Use RHILockBuffer() instead.")
 FORCEINLINE void* RHILockVertexBuffer(FRHIBuffer* VertexBuffer, uint32 Offset, uint32 SizeRHI, EResourceLockMode LockMode)
 {
+	check(IsInRenderingThread());
 	return FRHICommandListExecutor::GetImmediateCommandList().LockBuffer(VertexBuffer, Offset, SizeRHI, LockMode);
 }
 
 UE_DEPRECATED(5.0, "Buffer locks have been unified. Use RHIUnlockBuffer() instead.")
 FORCEINLINE void RHIUnlockVertexBuffer(FRHIBuffer* VertexBuffer)
 {
-	 FRHICommandListExecutor::GetImmediateCommandList().UnlockBuffer(VertexBuffer);
+	check(IsInRenderingThread());
+	FRHICommandListExecutor::GetImmediateCommandList().UnlockBuffer(VertexBuffer);
 }
 
-FORCEINLINE FBufferRHIRef RHICreateStructuredBuffer(uint32 Stride, uint32 Size, EBufferUsageFlags InUsage, ERHIAccess InResourceState, FRHIResourceCreateInfo& CreateInfo)
+FORCEINLINE FBufferRHIRef RHICreateStructuredBuffer(uint32 Stride, uint32 Size, EBufferUsageFlags Usage, ERHIAccess ResourceState, FRHIResourceCreateInfo& CreateInfo)
 {
-	FBufferRHIRef Buffer = GDynamicRHI->CreateBuffer_RenderThread(FRHICommandListExecutor::GetImmediateCommandList(), Size, InUsage | BUF_StructuredBuffer, Stride, InResourceState, CreateInfo);
-	Buffer->SetTrackedAccess_Unsafe(InResourceState);
-	return Buffer;
+	check(IsInRenderingThread());
+	return FRHICommandListExecutor::GetImmediateCommandList().CreateBuffer(Size, Usage | EBufferUsageFlags::StructuredBuffer, Stride, ResourceState, CreateInfo);
 }
 
 FORCEINLINE FBufferRHIRef RHICreateStructuredBuffer(uint32 Stride, uint32 Size, EBufferUsageFlags InUsage, FRHIResourceCreateInfo& CreateInfo)
@@ -5329,22 +5381,26 @@ FORCEINLINE FBufferRHIRef RHICreateStructuredBuffer(uint32 Stride, uint32 Size,
 UE_DEPRECATED(5.0, "Buffer locks have been unified. Use RHILockBuffer() instead.")
 FORCEINLINE void* RHILockStructuredBuffer(FRHIBuffer* StructuredBuffer, uint32 Offset, uint32 SizeRHI, EResourceLockMode LockMode)
 {
+	check(IsInRenderingThread());
 	return FRHICommandListExecutor::GetImmediateCommandList().LockBuffer(StructuredBuffer, Offset, SizeRHI, LockMode);
 }
 
 UE_DEPRECATED(5.0, "Buffer locks have been unified. Use RHIUnlockBuffer() instead.")
 FORCEINLINE void RHIUnlockStructuredBuffer(FRHIBuffer* StructuredBuffer)
 {
-	 FRHICommandListExecutor::GetImmediateCommandList().UnlockBuffer(StructuredBuffer);
+	check(IsInRenderingThread());
+	FRHICommandListExecutor::GetImmediateCommandList().UnlockBuffer(StructuredBuffer);
 }
 
 FORCEINLINE void* RHILockBuffer(FRHIBuffer* Buffer, uint32 Offset, uint32 SizeRHI, EResourceLockMode LockMode)
 {
+	check(IsInRenderingThread());
 	return FRHICommandListExecutor::GetImmediateCommandList().LockBuffer(Buffer, Offset, SizeRHI, LockMode);
 }
 
 FORCEINLINE void RHIUnlockBuffer(FRHIBuffer* Buffer)
 {
+	check(IsInRenderingThread());
 	FRHICommandListExecutor::GetImmediateCommandList().UnlockBuffer(Buffer);
 }
 
@@ -5385,6 +5441,7 @@ FORCEINLINE FShaderResourceViewRHIRef RHICreateShaderResourceView(const FShaderR
 
 FORCEINLINE void RHIUpdateRHIResources(FRHIResourceUpdateInfo* UpdateInfos, int32 Num, bool bNeedReleaseRefs)
 {
+	check(IsInRenderingThread());
 	return FRHICommandListExecutor::GetImmediateCommandList().UpdateRHIResources(UpdateInfos, Num, bNeedReleaseRefs);
 }
 
diff --git a/Engine/Source/Runtime/RHI/Public/RHICommandList.inl b/Engine/Source/Runtime/RHI/Public/RHICommandList.inl
--- a/Engine/Source/Runtime/RHI/Public/RHICommandList.inl
+++ b/Engine/Source/Runtime/RHI/Public/RHICommandList.inl
@@ -32,10 +32,20 @@ FORCEINLINE_DEBUGGABLE bool FRHICommandListBase::IsImmediateAsyncCompute() const
 	return this == &FRHICommandListExecutor::GetImmediateAsyncComputeCommandList();
 }
 
+FORCEINLINE_DEBUGGABLE FRHICommandListImmediate& FRHICommandListBase::GetAsImmediate()
+{
+	checkf(IsImmediate(), TEXT("This operation expects the immediate command list."));
+	return static_cast<FRHICommandListImmediate&>(*this);
+}
+
 FORCEINLINE_DEBUGGABLE bool FRHICommandListBase::Bypass() const
 {
 	check(!IsImmediate() || IsInRenderingThread() || IsInRHIThread());
-	return GRHICommandList.Bypass();
+	return GRHICommandList.Bypass()
+#if CAN_DISALLOW_COMMAND_LIST_BYPASS
+		&& bAllowBypass
+#endif
+		;
 }
 
 FORCEINLINE_DEBUGGABLE FScopedRHIThreadStaller::FScopedRHIThreadStaller(class FRHICommandListImmediate& InImmed, bool bDoStall)
diff --git a/Engine/Source/Runtime/RHI/Public/RHIValidation.h b/Engine/Source/Runtime/RHI/Public/RHIValidation.h
--- a/Engine/Source/Runtime/RHI/Public/RHIValidation.h
+++ b/Engine/Source/Runtime/RHI/Public/RHIValidation.h
@@ -380,9 +380,15 @@ public:
 	* @param ResourceArray - An optional pointer to a resource array containing the resource's data.
 	*/
 	// FlushType: Wait RHI Thread
-	virtual FBufferRHIRef RHICreateBuffer(uint32 Size, EBufferUsageFlags Usage, uint32 Stride, ERHIAccess ResourceState, FRHIResourceCreateInfo& CreateInfo) override final
+	virtual FBufferRHIRef RHICreateBuffer(FRHICommandListBase& RHICmdList, uint32 Size, EBufferUsageFlags Usage, uint32 Stride, ERHIAccess ResourceState, FRHIResourceCreateInfo& CreateInfo) override final
 	{
-		FBufferRHIRef Buffer = RHI->RHICreateBuffer(Size, Usage, Stride, ResourceState, CreateInfo);
+		if (CreateInfo.ResourceArray && RHICmdList.IsInsideRenderPass())
+		{
+			FString Msg = FString::Printf(TEXT("Creating buffers with initial data during a render pass is not supported, buffer name: \"%s\""), CreateInfo.DebugName);
+			RHI_VALIDATION_CHECK(false, *Msg);
+		}
+
+		FBufferRHIRef Buffer = RHI->RHICreateBuffer(RHICmdList, Size, Usage, Stride, ResourceState, CreateInfo);
 		Buffer->InitBarrierTracking(ResourceState, CreateInfo.DebugName);
 		return Buffer;
 	}
@@ -395,15 +401,15 @@ public:
 	}
 
 	// FlushType: Flush RHI Thread
-	virtual void* RHILockBuffer(class FRHICommandListImmediate& RHICmdList, FRHIBuffer* Buffer, uint32 Offset, uint32 SizeRHI, EResourceLockMode LockMode) override final;
-	virtual void* RHILockBufferMGPU(class FRHICommandListImmediate& RHICmdList, FRHIBuffer* Buffer, uint32 GPUIndex, uint32 Offset, uint32 SizeRHI, EResourceLockMode LockMode) override final;
+	virtual void* RHILockBuffer(class FRHICommandListBase& RHICmdList, FRHIBuffer* Buffer, uint32 Offset, uint32 SizeRHI, EResourceLockMode LockMode) override final;
+	virtual void* RHILockBufferMGPU(class FRHICommandListBase& RHICmdList, FRHIBuffer* Buffer, uint32 GPUIndex, uint32 Offset, uint32 SizeRHI, EResourceLockMode LockMode) override final;
 
 	// FlushType: Flush RHI Thread
-	virtual void RHIUnlockBuffer(class FRHICommandListImmediate& RHICmdList, FRHIBuffer* Buffer) override final
+	virtual void RHIUnlockBuffer(class FRHICommandListBase& RHICmdList, FRHIBuffer* Buffer) override final
 	{
 		RHI->RHIUnlockBuffer(RHICmdList, Buffer);
 	}
-	virtual void RHIUnlockBufferMGPU(class FRHICommandListImmediate& RHICmdList, FRHIBuffer* Buffer, uint32 GPUIndex) override final
+	virtual void RHIUnlockBufferMGPU(class FRHICommandListBase& RHICmdList, FRHIBuffer* Buffer, uint32 GPUIndex) override final
 	{
 		RHI->RHIUnlockBufferMGPU(RHICmdList, Buffer, GPUIndex);
 	}
@@ -1273,19 +1279,6 @@ public:
 		return RHI->RHIGetMinimumAlignmentForBufferBackedSRV(Format);
 	}
 
-	virtual FBufferRHIRef CreateBuffer_RenderThread(class FRHICommandListImmediate& RHICmdList, uint32 Size, EBufferUsageFlags Usage, uint32 Stride, ERHIAccess ResourceState, FRHIResourceCreateInfo& CreateInfo) override final
-	{
-		if (CreateInfo.ResourceArray && RHICmdList.IsInsideRenderPass())
-		{
-			FString Msg = FString::Printf(TEXT("Creating buffers with initial data during a render pass is not supported, buffer name: \"%s\""), CreateInfo.DebugName);
-			RHI_VALIDATION_CHECK(false, *Msg);
-		}
-
-		FBufferRHIRef Buffer = RHI->CreateBuffer_RenderThread(RHICmdList, Size, Usage, Stride, ResourceState, CreateInfo);
-		Buffer->InitBarrierTracking(ResourceState, CreateInfo.DebugName);
-		return Buffer;
-	}
-
 	virtual FShaderResourceViewRHIRef CreateShaderResourceView_RenderThread(class FRHICommandListImmediate& RHICmdList, FRHIBuffer* Buffer, uint32 Stride, uint8 Format) override final
 	{
 		FShaderResourceViewRHIRef SRV = RHI->CreateShaderResourceView_RenderThread(RHICmdList, Buffer, Stride, Format);
@@ -1334,13 +1327,13 @@ public:
 		return RHI->CancelAsyncReallocateTexture2D_RenderThread(RHICmdList, Texture2D, bBlockUntilCompleted);
 	}
 
-	virtual void* LockBuffer_BottomOfPipe(class FRHICommandListImmediate& RHICmdList, FRHIBuffer* Buffer, uint32 Offset, uint32 SizeRHI, EResourceLockMode LockMode) override final
+	virtual void* LockBuffer_BottomOfPipe(class FRHICommandListBase& RHICmdList, FRHIBuffer* Buffer, uint32 Offset, uint32 SizeRHI, EResourceLockMode LockMode) override final
 	{
 		RHI_VALIDATION_CHECK(LockMode != RLM_WriteOnly_NoOverwrite || GRHISupportsMapWriteNoOverwrite, TEXT("Using RLM_WriteOnly_NoOverwrite when the RHI doesn't support it."));
 		return RHI->LockBuffer_BottomOfPipe(RHICmdList, Buffer, Offset, SizeRHI, LockMode);
 	}
 	
-	virtual void UnlockBuffer_BottomOfPipe(class FRHICommandListImmediate& RHICmdList, FRHIBuffer* Buffer) override final
+	virtual void UnlockBuffer_BottomOfPipe(class FRHICommandListBase& RHICmdList, FRHIBuffer* Buffer) override final
 	{
 		RHI->UnlockBuffer_BottomOfPipe(RHICmdList, Buffer);
 	}
@@ -1651,7 +1644,7 @@ private:
 	void ValidatePipeline(const FGraphicsPipelineStateInitializer& Initializer);
 
 	// Shared validation logic, called from RHILockBuffer / RHILockBufferMGPU
-	void LockBufferValidate(class FRHICommandListImmediate& RHICmdList, FRHIBuffer* Buffer, EResourceLockMode LockMode);
+	void LockBufferValidate(class FRHICommandListBase& RHICmdList, FRHIBuffer* Buffer, EResourceLockMode LockMode);
 };
 
 #endif	// ENABLE_RHI_VALIDATION
diff --git a/Engine/Source/Runtime/RenderCore/Public/RenderResource.h b/Engine/Source/Runtime/RenderCore/Public/RenderResource.h
--- a/Engine/Source/Runtime/RenderCore/Public/RenderResource.h
+++ b/Engine/Source/Runtime/RenderCore/Public/RenderResource.h
@@ -12,6 +12,7 @@
 #include "RenderCore.h"
 #include "Serialization/MemoryLayout.h"
 #include "Containers/DynamicRHIResourceArray.h"
+#include "RenderingThread.h"
 
 /** Number of frames after which unused global resource allocations will be discarded. */
 extern int32 GGlobalBufferNumFramesUnusedThresold;
@@ -179,7 +180,8 @@ protected:
 			}
 			else
 			{
-				Buffer = RHIAsyncCreateVertexBuffer(SizeInBytes, InBufferUsageFlags, CreateInfo);
+				FRHIAsyncCommandList CommandList;
+				return CommandList->CreateBuffer(SizeInBytes, InBufferUsageFlags | EBufferUsageFlags::VertexBuffer, 0, ERHIAccess::SRVMask, CreateInfo);
 			}
 		}
 
diff --git a/Engine/Source/Runtime/RenderCore/Public/RenderingThread.h b/Engine/Source/Runtime/RenderCore/Public/RenderingThread.h
--- a/Engine/Source/Runtime/RenderCore/Public/RenderingThread.h
+++ b/Engine/Source/Runtime/RenderCore/Public/RenderingThread.h
@@ -13,8 +13,8 @@
 #include "Trace/Trace.h"
 #include "Serialization/MemoryLayout.h"
 #include "Delegates/Delegate.h"
-
-class FRHICommandListImmediate;
+#include "RHI.h"
+#include "RHICommandList.h"
 
 ////////////////////////////////////
 // Render thread API
@@ -326,6 +326,50 @@ extern RENDERCORE_API FPendingCleanupObjects* GetPendingCleanupObjects();
 // RenderThread scoped work
 ////////////////////////////////////
 
+/** A utility to record RHI commands asynchronously and then enqueue the resulting commands to the render thread. */
+class FRHIAsyncCommandList
+{
+public:
+	FRHIAsyncCommandList(FRHIGPUMask InGPUMask = FRHIGPUMask::All())
+		: RHICmdListStack(InGPUMask)
+	{
+		RHICmdListStack.DisallowBypass();
+	}
+
+	FRHICommandList& GetCommandList()
+	{
+		return RHICmdListStack;
+	}
+
+	FRHICommandList& operator*()
+	{
+		return RHICmdListStack;
+	}
+
+	FRHICommandList* operator->()
+	{
+		return &RHICmdListStack;
+	}
+
+	~FRHIAsyncCommandList()
+	{
+		if (RHICmdListStack.HasCommands())
+		{
+			FRHICommandList* RHICmdList = new FRHICommandList(FRHIGPUMask::All());
+			RHICmdList->ExchangeCmdList(RHICmdListStack);
+
+			ENQUEUE_RENDER_COMMAND(AsyncCommandListScope)(
+				[RHICmdList](FRHICommandListImmediate& RHICmdListImmediate)
+			{
+				RHICmdListImmediate.QueueCommandListSubmit(RHICmdList);
+			});
+		}
+	}
+
+private:
+	FRHICommandList RHICmdListStack;
+};
+
 class RENDERCORE_API FRenderThreadScope
 {
 	typedef TFunction<void(FRHICommandListImmediate&)> RenderCommandFunction;
diff --git a/Engine/Source/Runtime/VulkanRHI/Private/VulkanBuffer.cpp b/Engine/Source/Runtime/VulkanRHI/Private/VulkanBuffer.cpp
--- a/Engine/Source/Runtime/VulkanRHI/Private/VulkanBuffer.cpp
+++ b/Engine/Source/Runtime/VulkanRHI/Private/VulkanBuffer.cpp
@@ -112,7 +112,7 @@
 	return OutVkUsage;
 }
 
-FVulkanResourceMultiBuffer::FVulkanResourceMultiBuffer(FVulkanDevice* InDevice, uint32 InSize, EBufferUsageFlags InUEUsage, uint32 InStride, FRHIResourceCreateInfo& CreateInfo, class FRHICommandListImmediate* InRHICmdList, const FRHITransientHeapAllocation* InTransientHeapAllocation)
+FVulkanResourceMultiBuffer::FVulkanResourceMultiBuffer(FVulkanDevice* InDevice, uint32 InSize, EBufferUsageFlags InUEUsage, uint32 InStride, FRHIResourceCreateInfo& CreateInfo, FRHICommandListBase* InRHICmdList, const FRHITransientHeapAllocation* InTransientHeapAllocation)
 	: FRHIBuffer(InSize, InUEUsage, InStride)
 	, VulkanRHI::FDeviceChild(InDevice)
 	, BufferUsageFlags(0)
@@ -127,13 +127,15 @@
 	
 	if (!bZeroSize)
 	{
+		check(InDevice);
+
 		const bool bVolatile = EnumHasAnyFlags(InUEUsage, BUF_Volatile);
 		if (bVolatile)
 		{
-			bool bRenderThread = IsInRenderingThread();
+			check(InRHICmdList);
 
 			// Get a dummy buffer as sometimes the high-level misbehaves and tries to use SRVs off volatile buffers before filling them in...
-			void* Data = Lock(bRenderThread, RLM_WriteOnly, InSize, 0);
+			void* Data = Lock(*InRHICmdList, RLM_WriteOnly, InSize, 0);
 
 			if (CreateInfo.ResourceArray)
 			{
@@ -145,12 +147,10 @@
 				FMemory::Memzero(Data, InSize);
 			}
 
-			Unlock(bRenderThread);
+			Unlock(*InRHICmdList);
 		}
 		else
 		{
-			VkDevice VulkanDevice = InDevice->GetInstanceHandle();
-
 			NumBuffers = GetNumBuffersFromUsage(InUEUsage);
 			check(NumBuffers <= UE_ARRAY_COUNT(Buffers));
 
@@ -197,14 +197,10 @@
 			Current.Offset = Current.Alloc.Offset;
 			Current.Size = InSize;
 
-			bool bRenderThread = (InRHICmdList == nullptr);
-			if (bRenderThread)
-			{
-				ensure(IsInRenderingThread());
-			}
-
 			if (CreateInfo.ResourceArray)
 			{
+				check(InRHICmdList);
+
 				uint32 CopyDataSize = FMath::Min(InSize, CreateInfo.ResourceArray->GetResourceDataSize());
 				// We know this buffer is not in use by GPU atm. If we do have a direct access initialize it without extra copies
 				if (bUnifiedMem)
@@ -214,9 +210,9 @@
 				}
 				else
 				{
-					void* Data = Lock(bRenderThread, RLM_WriteOnly, CopyDataSize, 0);
+					void* Data = Lock(*InRHICmdList, RLM_WriteOnly, CopyDataSize, 0);
 					FMemory::Memcpy(Data, CreateInfo.ResourceArray->GetResourceData(), CopyDataSize);
-					Unlock(bRenderThread);
+					Unlock(*InRHICmdList);
 				}
 
 				CreateInfo.ResourceArray->Discard();
@@ -239,7 +235,7 @@
 	UpdateVulkanBufferStats(TotalSize, BufferUsageFlags, false);
 }
 
-void* FVulkanResourceMultiBuffer::Lock(bool bFromRenderingThread, EResourceLockMode LockMode, uint32 LockSize, uint32 Offset)
+void* FVulkanResourceMultiBuffer::Lock(FRHICommandListBase& RHICmdList, EResourceLockMode LockMode, uint32 LockSize, uint32 Offset)
 {
 	void* Data = nullptr;
 	uint32 DataOffset = 0;
@@ -278,6 +274,8 @@
 
 		if (LockMode == RLM_ReadOnly)
 		{
+			check(IsInRenderingThread() && RHICmdList.IsImmediate());
+
 			const bool bUnifiedMem = Device->HasUnifiedMemory();
 			if (bUnifiedMem)
 			{
@@ -428,7 +426,7 @@
 };
 
 
-void FVulkanResourceMultiBuffer::Unlock(bool bFromRenderingThread)
+void FVulkanResourceMultiBuffer::Unlock(FRHICommandListBase& RHICmdList)
 {
 	const bool bDynamic = EnumHasAnyFlags(GetUsage(), BUF_Dynamic);
 	const bool bVolatile = EnumHasAnyFlags(GetUsage(), BUF_Volatile);
@@ -459,14 +457,12 @@
 		checkf(bFound, TEXT("Mismatched lock/unlock IndexBuffer!"));
 		if (PendingLock.LockMode == RLM_WriteOnly)
 		{
-			FRHICommandList& RHICmdList = FRHICommandListExecutor::GetImmediateCommandList();
-			if (!bFromRenderingThread || (RHICmdList.Bypass() || !IsRunningRHIInSeparateThread()))
+			if (RHICmdList.IsBottomOfPipe())
 			{
 				FVulkanResourceMultiBuffer::InternalUnlock(Device->GetImmediateContext(), PendingLock, this, DynamicBufferIndex);
 			}
 			else
 			{
-				check(IsInRenderingThread());
 				ALLOC_COMMAND_CL(RHICmdList, FRHICommandMultiBufferUnlock)(Device, PendingLock, this, DynamicBufferIndex);
 			}
 		}
@@ -497,58 +493,31 @@
 	::Swap(VolatileLockInfo, Other.VolatileLockInfo);
 }
 
-FBufferRHIRef FVulkanDynamicRHI::RHICreateBuffer(uint32 Size, EBufferUsageFlags Usage, uint32 Stride, ERHIAccess ResourceState, FRHIResourceCreateInfo& ResourceCreateInfo)
+FBufferRHIRef FVulkanDynamicRHI::RHICreateBuffer(FRHICommandListBase& RHICmdList, uint32 Size, EBufferUsageFlags Usage, uint32 Stride, ERHIAccess ResourceState, FRHIResourceCreateInfo& ResourceCreateInfo)
 {
 	LLM_SCOPE_VULKAN(ELLMTagVulkan::VulkanBuffers);
 
 	if (ResourceCreateInfo.bWithoutNativeResource)
 	{
-		return new FVulkanResourceMultiBuffer(nullptr, 0, BUF_None, 0, ResourceCreateInfo, nullptr);
+		return new FVulkanResourceMultiBuffer(Device, 0, BUF_None, 0, ResourceCreateInfo, &RHICmdList);
 	}
-	return new FVulkanResourceMultiBuffer(Device, Size, Usage, Stride, ResourceCreateInfo, nullptr);
+	return new FVulkanResourceMultiBuffer(Device, Size, Usage, Stride, ResourceCreateInfo, &RHICmdList);
 }
 
-FRHIBuffer* FVulkanDynamicRHI::CreateBuffer(const FRHIBufferCreateInfo& InCreateInfo, FRHIResourceCreateInfo& InResourceCreateInfo, const FRHITransientHeapAllocation* InTransientHeapAllocation)
-{
-	LLM_SCOPE_VULKAN(ELLMTagVulkan::VulkanBuffers);
-
-	if (InTransientHeapAllocation == nullptr)
-	{
-		return this->RHICreateBuffer(InCreateInfo.Size, InCreateInfo.Usage, InCreateInfo.Stride, ERHIAccess::None, InResourceCreateInfo);
-	}
-
-	checkf(!EnumHasAnyFlags(InCreateInfo.Usage, BUF_AccelerationStructure), TEXT("AccelerationStructure not yet supported as TransientResource."));
-	checkf(!InResourceCreateInfo.bWithoutNativeResource, TEXT("WithoutNativeResource not yet supported as TransientResource."));
-
-	return new FVulkanResourceMultiBuffer(Device, InCreateInfo.Size, InCreateInfo.Usage, InCreateInfo.Stride, InResourceCreateInfo, nullptr, InTransientHeapAllocation);
-}
-
-void* FVulkanDynamicRHI::LockBuffer_BottomOfPipe(FRHICommandListImmediate& RHICmdList, FRHIBuffer* BufferRHI, uint32 Offset, uint32 Size, EResourceLockMode LockMode)
+void* FVulkanDynamicRHI::LockBuffer_BottomOfPipe(FRHICommandListBase& RHICmdList, FRHIBuffer* BufferRHI, uint32 Offset, uint32 Size, EResourceLockMode LockMode)
 {
 	LLM_SCOPE_VULKAN(ELLMTagVulkan::VulkanBuffers);
 	FVulkanResourceMultiBuffer* Buffer = ResourceCast(BufferRHI);
-	return Buffer->Lock(false, LockMode, Size, Offset);
+	return Buffer->Lock(RHICmdList, LockMode, Size, Offset);
 }
 
-void FVulkanDynamicRHI::UnlockBuffer_BottomOfPipe(FRHICommandListImmediate& RHICmdList, FRHIBuffer* BufferRHI)
+void FVulkanDynamicRHI::UnlockBuffer_BottomOfPipe(FRHICommandListBase& RHICmdList, FRHIBuffer* BufferRHI)
 {
 	LLM_SCOPE_VULKAN(ELLMTagVulkan::VulkanBuffers);
 	FVulkanResourceMultiBuffer* Buffer = ResourceCast(BufferRHI);
-	Buffer->Unlock(false);
+	Buffer->Unlock(RHICmdList);
 }
 
-#if VULKAN_BUFFER_LOCK_THREADSAFE
-void* FVulkanDynamicRHI::LockBuffer_RenderThread(class FRHICommandListImmediate& RHICmdList, FRHIBuffer* BufferRHI, uint32 Offset, uint32 SizeRHI, EResourceLockMode LockMode)
-{
-	return this->RHILockBuffer(BufferRHI, Offset, SizeRHI, LockMode);
-}
-
-void FVulkanDynamicRHI::UnlockBuffer_RenderThread(FRHICommandListImmediate& RHICmdList, FRHIBuffer* BufferRHI)
-{
-	this->RHIUnlockBuffer(BufferRHI);
-}
-#endif
-
 void FVulkanDynamicRHI::RHICopyBuffer(FRHIBuffer* SourceBufferRHI, FRHIBuffer* DestBufferRHI)
 {
 	VULKAN_SIGNAL_UNIMPLEMENTED();
diff --git a/Engine/Source/Runtime/VulkanRHI/Private/VulkanRHI.cpp b/Engine/Source/Runtime/VulkanRHI/Private/VulkanRHI.cpp
--- a/Engine/Source/Runtime/VulkanRHI/Private/VulkanRHI.cpp
+++ b/Engine/Source/Runtime/VulkanRHI/Private/VulkanRHI.cpp
@@ -215,6 +215,7 @@ FVulkanDynamicRHI::FVulkanDynamicRHI()
 	GPoolSizeVRAMPercentage = 0;
 	GTexturePoolSize = 0;
 	GRHISupportsMultithreading = true;
+	GRHISupportsMultithreadedResources = true;
 	GRHISupportsPipelineFileCache = true;
 	GRHITransitionPrivateData_SizeInBytes = sizeof(FVulkanPipelineBarrier);
 	GRHITransitionPrivateData_AlignInBytes = alignof(FVulkanPipelineBarrier);
diff --git a/Engine/Source/Runtime/VulkanRHI/Private/VulkanRayTracing.cpp b/Engine/Source/Runtime/VulkanRHI/Private/VulkanRayTracing.cpp
--- a/Engine/Source/Runtime/VulkanRHI/Private/VulkanRayTracing.cpp
+++ b/Engine/Source/Runtime/VulkanRHI/Private/VulkanRayTracing.cpp
@@ -530,9 +530,9 @@ void FVulkanRayTracingScene::BuildPerInstanceGeometryParameterBuffer()
 	check(PerInstanceGeometryParameterBuffer->GetSize() >= ParameterBufferSize);
 
 	check(IsInRHIThread() || !IsRunningRHIInSeparateThread());
-	const bool bTopOfPipe = false; // running on RHI timeline
+	FRHICommandListImmediate& RHICmdList = FRHICommandListExecutor::GetImmediateCommandList();
 
-	void* MappedBuffer = PerInstanceGeometryParameterBuffer->Lock(bTopOfPipe, RLM_WriteOnly, ParameterBufferSize, 0);
+	void* MappedBuffer = PerInstanceGeometryParameterBuffer->Lock(RHICmdList, RLM_WriteOnly, ParameterBufferSize, 0);
 	FVulkanRayTracingGeometryParameters* MappedParameters = reinterpret_cast<FVulkanRayTracingGeometryParameters*>(MappedBuffer);
 	uint32 ParameterIndex = 0;
 
@@ -577,7 +577,7 @@ void FVulkanRayTracingScene::BuildPerInstanceGeometryParameterBuffer()
 
 	check(ParameterIndex == Initializer.NumTotalSegments);
 
-	PerInstanceGeometryParameterBuffer->Unlock(bTopOfPipe);
+	PerInstanceGeometryParameterBuffer->Unlock(RHICmdList);
 }
 
 void FVulkanDynamicRHI::RHITransferRayTracingGeometryUnderlyingResource(FRHIRayTracingGeometry* DestGeometry, FRHIRayTracingGeometry* SrcGeometry)
diff --git a/Engine/Source/Runtime/VulkanRHI/Private/VulkanResources.h b/Engine/Source/Runtime/VulkanRHI/Private/VulkanResources.h
--- a/Engine/Source/Runtime/VulkanRHI/Private/VulkanResources.h
+++ b/Engine/Source/Runtime/VulkanRHI/Private/VulkanResources.h
@@ -860,7 +860,7 @@ protected:
 class FVulkanResourceMultiBuffer : public FRHIBuffer, public FVulkanEvictable, public VulkanRHI::FDeviceChild
 {
 public:
-	FVulkanResourceMultiBuffer(FVulkanDevice* InDevice, uint32 InSize, EBufferUsageFlags InUEUsage, uint32 InStride, FRHIResourceCreateInfo& CreateInfo, class FRHICommandListImmediate* InRHICmdList = nullptr, const FRHITransientHeapAllocation* InTransientHeapAllocation = nullptr);
+	FVulkanResourceMultiBuffer(FVulkanDevice* InDevice, uint32 InSize, EBufferUsageFlags InUEUsage, uint32 InStride, FRHIResourceCreateInfo& CreateInfo, class FRHICommandListBase* InRHICmdList = nullptr, const FRHITransientHeapAllocation* InTransientHeapAllocation = nullptr);
 	virtual ~FVulkanResourceMultiBuffer();
 
 	// FVulkanEvictable interface.
@@ -932,8 +932,8 @@ public:
 		return (GetStride() == 4)? VK_INDEX_TYPE_UINT32 : VK_INDEX_TYPE_UINT16;
 	}
 
-	void* Lock(bool bFromRenderingThread, EResourceLockMode LockMode, uint32 Size, uint32 Offset);
-	void Unlock(bool bFromRenderingThread);
+	void* Lock(FRHICommandListBase& RHICmdList, EResourceLockMode LockMode, uint32 Size, uint32 Offset);
+	void Unlock(FRHICommandListBase& RHICmdList);
 
 	void Swap(FVulkanResourceMultiBuffer& Other);
 
diff --git a/Engine/Source/Runtime/VulkanRHI/Private/VulkanTransientResourceAllocator.cpp b/Engine/Source/Runtime/VulkanRHI/Private/VulkanTransientResourceAllocator.cpp
--- a/Engine/Source/Runtime/VulkanRHI/Private/VulkanTransientResourceAllocator.cpp
+++ b/Engine/Source/Runtime/VulkanRHI/Private/VulkanTransientResourceAllocator.cpp
@@ -126,6 +126,9 @@ FRHITransientTexture* FVulkanTransientResourceAllocator::CreateTexture(const FRH
 
 FRHITransientBuffer* FVulkanTransientResourceAllocator::CreateBuffer(const FRHIBufferCreateInfo& InCreateInfo, const TCHAR* InDebugName, uint32 InPassIndex)
 {
+	checkf(!EnumHasAnyFlags(InCreateInfo.Usage, BUF_AccelerationStructure), TEXT("AccelerationStructure not yet supported as TransientResource."));
+	checkf(!EnumHasAnyFlags(InCreateInfo.Usage, BUF_Volatile), TEXT("The volatile flag is not supported for transient resources."));
+
 	const VkBufferUsageFlags VulkanBufferUsage = FVulkanResourceMultiBuffer::UEToVKBufferUsageFlags(Device, InCreateInfo.Usage, (InCreateInfo.Size == 0));
 	const uint32 Alignment = FMemoryManager::CalculateBufferAlignment(*Device, VulkanBufferUsage);
 	uint64 Size = Align(InCreateInfo.Size, Alignment) * FVulkanResourceMultiBuffer::GetNumBuffersFromUsage(InCreateInfo.Usage);
@@ -134,7 +137,7 @@ FRHITransientBuffer* FVulkanTransientResourceAllocator::CreateBuffer(const FRHIB
 		[&](const FRHITransientHeap::FResourceInitializer& Initializer)
 	{
 		FRHIResourceCreateInfo ResourceCreateInfo(InDebugName);
-		FRHIBuffer* Buffer = GVulkanRHI->CreateBuffer(InCreateInfo, ResourceCreateInfo, &Initializer.Allocation);
+		FRHIBuffer* Buffer = new FVulkanResourceMultiBuffer(Device, InCreateInfo.Size, InCreateInfo.Usage, InCreateInfo.Stride, ResourceCreateInfo, nullptr, &Initializer.Allocation);
 		Buffer->SetTrackedAccess_Unsafe(ERHIAccess::Discard);
 		return new FRHITransientBuffer(Buffer, 0/*GpuVirtualAddress*/, Initializer.Hash, Size, ERHITransientAllocationType::Heap, InCreateInfo);
 	});
diff --git a/Engine/Source/Runtime/VulkanRHI/Public/VulkanDynamicRHI.h b/Engine/Source/Runtime/VulkanRHI/Public/VulkanDynamicRHI.h
--- a/Engine/Source/Runtime/VulkanRHI/Public/VulkanDynamicRHI.h
+++ b/Engine/Source/Runtime/VulkanRHI/Public/VulkanDynamicRHI.h
@@ -14,9 +14,6 @@ class FVulkanQueue;
 class IHeadMountedDisplayVulkanExtensions;
 struct FRHITransientHeapAllocation;
 
-// TODO: fix Lock/Unlock Vertex/Index buffer
-#define VULKAN_BUFFER_LOCK_THREADSAFE 0
-
 struct FOptionalVulkanInstanceExtensions
 {
 	union
@@ -108,11 +105,11 @@ public:
 	virtual FComputePipelineStateRHIRef RHICreateComputePipelineState(FRHIComputeShader* ComputeShader) final override;
 	virtual FUniformBufferRHIRef RHICreateUniformBuffer(const void* Contents, const FRHIUniformBufferLayout* Layout, EUniformBufferUsage Usage, EUniformBufferValidation Validation) final override;
 	virtual void RHIUpdateUniformBuffer(FRHIUniformBuffer* UniformBufferRHI, const void* Contents) final override;
-	virtual FBufferRHIRef RHICreateBuffer(uint32 Size, EBufferUsageFlags Usage, uint32 Stride, ERHIAccess ResourceState, FRHIResourceCreateInfo& CreateInfo) final override;
+	virtual FBufferRHIRef RHICreateBuffer(FRHICommandListBase& RHICmdList, uint32 Size, EBufferUsageFlags Usage, uint32 Stride, ERHIAccess ResourceState, FRHIResourceCreateInfo& CreateInfo) final override;
 	virtual void RHICopyBuffer(FRHIBuffer* SourceBuffer, FRHIBuffer* DestBuffer) final override;
 	virtual void RHITransferBufferUnderlyingResource(FRHIBuffer* DestBuffer, FRHIBuffer* SrcBuffer) final override;
-	virtual void* LockBuffer_BottomOfPipe(FRHICommandListImmediate& RHICmdList, FRHIBuffer* Buffer, uint32 Offset, uint32 Size, EResourceLockMode LockMode) final override;
-	virtual void UnlockBuffer_BottomOfPipe(FRHICommandListImmediate& RHICmdList, FRHIBuffer* Buffer) final override;
+	virtual void* LockBuffer_BottomOfPipe(FRHICommandListBase& RHICmdList, FRHIBuffer* Buffer, uint32 Offset, uint32 Size, EResourceLockMode LockMode) final override;
+	virtual void UnlockBuffer_BottomOfPipe(FRHICommandListBase& RHICmdList, FRHIBuffer* Buffer) final override;
 	virtual FUnorderedAccessViewRHIRef RHICreateUnorderedAccessView(FRHIBuffer* Buffer, bool bUseUAVCounter, bool bAppendBuffer) final override;
 	virtual FUnorderedAccessViewRHIRef RHICreateUnorderedAccessView(FRHITexture* Texture, uint32 MipLevel, uint16 FirstArraySlice, uint16 NumArraySlices) final override;
 	virtual FUnorderedAccessViewRHIRef RHICreateUnorderedAccessView(FRHIBuffer* Buffer, uint8 Format) final override;
@@ -200,22 +197,10 @@ public:
 
 	// Transient Resource Functions and Helpers
 	virtual IRHITransientResourceAllocator* RHICreateTransientResourceAllocator() final override;
-	FRHIBuffer* CreateBuffer(const FRHIBufferCreateInfo & InCreateInfo, FRHIResourceCreateInfo & InResourceCreateInfo, const FRHITransientHeapAllocation* InAllocation);
-
-	virtual FBufferRHIRef CreateBuffer_RenderThread(class FRHICommandListImmediate& RHICmdList, uint32 Size, EBufferUsageFlags Usage, uint32 Stride, ERHIAccess ResourceState, FRHIResourceCreateInfo& CreateInfo) override final
-	{
-		// this-> is required to avoid calling the global version of this function
-		return this->RHICreateBuffer(Size, Usage, Stride, ResourceState, CreateInfo);
-	}
 
 	virtual FTexture2DRHIRef AsyncReallocateTexture2D_RenderThread(class FRHICommandListImmediate& RHICmdList, FRHITexture2D* Texture2D, int32 NewMipCount, int32 NewSizeX, int32 NewSizeY, FThreadSafeCounter* RequestStatus) override final;
 	//virtual ETextureReallocationStatus CancelAsyncReallocateTexture2D_RenderThread(class FRHICommandListImmediate& RHICmdList, FRHITexture2D* Texture2D, bool bBlockUntilCompleted) override final;
 	
-#if VULKAN_BUFFER_LOCK_THREADSAFE		
-	virtual void* LockBuffer_RenderThread(class FRHICommandListImmediate& RHICmdList, FRHIBuffer* Buffer, uint32 Offset, uint32 SizeRHI, EResourceLockMode LockMode) override final;
-	virtual void UnlockBuffer_RenderThread(class FRHICommandListImmediate& RHICmdList, FRHIBuffer* Buffer) override final;
-#endif
-
 	virtual FVertexShaderRHIRef CreateVertexShader_RenderThread(class FRHICommandListImmediate& RHICmdList, TArrayView<const uint8> Code, const FSHAHash& Hash) override final
 	{
 		return RHICreateVertexShader(Code, Hash);
diff --git a/Engine/Source/Runtime/Windows/D3D11RHI/Private/D3D11Buffer.cpp b/Engine/Source/Runtime/Windows/D3D11RHI/Private/D3D11Buffer.cpp
--- a/Engine/Source/Runtime/Windows/D3D11RHI/Private/D3D11Buffer.cpp
+++ b/Engine/Source/Runtime/Windows/D3D11RHI/Private/D3D11Buffer.cpp
@@ -13,7 +13,7 @@ TAutoConsoleVariable<int32> GCVarUseSharedKeyedMutex(
 	TEXT("with the D3D11_RESOURCE_MISC_SHARED_KEYEDMUTEX flag instead of D3D11_RESOURCE_MISC_SHARED (default).\n"),
 	ECVF_Default);
 
-FBufferRHIRef FD3D11DynamicRHI::RHICreateBuffer(uint32 Size, EBufferUsageFlags Usage, uint32 Stride, ERHIAccess ResourceState, FRHIResourceCreateInfo& CreateInfo)
+FBufferRHIRef FD3D11DynamicRHI::RHICreateBuffer(FRHICommandListBase&, uint32 Size, EBufferUsageFlags Usage, uint32 Stride, ERHIAccess ResourceState, FRHIResourceCreateInfo& CreateInfo)
 {
 	if (CreateInfo.bWithoutNativeResource)
 	{
@@ -120,12 +120,7 @@ FBufferRHIRef FD3D11DynamicRHI::RHICreateBuffer(uint32 Size, EBufferUsageFlags U
 	return NewBuffer;
 }
 
-FBufferRHIRef FD3D11DynamicRHI::CreateBuffer_RenderThread(class FRHICommandListImmediate& RHICmdList, uint32 Size, EBufferUsageFlags Usage, uint32 Stride, ERHIAccess ResourceState, FRHIResourceCreateInfo& CreateInfo)
-{
-	return RHICreateBuffer(Size, Usage, Stride, ResourceState, CreateInfo);
-}
-
-void* FD3D11DynamicRHI::LockBuffer_BottomOfPipe(FRHICommandListImmediate& RHICmdList, FRHIBuffer* BufferRHI, uint32 Offset, uint32 Size, EResourceLockMode LockMode)
+void* FD3D11DynamicRHI::LockBuffer_BottomOfPipe(FRHICommandListBase& RHICmdList, FRHIBuffer* BufferRHI, uint32 Offset, uint32 Size, EResourceLockMode LockMode)
 {
 	FD3D11Buffer* Buffer = ResourceCast(BufferRHI);
 	// If this resource is bound to the device, unbind it
@@ -198,7 +193,7 @@ void* FD3D11DynamicRHI::LockBuffer_BottomOfPipe(FRHICommandListImmediate& RHICmd
 	return (void*)((uint8*)LockedData.GetData() + Offset);
 }
 
-void FD3D11DynamicRHI::UnlockBuffer_BottomOfPipe(FRHICommandListImmediate& RHICmdList, FRHIBuffer* BufferRHI)
+void FD3D11DynamicRHI::UnlockBuffer_BottomOfPipe(FRHICommandListBase& RHICmdList, FRHIBuffer* BufferRHI)
 {
 	FD3D11Buffer* Buffer = ResourceCast(BufferRHI);
 
diff --git a/Engine/Source/Runtime/Windows/D3D11RHI/Private/D3D11RHIPrivate.h b/Engine/Source/Runtime/Windows/D3D11RHI/Private/D3D11RHIPrivate.h
--- a/Engine/Source/Runtime/Windows/D3D11RHI/Private/D3D11RHIPrivate.h
+++ b/Engine/Source/Runtime/Windows/D3D11RHI/Private/D3D11RHIPrivate.h
@@ -476,11 +476,11 @@ public:
 	virtual FBoundShaderStateRHIRef RHICreateBoundShaderState(FRHIVertexDeclaration* VertexDeclaration, FRHIVertexShader* VertexShader, FRHIPixelShader* PixelShader, FRHIGeometryShader* GeometryShader) final override;
 	virtual FUniformBufferRHIRef RHICreateUniformBuffer(const void* Contents, const FRHIUniformBufferLayout* Layout, EUniformBufferUsage Usage, EUniformBufferValidation Validation) final override;
 	virtual void RHIUpdateUniformBuffer(FRHIUniformBuffer* UniformBufferRHI, const void* Contents) final override;
-	virtual FBufferRHIRef RHICreateBuffer(uint32 Size, EBufferUsageFlags Usage, uint32 Stride, ERHIAccess ResourceState, FRHIResourceCreateInfo& CreateInfo) final override;
+	virtual FBufferRHIRef RHICreateBuffer(FRHICommandListBase& RHICmdList, uint32 Size, EBufferUsageFlags Usage, uint32 Stride, ERHIAccess ResourceState, FRHIResourceCreateInfo& CreateInfo) final override;
 	virtual void RHICopyBuffer(FRHIBuffer* SourceBuffer, FRHIBuffer* DestBuffer) final override;
 	virtual void RHITransferBufferUnderlyingResource(FRHIBuffer* DestBuffer, FRHIBuffer* SrcBuffer) final override;
-	virtual void* LockBuffer_BottomOfPipe(FRHICommandListImmediate& RHICmdList, FRHIBuffer* Buffer, uint32 Offset, uint32 Size, EResourceLockMode LockMode) final override;
-	virtual void UnlockBuffer_BottomOfPipe(FRHICommandListImmediate& RHICmdList, FRHIBuffer* Buffer) final override;
+	virtual void* LockBuffer_BottomOfPipe(FRHICommandListBase& RHICmdList, FRHIBuffer* Buffer, uint32 Offset, uint32 Size, EResourceLockMode LockMode) final override;
+	virtual void UnlockBuffer_BottomOfPipe(FRHICommandListBase& RHICmdList, FRHIBuffer* Buffer) final override;
 	virtual FUnorderedAccessViewRHIRef RHICreateUnorderedAccessView(FRHIBuffer* Buffer, bool bUseUAVCounter, bool bAppendBuffer) final override;
 	virtual FUnorderedAccessViewRHIRef RHICreateUnorderedAccessView(FRHITexture* Texture, uint32 MipLevel, uint16 FirstArraySlice, uint16 NumArraySlices) final override;
 	virtual FUnorderedAccessViewRHIRef RHICreateUnorderedAccessView(FRHIBuffer* Buffer, uint8 Format) final override;
@@ -629,7 +629,6 @@ public:
 	// *_RenderThread functions. Command lists call these functions on RT. You can implement your own behavior inside these functions.
 	// For example, deferring the actual creation to RHI thread by sending an RHI command.
 	// For D3D11, these functions mainly just remove RHIT stalls because ID3D11Device is thread safe.
-	virtual FBufferRHIRef CreateBuffer_RenderThread(class FRHICommandListImmediate& RHICmdList, uint32 Size, EBufferUsageFlags Usage, uint32 Stride, ERHIAccess ResourceState, FRHIResourceCreateInfo& CreateInfo) final override;
 	virtual FShaderResourceViewRHIRef CreateShaderResourceView_RenderThread(class FRHICommandListImmediate& RHICmdList, FRHIBuffer* Buffer, uint32 Stride, uint8 Format) final override;
 	virtual FShaderResourceViewRHIRef CreateShaderResourceView_RenderThread(class FRHICommandListImmediate& RHICmdList, const FShaderResourceViewInitializer& Initializer) final override;
 	virtual FShaderResourceViewRHIRef CreateShaderResourceView_RenderThread(class FRHICommandListImmediate& RHICmdList, FRHIBuffer* Buffer) final override;
diff --git a/Engine/Source/Runtime/Windows/D3D11RHI/Private/Windows/WindowsD3D11Device.cpp b/Engine/Source/Runtime/Windows/D3D11RHI/Private/Windows/WindowsD3D11Device.cpp
--- a/Engine/Source/Runtime/Windows/D3D11RHI/Private/Windows/WindowsD3D11Device.cpp
+++ b/Engine/Source/Runtime/Windows/D3D11RHI/Private/Windows/WindowsD3D11Device.cpp
@@ -2005,6 +2005,8 @@ void FD3D11DynamicRHI::InitD3DDevice()
 			GRHISupportsAsyncTextureCreation = false;
 		}
 
+		GRHISupportsMultithreadedResources = GRHISupportsAsyncTextureCreation;
+
 #ifdef NVAPI_INTERFACE
 
 		if (IsRHIDeviceNVIDIA() && bAllowVendorDevice)
