Fix mismatched disable deprecation pragmas
#rnx

diff --git a/Engine/Source/Runtime/RHI/Public/RHICommandList.h b/Engine/Source/Runtime/RHI/Public/RHICommandList.h
--- a/Engine/Source/Runtime/RHI/Public/RHICommandList.h
+++ b/Engine/Source/Runtime/RHI/Public/RHICommandList.h
@@ -4252,7 +4252,7 @@ public:
 		ERHIAccess ResourceState = RHIGetDefaultResourceState(Usage, true);
 	PRAGMA_DISABLE_DEPRECATION_WARNINGS
 		return CreateAndLockIndexBuffer(Stride, Size, Usage, ResourceState, CreateInfo, OutDataBuffer);
-	PRAGMA_DISABLE_DEPRECATION_WARNINGS
+	PRAGMA_ENABLE_DEPRECATION_WARNINGS
 	}
 	
 	UE_DEPRECATED(5.0, "Buffer locks have been unified. Use LockBuffer() instead.")
@@ -4292,7 +4292,7 @@ public:
 		ERHIAccess ResourceState = RHIGetDefaultResourceState(Usage, true);
 	PRAGMA_DISABLE_DEPRECATION_WARNINGS
 		return CreateAndLockVertexBuffer(Size, Usage, ResourceState, CreateInfo, OutDataBuffer);
-	PRAGMA_DISABLE_DEPRECATION_WARNINGS
+	PRAGMA_ENABLE_DEPRECATION_WARNINGS
 	}
 
 	UE_DEPRECATED(5.0, "Buffer locks have been unified. Use LockBuffer() instead.")
diff --git a/Engine/Source/Runtime/UMG/Private/Components/ScaleBoxSlot.cpp b/Engine/Source/Runtime/UMG/Private/Components/ScaleBoxSlot.cpp
--- a/Engine/Source/Runtime/UMG/Private/Components/ScaleBoxSlot.cpp
+++ b/Engine/Source/Runtime/UMG/Private/Components/ScaleBoxSlot.cpp
@@ -14,7 +14,7 @@ UScaleBoxSlot::UScaleBoxSlot(const FObjectInitializer& ObjectInitializer)
 PRAGMA_DISABLE_DEPRECATION_WARNINGS
 	HorizontalAlignment = HAlign_Center;
 	VerticalAlignment = VAlign_Center;
-PRAGMA_DISABLE_DEPRECATION_WARNINGS
+PRAGMA_ENABLE_DEPRECATION_WARNINGS
 }
 
 void UScaleBoxSlot::ReleaseSlateResources(bool bReleaseChildren)
