Wrap additional deprecated case that was previously masked by the mismatched deprecation disabling
#rnx

diff --git a/Engine/Source/Runtime/RHI/Public/RHICommandList.h b/Engine/Source/Runtime/RHI/Public/RHICommandList.h
--- a/Engine/Source/Runtime/RHI/Public/RHICommandList.h
+++ b/Engine/Source/Runtime/RHI/Public/RHICommandList.h
@@ -5315,7 +5315,9 @@ UE_DEPRECATED(5.0, "Use RHICreateBuffer() and RHILockBuffer() instead.")
 FORCEINLINE FBufferRHIRef RHICreateAndLockVertexBuffer(uint32 Size, EBufferUsageFlags InUsage, FRHIResourceCreateInfo& CreateInfo, void*& OutDataBuffer)
 {
 	check(IsInRenderingThread());
+PRAGMA_DISABLE_DEPRECATION_WARNINGS
 	return FRHICommandListExecutor::GetImmediateCommandList().CreateAndLockVertexBuffer(Size, InUsage, CreateInfo, OutDataBuffer);
+PRAGMA_ENABLE_DEPRECATION_WARNINGS
 }
 
 FORCEINLINE FBufferRHIRef RHICreateVertexBuffer(uint32 Size, EBufferUsageFlags Usage, ERHIAccess ResourceState, FRHIResourceCreateInfo& CreateInfo)
