Renamed RHI.DriverBlacklisted to RHI.DriverDenylisted

diff --git a/Engine/Source/Runtime/RHI/Private/DynamicRHI.cpp b/Engine/Source/Runtime/RHI/Private/DynamicRHI.cpp
--- a/Engine/Source/Runtime/RHI/Private/DynamicRHI.cpp
+++ b/Engine/Source/Runtime/RHI/Private/DynamicRHI.cpp
@@ -183,7 +183,7 @@ static void RHIDetectAndWarnOfBadDrivers(bool bHasEditorToken)
 		}
 
 		GRHIAdapterDriverOnDenyList = DenyListEntry.IsValid();
-		FGenericCrashContext::SetEngineData(TEXT("RHI.DriverBlacklisted"), DenyListEntry.IsValid() ? TEXT("true") : TEXT("false"));
+		FGenericCrashContext::SetEngineData(TEXT("RHI.DriverDenylisted"), DenyListEntry.IsValid() ? TEXT("true") : TEXT("false"));
 
 		// Only alert users who are capable of updating their driver. Assume vendors with an empty list can always update.
 		bool bShowPrompt = bDeviceCanUpdateDriver || !bVendorHasEntries;
