Adding Getters for ShaderPlatform Name, the Shader Format of the ShaderPlatform and the ShaderPlatformPropertiesHash as defined in DDPI

diff --git a/Engine/Source/Runtime/RHI/Private/RHI.cpp b/Engine/Source/Runtime/RHI/Private/RHI.cpp
--- a/Engine/Source/Runtime/RHI/Private/RHI.cpp
+++ b/Engine/Source/Runtime/RHI/Private/RHI.cpp
@@ -2444,14 +2444,27 @@ void FGenericDataDrivenShaderPlatformInfo::SetDefaultValues()
 void FGenericDataDrivenShaderPlatformInfo::ParseDataDrivenShaderInfo(const FConfigSection& Section, FGenericDataDrivenShaderPlatformInfo& Info)
 {
 	Info.Language = *GetSectionString(Section, "Language");
+	Info.ShaderFormat = *GetSectionString(Section, "ShaderFormat");
 	GetFeatureLevelFromName(*GetSectionString(Section, "MaxFeatureLevel"), Info.MaxFeatureLevel);
 
+	Info.ShaderPropertiesHash = 0;
+	FString ShaderPropertiesString = Info.Name.GetPlainNameString();
+
+#define ADD_TO_PROPERTIES_STRING(SettingName, SettingValue) \
+	ShaderPropertiesString += TEXT(#SettingName); \
+	ShaderPropertiesString += FString::Printf(TEXT("_%d"), SettingValue);
+
 #define GET_SECTION_BOOL_HELPER(SettingName)	\
-	Info.SettingName = GetSectionBool(Section, #SettingName, Info.SettingName)
+	Info.SettingName = GetSectionBool(Section, #SettingName, Info.SettingName);	\
+	ADD_TO_PROPERTIES_STRING(SettingName, Info.SettingName)
+
 #define GET_SECTION_INT_HELPER(SettingName)	\
-	Info.SettingName = GetSectionUint(Section, #SettingName, Info.SettingName)
+	Info.SettingName = GetSectionUint(Section, #SettingName, Info.SettingName); \
+	ADD_TO_PROPERTIES_STRING(SettingName, Info.SettingName)
+
 #define GET_SECTION_SUPPORT_HELPER(SettingName)	\
-	Info.SettingName = GetSectionFeatureSupport(Section, #SettingName, Info.SettingName)
+	Info.SettingName = GetSectionFeatureSupport(Section, #SettingName, Info.SettingName); \
+	ADD_TO_PROPERTIES_STRING(SettingName, Info.SettingName)
 
 	GET_SECTION_BOOL_HELPER(bIsMobile);
 	GET_SECTION_BOOL_HELPER(bIsMetalMRT);
@@ -2538,7 +2551,9 @@ void FGenericDataDrivenShaderPlatformInfo::ParseDataDrivenShaderInfo(const FConf
 #undef GET_SECTION_BOOL_HELPER
 #undef GET_SECTION_INT_HELPER
 #undef GET_SECTION_SUPPORT_HELPER
+#undef ADD_TO_PROPERTIES_STRING
 
+	Info.ShaderPropertiesHash = GetTypeHash(ShaderPropertiesString);
 #if WITH_EDITOR
 	FTextStringHelper::ReadFromBuffer(*GetSectionString(Section, FName("FriendlyName")), Info.FriendlyName);
 #endif
@@ -2583,6 +2598,7 @@ void FGenericDataDrivenShaderPlatformInfo::Initialize()
 				}
 				
 				// at this point, we can start pulling information out
+				Infos[ShaderPlatform].Name = *SectionName.Mid(15);
 				ParseDataDrivenShaderInfo(Section.Value, Infos[ShaderPlatform]);	
 				Infos[ShaderPlatform].bContainsValidPlatformInfo = true;
 			}
diff --git a/Engine/Source/Runtime/RHI/Public/RHIDefinitions.h b/Engine/Source/Runtime/RHI/Public/RHIDefinitions.h
--- a/Engine/Source/Runtime/RHI/Public/RHIDefinitions.h
+++ b/Engine/Source/Runtime/RHI/Public/RHIDefinitions.h
@@ -332,8 +332,11 @@ extern RHI_API const FName LANGUAGE_Nintendo;
 
 class RHI_API FGenericDataDrivenShaderPlatformInfo
 {
+	FName Name;
 	FName Language;
 	ERHIFeatureLevel::Type MaxFeatureLevel;
+	FName ShaderFormat;
+	uint32 ShaderPropertiesHash;
 	uint32 bIsMobile: 1;
 	uint32 bIsMetalMRT: 1;
 	uint32 bIsPC: 1;
@@ -437,6 +440,24 @@ public:
 	static void Initialize();
 	static void ParseDataDrivenShaderInfo(const FConfigSection& Section, FGenericDataDrivenShaderPlatformInfo& Info);
 
+	static FORCEINLINE_DEBUGGABLE const FName GetName(const FStaticShaderPlatform Platform)
+	{
+		check(IsValid(Platform));
+		return Infos[Platform].Name;
+	}
+
+	static FORCEINLINE_DEBUGGABLE const FName GetShaderFormat(const FStaticShaderPlatform Platform)
+	{
+		check(IsValid(Platform));
+		return Infos[Platform].ShaderFormat;
+	}
+
+	static FORCEINLINE_DEBUGGABLE uint32 GetShaderPlatformPropertiesHash(const FStaticShaderPlatform Platform)
+	{
+		check(IsValid(Platform));
+		return Infos[Platform].ShaderPropertiesHash;
+	}
+
 	static FORCEINLINE_DEBUGGABLE const bool GetIsLanguageD3D(const FStaticShaderPlatform Platform)
 	{
 		check(IsValid(Platform));
diff --git a/Engine/Source/Runtime/RenderCore/Private/Shader.cpp b/Engine/Source/Runtime/RenderCore/Private/Shader.cpp
--- a/Engine/Source/Runtime/RenderCore/Private/Shader.cpp
+++ b/Engine/Source/Runtime/RenderCore/Private/Shader.cpp
@@ -2004,6 +2004,9 @@ void ShaderMapAppendKeyString(EShaderPlatform Platform, FString& KeyString)
 	}
 
 	KeyString += FString::Printf(TEXT("_LWC-%d"), FMath::FloorToInt(FLargeWorldRenderScalar::GetTileSize()));
+
+	uint64 ShaderPlatformPropertiesHash = FDataDrivenShaderPlatformInfo::GetShaderPlatformPropertiesHash(Platform);
+	KeyString += FString::Printf(TEXT("_%u"), ShaderPlatformPropertiesHash);
 }
 
 EShaderPermutationFlags GetShaderPermutationFlags(const FPlatformTypeLayoutParameters& LayoutParams)
