Set the minimum nvidia driver version to 516.01 on Windows, as it contains the fix for the transient heap GPU crash on Ampere cards.
Set the minimum nvidia driver version to 515.48 on Linux, as it fixes the swapchain creation failure bug in fullscreen.
Start warning about old drivers on Intel GPUs (minimum driver version 30.0.101.1191, released in December 2021).
Send GRHIAdapterDriverOnDenyList to analytics, so crashes caused by old drivers can be filtered out.
Check for old drivers in unattended mode too, and regardless of the value of r.WarnOfBadDrivers. Unattended mode and r.WarnOfBadDrivers only suppress the warning dialog box, but the bad driver flag is still set now, and a warning is printed to the log.

#rnx

diff --git a/Engine/Config/BaseHardware.ini b/Engine/Config/BaseHardware.ini
--- a/Engine/Config/BaseHardware.ini
+++ b/Engine/Config/BaseHardware.ini
@@ -36,18 +36,18 @@
 
 
 [GPU_NVIDIA Windows]
-SuggestedDriverVersion="511.79"
+SuggestedDriverVersion="516.40"
 
 ; To prevent problems with older drivers.
-+DriverDenyList=(DriverVersion="<497.29", Reason="These driver versions have known stability issues and missing features")
++DriverDenyList=(DriverVersion="<516.01", Reason="These driver versions have known stability issues and missing features")
 
 ;-----------------------------------------------------------------------------------------------------------------
 
 [GPU_NVIDIA Linux]
-SuggestedDriverVersion="510.47"
+SuggestedDriverVersion="515.48"
 
 ; UE5 in Linux/Vulkan requires recent drivers for Nanite/Lumen and has other stability issues with older drivers (need >= 495.46)
-+DriverDenyList=(DriverVersion="<495.46", Reason="These driver versions have known stability issues and missing features in Vulkan")
++DriverDenyList=(DriverVersion="<515.48", Reason="These driver versions have known stability issues and missing features in Vulkan")
 
 
 
@@ -77,6 +77,7 @@ SuggestedDriverVersion="21.2.6"
 
 [GPU_Intel Windows]
 SuggestedDriverVersion="30.0.101.1340"
++DriverDenyList=(DriverVersion="<30.0.101.1191", Reason="These driver versions have known stability issues and missing features")
 
 ;-----------------------------------------------------------------------------------------------------------------
 
diff --git a/Engine/Source/Runtime/Engine/Private/Analytics/EngineAnalyticsSessionSummary.cpp b/Engine/Source/Runtime/Engine/Private/Analytics/EngineAnalyticsSessionSummary.cpp
--- a/Engine/Source/Runtime/Engine/Private/Analytics/EngineAnalyticsSessionSummary.cpp
+++ b/Engine/Source/Runtime/Engine/Private/Analytics/EngineAnalyticsSessionSummary.cpp
@@ -169,6 +169,7 @@ FEngineAnalyticsSessionSummary::FEngineAnalyticsSessionSummary(TSharedPtr<IAnaly
 	Store->Set(TEXT("GRHIAdapterInternalDriverVersion"), GRHIAdapterInternalDriverVersion);
 	Store->Set(TEXT("GRHIAdapterUserDriverVersion"), GRHIAdapterUserDriverVersion);
 	Store->Set(TEXT("GRHIName"), FString(GDynamicRHI ? GDynamicRHI->GetName() : TEXT("")));
+	Store->Set(TEXT("GRHIAdapterDriverOnDenyList"), GRHIAdapterDriverOnDenyList);
 	Store->Set(TEXT("ProjectName"), EngineAnalyticsProperties::GetProjectName(ProjectSettings));
 	Store->Set(TEXT("ProjectID"), ProjectSettings.ProjectID.ToString(EGuidFormats::DigitsWithHyphens));
 	Store->Set(TEXT("ProjectDescription"), ProjectSettings.Description);
diff --git a/Engine/Source/Runtime/RHI/Private/DynamicRHI.cpp b/Engine/Source/Runtime/RHI/Private/DynamicRHI.cpp
--- a/Engine/Source/Runtime/RHI/Private/DynamicRHI.cpp
+++ b/Engine/Source/Runtime/RHI/Private/DynamicRHI.cpp
@@ -90,15 +90,14 @@ void InitNullRHI()
 #if PLATFORM_WINDOWS || PLATFORM_UNIX
 static void RHIDetectAndWarnOfBadDrivers(bool bHasEditorToken)
 {
-	int32 CVarValue = CVarWarnOfBadDrivers.GetValueOnGameThread();
-
-	if(!GIsRHIInitialized || !CVarValue || (GRHIVendorId == 0) || FApp::IsUnattended())
+	if (!GIsRHIInitialized || (GRHIVendorId == 0))
 	{
-		UE_LOG(LogRHI, Log, TEXT("Skipping Driver Check: RHI%s initialized, WarnOfBadDrivers=%d, VendorId=0x%x, is%s unattended"), 
-			GIsRHIInitialized ? TEXT("") : TEXT(" NOT"), CVarValue, GRHIVendorId, FApp::IsUnattended() ? TEXT("") : TEXT(" NOT"));
+		UE_LOG(LogRHI, Log, TEXT("Skipping Driver Check: RHI%s initialized, VendorId=0x%x"), GIsRHIInitialized ? TEXT("") : TEXT(" NOT"), GRHIVendorId);
 		return;
 	}
 
+	int32 WarnMode = CVarWarnOfBadDrivers.GetValueOnGameThread();
+
 	FGPUDriverInfo DriverInfo;
 
 	// later we should make the globals use the struct directly
@@ -112,7 +111,7 @@ static void RHIDetectAndWarnOfBadDrivers(bool bHasEditorToken)
 
 #if !(UE_BUILD_SHIPPING || UE_BUILD_TEST)
 	// for testing
-	if(CVarValue == 2)
+	if(WarnMode == 2)
 	{
 		DriverInfo.SetNVIDIA();
 		DriverInfo.DeviceDescription = TEXT("Test NVIDIA (bad)");
@@ -120,7 +119,7 @@ static void RHIDetectAndWarnOfBadDrivers(bool bHasEditorToken)
 		DriverInfo.InternalDriverVersion = TEXT("9.18.134.643");
 		DriverInfo.DriverDate = TEXT("01-01-1900");
 	}
-	else if(CVarValue == 3)
+	else if(WarnMode == 3)
 	{
 		DriverInfo.SetAMD();
 		DriverInfo.DeviceDescription = TEXT("Test AMD (bad)");
@@ -128,7 +127,7 @@ static void RHIDetectAndWarnOfBadDrivers(bool bHasEditorToken)
 		DriverInfo.InternalDriverVersion = TEXT("13.152.1.1000");
 		DriverInfo.DriverDate = TEXT("09-10-13");
 	}
-	else if(CVarValue == 4)
+	else if(WarnMode == 4)
 	{
 		DriverInfo.SetAMD();
 		DriverInfo.DeviceDescription = TEXT("Test AMD (good)");
@@ -136,7 +135,7 @@ static void RHIDetectAndWarnOfBadDrivers(bool bHasEditorToken)
 		DriverInfo.InternalDriverVersion = TEXT("15.30.1025.1001");
 		DriverInfo.DriverDate = TEXT("01-01-16");
 	}
-	else if(CVarValue == 5)
+	else if(WarnMode == 5)
 	{
 		DriverInfo.SetIntel();
 		DriverInfo.DeviceDescription = TEXT("Test Intel (good)");
@@ -185,10 +184,16 @@ static void RHIDetectAndWarnOfBadDrivers(bool bHasEditorToken)
 		GRHIAdapterDriverOnDenyList = DenyListEntry.IsValid();
 		FGenericCrashContext::SetEngineData(TEXT("RHI.DriverDenylisted"), DenyListEntry.IsValid() ? TEXT("true") : TEXT("false"));
 
+		if(!GRHIAdapterDriverOnDenyList)
+		{
+			return;
+		}
+
 		// Only alert users who are capable of updating their driver. Assume vendors with an empty list can always update.
 		bool bShowPrompt = bDeviceCanUpdateDriver || !bVendorHasEntries;
+		bShowPrompt = bShowPrompt && !FApp::IsUnattended() && WarnMode != 0;
 
-		if (DenyListEntry.IsValid() && bShowPrompt)
+		if (bShowPrompt)
 		{
 			bool bLatestDenied = DetectedGPUHardware.IsLatestDenied();
 
@@ -260,6 +265,10 @@ static void RHIDetectAndWarnOfBadDrivers(bool bHasEditorToken)
 			}
 #endif
 		}
+		else
+		{
+			UE_LOG(LogRHI, Warning, TEXT("Running with bad GPU drivers but warning dialog will not be shown: bDeviceCanUpdateDriver=%d, !VendorHasEntries=%d, IsUnattended=%d, r.WarnOfBadDrivers=%d"), bDeviceCanUpdateDriver, bVendorHasEntries, FApp::IsUnattended(), WarnMode);
+		}
 	}
 }
 #elif PLATFORM_MAC
