RHIGPUReadback: Updated FRHIGPUTextureReadback Lock, Unlock and EnqueueCopy to handle the case of multi-index GPU mask, with the same pattern that the Buffer code path has.

diff --git a/Engine/Source/Runtime/RHI/Private/RHIGPUReadback.cpp b/Engine/Source/Runtime/RHI/Private/RHIGPUReadback.cpp
--- a/Engine/Source/Runtime/RHI/Private/RHIGPUReadback.cpp
+++ b/Engine/Source/Runtime/RHI/Private/RHIGPUReadback.cpp
@@ -125,54 +125,61 @@ void FRHIGPUTextureReadback::EnqueueCopyRDG(FRHICommandList& RHICmdList, FRHITex
 void FRHIGPUTextureReadback::EnqueueCopy(FRHICommandList& RHICmdList, FRHITexture* SourceTexture, FResolveRect Rect)
 {
 	Fence->Clear();
+	LastCopyGPUMask = RHICmdList.GetGPUMask();
 
 	if (SourceTexture)
 	{
 		check(SourceTexture->GetTexture2D());
 		check(!SourceTexture->IsMultisampled());
 
-		// Assume for now that every enqueue happens on a texture of the same format and size (when reused).
-		if (!DestinationStagingTexture)
+		for (uint32 GPUIndex : LastCopyGPUMask)
 		{
-			FIntVector TextureSize = SourceTexture->GetSizeXYZ();
-			FString FenceName = Fence->GetFName().ToString();
+			SCOPED_GPU_MASK(RHICmdList, FRHIGPUMask::FromIndex(GPUIndex));
 
-			const FRHITextureCreateDesc Desc =
-				FRHITextureCreateDesc::Create2D(*FenceName, TextureSize.X, TextureSize.Y, SourceTexture->GetFormat())
-				.SetFlags(ETextureCreateFlags::CPUReadback | ETextureCreateFlags::HideInVisualizeTexture)
-				.SetInitialState(ERHIAccess::CopyDest);
+			// Assume for now that every enqueue happens on a texture of the same format and size (when reused).
+			if (!DestinationStagingTextures[GPUIndex])
+			{
+				FIntVector TextureSize = SourceTexture->GetSizeXYZ();
+				FString FenceName = Fence->GetFName().ToString();
 
-			DestinationStagingTexture = RHICreateTexture(Desc);
-		}
+				const FRHITextureCreateDesc Desc =
+					FRHITextureCreateDesc::Create2D(*FenceName, TextureSize.X, TextureSize.Y, SourceTexture->GetFormat())
+					.SetFlags(ETextureCreateFlags::CPUReadback | ETextureCreateFlags::HideInVisualizeTexture)
+					.SetInitialState(ERHIAccess::CopyDest);
 
-		FRHICopyTextureInfo CopyInfo;
+				DestinationStagingTextures[GPUIndex] = RHICreateTexture(Desc);
+			}
 
-		if (Rect.IsValid())
-		{
-			CopyInfo.SourcePosition = FIntVector(Rect.X1, Rect.Y1, 1);
-			CopyInfo.DestPosition = CopyInfo.SourcePosition;
-			CopyInfo.Size = FIntVector(Rect.X2 - Rect.X1, Rect.Y2 - Rect.Y1, 1);
-		}
+			FRHICopyTextureInfo CopyInfo;
 
-		RHICmdList.CopyTexture(SourceTexture, DestinationStagingTexture, CopyInfo);
+			if (Rect.IsValid())
+			{
+				CopyInfo.SourcePosition = FIntVector(Rect.X1, Rect.Y1, 1);
+				CopyInfo.DestPosition = CopyInfo.SourcePosition;
+				CopyInfo.Size = FIntVector(Rect.X2 - Rect.X1, Rect.Y2 - Rect.Y1, 1);
+			}
 
-		RHICmdList.WriteGPUFence(Fence);
+			RHICmdList.CopyTexture(SourceTexture, DestinationStagingTextures[GPUIndex], CopyInfo);
 
-		LastCopyGPUMask = RHICmdList.GetGPUMask();
+			RHICmdList.WriteGPUFence(Fence);
+		}
 	}
 }
 
 void* FRHIGPUTextureReadback::Lock(uint32 NumBytes)
 {
-	if (DestinationStagingTexture)
+	uint32 GPUIndex = LastCopyGPUMask.GetFirstIndex();
+
+	if (DestinationStagingTextures[GPUIndex])
 	{
-		LastLockGPUIndex = LastCopyGPUMask.ToIndex();
+		LastLockGPUIndex = GPUIndex;
 
 		void* ResultsBuffer = nullptr;
 		int32 BufferWidth = 0, BufferHeight = 0;
 		FRHICommandListImmediate& RHICmdList = FRHICommandListExecutor::GetImmediateCommandList();
-		RHICmdList.Transition(FRHITransitionInfo(DestinationStagingTexture, ERHIAccess::CopyDest, ERHIAccess::CPURead));
-		RHICmdList.MapStagingSurface(DestinationStagingTexture, Fence.GetReference(), ResultsBuffer, BufferWidth, BufferHeight, LastLockGPUIndex);
+		RHICmdList.Transition(FRHITransitionInfo(DestinationStagingTextures[GPUIndex], ERHIAccess::CopyDest, ERHIAccess::CPURead));
+		RHICmdList.MapStagingSurface(DestinationStagingTextures[GPUIndex], Fence.GetReference(), ResultsBuffer, BufferWidth, BufferHeight, LastLockGPUIndex);
+
 		return ResultsBuffer;
 	}
 	else
@@ -183,20 +190,25 @@ void* FRHIGPUTextureReadback::Lock(uint32 NumBytes)
 
 void* FRHIGPUTextureReadback::Lock(int32& OutRowPitchInPixels, int32 *OutBufferHeight)
 {
-	if (DestinationStagingTexture)
+	uint32 GPUIndex = LastCopyGPUMask.GetFirstIndex();
+
+	if (DestinationStagingTextures[GPUIndex])
 	{
-		LastLockGPUIndex = LastCopyGPUMask.ToIndex();
+		LastLockGPUIndex = GPUIndex;
 
 		void* ResultsBuffer = nullptr;
 		int32 BufferWidth = 0, BufferHeight = 0;
 		FRHICommandListImmediate& RHICmdList = FRHICommandListExecutor::GetImmediateCommandList();
-		RHICmdList.Transition(FRHITransitionInfo(DestinationStagingTexture, ERHIAccess::CopyDest, ERHIAccess::CPURead));
-		RHICmdList.MapStagingSurface(DestinationStagingTexture, Fence.GetReference(), ResultsBuffer, BufferWidth, BufferHeight, LastLockGPUIndex);
+		RHICmdList.Transition(FRHITransitionInfo(DestinationStagingTextures[GPUIndex], ERHIAccess::CopyDest, ERHIAccess::CPURead));
+		RHICmdList.MapStagingSurface(DestinationStagingTextures[GPUIndex], Fence.GetReference(), ResultsBuffer, BufferWidth, BufferHeight, LastLockGPUIndex);
+
 		if (OutBufferHeight)
 		{
 			*OutBufferHeight = BufferHeight;
 		}
+
 		OutRowPitchInPixels = BufferWidth;
+
 		return ResultsBuffer;
 	}
 
@@ -215,9 +227,9 @@ void FRHIGPUTextureReadback::LockTexture(FRHICommandListImmediate& RHICmdList, v
 
 void FRHIGPUTextureReadback::Unlock()
 {
-	ensure(DestinationStagingTexture);
+	ensure(DestinationStagingTextures[LastLockGPUIndex]);
 
 	FRHICommandListImmediate& RHICmdList = FRHICommandListExecutor::GetImmediateCommandList();
-	RHICmdList.UnmapStagingSurface(DestinationStagingTexture, LastLockGPUIndex);
-	RHICmdList.Transition(FRHITransitionInfo(DestinationStagingTexture, ERHIAccess::CPURead, ERHIAccess::CopyDest));
+	RHICmdList.UnmapStagingSurface(DestinationStagingTextures[LastLockGPUIndex], LastLockGPUIndex);
+	RHICmdList.Transition(FRHITransitionInfo(DestinationStagingTextures[LastLockGPUIndex], ERHIAccess::CPURead, ERHIAccess::CopyDest));
 }
diff --git a/Engine/Source/Runtime/RHI/Public/RHIGPUReadback.h b/Engine/Source/Runtime/RHI/Public/RHIGPUReadback.h
--- a/Engine/Source/Runtime/RHI/Public/RHIGPUReadback.h
+++ b/Engine/Source/Runtime/RHI/Public/RHIGPUReadback.h
@@ -124,5 +124,9 @@ public:
 	UE_DEPRECATED(5.0, "Use FRHIGPUTextureReadback::Lock( int32& OutRowPitchInPixels) instead.")
 	void LockTexture(FRHICommandListImmediate& RHICmdList, void*& OutBufferPtr, int32& OutRowPitchInPixels);
 
-	FTextureRHIRef DestinationStagingTexture;
+#if WITH_MGPU
+	FTextureRHIRef DestinationStagingTextures[MAX_NUM_GPUS];
+#else
+	FTextureRHIRef DestinationStagingTextures[1];
+#endif
 };
