Add more options to r.AsyncPipelineCompile to allow async creation only for PSOs from precompile jobs or only non-precompile jobs.

diff --git a/Engine/Source/Runtime/RHI/Private/PipelineStateCache.cpp b/Engine/Source/Runtime/RHI/Private/PipelineStateCache.cpp
--- a/Engine/Source/Runtime/RHI/Private/PipelineStateCache.cpp
+++ b/Engine/Source/Runtime/RHI/Private/PipelineStateCache.cpp
@@ -56,11 +56,22 @@ static TAtomic<uint32> ComputePipelineCacheMisses;
 static TArray<uint32> ComputePipelineCacheMissesHistory;
 static bool	ReportFrameHitchThisFrame;
 
+enum class EPSOCompileAsyncMode
+{
+	None = 0,
+	All = 1,
+	Precompile = 2,
+	NonPrecompiled = 3,
+};
+
 static TAutoConsoleVariable<int32> GCVarAsyncPipelineCompile(
 	TEXT("r.AsyncPipelineCompile"),
-	1,
-	TEXT("0 to Create PSOs at the moment they are requested\n")\
-	TEXT("1 to Create Pipeline State Objects asynchronously(default)"),
+	(int32)EPSOCompileAsyncMode::All,
+	TEXT("0 to Create PSOs at the moment they are requested\n")
+	TEXT("1 to Create Pipeline State Objects asynchronously(default)\n")
+	TEXT("2 to Create Only precompile PSOs asynchronously\n")
+	TEXT("3 to Create Only non-precompile PSOs asynchronously")
+	,
 	ECVF_ReadOnly | ECVF_RenderThreadSafe
 );
 
@@ -1357,11 +1368,17 @@ void PipelineStateCache::FlushResources()
 
 }
 
-static bool IsAsyncCompilationAllowed(FRHIComputeCommandList& RHICmdList)
+static bool IsAsyncCompilationAllowed(FRHIComputeCommandList& RHICmdList, bool bIsPrecompileRequest)
 {
+	const EPSOCompileAsyncMode PSOCompileAsyncMode = (EPSOCompileAsyncMode)GCVarAsyncPipelineCompile.GetValueOnAnyThread();
+
+	const bool bCVarAllowsAsyncCreate = PSOCompileAsyncMode == EPSOCompileAsyncMode::All
+		|| (PSOCompileAsyncMode == EPSOCompileAsyncMode::Precompile && bIsPrecompileRequest)
+		|| (PSOCompileAsyncMode == EPSOCompileAsyncMode::NonPrecompiled && !bIsPrecompileRequest);
+
 	return GRHISupportsAsyncPipelinePrecompile &&
 		FDataDrivenShaderPlatformInfo::GetSupportsAsyncPipelineCompilation(GMaxRHIShaderPlatform) &&
-		GCVarAsyncPipelineCompile.GetValueOnAnyThread() && !RHICmdList.Bypass() && (IsRunningRHIInSeparateThread() && !IsInRHIThread()) && RHICmdList.AsyncPSOCompileAllowed();
+		bCVarAllowsAsyncCreate && !RHICmdList.Bypass() && (IsRunningRHIInSeparateThread() && !IsInRHIThread()) && RHICmdList.AsyncPSOCompileAllowed();
 }
 
 uint64 PipelineStateCache::RetrieveGraphicsPipelineStateSortKey(const FGraphicsPipelineState* GraphicsPipelineState)
@@ -1371,7 +1388,7 @@ uint64 PipelineStateCache::RetrieveGraphicsPipelineStateSortKey(const FGraphicsP
 
 FComputePipelineState* PipelineStateCache::GetAndOrCreateComputePipelineState(FRHIComputeCommandList& RHICmdList, FRHIComputeShader* ComputeShader, bool bFromFileCache)
 {	
-	bool DoAsyncCompile = IsAsyncCompilationAllowed(RHICmdList);
+	bool DoAsyncCompile = IsAsyncCompilationAllowed(RHICmdList, bFromFileCache);
 
 	FComputePipelineState* OutCachedState = nullptr;
 
@@ -1423,7 +1440,7 @@ FComputePipelineState* PipelineStateCache::GetAndOrCreateComputePipelineState(FR
 	GComputePipelineCache.Unlock(LockFlags);
 
 #if 0
-	bool DoAsyncCompile = IsAsyncCompilationAllowed(RHICmdList);
+	bool DoAsyncCompile = IsAsyncCompilationAllowed(RHICmdList, bFromFileCache);
 
 
 	TSharedPtr<FComputePipelineState> OutCachedState;
@@ -1579,7 +1596,7 @@ FRayTracingPipelineState* PipelineStateCache::GetAndOrCreateRayTracingPipelineSt
 
 	check(IsInRenderingThread() || IsInParallelRenderingThread());
 
-	const bool bDoAsyncCompile = IsAsyncCompilationAllowed(RHICmdList);
+	const bool bDoAsyncCompile = IsAsyncCompilationAllowed(RHICmdList, false);
 	const bool bNonBlocking = !!(Flags & ERayTracingPipelineCacheFlags::NonBlocking);
 
 	FRayTracingPipelineState* Result = nullptr;
@@ -1774,7 +1791,7 @@ FGraphicsPipelineState* PipelineStateCache::GetAndOrCreateGraphicsPipelineState(
 	}
 #endif
 
-	bool DoAsyncCompile = IsAsyncCompilationAllowed(RHICmdList);
+	bool DoAsyncCompile = IsAsyncCompilationAllowed(RHICmdList, Initializer.bFromPSOFileCache);
 
 	FGraphicsPipelineState* OutCachedState = nullptr;
 
