Fixed banned word warnings.

#rnx

diff --git a/Engine/Source/Runtime/RHI/Private/RHICommandList.cpp b/Engine/Source/Runtime/RHI/Private/RHICommandList.cpp
--- a/Engine/Source/Runtime/RHI/Private/RHICommandList.cpp
+++ b/Engine/Source/Runtime/RHI/Private/RHICommandList.cpp
@@ -70,7 +70,7 @@ static TAutoConsoleVariable<int32> CVarRHICmdUseDeferredContexts(
 TAutoConsoleVariable<int32> CVarRHICmdFlushRenderThreadTasks(
 	TEXT("r.RHICmdFlushRenderThreadTasks"),
 	0,
-	TEXT("If true, then we flush the render thread tasks every pass. For issue diagnosis. This is a master switch for more granular cvars."));
+	TEXT("If true, then we flush the render thread tasks every pass. For issue diagnosis. This is a main switch for more granular cvars."));
 
 static TAutoConsoleVariable<int32> CVarRHICmdFlushOnQueueParallelSubmit(
 	TEXT("r.RHICmdFlushOnQueueParallelSubmit"),
diff --git a/Engine/Source/Runtime/RHI/Public/GPUDefragAllocator.h b/Engine/Source/Runtime/RHI/Public/GPUDefragAllocator.h
--- a/Engine/Source/Runtime/RHI/Public/GPUDefragAllocator.h
+++ b/Engine/Source/Runtime/RHI/Public/GPUDefragAllocator.h
@@ -1124,7 +1124,7 @@ protected:
 	friend FScopedGPUDefragLock;
 };
 
-//FScopedGPUDefragLock can't cover any scope that will add dcb commands or we might deadlock with a master reserve failure.
+//FScopedGPUDefragLock can't cover any scope that will add dcb commands or we might deadlock.
 class FScopedGPUDefragLock
 {
 public:
