Fix deprecation warning pragmas broken in bad merge

diff --git a/Engine/Source/Runtime/RHI/Public/RHICommandList.h b/Engine/Source/Runtime/RHI/Public/RHICommandList.h
--- a/Engine/Source/Runtime/RHI/Public/RHICommandList.h
+++ b/Engine/Source/Runtime/RHI/Public/RHICommandList.h
@@ -4288,6 +4288,7 @@ public:
 		ERHIAccess ResourceState = RHIGetDefaultResourceState(Usage, true);
 	PRAGMA_DISABLE_DEPRECATION_WARNINGS
 		return CreateAndLockVertexBuffer(Size, Usage, ResourceState, CreateInfo, OutDataBuffer);
+	PRAGMA_ENABLE_DEPRECATION_WARNINGS
 	}
 
 	UE_DEPRECATED(5.0, "Buffer locks have been unified. Use LockBuffer() instead.")
