Adjust PartCount after adding new entry to rasterizer state cache.

#rnx

diff --git a/Engine/Source/Runtime/RHI/Private/PipelineFileCache.cpp b/Engine/Source/Runtime/RHI/Private/PipelineFileCache.cpp
--- a/Engine/Source/Runtime/RHI/Private/PipelineFileCache.cpp
+++ b/Engine/Source/Runtime/RHI/Private/PipelineFileCache.cpp
@@ -367,7 +367,7 @@ struct FPipelineCacheFileFormatPSOMetaData
 
 FString FPipelineFileCacheRasterizerState::ToString() const
 {
-	return FString::Printf(TEXT("<%f %f %u %u %u %u>")
+	return FString::Printf(TEXT("<%f %f %u %u %u %u %u>")
 		, DepthBias
 		, SlopeScaleDepthBias
 		, uint32(FillMode)
@@ -380,13 +380,13 @@ FString FPipelineFileCacheRasterizerState::ToString() const
 
 void FPipelineFileCacheRasterizerState::FromString(const FStringView& Src)
 {
-	constexpr int32 PartCount = 6;
+	constexpr int32 PartCount = 7;
 
 	TArray<FStringView, TInlineAllocator<PartCount>> Parts;
 	UE::String::ParseTokensMultiple(Src.TrimStartAndEnd(), {TEXT('\r'), TEXT('\n'), TEXT('\t'), TEXT('<'), TEXT('>'), TEXT(' ')},
 		[&Parts](FStringView Part) { if (!Part.IsEmpty()) { Parts.Add(Part); } });
 
-	check(Parts.Num() == PartCount && sizeof(FillMode) == 1 && sizeof(CullMode) == 1 && sizeof(bAllowMSAA) == 1 && sizeof(bEnableLineAA) == 1); //not a very robust parser
+	check(Parts.Num() == PartCount && sizeof(FillMode) == 1 && sizeof(CullMode) == 1 && sizeof(DepthClipMode) == 1 && sizeof(bAllowMSAA) == 1 && sizeof(bEnableLineAA) == 1); //not a very robust parser
 	const FStringView* PartIt = Parts.GetData();
 
 	LexFromString(DepthBias, *PartIt++);
