Multiple PSO cache processing.
Splitt PSO cache code to separate out bundled game caches and a single PSO recording user cache.
Any additional PSO caches are enqueued when the shader library state change adds shader library components.
removed unused PSO SaveMode::SortedBoundPSOs.

diff --git a/Engine/Source/Editor/UnrealEd/Private/Commandlets/MergeShaderPipelineCachesCommandlet.cpp b/Engine/Source/Editor/UnrealEd/Private/Commandlets/MergeShaderPipelineCachesCommandlet.cpp
--- a/Engine/Source/Editor/UnrealEd/Private/Commandlets/MergeShaderPipelineCachesCommandlet.cpp
+++ b/Engine/Source/Editor/UnrealEd/Private/Commandlets/MergeShaderPipelineCachesCommandlet.cpp
@@ -21,20 +21,20 @@ int32 UMergeShaderPipelineCachesCommandlet::Main(const FString& Params)
 
 	bool bOK = (Tokens.Num() >= 3);
 	
-	FPipelineFileCache::PSOOrder Order = FPipelineFileCache::PSOOrder::Default;
+	FPipelineFileCacheManager::PSOOrder Order = FPipelineFileCacheManager::PSOOrder::Default;
 	if (const FString* OrderVal = ParamVals.Find(FString(TEXT("Sort"))))
 	{
 		if(*OrderVal == TEXT("Default"))
 		{
-			Order = FPipelineFileCache::PSOOrder::Default;
+			Order = FPipelineFileCacheManager::PSOOrder::Default;
 		}
 		else if(*OrderVal == TEXT("FirstUsed"))
 		{
-			Order = FPipelineFileCache::PSOOrder::FirstToLatestUsed;
+			Order = FPipelineFileCacheManager::PSOOrder::FirstToLatestUsed;
 		}
 		else if(*OrderVal == TEXT("MostUsed"))
 		{
-			Order = FPipelineFileCache::PSOOrder::MostToLeastUsed;
+			Order = FPipelineFileCacheManager::PSOOrder::MostToLeastUsed;
 		}
 		else
 		{
@@ -61,5 +61,5 @@ int32 UMergeShaderPipelineCachesCommandlet::Main(const FString& Params)
 		return 0;
 	}
 	
-	return !FPipelineFileCache::MergePipelineFileCaches(A, B, Order, Output);
+	return !FPipelineFileCacheManager::MergePipelineFileCaches(A, B, Order, Output);
 }
diff --git a/Engine/Source/Editor/UnrealEd/Private/Commandlets/ShaderPipelineCacheToolsCommandlet.cpp b/Engine/Source/Editor/UnrealEd/Private/Commandlets/ShaderPipelineCacheToolsCommandlet.cpp
--- a/Engine/Source/Editor/UnrealEd/Private/Commandlets/ShaderPipelineCacheToolsCommandlet.cpp
+++ b/Engine/Source/Editor/UnrealEd/Private/Commandlets/ShaderPipelineCacheToolsCommandlet.cpp
@@ -412,7 +412,7 @@ int32 DumpPSOSC(FString& Token, const FString& StableKeyFileDir)
 	TSet<FPipelineCacheFileFormatPSO> PSOs;
 
 	UE_LOG(LogShaderPipelineCacheTools, Display, TEXT("Loading %s...."), *Token);
-	if (!FPipelineFileCache::LoadPipelineFileCacheInto(Token, PSOs))
+	if (!FPipelineFileCacheManager::LoadPipelineFileCacheInto(Token, PSOs))
 	{
 		UE_LOG(LogShaderPipelineCacheTools, Error, TEXT("Could not load %s or it was empty."), *Token);
 		return 1;
@@ -922,7 +922,7 @@ int32 ExpandPSOSC(const TArray<FString>& Tokens)
 		{
 			UE_LOG(LogShaderPipelineCacheTools, Display, TEXT("Loading %s...."), *Tokens[Index]);
 			TSet<FPipelineCacheFileFormatPSO> TempPSOs;
-			if (!FPipelineFileCache::LoadPipelineFileCacheInto(Tokens[Index], TempPSOs))
+			if (!FPipelineFileCacheManager::LoadPipelineFileCacheInto(Tokens[Index], TempPSOs))
 			{
 				UE_LOG(LogShaderPipelineCacheTools, Error, TEXT("Could not load %s or it was empty."), *Tokens[Index]);
 				continue;
@@ -1965,7 +1965,7 @@ int32 SaveBinaryPipelineCacheFile(const FString& OutputFilename, const EShaderPl
 		{
 			UE_LOG(LogShaderPipelineCacheTools, Fatal, TEXT("Could not delete %s"), *OutputFilename);
 		}
-		if (!FPipelineFileCache::SavePipelineFileCacheFrom(FShaderPipelineCache::GetGameVersionForPSOFileCache(), ShaderPlatform, OutputFilename, PSOs))
+		if (!FPipelineFileCacheManager::SavePipelineFileCacheFrom(FShaderPipelineCache::GetGameVersionForPSOFileCache(), ShaderPlatform, OutputFilename, PSOs))
 		{
 			UE_LOG(LogShaderPipelineCacheTools, Error, TEXT("Failed to save %s"), *OutputFilename);
 			return 1;
@@ -2024,9 +2024,12 @@ int32 SaveBinaryPipelineCacheFile(const FString& OutputFilename, const EShaderPl
 			ChunkIds.Add(kInvalidChunkId);
 		}
 
+		// This should be passed in.
+		FString ChunkDir = FPaths::GetPath(OutputFilename);
+
 		// prepare everything necessary to split the PSOs
 		ParallelFor(ChunkInfoFilenames.Num(),
-			[&ChunkInfoFilenames, &ChunkIds, &ChunkIdsAndResultAccessLock, &OverallResult, &StableNameMap, &StableMapConvTask, &ShaderPlatform, &ShaderFormat, &PSOs, &SaveSingleCacheFile](int32 Index)
+			[&ChunkDir, &ChunkInfoFilenames, &ChunkIds, &ChunkIdsAndResultAccessLock, &OverallResult, &StableNameMap, &StableMapConvTask, &ShaderPlatform, &ShaderFormat, &PSOs, &SaveSingleCacheFile](int32 Index)
 			{
 				int32 ChunkId;
 				FString OutputFilename;
@@ -2110,7 +2113,8 @@ int32 SaveBinaryPipelineCacheFile(const FString& OutputFilename, const EShaderPl
 
 				if (PSOsInChunk.Num())
 				{
-					int32 Result = SaveSingleCacheFile(OutputFilename, ShaderPlatform, PSOsInChunk);
+					FString FinalPath = FPaths::Combine(ChunkDir, OutputFilename);
+					int32 Result = SaveSingleCacheFile(FinalPath, ShaderPlatform, PSOsInChunk);
 					if (Result != 0)
 					{
 						FScopeLock Locker(&ChunkIdsAndResultAccessLock);
diff --git a/Engine/Source/Runtime/D3D12RHI/Private/Windows/WindowsD3D12PipelineState.cpp b/Engine/Source/Runtime/D3D12RHI/Private/Windows/WindowsD3D12PipelineState.cpp
--- a/Engine/Source/Runtime/D3D12RHI/Private/Windows/WindowsD3D12PipelineState.cpp
+++ b/Engine/Source/Runtime/D3D12RHI/Private/Windows/WindowsD3D12PipelineState.cpp
@@ -28,7 +28,7 @@
 #define D3D12RHI_USE_D3DDISASSEMBLE 1
 #endif
 
-// D3D12RHI PSO file cache doesn't work anymore. Use FPipelineFileCache instead
+// D3D12RHI PSO file cache doesn't work anymore. Use FPipelineFileCacheManager instead
 static TAutoConsoleVariable<int32> CVarPipelineStateDiskCache(
 	TEXT("D3D12.PSO.DiskCache"),
 	0,
diff --git a/Engine/Source/Runtime/RHI/Private/PipelineFileCache.cpp b/Engine/Source/Runtime/RHI/Private/PipelineFileCache.cpp
--- a/Engine/Source/Runtime/RHI/Private/PipelineFileCache.cpp
+++ b/Engine/Source/Runtime/RHI/Private/PipelineFileCache.cpp
@@ -22,6 +22,8 @@ PipelineFileCache.cpp: Pipeline state cache implementation.
 #include "String/LexFromString.h"
 #include "String/ParseTokens.h"
 #include "Misc/ScopeExit.h"
+#include <Misc/ConfigCacheIni.h>
+#include <Algo/ForEach.h>
 
 static FString JOURNAL_FILE_EXTENSION(TEXT(".jnl"));
 
@@ -160,18 +162,20 @@ static TAutoConsoleVariable<int32> CVarAlwaysGeneratePOSSOFileCache(
 														   );
 
 
-FRWLock FPipelineFileCache::FileCacheLock;
-FPipelineCacheFile* FPipelineFileCache::FileCache = nullptr;
-TMap<uint32, FPSOUsageData> FPipelineFileCache::RunTimeToPSOUsage;
-TMap<uint32, FPSOUsageData> FPipelineFileCache::NewPSOUsage;
-TMap<uint32, FPipelineStateStats*> FPipelineFileCache::Stats;
-TSet<FPipelineCacheFileFormatPSO> FPipelineFileCache::NewPSOs;
-TSet<uint32> FPipelineFileCache::NewPSOHashes;
-uint32 FPipelineFileCache::NumNewPSOs;
-FPipelineFileCache::PSOOrder FPipelineFileCache::RequestedOrder = FPipelineFileCache::PSOOrder::MostToLeastUsed;
-bool FPipelineFileCache::FileCacheEnabled = false;
-FPipelineFileCache::FPipelineStateLoggedEvent FPipelineFileCache::PSOLoggedEvent;
-uint64 FPipelineFileCache::GameUsageMask = 0;
+FRWLock FPipelineFileCacheManager::FileCacheLock;
+TMap<FString, TUniquePtr< class FPipelineCacheFile>> FPipelineFileCacheManager::FileCacheMap;
+TMap<FGuid, FString> FPipelineFileCacheManager::GameGuidToCacheKey;
+TMap<uint32, FPSOUsageData> FPipelineFileCacheManager::RunTimeToPSOUsage;
+TMap<uint32, FPSOUsageData> FPipelineFileCacheManager::NewPSOUsage;
+TMap<uint32, FPipelineStateStats*> FPipelineFileCacheManager::Stats;
+TSet<FPipelineCacheFileFormatPSO> FPipelineFileCacheManager::NewPSOs;
+TSet<uint32> FPipelineFileCacheManager::NewPSOHashes;
+uint32 FPipelineFileCacheManager::NumNewPSOs;
+FString FPipelineFileCacheManager::UserCacheKey;
+FPipelineFileCacheManager::PSOOrder FPipelineFileCacheManager::RequestedOrder = FPipelineFileCacheManager::PSOOrder::MostToLeastUsed;
+bool FPipelineFileCacheManager::FileCacheEnabled = false;
+FPipelineFileCacheManager::FPipelineStateLoggedEvent FPipelineFileCacheManager::PSOLoggedEvent;
+uint64 FPipelineFileCacheManager::GameUsageMask = 0;
 
 static int64 GetCurrentUnixTime()
 {
@@ -182,7 +186,7 @@ bool DefaultPSOMaskComparisonFunction(uint64 ReferenceMask, uint64 PSOMask)
 {
 	return (ReferenceMask & PSOMask) == ReferenceMask;
 }
-FPSOMaskComparisonFn FPipelineFileCache::MaskComparisonFn = DefaultPSOMaskComparisonFunction;
+FPSOMaskComparisonFn FPipelineFileCacheManager::MaskComparisonFn = DefaultPSOMaskComparisonFunction;
 
 static inline bool IsReferenceMaskSet(uint64 ReferenceMask, uint64 PSOMask)
 {
@@ -1626,12 +1630,21 @@ FPipelineCacheFileFormatPSO& FPipelineCacheFileFormatPSO::operator=(const FPipel
 struct FPipelineCacheFileFormatTOC
 {
 	FPipelineCacheFileFormatTOC()
-	: SortedOrder(FPipelineFileCache::PSOOrder::MostToLeastUsed)
+	: SortedOrder(FPipelineFileCacheManager::PSOOrder::MostToLeastUsed)
 	{}
 	
-	FPipelineFileCache::PSOOrder SortedOrder;
+	FPipelineFileCacheManager::PSOOrder SortedOrder;
 	TMap<uint32, FPipelineCacheFileFormatPSOMetaData> MetaData;
-	
+	void DumpToLog()
+	{
+		for (TMap<uint32, FPipelineCacheFileFormatPSOMetaData>::TConstIterator It(MetaData); It; ++It)
+		{
+			const FPipelineCacheFileFormatPSOMetaData& dat = It.Value();
+			UE_LOG(LogRHI, VeryVerbose, TEXT("PSO hash %u - guid (%s), stats(FF %d, LF %d, bind %d)"), It.Key(), *dat.FileGuid.ToString(), dat.Stats.FirstFrameUsed, dat.Stats.LastFrameUsed, dat.Stats.TotalBindCount)
+		}
+		UE_LOG(LogRHI, VeryVerbose, TEXT("Total PSOs %d"), MetaData.Num());
+	}
+
 	friend FArchive& operator<<(FArchive& Ar, FPipelineCacheFileFormatTOC& Info)
 	{
 		// TOC is assumed to be at the end of the file
@@ -1733,20 +1746,55 @@ struct FPipelineCacheFileFormatTOC
 	}
 };
 
+static bool ShouldDeleteExistingUserCache()
+{
+	static bool bOnce = false;
+	static bool bCmdLineForce = false;
+	if (!bOnce)
+	{
+		bOnce = true;
+		bCmdLineForce = FParse::Param(FCommandLine::Get(), TEXT("deleteuserpsocache")) || FParse::Param(FCommandLine::Get(), TEXT("logPSO"));
+		UE_CLOG(bCmdLineForce, LogRHI, Warning, TEXT("****************************** Deleting user-writable PSO cache as requested on command line"));
+	}
+	return bCmdLineForce;
+}
+
 class FPipelineCacheFile
 {
+	FString Name;
+	EShaderPlatform ShaderPlatform;
+	FName PlatformName;
+	uint64 TOCOffset;
+
+	FPipelineCacheFileFormatTOC TOC;
+	FGuid FileGuid;
+	FString FilePath;
+	TSharedPtr<IAsyncReadFileHandle, ESPMode::ThreadSafe> AsyncFileHandle;
+
+	FString RecordingFilename;
 public:
+	enum class EStatus : uint8
+	{
+		Unknown,
+		BundledCache,
+		UserCacheOpened,	// a user cache was successfully opened
+		NewUserCache,		// user cache failed to open, started empty.
+	};
+	EStatus CacheStatus = EStatus::Unknown;
+
 	static uint32 GameVersion;
 
 	FPipelineCacheFile()
 	: TOCOffset(0)
-	, UserFileGuid(FGuid::NewGuid())
-	, UserAsyncFileHandle(nullptr)
-	, GameAsyncFileHandle(nullptr)
+	, AsyncFileHandle(nullptr)
 	{
 	}
-	
-	bool OpenPipelineFileCache(FString const& FilePath, FGuid& Guid, TSharedPtr<IAsyncReadFileHandle, ESPMode::ThreadSafe>& Handle, FPipelineCacheFileFormatTOC& Content)
+	~FPipelineCacheFile()
+	{
+		DEC_MEMORY_STAT_BY(STAT_FileCacheMemory, TOC.MetaData.GetAllocatedSize());
+	}
+
+	static bool OpenPipelineFileCache(const FString& FilePath, EShaderPlatform ShaderPlatform, FGuid& Guid, TSharedPtr<IAsyncReadFileHandle, ESPMode::ThreadSafe>& Handle, FPipelineCacheFileFormatTOC& Content, uint64& TOCOffsetOUT)
 	{
 		bool bSuccess = false;
 
@@ -1810,7 +1858,7 @@ public:
 					UE_LOG(LogRHI, Log, TEXT("Opened FPipelineCacheFile: %s (GUID: %s) with %d entries."), *FilePath, *Header.Guid.ToString(), Content.MetaData.Num());
 					
 					Guid = Header.Guid;
-					TOCOffset = Header.TableOffset;
+					TOCOffsetOUT = Header.TableOffset;
 				}
 				else
 				{
@@ -1827,7 +1875,7 @@ public:
 		return bSuccess;
 	}
 
-	void GarbageCollectUserCache(FString const& UserCacheFilePath)
+	void GarbageCollectUserCache(FString const& UserCacheFilePath, const TSet<FGuid>& KnownGuids)
 	{
 		UE_LOG(LogRHI, Log, TEXT("FPipelineCacheFile: GarbageCollectUserCache() Begin"));
 		ON_SCOPE_EXIT{ UE_LOG(LogRHI, Log, TEXT("FPipelineCacheFile: GarbageCollectUserCache() End")); };
@@ -1894,7 +1942,7 @@ public:
 		FTimespan StaleTimespan = FTimespan::FromDays(StaleDays);
 		int64 EvictionTime = UnixTime - (int64)StaleTimespan.GetTotalSeconds();
 
-		auto EntryShouldBeRemovedFromUserCache = [&Header, GameFileGuid=this->GameFileGuid, EvictionTime](const FPipelineCacheFileFormatPSOMetaData& MetaData)
+		auto EntryShouldBeRemovedFromUserCache = [&Header, FileGuid=this->FileGuid, EvictionTime,&KnownGuids](const FPipelineCacheFileFormatPSOMetaData& MetaData)
 		{
 			// Remove the element if it is in the user cache and the time has elapsed, or if it was in a cache that no longer exists.
 			if (MetaData.FileGuid == Header.Guid)
@@ -1903,7 +1951,10 @@ public:
 			}
 			else
 			{
-				return MetaData.FileGuid != GameFileGuid;
+			// TODO: right now we do not have a way to supply known guids.
+			// we use only the expired date check for now.
+			//	return !KnownGuids.Contains(MetaData.FileGuid);
+				return EvictionTime >= MetaData.LastUsedUnixTime;
 			}
 		};
 
@@ -1970,7 +2021,7 @@ public:
 					*FileReader << ExistingPSO;
 					MemoryWriter << ExistingPSO;
 				}
-
+				check(NewEntry.FileGuid != FGuid());
 				NewTOC.MetaData.Add(Entry.Key, NewEntry);
 			}
 		}
@@ -2005,86 +2056,33 @@ public:
 		UE_LOG(LogRHI, Log, TEXT("Rewrote cache file. Old Size %lld, new size %lld (%lld byte reduction)"), OriginalSize, NewSize, OriginalSize - NewSize);
 	}
 	
-    bool ShouldDeleteExistingUserCache()
-    {
-        static bool bOnce = false;
-        static bool bCmdLineForce = false;
-        if (!bOnce)
-        {
-            bOnce = true;
-            bCmdLineForce = FParse::Param(FCommandLine::Get(), TEXT("deleteuserpsocache")) || FParse::Param(FCommandLine::Get(), TEXT("logPSO"));
-            UE_CLOG(bCmdLineForce, LogRHI, Warning, TEXT("****************************** Deleting user-writable PSO cache as requested on command line"));
-        }
-        return bCmdLineForce;
-    }
-    
-    bool ShouldLoadUserCache()
-    {
-		return FPipelineFileCache::LogPSOtoFileCache() && (CVarPSOFileCacheSaveUserCache.GetValueOnAnyThread() > 0);
-    }
-
-	bool OpenPipelineFileCache(FString const& FileName, EShaderPlatform Platform, FGuid& OutGameFileGuid)
+	bool OpenPipelineFileCache(FString const& NameIn, EShaderPlatform Platform, FGuid& OutGameFileGuid)
 	{
-		SET_DWORD_STAT(STAT_TotalGraphicsPipelineStateCount, 0);
-		SET_DWORD_STAT(STAT_TotalComputePipelineStateCount, 0);
-		SET_DWORD_STAT(STAT_TotalRayTracingPipelineStateCount, 0);
-		SET_DWORD_STAT(STAT_SerializedGraphicsPipelineStateCount, 0);
-		SET_DWORD_STAT(STAT_SerializedComputePipelineStateCount, 0);
-		SET_DWORD_STAT(STAT_NewGraphicsPipelineStateCount, 0);
-		SET_DWORD_STAT(STAT_NewComputePipelineStateCount, 0);
-		SET_DWORD_STAT(STAT_NewRayTracingPipelineStateCount, 0);
+		check(CacheStatus == EStatus::Unknown);
 
 		OutGameFileGuid = FGuid();
-		TOC.SortedOrder = FPipelineFileCache::PSOOrder::Default;
+		TOC.SortedOrder = FPipelineFileCacheManager::PSOOrder::Default;
 		TOC.MetaData.Empty();
 		
-		Name = FileName;
+		Name = NameIn;
 		
 		ShaderPlatform = Platform;
 		PlatformName = LegacyShaderPlatformToShaderFormat(Platform);
 		
-
-		FString GamePathStable = FPaths::ProjectContentDir() / TEXT("PipelineCaches") / ANSI_TO_TCHAR(FPlatformProperties::IniPlatformName()) / FString::Printf(TEXT("%s_%s.stable.upipelinecache"), *FileName, *PlatformName.ToString());
-		FString GamePath = FPaths::ProjectContentDir() / TEXT("PipelineCaches") / ANSI_TO_TCHAR(FPlatformProperties::IniPlatformName()) / FString::Printf(TEXT("%s_%s.upipelinecache"), *FileName, *PlatformName.ToString());
-
+		FString GamePath	= FPaths::ProjectContentDir() / TEXT("PipelineCaches") / ANSI_TO_TCHAR(FPlatformProperties::IniPlatformName()) / FString::Printf(TEXT("%s_%s.stable.upipelinecache"), *Name, *PlatformName.ToString());
 		static bool bCommandLineNotStable = FParse::Param(FCommandLine::Get(), TEXT("nostablepipelinecache"));
-		if (!bCommandLineNotStable)
-		{
-			GamePath = GamePathStable;
-		}
-
-		FString FilePath = FPaths::ProjectSavedDir() / FString::Printf(TEXT("%s_%s.upipelinecache"), *FileName, *PlatformName.ToString());
-
-		RecordingFilename = FString::Printf(TEXT("%s-CL-%u-"), *FEngineVersion::Current().GetBranchDescriptor(), FEngineVersion::Current().GetChangelist());
-
-		FGuid UniqueFileGuid;
-		FPlatformMisc::CreateGuid(UniqueFileGuid);  // not very unique on android, but won't matter much here
-		RecordingFilename += FString::Printf(TEXT("%s_%s_%s.rec.upipelinecache"), *FileName, *PlatformName.ToString(), *UniqueFileGuid.ToString());
-		RecordingFilename = FPaths::ProjectSavedDir() / TEXT("CollectedPSOs") / RecordingFilename;
-
-		UE_LOG(LogRHI, Log, TEXT("Base name for record PSOs is %s"), *RecordingFilename);
-
-		FString JournalPath = FilePath + JOURNAL_FILE_EXTENSION;
-        bool const bJournalFileExists = IFileManager::Get().FileExists(*JournalPath);
-		if (bJournalFileExists || ShouldDeleteExistingUserCache())
+		if (bCommandLineNotStable)
 		{
-			UE_LOG(LogRHI, Log, TEXT("Deleting FPipelineCacheFile: %s"), *FilePath);
-            // If either of the above are true we need to dispose of this case as we consider it invalid
-            if (IFileManager::Get().FileExists(*FilePath))
-            {
-                IFileManager::Get().Delete(*FilePath);
-            }
-            if (bJournalFileExists)
-            {
-                IFileManager::Get().Delete(*JournalPath);
-            }
+			GamePath.Empty();
 		}
 
-		const bool bGameFileOk = OpenPipelineFileCache(GamePath, GameFileGuid, GameAsyncFileHandle, GameTOC);
+		const bool bGameFileOk = OpenPipelineFileCache(GamePath, ShaderPlatform, FileGuid, AsyncFileHandle, TOC, TOCOffset);
 
 		if (bGameFileOk)
 		{
-			OutGameFileGuid = GameFileGuid;
+			FilePath = GamePath;
+			OutGameFileGuid = FileGuid;
+			CacheStatus = EStatus::BundledCache;
 		}
 
 		if (bGameFileOk && GRHISupportsLazyShaderCodeLoading && CVarLazyLoadShadersWhenPSOCacheIsPresent.GetValueOnAnyThread())
@@ -2093,61 +2091,21 @@ public:
 			GRHILazyShaderCodeLoading = true;
 		}
 
-        bool bUserFileOk = false;
-        
-        if (ShouldLoadUserCache())
-        {
-			GarbageCollectUserCache(FilePath);
-
-			FPipelineCacheFileFormatTOC UserTOC;
-            bUserFileOk = OpenPipelineFileCache(FilePath, UserFileGuid, UserAsyncFileHandle, UserTOC);
-            if (!bUserFileOk)
-            {
-                // Start the file again!
-                IFileManager::Get().Delete(*FilePath);
-                TOCOffset = 0;
-            }
-            else
-            {
-				for (auto const& Entry : UserTOC.MetaData)
-				{
-					// We want this entry that references the game version not the one from the Game TOC as that doesn't have ongoing UsageMasks bind counts etc...
-					auto* MetaPtr = TOC.MetaData.Find(Entry.Key);
-					if ((Entry.Value.FileGuid == UserFileGuid || Entry.Value.FileGuid == GameFileGuid) && (!MetaPtr || (MetaPtr->FileGuid != UserFileGuid && MetaPtr->FileGuid != GameFileGuid)))
-					{
-						TOC.MetaData.Add(Entry.Key, Entry.Value);
-					}
-				}
-
-				for (auto const& Entry : GameTOC.MetaData)
-                {
-                	// If its there - don't overwrite - we'll lose mutable user cache meta data unless an old entry
-					auto* MetaPtr = TOC.MetaData.Find(Entry.Key);
-                    if (!MetaPtr || (MetaPtr->FileGuid != UserFileGuid && MetaPtr->FileGuid != GameFileGuid))
-                    {
-                        TOC.MetaData.Add(Entry.Key, Entry.Value);
-                    }
-                }
-            }
-        }
-        
-        if (!bUserFileOk)
-        {
-            TOC = GameTOC;
-        }
-		
 		uint32 InvalidEntryCount = 0;
 		
         for (auto const& Entry : TOC.MetaData)
         {
-            FPipelineStateStats* Stat = FPipelineFileCache::Stats.FindRef(Entry.Key);
+            FPipelineStateStats* Stat = FPipelineFileCacheManager::Stats.FindRef(Entry.Key);
             if (!Stat)
             {
                 Stat = new FPipelineStateStats;
                 Stat->PSOHash = Entry.Key;
                 Stat->TotalBindCount = -1;
-                FPipelineFileCache::Stats.Add(Entry.Key, Stat);
+                FPipelineFileCacheManager::Stats.Add(Entry.Key, Stat);
             }
+
+			UE_CLOG(!!FPipelineFileCacheManager::NewPSOUsage.Find(Entry.Key), LogRHI, Warning, TEXT("loaded PSOFC %s contains entry (%u) previously marked as new "), *Name, (Entry.Key));
+			
 #if !UE_BUILD_SHIPPING
 			if((Entry.Value.EngineFlags & FPipelineCacheFlagInvalidPSO) != 0)
 			{
@@ -2161,12 +2119,123 @@ public:
         	UE_LOG(LogRHI, Warning, TEXT("Found %d / %d PSO entries marked as invalid."), InvalidEntryCount, TOC.MetaData.Num());
         }
 		
-		SET_MEMORY_STAT(STAT_FileCacheMemory, TOC.MetaData.GetAllocatedSize());
-		
-		return bGameFileOk || bUserFileOk;
+		INC_MEMORY_STAT_BY(STAT_FileCacheMemory, TOC.MetaData.GetAllocatedSize());
+
+		UE_LOG(LogRHI, VeryVerbose, TEXT("-- opened bundled %s cache:"), *NameIn);
+		TOC.DumpToLog();
+		UE_LOG(LogRHI, VeryVerbose, TEXT("-- opened bundled %s cache:"), *NameIn);
+
+		return bGameFileOk;
 	}
 	
-	void MergePSOUsageToMetaData(TMap<uint32, FPSOUsageData>& NewPSOUsage, TMap<uint32, FPipelineCacheFileFormatPSOMetaData>& MetaData, int64 CurrentUnixTime, bool bRemoveUpdatedentries = false)
+	// 
+	bool OpenUserPipelineFileCache(FString const& CacheNameIn, EShaderPlatform Platform)
+	{
+		bool bUserFileOk = false;
+
+		if (CacheStatus == EStatus::Unknown)
+		{
+			SET_DWORD_STAT(STAT_TotalGraphicsPipelineStateCount, 0);
+			SET_DWORD_STAT(STAT_TotalComputePipelineStateCount, 0);
+			SET_DWORD_STAT(STAT_TotalRayTracingPipelineStateCount, 0);
+			SET_DWORD_STAT(STAT_SerializedGraphicsPipelineStateCount, 0);
+			SET_DWORD_STAT(STAT_SerializedComputePipelineStateCount, 0);
+			SET_DWORD_STAT(STAT_NewGraphicsPipelineStateCount, 0);
+			SET_DWORD_STAT(STAT_NewComputePipelineStateCount, 0);
+			SET_DWORD_STAT(STAT_NewRayTracingPipelineStateCount, 0);
+
+			// one time attempt to open the user cache.
+			TOC.SortedOrder = FPipelineFileCacheManager::PSOOrder::Default;
+			TOC.MetaData.Empty();
+
+			Name = CacheNameIn;
+
+			ShaderPlatform = Platform;
+			PlatformName = LegacyShaderPlatformToShaderFormat(Platform);
+
+
+			FGuid UniqueFileGuid;
+			FPlatformMisc::CreateGuid(UniqueFileGuid);  // not very unique on android, but won't matter much here
+
+			RecordingFilename = FString::Printf(TEXT("%s-CL-%u-"), *FEngineVersion::Current().GetBranchDescriptor(), FEngineVersion::Current().GetChangelist());
+			RecordingFilename += FString::Printf(TEXT("%s_%s_%s.rec.upipelinecache"), *Name, *PlatformName.ToString(), *UniqueFileGuid.ToString());
+			RecordingFilename = FPaths::ProjectSavedDir() / TEXT("CollectedPSOs") / RecordingFilename;
+
+			UE_LOG(LogRHI, Log, TEXT("Base name for record PSOs is %s"), *RecordingFilename);
+			FilePath = FPaths::ProjectSavedDir() / FString::Printf(TEXT("%s_%s.upipelinecache"), *Name, *PlatformName.ToString());
+
+			if (ShouldDeleteExistingUserCache())
+			{
+				UE_LOG(LogRHI, Log, TEXT("Deleting FPipelineCacheFile: %s"), *FilePath);
+				if (IFileManager::Get().FileExists(*FilePath))
+				{
+					IFileManager::Get().Delete(*FilePath);
+				}
+			}
+
+			FString JournalPath = FilePath + JOURNAL_FILE_EXTENSION;
+			bool const bJournalFileExists = IFileManager::Get().FileExists(*JournalPath);
+			if (bJournalFileExists || ShouldDeleteExistingUserCache())
+			{
+				UE_LOG(LogRHI, Log, TEXT("Deleting FPipelineCacheFile: %s"), *FilePath);
+				// If either of the above are true we need to dispose of this case as we consider it invalid
+				if (IFileManager::Get().FileExists(*FilePath))
+				{
+					IFileManager::Get().Delete(*FilePath);
+				}
+				if (bJournalFileExists)
+				{
+					IFileManager::Get().Delete(*JournalPath);
+				}
+			}
+
+			// TODO: we currently do not know the full set of valid PSOFC guids.
+			// KnownGuids is a placeholder for all possible PSO cache guids.
+			TSet<FGuid> KnownGuids;
+			GarbageCollectUserCache(FilePath, KnownGuids);
+
+			FPipelineCacheFileFormatTOC UserTOC;
+			bUserFileOk = OpenPipelineFileCache(FilePath, ShaderPlatform, FileGuid, AsyncFileHandle, TOC, TOCOffset);
+			CacheStatus = bUserFileOk ? EStatus::UserCacheOpened : EStatus::NewUserCache;
+
+			if (!bUserFileOk)
+			{
+				FileGuid = FGuid::NewGuid();
+				// Start the file again!
+				IFileManager::Get().Delete(*FilePath);
+				TOCOffset = 0;
+			}
+
+			uint32 InvalidEntryCount = 0;
+
+			for (auto const& Entry : TOC.MetaData)
+			{
+				FPipelineStateStats* Stat = FPipelineFileCacheManager::Stats.FindRef(Entry.Key);
+				if (!Stat)
+				{
+					Stat = new FPipelineStateStats;
+					Stat->PSOHash = Entry.Key;
+					Stat->TotalBindCount = -1;
+					FPipelineFileCacheManager::Stats.Add(Entry.Key, Stat);
+				}
+#if !UE_BUILD_SHIPPING
+				if ((Entry.Value.EngineFlags & FPipelineCacheFlagInvalidPSO) != 0)
+				{
+					++InvalidEntryCount;
+				}
+#endif
+			}
+
+			if (InvalidEntryCount > 0)
+			{
+				UE_LOG(LogRHI, Warning, TEXT("Found %d / %d PSO entries marked as invalid."), InvalidEntryCount, TOC.MetaData.Num());
+			}
+		}
+		INC_MEMORY_STAT_BY(STAT_FileCacheMemory, TOC.MetaData.GetAllocatedSize());
+		return bUserFileOk;
+	}
+
+	static void MergePSOUsageToMetaData(TMap<uint32, FPSOUsageData>& NewPSOUsage, TMap<uint32, FPipelineCacheFileFormatPSOMetaData>& MetaData, int64 CurrentUnixTime, bool bRemoveUpdatedentries = false)
 	{
 		for(auto It = NewPSOUsage.CreateIterator(); It; ++It)
 		{
@@ -2188,20 +2257,23 @@ public:
 		}
 	}
 	
-	bool SavePipelineFileCache(FString const& FilePath, FPipelineFileCache::SaveMode Mode, TMap<uint32, FPipelineStateStats*> const& Stats, TSet<FPipelineCacheFileFormatPSO>& NewEntries, FPipelineFileCache::PSOOrder Order, TMap<uint32, FPSOUsageData>& NewPSOUsage)
+	bool SavePipelineFileCache(FPipelineFileCacheManager::SaveMode Mode, TMap<uint32, FPipelineStateStats*> const& Stats, TSet<FPipelineCacheFileFormatPSO>& NewEntries, FPipelineFileCacheManager::PSOOrder Order, TMap<uint32, FPSOUsageData>& NewPSOUsage)
 	{
+		check(CacheStatus == EStatus::NewUserCache || CacheStatus == EStatus::UserCacheOpened);
+		// remove from stats because this operation will modify the content and re-set the stat at the end.
+		DEC_MEMORY_STAT_BY(STAT_FileCacheMemory, TOC.MetaData.GetAllocatedSize());
 		QUICK_SCOPE_CYCLE_COUNTER(STAT_SavePipelineFileCache);
 		double StartTime = FPlatformTime::Seconds();
 		FString SaveFilePath = FilePath;
 		
-		if (FPipelineFileCache::SaveMode::BoundPSOsOnly == Mode)
+		if (FPipelineFileCacheManager::SaveMode::BoundPSOsOnly == Mode)
 		{
 			SaveFilePath = GetRecordingFilename();
 		}
 		
 		bool bFileWriteSuccess = false;
 		bool bPerformWrite = true;
-		if (FPipelineFileCache::SaveMode::Incremental == Mode)
+		if (FPipelineFileCacheManager::SaveMode::Incremental == Mode)
 		{
 			bPerformWrite = NewEntries.Num() || Order != TOC.SortedOrder || NewPSOUsage.Num();
 			bFileWriteSuccess = !bPerformWrite;
@@ -2213,7 +2285,7 @@ public:
             
 			int64 UnixTime = GetCurrentUnixTime();
 			FString JournalPath;
-			if (Mode != FPipelineFileCache::SaveMode::BoundPSOsOnly)
+			if (Mode != FPipelineFileCacheManager::SaveMode::BoundPSOsOnly)
 			{
 				JournalPath = SaveFilePath + JOURNAL_FILE_EXTENSION;
 				FArchive* JournalWriter = IFileManager::Get().CreateFileWriter(*JournalPath);
@@ -2227,7 +2299,7 @@ public:
 					Header.Version = FPipelineCacheFileFormatCurrentVersion;
 					Header.GameVersion = GameVersion;
 					Header.Platform = ShaderPlatform;
-					Header.Guid = UserFileGuid;
+					Header.Guid = FileGuid;
 					Header.TableOffset = 0;
 					Header.LastGCUnixTime = UnixTime;
 
@@ -2241,64 +2313,74 @@ public:
 			}
 			if (bPerformWrite)
 			{
-                FString GamePathStable = FPaths::ProjectContentDir() / TEXT("PipelineCaches") / ANSI_TO_TCHAR(FPlatformProperties::IniPlatformName()) / FString::Printf(TEXT("%s_%s.stable.upipelinecache"), *Name, *PlatformName.ToString());
-                FString GamePath = FPaths::ProjectContentDir() / TEXT("PipelineCaches") / ANSI_TO_TCHAR(FPlatformProperties::IniPlatformName()) / FString::Printf(TEXT("%s_%s.upipelinecache"), *Name, *PlatformName.ToString());
-				static bool bCommandLineNotStable = FParse::Param(FCommandLine::Get(), TEXT("nostablepipelinecache"));
-                if (!bCommandLineNotStable)
-                {
-                    GamePath = GamePathStable;
-                }
-
-                int64 GameFileSize = IFileManager::Get().FileSize(*GamePath);
-				TArray<uint8> GameFileBytes;
-				int64 FileSize = IFileManager::Get().FileSize(*FilePath);
-				TArray<uint8> UserFileBytes;
-				if(FPipelineFileCache::SaveMode::Incremental != Mode)
+				struct FMemoryReaderAndMemory 
 				{
-					if (GameFileSize > 0)
+					// Non-copyable
+					FMemoryReaderAndMemory(const FMemoryReaderAndMemory&) = delete;
+					FMemoryReaderAndMemory& operator=(const FMemoryReaderAndMemory&) = delete;
+
+					TArray<uint8> Bytes;
+					TUniquePtr<FMemoryReader> Reader;
+					explicit FMemoryReaderAndMemory(FPipelineCacheFile* PipelineFile)
 					{
-						if (GameAsyncFileHandle.IsValid())
-						{
-							GameFileBytes.AddZeroed(GameFileSize);
-							IAsyncReadRequest* Request = GameAsyncFileHandle->ReadRequest(0, GameFileSize, AIOP_Normal, nullptr, GameFileBytes.GetData());
-							Request->WaitCompletion();
-							delete Request;
-							// Can't report errors here because the AsyncIO requests have no such mechanism.
-						}
-						else
+						if(PipelineFile)
 						{
-							bool bReadOK = FFileHelper::LoadFileToArray(GameFileBytes, *GamePath);
-							UE_CLOG(!bReadOK, LogRHI, Warning, TEXT("Failed to read %lld bytes from %s while re-saving the PipelineFileCache!"), GameFileSize, *GamePath);
+							int64 FileSize = IFileManager::Get().FileSize(*PipelineFile->FilePath);
+							if (FileSize > 0)
+							{
+								Bytes.SetNumUninitialized(FileSize);
+								if (PipelineFile->AsyncFileHandle.IsValid())
+								{
+									IAsyncReadRequest* Request = PipelineFile->AsyncFileHandle->ReadRequest(0, FileSize, AIOP_Normal, nullptr, Bytes.GetData());
+									Request->WaitCompletion();
+									delete Request;
+									UE_LOG(LogRHI, VeryVerbose, TEXT("1 Opening %s as guid %s, size %d"), *PipelineFile->FilePath, *PipelineFile->GetFileGuid().ToString(), FileSize);
+								}
+								else
+								{
+									bool bReadOK = FFileHelper::LoadFileToArray(Bytes, *PipelineFile->FilePath);
+									UE_LOG(LogRHI, VeryVerbose, TEXT("2 Opening %s as guid %s, size %d"), *PipelineFile->FilePath, *PipelineFile->GetFileGuid().ToString(), FileSize);
+									UE_CLOG(!bReadOK, LogRHI, Warning, TEXT("Failed to read %lld bytes from %s while re-saving the PipelineFileCache!"), FileSize, *PipelineFile->FilePath);
+								}
+								Reader = MakeUnique<FMemoryReader>(Bytes);
+								Reader->SetGameNetVer(FPipelineCacheFileFormatCurrentVersion);
+							}
 						}
 					}
-		 
-					if (FileSize > 0)
+					FMemoryReader* GetReader() { return Reader.Get(); }
+				};
+
+				TMap<FGuid, TUniquePtr<FMemoryReaderAndMemory>> GuidToFileCacheReader;
+
+				auto GetFileCacheReaderFromGuid = [&GuidToFileCacheReader](const FGuid& guid)
+				{
+					check(guid != FGuid());
+					if(!GuidToFileCacheReader.Contains(guid))
 					{
-						if (UserAsyncFileHandle.IsValid())
-						{
-							UserFileBytes.AddZeroed(FileSize);
-							IAsyncReadRequest* Request = UserAsyncFileHandle->ReadRequest(0, FileSize, AIOP_Normal, nullptr, UserFileBytes.GetData());
-							Request->WaitCompletion();
-							delete Request;
-							// Can't report errors here because the AsyncIO requests have no such mechanism.
-						}
-						else
-						{
-							bool bReadOK = FFileHelper::LoadFileToArray(UserFileBytes, *FilePath);
-							UE_CLOG(!bReadOK, LogRHI, Warning, TEXT("Failed to read %lld bytes from %s while re-saving the PipelineFileCache!"), FileSize, *FilePath);
-						}
+						const FString& FoundKey = FPipelineFileCacheManager::GameGuidToCacheKey.FindRef(guid);
+
+						FPipelineCacheFile* FoundPipelineFile = FPipelineFileCacheManager::GetPipelineCacheFileFromKey(FoundKey);
+
+						GuidToFileCacheReader.Add(guid, MakeUnique<FMemoryReaderAndMemory>(FoundPipelineFile));
 					}
-				}
-		 
+
+					if (TUniquePtr<FMemoryReaderAndMemory>* Found = GuidToFileCacheReader.Find(guid))
+					{
+						return  (*Found)->GetReader();
+					}
+					checkNoEntry();
+					return static_cast<FMemoryReader*>(nullptr);
+				};
+
 				// Assume caller has handled Platform specifc path + filename
 				TArray<uint8> SaveBytes;
 				FArchive* FileWriter;
-				bool bUseMemoryWriter = (Mode == FPipelineFileCache::SaveMode::BoundPSOsOnly);
+				bool bUseMemoryWriter = (Mode == FPipelineFileCacheManager::SaveMode::BoundPSOsOnly);
 				FString TempPath = SaveFilePath;
 				// Only use a file switcheroo on Apple platforms as they are the only ones tested so far.
 				// At least two other platforms MoveFile implementation looks broken when moving from a writable source file to a writeable destination.
 				// They only handle moves/renames between the read-only -> writeable directories/devices.
-				if ((PLATFORM_APPLE || PLATFORM_ANDROID) && Mode != FPipelineFileCache::SaveMode::Incremental)
+				if ((PLATFORM_APPLE || PLATFORM_ANDROID) && Mode != FPipelineFileCacheManager::SaveMode::Incremental)
 				{
 					TempPath += TEXT(".tmp");
 				}
@@ -2325,7 +2407,7 @@ public:
 						Header.Version = FPipelineCacheFileFormatCurrentVersion;
 						Header.GameVersion = GameVersion;
 						Header.Platform = ShaderPlatform;
-						Header.Guid = UserFileGuid;
+						Header.Guid = FileGuid;
 						Header.TableOffset = 0;
 						Header.LastGCUnixTime = UnixTime;
 		 
@@ -2340,7 +2422,8 @@ public:
                     switch (Mode)
                     {
 						// This mode just writes new, used, entries to the end of the file and updates the TOC which will contain entries from the Game-Content file that are redundant.
-                        case FPipelineFileCache::SaveMode::Incremental:
+						// Any current tasks are unaffected as the prior offsets are still valid.
+                        case FPipelineFileCacheManager::SaveMode::Incremental:
                         {
                             // PSO Descriptors
                             uint64 PSOOffset = TOCOffset;
@@ -2348,44 +2431,53 @@ public:
                             FileWriter->Seek(PSOOffset);
                             
                             // Add new entries
-							TotalEntries = NewEntries.Num();
 							for(auto It = NewEntries.CreateIterator(); It; ++It)
                             {
 								FPipelineCacheFileFormatPSO& NewEntry = *It;
-                                check(!IsPSOEntryCached(NewEntry));
-                                
-                                uint32 PSOHash = GetTypeHash(NewEntry);
 
+                                uint32 PSOHash = GetTypeHash(NewEntry);
+								
+								bool bFound = FPipelineFileCacheManager::IsPSOEntryCached(NewEntry, nullptr);
+								if (bFound)
+								{
+									// this could happen if another PSOFC loads after the PSO was encountered, if desired we could remove things from newentries when a psofc is mounted..
+									UE_LOG(LogRHI, Display, TEXT("Incrementally saving new PSOs but entry (%u), is already cached.."), PSOHash);
+									// Not removing it as the cached item could be legit if co-owner is not always loaded.
+								}
 								FPipelineStateStats const* Stat = Stats.FindRef(PSOHash);
-                                if (Stat && Stat->TotalBindCount > 0)
+								if (Stat && Stat->TotalBindCount > 0)
 								{
 									FPipelineCacheFileFormatPSOMetaData Meta;
 									Meta.Stats.PSOHash = PSOHash;
 									Meta.FileOffset = PSOOffset;
-									Meta.FileGuid = UserFileGuid;
+									Meta.FileGuid = FileGuid;
 									Meta.AddShaders(NewEntry);
-									
+
 									TArray<uint8> Bytes;
 									FMemoryWriter Wr(Bytes);
 									Wr.SetGameNetVer(FPipelineCacheFileFormatCurrentVersion);
 									Wr << NewEntry;
-									
+
 									FileWriter->Serialize(Bytes.GetData(), Wr.TotalSize());
-									
+
 									Meta.FileSize = Wr.TotalSize();
-									
+
+									check(Meta.FileGuid != FGuid());
 									TOC.MetaData.Add(PSOHash, Meta);
 									PSOOffset += Meta.FileSize;
-									
+
 									check(PSOOffset == FileWriter->Tell());
-									
+
 									NumNewEntries++;
-									
+
 									It.RemoveCurrent();
+									UE_LOG(LogRHI, VeryVerbose, TEXT("Incremental save is appending new PSOs (%u)"), PSOHash);
 								}
                             }
-                            
-                            if(Order != FPipelineFileCache::PSOOrder::Default)
+							// We're appending to the current user cache here, Our TOC is the total.
+							TotalEntries = TOC.MetaData.Num();
+
+                            if(Order != FPipelineFileCacheManager::PSOOrder::Default)
                             {
                                 SortMetaData(TOC.MetaData, Order);
                                 TOC.SortedOrder = Order;
@@ -2393,7 +2485,7 @@ public:
                             else
                             {
                                 // Added new entries and not re-sorted - the sort order invalid - reset to default
-                                TOC.SortedOrder = FPipelineFileCache::PSOOrder::Default;
+                                TOC.SortedOrder = FPipelineFileCacheManager::PSOOrder::Default;
                             }
 							
 							// Update TOC Metadata usage and clear relevant entries in NewPSOUsage as we are saving this file cache TOC
@@ -2403,19 +2495,22 @@ public:
                             TOCOffset = PSOOffset;
 							
                             FileWriter->Seek(Header.TableOffset);
-                            *FileWriter << TOC;
+							// use a temp here because serializing can destroy our metadata guids.
+							FPipelineCacheFileFormatTOC TempTOC = TOC;
+                            *FileWriter << TempTOC;
                             break;
                         }
-						// These modes actually save to a separate file that records only PSOs that were bound.
+						// This mode actually saves to a separate file that records only PSOs that were bound.
 						// BoundPSOsOnly will record all those PSOs used in this run of the game.
-                        case FPipelineFileCache::SaveMode::BoundPSOsOnly:
+                        case FPipelineFileCacheManager::SaveMode::BoundPSOsOnly:
                         {
-							FMemoryReader UserFileBytesReader(UserFileBytes);
-							FMemoryReader GameFileBytesReader(GameFileBytes);
-							UserFileBytesReader.SetGameNetVer(FPipelineCacheFileFormatCurrentVersion);
-							GameFileBytesReader.SetGameNetVer(FPipelineCacheFileFormatCurrentVersion);
+                            FPipelineCacheFileFormatTOC TempTOC;
+							// Merge all of the existing PSO caches together, including this (user cache)
+							for (TPair<FString, TUniquePtr<class FPipelineCacheFile>>& PipelineCachePair : FPipelineFileCacheManager::FileCacheMap)
+							{
+								TempTOC.MetaData.Append(PipelineCachePair.Value->TOC.MetaData);
+							}
 
-                            FPipelineCacheFileFormatTOC TempTOC = TOC;
                             TMap<uint32, FPipelineCacheFileFormatPSO> PSOs;
 							
 							Header.Guid = FGuid::NewGuid();
@@ -2428,7 +2523,8 @@ public:
                                 Meta.FileSize = 0;
                                 Meta.FileGuid = Header.Guid;
 								Meta.AddShaders(Entry);
-								
+								check(Meta.FileGuid != FGuid());
+
                                 TempTOC.MetaData.Add(Meta.Stats.PSOHash, Meta);
                                 PSOs.Add(Meta.Stats.PSOHash, Entry);
                             }
@@ -2466,30 +2562,26 @@ public:
                                     if (!PSOs.Contains(It->Key))
                                     {
                                         check(It->Value.FileSize > 0);
-                                        if (It->Value.FileGuid == UserFileGuid)
-                                        {
-											check(It->Value.FileOffset < (uint32)UserFileBytes.Num());
-											UserFileBytesReader.Seek(It->Value.FileOffset);
-											
-											FPipelineCacheFileFormatPSO PSO;
-											UserFileBytesReader << PSO;
-											
-											PSOs.Add(It->Key, PSO);
-                                        }
-                                        else if (It->Value.FileGuid == GameFileGuid)
-                                        {
-											check(It->Value.FileOffset < (uint32)GameFileBytes.Num());
-											GameFileBytesReader.Seek(It->Value.FileOffset);
-											
+										FMemoryReader* Reader = GetFileCacheReaderFromGuid(It->Value.FileGuid);
+										if (Reader)
+										{
+											UE_LOG(LogRHI, VeryVerbose, TEXT("reading PSO (%u) from guid %s, Off %d // %d"), It->Key, *It->Value.FileGuid.ToString(), (It->Value.FileOffset), Reader->TotalSize());
+											check(It->Value.FileOffset < (uint32)Reader->TotalSize());
+											Reader->Seek(It->Value.FileOffset);
+
 											FPipelineCacheFileFormatPSO PSO;
-											GameFileBytesReader << PSO;
-											
+											(*Reader) << PSO;
+
 											PSOs.Add(It->Key, PSO);
-                                        }
+										}
 										else
 										{
-											UE_LOG(LogRHI, Verbose, TEXT("Trying to reconcile from unknown file GUID: %s but bound log file is: %s user file is: %s and game file is: %s - this means you have stale entries in a local cache file or the game content file is filled with bogus entries whose FileGUID doesn't match."), *(It->Value.FileGuid.ToString()), *(Header.Guid.ToString()), *(UserFileGuid.ToString()), *(GameFileGuid.ToString()));
-											
+											FString GameGuids;
+											Algo::ForEach(FPipelineFileCacheManager::FileCacheMap, [&GameGuids](auto& MapPair) { GameGuids += FString::Printf(TEXT("[%s - %s]"), *MapPair.Value->Name, *MapPair.Value->FileGuid.ToString()); });
+											UE_LOG(LogRHI, Display, \
+												TEXT("Trying to reconcile from unknown file GUID: %s but bound log file is: %s user file is: %s and the currently known game files are: %s - this means you have stale entries in a local cache file or the relevant game content file is yet to be mounted."), \
+												*(It->Value.FileGuid.ToString()), *(Header.Guid.ToString()), *(FileGuid.ToString()), *(GameGuids));
+
 											RemovedEntries++;
 											It.RemoveCurrent();
 										}
@@ -2537,136 +2629,6 @@ public:
                             FileWriter->Seek(Header.TableOffset);
                             *FileWriter << TempTOC;
                             
-                            break;
-                        }
-						// This mode should store all the PSOs that this device binds that weren't in a game-content cache.
-						// It will store the meta-data for all the PSOs that are ever bound, but it will omit PSO descriptors for entries that were cached in the game-content file
-						// This way the user builds up a log of uncaught entries but doesn't have to replicate the entire game-content file.
-                        case FPipelineFileCache::SaveMode::SortedBoundPSOs:
-                        {
-							TMap<uint32, FPipelineCacheFileFormatPSO> PSOs;
-                            for (auto& Entry : TOC.MetaData)
-                            {
-                                FPipelineCacheFileFormatPSO PSO;
-                                uint8* Bytes = nullptr;
-                                check(Entry.Value.FileSize > 0);
-                                if (Entry.Value.FileGuid == UserFileGuid)
-                                {
-                                    Bytes = &UserFileBytes[Entry.Value.FileOffset];
-									
-									TArray<uint8> PSOData(Bytes, Entry.Value.FileSize);
-									FMemoryReader Ar(PSOData);
-									Ar.SetGameNetVer(FPipelineCacheFileFormatCurrentVersion);
-									Ar << PSO;
-									PSOs.Add(Entry.Key, PSO);
-                                }
-								else if (Entry.Value.FileGuid != GameFileGuid)
-								{
-									UE_LOG(LogRHI, Verbose, TEXT("Trying to reconcile from unknown file GUID: %s but user file is: %s and game file is: %s - this means you have stale entries in a local cache file that reference a previous version of the game content cache."), *(Entry.Value.FileGuid.ToString()), *(UserFileGuid.ToString()), *(GameFileGuid.ToString()));
-								}
-                            }
-                            
-                            for (auto& Entry : NewEntries)
-                            {
-								FPipelineCacheFileFormatPSOMetaData Meta;
-                                Meta.Stats.PSOHash = GetTypeHash(Entry);
-                                Meta.FileOffset = 0;
-                                Meta.FileSize = 0;
-                                Meta.FileGuid = UserFileGuid;
-								Meta.AddShaders(Entry);
-                                
-                                TOC.MetaData.Add(Meta.Stats.PSOHash, Meta);
-                                PSOs.Add(Meta.Stats.PSOHash, Entry);
-                            }
-							
-							// Update TOC Metadata usage and clear updated entries in NewPSOUsage as file TOC is getting updated
-							MergePSOUsageToMetaData(NewPSOUsage, TOC.MetaData, UnixTime, true);
-							
-							FPipelineCacheFileFormatTOC TempTOC = TOC;
-							// Update PSO usage stats for new and old
-							for (auto& Pair : Stats)
-							{
-								auto* MetaPtr = TempTOC.MetaData.Find(Pair.Key);
-								if (MetaPtr)
-								{
-									auto& Meta = *MetaPtr;
-									check(Meta.Stats.PSOHash == Pair.Value->PSOHash);
-									Meta.Stats.CreateCount += Pair.Value->CreateCount;
-									if (Pair.Value->FirstFrameUsed > Meta.Stats.FirstFrameUsed)
-									{
-										Meta.Stats.FirstFrameUsed = Pair.Value->FirstFrameUsed;
-									}
-									if (Pair.Value->LastFrameUsed > Meta.Stats.LastFrameUsed)
-									{
-										Meta.Stats.LastFrameUsed = Pair.Value->LastFrameUsed;
-									}
-									Meta.Stats.TotalBindCount = (int64)FMath::Min((uint64)INT64_MAX, (uint64)FMath::Max(Meta.Stats.TotalBindCount, 0ll) + (uint64)FMath::Max(Pair.Value->TotalBindCount, 0ll));
-								}
-							}
-							
-							for (auto It = TempTOC.MetaData.CreateIterator(); It; ++It)
-							{
-								// If the entry doesn't belong to the game content or user local cache then remove it as it is invalid
-								// Anything that has never been compiled (BindCount < 0) is invalid and can be removed
-								// Or, if the BindCount is >= 0 and the same as in the GameTOC we have never seen it and we don't need to store it in the game cache
-								FPipelineCacheFileFormatPSOMetaData* GameData = GameTOC.MetaData.Find(It->Key);
-								if ((It->Value.FileGuid != UserFileGuid && It->Value.FileGuid != GameFileGuid)
-									|| It->Value.Stats.TotalBindCount < 0
-									|| (GameData && (It->Value.Stats.TotalBindCount == GameData->Stats.TotalBindCount)))
-								{
-									RemovedEntries++;
-									It.RemoveCurrent();
-								}
-							}
-							TotalEntries = TempTOC.MetaData.Num();
-							
-                            SortMetaData(TempTOC.MetaData, Order);
-                            TOC.SortedOrder = TempTOC.SortedOrder = Order;
-                            
-                            TOCOffset = (uint64)FileWriter->Tell();
-                            
-                            uint64 PSOOffset = TOCOffset;
-                            
-                            for (auto& Entry : TempTOC.MetaData)
-                            {
-                                // When saved in this mode the user local file only stores the PSO descriptor for entries that weren't in the game-content cache
-								// We don't need to store the PSO data for entries that come from the game cache
-								// We do store the meta-data for all PSOs that this device has ever seen and that are valid with the current game-content and user cache.
-								auto& CurrentMeta = TOC.MetaData.FindChecked(Entry.Key);
-								if (CurrentMeta.FileGuid == UserFileGuid)
-								{
-									CurrentMeta.FileOffset = PSOOffset;
-									Entry.Value.FileOffset = PSOOffset;
-									
-									FPipelineCacheFileFormatPSO& PSO = PSOs.FindChecked(Entry.Key);
-									
-									FileWriter->Seek(PSOOffset);
-									
-									TArray<uint8> Bytes;
-									FMemoryWriter Wr(Bytes);
-									Wr.SetGameNetVer(FPipelineCacheFileFormatCurrentVersion);
-									Wr << PSO;
-									
-									NewEntries.Remove(PSO);
-									
-									FileWriter->Serialize(Bytes.GetData(), Wr.TotalSize());
-									
-									CurrentMeta.FileSize = Wr.TotalSize();
-									Entry.Value.FileSize = Wr.TotalSize();
-									
-									PSOOffset += Entry.Value.FileSize;
-									check(PSOOffset == FileWriter->Tell());
-									
-									NumNewEntries++;
-								}
-                            }
-                            
-                            Header.TableOffset = PSOOffset;
-                            TOCOffset = PSOOffset;
-                            
-                            FileWriter->Seek(Header.TableOffset);
-                            *FileWriter << TempTOC;
-							
                             break;
                         }
                         default:
@@ -2698,7 +2660,7 @@ public:
 						{
 							delete FileWriter;
 							float ThisTimeMS = float(FPlatformTime::Seconds() - StartTime) * 1000.0f;
-							UE_LOG(LogRHI, Log, TEXT("FPipelineFileCache skipping saving empty .upipelinecache (took %6.2fms): %s."), ThisTimeMS, *SaveFilePath);
+							UE_LOG(LogRHI, Log, TEXT("FPipelineFileCacheManager skipping saving empty .upipelinecache (took %6.2fms): %s."), ThisTimeMS, *SaveFilePath);
 							return false;
 						}
 					}
@@ -2715,18 +2677,16 @@ public:
 							TCHAR const* ModeName = nullptr;
 							switch (Mode)
 							{
-							case FPipelineFileCache::SaveMode::Incremental:
+							default:
+								checkNoEntry();
+							case FPipelineFileCacheManager::SaveMode::Incremental:
 								ModeName = TEXT("Incremental");
 								break;
-							case FPipelineFileCache::SaveMode::BoundPSOsOnly:
+							case FPipelineFileCacheManager::SaveMode::BoundPSOsOnly:
 								ModeName = TEXT("BoundPSOsOnly");
 								break;
-							case FPipelineFileCache::SaveMode::SortedBoundPSOs:
-							default:
-								ModeName = TEXT("SortedBoundPSOs");
-								break;
 							}
-							UE_LOG(LogRHI, Log, TEXT("FPipelineFileCache %s saved %u total, %u new, %u removed, %u cons .upipelinecache (took %6.2fms): %s."), ModeName, TotalEntries, NumNewEntries, RemovedEntries, ConsolidatedEntries, ThisTimeMS, *SaveFilePath);
+							UE_LOG(LogRHI, Log, TEXT("FPipelineFileCacheManager %s saved %u total, %u new, %u removed, %u cons .upipelinecache (took %6.2fms): %s."), ModeName, TotalEntries, NumNewEntries, RemovedEntries, ConsolidatedEntries, ThisTimeMS, *SaveFilePath);
 							
 							if (JournalPath.Len())
 							{
@@ -2755,12 +2715,12 @@ public:
 			}
 		}
 		
-		SET_MEMORY_STAT(STAT_FileCacheMemory, TOC.MetaData.GetAllocatedSize());
+		INC_MEMORY_STAT_BY(STAT_FileCacheMemory, TOC.MetaData.GetAllocatedSize());
 		
 		return bFileWriteSuccess;
 	}
 	
-	bool IsPSOEntryCached(FPipelineCacheFileFormatPSO const& NewEntry, FPSOUsageData* EntryData = nullptr) const
+	bool IsPSOEntryCachedInternal(FPipelineCacheFileFormatPSO const& NewEntry, FPSOUsageData* EntryData = nullptr) const
 	{
 		uint32 PSOHash = GetTypeHash(NewEntry);
 		FPipelineCacheFileFormatPSOMetaData const * const Existing = TOC.MetaData.Find(PSOHash);
@@ -2774,9 +2734,9 @@ public:
 		return Existing != nullptr;
 	}
 	
-	bool IsBSSEquivalentPSOEntryCached(FPipelineCacheFileFormatPSO const& NewEntry) const
+	bool IsBSSEquivalentPSOEntryCachedInternal(FPipelineCacheFileFormatPSO const& NewEntry) const
 	{
-		check(!IsPSOEntryCached(NewEntry)); // this routine should only be called after we have done the much faster test
+		check(!IsPSOEntryCachedInternal(NewEntry)); // this routine should only be called after we have done the much faster test
 		bool bResult = false;
 		if (NewEntry.Type == FPipelineCacheFileFormatPSO::DescriptorType::Graphics)
 		{
@@ -2813,22 +2773,22 @@ public:
 		return bResult;
 	}
 	
-	static void SortMetaData(TMap<uint32, FPipelineCacheFileFormatPSOMetaData>& MetaData, FPipelineFileCache::PSOOrder Order)
+	static void SortMetaData(TMap<uint32, FPipelineCacheFileFormatPSOMetaData>& MetaData, FPipelineFileCacheManager::PSOOrder Order)
 	{
 		// Only sorting metadata ordering - this should not affect PSO data offsets / lookups
 		switch(Order)
 		{
-			case FPipelineFileCache::PSOOrder::FirstToLatestUsed:
+			case FPipelineFileCacheManager::PSOOrder::FirstToLatestUsed:
 			{
 				MetaData.ValueSort([](const FPipelineCacheFileFormatPSOMetaData& A, const FPipelineCacheFileFormatPSOMetaData& B) {return A.Stats.FirstFrameUsed > B.Stats.FirstFrameUsed;});
 				break;
 			}
-			case FPipelineFileCache::PSOOrder::MostToLeastUsed:
+			case FPipelineFileCacheManager::PSOOrder::MostToLeastUsed:
 			{
 				MetaData.ValueSort([](const FPipelineCacheFileFormatPSOMetaData& A, const FPipelineCacheFileFormatPSOMetaData& B) {return A.Stats.TotalBindCount > B.Stats.TotalBindCount;});
 				break;
 			}
-			case FPipelineFileCache::PSOOrder::Default:
+			case FPipelineFileCacheManager::PSOOrder::Default:
 			default:
 			{
 				// NOP - leave as is
@@ -2837,7 +2797,7 @@ public:
 		}
 	}
 	
-	void GetOrderedPSOHashes(TArray<FPipelineCachePSOHeader>& PSOHashes, FPipelineFileCache::PSOOrder Order, int64 MinBindCount, TSet<uint32> const& AlreadyCompiledHashes)
+	void GetOrderedPSOHashes(TArray<FPipelineCachePSOHeader>& PSOHashes, FPipelineFileCacheManager::PSOOrder Order, int64 MinBindCount, TSet<uint32> const& AlreadyCompiledHashes)
 	{
 		if(Order != TOC.SortedOrder)
 		{
@@ -2848,7 +2808,7 @@ public:
 		for (auto const& Hash : TOC.MetaData)
 		{
 			if( (Hash.Value.EngineFlags & FPipelineCacheFlagInvalidPSO) == 0 &&
-				FPipelineFileCache::MaskComparisonFn(FPipelineFileCache::GameUsageMask, Hash.Value.UsageMask) &&
+				FPipelineFileCacheManager::MaskComparisonFn(FPipelineFileCacheManager::GameUsageMask, Hash.Value.UsageMask) &&
 				Hash.Value.Stats.TotalBindCount >= MinBindCount &&
 				!AlreadyCompiledHashes.Contains(Hash.Key))
 			{
@@ -2894,41 +2854,27 @@ public:
 				Entry->bValid = false;
 				continue;
 			}
-			
-			if (Meta.FileGuid == GameFileGuid)
-			{
-                FPipelineCacheFileFormatPSOMetaData const* GameMeta = GameTOC.MetaData.Find(Entry->Hash);
-                if (GameMeta && GameAsyncFileHandle.IsValid())
-                {
-                    Entry->Data.SetNum(GameMeta->FileSize);
-                    Entry->ParentFileHandle = GameAsyncFileHandle;
-                    Entry->ReadRequest = MakeShareable(GameAsyncFileHandle->ReadRequest(GameMeta->FileOffset, GameMeta->FileSize, AIOP_Normal, nullptr, Entry->Data.GetData()));
-                }
-                else
-                {
-                    UE_LOG(LogRHI, Verbose, TEXT("Encountered a PSO entry %u that has been removed from the game-content file: %s or no game-content file"), Entry->Hash, *Meta.FileGuid.ToString());
-                    Entry->bValid = false;
-                    continue;
-                }
-			}
-			else if (Meta.FileGuid == UserFileGuid)
+		
+			if (Meta.FileGuid == FileGuid)
 			{
-				if(UserAsyncFileHandle.IsValid())
+				FPipelineCacheFileFormatPSOMetaData const* GameMeta = TOC.MetaData.Find(Entry->Hash);
+
+				if (GameMeta && ensure(AsyncFileHandle.IsValid()))
 				{
-					Entry->Data.SetNum(Meta.FileSize);
-					Entry->ParentFileHandle = UserAsyncFileHandle;
-					Entry->ReadRequest = MakeShareable(UserAsyncFileHandle->ReadRequest(Meta.FileOffset, Meta.FileSize, AIOP_Normal, nullptr, Entry->Data.GetData()));
+					Entry->Data.SetNum(GameMeta->FileSize);
+					Entry->ParentFileHandle = AsyncFileHandle;
+					Entry->ReadRequest = MakeShareable(AsyncFileHandle->ReadRequest(GameMeta->FileOffset, GameMeta->FileSize, AIOP_Normal, nullptr, Entry->Data.GetData()));
 				}
 				else
 				{
-					UE_LOG(LogRHI, Verbose, TEXT("Encountered a PSO entry %u that references user content file ID: %s but async handle not valid"), Entry->Hash, *Meta.FileGuid.ToString());
+					UE_LOG(LogRHI, Verbose, TEXT("Encountered a PSO entry %u that has been removed from the cache file: %s "), Entry->Hash, *Meta.FileGuid.ToString());
 					Entry->bValid = false;
-                    continue;
+					continue;
 				}
 			}
             else
             {
-                UE_LOG(LogRHI, Verbose, TEXT("Encountered a PSO entry %u that references unknown file ID: %s"), Entry->Hash, *Meta.FileGuid.ToString());
+                UE_LOG(LogRHI, Warning, TEXT("Encountered a PSO entry %u that references unknown file ID: %s"), Entry->Hash, *Meta.FileGuid.ToString());
                 Entry->bValid = false;
                 continue;
             }
@@ -2956,24 +2902,65 @@ public:
 	{
 		return RecordingFilename;
 	}
-	
-private:
-	FString Name;
-	EShaderPlatform ShaderPlatform;
-	FName PlatformName;
-	uint64 TOCOffset;
-    FPipelineCacheFileFormatTOC GameTOC; // < The game TOC is kept around separately to handle cases where a fast-saved user cache tries to load removed entries from the game file.
-	FPipelineCacheFileFormatTOC TOC;
-	FGuid UserFileGuid;
-	FGuid GameFileGuid;
-	TSharedPtr<IAsyncReadFileHandle, ESPMode::ThreadSafe> UserAsyncFileHandle;
-	TSharedPtr<IAsyncReadFileHandle, ESPMode::ThreadSafe> GameAsyncFileHandle;
 
-	FString RecordingFilename;
+	const FString& GetCacheFilename() const
+	{
+		return Name;
+	}
+
+	const FGuid& GetFileGuid() const
+	{
+		return FileGuid;
+	}
+
+	const int32 GetTOCMetaDataSize() const
+	{
+		return TOC.MetaData.Num();
+	}
 };
 uint32 FPipelineCacheFile::GameVersion = 0;
 
-bool FPipelineFileCache::IsPipelineFileCacheEnabled()
+//////////////////////////////////////////////////////////////////////////
+//////////////////////////////////////////////////////////////////////////
+//////////////////////////////////////////////////////////////////////////
+//////////////////////////////////////////////////////////////////////////
+//////////////////////////////////////////////////////////////////////////
+//////////////////////////////////////////////////////////////////////////
+//////////////////////////////////////////////////////////////////////////
+//////////////////////////////////////////////////////////////////////////
+//////////////////////////////////////////////////////////////////////////
+//////////////////////////////////////////////////////////////////////////
+//////////////////////////////////////////////////////////////////////////
+//////////////////////////////////////////////////////////////////////////
+
+
+bool FPipelineFileCacheManager::IsBSSEquivalentPSOEntryCached(FPipelineCacheFileFormatPSO const& NewEntry)
+{
+	check(!IsPSOEntryCached(NewEntry)); // this routine should only be called after we have done the much faster test
+
+	bool bFound = false;
+
+	for (auto MapIt = FileCacheMap.CreateIterator(); !bFound && MapIt; ++MapIt)
+	{
+		bFound = MapIt->Value->IsBSSEquivalentPSOEntryCachedInternal(NewEntry);
+	}
+
+	return bFound;
+}
+
+// note: when EntryData is supplied it also performs a state update, we may need to verify all occurrences are coherent...
+bool FPipelineFileCacheManager::IsPSOEntryCached(FPipelineCacheFileFormatPSO const& NewEntry, FPSOUsageData* EntryData)
+{
+	bool bFound = false;
+
+	for (auto MapIt = FileCacheMap.CreateIterator(); !bFound && MapIt; ++MapIt)
+	{
+		bFound = MapIt->Value->IsPSOEntryCachedInternal(NewEntry, EntryData);
+	}
+	return bFound;
+}
+
+bool FPipelineFileCacheManager::IsPipelineFileCacheEnabled()
 {
 	static bool bOnce = false;
 	static bool bCmdLineForce = false;
@@ -2986,7 +2973,7 @@ bool FPipelineFileCache::IsPipelineFileCacheEnabled()
 	return FileCacheEnabled && (bCmdLineForce || CVarPSOFileCacheEnabled.GetValueOnAnyThread() == 1);
 }
 
-bool FPipelineFileCache::LogPSOtoFileCache()
+bool FPipelineFileCacheManager::LogPSOtoFileCache()
 {
 	static bool bOnce = false;
 	static bool bCmdLineForce = false;
@@ -2999,7 +2986,7 @@ bool FPipelineFileCache::LogPSOtoFileCache()
 	return (bCmdLineForce || CVarPSOFileCacheLogPSO.GetValueOnAnyThread() == 1);
 }
 
-bool FPipelineFileCache::ReportNewPSOs()
+bool FPipelineFileCacheManager::ReportNewPSOs()
 {
     static bool bOnce = false;
     static bool bCmdLineForce = false;
@@ -3012,7 +2999,7 @@ bool FPipelineFileCache::ReportNewPSOs()
 	return (bCmdLineForce || CVarPSOFileCacheReportPSO.GetValueOnAnyThread() == 1);
 }
 
-bool FPipelineFileCache::LogPSODetails()
+bool FPipelineFileCacheManager::LogPSODetails()
 {
     static bool bOnce = false;
     static bool bCmdLineOption = false;
@@ -3026,7 +3013,7 @@ bool FPipelineFileCache::LogPSODetails()
 	return bCmdLineOption;
 }
 
-void FPipelineFileCache::Initialize(uint32 InGameVersion)
+void FPipelineFileCacheManager::Initialize(uint32 InGameVersion)
 {
 	ClearOSPipelineCache();
 	
@@ -3038,12 +3025,12 @@ void FPipelineFileCache::Initialize(uint32 InGameVersion)
 		// Defaulting the CL is fine though
 		FPipelineCacheFile::GameVersion = (uint32)FEngineVersion::Current().GetChangelist();
 	}
-	
+
 	SET_MEMORY_STAT(STAT_NewCachedPSOMemory, 0);
 	SET_MEMORY_STAT(STAT_PSOStatMemory, 0);
 }
 
-bool FPipelineFileCache::ShouldEnableFileCache()
+bool FPipelineFileCacheManager::ShouldEnableFileCache()
 {
 #if PLATFORM_IOS
 	if (CVarAlwaysGeneratePOSSOFileCache.GetValueOnAnyThread() == 0)
@@ -3061,7 +3048,7 @@ bool FPipelineFileCache::ShouldEnableFileCache()
 	return GRHISupportsPipelineFileCache;
 }
 
-void FPipelineFileCache::PreCompileComplete()
+void FPipelineFileCacheManager::PreCompileComplete()
 {
 #if PLATFORM_IOS
 	// write out a file signifying we have completed a pre-compile of the PSO cache. Used on successive runs of the game to determine how much caching we need to still perform
@@ -3074,7 +3061,7 @@ void FPipelineFileCache::PreCompileComplete()
 #endif
 }
 
-void FPipelineFileCache::ClearOSPipelineCache()
+void FPipelineFileCacheManager::ClearOSPipelineCache()
 {
 	UE_LOG(LogTemp, Display, TEXT("Clearing the OS Cache"));
 	
@@ -3149,28 +3136,40 @@ void FPipelineFileCache::ClearOSPipelineCache()
 	}
 }
 
-uint64 FPipelineFileCache::SetGameUsageMaskWithComparison(uint64 InGameUsageMask, FPSOMaskComparisonFn InComparisonFnPtr)
+int32 FPipelineFileCacheManager::GetTotalPSOCount(const FString& PSOCacheKey)
+{
+	int32 TotalPSOs = 0;
+	if (IsPipelineFileCacheEnabled())
+	{
+		FRWScopeLock Lock(FileCacheLock, SLT_ReadOnly);
+		FPipelineCacheFile* Found = GetPipelineCacheFileFromKey(PSOCacheKey);
+		TotalPSOs = Found ? Found->GetTOCMetaDataSize() : 0;
+	}
+	return TotalPSOs;
+}
+
+uint64 FPipelineFileCacheManager::SetGameUsageMaskWithComparison(uint64 InGameUsageMask, FPSOMaskComparisonFn InComparisonFnPtr)
 {
 	uint64 OldMask = 0;
 	if(IsPipelineFileCacheEnabled())
 	{
 		FRWScopeLock Lock(FileCacheLock, SLT_Write);
 		
-		OldMask = FPipelineFileCache::GameUsageMask;
-		FPipelineFileCache::GameUsageMask = InGameUsageMask;
+		OldMask = FPipelineFileCacheManager::GameUsageMask;
+		FPipelineFileCacheManager::GameUsageMask = InGameUsageMask;
 		
 		if(InComparisonFnPtr == nullptr)
 		{
 			InComparisonFnPtr = DefaultPSOMaskComparisonFunction;
 		}
 		
-		FPipelineFileCache::MaskComparisonFn = InComparisonFnPtr;
+		FPipelineFileCacheManager::MaskComparisonFn = InComparisonFnPtr;
 	}
 	
 	return OldMask;
 }
 
-void FPipelineFileCache::Shutdown()
+void FPipelineFileCacheManager::Shutdown()
 {
 	if(IsPipelineFileCacheEnabled())
 	{
@@ -3191,7 +3190,13 @@ void FPipelineFileCache::Shutdown()
 	}
 }
 
-bool FPipelineFileCache::OpenPipelineFileCache(FString const& Name, EShaderPlatform Platform, FGuid& OutGameFileGuid)
+bool FPipelineFileCacheManager::HasPipelineFileCache(const FString & Key)
+{
+	FRWScopeLock Lock(FileCacheLock, SLT_ReadOnly);
+	return FileCacheMap.Contains(Key);
+}
+
+bool FPipelineFileCacheManager::OpenPipelineFileCache(const FString& Key, const FString& CacheName, EShaderPlatform Platform, FGuid& OutGameFileGuid)
 {
 	bool bOk = false;
 	OutGameFileGuid = FGuid();
@@ -3200,201 +3205,229 @@ bool FPipelineFileCache::OpenPipelineFileCache(FString const& Name, EShaderPlatf
 	{
 		FRWScopeLock Lock(FileCacheLock, SLT_Write);
 		
-		if(FileCache == nullptr)
+		bool bFound = FileCacheMap.Contains(Key);
+
+		if(!bFound)
 		{
-			FileCache = new FPipelineCacheFile();
-			
-			bOk = FileCache->OpenPipelineFileCache(Name, Platform, OutGameFileGuid);
+			TUniquePtr<FPipelineCacheFile> NewFileCache = MakeUnique<FPipelineCacheFile>();
+
+			bOk = NewFileCache->OpenPipelineFileCache(CacheName, Platform, OutGameFileGuid);
 			
-			if (!bOk && !LogPSOtoFileCache()) // don't delete the FileCache if we couldn't open and we are trying to log it as that leads to no file cache at all, NOTE: we might need to also check to see if there is a cache file on the system at all here
+			if (!bOk)
+			{
+				NewFileCache = nullptr;
+			}
+			else
 			{
-				delete FileCache;
-				FileCache = nullptr;
+				UE_LOG(LogRHI, Display, TEXT("FPipelineCacheFile[%s] opened %s, filename %s, guid %s. "), *Key, *CacheName, *NewFileCache->GetCacheFilename(), *NewFileCache->GetFileGuid().ToString());
+				FileCacheMap.Add(Key, MoveTemp(NewFileCache));
+				check(!GameGuidToCacheKey.Contains(OutGameFileGuid));
+				GameGuidToCacheKey.Add(OutGameFileGuid, Key);
 			}
-			// File Cache now exists - these caches should be empty for this file otherwise will have false positives from any previous file caching - if not something has been caching when it should not be
-			check(NewPSOs.Num() == 0);
-			check(NewPSOHashes.Num() == 0);
-			check(RunTimeToPSOUsage.Num() == 0);
 		}
 	}
 	
 	return bOk;
 }
 
-bool FPipelineFileCache::SavePipelineFileCache(FString const& Name, SaveMode Mode)
+bool FPipelineFileCacheManager::OpenUserPipelineFileCache(const FString& Key, const FString& CacheName, EShaderPlatform Platform, FGuid& OutGameFileGuid)
 {
-	bool bOk = false;
-	
-	if(IsPipelineFileCacheEnabled() && LogPSOtoFileCache())
+	// close any existing cache and reset the user cache's PSO recording containers.
+	CloseUserPipelineFileCache();
+
+	bool bUserFileOpened = false;
+	if (IsPipelineFileCacheEnabled())
 	{
-		CSV_EVENT(PSO, TEXT("Saving PSO cache"));
 		FRWScopeLock Lock(FileCacheLock, SLT_Write);
-		
-		if(FileCache)
+
+		UserCacheKey = Key;
+		FPipelineCacheFile* FileCache = GetPipelineCacheFileFromKey(UserCacheKey);
+		if(ensure(!FileCache))
 		{
-			FName PlatformName = FileCache->GetPlatformName();
-			FString Path = FPaths::ProjectSavedDir() / FString::Printf(TEXT("%s_%s.upipelinecache"), *Name, *PlatformName.ToString());
-			bOk = FileCache->SavePipelineFileCache(Path, Mode, Stats, NewPSOs, RequestedOrder, NewPSOUsage);
-			
-			// If successful clear new PSO's as they should have been saved out
-			// Leave everything else in-tact (e.g stats) for subsequent in place save operations
-			if (bOk)
-			{
-                NumNewPSOs = NewPSOs.Num();
-				SET_MEMORY_STAT(STAT_NewCachedPSOMemory, (NumNewPSOs * (sizeof(FPipelineCacheFileFormatPSO) + sizeof(uint32) + sizeof(uint32))));
-			}
+			TUniquePtr<FPipelineCacheFile> NewFileCache = MakeUnique<FPipelineCacheFile>();
+			FileCache = NewFileCache.Get();
+			bUserFileOpened = NewFileCache->OpenUserPipelineFileCache(CacheName, Platform);
+
+			// we always add the user cache, even if we did not open a file
+			FileCacheMap.Add(UserCacheKey, MoveTemp(NewFileCache));
+			check(!GameGuidToCacheKey.Contains(FileCache->GetFileGuid()));
+			GameGuidToCacheKey.Add(FileCache->GetFileGuid(), UserCacheKey);
 		}
+		OutGameFileGuid = FileCache->GetFileGuid();
+		UE_LOG(LogRHI, Display, TEXT("FPipelineCacheFile User cache [key:%s] opened '%s'=%d, filename %s, guid %s. "), *UserCacheKey, *CacheName, (int)bUserFileOpened, *FileCache->GetCacheFilename(), *FileCache->GetFileGuid().ToString());
+
+		// User Cache now exists - these caches should be empty for this file otherwise will have false positives from any previous file caching - if not something has been caching when it should not be
+		check(NewPSOs.Num() == 0);
+		check(NewPSOHashes.Num() == 0);
+		check(RunTimeToPSOUsage.Num() == 0);
 	}
-	
-	return bOk;
+	return bUserFileOpened;
 }
 
-void FPipelineFileCache::ClosePipelineFileCache()
+void FPipelineFileCacheManager::CloseUserPipelineFileCache()
 {
-	if(IsPipelineFileCacheEnabled())
+	if (IsPipelineFileCacheEnabled())
 	{
 		FRWScopeLock Lock(FileCacheLock, SLT_Write);
-		
-		if(FileCache)
+		if(GetPipelineCacheFileFromKey(UserCacheKey))
 		{
-			delete FileCache;
-			FileCache = nullptr;
-			
-			// Reset stats tracking for the next file.
-			for (auto const& Pair : Stats)
-			{
-				FPlatformAtomics::InterlockedExchange((int64*)&Pair.Value->TotalBindCount, -1);
-				FPlatformAtomics::InterlockedExchange((int64*)&Pair.Value->FirstFrameUsed, -1);
-				FPlatformAtomics::InterlockedExchange((int64*)&Pair.Value->LastFrameUsed, -1);
-			}
-			
-			// Reset serialized counts
-			SET_DWORD_STAT(STAT_SerializedGraphicsPipelineStateCount, 0);
-			SET_DWORD_STAT(STAT_SerializedComputePipelineStateCount, 0);
-			
-			// Not tracking when there is no file clear other stats as well
-			SET_DWORD_STAT(STAT_TotalGraphicsPipelineStateCount, 0);
-			SET_DWORD_STAT(STAT_TotalComputePipelineStateCount, 0);
-			SET_DWORD_STAT(STAT_TotalRayTracingPipelineStateCount, 0);
-			SET_DWORD_STAT(STAT_NewGraphicsPipelineStateCount, 0);
-			SET_DWORD_STAT(STAT_NewComputePipelineStateCount, 0);
-			SET_DWORD_STAT(STAT_NewRayTracingPipelineStateCount, 0);
+			const FGuid& UserGuid = GetPipelineCacheFileFromKey(UserCacheKey)->GetFileGuid();
+			GameGuidToCacheKey.Remove(UserGuid);
+			FileCacheMap.Remove(UserCacheKey);
+		}
 
-			// Clear Runtime hashes otherwise we can't start adding newPSO's for a newly opened file
-			RunTimeToPSOUsage.Empty();
-			NewPSOUsage.Empty();
-			NewPSOs.Empty();
-			NewPSOHashes.Empty();
-            NumNewPSOs = 0;
-			
-			SET_MEMORY_STAT(STAT_NewCachedPSOMemory, 0);
-			SET_MEMORY_STAT(STAT_FileCacheMemory, 0);
+		// Reset stats tracking for the next file.
+		for (auto const& Pair : Stats)
+		{
+			FPlatformAtomics::InterlockedExchange((int64*)&Pair.Value->TotalBindCount, -1);
+			FPlatformAtomics::InterlockedExchange((int64*)&Pair.Value->FirstFrameUsed, -1);
+			FPlatformAtomics::InterlockedExchange((int64*)&Pair.Value->LastFrameUsed, -1);
+		}
+
+		// Reset serialized counts
+		SET_DWORD_STAT(STAT_SerializedGraphicsPipelineStateCount, 0);
+		SET_DWORD_STAT(STAT_SerializedComputePipelineStateCount, 0);
+
+		// Not tracking when there is no file clear other stats as well
+		SET_DWORD_STAT(STAT_TotalGraphicsPipelineStateCount, 0);
+		SET_DWORD_STAT(STAT_TotalComputePipelineStateCount, 0);
+		SET_DWORD_STAT(STAT_TotalRayTracingPipelineStateCount, 0);
+		SET_DWORD_STAT(STAT_NewGraphicsPipelineStateCount, 0);
+		SET_DWORD_STAT(STAT_NewComputePipelineStateCount, 0);
+		SET_DWORD_STAT(STAT_NewRayTracingPipelineStateCount, 0);
+
+		// Clear Runtime hashes otherwise we can't start adding newPSO's for a newly opened file
+		RunTimeToPSOUsage.Empty();
+		NewPSOUsage.Empty();
+		NewPSOs.Empty();
+		NewPSOHashes.Empty();
+		NumNewPSOs = 0;
+
+		SET_MEMORY_STAT(STAT_NewCachedPSOMemory, 0);
+	}
+}
+
+bool FPipelineFileCacheManager::SavePipelineFileCache(SaveMode Mode)
+{
+	bool bOk = false;
+
+	if (IsPipelineFileCacheEnabled() && LogPSOtoFileCache())
+	{
+		CSV_EVENT(PSO, TEXT("Saving PSO cache"));
+		FRWScopeLock Lock(FileCacheLock, SLT_Write);
+
+		FPipelineCacheFile* UserCache = GetPipelineCacheFileFromKey(UserCacheKey);
+		check(UserCache);
+		bOk = UserCache->SavePipelineFileCache(Mode, Stats, NewPSOs, RequestedOrder, NewPSOUsage);
+		// If successful clear new PSO's as they should have been saved out
+		// Leave everything else in-tact (e.g stats) for subsequent in place save operations
+		if (bOk)
+		{
+			NumNewPSOs = NewPSOs.Num();
+			SET_MEMORY_STAT(STAT_NewCachedPSOMemory, (NumNewPSOs * (sizeof(FPipelineCacheFileFormatPSO) + sizeof(uint32) + sizeof(uint32))));
 		}
 	}
+	return bOk;
 }
 
-void FPipelineFileCache::RegisterPSOUsageDataUpdateForNextSave(FPSOUsageData& UsageData)
+void FPipelineFileCacheManager::RegisterPSOUsageDataUpdateForNextSave(FPSOUsageData& UsageData)
 {
 	FPSOUsageData& CurrentEntry = NewPSOUsage.FindOrAdd(UsageData.PSOHash);
 	CurrentEntry.PSOHash = UsageData.PSOHash;
 	CurrentEntry.UsageMask |= UsageData.UsageMask;
 	CurrentEntry.EngineFlags |= UsageData.EngineFlags;
-}
+}  
 
-void FPipelineFileCache::CacheGraphicsPSO(uint32 RunTimeHash, FGraphicsPipelineStateInitializer const& Initializer)
+void FPipelineFileCacheManager::CacheGraphicsPSO(uint32 RunTimeHash, FGraphicsPipelineStateInitializer const& Initializer)
 {
 	if(IsPipelineFileCacheEnabled() && (LogPSOtoFileCache() || ReportNewPSOs()))
 	{
 		FRWScopeLock Lock(FileCacheLock, SLT_ReadOnly);
-	
-		if(FileCache)
+
+		FPSOUsageData* PSOUsage = RunTimeToPSOUsage.Find(RunTimeHash);
+		if(PSOUsage == nullptr || !IsReferenceMaskSet(FPipelineFileCacheManager::GameUsageMask, PSOUsage->UsageMask))
 		{
-			FPSOUsageData* PSOUsage = RunTimeToPSOUsage.Find(RunTimeHash);
-			if(PSOUsage == nullptr || !IsReferenceMaskSet(FPipelineFileCache::GameUsageMask, PSOUsage->UsageMask))
-			{
-				Lock.ReleaseReadOnlyLockAndAcquireWriteLock_USE_WITH_CAUTION();
-				PSOUsage = RunTimeToPSOUsage.Find(RunTimeHash);
+			Lock.ReleaseReadOnlyLockAndAcquireWriteLock_USE_WITH_CAUTION();
+			PSOUsage = RunTimeToPSOUsage.Find(RunTimeHash);
 				
-				if(PSOUsage == nullptr)
-				{
-					FPipelineCacheFileFormatPSO NewEntry;
-					bool bOK = FPipelineCacheFileFormatPSO::Init(NewEntry, Initializer);
-					check(bOK);
+			if(PSOUsage == nullptr)
+			{
+				FPipelineCacheFileFormatPSO NewEntry;
+				bool bOK = FPipelineCacheFileFormatPSO::Init(NewEntry, Initializer);
+				check(bOK);
 					
-					uint32 PSOHash = GetTypeHash(NewEntry);
-					FPSOUsageData CurrentUsageData(PSOHash, 0, 0);
+				uint32 PSOHash = GetTypeHash(NewEntry);
+				FPSOUsageData CurrentUsageData(PSOHash, 0, 0);
 					
-					if (!FileCache->IsPSOEntryCached(NewEntry, &CurrentUsageData))
+				if (!FPipelineFileCacheManager::IsPSOEntryCached(NewEntry, &CurrentUsageData))
+				{
+					bool bActuallyNewPSO = !NewPSOHashes.Contains(PSOHash);
+					if (bActuallyNewPSO && IsOpenGLPlatform(GMaxRHIShaderPlatform)) // OpenGL is a BSS platform and so we don't report BSS matches as missing.
 					{
-						bool bActuallyNewPSO = !NewPSOHashes.Contains(PSOHash);
-						if (bActuallyNewPSO && IsOpenGLPlatform(GMaxRHIShaderPlatform)) // OpenGL is a BSS platform and so we don't report BSS matches as missing.
-						{
-							bActuallyNewPSO = !FileCache->IsBSSEquivalentPSOEntryCached(NewEntry);
-						}
-						if (bActuallyNewPSO)
+						bActuallyNewPSO = !FPipelineFileCacheManager::IsBSSEquivalentPSOEntryCached(NewEntry);
+					}
+					if (bActuallyNewPSO)
+					{
+						CSV_EVENT(PSO, TEXT("Encountered new graphics PSO"));
+						UE_LOG(LogRHI, Display, TEXT("Encountered a new graphics PSO: %u"), PSOHash);
+						if (GPSOFileCachePrintNewPSODescriptors > 0)
 						{
-							CSV_EVENT(PSO, TEXT("Encountered new graphics PSO"));
-							UE_LOG(LogRHI, Display, TEXT("Encountered a new graphics PSO: %u"), PSOHash);
-							if (GPSOFileCachePrintNewPSODescriptors > 0)
+							UE_LOG(LogRHI, Display, TEXT("New Graphics PSO (%u)"), PSOHash);
+							if (LogPSODetails() || GPSOFileCachePrintNewPSODescriptors > 1)
 							{
-								UE_LOG(LogRHI, Display, TEXT("New Graphics PSO (%u)"), PSOHash);
-								if (LogPSODetails() || GPSOFileCachePrintNewPSODescriptors > 1)
-								{
-									UE_LOG(LogRHI, Display, TEXT("%s"), *NewEntry.ToStringReadable());
-								}
-								else
-								{
-									UE_LOG(LogRHI, Display, TEXT("%s"), *NewEntry.GraphicsDesc.ToString());
-								}
+								UE_LOG(LogRHI, Display, TEXT("%s"), *NewEntry.ToStringReadable());
 							}
-							if (LogPSOtoFileCache())
+							else
 							{
-								NewPSOs.Add(NewEntry);
-								INC_MEMORY_STAT_BY(STAT_NewCachedPSOMemory, sizeof(FPipelineCacheFileFormatPSO) + sizeof(uint32) + sizeof(uint32));
+								UE_LOG(LogRHI, Display, TEXT("%s"), *NewEntry.GraphicsDesc.ToString());
 							}
-							NewPSOHashes.Add(PSOHash);
+						}
+						if (LogPSOtoFileCache())
+						{
+							NewPSOs.Add(NewEntry);
+							INC_MEMORY_STAT_BY(STAT_NewCachedPSOMemory, sizeof(FPipelineCacheFileFormatPSO) + sizeof(uint32) + sizeof(uint32));
+						}
+						NewPSOHashes.Add(PSOHash);
 
-							NumNewPSOs++;
-							INC_DWORD_STAT(STAT_NewGraphicsPipelineStateCount);
-							INC_DWORD_STAT(STAT_TotalGraphicsPipelineStateCount);
+						NumNewPSOs++;
+						INC_DWORD_STAT(STAT_NewGraphicsPipelineStateCount);
+						INC_DWORD_STAT(STAT_TotalGraphicsPipelineStateCount);
 							
-							if (ReportNewPSOs() && PSOLoggedEvent.IsBound())
-							{
-								PSOLoggedEvent.Broadcast(NewEntry);
-							}
+						if (ReportNewPSOs() && PSOLoggedEvent.IsBound())
+						{
+							PSOLoggedEvent.Broadcast(NewEntry);
 						}
 					}
-					
-					// Only set if the file cache doesn't have this Mask for the PSO - avoid making more entries and unnessary file saves
-					if(!IsReferenceMaskSet(FPipelineFileCache::GameUsageMask, CurrentUsageData.UsageMask))
-					{
-						CurrentUsageData.UsageMask |= FPipelineFileCache::GameUsageMask;
-						RegisterPSOUsageDataUpdateForNextSave(CurrentUsageData);
-					}
-					
-					// Apply the existing file PSO Usage mask and current to our "fast" runtime check
-					RunTimeToPSOUsage.Add(RunTimeHash, CurrentUsageData);
 				}
-				else if(!IsReferenceMaskSet(FPipelineFileCache::GameUsageMask, PSOUsage->UsageMask))
+					
+				// Only set if the file cache doesn't have this Mask for the PSO - avoid making more entries and unnessary file saves
+				if(!IsReferenceMaskSet(FPipelineFileCacheManager::GameUsageMask, CurrentUsageData.UsageMask))
 				{
-					PSOUsage->UsageMask |= FPipelineFileCache::GameUsageMask;
-					RegisterPSOUsageDataUpdateForNextSave(*PSOUsage);
+					CurrentUsageData.UsageMask |= FPipelineFileCacheManager::GameUsageMask;
+					RegisterPSOUsageDataUpdateForNextSave(CurrentUsageData);
 				}
+					
+				// Apply the existing file PSO Usage mask and current to our "fast" runtime check
+				RunTimeToPSOUsage.Add(RunTimeHash, CurrentUsageData);
+			}
+			else if(!IsReferenceMaskSet(FPipelineFileCacheManager::GameUsageMask, PSOUsage->UsageMask))
+			{
+				PSOUsage->UsageMask |= FPipelineFileCacheManager::GameUsageMask;
+				RegisterPSOUsageDataUpdateForNextSave(*PSOUsage);
 			}
 		}
 	}
 }
 
-void FPipelineFileCache::CacheComputePSO(uint32 RunTimeHash, FRHIComputeShader const* Initializer)
+void FPipelineFileCacheManager::CacheComputePSO(uint32 RunTimeHash, FRHIComputeShader const* Initializer)
 {
 	if(IsPipelineFileCacheEnabled() && (LogPSOtoFileCache() || ReportNewPSOs()))
 	{
 		FRWScopeLock Lock(FileCacheLock, SLT_ReadOnly);
 		
-		if(FileCache)
 		{
 			FPSOUsageData* PSOUsage = RunTimeToPSOUsage.Find(RunTimeHash);
-			if(PSOUsage == nullptr || !IsReferenceMaskSet(FPipelineFileCache::GameUsageMask, PSOUsage->UsageMask))
+			if(PSOUsage == nullptr || !IsReferenceMaskSet(FPipelineFileCacheManager::GameUsageMask, PSOUsage->UsageMask))
 			{
 				Lock.ReleaseReadOnlyLockAndAcquireWriteLock_USE_WITH_CAUTION();
 				PSOUsage = RunTimeToPSOUsage.Find(RunTimeHash);
@@ -3408,7 +3441,7 @@ void FPipelineFileCache::CacheComputePSO(uint32 RunTimeHash, FRHIComputeShader c
 					uint32 PSOHash = GetTypeHash(NewEntry);
 					FPSOUsageData CurrentUsageData(PSOHash, 0, 0);
 					
-					if (!FileCache->IsPSOEntryCached(NewEntry, &CurrentUsageData))
+					if (!FPipelineFileCacheManager::IsPSOEntryCached(NewEntry, &CurrentUsageData))
 					{
 						bool bActuallyNewPSO = !NewPSOHashes.Contains(PSOHash);
 						if (bActuallyNewPSO)
@@ -3440,18 +3473,18 @@ void FPipelineFileCache::CacheComputePSO(uint32 RunTimeHash, FRHIComputeShader c
 					}
 					
 					// Only set if the file cache doesn't have this Mask for the PSO - avoid making more entries and unnessary file saves
-					if(!IsReferenceMaskSet(FPipelineFileCache::GameUsageMask, CurrentUsageData.UsageMask))
+					if(!IsReferenceMaskSet(FPipelineFileCacheManager::GameUsageMask, CurrentUsageData.UsageMask))
 					{
-						CurrentUsageData.UsageMask |= FPipelineFileCache::GameUsageMask;
+						CurrentUsageData.UsageMask |= FPipelineFileCacheManager::GameUsageMask;
 						RegisterPSOUsageDataUpdateForNextSave(CurrentUsageData);
 					}
 					
 					// Apply the existing file PSO Usage mask and current to our "fast" runtime check
 					RunTimeToPSOUsage.Add(RunTimeHash, CurrentUsageData);
 				}
-				else if(!IsReferenceMaskSet(FPipelineFileCache::GameUsageMask, PSOUsage->UsageMask))
+				else if(!IsReferenceMaskSet(FPipelineFileCacheManager::GameUsageMask, PSOUsage->UsageMask))
 				{
-					PSOUsage->UsageMask |= FPipelineFileCache::GameUsageMask;
+					PSOUsage->UsageMask |= FPipelineFileCacheManager::GameUsageMask;
 					RegisterPSOUsageDataUpdateForNextSave(*PSOUsage);
 				}
 			}
@@ -3459,7 +3492,7 @@ void FPipelineFileCache::CacheComputePSO(uint32 RunTimeHash, FRHIComputeShader c
 	}
 }
 
-void FPipelineFileCache::CacheRayTracingPSO(const FRayTracingPipelineStateInitializer& Initializer)
+void FPipelineFileCacheManager::CacheRayTracingPSO(const FRayTracingPipelineStateInitializer& Initializer)
 {
 	if (!IsPipelineFileCacheEnabled() || !(LogPSOtoFileCache() || ReportNewPSOs()))
 	{
@@ -3476,11 +3509,6 @@ void FPipelineFileCache::CacheRayTracingPSO(const FRayTracingPipelineStateInitia
 
 	FRWScopeLock Lock(FileCacheLock, SLT_ReadOnly);
 
-	if (!FileCache)
-	{
-		return;
-	}
-
 	for (TArrayView<FRHIRayTracingShader*>& Table : ShaderTables)
 	{
 		for (FRHIRayTracingShader* Shader : Table)
@@ -3489,7 +3517,7 @@ void FPipelineFileCache::CacheRayTracingPSO(const FRayTracingPipelineStateInitia
 			uint32 RunTimeHash = GetTypeHash(Desc);
 
 			FPSOUsageData* PSOUsage = RunTimeToPSOUsage.Find(RunTimeHash);
-			if (PSOUsage == nullptr || !IsReferenceMaskSet(FPipelineFileCache::GameUsageMask, PSOUsage->UsageMask))
+			if (PSOUsage == nullptr || !IsReferenceMaskSet(FPipelineFileCacheManager::GameUsageMask, PSOUsage->UsageMask))
 			{
 				Lock.ReleaseReadOnlyLockAndAcquireWriteLock_USE_WITH_CAUTION();
 				PSOUsage = RunTimeToPSOUsage.Find(RunTimeHash);
@@ -3502,7 +3530,7 @@ void FPipelineFileCache::CacheRayTracingPSO(const FRayTracingPipelineStateInitia
 					uint32 PSOHash = GetTypeHash(NewEntry);
 					FPSOUsageData CurrentUsageData(PSOHash, 0, 0);
 
-					if (!FileCache->IsPSOEntryCached(NewEntry, &CurrentUsageData))
+					if (!FPipelineFileCacheManager::IsPSOEntryCached(NewEntry, &CurrentUsageData))
 					{
 						//CSV_EVENT(PSO, TEXT("Encountered new ray tracing PSO"));
 						UE_LOG(LogRHI, Display, TEXT("Encountered a new ray tracing PSO: %u"), PSOHash);
@@ -3527,9 +3555,9 @@ void FPipelineFileCache::CacheRayTracingPSO(const FRayTracingPipelineStateInitia
 					}
 
 					// Only set if the file cache doesn't have this Mask for the PSO - avoid making more entries and unnessary file saves
-					if (!IsReferenceMaskSet(FPipelineFileCache::GameUsageMask, CurrentUsageData.UsageMask))
+					if (!IsReferenceMaskSet(FPipelineFileCacheManager::GameUsageMask, CurrentUsageData.UsageMask))
 					{
-						CurrentUsageData.UsageMask |= FPipelineFileCache::GameUsageMask;
+						CurrentUsageData.UsageMask |= FPipelineFileCacheManager::GameUsageMask;
 						RegisterPSOUsageDataUpdateForNextSave(CurrentUsageData);
 					}
 
@@ -3551,99 +3579,95 @@ void FPipelineFileCache::CacheRayTracingPSO(const FRayTracingPipelineStateInitia
 					}
 				}
 			}
-			else if (!IsReferenceMaskSet(FPipelineFileCache::GameUsageMask, PSOUsage->UsageMask))
+			else if (!IsReferenceMaskSet(FPipelineFileCacheManager::GameUsageMask, PSOUsage->UsageMask))
 			{
-				PSOUsage->UsageMask |= FPipelineFileCache::GameUsageMask;
+				PSOUsage->UsageMask |= FPipelineFileCacheManager::GameUsageMask;
 				RegisterPSOUsageDataUpdateForNextSave(*PSOUsage);
 			}
 		}
 	}
 }
 
-void FPipelineFileCache::RegisterPSOCompileFailure(uint32 RunTimeHash, FGraphicsPipelineStateInitializer const& Initializer)
+void FPipelineFileCacheManager::RegisterPSOCompileFailure(uint32 RunTimeHash, FGraphicsPipelineStateInitializer const& Initializer)
 {
 	if(IsPipelineFileCacheEnabled() && (LogPSOtoFileCache() || ReportNewPSOs()) && Initializer.bFromPSOFileCache)
 	{
 		FRWScopeLock Lock(FileCacheLock, SLT_ReadOnly);
-		
-		if(FileCache)
+
+		FPSOUsageData* PSOUsage = RunTimeToPSOUsage.Find(RunTimeHash);
+		if(PSOUsage == nullptr || !IsReferenceMaskSet(FPipelineCacheFlagInvalidPSO, PSOUsage->EngineFlags))
 		{
-			FPSOUsageData* PSOUsage = RunTimeToPSOUsage.Find(RunTimeHash);
-			if(PSOUsage == nullptr || !IsReferenceMaskSet(FPipelineCacheFlagInvalidPSO, PSOUsage->EngineFlags))
-			{
-				Lock.ReleaseReadOnlyLockAndAcquireWriteLock_USE_WITH_CAUTION();
-				PSOUsage = RunTimeToPSOUsage.Find(RunTimeHash);
+			Lock.ReleaseReadOnlyLockAndAcquireWriteLock_USE_WITH_CAUTION();
+			PSOUsage = RunTimeToPSOUsage.Find(RunTimeHash);
 				
-				if(PSOUsage == nullptr)
-				{
-					FPipelineCacheFileFormatPSO ShouldBeExistingEntry;
-					bool bOK = FPipelineCacheFileFormatPSO::Init(ShouldBeExistingEntry, Initializer);
-					check(bOK);
+			if(PSOUsage == nullptr)
+			{
+				FPipelineCacheFileFormatPSO ShouldBeExistingEntry;
+				bool bOK = FPipelineCacheFileFormatPSO::Init(ShouldBeExistingEntry, Initializer);
+				check(bOK);
 					
-					uint32 PSOHash = GetTypeHash(ShouldBeExistingEntry);
-					FPSOUsageData CurrentUsageData(PSOHash, 0, 0);
+				uint32 PSOHash = GetTypeHash(ShouldBeExistingEntry);
+				FPSOUsageData CurrentUsageData(PSOHash, 0, 0);
 					
-					bool bCached = FileCache->IsPSOEntryCached(ShouldBeExistingEntry, &CurrentUsageData);
-					check(bCached);	//bFromPSOFileCache was set but not in the cache something has gone wrong
-					{
-						CurrentUsageData.EngineFlags |= FPipelineCacheFlagInvalidPSO;
+				bool bCached = FPipelineFileCacheManager::IsPSOEntryCached(ShouldBeExistingEntry, &CurrentUsageData);
+				check(bCached);	//bFromPSOFileCache was set but not in the cache something has gone wrong
+				{
+					CurrentUsageData.EngineFlags |= FPipelineCacheFlagInvalidPSO;
 						
-						RegisterPSOUsageDataUpdateForNextSave(CurrentUsageData);
-						RunTimeToPSOUsage.Add(RunTimeHash, CurrentUsageData);
+					RegisterPSOUsageDataUpdateForNextSave(CurrentUsageData);
+					RunTimeToPSOUsage.Add(RunTimeHash, CurrentUsageData);
 						
-						UE_LOG(LogRHI, Warning, TEXT("Graphics PSO (%u) compile failure registering to File Cache"), PSOHash);
-					}
+					UE_LOG(LogRHI, Warning, TEXT("Graphics PSO (%u) compile failure registering to File Cache"), PSOHash);
 				}
-				else if(!IsReferenceMaskSet(FPipelineCacheFlagInvalidPSO, PSOUsage->EngineFlags))
-				{
-					PSOUsage->EngineFlags |= FPipelineCacheFlagInvalidPSO;
-					RegisterPSOUsageDataUpdateForNextSave(*PSOUsage);
+			}
+			else if(!IsReferenceMaskSet(FPipelineCacheFlagInvalidPSO, PSOUsage->EngineFlags))
+			{
+				PSOUsage->EngineFlags |= FPipelineCacheFlagInvalidPSO;
+				RegisterPSOUsageDataUpdateForNextSave(*PSOUsage);
 					
-					UE_LOG(LogRHI, Warning, TEXT("Graphics PSO (%u) compile failure registering to File Cache"), PSOUsage->PSOHash);
-				}
+				UE_LOG(LogRHI, Warning, TEXT("Graphics PSO (%u) compile failure registering to File Cache"), PSOUsage->PSOHash);
 			}
 		}
 	}
 }
 
-FPipelineStateStats* FPipelineFileCache::RegisterPSOStats(uint32 RunTimeHash)
+FPipelineStateStats* FPipelineFileCacheManager::RegisterPSOStats(uint32 RunTimeHash)
 {
 	FPipelineStateStats* Stat = nullptr;
 	if(IsPipelineFileCacheEnabled() && LogPSOtoFileCache())
 	{
 		FRWScopeLock Lock(FileCacheLock, SLT_ReadOnly);
-		
-		if(FileCache)
+
+		uint32 PSOHash = RunTimeToPSOUsage.FindChecked(RunTimeHash).PSOHash;
+		Stat = Stats.FindRef(PSOHash);
+		if (!Stat)
 		{
-			uint32 PSOHash = RunTimeToPSOUsage.FindChecked(RunTimeHash).PSOHash;
+			Lock.ReleaseReadOnlyLockAndAcquireWriteLock_USE_WITH_CAUTION();
 			Stat = Stats.FindRef(PSOHash);
 			if (!Stat)
 			{
-				Lock.ReleaseReadOnlyLockAndAcquireWriteLock_USE_WITH_CAUTION();
-				Stat = Stats.FindRef(PSOHash);
-				if(!Stat)
-				{
-					Stat = new FPipelineStateStats;
-					Stat->PSOHash = PSOHash;
-					Stats.Add(PSOHash, Stat);
-					
-					INC_MEMORY_STAT_BY(STAT_PSOStatMemory, sizeof(FPipelineStateStats) + sizeof(uint32));
-				}
+				Stat = new FPipelineStateStats;
+				Stat->PSOHash = PSOHash;
+				Stats.Add(PSOHash, Stat);
+
+				INC_MEMORY_STAT_BY(STAT_PSOStatMemory, sizeof(FPipelineStateStats) + sizeof(uint32));
 			}
-			Stat->CreateCount++;
 		}
+		Stat->CreateCount++;
 	}
 	return Stat;
 }
 
-void FPipelineFileCache::GetOrderedPSOHashes(TArray<FPipelineCachePSOHeader>& PSOHashes, PSOOrder Order, int64 MinBindCount, TSet<uint32> const& AlreadyCompiledHashes)
+void FPipelineFileCacheManager::GetOrderedPSOHashes(const FString& PSOCacheKey, TArray<FPipelineCachePSOHeader>& PSOHashes, PSOOrder Order, int64 MinBindCount, TSet<uint32> const& AlreadyCompiledHashes)
 {
 	if(IsPipelineFileCacheEnabled())
 	{
 		FRWScopeLock Lock(FileCacheLock, SLT_Write);
 		
 		RequestedOrder = Order;
-		
+	
+ 		FPipelineCacheFile* FileCache = GetPipelineCacheFileFromKey(PSOCacheKey);
+
 		if(FileCache)
 		{
 			FileCache->GetOrderedPSOHashes(PSOHashes, Order, MinBindCount, AlreadyCompiledHashes);
@@ -3651,12 +3675,12 @@ void FPipelineFileCache::GetOrderedPSOHashes(TArray<FPipelineCachePSOHeader>& PS
 	}
 }
 
-void FPipelineFileCache::FetchPSODescriptors(TDoubleLinkedList<FPipelineCacheFileFormatPSORead*>& Batch)
+void FPipelineFileCacheManager::FetchPSODescriptors(const FString& PSOCacheKey, TDoubleLinkedList<FPipelineCacheFileFormatPSORead*>& Batch)
 {
 	if(IsPipelineFileCacheEnabled())
 	{
 		FRWScopeLock Lock(FileCacheLock, SLT_ReadOnly);
-		
+		FPipelineCacheFile* FileCache = GetPipelineCacheFileFromKey(PSOCacheKey);
 		if(FileCache)
 		{
 			FileCache->FetchPSODescriptors(Batch);
@@ -3771,6 +3795,8 @@ struct FPipelineCacheFileData
 						{
 							FPipelineCacheFileFormatPSO& PSO = Data.PSOs.FindChecked(Entry.Key);
 							PSOs.Add(GetTypeHash(PSO), PSO);
+							check(Entry.Value.FileGuid != FGuid());
+
 							MetaData.Add(GetTypeHash(PSO), Entry.Value);
 						}
 						
@@ -3794,7 +3820,7 @@ struct FPipelineCacheFileData
 	}
 };
 		 
-uint32 FPipelineFileCache::NumPSOsLogged()
+uint32 FPipelineFileCacheManager::NumPSOsLogged()
 {
 	uint32 Result = 0;
 	if(IsPipelineFileCacheEnabled() && LogPSOtoFileCache())
@@ -3832,12 +3858,12 @@ uint32 FPipelineFileCache::NumPSOsLogged()
 	return Result;
 }
 
-FPipelineFileCache::FPipelineStateLoggedEvent& FPipelineFileCache::OnPipelineStateLogged()
+FPipelineFileCacheManager::FPipelineStateLoggedEvent& FPipelineFileCacheManager::OnPipelineStateLogged()
 {
 	return PSOLoggedEvent;
 }
 
-bool FPipelineFileCache::LoadPipelineFileCacheInto(FString const& Path, TSet<FPipelineCacheFileFormatPSO>& PSOs)
+bool FPipelineFileCacheManager::LoadPipelineFileCacheInto(FString const& Path, TSet<FPipelineCacheFileFormatPSO>& PSOs)
 {
 	FPipelineCacheFileData A = FPipelineCacheFileData::Open(Path);
 	bool bAny = false;
@@ -3849,7 +3875,7 @@ bool FPipelineFileCache::LoadPipelineFileCacheInto(FString const& Path, TSet<FPi
 	return bAny;
 }
 
-bool FPipelineFileCache::SavePipelineFileCacheFrom(uint32 GameVersion, EShaderPlatform Platform, FString const& Path, const TSet<FPipelineCacheFileFormatPSO>& PSOs)
+bool FPipelineFileCacheManager::SavePipelineFileCacheFrom(uint32 GameVersion, EShaderPlatform Platform, FString const& Path, const TSet<FPipelineCacheFileFormatPSO>& PSOs)
 {
 	FPipelineCacheFileData Output;
 	Output.Header.Magic = FPipelineCacheFileFormatMagic;
@@ -3913,6 +3939,7 @@ bool FPipelineFileCache::SavePipelineFileCacheFrom(uint32 GameVersion, EShaderPl
 			}
 		}
 
+		check(Meta.FileGuid != FGuid());
 		Output.TOC.MetaData.Add(Meta.Stats.PSOHash, Meta);
 		Output.PSOs.Add(Meta.Stats.PSOHash, Item);
 	}
@@ -3966,7 +3993,7 @@ bool FPipelineFileCache::SavePipelineFileCacheFrom(uint32 GameVersion, EShaderPl
 }
 
 
-bool FPipelineFileCache::MergePipelineFileCaches(FString const& PathA, FString const& PathB, FPipelineFileCache::PSOOrder Order, FString const& OutputPath)
+bool FPipelineFileCacheManager::MergePipelineFileCaches(FString const& PathA, FString const& PathB, FPipelineFileCacheManager::PSOOrder Order, FString const& OutputPath)
 {
 	bool bOK = false;
 	
diff --git a/Engine/Source/Runtime/RHI/Private/PipelineStateCache.cpp b/Engine/Source/Runtime/RHI/Private/PipelineStateCache.cpp
--- a/Engine/Source/Runtime/RHI/Private/PipelineStateCache.cpp
+++ b/Engine/Source/Runtime/RHI/Private/PipelineStateCache.cpp
@@ -206,7 +206,7 @@ static void HandlePipelineCreationFailure(const FGraphicsPipelineStateInitialize
 	if(Init.bFromPSOFileCache)
 	{
 		// Let the cache know so it hopefully won't give out this one again
-		FPipelineFileCache::RegisterPSOCompileFailure(GetTypeHash(Init), Init);
+		FPipelineFileCacheManager::RegisterPSOCompileFailure(GetTypeHash(Init), Init);
 	}
 	else
 	{
@@ -1398,11 +1398,11 @@ FComputePipelineState* PipelineStateCache::GetAndOrCreateComputePipelineState(FR
 
 	if (WasFound == false)
 	{
-		FPipelineFileCache::CacheComputePSO(GetTypeHash(ComputeShader), ComputeShader);
+		FPipelineFileCacheManager::CacheComputePSO(GetTypeHash(ComputeShader), ComputeShader);
 
 		// create new graphics state
 		OutCachedState = new FComputePipelineState(ComputeShader);
-		OutCachedState->Stats = FPipelineFileCache::RegisterPSOStats(GetTypeHash(ComputeShader));
+		OutCachedState->Stats = FPipelineFileCacheManager::RegisterPSOStats(GetTypeHash(ComputeShader));
 
 		if (!bFromFileCache)
 		{
@@ -1449,7 +1449,7 @@ FComputePipelineState* PipelineStateCache::GetAndOrCreateComputePipelineState(FR
 	bool WasFound = GComputePipelineCache.FindOrAdd(ComputeShader, OutCachedState, [&RHICmdList, &ComputeShader, &DoAsyncCompile] {
 			// create new graphics state
 			TSharedPtr<FComputePipelineState> PipelineState(new FComputePipelineState(ComputeShader));
-			PipelineState->Stats = FPipelineFileCache::RegisterPSOStats(GetTypeHash(ComputeShader));
+			PipelineState->Stats = FPipelineFileCacheManager::RegisterPSOStats(GetTypeHash(ComputeShader));
 
 			// create a compilation task, or just do it now...
 			if (DoAsyncCompile)
@@ -1632,7 +1632,7 @@ FRayTracingPipelineState* PipelineStateCache::GetAndOrCreateRayTracingPipelineSt
 	}
 	else
 	{
-		FPipelineFileCache::CacheRayTracingPSO(InInitializer);
+		FPipelineFileCacheManager::CacheRayTracingPSO(InInitializer);
 
 		// Copy the initializer as we may want to patch it below
 		FRayTracingPipelineStateInitializer Initializer = InInitializer;
@@ -1799,11 +1799,11 @@ FGraphicsPipelineState* PipelineStateCache::GetAndOrCreateGraphicsPipelineState(
 
 	if (bWasFound == false)
 	{
-		FPipelineFileCache::CacheGraphicsPSO(GetTypeHash(Initializer), Initializer);
+		FPipelineFileCacheManager::CacheGraphicsPSO(GetTypeHash(Initializer), Initializer);
 
 		// create new graphics state
 		OutCachedState = new FGraphicsPipelineState();
-		OutCachedState->Stats = FPipelineFileCache::RegisterPSOStats(GetTypeHash(Initializer));
+		OutCachedState->Stats = FPipelineFileCacheManager::RegisterPSOStats(GetTypeHash(Initializer));
 
 		if (!Initializer.bFromPSOFileCache)
 		{
@@ -2006,7 +2006,7 @@ void PipelineStateCache::Shutdown()
 		
 		GGraphicsPipelineCache.DiscardAndSwap();
 	}
-	FPipelineFileCache::Shutdown();
+	FPipelineFileCacheManager::Shutdown();
 
 	for (auto Pair : GVertexDeclarationCache)
 	{
diff --git a/Engine/Source/Runtime/RHI/Public/PipelineFileCache.h b/Engine/Source/Runtime/RHI/Public/PipelineFileCache.h
--- a/Engine/Source/Runtime/RHI/Public/PipelineFileCache.h
+++ b/Engine/Source/Runtime/RHI/Public/PipelineFileCache.h
@@ -301,7 +301,7 @@ extern RHI_API const uint32 FPipelineCacheFileFormatCurrentVersion;
 typedef bool(*FPSOMaskComparisonFn)(uint64 ReferenceMask, uint64 PSOMask);
 
 /**
- * FPipelineFileCache:
+ * FPipelineFileCacheManager:
  * The RHI-level backend for FShaderPipelineCache, responsible for tracking PSOs and their usage stats as well as dealing with the pipeline cache files.
  * It is not expected that games or end-users invoke this directly, they should be calling FShaderPipelineCache which exposes this functionality in a usable form. 
  */
@@ -315,7 +315,7 @@ struct FPSOUsageData
 	uint16 EngineFlags;
 };
 
-class RHI_API FPipelineFileCache
+class RHI_API FPipelineFileCacheManager
 {
     friend class FPipelineCacheFile;
 public:
@@ -323,7 +323,6 @@ public:
 	{
 		Incremental = 0, // Fast(er) approach which saves new entries incrementally at the end of the file, replacing the table-of-contents, but leaves everything else alone.
         BoundPSOsOnly = 1, // Slower approach which consolidates and saves all PSOs used in this run of the program, removing any entry that wasn't seen, and sorted by the desired sort-mode.
-		SortedBoundPSOs = 2 // Slow save consolidates all PSOs used on this device that were never part of a cache file delivered in game-content, sorts entries into the desired order and will thus read-back from disk.
 	};
 	
 	enum class PSOOrder : uint32
@@ -340,17 +339,28 @@ public:
 	
 	static bool LoadPipelineFileCacheInto(FString const& Path, TSet<FPipelineCacheFileFormatPSO>& PSOs);
 	static bool SavePipelineFileCacheFrom(uint32 GameVersion, EShaderPlatform Platform, FString const& Path, const TSet<FPipelineCacheFileFormatPSO>& PSOs);
-	static bool MergePipelineFileCaches(FString const& PathA, FString const& PathB, FPipelineFileCache::PSOOrder Order, FString const& OutputPath);
-
+	static bool MergePipelineFileCaches(FString const& PathA, FString const& PathB, FPipelineFileCacheManager::PSOOrder Order, FString const& OutputPath);
+																				
 	/* Open the pipeline file cache for the specfied name and platform. If successful, the GUID of the game file will be returned in OutGameFileGuid */
-	static bool OpenPipelineFileCache(FString const& Name, EShaderPlatform Platform, FGuid& OutGameFileGuid);
-	static bool SavePipelineFileCache(FString const& Name, SaveMode Mode);
-	static void ClosePipelineFileCache();
-	
+	static bool OpenPipelineFileCache(const FString& Key, const FString& CacheName, EShaderPlatform Platform, FGuid& OutGameFileGuid);
+
+	/* Open the user pipeline file cache for the specified name and platform. The user cache is always created even if the file was not present when opened.
+	* Name is the name used when opening the file, the key value for the user cache is held within UserCacheName.
+	* returns true if the file was opened.
+	*/
+	static bool OpenUserPipelineFileCache(const FString& Key, const FString& CacheName, EShaderPlatform Platform, FGuid& OutGameFileGuid);
+
+	static bool SavePipelineFileCache(SaveMode Mode);
+
+	static void CloseUserPipelineFileCache();
+
 	static void CacheGraphicsPSO(uint32 RunTimeHash, FGraphicsPipelineStateInitializer const& Initializer);
 	static void CacheComputePSO(uint32 RunTimeHash, FRHIComputeShader const* Initializer);
 	static void CacheRayTracingPSO(const FRayTracingPipelineStateInitializer& Initializer);
 
+	// true if the named PSOFC is currently open.
+	static bool HasPipelineFileCache(const FString& PSOCacheKey);
+
 	static FPipelineStateStats* RegisterPSOStats(uint32 RunTimeHash);
 	
 	/*
@@ -369,8 +379,11 @@ public:
 	 */
 	static FPipelineStateLoggedEvent& OnPipelineStateLogged();
 	
-	static void GetOrderedPSOHashes(TArray<FPipelineCachePSOHeader>& PSOHashes, PSOOrder Order, int64 MinBindCount, TSet<uint32> const& AlreadyCompiledHashes);
-	static void FetchPSODescriptors(TDoubleLinkedList<FPipelineCacheFileFormatPSORead*>& LoadedBatch);
+	static void GetOrderedPSOHashes(const FString& PSOCacheKey, TArray<FPipelineCachePSOHeader>& PSOHashes, PSOOrder Order, int64 MinBindCount, TSet<uint32> const& AlreadyCompiledHashes);
+	static void FetchPSODescriptors(const FString& PSOCacheKey, TDoubleLinkedList<FPipelineCacheFileFormatPSORead*>& LoadedBatch);
+
+	static int32 GetTotalPSOCount(const FString& PSOCacheKey);
+
 	static uint32 NumPSOsLogged();
 	
 	static bool IsPipelineFileCacheEnabled();
@@ -395,15 +408,32 @@ private:
 	static void ClearOSPipelineCache();
 	static bool ShouldEnableFileCache();
 	
+	static bool IsBSSEquivalentPSOEntryCached(FPipelineCacheFileFormatPSO const& NewEntry);
+	static bool IsPSOEntryCached(FPipelineCacheFileFormatPSO const& NewEntry, FPSOUsageData* EntryData = nullptr);
 private:
 	static FRWLock FileCacheLock;
-	static class FPipelineCacheFile* FileCache;
+
+	// Containers for the multiple bundled PSOFCs
+	// Name to PipelineCacheFile
+	static TMap<FString, TUniquePtr<class FPipelineCacheFile>> FileCacheMap;
+	// PipelineCacheFile GUID to Name
+	static TMap<FGuid, FString> GameGuidToCacheKey;
+	// User cache's key within FileCacheMap
+	static FString UserCacheKey;
+	// Helper for retrieving a file cache from the name.
+	static class FPipelineCacheFile* GetPipelineCacheFileFromKey(const FString& PSOCacheKey)
+	{
+		TUniquePtr<FPipelineCacheFile>* FileCacheFound = FileCacheMap.Find(PSOCacheKey);
+		return FileCacheFound ? FileCacheFound->Get() : nullptr;
+	}
+
+	// PSO recording
 	static TMap<uint32, FPSOUsageData> RunTimeToPSOUsage;		// Fast check structure - Not saved (External state cache runtime hash to seen usage data)
 	static TMap<uint32, FPSOUsageData> NewPSOUsage;				// For mask or engine updates - Merged + Saved (Our internal PSO hash to latest usage data) - temp working scratch, only holds updates since last "save" so is not the authority on state
 	static TMap<uint32, FPipelineStateStats*> Stats;
 	static TSet<FPipelineCacheFileFormatPSO> NewPSOs;
-	static TSet<uint32> NewPSOHashes;
-    static uint32 NumNewPSOs;
+ 	static TSet<uint32> NewPSOHashes;
+	static uint32 NumNewPSOs;
 	static PSOOrder RequestedOrder;
 	static bool FileCacheEnabled;
 	static FPipelineStateLoggedEvent PSOLoggedEvent;
diff --git a/Engine/Source/Runtime/RenderCore/Private/ShaderCodeLibrary.cpp b/Engine/Source/Runtime/RenderCore/Private/ShaderCodeLibrary.cpp
--- a/Engine/Source/Runtime/RenderCore/Private/ShaderCodeLibrary.cpp
+++ b/Engine/Source/Runtime/RenderCore/Private/ShaderCodeLibrary.cpp
@@ -1989,6 +1989,7 @@ public:
 
 		bool bResult = false;
 
+		TSet<int32> NewComponentIDs;
 		if (IsLibraryInitializedForRuntime())
 		{
 			LLM_SCOPE(ELLMTag::Shaders);
@@ -2026,8 +2027,7 @@ public:
 			}
 			else // attempt to open a chunked library
 			{
-				int32 PrevNumComponents = Library->GetNumComponents();
-
+				TSet<int32> PrevComponentSet = Library->PresentChunks;
 				{
 					FScopeLock KnownPakFilesLocker(&FMountedPakFileInfo::KnownPakFilesAccessLock);
 					for (TSet<FMountedPakFileInfo>::TConstIterator Iter(FMountedPakFileInfo::KnownPakFiles); Iter; ++Iter)
@@ -2035,8 +2035,9 @@ public:
 						Library->OnPakFileMounted(*Iter);
 					}
 				}
-
-				bResult = (Library->GetNumComponents() > PrevNumComponents);
+				NewComponentIDs = Library->PresentChunks.Difference(PrevComponentSet);
+				bResult = !NewComponentIDs.IsEmpty();
+				UE_LOG(LogShaderLibrary, Display, TEXT("Logical shader library '%s' component count %d, new components: %d"), *Name, Library->GetNumComponents(), NewComponentIDs.Num());
 
 #if UE_SHADERLIB_SUPPORT_CHUNK_DISCOVERY
 				if (!bResult)
@@ -2072,7 +2073,8 @@ public:
 								}
 							}
 
-							bResult = (Library->GetNumComponents() > PrevNumComponents);
+							NewComponentIDs = Library->PresentChunks.Difference(PrevComponentSet);
+							bResult = !NewComponentIDs.IsEmpty();
 						}
 						else
 						{
@@ -2106,7 +2108,11 @@ public:
 				// Inform the pipeline cache that the state of loaded libraries has changed (unless we had to delete the duplicate)
 				if (Library != nullptr)
 				{
-					FShaderPipelineCache::ShaderLibraryStateChanged(FShaderPipelineCache::Opened, ShaderPlatform, Name);
+					FShaderPipelineCache::ShaderLibraryStateChanged(FShaderPipelineCache::Opened, ShaderPlatform, Name, INDEX_NONE);
+					for (int32 ComponentID : NewComponentIDs)
+					{
+						FShaderPipelineCache::ShaderLibraryStateChanged(FShaderPipelineCache::OpenedComponent, ShaderPlatform, Name, ComponentID);
+					}
 				}
 			}
 		}
@@ -2148,7 +2154,7 @@ public:
 		}
 
 		// Inform the pipeline cache that the state of loaded libraries has changed
-		FShaderPipelineCache::ShaderLibraryStateChanged(FShaderPipelineCache::Closed, ShaderPlatform, Name);
+		FShaderPipelineCache::ShaderLibraryStateChanged(FShaderPipelineCache::Closed, ShaderPlatform, Name, -1);
 
 #if WITH_EDITOR
 		for (uint32 i = 0; i < EShaderPlatform::SP_NumPlatforms; i++)
@@ -2170,10 +2176,25 @@ public:
 	{		
 		if (IsLibraryInitializedForRuntime())
 		{
-			FRWScopeLock WriteLock(NamedLibrariesMutex, SLT_Write);
-			for (TTuple<FString, TUniquePtr<UE::ShaderLibrary::Private::FNamedShaderLibrary>>& NamedLibraryPair : NamedLibrariesStack)
+			TArray<TUniqueFunction<void()>> ShaderLibraryStateChanges;
+			{
+				FRWScopeLock WriteLock(NamedLibrariesMutex, SLT_Write);
+				for (TTuple<FString, TUniquePtr<UE::ShaderLibrary::Private::FNamedShaderLibrary>>& NamedLibraryPair : NamedLibrariesStack)
+				{
+					TSet<int32> PrevComponentSet = NamedLibraryPair.Value->PresentChunks;
+					NamedLibraryPair.Value->OnPakFileMounted(MountInfo);
+					for (int32 ComponentID : NamedLibraryPair.Value->PresentChunks.Difference(PrevComponentSet))
+					{
+						// Defer these for outside of the lock, ShaderPipelineCache may want to inspect shader library.
+						ShaderLibraryStateChanges.Add([ShaderPlatform = NamedLibraryPair.Value->ShaderPlatform, LogicalName = NamedLibraryPair.Value->LogicalName, ComponentID]() {
+							FShaderPipelineCache::ShaderLibraryStateChanged(FShaderPipelineCache::OpenedComponent, ShaderPlatform, LogicalName, ComponentID);
+						});
+					}
+				}
+			}
+			for (TUniqueFunction<void()>& OnShaderLibraryStateChange : ShaderLibraryStateChanges)
 			{
-				NamedLibraryPair.Value->OnPakFileMounted(MountInfo);
+				OnShaderLibraryStateChange();
 			}
 		}
 	}
diff --git a/Engine/Source/Runtime/RenderCore/Private/ShaderPipelineCache.cpp b/Engine/Source/Runtime/RenderCore/Private/ShaderPipelineCache.cpp
--- a/Engine/Source/Runtime/RenderCore/Private/ShaderPipelineCache.cpp
+++ b/Engine/Source/Runtime/RenderCore/Private/ShaderPipelineCache.cpp
@@ -23,6 +23,7 @@
 #include "Misc/ConfigCacheIni.h"
 #include "Async/AsyncFileHandle.h"
 #include "Misc/ScopeLock.h"
+#include <Algo/RemoveIf.h>
 
 DECLARE_DWORD_ACCUMULATOR_STAT(TEXT("Outstanding Tasks"), STAT_ShaderPipelineTaskCount, STATGROUP_PipelineStateCache );
 DECLARE_DWORD_ACCUMULATOR_STAT(TEXT("Waiting Tasks"), STAT_ShaderPipelineWaitingTaskCount, STATGROUP_PipelineStateCache );
@@ -38,7 +39,8 @@ DECLARE_DWORD_COUNTER_STAT(TEXT("# Batches Pre-Compiled"), STAT_PreCompileBatchN
 namespace FShaderPipelineCacheConstants
 {
 	static TCHAR const* SectionHeading = TEXT("ShaderPipelineCache.CacheFile");
-	static TCHAR const* LastOpenedKey = TEXT("LastOpened");
+	static TCHAR const* CacheFileNameOverride = TEXT("CacheFileNameOverride");
+	static TCHAR const* UserCacheFileNameOverride = TEXT("UserCacheFileNameOverride");
 	static TCHAR const* SortOrderKey = TEXT("SortOrder");
 	static TCHAR const* GameVersionKey = TEXT("GameVersion");
 }
@@ -172,34 +174,29 @@ static bool GetPSOFileCacheSaveUserCache()
 	return GRHISupportsPipelineFileCache && (CVarPSOFileCacheSaveUserCache && CVarPSOFileCacheSaveUserCache->GetInt() > 0);
 }
 
-void ConsoleCommandLoadPipelineFileCache(const TArray< FString >& Args)
+void ConsoleCommandLoadPipelineFileCache()
 {
-	FShaderPipelineCache::ClosePipelineFileCache();
-	FString Name = FApp::GetProjectName();
-	if (Args.Num() > 0)
-	{
-		Name = Args[0];
-	}
-	FShaderPipelineCache::OpenPipelineFileCache(Name, GMaxRHIShaderPlatform);
+	FShaderPipelineCache::CloseUserPipelineFileCache();
+	FShaderPipelineCache::OpenUserPipelineFileCache(GMaxRHIShaderPlatform);
+}
+
+void ConsoleCommandClosePipelineFileCache()
+{
+	FShaderPipelineCache::CloseUserPipelineFileCache();
 }
 
 void ConsoleCommandSavePipelineFileCache()
 {
 	if (GetShaderPipelineCacheSaveBoundPSOLog())
 	{
-		FShaderPipelineCache::SavePipelineFileCache(FPipelineFileCache::SaveMode::BoundPSOsOnly);
+		FShaderPipelineCache::SavePipelineFileCache(FPipelineFileCacheManager::SaveMode::BoundPSOsOnly);
 	}
 	if (GetPSOFileCacheSaveUserCache())
 	{
-		FShaderPipelineCache::SavePipelineFileCache(FPipelineFileCache::SaveMode::SortedBoundPSOs);
+		FShaderPipelineCache::SavePipelineFileCache(FPipelineFileCacheManager::SaveMode::Incremental);
 	}
 }
 
-void ConsoleCommandClosePipelineFileCache()
-{
-	FShaderPipelineCache::ClosePipelineFileCache();
-}
-
 void ConsoleCommandSwitchModePipelineCacheCmd(const TArray< FString >& Args)
 {
     if (Args.Num() > 0)
@@ -229,8 +226,8 @@ void ConsoleCommandSwitchModePipelineCacheCmd(const TArray< FString >& Args)
 
 static FAutoConsoleCommand LoadPipelineCacheCmd(
 												TEXT("r.ShaderPipelineCache.Open"),
-												TEXT("Takes the desired filename to open and then loads the pipeline file cache."),
-												FConsoleCommandWithArgsDelegate::CreateStatic(ConsoleCommandLoadPipelineFileCache)
+												TEXT("Close and reopen the user cache."),
+												FConsoleCommandDelegate::CreateStatic(ConsoleCommandLoadPipelineFileCache)
 												);
 
 static FAutoConsoleCommand SavePipelineCacheCmd(
@@ -259,6 +256,9 @@ static FAutoConsoleVariableRef CVarShaderPipelineCacheDoNotPrecompileComputePSO(
 												ECVF_Default
 												);
 
+uint32 FShaderPipelineCache::BatchSize = 0;
+float FShaderPipelineCache::BatchTime = 0.0f;
+
 class FShaderPipelineCacheArchive final : public FArchive
 {
 public:
@@ -309,7 +309,300 @@ private:
 	TArray<FExternalReadCallback> ExternalReadDependencies;
 };
 
+
+namespace UE
+{
+	namespace ShaderPipeline
+	{
+
+		FString LibNameToPSOFCName(FString const& Name, int32 ComponentID)
+		{
+			if (ComponentID == -1)
+			{
+				return Name;
+			}
+			return FString::Printf(TEXT("%s_Chunk%d"), *Name, ComponentID);
+		}
+
+		bool ShouldLoadUserCache()
+		{
+			return GetShaderPipelineCacheSaveBoundPSOLog() || GetPSOFileCacheSaveUserCache();
+		}
+
+
+		FString GetShaderPipelineBaseName(bool bWantUserCacheName)
+		{
+			FString Name = FApp::GetProjectName();
+
+			const TCHAR* OverrideFilenameParam = bWantUserCacheName ? FShaderPipelineCacheConstants::UserCacheFileNameOverride : FShaderPipelineCacheConstants::CacheFileNameOverride;
+			FString OverrideName;
+			if (GConfig && (
+					(GConfig->GetString(FShaderPipelineCacheConstants::SectionHeading, OverrideFilenameParam, OverrideName, *GGameUserSettingsIni) || GConfig->GetString(FShaderPipelineCacheConstants::SectionHeading, OverrideFilenameParam, OverrideName, *GGameIni))
+					&& OverrideName.Len())
+				)
+			{
+				Name = OverrideName;
+			}
+
+			return Name;
+		}
+
+
+		struct CompileJob
+		{
+			FPipelineCacheFileFormatPSO PSO;
+			FShaderPipelineCacheArchive* ReadRequests;
+
+			/** Tracks whether the shaders were preloaded. */
+			bool bShadersPreloaded = false;
+
+			/**
+			 * Schedules preload for all the shaders in the PSO. No shaders are preloaded if one of them isn't in the library.
+			 * Multiple jobs can request to preload the same shaders and then release them independently.
+			 *
+			 * @param OutRequiredShaders shader hashes of shaders that are actually required by this PSO
+			 * @param bOutCompatible if set to false, then this PSO type isn't supported (RTX as of now)
+			 */
+			bool PreloadShaders(TSet<FSHAHash>& OutRequiredShaders, bool& bOutCompatible);
+
+			/**
+			 * Releases preloaded entries for the shaders in the PSO. Symmetric to to PreloadShaders.
+			 * Multiple jobs can request to preload the same shaders and then release them independently.
+			 */
+			void ReleasePreloadedShaders();
+		};
+	}
+}
+
+/* 
+* FShaderPipelineCacheTask is encapsulates the precompile process for a single PSO file cache.
+*/
+class FShaderPipelineCacheTask
+{
+	friend class FShaderPipelineCache;
+
+public:
+	FShaderPipelineCacheTask(bool bIsUserCacheIn, EShaderPlatform InShaderPlatform) :
+		ShaderPlatform(InShaderPlatform), bIsUserCache(bIsUserCacheIn)
+	{
+		bReady = !CVarPSOFileCacheGameFileMaskEnabled.GetValueOnAnyThread() || CVarPSOFileCachePreOptimizeEnabled.GetValueOnAnyThread();
+	}
+
+	bool Open(const FString& Key, const FString& Name, EShaderPlatform Platform);
+	
+	void Close();
+	void Shutdown();
+
+	void BeginPrecompilingPipelineCache();
+
+	bool Precompile(FRHICommandListImmediate& RHICmdList, EShaderPlatform Platform, FPipelineCacheFileFormatPSO const& PSO);
+	void PreparePipelineBatch(TDoubleLinkedList<FPipelineCacheFileFormatPSORead*>& PipelineBatch);
+	bool ReadyForPrecompile();
+	void PrecompilePipelineBatch();
+	bool ReadyForNextBatch() const;
+	bool ReadyForAutoSave() const;
+	void PollShutdownItems();
+	void Flush();
+	void Tick(float DeltaTime);
+
+	bool HasRHIFenceCompleted() 
+	{
+		if (LastPrecompileRHIFence.GetReference() && LastPrecompileRHIFence->IsComplete())
+		{
+			LastPrecompileRHIFence = nullptr;
+		}
+		return LastPrecompileRHIFence.GetReference()==nullptr;
+	}
+	bool OnShaderLibraryStateChanged(FShaderPipelineCache::ELibraryState State, EShaderPlatform Platform, FString const& Name, int32 ComponentID);
+private:
+	FString PSOCacheKey;
+	int32 TotalPSOsCompiled = 0;
+	EShaderPlatform ShaderPlatform = EShaderPlatform::SP_NumPlatforms;
+	TArray<UE::ShaderPipeline::CompileJob> ReadTasks;
+	TArray<UE::ShaderPipeline::CompileJob> CompileTasks;
+	TArray<FPipelineCachePSOHeader> OrderedCompileTasks;
+	TDoubleLinkedList<FPipelineCacheFileFormatPSORead*> FetchTasks;
+	FGuid CacheFileGuid;
+
+	FGraphEventRef LastPrecompileRHIFence;
+
+	TSet<uint64> CompletedMasks;
+	FShaderPipelineCache::FShaderCachePrecompileContext PrecompileContext;
+	bool bPreOptimizing = false;
+	TArray<FPipelineCachePSOHeader> PreFetchedTasks;
+	bool bReady = false;
+	bool bOpened = false;
+	bool bIsUserCache;
+
+	// as there is only ever one active precompile task at a time
+	// these stats are available to read anywhere.
+	static float TotalPrecompileWallTime;
+	static int64 TotalPrecompileTasks;
+
+	static std::atomic_int64_t TotalActiveTasks;
+	static std::atomic_int64_t TotalWaitingTasks;
+	static std::atomic_int64_t TotalCompleteTasks;
+	static std::atomic_int64_t TotalPrecompileTime;
+
+	static double PrecompileStartTime;
+};
+
+float FShaderPipelineCacheTask::TotalPrecompileWallTime = 0.0f;
+int64 FShaderPipelineCacheTask::TotalPrecompileTasks = 0;
+std::atomic_int64_t FShaderPipelineCacheTask::TotalActiveTasks = 0;
+std::atomic_int64_t FShaderPipelineCacheTask::TotalWaitingTasks = 0;
+std::atomic_int64_t FShaderPipelineCacheTask::TotalCompleteTasks = 0;
+std::atomic_int64_t FShaderPipelineCacheTask::TotalPrecompileTime = 0;
+double FShaderPipelineCacheTask::PrecompileStartTime = 0.0;
+
+
+
+namespace UE
+{
+	namespace ShaderPipeline
+	{
+		// The shutdown container takes abandoned tasks off the PSO caches.
+		// PSO Cache completion does not require these tasks to have completed.
+		class FShutdownContainer
+		{
+		public:
+			static TUniquePtr<class FShutdownContainer> Instance;
+			std::atomic_int TotalRemainingTasks = 0;
+			TArray<UE::ShaderPipeline::CompileJob> ShutdownReadCompileTasks;
+			TDoubleLinkedList<FPipelineCacheFileFormatPSORead*> ShutdownFetchTasks;
+
+			static TUniquePtr<class FShutdownContainer>& Get()
+			{
+				if (!Instance.IsValid())
+				{
+					Instance = MakeUnique<FShutdownContainer>();
+				}
+				return Instance;
+			}
+
+			void AbandonCompileJobs(TArray<UE::ShaderPipeline::CompileJob>&& AbandonedCompileJobs)
+			{
+				TotalRemainingTasks += AbandonedCompileJobs.Num();
+				ShutdownReadCompileTasks.Append(MoveTemp(AbandonedCompileJobs));
+			}
+
+			void AbandonCompileJob(UE::ShaderPipeline::CompileJob&& AbandonedCompileJob)
+			{
+				ShutdownReadCompileTasks.Emplace(AbandonedCompileJob);
+				TotalRemainingTasks++;
+			}
+
+			void AbandonFetchTasks(TDoubleLinkedList<FPipelineCacheFileFormatPSORead*>&& AbandonedFetchTasks)
+			{
+				// Marshall the current fetch tasks into shutdown
+				for (TDoubleLinkedList<FPipelineCacheFileFormatPSORead*>::TIterator It(AbandonedFetchTasks.GetHead()); It; ++It)
+				{
+					FPipelineCacheFileFormatPSORead* Entry = *It;
+					check(Entry->ReadRequest.IsValid());
+					Entry->ReadRequest->Cancel();
+					ShutdownFetchTasks.AddTail(Entry);
+					TotalRemainingTasks++;
+				}
+
+			}
+
+			bool IsTickable() const
+			{
+				return TotalRemainingTasks > 0;
+			}
+
+			void Tick()
+			{
+				int64 RemovedTaskCount = 0;
+
+				const int64 InitialCompileTaskCount = ShutdownReadCompileTasks.Num();
+				if (ShutdownReadCompileTasks.Num() > 0)
+				{
+					for (uint32 i = 0; i < (uint32)ShutdownReadCompileTasks.Num(); )
+					{
+						check(ShutdownReadCompileTasks[i].ReadRequests);
+						if (ShutdownReadCompileTasks[i].ReadRequests->PollExternalReadDependencies())
+						{
+							ShutdownReadCompileTasks[i].ReleasePreloadedShaders();
+							delete ShutdownReadCompileTasks[i].ReadRequests;
+							ShutdownReadCompileTasks[i].ReadRequests = nullptr;
+
+							ShutdownReadCompileTasks.RemoveAtSwap(i, 1, false);
+							++RemovedTaskCount;
+						}
+						else
+						{
+							++i;
+						}
+					}
+
+					if (ShutdownReadCompileTasks.Num() == 0)
+					{
+						ShutdownReadCompileTasks.Shrink();
+					}
+				}
+
+				if (ShutdownFetchTasks.Num() > 0)
+				{
+					TDoubleLinkedList<FPipelineCacheFileFormatPSORead*>::TDoubleLinkedListNode* CurrentNode = ShutdownFetchTasks.GetHead();
+					while (CurrentNode)
+					{
+						FPipelineCacheFileFormatPSORead* PSORead = CurrentNode->GetValue();
+						check(PSORead);
+						FShaderPipelineCacheArchive* Archive = (FShaderPipelineCacheArchive*)(PSORead->Ar);
+						check(Archive);
+
+						TDoubleLinkedList<FPipelineCacheFileFormatPSORead*>::TDoubleLinkedListNode* PrevNode = CurrentNode;
+						CurrentNode = CurrentNode->GetNextNode();
+
+						if (PSORead->bReadCompleted || Archive->PollExternalReadDependencies())
+						{
+							delete PSORead;
+							ShutdownFetchTasks.RemoveNode(PrevNode);
+							++RemovedTaskCount;
+						}
+					}
+				}
+
+				if (RemovedTaskCount > 0)
+				{
+					TotalRemainingTasks -= RemovedTaskCount;
+				}
+			}
+
+			void ShutdownAndWait()
+			{
+				for (UE::ShaderPipeline::CompileJob& Entry : ShutdownReadCompileTasks)
+				{
+					if (Entry.ReadRequests != nullptr)
+					{
+						Entry.ReadRequests->BlockingWaitComplete();
+						Entry.ReleasePreloadedShaders();
+						delete Entry.ReadRequests;
+						Entry.ReadRequests = nullptr;
+					}
+				}
+				ShutdownReadCompileTasks.Empty();
+
+				for (FPipelineCacheFileFormatPSORead* Entry : ShutdownFetchTasks)
+				{
+					if (Entry->ReadRequest.IsValid())
+					{
+						Entry->ReadRequest->WaitCompletion(0.f);
+					}
+					delete Entry;
+				}
+				ShutdownFetchTasks.Empty();
+				TotalRemainingTasks = 0;
+				Instance = nullptr;
+			}
+		};
+		TUniquePtr<class FShutdownContainer> FShutdownContainer::Instance;
+	}
+}
+
 FShaderPipelineCache* FShaderPipelineCache::ShaderPipelineCache = nullptr;
+FString FShaderPipelineCache::UserCacheTaskKey;
 
 FShaderPipelineCache::FShaderCachePreOpenDelegate FShaderPipelineCache::OnCachePreOpen;
 FShaderPipelineCache::FShaderCacheOpenedDelegate FShaderPipelineCache::OnCachedOpened;
@@ -321,11 +614,11 @@ static void PipelineStateCacheOnAppDeactivate()
 {
 	if (GetShaderPipelineCacheSaveBoundPSOLog())
 	{
-		FShaderPipelineCache::SavePipelineFileCache(FPipelineFileCache::SaveMode::BoundPSOsOnly);
+		FShaderPipelineCache::SavePipelineFileCache(FPipelineFileCacheManager::SaveMode::BoundPSOsOnly);
 	}
 	if (GetPSOFileCacheSaveUserCache())
 	{
-		FShaderPipelineCache::SavePipelineFileCache(FPipelineFileCache::SaveMode::Incremental);
+		FShaderPipelineCache::SavePipelineFileCache(FPipelineFileCacheManager::SaveMode::Incremental);
 	}
 }
 
@@ -340,70 +633,37 @@ bool FShaderPipelineCache::SetGameUsageMaskWithComparison(uint64 InMask, FPSOMas
 {
 	static bool bMaskChanged = false;
 
-	if (ShaderPipelineCache != nullptr && CVarPSOFileCacheGameFileMaskEnabled.GetValueOnAnyThread() && !ShaderPipelineCache->bPreOptimizing)
+	if (ShaderPipelineCache != nullptr && CVarPSOFileCacheGameFileMaskEnabled.GetValueOnAnyThread())
 	{
 		FScopeLock Lock(&ShaderPipelineCache->Mutex);
 		
-		if (ShaderPipelineCache->bOpened)
+		if (!ShaderPipelineCache->ShaderCacheTasks.IsEmpty())
 		{
-			uint64 OldMask = FPipelineFileCache::SetGameUsageMaskWithComparison(InMask, InComparisonFnPtr);
+			uint64 OldMask = FPipelineFileCacheManager::SetGameUsageMaskWithComparison(InMask, InComparisonFnPtr);
 			bMaskChanged |= OldMask != InMask;
-		
-			ShaderPipelineCache->bReady = true;
-		
-			if(bMaskChanged)
+			if (bMaskChanged)
 			{
-				// Mask has changed and we have an open file refetch PSO's for this Mask - leave the FPipelineFileCache file open - no need to close - just pull out the relevant PSOs.
-				// If this PSO compile run has completed for this Mask in which case don't refetch + compile for that mask
-			
-				// Don't clear already compiled PSOHash list - this is not a full reset
-				ShaderPipelineCache->Flush(false);
-			
-				if(!ShaderPipelineCache->CompletedMasks.Contains(InMask))
+				// reset the current queue and all of the tasks.
+				ShaderPipelineCache->ClearPendingPrecompileTaskQueue();
+				int PendingPSOCacheCount = 0;
+				for (TPair < FString, TUniquePtr<class FShaderPipelineCacheTask>>& ShaderCacheTask : ShaderPipelineCache->ShaderCacheTasks)
 				{
-					int32 Order = (int32)FPipelineFileCache::PSOOrder::Default;
-				
-					if(!GConfig->GetInt(FShaderPipelineCacheConstants::SectionHeading, FShaderPipelineCacheConstants::SortOrderKey, Order, *GGameUserSettingsIni))
-					{
-						GConfig->GetInt(FShaderPipelineCacheConstants::SectionHeading, FShaderPipelineCacheConstants::SortOrderKey, Order, *GGameIni);
-					}
-				
-					TArray<FPipelineCachePSOHeader> LocalPreFetchedTasks;
-					FPipelineFileCache::GetOrderedPSOHashes(LocalPreFetchedTasks, (FPipelineFileCache::PSOOrder)Order, (int64)CVarPSOFileCacheMinBindCount.GetValueOnAnyThread(), ShaderPipelineCache->CompiledHashes);
-					// Iterate over all the tasks we haven't yet begun to read data for - these are the 'waiting' tasks
-					int64 Count = 0;
-					for (auto const& Task : LocalPreFetchedTasks)
-					{
-						bool bHasShaders = true;
-						for (FSHAHash const& Hash : Task.Shaders)
-						{
-							bHasShaders &= FShaderCodeLibrary::ContainsShaderCode(Hash);
-						}
-						if (bHasShaders)
-						{
-							Count++;
-						}
-					}
-				
-					FPlatformAtomics::InterlockedAdd(&ShaderPipelineCache->TotalWaitingTasks, Count);
-				
-					if (OnCachedOpened.IsBound())
+					if(!ShaderCacheTask.Value->CompletedMasks.Contains(InMask))
 					{
-						OnCachedOpened.Broadcast(ShaderPipelineCache->FileName, ShaderPipelineCache->CurrentPlatform, LocalPreFetchedTasks.Num(), ShaderPipelineCache->CacheFileGuid, ShaderPipelineCache->ShaderCachePrecompileContext);
+						// Mask has changed and we have an open file refetch PSO's for this Mask - leave the FPipelineFileCacheManager file open - no need to close - just pull out the relevant PSOs.
+						// If this PSO compile run has completed for this Mask in which case don't refetch + compile for that mask
+						// Don't clear already compiled PSOHash list - this is not a full reset
+						ShaderCacheTask.Value->bReady = true;
+						ShaderPipelineCache->EnqueuePendingPrecompileTask(ShaderCacheTask.Key);
+						PendingPSOCacheCount++;
 					}
-				
-					ShaderPipelineCache->PreFetchedTasks = LocalPreFetchedTasks;
-				
-					bMaskChanged = false;
-					
-					UE_LOG(LogRHI, Display, TEXT("New ShaderPipelineCache GameUsageMask [%llu=>%llu], Enqueued %d of %d tasks for precompile."), OldMask, InMask, Count, ShaderPipelineCache->PreFetchedTasks.Num());
-					
-					return OldMask != InMask;
-				}
-				else
-				{
-					UE_LOG(LogRHI, Display, TEXT("New ShaderPipelineCache GameUsageMask [%llu=>%llu], Target mask already precompiled."), OldMask, InMask);
 				}
+
+				bMaskChanged = false;
+						
+				UE_LOG(LogRHI, Display, TEXT("New ShaderPipelineCacheTask GameUsageMask [%llu=>%llu], queued %d cache(s) for precompile."), OldMask, InMask, PendingPSOCacheCount);
+
+				return OldMask != InMask;
 			}
 			else
 			{
@@ -413,10 +673,10 @@ bool FShaderPipelineCache::SetGameUsageMaskWithComparison(uint64 InMask, FPSOMas
 		else
 		{
 			// NOTE: if this is called and then the cache is opened, but this function is never called again, the PSOs will not get precompiled. Should probably update the open call itself to see if there was a new mask set and call the code above
-			uint64 OldMask = FPipelineFileCache::SetGameUsageMaskWithComparison(InMask, InComparisonFnPtr);
+			uint64 OldMask = FPipelineFileCacheManager::SetGameUsageMaskWithComparison(InMask, InComparisonFnPtr);
 			bMaskChanged |= OldMask != InMask;
 
-			UE_LOG(LogRHI, Display, TEXT("ShaderPipelineCache::SetGameUsageMaskWithComparison set a new mask but did not attempt to setup any tasks because the cache was not open"));
+			UE_LOG(LogRHI, Display, TEXT("ShaderPipelineCache::SetGameUsageMaskWithComparison set a new mask but did not attempt to setup any tasks because there are no PSO caching tasks."));
 			return OldMask != InMask;
 		}
 	}
@@ -434,7 +694,7 @@ void FShaderPipelineCache::Initialize(EShaderPlatform Platform)
 	
 	if(FShaderCodeLibrary::IsEnabled())
 	{
-		FPipelineFileCache::Initialize(GetGameVersionForPSOFileCache());
+		FPipelineFileCacheManager::Initialize(GetGameVersionForPSOFileCache());
 		ShaderPipelineCache = new FShaderPipelineCache(Platform);
 	}
 }
@@ -520,92 +780,154 @@ uint32 FShaderPipelineCache::NumPrecompilesRemaining()
 	
 	if (ShaderPipelineCache)
 	{
-		if (CVarPSOFileCacheMaxPrecompileTime.GetValueOnAnyThread() > 0.0 && FPlatformAtomics::AtomicRead(&ShaderPipelineCache->TotalPrecompileTasks) > 0)
+		if (CVarPSOFileCacheMaxPrecompileTime.GetValueOnAnyThread() > 0.0 && FShaderPipelineCacheTask::TotalPrecompileTasks > 0)
 		{
-			float Mult = ShaderPipelineCache->TotalPrecompileWallTime / CVarPSOFileCacheMaxPrecompileTime.GetValueOnAnyThread();
-			NumRemaining = uint32(FMath::Max(0.f, 1.f - Mult) * FPlatformAtomics::AtomicRead(&ShaderPipelineCache->TotalPrecompileTasks));
+			float Mult = FShaderPipelineCacheTask::TotalPrecompileWallTime / CVarPSOFileCacheMaxPrecompileTime.GetValueOnAnyThread();
+			NumRemaining = uint32(FMath::Max(0.f, 1.f - Mult) * FShaderPipelineCacheTask::TotalPrecompileTasks);
 		}
 		else
 		{
-        	int64 NumActiveTasksRemaining = FMath::Max(0ll, FPlatformAtomics::AtomicRead(&ShaderPipelineCache->TotalActiveTasks));
-        	int64 NumWaitingTasksRemaining = FMath::Max(0ll, FPlatformAtomics::AtomicRead(&ShaderPipelineCache->TotalWaitingTasks));
+			int64 NumActiveTasksRemaining = FMath::Max((int64_t)0, FShaderPipelineCacheTask::TotalActiveTasks.load());
+			int64 NumWaitingTasksRemaining = FMath::Max((int64_t)0, FShaderPipelineCacheTask::TotalWaitingTasks.load());
 			NumRemaining = (uint32)(NumActiveTasksRemaining + NumWaitingTasksRemaining);
 		}
 	}
 
-	return NumRemaining;
+	// NumRemaining is correct only for the current cache.
+	// This ensures we do not return 0 if there are more caches to process.
+	uint32 PendingCaches = IsPrecompiling() ? 1 : 0;
+	return FMath::Max(PendingCaches,NumRemaining);
 }
 
-uint32 FShaderPipelineCache::NumPrecompilesActive()
+bool FShaderPipelineCache::OpenPipelineFileCache(EShaderPlatform Platform)
 {
-    uint32 NumRemaining = 0;
-    
-    if (ShaderPipelineCache)
-    {
-        int64 NumTasksRemaining = FPlatformAtomics::AtomicRead(&ShaderPipelineCache->TotalActiveTasks);
-        if(NumTasksRemaining > 0)
-        {
-            NumRemaining = (uint32) NumTasksRemaining;
-        }
-		if (CVarPSOFileCacheMaxPrecompileTime.GetValueOnAnyThread() > 0.0 && FPlatformAtomics::AtomicRead(&ShaderPipelineCache->TotalPrecompileTasks) > 0)
-		{
-			float Mult = ShaderPipelineCache->TotalPrecompileWallTime / CVarPSOFileCacheMaxPrecompileTime.GetValueOnAnyThread();
-			NumRemaining = uint32(FMath::Max(0.f, 1.f - Mult) * FPlatformAtomics::AtomicRead(&ShaderPipelineCache->TotalPrecompileTasks));
-		}
-    }
-    
-    return NumRemaining;
+	bool bFileOpen = false;
+	if (ShaderPipelineCache)
+	{
+		FString CacheName = UE::ShaderPipeline::GetShaderPipelineBaseName(false);
+		// for Bundled files the Key == CacheName, only the user cache differs.
+		const FString& Key = CacheName;
+
+		bFileOpen = ShaderPipelineCache->OpenPipelineFileCacheInternal(false, Key, CacheName, Platform);
+	}
+	return bFileOpen;
 }
 
-bool FShaderPipelineCache::OpenPipelineFileCache(EShaderPlatform Platform)
+bool FShaderPipelineCache::OpenUserPipelineFileCache(EShaderPlatform Platform)
 {
 	bool bFileOpen = false;
-	if (GConfig)
+	if (ShaderPipelineCache)
 	{
-		FString LastOpenedName;
-		if ((GConfig->GetString(FShaderPipelineCacheConstants::SectionHeading, FShaderPipelineCacheConstants::LastOpenedKey, LastOpenedName, *GGameUserSettingsIni) || GConfig->GetString(FShaderPipelineCacheConstants::SectionHeading, FShaderPipelineCacheConstants::LastOpenedKey, LastOpenedName, *GGameIni)) && LastOpenedName.Len())
+		FScopeLock Lock(&ShaderPipelineCache->Mutex);
+		// Just use the project for the user cache name.
+		FString UserCacheName = UE::ShaderPipeline::GetShaderPipelineBaseName(true);
+		ShaderPipelineCache->UserCacheTaskKey = UserCacheName + TEXT("_usr");
+		if (UE::ShaderPipeline::ShouldLoadUserCache() && !ShaderPipelineCache->GetTask(UserCacheTaskKey))
 		{
-			bFileOpen = OpenPipelineFileCache(LastOpenedName, Platform);
+			bFileOpen = ShaderPipelineCache->OpenPipelineFileCacheInternal(true, ShaderPipelineCache->UserCacheTaskKey, UserCacheName, Platform);
 		}
 	}
+	return bFileOpen;
+}
 
-	if (!bFileOpen)
+bool FShaderPipelineCache::OpenPipelineFileCacheInternal(bool bWantUserCache, const FString& PSOCacheKey, const FString& CacheName, EShaderPlatform Platform)
+{
+	if(!ShaderCacheTasks.Contains(PSOCacheKey))
 	{
-		bFileOpen = OpenPipelineFileCache(FApp::GetProjectName(), Platform);
+		TUniquePtr<FShaderPipelineCacheTask> NewPSOTask = MakeUnique<FShaderPipelineCacheTask>(bWantUserCache, Platform);
+		bool bOpened = NewPSOTask->Open(PSOCacheKey, CacheName, Platform);
+		if(bOpened)
+		{
+			bool bReady = NewPSOTask->bReady;
+			ShaderCacheTasks.FindOrAdd(PSOCacheKey, MoveTemp(NewPSOTask));
+			if(bReady)
+			{
+				EnqueuePendingPrecompileTask(PSOCacheKey);
+			}
+		}
+		return bOpened;
 	}
-
-	return bFileOpen;
+	return false;
 }
 
-bool FShaderPipelineCache::OpenPipelineFileCache(FString const& Name, EShaderPlatform Platform)
+bool FShaderPipelineCache::SavePipelineFileCache(FPipelineFileCacheManager::SaveMode Mode)
 {
 	if (ShaderPipelineCache)
-		return ShaderPipelineCache->Open(Name, Platform);
-	else
-		return false;
+	{
+		FScopeLock Lock(&ShaderPipelineCache->Mutex);
+
+		bool bOK = FPipelineFileCacheManager::SavePipelineFileCache(Mode);
+		UE_CLOG(!bOK, LogRHI, Warning, TEXT("Failed to save shader pipeline cache for using save mode %d."), (uint32)Mode);
+
+		ShaderPipelineCache->LastAutoSaveTime = FPlatformTime::Seconds();
+
+		return bOK;
+	}
+
+	return false;
 }
 
-bool FShaderPipelineCache::SavePipelineFileCache(FPipelineFileCache::SaveMode Mode)
+void FShaderPipelineCache::CloseUserPipelineFileCache()
 {
 	if (ShaderPipelineCache)
-		return ShaderPipelineCache->Save(Mode);
-	else
-		return false;
+	{
+		FScopeLock Lock(&ShaderPipelineCache->Mutex);
+
+		// Log all bound PSOs
+		if (GetShaderPipelineCacheSaveBoundPSOLog())
+		{
+			SavePipelineFileCache(FPipelineFileCacheManager::SaveMode::BoundPSOsOnly);
+		}
+
+		if (GetPSOFileCacheSaveUserCache())
+		{
+			SavePipelineFileCache(FPipelineFileCacheManager::SaveMode::Incremental);
+		}
+
+		
+ 		if (FShaderPipelineCacheTask* Found = ShaderPipelineCache->GetTask(UserCacheTaskKey))
+ 		{
+			Found->Close();
+			ShaderPipelineCache->PendingPrecompilePSOCacheKeys.Remove(UserCacheTaskKey);
+			ShaderPipelineCache->ShaderCacheTasks.Remove(UserCacheTaskKey);
+ 		}
+
+		FPipelineFileCacheManager::CloseUserPipelineFileCache();
+	}
 }
 
-void FShaderPipelineCache::ClosePipelineFileCache()
+
+void FShaderPipelineCache::ShaderLibraryStateChanged(ELibraryState State, EShaderPlatform Platform, FString const& PSOCacheBaseName, int32 ComponentID)
 {
 	if (ShaderPipelineCache)
-		ShaderPipelineCache->Close();
-}
+	{
+		const FString PSOCacheName = UE::ShaderPipeline::LibNameToPSOFCName(PSOCacheBaseName, ComponentID);
+		// Key == Name for bundles PSOFCs
+		const FString& PSOCacheKey = PSOCacheName;
 
-void FShaderPipelineCache::ShaderLibraryStateChanged(ELibraryState State, EShaderPlatform Platform, FString const& Name)
-{
-    if (ShaderPipelineCache)
-        ShaderPipelineCache->OnShaderLibraryStateChanged(State, Platform, Name);
+		FScopeLock Lock(&ShaderPipelineCache->Mutex);
+		if (TUniquePtr<class FShaderPipelineCacheTask>* Found = ShaderPipelineCache->ShaderCacheTasks.Find(PSOCacheKey))
+		{
+			(*Found)->OnShaderLibraryStateChanged(State, Platform, PSOCacheName, ComponentID);
+		}
+		else
+		{
+			bool bSuccess = ShaderPipelineCache->OpenPipelineFileCacheInternal(false, PSOCacheKey, PSOCacheName, Platform);
+
+			const FString& PSOCacheBaseKey = PSOCacheBaseName;
+			// for some reason we're not always able to find the base PSO cache when the shader libs loads.
+			if( !bSuccess && !ShaderPipelineCache->ShaderCacheTasks.Contains(PSOCacheBaseKey))
+			{
+				ShaderPipelineCache->OpenPipelineFileCacheInternal(false, PSOCacheBaseKey, PSOCacheBaseName, Platform);
+			}
+		}
+
+		// An opportunity to open the user cache
+		OpenUserPipelineFileCache(Platform);
+	}
 }
 
-bool FShaderPipelineCache::Precompile(FRHICommandListImmediate& RHICmdList, EShaderPlatform Platform, FPipelineCacheFileFormatPSO const& PSO)
+bool FShaderPipelineCacheTask::Precompile(FRHICommandListImmediate& RHICmdList, EShaderPlatform Platform, FPipelineCacheFileFormatPSO const& PSO)
 {
 	INC_DWORD_STAT(STAT_PreCompileShadersTotal);
 	INC_DWORD_STAT(STAT_PreCompileShadersNum);
@@ -664,13 +986,13 @@ bool FShaderPipelineCache::Precompile(FRHICommandListImmediate& RHICmdList, ESha
 					GraphicsInitializer.BoundShaderState.PixelShaderRHI = FragmentShader;
 				}
 
-				auto BlendState = GetOrCreateBlendState(PSO.GraphicsDesc.BlendState);
+				auto BlendState = FShaderPipelineCache::Get()->GetOrCreateBlendState(PSO.GraphicsDesc.BlendState);
 				GraphicsInitializer.BlendState = BlendState;
 
-				auto RasterState = GetOrCreateRasterizerState(PSO.GraphicsDesc.RasterizerState);
+				auto RasterState = FShaderPipelineCache::Get()->GetOrCreateRasterizerState(PSO.GraphicsDesc.RasterizerState);
 				GraphicsInitializer.RasterizerState = RasterState;
 
-				auto DepthState = GetOrCreateDepthStencilState(PSO.GraphicsDesc.DepthStencilState);
+				auto DepthState = FShaderPipelineCache::Get()->GetOrCreateDepthStencilState(PSO.GraphicsDesc.DepthStencilState);
 				GraphicsInitializer.DepthStencilState = DepthState;
 
 				for (uint32 i = 0; i < MaxSimultaneousRenderTargets; ++i)
@@ -694,9 +1016,9 @@ bool FShaderPipelineCache::Precompile(FRHICommandListImmediate& RHICmdList, ESha
 				GraphicsInitializer.StencilTargetLoadAction = PSO.GraphicsDesc.StencilLoad;
 				GraphicsInitializer.DepthTargetStoreAction = PSO.GraphicsDesc.DepthStore;
 				GraphicsInitializer.StencilTargetStoreAction = PSO.GraphicsDesc.StencilStore;
-			
+
 				GraphicsInitializer.PrimitiveType = PSO.GraphicsDesc.PrimitiveType;
-			
+
 				// This indicates we do not want a fatal error if this compilation fails
 				// (ie, if this entry in the file cache is bad)
 				GraphicsInitializer.bFromPSOFileCache = true;
@@ -785,14 +1107,14 @@ bool FShaderPipelineCache::Precompile(FRHICommandListImmediate& RHICmdList, ESha
     // Otherwise we end up with outstanding compiles that we can't progress or external tools may think this has not been completed and may run again.
     {
         uint64 TimeDelta = FPlatformTime::Cycles64() - StartTime;
-        FPlatformAtomics::InterlockedIncrement(&TotalCompleteTasks);
-        FPlatformAtomics::InterlockedAdd(&TotalPrecompileTime, TimeDelta);
+        TotalCompleteTasks++;
+        TotalPrecompileTime+= TimeDelta;
     }
 	
 	return bOk;
 }
 
-bool FShaderPipelineCache::CompileJob::PreloadShaders(TSet<FSHAHash>& OutRequiredShaders, bool &bOutCompatible)
+bool UE::ShaderPipeline::CompileJob::PreloadShaders(TSet<FSHAHash>& OutRequiredShaders, bool &bOutCompatible)
 {
 	static FSHAHash EmptySHA;
 	bool bOK = true;
@@ -914,7 +1236,7 @@ bool FShaderPipelineCache::CompileJob::PreloadShaders(TSet<FSHAHash>& OutRequire
 	return bShadersPreloaded;
 }
 
-void FShaderPipelineCache::CompileJob::ReleasePreloadedShaders()
+void UE::ShaderPipeline::CompileJob::ReleasePreloadedShaders()
 {
 	if (bShadersPreloaded)
 	{
@@ -952,7 +1274,7 @@ void FShaderPipelineCache::CompileJob::ReleasePreloadedShaders()
 	}
 }
 
-void FShaderPipelineCache::PreparePipelineBatch(TDoubleLinkedList<FPipelineCacheFileFormatPSORead*>& PipelineBatch)
+void FShaderPipelineCacheTask::PreparePipelineBatch(TDoubleLinkedList<FPipelineCacheFileFormatPSORead*>& PipelineBatch)
 {
 	TDoubleLinkedList<FPipelineCacheFileFormatPSORead*>::TDoubleLinkedListNode* CurrentNode = PipelineBatch.GetHead();
 	while(CurrentNode)
@@ -982,7 +1304,7 @@ void FShaderPipelineCache::PreparePipelineBatch(TDoubleLinkedList<FPipelineCache
             // Shaders required.
             TSet<FSHAHash> RequiredShaders;
 			
-			CompileJob AsyncJob;
+			UE::ShaderPipeline::CompileJob AsyncJob;
 			AsyncJob.PSO = PSO;
 			AsyncJob.ReadRequests = new FShaderPipelineCacheArchive;
 		
@@ -1009,9 +1331,9 @@ void FShaderPipelineCache::PreparePipelineBatch(TDoubleLinkedList<FPipelineCache
 				{
 					UE_LOG(LogRHI, Error, TEXT("Invalid PSO entry in pipeline cache: %u!"), PSORead->Hash);
 				}
-
 				// Go to async shutdown instead - some shader code reads may have been requested - let it handle this
-				ShutdownReadCompileTasks.Add(AsyncJob);
+				UE::ShaderPipeline::FShutdownContainer::Get()->AbandonCompileJob(MoveTemp(AsyncJob));
+				TotalActiveTasks--;
 			}
 			
 			bRemoveEntry = true;
@@ -1021,7 +1343,7 @@ void FShaderPipelineCache::PreparePipelineBatch(TDoubleLinkedList<FPipelineCache
             UE_LOG(LogRHI, Error, TEXT("Invalid PSO entry in pipeline cache: %u!"), PSORead->Hash);
             
             // Invalid PSOs can be deleted
-            FPlatformAtomics::InterlockedDecrement(&TotalActiveTasks);
+			TotalActiveTasks--;
             bRemoveEntry = true;
         }
 		
@@ -1049,7 +1371,7 @@ static FAutoConsoleVariableRef GPipelinePrecompileTasksInFlightVar(
 	ECVF_RenderThreadSafe
 );
 
-bool FShaderPipelineCache::ReadyForPrecompile()
+bool FShaderPipelineCacheTask::ReadyForPrecompile()
 {
 	for(uint32 i = 0; i < (uint32)ReadTasks.Num();/*NOP*/)
 	{
@@ -1074,26 +1396,28 @@ bool FShaderPipelineCache::ReadyForPrecompile()
 		&& !LastPrecompileRHIFence.GetReference();
 }
 
-void FShaderPipelineCache::PrecompilePipelineBatch()
+void FShaderPipelineCacheTask::PrecompilePipelineBatch()
 {
 	INC_DWORD_STAT(STAT_PreCompileBatchTotal);
 	INC_DWORD_STAT(STAT_PreCompileBatchNum);
 	
-    int32 NumToPrecompile = FMath::Min<int32>(CompileTasks.Num(), BatchSize);
-    
+    int32 NumToPrecompile = FMath::Min<int32>(CompileTasks.Num(), FShaderPipelineCache::BatchSize);
+
+	FShaderPipelineCacheTask* UserCacheTask = FShaderPipelineCache::Get()->GetTask(FShaderPipelineCache::UserCacheTaskKey);
+
 	for(uint32 i = 0; i < (uint32)NumToPrecompile; i++)
 	{
-		CompileJob& CompileTask = CompileTasks[i];
+		UE::ShaderPipeline::CompileJob& CompileTask = CompileTasks[i];
 		
 		check(CompileTask.ReadRequests && CompileTask.ReadRequests->PollExternalReadDependencies());
 		
 		FRHICommandListImmediate& RHICmdList = GRHICommandList.GetImmediateCommandList();
 		
 		uint32 PSOHash = GetTypeHash(CompileTask.PSO);
-		UE_LOG(LogRHI, Verbose, TEXT("Precompiling PSO %u (%d/%d)"), PSOHash, i, NumToPrecompile);
+		UE_LOG(LogRHI, Verbose, TEXT("Precompiling PSO %u (type %d) (%d/%d)"), PSOHash, CompileTask.PSO.Type, i+1, NumToPrecompile);
 		Precompile(RHICmdList, GMaxRHIShaderPlatform, CompileTask.PSO);
-		CompiledHashes.Add(PSOHash);
-		
+		FShaderPipelineCache::Get()->CompiledHashes.Add(PSOHash);
+
 		CompileTask.ReleasePreloadedShaders();
 		delete CompileTask.ReadRequests;
 		CompileTask.ReadRequests = nullptr;
@@ -1124,156 +1448,58 @@ void FShaderPipelineCache::PrecompilePipelineBatch()
 		}
 #endif
 	}
-	
-    FPlatformAtomics::InterlockedAdd(&TotalActiveTasks, -NumToPrecompile);
+
+	TotalPSOsCompiled += NumToPrecompile;
+    TotalActiveTasks -= NumToPrecompile;
 
 #if PLATFORM_ANDROID
 	if (NumToPrecompile > 0 && IsRunningRHIInSeparateThread())
 	{
-		LastPrecompileRHIFence = FRHICommandListExecutor::GetImmediateCommandList().RHIThreadFence(false);
+ 		LastPrecompileRHIFence = FRHICommandListExecutor::GetImmediateCommandList().RHIThreadFence(false);
 	}
 #endif
 
 	CompileTasks.RemoveAt(0, NumToPrecompile);
 }
 
-bool FShaderPipelineCache::ReadyForNextBatch() const
+bool FShaderPipelineCacheTask::ReadyForNextBatch() const
 {
 	return ReadTasks.Num() == 0;
 }
 
-bool FShaderPipelineCache::ReadyForAutoSave() const
-{
-	bool bAutoSave = false;
-	uint32 SaveAfterNum = CVarPSOFileCacheSaveAfterPSOsLogged.GetValueOnAnyThread();
-	uint32 NumLogged = FPipelineFileCache::NumPSOsLogged();
-
-	const double TimeSinceSave = FPlatformTime::Seconds() - LastAutoSaveTime;
-
-	// autosave if its enabled, and we have more than the desired number, or it's been a while since our last save
-	if (SaveAfterNum > 0 && 
-			(NumLogged >= SaveAfterNum || (NumLogged > 0 && TimeSinceSave >= CVarPSOFileCacheAutoSaveTime.GetValueOnAnyThread()))
-		)
-	{
-		bAutoSave = true;
-	}
-	return bAutoSave;
-}
-
-void FShaderPipelineCache::PollShutdownItems()
+void FShaderPipelineCacheTask::Flush()
 {
-	int64 RemovedTaskCount = 0;
-	
-	const int64 InitialCompileTaskCount = ShutdownReadCompileTasks.Num();
-	if(ShutdownReadCompileTasks.Num() > 0)
-	{
-		for(uint32 i = 0; i < (uint32)ShutdownReadCompileTasks.Num(); )
-		{
-			check(ShutdownReadCompileTasks[i].ReadRequests);
-			if (ShutdownReadCompileTasks[i].ReadRequests->PollExternalReadDependencies())
-			{
-				ShutdownReadCompileTasks[i].ReleasePreloadedShaders();
-				delete ShutdownReadCompileTasks[i].ReadRequests;
-				ShutdownReadCompileTasks[i].ReadRequests = nullptr;
-				
-				ShutdownReadCompileTasks.RemoveAt(i, 1, false);
-				++RemovedTaskCount;
-			}
-			else
-			{
-				++i;
-			}
-		}
-		
-		if(ShutdownReadCompileTasks.Num() == 0)
-		{
-			ShutdownReadCompileTasks.Shrink();
-		}
-	}
-	
-	if(ShutdownFetchTasks.Num() > 0)
-	{
-		TDoubleLinkedList<FPipelineCacheFileFormatPSORead*>::TDoubleLinkedListNode* CurrentNode = ShutdownFetchTasks.GetHead();
-		while(CurrentNode)
-		{
-			FPipelineCacheFileFormatPSORead* PSORead = CurrentNode->GetValue();
-			check(PSORead);
-			FShaderPipelineCacheArchive* Archive = (FShaderPipelineCacheArchive*)(PSORead->Ar);
-			check(Archive);
-			
-			TDoubleLinkedList<FPipelineCacheFileFormatPSORead*>::TDoubleLinkedListNode* PrevNode = CurrentNode;
-			CurrentNode = CurrentNode->GetNextNode();
-			
-			if (PSORead->bReadCompleted || Archive->PollExternalReadDependencies())
-			{
-				delete PSORead;
-				ShutdownFetchTasks.RemoveNode(PrevNode);
-				++RemovedTaskCount;
-			}
-		}
-	}
-	
-	if(RemovedTaskCount > 0)
-    {
-        FPlatformAtomics::InterlockedAdd(&TotalActiveTasks, - RemovedTaskCount);
-	}
-}
-
-void FShaderPipelineCache::Flush(bool bClearCompiled /*= true*/)
-{
-	FScopeLock Lock(&Mutex);
-	
-	if(bClearCompiled)
-	{
-		CompiledHashes.Empty();
-	}
-	
 	// reset everything
 	// Abandon all the existing work.
 	// This must be done on the render-thread to avoid locks.
 	OrderedCompileTasks.Empty();
 	
-	// Marshall the current Compile Jobs into shutdown
-	ShutdownReadCompileTasks.Append(ReadTasks);
+	// Marshall the current Compile Jobs into the shutdown container.
+	int32 TotalAbandonedTasks = ReadTasks.Num() + CompileTasks.Num() + FetchTasks.Num();
+	UE::ShaderPipeline::FShutdownContainer::Get()->AbandonCompileJobs(MoveTemp(ReadTasks));
+	UE::ShaderPipeline::FShutdownContainer::Get()->AbandonCompileJobs(MoveTemp(CompileTasks));
+	UE::ShaderPipeline::FShutdownContainer::Get()->AbandonFetchTasks(MoveTemp(FetchTasks));
+	TotalActiveTasks -= TotalAbandonedTasks;
+
 	ReadTasks.Empty();
-	
-	ShutdownReadCompileTasks.Append(CompileTasks);
 	CompileTasks.Empty();
-	
-	// Marshall the current fetch tasks into shutdown
-	for (TDoubleLinkedList<FPipelineCacheFileFormatPSORead*>::TIterator It(FetchTasks.GetHead()); It; ++It)
-	{
-		FPipelineCacheFileFormatPSORead* Entry = *It;
-		check(Entry->ReadRequest.IsValid());
-		Entry->ReadRequest->Cancel();
-		ShutdownFetchTasks.AddTail(Entry);
-	}
 	FetchTasks.Empty();
 	
-	int64 StartTaskCount = OrderedCompileTasks.Num() + ShutdownReadCompileTasks.Num() + ShutdownFetchTasks.Num();;
-    FPlatformAtomics::InterlockedExchange(&TotalWaitingTasks, 0);
+	TotalWaitingTasks = 0;
+	TotalPrecompileTasks = 0;
+	TotalCompleteTasks = 0;
+	TotalPrecompileTime = 0;
 }
 
+static bool PreCompileMaskComparison(uint64 ReferenceGameMask, uint64 PSOMask);
+
 FShaderPipelineCache::FShaderPipelineCache(EShaderPlatform Platform)
 : FTickableObjectRenderThread(true, false) // (RegisterNow, HighFrequency)
-, CurrentPlatform((EShaderPlatform)-1)
-, BatchSize(0)
-, BatchTime(0.0f)
 , bPaused(false)
-, bOpened(false)
-, bReady(false)
-, bPreOptimizing(false)
 , PausedCount(0)
-, TotalActiveTasks(0)
-, TotalWaitingTasks(0)
-, TotalCompleteTasks(0)
-, TotalPrecompileTime(0)
-, PrecompileStartTime(0.0)
-, LastAutoSaveTime(0)
-, LastAutoSaveTimeLogBoundPSO(FPlatformTime::Seconds())
-, LastAutoSaveNum(-1)
-, TotalPrecompileWallTime(0.0)
-, TotalPrecompileTasks(0)
+ , LastAutoSaveTime(0)
+ , LastAutoSaveTimeLogBoundPSO(FPlatformTime::Seconds())
+ , LastAutoSaveNum(-1)
 {
 	SET_DWORD_STAT(STAT_ShaderPipelineTaskCount, 0);
     SET_DWORD_STAT(STAT_ShaderPipelineWaitingTaskCount, 0);
@@ -1285,7 +1511,9 @@ FShaderPipelineCache::FShaderPipelineCache(EShaderPlatform Platform)
 		case 0:
 			BatchSize = CVarPSOFileCacheBatchSize.GetValueOnAnyThread();
 			BatchTime = CVarPSOFileCacheBatchTime.GetValueOnAnyThread();
+			PausedCount = 1;
 			bPaused = true;
+			UE_LOG(LogRHI, Display, TEXT("ShaderPipelineCache: Starting paused. %d"), PausedCount);
 			break;
 		case 2:
 			BatchSize = CVarPSOFileCacheBackgroundBatchSize.GetValueOnAnyThread();
@@ -1301,113 +1529,261 @@ FShaderPipelineCache::FShaderPipelineCache(EShaderPlatform Platform)
 	BatchSize = CVarPSOFileCacheBatchSize.GetValueOnAnyThread();
 	BatchTime = CVarPSOFileCacheBatchTime.GetValueOnAnyThread();
 	
+	uint64 PreCompileMask = (uint64)CVarPSOFileCachePreCompileMask.GetValueOnAnyThread();
+	FPipelineFileCacheManager::SetGameUsageMaskWithComparison(PreCompileMask, PreCompileMaskComparison);
+
 	FCoreDelegates::ApplicationWillDeactivateDelegate.AddStatic(&PipelineStateCacheOnAppDeactivate);
-	
-	bReady = !CVarPSOFileCacheGameFileMaskEnabled.GetValueOnAnyThread() || CVarPSOFileCachePreOptimizeEnabled.GetValueOnAnyThread();
+}
+
+bool FShaderPipelineCache::ReadyForAutoSave() const
+{
+	bool bAutoSave = false;
+	uint32 SaveAfterNum = CVarPSOFileCacheSaveAfterPSOsLogged.GetValueOnAnyThread();
+	uint32 NumLogged = FPipelineFileCacheManager::NumPSOsLogged();
+
+	const double TimeSinceSave = FPlatformTime::Seconds() - LastAutoSaveTime;
+
+	// autosave if its enabled, and we have more than the desired number, or it's been a while since our last save
+	if (SaveAfterNum > 0 &&
+		(NumLogged >= SaveAfterNum || (NumLogged > 0 && TimeSinceSave >= CVarPSOFileCacheAutoSaveTime.GetValueOnAnyThread()))
+		)
+	{
+		bAutoSave = true;
+	}
+	return bAutoSave;
 }
 
 FShaderPipelineCache::~FShaderPipelineCache()
 {
-	Close();
-	
-	// The render thread tick should be dead now and we are safe to destroy everything that needs to wait or manual destruction
+	FScopeLock Lock(&Mutex);
+
+	FShaderPipelineCache::CloseUserPipelineFileCache();
 
-	// Close() should have called though to Flush() - all work is now in the shutdown lists
-	for(CompileJob& Entry: ShutdownReadCompileTasks)
+	for (auto& PSOTask : ShaderCacheTasks)
 	{
-		if(Entry.ReadRequests != nullptr)
-		{
-			Entry.ReadRequests->BlockingWaitComplete();
-			Entry.ReleasePreloadedShaders();
-			delete Entry.ReadRequests;
-			Entry.ReadRequests = nullptr;
-		}
+		PSOTask.Value->Close();
+		// The render thread tick should be dead now and we are safe to destroy everything that needs to wait or manual destruction
+		// Close() should have called though to Flush() - all work is now in the shutdown lists
 	}
-	
-	for (FPipelineCacheFileFormatPSORead* Entry : ShutdownFetchTasks)
+
+	//
+	// Clean up cached RHI resources
+	//
+	for (auto& Pair : BlendStateCache)
+	{
+		Pair.Value->Release();
+	}
+	BlendStateCache.Empty();
+
+	for (auto& Pair : RasterizerStateCache)
+	{
+		Pair.Value->Release();
+	}
+	RasterizerStateCache.Empty();
+
+	for (auto& Pair : DepthStencilStateCache)
+	{
+		Pair.Value->Release();
+	}
+	DepthStencilStateCache.Empty();
+	UE::ShaderPipeline::FShutdownContainer::Get()->ShutdownAndWait();
+}
+
+ void FShaderPipelineCache::BeginNextPrecompileCacheTask()
+{
+	// call begin on the GT, some of the precompile delegate handling is expecting the call to come from the GT.
+	// TODO: we should be able to call this from any thread.
+	if (IsInGameThread())
+	{
+		ShaderPipelineCache->BeginNextPrecompileCacheTaskInternal();
+	}
+	else
 	{
-		if(Entry->ReadRequest.IsValid())
+		static std::atomic_bool bHasDispatched = false;
+		// kick off next PSOFC compile on GT, ensure we only have one request in flight.
+		bool bExpected = false;
+		if (bHasDispatched.compare_exchange_strong(bExpected, true))
 		{
-			Entry->ReadRequest->WaitCompletion(0.f);
+			FFunctionGraphTask::CreateAndDispatchWhenReady(
+				[]()
+				{
+					ShaderPipelineCache->BeginNextPrecompileCacheTaskInternal();
+					bHasDispatched = false;
+				}
+			, TStatId(), NULL, ENamedThreads::GameThread);
 		}
-		delete Entry;
 	}
 }
 
-bool FShaderPipelineCache::IsTickable() const
+void FShaderPipelineCache::ClearPendingPrecompileTaskQueue()
 {
-	return FPlatformProperties::RequiresCookedData() && !bPaused &&
-		(FPlatformAtomics::AtomicRead(&TotalActiveTasks) != 0 ||
-			FPlatformAtomics::AtomicRead(&TotalWaitingTasks) != 0 ||
-			FPlatformAtomics::AtomicRead(&TotalCompleteTasks) != 0 ||
-			PipelineStateCache::GetNumActivePipelinePrecompileTasks()||
-			ReadyForAutoSave()||
-			GetShaderPipelineCacheSaveBoundPSOLog());
+	if (!CurrentPrecompilingPSOFCKey.IsEmpty())
+	{
+		GetTask(CurrentPrecompilingPSOFCKey)->Flush();
+		CurrentPrecompilingPSOFCKey.Empty();
+	}
+	PendingPrecompilePSOCacheKeys.Reset();
 }
 
-void FShaderPipelineCache::Tick( float DeltaTime )
+void FShaderPipelineCache::EnqueuePendingPrecompileTask(const FString& PSOCacheKey)
 {
-	FScopeLock Lock(&Mutex);
-    
-	if (LastPrecompileRHIFence.GetReference() && LastPrecompileRHIFence->IsComplete())
+	check(!PendingPrecompilePSOCacheKeys.Contains(PSOCacheKey));
+	check(ShaderCacheTasks.Contains(PSOCacheKey));
+	PendingPrecompilePSOCacheKeys.Add(PSOCacheKey);
+	BeginNextPrecompileCacheTask();
+}
+
+bool FShaderPipelineCache::IsPrecompiling()
+{
+	if(ShaderPipelineCache)
 	{
-		LastPrecompileRHIFence = nullptr;
+		return ShaderPipelineCache->bHasShaderPipelineTask;
 	}
+	return false;
+}
 
-	if (PrecompileStartTime > 0.0)
+void FShaderPipelineCache::BeginNextPrecompileCacheTaskInternal()
+{
+	FScopeLock Lock(&Mutex);
+	if (!CurrentPrecompilingPSOFCKey.IsEmpty())
 	{
-		TotalPrecompileWallTime = float(FPlatformTime::Seconds() - PrecompileStartTime);
+		return;
 	}
-
-	// HACK: note for the time being until the code further upstream properly uses the OnPrecompilationComplete delegate, we need to provide some delays to make sure it gets the messaging properly (the -10.0f and the -20.0f below).
-	if (((PipelineStateCache::GetNumActivePipelinePrecompileTasks() == 0 && FMath::Max(0ll, FPlatformAtomics::AtomicRead(&TotalWaitingTasks)) == 0 && FMath::Max(0ll, FPlatformAtomics::AtomicRead(&TotalActiveTasks)) == 0 && FPlatformAtomics::AtomicRead(&TotalCompleteTasks) != 0) || (CVarPSOFileCacheMaxPrecompileTime.GetValueOnAnyThread() > 0.0 && TotalPrecompileWallTime - 10.0f > CVarPSOFileCacheMaxPrecompileTime.GetValueOnAnyThread() && TotalPrecompileTasks > 0)) && !LastPrecompileRHIFence.GetReference())
-    {
-		UE_LOG(LogRHI, Warning, TEXT("FShaderPipelineCache completed %u tasks in %.2fs (%.2fs wall time since intial open)."), (uint32)TotalCompleteTasks, FPlatformTime::ToSeconds64(TotalPrecompileTime), TotalPrecompileWallTime);
-        if (OnPrecompilationComplete.IsBound())
-        {
-            OnPrecompilationComplete.Broadcast((uint32)TotalCompleteTasks, FPlatformTime::ToSeconds64(TotalPrecompileTime), ShaderCachePrecompileContext);
-        }
-		if (CVarPSOFileCacheMaxPrecompileTime.GetValueOnAnyThread() > 0.0 && TotalPrecompileWallTime - 20.0f > CVarPSOFileCacheMaxPrecompileTime.GetValueOnAnyThread() && TotalPrecompileTasks > 0)
+	
+	while(!PendingPrecompilePSOCacheKeys.IsEmpty())
+	{
+		CurrentPrecompilingPSOFCKey = PendingPrecompilePSOCacheKeys[0];
+		PendingPrecompilePSOCacheKeys.RemoveAtSwap(0);
+		FShaderPipelineCacheTask* NextPrecompileTask = ShaderCacheTasks.Find(CurrentPrecompilingPSOFCKey)->Get();
+		check(NextPrecompileTask->bReady);
+		UE_LOG(LogRHI, Warning, TEXT("FShaderPipelineCache::BeginNextPrecompileCacheTask() - %s begining compile."), *CurrentPrecompilingPSOFCKey);			
+		NextPrecompileTask->BeginPrecompilingPipelineCache();
+
+		// FShaderPipelineCache::OnCachePreOpen can set bReady=false, if so we cannot precompile that task.
+		if(NextPrecompileTask->bReady)
 		{
-			PrecompileStartTime = 0.0;
-			SetBatchMode(BatchMode::Background);
-			TotalPrecompileTasks = 0;
+			bHasShaderPipelineTask = true;
+			break;
 		}
 		else
 		{
-			FPipelineFileCache::PreCompileComplete();
+			UE_LOG(LogRHI, Warning, TEXT("FShaderPipelineCache::BeginNextPrecompileCacheTask() - skipping %s, not ready for precompile."), *CurrentPrecompilingPSOFCKey);
+			CurrentPrecompilingPSOFCKey.Empty();
+		}
+	}
+
+	if (PendingPrecompilePSOCacheKeys.IsEmpty() && CurrentPrecompilingPSOFCKey.IsEmpty())
+	{
+		// all jobs done.
+		UE_LOG(LogRHI, Log, TEXT("FShaderPipelineCache::BeginNextPrecompileCacheTask() - Finished, no jobs remaining."));
+		bHasShaderPipelineTask = false;
+	}
+}
+
+bool FShaderPipelineCache::IsTickable() const
+{
+	check(IsInRenderingThread());
+
+	return 	FPlatformProperties::RequiresCookedData() && !bPaused &&
+		(
+			bHasShaderPipelineTask
+  			|| ReadyForAutoSave()
+ 			|| GetShaderPipelineCacheSaveBoundPSOLog()
+			|| GetPSOFileCacheSaveUserCache()
+			|| UE::ShaderPipeline::FShutdownContainer::Get()->IsTickable()
+		);
+}
+
+void FShaderPipelineCache::Tick(float DeltaTime)
+{
+	FScopeLock Lock(&Mutex);
+
+	FShaderPipelineCacheTask* CurrentCacheTask = CurrentPrecompilingPSOFCKey.IsEmpty() ? nullptr : ShaderCacheTasks.Find(CurrentPrecompilingPSOFCKey)->Get();
+	const bool bHasActiveCacheTask = CurrentCacheTask != nullptr;
+
+	// HACK: note for the time being until the code further upstream properly uses the OnPrecompilationComplete delegate, we need to provide some delays to make sure it gets the messaging properly (the -10.0f and the -20.0f below).
+	const bool bHasCurrentCacheTaskCompleted = bHasActiveCacheTask &&
+		PipelineStateCache::GetNumActivePipelinePrecompileTasks() == 0 && CurrentCacheTask->HasRHIFenceCompleted() &&
+				(
+				FMath::Max((int64_t)0, FShaderPipelineCacheTask::TotalWaitingTasks.load()) == 0 &&
+				FMath::Max((int64_t)0, FShaderPipelineCacheTask::TotalActiveTasks.load()) == 0 &&
+				(FShaderPipelineCacheTask::TotalPrecompileTasks == 0 || FShaderPipelineCacheTask::TotalCompleteTasks != 0)
+				);
+
+	if (!bHasActiveCacheTask || bHasCurrentCacheTaskCompleted )
+	{
+		// drop to background mode if CVarPSOFileCacheMaxPrecompileTime exceeded.
+		if (CVarPSOFileCacheMaxPrecompileTime.GetValueOnAnyThread() > 0.0 && FShaderPipelineCacheTask::TotalPrecompileWallTime - 20.0f > CVarPSOFileCacheMaxPrecompileTime.GetValueOnAnyThread() && FShaderPipelineCacheTask::TotalPrecompileTasks > 0)
+ 		{
+ 			SetBatchMode(BatchMode::Background);
+ 		}
+
+		if (bHasCurrentCacheTaskCompleted)
+		{
+			UE_LOG(LogRHI, Warning, TEXT("FShaderPipelineCache %s completed %u tasks in %.2fs (%.2fs wall time since intial open)."), *CurrentCacheTask->PSOCacheKey, (uint32)FShaderPipelineCacheTask::TotalCompleteTasks, FPlatformTime::ToSeconds64(FShaderPipelineCacheTask::TotalPrecompileTime), FShaderPipelineCacheTask::TotalPrecompileWallTime);
+			if (OnPrecompilationComplete.IsBound())
+			{
+				OnPrecompilationComplete.Broadcast((uint32)FShaderPipelineCacheTask::TotalCompleteTasks, FPlatformTime::ToSeconds64(FShaderPipelineCacheTask::TotalPrecompileTime), CurrentCacheTask->PrecompileContext);
+			}
+
+			FPipelineFileCacheManager::PreCompileComplete();
 			PipelineStateCache::PreCompileComplete();
 
-        	FPlatformAtomics::InterlockedExchange(&TotalCompleteTasks, 0);
-        	FPlatformAtomics::InterlockedExchange(&TotalPrecompileTime, 0);
-			bPreOptimizing = false;
+			FShaderPipelineCacheTask::PrecompileStartTime = 0.0;
+			FShaderPipelineCacheTask::TotalPrecompileTasks = 0;
+			FShaderPipelineCacheTask::TotalCompleteTasks = 0;
+			FShaderPipelineCacheTask::TotalPrecompileTime = 0;
+			CurrentCacheTask->bPreOptimizing = false;
+
+			CurrentPrecompilingPSOFCKey.Empty();
+			CurrentCacheTask = nullptr;
+			BeginNextPrecompileCacheTask();
 		}
-    }
-	
+	}
+
 	if (ReadyForAutoSave())
 	{
 		if (GetPSOFileCacheSaveUserCache())
 		{
-			Save(FPipelineFileCache::SaveMode::Incremental);
+			SavePipelineFileCache(FPipelineFileCacheManager::SaveMode::Incremental);
 		}
 	}
+
 	if (GetShaderPipelineCacheSaveBoundPSOLog())
 	{
-		if (LastAutoSaveNum < int32(FPipelineFileCache::NumPSOsLogged()))
+		if (LastAutoSaveNum < int32(FPipelineFileCacheManager::NumPSOsLogged()))
 		{
 			const double TimeSinceSave = FPlatformTime::Seconds() - LastAutoSaveTimeLogBoundPSO;
 
 			if (TimeSinceSave >= CVarPSOFileCacheAutoSaveTimeBoundPSO.GetValueOnAnyThread())
 			{
-				Save(FPipelineFileCache::SaveMode::BoundPSOsOnly);
+				SavePipelineFileCache(FPipelineFileCacheManager::SaveMode::BoundPSOsOnly);
 				LastAutoSaveTimeLogBoundPSO = FPlatformTime::Seconds();
-				LastAutoSaveNum = FPipelineFileCache::NumPSOsLogged();
+				LastAutoSaveNum = FPipelineFileCacheManager::NumPSOsLogged();
 			}
 		}
 	}
 
-	PollShutdownItems();
+	UE::ShaderPipeline::FShutdownContainer::Get()->Tick();
+
+	if(CurrentCacheTask)
+	{
+		CurrentCacheTask->Tick(DeltaTime);
+	}
+}
+
+void FShaderPipelineCacheTask::Tick(float DeltaTime)
+{
+	if (LastPrecompileRHIFence.GetReference() && LastPrecompileRHIFence->IsComplete())
+	{
+		LastPrecompileRHIFence = nullptr;
+	}
 
+	if (PrecompileStartTime > 0.0)
+	{
+		TotalPrecompileWallTime = float(FPlatformTime::Seconds() - PrecompileStartTime);
+	}
+	
 	if (PrecompileStartTime == 0.0 && (PreFetchedTasks.Num() || FetchTasks.Num() || OrderedCompileTasks.Num()))
 	{
 		PrecompileStartTime = FPlatformTime::Seconds();
@@ -1431,22 +1807,22 @@ void FShaderPipelineCache::Tick( float DeltaTime )
 
 		uint32 End = FPlatformTime::Cycles();
 
-		if (BatchTime > 0.0f)
+		if (FShaderPipelineCache::BatchTime > 0.0f)
 		{
 			float ElapsedMs = FPlatformTime::ToMilliseconds(End - Start);
-			if (ElapsedMs < BatchTime)
+			if (ElapsedMs < FShaderPipelineCache::BatchTime)
 			{
-				BatchSize++;
+				FShaderPipelineCache::BatchSize++;
 			}
-			else if (ElapsedMs > BatchTime)
+			else if (ElapsedMs > FShaderPipelineCache::BatchTime)
 			{
-				if (BatchSize > 1)
+				if (FShaderPipelineCache::BatchSize > 1)
 				{
-					BatchSize--;
+					FShaderPipelineCache::BatchSize--;
 				}
 				else
 				{
-					UE_LOG(LogRHI, Warning, TEXT("FShaderPipelineCache: Cannot reduce BatchSize below 1 to meet target of %f ms, elapsed time was %f ms)"), BatchTime, ElapsedMs);
+					UE_LOG(LogRHI, Warning, TEXT("FShaderPipelineCache: Cannot reduce BatchSize below 1 to meet target of %f ms, elapsed time was %f ms)"), FShaderPipelineCache::BatchTime, ElapsedMs);
 				}
 			}
 		}
@@ -1455,9 +1831,9 @@ void FShaderPipelineCache::Tick( float DeltaTime )
 	if (ReadyForNextBatch() && (OrderedCompileTasks.Num() || FetchTasks.Num()))
 	{
         uint32 Num = 0;
-        if (BatchSize > static_cast<uint32>(FetchTasks.Num()))
+        if (FShaderPipelineCache::BatchSize > static_cast<uint32>(FetchTasks.Num()))
         {
-            Num = BatchSize - FetchTasks.Num();
+            Num = FShaderPipelineCache::BatchSize - FetchTasks.Num();
         }
         Num = FMath::Min(Num, (uint32)OrderedCompileTasks.Num());
             
@@ -1484,29 +1860,28 @@ void FShaderPipelineCache::Tick( float DeltaTime )
 					FetchTasks.AddTail(Entry);
 					
 	            	Iterator.RemoveCurrent();
-                    FPlatformAtomics::InterlockedIncrement(&TotalActiveTasks);
-                    FPlatformAtomics::InterlockedDecrement(&TotalWaitingTasks);
+                    TotalActiveTasks++;
+                    TotalWaitingTasks--;
                     --Num;
             	}
             }
 			
-			FPipelineFileCache::FetchPSODescriptors(NewBatch);
+			FPipelineFileCacheManager::FetchPSODescriptors(PSOCacheKey, NewBatch);
 		}
 
-        if (static_cast<uint32>(FetchTasks.Num()) > BatchSize)
+        if (static_cast<uint32>(FetchTasks.Num()) > FShaderPipelineCache::BatchSize)
         {
-            UE_LOG(LogRHI, Warning, TEXT("FShaderPipelineCache: Attempting to pre-compile more jobs (%d) than the batch size (%d)"), FetchTasks.Num(), BatchSize);
+            UE_LOG(LogRHI, Warning, TEXT("FShaderPipelineCache: Attempting to pre-compile more jobs (%d) than the batch size (%d)"), FetchTasks.Num(), FShaderPipelineCache::BatchSize);
         }
         
 		PreparePipelineBatch(FetchTasks);
 	}
 	
-	
 	if(CVarPSOFileCacheGameFileMaskEnabled.GetValueOnAnyThread())
 	{
-		if(FPlatformAtomics::AtomicRead(&TotalActiveTasks) + FPlatformAtomics::AtomicRead(&TotalWaitingTasks) == 0)
+		if(TotalActiveTasks + TotalWaitingTasks == 0)
 		{
-			const uint64 Mask = FPipelineFileCache::GetGameUsageMask();
+			const uint64 Mask = FPipelineFileCacheManager::GetGameUsageMask();
 			bool bAlreadyInSet = false;
 			CompletedMasks.Add(Mask, &bAlreadyInSet);
 			if(!bAlreadyInSet)
@@ -1517,28 +1892,30 @@ void FShaderPipelineCache::Tick( float DeltaTime )
 	}
 	
 #if STATS
-	int64 ActiveTaskCount = FMath::Max((int64)0, FPlatformAtomics::AtomicRead(&TotalActiveTasks));
-    int64 WaitingTaskCount = FMath::Max((int64)0, FPlatformAtomics::AtomicRead(&TotalWaitingTasks));
+	int64 ActiveTaskCount = FMath::Max((int64_t)0, TotalActiveTasks.load());
+    int64 WaitingTaskCount = FMath::Max((int64_t)0, TotalWaitingTasks.load());
 	SET_DWORD_STAT(STAT_ShaderPipelineTaskCount, ActiveTaskCount+WaitingTaskCount);
     SET_DWORD_STAT(STAT_ShaderPipelineWaitingTaskCount, WaitingTaskCount);
     SET_DWORD_STAT(STAT_ShaderPipelineActiveTaskCount, ActiveTaskCount);
 	
+	using UE::ShaderPipeline::FShutdownContainer;
+
 	// Calc in one place otherwise this computation will be splattered all over the place - this will not be exact but counts the expensive bits
 	int64 InUseMemory = OrderedCompileTasks.GetAllocatedSize() + 
-						CompiledHashes.GetAllocatedSize() + 
 						ReadTasks.GetAllocatedSize() + 
 						CompileTasks.GetAllocatedSize() + 
-						ShutdownReadCompileTasks.GetAllocatedSize();
+						FShaderPipelineCache::Get()->CompiledHashes.GetAllocatedSize() +
+						FShutdownContainer::Get()->ShutdownReadCompileTasks.GetAllocatedSize();
 	if(ActiveTaskCount+WaitingTaskCount > 0)
 	{
-		InUseMemory += (ReadTasks.Num() + CompileTasks.Num() + ShutdownReadCompileTasks.Num()) * (sizeof(FShaderPipelineCacheArchive));
-		InUseMemory += (FetchTasks.Num() + ShutdownFetchTasks.Num()) * (sizeof(FPipelineCacheFileFormatPSORead));
+		InUseMemory += (ReadTasks.Num() + CompileTasks.Num() + FShutdownContainer::Get()->ShutdownReadCompileTasks.Num()) * (sizeof(FShaderPipelineCacheArchive));
+		InUseMemory += (FetchTasks.Num() + FShutdownContainer::Get()->ShutdownFetchTasks.Num()) * (sizeof(FPipelineCacheFileFormatPSORead));
 		for (TDoubleLinkedList<FPipelineCacheFileFormatPSORead*>::TIterator It(FetchTasks.GetHead()); It; ++It)
 		{
 			FPipelineCacheFileFormatPSORead* Entry = *It;
 			InUseMemory += Entry->Data.Num();
 		}
-		for (TDoubleLinkedList<FPipelineCacheFileFormatPSORead*>::TIterator It(ShutdownFetchTasks.GetHead()); It; ++It)
+		for (TDoubleLinkedList<FPipelineCacheFileFormatPSORead*>::TIterator It(FShutdownContainer::Get()->ShutdownFetchTasks.GetHead()); It; ++It)
 		{
 			FPipelineCacheFileFormatPSORead* Entry = *It;
 			InUseMemory += Entry->Data.Num();
@@ -1572,191 +1949,141 @@ static bool PreCompileMaskComparison(uint64 ReferenceGameMask, uint64 PSOMask)
 	return bIgnoreGameMask || ((UsageMask & (7l << (MaxQualityCount*3+MaxPlaylistCount))) && (UsageMask & (7 << (MaxQualityCount*3))) && (UsageMask & (63 << MaxQualityCount*2)) && (UsageMask & (63 << MaxQualityCount)) && (UsageMask & 63));
 }
 
-bool FShaderPipelineCache::Open(FString const& Name, EShaderPlatform Platform)
+bool FShaderPipelineCacheTask::Open(const FString& Key, FString const& Name, EShaderPlatform Platform)
 {
-	FileName = Name;
-    CurrentPlatform = Platform;
-	
-	bool bOK = FPipelineFileCache::OpenPipelineFileCache(Name, Platform, CacheFileGuid);
-	if (bOK)
+	PSOCacheKey = Key;
+	ShaderPlatform = Platform;
+
+	if (bIsUserCache)
 	{
-		FScopeLock Lock(&Mutex);
-		
-		Flush();
+		return FPipelineFileCacheManager::OpenUserPipelineFileCache(Key, Name, Platform, CacheFileGuid);
+	}
+	else
+	{
+		return FPipelineFileCacheManager::OpenPipelineFileCache(Key, Name, Platform, CacheFileGuid);
+	}
+}
+
+void FShaderPipelineCacheTask::BeginPrecompilingPipelineCache()
+{
+	Flush();
+
+	FShaderPipelineCacheTask::TotalPrecompileWallTime = 0.0f;
+	FShaderPipelineCacheTask::PrecompileStartTime = 0.0;
+	FShaderPipelineCacheTask::TotalPrecompileTasks = 0;
+
+	// these should have been zeroed when the previous task finished.
+	check(FShaderPipelineCacheTask::TotalActiveTasks == 0);
+	check(FShaderPipelineCacheTask::TotalWaitingTasks == 0);
+	check(FShaderPipelineCacheTask::TotalCompleteTasks == 0);
+	check(FShaderPipelineCacheTask::TotalPrecompileTime == 0);
+
+	if (FShaderPipelineCache::OnCachePreOpen.IsBound())
+	{
+		FShaderPipelineCache::OnCachePreOpen.Broadcast(PSOCacheKey, ShaderPlatform, bReady);
+	}
 		
-		if (OnCachePreOpen.IsBound())
+	if(bReady)
+	{
+		int32 Order = (int32)FPipelineFileCacheManager::PSOOrder::Default;
+
+		if (!GConfig->GetInt(FShaderPipelineCacheConstants::SectionHeading, FShaderPipelineCacheConstants::SortOrderKey, Order, *GGameUserSettingsIni))
 		{
-			OnCachePreOpen.Broadcast(Name, Platform, bReady);
+			GConfig->GetInt(FShaderPipelineCacheConstants::SectionHeading, FShaderPipelineCacheConstants::SortOrderKey, Order, *GGameIni);
 		}
-		
-		if(bReady)
-		{
-			uint64 PreCompileMask = (uint64)CVarPSOFileCachePreCompileMask.GetValueOnAnyThread();
-			FPipelineFileCache::SetGameUsageMaskWithComparison(PreCompileMask, PreCompileMaskComparison);
 
-			int32 Order = (int32)FPipelineFileCache::PSOOrder::Default;
-			
-			if(!GConfig->GetInt(FShaderPipelineCacheConstants::SectionHeading, FShaderPipelineCacheConstants::SortOrderKey, Order, *GGameUserSettingsIni))
-			{
-				GConfig->GetInt(FShaderPipelineCacheConstants::SectionHeading, FShaderPipelineCacheConstants::SortOrderKey, Order, *GGameIni);
-			}
+		TArray<FPipelineCachePSOHeader> LocalPreFetchedTasks;
 
-			TArray<FPipelineCachePSOHeader> LocalPreFetchedTasks;
-			FPipelineFileCache::GetOrderedPSOHashes(LocalPreFetchedTasks, (FPipelineFileCache::PSOOrder)Order, (int64)CVarPSOFileCacheMinBindCount.GetValueOnAnyThread(), CompiledHashes);
-			// Iterate over all the tasks we haven't yet begun to read data for - these are the 'waiting' tasks
-			int64 Count = 0;
-			for (auto const& Task : LocalPreFetchedTasks)
+		FPipelineFileCacheManager::GetOrderedPSOHashes(PSOCacheKey, LocalPreFetchedTasks, (FPipelineFileCacheManager::PSOOrder)Order, (int64)CVarPSOFileCacheMinBindCount.GetValueOnAnyThread(), FShaderPipelineCache::Get()->CompiledHashes);
+		// Iterate over all the tasks we haven't yet begun to read data for - these are the 'waiting' tasks
+		int64 EligibleTaskCount = LocalPreFetchedTasks.Num();
+
+		int32 MissingShaders = 0;
+		int32 NewSize = Algo::StableRemoveIf(LocalPreFetchedTasks, [&MissingShaders](FPipelineCachePSOHeader& Task)
+		{
+			bool bHasShaders = true;
+			for (FSHAHash const& Hash : Task.Shaders)
 			{
-				bool bHasShaders = true;
-				for (FSHAHash const& Hash : Task.Shaders)
-				{
-					bHasShaders &= FShaderCodeLibrary::ContainsShaderCode(Hash);
-				}
-				if (bHasShaders)
-				{
-					Count++;
-				}
+				bHasShaders &= FShaderCodeLibrary::ContainsShaderCode(Hash);
 			}
-
-			FPlatformAtomics::InterlockedAdd(&TotalWaitingTasks, Count);
-
-			if (OnCachedOpened.IsBound())
+			if(!bHasShaders)
 			{
-				OnCachedOpened.Broadcast(Name, Platform, LocalPreFetchedTasks.Num(), CacheFileGuid, ShaderCachePrecompileContext);
+				MissingShaders++;
 			}
+			return !bHasShaders;
+		});
 
-			PreFetchedTasks = LocalPreFetchedTasks;
-			TotalPrecompileTasks = PreFetchedTasks.Num();
-			
-			bPreOptimizing = TotalPrecompileTasks > 0;
-			UE_LOG(LogRHI, Display, TEXT("Opened pipeline cache and enqueued %d of %d tasks for precompile with BatchSize %d and BatchTime %f."), Count, PreFetchedTasks.Num(), BatchSize, BatchTime);
-		}
-		else
+		LocalPreFetchedTasks.SetNum(NewSize);
+		int64 PossibleTasks = LocalPreFetchedTasks.Num();
+		TotalWaitingTasks = PossibleTasks;
+
+		if (FShaderPipelineCache::OnCachedOpened.IsBound())
 		{
-			UE_LOG(LogRHI, Display, TEXT("Opened pipeline cache - precompile deferred on UsageMask."));
+			FShaderPipelineCache::OnCachedOpened.Broadcast(PSOCacheKey, ShaderPlatform, LocalPreFetchedTasks.Num(), CacheFileGuid, PrecompileContext);
 		}
-	}
 
-	UE_CLOG(!bOK, LogRHI, Display, TEXT("Failed to open default shader pipeline cache for %s using shader platform %d."), *Name, (uint32)Platform);
-	
-	bOpened = bOK;
+		// Copy any new items over to our 'internal' safe array
+		check(PreFetchedTasks.IsEmpty());
+		check(OrderedCompileTasks.IsEmpty());
 
-	// OnPrecompilationBegin can bring up a modal loading screen, so we call it outside the scope lock to prevent deadlocks
-	if (bOK && OnPrecompilationBegin.IsBound())
-	{
-		OnPrecompilationBegin.Broadcast(PreFetchedTasks.Num(), ShaderCachePrecompileContext);
-	}
-
-	return bOK;
-}
-
-bool FShaderPipelineCache::Save(FPipelineFileCache::SaveMode Mode)
-{
-	FScopeLock Lock(&Mutex);
-		
-	bool bOK = FPipelineFileCache::SavePipelineFileCache(FileName, Mode);
-	UE_CLOG(!bOK, LogRHI, Warning, TEXT("Failed to save shader pipeline cache for %s using save mode %d."), *FileName, (uint32)Mode);
+		OrderedCompileTasks = MoveTemp(LocalPreFetchedTasks);
 
-	LastAutoSaveTime = FPlatformTime::Seconds();
+		TotalPrecompileTasks = OrderedCompileTasks.Num();
 
-	return bOK;
-}
+			
+		bPreOptimizing = TotalPrecompileTasks > 0;
+		int32 TotalPSOCount = FPipelineFileCacheManager::GetTotalPSOCount(PSOCacheKey);
+		UE_LOG(LogRHI, Display, TEXT("FShaderPipelineCache starting pipeline cache '%s' and enqueued %d tasks for precompile. (cache contains %d, %d eligible, %d had missing shaders. %d already compiled). BatchSize %d and BatchTime %f."),
+			*PSOCacheKey, TotalPrecompileTasks, TotalPSOCount, EligibleTaskCount, MissingShaders, TotalPSOsCompiled, FShaderPipelineCache::BatchSize, FShaderPipelineCache::BatchTime);
 
-void FShaderPipelineCache::Close()
-{
-	FScopeLock Lock(&Mutex);
-		
-	if(GConfig)
-	{
-		GConfig->SetString(FShaderPipelineCacheConstants::SectionHeading, FShaderPipelineCacheConstants::LastOpenedKey, *FileName, *GGameUserSettingsIni);
-		GConfig->Flush(false, *GGameUserSettingsIni);
-	}
-	
-	// Log all bound PSOs
-	if (GetShaderPipelineCacheSaveBoundPSOLog())
-	{
-		Save(FPipelineFileCache::SaveMode::BoundPSOsOnly);
+		if (FShaderPipelineCache::OnPrecompilationBegin.IsBound())
+		{
+			FShaderPipelineCache::OnPrecompilationBegin.Broadcast(TotalPrecompileTasks, PrecompileContext);
+		}
 	}
-	
-	if (GetPSOFileCacheSaveUserCache())
+	else
 	{
-		Save(FPipelineFileCache::SaveMode::Incremental);
+		UE_LOG(LogRHI, Display, TEXT("FShaderPipelineCache pipeline cache %s - precompile deferred on UsageMask."), *PSOCacheKey);
 	}
-	
+}
+
+void FShaderPipelineCacheTask::Close()
+{
 	// Signal flush of outstanding work to allow restarting for a new cache file
 	Flush();
     
-    if (OnCachedClosed.IsBound())
+    if (FShaderPipelineCache::OnCachedClosed.IsBound())
     {
-        OnCachedClosed.Broadcast(FileName, CurrentPlatform);
+		FShaderPipelineCache::OnCachedClosed.Broadcast(PSOCacheKey, ShaderPlatform);
     }
 	
 	bOpened = false;
 
-	FPipelineFileCache::ClosePipelineFileCache();
+	FPipelineFileCacheManager::CloseUserPipelineFileCache();
+}
 
-	//
-	// Clean up cached RHI resources
-	//
-	for (auto Pair : BlendStateCache)
+bool FShaderPipelineCacheTask::OnShaderLibraryStateChanged(FShaderPipelineCache::ELibraryState State, EShaderPlatform Platform, FString const& Name, int32 ComponentID)
+{
+	check(!bIsUserCache);
+	if (State == FShaderPipelineCache::Opened && Name == FApp::GetProjectName() && Platform == ShaderPlatform && !bOpened)
 	{
-		Pair.Value->Release();
+		Close();
 	}
-	BlendStateCache.Empty();
-	
-	for (auto Pair : RasterizerStateCache)
+	else if(State == FShaderPipelineCache::OpenedComponent)
 	{
-		Pair.Value->Release();
+		check(ComponentID != INDEX_NONE);
 	}
-	RasterizerStateCache.Empty();
-	
-	for (auto Pair : DepthStencilStateCache)
-	{
-		Pair.Value->Release();
-	}
-	DepthStencilStateCache.Empty();
-}
 
-void FShaderPipelineCache::OnShaderLibraryStateChanged(ELibraryState State, EShaderPlatform Platform, FString const& Name)
-{
-    FScopeLock Lock(&Mutex);
-	
-	if (State == FShaderPipelineCache::Opened && Name == FApp::GetProjectName() && Platform == CurrentPlatform && !bOpened)
+	if (ensure(!Name.IsEmpty()))
 	{
-		Close();
-		FString LastOpenedName;
-		if ((!GConfig->GetString(FShaderPipelineCacheConstants::SectionHeading, FShaderPipelineCacheConstants::LastOpenedKey, LastOpenedName, *GGameUserSettingsIni) && !GConfig->GetString(FShaderPipelineCacheConstants::SectionHeading, FShaderPipelineCacheConstants::LastOpenedKey, LastOpenedName, *GGameIni)) && !LastOpenedName.Len())
-		{
-			LastOpenedName = FApp::GetProjectName();
-		}
-		Open(LastOpenedName, Platform);
+		// for bundled files the key == name
+		const FString& Key = Name;
+		bool bDidOpen = Open(Key, Name, Platform);
+		UE_LOG(LogRHI, Display, TEXT("Shader library state change %d, pipeline cache %s was opened=%d"), State, *Name, bDidOpen);
+		return bDidOpen;
 	}
-	
-    // Copy any new items over to our 'internal' safe array
-    if (PreFetchedTasks.Num())
-    {
-        OrderedCompileTasks = PreFetchedTasks;
-        PreFetchedTasks.Empty();
-    }
-    
-    // Iterate over all the tasks we haven't yet begun to read data for - these are the 'waiting' tasks
-    int64 Count = 0;
-    for (auto const& Task : OrderedCompileTasks)
-    {
-        bool bHasShaders = true;
-        for (FSHAHash const& Hash : Task.Shaders)
-        {
-            bHasShaders &= FShaderCodeLibrary::ContainsShaderCode(Hash);
-        }
-        if (bHasShaders)
-        {
-            Count++;
-        }
-    }
-    
-    // Set the new waiting count that we can actually process.
-    FPlatformAtomics::InterlockedExchange(&TotalWaitingTasks, Count);
-	UE_LOG(LogRHI, Display, TEXT("Opened pipeline cache after state change and enqueued %d of %d tasks for precompile."), Count, OrderedCompileTasks.Num());
+
+	return false;
 }
 
 FRHIBlendState* FShaderPipelineCache::GetOrCreateBlendState(const FBlendStateInitializerRHI& Initializer)
diff --git a/Engine/Source/Runtime/RenderCore/Public/ShaderPipelineCache.h b/Engine/Source/Runtime/RenderCore/Public/ShaderPipelineCache.h
--- a/Engine/Source/Runtime/RenderCore/Public/ShaderPipelineCache.h
+++ b/Engine/Source/Runtime/RenderCore/Public/ShaderPipelineCache.h
@@ -7,24 +7,13 @@
  
 #pragma once
 
-#include "Async/TaskGraphInterfaces.h"
-#include "Containers/Array.h"
-#include "Containers/List.h"
-#include "Containers/Map.h"
-#include "Containers/Set.h"
-#include "Containers/UnrealString.h"
 #include "CoreMinimal.h"
-#include "Delegates/Delegate.h"
-#include "Delegates/DelegateCombinations.h"
-#include "HAL/CriticalSection.h"
-#include "HAL/Platform.h"
-#include "Misc/Guid.h"
-#include "PipelineFileCache.h"
-#include "RHI.h"
-#include "RHIDefinitions.h"
+#include "Containers/List.h"
 #include "Stats/Stats.h"
-#include "Stats/Stats2.h"
+#include "RHI.h"
 #include "TickableObjectRenderThread.h"
+#include "PipelineFileCache.h"
+#include "Delegates/DelegateCombinations.h"
 
 class FBlendStateInitializerRHI;
 class FRHIBlendState;
@@ -40,7 +29,7 @@ struct FRasterizerStateInitializerRHI;
  * FShaderPipelineCache:
  * The FShaderPipelineCache provides the new Pipeline State Object (PSO) logging, serialisation & precompilation mechanism that replaces FShaderCache.
  * Caching Pipeline State Objects and serialising the initialisers to disk allows for precompilation of these states the next time the game is run, which reduces hitching.
- * To achieve this FShaderPipelineCache relies upon FShaderCodeLibrary & "Share Material Shader Code" and the RHI-level backend FPipelineFileCache.
+ * To achieve this FShaderPipelineCache relies upon FShaderCodeLibrary & "Share Material Shader Code" and the RHI-level backend FPipelineFileCacheManager.
  *
  * Basic Runtime Usage:
  * - Enable the cache with r.ShaderPipelineCache.Enabled = 1, which allows the pipeline cache to load existing data from disk and precompile it.
@@ -87,30 +76,8 @@ struct FRasterizerStateInitializerRHI;
  */
 class RENDERCORE_API FShaderPipelineCache : public FTickableObjectRenderThread
 {
-	struct CompileJob
-	{
-		FPipelineCacheFileFormatPSO PSO;
-		FShaderPipelineCacheArchive* ReadRequests;
-
-		/** Tracks whether the shaders were preloaded. */
-		bool bShadersPreloaded = false;
-
-		/** 
-		 * Schedules preload for all the shaders in the PSO. No shaders are preloaded if one of them isn't in the library.
-		 * Multiple jobs can request to preload the same shaders and then release them independently.
-		 * 
-		 * @param OutRequiredShaders shader hashes of shaders that are actually required by this PSO
-		 * @param bOutCompatible if set to false, then this PSO type isn't supported (RTX as of now)
-		 */
-		bool PreloadShaders(TSet<FSHAHash>& OutRequiredShaders, bool& bOutCompatible);
-
-		/**
-		 * Releases preloaded entries for the shaders in the PSO. Symmetric to to PreloadShaders.
-		 * Multiple jobs can request to preload the same shaders and then release them independently.
-		 */
-		void ReleasePreloadedShaders();
-	};
-
+	// the FShaderPipelineCacheTask class represents a pipeline precompile of a single PSO file cache.
+	friend class FShaderPipelineCacheTask;
 public:
 	/**
 	 * Initializes the shader pipeline cache for the desired platform, called by the engine.
@@ -136,29 +103,30 @@ public:
 	/** Resumes precompilation batching. */
 	static void ResumeBatching();
 	
-	/** Returns the number of pipelines waiting for precompilation. */
+	/** true if pipeline cache(s) are precompiling. */
+	static bool IsPrecompiling();
+
+	/** Returns the number of pipelines waiting for precompilation of the current PSOFC task. 
+		if there are multiple PSOFCs pending this will return a minimum value of 1. It may increase as subsequent caches are processed.*/
 	static uint32 NumPrecompilesRemaining();
     
-    /** Returns the number of pipelines actively being precompiled this frame. */
-    static uint32 NumPrecompilesActive();
-
 	/** Opens the shader pipeline cache file with either the LastOpened setting if available, or the project name otherwise */
 	static bool OpenPipelineFileCache(EShaderPlatform Platform);
 
-	/** Opens the shader pipeline cache file with the given name and shader platform. */
-	static bool OpenPipelineFileCache(FString const& Name, EShaderPlatform Platform);
+ 	/** Opens the user pipeline cache file with either the LastOpened setting if available, or the project name otherwise */
+ 	static bool OpenUserPipelineFileCache(EShaderPlatform Platform);
 	
 	/** Saves the current shader pipeline cache to disk using one of the defined save modes, Fast uses an incremental approach whereas Slow will consolidate all data into the file. */
-	static bool SavePipelineFileCache(FPipelineFileCache::SaveMode Mode);
+	static bool SavePipelineFileCache(FPipelineFileCacheManager::SaveMode Mode);
 	
 	/** Closes the existing pipeline cache, allowing it to be reopened with a different file and/or shader platform. Will implicitly invoke a Fast Save. */
-	static void ClosePipelineFileCache();
+	static void CloseUserPipelineFileCache();
 
 	static int32 GetGameVersionForPSOFileCache();
 	
 	/**
 	 * Set the current PSO Game Usage Mask and comparsion function .  Returns true if this mask is different from the old mask or false if not or the cache system is disabled or if the mask feature is disabled
-	 * Any new PSO's found will be logged with this value or existing PSO should have their masks updated.  See FPipelineFileCache for more details.
+	 * Any new PSO's found will be logged with this value or existing PSO should have their masks updated.  See FPipelineFileCacheManager for more details.
 	 */
 	static bool SetGameUsageMaskWithComparison(uint64 Mask, FPSOMaskComparisonFn InComparisonFnPtr);
 
@@ -168,9 +136,9 @@ public:
 	virtual ~FShaderPipelineCache();
 
 	bool IsTickable() const;
-	
+
 	void Tick( float DeltaTime );
-	
+
 	bool NeedsRenderingResumedForRenderingThreadTick() const;
 	
 	TStatId GetStatId() const;
@@ -178,11 +146,12 @@ public:
 	enum ELibraryState
 	{
 		Opened,
+		OpenedComponent,
 		Closed
 	};
 	
     /** Called by FShaderCodeLibrary to notify us that the shader code library state changed and shader availability will need to be re-evaluated */
-    static void ShaderLibraryStateChanged(ELibraryState State, EShaderPlatform Platform, FString const& Name);
+    static void ShaderLibraryStateChanged(ELibraryState State, EShaderPlatform Platform, FString const& Name, int32 ComponentID);
 
 	/** Allows handlers of FShaderCacheOpenedDelegate to send state to begin and complete pre-compilation delegates */
 	class FShaderCachePrecompileContext
@@ -245,20 +214,16 @@ public:
     static FShaderPrecompilationCompleteDelegate& GetPrecompilationCompleteDelegate() { return OnPrecompilationComplete; }
 
 private:
-	bool Open(FString const& Name, EShaderPlatform Platform);
-	bool Save(FPipelineFileCache::SaveMode Mode);
-	void Close();
-
-	bool Precompile(FRHICommandListImmediate& RHICmdList, EShaderPlatform Platform, FPipelineCacheFileFormatPSO const& PSO);
-	void PreparePipelineBatch(TDoubleLinkedList<FPipelineCacheFileFormatPSORead*>& PipelineBatch);
-	bool ReadyForPrecompile();
-	void PrecompilePipelineBatch();
-	bool ReadyForNextBatch() const;
+	static inline FShaderPipelineCache* Get() { return ShaderPipelineCache;	}
+
+	bool OpenPipelineFileCacheInternal(bool bUserCache, const FString& PSOCacheKey, const FString& CacheName, EShaderPlatform Platform);
+
+	void BeginNextPrecompileCacheTask();
+	void BeginNextPrecompileCacheTaskInternal();
 	bool ReadyForAutoSave() const;
-	void PollShutdownItems();
-	void Flush(bool bClearCompiled = true);
-    
-    void OnShaderLibraryStateChanged(ELibraryState State, EShaderPlatform Platform, FString const& Name);
+
+	void ClearPendingPrecompileTaskQueue();
+	void EnqueuePendingPrecompileTask(const FString& PSOCacheKey);
 
 	FRHIBlendState* GetOrCreateBlendState(const FBlendStateInitializerRHI& Initializer);
 	FRHIRasterizerState* GetOrCreateRasterizerState(const FRasterizerStateInitializerRHI& Initializer);
@@ -266,51 +231,49 @@ private:
 	
 private:
 	static FShaderPipelineCache* ShaderPipelineCache;
-	TArray<CompileJob> ReadTasks;
-	TArray<CompileJob> CompileTasks;
-	TArray<FPipelineCachePSOHeader> OrderedCompileTasks;
-	TDoubleLinkedList<FPipelineCacheFileFormatPSORead*> FetchTasks;
-	TSet<uint32> CompiledHashes;
-	FString FileName;
-    EShaderPlatform CurrentPlatform;
-	FGuid CacheFileGuid;
-	uint32 BatchSize;
-	float BatchTime;
+	static uint32 BatchSize;
+	static float BatchTime;
 	bool bPaused;
-	bool bOpened;
-	bool bReady;
-	bool bPreOptimizing;
     int32 PausedCount;
-	FShaderCachePrecompileContext ShaderCachePrecompileContext;
-	
-	MS_ALIGN(8) volatile int64 TotalActiveTasks GCC_ALIGN(8);
-	MS_ALIGN(8) volatile int64 TotalWaitingTasks GCC_ALIGN(8);
-	MS_ALIGN(8) volatile int64 TotalCompleteTasks GCC_ALIGN(8);
-	MS_ALIGN(8) volatile int64 TotalPrecompileTime GCC_ALIGN(8);
-	double PrecompileStartTime;
-
-	FCriticalSection Mutex;
-	TArray<FPipelineCachePSOHeader> PreFetchedTasks;
-	
-	FGraphEventRef LastPrecompileRHIFence;
+	double LastAutoSaveTime = 0.0;
+	double LastAutoSaveTimeLogBoundPSO = 0.0;
+	int32 LastAutoSaveNum = 0;
+	TSet<uint32> CompiledHashes;
+
+	static FString UserCacheTaskKey;
+
+	struct FPipelineCachePSOContext
+	{
+		FString FileName;
+		EShaderPlatform CurrentPlatform;
 
-	TArray<CompileJob> ShutdownReadCompileTasks;
-	TDoubleLinkedList<FPipelineCacheFileFormatPSORead*> ShutdownFetchTasks;
+		FShaderCachePrecompileContext PrecompileContext;
+	};
 	
+	TMap<FString, TUniquePtr<class FShaderPipelineCacheTask>> ShaderCacheTasks;
+
+	class FShaderPipelineCacheTask* GetTask(const FString& Name) {
+		TUniquePtr<class FShaderPipelineCacheTask>* CacheTask = FShaderPipelineCache::Get()->ShaderCacheTasks.Find(Name);
+		return CacheTask ? CacheTask->Get() : nullptr;
+	}
+
+	// The list of pending PSO cache names to precompile.
+	TArray<FString> PendingPrecompilePSOCacheKeys;
+
+	// The name of the current precompiling PSO cache.
+	FString CurrentPrecompilingPSOFCKey;
+
+	// Used for threadsafe query of pending pipeline compiles.
+	std::atomic_bool bHasShaderPipelineTask = false;
+
+	mutable FCriticalSection Mutex;
+
 	static FShaderCachePreOpenDelegate OnCachePreOpen;
     static FShaderCacheOpenedDelegate OnCachedOpened;
     static FShaderCacheClosedDelegate OnCachedClosed;
 	static FShaderPrecompilationBeginDelegate OnPrecompilationBegin;
 	static FShaderPrecompilationCompleteDelegate OnPrecompilationComplete;
 
-	double LastAutoSaveTime;
-	double LastAutoSaveTimeLogBoundPSO;
-	int32 LastAutoSaveNum;
-	
-	TSet<uint64> CompletedMasks;
-	float TotalPrecompileWallTime;
-	int64 TotalPrecompileTasks;
-
 	TMap<FBlendStateInitializerRHI, FRHIBlendState*> BlendStateCache;
 	TMap<FRasterizerStateInitializerRHI, FRHIRasterizerState*> RasterizerStateCache;
 	TMap<FDepthStencilStateInitializerRHI, FRHIDepthStencilState*> DepthStencilStateCache;
