Removing unused variables.

diff --git a/Engine/Source/Runtime/RHI/Private/RHI.cpp b/Engine/Source/Runtime/RHI/Private/RHI.cpp
--- a/Engine/Source/Runtime/RHI/Private/RHI.cpp
+++ b/Engine/Source/Runtime/RHI/Private/RHI.cpp
@@ -1720,7 +1720,6 @@ RHI_API void GetShadingPathName(ERHIShadingPath::Type InShadingPath, FName& OutN
 }
 
 static FName NAME_PLATFORM_WINDOWS(TEXT("Windows"));
-static FName NAME_PLATFORM_XBOXONE(TEXT("XboxOne"));
 static FName NAME_PLATFORM_ANDROID(TEXT("Android"));
 static FName NAME_PLATFORM_IOS(TEXT("IOS"));
 static FName NAME_PLATFORM_MAC(TEXT("Mac"));
diff --git a/Engine/Source/Runtime/RHI/Public/RHIShaderFormatDefinitions.inl b/Engine/Source/Runtime/RHI/Public/RHIShaderFormatDefinitions.inl
--- a/Engine/Source/Runtime/RHI/Public/RHIShaderFormatDefinitions.inl
+++ b/Engine/Source/Runtime/RHI/Public/RHIShaderFormatDefinitions.inl
@@ -15,8 +15,6 @@ static FName NAME_PCD3D_ES3_1(TEXT("PCD3D_ES31"));
 static FName NAME_GLSL_150_ES31(TEXT("GLSL_150_ES31"));
 static FName NAME_GLSL_ES3_1_ANDROID(TEXT("GLSL_ES3_1_ANDROID"));
 
-static FName NAME_SF_XBOXONE_D3D12(TEXT("SF_XBOXONE_D3D12"));
-
 static FName NAME_SF_METAL(TEXT("SF_METAL"));
 static FName NAME_SF_METAL_MRT(TEXT("SF_METAL_MRT"));
 static FName NAME_SF_METAL_TVOS(TEXT("SF_METAL_TVOS"));
