Update size assertion for EPixelFormat in PSO file cache and move it to static_assert.

#rnx

diff --git a/Engine/Source/Runtime/RHI/Private/PipelineFileCache.cpp b/Engine/Source/Runtime/RHI/Private/PipelineFileCache.cpp
--- a/Engine/Source/Runtime/RHI/Private/PipelineFileCache.cpp
+++ b/Engine/Source/Runtime/RHI/Private/PipelineFileCache.cpp
@@ -627,6 +627,15 @@ void FPipelineCacheFileFormatPSO::GraphicsDescriptor::AddStateToReadableString(T
 
 bool FPipelineCacheFileFormatPSO::GraphicsDescriptor::StateFromString(const FStringView& Src)
 {
+	static_assert(sizeof(EPixelFormat) == 1);
+	static_assert(sizeof(ERenderTargetLoadAction) == 1);
+	static_assert(sizeof(ERenderTargetStoreAction) == 1);
+	static_assert(sizeof(DepthLoad) == 1);
+	static_assert(sizeof(DepthStore) == 1);
+	static_assert(sizeof(StencilLoad) == 1);
+	static_assert(sizeof(StencilStore) == 1);
+	static_assert(sizeof(PrimitiveType) == 4);
+
 	constexpr int32 PartCount = FPipelineCacheGraphicsDescPartsNum;
 
 	TArray<FStringView, TInlineAllocator<PartCount>> Parts;
@@ -647,12 +656,12 @@ bool FPipelineCacheFileFormatPSO::GraphicsDescriptor::StateFromString(const FStr
 	RasterizerState.FromString(*PartIt++);
 	DepthStencilState.FromString(*PartIt++);
 
-	check(PartEnd - PartIt >= 3 && sizeof(EPixelFormat) == sizeof(uint32)); //not a very robust parser
+	check(PartEnd - PartIt >= 3); //not a very robust parser
 	LexFromString(MSAASamples, *PartIt++);
 	LexFromString((uint32&)DepthStencilFormat, *PartIt++);
 	LexFromString(DepthStencilFlags, *PartIt++);
 
-	check(PartEnd - PartIt >= 5 && sizeof(DepthLoad) == 1 && sizeof(StencilLoad) == 1 && sizeof(DepthStore) == 1 && sizeof(StencilStore) == 1 && sizeof(PrimitiveType) == 4); //not a very robust parser
+	check(PartEnd - PartIt >= 5); //not a very robust parser
 	LexFromString((uint32&)DepthLoad, *PartIt++);
 	LexFromString((uint32&)StencilLoad, *PartIt++);
 	LexFromString((uint32&)DepthStore, *PartIt++);
@@ -664,7 +673,7 @@ bool FPipelineCacheFileFormatPSO::GraphicsDescriptor::StateFromString(const FStr
 
 	for (int32 Index = 0; Index < MaxSimultaneousRenderTargets; Index++)
 	{
-		check(PartEnd - PartIt >= 4 && sizeof(ERenderTargetLoadAction) == 1 && sizeof(ERenderTargetStoreAction) == 1 && sizeof(EPixelFormat) == sizeof(uint32)); //not a very robust parser
+		check(PartEnd - PartIt >= 4); //not a very robust parser
 		LexFromString((uint32&)(RenderTargetFormats[Index]), *PartIt++);
 		ETextureCreateFlags RTFlags;
 		LexFromString(RTFlags, *PartIt++);
