Launching Editor with no feature levels selected wasn't falling back to any kind of default feature level.

diff --git a/Engine/Source/Runtime/RHI/Private/Windows/WindowsDynamicRHI.cpp b/Engine/Source/Runtime/RHI/Private/Windows/WindowsDynamicRHI.cpp
--- a/Engine/Source/Runtime/RHI/Private/Windows/WindowsDynamicRHI.cpp
+++ b/Engine/Source/Runtime/RHI/Private/Windows/WindowsDynamicRHI.cpp
@@ -587,18 +587,18 @@ static ERHIFeatureLevel::Type ChooseFeatureLevel(EWindowsRHI ChosenRHI, const TO
 	else
 	{
 		FeatureLevel = Config.GetHighestSupportedFeatureLevel(ChosenRHI);
-	}
-
-	// If we were forced to a specific RHI while not forced to a specific feature level and the project isn't configured for it, find the default Feature Level for that RHI
-	if (!FeatureLevel && ForcedRHI)
-	{
-		FeatureLevel = GetDefaultFeatureLevelForRHI(ChosenRHI);
 
-		if (FPlatformProperties::RequiresCookedData())
+		// If we were forced to a specific RHI while not forced to a specific feature level and the project isn't configured for it, find the default Feature Level for that RHI
+		if (!FeatureLevel)
 		{
-			const TCHAR* RHIName = ModuleNameFromWindowsRHI(ForcedRHI.GetValue());
-			const FString FeatureLevelName = LexToString(FeatureLevel.GetValue());
-			UE_LOG(LogRHI, Warning, TEXT("User requested RHI '%s' but that is not supported by this project's data. Defaulting to Feature Level '%s'."), RHIName, *FeatureLevelName);
+			FeatureLevel = GetDefaultFeatureLevelForRHI(ChosenRHI);
+
+			if (FPlatformProperties::RequiresCookedData())
+			{
+				const TCHAR* RHIName = ModuleNameFromWindowsRHI(ForcedRHI.GetValue());
+				const FString FeatureLevelName = LexToString(FeatureLevel.GetValue());
+				UE_LOG(LogRHI, Warning, TEXT("User requested RHI '%s' but that is not supported by this project's data. Defaulting to Feature Level '%s'."), RHIName, *FeatureLevelName);
+			}
 		}
 	}
 
