Initialize the GPUProfiler calibration timestamps to 0 to prevent tracing random values when the calibration query is not implemented.

diff --git a/Engine/Source/Runtime/RHI/Public/GPUProfiler.h b/Engine/Source/Runtime/RHI/Public/GPUProfiler.h
--- a/Engine/Source/Runtime/RHI/Public/GPUProfiler.h
+++ b/Engine/Source/Runtime/RHI/Public/GPUProfiler.h
@@ -147,8 +147,8 @@ struct RHI_API FGPUProfilerEventNodeFrame
 */
 struct FGPUTimingCalibrationTimestamp
 {
-	uint64 GPUMicroseconds;
-	uint64 CPUMicroseconds;
+	uint64 GPUMicroseconds = 0;
+	uint64 CPUMicroseconds = 0;
 };
 
 /**
