Fixed unsafe type casts in a few widely-used modules

#rnx

diff --git a/Engine/Source/Runtime/ApplicationCore/Private/Linux/LinuxApplication.cpp b/Engine/Source/Runtime/ApplicationCore/Private/Linux/LinuxApplication.cpp
--- a/Engine/Source/Runtime/ApplicationCore/Private/Linux/LinuxApplication.cpp
+++ b/Engine/Source/Runtime/ApplicationCore/Private/Linux/LinuxApplication.cpp
@@ -357,7 +357,7 @@ void FLinuxApplication::ProcessDeferredMessage( SDL_Event Event )
 					int32 BorderSizeX, BorderSizeY;
 					CurrentEventWindow->GetNativeBordersSize(BorderSizeX, BorderSizeY);
 
-					LinuxCursor->SetCachedPosition(motionEvent.x + Props.Location.X + BorderSizeX, motionEvent.y + Props.Location.Y + BorderSizeY);
+					LinuxCursor->SetCachedPosition((int32)(motionEvent.x + Props.Location.X + BorderSizeX), (int32)(motionEvent.y + Props.Location.Y + BorderSizeY));
 				}
 
 				if( !CurrentEventWindow->GetDefinition().HasOSWindowBorder )
@@ -365,7 +365,7 @@ void FLinuxApplication::ProcessDeferredMessage( SDL_Event Event )
 					if ( CurrentEventWindow->IsRegularWindow() )
 					{
 						FVector2D CurrentPosition = LinuxCursor->GetPosition();
-						MessageHandler->GetWindowZoneForPoint( CurrentEventWindow.ToSharedRef(), CurrentPosition.X - Props.Location.X, CurrentPosition.Y - Props.Location.Y );
+						MessageHandler->GetWindowZoneForPoint(CurrentEventWindow.ToSharedRef(), (int32)(CurrentPosition.X - Props.Location.X), (int32)(CurrentPosition.Y - Props.Location.Y));
 						MessageHandler->OnCursorSet();
 					}
 				}
@@ -463,7 +463,7 @@ void FLinuxApplication::ProcessDeferredMessage( SDL_Event Event )
 	case SDL_MOUSEWHEEL:
 		{
 			SDL_MouseWheelEvent *WheelEvent = &Event.wheel;
-			float Amount = WheelEvent->y * fMouseWheelScrollAccel;
+			float Amount = (float)WheelEvent->y * fMouseWheelScrollAccel;
 
 			MessageHandler->OnMouseWheel(Amount);
 		}
@@ -1191,18 +1191,18 @@ TCHAR FLinuxApplication::ConvertChar( SDL_Keysym Keysym )
 		return 0;
 	}
 
-	TCHAR Char = SDL_GetKeyFromScancode(Keysym.scancode);
+	TCHAR Char = (TCHAR)SDL_GetKeyFromScancode(Keysym.scancode);
 
     if (Keysym.mod & (KMOD_LSHIFT | KMOD_RSHIFT))
     {
         // Convert to uppercase (FIXME: what about CAPS?)
 		if( SDL_GetKeyFromScancode(Keysym.scancode)  >= 97 && SDL_GetKeyFromScancode(Keysym.scancode)  <= 122)
         {
-            return Keysym.sym - 32;
+            return (TCHAR)(Keysym.sym - 32);
         }
 		else if( SDL_GetKeyFromScancode(Keysym.scancode) >= 91 && SDL_GetKeyFromScancode(Keysym.scancode)  <= 93)
         {
-			return SDL_GetKeyFromScancode(Keysym.scancode) + 32; // [ \ ] -> { | }
+			return (TCHAR)(SDL_GetKeyFromScancode(Keysym.scancode) + 32); // [ \ ] -> { | }
         }
         else
         {
@@ -1899,7 +1899,7 @@ bool FLinuxApplication::IsMouseAttached() const
 
 	for (int i=0; i<9; i++)
 	{
-		Mouse[MouseIdx] = '0' + i;
+		Mouse[MouseIdx] = (char)('0' + i);
 		if (access(Mouse, F_OK) == 0)
 		{
 			return true;
@@ -2026,15 +2026,15 @@ void FLinuxApplication::SDLControllerState::UpdateHapticEffect()
 	{
 		Effect.type = SDL_HAPTIC_LEFTRIGHT;
 		Effect.leftright.length = 1000;
-		Effect.leftright.large_magnitude = 32767.0f * LargeValue;
-		Effect.leftright.small_magnitude = 32767.0f * SmallValue;
+		Effect.leftright.large_magnitude = (uint16)(32767.0f * LargeValue);
+		Effect.leftright.small_magnitude = (uint16)(32767.0f * SmallValue);
 	}
 	else if (SDL_HapticQuery(Haptic) & SDL_HAPTIC_SINE)
 	{
 		Effect.type = SDL_HAPTIC_SINE;
 		Effect.periodic.length = 1000;
 		Effect.periodic.period = 1000;
-		Effect.periodic.magnitude = 32767.0f * FMath::Max(SmallValue, LargeValue);
+		Effect.periodic.magnitude = (uint16)(32767.0f * FMath::Max(SmallValue, LargeValue));
 	}
 	else
 	{
diff --git a/Engine/Source/Runtime/ApplicationCore/Private/Linux/LinuxPlatformSplash.cpp b/Engine/Source/Runtime/ApplicationCore/Private/Linux/LinuxPlatformSplash.cpp
--- a/Engine/Source/Runtime/ApplicationCore/Private/Linux/LinuxPlatformSplash.cpp
+++ b/Engine/Source/Runtime/ApplicationCore/Private/Linux/LinuxPlatformSplash.cpp
@@ -243,11 +243,11 @@ void FLinuxSplashState::DrawCharacter(int32 PenX, int32 PenY, FT_GlyphSlot Glyph
 			// write pixel
 			const float Alpha = Pixels[SourceIndex] / 255.0f;
 
-			ScratchSpace[DestIndex] = ScratchSpace[DestIndex]*(1.0 - Alpha) + Alpha*Red;
+			ScratchSpace[DestIndex] = (uint8)(ScratchSpace[DestIndex]*(1.0 - Alpha) + Alpha*Red);
 			DestIndex++;
-			ScratchSpace[DestIndex] = ScratchSpace[DestIndex]*(1.0 - Alpha) + Alpha*Green;
+			ScratchSpace[DestIndex] = (uint8)(ScratchSpace[DestIndex]*(1.0 - Alpha) + Alpha*Green);
 			DestIndex++;
-			ScratchSpace[DestIndex] = ScratchSpace[DestIndex]*(1.0 - Alpha) + Alpha*Blue;
+			ScratchSpace[DestIndex] = (uint8)(ScratchSpace[DestIndex]*(1.0 - Alpha) + Alpha*Blue);
 		}
 	}
 }
diff --git a/Engine/Source/Runtime/ApplicationCore/Private/Linux/LinuxWindow.cpp b/Engine/Source/Runtime/ApplicationCore/Private/Linux/LinuxWindow.cpp
--- a/Engine/Source/Runtime/ApplicationCore/Private/Linux/LinuxWindow.cpp
+++ b/Engine/Source/Runtime/ApplicationCore/Private/Linux/LinuxWindow.cpp
@@ -277,12 +277,12 @@ void FLinuxWindow::Initialize( FLinuxApplication* const Application, const TShar
 	}
 	else if (bIsMinSizeSet)
 	{
-		SDL_SetWindowMinimumSize(HWnd, *Definition->SizeLimits.GetMinWidth(), *Definition->SizeLimits.GetMinHeight());
+		SDL_SetWindowMinimumSize(HWnd, (int)*Definition->SizeLimits.GetMinWidth(), (int)*Definition->SizeLimits.GetMinHeight());
 	}
 
 	if (bIsMaxSizeSet)
 	{
-		SDL_SetWindowMaximumSize(HWnd, *Definition->SizeLimits.GetMaxWidth(), *Definition->SizeLimits.GetMaxHeight());
+		SDL_SetWindowMaximumSize(HWnd, (int)*Definition->SizeLimits.GetMaxWidth(), (int)*Definition->SizeLimits.GetMaxHeight());
 	}
 
 	if (Definition->AppearsInTaskbar)
diff --git a/Engine/Source/Runtime/ApplicationCore/Public/GenericPlatform/ICursor.h b/Engine/Source/Runtime/ApplicationCore/Public/GenericPlatform/ICursor.h
--- a/Engine/Source/Runtime/ApplicationCore/Public/GenericPlatform/ICursor.h
+++ b/Engine/Source/Runtime/ApplicationCore/Public/GenericPlatform/ICursor.h
@@ -69,6 +69,8 @@ namespace EMouseCursor
 class ICursor
 {
 public:
+	virtual ~ICursor() = default;
+
 	/** Creates a hardware cursor from file. Can return nullptr when not available. */
 	virtual void* CreateCursorFromFile(const FString& InPathToCursorWithoutExtension, FVector2D HotSpot) { return nullptr; }
 
diff --git a/Engine/Source/Runtime/RHI/Public/DynamicRHI.h b/Engine/Source/Runtime/RHI/Public/DynamicRHI.h
--- a/Engine/Source/Runtime/RHI/Public/DynamicRHI.h
+++ b/Engine/Source/Runtime/RHI/Public/DynamicRHI.h
@@ -1333,8 +1333,8 @@ FORCEINLINE FDynamicRHI::FRHICalcTextureSizeResult RHICalcTexturePlatformSize(FR
 //UE_DEPRECATED(5.1, "The separate RHICalcTexture... functions for each texture type are deprecated. Use RHICalcTexturePlatformSize instead.")
 FORCEINLINE uint64 RHICalcTexture2DPlatformSize(uint32 SizeX, uint32 SizeY, uint8 Format, uint32 NumMips, uint32 NumSamples, ETextureCreateFlags Flags, const FRHIResourceCreateInfo& CreateInfo, uint32& OutAlign)
 {
-	const uint32 Depth = 1;
-	const uint32 ArraySize = 1;
+	const uint16 Depth = 1;
+	const uint16 ArraySize = 1;
 
 	FRHITextureDesc Desc(
 		ETextureDimension::Texture2D,
@@ -1344,8 +1344,8 @@ FORCEINLINE uint64 RHICalcTexture2DPlatformSize(uint32 SizeX, uint32 SizeY, uint
 		{ (int32)SizeX, (int32)SizeY },
 		Depth,
 		ArraySize,
-		NumMips,
-		NumSamples,
+		(uint8)NumMips,
+		(uint8)NumSamples,
 		CreateInfo.ExtData
 	);
 
@@ -1357,7 +1357,7 @@ FORCEINLINE uint64 RHICalcTexture2DPlatformSize(uint32 SizeX, uint32 SizeY, uint
 //UE_DEPRECATED(5.1, "The separate RHICalcTexture... functions for each texture type are deprecated. Use RHICalcTexturePlatformSize instead.")
 FORCEINLINE uint64 RHICalcTexture2DArrayPlatformSize(uint32 SizeX, uint32 SizeY, uint32 ArraySize, uint8 Format, uint32 NumMips, uint32 NumSamples, ETextureCreateFlags Flags, const FRHIResourceCreateInfo& CreateInfo, uint32& OutAlign)
 {
-	const uint32 Depth = 1;
+	const uint16 Depth = 1;
 
 	FRHITextureDesc Desc(
 		ETextureDimension::Texture2DArray,
@@ -1366,9 +1366,9 @@ FORCEINLINE uint64 RHICalcTexture2DArrayPlatformSize(uint32 SizeX, uint32 SizeY,
 		CreateInfo.ClearValueBinding,
 		{ (int32)SizeX, (int32)SizeY },
 		Depth,
-		ArraySize,
-		NumMips,
-		NumSamples,
+		(uint16)ArraySize,
+		(uint8)NumMips,
+		(uint8)NumSamples,
 		CreateInfo.ExtData
 	);
 
@@ -1380,8 +1380,8 @@ FORCEINLINE uint64 RHICalcTexture2DArrayPlatformSize(uint32 SizeX, uint32 SizeY,
 //UE_DEPRECATED(5.1, "The separate RHICalcTexture... functions for each texture type are deprecated. Use RHICalcTexturePlatformSize instead.")
 FORCEINLINE uint64 RHICalcVMTexture2DPlatformSize(uint32 Mip0Width, uint32 Mip0Height, uint8 Format, uint32 NumMips, uint32 FirstMipIdx, uint32 NumSamples, ETextureCreateFlags Flags, uint32& OutAlign)
 {
-	const uint32 Depth     = 1;
-	const uint32 ArraySize = 1;
+	const uint16 Depth     = 1;
+	const uint16 ArraySize = 1;
 	const uint32 ExtData   = 0;
 
 	FRHITextureDesc Desc(
@@ -1392,8 +1392,8 @@ FORCEINLINE uint64 RHICalcVMTexture2DPlatformSize(uint32 Mip0Width, uint32 Mip0H
 		{ (int32)Mip0Width, (int32)Mip0Height },
 		Depth,
 		ArraySize,
-		NumMips,
-		NumSamples,
+		(uint8)NumMips,
+		(uint8)NumSamples,
 		ExtData
 	);
 
@@ -1405,8 +1405,8 @@ FORCEINLINE uint64 RHICalcVMTexture2DPlatformSize(uint32 Mip0Width, uint32 Mip0H
 //UE_DEPRECATED(5.1, "The separate RHICalcTexture... functions for each texture type are deprecated. Use RHICalcTexturePlatformSize instead.")
 FORCEINLINE uint64 RHICalcTexture3DPlatformSize(uint32 SizeX, uint32 SizeY, uint32 SizeZ, uint8 Format, uint32 NumMips, ETextureCreateFlags Flags, const FRHIResourceCreateInfo& CreateInfo, uint32& OutAlign)
 {
-	const uint32 ArraySize  = 1;
-	const uint32 NumSamples = 1;
+	const uint16 ArraySize  = 1;
+	const uint8 NumSamples = 1;
 
 	FRHITextureDesc Desc(
 		ETextureDimension::Texture3D,
@@ -1414,9 +1414,9 @@ FORCEINLINE uint64 RHICalcTexture3DPlatformSize(uint32 SizeX, uint32 SizeY, uint
 		(EPixelFormat)Format,
 		CreateInfo.ClearValueBinding,
 		{ (int32)SizeX, (int32)SizeY },
-		SizeZ,
+		(uint16)SizeZ,
 		ArraySize,
-		NumMips,
+		(uint8)NumMips,
 		NumSamples,
 		CreateInfo.ExtData
 	);
@@ -1429,9 +1429,9 @@ FORCEINLINE uint64 RHICalcTexture3DPlatformSize(uint32 SizeX, uint32 SizeY, uint
 //UE_DEPRECATED(5.1, "The separate RHICalcTexture... functions for each texture type are deprecated. Use RHICalcTexturePlatformSize instead.")
 FORCEINLINE uint64 RHICalcTextureCubePlatformSize(uint32 Size, uint8 Format, uint32 NumMips, ETextureCreateFlags Flags, const FRHIResourceCreateInfo& CreateInfo, uint32& OutAlign)
 {
-	const uint32 Depth      = 1;
-	const uint32 ArraySize  = 1;
-	const uint32 NumSamples = 1;
+	const uint16 Depth      = 1;
+	const uint16 ArraySize  = 1;
+	const uint8 NumSamples = 1;
 
 	FRHITextureDesc Desc(
 		ETextureDimension::TextureCube,
@@ -1441,7 +1441,7 @@ FORCEINLINE uint64 RHICalcTextureCubePlatformSize(uint32 Size, uint8 Format, uin
 		{ (int32)Size, (int32)Size },
 		Depth,
 		ArraySize,
-		NumMips,
+		(uint8)NumMips,
 		NumSamples,
 		CreateInfo.ExtData
 	);
diff --git a/Engine/Source/Runtime/RHI/Public/RHI.h b/Engine/Source/Runtime/RHI/Public/RHI.h
--- a/Engine/Source/Runtime/RHI/Public/RHI.h
+++ b/Engine/Source/Runtime/RHI/Public/RHI.h
@@ -2414,7 +2414,7 @@ private:
 
 	// Private constructor. Memory for transitions is allocated manually with extra space at the tail of the structure for RHI use.
 	FRHITransition(ERHIPipeline SrcPipelines, ERHIPipeline DstPipelines)
-		: State(int8(SrcPipelines) | (int8(DstPipelines) << int32(ERHIPipeline::Num)))
+		: State(int8(int32(SrcPipelines) | (int32(DstPipelines) << int32(ERHIPipeline::Num))))
 #if DO_CHECK || USING_CODE_ANALYSIS
 		, AllowedSrc(SrcPipelines)
 		, AllowedDst(DstPipelines)
@@ -2460,7 +2460,7 @@ private:
 	{
 		checkf(EnumHasAllFlags(AllowedDst, Pipeline), TEXT("Transition is being used on a destination pipeline that it wasn't created for."));
 
-		int8 Mask = int8(Pipeline) << int32(ERHIPipeline::Num);
+		int8 Mask = int8(int32(Pipeline) << int32(ERHIPipeline::Num));
 		int8 PreviousValue = FPlatformAtomics::InterlockedAnd(&State, ~Mask);
 		checkf((PreviousValue & Mask) == Mask, TEXT("RHIEndTransitions has been called twice on this transition for at least one pipeline."));
 
diff --git a/Engine/Source/Runtime/RHI/Public/RHICommandList.h b/Engine/Source/Runtime/RHI/Public/RHICommandList.h
--- a/Engine/Source/Runtime/RHI/Public/RHICommandList.h
+++ b/Engine/Source/Runtime/RHI/Public/RHICommandList.h
@@ -5462,10 +5462,10 @@ FORCEINLINE FTexture2DRHIRef RHICreateTexture2D(uint32 SizeX, uint32 SizeY, uint
 {
 	return RHICreateTexture(
 		FRHITextureCreateDesc::Create2D(CreateInfo.DebugName)
-			.SetExtent(SizeX, SizeY)
+			.SetExtent((int32)SizeX, (int32)SizeY)
 			.SetFormat((EPixelFormat)Format)
-			.SetNumMips(NumMips)
-			.SetNumSamples(NumSamples)
+			.SetNumMips((uint8)NumMips)
+			.SetNumSamples((uint8)NumSamples)
 			.SetFlags(Flags)
 			.SetInitialState(ResourceState)
 			.SetExtData(CreateInfo.ExtData)
@@ -5496,11 +5496,11 @@ FORCEINLINE FTexture2DArrayRHIRef RHICreateTexture2DArray(uint32 SizeX, uint32 S
 {
 	return RHICreateTexture(
 		FRHITextureCreateDesc::Create2DArray(CreateInfo.DebugName)
-			.SetExtent(SizeX, SizeY)
-			.SetArraySize(ArraySize)
+			.SetExtent((int32)SizeX, (int32)SizeY)
+			.SetArraySize((uint16)ArraySize)
 			.SetFormat((EPixelFormat)Format)
-			.SetNumMips(NumMips)
-			.SetNumSamples(NumSamples)
+			.SetNumMips((uint8)NumMips)
+			.SetNumSamples((uint8)NumSamples)
 			.SetFlags(Flags)
 			.SetInitialState(ResourceState)
 			.SetExtData(CreateInfo.ExtData)
@@ -5515,10 +5515,10 @@ FORCEINLINE FTexture3DRHIRef RHICreateTexture3D(uint32 SizeX, uint32 SizeY, uint
 {
 	return RHICreateTexture(
 		FRHITextureCreateDesc::Create3D(CreateInfo.DebugName)
-			.SetExtent(SizeX, SizeY)
-			.SetDepth(SizeZ)
+			.SetExtent((int32)SizeX, (int32)SizeY)
+			.SetDepth((uint16)SizeZ)
 			.SetFormat((EPixelFormat)Format)
-			.SetNumMips(NumMips)
+			.SetNumMips((uint8)NumMips)
 			.SetFlags(Flags)
 			.SetInitialState(ResourceState)
 			.SetExtData(CreateInfo.ExtData)
@@ -5535,7 +5535,7 @@ FORCEINLINE FTextureCubeRHIRef RHICreateTextureCube(uint32 Size, uint8 Format, u
 		FRHITextureCreateDesc::CreateCube(CreateInfo.DebugName)
 			.SetExtent(Size)
 			.SetFormat((EPixelFormat)Format)
-			.SetNumMips(NumMips)
+			.SetNumMips((uint8)NumMips)
 			.SetFlags(Flags)
 			.SetInitialState(ResourceState)
 			.SetExtData(CreateInfo.ExtData)
@@ -5551,9 +5551,9 @@ FORCEINLINE FTextureCubeRHIRef RHICreateTextureCubeArray(uint32 Size, uint32 Arr
 	return RHICreateTexture(
 		FRHITextureCreateDesc::CreateCubeArray(CreateInfo.DebugName)
 			.SetExtent(Size)
-			.SetArraySize(ArraySize)
+			.SetArraySize((uint16)ArraySize)
 			.SetFormat((EPixelFormat)Format)
-			.SetNumMips(NumMips)
+			.SetNumMips((uint8)NumMips)
 			.SetFlags(Flags)
 			.SetInitialState(ResourceState)
 			.SetExtData(CreateInfo.ExtData)
diff --git a/Engine/Source/Runtime/RHI/Public/RHIResources.h b/Engine/Source/Runtime/RHI/Public/RHIResources.h
--- a/Engine/Source/Runtime/RHI/Public/RHIResources.h
+++ b/Engine/Source/Runtime/RHI/Public/RHIResources.h
@@ -1221,8 +1221,8 @@ struct RHI_API FRHITextureDesc
 		, uint32              ExtData    = 0
 		)
 	{
-		const uint32 Depth     = 1;
-		const uint32 ArraySize = 1;
+		const uint16 Depth     = 1;
+		const uint16 ArraySize = 1;
 		return FRHITextureDesc(ETextureDimension::Texture2D, Flags, Format, ClearValue, { Size.X, Size.Y }, Depth, ArraySize, NumMips, NumSamples, ExtData);
 	}
 
@@ -1238,7 +1238,7 @@ struct RHI_API FRHITextureDesc
 		, uint32              ExtData    = 0
 		)
 	{
-		const uint32 Depth   = 1;
+		const uint16 Depth   = 1;
 		return FRHITextureDesc(ETextureDimension::Texture2DArray, Flags, Format, ClearValue, { Size.X, Size.Y }, Depth, ArraySize, NumMips, NumSamples, ExtData);
 	}
 
@@ -1252,12 +1252,13 @@ struct RHI_API FRHITextureDesc
 		, uint32              ExtData    = 0
 		)
 	{
-		const uint32 ArraySize  = 1;
-		const uint32 LocalNumSamples = 1;
+		const uint16 Depth = (uint16)Size.Z;
+		const uint16 ArraySize  = 1;
+		const uint16 LocalNumSamples = 1;
 
 		checkf(Size.Z <= TNumericLimits<decltype(FRHITextureDesc::Depth)>::Max(), TEXT("Depth parameter (Size.Z) exceeds valid range"));
 
-		return FRHITextureDesc(ETextureDimension::Texture3D, Flags, Format, ClearValue, { Size.X, Size.Y }, Size.Z, ArraySize, NumMips, LocalNumSamples, ExtData);
+		return FRHITextureDesc(ETextureDimension::Texture3D, Flags, Format, ClearValue, { Size.X, Size.Y }, Depth, ArraySize, NumMips, LocalNumSamples, ExtData);
 	}
 
 	UE_DEPRECATED(5.1, "FRHITextureDesc Create functions have been moved to FRHITextureCreateDesc.")
@@ -1273,8 +1274,8 @@ struct RHI_API FRHITextureDesc
 	{
 		checkf(Size <= (uint32)TNumericLimits<int32>::Max(), TEXT("Size parameter exceeds valid range"));
 
-		const uint32 Depth     = 1;
-		const uint32 ArraySize = 1;
+		const uint16 Depth     = 1;
+		const uint16 ArraySize = 1;
 		return FRHITextureDesc(ETextureDimension::TextureCube, Flags, Format, ClearValue, { (int32)Size, (int32)Size }, Depth, ArraySize, NumMips, NumSamples, ExtData);
 	}
 
@@ -1292,7 +1293,7 @@ struct RHI_API FRHITextureDesc
 	{
 		checkf(Size <= (uint32)TNumericLimits<int32>::Max(), TEXT("Size parameter exceeds valid range"));
 
-		const uint32 Depth   = 1;
+		const uint16 Depth   = 1;
 		return FRHITextureDesc(ETextureDimension::TextureCubeArray, Flags, Format, ClearValue, { (int32)Size, (int32)Size }, Depth, ArraySize, NumMips, NumSamples, ExtData);
 	}
 
@@ -1560,7 +1561,7 @@ struct FRHITextureCreateDesc : public FRHITextureDesc
 		return Create2DArray(DebugName)
 			.SetExtent(Size)
 			.SetFormat(Format)
-			.SetArraySize(ArraySize);
+			.SetArraySize((uint16)ArraySize);
 	}
 
 	static FRHITextureCreateDesc Create2DArray(const TCHAR* DebugName, int32 SizeX, int32 SizeY, int32 ArraySize, EPixelFormat Format)
@@ -1568,14 +1569,14 @@ struct FRHITextureCreateDesc : public FRHITextureDesc
 		return Create2DArray(DebugName)
 			.SetExtent(SizeX, SizeY)
 			.SetFormat(Format)
-			.SetArraySize(ArraySize);
+			.SetArraySize((uint16)ArraySize);
 	}
 
 	static FRHITextureCreateDesc Create3D(const TCHAR* DebugName, FIntVector Size, EPixelFormat Format)
 	{
 		return Create3D(DebugName)
 			.SetExtent(Size.X, Size.Y)
-			.SetDepth(Size.Z)
+			.SetDepth((uint16)Size.Z)
 			.SetFormat(Format);
 	}
 
@@ -1583,7 +1584,7 @@ struct FRHITextureCreateDesc : public FRHITextureDesc
 	{
 		return Create3D(DebugName)
 			.SetExtent(SizeX, SizeY)
-			.SetDepth(SizeZ)
+			.SetDepth((uint16)SizeZ)
 			.SetFormat(Format);
 	}
 
@@ -1599,7 +1600,7 @@ struct FRHITextureCreateDesc : public FRHITextureDesc
 		return CreateCubeArray(DebugName)
 			.SetExtent(Size)
 			.SetFormat(Format)
-			.SetArraySize(ArraySize);
+			.SetArraySize((uint16)ArraySize);
 	}
 
 	FRHITextureCreateDesc() = default;
diff --git a/Engine/Source/Runtime/RHI/Public/RHIUtilities.h b/Engine/Source/Runtime/RHI/Public/RHIUtilities.h
--- a/Engine/Source/Runtime/RHI/Public/RHIUtilities.h
+++ b/Engine/Source/Runtime/RHI/Public/RHIUtilities.h
@@ -793,8 +793,8 @@ inline void ClearRenderTarget(FRHICommandList& RHICmdList, FRHITexture* Texture,
 	check(Texture);
 	const FIntPoint Extent = Texture->GetSizeXY();
 	FRHIRenderPassInfo Info(Texture, ERenderTargetActions::Clear_Store);
-	Info.ColorRenderTargets[0].MipIndex = MipIndex;
-	Info.ColorRenderTargets[0].ArraySlice = ArraySlice;
+	Info.ColorRenderTargets[0].MipIndex = (uint8)MipIndex;
+	Info.ColorRenderTargets[0].ArraySlice = (int32)ArraySlice;
 	RHICmdList.BeginRenderPass(Info, TEXT("ClearRenderTarget"));
 	RHICmdList.EndRenderPass();
 }
diff --git a/Engine/Source/Runtime/Sockets/Private/BSDSockets/SocketsBSD.cpp b/Engine/Source/Runtime/Sockets/Private/BSDSockets/SocketsBSD.cpp
--- a/Engine/Source/Runtime/Sockets/Private/BSDSockets/SocketsBSD.cpp
+++ b/Engine/Source/Runtime/Sockets/Private/BSDSockets/SocketsBSD.cpp
@@ -730,7 +730,7 @@ ESocketBSDReturn FSocketBSD::HasState(ESocketBSDParam State, FTimespan WaitTime)
 		FDSet.events = POLLERR | POLLPRI;
 		break;
 	}
-	int Result = ::poll(&FDSet, 1, WaitTime.GetTotalMilliseconds());
+	int Result = ::poll(&FDSet, 1, (int)WaitTime.GetTotalMilliseconds());
 	if (Result >= 0)
 	{
 		if ((FDSet.revents & FDSet.events) > 0)
