Fixed parallel translate latency. With lock / unlock supported in queued async command lists, an RHI thread fence was being added conservatively with every async command list submission. When RDG parallel execution is enabled, the parallel translate jobs cannot launch until they are supposed to be submitted, which introduces massive stalls on the RHI thread and extends parallel translation into the next frame, clashing with InitViews. This change fixes the issue by introducing a tiny job that launches when all dependent command lists are finished recording (which happens quickly in RDG) and then scans for the latest RHI thread fence on any of the command lists.

diff --git a/Engine/Source/Runtime/Engine/Private/Materials/MaterialShared.cpp b/Engine/Source/Runtime/Engine/Private/Materials/MaterialShared.cpp
--- a/Engine/Source/Runtime/Engine/Private/Materials/MaterialShared.cpp
+++ b/Engine/Source/Runtime/Engine/Private/Materials/MaterialShared.cpp
@@ -3679,6 +3679,8 @@ public:
 				RHICmdList->UpdateUniformBuffer(Item.UniformBuffer, TempBuffer);
 			}
 
+			RHICmdList->FinishRecording();
+
 		}, TStatId(), nullptr, ENamedThreads::AnyHiPriThreadHiPriTask);
 
 		RHICmdListImmediate.QueueRenderThreadCommandListSubmit(Event, RHICmdList);
diff --git a/Engine/Source/Runtime/RHI/Private/RHICommandList.cpp b/Engine/Source/Runtime/RHI/Private/RHICommandList.cpp
--- a/Engine/Source/Runtime/RHI/Private/RHICommandList.cpp
+++ b/Engine/Source/Runtime/RHI/Private/RHICommandList.cpp
@@ -107,11 +107,6 @@ static TAutoConsoleVariable<int32> CVarRHICmdForceRHIFlush(
 	0,
 	TEXT("Force a flush for every task sent to the RHI thread. For issue diagnosis."));
 
-static TAutoConsoleVariable<int32> CVarRHICmdBalanceTranslatesAfterTasks(
-	TEXT("r.RHICmdBalanceTranslatesAfterTasks"),
-	0,
-	TEXT("Experimental option to balance the parallel translates after the render tasks are complete. This minimizes the number of deferred contexts, but adds latency to starting the translates. r.RHICmdBalanceParallelLists overrides and disables this option"));
-
 static TAutoConsoleVariable<int32> CVarRHICmdMinCmdlistForParallelTranslate(
 	TEXT("r.RHICmdMinCmdlistForParallelTranslate"),
 	2,
@@ -695,6 +690,8 @@ void FRHICommandListExecutor::ExecuteInner(FRHICommandListBase& CmdList)
 				CmdList.InitialGPUMask = SwapCmdList->GPUMask;
 				CmdList.PSOContext = SwapCmdList->PSOContext;
 				CmdList.RHIThreadBufferLockFence = SwapCmdList->RHIThreadBufferLockFence;
+				CmdList.QueuedFenceCandidateEvents = SwapCmdList->QueuedFenceCandidateEvents;
+				CmdList.QueuedFenceCandidates = SwapCmdList->QueuedFenceCandidates;
 				CmdList.Data.bInsideRenderPass = SwapCmdList->Data.bInsideRenderPass;
 				CmdList.Data.bInsideComputePass = SwapCmdList->Data.bInsideComputePass;
 			}
@@ -998,6 +995,9 @@ FGraphEventRef FRHICommandListBase::RHIThreadFence(bool bSetLockFence)
 		if (bSetLockFence)
 		{
 			RHIThreadBufferLockFence = Cmd->Fence;
+			QueuedFenceCandidateEvents.Reset();
+			QueuedFenceCandidates.Reset();
+
 		}
 		return Cmd->Fence;
 	}
@@ -1146,6 +1146,11 @@ void FRHICommandListBase::Reset()
 
 	InitialGPUMask = GPUMask;
 
+	if (!IsImmediate())
+	{
+		FenceCandidate = new FFenceCandidate;
+	}
+
 #if RHI_WANT_BREADCRUMB_EVENTS
 	BreadcrumbStack.ValidateEmpty();
 	BreadcrumbStack.Reset();
@@ -1463,45 +1468,6 @@ void FRHICommandListBase::QueueParallelAsyncCommandListSubmit(FGraphEventRef* An
 
 	if (Num && IsRunningRHIInSeparateThread())
 	{
-		static const auto ICVarRHICmdBalanceParallelLists = IConsoleManager::Get().FindTConsoleVariableDataInt(TEXT("r.RHICmdBalanceParallelLists"));
-
-		if (ICVarRHICmdBalanceParallelLists->GetValueOnRenderThread() == 0 && CVarRHICmdBalanceTranslatesAfterTasks.GetValueOnRenderThread() > 0 && GRHISupportsParallelRHIExecute && CVarRHICmdUseDeferredContexts.GetValueOnAnyThread() > 0)
-		{
-			FGraphEventArray Prereq;
-			FRHICommandListBase** RHICmdLists = (FRHICommandListBase**)Alloc(sizeof(FRHICommandListBase*) * Num, alignof(FRHICommandListBase*));
-			for (int32 Index = 0; Index < Num; Index++)
-			{
-				FGraphEventRef& AnyThreadCompletionEvent = AnyThreadCompletionEvents[Index];
-				FRHICommandList* CmdList = CmdLists[Index];
-				RHICmdLists[Index] = CmdList;
-				if (AnyThreadCompletionEvent.GetReference())
-				{
-					Prereq.Add(AnyThreadCompletionEvent);
-					WaitOutstandingTasks.Add(AnyThreadCompletionEvent);
-				}
-			}
-			// this is used to ensure that any old buffer locks are completed before we start any parallel translates
-			if (RHIThreadBufferLockFence.GetReference())
-			{
-				Prereq.Add(RHIThreadBufferLockFence);
-			}
-			FRHICommandList* CmdList = new FRHICommandList(GetGPUMask());
-			FGraphEventRef TranslateSetupCompletionEvent = TGraphTask<FParallelTranslateSetupCommandList>::CreateTask(&Prereq, ENamedThreads::GetRenderThread()).ConstructAndDispatchWhenReady(CmdList, &RHICmdLists[0], Num, bIsPrepass);
-			QueueCommandListSubmit(CmdList);
-			AllOutstandingTasks.Add(TranslateSetupCompletionEvent);
-			if (IsRunningRHIInSeparateThread())
-			{
-				FRHICommandListExecutor::GetImmediateCommandList().ImmediateFlush(EImmediateFlushType::DispatchToRHIThread); // we don't want stuff after the async cmd list to be bundled with it
-			}
-#if !UE_BUILD_SHIPPING
-			if (CVarRHICmdFlushOnQueueParallelSubmit.GetValueOnRenderThread())
-			{
-				CSV_SCOPED_TIMING_STAT(RHITFlushes, QueueParallelAsyncCommandListSubmit);
-				FRHICommandListExecutor::GetImmediateCommandList().ImmediateFlush(EImmediateFlushType::FlushRHIThread);
-			}
-#endif
-			return;
-		}
 		IRHICommandContextContainer* ContextContainer = nullptr;
 		bool bMerge = !!CVarRHICmdMergeSmallDeferredContexts.GetValueOnRenderThread();
 		int32 EffectiveThreads = 0;
@@ -1573,8 +1539,39 @@ void FRHICommandListBase::QueueParallelAsyncCommandListSubmit(FGraphEventRef* An
 				}
 				UE_CLOG(bSpewMerge, LogTemp, Display, TEXT("Parallel translate %d->%d    %dKB mem   %d draws (-1 = unknown)"), Start, Last, FMath::DivideAndRoundUp(TotalMem, 1024), DrawCnt);
 
-				// this is used to ensure that any old buffer locks are completed before we start any parallel translates
-				if (RHIThreadBufferLockFence.GetReference())
+				if (QueuedFenceCandidates.Num() > 0)
+				{
+					FGraphEventRef FenceCandidateEvent = FGraphEvent::CreateGraphEvent();
+
+					if (RHIThreadBufferLockFence.GetReference())
+					{
+						FenceCandidateEvent->DontCompleteUntil(RHIThreadBufferLockFence);
+						RHIThreadBufferLockFence = FenceCandidateEvent;
+					}
+
+					Prereq.Add(FenceCandidateEvent);
+
+					FFunctionGraphTask::CreateAndDispatchWhenReady(
+						[FenceCandidates = MoveTemp(QueuedFenceCandidates), FenceCandidateEvent](ENamedThreads::Type, const FGraphEventRef&) mutable
+					{
+						QUICK_SCOPE_CYCLE_COUNTER(STAT_FRHICommandListBase_SignalLockFence);
+
+						for (int32 Index = FenceCandidates.Num() - 1; Index >= 0; Index--)
+						{
+							if (FenceCandidates[Index]->Fence)
+							{
+								FenceCandidateEvent->DontCompleteUntil(FenceCandidates[Index]->Fence);
+								break;
+							}
+						}
+
+						FenceCandidateEvent->DispatchSubsequents();
+
+					}, TStatId(), &QueuedFenceCandidateEvents);
+
+					QueuedFenceCandidateEvents.Reset();
+				}
+				else if (RHIThreadBufferLockFence.GetReference())
 				{
 					Prereq.Add(RHIThreadBufferLockFence);
 				}
@@ -1649,11 +1646,6 @@ void FRHICommandListBase::QueueAsyncCommandListSubmit(FGraphEventRef& AnyThreadC
 	ALLOC_COMMAND(FRHICommandWaitForAndSubmitSubList)(AnyThreadCompletionEvent, CmdList);
 	if (IsRunningRHIInSeparateThread())
 	{
-		// The async command list is being recorded in parallel, we have to take a conservative RHI thread fence
-		// for all the work in case a lock was used. The RHI thread waits on the work, so the render thread can't
-		// access the pending RHI thread fence.
-		RHIThreadFence(true);
-
 		FRHICommandListExecutor::GetImmediateCommandList().ImmediateFlush(EImmediateFlushType::DispatchToRHIThread); // we don't want stuff after the async cmd list to be bundled with it
 	}
 }
@@ -1688,12 +1680,10 @@ void FRHICommandListBase::QueueRenderThreadCommandListSubmit(FGraphEventRef& Ren
 		return;
 	}
 
-	ALLOC_COMMAND(FRHICommandWaitForAndSubmitRTSubList)(RenderThreadCompletionEvent, CmdList);
+	QueuedFenceCandidateEvents.Emplace(RenderThreadCompletionEvent);
+	QueuedFenceCandidates.Emplace(CmdList->FenceCandidate);
 
-	// The async command list is being recorded in parallel, we have to take a conservative RHI thread fence
-	// for all the work in case a lock was used. The RHI thread waits on the work, so the render thread can't
-	// access the pending RHI thread fence.
-	RHIThreadFence(true);
+	ALLOC_COMMAND(FRHICommandWaitForAndSubmitRTSubList)(RenderThreadCompletionEvent, CmdList);
 
 #if WITH_MGPU
 	// This will restore the context GPU masks to whatever they were set to
diff --git a/Engine/Source/Runtime/RHI/Public/RHICommandList.h b/Engine/Source/Runtime/RHI/Public/RHICommandList.h
--- a/Engine/Source/Runtime/RHI/Public/RHICommandList.h
+++ b/Engine/Source/Runtime/RHI/Public/RHICommandList.h
@@ -661,6 +661,11 @@ public:
 #endif
 	}
 
+	void FinishRecording()
+	{
+		FenceCandidate->Fence = RHIThreadBufferLockFence;
+	}
+
 private:
 	FRHICommandBase* Root;
 	FRHICommandBase** CommandLink;
@@ -676,6 +681,15 @@ private:
 	FGraphEventArray RTTasks;
 	FGraphEventRef RHIThreadBufferLockFence;
 
+	struct FFenceCandidate : public TConcurrentLinearObject<FFenceCandidate>, public FRefCountBase
+	{
+		FGraphEventRef Fence;
+	};
+
+	TRefCountPtr<FFenceCandidate> FenceCandidate;
+	FGraphEventArray QueuedFenceCandidateEvents;
+	TArray<TRefCountPtr<FFenceCandidate>, FConcurrentLinearArrayAllocator> QueuedFenceCandidates;
+
 	friend class FRHICommandListExecutor;
 	friend class FRHICommandListIterator;
 	friend class FRHICommandListScopedFlushAndExecute;
diff --git a/Engine/Source/Runtime/RenderCore/Private/RenderGraphBuilder.cpp b/Engine/Source/Runtime/RenderCore/Private/RenderGraphBuilder.cpp
--- a/Engine/Source/Runtime/RenderCore/Private/RenderGraphBuilder.cpp
+++ b/Engine/Source/Runtime/RenderCore/Private/RenderGraphBuilder.cpp
@@ -2375,6 +2375,7 @@ void FRDGBuilder::DispatchParallelExecute(IRHICommandContext* RHICmdContext)
 			}
 
 			RHICmdListPass.HandleRTThreadTaskCompletion(MyCompletionGraphEvent);
+			RHICmdListPass.FinishRecording();
 
 		}, TStatId(), nullptr, ENamedThreads::AnyHiPriThreadHiPriTask);
 
