Fix for assert due to race with FRHIAsyncCommandList and GRHICommandList.OutstandingCmdListCount.

diff --git a/Engine/Source/Runtime/RHI/Private/RHICommandList.cpp b/Engine/Source/Runtime/RHI/Private/RHICommandList.cpp
--- a/Engine/Source/Runtime/RHI/Private/RHICommandList.cpp
+++ b/Engine/Source/Runtime/RHI/Private/RHICommandList.cpp
@@ -1104,10 +1104,11 @@ void FRHICommandListExecutor::Transition(TArrayView<const FRHITransitionInfo> In
 	}
 }
 
-FRHICommandListBase::FRHICommandListBase(FRHIGPUMask InGPUMask)
+FRHICommandListBase::FRHICommandListBase(FRHIGPUMask InGPUMask, ERecordingThread InRecordingThread)
 	: Root(nullptr)
 	, CommandLink(nullptr)
 	, bExecuting(false)
+	, RecordingThread(InRecordingThread)
 	, NumCommands(0)
 	, UID(UINT32_MAX)
 	, Context(nullptr)
@@ -1118,14 +1119,26 @@ FRHICommandListBase::FRHICommandListBase(FRHIGPUMask InGPUMask)
 	, InitialGPUMask(InGPUMask)
 	, BoundComputeShaderRHI(nullptr)
 {
-	GRHICommandList.OutstandingCmdListCount.Increment();
+	if (RecordingThread == ERecordingThread::Render)
+	{
+		GRHICommandList.OutstandingCmdListCount.Increment();
+	}
+	else
+	{
+		checkf(!PLATFORM_RHITHREAD_DEFAULT_BYPASS, TEXT("The platform has enabled RHI command list bypass mode for shipping builds. Only render thread command lists are allowed."));
+	}
+
 	Reset();
 }
 
 FRHICommandListBase::~FRHICommandListBase()
 {
 	Flush();
-	GRHICommandList.OutstandingCmdListCount.Decrement();
+
+	if (RecordingThread == ERecordingThread::Render)
+	{
+		GRHICommandList.OutstandingCmdListCount.Decrement();
+	}
 }
 
 const int32 FRHICommandListBase::GetUsedMemory() const
diff --git a/Engine/Source/Runtime/RHI/Public/RHICommandList.h b/Engine/Source/Runtime/RHI/Public/RHICommandList.h
--- a/Engine/Source/Runtime/RHI/Public/RHICommandList.h
+++ b/Engine/Source/Runtime/RHI/Public/RHICommandList.h
@@ -443,6 +443,12 @@ extern RHI_API FRHICommandListFenceAllocator GRHIFenceAllocator;
 class RHI_API FRHICommandListBase : public FNoncopyable
 {
 public:
+	enum class ERecordingThread
+	{
+		Render,
+		Any
+	};
+
 	~FRHICommandListBase();
 
 	inline void Flush();
@@ -654,13 +660,6 @@ public:
 		GDynamicRHI->RHIUpdateUniformBuffer(*this, UniformBufferRHI, Contents);
 	}
 
-	void DisallowBypass()
-	{
-#if CAN_DISALLOW_COMMAND_LIST_BYPASS
-		bAllowBypass = false;
-#endif
-	}
-
 	void FinishRecording()
 	{
 		FenceCandidate->Fence = RHIThreadBufferLockFence;
@@ -670,9 +669,7 @@ private:
 	FRHICommandBase* Root;
 	FRHICommandBase** CommandLink;
 	bool bExecuting;
-#if CAN_DISALLOW_COMMAND_LIST_BYPASS
-	bool bAllowBypass = true;
-#endif
+	ERecordingThread RecordingThread;
 	uint32 NumCommands;
 	uint32 UID;
 	IRHICommandContext* Context;
@@ -695,7 +692,7 @@ private:
 	friend class FRHICommandListScopedFlushAndExecute;
 
 protected:
-	FRHICommandListBase(FRHIGPUMask InGPUMask);
+	FRHICommandListBase(FRHIGPUMask InGPUMask, ERecordingThread InRecordingThread);
 
 	FMemStackBase& GetAllocator() { return MemManager; }
 
@@ -2427,7 +2424,7 @@ template<> RHI_API void FRHICommandSetUAVParameter<FRHIComputeShader>::Execute(F
 class RHI_API FRHIComputeCommandList : public FRHICommandListBase
 {
 public:
-	FRHIComputeCommandList(FRHIGPUMask GPUMask) : FRHICommandListBase(GPUMask) {}
+	FRHIComputeCommandList(FRHIGPUMask GPUMask, ERecordingThread InRecordingThread = ERecordingThread::Render) : FRHICommandListBase(GPUMask, InRecordingThread) {}
 
 	template <typename LAMBDA>
 	FORCEINLINE_DEBUGGABLE void EnqueueLambda(LAMBDA&& Lambda)
@@ -3141,7 +3138,7 @@ template<> RHI_API void FRHICommandSetUAVParameter<FRHIPixelShader>::Execute(FRH
 class RHI_API FRHICommandList : public FRHIComputeCommandList
 {
 public:
-	FRHICommandList(FRHIGPUMask GPUMask) : FRHIComputeCommandList(GPUMask) {}
+	FRHICommandList(FRHIGPUMask GPUMask, ERecordingThread InRecordingThread = ERecordingThread::Render) : FRHIComputeCommandList(GPUMask, InRecordingThread) {}
 
 	inline FRHIVertexShader* GetBoundVertexShader() const { return BoundShaderInput.VertexShaderRHI; }
 	inline FRHIMeshShader* GetBoundMeshShader() const { return BoundShaderInput.GetMeshShader(); }
@@ -4080,7 +4077,7 @@ class RHI_API FRHICommandListImmediate : public FRHICommandList
 {
 	friend class FRHICommandListExecutor;
 	FRHICommandListImmediate()
-		: FRHICommandList(FRHIGPUMask::All())
+		: FRHICommandList(FRHIGPUMask::All(), ERecordingThread::Render)
 	{
 		Data.Type = FRHICommandListBase::FCommonData::ECmdListType::Immediate;
 	}
@@ -4961,7 +4958,7 @@ class RHI_API FRHIAsyncComputeCommandListImmediate : public FRHIComputeCommandLi
 {
 public:
 	FRHIAsyncComputeCommandListImmediate()
-		: FRHIComputeCommandList(FRHIGPUMask::All())
+		: FRHIComputeCommandList(FRHIGPUMask::All(), ERecordingThread::Render)
 	{
 		// TODO: handle async compute command list task prerequisites
 		bAsyncPSOCompileAllowed = false;
diff --git a/Engine/Source/Runtime/RHI/Public/RHICommandList.inl b/Engine/Source/Runtime/RHI/Public/RHICommandList.inl
--- a/Engine/Source/Runtime/RHI/Public/RHICommandList.inl
+++ b/Engine/Source/Runtime/RHI/Public/RHICommandList.inl
@@ -42,8 +42,8 @@ FORCEINLINE_DEBUGGABLE bool FRHICommandListBase::Bypass() const
 {
 	check(!IsImmediate() || IsInRenderingThread() || IsInRHIThread());
 	return GRHICommandList.Bypass()
-#if CAN_DISALLOW_COMMAND_LIST_BYPASS
-		&& bAllowBypass
+#if CAN_TOGGLE_COMMAND_LIST_BYPASS
+		&& RecordingThread == ERecordingThread::Render
 #endif
 		;
 }
diff --git a/Engine/Source/Runtime/RenderCore/Public/RenderingThread.h b/Engine/Source/Runtime/RenderCore/Public/RenderingThread.h
--- a/Engine/Source/Runtime/RenderCore/Public/RenderingThread.h
+++ b/Engine/Source/Runtime/RenderCore/Public/RenderingThread.h
@@ -344,10 +344,8 @@ class FRHIAsyncCommandList
 {
 public:
 	FRHIAsyncCommandList(FRHIGPUMask InGPUMask = FRHIGPUMask::All())
-		: RHICmdListStack(InGPUMask)
-	{
-		RHICmdListStack.DisallowBypass();
-	}
+		: RHICmdListStack(InGPUMask, FRHICommandList::ERecordingThread::Any)
+	{}
 
 	FRHICommandList& GetCommandList()
 	{
@@ -368,7 +366,7 @@ public:
 	{
 		if (RHICmdListStack.HasCommands())
 		{
-			FRHICommandList* RHICmdList = new FRHICommandList(FRHIGPUMask::All());
+			FRHICommandList* RHICmdList = new FRHICommandList(FRHIGPUMask::All(), FRHICommandList::ERecordingThread::Any);
 			RHICmdList->ExchangeCmdList(RHICmdListStack);
 
 			ENQUEUE_RENDER_COMMAND(AsyncCommandListScope)(
