CoreMinimal.h: Add includes to remove dependency on CoreMinimal.h from public headers.
#rnx

diff --git a/Engine/Source/Editor/UnrealEd/Classes/Editor/EditorEngine.h b/Engine/Source/Editor/UnrealEd/Classes/Editor/EditorEngine.h
--- a/Engine/Source/Editor/UnrealEd/Classes/Editor/EditorEngine.h
+++ b/Engine/Source/Editor/UnrealEd/Classes/Editor/EditorEngine.h
@@ -66,6 +66,7 @@ class USoundNode;
 class UTextureRenderTarget2D;
 class UTransactor;
 class FTransactionObjectEvent;
+struct FEdge;
 struct FTransactionContext;
 struct FEditorTransactionDeltaContext;
 struct FTypedElementHandle;
diff --git a/Engine/Source/Runtime/CoreUObject/Public/UObject/Script.h b/Engine/Source/Runtime/CoreUObject/Public/UObject/Script.h
--- a/Engine/Source/Runtime/CoreUObject/Public/UObject/Script.h
+++ b/Engine/Source/Runtime/CoreUObject/Public/UObject/Script.h
@@ -8,6 +8,7 @@
 
 #include "CoreMinimal.h"
 #include "HAL/ThreadSingleton.h"
+#include "Internationalization/Text.h"
 #include "Stats/Stats.h"
 #include "Misc/EnumClassFlags.h"
 #include "Misc/CoreMisc.h"
diff --git a/Engine/Source/Runtime/Engine/Public/PrimitiveSceneProxy.h b/Engine/Source/Runtime/Engine/Public/PrimitiveSceneProxy.h
--- a/Engine/Source/Runtime/Engine/Public/PrimitiveSceneProxy.h
+++ b/Engine/Source/Runtime/Engine/Public/PrimitiveSceneProxy.h
@@ -16,6 +16,7 @@
 #include "SceneView.h"
 #include "PrimitiveUniformShaderParameters.h"
 #include "DrawDebugHelpers.h"
+#include "Math/CapsuleShape.h"
 
 class FLightSceneInfo;
 class FLightSceneProxy;
diff --git a/Engine/Source/Runtime/Engine/Public/Rendering/StaticMeshVertexBuffer.h b/Engine/Source/Runtime/Engine/Public/Rendering/StaticMeshVertexBuffer.h
--- a/Engine/Source/Runtime/Engine/Public/Rendering/StaticMeshVertexBuffer.h
+++ b/Engine/Source/Runtime/Engine/Public/Rendering/StaticMeshVertexBuffer.h
@@ -7,6 +7,7 @@
 #include "StaticMeshVertexData.h"
 #include "PackedNormal.h"
 #include "Components.h"
+#include "Math/Vector2DHalf.h"
 
 class FStaticMeshVertexDataInterface;
 
diff --git a/Engine/Source/Runtime/Engine/Public/SceneView.h b/Engine/Source/Runtime/Engine/Public/SceneView.h
--- a/Engine/Source/Runtime/Engine/Public/SceneView.h
+++ b/Engine/Source/Runtime/Engine/Public/SceneView.h
@@ -19,6 +19,7 @@
 #include "RendererInterface.h"
 #include "Interfaces/Interface_PostProcessVolume.h"
 #include "DynamicRenderScaling.h"
+#include "Math/MirrorMatrix.h"
 
 class FSceneView;
 class FSceneViewFamily;
diff --git a/Engine/Source/Runtime/RHI/Public/RHI.h b/Engine/Source/Runtime/RHI/Public/RHI.h
--- a/Engine/Source/Runtime/RHI/Public/RHI.h
+++ b/Engine/Source/Runtime/RHI/Public/RHI.h
@@ -12,6 +12,10 @@
 #include "RHIDefinitions.h"
 #include "Containers/StaticArray.h"
 #include "Containers/StringFwd.h"
+#include "Math/IntRect.h"
+#include "Math/PerspectiveMatrix.h"
+#include "Math/ScaleMatrix.h"
+#include "Math/TranslationMatrix.h"
 
 #ifndef RHI_COMMAND_LIST_DEBUG_TRACES
 #define RHI_COMMAND_LIST_DEBUG_TRACES 0
