Temporary fix for Asset Reference CIS

diff --git a/Engine/Source/Runtime/RHI/Private/RHI.cpp b/Engine/Source/Runtime/RHI/Private/RHI.cpp
--- a/Engine/Source/Runtime/RHI/Private/RHI.cpp
+++ b/Engine/Source/Runtime/RHI/Private/RHI.cpp
@@ -2653,13 +2653,18 @@ void FGenericDataDrivenShaderPlatformInfo::Initialize()
 					// at this point, we can start pulling information out
 					/*Infos[ShaderPlatform].Name*/ 
 					FName ShaderPlatformName = *GetSectionString(Section.Value, FName("ShaderPlatform"));
-					FMemory::Memcpy(&Infos[ShaderPlatform], &ShaderPlatformNameToConfigSection[ShaderPlatformName].Info, sizeof(FGenericDataDrivenShaderPlatformInfo));
 
-					Infos[ShaderPlatform].Name = *(ShaderPlatformName.ToString() + TEXT("_PREVIEW"));
-					Infos[ShaderPlatform].PreviewShaderPlatformParent = ShaderPlatformNameToConfigSection[ShaderPlatformName].ShaderPlatform;
+					PlatformInfoAndPlatformEnum* ConfigSection = ShaderPlatformNameToConfigSection.Find(ShaderPlatformName);
+					if (ConfigSection != nullptr)
+					{
+						FMemory::Memcpy(&Infos[ShaderPlatform], &(ConfigSection->Info), sizeof(FGenericDataDrivenShaderPlatformInfo));
+
+						Infos[ShaderPlatform].Name = *(ShaderPlatformName.ToString() + TEXT("_PREVIEW"));
+						Infos[ShaderPlatform].PreviewShaderPlatformParent = ConfigSection->ShaderPlatform;
 
-					//set previewable to true, bContainsValidPlatformInfo is set in the update
-					Infos[ShaderPlatform].bIsPreviewPlatform = true;
+						//set previewable to true, bContainsValidPlatformInfo is set in the update
+						Infos[ShaderPlatform].bIsPreviewPlatform = true;
+					}
 				}
 			}
 		}
