Fix Linux Asan

diff --git a/Engine/Source/Runtime/RHI/Private/RHI.cpp b/Engine/Source/Runtime/RHI/Private/RHI.cpp
--- a/Engine/Source/Runtime/RHI/Private/RHI.cpp
+++ b/Engine/Source/Runtime/RHI/Private/RHI.cpp
@@ -2680,19 +2680,21 @@ void FGenericDataDrivenShaderPlatformInfo::UpdatePreviewPlatforms()
 		EShaderPlatform ShaderPlatform = EShaderPlatform(i);
 		if (IsValid(ShaderPlatform))
 		{
-			if (Infos[ShaderPlatform].bIsPreviewPlatform)
+			ERHIFeatureLevel::Type PreviewSPMaxFeatureLevel = Infos[ShaderPlatform].MaxFeatureLevel;
+			EShaderPlatform EditorSPForPreviewMaxFeatureLevel = GShaderPlatformForFeatureLevel[PreviewSPMaxFeatureLevel];
+			if (Infos[ShaderPlatform].bIsPreviewPlatform && EditorSPForPreviewMaxFeatureLevel < EShaderPlatform::SP_NumPlatforms)
 			{
-				Infos[ShaderPlatform].ShaderFormat = Infos[GShaderPlatformForFeatureLevel[Infos[ShaderPlatform].MaxFeatureLevel]].ShaderFormat;
-				Infos[ShaderPlatform].Language = Infos[GShaderPlatformForFeatureLevel[Infos[ShaderPlatform].MaxFeatureLevel]].Language;
-				Infos[ShaderPlatform].bIsHlslcc = Infos[GShaderPlatformForFeatureLevel[Infos[ShaderPlatform].MaxFeatureLevel]].bIsHlslcc;
-				Infos[ShaderPlatform].bSupportsDxc = Infos[GShaderPlatformForFeatureLevel[Infos[ShaderPlatform].MaxFeatureLevel]].bSupportsDxc;
-				Infos[ShaderPlatform].bSupportsGPUScene = Infos[GShaderPlatformForFeatureLevel[Infos[ShaderPlatform].MaxFeatureLevel]].bSupportsGPUScene;
+				Infos[ShaderPlatform].ShaderFormat = Infos[EditorSPForPreviewMaxFeatureLevel].ShaderFormat;
+				Infos[ShaderPlatform].Language = Infos[EditorSPForPreviewMaxFeatureLevel].Language;
+				Infos[ShaderPlatform].bIsHlslcc = Infos[EditorSPForPreviewMaxFeatureLevel].bIsHlslcc;
+				Infos[ShaderPlatform].bSupportsDxc = Infos[EditorSPForPreviewMaxFeatureLevel].bSupportsDxc;
+				Infos[ShaderPlatform].bSupportsGPUScene = Infos[EditorSPForPreviewMaxFeatureLevel].bSupportsGPUScene;
 				Infos[ShaderPlatform].bIsPC = true;
 				Infos[ShaderPlatform].bIsConsole = false;
-				Infos[ShaderPlatform].bSupportsSceneDataCompressedTransforms = Infos[GShaderPlatformForFeatureLevel[Infos[ShaderPlatform].MaxFeatureLevel]].bSupportsSceneDataCompressedTransforms;
-				Infos[ShaderPlatform].bSupportsNanite = Infos[GShaderPlatformForFeatureLevel[Infos[ShaderPlatform].MaxFeatureLevel]].bSupportsNanite;
-				Infos[ShaderPlatform].bSupportsUInt64ImageAtomics = Infos[GShaderPlatformForFeatureLevel[Infos[ShaderPlatform].MaxFeatureLevel]].bSupportsUInt64ImageAtomics;
-				Infos[ShaderPlatform].bSupportsGen5TemporalAA = Infos[GShaderPlatformForFeatureLevel[Infos[ShaderPlatform].MaxFeatureLevel]].bSupportsGen5TemporalAA;
+				Infos[ShaderPlatform].bSupportsSceneDataCompressedTransforms = Infos[EditorSPForPreviewMaxFeatureLevel].bSupportsSceneDataCompressedTransforms;
+				Infos[ShaderPlatform].bSupportsNanite = Infos[EditorSPForPreviewMaxFeatureLevel].bSupportsNanite;
+				Infos[ShaderPlatform].bSupportsUInt64ImageAtomics = Infos[EditorSPForPreviewMaxFeatureLevel].bSupportsUInt64ImageAtomics;
+				Infos[ShaderPlatform].bSupportsGen5TemporalAA = Infos[EditorSPForPreviewMaxFeatureLevel].bSupportsGen5TemporalAA;
 				Infos[ShaderPlatform].bSupportsRenderTargetWriteMask = false;
 				Infos[ShaderPlatform].bSupportsIntrinsicWaveOnce = false;
 				Infos[ShaderPlatform].bSupportsDOFHybridScattering = false;
