Fix ASAN Linux heap-use-after-free

diff --git a/Engine/Source/Runtime/Launch/Private/LaunchEngineLoop.cpp b/Engine/Source/Runtime/Launch/Private/LaunchEngineLoop.cpp
--- a/Engine/Source/Runtime/Launch/Private/LaunchEngineLoop.cpp
+++ b/Engine/Source/Runtime/Launch/Private/LaunchEngineLoop.cpp
@@ -2845,7 +2845,9 @@ PRAGMA_ENABLE_DEPRECATION_WARNINGS
 		SCOPED_BOOT_TIMING("RHIInit");
 		// Initialize the RHI.
 		RHIInit(bHasEditorToken);
+#if WITH_EDITOR
 		FGenericDataDrivenShaderPlatformInfo::UpdatePreviewPlatforms();
+#endif
 	}
 
 	{
diff --git a/Engine/Source/Runtime/RHI/Private/RHI.cpp b/Engine/Source/Runtime/RHI/Private/RHI.cpp
--- a/Engine/Source/Runtime/RHI/Private/RHI.cpp
+++ b/Engine/Source/Runtime/RHI/Private/RHI.cpp
@@ -2580,7 +2580,6 @@ void FGenericDataDrivenShaderPlatformInfo::Initialize()
 		EShaderPlatform ShaderPlatform;
 	};
 
-	TMap<FName, PlatformInfoAndPlatformEnum> ShaderPlatformNameToConfigSection;
 	for (int32 Index = 0; Index < NumDDInfoFiles; Index++)
 	{
 		FConfigFile IniFile;
@@ -2614,62 +2613,24 @@ void FGenericDataDrivenShaderPlatformInfo::Initialize()
 				ParseDataDrivenShaderInfo(Section.Value, Infos[ShaderPlatform]);	
 				Infos[ShaderPlatform].bContainsValidPlatformInfo = true;
 
-				PlatformInfoAndPlatformEnum ItemToAdd;
-				ItemToAdd.Info = Infos[ShaderPlatform];
-				ItemToAdd.ShaderPlatform = ShaderPlatform;
-
-				ShaderPlatformNameToConfigSection.Add(Infos[ShaderPlatform].Name, ItemToAdd);
-			}
-		}
-	}
-
-	if (GIsEditor)
-	{
-		for (int32 Index = 0; Index < NumDDInfoFiles; Index++)
-		{
-			// find all PreviewPlatform sections for all platforms
-			// copy the [PreeviewPlatform]ShaderPlatform= customshaderplatform to a new customshaderplatform with name (orignal)_PREVIEW including its featurelevel.
-			// replace shader format with current editor's shader format.
-
-			FConfigFile IniFile;
-			FString PlatformName;
-
-			FDataDrivenPlatformInfoRegistry::LoadDataDrivenIniFile(Index, IniFile, PlatformName);
-
-			for (auto Section : IniFile)
-			{
-				if (Section.Key.StartsWith(TEXT("PreviewPlatform ")))
+#if WITH_EDITOR
+				for (const FPreviewPlatformMenuItem& Item : FDataDrivenPlatformInfoRegistry::GetAllPreviewPlatformMenuItems())
 				{
-					const FString& SectionName = Section.Key;
-					EShaderPlatform ShaderPlatform = EShaderPlatform(CustomShaderPlatform++);
-					// get enum value for the string name
-					//LexFromString(ShaderPlatform, *SectionName.Mid(15));
-					if (ShaderPlatform == EShaderPlatform::SP_NumPlatforms)
-					{
-						UE_LOG(LogRHI, Warning, TEXT("Found an unknown shader platform %s in a DataDriven ini file"), *SectionName.Mid(15));
-						continue;
-					}
-
-					// at this point, we can start pulling information out
-					/*Infos[ShaderPlatform].Name*/ 
-					FName ShaderPlatformName = *GetSectionString(Section.Value, FName("ShaderPlatform"));
-
-					PlatformInfoAndPlatformEnum* ConfigSection = ShaderPlatformNameToConfigSection.Find(ShaderPlatformName);
-					if (ConfigSection != nullptr)
+					FName PreviewPlatformName = *((Infos[ShaderPlatform].Name).ToString() + TEXT("_PREVIEW"));
+					if (Item.ShaderPlatformPreview == PreviewPlatformName)
 					{
-						FMemory::Memcpy(&Infos[ShaderPlatform], &(ConfigSection->Info), sizeof(FGenericDataDrivenShaderPlatformInfo));
-
-						Infos[ShaderPlatform].Name = *(ShaderPlatformName.ToString() + TEXT("_PREVIEW"));
-						Infos[ShaderPlatform].PreviewShaderPlatformParent = ConfigSection->ShaderPlatform;
-
-						//set previewable to true, bContainsValidPlatformInfo is set in the update
-						Infos[ShaderPlatform].bIsPreviewPlatform = true;
+						EShaderPlatform PreviewShaderPlatform = EShaderPlatform(CustomShaderPlatform++);
+						ParseDataDrivenShaderInfo(Section.Value, Infos[PreviewShaderPlatform]);
+						Infos[PreviewShaderPlatform].Name = PreviewPlatformName;
+						Infos[PreviewShaderPlatform].PreviewShaderPlatformParent = ShaderPlatform;
+						Infos[PreviewShaderPlatform].bIsPreviewPlatform = true;
+						Infos[PreviewShaderPlatform].bContainsValidPlatformInfo = true;
 					}
 				}
+#endif
 			}
 		}
 	}
-
 	bInitialized = true;
 }
 
