No-PCH compile fix

diff --git a/Engine/Source/Runtime/RHI/Public/RHIContext.h b/Engine/Source/Runtime/RHI/Public/RHIContext.h
--- a/Engine/Source/Runtime/RHI/Public/RHIContext.h
+++ b/Engine/Source/Runtime/RHI/Public/RHIContext.h
@@ -17,6 +17,7 @@
 #include "Math/Float16Color.h"
 #include "Modules/ModuleInterface.h"
 #include "RHIBreadcrumbs.h"
+#include "RHIResources.h"
 
 class FRHIDepthRenderTargetView;
 class FRHIRenderTargetView;
