[Insights]
  - Fixed TextureProfiler to emit the inline stat spec events only once. The spec events are "important" trace events. These events are cached by the trace system, so they should not be emitted continuously (as the memory cost for those will increase as well).
  - Fixed a double string conversion in FTextureProfiler::GetTextureNameString.
  - Fixed matching #endif comment for #if CSVPROFILERTRACE_ENABLED.

diff --git a/Engine/Source/Runtime/RHI/Private/TextureProfiler.cpp b/Engine/Source/Runtime/RHI/Private/TextureProfiler.cpp
--- a/Engine/Source/Runtime/RHI/Private/TextureProfiler.cpp
+++ b/Engine/Source/Runtime/RHI/Private/TextureProfiler.cpp
@@ -9,6 +9,7 @@
 #include "Containers/StringConv.h"
 
 #if TEXTURE_PROFILER_ENABLED
+
 static int32 ToMB(uint64 Value)
 {
 	const size_t MB = 1024 * 1024;
@@ -130,6 +131,16 @@ FTextureProfiler::FTexureDetails& FTextureProfiler::FTexureDetails::operator-=(c
 void FTextureProfiler::Init()
 {
 	FCoreDelegates::OnEndFrameRT.AddRaw(this, &FTextureProfiler::Update);
+
+	FCsvProfilerTrace::OutputInlineStat("Total", CSV_CATEGORY_INDEX(RenderTargetProfiler));
+	FCsvProfilerTrace::OutputInlineStat("Total", CSV_CATEGORY_INDEX(RenderTargetWasteProfiler));
+	FCsvProfilerTrace::OutputInlineStat("Other", CSV_CATEGORY_INDEX(RenderTargetProfiler));
+	FCsvProfilerTrace::OutputInlineStat("Other", CSV_CATEGORY_INDEX(RenderTargetWasteProfiler));
+
+	FCsvProfilerTrace::OutputInlineStat("Total", CSV_CATEGORY_INDEX(TextureProfiler));
+	FCsvProfilerTrace::OutputInlineStat("Total", CSV_CATEGORY_INDEX(TextureWasteProfiler));
+	FCsvProfilerTrace::OutputInlineStat("Other", CSV_CATEGORY_INDEX(TextureProfiler));
+	FCsvProfilerTrace::OutputInlineStat("Other", CSV_CATEGORY_INDEX(TextureWasteProfiler));
 }
 
 FTextureProfiler* FTextureProfiler::Get()
@@ -165,6 +176,17 @@ void FTextureProfiler::AddTextureAllocation(FRHITexture* UniqueTexturePtr, size_
 	if (NamedValue.GetTextureName().IsNone() || !NamedValue.GetTextureName().IsValid())
 	{
 		NamedValue.SetName(AddedDetails.GetTextureName());
+
+		if (AddedDetails.IsRenderTarget)
+		{
+			FCsvProfilerTrace::OutputInlineStat(NamedValue.GetTextureNameString(), CSV_CATEGORY_INDEX(RenderTargetProfiler));
+			FCsvProfilerTrace::OutputInlineStat(NamedValue.GetTextureNameString(), CSV_CATEGORY_INDEX(RenderTargetWasteProfiler));
+		}
+		else
+		{
+			FCsvProfilerTrace::OutputInlineStat(NamedValue.GetTextureNameString(), CSV_CATEGORY_INDEX(TextureProfiler));
+			FCsvProfilerTrace::OutputInlineStat(NamedValue.GetTextureNameString(), CSV_CATEGORY_INDEX(TextureWasteProfiler));
+		}
 	}
 
 	TexturesMap.Add(UniqueTexturePtr, AddedDetails);
@@ -243,9 +265,10 @@ const char* FTextureProfiler::GetTextureNameString(FName TextureName)
 		{
 			WIDECHAR NewTextureNameString[NAME_SIZE];
 			TextureName.GetPlainWIDEString(NewTextureNameString);
-			int Length = StringCast<ANSICHAR>(NewTextureNameString).Length();
+			const auto ConvertedTextureNameString = StringCast<ANSICHAR>(NewTextureNameString);
+			int Length = ConvertedTextureNameString.Length();
 			TextureNameString = new char[Length + 1];
-			FCStringAnsi::Strcpy(const_cast<char*>(TextureNameString), Length + 1, StringCast<ANSICHAR>(NewTextureNameString).Get());
+			FCStringAnsi::Strcpy(const_cast<char*>(TextureNameString), Length + 1, ConvertedTextureNameString.Get());
 		}
 		else
 		{
@@ -264,10 +287,13 @@ const char* FTextureProfiler::GetTextureNameString(FName TextureName)
 
 static void ReportTextureStat(const TAutoConsoleVariable<bool>& EnableVar, const char* StatName, uint32 CategoryIndex, int32 Size, ECsvCustomStatOp Op)
 {
-	FCsvProfilerTrace::OutputInlineStat(StatName, CategoryIndex);
-	FCsvProfiler::RecordCustomStat(StatName, CategoryIndex, Size, Op);
+	if (EnableVar.GetValueOnRenderThread())
+	{
+		FCsvProfiler::RecordCustomStat(StatName, CategoryIndex, Size, Op);
+	}
 }
-#endif //CSV_PROFILER
+
+#endif //CSVPROFILERTRACE_ENABLED
 
 void FTextureProfiler::Update()
 {
@@ -343,7 +369,8 @@ void FTextureProfiler::Update()
 
 	TotalRenderTargetSize.ResetPeakSize();
 	TotalTextureSize.ResetPeakSize();
-#endif //CSV_PROFILER
+
+#endif //CSVPROFILERTRACE_ENABLED
 }
 
 void FTextureProfiler::DumpTextures(bool RenderTargets, bool CombineTextureNames, bool AsCSV, FOutputDevice& OutputDevice)
