Texture profiler compiler errors in some builds.

diff --git a/Engine/Source/Runtime/RHI/Private/TextureProfiler.cpp b/Engine/Source/Runtime/RHI/Private/TextureProfiler.cpp
--- a/Engine/Source/Runtime/RHI/Private/TextureProfiler.cpp
+++ b/Engine/Source/Runtime/RHI/Private/TextureProfiler.cpp
@@ -132,6 +132,7 @@ void FTextureProfiler::Init()
 {
 	FCoreDelegates::OnEndFrameRT.AddRaw(this, &FTextureProfiler::Update);
 
+#if CSVPROFILERTRACE_ENABLED
 	FCsvProfilerTrace::OutputInlineStat("Total", CSV_CATEGORY_INDEX(RenderTargetProfiler));
 	FCsvProfilerTrace::OutputInlineStat("Total", CSV_CATEGORY_INDEX(RenderTargetWasteProfiler));
 	FCsvProfilerTrace::OutputInlineStat("Other", CSV_CATEGORY_INDEX(RenderTargetProfiler));
@@ -141,6 +142,7 @@ void FTextureProfiler::Init()
 	FCsvProfilerTrace::OutputInlineStat("Total", CSV_CATEGORY_INDEX(TextureWasteProfiler));
 	FCsvProfilerTrace::OutputInlineStat("Other", CSV_CATEGORY_INDEX(TextureProfiler));
 	FCsvProfilerTrace::OutputInlineStat("Other", CSV_CATEGORY_INDEX(TextureWasteProfiler));
+#endif //CSVPROFILERTRACE_ENABLED
 }
 
 FTextureProfiler* FTextureProfiler::Get()
@@ -177,6 +179,7 @@ void FTextureProfiler::AddTextureAllocation(FRHITexture* UniqueTexturePtr, size_
 	{
 		NamedValue.SetName(AddedDetails.GetTextureName());
 
+#if CSVPROFILERTRACE_ENABLED
 		if (AddedDetails.IsRenderTarget)
 		{
 			FCsvProfilerTrace::OutputInlineStat(NamedValue.GetTextureNameString(), CSV_CATEGORY_INDEX(RenderTargetProfiler));
@@ -187,6 +190,7 @@ void FTextureProfiler::AddTextureAllocation(FRHITexture* UniqueTexturePtr, size_
 			FCsvProfilerTrace::OutputInlineStat(NamedValue.GetTextureNameString(), CSV_CATEGORY_INDEX(TextureProfiler));
 			FCsvProfilerTrace::OutputInlineStat(NamedValue.GetTextureNameString(), CSV_CATEGORY_INDEX(TextureWasteProfiler));
 		}
+#endif //CSVPROFILERTRACE_ENABLED
 	}
 
 	TexturesMap.Add(UniqueTexturePtr, AddedDetails);
