Remove delete of MacOSXVersionCompare causing cook error on Mac

diff --git a/Engine/Source/Runtime/RHI/Private/DynamicRHI.cpp b/Engine/Source/Runtime/RHI/Private/DynamicRHI.cpp
--- a/Engine/Source/Runtime/RHI/Private/DynamicRHI.cpp
+++ b/Engine/Source/Runtime/RHI/Private/DynamicRHI.cpp
@@ -277,19 +277,22 @@ static void RHIDetectAndWarnOfBadDrivers(bool bHasEditorToken)
 		return;
 	}
 
-	// this message can be suppressed with r.WarnOfBadDrivers=0
-	FPlatformMisc::MessageBoxExt(EAppMsgType::Ok,
+	if (FPlatformMisc::MacOSXVersionCompare(12, 0, 0) < 0)
+	{
+		// this message can be suppressed with r.WarnOfBadDrivers=0
+		FPlatformMisc::MessageBoxExt(EAppMsgType::Ok,
 								 *NSLOCTEXT("MessageDialog", "UpdateMacOSX_Body", "Please update to the latest version of macOS for best performance and stability.").ToString(),
 								 *NSLOCTEXT("MessageDialog", "UpdateMacOSX_Title", "Update macOS").ToString());
 		
 #if !UE_BUILD_SHIPPING
-	if (CVarBadDriverWarningIsFatal.GetValueOnGameThread())
-	{
-		// Force a fatal error depending on CVar
+		if (CVarBadDriverWarningIsFatal.GetValueOnGameThread())
+		{
+			// Force a fatal error depending on CVar
 		UE_LOG(LogRHI, Fatal, TEXT("Fatal crash requested when graphics drivers are out of date.\n")
 			TEXT("To prevent this crash, please update macOS."));
-	}
+		}
 #endif
+	}
 }
 #endif // PLATFORM_WINDOWS
 
