Fixed RHI validation error about multiple view uniform buffer bindings.

diff --git a/Engine/Source/Runtime/RHI/Public/RHIContext.h b/Engine/Source/Runtime/RHI/Public/RHIContext.h
--- a/Engine/Source/Runtime/RHI/Public/RHIContext.h
+++ b/Engine/Source/Runtime/RHI/Public/RHIContext.h
@@ -60,7 +60,10 @@ public:
 		checkf(IsUniformBufferStaticSlotValid(Slot), TEXT("Attempted to set a global uniform buffer %s with an invalid slot."), *Layout.GetDebugName());
 
 #if VALIDATE_UNIFORM_BUFFER_STATIC_BINDINGS
-		ensureMsgf(INDEX_NONE == Slots.Find(Slot), TEXT("Uniform Buffer %s was added twice to the binding array."), *Layout.GetDebugName());
+		if (int32 SlotIndex = Slots.Find(Slot); SlotIndex != INDEX_NONE)
+		{
+			checkf(UniformBuffers[SlotIndex] == UniformBuffer, TEXT("Uniform Buffer %s was added multiple times to the binding array but with different values."), *Layout.GetDebugName());
+		}
 #endif
 
 		Slots.Add(Slot);
diff --git a/Engine/Source/Runtime/Renderer/Private/SingleLayerWaterRendering.cpp b/Engine/Source/Runtime/Renderer/Private/SingleLayerWaterRendering.cpp
--- a/Engine/Source/Runtime/Renderer/Private/SingleLayerWaterRendering.cpp
+++ b/Engine/Source/Runtime/Renderer/Private/SingleLayerWaterRendering.cpp
@@ -733,7 +733,7 @@ void FDeferredShadingSceneRenderer::RenderSingleLayerWaterReflections(
 
 			FWaterCompositeParameters* PassParameters = GraphBuilder.AllocParameters<FWaterCompositeParameters>();
 
-			PassParameters->VS.ViewUniformBuffer = GetShaderBinding(View.ViewUniformBuffer);
+			PassParameters->VS.ViewUniformBuffer = View.ViewUniformBuffer;
 			PassParameters->VS.TileListData = TiledScreenSpaceReflection.TileListDataBufferSRV;
 
 			SetCommonParameters(PassParameters->PS.CommonParameters);
