Added support for RHI validation breadcrumb logging with transition validation errors.

diff --git a/Engine/Source/Runtime/RHI/Private/RHIValidation.cpp b/Engine/Source/Runtime/RHI/Private/RHIValidation.cpp
--- a/Engine/Source/Runtime/RHI/Private/RHIValidation.cpp
+++ b/Engine/Source/Runtime/RHI/Private/RHIValidation.cpp
@@ -426,13 +426,54 @@ void* FValidationRHI::RHILockBufferMGPU(class FRHICommandListBase& RHICmdList, F
 	return RHI->RHILockBufferMGPU(RHICmdList, Buffer, GPUIndex, Offset, SizeRHI, LockMode);
 }
 
+class FRHIValidationBreadcrumbScope
+{
+public:
+	FRHIValidationBreadcrumbScope(TConstArrayView<const TCHAR*> InBreadcrumbs)
+	{
+		Breadcrumbs = InBreadcrumbs;
+	}
+
+	~FRHIValidationBreadcrumbScope()
+	{
+		Breadcrumbs = {};
+	}
+
+	thread_local static TConstArrayView<const TCHAR*> Breadcrumbs;
+};
+
+thread_local TConstArrayView<const TCHAR*> FRHIValidationBreadcrumbScope::Breadcrumbs;
+
 void FValidationRHI::ReportValidationFailure(const TCHAR* InMessage)
 {
+	FString Message;
+
+	if (!FRHIValidationBreadcrumbScope::Breadcrumbs.IsEmpty())
+	{
+		FString BreadcrumbMessage;
+
+		for (int32 Index = 0; Index < FRHIValidationBreadcrumbScope::Breadcrumbs.Num() - 1; ++Index)
+		{
+			BreadcrumbMessage += (FRHIValidationBreadcrumbScope::Breadcrumbs[Index] + FString(TEXT("/")));
+		}
+
+		BreadcrumbMessage += FRHIValidationBreadcrumbScope::Breadcrumbs.Last();
+		Message = FString::Printf(
+			TEXT("%s")
+			TEXT("Breadcrumbs: %s\n")
+			TEXT("--------------------------------------------------------------------\n"),
+			InMessage, *BreadcrumbMessage);
+	}
+	else
+	{
+		Message = InMessage;
+	}
+
 	// Report failures only once per session, since many of them will happen repeatedly. This is similar to what ensure() does, but
 	// ensure() looks at the source location to determine if it's seen the error before. We want to look at the actual message, since
 	// all failures of a given kind will come from the same place, but (hopefully) the error message contains the name of the resource
 	// and a description of the state, so it should be unique for each failure.
-	uint32 Hash = FCrc::StrCrc32<TCHAR>(InMessage);
+	uint32 Hash = FCrc::StrCrc32<TCHAR>(*Message);
 	
 	SeenFailureHashesMutex.Lock();
 	bool bIsAlreadyInSet;
@@ -444,13 +485,13 @@ void FValidationRHI::ReportValidationFailure(const TCHAR* InMessage)
 		return;
 	}
 
-	UE_LOG(LogRHI, Error, TEXT("%s"), InMessage);
+	UE_LOG(LogRHI, Error, TEXT("%s"), *Message);
 
 	if (FPlatformMisc::IsDebuggerPresent() && RHIValidation::GBreakOnTransitionError)
 	{
 		// Print the message again using the debug output function, because UE_LOG doesn't always reach
 		// the VS output window before the breakpoint is triggered, despite the log flush call below.
-		FPlatformMisc::LowLevelOutputDebugStringf(TEXT("%s\n"), InMessage);
+		FPlatformMisc::LowLevelOutputDebugStringf(TEXT("%s\n"), *Message);
 		GLog->Flush();
 		PLATFORM_BREAK();
 	}
@@ -468,7 +509,7 @@ void FValidationComputeContext::FState::Reset()
 	ComputePassName.Reset();
 	bComputePSOSet = false;
 	TrackerInstance.ResetAllUAVState();
-	GlobalUniformBuffers.Reset();
+	StaticUniformBuffers.Reset();
 }
 
 FValidationContext::FValidationContext(EType InType)
@@ -552,18 +593,18 @@ void FValidationContext::FState::Reset()
 	ComputePassName.Reset();
 	bComputePSOSet = false;
 	TrackerInstance.ResetAllUAVState();
-	GlobalUniformBuffers.Reset();
+	StaticUniformBuffers.Reset();
 }
 
 namespace RHIValidation
 {
-	void FGlobalUniformBuffers::Reset()
+	void FStaticUniformBuffers::Reset()
 	{
 		Bindings.Reset();
 		check(!bInSetPipelineStateCall);
 	}
 
-	void FGlobalUniformBuffers::ValidateSetShaderUniformBuffer(FRHIUniformBuffer* UniformBuffer)
+	void FStaticUniformBuffers::ValidateSetShaderUniformBuffer(FRHIUniformBuffer* UniformBuffer)
 	{
 		check(UniformBuffer);
 		UniformBuffer->ValidateLifeTime();
@@ -1402,8 +1443,25 @@ namespace RHIValidation
 		}
 	}
 
-	RHI_API EReplayStatus FOperation::Replay(ERHIPipeline Pipeline, bool& bAllowAllUAVsOverlap) const
+	RHI_API EReplayStatus FOperation::Replay(ERHIPipeline Pipeline, bool& bAllowAllUAVsOverlap, FBreadcrumbStack& Breadcrumbs) const
 	{
+		switch (Type)
+		{
+		case EOpType::PushBreadcrumb:
+			Breadcrumbs.Push(Data_PushBreadcrumb.Breadcrumb);
+			return EReplayStatus::Normal;
+
+		case EOpType::PopBreadcrumb:
+			if (!Breadcrumbs.IsEmpty())
+			{
+				delete[] Breadcrumbs.Last();
+				Breadcrumbs.Pop();
+			}
+			return EReplayStatus::Normal;
+		}
+
+		FRHIValidationBreadcrumbScope BreadcrumbScope(Breadcrumbs);
+
 		switch (Type)
 		{
 		case EOpType::Rename:
@@ -1547,13 +1605,13 @@ namespace RHIValidation
 				FOpQueueState& CurrentQueue = OpQueues[CurrentIndex];
 				if (CurrentQueue.bWaiting)
 				{
-					Status = CurrentQueue.Ops.Replay(CurrentPipeline, CurrentQueue.bAllowAllUAVsOverlap);
+					Status = CurrentQueue.Ops.Replay(CurrentPipeline, CurrentQueue.bAllowAllUAVsOverlap, CurrentQueue.Breadcrumbs);
 					if (!EnumHasAllFlags(Status, EReplayStatus::Waiting))
 					{
 						CurrentQueue.Ops.Reset();
 						if (CurrentIndex == DstOpQueueIndex && InOpsList.Incomplete())
 						{
-							Status |= InOpsList.Replay(CurrentPipeline, CurrentQueue.bAllowAllUAVsOverlap);
+							Status |= InOpsList.Replay(CurrentPipeline, CurrentQueue.bAllowAllUAVsOverlap, CurrentQueue.Breadcrumbs);
 							CurrentQueue.bWaiting = InOpsList.Incomplete();
 						}
 						else
diff --git a/Engine/Source/Runtime/RHI/Public/RHI.h b/Engine/Source/Runtime/RHI/Public/RHI.h
--- a/Engine/Source/Runtime/RHI/Public/RHI.h
+++ b/Engine/Source/Runtime/RHI/Public/RHI.h
@@ -2787,13 +2787,12 @@ inline void RHIValidation::FTracker::AddOp(const RHIValidation::FOperation& Op)
 	if (GRHICommandList.Bypass() && CurrentList.Operations.Num() == 0)
 	{
 		auto& OpQueue = OpQueues[GetOpQueueIndex(Pipeline)];
-		if (!EnumHasAllFlags(Op.Replay(Pipeline, OpQueue.bAllowAllUAVsOverlap), EReplayStatus::Waiting))
+		if (!EnumHasAllFlags(Op.Replay(Pipeline, OpQueue.bAllowAllUAVsOverlap, OpQueue.Breadcrumbs), EReplayStatus::Waiting))
 		{
 			return;
 		}
 	}
 
-	// @todo: improve performance. Use a fast linear allocator
 	CurrentList.Operations.Add(Op);
 }
 
diff --git a/Engine/Source/Runtime/RHI/Public/RHIValidationCommon.h b/Engine/Source/Runtime/RHI/Public/RHIValidationCommon.h
--- a/Engine/Source/Runtime/RHI/Public/RHIValidationCommon.h
+++ b/Engine/Source/Runtime/RHI/Public/RHIValidationCommon.h
@@ -6,6 +6,8 @@
 
 #pragma once
 
+#include "Experimental/ConcurrentLinearAllocator.h"
+
 #if ENABLE_RHI_VALIDATION
 extern RHI_API bool GRHIValidationEnabled;
 #else
@@ -23,7 +25,7 @@ static inline bool IsStencilFormat(EPixelFormat Format);
 
 namespace RHIValidation
 {
-	struct FGlobalUniformBuffers
+	struct FStaticUniformBuffers
 	{
 		TArray<FRHIUniformBuffer*> Bindings;
 		bool bInSetPipelineStateCall{};
@@ -582,7 +584,9 @@ namespace RHIValidation
 		Signal,
 		Wait,
 		AllUAVsOverlap,
-		SpecificUAVOverlap
+		SpecificUAVOverlap,
+		PushBreadcrumb,
+		PopBreadcrumb
 	};
 
 	struct FUniformBufferResource
@@ -595,7 +599,9 @@ namespace RHIValidation
 		void InitLifetimeTracking(uint64 FrameID, const void* Contents, EUniformBufferUsage Usage);
 		void UpdateAllocation(uint64 FrameID);
 		void ValidateLifeTime();
-	};	
+	};
+
+	using FBreadcrumbStack = TArray<const TCHAR*, FConcurrentLinearArrayAllocator>;
 
 	struct FOperation
 	{
@@ -680,9 +686,14 @@ namespace RHIValidation
 				FResourceIdentity Identity;
 				bool bAllow;
 			} Data_SpecificUAVOverlap;
+
+			struct
+			{
+				TCHAR* Breadcrumb;
+			} Data_PushBreadcrumb;
 		};
 
-		RHI_API EReplayStatus Replay(ERHIPipeline Pipeline, bool& bAllowAllUAVsOverlap) const;
+		RHI_API EReplayStatus Replay(ERHIPipeline Pipeline, bool& bAllowAllUAVsOverlap, FBreadcrumbStack& Breadcrumbs) const;
 
 		static inline FOperation BeginTransitionResource(FResourceIdentity Identity, FState PreviousState, FState NextState, EResourceTransitionFlags Flags, void* CreateBacktrace)
 		{
@@ -787,12 +798,7 @@ namespace RHIValidation
 			FOperation Op;
 			Op.Type = EOpType::Rename;
 			Op.Data_Rename.Resource = Resource;
-
-			int32 NameLen = FCString::Strlen(NewName);
-			Op.Data_Rename.DebugName = new TCHAR[NameLen + 1];
-			FMemory::Memcpy(Op.Data_Rename.DebugName, NewName, NameLen * sizeof(TCHAR));
-			Op.Data_Rename.DebugName[NameLen] = 0;
-
+			AllocStringCopy(Op.Data_Rename.DebugName, NewName);
 			Op.Data_Rename.Suffix = Suffix;
 			return MoveTemp(Op);
 		}
@@ -833,19 +839,45 @@ namespace RHIValidation
 			Op.Data_SpecificUAVOverlap.bAllow = bAllow;
 			return MoveTemp(Op);
 		}
+
+		static inline FOperation PushBreadcrumb(const TCHAR* Breadcrumb)
+		{
+			FOperation Op;
+			Op.Type = EOpType::PushBreadcrumb;
+			AllocStringCopy(Op.Data_PushBreadcrumb.Breadcrumb, Breadcrumb);
+			return Op;
+		}
+
+		static inline FOperation PopBreadcrumb()
+		{
+			FOperation Op;
+			Op.Type = EOpType::PopBreadcrumb;
+			return Op;
+		}
+
+	private:
+		static inline void AllocStringCopy(TCHAR*& OutString, const TCHAR* InString)
+		{
+			int32 Len = FCString::Strlen(InString);
+			OutString = new TCHAR[Len + 1];
+			FMemory::Memcpy(OutString, InString, Len * sizeof(TCHAR));
+			OutString[Len] = 0;
+		}
 	};
 
 	struct FOperationsList
 	{
-		TArray<FOperation> Operations;
+		using ListType = TArray<FOperation, FConcurrentLinearArrayAllocator>;
+
+		ListType Operations;
 		int32 OperationPos = 0;
 
-		inline EReplayStatus Replay(ERHIPipeline Pipeline, bool& bAllowAllUAVsOverlap)
+		inline EReplayStatus Replay(ERHIPipeline Pipeline, bool& bAllowAllUAVsOverlap, FBreadcrumbStack& Breadcrumbs)
 		{
 			EReplayStatus Status = EReplayStatus::Normal;
 			for (; OperationPos < Operations.Num(); ++OperationPos)
 			{
-				Status |= Operations[OperationPos].Replay(Pipeline, bAllowAllUAVsOverlap);
+				Status |= Operations[OperationPos].Replay(Pipeline, bAllowAllUAVsOverlap, Breadcrumbs);
 				if (EnumHasAllFlags(Status, EReplayStatus::Waiting))
 				{
 					break;
@@ -952,7 +984,7 @@ namespace RHIValidation
 
 		inline void AddOp(const FOperation& Op);
 
-		inline void AddOps(const TArray<FOperation>& Ops)
+		inline void AddOps(const FOperationsList::ListType& Ops)
 		{
 			for (const FOperation& Op : Ops)
 			{
@@ -969,6 +1001,16 @@ namespace RHIValidation
 			return MoveTemp(CurrentList);
 		}
 
+		inline void PushBreadcrumb(const TCHAR* Breadcrumb)
+		{
+			AddOp(FOperation::PushBreadcrumb(Breadcrumb));
+		}
+
+		inline void PopBreadcrumb()
+		{
+			AddOp(FOperation::PopBreadcrumb());
+		}
+
 		inline void SetTrackedAccess(FResource* Resource, ERHIAccess Access)
 		{
 			AddOp(FOperation::SetTrackedAccess(Resource, Access));
@@ -1066,6 +1108,8 @@ namespace RHIValidation
 			bool bAllowAllUAVsOverlap = false;
 
 			FOperationsList Ops;
+			TArray<const TCHAR*, FConcurrentLinearArrayAllocator> Breadcrumbs;
+
 		} static RHI_API OpQueues[int32(ERHIPipeline::Num)];
 	};
 
diff --git a/Engine/Source/Runtime/RHI/Public/RHIValidationContext.h b/Engine/Source/Runtime/RHI/Public/RHIValidationContext.h
--- a/Engine/Source/Runtime/RHI/Public/RHIValidationContext.h
+++ b/Engine/Source/Runtime/RHI/Public/RHIValidationContext.h
@@ -28,9 +28,9 @@ public:
 		// Reset the compute UAV tracker since the renderer must re-bind all resources after changing a shader.
 		Tracker->ResetUAVState(RHIValidation::EUAVMode::Compute);
 
-		State.GlobalUniformBuffers.bInSetPipelineStateCall = true;
+		State.StaticUniformBuffers.bInSetPipelineStateCall = true;
 		RHIContext->RHISetComputePipelineState(ComputePipelineState);
-		State.GlobalUniformBuffers.bInSetPipelineStateCall = false;
+		State.StaticUniformBuffers.bInSetPipelineStateCall = false;
 	}
 
 	virtual void RHIDispatchComputeShader(uint32 ThreadGroupCountX, uint32 ThreadGroupCountY, uint32 ThreadGroupCountZ) override final
@@ -170,7 +170,7 @@ public:
 	virtual void RHISetShaderUniformBuffer(FRHIComputeShader* Shader, uint32 BufferIndex, FRHIUniformBuffer* Buffer) override final
 	{
 		checkf(State.bComputePSOSet, TEXT("A Compute PSO has to be set to set resources into a shader!"));
-		State.GlobalUniformBuffers.ValidateSetShaderUniformBuffer(Buffer);
+		State.StaticUniformBuffers.ValidateSetShaderUniformBuffer(Buffer);
 		RHIContext->RHISetShaderUniformBuffer(Shader, BufferIndex, Buffer);
 	}
 
@@ -182,17 +182,19 @@ public:
 
 	virtual void RHISetStaticUniformBuffers(const FUniformBufferStaticBindings& InUniformBuffers) override final
 	{
-		InUniformBuffers.Bind(State.GlobalUniformBuffers.Bindings);
+		InUniformBuffers.Bind(State.StaticUniformBuffers.Bindings);
 		RHIContext->RHISetStaticUniformBuffers(InUniformBuffers);
 	}
 
 	virtual void RHIPushEvent(const TCHAR* Name, FColor Color) override final
 	{
+		Tracker->PushBreadcrumb(Name);
 		RHIContext->RHIPushEvent(Name, Color);
 	}
 
 	virtual void RHIPopEvent() override final
 	{
+		Tracker->PopBreadcrumb();
 		RHIContext->RHIPopEvent();
 	}
 
@@ -237,7 +239,7 @@ protected:
 	struct FState
 	{
 		RHIValidation::FTracker TrackerInstance{ ERHIPipeline::AsyncCompute };
-		RHIValidation::FGlobalUniformBuffers GlobalUniformBuffers;
+		RHIValidation::FStaticUniformBuffers StaticUniformBuffers;
 
 		FString ComputePassName;
 		bool bComputePSOSet{};
@@ -277,9 +279,9 @@ public:
 		// Reset the compute UAV tracker since the renderer must re-bind all resources after changing a shader.
 		Tracker->ResetUAVState(RHIValidation::EUAVMode::Compute);
 
-		State.GlobalUniformBuffers.bInSetPipelineStateCall = true;
+		State.StaticUniformBuffers.bInSetPipelineStateCall = true;
 		RHIContext->RHISetComputePipelineState(ComputePipelineState);
-		State.GlobalUniformBuffers.bInSetPipelineStateCall = false;
+		State.StaticUniformBuffers.bInSetPipelineStateCall = false;
 	}
 
 	virtual void RHIDispatchComputeShader(uint32 ThreadGroupCountX, uint32 ThreadGroupCountY, uint32 ThreadGroupCountZ) override final
@@ -572,9 +574,9 @@ public:
 		// Setting a new PSO unbinds all previous bound resources
 		Tracker->ResetUAVState(RHIValidation::EUAVMode::Graphics);
 
-		State.GlobalUniformBuffers.bInSetPipelineStateCall = true;
+		State.StaticUniformBuffers.bInSetPipelineStateCall = true;
 		RHIContext->RHISetGraphicsPipelineState(GraphicsState, StencilRef, bApplyAdditionalState);
-		State.GlobalUniformBuffers.bInSetPipelineStateCall = false;
+		State.StaticUniformBuffers.bInSetPipelineStateCall = false;
 	}
 
 #if PLATFORM_USE_FALLBACK_PSO
@@ -589,9 +591,9 @@ public:
 		// Setting a new PSO unbinds all previous bound resources
 		Tracker->ResetUAVState(RHIValidation::EUAVMode::Graphics);
 
-		State.GlobalUniformBuffers.bInSetPipelineStateCall = true;
+		State.StaticUniformBuffers.bInSetPipelineStateCall = true;
 		RHIContext->RHISetGraphicsPipelineState(PsoInit, StencilRef, bApplyAdditionalState);
-		State.GlobalUniformBuffers.bInSetPipelineStateCall = false;
+		State.StaticUniformBuffers.bInSetPipelineStateCall = false;
 	}
 #endif
 
@@ -701,14 +703,14 @@ public:
 	virtual void RHISetShaderUniformBuffer(FRHIGraphicsShader* Shader, uint32 BufferIndex, FRHIUniformBuffer* Buffer) override final
 	{
 		checkf(State.bGfxPSOSet, TEXT("A Graphics PSO has to be set to set resources into a shader!"));
-		State.GlobalUniformBuffers.ValidateSetShaderUniformBuffer(Buffer);
+		State.StaticUniformBuffers.ValidateSetShaderUniformBuffer(Buffer);
 		RHIContext->RHISetShaderUniformBuffer(Shader, BufferIndex, Buffer);
 	}
 
 	virtual void RHISetShaderUniformBuffer(FRHIComputeShader* Shader, uint32 BufferIndex, FRHIUniformBuffer* Buffer) override final
 	{
 		checkf(State.bComputePSOSet, TEXT("A Compute PSO has to be set to set resources into a shader!"));
-		State.GlobalUniformBuffers.ValidateSetShaderUniformBuffer(Buffer);
+		State.StaticUniformBuffers.ValidateSetShaderUniformBuffer(Buffer);
 		RHIContext->RHISetShaderUniformBuffer(Shader, BufferIndex, Buffer);
 	}
 
@@ -726,7 +728,7 @@ public:
 
 	virtual void RHISetStaticUniformBuffers(const FUniformBufferStaticBindings& InUniformBuffers) override final
 	{
-		InUniformBuffers.Bind(State.GlobalUniformBuffers.Bindings);
+		InUniformBuffers.Bind(State.StaticUniformBuffers.Bindings);
 		RHIContext->RHISetStaticUniformBuffers(InUniformBuffers);
 	}
 
@@ -826,11 +828,13 @@ public:
 
 	virtual void RHIPushEvent(const TCHAR* Name, FColor Color) override final
 	{
+		Tracker->PushBreadcrumb(Name);
 		RHIContext->RHIPushEvent(Name, Color);
 	}
 
 	virtual void RHIPopEvent() override final
 	{
+		Tracker->PopBreadcrumb();
 		RHIContext->RHIPopEvent();
 	}
 
@@ -1113,7 +1117,7 @@ protected:
 	struct FState
 	{
 		RHIValidation::FTracker TrackerInstance{ ERHIPipeline::Graphics };
-		RHIValidation::FGlobalUniformBuffers GlobalUniformBuffers;
+		RHIValidation::FStaticUniformBuffers StaticUniformBuffers;
 
 		void* PreviousBeginFrame = nullptr;
 		void* PreviousEndFrame = nullptr;
