Fix CIS compilation issue.

diff --git a/Engine/Source/Runtime/RHI/Public/RHIResources.h b/Engine/Source/Runtime/RHI/Public/RHIResources.h
--- a/Engine/Source/Runtime/RHI/Public/RHIResources.h
+++ b/Engine/Source/Runtime/RHI/Public/RHIResources.h
@@ -1673,7 +1673,7 @@ struct FRHITextureCreateDesc : public FRHITextureDesc
 	FRHITextureCreateDesc& SetGPUMask(FRHIGPUMask InGPUMask)                   { GPUMask = InGPUMask;                      return *this; }
 	FRHITextureCreateDesc& SetBulkData(FResourceBulkDataInterface* InBulkData) { BulkData = InBulkData;                    return *this; }
 	FRHITextureCreateDesc& DetermineInititialState()                           { if (InitialState == ERHIAccess::Unknown) InitialState = RHIGetDefaultResourceState(Flags, BulkData != nullptr); return *this; }
-	FRHITextureCreateDesc& SetFastVRAMPercentage(float In)                     { FastVRAMPercentage = FMath::Clamp(In, 0.f, 1.0f) * 0xFF; return *this; }
+	FRHITextureCreateDesc& SetFastVRAMPercentage(float In)                     { FastVRAMPercentage = uint8(FMath::Clamp(In, 0.f, 1.0f) * 0xFF); return *this; }
 
 	/* The RHI access state that the resource will be created in. */
 	ERHIAccess InitialState = ERHIAccess::Unknown;
