Temporary fix for Editor

diff --git a/Engine/Source/Runtime/RHI/Private/RHI.cpp b/Engine/Source/Runtime/RHI/Private/RHI.cpp
--- a/Engine/Source/Runtime/RHI/Private/RHI.cpp
+++ b/Engine/Source/Runtime/RHI/Private/RHI.cpp
@@ -2671,11 +2671,11 @@ void FGenericDataDrivenShaderPlatformInfo::UpdatePreviewPlatforms()
 #if WITH_EDITOR
 FText FGenericDataDrivenShaderPlatformInfo::GetFriendlyName(const FStaticShaderPlatform Platform)
 {
-	if (IsRunningCommandlet() || GUsingNullRHI)
+	if (IsRunningCommandlet() || GUsingNullRHI || !IsValid(Platform))
 	{
 		return FText();
 	}
-	check(IsValid(Platform));
+	//check(IsValid(Platform));
 	return Infos[Platform].FriendlyName;
 }
 #endif
