Memory leak in RHIThreadFence

#localization none
#tests Zach reproduced locally

diff --git a/Engine/Source/Runtime/RHI/Private/RHICommandList.cpp b/Engine/Source/Runtime/RHI/Private/RHICommandList.cpp
--- a/Engine/Source/Runtime/RHI/Private/RHICommandList.cpp
+++ b/Engine/Source/Runtime/RHI/Private/RHICommandList.cpp
@@ -989,15 +989,18 @@ FRHICOMMAND_MACRO(FRHICommandRHIThreadFence)
 
 FGraphEventRef FRHICommandListBase::RHIThreadFence(bool bSetLockFence)
 {
+	if (bSetLockFence)
+	{
+		QueuedFenceCandidateEvents.Empty();
+		QueuedFenceCandidates.Empty();
+	}
+
 	if (IsRunningRHIInSeparateThread())
 	{
 		FRHICommandRHIThreadFence* Cmd = ALLOC_COMMAND(FRHICommandRHIThreadFence)();
 		if (bSetLockFence)
 		{
 			RHIThreadBufferLockFence = Cmd->Fence;
-			QueuedFenceCandidateEvents.Reset();
-			QueuedFenceCandidates.Reset();
-
 		}
 		return Cmd->Fence;
 	}
@@ -1694,8 +1697,11 @@ void FRHICommandListBase::QueueRenderThreadCommandListSubmit(FGraphEventRef& Ren
 		return;
 	}
 
-	QueuedFenceCandidateEvents.Emplace(RenderThreadCompletionEvent);
-	QueuedFenceCandidates.Emplace(CmdList->FenceCandidate);
+	if (IsRunningRHIInSeparateThread())
+	{
+		QueuedFenceCandidateEvents.Emplace(RenderThreadCompletionEvent);
+		QueuedFenceCandidates.Emplace(CmdList->FenceCandidate);
+	}
 
 	ALLOC_COMMAND(FRHICommandWaitForAndSubmitRTSubList)(RenderThreadCompletionEvent, CmdList);
 
