Preview Platforms set bSupportsMobileMultiView, bSupportsInstancedStereo from the parent platform to avoid shader compilation issues.

diff --git a/Engine/Source/Runtime/RHI/Private/RHI.cpp b/Engine/Source/Runtime/RHI/Private/RHI.cpp
--- a/Engine/Source/Runtime/RHI/Private/RHI.cpp
+++ b/Engine/Source/Runtime/RHI/Private/RHI.cpp
@@ -2658,6 +2658,8 @@ void FGenericDataDrivenShaderPlatformInfo::UpdatePreviewPlatforms()
 				Infos[ShaderPlatform].bSupportsNanite = Infos[EditorSPForPreviewMaxFeatureLevel].bSupportsNanite;
 				Infos[ShaderPlatform].bSupportsUInt64ImageAtomics = Infos[EditorSPForPreviewMaxFeatureLevel].bSupportsUInt64ImageAtomics;
 				Infos[ShaderPlatform].bSupportsGen5TemporalAA = Infos[EditorSPForPreviewMaxFeatureLevel].bSupportsGen5TemporalAA;
+				Infos[ShaderPlatform].bSupportsMobileMultiView = Infos[EditorSPForPreviewMaxFeatureLevel].bSupportsMobileMultiView;
+				Infos[ShaderPlatform].bSupportsInstancedStereo = Infos[EditorSPForPreviewMaxFeatureLevel].bSupportsInstancedStereo;
 				Infos[ShaderPlatform].bSupportsRenderTargetWriteMask = false;
 				Infos[ShaderPlatform].bSupportsIntrinsicWaveOnce = false;
 				Infos[ShaderPlatform].bSupportsDOFHybridScattering = false;
