- More than one PreviewPlatform can now use the same ShaderPlatforms
- Clean up temporary fix
- Modify UpdatePreviewPlatforms to only restrict DDSPI properties and not enable them if the preview platform doesn't have them on

diff --git a/Engine/Source/Editor/LevelEditor/Private/LevelEditor.cpp b/Engine/Source/Editor/LevelEditor/Private/LevelEditor.cpp
--- a/Engine/Source/Editor/LevelEditor/Private/LevelEditor.cpp
+++ b/Engine/Source/Editor/LevelEditor/Private/LevelEditor.cpp
@@ -1793,32 +1793,17 @@ void FLevelEditorModule::BindGlobalLevelEditorCommands()
 	for (int32 Index=0; Index < MenuItems.Num(); Index++)
 	{
 		const FPreviewPlatformMenuItem& Item = MenuItems[Index];
-		EShaderPlatform ShaderPlatform = ShaderFormatToLegacyShaderPlatform(Item.ShaderFormat);
+		EShaderPlatform ShaderPlatform = FDataDrivenShaderPlatformInfo::GetShaderPlatformFromName(Item.PreviewShaderPlatformName);
 		
 		if (ShaderPlatform < SP_NumPlatforms)
 		{
 			ERHIFeatureLevel::Type FeatureLevel = GetMaxSupportedFeatureLevel(ShaderPlatform);
 
-			const bool IsDefaultActive = ShaderPlatform == GMaxRHIShaderPlatform;
+			const bool IsDefaultActive = FDataDrivenShaderPlatformInfo::GetPreviewShaderPlatformParent(ShaderPlatform) == GMaxRHIShaderPlatform;
 			const bool AllowPreview = !IsDefaultActive;
 
-			EShaderPlatform InShaderPlatform = EShaderPlatform::SP_NumPlatforms;
-			for (uint32 j = 0; j < SP_NumPlatforms; ++j)
-			{
-				InShaderPlatform = (EShaderPlatform)j;
-				if (!FDataDrivenShaderPlatformInfo::IsValid(InShaderPlatform))
-				{
-					continue;
-				}
-
-				FName ShaderPlatformTest = FDataDrivenShaderPlatformInfo::GetName(InShaderPlatform);
-				if (Item.ShaderPlatformPreview == ShaderPlatformTest)
-				{
-					break;
-				}
-			}
 
-			FPreviewPlatformInfo PreviewFeatureLevelInfo(FeatureLevel, (EShaderPlatform)InShaderPlatform, IsDefaultActive ? NAME_None : Item.PlatformName, IsDefaultActive ? NAME_None : Item.ShaderFormat, IsDefaultActive ? NAME_None : Item.DeviceProfileName, AllowPreview, Item.ShaderPlatformPreview);
+			FPreviewPlatformInfo PreviewFeatureLevelInfo(FeatureLevel, (EShaderPlatform)ShaderPlatform, IsDefaultActive ? NAME_None : Item.PlatformName, IsDefaultActive ? NAME_None : Item.ShaderFormat, IsDefaultActive ? NAME_None : Item.DeviceProfileName, AllowPreview, Item.PreviewShaderPlatformName);
 
 			ActionList.MapAction(
 				Commands.PreviewPlatformOverrides[Index],
diff --git a/Engine/Source/Editor/LevelEditor/Private/LevelEditorActions.cpp b/Engine/Source/Editor/LevelEditor/Private/LevelEditorActions.cpp
--- a/Engine/Source/Editor/LevelEditor/Private/LevelEditorActions.cpp
+++ b/Engine/Source/Editor/LevelEditor/Private/LevelEditorActions.cpp
@@ -3718,7 +3718,7 @@ void FLevelEditorCommands::RegisterCommands()
 		FText FriendlyName;
 		if (!IsRunningCommandlet() && !GUsingNullRHI)
 		{
-			EShaderPlatform ShaderPlatform = FDataDrivenShaderPlatformInfo::GetShaderPlatformFromName(Item.ShaderPlatformPreview);
+			EShaderPlatform ShaderPlatform = FDataDrivenShaderPlatformInfo::GetShaderPlatformFromName(Item.PreviewShaderPlatformName);
 			FriendlyName = FDataDrivenShaderPlatformInfo::GetFriendlyName(ShaderPlatform);
 		}
 
diff --git a/Engine/Source/Editor/LevelEditor/Private/LevelEditorToolBar.cpp b/Engine/Source/Editor/LevelEditor/Private/LevelEditorToolBar.cpp
--- a/Engine/Source/Editor/LevelEditor/Private/LevelEditorToolBar.cpp
+++ b/Engine/Source/Editor/LevelEditor/Private/LevelEditorToolBar.cpp
@@ -1400,15 +1400,15 @@ void FLevelEditorToolBar::RegisterLevelEditorToolBar( const TSharedRef<FUIComman
 			{
 				const FPreviewPlatformMenuItem* Item = FDataDrivenPlatformInfoRegistry::GetAllPreviewPlatformMenuItems().FindByPredicate([](const FPreviewPlatformMenuItem& TestItem)
 					{
-						return GEditor->PreviewPlatform.PreviewPlatformName == TestItem.PlatformName && GEditor->PreviewPlatform.PreviewShaderFormatName == TestItem.ShaderFormat && GEditor->PreviewPlatform.ShaderPlatformPreview == TestItem.ShaderPlatformPreview;
+						return GEditor->PreviewPlatform.PreviewPlatformName == TestItem.PlatformName && GEditor->PreviewPlatform.PreviewShaderFormatName == TestItem.ShaderFormat && GEditor->PreviewPlatform.PreviewShaderPlatformName == TestItem.PreviewShaderPlatformName;
 					});
 				return Item ? Item->IconText : FText();
 			}
 
 			static FText GetPreviewModeTooltip()
 			{
-				EShaderPlatform PreviewShaderPlatform = GEditor->PreviewPlatform.PreviewShaderFormatName != NAME_None ?
-					ShaderFormatToLegacyShaderPlatform(GEditor->PreviewPlatform.PreviewShaderFormatName) :
+				EShaderPlatform PreviewShaderPlatform = GEditor->PreviewPlatform.PreviewShaderPlatformName != NAME_None ?
+					FDataDrivenShaderPlatformInfo::GetShaderPlatformFromName(GEditor->PreviewPlatform.PreviewShaderPlatformName) :
 					GetFeatureLevelShaderPlatform(GEditor->PreviewPlatform.PreviewFeatureLevel);
 
 				EShaderPlatform MaxRHIFeatureLevelPlatform = GetFeatureLevelShaderPlatform(GMaxRHIFeatureLevel);
@@ -1435,14 +1435,14 @@ void FLevelEditorToolBar::RegisterLevelEditorToolBar( const TSharedRef<FUIComman
 			{
 				const FPreviewPlatformMenuItem* Item = FDataDrivenPlatformInfoRegistry::GetAllPreviewPlatformMenuItems().FindByPredicate([](const FPreviewPlatformMenuItem& TestItem)
 					{
-						return GEditor->PreviewPlatform.PreviewPlatformName == TestItem.PlatformName && GEditor->PreviewPlatform.PreviewShaderFormatName == TestItem.ShaderFormat && GEditor->PreviewPlatform.ShaderPlatformPreview == TestItem.ShaderPlatformPreview;
+						return GEditor->PreviewPlatform.PreviewPlatformName == TestItem.PlatformName && GEditor->PreviewPlatform.PreviewShaderFormatName == TestItem.ShaderFormat && GEditor->PreviewPlatform.PreviewShaderPlatformName == TestItem.PreviewShaderPlatformName;
 					});
 				if(Item)
 				{
 					return FSlateIcon(FAppStyle::GetAppStyleSetName(), GEditor->IsFeatureLevelPreviewActive() ? Item->ActiveIconName : Item->InactiveIconName);
 				}
 
-				EShaderPlatform ShaderPlatform = ShaderFormatToLegacyShaderPlatform(GEditor->PreviewPlatform.PreviewShaderFormatName);
+				EShaderPlatform ShaderPlatform = FDataDrivenShaderPlatformInfo::GetShaderPlatformFromName(GEditor->PreviewPlatform.PreviewShaderPlatformName);
 
 				if (ShaderPlatform == SP_NumPlatforms)
 				{
diff --git a/Engine/Source/Editor/UnrealEd/Classes/Editor/EditorEngine.h b/Engine/Source/Editor/UnrealEd/Classes/Editor/EditorEngine.h
--- a/Engine/Source/Editor/UnrealEd/Classes/Editor/EditorEngine.h
+++ b/Engine/Source/Editor/UnrealEd/Classes/Editor/EditorEngine.h
@@ -213,7 +213,7 @@ struct FPreviewPlatformInfo
 	,	PreviewPlatformName(NAME_None)
 	,	PreviewShaderFormatName(NAME_None)
 	,	bPreviewFeatureLevelActive(false)
-	,	ShaderPlatformPreview(NAME_None)
+	,	PreviewShaderPlatformName(NAME_None)
 	{}
 
 	FPreviewPlatformInfo(ERHIFeatureLevel::Type InFeatureLevel, EShaderPlatform InShaderPlatform = EShaderPlatform::SP_NumPlatforms, FName InPreviewPlatformName = NAME_None, FName InPreviewShaderFormatName = NAME_None, FName InDeviceProfileName = NAME_None, bool InbPreviewFeatureLevelActive = false, FName InShaderPlatformName = NAME_None)
@@ -223,7 +223,7 @@ struct FPreviewPlatformInfo
 	,	PreviewShaderFormatName(InPreviewShaderFormatName)
 	,	DeviceProfileName(InDeviceProfileName)
 	,	bPreviewFeatureLevelActive(InbPreviewFeatureLevelActive)
-	,	ShaderPlatformPreview(InShaderPlatformName)
+	,	PreviewShaderPlatformName(InShaderPlatformName)
 	{}
 
 	/** The feature level we should use when loading or creating a new world */
@@ -244,13 +244,13 @@ struct FPreviewPlatformInfo
 	/** Is feature level preview currently active */
 	bool bPreviewFeatureLevelActive;
 	
-	/** The shader platform to preview*/
-	FName ShaderPlatformPreview;
+	/** Preview Shader Platform Name */
+	FName PreviewShaderPlatformName;
 
 	/** Checks if two FPreviewPlatformInfos are for the same preview platform. Note, this does NOT compare the bPreviewFeatureLevelActive flag */
 	bool Matches(const FPreviewPlatformInfo& Other) const
 	{
-		return PreviewFeatureLevel == Other.PreviewFeatureLevel && ShaderPlatform == Other.ShaderPlatform && PreviewPlatformName == Other.PreviewPlatformName && PreviewShaderFormatName == Other.PreviewShaderFormatName && DeviceProfileName == Other.DeviceProfileName && ShaderPlatformPreview == Other.ShaderPlatformPreview;
+		return PreviewFeatureLevel == Other.PreviewFeatureLevel && ShaderPlatform == Other.ShaderPlatform && PreviewPlatformName == Other.PreviewPlatformName && PreviewShaderFormatName == Other.PreviewShaderFormatName && DeviceProfileName == Other.DeviceProfileName && PreviewShaderPlatformName == Other.PreviewShaderPlatformName;
 	}
 
 	/** Return platform name like "Android", or NAME_None if none is set or the preview feature level is not active */
diff --git a/Engine/Source/Editor/UnrealEd/Private/EditorEngine.cpp b/Engine/Source/Editor/UnrealEd/Private/EditorEngine.cpp
--- a/Engine/Source/Editor/UnrealEd/Private/EditorEngine.cpp
+++ b/Engine/Source/Editor/UnrealEd/Private/EditorEngine.cpp
@@ -7779,7 +7779,7 @@ void UEditorEngine::ToggleFeatureLevelPreview()
 
 bool UEditorEngine::IsFeatureLevelPreviewEnabled() const
 {
-	return PreviewPlatform.PreviewFeatureLevel != GMaxRHIFeatureLevel || PreviewPlatform.PreviewShaderFormatName != NAME_None || PreviewPlatform.ShaderPlatformPreview != NAME_None;
+	return PreviewPlatform.PreviewFeatureLevel != GMaxRHIFeatureLevel || PreviewPlatform.PreviewShaderFormatName != NAME_None || PreviewPlatform.PreviewShaderPlatformName != NAME_None;
 }
 
 bool UEditorEngine::IsFeatureLevelPreviewActive() const
@@ -7798,7 +7798,7 @@ void UEditorEngine::LoadEditorFeatureLevel()
 
 	EShaderPlatform ShaderPlatformToPreview = FDataDrivenShaderPlatformInfo::GetShaderPlatformFromName(Settings->PreviewShaderPlatformName);
 
-	if (Settings->PreviewFeatureLevel >= 0 && Settings->PreviewFeatureLevel < (int32)ERHIFeatureLevel::Num && ShaderPlatformToPreview < (int32)EShaderPlatform::SP_NumPlatforms)
+	if (Settings->PreviewFeatureLevel >= 0 && Settings->PreviewFeatureLevel < (int32)ERHIFeatureLevel::Num && ShaderPlatformToPreview < EShaderPlatform::SP_NumPlatforms)
 	{
 		// Try to map a saved ShaderFormatName to the PreviewPlatformName using ITargetPlatform if we don't have one. 
 		// We now store the PreviewPlatformName explicitly to support preview for platforms we don't have an ITargetPlatform of.
@@ -7830,7 +7830,7 @@ void UEditorEngine::SaveEditorFeatureLevel()
 	Settings->bPreviewFeatureLevelActive = PreviewPlatform.bPreviewFeatureLevelActive;
 	Settings->bPreviewFeatureLevelWasDefault = (PreviewPlatform.PreviewFeatureLevel == GMaxRHIFeatureLevel);
 	Settings->PreviewDeviceProfileName = PreviewPlatform.DeviceProfileName;
-	Settings->PreviewShaderPlatformName = PreviewPlatform.ShaderPlatformPreview;
+	Settings->PreviewShaderPlatformName = PreviewPlatform.PreviewShaderPlatformName;
 	Settings->PostEditChange();
 }
 
diff --git a/Engine/Source/Runtime/Core/Private/Misc/DataDrivenPlatformInfoRegistry.cpp b/Engine/Source/Runtime/Core/Private/Misc/DataDrivenPlatformInfoRegistry.cpp
--- a/Engine/Source/Runtime/Core/Private/Misc/DataDrivenPlatformInfoRegistry.cpp
+++ b/Engine/Source/Runtime/Core/Private/Misc/DataDrivenPlatformInfoRegistry.cpp
@@ -252,7 +252,7 @@ static void ParsePreviewPlatforms(const FConfigFile& IniFile)
 		if (Section.Key.StartsWith(TEXT("PreviewPlatform ")))
 		{
 			const FString& SectionName = Section.Key;
-			FName PreviewPlatformName = *SectionName.Mid(16);
+			FName PreviewPlatformName = *(SectionName.Mid(16) + TEXT("_Preview"));
 
 			// Early-out if enabled cvar is specified and not set
 			TArray<FString> Tokens;
@@ -277,6 +277,7 @@ static void ParsePreviewPlatforms(const FConfigFile& IniFile)
 
 			FPreviewPlatformMenuItem Item;
 			Item.PlatformName = PlatformName;
+			Item.PreviewShaderPlatformName = PreviewPlatformName;
 			Item.ShaderFormat = *GetSectionString(Section.Value, FName("ShaderFormat"));
 			checkf(Item.ShaderFormat != NAME_None, TEXT("DataDrivenPlatformInfo section [PreviewPlatform %s] must specify a ShaderFormat"), *SectionName);
 			Item.ActiveIconPath = GetSectionString(Section.Value, FName("ActiveIconPath"));
@@ -284,7 +285,9 @@ static void ParsePreviewPlatforms(const FConfigFile& IniFile)
 			Item.InactiveIconPath = GetSectionString(Section.Value, FName("InactiveIconPath"));
 			Item.InactiveIconName = *GetSectionString(Section.Value, FName("InactiveIconName"));
 			Item.DeviceProfileName = *GetSectionString(Section.Value, FName("DeviceProfileName"));
-			Item.ShaderPlatformPreview = *(GetSectionString(Section.Value, FName("ShaderPlatform")) + TEXT("_PREVIEW"));
+			Item.ShaderPlatformToPreview = *GetSectionString(Section.Value, FName("ShaderPlatform"));
+			checkf(Item.ShaderPlatformToPreview != NAME_None, TEXT("DataDrivenPlatformInfo section [PreviewPlatform %s] must specify a ShaderPlatform"), *SectionName);
+			FTextStringHelper::ReadFromBuffer(*GetSectionString(Section.Value, FName("FriendlyName")), Item.OptionalFriendlyNameOverride);
 			FTextStringHelper::ReadFromBuffer(*GetSectionString(Section.Value, FName("MenuTooltip")), Item.MenuTooltip);
 			FTextStringHelper::ReadFromBuffer(*GetSectionString(Section.Value, FName("IconText")), Item.IconText);
 			PreviewPlatformMenuItems.Add(Item);
diff --git a/Engine/Source/Runtime/Core/Public/Misc/DataDrivenPlatformInfoRegistry.h b/Engine/Source/Runtime/Core/Public/Misc/DataDrivenPlatformInfoRegistry.h
--- a/Engine/Source/Runtime/Core/Public/Misc/DataDrivenPlatformInfoRegistry.h
+++ b/Engine/Source/Runtime/Core/Public/Misc/DataDrivenPlatformInfoRegistry.h
@@ -99,11 +99,13 @@ struct FPreviewPlatformMenuItem
 {
 	FName PlatformName;
 	FName ShaderFormat;
-	FName ShaderPlatformPreview;
+	FName ShaderPlatformToPreview;
+	FName PreviewShaderPlatformName;
 	FString ActiveIconPath;
 	FName ActiveIconName;
 	FString InactiveIconPath;
 	FName InactiveIconName;
+	FText OptionalFriendlyNameOverride;
 	FText MenuTooltip;
 	FText IconText;
 	FName DeviceProfileName;
diff --git a/Engine/Source/Runtime/RHI/Private/RHI.cpp b/Engine/Source/Runtime/RHI/Private/RHI.cpp
--- a/Engine/Source/Runtime/RHI/Private/RHI.cpp
+++ b/Engine/Source/Runtime/RHI/Private/RHI.cpp
@@ -2609,12 +2609,16 @@ void FGenericDataDrivenShaderPlatformInfo::Initialize()
 #if WITH_EDITOR
 				for (const FPreviewPlatformMenuItem& Item : FDataDrivenPlatformInfoRegistry::GetAllPreviewPlatformMenuItems())
 				{
-					FName PreviewPlatformName = *((Infos[ShaderPlatform].Name).ToString() + TEXT("_PREVIEW"));
-					if (Item.ShaderPlatformPreview == PreviewPlatformName)
+					FName PreviewPlatformName = *(Infos[ShaderPlatform].Name).ToString();
+					if (Item.ShaderPlatformToPreview == PreviewPlatformName)
 					{
 						EShaderPlatform PreviewShaderPlatform = EShaderPlatform(CustomShaderPlatform++);
 						ParseDataDrivenShaderInfo(Section.Value, Infos[PreviewShaderPlatform]);
-						Infos[PreviewShaderPlatform].Name = PreviewPlatformName;
+						if (!Item.OptionalFriendlyNameOverride.IsEmpty())
+						{
+							Infos[PreviewShaderPlatform].FriendlyName = Item.OptionalFriendlyNameOverride;
+						}
+						Infos[PreviewShaderPlatform].Name = Item.PreviewShaderPlatformName;
 						Infos[PreviewShaderPlatform].PreviewShaderPlatformParent = ShaderPlatform;
 						Infos[PreviewShaderPlatform].bIsPreviewPlatform = true;
 						Infos[PreviewShaderPlatform].bContainsValidPlatformInfo = true;
@@ -2642,15 +2646,15 @@ void FGenericDataDrivenShaderPlatformInfo::UpdatePreviewPlatforms()
 				Infos[ShaderPlatform].Language = Infos[EditorSPForPreviewMaxFeatureLevel].Language;
 				Infos[ShaderPlatform].bIsHlslcc = Infos[EditorSPForPreviewMaxFeatureLevel].bIsHlslcc;
 				Infos[ShaderPlatform].bSupportsDxc = Infos[EditorSPForPreviewMaxFeatureLevel].bSupportsDxc;
-				Infos[ShaderPlatform].bSupportsGPUScene = Infos[EditorSPForPreviewMaxFeatureLevel].bSupportsGPUScene;
+				Infos[ShaderPlatform].bSupportsGPUScene &= Infos[EditorSPForPreviewMaxFeatureLevel].bSupportsGPUScene;
 				Infos[ShaderPlatform].bIsPC = true;
 				Infos[ShaderPlatform].bIsConsole = false;
-				Infos[ShaderPlatform].bSupportsSceneDataCompressedTransforms = Infos[EditorSPForPreviewMaxFeatureLevel].bSupportsSceneDataCompressedTransforms;
-				Infos[ShaderPlatform].bSupportsNanite = Infos[EditorSPForPreviewMaxFeatureLevel].bSupportsNanite;
-				Infos[ShaderPlatform].bSupportsUInt64ImageAtomics = Infos[EditorSPForPreviewMaxFeatureLevel].bSupportsUInt64ImageAtomics;
-				Infos[ShaderPlatform].bSupportsGen5TemporalAA = Infos[EditorSPForPreviewMaxFeatureLevel].bSupportsGen5TemporalAA;
-				Infos[ShaderPlatform].bSupportsMobileMultiView = Infos[EditorSPForPreviewMaxFeatureLevel].bSupportsMobileMultiView;
-				Infos[ShaderPlatform].bSupportsInstancedStereo = Infos[EditorSPForPreviewMaxFeatureLevel].bSupportsInstancedStereo;
+				Infos[ShaderPlatform].bSupportsSceneDataCompressedTransforms &= Infos[EditorSPForPreviewMaxFeatureLevel].bSupportsSceneDataCompressedTransforms;
+				Infos[ShaderPlatform].bSupportsNanite &= Infos[EditorSPForPreviewMaxFeatureLevel].bSupportsNanite;
+				Infos[ShaderPlatform].bSupportsUInt64ImageAtomics &= Infos[EditorSPForPreviewMaxFeatureLevel].bSupportsUInt64ImageAtomics;
+				Infos[ShaderPlatform].bSupportsGen5TemporalAA &= Infos[EditorSPForPreviewMaxFeatureLevel].bSupportsGen5TemporalAA;
+				Infos[ShaderPlatform].bSupportsMobileMultiView &= Infos[EditorSPForPreviewMaxFeatureLevel].bSupportsMobileMultiView;
+				Infos[ShaderPlatform].bSupportsInstancedStereo &= Infos[EditorSPForPreviewMaxFeatureLevel].bSupportsInstancedStereo;
 				Infos[ShaderPlatform].bSupportsRenderTargetWriteMask = false;
 				Infos[ShaderPlatform].bSupportsIntrinsicWaveOnce = false;
 				Infos[ShaderPlatform].bSupportsDOFHybridScattering = false;
@@ -2664,11 +2668,11 @@ void FGenericDataDrivenShaderPlatformInfo::UpdatePreviewPlatforms()
 #if WITH_EDITOR
 FText FGenericDataDrivenShaderPlatformInfo::GetFriendlyName(const FStaticShaderPlatform Platform)
 {
-	if (IsRunningCommandlet() || GUsingNullRHI || !IsValid(Platform))
+	if (IsRunningCommandlet() || GUsingNullRHI)
 	{
 		return FText();
 	}
-	//check(IsValid(Platform));
+	check(IsValid(Platform));
 	return Infos[Platform].FriendlyName;
 }
 #endif
diff --git a/Engine/Source/Runtime/RenderCore/Private/RenderUtils.cpp b/Engine/Source/Runtime/RenderCore/Private/RenderUtils.cpp
--- a/Engine/Source/Runtime/RenderCore/Private/RenderUtils.cpp
+++ b/Engine/Source/Runtime/RenderCore/Private/RenderUtils.cpp
@@ -1287,15 +1287,14 @@ RENDERCORE_API ShaderPlatformMaskType GRayTracingPlatformMask;
 // Value may be queried using IsRayTracingEnabled().
 RENDERCORE_API bool GUseRayTracing = false;
 
-void GetAllPossiblePreviewPlatformsForMainShaderPlatform(TArray<EShaderPlatform>& OutPreviewPlatforms, FName ShaderPlatformName)
+void GetAllPossiblePreviewPlatformsForMainShaderPlatform(TArray<EShaderPlatform>& OutPreviewPlatforms, EShaderPlatform ParentShaderPlatform)
 {
-	FName PreviewPlatformName = FName(ShaderPlatformName.ToString() + TEXT("_PREVIEW"));
 	for (int i = 0; i < EShaderPlatform::SP_NumPlatforms; ++i)
 	{
 		EShaderPlatform ShaderPlatform = EShaderPlatform(i);
 		if (FDataDrivenShaderPlatformInfo::IsValid(ShaderPlatform))
 		{
-			bool bIsPreviewPlatform = FDataDrivenShaderPlatformInfo::GetIsPreviewPlatform(ShaderPlatform) && (FDataDrivenShaderPlatformInfo::GetName(ShaderPlatform) == PreviewPlatformName);
+			bool bIsPreviewPlatform = FDataDrivenShaderPlatformInfo::GetIsPreviewPlatform(ShaderPlatform) && (FDataDrivenShaderPlatformInfo::GetPreviewShaderPlatformParent(ShaderPlatform) == ParentShaderPlatform);
 			if (bIsPreviewPlatform)
 			{
 				OutPreviewPlatforms.Add(ShaderPlatform);
@@ -1341,7 +1340,7 @@ RENDERCORE_API void RenderUtilsInit()
 			{
 				EShaderPlatform ShaderPlatform = ShaderFormatNameToShaderPlatform(Format);
 				TArray<EShaderPlatform> PossiblePreviewPlatformsAndMainPlatform;
-				GetAllPossiblePreviewPlatformsForMainShaderPlatform(PossiblePreviewPlatformsAndMainPlatform, FDataDrivenShaderPlatformInfo::GetName(ShaderPlatform));
+				GetAllPossiblePreviewPlatformsForMainShaderPlatform(PossiblePreviewPlatformsAndMainPlatform, ShaderPlatform);
 				PossiblePreviewPlatformsAndMainPlatform.Add(ShaderPlatform);
 
 				for (EShaderPlatform ShaderPlatformToEdit : PossiblePreviewPlatformsAndMainPlatform)
@@ -1379,10 +1378,10 @@ RENDERCORE_API void RenderUtilsInit()
 
 				for (FName FormatName : PlatformRayTracingShaderFormats)
 				{
+					EShaderPlatform MainShaderPlatform = ShaderFormatNameToShaderPlatform(FormatName);
 					TArray<EShaderPlatform> PossiblePreviewPlatformsAndMainPlatform;
-					GetAllPossiblePreviewPlatformsForMainShaderPlatform(PossiblePreviewPlatformsAndMainPlatform, FormatName);
+					GetAllPossiblePreviewPlatformsForMainShaderPlatform(PossiblePreviewPlatformsAndMainPlatform, MainShaderPlatform);
 
-					EShaderPlatform MainShaderPlatform = ShaderFormatNameToShaderPlatform(FormatName);
 					PossiblePreviewPlatformsAndMainPlatform.Add(MainShaderPlatform);
 
 					for (EShaderPlatform ShaderPlatform : PossiblePreviewPlatformsAndMainPlatform)
