Fix RHI transition validation for ray tracing scene buffer

- Remove transitions on the scene object and only perform them on the underlying buffer itself
- Add ERHIAccess::SRVMask when transitioning the acceleration structure buffer to readable state

#jrra UE-150632

diff --git a/Engine/Source/Runtime/RHI/Public/RHIValidationContext.h b/Engine/Source/Runtime/RHI/Public/RHIValidationContext.h
--- a/Engine/Source/Runtime/RHI/Public/RHIValidationContext.h
+++ b/Engine/Source/Runtime/RHI/Public/RHIValidationContext.h
@@ -1030,8 +1030,6 @@ public:
 		FRHIUnorderedAccessView* Output,
 		uint32 NumRays) override final
 	{
-		// #yuriy_todo: validate all referenced BLAS states
-		Tracker->Assert(Scene->GetWholeResourceIdentity(), ERHIAccess::BVHRead);
 		RHIContext->RHIRayTraceOcclusion(Scene, Rays, Output, NumRays);
 	}
 
@@ -1040,8 +1038,6 @@ public:
 		FRHIUnorderedAccessView* Output,
 		uint32 NumRays) override final
 	{
-		// #yuriy_todo: validate all referenced BLAS states
-		Tracker->Assert(Scene->GetWholeResourceIdentity(), ERHIAccess::BVHRead);
 		RHIContext->RHIRayTraceIntersection(Scene, Rays, Output, NumRays);
 	}
 
@@ -1050,8 +1046,6 @@ public:
 		const FRayTracingShaderBindings& GlobalResourceBindings,
 		uint32 Width, uint32 Height) override final
 	{
-		// #yuriy_todo: validate all referenced BLAS states
-		Tracker->Assert(Scene->GetWholeResourceIdentity(), ERHIAccess::BVHRead);
 		RHIContext->RHIRayTraceDispatch(RayTracingPipelineState, RayGenShader, Scene, GlobalResourceBindings, Width, Height);
 	}
 
@@ -1060,10 +1054,8 @@ public:
 		const FRayTracingShaderBindings& GlobalResourceBindings,
 		FRHIBuffer* ArgumentBuffer, uint32 ArgumentOffset) override final
 	{
-		// #yuriy_todo: validate all referenced BLAS states
 		Tracker->Assert(ArgumentBuffer->GetWholeResourceIdentity(), ERHIAccess::IndirectArgs);
 		Tracker->Assert(ArgumentBuffer->GetWholeResourceIdentity(), ERHIAccess::SRVCompute);
-		Tracker->Assert(Scene->GetWholeResourceIdentity(), ERHIAccess::BVHRead);
 
 		RHIContext->RHIRayTraceDispatchIndirect(RayTracingPipelineState, RayGenShader, Scene, GlobalResourceBindings, ArgumentBuffer, ArgumentOffset);
 	}
diff --git a/Engine/Source/Runtime/Renderer/Private/DeferredShadingRenderer.cpp b/Engine/Source/Runtime/Renderer/Private/DeferredShadingRenderer.cpp
--- a/Engine/Source/Runtime/Renderer/Private/DeferredShadingRenderer.cpp
+++ b/Engine/Source/Runtime/Renderer/Private/DeferredShadingRenderer.cpp
@@ -1772,6 +1772,9 @@ bool FDeferredShadingSceneRenderer::DispatchRayTracingWorldUpdates(FRDGBuilder&
 
 	RayTracingScene.CreateWithInitializationData(GraphBuilder, &Scene->GPUScene, ReferenceView.ViewMatrices, MoveTemp(ReferenceView.RayTracingSceneInitData));
 
+	// Transition internal resources before building
+	RayTracingScene.Transition(GraphBuilder, ERayTracingSceneState::Writable);
+
 	const uint32 BLASScratchSize = Scene->GetRayTracingDynamicGeometryCollection()->ComputeScratchBufferSize();
 	if (BLASScratchSize > 0)
 	{
@@ -2014,14 +2017,14 @@ void FDeferredShadingSceneRenderer::WaitForRayTracingScene(FRDGBuilder& GraphBui
 			RayTracingDynamicGeometryUpdateEndTransition = nullptr;
 		}
 
-		FRHIRayTracingScene* RayTracingScene = ReferenceView.GetRayTracingSceneChecked();
-		RHICmdList.Transition(FRHITransitionInfo(RayTracingScene, ERHIAccess::BVHWrite, ERHIAccess::BVHRead));
-
 		if (ReferenceView.LumenHardwareRayTracingHitDataBuffer)
 		{
 			RHICmdList.Transition(FRHITransitionInfo(ReferenceView.LumenHardwareRayTracingHitDataBuffer, ERHIAccess::None, ERHIAccess::SRVMask));
-		}			
+		}
 	});
+
+    // Transition to readable state, synchronizing with previous build operation
+	Scene->RayTracingScene.Transition(GraphBuilder, ERayTracingSceneState::Readable);
 }
 
 #endif // RHI_RAYTRACING
diff --git a/Engine/Source/Runtime/Renderer/Private/RayTracing/RayTracingScene.cpp b/Engine/Source/Runtime/Renderer/Private/RayTracing/RayTracingScene.cpp
--- a/Engine/Source/Runtime/Renderer/Private/RayTracing/RayTracingScene.cpp
+++ b/Engine/Source/Runtime/Renderer/Private/RayTracing/RayTracingScene.cpp
@@ -84,6 +84,7 @@ void FRayTracingScene::CreateWithInitializationData(FRDGBuilder& GraphBuilder, c
 	{
 		FRHIResourceCreateInfo CreateInfo(TEXT("FRayTracingScene::SceneBuffer"));
 		RayTracingSceneBuffer = RHICreateBuffer(uint32(SizeInfo.ResultSize), EBufferUsageFlags::AccelerationStructure, 0, ERHIAccess::BVHWrite, CreateInfo);
+		State = ERayTracingSceneState::Writable;
 	}
 
 	LayerSRVs.SetNum(NumLayers);
@@ -340,6 +341,30 @@ void FRayTracingScene::ResetAndReleaseResources()
 	UsedCoarseMeshStreamingHandles.Empty();
 	RayTracingSceneBuffer = nullptr;
 	RayTracingSceneRHI = nullptr;
+	State = ERayTracingSceneState::Writable;
+}
+
+void FRayTracingScene::Transition(FRDGBuilder& GraphBuilder, ERayTracingSceneState InState)
+{
+	if (State == InState)
+	{
+		return;
+	}
+
+	GraphBuilder.AddPass(RDG_EVENT_NAME("RayTracingTransition"), ERDGPassFlags::None, 
+		[this, InState](FRHIComputeCommandList& RHICmdList)
+	{
+		if (InState == ERayTracingSceneState::Writable)
+		{
+			RHICmdList.Transition(FRHITransitionInfo(RayTracingSceneBuffer, ERHIAccess::BVHRead | ERHIAccess::SRVMask, ERHIAccess::BVHWrite));
+		}
+		else
+		{
+			RHICmdList.Transition(FRHITransitionInfo(RayTracingSceneBuffer, ERHIAccess::BVHWrite, ERHIAccess::BVHRead | ERHIAccess::SRVMask));
+		}
+	});
+
+	State = InState;
 }
 
 #endif // RHI_RAYTRACING
diff --git a/Engine/Source/Runtime/Renderer/Private/RayTracing/RayTracingScene.h b/Engine/Source/Runtime/Renderer/Private/RayTracing/RayTracingScene.h
--- a/Engine/Source/Runtime/Renderer/Private/RayTracing/RayTracingScene.h
+++ b/Engine/Source/Runtime/Renderer/Private/RayTracing/RayTracingScene.h
@@ -30,6 +30,12 @@ enum class ERayTracingSceneLayer : uint8
 	NUM
 };
 
+enum class ERayTracingSceneState
+{
+	Writable, // Scene is being built and can't be bound as SRV
+	Readable, // Scene can be used in ray tracing commands (can be bound as SRV)
+};
+
 /**
 * Persistent representation of the scene for ray tracing.
 * Manages top level acceleration structure instances, memory and build process.
@@ -58,6 +64,8 @@ public:
 	// Similar to Reset(), but also releases any persistent CPU and GPU memory allocations.
 	void ResetAndReleaseResources();
 
+	void Transition(FRDGBuilder& GraphBuilder, ERayTracingSceneState InState);
+
 	// Allocates temporary memory that will be valid until the next Reset().
 	// Can be used to store temporary instance transforms, user data, etc.
 	template <typename T>
@@ -143,6 +151,7 @@ private:
 
 	mutable FGraphEventRef FillInstanceUploadBufferTask;
 
+	ERayTracingSceneState State = ERayTracingSceneState::Writable;
 };
 
 #endif // RHI_RAYTRACING
diff --git a/Engine/Source/Runtime/VulkanRHI/Private/VulkanBarriers.cpp b/Engine/Source/Runtime/VulkanRHI/Private/VulkanBarriers.cpp
--- a/Engine/Source/Runtime/VulkanRHI/Private/VulkanBarriers.cpp
+++ b/Engine/Source/Runtime/VulkanRHI/Private/VulkanBarriers.cpp
@@ -169,6 +169,14 @@ static void GetVkStageAndAccessFlags(ERHIAccess RHIAccess, FRHITransitionInfo::E
 		RHIAccess = ERHIAccess::RTV;
 	}
 
+#if VULKAN_RHI_RAYTRACING
+	// BVHRead state may be combined with SRV, but we always treat this as just BVHRead by clearing the SRV mask
+	if (EnumHasAnyFlags(RHIAccess, ERHIAccess::BVHRead))
+	{
+		RHIAccess &= ~ERHIAccess::SRVMask;
+	}
+#endif
+
 	Layout = VK_IMAGE_LAYOUT_UNDEFINED;
 
 	// The layout to use if SRV access is requested. In case of depth/stencil buffers, we don't need to worry about different states for the separate aspects, since that's handled explicitly elsewhere,
