- Fix some Preview Platforms by updating the UpdatePreviewPlatforms

diff --git a/Engine/Source/Runtime/RHI/Private/RHI.cpp b/Engine/Source/Runtime/RHI/Private/RHI.cpp
--- a/Engine/Source/Runtime/RHI/Private/RHI.cpp
+++ b/Engine/Source/Runtime/RHI/Private/RHI.cpp
@@ -2649,12 +2649,14 @@ void FGenericDataDrivenShaderPlatformInfo::UpdatePreviewPlatforms()
 				Infos[ShaderPlatform].bSupportsGPUScene &= Infos[EditorSPForPreviewMaxFeatureLevel].bSupportsGPUScene;
 				Infos[ShaderPlatform].bIsPC = true;
 				Infos[ShaderPlatform].bIsConsole = false;
-				Infos[ShaderPlatform].bSupportsSceneDataCompressedTransforms &= Infos[EditorSPForPreviewMaxFeatureLevel].bSupportsSceneDataCompressedTransforms;
+				Infos[ShaderPlatform].bSupportsSceneDataCompressedTransforms = Infos[EditorSPForPreviewMaxFeatureLevel].bSupportsSceneDataCompressedTransforms;
 				Infos[ShaderPlatform].bSupportsNanite &= Infos[EditorSPForPreviewMaxFeatureLevel].bSupportsNanite;
+				Infos[ShaderPlatform].bSupportsLumenGI &= Infos[EditorSPForPreviewMaxFeatureLevel].bSupportsLumenGI;
 				Infos[ShaderPlatform].bSupportsUInt64ImageAtomics &= Infos[EditorSPForPreviewMaxFeatureLevel].bSupportsUInt64ImageAtomics;
 				Infos[ShaderPlatform].bSupportsGen5TemporalAA &= Infos[EditorSPForPreviewMaxFeatureLevel].bSupportsGen5TemporalAA;
 				Infos[ShaderPlatform].bSupportsMobileMultiView &= Infos[EditorSPForPreviewMaxFeatureLevel].bSupportsMobileMultiView;
 				Infos[ShaderPlatform].bSupportsInstancedStereo &= Infos[EditorSPForPreviewMaxFeatureLevel].bSupportsInstancedStereo;
+				Infos[ShaderPlatform].bSupportsManualVertexFetch &= Infos[EditorSPForPreviewMaxFeatureLevel].bSupportsManualVertexFetch;
 				Infos[ShaderPlatform].bSupportsRenderTargetWriteMask = false;
 				Infos[ShaderPlatform].bSupportsIntrinsicWaveOnce = false;
 				Infos[ShaderPlatform].bSupportsDOFHybridScattering = false;
