When launching a game that isn't cooked with any shader formats for the selected RHI, show an error dialog before fatally quitting. Since there are no formats available, we can't proceed.

diff --git a/Engine/Source/Runtime/RHI/Private/Windows/WindowsDynamicRHI.cpp b/Engine/Source/Runtime/RHI/Private/Windows/WindowsDynamicRHI.cpp
--- a/Engine/Source/Runtime/RHI/Private/Windows/WindowsDynamicRHI.cpp
+++ b/Engine/Source/Runtime/RHI/Private/Windows/WindowsDynamicRHI.cpp
@@ -591,7 +591,13 @@ static ERHIFeatureLevel::Type ChooseFeatureLevel(EWindowsRHI ChosenRHI, const TO
 
 			if (FPlatformProperties::RequiresCookedData())
 			{
-				const TCHAR* RHIName = ModuleNameFromWindowsRHI(ForcedRHI.GetValue());
+				const TCHAR* RHIName = ModuleNameFromWindowsRHI(ChosenRHI);
+				FMessageDialog::Open(EAppMsgType::Ok, FText::Format(NSLOCTEXT("WindowsDynamicRHI", "MissingTargetedShaderFormats", "Unable to launch with RHI '{0}' since the project is not configured to support it."), FText::AsCultureInvariant(RHIName)));
+				UE_LOG(LogRHI, Fatal, TEXT("Unable to launch with RHI '%s' since the project is not configured to support it."), RHIName);
+			}
+			else
+			{
+				const TCHAR* RHIName = ModuleNameFromWindowsRHI(ChosenRHI);
 				const FString FeatureLevelName = LexToString(FeatureLevel.GetValue());
 				UE_LOG(LogRHI, Warning, TEXT("User requested RHI '%s' but that is not supported by this project's data. Defaulting to Feature Level '%s'."), RHIName, *FeatureLevelName);
 			}
