Remove check(!Bypass()) in FRHICommandListImmediate::QueueAsyncCommandListSubmit().
 - Even when "-forcerhibypass -norhithread" command line args are used, the engine may still need to create parallel command lists for async mesh/texture streaming. This means some command lists will always be recorded (act as though bypass is disabled).
 - The check is unnecessary, since any command lists submitted via QueueAsyncCommandListSubmit will simply be replayed on the bottom-of-pipe thread (RHI thread, or render thread when in -norhithread mode).

#rnx

diff --git a/Engine/Source/Runtime/RHI/Private/RHICommandList.cpp b/Engine/Source/Runtime/RHI/Private/RHICommandList.cpp
--- a/Engine/Source/Runtime/RHI/Private/RHICommandList.cpp
+++ b/Engine/Source/Runtime/RHI/Private/RHICommandList.cpp
@@ -437,7 +437,6 @@ void FRHICommandListBase::Execute(TRHIPipelineArray<IRHIComputeContext*>& InOutC
 void FRHICommandListImmediate::QueueAsyncCommandListSubmit(TArrayView<FQueuedCommandList> CommandLists, ETranslatePriority ParallelTranslatePriority, int32 MinDrawsPerTranslate)
 {
 	check(IsInRenderingThread());
-	checkf(!Bypass(), TEXT("Parallel execution cannot be used in RHI bypass mode."));
 
 	if (CommandLists.Num() == 0)
 		return;
