Cleaning up deprecated items in RHI for 5.2

diff --git a/Engine/Source/Runtime/RHI/Private/DynamicRHI.cpp b/Engine/Source/Runtime/RHI/Private/DynamicRHI.cpp
--- a/Engine/Source/Runtime/RHI/Private/DynamicRHI.cpp
+++ b/Engine/Source/Runtime/RHI/Private/DynamicRHI.cpp
@@ -523,16 +523,6 @@ void FDynamicRHI::RHITransferBufferUnderlyingResource(FRHIBuffer* DestBuffer, FR
 	UE_LOG(LogRHI, Fatal, TEXT("RHITransferBufferUnderlyingResource isn't implemented for the current RHI"));
 }
 
-FUnorderedAccessViewRHIRef FDynamicRHI::RHICreateUnorderedAccessView(FRHITexture* Texture, uint32 MipLevel)
-{
-	return RHICreateUnorderedAccessView(Texture, MipLevel, 0, 0);
-}
-
-FUnorderedAccessViewRHIRef FDynamicRHI::RHICreateUnorderedAccessView(FRHITexture* Texture, uint32 MipLevel, uint8 Format)
-{
-	return RHICreateUnorderedAccessView(Texture, MipLevel, Format, 0, 0);
-}
-
 FUnorderedAccessViewRHIRef FDynamicRHI::RHICreateUnorderedAccessView(FRHITexture* Texture, uint32 MipLevel, uint8 Format, uint16 FirstArraySlice, uint16 NumArraySlices)
 {
 	UE_LOG(LogRHI, Fatal, TEXT("RHICreateUnorderedAccessView with Format parameter isn't implemented for the current RHI"));
diff --git a/Engine/Source/Runtime/RHI/Private/RHI.cpp b/Engine/Source/Runtime/RHI/Private/RHI.cpp
--- a/Engine/Source/Runtime/RHI/Private/RHI.cpp
+++ b/Engine/Source/Runtime/RHI/Private/RHI.cpp
@@ -1447,7 +1447,6 @@ bool GRHISupportsRHIThread = false;
 bool GRHISupportsRHIOnTaskThread = false;
 bool GRHISupportsParallelRHIExecute = false;
 bool GSupportsParallelOcclusionQueries = false;
-bool GSupportsTransientResourceAliasing = false;
 bool GRHIRequiresRenderTargetForPixelShaderUAVs = false;
 bool GRHISupportsUAVFormatAliasing = false;
 bool GRHISupportsDirectGPUMemoryLock = false;
diff --git a/Engine/Source/Runtime/RHI/Public/DynamicRHI.h b/Engine/Source/Runtime/RHI/Public/DynamicRHI.h
--- a/Engine/Source/Runtime/RHI/Public/DynamicRHI.h
+++ b/Engine/Source/Runtime/RHI/Public/DynamicRHI.h
@@ -408,12 +408,6 @@ public:
 	// FlushType: Thread safe, but varies depending on the RHI
 	virtual FUniformBufferRHIRef RHICreateUniformBuffer(const void* Contents, const FRHIUniformBufferLayout* Layout, EUniformBufferUsage Usage, EUniformBufferValidation Validation) = 0;
 
-	UE_DEPRECATED(5.0, "Use Layout pointers instead")
-	virtual FUniformBufferRHIRef RHICreateUniformBuffer(const void* Contents, const FRHIUniformBufferLayout& Layout, EUniformBufferUsage Usage, EUniformBufferValidation Validation)
-	{
-		return RHICreateUniformBuffer(Contents, &Layout, Usage, Validation);
-	}
-
 	virtual void RHIUpdateUniformBuffer(FRHICommandListBase& RHICmdList, FRHIUniformBuffer* UniformBufferRHI, const void* Contents) = 0;
 
 	/** Copies the contents of one buffer to another buffer. They must have identical sizes. */
@@ -447,20 +441,10 @@ public:
 	// FlushType: Wait RHI Thread
 	virtual FUnorderedAccessViewRHIRef RHICreateUnorderedAccessView(FRHIBuffer* Buffer, bool bUseUAVCounter, bool bAppendBuffer) = 0;
 
-	/** Creates an unordered access view of the given texture. */
-	// FlushType: Wait RHI Thread
-	UE_DEPRECATED(5.0, "RHICreateUnorderedAccessView now requires FirstArraySlice and NumArraySlices parameters.")
-	virtual FUnorderedAccessViewRHIRef RHICreateUnorderedAccessView(FRHITexture* Texture, uint32 MipLevel);
-
 	/** Creates an unordered access view of the given texture. */
 	// FlushType: Wait RHI Thread
 	virtual FUnorderedAccessViewRHIRef RHICreateUnorderedAccessView(FRHITexture* Texture, uint32 MipLevel, uint16 FirstArraySlice, uint16 NumArraySlices) = 0;
 
-	/** Creates an unordered access view of the given texture. */
-	// FlushType: Wait RHI Thread
-	UE_DEPRECATED(5.0, "RHICreateUnorderedAccessView now requires FirstArraySlice and NumArraySlices parameters.")
-	virtual FUnorderedAccessViewRHIRef RHICreateUnorderedAccessView(FRHITexture* Texture, uint32 MipLevel, uint8 Format);
-
 	/** Creates an unordered access view of the given texture. */
 	// FlushType: Wait RHI Thread
 	virtual FUnorderedAccessViewRHIRef RHICreateUnorderedAccessView(FRHITexture* Texture, uint32 MipLevel, uint8 Format, uint16 FirstArraySlice, uint16 NumArraySlices);
@@ -565,21 +549,6 @@ public:
 	// FlushType: Flush RHI Thread
 	virtual void RHICopySharedMips(FRHITexture2D* DestTexture2D, FRHITexture2D* SrcTexture2D) = 0;
 
-	/**
-	* @param Ref may be 0
-	*/
-	// FlushType: Thread safe
-	UE_DEPRECATED(5.0, "RHIGetResourceInfo is no longer implemented in favor of FRHIResource::GetResourceInfo.")
-	virtual void RHIGetResourceInfo(FRHITexture* Ref, FRHIResourceInfo& OutInfo)
-	{
-#if RHI_ENABLE_RESOURCE_INFO
-		if (Ref)
-		{
-			Ref->GetResourceInfo(OutInfo);
-		}
-#endif
-	}
-
 	/**
 	* Creates a shader resource view for a texture
 	*/
@@ -1139,35 +1108,11 @@ public:
 	virtual void* RHILockTextureCubeFace_RenderThread(class FRHICommandListImmediate& RHICmdList, FRHITextureCube* Texture, uint32 FaceIndex, uint32 ArrayIndex, uint32 MipIndex, EResourceLockMode LockMode, uint32& DestStride, bool bLockWithinMiptail);
 	virtual void RHIUnlockTextureCubeFace_RenderThread(class FRHICommandListImmediate& RHICmdList, FRHITextureCube* Texture, uint32 FaceIndex, uint32 ArrayIndex, uint32 MipIndex, bool bLockWithinMiptail);
 
-	UE_DEPRECATED(5.0, "AcquireTransientResource_RenderThread API is deprecated; use IRHITransientResourceAllocator instead.")
-	virtual void RHIAcquireTransientResource_RenderThread(FRHITexture*) {}
-
-	UE_DEPRECATED(5.0, "DiscardTransientResource_RenderThread API is deprecated; use IRHITransientResourceAllocator instead.")
-	virtual void RHIDiscardTransientResource_RenderThread(FRHITexture*) {}
-
-	UE_DEPRECATED(5.0, "AcquireTransientResource_RenderThread API is deprecated; use IRHITransientResourceAllocator instead.")
-	virtual void RHIAcquireTransientResource_RenderThread(FRHIBuffer*)  {}
-
-	UE_DEPRECATED(5.0, "DiscardTransientResource_RenderThread API is deprecated; use IRHITransientResourceAllocator instead.")
-	virtual void RHIDiscardTransientResource_RenderThread(FRHIBuffer*)  {}
-
 	virtual void RHIMapStagingSurface_RenderThread(class FRHICommandListImmediate& RHICmdList, FRHITexture* Texture, uint32 GPUIndex, FRHIGPUFence* Fence, void*& OutData, int32& OutWidth, int32& OutHeight);
 	virtual void RHIUnmapStagingSurface_RenderThread(class FRHICommandListImmediate& RHICmdList, FRHITexture* Texture, uint32 GPUIndex);
 	virtual void RHIReadSurfaceFloatData_RenderThread(class FRHICommandListImmediate& RHICmdList, FRHITexture* Texture, FIntRect Rect, TArray<FFloat16Color>& OutData, ECubeFace CubeFace, int32 ArrayIndex, int32 MipIndex);
 	virtual void RHIReadSurfaceFloatData_RenderThread(class FRHICommandListImmediate& RHICmdList, FRHITexture* Texture, FIntRect Rect, TArray<FFloat16Color>& OutData, FReadSurfaceDataFlags Flags);
 
-	UE_DEPRECATED(5.0, "Use the version of this function that accepts GPUIndex -- can pass INDEX_NONE to pull from the active GPUMask (the original behavior), but this will fail if GPUMask contains more than one GPU")
-	inline void RHIMapStagingSurface_RenderThread(class FRHICommandListImmediate& RHICmdList, FRHITexture* Texture, FRHIGPUFence* Fence, void*& OutData, int32& OutWidth, int32& OutHeight)
-	{
-		RHIMapStagingSurface_RenderThread(RHICmdList, Texture, INDEX_NONE, Fence, OutData, OutWidth, OutHeight);
-	}
-
-	UE_DEPRECATED(5.0, "Use the version of this function that accepts GPUIndex -- can pass INDEX_NONE to pull from the active GPUMask (the original behavior), but this will fail if GPUMask contains more than one GPU")
-	inline void RHIUnmapStagingSurface_RenderThread(class FRHICommandListImmediate& RHICmdList, FRHITexture* Texture)
-	{
-		RHIUnmapStagingSurface_RenderThread(RHICmdList, Texture, INDEX_NONE);
-	}
-
 	// Buffer Lock/Unlock
 	virtual void* LockBuffer_BottomOfPipe(class FRHICommandListBase& RHICmdList, FRHIBuffer* Buffer, uint32 Offset, uint32 SizeRHI, EResourceLockMode LockMode)
 	{
@@ -1197,9 +1142,6 @@ public:
 	virtual void RHICalibrateTimers() {}
 	virtual void RHIPollRenderQueryResults() {}
 
-	UE_DEPRECATED(5.0, "The global version of RHIIsTypedUAVLoadSupported should be used")
-	virtual bool RHIIsTypedUAVLoadSupported(EPixelFormat PixelFormat) { return ::RHIIsTypedUAVLoadSupported(PixelFormat); }
-
 	virtual uint16 RHIGetPlatformTextureMaxSampleCount() { return 8; };
 
 	virtual bool RHIRequiresComputeGenerateMips() const { return false; };
@@ -1349,12 +1291,6 @@ FORCEINLINE FUniformBufferRHIRef RHICreateUniformBuffer(const void* Contents, co
 	return GDynamicRHI->RHICreateUniformBuffer(Contents, Layout, Usage, Validation);
 }
 
-UE_DEPRECATED(5.0, "Use Layout pointers instead")
-FORCEINLINE FUniformBufferRHIRef RHICreateUniformBuffer(const void* Contents, const FRHIUniformBufferLayout& Layout, EUniformBufferUsage Usage, EUniformBufferValidation Validation = EUniformBufferValidation::ValidateResources)
-{
-	return GDynamicRHI->RHICreateUniformBuffer(Contents, &Layout, Usage, Validation);
-}
-
 FORCEINLINE FDynamicRHI::FRHICalcTextureSizeResult RHICalcTexturePlatformSize(FRHITextureDesc const& Desc, uint32 FirstMipIndex = 0)
 {
 	if ( ! Desc.IsValid() )
@@ -1498,14 +1434,6 @@ FORCEINLINE void RHIGetTextureMemoryStats(FTextureMemoryStats& OutStats)
 	GDynamicRHI->RHIGetTextureMemoryStats(OutStats);
 }
 
-UE_DEPRECATED(5.0, "RHIGetResourceInfo is no longer implemented in favor of FRHIResource::GetResourceInfo.")
-FORCEINLINE void RHIGetResourceInfo(FRHITexture* Ref, FRHIResourceInfo& OutInfo)
-{
-PRAGMA_DISABLE_DEPRECATION_WARNINGS
-	GDynamicRHI->RHIGetResourceInfo(Ref, OutInfo);
-PRAGMA_ENABLE_DEPRECATION_WARNINGS
-}
-
 FORCEINLINE uint32 RHIComputeMemorySize(FRHITexture* TextureRHI)
 {
 	return GDynamicRHI->RHIComputeMemorySize(TextureRHI);
@@ -1651,12 +1579,6 @@ FORCEINLINE const FRHITransition* RHICreateTransition(const FRHITransitionCreate
 	return Transition;
 }
 
-UE_DEPRECATED(5.0, "Use the FRHITransitionCreateInfo version of RHICreateTransition instead.")
-FORCEINLINE const FRHITransition* RHICreateTransition(ERHIPipeline SrcPipelines, ERHIPipeline DstPipelines, ERHITransitionCreateFlags CreateFlags, TArrayView<const FRHITransitionInfo> Infos)
-{
-	return RHICreateTransition(FRHITransitionCreateInfo(SrcPipelines, DstPipelines, CreateFlags, Infos));
-}
-
 FORCEINLINE void RHIReleaseTransition(FRHITransition* Transition)
 {
 	GDynamicRHI->RHIReleaseTransition(Transition);
diff --git a/Engine/Source/Runtime/RHI/Public/RHI.h b/Engine/Source/Runtime/RHI/Public/RHI.h
--- a/Engine/Source/Runtime/RHI/Public/RHI.h
+++ b/Engine/Source/Runtime/RHI/Public/RHI.h
@@ -467,10 +467,6 @@ extern RHI_API bool GSupportsEfficientAsyncCompute;
 /** True if the RHI supports getting the result of occlusion queries when on a thread other than the render thread */
 extern RHI_API bool GSupportsParallelOcclusionQueries;
 
-/** true if the RHI supports aliasing of transient resources */
-UE_DEPRECATED(5.0, "GSupportsTransientResourceAliasing has been deprecated. Transient resource aliasing is handled through IRHITransientResourceAllocator.")
-extern RHI_API bool GSupportsTransientResourceAliasing;
-
 /** true if the RHI requires a valid RT bound during UAV scatter operation inside the pixel shader */
 extern RHI_API bool GRHIRequiresRenderTargetForPixelShaderUAVs;
 
diff --git a/Engine/Source/Runtime/RHI/Public/RHICommandList.h b/Engine/Source/Runtime/RHI/Public/RHICommandList.h
--- a/Engine/Source/Runtime/RHI/Public/RHICommandList.h
+++ b/Engine/Source/Runtime/RHI/Public/RHICommandList.h
@@ -2428,12 +2428,6 @@ public:
 
 	inline FRHIComputeShader* GetBoundComputeShader() const { return PersistentState.BoundComputeShaderRHI; }
 
-	UE_DEPRECATED(5.0, "Please rename to SetStaticUniformBuffers")
-	FORCEINLINE_DEBUGGABLE void SetGlobalUniformBuffers(const FUniformBufferStaticBindings& UniformBuffers)
-	{
-		SetStaticUniformBuffers(UniformBuffers);
-	}
-
 	FORCEINLINE_DEBUGGABLE void SetStaticUniformBuffers(const FUniformBufferStaticBindings& UniformBuffers)
 	{
 		if (Bypass())
@@ -2477,14 +2471,6 @@ public:
 		return Result;
 	}
 
-	UE_DEPRECATED(5.0, "Use Layout pointers instead")
-	FORCEINLINE_DEBUGGABLE FLocalUniformBuffer BuildLocalUniformBuffer(const void* Contents, uint32 ContentsSize, const FRHIUniformBufferLayout& Layout)
-	{
-	PRAGMA_DISABLE_DEPRECATION_WARNINGS
-		return BuildLocalUniformBuffer(Contents, ContentsSize, &Layout);
-	PRAGMA_ENABLE_DEPRECATION_WARNINGS
-	}
-
 	template <typename TRHIShader>
 	UE_DEPRECATED(5.1, "Local uniform buffers are now deprecated. Use RHICreateUniformBuffer instead.")
 	FORCEINLINE_DEBUGGABLE void SetLocalShaderUniformBuffer(TRHIShader* Shader, uint32 BaseIndex, const FLocalUniformBuffer& UniformBuffer)
@@ -3467,12 +3453,6 @@ public:
 		ALLOC_COMMAND(FRHICommandSetGraphicsPipelineState)(GraphicsPipelineState, StencilRef, bApplyAdditionalState);
 	}
 
-	UE_DEPRECATED(5.0, "SetGraphicsPipelineState now requires a StencilRef argument")
-	FORCEINLINE_DEBUGGABLE void SetGraphicsPipelineState(class FGraphicsPipelineState* GraphicsPipelineState, const FBoundShaderStateInput& ShaderInput, bool bApplyAdditionalState)
-	{
-		SetGraphicsPipelineState(GraphicsPipelineState, ShaderInput, 0, bApplyAdditionalState);
-	}
-
 #if PLATFORM_USE_FALLBACK_PSO
 	FORCEINLINE_DEBUGGABLE void SetGraphicsPipelineState(const FGraphicsPipelineStateInitializer& PsoInit, uint32 StencilRef, bool bApplyAdditionalState)
 	{
@@ -4267,12 +4247,6 @@ public:
 		return RHICreateUniformBuffer(Contents, Layout, Usage);
 	}
 
-	UE_DEPRECATED(5.0, "Use Layout pointers instead")
-	FORCEINLINE FUniformBufferRHIRef CreateUniformBuffer(const void* Contents, const FRHIUniformBufferLayout& Layout, EUniformBufferUsage Usage)
-	{
-		return RHICreateUniformBuffer(Contents, &Layout, Usage);
-	}
-
 	UE_DEPRECATED(5.1, "Use CreateBuffer and LockBuffer separately")
 	FORCEINLINE FBufferRHIRef CreateAndLockIndexBuffer(uint32 Stride, uint32 Size, EBufferUsageFlags InUsage, ERHIAccess InResourceState, FRHIResourceCreateInfo& CreateInfo, void*& OutDataBuffer)
 	{
@@ -4291,18 +4265,6 @@ public:
 	PRAGMA_ENABLE_DEPRECATION_WARNINGS
 	}
 	
-	UE_DEPRECATED(5.0, "Buffer locks have been unified. Use LockBuffer() instead.")
-	FORCEINLINE void* LockIndexBuffer(FRHIBuffer* IndexBuffer, uint32 Offset, uint32 SizeRHI, EResourceLockMode LockMode)
-	{
-		return GDynamicRHI->RHILockBuffer(*this, IndexBuffer, Offset, SizeRHI, LockMode);
-	}
-	
-	UE_DEPRECATED(5.0, "Buffer locks have been unified. Use UnlockBuffer() instead.")
-	FORCEINLINE void UnlockIndexBuffer(FRHIBuffer* IndexBuffer)
-	{
-		GDynamicRHI->RHIUnlockBuffer(*this, IndexBuffer);
-	}
-	
 	FORCEINLINE void* LockStagingBuffer(FRHIStagingBuffer* StagingBuffer, FRHIGPUFence* Fence, uint32 Offset, uint32 SizeRHI)
 	{
 		return GDynamicRHI->LockStagingBuffer_RenderThread(*this, StagingBuffer, Fence, Offset, SizeRHI);
@@ -4331,18 +4293,6 @@ public:
 	PRAGMA_ENABLE_DEPRECATION_WARNINGS
 	}
 
-	UE_DEPRECATED(5.0, "Buffer locks have been unified. Use LockBuffer() instead.")
-	FORCEINLINE void* LockVertexBuffer(FRHIBuffer* VertexBuffer, uint32 Offset, uint32 SizeRHI, EResourceLockMode LockMode)
-	{
-		return GDynamicRHI->RHILockBuffer(*this, VertexBuffer, Offset, SizeRHI, LockMode);
-	}
-	
-	UE_DEPRECATED(5.0, "Buffer locks have been unified. Use UnlockBuffer() instead.")
-	FORCEINLINE void UnlockVertexBuffer(FRHIBuffer* VertexBuffer)
-	{
-		GDynamicRHI->RHIUnlockBuffer(*this, VertexBuffer);
-	}
-	
 	FORCEINLINE void CopyBuffer(FRHIBuffer* SourceBuffer, FRHIBuffer* DestBuffer)
 	{
 		QUICK_SCOPE_CYCLE_COUNTER(STAT_RHIMETHOD_CopyBuffer_Flush);
@@ -4350,24 +4300,6 @@ public:
 		GDynamicRHI->RHICopyBuffer(SourceBuffer, DestBuffer);
 	}
 
-	UE_DEPRECATED(5.0, "CopyVertexBuffer() has been replaced with a general CopyBuffer() function.")
-	FORCEINLINE void CopyVertexBuffer(FRHIBuffer* SourceBuffer, FRHIBuffer* DestBuffer)
-	{
-		CopyBuffer(SourceBuffer, DestBuffer);
-	}
-
-	UE_DEPRECATED(5.0, "Buffer locks have been unified. Use LockBuffer() instead.")
-	FORCEINLINE void* LockStructuredBuffer(FRHIBuffer* StructuredBuffer, uint32 Offset, uint32 SizeRHI, EResourceLockMode LockMode)
-	{
-		return GDynamicRHI->RHILockBuffer(*this, StructuredBuffer, Offset, SizeRHI, LockMode);
-	}
-
-	UE_DEPRECATED(5.0, "Buffer locks have been unified. Use UnlockBuffer() instead.")
-	FORCEINLINE void UnlockStructuredBuffer(FRHIBuffer* StructuredBuffer)
-	{
-		GDynamicRHI->RHIUnlockBuffer(*this, StructuredBuffer);
-	}
-
 	// LockBufferMGPU / UnlockBufferMGPU may ONLY be called for buffers with the EBufferUsageFlags::MultiGPUAllocate flag set!
 	// And buffers with that flag set may not call the regular (single GPU) LockBuffer / UnlockBuffer.  The single GPU version
 	// of LockBuffer uses driver mirroring to propagate the updated buffer to other GPUs, while the MGPU / MultiGPUAllocate
@@ -4487,14 +4419,6 @@ public:
 		});
 	}
 
-	UE_DEPRECATED(5.0, "RHIGetResourceInfo is no longer implemented in favor of FRHIResource::GetResourceInfo.")
-	FORCEINLINE void GetResourceInfo(FRHITexture* Ref, FRHIResourceInfo& OutInfo)
-	{
-PRAGMA_DISABLE_DEPRECATION_WARNINGS
-		return RHIGetResourceInfo(Ref, OutInfo);
-PRAGMA_ENABLE_DEPRECATION_WARNINGS
-	}
-	
 	FORCEINLINE FShaderResourceViewRHIRef CreateShaderResourceView(FRHITexture* Texture, const FRHITextureSRVCreateInfo& CreateInfo)
 	{
 		LLM_SCOPE_BYNAME(TEXT("RHIMisc/CreateShaderResourceView"));
@@ -4706,18 +4630,6 @@ PRAGMA_ENABLE_DEPRECATION_WARNINGS
 		GDynamicRHI->RHIRead3DSurfaceFloatData(Texture,Rect,ZMinMax,OutData,Flags);
 	}
 
-	UE_DEPRECATED(5.0, "AcquireTransientResource_RenderThread API is deprecated; use IRHITransientResourceAllocator instead.")
-	FORCEINLINE void AcquireTransientResource_RenderThread(FRHITexture* Texture) {}
-
-	UE_DEPRECATED(5.0, "DiscardTransientResource_RenderThread API is deprecated; use IRHITransientResourceAllocator instead.")
-	FORCEINLINE void DiscardTransientResource_RenderThread(FRHITexture* Texture) {}
-
-	UE_DEPRECATED(5.0, "AcquireTransientResource_RenderThread API is deprecated; use IRHITransientResourceAllocator instead.")
-	FORCEINLINE void AcquireTransientResource_RenderThread(FRHIBuffer* Buffer) {}
-
-	UE_DEPRECATED(5.0, "DiscardTransientResource_RenderThread API is deprecated; use IRHITransientResourceAllocator instead.")
-	FORCEINLINE void DiscardTransientResource_RenderThread(FRHIBuffer* Buffer) {}
-
 	FORCEINLINE bool GetRenderQueryResult(FRHIRenderQuery* RenderQuery, uint64& OutResult, bool bWait, uint32 GPUIndex = INDEX_NONE)
 	{
 		return RHIGetRenderQueryResult(RenderQuery, OutResult, bWait, GPUIndex);
@@ -5036,12 +4948,6 @@ struct RHI_API FScopedUniformBufferStaticBindings
 
 #define SCOPED_UNIFORM_BUFFER_STATIC_BINDINGS(RHICmdList, UniformBuffers) FScopedUniformBufferStaticBindings PREPROCESSOR_JOIN(UniformBuffers, __LINE__){ RHICmdList, UniformBuffers }
 
-UE_DEPRECATED(5.0, "Please rename to FScopedUniformBufferStaticBindings, or use the SCOPED_UNIFORM_BUFFER_STATIC_BINDINGS instead.")
-typedef FScopedUniformBufferStaticBindings FScopedUniformBufferGlobalBindings;
-
-// UE_DEPRECATED(5.0)
-#define SCOPED_UNIFORM_BUFFER_GLOBAL_BINDINGS(RHICmdList, UniformBuffers) FScopedUniformBufferGlobalBindings PREPROCESSOR_JOIN(UniformBuffers, __LINE__){ RHICmdList, UniformBuffers }
-
 // Helper to enable the use of graphics RHI command lists from within platform RHI implementations.
 // Recorded commands are dispatched when the command list is destructed. Intended for use on the stack / in a scope block.
 class RHI_API FRHICommandList_RecursiveHazardous : public FRHICommandList
@@ -5269,15 +5175,6 @@ FORCEINLINE FStagingBufferRHIRef RHICreateStagingBuffer()
 	return FRHICommandListExecutor::GetImmediateCommandList().CreateStagingBuffer();
 }
 
-UE_DEPRECATED(5.0, "Use RHICreateBuffer() and RHILockBuffer() instead.")
-FORCEINLINE FBufferRHIRef RHICreateAndLockIndexBuffer(uint32 Stride, uint32 Size, EBufferUsageFlags InUsage, FRHIResourceCreateInfo& CreateInfo, void*& OutDataBuffer)
-{
-	check(IsInRenderingThread());
-PRAGMA_DISABLE_DEPRECATION_WARNINGS
-	return FRHICommandListExecutor::GetImmediateCommandList().CreateAndLockIndexBuffer(Stride, Size, InUsage, CreateInfo, OutDataBuffer);
-PRAGMA_ENABLE_DEPRECATION_WARNINGS
-}
-
 FORCEINLINE FBufferRHIRef RHICreateBuffer(uint32 Size, EBufferUsageFlags Usage, uint32 Stride, ERHIAccess ResourceState, FRHIResourceCreateInfo& CreateInfo)
 {
 	check(IsInRenderingThread());
@@ -5316,34 +5213,11 @@ PRAGMA_DISABLE_DEPRECATION_WARNINGS
 PRAGMA_ENABLE_DEPRECATION_WARNINGS
 }
 
-UE_DEPRECATED(5.0, "Buffer locks have been unified. Use RHILockBuffer() instead.")
-FORCEINLINE void* RHILockIndexBuffer(FRHIBuffer* IndexBuffer, uint32 Offset, uint32 Size, EResourceLockMode LockMode)
-{
-	check(IsInRenderingThread());
-	return FRHICommandListExecutor::GetImmediateCommandList().LockBuffer(IndexBuffer, Offset, Size, LockMode);
-}
-
-UE_DEPRECATED(5.0, "Buffer locks have been unified. Use RHIUnlockBuffer() instead.")
-FORCEINLINE void RHIUnlockIndexBuffer(FRHIBuffer* IndexBuffer)
-{
-	check(IsInRenderingThread());
-	FRHICommandListExecutor::GetImmediateCommandList().UnlockBuffer(IndexBuffer);
-}
-
 FORCEINLINE void RHIUpdateUniformBuffer(FRHIUniformBuffer* UniformBufferRHI, const void* Contents)
 {
 	return FRHICommandListExecutor::GetImmediateCommandList().UpdateUniformBuffer(UniformBufferRHI, Contents);
 }
 
-UE_DEPRECATED(5.0, "Use RHICreateBuffer() and RHILockBuffer() instead.")
-FORCEINLINE FBufferRHIRef RHICreateAndLockVertexBuffer(uint32 Size, EBufferUsageFlags InUsage, FRHIResourceCreateInfo& CreateInfo, void*& OutDataBuffer)
-{
-	check(IsInRenderingThread());
-PRAGMA_DISABLE_DEPRECATION_WARNINGS
-	return FRHICommandListExecutor::GetImmediateCommandList().CreateAndLockVertexBuffer(Size, InUsage, CreateInfo, OutDataBuffer);
-PRAGMA_ENABLE_DEPRECATION_WARNINGS
-}
-
 FORCEINLINE FBufferRHIRef RHICreateVertexBuffer(uint32 Size, EBufferUsageFlags Usage, ERHIAccess ResourceState, FRHIResourceCreateInfo& CreateInfo)
 {
 	check(IsInRenderingThread());
@@ -5376,20 +5250,6 @@ PRAGMA_DISABLE_DEPRECATION_WARNINGS
 PRAGMA_ENABLE_DEPRECATION_WARNINGS
 }
 
-UE_DEPRECATED(5.0, "Buffer locks have been unified. Use RHILockBuffer() instead.")
-FORCEINLINE void* RHILockVertexBuffer(FRHIBuffer* VertexBuffer, uint32 Offset, uint32 SizeRHI, EResourceLockMode LockMode)
-{
-	check(IsInRenderingThread());
-	return FRHICommandListExecutor::GetImmediateCommandList().LockBuffer(VertexBuffer, Offset, SizeRHI, LockMode);
-}
-
-UE_DEPRECATED(5.0, "Buffer locks have been unified. Use RHIUnlockBuffer() instead.")
-FORCEINLINE void RHIUnlockVertexBuffer(FRHIBuffer* VertexBuffer)
-{
-	check(IsInRenderingThread());
-	FRHICommandListExecutor::GetImmediateCommandList().UnlockBuffer(VertexBuffer);
-}
-
 FORCEINLINE FBufferRHIRef RHICreateStructuredBuffer(uint32 Stride, uint32 Size, EBufferUsageFlags Usage, ERHIAccess ResourceState, FRHIResourceCreateInfo& CreateInfo)
 {
 	check(IsInRenderingThread());
@@ -5404,20 +5264,6 @@ FORCEINLINE FBufferRHIRef RHICreateStructuredBuffer(uint32 Stride, uint32 Size,
 	return RHICreateStructuredBuffer(Stride, Size, Usage, ResourceState, CreateInfo);
 }
 
-UE_DEPRECATED(5.0, "Buffer locks have been unified. Use RHILockBuffer() instead.")
-FORCEINLINE void* RHILockStructuredBuffer(FRHIBuffer* StructuredBuffer, uint32 Offset, uint32 SizeRHI, EResourceLockMode LockMode)
-{
-	check(IsInRenderingThread());
-	return FRHICommandListExecutor::GetImmediateCommandList().LockBuffer(StructuredBuffer, Offset, SizeRHI, LockMode);
-}
-
-UE_DEPRECATED(5.0, "Buffer locks have been unified. Use RHIUnlockBuffer() instead.")
-FORCEINLINE void RHIUnlockStructuredBuffer(FRHIBuffer* StructuredBuffer)
-{
-	check(IsInRenderingThread());
-	FRHICommandListExecutor::GetImmediateCommandList().UnlockBuffer(StructuredBuffer);
-}
-
 FORCEINLINE void* RHILockBuffer(FRHIBuffer* Buffer, uint32 Offset, uint32 SizeRHI, EResourceLockMode LockMode)
 {
 	check(IsInRenderingThread());
@@ -5476,12 +5322,6 @@ FORCEINLINE FTextureReferenceRHIRef RHICreateTextureReference()
 	return new FRHITextureReference();
 }
 
-UE_DEPRECATED(5.0, "The LastRenderTime parameter will be removed in the future")
-FORCEINLINE FTextureReferenceRHIRef RHICreateTextureReference(FLastRenderTimeContainer* LastRenderTime)
-{
-	return RHICreateTextureReference();
-}
-
 FORCEINLINE void RHIUpdateTextureReference(FRHITextureReference* TextureRef, FRHITexture* NewTexture)
 {
 	check(IsInRenderingThread());
@@ -5769,18 +5609,6 @@ FORCEINLINE void RHIUnlockTextureCubeFace(FRHITextureCube* Texture, uint32 FaceI
 	 FRHICommandListExecutor::GetImmediateCommandList().UnlockTextureCubeFace(Texture, FaceIndex, ArrayIndex, MipIndex, bLockWithinMiptail);
 }
 
-UE_DEPRECATED(5.0, "RHIAcquireTransientResource API is deprecated; use IRHITransientResourceAllocator instead.")
-FORCEINLINE void RHIAcquireTransientResource(FRHITexture*) {}
-
-UE_DEPRECATED(5.0, "RHIDiscardTransientResource API is deprecated; use IRHITransientResourceAllocator instead.")
-FORCEINLINE void RHIDiscardTransientResource(FRHITexture*) {}
-
-UE_DEPRECATED(5.0, "RHIAcquireTransientResource API is deprecated; use IRHITransientResourceAllocator instead.")
-FORCEINLINE void RHIAcquireTransientResource(FRHIBuffer*)  {}
-
-UE_DEPRECATED(5.0, "RHIDiscardTransientResource API is deprecated; use IRHITransientResourceAllocator instead.")
-FORCEINLINE void RHIDiscardTransientResource(FRHIBuffer*)  {}
-
 FORCEINLINE void RHIAcquireThreadOwnership()
 {
 	return FRHICommandListExecutor::GetImmediateCommandList().AcquireThreadOwnership();
diff --git a/Engine/Source/Runtime/RHI/Public/RHIContext.h b/Engine/Source/Runtime/RHI/Public/RHIContext.h
--- a/Engine/Source/Runtime/RHI/Public/RHIContext.h
+++ b/Engine/Source/Runtime/RHI/Public/RHIContext.h
@@ -117,9 +117,6 @@ private:
 	int32 SlotCount = 0;
 };
 
-UE_DEPRECATED(5.0, "Please rename to FUniformBufferStaticBindings")
-typedef FUniformBufferStaticBindings FUniformBufferGlobalBindings;
-
 struct FTransferResourceFenceData
 {
 	TStaticArray<void*, MAX_NUM_GPUS> SyncPoints;
@@ -315,9 +312,6 @@ public:
 		/** empty default implementation. */
 	}
 
-	UE_DEPRECATED(5.0, "Please rename to RHISetStaticUniformBuffers.")
-	virtual void RHISetGlobalUniformBuffers(const FUniformBufferStaticBindings& InUniformBuffers) {}
-
 	virtual void RHIPushEvent(const TCHAR* Name, FColor Color) = 0;
 
 	virtual void RHIPopEvent() = 0;
@@ -661,11 +655,6 @@ public:
 
 	virtual void RHISetGraphicsPipelineState(FRHIGraphicsPipelineState* GraphicsState, uint32 StencilRef, bool bApplyAdditionalState) = 0;
 
-	UE_DEPRECATED(5.0, "SetGraphicsPipelineState now requires a StencilRef argument")
-	void RHISetGraphicsPipelineState(FRHIGraphicsPipelineState* GraphicsState, bool bApplyAdditionalState)
-	{
-		RHISetGraphicsPipelineState(GraphicsState, 0, bApplyAdditionalState);
-	}
 #if PLATFORM_USE_FALLBACK_PSO
 	virtual void RHISetGraphicsPipelineState(const FGraphicsPipelineStateInitializer& PsoInit, uint32 StencilRef, bool bApplyAdditionalState) = 0;
 #endif
diff --git a/Engine/Source/Runtime/RHI/Public/RHIDefinitions.h b/Engine/Source/Runtime/RHI/Public/RHIDefinitions.h
--- a/Engine/Source/Runtime/RHI/Public/RHIDefinitions.h
+++ b/Engine/Source/Runtime/RHI/Public/RHIDefinitions.h
@@ -86,14 +86,11 @@ enum EShaderPlatform
 	SP_OPENGL_PCES3_1				= 15,
 	SP_METAL_SM5					= 16,
 	SP_VULKAN_PCES3_1				= 17,
-	SP_METAL_SM5_NOTESS_REMOVED		UE_DEPRECATED(5.0, "ShaderPlatform is removed; please don't use.") = 18,
 	SP_VULKAN_SM5					= 20,
 	SP_VULKAN_ES3_1_ANDROID			= 21,
 	SP_METAL_MACES3_1 				= 22,
 	SP_OPENGL_ES3_1_ANDROID			= 24,
 	SP_METAL_MRT_MAC				= 27,
-	SP_VULKAN_SM5_LUMIN_REMOVED		UE_DEPRECATED(5.0, "ShaderPlatform is removed; please don't use.") = 28,
-	SP_VULKAN_ES3_1_LUMIN_REMOVED	UE_DEPRECATED(5.0, "ShaderPlatform is removed; please don't use.") = 29,
 	SP_METAL_TVOS					= 30,
 	SP_METAL_MRT_TVOS				= 31,
 	/**********************************************************************************/
@@ -1717,9 +1714,6 @@ enum class EBufferUsageFlags : uint32
 	/** Buffer should go in fast vram (hint only). Requires BUF_Transient */
 	FastVRAM                = 1 << 10,
 
-	/** Buffer should be allocated from transient memory. */
-	Transient UE_DEPRECATED(5.0, "EBufferUsageFlags::Transient flag is no longer used.") = None,
-
 	/** Create a buffer that can be shared with an external RHI or process. */
 	Shared                  = 1 << 12,
 
@@ -1931,8 +1925,6 @@ enum class ETextureCreateFlags : uint64
     AFRManual                         = 1ull << 30,
     // Workaround for 128^3 volume textures getting bloated 4x due to tiling mode on some platforms.
     ReduceMemoryWithTilingMode        = 1ull << 31,
-    /** Texture should be allocated from transient memory. */
-    Transient UE_DEPRECATED(5.0, "ETextureCreateFlags::Transient flag is no longer used.") = None,
     /** Texture needs to support atomic operations */
     AtomicCompatible                  = 1ull << 33,
 	/** Texture should be allocated for external access. Vulkan only */
diff --git a/Engine/Source/Runtime/RHI/Public/RHIResources.h b/Engine/Source/Runtime/RHI/Public/RHIResources.h
--- a/Engine/Source/Runtime/RHI/Public/RHIResources.h
+++ b/Engine/Source/Runtime/RHI/Public/RHIResources.h
@@ -50,16 +50,6 @@ enum class EClearBinding;
 class RHI_API FRHIResource
 {
 public:
-	UE_DEPRECATED(5.0, "FRHIResource(bool) is deprecated, please use FRHIResource(ERHIResourceType)")
-	FRHIResource(bool InbDoNotDeferDelete=false)
-		: ResourceType(RRT_None)
-		, bCommitted(true)
-#if RHI_ENABLE_RESOURCE_INFO
-		, bBeingTracked(false)
-#endif
-	{
-	}
-
 	FRHIResource(ERHIResourceType InResourceType)
 		: ResourceType(InResourceType)
 		, bCommitted(true)
@@ -1195,10 +1185,6 @@ private:
 	EBufferUsageFlags Usage{};
 };
 
-UE_DEPRECATED(5.0, "FRHIIndexBuffer is deprecated, please use FRHIBuffer.")      typedef class FRHIBuffer FRHIIndexBuffer;
-UE_DEPRECATED(5.0, "FRHIVertexBuffer is deprecated, please use FRHIBuffer.")     typedef class FRHIBuffer FRHIVertexBuffer;
-UE_DEPRECATED(5.0, "FRHIStructuredBuffer is deprecated, please use FRHIBuffer.") typedef class FRHIBuffer FRHIStructuredBuffer;
-
 //
 // Textures
 //
@@ -1857,18 +1843,6 @@ public:
 		return LastRenderTime.GetLastRenderTime();
 	}
 
-	/** Returns the last render time container, or NULL if none were specified at creation. */
-	UE_DEPRECATED(5.0, "GetLastRenderTimeContainer is deprecated and will be removed in the future")
-	FLastRenderTimeContainer* GetLastRenderTimeContainer()
-	{
-		return nullptr;
-	}
-
-	UE_DEPRECATED(5.0, "SetDefaultLastRenderTimeContainer is deprecated and will be removed in the future")
-	FORCEINLINE_DEBUGGABLE void SetDefaultLastRenderTimeContainer()
-	{
-	}
-
 	void SetName(const FName& InName)
 	{
 		Name = InName;
@@ -1955,11 +1929,6 @@ public:
 		ReferencedTexture = DefaultTexture;
 	}
 
-	UE_DEPRECATED(5.0, "The InLastRenderTime parameter will be removed in the future")
-	explicit FRHITextureReference(FLastRenderTimeContainer* InLastRenderTime)
-		: FRHITextureReference()
-	{}
-
 	virtual class FRHITextureReference* GetTextureReference() override
 	{
 		return this;
@@ -2311,9 +2280,6 @@ typedef TRefCountPtr<FRHIBoundShaderState> FBoundShaderStateRHIRef;
 typedef TRefCountPtr<const FRHIUniformBufferLayout> FUniformBufferLayoutRHIRef;
 typedef TRefCountPtr<FRHIUniformBuffer> FUniformBufferRHIRef;
 typedef TRefCountPtr<FRHIBuffer> FBufferRHIRef;
-UE_DEPRECATED(5.0, "FIndexBufferRHIRef is deprecated, please use FBufferRHIRef.")      typedef FBufferRHIRef FIndexBufferRHIRef;
-UE_DEPRECATED(5.0, "FVertexBufferRHIRef is deprecated, please use FBufferRHIRef.")     typedef FBufferRHIRef FVertexBufferRHIRef;
-UE_DEPRECATED(5.0, "FStructuredBufferRHIRef is deprecated, please use FBufferRHIRef.") typedef FBufferRHIRef FStructuredBufferRHIRef;
 typedef TRefCountPtr<FRHITexture> FTextureRHIRef;
 /*UE_DEPRECATED(5.1, "FTexture2DRHIRef is deprecated, please use FTextureRHIRef.")      */ typedef FTextureRHIRef FTexture2DRHIRef;
 /*UE_DEPRECATED(5.1, "FTexture2DArrayRHIRef is deprecated, please use FTextureRHIRef.") */ typedef FTextureRHIRef FTexture2DArrayRHIRef;
diff --git a/Engine/Source/Runtime/RHI/Public/RHIUtilities.h b/Engine/Source/Runtime/RHI/Public/RHIUtilities.h
--- a/Engine/Source/Runtime/RHI/Public/RHIUtilities.h
+++ b/Engine/Source/Runtime/RHI/Public/RHIUtilities.h
@@ -237,12 +237,6 @@ struct FTextureRWBuffer2D : public FTextureRWBuffer
 	{
 		Initialize2D(InDebugName, BytesPerElement, SizeX, SizeY, Format, Flags);
 	}
-
-	UE_DEPRECATED(5.0, "AcquireTransientResource is deprecated. Transient resources are allocated through IRHITransientResourceAllocator instead.")
-	void AcquireTransientResource() {}
-
-	UE_DEPRECATED(5.0, "DiscardTransientResource is deprecated. Transient resources are allocated through IRHITransientResourceAllocator instead.")
-	void DiscardTransientResource() {}
 };
 
 struct UE_DEPRECATED(5.1, "FTextureRWBuffer should be used instead of FTextureRWBuffer3D.") FTextureRWBuffer3D;
@@ -253,12 +247,6 @@ struct FTextureRWBuffer3D : public FTextureRWBuffer
 	{
 		Initialize3D(InDebugName, BytesPerElement, SizeX, SizeY, SizeZ, Format, Flags);
 	}
-
-	UE_DEPRECATED(5.0, "AcquireTransientResource is deprecated. Transient resources are allocated through IRHITransientResourceAllocator instead.")
-	void AcquireTransientResource() {}
-
-	UE_DEPRECATED(5.0, "DiscardTransientResource is deprecated. Transient resources are allocated through IRHITransientResourceAllocator instead.")
-	void DiscardTransientResource() {}
 };
 
 /** Encapsulates a GPU read/write buffer with its UAV and SRV. */
@@ -334,12 +322,6 @@ struct FRWBuffer
 		Initialize(InDebugName, BytesPerElement, NumElements, Format, ERHIAccess::UAVCompute, AdditionalUsage, InResourceArray);
 	}
 
-	UE_DEPRECATED(5.0, "AcquireTransientResource is deprecated. Transient resources are allocated through IRHITransientResourceAllocator instead.")
-	void AcquireTransientResource() {}
-
-	UE_DEPRECATED(5.0, "DiscardTransientResource is deprecated. Transient resources are allocated through IRHITransientResourceAllocator instead.")
-	void DiscardTransientResource() {}
-
 	void Release()
 	{
 		NumBytes = 0;
@@ -380,12 +362,6 @@ struct FTextureReadBuffer2D
 		SRV = RHICreateShaderResourceView(Buffer, 0);
 	}
 
-	UE_DEPRECATED(5.0, "AcquireTransientResource is deprecated. Transient resources are allocated through IRHITransientResourceAllocator instead.")
-	void AcquireTransientResource() {}
-
-	UE_DEPRECATED(5.0, "DiscardTransientResource is deprecated. Transient resources are allocated through IRHITransientResourceAllocator instead.")
-	void DiscardTransientResource() {}
-
 	void Release()
 	{
 		NumBytes = 0;
@@ -455,12 +431,6 @@ struct FRWBufferStructured
 		UAV.SafeRelease();
 		SRV.SafeRelease();
 	}
-
-	UE_DEPRECATED(5.0, "AcquireTransientResource is deprecated. Transient resources are allocated through IRHITransientResourceAllocator instead.")
-	void AcquireTransientResource() {}
-
-	UE_DEPRECATED(5.0, "DiscardTransientResource is deprecated. Transient resources are allocated through IRHITransientResourceAllocator instead.")
-	void DiscardTransientResource() {}
 };
 
 struct FByteAddressBuffer
