Add option to dump rhi resource info to csv file directly

diff --git a/Engine/Source/Runtime/RHI/Private/RHI.cpp b/Engine/Source/Runtime/RHI/Private/RHI.cpp
--- a/Engine/Source/Runtime/RHI/Private/RHI.cpp
+++ b/Engine/Source/Runtime/RHI/Private/RHI.cpp
@@ -745,7 +745,7 @@
 	return RRT_None;
 }
 
-static const TCHAR* StringFromRHIResourceType(ERHIResourceType ResourceType)
+const TCHAR* StringFromRHIResourceType(ERHIResourceType ResourceType)
 {
 	for (const auto& TypeName : GRHIResourceTypeNames)
 	{
@@ -829,62 +829,13 @@
 	BufferedOutput.RedirectTo(OutputDevice);
 }));
 
-static FAutoConsoleCommandWithWorldArgsAndOutputDevice GDumpRHIResourceMemoryCmd(
-	TEXT("rhi.DumpResourceMemory"),
-	TEXT("Dumps RHI resource memory stats to the log\n")
-	TEXT("Usage: rhi.DumpResourceMemory [<Number To Show>] [all] [summary] [Name=<Filter Text>] [Type=<RHI Resource Type>] [Transient=<no, yes, or all> [csv]"),
-	FConsoleCommandWithWorldArgsAndOutputDeviceDelegate::CreateStatic([](const TArray<FString>& Args, UWorld*, FOutputDevice& OutputDevice)
-{
-	FString NameFilter;
-	ERHIResourceType TypeFilter = RRT_None;
-	EBooleanFilter TransientFilter = EBooleanFilter::No;
-	int32 NumberOfResourcesToShow = 50;
-	bool bUseCSVOutput = false;
-	bool bSummaryOutput = false;
-
-	for (const FString& Argument : Args)
+void RHIDumpResourceMemory(const FString& NameFilter, ERHIResourceType TypeFilter, EBooleanFilter TransientFilter, int32 NumberOfResourcesToShow, bool bUseCSVOutput, bool bSummaryOutput, bool bOutputToCSVFile, FBufferedOutputDevice& BufferedOutput)
+{	
+	FArchive* CSVFile{ nullptr };
+	if (bOutputToCSVFile)
 	{
-		if (Argument.Equals(TEXT("all"), ESearchCase::IgnoreCase))
-		{
-			NumberOfResourcesToShow = -1;
-		}
-		else if (Argument.Equals(TEXT("-csv"), ESearchCase::IgnoreCase))
-		{
-			bUseCSVOutput = true;
-		}
-		else if (Argument.StartsWith(TEXT("Name="), ESearchCase::IgnoreCase))
-		{
-			NameFilter = Argument.RightChop(5);
-		}
-		else if (Argument.StartsWith(TEXT("Type="), ESearchCase::IgnoreCase))
-		{
-			TypeFilter = RHIResourceTypeFromString(Argument.RightChop(5));
-		}
-		else if (Argument.StartsWith(TEXT("Transient="), ESearchCase::IgnoreCase))
-		{
-			TransientFilter = ParseBooleanFilter(Argument.RightChop(10));
-		}
-		else if (FCString::IsNumeric(*Argument))
-		{
-			LexFromString(NumberOfResourcesToShow, *Argument);
-		}
-		else if (Argument.Equals(TEXT("summary"), ESearchCase::IgnoreCase))
-		{
-			// Respects name, type and transient filters but only reports total sizes.
-			// Does not report a list of individual resources.
-			bSummaryOutput = true;
-			NumberOfResourcesToShow = -1;
-		}
-		else
-		{
-			NameFilter = Argument;
-		}
+		const FString Filename = FString::Printf(TEXT("%srhiDumpResourceMemory-%s.csv"), *FPaths::ProfilingDir(), *FDateTime::Now().ToString());
+		CSVFile = IFileManager::Get().CreateFileWriter(*Filename, FILEWRITE_AllowRead);
 	}
 
-	// Summary output and csv are mutually exclusive.  So if both are on, summary takes precedence
-	if (bSummaryOutput && bUseCSVOutput)
-	{
-		bUseCSVOutput = false;
-	}
-
 	TCHAR ResourceNameBuffer[FName::StringBufferSize];
@@ -948,6 +899,8 @@
 		{
 			if (Resource && Resource->GetResourceInfo(ResourceInfo))
 			{
+				ResourceInfo.bValid = Resource->IsValid();
+
 				if (ShouldIncludeResource(Resource, ResourceInfo))
 				{
 					Resources.Emplace(FLocalResourceEntry{ Resource, ResourceInfo });
@@ -977,12 +930,16 @@
 		return EntryA.ResourceInfo.VRamAllocation.AllocationSize > EntryB.ResourceInfo.VRamAllocation.AllocationSize;
 	});
 
-	FBufferedOutputDevice BufferedOutput;
 	FName CategoryName(TEXT("RHIResources"));
 
-	if (bUseCSVOutput)
+	if (bOutputToCSVFile)
 	{
-		BufferedOutput.CategorizedLogf(CategoryName, ELogVerbosity::Log, TEXT("Name,Type,Size,Transient,Streaming,RenderTarget,UAV,\"Raytracing Acceleration Structure\""));
+		const TCHAR* Header = TEXT("Name,Type,Size,MarkedForDelete,Transient,Streaming,RenderTarget,UAV,\"Raytracing Acceleration Structure\"\n");
+		CSVFile->Serialize(TCHAR_TO_ANSI(Header), FPlatformString::Strlen(Header));
 	}
+	else if (bUseCSVOutput)
+	{
+		BufferedOutput.CategorizedLogf(CategoryName, ELogVerbosity::Log, TEXT("Name,Type,Size,MarkedForDelete,Transient,Streaming,RenderTarget,UAV,\"Raytracing Acceleration Structure\""));
+	}
 	else
 	{
@@ -1009,6 +966,7 @@
 			const TCHAR* ResourceType = StringFromRHIResourceType(ResourceInfo.Type);
 			const int64 SizeInBytes = ResourceInfo.VRamAllocation.AllocationSize;
 
+			bool bMarkedForDelete = !ResourceInfo.bValid;
 			bool bTransient = ResourceInfo.IsTransient;
 			bool bStreaming = false;
 			bool bRT = false;
@@ -1016,7 +974,8 @@
 			bool bUAV = false;
 			bool bRTAS = false;
 
-			bool bIsTexture = ResourceInfo.Type == RRT_Texture2D ||
+			bool bIsTexture = ResourceInfo.Type == RRT_Texture ||
+				ResourceInfo.Type == RRT_Texture2D ||
 				ResourceInfo.Type == RRT_Texture2DArray ||
 				ResourceInfo.Type == RRT_Texture3D ||
 				ResourceInfo.Type == RRT_TextureCube;
@@ -1037,12 +996,28 @@
 
 			if (bSummaryOutput == false)
 			{
-				if (bUseCSVOutput)
+				if (bOutputToCSVFile)
+				{					
+					const FString Row = FString::Printf(TEXT("%s,%s,%.9f,%s,%s,%s,%s,%s,%s\n"),
+						ResourceNameBuffer,
+						ResourceType,
+						SizeInBytes / double(1 << 20),
+						bMarkedForDelete ? TEXT("Yes") : TEXT("No"),
+						bTransient ? TEXT("Yes") : TEXT("No"),
+						bStreaming ? TEXT("Yes") : TEXT("No"),
+						(bRT || bDS) ? TEXT("Yes") : TEXT("No"),
+						bUAV ? TEXT("Yes") : TEXT("No"),
+						bRTAS ? TEXT("Yes") : TEXT("No"));
+
+					CSVFile->Serialize(TCHAR_TO_ANSI(*Row), Row.Len());
+				}
+				else if (bUseCSVOutput)
 				{
-					BufferedOutput.CategorizedLogf(CategoryName, ELogVerbosity::Log, TEXT("%s,%s,%.9f,%s,%s,%s,%s,%s"),
+					BufferedOutput.CategorizedLogf(CategoryName, ELogVerbosity::Log, TEXT("%s,%s,%.9f,%s,%s,%s,%s,%s,%s"),
 						ResourceNameBuffer,
 						ResourceType,
 						SizeInBytes / double(1 << 20),
+						bMarkedForDelete ? TEXT("Yes") : TEXT("No"),
 						bTransient ? TEXT("Yes") : TEXT("No"),
 						bStreaming ? TEXT("Yes") : TEXT("No"),
 						(bRT || bDS) ? TEXT("Yes") : TEXT("No"),
@@ -1097,7 +1072,12 @@
 		}
 	}
 
-	if (!bUseCSVOutput)
+	if (bOutputToCSVFile)
 	{
+		delete CSVFile;
+		CSVFile = nullptr;
+	}
+	else if (!bUseCSVOutput)
+	{
 		const double TotalNonTransientSizeF = TotalTrackedResourceSize / double(1 << 20);
 		const double TotalTransientSizeF = TotalTrackedTransientResourceSize / double(1 << 20);
@@ -1144,10 +1124,84 @@
 			}
 		}
 	}
+}
 
+static FAutoConsoleCommandWithWorldArgsAndOutputDevice GDumpRHIResourceMemoryCmd(
+	TEXT("rhi.DumpResourceMemory"),
+	TEXT("Dumps RHI resource memory stats to the log\n")
+	TEXT("Usage: rhi.DumpResourceMemory [<Number To Show>] [all] [summary] [Name=<Filter Text>] [Type=<RHI Resource Type>] [Transient=<no, yes, or all> [csv]"),
+	FConsoleCommandWithWorldArgsAndOutputDeviceDelegate::CreateStatic([](const TArray<FString>& Args, UWorld*, FOutputDevice& OutputDevice)
+{
+	FString NameFilter;
+	ERHIResourceType TypeFilter = RRT_None;
+	EBooleanFilter TransientFilter = EBooleanFilter::No;
+	int32 NumberOfResourcesToShow = 50;
+	bool bUseCSVOutput = false;
+	bool bSummaryOutput = false;
+	bool bOutputToCSVFile = false;
+
+	for (const FString& Argument : Args)
+	{
+		if (Argument.Equals(TEXT("all"), ESearchCase::IgnoreCase))
+		{
+			NumberOfResourcesToShow = -1;
+		}
+		else if (Argument.Equals(TEXT("-csv"), ESearchCase::IgnoreCase))
+		{
+			bUseCSVOutput = true;
+		}
+		else if (Argument.Equals(TEXT("-csvfile"), ESearchCase::IgnoreCase))
+		{
+			bOutputToCSVFile = true;
+		}
+		else if (Argument.StartsWith(TEXT("Name="), ESearchCase::IgnoreCase))
+		{
+			NameFilter = Argument.RightChop(5);
+		}
+		else if (Argument.StartsWith(TEXT("Type="), ESearchCase::IgnoreCase))
+		{
+			TypeFilter = RHIResourceTypeFromString(Argument.RightChop(5));
+		}
+		else if (Argument.StartsWith(TEXT("Transient="), ESearchCase::IgnoreCase))
+		{
+			TransientFilter = ParseBooleanFilter(Argument.RightChop(10));
+		}
+		else if (FCString::IsNumeric(*Argument))
+		{
+			LexFromString(NumberOfResourcesToShow, *Argument);
+		}
+		else if (Argument.Equals(TEXT("summary"), ESearchCase::IgnoreCase))
+		{
+			// Respects name, type and transient filters but only reports total sizes.
+			// Does not report a list of individual resources.
+			bSummaryOutput = true;
+			NumberOfResourcesToShow = -1;
+		}
+		else
+		{
+			NameFilter = Argument;
+		}
+	}
+
+	FBufferedOutputDevice BufferedOutput;
+	RHIDumpResourceMemory(NameFilter, TypeFilter, TransientFilter, NumberOfResourcesToShow, bUseCSVOutput, bSummaryOutput, bOutputToCSVFile, BufferedOutput);
 	BufferedOutput.RedirectTo(OutputDevice);
 }));
 
+void RHIDumpResourceMemoryToCSV()
+{
+	FString NameFilter;
+	ERHIResourceType TypeFilter = RRT_None;
+	EBooleanFilter TransientFilter = EBooleanFilter::No;
+	int32 NumberOfResourcesToShow = -1;
+	bool bUseCSVOutput = false;
+	bool bSummaryOutput = false;
+	bool bOutputToCSVFile = true;
+
+	FBufferedOutputDevice BufferedOutput;
+	RHIDumpResourceMemory(NameFilter, TypeFilter, TransientFilter, NumberOfResourcesToShow, bUseCSVOutput, bSummaryOutput, bOutputToCSVFile, BufferedOutput);
+}
+
 #endif // RHI_ENABLE_RESOURCE_INFO
 
 bool FRHIResource::Bypass()
diff --git a/Engine/Source/Runtime/RHI/Public/RHI.h b/Engine/Source/Runtime/RHI/Public/RHI.h
--- a/Engine/Source/Runtime/RHI/Public/RHI.h
+++ b/Engine/Source/Runtime/RHI/Public/RHI.h
@@ -142,6 +142,9 @@ RHI_API EPixelFormat RHIPreferredPixelFormatHint(EPixelFormat PreferredPixelForm
 // helper to check which resource type should be used for clear (UAV) replacement shaders.
 RHI_API int32 RHIGetPreferredClearUAVRectPSResourceType(const FStaticShaderPlatform Platform);
 
+// helper to force dump all RHI resource to CSV file
+RHI_API void RHIDumpResourceMemoryToCSV();
+
 inline bool RHISupportsInstancedStereo(const FStaticShaderPlatform Platform)
 {
 	return FDataDrivenShaderPlatformInfo::GetSupportsInstancedStereo(Platform);
@@ -957,6 +960,9 @@ extern RHI_API FString GetBufferUsageFlagsName(EBufferUsageFlags BufferUsage);
 /** Returns a string of friendly name bits for the texture create flags enum. */
 extern RHI_API FString GetTextureCreateFlagsName(ETextureCreateFlags TextureCreateFlags);
 
+/** Returns a string of friendly name bits for the texture create flags enum. */
+extern RHI_API const TCHAR* StringFromRHIResourceType(ERHIResourceType ResourceType);
+
 enum class ERHIPipeline : uint8
 {
 	Graphics = 1 << 0,
@@ -1640,6 +1646,7 @@ struct FRHIResourceInfo
 	ERHIResourceType Type{ RRT_None };
 	FVRamAllocation VRamAllocation;
 	bool IsTransient{ false };
+	bool bValid{ true };
 };
 
 enum class EClearBinding
