- Use background tasks for PSO precaching when PSOPrecacheThreadPool is not used - default (rename PSO thread pool)
- Make sure graph dependencies are not added when rhi command lists doesn't supported multithreaded rendering
- Fix leak of PSOCollectors

#rnx

diff --git a/Engine/Source/Runtime/Engine/Private/Materials/MaterialShader.cpp b/Engine/Source/Runtime/Engine/Private/Materials/MaterialShader.cpp
--- a/Engine/Source/Runtime/Engine/Private/Materials/MaterialShader.cpp
+++ b/Engine/Source/Runtime/Engine/Private/Materials/MaterialShader.cpp
@@ -2743,9 +2743,6 @@ FGraphEventArray FMaterialShaderMap::CollectPSOs(ERHIFeatureLevel::Type InFeatur
 	SceneTexturesConfig.Init(SceneTexturesConfigInitSettings);
 
 	const FMaterialShaderMapContent* LocalContent = GetContent();
-	const EShaderPlatform Platform = LocalContent->GetShaderPlatform();
-	const FMaterialShaderParameters MaterialParameters(Material);
-
 	const EShadingPath ShadingPath = FSceneInterface::GetShadingPath(InFeatureLevel);
 
 	TArray<FPSOPrecacheData> PSOInitializers;
@@ -2766,6 +2763,8 @@ FGraphEventArray FMaterialShaderMap::CollectPSOs(ERHIFeatureLevel::Type InFeatur
 						PSOCollector->CollectPSOInitializers(SceneTexturesConfig, *Material, VFType, PreCacheParams, PSOInitializers);
 					}
 				}
+			
+				delete PSOCollector;
 			}
 		}
 	}
diff --git a/Engine/Source/Runtime/RHI/Private/PipelineStateCache.cpp b/Engine/Source/Runtime/RHI/Private/PipelineStateCache.cpp
--- a/Engine/Source/Runtime/RHI/Private/PipelineStateCache.cpp
+++ b/Engine/Source/Runtime/RHI/Private/PipelineStateCache.cpp
@@ -101,7 +101,7 @@ static TAutoConsoleVariable<int32> CVarRTPSOCacheSize(
 #endif // RHI_RAYTRACING
 
 int32 GPSOPrecaching = 0;
-static TAutoConsoleVariable<int32> CVarPSOPrecaching(
+static FAutoConsoleVariableRef CVarPSOPrecaching(
 	TEXT("r.PSOPrecaching"),
 	GPSOPrecaching,
 	TEXT("0 to Disable PSOs precaching\n")
@@ -132,10 +132,10 @@ static FAutoConsoleVariableRef GPSOPrecompileThreadPoolSizeVar(
 	ECVF_RenderThreadSafe | ECVF_ReadOnly
 );
 
-class FPSOThreadPool
+class FPSOPrecacheThreadPool
 {
 public:
-	~FPSOThreadPool()
+	~FPSOPrecacheThreadPool()
 	{
 		ShutdownThreadPool();
 	}
@@ -147,6 +147,8 @@ public:
 			FScopeLock lock(&LockCS);
 			if (PSOPrecompileCompileThreadPool.load() == nullptr)
 			{
+				check(UsePool());
+
 				FQueuedThreadPool* PSOPrecompileCompileThreadPoolLocal = FQueuedThreadPool::Allocate();
 				PSOPrecompileCompileThreadPoolLocal->Create(GPSOPrecompileThreadPoolSize, 512 * 1024, EThreadPriority::TPri_BelowNormal, TEXT("PSOPrecompilePool"));
 				PSOPrecompileCompileThreadPool = PSOPrecompileCompileThreadPoolLocal;
@@ -166,17 +168,22 @@ public:
 		}
 	}
 
+	static bool UsePool()
+	{
+		return GPSOPrecompileThreadPoolSize > 0;
+	}
+
 private:
 	FCriticalSection LockCS;
 	std::atomic<FQueuedThreadPool*> PSOPrecompileCompileThreadPool {};
 };
 
-static FPSOThreadPool GPSOThreadPool;
+static FPSOPrecacheThreadPool GPSOPrecacheThreadPool;
 
 void PipelineStateCache::PreCompileComplete()
 {
 	// free up our threads when the precompile completes.
-	GPSOThreadPool.ShutdownThreadPool();
+	GPSOPrecacheThreadPool.ShutdownThreadPool();
 }
 
 extern RHI_API FRHIComputePipelineState* ExecuteSetComputePipelineState(FComputePipelineState* ComputePipelineState);
@@ -894,29 +901,49 @@ public:
 		return PrecachedPSOInitializers.Contains(Initializer);
 	}
 
-	void Add(const FGraphicsPipelineStateInitializer& Initializer, FGraphicsPipelineState* PipelineState, bool bAsync)
-	{		
-		// First add to set of precached PSOs initializers
+	FGraphicsPipelineState* TryAddNewState(const FGraphicsPipelineStateInitializer& Initializer, bool bDoAsyncCompile)
+	{
+		// Fast check first with read lock
+		if (Contains(Initializer))
+		{
+			return nullptr;
+		}
+
+		FGraphicsPipelineState* NewState = nullptr;
+
+		// Now try and add with write lock
 		{
 			FRWScopeLock WriteLock(PrecachePSOsRWLock, SLT_Write);
-			check(!PrecachedPSOInitializers.Contains(Initializer));
+
+			// check again
+			if (PrecachedPSOInitializers.Contains(Initializer))
+			{
+				return nullptr;
+			}
+
+			// create new graphics state
+			NewState = new FGraphicsPipelineState();
+
+			// Add to cache before starting async operation, because the async op will remove the task from active ops when done			
 			PrecachedPSOInitializers.Add(Initializer);
 		}
 
 		// Add to active set of precaching PSOs if it's created async
-		if (bAsync)
+		if (bDoAsyncCompile)
 		{
 			// Only need the hash of the PSO here
 			uint64 PrecachePSOHash = RHIComputePrecachePSOHash(Initializer);
 
 			FPrecacheTask PrecacheTask;
 			PrecacheTask.Initializer = Initializer;
-			PrecacheTask.PSO = PipelineState;
+			PrecacheTask.PSO = NewState;
 
 			FRWScopeLock WriteLock(ActivePrecachingTasksRWLock, SLT_Write);
 			check(!ActivePrecachingTasks.Contains(PrecachePSOHash));
 			ActivePrecachingTasks.Add(PrecachePSOHash, PrecacheTask);
 		}
+
+		return NewState;
 	}
 
 	void WaitTasksComplete()
@@ -1451,7 +1478,8 @@ public:
 
 		// On Mac the compilation is handled using external processes, so engine threads have very little work to do
 		// and it's better to leave more CPU time to these extrenal processes and other engine threads.
-		return PLATFORM_MAC ? ENamedThreads::AnyBackgroundThreadNormalTask : DesiredThread;
+		// Also use background threads for PSO precaching when the PSO thread pool is not used
+		return (PLATFORM_MAC || PSOPreCacheResult == EPSOPrecacheResult::Active) ? ENamedThreads::AnyBackgroundThreadNormalTask : DesiredThread;
 	}
 };
 
@@ -1936,15 +1964,15 @@ inline void ValidateGraphicsPipelineStateInitializer(const FGraphicsPipelineStat
 	check(Initializer.DepthStencilState && Initializer.BlendState && Initializer.RasterizerState);
 }
 
-static FGraphEventRef CreateGraphicsPipelineState(const FGraphicsPipelineStateInitializer& Initializer, EPSOPrecacheResult PSOPrecacheResult, bool bDoAsyncCompile, FGraphicsPipelineState* CachedState)
+static FGraphEventRef CreateGraphicsPipelineState(const FGraphicsPipelineStateInitializer& Initializer, EPSOPrecacheResult PSOPrecacheResult, bool bDoAsyncCompile, bool bPSOPrecache, FGraphicsPipelineState* CachedState)
 {
 	FGraphEventRef GraphEvent;
 
 	// create a compilation task, or just do it now...
 	if (bDoAsyncCompile)
 	{
-		// Use normal task graph for non-precompile jobs.
-		if (GPSOPrecompileThreadPoolSize <= 0)
+		// Use normal task graph for non-precompile jobs (or when thread pool is not enabled)
+		if (!bPSOPrecache || !FPSOPrecacheThreadPool::UsePool())
 		{
 			CachedState->CompletionEvent = TGraphTask<FCompilePipelineStateTask>::CreateTask().ConstructAndDispatchWhenReady(CachedState, Initializer, PSOPrecacheResult);
 			GraphEvent = CachedState->CompletionEvent;
@@ -1974,7 +2002,7 @@ static FGraphEventRef CreateGraphicsPipelineState(const FGraphicsPipelineStateIn
 #endif
 			}
 			);
-			CachedState->PrecompileTask->StartBackgroundTask(&GPSOThreadPool.Get(), EQueuedWorkPriority::Normal);
+			CachedState->PrecompileTask->StartBackgroundTask(&GPSOPrecacheThreadPool.Get(), EQueuedWorkPriority::Normal);
 		}
 	}
 	else
@@ -2034,11 +2062,10 @@ FGraphicsPipelineState* PipelineStateCache::GetAndOrCreateGraphicsPipelineState(
 	}
 #endif
 
-	bool DoAsyncCompile = IsAsyncCompilationAllowed(RHICmdList, Initializer.bFromPSOFileCache);
-
 	FGraphicsPipelineState* OutCachedState = nullptr;
 
 	bool bWasFound = GGraphicsPipelineCache.Find(Initializer, OutCachedState);
+	bool DoAsyncCompile = IsAsyncCompilationAllowed(RHICmdList, Initializer.bFromPSOFileCache);
 
 	if (bWasFound == false)
 	{
@@ -2053,12 +2080,14 @@ FGraphicsPipelineState* PipelineStateCache::GetAndOrCreateGraphicsPipelineState(
 			GraphicsPipelineCacheMisses++;
 		}
 
-		FGraphEventRef GraphEvent = CreateGraphicsPipelineState(Initializer, PSOPrecacheResult, DoAsyncCompile, OutCachedState);
+		bool bPSOPrecache = Initializer.bFromPSOFileCache;
+		FGraphEventRef GraphEvent = CreateGraphicsPipelineState(Initializer, PSOPrecacheResult, DoAsyncCompile, bPSOPrecache, OutCachedState);
 
-		// Use normal task graph for non-precompile jobs.
-		if (DoAsyncCompile && (!Initializer.bFromPSOFileCache || GPSOPrecompileThreadPoolSize <= 0))
+		// Add dispatch pre requisite for non precaching jobs only
+		// NOTE: do we need to add dispatch prerequisite for PSO precache when precache pool is disabled and it's using regular task graph?
+		if (GraphEvent.IsValid() && (!bPSOPrecache || !FPSOPrecacheThreadPool::UsePool()))
 		{
-			check(GraphEvent);
+			check(DoAsyncCompile);
 			RHICmdList.AddDispatchPrerequisite(GraphEvent);
 		}
 
@@ -2072,7 +2101,7 @@ FGraphicsPipelineState* PipelineStateCache::GetAndOrCreateGraphicsPipelineState(
 			{
 				// if this is an in-progress threadpool precompile task then it could be seconds away in the queue.
 				// Reissue this task so that it jumps the precompile queue.
-				OutCachedState->PrecompileTask->Reschedule(&GPSOThreadPool.Get(), EQueuedWorkPriority::Highest);
+				OutCachedState->PrecompileTask->Reschedule(&GPSOPrecacheThreadPool.Get(), EQueuedWorkPriority::Highest);
 #if PSO_TRACK_CACHE_STATS
 		UE_LOG(LogRHI, Log, TEXT("An incomplete precompile task was required for rendering!"));
 #endif
@@ -2130,8 +2159,8 @@ FGraphEventRef PipelineStateCache::PrecacheComputePipelineState(FRHIComputeShade
 	static bool bDoAsyncCompile = FApp::ShouldUseThreadingForPerformance();
 	if (bDoAsyncCompile)
 	{
-		// Use normal task graph for non-precompile jobs.
-		if (GPSOPrecompileThreadPoolSize <= 0)
+		// Thread pool disabled?
+		if (!FPSOPrecacheThreadPool::UsePool())
 		{
 			CachedState->CompletionEvent = TGraphTask<FCompilePipelineStateTask>::CreateTask().ConstructAndDispatchWhenReady(CachedState, FGraphicsPipelineStateInitializer(), EPSOPrecacheResult::Active);
 			GraphEvent = CachedState->CompletionEvent;
@@ -2161,7 +2190,7 @@ FGraphEventRef PipelineStateCache::PrecacheComputePipelineState(FRHIComputeShade
 #endif
 			}
 			);
-			CachedState->PrecompileTask->StartBackgroundTask(&GPSOThreadPool.Get(), EQueuedWorkPriority::Normal);
+			CachedState->PrecompileTask->StartBackgroundTask(&GPSOPrecacheThreadPool.Get(), EQueuedWorkPriority::Normal);
 		}
 	}
 	else
@@ -2175,34 +2204,30 @@ FGraphEventRef PipelineStateCache::PrecacheComputePipelineState(FRHIComputeShade
 
 FGraphEventRef PipelineStateCache::PrecacheGraphicsPipelineState(const FGraphicsPipelineStateInitializer& Initializer)
 {
-	FGraphEventRef GraphEvent = nullptr;
 	if (!IsPSOPrecachingEnabled())
 	{
-		return GraphEvent;
+		return nullptr;
 	}
 
 	LLM_SCOPE(ELLMTag::PSO);
-	ValidateGraphicsPipelineStateInitializer(Initializer);
-		
-	// Already in the cache?
-	if (GPrecacheGraphicsPipelineCache.Contains(Initializer))
-	{
-		return GraphEvent;
-	}
 
-	// create new graphics state
-	FGraphicsPipelineState* CachedState = new FGraphicsPipelineState();
-	
-	// Use async compilation
+	// Use async compilation if available
 	static bool bDoAsyncCompile = FApp::ShouldUseThreadingForPerformance();
 
-	// Add to cache before starting async operation, because the async op will remove the task from active ops when done
-	GPrecacheGraphicsPipelineCache.Add(Initializer, CachedState, bDoAsyncCompile);
+	// try and create new graphics state
+	FGraphicsPipelineState* NewGraphicsPipelineState = GPrecacheGraphicsPipelineCache.TryAddNewState(Initializer, bDoAsyncCompile);
+	if (NewGraphicsPipelineState == nullptr)
+	{
+		return nullptr;
+	}
+	
+	ValidateGraphicsPipelineStateInitializer(Initializer);
 
-	// Start the async 
-	GraphEvent = CreateGraphicsPipelineState(Initializer, EPSOPrecacheResult::Active, bDoAsyncCompile, CachedState);
+	// Mark as precache so it will try and use the background thread pool if available
+	bool bPSOPrecache = true;
 
-	return GraphEvent;
+	// Start the precache task
+	return CreateGraphicsPipelineState(Initializer, EPSOPrecacheResult::Active, bDoAsyncCompile, bPSOPrecache, NewGraphicsPipelineState);
 }
 
 EPSOPrecacheResult PipelineStateCache::CheckPipelineStateInCache(const FGraphicsPipelineStateInitializer& PipelineStateInitializer)
