Allowing SM5 to support bindless only in RayTracing shaders.

diff --git a/Engine/Config/VulkanPC/DataDrivenPlatformInfo.ini b/Engine/Config/VulkanPC/DataDrivenPlatformInfo.ini
--- a/Engine/Config/VulkanPC/DataDrivenPlatformInfo.ini
+++ b/Engine/Config/VulkanPC/DataDrivenPlatformInfo.ini
@@ -134,7 +134,7 @@ bSupportsOIT=false
 EnablesHLSL2021ByDefault=1
 FriendlyName=LOCTEXT("FriendlyShaderPlatformName_VulkanPC_SM5", "Vulkan PC Shader Model 5")
 
-bSupportsBindless=true
+BindlessSupport=AllShaderTypes
 
 
 [ShaderPlatform VULKAN_PCES3_1]
diff --git a/Engine/Config/Windows/DataDrivenPlatformInfo.ini b/Engine/Config/Windows/DataDrivenPlatformInfo.ini
--- a/Engine/Config/Windows/DataDrivenPlatformInfo.ini
+++ b/Engine/Config/Windows/DataDrivenPlatformInfo.ini
@@ -73,6 +73,7 @@ bSupportsVolumetricFog=true
 bSupportsIndexBufferUAVs=true
 bSupportsROV=true
 bSupportsOIT=true
+BindlessSupport=RayTracingOnly
 SupportsMultiViewport=RuntimeDependent
 FriendlyName=LOCTEXT("FriendlyShaderPlatformName_D3D_SM5", "D3D Shader Model 5")
 
@@ -126,7 +127,7 @@ bSupportsCapsuleShadows=true
 bSupportsPercentageCloserShadows=true
 bSupportsVolumetricFog=true
 bSupportsIndexBufferUAVs=true
-bSupportsBindless=true
+BindlessSupport=AllShaderTypes
 bSupportsROV=true
 bSupportsOIT=true
 bSupportsRealTypes=RuntimeDependent
diff --git a/Engine/Source/Runtime/D3D12RHI/Private/D3D12Adapter.cpp b/Engine/Source/Runtime/D3D12RHI/Private/D3D12Adapter.cpp
--- a/Engine/Source/Runtime/D3D12RHI/Private/D3D12Adapter.cpp
+++ b/Engine/Source/Runtime/D3D12RHI/Private/D3D12Adapter.cpp
@@ -940,7 +940,10 @@ void FD3D12Adapter::InitializeDevices()
 			//     ResourceDescriptorHeap/SamplerDescriptorHeap must be supported on devices that support both D3D12_RESOURCE_BINDING_TIER_3 and D3D_SHADER_MODEL_6_6
 			if (GetHighestShaderModel() >= D3D_SHADER_MODEL_6_6 && GetResourceBindingTier() >= D3D12_RESOURCE_BINDING_TIER_3)
 			{
+				GRHIBindlessSupport = GMaxRHIFeatureLevel == ERHIFeatureLevel::SM5 ? ERHIBindlessSupport::RayTracingOnly : ERHIBindlessSupport::AllShaderTypes;
+PRAGMA_DISABLE_DEPRECATION_WARNINGS
 				GRHISupportsBindless = true;
+PRAGMA_ENABLE_DEPRECATION_WARNINGS
 				UE_LOG(LogD3D12RHI, Log, TEXT("Bindless resources are supported"));
 			}
 
@@ -969,7 +972,7 @@ void FD3D12Adapter::InitializeDevices()
 					&& !FParse::Param(FCommandLine::Get(), TEXT("noraytracing")))
 				{
 					if (D3D12Caps5.RaytracingTier >= D3D12_RAYTRACING_TIER_1_1
-						&& GRHISupportsBindless
+						&& GRHIBindlessSupport != ERHIBindlessSupport::Unsupported
 						&& RootDevice7)
 					{
 						if (bRayTracingAllowedOnCurrentShaderPlatform)
@@ -987,7 +990,7 @@ void FD3D12Adapter::InitializeDevices()
  							UE_LOG(LogD3D12RHI, Log, TEXT("Ray tracing is disabled because SM6 shader platform is required (r.RayTracing.RequireSM6=1)."));
 						}
 					}
-					else if (!GRHISupportsBindless)
+					else if (GRHIBindlessSupport == ERHIBindlessSupport::Unsupported)
 					{
 						UE_LOG(LogD3D12RHI, Log, TEXT("Ray tracing is disabled because bindless resources are not supported (Shader Model 6.6 and Resource Binding Tier 3 are required)."));
 					}
diff --git a/Engine/Source/Runtime/D3D12RHI/Private/D3D12Shaders.cpp b/Engine/Source/Runtime/D3D12RHI/Private/D3D12Shaders.cpp
--- a/Engine/Source/Runtime/D3D12RHI/Private/D3D12Shaders.cpp
+++ b/Engine/Source/Runtime/D3D12RHI/Private/D3D12Shaders.cpp
@@ -68,9 +68,13 @@ static bool ValidateShaderIsUsable(FD3D12ShaderData* InShader, EShaderFrequency
 		return false;
 	}
 
-	if (EnumHasAnyFlags(InShader->Features, EShaderCodeFeatures::BindlessResources | EShaderCodeFeatures::BindlessSamplers) && !GRHISupportsBindless)
+	if (EnumHasAnyFlags(InShader->Features, EShaderCodeFeatures::BindlessResources | EShaderCodeFeatures::BindlessSamplers))
 	{
-		return false;
+		if (GRHIBindlessSupport == ERHIBindlessSupport::Unsupported ||
+			(GRHIBindlessSupport == ERHIBindlessSupport::RayTracingOnly && !IsRayTracingShaderFrequency(InFrequency)))
+		{
+			return false;
+		}
 	}
 #endif
 
diff --git a/Engine/Source/Runtime/Engine/Private/ShaderCompiler/ShaderCompiler.cpp b/Engine/Source/Runtime/Engine/Private/ShaderCompiler/ShaderCompiler.cpp
--- a/Engine/Source/Runtime/Engine/Private/ShaderCompiler/ShaderCompiler.cpp
+++ b/Engine/Source/Runtime/Engine/Private/ShaderCompiler/ShaderCompiler.cpp
@@ -6285,10 +6285,10 @@ void GlobalBeginCompileShader(
 	Input.Environment.SetDefine(TEXT("PLATFORM_SUPPORTS_DISTANCE_FIELDS"), DoesPlatformSupportDistanceFields(EShaderPlatform(Target.Platform)) ? 1 : 0);
 	Input.Environment.SetDefine(TEXT("PLATFORM_SUPPORTS_MESH_SHADERS_TIER0"), RHISupportsMeshShadersTier0(EShaderPlatform(Target.Platform)) ? 1 : 0);
 	Input.Environment.SetDefine(TEXT("PLATFORM_SUPPORTS_MESH_SHADERS_TIER1"), RHISupportsMeshShadersTier1(EShaderPlatform(Target.Platform)) ? 1 : 0);
-	Input.Environment.SetDefine(TEXT("PLATFORM_SUPPORTS_BINDLESS"), RHISupportsBindless(EShaderPlatform(Target.Platform)) ? 1 : 0);
 	Input.Environment.SetDefine(TEXT("PLATFORM_ALLOW_SCENE_DATA_COMPRESSED_TRANSFORMS"), FDataDrivenShaderPlatformInfo::GetSupportSceneDataCompressedTransforms(EShaderPlatform(Target.Platform)) ? 1 : 0);
 
-	if (RHISupportsBindless(EShaderPlatform(Target.Platform)))
+	bool bEnableBindlessMacro = false;
+	if (RHIGetBindlessSupport(EShaderPlatform(Target.Platform)) != ERHIBindlessSupport::Unsupported)
 	{
 		const bool bIsRaytracingShader = IsRayTracingShaderFrequency(Input.Target.GetFrequency());
 
@@ -6297,17 +6297,21 @@ void GlobalBeginCompileShader(
 
 		if (ResourcesConfig == ERHIBindlessConfiguration::AllShaders || (ResourcesConfig == ERHIBindlessConfiguration::RayTracingShaders && bIsRaytracingShader))
 		{
+			bEnableBindlessMacro = true;
 			Input.Environment.CompilerFlags.Add(CFLAG_BindlessResources);
 			Input.Environment.SetDefine(TEXT("ENABLE_BINDLESS_RESOURCES"), true);
 		}
 
 		if (SamplersConfig == ERHIBindlessConfiguration::AllShaders || (SamplersConfig == ERHIBindlessConfiguration::RayTracingShaders && bIsRaytracingShader))
 		{
+			bEnableBindlessMacro = true;
 			Input.Environment.CompilerFlags.Add(CFLAG_BindlessSamplers);
 			Input.Environment.SetDefine(TEXT("ENABLE_BINDLESS_SAMPLERS"), true);
 		}
 	}
 
+	Input.Environment.SetDefine(TEXT("PLATFORM_SUPPORTS_BINDLESS"), bEnableBindlessMacro);
+
 	if (CVarShadersRemoveDeadCode.GetValueOnAnyThread())
 	{
 		Input.Environment.CompilerFlags.Add(CFLAG_RemoveDeadCode);
diff --git a/Engine/Source/Runtime/RHI/Private/RHI.cpp b/Engine/Source/Runtime/RHI/Private/RHI.cpp
--- a/Engine/Source/Runtime/RHI/Private/RHI.cpp
+++ b/Engine/Source/Runtime/RHI/Private/RHI.cpp
@@ -1364,11 +1364,18 @@ static FAutoConsoleVariableRef CVarEnableBindlessSamplers(
 
 static ERHIBindlessConfiguration DetermineBindlessConfiguration(EShaderPlatform Platform, int32 BindlessConfigSetting)
 {
-	if (!RHISupportsBindless(Platform) || BindlessConfigSetting == 0)
+	const ERHIBindlessSupport BindlessSupport = RHIGetBindlessSupport(Platform);
+
+	if (BindlessSupport == ERHIBindlessSupport::Unsupported || BindlessConfigSetting == 0)
 	{
 		return ERHIBindlessConfiguration::Disabled;
 	}
 
+	if (BindlessSupport == ERHIBindlessSupport::RayTracingOnly)
+	{
+		return ERHIBindlessConfiguration::RayTracingShaders;
+	}
+
 	return BindlessConfigSetting == 2 ? ERHIBindlessConfiguration::RayTracingShaders : ERHIBindlessConfiguration::AllShaders;
 }
 
@@ -1561,6 +1568,7 @@ bool GRHISupportsMapWriteNoOverwrite = false;
 
 bool GRHISupportsSeparateDepthStencilCopyAccess = true;
 
+ERHIBindlessSupport GRHIBindlessSupport = ERHIBindlessSupport::Unsupported;
 bool GRHISupportsBindless = false;
 
 FVertexElementTypeSupportInfo GVertexElementTypeSupport;
@@ -2456,6 +2464,33 @@ static inline uint32 GetSectionFeatureSupport(const FConfigSection& Section, FNa
 	return OriginalValue;
 }
 
+static inline uint32 GetSectionBindlessSupport(const FConfigSection& Section, FName Key, uint32 OriginalValue)
+{
+	const FConfigValue* ConfigValue = Section.Find(Key);
+	if (ConfigValue != nullptr)
+	{
+		FString Value = ConfigValue->GetValue();
+		if (Value == TEXT("Unsupported"))
+		{
+			return uint32(ERHIBindlessSupport::Unsupported);
+		}
+		if (Value == TEXT("RayTracingOnly"))
+		{
+			return uint32(ERHIBindlessSupport::RayTracingOnly);
+		}
+		else if (Value == TEXT("AllShaderTypes"))
+		{
+			return uint32(ERHIBindlessSupport::AllShaderTypes);
+		}
+		else
+		{
+			checkf(false, TEXT("Unknown ERHIBindlessSupport value \"%s\" for %s"), *Value, *Key.ToString());
+		}
+	}
+
+	return OriginalValue;
+}
+
 FRHIViewableResource* GetViewableResource(const FRHITransitionInfo& Info)
 {
 	switch (Info.Type)
@@ -2511,6 +2546,10 @@ void FGenericDataDrivenShaderPlatformInfo::ParseDataDrivenShaderInfo(const FConf
 	Info.SettingName = GetSectionFeatureSupport(Section, #SettingName, Info.SettingName); \
 	ADD_TO_PROPERTIES_STRING(SettingName, Info.SettingName)
 
+#define GET_SECTION_BINDLESS_SUPPORT_HELPER(SettingName) \
+	Info.SettingName = GetSectionBindlessSupport(Section, #SettingName, Info.SettingName); \
+	ADD_TO_PROPERTIES_STRING(SettingName, Info.SettingName)
+
 	GET_SECTION_BOOL_HELPER(bIsMobile);
 	GET_SECTION_BOOL_HELPER(bIsMetalMRT);
 	GET_SECTION_BOOL_HELPER(bIsPC);
@@ -2589,7 +2628,7 @@ void FGenericDataDrivenShaderPlatformInfo::ParseDataDrivenShaderInfo(const FConf
 	GET_SECTION_BOOL_HELPER(bSupportsMobileDistanceField);
 	GET_SECTION_BOOL_HELPER(bSupportsFFTBloom);
 	GET_SECTION_BOOL_HELPER(bSupportsVertexShaderLayer);
-	GET_SECTION_BOOL_HELPER(bSupportsBindless);
+	GET_SECTION_BINDLESS_SUPPORT_HELPER(BindlessSupport);
 	GET_SECTION_BOOL_HELPER(bSupportsVolumeTextureAtomics);
 	GET_SECTION_BOOL_HELPER(bSupportsROV);
 	GET_SECTION_BOOL_HELPER(bSupportsOIT);
diff --git a/Engine/Source/Runtime/RHI/Public/RHI.h b/Engine/Source/Runtime/RHI/Public/RHI.h
--- a/Engine/Source/Runtime/RHI/Public/RHI.h
+++ b/Engine/Source/Runtime/RHI/Public/RHI.h
@@ -302,9 +302,15 @@ inline RHI_API bool RHISupportsConservativeRasterization(const FStaticShaderPlat
 }
 
 /** True if the given shader platform supports bindless resources/views. */
+inline ERHIBindlessSupport RHIGetBindlessSupport(const FStaticShaderPlatform Platform)
+{
+	return FDataDrivenShaderPlatformInfo::GetBindlessSupport(Platform);
+}
+
+UE_DEPRECATED(5.2, "You must use RHIGetBindlessSupport instead.")
 inline bool RHISupportsBindless(EShaderPlatform Platform)
 {
-	return FDataDrivenShaderPlatformInfo::GetSupportsBindless(Platform);
+	return RHIGetBindlessSupport(Platform) == ERHIBindlessSupport::AllShaderTypes;
 }
 
 inline bool RHISupportsVolumeTextureAtomics(EShaderPlatform Platform)
@@ -817,6 +823,9 @@ extern RHI_API FVector2f GRHIDefaultMSAASampleOffsets[1 + 2 + 4 + 8 + 16];
 extern RHI_API bool GRHISupportsAsyncPipelinePrecompile;
 
 /** Whether dynamic (bindless) resources are supported */
+extern RHI_API ERHIBindlessSupport GRHIBindlessSupport;
+
+UE_DEPRECATED(5.2, "You must use GRHIBindlessSupport instead.")
 extern RHI_API bool GRHISupportsBindless;
 
 // Calculate the index of the sample in GRHIDefaultMSAASampleOffsets
diff --git a/Engine/Source/Runtime/RHI/Public/RHIDefinitions.h b/Engine/Source/Runtime/RHI/Public/RHIDefinitions.h
--- a/Engine/Source/Runtime/RHI/Public/RHIDefinitions.h
+++ b/Engine/Source/Runtime/RHI/Public/RHIDefinitions.h
@@ -51,6 +51,15 @@ enum class ERHIFeatureSupport : uint8
 	NumBits = 2,
 };
 
+enum class ERHIBindlessSupport : uint8
+{
+	Unsupported,
+	RayTracingOnly,
+	AllShaderTypes,
+
+	NumBits = 2
+};
+
 enum EShaderFrequency : uint8
 {
 	SF_Vertex			= 0,
@@ -416,7 +425,7 @@ class RHI_API FGenericDataDrivenShaderPlatformInfo
 	uint32 bSupportsInlineRayTracing : 1;
 	uint32 bSupportsRayTracingShaders : 1;
 	uint32 bSupportsVertexShaderLayer : 1;
-	uint32 bSupportsBindless : 1;
+	uint32 BindlessSupport : int32(ERHIBindlessSupport::NumBits);
 	uint32 bSupportsVolumeTextureAtomics : 1;
 	uint32 bSupportsROV : 1;
 	uint32 bSupportsOIT : 1;
@@ -1002,9 +1011,16 @@ public:
 		return Infos[Platform].bSupportsVertexShaderLayer;
 	}
 
+	static FORCEINLINE_DEBUGGABLE const ERHIBindlessSupport GetBindlessSupport(const FStaticShaderPlatform Platform)
+	{
+		check(IsValid(Platform));
+		return static_cast<ERHIBindlessSupport>(Infos[Platform].BindlessSupport);
+	}
+
+	UE_DEPRECATED(5.2, "You must use GetBindlessSupport instead.")
 	static FORCEINLINE_DEBUGGABLE const bool GetSupportsBindless(const FStaticShaderPlatform Platform)
 	{
-		return Infos[Platform].bSupportsBindless;
+		return GetBindlessSupport(Platform) == ERHIBindlessSupport::AllShaderTypes;
 	}
 
 	static FORCEINLINE_DEBUGGABLE const bool GetSupportsVolumeTextureAtomics(const FStaticShaderPlatform Platform)
diff --git a/Engine/Source/Runtime/RenderCore/Private/Shader.cpp b/Engine/Source/Runtime/RenderCore/Private/Shader.cpp
--- a/Engine/Source/Runtime/RenderCore/Private/Shader.cpp
+++ b/Engine/Source/Runtime/RenderCore/Private/Shader.cpp
@@ -1979,21 +1979,19 @@ void ShaderMapAppendKeyString(EShaderPlatform Platform, FString& KeyString)
 		KeyString += TEXT("_MS_T1");
 	}
 
-	if (RHISupportsBindless(Platform))
+	if (RHIGetBindlessSupport(Platform) != ERHIBindlessSupport::Unsupported)
 	{
-		KeyString += TEXT("_BNDLS");
-
 		const ERHIBindlessConfiguration ResourcesConfig = RHIGetBindlessResourcesConfiguration(Platform);
 		const ERHIBindlessConfiguration SamplersConfig = RHIGetBindlessSamplersConfiguration(Platform);
 
 		if (ResourcesConfig != ERHIBindlessConfiguration::Disabled)
 		{
-			KeyString += ResourcesConfig == ERHIBindlessConfiguration::RayTracingShaders ? TEXT("_RTRES") : TEXT("ALLRES");
+			KeyString += ResourcesConfig == ERHIBindlessConfiguration::RayTracingShaders ? TEXT("_BNDLSRTRES") : TEXT("_BNDLSRES");
 		}
 
 		if (SamplersConfig != ERHIBindlessConfiguration::Disabled)
 		{
-			KeyString += SamplersConfig == ERHIBindlessConfiguration::RayTracingShaders ? TEXT("_RTSAM") : TEXT("ALLSAM");
+			KeyString += SamplersConfig == ERHIBindlessConfiguration::RayTracingShaders ? TEXT("_BNDLSRTSAM") : TEXT("_BNDLSSAM");
 		}
 	}
 
diff --git a/Engine/Source/Runtime/VulkanRHI/Private/VulkanDevice.cpp b/Engine/Source/Runtime/VulkanRHI/Private/VulkanDevice.cpp
--- a/Engine/Source/Runtime/VulkanRHI/Private/VulkanDevice.cpp
+++ b/Engine/Source/Runtime/VulkanRHI/Private/VulkanDevice.cpp
@@ -1080,9 +1080,7 @@ void FVulkanDevice::InitGPU()
 	DescriptorPoolsManager = new FVulkanDescriptorPoolsManager();
 	DescriptorPoolsManager->Init(this);
 
-	if (RHISupportsBindless(GMaxRHIShaderPlatform) &&
-		OptionalDeviceExtensions.HasEXTDescriptorIndexing &&
-		(GpuProps.limits.maxBoundDescriptorSets >= VulkanBindless::MaxNumSets))
+	if (OptionalDeviceExtensions.HasEXTDescriptorIndexing && GpuProps.limits.maxBoundDescriptorSets >= VulkanBindless::MaxNumSets)
 	{
 		const bool bBindlessResources = RHIGetBindlessResourcesConfiguration(GMaxRHIShaderPlatform) != ERHIBindlessConfiguration::Disabled;
 		const bool bBindlessSamplers = RHIGetBindlessSamplersConfiguration(GMaxRHIShaderPlatform) != ERHIBindlessConfiguration::Disabled;
diff --git a/Engine/Source/Runtime/VulkanRHI/Private/VulkanRHI.cpp b/Engine/Source/Runtime/VulkanRHI/Private/VulkanRHI.cpp
--- a/Engine/Source/Runtime/VulkanRHI/Private/VulkanRHI.cpp
+++ b/Engine/Source/Runtime/VulkanRHI/Private/VulkanRHI.cpp
@@ -822,7 +822,10 @@ void FVulkanDynamicRHI::InitInstance()
 		GRHISupportsBaseVertexIndex = true;
 		GSupportsSeparateRenderTargetBlendState = true;
 		GRHISupportsSeparateDepthStencilCopyAccess = false;
-		GRHISupportsBindless =	RHISupportsBindless(GMaxRHIShaderPlatform) && Device->SupportsBindless();
+		GRHIBindlessSupport = Device->SupportsBindless() ? RHIGetBindlessSupport(GMaxRHIShaderPlatform) : ERHIBindlessSupport::Unsupported;
+PRAGMA_DISABLE_DEPRECATION_WARNINGS
+		GRHISupportsBindless = GRHIBindlessSupport != ERHIBindlessSupport::Unsupported;
+PRAGMA_ENABLE_DEPRECATION_WARNINGS
 
 		GRHIMaxDispatchThreadGroupsPerDimension.X = FMath::Min<uint32>(Limits.maxComputeWorkGroupCount[0], 0x7fffffff);
 		GRHIMaxDispatchThreadGroupsPerDimension.Y = FMath::Min<uint32>(Limits.maxComputeWorkGroupCount[1], 0x7fffffff);
