- Add support to get the maximum number of bytes we can store in a constant buffer

#rnx

diff --git a/Engine/Source/Runtime/RHI/Private/RHI.cpp b/Engine/Source/Runtime/RHI/Private/RHI.cpp
--- a/Engine/Source/Runtime/RHI/Private/RHI.cpp
+++ b/Engine/Source/Runtime/RHI/Private/RHI.cpp
@@ -1469,6 +1469,7 @@ TRHIGlobal<int32> GMaxShadowDepthBufferSizeX(2048);
 TRHIGlobal<int32> GMaxShadowDepthBufferSizeY(2048);
 TRHIGlobal<int32> GMaxTextureDimensions(2048);
 TRHIGlobal<int64> GMaxBufferDimensions(1<<27);
+TRHIGlobal<int64> GRHIMaxConstantBufferByteSize(1<<27);
 TRHIGlobal<int64> GMaxComputeSharedMemory(1<<15);
 TRHIGlobal<int32> GMaxVolumeTextureDimensions(2048);
 TRHIGlobal<int32> GMaxCubeTextureDimensions(2048);
diff --git a/Engine/Source/Runtime/RHI/Public/RHI.h b/Engine/Source/Runtime/RHI/Public/RHI.h
--- a/Engine/Source/Runtime/RHI/Public/RHI.h
+++ b/Engine/Source/Runtime/RHI/Public/RHI.h
@@ -526,6 +526,9 @@ extern RHI_API TRHIGlobal<int32> GMaxTextureDimensions;
 /** The maximum size allowed for 2D textures in both dimensions. */
 extern RHI_API TRHIGlobal<int64> GMaxBufferDimensions;
 
+/** The maximum size allowed for a contant buffer. */
+extern RHI_API TRHIGlobal<int64> GRHIMaxConstantBufferByteSize;
+
 /** The maximum size allowed for Shared Compute Memory. */
 extern RHI_API TRHIGlobal<int64> GMaxComputeSharedMemory;
 
@@ -541,12 +544,16 @@ extern RHI_API bool GRHISupportsRawViewsForAnyBuffer;
 /** Whether depth or stencil can individually be set to CopySrc/Dest access. */
 extern RHI_API bool GRHISupportsSeparateDepthStencilCopyAccess;
 
-
 FORCEINLINE uint64 GetMaxBufferDimension()
 {
 	return GMaxBufferDimensions;
 }
 
+FORCEINLINE uint64 GetMaxConstantBufferByteSize()
+{
+	return GRHIMaxConstantBufferByteSize;
+}
+
 FORCEINLINE uint64 GetMaxComputeSharedMemory()
 {
 	return GMaxComputeSharedMemory;
diff --git a/Engine/Source/Runtime/Windows/D3D11RHI/Private/D3D11Device.cpp b/Engine/Source/Runtime/Windows/D3D11RHI/Private/D3D11Device.cpp
--- a/Engine/Source/Runtime/Windows/D3D11RHI/Private/D3D11Device.cpp
+++ b/Engine/Source/Runtime/Windows/D3D11RHI/Private/D3D11Device.cpp
@@ -5,6 +5,7 @@
 =============================================================================*/
 
 #include "D3D11RHIPrivate.h"
+#include "D3D11ConstantBuffer.h"
 #include "Misc/CommandLine.h"
 #include "Misc/ConfigCacheIni.h"
 #include "Modules/ModuleManager.h"
@@ -207,6 +208,7 @@ FD3D11DynamicRHI::FD3D11DynamicRHI(IDXGIFactory1* InDXGIFactory1, D3D_FEATURE_LE
 	GMaxTextureDimensions = D3D11_REQ_TEXTURE2D_U_OR_V_DIMENSION;
 	GMaxCubeTextureDimensions = D3D11_REQ_TEXTURECUBE_DIMENSION;
 	GMaxTextureArrayLayers = D3D11_REQ_TEXTURE2D_ARRAY_AXIS_DIMENSION;
+	GRHIMaxConstantBufferByteSize = MAX_GLOBAL_CONSTANT_BUFFER_SIZE;
 	GRHISupportsMSAADepthSampleAccess = true;
 	GRHISupportsRHIThread = !!EXPERIMENTAL_D3D11_RHITHREAD;
 
