Deprecate RHIRayTraceIntersection and RHIRayTraceOcclusion built-in fixed function ray tracing methods. Explicit ray tracing pipelines should be used instead.

diff --git a/Engine/Source/Runtime/D3D12RHI/Private/D3D12RayTracing.cpp b/Engine/Source/Runtime/D3D12RHI/Private/D3D12RayTracing.cpp
--- a/Engine/Source/Runtime/D3D12RHI/Private/D3D12RayTracing.cpp
+++ b/Engine/Source/Runtime/D3D12RHI/Private/D3D12RayTracing.cpp
@@ -2218,9 +2218,11 @@ public:
 template<typename ShaderType>
 static FD3D12RayTracingShader* GetBuildInRayTracingShader()
 {
+	PRAGMA_DISABLE_DEPRECATION_WARNINGS
 	FGlobalShaderMap* ShaderMap = GetGlobalShaderMap(GMaxRHIFeatureLevel);
 	auto Shader = ShaderMap->GetShader<ShaderType>();
 	FD3D12RayTracingShader* RayTracingShader = static_cast<FD3D12RayTracingShader*>(Shader.GetRayTracingShader());
+	PRAGMA_ENABLE_DEPRECATION_WARNINGS
 	return RayTracingShader;
 }
 
@@ -2858,6 +2860,8 @@ public:
 	{
 		// Occlusion pipeline
 		{
+			PRAGMA_DISABLE_DEPRECATION_WARNINGS
+
 			FRayTracingPipelineStateInitializer OcclusionInitializer;
 
 			FRHIRayTracingShader* OcclusionRGSTable[] = { GetBuildInRayTracingShader<FOcclusionMainRG>() };
@@ -2869,10 +2873,14 @@ public:
 			OcclusionInitializer.bAllowHitGroupIndexing = false;
 
 			Occlusion = new FD3D12RayTracingPipelineState(Device, OcclusionInitializer);
+
+			PRAGMA_ENABLE_DEPRECATION_WARNINGS
 		}
 
 		// Intersection pipeline
 		{
+			PRAGMA_DISABLE_DEPRECATION_WARNINGS
+
 			FRayTracingPipelineStateInitializer IntersectionInitializer;
 
 			FRHIRayTracingShader* IntersectionRGSTable[] = { GetBuildInRayTracingShader<FIntersectionMainRG>() };
@@ -2887,6 +2895,8 @@ public:
 			IntersectionInitializer.bAllowHitGroupIndexing = false;
 
 			Intersection = new FD3D12RayTracingPipelineState(Device, IntersectionInitializer);
+
+			PRAGMA_ENABLE_DEPRECATION_WARNINGS
 		}
 	}
 
diff --git a/Engine/Source/Runtime/RHI/Public/RHICommandList.h b/Engine/Source/Runtime/RHI/Public/RHICommandList.h
--- a/Engine/Source/Runtime/RHI/Public/RHICommandList.h
+++ b/Engine/Source/Runtime/RHI/Public/RHICommandList.h
@@ -199,7 +199,7 @@ enum class ERayTracingBindingType : uint8
 };
 
 // C++ counter-part of FBasicRayData declared in RayTracingCommon.ush
-struct FBasicRayData
+struct UE_DEPRECATED(5.1, "Please use an explicit ray generation shader with a custom intersection structure instead.") FBasicRayData
 {
 	float Origin[3];
 	uint32 Mask;
@@ -208,7 +208,7 @@ struct FBasicRayData
 };
 
 // C++ counter-part of FIntersectionPayload declared in RayTracingCommon.ush
-struct FIntersectionPayload
+struct UE_DEPRECATED(5.1, "Please use an explicit ray generation shader with a custom intersection structure instead.") FIntersectionPayload
 {
 	float  HitT;            // Distance from ray origin to the intersection point in the ray direction. Negative on miss.
 	uint32 PrimitiveIndex;  // Index of the primitive within the geometry inside the bottom-level acceleration structure instance. Undefined on miss.
@@ -3764,11 +3764,13 @@ public:
 	 * Binary intersection results are written to output buffer as R32_UINTs.
 	 * 0xFFFFFFFF is written if ray intersects any scene triangle, 0 otherwise.
 	 */
+	UE_DEPRECATED(5.1, "Please use an explicit ray generation shader and RayTraceDispatch() instead.")
 	FORCEINLINE_DEBUGGABLE void RayTraceOcclusion(FRHIRayTracingScene* Scene,
 		FRHIShaderResourceView* Rays,
 		FRHIUnorderedAccessView* Output,
 		uint32 NumRays)
 	{
+		PRAGMA_DISABLE_DEPRECATION_WARNINGS
 		if (Bypass())
 		{
 			GetContext().RHIRayTraceOcclusion(Scene, Rays, Output, NumRays);
@@ -3777,17 +3779,20 @@ public:
 		{
 			ALLOC_COMMAND(FRHICommandRayTraceOcclusion)(Scene, Rays, Output, NumRays);
 		}
+		PRAGMA_ENABLE_DEPRECATION_WARNINGS
 	}
 
 	/**
 	 * Trace rays from an input buffer of FBasicRayData.
 	 * Primitive intersection results are written to output buffer as FIntersectionPayload.
 	 */
+	UE_DEPRECATED(5.1, "Please use an explicit ray generation shader and RayTraceDispatch() instead.")
 	FORCEINLINE_DEBUGGABLE void RayTraceIntersection(FRHIRayTracingScene* Scene,
 		FRHIShaderResourceView* Rays,
 		FRHIUnorderedAccessView* Output,
 		uint32 NumRays)
 	{
+		PRAGMA_DISABLE_DEPRECATION_WARNINGS
 		if (Bypass())
 		{
 			GetContext().RHIRayTraceIntersection(Scene, Rays, Output, NumRays);
@@ -3796,6 +3801,7 @@ public:
 		{
 			ALLOC_COMMAND(FRHICommandRayTraceIntersection)(Scene, Rays, Output, NumRays);
 		}
+		PRAGMA_ENABLE_DEPRECATION_WARNINGS
 	}
 
 	FORCEINLINE_DEBUGGABLE void RayTraceDispatch(FRayTracingPipelineState* Pipeline, FRHIRayTracingShader* RayGenShader, FRHIRayTracingScene* Scene, const FRayTracingShaderBindings& GlobalResourceBindings, uint32 Width, uint32 Height)
diff --git a/Engine/Source/Runtime/RHI/Public/RHICommandListCommandExecutes.inl b/Engine/Source/Runtime/RHI/Public/RHICommandListCommandExecutes.inl
--- a/Engine/Source/Runtime/RHI/Public/RHICommandListCommandExecutes.inl
+++ b/Engine/Source/Runtime/RHI/Public/RHICommandListCommandExecutes.inl
@@ -596,13 +596,17 @@ void FRHICommandBuildAccelerationStructures::Execute(FRHICommandListBase& CmdLis
 void FRHICommandRayTraceOcclusion::Execute(FRHICommandListBase& CmdList)
 {
 	RHISTAT(RayTraceOcclusion);
+PRAGMA_DISABLE_DEPRECATION_WARNINGS
 	INTERNAL_DECORATOR(RHIRayTraceOcclusion)(Scene, Rays, Output, NumRays);
+PRAGMA_ENABLE_DEPRECATION_WARNINGS
 }
 
 void FRHICommandRayTraceIntersection::Execute(FRHICommandListBase& CmdList)
 {
 	RHISTAT(RayTraceIntersection);
+PRAGMA_DISABLE_DEPRECATION_WARNINGS
 	INTERNAL_DECORATOR(RHIRayTraceIntersection)(Scene, Rays, Output, NumRays);
+PRAGMA_ENABLE_DEPRECATION_WARNINGS
 }
 
 void FRHICommandRayTraceDispatch::Execute(FRHICommandListBase& CmdList)
diff --git a/Engine/Source/Runtime/RHI/Public/RHIContext.h b/Engine/Source/Runtime/RHI/Public/RHIContext.h
--- a/Engine/Source/Runtime/RHI/Public/RHIContext.h
+++ b/Engine/Source/Runtime/RHI/Public/RHIContext.h
@@ -797,6 +797,7 @@ public:
 		checkNoEntry();
 	}
 
+	UE_DEPRECATED(5.1, "Please use an explicit ray generation shader and RHIRayTraceDispatch() instead.")
 	virtual void RHIRayTraceOcclusion(FRHIRayTracingScene* Scene,
 		FRHIShaderResourceView* Rays,
 		FRHIUnorderedAccessView* Output,
@@ -805,6 +806,7 @@ public:
 		checkNoEntry();
 	}
 
+	UE_DEPRECATED(5.1, "Please use an explicit ray generation shader and RHIRayTraceDispatch() instead.")
 	virtual void RHIRayTraceIntersection(FRHIRayTracingScene* Scene,
 		FRHIShaderResourceView* Rays,
 		FRHIUnorderedAccessView* Output,
diff --git a/Engine/Source/Runtime/RHI/Public/RHIValidationContext.h b/Engine/Source/Runtime/RHI/Public/RHIValidationContext.h
--- a/Engine/Source/Runtime/RHI/Public/RHIValidationContext.h
+++ b/Engine/Source/Runtime/RHI/Public/RHIValidationContext.h
@@ -1030,7 +1030,9 @@ public:
 		FRHIUnorderedAccessView* Output,
 		uint32 NumRays) override final
 	{
+		PRAGMA_DISABLE_DEPRECATION_WARNINGS
 		RHIContext->RHIRayTraceOcclusion(Scene, Rays, Output, NumRays);
+		PRAGMA_ENABLE_DEPRECATION_WARNINGS
 	}
 
 	virtual void RHIRayTraceIntersection(FRHIRayTracingScene* Scene,
@@ -1038,7 +1040,9 @@ public:
 		FRHIUnorderedAccessView* Output,
 		uint32 NumRays) override final
 	{
+		PRAGMA_DISABLE_DEPRECATION_WARNINGS
 		RHIContext->RHIRayTraceIntersection(Scene, Rays, Output, NumRays);
+		PRAGMA_ENABLE_DEPRECATION_WARNINGS
 	}
 
 	virtual void RHIRayTraceDispatch(FRHIRayTracingPipelineState* RayTracingPipelineState, FRHIRayTracingShader* RayGenShader,
diff --git a/Engine/Source/Runtime/RenderCore/Private/BuiltInRayTracingShaders.cpp b/Engine/Source/Runtime/RenderCore/Private/BuiltInRayTracingShaders.cpp
--- a/Engine/Source/Runtime/RenderCore/Private/BuiltInRayTracingShaders.cpp
+++ b/Engine/Source/Runtime/RenderCore/Private/BuiltInRayTracingShaders.cpp
@@ -6,9 +6,12 @@
 
 #if RHI_RAYTRACING
 
+PRAGMA_DISABLE_DEPRECATION_WARNINGS
 IMPLEMENT_GLOBAL_SHADER( FOcclusionMainRG,		"/Engine/Private/RayTracing/RayTracingBuiltInShaders.usf",		"OcclusionMainRG",				SF_RayGen);
 IMPLEMENT_GLOBAL_SHADER( FIntersectionMainRG,	"/Engine/Private/RayTracing/RayTracingBuiltInShaders.usf",		"IntersectionMainRG",			SF_RayGen);
 IMPLEMENT_SHADER_TYPE(, FIntersectionMainCHS,	TEXT("/Engine/Private/RayTracing/RayTracingBuiltInShaders.usf"), TEXT("IntersectionMainCHS"),	SF_RayHitGroup);
+PRAGMA_ENABLE_DEPRECATION_WARNINGS
+
 IMPLEMENT_SHADER_TYPE(, FDefaultMainCHS,		TEXT("/Engine/Private/RayTracing/RayTracingBuiltInShaders.usf"), TEXT("DefaultMainCHS"),		SF_RayHitGroup);
 IMPLEMENT_SHADER_TYPE(, FDefaultMainCHSOpaqueAHS, TEXT("/Engine/Private/RayTracing/RayTracingBuiltInShaders.usf"), TEXT("closesthit=DefaultMainCHS anyhit=DefaultOpaqueAHS"), SF_RayHitGroup);
 IMPLEMENT_SHADER_TYPE(, FDefaultPayloadMS,		TEXT("/Engine/Private/RayTracing/RayTracingBuiltInShaders.usf"), TEXT("DefaultPayloadMS"),			SF_RayMiss);
diff --git a/Engine/Source/Runtime/RenderCore/Public/BuiltInRayTracingShaders.h b/Engine/Source/Runtime/RenderCore/Public/BuiltInRayTracingShaders.h
--- a/Engine/Source/Runtime/RenderCore/Public/BuiltInRayTracingShaders.h
+++ b/Engine/Source/Runtime/RenderCore/Public/BuiltInRayTracingShaders.h
@@ -28,7 +28,7 @@ public:
 	{}
 };
 
-class RENDERCORE_API FOcclusionMainRG : public FBuiltInRayTracingShader
+class RENDERCORE_API UE_DEPRECATED(5.1, "Please use an explicit ray generation shader instead.") FOcclusionMainRG : public FBuiltInRayTracingShader
 {
 	DECLARE_GLOBAL_SHADER(FOcclusionMainRG);
 	SHADER_USE_ROOT_PARAMETER_STRUCT(FOcclusionMainRG, FBuiltInRayTracingShader);
@@ -40,7 +40,7 @@ class RENDERCORE_API FOcclusionMainRG : public FBuiltInRayTracingShader
 	END_SHADER_PARAMETER_STRUCT()
 };
 
-class RENDERCORE_API FIntersectionMainRG : public FBuiltInRayTracingShader
+class RENDERCORE_API UE_DEPRECATED(5.1, "Please use an explicit ray generation shader instead.") FIntersectionMainRG : public FBuiltInRayTracingShader
 {
 	DECLARE_GLOBAL_SHADER(FIntersectionMainRG);
 	SHADER_USE_ROOT_PARAMETER_STRUCT(FIntersectionMainRG, FBuiltInRayTracingShader);
@@ -52,7 +52,7 @@ class RENDERCORE_API FIntersectionMainRG : public FBuiltInRayTracingShader
 	END_SHADER_PARAMETER_STRUCT()
 };
 
-class FIntersectionMainCHS : public FBuiltInRayTracingShader
+class UE_DEPRECATED(5.1, "Please use an explicit ray generation and hit shaders instead.") FIntersectionMainCHS : public FBuiltInRayTracingShader
 {
 	DECLARE_EXPORTED_SHADER_TYPE(FIntersectionMainCHS, Global, RENDERCORE_API);
 public:
diff --git a/Engine/Source/Runtime/Renderer/Private/Tests/RayTracingTestbed.cpp b/Engine/Source/Runtime/Renderer/Private/Tests/RayTracingTestbed.cpp
--- a/Engine/Source/Runtime/Renderer/Private/Tests/RayTracingTestbed.cpp
+++ b/Engine/Source/Runtime/Renderer/Private/Tests/RayTracingTestbed.cpp
@@ -63,6 +63,7 @@ bool RunRayTracingTestbed_RenderThread(const FString& Parameters)
 	FShaderResourceViewRHIRef RayBufferView;
 
 	{
+		PRAGMA_DISABLE_DEPRECATION_WARNINGS
 		TResourceArray<FBasicRayData> RayData;
 		RayData.SetNumUninitialized(NumRays);
 		RayData[0] = FBasicRayData{ { 0.75f, 0.0f, -1.0f}, 0xFFFFFFFF, {0.0f, 0.0f,  1.0f}, 100000.0f }; // expected to hit
@@ -75,6 +76,7 @@ bool RunRayTracingTestbed_RenderThread(const FString& Parameters)
 
 		RayBuffer = RHICreateStructuredBuffer(sizeof(FBasicRayData), RayData.GetResourceDataSize(), BUF_Static | BUF_ShaderResource, CreateInfo);
 		RayBufferView = RHICreateShaderResourceView(RayBuffer);
+		PRAGMA_ENABLE_DEPRECATION_WARNINGS
 	}
 
 	FBufferRHIRef OcclusionResultBuffer;
@@ -91,7 +93,9 @@ bool RunRayTracingTestbed_RenderThread(const FString& Parameters)
 
 	{
 		FRHIResourceCreateInfo CreateInfo(TEXT("IntersectionResultBuffer"));
+		PRAGMA_DISABLE_DEPRECATION_WARNINGS
 		IntersectionResultBuffer = RHICreateVertexBuffer(sizeof(FIntersectionPayload)*NumRays, BUF_Static | BUF_UnorderedAccess, CreateInfo);
+		PRAGMA_ENABLE_DEPRECATION_WARNINGS
 		IntersectionResultBufferView = RHICreateUnorderedAccessView(IntersectionResultBuffer, PF_R32_UINT);
 	}
 
@@ -229,8 +233,10 @@ bool RunRayTracingTestbed_RenderThread(const FString& Parameters)
 
 	RHICmdList.Transition(FRHITransitionInfo(RayTracingScene.Scene.GetReference(), ERHIAccess::BVHWrite, ERHIAccess::BVHRead));
 
+	PRAGMA_DISABLE_DEPRECATION_WARNINGS
 	RHICmdList.RayTraceOcclusion(RayTracingScene.Scene, RayBufferView, OcclusionResultBufferView, NumRays);
 	RHICmdList.RayTraceIntersection(RayTracingScene.Scene, RayBufferView, IntersectionResultBufferView, NumRays);
+	PRAGMA_ENABLE_DEPRECATION_WARNINGS
 
 	const bool bValidateResults = true;
 	bool bOcclusionTestOK = false;
@@ -261,6 +267,7 @@ bool RunRayTracingTestbed_RenderThread(const FString& Parameters)
 		// Read back and validate intersection trace results
 
 		{
+			PRAGMA_DISABLE_DEPRECATION_WARNINGS
 			auto MappedResults = (const FIntersectionPayload*)RHILockBuffer(IntersectionResultBuffer, 0, sizeof(FIntersectionPayload)*NumRays, RLM_ReadOnly);
 
 			check(MappedResults);
@@ -279,6 +286,7 @@ bool RunRayTracingTestbed_RenderThread(const FString& Parameters)
 			RHIUnlockBuffer(IntersectionResultBuffer);
 
 			bIntersectionTestOK = (MappedResults[0].HitT >= 0) && (MappedResults[1].HitT < 0) && (MappedResults[2].HitT >= 0) && (MappedResults[3].HitT < 0);
+			PRAGMA_ENABLE_DEPRECATION_WARNINGS
 		}
 	}
 
diff --git a/Engine/Source/Runtime/VulkanRHI/Private/VulkanRayTracing.cpp b/Engine/Source/Runtime/VulkanRHI/Private/VulkanRayTracing.cpp
--- a/Engine/Source/Runtime/VulkanRHI/Private/VulkanRayTracing.cpp
+++ b/Engine/Source/Runtime/VulkanRHI/Private/VulkanRayTracing.cpp
@@ -1124,6 +1124,8 @@ FVulkanBasicRaytracingPipeline::FVulkanBasicRaytracingPipeline(FVulkanDevice* co
 
 	// Occlusion pipeline
 	{
+		PRAGMA_DISABLE_DEPRECATION_WARNINGS
+
 		FRayTracingPipelineStateInitializer OcclusionInitializer;
 
 		FRHIRayTracingShader* OcclusionRGSTable[] = { GetBuiltInRayTracingShader<FOcclusionMainRG>() };
@@ -1138,6 +1140,8 @@ FVulkanBasicRaytracingPipeline::FVulkanBasicRaytracingPipeline(FVulkanDevice* co
 		OcclusionInitializer.bAllowHitGroupIndexing = false;
 
 		Occlusion = new FVulkanRayTracingPipelineState(InDevice, OcclusionInitializer);
+
+		PRAGMA_ENABLE_DEPRECATION_WARNINGS
 	}
 }
 
