fix RHI display info with validation RHI

diff --git a/Engine/Source/Runtime/RHI/Public/RHIValidation.h b/Engine/Source/Runtime/RHI/Public/RHIValidation.h
--- a/Engine/Source/Runtime/RHI/Public/RHIValidation.h
+++ b/Engine/Source/Runtime/RHI/Public/RHIValidation.h
@@ -1034,6 +1034,11 @@ public:
 		return RHI->RHICreateAliasedTexture(SourceTexture);
 	}
 
+	virtual void RHIGetDisplaysInformation(FDisplayInformationArray& OutDisplayInformation) override final
+	{
+		RHI->RHIGetDisplaysInformation(OutDisplayInformation);
+	}
+
 	virtual void RHIAdvanceFrameFence() override final
 	{
 		check(IsInRenderingThread());
