- Only allow PSOPrecaching outside of editor builds for now (needs more testing in editor)
- Enable PSO precaching by default in Lyra & CitySample projects for Windows

#rnx

diff --git a/Engine/Source/Runtime/RHI/Private/PipelineStateCache.cpp b/Engine/Source/Runtime/RHI/Private/PipelineStateCache.cpp
--- a/Engine/Source/Runtime/RHI/Private/PipelineStateCache.cpp
+++ b/Engine/Source/Runtime/RHI/Private/PipelineStateCache.cpp
@@ -2129,8 +2129,13 @@ FGraphicsPipelineState* PipelineStateCache::FindGraphicsPipelineState(const FGra
 }
 
 bool PipelineStateCache::IsPSOPrecachingEnabled()
-{
+{	
+#if WITH_EDITOR
+	// Disables in the editor for now by default untill more testing is done - still WIP
+	return false;
+#else
 	return GPSOPrecaching != 0;
+#endif // WITH_EDITOR
 }
 
 FGraphEventRef PipelineStateCache::PrecacheComputePipelineState(FRHIComputeShader* ComputeShader)
