RHI thread perf:  Zach's CL 22019956 moved to Release 5.1, fix for parallel translate stalling due to unnecessary RHI thread fence delaying the start of tasks.  Saves 2.7 ms CPU on Valley test scene Node 2.

#rnx

diff --git a/Engine/Source/Runtime/RHI/Private/RHICommandList.cpp b/Engine/Source/Runtime/RHI/Private/RHICommandList.cpp
--- a/Engine/Source/Runtime/RHI/Private/RHICommandList.cpp
+++ b/Engine/Source/Runtime/RHI/Private/RHICommandList.cpp
@@ -151,6 +151,11 @@ FRHICommandListBase::FRHICommandListBase(FPersistentState&& InPersistentState)
 	CommandLink = &Root;
 	UID = GRHICommandList.UIDCounter.Increment();
 
+	if (!IsImmediate())
+	{
+		PersistentState.FenceCandidate = new FPersistentState::FFenceCandidate;
+	}
+
 #if DO_CHECK
 	if (PersistentState.RecordingThread == ERecordingThread::Render)
 	{
@@ -259,6 +264,8 @@ void FRHICommandListBase::FinishRecording()
 {
 	checkf(!IsImmediate(), TEXT("Do not call FinishRecording() on the immediate RHI command list."));
 
+	PersistentState.FenceCandidate->Fence = PersistentState.RHIThreadBufferLockFence;
+
 	// "Complete" the dispatch event. This unblocks waiting tasks but only when
 	// all dependencies added via AddDispatchPrerequisite() have been resolved.
 	DispatchEvent->DispatchSubsequents();
@@ -507,8 +514,39 @@ void FRHICommandListImmediate::QueueAsyncCommandListSubmit(TArrayView<FQueuedCom
 				Prereqs.Add(CmdList->DispatchEvent);
 			}
 
-			// This is used to ensure that any old buffer locks are completed before we start any parallel translates.
-			if (PersistentState.RHIThreadBufferLockFence)
+			if (PersistentState.QueuedFenceCandidates.Num() > 0)
+			{
+				FGraphEventRef FenceCandidateEvent = FGraphEvent::CreateGraphEvent();
+
+				if (PersistentState.RHIThreadBufferLockFence.GetReference())
+				{
+					FenceCandidateEvent->DontCompleteUntil(PersistentState.RHIThreadBufferLockFence);
+				}
+
+				PersistentState.RHIThreadBufferLockFence = FenceCandidateEvent;
+				Prereqs.Add(FenceCandidateEvent);
+
+				FFunctionGraphTask::CreateAndDispatchWhenReady(
+					[FenceCandidates = MoveTemp(PersistentState.QueuedFenceCandidates), FenceCandidateEvent](ENamedThreads::Type, const FGraphEventRef&) mutable
+				{
+					SCOPED_NAMED_EVENT(STAT_FRHICommandListBase_SignalLockFence, FColor::Magenta);
+
+					for (int32 Index = FenceCandidates.Num() - 1; Index >= 0; Index--)
+					{
+						if (FenceCandidates[Index]->Fence)
+						{
+							FenceCandidateEvent->DontCompleteUntil(FenceCandidates[Index]->Fence);
+							break;
+						}
+					}
+
+					FenceCandidateEvent->DispatchSubsequents();
+
+				}, TStatId(), &PersistentState.QueuedFenceCandidateEvents);
+
+				PersistentState.QueuedFenceCandidateEvents.Reset();
+			}
+			else if (PersistentState.RHIThreadBufferLockFence)
 			{
 				Prereqs.Add(PersistentState.RHIThreadBufferLockFence);
 			}
@@ -553,8 +591,6 @@ void FRHICommandListImmediate::QueueAsyncCommandListSubmit(TArrayView<FQueuedCom
 			);
 		}
 
-		PersistentState.RHIThreadBufferLockFence = nullptr;
-
 		// Resize the tasks array view to how many tasks we actually created after merging
 		Tasks = TArrayView<FTask>(Tasks.GetData(), NumTasks);
 
@@ -587,7 +623,10 @@ void FRHICommandListImmediate::QueueAsyncCommandListSubmit(TArrayView<FQueuedCom
 		TArrayView<FRHICommandListBase*> CmdListsView = AllocArrayUninitialized<FRHICommandListBase*>(CommandLists.Num());
 		for (int32 Index = 0; Index < CommandLists.Num(); ++Index)
 		{
-			CmdListsView[Index] = CommandLists[Index].CmdList;
+			FRHICommandListBase* CommandList = CommandLists[Index].CmdList;
+			PersistentState.QueuedFenceCandidateEvents.Emplace(CommandList->DispatchEvent);
+			PersistentState.QueuedFenceCandidates.Emplace(CommandList->PersistentState.FenceCandidate);
+			CmdListsView[Index] = CommandList;
 		}
 
 		EnqueueLambda([CmdListsView](FRHICommandListBase& ParentCmdList)
@@ -599,11 +638,6 @@ void FRHICommandListImmediate::QueueAsyncCommandListSubmit(TArrayView<FQueuedCom
 				delete CmdList;
 			}
 		});
-
-		// The async command list is being recorded in parallel, we have to take a conservative RHI thread fence
-		// for all the work in case a lock was used. The RHI thread waits on the work, so the render thread can't
-		// access the pending RHI thread fence.
-		RHIThreadFence(true);
 	}
 }
 
@@ -754,6 +788,12 @@ FRHICOMMAND_MACRO(FRHICommandRHIThreadFence)
 
 FGraphEventRef FRHICommandListBase::RHIThreadFence(bool bSetLockFence)
 {
+	if (bSetLockFence)
+	{
+		PersistentState.QueuedFenceCandidateEvents.Empty();
+		PersistentState.QueuedFenceCandidates.Empty();
+	}
+
 	if (IsRunningRHIInSeparateThread())
 	{
 		FRHICommandRHIThreadFence* Cmd = ALLOC_COMMAND(FRHICommandRHIThreadFence)();
diff --git a/Engine/Source/Runtime/RHI/Public/RHICommandList.h b/Engine/Source/Runtime/RHI/Public/RHICommandList.h
--- a/Engine/Source/Runtime/RHI/Public/RHICommandList.h
+++ b/Engine/Source/Runtime/RHI/Public/RHICommandList.h
@@ -750,6 +750,15 @@ protected:
 
 		FGraphEventRef RHIThreadBufferLockFence;
 
+		struct FFenceCandidate : public TConcurrentLinearObject<FFenceCandidate>, public FRefCountBase
+		{
+			FGraphEventRef Fence;
+		};
+
+		TRefCountPtr<FFenceCandidate> FenceCandidate;
+		FGraphEventArray QueuedFenceCandidateEvents;
+		TArray<TRefCountPtr<FFenceCandidate>, FConcurrentLinearArrayAllocator> QueuedFenceCandidates;
+
 		FPersistentState(FRHIGPUMask InInitialGPUMask, ERecordingThread InRecordingThread)
 			: RecordingThread(InRecordingThread)
 			, CurrentGPUMask(InInitialGPUMask)
