Added an explicit way of selecting RHI renderer/feature level.
Users can specify PreferredRHI and PreferredFeatureLevel instead of the less expressive bUseD3D12InGame and bPreferFeatureLevelES31.
The boolean settings are still read and applied, to preserve backwards compatibility.

#rnx

diff --git a/Engine/Source/Runtime/RHI/Private/Windows/WindowsDynamicRHI.cpp b/Engine/Source/Runtime/RHI/Private/Windows/WindowsDynamicRHI.cpp
--- a/Engine/Source/Runtime/RHI/Private/Windows/WindowsDynamicRHI.cpp
+++ b/Engine/Source/Runtime/RHI/Private/Windows/WindowsDynamicRHI.cpp
@@ -42,6 +42,18 @@ static const TCHAR* ModuleNameFromWindowsRHI(EWindowsRHI InWindowsRHI)
 	}
 }
 
+static const TCHAR* GetRHINameFromWindowsRHI(EWindowsRHI InWindowsRHI)
+{
+	switch (InWindowsRHI)
+	{
+	default: check(false);
+	case EWindowsRHI::D3D11:  return TEXT("DirectX 11");
+	case EWindowsRHI::D3D12:  return TEXT("DirectX 12");
+	case EWindowsRHI::Vulkan: return TEXT("Vulkan");
+	case EWindowsRHI::OpenGL: return TEXT("OpenGL");
+	}
+}
+
 static FString GetRHINameFromWindowsRHI(EWindowsRHI InWindowsRHI, ERHIFeatureLevel::Type InFeatureLevel)
 {
 	switch (InWindowsRHI)
@@ -376,30 +388,54 @@ static bool DefaultFeatureLevelES31()
 	return false;
 }
 
-static bool PreferFeatureLevelES31()
+static TOptional<ERHIFeatureLevel::Type> GetPreferredFeatureLevel()
 {
+	TOptional<ERHIFeatureLevel::Type> PreferredFeatureLevel;
 	if (!GIsEditor)
 	{
-		bool bPreferFeatureLevelES31 = false;
-		bool bFoundPreference = GConfig->GetBool(TEXT("D3DRHIPreference"), TEXT("bPreferFeatureLevelES31"), bPreferFeatureLevelES31, GGameUserSettingsIni);
-
-		// Force low-spec users into performance mode but respect their choice once they have set a preference
-		bool bDefaultES31 = false;
-		if (!bFoundPreference && !CVarIgnorePerformanceModeCheck.GetValueOnAnyThread())
+		FString PreferredFeatureLevelName;
+		if (GConfig->GetString(TEXT("D3DRHIPreference"), TEXT("PreferredFeatureLevel"), PreferredFeatureLevelName, GGameUserSettingsIni))
 		{
-			bDefaultES31 = DefaultFeatureLevelES31();
+			if (PreferredFeatureLevelName == TEXT("sm5"))
+			{
+				PreferredFeatureLevel = ERHIFeatureLevel::SM5;
+			}
+			else if (PreferredFeatureLevelName == TEXT("sm6"))
+			{
+				PreferredFeatureLevel = ERHIFeatureLevel::SM6;
+			}
+			else if (PreferredFeatureLevelName == TEXT("es31"))
+			{
+				PreferredFeatureLevel = ERHIFeatureLevel::ES3_1;
+			}
+			else
+			{
+				UE_LOG(LogRHI, Error, TEXT("unknown feature level name \"%s\" in game user settings, using default"), *PreferredFeatureLevelName);
+			}
 		}
-
-		if (bPreferFeatureLevelES31 || bDefaultES31)
+		else
 		{
-			if (!bFoundPreference)
+			bool bPreferFeatureLevelES31 = false;
+			bool bFoundPreference = GConfig->GetBool(TEXT("D3DRHIPreference"), TEXT("bPreferFeatureLevelES31"), bPreferFeatureLevelES31, GGameUserSettingsIni);
+
+			// Force low-spec users into performance mode but respect their choice once they have set a preference
+			bool bDefaultES31 = false;
+			if (!bFoundPreference && !CVarIgnorePerformanceModeCheck.GetValueOnAnyThread())
 			{
-				GConfig->SetBool(TEXT("D3DRHIPreference"), TEXT("bPreferFeatureLevelES31"), true, GGameUserSettingsIni);
+				bDefaultES31 = DefaultFeatureLevelES31();
+			}
+
+			if (bPreferFeatureLevelES31 || bDefaultES31)
+			{
+				if (!bFoundPreference)
+				{
+					GConfig->SetBool(TEXT("D3DRHIPreference"), TEXT("bPreferFeatureLevelES31"), true, GGameUserSettingsIni);
+				}
+				PreferredFeatureLevel = ERHIFeatureLevel::ES3_1;
 			}
-			return true;
 		}
 	}
-	return false;
+	return PreferredFeatureLevel;
 }
 
 static bool IsES31D3DOnly()
@@ -459,9 +495,33 @@ static TOptional<EWindowsRHI> ChoosePreferredRHI(EWindowsRHI InDefaultRHI)
 	if (!GIsEditor && (InDefaultRHI == EWindowsRHI::D3D11 || InDefaultRHI == EWindowsRHI::D3D12))
 	{
 		bool bUseD3D12InGame = false;
-		if (GConfig->GetBool(TEXT("D3DRHIPreference"), TEXT("bUseD3D12InGame"), bUseD3D12InGame, GGameUserSettingsIni) && bUseD3D12InGame)
+		FString PreferredRHIName;
+		if (GConfig->GetString(TEXT("D3DRHIPreference"), TEXT("PreferredRHI"), PreferredRHIName, GGameUserSettingsIni))
 		{
-			RHIPreference = EWindowsRHI::D3D12;
+			if (PreferredRHIName == TEXT("dx12"))
+			{
+				RHIPreference = EWindowsRHI::D3D12;
+			}
+			else if (PreferredRHIName == TEXT("dx11"))
+			{
+				RHIPreference = EWindowsRHI::D3D11;
+			}
+			else if (PreferredRHIName == TEXT("vulkan"))
+			{
+				RHIPreference = EWindowsRHI::Vulkan;
+			}
+			else if (PreferredRHIName == TEXT("opengl"))
+			{
+				RHIPreference = EWindowsRHI::OpenGL;
+			}
+			else
+			{
+				UE_LOG(LogRHI, Error, TEXT("unknown RHI name \"%s\" in game user settings, using default"), *PreferredRHIName);
+			}
+		}
+		else if (GConfig->GetBool(TEXT("D3DRHIPreference"), TEXT("bUseD3D12InGame"), bUseD3D12InGame, GGameUserSettingsIni))
+		{
+			RHIPreference = bUseD3D12InGame ? EWindowsRHI::D3D12 : EWindowsRHI::D3D11;
 		}
 	}
 
@@ -574,13 +634,9 @@ static ERHIFeatureLevel::Type ChooseFeatureLevel(EWindowsRHI ChosenRHI, const TO
 		}
 	}
 
-	TOptional<ERHIFeatureLevel::Type> FeatureLevel{};
+	TOptional<ERHIFeatureLevel::Type> FeatureLevel = GetPreferredFeatureLevel();
 
-	if ((ChosenRHI == EWindowsRHI::D3D11 || ChosenRHI == EWindowsRHI::D3D12) && Config.IsFeatureLevelTargeted(ChosenRHI, ERHIFeatureLevel::ES3_1) && PreferFeatureLevelES31())
-	{
-		FeatureLevel = TOptional<ERHIFeatureLevel::Type>(ERHIFeatureLevel::ES3_1);
-	}
-	else
+	if (!FeatureLevel || (FPlatformProperties::RequiresCookedData() && !Config.IsFeatureLevelTargeted(ChosenRHI, FeatureLevel.GetValue())))
 	{
 		FeatureLevel = Config.GetHighestSupportedFeatureLevel(ChosenRHI);
 
