Fix an RHI validation error about Nanite.VisBuffer64 being in the wrong state when it is not

diff --git a/Engine/Source/Runtime/RHI/Public/RHIValidationContext.h b/Engine/Source/Runtime/RHI/Public/RHIValidationContext.h
--- a/Engine/Source/Runtime/RHI/Public/RHIValidationContext.h
+++ b/Engine/Source/Runtime/RHI/Public/RHIValidationContext.h
@@ -116,6 +116,36 @@ public:
 		RHIContext->RHIClearUAVUint(UnorderedAccessViewRHI, Values);
 	}
 
+	virtual void RHIBeginUAVOverlap() final override
+	{
+		Tracker->AllUAVsOverlap(true);
+		RHIContext->RHIBeginUAVOverlap();
+	}
+
+	virtual void RHIEndUAVOverlap() final override
+	{
+		Tracker->AllUAVsOverlap(false);
+		RHIContext->RHIEndUAVOverlap();
+	}
+
+	virtual void RHIBeginUAVOverlap(TArrayView<FRHIUnorderedAccessView* const> UAVs) final override
+	{
+		for (FRHIUnorderedAccessView* UAV : UAVs)
+		{
+			Tracker->SpecificUAVOverlap(UAV->ViewIdentity, true);
+		}
+		RHIContext->RHIBeginUAVOverlap(UAVs);
+	}
+
+	virtual void RHIEndUAVOverlap(TArrayView<FRHIUnorderedAccessView* const> UAVs) final override
+	{
+		for (FRHIUnorderedAccessView* UAV : UAVs)
+		{
+			Tracker->SpecificUAVOverlap(UAV->ViewIdentity, false);
+		}
+		RHIContext->RHIEndUAVOverlap(UAVs);
+	}
+
 	virtual void RHISubmitCommandsHint() override final
 	{
 		RHIValidation::FTracker::ReplayOpQueue(ERHIPipeline::AsyncCompute, Tracker->Finalize());
