Bindless: Add missing handles from TextureReference during UniformBuffer traversal.

#rnx

diff --git a/Engine/Source/Runtime/RHI/Public/RHIResources.h b/Engine/Source/Runtime/RHI/Public/RHIResources.h
--- a/Engine/Source/Runtime/RHI/Public/RHIResources.h
+++ b/Engine/Source/Runtime/RHI/Public/RHIResources.h
@@ -1932,6 +1932,12 @@ public:
 		return this;
 	}
 
+	virtual FRHIDescriptorHandle GetDefaultBindlessHandle() const override 
+	{ 
+		check(ReferencedTexture);
+		return ReferencedTexture->GetDefaultBindlessHandle();
+	}
+
 	virtual void* GetNativeResource() const override 
 	{
 		check(ReferencedTexture);
