[RHI]
* RHICommandList - Removed include of RHI.h to prevent circular dependencies

diff --git a/Engine/Source/Runtime/RHI/Public/RHICommandList.h b/Engine/Source/Runtime/RHI/Public/RHICommandList.h
--- a/Engine/Source/Runtime/RHI/Public/RHICommandList.h
+++ b/Engine/Source/Runtime/RHI/Public/RHICommandList.h
@@ -25,7 +25,6 @@
 #include "HAL/IConsoleManager.h"
 #include "Async/TaskGraphInterfaces.h"
 #include "HAL/LowLevelMemTracker.h"
-#include "RHI.h"
 #include "ProfilingDebugging/CsvProfiler.h"
 #include "ProfilingDebugging/CpuProfilerTrace.h"
 #include "Trace/Trace.h"
