[Engine]
* Fixes to make code work with header units

diff --git a/Engine/Source/Runtime/Engine/Classes/Engine/NetConnection.h b/Engine/Source/Runtime/Engine/Classes/Engine/NetConnection.h
--- a/Engine/Source/Runtime/Engine/Classes/Engine/NetConnection.h
+++ b/Engine/Source/Runtime/Engine/Classes/Engine/NetConnection.h
@@ -65,7 +65,7 @@ typedef TMap<TWeakObjectPtr<AActor>, UActorChannel*, FDefaultSetAllocator, TWeak
 namespace NetConnectionHelper
 {
 	/** Number of bits to use in the packet header for sending the milliseconds on the clock when the packet is sent */
-	constexpr int32 NumBitsForJitterClockTimeInHeader = 10;
+	inline constexpr int32 NumBitsForJitterClockTimeInHeader = 10;
 }
 
 extern ENGINE_API TAutoConsoleVariable<int32> CVarNetEnableCongestionControl;
diff --git a/Engine/Source/Runtime/RHI/Private/RHICommandList.cpp b/Engine/Source/Runtime/RHI/Private/RHICommandList.cpp
--- a/Engine/Source/Runtime/RHI/Private/RHICommandList.cpp
+++ b/Engine/Source/Runtime/RHI/Private/RHICommandList.cpp
@@ -1913,4 +1913,9 @@ RHI_API void RHISetComputeShaderBackwardsCompatible(IRHIComputeContext* InContex
 {
 	TRefCountPtr<FRHIComputePipelineState> ComputePipelineState = RHICreateComputePipelineState(InShader);
 	InContext->RHISetComputePipelineState(ComputePipelineState);
-}
\ No newline at end of file
+}
+
+void IRHIComputeContext::RHISetComputeShader(FRHIComputeShader* ComputeShader)
+{
+	RHISetComputeShaderBackwardsCompatible(this, ComputeShader);
+}
diff --git a/Engine/Source/Runtime/RHI/Public/RHIContext.h b/Engine/Source/Runtime/RHI/Public/RHIContext.h
--- a/Engine/Source/Runtime/RHI/Public/RHIContext.h
+++ b/Engine/Source/Runtime/RHI/Public/RHIContext.h
@@ -227,11 +227,7 @@ public:
 	}
 
 	UE_DEPRECATED(5.1, "ComputePipelineStates should be used instead of direct ComputeShaders.")
-	virtual void RHISetComputeShader(FRHIComputeShader* ComputeShader)
-	{
-		RHI_API void RHISetComputeShaderBackwardsCompatible(IRHIComputeContext*, FRHIComputeShader*);
-		RHISetComputeShaderBackwardsCompatible(this, ComputeShader);
-	}
+	virtual void RHISetComputeShader(FRHIComputeShader* ComputeShader);
 
 	virtual void RHISetComputePipelineState(FRHIComputePipelineState* ComputePipelineState) = 0;
 
