Undo Engine/Source/Runtime/... changelist 22838017

diff --git a/Engine/Source/Runtime/Engine/Classes/Engine/NetConnection.h b/Engine/Source/Runtime/Engine/Classes/Engine/NetConnection.h
--- a/Engine/Source/Runtime/Engine/Classes/Engine/NetConnection.h
+++ b/Engine/Source/Runtime/Engine/Classes/Engine/NetConnection.h
@@ -65,7 +65,7 @@ typedef TMap<TWeakObjectPtr<AActor>, UActorChannel*, FDefaultSetAllocator, TWeak
 namespace NetConnectionHelper
 {
 	/** Number of bits to use in the packet header for sending the milliseconds on the clock when the packet is sent */
-	inline constexpr int32 NumBitsForJitterClockTimeInHeader = 10;
+	constexpr int32 NumBitsForJitterClockTimeInHeader = 10;
 }
 
 extern ENGINE_API TAutoConsoleVariable<int32> CVarNetEnableCongestionControl;
diff --git a/Engine/Source/Runtime/RHI/Private/RHICommandList.cpp b/Engine/Source/Runtime/RHI/Private/RHICommandList.cpp
--- a/Engine/Source/Runtime/RHI/Private/RHICommandList.cpp
+++ b/Engine/Source/Runtime/RHI/Private/RHICommandList.cpp
@@ -1913,9 +1913,4 @@ RHI_API void RHISetComputeShaderBackwardsCompatible(IRHIComputeContext* InContex
 {
 	TRefCountPtr<FRHIComputePipelineState> ComputePipelineState = RHICreateComputePipelineState(InShader);
 	InContext->RHISetComputePipelineState(ComputePipelineState);
-}
-
-void IRHIComputeContext::RHISetComputeShader(FRHIComputeShader* ComputeShader)
-{
-	RHISetComputeShaderBackwardsCompatible(this, ComputeShader);
-}
+}
\ No newline at end of file
diff --git a/Engine/Source/Runtime/RHI/Public/RHIContext.h b/Engine/Source/Runtime/RHI/Public/RHIContext.h
--- a/Engine/Source/Runtime/RHI/Public/RHIContext.h
+++ b/Engine/Source/Runtime/RHI/Public/RHIContext.h
@@ -227,7 +227,11 @@ public:
 	}
 
 	UE_DEPRECATED(5.1, "ComputePipelineStates should be used instead of direct ComputeShaders.")
-	virtual void RHISetComputeShader(FRHIComputeShader* ComputeShader);
+	virtual void RHISetComputeShader(FRHIComputeShader* ComputeShader)
+	{
+		RHI_API void RHISetComputeShaderBackwardsCompatible(IRHIComputeContext*, FRHIComputeShader*);
+		RHISetComputeShaderBackwardsCompatible(this, ComputeShader);
+	}
 
 	virtual void RHISetComputePipelineState(FRHIComputePipelineState* ComputePipelineState) = 0;
 
