ShaderMap now serializes the ShaderPlatform name instead of ShaderPlatform enum

diff --git a/Engine/Shaders/Public/ShaderVersion.ush b/Engine/Shaders/Public/ShaderVersion.ush
--- a/Engine/Shaders/Public/ShaderVersion.ush
+++ b/Engine/Shaders/Public/ShaderVersion.ush
@@ -2,4 +2,4 @@
 // Each time the console command is executed it generates a new GUID. As this file is included
 // in Platform.ush (which should be included in any shader) it allows to invalidate the shader DDC.
 // 
-// GUID = D22FFC72-E8D1-4B88-B4D2-A05F60748A10
\ No newline at end of file
+// GUID = 22434BFA-4B31-4043-9494-8D69F6E789DD
\ No newline at end of file
diff --git a/Engine/Source/Runtime/RHI/Private/RHI.cpp b/Engine/Source/Runtime/RHI/Private/RHI.cpp
--- a/Engine/Source/Runtime/RHI/Private/RHI.cpp
+++ b/Engine/Source/Runtime/RHI/Private/RHI.cpp
@@ -2414,6 +2414,7 @@ const FName LANGUAGE_Sony("Sony");
 const FName LANGUAGE_Nintendo("Nintendo");
 
 RHI_API FGenericDataDrivenShaderPlatformInfo FGenericDataDrivenShaderPlatformInfo::Infos[SP_NumPlatforms];
+TMap<FName, EShaderPlatform> PlatformNameToShaderPlatformMap;
 
 // Gets a string from a section, or empty string if it didn't exist
 static inline FString GetSectionString(const FConfigSection& Section, FName Key)
@@ -2667,6 +2668,7 @@ void FGenericDataDrivenShaderPlatformInfo::Initialize()
 	{
 		return;
 	}
+	PlatformNameToShaderPlatformMap.Empty();
 
 	// look for the standard DataDriven ini files
 	int32 NumDDInfoFiles = FDataDrivenPlatformInfoRegistry::GetNumDataDrivenIniFiles();
@@ -2708,6 +2710,7 @@ void FGenericDataDrivenShaderPlatformInfo::Initialize()
 				
 				// at this point, we can start pulling information out
 				Infos[ShaderPlatform].Name = *SectionName.Mid(15);
+				PlatformNameToShaderPlatformMap.FindOrAdd(Infos[ShaderPlatform].Name) = ShaderPlatform;
 				ParseDataDrivenShaderInfo(Section.Value, Infos[ShaderPlatform]);	
 				Infos[ShaderPlatform].bContainsValidPlatformInfo = true;
 
@@ -2718,12 +2721,14 @@ void FGenericDataDrivenShaderPlatformInfo::Initialize()
 					if (Item.ShaderPlatformToPreview == PreviewPlatformName)
 					{
 						EShaderPlatform PreviewShaderPlatform = EShaderPlatform(CustomShaderPlatform++);
+						
 						ParseDataDrivenShaderInfo(Section.Value, Infos[PreviewShaderPlatform]);
 						if (!Item.OptionalFriendlyNameOverride.IsEmpty())
 						{
 							Infos[PreviewShaderPlatform].FriendlyName = Item.OptionalFriendlyNameOverride;
 						}
 						Infos[PreviewShaderPlatform].Name = Item.PreviewShaderPlatformName;
+						PlatformNameToShaderPlatformMap.FindOrAdd(Infos[PreviewShaderPlatform].Name) = PreviewShaderPlatform;
 						Infos[PreviewShaderPlatform].PreviewShaderPlatformParent = ShaderPlatform;
 						Infos[PreviewShaderPlatform].bIsPreviewPlatform = true;
 						Infos[PreviewShaderPlatform].bContainsValidPlatformInfo = true;
@@ -2786,18 +2791,9 @@ FText FGenericDataDrivenShaderPlatformInfo::GetFriendlyName(const FStaticShaderP
 
 const EShaderPlatform FGenericDataDrivenShaderPlatformInfo::GetShaderPlatformFromName(const FName ShaderPlatformName)
 {
-	for (int32 i = 0; i < SP_NumPlatforms; ++i)
+	if (EShaderPlatform* ShaderPlatform = PlatformNameToShaderPlatformMap.Find(ShaderPlatformName))
 	{
-		const EShaderPlatform Platform = static_cast<EShaderPlatform>(i);
-		if (!Infos[Platform].bContainsValidPlatformInfo)
-		{
-			continue;
-		}
-
-		if (Infos[Platform].Name == ShaderPlatformName)
-		{
-			return Platform;
-		}
+		return *ShaderPlatform;
 	}
 	return SP_NumPlatforms;
 }
diff --git a/Engine/Source/Runtime/RenderCore/Private/ShaderMap.cpp b/Engine/Source/Runtime/RenderCore/Private/ShaderMap.cpp
--- a/Engine/Source/Runtime/RenderCore/Private/ShaderMap.cpp
+++ b/Engine/Source/Runtime/RenderCore/Private/ShaderMap.cpp
@@ -169,8 +169,9 @@ bool FShaderMapBase::Serialize(FArchive& Ar, bool bInlineShaderResources, bool b
 
 		// Serialize a copy of ShaderPlatform directly into the archive
 		// This will allow us to correctly deserialize the stream, even if we're not able to load the frozen content
-		TEnumAsByte<EShaderPlatform> ShaderPlatform = GetShaderPlatform();
-		Ar << ShaderPlatform;
+		const EShaderPlatform ShaderPlatform = GetShaderPlatform();
+		FString ShaderPlatformName = FDataDrivenShaderPlatformInfo::GetName(GetShaderPlatform()).ToString();
+		Ar << ShaderPlatformName;
 
 		if (Ar.IsCooking())
 		{
@@ -205,8 +206,10 @@ bool FShaderMapBase::Serialize(FArchive& Ar, bool bInlineShaderResources, bool b
 		bool bShareCode = false;
 		Ar << bShareCode;
 
-		TEnumAsByte<EShaderPlatform> ShaderPlatform = SP_NumPlatforms;
-		Ar << ShaderPlatform;
+		FString ShaderPlatformName;
+		Ar << ShaderPlatformName;
+		
+		const EShaderPlatform ShaderPlatform = FDataDrivenShaderPlatformInfo::GetShaderPlatformFromName(FName(ShaderPlatformName));
 
 		if (bShareCode)
 		{
@@ -483,14 +486,15 @@ void FShaderMapContent::GetShaderList(const FShaderMapBase& InShaderMap, TMap<FH
 
 void FShaderMapContent::GetShaderPipelineList(const FShaderMapBase& InShaderMap, TArray<FShaderPipelineRef>& OutShaderPipelines, FShaderPipeline::EFilter Filter) const
 {
+	const EShaderPlatform ShaderPlatform = GetShaderPlatform();
 	for (FShaderPipeline* Pipeline : ShaderPipelines)
 	{
 		const FShaderPipelineType* PipelineType = FShaderPipelineType::GetShaderPipelineTypeByName(Pipeline->TypeName);
-		if (PipelineType->ShouldOptimizeUnusedOutputs(Platform) && Filter == FShaderPipeline::EOnlyShared)
+		if (PipelineType->ShouldOptimizeUnusedOutputs(ShaderPlatform) && Filter == FShaderPipeline::EOnlyShared)
 		{
 			continue;
 		}
-		else if (!PipelineType->ShouldOptimizeUnusedOutputs(Platform) && Filter == FShaderPipeline::EOnlyUnique)
+		else if (!PipelineType->ShouldOptimizeUnusedOutputs(ShaderPlatform) && Filter == FShaderPipeline::EOnlyUnique)
 		{
 			continue;
 		}
diff --git a/Engine/Source/Runtime/RenderCore/Public/Shader.h b/Engine/Source/Runtime/RenderCore/Public/Shader.h
--- a/Engine/Source/Runtime/RenderCore/Public/Shader.h
+++ b/Engine/Source/Runtime/RenderCore/Public/Shader.h
@@ -2011,7 +2011,7 @@ public:
 	};
 	/** Default constructor. */
 	explicit FShaderMapContent(EShaderPlatform InPlatform)
-		: ShaderHash(128u), Platform(InPlatform)
+		: ShaderHash(128u), ShaderPlatformName(FDataDrivenShaderPlatformInfo::GetName(InPlatform).ToString())
 	{}
 
 	/** Destructor ensures pipelines cleared up. */
@@ -2020,7 +2020,7 @@ public:
 		Empty();
 	}
 
-	EShaderPlatform GetShaderPlatform() const { return Platform; }
+	EShaderPlatform GetShaderPlatform() const { return FDataDrivenShaderPlatformInfo::GetShaderPlatformFromName(FName(ShaderPlatformName)); };
 
 	void Validate(const FShaderMapBase& InShaderMap) const;
 
@@ -2029,7 +2029,7 @@ public:
 	ShaderType* GetShader(int32 PermutationId = 0) const
 	{
 		FShader* Shader = GetShader(&ShaderType::StaticType, PermutationId);
-		checkf(Shader != nullptr, TEXT("Failed to find shader type %s in Platform %s"), ShaderType::StaticType.GetName(), *LegacyShaderPlatformToShaderFormat(Platform).ToString());
+		checkf(Shader != nullptr, TEXT("Failed to find shader type %s in Platform %s"), ShaderType::StaticType.GetName(), *LegacyShaderPlatformToShaderFormat(GetShaderPlatform()).ToString());
 		return static_cast<ShaderType*>(Shader);
 	}
 
@@ -2157,8 +2157,8 @@ protected:
 	LAYOUT_FIELD(TMemoryImageArray<int32>, ShaderPermutations);
 	LAYOUT_FIELD(TMemoryImageArray<TMemoryImagePtr<FShader>>, Shaders);
 	LAYOUT_FIELD(TMemoryImageArray<TMemoryImagePtr<FShaderPipeline>>, ShaderPipelines);
-	/** The platform this shader map was compiled with */
-	LAYOUT_FIELD(TEnumAsByte<EShaderPlatform>, Platform);
+	/** The ShaderPlatform Name this shader map was compiled with */
+	LAYOUT_FIELD(FMemoryImageString, ShaderPlatformName);
 };
 
 class RENDERCORE_API FShaderMapBase
