Fix LTO link error
 - The constructor for FRHIDrawStats was very large due to nested TStaticArrays, and causing "dynamic initializer for 'FRHICommandListImmediate::FrameDrawStats': function size suppresses optimizations" when LTO is enabled.
 - This should have been a simple Memzero. Refactored the FRHIDrawStats struct to avoid the expensive generated constructor.

Reduced size of FRHIDrawStats structure
 - By ignoring the PT_ enum values for tessellation, and counting only stats for triangles, lines, quads, points and rectangles.

diff --git a/Engine/Source/Runtime/RHI/Private/RHI.cpp b/Engine/Source/Runtime/RHI/Private/RHI.cpp
--- a/Engine/Source/Runtime/RHI/Private/RHI.cpp
+++ b/Engine/Source/Runtime/RHI/Private/RHI.cpp
@@ -1614,20 +1614,15 @@ RHI_API void FRHIDrawStats::Accumulate(FRHIDrawStats& Other)
 {
 	for (uint32 GPUIndex = 0; GPUIndex < GNumExplicitGPUsForRendering; ++GPUIndex)
 	{
-		FPerGPUStats& LeftGPU = GPUs[GPUIndex];
-		FPerGPUStats& RightGPU = Other.GPUs[GPUIndex];
+		FPerGPUStats& LeftGPU = GetGPU(GPUIndex);
+		FPerGPUStats& RightGPU = Other.GetGPU(GPUIndex);
 
 		for (int32 CategoryIndex = 0; CategoryIndex < NumCategories; ++CategoryIndex)
 		{
-			FPerCategoryStats& LeftCategory = LeftGPU.Categories[CategoryIndex];
-			FPerCategoryStats& RightCategory = RightGPU.Categories[CategoryIndex];
+			FPerCategoryStats& LeftCategory = LeftGPU.GetCategory(CategoryIndex);
+			FPerCategoryStats& RightCategory = RightGPU.GetCategory(CategoryIndex);
 
-			LeftCategory.Draws += RightCategory.Draws;
-
-			for (uint32 PrimitiveType = 0; PrimitiveType < PT_Num; ++PrimitiveType)
-			{
-				LeftCategory.Primitives[PrimitiveType] += RightCategory.Primitives[PrimitiveType];
-			}
+			LeftCategory += RightCategory;
 		}
 	}
 }
@@ -1651,43 +1646,23 @@ void FRHICommandListImmediate::ProcessStats()
 #endif
 
 	// Summed stats across all GPUs
-	uint32 TotalDraws     = 0;
-	uint32 TotalTriangles = 0;
-	uint32 TotalLines     = 0;
-	TStaticArray<uint32, FRHIDrawStats::NumCategories> TotalDrawsPerCategory { InPlace, 0 };
+	FRHIDrawStats::FPerCategoryStats Total = {};
+	TStaticArray<FRHIDrawStats::FPerCategoryStats, FRHIDrawStats::NumCategories> TotalPerCategory;
+	FMemory::Memzero(TotalPerCategory);
 
 	for (int32 GPUIndex = 0; GPUIndex < MAX_NUM_GPUS; ++GPUIndex)
 	{
-		uint32 TotalDraws_ThisGPU      = 0;
-		uint32 TotalPrimitives_ThisGPU = 0;
+		FRHIDrawStats::FPerCategoryStats TotalPerGPU = {};
 
-		FRHIDrawStats::FPerGPUStats& GPUStats = FrameDrawStats.GPUs[GPUIndex];
+		FRHIDrawStats::FPerGPUStats& GPUStats = FrameDrawStats.GetGPU(GPUIndex);
 
 		for (int32 CategoryIndex = 0; CategoryIndex < FRHIDrawStats::NumCategories; ++CategoryIndex)
 		{
-			FRHIDrawStats::FPerCategoryStats& Category = GPUStats.Categories[CategoryIndex];
+			FRHIDrawStats::FPerCategoryStats& Category = GPUStats.GetCategory(CategoryIndex);
 
-			TotalDraws                           += Category.Draws;
-			TotalDraws_ThisGPU                   += Category.Draws;
-			TotalDrawsPerCategory[CategoryIndex] += Category.Draws;
-
-			for (uint32 PrimitiveType = 0; PrimitiveType < PT_Num; ++PrimitiveType)
-			{
-				uint32 PrimitiveCount = Category.Primitives[PrimitiveType];
-				TotalPrimitives_ThisGPU += PrimitiveCount;
-
-				switch (PrimitiveType)
-				{
-				case PT_TriangleList:
-				case PT_TriangleStrip:
-					TotalTriangles += PrimitiveCount;
-					break;
-
-				case PT_LineList:
-					TotalLines += PrimitiveCount;
-					break;
-				}
-			}
+			TotalPerCategory[CategoryIndex] += Category;
+			TotalPerGPU                     += Category;
+			Total                           += Category;
 
 #if HAS_GPU_STATS
 			if (bCopyDisplayFrames && CategoryIndex < FDrawCallCategoryName::NumCategory)
@@ -1697,24 +1672,24 @@ void FRHICommandListImmediate::ProcessStats()
 #endif // HAS_GPU_STATS
 		}
 		
-		GNumDrawCallsRHI      [GPUIndex] = TotalDraws_ThisGPU;
-		GNumPrimitivesDrawnRHI[GPUIndex] = TotalPrimitives_ThisGPU;
+		GNumDrawCallsRHI      [GPUIndex] = TotalPerGPU.Draws;
+		GNumPrimitivesDrawnRHI[GPUIndex] = TotalPerGPU.GetTotalPrimitives();
 	}
 
 	// Multi-GPU support : CSV stats do not support MGPU yet. We're summing the totals across all GPUs here.
-	CSV_CUSTOM_STAT(RHI, DrawCalls      , int32(TotalDraws                 ), ECsvCustomStatOp::Set);
-	CSV_CUSTOM_STAT(RHI, PrimitivesDrawn, int32(TotalTriangles + TotalLines), ECsvCustomStatOp::Set);
+	CSV_CUSTOM_STAT(RHI, DrawCalls      , int32(Total.Draws               ), ECsvCustomStatOp::Set);
+	CSV_CUSTOM_STAT(RHI, PrimitivesDrawn, int32(Total.GetTotalPrimitives()), ECsvCustomStatOp::Set);
 
 #if HAS_GPU_STATS
-	SET_DWORD_STAT(STAT_RHITriangles         , TotalTriangles);
-	SET_DWORD_STAT(STAT_RHILines             , TotalLines    );
-	SET_DWORD_STAT(STAT_RHIDrawPrimitiveCalls, TotalDraws    );
+	SET_DWORD_STAT(STAT_RHITriangles         , Total.Triangles);
+	SET_DWORD_STAT(STAT_RHILines             , Total.Lines    );
+	SET_DWORD_STAT(STAT_RHIDrawPrimitiveCalls, Total.Draws    );
 
 	#if CSV_PROFILER
 	for (int32 CategoryIndex = 0; CategoryIndex < FDrawCallCategoryName::NumCategory; ++CategoryIndex)
 	{
 		FDrawCallCategoryName* CategoryName = FDrawCallCategoryName::Array[CategoryIndex];
-		FCsvProfiler::RecordCustomStat(CategoryName->Name, CSV_CATEGORY_INDEX(DrawCall), int32(TotalDrawsPerCategory[CategoryIndex]), ECsvCustomStatOp::Set);
+		FCsvProfiler::RecordCustomStat(CategoryName->Name, CSV_CATEGORY_INDEX(DrawCall), int32(TotalPerCategory[CategoryIndex].Draws), ECsvCustomStatOp::Set);
 	}
 	#endif
 #endif // HAS_GPU_STATS
diff --git a/Engine/Source/Runtime/RHI/Public/RHI.h b/Engine/Source/Runtime/RHI/Public/RHI.h
--- a/Engine/Source/Runtime/RHI/Public/RHI.h
+++ b/Engine/Source/Runtime/RHI/Public/RHI.h
@@ -2707,11 +2707,19 @@ extern RHI_API int32 GNumPrimitivesDrawnRHI[MAX_NUM_GPUS];
 // Updates the Stats structure on the underlying RHI context class (IRHICommandContext)
 
 #define RHI_DRAW_CALL_INC() do { Stats->Draws++; } while (false)
-#define RHI_DRAW_CALL_STATS(PrimitiveType,NumPrimitives)   \
-	do                                                     \
-	{												       \
-		Stats->Primitives[PrimitiveType] += NumPrimitives; \
-		Stats->Draws++;                                    \
+#define RHI_DRAW_CALL_STATS(PrimitiveType,NumPrimitives)                  \
+	do                                                                    \
+	{												                      \
+		switch (PrimitiveType)                                            \
+		{                                                                 \
+		case PT_TriangleList : Stats->Triangles  += NumPrimitives; break; \
+		case PT_TriangleStrip: Stats->Triangles  += NumPrimitives; break; \
+		case PT_LineList     : Stats->Lines      += NumPrimitives; break; \
+		case PT_QuadList     : Stats->Quads      += NumPrimitives; break; \
+		case PT_PointList    : Stats->Points     += NumPrimitives; break; \
+		case PT_RectList     : Stats->Rectangles += NumPrimitives; break; \
+		}                                                                 \
+		Stats->Draws++;                                                   \
 	} while(false)
 
 struct FRHIDrawStats
@@ -2727,16 +2735,51 @@ struct FRHIDrawStats
 
 	struct FPerCategoryStats
 	{
-		uint32 Draws = 0;
-		TStaticArray<uint32, PT_Num> Primitives { InPlace, 0 };
+		uint32 Draws;
+		uint32 Triangles;
+		uint32 Lines;
+		uint32 Quads;
+		uint32 Points;
+		uint32 Rectangles;
+
+		uint32 GetTotalPrimitives() const
+		{
+			return Triangles
+				+ Lines
+				+ Quads
+				+ Points
+				+ Rectangles;
+		}
+
+		FPerCategoryStats& operator += (FPerCategoryStats const& RHS)
+		{
+			Draws      += RHS.Draws;
+			Triangles  += RHS.Triangles;
+			Lines      += RHS.Lines;
+			Quads      += RHS.Quads;
+			Points     += RHS.Points;
+			Rectangles += RHS.Rectangles;
+			return *this;
+		}
 	};
 
 	struct FPerGPUStats
 	{
-		TStaticArray<FPerCategoryStats, NumCategories> Categories { InPlace };
+		FPerCategoryStats& GetCategory(uint32 Category)
+		{
+			checkSlow(Category < UE_ARRAY_COUNT(Categories));
+			return Categories[Category];
+		}
+
+	private:
+		FPerCategoryStats Categories[NumCategories];
 	};
 
-	TStaticArray<FPerGPUStats, MAX_NUM_GPUS> GPUs { InPlace };
+	FPerGPUStats& GetGPU(uint32 GPUIndex)
+	{
+		checkSlow(GPUIndex < UE_ARRAY_COUNT(GPUs));
+		return GPUs[GPUIndex];
+	}
 
 	FRHIDrawStats()
 	{
@@ -2748,8 +2791,10 @@ struct FRHIDrawStats
 		FMemory::Memzero(*this);
 	}
 
-	// Adds the values from Other into the current instance.
-	RHI_API void Accumulate(FRHIDrawStats& Other);
+	RHI_API void Accumulate(FRHIDrawStats& RHS);
+
+private:
+	FPerGPUStats GPUs[MAX_NUM_GPUS];
 };
 
 // RHI memory stats.
diff --git a/Engine/Source/Runtime/RHI/Public/RHIContext.h b/Engine/Source/Runtime/RHI/Public/RHIContext.h
--- a/Engine/Source/Runtime/RHI/Public/RHIContext.h
+++ b/Engine/Source/Runtime/RHI/Public/RHIContext.h
@@ -470,9 +470,7 @@ protected:
 public:
 	void StatsSetCategory(FRHIDrawStats* InStats, uint32 InCategoryID, uint32 InGPUIndex)
 	{
-		FRHIDrawStats::FPerGPUStats& GPU = InStats->GPUs[InGPUIndex];
-		FRHIDrawStats::FPerCategoryStats& Category = GPU.Categories[InCategoryID];
-		Stats = &Category;
+		Stats = &InStats->GetGPU(InGPUIndex).GetCategory(InCategoryID);
 	}
 
 #if WITH_MGPU || ENABLE_RHI_VALIDATION
