Fix crash with csv profiler and draw call category stats
 - The categories are global objects, so the initialization order is undefined.
 - The FDrawCallCategoryName::Array static object was being initialized after some of the category objects, clearing their pointers from the array, causing a nullptr dereference later.
 - Moved the global arrays etc into a static "manager" singleton which is lazy-initialized by a getter function. This fixes the undefined init order problem.

diff --git a/Engine/Source/Runtime/Engine/Private/UnrealEngine.cpp b/Engine/Source/Runtime/Engine/Private/UnrealEngine.cpp
--- a/Engine/Source/Runtime/Engine/Private/UnrealEngine.cpp
+++ b/Engine/Source/Runtime/Engine/Private/UnrealEngine.cpp
@@ -17644,39 +17644,50 @@ int32 UEngine::RenderStatUnit(UWorld* World, FViewport* Viewport, FCanvas* Canva
 int32 UEngine::RenderStatDrawCount(UWorld* World, FViewport* Viewport, FCanvas* Canvas, int32 X, int32 Y, const FVector* ViewLocation, const FRotator* ViewRotation)
 {
 #if HAS_GPU_STATS
+
 	int32 TotalCount[MAX_NUM_GPUS] = { 0 };
 	// Display all the categories of draw counts. This may always report 0 in some modes if AreGPUStatsEnabled is not enabled.
 	// Most likely because we are not currently capturing a CSV.
-	for (int32 Index = 0; Index < FDrawCallCategoryName::NumCategory; ++Index)
+	FDrawCallCategoryName::FManager const& Manager = FDrawCallCategoryName::GetManager();
+	for (int32 Index = 0; Index < Manager.NumCategory; ++Index)
 	{
 		for (uint32 GPUIndex : FRHIGPUMask::All())
 		{
-			TotalCount[GPUIndex] += FDrawCallCategoryName::DisplayCounts[Index][GPUIndex];
-		FDrawCallCategoryName* CategoryName = FDrawCallCategoryName::Array[Index];
-			Canvas->DrawShadowedString(X - 100,
-			Y,
+			TotalCount[GPUIndex] += Manager.DisplayCounts[Index][GPUIndex];
+			FDrawCallCategoryName* CategoryName = Manager.Array[Index];
+
+			Canvas->DrawShadowedString(
+				X - 100, Y,
 				*FString::Printf(
 					TEXT("%s%s: %i"),
 					*CategoryName->Name.ToString(),
 					GNumExplicitGPUsForRendering > 1 ? *FString::Printf(TEXT(" (GPU%u)"), GPUIndex) : TEXT(""),
-					FDrawCallCategoryName::DisplayCounts[Index][GPUIndex]),
-			GetSmallFont(),
-			FColor::Green);
-		Y += 12;
-	}
+					Manager.DisplayCounts[Index][GPUIndex]
+				),
+				GetSmallFont(),
+				FColor::Green
+			);
+
+			Y += 12;
+		}
 	}
+
 	for (uint32 GPUIndex : FRHIGPUMask::All())
 	{
-		Canvas->DrawShadowedString(X - 100,
-		Y,
+		Canvas->DrawShadowedString(
+			X - 100, Y,
 			*FString::Printf(
 				TEXT("Total%s: %i"),
 				GNumExplicitGPUsForRendering > 1 ? *FString::Printf(TEXT(" (GPU%u)"), GPUIndex) : TEXT(""),
-				TotalCount[GPUIndex]),
-		GetSmallFont(),
-		FColor::Green);
-	Y += 12;
+				TotalCount[GPUIndex]
+			),
+			GetSmallFont(),
+			FColor::Green
+		);
+
+		Y += 12;
 	}
+
 #else
 	Canvas->DrawShadowedString(X - 200,
 		Y,
diff --git a/Engine/Source/Runtime/RHI/Private/RHI.cpp b/Engine/Source/Runtime/RHI/Private/RHI.cpp
--- a/Engine/Source/Runtime/RHI/Private/RHI.cpp
+++ b/Engine/Source/Runtime/RHI/Private/RHI.cpp
@@ -300,9 +300,36 @@ FRHIResource* FRHIResource::CurrentlyDeleting = nullptr;
 
 #if HAS_GPU_STATS
 
-	TStaticArray<FDrawCallCategoryName*, FDrawCallCategoryName::MAX_DRAWCALL_CATEGORY> FDrawCallCategoryName::Array { InPlace, nullptr };
-	TStaticArray<TStaticArray<int32, MAX_NUM_GPUS>, FDrawCallCategoryName::MAX_DRAWCALL_CATEGORY> FDrawCallCategoryName::DisplayCounts { InPlace };
-	int32 FDrawCallCategoryName::NumCategory = 0;
+	FDrawCallCategoryName::FDrawCallCategoryName()
+		: Name(NAME_None)
+		, Index(-1)
+	{}
+
+	FDrawCallCategoryName::FDrawCallCategoryName(FName InName)
+		: Name(InName)
+		, Index(GetManager().NumCategory++)
+	{
+		check(Index < MAX_DRAWCALL_CATEGORY);
+		if (Index < MAX_DRAWCALL_CATEGORY)
+		{
+			GetManager().Array[Index] = this;
+		}
+	}
+
+	FDrawCallCategoryName::FManager::FManager()
+		: NumCategory(0)
+	{
+		FMemory::Memzero(Array);
+		FMemory::Memzero(DisplayCounts);
+	}
+
+	FDrawCallCategoryName::FManager& FDrawCallCategoryName::GetManager()
+	{
+		// Categories are global scope objects, so the initialization order is undefined.
+		// Lazy init the manager on first use.
+		static FManager Manager;
+		return Manager;
+	}
 
 #endif
 
@@ -1643,6 +1670,8 @@ void FRHICommandListImmediate::ProcessStats()
 		LastTime = CurrentTime;
 		bCopyDisplayFrames = true;
 	}
+
+	FDrawCallCategoryName::FManager& Manager = FDrawCallCategoryName::GetManager();
 #endif
 
 	// Summed stats across all GPUs
@@ -1665,9 +1694,9 @@ void FRHICommandListImmediate::ProcessStats()
 			Total                           += Category;
 
 #if HAS_GPU_STATS
-			if (bCopyDisplayFrames && CategoryIndex < FDrawCallCategoryName::NumCategory)
+			if (bCopyDisplayFrames && CategoryIndex < Manager.NumCategory)
 			{
-				FDrawCallCategoryName::DisplayCounts[CategoryIndex][GPUIndex] = Category.Draws;
+				Manager.DisplayCounts[CategoryIndex][GPUIndex] = Category.Draws;
 			}
 #endif // HAS_GPU_STATS
 		}
@@ -1686,10 +1715,9 @@ void FRHICommandListImmediate::ProcessStats()
 	SET_DWORD_STAT(STAT_RHIDrawPrimitiveCalls, Total.Draws    );
 
 	#if CSV_PROFILER
-	for (int32 CategoryIndex = 0; CategoryIndex < FDrawCallCategoryName::NumCategory; ++CategoryIndex)
+	for (int32 CategoryIndex = 0; CategoryIndex < Manager.NumCategory; ++CategoryIndex)
 	{
-		FDrawCallCategoryName* CategoryName = FDrawCallCategoryName::Array[CategoryIndex];
-		FCsvProfiler::RecordCustomStat(CategoryName->Name, CSV_CATEGORY_INDEX(DrawCall), int32(TotalPerCategory[CategoryIndex].Draws), ECsvCustomStatOp::Set);
+		FCsvProfiler::RecordCustomStat(Manager.Array[CategoryIndex]->Name, CSV_CATEGORY_INDEX(DrawCall), int32(TotalPerCategory[CategoryIndex].Draws), ECsvCustomStatOp::Set);
 	}
 	#endif
 #endif // HAS_GPU_STATS
diff --git a/Engine/Source/Runtime/RHI/Public/RHI.h b/Engine/Source/Runtime/RHI/Public/RHI.h
--- a/Engine/Source/Runtime/RHI/Public/RHI.h
+++ b/Engine/Source/Runtime/RHI/Public/RHI.h
@@ -2664,21 +2664,8 @@ extern RHI_API int32 GNumPrimitivesDrawnRHI[MAX_NUM_GPUS];
 
 	struct RHI_API FDrawCallCategoryName
 	{
-		FDrawCallCategoryName()
-			: Name(NAME_None)
-			, Index(-1)
-		{}
-
-		FDrawCallCategoryName(FName InName)
-			: Name(InName)
-			, Index(NumCategory++)
-		{
-			check(Index < MAX_DRAWCALL_CATEGORY);
-			if (Index < MAX_DRAWCALL_CATEGORY)
-			{
-				Array[Index] = this;
-			}
-		}
+		FDrawCallCategoryName();
+		FDrawCallCategoryName(FName InName);
 
 		bool ShouldCountDraws() const { return Index != -1; }
 
@@ -2687,9 +2674,19 @@ extern RHI_API int32 GNumPrimitivesDrawnRHI[MAX_NUM_GPUS];
 
 		static constexpr int32 MAX_DRAWCALL_CATEGORY = 256;
 
-		static TStaticArray<FDrawCallCategoryName*, MAX_DRAWCALL_CATEGORY> Array;
-		static TStaticArray<TStaticArray<int32, MAX_NUM_GPUS>, MAX_DRAWCALL_CATEGORY> DisplayCounts; // A backup of the counts that can be used to display on screen to avoid flickering.
-		static int32 NumCategory;
+		struct FManager
+		{
+			TStaticArray<FDrawCallCategoryName*, MAX_DRAWCALL_CATEGORY> Array;
+
+			// A backup of the counts that can be used to display on screen to avoid flickering.
+			TStaticArray<TStaticArray<int32, MAX_NUM_GPUS>, MAX_DRAWCALL_CATEGORY> DisplayCounts;
+
+			int32 NumCategory;
+
+			FManager();
+		};
+
+		static FManager& GetManager();
 	};
 
 	// RHI counter stats.
