RHI Cleanup
- Adding FRHICOMMAND_UNNAMED to make it easier to add unnamed RHI commands.
- Adding FRHICOMMAND_MACRO_TPL and FRHICOMMAND_UNNAMED_TPL to make it easier to add RHI commands with template arguments.
- Made more public RHI commands use  FRHICOMMAND macros.
- Removing FRHIComputeShader re-declarations of IRHIComputeContext RHISet methods inside of IRHICommandContext and replacing them with using IRHIComputeContext::RHISetXXX statements so we don't duplicate signatures. A comment was also left to explain why they are there.

diff --git a/Engine/Source/Runtime/RHI/Public/RHICommandList.h b/Engine/Source/Runtime/RHI/Public/RHICommandList.h
--- a/Engine/Source/Runtime/RHI/Public/RHICommandList.h
+++ b/Engine/Source/Runtime/RHI/Public/RHICommandList.h
@@ -878,12 +878,27 @@ struct FRHICommand : public FRHICommandBase
 	virtual void StoreDebugInfo(FRHICommandListDebugContext& Context) {};
 };
 
+#define FRHICOMMAND_UNNAMED(CommandName)							\
+	struct CommandName final : public FRHICommand<CommandName, FUnnamedRhiCommand>
+
+#define FRHICOMMAND_UNNAMED_TPL(TemplateType, CommandName)			\
+	template<typename TemplateType>									\
+	struct CommandName final : public FRHICommand<CommandName<TemplateType>, FUnnamedRhiCommand>
+
 #define FRHICOMMAND_MACRO(CommandName)								\
-struct PREPROCESSOR_JOIN(CommandName##String, __LINE__)				\
-{																	\
-	static const TCHAR* TStr() { return TEXT(#CommandName); }		\
-};																	\
-struct CommandName final : public FRHICommand<CommandName, PREPROCESSOR_JOIN(CommandName##String, __LINE__)>
+	struct PREPROCESSOR_JOIN(CommandName##String, __LINE__)			\
+	{																\
+		static const TCHAR* TStr() { return TEXT(#CommandName); }	\
+	};																\
+	struct CommandName final : public FRHICommand<CommandName, PREPROCESSOR_JOIN(CommandName##String, __LINE__)>
+
+#define FRHICOMMAND_MACRO_TPL(TemplateType, CommandName)			\
+	struct PREPROCESSOR_JOIN(CommandName##String, __LINE__)			\
+	{																\
+		static const TCHAR* TStr() { return TEXT(#CommandName); }	\
+	};																\
+	template<typename TemplateType>									\
+	struct CommandName final : public FRHICommand<CommandName<TemplateType>, PREPROCESSOR_JOIN(CommandName##String, __LINE__)>
 
 FRHICOMMAND_MACRO(FRHICommandBeginUpdateMultiFrameResource)
 {
@@ -946,12 +961,7 @@ FRHICOMMAND_MACRO(FRHICommandWaitForTemporalEffect)
 	RHI_API void Execute(FRHICommandListBase& CmdList);
 };
 
-struct FRHICommandBroadcastTemporalEffectString
-{
-	static const TCHAR* TStr() { return TEXT("FRHICommandBroadcastTemporalEffect"); }
-};
-template <typename TRHIResource>
-struct FRHICommandBroadcastTemporalEffect final	: public FRHICommand<FRHICommandBroadcastTemporalEffect<TRHIResource>, FRHICommandBroadcastTemporalEffectString>
+FRHICOMMAND_MACRO_TPL(TRHIResource, FRHICommandBroadcastTemporalEffect)
 {
 	FName EffectName;
 	const TArrayView<TRHIResource*> Resources;
@@ -1012,12 +1022,7 @@ FRHICOMMAND_MACRO(FRHICommandSetStencilRef)
 	RHI_API void Execute(FRHICommandListBase& CmdList);
 };
 
-struct FRHICommandSetShaderParameterString
-{
-	static const TCHAR* TStr() { return TEXT("FRHICommandSetShaderParameter"); }
-};
-template <typename TRHIShader>
-struct FRHICommandSetShaderParameter final : public FRHICommand<FRHICommandSetShaderParameter<TRHIShader>, FRHICommandSetShaderParameterString>
+FRHICOMMAND_MACRO_TPL(TRHIShader, FRHICommandSetShaderParameter)
 {
 	TRHIShader* Shader;
 	const void* NewValue;
@@ -1035,12 +1040,7 @@ struct FRHICommandSetShaderParameter final : public FRHICommand<FRHICommandSetSh
 	RHI_API void Execute(FRHICommandListBase& CmdList);
 };
 
-struct FRHICommandSetShaderUniformBufferString
-{
-	static const TCHAR* TStr() { return TEXT("FRHICommandSetShaderUniformBuffer"); }
-};
-template <typename TRHIShader>
-struct FRHICommandSetShaderUniformBuffer final : public FRHICommand<FRHICommandSetShaderUniformBuffer<TRHIShader>, FRHICommandSetShaderUniformBufferString>
+FRHICOMMAND_MACRO_TPL(TRHIShader, FRHICommandSetShaderUniformBuffer)
 {
 	TRHIShader* Shader;
 	uint32 BaseIndex;
@@ -1054,12 +1054,7 @@ struct FRHICommandSetShaderUniformBuffer final : public FRHICommand<FRHICommandS
 	RHI_API void Execute(FRHICommandListBase& CmdList);
 };
 
-struct FRHICommandSetShaderTextureString
-{
-	static const TCHAR* TStr() { return TEXT("FRHICommandSetShaderTexture"); }
-};
-template <typename TRHIShader>
-struct FRHICommandSetShaderTexture final : public FRHICommand<FRHICommandSetShaderTexture<TRHIShader>, FRHICommandSetShaderTextureString >
+FRHICOMMAND_MACRO_TPL(TRHIShader, FRHICommandSetShaderTexture)
 {
 	TRHIShader* Shader;
 	uint32 TextureIndex;
@@ -1073,12 +1068,7 @@ struct FRHICommandSetShaderTexture final : public FRHICommand<FRHICommandSetShad
 	RHI_API void Execute(FRHICommandListBase& CmdList);
 };
 
-struct FRHICommandSetShaderResourceViewParameterString
-{
-	static const TCHAR* TStr() { return TEXT("FRHICommandSetShaderResourceViewParameter"); }
-};
-template <typename TRHIShader>
-struct FRHICommandSetShaderResourceViewParameter final : public FRHICommand<FRHICommandSetShaderResourceViewParameter<TRHIShader>, FRHICommandSetShaderResourceViewParameterString >
+FRHICOMMAND_MACRO_TPL(TRHIShader, FRHICommandSetShaderResourceViewParameter)
 {
 	TRHIShader* Shader;
 	uint32 SamplerIndex;
@@ -1092,12 +1082,7 @@ struct FRHICommandSetShaderResourceViewParameter final : public FRHICommand<FRHI
 	RHI_API void Execute(FRHICommandListBase& CmdList);
 };
 
-struct FRHICommandSetUAVParameterString
-{
-	static const TCHAR* TStr() { return TEXT("FRHICommandSetUAVParameter"); }
-};
-template <typename TRHIShader>
-struct FRHICommandSetUAVParameter final : public FRHICommand<FRHICommandSetUAVParameter<TRHIShader>, FRHICommandSetUAVParameterString >
+FRHICOMMAND_MACRO_TPL(TRHIShader, FRHICommandSetUAVParameter)
 {
 	TRHIShader* Shader;
 	uint32 UAVIndex;
@@ -1111,11 +1096,7 @@ struct FRHICommandSetUAVParameter final : public FRHICommand<FRHICommandSetUAVPa
 	RHI_API void Execute(FRHICommandListBase& CmdList);
 };
 
-struct FRHICommandSetUAVParameter_InitialCountString
-{
-	static const TCHAR* TStr() { return TEXT("FRHICommandSetUAVParameter_InitialCount"); }
-};
-struct FRHICommandSetUAVParameter_InitialCount final : public FRHICommand<FRHICommandSetUAVParameter_InitialCount, FRHICommandSetUAVParameter_InitialCountString >
+FRHICOMMAND_MACRO(FRHICommandSetUAVParameter_InitialCount)
 {
 	FRHIComputeShader* Shader;
 	uint32 UAVIndex;
@@ -1131,12 +1112,7 @@ struct FRHICommandSetUAVParameter_InitialCount final : public FRHICommand<FRHICo
 	RHI_API void Execute(FRHICommandListBase& CmdList);
 };
 
-struct FRHICommandSetShaderSamplerString
-{
-	static const TCHAR* TStr() { return TEXT("FRHICommandSetShaderSampler"); }
-};
-template <typename TRHIShader>
-struct FRHICommandSetShaderSampler final : public FRHICommand<FRHICommandSetShaderSampler<TRHIShader>, FRHICommandSetShaderSamplerString >
+FRHICOMMAND_MACRO_TPL(TRHIShader, FRHICommandSetShaderSampler)
 {
 	TRHIShader* Shader;
 	uint32 SamplerIndex;
@@ -1308,12 +1284,7 @@ FRHICOMMAND_MACRO(FRHICommandNextSubpass)
 	RHI_API void Execute(FRHICommandListBase& CmdList);
 };
 
-struct FRHICommandSetComputeShaderString
-{
-	static const TCHAR* TStr() { return TEXT("FRHICommandSetComputeShader"); }
-};
-
-struct FRHICommandSetComputeShader final : public FRHICommand<FRHICommandSetComputeShader, FRHICommandSetComputeShaderString>
+FRHICOMMAND_MACRO(FRHICommandSetComputeShader)
 {
 	FRHIComputeShader* ComputeShader;
 	FORCEINLINE_DEBUGGABLE FRHICommandSetComputeShader(FRHIComputeShader* InComputeShader)
@@ -1323,12 +1294,7 @@ struct FRHICommandSetComputeShader final : public FRHICommand<FRHICommandSetComp
 	RHI_API void Execute(FRHICommandListBase& CmdList);
 };
 
-struct FRHICommandSetComputePipelineStateString
-{
-	static const TCHAR* TStr() { return TEXT("FRHICommandSetComputePipelineState"); }
-};
-
-struct FRHICommandSetComputePipelineState final : public FRHICommand<FRHICommandSetComputePipelineState, FRHICommandSetComputePipelineStateString>
+FRHICOMMAND_MACRO(FRHICommandSetComputePipelineState)
 {
 	FComputePipelineState* ComputePipelineState;
 	FORCEINLINE_DEBUGGABLE FRHICommandSetComputePipelineState(FComputePipelineState* InComputePipelineState)
@@ -1368,12 +1334,7 @@ FRHICOMMAND_MACRO(FRHICommandSetGraphicsPipelineStateFromInitializer)
 };
 #endif
 
-struct FRHICommandDispatchComputeShaderString
-{
-	static const TCHAR* TStr() { return TEXT("FRHICommandDispatchComputeShader"); }
-};
-
-struct FRHICommandDispatchComputeShader final : public FRHICommand<FRHICommandDispatchComputeShader, FRHICommandDispatchComputeShaderString>
+FRHICOMMAND_MACRO(FRHICommandDispatchComputeShader)
 {
 	uint32 ThreadGroupCountX;
 	uint32 ThreadGroupCountY;
@@ -1387,12 +1348,7 @@ struct FRHICommandDispatchComputeShader final : public FRHICommand<FRHICommandDi
 	RHI_API void Execute(FRHICommandListBase& CmdList);
 };
 
-struct FRHICommandDispatchIndirectComputeShaderString
-{
-	static const TCHAR* TStr() { return TEXT("FRHICommandDispatchIndirectComputeShader"); }
-};
-
-struct FRHICommandDispatchIndirectComputeShader final : public FRHICommand<FRHICommandDispatchIndirectComputeShader, FRHICommandDispatchIndirectComputeShaderString>
+FRHICOMMAND_MACRO(FRHICommandDispatchIndirectComputeShader)
 {
 	FRHIBuffer* ArgumentBuffer;
 	uint32 ArgumentOffset;
@@ -1645,12 +1601,7 @@ FRHICOMMAND_MACRO(FRHICommandSetTrackedAccess)
 	RHI_API void Execute(FRHICommandListBase& CmdList);
 };
 
-struct FRHICommandSetAsyncComputeBudgetString
-{
-	static const TCHAR* TStr() { return TEXT("FRHICommandSetAsyncComputeBudget"); }
-};
-
-struct FRHICommandSetAsyncComputeBudget final : public FRHICommand<FRHICommandSetAsyncComputeBudget, FRHICommandSetAsyncComputeBudgetString>
+FRHICOMMAND_MACRO(FRHICommandSetAsyncComputeBudget)
 {
 	EAsyncComputeBudget Budget;
 
@@ -1661,12 +1612,7 @@ struct FRHICommandSetAsyncComputeBudget final : public FRHICommand<FRHICommandSe
 	RHI_API void Execute(FRHICommandListBase& CmdList);
 };
 
-struct FRHICommandCopyToStagingBufferString
-{
-	static const TCHAR* TStr() { return TEXT("FRHICommandCopyToStagingBuffer"); }
-};
-
-struct FRHICommandCopyToStagingBuffer final : public FRHICommand<FRHICommandCopyToStagingBuffer, FRHICommandCopyToStagingBufferString>
+FRHICOMMAND_MACRO(FRHICommandCopyToStagingBuffer)
 {
 	FRHIBuffer* SourceBuffer;
 	FRHIStagingBuffer* DestinationStagingBuffer;
@@ -1809,12 +1755,7 @@ FRHICOMMAND_MACRO(FRHICommandBuildLocalUniformBuffer)
 	RHI_API void Execute(FRHICommandListBase& CmdList);
 };
 
-struct FRHICommandSetLocalUniformBufferString
-{
-	static const TCHAR* TStr() { return TEXT("FRHICommandSetLocalUniformBuffer"); }
-};
-template <typename TRHIShader>
-struct FRHICommandSetLocalUniformBuffer final : public FRHICommand<FRHICommandSetLocalUniformBuffer<TRHIShader>, FRHICommandSetLocalUniformBufferString >
+FRHICOMMAND_MACRO_TPL(TRHIShader, FRHICommandSetLocalUniformBuffer)
 {
 	TRHIShader* Shader;
 	uint32 BaseIndex;
@@ -1939,12 +1880,7 @@ FRHICOMMAND_MACRO(FRHICommandEndDrawingViewport)
 	RHI_API void Execute(FRHICommandListBase& CmdList);
 };
 
-struct FRHICommandPushEventString
-{
-	static const TCHAR* TStr() { return TEXT("FRHICommandPushEventString"); }
-};
-
-struct FRHICommandPushEvent final : public FRHICommand<FRHICommandPushEvent, FRHICommandPushEventString>
+FRHICOMMAND_MACRO(FRHICommandPushEvent)
 {
 	const TCHAR *Name;
 	FColor Color;
@@ -1962,12 +1898,7 @@ struct FRHICommandPushEvent final : public FRHICommand<FRHICommandPushEvent, FRH
 	};
 };
 
-struct FRHICommandPopEventString
-{
-	static const TCHAR* TStr() { return TEXT("FRHICommandPopEvent"); }
-};
-
-struct FRHICommandPopEvent final : public FRHICommand<FRHICommandPopEvent, FRHICommandPopEventString>
+FRHICOMMAND_MACRO(FRHICommandPopEvent)
 {
 	RHI_API void Execute(FRHICommandListBase& CmdList);
 
@@ -2097,7 +2028,7 @@ FRHICOMMAND_MACRO(FRHICommandCopyBufferRegion)
 
 #if RHI_RAYTRACING
 
-struct FRHICommandBindAccelerationStructureMemory final : public FRHICommand<FRHICommandBindAccelerationStructureMemory>
+FRHICOMMAND_UNNAMED(FRHICommandBindAccelerationStructureMemory)
 {
 	FRHIRayTracingScene* Scene;
 	FRHIBuffer* Buffer;
@@ -2112,7 +2043,7 @@ struct FRHICommandBindAccelerationStructureMemory final : public FRHICommand<FRH
 	RHI_API void Execute(FRHICommandListBase& CmdList);
 };
 
-struct FRHICommandBuildAccelerationStructure final : public FRHICommand<FRHICommandBuildAccelerationStructure>
+FRHICOMMAND_UNNAMED(FRHICommandBuildAccelerationStructure)
 {
 	FRayTracingSceneBuildParams SceneBuildParams;
 
@@ -2134,7 +2065,7 @@ FRHICOMMAND_MACRO(FRHICommandClearRayTracingBindings)
 	RHI_API void Execute(FRHICommandListBase& CmdList);
 };
 
-struct FRHICommandBuildAccelerationStructures final : public FRHICommand<FRHICommandBuildAccelerationStructures>
+FRHICOMMAND_UNNAMED(FRHICommandBuildAccelerationStructures)
 {
 	const TArrayView<const FRayTracingGeometryBuildParams> Params;
 	FRHIBufferRange ScratchBufferRange;
diff --git a/Engine/Source/Runtime/RHI/Public/RHIContext.h b/Engine/Source/Runtime/RHI/Public/RHIContext.h
--- a/Engine/Source/Runtime/RHI/Public/RHIContext.h
+++ b/Engine/Source/Runtime/RHI/Public/RHIContext.h
@@ -270,7 +270,7 @@ public:
 	virtual void RHIEndUAVOverlap(TArrayView<FRHIUnorderedAccessView* const> UAVs) {}
 
 	/** Set the shader resource view of a surface.  This is used for binding TextureMS parameter types that need a multi sampled view. */
-	virtual void RHISetShaderTexture(FRHIComputeShader* PixelShader, uint32 TextureIndex, FRHITexture* NewTexture) = 0;
+	virtual void RHISetShaderTexture(FRHIComputeShader* ComputeShader, uint32 TextureIndex, FRHITexture* NewTexture) = 0;
 
 	/**
 	* Sets sampler state.
@@ -667,20 +667,18 @@ public:
 	virtual void RHISetGraphicsPipelineState(const FGraphicsPipelineStateInitializer& PsoInit, uint32 StencilRef, bool bApplyAdditionalState) = 0;
 #endif
 
+	// Inherit the parent context's RHISet functions that take FRHIComputeShader arguments
+	// Required to avoid warning C4263 : 'function' : member function does not override any base class virtual member function
+	using IRHIComputeContext::RHISetShaderTexture;
+	using IRHIComputeContext::RHISetShaderSampler;
+	using IRHIComputeContext::RHISetUAVParameter;
+	using IRHIComputeContext::RHISetShaderResourceViewParameter;
+	using IRHIComputeContext::RHISetShaderUniformBuffer;
+	using IRHIComputeContext::RHISetShaderParameter;
+
 	/** Set the shader resource view of a surface. */
 	virtual void RHISetShaderTexture(FRHIGraphicsShader* Shader, uint32 TextureIndex, FRHITexture* NewTexture) = 0;
 
-	/** Set the shader resource view of a surface.  This is used for binding TextureMS parameter types that need a multi sampled view. */
-	virtual void RHISetShaderTexture(FRHIComputeShader* PixelShader, uint32 TextureIndex, FRHITexture* NewTexture) = 0;
-
-	/**
-	* Sets sampler state.
-	* @param ComputeShader		The compute shader to set the sampler for.
-	* @param SamplerIndex		The index of the sampler.
-	* @param NewState			The new sampler state.
-	*/
-	virtual void RHISetShaderSampler(FRHIComputeShader* ComputeShader, uint32 SamplerIndex, FRHISamplerState* NewState) = 0;
-
 	/**
 	* Sets sampler state.
 	* @param Shader				The shader to set the sampler for.
@@ -698,35 +696,12 @@ public:
 	virtual void RHISetUAVParameter(FRHIPixelShader* PixelShader, uint32 UAVIndex, FRHIUnorderedAccessView* UAV) = 0;
 
 
-	/**
-	* Sets a compute shader UAV parameter.
-	* @param ComputeShader	The compute shader to set the UAV for.
-	* @param UAVIndex		The index of the UAVIndex.
-	* @param UAV			The new UAV.
-	*/
-	virtual void RHISetUAVParameter(FRHIComputeShader* ComputeShader, uint32 UAVIndex, FRHIUnorderedAccessView* UAV) = 0;
-
-	/**
-	* Sets a compute shader counted UAV parameter and initial count
-	* @param ComputeShader	The compute shader to set the UAV for.
-	* @param UAVIndex		The index of the UAVIndex.
-	* @param UAV			The new UAV.
-	* @param InitialCount	The initial number of items in the UAV.
-	*/
-	virtual void RHISetUAVParameter(FRHIComputeShader* ComputeShader, uint32 UAVIndex, FRHIUnorderedAccessView* UAV, uint32 InitialCount) = 0;
-
-	virtual void RHISetShaderResourceViewParameter(FRHIComputeShader* ComputeShader, uint32 SamplerIndex, FRHIShaderResourceView* SRV) = 0;
-
 	virtual void RHISetShaderResourceViewParameter(FRHIGraphicsShader* Shader, uint32 SamplerIndex, FRHIShaderResourceView* SRV) = 0;
 
 	virtual void RHISetShaderUniformBuffer(FRHIGraphicsShader* Shader, uint32 BufferIndex, FRHIUniformBuffer* Buffer) = 0;
 
-	virtual void RHISetShaderUniformBuffer(FRHIComputeShader* ComputeShader, uint32 BufferIndex, FRHIUniformBuffer* Buffer) = 0;
-
 	virtual void RHISetShaderParameter(FRHIGraphicsShader* Shader, uint32 BufferIndex, uint32 BaseIndex, uint32 NumBytes, const void* NewValue) = 0;
 
-	virtual void RHISetShaderParameter(FRHIComputeShader* ComputeShader, uint32 BufferIndex, uint32 BaseIndex, uint32 NumBytes, const void* NewValue) = 0;
-
 	virtual void RHISetStencilRef(uint32 StencilRef) {}
 
 	virtual void RHISetBlendFactor(const FLinearColor& BlendFactor) {}
@@ -854,8 +829,8 @@ public:
 		checkNoEntry();
 	}
 
-	protected:
-		FRHIRenderPassInfo RenderPassInfo;
+protected:
+	FRHIRenderPassInfo RenderPassInfo;
 };
 
 
