Suppress some verbose logging related to RTPSO creation when non-blocking mode is used

* Long PSO creation time typically does not cause a frame time spike when background compilation is used
* Log events are still emitted if a PSO is created in a blocking mode (this can happen if non-blocking mode is enabled in some cases, such as when adding a new ray generation shader)

diff --git a/Engine/Source/Runtime/D3D12RHI/Private/D3D12RayTracing.cpp b/Engine/Source/Runtime/D3D12RHI/Private/D3D12RayTracing.cpp
--- a/Engine/Source/Runtime/D3D12RHI/Private/D3D12RayTracing.cpp
+++ b/Engine/Source/Runtime/D3D12RHI/Private/D3D12RayTracing.cpp
@@ -2765,7 +2765,8 @@ public:
 
 		const double TotalCreationTimeMS = 1000.0 * FPlatformTime::ToSeconds64(TotalCreationTime);
 		const float CreationTimeWarningThresholdMS = 10.0f;
-		if (TotalCreationTimeMS > CreationTimeWarningThresholdMS)
+		const bool bAllowLogSlowCreation = !Initializer.bBackgroundCompilation; // Only report creation stalls on the critical path
+		if (bAllowLogSlowCreation && TotalCreationTimeMS > CreationTimeWarningThresholdMS)
 		{
 			const double CompileTimeMS = 1000.0 * FPlatformTime::ToSeconds64(CompileTime);
 			const double LinkTimeMS = 1000.0 * FPlatformTime::ToSeconds64(LinkTime);
diff --git a/Engine/Source/Runtime/RHI/Private/PipelineFileCache.cpp b/Engine/Source/Runtime/RHI/Private/PipelineFileCache.cpp
--- a/Engine/Source/Runtime/RHI/Private/PipelineFileCache.cpp
+++ b/Engine/Source/Runtime/RHI/Private/PipelineFileCache.cpp
@@ -3495,7 +3495,7 @@ void FPipelineFileCacheManager::CacheComputePSO(uint32 RunTimeHash, FRHIComputeS
 	}
 }
 
-void FPipelineFileCacheManager::CacheRayTracingPSO(const FRayTracingPipelineStateInitializer& Initializer)
+void FPipelineFileCacheManager::CacheRayTracingPSO(const FRayTracingPipelineStateInitializer& Initializer, ERayTracingPipelineCacheFlags Flags)
 {
 	if (!IsPipelineFileCacheEnabled() || !(LogPSOtoFileCache() || ReportNewPSOs()))
 	{
@@ -3510,6 +3510,9 @@ void FPipelineFileCacheManager::CacheRayTracingPSO(const FRayTracingPipelineStat
 		Initializer.GetCallableTable()
 	};
 
+	// When non-blocking creation is used, encountering a non-cached RTPSO is not likely to cause a hitch and so the logging is not useful/actionable.
+	const bool bShouldReportNewRTPSOToLog = !EnumHasAnyFlags(Flags, ERayTracingPipelineCacheFlags::NonBlocking);
+
 	FRWScopeLock Lock(FileCacheLock, SLT_ReadOnly);
 
 	for (TArrayView<FRHIRayTracingShader*>& Table : ShaderTables)
@@ -3536,11 +3539,15 @@ void FPipelineFileCacheManager::CacheRayTracingPSO(const FRayTracingPipelineStat
 					if (!FPipelineFileCacheManager::IsPSOEntryCached(NewEntry, &CurrentUsageData))
 					{
 						//CSV_EVENT(PSO, TEXT("Encountered new ray tracing PSO"));
-						UE_LOG(LogRHI, Display, TEXT("Encountered a new ray tracing PSO: %u"), PSOHash);
-						if (GPSOFileCachePrintNewPSODescriptors > 0)
+						if (bShouldReportNewRTPSOToLog)
 						{
-							UE_LOG(LogRHI, Display, TEXT("New ray tracing PSO (%u) Description: %s"), PSOHash, *NewEntry.RayTracingDesc.ToString());
+							UE_LOG(LogRHI, Display, TEXT("Encountered a new ray tracing PSO: %u"), PSOHash);
+							if (GPSOFileCachePrintNewPSODescriptors > 0)
+							{
+								UE_LOG(LogRHI, Display, TEXT("New ray tracing PSO (%u) Description: %s"), PSOHash, *NewEntry.RayTracingDesc.ToString());
+							}
 						}
+
 						if (LogPSOtoFileCache())
 						{
 							NewPSOs.Add(NewEntry);
diff --git a/Engine/Source/Runtime/RHI/Private/PipelineStateCache.cpp b/Engine/Source/Runtime/RHI/Private/PipelineStateCache.cpp
--- a/Engine/Source/Runtime/RHI/Private/PipelineStateCache.cpp
+++ b/Engine/Source/Runtime/RHI/Private/PipelineStateCache.cpp
@@ -1722,6 +1722,8 @@ public:
 		, Initializer(InInitializer)
 		, bBackgroundTask(bInBackgroundTask)
 	{
+		Initializer.bBackgroundCompilation = bBackgroundTask;
+
 		// Copy all referenced shaders and AddRef them while the task is alive
 
 		RayGenTable   = CopyShaderTable(InInitializer.GetRayGenTable());
@@ -1887,7 +1889,7 @@ FRayTracingPipelineState* PipelineStateCache::GetAndOrCreateRayTracingPipelineSt
 	}
 	else
 	{
-		FPipelineFileCacheManager::CacheRayTracingPSO(InInitializer);
+		FPipelineFileCacheManager::CacheRayTracingPSO(InInitializer, Flags);
 
 		// Copy the initializer as we may want to patch it below
 		FRayTracingPipelineStateInitializer Initializer = InInitializer;
diff --git a/Engine/Source/Runtime/RHI/Public/PipelineFileCache.h b/Engine/Source/Runtime/RHI/Public/PipelineFileCache.h
--- a/Engine/Source/Runtime/RHI/Public/PipelineFileCache.h
+++ b/Engine/Source/Runtime/RHI/Public/PipelineFileCache.h
@@ -82,6 +82,7 @@ struct RHI_API FPipelineFileCacheRasterizerState
 
 class FRayTracingPipelineStateInitializer;
 class FRHIRayTracingShader;
+enum class ERayTracingPipelineCacheFlags : uint8;
 
 /**
  * Tracks stats for the current session between opening & closing the file-cache.
@@ -356,7 +357,7 @@ public:
 
 	static void CacheGraphicsPSO(uint32 RunTimeHash, FGraphicsPipelineStateInitializer const& Initializer);
 	static void CacheComputePSO(uint32 RunTimeHash, FRHIComputeShader const* Initializer);
-	static void CacheRayTracingPSO(const FRayTracingPipelineStateInitializer& Initializer);
+	static void CacheRayTracingPSO(const FRayTracingPipelineStateInitializer& Initializer, ERayTracingPipelineCacheFlags Flags);
 
 	// true if the named PSOFC is currently open.
 	static bool HasPipelineFileCache(const FString& PSOCacheKey);
diff --git a/Engine/Source/Runtime/RHI/Public/PipelineStateCache.h b/Engine/Source/Runtime/RHI/Public/PipelineStateCache.h
--- a/Engine/Source/Runtime/RHI/Public/PipelineStateCache.h
+++ b/Engine/Source/Runtime/RHI/Public/PipelineStateCache.h
@@ -24,7 +24,7 @@ enum class EApplyRendertargetOption : int
 
 ENUM_CLASS_FLAGS(EApplyRendertargetOption);
 
-enum class ERayTracingPipelineCacheFlags : int
+enum class ERayTracingPipelineCacheFlags : uint8
 {
 	// Query the pipeline cache, create pipeline if necessary.
 	// Compilation may happen on a task, but RHIThread will block on it before translating the RHICmdList.
diff --git a/Engine/Source/Runtime/RHI/Public/RHIResources.h b/Engine/Source/Runtime/RHI/Public/RHIResources.h
--- a/Engine/Source/Runtime/RHI/Public/RHIResources.h
+++ b/Engine/Source/Runtime/RHI/Public/RHIResources.h
@@ -3486,6 +3486,12 @@ public:
 	// at least one shader must be present in total (completely empty pipelines are not allowed).
 	bool bPartial = false;
 
+	// Hints to the RHI that this PSO is being compiled by a background task and will not be needed immediately for rendering.
+	// Speculative PSO pre-caching or non-blocking PSO creation should set this flag.
+	// This may be used by the RHI to decide if a hitch warning should be reported, change priority of any internally dispatched tasks, etc.
+	// Does not affect the creation of the PSO itself.
+	bool bBackgroundCompilation = false;
+
 	// Ray tracing pipeline may be created by deriving from the existing base.
 	// Base pipeline will be extended by adding new shaders into it, potentially saving substantial amount of CPU time.
 	// Depends on GRHISupportsRayTracingPSOAdditions support at runtime (base pipeline is simply ignored if it is unsupported).
