Fixed race condition in RHI validation logic. The op list was being finalized after pushing the context onto the 'available' queue, meaning another thread could allocate it and start allocating ops.

diff --git a/Engine/Source/Runtime/RHI/Private/RHIValidation.cpp b/Engine/Source/Runtime/RHI/Private/RHIValidation.cpp
--- a/Engine/Source/Runtime/RHI/Private/RHIValidation.cpp
+++ b/Engine/Source/Runtime/RHI/Private/RHIValidation.cpp
@@ -186,8 +186,9 @@ IRHIPlatformCommandList* FValidationRHI::RHIFinalizeContext(IRHIComputeContext*
 
 	FValidationCommandList* OuterCommandList = new FValidationCommandList();
 
-	OuterCommandList->InnerCommandList = RHI->RHIFinalizeContext(&InnerContext);
+	// RHIFinalizeContext makes the context available to other threads, so finalize the tracker beforehand.
 	OuterCommandList->CompletedOpList = InnerContext.Tracker->Finalize();
+	OuterCommandList->InnerCommandList = RHI->RHIFinalizeContext(&InnerContext);
 	OuterCommandList->Pipeline = OuterContext->GetPipeline();
 
 	switch (OuterCommandList->Pipeline)
