Fix some spelling and unit errors.

diff --git a/Engine/Source/Runtime/RHI/Private/TextureProfiler.cpp b/Engine/Source/Runtime/RHI/Private/TextureProfiler.cpp
--- a/Engine/Source/Runtime/RHI/Private/TextureProfiler.cpp
+++ b/Engine/Source/Runtime/RHI/Private/TextureProfiler.cpp
@@ -19,6 +19,15 @@ static int32 ToMB(uint64 Value)
 	return Value32;
 }
 
+static int32 ToKB(uint64 Value)
+{
+	const size_t KB = 1024;
+	uint64 Value64 = (Value + (KB - 1)) / KB;
+	int32 Value32 = static_cast<int32>(Value64);
+	check(Value64 == Value32);
+	return Value32;
+}
+
 CSV_DEFINE_CATEGORY(RenderTargetProfiler, true);
 CSV_DEFINE_CATEGORY(RenderTargetWasteProfiler, true);
 CSV_DEFINE_CATEGORY(TextureProfiler, true);
@@ -76,7 +85,7 @@ static FAutoConsoleCommand CmdTextureProfilerDumpTextures(
 
 FTextureProfiler* FTextureProfiler::Instance = nullptr;
 
-FTextureProfiler::FTexureDetails::FTexureDetails(FRHITexture* Texture, size_t InSize, uint32 InAlign, size_t InAllocationWaste)
+FTextureProfiler::FTextureDetails::FTextureDetails(FRHITexture* Texture, size_t InSize, uint32 InAlign, size_t InAllocationWaste)
 	: Size(InSize)
 	, PeakSize(0)
 	, Align(InAlign)
@@ -87,7 +96,7 @@ FTextureProfiler::FTexureDetails::FTexureDetails(FRHITexture* Texture, size_t In
 	SetName(TextureName);
 }
 
-void FTextureProfiler::FTexureDetails::SetName(FName InTextureName)
+void FTextureProfiler::FTextureDetails::SetName(FName InTextureName)
 {
 	TextureName = InTextureName;
 	TextureName.SetNumber(0);
@@ -98,18 +107,18 @@ void FTextureProfiler::FTexureDetails::SetName(FName InTextureName)
 	TextureNameString = FTextureProfiler::Get()->GetTextureNameString(InTextureName);
 }
 
-void FTextureProfiler::FTexureDetails::ResetPeakSize()
+void FTextureProfiler::FTextureDetails::ResetPeakSize()
 {
 	PeakSize = Size;
 }
 
-FTextureProfiler::FTexureDetails::~FTexureDetails()
+FTextureProfiler::FTextureDetails::~FTextureDetails()
 {
 	// Do not delete the texture name string, it is not owned by this object
 	//delete TextureNameString;
 }
 
-FTextureProfiler::FTexureDetails& FTextureProfiler::FTexureDetails::operator+=(const FTextureProfiler::FTexureDetails& Other)
+FTextureProfiler::FTextureDetails& FTextureProfiler::FTextureDetails::operator+=(const FTextureProfiler::FTextureDetails& Other)
 {
 	Size += Other.Size;
 	PeakSize = FMath::Max(PeakSize, Size);
@@ -119,7 +128,7 @@ FTextureProfiler::FTexureDetails& FTextureProfiler::FTexureDetails::operator+=(c
 	return *this;
 }
 
-FTextureProfiler::FTexureDetails& FTextureProfiler::FTexureDetails::operator-=(const FTextureProfiler::FTexureDetails& Other)
+FTextureProfiler::FTextureDetails& FTextureProfiler::FTextureDetails::operator-=(const FTextureProfiler::FTextureDetails& Other)
 {
 	Size -= Other.Size;
 	AllocationWaste -= Other.AllocationWaste;
@@ -158,7 +167,7 @@ void FTextureProfiler::AddTextureAllocation(FRHITexture* UniqueTexturePtr, size_
 {
 	FScopeLock Lock(&TextureMapCS);
 
-	FTexureDetails AddedDetails(
+	FTextureDetails AddedDetails(
 		UniqueTexturePtr,
 		Size,
 		Alignment,
@@ -166,13 +175,13 @@ void FTextureProfiler::AddTextureAllocation(FRHITexture* UniqueTexturePtr, size_
 
 	AddedDetails.Count = 1;
 
-	FTexureDetails& TotalValue = AddedDetails.IsRenderTarget ? TotalRenderTargetSize : TotalTextureSize;
+	FTextureDetails& TotalValue = AddedDetails.IsRenderTarget ? TotalRenderTargetSize : TotalTextureSize;
 	TotalValue += AddedDetails;
 
 	// Make sure we have a sane value
 	check(ToMB(TotalValue.Size) >= 0);
 
-	FTexureDetails& NamedValue = AddedDetails.IsRenderTarget ? CombinedRenderTargetSizes.FindOrAdd(AddedDetails.GetTextureName()) : CombinedTextureSizes.FindOrAdd(AddedDetails.GetTextureName());
+	FTextureDetails& NamedValue = AddedDetails.IsRenderTarget ? CombinedRenderTargetSizes.FindOrAdd(AddedDetails.GetTextureName()) : CombinedTextureSizes.FindOrAdd(AddedDetails.GetTextureName());
 	NamedValue += AddedDetails;
 
 	if (NamedValue.GetTextureName().IsNone() || !NamedValue.GetTextureName().IsValid())
@@ -200,21 +209,21 @@ void FTextureProfiler::UpdateTextureAllocation(FRHITexture* UniqueTexturePtr, si
 {
 	FScopeLock Lock(&TextureMapCS);
 
-	FTexureDetails UpdatedDetails(
+	FTextureDetails UpdatedDetails(
 		UniqueTexturePtr,
 		Size,
 		Alignment,
 		AllocationWaste);
 
-	FTexureDetails& ExistingTexture = TexturesMap[UniqueTexturePtr];
-	FTexureDetails& TotalValue = UpdatedDetails.IsRenderTarget ? TotalRenderTargetSize : TotalTextureSize;
+	FTextureDetails& ExistingTexture = TexturesMap[UniqueTexturePtr];
+	FTextureDetails& TotalValue = UpdatedDetails.IsRenderTarget ? TotalRenderTargetSize : TotalTextureSize;
 
 	TotalValue -= ExistingTexture;
 	TotalValue += UpdatedDetails;
 
 	check(ToMB(TotalValue.Size) >= 0);
 
-	FTexureDetails& NamedValue = UpdatedDetails.IsRenderTarget ? CombinedRenderTargetSizes.FindOrAdd(UpdatedDetails.GetTextureName()) : CombinedTextureSizes.FindOrAdd(UpdatedDetails.GetTextureName());
+	FTextureDetails& NamedValue = UpdatedDetails.IsRenderTarget ? CombinedRenderTargetSizes.FindOrAdd(UpdatedDetails.GetTextureName()) : CombinedTextureSizes.FindOrAdd(UpdatedDetails.GetTextureName());
 	NamedValue -= ExistingTexture;
 	NamedValue += UpdatedDetails;
 
@@ -224,12 +233,12 @@ void FTextureProfiler::UpdateTextureAllocation(FRHITexture* UniqueTexturePtr, si
 void FTextureProfiler::RemoveTextureAllocation(FRHITexture* UniqueTexturePtr)
 {
 	FScopeLock Lock(&TextureMapCS);
-	FTexureDetails Details = TexturesMap.FindAndRemoveChecked(UniqueTexturePtr);
+	FTextureDetails Details = TexturesMap.FindAndRemoveChecked(UniqueTexturePtr);
 	
-	FTexureDetails& TotalValue = Details.IsRenderTarget ? TotalRenderTargetSize : TotalTextureSize;
+	FTextureDetails& TotalValue = Details.IsRenderTarget ? TotalRenderTargetSize : TotalTextureSize;
 	TotalValue -= Details;
 
-	FTexureDetails& NamedValue = Details.IsRenderTarget ? CombinedRenderTargetSizes[Details.GetTextureName()] : CombinedTextureSizes[Details.GetTextureName()];
+	FTextureDetails& NamedValue = Details.IsRenderTarget ? CombinedRenderTargetSizes[Details.GetTextureName()] : CombinedTextureSizes[Details.GetTextureName()];
 	NamedValue -= Details;
 }
 
@@ -237,18 +246,18 @@ void FTextureProfiler::UpdateTextureName(FRHITexture* UniqueTexturePtr)
 {
 	FScopeLock Lock(&TextureMapCS);
 
-	FTexureDetails* DetailsPtr = TexturesMap.Find(UniqueTexturePtr);
+	FTextureDetails* DetailsPtr = TexturesMap.Find(UniqueTexturePtr);
 	if (DetailsPtr == nullptr)
 	{
 		return;
 	}
 
-	FTexureDetails& Details = *DetailsPtr;
+	FTextureDetails& Details = *DetailsPtr;
 	FName OldName = Details.GetTextureName();
 	Details.SetName(UniqueTexturePtr->GetName());
 
-	FTexureDetails& NewCombinedValue = Details.IsRenderTarget ? CombinedRenderTargetSizes.FindOrAdd(Details.GetTextureName()) : CombinedTextureSizes.FindOrAdd(Details.GetTextureName());
-	FTexureDetails& OldCombinedValue = Details.IsRenderTarget ? CombinedRenderTargetSizes[OldName] : CombinedTextureSizes[OldName];
+	FTextureDetails& NewCombinedValue = Details.IsRenderTarget ? CombinedRenderTargetSizes.FindOrAdd(Details.GetTextureName()) : CombinedTextureSizes.FindOrAdd(Details.GetTextureName());
+	FTextureDetails& OldCombinedValue = Details.IsRenderTarget ? CombinedRenderTargetSizes[OldName] : CombinedTextureSizes[OldName];
 
 	if (NewCombinedValue.GetTextureName().IsNone() || !NewCombinedValue.GetTextureName().IsValid())
 	{
@@ -329,7 +338,7 @@ void FTextureProfiler::Update()
 		};
 	};
 	
-	FTexureDetails OtherRenderTargetSizes;
+	FTextureDetails OtherRenderTargetSizes;
 	for (auto& Pair : CombinedRenderTargetSizes)
 	{
 		// Don't report empty sizes or sizes less than the minimum
@@ -345,7 +354,7 @@ void FTextureProfiler::Update()
 		Pair.Value.ResetPeakSize();
 	}
 
-	FTexureDetails OtherTextureSizes;
+	FTextureDetails OtherTextureSizes;
 	for (auto& Pair : CombinedTextureSizes)
 	{
 		// Don't report empty sizes or sizes less than the minimum
@@ -394,53 +403,82 @@ void FTextureProfiler::DumpTextures(bool RenderTargets, bool CombineTextureNames
 		}
 	}
 
-	FTexureDetails& TotalValue = RenderTargets ? TotalRenderTargetSize : TotalTextureSize;
+	FTextureDetails& TotalValue = RenderTargets ? TotalRenderTargetSize : TotalTextureSize;
 
 	if (CombineTextureNames)
 	{
-		OutputDevice.Logf(TEXT("%s%c%s%c%s%c%s"), TEXT("Texture"), Sep, TEXT("Combined Size(KB)"), Sep, TEXT("Count"), Sep, TEXT("Combined Wasted(KB)"));
+		OutputDevice.Logf(TEXT("%s%c%s%c%s%c%s"), TEXT("Texture"), Sep, TEXT("Combined Size(MB)"), Sep, TEXT("Count"), Sep, TEXT("Combined Wasted(KB)"));
 
 		size_t MinTextureSize = CVarTextureProfilerMinTextureSizeMB.GetValueOnAnyThread() * 1024 * 1024;
 		size_t MinRenderTargetSize = CVarTextureProfilerMinRenderTargetSizeMB.GetValueOnAnyThread() * 1024 * 1024;
 
 		size_t MinSize = RenderTargets ? MinRenderTargetSize : MinTextureSize;
 
-		TMap<FName, FTexureDetails>& CombinedSizes = RenderTargets ? CombinedRenderTargetSizes : CombinedTextureSizes;
-		FTexureDetails OtherTextureSizes;
+		const TMap<FName, FTextureDetails>& CombinedSizes = RenderTargets ? CombinedRenderTargetSizes : CombinedTextureSizes;
+
+		TArray<const FTextureDetails*> CombinedSizesSorted;
+		CombinedSizesSorted.Reserve(CombinedSizes.Num());
+
+		using FSortTextureDetails = FTextureDetails;
 		for (const auto& Pair : CombinedSizes)
+		{
+			int32 InsertIndex = Algo::LowerBound(CombinedSizesSorted, &Pair.Value,
+				[](const FTextureDetails* Value1, const FTextureDetails* Value2)
+				{
+					return Value1->Size >= Value2->Size;
+				});
+			CombinedSizesSorted.Insert(&Pair.Value, InsertIndex);
+		}
+
+		FTextureDetails OtherTextureSizes;
+		for (const FTextureDetails* TextureDetails : CombinedSizesSorted)
 		{
 			// Don't report empty sizes or sizes less than the minimum
-			if (Pair.Value.Size == 0 || Pair.Value.Size < MinSize)
+			if (TextureDetails->Size == 0 || TextureDetails->Size < MinSize)
 			{
-				OtherTextureSizes += Pair.Value;
+				OtherTextureSizes += *TextureDetails;
 				continue;
 			}
 
-			FString NameString = Pair.Value.GetTextureName().ToString();
-			OutputDevice.Logf(TEXT("%s%c%d%c%d%c%d"), *NameString, Sep, ToMB(Pair.Value.Size), Sep, Pair.Value.Count, Sep, ToMB(Pair.Value.AllocationWaste));
+			FString NameString = TextureDetails->GetTextureName().ToString();
+			OutputDevice.Logf(TEXT("%s%c%d%c%d%c%d"), *NameString, Sep, ToMB(TextureDetails->Size), Sep, TextureDetails->Count, Sep, ToKB(TextureDetails->AllocationWaste));
 		}
 
 		if (OtherTextureSizes.Size > 0)
 		{
-			OutputDevice.Logf(TEXT("%s%c%d%c%d%c%d"), TEXT("Other"), Sep, ToMB(OtherTextureSizes.Size), Sep, OtherTextureSizes.Count, Sep, ToMB(OtherTextureSizes.AllocationWaste));
+			OutputDevice.Logf(TEXT("%s%c%d%c%d%c%d"), TEXT("Other"), Sep, ToMB(OtherTextureSizes.Size), Sep, OtherTextureSizes.Count, Sep, ToKB(OtherTextureSizes.AllocationWaste));
 		}
-		OutputDevice.Logf(TEXT("%s%c%d%c%d%c%d"), TEXT("Total"), Sep, ToMB(TotalValue.Size), Sep, TotalValue.Count, Sep, ToMB(TotalValue.AllocationWaste));
+		OutputDevice.Logf(TEXT("%s%c%d%c%d%c%d"), TEXT("Total"), Sep, ToMB(TotalValue.Size), Sep, TotalValue.Count, Sep, ToKB(TotalValue.AllocationWaste));
 	}
 	else
 	{
-		OutputDevice.Logf(TEXT("%s%c%s%c%s%c%s"), TEXT("Texture"), Sep, TEXT("Size(KB)"), Sep, TEXT("Wasted(KB)"), Sep, TEXT("Alignment"));
+		TArray<const FTextureDetails*> TexuteSizesSorted;
+		TexuteSizesSorted.Reserve(TexturesMap.Num());
+
+		using FSortTextureDetails = FTextureDetails;
 		for (const auto& Pair : TexturesMap)
 		{
-			if (Pair.Value.IsRenderTarget != RenderTargets)
+			int32 InsertIndex = Algo::LowerBound(TexuteSizesSorted, &Pair.Value,
+				[](const FTextureDetails* Value1, const FTextureDetails* Value2)
+				{
+					return Value1->Size >= Value2->Size;
+				});
+			TexuteSizesSorted.Insert(&Pair.Value, InsertIndex);
+		}
+
+		OutputDevice.Logf(TEXT("%s%c%s%c%s%c%s"), TEXT("Texture"), Sep, TEXT("Size(MB)"), Sep, TEXT("Wasted(KB)"), Sep, TEXT("Alignment"));
+		for (const FTextureDetails* TextureDetails : TexuteSizesSorted)
+		{
+			if (TextureDetails->IsRenderTarget != RenderTargets)
 			{
 				continue;
 			}
 
-			FString NameString = Pair.Value.GetTextureName().ToString();
-			OutputDevice.Logf(TEXT("%s%c%d%c%d%c%p"), *NameString, Sep, ToMB(Pair.Value.Size), Sep, ToMB(Pair.Value.AllocationWaste), Sep, Pair.Value.Align);
+			FString NameString = TextureDetails->GetTextureName().ToString();
+			OutputDevice.Logf(TEXT("%s%c%d%c%d%c%p"), *NameString, Sep, ToMB(TextureDetails->Size), Sep, ToKB(TextureDetails->AllocationWaste), Sep, TextureDetails->Align);
 		}
 
-		OutputDevice.Logf(TEXT("%s%c%d%c%d%c%p"), TEXT("Total"), Sep, ToMB(TotalValue.Size), Sep, ToMB(TotalValue.AllocationWaste), Sep, TotalValue.Align);
+		OutputDevice.Logf(TEXT("%s%c%d%c%d%c%p"), TEXT("Total"), Sep, ToMB(TotalValue.Size), Sep, ToKB(TotalValue.AllocationWaste), Sep, TotalValue.Align);
 	}
 }
 
diff --git a/Engine/Source/Runtime/RHI/Public/TextureProfiler.h b/Engine/Source/Runtime/RHI/Public/TextureProfiler.h
--- a/Engine/Source/Runtime/RHI/Public/TextureProfiler.h
+++ b/Engine/Source/Runtime/RHI/Public/TextureProfiler.h
@@ -46,23 +46,23 @@ private:
 
 	FCriticalSection TextureMapCS;
 
-	class FTexureDetails
+	class FTextureDetails
 	{
 	public:
-		FTexureDetails() = default;
-		FTexureDetails(FRHITexture* Texture, size_t InSize, uint32 InAlign, size_t InAllocationWaste);
-		~FTexureDetails();
+		FTextureDetails() = default;
+		FTextureDetails(FRHITexture* Texture, size_t InSize, uint32 InAlign, size_t InAllocationWaste);
+		~FTextureDetails();
 		
 		void SetName(FName InTextureName);
 
 		void ResetPeakSize();
-		void SetValues(const FTexureDetails& Values);
+		void SetValues(const FTextureDetails& Values);
 		FName GetTextureName() const { return TextureName; }
 		const char* GetTextureNameString() const { check(TextureNameString != nullptr);  return TextureNameString; }
 
-		FTexureDetails& operator+=(const FTexureDetails& Other);
-		FTexureDetails& operator-=(const FTexureDetails& Other);
-		FTexureDetails& operator=(const FTexureDetails& Other) = default;
+		FTextureDetails& operator+=(const FTextureDetails& Other);
+		FTextureDetails& operator-=(const FTextureDetails& Other);
+		FTextureDetails& operator=(const FTextureDetails& Other) = default;
 
 		size_t Size = 0;
 		size_t PeakSize = 0;
@@ -80,13 +80,13 @@ private:
 
 	// Allocated so that a reference to the string can be stored in FTextureDetails
 	TMap<FName, const char*> TextureNameStrings;
-	TMap<void*, FTexureDetails> TexturesMap;
+	TMap<void*, FTextureDetails> TexturesMap;
 
 	// Keep track of the totals separately to reduce the cost of rounding error for sizes
-	FTexureDetails TotalTextureSize;
-	FTexureDetails TotalRenderTargetSize;
-	TMap<FName, FTexureDetails> CombinedTextureSizes;
-	TMap<FName, FTexureDetails> CombinedRenderTargetSizes;
-	friend class FTexureDetails;
+	FTextureDetails TotalTextureSize;
+	FTextureDetails TotalRenderTargetSize;
+	TMap<FName, FTextureDetails> CombinedTextureSizes;
+	TMap<FName, FTextureDetails> CombinedRenderTargetSizes;
+	friend class FTextureDetails;
 };
 #endif //TEXTURE_PROFILER_ENABLED
