Only compare the relevant bits of render target flags when comparing FGraphicsPipelineStateInitializer-s

* Fixes PSO cache lookup when the only difference is some miscellaneous flag, such as TexCreate_FastVRAM, TexCreate_Transient, etc.
* Fixes pipeline sort key calculation in BuildNaniteMaterialPassCommands()

diff --git a/Engine/Source/Runtime/RHI/Private/DynamicRHI.cpp b/Engine/Source/Runtime/RHI/Private/DynamicRHI.cpp
--- a/Engine/Source/Runtime/RHI/Private/DynamicRHI.cpp
+++ b/Engine/Source/Runtime/RHI/Private/DynamicRHI.cpp
@@ -599,6 +599,11 @@ uint64 FDynamicRHI::RHIComputePrecachePSOHash(const FGraphicsPipelineStateInitia
 	HashKey.bHasFragmentDensityAttachment = Initializer.bHasFragmentDensityAttachment;
 	HashKey.ShadingRate					= Initializer.ShadingRate;
 
+	for (ETextureCreateFlags& Flags : HashKey.RenderTargetFlags)
+	{
+		Flags = Flags & FGraphicsPipelineStateInitializer::RelevantRenderTargetFlagMask;
+	}
+
 	return CityHash64((const char*)&HashKey, sizeof(FNonStateHashKey));
 }
 
@@ -613,7 +618,7 @@ bool FDynamicRHI::RHIMatchPrecachePSOInitializers(const FGraphicsPipelineStateIn
 		LHS.bHasFragmentDensityAttachment != RHS.bHasFragmentDensityAttachment ||
 		LHS.RenderTargetsEnabled != RHS.RenderTargetsEnabled ||
 		LHS.RenderTargetFormats != RHS.RenderTargetFormats ||
-		LHS.RenderTargetFlags != RHS.RenderTargetFlags ||
+		!FGraphicsPipelineStateInitializer::RelevantRenderTargetFlagsEqual(LHS.RenderTargetFlags, RHS.RenderTargetFlags) ||
 		LHS.DepthStencilTargetFormat != RHS.DepthStencilTargetFormat ||
 		LHS.DepthStencilTargetFlag != RHS.DepthStencilTargetFlag ||
 		LHS.DepthTargetLoadAction != RHS.DepthTargetLoadAction ||
diff --git a/Engine/Source/Runtime/RHI/Private/PipelineFileCache.cpp b/Engine/Source/Runtime/RHI/Private/PipelineFileCache.cpp
--- a/Engine/Source/Runtime/RHI/Private/PipelineFileCache.cpp
+++ b/Engine/Source/Runtime/RHI/Private/PipelineFileCache.cpp
@@ -758,7 +758,7 @@ ETextureCreateFlags FPipelineCacheFileFormatPSO::GraphicsDescriptor::ReduceRTFla
 {
 	// We care about flags that influence RT formats (which is the only thing the underlying API cares about).
 	// In most RHIs, the format is only influenced by TexCreate_SRGB. D3D12 additionally uses TexCreate_Shared in its format selection logic.
-	return (InFlags & (TexCreate_SRGB | TexCreate_Shared));
+	return (InFlags & FGraphicsPipelineStateInitializer::RelevantRenderTargetFlagMask);
 }
 
 FString FPipelineCacheFileFormatPSO::GraphicsDescriptor::StateHeaderLine()
diff --git a/Engine/Source/Runtime/RHI/Public/RHIResources.h b/Engine/Source/Runtime/RHI/Public/RHIResources.h
--- a/Engine/Source/Runtime/RHI/Public/RHIResources.h
+++ b/Engine/Source/Runtime/RHI/Public/RHIResources.h
@@ -3348,7 +3348,7 @@ public:
 			bHasFragmentDensityAttachment != rhs.bHasFragmentDensityAttachment ||
 			RenderTargetsEnabled != rhs.RenderTargetsEnabled ||
 			RenderTargetFormats != rhs.RenderTargetFormats || 
-			RenderTargetFlags != rhs.RenderTargetFlags || 
+			!RelevantRenderTargetFlagsEqual(RenderTargetFlags, rhs.RenderTargetFlags) || 
 			DepthStencilTargetFormat != rhs.DepthStencilTargetFormat || 
 			DepthStencilTargetFlag != rhs.DepthStencilTargetFlag ||
 			DepthTargetLoadAction != rhs.DepthTargetLoadAction ||
@@ -3367,6 +3367,24 @@ public:
 		return true;
 	}
 
+	// We care about flags that influence RT formats (which is the only thing the underlying API cares about).
+	// In most RHIs, the format is only influenced by TexCreate_SRGB. D3D12 additionally uses TexCreate_Shared in its format selection logic.
+	static constexpr ETextureCreateFlags RelevantRenderTargetFlagMask = ETextureCreateFlags::SRGB | ETextureCreateFlags::Shared;
+
+	static bool RelevantRenderTargetFlagsEqual(const TRenderTargetFlags& A, const TRenderTargetFlags& B)
+	{
+		for (int32 Index = 0; Index < A.Num(); ++Index)
+		{
+			ETextureCreateFlags FlagsA = A[Index] & RelevantRenderTargetFlagMask;
+			ETextureCreateFlags FlagsB = B[Index] & RelevantRenderTargetFlagMask;
+			if (FlagsA != FlagsB)
+			{
+				return false;
+			}
+		}
+		return true;
+	}
+
 	uint32 ComputeNumValidRenderTargets() const
 	{
 		// Get the count of valid render targets (ignore those at the end of the array with PF_Unknown)
