Change Atomic Queue of PendingDeletes to a simpler Algorithm

diff --git a/Engine/Source/Runtime/RHI/Private/RHI.cpp b/Engine/Source/Runtime/RHI/Private/RHI.cpp
--- a/Engine/Source/Runtime/RHI/Private/RHI.cpp
+++ b/Engine/Source/Runtime/RHI/Private/RHI.cpp
@@ -294,8 +294,7 @@ const FClearValueBinding FClearValueBinding::Green(FLinearColor(0.0f, 1.0f, 0.0f
 // Note: this is used as the default normal for DBuffer decals.  It must decode to a value of 0 in DecodeDBufferData.
 const FClearValueBinding FClearValueBinding::DefaultNormal8Bit(FLinearColor(128.0f / 255.0f, 128.0f / 255.0f, 128.0f / 255.0f, 1.0f));
 
-std::atomic<TClosableMpscQueue<FRHIResource*>*> FRHIResource::PendingDeletes { new TClosableMpscQueue<FRHIResource*>() };
-FHazardPointerCollection FRHIResource::PendingDeletesHPC;
+UE::TDepletableMpscQueue<FRHIResource*, FConcurrentLinearAllocator> FRHIResource::PendingDeletes;
 FRHIResource* FRHIResource::CurrentlyDeleting = nullptr;
 
 #if HAS_GPU_STATS
@@ -1257,12 +1256,10 @@ int32 FRHIResource::FlushPendingDeletes(FRHICommandListImmediate& RHICmdList)
 	check(IsInRenderingThread());
 
 	TArray<FRHIResource*> DeletedResources;
-	TClosableMpscQueue<FRHIResource*>* PendingDeletesPtr = PendingDeletes.exchange(new TClosableMpscQueue<FRHIResource*>());
-	PendingDeletesPtr->Close([&DeletedResources](FRHIResource* Resource)
-	{	
+	PendingDeletes.Deplete([&DeletedResources](FRHIResource* Resource)
+	{
 		DeletedResources.Push(Resource);
 	});
-	PendingDeletesHPC.Delete(PendingDeletesPtr);
 
 	const int32 NumDeletes = DeletedResources.Num();
 
diff --git a/Engine/Source/Runtime/RHI/Public/RHIResources.h b/Engine/Source/Runtime/RHI/Public/RHIResources.h
--- a/Engine/Source/Runtime/RHI/Public/RHIResources.h
+++ b/Engine/Source/Runtime/RHI/Public/RHIResources.h
@@ -23,7 +23,7 @@
 #include "Async/TaskGraphInterfaces.h"
 #include "Serialization/MemoryImage.h"
 #include "Experimental/Containers/HazardPointer.h"
-#include "Containers/ClosableMpscQueue.h"
+#include "Containers/DepletableMpscQueue.h"
 #include "Misc/CoreDelegates.h"
 
 // RHI_WANT_RESOURCE_INFO should be controlled by the RHI module.
@@ -85,15 +85,7 @@ private:
 	{
 		if (!AtomicFlags.MarkForDelete(std::memory_order_release))
 		{
-			while (true)
-			{
-				auto HP = MakeHazardPointer(PendingDeletes, PendingDeletesHPC);
-				TClosableMpscQueue<FRHIResource*>* PendingDeletesPtr = HP.Get();
-				if (PendingDeletesPtr->Enqueue(const_cast<FRHIResource*>(this)))
-				{
-					break;
-				}
-			}
+			PendingDeletes.Enqueue(const_cast<FRHIResource*>(this));
 		}
 	}
 
@@ -231,8 +223,7 @@ private:
 	uint8 bBeingTracked : 1;
 #endif
 
-	static std::atomic<TClosableMpscQueue<FRHIResource*>*> PendingDeletes;
-	static FHazardPointerCollection PendingDeletesHPC;
+	static UE::TDepletableMpscQueue<FRHIResource*, FConcurrentLinearAllocator> PendingDeletes;
 	static FRHIResource* CurrentlyDeleting;
 
 	// Some APIs don't do internal reference counting, so we have to wait an extra couple of frames before deleting resources
