Rollback 23059887- Fixes OOMs on all platforms.

#rnx

diff --git a/Engine/Source/Runtime/RHI/Private/RHI.cpp b/Engine/Source/Runtime/RHI/Private/RHI.cpp
--- a/Engine/Source/Runtime/RHI/Private/RHI.cpp
+++ b/Engine/Source/Runtime/RHI/Private/RHI.cpp
@@ -294,7 +294,8 @@ const FClearValueBinding FClearValueBinding::Green(FLinearColor(0.0f, 1.0f, 0.0f
 // Note: this is used as the default normal for DBuffer decals.  It must decode to a value of 0 in DecodeDBufferData.
 const FClearValueBinding FClearValueBinding::DefaultNormal8Bit(FLinearColor(128.0f / 255.0f, 128.0f / 255.0f, 128.0f / 255.0f, 1.0f));
 
-UE::TDepletableMpscQueue<FRHIResource*, FConcurrentLinearAllocator> FRHIResource::PendingDeletes;
+std::atomic<TClosableMpscQueue<FRHIResource*>*> FRHIResource::PendingDeletes { new TClosableMpscQueue<FRHIResource*>() };
+FHazardPointerCollection FRHIResource::PendingDeletesHPC;
 FRHIResource* FRHIResource::CurrentlyDeleting = nullptr;
 
 #if HAS_GPU_STATS
@@ -1256,10 +1257,12 @@ int32 FRHIResource::FlushPendingDeletes(FRHICommandListImmediate& RHICmdList)
 	check(IsInRenderingThread());
 
 	TArray<FRHIResource*> DeletedResources;
-	PendingDeletes.Deplete([&DeletedResources](FRHIResource* Resource)
-	{
+	TClosableMpscQueue<FRHIResource*>* PendingDeletesPtr = PendingDeletes.exchange(new TClosableMpscQueue<FRHIResource*>());
+	PendingDeletesPtr->Close([&DeletedResources](FRHIResource* Resource)
+	{	
 		DeletedResources.Push(Resource);
 	});
+	PendingDeletesHPC.Delete(PendingDeletesPtr);
 
 	const int32 NumDeletes = DeletedResources.Num();
 
diff --git a/Engine/Source/Runtime/RHI/Public/RHIResources.h b/Engine/Source/Runtime/RHI/Public/RHIResources.h
--- a/Engine/Source/Runtime/RHI/Public/RHIResources.h
+++ b/Engine/Source/Runtime/RHI/Public/RHIResources.h
@@ -23,7 +23,7 @@
 #include "Async/TaskGraphInterfaces.h"
 #include "Serialization/MemoryImage.h"
 #include "Experimental/Containers/HazardPointer.h"
-#include "Containers/DepletableMpscQueue.h"
+#include "Containers/ClosableMpscQueue.h"
 #include "Misc/CoreDelegates.h"
 
 // RHI_WANT_RESOURCE_INFO should be controlled by the RHI module.
@@ -85,7 +85,15 @@ private:
 	{
 		if (!AtomicFlags.MarkForDelete(std::memory_order_release))
 		{
-			PendingDeletes.Enqueue(const_cast<FRHIResource*>(this));
+			while (true)
+			{
+				auto HP = MakeHazardPointer(PendingDeletes, PendingDeletesHPC);
+				TClosableMpscQueue<FRHIResource*>* PendingDeletesPtr = HP.Get();
+				if (PendingDeletesPtr->Enqueue(const_cast<FRHIResource*>(this)))
+				{
+					break;
+				}
+			}
 		}
 	}
 
@@ -223,7 +231,8 @@ private:
 	uint8 bBeingTracked : 1;
 #endif
 
-	static UE::TDepletableMpscQueue<FRHIResource*, FConcurrentLinearAllocator> PendingDeletes;
+	static std::atomic<TClosableMpscQueue<FRHIResource*>*> PendingDeletes;
+	static FHazardPointerCollection PendingDeletesHPC;
 	static FRHIResource* CurrentlyDeleting;
 
 	// Some APIs don't do internal reference counting, so we have to wait an extra couple of frames before deleting resources
