Log out more information when using RHI/FeatureLevel preferences. This will make it easier to diagnose why an RHI+FL combo was used at runtime.

diff --git a/Engine/Source/Runtime/RHI/Private/Windows/WindowsDynamicRHI.cpp b/Engine/Source/Runtime/RHI/Private/Windows/WindowsDynamicRHI.cpp
--- a/Engine/Source/Runtime/RHI/Private/Windows/WindowsDynamicRHI.cpp
+++ b/Engine/Source/Runtime/RHI/Private/Windows/WindowsDynamicRHI.cpp
@@ -308,6 +308,7 @@ static bool DefaultFeatureLevelES31()
 	int MinCoreCount = 0;
 	if (GConfig->GetInt(TEXT("PerformanceMode"), TEXT("MinCoreCount"), MinCoreCount, GEngineIni) && FPlatformMisc::NumberOfCoresIncludingHyperthreads() < MinCoreCount)
 	{
+		UE_LOG(LogRHI, Log, TEXT("MinCoreCount found (%d) and user core count (%d) forced feature level to ES3_1"), MinCoreCount, FPlatformMisc::NumberOfCoresIncludingHyperthreads());
 		ForceES31 = true;
 		return true;
 	}
@@ -326,6 +327,8 @@ static bool DefaultFeatureLevelES31()
 			{
 				if (FPlatformMemory::GetMemorySizeBucket() >= EPlatformMemorySizeBucket(EnumIndex))
 				{
+					UE_LOG(LogRHI, Log, TEXT("MinMemorySizeBucket found (%s) and user memory (%d) forced feature level to ES3_1"), *MinMemorySizeBucketString, int32(FPlatformMemory::GetMemorySizeBucket()));
+
 					ForceES31 = true;
 					return true;
 				}
@@ -337,8 +340,9 @@ static bool DefaultFeatureLevelES31()
 				const int MIN_GPU_MEMORY = 512 * 1024 * 1024;
 				if (FPlatformMemory::GetMemorySizeBucket() >= EPlatformMemorySizeBucket(EnumIndex) && BestGPUInfo.DedicatedVideoMemory < MIN_GPU_MEMORY)
 				{
-					ForceES31 = true;
+					UE_LOG(LogRHI, Log, TEXT("MinMemorySizeBucket found (%s) and user memory (%d) forced feature level to ES3_1"), *MinMemorySizeBucketString, int32(FPlatformMemory::GetMemorySizeBucket()));
 
+					ForceES31 = true;
 					return true;
 				}
 			}
@@ -363,8 +367,9 @@ static bool DefaultFeatureLevelES31()
 
 		if (RHIName.Compare("D3D11_ES31", ESearchCase::IgnoreCase) == 0 && GPUBrand.Compare(DeviceName, ESearchCase::IgnoreCase) == 0)
 		{
-			ForceES31 = true;
+			UE_LOG(LogRHI, Log, TEXT("Found RHIName %s with DeviceName %s which has forced feature level to ES3_1"), *RHIName, *DeviceName);
 
+			ForceES31 = true;
 			return true;
 		}
 
@@ -380,8 +385,9 @@ static bool DefaultFeatureLevelES31()
 			BestGPUInfo.VendorId == VendorIdInt && BestGPUInfo.DeviceId == DeviceIdInt &&
 			RHIName.Compare("D3D11_ES31", ESearchCase::IgnoreCase) == 0)
 		{
-			ForceES31 = true;
+			UE_LOG(LogRHI, Log, TEXT("Found RHIName %s with DeviceName %s, VendorId %s, and DeviceId %s which has forced feature level to ES3_1"), *RHIName, *DeviceName, *VendorId, *DeviceId);
 
+			ForceES31 = true;
 			return true;
 		}
 	}
@@ -398,6 +404,8 @@ static TOptional<ERHIFeatureLevel::Type> GetPreferredFeatureLevel()
 		FString PreferredFeatureLevelName;
 		if (GConfig->GetString(TEXT("D3DRHIPreference"), TEXT("PreferredFeatureLevel"), PreferredFeatureLevelName, GGameUserSettingsIni))
 		{
+			UE_LOG(LogRHI, Log, TEXT("Found D3DRHIPreference PreferredFeatureLevel: %s"), *PreferredFeatureLevelName);
+
 			if (PreferredFeatureLevelName == TEXT("sm5"))
 			{
 				PreferredFeatureLevel = ERHIFeatureLevel::SM5;
@@ -419,6 +427,10 @@ static TOptional<ERHIFeatureLevel::Type> GetPreferredFeatureLevel()
 		{
 			bool bPreferFeatureLevelES31 = false;
 			bool bFoundPreference = GConfig->GetBool(TEXT("D3DRHIPreference"), TEXT("bPreferFeatureLevelES31"), bPreferFeatureLevelES31, GGameUserSettingsIni);
+			if (bFoundPreference)
+			{
+				UE_LOG(LogRHI, Log, TEXT("Found D3DRHIPreference bPreferFeatureLevelES31: %s"), bPreferFeatureLevelES31 ? TEXT("true") : TEXT("false"));
+			}
 
 			// Force low-spec users into performance mode but respect their choice once they have set a preference
 			bool bDefaultES31 = false;
@@ -500,6 +512,8 @@ static TOptional<EWindowsRHI> ChoosePreferredRHI(EWindowsRHI InDefaultRHI)
 		FString PreferredRHIName;
 		if (GConfig->GetString(TEXT("D3DRHIPreference"), TEXT("PreferredRHI"), PreferredRHIName, GGameUserSettingsIni))
 		{
+			UE_LOG(LogRHI, Log, TEXT("Found D3DRHIPreference PreferredRHI: %s"), *PreferredRHIName);
+
 			if (PreferredRHIName == TEXT("dx12"))
 			{
 				RHIPreference = EWindowsRHI::D3D12;
@@ -523,6 +537,8 @@ static TOptional<EWindowsRHI> ChoosePreferredRHI(EWindowsRHI InDefaultRHI)
 		}
 		else if (GConfig->GetBool(TEXT("D3DRHIPreference"), TEXT("bUseD3D12InGame"), bUseD3D12InGame, GGameUserSettingsIni))
 		{
+			UE_LOG(LogRHI, Log, TEXT("Found D3DRHIPreference bUseD3D12InGame: %s"), bUseD3D12InGame ? TEXT("true") : TEXT("false"));
+
 			RHIPreference = bUseD3D12InGame ? EWindowsRHI::D3D12 : EWindowsRHI::D3D11;
 		}
 	}
@@ -626,12 +642,14 @@ static ERHIFeatureLevel::Type ChooseFeatureLevel(EWindowsRHI ChosenRHI, const TO
 		// Allow the forced feature level if we're in a position to compile its shaders
 		if (!FPlatformProperties::RequiresCookedData())
 		{
+			UE_LOG(LogRHI, Log, TEXT("Using Forced Feature Level in Editor: %s"), *LexToString(ForcedFeatureLevel.GetValue()));
 			return ForcedFeatureLevel.GetValue();
 		}
 
 		// Make sure the feature level is supported by the runtime, otherwise fall back to the default
 		if (Config.IsFeatureLevelTargeted(ChosenRHI, ForcedFeatureLevel.GetValue()))
 		{
+			UE_LOG(LogRHI, Log, TEXT("Using Forced Feature Level the Game is configured for: %s"), *LexToString(ForcedFeatureLevel.GetValue()));
 			return ForcedFeatureLevel.GetValue();
 		}
 	}
@@ -660,6 +678,10 @@ static ERHIFeatureLevel::Type ChooseFeatureLevel(EWindowsRHI ChosenRHI, const TO
 				UE_LOG(LogRHI, Warning, TEXT("User requested RHI '%s' but that is not supported by this project's data. Defaulting to Feature Level '%s'."), RHIName, *FeatureLevelName);
 			}
 		}
+		else
+		{
+			UE_LOG(LogRHI, Log, TEXT("Using Highest Feature Level of %s: %s"), ModuleNameFromWindowsRHI(ChosenRHI), *LexToString(FeatureLevel.GetValue()));
+		}
 	}
 
 	// If the user wanted to force a feature level and we couldn't set it, log out why and what we're actually running with
@@ -781,10 +803,18 @@ static IDynamicRHIModule* LoadDynamicRHIModule(ERHIFeatureLevel::Type& DesiredFe
 	if (ForcedRHI)
 	{
 		ChosenRHI = ForcedRHI.GetValue();
+
+		UE_LOG(LogRHI, Log, TEXT("Using Forced RHI: %s"), ModuleNameFromWindowsRHI(ChosenRHI));
 	}
 	else if (PreferredRHI)
 	{
 		ChosenRHI = PreferredRHI.GetValue();
+
+		UE_LOG(LogRHI, Log, TEXT("Using Preferred RHI: %s"), ModuleNameFromWindowsRHI(ChosenRHI));
+	}
+	else
+	{
+		UE_LOG(LogRHI, Log, TEXT("Using Default RHI: %s"), ModuleNameFromWindowsRHI(ChosenRHI));
 	}
 
 	const TOptional<ERHIFeatureLevel::Type> ForcedFeatureLevel = GetForcedFeatureLevel();
