Don't add dependency to command list for compute PSOs create from file cache (fix possible hitches due to slow compute PSO creation during startup)

diff --git a/Engine/Source/Runtime/RHI/Private/PipelineStateCache.cpp b/Engine/Source/Runtime/RHI/Private/PipelineStateCache.cpp
--- a/Engine/Source/Runtime/RHI/Private/PipelineStateCache.cpp
+++ b/Engine/Source/Runtime/RHI/Private/PipelineStateCache.cpp
@@ -1628,7 +1628,10 @@ FComputePipelineState* PipelineStateCache::GetAndOrCreateComputePipelineState(FR
 		if (DoAsyncCompile)
 		{
 			OutCachedState->CompletionEvent = TGraphTask<FCompilePipelineStateTask>::CreateTask().ConstructAndDispatchWhenReady(OutCachedState, FGraphicsPipelineStateInitializer(), EPSOPrecacheResult::Untracked);
-			RHICmdList.AddDispatchPrerequisite(OutCachedState->CompletionEvent);
+			if (!bFromFileCache)
+			{
+				RHICmdList.AddDispatchPrerequisite(OutCachedState->CompletionEvent);
+			}
 		}
 		else
 		{
