VulkanRHI: Bindless: allow the 8bit type in the bindless handle to be configurable per RHI if needed.

#rnx

diff --git a/Engine/Source/Runtime/D3D12RHI/Private/D3D12Device.cpp b/Engine/Source/Runtime/D3D12RHI/Private/D3D12Device.cpp
--- a/Engine/Source/Runtime/D3D12RHI/Private/D3D12Device.cpp
+++ b/Engine/Source/Runtime/D3D12RHI/Private/D3D12Device.cpp
@@ -118,7 +118,7 @@ FD3D12Device::FD3D12Device(FRHIGPUMask InGPUMask, FD3D12Adapter* InAdapter)
 {
 	check(IsInGameThread());
 
-	for (uint32 HeapType = 0; HeapType < (uint32)ERHIDescriptorHeapType::count; ++HeapType)
+	for (uint32 HeapType = 0; HeapType < (uint32)ERHIDescriptorHeapType::Count; ++HeapType)
 	{
 		OfflineDescriptorManagers.Emplace(this, (ERHIDescriptorHeapType)HeapType);
 	}
diff --git a/Engine/Source/Runtime/D3D12RHI/Private/D3D12Device.h b/Engine/Source/Runtime/D3D12RHI/Private/D3D12Device.h
--- a/Engine/Source/Runtime/D3D12RHI/Private/D3D12Device.h
+++ b/Engine/Source/Runtime/D3D12RHI/Private/D3D12Device.h
@@ -305,7 +305,7 @@ public:
 	inline FD3D12OnlineDescriptorManager&   GetOnlineDescriptorManager  () { return OnlineDescriptorManager;   }
 	inline FD3D12OfflineDescriptorManager&  GetOfflineDescriptorManager (ERHIDescriptorHeapType InType)
 	{
-		check(InType < ERHIDescriptorHeapType::count);
+		check(InType < ERHIDescriptorHeapType::Count);
 		return OfflineDescriptorManagers[static_cast<int>(InType)];
 	}
 
@@ -367,7 +367,7 @@ private:
 
 	FD3D12DescriptorHeapManager     DescriptorHeapManager;
 	FD3D12BindlessDescriptorManager BindlessDescriptorManager;
-	TArray<FD3D12OfflineDescriptorManager, TInlineAllocator<(uint32)ERHIDescriptorHeapType::count>> OfflineDescriptorManagers;
+	TArray<FD3D12OfflineDescriptorManager, TInlineAllocator<(uint32)ERHIDescriptorHeapType::Count>> OfflineDescriptorManagers;
 
 	FD3D12GlobalOnlineSamplerHeap GlobalSamplerHeap;
 	FD3D12OnlineDescriptorManager OnlineDescriptorManager;
diff --git a/Engine/Source/Runtime/RHI/Public/RHIDefinitions.h b/Engine/Source/Runtime/RHI/Public/RHIDefinitions.h
--- a/Engine/Source/Runtime/RHI/Public/RHIDefinitions.h
+++ b/Engine/Source/Runtime/RHI/Public/RHIDefinitions.h
@@ -2066,26 +2066,33 @@ enum class ERHIDescriptorHeapType : uint8
 	Sampler,
 	RenderTarget,
 	DepthStencil,
-	count
+	Count,
+	Invalid = MAX_uint8
 };
 
 struct FRHIDescriptorHandle
 {
 	FRHIDescriptorHandle() = default;
 	FRHIDescriptorHandle(ERHIDescriptorHeapType InType, uint32 InIndex)
+		: Index(InIndex)
+		, Type((uint8)InType)
+	{
+	}
+	FRHIDescriptorHandle(uint8 InType, uint32 InIndex)
 		: Index(InIndex)
 		, Type(InType)
 	{
 	}
 
-	inline uint32                GetIndex() const { return Index; }
-	inline ERHIDescriptorHeapType GetType() const { return Type; }
+	inline uint32                 GetIndex() const { return Index; }
+	inline ERHIDescriptorHeapType GetType() const { return (ERHIDescriptorHeapType)Type; }
+	inline uint8                  GetRawType() const { return Type; }
 
-	inline bool IsValid() const { return Index != UINT_MAX && Type != ERHIDescriptorHeapType::count; }
+	inline bool IsValid() const { return Index != UINT_MAX && Type != (uint8)ERHIDescriptorHeapType::Invalid; }
 
 private:
-	uint32                 Index{ UINT_MAX };
-	ERHIDescriptorHeapType Type{ ERHIDescriptorHeapType::count };
+	uint32    Index{ UINT_MAX };
+	uint8     Type{ (uint8)ERHIDescriptorHeapType::Invalid };
 };
 
 using FDisplayInformationArray = TArray<struct FDisplayInformation>;
