On Windows, a forced feature level can imply a forced RHI. The specific case we need to address is -es31 needs to force -d3d11 and error if a different RHI was forced.
bPreferFeatureLevelES31 also should only apply to RHIs that ES3.1 is enabled for.

diff --git a/Engine/Source/Runtime/RHI/Private/Windows/WindowsDynamicRHI.cpp b/Engine/Source/Runtime/RHI/Private/Windows/WindowsDynamicRHI.cpp
--- a/Engine/Source/Runtime/RHI/Private/Windows/WindowsDynamicRHI.cpp
+++ b/Engine/Source/Runtime/RHI/Private/Windows/WindowsDynamicRHI.cpp
@@ -396,7 +396,7 @@ static bool DefaultFeatureLevelES31()
 	return false;
 }
 
-static TOptional<ERHIFeatureLevel::Type> GetPreferredFeatureLevel()
+static TOptional<ERHIFeatureLevel::Type> GetPreferredFeatureLevel(EWindowsRHI ChosenRHI, const FParsedWindowsDynamicRHIConfig& Config)
 {
 	TOptional<ERHIFeatureLevel::Type> PreferredFeatureLevel;
 	if (!GIsEditor)
@@ -423,7 +423,7 @@ static TOptional<ERHIFeatureLevel::Type> GetPreferredFeatureLevel()
 				UE_LOG(LogRHI, Error, TEXT("unknown feature level name \"%s\" in game user settings, using default"), *PreferredFeatureLevelName);
 			}
 		}
-		else
+		else if (Config.IsFeatureLevelTargeted(ChosenRHI, ERHIFeatureLevel::ES3_1))
 		{
 			bool bPreferFeatureLevelES31 = false;
 			bool bFoundPreference = GConfig->GetBool(TEXT("D3DRHIPreference"), TEXT("bPreferFeatureLevelES31"), bPreferFeatureLevelES31, GGameUserSettingsIni);
@@ -452,27 +452,6 @@ static TOptional<ERHIFeatureLevel::Type> GetPreferredFeatureLevel()
 	return PreferredFeatureLevel;
 }
 
-static bool IsES31D3DOnly()
-{
-	bool bES31DXOnly = false;
-#if !WITH_EDITOR
-	if (!GIsEditor)
-	{
-		GConfig->GetBool(TEXT("PerformanceMode"), TEXT("bES31DXOnly"), bES31DXOnly, GEngineIni);
-	}
-#endif
-	return bES31DXOnly;
-}
-
-static bool AllowD3D12FeatureLevelES31(const FParsedWindowsDynamicRHIConfig& Config)
-{
-	if (!GIsEditor)
-	{
-		return Config.IsFeatureLevelTargeted(EWindowsRHI::D3D12, ERHIFeatureLevel::ES3_1);
-	}
-	return true;
-}
-
 // Choose the default from DefaultGraphicsRHI or TargetedRHIs. DefaultGraphicsRHI has precedence.
 static EWindowsRHI ChooseDefaultRHI(const FParsedWindowsDynamicRHIConfig& Config)
 {
@@ -546,7 +525,7 @@ static TOptional<EWindowsRHI> ChoosePreferredRHI(EWindowsRHI InDefaultRHI)
 	return RHIPreference;
 }
 
-static TOptional<EWindowsRHI> ChooseForcedRHI(const FParsedWindowsDynamicRHIConfig& Config)
+static TOptional<EWindowsRHI> ChooseForcedRHI(TOptional<ERHIFeatureLevel::Type> ForcedFeatureLevel, const FParsedWindowsDynamicRHIConfig& Config)
 {
 	TOptional<EWindowsRHI> ForcedRHI = {};
 
@@ -579,36 +558,37 @@ static TOptional<EWindowsRHI> ChooseForcedRHI(const FParsedWindowsDynamicRHIConf
 		UE_LOG(LogRHI, Fatal, TEXT("-d3d12, -d3d11, -vulkan, and -opengl are mutually exclusive options, but more than one was specified on the command-line."));
 	}
 
-#if	!WITH_EDITOR && UE_BUILD_SHIPPING
-	// In Shipping builds we can limit ES31 on Windows to only DX11. All RHIs are allowed by default.
-
 	// FeatureLevelES31 is also a command line override, so it will determine the underlying RHI unless one is specified
-	if (IsES31D3DOnly() && (FParse::Param(FCommandLine::Get(), TEXT("FeatureLevelES31")) || FParse::Param(FCommandLine::Get(), TEXT("FeatureLevelES3_1"))))
+	if ((FPlatformProperties::RequiresCookedData() || !GIsEditor) && ForcedFeatureLevel == ERHIFeatureLevel::ES3_1)
 	{
-		if (ForcedRHI == EWindowsRHI::OpenGL)
+		if (ForcedRHI)
 		{
-			FMessageDialog::Open(EAppMsgType::Ok, NSLOCTEXT("WindowsDynamicRHI", "RHIPerformanceOpenGL", "OpenGL is not supported for Performance Mode."));
-			UE_LOG(LogRHI, Fatal, TEXT("OpenGL is not supported for Performance Mode."));
-		}
-		else if (ForcedRHI == EWindowsRHI::Vulkan)
-		{
-			FMessageDialog::Open(EAppMsgType::Ok, NSLOCTEXT("WindowsDynamicRHI", "RHIPerformanceVulkan", "Vulkan is not supported for Performance Mode."));
-			UE_LOG(LogRHI, Fatal, TEXT("Vulkan is not supported for Performance Mode."));
-		}
-		else if (ForcedRHI == EWindowsRHI::D3D12)
-		{
-			if (!AllowD3D12FeatureLevelES31(Config))
+			if (!Config.IsFeatureLevelTargeted(ForcedRHI.GetValue(), ForcedFeatureLevel.GetValue()))
 			{
-				FMessageDialog::Open(EAppMsgType::Ok, NSLOCTEXT("WindowsDynamicRHI", "RHIPerformanceDX12", "DirectX 12 is not supported for Performance Mode."));
-				UE_LOG(LogRHI, Fatal, TEXT("DirectX 12 is not supported for Performance Mode."));
+				if (ForcedRHI == EWindowsRHI::OpenGL)
+				{
+					FMessageDialog::Open(EAppMsgType::Ok, NSLOCTEXT("WindowsDynamicRHI", "RHIPerformanceOpenGL", "OpenGL is not supported for Performance Mode."));
+					UE_LOG(LogRHI, Fatal, TEXT("OpenGL is not supported for Performance Mode."));
+				}
+				else if (ForcedRHI == EWindowsRHI::Vulkan)
+				{
+					FMessageDialog::Open(EAppMsgType::Ok, NSLOCTEXT("WindowsDynamicRHI", "RHIPerformanceVulkan", "Vulkan is not supported for Performance Mode."));
+					UE_LOG(LogRHI, Fatal, TEXT("Vulkan is not supported for Performance Mode."));
+				}
+				else if (ForcedRHI == EWindowsRHI::D3D12)
+				{
+					FMessageDialog::Open(EAppMsgType::Ok, NSLOCTEXT("WindowsDynamicRHI", "RHIPerformanceDX12", "DirectX 12 is not supported for Performance Mode."));
+					UE_LOG(LogRHI, Fatal, TEXT("DirectX 12 is not supported for Performance Mode."));
+				}
 			}
 		}
 		else
 		{
 			ForcedRHI = EWindowsRHI::D3D11;
+
+			UE_LOG(LogRHI, Log, TEXT("Forcing RHI to %s since Feature Level %s was forced"), ModuleNameFromWindowsRHI(ForcedRHI.GetValue()), *LexToString(ForcedFeatureLevel.GetValue()));
 		}
 	}
-#endif //!WITH_EDITOR && UE_BUILD_SHIPPING
 
 	return ForcedRHI;
 }
@@ -654,7 +634,7 @@ static ERHIFeatureLevel::Type ChooseFeatureLevel(EWindowsRHI ChosenRHI, const TO
 		}
 	}
 
-	TOptional<ERHIFeatureLevel::Type> FeatureLevel = GetPreferredFeatureLevel();
+	TOptional<ERHIFeatureLevel::Type> FeatureLevel = GetPreferredFeatureLevel(ChosenRHI, Config);
 
 	if (!FeatureLevel || (FPlatformProperties::RequiresCookedData() && !Config.IsFeatureLevelTargeted(ChosenRHI, FeatureLevel.GetValue())))
 	{
@@ -797,7 +777,8 @@ static IDynamicRHIModule* LoadDynamicRHIModule(ERHIFeatureLevel::Type& DesiredFe
 
 	EWindowsRHI DefaultRHI = ChooseDefaultRHI(Config);
 	const TOptional<EWindowsRHI> PreferredRHI = ChoosePreferredRHI(DefaultRHI);
-	const TOptional<EWindowsRHI> ForcedRHI = ChooseForcedRHI(Config);
+	const TOptional<ERHIFeatureLevel::Type> ForcedFeatureLevel = GetForcedFeatureLevel();
+	const TOptional<EWindowsRHI> ForcedRHI = ChooseForcedRHI(ForcedFeatureLevel, Config);
 
 	EWindowsRHI ChosenRHI = DefaultRHI;
 	if (ForcedRHI)
@@ -817,8 +798,6 @@ static IDynamicRHIModule* LoadDynamicRHIModule(ERHIFeatureLevel::Type& DesiredFe
 		UE_LOG(LogRHI, Log, TEXT("Using Default RHI: %s"), ModuleNameFromWindowsRHI(ChosenRHI));
 	}
 
-	const TOptional<ERHIFeatureLevel::Type> ForcedFeatureLevel = GetForcedFeatureLevel();
-
 	DesiredFeatureLevel = ChooseFeatureLevel(ChosenRHI, ForcedRHI, ForcedFeatureLevel, Config);
 
 	// Load the dynamic RHI module.
