Batching parameters to be set in a single SetShaderParameters call
- SetShaderParameter, SetShaderTexture, SetShaderResourceViewParameter, SetShaderSampler, and SetUAVParameter all go through a new class, FRHIParameterBatcher, that batches them all together
- FRHIParameterBatcher will be flushed on every Draw, Dispatch
- FRHIParameterBatcher will also be flushed before every PSO swap and before other SetShaderParameters calls.

diff --git a/Engine/Source/Runtime/RHI/Private/RHICommandList.cpp b/Engine/Source/Runtime/RHI/Private/RHICommandList.cpp
--- a/Engine/Source/Runtime/RHI/Private/RHICommandList.cpp
+++ b/Engine/Source/Runtime/RHI/Private/RHICommandList.cpp
@@ -136,6 +136,68 @@ static TStatId GCurrentExecuteStat;
 static FCriticalSection GRHIThreadOnTasksCritical;
 static std::atomic<int32> GRHIThreadStallRequestCount;
 
+FRHIParameterBatcher::FRHIParameterBatcher() = default;
+FRHIParameterBatcher::FRHIParameterBatcher(FRHIParameterBatcher&&) = default;
+FRHIParameterBatcher::~FRHIParameterBatcher() = default;
+
+void FRHIParameterBatcher::OnBoundShaderChanged(FRHICommandList& InCommandList, const FBoundShaderStateInput& InBoundShaderStateInput)
+{
+	for (int32 Index = 0; Index < SF_NumGraphicsFrequencies; Index++)
+	{
+		InCommandList.SetBatchedShaderParameters(GetBatchedGraphicsShader(Index), AllBatchedShaderParameters[Index]);
+		AllBatchedShaderParameters[Index].Reset();
+	}
+
+	AllBatchedShaders[SF_Vertex]        = InBoundShaderStateInput.GetVertexShader();
+	AllBatchedShaders[SF_Mesh]          = InBoundShaderStateInput.GetMeshShader();
+	AllBatchedShaders[SF_Amplification] = InBoundShaderStateInput.GetAmplificationShader();
+	AllBatchedShaders[SF_Pixel]         = InBoundShaderStateInput.GetPixelShader();
+	AllBatchedShaders[SF_Geometry]      = InBoundShaderStateInput.GetGeometryShader();
+}
+
+void FRHIParameterBatcher::OnBoundShaderChanged(FRHIComputeCommandList& InCommandList, FRHIComputeShader* InBoundComputeShaderRHI)
+{
+	InCommandList.SetBatchedShaderParameters(GetBatchedComputeShader(), AllBatchedShaderParameters[SF_Compute]);
+	AllBatchedShaderParameters[SF_Compute].Reset();
+
+	AllBatchedShaders[SF_Compute] = InBoundComputeShaderRHI;
+}
+
+void FRHIParameterBatcher::PreDispatch(FRHIComputeCommandList& InCommandList)
+{
+	InCommandList.SetBatchedShaderParameters(GetBatchedComputeShader(), AllBatchedShaderParameters[SF_Compute]);
+	AllBatchedShaderParameters[SF_Compute].Reset();
+}
+
+void FRHIParameterBatcher::PreDraw(FRHICommandList& InCommandList)
+{
+	for (uint32 Index = 0; Index < SF_NumGraphicsFrequencies; Index++)
+	{
+		InCommandList.SetBatchedShaderParameters(GetBatchedGraphicsShader(Index), AllBatchedShaderParameters[Index]);
+		AllBatchedShaderParameters[Index].Reset();
+	}
+}
+
+void FRHIParameterBatcher::FlushPendingParameters(FRHIComputeCommandList& InCommandList, FRHIComputeShader* InShader)
+{
+	const EShaderFrequency Frequency = InShader->GetFrequency();
+	if (Frequency == SF_Compute)
+	{
+		InCommandList.SetBatchedShaderParameters(GetBatchedComputeShader(), AllBatchedShaderParameters[SF_Compute]);
+		AllBatchedShaderParameters[SF_Compute].Reset();
+	}
+}
+
+void FRHIParameterBatcher::FlushPendingParameters(FRHICommandList& InCommandList, FRHIGraphicsShader* InShader)
+{
+	const EShaderFrequency Frequency = InShader->GetFrequency();
+	if (Frequency < SF_NumGraphicsFrequencies)
+	{
+		InCommandList.SetBatchedShaderParameters(GetBatchedGraphicsShader(Frequency), AllBatchedShaderParameters[Frequency]);
+		AllBatchedShaderParameters[Frequency].Reset();
+	}
+}
+
 FRHICommandListBase::FRHICommandListBase(FRHIGPUMask InGPUMask, ERecordingThread InRecordingThread)
 	: FRHICommandListBase(FPersistentState(InGPUMask, InRecordingThread))
 {}
@@ -172,6 +234,7 @@ FRHICommandListBase::FRHICommandListBase(FRHICommandListBase&& Other)
     , GraphicsContext (MoveTemp(Other.GraphicsContext))
     , ComputeContext  (MoveTemp(Other.ComputeContext))
     , Contexts        (MoveTemp(Other.Contexts))
+	, ParameterBatcher(MoveTemp(Other.ParameterBatcher))
 #if RHI_COUNT_COMMANDS
     , NumCommands     (MoveTemp(Other.NumCommands))
 #endif				  
@@ -1370,6 +1433,7 @@ void FRHIComputeCommandList::BuildAccelerationStructures(const TArrayView<const
 	BuildAccelerationStructures(Params, ScratchBufferRange);
 }
 #endif
+
 FBufferRHIRef FDynamicRHI::CreateBuffer_RenderThread(class FRHICommandListBase& RHICmdList, uint32 Size, EBufferUsageFlags Usage, uint32 Stride, ERHIAccess ResourceState, FRHIResourceCreateInfo& CreateInfo)
 {
 	CSV_SCOPED_TIMING_STAT(RHITStalls, CreateBuffer_RenderThread);
diff --git a/Engine/Source/Runtime/RHI/Public/RHICommandList.h b/Engine/Source/Runtime/RHI/Public/RHICommandList.h
--- a/Engine/Source/Runtime/RHI/Public/RHICommandList.h
+++ b/Engine/Source/Runtime/RHI/Public/RHICommandList.h
@@ -406,6 +406,89 @@ public:
 	~FRHICommandListScopedPipelineGuard();
 };
 
+struct FRHIBatchedShaderParameters
+{
+	TArray<uint8> ParametersData;
+	TArray<FRHIShaderParameter> Parameters;
+	TArray<FRHIShaderParameterResource> ResourceParameters;
+	TArray<FRHIShaderParameterResource> BindlessParameters;
+
+	inline bool HasParameters() const
+	{
+		return (Parameters.Num() + ResourceParameters.Num() + BindlessParameters.Num()) > 0;
+	}
+
+	void Reset()
+	{
+		ParametersData.Reset();
+		Parameters.Reset();
+		ResourceParameters.Reset();
+		BindlessParameters.Reset();
+	}
+};
+
+class RHI_API FRHIParameterBatcher
+{
+protected:
+	FRHIShader* AllBatchedShaders[SF_NumStandardFrequencies]{};
+	FRHIBatchedShaderParameters AllBatchedShaderParameters[SF_NumStandardFrequencies];
+
+	FRHIGraphicsShader* GetBatchedGraphicsShader(int32 Index) { return static_cast<FRHIGraphicsShader*>(AllBatchedShaders[Index]); }
+	FRHIComputeShader*  GetBatchedComputeShader()             { return static_cast<FRHIComputeShader*>(AllBatchedShaders[SF_Compute]);   }
+
+	FORCEINLINE_DEBUGGABLE FRHIBatchedShaderParameters& GetBatchedShaderParameters(FRHIShader* InShader)
+	{
+		checkSlow(InShader);
+		const EShaderFrequency Frequency = InShader->GetFrequency();
+		checkSlow(Frequency < SF_NumStandardFrequencies);
+		checkSlow(AllBatchedShaders[Frequency] != nullptr);
+		return AllBatchedShaderParameters[Frequency];
+	}
+
+public:
+	FRHIParameterBatcher();
+	FRHIParameterBatcher(FRHIParameterBatcher&&);
+	~FRHIParameterBatcher();
+
+	void OnBoundShaderChanged(class FRHICommandList& InCommandList, const FBoundShaderStateInput& InBoundShaderStateInput);
+	void OnBoundShaderChanged(class FRHIComputeCommandList& InCommandList, FRHIComputeShader* InBoundComputeShaderRHI);
+
+	void PreDispatch(class FRHIComputeCommandList& InCommandList);
+	void PreDraw(class FRHICommandList& InCommandList);
+
+	void FlushPendingParameters(class FRHIComputeCommandList& InCommandList, FRHIComputeShader* InShader);
+	void FlushPendingParameters(class FRHICommandList& InCommandList, FRHIGraphicsShader* InShader);
+
+	FORCEINLINE_DEBUGGABLE void SetShaderParameter(FRHIShader* Shader, uint32 BufferIndex, uint32 BaseIndex, uint32 NumBytes, const void* NewValue)
+	{
+		FRHIBatchedShaderParameters& BatchedParameters = GetBatchedShaderParameters(Shader);
+
+		const int32 DestDataOffset = BatchedParameters.ParametersData.Num();
+		BatchedParameters.ParametersData.Append((const uint8*)NewValue, NumBytes);
+		BatchedParameters.Parameters.Emplace((uint16)BufferIndex, (uint16)BaseIndex, (uint16)DestDataOffset, (uint16)NumBytes);
+	}
+	FORCEINLINE_DEBUGGABLE void SetShaderTexture(FRHIShader* Shader, uint32 TextureIndex, FRHITexture* Texture)
+	{
+		GetBatchedShaderParameters(Shader).ResourceParameters.Emplace(FRHIShaderParameterResource::EType::Texture, Texture, (uint16)TextureIndex);
+	}
+	FORCEINLINE_DEBUGGABLE void SetShaderResourceViewParameter(FRHIShader* Shader, uint32 SamplerIndex, FRHIShaderResourceView* SRV)
+	{
+		GetBatchedShaderParameters(Shader).ResourceParameters.Emplace(FRHIShaderParameterResource::EType::ResourceView, SRV, (uint16)SamplerIndex);
+	}
+	FORCEINLINE_DEBUGGABLE void SetShaderSampler(FRHIShader* Shader, uint32 SamplerIndex, FRHISamplerState* State)
+	{
+		GetBatchedShaderParameters(Shader).ResourceParameters.Emplace(FRHIShaderParameterResource::EType::Sampler, State, (uint16)SamplerIndex);
+	}
+	FORCEINLINE_DEBUGGABLE void SetUAVParameter(FRHIShader* Shader, uint32 UAVIndex, FRHIUnorderedAccessView* UAV)
+	{
+		GetBatchedShaderParameters(Shader).ResourceParameters.Emplace(FRHIShaderParameterResource::EType::UnorderedAccessView, UAV, (uint16)UAVIndex);
+	}
+	FORCEINLINE_DEBUGGABLE void SetUAVParameter(FRHIShader* Shader, uint32 UAVIndex, FRHIUnorderedAccessView* UAV, uint32 InitialCount)
+	{
+		checkNoEntry(); // @todo: support append/consume buffers
+	}
+};
+
 class RHI_API FRHICommandListBase : public FNoncopyable
 {
 public:
@@ -727,6 +810,8 @@ protected:
 	// These are always set for the immediate command list, see InitializeImmediateContexts().
 	TRHIPipelineArray<IRHIComputeContext*> Contexts = {};
 
+	FRHIParameterBatcher ParameterBatcher;
+
 #if RHI_COUNT_COMMANDS
 	uint32 NumCommands = 0;
 #endif
@@ -1093,24 +1178,6 @@ FRHICOMMAND_MACRO_TPL(TRHIShader, FRHICommandSetShaderParameters)
 	RHI_API void Execute(FRHICommandListBase& CmdList);
 };
 
-FRHICOMMAND_MACRO_TPL(TRHIShader, FRHICommandSetShaderParameter)
-{
-	TRHIShader* Shader;
-	const void* NewValue;
-	uint32 BufferIndex;
-	uint32 BaseIndex;
-	uint32 NumBytes;
-	FORCEINLINE_DEBUGGABLE FRHICommandSetShaderParameter(TRHIShader* InShader, uint32 InBufferIndex, uint32 InBaseIndex, uint32 InNumBytes, const void* InNewValue)
-		: Shader(InShader)
-		, NewValue(InNewValue)
-		, BufferIndex(InBufferIndex)
-		, BaseIndex(InBaseIndex)
-		, NumBytes(InNumBytes)
-	{
-	}
-	RHI_API void Execute(FRHICommandListBase& CmdList);
-};
-
 FRHICOMMAND_MACRO_TPL(TRHIShader, FRHICommandSetShaderUniformBuffer)
 {
 	TRHIShader* Shader;
@@ -1125,78 +1192,6 @@ FRHICOMMAND_MACRO_TPL(TRHIShader, FRHICommandSetShaderUniformBuffer)
 	RHI_API void Execute(FRHICommandListBase& CmdList);
 };
 
-FRHICOMMAND_MACRO_TPL(TRHIShader, FRHICommandSetShaderTexture)
-{
-	TRHIShader* Shader;
-	uint32 TextureIndex;
-	FRHITexture* Texture;
-	FORCEINLINE_DEBUGGABLE FRHICommandSetShaderTexture(TRHIShader* InShader, uint32 InTextureIndex, FRHITexture* InTexture)
-		: Shader(InShader)
-		, TextureIndex(InTextureIndex)
-		, Texture(InTexture)
-	{
-	}
-	RHI_API void Execute(FRHICommandListBase& CmdList);
-};
-
-FRHICOMMAND_MACRO_TPL(TRHIShader, FRHICommandSetShaderResourceViewParameter)
-{
-	TRHIShader* Shader;
-	uint32 SamplerIndex;
-	FRHIShaderResourceView* SRV;
-	FORCEINLINE_DEBUGGABLE FRHICommandSetShaderResourceViewParameter(TRHIShader* InShader, uint32 InSamplerIndex, FRHIShaderResourceView* InSRV)
-		: Shader(InShader)
-		, SamplerIndex(InSamplerIndex)
-		, SRV(InSRV)
-	{
-	}
-	RHI_API void Execute(FRHICommandListBase& CmdList);
-};
-
-FRHICOMMAND_MACRO_TPL(TRHIShader, FRHICommandSetUAVParameter)
-{
-	TRHIShader* Shader;
-	uint32 UAVIndex;
-	FRHIUnorderedAccessView* UAV;
-	FORCEINLINE_DEBUGGABLE FRHICommandSetUAVParameter(TRHIShader* InShader, uint32 InUAVIndex, FRHIUnorderedAccessView* InUAV)
-		: Shader(InShader)
-		, UAVIndex(InUAVIndex)
-		, UAV(InUAV)
-	{
-	}
-	RHI_API void Execute(FRHICommandListBase& CmdList);
-};
-
-FRHICOMMAND_MACRO(FRHICommandSetUAVParameter_InitialCount)
-{
-	FRHIComputeShader* Shader;
-	uint32 UAVIndex;
-	FRHIUnorderedAccessView* UAV;
-	uint32 InitialCount;
-	FORCEINLINE_DEBUGGABLE FRHICommandSetUAVParameter_InitialCount(FRHIComputeShader* InShader, uint32 InUAVIndex, FRHIUnorderedAccessView* InUAV, uint32 InInitialCount)
-		: Shader(InShader)
-		, UAVIndex(InUAVIndex)
-		, UAV(InUAV)
-		, InitialCount(InInitialCount)
-	{
-	}
-	RHI_API void Execute(FRHICommandListBase& CmdList);
-};
-
-FRHICOMMAND_MACRO_TPL(TRHIShader, FRHICommandSetShaderSampler)
-{
-	TRHIShader* Shader;
-	uint32 SamplerIndex;
-	FRHISamplerState* Sampler;
-	FORCEINLINE_DEBUGGABLE FRHICommandSetShaderSampler(TRHIShader* InShader, uint32 InSamplerIndex, FRHISamplerState* InSampler)
-		: Shader(InShader)
-		, SamplerIndex(InSamplerIndex)
-		, Sampler(InSampler)
-	{
-	}
-	RHI_API void Execute(FRHICommandListBase& CmdList);
-};
-
 FRHICOMMAND_MACRO(FRHICommandDrawPrimitive)
 {
 	uint32 BaseVertexIndex;
@@ -2259,12 +2254,7 @@ FRHICOMMAND_MACRO(FRHICommandSetRayTracingBindings)
 #endif // RHI_RAYTRACING
 
 template<> RHI_API void FRHICommandSetShaderParameters           <FRHIComputeShader>::Execute(FRHICommandListBase& CmdList);
-template<> RHI_API void FRHICommandSetShaderParameter            <FRHIComputeShader>::Execute(FRHICommandListBase& CmdList);
 template<> RHI_API void FRHICommandSetShaderUniformBuffer        <FRHIComputeShader>::Execute(FRHICommandListBase& CmdList);
-template<> RHI_API void FRHICommandSetShaderTexture              <FRHIComputeShader>::Execute(FRHICommandListBase& CmdList);
-template<> RHI_API void FRHICommandSetShaderResourceViewParameter<FRHIComputeShader>::Execute(FRHICommandListBase& CmdList);
-template<> RHI_API void FRHICommandSetShaderSampler              <FRHIComputeShader>::Execute(FRHICommandListBase& CmdList);
-template<> RHI_API void FRHICommandSetUAVParameter               <FRHIComputeShader>::Execute(FRHICommandListBase& CmdList);
 
 extern RHI_API FRHIComputePipelineState* ExecuteSetComputePipelineState(FComputePipelineState* ComputePipelineState);
 extern RHI_API FRHIGraphicsPipelineState* ExecuteSetGraphicsPipelineState(class FGraphicsPipelineState* GraphicsPipelineState);
@@ -2272,6 +2262,23 @@ extern RHI_API FRHIRayTracingPipelineState* GetRHIRayTracingPipelineState(FRayTr
 
 class RHI_API FRHIComputeCommandList : public FRHICommandListBase
 {
+protected:
+	void OnBoundShaderChanged(FRHIComputeShader* InBoundComputeShaderRHI)
+	{
+		ParameterBatcher.OnBoundShaderChanged(*this, InBoundComputeShaderRHI);
+		PersistentState.BoundComputeShaderRHI = InBoundComputeShaderRHI;
+	}
+
+	void PreDispatch()
+	{
+		ParameterBatcher.PreDispatch(*this);
+	}
+
+	void FlushPendingComputeParameters(FRHIComputeShader* InShader)
+	{
+		ParameterBatcher.FlushPendingParameters(*this, InShader);
+	}
+
 public:
 	FRHIComputeCommandList(FRHIGPUMask GPUMask = FRHIGPUMask::All(), ERecordingThread InRecordingThread = ERecordingThread::Render)
 		: FRHICommandListBase(GPUMask, InRecordingThread)
@@ -2364,9 +2371,16 @@ public:
 		, TConstArrayView<FRHIShaderParameter> InParameters
 		, TConstArrayView<FRHIShaderParameterResource> InResourceParameters
 		, TConstArrayView<FRHIShaderParameterResource> InBindlessParameters
+		, bool bSubmitBatchedParameters = true
 	)
 	{
 		ValidateBoundShader(InShader);
+
+		if (bSubmitBatchedParameters)
+		{
+			FlushPendingComputeParameters(InShader);
+		}
+
 		if (Bypass())
 		{
 			GetComputeContext().RHISetShaderParameters(InShader, InParametersData, InParameters, InResourceParameters, InBindlessParameters);
@@ -2382,17 +2396,25 @@ public:
 		);
 	}
 
-	FORCEINLINE_DEBUGGABLE void SetShaderParameter(FRHIComputeShader* Shader, uint32 BufferIndex, uint32 BaseIndex, uint32 NumBytes, const void* NewValue)
+	FORCEINLINE_DEBUGGABLE void SetBatchedShaderParameters(FRHIComputeShader* InShader, const FRHIBatchedShaderParameters& InBatchedParameters)
 	{
-		ValidateBoundShader(Shader);
-		if (Bypass())
+		if (InBatchedParameters.HasParameters())
 		{
-			GetComputeContext().RHISetShaderParameter(Shader, BufferIndex, BaseIndex, NumBytes, NewValue);
-			return;
+			SetShaderParameters(
+				InShader,
+				InBatchedParameters.ParametersData,
+				InBatchedParameters.Parameters,
+				InBatchedParameters.ResourceParameters,
+				InBatchedParameters.BindlessParameters,
+				false /* bSubmitBatchedParameters */
+			);
 		}
-		void* UseValue = Alloc(NumBytes, 16);
-		FMemory::Memcpy(UseValue, NewValue, NumBytes);
-		ALLOC_COMMAND(FRHICommandSetShaderParameter<FRHIComputeShader>)(Shader, BufferIndex, BaseIndex, NumBytes, UseValue);
+	}
+
+	FORCEINLINE_DEBUGGABLE void SetShaderParameter(FRHIComputeShader* Shader, uint32 BufferIndex, uint32 BaseIndex, uint32 NumBytes, const void* NewValue)
+	{
+		ValidateBoundShader(Shader);
+		ParameterBatcher.SetShaderParameter(Shader, BufferIndex, BaseIndex, NumBytes, NewValue);
 	}
 
 	FORCEINLINE void SetShaderParameter(FComputeShaderRHIRef& Shader, uint32 BufferIndex, uint32 BaseIndex, uint32 NumBytes, const void* NewValue)
@@ -2403,23 +2425,13 @@ public:
 	FORCEINLINE_DEBUGGABLE void SetShaderTexture(FRHIComputeShader* Shader, uint32 TextureIndex, FRHITexture* Texture)
 	{
 		ValidateBoundShader(Shader);
-		if (Bypass())
-		{
-			GetComputeContext().RHISetShaderTexture(Shader, TextureIndex, Texture);
-			return;
-		}
-		ALLOC_COMMAND(FRHICommandSetShaderTexture<FRHIComputeShader>)(Shader, TextureIndex, Texture);
+		ParameterBatcher.SetShaderTexture(Shader, TextureIndex, Texture);
 	}
 
 	FORCEINLINE_DEBUGGABLE void SetShaderResourceViewParameter(FRHIComputeShader* Shader, uint32 SamplerIndex, FRHIShaderResourceView* SRV)
 	{
 		ValidateBoundShader(Shader);
-		if (Bypass())
-		{
-			GetComputeContext().RHISetShaderResourceViewParameter(Shader, SamplerIndex, SRV);
-			return;
-		}
-		ALLOC_COMMAND(FRHICommandSetShaderResourceViewParameter<FRHIComputeShader>)(Shader, SamplerIndex, SRV);
+		ParameterBatcher.SetShaderResourceViewParameter(Shader, SamplerIndex, SRV);
 	}
 
 	FORCEINLINE_DEBUGGABLE void SetShaderSampler(FRHIComputeShader* Shader, uint32 SamplerIndex, FRHISamplerState* State)
@@ -2431,40 +2443,25 @@ public:
 			return;
 		}
 
-		if (Bypass())
-		{
-			GetComputeContext().RHISetShaderSampler(Shader, SamplerIndex, State);
-			return;
-		}
-		ALLOC_COMMAND(FRHICommandSetShaderSampler<FRHIComputeShader>)(Shader, SamplerIndex, State);
+		ParameterBatcher.SetShaderSampler(Shader, SamplerIndex, State);
 	}
 
 	FORCEINLINE_DEBUGGABLE void SetUAVParameter(FRHIComputeShader* Shader, uint32 UAVIndex, FRHIUnorderedAccessView* UAV)
 	{
 		ValidateBoundShader(Shader);
-		if (Bypass())
-		{
-			GetComputeContext().RHISetUAVParameter(Shader, UAVIndex, UAV);
-			return;
-		}
-		ALLOC_COMMAND(FRHICommandSetUAVParameter<FRHIComputeShader>)(Shader, UAVIndex, UAV);
+		ParameterBatcher.SetUAVParameter(Shader, UAVIndex, UAV);
 	}
 
 	FORCEINLINE_DEBUGGABLE void SetUAVParameter(FRHIComputeShader* Shader, uint32 UAVIndex, FRHIUnorderedAccessView* UAV, uint32 InitialCount)
 	{
 		ValidateBoundShader(Shader);
-		if (Bypass())
-		{
-			GetComputeContext().RHISetUAVParameter(Shader, UAVIndex, UAV, InitialCount);
-			return;
-		}
-		ALLOC_COMMAND(FRHICommandSetUAVParameter_InitialCount)(Shader, UAVIndex, UAV, InitialCount);
+		ParameterBatcher.SetUAVParameter(Shader, UAVIndex, UAV, InitialCount);
 	}
 
 	UE_DEPRECATED(5.1, "ComputePipelineStates should be used instead of direct ComputeShaders. You can use SetComputePipelineState(RHICmdList, ComputeShader).")
 	FORCEINLINE_DEBUGGABLE void SetComputeShader(FRHIComputeShader* ComputeShader)
 	{
-		PersistentState.BoundComputeShaderRHI = ComputeShader;
+		OnBoundShaderChanged(ComputeShader);
 		ComputeShader->UpdateStats();
 		if (Bypass())
 		{
@@ -2478,7 +2475,7 @@ PRAGMA_ENABLE_DEPRECATION_WARNINGS
 
 	FORCEINLINE_DEBUGGABLE void SetComputePipelineState(FComputePipelineState* ComputePipelineState, FRHIComputeShader* ComputeShader)
 	{
-		PersistentState.BoundComputeShaderRHI = ComputeShader;
+		OnBoundShaderChanged(ComputeShader);
 		if (Bypass())
 		{
 			FRHIComputePipelineState* RHIComputePipelineState = ExecuteSetComputePipelineState(ComputePipelineState);
@@ -2500,6 +2497,7 @@ PRAGMA_ENABLE_DEPRECATION_WARNINGS
 
 	FORCEINLINE_DEBUGGABLE void DispatchComputeShader(uint32 ThreadGroupCountX, uint32 ThreadGroupCountY, uint32 ThreadGroupCountZ)
 	{
+		PreDispatch();
 		if (Bypass())
 		{
 			GetComputeContext().RHIDispatchComputeShader(ThreadGroupCountX, ThreadGroupCountY, ThreadGroupCountZ);
@@ -2510,6 +2508,7 @@ PRAGMA_ENABLE_DEPRECATION_WARNINGS
 
 	FORCEINLINE_DEBUGGABLE void DispatchIndirectComputeShader(FRHIBuffer* ArgumentBuffer, uint32 ArgumentOffset)
 	{
+		PreDispatch();
 		if (Bypass())
 		{
 			GetComputeContext().RHIDispatchIndirectComputeShader(ArgumentBuffer, ArgumentOffset);
@@ -2977,15 +2976,29 @@ PRAGMA_ENABLE_DEPRECATION_WARNINGS
 };
 
 template<> RHI_API void FRHICommandSetShaderParameters           <FRHIGraphicsShader>::Execute(FRHICommandListBase& CmdList);
-template<> RHI_API void FRHICommandSetShaderParameter            <FRHIGraphicsShader>::Execute(FRHICommandListBase& CmdList);
 template<> RHI_API void FRHICommandSetShaderUniformBuffer        <FRHIGraphicsShader>::Execute(FRHICommandListBase& CmdList);
-template<> RHI_API void FRHICommandSetShaderTexture              <FRHIGraphicsShader>::Execute(FRHICommandListBase& CmdList);
-template<> RHI_API void FRHICommandSetShaderResourceViewParameter<FRHIGraphicsShader>::Execute(FRHICommandListBase& CmdList);
-template<> RHI_API void FRHICommandSetShaderSampler              <FRHIGraphicsShader>::Execute(FRHICommandListBase& CmdList);
-template<> RHI_API void FRHICommandSetUAVParameter               <FRHIPixelShader   >::Execute(FRHICommandListBase& CmdList);
 
 class RHI_API FRHICommandList : public FRHIComputeCommandList
 {
+protected:
+	using FRHIComputeCommandList::OnBoundShaderChanged;
+
+	void OnBoundShaderChanged(const FBoundShaderStateInput& InBoundShaderStateInput)
+	{
+		ParameterBatcher.OnBoundShaderChanged(*this, InBoundShaderStateInput);
+		PersistentState.BoundShaderInput = InBoundShaderStateInput;
+	}
+
+	void PreDraw()
+	{
+		ParameterBatcher.PreDraw(*this);
+	}
+
+	void FlushPendingGraphicsParameters(FRHIGraphicsShader* InShader)
+	{
+		ParameterBatcher.FlushPendingParameters(*this, InShader);
+	}
+
 public:
 	FRHICommandList(FRHIGPUMask GPUMask = FRHIGPUMask::All(), ERecordingThread InRecordingThread = ERecordingThread::Render)
 		: FRHIComputeCommandList(GPUMask, InRecordingThread)
@@ -3086,9 +3099,16 @@ public:
 		, TConstArrayView<FRHIShaderParameter> InParameters
 		, TConstArrayView<FRHIShaderParameterResource> InResourceParameters
 		, TConstArrayView<FRHIShaderParameterResource> InBindlessParameters
+		, bool bSubmitBatchedParameters = true
 	)
 	{
 		ValidateBoundShader(InShader);
+
+		if (bSubmitBatchedParameters)
+		{
+			FlushPendingGraphicsParameters(InShader);
+		}
+
 		if (Bypass())
 		{
 			GetContext().RHISetShaderParameters(InShader, InParametersData, InParameters, InResourceParameters, InBindlessParameters);
@@ -3104,20 +3124,30 @@ public:
 			);
 	}
 
+	using FRHIComputeCommandList::SetBatchedShaderParameters;
+
+	FORCEINLINE_DEBUGGABLE void SetBatchedShaderParameters(FRHIGraphicsShader* InShader, const FRHIBatchedShaderParameters& InBatchedParameters)
+	{
+		if (InBatchedParameters.HasParameters())
+		{
+			SetShaderParameters(
+				InShader,
+				InBatchedParameters.ParametersData,
+				InBatchedParameters.Parameters,
+				InBatchedParameters.ResourceParameters,
+				InBatchedParameters.BindlessParameters,
+				false /* bSubmitBatchedParameters */
+			);
+		}
+	}
+
 	using FRHIComputeCommandList::SetShaderParameter;
 
 	FORCEINLINE_DEBUGGABLE void SetShaderParameter(FRHIGraphicsShader* Shader, uint32 BufferIndex, uint32 BaseIndex, uint32 NumBytes, const void* NewValue)
 	{
 		//check(IsOutsideRenderPass());
 		ValidateBoundShader(Shader);
-		if (Bypass())
-		{
-			GetContext().RHISetShaderParameter(Shader, BufferIndex, BaseIndex, NumBytes, NewValue);
-			return;
-		}
-		void* UseValue = Alloc(NumBytes, 16);
-		FMemory::Memcpy(UseValue, NewValue, NumBytes);
-		ALLOC_COMMAND(FRHICommandSetShaderParameter<FRHIGraphicsShader>)(Shader, BufferIndex, BaseIndex, NumBytes, UseValue);
+		ParameterBatcher.SetShaderParameter(Shader, BufferIndex, BaseIndex, NumBytes, NewValue);
 	}
 
 	template <typename TShaderRHI>
@@ -3132,12 +3162,7 @@ public:
 	{
 		//check(IsOutsideRenderPass());
 		ValidateBoundShader(Shader);
-		if (Bypass())
-		{
-			GetContext().RHISetShaderTexture(Shader, TextureIndex, Texture);
-			return;
-		}
-		ALLOC_COMMAND(FRHICommandSetShaderTexture<FRHIGraphicsShader>)(Shader, TextureIndex, Texture);
+		ParameterBatcher.SetShaderTexture(Shader, TextureIndex, Texture);
 	}
 
 	template <typename TShaderRHI>
@@ -3152,12 +3177,7 @@ public:
 	{
 		//check(IsOutsideRenderPass());
 		ValidateBoundShader(Shader);
-		if (Bypass())
-		{
-			GetContext().RHISetShaderResourceViewParameter(Shader, SamplerIndex, SRV);
-			return;
-		}
-		ALLOC_COMMAND(FRHICommandSetShaderResourceViewParameter<FRHIGraphicsShader>)(Shader, SamplerIndex, SRV);
+		ParameterBatcher.SetShaderResourceViewParameter(Shader, SamplerIndex, SRV);
 	}
 
 	template <typename TShaderRHI>
@@ -3180,12 +3200,7 @@ public:
 			return;
 		}
 
-		if (Bypass())
-		{
-			GetContext().RHISetShaderSampler(Shader, SamplerIndex, State);
-			return;
-		}
-		ALLOC_COMMAND(FRHICommandSetShaderSampler<FRHIGraphicsShader>)(Shader, SamplerIndex, State);
+		ParameterBatcher.SetShaderSampler(Shader, SamplerIndex, State);
 	}
 
 	template <typename TShaderRHI>
@@ -3199,12 +3214,7 @@ public:
 	FORCEINLINE_DEBUGGABLE void SetUAVParameter(FRHIPixelShader* Shader, uint32 UAVIndex, FRHIUnorderedAccessView* UAV)
 	{
 		ValidateBoundShader(Shader);
-		if (Bypass())
-		{
-			GetContext().RHISetUAVParameter(Shader, UAVIndex, UAV);
-			return;
-		}
-		ALLOC_COMMAND(FRHICommandSetUAVParameter<FRHIPixelShader>)(Shader, UAVIndex, UAV);
+		ParameterBatcher.SetUAVParameter(Shader, UAVIndex, UAV);
 	}
 
 	FORCEINLINE_DEBUGGABLE void SetUAVParameter(const TRefCountPtr<FRHIPixelShader>& Shader, uint32 UAVIndex, FRHIUnorderedAccessView* UAV)
@@ -3226,6 +3236,7 @@ public:
 	FORCEINLINE_DEBUGGABLE void DrawPrimitive(uint32 BaseVertexIndex, uint32 NumPrimitives, uint32 NumInstances)
 	{
 		//check(IsOutsideRenderPass());
+		PreDraw();
 		if (Bypass())
 		{
 			GetContext().RHIDrawPrimitive(BaseVertexIndex, NumPrimitives, NumInstances);
@@ -3242,6 +3253,7 @@ public:
 		}
 
 		//check(IsOutsideRenderPass());
+		PreDraw();
 		if (Bypass())
 		{
 			GetContext().RHIDrawIndexedPrimitive(IndexBuffer, BaseVertexIndex, FirstInstance, NumVertices, StartIndex, NumPrimitives, NumInstances);
@@ -3360,7 +3372,7 @@ public:
 	FORCEINLINE_DEBUGGABLE void SetGraphicsPipelineState(class FGraphicsPipelineState* GraphicsPipelineState, const FBoundShaderStateInput& ShaderInput, uint32 StencilRef, bool bApplyAdditionalState)
 	{
 		//check(IsOutsideRenderPass());
-		PersistentState.BoundShaderInput = ShaderInput;
+		OnBoundShaderChanged(ShaderInput);
 		if (Bypass())
 		{
 			FRHIGraphicsPipelineState* RHIGraphicsPipelineState = ExecuteSetGraphicsPipelineState(GraphicsPipelineState);
@@ -3374,7 +3386,7 @@ public:
 	FORCEINLINE_DEBUGGABLE void SetGraphicsPipelineState(const FGraphicsPipelineStateInitializer& PsoInit, uint32 StencilRef, bool bApplyAdditionalState)
 	{
 		//check(IsOutsideRenderPass());
-		PersistentState.BoundShaderInput = PsoInit.BoundShaderState;
+		OnBoundShaderChanged(PsoInit.BoundShaderState);
 		if (Bypass())
 		{
 			GetContext().RHISetGraphicsPipelineState(PsoInit, StencilRef, bApplyAdditionalState);
@@ -3387,6 +3399,7 @@ public:
 	FORCEINLINE_DEBUGGABLE void DrawPrimitiveIndirect(FRHIBuffer* ArgumentBuffer, uint32 ArgumentOffset)
 	{
 		//check(IsOutsideRenderPass());
+		PreDraw();
 		if (Bypass())
 		{
 			GetContext().RHIDrawPrimitiveIndirect(ArgumentBuffer, ArgumentOffset);
@@ -3398,6 +3411,7 @@ public:
 	FORCEINLINE_DEBUGGABLE void DrawIndexedIndirect(FRHIBuffer* IndexBufferRHI, FRHIBuffer* ArgumentsBufferRHI, uint32 DrawArgumentsIndex, uint32 NumInstances)
 	{
 		//check(IsOutsideRenderPass());
+		PreDraw();
 		if (Bypass())
 		{
 			GetContext().RHIDrawIndexedIndirect(IndexBufferRHI, ArgumentsBufferRHI, DrawArgumentsIndex, NumInstances);
@@ -3409,6 +3423,7 @@ public:
 	FORCEINLINE_DEBUGGABLE void DrawIndexedPrimitiveIndirect(FRHIBuffer* IndexBuffer, FRHIBuffer* ArgumentsBuffer, uint32 ArgumentOffset)
 	{
 		//check(IsOutsideRenderPass());
+		PreDraw();
 		if (Bypass())
 		{
 			GetContext().RHIDrawIndexedPrimitiveIndirect(IndexBuffer, ArgumentsBuffer, ArgumentOffset);
@@ -3419,6 +3434,7 @@ public:
 
 	FORCEINLINE_DEBUGGABLE void DispatchMeshShader(uint32 ThreadGroupCountX, uint32 ThreadGroupCountY, uint32 ThreadGroupCountZ)
 	{
+		PreDraw();
 		if (Bypass())
 		{
 			GetContext().RHIDispatchMeshShader(ThreadGroupCountX, ThreadGroupCountY, ThreadGroupCountZ);
@@ -3429,6 +3445,7 @@ public:
 
 	FORCEINLINE_DEBUGGABLE void DispatchIndirectMeshShader(FRHIBuffer* ArgumentBuffer, uint32 ArgumentOffset)
 	{
+		PreDraw();
 		if (Bypass())
 		{
 			GetContext().RHIDispatchIndirectMeshShader(ArgumentBuffer, ArgumentOffset);
diff --git a/Engine/Source/Runtime/RHI/Public/RHICommandListCommandExecutes.inl b/Engine/Source/Runtime/RHI/Public/RHICommandListCommandExecutes.inl
--- a/Engine/Source/Runtime/RHI/Public/RHICommandListCommandExecutes.inl
+++ b/Engine/Source/Runtime/RHI/Public/RHICommandListCommandExecutes.inl
@@ -155,18 +155,6 @@ template<> RHI_API void FRHICommandSetShaderParameters<FRHIGraphicsShader>::Exec
 	INTERNAL_DECORATOR(RHISetShaderParameters)(Shader, ParametersData, Parameters, ResourceParameters, BindlessParameters);
 }
 
-template<> RHI_API void FRHICommandSetShaderParameter<FRHIComputeShader>::Execute(FRHICommandListBase& CmdList)
-{
-	RHISTAT(SetShaderParameter);
-	INTERNAL_DECORATOR_COMPUTE(RHISetShaderParameter)(Shader, BufferIndex, BaseIndex, NumBytes, NewValue);
-}
-
-template<> RHI_API void FRHICommandSetShaderParameter<FRHIGraphicsShader>::Execute(FRHICommandListBase& CmdList)
-{
-	RHISTAT(SetShaderParameter);
-	INTERNAL_DECORATOR(RHISetShaderParameter)(Shader, BufferIndex, BaseIndex, NumBytes, NewValue);
-}
-
 template<> RHI_API void FRHICommandSetShaderUniformBuffer<FRHIComputeShader>::Execute(FRHICommandListBase& CmdList)
 {
 	RHISTAT(SetShaderUniformBuffer);
@@ -179,60 +167,6 @@ template<> RHI_API void FRHICommandSetShaderUniformBuffer<FRHIGraphicsShader>::E
 	INTERNAL_DECORATOR(RHISetShaderUniformBuffer)(Shader, BaseIndex, UniformBuffer);
 }
 
-template<> RHI_API void FRHICommandSetShaderTexture<FRHIComputeShader>::Execute(FRHICommandListBase& CmdList)
-{
-	RHISTAT(SetShaderTexture);
-	INTERNAL_DECORATOR_COMPUTE(RHISetShaderTexture)(Shader, TextureIndex, Texture);
-}
-
-template<> RHI_API void FRHICommandSetShaderTexture<FRHIGraphicsShader>::Execute(FRHICommandListBase& CmdList)
-{
-	RHISTAT(SetShaderTexture);
-	INTERNAL_DECORATOR(RHISetShaderTexture)(Shader, TextureIndex, Texture);
-}
-
-template<> RHI_API void FRHICommandSetShaderResourceViewParameter<FRHIComputeShader>::Execute(FRHICommandListBase& CmdList)
-{
-	RHISTAT(SetShaderResourceViewParameter);
-	INTERNAL_DECORATOR_COMPUTE(RHISetShaderResourceViewParameter)(Shader, SamplerIndex, SRV);
-}
-
-template<> RHI_API void FRHICommandSetShaderResourceViewParameter<FRHIGraphicsShader>::Execute(FRHICommandListBase& CmdList)
-{
-	RHISTAT(SetShaderResourceViewParameter);
-	INTERNAL_DECORATOR(RHISetShaderResourceViewParameter)(Shader, SamplerIndex, SRV);
-}
-
-template<> RHI_API void FRHICommandSetUAVParameter<FRHIComputeShader>::Execute(FRHICommandListBase& CmdList)
-{
-	RHISTAT(SetUAVParameter);
-	INTERNAL_DECORATOR_COMPUTE(RHISetUAVParameter)(Shader, UAVIndex, UAV);
-}
-
-template<> RHI_API void FRHICommandSetUAVParameter<FRHIPixelShader>::Execute(FRHICommandListBase& CmdList)
-{
-	RHISTAT(SetUAVParameter);
-	INTERNAL_DECORATOR(RHISetUAVParameter)(Shader, UAVIndex, UAV);
-}
-
-void FRHICommandSetUAVParameter_InitialCount::Execute(FRHICommandListBase& CmdList)
-{
-	RHISTAT(SetUAVParameter);
-	INTERNAL_DECORATOR_COMPUTE(RHISetUAVParameter)(Shader, UAVIndex, UAV, InitialCount);
-}
-
-template<> RHI_API void FRHICommandSetShaderSampler<FRHIComputeShader>::Execute(FRHICommandListBase& CmdList)
-{
-	RHISTAT(SetShaderSampler);
-	INTERNAL_DECORATOR_COMPUTE(RHISetShaderSampler)(Shader, SamplerIndex, Sampler);
-}
-
-template<> RHI_API void FRHICommandSetShaderSampler<FRHIGraphicsShader>::Execute(FRHICommandListBase& CmdList)
-{
-	RHISTAT(SetShaderSampler);
-	INTERNAL_DECORATOR(RHISetShaderSampler)(Shader, SamplerIndex, Sampler);
-}
-
 void FRHICommandDrawPrimitive::Execute(FRHICommandListBase& CmdList)
 {
 	RHISTAT(DrawPrimitive);
