- Enable PSO precaching by default on Windows client
- Only enable PSO precaching when GRHISupportsPipelineFileCache is set (not enabled in d3d11)

#localization none
#tests did a lot of local replays with PSO precaching enabled - slightly longer loading times on empty cache and slightly higher memory load due to more uncompressed shaders stored

diff --git a/Engine/Source/Runtime/RHI/Private/PipelineStateCache.cpp b/Engine/Source/Runtime/RHI/Private/PipelineStateCache.cpp
--- a/Engine/Source/Runtime/RHI/Private/PipelineStateCache.cpp
+++ b/Engine/Source/Runtime/RHI/Private/PipelineStateCache.cpp
@@ -908,7 +908,7 @@ public:
 		if (Contains(Initializer))
 		{
 			return nullptr;
-		}
+		}		
 
 		FGraphicsPipelineState* NewState = nullptr;
 
@@ -2189,7 +2189,7 @@ bool PipelineStateCache::IsPSOPrecachingEnabled()
 	// Disables in the editor for now by default untill more testing is done - still WIP
 	return false;
 #else
-	return GPSOPrecaching != 0;
+	return GPSOPrecaching != 0 && GRHISupportsPipelineFileCache;
 #endif // WITH_EDITOR
 }
 
