Parameter batching needs more flushes
- Adding check in FRHIParameterBatcher's destructor to make sure all commands went through.
- Adding FlushAllParameters implementations for each command list type to make sure they commit all their pending parameters.

diff --git a/Engine/Source/Runtime/RHI/Private/RHICommandList.cpp b/Engine/Source/Runtime/RHI/Private/RHICommandList.cpp
--- a/Engine/Source/Runtime/RHI/Private/RHICommandList.cpp
+++ b/Engine/Source/Runtime/RHI/Private/RHICommandList.cpp
@@ -138,7 +138,13 @@ static std::atomic<int32> GRHIThreadStallRequestCount;
 
 FRHIParameterBatcher::FRHIParameterBatcher() = default;
 FRHIParameterBatcher::FRHIParameterBatcher(FRHIParameterBatcher&&) = default;
-FRHIParameterBatcher::~FRHIParameterBatcher() = default;
+FRHIParameterBatcher::~FRHIParameterBatcher()
+{
+	for (int32 Index = 0; Index < SF_NumStandardFrequencies; Index++)
+	{
+		checkSlow(AllBatchedShaderParameters[Index].HasParameters() == false);
+	}
+}
 
 void FRHIParameterBatcher::OnBoundShaderChanged(FRHICommandList& InCommandList, const FBoundShaderStateInput& InBoundShaderStateInput)
 {
@@ -178,6 +184,17 @@ void FRHIParameterBatcher::PreDraw(FRHICommandList& InCommandList)
 	}
 }
 
+void FRHIParameterBatcher::FlushAllParameters(class FRHIComputeCommandList& InCommandList)
+{
+	PreDispatch(InCommandList);
+}
+
+void FRHIParameterBatcher::FlushAllParameters(class FRHICommandList& InCommandList)
+{
+	PreDraw(InCommandList);
+	PreDispatch(InCommandList);
+}
+
 void FRHIParameterBatcher::FlushPendingParameters(FRHIComputeCommandList& InCommandList, FRHIComputeShader* InShader)
 {
 	const EShaderFrequency Frequency = InShader->GetFrequency();
@@ -329,6 +346,9 @@ void FRHICommandListBase::FinishRecording()
 {
 	checkf(!IsImmediate(), TEXT("Do not call FinishRecording() on the immediate RHI command list."));
 
+	// Make sure the batcher is cleared out. If there are pending graphics parameters, this cast is safe.
+	ParameterBatcher.FlushAllParameters(static_cast<FRHICommandList&>(*this));
+
 	PersistentState.FenceCandidate->Fence = PersistentState.RHIThreadBufferLockFence;
 
 	// "Complete" the dispatch event. This unblocks waiting tasks but only when
@@ -760,6 +780,8 @@ void FRHICommandListImmediate::ExecuteAndReset()
 	// Always reset the immediate command list when we're done.
 	ON_SCOPE_EXIT { Reset(); };
 
+	ParameterBatcher.FlushAllParameters(*this);
+
 	//
 	// In bypass mode, the immediate command list will never contain recorded commands (since these were forwarded directly into the immediate RHI contexts).
 	// However, platforms that use explicit command list submission still need to be told to submit the enqueued work in the default contexts. Do that now.
diff --git a/Engine/Source/Runtime/RHI/Public/RHICommandList.h b/Engine/Source/Runtime/RHI/Public/RHICommandList.h
--- a/Engine/Source/Runtime/RHI/Public/RHICommandList.h
+++ b/Engine/Source/Runtime/RHI/Public/RHICommandList.h
@@ -456,6 +456,9 @@ public:
 	void PreDispatch(class FRHIComputeCommandList& InCommandList);
 	void PreDraw(class FRHICommandList& InCommandList);
 
+	void FlushAllParameters(class FRHIComputeCommandList& InCommandList);
+	void FlushAllParameters(class FRHICommandList& InCommandList);
+
 	void FlushPendingParameters(class FRHIComputeCommandList& InCommandList, FRHIComputeShader* InShader);
 	void FlushPendingParameters(class FRHICommandList& InCommandList, FRHIGraphicsShader* InShader);
 
