Prevent crash when saving user PSO cache file on exit.

diff --git a/Engine/Source/Runtime/RHI/Private/PipelineFileCache.cpp b/Engine/Source/Runtime/RHI/Private/PipelineFileCache.cpp
--- a/Engine/Source/Runtime/RHI/Private/PipelineFileCache.cpp
+++ b/Engine/Source/Runtime/RHI/Private/PipelineFileCache.cpp
@@ -3330,8 +3330,7 @@ bool FPipelineFileCacheManager::SavePipelineFileCache(SaveMode Mode)
 		FRWScopeLock Lock(FileCacheLock, SLT_Write);
 
 		FPipelineCacheFile* UserCache = GetPipelineCacheFileFromKey(UserCacheKey);
-		check(UserCache);
-		bOk = UserCache->SavePipelineFileCache(Mode, Stats, NewPSOs, RequestedOrder, NewPSOUsage);
+		bOk = (UserCache != nullptr) && UserCache->SavePipelineFileCache(Mode, Stats, NewPSOs, RequestedOrder, NewPSOUsage);
 		// If successful clear new PSO's as they should have been saved out
 		// Leave everything else in-tact (e.g stats) for subsequent in place save operations
 		if (bOk)
