Do not add PSOs being precompiled from the cache as cmdlist dependencies.

- This greatly reduces hitching while navigating the UI.

diff --git a/Engine/Source/Runtime/RHI/Private/PipelineStateCache.cpp b/Engine/Source/Runtime/RHI/Private/PipelineStateCache.cpp
--- a/Engine/Source/Runtime/RHI/Private/PipelineStateCache.cpp
+++ b/Engine/Source/Runtime/RHI/Private/PipelineStateCache.cpp
@@ -101,7 +101,7 @@ static TAutoConsoleVariable<int32> CVarRTPSOCacheSize(
 );
 #endif // RHI_RAYTRACING
 
-int32 GPSOPrecaching = 0;
+int32 GPSOPrecaching = 1;
 static FAutoConsoleVariableRef CVarPSOPrecaching(
 	TEXT("r.PSOPrecaching"),
 	GPSOPrecaching,
@@ -1489,7 +1489,8 @@ public:
 		// On Mac the compilation is handled using external processes, so engine threads have very little work to do
 		// and it's better to leave more CPU time to these extrenal processes and other engine threads.
 		// Also use background threads for PSO precaching when the PSO thread pool is not used
-		return (PLATFORM_MAC || PSOPreCacheResult == EPSOPrecacheResult::Active) ? ENamedThreads::AnyBackgroundThreadNormalTask : DesiredThread;
+		// Compute pipelines usually take much longer to compile, compile them on background thread as well.
+		return (PLATFORM_MAC || PSOPreCacheResult == EPSOPrecacheResult::Active || (Pipeline && Pipeline->IsCompute())) ? ENamedThreads::AnyBackgroundThreadNormalTask : DesiredThread;
 	}
 };
 
@@ -2139,8 +2140,8 @@ FGraphicsPipelineState* PipelineStateCache::GetAndOrCreateGraphicsPipelineState(
 		FGraphEventRef GraphEvent = CreateGraphicsPipelineState(Initializer, PSOPrecacheResult, DoAsyncCompile, bPSOPrecache, OutCachedState);
 
 		// Add dispatch pre requisite for non precaching jobs only
-		// NOTE: do we need to add dispatch prerequisite for PSO precache when precache pool is disabled and it's using regular task graph?
-		if (GraphEvent.IsValid() && (!bPSOPrecache || !FPSOPrecacheThreadPool::UsePool()))
+		//if (GraphEvent.IsValid() && (!bPSOPrecache || !FPSOPrecacheThreadPool::UsePool()))
+		if (GraphEvent.IsValid() && !bPSOPrecache)
 		{
 			check(DoAsyncCompile);
 			RHICmdList.AddDispatchPrerequisite(GraphEvent);
