Fixing bad parameter flushes. Adding another parameter flush to SwitchPipeline.

diff --git a/Engine/Source/Runtime/RHI/Private/RHICommandList.cpp b/Engine/Source/Runtime/RHI/Private/RHICommandList.cpp
--- a/Engine/Source/Runtime/RHI/Private/RHICommandList.cpp
+++ b/Engine/Source/Runtime/RHI/Private/RHICommandList.cpp
@@ -346,8 +346,15 @@ void FRHICommandListBase::FinishRecording()
 {
 	checkf(!IsImmediate(), TEXT("Do not call FinishRecording() on the immediate RHI command list."));
 
-	// Make sure the batcher is cleared out. If there are pending graphics parameters, this cast is safe.
-	ParameterBatcher.FlushAllParameters(static_cast<FRHICommandList&>(*this));
+	if (IsGraphics())
+	{
+		// Make sure the batcher is cleared out. If there are pending graphics parameters, this cast is safe.
+		ParameterBatcher.FlushAllParameters(static_cast<FRHICommandList&>(*this));
+	}
+	else
+	{
+		ParameterBatcher.FlushAllParameters(static_cast<FRHIComputeCommandList&>(*this));
+	}
 
 	PersistentState.FenceCandidate->Fence = PersistentState.RHIThreadBufferLockFence;
 
@@ -415,6 +422,15 @@ ERHIPipeline FRHICommandListBase::SwitchPipeline(ERHIPipeline Pipeline)
 	Exchange(ActivePipeline, Pipeline);
 	if (ActivePipeline != Pipeline)
 	{
+		if (Pipeline == ERHIPipeline::Graphics)
+		{
+			ParameterBatcher.FlushAllParameters(static_cast<FRHICommandList&>(*this));
+		}
+		else if (Pipeline == ERHIPipeline::AsyncCompute)
+		{
+			ParameterBatcher.FlushAllParameters(static_cast<FRHIComputeCommandList&>(*this));
+		}
+
 		EnqueueLambda([NewPipeline = ActivePipeline](FRHICommandListBase& ExecutingCmdList)
 		{
 			ExecutingCmdList.ActivePipeline = NewPipeline;
