Disabling parameter batching for the time being. There's too many corner cases on flushing that we need to take care of.
This requires us to restore the previously removed commands that the batcher removed.

diff --git a/Engine/Source/Runtime/RHI/Private/RHICommandList.cpp b/Engine/Source/Runtime/RHI/Private/RHICommandList.cpp
--- a/Engine/Source/Runtime/RHI/Private/RHICommandList.cpp
+++ b/Engine/Source/Runtime/RHI/Private/RHICommandList.cpp
@@ -136,7 +136,18 @@ static TStatId GCurrentExecuteStat;
 static FCriticalSection GRHIThreadOnTasksCritical;
 static std::atomic<int32> GRHIThreadStallRequestCount;
 
-FRHIParameterBatcher::FRHIParameterBatcher() = default;
+int32 GRHICommandParameterBatching = 0;
+static FAutoConsoleVariableRef CVarRHICommandParameterBatching(
+	TEXT("r.RHICmdParamBatching"),
+	GRHICommandParameterBatching,
+	TEXT("Controls if parameters can be batched in RHI command lists."),
+	ECVF_ReadOnly
+);
+
+FRHIParameterBatcher::FRHIParameterBatcher()
+{
+	bEnabled = GRHICommandParameterBatching != 0;
+}
 FRHIParameterBatcher::FRHIParameterBatcher(FRHIParameterBatcher&&) = default;
 FRHIParameterBatcher::~FRHIParameterBatcher()
 {
diff --git a/Engine/Source/Runtime/RHI/Public/RHICommandList.h b/Engine/Source/Runtime/RHI/Public/RHICommandList.h
--- a/Engine/Source/Runtime/RHI/Public/RHICommandList.h
+++ b/Engine/Source/Runtime/RHI/Public/RHICommandList.h
@@ -432,6 +432,7 @@ class RHI_API FRHIParameterBatcher
 protected:
 	FRHIShader* AllBatchedShaders[SF_NumStandardFrequencies]{};
 	FRHIBatchedShaderParameters AllBatchedShaderParameters[SF_NumStandardFrequencies];
+	bool bEnabled = false;
 
 	FRHIGraphicsShader* GetBatchedGraphicsShader(int32 Index) { return static_cast<FRHIGraphicsShader*>(AllBatchedShaders[Index]); }
 	FRHIComputeShader*  GetBatchedComputeShader()             { return static_cast<FRHIComputeShader*>(AllBatchedShaders[SF_Compute]);   }
@@ -450,6 +451,8 @@ public:
 	FRHIParameterBatcher(FRHIParameterBatcher&&);
 	~FRHIParameterBatcher();
 
+	FORCEINLINE bool IsEnabled() const { return bEnabled; }
+
 	void OnBoundShaderChanged(class FRHICommandList& InCommandList, const FBoundShaderStateInput& InBoundShaderStateInput);
 	void OnBoundShaderChanged(class FRHIComputeCommandList& InCommandList, FRHIComputeShader* InBoundComputeShaderRHI);
 
@@ -1181,6 +1184,24 @@ FRHICOMMAND_MACRO_TPL(TRHIShader, FRHICommandSetShaderParameters)
 	RHI_API void Execute(FRHICommandListBase& CmdList);
 };
 
+FRHICOMMAND_MACRO_TPL(TRHIShader, FRHICommandSetShaderParameter)
+{
+	TRHIShader* Shader;
+	const void* NewValue;
+	uint32 BufferIndex;
+	uint32 BaseIndex;
+	uint32 NumBytes;
+	FORCEINLINE_DEBUGGABLE FRHICommandSetShaderParameter(TRHIShader* InShader, uint32 InBufferIndex, uint32 InBaseIndex, uint32 InNumBytes, const void* InNewValue)
+		: Shader(InShader)
+		, NewValue(InNewValue)
+		, BufferIndex(InBufferIndex)
+		, BaseIndex(InBaseIndex)
+		, NumBytes(InNumBytes)
+	{
+	}
+	RHI_API void Execute(FRHICommandListBase& CmdList);
+};
+
 FRHICOMMAND_MACRO_TPL(TRHIShader, FRHICommandSetShaderUniformBuffer)
 {
 	TRHIShader* Shader;
@@ -1195,6 +1216,78 @@ FRHICOMMAND_MACRO_TPL(TRHIShader, FRHICommandSetShaderUniformBuffer)
 	RHI_API void Execute(FRHICommandListBase& CmdList);
 };
 
+FRHICOMMAND_MACRO_TPL(TRHIShader, FRHICommandSetShaderTexture)
+{
+	TRHIShader* Shader;
+	uint32 TextureIndex;
+	FRHITexture* Texture;
+	FORCEINLINE_DEBUGGABLE FRHICommandSetShaderTexture(TRHIShader* InShader, uint32 InTextureIndex, FRHITexture* InTexture)
+		: Shader(InShader)
+		, TextureIndex(InTextureIndex)
+		, Texture(InTexture)
+	{
+	}
+	RHI_API void Execute(FRHICommandListBase& CmdList);
+};
+
+FRHICOMMAND_MACRO_TPL(TRHIShader, FRHICommandSetShaderResourceViewParameter)
+{
+	TRHIShader* Shader;
+	uint32 SamplerIndex;
+	FRHIShaderResourceView* SRV;
+	FORCEINLINE_DEBUGGABLE FRHICommandSetShaderResourceViewParameter(TRHIShader* InShader, uint32 InSamplerIndex, FRHIShaderResourceView* InSRV)
+		: Shader(InShader)
+		, SamplerIndex(InSamplerIndex)
+		, SRV(InSRV)
+	{
+	}
+	RHI_API void Execute(FRHICommandListBase& CmdList);
+};
+
+FRHICOMMAND_MACRO_TPL(TRHIShader, FRHICommandSetUAVParameter)
+{
+	TRHIShader* Shader;
+	uint32 UAVIndex;
+	FRHIUnorderedAccessView* UAV;
+	FORCEINLINE_DEBUGGABLE FRHICommandSetUAVParameter(TRHIShader* InShader, uint32 InUAVIndex, FRHIUnorderedAccessView* InUAV)
+		: Shader(InShader)
+		, UAVIndex(InUAVIndex)
+		, UAV(InUAV)
+	{
+	}
+	RHI_API void Execute(FRHICommandListBase& CmdList);
+};
+
+FRHICOMMAND_MACRO(FRHICommandSetUAVParameter_InitialCount)
+{
+	FRHIComputeShader* Shader;
+	uint32 UAVIndex;
+	FRHIUnorderedAccessView* UAV;
+	uint32 InitialCount;
+	FORCEINLINE_DEBUGGABLE FRHICommandSetUAVParameter_InitialCount(FRHIComputeShader* InShader, uint32 InUAVIndex, FRHIUnorderedAccessView* InUAV, uint32 InInitialCount)
+		: Shader(InShader)
+		, UAVIndex(InUAVIndex)
+		, UAV(InUAV)
+		, InitialCount(InInitialCount)
+	{
+	}
+	RHI_API void Execute(FRHICommandListBase& CmdList);
+};
+
+FRHICOMMAND_MACRO_TPL(TRHIShader, FRHICommandSetShaderSampler)
+{
+	TRHIShader* Shader;
+	uint32 SamplerIndex;
+	FRHISamplerState* Sampler;
+	FORCEINLINE_DEBUGGABLE FRHICommandSetShaderSampler(TRHIShader* InShader, uint32 InSamplerIndex, FRHISamplerState* InSampler)
+		: Shader(InShader)
+		, SamplerIndex(InSamplerIndex)
+		, Sampler(InSampler)
+	{
+	}
+	RHI_API void Execute(FRHICommandListBase& CmdList);
+};
+
 FRHICOMMAND_MACRO(FRHICommandDrawPrimitive)
 {
 	uint32 BaseVertexIndex;
@@ -2257,7 +2350,12 @@ FRHICOMMAND_MACRO(FRHICommandSetRayTracingBindings)
 #endif // RHI_RAYTRACING
 
 template<> RHI_API void FRHICommandSetShaderParameters           <FRHIComputeShader>::Execute(FRHICommandListBase& CmdList);
+template<> RHI_API void FRHICommandSetShaderParameter            <FRHIComputeShader>::Execute(FRHICommandListBase& CmdList);
 template<> RHI_API void FRHICommandSetShaderUniformBuffer        <FRHIComputeShader>::Execute(FRHICommandListBase& CmdList);
+template<> RHI_API void FRHICommandSetShaderTexture              <FRHIComputeShader>::Execute(FRHICommandListBase& CmdList);
+template<> RHI_API void FRHICommandSetShaderResourceViewParameter<FRHIComputeShader>::Execute(FRHICommandListBase& CmdList);
+template<> RHI_API void FRHICommandSetShaderSampler              <FRHIComputeShader>::Execute(FRHICommandListBase& CmdList);
+template<> RHI_API void FRHICommandSetUAVParameter               <FRHIComputeShader>::Execute(FRHICommandListBase& CmdList);
 
 extern RHI_API FRHIComputePipelineState* ExecuteSetComputePipelineState(FComputePipelineState* ComputePipelineState);
 extern RHI_API FRHIGraphicsPipelineState* ExecuteSetGraphicsPipelineState(class FGraphicsPipelineState* GraphicsPipelineState);
@@ -2417,7 +2515,20 @@ public:
 	FORCEINLINE_DEBUGGABLE void SetShaderParameter(FRHIComputeShader* Shader, uint32 BufferIndex, uint32 BaseIndex, uint32 NumBytes, const void* NewValue)
 	{
 		ValidateBoundShader(Shader);
-		ParameterBatcher.SetShaderParameter(Shader, BufferIndex, BaseIndex, NumBytes, NewValue);
+		if (ParameterBatcher.IsEnabled())
+		{
+			ParameterBatcher.SetShaderParameter(Shader, BufferIndex, BaseIndex, NumBytes, NewValue);
+			return;
+		}
+
+		if (Bypass())
+		{
+			GetComputeContext().RHISetShaderParameter(Shader, BufferIndex, BaseIndex, NumBytes, NewValue);
+			return;
+		}
+		void* UseValue = Alloc(NumBytes, 16);
+		FMemory::Memcpy(UseValue, NewValue, NumBytes);
+		ALLOC_COMMAND(FRHICommandSetShaderParameter<FRHIComputeShader>)(Shader, BufferIndex, BaseIndex, NumBytes, UseValue);
 	}
 
 	FORCEINLINE void SetShaderParameter(FComputeShaderRHIRef& Shader, uint32 BufferIndex, uint32 BaseIndex, uint32 NumBytes, const void* NewValue)
@@ -2428,13 +2539,35 @@ public:
 	FORCEINLINE_DEBUGGABLE void SetShaderTexture(FRHIComputeShader* Shader, uint32 TextureIndex, FRHITexture* Texture)
 	{
 		ValidateBoundShader(Shader);
-		ParameterBatcher.SetShaderTexture(Shader, TextureIndex, Texture);
+		if (ParameterBatcher.IsEnabled())
+		{
+			ParameterBatcher.SetShaderTexture(Shader, TextureIndex, Texture);
+			return;
+		}
+
+		if (Bypass())
+		{
+			GetComputeContext().RHISetShaderTexture(Shader, TextureIndex, Texture);
+			return;
+		}
+		ALLOC_COMMAND(FRHICommandSetShaderTexture<FRHIComputeShader>)(Shader, TextureIndex, Texture);
 	}
 
 	FORCEINLINE_DEBUGGABLE void SetShaderResourceViewParameter(FRHIComputeShader* Shader, uint32 SamplerIndex, FRHIShaderResourceView* SRV)
 	{
 		ValidateBoundShader(Shader);
-		ParameterBatcher.SetShaderResourceViewParameter(Shader, SamplerIndex, SRV);
+		if (ParameterBatcher.IsEnabled())
+		{
+			ParameterBatcher.SetShaderResourceViewParameter(Shader, SamplerIndex, SRV);
+			return;
+		}
+
+		if (Bypass())
+		{
+			GetComputeContext().RHISetShaderResourceViewParameter(Shader, SamplerIndex, SRV);
+			return;
+		}
+		ALLOC_COMMAND(FRHICommandSetShaderResourceViewParameter<FRHIComputeShader>)(Shader, SamplerIndex, SRV);
 	}
 
 	FORCEINLINE_DEBUGGABLE void SetShaderSampler(FRHIComputeShader* Shader, uint32 SamplerIndex, FRHISamplerState* State)
@@ -2446,19 +2579,54 @@ public:
 			return;
 		}
 
-		ParameterBatcher.SetShaderSampler(Shader, SamplerIndex, State);
+		if (ParameterBatcher.IsEnabled())
+		{
+			ParameterBatcher.SetShaderSampler(Shader, SamplerIndex, State);
+			return;
+		}
+
+		if (Bypass())
+		{
+			GetComputeContext().RHISetShaderSampler(Shader, SamplerIndex, State);
+			return;
+		}
+		ALLOC_COMMAND(FRHICommandSetShaderSampler<FRHIComputeShader>)(Shader, SamplerIndex, State);
 	}
 
 	FORCEINLINE_DEBUGGABLE void SetUAVParameter(FRHIComputeShader* Shader, uint32 UAVIndex, FRHIUnorderedAccessView* UAV)
 	{
 		ValidateBoundShader(Shader);
-		ParameterBatcher.SetUAVParameter(Shader, UAVIndex, UAV);
+
+		if (ParameterBatcher.IsEnabled())
+		{
+			ParameterBatcher.SetUAVParameter(Shader, UAVIndex, UAV);
+			return;
+		}
+
+		if (Bypass())
+		{
+			GetComputeContext().RHISetUAVParameter(Shader, UAVIndex, UAV);
+			return;
+		}
+		ALLOC_COMMAND(FRHICommandSetUAVParameter<FRHIComputeShader>)(Shader, UAVIndex, UAV);
 	}
 
 	FORCEINLINE_DEBUGGABLE void SetUAVParameter(FRHIComputeShader* Shader, uint32 UAVIndex, FRHIUnorderedAccessView* UAV, uint32 InitialCount)
 	{
 		ValidateBoundShader(Shader);
-		ParameterBatcher.SetUAVParameter(Shader, UAVIndex, UAV, InitialCount);
+
+		if (ParameterBatcher.IsEnabled())
+		{
+			ParameterBatcher.SetUAVParameter(Shader, UAVIndex, UAV, InitialCount);
+			return;
+		}
+
+		if (Bypass())
+		{
+			GetComputeContext().RHISetUAVParameter(Shader, UAVIndex, UAV, InitialCount);
+			return;
+		}
+		ALLOC_COMMAND(FRHICommandSetUAVParameter_InitialCount)(Shader, UAVIndex, UAV, InitialCount);
 	}
 
 	UE_DEPRECATED(5.1, "ComputePipelineStates should be used instead of direct ComputeShaders. You can use SetComputePipelineState(RHICmdList, ComputeShader).")
@@ -2979,7 +3147,12 @@ PRAGMA_ENABLE_DEPRECATION_WARNINGS
 };
 
 template<> RHI_API void FRHICommandSetShaderParameters           <FRHIGraphicsShader>::Execute(FRHICommandListBase& CmdList);
+template<> RHI_API void FRHICommandSetShaderParameter            <FRHIGraphicsShader>::Execute(FRHICommandListBase& CmdList);
 template<> RHI_API void FRHICommandSetShaderUniformBuffer        <FRHIGraphicsShader>::Execute(FRHICommandListBase& CmdList);
+template<> RHI_API void FRHICommandSetShaderTexture              <FRHIGraphicsShader>::Execute(FRHICommandListBase& CmdList);
+template<> RHI_API void FRHICommandSetShaderResourceViewParameter<FRHIGraphicsShader>::Execute(FRHICommandListBase& CmdList);
+template<> RHI_API void FRHICommandSetShaderSampler              <FRHIGraphicsShader>::Execute(FRHICommandListBase& CmdList);
+template<> RHI_API void FRHICommandSetUAVParameter               <FRHIPixelShader   >::Execute(FRHICommandListBase& CmdList);
 
 class RHI_API FRHICommandList : public FRHIComputeCommandList
 {
@@ -3150,7 +3323,21 @@ public:
 	{
 		//check(IsOutsideRenderPass());
 		ValidateBoundShader(Shader);
-		ParameterBatcher.SetShaderParameter(Shader, BufferIndex, BaseIndex, NumBytes, NewValue);
+
+		if (ParameterBatcher.IsEnabled())
+		{
+			ParameterBatcher.SetShaderParameter(Shader, BufferIndex, BaseIndex, NumBytes, NewValue);
+			return;
+		}
+
+		if (Bypass())
+		{
+			GetContext().RHISetShaderParameter(Shader, BufferIndex, BaseIndex, NumBytes, NewValue);
+			return;
+		}
+		void* UseValue = Alloc(NumBytes, 16);
+		FMemory::Memcpy(UseValue, NewValue, NumBytes);
+		ALLOC_COMMAND(FRHICommandSetShaderParameter<FRHIGraphicsShader>)(Shader, BufferIndex, BaseIndex, NumBytes, UseValue);
 	}
 
 	template <typename TShaderRHI>
@@ -3165,7 +3352,19 @@ public:
 	{
 		//check(IsOutsideRenderPass());
 		ValidateBoundShader(Shader);
-		ParameterBatcher.SetShaderTexture(Shader, TextureIndex, Texture);
+
+		if (ParameterBatcher.IsEnabled())
+		{
+			ParameterBatcher.SetShaderTexture(Shader, TextureIndex, Texture);
+			return;
+		}
+
+		if (Bypass())
+		{
+			GetContext().RHISetShaderTexture(Shader, TextureIndex, Texture);
+			return;
+		}
+		ALLOC_COMMAND(FRHICommandSetShaderTexture<FRHIGraphicsShader>)(Shader, TextureIndex, Texture);
 	}
 
 	template <typename TShaderRHI>
@@ -3180,7 +3379,19 @@ public:
 	{
 		//check(IsOutsideRenderPass());
 		ValidateBoundShader(Shader);
-		ParameterBatcher.SetShaderResourceViewParameter(Shader, SamplerIndex, SRV);
+
+		if (ParameterBatcher.IsEnabled())
+		{
+			ParameterBatcher.SetShaderResourceViewParameter(Shader, SamplerIndex, SRV);
+			return;
+		}
+
+		if (Bypass())
+		{
+			GetContext().RHISetShaderResourceViewParameter(Shader, SamplerIndex, SRV);
+			return;
+		}
+		ALLOC_COMMAND(FRHICommandSetShaderResourceViewParameter<FRHIGraphicsShader>)(Shader, SamplerIndex, SRV);
 	}
 
 	template <typename TShaderRHI>
@@ -3203,7 +3414,18 @@ public:
 			return;
 		}
 
-		ParameterBatcher.SetShaderSampler(Shader, SamplerIndex, State);
+		if (ParameterBatcher.IsEnabled())
+		{
+			ParameterBatcher.SetShaderSampler(Shader, SamplerIndex, State);
+			return;
+		}
+
+		if (Bypass())
+		{
+			GetContext().RHISetShaderSampler(Shader, SamplerIndex, State);
+			return;
+		}
+		ALLOC_COMMAND(FRHICommandSetShaderSampler<FRHIGraphicsShader>)(Shader, SamplerIndex, State);
 	}
 
 	template <typename TShaderRHI>
@@ -3217,7 +3439,19 @@ public:
 	FORCEINLINE_DEBUGGABLE void SetUAVParameter(FRHIPixelShader* Shader, uint32 UAVIndex, FRHIUnorderedAccessView* UAV)
 	{
 		ValidateBoundShader(Shader);
-		ParameterBatcher.SetUAVParameter(Shader, UAVIndex, UAV);
+
+		if (ParameterBatcher.IsEnabled())
+		{
+			ParameterBatcher.SetUAVParameter(Shader, UAVIndex, UAV);
+			return;
+		}
+
+		if (Bypass())
+		{
+			GetContext().RHISetUAVParameter(Shader, UAVIndex, UAV);
+			return;
+		}
+		ALLOC_COMMAND(FRHICommandSetUAVParameter<FRHIPixelShader>)(Shader, UAVIndex, UAV);
 	}
 
 	FORCEINLINE_DEBUGGABLE void SetUAVParameter(const TRefCountPtr<FRHIPixelShader>& Shader, uint32 UAVIndex, FRHIUnorderedAccessView* UAV)
diff --git a/Engine/Source/Runtime/RHI/Public/RHICommandListCommandExecutes.inl b/Engine/Source/Runtime/RHI/Public/RHICommandListCommandExecutes.inl
--- a/Engine/Source/Runtime/RHI/Public/RHICommandListCommandExecutes.inl
+++ b/Engine/Source/Runtime/RHI/Public/RHICommandListCommandExecutes.inl
@@ -155,6 +155,18 @@ template<> RHI_API void FRHICommandSetShaderParameters<FRHIGraphicsShader>::Exec
 	INTERNAL_DECORATOR(RHISetShaderParameters)(Shader, ParametersData, Parameters, ResourceParameters, BindlessParameters);
 }
 
+template<> RHI_API void FRHICommandSetShaderParameter<FRHIComputeShader>::Execute(FRHICommandListBase& CmdList)
+{
+	RHISTAT(SetShaderParameter);
+	INTERNAL_DECORATOR_COMPUTE(RHISetShaderParameter)(Shader, BufferIndex, BaseIndex, NumBytes, NewValue);
+}
+
+template<> RHI_API void FRHICommandSetShaderParameter<FRHIGraphicsShader>::Execute(FRHICommandListBase& CmdList)
+{
+	RHISTAT(SetShaderParameter);
+	INTERNAL_DECORATOR(RHISetShaderParameter)(Shader, BufferIndex, BaseIndex, NumBytes, NewValue);
+}
+
 template<> RHI_API void FRHICommandSetShaderUniformBuffer<FRHIComputeShader>::Execute(FRHICommandListBase& CmdList)
 {
 	RHISTAT(SetShaderUniformBuffer);
@@ -167,6 +179,60 @@ template<> RHI_API void FRHICommandSetShaderUniformBuffer<FRHIGraphicsShader>::E
 	INTERNAL_DECORATOR(RHISetShaderUniformBuffer)(Shader, BaseIndex, UniformBuffer);
 }
 
+template<> RHI_API void FRHICommandSetShaderTexture<FRHIComputeShader>::Execute(FRHICommandListBase& CmdList)
+{
+	RHISTAT(SetShaderTexture);
+	INTERNAL_DECORATOR_COMPUTE(RHISetShaderTexture)(Shader, TextureIndex, Texture);
+}
+
+template<> RHI_API void FRHICommandSetShaderTexture<FRHIGraphicsShader>::Execute(FRHICommandListBase& CmdList)
+{
+	RHISTAT(SetShaderTexture);
+	INTERNAL_DECORATOR(RHISetShaderTexture)(Shader, TextureIndex, Texture);
+}
+
+template<> RHI_API void FRHICommandSetShaderResourceViewParameter<FRHIComputeShader>::Execute(FRHICommandListBase& CmdList)
+{
+	RHISTAT(SetShaderResourceViewParameter);
+	INTERNAL_DECORATOR_COMPUTE(RHISetShaderResourceViewParameter)(Shader, SamplerIndex, SRV);
+}
+
+template<> RHI_API void FRHICommandSetShaderResourceViewParameter<FRHIGraphicsShader>::Execute(FRHICommandListBase& CmdList)
+{
+	RHISTAT(SetShaderResourceViewParameter);
+	INTERNAL_DECORATOR(RHISetShaderResourceViewParameter)(Shader, SamplerIndex, SRV);
+}
+
+template<> RHI_API void FRHICommandSetUAVParameter<FRHIComputeShader>::Execute(FRHICommandListBase& CmdList)
+{
+	RHISTAT(SetUAVParameter);
+	INTERNAL_DECORATOR_COMPUTE(RHISetUAVParameter)(Shader, UAVIndex, UAV);
+}
+
+template<> RHI_API void FRHICommandSetUAVParameter<FRHIPixelShader>::Execute(FRHICommandListBase& CmdList)
+{
+	RHISTAT(SetUAVParameter);
+	INTERNAL_DECORATOR(RHISetUAVParameter)(Shader, UAVIndex, UAV);
+}
+
+void FRHICommandSetUAVParameter_InitialCount::Execute(FRHICommandListBase& CmdList)
+{
+	RHISTAT(SetUAVParameter);
+	INTERNAL_DECORATOR_COMPUTE(RHISetUAVParameter)(Shader, UAVIndex, UAV, InitialCount);
+}
+
+template<> RHI_API void FRHICommandSetShaderSampler<FRHIComputeShader>::Execute(FRHICommandListBase& CmdList)
+{
+	RHISTAT(SetShaderSampler);
+	INTERNAL_DECORATOR_COMPUTE(RHISetShaderSampler)(Shader, SamplerIndex, Sampler);
+}
+
+template<> RHI_API void FRHICommandSetShaderSampler<FRHIGraphicsShader>::Execute(FRHICommandListBase& CmdList)
+{
+	RHISTAT(SetShaderSampler);
+	INTERNAL_DECORATOR(RHISetShaderSampler)(Shader, SamplerIndex, Sampler);
+}
+
 void FRHICommandDrawPrimitive::Execute(FRHICommandListBase& CmdList)
 {
 	RHISTAT(DrawPrimitive);
