Minimal version of RHICopySharedMips for a platform it is required on
 - Implemented using a common RHI function that implements the copy using RHICopyTexture (all platforms will use this function in a future change, and RHICopySharedMips will be removed).

diff --git a/Engine/Source/Runtime/Apple/MetalRHI/Private/MetalDynamicRHI.h b/Engine/Source/Runtime/Apple/MetalRHI/Private/MetalDynamicRHI.h
--- a/Engine/Source/Runtime/Apple/MetalRHI/Private/MetalDynamicRHI.h
+++ b/Engine/Source/Runtime/Apple/MetalRHI/Private/MetalDynamicRHI.h
@@ -73,7 +73,7 @@ public:
 	virtual bool RHIGetTextureMemoryVisualizeData(FColor* TextureData,int32 SizeX,int32 SizeY,int32 Pitch,int32 PixelSize) final override;
 	virtual FTextureRHIRef RHICreateTexture(const FRHITextureCreateDesc& CreateDesc) final override;
 	virtual FTexture2DRHIRef RHIAsyncCreateTexture2D(uint32 SizeX, uint32 SizeY, uint8 Format, uint32 NumMips, ETextureCreateFlags Flags, ERHIAccess InResourceState, void** InitialMipData, uint32 NumInitialMips) final override;
-	virtual void RHICopySharedMips(FRHITexture2D* DestTexture2D, FRHITexture2D* SrcTexture2D) final override;
+	virtual void RHICopySharedMips(FRHICommandList& RHICmdList, FRHITexture2D* DestTexture2D, FRHITexture2D* SrcTexture2D) final override;
 	virtual void RHIGenerateMips(FRHITexture* Texture) final override;
 	virtual uint32 RHIComputeMemorySize(FRHITexture* TextureRHI) final override;
 	virtual FTexture2DRHIRef RHIAsyncReallocateTexture2D(FRHITexture2D* Texture2D, int32 NewMipCount, int32 NewSizeX, int32 NewSizeY, FThreadSafeCounter* RequestStatus) final override;
diff --git a/Engine/Source/Runtime/Apple/MetalRHI/Private/MetalTexture.cpp b/Engine/Source/Runtime/Apple/MetalRHI/Private/MetalTexture.cpp
--- a/Engine/Source/Runtime/Apple/MetalRHI/Private/MetalTexture.cpp
+++ b/Engine/Source/Runtime/Apple/MetalRHI/Private/MetalTexture.cpp
@@ -1521,7 +1521,7 @@ FTexture2DRHIRef FMetalDynamicRHI::RHIAsyncCreateTexture2D(uint32 SizeX, uint32
 	return FTexture2DRHIRef();
 }
 
-void FMetalDynamicRHI::RHICopySharedMips(FRHITexture2D* DestTexture2D, FRHITexture2D* SrcTexture2D)
+void FMetalDynamicRHI::RHICopySharedMips(FRHICommandList& RHICmdList, FRHITexture2D* DestTexture2D, FRHITexture2D* SrcTexture2D)
 {
 	NOT_SUPPORTED("RHICopySharedMips");
 }
diff --git a/Engine/Source/Runtime/D3D12RHI/Private/D3D12RHIPrivate.h b/Engine/Source/Runtime/D3D12RHI/Private/D3D12RHIPrivate.h
--- a/Engine/Source/Runtime/D3D12RHI/Private/D3D12RHIPrivate.h
+++ b/Engine/Source/Runtime/D3D12RHI/Private/D3D12RHIPrivate.h
@@ -466,7 +466,7 @@ public:
 	virtual bool RHIGetTextureMemoryVisualizeData(FColor* TextureData, int32 SizeX, int32 SizeY, int32 Pitch, int32 PixelSize) final override;
 	virtual FTextureRHIRef RHIAsyncCreateTexture2D(uint32 SizeX, uint32 SizeY, uint8 Format, uint32 NumMips, ETextureCreateFlags Flags, ERHIAccess InResourceState, void** InitialMipData, uint32 NumInitialMips) final override;
 	virtual FTextureRHIRef RHICreateTexture(const FRHITextureCreateDesc& CreateDesc) override;
-	virtual void RHICopySharedMips(FRHITexture2D* DestTexture2D, FRHITexture2D* SrcTexture2D) final override;
+	virtual void RHICopySharedMips(FRHICommandList& RHICmdList, FRHITexture2D* DestTexture2D, FRHITexture2D* SrcTexture2D) final override;
 	virtual FShaderResourceViewRHIRef RHICreateShaderResourceView(FRHITexture* Texture, const FRHITextureSRVCreateInfo& CreateInfo) final override;
 	virtual uint32 RHIComputeMemorySize(FRHITexture* TextureRHI) final override;
 	virtual FTexture2DRHIRef RHIAsyncReallocateTexture2D(FRHITexture2D* Texture2D, int32 NewMipCount, int32 NewSizeX, int32 NewSizeY, FThreadSafeCounter* RequestStatus) override;
diff --git a/Engine/Source/Runtime/D3D12RHI/Private/D3D12Texture.cpp b/Engine/Source/Runtime/D3D12RHI/Private/D3D12Texture.cpp
--- a/Engine/Source/Runtime/D3D12RHI/Private/D3D12Texture.cpp
+++ b/Engine/Source/Runtime/D3D12RHI/Private/D3D12Texture.cpp
@@ -1169,7 +1169,7 @@ FTextureRHIRef FD3D12DynamicRHI::RHIAsyncCreateTexture2D(uint32 SizeX, uint32 Si
 }
 
 
-void FD3D12DynamicRHI::RHICopySharedMips(FRHITexture2D* DestTexture2DRHI, FRHITexture2D* SrcTexture2DRHI)
+void FD3D12DynamicRHI::RHICopySharedMips(FRHICommandList& RHICmdList, FRHITexture2D* DestTexture2DRHI, FRHITexture2D* SrcTexture2DRHI)
 {
 	FD3D12Texture*  DestTexture2D = GetD3D12TextureFromRHITexture(DestTexture2DRHI);
 	FD3D12Texture*  SrcTexture2D = GetD3D12TextureFromRHITexture(SrcTexture2DRHI);
diff --git a/Engine/Source/Runtime/NullDrv/Public/NullRHI.h b/Engine/Source/Runtime/NullDrv/Public/NullRHI.h
--- a/Engine/Source/Runtime/NullDrv/Public/NullRHI.h
+++ b/Engine/Source/Runtime/NullDrv/Public/NullRHI.h
@@ -227,7 +227,7 @@ public:
 		return this->RHICreateTexture(Desc);
 	}
 
-	virtual void RHICopySharedMips(FRHITexture2D* DestTexture2D, FRHITexture2D* SrcTexture2D) final override
+	virtual void RHICopySharedMips(FRHICommandList& RHICmdList, FRHITexture2D* DestTexture2D, FRHITexture2D* SrcTexture2D) final override
 	{
 	}
 
diff --git a/Engine/Source/Runtime/OpenGLDrv/Private/OpenGLTexture.cpp b/Engine/Source/Runtime/OpenGLDrv/Private/OpenGLTexture.cpp
--- a/Engine/Source/Runtime/OpenGLDrv/Private/OpenGLTexture.cpp
+++ b/Engine/Source/Runtime/OpenGLDrv/Private/OpenGLTexture.cpp
@@ -1646,7 +1646,7 @@ FTextureRHIRef FOpenGLDynamicRHI::RHIAsyncCreateTexture2D(uint32 SizeX, uint32 S
 	return FTexture2DRHIRef();
 }
 
-void FOpenGLDynamicRHI::RHICopySharedMips(FRHITexture2D* DestTexture2D, FRHITexture2D* SrcTexture2D)
+void FOpenGLDynamicRHI::RHICopySharedMips(FRHICommandList& RHICmdList, FRHITexture2D* DestTexture2D, FRHITexture2D* SrcTexture2D)
 {
 	check(0);
 }
diff --git a/Engine/Source/Runtime/OpenGLDrv/Public/OpenGLDrv.h b/Engine/Source/Runtime/OpenGLDrv/Public/OpenGLDrv.h
--- a/Engine/Source/Runtime/OpenGLDrv/Public/OpenGLDrv.h
+++ b/Engine/Source/Runtime/OpenGLDrv/Public/OpenGLDrv.h
@@ -389,7 +389,7 @@ public:
 	virtual bool RHIGetTextureMemoryVisualizeData(FColor* TextureData,int32 SizeX,int32 SizeY,int32 Pitch,int32 PixelSize) final override;
 	virtual FTextureRHIRef RHICreateTexture(const FRHITextureCreateDesc& CreateDesc) final override;
 	virtual FTextureRHIRef RHIAsyncCreateTexture2D(uint32 SizeX, uint32 SizeY, uint8 Format, uint32 NumMips, ETextureCreateFlags Flags, ERHIAccess InResourceState, void** InitialMipData, uint32 NumInitialMips) final override;
-	virtual void RHICopySharedMips(FRHITexture2D* DestTexture2D, FRHITexture2D* SrcTexture2D) final override;
+	virtual void RHICopySharedMips(FRHICommandList& RHICmdList, FRHITexture2D* DestTexture2D, FRHITexture2D* SrcTexture2D) final override;
 	virtual FShaderResourceViewRHIRef RHICreateShaderResourceView(FRHITexture* Texture, const FRHITextureSRVCreateInfo& CreateInfo) final override;
 	virtual void RHIGenerateMips(FRHITexture* Texture) final override;
 	virtual uint32 RHIComputeMemorySize(FRHITexture* TextureRHI) final override;
diff --git a/Engine/Source/Runtime/RHI/Private/RHI.cpp b/Engine/Source/Runtime/RHI/Private/RHI.cpp
--- a/Engine/Source/Runtime/RHI/Private/RHI.cpp
+++ b/Engine/Source/Runtime/RHI/Private/RHI.cpp
@@ -3199,3 +3199,54 @@ void FDebugName::AppendString(FStringBuilderBase& Builder) const
 		Builder << '_' << Number;
 	}
 }
+
+// ------------------------------------------------------------------------------------------------------------------------------------------------
+// 
+// @todo move these functions into RHICore when the RHICmdList.CopySharedMips and RHICopySharedMips functions are removed (currently deprecated).
+//
+RHI_API void RHICopySharedMips_PRIVATE(FRHICommandList& RHICmdList, FRHITexture* SrcTexture, FRHITexture* DstTexture)
+{
+	FRHITextureDesc const& Desc = DstTexture->GetNumMips() < SrcTexture->GetNumMips()
+		? DstTexture->GetDesc()
+		: SrcTexture->GetDesc();
+
+	FRHICopyTextureInfo CopyInfo;
+	CopyInfo.Size.X         = Desc.Extent.X;
+	CopyInfo.Size.Y         = Desc.Extent.Y;
+	CopyInfo.Size.Z         = Desc.Depth;
+	CopyInfo.NumSlices      = Desc.ArraySize;
+	CopyInfo.NumMips        = Desc.NumMips;
+	CopyInfo.SourceMipIndex = SrcTexture->GetNumMips() - CopyInfo.NumMips;
+	CopyInfo.DestMipIndex   = DstTexture->GetNumMips() - CopyInfo.NumMips;
+
+	RHICmdList.CopyTexture(SrcTexture, DstTexture, CopyInfo);
+}
+
+RHI_API void RHICopySharedMips_AssumeSRVMaskState_PRIVATE(FRHICommandList& RHICmdList, FRHITexture* SrcTexture, FRHITexture* DstTexture)
+{
+	// Transition to copy source and dest
+	{
+		FRHITransitionInfo TransitionsBefore[] =
+		{
+			FRHITransitionInfo(SrcTexture, ERHIAccess::SRVMask, ERHIAccess::CopySrc ),
+			FRHITransitionInfo(DstTexture, ERHIAccess::SRVMask, ERHIAccess::CopyDest)
+		};
+		RHICmdList.Transition(MakeArrayView(TransitionsBefore, UE_ARRAY_COUNT(TransitionsBefore)));
+	}
+
+	RHICopySharedMips_PRIVATE(RHICmdList, SrcTexture, DstTexture);
+
+	// Transition to SRV
+	{
+		FRHITransitionInfo TransitionsAfter[] =
+		{
+			FRHITransitionInfo(SrcTexture, ERHIAccess::CopySrc , ERHIAccess::SRVMask),
+			FRHITransitionInfo(DstTexture, ERHIAccess::CopyDest, ERHIAccess::SRVMask)
+		};
+		RHICmdList.Transition(MakeArrayView(TransitionsAfter, UE_ARRAY_COUNT(TransitionsAfter)));
+	}
+}
+//
+// @todo move these functions into RHICore when the RHICmdList.CopySharedMips and RHICopySharedMips functions are removed (currently deprecated).
+// 
+// ------------------------------------------------------------------------------------------------------------------------------------------------
\ No newline at end of file
diff --git a/Engine/Source/Runtime/RHI/Public/DynamicRHI.h b/Engine/Source/Runtime/RHI/Public/DynamicRHI.h
--- a/Engine/Source/Runtime/RHI/Public/DynamicRHI.h
+++ b/Engine/Source/Runtime/RHI/Public/DynamicRHI.h
@@ -547,7 +547,7 @@ public:
 	* @param SrcTexture2D - source texture
 	*/
 	// FlushType: Flush RHI Thread
-	virtual void RHICopySharedMips(FRHITexture2D* DestTexture2D, FRHITexture2D* SrcTexture2D) = 0;
+	virtual void RHICopySharedMips(FRHICommandList& RHICmdList, FRHITexture2D* DestTexture2D, FRHITexture2D* SrcTexture2D) = 0;
 
 	/**
 	* Creates a shader resource view for a texture
diff --git a/Engine/Source/Runtime/RHI/Public/RHICommandList.h b/Engine/Source/Runtime/RHI/Public/RHICommandList.h
--- a/Engine/Source/Runtime/RHI/Public/RHICommandList.h
+++ b/Engine/Source/Runtime/RHI/Public/RHICommandList.h
@@ -4542,11 +4542,11 @@ public:
 		QUICK_SCOPE_CYCLE_COUNTER(STAT_RHIMETHOD_CopySharedMips_Flush);
 		DestTexture->AddRef();
 		SrcTexture->AddRef();
-		EnqueueLambda([DestTexture, SrcTexture](FRHICommandList&)
+		EnqueueLambda([DestTexture, SrcTexture](FRHICommandList& ExecutingCmdList)
 		{
 			LLM_SCOPE(ELLMTag::Textures);
 			//@todo move this function onto the IRHIComputeCommandContext
-			GDynamicRHI->RHICopySharedMips(DestTexture, SrcTexture);
+			GDynamicRHI->RHICopySharedMips(ExecutingCmdList, DestTexture, SrcTexture);
 			DestTexture->Release();
 			SrcTexture->Release();
 		});
diff --git a/Engine/Source/Runtime/RHI/Public/RHIValidation.h b/Engine/Source/Runtime/RHI/Public/RHIValidation.h
--- a/Engine/Source/Runtime/RHI/Public/RHIValidation.h
+++ b/Engine/Source/Runtime/RHI/Public/RHIValidation.h
@@ -594,9 +594,9 @@ public:
 	* @param SrcTexture2D - source texture
 	*/
 	// FlushType: Flush RHI Thread
-	virtual void RHICopySharedMips(FRHITexture2D* DestTexture2D, FRHITexture2D* SrcTexture2D) override final
+	virtual void RHICopySharedMips(FRHICommandList& RHICmdList, FRHITexture2D* DestTexture2D, FRHITexture2D* SrcTexture2D) override final
 	{
-		RHI->RHICopySharedMips(DestTexture2D, SrcTexture2D);
+		RHI->RHICopySharedMips(RHICmdList, DestTexture2D, SrcTexture2D);
 	}
 
 	void RHITransferBufferUnderlyingResource(FRHIBuffer* DestBuffer, FRHIBuffer* SrcBuffer) override final
diff --git a/Engine/Source/Runtime/VulkanRHI/Private/VulkanTexture.cpp b/Engine/Source/Runtime/VulkanRHI/Private/VulkanTexture.cpp
--- a/Engine/Source/Runtime/VulkanRHI/Private/VulkanTexture.cpp
+++ b/Engine/Source/Runtime/VulkanRHI/Private/VulkanTexture.cpp
@@ -925,7 +925,7 @@ FTextureRHIRef FVulkanDynamicRHI::RHIAsyncCreateTexture2D(uint32 SizeX,uint32 Si
 	return FTextureRHIRef();
 }
 
-void FVulkanDynamicRHI::RHICopySharedMips(FRHITexture2D* DestTexture2D, FRHITexture2D* SrcTexture2D)
+void FVulkanDynamicRHI::RHICopySharedMips(FRHICommandList& RHICmdList, FRHITexture2D* DestTexture2D, FRHITexture2D* SrcTexture2D)
 {
 	VULKAN_SIGNAL_UNIMPLEMENTED();
 }
diff --git a/Engine/Source/Runtime/VulkanRHI/Public/VulkanDynamicRHI.h b/Engine/Source/Runtime/VulkanRHI/Public/VulkanDynamicRHI.h
--- a/Engine/Source/Runtime/VulkanRHI/Public/VulkanDynamicRHI.h
+++ b/Engine/Source/Runtime/VulkanRHI/Public/VulkanDynamicRHI.h
@@ -127,7 +127,7 @@ public:
 	virtual bool RHIGetTextureMemoryVisualizeData(FColor* TextureData, int32 SizeX, int32 SizeY, int32 Pitch, int32 PixelSize) final override;
 	virtual FTextureRHIRef RHICreateTexture(const FRHITextureCreateDesc& CreateDesc) final override;
 	virtual FTextureRHIRef RHIAsyncCreateTexture2D(uint32 SizeX, uint32 SizeY, uint8 Format, uint32 NumMips, ETextureCreateFlags Flags, ERHIAccess InResourceState, void** InitialMipData, uint32 NumInitialMips) final override;
-	virtual void RHICopySharedMips(FRHITexture2D* DestTexture2D, FRHITexture2D* SrcTexture2D) final override;
+	virtual void RHICopySharedMips(FRHICommandList& RHICmdList, FRHITexture2D* DestTexture2D, FRHITexture2D* SrcTexture2D) final override;
 	virtual FShaderResourceViewRHIRef RHICreateShaderResourceView(FRHITexture* Texture, const FRHITextureSRVCreateInfo& CreateInfo) final override;
 	virtual uint32 RHIComputeMemorySize(FRHITexture* TextureRHI) final override;
 	virtual FTexture2DRHIRef RHIAsyncReallocateTexture2D(FRHITexture2D* Texture2D, int32 NewMipCount, int32 NewSizeX, int32 NewSizeY, FThreadSafeCounter* RequestStatus) final override;
diff --git a/Engine/Source/Runtime/Windows/D3D11RHI/Private/D3D11RHIPrivate.h b/Engine/Source/Runtime/Windows/D3D11RHI/Private/D3D11RHIPrivate.h
--- a/Engine/Source/Runtime/Windows/D3D11RHI/Private/D3D11RHIPrivate.h
+++ b/Engine/Source/Runtime/Windows/D3D11RHI/Private/D3D11RHIPrivate.h
@@ -494,7 +494,7 @@ public:
 	virtual bool RHIGetTextureMemoryVisualizeData(FColor* TextureData,int32 SizeX,int32 SizeY,int32 Pitch,int32 PixelSize) final override;
 	virtual FTextureRHIRef RHICreateTexture(const FRHITextureCreateDesc& CreateDesc) final override;
 	virtual FTextureRHIRef RHIAsyncCreateTexture2D(uint32 SizeX, uint32 SizeY, uint8 Format, uint32 NumMips, ETextureCreateFlags Flags, ERHIAccess InResourceState, void** InitialMipData, uint32 NumInitialMips) final override;
-	virtual void RHICopySharedMips(FRHITexture2D* DestTexture2D, FRHITexture2D* SrcTexture2D) final override;
+	virtual void RHICopySharedMips(FRHICommandList& RHICmdList, FRHITexture2D* DestTexture2D, FRHITexture2D* SrcTexture2D) final override;
 	virtual FShaderResourceViewRHIRef RHICreateShaderResourceView(FRHITexture* TextureRHI, const FRHITextureSRVCreateInfo& CreateInfo) final override;
 	virtual void RHIGenerateMips(FRHITexture* Texture) final override;
 	virtual uint32 RHIComputeMemorySize(FRHITexture* TextureRHI) final override;
diff --git a/Engine/Source/Runtime/Windows/D3D11RHI/Private/D3D11Texture.cpp b/Engine/Source/Runtime/Windows/D3D11RHI/Private/D3D11Texture.cpp
--- a/Engine/Source/Runtime/Windows/D3D11RHI/Private/D3D11Texture.cpp
+++ b/Engine/Source/Runtime/Windows/D3D11RHI/Private/D3D11Texture.cpp
@@ -1020,7 +1020,7 @@ FTextureRHIRef FD3D11DynamicRHI::RHIAsyncCreateTexture2D(uint32 SizeX, uint32 Si
 	return Texture;
 }
 
-void FD3D11DynamicRHI::RHICopySharedMips(FRHITexture2D* DestTexture2DRHI, FRHITexture2D* SrcTexture2DRHI)
+void FD3D11DynamicRHI::RHICopySharedMips(FRHICommandList& RHICmdList, FRHITexture2D* DestTexture2DRHI, FRHITexture2D* SrcTexture2DRHI)
 {
 	FD3D11Texture* DestTexture2D = ResourceCast(DestTexture2DRHI);
 	FD3D11Texture* SrcTexture2D = ResourceCast(SrcTexture2DRHI);
