Disable r.PSOPrecaching by default in code as probably unintended change in CL: 23220618 - now it will run on all platforms where GRHISupportsPipelineFileCache which includes OpenGL & Vulkan where it has never been tested yet

#[fyi] Arciel.Rekman

#localization none
#tests Restore back to default cvar - ran replay but it's enabled via ini file for Windows only right now

diff --git a/Engine/Source/Runtime/RHI/Private/PipelineStateCache.cpp b/Engine/Source/Runtime/RHI/Private/PipelineStateCache.cpp
--- a/Engine/Source/Runtime/RHI/Private/PipelineStateCache.cpp
+++ b/Engine/Source/Runtime/RHI/Private/PipelineStateCache.cpp
@@ -101,7 +101,7 @@ static TAutoConsoleVariable<int32> CVarRTPSOCacheSize(
 );
 #endif // RHI_RAYTRACING
 
-int32 GPSOPrecaching = 1;
+int32 GPSOPrecaching = 0;
 static FAutoConsoleVariableRef CVarPSOPrecaching(
 	TEXT("r.PSOPrecaching"),
 	GPSOPrecaching,
