Added Render Resource Viewer under Tools menu in editor listing tracked RHI resources, with filtering ability by resource flags and search by keyword

diff --git a/Engine/Source/Editor/RenderResourceViewer/Private/RenderResourceViewerModule.cpp b/Engine/Source/Editor/RenderResourceViewer/Private/RenderResourceViewerModule.cpp
new file mode 100644
--- /dev/null
+++ b/Engine/Source/Editor/RenderResourceViewer/Private/RenderResourceViewerModule.cpp
@@ -0,0 +1,44 @@
+// Copyright Epic Games, Inc. All Rights Reserved.
+
+#include "RenderResourceViewerModule.h"
+#include "WorkspaceMenuStructure.h"
+#include "WorkspaceMenuStructureModule.h"
+#include "Framework/Docking/TabManager.h"
+#include "Framework/Application/SlateApplication.h"
+#include "Modules/ModuleManager.h"
+#include "Widgets/Docking/SDockTab.h"
+#include "SRenderResourceViewer.h"
+
+#define LOCTEXT_NAMESPACE "RenderResourceViewerModule"
+
+IMPLEMENT_MODULE(FRenderResourceViewerModule, RenderResourceViewer);
+
+FName RenderResourceViewerTabId("RenderResourceViewer");
+
+void FRenderResourceViewerModule::StartupModule()
+{
+	FGlobalTabmanager::Get()->RegisterNomadTabSpawner(RenderResourceViewerTabId, FOnSpawnTab::CreateRaw(this, &FRenderResourceViewerModule::CreateTab))
+		.SetDisplayName(NSLOCTEXT("RenderResourceViewer", "TabTitle", "Render Resource Viewer"))
+		.SetTooltipText(NSLOCTEXT("RenderResourceViewer", "TooltipText", "Open the render resource viewer."))
+		.SetGroup(WorkspaceMenu::GetMenuStructure().GetToolsCategory())
+		.SetIcon(FSlateIcon(FAppStyle::GetAppStyleSetName(), "ClassIcon.UserDefinedStruct"));
+}
+
+void FRenderResourceViewerModule::ShutdownModule()
+{
+	if (FSlateApplication::IsInitialized())
+	{
+		FGlobalTabmanager::Get()->UnregisterNomadTabSpawner(RenderResourceViewerTabId);
+	}
+}
+
+TSharedRef<SDockTab> FRenderResourceViewerModule::CreateTab(const FSpawnTabArgs& Args)
+{
+	const TSharedRef<SDockTab> DockTab = SNew(SDockTab).TabRole(ETabRole::NomadTab);
+	TSharedRef<SRenderResourceViewerWidget> WindowRef = SNew(SRenderResourceViewerWidget, DockTab, Args.GetOwnerWindow());
+	DockTab->SetContent(WindowRef);
+	AssignWindow(WindowRef);
+	return DockTab;
+}
+
+#undef LOCTEXT_NAMESPACE
diff --git a/Engine/Source/Editor/RenderResourceViewer/Private/SRenderResourceViewer.cpp b/Engine/Source/Editor/RenderResourceViewer/Private/SRenderResourceViewer.cpp
new file mode 100644
--- /dev/null
+++ b/Engine/Source/Editor/RenderResourceViewer/Private/SRenderResourceViewer.cpp
@@ -0,0 +1,325 @@
+// Copyright Epic Games, Inc. All Rights Reserved.
+
+#include "SRenderResourceViewer.h"
+#include "Widgets/Input/SSearchBox.h"
+#include "Widgets/Input/SCheckBox.h"
+#include "Widgets/Input/SButton.h"
+#include "Widgets/Docking/SDockTab.h"
+#include "RHI.h"
+
+#define LOCTEXT_NAMESPACE "SRenderResourceView"
+
+namespace RenderResourceViewerInternal
+{
+	const FName ColumnName("Name");
+	const FName ColumnType("Type");
+	const FName ColumnSize("Size");
+	const FName ColumnFlags("Flags");
+	
+	// Format the size to nearest size unit of Byte/KB/MB
+	FString GetFormatedSize(uint64 SizeInBytes)
+	{
+		const uint64 KB = 1024;
+		const uint64 MB = KB * KB;
+		if (SizeInBytes < KB)
+		{
+			return FString::FromInt(SizeInBytes) + TEXT(" B");
+		}
+		else if (SizeInBytes < MB)
+		{
+			uint32 SizeInKB = SizeInBytes / KB;
+			return FString::FromInt(SizeInKB) + TEXT(" KB");
+		}
+		else
+		{
+			uint32 SizeInMB = SizeInBytes / MB;
+			return FString::FromInt(SizeInMB) + TEXT(" MB");
+		}
+	}
+	
+	struct FResourceInfoRow : public SMultiColumnTableRow<TSharedPtr<FRHIResourceStats>>
+	{
+		SLATE_BEGIN_ARGS(FResourceInfoRow) {}
+		SLATE_END_ARGS()
+	
+		TSharedPtr<FRHIResourceStats> Info;
+	
+		void Construct(const FArguments& InArgs, const TSharedRef<STableViewBase>& InOwnerTable, TSharedPtr<FRHIResourceStats> InItem)
+		{
+			Info = InItem;
+			SMultiColumnTableRow<TSharedPtr<FRHIResourceStats>>::Construct(FSuperRowType::FArguments(), InOwnerTable);
+		}
+	
+		TSharedRef<SWidget> GenerateWidgetForColumn(const FName& Column)
+		{
+			if (Column == RenderResourceViewerInternal::ColumnName)
+			{
+				return SNew(STextBlock)
+					.Text(FText::FromString(Info->Name));
+			}
+			else if (Column == RenderResourceViewerInternal::ColumnType)
+			{
+				return SNew(STextBlock)
+					.Text(FText::FromString(Info->Type));
+			}
+			else if (Column == RenderResourceViewerInternal::ColumnSize)
+			{
+				return SNew(STextBlock)
+					.Text(FText::FromString(RenderResourceViewerInternal::GetFormatedSize(Info->SizeInBytes)));
+			}
+			else if (Column == RenderResourceViewerInternal::ColumnFlags)
+			{
+				return SNew(STextBlock)
+					.Text(FText::FromString(Info->Flags));
+			}
+			return SNullWidget::NullWidget;
+		}
+	};
+}
+
+#define RENDER_RESOURCE_VIEWER_ADD_CHECKBOX(bIsChecked, DisplayText, Tooltip, OnCheckStateChangedFunc) \
+	+ SHorizontalBox::Slot() \
+		.AutoWidth() \
+		.Padding(FMargin(0.0f, 2.0f, 2.0f, 0.0f)) \
+		[ \
+			SNew(STextBlock) \
+			.Text(DisplayText) \
+			.ToolTipText(Tooltip) \
+		] \
+	+ SHorizontalBox::Slot() \
+		.AutoWidth() \
+		.Padding(FMargin(0.0f, 0.0f, 10.0f, 0.0f)) \
+		[ \
+			SNew(SCheckBox) \
+			.IsChecked(bIsChecked) \
+			.OnCheckStateChanged(this, &OnCheckStateChangedFunc) \
+		]
+
+void SRenderResourceViewerWidget::Construct(const FArguments& InArgs, const TSharedRef<SDockTab>& ConstructUnderMajorTab, const TSharedPtr<SWindow>& ConstructUnderWindow)
+{
+	SortByColumn = RenderResourceViewerInternal::ColumnSize;
+	SortMode = EColumnSortMode::Descending;
+
+	ChildSlot
+	[
+		SNew(SVerticalBox)
+
+		+SVerticalBox::Slot()
+		.FillHeight(0.5f)
+		[
+			SNew(SBox)
+			[
+				SNew(SVerticalBox)
+
+				+ SVerticalBox::Slot()
+				.AutoHeight()
+				[
+					SNew(SHorizontalBox)
+
+					+ SHorizontalBox::Slot()
+					.AutoWidth()
+					.Padding(FMargin(0.0f, 2.0f, 10.0f, 0.0f))
+					[
+						SNew(STextBlock)
+						.Text(LOCTEXT("FilterText", "FilterByFlags:"))
+					]
+
+					// A row of checkboxes to filter resource list by flags
+					RENDER_RESOURCE_VIEWER_ADD_CHECKBOX(bShowTransient, LOCTEXT("TransientText", "Transient"), LOCTEXT("TransientTooltip", "Transient"), SRenderResourceViewerWidget::OnTransientCheckboxChanged)
+					RENDER_RESOURCE_VIEWER_ADD_CHECKBOX(bShowStreaming, LOCTEXT("StreamingText", "Streaming"), LOCTEXT("StreamingTooltip", "Streaming"), SRenderResourceViewerWidget::OnStreamingCheckboxChanged)
+					RENDER_RESOURCE_VIEWER_ADD_CHECKBOX(bShowRT, LOCTEXT("RTText", "RT"), LOCTEXT("RTTooltip", "Render Target"), SRenderResourceViewerWidget::OnRTCheckboxChanged)
+					RENDER_RESOURCE_VIEWER_ADD_CHECKBOX(bShowDS, LOCTEXT("DSText", "DS"), LOCTEXT("DSTooltip", "Depth Stencil"), SRenderResourceViewerWidget::OnDSCheckboxChanged)
+					RENDER_RESOURCE_VIEWER_ADD_CHECKBOX(bShowUAV, LOCTEXT("UAVText", "UAV"), LOCTEXT("UAVTooltip", "Unordered Access View"), SRenderResourceViewerWidget::OnUAVCheckboxChanged)
+					RENDER_RESOURCE_VIEWER_ADD_CHECKBOX(bShowRTAS, LOCTEXT("RTASText", "RTAS"), LOCTEXT("RTASTooltip", "Ray Tracing Acceleration Structure"), SRenderResourceViewerWidget::OnRTASCheckboxChanged)
+					RENDER_RESOURCE_VIEWER_ADD_CHECKBOX(bShowRTAS, LOCTEXT("NoneText", "None"), LOCTEXT("NoneTooltip", "Resource with no flags set"), SRenderResourceViewerWidget::OnNoneCheckboxChanged)
+
+					// Refresh button to update the resource list
+					+ SHorizontalBox::Slot()
+					.AutoWidth()
+					.Padding(FMargin(0.0f, 0.0f, 10.0f, 0.0f))
+					[
+						SNew(SButton)
+						.Text(LOCTEXT("RefreshText", "Refresh"))
+						.ToolTipText(LOCTEXT("RefreshTooltip", "Refresh the resource list"))
+						.OnClicked(this, &SRenderResourceViewerWidget::OnRefreshButtonClicked)
+					]
+				]
+
+				// Text filter box to search by resource name
+				+ SVerticalBox::Slot()
+				.AutoHeight()
+				[
+					SAssignNew(FilterTextBox, SSearchBox)
+					.HintText(LOCTEXT("FilterTextBoxHint", "Search resources by name."))
+					.ToolTipText(LOCTEXT("FilterTextBoxToolTip", "Type here to filter the list of render resources by name."))
+					.OnTextChanged(this, &SRenderResourceViewerWidget::FilterTextBox_OnTextChanged)
+				]
+
+				// Display total resource count and size in text
+				+ SVerticalBox::Slot()
+				.AutoHeight()
+				[
+					SNew(SHorizontalBox)
+
+					+SHorizontalBox::Slot()
+					.AutoWidth()
+					.Padding(FMargin(0.0f, 0.0f, 10.0f, 0.0f))
+					[
+						SNew(STextBlock)
+						.Text(LOCTEXT("TotalResourceCountText", "Total Resource Count:"))
+					]
+
+					+ SHorizontalBox::Slot()
+					.AutoWidth()
+					.Padding(FMargin(0.0f, 0.0f, 10.0f, 0.0f))
+					[
+						SNew(STextBlock)
+						.Text(this, &SRenderResourceViewerWidget::GetResourceCountText)
+					]
+
+					+ SHorizontalBox::Slot()
+					.Padding(FMargin(0.0f, 0.0f, 10.0f, 0.0f))
+					.AutoWidth()
+					[
+						SNew(STextBlock)
+						.Text(LOCTEXT("TotalResourceSizeText", "Total Resource Size:"))
+					]
+
+					+ SHorizontalBox::Slot()
+					.AutoWidth()
+					[
+						SNew(STextBlock)
+						.Text(this, &SRenderResourceViewerWidget::GetResourceSizeText)
+					]
+				]
+
+				// Display resource list
+				+ SVerticalBox::Slot()
+				.FillHeight(1.f)
+				[
+					SAssignNew(ResourceListView, SListView<TSharedPtr<FRHIResourceStats>>)
+					.ScrollbarVisibility(EVisibility::Visible)
+					.ItemHeight(24.0f)
+					.ListItemsSource(&ResourceInfos)
+					.SelectionMode(ESelectionMode::SingleToggle)
+					.OnGenerateRow(this, &SRenderResourceViewerWidget::HandleResourceGenerateRow)
+					.HeaderRow
+					(
+						SNew(SHeaderRow)
+
+						+ SHeaderRow::Column(RenderResourceViewerInternal::ColumnName)
+						.FillWidth(0.6f)
+						.DefaultLabel(LOCTEXT("NameColumn", "Resource Name"))
+						.SortMode(this, &SRenderResourceViewerWidget::GetColumnSortMode, RenderResourceViewerInternal::ColumnName)
+						.OnSort(this, &SRenderResourceViewerWidget::OnColumnSortModeChanged)
+
+						+ SHeaderRow::Column(RenderResourceViewerInternal::ColumnType)
+						.FillWidth(0.1f)
+						.DefaultLabel(LOCTEXT("TypeColumn", "Type"))
+						.SortMode(this, &SRenderResourceViewerWidget::GetColumnSortMode, RenderResourceViewerInternal::ColumnType)
+						.OnSort(this, &SRenderResourceViewerWidget::OnColumnSortModeChanged)
+
+						+ SHeaderRow::Column(RenderResourceViewerInternal::ColumnSize)
+						.FillWidth(0.1f)
+						.DefaultLabel(LOCTEXT("AllocationSizeColumn", "Size"))
+						.SortMode(this, &SRenderResourceViewerWidget::GetColumnSortMode, RenderResourceViewerInternal::ColumnSize)
+						.OnSort(this, &SRenderResourceViewerWidget::OnColumnSortModeChanged)
+
+						+ SHeaderRow::Column(RenderResourceViewerInternal::ColumnFlags)
+						.FillWidth(0.2f)
+						.DefaultLabel(LOCTEXT("FlagsColumn", "Flags"))
+						.SortMode(this, &SRenderResourceViewerWidget::GetColumnSortMode, RenderResourceViewerInternal::ColumnFlags)
+						.OnSort(this, &SRenderResourceViewerWidget::OnColumnSortModeChanged)
+					)
+				]
+			]
+		]
+	];
+
+	RefreshNodes();
+}
+
+void SRenderResourceViewerWidget::RefreshNodes()
+{
+	const int32 NumberOfResourcesToShow = -1;
+
+	// Retrieve the whole list of RHI tracked resources
+	TArray<TSharedPtr<FRHIResourceStats>> RHIResources;
+	RHIGetTrackedResourceStats(RHIResources);
+
+	// Filter the resource list by search keyword and flags
+	ResourceInfos.Reset();
+	TotalResourceSize = 0;
+	for (const TSharedPtr<FRHIResourceStats>& Info : RHIResources)
+	{
+		bool bContainsFilterText = FilterText.IsEmpty() || Info->Name.Contains(FilterText.ToString());
+		// Note bMarkedForDelete resources are excluded from display
+		bool bContainsFilterFlags = (bShowTransient && Info->bTransient) || (bShowStreaming && Info->bStreaming)
+									|| (bShowRT && Info->bRenderTarget) || (bShowDS && Info->bDepthStencil) || (bShowUAV && Info->bUnorderedAccessView) || (bShowRTAS && Info->bRayTracingAccelerationStructure)
+									|| (bShowNone && !Info->bHasFlags);
+		if (bContainsFilterText && bContainsFilterFlags)
+		{
+			ResourceInfos.Add(Info);
+			TotalResourceSize += Info->SizeInBytes;
+		}
+	}
+
+	TotalResourceCount = ResourceInfos.Num();
+
+	// Sort the resource list based on columns.
+	{
+		if (SortByColumn == RenderResourceViewerInternal::ColumnName)
+		{
+			ResourceInfos.Sort([this](const TSharedPtr<FRHIResourceStats>& A, const TSharedPtr<FRHIResourceStats>& B) {
+				if (SortMode == EColumnSortMode::Ascending) return (A->Name < B->Name);
+				else if (SortMode == EColumnSortMode::Descending) return (A->Name >= B->Name);
+				return true;
+			});
+		}
+		else if (SortByColumn == RenderResourceViewerInternal::ColumnType)
+		{
+			ResourceInfos.Sort([this](const TSharedPtr<FRHIResourceStats>& A, const TSharedPtr<FRHIResourceStats>& B) {
+				if (SortMode == EColumnSortMode::Ascending) return (A->Type < B->Type);
+				else if (SortMode == EColumnSortMode::Descending) return (A->Type >= B->Type);
+				return true;
+			});
+		}
+		else if (SortByColumn == RenderResourceViewerInternal::ColumnSize)
+		{
+			ResourceInfos.Sort([this](const TSharedPtr<FRHIResourceStats>& A, const TSharedPtr<FRHIResourceStats>& B) {
+				if (SortMode == EColumnSortMode::Ascending) return (A->SizeInBytes < B->SizeInBytes);
+				else if (SortMode == EColumnSortMode::Descending) return (A->SizeInBytes >= B->SizeInBytes);
+				return true;
+			});
+		}
+		else if (SortByColumn == RenderResourceViewerInternal::ColumnFlags)
+		{
+			ResourceInfos.Sort([this](const TSharedPtr<FRHIResourceStats>& A, const TSharedPtr<FRHIResourceStats>& B) {
+				if (SortMode == EColumnSortMode::Ascending) return (A->Flags < B->Flags);
+				else if (SortMode == EColumnSortMode::Descending) return (A->Flags >= B->Flags);
+				return true;
+			});
+		}
+	}
+
+	ResourceListView->RebuildList();
+}
+
+TSharedRef<ITableRow> SRenderResourceViewerWidget::HandleResourceGenerateRow(TSharedPtr<FRHIResourceStats> Item, const TSharedRef<STableViewBase>& OwnerTable)
+{
+	return SNew(RenderResourceViewerInternal::FResourceInfoRow, OwnerTable, Item);
+}
+
+void SRenderResourceViewerWidget::OnColumnSortModeChanged(const EColumnSortPriority::Type SortPriority, const FName& ColumnId, const EColumnSortMode::Type InSortMode)
+{
+	SortByColumn = ColumnId;
+	SortMode = InSortMode;
+	RefreshNodes();
+}
+
+FText SRenderResourceViewerWidget::GetResourceSizeText() const
+{
+	return FText::FromString(RenderResourceViewerInternal::GetFormatedSize(TotalResourceSize));
+}
+
+#undef LOCTEXT_NAMESPACE
diff --git a/Engine/Source/Editor/RenderResourceViewer/Private/SRenderResourceViewer.h b/Engine/Source/Editor/RenderResourceViewer/Private/SRenderResourceViewer.h
new file mode 100644
--- /dev/null
+++ b/Engine/Source/Editor/RenderResourceViewer/Private/SRenderResourceViewer.h
@@ -0,0 +1,59 @@
+// Copyright Epic Games, Inc. All Rights Reserved.
+
+#pragma once
+
+#include "CoreMinimal.h"
+#include "Widgets/SCompoundWidget.h"
+#include "Styling/SlateTypes.h"
+#include "Widgets/Views/SHeaderRow.h"
+#include "Widgets/Views/SListView.h"
+
+struct FRHIResourceStats;
+class ITableRow;
+class STableViewBase;
+class SDockTab;
+
+class SRenderResourceViewerWidget : public SCompoundWidget
+{
+public:
+	SLATE_BEGIN_ARGS(SRenderResourceViewerWidget) {}
+	SLATE_END_ARGS()
+
+	void Construct(const FArguments& InArgs, const TSharedRef<SDockTab>& ConstructUnderMajorTab, const TSharedPtr<SWindow>& ConstructUnderWindow);
+	~SRenderResourceViewerWidget() {}
+
+private:
+	void RefreshNodes();
+	TSharedRef<ITableRow> HandleResourceGenerateRow(TSharedPtr<FRHIResourceStats> Item, const TSharedRef<STableViewBase>& OwnerTable);
+	void OnColumnSortModeChanged(const EColumnSortPriority::Type SortPriority, const FName& ColumnId, const EColumnSortMode::Type InSortMode);
+
+	EColumnSortMode::Type GetColumnSortMode(const FName ColumnId) const { return SortByColumn == ColumnId ? SortMode : EColumnSortMode::None; }
+	FText GetResourceCountText() const									{ return FText::AsNumber(TotalResourceCount); }
+	FText GetResourceSizeText() const;
+	void FilterTextBox_OnTextChanged(const FText& InFilterText)			{ FilterText = InFilterText; RefreshNodes(); }
+	void OnTransientCheckboxChanged(ECheckBoxState NewState)			{ bShowTransient = (NewState == ECheckBoxState::Checked); RefreshNodes(); }
+	void OnStreamingCheckboxChanged(ECheckBoxState NewState)			{ bShowStreaming = (NewState == ECheckBoxState::Checked); RefreshNodes(); }
+	void OnRTCheckboxChanged(ECheckBoxState NewState)					{ bShowRT = (NewState == ECheckBoxState::Checked); RefreshNodes(); }
+	void OnDSCheckboxChanged(ECheckBoxState NewState)					{ bShowDS = (NewState == ECheckBoxState::Checked); RefreshNodes(); }
+	void OnUAVCheckboxChanged(ECheckBoxState NewState)					{ bShowUAV = (NewState == ECheckBoxState::Checked); RefreshNodes(); }
+	void OnRTASCheckboxChanged(ECheckBoxState NewState)					{ bShowRTAS = (NewState == ECheckBoxState::Checked); RefreshNodes(); }
+	void OnNoneCheckboxChanged(ECheckBoxState NewState)					{ bShowNone = (NewState == ECheckBoxState::Checked); RefreshNodes(); }
+	FReply OnRefreshButtonClicked()										{ RefreshNodes(); return FReply::Handled(); }
+
+	TSharedPtr<SListView<TSharedPtr<FRHIResourceStats>>> ResourceListView;
+	TArray<TSharedPtr<FRHIResourceStats>> ResourceInfos;
+	FName SortByColumn;
+	EColumnSortMode::Type SortMode;
+	FText FilterText;
+	TSharedPtr< class SEditableTextBox > FilterTextBox;
+	uint64 TotalResourceCount = 0;
+	uint64 TotalResourceSize = 0;
+
+	bool bShowTransient = true;			// Show resource with Transient flag set
+	bool bShowStreaming = true;			// Show resource with Streaming flag set
+	bool bShowRT = true;				// Show resource with RenderTarget flag set
+	bool bShowDS = true;				// Show resource with DepthStencil flag set
+	bool bShowUAV = true;				// Show resource with UAV flag set
+	bool bShowRTAS = true;				// Show resource with RayTracingAccelationStructure flag set
+	bool bShowNone = true;				// Show resource with no flag set
+};
diff --git a/Engine/Source/Editor/RenderResourceViewer/Public/RenderResourceViewerModule.h b/Engine/Source/Editor/RenderResourceViewer/Public/RenderResourceViewerModule.h
new file mode 100644
--- /dev/null
+++ b/Engine/Source/Editor/RenderResourceViewer/Public/RenderResourceViewerModule.h
@@ -0,0 +1,38 @@
+// Copyright Epic Games, Inc. All Rights Reserved.
+
+#pragma once
+
+#include "CoreMinimal.h"
+#include "Modules/ModuleInterface.h"
+#include "Templates/SharedPointer.h"
+
+class SRenderResourceViewerWidget;
+class SDockTab;
+class FSpawnTabArgs;
+
+/**
+ * Struct Viewer module
+ */
+class FRenderResourceViewerModule : public IModuleInterface
+{
+public:
+	/**
+	 * Called right after the module DLL has been loaded and the module object has been created
+	 */
+	virtual void StartupModule() override;
+
+	/**
+	 * Called before the module is unloaded, right before the module object is destroyed.
+	 */
+	virtual void ShutdownModule() override;
+
+private:
+	TSharedRef<SDockTab> CreateTab(const FSpawnTabArgs& Args);
+
+	void AssignWindow(const TSharedRef<SRenderResourceViewerWidget>& InWindow)
+	{
+		RenderResourceViewerWindow = InWindow;
+	}
+
+	TWeakPtr<class SRenderResourceViewerWidget> RenderResourceViewerWindow;
+};
diff --git a/Engine/Source/Editor/RenderResourceViewer/RenderResourceViewer.Build.cs b/Engine/Source/Editor/RenderResourceViewer/RenderResourceViewer.Build.cs
new file mode 100644
--- /dev/null
+++ b/Engine/Source/Editor/RenderResourceViewer/RenderResourceViewer.Build.cs
@@ -0,0 +1,27 @@
+// Copyright Epic Games, Inc. All Rights Reserved.
+
+using UnrealBuildTool;
+
+public class RenderResourceViewer : ModuleRules
+{
+	public RenderResourceViewer(ReadOnlyTargetRules Target) : base(Target)
+	{
+		PrivateIncludePathModuleNames.AddRange(
+			new string[]
+			{
+				"WorkspaceMenuStructure",
+			}
+		);
+
+		PrivateDependencyModuleNames.AddRange(
+			new string[]
+			{
+				"Core",
+				"InputCore",
+				"SlateCore",
+				"Slate",
+				"RHI",
+			}
+		);
+	}
+}
diff --git a/Engine/Source/Editor/UnrealEd/Private/EditorEngine.cpp b/Engine/Source/Editor/UnrealEd/Private/EditorEngine.cpp
--- a/Engine/Source/Editor/UnrealEd/Private/EditorEngine.cpp
+++ b/Engine/Source/Editor/UnrealEd/Private/EditorEngine.cpp
@@ -1164,7 +1164,8 @@ void UEditorEngine::Init(IEngineLoop* InEngineLoop)
 			TEXT("CSVtoSVG"),
 			TEXT("GeometryFramework"),
 			TEXT("VirtualizationEditor"),
-			TEXT("AnimationSettings")
+			TEXT("AnimationSettings"),
+			TEXT("RenderResourceViewer"),
 		};
 
 		FScopedSlowTask ModuleSlowTask((float)UE_ARRAY_COUNT(ModuleNames));
diff --git a/Engine/Source/Editor/UnrealEd/UnrealEd.Build.cs b/Engine/Source/Editor/UnrealEd/UnrealEd.Build.cs
--- a/Engine/Source/Editor/UnrealEd/UnrealEd.Build.cs
+++ b/Engine/Source/Editor/UnrealEd/UnrealEd.Build.cs
@@ -43,6 +43,7 @@ public class UnrealEd : ModuleRules
 				"MessagingCommon",
 				"MovieSceneCapture",
 				"NaniteTools",
+				"RenderResourceViewer",
 				"PlacementMode",
 				"Settings",
 				"SettingsEditor",
@@ -255,6 +256,7 @@ public class UnrealEd : ModuleRules
 				"MeshUtilities",
 				"MergeActors",
 				"NaniteTools",
+				"RenderResourceViewer",
 				"ProjectSettingsViewer",
 				"BehaviorTreeEditor",
 				"ViewportSnapping",
diff --git a/Engine/Source/Runtime/RHI/Private/RHI.cpp b/Engine/Source/Runtime/RHI/Private/RHI.cpp
--- a/Engine/Source/Runtime/RHI/Private/RHI.cpp
+++ b/Engine/Source/Runtime/RHI/Private/RHI.cpp
@@ -16,6 +16,7 @@
 #include "String/ParseTokens.h"
 #include "Misc/BufferedOutputDevice.h"
 #include "Misc/OutputDeviceFile.h"
+#include "Async/ParallelFor.h"
 
 #if RHI_ENABLE_RESOURCE_INFO
 #include "HAL/FileManager.h"
@@ -867,106 +868,231 @@ static FAutoConsoleCommandWithWorldArgsAndOutputDevice GDumpRHIResourceCountsCmd
 	BufferedOutput.RedirectTo(OutputDevice);
 }));
 
-void RHIDumpResourceMemory(const FString& NameFilter, ERHIResourceType TypeFilter, EBooleanFilter TransientFilter, int32 NumberOfResourcesToShow, bool bUseCSVOutput, bool bSummaryOutput, bool bOutputToCSVFile, FBufferedOutputDevice& BufferedOutput)
-{	
-	FArchive* CSVFile{ nullptr };
-	if (bOutputToCSVFile)
-	{
-		const FString Filename = FString::Printf(TEXT("%srhiDumpResourceMemory-%s.csv"), *FPaths::ProfilingDir(), *FDateTime::Now().ToString());
-		CSVFile = IFileManager::Get().CreateFileWriter(*Filename, FILEWRITE_AllowRead);
-	}
 
-	TCHAR ResourceNameBuffer[FName::StringBufferSize];
+namespace RHIInternal
+{
+	struct FResourceEntry
+	{
+		const FRHIResource* Resource;
+		FRHIResourceInfo ResourceInfo;
+	};
 
-	auto ShouldIncludeResource = [&](const FRHIResource* Resource, const FRHIResourceInfo& ResourceInfo) -> bool
+	struct FResourceFlags
 	{
-		if (!NameFilter.IsEmpty())
+		bool bMarkedForDelete = false;
+		bool bTransient = false;
+		bool bStreaming = false;
+		bool bRT = false;
+		bool bDS = false;
+		bool bUAV = false;
+		bool bRTAS = false;
+		bool bHasFlags = false;
+
+		FString GetString()
 		{
-			if (ResourceInfo.Name.ToString(ResourceNameBuffer) == 0 || UE::String::FindFirst(ResourceNameBuffer, *NameFilter, ESearchCase::IgnoreCase) == INDEX_NONE)
+			FString FlagsString;
+			bool bHasFlag = false;
+			if (bMarkedForDelete)
 			{
-				return false;
+				FlagsString += "MarkedForDelete";
+				bHasFlag = true;
 			}
-		}
-		if (TypeFilter != RRT_None)
-		{
-			if (TypeFilter == RRT_Texture)
+			if (bTransient)
 			{
-				if (ResourceInfo.Type != RRT_Texture2D &&
-					ResourceInfo.Type != RRT_Texture2DArray &&
-					ResourceInfo.Type != RRT_Texture3D &&
-					ResourceInfo.Type != RRT_TextureCube)
-				{
-					return false;
-				}
+				FlagsString += bHasFlag ? " | Transient" : "Transient";
+				bHasFlag = true;
 			}
-			else if (TypeFilter != ResourceInfo.Type)
+			if (bStreaming)
 			{
-				return false;
+				FlagsString += bHasFlag ? " | Streaming" : "Streaming";
+				bHasFlag = true;
 			}
-		}
-		if (TransientFilter != EBooleanFilter::All)
-		{
-			const bool bAllowedFlag = TransientFilter == EBooleanFilter::Yes ? true : false;
-			if (ResourceInfo.IsTransient != bAllowedFlag)
+			if (bRT)
+			{
+				FlagsString += bHasFlag ? " | RT" : "RT";
+				bHasFlag = true;
+			}
+			else if (bDS)
+			{
+				FlagsString += bHasFlag ? " | DS" : "DS";
+				bHasFlag = true;
+			}
+			if (bUAV)
 			{
-				return false;
+				FlagsString += bHasFlag ? " | UAV" : "UAV";
+				bHasFlag = true;
 			}
+			if (bRTAS)
+			{
+				FlagsString += bHasFlag ? " | RTAS" : "RTAS";
+				bHasFlag = true;
+			}
+			return FlagsString;
 		}
-		return true;
 	};
 
-	struct FLocalResourceEntry
+	void GetTrackedResourcesInternal(const FString& NameFilter, ERHIResourceType TypeFilter, EBooleanFilter TransientFilter, TArray<FResourceEntry>& OutResources, int32& OutNumberOfResourcesToShow,
+		int32& OutTotalResourcesWithInfo, int32& OutTotalTrackedResources, int64& OutTotalTrackedResourceSize, int64& OutTotalTrackedTransientResourceSize)
 	{
-		const FRHIResource* Resource;
-		FRHIResourceInfo ResourceInfo;
-	};
-	TArray<FLocalResourceEntry> Resources;
-
-	FScopeLock Lock(&GRHIResourceTrackingCriticalSection);
-
-	int32 TotalResourcesWithInfo = 0;
-	int32 TotalTrackedResources = 0;
-	int64 TotalTrackedResourceSize = 0;
-	int64 TotalTrackedTransientResourceSize = 0;
+		TRACE_CPUPROFILER_EVENT_SCOPE(RHIInternal_GetTrackedResourcesInternal);
 
-	{
-		FRHIResourceInfo ResourceInfo;
-		TotalTrackedResources = GRHITrackedResources.Num();
+		TCHAR ResourceNameBuffer[FName::StringBufferSize];
 
-		for (const FRHIResource* Resource : GRHITrackedResources)
+		auto ShouldIncludeResource = [&](const FRHIResource* Resource, const FRHIResourceInfo& ResourceInfo) -> bool
 		{
-			if (Resource && Resource->GetResourceInfo(ResourceInfo))
+			if (!NameFilter.IsEmpty())
 			{
-				ResourceInfo.bValid = Resource->IsValid();
-
-				if (ShouldIncludeResource(Resource, ResourceInfo))
+				if (ResourceInfo.Name.ToString(ResourceNameBuffer) == 0 || UE::String::FindFirst(ResourceNameBuffer, *NameFilter, ESearchCase::IgnoreCase) == INDEX_NONE)
 				{
-					Resources.Emplace(FLocalResourceEntry{ Resource, ResourceInfo });
+					return false;
 				}
-
-				TotalResourcesWithInfo++;
-				if (ResourceInfo.IsTransient)
+			}
+			if (TypeFilter != RRT_None)
+			{
+				if (TypeFilter == RRT_Texture)
 				{
-					TotalTrackedTransientResourceSize += ResourceInfo.VRamAllocation.AllocationSize;
+					if (ResourceInfo.Type != RRT_Texture2D &&
+						ResourceInfo.Type != RRT_Texture2DArray &&
+						ResourceInfo.Type != RRT_Texture3D &&
+						ResourceInfo.Type != RRT_TextureCube)
+					{
+						return false;
+					}
 				}
-				else
+				else if (TypeFilter != ResourceInfo.Type)
+				{
+					return false;
+				}
+			}
+			if (TransientFilter != EBooleanFilter::All)
+			{
+				const bool bAllowedFlag = TransientFilter == EBooleanFilter::Yes ? true : false;
+				if (ResourceInfo.IsTransient != bAllowedFlag)
 				{
-					TotalTrackedResourceSize += ResourceInfo.VRamAllocation.AllocationSize;
+					return false;
+				}
+			}
+			return true;
+		};
+
+		OutResources.Reset();
+
+		{
+			FRHIResourceInfo ResourceInfo;
+			OutTotalTrackedResources = GRHITrackedResources.Num();
+
+			for (const FRHIResource* Resource : GRHITrackedResources)
+			{
+				if (Resource && Resource->GetResourceInfo(ResourceInfo))
+				{
+					ResourceInfo.bValid = Resource->IsValid();
+
+					if (ShouldIncludeResource(Resource, ResourceInfo))
+					{
+						OutResources.Emplace(FResourceEntry{ Resource, ResourceInfo });
+					}
+
+					OutTotalResourcesWithInfo++;
+					if (ResourceInfo.IsTransient)
+					{
+						OutTotalTrackedTransientResourceSize += ResourceInfo.VRamAllocation.AllocationSize;
+					}
+					else
+					{
+						OutTotalTrackedResourceSize += ResourceInfo.VRamAllocation.AllocationSize;
+					}
 				}
 			}
 		}
+
+		if (OutNumberOfResourcesToShow < 0 || OutNumberOfResourcesToShow > OutResources.Num())
+		{
+			OutNumberOfResourcesToShow = OutResources.Num();
+		}
+
+		OutResources.Sort([](const FResourceEntry& EntryA, const FResourceEntry& EntryB)
+		{
+			return EntryA.ResourceInfo.VRamAllocation.AllocationSize > EntryB.ResourceInfo.VRamAllocation.AllocationSize;
+		});
 	}
 
-	const int32 NumberOfResourcesBeforeNumberFilter = Resources.Num();
-	if (NumberOfResourcesToShow < 0 || NumberOfResourcesToShow > Resources.Num())
+	FResourceFlags GetResourceFlagsInternal(const FResourceEntry& Resource)
 	{
-		NumberOfResourcesToShow = Resources.Num();
+		FResourceFlags Flags;
+		Flags.bMarkedForDelete = !Resource.ResourceInfo.bValid;
+		Flags.bTransient = Resource.ResourceInfo.IsTransient;
+
+		bool bIsTexture = Resource.ResourceInfo.Type == RRT_Texture ||
+			Resource.ResourceInfo.Type == RRT_Texture2D ||
+			Resource.ResourceInfo.Type == RRT_Texture2DArray ||
+			Resource.ResourceInfo.Type == RRT_Texture3D ||
+			Resource.ResourceInfo.Type == RRT_TextureCube;
+		if (bIsTexture)
+		{
+			FRHITexture* Texture = (FRHITexture*)Resource.Resource;
+			Flags.bRT = EnumHasAnyFlags(Texture->GetFlags(), TexCreate_RenderTargetable);
+			Flags.bDS = EnumHasAnyFlags(Texture->GetFlags(), TexCreate_DepthStencilTargetable);
+			Flags.bUAV = EnumHasAnyFlags(Texture->GetFlags(), TexCreate_UAV);
+			Flags.bStreaming = EnumHasAnyFlags(Texture->GetFlags(), TexCreate_Streamable);
+		}
+		else if (Resource.ResourceInfo.Type == RRT_Buffer)
+		{
+			FRHIBuffer* Buffer = (FRHIBuffer*)Resource.Resource;
+			Flags.bUAV = EnumHasAnyFlags((EBufferUsageFlags)Buffer->GetUsage(), BUF_UnorderedAccess);
+			Flags.bRTAS = EnumHasAnyFlags((EBufferUsageFlags)Buffer->GetUsage(), BUF_AccelerationStructure);
+		}
+
+		Flags.bHasFlags = Flags.bMarkedForDelete || Flags.bTransient || Flags.bStreaming || Flags.bRT || Flags.bDS || Flags.bUAV || Flags.bRTAS;
+		return Flags;
 	}
+}
+
+void RHIGetTrackedResourceStats(TArray<TSharedPtr<FRHIResourceStats>>& OutResourceStats)
+{
+	TRACE_CPUPROFILER_EVENT_SCOPE(RHIGetTrackedResourceStats);
+	FScopeLock Lock(&GRHIResourceTrackingCriticalSection);
 
-	Resources.Sort([](const FLocalResourceEntry& EntryA, const FLocalResourceEntry& EntryB)
+	TArray<RHIInternal::FResourceEntry> Resources;
+	int32 TotalResourcesWithInfo = 0;
+	int32 TotalTrackedResources = 0;
+	int64 TotalTrackedResourceSize = 0;
+	int64 TotalTrackedTransientResourceSize = 0;
+	int32 NumberOfResourcesToShow = -1;
+
+	RHIInternal::GetTrackedResourcesInternal(TEXT(""), ERHIResourceType::RRT_None, EBooleanFilter::All, Resources, NumberOfResourcesToShow, TotalResourcesWithInfo, TotalTrackedResources, TotalTrackedResourceSize, TotalTrackedTransientResourceSize);
+
+	OutResourceStats.SetNum(Resources.Num());
+	ParallelFor(Resources.Num(), [&](int32 Index)
 	{
-		return EntryA.ResourceInfo.VRamAllocation.AllocationSize > EntryB.ResourceInfo.VRamAllocation.AllocationSize;
+		const FRHIResourceInfo& ResourceInfo = Resources[Index].ResourceInfo;
+		FString ResourceName = ResourceInfo.Name.ToString();
+		const TCHAR* ResourceType = StringFromRHIResourceType(ResourceInfo.Type);
+		const int64 SizeInBytes = ResourceInfo.VRamAllocation.AllocationSize;
+		RHIInternal::FResourceFlags Flags = GetResourceFlagsInternal(Resources[Index]);
+		OutResourceStats[Index] = MakeShared<FRHIResourceStats>(ResourceName, ResourceType, Flags.GetString(), SizeInBytes,
+									Flags.bMarkedForDelete, Flags.bTransient, Flags.bStreaming, Flags.bRT, Flags.bDS, Flags.bUAV, Flags.bRTAS, Flags.bHasFlags);
 	});
+}
+
+void RHIDumpResourceMemory(const FString& NameFilter, ERHIResourceType TypeFilter, EBooleanFilter TransientFilter, int32 NumberOfResourcesToShow, bool bUseCSVOutput, bool bSummaryOutput, bool bOutputToCSVFile, FBufferedOutputDevice& BufferedOutput)
+{	
+	FArchive* CSVFile{ nullptr };
+	if (bOutputToCSVFile)
+	{
+		const FString Filename = FString::Printf(TEXT("%srhiDumpResourceMemory-%s.csv"), *FPaths::ProfilingDir(), *FDateTime::Now().ToString());
+		CSVFile = IFileManager::Get().CreateFileWriter(*Filename, FILEWRITE_AllowRead);
+	}
+
+	FScopeLock Lock(&GRHIResourceTrackingCriticalSection);
+
+	TArray<RHIInternal::FResourceEntry> Resources;
+	int32 TotalResourcesWithInfo = 0;
+	int32 TotalTrackedResources = 0;
+	int64 TotalTrackedResourceSize = 0;
+	int64 TotalTrackedTransientResourceSize = 0;
+
+	RHIInternal::GetTrackedResourcesInternal(NameFilter, TypeFilter, TransientFilter, Resources, NumberOfResourcesToShow, TotalResourcesWithInfo, TotalTrackedResources, TotalTrackedResourceSize, TotalTrackedTransientResourceSize);
+	const int32 NumberOfResourcesBeforeNumberFilter = Resources.Num();
 
 	FName CategoryName(TEXT("RHIResources"));
 
@@ -992,6 +1118,7 @@ void RHIDumpResourceMemory(const FString& NameFilter, ERHIResourceType TypeFilte
 		}
 	}
 
+	TCHAR ResourceNameBuffer[FName::StringBufferSize];
 	int64 TotalShownResourceSize = 0;
 
 	for (int32 Index = 0; Index < Resources.Num(); Index++)
@@ -1004,33 +1131,7 @@ void RHIDumpResourceMemory(const FString& NameFilter, ERHIResourceType TypeFilte
 			const TCHAR* ResourceType = StringFromRHIResourceType(ResourceInfo.Type);
 			const int64 SizeInBytes = ResourceInfo.VRamAllocation.AllocationSize;
 
-			bool bMarkedForDelete = !ResourceInfo.bValid;
-			bool bTransient = ResourceInfo.IsTransient;
-			bool bStreaming = false;
-			bool bRT = false;
-			bool bDS = false;
-			bool bUAV = false;
-			bool bRTAS = false;
-
-			bool bIsTexture = ResourceInfo.Type == RRT_Texture ||
-				ResourceInfo.Type == RRT_Texture2D ||
-				ResourceInfo.Type == RRT_Texture2DArray ||
-				ResourceInfo.Type == RRT_Texture3D ||
-				ResourceInfo.Type == RRT_TextureCube;
-			if (bIsTexture)
-			{
-				FRHITexture* Texture = (FRHITexture*)Resources[Index].Resource;
-				bRT = EnumHasAnyFlags(Texture->GetFlags(), TexCreate_RenderTargetable);
-				bDS = EnumHasAnyFlags(Texture->GetFlags(), TexCreate_DepthStencilTargetable);
-				bUAV = EnumHasAnyFlags(Texture->GetFlags(), TexCreate_UAV);
-				bStreaming = EnumHasAnyFlags(Texture->GetFlags(), TexCreate_Streamable);
-			}
-			else if (ResourceInfo.Type == RRT_Buffer)
-			{
-				FRHIBuffer* Buffer = (FRHIBuffer*)Resources[Index].Resource;
-				bUAV = EnumHasAnyFlags((EBufferUsageFlags)Buffer->GetUsage(), BUF_UnorderedAccess);
-				bRTAS = EnumHasAnyFlags((EBufferUsageFlags)Buffer->GetUsage(), BUF_AccelerationStructure);
-			}
+			RHIInternal::FResourceFlags Flags = GetResourceFlagsInternal(Resources[Index]);
 
 			if (bSummaryOutput == false)
 			{
@@ -1040,12 +1141,12 @@ void RHIDumpResourceMemory(const FString& NameFilter, ERHIResourceType TypeFilte
 						ResourceNameBuffer,
 						ResourceType,
 						SizeInBytes / double(1 << 20),
-						bMarkedForDelete ? TEXT("Yes") : TEXT("No"),
-						bTransient ? TEXT("Yes") : TEXT("No"),
-						bStreaming ? TEXT("Yes") : TEXT("No"),
-						(bRT || bDS) ? TEXT("Yes") : TEXT("No"),
-						bUAV ? TEXT("Yes") : TEXT("No"),
-						bRTAS ? TEXT("Yes") : TEXT("No"));
+						Flags.bMarkedForDelete ? TEXT("Yes") : TEXT("No"),
+						Flags.bTransient ? TEXT("Yes") : TEXT("No"),
+						Flags.bStreaming ? TEXT("Yes") : TEXT("No"),
+						(Flags.bRT || Flags.bDS) ? TEXT("Yes") : TEXT("No"),
+						Flags.bUAV ? TEXT("Yes") : TEXT("No"),
+						Flags.bRTAS ? TEXT("Yes") : TEXT("No"));
 
 					CSVFile->Serialize(TCHAR_TO_ANSI(*Row), Row.Len());
 				}
@@ -1055,54 +1156,21 @@ void RHIDumpResourceMemory(const FString& NameFilter, ERHIResourceType TypeFilte
 						ResourceNameBuffer,
 						ResourceType,
 						SizeInBytes / double(1 << 20),
-						bMarkedForDelete ? TEXT("Yes") : TEXT("No"),
-						bTransient ? TEXT("Yes") : TEXT("No"),
-						bStreaming ? TEXT("Yes") : TEXT("No"),
-						(bRT || bDS) ? TEXT("Yes") : TEXT("No"),
-						bUAV ? TEXT("Yes") : TEXT("No"),
-						bRTAS ? TEXT("Yes") : TEXT("No"));
+						Flags.bMarkedForDelete ? TEXT("Yes") : TEXT("No"),
+						Flags.bTransient ? TEXT("Yes") : TEXT("No"),
+						Flags.bStreaming ? TEXT("Yes") : TEXT("No"),
+						(Flags.bRT || Flags.bDS) ? TEXT("Yes") : TEXT("No"),
+						Flags.bUAV ? TEXT("Yes") : TEXT("No"),
+						Flags.bRTAS ? TEXT("Yes") : TEXT("No"));
 				}
 				else
 				{
-					FString ResoureFlags;
-					bool bHasFlag = false;
-					if (bTransient)
-					{
-						ResoureFlags += "Transient";
-						bHasFlag = true;
-					}
-					if (bStreaming)
-					{
-						ResoureFlags += bHasFlag ? "|Streaming" : "Streaming";
-						bHasFlag = true;
-					}
-					if (bRT)
-					{
-						ResoureFlags += bHasFlag ? "|RT" : "RT";
-						bHasFlag = true;
-					}
-					else if (bDS)
-					{
-						ResoureFlags += bHasFlag ? "|DS" : "DS";
-						bHasFlag = true;
-					}
-					if (bUAV)
-					{
-						ResoureFlags += bHasFlag ? "|UAV" : "UAV";
-						bHasFlag = true;
-					}
-					if (bRTAS)
-					{
-						ResoureFlags += bHasFlag ? "|RTAS" : "RTAS";
-						bHasFlag = true;
-
-					}
-
+					FString ResoureFlags = Flags.GetString();
 					BufferedOutput.CategorizedLogf(CategoryName, ELogVerbosity::Log, TEXT("Name: %s - Type: %s - Size: %.9f MB - Flags: %s"),
 						ResourceNameBuffer,
 						ResourceType,
 						SizeInBytes / double(1 << 20),
-						bHasFlag ? *ResoureFlags : TEXT("None"));
+						ResoureFlags.IsEmpty() ? TEXT("None") : *ResoureFlags);
 				}
 			}
 
diff --git a/Engine/Source/Runtime/RHI/Public/RHI.h b/Engine/Source/Runtime/RHI/Public/RHI.h
--- a/Engine/Source/Runtime/RHI/Public/RHI.h
+++ b/Engine/Source/Runtime/RHI/Public/RHI.h
@@ -151,6 +151,40 @@ RHI_API int32 RHIGetPreferredClearUAVRectPSResourceType(const FStaticShaderPlatf
 // helper to force dump all RHI resource to CSV file
 RHI_API void RHIDumpResourceMemoryToCSV();
 
+struct FRHIResourceStats
+{
+	FString Name;
+	FString Type;
+	FString Flags;
+	uint64	SizeInBytes = 0;
+	bool	bMarkedForDelete = false;
+	bool	bTransient = false;
+	bool	bStreaming = false;
+	bool	bRenderTarget = false;
+	bool	bDepthStencil = false;
+	bool	bUnorderedAccessView = false;
+	bool	bRayTracingAccelerationStructure = false;
+	bool	bHasFlags = false;
+
+	FRHIResourceStats(const FString& InName, const FString& InType, const FString& InFlags, const uint64& InSizeInBytes, 
+						bool bInMarkedForDelete, bool bInTransient, bool bInStreaming, bool bInRT, bool bInDS, bool bInUAV, bool bInRTAS, bool bInHasFlags)
+		: Name(InName)
+		, Type(InType)
+		, Flags(InFlags)
+		, SizeInBytes(InSizeInBytes)
+		, bMarkedForDelete(bInMarkedForDelete)
+		, bTransient(bInTransient)
+		, bStreaming(bInStreaming)
+		, bRenderTarget(bInRT)
+		, bDepthStencil(bInDS)
+		, bUnorderedAccessView(bInUAV)
+		, bRayTracingAccelerationStructure(bInRTAS)
+		, bHasFlags(bInHasFlags)
+	{ }
+};
+
+RHI_API void RHIGetTrackedResourceStats(TArray<TSharedPtr<FRHIResourceStats>>& OutResourceStats);
+
 inline bool RHISupportsInstancedStereo(const FStaticShaderPlatform Platform)
 {
 	return FDataDrivenShaderPlatformInfo::GetSupportsInstancedStereo(Platform);
