During shutdown some static renderresources clean up late and in unpredicatable order they get stuck in the queue and never propperly cleanup the concurrent linear allocator.

diff --git a/Engine/Source/Runtime/RHI/Private/RHI.cpp b/Engine/Source/Runtime/RHI/Private/RHI.cpp
--- a/Engine/Source/Runtime/RHI/Private/RHI.cpp
+++ b/Engine/Source/Runtime/RHI/Private/RHI.cpp
@@ -295,7 +295,7 @@ const FClearValueBinding FClearValueBinding::Green(FLinearColor(0.0f, 1.0f, 0.0f
 // Note: this is used as the default normal for DBuffer decals.  It must decode to a value of 0 in DecodeDBufferData.
 const FClearValueBinding FClearValueBinding::DefaultNormal8Bit(FLinearColor(128.0f / 255.0f, 128.0f / 255.0f, 128.0f / 255.0f, 1.0f));
 
-UE::TConsumeAllMpmcQueue<FRHIResource*, FConcurrentLinearAllocator> FRHIResource::PendingDeletes;
+UE::TConsumeAllMpmcQueue<FRHIResource*> FRHIResource::PendingDeletes;
 FRHIResource* FRHIResource::CurrentlyDeleting = nullptr;
 
 #if HAS_GPU_STATS
diff --git a/Engine/Source/Runtime/RHI/Public/RHIResources.h b/Engine/Source/Runtime/RHI/Public/RHIResources.h
--- a/Engine/Source/Runtime/RHI/Public/RHIResources.h
+++ b/Engine/Source/Runtime/RHI/Public/RHIResources.h
@@ -223,7 +223,7 @@ private:
 	uint8 bBeingTracked : 1;
 #endif
 
-	static UE::TConsumeAllMpmcQueue<FRHIResource*, FConcurrentLinearAllocator> PendingDeletes;
+	static UE::TConsumeAllMpmcQueue<FRHIResource*> PendingDeletes;
 	static FRHIResource* CurrentlyDeleting;
 
 	// Some APIs don't do internal reference counting, so we have to wait an extra couple of frames before deleting resources
