Allow projects to specify a VRAM threshold for defaulting to DX12 SM6.
D3D12 adapter selection should now pick the adapter with the highest video memory.

diff --git a/Engine/Source/Runtime/D3D12RHI/Private/Windows/WindowsD3D12Device.cpp b/Engine/Source/Runtime/D3D12RHI/Private/Windows/WindowsD3D12Device.cpp
--- a/Engine/Source/Runtime/D3D12RHI/Private/Windows/WindowsD3D12Device.cpp
+++ b/Engine/Source/Runtime/D3D12RHI/Private/Windows/WindowsD3D12Device.cpp
@@ -854,6 +854,7 @@ void FD3D12DynamicRHIModule::FindAdapter()
 	const D3D_FEATURE_LEVEL MinRequiredFeatureLevel = GetRequiredD3DFeatureLevel();
 
 	FD3D12AdapterDesc FirstDiscreteAdapter;
+	FD3D12AdapterDesc BestMemoryAdapter;
 	FD3D12AdapterDesc FirstAdapter;
 
 	bool bRequestedWARP = D3D12RHI_ShouldCreateWithWarp();
@@ -947,13 +948,21 @@ void FD3D12DynamicRHIModule::FindAdapter()
 
 				if (!bSkipAdapter)
 				{
-					if (!bIsWARP && !CurrentAdapter.bIsIntegrated && !FirstDiscreteAdapter.IsValid())
+					if (PreferredVendor == AdapterDesc.VendorId && FirstDiscreteAdapter.IsValid())
 					{
 						FirstDiscreteAdapter = CurrentAdapter;
 					}
-					else if (PreferredVendor == AdapterDesc.VendorId && FirstDiscreteAdapter.IsValid())
+					else if (!bIsWARP && !CurrentAdapter.bIsIntegrated)
 					{
-						FirstDiscreteAdapter = CurrentAdapter;
+						if (!FirstDiscreteAdapter.IsValid())
+						{
+							FirstDiscreteAdapter = CurrentAdapter;
+						}
+
+						if (CurrentAdapter.Desc.DedicatedVideoMemory > BestMemoryAdapter.Desc.DedicatedVideoMemory)
+						{
+							BestMemoryAdapter = CurrentAdapter;
+						}
 					}
 
 					if (!FirstAdapter.IsValid())
@@ -973,7 +982,12 @@ void FD3D12DynamicRHIModule::FindAdapter()
 	if (bFavorDiscreteAdapter)
 	{
 		// We assume Intel is integrated graphics (slower than discrete) than NVIDIA or AMD cards and rather take a different one
-		if (FirstDiscreteAdapter.IsValid())
+		if (BestMemoryAdapter.IsValid())
+		{
+			NewAdapter = TSharedPtr<FD3D12Adapter>(new FWindowsD3D12Adapter(BestMemoryAdapter));
+			ChosenAdapters.Add(NewAdapter);
+		}
+		else if (FirstDiscreteAdapter.IsValid())
 		{
 			NewAdapter = TSharedPtr<FD3D12Adapter>(new FWindowsD3D12Adapter(FirstDiscreteAdapter));
 			ChosenAdapters.Add(NewAdapter);
diff --git a/Engine/Source/Runtime/RHI/Private/Windows/WindowsDynamicRHI.cpp b/Engine/Source/Runtime/RHI/Private/Windows/WindowsDynamicRHI.cpp
--- a/Engine/Source/Runtime/RHI/Private/Windows/WindowsDynamicRHI.cpp
+++ b/Engine/Source/Runtime/RHI/Private/Windows/WindowsDynamicRHI.cpp
@@ -204,8 +204,6 @@ TOptional<EWindowsRHI> ParseDefaultWindowsRHI()
 		const FString NAME_DX12(TEXT("DefaultGraphicsRHI_DX12"));
 		const FString NAME_VULKAN(TEXT("DefaultGraphicsRHI_Vulkan"));
 
-		DefaultRHI = EWindowsRHI::D3D11;
-
 		if (DefaultGraphicsRHI == NAME_DX11)
 		{
 			DefaultRHI = EWindowsRHI::D3D11;
@@ -452,6 +450,36 @@ static TOptional<ERHIFeatureLevel::Type> GetPreferredFeatureLevel(EWindowsRHI Ch
 	return PreferredFeatureLevel;
 }
 
+static bool IsRHIAllowedAsDefault(EWindowsRHI InRHI, ERHIFeatureLevel::Type InFeatureLevel)
+{
+	if (InRHI == EWindowsRHI::D3D12 && InFeatureLevel >= ERHIFeatureLevel::SM6)
+	{
+		int32 MinDedicatedMemoryMB = 0;
+		if (GConfig->GetInt(TEXT("D3D12_SM6"), TEXT("MinDedicatedMemory"), MinDedicatedMemoryMB, GEngineIni))
+		{
+			const int32 MinDedicatedMemory = MinDedicatedMemoryMB << 20;
+			static const FWindowsPlatformApplicationMisc::FGPUInfo BestGPUInfo = FWindowsPlatformApplicationMisc::GetBestGPUInfo();
+			return BestGPUInfo.DedicatedVideoMemory >= MinDedicatedMemory;
+		}
+	}
+	return true;
+}
+
+static TOptional<ERHIFeatureLevel::Type> ChooseDefaultFeatureLevel(EWindowsRHI InRHI, const FParsedWindowsDynamicRHIConfig& Config)
+{
+	TOptional<ERHIFeatureLevel::Type> HighestFL = Config.GetHighestSupportedFeatureLevel(InRHI);
+	while (HighestFL)
+	{
+		if (IsRHIAllowedAsDefault(InRHI, HighestFL.GetValue()))
+		{
+			return HighestFL;
+		}
+		HighestFL = Config.GetNextHighestTargetedFeatureLevel(InRHI, HighestFL.GetValue());
+	}
+
+	return TOptional<ERHIFeatureLevel::Type>();
+}
+
 // Choose the default from DefaultGraphicsRHI or TargetedRHIs. DefaultGraphicsRHI has precedence.
 static EWindowsRHI ChooseDefaultRHI(const FParsedWindowsDynamicRHIConfig& Config)
 {
@@ -471,7 +499,7 @@ static EWindowsRHI ChooseDefaultRHI(const FParsedWindowsDynamicRHIConfig& Config
 	// Find the first RHI with configured support based on the order above
 	for (EWindowsRHI DefaultRHI : DefaultRHIOrder)
 	{
-		if (TOptional<ERHIFeatureLevel::Type> HighestFL = Config.GetHighestSupportedFeatureLevel(DefaultRHI))
+		if (TOptional<ERHIFeatureLevel::Type> HighestFL = ChooseDefaultFeatureLevel(DefaultRHI, Config))
 		{
 			return DefaultRHI;
 		}
@@ -488,6 +516,7 @@ static TOptional<EWindowsRHI> ChoosePreferredRHI(EWindowsRHI InDefaultRHI)
 	if (!GIsEditor && (InDefaultRHI == EWindowsRHI::D3D11 || InDefaultRHI == EWindowsRHI::D3D12))
 	{
 		bool bUseD3D12InGame = false;
+		bool bPreferFeatureLevelES31 = false;
 		FString PreferredRHIName;
 		if (GConfig->GetString(TEXT("D3DRHIPreference"), TEXT("PreferredRHI"), PreferredRHIName, GGameUserSettingsIni))
 		{
@@ -514,14 +543,18 @@ static TOptional<EWindowsRHI> ChoosePreferredRHI(EWindowsRHI InDefaultRHI)
 				UE_LOG(LogRHI, Error, TEXT("unknown RHI name \"%s\" in game user settings, using default"), *PreferredRHIName);
 			}
 		}
-		else if (GConfig->GetBool(TEXT("D3DRHIPreference"), TEXT("bUseD3D12InGame"), bUseD3D12InGame, GGameUserSettingsIni))
+		else if (GConfig->GetBool(TEXT("D3DRHIPreference"), TEXT("bUseD3D12InGame"), bUseD3D12InGame, GGameUserSettingsIni) &&
+			GConfig->GetBool(TEXT("D3DRHIPreference"), TEXT("bPreferFeatureLevelES31"), bPreferFeatureLevelES31, GGameUserSettingsIni))
 		{
+			// Respect legacy GameUserSettings values if present
 			UE_LOG(LogRHI, Log, TEXT("Found D3DRHIPreference bUseD3D12InGame: %s"), bUseD3D12InGame ? TEXT("true") : TEXT("false"));
+			UE_LOG(LogRHI, Log, TEXT("Found D3DRHIPreference bPreferFeatureLevelES31: %s"), bPreferFeatureLevelES31 ? TEXT("true") : TEXT("false"));
 
-			bool bPreferFeatureLevelES31 = false;
-			GConfig->GetBool(TEXT("D3DRHIPreference"), TEXT("bPreferFeatureLevelES31"), bPreferFeatureLevelES31, GGameUserSettingsIni);
-
-			if (bPreferFeatureLevelES31)
+			if (bUseD3D12InGame)
+			{
+				RHIPreference = EWindowsRHI::D3D12;
+			}
+			else if (bPreferFeatureLevelES31)
 			{
 				RHIPreference = EWindowsRHI::D3D11;
 			}
