Fixing Default RHI detection when comparing memory units higher than 2^32.

diff --git a/Engine/Source/Runtime/RHI/Private/Windows/WindowsDynamicRHI.cpp b/Engine/Source/Runtime/RHI/Private/Windows/WindowsDynamicRHI.cpp
--- a/Engine/Source/Runtime/RHI/Private/Windows/WindowsDynamicRHI.cpp
+++ b/Engine/Source/Runtime/RHI/Private/Windows/WindowsDynamicRHI.cpp
@@ -457,7 +457,7 @@ static bool IsRHIAllowedAsDefault(EWindowsRHI InRHI, ERHIFeatureLevel::Type InFe
 		int32 MinDedicatedMemoryMB = 0;
 		if (GConfig->GetInt(TEXT("D3D12_SM6"), TEXT("MinDedicatedMemory"), MinDedicatedMemoryMB, GEngineIni))
 		{
-			const int32 MinDedicatedMemory = MinDedicatedMemoryMB << 20;
+			const uint64 MinDedicatedMemory = static_cast<uint64>(MinDedicatedMemoryMB) << 20;
 			static const FWindowsPlatformApplicationMisc::FGPUInfo BestGPUInfo = FWindowsPlatformApplicationMisc::GetBestGPUInfo();
 			return BestGPUInfo.DedicatedVideoMemory >= MinDedicatedMemory;
 		}
