Remove platform RHICopySharedMips function
 - Replaced with a common UE::RHICore::CopySharedMips function that relies on RHICopyTexture, which is already implemented for all platforms.
 - Global scope / command list versions of RHICopySharedMips are deprecated, and forward to the new implementation.

Added an RHI test for RHICopyTexture on a 3D texture.
 - Fixed a handful of platform RHI bugs where RHICopyTexture failed on 3D textures with mips.

diff --git a/Engine/Plugins/Tests/RHITests/Source/RHITests/Private/RHITests.spec.cpp b/Engine/Plugins/Tests/RHITests/Source/RHITests/Private/RHITests.spec.cpp
--- a/Engine/Plugins/Tests/RHITests/Source/RHITests/Private/RHITests.spec.cpp
+++ b/Engine/Plugins/Tests/RHITests/Source/RHITests/Private/RHITests.spec.cpp
@@ -83,4 +83,13 @@ void FAutomationRHITest::Define()
 			TestEqual("RHICopyTexture", bResult, 1);
 		});
 	});
+
+	Describe("Test RHI Create Buffer Parallel", [this]
+	{
+		It("RHICreateBuffer_Parallel", [this]()
+		{
+			bool bResult = RunOnRenderThreadSynchronous(FRHIBufferTests::Test_RHICreateBuffer_Parallel);
+			TestEqual("RHICreateBuffer_Parallel", bResult, 1);
+		});
+	});
 }
diff --git a/Engine/Plugins/Tests/RHITests/Source/RHITests/Private/RHITestsModule.cpp b/Engine/Plugins/Tests/RHITests/Source/RHITests/Private/RHITestsModule.cpp
--- a/Engine/Plugins/Tests/RHITests/Source/RHITests/Private/RHITestsModule.cpp
+++ b/Engine/Plugins/Tests/RHITests/Source/RHITests/Private/RHITestsModule.cpp
@@ -15,34 +15,42 @@ static bool RunTests_RenderThread(FRHICommandListImmediate& RHICmdList)
 	bool bResult = true;
 
 	// ------------------------------------------------
-	// RHIClearUAVUint / RHIClearUAVFloat tests
+	// RHI Formats
 	// ------------------------------------------------
+	{
+		RUN_TEST(FRHITextureTests::Test_RHIFormats(RHICmdList));
+	}
 
-	// Vertex/Structured Buffer
+	// ------------------------------------------------
+	// RHIClearUAVUint / RHIClearUAVFloat tests
+	// ------------------------------------------------
 	{
+		// Vertex/Structured Buffer
 		RUN_TEST(FRHIBufferTests::Test_RHIClearUAVUint_VertexBuffer(RHICmdList));
 		RUN_TEST(FRHIBufferTests::Test_RHIClearUAVFloat_VertexBuffer(RHICmdList));
 
 		RUN_TEST(FRHIBufferTests::Test_RHIClearUAVUint_StructuredBuffer(RHICmdList));
 		RUN_TEST(FRHIBufferTests::Test_RHIClearUAVFloat_StructuredBuffer(RHICmdList));
 
-		RUN_TEST(FRHIBufferTests::Test_RHICreateBuffer_Parallel(RHICmdList));
-	}
-
-	// Texture2D/3D
-	{
+		// Texture2D/3D
 		RUN_TEST(FRHITextureTests::Test_RHIClearUAV_Texture2D(RHICmdList));
 		RUN_TEST(FRHITextureTests::Test_RHIClearUAV_Texture3D(RHICmdList));
-		RUN_TEST(FRHITextureTests::Test_UpdateTexture(RHICmdList));
-		RUN_TEST(FRHITextureTests::Test_MultipleLockTexture2D(RHICmdList));
 	}
 
+	// ------------------------------------------------
+	// Texture Operations
+	// ------------------------------------------------
 	{
-		RUN_TEST(FRHITextureTests::Test_RHIFormats(RHICmdList));
+		RUN_TEST(FRHITextureTests::Test_RHICopyTexture(RHICmdList));
+		RUN_TEST(FRHITextureTests::Test_UpdateTexture(RHICmdList));
+		RUN_TEST(FRHITextureTests::Test_MultipleLockTexture2D(RHICmdList));
 	}
 
+	// ------------------------------------------------
+	// Buffer Operations
+	// ------------------------------------------------
 	{
-		RUN_TEST(FRHITextureTests::Test_RHICopyTexture(RHICmdList));
+		RUN_TEST(FRHIBufferTests::Test_RHICreateBuffer_Parallel(RHICmdList));
 	}
 
 	// @todo - add more tests
diff --git a/Engine/Plugins/Tests/RHITests/Source/RHITests/Public/RHITextureTests.h b/Engine/Plugins/Tests/RHITests/Source/RHITests/Public/RHITextureTests.h
--- a/Engine/Plugins/Tests/RHITests/Source/RHITests/Public/RHITextureTests.h
+++ b/Engine/Plugins/Tests/RHITests/Source/RHITests/Public/RHITextureTests.h
@@ -1200,6 +1200,64 @@ public:
 					return CheckTestData(Data, Width * sizeof(uint32), PatternSize, PatternOffset, ImageIndex);
 				}));
 			}
+
+			// Copy 3D texture
+			{
+				const TCHAR* TestName = TEXT("Copy 3D Texture");
+				
+				const FRHITextureCreateDesc SourceDesc =
+					FRHITextureCreateDesc::Create3D(TestName)
+					.SetExtent(64, 64)
+					.SetDepth(64)
+					.SetNumMips(6)
+					.SetFormat(PF_R32_UINT)
+					.SetInitialState(ERHIAccess::CopySrc);
+
+				const FRHITextureCreateDesc DestDesc = FRHITextureCreateDesc(SourceDesc)
+					.SetInitialState(ERHIAccess::CopyDest);
+
+				FTextureRHIRef SourceTexture = RHICreateTexture(SourceDesc);
+				FTextureRHIRef DestTexture = RHICreateTexture(DestDesc);
+
+				// Fill the source texture with data
+				for (int32 Mip = 0, ImageIndex = 1; Mip < SourceDesc.NumMips; ++Mip, ImageIndex += SourceDesc.Depth)
+				{
+					FUpdateTextureRegion3D Region(FIntVector::ZeroValue, FIntVector::ZeroValue, FIntVector(
+						FMath::Max(SourceDesc.Extent.X >> Mip, 1),
+						FMath::Max(SourceDesc.Extent.Y >> Mip, 1),
+						FMath::Max(SourceDesc.Depth    >> Mip, 1)
+					));
+
+					FUpdateTexture3DData Data = RHICmdList.BeginUpdateTexture3D(SourceTexture, Mip, Region);
+					
+					const FIntPoint PatternOffset(0, 0);
+					const FIntPoint PatternSize(Region.Width, Region.Height);
+
+					for (uint32 Z = 0; Z < Region.Depth; ++Z)
+					{
+						WriteTestData(Data.Data + Data.DepthPitch * Z, Data.RowPitch, PatternSize, PatternOffset, ImageIndex + Z);
+					}
+
+					RHICmdList.EndUpdateTexture3D(Data);
+				}
+
+				FRHICopyTextureInfo CopyInfo;
+				CopyInfo.NumMips = SourceDesc.NumMips;
+				CopyInfo.NumSlices = SourceDesc.ArraySize;
+
+				RHICmdList.CopyTexture(SourceTexture, DestTexture, CopyInfo);
+				RHICmdList.Transition(FRHITransitionInfo(DestTexture, ERHIAccess::CopyDest, ERHIAccess::CopySrc));
+
+				RUN_TEST(VerifyTextureContents(TestName, RHICmdList, DestTexture, [&](void* Data, uint32 MipWidth, uint32 MipHeight, uint32 Width, uint32 Height, uint32 MipIndex, uint32 Depth)
+				{
+					int32 ImageIndex = SourceDesc.Depth * MipIndex + Depth + 1;
+
+					const FIntPoint PatternOffset(0, 0);
+					const FIntPoint PatternSize(MipWidth, MipHeight);
+
+					return CheckTestData(Data, Width * sizeof(uint32), PatternSize, PatternOffset, ImageIndex);
+				}));
+			}
 		}
 
 		return bResult;
diff --git a/Engine/Source/Runtime/Apple/MetalRHI/Private/MetalDynamicRHI.h b/Engine/Source/Runtime/Apple/MetalRHI/Private/MetalDynamicRHI.h
--- a/Engine/Source/Runtime/Apple/MetalRHI/Private/MetalDynamicRHI.h
+++ b/Engine/Source/Runtime/Apple/MetalRHI/Private/MetalDynamicRHI.h
@@ -73,7 +73,6 @@ public:
 	virtual bool RHIGetTextureMemoryVisualizeData(FColor* TextureData,int32 SizeX,int32 SizeY,int32 Pitch,int32 PixelSize) final override;
 	virtual FTextureRHIRef RHICreateTexture(const FRHITextureCreateDesc& CreateDesc) final override;
 	virtual FTexture2DRHIRef RHIAsyncCreateTexture2D(uint32 SizeX, uint32 SizeY, uint8 Format, uint32 NumMips, ETextureCreateFlags Flags, ERHIAccess InResourceState, void** InitialMipData, uint32 NumInitialMips) final override;
-	virtual void RHICopySharedMips(FRHICommandList& RHICmdList, FRHITexture2D* DestTexture2D, FRHITexture2D* SrcTexture2D) final override;
 	virtual void RHIGenerateMips(FRHITexture* Texture) final override;
 	virtual uint32 RHIComputeMemorySize(FRHITexture* TextureRHI) final override;
 	virtual FTexture2DRHIRef RHIAsyncReallocateTexture2D(FRHITexture2D* Texture2D, int32 NewMipCount, int32 NewSizeX, int32 NewSizeY, FThreadSafeCounter* RequestStatus) final override;
diff --git a/Engine/Source/Runtime/Apple/MetalRHI/Private/MetalTexture.cpp b/Engine/Source/Runtime/Apple/MetalRHI/Private/MetalTexture.cpp
--- a/Engine/Source/Runtime/Apple/MetalRHI/Private/MetalTexture.cpp
+++ b/Engine/Source/Runtime/Apple/MetalRHI/Private/MetalTexture.cpp
@@ -1523,11 +1523,6 @@ FTexture2DRHIRef FMetalDynamicRHI::RHIAsyncCreateTexture2D(uint32 SizeX, uint32
 	return FTexture2DRHIRef();
 }
 
-void FMetalDynamicRHI::RHICopySharedMips(FRHICommandList& RHICmdList, FRHITexture2D* DestTexture2D, FRHITexture2D* SrcTexture2D)
-{
-	NOT_SUPPORTED("RHICopySharedMips");
-}
-
 void FMetalDynamicRHI::RHIGenerateMips(FRHITexture* SourceSurfaceRHI)
 {
 	@autoreleasepool {
diff --git a/Engine/Source/Runtime/D3D12RHI/Private/D3D12RHIPrivate.h b/Engine/Source/Runtime/D3D12RHI/Private/D3D12RHIPrivate.h
--- a/Engine/Source/Runtime/D3D12RHI/Private/D3D12RHIPrivate.h
+++ b/Engine/Source/Runtime/D3D12RHI/Private/D3D12RHIPrivate.h
@@ -466,7 +466,6 @@ public:
 	virtual bool RHIGetTextureMemoryVisualizeData(FColor* TextureData, int32 SizeX, int32 SizeY, int32 Pitch, int32 PixelSize) final override;
 	virtual FTextureRHIRef RHIAsyncCreateTexture2D(uint32 SizeX, uint32 SizeY, uint8 Format, uint32 NumMips, ETextureCreateFlags Flags, ERHIAccess InResourceState, void** InitialMipData, uint32 NumInitialMips) final override;
 	virtual FTextureRHIRef RHICreateTexture(const FRHITextureCreateDesc& CreateDesc) override;
-	virtual void RHICopySharedMips(FRHICommandList& RHICmdList, FRHITexture2D* DestTexture2D, FRHITexture2D* SrcTexture2D) final override;
 	virtual FShaderResourceViewRHIRef RHICreateShaderResourceView(FRHITexture* Texture, const FRHITextureSRVCreateInfo& CreateInfo) final override;
 	virtual uint32 RHIComputeMemorySize(FRHITexture* TextureRHI) final override;
 	virtual FTexture2DRHIRef RHIAsyncReallocateTexture2D(FRHITexture2D* Texture2D, int32 NewMipCount, int32 NewSizeX, int32 NewSizeY, FThreadSafeCounter* RequestStatus) override;
diff --git a/Engine/Source/Runtime/D3D12RHI/Private/D3D12Texture.cpp b/Engine/Source/Runtime/D3D12RHI/Private/D3D12Texture.cpp
--- a/Engine/Source/Runtime/D3D12RHI/Private/D3D12Texture.cpp
+++ b/Engine/Source/Runtime/D3D12RHI/Private/D3D12Texture.cpp
@@ -1169,60 +1169,6 @@ FTextureRHIRef FD3D12DynamicRHI::RHIAsyncCreateTexture2D(uint32 SizeX, uint32 Si
 }
 
 
-void FD3D12DynamicRHI::RHICopySharedMips(FRHICommandList& RHICmdList, FRHITexture2D* DestTexture2DRHI, FRHITexture2D* SrcTexture2DRHI)
-{
-	FD3D12Texture*  DestTexture2D = GetD3D12TextureFromRHITexture(DestTexture2DRHI);
-	FD3D12Texture*  SrcTexture2D = GetD3D12TextureFromRHITexture(SrcTexture2DRHI);
-
-	// Use the GPU to asynchronously copy the old mip-maps into the new texture.
-	const uint32 NumSharedMips = FMath::Min(DestTexture2D->GetNumMips(), SrcTexture2D->GetNumMips());
-	const uint32 SourceMipOffset = SrcTexture2D->GetNumMips() - NumSharedMips;
-	const uint32 DestMipOffset = DestTexture2D->GetNumMips() - NumSharedMips;
-
-	uint32 srcSubresource = 0;
-	uint32 destSubresource = 0;
-
-	FD3D12Adapter* Adapter = &GetAdapter();
-
-	for (FD3D12Texture::FDualLinkedObjectIterator It(DestTexture2D, SrcTexture2D); It; ++It)
-	{
-		DestTexture2D = static_cast<FD3D12Texture*>(It.GetFirst());
-		SrcTexture2D = static_cast<FD3D12Texture*>(It.GetSecond());
-
-		FD3D12CommandContext& Context = DestTexture2D->GetParentDevice()->GetDefaultCommandContext();
-
-		{
-			FScopedResourceBarrier ScopeResourceBarrierDst(Context, DestTexture2D->GetResource(), D3D12_RESOURCE_STATE_COPY_DEST  , D3D12_RESOURCE_BARRIER_ALL_SUBRESOURCES);
-			FScopedResourceBarrier ScopeResourceBarrierSrc(Context, SrcTexture2D->GetResource() , D3D12_RESOURCE_STATE_COPY_SOURCE, D3D12_RESOURCE_BARRIER_ALL_SUBRESOURCES);
-			Context.FlushResourceBarriers();
-
-			for (uint32 MipIndex = 0; MipIndex < NumSharedMips; ++MipIndex)
-			{
-				// Use the GPU to copy between mip-maps.
-				srcSubresource = CalcSubresource(MipIndex + SourceMipOffset, 0, SrcTexture2D->GetNumMips());
-				destSubresource = CalcSubresource(MipIndex + DestMipOffset, 0, DestTexture2D->GetNumMips());
-
-				CD3DX12_TEXTURE_COPY_LOCATION DestCopyLocation(DestTexture2D->GetResource()->GetResource(), destSubresource);
-				CD3DX12_TEXTURE_COPY_LOCATION SourceCopyLocation(SrcTexture2D->GetResource()->GetResource(), srcSubresource);
-
-				Context.GraphicsCommandList()->CopyTextureRegion(
-					&DestCopyLocation,
-					0, 0, 0,
-					&SourceCopyLocation,
-					nullptr);
-
-				Context.UpdateResidency(DestTexture2D->GetResource());
-				Context.UpdateResidency(SrcTexture2D->GetResource());
-			}
-		}
-
-		Context.ConditionalSplitCommandList();
-
-		DEBUG_EXECUTE_COMMAND_CONTEXT(Context);
-	}
-}
-
-
 /**
  * Computes the size in memory required by a given texture.
  *
diff --git a/Engine/Source/Runtime/Engine/Private/Streaming/Texture2DArrayStreaming.cpp b/Engine/Source/Runtime/Engine/Private/Streaming/Texture2DArrayStreaming.cpp
--- a/Engine/Source/Runtime/Engine/Private/Streaming/Texture2DArrayStreaming.cpp
+++ b/Engine/Source/Runtime/Engine/Private/Streaming/Texture2DArrayStreaming.cpp
@@ -7,51 +7,6 @@ Texture2DArrayStreaming.cpp: Helpers to stream in and out texture 2D array LODs.
 #include "Streaming/Texture2DArrayStreaming.h"
 #include "RenderUtils.h"
 
-//*****************************************************************************
-//***************************** Global Definitions ****************************
-//*****************************************************************************
-
-void RHICopySharedMips(FRHICommandList& RHICmdList, FRHITexture2DArray* DestTexture, FRHITexture2DArray* SrcTexture)
-{
-	// Transition to copy source and dest
-	{
-		FRHITransitionInfo TransitionsBefore[] = { FRHITransitionInfo(SrcTexture, ERHIAccess::SRVMask, ERHIAccess::CopySrc), FRHITransitionInfo(DestTexture, ERHIAccess::SRVMask, ERHIAccess::CopyDest) };
-		RHICmdList.Transition(MakeArrayView(TransitionsBefore, UE_ARRAY_COUNT(TransitionsBefore)));
-	}
-
-	// Copy 
-	{
-		FRHICopyTextureInfo CopyInfo;
-
-		auto SetCopyInfo = [&](FRHITexture2DArray* Texture2DArray)
-		{
-			CopyInfo.Size.X = Texture2DArray->GetSizeX();
-			CopyInfo.Size.Y = Texture2DArray->GetSizeY();
-			CopyInfo.NumSlices = Texture2DArray->GetSizeZ();
-			CopyInfo.NumMips = Texture2DArray->GetNumMips();
-		};
-
-		if (DestTexture->GetNumMips() < SrcTexture->GetNumMips())
-		{
-			SetCopyInfo(DestTexture);
-		}
-		else
-		{
-			SetCopyInfo(SrcTexture);
-		}
-
-		CopyInfo.SourceMipIndex = SrcTexture->GetNumMips() - CopyInfo.NumMips;
-		CopyInfo.DestMipIndex = DestTexture->GetNumMips() - CopyInfo.NumMips;
-		RHICmdList.CopyTexture(SrcTexture, DestTexture, CopyInfo);
-	}
-
-	// Transition to SRV
-	{
-		FRHITransitionInfo TransitionsAfter[] = { FRHITransitionInfo(SrcTexture, ERHIAccess::CopySrc, ERHIAccess::SRVMask), FRHITransitionInfo(DestTexture, ERHIAccess::CopyDest, ERHIAccess::SRVMask) };
-		RHICmdList.Transition(MakeArrayView(TransitionsAfter, UE_ARRAY_COUNT(TransitionsAfter)));
-	}
-}
-
 //*****************************************************************************
 //******************* FTexture2DArrayMipAllocator_Reallocate ******************
 //*****************************************************************************
@@ -121,7 +76,11 @@ bool FTexture2DArrayMipAllocator_Reallocate::FinalizeMips(const FTextureUpdateCo
 		ENQUEUE_RENDER_COMMAND(FCopySharedMipsForTexture2DArray)(
 			[&](FRHICommandListImmediate& RHICmdList)
 		{
-			RHICopySharedMips(RHICmdList, IntermediateTextureRHI.GetReference(), Context.Resource->GetTexture2DArrayRHI());
+			FRHITexture* SrcTexture = Context.Resource->GetTexture2DArrayRHI();
+			FRHITexture* DstTexture = IntermediateTextureRHI.GetReference();
+
+			UE::RHI::CopySharedMips_AssumeSRVMaskState(RHICmdList, SrcTexture, DstTexture);
+
 			bCopySharedMipsDone = true;
 		});
 		// Expected to execute immediately since ran on the renderthread.
diff --git a/Engine/Source/Runtime/Engine/Private/Streaming/Texture2DMipAllocator_AsyncCreate.cpp b/Engine/Source/Runtime/Engine/Private/Streaming/Texture2DMipAllocator_AsyncCreate.cpp
--- a/Engine/Source/Runtime/Engine/Private/Streaming/Texture2DMipAllocator_AsyncCreate.cpp
+++ b/Engine/Source/Runtime/Engine/Private/Streaming/Texture2DMipAllocator_AsyncCreate.cpp
@@ -92,7 +92,10 @@ bool FTexture2DMipAllocator_AsyncCreate::FinalizeMips(const FTextureUpdateContex
 	else
 	{
 		// Copy the mips.
-		RHICopySharedMips(IntermediateTextureRHI, Texture2DRHI);
+		UE::RHI::CopySharedMips_AssumeSRVMaskState(
+			FRHICommandListExecutor::GetImmediateCommandList(),
+			Texture2DRHI,
+			IntermediateTextureRHI);
 		// Use the new texture resource for the texture asset, must run on the renderthread.
 		Context.Resource->FinalizeStreaming(IntermediateTextureRHI);
 		// No need for the intermediate texture anymore.
diff --git a/Engine/Source/Runtime/Engine/Private/Streaming/Texture2DStreamIn.cpp b/Engine/Source/Runtime/Engine/Private/Streaming/Texture2DStreamIn.cpp
--- a/Engine/Source/Runtime/Engine/Private/Streaming/Texture2DStreamIn.cpp
+++ b/Engine/Source/Runtime/Engine/Private/Streaming/Texture2DStreamIn.cpp
@@ -112,7 +112,10 @@ void FTexture2DStreamIn::DoCopySharedMips(const FContext& Context)
 
 	if (!IsCancelled() && IntermediateTextureRHI && Context.Resource)
 	{
-		RHICopySharedMips(IntermediateTextureRHI, Context.Resource->GetTexture2DRHI());
+		UE::RHI::CopySharedMips_AssumeSRVMaskState(
+			FRHICommandListExecutor::GetImmediateCommandList(),
+			Context.Resource->GetTexture2DRHI(),
+			IntermediateTextureRHI);
 	}
 }
 
diff --git a/Engine/Source/Runtime/Engine/Private/Streaming/Texture2DStreamOut_AsyncCreate.cpp b/Engine/Source/Runtime/Engine/Private/Streaming/Texture2DStreamOut_AsyncCreate.cpp
--- a/Engine/Source/Runtime/Engine/Private/Streaming/Texture2DStreamOut_AsyncCreate.cpp
+++ b/Engine/Source/Runtime/Engine/Private/Streaming/Texture2DStreamOut_AsyncCreate.cpp
@@ -52,7 +52,10 @@ void FTexture2DStreamOut_AsyncCreate::Finalize(const FContext& Context)
 
 	if (!IsCancelled() && IntermediateTextureRHI && Context.Resource)
 	{
-		RHICopySharedMips(IntermediateTextureRHI, Context.Resource->GetTexture2DRHI());
+		UE::RHI::CopySharedMips_AssumeSRVMaskState(
+			FRHICommandListExecutor::GetImmediateCommandList(),
+			Context.Resource->GetTexture2DRHI(),
+			IntermediateTextureRHI);
 	}
 
 	RHIFinalizeAsyncReallocateTexture2D(IntermediateTextureRHI, true);
diff --git a/Engine/Source/Runtime/Engine/Private/Streaming/Texture2DUpdate.cpp b/Engine/Source/Runtime/Engine/Private/Streaming/Texture2DUpdate.cpp
--- a/Engine/Source/Runtime/Engine/Private/Streaming/Texture2DUpdate.cpp
+++ b/Engine/Source/Runtime/Engine/Private/Streaming/Texture2DUpdate.cpp
@@ -104,7 +104,11 @@ void FTexture2DUpdate::DoConvertToVirtualWithNewMips(const FContext& Context)
 
 			RHIVirtualTextureSetFirstMipInMemory(IntermediateTextureRHI, CurrentFirstLODIdx);
 			RHIVirtualTextureSetFirstMipVisible(IntermediateTextureRHI, CurrentFirstLODIdx);
-			RHICopySharedMips(IntermediateTextureRHI, Context.Resource->GetTexture2DRHI());
+
+			UE::RHI::CopySharedMips_AssumeSRVMaskState(
+				FRHICommandListExecutor::GetImmediateCommandList(),
+				Context.Resource->GetTexture2DRHI(),
+				IntermediateTextureRHI);
 		}
 		else
 		{
@@ -136,7 +140,10 @@ bool FTexture2DUpdate::DoConvertToNonVirtual(const FContext& Context)
 
 			IntermediateTextureRHI = RHICreateTexture(Desc);
 
-			RHICopySharedMips(IntermediateTextureRHI, Context.Resource->GetTexture2DRHI());
+			UE::RHI::CopySharedMips_AssumeSRVMaskState(
+				FRHICommandListExecutor::GetImmediateCommandList(),
+				Context.Resource->GetTexture2DRHI(),
+				IntermediateTextureRHI);
 
 			return true;
 		}
diff --git a/Engine/Source/Runtime/Engine/Private/Streaming/VolumeTextureStreaming.cpp b/Engine/Source/Runtime/Engine/Private/Streaming/VolumeTextureStreaming.cpp
--- a/Engine/Source/Runtime/Engine/Private/Streaming/VolumeTextureStreaming.cpp
+++ b/Engine/Source/Runtime/Engine/Private/Streaming/VolumeTextureStreaming.cpp
@@ -95,7 +95,11 @@ bool FVolumeTextureMipAllocator_Reallocate::FinalizeMips(const FTextureUpdateCon
 		ENQUEUE_RENDER_COMMAND(FCopySharedMipsForTexture3D)(
 			[&](FRHICommandListImmediate& RHICmdList)
 		{
-			RHICmdList.CopySharedMips(IntermediateTextureRHI.GetReference(), Context.Resource->GetTexture3DRHI());
+			FRHITexture* SrcTexture = Context.Resource->GetTexture3DRHI();
+			FRHITexture* DstTexture = IntermediateTextureRHI.GetReference();
+
+			UE::RHI::CopySharedMips_AssumeSRVMaskState(RHICmdList, SrcTexture, DstTexture);
+
 			bCopySharedMipsDone = true;
 		});
 		// Expected to execute immediately since ran on the renderthread.
diff --git a/Engine/Source/Runtime/NullDrv/Public/NullRHI.h b/Engine/Source/Runtime/NullDrv/Public/NullRHI.h
--- a/Engine/Source/Runtime/NullDrv/Public/NullRHI.h
+++ b/Engine/Source/Runtime/NullDrv/Public/NullRHI.h
@@ -227,10 +227,6 @@ public:
 		return this->RHICreateTexture(Desc);
 	}
 
-	virtual void RHICopySharedMips(FRHICommandList& RHICmdList, FRHITexture2D* DestTexture2D, FRHITexture2D* SrcTexture2D) final override
-	{
-	}
-
 	virtual FShaderResourceViewRHIRef RHICreateShaderResourceView(FRHITexture* Texture, const FRHITextureSRVCreateInfo& CreateInfo) final override
 	{ 
 		return new FRHIShaderResourceView(Texture); 
diff --git a/Engine/Source/Runtime/OpenGLDrv/Private/OpenGLTexture.cpp b/Engine/Source/Runtime/OpenGLDrv/Private/OpenGLTexture.cpp
--- a/Engine/Source/Runtime/OpenGLDrv/Private/OpenGLTexture.cpp
+++ b/Engine/Source/Runtime/OpenGLDrv/Private/OpenGLTexture.cpp
@@ -1646,11 +1646,6 @@ FTextureRHIRef FOpenGLDynamicRHI::RHIAsyncCreateTexture2D(uint32 SizeX, uint32 S
 	return FTexture2DRHIRef();
 }
 
-void FOpenGLDynamicRHI::RHICopySharedMips(FRHICommandList& RHICmdList, FRHITexture2D* DestTexture2D, FRHITexture2D* SrcTexture2D)
-{
-	check(0);
-}
-
 FOpenGLShaderResourceView::FOpenGLShaderResourceView(FOpenGLTexture* InTexture, const FRHITextureSRVCreateInfo& CreateInfo)
 	: FRHIShaderResourceView(InTexture)
 	, Target      (InTexture->Target)
diff --git a/Engine/Source/Runtime/OpenGLDrv/Public/OpenGLDrv.h b/Engine/Source/Runtime/OpenGLDrv/Public/OpenGLDrv.h
--- a/Engine/Source/Runtime/OpenGLDrv/Public/OpenGLDrv.h
+++ b/Engine/Source/Runtime/OpenGLDrv/Public/OpenGLDrv.h
@@ -389,7 +389,6 @@ public:
 	virtual bool RHIGetTextureMemoryVisualizeData(FColor* TextureData,int32 SizeX,int32 SizeY,int32 Pitch,int32 PixelSize) final override;
 	virtual FTextureRHIRef RHICreateTexture(const FRHITextureCreateDesc& CreateDesc) final override;
 	virtual FTextureRHIRef RHIAsyncCreateTexture2D(uint32 SizeX, uint32 SizeY, uint8 Format, uint32 NumMips, ETextureCreateFlags Flags, ERHIAccess InResourceState, void** InitialMipData, uint32 NumInitialMips) final override;
-	virtual void RHICopySharedMips(FRHICommandList& RHICmdList, FRHITexture2D* DestTexture2D, FRHITexture2D* SrcTexture2D) final override;
 	virtual FShaderResourceViewRHIRef RHICreateShaderResourceView(FRHITexture* Texture, const FRHITextureSRVCreateInfo& CreateInfo) final override;
 	virtual void RHIGenerateMips(FRHITexture* Texture) final override;
 	virtual uint32 RHIComputeMemorySize(FRHITexture* TextureRHI) final override;
diff --git a/Engine/Source/Runtime/RHI/Private/RHI.cpp b/Engine/Source/Runtime/RHI/Private/RHI.cpp
--- a/Engine/Source/Runtime/RHI/Private/RHI.cpp
+++ b/Engine/Source/Runtime/RHI/Private/RHI.cpp
@@ -3268,53 +3268,50 @@ void FDebugName::AppendString(FStringBuilderBase& Builder) const
 	}
 }
 
-// ------------------------------------------------------------------------------------------------------------------------------------------------
-// 
-// @todo move these functions into RHICore when the RHICmdList.CopySharedMips and RHICopySharedMips functions are removed (currently deprecated).
-//
-RHI_API void RHICopySharedMips_PRIVATE(FRHICommandList& RHICmdList, FRHITexture* SrcTexture, FRHITexture* DstTexture)
+namespace UE::RHI
 {
-	FRHITextureDesc const& Desc = DstTexture->GetNumMips() < SrcTexture->GetNumMips()
-		? DstTexture->GetDesc()
-		: SrcTexture->GetDesc();
 
-	FRHICopyTextureInfo CopyInfo;
-	CopyInfo.Size.X         = Desc.Extent.X;
-	CopyInfo.Size.Y         = Desc.Extent.Y;
-	CopyInfo.Size.Z         = Desc.Depth;
-	CopyInfo.NumSlices      = Desc.ArraySize;
-	CopyInfo.NumMips        = Desc.NumMips;
-	CopyInfo.SourceMipIndex = SrcTexture->GetNumMips() - CopyInfo.NumMips;
-	CopyInfo.DestMipIndex   = DstTexture->GetNumMips() - CopyInfo.NumMips;
+	RHI_API void CopySharedMips(FRHICommandList& RHICmdList, FRHITexture* SrcTexture, FRHITexture* DstTexture)
+	{
+		FRHITextureDesc const& Desc = DstTexture->GetNumMips() < SrcTexture->GetNumMips()
+			? DstTexture->GetDesc()
+			: SrcTexture->GetDesc();
 
-	RHICmdList.CopyTexture(SrcTexture, DstTexture, CopyInfo);
-}
+		FRHICopyTextureInfo CopyInfo;
+		CopyInfo.Size.X         = Desc.Extent.X;
+		CopyInfo.Size.Y         = Desc.Extent.Y;
+		CopyInfo.Size.Z         = Desc.Depth;
+		CopyInfo.NumSlices      = Desc.ArraySize;
+		CopyInfo.NumMips        = Desc.NumMips;
+		CopyInfo.SourceMipIndex = SrcTexture->GetNumMips() - CopyInfo.NumMips;
+		CopyInfo.DestMipIndex   = DstTexture->GetNumMips() - CopyInfo.NumMips;
 
-RHI_API void RHICopySharedMips_AssumeSRVMaskState_PRIVATE(FRHICommandList& RHICmdList, FRHITexture* SrcTexture, FRHITexture* DstTexture)
-{
-	// Transition to copy source and dest
+		RHICmdList.CopyTexture(SrcTexture, DstTexture, CopyInfo);
+	}
+
+	RHI_API void CopySharedMips_AssumeSRVMaskState(FRHICommandList& RHICmdList, FRHITexture* SrcTexture, FRHITexture* DstTexture)
 	{
-		FRHITransitionInfo TransitionsBefore[] =
+		// Transition to copy source and dest
 		{
-			FRHITransitionInfo(SrcTexture, ERHIAccess::SRVMask, ERHIAccess::CopySrc ),
-			FRHITransitionInfo(DstTexture, ERHIAccess::SRVMask, ERHIAccess::CopyDest)
-		};
-		RHICmdList.Transition(MakeArrayView(TransitionsBefore, UE_ARRAY_COUNT(TransitionsBefore)));
-	}
+			FRHITransitionInfo TransitionsBefore[] =
+			{
+				FRHITransitionInfo(SrcTexture, ERHIAccess::SRVMask, ERHIAccess::CopySrc),
+				FRHITransitionInfo(DstTexture, ERHIAccess::SRVMask, ERHIAccess::CopyDest)
+			};
+			RHICmdList.Transition(MakeArrayView(TransitionsBefore, UE_ARRAY_COUNT(TransitionsBefore)));
+		}
 
-	RHICopySharedMips_PRIVATE(RHICmdList, SrcTexture, DstTexture);
+		CopySharedMips(RHICmdList, SrcTexture, DstTexture);
 
-	// Transition to SRV
-	{
-		FRHITransitionInfo TransitionsAfter[] =
+		// Transition to SRV
 		{
-			FRHITransitionInfo(SrcTexture, ERHIAccess::CopySrc , ERHIAccess::SRVMask),
-			FRHITransitionInfo(DstTexture, ERHIAccess::CopyDest, ERHIAccess::SRVMask)
-		};
-		RHICmdList.Transition(MakeArrayView(TransitionsAfter, UE_ARRAY_COUNT(TransitionsAfter)));
+			FRHITransitionInfo TransitionsAfter[] =
+			{
+				FRHITransitionInfo(SrcTexture, ERHIAccess::CopySrc, ERHIAccess::SRVMask),
+				FRHITransitionInfo(DstTexture, ERHIAccess::CopyDest, ERHIAccess::SRVMask)
+			};
+			RHICmdList.Transition(MakeArrayView(TransitionsAfter, UE_ARRAY_COUNT(TransitionsAfter)));
+		}
 	}
-}
-//
-// @todo move these functions into RHICore when the RHICmdList.CopySharedMips and RHICopySharedMips functions are removed (currently deprecated).
-// 
-// ------------------------------------------------------------------------------------------------------------------------------------------------
\ No newline at end of file
+
+} //! UE::RHI
diff --git a/Engine/Source/Runtime/RHI/Public/DynamicRHI.h b/Engine/Source/Runtime/RHI/Public/DynamicRHI.h
--- a/Engine/Source/Runtime/RHI/Public/DynamicRHI.h
+++ b/Engine/Source/Runtime/RHI/Public/DynamicRHI.h
@@ -539,16 +539,6 @@ public:
 	// FlushType: Thread safe
 	virtual FTextureRHIRef RHIAsyncCreateTexture2D(uint32 SizeX, uint32 SizeY, uint8 Format, uint32 NumMips, ETextureCreateFlags Flags, ERHIAccess InResourceState, void** InitialMipData, uint32 NumInitialMips) = 0;
 
-	/**
-	* Copies shared mip levels from one texture to another. The textures must have
-	* full mip chains, share the same format, and have the same aspect ratio. This
-	* copy will not cause synchronization with the GPU.
-	* @param DestTexture2D - destination texture
-	* @param SrcTexture2D - source texture
-	*/
-	// FlushType: Flush RHI Thread
-	virtual void RHICopySharedMips(FRHICommandList& RHICmdList, FRHITexture2D* DestTexture2D, FRHITexture2D* SrcTexture2D) = 0;
-
 	/**
 	* Creates a shader resource view for a texture
 	*/
diff --git a/Engine/Source/Runtime/RHI/Public/RHICommandList.h b/Engine/Source/Runtime/RHI/Public/RHICommandList.h
--- a/Engine/Source/Runtime/RHI/Public/RHICommandList.h
+++ b/Engine/Source/Runtime/RHI/Public/RHICommandList.h
@@ -4537,20 +4537,8 @@ public:
 		return GDynamicRHI->RHIGetTextureMemoryVisualizeData(TextureData,SizeX,SizeY,Pitch,PixelSize);
 	}
 	
-	FORCEINLINE void CopySharedMips(FRHITexture* DestTexture, FRHITexture* SrcTexture)
-	{
-		QUICK_SCOPE_CYCLE_COUNTER(STAT_RHIMETHOD_CopySharedMips_Flush);
-		DestTexture->AddRef();
-		SrcTexture->AddRef();
-		EnqueueLambda([DestTexture, SrcTexture](FRHICommandList& ExecutingCmdList)
-		{
-			LLM_SCOPE(ELLMTag::Textures);
-			//@todo move this function onto the IRHIComputeCommandContext
-			GDynamicRHI->RHICopySharedMips(ExecutingCmdList, DestTexture, SrcTexture);
-			DestTexture->Release();
-			SrcTexture->Release();
-		});
-	}
+	UE_DEPRECATED(5.2, "CopySharedMips should be replaced with UE::RHI::CopySharedMips(). ** Note ** The argument ordering of the new function has changed. The new function also does not issue RHI transitions, so the caller is responsible for ensuring src/dst textures are in CopySrc/CopyDest state.")
+	FORCEINLINE void CopySharedMips(FRHITexture* DestTexture, FRHITexture* SrcTexture);
 
 	FORCEINLINE FShaderResourceViewRHIRef CreateShaderResourceView(FRHITexture* Texture, const FRHITextureSRVCreateInfo& CreateInfo)
 	{
@@ -5618,9 +5606,12 @@ PRAGMA_DISABLE_DEPRECATION_WARNINGS
 PRAGMA_ENABLE_DEPRECATION_WARNINGS
 }
 
+UE_DEPRECATED(5.2, "RHICopySharedMips should be replaced with UE::RHI::CopySharedMips(). ** Note ** The argument ordering of the new function has changed. The new function also does not issue RHI transitions, so the caller is responsible for ensuring src/dst textures are in CopySrc/CopyDest state.")
 FORCEINLINE void RHICopySharedMips(FRHITexture* DestTexture, FRHITexture* SrcTexture)
 {
+PRAGMA_DISABLE_DEPRECATION_WARNINGS
 	return FRHICommandListExecutor::GetImmediateCommandList().CopySharedMips(DestTexture, SrcTexture);
+PRAGMA_ENABLE_DEPRECATION_WARNINGS
 }
 
 FORCEINLINE FShaderResourceViewRHIRef RHICreateShaderResourceView(FRHITexture* Texture, uint8 MipLevel)
@@ -5798,6 +5789,26 @@ FORCEINLINE void RHIUnlockStagingBuffer(FRHIStagingBuffer* StagingBuffer)
 	 FRHICommandListExecutor::GetImmediateCommandList().UnlockStagingBuffer(StagingBuffer);
 }
 
+namespace UE::RHI
+{
+
+	//
+	// Copies shared mip levels from one texture to another.
+	// Both textures must have full mip chains, share the same format, and have the same aspect ratio.
+	// The source texture must be in the CopySrc state, and the destination texture must be in the CopyDest state.
+	// 
+	RHI_API void CopySharedMips(FRHICommandList& RHICmdList, FRHITexture* SrcTexture, FRHITexture* DstTexture);
+
+	//
+	// Same as CopySharedMips(), but assumes both source and destination textures are in the SRVMask state.
+	// Adds transitions to move the textures to/from the CopySrc/CopyDest states, restoring SRVMask when done.
+	//
+	// Provided for backwards compatibility. Caller should prefer CopySharedMips() with optimally batched transitions.
+	//
+	RHI_API void CopySharedMips_AssumeSRVMaskState(FRHICommandList& RHICmdList, FRHITexture* SrcTexture, FRHITexture* DstTexture);
+
+} //! UE::RHI
+
 template <uint32 MaxNumUpdates>
 struct TRHIResourceUpdateBatcher
 {
diff --git a/Engine/Source/Runtime/RHI/Public/RHICommandList.inl b/Engine/Source/Runtime/RHI/Public/RHICommandList.inl
--- a/Engine/Source/Runtime/RHI/Public/RHICommandList.inl
+++ b/Engine/Source/Runtime/RHI/Public/RHICommandList.inl
@@ -161,4 +161,12 @@ inline FRHICommandListScopedPipelineGuard::~FRHICommandListScopedPipelineGuard()
 	{
 		RHICmdList.SwitchPipeline(ERHIPipeline::None);
 	}
-}
\ No newline at end of file
+}
+
+FORCEINLINE void FRHICommandListImmediate::CopySharedMips(FRHITexture* DestTexture, FRHITexture* SrcTexture)
+{
+	QUICK_SCOPE_CYCLE_COUNTER(STAT_RHIMETHOD_CopySharedMips_Flush);
+
+	// Use the version that assumes SRVMask state on both textures to maintain behaviour of the old code path.
+	UE::RHI::CopySharedMips_AssumeSRVMaskState(*this, SrcTexture, DestTexture);
+}
diff --git a/Engine/Source/Runtime/RHI/Public/RHIValidation.h b/Engine/Source/Runtime/RHI/Public/RHIValidation.h
--- a/Engine/Source/Runtime/RHI/Public/RHIValidation.h
+++ b/Engine/Source/Runtime/RHI/Public/RHIValidation.h
@@ -586,19 +586,6 @@ public:
 		return Texture;
 	}
 
-	/**
-	* Copies shared mip levels from one texture to another. The textures must have
-	* full mip chains, share the same format, and have the same aspect ratio. This
-	* copy will not cause synchronization with the GPU.
-	* @param DestTexture2D - destination texture
-	* @param SrcTexture2D - source texture
-	*/
-	// FlushType: Flush RHI Thread
-	virtual void RHICopySharedMips(FRHICommandList& RHICmdList, FRHITexture2D* DestTexture2D, FRHITexture2D* SrcTexture2D) override final
-	{
-		RHI->RHICopySharedMips(RHICmdList, DestTexture2D, SrcTexture2D);
-	}
-
 	void RHITransferBufferUnderlyingResource(FRHIBuffer* DestBuffer, FRHIBuffer* SrcBuffer) override final
 	{
 		RHI->RHITransferBufferUnderlyingResource(DestBuffer, SrcBuffer);
diff --git a/Engine/Source/Runtime/VulkanRHI/Private/VulkanTexture.cpp b/Engine/Source/Runtime/VulkanRHI/Private/VulkanTexture.cpp
--- a/Engine/Source/Runtime/VulkanRHI/Private/VulkanTexture.cpp
+++ b/Engine/Source/Runtime/VulkanRHI/Private/VulkanTexture.cpp
@@ -925,11 +925,6 @@ FTextureRHIRef FVulkanDynamicRHI::RHIAsyncCreateTexture2D(uint32 SizeX,uint32 Si
 	return FTextureRHIRef();
 }
 
-void FVulkanDynamicRHI::RHICopySharedMips(FRHICommandList& RHICmdList, FRHITexture2D* DestTexture2D, FRHITexture2D* SrcTexture2D)
-{
-	VULKAN_SIGNAL_UNIMPLEMENTED();
-}
-
 static void DoAsyncReallocateTexture2D(FVulkanCommandListContext& Context, FVulkanTexture* OldTexture, FVulkanTexture* NewTexture, int32 NewMipCount, int32 NewSizeX, int32 NewSizeY, FThreadSafeCounter* RequestStatus)
 {
 	LLM_SCOPE_VULKAN(ELLMTagVulkan::VulkanTextures);
diff --git a/Engine/Source/Runtime/VulkanRHI/Public/VulkanDynamicRHI.h b/Engine/Source/Runtime/VulkanRHI/Public/VulkanDynamicRHI.h
--- a/Engine/Source/Runtime/VulkanRHI/Public/VulkanDynamicRHI.h
+++ b/Engine/Source/Runtime/VulkanRHI/Public/VulkanDynamicRHI.h
@@ -127,7 +127,6 @@ public:
 	virtual bool RHIGetTextureMemoryVisualizeData(FColor* TextureData, int32 SizeX, int32 SizeY, int32 Pitch, int32 PixelSize) final override;
 	virtual FTextureRHIRef RHICreateTexture(const FRHITextureCreateDesc& CreateDesc) final override;
 	virtual FTextureRHIRef RHIAsyncCreateTexture2D(uint32 SizeX, uint32 SizeY, uint8 Format, uint32 NumMips, ETextureCreateFlags Flags, ERHIAccess InResourceState, void** InitialMipData, uint32 NumInitialMips) final override;
-	virtual void RHICopySharedMips(FRHICommandList& RHICmdList, FRHITexture2D* DestTexture2D, FRHITexture2D* SrcTexture2D) final override;
 	virtual FShaderResourceViewRHIRef RHICreateShaderResourceView(FRHITexture* Texture, const FRHITextureSRVCreateInfo& CreateInfo) final override;
 	virtual uint32 RHIComputeMemorySize(FRHITexture* TextureRHI) final override;
 	virtual FTexture2DRHIRef RHIAsyncReallocateTexture2D(FRHITexture2D* Texture2D, int32 NewMipCount, int32 NewSizeX, int32 NewSizeY, FThreadSafeCounter* RequestStatus) final override;
diff --git a/Engine/Source/Runtime/Windows/D3D11RHI/Private/D3D11RHIPrivate.h b/Engine/Source/Runtime/Windows/D3D11RHI/Private/D3D11RHIPrivate.h
--- a/Engine/Source/Runtime/Windows/D3D11RHI/Private/D3D11RHIPrivate.h
+++ b/Engine/Source/Runtime/Windows/D3D11RHI/Private/D3D11RHIPrivate.h
@@ -494,7 +494,6 @@ public:
 	virtual bool RHIGetTextureMemoryVisualizeData(FColor* TextureData,int32 SizeX,int32 SizeY,int32 Pitch,int32 PixelSize) final override;
 	virtual FTextureRHIRef RHICreateTexture(const FRHITextureCreateDesc& CreateDesc) final override;
 	virtual FTextureRHIRef RHIAsyncCreateTexture2D(uint32 SizeX, uint32 SizeY, uint8 Format, uint32 NumMips, ETextureCreateFlags Flags, ERHIAccess InResourceState, void** InitialMipData, uint32 NumInitialMips) final override;
-	virtual void RHICopySharedMips(FRHICommandList& RHICmdList, FRHITexture2D* DestTexture2D, FRHITexture2D* SrcTexture2D) final override;
 	virtual FShaderResourceViewRHIRef RHICreateShaderResourceView(FRHITexture* TextureRHI, const FRHITextureSRVCreateInfo& CreateInfo) final override;
 	virtual void RHIGenerateMips(FRHITexture* Texture) final override;
 	virtual uint32 RHIComputeMemorySize(FRHITexture* TextureRHI) final override;
diff --git a/Engine/Source/Runtime/Windows/D3D11RHI/Private/D3D11Texture.cpp b/Engine/Source/Runtime/Windows/D3D11RHI/Private/D3D11Texture.cpp
--- a/Engine/Source/Runtime/Windows/D3D11RHI/Private/D3D11Texture.cpp
+++ b/Engine/Source/Runtime/Windows/D3D11RHI/Private/D3D11Texture.cpp
@@ -853,7 +853,6 @@ FD3D11Texture* FD3D11DynamicRHI::CreateD3D11Texture3D(FRHITextureCreateDesc cons
 
 	// Set up the texture bind flags.
 	check(!EnumHasAnyFlags(Flags, TexCreate_DepthStencilTargetable | TexCreate_ResolveTargetable));
-	check(EnumHasAllFlags(Flags, TexCreate_ShaderResource));
 
 	TArray<D3D11_SUBRESOURCE_DATA> SubResourceData;
 
@@ -1020,31 +1019,6 @@ FTextureRHIRef FD3D11DynamicRHI::RHIAsyncCreateTexture2D(uint32 SizeX, uint32 Si
 	return Texture;
 }
 
-void FD3D11DynamicRHI::RHICopySharedMips(FRHICommandList& RHICmdList, FRHITexture2D* DestTexture2DRHI, FRHITexture2D* SrcTexture2DRHI)
-{
-	FD3D11Texture* DestTexture2D = ResourceCast(DestTexture2DRHI);
-	FD3D11Texture* SrcTexture2D = ResourceCast(SrcTexture2DRHI);
-
-	// Use the GPU to asynchronously copy the old mip-maps into the new texture.
-	const uint32 NumSharedMips = FMath::Min(DestTexture2D->GetNumMips(),SrcTexture2D->GetNumMips());
-	const uint32 SourceMipOffset = SrcTexture2D->GetNumMips() - NumSharedMips;
-	const uint32 DestMipOffset   = DestTexture2D->GetNumMips() - NumSharedMips;
-	for(uint32 MipIndex = 0;MipIndex < NumSharedMips;++MipIndex)
-	{
-		// Use the GPU to copy between mip-maps.
-		Direct3DDeviceIMContext->CopySubresourceRegion(
-			DestTexture2D->GetResource(),
-			D3D11CalcSubresource(MipIndex + DestMipOffset,0,DestTexture2D->GetNumMips()),
-			0,
-			0,
-			0,
-			SrcTexture2D->GetResource(),
-			D3D11CalcSubresource(MipIndex + SourceMipOffset,0,SrcTexture2D->GetNumMips()),
-			NULL
-			);
-	}
-}
-
 FShaderResourceViewRHIRef FD3D11DynamicRHI::RHICreateShaderResourceView(FRHITexture* TextureRHI, const FRHITextureSRVCreateInfo& CreateInfo)
 {
 	FD3D11Texture* Texture = GetD3D11TextureFromRHITexture(TextureRHI);
