Add cvars to dynamically control the PSO precompile pool size based on number of hardware threads. This is an alternative to the existing fixed r.pso.PrecompileThreadPoolSize.
- r.pso.PrecompileThreadPoolPercentOfHardwareThreads sets a percentage of hardware threads to use. The final value is clamped with r.pso.PrecompileThreadPoolSizeMin and r.pso.PrecompileThreadPoolSizeMax.
- Priority of the pool threads is also exposed via r.pso.PrecompileThreadPoolThreadPriority

diff --git a/Engine/Source/Runtime/RHI/Private/PipelineStateCache.cpp b/Engine/Source/Runtime/RHI/Private/PipelineStateCache.cpp
--- a/Engine/Source/Runtime/RHI/Private/PipelineStateCache.cpp
+++ b/Engine/Source/Runtime/RHI/Private/PipelineStateCache.cpp
@@ -13,6 +13,7 @@ PipelineStateCache.cpp: Pipeline state cache implementation.
 #include "Misc/TimeGuard.h"
 #include "Containers/DiscardableKeyValueCache.h"
 #include "Async/Async.h"
+#include "HAL/PlatformMisc.h"
 
 // perform cache eviction each frame, used to stress the system and flush out bugs
 #define PSO_DO_CACHE_EVICT_EACH_FRAME 0
@@ -127,12 +128,53 @@ static int32 GPSOPrecompileThreadPoolSize = 0;
 static FAutoConsoleVariableRef GPSOPrecompileThreadPoolSizeVar(
 	TEXT("r.pso.PrecompileThreadPoolSize"),
 	GPSOPrecompileThreadPoolSize,
-	TEXT("The maximum number of threads available for concurrent PSO Precompiling.\n")
+	TEXT("The number of threads available for concurrent PSO Precompiling.\n")
+	TEXT("0 to disable threadpool usage when precompiling PSOs. (default)")
+	,
+	ECVF_RenderThreadSafe | ECVF_ReadOnly
+);
+
+static int32 GPSOPrecompileThreadPoolPercentOfHardwareThreads = 0;
+static FAutoConsoleVariableRef GPSOPrecompileThreadPoolPercentOfHardwareThreadsVar(
+	TEXT("r.pso.PrecompileThreadPoolPercentOfHardwareThreads"),
+	GPSOPrecompileThreadPoolPercentOfHardwareThreads,
+	TEXT("If > 0, use this percentage of cores (rounded up) for the PSO precompile thread pool\n")
+	TEXT("Use this as an alternative to r.pso.PrecompileThreadPoolSize\n")
 	TEXT("0 to disable threadpool usage when precompiling PSOs. (default)")
 	,
 	ECVF_RenderThreadSafe | ECVF_ReadOnly
 );
 
+static int32 GPSOPrecompileThreadPoolSizeMin = 0;
+static FAutoConsoleVariableRef GPSOPrecompileThreadPoolSizeMinVar(
+	TEXT("r.pso.PrecompileThreadPoolSizeMin"),
+	GPSOPrecompileThreadPoolSizeMin,
+	TEXT("The minimum number of threads available for concurrent PSO Precompiling.\n")
+	TEXT("Ignored unless r.pso.PrecompileThreadPoolPercentOfHardwareThreads is specified\n")
+	TEXT("0 = no minimum (default)")
+	,
+	ECVF_RenderThreadSafe | ECVF_ReadOnly
+);
+
+static int32 GPSOPrecompileThreadPoolSizeMax = INT_MAX;
+static FAutoConsoleVariableRef GPSOPrecompileThreadPoolSizeMaxVar(
+	TEXT("r.pso.PrecompileThreadPoolSizeMax"),
+	GPSOPrecompileThreadPoolSizeMax,
+	TEXT("The maximum number of threads available for concurrent PSO Precompiling.\n")
+	TEXT("Ignored unless r.pso.PrecompileThreadPoolPercentOfHardwareThreads is specified\n")
+	TEXT("Default is no maximum (INT_MAX)")
+	,
+	ECVF_RenderThreadSafe | ECVF_ReadOnly
+);
+
+int32 GPSOPrecompileThreadPoolThreadPriority = (int32)EThreadPriority::TPri_BelowNormal;
+static FAutoConsoleVariableRef CVarStreamingTextureIOPriority(
+	TEXT("r.pso.PrecompileThreadPoolThreadPriority"),
+	GPSOPrecompileThreadPoolThreadPriority,
+	TEXT("Thread priority for the PSO precompile pool"),
+	ECVF_RenderThreadSafe);
+
+
 class FPSOPrecacheThreadPool
 {
 public:
@@ -151,7 +193,7 @@ public:
 				check(UsePool());
 
 				FQueuedThreadPool* PSOPrecompileCompileThreadPoolLocal = FQueuedThreadPool::Allocate();
-				PSOPrecompileCompileThreadPoolLocal->Create(GPSOPrecompileThreadPoolSize, 512 * 1024, EThreadPriority::TPri_BelowNormal, TEXT("PSOPrecompilePool"));
+				PSOPrecompileCompileThreadPoolLocal->Create(GetDesiredPoolSize(), 512 * 1024, (EThreadPriority)GPSOPrecompileThreadPoolThreadPriority, TEXT("PSOPrecompilePool"));
 				PSOPrecompileCompileThreadPool = PSOPrecompileCompileThreadPoolLocal;
 			}
 		}
@@ -169,9 +211,24 @@ public:
 		}
 	}
 
+	static int32 GetDesiredPoolSize()
+	{
+		if (GPSOPrecompileThreadPoolSize > 0)
+		{
+			ensure(GPSOPrecompileThreadPoolPercentOfHardwareThreads == 0); // These settings are mutually exclusive
+			return GPSOPrecompileThreadPoolSize;
+		}
+		if (GPSOPrecompileThreadPoolPercentOfHardwareThreads > 0)
+		{
+			int32 NumThreads = FMath::CeilToInt((float)FPlatformMisc::NumberOfCoresIncludingHyperthreads() * (float)GPSOPrecompileThreadPoolPercentOfHardwareThreads / 100.0f);
+			return FMath::Clamp(NumThreads, GPSOPrecompileThreadPoolSizeMin, GPSOPrecompileThreadPoolSizeMax);
+		}
+		return 0;
+	}
+
 	static bool UsePool()
 	{
-		return GPSOPrecompileThreadPoolSize > 0;
+		return GetDesiredPoolSize() > 0;
 	}
 
 private:
