Enabling RHI Command List parameter batching after extensive testing.

diff --git a/Engine/Source/Runtime/RHI/Private/RHICommandList.cpp b/Engine/Source/Runtime/RHI/Private/RHICommandList.cpp
--- a/Engine/Source/Runtime/RHI/Private/RHICommandList.cpp
+++ b/Engine/Source/Runtime/RHI/Private/RHICommandList.cpp
@@ -136,7 +136,7 @@ static TStatId GCurrentExecuteStat;
 static FCriticalSection GRHIThreadOnTasksCritical;
 static std::atomic<int32> GRHIThreadStallRequestCount;
 
-int32 GRHICommandParameterBatching = 0;
+int32 GRHICommandParameterBatching = 1;
 static FAutoConsoleVariableRef CVarRHICommandParameterBatching(
 	TEXT("r.RHICmdParamBatching"),
 	GRHICommandParameterBatching,
