Disable printing new PSOs to the log unless -logpsodetails is on the command line. Previous behavior had this enabled by default in all but shipping configs.

diff --git a/Engine/Source/Runtime/RHI/Private/PipelineFileCache.cpp b/Engine/Source/Runtime/RHI/Private/PipelineFileCache.cpp
--- a/Engine/Source/Runtime/RHI/Private/PipelineFileCache.cpp
+++ b/Engine/Source/Runtime/RHI/Private/PipelineFileCache.cpp
@@ -112,11 +112,11 @@ static TAutoConsoleVariable<int32> CVarPSOFileCacheReportPSO(
 														   ECVF_Default | ECVF_RenderThreadSafe
 														   );
 
-static int32 GPSOFileCachePrintNewPSODescriptors = UE_BUILD_SHIPPING ? 0 : 1;
+static int32 GPSOFileCachePrintNewPSODescriptors = 0;
 static FAutoConsoleVariableRef CVarPSOFileCachePrintNewPSODescriptors(
 														   TEXT("r.ShaderPipelineCache.PrintNewPSODescriptors"),
 														   GPSOFileCachePrintNewPSODescriptors,
-														   TEXT("1 prints descriptions for all new PSO entries to the log/console while 0 does not. 2 prints additional details about the PSO. Defaults to 0 in *Shipping* builds, otherwise 1."),
+														   TEXT("1 prints descriptions for all new PSO entries to the log/console while 0 does not. 2 prints additional details about graphics PSO. Defaults to 0."),
 														   ECVF_Default
 														   );
 
@@ -3382,17 +3382,14 @@ void FPipelineFileCacheManager::CacheGraphicsPSO(uint32 RunTimeHash, FGraphicsPi
 					{
 						CSV_EVENT(PSO, TEXT("Encountered new graphics PSO"));
 						UE_LOG(LogRHI, Display, TEXT("Encountered a new graphics PSO: %u"), PSOHash);
-						if (GPSOFileCachePrintNewPSODescriptors > 0)
+						int32 LogDetailLevel = LogPSODetails() ? 2 : GPSOFileCachePrintNewPSODescriptors;
+						if (LogDetailLevel > 0)
 						{
 							UE_LOG(LogRHI, Display, TEXT("New Graphics PSO (%u)"), PSOHash);
-							if (LogPSODetails() || GPSOFileCachePrintNewPSODescriptors > 1)
+							if (LogDetailLevel > 1)
 							{
 								UE_LOG(LogRHI, Display, TEXT("%s"), *NewEntry.ToStringReadable());
 							}
-							else
-							{
-								UE_LOG(LogRHI, Display, TEXT("%s"), *NewEntry.GraphicsDesc.ToString());
-							}
 						}
 						if (LogPSOtoFileCache())
 						{
