Reimplementing ERHIBindlessSupport support in new DataDrivenShaderPlatformInfo files.

diff --git a/Engine/Source/Runtime/RHI/Private/DataDrivenShaderPlatformInfo.cpp b/Engine/Source/Runtime/RHI/Private/DataDrivenShaderPlatformInfo.cpp
--- a/Engine/Source/Runtime/RHI/Private/DataDrivenShaderPlatformInfo.cpp
+++ b/Engine/Source/Runtime/RHI/Private/DataDrivenShaderPlatformInfo.cpp
@@ -91,6 +91,33 @@ static inline uint32 GetSectionFeatureSupport(const FConfigSection& Section, FNa
 	return OriginalValue;
 }
 
+static inline uint32 GetSectionBindlessSupport(const FConfigSection& Section, FName Key, uint32 OriginalValue)
+{
+	const FConfigValue* ConfigValue = Section.Find(Key);
+	if (ConfigValue != nullptr)
+	{
+		FString Value = ConfigValue->GetValue();
+		if (Value == TEXT("Unsupported"))
+		{
+			return uint32(ERHIBindlessSupport::Unsupported);
+		}
+		if (Value == TEXT("RayTracingOnly"))
+		{
+			return uint32(ERHIBindlessSupport::RayTracingOnly);
+		}
+		else if (Value == TEXT("AllShaderTypes"))
+		{
+			return uint32(ERHIBindlessSupport::AllShaderTypes);
+		}
+		else
+		{
+			checkf(false, TEXT("Unknown ERHIBindlessSupport value \"%s\" for %s"), *Value, *Key.ToString());
+		}
+	}
+
+	return OriginalValue;
+}
+
 void FGenericDataDrivenShaderPlatformInfo::SetDefaultValues()
 {
 	MaxFeatureLevel = ERHIFeatureLevel::Num;
@@ -131,6 +158,10 @@ void FGenericDataDrivenShaderPlatformInfo::ParseDataDrivenShaderInfo(const FConf
 	Info.SettingName = GetSectionFeatureSupport(Section, #SettingName, Info.SettingName); \
 	ADD_TO_PROPERTIES_STRING(SettingName, Info.SettingName)
 
+#define GET_SECTION_BINDLESS_SUPPORT_HELPER(SettingName) \
+	Info.SettingName = GetSectionBindlessSupport(Section, #SettingName, Info.SettingName); \
+	ADD_TO_PROPERTIES_STRING(SettingName, Info.SettingName)
+
 	GET_SECTION_BOOL_HELPER(bIsMobile);
 	GET_SECTION_BOOL_HELPER(bIsMetalMRT);
 	GET_SECTION_BOOL_HELPER(bIsPC);
@@ -210,7 +241,7 @@ void FGenericDataDrivenShaderPlatformInfo::ParseDataDrivenShaderInfo(const FConf
 	GET_SECTION_BOOL_HELPER(bSupportsMobileDistanceField);
 	GET_SECTION_BOOL_HELPER(bSupportsFFTBloom);
 	GET_SECTION_BOOL_HELPER(bSupportsVertexShaderLayer);
-	GET_SECTION_BOOL_HELPER(bSupportsBindless);
+	GET_SECTION_BINDLESS_SUPPORT_HELPER(BindlessSupport);
 	GET_SECTION_BOOL_HELPER(bSupportsVolumeTextureAtomics);
 	GET_SECTION_BOOL_HELPER(bSupportsROV);
 	GET_SECTION_BOOL_HELPER(bSupportsOIT);
diff --git a/Engine/Source/Runtime/RHI/Public/DataDrivenShaderPlatformInfo.h b/Engine/Source/Runtime/RHI/Public/DataDrivenShaderPlatformInfo.h
--- a/Engine/Source/Runtime/RHI/Public/DataDrivenShaderPlatformInfo.h
+++ b/Engine/Source/Runtime/RHI/Public/DataDrivenShaderPlatformInfo.h
@@ -100,7 +100,7 @@ class RHI_API FGenericDataDrivenShaderPlatformInfo
 	uint32 bSupportsInlineRayTracing : 1;
 	uint32 bSupportsRayTracingShaders : 1;
 	uint32 bSupportsVertexShaderLayer : 1;
-	uint32 bSupportsBindless : 1;
+	uint32 BindlessSupport : int32(ERHIBindlessSupport::NumBits);
 	uint32 bSupportsVolumeTextureAtomics : 1;
 	uint32 bSupportsROV : 1;
 	uint32 bSupportsOIT : 1;
@@ -692,9 +692,16 @@ public:
 		return Infos[Platform].bSupportsVertexShaderLayer;
 	}
 
+	static FORCEINLINE_DEBUGGABLE const ERHIBindlessSupport GetBindlessSupport(const FStaticShaderPlatform Platform)
+	{
+		check(IsValid(Platform));
+		return static_cast<ERHIBindlessSupport>(Infos[Platform].BindlessSupport);
+	}
+
+	UE_DEPRECATED(5.2, "You must use GetBindlessSupport instead.")
 	static FORCEINLINE_DEBUGGABLE const bool GetSupportsBindless(const FStaticShaderPlatform Platform)
 	{
-		return Infos[Platform].bSupportsBindless;
+		return GetBindlessSupport(Platform) == ERHIBindlessSupport::AllShaderTypes;
 	}
 
 	static FORCEINLINE_DEBUGGABLE const bool GetSupportsVolumeTextureAtomics(const FStaticShaderPlatform Platform)
