Fixing bad merge.

diff --git a/Engine/Source/Runtime/RHI/Public/DataDrivenShaderPlatformInfo.h b/Engine/Source/Runtime/RHI/Public/DataDrivenShaderPlatformInfo.h
--- a/Engine/Source/Runtime/RHI/Public/DataDrivenShaderPlatformInfo.h
+++ b/Engine/Source/Runtime/RHI/Public/DataDrivenShaderPlatformInfo.h
@@ -1115,9 +1115,15 @@ inline RHI_API bool RHISupportsConservativeRasterization(const FStaticShaderPlat
 }
 
 /** True if the given shader platform supports bindless resources/views. */
+inline ERHIBindlessSupport RHIGetBindlessSupport(const FStaticShaderPlatform Platform)
+{
+	return FDataDrivenShaderPlatformInfo::GetBindlessSupport(Platform);
+}
+
+UE_DEPRECATED(5.2, "You must use RHIGetBindlessSupport instead.")
 inline bool RHISupportsBindless(EShaderPlatform Platform)
 {
-	return FDataDrivenShaderPlatformInfo::GetSupportsBindless(Platform);
+	return RHIGetBindlessSupport(Platform) == ERHIBindlessSupport::AllShaderTypes;
 }
 
 inline bool RHISupportsVolumeTextureAtomics(EShaderPlatform Platform)
diff --git a/Engine/Source/Runtime/RHI/Public/RHI.h b/Engine/Source/Runtime/RHI/Public/RHI.h
--- a/Engine/Source/Runtime/RHI/Public/RHI.h
+++ b/Engine/Source/Runtime/RHI/Public/RHI.h
@@ -155,14 +155,6 @@ struct FRHIResourceStats
 
 RHI_API void RHIGetTrackedResourceStats(TArray<TSharedPtr<FRHIResourceStats>>& OutResourceStats);
 
-inline ERHIBindlessSupport RHIGetBindlessSupport(const FStaticShaderPlatform Platform)
-{
-	return FDataDrivenShaderPlatformInfo::GetBindlessSupport(Platform);
-}
-
-UE_DEPRECATED(5.2, "You must use RHIGetBindlessSupport instead.")
-	return FDataDrivenShaderPlatformInfo::GetSupportsBindless(Platform);
-
 // Wrapper for GRHI## global variables, allows values to be overridden for mobile preview modes.
 template <typename TValueType>
 class TRHIGlobal
