Adding missing headers that should have been merged but somehow got ignored.

diff --git a/Engine/Source/Runtime/RHI/Private/RHI.cpp b/Engine/Source/Runtime/RHI/Private/RHI.cpp
--- a/Engine/Source/Runtime/RHI/Private/RHI.cpp
+++ b/Engine/Source/Runtime/RHI/Private/RHI.cpp
@@ -16,6 +16,8 @@
 #include "String/ParseTokens.h"
 #include "Misc/BufferedOutputDevice.h"
 #include "Misc/OutputDeviceFile.h"
+#include "Async/ParallelFor.h"
+#include "DataDrivenShaderPlatformInfo.h"
 
 #if RHI_ENABLE_RESOURCE_INFO
 #include "HAL/FileManager.h"
