Moving 5.2 include order of #include "DataDrivenShaderPlatformInfo.h" in RHIDefinitions.h to after the definitions of types since DataDrivenShaderPlatformInfo.h depends on them.

diff --git a/Engine/Source/Runtime/RHI/Public/RHIDefinitions.h b/Engine/Source/Runtime/RHI/Public/RHIDefinitions.h
--- a/Engine/Source/Runtime/RHI/Public/RHIDefinitions.h
+++ b/Engine/Source/Runtime/RHI/Public/RHIDefinitions.h
@@ -15,7 +15,6 @@
 #include "HAL/IConsoleManager.h"
 #include "RHIShaderPlatform.h"
 #include "RHIFeatureLevel.h"
-#include "DataDrivenShaderPlatformInfo.h"
 #endif
 
 /** Alignment of the shader parameters struct is required to be 16-byte boundaries. */
@@ -1470,3 +1469,7 @@ inline ERHIResourceType GetRHIResourceType(ETextureDimension Dimension)
 #else
 	#define GEOMETRY_SHADER(GeometryShader)	nullptr
 #endif
+
+#if UE_ENABLE_INCLUDE_ORDER_DEPRECATED_IN_5_2
+#include "DataDrivenShaderPlatformInfo.h"
+#endif
