Fixing Mac compile.

diff --git a/Engine/Source/Runtime/RHI/Private/Apple/AppleDynamicRHI.cpp b/Engine/Source/Runtime/RHI/Private/Apple/AppleDynamicRHI.cpp
--- a/Engine/Source/Runtime/RHI/Private/Apple/AppleDynamicRHI.cpp
+++ b/Engine/Source/Runtime/RHI/Private/Apple/AppleDynamicRHI.cpp
@@ -6,6 +6,7 @@
 #include "Misc/CommandLine.h"
 #include "Misc/ConfigCacheIni.h"
 #include "Misc/MessageDialog.h"
+#include "DataDrivenShaderPlatformInfo.h"
 #include COMPILED_PLATFORM_HEADER_WITH_PREFIX(Apple/Platform, PlatformDynamicRHI.h)
 
 
