Fixing define inconsistencies by adding USE_STATIC_FEATURE_LEVEL_ENUMS alongside USE_STATIC_SHADER_PLATFORM_ENUMS.

diff --git a/Engine/Source/Runtime/RHI/Public/DataDrivenShaderPlatformInfo.h b/Engine/Source/Runtime/RHI/Public/DataDrivenShaderPlatformInfo.h
--- a/Engine/Source/Runtime/RHI/Public/DataDrivenShaderPlatformInfo.h
+++ b/Engine/Source/Runtime/RHI/Public/DataDrivenShaderPlatformInfo.h
@@ -748,7 +748,7 @@ public:
 	}
 };
 
-#if USE_STATIC_SHADER_PLATFORM_ENUMS || USE_STATIC_SHADER_PLATFORM_INFO
+#if USE_STATIC_SHADER_PLATFORM_ENUMS || USE_STATIC_FEATURE_LEVEL_ENUMS || USE_STATIC_SHADER_PLATFORM_INFO
 
 #define IMPLEMENT_DDPSPI_SETTING_WITH_RETURN_TYPE(ReturnType, Function, Value) \
 	static FORCEINLINE_DEBUGGABLE const ReturnType Function(const FStaticShaderPlatform Platform) \
diff --git a/Engine/Source/Runtime/RHI/Public/RHIFeatureLevel.h b/Engine/Source/Runtime/RHI/Public/RHIFeatureLevel.h
--- a/Engine/Source/Runtime/RHI/Public/RHIFeatureLevel.h
+++ b/Engine/Source/Runtime/RHI/Public/RHIFeatureLevel.h
@@ -4,8 +4,8 @@
 
 #include "Containers/EnumAsByte.h"
 
-#ifndef USE_STATIC_SHADER_PLATFORM_ENUMS
-#define USE_STATIC_SHADER_PLATFORM_ENUMS 0
+#ifndef USE_STATIC_FEATURE_LEVEL_ENUMS
+#define USE_STATIC_FEATURE_LEVEL_ENUMS 0
 #endif
 
 /**
@@ -99,7 +99,7 @@ private:
 	ERHIFeatureLevel::Type FeatureLevel;
 };
 
-#if USE_STATIC_SHADER_PLATFORM_ENUMS
+#if USE_STATIC_FEATURE_LEVEL_ENUMS
 #include COMPILED_PLATFORM_HEADER(StaticFeatureLevel.inl)
 #else
 using FStaticFeatureLevel = FGenericStaticFeatureLevel;
diff --git a/Engine/Source/Runtime/RHI/Public/RHIValidationContext.h b/Engine/Source/Runtime/RHI/Public/RHIValidationContext.h
--- a/Engine/Source/Runtime/RHI/Public/RHIValidationContext.h
+++ b/Engine/Source/Runtime/RHI/Public/RHIValidationContext.h
@@ -78,12 +78,14 @@ public:
 
 	virtual void RHIDispatchComputeShader(uint32 ThreadGroupCountX, uint32 ThreadGroupCountY, uint32 ThreadGroupCountZ) override final
 	{
+		checkf(State.bComputePSOSet, TEXT("A Compute PSO has to be set to set resources into a shader!"));
 		RHIContext->RHIDispatchComputeShader(ThreadGroupCountX, ThreadGroupCountY, ThreadGroupCountZ);
 		Tracker->Dispatch();
 	}
 
 	virtual void RHIDispatchIndirectComputeShader(FRHIBuffer* ArgumentBuffer, uint32 ArgumentOffset) override final
 	{
+		checkf(State.bComputePSOSet, TEXT("A Compute PSO has to be set to set resources into a shader!"));
 		Tracker->Assert(ArgumentBuffer->GetWholeResourceIdentity(), ERHIAccess::IndirectArgs);
 
 		RHIContext->RHIDispatchIndirectComputeShader(ArgumentBuffer, ArgumentOffset);
@@ -364,12 +366,14 @@ public:
 
 	virtual void RHIDispatchComputeShader(uint32 ThreadGroupCountX, uint32 ThreadGroupCountY, uint32 ThreadGroupCountZ) override final
 	{
+		checkf(State.bComputePSOSet, TEXT("A Compute PSO has to be set to set resources into a shader!"));
 		RHIContext->RHIDispatchComputeShader(ThreadGroupCountX, ThreadGroupCountY, ThreadGroupCountZ);
 		Tracker->Dispatch();
 	}
 
 	virtual void RHIDispatchIndirectComputeShader(FRHIBuffer* ArgumentBuffer, uint32 ArgumentOffset) override final
 	{
+		checkf(State.bComputePSOSet, TEXT("A Compute PSO has to be set to set resources into a shader!"));
 		Tracker->Assert(ArgumentBuffer->GetWholeResourceIdentity(), ERHIAccess::IndirectArgs);
 
 		RHIContext->RHIDispatchIndirectComputeShader(ArgumentBuffer, ArgumentOffset);
