Merging 23.00-CL-23256477 to 23.10. Previously deadended.

#desc Added additional D3D12_SM6 filter for default selection; Explicitly curated list of GPUs (starting with GTX 10xx series) from automatic migration. GPUs in this list can see solid performance with SM6 though not necessarily with all UE5 features enabled and may benefit from staying with SM5.
Feature is opt-in for the project and clients can manually circumvent with forced RHI selection.

diff --git a/Engine/Source/Runtime/RHI/Private/Windows/WindowsDynamicRHI.cpp b/Engine/Source/Runtime/RHI/Private/Windows/WindowsDynamicRHI.cpp
--- a/Engine/Source/Runtime/RHI/Private/Windows/WindowsDynamicRHI.cpp
+++ b/Engine/Source/Runtime/RHI/Private/Windows/WindowsDynamicRHI.cpp
@@ -451,19 +451,60 @@ static TOptional<ERHIFeatureLevel::Type> GetPreferredFeatureLevel(EWindowsRHI Ch
 	return PreferredFeatureLevel;
 }
 
+static bool ShouldDevicePreferSM5(uint32 DeviceId)
+{
+	uint32 SM5PreferedDeviceIds[] =
+	{
+		0x1B80, // "NVIDIA GeForce GTX 1080"
+		0x1B81, // "NVIDIA GeForce GTX 1070"
+		0x1B82, // "NVIDIA GeForce GTX 1070 Ti"
+		0x1B83, // "NVIDIA GeForce GTX 1060 6GB"
+		0x1B84, // "NVIDIA GeForce GTX 1060 3GB"
+		0x1C01, // "NVIDIA GeForce GTX 1050 Ti"
+		0x1C02, // "NVIDIA GeForce GTX 1060 3GB"
+		0x1C03, // "NVIDIA GeForce GTX 1060 6GB"
+		0x1C04, // "NVIDIA GeForce GTX 1060 5GB"
+		0x1C06, // "NVIDIA GeForce GTX 1060 6GB"
+		0x1C08, // "NVIDIA GeForce GTX 1050"
+		0x1C81, // "NVIDIA GeForce GTX 1050"
+		0x1C82, // "NVIDIA GeForce GTX 1050 Ti"
+		0x1C83, // "NVIDIA GeForce GTX 1050"
+		0x1B06, // "NVIDIA GeForce GTX 1080 Ti"
+	};
+
+	for (int Index = 0; Index < UE_ARRAY_COUNT(SM5PreferedDeviceIds); ++Index)
+	{
+		if (DeviceId == SM5PreferedDeviceIds[Index])
+		{
+			return true;
+		}
+	}
+
+	return false;
+}
+
 static bool IsRHIAllowedAsDefault(EWindowsRHI InRHI, ERHIFeatureLevel::Type InFeatureLevel)
 {
+	bool bAllowed = true;
 	if (InRHI == EWindowsRHI::D3D12 && InFeatureLevel >= ERHIFeatureLevel::SM6)
 	{
+		static const FWindowsPlatformApplicationMisc::FGPUInfo BestGPUInfo = FWindowsPlatformApplicationMisc::GetBestGPUInfo();
+
 		int32 MinDedicatedMemoryMB = 0;
 		if (GConfig->GetInt(TEXT("D3D12_SM6"), TEXT("MinDedicatedMemory"), MinDedicatedMemoryMB, GEngineIni))
 		{
 			const uint64 MinDedicatedMemory = static_cast<uint64>(MinDedicatedMemoryMB) << 20;
-			static const FWindowsPlatformApplicationMisc::FGPUInfo BestGPUInfo = FWindowsPlatformApplicationMisc::GetBestGPUInfo();
-			return BestGPUInfo.DedicatedVideoMemory >= MinDedicatedMemory;
+			bAllowed &= BestGPUInfo.DedicatedVideoMemory >= MinDedicatedMemory;
+		}
+
+		bool bUseSM5PreferredGPUList = true;
+		GConfig->GetBool(TEXT("D3D12_SM6"), TEXT("bUseSM5PreferredGPUList"), bUseSM5PreferredGPUList, GEngineIni);
+		if (bUseSM5PreferredGPUList)
+		{
+			bAllowed &= !ShouldDevicePreferSM5(BestGPUInfo.DeviceId);
 		}
 	}
-	return true;
+	return bAllowed;
 }
 
 static TOptional<ERHIFeatureLevel::Type> ChooseDefaultFeatureLevel(EWindowsRHI InRHI, const FParsedWindowsDynamicRHIConfig& Config)
