RHI: Fixed incorrect Z source position when copying texture subrects instead of full texture.

It was using an source position of 1 instead of 0, which would cause an assertion later on.

diff --git a/Engine/Source/Runtime/RHI/Private/RHIGPUReadback.cpp b/Engine/Source/Runtime/RHI/Private/RHIGPUReadback.cpp
--- a/Engine/Source/Runtime/RHI/Private/RHIGPUReadback.cpp
+++ b/Engine/Source/Runtime/RHI/Private/RHIGPUReadback.cpp
@@ -168,7 +168,7 @@ void FRHIGPUTextureReadback::EnqueueCopy(FRHICommandList& RHICmdList, FRHITextur
 
 			if (Rect.IsValid())
 			{
-				CopyInfo.SourcePosition = FIntVector(Rect.X1, Rect.Y1, 1);
+				CopyInfo.SourcePosition = FIntVector(Rect.X1, Rect.Y1, 0);
 				CopyInfo.DestPosition = CopyInfo.SourcePosition;
 				CopyInfo.Size = FIntVector(Rect.X2 - Rect.X1, Rect.Y2 - Rect.Y1, 1);
 			}
