Improvements to PSO precaching:
* Whenever a task was issued, we would not return the GraphEvent associated to it, and assume the task was already done. This is fixed by returning the CompletionEvent if the object is still in ActivePrecachingTasks

* If a task was in the process of being issued, we would miss it because the lock was not kept after TryAddNewState. We would call CreateGraphicsPipelineState, and assign the CompletionEvent to FPipelineState without any lock. This is fixed by keeping the lock until we do the assignment. This assumes that bDoAsyncCompile is true, and the call to CreateGraphicsPipelineState is cheap, otherwise we will be keeping the lock for too long. If lock contention becomes an issue, I think we could create an empty GraphEvent first and then link the dependency through DontCompleteUntil, outside the lock. TBD

* Change ActivePrecachingTasks to use the same hashing-then-match functions as PrecachedPSOInitializers. I got some hash collisions during stress testing, and I think it's better to have both data structures return the same answer about the existence of an existing PSO

* Have Compute shader use a similar structure for precaching

diff --git a/Engine/Source/Runtime/RHI/Private/PipelineStateCache.cpp b/Engine/Source/Runtime/RHI/Private/PipelineStateCache.cpp
--- a/Engine/Source/Runtime/RHI/Private/PipelineStateCache.cpp
+++ b/Engine/Source/Runtime/RHI/Private/PipelineStateCache.cpp
@@ -612,6 +612,12 @@ int32 FindRayTracingMissShaderIndex(FRayTracingPipelineState* Pipeline, FRHIRayT
 	return INDEX_NONE;
 }
 
+bool IsPrecachedPSO(const FGraphicsPipelineStateInitializer& Initializer)
+{
+	check((Initializer.bFromPSOFileCache && Initializer.bPSOPrecache) == 0);
+	return Initializer.bFromPSOFileCache || Initializer.bPSOPrecache;
+}
+
 void SetGraphicsPipelineState(FRHICommandList& RHICmdList, const FGraphicsPipelineStateInitializer& Initializer, uint32 StencilRef, EApplyRendertargetOption ApplyFlags, bool bApplyAdditionalState, EPSOPrecacheResult PSOPrecacheResult)
 {
 #if PLATFORM_USE_FALLBACK_PSO
@@ -944,80 +950,65 @@ private:
 
 };
 
-class FPrecacheGraphicsPipelineCache
+template<class TPrecachePipelineCacheDerived, class TPrecachedPSOInitializer, class TPipelineState>
+class TPrecachePipelineCacheBase
 {
 public:
-
-	~FPrecacheGraphicsPipelineCache()
+	~TPrecachePipelineCacheBase()
 	{
 		// Wait for all precache tasks to finished
 		WaitTasksComplete();
 	}
-
-	bool Contains(const FGraphicsPipelineStateInitializer& Initializer)
-	{
-		FRWScopeLock ReadLock(PrecachePSOsRWLock, SLT_ReadOnly);
-		return PrecachedPSOInitializers.Contains(Initializer);
-	}
-
-	FGraphicsPipelineState* TryAddNewState(const FGraphicsPipelineStateInitializer& Initializer, bool bDoAsyncCompile)
+	
+protected:
+	void TryAddNewState(const TPrecachedPSOInitializer& Initializer, bool bDoAsyncCompile, FGraphEventRef& GraphEvent)
 	{
+		uint32 InitializerHash = TPrecachePipelineCacheDerived::PipelineStateInitializerHash(Initializer);
 		// Fast check first with read lock
-		if (Contains(Initializer))
 		{
-			return nullptr;
-		}		
-
-		FGraphicsPipelineState* NewState = nullptr;
+			FRWScopeLock ReadLock(PrecachePSOsRWLock, SLT_ReadOnly);
+			if (HasPSOBeenRequested(Initializer, InitializerHash, GraphEvent))
+			{
+				return;
+			}
+		}
 
 		// Now try and add with write lock
 		{
 			FRWScopeLock WriteLock(PrecachePSOsRWLock, SLT_Write);
-
-			// check again
-			if (PrecachedPSOInitializers.Contains(Initializer))
+			if (HasPSOBeenRequested(Initializer, InitializerHash, GraphEvent))
 			{
-				return nullptr;
+				return;
 			}
 
-			// Add to cache before starting async operation, because the async op will remove the task from active ops when done			
-			PrecachedPSOInitializers.Add(Initializer);
-		}
 
-		// Add to active set of precaching PSOs if it's created async
-		if (bDoAsyncCompile)
-		{
-			// Only need the hash of the PSO here
-			uint64 PrecachePSOHash = RHIComputePrecachePSOHash(Initializer);
+			// create new graphics state
+			TPipelineState* NewState = TPrecachePipelineCacheDerived::CreateNewPSO(Initializer);
 
-			FRWScopeLock WriteLock(ActivePrecachingTasksRWLock, SLT_Write);
+			// Add to cache before starting async operation, because the async op will remove the task from active ops when done			
+			PrecachedPSOInitializers.AddByHash(InitializerHash, Initializer);
 
-			if (ActivePrecachingTasks.Contains(PrecachePSOHash))
+			// Add to active set of precaching PSOs if it's created async
+			if (bDoAsyncCompile)
 			{
-				return nullptr;
+				FPrecacheTask PrecacheTask;
+				PrecacheTask.Initializer = Initializer;
+				PrecacheTask.PSO = NewState;
+#if DO_ENSURE
+				ensure(!ActivePrecachingTasks.ContainsByHash(InitializerHash, Initializer));
+#endif
+				ActivePrecachingTasks.AddByHash(InitializerHash, Initializer, PrecacheTask);
 			}
 
-			// create new graphics state
-			NewState = new FGraphicsPipelineState();
-
-			FPrecacheTask PrecacheTask;
-			PrecacheTask.Initializer = Initializer;
-			PrecacheTask.PSO = NewState;
-
-			ActivePrecachingTasks.Add(PrecachePSOHash, PrecacheTask);
+			TPrecachePipelineCacheDerived::OnNewPipelineStateCreated(Initializer, NewState, bDoAsyncCompile);
+			GraphEvent = NewState->CompletionEvent;
 		}
-		else
-		{
-			// create new graphics state
-			NewState = new FGraphicsPipelineState();
-		}
-
-		return NewState;
 	}
+public:
 
 	void WaitTasksComplete()
 	{
-		FRWScopeLock WriteLock(ActivePrecachingTasksRWLock, SLT_Write);
+		FRWScopeLock WriteLock(GetActivePrecachingTasksRWLock(), SLT_Write);
 
 		for (auto Iterator = ActivePrecachingTasks.CreateIterator(); Iterator; ++Iterator)
 		{
@@ -1029,48 +1020,50 @@ public:
 		ActivePrecachingTasks.Empty();
 	}
 
-	bool IsPrecaching(const FGraphicsPipelineStateInitializer& Initializer)
+	EPSOPrecacheResult GetPrecachingState(const TPrecachedPSOInitializer& Initializer)
 	{
-		uint64 PrecachePSOHash = RHIComputePrecachePSOHash(Initializer);
+		uint32 InitializerHash = TPrecachePipelineCacheDerived::PipelineStateInitializerHash(Initializer);
+		FRWScopeLock ReadLock(GetActivePrecachingTasksRWLock(), SLT_ReadOnly);
+		FPrecacheTask* FindResult = ActivePrecachingTasks.FindByHash(InitializerHash, Initializer);
+		if (FindResult && !FindResult->PSO->IsComplete())
+		{
+			return EPSOPrecacheResult::Active;
+		}
 
-		FRWScopeLock ReadLock(ActivePrecachingTasksRWLock, SLT_ReadOnly);
-		FPrecacheTask* FindResult = ActivePrecachingTasks.Find(PrecachePSOHash);
-		if (FindResult)
+		if (!PrecachedPSOInitializers.ContainsByHash(InitializerHash, Initializer))
 		{
-			// Check if not complete yet
-			return !(FindResult->PSO->IsComplete());
+			return EPSOPrecacheResult::Missed;
 		}
 
 		// not found then assume it's not precaching
-		return false;
+		return EPSOPrecacheResult::Complete;
 	}
 
 	bool IsPrecaching()
 	{
-		FRWScopeLock ReadLock(ActivePrecachingTasksRWLock, SLT_ReadOnly);
+		FRWScopeLock ReadLock(GetActivePrecachingTasksRWLock(), SLT_ReadOnly);
 		return !ActivePrecachingTasks.IsEmpty();
 	}
 
-	void PrecacheFinished(const FGraphicsPipelineStateInitializer& Initializer)
+	void PrecacheFinished(const TPrecachedPSOInitializer& Initializer)
 	{
-		uint64 PrecachePSOHash = RHIComputePrecachePSOHash(Initializer);
-
-		FRWScopeLock WriteLock(ActivePrecachingTasksRWLock, SLT_Write);
-		FPrecacheTask* FindResult = ActivePrecachingTasks.Find(PrecachePSOHash);
+		uint32 InitializerHash = TPrecachePipelineCacheDerived::PipelineStateInitializerHash(Initializer);
+		FRWScopeLock WriteLock(GetActivePrecachingTasksRWLock(), SLT_Write);
+		FPrecacheTask* FindResult = ActivePrecachingTasks.FindByHash(InitializerHash, Initializer);
 		// Might have already been processed because the task got rescheduled
 		if (FindResult != nullptr)
 		{
-			verify(ActivePrecachingTasks.Remove(PrecachePSOHash) == 1);
+			verify(ActivePrecachingTasks.RemoveByHash(InitializerHash, Initializer) == 1);
 			PrecachedPSOs.Add(FindResult->PSO);
 		}
 	}
 
 	void ProcessDelayedCleanup()
 	{
-		FRWScopeLock WriteLock(ActivePrecachingTasksRWLock, SLT_Write);
+		FRWScopeLock WriteLock(GetActivePrecachingTasksRWLock(), SLT_Write);
 		for (int32 Index = 0; Index < PrecachedPSOs.Num(); ++Index)
 		{
-			FGraphicsPipelineState* GraphicsPSO = PrecachedPSOs[Index];
+			TPipelineState* GraphicsPSO = PrecachedPSOs[Index];
 			if (GraphicsPSO->IsComplete())
 			{
 				// This is needed to cleanup the members - bit strange because it's complete already
@@ -1079,40 +1072,125 @@ public:
 				delete GraphicsPSO;
 				PrecachedPSOs.RemoveAtSwap(Index);
 				Index--;
-			}			
+			}
 		}
-
 	}
 
 private:
-
-	struct FGraphicsPSOInitializerKeyFuncs : DefaultKeyFuncs<FGraphicsPipelineStateInitializer>
+	bool HasPSOBeenRequested(const TPrecachedPSOInitializer& Initializer, uint32 InitializerHash, FGraphEventRef& GraphEvent)
 	{
-		static FORCEINLINE bool Matches(KeyInitType A, KeyInitType B)
+		FPrecacheTask* FindResult = ActivePrecachingTasks.FindByHash(InitializerHash, Initializer);
+		if (FindResult)
 		{
-			return RHIMatchPrecachePSOInitializers(A, B);
+			GraphEvent = FindResult->PSO->CompletionEvent;
+			check(GraphEvent != nullptr);
+			return true;
 		}
 
-		static FORCEINLINE uint32 GetKeyHash(KeyInitType Key)
+		if (PrecachedPSOInitializers.ContainsByHash(InitializerHash, Initializer))
 		{
-			return RHIComputePrecachePSOHash(Key);
+			return true;
 		}
-	};
+
+		return false;
+	}
+
+	FRWLock& GetActivePrecachingTasksRWLock()
+	{
+		return PrecachePSOsRWLock;
+	}
 
 	// Cache of all precached PSOs - need correct compare to make sure we precache all of them
 	FRWLock PrecachePSOsRWLock;
-	TSet<FGraphicsPipelineStateInitializer, FGraphicsPSOInitializerKeyFuncs> PrecachedPSOInitializers;
+
+	using TSetKeyFuncsBaseClass = DefaultKeyFuncs<TPrecachedPSOInitializer>;
+	struct TSetKeyFuncs : public TSetKeyFuncsBaseClass
+	{
+		static FORCEINLINE bool Matches(typename TSetKeyFuncsBaseClass::KeyInitType A, typename TSetKeyFuncsBaseClass::KeyInitType B)
+		{
+			return TPrecachePipelineCacheDerived::PipelineStateInitializerMatch(A, B);
+		}
+		static FORCEINLINE uint32 GetKeyHash(typename TSetKeyFuncsBaseClass::KeyInitType Key)
+		{
+			return TPrecachePipelineCacheDerived::PipelineStateInitializerHash(Key);
+		}
+	};
+	TSet<TPrecachedPSOInitializer, TSetKeyFuncs> PrecachedPSOInitializers;
 
 	// Current active background precache operations - full PSO hash can be used for fast retrieval
 	// (on hash clash it might report that a PSO is still compiling)
-	FRWLock ActivePrecachingTasksRWLock;
 	struct FPrecacheTask
 	{
-		FGraphicsPipelineStateInitializer Initializer;
-		FGraphicsPipelineState* PSO = nullptr;
+		TPrecachedPSOInitializer Initializer;
+		TPipelineState* PSO = nullptr;
+	};
+
+	using TMapKeyFuncsBaseClass = TDefaultMapHashableKeyFuncs<TPrecachedPSOInitializer, FPrecacheTask, false>;
+	struct TMapKeyFuncs : public TMapKeyFuncsBaseClass
+	{
+		static FORCEINLINE bool Matches(typename TMapKeyFuncsBaseClass::KeyInitType A, typename TMapKeyFuncsBaseClass::KeyInitType B)
+		{
+			return TPrecachePipelineCacheDerived::PipelineStateInitializerMatch(A, B);
+		}
+		static FORCEINLINE uint32 GetKeyHash(typename TMapKeyFuncsBaseClass::KeyInitType Key)
+		{
+			return TPrecachePipelineCacheDerived::PipelineStateInitializerHash(Key);
+		}
 	};
-	TMap<uint64, FPrecacheTask> ActivePrecachingTasks;
-	TArray<FGraphicsPipelineState*> PrecachedPSOs;
+
+	TMap<TPrecachedPSOInitializer, FPrecacheTask, FDefaultSetAllocator, TMapKeyFuncs> ActivePrecachingTasks;
+	TArray<TPipelineState*> PrecachedPSOs;
+};
+
+class FPrecacheComputePipelineCache : public TPrecachePipelineCacheBase<FPrecacheComputePipelineCache, FRHIComputeShader*, FComputePipelineState>
+{
+public:
+	static FComputePipelineState* CreateNewPSO(const FRHIComputeShader* ComputeShader)
+	{
+		return new FComputePipelineState((FRHIComputeShader*)ComputeShader);
+	}
+
+	FGraphEventRef PrecacheComputePipelineState(FRHIComputeShader* ComputeShader, bool bForcePrecache);
+	static void OnNewPipelineStateCreated(FRHIComputeShader* ComputeShader, FComputePipelineState* NewComputePipelineState, bool bDoAsyncCompile);
+
+	static FORCEINLINE bool PipelineStateInitializerMatch(const FRHIComputeShader* A, const FRHIComputeShader* B)
+	{
+		// todo: would be good/more robust to have the CS equivalent of RHIMatchPrecachePSOInitializers instead of relying on pointers
+		// (for example if an FRHIComputeShader gets released and a new one is allocated with the same pointer value)
+		return A == B;
+	}
+
+	static FORCEINLINE uint32 PipelineStateInitializerHash(const FRHIComputeShader* Key)
+	{
+		// will hash the pointer through PointerHash(const void* Key)
+		return GetTypeHash(Key);
+	}
+
+};
+
+class FPrecacheGraphicsPipelineCache : public TPrecachePipelineCacheBase<FPrecacheGraphicsPipelineCache, FGraphicsPipelineStateInitializer, FGraphicsPipelineState>
+{
+public:
+	static FGraphicsPipelineState* CreateNewPSO(const FGraphicsPipelineStateInitializer& Initializer)
+		{
+		return new FGraphicsPipelineState;
+	}
+
+	static FORCEINLINE bool PipelineStateInitializerMatch(const FGraphicsPipelineStateInitializer& A, const FGraphicsPipelineStateInitializer& B)
+	{
+		return RHIMatchPrecachePSOInitializers(A, B);
+	}
+
+	static FORCEINLINE uint32 PipelineStateInitializerHash(const FGraphicsPipelineStateInitializer& Key)
+	{
+		return RHIComputePrecachePSOHash(Key);
+	}
+
+
+	FGraphEventRef PrecacheGraphicsPipelineState(const FGraphicsPipelineStateInitializer& Initializer);
+
+	static void OnNewPipelineStateCreated(const FGraphicsPipelineStateInitializer& Initializer, FGraphicsPipelineState* NewGraphicsPipelineState, bool bDoAsyncCompile);
+
 };
 
 // Typed caches for compute and graphics
@@ -1123,6 +1201,7 @@ typedef TSharedPipelineStateCache<FGraphicsPipelineStateInitializer, FGraphicsPi
 FComputePipelineCache GComputePipelineCache;
 FGraphicsPipelineCache GGraphicsPipelineCache;
 FPrecacheGraphicsPipelineCache GPrecacheGraphicsPipelineCache;
+FPrecacheComputePipelineCache GPrecacheComputePipelineCache;
 
 FAutoConsoleTaskPriority CPrio_FCompilePipelineStateTask(
 	TEXT("TaskGraph.TaskPriorities.CompilePipelineStateTask"),
@@ -1353,7 +1432,7 @@ public:
 		, Initializer(InInitializer)
 		, PSOPreCacheResult(InPSOPreCacheResult)
 	{
-		if(Initializer.bFromPSOFileCache)
+		if(IsPrecachedPSO(Initializer))
 		{
 			GPipelinePrecompileTasksInFlight++;
 		}
@@ -1414,7 +1493,7 @@ public:
 
 	~FCompilePipelineStateTask()
 	{
-		if (Initializer.bFromPSOFileCache)
+		if (IsPrecachedPSO(Initializer))
 		{
 			GPipelinePrecompileTasksInFlight--;
 		}
@@ -1436,6 +1515,10 @@ public:
 		{
 			FComputePipelineState* ComputePipeline = static_cast<FComputePipelineState*>(Pipeline);
 			ComputePipeline->RHIPipeline = RHICreateComputePipelineState(ComputePipeline->ComputeShader);
+			if (Initializer.bPSOPrecache)
+			{
+				GPrecacheComputePipelineCache.PrecacheFinished(ComputePipeline->ComputeShader);
+			}
 		}
 		else
 		{
@@ -1578,6 +1661,7 @@ void PipelineStateCache::FlushResources()
 	GGraphicsPipelineCache.ProcessDelayedCleanup();
 
 	GPrecacheGraphicsPipelineCache.ProcessDelayedCleanup();
+	GPrecacheComputePipelineCache.ProcessDelayedCleanup();
 
 	{
 		int32 NumMissesThisFrame = GraphicsPipelineCacheMisses.Load(EMemoryOrder::Relaxed);
@@ -1705,7 +1789,9 @@ FComputePipelineState* PipelineStateCache::GetAndOrCreateComputePipelineState(FR
 		// create a compilation task, or just do it now...
 		if (DoAsyncCompile)
 		{
-			OutCachedState->CompletionEvent = TGraphTask<FCompilePipelineStateTask>::CreateTask().ConstructAndDispatchWhenReady(OutCachedState, FGraphicsPipelineStateInitializer(), EPSOPrecacheResult::Untracked);
+			FGraphicsPipelineStateInitializer GraphicsPipelineStateInitializer;
+			GraphicsPipelineStateInitializer.bFromPSOFileCache = bFromFileCache;
+			OutCachedState->CompletionEvent = TGraphTask<FCompilePipelineStateTask>::CreateTask().ConstructAndDispatchWhenReady(OutCachedState, GraphicsPipelineStateInitializer, EPSOPrecacheResult::Untracked);
 			if (!bFromFileCache)
 			{
 				RHICmdList.AddDispatchPrerequisite(OutCachedState->CompletionEvent);
@@ -2212,6 +2298,7 @@ FGraphicsPipelineState* PipelineStateCache::GetAndOrCreateGraphicsPipelineState(
 	FGraphicsPipelineState* OutCachedState = nullptr;
 
 	bool bWasFound = GGraphicsPipelineCache.Find(Initializer, OutCachedState);
+    ensure(!Initializer.bPSOPrecache);
 	bool DoAsyncCompile = IsAsyncCompilationAllowed(RHICmdList, Initializer.bFromPSOFileCache);
 
 	if (bWasFound == false)
@@ -2285,10 +2372,10 @@ bool PipelineStateCache::IsPSOPrecachingEnabled()
 #endif // WITH_EDITOR
 }
 
-FGraphEventRef PipelineStateCache::PrecacheComputePipelineState(FRHIComputeShader* ComputeShader, bool bForcePrecache)
+FGraphEventRef FPrecacheComputePipelineCache::PrecacheComputePipelineState(FRHIComputeShader* ComputeShader, bool bForcePrecache)
 {
 	FGraphEventRef GraphEvent = nullptr;
-	if (!IsPSOPrecachingEnabled() && !bForcePrecache)
+	if (!PipelineStateCache::IsPSOPrecachingEnabled() && !bForcePrecache)
 	{
 		return GraphEvent;
 	}
@@ -2298,31 +2385,36 @@ FGraphEventRef PipelineStateCache::PrecacheComputePipelineState(FRHIComputeShade
 		return GraphEvent;
 	}
 
-	FComputePipelineState* CachedState = nullptr;
-	if (GComputePipelineCache.Find(ComputeShader, CachedState))
+	static bool bDoAsyncCompile = FApp::ShouldUseThreadingForPerformance();
+	GPrecacheComputePipelineCache.TryAddNewState(ComputeShader, bDoAsyncCompile, GraphEvent);
+
+	if (GraphEvent == nullptr)
 	{
-		return GraphEvent;
+		check(GPrecacheComputePipelineCache.GetPrecachingState(ComputeShader) == EPSOPrecacheResult::Complete)
 	}
 
-	// create new graphics state
-	CachedState = new FComputePipelineState(ComputeShader);
+	return GraphEvent;
+}
 
+void FPrecacheComputePipelineCache::OnNewPipelineStateCreated(FRHIComputeShader* ComputeShader, FComputePipelineState* CachedState, bool bDoAsyncCompile)
+{
 	// create a compilation task, or just do it now...
-	static bool bDoAsyncCompile = FApp::ShouldUseThreadingForPerformance();
 	if (bDoAsyncCompile)
 	{
+		FGraphicsPipelineStateInitializer GraphicsPipelineStateInitializer;
+		GraphicsPipelineStateInitializer.bPSOPrecache = true;
 		// Thread pool disabled?
 		if (!FPSOPrecacheThreadPool::UsePool())
 		{
-			CachedState->CompletionEvent = TGraphTask<FCompilePipelineStateTask>::CreateTask().ConstructAndDispatchWhenReady(CachedState, FGraphicsPipelineStateInitializer(), EPSOPrecacheResult::Active);
-			GraphEvent = CachedState->CompletionEvent;
+			CachedState->CompletionEvent = TGraphTask<FCompilePipelineStateTask>::CreateTask().ConstructAndDispatchWhenReady(CachedState, GraphicsPipelineStateInitializer, EPSOPrecacheResult::Active);
 		}
 		else
 		{
 			// Here, PSO precompiles use a separate thread pool.
 			// Note that we do not add precompile tasks as cmdlist prerequisites.
-			GraphEvent = FGraphEvent::CreateGraphEvent();
-			TUniquePtr<FCompilePipelineStateTask> ThreadPoolTask = MakeUnique<FCompilePipelineStateTask>(CachedState, FGraphicsPipelineStateInitializer(), EPSOPrecacheResult::Active);
+			FGraphEventRef GraphEvent = FGraphEvent::CreateGraphEvent();
+			CachedState->CompletionEvent = GraphEvent;
+			TUniquePtr<FCompilePipelineStateTask> ThreadPoolTask = MakeUnique<FCompilePipelineStateTask>(CachedState, GraphicsPipelineStateInitializer, EPSOPrecacheResult::Active);
 			CachedState->CompletionEvent = GraphEvent;
 			uint64 StartTime = FPlatformTime::Cycles64();
 #if	PSO_TRACK_CACHE_STATS
@@ -2349,14 +2441,16 @@ FGraphEventRef PipelineStateCache::PrecacheComputePipelineState(FRHIComputeShade
 	{
 		CachedState->RHIPipeline = RHICreateComputePipelineState(CachedState->ComputeShader);
 	}
+}
 
-	GComputePipelineCache.Add(ComputeShader, CachedState);
-	return GraphEvent;
+FGraphEventRef PipelineStateCache::PrecacheComputePipelineState(FRHIComputeShader* ComputeShader, bool bForcePrecache)
+{
+	return GPrecacheComputePipelineCache.PrecacheComputePipelineState(ComputeShader, bForcePrecache);
 }
 
-FGraphEventRef PipelineStateCache::PrecacheGraphicsPipelineState(const FGraphicsPipelineStateInitializer& Initializer)
+FGraphEventRef FPrecacheGraphicsPipelineCache::PrecacheGraphicsPipelineState(const FGraphicsPipelineStateInitializer& Initializer)
 {
-	if (!IsPSOPrecachingEnabled())
+	if (!PipelineStateCache::IsPSOPrecachingEnabled())
 	{
 		return nullptr;
 	}
@@ -2367,12 +2461,17 @@ FGraphEventRef PipelineStateCache::PrecacheGraphicsPipelineState(const FGraphics
 	static bool bDoAsyncCompile = FApp::ShouldUseThreadingForPerformance();
 
 	// try and create new graphics state
-	FGraphicsPipelineState* NewGraphicsPipelineState = GPrecacheGraphicsPipelineCache.TryAddNewState(Initializer, bDoAsyncCompile);
-	if (NewGraphicsPipelineState == nullptr)
+	FGraphEventRef PrecachingEvent;
+	GPrecacheGraphicsPipelineCache.TryAddNewState(Initializer, bDoAsyncCompile, PrecachingEvent);
+	if (PrecachingEvent == nullptr)
 	{
-		return nullptr;
+		check(GPrecacheGraphicsPipelineCache.GetPrecachingState(Initializer) == EPSOPrecacheResult::Complete);
 	}
+	return PrecachingEvent;
+}
 	
+void FPrecacheGraphicsPipelineCache::OnNewPipelineStateCreated(const FGraphicsPipelineStateInitializer & Initializer, FGraphicsPipelineState * NewGraphicsPipelineState, bool bDoAsyncCompile)
+{
 	ValidateGraphicsPipelineStateInitializer(Initializer);
 
 	// Mark as precache so it will try and use the background thread pool if available
@@ -2382,7 +2481,14 @@ FGraphEventRef PipelineStateCache::PrecacheGraphicsPipelineState(const FGraphics
 	InitializerCopy.bPSOPrecache = true;
 
 	// Start the precache task
-	return CreateGraphicsPipelineState(InitializerCopy, EPSOPrecacheResult::Active, bDoAsyncCompile, bPSOPrecache, NewGraphicsPipelineState);
+	FGraphEventRef GraphEventRef = CreateGraphicsPipelineState(InitializerCopy, EPSOPrecacheResult::Active, bDoAsyncCompile, bPSOPrecache, NewGraphicsPipelineState);
+
+	check(GraphEventRef == NewGraphicsPipelineState->CompletionEvent);
+}
+
+FGraphEventRef PipelineStateCache::PrecacheGraphicsPipelineState(const FGraphicsPipelineStateInitializer& Initializer)
+{
+	return GPrecacheGraphicsPipelineCache.PrecacheGraphicsPipelineState(Initializer);
 }
 
 EPSOPrecacheResult PipelineStateCache::CheckPipelineStateInCache(const FGraphicsPipelineStateInitializer& PipelineStateInitializer)
@@ -2392,12 +2498,7 @@ EPSOPrecacheResult PipelineStateCache::CheckPipelineStateInCache(const FGraphics
 		return EPSOPrecacheResult::Unknown;
 	}
 
-	if (!GPrecacheGraphicsPipelineCache.Contains(PipelineStateInitializer))
-	{
-		return EPSOPrecacheResult::Missed;
-	}
-
-	return GPrecacheGraphicsPipelineCache.IsPrecaching(PipelineStateInitializer) ? EPSOPrecacheResult::Active : EPSOPrecacheResult::Complete;
+	return GPrecacheGraphicsPipelineCache.GetPrecachingState(PipelineStateInitializer);
 }
 
 EPSOPrecacheResult PipelineStateCache::CheckPipelineStateInCache(FRHIComputeShader* ComputeShader)
@@ -2407,28 +2508,27 @@ EPSOPrecacheResult PipelineStateCache::CheckPipelineStateInCache(FRHIComputeShad
 		return EPSOPrecacheResult::Unknown;
 	}
 
-	FComputePipelineState* CachedState = nullptr;
-	if (!GComputePipelineCache.Find(ComputeShader, CachedState))
-	{
-		return EPSOPrecacheResult::Missed;
-	}
+	return GPrecacheComputePipelineCache.GetPrecachingState(ComputeShader);
+}
 
-	if (!CachedState->IsComplete())
+bool PipelineStateCache::IsPrecaching(const FGraphicsPipelineStateInitializer& PipelineStateInitializer)
+{
+	if (!IsPSOPrecachingEnabled())
 	{
-		return EPSOPrecacheResult::Active;
+		return false;
 	}
 
-	return CachedState->RHIPipeline && CachedState->RHIPipeline->IsValid() ? EPSOPrecacheResult::Complete : EPSOPrecacheResult::NotSupported;
+	return GPrecacheGraphicsPipelineCache.GetPrecachingState(PipelineStateInitializer) == EPSOPrecacheResult::Active;
 }
 
-bool PipelineStateCache::IsPrecaching(const FGraphicsPipelineStateInitializer& PipelineStateInitializer)
+bool PipelineStateCache::IsPrecaching(FRHIComputeShader* ComputeShader)
 {
 	if (!IsPSOPrecachingEnabled())
 	{
 		return false;
 	}
 
-	return GPrecacheGraphicsPipelineCache.IsPrecaching(PipelineStateInitializer);
+	return GPrecacheComputePipelineCache.GetPrecachingState(ComputeShader) == EPSOPrecacheResult::Active;
 }
 
 bool PipelineStateCache::IsPrecaching()
@@ -2438,7 +2538,7 @@ bool PipelineStateCache::IsPrecaching()
 		return false;
 	}
 
-	return GPrecacheGraphicsPipelineCache.IsPrecaching();
+	return GPrecacheGraphicsPipelineCache.IsPrecaching() || GPrecacheComputePipelineCache.IsPrecaching();
 }
 
 FRHIGraphicsPipelineState* ExecuteSetGraphicsPipelineState(FGraphicsPipelineState* GraphicsPipelineState)
diff --git a/Engine/Source/Runtime/RHI/Public/PipelineStateCache.h b/Engine/Source/Runtime/RHI/Public/PipelineStateCache.h
--- a/Engine/Source/Runtime/RHI/Public/PipelineStateCache.h
+++ b/Engine/Source/Runtime/RHI/Public/PipelineStateCache.h
@@ -112,6 +112,9 @@ namespace PipelineStateCache
 	/* Is the given PSO initializer still precaching? */
 	extern RHI_API bool						IsPrecaching(const FGraphicsPipelineStateInitializer& PipelineStateInitializer);
 
+	/* Is the given PSO initializer still precaching? */
+	extern RHI_API bool						IsPrecaching(FRHIComputeShader* ComputeShader);
+
 	/* Any async precaching operations still busy */
 	extern RHI_API bool						IsPrecaching();
 
