With batched shader parameters implemented, we can now batch bindless parameter setting. This avoids doing direct ShaderParameter setting using FRHIDescriptorHandle outside of the RHI.
This also makes sure bindless resources all go through the validation layer.

diff --git a/Engine/Source/Runtime/RHI/Public/RHICommandList.h b/Engine/Source/Runtime/RHI/Public/RHICommandList.h
--- a/Engine/Source/Runtime/RHI/Public/RHICommandList.h
+++ b/Engine/Source/Runtime/RHI/Public/RHICommandList.h
@@ -473,22 +473,41 @@ public:
 		BatchedParameters.ParametersData.Append((const uint8*)NewValue, NumBytes);
 		BatchedParameters.Parameters.Emplace((uint16)BufferIndex, (uint16)BaseIndex, (uint16)DestDataOffset, (uint16)NumBytes);
 	}
+
 	FORCEINLINE_DEBUGGABLE void SetShaderTexture(FRHIShader* Shader, uint32 TextureIndex, FRHITexture* Texture)
 	{
-		GetBatchedShaderParameters(Shader).ResourceParameters.Emplace(FRHIShaderParameterResource::EType::Texture, Texture, (uint16)TextureIndex);
+		GetBatchedShaderParameters(Shader).ResourceParameters.Emplace(Texture, (uint16)TextureIndex);
 	}
 	FORCEINLINE_DEBUGGABLE void SetShaderResourceViewParameter(FRHIShader* Shader, uint32 SamplerIndex, FRHIShaderResourceView* SRV)
 	{
-		GetBatchedShaderParameters(Shader).ResourceParameters.Emplace(FRHIShaderParameterResource::EType::ResourceView, SRV, (uint16)SamplerIndex);
+		GetBatchedShaderParameters(Shader).ResourceParameters.Emplace(SRV, (uint16)SamplerIndex);
 	}
 	FORCEINLINE_DEBUGGABLE void SetShaderSampler(FRHIShader* Shader, uint32 SamplerIndex, FRHISamplerState* State)
 	{
-		GetBatchedShaderParameters(Shader).ResourceParameters.Emplace(FRHIShaderParameterResource::EType::Sampler, State, (uint16)SamplerIndex);
+		GetBatchedShaderParameters(Shader).ResourceParameters.Emplace(State, (uint16)SamplerIndex);
 	}
 	FORCEINLINE_DEBUGGABLE void SetUAVParameter(FRHIShader* Shader, uint32 UAVIndex, FRHIUnorderedAccessView* UAV)
 	{
-		GetBatchedShaderParameters(Shader).ResourceParameters.Emplace(FRHIShaderParameterResource::EType::UnorderedAccessView, UAV, (uint16)UAVIndex);
+		GetBatchedShaderParameters(Shader).ResourceParameters.Emplace(UAV, (uint16)UAVIndex);
+	}
+
+	FORCEINLINE_DEBUGGABLE void SetBindlessTexture(FRHIShader* Shader, uint32 Index, FRHITexture* Texture)
+	{
+		GetBatchedShaderParameters(Shader).BindlessParameters.Emplace(Texture, (uint16)Index);
+	}
+	FORCEINLINE_DEBUGGABLE void SetBindlessResourceView(FRHIShader* Shader, uint32 Index, FRHIShaderResourceView* SRV)
+	{
+		GetBatchedShaderParameters(Shader).BindlessParameters.Emplace(SRV, (uint16)Index);
+	}
+	FORCEINLINE_DEBUGGABLE void SetBindlessSampler(FRHIShader* Shader, uint32 Index, FRHISamplerState* State)
+	{
+		GetBatchedShaderParameters(Shader).BindlessParameters.Emplace(State, (uint16)Index);
+	}
+	FORCEINLINE_DEBUGGABLE void SetBindlessUAV(FRHIShader* Shader, uint32 Index, FRHIUnorderedAccessView* UAV)
+	{
+		GetBatchedShaderParameters(Shader).BindlessParameters.Emplace(UAV, (uint16)Index);
 	}
+
 	FORCEINLINE_DEBUGGABLE void SetUAVParameter(FRHIShader* Shader, uint32 UAVIndex, FRHIUnorderedAccessView* UAV, uint32 InitialCount)
 	{
 		checkNoEntry(); // @todo: support append/consume buffers
@@ -2629,6 +2648,70 @@ public:
 		ALLOC_COMMAND(FRHICommandSetUAVParameter_InitialCount)(Shader, UAVIndex, UAV, InitialCount);
 	}
 
+	FORCEINLINE_DEBUGGABLE void SetBindlessTexture(FRHIComputeShader* Shader, uint32 Index, FRHITexture* Texture)
+	{
+		ValidateBoundShader(Shader);
+
+		if (ParameterBatcher.IsEnabled())
+		{
+			ParameterBatcher.SetBindlessTexture(Shader, Index, Texture);
+			return;
+		}
+
+		FRHIShaderParameterResource Resource(Texture, (uint16)Index);
+		SetShaderParameters(Shader, {}, {}, {}, MakeArrayView(&Resource, 1));
+	}
+
+	FORCEINLINE_DEBUGGABLE void SetBindlessResourceView(FRHIComputeShader* Shader, uint32 Index, FRHIShaderResourceView* SRV)
+	{
+		ValidateBoundShader(Shader);
+
+		if (ParameterBatcher.IsEnabled())
+		{
+			ParameterBatcher.SetBindlessResourceView(Shader, Index, SRV);
+			return;
+		}
+
+		FRHIShaderParameterResource Resource(SRV, (uint16)Index);
+		SetShaderParameters(Shader, {}, {}, {}, MakeArrayView(&Resource, 1));
+	}
+
+	FORCEINLINE_DEBUGGABLE void SetBindlessSampler(FRHIComputeShader* Shader, uint32 Index, FRHISamplerState* State)
+	{
+		ValidateBoundShader(Shader);
+
+		// Immutable samplers can't be set dynamically
+		check(!State->IsImmutable());
+		if (State->IsImmutable())
+		{
+			return;
+		}
+
+		// Immutable samplers can't be set dynamically
+		if (ParameterBatcher.IsEnabled())
+		{
+			ParameterBatcher.SetBindlessSampler(Shader, Index, State);
+			return;
+		}
+
+		FRHIShaderParameterResource Resource(State, (uint16)Index);
+		SetShaderParameters(Shader, {}, {}, {}, MakeArrayView(&Resource, 1));
+	}
+
+	FORCEINLINE_DEBUGGABLE void SetBindlessUAV(FRHIComputeShader* Shader, uint32 Index, FRHIUnorderedAccessView* UAV)
+	{
+		ValidateBoundShader(Shader);
+
+		if (ParameterBatcher.IsEnabled())
+		{
+			ParameterBatcher.SetBindlessUAV(Shader, Index, UAV);
+			return;
+		}
+
+		FRHIShaderParameterResource Resource(UAV, (uint16)Index);
+		SetShaderParameters(Shader, {}, {}, {}, MakeArrayView(&Resource, 1));
+	}
+
 	UE_DEPRECATED(5.1, "ComputePipelineStates should be used instead of direct ComputeShaders. You can use SetComputePipelineState(RHICmdList, ComputeShader).")
 	FORCEINLINE_DEBUGGABLE void SetComputeShader(FRHIComputeShader* ComputeShader)
 	{
@@ -3459,6 +3542,77 @@ public:
 		SetUAVParameter(Shader.GetReference(), UAVIndex, UAV);
 	}
 
+	using FRHIComputeCommandList::SetBindlessTexture;
+
+	FORCEINLINE_DEBUGGABLE void SetBindlessTexture(FRHIGraphicsShader* Shader, uint32 Index, FRHITexture* Texture)
+	{
+		ValidateBoundShader(Shader);
+
+		if (ParameterBatcher.IsEnabled())
+		{
+			ParameterBatcher.SetBindlessTexture(Shader, Index, Texture);
+			return;
+		}
+
+		FRHIShaderParameterResource Resource(Texture, (uint16)Index);
+		SetShaderParameters(Shader, {}, {}, {}, MakeArrayView(&Resource, 1));
+	}
+
+	using FRHIComputeCommandList::SetBindlessResourceView;
+
+	FORCEINLINE_DEBUGGABLE void SetBindlessResourceView(FRHIGraphicsShader* Shader, uint32 Index, FRHIShaderResourceView* SRV)
+	{
+		ValidateBoundShader(Shader);
+
+		if (ParameterBatcher.IsEnabled())
+		{
+			ParameterBatcher.SetBindlessResourceView(Shader, Index, SRV);
+			return;
+		}
+
+		FRHIShaderParameterResource Resource(SRV, (uint16)Index);
+		SetShaderParameters(Shader, {}, {}, {}, MakeArrayView(&Resource, 1));
+	}
+
+	using FRHIComputeCommandList::SetBindlessSampler;
+
+	FORCEINLINE_DEBUGGABLE void SetBindlessSampler(FRHIGraphicsShader* Shader, uint32 Index, FRHISamplerState* State)
+	{
+		ValidateBoundShader(Shader);
+
+		// Immutable samplers can't be set dynamically
+		check(!State->IsImmutable());
+		if (State->IsImmutable())
+		{
+			return;
+		}
+
+		if (ParameterBatcher.IsEnabled())
+		{
+			ParameterBatcher.SetBindlessSampler(Shader, Index, State);
+			return;
+		}
+
+		FRHIShaderParameterResource Resource(State, (uint16)Index);
+		SetShaderParameters(Shader, {}, {}, {}, MakeArrayView(&Resource, 1));
+	}
+
+	using FRHIComputeCommandList::SetBindlessUAV;
+
+	FORCEINLINE_DEBUGGABLE void SetBindlessUAV(FRHIPixelShader* Shader, uint32 Index, FRHIUnorderedAccessView* UAV)
+	{
+		ValidateBoundShader(Shader);
+
+		if (ParameterBatcher.IsEnabled())
+		{
+			ParameterBatcher.SetBindlessUAV(Shader, Index, UAV);
+			return;
+		}
+
+		FRHIShaderParameterResource Resource(UAV, (uint16)Index);
+		SetShaderParameters(Shader, {}, {}, {}, MakeArrayView(&Resource, 1));
+	}
+
 	FORCEINLINE_DEBUGGABLE void SetBlendFactor(const FLinearColor& BlendFactor = FLinearColor::White)
 	{
 		//check(IsOutsideRenderPass());
diff --git a/Engine/Source/Runtime/RenderCore/Public/ShaderParameterUtils.h b/Engine/Source/Runtime/RenderCore/Public/ShaderParameterUtils.h
--- a/Engine/Source/Runtime/RenderCore/Public/ShaderParameterUtils.h
+++ b/Engine/Source/Runtime/RenderCore/Public/ShaderParameterUtils.h
@@ -152,36 +152,24 @@ GUARD_SETSHADERVALUE(FQuat4)
 GUARD_SETSHADERVALUE(FSphere3)
 GUARD_SETSHADERVALUE(FBox3)
 
-template<typename TRHIShader, typename TRHICmdList>
-FORCEINLINE void SetBindlessParameter(TRHICmdList& RHICmdList, TRHIShader* Shader, const FShaderResourceParameter& Parameter, FRHIDescriptorHandle Handle)
-{
-	if (Handle.IsValid())
-	{
-		const uint32 BindlessIndex = Handle.GetIndex();
-		RHICmdList.SetShaderParameter(Shader, 0, Parameter.GetBaseIndex(), sizeof(BindlessIndex), &BindlessIndex);
-	}
-}
-
-
 /**
  * Sets the value of a shader surface parameter (e.g. to access MSAA samples).
  * Template'd on shader type (e.g. pixel shader or compute shader).
  */
 template<typename TRHIShader, typename TRHICmdList>
-FORCEINLINE void SetTextureParameter(TRHICmdList& RHICmdList, TRHIShader* Shader, const FShaderResourceParameter& Parameter, FRHITexture* NewTextureRHI, uint32 ElementIndex = 0)
+FORCEINLINE void SetTextureParameter(TRHICmdList& RHICmdList, TRHIShader* Shader, const FShaderResourceParameter& Parameter, FRHITexture* TextureRHI, uint32 ElementIndex = 0)
 {
 	if (Parameter.IsBound() && ElementIndex < Parameter.GetNumResources())
 	{
 #if PLATFORM_SUPPORTS_BINDLESS_RENDERING
 		if (Parameter.GetType() == EShaderParameterType::BindlessResourceIndex)
 		{
-			const FRHIDescriptorHandle Handle = NewTextureRHI ? NewTextureRHI->GetDefaultBindlessHandle() : FRHIDescriptorHandle();
-			SetBindlessParameter(RHICmdList, Shader, Parameter, Handle);
+			RHICmdList.SetBindlessTexture(Shader, Parameter.GetBaseIndex(), TextureRHI);
 		}
 		else
 #endif
 		{
-			RHICmdList.SetShaderTexture(Shader, Parameter.GetBaseIndex() + ElementIndex, NewTextureRHI);
+			RHICmdList.SetShaderTexture(Shader, Parameter.GetBaseIndex() + ElementIndex, TextureRHI);
 		}
 	}
 }
@@ -197,8 +185,7 @@ FORCEINLINE void SetSamplerParameter(TRHICmdList& RHICmdList, TRHIShader* Shader
 #if PLATFORM_SUPPORTS_BINDLESS_RENDERING
 		if (Parameter.GetType() == EShaderParameterType::BindlessSamplerIndex)
 		{
-			const FRHIDescriptorHandle Handle = SamplerStateRHI ? SamplerStateRHI->GetBindlessHandle() : FRHIDescriptorHandle();
-			SetBindlessParameter(RHICmdList, Shader, Parameter, Handle);
+			RHICmdList.SetBindlessSampler(Shader, Parameter.GetBaseIndex(), SamplerStateRHI);
 		}
 		else
 #endif
@@ -261,33 +248,27 @@ FORCEINLINE void SetTextureParameter(
  * Template'd on shader type (e.g. pixel shader or compute shader).
  */
 template<typename TRHIShader, typename TRHICmdList>
-FORCEINLINE void SetSRVParameter(
-	TRHICmdList& RHICmdList,
-	TRHIShader* Shader,
-	const FShaderResourceParameter& Parameter,
-	FRHIShaderResourceView* NewShaderResourceViewRHI
-	)
+FORCEINLINE void SetSRVParameter(TRHICmdList& RHICmdList, TRHIShader* Shader, const FShaderResourceParameter& Parameter, FRHIShaderResourceView* SRV)
 {
 	if (Parameter.IsBound())
 	{
 #if PLATFORM_SUPPORTS_BINDLESS_RENDERING
 		if (Parameter.GetType() == EShaderParameterType::BindlessResourceIndex)
 		{
-			const FRHIDescriptorHandle Handle = NewShaderResourceViewRHI ? NewShaderResourceViewRHI->GetBindlessHandle() : FRHIDescriptorHandle();
-			SetBindlessParameter(RHICmdList, Shader, Parameter, Handle);
+			RHICmdList.SetBindlessResourceView(Shader, Parameter.GetBaseIndex(), SRV);
 		}
 		else
 #endif
 		{
-			RHICmdList.SetShaderResourceViewParameter(Shader, Parameter.GetBaseIndex(), NewShaderResourceViewRHI);
+			RHICmdList.SetShaderResourceViewParameter(Shader, Parameter.GetBaseIndex(), SRV);
 		}
 	}
 }
 
 template<typename TRHIShader, typename TRHICmdList>
-FORCEINLINE void SetSRVParameter(TRHICmdList& RHICmdList, const TRefCountPtr<TRHIShader>& Shader, const FShaderResourceParameter& Parameter, FRHIShaderResourceView* NewShaderResourceViewRHI)
+FORCEINLINE void SetSRVParameter(TRHICmdList& RHICmdList, const TRefCountPtr<TRHIShader>& Shader, const FShaderResourceParameter& Parameter, FRHIShaderResourceView* SRV)
 {
-	SetSRVParameter(RHICmdList, Shader.GetReference(), Parameter, NewShaderResourceViewRHI);
+	SetSRVParameter(RHICmdList, Shader.GetReference(), Parameter, SRV);
 }
 
 template<typename TRHIShader, typename TRHICmdList>
@@ -298,8 +279,7 @@ FORCEINLINE void SetUAVParameterSafeShader(TRHICmdList& RHICmdList, TRHIShader*
 #if PLATFORM_SUPPORTS_BINDLESS_RENDERING
 		if (Parameter.GetType() == EShaderParameterType::BindlessResourceIndex)
 		{
-			const FRHIDescriptorHandle Handle = UAV ? UAV->GetBindlessHandle() : FRHIDescriptorHandle();
-			SetBindlessParameter(RHICmdList, Shader, Parameter, Handle);
+			RHICmdList.SetBindlessUAV(Shader, Parameter.GetBaseIndex(), UAV);
 		}
 		else
 #endif
