ParameterBatcher wasn't inheriting the persistent state when a command list was being reset.

diff --git a/Engine/Source/Runtime/RHI/Private/RHICommandList.cpp b/Engine/Source/Runtime/RHI/Private/RHICommandList.cpp
--- a/Engine/Source/Runtime/RHI/Private/RHICommandList.cpp
+++ b/Engine/Source/Runtime/RHI/Private/RHICommandList.cpp
@@ -148,6 +148,17 @@ FRHIParameterBatcher::FRHIParameterBatcher()
 {
 	bEnabled = GRHICommandParameterBatching != 0;
 }
+
+FRHIParameterBatcher::FRHIParameterBatcher(const FBoundShaderStateInput& InBoundShaderStateInput, FRHIComputeShader* InBoundComputeShaderRHI)
+{
+	AllBatchedShaders[SF_Vertex] = InBoundShaderStateInput.GetVertexShader();
+	AllBatchedShaders[SF_Mesh] = InBoundShaderStateInput.GetMeshShader();
+	AllBatchedShaders[SF_Amplification] = InBoundShaderStateInput.GetAmplificationShader();
+	AllBatchedShaders[SF_Pixel] = InBoundShaderStateInput.GetPixelShader();
+	AllBatchedShaders[SF_Geometry] = InBoundShaderStateInput.GetGeometryShader();
+	AllBatchedShaders[SF_Compute] = InBoundComputeShaderRHI;
+}
+
 FRHIParameterBatcher::FRHIParameterBatcher(FRHIParameterBatcher&&) = default;
 FRHIParameterBatcher::~FRHIParameterBatcher()
 {
@@ -231,7 +242,8 @@ FRHICommandListBase::FRHICommandListBase(FRHIGPUMask InGPUMask, ERecordingThread
 {}
 
 FRHICommandListBase::FRHICommandListBase(FPersistentState&& InPersistentState)
-	: DispatchEvent  (FGraphEvent::CreateGraphEvent())
+	: ParameterBatcher(InPersistentState.BoundShaderInput, InPersistentState.BoundComputeShaderRHI)
+	, DispatchEvent(FGraphEvent::CreateGraphEvent())
 	, PersistentState(MoveTemp(InPersistentState))
 {
 	DispatchEvent->SetDebugName(TEXT("FRHICommandListBase::DispatchEvent"));
diff --git a/Engine/Source/Runtime/RHI/Public/RHICommandList.h b/Engine/Source/Runtime/RHI/Public/RHICommandList.h
--- a/Engine/Source/Runtime/RHI/Public/RHICommandList.h
+++ b/Engine/Source/Runtime/RHI/Public/RHICommandList.h
@@ -442,12 +442,13 @@ protected:
 		checkSlow(InShader);
 		const EShaderFrequency Frequency = InShader->GetFrequency();
 		checkSlow(Frequency < SF_NumStandardFrequencies);
-		checkSlow(AllBatchedShaders[Frequency] != nullptr);
+		check(AllBatchedShaders[Frequency] == InShader);
 		return AllBatchedShaderParameters[Frequency];
 	}
 
 public:
 	FRHIParameterBatcher();
+	FRHIParameterBatcher(const FBoundShaderStateInput& InBoundShaderStateInput, FRHIComputeShader* InBoundComputeShaderRHI);
 	FRHIParameterBatcher(FRHIParameterBatcher&&);
 	~FRHIParameterBatcher();
 
