Adding missing header dependency.

diff --git a/Engine/Source/Runtime/RHI/Private/Linux/LinuxDynamicRHI.cpp b/Engine/Source/Runtime/RHI/Private/Linux/LinuxDynamicRHI.cpp
--- a/Engine/Source/Runtime/RHI/Private/Linux/LinuxDynamicRHI.cpp
+++ b/Engine/Source/Runtime/RHI/Private/Linux/LinuxDynamicRHI.cpp
@@ -6,6 +6,7 @@
 #include "Modules/ModuleManager.h"
 #include "Misc/ConfigCacheIni.h"
 #include "HAL/PlatformApplicationMisc.h"
+#include "DataDrivenShaderPlatformInfo.h"
 
 FDynamicRHI* PlatformCreateDynamicRHI()
 {
