Added bSupportsClipDistance to DDSPI for mobile shader platforms that support SV_ClipDistance

diff --git a/Engine/Config/Android/DataDrivenPlatformInfo.ini b/Engine/Config/Android/DataDrivenPlatformInfo.ini
--- a/Engine/Config/Android/DataDrivenPlatformInfo.ini
+++ b/Engine/Config/Android/DataDrivenPlatformInfo.ini
@@ -63,6 +63,7 @@ bTargetsTiledGPU=true
 bSupportsMobileMultiView=true
 bNeedsOfflineCompiler=true
 bSupportsManualVertexFetch=false
+bSupportsClipDistance = false
 FriendlyName=LOCTEXT("FriendlyShaderPlatformName_AndroidES31", "Android ES 3_1")
 
 [ShaderPlatform VULKAN_ES3_1_ANDROID]
@@ -75,6 +76,7 @@ bSupportsMobileMultiView=true
 bNeedsOfflineCompiler=true
 bSupportsIndexBufferUAVs=true
 bSupportsManualVertexFetch=false
+bSupportsClipDistance = false
 FriendlyName=LOCTEXT("FriendlyShaderPlatformName_AndroidVulkan", "Android Vulkan")
 
 [ShaderPlatform VULKAN_SM5_ANDROID]
diff --git a/Engine/Config/HoloLens/DataDrivenPlatformInfo.ini b/Engine/Config/HoloLens/DataDrivenPlatformInfo.ini
--- a/Engine/Config/HoloLens/DataDrivenPlatformInfo.ini
+++ b/Engine/Config/HoloLens/DataDrivenPlatformInfo.ini
@@ -33,4 +33,5 @@ bIsMobile=true
 bSupportsInstancedStereo=true
 bSupportsManualVertexFetch=false
 bSupportsVertexShaderLayer=true
+bSupportsClipDistance = false
 FriendlyName=LOCTEXT("FriendlyShaderPlatformName_HololensES31", "Hololens Shader Model 3_1")
\ No newline at end of file
diff --git a/Engine/Shaders/Private/MobileBasePassPixelShader.usf b/Engine/Shaders/Private/MobileBasePassPixelShader.usf
--- a/Engine/Shaders/Private/MobileBasePassPixelShader.usf
+++ b/Engine/Shaders/Private/MobileBasePassPixelShader.usf
@@ -299,7 +299,7 @@ void Main(
 	ResolvedView = ResolveView();
 #endif
 
-#if USE_GLOBAL_CLIP_PLANE
+#if USE_GLOBAL_CLIP_PLANE && !PLATFORM_SUPPORTS_CLIP_DISTANCE
 	clip(BasePassInterpolants.OutClipDistance);
 #endif
 
diff --git a/Engine/Shaders/Public/Platform.ush b/Engine/Shaders/Public/Platform.ush
--- a/Engine/Shaders/Public/Platform.ush
+++ b/Engine/Shaders/Public/Platform.ush
@@ -372,8 +372,8 @@
 	#define UB_FIXED_FLOAT(Dim) float##Dim
 #endif
 
-// On mobile use interpolator for ClipDistance as it does not work/supported on all mobile platforms
-#if (FEATURE_LEVEL == FEATURE_LEVEL_ES3_1)
+// Use interpolator for platforms that do not support ClipDistance
+#if !PLATFORM_SUPPORTS_CLIP_DISTANCE
 #define SV_ClipDistance OUTCLIPDIST
 #endif
 
diff --git a/Engine/Source/Runtime/Engine/Private/ShaderCompiler/ShaderCompiler.cpp b/Engine/Source/Runtime/Engine/Private/ShaderCompiler/ShaderCompiler.cpp
--- a/Engine/Source/Runtime/Engine/Private/ShaderCompiler/ShaderCompiler.cpp
+++ b/Engine/Source/Runtime/Engine/Private/ShaderCompiler/ShaderCompiler.cpp
@@ -6248,6 +6248,11 @@ void GlobalBeginCompileShader(
 		Input.Environment.SetDefine(TEXT("PROJECT_ALLOW_GLOBAL_CLIP_PLANE"), CVar ? (CVar->GetInt() != 0) : 0);
 	}
 
+	{
+		const bool bSupportsClipDistance = FDataDrivenShaderPlatformInfo::GetSupportsClipDistance((EShaderPlatform)Target.Platform);
+		Input.Environment.SetDefine(TEXT("PLATFORM_SUPPORTS_CLIP_DISTANCE"), bSupportsClipDistance ? 1u : 0u);
+	}
+
 	ITargetPlatform* TargetPlatform = GetTargetPlatformManager()->FindTargetPlatformWithSupport(TEXT("ShaderFormat"), ShaderFormatName);
 	bool bForwardShading = false;
 	{
diff --git a/Engine/Source/Runtime/RHI/Private/DataDrivenShaderPlatformInfo.cpp b/Engine/Source/Runtime/RHI/Private/DataDrivenShaderPlatformInfo.cpp
--- a/Engine/Source/Runtime/RHI/Private/DataDrivenShaderPlatformInfo.cpp
+++ b/Engine/Source/Runtime/RHI/Private/DataDrivenShaderPlatformInfo.cpp
@@ -128,6 +128,7 @@ void FGenericDataDrivenShaderPlatformInfo::SetDefaultValues()
 	bSupportsAsyncPipelineCompilation = true;
 	bSupportsManualVertexFetch = true;
 	bSupportsVolumeTextureAtomics = true;
+	bSupportsClipDistance = true;
 	PreviewShaderPlatformParent = EShaderPlatform::SP_NumPlatforms;
 }
 
@@ -249,6 +250,7 @@ void FGenericDataDrivenShaderPlatformInfo::ParseDataDrivenShaderInfo(const FConf
 	GET_SECTION_INT_HELPER(EnablesHLSL2021ByDefault);
 	GET_SECTION_BOOL_HELPER(bSupportsSceneDataCompressedTransforms);
 	GET_SECTION_BOOL_HELPER(bSupportsSwapchainUAVs);
+	GET_SECTION_BOOL_HELPER(bSupportsClipDistance);
 #undef GET_SECTION_BOOL_HELPER
 #undef GET_SECTION_INT_HELPER
 #undef GET_SECTION_SUPPORT_HELPER
diff --git a/Engine/Source/Runtime/RHI/Public/DataDrivenShaderPlatformInfo.h b/Engine/Source/Runtime/RHI/Public/DataDrivenShaderPlatformInfo.h
--- a/Engine/Source/Runtime/RHI/Public/DataDrivenShaderPlatformInfo.h
+++ b/Engine/Source/Runtime/RHI/Public/DataDrivenShaderPlatformInfo.h
@@ -109,6 +109,7 @@ class RHI_API FGenericDataDrivenShaderPlatformInfo
 	uint32 bSupportsSceneDataCompressedTransforms : 1;
 	uint32 bIsPreviewPlatform : 1;
 	uint32 bSupportsSwapchainUAVs : 1;
+	uint32 bSupportsClipDistance : 1;
 
 #if WITH_EDITOR
 	FText FriendlyName;
@@ -733,6 +734,12 @@ public:
 	{
 		return Infos[Platform].EnablesHLSL2021ByDefault;
 	}
+	
+	static FORCEINLINE_DEBUGGABLE const bool GetSupportsClipDistance(const FStaticShaderPlatform Platform)
+	{
+		check(IsValid(Platform));
+		return Infos[Platform].bSupportsClipDistance;
+	}
 
 #if WITH_EDITOR
 	static FText GetFriendlyName(const FStaticShaderPlatform Platform);
