resubmit of 23347569, with a few changes:

Improvements to PSO precaching:
* Whenever a task was issued, we would not return the GraphEvent associated to it, and assume the task was already done. This is fixed by returning the CompletionEvent if the object is still in ActivePrecachingTasks
* If a task was in the process of being issued, we would miss it because the lock was not kept after TryAddNewState. We would call CreateGraphicsPipelineState, and assign the CompletionEvent to FPipelineState without any lock. We now assign the event during the pipeline creation, and before calling the PSO compilation. This is done for both precaching and JIT PSO requests: InternalCreateGraphicsPipelineState does not create graphevents anymore
* Change ActivePrecachingTasks to use the same hashing-then-match functions as PrecachedPSOInitializers. I got some hash collisions during stress testing, and I think it's better to have both data structures return the same answer about the existence of an existing PSO
* Have Compute shader use a similar structure for precaching

diff --git a/Engine/Source/Runtime/Core/Public/Async/TaskGraphInterfaces.h b/Engine/Source/Runtime/Core/Public/Async/TaskGraphInterfaces.h
--- a/Engine/Source/Runtime/Core/Public/Async/TaskGraphInterfaces.h
+++ b/Engine/Source/Runtime/Core/Public/Async/TaskGraphInterfaces.h
@@ -1284,6 +1284,19 @@ public:
 		return FConstructor(new TGraphTask(TTask::GetSubsequentsMode() == ESubsequentsMode::FireAndForget ? NULL : FGraphEvent::CreateGraphEvent(), NumPrereq), Prerequisites, CurrentThreadIfKnown);
 	}
 
+	/**
+	 *	Factory to create a task and return the helper object to construct the embedded task and set it up for execution.
+	 *	@param Prerequisites; the list of FGraphEvents that must be completed prior to this task executing.
+	 *	@param CurrentThreadIfKnown; provides the index of the thread we are running on. Can be ENamedThreads::AnyThread if the current thread is unknown.
+	 *	@return a temporary helper class which can be used to complete the process.
+	**/
+	static FConstructor CreateTaskWithEvent(FGraphEventRef SubsequentEventToAssume, const FGraphEventArray* Prerequisites = NULL, ENamedThreads::Type CurrentThreadIfKnown = ENamedThreads::AnyThread)
+	{
+		int32 NumPrereq = Prerequisites ? Prerequisites->Num() : 0;
+		check(SubsequentEventToAssume == NULL || TTask::GetSubsequentsMode() == ESubsequentsMode::TrackSubsequents);
+		return FConstructor(new TGraphTask(SubsequentEventToAssume, NumPrereq), Prerequisites, CurrentThreadIfKnown);
+	}
+
 	void Unlock(ENamedThreads::Type CurrentThreadIfKnown = ENamedThreads::AnyThread)
 	{
 		TaskTrace::Launched(GetTraceId(), nullptr, Subsequents.IsValid(), ((TTask*)&TaskStorage)->GetDesiredThread());
diff --git a/Engine/Source/Runtime/RHI/Private/PipelineStateCache.cpp b/Engine/Source/Runtime/RHI/Private/PipelineStateCache.cpp
--- a/Engine/Source/Runtime/RHI/Private/PipelineStateCache.cpp
+++ b/Engine/Source/Runtime/RHI/Private/PipelineStateCache.cpp
@@ -612,6 +612,11 @@ int32 FindRayTracingMissShaderIndex(FRayTracingPipelineState* Pipeline, FRHIRayT
 	return INDEX_NONE;
 }
 
+bool IsPrecachedPSO(const FGraphicsPipelineStateInitializer& Initializer)
+{
+	return Initializer.bFromPSOFileCache || Initializer.bPSOPrecache;
+}
+
 void SetGraphicsPipelineState(FRHICommandList& RHICmdList, const FGraphicsPipelineStateInitializer& Initializer, uint32 StencilRef, EApplyRendertargetOption ApplyFlags, bool bApplyAdditionalState, EPSOPrecacheResult PSOPrecacheResult)
 {
 #if PLATFORM_USE_FALLBACK_PSO
@@ -944,80 +949,73 @@ private:
 
 };
 
-class FPrecacheGraphicsPipelineCache
+template<class TPrecachePipelineCacheDerived, class TPrecachedPSOInitializer, class TPipelineState>
+class TPrecachePipelineCacheBase
 {
 public:
-
-	~FPrecacheGraphicsPipelineCache()
+	~TPrecachePipelineCacheBase()
 	{
 		// Wait for all precache tasks to finished
 		WaitTasksComplete();
 	}
-
-	bool Contains(const FGraphicsPipelineStateInitializer& Initializer)
-	{
-		FRWScopeLock ReadLock(PrecachePSOsRWLock, SLT_ReadOnly);
-		return PrecachedPSOInitializers.Contains(Initializer);
-	}
-
-	FGraphicsPipelineState* TryAddNewState(const FGraphicsPipelineStateInitializer& Initializer, bool bDoAsyncCompile)
+	
+protected:
+	void TryAddNewState(const TPrecachedPSOInitializer& Initializer, bool bDoAsyncCompile, FGraphEventRef& GraphEvent)
 	{
+		uint32 InitializerHash = TPrecachePipelineCacheDerived::PipelineStateInitializerHash(Initializer);
 		// Fast check first with read lock
-		if (Contains(Initializer))
 		{
-			return nullptr;
-		}		
-
-		FGraphicsPipelineState* NewState = nullptr;
-
-		// Now try and add with write lock
-		{
-			FRWScopeLock WriteLock(PrecachePSOsRWLock, SLT_Write);
-
-			// check again
-			if (PrecachedPSOInitializers.Contains(Initializer))
+			FRWScopeLock ReadLock(PrecachePSOsRWLock, SLT_ReadOnly);
+			if (HasPSOBeenRequested(Initializer, InitializerHash, GraphEvent))
 			{
-				return nullptr;
+				return;
 			}
-
-			// Add to cache before starting async operation, because the async op will remove the task from active ops when done			
-			PrecachedPSOInitializers.Add(Initializer);
 		}
 
-		// Add to active set of precaching PSOs if it's created async
-		if (bDoAsyncCompile)
-		{
-			// Only need the hash of the PSO here
-			uint64 PrecachePSOHash = RHIComputePrecachePSOHash(Initializer);
-
-			FRWScopeLock WriteLock(ActivePrecachingTasksRWLock, SLT_Write);
+		// Now try and add with write lock
 
-			if (ActivePrecachingTasks.Contains(PrecachePSOHash))
+		TPipelineState* NewState = nullptr;
+		{
+			FRWScopeLock WriteLock(PrecachePSOsRWLock, SLT_Write);
+			if (HasPSOBeenRequested(Initializer, InitializerHash, GraphEvent))
 			{
-				return nullptr;
+				return;
 			}
 
 			// create new graphics state
-			NewState = new FGraphicsPipelineState();
+			NewState = TPrecachePipelineCacheDerived::CreateNewPSO(Initializer);
 
-			FPrecacheTask PrecacheTask;
-			PrecacheTask.Initializer = Initializer;
-			PrecacheTask.PSO = NewState;
+			// Add to cache before starting async operation, because the async op will remove the task from active ops when done			
+			PrecachedPSOInitializers.AddByHash(InitializerHash, Initializer);
 
-			ActivePrecachingTasks.Add(PrecachePSOHash, PrecacheTask);
-		}
-		else
-		{
-			// create new graphics state
-			NewState = new FGraphicsPipelineState();
+			// Add to active set of precaching PSOs if it's created async
+			if (bDoAsyncCompile)
+			{
+				FPrecacheTask PrecacheTask;
+				PrecacheTask.Initializer = Initializer;
+				PrecacheTask.PSO = NewState;
+#if DO_ENSURE
+				ensure(!ActivePrecachingTasks.ContainsByHash(InitializerHash, Initializer));
+#endif
+				ActivePrecachingTasks.AddByHash(InitializerHash, Initializer, PrecacheTask);
+    			// Assign the event at this point because we need to release the lock before calling OnNewPipelineStateCreated which
+	    		// might call PrecacheFinished directly (The background task might get abandoned) and FRWLock can't be acquired recursively
+	    		// Note that calling IsComplete will return false until we link it somehow like we do below
+				NewState->CompletionEvent = FGraphEvent::CreateGraphEvent();
+				GraphEvent = NewState->CompletionEvent;
+			}
 		}
 
-		return NewState;
+		TPrecachePipelineCacheDerived::OnNewPipelineStateCreated(Initializer, NewState, bDoAsyncCompile);
+
+		// Make sure that we don't access NewState here as the task might have already been finished, ProcessDelayedCleanup may have been called
+		// and NewState already been deleted
 	}
+public:
 
 	void WaitTasksComplete()
 	{
-		FRWScopeLock WriteLock(ActivePrecachingTasksRWLock, SLT_Write);
+		FRWScopeLock WriteLock(PrecachePSOsRWLock, SLT_Write);
 
 		for (auto Iterator = ActivePrecachingTasks.CreateIterator(); Iterator; ++Iterator)
 		{
@@ -1029,48 +1027,50 @@ public:
 		ActivePrecachingTasks.Empty();
 	}
 
-	bool IsPrecaching(const FGraphicsPipelineStateInitializer& Initializer)
+	EPSOPrecacheResult GetPrecachingState(const TPrecachedPSOInitializer& Initializer)
 	{
-		uint64 PrecachePSOHash = RHIComputePrecachePSOHash(Initializer);
+		uint32 InitializerHash = TPrecachePipelineCacheDerived::PipelineStateInitializerHash(Initializer);
+		FRWScopeLock ReadLock(PrecachePSOsRWLock, SLT_ReadOnly);
+		FPrecacheTask* FindResult = ActivePrecachingTasks.FindByHash(InitializerHash, Initializer);
+		if (FindResult && !FindResult->PSO->IsComplete())
+		{
+			return EPSOPrecacheResult::Active;
+		}
 
-		FRWScopeLock ReadLock(ActivePrecachingTasksRWLock, SLT_ReadOnly);
-		FPrecacheTask* FindResult = ActivePrecachingTasks.Find(PrecachePSOHash);
-		if (FindResult)
+		if (!PrecachedPSOInitializers.ContainsByHash(InitializerHash, Initializer))
 		{
-			// Check if not complete yet
-			return !(FindResult->PSO->IsComplete());
+			return EPSOPrecacheResult::Missed;
 		}
 
 		// not found then assume it's not precaching
-		return false;
+		return EPSOPrecacheResult::Complete;
 	}
 
 	bool IsPrecaching()
 	{
-		FRWScopeLock ReadLock(ActivePrecachingTasksRWLock, SLT_ReadOnly);
+		FRWScopeLock ReadLock(PrecachePSOsRWLock, SLT_ReadOnly);
 		return !ActivePrecachingTasks.IsEmpty();
 	}
 
-	void PrecacheFinished(const FGraphicsPipelineStateInitializer& Initializer)
+	void PrecacheFinished(const TPrecachedPSOInitializer& Initializer)
 	{
-		uint64 PrecachePSOHash = RHIComputePrecachePSOHash(Initializer);
-
-		FRWScopeLock WriteLock(ActivePrecachingTasksRWLock, SLT_Write);
-		FPrecacheTask* FindResult = ActivePrecachingTasks.Find(PrecachePSOHash);
+		uint32 InitializerHash = TPrecachePipelineCacheDerived::PipelineStateInitializerHash(Initializer);
+		FRWScopeLock WriteLock(PrecachePSOsRWLock, SLT_Write);
+		FPrecacheTask* FindResult = ActivePrecachingTasks.FindByHash(InitializerHash, Initializer);
 		// Might have already been processed because the task got rescheduled
 		if (FindResult != nullptr)
 		{
-			verify(ActivePrecachingTasks.Remove(PrecachePSOHash) == 1);
+			verify(ActivePrecachingTasks.RemoveByHash(InitializerHash, Initializer) == 1);
 			PrecachedPSOs.Add(FindResult->PSO);
 		}
 	}
 
 	void ProcessDelayedCleanup()
 	{
-		FRWScopeLock WriteLock(ActivePrecachingTasksRWLock, SLT_Write);
+		FRWScopeLock WriteLock(PrecachePSOsRWLock, SLT_Write);
 		for (int32 Index = 0; Index < PrecachedPSOs.Num(); ++Index)
 		{
-			FGraphicsPipelineState* GraphicsPSO = PrecachedPSOs[Index];
+			TPipelineState* GraphicsPSO = PrecachedPSOs[Index];
 			if (GraphicsPSO->IsComplete())
 			{
 				// This is needed to cleanup the members - bit strange because it's complete already
@@ -1079,40 +1079,120 @@ public:
 				delete GraphicsPSO;
 				PrecachedPSOs.RemoveAtSwap(Index);
 				Index--;
-			}			
+			}
 		}
-
 	}
 
 private:
-
-	struct FGraphicsPSOInitializerKeyFuncs : DefaultKeyFuncs<FGraphicsPipelineStateInitializer>
+	bool HasPSOBeenRequested(const TPrecachedPSOInitializer& Initializer, uint32 InitializerHash, FGraphEventRef& GraphEvent)
 	{
-		static FORCEINLINE bool Matches(KeyInitType A, KeyInitType B)
+		FPrecacheTask* FindResult = ActivePrecachingTasks.FindByHash(InitializerHash, Initializer);
+		if (FindResult)
 		{
-			return RHIMatchPrecachePSOInitializers(A, B);
+			GraphEvent = FindResult->PSO->CompletionEvent;
+			check(GraphEvent != nullptr);
+			return true;
 		}
 
-		static FORCEINLINE uint32 GetKeyHash(KeyInitType Key)
+		if (PrecachedPSOInitializers.ContainsByHash(InitializerHash, Initializer))
 		{
-			return RHIComputePrecachePSOHash(Key);
+			return true;
 		}
-	};
+
+		return false;
+	}
 
 	// Cache of all precached PSOs - need correct compare to make sure we precache all of them
 	FRWLock PrecachePSOsRWLock;
-	TSet<FGraphicsPipelineStateInitializer, FGraphicsPSOInitializerKeyFuncs> PrecachedPSOInitializers;
+
+	using TSetKeyFuncsBaseClass = DefaultKeyFuncs<TPrecachedPSOInitializer>;
+	struct TSetKeyFuncs : public TSetKeyFuncsBaseClass
+	{
+		static FORCEINLINE bool Matches(typename TSetKeyFuncsBaseClass::KeyInitType A, typename TSetKeyFuncsBaseClass::KeyInitType B)
+		{
+			return TPrecachePipelineCacheDerived::PipelineStateInitializerMatch(A, B);
+		}
+		static FORCEINLINE uint32 GetKeyHash(typename TSetKeyFuncsBaseClass::KeyInitType Key)
+		{
+			return TPrecachePipelineCacheDerived::PipelineStateInitializerHash(Key);
+		}
+	};
+	TSet<TPrecachedPSOInitializer, TSetKeyFuncs> PrecachedPSOInitializers;
 
 	// Current active background precache operations - full PSO hash can be used for fast retrieval
 	// (on hash clash it might report that a PSO is still compiling)
-	FRWLock ActivePrecachingTasksRWLock;
 	struct FPrecacheTask
 	{
-		FGraphicsPipelineStateInitializer Initializer;
-		FGraphicsPipelineState* PSO = nullptr;
+		TPrecachedPSOInitializer Initializer;
+		TPipelineState* PSO = nullptr;
 	};
-	TMap<uint64, FPrecacheTask> ActivePrecachingTasks;
-	TArray<FGraphicsPipelineState*> PrecachedPSOs;
+
+	using TMapKeyFuncsBaseClass = TDefaultMapHashableKeyFuncs<TPrecachedPSOInitializer, FPrecacheTask, false>;
+	struct TMapKeyFuncs : public TMapKeyFuncsBaseClass
+	{
+		static FORCEINLINE bool Matches(typename TMapKeyFuncsBaseClass::KeyInitType A, typename TMapKeyFuncsBaseClass::KeyInitType B)
+		{
+			return TPrecachePipelineCacheDerived::PipelineStateInitializerMatch(A, B);
+		}
+		static FORCEINLINE uint32 GetKeyHash(typename TMapKeyFuncsBaseClass::KeyInitType Key)
+		{
+			return TPrecachePipelineCacheDerived::PipelineStateInitializerHash(Key);
+		}
+	};
+
+	TMap<TPrecachedPSOInitializer, FPrecacheTask, FDefaultSetAllocator, TMapKeyFuncs> ActivePrecachingTasks;
+	TArray<TPipelineState*> PrecachedPSOs;
+};
+
+class FPrecacheComputePipelineCache : public TPrecachePipelineCacheBase<FPrecacheComputePipelineCache, FRHIComputeShader*, FComputePipelineState>
+{
+public:
+	static FComputePipelineState* CreateNewPSO(const FRHIComputeShader* ComputeShader)
+	{
+		return new FComputePipelineState((FRHIComputeShader*)ComputeShader);
+	}
+
+	FGraphEventRef PrecacheComputePipelineState(FRHIComputeShader* ComputeShader, bool bForcePrecache);
+	static void OnNewPipelineStateCreated(FRHIComputeShader* ComputeShader, FComputePipelineState* NewComputePipelineState, bool bDoAsyncCompile);
+
+	static FORCEINLINE bool PipelineStateInitializerMatch(const FRHIComputeShader* A, const FRHIComputeShader* B)
+	{
+		// todo: would be good/more robust to have the CS equivalent of RHIMatchPrecachePSOInitializers instead of relying on pointers
+		// (for example if an FRHIComputeShader gets released and a new one is allocated with the same pointer value)
+		return A == B;
+	}
+
+	static FORCEINLINE uint32 PipelineStateInitializerHash(const FRHIComputeShader* Key)
+	{
+		// will hash the pointer through PointerHash(const void* Key)
+		return GetTypeHash(Key);
+	}
+
+};
+
+class FPrecacheGraphicsPipelineCache : public TPrecachePipelineCacheBase<FPrecacheGraphicsPipelineCache, FGraphicsPipelineStateInitializer, FGraphicsPipelineState>
+{
+public:
+	static FGraphicsPipelineState* CreateNewPSO(const FGraphicsPipelineStateInitializer& Initializer)
+	{
+		return new FGraphicsPipelineState;
+	}
+
+	static FORCEINLINE bool PipelineStateInitializerMatch(const FGraphicsPipelineStateInitializer& A, const FGraphicsPipelineStateInitializer& B)
+	{
+		return RHIMatchPrecachePSOInitializers(A, B);
+	}
+
+	static FORCEINLINE uint32 PipelineStateInitializerHash(const FGraphicsPipelineStateInitializer& Key)
+	{
+		return RHIComputePrecachePSOHash(Key);
+	}
+
+
+	FGraphEventRef PrecacheGraphicsPipelineState(const FGraphicsPipelineStateInitializer& Initializer);
+
+	static void OnNewPipelineStateCreated(const FGraphicsPipelineStateInitializer& Initializer, FGraphicsPipelineState* NewGraphicsPipelineState, bool bDoAsyncCompile);
+
 };
 
 // Typed caches for compute and graphics
@@ -1123,6 +1203,7 @@ typedef TSharedPipelineStateCache<FGraphicsPipelineStateInitializer, FGraphicsPi
 FComputePipelineCache GComputePipelineCache;
 FGraphicsPipelineCache GGraphicsPipelineCache;
 FPrecacheGraphicsPipelineCache GPrecacheGraphicsPipelineCache;
+FPrecacheComputePipelineCache GPrecacheComputePipelineCache;
 
 FAutoConsoleTaskPriority CPrio_FCompilePipelineStateTask(
 	TEXT("TaskGraph.TaskPriorities.CompilePipelineStateTask"),
@@ -1353,7 +1434,7 @@ public:
 		, Initializer(InInitializer)
 		, PSOPreCacheResult(InPSOPreCacheResult)
 	{
-		if(Initializer.bFromPSOFileCache)
+		if(IsPrecachedPSO(Initializer))
 		{
 			GPipelinePrecompileTasksInFlight++;
 		}
@@ -1414,13 +1495,13 @@ public:
 
 	~FCompilePipelineStateTask()
 	{
-		if (Initializer.bFromPSOFileCache)
+		if (IsPrecachedPSO(Initializer))
 		{
 			GPipelinePrecompileTasksInFlight--;
 		}
 	}
 
-	static ESubsequentsMode::Type GetSubsequentsMode() { return ESubsequentsMode::TrackSubsequents; }
+	static constexpr ESubsequentsMode::Type GetSubsequentsMode() { return ESubsequentsMode::TrackSubsequents; }
 
 	void DoTask(ENamedThreads::Type CurrentThread, const FGraphEventRef& MyCompletionGraphEvent)
 	{
@@ -1436,6 +1517,10 @@ public:
 		{
 			FComputePipelineState* ComputePipeline = static_cast<FComputePipelineState*>(Pipeline);
 			ComputePipeline->RHIPipeline = RHICreateComputePipelineState(ComputePipeline->ComputeShader);
+			if (Initializer.bPSOPrecache)
+			{
+				GPrecacheComputePipelineCache.PrecacheFinished(ComputePipeline->ComputeShader);
+			}
 		}
 		else
 		{
@@ -1578,6 +1663,7 @@ void PipelineStateCache::FlushResources()
 	GGraphicsPipelineCache.ProcessDelayedCleanup();
 
 	GPrecacheGraphicsPipelineCache.ProcessDelayedCleanup();
+	GPrecacheComputePipelineCache.ProcessDelayedCleanup();
 
 	{
 		int32 NumMissesThisFrame = GraphicsPipelineCacheMisses.Load(EMemoryOrder::Relaxed);
@@ -1705,7 +1791,9 @@ FComputePipelineState* PipelineStateCache::GetAndOrCreateComputePipelineState(FR
 		// create a compilation task, or just do it now...
 		if (DoAsyncCompile)
 		{
-			OutCachedState->CompletionEvent = TGraphTask<FCompilePipelineStateTask>::CreateTask().ConstructAndDispatchWhenReady(OutCachedState, FGraphicsPipelineStateInitializer(), EPSOPrecacheResult::Untracked);
+			FGraphicsPipelineStateInitializer GraphicsPipelineStateInitializer;
+			GraphicsPipelineStateInitializer.bFromPSOFileCache = bFromFileCache;
+			OutCachedState->CompletionEvent = TGraphTask<FCompilePipelineStateTask>::CreateTask().ConstructAndDispatchWhenReady(OutCachedState, GraphicsPipelineStateInitializer, EPSOPrecacheResult::Untracked);
 			if (!bFromFileCache)
 			{
 				RHICmdList.AddDispatchPrerequisite(OutCachedState->CompletionEvent);
@@ -2111,26 +2199,24 @@ inline void ValidateGraphicsPipelineStateInitializer(const FGraphicsPipelineStat
 	check(Initializer.DepthStencilState && Initializer.BlendState && Initializer.RasterizerState);
 }
 
-static FGraphEventRef CreateGraphicsPipelineState(const FGraphicsPipelineStateInitializer& Initializer, EPSOPrecacheResult PSOPrecacheResult, bool bDoAsyncCompile, bool bPSOPrecache, FGraphicsPipelineState* CachedState)
+static void InternalCreateGraphicsPipelineState(const FGraphicsPipelineStateInitializer& Initializer, EPSOPrecacheResult PSOPrecacheResult, bool bDoAsyncCompile, bool bPSOPrecache, FGraphicsPipelineState* CachedState)
 {
-	FGraphEventRef GraphEvent;
+	FGraphEventRef GraphEvent = CachedState->CompletionEvent;
 
 	// create a compilation task, or just do it now...
 	if (bDoAsyncCompile)
 	{
+		check(GraphEvent != nullptr);
 		// Use normal task graph for non-precompile jobs (or when thread pool is not enabled)
 		if (!bPSOPrecache || !FPSOPrecacheThreadPool::UsePool())
 		{
-			CachedState->CompletionEvent = TGraphTask<FCompilePipelineStateTask>::CreateTask().ConstructAndDispatchWhenReady(CachedState, Initializer, PSOPrecacheResult);
-			GraphEvent = CachedState->CompletionEvent;
+			TGraphTask<FCompilePipelineStateTask>::CreateTaskWithEvent(GraphEvent).ConstructAndDispatchWhenReady(CachedState, Initializer, PSOPrecacheResult);
 		}
 		else
 		{
 			// Here, PSO precompiles use a separate thread pool.
 			// Note that we do not add precompile tasks as cmdlist prerequisites.
-			GraphEvent = FGraphEvent::CreateGraphEvent();
 			TUniquePtr<FCompilePipelineStateTask> ThreadPoolTask = MakeUnique<FCompilePipelineStateTask>(CachedState, Initializer, PSOPrecacheResult);
-			CachedState->CompletionEvent = GraphEvent;
 			uint64 StartTime = FPlatformTime::Cycles64();
 #if	PSO_TRACK_CACHE_STATS
 			StatsStartPrecompile();
@@ -2154,6 +2240,7 @@ static FGraphEventRef CreateGraphicsPipelineState(const FGraphicsPipelineStateIn
 	}
 	else
 	{
+		check(GraphEvent == nullptr);
 		CachedState->RHIPipeline = RHICreateGraphicsPipelineState(Initializer);
 		if (CachedState->RHIPipeline)
 		{
@@ -2164,8 +2251,6 @@ static FGraphEventRef CreateGraphicsPipelineState(const FGraphicsPipelineStateIn
 			HandlePipelineCreationFailure(Initializer);
 		}
 	}
-
-	return GraphEvent;
 }
 
 FGraphicsPipelineState* PipelineStateCache::GetAndOrCreateGraphicsPipelineState(FRHICommandList& RHICmdList, const FGraphicsPipelineStateInitializer& Initializer, EApplyRendertargetOption ApplyFlags, EPSOPrecacheResult PSOPrecacheResult)
@@ -2212,6 +2297,7 @@ FGraphicsPipelineState* PipelineStateCache::GetAndOrCreateGraphicsPipelineState(
 	FGraphicsPipelineState* OutCachedState = nullptr;
 
 	bool bWasFound = GGraphicsPipelineCache.Find(Initializer, OutCachedState);
+    ensure(!Initializer.bPSOPrecache);
 	bool DoAsyncCompile = IsAsyncCompilationAllowed(RHICmdList, Initializer.bFromPSOFileCache);
 
 	if (bWasFound == false)
@@ -2221,14 +2307,20 @@ FGraphicsPipelineState* PipelineStateCache::GetAndOrCreateGraphicsPipelineState(
 		// create new graphics state
 		OutCachedState = new FGraphicsPipelineState();
 		OutCachedState->Stats = FPipelineFileCacheManager::RegisterPSOStats(GetTypeHash(Initializer));
-		
+		if (DoAsyncCompile)
+		{
+			OutCachedState->CompletionEvent = FGraphEvent::CreateGraphEvent();
+		}
+
 		if (!Initializer.bFromPSOFileCache)
 		{
 			GraphicsPipelineCacheMisses++;
 		}
 
 		bool bPSOPrecache = Initializer.bFromPSOFileCache;
-		FGraphEventRef GraphEvent = CreateGraphicsPipelineState(Initializer, PSOPrecacheResult, DoAsyncCompile, bPSOPrecache, OutCachedState);
+		FGraphEventRef GraphEvent = OutCachedState->CompletionEvent;
+		check(PSOPrecacheResult != EPSOPrecacheResult::Active);
+		InternalCreateGraphicsPipelineState(Initializer, PSOPrecacheResult, DoAsyncCompile, bPSOPrecache, OutCachedState);
 
 		// Add dispatch pre requisite for non precaching jobs only
 		//if (GraphEvent.IsValid() && (!bPSOPrecache || !FPSOPrecacheThreadPool::UsePool()))
@@ -2285,10 +2377,10 @@ bool PipelineStateCache::IsPSOPrecachingEnabled()
 #endif // WITH_EDITOR
 }
 
-FGraphEventRef PipelineStateCache::PrecacheComputePipelineState(FRHIComputeShader* ComputeShader, bool bForcePrecache)
+FGraphEventRef FPrecacheComputePipelineCache::PrecacheComputePipelineState(FRHIComputeShader* ComputeShader, bool bForcePrecache)
 {
 	FGraphEventRef GraphEvent = nullptr;
-	if (!IsPSOPrecachingEnabled() && !bForcePrecache)
+	if (!PipelineStateCache::IsPSOPrecachingEnabled() && !bForcePrecache)
 	{
 		return GraphEvent;
 	}
@@ -2298,32 +2390,36 @@ FGraphEventRef PipelineStateCache::PrecacheComputePipelineState(FRHIComputeShade
 		return GraphEvent;
 	}
 
-	FComputePipelineState* CachedState = nullptr;
-	if (GComputePipelineCache.Find(ComputeShader, CachedState))
+	static bool bDoAsyncCompile = FApp::ShouldUseThreadingForPerformance();
+	GPrecacheComputePipelineCache.TryAddNewState(ComputeShader, bDoAsyncCompile, GraphEvent);
+
+	if (GraphEvent == nullptr)
 	{
-		return GraphEvent;
+		check(GPrecacheComputePipelineCache.GetPrecachingState(ComputeShader) == EPSOPrecacheResult::Complete)
 	}
 
-	// create new graphics state
-	CachedState = new FComputePipelineState(ComputeShader);
+	return GraphEvent;
+}
 
+void FPrecacheComputePipelineCache::OnNewPipelineStateCreated(FRHIComputeShader* ComputeShader, FComputePipelineState* CachedState, bool bDoAsyncCompile)
+{
 	// create a compilation task, or just do it now...
-	static bool bDoAsyncCompile = FApp::ShouldUseThreadingForPerformance();
+	FGraphEventRef GraphEvent = CachedState->CompletionEvent;
 	if (bDoAsyncCompile)
 	{
+		check(GraphEvent != nullptr);
+		FGraphicsPipelineStateInitializer GraphicsPipelineStateInitializer;
+		GraphicsPipelineStateInitializer.bPSOPrecache = true;
 		// Thread pool disabled?
 		if (!FPSOPrecacheThreadPool::UsePool())
 		{
-			CachedState->CompletionEvent = TGraphTask<FCompilePipelineStateTask>::CreateTask().ConstructAndDispatchWhenReady(CachedState, FGraphicsPipelineStateInitializer(), EPSOPrecacheResult::Active);
-			GraphEvent = CachedState->CompletionEvent;
+			TGraphTask<FCompilePipelineStateTask>::CreateTaskWithEvent(GraphEvent).ConstructAndDispatchWhenReady(CachedState, GraphicsPipelineStateInitializer, EPSOPrecacheResult::Active);
 		}
 		else
 		{
 			// Here, PSO precompiles use a separate thread pool.
 			// Note that we do not add precompile tasks as cmdlist prerequisites.
-			GraphEvent = FGraphEvent::CreateGraphEvent();
-			TUniquePtr<FCompilePipelineStateTask> ThreadPoolTask = MakeUnique<FCompilePipelineStateTask>(CachedState, FGraphicsPipelineStateInitializer(), EPSOPrecacheResult::Active);
-			CachedState->CompletionEvent = GraphEvent;
+			TUniquePtr<FCompilePipelineStateTask> ThreadPoolTask = MakeUnique<FCompilePipelineStateTask>(CachedState, GraphicsPipelineStateInitializer, EPSOPrecacheResult::Active);
 			uint64 StartTime = FPlatformTime::Cycles64();
 #if	PSO_TRACK_CACHE_STATS
 			StatsStartPrecompile();
@@ -2347,16 +2443,19 @@ FGraphEventRef PipelineStateCache::PrecacheComputePipelineState(FRHIComputeShade
 	}
 	else
 	{
+		check(GraphEvent == nullptr);
 		CachedState->RHIPipeline = RHICreateComputePipelineState(CachedState->ComputeShader);
 	}
+}
 
-	GComputePipelineCache.Add(ComputeShader, CachedState);
-	return GraphEvent;
+FGraphEventRef PipelineStateCache::PrecacheComputePipelineState(FRHIComputeShader* ComputeShader, bool bForcePrecache)
+{
+	return GPrecacheComputePipelineCache.PrecacheComputePipelineState(ComputeShader, bForcePrecache);
 }
 
-FGraphEventRef PipelineStateCache::PrecacheGraphicsPipelineState(const FGraphicsPipelineStateInitializer& Initializer)
+FGraphEventRef FPrecacheGraphicsPipelineCache::PrecacheGraphicsPipelineState(const FGraphicsPipelineStateInitializer& Initializer)
 {
-	if (!IsPSOPrecachingEnabled())
+	if (!PipelineStateCache::IsPSOPrecachingEnabled())
 	{
 		return nullptr;
 	}
@@ -2367,12 +2466,17 @@ FGraphEventRef PipelineStateCache::PrecacheGraphicsPipelineState(const FGraphics
 	static bool bDoAsyncCompile = FApp::ShouldUseThreadingForPerformance();
 
 	// try and create new graphics state
-	FGraphicsPipelineState* NewGraphicsPipelineState = GPrecacheGraphicsPipelineCache.TryAddNewState(Initializer, bDoAsyncCompile);
-	if (NewGraphicsPipelineState == nullptr)
+	FGraphEventRef PrecachingEvent;
+	GPrecacheGraphicsPipelineCache.TryAddNewState(Initializer, bDoAsyncCompile, PrecachingEvent);
+	if (PrecachingEvent == nullptr)
 	{
-		return nullptr;
+		check(GPrecacheGraphicsPipelineCache.GetPrecachingState(Initializer) == EPSOPrecacheResult::Complete);
 	}
+	return PrecachingEvent;
+}
 	
+void FPrecacheGraphicsPipelineCache::OnNewPipelineStateCreated(const FGraphicsPipelineStateInitializer & Initializer, FGraphicsPipelineState * NewGraphicsPipelineState, bool bDoAsyncCompile)
+{
 	ValidateGraphicsPipelineStateInitializer(Initializer);
 
 	// Mark as precache so it will try and use the background thread pool if available
@@ -2381,8 +2485,17 @@ FGraphEventRef PipelineStateCache::PrecacheGraphicsPipelineState(const FGraphics
 	FGraphicsPipelineStateInitializer InitializerCopy(Initializer);
 	InitializerCopy.bPSOPrecache = true;
 
+
+	bool bHasCompletionEvent = (NewGraphicsPipelineState->CompletionEvent != nullptr);
+	check(bHasCompletionEvent == bDoAsyncCompile);
+
 	// Start the precache task
-	return CreateGraphicsPipelineState(InitializerCopy, EPSOPrecacheResult::Active, bDoAsyncCompile, bPSOPrecache, NewGraphicsPipelineState);
+	InternalCreateGraphicsPipelineState(InitializerCopy, EPSOPrecacheResult::Active, bDoAsyncCompile, bPSOPrecache, NewGraphicsPipelineState);
+}
+
+FGraphEventRef PipelineStateCache::PrecacheGraphicsPipelineState(const FGraphicsPipelineStateInitializer& Initializer)
+{
+	return GPrecacheGraphicsPipelineCache.PrecacheGraphicsPipelineState(Initializer);
 }
 
 EPSOPrecacheResult PipelineStateCache::CheckPipelineStateInCache(const FGraphicsPipelineStateInitializer& PipelineStateInitializer)
@@ -2392,12 +2505,7 @@ EPSOPrecacheResult PipelineStateCache::CheckPipelineStateInCache(const FGraphics
 		return EPSOPrecacheResult::Unknown;
 	}
 
-	if (!GPrecacheGraphicsPipelineCache.Contains(PipelineStateInitializer))
-	{
-		return EPSOPrecacheResult::Missed;
-	}
-
-	return GPrecacheGraphicsPipelineCache.IsPrecaching(PipelineStateInitializer) ? EPSOPrecacheResult::Active : EPSOPrecacheResult::Complete;
+	return GPrecacheGraphicsPipelineCache.GetPrecachingState(PipelineStateInitializer);
 }
 
 EPSOPrecacheResult PipelineStateCache::CheckPipelineStateInCache(FRHIComputeShader* ComputeShader)
@@ -2407,28 +2515,27 @@ EPSOPrecacheResult PipelineStateCache::CheckPipelineStateInCache(FRHIComputeShad
 		return EPSOPrecacheResult::Unknown;
 	}
 
-	FComputePipelineState* CachedState = nullptr;
-	if (!GComputePipelineCache.Find(ComputeShader, CachedState))
-	{
-		return EPSOPrecacheResult::Missed;
-	}
+	return GPrecacheComputePipelineCache.GetPrecachingState(ComputeShader);
+}
 
-	if (!CachedState->IsComplete())
+bool PipelineStateCache::IsPrecaching(const FGraphicsPipelineStateInitializer& PipelineStateInitializer)
+{
+	if (!IsPSOPrecachingEnabled())
 	{
-		return EPSOPrecacheResult::Active;
+		return false;
 	}
 
-	return CachedState->RHIPipeline && CachedState->RHIPipeline->IsValid() ? EPSOPrecacheResult::Complete : EPSOPrecacheResult::NotSupported;
+	return GPrecacheGraphicsPipelineCache.GetPrecachingState(PipelineStateInitializer) == EPSOPrecacheResult::Active;
 }
 
-bool PipelineStateCache::IsPrecaching(const FGraphicsPipelineStateInitializer& PipelineStateInitializer)
+bool PipelineStateCache::IsPrecaching(FRHIComputeShader* ComputeShader)
 {
 	if (!IsPSOPrecachingEnabled())
 	{
 		return false;
 	}
 
-	return GPrecacheGraphicsPipelineCache.IsPrecaching(PipelineStateInitializer);
+	return GPrecacheComputePipelineCache.GetPrecachingState(ComputeShader) == EPSOPrecacheResult::Active;
 }
 
 bool PipelineStateCache::IsPrecaching()
@@ -2438,7 +2545,7 @@ bool PipelineStateCache::IsPrecaching()
 		return false;
 	}
 
-	return GPrecacheGraphicsPipelineCache.IsPrecaching();
+	return GPrecacheGraphicsPipelineCache.IsPrecaching() || GPrecacheComputePipelineCache.IsPrecaching();
 }
 
 FRHIGraphicsPipelineState* ExecuteSetGraphicsPipelineState(FGraphicsPipelineState* GraphicsPipelineState)
diff --git a/Engine/Source/Runtime/RHI/Public/PipelineStateCache.h b/Engine/Source/Runtime/RHI/Public/PipelineStateCache.h
--- a/Engine/Source/Runtime/RHI/Public/PipelineStateCache.h
+++ b/Engine/Source/Runtime/RHI/Public/PipelineStateCache.h
@@ -112,6 +112,9 @@ namespace PipelineStateCache
 	/* Is the given PSO initializer still precaching? */
 	extern RHI_API bool						IsPrecaching(const FGraphicsPipelineStateInitializer& PipelineStateInitializer);
 
+	/* Is the given PSO initializer still precaching? */
+	extern RHI_API bool						IsPrecaching(FRHIComputeShader* ComputeShader);
+
 	/* Any async precaching operations still busy */
 	extern RHI_API bool						IsPrecaching();
 
