Moving IsCustomPlatform to RHIShaderPlatform.h

diff --git a/Engine/Source/Runtime/RHI/Public/DataDrivenShaderPlatformInfo.h b/Engine/Source/Runtime/RHI/Public/DataDrivenShaderPlatformInfo.h
--- a/Engine/Source/Runtime/RHI/Public/DataDrivenShaderPlatformInfo.h
+++ b/Engine/Source/Runtime/RHI/Public/DataDrivenShaderPlatformInfo.h
@@ -764,10 +764,6 @@ public:
 using FDataDrivenShaderPlatformInfo = FGenericDataDrivenShaderPlatformInfo;
 #endif
 
-
-
-
-
 inline bool IsPCPlatform(const FStaticShaderPlatform Platform)
 {
 	return FDataDrivenShaderPlatformInfo::GetIsPC(Platform);
@@ -779,11 +775,6 @@ inline bool IsMobilePlatform(const FStaticShaderPlatform Platform)
 	return FDataDrivenShaderPlatformInfo::GetMaxFeatureLevel(Platform) == ERHIFeatureLevel::ES3_1;
 }
 
-inline bool IsCustomPlatform(const FStaticShaderPlatform Platform)
-{
-	return (Platform >= EShaderPlatform::SP_CUSTOM_PLATFORM_FIRST) && (Platform < EShaderPlatform::SP_CUSTOM_PLATFORM_LAST);
-}
-
 inline bool IsOpenGLPlatform(const FStaticShaderPlatform Platform)
 {
 	return FDataDrivenShaderPlatformInfo::GetIsLanguageOpenGL(Platform);
diff --git a/Engine/Source/Runtime/RHI/Public/RHIShaderPlatform.h b/Engine/Source/Runtime/RHI/Public/RHIShaderPlatform.h
--- a/Engine/Source/Runtime/RHI/Public/RHIShaderPlatform.h
+++ b/Engine/Source/Runtime/RHI/Public/RHIShaderPlatform.h
@@ -170,3 +170,9 @@ extern RHI_API FString LexToString(EShaderPlatform Platform, bool bError = true)
 
 // The maximum shader platform available on this system
 extern RHI_API EShaderPlatform GMaxRHIShaderPlatform;
+
+inline bool IsCustomPlatform(const FStaticShaderPlatform Platform)
+{
+	return (Platform >= SP_CUSTOM_PLATFORM_FIRST && Platform < SP_CUSTOM_PLATFORM_LAST);
+}
+
