Adding RHIFwd.h for forward declaring common RHI types.

diff --git a/Engine/Source/Runtime/RHI/Public/RHIFwd.h b/Engine/Source/Runtime/RHI/Public/RHIFwd.h
new file mode 100644
--- /dev/null
+++ b/Engine/Source/Runtime/RHI/Public/RHIFwd.h
@@ -0,0 +1,104 @@
+// Copyright Epic Games, Inc. All Rights Reserved.
+
+#pragma once
+
+#include "Templates/RefCounting.h"
+
+// Basic Types
+
+namespace ERHIFeatureLevel { enum Type : int; }
+enum EShaderPlatform : uint16;
+
+// Command Lists
+class FRHICommandListBase;
+class FRHIComputeCommandList;
+class FRHICommandList;
+class FRHICommandListImmediate;
+
+// Resources
+class FRHIAmplificationShader;
+class FRHIBlendState;
+class FRHIBoundShaderState;
+class FRHIBuffer;
+class FRHIComputeFence;
+class FRHIComputePipelineState;
+class FRHIComputeShader;
+class FRHICustomPresent;
+class FRHIDepthStencilState;
+class FRHIGeometryShader;
+class FRHIGPUFence;
+class FRHIGraphicsPipelineState;
+class FRHIMeshShader;
+class FRHIPipelineBinaryLibrary;
+class FRHIPixelShader;
+class FRHIRasterizerState;
+class FRHIRayTracingGeometry;
+class FRHIRayTracingPipelineState;
+class FRHIRayTracingScene;
+class FRHIRayTracingShader;
+class FRHIRenderQuery;
+class FRHIRenderQueryPool;
+class FRHIResource;
+class FRHISamplerState;
+class FRHIShaderLibrary;
+class FRHIShaderResourceView;
+class FRHIStagingBuffer;
+class FRHITexture;
+class FRHITextureReference;
+class FRHITimestampCalibrationQuery;
+class FRHIUniformBuffer;
+class FRHIUnorderedAccessView;
+class FRHIVertexDeclaration;
+class FRHIVertexShader;
+class FRHIViewableResource;
+class FRHIViewport;
+
+struct FRHIUniformBufferLayout;
+
+// Pointers
+
+using FAmplificationShaderRHIRef       = TRefCountPtr<FRHIAmplificationShader>;
+using FBlendStateRHIRef                = TRefCountPtr<FRHIBlendState>;
+using FBoundShaderStateRHIRef          = TRefCountPtr<FRHIBoundShaderState>;
+using FBufferRHIRef                    = TRefCountPtr<FRHIBuffer>;
+using FComputeFenceRHIRef              = TRefCountPtr<FRHIComputeFence>;
+using FComputePipelineStateRHIRef      = TRefCountPtr<FRHIComputePipelineState>;
+using FComputeShaderRHIRef             = TRefCountPtr<FRHIComputeShader>;
+using FCustomPresentRHIRef             = TRefCountPtr<FRHICustomPresent>;
+using FDepthStencilStateRHIRef         = TRefCountPtr<FRHIDepthStencilState>;
+using FGeometryShaderRHIRef            = TRefCountPtr<FRHIGeometryShader>;
+using FGPUFenceRHIRef                  = TRefCountPtr<FRHIGPUFence>;
+using FGraphicsPipelineStateRHIRef     = TRefCountPtr<FRHIGraphicsPipelineState>;
+using FMeshShaderRHIRef                = TRefCountPtr<FRHIMeshShader>;
+using FPixelShaderRHIRef               = TRefCountPtr<FRHIPixelShader>;
+using FRasterizerStateRHIRef           = TRefCountPtr<FRHIRasterizerState>;
+using FRayTracingGeometryRHIRef        = TRefCountPtr<FRHIRayTracingGeometry>;
+using FRayTracingPipelineStateRHIRef   = TRefCountPtr<FRHIRayTracingPipelineState>;
+using FRayTracingSceneRHIRef           = TRefCountPtr<FRHIRayTracingScene>;
+using FRayTracingShaderRHIRef          = TRefCountPtr<FRHIRayTracingShader>;
+using FRenderQueryPoolRHIRef           = TRefCountPtr<FRHIRenderQueryPool>;
+using FRenderQueryRHIRef               = TRefCountPtr<FRHIRenderQuery>;
+using FRHIPipelineBinaryLibraryRef     = TRefCountPtr<FRHIPipelineBinaryLibrary>;
+using FRHIShaderLibraryRef             = TRefCountPtr<FRHIShaderLibrary>;
+using FSamplerStateRHIRef              = TRefCountPtr<FRHISamplerState>;
+using FShaderResourceViewRHIRef        = TRefCountPtr<FRHIShaderResourceView>;
+using FStagingBufferRHIRef             = TRefCountPtr<FRHIStagingBuffer>;
+using FTextureReferenceRHIRef          = TRefCountPtr<FRHITextureReference>;
+using FTextureRHIRef                   = TRefCountPtr<FRHITexture>;
+using FTimestampCalibrationQueryRHIRef = TRefCountPtr<FRHITimestampCalibrationQuery>;
+using FUniformBufferLayoutRHIRef       = TRefCountPtr<const FRHIUniformBufferLayout>;
+using FUniformBufferRHIRef             = TRefCountPtr<FRHIUniformBuffer>;
+using FUnorderedAccessViewRHIRef       = TRefCountPtr<FRHIUnorderedAccessView>;
+using FVertexDeclarationRHIRef         = TRefCountPtr<FRHIVertexDeclaration>;
+using FVertexShaderRHIRef              = TRefCountPtr<FRHIVertexShader>;
+using FViewportRHIRef                  = TRefCountPtr<FRHIViewport>;
+
+// Deprecated typenames
+using FRHITexture2D                    = FRHITexture;
+using FRHITexture2DArray               = FRHITexture;
+using FRHITexture3D                    = FRHITexture;
+using FRHITextureCube                  = FRHITexture;
+using FTexture2DRHIRef                 = FTextureRHIRef;
+using FTexture2DArrayRHIRef            = FTextureRHIRef;
+using FTexture3DRHIRef                 = FTextureRHIRef;
+using FTextureCubeRHIRef               = FTextureRHIRef;
diff --git a/Engine/Source/Runtime/RHI/Public/RHIResources.h b/Engine/Source/Runtime/RHI/Public/RHIResources.h
--- a/Engine/Source/Runtime/RHI/Public/RHIResources.h
+++ b/Engine/Source/Runtime/RHI/Public/RHIResources.h
@@ -25,6 +25,7 @@
 #include "Experimental/Containers/HazardPointer.h"
 #include "Containers/ConsumeAllMpmcQueue.h"
 #include "Misc/CoreDelegates.h"
+#include "RHIFwd.h"
 
 // RHI_WANT_RESOURCE_INFO should be controlled by the RHI module.
 #ifndef RHI_WANT_RESOURCE_INFO
@@ -553,6 +554,7 @@ static bool MatchRHIState(RHIState* LHSState, RHIState* RHSState)
 //
 
 typedef TArray<struct FVertexElement,TFixedAllocator<MaxVertexElementCount> > FVertexDeclarationElementList;
+
 class FRHIVertexDeclaration : public FRHIResource
 {
 public:
@@ -1698,13 +1700,6 @@ struct FRHITextureCreateDesc : public FRHITextureDesc
 	FResourceBulkDataInterface* BulkData = nullptr;
 };
 
-class FRHITexture;
-/*UE_DEPRECATED(5.1, "FRHITexture2D is deprecated, please use FRHITexture.")      */ typedef class FRHITexture FRHITexture2D;
-/*UE_DEPRECATED(5.1, "FRHITexture2DArray is deprecated, please use FRHITexture.") */ typedef class FRHITexture FRHITexture2DArray;
-/*UE_DEPRECATED(5.1, "FRHITexture3D is deprecated, please use FRHITexture.")      */ typedef class FRHITexture FRHITexture3D;
-/*UE_DEPRECATED(5.1, "FRHITextureCube is deprecated, please use FRHITexture.")    */ typedef class FRHITexture FRHITextureCube;
-
-
 class RHI_API FRHITexture : public FRHIViewableResource
 #if ENABLE_RHI_VALIDATION
 	, public RHIValidation::FTextureResource
@@ -2297,41 +2292,6 @@ public:
 	{}
 };
 
-typedef TRefCountPtr<FRHISamplerState> FSamplerStateRHIRef;
-typedef TRefCountPtr<FRHIRasterizerState> FRasterizerStateRHIRef;
-typedef TRefCountPtr<FRHIDepthStencilState> FDepthStencilStateRHIRef;
-typedef TRefCountPtr<FRHIBlendState> FBlendStateRHIRef;
-typedef TRefCountPtr<FRHIVertexDeclaration> FVertexDeclarationRHIRef;
-typedef TRefCountPtr<FRHIVertexShader> FVertexShaderRHIRef;
-typedef TRefCountPtr<FRHIMeshShader> FMeshShaderRHIRef;
-typedef TRefCountPtr<FRHIAmplificationShader> FAmplificationShaderRHIRef;
-typedef TRefCountPtr<FRHIPixelShader> FPixelShaderRHIRef;
-typedef TRefCountPtr<FRHIGeometryShader> FGeometryShaderRHIRef;
-typedef TRefCountPtr<FRHIComputeShader> FComputeShaderRHIRef;
-typedef TRefCountPtr<FRHIRayTracingShader>          FRayTracingShaderRHIRef;
-typedef TRefCountPtr<FRHIComputeFence>	FComputeFenceRHIRef;
-typedef TRefCountPtr<FRHIBoundShaderState> FBoundShaderStateRHIRef;
-typedef TRefCountPtr<const FRHIUniformBufferLayout> FUniformBufferLayoutRHIRef;
-typedef TRefCountPtr<FRHIUniformBuffer> FUniformBufferRHIRef;
-typedef TRefCountPtr<FRHIBuffer> FBufferRHIRef;
-typedef TRefCountPtr<FRHITexture> FTextureRHIRef;
-/*UE_DEPRECATED(5.1, "FTexture2DRHIRef is deprecated, please use FTextureRHIRef.")      */ typedef FTextureRHIRef FTexture2DRHIRef;
-/*UE_DEPRECATED(5.1, "FTexture2DArrayRHIRef is deprecated, please use FTextureRHIRef.") */ typedef FTextureRHIRef FTexture2DArrayRHIRef;
-/*UE_DEPRECATED(5.1, "FTexture3DRHIRef is deprecated, please use FTextureRHIRef.")      */ typedef FTextureRHIRef FTexture3DRHIRef;
-/*UE_DEPRECATED(5.1, "FTextureCubeRHIRef is deprecated, please use FTextureRHIRef.")    */ typedef FTextureRHIRef FTextureCubeRHIRef;
-typedef TRefCountPtr<FRHITextureReference> FTextureReferenceRHIRef;
-typedef TRefCountPtr<FRHIRenderQuery> FRenderQueryRHIRef;
-typedef TRefCountPtr<FRHIRenderQueryPool> FRenderQueryPoolRHIRef;
-typedef TRefCountPtr<FRHITimestampCalibrationQuery> FTimestampCalibrationQueryRHIRef;
-typedef TRefCountPtr<FRHIGPUFence>	FGPUFenceRHIRef;
-typedef TRefCountPtr<FRHIViewport> FViewportRHIRef;
-typedef TRefCountPtr<FRHIUnorderedAccessView> FUnorderedAccessViewRHIRef;
-typedef TRefCountPtr<FRHIShaderResourceView> FShaderResourceViewRHIRef;
-typedef TRefCountPtr<FRHIGraphicsPipelineState> FGraphicsPipelineStateRHIRef;
-typedef TRefCountPtr<FRHIComputePipelineState> FComputePipelineStateRHIRef;
-typedef TRefCountPtr<FRHIRayTracingPipelineState> FRayTracingPipelineStateRHIRef;
-
-
 /**
  * A type used only for printing a string for debugging/profiling.
  * Adds Number as a suffix to the printed string even if the base name includes a number, so may prints a string like: Base_1_1
@@ -2643,8 +2603,6 @@ protected:
 	ERayTracingGeometryInitializerType InitializedType = ERayTracingGeometryInitializerType::Rendering;
 };
 
-typedef TRefCountPtr<FRHIRayTracingGeometry>     FRayTracingGeometryRHIRef;
-
 /** Top level ray tracing acceleration structure (contains instances of meshes). */
 class FRHIRayTracingScene : public FRHIRayTracingAccelerationStructure
 {
@@ -2662,9 +2620,6 @@ public:
 	virtual uint32 GetLayerBufferOffset(uint32 LayerIndex) const = 0;
 };
 
-typedef TRefCountPtr<FRHIRayTracingScene>        FRayTracingSceneRHIRef;
-
-
 /* Generic staging buffer class used by FRHIGPUMemoryReadback
 * RHI specific staging buffers derive from this
 */
@@ -2705,8 +2660,6 @@ public:
 	uint32 Offset;
 };
 
-typedef TRefCountPtr<FRHIStagingBuffer>	FStagingBufferRHIRef;
-
 class FRHIRenderTargetView
 {
 public:
@@ -3027,8 +2980,6 @@ public:
 };
 
 
-typedef TRefCountPtr<FRHICustomPresent> FCustomPresentRHIRef;
-
 // Templates to convert an FRHI*Shader to its enum
 template<typename TRHIShader> struct TRHIShaderToEnum {};
 template<> struct TRHIShaderToEnum<FRHIVertexShader>           { enum { ShaderFrequency = SF_Vertex        }; };
@@ -3673,8 +3624,6 @@ protected:
 	uint32 LibraryId;
 };
 
-typedef TRefCountPtr<FRHIShaderLibrary>	FRHIShaderLibraryRef;
-
 class FRHIPipelineBinaryLibrary : public FRHIResource
 {
 public:
@@ -3687,8 +3636,6 @@ protected:
 	EShaderPlatform Platform;
 };
 
-typedef TRefCountPtr<FRHIPipelineBinaryLibrary>	FRHIPipelineBinaryLibraryRef;
-
 enum class ERenderTargetActions : uint8
 {
 	LoadOpMask = 2,
