TRHIResourceUpdateBatcher cleanup
- Moving TRHIResourceUpdateBatcher and FRHIResourceUpdateInfo into their own file, RHIResourceUpdates
- TRHIResourceUpdateBatcher is now only the data class and has a base class to manage the functionality.
- The new base class can be used generically, removing the need for templates on all prior uses of TRHIResourceUpdateBatcher
- Changing all functional uses of TRHIResourceUpdateBatcher to use the base class and moving all implementations into cpp files.

diff --git a/Engine/Source/Runtime/Engine/Private/Animation/SkinWeightProfile.cpp b/Engine/Source/Runtime/Engine/Private/Animation/SkinWeightProfile.cpp
--- a/Engine/Source/Runtime/Engine/Private/Animation/SkinWeightProfile.cpp
+++ b/Engine/Source/Runtime/Engine/Private/Animation/SkinWeightProfile.cpp
@@ -779,6 +779,24 @@ void FSkinWeightProfilesData::CreateRHIBuffers_Async(TArray<TPair<FName, FSkinWe
 	CreateRHIBuffers_Internal<false>(OutBuffers);
 }
 
+void FSkinWeightProfilesData::InitRHIForStreaming(const TArray<TPair<FName, FSkinWeightRHIInfo>>& IntermediateBuffers, FRHIResourceUpdateBatcher& Batcher)
+{
+	for (int32 Idx = 0; Idx < IntermediateBuffers.Num(); ++Idx)
+	{
+		const FName& ProfileName = IntermediateBuffers[Idx].Key;
+		const FSkinWeightRHIInfo& IntermediateBuffer = IntermediateBuffers[Idx].Value;
+		ProfileNameToBuffer.FindChecked(ProfileName)->InitRHIForStreaming(IntermediateBuffer, Batcher);
+	}
+}
+
+void FSkinWeightProfilesData::ReleaseRHIForStreaming(FRHIResourceUpdateBatcher& Batcher)
+{
+	for (TMap<FName, FSkinWeightVertexBuffer*>::TIterator It(ProfileNameToBuffer); It; ++It)
+	{
+		It->Value->ReleaseRHIForStreaming(Batcher);
+	}
+}
+
 void FRuntimeSkinWeightProfileData::ApplyOverrides(FSkinWeightVertexBuffer* OverrideBuffer, const void* DataBuffer, const int32 NumVerts) const
 {
 	if (DataBuffer)
diff --git a/Engine/Source/Runtime/Engine/Private/PositionVertexBuffer.cpp b/Engine/Source/Runtime/Engine/Private/PositionVertexBuffer.cpp
--- a/Engine/Source/Runtime/Engine/Private/PositionVertexBuffer.cpp
+++ b/Engine/Source/Runtime/Engine/Private/PositionVertexBuffer.cpp
@@ -8,6 +8,7 @@
 
 #include "StaticMeshVertexData.h"
 #include "GPUSkinCache.h"
+#include "RHIResourceUpdates.h"
 
 /*-----------------------------------------------------------------------------
 FPositionVertexBuffer
@@ -232,6 +233,29 @@ void FPositionVertexBuffer::CopyRHIForStreaming(const FPositionVertexBuffer& Oth
 	PositionComponentSRV = Other.PositionComponentSRV;
 }
 
+void FPositionVertexBuffer::InitRHIForStreaming(FRHIBuffer* IntermediateBuffer, FRHIResourceUpdateBatcher& Batcher)
+{
+	check(VertexBufferRHI);
+	if (IntermediateBuffer)
+	{
+		Batcher.QueueUpdateRequest(VertexBufferRHI, IntermediateBuffer);
+		if (PositionComponentSRV)
+		{
+			Batcher.QueueUpdateRequest(PositionComponentSRV, VertexBufferRHI, 4, PF_R32_FLOAT);
+		}
+	}
+}
+
+void FPositionVertexBuffer::ReleaseRHIForStreaming(FRHIResourceUpdateBatcher& Batcher)
+{
+	check(VertexBufferRHI);
+	Batcher.QueueUpdateRequest(VertexBufferRHI, nullptr);
+	if (PositionComponentSRV)
+	{
+		Batcher.QueueUpdateRequest(PositionComponentSRV, nullptr, 0, 0);
+	}
+}
+
 void FPositionVertexBuffer::InitRHI()
 {
 	TRACE_CPUPROFILER_EVENT_SCOPE(FPositionVertexBuffer::InitRHI);
diff --git a/Engine/Source/Runtime/Engine/Private/RawIndexBuffer.cpp b/Engine/Source/Runtime/Engine/Private/RawIndexBuffer.cpp
--- a/Engine/Source/Runtime/Engine/Private/RawIndexBuffer.cpp
+++ b/Engine/Source/Runtime/Engine/Private/RawIndexBuffer.cpp
@@ -9,6 +9,7 @@
 #include "Modules/ModuleManager.h"
 #include "RenderingThread.h"
 #include "DataDrivenShaderPlatformInfo.h"
+#include "RHIResourceUpdates.h"
 
 #if WITH_EDITOR
 #include "MeshUtilities.h"
@@ -381,6 +382,22 @@ void FRawStaticIndexBuffer::CopyRHIForStreaming(const FRawStaticIndexBuffer& Oth
 	IndexBufferRHI = Other.IndexBufferRHI;
 }
 
+void FRawStaticIndexBuffer::InitRHIForStreaming(FRHIBuffer* IntermediateBuffer, FRHIResourceUpdateBatcher& Batcher)
+{
+	if (IndexBufferRHI && IntermediateBuffer)
+	{
+		Batcher.QueueUpdateRequest(IndexBufferRHI, IntermediateBuffer);
+	}
+}
+
+void FRawStaticIndexBuffer::ReleaseRHIForStreaming(FRHIResourceUpdateBatcher& Batcher)
+{
+	if (IndexBufferRHI)
+	{
+		Batcher.QueueUpdateRequest(IndexBufferRHI, nullptr);
+	}
+}
+
 void FRawStaticIndexBuffer::InitRHI()
 {
 	TRACE_CPUPROFILER_EVENT_SCOPE(FRawStaticIndexBuffer::InitRHI);
@@ -444,6 +461,30 @@ void FRawStaticIndexBuffer::Discard()
     IndexStorage.Discard();
 }
 
+void FRawStaticIndexBuffer16or32Interface::InitRHIForStreaming(FRHIBuffer* IntermediateBuffer, size_t IndexSize, FRHIResourceUpdateBatcher& Batcher)
+{
+	if (IndexBufferRHI && IntermediateBuffer)
+	{
+		Batcher.QueueUpdateRequest(IndexBufferRHI, IntermediateBuffer);
+		if (SRVValue)
+		{
+			Batcher.QueueUpdateRequest(SRVValue, IndexBufferRHI, IndexSize, IndexSize == 2 ? PF_R16_UINT : PF_R32_UINT);
+		}
+	}
+}
+
+void FRawStaticIndexBuffer16or32Interface::ReleaseRHIForStreaming(FRHIResourceUpdateBatcher& Batcher)
+{
+	if (IndexBufferRHI)
+	{
+		Batcher.QueueUpdateRequest(IndexBufferRHI, nullptr);
+	}
+	if (SRVValue)
+	{
+		Batcher.QueueUpdateRequest(SRVValue, nullptr, 0, 0);
+	}
+}
+
 FBufferRHIRef FRawStaticIndexBuffer16or32Interface::CreateRHIIndexBufferInternal(
 	const TCHAR* InDebugName,
 	const FName& InOwnerName,
diff --git a/Engine/Source/Runtime/Engine/Private/Rendering/ColorVertexBuffer.cpp b/Engine/Source/Runtime/Engine/Private/Rendering/ColorVertexBuffer.cpp
--- a/Engine/Source/Runtime/Engine/Private/Rendering/ColorVertexBuffer.cpp
+++ b/Engine/Source/Runtime/Engine/Private/Rendering/ColorVertexBuffer.cpp
@@ -7,6 +7,7 @@
 #include "ProfilingDebugging/LoadTimeTracker.h"
 #include "VertexFactory.h"
 #include "DataDrivenShaderPlatformInfo.h"
+#include "RHIResourceUpdates.h"
 
 /*-----------------------------------------------------------------------------
 FColorVertexBuffer
@@ -430,6 +431,30 @@ void FColorVertexBuffer::CopyRHIForStreaming(const FColorVertexBuffer& Other, bo
 	ColorComponentsSRV = Other.ColorComponentsSRV;
 }
 
+void FColorVertexBuffer::InitRHIForStreaming(FRHIBuffer* IntermediateBuffer, FRHIResourceUpdateBatcher& Batcher)
+{
+	if (VertexBufferRHI && IntermediateBuffer)
+	{
+		Batcher.QueueUpdateRequest(VertexBufferRHI, IntermediateBuffer);
+		if (ColorComponentsSRV)
+		{
+			Batcher.QueueUpdateRequest(ColorComponentsSRV, VertexBufferRHI, 4, PF_R8G8B8A8);
+		}
+	}
+}
+
+void FColorVertexBuffer::ReleaseRHIForStreaming(FRHIResourceUpdateBatcher& Batcher)
+{
+	if (VertexBufferRHI)
+	{
+		Batcher.QueueUpdateRequest(VertexBufferRHI, nullptr);
+	}
+	if (ColorComponentsSRV)
+	{
+		Batcher.QueueUpdateRequest(ColorComponentsSRV, nullptr, 0, 0);
+	}
+}
+
 void FColorVertexBuffer::InitRHI()
 {
 	TRACE_CPUPROFILER_EVENT_SCOPE(FColorVertexBuffer::InitRHI);
diff --git a/Engine/Source/Runtime/Engine/Private/Rendering/MultiSizeIndexContainer.cpp b/Engine/Source/Runtime/Engine/Private/Rendering/MultiSizeIndexContainer.cpp
--- a/Engine/Source/Runtime/Engine/Private/Rendering/MultiSizeIndexContainer.cpp
+++ b/Engine/Source/Runtime/Engine/Private/Rendering/MultiSizeIndexContainer.cpp
@@ -223,6 +223,37 @@ FBufferRHIRef FMultiSizeIndexContainer::CreateRHIBuffer_Async()
 	return nullptr;
 }
 
+void FMultiSizeIndexContainer::InitRHIForStreaming(FRHIBuffer* IntermediateBuffer, FRHIResourceUpdateBatcher& Batcher)
+{
+	check(!((uint32)!!IntermediateBuffer ^ (uint32)!!IndexBuffer));
+	if (IntermediateBuffer)
+	{
+		if (DataTypeSize == sizeof(uint16))
+		{
+			static_cast<FRawStaticIndexBuffer16or32<uint16>*>(IndexBuffer)->InitRHIForStreaming(IntermediateBuffer, Batcher);
+		}
+		else
+		{
+			static_cast<FRawStaticIndexBuffer16or32<uint32>*>(IndexBuffer)->InitRHIForStreaming(IntermediateBuffer, Batcher);
+		}
+	}
+}
+
+void FMultiSizeIndexContainer::ReleaseRHIForStreaming(FRHIResourceUpdateBatcher& Batcher)
+{
+	if (IndexBuffer)
+	{
+		if (DataTypeSize == sizeof(uint16))
+		{
+			static_cast<FRawStaticIndexBuffer16or32<uint16>*>(IndexBuffer)->ReleaseRHIForStreaming(Batcher);
+		}
+		else
+		{
+			static_cast<FRawStaticIndexBuffer16or32<uint32>*>(IndexBuffer)->ReleaseRHIForStreaming(Batcher);
+		}
+	}
+}
+
 #if WITH_EDITOR
 /**
 * Retrieves index buffer related data
diff --git a/Engine/Source/Runtime/Engine/Private/Rendering/SkeletalMeshVertexClothBuffer.cpp b/Engine/Source/Runtime/Engine/Private/Rendering/SkeletalMeshVertexClothBuffer.cpp
--- a/Engine/Source/Runtime/Engine/Private/Rendering/SkeletalMeshVertexClothBuffer.cpp
+++ b/Engine/Source/Runtime/Engine/Private/Rendering/SkeletalMeshVertexClothBuffer.cpp
@@ -4,6 +4,7 @@
 #include "Rendering/SkeletalMeshVertexBuffer.h"
 #include "EngineUtils.h"
 #include "ProfilingDebugging/LoadTimeTracker.h"
+#include "RHIResourceUpdates.h"
 
 /**
 * Constructor
@@ -87,6 +88,28 @@ FBufferRHIRef FSkeletalMeshVertexClothBuffer::CreateRHIBuffer_Async()
 	return CreateRHIBuffer_Internal<false>();
 }
 
+void FSkeletalMeshVertexClothBuffer::InitRHIForStreaming(FRHIBuffer* IntermediateBuffer, FRHIResourceUpdateBatcher& Batcher)
+{
+	if (VertexBufferRHI && IntermediateBuffer)
+	{
+		check(VertexBufferSRV);
+		Batcher.QueueUpdateRequest(VertexBufferRHI, IntermediateBuffer);
+		Batcher.QueueUpdateRequest(VertexBufferSRV, VertexBufferRHI, sizeof(FVector4f), PF_A32B32G32R32F);
+	}
+}
+
+void FSkeletalMeshVertexClothBuffer::ReleaseRHIForStreaming(FRHIResourceUpdateBatcher& Batcher)
+{
+	if (VertexBufferRHI)
+	{
+		Batcher.QueueUpdateRequest(VertexBufferRHI, nullptr);
+	}
+	if (VertexBufferSRV)
+	{
+		Batcher.QueueUpdateRequest(VertexBufferSRV, nullptr, 0, 0);
+	}
+}
+
 /**
 * Initialize the RHI resource for this vertex buffer
 */
diff --git a/Engine/Source/Runtime/Engine/Private/Rendering/SkinWeightVertexBuffer.cpp b/Engine/Source/Runtime/Engine/Private/Rendering/SkinWeightVertexBuffer.cpp
--- a/Engine/Source/Runtime/Engine/Private/Rendering/SkinWeightVertexBuffer.cpp
+++ b/Engine/Source/Runtime/Engine/Private/Rendering/SkinWeightVertexBuffer.cpp
@@ -9,6 +9,7 @@
 #include "SkeletalMeshTypes.h"
 #include "UObject/AnimObjectVersion.h"
 #include "UObject/UE5MainStreamObjectVersion.h"
+#include "RHIResourceUpdates.h"
 
 /*-----------------------------------------------------------------------------
 FSkinWeightLookupVertexBuffer
@@ -163,6 +164,28 @@ FBufferRHIRef FSkinWeightLookupVertexBuffer::CreateRHIBuffer_Async()
 	return CreateRHIBuffer_Internal<false>();
 }
 
+void FSkinWeightLookupVertexBuffer::InitRHIForStreaming(FRHIBuffer* IntermediateBuffer, FRHIResourceUpdateBatcher& Batcher)
+{
+	if (VertexBufferRHI && IntermediateBuffer)
+	{
+		check(SRVValue);
+		Batcher.QueueUpdateRequest(VertexBufferRHI, IntermediateBuffer);
+		Batcher.QueueUpdateRequest(SRVValue, VertexBufferRHI, PixelFormatStride, PixelFormat);
+	}
+}
+
+void FSkinWeightLookupVertexBuffer::ReleaseRHIForStreaming(FRHIResourceUpdateBatcher& Batcher)
+{
+	if (VertexBufferRHI)
+	{
+		Batcher.QueueUpdateRequest(VertexBufferRHI, nullptr);
+	}
+	if (SRVValue)
+	{
+		Batcher.QueueUpdateRequest(SRVValue, nullptr, 0, 0);
+	}
+}
+
 void FSkinWeightLookupVertexBuffer::GetWeightOffsetAndInfluenceCount(uint32 VertexIndex, uint32& OutWeightOffset, uint32& OutInfluenceCount) const
 {
 	uint32 Offset = VertexIndex * 4;
@@ -396,6 +419,28 @@ FBufferRHIRef FSkinWeightDataVertexBuffer::CreateRHIBuffer_Async()
 	return CreateRHIBuffer_Internal<false>();
 }
 
+void FSkinWeightDataVertexBuffer::InitRHIForStreaming(FRHIBuffer* IntermediateBuffer, FRHIResourceUpdateBatcher& Batcher)
+{
+	if (VertexBufferRHI && IntermediateBuffer)
+	{
+		check(SRVValue);
+		Batcher.QueueUpdateRequest(VertexBufferRHI, IntermediateBuffer);
+		Batcher.QueueUpdateRequest(SRVValue, VertexBufferRHI, GetPixelFormatStride(), GetPixelFormat());
+	}
+}
+
+void FSkinWeightDataVertexBuffer::ReleaseRHIForStreaming(FRHIResourceUpdateBatcher& Batcher)
+{
+	if (VertexBufferRHI)
+	{
+		Batcher.QueueUpdateRequest(VertexBufferRHI, nullptr);
+	}
+	if (SRVValue)
+	{
+		Batcher.QueueUpdateRequest(SRVValue, nullptr, 0, 0);
+	}
+}
+
 bool FSkinWeightDataVertexBuffer::IsWeightDataValid() const
 {
 	return WeightData != nullptr;
diff --git a/Engine/Source/Runtime/Engine/Private/Rendering/StaticMeshVertexBuffer.cpp b/Engine/Source/Runtime/Engine/Private/Rendering/StaticMeshVertexBuffer.cpp
--- a/Engine/Source/Runtime/Engine/Private/Rendering/StaticMeshVertexBuffer.cpp
+++ b/Engine/Source/Runtime/Engine/Private/Rendering/StaticMeshVertexBuffer.cpp
@@ -5,6 +5,7 @@
 #include "Components.h"
 #include "GPUSkinCache.h"
 #include "ProfilingDebugging/LoadTimeTracker.h"
+#include "RHIResourceUpdates.h"
 
 FStaticMeshVertexBuffer::FStaticMeshVertexBuffer() :
 	TangentsData(nullptr),
@@ -326,6 +327,53 @@ void FStaticMeshVertexBuffer::CopyRHIForStreaming(const FStaticMeshVertexBuffer&
 	TextureCoordinatesSRV = Other.TextureCoordinatesSRV;
 }
 
+void FStaticMeshVertexBuffer::InitRHIForStreaming(
+	FRHIBuffer* IntermediateTangentsBuffer,
+	FRHIBuffer* IntermediateTexCoordBuffer,
+	FRHIResourceUpdateBatcher& Batcher)
+{
+	check(TangentsVertexBuffer.VertexBufferRHI && TexCoordVertexBuffer.VertexBufferRHI);
+	if (IntermediateTangentsBuffer)
+	{
+		Batcher.QueueUpdateRequest(TangentsVertexBuffer.VertexBufferRHI, IntermediateTangentsBuffer);
+		if (TangentsSRV)
+		{
+			Batcher.QueueUpdateRequest(
+				TangentsSRV,
+				TangentsVertexBuffer.VertexBufferRHI,
+				GetUseHighPrecisionTangentBasis() ? 8u : 4u,
+				GetUseHighPrecisionTangentBasis() ? (uint8)PF_R16G16B16A16_SNORM : (uint8)PF_R8G8B8A8_SNORM);;
+		}
+	}
+	if (IntermediateTexCoordBuffer)
+	{
+		Batcher.QueueUpdateRequest(TexCoordVertexBuffer.VertexBufferRHI, IntermediateTexCoordBuffer);
+		if (TextureCoordinatesSRV)
+		{
+			Batcher.QueueUpdateRequest(
+				TextureCoordinatesSRV,
+				TexCoordVertexBuffer.VertexBufferRHI,
+				GetUseFullPrecisionUVs() ? 8u : 4u,
+				GetUseFullPrecisionUVs() ? (uint8)PF_G32R32F : (uint8)PF_G16R16F);
+		}
+	}
+}
+
+void FStaticMeshVertexBuffer::ReleaseRHIForStreaming(FRHIResourceUpdateBatcher& Batcher)
+{
+	check(TangentsVertexBuffer.VertexBufferRHI && TexCoordVertexBuffer.VertexBufferRHI);
+	Batcher.QueueUpdateRequest(TangentsVertexBuffer.VertexBufferRHI, nullptr);
+	Batcher.QueueUpdateRequest(TexCoordVertexBuffer.VertexBufferRHI, nullptr);
+	if (TangentsSRV)
+	{
+		Batcher.QueueUpdateRequest(TangentsSRV, nullptr, 0, 0);
+	}
+	if (TextureCoordinatesSRV)
+	{
+		Batcher.QueueUpdateRequest(TextureCoordinatesSRV, nullptr, 0, 0);
+	}
+}
+
 void FStaticMeshVertexBuffer::InitRHI()
 {
 	TRACE_CPUPROFILER_EVENT_SCOPE(FStaticMeshVertexBuffer::InitRHI);
diff --git a/Engine/Source/Runtime/Engine/Private/Streaming/SkeletalMeshUpdate.cpp b/Engine/Source/Runtime/Engine/Private/Streaming/SkeletalMeshUpdate.cpp
--- a/Engine/Source/Runtime/Engine/Private/Streaming/SkeletalMeshUpdate.cpp
+++ b/Engine/Source/Runtime/Engine/Private/Streaming/SkeletalMeshUpdate.cpp
@@ -14,6 +14,7 @@ SkeletalMeshUpdate.cpp: Helpers to stream in and out skeletal mesh LODs.
 #include "Rendering/SkeletalMeshRenderData.h"
 #include "Components/SkinnedMeshComponent.h"
 #include "Streaming/RenderAssetUpdate.inl"
+#include "RHIResourceUpdates.h"
 
 int32 GStreamingSkeletalMeshIOPriority = (int32)AIOP_Low;
 static FAutoConsoleVariableRef CVarStreamingSkeletalMeshIOPriority(
@@ -95,8 +96,7 @@ void FSkeletalMeshStreamIn::FIntermediateBuffers::SafeRelease()
 	AltSkinWeightVertexBuffers.Empty();
 }
 
-template <uint32 MaxNumUpdates>
-void FSkeletalMeshStreamIn::FIntermediateBuffers::TransferBuffers(FSkeletalMeshLODRenderData& LODResource, TRHIResourceUpdateBatcher<MaxNumUpdates>& Batcher)
+void FSkeletalMeshStreamIn::FIntermediateBuffers::TransferBuffers(FSkeletalMeshLODRenderData& LODResource, FRHIResourceUpdateBatcher& Batcher)
 {
 	FStaticMeshVertexBuffers& VBs = LODResource.StaticVertexBuffers;
 	VBs.StaticMeshVertexBuffer.InitRHIForStreaming(TangentsVertexBuffer, TexCoordVertexBuffer, Batcher);
diff --git a/Engine/Source/Runtime/Engine/Private/Streaming/SkeletalMeshUpdate.h b/Engine/Source/Runtime/Engine/Private/Streaming/SkeletalMeshUpdate.h
--- a/Engine/Source/Runtime/Engine/Private/Streaming/SkeletalMeshUpdate.h
+++ b/Engine/Source/Runtime/Engine/Private/Streaming/SkeletalMeshUpdate.h
@@ -90,8 +90,7 @@ protected:
 		void SafeRelease();
 
 		/** Transfer ownership of buffers to a LOD resource */
-		template <uint32 MaxNumUpdates>
-		void TransferBuffers(FSkeletalMeshLODRenderData& LODResource, TRHIResourceUpdateBatcher<MaxNumUpdates>& Batcher);
+		void TransferBuffers(FSkeletalMeshLODRenderData& LODResource, FRHIResourceUpdateBatcher& Batcher);
 
 		void CheckIsNull() const;
 	};
diff --git a/Engine/Source/Runtime/Engine/Private/Streaming/StaticMeshUpdate.cpp b/Engine/Source/Runtime/Engine/Private/Streaming/StaticMeshUpdate.cpp
--- a/Engine/Source/Runtime/Engine/Private/Streaming/StaticMeshUpdate.cpp
+++ b/Engine/Source/Runtime/Engine/Private/Streaming/StaticMeshUpdate.cpp
@@ -13,6 +13,7 @@ StaticMeshUpdate.cpp: Helpers to stream in and out static mesh LODs.
 #include "StaticMeshResources.h"
 #include "Streaming/RenderAssetUpdate.inl"
 #include "ContentStreaming.h"
+#include "RHIResourceUpdates.h"
 
 int32 GStreamingMaxReferenceChecks = 2;
 static FAutoConsoleVariableRef CVarStreamingMaxReferenceChecksBeforeStreamOut(
@@ -112,8 +113,7 @@ void FStaticMeshStreamIn::FIntermediateBuffers::SafeRelease()
 	WireframeIndexBuffer.SafeRelease();
 }
 
-template <uint32 MaxNumUpdates>
-void FStaticMeshStreamIn::FIntermediateBuffers::TransferBuffers(FStaticMeshLODResources& LODResource, TRHIResourceUpdateBatcher<MaxNumUpdates>& Batcher)
+void FStaticMeshStreamIn::FIntermediateBuffers::TransferBuffers(FStaticMeshLODResources& LODResource, FRHIResourceUpdateBatcher& Batcher)
 {
 	FStaticMeshVertexBuffers& VBs = LODResource.VertexBuffers;
 	VBs.StaticMeshVertexBuffer.InitRHIForStreaming(TangentsVertexBuffer, TexCoordVertexBuffer, Batcher);
diff --git a/Engine/Source/Runtime/Engine/Private/Streaming/StaticMeshUpdate.h b/Engine/Source/Runtime/Engine/Private/Streaming/StaticMeshUpdate.h
--- a/Engine/Source/Runtime/Engine/Private/Streaming/StaticMeshUpdate.h
+++ b/Engine/Source/Runtime/Engine/Private/Streaming/StaticMeshUpdate.h
@@ -92,8 +92,7 @@ protected:
 		void SafeRelease();
 
 		/** Transfer ownership of buffers to a LOD resource */
-		template <uint32 MaxNumUpdates>
-		void TransferBuffers(FStaticMeshLODResources& LODResource, TRHIResourceUpdateBatcher<MaxNumUpdates>& Batcher);
+		void TransferBuffers(FStaticMeshLODResources& LODResource, FRHIResourceUpdateBatcher& Batcher);
 
 		void CheckIsNull() const;
 	};
diff --git a/Engine/Source/Runtime/Engine/Public/Animation/SkinWeightProfile.h b/Engine/Source/Runtime/Engine/Public/Animation/SkinWeightProfile.h
--- a/Engine/Source/Runtime/Engine/Public/Animation/SkinWeightProfile.h
+++ b/Engine/Source/Runtime/Engine/Public/Animation/SkinWeightProfile.h
@@ -166,25 +166,8 @@ struct ENGINE_API FSkinWeightProfilesData
 	void CreateRHIBuffers_RenderThread(TArray<TPair<FName, FSkinWeightRHIInfo>>& OutBuffers);
 	void CreateRHIBuffers_Async(TArray<TPair<FName, FSkinWeightRHIInfo>>& OutBuffers);
 
-	template <uint32 MaxNumUpdates>
-	void InitRHIForStreaming(const TArray<TPair<FName, FSkinWeightRHIInfo>>& IntermediateBuffers, TRHIResourceUpdateBatcher<MaxNumUpdates>& Batcher)
-	{
-		for (int32 Idx = 0; Idx < IntermediateBuffers.Num(); ++Idx)
-		{
-			const FName& ProfileName = IntermediateBuffers[Idx].Key;
-			const FSkinWeightRHIInfo& IntermediateBuffer = IntermediateBuffers[Idx].Value;
-			ProfileNameToBuffer.FindChecked(ProfileName)->InitRHIForStreaming(IntermediateBuffer, Batcher);
-		}
-	}
-	
-	template <uint32 MaxNumUpdates>
-	void ReleaseRHIForStreaming(TRHIResourceUpdateBatcher<MaxNumUpdates>& Batcher)
-	{
-		for (TMap<FName, FSkinWeightVertexBuffer*>::TIterator It(ProfileNameToBuffer); It; ++It)
-		{
-			It->Value->ReleaseRHIForStreaming(Batcher);
-		}
-	}
+	void InitRHIForStreaming(const TArray<TPair<FName, FSkinWeightRHIInfo>>& IntermediateBuffers, FRHIResourceUpdateBatcher& Batcher);
+	void ReleaseRHIForStreaming(FRHIResourceUpdateBatcher& Batcher);
 
 	bool IsPendingReadback() const;
 	void EnqueueGPUReadback();
diff --git a/Engine/Source/Runtime/Engine/Public/RawIndexBuffer.h b/Engine/Source/Runtime/Engine/Public/RawIndexBuffer.h
--- a/Engine/Source/Runtime/Engine/Public/RawIndexBuffer.h
+++ b/Engine/Source/Runtime/Engine/Public/RawIndexBuffer.h
@@ -282,24 +282,10 @@ public:
 	void CopyRHIForStreaming(const FRawStaticIndexBuffer& Other, bool InAllowCPUAccess);
 
 	/** Take over ownership of IntermediateBuffer */
-	template <uint32 MaxNumUpdates>
-	void InitRHIForStreaming(FRHIBuffer* IntermediateBuffer, TRHIResourceUpdateBatcher<MaxNumUpdates>& Batcher)
-	{
-		if (IndexBufferRHI && IntermediateBuffer)
-		{
-			Batcher.QueueUpdateRequest(IndexBufferRHI, IntermediateBuffer);
-		}
-	}
+	void InitRHIForStreaming(FRHIBuffer* IntermediateBuffer, FRHIResourceUpdateBatcher& Batcher);
 
 	/** Release any GPU resource owned by the RHI object */
-	template <uint32 MaxNumUpdates>
-	void ReleaseRHIForStreaming(TRHIResourceUpdateBatcher<MaxNumUpdates>& Batcher)
-	{
-		if (IndexBufferRHI)
-		{
-			Batcher.QueueUpdateRequest(IndexBufferRHI, nullptr);
-		}
-	}
+	void ReleaseRHIForStreaming(FRHIResourceUpdateBatcher& Batcher);
 
 	/**
 	 * Serialization.
@@ -374,6 +360,10 @@ public:
 	}
 
 protected:
+	/** Similar to Init/ReleaseRHI but only update existing SRV so references to the SRV stays valid */
+	void InitRHIForStreaming(FRHIBuffer* IntermediateBuffer, size_t IndexSize, FRHIResourceUpdateBatcher& Batcher);
+	void ReleaseRHIForStreaming(FRHIResourceUpdateBatcher& Batcher);
+
 	static ENGINE_API FBufferRHIRef CreateRHIIndexBufferInternal(
 		const TCHAR* InDebugName,
 		const FName& InOwnerName,
@@ -511,30 +501,14 @@ public:
 	FBufferRHIRef CreateRHIBuffer_Async() { return CreateRHIBuffer_Internal<false>(); }
 
 	/** Similar to Init/ReleaseRHI but only update existing SRV so references to the SRV stays valid */
-	template <uint32 MaxNumUpdates>
-	void InitRHIForStreaming(FRHIBuffer* IntermediateBuffer, TRHIResourceUpdateBatcher<MaxNumUpdates>& Batcher)
+	void InitRHIForStreaming(FRHIBuffer* IntermediateBuffer, FRHIResourceUpdateBatcher& Batcher)
 	{
-		if (IndexBufferRHI && IntermediateBuffer)
-		{
-			Batcher.QueueUpdateRequest(IndexBufferRHI, IntermediateBuffer);
-			if (SRVValue)
-			{
-				Batcher.QueueUpdateRequest(SRVValue, IndexBufferRHI, sizeof(INDEX_TYPE), sizeof(INDEX_TYPE) == 2 ? PF_R16_UINT : PF_R32_UINT);
-			}
-		}
+		FRawStaticIndexBuffer16or32Interface::InitRHIForStreaming(IntermediateBuffer, sizeof(INDEX_TYPE), Batcher);
 	}
 
-	template <uint32 MaxNumUpdates>
-	void ReleaseRHIForStreaming(TRHIResourceUpdateBatcher<MaxNumUpdates>& Batcher)
+	void ReleaseRHIForStreaming(FRHIResourceUpdateBatcher& Batcher)
 	{
-		if (IndexBufferRHI)
-		{
-			Batcher.QueueUpdateRequest(IndexBufferRHI, nullptr);
-		}
-		if (SRVValue)
-		{
-			Batcher.QueueUpdateRequest(SRVValue, nullptr, 0, 0);
-		}
+		FRawStaticIndexBuffer16or32Interface::ReleaseRHIForStreaming(Batcher);
 	}
 
 private:
diff --git a/Engine/Source/Runtime/Engine/Public/Rendering/ColorVertexBuffer.h b/Engine/Source/Runtime/Engine/Public/Rendering/ColorVertexBuffer.h
--- a/Engine/Source/Runtime/Engine/Public/Rendering/ColorVertexBuffer.h
+++ b/Engine/Source/Runtime/Engine/Public/Rendering/ColorVertexBuffer.h
@@ -147,31 +147,8 @@ public:
 	void CopyRHIForStreaming(const FColorVertexBuffer& Other, bool InAllowCPUAccess);
 
 	/** Similar to Init/ReleaseRHI but only update existing SRV so references to the SRV stays valid */
-	template <uint32 MaxNumUpdates>
-	void InitRHIForStreaming(FRHIBuffer* IntermediateBuffer, TRHIResourceUpdateBatcher<MaxNumUpdates>& Batcher)
-	{
-		if (VertexBufferRHI && IntermediateBuffer)
-		{
-			Batcher.QueueUpdateRequest(VertexBufferRHI, IntermediateBuffer);
-			if (ColorComponentsSRV)
-			{
-				Batcher.QueueUpdateRequest(ColorComponentsSRV, VertexBufferRHI, 4, PF_R8G8B8A8);
-			}
-		}
-	}
-
-	template <uint32 MaxNumUpdates>
-	void ReleaseRHIForStreaming(TRHIResourceUpdateBatcher<MaxNumUpdates>& Batcher)
-	{
-		if (VertexBufferRHI)
-		{
-			Batcher.QueueUpdateRequest(VertexBufferRHI, nullptr);
-		}
-		if (ColorComponentsSRV)
-		{
-			Batcher.QueueUpdateRequest(ColorComponentsSRV, nullptr, 0, 0);
-		}
-	}
+	void InitRHIForStreaming(FRHIBuffer* IntermediateBuffer, FRHIResourceUpdateBatcher& Batcher);
+	void ReleaseRHIForStreaming(FRHIResourceUpdateBatcher& Batcher);
 
 	// FRenderResource interface.
 	ENGINE_API virtual void InitRHI() override;
diff --git a/Engine/Source/Runtime/Engine/Public/Rendering/MultiSizeIndexContainer.h b/Engine/Source/Runtime/Engine/Public/Rendering/MultiSizeIndexContainer.h
--- a/Engine/Source/Runtime/Engine/Public/Rendering/MultiSizeIndexContainer.h
+++ b/Engine/Source/Runtime/Engine/Public/Rendering/MultiSizeIndexContainer.h
@@ -102,38 +102,8 @@ public:
 	FBufferRHIRef CreateRHIBuffer_RenderThread();
 	FBufferRHIRef CreateRHIBuffer_Async();
 
-	template <uint32 MaxNumUpdates>
-	void InitRHIForStreaming(FRHIBuffer* IntermediateBuffer, TRHIResourceUpdateBatcher<MaxNumUpdates>& Batcher)
-	{
-		check(!((uint32)!!IntermediateBuffer ^ (uint32)!!IndexBuffer));
-		if (IntermediateBuffer)
-		{
-			if (DataTypeSize == sizeof(uint16))
-			{
-				static_cast<FRawStaticIndexBuffer16or32<uint16>*>(IndexBuffer)->InitRHIForStreaming(IntermediateBuffer, Batcher);
-			}
-			else
-			{
-				static_cast<FRawStaticIndexBuffer16or32<uint32>*>(IndexBuffer)->InitRHIForStreaming(IntermediateBuffer, Batcher);
-			}
-		}
-	}
-
-	template <uint32 MaxNumUpdates>
-	void ReleaseRHIForStreaming(TRHIResourceUpdateBatcher<MaxNumUpdates>& Batcher)
-	{
-		if (IndexBuffer)
-		{
-			if (DataTypeSize == sizeof(uint16))
-			{
-				static_cast<FRawStaticIndexBuffer16or32<uint16>*>(IndexBuffer)->ReleaseRHIForStreaming(Batcher);
-			}
-			else
-			{
-				static_cast<FRawStaticIndexBuffer16or32<uint32>*>(IndexBuffer)->ReleaseRHIForStreaming(Batcher);
-			}
-		}
-	}
+	void InitRHIForStreaming(FRHIBuffer* IntermediateBuffer, FRHIResourceUpdateBatcher& Batcher);
+	void ReleaseRHIForStreaming(FRHIResourceUpdateBatcher& Batcher);
 
 private:
 	/** Size of the index buffer's index type (should be 2 or 4 bytes) */
diff --git a/Engine/Source/Runtime/Engine/Public/Rendering/PositionVertexBuffer.h b/Engine/Source/Runtime/Engine/Public/Rendering/PositionVertexBuffer.h
--- a/Engine/Source/Runtime/Engine/Public/Rendering/PositionVertexBuffer.h
+++ b/Engine/Source/Runtime/Engine/Public/Rendering/PositionVertexBuffer.h
@@ -104,30 +104,8 @@ public:
 	void CopyRHIForStreaming(const FPositionVertexBuffer& Other, bool InAllowCPUAccess);
 
 	/** Similar to Init/ReleaseRHI but only update existing SRV so references to the SRV stays valid */
-	template <uint32 MaxNumUpdates>
-	void InitRHIForStreaming(FRHIBuffer* IntermediateBuffer, TRHIResourceUpdateBatcher<MaxNumUpdates>& Batcher)
-	{
-		check(VertexBufferRHI);
-		if (IntermediateBuffer)
-		{
-			Batcher.QueueUpdateRequest(VertexBufferRHI, IntermediateBuffer);
-			if (PositionComponentSRV)
-			{
-				Batcher.QueueUpdateRequest(PositionComponentSRV, VertexBufferRHI, 4, PF_R32_FLOAT);
-			}
-		}
-	}
-
-	template <uint32 MaxNumUpdates>
-	void ReleaseRHIForStreaming(TRHIResourceUpdateBatcher<MaxNumUpdates>& Batcher)
-	{
-		check(VertexBufferRHI);
-		Batcher.QueueUpdateRequest(VertexBufferRHI, nullptr);
-		if (PositionComponentSRV)
-		{
-			Batcher.QueueUpdateRequest(PositionComponentSRV, nullptr, 0, 0);
-		}
-	}
+	void InitRHIForStreaming(FRHIBuffer* IntermediateBuffer, FRHIResourceUpdateBatcher& Batcher);
+	void ReleaseRHIForStreaming(FRHIResourceUpdateBatcher& Batcher);
 
 	// FRenderResource interface.
 	ENGINE_API virtual void InitRHI() override;
diff --git a/Engine/Source/Runtime/Engine/Public/Rendering/SkeletalMeshVertexClothBuffer.h b/Engine/Source/Runtime/Engine/Public/Rendering/SkeletalMeshVertexClothBuffer.h
--- a/Engine/Source/Runtime/Engine/Public/Rendering/SkeletalMeshVertexClothBuffer.h
+++ b/Engine/Source/Runtime/Engine/Public/Rendering/SkeletalMeshVertexClothBuffer.h
@@ -128,29 +128,8 @@ public:
 	FBufferRHIRef CreateRHIBuffer_Async();
 
 	/** Similar to Init/ReleaseRHI but only update existing SRV so references to the SRV stays valid */
-	template <uint32 MaxNumUpdates>
-	void InitRHIForStreaming(FRHIBuffer* IntermediateBuffer, TRHIResourceUpdateBatcher<MaxNumUpdates>& Batcher)
-	{
-		if (VertexBufferRHI && IntermediateBuffer)
-		{
-			check(VertexBufferSRV);
-			Batcher.QueueUpdateRequest(VertexBufferRHI, IntermediateBuffer);
-			Batcher.QueueUpdateRequest(VertexBufferSRV, VertexBufferRHI, sizeof(FVector4f), PF_A32B32G32R32F);
-		}
-	}
-
-	template <uint32 MaxNumUpdates>
-	void ReleaseRHIForStreaming(TRHIResourceUpdateBatcher<MaxNumUpdates>& Batcher)
-	{
-		if (VertexBufferRHI)
-		{
-			Batcher.QueueUpdateRequest(VertexBufferRHI, nullptr);
-		}
-		if (VertexBufferSRV)
-		{
-			Batcher.QueueUpdateRequest(VertexBufferSRV, nullptr, 0, 0);
-		}
-	}
+	void InitRHIForStreaming(FRHIBuffer* IntermediateBuffer, FRHIResourceUpdateBatcher& Batcher);
+	void ReleaseRHIForStreaming(FRHIResourceUpdateBatcher& Batcher);
 
 private:
 	FShaderResourceViewRHIRef VertexBufferSRV;
diff --git a/Engine/Source/Runtime/Engine/Public/Rendering/SkinWeightVertexBuffer.h b/Engine/Source/Runtime/Engine/Public/Rendering/SkinWeightVertexBuffer.h
--- a/Engine/Source/Runtime/Engine/Public/Rendering/SkinWeightVertexBuffer.h
+++ b/Engine/Source/Runtime/Engine/Public/Rendering/SkinWeightVertexBuffer.h
@@ -142,29 +142,8 @@ public:
 	FBufferRHIRef CreateRHIBuffer_Async();
 
 	/** Similar to Init/ReleaseRHI but only update existing SRV so references to the SRV stays valid */
-	template <uint32 MaxNumUpdates>
-	void InitRHIForStreaming(FRHIBuffer* IntermediateBuffer, TRHIResourceUpdateBatcher<MaxNumUpdates>& Batcher)
-	{
-		if (VertexBufferRHI && IntermediateBuffer)
-		{
-			check(SRVValue);
-			Batcher.QueueUpdateRequest(VertexBufferRHI, IntermediateBuffer);
-			Batcher.QueueUpdateRequest(SRVValue, VertexBufferRHI, PixelFormatStride, PixelFormat);
-		}
-	}
-
-	template <uint32 MaxNumUpdates>
-	void ReleaseRHIForStreaming(TRHIResourceUpdateBatcher<MaxNumUpdates>& Batcher)
-	{
-		if (VertexBufferRHI)
-		{
-			Batcher.QueueUpdateRequest(VertexBufferRHI, nullptr);
-		}
-		if (SRVValue)
-		{
-			Batcher.QueueUpdateRequest(SRVValue, nullptr, 0, 0);
-		}
-	}
+	void InitRHIForStreaming(FRHIBuffer* IntermediateBuffer, FRHIResourceUpdateBatcher& Batcher);
+	void ReleaseRHIForStreaming(FRHIResourceUpdateBatcher& Batcher);
 
 	void GetWeightOffsetAndInfluenceCount(uint32 VertexIndex, uint32& OutWeightOffset, uint32& OutInfluenceCount) const;
 	void SetWeightOffsetAndInfluenceCount(uint32 VertexIndex, uint32 WeightOffset, uint32 InfluenceCount);
@@ -325,29 +304,8 @@ public:
 	FBufferRHIRef CreateRHIBuffer_Async();
 
 	/** Similar to Init/ReleaseRHI but only update existing SRV so references to the SRV stays valid */
-	template <uint32 MaxNumUpdates>
-	void InitRHIForStreaming(FRHIBuffer* IntermediateBuffer, TRHIResourceUpdateBatcher<MaxNumUpdates>& Batcher)
-	{
-		if (VertexBufferRHI && IntermediateBuffer)
-		{
-			check(SRVValue);
-			Batcher.QueueUpdateRequest(VertexBufferRHI, IntermediateBuffer);
-			Batcher.QueueUpdateRequest(SRVValue, VertexBufferRHI, GetPixelFormatStride(), GetPixelFormat());
-		}
-	}
-
-	template <uint32 MaxNumUpdates>
-	void ReleaseRHIForStreaming(TRHIResourceUpdateBatcher<MaxNumUpdates>& Batcher)
-	{
-		if (VertexBufferRHI)
-		{
-			Batcher.QueueUpdateRequest(VertexBufferRHI, nullptr);
-		}
-		if (SRVValue)
-		{
-			Batcher.QueueUpdateRequest(SRVValue, nullptr, 0, 0);
-		}
-	}
+	void InitRHIForStreaming(FRHIBuffer* IntermediateBuffer, FRHIResourceUpdateBatcher& Batcher);
+	void ReleaseRHIForStreaming(FRHIResourceUpdateBatcher& Batcher);
 
 	bool IsWeightDataValid() const;
 
@@ -521,15 +479,13 @@ public:
 	ENGINE_API void ReleaseResources();
 
 	/** Similar to Init/ReleaseRHI but only update existing SRV so references to the SRV stays valid */
-	template <uint32 MaxNumUpdates>
-	void InitRHIForStreaming(const FSkinWeightRHIInfo& RHIInfo, TRHIResourceUpdateBatcher<MaxNumUpdates>& Batcher)
+	void InitRHIForStreaming(const FSkinWeightRHIInfo& RHIInfo, FRHIResourceUpdateBatcher& Batcher)
 	{
 		DataVertexBuffer.InitRHIForStreaming(RHIInfo.DataVertexBufferRHI, Batcher);
 		LookupVertexBuffer.InitRHIForStreaming(RHIInfo.LookupVertexBufferRHI, Batcher);
 	}
 
-	template <uint32 MaxNumUpdates>
-	void ReleaseRHIForStreaming(TRHIResourceUpdateBatcher<MaxNumUpdates>& Batcher)
+	void ReleaseRHIForStreaming(FRHIResourceUpdateBatcher& Batcher)
 	{
 		DataVertexBuffer.ReleaseRHIForStreaming(Batcher);
 		LookupVertexBuffer.ReleaseRHIForStreaming(Batcher);
diff --git a/Engine/Source/Runtime/Engine/Public/Rendering/StaticMeshVertexBuffer.h b/Engine/Source/Runtime/Engine/Public/Rendering/StaticMeshVertexBuffer.h
--- a/Engine/Source/Runtime/Engine/Public/Rendering/StaticMeshVertexBuffer.h
+++ b/Engine/Source/Runtime/Engine/Public/Rendering/StaticMeshVertexBuffer.h
@@ -437,54 +437,8 @@ public:
 	void CopyRHIForStreaming(const FStaticMeshVertexBuffer& Other, bool InAllowCPUAccess);
 
 	/** Similar to Init/ReleaseRHI but only update existing SRV so references to the SRV stays valid */
-	template <uint32 MaxNumUpdates>
-	void InitRHIForStreaming(
-		FRHIBuffer* IntermediateTangentsBuffer,
-		FRHIBuffer* IntermediateTexCoordBuffer,
-		TRHIResourceUpdateBatcher<MaxNumUpdates>& Batcher)
-	{
-		check(TangentsVertexBuffer.VertexBufferRHI && TexCoordVertexBuffer.VertexBufferRHI);
-		if (IntermediateTangentsBuffer)
-		{
-			Batcher.QueueUpdateRequest(TangentsVertexBuffer.VertexBufferRHI, IntermediateTangentsBuffer);
-			if (TangentsSRV)
-			{
-				Batcher.QueueUpdateRequest(
-					TangentsSRV,
-					TangentsVertexBuffer.VertexBufferRHI,
-					GetUseHighPrecisionTangentBasis() ? 8u : 4u,
-					GetUseHighPrecisionTangentBasis() ? (uint8)PF_R16G16B16A16_SNORM : (uint8)PF_R8G8B8A8_SNORM);;
-			}
-		}
-		if (IntermediateTexCoordBuffer)
-		{
-			Batcher.QueueUpdateRequest(TexCoordVertexBuffer.VertexBufferRHI, IntermediateTexCoordBuffer);
-			if (TextureCoordinatesSRV)
-			{
-				Batcher.QueueUpdateRequest(
-					TextureCoordinatesSRV,
-					TexCoordVertexBuffer.VertexBufferRHI,
-					GetUseFullPrecisionUVs() ? 8u : 4u,
-					GetUseFullPrecisionUVs() ? (uint8)PF_G32R32F : (uint8)PF_G16R16F);
-			}
-		}
-	}
-
-	template<uint32 MaxNumUpdates>
-	void ReleaseRHIForStreaming(TRHIResourceUpdateBatcher<MaxNumUpdates>& Batcher)
-	{
-		check(TangentsVertexBuffer.VertexBufferRHI && TexCoordVertexBuffer.VertexBufferRHI);
-		Batcher.QueueUpdateRequest(TangentsVertexBuffer.VertexBufferRHI, nullptr);
-		Batcher.QueueUpdateRequest(TexCoordVertexBuffer.VertexBufferRHI, nullptr);
-		if (TangentsSRV)
-		{
-			Batcher.QueueUpdateRequest(TangentsSRV, nullptr, 0, 0);
-		}
-		if (TextureCoordinatesSRV)
-		{
-			Batcher.QueueUpdateRequest(TextureCoordinatesSRV, nullptr, 0, 0);
-		}
-	}
+	void InitRHIForStreaming(FRHIBuffer* IntermediateTangentsBuffer, FRHIBuffer* IntermediateTexCoordBuffer, FRHIResourceUpdateBatcher& Batcher);
+	void ReleaseRHIForStreaming(FRHIResourceUpdateBatcher& Batcher);
 
 	// FRenderResource interface.
 	ENGINE_API virtual void InitRHI() override;
diff --git a/Engine/Source/Runtime/RHI/Private/RHICommandList.cpp b/Engine/Source/Runtime/RHI/Private/RHICommandList.cpp
--- a/Engine/Source/Runtime/RHI/Private/RHICommandList.cpp
+++ b/Engine/Source/Runtime/RHI/Private/RHICommandList.cpp
@@ -14,6 +14,7 @@
 #include "ProfilingDebugging/CsvProfiler.h"
 #include "Trace/Trace.inl"
 #include "GenericPlatform/GenericPlatformCrashContext.h"
+#include "RHIResourceUpdates.h"
 
 CSV_DEFINE_CATEGORY_MODULE(RHI_API, RHITStalls, false);
 CSV_DEFINE_CATEGORY_MODULE(RHI_API, RHITFlushes, false);
diff --git a/Engine/Source/Runtime/RHI/Private/RHIResourceUpdates.cpp b/Engine/Source/Runtime/RHI/Private/RHIResourceUpdates.cpp
new file mode 100644
--- /dev/null
+++ b/Engine/Source/Runtime/RHI/Private/RHIResourceUpdates.cpp
@@ -0,0 +1,87 @@
+// Copyright Epic Games, Inc. All Rights Reserved.
+
+#include "RHIResourceUpdates.h"
+#include "RHI.h"
+
+FRHIResourceUpdateBatcher::FRHIResourceUpdateBatcher(TArrayView<FRHIResourceUpdateInfo> InUpdateInfos)
+	: UpdateInfos(InUpdateInfos)
+{}
+
+FRHIResourceUpdateBatcher::~FRHIResourceUpdateBatcher()
+{
+	Flush();
+}
+
+void FRHIResourceUpdateBatcher::Flush()
+{
+	if (NumBatched > 0)
+	{
+		RHIUpdateRHIResources(&UpdateInfos[0], NumBatched, true);
+		NumBatched = 0;
+	}
+}
+
+void FRHIResourceUpdateBatcher::QueueUpdateRequest(FRHIBuffer* DestBuffer, FRHIBuffer* SrcBuffer)
+{
+	FRHIResourceUpdateInfo& UpdateInfo = GetNextUpdateInfo();
+	UpdateInfo.Type = FRHIResourceUpdateInfo::UT_Buffer;
+	UpdateInfo.Buffer = { DestBuffer, SrcBuffer };
+	DestBuffer->AddRef();
+	if (SrcBuffer)
+	{
+		SrcBuffer->AddRef();
+	}
+}
+
+void FRHIResourceUpdateBatcher::QueueUpdateRequest(FRHIRayTracingGeometry* DestGeometry, FRHIRayTracingGeometry* SrcGeometry)
+{
+	FRHIResourceUpdateInfo& UpdateInfo = GetNextUpdateInfo();
+	UpdateInfo.Type = FRHIResourceUpdateInfo::UT_RayTracingGeometry;
+	UpdateInfo.RayTracingGeometry = { DestGeometry, SrcGeometry };
+	DestGeometry->AddRef();
+	if (SrcGeometry)
+	{
+		SrcGeometry->AddRef();
+	}
+}
+
+void FRHIResourceUpdateBatcher::QueueUpdateRequest(FRHIShaderResourceView* SRV, FRHIBuffer* Buffer, uint32 Stride, uint8 Format)
+{
+	FRHIResourceUpdateInfo& UpdateInfo = GetNextUpdateInfo();
+	UpdateInfo.Type = FRHIResourceUpdateInfo::UT_BufferFormatSRV;
+	UpdateInfo.BufferSRV = { SRV, Buffer, Stride, Format };
+	SRV->AddRef();
+	if (Buffer)
+	{
+		Buffer->AddRef();
+	}
+}
+
+void FRHIResourceUpdateBatcher::QueueUpdateRequest(FRHIShaderResourceView* SRV, FRHIBuffer* Buffer)
+{
+	FRHIResourceUpdateInfo& UpdateInfo = GetNextUpdateInfo();
+	UpdateInfo.Type = FRHIResourceUpdateInfo::UT_BufferSRV;
+	UpdateInfo.BufferSRV = { SRV, Buffer };
+	SRV->AddRef();
+	if (Buffer)
+	{
+		Buffer->AddRef();
+	}
+}
+
+FRHIResourceUpdateInfo& FRHIResourceUpdateBatcher::GetNextUpdateInfo()
+{
+	check(NumBatched <= UpdateInfos.Num());
+	if (NumBatched >= UpdateInfos.Num())
+	{
+		Flush();
+	}
+#ifdef _MSC_VER
+#pragma warning(push)
+#pragma warning(disable : 6385) // Access is always in-bound due to the Flush above
+#endif
+	return UpdateInfos[NumBatched++];
+#ifdef _MSC_VER
+#pragma warning(pop)
+#endif
+}
diff --git a/Engine/Source/Runtime/RHI/Public/RHICommandList.h b/Engine/Source/Runtime/RHI/Public/RHICommandList.h
--- a/Engine/Source/Runtime/RHI/Public/RHICommandList.h
+++ b/Engine/Source/Runtime/RHI/Public/RHICommandList.h
@@ -2125,53 +2125,6 @@ FRHICOMMAND_MACRO(FRHICommandDiscardRenderTargets)
 	RHI_API void Execute(FRHICommandListBase& CmdList);
 };
 
-struct FRHIShaderResourceViewUpdateInfo
-{
-	FRHIShaderResourceView* SRV;
-	FRHIBuffer* Buffer;
-	uint32 Stride;
-	uint8 Format;
-};
-
-struct FRHIBufferUpdateInfo
-{
-	FRHIBuffer* DestBuffer;
-	FRHIBuffer* SrcBuffer;
-};
-
-struct FRHIRayTracingGeometryUpdateInfo
-{
-	FRHIRayTracingGeometry* DestGeometry;
-	FRHIRayTracingGeometry* SrcGeometry;
-};
-
-struct FRHIResourceUpdateInfo
-{
-	enum EUpdateType
-	{
-		/** Take over underlying resource from an intermediate buffer */
-		UT_Buffer,
-		/** Update an SRV to view on a different buffer */
-		UT_BufferSRV,
-		/** Update an SRV to view on a different buffer with a given format */
-		UT_BufferFormatSRV,
-		/** Take over underlying resource from an intermediate geometry */
-		UT_RayTracingGeometry,
-		/** Number of update types */
-		UT_Num
-	};
-
-	EUpdateType Type;
-	union
-	{
-		FRHIBufferUpdateInfo Buffer;
-		FRHIShaderResourceViewUpdateInfo BufferSRV;
-		FRHIRayTracingGeometryUpdateInfo RayTracingGeometry;
-	};
-
-	void ReleaseRefs();
-};
-
 FRHICOMMAND_MACRO(FRHICommandUpdateRHIResources)
 {
 	FRHIResourceUpdateInfo* UpdateInfos;
@@ -5964,97 +5917,6 @@ namespace UE::RHI
 
 } //! UE::RHI
 
-template <uint32 MaxNumUpdates>
-struct TRHIResourceUpdateBatcher
-{
-	FRHIResourceUpdateInfo UpdateInfos[MaxNumUpdates];
-	uint32 NumBatched;
-
-	TRHIResourceUpdateBatcher()
-		: NumBatched(0)
-	{}
-
-	~TRHIResourceUpdateBatcher()
-	{
-		Flush();
-	}
-
-	void Flush()
-	{
-		if (NumBatched > 0)
-		{
-			RHIUpdateRHIResources(UpdateInfos, NumBatched, true);
-			NumBatched = 0;
-		}
-	}
-
-	void QueueUpdateRequest(FRHIBuffer* DestBuffer, FRHIBuffer* SrcBuffer)
-	{
-		FRHIResourceUpdateInfo& UpdateInfo = GetNextUpdateInfo();
-		UpdateInfo.Type = FRHIResourceUpdateInfo::UT_Buffer;
-		UpdateInfo.Buffer = { DestBuffer, SrcBuffer };
-		DestBuffer->AddRef();
-		if (SrcBuffer)
-		{
-			SrcBuffer->AddRef();
-		}
-	}
-
-	void QueueUpdateRequest(FRHIRayTracingGeometry* DestGeometry, FRHIRayTracingGeometry* SrcGeometry)
-	{
-		FRHIResourceUpdateInfo& UpdateInfo = GetNextUpdateInfo();
-		UpdateInfo.Type = FRHIResourceUpdateInfo::UT_RayTracingGeometry;
-		UpdateInfo.RayTracingGeometry = { DestGeometry, SrcGeometry };
-		DestGeometry->AddRef();
-		if (SrcGeometry)
-		{
-			SrcGeometry->AddRef();
-		}
-	}
-
-	void QueueUpdateRequest(FRHIShaderResourceView* SRV, FRHIBuffer* Buffer, uint32 Stride, uint8 Format)
-	{
-		FRHIResourceUpdateInfo& UpdateInfo = GetNextUpdateInfo();
-		UpdateInfo.Type = FRHIResourceUpdateInfo::UT_BufferFormatSRV;
-		UpdateInfo.BufferSRV = { SRV, Buffer, Stride, Format };
-		SRV->AddRef();
-		if (Buffer)
-		{
-			Buffer->AddRef();
-		}
-	}
-
-	void QueueUpdateRequest(FRHIShaderResourceView* SRV, FRHIBuffer* Buffer)
-	{
-		FRHIResourceUpdateInfo& UpdateInfo = GetNextUpdateInfo();
-		UpdateInfo.Type = FRHIResourceUpdateInfo::UT_BufferSRV;
-		UpdateInfo.BufferSRV = { SRV, Buffer };
-		SRV->AddRef();
-		if (Buffer)
-		{
-			Buffer->AddRef();
-		}
-	}
-
-private:
-	FRHIResourceUpdateInfo & GetNextUpdateInfo()
-	{
-		check(NumBatched <= MaxNumUpdates);
-		if (NumBatched >= MaxNumUpdates)
-		{
-			Flush();
-		}
-#ifdef _MSC_VER
-#pragma warning(push)
-#pragma warning(disable : 6385) // Access is always in-bound due to the Flush above
-#endif
-		return UpdateInfos[NumBatched++];
-#ifdef _MSC_VER
-#pragma warning(pop)
-#endif
-	}
-};
-
 #undef RHICOMMAND_CALLSTACK
 
 #include "RHICommandList.inl" // IWYU pragma: export
diff --git a/Engine/Source/Runtime/RHI/Public/RHICommandListCommandExecutes.inl b/Engine/Source/Runtime/RHI/Public/RHICommandListCommandExecutes.inl
--- a/Engine/Source/Runtime/RHI/Public/RHICommandListCommandExecutes.inl
+++ b/Engine/Source/Runtime/RHI/Public/RHICommandListCommandExecutes.inl
@@ -13,6 +13,8 @@
 #define INTERNAL_DECORATOR_COMPUTE(Method) CmdList.GetComputeContext().Method
 #endif
 
+#include "RHIResourceUpdates.h"
+
 class FRHICommandListBase;
 class IRHIComputeContext;
 struct FComputedBSS;
diff --git a/Engine/Source/Runtime/RHI/Public/RHIFwd.h b/Engine/Source/Runtime/RHI/Public/RHIFwd.h
--- a/Engine/Source/Runtime/RHI/Public/RHIFwd.h
+++ b/Engine/Source/Runtime/RHI/Public/RHIFwd.h
@@ -13,12 +13,23 @@ enum ECubeFace : uint32;
 enum EPixelFormat : uint8;
 enum class EPixelFormatChannelFlags : uint8;
 
+enum class EBufferUsageFlags : uint32;
+enum class ETextureCreateFlags : uint64;
+
 // Command Lists
 class FRHICommandListBase;
 class FRHIComputeCommandList;
 class FRHICommandList;
 class FRHICommandListImmediate;
 
+struct FRHIResourceUpdateInfo;
+struct FRHIResourceUpdateBatcher;
+
+struct FSamplerStateInitializerRHI;
+struct FRasterizerStateInitializerRHI;
+struct FDepthStencilStateInitializerRHI;
+class FBlendStateInitializerRHI;
+
 // Resources
 class FRHIAmplificationShader;
 class FRHIBlendState;
diff --git a/Engine/Source/Runtime/RHI/Public/RHIResourceUpdates.h b/Engine/Source/Runtime/RHI/Public/RHIResourceUpdates.h
new file mode 100644
--- /dev/null
+++ b/Engine/Source/Runtime/RHI/Public/RHIResourceUpdates.h
@@ -0,0 +1,84 @@
+// Copyright Epic Games, Inc. All Rights Reserved.
+
+#pragma once
+
+#include "RHIFwd.h"
+#include "Containers/ArrayView.h"
+
+struct FRHIShaderResourceViewUpdateInfo
+{
+	FRHIShaderResourceView* SRV;
+	FRHIBuffer* Buffer;
+	uint32 Stride;
+	uint8 Format;
+};
+
+struct FRHIBufferUpdateInfo
+{
+	FRHIBuffer* DestBuffer;
+	FRHIBuffer* SrcBuffer;
+};
+
+struct FRHIRayTracingGeometryUpdateInfo
+{
+	FRHIRayTracingGeometry* DestGeometry;
+	FRHIRayTracingGeometry* SrcGeometry;
+};
+
+struct FRHIResourceUpdateInfo
+{
+	enum EUpdateType
+	{
+		/** Take over underlying resource from an intermediate buffer */
+		UT_Buffer,
+		/** Update an SRV to view on a different buffer */
+		UT_BufferSRV,
+		/** Update an SRV to view on a different buffer with a given format */
+		UT_BufferFormatSRV,
+		/** Take over underlying resource from an intermediate geometry */
+		UT_RayTracingGeometry,
+		/** Number of update types */
+		UT_Num
+	};
+
+	EUpdateType Type;
+	union
+	{
+		FRHIBufferUpdateInfo Buffer;
+		FRHIShaderResourceViewUpdateInfo BufferSRV;
+		FRHIRayTracingGeometryUpdateInfo RayTracingGeometry;
+	};
+
+	void ReleaseRefs();
+};
+
+struct RHI_API FRHIResourceUpdateBatcher
+{
+	TArrayView<FRHIResourceUpdateInfo> UpdateInfos;
+	int32 NumBatched = 0;
+
+	~FRHIResourceUpdateBatcher();
+
+	void Flush();
+
+	void QueueUpdateRequest(FRHIBuffer* DestBuffer, FRHIBuffer* SrcBuffer);
+	void QueueUpdateRequest(FRHIRayTracingGeometry* DestGeometry, FRHIRayTracingGeometry* SrcGeometry);
+	void QueueUpdateRequest(FRHIShaderResourceView* SRV, FRHIBuffer* Buffer, uint32 Stride, uint8 Format);
+	void QueueUpdateRequest(FRHIShaderResourceView* SRV, FRHIBuffer* Buffer);
+
+protected:
+	FRHIResourceUpdateBatcher() = delete;
+	FRHIResourceUpdateBatcher(TArrayView<FRHIResourceUpdateInfo> InUpdateInfos);
+
+private:
+	FRHIResourceUpdateInfo& GetNextUpdateInfo();
+};
+
+template <uint32 MaxNumUpdates>
+struct TRHIResourceUpdateBatcher : public FRHIResourceUpdateBatcher
+{
+	FRHIResourceUpdateInfo UpdateInfoStorage[MaxNumUpdates];
+
+	TRHIResourceUpdateBatcher() : FRHIResourceUpdateBatcher(UpdateInfoStorage)
+	{}
+};
diff --git a/Engine/Source/Runtime/RenderCore/Private/RayTracingGeometry.cpp b/Engine/Source/Runtime/RenderCore/Private/RayTracingGeometry.cpp
--- a/Engine/Source/Runtime/RenderCore/Private/RayTracingGeometry.cpp
+++ b/Engine/Source/Runtime/RenderCore/Private/RayTracingGeometry.cpp
@@ -3,6 +3,7 @@
 #include "RayTracingGeometry.h"
 #include "RayTracingGeometryManager.h"
 #include "RenderUtils.h"
+#include "RHIResourceUpdates.h"
 
 int32 GVarDebugForceRuntimeBLAS = 0;
 FAutoConsoleVariableRef CVarDebugForceRuntimeBLAS(
@@ -18,6 +19,31 @@ FRayTracingGeometry::~FRayTracingGeometry() = default;
 
 #if RHI_RAYTRACING
 
+void FRayTracingGeometry::InitRHIForStreaming(FRHIRayTracingGeometry* IntermediateGeometry, FRHIResourceUpdateBatcher& Batcher)
+{
+	EnumAddFlags(GeometryState, EGeometryStateFlags::StreamedIn);
+
+	if (RayTracingGeometryRHI && IntermediateGeometry)
+	{
+		Batcher.QueueUpdateRequest(RayTracingGeometryRHI, IntermediateGeometry);
+		EnumAddFlags(GeometryState, EGeometryStateFlags::Valid);
+	}
+}
+
+void FRayTracingGeometry::ReleaseRHIForStreaming(FRHIResourceUpdateBatcher& Batcher)
+{
+	Initializer = {};
+
+	RemoveBuildRequest();
+	EnumRemoveFlags(GeometryState, EGeometryStateFlags::StreamedIn);
+
+	if (RayTracingGeometryRHI)
+	{
+		Batcher.QueueUpdateRequest(RayTracingGeometryRHI, nullptr);
+		EnumRemoveFlags(GeometryState, EGeometryStateFlags::Valid);
+	}
+}
+
 void FRayTracingGeometry::CreateRayTracingGeometryFromCPUData(TResourceArray<uint8>& OfflineData)
 {
 	check(OfflineData.Num() == 0 || Initializer.OfflineData == nullptr);
diff --git a/Engine/Source/Runtime/RenderCore/Public/RayTracingGeometry.h b/Engine/Source/Runtime/RenderCore/Public/RayTracingGeometry.h
--- a/Engine/Source/Runtime/RenderCore/Public/RayTracingGeometry.h
+++ b/Engine/Source/Runtime/RenderCore/Public/RayTracingGeometry.h
@@ -3,6 +3,7 @@
 #pragma once
 
 #include "CoreMinimal.h"
+#include "Containers/DynamicRHIResourceArray.h"
 #include "RenderResource.h"
 #if RHI_RAYTRACING
 #include "RHI.h"
@@ -96,32 +97,9 @@ public:
 		return GeometryState;
 	}
 
-	template <uint32 MaxNumUpdates>
-	void InitRHIForStreaming(FRHIRayTracingGeometry* IntermediateGeometry, TRHIResourceUpdateBatcher<MaxNumUpdates>& Batcher)
-	{
-		EnumAddFlags(GeometryState, EGeometryStateFlags::StreamedIn);
-
-		if (RayTracingGeometryRHI && IntermediateGeometry)
-		{
-			Batcher.QueueUpdateRequest(RayTracingGeometryRHI, IntermediateGeometry);
-			EnumAddFlags(GeometryState, EGeometryStateFlags::Valid);
-		}
-	}
+	void InitRHIForStreaming(FRHIRayTracingGeometry* IntermediateGeometry, FRHIResourceUpdateBatcher& Batcher);
+	void ReleaseRHIForStreaming(FRHIResourceUpdateBatcher& Batcher);
 
-	template <uint32 MaxNumUpdates>
-	void ReleaseRHIForStreaming(TRHIResourceUpdateBatcher<MaxNumUpdates>& Batcher)
-	{
-		Initializer = {};
-
-		RemoveBuildRequest();
-		EnumRemoveFlags(GeometryState, EGeometryStateFlags::StreamedIn);
-
-		if (RayTracingGeometryRHI)
-		{
-			Batcher.QueueUpdateRequest(RayTracingGeometryRHI, nullptr);
-			EnumRemoveFlags(GeometryState, EGeometryStateFlags::Valid);
-		}
-	}
 	void CreateRayTracingGeometryFromCPUData(TResourceArray<uint8>& OfflineData);
 	void RequestBuildIfNeeded(ERTAccelerationStructureBuildPriority InBuildPriority);
 
