Add GetPrecachingStateForInitializer to handle compilation failures for PSO CS precaching. This is useful to properly handle to use case of 23261036, where some shader compilations may succeed offline but fail on the final client

diff --git a/Engine/Source/Runtime/RHI/Private/PipelineStateCache.cpp b/Engine/Source/Runtime/RHI/Private/PipelineStateCache.cpp
--- a/Engine/Source/Runtime/RHI/Private/PipelineStateCache.cpp
+++ b/Engine/Source/Runtime/RHI/Private/PipelineStateCache.cpp
@@ -1037,13 +1037,13 @@ public:
 			return EPSOPrecacheResult::Active;
 		}
 
-		if (!PrecachedPSOInitializers.ContainsByHash(InitializerHash, Initializer))
+		const TPrecachedPSOInitializer* InitializerInSet = PrecachedPSOInitializers.FindByHash(InitializerHash, Initializer);
+		if (InitializerInSet == nullptr)
 		{
 			return EPSOPrecacheResult::Missed;
 		}
 
-		// not found then assume it's not precaching
-		return EPSOPrecacheResult::Complete;
+		return TPrecachePipelineCacheDerived::GetPrecachingStateForInitializer(*InitializerInSet);
 	}
 
 	bool IsPrecaching()
@@ -1083,7 +1083,7 @@ public:
 		}
 	}
 
-private:
+protected:
 	bool HasPSOBeenRequested(const TPrecachedPSOInitializer& Initializer, uint32 InitializerHash, FGraphEventRef& GraphEvent)
 	{
 		FPrecacheTask* FindResult = ActivePrecachingTasks.FindByHash(InitializerHash, Initializer);
@@ -1127,7 +1127,7 @@ private:
 		TPipelineState* PSO = nullptr;
 	};
 
-	using TMapKeyFuncsBaseClass = TDefaultMapHashableKeyFuncs<TPrecachedPSOInitializer, FPrecacheTask, false>;
+	using TMapKeyFuncsBaseClass = TDefaultMapKeyFuncs<TPrecachedPSOInitializer, FPrecacheTask, false>;
 	struct TMapKeyFuncs : public TMapKeyFuncsBaseClass
 	{
 		static FORCEINLINE bool Matches(typename TMapKeyFuncsBaseClass::KeyInitType A, typename TMapKeyFuncsBaseClass::KeyInitType B)
@@ -1144,28 +1144,73 @@ private:
 	TArray<TPipelineState*> PrecachedPSOs;
 };
 
-class FPrecacheComputePipelineCache : public TPrecachePipelineCacheBase<FPrecacheComputePipelineCache, FRHIComputeShader*, FComputePipelineState>
+struct FPrecacheComputeInitializer
+{
+	FPrecacheComputeInitializer()
+	: RHIComputeShaderAsU64(0)
+	, bCompilationFailed(0)
+	{}
+
+
+	FPrecacheComputeInitializer(const FRHIComputeShader* InRHIComputeShader)
+		: RHIComputeShaderAsU64((uint64)InRHIComputeShader)
+		, bCompilationFailed(0)
+	{}
+
+	uint64 RHIComputeShaderAsU64; // Using a U64 here rather than FRHIComputeShader* because we keep FPrecacheComputeInitializer but the FRHIComputeShader might get deleted
+	uint32 bCompilationFailed : 1;
+};
+
+class FPrecacheComputePipelineCache : public TPrecachePipelineCacheBase<FPrecacheComputePipelineCache, FPrecacheComputeInitializer, FComputePipelineState>
 {
 public:
-	static FComputePipelineState* CreateNewPSO(const FRHIComputeShader* ComputeShader)
+	static FComputePipelineState* CreateNewPSO(const FPrecacheComputeInitializer& ComputeShaderInitializer)
 	{
-		return new FComputePipelineState((FRHIComputeShader*)ComputeShader);
+		return new FComputePipelineState((FRHIComputeShader*)ComputeShaderInitializer.RHIComputeShaderAsU64);
 	}
 
 	FGraphEventRef PrecacheComputePipelineState(FRHIComputeShader* ComputeShader, bool bForcePrecache);
-	static void OnNewPipelineStateCreated(FRHIComputeShader* ComputeShader, FComputePipelineState* NewComputePipelineState, bool bDoAsyncCompile);
+	static void OnNewPipelineStateCreated(const FPrecacheComputeInitializer& ComputeInitializer, FComputePipelineState* NewComputePipelineState, bool bDoAsyncCompile);
+ 
 
-	static FORCEINLINE bool PipelineStateInitializerMatch(const FRHIComputeShader* A, const FRHIComputeShader* B)
+	static FORCEINLINE bool PipelineStateInitializerMatch(const FPrecacheComputeInitializer& ComputeShaderInitializerA, const FPrecacheComputeInitializer& ComputeShaderInitializerB)
 	{
 		// todo: would be good/more robust to have the CS equivalent of RHIMatchPrecachePSOInitializers instead of relying on pointers
 		// (for example if an FRHIComputeShader gets released and a new one is allocated with the same pointer value)
-		return A == B;
+		return ComputeShaderInitializerA.RHIComputeShaderAsU64 == ComputeShaderInitializerB.RHIComputeShaderAsU64;
+	}
+
+	static FORCEINLINE uint32 PipelineStateInitializerHash(const FPrecacheComputeInitializer& Key)
+	{
+		return GetTypeHash(Key.RHIComputeShaderAsU64);
+	}
+
+	static EPSOPrecacheResult GetPrecachingStateForInitializer(const FPrecacheComputeInitializer& Key)
+	{
+		return (Key.bCompilationFailed != 0) ? EPSOPrecacheResult::NotSupported : EPSOPrecacheResult::Complete;
 	}
 
-	static FORCEINLINE uint32 PipelineStateInitializerHash(const FRHIComputeShader* Key)
+	void RegisterCSFailure(const FComputePipelineState* ComputePipelineState)
 	{
-		// will hash the pointer through PointerHash(const void* Key)
-		return GetTypeHash(Key);
+		check(ComputePipelineState->RHIPipeline == nullptr);
+		FPrecacheComputeInitializer PrecacheComputeInitializer(ComputePipelineState->ComputeShader);
+		uint32 Hash = PipelineStateInitializerHash(PrecacheComputeInitializer);
+
+		FRWScopeLock WriteLock(PrecachePSOsRWLock, SLT_Write);
+		FPrecacheComputeInitializer* PrecacheComputeInitializerRegistered = PrecachedPSOInitializers.FindByHash(Hash, PrecacheComputeInitializer);
+		ensure(PrecacheComputeInitializerRegistered != nullptr);
+		if (PrecacheComputeInitializerRegistered != nullptr)
+		{
+			PrecacheComputeInitializerRegistered->bCompilationFailed = 1;
+		}
+
+		// Make sure we don't leak the PSO
+		FPrecacheTask* FindResult = ActivePrecachingTasks.FindByHash(Hash, PrecacheComputeInitializer);
+		if (FindResult != nullptr)
+		{
+			verify(ActivePrecachingTasks.RemoveByHash(Hash, PrecacheComputeInitializer) == 1);
+			PrecachedPSOs.Add(FindResult->PSO);
+		}
 	}
 
 };
@@ -1188,6 +1233,10 @@ public:
 		return RHIComputePrecachePSOHash(Key);
 	}
 
+	static EPSOPrecacheResult GetPrecachingStateForInitializer(const FGraphicsPipelineStateInitializer& Key)
+	{
+		return EPSOPrecacheResult::Complete;
+	}
 
 	FGraphEventRef PrecacheGraphicsPipelineState(const FGraphicsPipelineStateInitializer& Initializer);
 
@@ -1519,7 +1568,14 @@ public:
 			ComputePipeline->RHIPipeline = RHICreateComputePipelineState(ComputePipeline->ComputeShader);
 			if (Initializer.bPSOPrecache)
 			{
-				GPrecacheComputePipelineCache.PrecacheFinished(ComputePipeline->ComputeShader);
+				if (ComputePipeline->RHIPipeline == nullptr)
+				{
+					GPrecacheComputePipelineCache.RegisterCSFailure(ComputePipeline);
+				}
+				else
+				{
+					GPrecacheComputePipelineCache.PrecacheFinished(ComputePipeline->ComputeShader);
+				}
 			}
 		}
 		else
@@ -2319,7 +2375,6 @@ FGraphicsPipelineState* PipelineStateCache::GetAndOrCreateGraphicsPipelineState(
 
 		bool bPSOPrecache = Initializer.bFromPSOFileCache;
 		FGraphEventRef GraphEvent = OutCachedState->CompletionEvent;
-		check(PSOPrecacheResult != EPSOPrecacheResult::Active);
 		InternalCreateGraphicsPipelineState(Initializer, PSOPrecacheResult, DoAsyncCompile, bPSOPrecache, OutCachedState);
 
 		// Add dispatch pre requisite for non precaching jobs only
@@ -2390,18 +2445,21 @@ FGraphEventRef FPrecacheComputePipelineCache::PrecacheComputePipelineState(FRHIC
 		return GraphEvent;
 	}
 
+	FPrecacheComputeInitializer PrecacheComputeInitializer(ComputeShader);
+
 	static bool bDoAsyncCompile = FApp::ShouldUseThreadingForPerformance();
-	GPrecacheComputePipelineCache.TryAddNewState(ComputeShader, bDoAsyncCompile, GraphEvent);
+	GPrecacheComputePipelineCache.TryAddNewState(PrecacheComputeInitializer, bDoAsyncCompile, GraphEvent);
 
 	if (GraphEvent == nullptr)
 	{
-		check(GPrecacheComputePipelineCache.GetPrecachingState(ComputeShader) == EPSOPrecacheResult::Complete)
+		EPSOPrecacheResult PSOPrecacheResult = GPrecacheComputePipelineCache.GetPrecachingState(ComputeShader);
+		check(PSOPrecacheResult == EPSOPrecacheResult::Complete|| PSOPrecacheResult == EPSOPrecacheResult::NotSupported);
 	}
 
 	return GraphEvent;
 }
 
-void FPrecacheComputePipelineCache::OnNewPipelineStateCreated(FRHIComputeShader* ComputeShader, FComputePipelineState* CachedState, bool bDoAsyncCompile)
+void FPrecacheComputePipelineCache::OnNewPipelineStateCreated(const FPrecacheComputeInitializer& ComputeInitializer, FComputePipelineState* CachedState, bool bDoAsyncCompile)
 {
 	// create a compilation task, or just do it now...
 	FGraphEventRef GraphEvent = CachedState->CompletionEvent;
