Update SupportsInlineRayTracing for preview shader platforms based on the support in the editor.

diff --git a/Engine/Source/Runtime/RHI/Private/DataDrivenShaderPlatformInfo.cpp b/Engine/Source/Runtime/RHI/Private/DataDrivenShaderPlatformInfo.cpp
--- a/Engine/Source/Runtime/RHI/Private/DataDrivenShaderPlatformInfo.cpp
+++ b/Engine/Source/Runtime/RHI/Private/DataDrivenShaderPlatformInfo.cpp
@@ -364,6 +364,8 @@ void FGenericDataDrivenShaderPlatformInfo::UpdatePreviewPlatforms()
 				Infos[ShaderPlatform].bSupportsLumenGI &= Infos[EditorSPForPreviewMaxFeatureLevel].bSupportsLumenGI;
 				Infos[ShaderPlatform].bSupportsUInt64ImageAtomics &= Infos[EditorSPForPreviewMaxFeatureLevel].bSupportsUInt64ImageAtomics;
 				Infos[ShaderPlatform].bSupportsGen5TemporalAA &= Infos[EditorSPForPreviewMaxFeatureLevel].bSupportsGen5TemporalAA;
+				Infos[ShaderPlatform].bSupportsInlineRayTracing &= Infos[EditorSPForPreviewMaxFeatureLevel].bSupportsInlineRayTracing;
+
 				// Support for stereo features requires extra consideration. The editor may not use the same technique as the preview platform,
 				// particularly MobileMultiView may be substituted by a fallback path. In order to avoid inundating real mobile platforms
 				// with the properties needed for the desktop MMV fallback path, override them here with the editor ones to make MMV preview possible
