ParameterBatcher wasn't being enabled after a command list was reset.

diff --git a/Engine/Source/Runtime/RHI/Private/RHICommandList.cpp b/Engine/Source/Runtime/RHI/Private/RHICommandList.cpp
--- a/Engine/Source/Runtime/RHI/Private/RHICommandList.cpp
+++ b/Engine/Source/Runtime/RHI/Private/RHICommandList.cpp
@@ -158,6 +158,8 @@ FRHIParameterBatcher::FRHIParameterBatcher(const FBoundShaderStateInput& InBound
 	AllBatchedShaders[SF_Pixel] = InBoundShaderStateInput.GetPixelShader();
 	AllBatchedShaders[SF_Geometry] = InBoundShaderStateInput.GetGeometryShader();
 	AllBatchedShaders[SF_Compute] = InBoundComputeShaderRHI;
+
+	bEnabled = GRHICommandParameterBatching != 0;
 }
 
 FRHIParameterBatcher::FRHIParameterBatcher(FRHIParameterBatcher&&) = default;
