- Add stats on actual PSO creation hitches and active PSO precache requests controlled via r.PSO.RuntimeCreationHitchThreshold (default set to 20msec for now)
- Add stats on active compute and graphics preached PSO requests
- Add function to check the total amount of open active PSO requests

#rnx

diff --git a/Engine/Source/Runtime/RHI/Private/PipelineStateCache.cpp b/Engine/Source/Runtime/RHI/Private/PipelineStateCache.cpp
--- a/Engine/Source/Runtime/RHI/Private/PipelineStateCache.cpp
+++ b/Engine/Source/Runtime/RHI/Private/PipelineStateCache.cpp
@@ -30,6 +30,13 @@ PipelineStateCache.cpp: Pipeline state cache implementation.
 
 CSV_DECLARE_CATEGORY_EXTERN(PSO);
 
+DECLARE_DWORD_ACCUMULATOR_STAT(TEXT("Runtime Graphics PSO Hitch Count"), STAT_RuntimeGraphicsPSOHitchCount, STATGROUP_PipelineStateCache);
+DECLARE_DWORD_ACCUMULATOR_STAT(TEXT("Runtime Compute PSO Hitch Count"), STAT_RuntimeComputePSOHitchCount, STATGROUP_PipelineStateCache);
+
+DECLARE_DWORD_ACCUMULATOR_STAT(TEXT("Active Graphics PSO Precache Requests"), STAT_ActiveGraphicsPSOPrecacheRequests, STATGROUP_PipelineStateCache);
+DECLARE_DWORD_ACCUMULATOR_STAT(TEXT("Active Compute PSO Precache Requestst"), STAT_ActiveComputePSOPrecacheRequests, STATGROUP_PipelineStateCache);
+
+
 static inline uint32 GetTypeHash(const FBoundShaderStateInput& Input)
 {
 	return GetTypeHash(Input.VertexDeclarationRHI)
@@ -94,6 +101,13 @@ static TAutoConsoleVariable<int32> CVarPSOEvictionTime(
 	ECVF_ReadOnly | ECVF_RenderThreadSafe
 );
 
+static TAutoConsoleVariable<int32> CVarPSORuntimeCreationHitchThreshold(
+	TEXT("r.PSO.RuntimeCreationHitchThreshold"),
+	20,
+	TEXT("Threshold for runtime PSO creation to count as a hitch (in msec) (default 20)"),
+	ECVF_ReadOnly | ECVF_RenderThreadSafe
+);
+
 #if RHI_RAYTRACING
 static TAutoConsoleVariable<int32> CVarRTPSOCacheSize(
 	TEXT("r.RayTracing.PSOCacheSize"),
@@ -120,6 +134,19 @@ static FAutoConsoleCommand DumpPipelineCmd(
 	FConsoleCommandDelegate::CreateStatic(DumpPipelineCacheStats)
 );
 
+static inline void CheckAndUpdateHitchCountStat(const FName& StatName, bool bIsRuntimePSO, uint64 StartTime)
+{
+	if (bIsRuntimePSO)
+	{
+		int32 RuntimePSOCreationHitchThreshold = CVarPSORuntimeCreationHitchThreshold.GetValueOnAnyThread();
+		double PSOCreationTimeMs = FPlatformTime::ToMilliseconds64(FPlatformTime::Cycles64() - StartTime);
+		if (PSOCreationTimeMs > RuntimePSOCreationHitchThreshold)
+		{
+			INC_DWORD_STAT_FName(StatName);
+		}
+	}
+}
+
 void SetComputePipelineState(FRHIComputeCommandList& RHICmdList, FRHIComputeShader* ComputeShader)
 {
 	RHICmdList.SetComputePipelineState(PipelineStateCache::GetAndOrCreateComputePipelineState(RHICmdList, ComputeShader, false), ComputeShader);
@@ -1056,6 +1083,12 @@ public:
 		return !ActivePrecachingTasks.IsEmpty();
 	}
 
+	uint32 NumActivePrecacheRequests()
+	{
+		FRWScopeLock ReadLock(PrecachePSOsRWLock, SLT_ReadOnly);
+		return ActivePrecachingTasks.Num();
+	}
+
 	void PrecacheFinished(const TPrecachedPSOInitializer& Initializer)
 	{
 		uint32 InitializerHash = TPrecachePipelineCacheDerived::PipelineStateInitializerHash(Initializer);
@@ -1569,7 +1602,11 @@ public:
 		if (Pipeline->IsCompute())
 		{
 			FComputePipelineState* ComputePipeline = static_cast<FComputePipelineState*>(Pipeline);
+			
+			uint64 StartTime = FPlatformTime::Cycles64();
 			ComputePipeline->RHIPipeline = RHICreateComputePipelineState(ComputePipeline->ComputeShader);
+			CheckAndUpdateHitchCountStat(GET_STATFNAME(STAT_RuntimeComputePSOHitchCount), !IsPrecachedPSO(Initializer), StartTime);
+
 			if (Initializer.bPSOPrecache)
 			{
 				if (ComputePipeline->RHIPipeline == nullptr)
@@ -1580,6 +1617,7 @@ public:
 				{
 					GPrecacheComputePipelineCache.PrecacheFinished(ComputePipeline->ComputeShader);
 				}
+				DEC_DWORD_STAT(STAT_ActiveComputePSOPrecacheRequests);
 			}
 		}
 		else
@@ -1621,7 +1659,10 @@ public:
 			}
 
 			FGraphicsPipelineState* GfxPipeline = static_cast<FGraphicsPipelineState*>(Pipeline);
+
+			uint64 StartTime = FPlatformTime::Cycles64();
 			GfxPipeline->RHIPipeline = bSkipCreation ? nullptr : RHICreateGraphicsPipelineState(Initializer);
+			CheckAndUpdateHitchCountStat(GET_STATFNAME(STAT_RuntimeGraphicsPSOHitchCount), !IsPrecachedPSO(Initializer), StartTime);
 
 			if (GfxPipeline->RHIPipeline)
 			{
@@ -1636,6 +1677,7 @@ public:
 			if (PSOPreCacheResult == EPSOPrecacheResult::Active)
 			{
 				GPrecacheGraphicsPipelineCache.PrecacheFinished(Initializer);
+				DEC_DWORD_STAT(STAT_ActiveGraphicsPSOPrecacheRequests);
 			}
 
 			if (Initializer.BoundShaderState.GetMeshShader())
@@ -1861,7 +1903,9 @@ FComputePipelineState* PipelineStateCache::GetAndOrCreateComputePipelineState(FR
 		}
 		else
 		{
+			uint64 StartTime = FPlatformTime::Cycles64();
 			OutCachedState->RHIPipeline = RHICreateComputePipelineState(OutCachedState->ComputeShader);
+			CheckAndUpdateHitchCountStat(GET_STATFNAME(STAT_RuntimeComputePSOHitchCount), !bFromFileCache, StartTime);
 		}
 
 		GComputePipelineCache.Add(ComputeShader, OutCachedState, LockFlags);
@@ -2301,7 +2345,10 @@ static void InternalCreateGraphicsPipelineState(const FGraphicsPipelineStateInit
 	else
 	{
 		check(GraphEvent == nullptr);
+		uint64 StartTime = FPlatformTime::Cycles64();
 		CachedState->RHIPipeline = RHICreateGraphicsPipelineState(Initializer);
+		CheckAndUpdateHitchCountStat(GET_STATFNAME(STAT_RuntimeGraphicsPSOHitchCount), !IsPrecachedPSO(Initializer), StartTime);
+
 		if (CachedState->RHIPipeline)
 		{
 			CachedState->SortKey = CachedState->RHIPipeline->GetSortKey();
@@ -2469,6 +2516,8 @@ void FPrecacheComputePipelineCache::OnNewPipelineStateCreated(const FPrecacheCom
 	FGraphEventRef GraphEvent = CachedState->CompletionEvent;
 	if (bDoAsyncCompile)
 	{
+		INC_DWORD_STAT(STAT_ActiveComputePSOPrecacheRequests);
+
 		check(GraphEvent != nullptr);
 		FGraphicsPipelineStateInitializer GraphicsPipelineStateInitializer;
 		GraphicsPipelineStateInitializer.bPSOPrecache = true;
@@ -2540,16 +2589,18 @@ FGraphEventRef FPrecacheGraphicsPipelineCache::PrecacheGraphicsPipelineState(con
 void FPrecacheGraphicsPipelineCache::OnNewPipelineStateCreated(const FGraphicsPipelineStateInitializer & Initializer, FGraphicsPipelineState * NewGraphicsPipelineState, bool bDoAsyncCompile)
 {
 	ValidateGraphicsPipelineStateInitializer(Initializer);
+	check((NewGraphicsPipelineState->CompletionEvent != nullptr) == bDoAsyncCompile);
 
 	// Mark as precache so it will try and use the background thread pool if available
 	bool bPSOPrecache = true;
 
 	FGraphicsPipelineStateInitializer InitializerCopy(Initializer);
-	InitializerCopy.bPSOPrecache = true;
-
+	InitializerCopy.bPSOPrecache = bPSOPrecache;
 
-	bool bHasCompletionEvent = (NewGraphicsPipelineState->CompletionEvent != nullptr);
-	check(bHasCompletionEvent == bDoAsyncCompile);
+	if (bDoAsyncCompile)
+	{
+		INC_DWORD_STAT(STAT_ActiveGraphicsPSOPrecacheRequests);
+	}
 
 	// Start the precache task
 	InternalCreateGraphicsPipelineState(InitializerCopy, EPSOPrecacheResult::Active, bDoAsyncCompile, bPSOPrecache, NewGraphicsPipelineState);
@@ -2610,6 +2661,16 @@ bool PipelineStateCache::IsPrecaching()
 	return GPrecacheGraphicsPipelineCache.IsPrecaching() || GPrecacheComputePipelineCache.IsPrecaching();
 }
 
+uint32 PipelineStateCache::NumActivePrecacheRequests()
+{
+	if (!IsPSOPrecachingEnabled())
+	{
+		return 0;
+	}
+
+	return GPrecacheGraphicsPipelineCache.NumActivePrecacheRequests() + GPrecacheComputePipelineCache.NumActivePrecacheRequests();
+}
+
 FRHIGraphicsPipelineState* ExecuteSetGraphicsPipelineState(FGraphicsPipelineState* GraphicsPipelineState)
 {
 	FRHIGraphicsPipelineState* RHIPipeline = GraphicsPipelineState->RHIPipeline;
diff --git a/Engine/Source/Runtime/RHI/Public/PipelineStateCache.h b/Engine/Source/Runtime/RHI/Public/PipelineStateCache.h
--- a/Engine/Source/Runtime/RHI/Public/PipelineStateCache.h
+++ b/Engine/Source/Runtime/RHI/Public/PipelineStateCache.h
@@ -118,7 +118,8 @@ namespace PipelineStateCache
 	/* Any async precaching operations still busy */
 	extern RHI_API bool						IsPrecaching();
 
-	
+	/* Return number of active or pending PSO precache requests */
+	extern RHI_API uint32					NumActivePrecacheRequests();
 }
 
 // Returns the shader index within the ray tracing pipeline or INDEX_NONE if given shader does not exist.
