Added breadcrumbs to RHI validation logging.

diff --git a/Engine/Source/Runtime/RHI/Private/RHIValidation.cpp b/Engine/Source/Runtime/RHI/Private/RHIValidation.cpp
--- a/Engine/Source/Runtime/RHI/Private/RHIValidation.cpp
+++ b/Engine/Source/Runtime/RHI/Private/RHIValidation.cpp
@@ -532,6 +532,19 @@ public:
 
 thread_local TConstArrayView<const TCHAR*> FRHIValidationBreadcrumbScope::Breadcrumbs;
 
+static FString GetBreadcrumbPath()
+{
+	FString BreadcrumbMessage;
+
+	for (int32 Index = 0; Index < FRHIValidationBreadcrumbScope::Breadcrumbs.Num() - 1; ++Index)
+	{
+		BreadcrumbMessage += (FRHIValidationBreadcrumbScope::Breadcrumbs[Index] + FString(TEXT("/")));
+	}
+
+	BreadcrumbMessage += FRHIValidationBreadcrumbScope::Breadcrumbs.Last();
+	return BreadcrumbMessage;
+}
+
 void FValidationRHI::ReportValidationFailure(const TCHAR* InMessage)
 {
 	// Report failures only once per session, since many of them will happen repeatedly. This is similar to what ensure() does, but
@@ -554,19 +567,11 @@ void FValidationRHI::ReportValidationFailure(const TCHAR* InMessage)
 
 	if (!FRHIValidationBreadcrumbScope::Breadcrumbs.IsEmpty())
 	{
-		FString BreadcrumbMessage;
-
-		for (int32 Index = 0; Index < FRHIValidationBreadcrumbScope::Breadcrumbs.Num() - 1; ++Index)
-		{
-			BreadcrumbMessage += (FRHIValidationBreadcrumbScope::Breadcrumbs[Index] + FString(TEXT("/")));
-		}
-
-		BreadcrumbMessage += FRHIValidationBreadcrumbScope::Breadcrumbs.Last();
 		Message = FString::Printf(
 			TEXT("%s")
 			TEXT("Breadcrumbs: %s\n")
 			TEXT("--------------------------------------------------------------------\n"),
-			InMessage, *BreadcrumbMessage);
+			InMessage, *GetBreadcrumbPath());
 	}
 	else
 	{
@@ -1199,23 +1204,27 @@ namespace RHIValidation
 	{
 		void* Trace = CaptureBacktrace();
 
+		FString BreadcrumbMessage = GetBreadcrumbPath();
+
 		if (CreateTrace)
 		{
-			FPlatformMisc::LowLevelOutputDebugStringf(TEXT("\n%s: Type: %s, %s, CreateTrace: 0x%p, %sTrace: 0x%p\n"),
+			FPlatformMisc::LowLevelOutputDebugStringf(TEXT("\n%s: Type: %s, %s, CreateTrace: 0x%p, %sTrace: 0x%p, %s\n"),
 				*GetResourceDebugName(Resource, SubresourceIndex),
 				Type,
 				LogStr,
 				CreateTrace,
 				TracePrefix,
-				Trace);
+				Trace,
+				*BreadcrumbMessage);
 		}
 		else
 		{
-			FPlatformMisc::LowLevelOutputDebugStringf(TEXT("\n%s: Type: %s, %s, Trace: 0x%p\n"),
+			FPlatformMisc::LowLevelOutputDebugStringf(TEXT("\n%s: Type: %s, %s, Trace: 0x%p, %s\n"),
 				*GetResourceDebugName(Resource, SubresourceIndex),
 				Type,
 				LogStr,
-				Trace);
+				Trace,
+				*BreadcrumbMessage);
 		}
 
 		return Trace;
