Disable some checks in the FRHICommandListBase destructor in shipping/test configurations, as they can cause crashes on shutdown due to otherwise benign bugs.

#rnx

diff --git a/Engine/Source/Runtime/RHI/Private/RHICommandList.cpp b/Engine/Source/Runtime/RHI/Private/RHICommandList.cpp
--- a/Engine/Source/Runtime/RHI/Private/RHICommandList.cpp
+++ b/Engine/Source/Runtime/RHI/Private/RHICommandList.cpp
@@ -312,13 +312,18 @@ FRHICommandListBase::FRHICommandListBase(FRHICommandListBase&& Other)
 
 FRHICommandListBase::~FRHICommandListBase()
 {
+#if DO_CHECK
+	// Some configurations enable checks in shipping/test, particularly server builds. Skip these checks explicitly in that case, as they can fire very late in
+	// the shutdown process and crash in unexpected ways because the log output channel has already been destroyed. Also, having pending commands
+	// on shutdown shouldn't really be a fatal error, it's a fairly harmless condition.
+#if !UE_BUILD_SHIPPING && !UE_BUILD_TEST
 	checkf(!HasCommands(), TEXT("FRHICommandListBase has been deleted while it still contained commands. The command list was not submitted."));
 
-#if DO_CHECK
 	for (void* Data : QueryBatchData)
 	{
 		check(Data == nullptr);
 	}
+#endif
 
 	if (PersistentState.RecordingThread == ERecordingThread::Render)
 	{
