Fixing RHIValidationContext crash on null Texture resources.
bGfxPSOSet shouldn't be reset when calling BeginRenderPass and EndRenderPass.

diff --git a/Engine/Source/Runtime/RHI/Public/RHIValidationContext.h b/Engine/Source/Runtime/RHI/Public/RHIValidationContext.h
--- a/Engine/Source/Runtime/RHI/Public/RHIValidationContext.h
+++ b/Engine/Source/Runtime/RHI/Public/RHIValidationContext.h
@@ -19,7 +19,10 @@ inline void ValidateShaderParameters(RHIValidation::FTracker* Tracker, RHIValida
 		switch (Parameter.Type)
 		{
 		case FRHIShaderParameterResource::EType::Texture:
-			Tracker->Assert(static_cast<FRHITexture*>(Parameter.Resource)->GetWholeResourceIdentitySRV(), InRequiredAccess);
+			if (FRHITexture* Texture = static_cast<FRHITexture*>(Parameter.Resource))
+			{
+				Tracker->Assert(Texture->GetWholeResourceIdentitySRV(), InRequiredAccess);
+			}
 			break;
 		case FRHIShaderParameterResource::EType::ResourceView:
 			if (FRHIShaderResourceView* SRV = static_cast<FRHIShaderResourceView*>(Parameter.Resource))
@@ -947,7 +950,6 @@ public:
 		State.bInsideBeginRenderPass = true;
 		State.RenderPassInfo = InInfo;
 		State.RenderPassName = InName;
-		State.bGfxPSOSet = false;
 
 		FIntVector ViewDimensions = FIntVector(0);
 
@@ -1064,7 +1066,6 @@ public:
 		RHIContext->RHIEndRenderPass();
 		State.bInsideBeginRenderPass = false;
 		State.PreviousRenderPassName = State.RenderPassName;
-		State.bGfxPSOSet = false;
 	}
 
 	virtual void RHIWriteGPUFence(FRHIGPUFence* FenceRHI) override final
