Validate size of indirect buffer args in RHIDispatchIndirectComputeShader, RHIDrawPrimitiveIndirect, RHIDrawIndexedIndirect, RHIDrawIndexedPrimitiveIndirect, RHIDispatchIndirectMeshShader, and RHIRayTraceDispatchIndirect.

diff --git a/Engine/Source/Runtime/RHI/Private/RHIValidation.cpp b/Engine/Source/Runtime/RHI/Private/RHIValidation.cpp
--- a/Engine/Source/Runtime/RHI/Private/RHIValidation.cpp
+++ b/Engine/Source/Runtime/RHI/Private/RHIValidation.cpp
@@ -545,6 +545,45 @@ static FString GetBreadcrumbPath()
 	return BreadcrumbMessage;
 }
 
+// FlushType: Thread safe
+void FValidationRHI::RHIBindDebugLabelName(FRHITexture* Texture, const TCHAR* Name)
+{
+	check(IsInRenderingThread());
+
+	FString NameCopyRT = Name;
+	FRHICommandListExecutor::GetImmediateCommandList().EnqueueLambda([Texture, NameCopyRHIT = MoveTemp(NameCopyRT)](FRHICommandListImmediate& RHICmdList)
+		{
+			((FValidationContext&)RHICmdList.GetContext()).Tracker->Rename(Texture->GetTrackerResource(), *NameCopyRHIT);
+		});
+
+	RHI->RHIBindDebugLabelName(Texture, Name);
+}
+
+void FValidationRHI::RHIBindDebugLabelName(FRHIBuffer* Buffer, const TCHAR* Name)
+{
+	check(IsInRenderingThread());
+
+	FString NameCopyRT = Name;
+	FRHICommandListExecutor::GetImmediateCommandList().EnqueueLambda([Buffer, NameCopyRHIT = MoveTemp(NameCopyRT)](FRHICommandListImmediate& RHICmdList)
+		{
+			((FValidationContext&)RHICmdList.GetContext()).Tracker->Rename(Buffer, *NameCopyRHIT);
+		});
+
+	RHI->RHIBindDebugLabelName(Buffer, Name);
+}
+
+void FValidationRHI::RHIBindDebugLabelName(FRHIUnorderedAccessView* UnorderedAccessViewRHI, const TCHAR* Name)
+{
+	RHIValidation::FResource* Resource = UnorderedAccessViewRHI->ViewIdentity.Resource;
+	FString NameCopyRT = Name;
+	FRHICommandListExecutor::GetImmediateCommandList().EnqueueLambda([Resource, NameCopyRHIT = MoveTemp(NameCopyRT)](FRHICommandListImmediate& RHICmdList)
+		{
+			((FValidationContext&)RHICmdList.GetContext()).Tracker->Rename(Resource, *NameCopyRHIT);
+		});
+
+	RHI->RHIBindDebugLabelName(UnorderedAccessViewRHI, Name);
+}
+
 void FValidationRHI::ReportValidationFailure(const TCHAR* InMessage)
 {
 	// Report failures only once per session, since many of them will happen repeatedly. This is similar to what ensure() does, but
@@ -567,11 +606,19 @@ void FValidationRHI::ReportValidationFailure(const TCHAR* InMessage)
 
 	if (!FRHIValidationBreadcrumbScope::Breadcrumbs.IsEmpty())
 	{
+		FString BreadcrumbMessage;
+
+		for (int32 Index = 0; Index < FRHIValidationBreadcrumbScope::Breadcrumbs.Num() - 1; ++Index)
+		{
+			BreadcrumbMessage += (FRHIValidationBreadcrumbScope::Breadcrumbs[Index] + FString(TEXT("/")));
+		}
+
+		BreadcrumbMessage += FRHIValidationBreadcrumbScope::Breadcrumbs.Last();
 		Message = FString::Printf(
 			TEXT("%s")
 			TEXT("Breadcrumbs: %s\n")
 			TEXT("--------------------------------------------------------------------\n"),
-			InMessage, *GetBreadcrumbPath());
+			InMessage, *BreadcrumbMessage);
 	}
 	else
 	{
@@ -1204,27 +1251,23 @@ namespace RHIValidation
 	{
 		void* Trace = CaptureBacktrace();
 
-		FString BreadcrumbMessage = GetBreadcrumbPath();
-
 		if (CreateTrace)
 		{
-			FPlatformMisc::LowLevelOutputDebugStringf(TEXT("\n%s: Type: %s, %s, CreateTrace: 0x%p, %sTrace: 0x%p, %s\n"),
+			FPlatformMisc::LowLevelOutputDebugStringf(TEXT("\n%s: Type: %s, %s, CreateTrace: 0x%p, %sTrace: 0x%p\n"),
 				*GetResourceDebugName(Resource, SubresourceIndex),
 				Type,
 				LogStr,
 				CreateTrace,
 				TracePrefix,
-				Trace,
-				*BreadcrumbMessage);
+				Trace);
 		}
 		else
 		{
-			FPlatformMisc::LowLevelOutputDebugStringf(TEXT("\n%s: Type: %s, %s, Trace: 0x%p, %s\n"),
+			FPlatformMisc::LowLevelOutputDebugStringf(TEXT("\n%s: Type: %s, %s, Trace: 0x%p\n"),
 				*GetResourceDebugName(Resource, SubresourceIndex),
 				Type,
 				LogStr,
-				Trace,
-				*BreadcrumbMessage);
+				Trace);
 		}
 
 		return Trace;
diff --git a/Engine/Source/Runtime/RHI/Public/RHIValidation.h b/Engine/Source/Runtime/RHI/Public/RHIValidation.h
--- a/Engine/Source/Runtime/RHI/Public/RHIValidation.h
+++ b/Engine/Source/Runtime/RHI/Public/RHIValidation.h
@@ -9,12 +9,9 @@
 #include "RHI.h"
 #include "RHIValidationCommon.h"
 #include "RHIValidationUtils.h"
-#include "RHIValidationContext.h"
 #include "DataDrivenShaderPlatformInfo.h"
 
 #if ENABLE_RHI_VALIDATION
-class FValidationComputeContext;
-class FValidationContext;
 
 // Controls whether BUF_SourceCopy should be validate or not.
 extern RHI_API bool GRHIValidateBufferSourceCopy;
@@ -36,6 +33,20 @@ public:
 	RHI_API FValidationRHI(FDynamicRHI* InRHI);
 	RHI_API virtual ~FValidationRHI();
 
+	static inline void ValidateIndirectArgsBuffer(FRHIBuffer* ArgumentBuffer, uint32 ArgumentOffset)
+	{
+		const TCHAR* BufferName = ArgumentBuffer->GetDebugName();
+		RHI_VALIDATION_CHECK(EnumHasAnyFlags(ArgumentBuffer->GetUsage(), EBufferUsageFlags::VertexBuffer), *FString::Printf(TEXT("The buffer %s needs to be a vertex buffer to be used as an indirect dispatch parameter."), BufferName));
+		RHI_VALIDATION_CHECK(EnumHasAnyFlags(ArgumentBuffer->GetUsage(), EBufferUsageFlags::DrawIndirect), *FString::Printf(TEXT("The buffer %s for indirect dispatch parameters was not flagged with BUF_DrawIndirect."), BufferName));
+		RHI_VALIDATION_CHECK((ArgumentOffset % 4) == 0, TEXT("IndirectArgOffset for compute shader indirect dispatch needs to be a multiple of 4."));
+		RHI_VALIDATION_CHECK((ArgumentOffset + ArgumentBuffer->GetStride()) <= ArgumentBuffer->GetSize(),
+			*FString::Printf(TEXT("Indirect parameters buffer for compute shader indirect dispatch at byte offset %d doesn't have enough room for one element."), ArgumentOffset));
+#if PLATFORM_DISPATCH_INDIRECT_ARGUMENT_BOUNDARY_SIZE != 0
+		RHI_VALIDATION_CHECK(IndirectArgOffset / PLATFORM_DISPATCH_INDIRECT_ARGUMENT_BOUNDARY_SIZE == (ArgumentOffset + Stride - 1) / PLATFORM_DISPATCH_INDIRECT_ARGUMENT_BOUNDARY_SIZE,
+			*FString::Printf(TEXT("Compute indirect dispatch arguments cannot cross %d byte boundary."), PLATFORM_DISPATCH_INDIRECT_ARGUMENT_BOUNDARY_SIZE));
+#endif // #if PLATFORM_DISPATCH_INDIRECT_ARGUMENT_BOUNDARY_SIZE != 0
+	}
+
 	virtual void Init() override final
 	{
 		RHI->Init();
@@ -839,43 +850,11 @@ public:
 	}
 
 	// FlushType: Thread safe
-	virtual void RHIBindDebugLabelName(FRHITexture* Texture, const TCHAR* Name) override final
-	{
-		check(IsInRenderingThread());
+	virtual void RHIBindDebugLabelName(FRHITexture* Texture, const TCHAR* Name) override final;
 
-		FString NameCopyRT = Name;
-		FRHICommandListExecutor::GetImmediateCommandList().EnqueueLambda([Texture, NameCopyRHIT = MoveTemp(NameCopyRT)](FRHICommandListImmediate& RHICmdList)
-		{
-			((FValidationContext&)RHICmdList.GetContext()).Tracker->Rename(Texture->GetTrackerResource(), *NameCopyRHIT);
-		});
-
-		RHI->RHIBindDebugLabelName(Texture, Name);
-	}
+	virtual void RHIBindDebugLabelName(FRHIBuffer* Buffer, const TCHAR* Name) override final;
 
-	virtual void RHIBindDebugLabelName(FRHIBuffer* Buffer, const TCHAR* Name) override final
-	{
-		check(IsInRenderingThread());
-
-		FString NameCopyRT = Name;
-		FRHICommandListExecutor::GetImmediateCommandList().EnqueueLambda([Buffer, NameCopyRHIT = MoveTemp(NameCopyRT)](FRHICommandListImmediate& RHICmdList)
-		{
-			((FValidationContext&)RHICmdList.GetContext()).Tracker->Rename(Buffer, *NameCopyRHIT);
-		});
-
-		RHI->RHIBindDebugLabelName(Buffer, Name);
-	}
-
-	virtual void RHIBindDebugLabelName(FRHIUnorderedAccessView* UnorderedAccessViewRHI, const TCHAR* Name) override final
-	{
-		RHIValidation::FResource* Resource = UnorderedAccessViewRHI->ViewIdentity.Resource;
-		FString NameCopyRT = Name;
-		FRHICommandListExecutor::GetImmediateCommandList().EnqueueLambda([Resource, NameCopyRHIT = MoveTemp(NameCopyRT)](FRHICommandListImmediate& RHICmdList)
-		{
-			((FValidationContext&)RHICmdList.GetContext()).Tracker->Rename(Resource, *NameCopyRHIT);
-		});
-
-		RHI->RHIBindDebugLabelName(UnorderedAccessViewRHI, Name);
-	}
+	virtual void RHIBindDebugLabelName(FRHIUnorderedAccessView* UnorderedAccessViewRHI, const TCHAR* Name) override final;
 
 	/**
 	* Reads the contents of a texture to an output buffer (non MSAA and MSAA) and returns it as a FColor array.
diff --git a/Engine/Source/Runtime/RHI/Public/RHIValidationContext.h b/Engine/Source/Runtime/RHI/Public/RHIValidationContext.h
--- a/Engine/Source/Runtime/RHI/Public/RHIValidationContext.h
+++ b/Engine/Source/Runtime/RHI/Public/RHIValidationContext.h
@@ -89,6 +89,7 @@ public:
 	virtual void RHIDispatchIndirectComputeShader(FRHIBuffer* ArgumentBuffer, uint32 ArgumentOffset) override final
 	{
 		checkf(State.bComputePSOSet, TEXT("A Compute PSO has to be set to set resources into a shader!"));
+		FValidationRHI::ValidateIndirectArgsBuffer(ArgumentBuffer, ArgumentOffset);
 		Tracker->Assert(ArgumentBuffer->GetWholeResourceIdentity(), ERHIAccess::IndirectArgs);
 
 		RHIContext->RHIDispatchIndirectComputeShader(ArgumentBuffer, ArgumentOffset);
@@ -377,6 +378,7 @@ public:
 	virtual void RHIDispatchIndirectComputeShader(FRHIBuffer* ArgumentBuffer, uint32 ArgumentOffset) override final
 	{
 		checkf(State.bComputePSOSet, TEXT("A Compute PSO has to be set to set resources into a shader!"));
+		FValidationRHI::ValidateIndirectArgsBuffer(ArgumentBuffer, ArgumentOffset);
 		Tracker->Assert(ArgumentBuffer->GetWholeResourceIdentity(), ERHIAccess::IndirectArgs);
 
 		RHIContext->RHIDispatchIndirectComputeShader(ArgumentBuffer, ArgumentOffset);
@@ -860,6 +862,7 @@ public:
 	virtual void RHIDrawPrimitiveIndirect(FRHIBuffer* ArgumentBuffer, uint32 ArgumentOffset) override final
 	{
 		checkf(State.bGfxPSOSet, TEXT("A Graphics PSO has to be set to draw!"));
+		FValidationRHI::ValidateIndirectArgsBuffer(ArgumentBuffer, ArgumentOffset);
 		SetupDrawing();
 		Tracker->Assert(ArgumentBuffer->GetWholeResourceIdentity(), ERHIAccess::IndirectArgs);
 		RHIContext->RHIDrawPrimitiveIndirect(ArgumentBuffer, ArgumentOffset);
@@ -869,6 +872,7 @@ public:
 	virtual void RHIDrawIndexedIndirect(FRHIBuffer* IndexBufferRHI, FRHIBuffer* ArgumentsBufferRHI, int32 DrawArgumentsIndex, uint32 NumInstances) override final
 	{
 		checkf(State.bGfxPSOSet, TEXT("A Graphics PSO has to be set to draw!"));
+		FValidationRHI::ValidateIndirectArgsBuffer(ArgumentsBufferRHI, DrawArgumentsIndex * ArgumentsBufferRHI->GetStride());
 		SetupDrawing();
 		Tracker->Assert(ArgumentsBufferRHI->GetWholeResourceIdentity(), ERHIAccess::IndirectArgs);
 		Tracker->Assert(IndexBufferRHI->GetWholeResourceIdentity(), ERHIAccess::VertexOrIndexBuffer);
@@ -889,6 +893,7 @@ public:
 	virtual void RHIDrawIndexedPrimitiveIndirect(FRHIBuffer* IndexBuffer, FRHIBuffer* ArgumentBuffer, uint32 ArgumentOffset) override final
 	{
 		checkf(State.bGfxPSOSet, TEXT("A Graphics PSO has to be set to draw!"));
+		FValidationRHI::ValidateIndirectArgsBuffer(ArgumentBuffer, ArgumentOffset);
 		SetupDrawing();
 		Tracker->Assert(ArgumentBuffer->GetWholeResourceIdentity(), ERHIAccess::IndirectArgs);
 		Tracker->Assert(IndexBuffer->GetWholeResourceIdentity(), ERHIAccess::VertexOrIndexBuffer);
@@ -907,6 +912,7 @@ public:
 	virtual void RHIDispatchIndirectMeshShader(FRHIBuffer* ArgumentBuffer, uint32 ArgumentOffset) final override
 	{
 		checkf(State.bGfxPSOSet, TEXT("A Graphics PSO has to be set to draw!"));
+		FValidationRHI::ValidateIndirectArgsBuffer(ArgumentBuffer, ArgumentOffset);
 		SetupDrawing();
 		Tracker->Assert(ArgumentBuffer->GetWholeResourceIdentity(), ERHIAccess::IndirectArgs);
 		RHIContext->RHIDispatchIndirectMeshShader(ArgumentBuffer, ArgumentOffset);
@@ -1134,6 +1140,7 @@ public:
 		const FRayTracingShaderBindings& GlobalResourceBindings,
 		FRHIBuffer* ArgumentBuffer, uint32 ArgumentOffset) override final
 	{
+		FValidationRHI::ValidateIndirectArgsBuffer(ArgumentBuffer, ArgumentOffset);
 		Tracker->Assert(ArgumentBuffer->GetWholeResourceIdentity(), ERHIAccess::IndirectArgs);
 		Tracker->Assert(ArgumentBuffer->GetWholeResourceIdentity(), ERHIAccess::SRVCompute);
 
