Fix incremental build error when PLATFORM_DISPATCH_INDIRECT_ARGUMENT_BOUNDARY_SIZE != 0

diff --git a/Engine/Source/Runtime/RHI/Public/RHIValidation.h b/Engine/Source/Runtime/RHI/Public/RHIValidation.h
--- a/Engine/Source/Runtime/RHI/Public/RHIValidation.h
+++ b/Engine/Source/Runtime/RHI/Public/RHIValidation.h
@@ -42,7 +42,7 @@ public:
 		RHI_VALIDATION_CHECK((ArgumentOffset + ArgumentBuffer->GetStride()) <= ArgumentBuffer->GetSize(),
 			*FString::Printf(TEXT("Indirect parameters buffer for compute shader indirect dispatch at byte offset %d doesn't have enough room for one element."), ArgumentOffset));
 #if PLATFORM_DISPATCH_INDIRECT_ARGUMENT_BOUNDARY_SIZE != 0
-		RHI_VALIDATION_CHECK(IndirectArgOffset / PLATFORM_DISPATCH_INDIRECT_ARGUMENT_BOUNDARY_SIZE == (ArgumentOffset + Stride - 1) / PLATFORM_DISPATCH_INDIRECT_ARGUMENT_BOUNDARY_SIZE,
+		RHI_VALIDATION_CHECK(ArgumentOffset / PLATFORM_DISPATCH_INDIRECT_ARGUMENT_BOUNDARY_SIZE == (ArgumentOffset + ArgumentBuffer->GetStride() - 1) / PLATFORM_DISPATCH_INDIRECT_ARGUMENT_BOUNDARY_SIZE,
 			*FString::Printf(TEXT("Compute indirect dispatch arguments cannot cross %d byte boundary."), PLATFORM_DISPATCH_INDIRECT_ARGUMENT_BOUNDARY_SIZE));
 #endif // #if PLATFORM_DISPATCH_INDIRECT_ARGUMENT_BOUNDARY_SIZE != 0
 	}
