Validate ThreadGroupCountX/Y/Z in RHIDispatchComputeShader and RHIDispatchMeshShader.

diff --git a/Engine/Source/Runtime/RHI/Public/RHIValidation.h b/Engine/Source/Runtime/RHI/Public/RHIValidation.h
--- a/Engine/Source/Runtime/RHI/Public/RHIValidation.h
+++ b/Engine/Source/Runtime/RHI/Public/RHIValidation.h
@@ -33,6 +33,17 @@ public:
 	RHI_API FValidationRHI(FDynamicRHI* InRHI);
 	RHI_API virtual ~FValidationRHI();
 
+	static inline void ValidateThreadGroupCount(uint32 ThreadGroupCountX, uint32 ThreadGroupCountY, uint32 ThreadGroupCountZ)
+	{
+		RHI_VALIDATION_CHECK(ThreadGroupCountX > 0, TEXT("ThreadGroupCount X must be greater than 0."));
+		RHI_VALIDATION_CHECK(ThreadGroupCountY > 0, TEXT("ThreadGroupCount Y must be greater than 0."));
+		RHI_VALIDATION_CHECK(ThreadGroupCountZ > 0, TEXT("ThreadGroupCount Z must be greater than 0."));
+
+		RHI_VALIDATION_CHECK((ThreadGroupCountX <= (uint32)GRHIMaxDispatchThreadGroupsPerDimension.X), *FString::Printf(TEXT("ThreadGroupCountX is invalid: %u. Must be greater than 0 and less than %d"), ThreadGroupCountX, GRHIMaxDispatchThreadGroupsPerDimension.X));
+		RHI_VALIDATION_CHECK((ThreadGroupCountY <= (uint32)GRHIMaxDispatchThreadGroupsPerDimension.Y), *FString::Printf(TEXT("ThreadGroupCountY is invalid: %u. Must be greater than 0 and less than %d"), ThreadGroupCountY, GRHIMaxDispatchThreadGroupsPerDimension.Y));
+		RHI_VALIDATION_CHECK((ThreadGroupCountZ <= (uint32)GRHIMaxDispatchThreadGroupsPerDimension.Z), *FString::Printf(TEXT("ThreadGroupCountZ is invalid: %u. Must be greater than 0 and less than %d"), ThreadGroupCountZ, GRHIMaxDispatchThreadGroupsPerDimension.Z));
+	}
+
 	static inline void ValidateIndirectArgsBuffer(FRHIBuffer* ArgumentBuffer, uint32 ArgumentOffset)
 	{
 		const TCHAR* BufferName = ArgumentBuffer->GetDebugName();
diff --git a/Engine/Source/Runtime/RHI/Public/RHIValidationContext.h b/Engine/Source/Runtime/RHI/Public/RHIValidationContext.h
--- a/Engine/Source/Runtime/RHI/Public/RHIValidationContext.h
+++ b/Engine/Source/Runtime/RHI/Public/RHIValidationContext.h
@@ -82,6 +82,7 @@ public:
 	virtual void RHIDispatchComputeShader(uint32 ThreadGroupCountX, uint32 ThreadGroupCountY, uint32 ThreadGroupCountZ) override final
 	{
 		checkf(State.bComputePSOSet, TEXT("A Compute PSO has to be set to set resources into a shader!"));
+		FValidationRHI::ValidateThreadGroupCount(ThreadGroupCountX, ThreadGroupCountY, ThreadGroupCountZ);
 		RHIContext->RHIDispatchComputeShader(ThreadGroupCountX, ThreadGroupCountY, ThreadGroupCountZ);
 		Tracker->Dispatch();
 	}
@@ -371,6 +372,7 @@ public:
 	virtual void RHIDispatchComputeShader(uint32 ThreadGroupCountX, uint32 ThreadGroupCountY, uint32 ThreadGroupCountZ) override final
 	{
 		checkf(State.bComputePSOSet, TEXT("A Compute PSO has to be set to set resources into a shader!"));
+		FValidationRHI::ValidateThreadGroupCount(ThreadGroupCountX, ThreadGroupCountY, ThreadGroupCountZ);
 		RHIContext->RHIDispatchComputeShader(ThreadGroupCountX, ThreadGroupCountY, ThreadGroupCountZ);
 		Tracker->Dispatch();
 	}
@@ -904,6 +906,7 @@ public:
 	virtual void RHIDispatchMeshShader(uint32 ThreadGroupCountX, uint32 ThreadGroupCountY, uint32 ThreadGroupCountZ) final override
 	{
 		checkf(State.bGfxPSOSet, TEXT("A Graphics PSO has to be set to draw!"));
+		FValidationRHI::ValidateThreadGroupCount(ThreadGroupCountX, ThreadGroupCountY, ThreadGroupCountZ);
 		SetupDrawing();
 		RHIContext->RHIDispatchMeshShader(ThreadGroupCountX, ThreadGroupCountY, ThreadGroupCountZ);
 		Tracker->Draw();
