Moving the RHI_ENABLE_RESOURCE_INFO defines to RHIFwd.h

diff --git a/Engine/Source/Runtime/RHI/Public/RHIFwd.h b/Engine/Source/Runtime/RHI/Public/RHIFwd.h
--- a/Engine/Source/Runtime/RHI/Public/RHIFwd.h
+++ b/Engine/Source/Runtime/RHI/Public/RHIFwd.h
@@ -4,6 +4,18 @@
 
 #include "Templates/RefCounting.h"
 
+// RHI_WANT_RESOURCE_INFO should be controlled by the RHI module.
+#ifndef RHI_WANT_RESOURCE_INFO
+#define RHI_WANT_RESOURCE_INFO 0
+#endif
+
+// RHI_FORCE_DISABLE_RESOURCE_INFO can be defined anywhere else, like in GlobalDefinitions.
+#ifndef RHI_FORCE_DISABLE_RESOURCE_INFO
+#define RHI_FORCE_DISABLE_RESOURCE_INFO 0
+#endif
+
+#define RHI_ENABLE_RESOURCE_INFO (RHI_WANT_RESOURCE_INFO && !RHI_FORCE_DISABLE_RESOURCE_INFO)
+
 // Basic Types
 
 namespace ERHIFeatureLevel { enum Type : int; }
diff --git a/Engine/Source/Runtime/RHI/Public/RHIResources.h b/Engine/Source/Runtime/RHI/Public/RHIResources.h
--- a/Engine/Source/Runtime/RHI/Public/RHIResources.h
+++ b/Engine/Source/Runtime/RHI/Public/RHIResources.h
@@ -27,18 +27,6 @@
 #include "Misc/CoreDelegates.h"
 #include "RHIFwd.h"
 
-// RHI_WANT_RESOURCE_INFO should be controlled by the RHI module.
-#ifndef RHI_WANT_RESOURCE_INFO
-#define RHI_WANT_RESOURCE_INFO 0
-#endif
-
-// RHI_FORCE_DISABLE_RESOURCE_INFO can be defined anywhere else, like in GlobalDefinitions.
-#ifndef RHI_FORCE_DISABLE_RESOURCE_INFO
-#define RHI_FORCE_DISABLE_RESOURCE_INFO 0
-#endif
-
-#define RHI_ENABLE_RESOURCE_INFO (RHI_WANT_RESOURCE_INFO && !RHI_FORCE_DISABLE_RESOURCE_INFO)
-
 class FRHIComputeCommandList;
 class FRHICommandListImmediate;
 struct FClearValueBinding;
