Allow running editor with: -game -featureleveles31

diff --git a/Engine/Source/Runtime/RHI/Private/Windows/WindowsDynamicRHI.cpp b/Engine/Source/Runtime/RHI/Private/Windows/WindowsDynamicRHI.cpp
--- a/Engine/Source/Runtime/RHI/Private/Windows/WindowsDynamicRHI.cpp
+++ b/Engine/Source/Runtime/RHI/Private/Windows/WindowsDynamicRHI.cpp
@@ -640,7 +640,7 @@ static TOptional<EWindowsRHI> ChooseForcedRHI(TOptional<ERHIFeatureLevel::Type>
 	}
 
 	// FeatureLevelES31 is also a command line override, so it will determine the underlying RHI unless one is specified
-	if ((FPlatformProperties::RequiresCookedData() || !GIsEditor) && ForcedFeatureLevel == ERHIFeatureLevel::ES3_1)
+	if (FPlatformProperties::RequiresCookedData() && ForcedFeatureLevel == ERHIFeatureLevel::ES3_1)
 	{
 		if (ForcedRHI)
 		{
