[RHI]
* Added IWYU pragmas to includes to make sure iwyu is not replacing RHI includes with the internal ones

diff --git a/Engine/Source/Runtime/RHI/Public/RHI.h b/Engine/Source/Runtime/RHI/Public/RHI.h
--- a/Engine/Source/Runtime/RHI/Public/RHI.h
+++ b/Engine/Source/Runtime/RHI/Public/RHI.h
@@ -451,7 +451,7 @@ extern RHI_API int32 GDrawUPIndexCheckCount;
 /** true for each VET that is supported. One-to-one mapping with EVertexElementType */
 extern RHI_API class FVertexElementTypeSupportInfo GVertexElementTypeSupport;
 
-#include "MultiGPU.h"
+#include "MultiGPU.h" // IWYU pragma: export
 
 /** Whether the next frame should profile the GPU. */
 extern RHI_API bool GTriggerGPUProfile;
@@ -2079,7 +2079,7 @@ struct FRHITrackedAccessInfo
 	ERHIAccess Access = ERHIAccess::Unknown;
 };
 
-#include "RHIValidationCommon.h"
+#include "RHIValidationCommon.h" // IWYU pragma: export
 
 // Opaque data structure used to represent a pending resource transition in the RHI.
 struct FRHITransition
@@ -2513,8 +2513,8 @@ DECLARE_MEMORY_STAT_POOL_EXTERN(TEXT("Pixel buffer memory"),STAT_PixelBufferMemo
 
 
 // RHI base resource types.
-#include "RHIResources.h"
-#include "DynamicRHI.h"
+#include "RHIResources.h" // IWYU pragma: export
+#include "DynamicRHI.h" // IWYU pragma: export
 
 /** Initializes the RHI. */
 extern RHI_API void RHIInit(bool bHasEditorToken);
@@ -2531,7 +2531,7 @@ DECLARE_DELEGATE_OneParam(FRHIPanicEvent, const FName&);
 extern RHI_API FRHIPanicEvent& RHIGetPanicDelegate();
 
 // RHI utility functions that depend on the RHI definitions.
-#include "RHIUtilities.h"
+#include "RHIUtilities.h" // IWYU pragma: export
 
 inline void FRHITransition::Cleanup() const
 {
