Adding FRHIShader to RHIFwd.h

diff --git a/Engine/Source/Runtime/RHI/Public/RHIFwd.h b/Engine/Source/Runtime/RHI/Public/RHIFwd.h
--- a/Engine/Source/Runtime/RHI/Public/RHIFwd.h
+++ b/Engine/Source/Runtime/RHI/Public/RHIFwd.h
@@ -67,6 +67,7 @@ class FRHIRenderQuery;
 class FRHIRenderQueryPool;
 class FRHIResource;
 class FRHISamplerState;
+class FRHIShader;
 class FRHIShaderLibrary;
 class FRHIShaderResourceView;
 class FRHIStagingBuffer;
