Add indirect buffer name, stride, and offset to all validation error messages.

diff --git a/Engine/Source/Runtime/RHI/Public/RHIValidation.h b/Engine/Source/Runtime/RHI/Public/RHIValidation.h
--- a/Engine/Source/Runtime/RHI/Public/RHIValidation.h
+++ b/Engine/Source/Runtime/RHI/Public/RHIValidation.h
@@ -46,15 +46,16 @@ public:
 
 	static inline void ValidateIndirectArgsBuffer(FRHIBuffer* ArgumentBuffer, uint32 ArgumentOffset)
 	{
-		const TCHAR* BufferName = ArgumentBuffer->GetDebugName();
-		RHI_VALIDATION_CHECK(EnumHasAnyFlags(ArgumentBuffer->GetUsage(), EBufferUsageFlags::VertexBuffer), *FString::Printf(TEXT("The buffer %s needs to be a vertex buffer to be used as an indirect dispatch parameter."), BufferName));
-		RHI_VALIDATION_CHECK(EnumHasAnyFlags(ArgumentBuffer->GetUsage(), EBufferUsageFlags::DrawIndirect), *FString::Printf(TEXT("The buffer %s for indirect dispatch parameters was not flagged with BUF_DrawIndirect."), BufferName));
-		RHI_VALIDATION_CHECK((ArgumentOffset % 4) == 0, TEXT("IndirectArgOffset for compute shader indirect dispatch needs to be a multiple of 4."));
+		auto GetBufferDesc = [&]() -> FString { return FString::Printf(TEXT("Buffer: %s, Size: %u, Stride: %u, Offset: %u"), ArgumentBuffer->GetDebugName(), ArgumentBuffer->GetSize(), ArgumentBuffer->GetStride(), ArgumentOffset); };
+
+		RHI_VALIDATION_CHECK(EnumHasAnyFlags(ArgumentBuffer->GetUsage(), EBufferUsageFlags::VertexBuffer), *FString::Printf(TEXT("Indirect argument buffer must be a vertex buffer to be used as an indirect dispatch parameter. %s"), *GetBufferDesc()));
+		RHI_VALIDATION_CHECK(EnumHasAnyFlags(ArgumentBuffer->GetUsage(), EBufferUsageFlags::DrawIndirect), *FString::Printf(TEXT("Indirect dispatch parameter buffer was not flagged with BUF_DrawIndirect. %s"), *GetBufferDesc()));
+		RHI_VALIDATION_CHECK((ArgumentOffset % 4) == 0, *FString::Printf(TEXT("Offset into buffer for compute shader indirect dispatch must be a multiple of 4. %s"), *GetBufferDesc()));
 		RHI_VALIDATION_CHECK((ArgumentOffset + ArgumentBuffer->GetStride()) <= ArgumentBuffer->GetSize(),
-			*FString::Printf(TEXT("Indirect parameters buffer for compute shader indirect dispatch at byte offset %d doesn't have enough room for one element."), ArgumentOffset));
+			*FString::Printf(TEXT("Indirect parameter buffer for compute shader indirect dispatch at byte offset %u doesn't have enough room for one element. %s"), ArgumentOffset, *GetBufferDesc()));
 #if PLATFORM_DISPATCH_INDIRECT_ARGUMENT_BOUNDARY_SIZE != 0
 		RHI_VALIDATION_CHECK(ArgumentOffset / PLATFORM_DISPATCH_INDIRECT_ARGUMENT_BOUNDARY_SIZE == (ArgumentOffset + ArgumentBuffer->GetStride() - 1) / PLATFORM_DISPATCH_INDIRECT_ARGUMENT_BOUNDARY_SIZE,
-			*FString::Printf(TEXT("Compute indirect dispatch arguments cannot cross %d byte boundary."), PLATFORM_DISPATCH_INDIRECT_ARGUMENT_BOUNDARY_SIZE));
+			*FString::Printf(TEXT("Compute indirect dispatch arguments cannot cross %d byte boundary. %s"), PLATFORM_DISPATCH_INDIRECT_ARGUMENT_BOUNDARY_SIZE, *GetBufferDesc()));
 #endif // #if PLATFORM_DISPATCH_INDIRECT_ARGUMENT_BOUNDARY_SIZE != 0
 	}
 
