Reworking logging and messaging in WindowsDynamicRHI to be more consistent wth names.

diff --git a/Engine/Source/Runtime/RHI/Private/Windows/WindowsDynamicRHI.cpp b/Engine/Source/Runtime/RHI/Private/Windows/WindowsDynamicRHI.cpp
--- a/Engine/Source/Runtime/RHI/Private/Windows/WindowsDynamicRHI.cpp
+++ b/Engine/Source/Runtime/RHI/Private/Windows/WindowsDynamicRHI.cpp
@@ -12,6 +12,8 @@
 
 #include "Windows/WindowsPlatformApplicationMisc.h"
 
+#define LOCTEXT_NAMESPACE "WindowsDynamicRHI"
+
 bool GDynamicRHIFailedToInitializeAdvancedPlatform = false;
 
 static const TCHAR* GLoadedRHIModuleName;
@@ -26,6 +28,60 @@ enum class EWindowsRHI
 };
 static constexpr int32 EWindowsRHICount = static_cast<int32>(EWindowsRHI::count);
 
+inline const TCHAR* GetLogName(EWindowsRHI InWindowsRHI)
+{
+	switch (InWindowsRHI)
+	{
+	case EWindowsRHI::D3D11:  return TEXT("D3D11");
+	case EWindowsRHI::D3D12:  return TEXT("D3D12");
+	case EWindowsRHI::Vulkan: return TEXT("Vulkan");
+	case EWindowsRHI::OpenGL: return TEXT("OpenGL");
+	default:                  return TEXT("<unknown>");
+	}
+}
+
+inline FText GetTextName(EWindowsRHI InWindowsRHI)
+{
+	switch (InWindowsRHI)
+	{
+	case EWindowsRHI::D3D11:  return LOCTEXT("D3D11",      "DirectX 11");
+	case EWindowsRHI::D3D12:  return LOCTEXT("D3D12",      "DirectX 12");
+	case EWindowsRHI::Vulkan: return LOCTEXT("Vulkan",     "Vulkan");
+	case EWindowsRHI::OpenGL: return LOCTEXT("OpenGL",     "OpenGL");
+	default:                  return LOCTEXT("UnknownRHI", "<Unknown>");
+	}
+}
+
+inline const TCHAR* GetLogName(ERHIFeatureLevel::Type InFeatureLevel)
+{
+	switch (InFeatureLevel)
+	{
+	case ERHIFeatureLevel::ES3_1: return TEXT("ES3_1");
+	case ERHIFeatureLevel::SM5:   return TEXT("SM5");
+	case ERHIFeatureLevel::SM6:   return TEXT("SM6");
+	default:                      return TEXT("<unknown>");
+	}
+}
+
+inline FText GetTextName(ERHIFeatureLevel::Type InFeatureLevel)
+{
+	switch (InFeatureLevel)
+	{
+	case ERHIFeatureLevel::ES3_1: return LOCTEXT("FeatureLevelES31",    "ES3_1");
+	case ERHIFeatureLevel::SM5:   return LOCTEXT("FeatureLevelSM4",     "SM5");
+	case ERHIFeatureLevel::SM6:   return LOCTEXT("FeatureLevelSM6",     "SM6");
+	default:                      return LOCTEXT("FeatureLevelUnknown", "<unknown>");
+	}
+}
+
+const EWindowsRHI GRHISearchOrder[] =
+{
+	EWindowsRHI::D3D12,
+	EWindowsRHI::D3D11,
+	EWindowsRHI::Vulkan,
+	EWindowsRHI::OpenGL,
+};
+
 static TAutoConsoleVariable<bool> CVarIgnorePerformanceModeCheck(
 	TEXT("r.IgnorePerformanceModeCheck"),
 	false,
@@ -105,6 +161,7 @@ struct FParsedWindowsDynamicRHIConfig
 	TOptional<ERHIFeatureLevel::Type> GetNextHighestTargetedFeatureLevel(EWindowsRHI InWindowsRHI, ERHIFeatureLevel::Type InFeatureLevel) const;
 
 	bool IsFeatureLevelTargeted(EWindowsRHI InWindowsRHI, ERHIFeatureLevel::Type InFeatureLevel) const;
+	TOptional<EWindowsRHI> GetFirstRHIWithFeatureLevelSupport(ERHIFeatureLevel::Type InFeatureLevel) const;
 
 	void MergeConfig(EWindowsRHI InWindowsRHI, const FWindowsRHIConfig& Other);
 };
@@ -179,6 +236,18 @@ bool FParsedWindowsDynamicRHIConfig::IsFeatureLevelTargeted(EWindowsRHI InWindow
 	return false;
 }
 
+TOptional<EWindowsRHI> FParsedWindowsDynamicRHIConfig::GetFirstRHIWithFeatureLevelSupport(ERHIFeatureLevel::Type InFeatureLevel) const
+{
+	for (EWindowsRHI WindowsRHI : GRHISearchOrder)
+	{
+		if (IsFeatureLevelTargeted(WindowsRHI, InFeatureLevel))
+		{
+			return WindowsRHI;
+		}
+	}
+	return {};
+}
+
 void FParsedWindowsDynamicRHIConfig::MergeConfig(EWindowsRHI InWindowsRHI, const FWindowsRHIConfig& Other)
 {
 	FWindowsRHIConfig& ExistingConfig = RHIConfigs[(int32)InWindowsRHI];
@@ -307,7 +376,7 @@ static bool DefaultFeatureLevelES31()
 	int MinCoreCount = 0;
 	if (GConfig->GetInt(TEXT("PerformanceMode"), TEXT("MinCoreCount"), MinCoreCount, GEngineIni) && FPlatformMisc::NumberOfCoresIncludingHyperthreads() < MinCoreCount)
 	{
-		UE_LOG(LogRHI, Log, TEXT("MinCoreCount found (%d) and user core count (%d) forced feature level to ES3_1"), MinCoreCount, FPlatformMisc::NumberOfCoresIncludingHyperthreads());
+		UE_LOG(LogRHI, Log, TEXT("MinCoreCount found (%d) and user core count (%d) forced feature level to %s"), MinCoreCount, FPlatformMisc::NumberOfCoresIncludingHyperthreads(), GetLogName(ERHIFeatureLevel::ES3_1));
 		ForceES31 = true;
 		return true;
 	}
@@ -326,7 +395,7 @@ static bool DefaultFeatureLevelES31()
 			{
 				if (FPlatformMemory::GetMemorySizeBucket() >= EPlatformMemorySizeBucket(EnumIndex))
 				{
-					UE_LOG(LogRHI, Log, TEXT("MinMemorySizeBucket found (%s) and user memory (%d) forced feature level to ES3_1"), *MinMemorySizeBucketString, int32(FPlatformMemory::GetMemorySizeBucket()));
+					UE_LOG(LogRHI, Log, TEXT("MinMemorySizeBucket found (%s) and user memory (%d) forced feature level to %s"), *MinMemorySizeBucketString, int32(FPlatformMemory::GetMemorySizeBucket()), GetLogName(ERHIFeatureLevel::ES3_1));
 
 					ForceES31 = true;
 					return true;
@@ -339,7 +408,7 @@ static bool DefaultFeatureLevelES31()
 				const int MIN_GPU_MEMORY = 512 * 1024 * 1024;
 				if (FPlatformMemory::GetMemorySizeBucket() >= EPlatformMemorySizeBucket(EnumIndex) && BestGPUInfo.DedicatedVideoMemory < MIN_GPU_MEMORY)
 				{
-					UE_LOG(LogRHI, Log, TEXT("MinMemorySizeBucket found (%s) and user memory (%d) forced feature level to ES3_1"), *MinMemorySizeBucketString, int32(FPlatformMemory::GetMemorySizeBucket()));
+					UE_LOG(LogRHI, Log, TEXT("MinMemorySizeBucket found (%s) and user memory (%d) forced feature level to %s"), *MinMemorySizeBucketString, int32(FPlatformMemory::GetMemorySizeBucket()), GetLogName(ERHIFeatureLevel::ES3_1));
 
 					ForceES31 = true;
 					return true;
@@ -366,7 +435,7 @@ static bool DefaultFeatureLevelES31()
 
 		if (RHIName.Compare("D3D11_ES31", ESearchCase::IgnoreCase) == 0 && GPUBrand.Compare(DeviceName, ESearchCase::IgnoreCase) == 0)
 		{
-			UE_LOG(LogRHI, Log, TEXT("Found RHIName %s with DeviceName %s which has forced feature level to ES3_1"), *RHIName, *DeviceName);
+			UE_LOG(LogRHI, Log, TEXT("Found RHIName %s with DeviceName %s which has forced feature level to %s"), *RHIName, *DeviceName, GetLogName(ERHIFeatureLevel::ES3_1));
 
 			ForceES31 = true;
 			return true;
@@ -384,7 +453,7 @@ static bool DefaultFeatureLevelES31()
 			BestGPUInfo.VendorId == VendorIdInt && BestGPUInfo.DeviceId == DeviceIdInt &&
 			RHIName.Compare("D3D11_ES31", ESearchCase::IgnoreCase) == 0)
 		{
-			UE_LOG(LogRHI, Log, TEXT("Found RHIName %s with DeviceName %s, VendorId %s, and DeviceId %s which has forced feature level to ES3_1"), *RHIName, *DeviceName, *VendorId, *DeviceId);
+			UE_LOG(LogRHI, Log, TEXT("Found RHIName %s with DeviceName %s, VendorId %s, and DeviceId %s which has forced feature level to %s"), *RHIName, *DeviceName, *VendorId, *DeviceId, GetLogName(ERHIFeatureLevel::ES3_1));
 
 			ForceES31 = true;
 			return true;
@@ -531,15 +600,8 @@ static EWindowsRHI ChooseDefaultRHI(const FParsedWindowsDynamicRHIConfig& Config
 		return ConfigDefault.GetValue();
 	}
 
-	const EWindowsRHI DefaultRHIOrder[] =
-	{
-		EWindowsRHI::D3D12,
-		EWindowsRHI::D3D11,
-		EWindowsRHI::Vulkan,
-	};
-
 	// Find the first RHI with configured support based on the order above
-	for (EWindowsRHI DefaultRHI : DefaultRHIOrder)
+	for (EWindowsRHI DefaultRHI : GRHISearchOrder)
 	{
 		if (TOptional<ERHIFeatureLevel::Type> HighestFL = ChooseDefaultFeatureLevel(DefaultRHI, Config))
 		{
@@ -635,39 +697,24 @@ static TOptional<EWindowsRHI> ChooseForcedRHI(TOptional<ERHIFeatureLevel::Type>
 
 	if (Sum > 1)
 	{
-		FMessageDialog::Open(EAppMsgType::Ok, NSLOCTEXT("WindowsDynamicRHI", "RHIOptionsError", "-d3d12/dx12, -d3d11/dx11, -vulkan, and -opengl are mutually exclusive options, but more than one was specified on the command-line."));
+		FMessageDialog::Open(EAppMsgType::Ok, LOCTEXT("RHIOptionsError", "-d3d12/dx12, -d3d11/dx11, -vulkan, and -opengl are mutually exclusive options, but more than one was specified on the command-line."));
 		UE_LOG(LogRHI, Fatal, TEXT("-d3d12, -d3d11, -vulkan, and -opengl are mutually exclusive options, but more than one was specified on the command-line."));
 	}
 
 	// FeatureLevelES31 is also a command line override, so it will determine the underlying RHI unless one is specified
-	if (FPlatformProperties::RequiresCookedData() && ForcedFeatureLevel == ERHIFeatureLevel::ES3_1)
+	if ((FPlatformProperties::RequiresCookedData() || !GIsEditor) && ForcedFeatureLevel)
 	{
 		if (ForcedRHI)
 		{
 			if (!Config.IsFeatureLevelTargeted(ForcedRHI.GetValue(), ForcedFeatureLevel.GetValue()))
 			{
-				if (ForcedRHI == EWindowsRHI::OpenGL)
-				{
-					FMessageDialog::Open(EAppMsgType::Ok, NSLOCTEXT("WindowsDynamicRHI", "RHIPerformanceOpenGL", "OpenGL is not supported for Performance Mode."));
-					UE_LOG(LogRHI, Fatal, TEXT("OpenGL is not supported for Performance Mode."));
-				}
-				else if (ForcedRHI == EWindowsRHI::Vulkan)
-				{
-					FMessageDialog::Open(EAppMsgType::Ok, NSLOCTEXT("WindowsDynamicRHI", "RHIPerformanceVulkan", "Vulkan is not supported for Performance Mode."));
-					UE_LOG(LogRHI, Fatal, TEXT("Vulkan is not supported for Performance Mode."));
-				}
-				else if (ForcedRHI == EWindowsRHI::D3D12)
-				{
-					FMessageDialog::Open(EAppMsgType::Ok, NSLOCTEXT("WindowsDynamicRHI", "RHIPerformanceDX12", "DirectX 12 is not supported for Performance Mode."));
-					UE_LOG(LogRHI, Fatal, TEXT("DirectX 12 is not supported for Performance Mode."));
-				}
+				FMessageDialog::Open(EAppMsgType::Ok, FText::Format(LOCTEXT("RHINotSupported", "{0} does not support {1}."), GetTextName(ForcedRHI.GetValue()), GetTextName(ForcedFeatureLevel.GetValue())));
+				UE_LOG(LogRHI, Fatal, TEXT("%s does not support %s."), GetLogName(ForcedRHI.GetValue()), GetLogName(ForcedFeatureLevel.GetValue()));
 			}
 		}
-		else
+		else if (ForcedRHI = Config.GetFirstRHIWithFeatureLevelSupport(ForcedFeatureLevel.GetValue()))
 		{
-			ForcedRHI = EWindowsRHI::D3D11;
-
-			UE_LOG(LogRHI, Log, TEXT("Forcing RHI to %s since Feature Level %s was forced"), ModuleNameFromWindowsRHI(ForcedRHI.GetValue()), *LexToString(ForcedFeatureLevel.GetValue()));
+			UE_LOG(LogRHI, Log, TEXT("Forcing RHI to %s since Feature Level %s was forced"), GetLogName(ForcedRHI.GetValue()), GetLogName(ForcedFeatureLevel.GetValue()));
 		}
 	}
 
@@ -703,14 +750,14 @@ static ERHIFeatureLevel::Type ChooseFeatureLevel(EWindowsRHI ChosenRHI, const TO
 		// Allow the forced feature level if we're in a position to compile its shaders
 		if (!FPlatformProperties::RequiresCookedData())
 		{
-			UE_LOG(LogRHI, Log, TEXT("Using Forced Feature Level in Editor: %s"), *LexToString(ForcedFeatureLevel.GetValue()));
+			UE_LOG(LogRHI, Log, TEXT("Using Forced Feature Level in Editor: %s"), GetLogName(ForcedFeatureLevel.GetValue()));
 			return ForcedFeatureLevel.GetValue();
 		}
 
 		// Make sure the feature level is supported by the runtime, otherwise fall back to the default
 		if (Config.IsFeatureLevelTargeted(ChosenRHI, ForcedFeatureLevel.GetValue()))
 		{
-			UE_LOG(LogRHI, Log, TEXT("Using Forced Feature Level the Game is configured for: %s"), *LexToString(ForcedFeatureLevel.GetValue()));
+			UE_LOG(LogRHI, Log, TEXT("Using Forced Feature Level the Game is configured for: %s"), GetLogName(ForcedFeatureLevel.GetValue()));
 			return ForcedFeatureLevel.GetValue();
 		}
 	}
@@ -728,29 +775,24 @@ static ERHIFeatureLevel::Type ChooseFeatureLevel(EWindowsRHI ChosenRHI, const TO
 
 			if (FPlatformProperties::RequiresCookedData())
 			{
-				const TCHAR* RHIName = ModuleNameFromWindowsRHI(ChosenRHI);
-				FMessageDialog::Open(EAppMsgType::Ok, FText::Format(NSLOCTEXT("WindowsDynamicRHI", "MissingTargetedShaderFormats", "Unable to launch with RHI '{0}' since the project is not configured to support it."), FText::AsCultureInvariant(RHIName)));
-				UE_LOG(LogRHI, Fatal, TEXT("Unable to launch with RHI '%s' since the project is not configured to support it."), RHIName);
+				FMessageDialog::Open(EAppMsgType::Ok, FText::Format(LOCTEXT("MissingTargetedShaderFormats", "Unable to launch with RHI '{0}' since the project is not configured to support it."), GetTextName(ChosenRHI)));
+				UE_LOG(LogRHI, Fatal, TEXT("Unable to launch with RHI '%s' since the project is not configured to support it."), GetLogName(ChosenRHI));
 			}
 			else
 			{
-				const TCHAR* RHIName = ModuleNameFromWindowsRHI(ChosenRHI);
-				const FString FeatureLevelName = LexToString(FeatureLevel.GetValue());
-				UE_LOG(LogRHI, Warning, TEXT("User requested RHI '%s' but that is not supported by this project's data. Defaulting to Feature Level '%s'."), RHIName, *FeatureLevelName);
+				UE_LOG(LogRHI, Warning, TEXT("User requested RHI '%s' but that is not supported by this project's data. Defaulting to Feature Level '%s'."), GetLogName(ChosenRHI), GetLogName(FeatureLevel.GetValue()));
 			}
 		}
 		else
 		{
-			UE_LOG(LogRHI, Log, TEXT("Using Highest Feature Level of %s: %s"), ModuleNameFromWindowsRHI(ChosenRHI), *LexToString(FeatureLevel.GetValue()));
+			UE_LOG(LogRHI, Log, TEXT("Using Highest Feature Level of %s: %s"), GetLogName(ChosenRHI), GetLogName(FeatureLevel.GetValue()));
 		}
 	}
 
 	// If the user wanted to force a feature level and we couldn't set it, log out why and what we're actually running with
 	if (ForcedFeatureLevel)
 	{
-		const FString ForcedName = LexToString(ForcedFeatureLevel.GetValue());
-		const FString UsedName = LexToString(FeatureLevel.GetValue());
-		UE_LOG(LogRHI, Warning, TEXT("User requested Feature Level '%s' but that is not supported by this project. Falling back to Feature Level '%s'."), *ForcedName, *UsedName);
+		UE_LOG(LogRHI, Warning, TEXT("User requested Feature Level '%s' but that is not supported by this project. Falling back to Feature Level '%s'."), GetLogName(ForcedFeatureLevel.GetValue()), GetLogName(FeatureLevel.GetValue()));
 	}
 
 	return FeatureLevel.GetValue();
@@ -762,9 +804,9 @@ static bool HandleUnsupportedFeatureLevel(EWindowsRHI& WindowsRHI, ERHIFeatureLe
 	{
 		if (WindowsRHI == EWindowsRHI::D3D12 && FeatureLevel == ERHIFeatureLevel::SM6)
 		{
-			UE_LOG(LogRHI, Log, TEXT("RHI %s with Feature Level %s is not supported on your system."), ModuleNameFromWindowsRHI(WindowsRHI), *LexToString(FeatureLevel));
+			UE_LOG(LogRHI, Log, TEXT("RHI %s with Feature Level %s is not supported on your system."), GetLogName(WindowsRHI), GetLogName(FeatureLevel));
 
-			FMessageDialog::Open(EAppMsgType::Ok, NSLOCTEXT("WindowsDynamicRHI", "RequiredDX12SM6", "DirectX 12 with Feature Level SM6 is not supported on your system. Try running without the -sm6 command line argument."));
+			FMessageDialog::Open(EAppMsgType::Ok, LOCTEXT("RequiredDX12SM6", "DirectX 12 with Feature Level SM6 is not supported on your system. Try running without the -sm6 command line argument."));
 			FPlatformMisc::RequestExit(1);
 		}
 
@@ -791,9 +833,9 @@ static bool HandleUnsupportedRHI(EWindowsRHI& WindowsRHI, ERHIFeatureLevel::Type
 	{
 		if (ForcedRHI == EWindowsRHI::D3D12)
 		{
-			UE_LOG(LogRHI, Log, TEXT("RHI %s with Feature Level %s is not supported on your system."), ModuleNameFromWindowsRHI(WindowsRHI), *LexToString(FeatureLevel));
+			UE_LOG(LogRHI, Log, TEXT("RHI %s with Feature Level %s is not supported on your system."), GetLogName(WindowsRHI), GetLogName(FeatureLevel));
 
-			FMessageDialog::Open(EAppMsgType::Ok, NSLOCTEXT("WindowsDynamicRHI", "RequiredDX12", "DirectX 12 is not supported on your system. Try running without the -dx12 or -d3d12 command line argument."));
+			FMessageDialog::Open(EAppMsgType::Ok, LOCTEXT("RequiredDX12", "DirectX 12 is not supported on your system. Try running without the -dx12 or -d3d12 command line argument."));
 			FPlatformMisc::RequestExit(1);
 			return false;
 		}
@@ -809,29 +851,29 @@ static bool HandleUnsupportedRHI(EWindowsRHI& WindowsRHI, ERHIFeatureLevel::Type
 		}
 	}
 
-	UE_LOG(LogRHI, Log, TEXT("RHI %s is not supported on your system."), ModuleNameFromWindowsRHI(WindowsRHI));
+	UE_LOG(LogRHI, Log, TEXT("RHI %s is not supported on your system."), GetLogName(WindowsRHI));
 
 	if (WindowsRHI == EWindowsRHI::D3D12)
 	{
-		FMessageDialog::Open(EAppMsgType::Ok, NSLOCTEXT("WindowsDynamicRHI", "RequiredDX12", "DirectX 12 is not supported on your system. Try running without the -dx12 or -d3d12 command line argument."));
+		FMessageDialog::Open(EAppMsgType::Ok, LOCTEXT("RequiredDX12", "DirectX 12 is not supported on your system. Try running without the -dx12 or -d3d12 command line argument."));
 		FPlatformMisc::RequestExit(1);
 	}
 
 	if (WindowsRHI == EWindowsRHI::D3D11)
 	{
-		FMessageDialog::Open(EAppMsgType::Ok, NSLOCTEXT("WindowsDynamicRHI", "RequiredDX11Feature_11_SM5", "A D3D11-compatible GPU (Feature Level 11.0, Shader Model 5.0) is required to run the engine."));
+		FMessageDialog::Open(EAppMsgType::Ok, LOCTEXT("RequiredDX11Feature_11_SM5", "A D3D11-compatible GPU (Feature Level 11.0, Shader Model 5.0) is required to run the engine."));
 		FPlatformMisc::RequestExit(1);
 	}
 
 	if (WindowsRHI == EWindowsRHI::Vulkan)
 	{
-		FMessageDialog::Open(EAppMsgType::Ok, NSLOCTEXT("WindowsDynamicRHI", "RequiredVulkan", "Vulkan Driver is required to run the engine."));
+		FMessageDialog::Open(EAppMsgType::Ok, LOCTEXT("RequiredVulkan", "Vulkan Driver is required to run the engine."));
 		FPlatformMisc::RequestExit(1);
 	}
 
 	if (WindowsRHI == EWindowsRHI::OpenGL)
 	{
-		FMessageDialog::Open(EAppMsgType::Ok, NSLOCTEXT("WindowsDynamicRHI", "RequiredOpenGL", "OpenGL 4.3 is required to run the engine."));
+		FMessageDialog::Open(EAppMsgType::Ok, LOCTEXT("RequiredOpenGL", "OpenGL 4.3 is required to run the engine."));
 		FPlatformMisc::RequestExit(1);
 	}
 
@@ -866,17 +908,17 @@ static IDynamicRHIModule* LoadDynamicRHIModule(ERHIFeatureLevel::Type& DesiredFe
 	{
 		ChosenRHI = ForcedRHI.GetValue();
 
-		UE_LOG(LogRHI, Log, TEXT("Using Forced RHI: %s"), ModuleNameFromWindowsRHI(ChosenRHI));
+		UE_LOG(LogRHI, Log, TEXT("Using Forced RHI: %s"), GetLogName(ChosenRHI));
 	}
 	else if (PreferredRHI)
 	{
 		ChosenRHI = PreferredRHI.GetValue();
 
-		UE_LOG(LogRHI, Log, TEXT("Using Preferred RHI: %s"), ModuleNameFromWindowsRHI(ChosenRHI));
+		UE_LOG(LogRHI, Log, TEXT("Using Preferred RHI: %s"), GetLogName(ChosenRHI));
 	}
 	else
 	{
-		UE_LOG(LogRHI, Log, TEXT("Using Default RHI: %s"), ModuleNameFromWindowsRHI(ChosenRHI));
+		UE_LOG(LogRHI, Log, TEXT("Using Default RHI: %s"), GetLogName(ChosenRHI));
 	}
 
 	DesiredFeatureLevel = ChooseFeatureLevel(ChosenRHI, ForcedRHI, ForcedFeatureLevel, Config);
@@ -891,15 +933,15 @@ static IDynamicRHIModule* LoadDynamicRHIModule(ERHIFeatureLevel::Type& DesiredFe
 
 		const TCHAR* ModuleName = ModuleNameFromWindowsRHI(ChosenRHI);
 
-		UE_LOG(LogRHI, Log, TEXT("Loading RHI %s"), ModuleName);
+		UE_LOG(LogRHI, Log, TEXT("Loading RHI module %s"), ModuleName);
 
 		IDynamicRHIModule* DynamicRHIModule = FModuleManager::LoadModulePtr<IDynamicRHIModule>(ModuleName);
 
-		UE_LOG(LogRHI, Log, TEXT("Checking if RHI %s with Feature Level %s is supported by your system."), ModuleName, *LexToString(DesiredFeatureLevel));
+		UE_LOG(LogRHI, Log, TEXT("Checking if RHI %s with Feature Level %s is supported by your system."), GetLogName(ChosenRHI), GetLogName(DesiredFeatureLevel));
 
 		if (DynamicRHIModule && DynamicRHIModule->IsSupported(DesiredFeatureLevel))
 		{
-			UE_LOG(LogRHI, Log, TEXT("RHI %s with Feature Level %s is supported and will be used."), ModuleName, *LexToString(DesiredFeatureLevel));
+			UE_LOG(LogRHI, Log, TEXT("RHI %s with Feature Level %s is supported and will be used."), GetLogName(ChosenRHI), GetLogName(DesiredFeatureLevel));
 
 			LoadedRHIModuleName = ModuleName;
 			return DynamicRHIModule;
@@ -918,12 +960,12 @@ static IDynamicRHIModule* LoadDynamicRHIModule(ERHIFeatureLevel::Type& DesiredFe
 		if (bTryWithNewConfig)
 		{
 			UE_LOG(LogRHI, Log, TEXT("RHI %s with Feature Level %s is not supported on your system, attempting to fall back to RHI %s with Feature Level %s"),
-				ModuleNameFromWindowsRHI(PreviousRHI), *LexToString(PreviousFeatureLevel),
-				ModuleNameFromWindowsRHI(ChosenRHI), *LexToString(DesiredFeatureLevel));
+				GetLogName(PreviousRHI), GetLogName(PreviousFeatureLevel),
+				GetLogName(ChosenRHI), GetLogName(DesiredFeatureLevel));
 		}
 	} while (bTryWithNewConfig);
 
-	UE_LOG(LogRHI, Log, TEXT("RHI %s with Feature Level %s is not supported on your system. No RHI was supported, failing initialization."), ModuleNameFromWindowsRHI(ChosenRHI), *LexToString(DesiredFeatureLevel));
+	UE_LOG(LogRHI, Log, TEXT("RHI %s with Feature Level %s is not supported on your system. No RHI was supported, failing initialization."), GetLogName(ChosenRHI), GetLogName(DesiredFeatureLevel));
 
 	return nullptr;
 }
@@ -971,4 +1013,6 @@ const TCHAR* GetSelectedDynamicRHIModuleName(bool bCleanup)
 	}
 }
 
+#undef LOCTEXT_NAMESPACE
+
 #endif //WINDOWS_USE_FEATURE_DYNAMIC_RHI
