remove CreateTaskWithEvent for PSOs as we can do with only a regular CreateTask and not burden the task library with another path

#rnx

diff --git a/Engine/Source/Runtime/Core/Public/Async/TaskGraphInterfaces.h b/Engine/Source/Runtime/Core/Public/Async/TaskGraphInterfaces.h
--- a/Engine/Source/Runtime/Core/Public/Async/TaskGraphInterfaces.h
+++ b/Engine/Source/Runtime/Core/Public/Async/TaskGraphInterfaces.h
@@ -1262,21 +1262,6 @@ public:
 		return FConstructor(new TGraphTask(TTask::GetSubsequentsMode() == ESubsequentsMode::FireAndForget ? NULL : FGraphEvent::CreateGraphEvent(), NumPrereq), Prerequisites, CurrentThreadIfKnown);
 	}
 
-	/**
-	 *	Factory to create a task and return the helper object to construct the embedded task and set it up for execution.
-	 *	@param SubsequentEventToAssume; subsequent to "assume" from an existing task
-	 *	@param Prerequisites; the list of FGraphEvents that must be completed prior to this task executing.
-	 *	@param CurrentThreadIfKnown; provides the index of the thread we are running on. Can be ENamedThreads::AnyThread if the current thread is unknown.
-	 *	@return a temporary helper class which can be used to complete the process.
-	**/
-	static FConstructor CreateTaskWithEvent(FGraphEventRef SubsequentEventToAssume, const FGraphEventArray* Prerequisites = NULL, ENamedThreads::Type CurrentThreadIfKnown = ENamedThreads::AnyThread)
-	{
-		int32 NumPrereq = Prerequisites ? Prerequisites->Num() : 0;
-		// If an event is provided, make sure the task will call DispatchSubsequents on it in ExecuteTask
-		check(SubsequentEventToAssume == NULL || TTask::GetSubsequentsMode() == ESubsequentsMode::TrackSubsequents);
-		return FConstructor(new TGraphTask(SubsequentEventToAssume, NumPrereq), Prerequisites, CurrentThreadIfKnown);
-	}
-
 	void Unlock(ENamedThreads::Type CurrentThreadIfKnown = ENamedThreads::AnyThread)
 	{
 		TaskTrace::Launched(GetTraceId(), nullptr, Subsequents.IsValid(), ((TTask*)&TaskStorage)->GetDesiredThread());
diff --git a/Engine/Source/Runtime/RHI/Private/PipelineStateCache.cpp b/Engine/Source/Runtime/RHI/Private/PipelineStateCache.cpp
--- a/Engine/Source/Runtime/RHI/Private/PipelineStateCache.cpp
+++ b/Engine/Source/Runtime/RHI/Private/PipelineStateCache.cpp
@@ -1520,6 +1520,7 @@ public:
 		, Initializer(InInitializer)
 		, PSOPreCacheResult(InPSOPreCacheResult)
 	{
+		ensure(Pipeline->CompletionEvent != nullptr);
 		if(IsPrecachedPSO(Initializer))
 		{
 			GPipelinePrecompileTasksInFlight++;
@@ -1731,6 +1732,14 @@ public:
 				Initializer.DepthStencilState->Release();
 			}
 		}
+
+		// We kicked a task: the event really should be there
+		if (ensure(Pipeline->CompletionEvent))
+		{
+			Pipeline->CompletionEvent->DispatchSubsequents();
+			// At this point, it's not safe to use Pipeline anymore, as it might get picked up by ProcessDelayedCleanup and deleted
+			Pipeline = nullptr;
+		}
 	}
 
 	FORCEINLINE TStatId GetStatId() const
@@ -1894,9 +1903,10 @@ FComputePipelineState* PipelineStateCache::GetAndOrCreateComputePipelineState(FR
 		// create a compilation task, or just do it now...
 		if (DoAsyncCompile)
 		{
+			OutCachedState->CompletionEvent = FGraphEvent::CreateGraphEvent();
 			FGraphicsPipelineStateInitializer GraphicsPipelineStateInitializer;
 			GraphicsPipelineStateInitializer.bFromPSOFileCache = bFromFileCache;
-			OutCachedState->CompletionEvent = TGraphTask<FCompilePipelineStateTask>::CreateTask().ConstructAndDispatchWhenReady(OutCachedState, GraphicsPipelineStateInitializer, EPSOPrecacheResult::Untracked);
+			TGraphTask<FCompilePipelineStateTask>::CreateTask().ConstructAndDispatchWhenReady(OutCachedState, GraphicsPipelineStateInitializer, EPSOPrecacheResult::Untracked);
 			if (!bFromFileCache)
 			{
 				RHICmdList.AddDispatchPrerequisite(OutCachedState->CompletionEvent);
@@ -1943,7 +1953,8 @@ FComputePipelineState* PipelineStateCache::GetAndOrCreateComputePipelineState(FR
 			// create a compilation task, or just do it now...
 			if (DoAsyncCompile)
 			{
-				PipelineState->CompletionEvent = TGraphTask<FCompilePipelineStateTask>::CreateTask().ConstructAndDispatchWhenReady(PipelineState.Get(), FGraphicsPipelineStateInitializer(), EPSOPrecacheResult::Untracked);
+				PipelineState->CompletionEvent = FGraphEvent::CreateGraphEvent();
+				TGraphTask<FCompilePipelineStateTask>::CreateTask().ConstructAndDispatchWhenReady(PipelineState.Get(), FGraphicsPipelineStateInitializer(), EPSOPrecacheResult::Untracked);
 				RHICmdList.QueueAsyncPipelineStateCompile(PipelineState->CompletionEvent);
 			}
 			else
@@ -2315,7 +2326,7 @@ static void InternalCreateGraphicsPipelineState(const FGraphicsPipelineStateInit
 		// Use normal task graph for non-precompile jobs (or when thread pool is not enabled)
 		if (!bPSOPrecache || !FPSOPrecacheThreadPool::UsePool())
 		{
-			TGraphTask<FCompilePipelineStateTask>::CreateTaskWithEvent(GraphEvent).ConstructAndDispatchWhenReady(CachedState, Initializer, PSOPrecacheResult);
+			TGraphTask<FCompilePipelineStateTask>::CreateTask().ConstructAndDispatchWhenReady(CachedState, Initializer, PSOPrecacheResult);
 		}
 		else
 		{
@@ -2332,9 +2343,8 @@ static void InternalCreateGraphicsPipelineState(const FGraphicsPipelineStateInit
 				ThreadPoolTask->CompilePSO();
 			}
 			,
-				[GraphEvent, StartTime]()
+				[StartTime]()
 			{
-				GraphEvent->DispatchSubsequents();
 #if PSO_TRACK_CACHE_STATS
 				StatsEndPrecompile(FPlatformTime::Cycles64() - StartTime);
 #endif
@@ -2514,18 +2524,17 @@ FGraphEventRef FPrecacheComputePipelineCache::PrecacheComputePipelineState(FRHIC
 void FPrecacheComputePipelineCache::OnNewPipelineStateCreated(const FPrecacheComputeInitializer& ComputeInitializer, FComputePipelineState* CachedState, bool bDoAsyncCompile)
 {
 	// create a compilation task, or just do it now...
-	FGraphEventRef GraphEvent = CachedState->CompletionEvent;
 	if (bDoAsyncCompile)
 	{
 		INC_DWORD_STAT(STAT_ActiveComputePSOPrecacheRequests);
 
-		check(GraphEvent != nullptr);
+		check(CachedState->CompletionEvent != nullptr);
 		FGraphicsPipelineStateInitializer GraphicsPipelineStateInitializer;
 		GraphicsPipelineStateInitializer.bPSOPrecache = true;
 		// Thread pool disabled?
 		if (!FPSOPrecacheThreadPool::UsePool())
 		{
-			TGraphTask<FCompilePipelineStateTask>::CreateTaskWithEvent(GraphEvent).ConstructAndDispatchWhenReady(CachedState, GraphicsPipelineStateInitializer, EPSOPrecacheResult::Active);
+			TGraphTask<FCompilePipelineStateTask>::CreateTask().ConstructAndDispatchWhenReady(CachedState, GraphicsPipelineStateInitializer, EPSOPrecacheResult::Active);
 		}
 		else
 		{
@@ -2542,9 +2551,8 @@ void FPrecacheComputePipelineCache::OnNewPipelineStateCreated(const FPrecacheCom
 				ThreadPoolTask->CompilePSO();
 			}
 			,
-				[GraphEvent, StartTime]()
+				[StartTime]()
 			{
-				GraphEvent->DispatchSubsequents();
 #if PSO_TRACK_CACHE_STATS
 				StatsEndPrecompile(FPlatformTime::Cycles64() - StartTime);
 #endif
@@ -2555,7 +2563,7 @@ void FPrecacheComputePipelineCache::OnNewPipelineStateCreated(const FPrecacheCom
 	}
 	else
 	{
-		check(GraphEvent == nullptr);
+		check(CachedState->CompletionEvent == nullptr);
 		CachedState->RHIPipeline = RHICreateComputePipelineState(CachedState->ComputeShader);
 	}
 }
