[RHI]
* Added markups to RHI files to make sure iwyu is not using RHIResources.h, RHIUtilities.h and RHICommandList.h directly.. all these includes needs RHI.h to be included for them to compile (and they are automatically included by RHI.h)
* Removed includes to any of the above includes and replaced them with RHI.h

diff --git a/Engine/Source/Runtime/Engine/Public/DebugViewModeInterface.h b/Engine/Source/Runtime/Engine/Public/DebugViewModeInterface.h
--- a/Engine/Source/Runtime/Engine/Public/DebugViewModeInterface.h
+++ b/Engine/Source/Runtime/Engine/Public/DebugViewModeInterface.h
@@ -10,7 +10,6 @@ DebugViewModeInterface.h: Contains definitions for rendering debug viewmodes.
 #include "DebugViewModeHelpers.h"
 #include "Engine/EngineTypes.h"
 #include "RHI.h"
-#include "RHIResources.h"
 #include "Shader.h"
 #include "DrawDebugHelpers.h"
 
diff --git a/Engine/Source/Runtime/Engine/Public/ExternalTexture.h b/Engine/Source/Runtime/Engine/Public/ExternalTexture.h
--- a/Engine/Source/Runtime/Engine/Public/ExternalTexture.h
+++ b/Engine/Source/Runtime/Engine/Public/ExternalTexture.h
@@ -4,7 +4,6 @@
 
 #include "CoreMinimal.h"
 #include "RHI.h"
-#include "RHIResources.h"
 #include "MaterialShared.h"
 #include "HAL/CriticalSection.h"
 
diff --git a/Engine/Source/Runtime/Engine/Public/GPUSortManager.h b/Engine/Source/Runtime/Engine/Public/GPUSortManager.h
--- a/Engine/Source/Runtime/Engine/Public/GPUSortManager.h
+++ b/Engine/Source/Runtime/Engine/Public/GPUSortManager.h
@@ -9,7 +9,6 @@
 #include "CoreMinimal.h"
 #include "Templates/RefCounting.h"
 #include "RHI.h"
-#include "RHIResources.h"
 #include "RenderResource.h"
 
 /**
diff --git a/Engine/Source/Runtime/RHI/Public/DynamicRHI.h b/Engine/Source/Runtime/RHI/Public/DynamicRHI.h
--- a/Engine/Source/Runtime/RHI/Public/DynamicRHI.h
+++ b/Engine/Source/Runtime/RHI/Public/DynamicRHI.h
@@ -1,5 +1,7 @@
 // Copyright Epic Games, Inc. All Rights Reserved.
 
+// IWYU pragma: private, include "RHI.h"
+
 /*=============================================================================
 DynamicRHI.h: Dynamically bound Render Hardware Interface definitions.
 =============================================================================*/
diff --git a/Engine/Source/Runtime/RHI/Public/RHI.h b/Engine/Source/Runtime/RHI/Public/RHI.h
--- a/Engine/Source/Runtime/RHI/Public/RHI.h
+++ b/Engine/Source/Runtime/RHI/Public/RHI.h
@@ -2513,8 +2513,8 @@ DECLARE_MEMORY_STAT_POOL_EXTERN(TEXT("Pixel buffer memory"),STAT_PixelBufferMemo
 
 
 // RHI base resource types.
-#include "RHIResources.h" // IWYU pragma: export
-#include "DynamicRHI.h" // IWYU pragma: export
+#include "RHIResources.h" // IWYU pragma: keep
+#include "DynamicRHI.h" // IWYU pragma: keep
 
 /** Initializes the RHI. */
 extern RHI_API void RHIInit(bool bHasEditorToken);
@@ -2531,7 +2531,7 @@ DECLARE_DELEGATE_OneParam(FRHIPanicEvent, const FName&);
 extern RHI_API FRHIPanicEvent& RHIGetPanicDelegate();
 
 // RHI utility functions that depend on the RHI definitions.
-#include "RHIUtilities.h" // IWYU pragma: export
+#include "RHIUtilities.h" // IWYU pragma: keep
 
 inline void FRHITransition::Cleanup() const
 {
diff --git a/Engine/Source/Runtime/RHI/Public/RHICommandList.h b/Engine/Source/Runtime/RHI/Public/RHICommandList.h
--- a/Engine/Source/Runtime/RHI/Public/RHICommandList.h
+++ b/Engine/Source/Runtime/RHI/Public/RHICommandList.h
@@ -4,7 +4,12 @@
 	RHICommandList.h: RHI Command List definitions for queueing up & executing later.
 =============================================================================*/
 
+// IWYU pragma: private, include "RHI.h"
+
 #pragma once
+
+// HEADER_UNIT_SKIP - Should always be included through RHI.h (this code needs cleanup.. this file is included by RHIUtilities.h which is included in the middle of RHI.h)
+
 #include "CoreTypes.h"
 #include "Misc/AssertionMacros.h"
 #include "HAL/UnrealMemory.h"
diff --git a/Engine/Source/Runtime/RHI/Public/RHIResources.h b/Engine/Source/Runtime/RHI/Public/RHIResources.h
--- a/Engine/Source/Runtime/RHI/Public/RHIResources.h
+++ b/Engine/Source/Runtime/RHI/Public/RHIResources.h
@@ -1,10 +1,11 @@
 // Copyright Epic Games, Inc. All Rights Reserved.
 
+// IWYU pragma: private, include "RHI.h"
+
 #pragma once
 
 // HEADER_UNIT_SKIP - Should always be included through RHI.h (this code needs cleanup.. this file is included in the middle of RHI.h)
 
-
 #include <atomic>
 
 #include "CoreTypes.h"
diff --git a/Engine/Source/Runtime/RHI/Public/RHIUtilities.h b/Engine/Source/Runtime/RHI/Public/RHIUtilities.h
--- a/Engine/Source/Runtime/RHI/Public/RHIUtilities.h
+++ b/Engine/Source/Runtime/RHI/Public/RHIUtilities.h
@@ -1,8 +1,11 @@
 // Copyright Epic Games, Inc. All Rights Reserved.
 
+// IWYU pragma: private, include "RHI.h"
 
 #pragma once
 
+// HEADER_UNIT_SKIP - Should always be included through RHI.h (this code needs cleanup.. this file is included in the middle of RHI.h)
+
 #include "CoreTypes.h"
 #include "Misc/AssertionMacros.h"
 #include "HAL/UnrealMemory.h"
