Fixing clang static analysis warning.

diff --git a/Engine/Source/Runtime/RHI/Private/Windows/WindowsDynamicRHI.cpp b/Engine/Source/Runtime/RHI/Private/Windows/WindowsDynamicRHI.cpp
--- a/Engine/Source/Runtime/RHI/Private/Windows/WindowsDynamicRHI.cpp
+++ b/Engine/Source/Runtime/RHI/Private/Windows/WindowsDynamicRHI.cpp
@@ -712,7 +712,7 @@ static TOptional<EWindowsRHI> ChooseForcedRHI(TOptional<ERHIFeatureLevel::Type>
 				UE_LOG(LogRHI, Fatal, TEXT("%s does not support %s."), GetLogName(ForcedRHI.GetValue()), GetLogName(ForcedFeatureLevel.GetValue()));
 			}
 		}
-		else if (ForcedRHI = Config.GetFirstRHIWithFeatureLevelSupport(ForcedFeatureLevel.GetValue()))
+		else if ((ForcedRHI = Config.GetFirstRHIWithFeatureLevelSupport(ForcedFeatureLevel.GetValue())))
 		{
 			UE_LOG(LogRHI, Log, TEXT("Forcing RHI to %s since Feature Level %s was forced"), GetLogName(ForcedRHI.GetValue()), GetLogName(ForcedFeatureLevel.GetValue()));
 		}
