Remove STAT prefix from two named events.

diff --git a/Engine/Source/Runtime/RHI/Private/RHICommandList.cpp b/Engine/Source/Runtime/RHI/Private/RHICommandList.cpp
--- a/Engine/Source/Runtime/RHI/Private/RHICommandList.cpp
+++ b/Engine/Source/Runtime/RHI/Private/RHICommandList.cpp
@@ -685,7 +685,7 @@ void FRHICommandListImmediate::QueueAsyncCommandListSubmit(TArrayView<FQueuedCom
 				FFunctionGraphTask::CreateAndDispatchWhenReady(
 					[FenceCandidates = MoveTemp(PersistentState.QueuedFenceCandidates), FenceCandidateEvent](ENamedThreads::Type, const FGraphEventRef&) mutable
 				{
-					SCOPED_NAMED_EVENT(STAT_FRHICommandListBase_SignalLockFence, FColor::Magenta);
+					SCOPED_NAMED_EVENT(FRHICommandListBase_SignalLockFence, FColor::Magenta);
 
 					for (int32 Index = FenceCandidates.Num() - 1; Index >= 0; Index--)
 					{
diff --git a/Engine/Source/Runtime/Renderer/Private/GPUScene.cpp b/Engine/Source/Runtime/Renderer/Private/GPUScene.cpp
--- a/Engine/Source/Runtime/Renderer/Private/GPUScene.cpp
+++ b/Engine/Source/Runtime/Renderer/Private/GPUScene.cpp
@@ -727,7 +727,7 @@ void FGPUScene::UpdateInternal(FRDGBuilder& GraphBuilder, FScene& Scene, FRDGExt
 	PrimitiveDirtyState.Init(EPrimitiveDirtyState::None, PrimitiveDirtyState.Num());
 
 	{
-		SCOPED_NAMED_EVENT(STAT_UpdateGPUScene, FColor::Green);
+		SCOPED_NAMED_EVENT(UpdateGPUScene, FColor::Green);
 		QUICK_SCOPE_CYCLE_COUNTER(STAT_UpdateGPUScene);
 		SCOPE_CYCLE_COUNTER(STAT_UpdateGPUSceneTime);
 
