Allow running editor with: -game -featureleveles31

diff --git a/Engine/Source/Runtime/RHI/Private/Windows/WindowsDynamicRHI.cpp b/Engine/Source/Runtime/RHI/Private/Windows/WindowsDynamicRHI.cpp
--- a/Engine/Source/Runtime/RHI/Private/Windows/WindowsDynamicRHI.cpp
+++ b/Engine/Source/Runtime/RHI/Private/Windows/WindowsDynamicRHI.cpp
@@ -702,7 +702,7 @@ static TOptional<EWindowsRHI> ChooseForcedRHI(TOptional<ERHIFeatureLevel::Type>
 	}
 
 	// FeatureLevelES31 is also a command line override, so it will determine the underlying RHI unless one is specified
-	if ((FPlatformProperties::RequiresCookedData() || !GIsEditor) && ForcedFeatureLevel)
+	if (FPlatformProperties::RequiresCookedData() && ForcedFeatureLevel)
 	{
 		if (ForcedRHI)
 		{
