RHI header dependency cleanup
- Reducing other header dependencies from RHIDefinitions.h
- Moved a few functions to the new RHIStrings.h
- Uniform buffer layout initializer moved to its own file to reduce RHI dependency on MemoryLayout.

diff --git a/Engine/Plugins/MeshPainting/Source/MeshPaintingToolset/Private/TexturePaintToolset.cpp b/Engine/Plugins/MeshPainting/Source/MeshPaintingToolset/Private/TexturePaintToolset.cpp
--- a/Engine/Plugins/MeshPainting/Source/MeshPaintingToolset/Private/TexturePaintToolset.cpp
+++ b/Engine/Plugins/MeshPainting/Source/MeshPaintingToolset/Private/TexturePaintToolset.cpp
@@ -317,6 +317,16 @@ bool UTexturePaintToolset::GenerateSeamMask(UMeshComponent* MeshComponent, int32
 	return RetVal;
 }
 
+int32 UTexturePaintToolset::GetMaxSupportedBytesPerPixelForPainting()
+{
+	return GPixelFormats[GetTempUncompressedTexturePixelFormat()].BlockBytes;
+}
+
+EPixelFormat UTexturePaintToolset::GetTempUncompressedTexturePixelFormat()
+{
+	return PF_B8G8R8A8;
+}
+
 UTexture2D* UTexturePaintToolset::CreateScratchUncompressedTexture(UTexture2D* SourceTexture)
 {
 	check(SourceTexture->Source.IsValid());
diff --git a/Engine/Plugins/MeshPainting/Source/MeshPaintingToolset/Public/TexturePaintToolset.h b/Engine/Plugins/MeshPainting/Source/MeshPaintingToolset/Public/TexturePaintToolset.h
--- a/Engine/Plugins/MeshPainting/Source/MeshPaintingToolset/Public/TexturePaintToolset.h
+++ b/Engine/Plugins/MeshPainting/Source/MeshPaintingToolset/Public/TexturePaintToolset.h
@@ -72,16 +72,10 @@ public:
 	static bool GenerateSeamMask(UMeshComponent* MeshComponent, int32 UVSet, UTextureRenderTarget2D* SeamRenderTexture, UTexture2D* Texture, UTextureRenderTarget2D* RenderTargetTexture );
 
 	/** Returns the maximum bytes per pixel that are supported for source textures when painting. This limitation is set by CreateTempUncompressedTexture. */
-	static constexpr int32 GetMaxSupportedBytesPerPixelForPainting()
-	{
-		return GPixelFormats[GetTempUncompressedTexturePixelFormat()].BlockBytes;
-	}
+	static int32 GetMaxSupportedBytesPerPixelForPainting();
 
 	/** Returns the pixel format that CreateTempUncompressedTexture uses to create render target data for painting. */
-	static constexpr EPixelFormat GetTempUncompressedTexturePixelFormat()
-	{
-		return EPixelFormat::PF_B8G8R8A8;
-	}
+	static EPixelFormat GetTempUncompressedTexturePixelFormat();
 
 	/** Static: Creates a temporary texture used to transfer data to a render target in memory */
 	static UTexture2D* CreateScratchUncompressedTexture(UTexture2D* SourceTexture);
diff --git a/Engine/Plugins/RenderGraphInsights/Source/RenderGraphInsights/Private/RenderGraphTrack.cpp b/Engine/Plugins/RenderGraphInsights/Source/RenderGraphInsights/Private/RenderGraphTrack.cpp
--- a/Engine/Plugins/RenderGraphInsights/Source/RenderGraphInsights/Private/RenderGraphTrack.cpp
+++ b/Engine/Plugins/RenderGraphInsights/Source/RenderGraphInsights/Private/RenderGraphTrack.cpp
@@ -5,6 +5,7 @@
 #include "Modules/ModuleManager.h"
 #include "RenderGraphProvider.h"
 #include "RenderGraphTimingViewSession.h"
+#include "RHIStrings.h"
 
 #include "Application/SlateApplicationBase.h"
 #include "Fonts/FontMeasure.h"
diff --git a/Engine/Source/Developer/Windows/ShaderFormatD3D/Private/D3DShaderCompiler.cpp b/Engine/Source/Developer/Windows/ShaderFormatD3D/Private/D3DShaderCompiler.cpp
--- a/Engine/Source/Developer/Windows/ShaderFormatD3D/Private/D3DShaderCompiler.cpp
+++ b/Engine/Source/Developer/Windows/ShaderFormatD3D/Private/D3DShaderCompiler.cpp
@@ -9,6 +9,7 @@
 #include "Misc/Paths.h"
 #include "Misc/FileHelper.h"
 #include "HAL/FileManager.h"
+#include "RHIStrings.h"
 #include "Serialization/MemoryWriter.h"
 #include "RayTracingDefinitions.h"
 #include "SpirvCommon.h"
diff --git a/Engine/Source/Developer/Windows/ShaderFormatD3D/Private/D3DShaderCompilerDXC.cpp b/Engine/Source/Developer/Windows/ShaderFormatD3D/Private/D3DShaderCompilerDXC.cpp
--- a/Engine/Source/Developer/Windows/ShaderFormatD3D/Private/D3DShaderCompilerDXC.cpp
+++ b/Engine/Source/Developer/Windows/ShaderFormatD3D/Private/D3DShaderCompilerDXC.cpp
@@ -10,6 +10,7 @@
 #include "HAL/FileManager.h"
 #include "Serialization/MemoryWriter.h"
 #include "RayTracingDefinitions.h"
+#include "RHIStrings.h"
 
 DEFINE_LOG_CATEGORY_STATIC(LogD3D12ShaderCompiler, Log, All);
 
diff --git a/Engine/Source/Editor/MeshPaint/Private/TexturePaintHelpers.cpp b/Engine/Source/Editor/MeshPaint/Private/TexturePaintHelpers.cpp
--- a/Engine/Source/Editor/MeshPaint/Private/TexturePaintHelpers.cpp
+++ b/Engine/Source/Editor/MeshPaint/Private/TexturePaintHelpers.cpp
@@ -315,6 +315,17 @@ bool TexturePaintHelpers::GenerateSeamMask(UMeshComponent* MeshComponent, int32
 	return RetVal;
 }
 
+int32 TexturePaintHelpers::GetMaxSupportedBytesPerPixelForPainting()
+{
+	return GPixelFormats[GetTempUncompressedTexturePixelFormat()].BlockBytes;
+}
+
+/** Returns the pixel format that CreateTempUncompressedTexture uses to create render target data for painting. */
+EPixelFormat TexturePaintHelpers::GetTempUncompressedTexturePixelFormat()
+{
+	return PF_B8G8R8A8;
+}
+
 UTexture2D* TexturePaintHelpers::CreateTempUncompressedTexture(UTexture2D* SourceTexture)
 {
 	check(SourceTexture->Source.IsValid());
diff --git a/Engine/Source/Editor/MeshPaint/Public/TexturePaintHelpers.h b/Engine/Source/Editor/MeshPaint/Public/TexturePaintHelpers.h
--- a/Engine/Source/Editor/MeshPaint/Public/TexturePaintHelpers.h
+++ b/Engine/Source/Editor/MeshPaint/Public/TexturePaintHelpers.h
@@ -66,16 +66,10 @@ public:
 	static bool GenerateSeamMask(UMeshComponent* MeshComponent, int32 UVSet, UTextureRenderTarget2D* SeamRenderTexture, UTexture2D* Texture, UTextureRenderTarget2D* RenderTargetTexture );
 
 	/** Returns the maximum bytes per pixel that are supported for source textures when painting. This limitation is set by CreateTempUncompressedTexture. */
-	static constexpr int32 GetMaxSupportedBytesPerPixelForPainting()
-	{
-		return GPixelFormats[GetTempUncompressedTexturePixelFormat()].BlockBytes;
-	}
+	static int32 GetMaxSupportedBytesPerPixelForPainting();
 
 	/** Returns the pixel format that CreateTempUncompressedTexture uses to create render target data for painting. */
-	static constexpr EPixelFormat GetTempUncompressedTexturePixelFormat()
-	{
-		return EPixelFormat::PF_B8G8R8A8;
-	}
+	static EPixelFormat GetTempUncompressedTexturePixelFormat();
 
 	/** Static: Creates a temporary texture used to transfer data to a render target in memory */
 	static UTexture2D* CreateTempUncompressedTexture(UTexture2D* SourceTexture);
diff --git a/Engine/Source/Programs/UnrealBuildTool/Configuration/UEBuildModuleCPP.cs b/Engine/Source/Programs/UnrealBuildTool/Configuration/UEBuildModuleCPP.cs
--- a/Engine/Source/Programs/UnrealBuildTool/Configuration/UEBuildModuleCPP.cs
+++ b/Engine/Source/Programs/UnrealBuildTool/Configuration/UEBuildModuleCPP.cs
@@ -1532,6 +1532,7 @@ namespace UnrealBuildTool
 			}
 			switch(Name)
 			{
+				case "RHI":
 				case "D3D11RHI":
 				case "D3D12RHI":
 				case "VulkanRHI":
diff --git a/Engine/Source/Runtime/Engine/Private/LightMap.cpp b/Engine/Source/Runtime/Engine/Private/LightMap.cpp
--- a/Engine/Source/Runtime/Engine/Private/LightMap.cpp
+++ b/Engine/Source/Runtime/Engine/Private/LightMap.cpp
@@ -6,6 +6,7 @@
 
 #include "LightMap.h"
 #include "Components/PrimitiveComponent.h"
+#include "DataDrivenShaderPlatformInfo.h"
 #include "Engine/Level.h"
 #include "Misc/QueuedThreadPool.h"
 #include "ShadowMap.h"
diff --git a/Engine/Source/Runtime/Landscape/Public/LandscapeLight.h b/Engine/Source/Runtime/Landscape/Public/LandscapeLight.h
--- a/Engine/Source/Runtime/Landscape/Public/LandscapeLight.h
+++ b/Engine/Source/Runtime/Landscape/Public/LandscapeLight.h
@@ -7,6 +7,7 @@
 #pragma once
 
 #include "CoreMinimal.h"
+#include "PixelFormat.h"
 #include "StaticLighting.h"
 #include "RenderUtils.h"
 
diff --git a/Engine/Source/Runtime/OpenGLDrv/Private/OpenGLCommands.cpp b/Engine/Source/Runtime/OpenGLDrv/Private/OpenGLCommands.cpp
--- a/Engine/Source/Runtime/OpenGLDrv/Private/OpenGLCommands.cpp
+++ b/Engine/Source/Runtime/OpenGLDrv/Private/OpenGLCommands.cpp
@@ -10,6 +10,7 @@
 #include "Misc/App.h"
 #include "RHIDefinitions.h"
 #include "RHI.h"
+#include "RHIStrings.h"
 #include "EngineGlobals.h"
 #include "RenderResource.h"
 #include "OpenGLDrv.h"
diff --git a/Engine/Source/Runtime/RHI/Private/DynamicRHI.cpp b/Engine/Source/Runtime/RHI/Private/DynamicRHI.cpp
--- a/Engine/Source/Runtime/RHI/Private/DynamicRHI.cpp
+++ b/Engine/Source/Runtime/RHI/Private/DynamicRHI.cpp
@@ -17,9 +17,6 @@
 #include "TextureProfiler.h"
 #include "DataDrivenShaderPlatformInfo.h"
 
-IMPLEMENT_TYPE_LAYOUT(FRayTracingGeometryInitializer);
-IMPLEMENT_TYPE_LAYOUT(FRayTracingGeometrySegment);
-
 static_assert(sizeof(FRayTracingGeometryInstance) <= 104,
 	"Ray tracing instance descriptor is expected to be no more than 104 bytes, as there may be a very large number of them.");
 
diff --git a/Engine/Source/Runtime/RHI/Private/RHI.cpp b/Engine/Source/Runtime/RHI/Private/RHI.cpp
--- a/Engine/Source/Runtime/RHI/Private/RHI.cpp
+++ b/Engine/Source/Runtime/RHI/Private/RHI.cpp
@@ -18,6 +18,8 @@
 #include "Misc/OutputDeviceFile.h"
 #include "Async/ParallelFor.h"
 #include "DataDrivenShaderPlatformInfo.h"
+#include "RHIStrings.h"
+#include "RHIUniformBufferLayoutInitializer.h"
 
 #if RHI_ENABLE_RESOURCE_INFO
 #include "HAL/FileManager.h"
@@ -57,7 +59,7 @@ DEFINE_STAT(STAT_StructuredBufferMemory);
 DEFINE_STAT(STAT_PixelBufferMemory);
 
 IMPLEMENT_TYPE_LAYOUT(FRHIUniformBufferLayoutInitializer);
-IMPLEMENT_TYPE_LAYOUT(FRHIUniformBufferResource);
+IMPLEMENT_TYPE_LAYOUT(FRHIUniformBufferResourceInitializer);
 
 #if !defined(RHIRESOURCE_NUM_FRAMES_TO_EXPIRE)
 	#define RHIRESOURCE_NUM_FRAMES_TO_EXPIRE 3
diff --git a/Engine/Source/Runtime/RHI/Private/RHIResources.cpp b/Engine/Source/Runtime/RHI/Private/RHIResources.cpp
new file mode 100644
--- /dev/null
+++ b/Engine/Source/Runtime/RHI/Private/RHIResources.cpp
@@ -0,0 +1,29 @@
+// Copyright Epic Games, Inc. All Rights Reserved.
+
+#include "RHI.h"
+#include "RHIUniformBufferLayoutInitializer.h"
+
+template <>
+struct TIsBitwiseConstructible<FRHIUniformBufferResource, FRHIUniformBufferResourceInitializer>
+{
+	enum { Value = true };
+};
+
+FRHIUniformBufferLayout::FRHIUniformBufferLayout(const FRHIUniformBufferLayoutInitializer& Initializer)
+	: FRHIResource(RRT_UniformBufferLayout)
+	, Name(Initializer.GetDebugName())
+	, Resources(TConstArrayView<FRHIUniformBufferResourceInitializer>(Initializer.Resources))
+	, GraphResources(TConstArrayView<FRHIUniformBufferResourceInitializer>(Initializer.GraphResources))
+	, GraphTextures(TConstArrayView<FRHIUniformBufferResourceInitializer>(Initializer.GraphTextures))
+	, GraphBuffers(TConstArrayView<FRHIUniformBufferResourceInitializer>(Initializer.GraphBuffers))
+	, GraphUniformBuffers(TConstArrayView<FRHIUniformBufferResourceInitializer>(Initializer.GraphUniformBuffers))
+	, UniformBuffers(TConstArrayView<FRHIUniformBufferResourceInitializer>(Initializer.UniformBuffers))
+	, Hash(Initializer.GetHash())
+	, ConstantBufferSize(Initializer.ConstantBufferSize)
+	, RenderTargetsOffset(Initializer.RenderTargetsOffset)
+	, StaticSlot(Initializer.StaticSlot)
+	, BindingFlags(Initializer.BindingFlags)
+	, bHasNonGraphOutputs(Initializer.bHasNonGraphOutputs)
+	, bNoEmulatedUniformBuffer(Initializer.bNoEmulatedUniformBuffer)
+{
+}
diff --git a/Engine/Source/Runtime/RHI/Public/DynamicRHI.h b/Engine/Source/Runtime/RHI/Public/DynamicRHI.h
--- a/Engine/Source/Runtime/RHI/Public/DynamicRHI.h
+++ b/Engine/Source/Runtime/RHI/Public/DynamicRHI.h
@@ -35,6 +35,7 @@ struct FSamplerStateInitializerRHI;
 struct FTextureMemoryStats;
 struct FRHIGPUMask;
 
+using FDisplayInformationArray = TArray<struct FDisplayInformation>;
 
 /** Struct to hold common data between begin/end updatetexture3d */
 struct FUpdateTexture3DData
diff --git a/Engine/Source/Runtime/RHI/Public/RHI.h b/Engine/Source/Runtime/RHI/Public/RHI.h
--- a/Engine/Source/Runtime/RHI/Public/RHI.h
+++ b/Engine/Source/Runtime/RHI/Public/RHI.h
@@ -6,7 +6,9 @@
 
 #pragma once
 
+#if UE_ENABLE_INCLUDE_ORDER_DEPRECATED_IN_5_2
 #include "CoreMinimal.h"
+#endif
 #include "Containers/ContainersFwd.h"
 #include "Stats/Stats.h"
 #include "RHIDefinitions.h"
@@ -16,6 +18,7 @@
 #include "Math/PerspectiveMatrix.h"
 #include "Math/ScaleMatrix.h"
 #include "Math/TranslationMatrix.h"
+#include "PixelFormat.h"
 #include "ProfilingDebugging/CsvProfilerConfig.h"
 #include "GpuProfilerTrace.h"
 #include "RHIShaderPlatform.h"
diff --git a/Engine/Source/Runtime/RHI/Public/RHIDefinitions.h b/Engine/Source/Runtime/RHI/Public/RHIDefinitions.h
--- a/Engine/Source/Runtime/RHI/Public/RHIDefinitions.h
+++ b/Engine/Source/Runtime/RHI/Public/RHIDefinitions.h
@@ -7,15 +7,17 @@
 
 #pragma once
 
-#include "PixelFormat.h"
-#include "Serialization/MemoryLayout.h"
+#include "HAL/Platform.h"
+#include "Misc/EnumClassFlags.h"
 
 #if UE_ENABLE_INCLUDE_ORDER_DEPRECATED_IN_5_2
 #include "CoreMinimal.h"
+#include "PixelFormat.h"
 #include "HAL/IConsoleManager.h"
 #include "RHIShaderPlatform.h"
 #include "RHIFeatureLevel.h"
 #include "RHIImmutableSamplerState.h"
+#include "Serialization/MemoryLayout.h"
 #endif
 
 /** Alignment of the shader parameters struct is required to be 16-byte boundaries. */
@@ -399,7 +401,6 @@ enum EVertexElementType
 	VET_NumBits = 5,
 };
 static_assert(VET_MAX <= (1 << VET_NumBits), "VET_MAX will not fit on VET_NumBits");
-DECLARE_INTRINSIC_TYPE_LAYOUT(EVertexElementType);
 
 enum ECubeFace : uint32
 {
@@ -518,7 +519,6 @@ enum EUniformBufferBaseType : uint8
 	EUniformBufferBaseType_NumBits = 5,
 };
 static_assert(EUniformBufferBaseType_Num <= (1 << EUniformBufferBaseType_NumBits), "EUniformBufferBaseType_Num will not fit on EUniformBufferBaseType_NumBits");
-DECLARE_INTRINSIC_TYPE_LAYOUT(EUniformBufferBaseType);
 
 /** The list of flags declaring which binding models are allowed for a uniform buffer layout. */
 enum class EUniformBufferBindingFlags : uint8
@@ -537,7 +537,6 @@ enum class EUniformBufferBindingFlags : uint8
 	StaticAndShader = Static | Shader
 };
 ENUM_CLASS_FLAGS(EUniformBufferBindingFlags);
-DECLARE_INTRINSIC_TYPE_LAYOUT(EUniformBufferBindingFlags);
 
 /** Numerical type used to store the static slot indices. */
 using FUniformBufferStaticSlot = uint8;
@@ -743,7 +742,6 @@ enum class EBufferUsageFlags : uint32
 
 	/** Buffer should go in fast vram (hint only). Requires BUF_Transient */
 	FastVRAM                = 1 << 10,
-
 	/** Create a buffer that can be shared with an external RHI or process. */
 	Shared                  = 1 << 12,
 
@@ -1103,8 +1101,7 @@ enum class ERHIDescriptorHeapType : uint8
 	Sampler,
 	RenderTarget,
 	DepthStencil,
-	Count,
-	Invalid = MAX_uint8
+	Count
 };
 
 struct FRHIDescriptorHandle
@@ -1125,14 +1122,15 @@ struct FRHIDescriptorHandle
 	inline ERHIDescriptorHeapType GetType() const { return (ERHIDescriptorHeapType)Type; }
 	inline uint8                  GetRawType() const { return Type; }
 
-	inline bool IsValid() const { return Index != UINT_MAX && Type != (uint8)ERHIDescriptorHeapType::Invalid; }
+	inline bool IsValid() const { return Index != InvalidIndex && Type !=InvalidType; }
 
 private:
-	uint32    Index{ UINT_MAX };
-	uint8     Type{ (uint8)ERHIDescriptorHeapType::Invalid };
-};
+	static constexpr uint32 InvalidIndex = 0xffffffff;
+	static constexpr uint8  InvalidType  = 0xff;
 
-using FDisplayInformationArray = TArray<struct FDisplayInformation>;
+	uint32    Index{ InvalidIndex };
+	uint8     Type{ InvalidType };
+};
 
 enum class ERHIBindlessConfiguration
 {
@@ -1235,222 +1233,6 @@ inline EGpuVendorId RHIConvertToGpuVendorId(uint32 VendorId)
 	return EGpuVendorId::Unknown;
 }
 
-inline const TCHAR* GetShaderFrequencyString(EShaderFrequency Frequency, bool bIncludePrefix = true)
-{
-	const TCHAR* String = TEXT("SF_NumFrequencies");
-	switch (Frequency)
-	{
-	case SF_Vertex:			String = TEXT("SF_Vertex"); break;
-	case SF_Mesh:			String = TEXT("SF_Mesh"); break;
-	case SF_Amplification:	String = TEXT("SF_Amplification"); break;
-	case SF_Geometry:		String = TEXT("SF_Geometry"); break;
-	case SF_Pixel:			String = TEXT("SF_Pixel"); break;
-	case SF_Compute:		String = TEXT("SF_Compute"); break;
-	case SF_RayGen:			String = TEXT("SF_RayGen"); break;
-	case SF_RayMiss:		String = TEXT("SF_RayMiss"); break;
-	case SF_RayHitGroup:	String = TEXT("SF_RayHitGroup"); break;
-	case SF_RayCallable:	String = TEXT("SF_RayCallable"); break;
-
-	default:
-		checkf(0, TEXT("Unknown ShaderFrequency %d"), (int32)Frequency);
-		break;
-	}
-
-	// Skip SF_
-	int32 Index = bIncludePrefix ? 0 : 3;
-	String += Index;
-	return String;
-};
-
-inline const TCHAR* GetTextureDimensionString(ETextureDimension Dimension)
-{
-	switch (Dimension)
-	{
-	case ETextureDimension::Texture2D:
-		return TEXT("Texture2D");
-	case ETextureDimension::Texture2DArray:
-		return TEXT("Texture2DArray");
-	case ETextureDimension::Texture3D:
-		return TEXT("Texture3D");
-	case ETextureDimension::TextureCube:
-		return TEXT("TextureCube");
-	case ETextureDimension::TextureCubeArray:
-		return TEXT("TextureCubeArray");
-	}
-	return TEXT("");
-}
-
-inline const TCHAR* GetTextureCreateFlagString(ETextureCreateFlags TextureCreateFlag)
-{
-	switch (TextureCreateFlag)
-	{
-	case ETextureCreateFlags::None:
-		return TEXT("None");
-	case ETextureCreateFlags::RenderTargetable:
-		return TEXT("RenderTargetable");
-	case ETextureCreateFlags::ResolveTargetable:
-		return TEXT("ResolveTargetable");
-	case ETextureCreateFlags::DepthStencilTargetable:
-		return TEXT("DepthStencilTargetable");
-	case ETextureCreateFlags::ShaderResource:
-		return TEXT("ShaderResource");
-	case ETextureCreateFlags::SRGB:
-		return TEXT("SRGB");
-	case ETextureCreateFlags::CPUWritable:
-		return TEXT("CPUWritable");
-	case ETextureCreateFlags::NoTiling:
-		return TEXT("NoTiling");
-	case ETextureCreateFlags::VideoDecode:
-		return TEXT("VideoDecode");
-	case ETextureCreateFlags::Dynamic:
-		return TEXT("Dynamic");
-	case ETextureCreateFlags::InputAttachmentRead:
-		return TEXT("InputAttachmentRead");
-	case ETextureCreateFlags::Foveation:
-		return TEXT("Foveation");
-	case ETextureCreateFlags::Tiling3D:
-		return TEXT("Tiling3D");
-	case ETextureCreateFlags::Memoryless:
-		return TEXT("Memoryless");
-	case ETextureCreateFlags::GenerateMipCapable:
-		return TEXT("GenerateMipCapable");
-	case ETextureCreateFlags::FastVRAMPartialAlloc:
-		return TEXT("FastVRAMPartialAlloc");
-	case ETextureCreateFlags::DisableSRVCreation:
-		return TEXT("DisableSRVCreation");
-	case ETextureCreateFlags::DisableDCC:
-		return TEXT("DisableDCC");
-	case ETextureCreateFlags::UAV:
-		return TEXT("UAV");
-	case ETextureCreateFlags::Presentable:
-		return TEXT("Presentable");
-	case ETextureCreateFlags::CPUReadback:
-		return TEXT("CPUReadback");
-	case ETextureCreateFlags::OfflineProcessed:
-		return TEXT("OfflineProcessed");
-	case ETextureCreateFlags::FastVRAM:
-		return TEXT("FastVRAM");
-	case ETextureCreateFlags::HideInVisualizeTexture:
-		return TEXT("HideInVisualizeTexture");
-	case ETextureCreateFlags::Virtual:
-		return TEXT("Virtual");
-	case ETextureCreateFlags::TargetArraySlicesIndependently:
-		return TEXT("TargetArraySlicesIndependently");
-	case ETextureCreateFlags::Shared:
-		return TEXT("Shared");
-	case ETextureCreateFlags::NoFastClear:
-		return TEXT("NoFastClear");
-	case ETextureCreateFlags::DepthStencilResolveTarget:
-		return TEXT("DepthStencilResolveTarget");
-	case ETextureCreateFlags::Streamable:
-		return TEXT("Streamable");
-	case ETextureCreateFlags::NoFastClearFinalize:
-		return TEXT("NoFastClearFinalize");
-	case ETextureCreateFlags::AFRManual:
-		return TEXT("AFRManual");
-	case ETextureCreateFlags::ReduceMemoryWithTilingMode:
-		return TEXT("ReduceMemoryWithTilingMode");
-	}
-	return TEXT("");
-}
-
-inline const TCHAR* GetBufferUsageFlagString(EBufferUsageFlags BufferUsage)
-{
-	switch (BufferUsage)
-	{
-	case EBufferUsageFlags::None:
-		return TEXT("None");
-	case EBufferUsageFlags::Static:
-		return TEXT("Static");
-	case EBufferUsageFlags::Dynamic:
-		return TEXT("Dynamic");
-	case EBufferUsageFlags::Volatile:
-		return TEXT("Volatile");
-	case EBufferUsageFlags::UnorderedAccess:
-		return TEXT("UnorderedAccess");
-	case EBufferUsageFlags::ByteAddressBuffer:
-		return TEXT("ByteAddressBuffer");
-	case EBufferUsageFlags::SourceCopy:
-		return TEXT("SourceCopy");
-	case EBufferUsageFlags::StreamOutput:
-		return TEXT("StreamOutput");
-	case EBufferUsageFlags::DrawIndirect:
-		return TEXT("DrawIndirect");
-	case EBufferUsageFlags::ShaderResource:
-		return TEXT("ShaderResource");
-	case EBufferUsageFlags::KeepCPUAccessible:
-		return TEXT("KeepCPUAccessible");
-	case EBufferUsageFlags::FastVRAM:
-		return TEXT("FastVRAM");
-	case EBufferUsageFlags::Shared:
-		return TEXT("Shared");
-	case EBufferUsageFlags::AccelerationStructure:
-		return TEXT("AccelerationStructure");
-	case EBufferUsageFlags::VertexBuffer:
-		return TEXT("VertexBuffer");
-	case EBufferUsageFlags::IndexBuffer:
-		return TEXT("IndexBuffer");
-	case EBufferUsageFlags::StructuredBuffer:
-		return TEXT("StructuredBuffer");
-	}
-	return TEXT("");
-}
-
-inline const TCHAR* GetUniformBufferBaseTypeString(EUniformBufferBaseType BaseType)
-{
-	switch (BaseType)
-	{
-	case UBMT_INVALID:
-		return TEXT("UBMT_INVALID");
-	case UBMT_BOOL:
-		return TEXT("UBMT_BOOL");
-	case UBMT_INT32:
-		return TEXT("UBMT_INT32");
-	case UBMT_UINT32:
-		return TEXT("UBMT_UINT32");
-	case UBMT_FLOAT32:
-		return TEXT("UBMT_FLOAT32");
-	case UBMT_TEXTURE:
-		return TEXT("UBMT_TEXTURE");
-	case UBMT_SRV:
-		return TEXT("UBMT_SRV");
-	case UBMT_UAV:
-		return TEXT("UBMT_UAV");
-	case UBMT_SAMPLER:
-		return TEXT("UBMT_SAMPLER");
-	case UBMT_RDG_TEXTURE:
-		return TEXT("UBMT_RDG_TEXTURE");
-	case UBMT_RDG_TEXTURE_ACCESS:
-		return TEXT("UBMT_RDG_TEXTURE_ACCESS");
-	case UBMT_RDG_TEXTURE_ACCESS_ARRAY:
-		return TEXT("UBMT_RDG_TEXTURE_ACCESS_ARRAY");
-	case UBMT_RDG_TEXTURE_SRV:
-		return TEXT("UBMT_RDG_TEXTURE_SRV");
-	case UBMT_RDG_TEXTURE_UAV:
-		return TEXT("UBMT_RDG_TEXTURE_UAV");
-	case UBMT_RDG_BUFFER_ACCESS:
-		return TEXT("UBMT_RDG_BUFFER_ACCESS");
-	case UBMT_RDG_BUFFER_ACCESS_ARRAY:
-		return TEXT("UBMT_RDG_BUFFER_ACCESS_ARRAY");
-	case UBMT_RDG_BUFFER_SRV:
-		return TEXT("UBMT_RDG_BUFFER_SRV");
-	case UBMT_RDG_BUFFER_UAV:
-		return TEXT("UBMT_RDG_BUFFER_UAV");
-	case UBMT_RDG_UNIFORM_BUFFER:
-		return TEXT("UBMT_RDG_UNIFORM_BUFFER");
-	case UBMT_NESTED_STRUCT:
-		return TEXT("UBMT_NESTED_STRUCT");
-	case UBMT_INCLUDED_STRUCT:
-		return TEXT("UBMT_INCLUDED_STRUCT");
-	case UBMT_REFERENCED_STRUCT:
-		return TEXT("UBMT_REFERENCED_STRUCT");
-	case UBMT_RENDER_TARGET_BINDING_SLOTS:
-		return TEXT("UBMT_RENDER_TARGET_BINDING_SLOTS");
-	}
-	return TEXT("");
-}
-
-
 inline bool IsGeometryPipelineShaderFrequency(EShaderFrequency Frequency)
 {
 	return Frequency == SF_Mesh || Frequency == SF_Amplification;
@@ -1470,24 +1252,6 @@ inline bool IsRayTracingShaderFrequency(EShaderFrequency Frequency)
 	}
 }
 
-inline ERHIResourceType GetRHIResourceType(ETextureDimension Dimension)
-{
-	switch (Dimension)
-	{
-	case ETextureDimension::Texture2D:
-		return ERHIResourceType::RRT_Texture2D;
-	case ETextureDimension::Texture2DArray:
-		return ERHIResourceType::RRT_Texture2DArray;
-	case ETextureDimension::Texture3D:
-		return ERHIResourceType::RRT_Texture3D;
-	case ETextureDimension::TextureCube:
-	case ETextureDimension::TextureCubeArray:
-		return ERHIResourceType::RRT_TextureCube;
-	}
-	checkNoEntry();
-	return ERHIResourceType::RRT_None;
-}
-
 #if PLATFORM_SUPPORTS_GEOMETRY_SHADERS
 	#define GEOMETRY_SHADER(GeometryShader)	(GeometryShader)
 #else
@@ -1496,4 +1260,5 @@ inline ERHIResourceType GetRHIResourceType(ETextureDimension Dimension)
 
 #if UE_ENABLE_INCLUDE_ORDER_DEPRECATED_IN_5_2
 #include "DataDrivenShaderPlatformInfo.h"
+#include "RHIStrings.h"
 #endif
diff --git a/Engine/Source/Runtime/RHI/Public/RHIFeatureLevel.h b/Engine/Source/Runtime/RHI/Public/RHIFeatureLevel.h
--- a/Engine/Source/Runtime/RHI/Public/RHIFeatureLevel.h
+++ b/Engine/Source/Runtime/RHI/Public/RHIFeatureLevel.h
@@ -3,7 +3,9 @@
 #pragma once
 
 #include "Containers/EnumAsByte.h"
+#if UE_ENABLE_INCLUDE_ORDER_DEPRECATED_IN_5_2
 #include "Serialization/MemoryLayout.h"
+#endif
 
 #ifndef USE_STATIC_FEATURE_LEVEL_ENUMS
 #define USE_STATIC_FEATURE_LEVEL_ENUMS 0
@@ -54,7 +56,6 @@ namespace ERHIFeatureLevel
 		Num
 	};
 };
-DECLARE_INTRINSIC_TYPE_LAYOUT(ERHIFeatureLevel::Type);
 
 struct FGenericStaticFeatureLevel
 {
diff --git a/Engine/Source/Runtime/RHI/Public/RHIMemoryLayout.h b/Engine/Source/Runtime/RHI/Public/RHIMemoryLayout.h
new file mode 100644
--- /dev/null
+++ b/Engine/Source/Runtime/RHI/Public/RHIMemoryLayout.h
@@ -0,0 +1,12 @@
+// Copyright Epic Games, Inc. All Rights Reserved.
+
+#pragma once
+
+#include "RHIDefinitions.h"
+#include "RHIFeatureLevel.h"
+#include "Serialization/MemoryLayout.h"
+
+DECLARE_INTRINSIC_TYPE_LAYOUT(EUniformBufferBaseType);
+DECLARE_INTRINSIC_TYPE_LAYOUT(EUniformBufferBindingFlags);
+DECLARE_INTRINSIC_TYPE_LAYOUT(EVertexElementType);
+DECLARE_INTRINSIC_TYPE_LAYOUT(ERHIFeatureLevel::Type);
diff --git a/Engine/Source/Runtime/RHI/Public/RHIResources.h b/Engine/Source/Runtime/RHI/Public/RHIResources.h
--- a/Engine/Source/Runtime/RHI/Public/RHIResources.h
+++ b/Engine/Source/Runtime/RHI/Public/RHIResources.h
@@ -30,12 +30,15 @@
 #include "Containers/ConsumeAllMpmcQueue.h"
 #include "Misc/CoreDelegates.h"
 #include "RHIFwd.h"
+#include "RHIAccess.h"
+#include "RHIDefinitions.h"
 #include "RHIImmutableSamplerState.h"
 
 class FRHIComputeCommandList;
 class FRHICommandListImmediate;
 struct FClearValueBinding;
 struct FRHIResourceInfo;
+struct FRHIUniformBufferLayoutInitializer;
 struct FGenerateMipsStruct;
 enum class EClearBinding;
 
@@ -759,13 +762,11 @@ public:
 /** Data structure to store information about resource parameter in a shader parameter structure. */
 struct FRHIUniformBufferResource
 {
-	DECLARE_EXPORTED_TYPE_LAYOUT(FRHIUniformBufferResource, RHI_API, NonVirtual);
-
 	/** Byte offset to each resource in the uniform buffer memory. */
-	LAYOUT_FIELD(uint16, MemberOffset);
+	uint16 MemberOffset;
 
 	/** Type of the member that allow (). */
-	LAYOUT_FIELD(EUniformBufferBaseType, MemberType);
+	EUniformBufferBaseType MemberType;
 
 	/** Compare two uniform buffer layout resources. */
 	friend inline bool operator==(const FRHIUniformBufferResource& A, const FRHIUniformBufferResource& B)
@@ -775,198 +776,12 @@ struct FRHIUniformBufferResource
 	}
 };
 
-inline FArchive& operator<<(FArchive& Ar, FRHIUniformBufferResource& Ref)
-{
-	uint8 Type = (uint8)Ref.MemberType;
-	Ar << Ref.MemberOffset;
-	Ar << Type;
-	Ref.MemberType = (EUniformBufferBaseType)Type;
-	return Ar;
-}
-
-inline constexpr uint16 kUniformBufferInvalidOffset = TNumericLimits<uint16>::Max();
-
-/** Initializer for the layout of a uniform buffer in memory. */
-struct FRHIUniformBufferLayoutInitializer
-{
-	DECLARE_EXPORTED_TYPE_LAYOUT(FRHIUniformBufferLayoutInitializer, RHI_API, NonVirtual);
-
-	FRHIUniformBufferLayoutInitializer() = default;
-
-	explicit FRHIUniformBufferLayoutInitializer(const TCHAR * InName)
-		: Name(InName)
-	{}
-	explicit FRHIUniformBufferLayoutInitializer(const TCHAR * InName, uint32 InConstantBufferSize)
-		: Name(InName)
-		, ConstantBufferSize(InConstantBufferSize)
-	{
-		ComputeHash();
-	}
-
-	inline uint32 GetHash() const
-	{
-		checkSlow(Hash != 0);
-		return Hash;
-	}
-
-	void ComputeHash()
-	{
-		// Static slot is not stable. Just track whether we have one at all.
-		uint32 TmpHash = ConstantBufferSize << 16 | static_cast<uint32>(BindingFlags) << 8 | static_cast<uint32>(StaticSlot != MAX_UNIFORM_BUFFER_STATIC_SLOTS);
-
-		for (int32 ResourceIndex = 0; ResourceIndex < Resources.Num(); ResourceIndex++)
-		{
-			// Offset and therefore hash must be the same regardless of pointer size
-			checkSlow(Resources[ResourceIndex].MemberOffset == Align(Resources[ResourceIndex].MemberOffset, SHADER_PARAMETER_POINTER_ALIGNMENT));
-			TmpHash ^= Resources[ResourceIndex].MemberOffset;
-		}
-
-		uint32 N = Resources.Num();
-		while (N >= 4)
-		{
-			TmpHash ^= (Resources[--N].MemberType << 0);
-			TmpHash ^= (Resources[--N].MemberType << 8);
-			TmpHash ^= (Resources[--N].MemberType << 16);
-			TmpHash ^= (Resources[--N].MemberType << 24);
-		}
-		while (N >= 2)
-		{
-			TmpHash ^= Resources[--N].MemberType << 0;
-			TmpHash ^= Resources[--N].MemberType << 16;
-		}
-		while (N > 0)
-		{
-			TmpHash ^= Resources[--N].MemberType;
-		}
-		Hash = TmpHash;
-	}
-
-	void CopyFrom(const FRHIUniformBufferLayoutInitializer& Source)
-	{
-		ConstantBufferSize = Source.ConstantBufferSize;
-		StaticSlot = Source.StaticSlot;
-		BindingFlags = Source.BindingFlags;
-		Resources = Source.Resources;
-		Name = Source.Name;
-		Hash = Source.Hash;
-	}
-
-	const FMemoryImageString& GetDebugName() const
-	{
-		return Name;
-	}
-
-	bool HasRenderTargets() const
-	{
-		return RenderTargetsOffset != kUniformBufferInvalidOffset;
-	}
-
-	bool HasExternalOutputs() const
-	{
-		return bHasNonGraphOutputs;
-	}
-
-	bool HasStaticSlot() const
-	{
-		return IsUniformBufferStaticSlotValid(StaticSlot);
-	}
-
-	friend FArchive& operator<<(FArchive& Ar, FRHIUniformBufferLayoutInitializer& Ref)
-	{
-		Ar << Ref.ConstantBufferSize;
-		Ar << Ref.StaticSlot;
-		Ar << Ref.RenderTargetsOffset;
-		Ar << Ref.bHasNonGraphOutputs;
-		Ar << Ref.BindingFlags;
-		Ar << Ref.Resources;
-		Ar << Ref.GraphResources;
-		Ar << Ref.GraphTextures;
-		Ar << Ref.GraphBuffers;
-		Ar << Ref.GraphUniformBuffers;
-		Ar << Ref.UniformBuffers;
-		Ar << Ref.Name;
-		Ar << Ref.Hash;
-		Ar << Ref.bNoEmulatedUniformBuffer;
-		return Ar;
-	}
-
-private:
-	// for debugging / error message
-	LAYOUT_FIELD(FMemoryImageString, Name);
-
-public:
-	/** The list of all resource inlined into the shader parameter structure. */
-	LAYOUT_FIELD(TMemoryImageArray<FRHIUniformBufferResource>, Resources);
-
-	/** The list of all RDG resource references inlined into the shader parameter structure. */
-	LAYOUT_FIELD(TMemoryImageArray<FRHIUniformBufferResource>, GraphResources);
-
-	/** The list of all RDG texture references inlined into the shader parameter structure. */
-	LAYOUT_FIELD(TMemoryImageArray<FRHIUniformBufferResource>, GraphTextures);
-
-	/** The list of all RDG buffer references inlined into the shader parameter structure. */
-	LAYOUT_FIELD(TMemoryImageArray<FRHIUniformBufferResource>, GraphBuffers);
-
-	/** The list of all RDG uniform buffer references inlined into the shader parameter structure. */
-	LAYOUT_FIELD(TMemoryImageArray<FRHIUniformBufferResource>, GraphUniformBuffers);
-
-	/** The list of all non-RDG uniform buffer references inlined into the shader parameter structure. */
-	LAYOUT_FIELD(TMemoryImageArray<FRHIUniformBufferResource>, UniformBuffers);
-
-private:
-	LAYOUT_FIELD_INITIALIZED(uint32, Hash, 0);
-
-public:
-	/** The size of the constant buffer in bytes. */
-	LAYOUT_FIELD_INITIALIZED(uint32, ConstantBufferSize, 0);
-
-	/** The render target binding slots offset, if it exists. */
-	LAYOUT_FIELD_INITIALIZED(uint16, RenderTargetsOffset, kUniformBufferInvalidOffset);
-
-	/** The static slot (if applicable). */
-	LAYOUT_FIELD_INITIALIZED(FUniformBufferStaticSlot, StaticSlot, MAX_UNIFORM_BUFFER_STATIC_SLOTS);
-
-	/** The binding flags describing how this resource can be bound to the RHI. */
-	LAYOUT_FIELD_INITIALIZED(EUniformBufferBindingFlags, BindingFlags, EUniformBufferBindingFlags::Shader);
-
-	/** Whether this layout may contain non-render-graph outputs (e.g. RHI UAVs). */
-	LAYOUT_FIELD_INITIALIZED(bool, bHasNonGraphOutputs, false);
-
-	/** Used for platforms which use emulated ub's, forces a real uniform buffer instead */
-	LAYOUT_FIELD_INITIALIZED(bool, bNoEmulatedUniformBuffer, false);
-
-	/** Compare two uniform buffer layout initializers. */
-	friend inline bool operator==(const FRHIUniformBufferLayoutInitializer& A, const FRHIUniformBufferLayoutInitializer& B)
-	{
-		return A.ConstantBufferSize == B.ConstantBufferSize
-			&& A.StaticSlot == B.StaticSlot
-			&& A.BindingFlags == B.BindingFlags
-			&& A.Resources == B.Resources;
-	}
-};
-
 /** The layout of a uniform buffer in memory. */
 struct FRHIUniformBufferLayout : public FRHIResource
 {
 	FRHIUniformBufferLayout() = delete;
 
-	explicit FRHIUniformBufferLayout(const FRHIUniformBufferLayoutInitializer& Initializer)
-		: FRHIResource(RRT_UniformBufferLayout)
-		, Name(Initializer.GetDebugName())
-		, Resources(Initializer.Resources)
-		, GraphResources(Initializer.GraphResources)
-		, GraphTextures(Initializer.GraphTextures)
-		, GraphBuffers(Initializer.GraphBuffers)
-		, GraphUniformBuffers(Initializer.GraphUniformBuffers)
-		, UniformBuffers(Initializer.UniformBuffers)
-		, Hash(Initializer.GetHash())
-		, ConstantBufferSize(Initializer.ConstantBufferSize)
-		, RenderTargetsOffset(Initializer.RenderTargetsOffset)
-		, StaticSlot(Initializer.StaticSlot)
-		, BindingFlags(Initializer.BindingFlags)
-		, bHasNonGraphOutputs(Initializer.bHasNonGraphOutputs)
-		, bNoEmulatedUniformBuffer(Initializer.bNoEmulatedUniformBuffer)
-	{}
+	RHI_API explicit FRHIUniformBufferLayout(const FRHIUniformBufferLayoutInitializer& Initializer);
 
 	inline const FString& GetDebugName() const
 	{
@@ -981,7 +796,7 @@ struct FRHIUniformBufferLayout : public FRHIResource
 
 	inline bool HasRenderTargets() const
 	{
-		return RenderTargetsOffset != kUniformBufferInvalidOffset;
+		return RenderTargetsOffset != TNumericLimits<uint16>::Max();
 	}
 
 	inline bool HasExternalOutputs() const
@@ -2409,71 +2224,68 @@ DECLARE_INTRINSIC_TYPE_LAYOUT(ERayTracingGeometryInitializerType);
 
 struct FRayTracingGeometrySegment
 {
-	DECLARE_TYPE_LAYOUT(FRayTracingGeometrySegment, NonVirtual);
-public:
-	LAYOUT_FIELD_INITIALIZED(FBufferRHIRef, VertexBuffer, nullptr);
-	LAYOUT_FIELD_INITIALIZED(EVertexElementType, VertexBufferElementType, VET_Float3);
+	FBufferRHIRef VertexBuffer = nullptr;
+	EVertexElementType VertexBufferElementType = VET_Float3;
 
 	// Offset in bytes from the base address of the vertex buffer.
-	LAYOUT_FIELD_INITIALIZED(uint32, VertexBufferOffset, 0);
+	uint32 VertexBufferOffset = 0;
 
 	// Number of bytes between elements of the vertex buffer (sizeof VET_Float3 by default).
 	// Must be equal or greater than the size of the position vector.
-	LAYOUT_FIELD_INITIALIZED(uint32, VertexBufferStride, 12);
+	uint32 VertexBufferStride = 12;
 
 	// Number of vertices (positions) in VertexBuffer.
 	// If an index buffer is present, this must be at least the maximum index value in the index buffer + 1.
-	LAYOUT_FIELD_INITIALIZED(uint32, MaxVertices, 0);
+	uint32 MaxVertices = 0;
 
 	// Primitive range for this segment.
-	LAYOUT_FIELD_INITIALIZED(uint32, FirstPrimitive, 0);
-	LAYOUT_FIELD_INITIALIZED(uint32, NumPrimitives, 0);
+	uint32 FirstPrimitive = 0;
+	uint32 NumPrimitives = 0;
 
 	// Indicates whether any-hit shader could be invoked when hitting this geometry segment.
 	// Setting this to `false` turns off any-hit shaders, making the section "opaque" and improving ray tracing performance.
-	LAYOUT_FIELD_INITIALIZED(bool, bForceOpaque, false);
+	bool bForceOpaque = false;
 
 	// Any-hit shader may be invoked multiple times for the same primitive during ray traversal.
 	// Setting this to `false` guarantees that only a single instance of any-hit shader will run per primitive, at some performance cost.
-	LAYOUT_FIELD_INITIALIZED(bool, bAllowDuplicateAnyHitShaderInvocation, true);
+	bool bAllowDuplicateAnyHitShaderInvocation = true;
 
 	// Indicates whether this section is enabled and should be taken into account during acceleration structure creation
-	LAYOUT_FIELD_INITIALIZED(bool, bEnabled, true);
+	bool bEnabled = true;
 };
 
 struct FRayTracingGeometryInitializer
 {
-	DECLARE_EXPORTED_TYPE_LAYOUT(FRayTracingGeometryInitializer, RHI_API, NonVirtual);
-public:
-	LAYOUT_FIELD_INITIALIZED(FBufferRHIRef, IndexBuffer, nullptr);
+	FBufferRHIRef IndexBuffer = nullptr;
 
 	// Offset in bytes from the base address of the index buffer.
-	LAYOUT_FIELD_INITIALIZED(uint32, IndexBufferOffset, 0);
+	uint32 IndexBufferOffset = 0;
 
-	LAYOUT_FIELD_INITIALIZED(ERayTracingGeometryType, GeometryType, RTGT_Triangles);
+	ERayTracingGeometryType GeometryType = RTGT_Triangles;
 
 	// Total number of primitives in all segments of the geometry. Only used for validation.
-	LAYOUT_FIELD_INITIALIZED(uint32, TotalPrimitiveCount, 0);
+	uint32 TotalPrimitiveCount = 0;
 
 	// Partitions of geometry to allow different shader and resource bindings.
 	// All ray tracing geometries must have at least one segment.
-	LAYOUT_FIELD(TMemoryImageArray<FRayTracingGeometrySegment>, Segments);
+	TMemoryImageArray<FRayTracingGeometrySegment> Segments;
 
 	// Offline built geometry data. If null, the geometry will be built by the RHI at runtime.
-	LAYOUT_FIELD_INITIALIZED(FResourceArrayInterface*, OfflineData, nullptr);
+	FResourceArrayInterface* OfflineData = nullptr;
 
 	// Pointer to an existing ray tracing geometry which the new geometry is built from.
-	LAYOUT_FIELD_INITIALIZED(FRHIRayTracingGeometry*, SourceGeometry, nullptr);
+	FRHIRayTracingGeometry* SourceGeometry = nullptr;
 
-	LAYOUT_FIELD_INITIALIZED(bool, bFastBuild, false);
-	LAYOUT_FIELD_INITIALIZED(bool, bAllowUpdate, false);
-	LAYOUT_FIELD_INITIALIZED(bool, bAllowCompaction, true);
-	LAYOUT_FIELD_INITIALIZED(ERayTracingGeometryInitializerType, Type, ERayTracingGeometryInitializerType::Rendering);
+	bool bFastBuild = false;
+	bool bAllowUpdate = false;
+	bool bAllowCompaction = true;
+	ERayTracingGeometryInitializerType Type = ERayTracingGeometryInitializerType::Rendering;
 
 	// Use FDebugName for auto-generated debug names with numbered suffixes, it is a variation of FMemoryImageName with optional number postfix.
-	LAYOUT_FIELD(FDebugName, DebugName);
+	FDebugName DebugName;
+
 	// Store the path name of the owner object for resource tracking. FMemoryImageName allows a conversion to/from FName.
-	LAYOUT_FIELD(FMemoryImageName, OwnerName);
+	FMemoryImageName OwnerName;
 };
 
 enum ERayTracingSceneLifetime
@@ -2997,12 +2809,6 @@ template<> struct TRHIShaderToEnum<FPixelShaderRHIRef>         { enum { ShaderFr
 template<> struct TRHIShaderToEnum<FGeometryShaderRHIRef>      { enum { ShaderFrequency = SF_Geometry      }; };
 template<> struct TRHIShaderToEnum<FComputeShaderRHIRef>       { enum { ShaderFrequency = SF_Compute       }; };
 
-template<typename TRHIShaderType>
-inline const TCHAR* GetShaderFrequencyString(bool bIncludePrefix = true)
-{
-	return GetShaderFrequencyString(static_cast<EShaderFrequency>(TRHIShaderToEnum<TRHIShaderType>::ShaderFrequency), bIncludePrefix);
-}
-
 struct FBoundShaderStateInput
 {
 	inline FBoundShaderStateInput() {}
diff --git a/Engine/Source/Runtime/RHI/Public/RHIStrings.h b/Engine/Source/Runtime/RHI/Public/RHIStrings.h
new file mode 100644
--- /dev/null
+++ b/Engine/Source/Runtime/RHI/Public/RHIStrings.h
@@ -0,0 +1,227 @@
+// Copyright Epic Games, Inc. All Rights Reserved.
+
+#pragma once
+
+#include "RHIDefinitions.h"
+#include "Misc/AssertionMacros.h"
+
+inline const TCHAR* GetShaderFrequencyString(EShaderFrequency Frequency, bool bIncludePrefix = true)
+{
+	const TCHAR* String = TEXT("SF_NumFrequencies");
+	switch (Frequency)
+	{
+	case SF_Vertex:			String = TEXT("SF_Vertex"); break;
+	case SF_Mesh:			String = TEXT("SF_Mesh"); break;
+	case SF_Amplification:	String = TEXT("SF_Amplification"); break;
+	case SF_Geometry:		String = TEXT("SF_Geometry"); break;
+	case SF_Pixel:			String = TEXT("SF_Pixel"); break;
+	case SF_Compute:		String = TEXT("SF_Compute"); break;
+	case SF_RayGen:			String = TEXT("SF_RayGen"); break;
+	case SF_RayMiss:		String = TEXT("SF_RayMiss"); break;
+	case SF_RayHitGroup:	String = TEXT("SF_RayHitGroup"); break;
+	case SF_RayCallable:	String = TEXT("SF_RayCallable"); break;
+
+	default:
+		checkf(0, TEXT("Unknown ShaderFrequency %d"), (int32)Frequency);
+		break;
+	}
+
+	// Skip SF_
+	int32 Index = bIncludePrefix ? 0 : 3;
+	String += Index;
+	return String;
+}
+
+template<typename TRHIShaderType>
+inline const TCHAR* GetShaderFrequencyString(bool bIncludePrefix = true)
+{
+	return GetShaderFrequencyString(static_cast<EShaderFrequency>(TRHIShaderToEnum<TRHIShaderType>::ShaderFrequency), bIncludePrefix);
+}
+
+inline const TCHAR* GetTextureDimensionString(ETextureDimension Dimension)
+{
+	switch (Dimension)
+	{
+	case ETextureDimension::Texture2D:
+		return TEXT("Texture2D");
+	case ETextureDimension::Texture2DArray:
+		return TEXT("Texture2DArray");
+	case ETextureDimension::Texture3D:
+		return TEXT("Texture3D");
+	case ETextureDimension::TextureCube:
+		return TEXT("TextureCube");
+	case ETextureDimension::TextureCubeArray:
+		return TEXT("TextureCubeArray");
+	}
+	return TEXT("");
+}
+
+inline const TCHAR* GetTextureCreateFlagString(ETextureCreateFlags TextureCreateFlag)
+{
+	switch (TextureCreateFlag)
+	{
+	case ETextureCreateFlags::None:
+		return TEXT("None");
+	case ETextureCreateFlags::RenderTargetable:
+		return TEXT("RenderTargetable");
+	case ETextureCreateFlags::ResolveTargetable:
+		return TEXT("ResolveTargetable");
+	case ETextureCreateFlags::DepthStencilTargetable:
+		return TEXT("DepthStencilTargetable");
+	case ETextureCreateFlags::ShaderResource:
+		return TEXT("ShaderResource");
+	case ETextureCreateFlags::SRGB:
+		return TEXT("SRGB");
+	case ETextureCreateFlags::CPUWritable:
+		return TEXT("CPUWritable");
+	case ETextureCreateFlags::NoTiling:
+		return TEXT("NoTiling");
+	case ETextureCreateFlags::VideoDecode:
+		return TEXT("VideoDecode");
+	case ETextureCreateFlags::Dynamic:
+		return TEXT("Dynamic");
+	case ETextureCreateFlags::InputAttachmentRead:
+		return TEXT("InputAttachmentRead");
+	case ETextureCreateFlags::Foveation:
+		return TEXT("Foveation");
+	case ETextureCreateFlags::Tiling3D:
+		return TEXT("Tiling3D");
+	case ETextureCreateFlags::Memoryless:
+		return TEXT("Memoryless");
+	case ETextureCreateFlags::GenerateMipCapable:
+		return TEXT("GenerateMipCapable");
+	case ETextureCreateFlags::FastVRAMPartialAlloc:
+		return TEXT("FastVRAMPartialAlloc");
+	case ETextureCreateFlags::DisableSRVCreation:
+		return TEXT("DisableSRVCreation");
+	case ETextureCreateFlags::DisableDCC:
+		return TEXT("DisableDCC");
+	case ETextureCreateFlags::UAV:
+		return TEXT("UAV");
+	case ETextureCreateFlags::Presentable:
+		return TEXT("Presentable");
+	case ETextureCreateFlags::CPUReadback:
+		return TEXT("CPUReadback");
+	case ETextureCreateFlags::OfflineProcessed:
+		return TEXT("OfflineProcessed");
+	case ETextureCreateFlags::FastVRAM:
+		return TEXT("FastVRAM");
+	case ETextureCreateFlags::HideInVisualizeTexture:
+		return TEXT("HideInVisualizeTexture");
+	case ETextureCreateFlags::Virtual:
+		return TEXT("Virtual");
+	case ETextureCreateFlags::TargetArraySlicesIndependently:
+		return TEXT("TargetArraySlicesIndependently");
+	case ETextureCreateFlags::Shared:
+		return TEXT("Shared");
+	case ETextureCreateFlags::NoFastClear:
+		return TEXT("NoFastClear");
+	case ETextureCreateFlags::DepthStencilResolveTarget:
+		return TEXT("DepthStencilResolveTarget");
+	case ETextureCreateFlags::Streamable:
+		return TEXT("Streamable");
+	case ETextureCreateFlags::NoFastClearFinalize:
+		return TEXT("NoFastClearFinalize");
+	case ETextureCreateFlags::AFRManual:
+		return TEXT("AFRManual");
+	case ETextureCreateFlags::ReduceMemoryWithTilingMode:
+		return TEXT("ReduceMemoryWithTilingMode");
+	}
+	return TEXT("");
+}
+
+inline const TCHAR* GetBufferUsageFlagString(EBufferUsageFlags BufferUsage)
+{
+	switch (BufferUsage)
+	{
+	case EBufferUsageFlags::None:
+		return TEXT("None");
+	case EBufferUsageFlags::Static:
+		return TEXT("Static");
+	case EBufferUsageFlags::Dynamic:
+		return TEXT("Dynamic");
+	case EBufferUsageFlags::Volatile:
+		return TEXT("Volatile");
+	case EBufferUsageFlags::UnorderedAccess:
+		return TEXT("UnorderedAccess");
+	case EBufferUsageFlags::ByteAddressBuffer:
+		return TEXT("ByteAddressBuffer");
+	case EBufferUsageFlags::SourceCopy:
+		return TEXT("SourceCopy");
+	case EBufferUsageFlags::StreamOutput:
+		return TEXT("StreamOutput");
+	case EBufferUsageFlags::DrawIndirect:
+		return TEXT("DrawIndirect");
+	case EBufferUsageFlags::ShaderResource:
+		return TEXT("ShaderResource");
+	case EBufferUsageFlags::KeepCPUAccessible:
+		return TEXT("KeepCPUAccessible");
+	case EBufferUsageFlags::FastVRAM:
+		return TEXT("FastVRAM");
+	case EBufferUsageFlags::Shared:
+		return TEXT("Shared");
+	case EBufferUsageFlags::AccelerationStructure:
+		return TEXT("AccelerationStructure");
+	case EBufferUsageFlags::VertexBuffer:
+		return TEXT("VertexBuffer");
+	case EBufferUsageFlags::IndexBuffer:
+		return TEXT("IndexBuffer");
+	case EBufferUsageFlags::StructuredBuffer:
+		return TEXT("StructuredBuffer");
+	}
+	return TEXT("");
+}
+
+inline const TCHAR* GetUniformBufferBaseTypeString(EUniformBufferBaseType BaseType)
+{
+	switch (BaseType)
+	{
+	case UBMT_INVALID:
+		return TEXT("UBMT_INVALID");
+	case UBMT_BOOL:
+		return TEXT("UBMT_BOOL");
+	case UBMT_INT32:
+		return TEXT("UBMT_INT32");
+	case UBMT_UINT32:
+		return TEXT("UBMT_UINT32");
+	case UBMT_FLOAT32:
+		return TEXT("UBMT_FLOAT32");
+	case UBMT_TEXTURE:
+		return TEXT("UBMT_TEXTURE");
+	case UBMT_SRV:
+		return TEXT("UBMT_SRV");
+	case UBMT_UAV:
+		return TEXT("UBMT_UAV");
+	case UBMT_SAMPLER:
+		return TEXT("UBMT_SAMPLER");
+	case UBMT_RDG_TEXTURE:
+		return TEXT("UBMT_RDG_TEXTURE");
+	case UBMT_RDG_TEXTURE_ACCESS:
+		return TEXT("UBMT_RDG_TEXTURE_ACCESS");
+	case UBMT_RDG_TEXTURE_ACCESS_ARRAY:
+		return TEXT("UBMT_RDG_TEXTURE_ACCESS_ARRAY");
+	case UBMT_RDG_TEXTURE_SRV:
+		return TEXT("UBMT_RDG_TEXTURE_SRV");
+	case UBMT_RDG_TEXTURE_UAV:
+		return TEXT("UBMT_RDG_TEXTURE_UAV");
+	case UBMT_RDG_BUFFER_ACCESS:
+		return TEXT("UBMT_RDG_BUFFER_ACCESS");
+	case UBMT_RDG_BUFFER_ACCESS_ARRAY:
+		return TEXT("UBMT_RDG_BUFFER_ACCESS_ARRAY");
+	case UBMT_RDG_BUFFER_SRV:
+		return TEXT("UBMT_RDG_BUFFER_SRV");
+	case UBMT_RDG_BUFFER_UAV:
+		return TEXT("UBMT_RDG_BUFFER_UAV");
+	case UBMT_RDG_UNIFORM_BUFFER:
+		return TEXT("UBMT_RDG_UNIFORM_BUFFER");
+	case UBMT_NESTED_STRUCT:
+		return TEXT("UBMT_NESTED_STRUCT");
+	case UBMT_INCLUDED_STRUCT:
+		return TEXT("UBMT_INCLUDED_STRUCT");
+	case UBMT_REFERENCED_STRUCT:
+		return TEXT("UBMT_REFERENCED_STRUCT");
+	case UBMT_RENDER_TARGET_BINDING_SLOTS:
+		return TEXT("UBMT_RENDER_TARGET_BINDING_SLOTS");
+	}
+	return TEXT("");
+}
diff --git a/Engine/Source/Runtime/RHI/Public/RHIUniformBufferLayoutInitializer.h b/Engine/Source/Runtime/RHI/Public/RHIUniformBufferLayoutInitializer.h
new file mode 100644
--- /dev/null
+++ b/Engine/Source/Runtime/RHI/Public/RHIUniformBufferLayoutInitializer.h
@@ -0,0 +1,195 @@
+// Copyright Epic Games, Inc. All Rights Reserved.
+
+#pragma once
+
+#include "RHIMemoryLayout.h"
+
+/** Data structure to store information about resource parameter in a shader parameter structure. */
+struct FRHIUniformBufferResourceInitializer
+{
+	DECLARE_EXPORTED_TYPE_LAYOUT(FRHIUniformBufferResourceInitializer, RHI_API, NonVirtual);
+
+	/** Byte offset to each resource in the uniform buffer memory. */
+	LAYOUT_FIELD(uint16, MemberOffset);
+
+	/** Type of the member that allow (). */
+	LAYOUT_FIELD(EUniformBufferBaseType, MemberType);
+
+	/** Compare two uniform buffer layout resources. */
+	friend inline bool operator==(const FRHIUniformBufferResourceInitializer& A, const FRHIUniformBufferResourceInitializer& B)
+	{
+		return A.MemberOffset == B.MemberOffset
+			&& A.MemberType == B.MemberType;
+	}
+};
+
+inline FArchive& operator<<(FArchive& Ar, FRHIUniformBufferResourceInitializer& Ref)
+{
+	uint8 Type = (uint8)Ref.MemberType;
+	Ar << Ref.MemberOffset;
+	Ar << Type;
+	Ref.MemberType = (EUniformBufferBaseType)Type;
+	return Ar;
+}
+
+inline constexpr uint16 kUniformBufferInvalidOffset = TNumericLimits<uint16>::Max();
+
+/** Initializer for the layout of a uniform buffer in memory. */
+struct FRHIUniformBufferLayoutInitializer
+{
+	DECLARE_EXPORTED_TYPE_LAYOUT(FRHIUniformBufferLayoutInitializer, RHI_API, NonVirtual);
+
+	FRHIUniformBufferLayoutInitializer() = default;
+
+	explicit FRHIUniformBufferLayoutInitializer(const TCHAR* InName)
+		: Name(InName)
+	{}
+	explicit FRHIUniformBufferLayoutInitializer(const TCHAR* InName, uint32 InConstantBufferSize)
+		: Name(InName)
+		, ConstantBufferSize(InConstantBufferSize)
+	{
+		ComputeHash();
+	}
+
+	inline uint32 GetHash() const
+	{
+		checkSlow(Hash != 0);
+		return Hash;
+	}
+
+	void ComputeHash()
+	{
+		// Static slot is not stable. Just track whether we have one at all.
+		uint32 TmpHash = ConstantBufferSize << 16 | static_cast<uint32>(BindingFlags) << 8 | static_cast<uint32>(StaticSlot != MAX_UNIFORM_BUFFER_STATIC_SLOTS);
+
+		for (int32 ResourceIndex = 0; ResourceIndex < Resources.Num(); ResourceIndex++)
+		{
+			// Offset and therefore hash must be the same regardless of pointer size
+			checkSlow(Resources[ResourceIndex].MemberOffset == Align(Resources[ResourceIndex].MemberOffset, SHADER_PARAMETER_POINTER_ALIGNMENT));
+			TmpHash ^= Resources[ResourceIndex].MemberOffset;
+		}
+
+		uint32 N = Resources.Num();
+		while (N >= 4)
+		{
+			TmpHash ^= (Resources[--N].MemberType << 0);
+			TmpHash ^= (Resources[--N].MemberType << 8);
+			TmpHash ^= (Resources[--N].MemberType << 16);
+			TmpHash ^= (Resources[--N].MemberType << 24);
+		}
+		while (N >= 2)
+		{
+			TmpHash ^= Resources[--N].MemberType << 0;
+			TmpHash ^= Resources[--N].MemberType << 16;
+		}
+		while (N > 0)
+		{
+			TmpHash ^= Resources[--N].MemberType;
+		}
+		Hash = TmpHash;
+	}
+
+	void CopyFrom(const FRHIUniformBufferLayoutInitializer& Source)
+	{
+		ConstantBufferSize = Source.ConstantBufferSize;
+		StaticSlot = Source.StaticSlot;
+		BindingFlags = Source.BindingFlags;
+		Resources = Source.Resources;
+		Name = Source.Name;
+		Hash = Source.Hash;
+	}
+
+	const FMemoryImageString& GetDebugName() const
+	{
+		return Name;
+	}
+
+	bool HasRenderTargets() const
+	{
+		return RenderTargetsOffset != kUniformBufferInvalidOffset;
+	}
+
+	bool HasExternalOutputs() const
+	{
+		return bHasNonGraphOutputs;
+	}
+
+	bool HasStaticSlot() const
+	{
+		return IsUniformBufferStaticSlotValid(StaticSlot);
+	}
+
+	friend FArchive& operator<<(FArchive& Ar, FRHIUniformBufferLayoutInitializer& Ref)
+	{
+		Ar << Ref.ConstantBufferSize;
+		Ar << Ref.StaticSlot;
+		Ar << Ref.RenderTargetsOffset;
+		Ar << Ref.bHasNonGraphOutputs;
+		Ar << Ref.BindingFlags;
+		Ar << Ref.Resources;
+		Ar << Ref.GraphResources;
+		Ar << Ref.GraphTextures;
+		Ar << Ref.GraphBuffers;
+		Ar << Ref.GraphUniformBuffers;
+		Ar << Ref.UniformBuffers;
+		Ar << Ref.Name;
+		Ar << Ref.Hash;
+		Ar << Ref.bNoEmulatedUniformBuffer;
+		return Ar;
+	}
+
+private:
+	// for debugging / error message
+	LAYOUT_FIELD(FMemoryImageString, Name);
+
+public:
+	/** The list of all resource inlined into the shader parameter structure. */
+	LAYOUT_FIELD(TMemoryImageArray<FRHIUniformBufferResourceInitializer>, Resources);
+
+	/** The list of all RDG resource references inlined into the shader parameter structure. */
+	LAYOUT_FIELD(TMemoryImageArray<FRHIUniformBufferResourceInitializer>, GraphResources);
+
+	/** The list of all RDG texture references inlined into the shader parameter structure. */
+	LAYOUT_FIELD(TMemoryImageArray<FRHIUniformBufferResourceInitializer>, GraphTextures);
+
+	/** The list of all RDG buffer references inlined into the shader parameter structure. */
+	LAYOUT_FIELD(TMemoryImageArray<FRHIUniformBufferResourceInitializer>, GraphBuffers);
+
+	/** The list of all RDG uniform buffer references inlined into the shader parameter structure. */
+	LAYOUT_FIELD(TMemoryImageArray<FRHIUniformBufferResourceInitializer>, GraphUniformBuffers);
+
+	/** The list of all non-RDG uniform buffer references inlined into the shader parameter structure. */
+	LAYOUT_FIELD(TMemoryImageArray<FRHIUniformBufferResourceInitializer>, UniformBuffers);
+
+private:
+	LAYOUT_FIELD_INITIALIZED(uint32, Hash, 0);
+
+public:
+	/** The size of the constant buffer in bytes. */
+	LAYOUT_FIELD_INITIALIZED(uint32, ConstantBufferSize, 0);
+
+	/** The render target binding slots offset, if it exists. */
+	LAYOUT_FIELD_INITIALIZED(uint16, RenderTargetsOffset, kUniformBufferInvalidOffset);
+
+	/** The static slot (if applicable). */
+	LAYOUT_FIELD_INITIALIZED(FUniformBufferStaticSlot, StaticSlot, MAX_UNIFORM_BUFFER_STATIC_SLOTS);
+
+	/** The binding flags describing how this resource can be bound to the RHI. */
+	LAYOUT_FIELD_INITIALIZED(EUniformBufferBindingFlags, BindingFlags, EUniformBufferBindingFlags::Shader);
+
+	/** Whether this layout may contain non-render-graph outputs (e.g. RHI UAVs). */
+	LAYOUT_FIELD_INITIALIZED(bool, bHasNonGraphOutputs, false);
+
+	/** Used for platforms which use emulated ub's, forces a real uniform buffer instead */
+	LAYOUT_FIELD_INITIALIZED(bool, bNoEmulatedUniformBuffer, false);
+
+	/** Compare two uniform buffer layout initializers. */
+	friend inline bool operator==(const FRHIUniformBufferLayoutInitializer& A, const FRHIUniformBufferLayoutInitializer& B)
+	{
+		return A.ConstantBufferSize == B.ConstantBufferSize
+			&& A.StaticSlot == B.StaticSlot
+			&& A.BindingFlags == B.BindingFlags
+			&& A.Resources == B.Resources;
+	}
+};
+
diff --git a/Engine/Source/Runtime/RenderCore/Private/DumpGPU.cpp b/Engine/Source/Runtime/RenderCore/Private/DumpGPU.cpp
--- a/Engine/Source/Runtime/RenderCore/Private/DumpGPU.cpp
+++ b/Engine/Source/Runtime/RenderCore/Private/DumpGPU.cpp
@@ -17,6 +17,7 @@
 #include "GenericPlatform/GenericPlatformDriver.h"
 #include "RenderGraphUtils.h"
 #include "GlobalShader.h"
+#include "RHIStrings.h"
 #include "RHIValidation.h"
 #include "RenderGraphPrivate.h"
 #include "RenderCore.h"
diff --git a/Engine/Source/Runtime/RenderCore/Private/Shader.cpp b/Engine/Source/Runtime/RenderCore/Private/Shader.cpp
--- a/Engine/Source/Runtime/RenderCore/Private/Shader.cpp
+++ b/Engine/Source/Runtime/RenderCore/Private/Shader.cpp
@@ -18,6 +18,7 @@
 #include "ShaderCore.h"
 #include "ShaderCompilerCore.h"
 #include "RenderUtils.h"
+#include "RHIStrings.h"
 #include "StereoRenderUtils.h"
 #include "Misc/ConfigCacheIni.h"
 #include "Misc/ScopeLock.h"
diff --git a/Engine/Source/Runtime/RenderCore/Private/ShaderParameterMetadata.cpp b/Engine/Source/Runtime/RenderCore/Private/ShaderParameterMetadata.cpp
--- a/Engine/Source/Runtime/RenderCore/Private/ShaderParameterMetadata.cpp
+++ b/Engine/Source/Runtime/RenderCore/Private/ShaderParameterMetadata.cpp
@@ -9,6 +9,7 @@
 #include "ShaderCore.h"
 #include "ShaderParameters.h"
 #include "DataDrivenShaderPlatformInfo.h"
+#include "RHIUniformBufferLayoutInitializer.h"
 
 bool SupportShaderPrecisionModifier(EShaderPlatform Platform)
 {
@@ -656,8 +657,8 @@ void FShaderParametersMetadata::InitializeLayout(FRHIUniformBufferLayoutInitiali
 			for (uint32 ArrayElementId = 0; ArrayElementId < (bIsArray ? ArraySize : 1u); ArrayElementId++)
 			{
 				const uint32 AbsoluteMemberOffset = CurrentMember.GetOffset() + MemberStack[i].StructOffset + ArrayElementId * SHADER_PARAMETER_POINTER_ALIGNMENT;
-				check(AbsoluteMemberOffset < (1u << (sizeof(FRHIUniformBufferResource::MemberOffset) * 8)));
-				const FRHIUniformBufferResource ResourceParameter{ uint16(AbsoluteMemberOffset), BaseType };
+				check(AbsoluteMemberOffset < (1u << (sizeof(FRHIUniformBufferResourceInitializer::MemberOffset) * 8)));
+				const FRHIUniformBufferResourceInitializer ResourceParameter{ uint16(AbsoluteMemberOffset), BaseType };
 
 				LayoutInitializer.Resources.Add(ResourceParameter);
 
@@ -705,15 +706,15 @@ void FShaderParametersMetadata::InitializeLayout(FRHIUniformBufferLayoutInitiali
 	}
 
 	const auto ByMemberOffset = [](
-		const FRHIUniformBufferResource& A,
-		const FRHIUniformBufferResource& B)
+		const FRHIUniformBufferResourceInitializer& A,
+		const FRHIUniformBufferResourceInitializer& B)
 	{
 		return A.MemberOffset < B.MemberOffset;
 	};
 
 	const auto ByTypeThenMemberOffset = [](
-		const FRHIUniformBufferResource& A,
-		const FRHIUniformBufferResource& B)
+		const FRHIUniformBufferResourceInitializer& A,
+		const FRHIUniformBufferResourceInitializer& B)
 	{
 		if (A.MemberType == B.MemberType)
 		{
diff --git a/Engine/Source/Runtime/RenderCore/Private/ShaderParameters.cpp b/Engine/Source/Runtime/RenderCore/Private/ShaderParameters.cpp
--- a/Engine/Source/Runtime/RenderCore/Private/ShaderParameters.cpp
+++ b/Engine/Source/Runtime/RenderCore/Private/ShaderParameters.cpp
@@ -7,6 +7,7 @@
 #include "ShaderParameters.h"
 #include "Containers/List.h"
 #include "UniformBuffer.h"
+#include "RHIStrings.h"
 #include "ShaderCore.h"
 #include "Shader.h"
 #include "VertexFactory.h"
diff --git a/Engine/Source/Runtime/RenderCore/Private/ShaderResource.cpp b/Engine/Source/Runtime/RenderCore/Private/ShaderResource.cpp
--- a/Engine/Source/Runtime/RenderCore/Private/ShaderResource.cpp
+++ b/Engine/Source/Runtime/RenderCore/Private/ShaderResource.cpp
@@ -26,6 +26,7 @@
 #include "ShaderCompilerCore.h"
 #include "Compression/OodleDataCompression.h"
 #include "RHIResources.h"	// Access to FRHIRayTracingShader::RayTracingPayloadType requires this
+#include "RHIStrings.h"
 #include "DataDrivenShaderPlatformInfo.h"
 
 #if WITH_EDITORONLY_DATA
diff --git a/Engine/Source/Runtime/RenderCore/Public/HDRHelper.h b/Engine/Source/Runtime/RenderCore/Public/HDRHelper.h
--- a/Engine/Source/Runtime/RenderCore/Public/HDRHelper.h
+++ b/Engine/Source/Runtime/RenderCore/Public/HDRHelper.h
@@ -31,6 +31,8 @@ struct FDisplayInformation
 	bool   bHDRSupported;
 };
 
+using FDisplayInformationArray = TArray<struct FDisplayInformation>;
+
 struct FACESTonemapParams
 {
 	// Default Values such as MinLuminanceLog10=-4, MidLuminance=15, MaxLuminance=1000, SceneColorMultiplier=1.5
diff --git a/Engine/Source/Runtime/RenderCore/Public/Shader.h b/Engine/Source/Runtime/RenderCore/Public/Shader.h
--- a/Engine/Source/Runtime/RenderCore/Public/Shader.h
+++ b/Engine/Source/Runtime/RenderCore/Public/Shader.h
@@ -36,6 +36,7 @@
 #include "RenderingThread.h"
 #endif
 #include "RenderDeferredCleanup.h"
+#include "RHIMemoryLayout.h"
 #include "Serialization/Archive.h"
 #include "Serialization/ArchiveProxy.h"
 #include "Serialization/MemoryImage.h"
