Adding some missing dependencies to RHIStrings.h

diff --git a/Engine/Source/Runtime/RHI/Public/RHIStrings.h b/Engine/Source/Runtime/RHI/Public/RHIStrings.h
--- a/Engine/Source/Runtime/RHI/Public/RHIStrings.h
+++ b/Engine/Source/Runtime/RHI/Public/RHIStrings.h
@@ -2,7 +2,7 @@
 
 #pragma once
 
-#include "RHIDefinitions.h"
+#include "RHI.h"
 #include "Misc/AssertionMacros.h"
 
 inline const TCHAR* GetShaderFrequencyString(EShaderFrequency Frequency, bool bIncludePrefix = true)
