Dump residency state during rhi.DumpResourceInfo to csv file

#rnx

diff --git a/Engine/Source/Runtime/D3D12RHI/Private/D3D12Texture.cpp b/Engine/Source/Runtime/D3D12RHI/Private/D3D12Texture.cpp
--- a/Engine/Source/Runtime/D3D12RHI/Private/D3D12Texture.cpp
+++ b/Engine/Source/Runtime/D3D12RHI/Private/D3D12Texture.cpp
@@ -1334,6 +1334,10 @@ bool FD3D12Texture::GetResourceInfo(FRHIResourceInfo& OutResourceInfo) const
 	OutResourceInfo.Type = GetType();
 	OutResourceInfo.VRamAllocation.AllocationSize = ResourceLocation.GetSize();
 	OutResourceInfo.IsTransient = ResourceLocation.IsTransient();
+#if ENABLE_RESIDENCY_MANAGEMENT
+	OutResourceInfo.bResident = GetResource() && GetResource()->GetResidencyHandle().ResidencyStatus == D3DX12Residency::ManagedObject::RESIDENCY_STATUS::RESIDENT;
+#endif
+
 	return true;
 }
 #endif
diff --git a/Engine/Source/Runtime/D3D12RHI/Public/D3D12Resources.h b/Engine/Source/Runtime/D3D12RHI/Public/D3D12Resources.h
--- a/Engine/Source/Runtime/D3D12RHI/Public/D3D12Resources.h
+++ b/Engine/Source/Runtime/D3D12RHI/Public/D3D12Resources.h
@@ -873,7 +873,11 @@ public:
 		OutResourceInfo.Name = GetName();
 		OutResourceInfo.Type = GetType();
 		OutResourceInfo.VRamAllocation.AllocationSize = ResourceLocation.GetSize();
-		OutResourceInfo.IsTransient = this->ResourceLocation.IsTransient();
+		OutResourceInfo.IsTransient = ResourceLocation.IsTransient();
+#if ENABLE_RESIDENCY_MANAGEMENT
+		OutResourceInfo.bResident = GetResource() && GetResource()->GetResidencyHandle().ResidencyStatus == D3DX12Residency::ManagedObject::RESIDENCY_STATUS::RESIDENT;
+#endif
+		
 		return true;
 	}
 #endif
diff --git a/Engine/Source/Runtime/RHI/Private/RHI.cpp b/Engine/Source/Runtime/RHI/Private/RHI.cpp
--- a/Engine/Source/Runtime/RHI/Private/RHI.cpp
+++ b/Engine/Source/Runtime/RHI/Private/RHI.cpp
@@ -1115,12 +1115,12 @@ void RHIDumpResourceMemory(const FString& NameFilter, ERHIResourceType TypeFilte
 
 	if (bOutputToCSVFile)
 	{
-		const TCHAR* Header = TEXT("Name,Type,Size,MarkedForDelete,Transient,Streaming,RenderTarget,UAV,\"Raytracing Acceleration Structure\"\n");
+		const TCHAR* Header = TEXT("Name,Type,Size,Resident,MarkedForDelete,Transient,Streaming,RenderTarget,UAV,\"Raytracing Acceleration Structure\"\n");
 		CSVFile->Serialize(TCHAR_TO_ANSI(Header), FPlatformString::Strlen(Header));
 	}
 	else if (bUseCSVOutput)
 	{
-		BufferedOutput.CategorizedLogf(CategoryName, ELogVerbosity::Log, TEXT("Name,Type,Size,MarkedForDelete,Transient,Streaming,RenderTarget,UAV,\"Raytracing Acceleration Structure\""));
+		BufferedOutput.CategorizedLogf(CategoryName, ELogVerbosity::Log, TEXT("Name,Type,Size,Resident,MarkedForDelete,Transient,Streaming,RenderTarget,UAV,\"Raytracing Acceleration Structure\""));
 	}
 	else
 	{
@@ -1147,17 +1147,19 @@ void RHIDumpResourceMemory(const FString& NameFilter, ERHIResourceType TypeFilte
 			ResourceInfo.Name.ToString(ResourceNameBuffer);
 			const TCHAR* ResourceType = StringFromRHIResourceType(ResourceInfo.Type);
 			const int64 SizeInBytes = ResourceInfo.VRamAllocation.AllocationSize;
-
+			
+			bool bResident = ResourceInfo.bResident;
 			RHIInternal::FResourceFlags Flags = GetResourceFlagsInternal(Resources[Index]);
 
 			if (bSummaryOutput == false)
 			{
-				if (bOutputToCSVFile)
-				{					
-					const FString Row = FString::Printf(TEXT("%s,%s,%.9f,%s,%s,%s,%s,%s,%s\n"),
+				if (bOutputToCSVFile || bUseCSVOutput)
+				{		
+					const FString Row = FString::Printf(TEXT("%s,%s,%.9f,%s,%s,%s,%s,%s,%s,%s\n"),
 						ResourceNameBuffer,
 						ResourceType,
 						SizeInBytes / double(1 << 20),
+						bResident ? TEXT("Yes") : TEXT("No"),
 						Flags.bMarkedForDelete ? TEXT("Yes") : TEXT("No"),
 						Flags.bTransient ? TEXT("Yes") : TEXT("No"),
 						Flags.bStreaming ? TEXT("Yes") : TEXT("No"),
@@ -1165,20 +1167,14 @@ void RHIDumpResourceMemory(const FString& NameFilter, ERHIResourceType TypeFilte
 						Flags.bUAV ? TEXT("Yes") : TEXT("No"),
 						Flags.bRTAS ? TEXT("Yes") : TEXT("No"));
 
-					CSVFile->Serialize(TCHAR_TO_ANSI(*Row), Row.Len());
-				}
-				else if (bUseCSVOutput)
-				{
-					BufferedOutput.CategorizedLogf(CategoryName, ELogVerbosity::Log, TEXT("%s,%s,%.9f,%s,%s,%s,%s,%s,%s"),
-						ResourceNameBuffer,
-						ResourceType,
-						SizeInBytes / double(1 << 20),
-						Flags.bMarkedForDelete ? TEXT("Yes") : TEXT("No"),
-						Flags.bTransient ? TEXT("Yes") : TEXT("No"),
-						Flags.bStreaming ? TEXT("Yes") : TEXT("No"),
-						(Flags.bRT || Flags.bDS) ? TEXT("Yes") : TEXT("No"),
-						Flags.bUAV ? TEXT("Yes") : TEXT("No"),
-						Flags.bRTAS ? TEXT("Yes") : TEXT("No"));
+					if (bOutputToCSVFile)
+					{
+						CSVFile->Serialize(TCHAR_TO_ANSI(*Row), Row.Len());
+					}
+					else
+					{
+						BufferedOutput.CategorizedLogf(CategoryName, ELogVerbosity::Log, TEXT("%s"), *Row);
+					}
 				}
 				else
 				{
diff --git a/Engine/Source/Runtime/RHI/Public/RHI.h b/Engine/Source/Runtime/RHI/Public/RHI.h
--- a/Engine/Source/Runtime/RHI/Public/RHI.h
+++ b/Engine/Source/Runtime/RHI/Public/RHI.h
@@ -1356,6 +1356,7 @@ struct FRHIResourceInfo
 	FVRamAllocation VRamAllocation;
 	bool IsTransient{ false };
 	bool bValid{ true };
+	bool bResident{ true };
 };
 
 enum class EClearBinding
