When switching from Compute -> Graphics, make sure we flush the Compute parameters.

diff --git a/Engine/Source/Runtime/RHI/Private/RHICommandList.cpp b/Engine/Source/Runtime/RHI/Private/RHICommandList.cpp
--- a/Engine/Source/Runtime/RHI/Private/RHICommandList.cpp
+++ b/Engine/Source/Runtime/RHI/Private/RHICommandList.cpp
@@ -173,11 +173,9 @@ FRHIParameterBatcher::~FRHIParameterBatcher()
 
 void FRHIParameterBatcher::OnBoundShaderChanged(FRHICommandList& InCommandList, const FBoundShaderStateInput& InBoundShaderStateInput)
 {
-	for (int32 Index = 0; Index < SF_NumGraphicsFrequencies; Index++)
-	{
-		InCommandList.SetBatchedShaderParameters(GetBatchedGraphicsShader(Index), AllBatchedShaderParameters[Index]);
-		AllBatchedShaderParameters[Index].Reset();
-	}
+	// Make sure we flush all parameters before binding new shaders, even if we make a pipeline switch
+	PreDispatch(InCommandList);
+	PreDraw(InCommandList);
 
 	AllBatchedShaders[SF_Vertex]        = InBoundShaderStateInput.GetVertexShader();
 	AllBatchedShaders[SF_Mesh]          = InBoundShaderStateInput.GetMeshShader();
@@ -188,8 +186,9 @@ void FRHIParameterBatcher::OnBoundShaderChanged(FRHICommandList& InCommandList,
 
 void FRHIParameterBatcher::OnBoundShaderChanged(FRHIComputeCommandList& InCommandList, FRHIComputeShader* InBoundComputeShaderRHI)
 {
-	InCommandList.SetBatchedShaderParameters(GetBatchedComputeShader(), AllBatchedShaderParameters[SF_Compute]);
-	AllBatchedShaderParameters[SF_Compute].Reset();
+	// Make sure we flush all parameters before binding new shaders, even if we make a pipeline switch
+	PreDraw((FRHICommandList&)InCommandList);
+	PreDispatch(InCommandList);
 
 	AllBatchedShaders[SF_Compute] = InBoundComputeShaderRHI;
 }
