Fixes for crashes when when switching preview platforms on Mac Editor

diff --git a/Engine/Source/Editor/UnrealEd/Private/Settings/EditorPerProjectUserSettings.cpp b/Engine/Source/Editor/UnrealEd/Private/Settings/EditorPerProjectUserSettings.cpp
--- a/Engine/Source/Editor/UnrealEd/Private/Settings/EditorPerProjectUserSettings.cpp
+++ b/Engine/Source/Editor/UnrealEd/Private/Settings/EditorPerProjectUserSettings.cpp
@@ -45,7 +45,9 @@ void UEditorPerProjectUserSettings::PostInitProperties()
 	if (bPreviewFeatureLevelWasDefault || PreviewFeatureLevel > GMaxRHIFeatureLevel)
 	{
 		PreviewFeatureLevel = GMaxRHIFeatureLevel;
+		PreviewShaderPlatformName = NAME_None;
 		PreviewShaderFormatName = NAME_None;
+        PreviewPlatformName = NAME_None;
 		bPreviewFeatureLevelActive = false;
 		bPreviewFeatureLevelWasDefault = true;
 		PreviewDeviceProfileName = NAME_None;
diff --git a/Engine/Source/Runtime/RHI/Private/DataDrivenShaderPlatformInfo.cpp b/Engine/Source/Runtime/RHI/Private/DataDrivenShaderPlatformInfo.cpp
--- a/Engine/Source/Runtime/RHI/Private/DataDrivenShaderPlatformInfo.cpp
+++ b/Engine/Source/Runtime/RHI/Private/DataDrivenShaderPlatformInfo.cpp
@@ -356,7 +356,7 @@ void FGenericDataDrivenShaderPlatformInfo::UpdatePreviewPlatforms()
 				Infos[ShaderPlatform].Language = Infos[EditorSPForPreviewMaxFeatureLevel].Language;
 				Infos[ShaderPlatform].bIsHlslcc = Infos[EditorSPForPreviewMaxFeatureLevel].bIsHlslcc;
 				Infos[ShaderPlatform].bSupportsDxc = Infos[EditorSPForPreviewMaxFeatureLevel].bSupportsDxc;
-				Infos[ShaderPlatform].bSupportsGPUScene &= Infos[EditorSPForPreviewMaxFeatureLevel].bSupportsGPUScene;
+				Infos[ShaderPlatform].bSupportsGPUScene = Infos[EditorSPForPreviewMaxFeatureLevel].bSupportsGPUScene;
 				Infos[ShaderPlatform].bIsPC = true;
 				Infos[ShaderPlatform].bIsConsole = false;
 				Infos[ShaderPlatform].bSupportsSceneDataCompressedTransforms = Infos[EditorSPForPreviewMaxFeatureLevel].bSupportsSceneDataCompressedTransforms;
@@ -379,7 +379,7 @@ void FGenericDataDrivenShaderPlatformInfo::UpdatePreviewPlatforms()
 					Infos[ShaderPlatform].bSupportsInstancedStereo &= Infos[EditorSPForPreviewMaxFeatureLevel].bSupportsInstancedStereo;
 				}
 				Infos[ShaderPlatform].bSupportsMobileMultiView &= Infos[EditorSPForPreviewMaxFeatureLevel].bSupportsMobileMultiView;
-				Infos[ShaderPlatform].bSupportsManualVertexFetch &= Infos[EditorSPForPreviewMaxFeatureLevel].bSupportsManualVertexFetch;
+				Infos[ShaderPlatform].bSupportsManualVertexFetch = Infos[EditorSPForPreviewMaxFeatureLevel].bSupportsManualVertexFetch;
 				Infos[ShaderPlatform].bSupportsRenderTargetWriteMask = false;
 				Infos[ShaderPlatform].bSupportsIntrinsicWaveOnce = false;
 				Infos[ShaderPlatform].bSupportsDOFHybridScattering = false;
@@ -409,4 +409,4 @@ const EShaderPlatform FGenericDataDrivenShaderPlatformInfo::GetShaderPlatformFro
 		return *ShaderPlatform;
 	}
 	return SP_NumPlatforms;
-}
\ No newline at end of file
+}
diff --git a/Engine/Source/Runtime/Renderer/Private/TranslucentLighting.cpp b/Engine/Source/Runtime/Renderer/Private/TranslucentLighting.cpp
--- a/Engine/Source/Runtime/Renderer/Private/TranslucentLighting.cpp
+++ b/Engine/Source/Runtime/Renderer/Private/TranslucentLighting.cpp
@@ -64,7 +64,7 @@ FAutoConsoleVariableRef CVarUseTranslucentLightingVolumes(
 	GUseTranslucentLightingVolumes,
 	TEXT("Whether to allow updating the translucent lighting volumes.\n")
 	TEXT("0:off, otherwise on, default is 1"),
-	ECVF_RenderThreadSafe
+	ECVF_RenderThreadSafe | ECVF_Scalability
 	);
 
 float GTranslucentVolumeMinFOV = 45;
@@ -72,7 +72,7 @@ static FAutoConsoleVariableRef CVarTranslucentVolumeMinFOV(
 	TEXT("r.TranslucentVolumeMinFOV"),
 	GTranslucentVolumeMinFOV,
 	TEXT("Minimum FOV for translucent lighting volume.  Prevents popping in lighting when zooming in."),
-	ECVF_RenderThreadSafe
+	ECVF_RenderThreadSafe | ECVF_Scalability
 	);
 
 float GTranslucentVolumeFOVSnapFactor = 10;
@@ -80,7 +80,7 @@ static FAutoConsoleVariableRef CTranslucentVolumeFOVSnapFactor(
 	TEXT("r.TranslucentVolumeFOVSnapFactor"),
 	GTranslucentVolumeFOVSnapFactor,
 	TEXT("FOV will be snapped to a factor of this before computing volume bounds."),
-	ECVF_RenderThreadSafe
+	ECVF_RenderThreadSafe | ECVF_Scalability
 	);
 
 int32 GUseTranslucencyVolumeBlur = 1;
@@ -104,13 +104,13 @@ static TAutoConsoleVariable<float> CVarTranslucencyLightingVolumeInnerDistance(
 	TEXT("r.TranslucencyLightingVolumeInnerDistance"),
 	1500.0f,
 	TEXT("Distance from the camera that the first volume cascade should end"),
-	ECVF_RenderThreadSafe);
+	ECVF_RenderThreadSafe | ECVF_Scalability);
 
 static TAutoConsoleVariable<float> CVarTranslucencyLightingVolumeOuterDistance(
 	TEXT("r.TranslucencyLightingVolumeOuterDistance"),
 	5000.0f,
 	TEXT("Distance from the camera that the second volume cascade should end"),
-	ECVF_RenderThreadSafe);
+	ECVF_RenderThreadSafe | ECVF_Scalability);
 
 /** Function returning current translucency lighting volume dimensions. */
 int32 GetTranslucencyLightingVolumeDim()
