RHI header dependency cleanup.
- Reordering a lot of type defines to allow RHIResources.h to be included indepdently of other RHI headers.
- Moving all RHI string functions to RHIStrings.h/cpp

diff --git a/Engine/Source/Runtime/RHI/Private/DataDrivenShaderPlatformInfo.cpp b/Engine/Source/Runtime/RHI/Private/DataDrivenShaderPlatformInfo.cpp
--- a/Engine/Source/Runtime/RHI/Private/DataDrivenShaderPlatformInfo.cpp
+++ b/Engine/Source/Runtime/RHI/Private/DataDrivenShaderPlatformInfo.cpp
@@ -138,7 +138,7 @@ void FGenericDataDrivenShaderPlatformInfo::ParseDataDrivenShaderInfo(const FConf
 	Info.ShaderFormat = *GetSectionString(Section, "ShaderFormat");
 	checkf(!Info.ShaderFormat.IsNone(), TEXT("Missing ShaderFormat for ShaderPlatform %s  ShaderFormat %s"), *Info.Name.ToString(), *Info.ShaderFormat.ToString());
 
-	GetFeatureLevelFromName(*GetSectionString(Section, "MaxFeatureLevel"), Info.MaxFeatureLevel);
+	GetFeatureLevelFromName(GetSectionString(Section, "MaxFeatureLevel"), Info.MaxFeatureLevel);
 
 	Info.ShaderPropertiesHash = 0;
 	FString ShaderPropertiesString = Info.Name.GetPlainNameString();
diff --git a/Engine/Source/Runtime/RHI/Private/RHI.cpp b/Engine/Source/Runtime/RHI/Private/RHI.cpp
--- a/Engine/Source/Runtime/RHI/Private/RHI.cpp
+++ b/Engine/Source/Runtime/RHI/Private/RHI.cpp
@@ -9,7 +9,6 @@
 #include "Modules/ModuleManager.h"
 #include "Misc/ConfigCacheIni.h"
 #include "Misc/MessageDialog.h"
-#include "RHIShaderFormatDefinitions.inl"
 #include "ProfilingDebugging/CsvProfiler.h"
 #include "String/Find.h"
 #include "String/LexFromString.h"
@@ -82,175 +81,6 @@ static TAutoConsoleVariable<int32> CVarGraphicsAdapter(
 	TEXT("  1: Adapter #1, ..."),
 	ECVF_ReadOnly | ECVF_RenderThreadSafe);
 
-template<typename EnumType>
-inline FString BuildEnumNameBitList(EnumType Value, const TCHAR*(*GetEnumName)(EnumType))
-{
-	if (Value == EnumType(0))
-	{
-		return GetEnumName(Value);
-	}
-
-	using T = __underlying_type(EnumType);
-	T StateValue = (T)Value;
-
-	FString Name;
-
-	int32 BitIndex = 0;
-	while (StateValue)
-	{
-		if (StateValue & 1)
-		{
-			if (Name.Len() > 0 && StateValue > 0)
-			{
-				Name += TEXT("|");
-			}
-
-			Name += GetEnumName(EnumType(T(1) << BitIndex));
-		}
-
-		BitIndex++;
-		StateValue >>= 1;
-	}
-
-	return MoveTemp(Name);
-}
-
-FString GetRHIAccessName(ERHIAccess Access)
-{
-		return BuildEnumNameBitList<ERHIAccess>(Access, [](ERHIAccess AccessBit)
-		{
-			switch (AccessBit)
-			{
-			default: checkNoEntry(); // fall through
-			case ERHIAccess::Unknown:             return TEXT("Unknown");
-			case ERHIAccess::CPURead:             return TEXT("CPURead");
-			case ERHIAccess::Present:             return TEXT("Present");
-			case ERHIAccess::IndirectArgs:        return TEXT("IndirectArgs");
-			case ERHIAccess::VertexOrIndexBuffer: return TEXT("VertexOrIndexBuffer");
-			case ERHIAccess::SRVCompute:          return TEXT("SRVCompute");
-			case ERHIAccess::SRVGraphics:         return TEXT("SRVGraphics");
-			case ERHIAccess::CopySrc:             return TEXT("CopySrc");
-			case ERHIAccess::ResolveSrc:          return TEXT("ResolveSrc");
-			case ERHIAccess::DSVRead:             return TEXT("DSVRead");
-			case ERHIAccess::UAVCompute:          return TEXT("UAVCompute");
-			case ERHIAccess::UAVGraphics:         return TEXT("UAVGraphics");
-			case ERHIAccess::RTV:                 return TEXT("RTV");
-			case ERHIAccess::CopyDest:            return TEXT("CopyDest");
-			case ERHIAccess::ResolveDst:          return TEXT("ResolveDst");
-			case ERHIAccess::DSVWrite:            return TEXT("DSVWrite");
-			case ERHIAccess::ShadingRateSource:	  return TEXT("ShadingRateSource");
-			case ERHIAccess::BVHRead:             return TEXT("BVHRead");
-			case ERHIAccess::BVHWrite:            return TEXT("BVHWrite");
-			case ERHIAccess::Discard:             return TEXT("Discard");
-			}
-		});
-}
-
-FString GetResourceTransitionFlagsName(EResourceTransitionFlags Flags)
-{
-	return BuildEnumNameBitList<EResourceTransitionFlags>(Flags, [](EResourceTransitionFlags Value)
-	{
-		switch (Value)
-		{
-		default: checkNoEntry(); // fall through
-		case EResourceTransitionFlags::None:                return TEXT("None");
-		case EResourceTransitionFlags::MaintainCompression: return TEXT("MaintainCompression");
-		case EResourceTransitionFlags::Discard:				return TEXT("Discard");
-		case EResourceTransitionFlags::Clear:				return TEXT("Clear");
-		}
-	});
-}
-
-FString GetRHIPipelineName(ERHIPipeline Pipeline)
-{
-	return BuildEnumNameBitList<ERHIPipeline>(Pipeline, [](ERHIPipeline Value)
-	{
-		if (Value == ERHIPipeline(0)) { return TEXT("None"); }
-
-		switch (Value)
-		{
-		default: checkNoEntry(); // fall through
-		case ERHIPipeline::Graphics:     return TEXT("Graphics");
-		case ERHIPipeline::AsyncCompute: return TEXT("AsyncCompute");
-		}
-	});
-}
-
-FString GetBufferUsageFlagsName(EBufferUsageFlags BufferUsage)
-{
-	return BuildEnumNameBitList<EBufferUsageFlags>(BufferUsage, [](EBufferUsageFlags BufferUsage)
-	{
-		switch (BufferUsage)
-		{
-		default: checkNoEntry(); // fall through
-		case BUF_None:					return TEXT("BUF_None");
-		case BUF_Static:				return TEXT("BUF_Static");
-		case BUF_Dynamic:				return TEXT("BUF_Dynamic");
-		case BUF_Volatile:				return TEXT("BUF_Volitile");
-		case BUF_UnorderedAccess:		return TEXT("BUF_UnorderedAccess");
-		case BUF_ByteAddressBuffer:		return TEXT("BUF_ByteAddressBuffer");
-		case BUF_SourceCopy:			return TEXT("BUF_SourceCopy");
-		case BUF_StreamOutput:			return TEXT("BUF_StreamOutput");
-		case BUF_DrawIndirect:			return TEXT("BUF_DrawIndirect");
-		case BUF_ShaderResource:		return TEXT("BUF_ShaderResource");
-		case BUF_KeepCPUAccessible:		return TEXT("BUF_KeepCPUAccessible");
-		case BUF_FastVRAM:				return TEXT("BUF_FastVRAM");
-		case BUF_Shared:				return TEXT("BUF_Shared");
-		case BUF_AccelerationStructure:	return TEXT("BUF_AccelerationStructure");
-		case BUF_RayTracingScratch:		return TEXT("BUF_RayTracingScratch");
-		case BUF_VertexBuffer:			return TEXT("BUF_VertexBuffer");
-		case BUF_IndexBuffer:			return TEXT("BUF_IndexBuffer");
-		case BUF_StructuredBuffer:		return TEXT("BUF_StructuredBuffer");
-		}
-	});
-}
-
-FString GetTextureCreateFlagsName(ETextureCreateFlags TextureCreateFlags)
-{
-	return BuildEnumNameBitList<ETextureCreateFlags>(TextureCreateFlags, [](ETextureCreateFlags TextureCreateFlags)
-	{
-		switch (TextureCreateFlags)
-		{
-		default: checkNoEntry(); // fall through
-		case TexCreate_None:							return TEXT("TexCreate_None");
-		case TexCreate_RenderTargetable:				return TEXT("TexCreate_RenderTargetable");
-		case TexCreate_ResolveTargetable:				return TEXT("TexCreate_ResolveTargetable");
-		case TexCreate_DepthStencilTargetable:			return TEXT("TexCreate_DepthStencilTargetable");
-		case TexCreate_ShaderResource:					return TEXT("TexCreate_ShaderResource");
-		case TexCreate_SRGB:							return TEXT("TexCreate_SRGB");
-		case TexCreate_CPUWritable:						return TEXT("TexCreate_CPUWritable");
-		case TexCreate_NoTiling:						return TEXT("TexCreate_NoTiling");
-		case TexCreate_VideoDecode:						return TEXT("TexCreate_VideoDecode");
-		case TexCreate_Dynamic:							return TEXT("TexCreate_Dynamic");
-		case TexCreate_InputAttachmentRead:				return TEXT("TexCreate_InputAttachmentRead");
-		case TexCreate_Foveation:						return TEXT("TexCreate_Foveation");
-		case TexCreate_3DTiling:						return TEXT("TexCreate_3DTiling");
-		case TexCreate_Memoryless:						return TEXT("TexCreate_Memoryless");
-		case TexCreate_GenerateMipCapable:				return TEXT("TexCreate_GenerateMipCapable");
-		case TexCreate_FastVRAMPartialAlloc:			return TEXT("TexCreate_FastVRAMPartialAlloc");
-		case TexCreate_DisableSRVCreation:				return TEXT("TexCreate_DisableSRVCreation");
-		case TexCreate_DisableDCC:						return TEXT("TexCreate_DisableDCC");
-		case TexCreate_UAV:								return TEXT("TexCreate_UAV");
-		case TexCreate_Presentable:						return TEXT("TexCreate_Presentable");
-		case TexCreate_CPUReadback:						return TEXT("TexCreate_CPUReadback");
-		case TexCreate_OfflineProcessed:				return TEXT("TexCreate_OfflineProcessed");
-		case TexCreate_FastVRAM:						return TEXT("TexCreate_FastVRAM");
-		case TexCreate_HideInVisualizeTexture:			return TEXT("TexCreate_HideInVisualizeTexture");
-		case TexCreate_Virtual:							return TEXT("TexCreate_Virtual");
-		case TexCreate_TargetArraySlicesIndependently:	return TEXT("TexCreate_TargetArraySlicesIndependently");
-		case TexCreate_Shared:							return TEXT("TexCreate_Shared");
-		case TexCreate_NoFastClear:						return TEXT("TexCreate_NoFastClear");
-		case TexCreate_DepthStencilResolveTarget:		return TEXT("TexCreate_DepthStencilResolveTarget");
-		case TexCreate_Streamable:						return TEXT("TexCreate_Streamable");
-		case TexCreate_NoFastClearFinalize:				return TEXT("TexCreate_NoFastClearFinalize");
-		case TexCreate_AFRManual:						return TEXT("TexCreate_AFRManual");
-		case TexCreate_ReduceMemoryWithTilingMode:		return TEXT("TexCreate_ReduceMemoryWithTilingMode");
-		case TexCreate_AtomicCompatible:				return TEXT("TexCreate_AtomicCompatible");
-		case TexCreate_External:						return TEXT("TexCreate_External");
-		}
-	});
-}
-
 #if STATS
 #include "Stats/StatsData.h"
 static void DumpRHIMemory(FOutputDevice& OutputDevice)
@@ -296,9 +126,6 @@ const FClearValueBinding FClearValueBinding::Green(FLinearColor(0.0f, 1.0f, 0.0f
 // Note: this is used as the default normal for DBuffer decals.  It must decode to a value of 0 in DecodeDBufferData.
 const FClearValueBinding FClearValueBinding::DefaultNormal8Bit(FLinearColor(128.0f / 255.0f, 128.0f / 255.0f, 128.0f / 255.0f, 1.0f));
 
-UE::TConsumeAllMpmcQueue<FRHIResource*> FRHIResource::PendingDeletes;
-FRHIResource* FRHIResource::CurrentlyDeleting = nullptr;
-
 #if HAS_GPU_STATS
 
 	FDrawCallCategoryName::FDrawCallCategoryName()
@@ -739,80 +566,6 @@ void FRHIResource::StopTrackingAllResources()
 	GRHITrackingResources = false;
 }
 
-struct FRHIResourceTypeName
-{
-	ERHIResourceType Type;
-	const TCHAR* Name;
-};
-static const FRHIResourceTypeName GRHIResourceTypeNames[] =
-{
-#define RHI_RESOURCE_TYPE_DEF(Name) { RRT_##Name, TEXT(#Name) }
-
-	RHI_RESOURCE_TYPE_DEF(None),
-	RHI_RESOURCE_TYPE_DEF(SamplerState),
-	RHI_RESOURCE_TYPE_DEF(RasterizerState),
-	RHI_RESOURCE_TYPE_DEF(DepthStencilState),
-	RHI_RESOURCE_TYPE_DEF(BlendState),
-	RHI_RESOURCE_TYPE_DEF(VertexDeclaration),
-	RHI_RESOURCE_TYPE_DEF(VertexShader),
-	RHI_RESOURCE_TYPE_DEF(MeshShader),
-	RHI_RESOURCE_TYPE_DEF(AmplificationShader),
-	RHI_RESOURCE_TYPE_DEF(PixelShader),
-	RHI_RESOURCE_TYPE_DEF(GeometryShader),
-	RHI_RESOURCE_TYPE_DEF(RayTracingShader),
-	RHI_RESOURCE_TYPE_DEF(ComputeShader),
-	RHI_RESOURCE_TYPE_DEF(GraphicsPipelineState),
-	RHI_RESOURCE_TYPE_DEF(ComputePipelineState),
-	RHI_RESOURCE_TYPE_DEF(RayTracingPipelineState),
-	RHI_RESOURCE_TYPE_DEF(BoundShaderState),
-	RHI_RESOURCE_TYPE_DEF(UniformBufferLayout),
-	RHI_RESOURCE_TYPE_DEF(UniformBuffer),
-	RHI_RESOURCE_TYPE_DEF(Buffer),
-	RHI_RESOURCE_TYPE_DEF(Texture),
-	RHI_RESOURCE_TYPE_DEF(Texture2D),
-	RHI_RESOURCE_TYPE_DEF(Texture2DArray),
-	RHI_RESOURCE_TYPE_DEF(Texture3D),
-	RHI_RESOURCE_TYPE_DEF(TextureCube),
-	RHI_RESOURCE_TYPE_DEF(TextureReference),
-	RHI_RESOURCE_TYPE_DEF(TimestampCalibrationQuery),
-	RHI_RESOURCE_TYPE_DEF(GPUFence),
-	RHI_RESOURCE_TYPE_DEF(RenderQuery),
-	RHI_RESOURCE_TYPE_DEF(RenderQueryPool),
-	RHI_RESOURCE_TYPE_DEF(ComputeFence),
-	RHI_RESOURCE_TYPE_DEF(Viewport),
-	RHI_RESOURCE_TYPE_DEF(UnorderedAccessView),
-	RHI_RESOURCE_TYPE_DEF(ShaderResourceView),
-	RHI_RESOURCE_TYPE_DEF(RayTracingAccelerationStructure),
-	RHI_RESOURCE_TYPE_DEF(StagingBuffer),
-	RHI_RESOURCE_TYPE_DEF(CustomPresent),
-	RHI_RESOURCE_TYPE_DEF(ShaderLibrary),
-	RHI_RESOURCE_TYPE_DEF(PipelineBinaryLibrary),
-};
-
-static ERHIResourceType RHIResourceTypeFromString(const FString& InString)
-{
-	for (const auto& TypeName : GRHIResourceTypeNames)
-	{
-		if (InString.Equals(TypeName.Name, ESearchCase::IgnoreCase))
-		{
-			return TypeName.Type;
-		}
-	}
-	return RRT_None;
-}
-
-const TCHAR* StringFromRHIResourceType(ERHIResourceType ResourceType)
-{
-	for (const auto& TypeName : GRHIResourceTypeNames)
-	{
-		if (TypeName.Type == ResourceType)
-		{
-			return TypeName.Name;
-		}
-	}
-	return TEXT("<unknown>");
-}
-
 enum class EBooleanFilter
 {
 	No,
@@ -1323,48 +1076,6 @@ void RHIDumpResourceMemoryToCSV()
 
 #endif // RHI_ENABLE_RESOURCE_INFO
 
-bool FRHIResource::Bypass()
-{
-	return GRHICommandList.Bypass();
-}
-
-DECLARE_CYCLE_STAT(TEXT("Delete Resources"), STAT_DeleteResources, STATGROUP_RHICMDLIST);
-
-int32 FRHIResource::FlushPendingDeletes(FRHICommandListImmediate& RHICmdList)
-{
-	SCOPE_CYCLE_COUNTER(STAT_DeleteResources);
-
-	check(IsInRenderingThread());
-
-	TArray<FRHIResource*> DeletedResources;
-	PendingDeletes.ConsumeAllLifo([&DeletedResources](FRHIResource* Resource)
-	{
-		DeletedResources.Push(Resource);
-	});
-
-	const int32 NumDeletes = DeletedResources.Num();
-
-	RHICmdList.EnqueueLambda([DeletedResources = MoveTemp(DeletedResources)](FRHICommandListImmediate& RHICmdList) mutable
-	{
-		if (GDynamicRHI)
-		{
-			GDynamicRHI->RHIPerFrameRHIFlushComplete();
-		}
-
-		for (int32 i = DeletedResources.Num() - 1; i >= 0; i--)
-		{
-			FRHIResource* Resource = DeletedResources[i];
-			if (Resource->AtomicFlags.Deleteing())
-			{
-				FRHIResource::CurrentlyDeleting = Resource;
-				delete Resource;
-			}
-		}
-	});
-
-	return NumDeletes;
-}
-
 static_assert(ERHIZBuffer::FarPlane != ERHIZBuffer::NearPlane, "Near and Far planes must be different!");
 static_assert((int32)ERHIZBuffer::NearPlane == 0 || (int32)ERHIZBuffer::NearPlane == 1, "Invalid Values for Near Plane, can only be 0 or 1!");
 static_assert((int32)ERHIZBuffer::FarPlane == 0 || (int32)ERHIZBuffer::FarPlane == 1, "Invalid Values for Far Plane, can only be 0 or 1");
@@ -1708,9 +1419,6 @@ RHI_API EShaderPlatform GShaderPlatformForFeatureLevel[ERHIFeatureLevel::Num] =
 RHI_API int32 GNumDrawCallsRHI      [MAX_NUM_GPUS] = {};
 RHI_API int32 GNumPrimitivesDrawnRHI[MAX_NUM_GPUS] = {};
 
-RHI_API uint64 GRHITransitionPrivateData_SizeInBytes = 0;
-RHI_API uint64 GRHITransitionPrivateData_AlignInBytes = 0;
-
 // By default, read only states and UAV states are allowed to participate in state merging.
 ERHIAccess GRHIMergeableAccessMask = ERHIAccess::ReadOnlyMask | ERHIAccess::UAVMask;
 
@@ -1817,171 +1525,6 @@ RHI_API EShaderPlatform GMaxRHIShaderPlatform = SP_PCD3D_SM5;
 /** The maximum feature level supported on this machine */
 RHI_API ERHIFeatureLevel::Type GMaxRHIFeatureLevel = ERHIFeatureLevel::SM5;
 
-FName FeatureLevelNames[] = 
-{
-	FName(TEXT("ES2")),
-	FName(TEXT("ES3_1")),
-	FName(TEXT("SM4_REMOVED")),
-	FName(TEXT("SM5")),
-	FName(TEXT("SM6")),
-};
-
-static_assert(UE_ARRAY_COUNT(FeatureLevelNames) == ERHIFeatureLevel::Num, "Missing entry from feature level names.");
-
-RHI_API bool GetFeatureLevelFromName(FName Name, ERHIFeatureLevel::Type& OutFeatureLevel)
-{
-	for (int32 NameIndex = 0; NameIndex < UE_ARRAY_COUNT(FeatureLevelNames); NameIndex++)
-	{
-		if (FeatureLevelNames[NameIndex] == Name)
-		{
-			OutFeatureLevel = (ERHIFeatureLevel::Type)NameIndex;
-			return true;
-		}
-	}
-
-	OutFeatureLevel = ERHIFeatureLevel::Num;
-	return false;
-}
-
-RHI_API void GetFeatureLevelName(ERHIFeatureLevel::Type InFeatureLevel, FString& OutName)
-{
-	check(InFeatureLevel < UE_ARRAY_COUNT(FeatureLevelNames));
-	if (InFeatureLevel < UE_ARRAY_COUNT(FeatureLevelNames))
-	{
-		FeatureLevelNames[(int32)InFeatureLevel].ToString(OutName);
-	}
-	else
-	{
-		OutName = TEXT("InvalidFeatureLevel");
-	}	
-}
-
-static FName InvalidFeatureLevelName(TEXT("InvalidFeatureLevel"));
-RHI_API void GetFeatureLevelName(ERHIFeatureLevel::Type InFeatureLevel, FName& OutName)
-{
-	check(InFeatureLevel < UE_ARRAY_COUNT(FeatureLevelNames));
-	if (InFeatureLevel < UE_ARRAY_COUNT(FeatureLevelNames))
-	{
-		OutName = FeatureLevelNames[(int32)InFeatureLevel];
-	}
-	else
-	{
-		
-		OutName = InvalidFeatureLevelName;
-	}
-}
-
-FName ShadingPathNames[] =
-{
-	FName(TEXT("Deferred")),
-	FName(TEXT("Forward")),
-	FName(TEXT("Mobile")),
-};
-
-static_assert(UE_ARRAY_COUNT(ShadingPathNames) == ERHIShadingPath::Num, "Missing entry from shading path names.");
-
-RHI_API bool GetShadingPathFromName(FName Name, ERHIShadingPath::Type& OutShadingPath)
-{
-	for (int32 NameIndex = 0; NameIndex < UE_ARRAY_COUNT(ShadingPathNames); NameIndex++)
-	{
-		if (ShadingPathNames[NameIndex] == Name)
-		{
-			OutShadingPath = (ERHIShadingPath::Type)NameIndex;
-			return true;
-		}
-	}
-
-	OutShadingPath = ERHIShadingPath::Num;
-	return false;
-}
-
-RHI_API void GetShadingPathName(ERHIShadingPath::Type InShadingPath, FString& OutName)
-{
-	check(InShadingPath < UE_ARRAY_COUNT(ShadingPathNames));
-	if (InShadingPath < UE_ARRAY_COUNT(ShadingPathNames))
-	{
-		ShadingPathNames[(int32)InShadingPath].ToString(OutName);
-	}
-	else
-	{
-		OutName = TEXT("InvalidShadingPath");
-	}
-}
-
-static FName InvalidShadingPathName(TEXT("InvalidShadingPath"));
-RHI_API void GetShadingPathName(ERHIShadingPath::Type InShadingPath, FName& OutName)
-{
-	check(InShadingPath < UE_ARRAY_COUNT(ShadingPathNames));
-	if (InShadingPath < UE_ARRAY_COUNT(ShadingPathNames))
-	{
-		OutName = ShadingPathNames[(int32)InShadingPath];
-	}
-	else
-	{
-
-		OutName = InvalidShadingPathName;
-	}
-}
-
-static FName NAME_PLATFORM_WINDOWS(TEXT("Windows"));
-static FName NAME_PLATFORM_ANDROID(TEXT("Android"));
-static FName NAME_PLATFORM_IOS(TEXT("IOS"));
-static FName NAME_PLATFORM_MAC(TEXT("Mac"));
-static FName NAME_PLATFORM_TVOS(TEXT("TVOS"));
-
-// @todo platplug: This is still here, only being used now by UMaterialShaderQualitySettings::GetOrCreatePlatformSettings
-// since I have moved the other uses to FindTargetPlatformWithSupport
-// But I'd like to delete it anyway!
-FName ShaderPlatformToPlatformName(EShaderPlatform Platform)
-{
-	switch (Platform)
-	{
-	case SP_PCD3D_SM6:
-	case SP_PCD3D_SM5:
-	case SP_PCD3D_ES3_1:
-	case SP_OPENGL_PCES3_1:
-	case SP_VULKAN_PCES3_1:
-	case SP_VULKAN_SM5:
-	case SP_D3D_ES3_1_HOLOLENS:
-		return NAME_PLATFORM_WINDOWS;
-	case SP_VULKAN_ES3_1_ANDROID:
-	case SP_VULKAN_SM5_ANDROID:
-	case SP_OPENGL_ES3_1_ANDROID:
-		return NAME_PLATFORM_ANDROID;
-	case SP_METAL:
-	case SP_METAL_MRT:
-		return NAME_PLATFORM_IOS;
-	case SP_METAL_SM5:
-	case SP_METAL_MACES3_1:
-	case SP_METAL_MRT_MAC:
-		return NAME_PLATFORM_MAC;
-	case SP_METAL_TVOS:
-	case SP_METAL_MRT_TVOS:
-		return NAME_PLATFORM_TVOS;
-
-
-	default:
-		if (FStaticShaderPlatformNames::IsStaticPlatform(Platform))
-		{
-			return FStaticShaderPlatformNames::Get().GetPlatformName(Platform);
-		}
-		else
-		{
-			return NAME_None;
-		}
-	}
-}
-
-FName LegacyShaderPlatformToShaderFormat(EShaderPlatform Platform)
-{
-	return  FDataDrivenShaderPlatformInfo::GetShaderFormat(Platform);
-}
-
-EShaderPlatform ShaderFormatToLegacyShaderPlatform(FName ShaderFormat)
-{
-	return ShaderFormatNameToShaderPlatform(ShaderFormat);
-}
-
 RHI_API bool IsRHIDeviceAMD()
 {
 	check(GRHIVendorId != 0);
@@ -2002,36 +1545,6 @@ RHI_API bool IsRHIDeviceNVIDIA()
 	return GRHIVendorId == 0x10DE;
 }
 
-RHI_API const TCHAR* RHIVendorIdToString()
-{
-	return RHIVendorIdToString((EGpuVendorId)GRHIVendorId);
-}
-
-RHI_API const TCHAR* RHIVendorIdToString(EGpuVendorId VendorId)
-{
-	switch (VendorId)
-	{
-	case EGpuVendorId::Amd:			return TEXT("AMD");
-	case EGpuVendorId::ImgTec:		return TEXT("ImgTec");
-	case EGpuVendorId::Nvidia:		return TEXT("NVIDIA");
-	case EGpuVendorId::Arm:			return TEXT("ARM");
-	case EGpuVendorId::Broadcom:	return TEXT("Broadcom");
-	case EGpuVendorId::Qualcomm:	return TEXT("Qualcomm");
-	case EGpuVendorId::Apple:		return TEXT("Apple");
-	case EGpuVendorId::Intel:		return TEXT("Intel");
-	case EGpuVendorId::Vivante:		return TEXT("Vivante");
-	case EGpuVendorId::VeriSilicon:	return TEXT("VeriSilicon");
-	case EGpuVendorId::Kazan:		return TEXT("Kazan");
-	case EGpuVendorId::Codeplay:	return TEXT("Codeplay");
-	case EGpuVendorId::Mesa:		return TEXT("Mesa");
-	case EGpuVendorId::NotQueried:	return TEXT("Not Queried");
-	default:
-		break;
-	}
-
-	return TEXT("Unknown");
-}
-
 RHI_API uint32 RHIGetMetalShaderLanguageVersion(const FStaticShaderPlatform Platform)
 {
     if (IsMetalPlatform(Platform))
@@ -2456,95 +1969,6 @@ FRHIPanicEvent& RHIGetPanicDelegate()
 	return RHIPanicEvent;
 }
 
-FString LexToString(EShaderPlatform Platform, bool bError)
-{
-	switch (Platform)
-	{
-	case SP_PCD3D_SM6: return TEXT("PCD3D_SM6");
-	case SP_PCD3D_SM5: return TEXT("PCD3D_SM5");
-	case SP_PCD3D_ES3_1: return TEXT("PCD3D_ES3_1");
-	case SP_OPENGL_PCES3_1: return TEXT("OPENGL_PCES3_1");
-	case SP_OPENGL_ES3_1_ANDROID: return TEXT("OPENGL_ES3_1_ANDROID");
-	case SP_METAL: return TEXT("METAL");
-	case SP_METAL_MRT: return TEXT("METAL_MRT");
-	case SP_METAL_TVOS: return TEXT("METAL_TVOS");
-	case SP_METAL_MRT_TVOS: return TEXT("METAL_MRT_TVOS");
-	case SP_METAL_MRT_MAC: return TEXT("METAL_MRT_MAC");
-	case SP_METAL_SM5: return TEXT("METAL_SM5");
-	case SP_METAL_MACES3_1: return TEXT("METAL_MACES3_1");
-	case SP_VULKAN_ES3_1_ANDROID: return TEXT("VULKAN_ES3_1_ANDROID");
-	case SP_VULKAN_PCES3_1: return TEXT("VULKAN_PCES3_1");
-	case SP_VULKAN_SM5: return TEXT("VULKAN_SM5");
-	case SP_VULKAN_SM5_ANDROID: return TEXT("VULKAN_SM5_ANDROID");
-	case SP_D3D_ES3_1_HOLOLENS: return TEXT("D3D_ES3_1_HOLOLENS");
-	case SP_CUSTOM_PLATFORM_FIRST: return TEXT("SP_CUSTOM_PLATFORM_FIRST");
-	case SP_CUSTOM_PLATFORM_LAST: return TEXT("SP_CUSTOM_PLATFORM_LAST");
-
-	default:
-		if (FStaticShaderPlatformNames::IsStaticPlatform(Platform))
-		{
-			return FStaticShaderPlatformNames::Get().GetShaderPlatform(Platform).ToString();
-		}
-		else if (Platform > SP_CUSTOM_PLATFORM_FIRST && Platform < SP_CUSTOM_PLATFORM_LAST)
-		{
-			return FString::Printf(TEXT("SP_CUSTOM_PLATFORM_%d"), Platform - SP_CUSTOM_PLATFORM_FIRST);
-		}
-		else
-		{
-			checkf(!bError, TEXT("Unknown or removed EShaderPlatform %d!"), (int32)Platform);
-			return TEXT("");
-		}
-	}
-}
-
-FString LexToString(ERHIFeatureLevel::Type Level)
-{
-	switch (Level)
-	{
-		case ERHIFeatureLevel::ES2_REMOVED:
-			return TEXT("ES2_REMOVED");
-		case ERHIFeatureLevel::ES3_1:
-			return TEXT("ES3_1");
-		case ERHIFeatureLevel::SM4_REMOVED:
-			return TEXT("SM4_REMOVED");
-		case ERHIFeatureLevel::SM5:
-			return TEXT("SM5");
-		case ERHIFeatureLevel::SM6:
-			return TEXT("SM6");
-		default:
-			break;
-	}
-	return TEXT("UnknownFeatureLevel");
-}
-
-const TCHAR* LexToString(ERHIDescriptorHeapType InHeapType)
-{
-	switch (InHeapType)
-	{
-	default: checkNoEntry();                   return TEXT("UnknownDescriptorHeapType");
-	case ERHIDescriptorHeapType::Standard:     return TEXT("Standard");
-	case ERHIDescriptorHeapType::RenderTarget: return TEXT("RenderTarget");
-	case ERHIDescriptorHeapType::DepthStencil: return TEXT("DepthStencil");
-	case ERHIDescriptorHeapType::Sampler:      return TEXT("Sampler");
-	}
-}
-
-FRHIViewableResource* GetViewableResource(const FRHITransitionInfo& Info)
-{
-	switch (Info.Type)
-	{
-	case FRHITransitionInfo::EType::Buffer:
-	case FRHITransitionInfo::EType::Texture:
-		return Info.ViewableResource;
-
-	case FRHITransitionInfo::EType::UAV:
-		return Info.UAV ? Info.UAV->GetParentResource() : nullptr;
-	}
-
-	return nullptr;
-}
-
-
 //
 //	MSAA sample offsets.
 //
diff --git a/Engine/Source/Runtime/RHI/Private/RHIResources.cpp b/Engine/Source/Runtime/RHI/Private/RHIResources.cpp
new file mode 100644
--- /dev/null
+++ b/Engine/Source/Runtime/RHI/Private/RHIResources.cpp
@@ -0,0 +1,51 @@
+// Copyright Epic Games, Inc. All Rights Reserved.
+
+#include "RHIResources.h"
+#include "RHI.h"
+
+std::atomic<TClosableMpscQueue<FRHIResource*>*> FRHIResource::PendingDeletes{ new TClosableMpscQueue<FRHIResource*>() };
+FHazardPointerCollection FRHIResource::PendingDeletesHPC;
+FRHIResource* FRHIResource::CurrentlyDeleting = nullptr;
+
+bool FRHIResource::Bypass()
+{
+	return GRHICommandList.Bypass();
+}
+
+DECLARE_CYCLE_STAT(TEXT("Delete Resources"), STAT_DeleteResources, STATGROUP_RHICMDLIST);
+
+int32 FRHIResource::FlushPendingDeletes(FRHICommandListImmediate& RHICmdList)
+{
+	SCOPE_CYCLE_COUNTER(STAT_DeleteResources);
+
+	check(IsInRenderingThread());
+
+	TArray<FRHIResource*> DeletedResources;
+	TClosableMpscQueue<FRHIResource*>* PendingDeletesPtr = PendingDeletes.exchange(new TClosableMpscQueue<FRHIResource*>());
+	PendingDeletesPtr->Close([&DeletedResources](FRHIResource* Resource)
+	{
+		DeletedResources.Push(Resource);
+	});
+	PendingDeletesHPC.Delete(PendingDeletesPtr);
+
+	const int32 NumDeletes = DeletedResources.Num();
+
+	RHICmdList.EnqueueLambda([DeletedResources = MoveTemp(DeletedResources)](FRHICommandListImmediate& RHICmdList) mutable
+	{
+		if (GDynamicRHI)
+		{
+			GDynamicRHI->RHIPerFrameRHIFlushComplete();
+		}
+
+		for (FRHIResource* Resource : DeletedResources)
+		{
+			if (Resource->AtomicFlags.Deleteing())
+			{
+				FRHIResource::CurrentlyDeleting = Resource;
+				delete Resource;
+			}
+		}
+	});
+
+	return NumDeletes;
+}
diff --git a/Engine/Source/Runtime/RHI/Private/RHIStrings.cpp b/Engine/Source/Runtime/RHI/Private/RHIStrings.cpp
new file mode 100644
--- /dev/null
+++ b/Engine/Source/Runtime/RHI/Private/RHIStrings.cpp
@@ -0,0 +1,696 @@
+// Copyright Epic Games, Inc. All Rights Reserved.
+
+#include "RHIStrings.h"
+#include "DataDrivenShaderPlatformInfo.h"
+#include "RHI.h"
+#include "RHIShaderFormatDefinitions.inl"
+#include "UObject/NameTypes.h"
+
+static FName NAME_PLATFORM_WINDOWS(TEXT("Windows"));
+static FName NAME_PLATFORM_ANDROID(TEXT("Android"));
+static FName NAME_PLATFORM_IOS(TEXT("IOS"));
+static FName NAME_PLATFORM_MAC(TEXT("Mac"));
+static FName NAME_PLATFORM_TVOS(TEXT("TVOS"));
+
+static FName InvalidFeatureLevelName(TEXT("InvalidFeatureLevel"));
+static FName InvalidShadingPathName(TEXT("InvalidShadingPath"));
+
+template<typename EnumType>
+inline FString BuildEnumNameBitList(EnumType Value, const TCHAR* (*GetEnumName)(EnumType))
+{
+	if (Value == EnumType(0))
+	{
+		return GetEnumName(Value);
+	}
+
+	using T = __underlying_type(EnumType);
+	T StateValue = (T)Value;
+
+	FString Name;
+
+	int32 BitIndex = 0;
+	while (StateValue)
+	{
+		if (StateValue & 1)
+		{
+			if (Name.Len() > 0 && StateValue > 0)
+			{
+				Name += TEXT("|");
+			}
+
+			Name += GetEnumName(EnumType(T(1) << BitIndex));
+		}
+
+		BitIndex++;
+		StateValue >>= 1;
+	}
+
+	return MoveTemp(Name);
+}
+
+const TCHAR* RHIVendorIdToString()
+{
+	return RHIVendorIdToString((EGpuVendorId)GRHIVendorId);
+}
+
+const TCHAR* RHIVendorIdToString(EGpuVendorId VendorId)
+{
+	switch (VendorId)
+	{
+	case EGpuVendorId::Amd:			return TEXT("AMD");
+	case EGpuVendorId::ImgTec:		return TEXT("ImgTec");
+	case EGpuVendorId::Nvidia:		return TEXT("NVIDIA");
+	case EGpuVendorId::Arm:			return TEXT("ARM");
+	case EGpuVendorId::Broadcom:	return TEXT("Broadcom");
+	case EGpuVendorId::Qualcomm:	return TEXT("Qualcomm");
+	case EGpuVendorId::Apple:		return TEXT("Apple");
+	case EGpuVendorId::Intel:		return TEXT("Intel");
+	case EGpuVendorId::Vivante:		return TEXT("Vivante");
+	case EGpuVendorId::VeriSilicon:	return TEXT("VeriSilicon");
+	case EGpuVendorId::Kazan:		return TEXT("Kazan");
+	case EGpuVendorId::Codeplay:	return TEXT("Codeplay");
+	case EGpuVendorId::Mesa:		return TEXT("Mesa");
+	case EGpuVendorId::NotQueried:	return TEXT("Not Queried");
+	default:                        return TEXT("Unknown");
+	}
+}
+
+FString LexToString(EShaderPlatform Platform, bool bError)
+{
+	switch (Platform)
+	{
+	case SP_PCD3D_SM6: return TEXT("PCD3D_SM6");
+	case SP_PCD3D_SM5: return TEXT("PCD3D_SM5");
+	case SP_PCD3D_ES3_1: return TEXT("PCD3D_ES3_1");
+	case SP_OPENGL_PCES3_1: return TEXT("OPENGL_PCES3_1");
+	case SP_OPENGL_ES3_1_ANDROID: return TEXT("OPENGL_ES3_1_ANDROID");
+	case SP_METAL: return TEXT("METAL");
+	case SP_METAL_MRT: return TEXT("METAL_MRT");
+	case SP_METAL_TVOS: return TEXT("METAL_TVOS");
+	case SP_METAL_MRT_TVOS: return TEXT("METAL_MRT_TVOS");
+	case SP_METAL_MRT_MAC: return TEXT("METAL_MRT_MAC");
+	case SP_METAL_SM5: return TEXT("METAL_SM5");
+	case SP_METAL_MACES3_1: return TEXT("METAL_MACES3_1");
+	case SP_VULKAN_ES3_1_ANDROID: return TEXT("VULKAN_ES3_1_ANDROID");
+	case SP_VULKAN_PCES3_1: return TEXT("VULKAN_PCES3_1");
+	case SP_VULKAN_SM5: return TEXT("VULKAN_SM5");
+	case SP_VULKAN_SM5_ANDROID: return TEXT("VULKAN_SM5_ANDROID");
+	case SP_D3D_ES3_1_HOLOLENS: return TEXT("D3D_ES3_1_HOLOLENS");
+	case SP_CUSTOM_PLATFORM_FIRST: return TEXT("SP_CUSTOM_PLATFORM_FIRST");
+	case SP_CUSTOM_PLATFORM_LAST: return TEXT("SP_CUSTOM_PLATFORM_LAST");
+
+	default:
+		if (FStaticShaderPlatformNames::IsStaticPlatform(Platform))
+		{
+			return FStaticShaderPlatformNames::Get().GetShaderPlatform(Platform).ToString();
+		}
+		else if (Platform > SP_CUSTOM_PLATFORM_FIRST && Platform < SP_CUSTOM_PLATFORM_LAST)
+		{
+			return FString::Printf(TEXT("SP_CUSTOM_PLATFORM_%d"), Platform - SP_CUSTOM_PLATFORM_FIRST);
+		}
+		else
+		{
+			checkf(!bError, TEXT("Unknown or removed EShaderPlatform %d!"), (int32)Platform);
+			return TEXT("");
+		}
+	}
+}
+
+FString LexToString(ERHIFeatureLevel::Type Level)
+{
+	switch (Level)
+	{
+	case ERHIFeatureLevel::ES2_REMOVED: return TEXT("ES2_REMOVED");
+	case ERHIFeatureLevel::ES3_1:       return TEXT("ES3_1");
+	case ERHIFeatureLevel::SM4_REMOVED: return TEXT("SM4_REMOVED");
+	case ERHIFeatureLevel::SM5:         return TEXT("SM5");
+	case ERHIFeatureLevel::SM6:         return TEXT("SM6");
+	default:                            return TEXT("UnknownFeatureLevel");
+	}
+}
+
+static const TCHAR* FeatureLevelNames[] =
+{
+	TEXT("ES2"),
+	TEXT("ES3_1"),
+	TEXT("SM4_REMOVED"),
+	TEXT("SM5"),
+	TEXT("SM6"),
+};
+static_assert(UE_ARRAY_COUNT(FeatureLevelNames) == ERHIFeatureLevel::Num, "Missing entry from feature level names.");
+
+bool GetFeatureLevelFromName(const FString& Name, ERHIFeatureLevel::Type& OutFeatureLevel)
+{
+	for (int32 NameIndex = 0; NameIndex < UE_ARRAY_COUNT(FeatureLevelNames); NameIndex++)
+	{
+		if (FCString::Strcmp(FeatureLevelNames[NameIndex], *Name) == 0)
+		{
+			OutFeatureLevel = (ERHIFeatureLevel::Type)NameIndex;
+			return true;
+		}
+	}
+
+	OutFeatureLevel = ERHIFeatureLevel::Num;
+	return false;
+}
+
+bool GetFeatureLevelFromName(FName Name, ERHIFeatureLevel::Type& OutFeatureLevel)
+{
+	const FString StringName(Name.ToString());
+	return GetFeatureLevelFromName(StringName, OutFeatureLevel);
+}
+
+void GetFeatureLevelName(ERHIFeatureLevel::Type InFeatureLevel, FString& OutName)
+{
+	check(InFeatureLevel < UE_ARRAY_COUNT(FeatureLevelNames));
+	if (InFeatureLevel < UE_ARRAY_COUNT(FeatureLevelNames))
+	{
+		OutName = FeatureLevelNames[(int32)InFeatureLevel];
+	}
+	else
+	{
+		OutName = TEXT("InvalidFeatureLevel");
+	}
+}
+
+void GetFeatureLevelName(ERHIFeatureLevel::Type InFeatureLevel, FName& OutName)
+{
+	check(InFeatureLevel < UE_ARRAY_COUNT(FeatureLevelNames));
+	if (InFeatureLevel < UE_ARRAY_COUNT(FeatureLevelNames))
+	{
+		OutName = FName(FeatureLevelNames[(int32)InFeatureLevel]);
+	}
+	else
+	{
+		OutName = InvalidFeatureLevelName;
+	}
+}
+
+// @todo platplug: This is still here, only being used now by UMaterialShaderQualitySettings::GetOrCreatePlatformSettings
+// since I have moved the other uses to FindTargetPlatformWithSupport
+// But I'd like to delete it anyway!
+FName ShaderPlatformToPlatformName(EShaderPlatform Platform)
+{
+	switch (Platform)
+	{
+	case SP_PCD3D_SM6:
+	case SP_PCD3D_SM5:
+	case SP_PCD3D_ES3_1:
+	case SP_OPENGL_PCES3_1:
+	case SP_VULKAN_PCES3_1:
+	case SP_VULKAN_SM5:
+	case SP_D3D_ES3_1_HOLOLENS:
+		return NAME_PLATFORM_WINDOWS;
+	case SP_VULKAN_ES3_1_ANDROID:
+	case SP_VULKAN_SM5_ANDROID:
+	case SP_OPENGL_ES3_1_ANDROID:
+		return NAME_PLATFORM_ANDROID;
+	case SP_METAL:
+	case SP_METAL_MRT:
+		return NAME_PLATFORM_IOS;
+	case SP_METAL_SM5:
+	case SP_METAL_MACES3_1:
+	case SP_METAL_MRT_MAC:
+		return NAME_PLATFORM_MAC;
+	case SP_METAL_TVOS:
+	case SP_METAL_MRT_TVOS:
+		return NAME_PLATFORM_TVOS;
+
+	default:
+		if (FStaticShaderPlatformNames::IsStaticPlatform(Platform))
+		{
+			return FStaticShaderPlatformNames::Get().GetPlatformName(Platform);
+		}
+		else
+		{
+			return NAME_None;
+		}
+	}
+}
+
+FName LegacyShaderPlatformToShaderFormat(EShaderPlatform Platform)
+{
+	return FDataDrivenShaderPlatformInfo::GetShaderFormat(Platform);
+}
+
+EShaderPlatform ShaderFormatToLegacyShaderPlatform(FName ShaderFormat)
+{
+	return ShaderFormatNameToShaderPlatform(ShaderFormat);
+}
+
+const TCHAR* LexToString(ERHIDescriptorHeapType InHeapType)
+{
+	switch (InHeapType)
+	{
+	default: checkNoEntry();                   return TEXT("UnknownDescriptorHeapType");
+	case ERHIDescriptorHeapType::Standard:     return TEXT("Standard");
+	case ERHIDescriptorHeapType::RenderTarget: return TEXT("RenderTarget");
+	case ERHIDescriptorHeapType::DepthStencil: return TEXT("DepthStencil");
+	case ERHIDescriptorHeapType::Sampler:      return TEXT("Sampler");
+	}
+}
+
+static const FName ShadingPathNames[] =
+{
+	FName(TEXT("Deferred")),
+	FName(TEXT("Forward")),
+	FName(TEXT("Mobile")),
+};
+static_assert(UE_ARRAY_COUNT(ShadingPathNames) == ERHIShadingPath::Num, "Missing entry from shading path names.");
+
+bool GetShadingPathFromName(FName Name, ERHIShadingPath::Type& OutShadingPath)
+{
+	for (int32 NameIndex = 0; NameIndex < UE_ARRAY_COUNT(ShadingPathNames); NameIndex++)
+	{
+		if (ShadingPathNames[NameIndex] == Name)
+		{
+			OutShadingPath = (ERHIShadingPath::Type)NameIndex;
+			return true;
+		}
+	}
+
+	OutShadingPath = ERHIShadingPath::Num;
+	return false;
+}
+
+void GetShadingPathName(ERHIShadingPath::Type InShadingPath, FString& OutName)
+{
+	check(InShadingPath < UE_ARRAY_COUNT(ShadingPathNames));
+	if (InShadingPath < UE_ARRAY_COUNT(ShadingPathNames))
+	{
+		ShadingPathNames[(int32)InShadingPath].ToString(OutName);
+	}
+	else
+	{
+		OutName = TEXT("InvalidShadingPath");
+	}
+}
+
+void GetShadingPathName(ERHIShadingPath::Type InShadingPath, FName& OutName)
+{
+	check(InShadingPath < UE_ARRAY_COUNT(ShadingPathNames));
+	if (InShadingPath < UE_ARRAY_COUNT(ShadingPathNames))
+	{
+		OutName = ShadingPathNames[(int32)InShadingPath];
+	}
+	else
+	{
+		OutName = InvalidShadingPathName;
+	}
+}
+
+FString GetBufferUsageFlagsName(EBufferUsageFlags BufferUsage)
+{
+	return BuildEnumNameBitList<EBufferUsageFlags>(BufferUsage, [](EBufferUsageFlags BufferUsage)
+	{
+		switch (BufferUsage)
+		{
+		default: checkNoEntry(); // fall through
+		case BUF_None:					return TEXT("BUF_None");
+		case BUF_Static:				return TEXT("BUF_Static");
+		case BUF_Dynamic:				return TEXT("BUF_Dynamic");
+		case BUF_Volatile:				return TEXT("BUF_Volitile");
+		case BUF_UnorderedAccess:		return TEXT("BUF_UnorderedAccess");
+		case BUF_ByteAddressBuffer:		return TEXT("BUF_ByteAddressBuffer");
+		case BUF_SourceCopy:			return TEXT("BUF_SourceCopy");
+		case BUF_StreamOutput:			return TEXT("BUF_StreamOutput");
+		case BUF_DrawIndirect:			return TEXT("BUF_DrawIndirect");
+		case BUF_ShaderResource:		return TEXT("BUF_ShaderResource");
+		case BUF_KeepCPUAccessible:		return TEXT("BUF_KeepCPUAccessible");
+		case BUF_FastVRAM:				return TEXT("BUF_FastVRAM");
+		case BUF_Shared:				return TEXT("BUF_Shared");
+		case BUF_AccelerationStructure:	return TEXT("BUF_AccelerationStructure");
+		case BUF_RayTracingScratch:		return TEXT("BUF_RayTracingScratch");
+		case BUF_VertexBuffer:			return TEXT("BUF_VertexBuffer");
+		case BUF_IndexBuffer:			return TEXT("BUF_IndexBuffer");
+		case BUF_StructuredBuffer:		return TEXT("BUF_StructuredBuffer");
+		}
+	});
+}
+
+FString GetTextureCreateFlagsName(ETextureCreateFlags TextureCreateFlags)
+{
+	return BuildEnumNameBitList<ETextureCreateFlags>(TextureCreateFlags, [](ETextureCreateFlags TextureCreateFlags)
+	{
+		switch (TextureCreateFlags)
+		{
+		default: checkNoEntry(); // fall through
+		case TexCreate_None:							return TEXT("TexCreate_None");
+		case TexCreate_RenderTargetable:				return TEXT("TexCreate_RenderTargetable");
+		case TexCreate_ResolveTargetable:				return TEXT("TexCreate_ResolveTargetable");
+		case TexCreate_DepthStencilTargetable:			return TEXT("TexCreate_DepthStencilTargetable");
+		case TexCreate_ShaderResource:					return TEXT("TexCreate_ShaderResource");
+		case TexCreate_SRGB:							return TEXT("TexCreate_SRGB");
+		case TexCreate_CPUWritable:						return TEXT("TexCreate_CPUWritable");
+		case TexCreate_NoTiling:						return TEXT("TexCreate_NoTiling");
+		case TexCreate_VideoDecode:						return TEXT("TexCreate_VideoDecode");
+		case TexCreate_Dynamic:							return TEXT("TexCreate_Dynamic");
+		case TexCreate_InputAttachmentRead:				return TEXT("TexCreate_InputAttachmentRead");
+		case TexCreate_Foveation:						return TEXT("TexCreate_Foveation");
+		case TexCreate_3DTiling:						return TEXT("TexCreate_3DTiling");
+		case TexCreate_Memoryless:						return TEXT("TexCreate_Memoryless");
+		case TexCreate_GenerateMipCapable:				return TEXT("TexCreate_GenerateMipCapable");
+		case TexCreate_FastVRAMPartialAlloc:			return TEXT("TexCreate_FastVRAMPartialAlloc");
+		case TexCreate_DisableSRVCreation:				return TEXT("TexCreate_DisableSRVCreation");
+		case TexCreate_DisableDCC:						return TEXT("TexCreate_DisableDCC");
+		case TexCreate_UAV:								return TEXT("TexCreate_UAV");
+		case TexCreate_Presentable:						return TEXT("TexCreate_Presentable");
+		case TexCreate_CPUReadback:						return TEXT("TexCreate_CPUReadback");
+		case TexCreate_OfflineProcessed:				return TEXT("TexCreate_OfflineProcessed");
+		case TexCreate_FastVRAM:						return TEXT("TexCreate_FastVRAM");
+		case TexCreate_HideInVisualizeTexture:			return TEXT("TexCreate_HideInVisualizeTexture");
+		case TexCreate_Virtual:							return TEXT("TexCreate_Virtual");
+		case TexCreate_TargetArraySlicesIndependently:	return TEXT("TexCreate_TargetArraySlicesIndependently");
+		case TexCreate_Shared:							return TEXT("TexCreate_Shared");
+		case TexCreate_NoFastClear:						return TEXT("TexCreate_NoFastClear");
+		case TexCreate_DepthStencilResolveTarget:		return TEXT("TexCreate_DepthStencilResolveTarget");
+		case TexCreate_Streamable:						return TEXT("TexCreate_Streamable");
+		case TexCreate_NoFastClearFinalize:				return TEXT("TexCreate_NoFastClearFinalize");
+		case TexCreate_AFRManual:						return TEXT("TexCreate_AFRManual");
+		case TexCreate_ReduceMemoryWithTilingMode:		return TEXT("TexCreate_ReduceMemoryWithTilingMode");
+		case TexCreate_AtomicCompatible:				return TEXT("TexCreate_AtomicCompatible");
+		case TexCreate_External:						return TEXT("TexCreate_External");
+		}
+	});
+}
+
+struct FRHIResourceTypeName
+{
+	ERHIResourceType Type;
+	const TCHAR* Name;
+};
+static const FRHIResourceTypeName GRHIResourceTypeNames[] =
+{
+#define RHI_RESOURCE_TYPE_DEF(Name) { RRT_##Name, TEXT(#Name) }
+
+	RHI_RESOURCE_TYPE_DEF(None),
+	RHI_RESOURCE_TYPE_DEF(SamplerState),
+	RHI_RESOURCE_TYPE_DEF(RasterizerState),
+	RHI_RESOURCE_TYPE_DEF(DepthStencilState),
+	RHI_RESOURCE_TYPE_DEF(BlendState),
+	RHI_RESOURCE_TYPE_DEF(VertexDeclaration),
+	RHI_RESOURCE_TYPE_DEF(VertexShader),
+	RHI_RESOURCE_TYPE_DEF(MeshShader),
+	RHI_RESOURCE_TYPE_DEF(AmplificationShader),
+	RHI_RESOURCE_TYPE_DEF(PixelShader),
+	RHI_RESOURCE_TYPE_DEF(GeometryShader),
+	RHI_RESOURCE_TYPE_DEF(RayTracingShader),
+	RHI_RESOURCE_TYPE_DEF(ComputeShader),
+	RHI_RESOURCE_TYPE_DEF(GraphicsPipelineState),
+	RHI_RESOURCE_TYPE_DEF(ComputePipelineState),
+	RHI_RESOURCE_TYPE_DEF(RayTracingPipelineState),
+	RHI_RESOURCE_TYPE_DEF(BoundShaderState),
+	RHI_RESOURCE_TYPE_DEF(UniformBufferLayout),
+	RHI_RESOURCE_TYPE_DEF(UniformBuffer),
+	RHI_RESOURCE_TYPE_DEF(Buffer),
+	RHI_RESOURCE_TYPE_DEF(Texture),
+	RHI_RESOURCE_TYPE_DEF(Texture2D),
+	RHI_RESOURCE_TYPE_DEF(Texture2DArray),
+	RHI_RESOURCE_TYPE_DEF(Texture3D),
+	RHI_RESOURCE_TYPE_DEF(TextureCube),
+	RHI_RESOURCE_TYPE_DEF(TextureReference),
+	RHI_RESOURCE_TYPE_DEF(TimestampCalibrationQuery),
+	RHI_RESOURCE_TYPE_DEF(GPUFence),
+	RHI_RESOURCE_TYPE_DEF(RenderQuery),
+	RHI_RESOURCE_TYPE_DEF(RenderQueryPool),
+	RHI_RESOURCE_TYPE_DEF(ComputeFence),
+	RHI_RESOURCE_TYPE_DEF(Viewport),
+	RHI_RESOURCE_TYPE_DEF(UnorderedAccessView),
+	RHI_RESOURCE_TYPE_DEF(ShaderResourceView),
+	RHI_RESOURCE_TYPE_DEF(RayTracingAccelerationStructure),
+	RHI_RESOURCE_TYPE_DEF(StagingBuffer),
+	RHI_RESOURCE_TYPE_DEF(CustomPresent),
+	RHI_RESOURCE_TYPE_DEF(ShaderLibrary),
+	RHI_RESOURCE_TYPE_DEF(PipelineBinaryLibrary),
+};
+
+ERHIResourceType RHIResourceTypeFromString(const FString& InString)
+{
+	for (const auto& TypeName : GRHIResourceTypeNames)
+	{
+		if (InString.Equals(TypeName.Name, ESearchCase::IgnoreCase))
+		{
+			return TypeName.Type;
+		}
+	}
+	return RRT_None;
+}
+
+const TCHAR* StringFromRHIResourceType(ERHIResourceType ResourceType)
+{
+	for (const auto& TypeName : GRHIResourceTypeNames)
+	{
+		if (TypeName.Type == ResourceType)
+		{
+			return TypeName.Name;
+		}
+	}
+	return TEXT("<unknown>");
+}
+
+FString GetRHIAccessName(ERHIAccess Access)
+{
+	return BuildEnumNameBitList<ERHIAccess>(Access, [](ERHIAccess AccessBit)
+	{
+		switch (AccessBit)
+		{
+		default: checkNoEntry(); // fall through
+		case ERHIAccess::Unknown:             return TEXT("Unknown");
+		case ERHIAccess::CPURead:             return TEXT("CPURead");
+		case ERHIAccess::Present:             return TEXT("Present");
+		case ERHIAccess::IndirectArgs:        return TEXT("IndirectArgs");
+		case ERHIAccess::VertexOrIndexBuffer: return TEXT("VertexOrIndexBuffer");
+		case ERHIAccess::SRVCompute:          return TEXT("SRVCompute");
+		case ERHIAccess::SRVGraphics:         return TEXT("SRVGraphics");
+		case ERHIAccess::CopySrc:             return TEXT("CopySrc");
+		case ERHIAccess::ResolveSrc:          return TEXT("ResolveSrc");
+		case ERHIAccess::DSVRead:             return TEXT("DSVRead");
+		case ERHIAccess::UAVCompute:          return TEXT("UAVCompute");
+		case ERHIAccess::UAVGraphics:         return TEXT("UAVGraphics");
+		case ERHIAccess::RTV:                 return TEXT("RTV");
+		case ERHIAccess::CopyDest:            return TEXT("CopyDest");
+		case ERHIAccess::ResolveDst:          return TEXT("ResolveDst");
+		case ERHIAccess::DSVWrite:            return TEXT("DSVWrite");
+		case ERHIAccess::ShadingRateSource:	  return TEXT("ShadingRateSource");
+		case ERHIAccess::BVHRead:             return TEXT("BVHRead");
+		case ERHIAccess::BVHWrite:            return TEXT("BVHWrite");
+		case ERHIAccess::Discard:             return TEXT("Discard");
+		}
+	});
+}
+
+FString GetResourceTransitionFlagsName(EResourceTransitionFlags Flags)
+{
+	return BuildEnumNameBitList<EResourceTransitionFlags>(Flags, [](EResourceTransitionFlags Value)
+	{
+		switch (Value)
+		{
+		default: checkNoEntry(); // fall through
+		case EResourceTransitionFlags::None:                return TEXT("None");
+		case EResourceTransitionFlags::MaintainCompression: return TEXT("MaintainCompression");
+		case EResourceTransitionFlags::Discard:				return TEXT("Discard");
+		case EResourceTransitionFlags::Clear:				return TEXT("Clear");
+		}
+	});
+}
+
+FString GetRHIPipelineName(ERHIPipeline Pipeline)
+{
+	return BuildEnumNameBitList<ERHIPipeline>(Pipeline, [](ERHIPipeline Value)
+	{
+		if (Value == ERHIPipeline(0))
+		{
+			return TEXT("None");
+		}
+
+		switch (Value)
+		{
+		default: checkNoEntry(); // fall through
+		case ERHIPipeline::Graphics:     return TEXT("Graphics");
+		case ERHIPipeline::AsyncCompute: return TEXT("AsyncCompute");
+		}
+	});
+}
+
+const TCHAR* GetTextureDimensionString(ETextureDimension Dimension)
+{
+	switch (Dimension)
+	{
+	case ETextureDimension::Texture2D:        return TEXT("Texture2D");
+	case ETextureDimension::Texture2DArray:   return TEXT("Texture2DArray");
+	case ETextureDimension::Texture3D:        return TEXT("Texture3D");
+	case ETextureDimension::TextureCube:      return TEXT("TextureCube");
+	case ETextureDimension::TextureCubeArray: return TEXT("TextureCubeArray");
+	default:                                  return TEXT("");
+	}
+}
+
+const TCHAR* GetTextureCreateFlagString(ETextureCreateFlags TextureCreateFlag)
+{
+	switch (TextureCreateFlag)
+	{
+	case ETextureCreateFlags::None:
+		return TEXT("None");
+	case ETextureCreateFlags::RenderTargetable:
+		return TEXT("RenderTargetable");
+	case ETextureCreateFlags::ResolveTargetable:
+		return TEXT("ResolveTargetable");
+	case ETextureCreateFlags::DepthStencilTargetable:
+		return TEXT("DepthStencilTargetable");
+	case ETextureCreateFlags::ShaderResource:
+		return TEXT("ShaderResource");
+	case ETextureCreateFlags::SRGB:
+		return TEXT("SRGB");
+	case ETextureCreateFlags::CPUWritable:
+		return TEXT("CPUWritable");
+	case ETextureCreateFlags::NoTiling:
+		return TEXT("NoTiling");
+	case ETextureCreateFlags::VideoDecode:
+		return TEXT("VideoDecode");
+	case ETextureCreateFlags::Dynamic:
+		return TEXT("Dynamic");
+	case ETextureCreateFlags::InputAttachmentRead:
+		return TEXT("InputAttachmentRead");
+	case ETextureCreateFlags::Foveation:
+		return TEXT("Foveation");
+	case ETextureCreateFlags::Tiling3D:
+		return TEXT("Tiling3D");
+	case ETextureCreateFlags::Memoryless:
+		return TEXT("Memoryless");
+	case ETextureCreateFlags::GenerateMipCapable:
+		return TEXT("GenerateMipCapable");
+	case ETextureCreateFlags::FastVRAMPartialAlloc:
+		return TEXT("FastVRAMPartialAlloc");
+	case ETextureCreateFlags::DisableSRVCreation:
+		return TEXT("DisableSRVCreation");
+	case ETextureCreateFlags::DisableDCC:
+		return TEXT("DisableDCC");
+	case ETextureCreateFlags::UAV:
+		return TEXT("UAV");
+	case ETextureCreateFlags::Presentable:
+		return TEXT("Presentable");
+	case ETextureCreateFlags::CPUReadback:
+		return TEXT("CPUReadback");
+	case ETextureCreateFlags::OfflineProcessed:
+		return TEXT("OfflineProcessed");
+	case ETextureCreateFlags::FastVRAM:
+		return TEXT("FastVRAM");
+	case ETextureCreateFlags::HideInVisualizeTexture:
+		return TEXT("HideInVisualizeTexture");
+	case ETextureCreateFlags::Virtual:
+		return TEXT("Virtual");
+	case ETextureCreateFlags::TargetArraySlicesIndependently:
+		return TEXT("TargetArraySlicesIndependently");
+	case ETextureCreateFlags::Shared:
+		return TEXT("Shared");
+	case ETextureCreateFlags::NoFastClear:
+		return TEXT("NoFastClear");
+	case ETextureCreateFlags::DepthStencilResolveTarget:
+		return TEXT("DepthStencilResolveTarget");
+	case ETextureCreateFlags::Streamable:
+		return TEXT("Streamable");
+	case ETextureCreateFlags::NoFastClearFinalize:
+		return TEXT("NoFastClearFinalize");
+	case ETextureCreateFlags::AFRManual:
+		return TEXT("AFRManual");
+	case ETextureCreateFlags::ReduceMemoryWithTilingMode:
+		return TEXT("ReduceMemoryWithTilingMode");
+	}
+	return TEXT("");
+}
+
+const TCHAR* GetBufferUsageFlagString(EBufferUsageFlags BufferUsage)
+{
+	switch (BufferUsage)
+	{
+	case EBufferUsageFlags::None:
+		return TEXT("None");
+	case EBufferUsageFlags::Static:
+		return TEXT("Static");
+	case EBufferUsageFlags::Dynamic:
+		return TEXT("Dynamic");
+	case EBufferUsageFlags::Volatile:
+		return TEXT("Volatile");
+	case EBufferUsageFlags::UnorderedAccess:
+		return TEXT("UnorderedAccess");
+	case EBufferUsageFlags::ByteAddressBuffer:
+		return TEXT("ByteAddressBuffer");
+	case EBufferUsageFlags::SourceCopy:
+		return TEXT("SourceCopy");
+	case EBufferUsageFlags::StreamOutput:
+		return TEXT("StreamOutput");
+	case EBufferUsageFlags::DrawIndirect:
+		return TEXT("DrawIndirect");
+	case EBufferUsageFlags::ShaderResource:
+		return TEXT("ShaderResource");
+	case EBufferUsageFlags::KeepCPUAccessible:
+		return TEXT("KeepCPUAccessible");
+	case EBufferUsageFlags::FastVRAM:
+		return TEXT("FastVRAM");
+	case EBufferUsageFlags::Shared:
+		return TEXT("Shared");
+	case EBufferUsageFlags::AccelerationStructure:
+		return TEXT("AccelerationStructure");
+	case EBufferUsageFlags::VertexBuffer:
+		return TEXT("VertexBuffer");
+	case EBufferUsageFlags::IndexBuffer:
+		return TEXT("IndexBuffer");
+	case EBufferUsageFlags::StructuredBuffer:
+		return TEXT("StructuredBuffer");
+	}
+	return TEXT("");
+}
+
+const TCHAR* GetUniformBufferBaseTypeString(EUniformBufferBaseType BaseType)
+{
+	switch (BaseType)
+	{
+	case UBMT_INVALID:
+		return TEXT("UBMT_INVALID");
+	case UBMT_BOOL:
+		return TEXT("UBMT_BOOL");
+	case UBMT_INT32:
+		return TEXT("UBMT_INT32");
+	case UBMT_UINT32:
+		return TEXT("UBMT_UINT32");
+	case UBMT_FLOAT32:
+		return TEXT("UBMT_FLOAT32");
+	case UBMT_TEXTURE:
+		return TEXT("UBMT_TEXTURE");
+	case UBMT_SRV:
+		return TEXT("UBMT_SRV");
+	case UBMT_UAV:
+		return TEXT("UBMT_UAV");
+	case UBMT_SAMPLER:
+		return TEXT("UBMT_SAMPLER");
+	case UBMT_RDG_TEXTURE:
+		return TEXT("UBMT_RDG_TEXTURE");
+	case UBMT_RDG_TEXTURE_ACCESS:
+		return TEXT("UBMT_RDG_TEXTURE_ACCESS");
+	case UBMT_RDG_TEXTURE_ACCESS_ARRAY:
+		return TEXT("UBMT_RDG_TEXTURE_ACCESS_ARRAY");
+	case UBMT_RDG_TEXTURE_SRV:
+		return TEXT("UBMT_RDG_TEXTURE_SRV");
+	case UBMT_RDG_TEXTURE_UAV:
+		return TEXT("UBMT_RDG_TEXTURE_UAV");
+	case UBMT_RDG_BUFFER_ACCESS:
+		return TEXT("UBMT_RDG_BUFFER_ACCESS");
+	case UBMT_RDG_BUFFER_ACCESS_ARRAY:
+		return TEXT("UBMT_RDG_BUFFER_ACCESS_ARRAY");
+	case UBMT_RDG_BUFFER_SRV:
+		return TEXT("UBMT_RDG_BUFFER_SRV");
+	case UBMT_RDG_BUFFER_UAV:
+		return TEXT("UBMT_RDG_BUFFER_UAV");
+	case UBMT_RDG_UNIFORM_BUFFER:
+		return TEXT("UBMT_RDG_UNIFORM_BUFFER");
+	case UBMT_NESTED_STRUCT:
+		return TEXT("UBMT_NESTED_STRUCT");
+	case UBMT_INCLUDED_STRUCT:
+		return TEXT("UBMT_INCLUDED_STRUCT");
+	case UBMT_REFERENCED_STRUCT:
+		return TEXT("UBMT_REFERENCED_STRUCT");
+	case UBMT_RENDER_TARGET_BINDING_SLOTS:
+		return TEXT("UBMT_RENDER_TARGET_BINDING_SLOTS");
+	}
+	return TEXT("");
+}
diff --git a/Engine/Source/Runtime/RHI/Private/RHITransition.cpp b/Engine/Source/Runtime/RHI/Private/RHITransition.cpp
new file mode 100644
--- /dev/null
+++ b/Engine/Source/Runtime/RHI/Private/RHITransition.cpp
@@ -0,0 +1,32 @@
+// Copyright Epic Games, Inc. All Rights Reserved.
+
+#include "RHITransition.h"
+#include "RHI.h"
+
+RHI_API uint64 GRHITransitionPrivateData_SizeInBytes = 0;
+RHI_API uint64 GRHITransitionPrivateData_AlignInBytes = 0;
+
+FRHIViewableResource* GetViewableResource(const FRHITransitionInfo& Info)
+{
+	switch (Info.Type)
+	{
+	case FRHITransitionInfo::EType::Buffer:
+	case FRHITransitionInfo::EType::Texture:
+		return Info.ViewableResource;
+
+	case FRHITransitionInfo::EType::UAV:
+		return Info.UAV ? Info.UAV->GetParentResource() : nullptr;
+	}
+
+	return nullptr;
+}
+
+void FRHITransition::Cleanup() const
+{
+	FRHITransition* Transition = const_cast<FRHITransition*>(this);
+	RHIReleaseTransition(Transition);
+
+	// Explicit destruction of the transition.
+	Transition->~FRHITransition();
+	FConcurrentLinearAllocator::Free(Transition);
+}
diff --git a/Engine/Source/Runtime/RHI/Private/RHIValidation.cpp b/Engine/Source/Runtime/RHI/Private/RHIValidation.cpp
--- a/Engine/Source/Runtime/RHI/Private/RHIValidation.cpp
+++ b/Engine/Source/Runtime/RHI/Private/RHIValidation.cpp
@@ -84,6 +84,105 @@ namespace RHIValidation
 			CreateDesc.InitialState,
 			CreateDesc.DebugName);
 	}
+
+	void FTextureResource::InitBarrierTracking(int32 InNumMips, int32 InNumArraySlices, EPixelFormat PixelFormat, ETextureCreateFlags Flags, ERHIAccess InResourceState, const TCHAR* InDebugName)
+	{
+		FResource* Resource = GetTrackerResource();
+		if (!Resource)
+			return;
+
+		int32 InNumPlanes = 1;
+
+		// @todo: htile tracking
+		if (IsStencilFormat(PixelFormat))
+		{
+			InNumPlanes = 2; // Depth + Stencil
+		}
+		else
+		{
+			InNumPlanes = 1; // Depth only
+		}
+
+		Resource->InitBarrierTracking(InNumMips, InNumArraySlices, InNumPlanes, InResourceState, InDebugName);
+	}
+
+	FResourceIdentity FTextureResource::GetViewIdentity(uint32 InMipIndex, uint32 InNumMips, uint32 InArraySlice, uint32 InNumArraySlices, uint32 InPlaneIndex, uint32 InNumPlanes)
+	{
+		FResource* Resource = GetTrackerResource();
+
+		checkSlow((InMipIndex + InNumMips) <= Resource->NumMips);
+		checkSlow((InArraySlice + InNumArraySlices) <= Resource->NumArraySlices);
+		checkSlow((InPlaneIndex + InNumPlanes) <= Resource->NumPlanes);
+
+		if (InNumMips == 0)
+		{
+			InNumMips = Resource->NumMips;
+		}
+		if (InNumArraySlices == 0)
+		{
+			InNumArraySlices = Resource->NumArraySlices;
+		}
+		if (InNumPlanes == 0)
+		{
+			InNumPlanes = Resource->NumPlanes;
+		}
+
+		FResourceIdentity Identity;
+		Identity.Resource = Resource;
+		Identity.SubresourceRange.MipIndex = InMipIndex;
+		Identity.SubresourceRange.NumMips = InNumMips;
+		Identity.SubresourceRange.ArraySlice = InArraySlice;
+		Identity.SubresourceRange.NumArraySlices = InNumArraySlices;
+		Identity.SubresourceRange.PlaneIndex = InPlaneIndex;
+		Identity.SubresourceRange.NumPlanes = InNumPlanes;
+		return Identity;
+	}
+
+	FResourceIdentity FTextureResource::GetTransitionIdentity(const FRHITransitionInfo& Info)
+	{
+		FResource* Resource = GetTrackerResource();
+
+		FResourceIdentity Identity;
+		Identity.Resource = Resource;
+
+		if (Info.IsAllMips())
+		{
+			Identity.SubresourceRange.MipIndex = 0;
+			Identity.SubresourceRange.NumMips = Resource->NumMips;
+		}
+		else
+		{
+			check(Info.MipIndex < uint32(Resource->NumMips));
+			Identity.SubresourceRange.MipIndex = Info.MipIndex;
+			Identity.SubresourceRange.NumMips = 1;
+		}
+
+		if (Info.IsAllArraySlices())
+		{
+			Identity.SubresourceRange.ArraySlice = 0;
+			Identity.SubresourceRange.NumArraySlices = Resource->NumArraySlices;
+		}
+		else
+		{
+			check(Info.ArraySlice < uint32(Resource->NumArraySlices));
+			Identity.SubresourceRange.ArraySlice = Info.ArraySlice;
+			Identity.SubresourceRange.NumArraySlices = 1;
+		}
+
+		if (Info.IsAllPlaneSlices())
+		{
+			Identity.SubresourceRange.PlaneIndex = 0;
+			Identity.SubresourceRange.NumPlanes = Resource->NumPlanes;
+		}
+		else
+		{
+			check(Info.PlaneSlice < uint32(Resource->NumPlanes));
+			Identity.SubresourceRange.PlaneIndex = Info.PlaneSlice;
+			Identity.SubresourceRange.NumPlanes = 1;
+		}
+
+		return Identity;
+	}
 }
 
 TSet<uint32> FValidationRHI::SeenFailureHashes;
@@ -1727,6 +1826,20 @@ namespace RHIValidation
 		return EReplayStatus::Normal;
 	}
 
+	void FTracker::AddOp(const RHIValidation::FOperation& Op)
+	{
+		if (GRHICommandList.Bypass() && CurrentList.Operations.Num() == 0)
+		{
+			auto& OpQueue = OpQueues[GetOpQueueIndex(Pipeline)];
+			if (!EnumHasAllFlags(Op.Replay(Pipeline, OpQueue.bAllowAllUAVsOverlap, OpQueue.Breadcrumbs), EReplayStatus::Waiting))
+			{
+				return;
+			}
+		}
+
+		CurrentList.Operations.Add(Op);
+	}
+
 	void FTracker::ReplayOpQueue(ERHIPipeline DstOpQueue, FOperationsList&& InOpsList)
 	{
 		int32 DstOpQueueIndex = GetOpQueueIndex(DstOpQueue);
diff --git a/Engine/Source/Runtime/RHI/Public/RHI.h b/Engine/Source/Runtime/RHI/Public/RHI.h
--- a/Engine/Source/Runtime/RHI/Public/RHI.h
+++ b/Engine/Source/Runtime/RHI/Public/RHI.h
@@ -16,6 +16,7 @@
 #include "Math/PerspectiveMatrix.h"
 #include "Math/ScaleMatrix.h"
 #include "Math/TranslationMatrix.h"
+#include "PixelFormat.h"
 #include "ProfilingDebugging/CsvProfilerConfig.h"
 #include "GpuProfilerTrace.h"
 #include "RHIShaderPlatform.h"
@@ -101,12 +102,6 @@ RHI_API bool IsRHIDeviceIntel();
 // to trigger GPU specific optimizations and fallbacks
 RHI_API bool IsRHIDeviceNVIDIA();
 
-// helper to convert GRHIVendorId into a printable string, or "Unknown" if unknown.
-RHI_API const TCHAR* RHIVendorIdToString();
-
-// helper to convert VendorId into a printable string, or "Unknown" if unknown.
-RHI_API const TCHAR* RHIVendorIdToString(EGpuVendorId VendorId);
-
 // helper to return the shader language version for Metal shader.
 RHI_API uint32 RHIGetMetalShaderLanguageVersion(const FStaticShaderPlatform Platform);
 
@@ -728,92 +723,6 @@ inline EShaderPlatform GetFeatureLevelShaderPlatform(const FStaticFeatureLevel I
 	return GShaderPlatformForFeatureLevel[InFeatureLevel];
 }
 
-/** Stringifies ERHIDescriptorHeapType */
-extern RHI_API const TCHAR* LexToString(ERHIDescriptorHeapType InHeapType);
-
-/** Finds a corresponding ERHIShadingPath::Type given an FName, or returns false if one could not be found. */
-extern RHI_API bool GetShadingPathFromName(FName Name, ERHIShadingPath::Type& OutShadingPath);
-
-/** Creates a string for the given shading path. */
-extern RHI_API void GetShadingPathName(ERHIShadingPath::Type InShadingPath, FString& OutName);
-
-/** Creates an FName for the given shading path. */
-extern RHI_API void GetShadingPathName(ERHIShadingPath::Type InShadingPath, FName& OutName);
-
-/** Returns a string of friendly name bits for the buffer usage flags enum. */
-extern RHI_API FString GetBufferUsageFlagsName(EBufferUsageFlags BufferUsage);
-
-/** Returns a string of friendly name bits for the texture create flags enum. */
-extern RHI_API FString GetTextureCreateFlagsName(ETextureCreateFlags TextureCreateFlags);
-
-/** Returns a string of friendly name bits for the texture create flags enum. */
-extern RHI_API const TCHAR* StringFromRHIResourceType(ERHIResourceType ResourceType);
-
-enum class ERHIPipeline : uint8
-{
-	Graphics = 1 << 0,
-	AsyncCompute = 1 << 1,
-
-	None = 0,
-	All = Graphics | AsyncCompute,
-	Num = 2
-};
-ENUM_CLASS_FLAGS(ERHIPipeline)
-
-inline constexpr uint32 GetRHIPipelineIndex(ERHIPipeline Pipeline)
-{
-	switch (Pipeline)
-	{
-	default:
-	case ERHIPipeline::Graphics:
-		return 0;
-	case ERHIPipeline::AsyncCompute:
-		return 1;
-	}
-}
-
-inline constexpr uint32 GetRHIPipelineCount()
-{
-	return uint32(ERHIPipeline::Num);
-}
-
-inline TArrayView<const ERHIPipeline> GetRHIPipelines()
-{
-	static const ERHIPipeline Pipelines[] = { ERHIPipeline::Graphics, ERHIPipeline::AsyncCompute };
-	return Pipelines;
-}
-
-template <typename FunctionType>
-inline void EnumerateRHIPipelines(ERHIPipeline PipelineMask, FunctionType Function)
-{
-	for (ERHIPipeline Pipeline : GetRHIPipelines())
-	{
-		if (EnumHasAnyFlags(PipelineMask, Pipeline))
-		{
-			Function(Pipeline);
-		}
-	}
-}
-
-/** Array of pass handles by RHI pipeline, with overloads to help with enum conversion. */
-template <typename ElementType>
-class TRHIPipelineArray : public TStaticArray<ElementType, GetRHIPipelineCount()>
-{
-	using Base = TStaticArray<ElementType, GetRHIPipelineCount()>;
-public:
-	using Base::Base;
-
-	FORCEINLINE ElementType& operator[](ERHIPipeline Pipeline)
-	{
-		return Base::operator[](GetRHIPipelineIndex(Pipeline));
-	}
-
-	FORCEINLINE const ElementType& operator[](ERHIPipeline Pipeline) const
-	{
-		return Base::operator[](GetRHIPipelineIndex(Pipeline));
-	}
-};
-
 /** to customize the RHIReadSurfaceData() output */
 class FReadSurfaceDataFlags
 {
@@ -1359,308 +1268,6 @@ struct FRHIResourceInfo
 	bool bResident{ true };
 };
 
-enum class EClearBinding
-{
-	ENoneBound, //no clear color associated with this target.  Target will not do hardware clears on most platforms
-	EColorBound, //target has a clear color bound.  Clears will use the bound color, and do hardware clears.
-	EDepthStencilBound, //target has a depthstencil value bound.  Clears will use the bound values and do hardware clears.
-};
-
-
-enum class EColorSpaceAndEOTF
-{
-	EUnknown = 0,
-
-	EColorSpace_Rec709  = 1,		// Color Space Uses Rec 709  Primaries
-	EColorSpace_Rec2020 = 2,		// Color Space Uses Rec 2020 Primaries
-	EColorSpace_DCIP3   = 3,		// Color Space Uses DCI-P3   Primaries
-	EEColorSpace_MASK   = 0xf,
-
-	EEOTF_Linear		= 1 << 4,   // Transfer Function Uses Linear Encoding
-	EEOTF_sRGB			= 2 << 4,	// Transfer Function Uses sRGB Encoding
-	EEOTF_PQ			= 3 << 4,	// Transfer Function Uses PQ Encoding
-	EEOTF_MASK			= 0xf << 4,
-
-	ERec709_sRGB		= EColorSpace_Rec709  | EEOTF_sRGB,
-	ERec709_Linear		= EColorSpace_Rec709  | EEOTF_Linear,
-	
-	ERec2020_PQ			= EColorSpace_Rec2020 | EEOTF_PQ,
-	ERec2020_Linear		= EColorSpace_Rec2020 | EEOTF_Linear,
-	
-	EDCIP3_PQ			= EColorSpace_DCIP3 | EEOTF_PQ,
-	EDCIP3_Linear		= EColorSpace_DCIP3 | EEOTF_Linear,
-	
-};
-
-
-struct FClearValueBinding
-{
-	struct DSVAlue
-	{
-		float Depth;
-		uint32 Stencil;
-	};
-
-	FClearValueBinding()
-		: ColorBinding(EClearBinding::EColorBound)
-	{
-		Value.Color[0] = 0.0f;
-		Value.Color[1] = 0.0f;
-		Value.Color[2] = 0.0f;
-		Value.Color[3] = 0.0f;
-	}
-
-	FClearValueBinding(EClearBinding NoBinding)
-		: ColorBinding(NoBinding)
-	{
-		check(ColorBinding == EClearBinding::ENoneBound);
-	}
-
-	explicit FClearValueBinding(const FLinearColor& InClearColor)
-		: ColorBinding(EClearBinding::EColorBound)
-	{
-		Value.Color[0] = InClearColor.R;
-		Value.Color[1] = InClearColor.G;
-		Value.Color[2] = InClearColor.B;
-		Value.Color[3] = InClearColor.A;
-	}
-
-	explicit FClearValueBinding(float DepthClearValue, uint32 StencilClearValue = 0)
-		: ColorBinding(EClearBinding::EDepthStencilBound)
-	{
-		Value.DSValue.Depth = DepthClearValue;
-		Value.DSValue.Stencil = StencilClearValue;
-	}
-
-	FLinearColor GetClearColor() const
-	{
-		ensure(ColorBinding == EClearBinding::EColorBound);
-		return FLinearColor(Value.Color[0], Value.Color[1], Value.Color[2], Value.Color[3]);
-	}
-
-	void GetDepthStencil(float& OutDepth, uint32& OutStencil) const
-	{
-		ensure(ColorBinding == EClearBinding::EDepthStencilBound);
-		OutDepth = Value.DSValue.Depth;
-		OutStencil = Value.DSValue.Stencil;
-	}
-
-	bool operator==(const FClearValueBinding& Other) const
-	{
-		if (ColorBinding == Other.ColorBinding)
-		{
-			if (ColorBinding == EClearBinding::EColorBound)
-			{
-				return
-					Value.Color[0] == Other.Value.Color[0] &&
-					Value.Color[1] == Other.Value.Color[1] &&
-					Value.Color[2] == Other.Value.Color[2] &&
-					Value.Color[3] == Other.Value.Color[3];
-
-			}
-			if (ColorBinding == EClearBinding::EDepthStencilBound)
-			{
-				return
-					Value.DSValue.Depth == Other.Value.DSValue.Depth &&
-					Value.DSValue.Stencil == Other.Value.DSValue.Stencil;
-			}
-			return true;
-		}
-		return false;
-	}
-
-	friend inline uint32 GetTypeHash(FClearValueBinding const& Binding)
-	{
-		uint32 Hash = GetTypeHash(Binding.ColorBinding);
-
-		if (Binding.ColorBinding == EClearBinding::EColorBound)
-		{
-			Hash = HashCombine(Hash, GetTypeHash(Binding.Value.Color[0]));
-			Hash = HashCombine(Hash, GetTypeHash(Binding.Value.Color[1]));
-			Hash = HashCombine(Hash, GetTypeHash(Binding.Value.Color[2]));
-			Hash = HashCombine(Hash, GetTypeHash(Binding.Value.Color[3]));
-		}
-		else if (Binding.ColorBinding == EClearBinding::EDepthStencilBound)
-		{
-			Hash = HashCombine(Hash, GetTypeHash(Binding.Value.DSValue.Depth  ));
-			Hash = HashCombine(Hash, GetTypeHash(Binding.Value.DSValue.Stencil));
-		}
-
-		return Hash;
-	}
-
-	EClearBinding ColorBinding;
-
-	union ClearValueType
-	{
-		float Color[4];
-		DSVAlue DSValue;
-	} Value;
-
-	// common clear values
-	static RHI_API const FClearValueBinding None;
-	static RHI_API const FClearValueBinding Black;
-	static RHI_API const FClearValueBinding BlackMaxAlpha;
-	static RHI_API const FClearValueBinding White;
-	static RHI_API const FClearValueBinding Transparent;
-	static RHI_API const FClearValueBinding DepthOne;
-	static RHI_API const FClearValueBinding DepthZero;
-	static RHI_API const FClearValueBinding DepthNear;
-	static RHI_API const FClearValueBinding DepthFar;	
-	static RHI_API const FClearValueBinding Green;
-	static RHI_API const FClearValueBinding DefaultNormal8Bit;
-};
-
-struct FRHIResourceCreateInfo
-{
-	FRHIResourceCreateInfo(const TCHAR* InDebugName)
-		: BulkData(nullptr)
-		, ResourceArray(nullptr)
-		, ClearValueBinding(FLinearColor::Transparent)
-		, GPUMask(FRHIGPUMask::All())
-		, bWithoutNativeResource(false)
-		, DebugName(InDebugName)
-		, ExtData(0)
-	{
-		check(InDebugName);
-	}
-
-	// for CreateTexture calls
-	FRHIResourceCreateInfo(const TCHAR* InDebugName, FResourceBulkDataInterface* InBulkData)
-		: FRHIResourceCreateInfo(InDebugName)
-	{
-		BulkData = InBulkData;
-	}
-
-	// for CreateBuffer calls
-	FRHIResourceCreateInfo(const TCHAR* InDebugName, FResourceArrayInterface* InResourceArray)
-		: FRHIResourceCreateInfo(InDebugName)
-	{
-		ResourceArray = InResourceArray;
-	}
-
-	FRHIResourceCreateInfo(const TCHAR* InDebugName, const FClearValueBinding& InClearValueBinding)
-		: FRHIResourceCreateInfo(InDebugName)
-	{
-		ClearValueBinding = InClearValueBinding;
-	}
-
-	FRHIResourceCreateInfo(uint32 InExtData)
-		: FRHIResourceCreateInfo(TEXT(""))
-	{
-		ExtData = InExtData;
-	}
-
-	// for CreateTexture calls
-	FResourceBulkDataInterface* BulkData;
-	// for CreateBuffer calls
-	FResourceArrayInterface* ResourceArray;
-
-	// for binding clear colors to render targets.
-	FClearValueBinding ClearValueBinding;
-
-	// set of GPUs on which to create the resource
-	FRHIGPUMask GPUMask;
-
-	// whether to create an RHI object with no underlying resource
-	bool bWithoutNativeResource;
-	const TCHAR* DebugName;
-
-	// optional data that would have come from an offline cooker or whatever - general purpose
-	uint32 ExtData;
-};
-
-struct FResolveRect
-{
-	int32 X1;
-	int32 Y1;
-	int32 X2;
-	int32 Y2;
-
-	// e.g. for a a full 256 x 256 area starting at (0, 0) it would be 
-	// the values would be 0, 0, 256, 256
-	FResolveRect(int32 InX1=-1, int32 InY1=-1, int32 InX2=-1, int32 InY2=-1)
-	:	X1(InX1)
-	,	Y1(InY1)
-	,	X2(InX2)
-	,	Y2(InY2)
-	{}
-
-	FResolveRect(const FResolveRect& Other)
-		: X1(Other.X1)
-		, Y1(Other.Y1)
-		, X2(Other.X2)
-		, Y2(Other.Y2)
-	{}
-
-	explicit FResolveRect(FIntRect Other)
-		: X1(Other.Min.X)
-		, Y1(Other.Min.Y)
-		, X2(Other.Max.X)
-		, Y2(Other.Max.Y)
-	{}
-
-	bool operator==(FResolveRect Other) const
-	{
-		return X1 == Other.X1 && Y1 == Other.Y1 && X2 == Other.X2 && Y2 == Other.Y2;
-	}
-
-	bool operator!=(FResolveRect Other) const
-	{
-		return !(*this == Other);
-	}
-
-	bool IsValid() const
-	{
-		return X1 >= 0 && Y1 >= 0 && X2 - X1 > 0 && Y2 - Y1 > 0;
-	}
-};
-
-struct FResolveParams
-{
-	/** used to specify face when resolving to a cube map texture */
-	ECubeFace CubeFace;
-	/** resolve RECT bounded by [X1,Y1]..[X2,Y2]. Or -1 for fullscreen */
-	FResolveRect Rect;
-	FResolveRect DestRect;
-	/** The mip index to resolve in both source and dest. */
-	int32 MipIndex;
-	/** Array index to resolve in the source. */
-	int32 SourceArrayIndex;
-	/** Array index to resolve in the dest. */
-	int32 DestArrayIndex;
-	/** States to transition to at the end of the resolve operation. */
-	ERHIAccess SourceAccessFinal = ERHIAccess::SRVMask;
-	ERHIAccess DestAccessFinal = ERHIAccess::SRVMask;
-
-	/** constructor */
-	FResolveParams(
-		const FResolveRect& InRect = FResolveRect(), 
-		ECubeFace InCubeFace = CubeFace_PosX,
-		int32 InMipIndex = 0,
-		int32 InSourceArrayIndex = 0,
-		int32 InDestArrayIndex = 0,
-		const FResolveRect& InDestRect = FResolveRect())
-		:	CubeFace(InCubeFace)
-		,	Rect(InRect)
-		,	DestRect(InDestRect)
-		,	MipIndex(InMipIndex)
-		,	SourceArrayIndex(InSourceArrayIndex)
-		,	DestArrayIndex(InDestArrayIndex)
-	{}
-
-	FORCEINLINE FResolveParams(const FResolveParams& Other)
-		: CubeFace(Other.CubeFace)
-		, Rect(Other.Rect)
-		, DestRect(Other.DestRect)
-		, MipIndex(Other.MipIndex)
-		, SourceArrayIndex(Other.SourceArrayIndex)
-		, DestArrayIndex(Other.DestArrayIndex)
-		, SourceAccessFinal(Other.SourceAccessFinal)
-		, DestAccessFinal(Other.DestAccessFinal)
-	{}
-};
-
 
 struct FRHICopyTextureInfo
 {
@@ -1691,502 +1298,6 @@ struct FRHICopyTextureInfo
 	uint32 NumMips = 1;
 };
 
-enum class ERHITransitionCreateFlags
-{
-	None = 0,
-
-	// Disables fencing between pipelines during the transition.
-	NoFence = 1 << 0,
-
-	// Indicates the transition will have no useful work between the Begin/End calls,
-	// so should use a partial flush rather than a fence as this is more optimal.
-	NoSplit = 1 << 1,
-
-	BeginSimpleMode
-};
-ENUM_CLASS_FLAGS(ERHITransitionCreateFlags);
-
-enum class EResourceTransitionFlags
-{
-	None                = 0,
-
-	MaintainCompression = 1 << 0, // Specifies that the transition should not decompress the resource, allowing us to read a compressed resource directly in its compressed state.
-	Discard				= 1 << 1, // Specifies that the data in the resource should be discarded during the transition - used for transient resource acquire when the resource will be fully overwritten
-	Clear				= 1 << 2, // Specifies that the data in the resource should be cleared during the transition - used for transient resource acquire when the resource might not be fully overwritten
-
-	Last = Clear,
-	Mask = (Last << 1) - 1
-};
-ENUM_CLASS_FLAGS(EResourceTransitionFlags);
-
-RHI_API FString GetResourceTransitionFlagsName(EResourceTransitionFlags Flags);
-RHI_API FString GetRHIPipelineName(ERHIPipeline Pipeline);
-
-// The size in bytes of the storage required by the platform RHI for each resource transition.
-extern RHI_API uint64 GRHITransitionPrivateData_SizeInBytes;
-extern RHI_API uint64 GRHITransitionPrivateData_AlignInBytes;
-
-struct FRHISubresourceRange
-{
-	static const uint32 kDepthPlaneSlice = 0;
-	static const uint32 kStencilPlaneSlice = 1;
-	static const uint32 kAllSubresources = TNumericLimits<uint32>::Max();
-
-	uint32 MipIndex = kAllSubresources;
-	uint32 ArraySlice = kAllSubresources;
-	uint32 PlaneSlice = kAllSubresources;
-
-	FRHISubresourceRange() = default;
-
-	FRHISubresourceRange(
-		uint32 InMipIndex,
-		uint32 InArraySlice,
-		uint32 InPlaneSlice)
-		: MipIndex(InMipIndex)
-		, ArraySlice(InArraySlice)
-		, PlaneSlice(InPlaneSlice)
-	{}
-
-	inline bool IsAllMips() const
-	{
-		return MipIndex == kAllSubresources;
-	}
-
-	inline bool IsAllArraySlices() const
-	{
-		return ArraySlice == kAllSubresources;
-	}
-
-	inline bool IsAllPlaneSlices() const
-	{
-		return PlaneSlice == kAllSubresources;
-	}
-
-	inline bool IsWholeResource() const
-	{
-		return IsAllMips() && IsAllArraySlices() && IsAllPlaneSlices();
-	}
-
-	inline bool IgnoreDepthPlane() const
-	{
-		return PlaneSlice == kStencilPlaneSlice;
-	}
-
-	inline bool IgnoreStencilPlane() const
-	{
-		return PlaneSlice == kDepthPlaneSlice;
-	}
-
-	inline bool operator == (FRHISubresourceRange const& RHS) const
-	{
-		return MipIndex == RHS.MipIndex
-			&& ArraySlice == RHS.ArraySlice
-			&& PlaneSlice == RHS.PlaneSlice;
-	}
-
-	inline bool operator != (FRHISubresourceRange const& RHS) const
-	{
-		return !(*this == RHS);
-	}
-};
-
-struct FRHITransitionInfo : public FRHISubresourceRange
-{
-	union
-	{
-		class FRHIResource* Resource = nullptr;
-		class FRHIViewableResource* ViewableResource;
-		class FRHITexture* Texture;
-		class FRHIBuffer* Buffer;
-		class FRHIUnorderedAccessView* UAV;
-		class FRHIRayTracingAccelerationStructure* BVH;
-	};
-
-	enum class EType : uint8
-	{
-		Unknown,
-		Texture,
-		Buffer,
-		UAV,
-		BVH,
-	} Type = EType::Unknown;
-
-	ERHIAccess AccessBefore = ERHIAccess::Unknown;
-	ERHIAccess AccessAfter = ERHIAccess::Unknown;
-	EResourceTransitionFlags Flags = EResourceTransitionFlags::None;
-
-	FRHITransitionInfo() = default;
-
-	FRHITransitionInfo(
-		class FRHITexture* InTexture,
-		ERHIAccess InPreviousState,
-		ERHIAccess InNewState,
-		EResourceTransitionFlags InFlags = EResourceTransitionFlags::None,
-		uint32 InMipIndex = kAllSubresources,
-		uint32 InArraySlice = kAllSubresources,
-		uint32 InPlaneSlice = kAllSubresources)
-		: FRHISubresourceRange(InMipIndex, InArraySlice, InPlaneSlice)
-		, Texture(InTexture)
-		, Type(EType::Texture)
-		, AccessBefore(InPreviousState)
-		, AccessAfter(InNewState)
-		, Flags(InFlags)
-	{}
-
-	FRHITransitionInfo(class FRHIUnorderedAccessView* InUAV, ERHIAccess InPreviousState, ERHIAccess InNewState, EResourceTransitionFlags InFlags = EResourceTransitionFlags::None)
-		: UAV(InUAV)
-		, Type(EType::UAV)
-		, AccessBefore(InPreviousState)
-		, AccessAfter(InNewState)
-		, Flags(InFlags)
-	{}
-
-	FRHITransitionInfo(class FRHIBuffer* InRHIBuffer, ERHIAccess InPreviousState, ERHIAccess InNewState, EResourceTransitionFlags InFlags = EResourceTransitionFlags::None)
-		: Buffer(InRHIBuffer)
-		, Type(EType::Buffer)
-		, AccessBefore(InPreviousState)
-		, AccessAfter(InNewState)
-		, Flags(InFlags)
-	{}
-
-	FRHITransitionInfo(class FRHIRayTracingAccelerationStructure* InBVH, ERHIAccess InPreviousState, ERHIAccess InNewState, EResourceTransitionFlags InFlags = EResourceTransitionFlags::None)
-		: BVH(InBVH)
-		, Type(EType::BVH)
-		, AccessBefore(InPreviousState)
-		, AccessAfter(InNewState)
-		, Flags(InFlags)
-	{}
-
-	FRHITransitionInfo(class FRHITexture* InTexture, ERHIAccess InNewState)
-		: Texture(InTexture)
-		, Type(EType::Texture)
-		, AccessAfter(InNewState)
-	{}
-
-	FRHITransitionInfo(class FRHIUnorderedAccessView* InUAV, ERHIAccess InNewState)
-		: UAV(InUAV)
-		, Type(EType::UAV)
-		, AccessAfter(InNewState)
-	{}
-
-	FRHITransitionInfo(class FRHIBuffer* InRHIBuffer, ERHIAccess InNewState)
-		: Buffer(InRHIBuffer)
-		, Type(EType::Buffer)
-		, AccessAfter(InNewState)
-	{}
-
-	inline bool operator == (FRHITransitionInfo const& RHS) const
-	{
-		return Resource == RHS.Resource
-			&& Type == RHS.Type
-			&& AccessBefore == RHS.AccessBefore
-			&& AccessAfter == RHS.AccessAfter
-			&& Flags == RHS.Flags
-			&& FRHISubresourceRange::operator==(RHS);
-	}
-
-	inline bool operator != (FRHITransitionInfo const& RHS) const
-	{
-		return !(*this == RHS);
-	}
-};
-
-RHI_API FRHIViewableResource* GetViewableResource(const FRHITransitionInfo& Info);
-
-struct FRHITransientAliasingOverlap
-{
-	union
-	{
-		class FRHIResource* Resource = nullptr;
-		class FRHITexture* Texture;
-		class FRHIBuffer* Buffer;
-	};
-
-	enum class EType : uint8
-	{
-		Texture,
-		Buffer
-	} Type = EType::Texture;
-
-	FRHITransientAliasingOverlap() = default;
-
-	FRHITransientAliasingOverlap(FRHIResource* InResource, EType InType)
-		: Resource(InResource)
-		, Type(InType)
-	{}
-
-	FRHITransientAliasingOverlap(FRHITexture* InTexture)
-		: Texture(InTexture)
-		, Type(EType::Texture)
-	{}
-
-	FRHITransientAliasingOverlap(FRHIBuffer* InBuffer)
-		: Buffer(InBuffer)
-		, Type(EType::Buffer)
-	{}
-
-	bool IsTexture() const
-	{
-		return Type == EType::Texture;
-	}
-
-	bool IsBuffer() const
-	{
-		return Type == EType::Buffer;
-	}
-
-	bool operator == (const FRHITransientAliasingOverlap& Other) const
-	{
-		return Resource == Other.Resource;
-	}
-
-	inline bool operator != (const FRHITransientAliasingOverlap& RHS) const
-	{
-		return !(*this == RHS);
-	}
-};
-
-struct FRHITransientAliasingInfo
-{
-	union
-	{
-		class FRHIResource* Resource = nullptr;
-		class FRHITexture* Texture;
-		class FRHIBuffer* Buffer;
-	};
-
-	// List of prior resource overlaps to use when acquiring. Must be empty for discard operations.
-	TArrayView<const FRHITransientAliasingOverlap> Overlaps;
-
-	enum class EType : uint8
-	{
-		Texture,
-		Buffer
-	} Type = EType::Texture;
-
-	enum class EAction : uint8
-	{
-		Acquire,
-		Discard
-	} Action = EAction::Acquire;
-
-	FRHITransientAliasingInfo() = default;
-
-	static FRHITransientAliasingInfo Acquire(class FRHITexture* Texture, TArrayView<const FRHITransientAliasingOverlap> InOverlaps)
-	{
-		FRHITransientAliasingInfo Info;
-		Info.Texture = Texture;
-		Info.Overlaps = InOverlaps;
-		Info.Type = EType::Texture;
-		Info.Action = EAction::Acquire;
-		return Info;
-	}
-
-	static FRHITransientAliasingInfo Acquire(class FRHIBuffer* Buffer, TArrayView<const FRHITransientAliasingOverlap> InOverlaps)
-	{
-		FRHITransientAliasingInfo Info;
-		Info.Buffer = Buffer;
-		Info.Overlaps = InOverlaps;
-		Info.Type = EType::Buffer;
-		Info.Action = EAction::Acquire;
-		return Info;
-	}
-
-	static FRHITransientAliasingInfo Discard(class FRHITexture* Texture)
-	{
-		FRHITransientAliasingInfo Info;
-		Info.Texture = Texture;
-		Info.Type = EType::Texture;
-		Info.Action = EAction::Discard;
-		return Info;
-	}
-
-	static FRHITransientAliasingInfo Discard(class FRHIBuffer* Buffer)
-	{
-		FRHITransientAliasingInfo Info;
-		Info.Buffer = Buffer;
-		Info.Type = EType::Buffer;
-		Info.Action = EAction::Discard;
-		return Info;
-	}
-
-	bool IsAcquire() const
-	{
-		return Action == EAction::Acquire;
-	}
-
-	bool IsDiscard() const
-	{
-		return Action == EAction::Discard;
-	}
-
-	bool IsTexture() const
-	{
-		return Type == EType::Texture;
-	}
-
-	bool IsBuffer() const
-	{
-		return Type == EType::Buffer;
-	}
-
-	inline bool operator == (const FRHITransientAliasingInfo& RHS) const
-	{
-		return Resource == RHS.Resource
-			&& Type == RHS.Type
-			&& Action == RHS.Action;
-	}
-
-	inline bool operator != (const FRHITransientAliasingInfo& RHS) const
-	{
-		return !(*this == RHS);
-	}
-};
-
-struct FRHITransitionCreateInfo
-{
-	FRHITransitionCreateInfo() = default;
-
-	FRHITransitionCreateInfo(
-		ERHIPipeline InSrcPipelines,
-		ERHIPipeline InDstPipelines,
-		ERHITransitionCreateFlags InFlags = ERHITransitionCreateFlags::None,
-		TArrayView<const FRHITransitionInfo> InTransitionInfos = {},
-		TArrayView<const FRHITransientAliasingInfo> InAliasingInfos = {})
-		: SrcPipelines(InSrcPipelines)
-		, DstPipelines(InDstPipelines)
-		, Flags(InFlags)
-		, TransitionInfos(InTransitionInfos)
-		, AliasingInfos(InAliasingInfos)
-	{}
-
-	ERHIPipeline SrcPipelines = ERHIPipeline::None;
-	ERHIPipeline DstPipelines = ERHIPipeline::None;
-	ERHITransitionCreateFlags Flags = ERHITransitionCreateFlags::None;
-	TArrayView<const FRHITransitionInfo> TransitionInfos;
-	TArrayView<const FRHITransientAliasingInfo> AliasingInfos;
-};
-
-struct FRHITrackedAccessInfo
-{
-	FRHITrackedAccessInfo() = default;
-
-	FRHITrackedAccessInfo(FRHIViewableResource* InResource, ERHIAccess InAccess)
-		: Resource(InResource)
-		, Access(InAccess)
-	{}
-
-	FRHIViewableResource* Resource = nullptr;
-	ERHIAccess Access = ERHIAccess::Unknown;
-};
-
-#include "RHIValidationCommon.h" // IWYU pragma: export
-
-// Opaque data structure used to represent a pending resource transition in the RHI.
-struct FRHITransition
-{
-public:
-	template <typename T>
-	inline T* GetPrivateData()
-	{
-		checkSlow(sizeof(T) == GRHITransitionPrivateData_SizeInBytes && GRHITransitionPrivateData_AlignInBytes != 0);
-		uintptr_t Addr = Align(uintptr_t(this + 1), GRHITransitionPrivateData_AlignInBytes);
-		checkSlow(Addr + GRHITransitionPrivateData_SizeInBytes - (uintptr_t)this == GetTotalAllocationSize());
-		return reinterpret_cast<T*>(Addr);
-	}
-
-	template <typename T>
-	inline const T* GetPrivateData() const
-	{
-		return const_cast<FRHITransition*>(this)->GetPrivateData<T>();
-	}
-
-private:
-	// Prevent copying and moving. Only pointers to these structures are allowed.
-	FRHITransition(const FRHITransition&) = delete;
-	FRHITransition(FRHITransition&&) = delete;
-
-	// Private constructor. Memory for transitions is allocated manually with extra space at the tail of the structure for RHI use.
-	FRHITransition(ERHIPipeline SrcPipelines, ERHIPipeline DstPipelines)
-		: State(int8(int32(SrcPipelines) | (int32(DstPipelines) << int32(ERHIPipeline::Num))))
-#if DO_CHECK || USING_CODE_ANALYSIS
-		, AllowedSrc(SrcPipelines)
-		, AllowedDst(DstPipelines)
-#endif
-	{}
-
-	~FRHITransition()
-	{}
-
-	// Give private access to specific functions/RHI commands that need to allocate or control transitions.
-	friend const FRHITransition* RHICreateTransition(const FRHITransitionCreateInfo&);
-	friend class FRHIComputeCommandList;
-	friend struct FRHICommandBeginTransitions;
-	friend struct FRHICommandEndTransitions;
-	friend struct FRHICommandResourceTransition;
-
-	static uint64 GetTotalAllocationSize()
-	{
-		// Allocate extra space at the end of this structure for private RHI use. This is determined by GRHITransitionPrivateData_SizeInBytes.
-		return Align(sizeof(FRHITransition), FMath::Max(GRHITransitionPrivateData_AlignInBytes, 1ull)) + GRHITransitionPrivateData_SizeInBytes;
-	}
-
-	static uint64 GetAlignment()
-	{
-		return FMath::Max((uint64)alignof(FRHITransition), GRHITransitionPrivateData_AlignInBytes);
-	}
-
-	inline void MarkBegin(ERHIPipeline Pipeline) const
-	{
-		checkf(EnumHasAllFlags(AllowedSrc, Pipeline), TEXT("Transition is being used on a source pipeline that it wasn't created for."));
-
-		int8 Mask = int8(Pipeline);
-		int8 PreviousValue = FPlatformAtomics::InterlockedAnd(&State, ~Mask);
-		checkf((PreviousValue & Mask) == Mask, TEXT("RHIBeginTransitions has been called twice on this transition for at least one pipeline."));
-
-		if (PreviousValue == Mask)
-		{
-			Cleanup();
-		}
-	}
-
-	inline void MarkEnd(ERHIPipeline Pipeline) const
-	{
-		checkf(EnumHasAllFlags(AllowedDst, Pipeline), TEXT("Transition is being used on a destination pipeline that it wasn't created for."));
-
-		int8 Mask = int8(int32(Pipeline) << int32(ERHIPipeline::Num));
-		int8 PreviousValue = FPlatformAtomics::InterlockedAnd(&State, ~Mask);
-		checkf((PreviousValue & Mask) == Mask, TEXT("RHIEndTransitions has been called twice on this transition for at least one pipeline."));
-
-		if (PreviousValue == Mask)
-		{
-			Cleanup();
-		}
-	}
-
-	inline void Cleanup() const;
-
-	mutable int8 State;
-	static_assert((int32(ERHIPipeline::Num) * 2) < (sizeof(State) * 8), "Not enough bits to hold pipeline state.");
-
-#if DO_CHECK || USING_CODE_ANALYSIS
-	mutable ERHIPipeline AllowedSrc;
-	mutable ERHIPipeline AllowedDst;
-#endif
-
-#if ENABLE_RHI_VALIDATION
-	friend class FValidationRHI;
-	friend class FValidationComputeContext;
-	friend class FValidationContext;
-
-	RHIValidation::FOperationsList PendingSignals;
-	RHIValidation::FOperationsList PendingWaits;
-	RHIValidation::FOperationsList PendingAliases;
-	RHIValidation::FOperationsList PendingAliasingOverlaps;
-	RHIValidation::FOperationsList PendingOperationsBegin;
-	RHIValidation::FOperationsList PendingOperationsEnd;
-#endif
-};
 
 /** specifies an update region for a texture */
 struct FUpdateTextureRegion2D
@@ -2264,7 +1375,7 @@ struct FUpdateTextureRegion3D
 
 struct FRHIBufferRange
 {
-	FRHIBuffer* Buffer{ nullptr };
+	class FRHIBuffer* Buffer{ nullptr };
 	uint64 Offset{ 0 };
 	uint64 Size{ 0 };
 };
@@ -2534,34 +1645,6 @@ extern RHI_API FRHIPanicEvent& RHIGetPanicDelegate();
 // RHI utility functions that depend on the RHI definitions.
 #include "RHIUtilities.h" // IWYU pragma: keep
 
-inline void FRHITransition::Cleanup() const
-{
-	FRHITransition* Transition = const_cast<FRHITransition*>(this);
-	RHIReleaseTransition(Transition);
-
-	// Explicit destruction of the transition.
-	Transition->~FRHITransition();
-	FConcurrentLinearAllocator::Free(Transition);
-}
-
-#if ENABLE_RHI_VALIDATION
-
-inline void RHIValidation::FTracker::AddOp(const RHIValidation::FOperation& Op)
-{
-	if (GRHICommandList.Bypass() && CurrentList.Operations.Num() == 0)
-	{
-		auto& OpQueue = OpQueues[GetOpQueueIndex(Pipeline)];
-		if (!EnumHasAllFlags(Op.Replay(Pipeline, OpQueue.bAllowAllUAVsOverlap, OpQueue.Breadcrumbs), EReplayStatus::Waiting))
-		{
-			return;
-		}
-	}
-
-	CurrentList.Operations.Add(Op);
-}
-
-#endif // ENABLE_RHI_VALIDATION
-
 // Return what the expected number of samplers will be supported by a feature level
 // Note that since the Feature Level is pretty orthogonal to the RHI/HW, this is not going to be perfect
 // If should only be used for a guess at the limit, the real limit will not be known until runtime
@@ -2572,3 +1655,7 @@ inline uint32 GetExpectedFeatureLevelMaxTextureSamplers(const FStaticFeatureLeve
 
 RHI_API ERHIBindlessConfiguration RHIGetBindlessResourcesConfiguration(EShaderPlatform Platform);
 RHI_API ERHIBindlessConfiguration RHIGetBindlessSamplersConfiguration(EShaderPlatform Platform);
+
+#if UE_ENABLE_INCLUDE_ORDER_DEPRECATED_IN_5_2
+#include "RHIStrings.h"
+#endif
diff --git a/Engine/Source/Runtime/RHI/Public/RHIAccess.h b/Engine/Source/Runtime/RHI/Public/RHIAccess.h
--- a/Engine/Source/Runtime/RHI/Public/RHIAccess.h
+++ b/Engine/Source/Runtime/RHI/Public/RHIAccess.h
@@ -115,11 +115,12 @@ inline constexpr bool IsValidAccess(ERHIAccess Access)
 	return !IsInvalidAccess(Access);
 }
 
-RHI_API FString GetRHIAccessName(ERHIAccess Access);
-
 /** Mask of states which are allowed to be considered for state merging. */
 extern RHI_API ERHIAccess GRHIMergeableAccessMask;
 
 /** Mask of states which are allowed to be considered for multi-pipeline state merging. This should be a subset of GRHIMergeableAccessMask. */
 extern RHI_API ERHIAccess GRHIMultiPipelineMergeableAccessMask;
 
+#if UE_ENABLE_INCLUDE_ORDER_DEPRECATED_IN_5_2
+#include "RHIStrings.h"
+#endif
diff --git a/Engine/Source/Runtime/RHI/Public/RHIDefinitions.h b/Engine/Source/Runtime/RHI/Public/RHIDefinitions.h
--- a/Engine/Source/Runtime/RHI/Public/RHIDefinitions.h
+++ b/Engine/Source/Runtime/RHI/Public/RHIDefinitions.h
@@ -7,17 +7,8 @@
 
 #pragma once
 
-#include "PixelFormat.h"
 #include "Serialization/MemoryLayout.h"
 
-#if UE_ENABLE_INCLUDE_ORDER_DEPRECATED_IN_5_2
-#include "CoreMinimal.h"
-#include "HAL/IConsoleManager.h"
-#include "RHIShaderPlatform.h"
-#include "RHIFeatureLevel.h"
-#include "RHIImmutableSamplerState.h"
-#endif
-
 /** Alignment of the shader parameters struct is required to be 16-byte boundaries. */
 #define SHADER_PARAMETER_STRUCT_ALIGNMENT 16
 
@@ -192,7 +183,7 @@ enum class ERHIZBuffer
 */
 namespace ERHIShadingPath
 {
-	enum Type
+	enum Type : int
 	{
 		Deferred,
 		Forward,
@@ -1141,6 +1132,58 @@ enum class ERHIBindlessConfiguration
 	RayTracingShaders,
 };
 
+enum class EColorSpaceAndEOTF
+{
+	EUnknown = 0,
+
+	EColorSpace_Rec709  = 1,		// Color Space Uses Rec 709  Primaries
+	EColorSpace_Rec2020 = 2,		// Color Space Uses Rec 2020 Primaries
+	EColorSpace_DCIP3   = 3,		// Color Space Uses DCI-P3   Primaries
+	EEColorSpace_MASK   = 0xf,
+
+	EEOTF_Linear		= 1 << 4,   // Transfer Function Uses Linear Encoding
+	EEOTF_sRGB			= 2 << 4,	// Transfer Function Uses sRGB Encoding
+	EEOTF_PQ			= 3 << 4,	// Transfer Function Uses PQ Encoding
+	EEOTF_MASK			= 0xf << 4,
+
+	ERec709_sRGB		= EColorSpace_Rec709  | EEOTF_sRGB,
+	ERec709_Linear		= EColorSpace_Rec709  | EEOTF_Linear,
+	
+	ERec2020_PQ			= EColorSpace_Rec2020 | EEOTF_PQ,
+	ERec2020_Linear		= EColorSpace_Rec2020 | EEOTF_Linear,
+	
+	EDCIP3_PQ			= EColorSpace_DCIP3 | EEOTF_PQ,
+	EDCIP3_Linear		= EColorSpace_DCIP3 | EEOTF_Linear,
+};
+
+enum class ERHITransitionCreateFlags
+{
+	None = 0,
+
+	// Disables fencing between pipelines during the transition.
+	NoFence = 1 << 0,
+
+	// Indicates the transition will have no useful work between the Begin/End calls,
+	// so should use a partial flush rather than a fence as this is more optimal.
+	NoSplit = 1 << 1,
+
+	BeginSimpleMode
+};
+ENUM_CLASS_FLAGS(ERHITransitionCreateFlags);
+
+enum class EResourceTransitionFlags
+{
+	None                = 0,
+
+	MaintainCompression = 1 << 0, // Specifies that the transition should not decompress the resource, allowing us to read a compressed resource directly in its compressed state.
+	Discard				= 1 << 1, // Specifies that the data in the resource should be discarded during the transition - used for transient resource acquire when the resource will be fully overwritten
+	Clear				= 1 << 2, // Specifies that the data in the resource should be cleared during the transition - used for transient resource acquire when the resource might not be fully overwritten
+
+	Last = Clear,
+	Mask = (Last << 1) - 1
+};
+ENUM_CLASS_FLAGS(EResourceTransitionFlags);
+
 /** Returns whether the shader parameter type references an RDG texture. */
 inline bool IsRDGTextureReferenceShaderParameterType(EUniformBufferBaseType BaseType)
 {
@@ -1235,222 +1278,6 @@ inline EGpuVendorId RHIConvertToGpuVendorId(uint32 VendorId)
 	return EGpuVendorId::Unknown;
 }
 
-inline const TCHAR* GetShaderFrequencyString(EShaderFrequency Frequency, bool bIncludePrefix = true)
-{
-	const TCHAR* String = TEXT("SF_NumFrequencies");
-	switch (Frequency)
-	{
-	case SF_Vertex:			String = TEXT("SF_Vertex"); break;
-	case SF_Mesh:			String = TEXT("SF_Mesh"); break;
-	case SF_Amplification:	String = TEXT("SF_Amplification"); break;
-	case SF_Geometry:		String = TEXT("SF_Geometry"); break;
-	case SF_Pixel:			String = TEXT("SF_Pixel"); break;
-	case SF_Compute:		String = TEXT("SF_Compute"); break;
-	case SF_RayGen:			String = TEXT("SF_RayGen"); break;
-	case SF_RayMiss:		String = TEXT("SF_RayMiss"); break;
-	case SF_RayHitGroup:	String = TEXT("SF_RayHitGroup"); break;
-	case SF_RayCallable:	String = TEXT("SF_RayCallable"); break;
-
-	default:
-		checkf(0, TEXT("Unknown ShaderFrequency %d"), (int32)Frequency);
-		break;
-	}
-
-	// Skip SF_
-	int32 Index = bIncludePrefix ? 0 : 3;
-	String += Index;
-	return String;
-};
-
-inline const TCHAR* GetTextureDimensionString(ETextureDimension Dimension)
-{
-	switch (Dimension)
-	{
-	case ETextureDimension::Texture2D:
-		return TEXT("Texture2D");
-	case ETextureDimension::Texture2DArray:
-		return TEXT("Texture2DArray");
-	case ETextureDimension::Texture3D:
-		return TEXT("Texture3D");
-	case ETextureDimension::TextureCube:
-		return TEXT("TextureCube");
-	case ETextureDimension::TextureCubeArray:
-		return TEXT("TextureCubeArray");
-	}
-	return TEXT("");
-}
-
-inline const TCHAR* GetTextureCreateFlagString(ETextureCreateFlags TextureCreateFlag)
-{
-	switch (TextureCreateFlag)
-	{
-	case ETextureCreateFlags::None:
-		return TEXT("None");
-	case ETextureCreateFlags::RenderTargetable:
-		return TEXT("RenderTargetable");
-	case ETextureCreateFlags::ResolveTargetable:
-		return TEXT("ResolveTargetable");
-	case ETextureCreateFlags::DepthStencilTargetable:
-		return TEXT("DepthStencilTargetable");
-	case ETextureCreateFlags::ShaderResource:
-		return TEXT("ShaderResource");
-	case ETextureCreateFlags::SRGB:
-		return TEXT("SRGB");
-	case ETextureCreateFlags::CPUWritable:
-		return TEXT("CPUWritable");
-	case ETextureCreateFlags::NoTiling:
-		return TEXT("NoTiling");
-	case ETextureCreateFlags::VideoDecode:
-		return TEXT("VideoDecode");
-	case ETextureCreateFlags::Dynamic:
-		return TEXT("Dynamic");
-	case ETextureCreateFlags::InputAttachmentRead:
-		return TEXT("InputAttachmentRead");
-	case ETextureCreateFlags::Foveation:
-		return TEXT("Foveation");
-	case ETextureCreateFlags::Tiling3D:
-		return TEXT("Tiling3D");
-	case ETextureCreateFlags::Memoryless:
-		return TEXT("Memoryless");
-	case ETextureCreateFlags::GenerateMipCapable:
-		return TEXT("GenerateMipCapable");
-	case ETextureCreateFlags::FastVRAMPartialAlloc:
-		return TEXT("FastVRAMPartialAlloc");
-	case ETextureCreateFlags::DisableSRVCreation:
-		return TEXT("DisableSRVCreation");
-	case ETextureCreateFlags::DisableDCC:
-		return TEXT("DisableDCC");
-	case ETextureCreateFlags::UAV:
-		return TEXT("UAV");
-	case ETextureCreateFlags::Presentable:
-		return TEXT("Presentable");
-	case ETextureCreateFlags::CPUReadback:
-		return TEXT("CPUReadback");
-	case ETextureCreateFlags::OfflineProcessed:
-		return TEXT("OfflineProcessed");
-	case ETextureCreateFlags::FastVRAM:
-		return TEXT("FastVRAM");
-	case ETextureCreateFlags::HideInVisualizeTexture:
-		return TEXT("HideInVisualizeTexture");
-	case ETextureCreateFlags::Virtual:
-		return TEXT("Virtual");
-	case ETextureCreateFlags::TargetArraySlicesIndependently:
-		return TEXT("TargetArraySlicesIndependently");
-	case ETextureCreateFlags::Shared:
-		return TEXT("Shared");
-	case ETextureCreateFlags::NoFastClear:
-		return TEXT("NoFastClear");
-	case ETextureCreateFlags::DepthStencilResolveTarget:
-		return TEXT("DepthStencilResolveTarget");
-	case ETextureCreateFlags::Streamable:
-		return TEXT("Streamable");
-	case ETextureCreateFlags::NoFastClearFinalize:
-		return TEXT("NoFastClearFinalize");
-	case ETextureCreateFlags::AFRManual:
-		return TEXT("AFRManual");
-	case ETextureCreateFlags::ReduceMemoryWithTilingMode:
-		return TEXT("ReduceMemoryWithTilingMode");
-	}
-	return TEXT("");
-}
-
-inline const TCHAR* GetBufferUsageFlagString(EBufferUsageFlags BufferUsage)
-{
-	switch (BufferUsage)
-	{
-	case EBufferUsageFlags::None:
-		return TEXT("None");
-	case EBufferUsageFlags::Static:
-		return TEXT("Static");
-	case EBufferUsageFlags::Dynamic:
-		return TEXT("Dynamic");
-	case EBufferUsageFlags::Volatile:
-		return TEXT("Volatile");
-	case EBufferUsageFlags::UnorderedAccess:
-		return TEXT("UnorderedAccess");
-	case EBufferUsageFlags::ByteAddressBuffer:
-		return TEXT("ByteAddressBuffer");
-	case EBufferUsageFlags::SourceCopy:
-		return TEXT("SourceCopy");
-	case EBufferUsageFlags::StreamOutput:
-		return TEXT("StreamOutput");
-	case EBufferUsageFlags::DrawIndirect:
-		return TEXT("DrawIndirect");
-	case EBufferUsageFlags::ShaderResource:
-		return TEXT("ShaderResource");
-	case EBufferUsageFlags::KeepCPUAccessible:
-		return TEXT("KeepCPUAccessible");
-	case EBufferUsageFlags::FastVRAM:
-		return TEXT("FastVRAM");
-	case EBufferUsageFlags::Shared:
-		return TEXT("Shared");
-	case EBufferUsageFlags::AccelerationStructure:
-		return TEXT("AccelerationStructure");
-	case EBufferUsageFlags::VertexBuffer:
-		return TEXT("VertexBuffer");
-	case EBufferUsageFlags::IndexBuffer:
-		return TEXT("IndexBuffer");
-	case EBufferUsageFlags::StructuredBuffer:
-		return TEXT("StructuredBuffer");
-	}
-	return TEXT("");
-}
-
-inline const TCHAR* GetUniformBufferBaseTypeString(EUniformBufferBaseType BaseType)
-{
-	switch (BaseType)
-	{
-	case UBMT_INVALID:
-		return TEXT("UBMT_INVALID");
-	case UBMT_BOOL:
-		return TEXT("UBMT_BOOL");
-	case UBMT_INT32:
-		return TEXT("UBMT_INT32");
-	case UBMT_UINT32:
-		return TEXT("UBMT_UINT32");
-	case UBMT_FLOAT32:
-		return TEXT("UBMT_FLOAT32");
-	case UBMT_TEXTURE:
-		return TEXT("UBMT_TEXTURE");
-	case UBMT_SRV:
-		return TEXT("UBMT_SRV");
-	case UBMT_UAV:
-		return TEXT("UBMT_UAV");
-	case UBMT_SAMPLER:
-		return TEXT("UBMT_SAMPLER");
-	case UBMT_RDG_TEXTURE:
-		return TEXT("UBMT_RDG_TEXTURE");
-	case UBMT_RDG_TEXTURE_ACCESS:
-		return TEXT("UBMT_RDG_TEXTURE_ACCESS");
-	case UBMT_RDG_TEXTURE_ACCESS_ARRAY:
-		return TEXT("UBMT_RDG_TEXTURE_ACCESS_ARRAY");
-	case UBMT_RDG_TEXTURE_SRV:
-		return TEXT("UBMT_RDG_TEXTURE_SRV");
-	case UBMT_RDG_TEXTURE_UAV:
-		return TEXT("UBMT_RDG_TEXTURE_UAV");
-	case UBMT_RDG_BUFFER_ACCESS:
-		return TEXT("UBMT_RDG_BUFFER_ACCESS");
-	case UBMT_RDG_BUFFER_ACCESS_ARRAY:
-		return TEXT("UBMT_RDG_BUFFER_ACCESS_ARRAY");
-	case UBMT_RDG_BUFFER_SRV:
-		return TEXT("UBMT_RDG_BUFFER_SRV");
-	case UBMT_RDG_BUFFER_UAV:
-		return TEXT("UBMT_RDG_BUFFER_UAV");
-	case UBMT_RDG_UNIFORM_BUFFER:
-		return TEXT("UBMT_RDG_UNIFORM_BUFFER");
-	case UBMT_NESTED_STRUCT:
-		return TEXT("UBMT_NESTED_STRUCT");
-	case UBMT_INCLUDED_STRUCT:
-		return TEXT("UBMT_INCLUDED_STRUCT");
-	case UBMT_REFERENCED_STRUCT:
-		return TEXT("UBMT_REFERENCED_STRUCT");
-	case UBMT_RENDER_TARGET_BINDING_SLOTS:
-		return TEXT("UBMT_RENDER_TARGET_BINDING_SLOTS");
-	}
-	return TEXT("");
-}
-
-
 inline bool IsGeometryPipelineShaderFrequency(EShaderFrequency Frequency)
 {
 	return Frequency == SF_Mesh || Frequency == SF_Amplification;
@@ -1495,5 +1322,12 @@ inline ERHIResourceType GetRHIResourceType(ETextureDimension Dimension)
 #endif
 
 #if UE_ENABLE_INCLUDE_ORDER_DEPRECATED_IN_5_2
+#include "CoreMinimal.h"
 #include "DataDrivenShaderPlatformInfo.h"
+#include "HAL/IConsoleManager.h"
+#include "PixelFormat.h"
+#include "RHIFeatureLevel.h"
+#include "RHIImmutableSamplerState.h"
+#include "RHIShaderPlatform.h"
+#include "RHIStrings.h"
 #endif
diff --git a/Engine/Source/Runtime/RHI/Public/RHIFeatureLevel.h b/Engine/Source/Runtime/RHI/Public/RHIFeatureLevel.h
--- a/Engine/Source/Runtime/RHI/Public/RHIFeatureLevel.h
+++ b/Engine/Source/Runtime/RHI/Public/RHIFeatureLevel.h
@@ -106,18 +106,9 @@ private:
 using FStaticFeatureLevel = FGenericStaticFeatureLevel;
 #endif
 
-
-/** Finds a corresponding ERHIFeatureLevel::Type given an FName, or returns false if one could not be found. */
-extern RHI_API bool GetFeatureLevelFromName(FName Name, ERHIFeatureLevel::Type& OutFeatureLevel);
-
-/** Creates a string for the given feature level. */
-extern RHI_API void GetFeatureLevelName(ERHIFeatureLevel::Type InFeatureLevel, FString& OutName);
-
-/** Creates an FName for the given feature level. */
-extern RHI_API void GetFeatureLevelName(ERHIFeatureLevel::Type InFeatureLevel, FName& OutName);
-
-/** Stringifies ERHIFeatureLevel */
-extern RHI_API FString LexToString(ERHIFeatureLevel::Type Level);
-
 // The maximum feature level available on this system
 extern RHI_API ERHIFeatureLevel::Type GMaxRHIFeatureLevel;
+
+#if UE_ENABLE_INCLUDE_ORDER_DEPRECATED_IN_5_2
+#include "RHIStrings.h"
+#endif
diff --git a/Engine/Source/Runtime/RHI/Public/RHIPipeline.h b/Engine/Source/Runtime/RHI/Public/RHIPipeline.h
new file mode 100644
--- /dev/null
+++ b/Engine/Source/Runtime/RHI/Public/RHIPipeline.h
@@ -0,0 +1,73 @@
+// Copyright Epic Games, Inc. All Rights Reserved.
+
+#pragma once
+
+#include "HAL/Platform.h"
+#include "Misc/EnumClassFlags.h"
+#include "Containers/ArrayView.h"
+#include "Containers/StaticArray.h"
+
+enum class ERHIPipeline : uint8
+{
+	Graphics = 1 << 0,
+	AsyncCompute = 1 << 1,
+
+	None = 0,
+	All = Graphics | AsyncCompute,
+	Num = 2
+};
+ENUM_CLASS_FLAGS(ERHIPipeline)
+
+inline constexpr uint32 GetRHIPipelineIndex(ERHIPipeline Pipeline)
+{
+	switch (Pipeline)
+	{
+	default:
+	case ERHIPipeline::Graphics:
+		return 0;
+	case ERHIPipeline::AsyncCompute:
+		return 1;
+	}
+}
+
+inline constexpr uint32 GetRHIPipelineCount()
+{
+	return uint32(ERHIPipeline::Num);
+}
+
+inline TArrayView<const ERHIPipeline> GetRHIPipelines()
+{
+	static const ERHIPipeline Pipelines[] = { ERHIPipeline::Graphics, ERHIPipeline::AsyncCompute };
+	return Pipelines;
+}
+
+template <typename FunctionType>
+inline void EnumerateRHIPipelines(ERHIPipeline PipelineMask, FunctionType Function)
+{
+	for (ERHIPipeline Pipeline : GetRHIPipelines())
+	{
+		if (EnumHasAnyFlags(PipelineMask, Pipeline))
+		{
+			Function(Pipeline);
+		}
+	}
+}
+
+/** Array of pass handles by RHI pipeline, with overloads to help with enum conversion. */
+template <typename ElementType>
+class TRHIPipelineArray : public TStaticArray<ElementType, GetRHIPipelineCount()>
+{
+	using Base = TStaticArray<ElementType, GetRHIPipelineCount()>;
+public:
+	using Base::Base;
+
+	FORCEINLINE ElementType& operator[](ERHIPipeline Pipeline)
+	{
+		return Base::operator[](GetRHIPipelineIndex(Pipeline));
+	}
+
+	FORCEINLINE const ElementType& operator[](ERHIPipeline Pipeline) const
+	{
+		return Base::operator[](GetRHIPipelineIndex(Pipeline));
+	}
+};
diff --git a/Engine/Source/Runtime/RHI/Public/RHIResources.h b/Engine/Source/Runtime/RHI/Public/RHIResources.h
--- a/Engine/Source/Runtime/RHI/Public/RHIResources.h
+++ b/Engine/Source/Runtime/RHI/Public/RHIResources.h
@@ -31,6 +31,8 @@
 #include "Misc/CoreDelegates.h"
 #include "RHIFwd.h"
 #include "RHIImmutableSamplerState.h"
+#include "RHITransition.h"
+#include "MultiGPU.h"
 
 class FRHIComputeCommandList;
 class FRHICommandListImmediate;
@@ -233,6 +235,194 @@ private:
 	};
 };
 
+enum class EClearBinding
+{
+	ENoneBound, //no clear color associated with this target.  Target will not do hardware clears on most platforms
+	EColorBound, //target has a clear color bound.  Clears will use the bound color, and do hardware clears.
+	EDepthStencilBound, //target has a depthstencil value bound.  Clears will use the bound values and do hardware clears.
+};
+
+struct FClearValueBinding
+{
+	struct DSVAlue
+	{
+		float Depth;
+		uint32 Stencil;
+	};
+
+	FClearValueBinding()
+		: ColorBinding(EClearBinding::EColorBound)
+	{
+		Value.Color[0] = 0.0f;
+		Value.Color[1] = 0.0f;
+		Value.Color[2] = 0.0f;
+		Value.Color[3] = 0.0f;
+	}
+
+	FClearValueBinding(EClearBinding NoBinding)
+		: ColorBinding(NoBinding)
+	{
+		check(ColorBinding == EClearBinding::ENoneBound);
+	}
+
+	explicit FClearValueBinding(const FLinearColor& InClearColor)
+		: ColorBinding(EClearBinding::EColorBound)
+	{
+		Value.Color[0] = InClearColor.R;
+		Value.Color[1] = InClearColor.G;
+		Value.Color[2] = InClearColor.B;
+		Value.Color[3] = InClearColor.A;
+	}
+
+	explicit FClearValueBinding(float DepthClearValue, uint32 StencilClearValue = 0)
+		: ColorBinding(EClearBinding::EDepthStencilBound)
+	{
+		Value.DSValue.Depth = DepthClearValue;
+		Value.DSValue.Stencil = StencilClearValue;
+	}
+
+	FLinearColor GetClearColor() const
+	{
+		ensure(ColorBinding == EClearBinding::EColorBound);
+		return FLinearColor(Value.Color[0], Value.Color[1], Value.Color[2], Value.Color[3]);
+	}
+
+	void GetDepthStencil(float& OutDepth, uint32& OutStencil) const
+	{
+		ensure(ColorBinding == EClearBinding::EDepthStencilBound);
+		OutDepth = Value.DSValue.Depth;
+		OutStencil = Value.DSValue.Stencil;
+	}
+
+	bool operator==(const FClearValueBinding& Other) const
+	{
+		if (ColorBinding == Other.ColorBinding)
+		{
+			if (ColorBinding == EClearBinding::EColorBound)
+			{
+				return
+					Value.Color[0] == Other.Value.Color[0] &&
+					Value.Color[1] == Other.Value.Color[1] &&
+					Value.Color[2] == Other.Value.Color[2] &&
+					Value.Color[3] == Other.Value.Color[3];
+
+			}
+			if (ColorBinding == EClearBinding::EDepthStencilBound)
+			{
+				return
+					Value.DSValue.Depth == Other.Value.DSValue.Depth &&
+					Value.DSValue.Stencil == Other.Value.DSValue.Stencil;
+			}
+			return true;
+		}
+		return false;
+	}
+
+	friend inline uint32 GetTypeHash(FClearValueBinding const& Binding)
+	{
+		uint32 Hash = GetTypeHash(Binding.ColorBinding);
+
+		if (Binding.ColorBinding == EClearBinding::EColorBound)
+		{
+			Hash = HashCombine(Hash, GetTypeHash(Binding.Value.Color[0]));
+			Hash = HashCombine(Hash, GetTypeHash(Binding.Value.Color[1]));
+			Hash = HashCombine(Hash, GetTypeHash(Binding.Value.Color[2]));
+			Hash = HashCombine(Hash, GetTypeHash(Binding.Value.Color[3]));
+		}
+		else if (Binding.ColorBinding == EClearBinding::EDepthStencilBound)
+		{
+			Hash = HashCombine(Hash, GetTypeHash(Binding.Value.DSValue.Depth  ));
+			Hash = HashCombine(Hash, GetTypeHash(Binding.Value.DSValue.Stencil));
+		}
+
+		return Hash;
+	}
+
+	EClearBinding ColorBinding;
+
+	union ClearValueType
+	{
+		float Color[4];
+		DSVAlue DSValue;
+	} Value;
+
+	// common clear values
+	static RHI_API const FClearValueBinding None;
+	static RHI_API const FClearValueBinding Black;
+	static RHI_API const FClearValueBinding BlackMaxAlpha;
+	static RHI_API const FClearValueBinding White;
+	static RHI_API const FClearValueBinding Transparent;
+	static RHI_API const FClearValueBinding DepthOne;
+	static RHI_API const FClearValueBinding DepthZero;
+	static RHI_API const FClearValueBinding DepthNear;
+	static RHI_API const FClearValueBinding DepthFar;	
+	static RHI_API const FClearValueBinding Green;
+	static RHI_API const FClearValueBinding DefaultNormal8Bit;
+};
+
+class FResourceBulkDataInterface;
+class FResourceArrayInterface;
+
+struct FRHIResourceCreateInfo
+{
+	FRHIResourceCreateInfo(const TCHAR* InDebugName)
+		: BulkData(nullptr)
+		, ResourceArray(nullptr)
+		, ClearValueBinding(FLinearColor::Transparent)
+		, GPUMask(FRHIGPUMask::All())
+		, bWithoutNativeResource(false)
+		, DebugName(InDebugName)
+		, ExtData(0)
+	{
+		check(InDebugName);
+	}
+
+	// for CreateTexture calls
+	FRHIResourceCreateInfo(const TCHAR* InDebugName, FResourceBulkDataInterface* InBulkData)
+		: FRHIResourceCreateInfo(InDebugName)
+	{
+		BulkData = InBulkData;
+	}
+
+	// for CreateBuffer calls
+	FRHIResourceCreateInfo(const TCHAR* InDebugName, FResourceArrayInterface* InResourceArray)
+		: FRHIResourceCreateInfo(InDebugName)
+	{
+		ResourceArray = InResourceArray;
+	}
+
+	FRHIResourceCreateInfo(const TCHAR* InDebugName, const FClearValueBinding& InClearValueBinding)
+		: FRHIResourceCreateInfo(InDebugName)
+	{
+		ClearValueBinding = InClearValueBinding;
+	}
+
+	FRHIResourceCreateInfo(uint32 InExtData)
+		: FRHIResourceCreateInfo(TEXT(""))
+	{
+		ExtData = InExtData;
+	}
+
+	// for CreateTexture calls
+	FResourceBulkDataInterface* BulkData;
+
+	// for CreateBuffer calls
+	FResourceArrayInterface* ResourceArray;
+
+	// for binding clear colors to render targets.
+	FClearValueBinding ClearValueBinding;
+
+	// set of GPUs on which to create the resource
+	FRHIGPUMask GPUMask;
+
+	// whether to create an RHI object with no underlying resource
+	bool bWithoutNativeResource;
+	const TCHAR* DebugName;
+
+	// optional data that would have come from an offline cooker or whatever - general purpose
+	uint32 ExtData;
+};
+
 class FExclusiveDepthStencil
 {
 public:
@@ -3681,6 +3871,97 @@ inline ERenderTargetActions GetStencilActions(EDepthStencilTargetActions Action)
 	return (ERenderTargetActions)((uint8)Action & ((1 << (uint8)EDepthStencilTargetActions::DepthMask) - 1));
 }
 
+struct FResolveRect
+{
+	int32 X1;
+	int32 Y1;
+	int32 X2;
+	int32 Y2;
+
+	// e.g. for a a full 256 x 256 area starting at (0, 0) it would be 
+	// the values would be 0, 0, 256, 256
+	FResolveRect(int32 InX1 = -1, int32 InY1 = -1, int32 InX2 = -1, int32 InY2 = -1)
+		: X1(InX1)
+		, Y1(InY1)
+		, X2(InX2)
+		, Y2(InY2)
+	{}
+
+	FResolveRect(const FResolveRect& Other)
+		: X1(Other.X1)
+		, Y1(Other.Y1)
+		, X2(Other.X2)
+		, Y2(Other.Y2)
+	{}
+
+	explicit FResolveRect(FIntRect Other)
+		: X1(Other.Min.X)
+		, Y1(Other.Min.Y)
+		, X2(Other.Max.X)
+		, Y2(Other.Max.Y)
+	{}
+
+	bool operator==(FResolveRect Other) const
+	{
+		return X1 == Other.X1 && Y1 == Other.Y1 && X2 == Other.X2 && Y2 == Other.Y2;
+	}
+
+	bool operator!=(FResolveRect Other) const
+	{
+		return !(*this == Other);
+	}
+
+	bool IsValid() const
+	{
+		return X1 >= 0 && Y1 >= 0 && X2 - X1 > 0 && Y2 - Y1 > 0;
+	}
+};
+
+struct FResolveParams
+{
+	/** used to specify face when resolving to a cube map texture */
+	ECubeFace CubeFace;
+	/** resolve RECT bounded by [X1,Y1]..[X2,Y2]. Or -1 for fullscreen */
+	FResolveRect Rect;
+	FResolveRect DestRect;
+	/** The mip index to resolve in both source and dest. */
+	int32 MipIndex;
+	/** Array index to resolve in the source. */
+	int32 SourceArrayIndex;
+	/** Array index to resolve in the dest. */
+	int32 DestArrayIndex;
+	/** States to transition to at the end of the resolve operation. */
+	ERHIAccess SourceAccessFinal = ERHIAccess::SRVMask;
+	ERHIAccess DestAccessFinal = ERHIAccess::SRVMask;
+
+	/** constructor */
+	FResolveParams(
+		const FResolveRect& InRect = FResolveRect(),
+		ECubeFace InCubeFace = CubeFace_PosX,
+		int32 InMipIndex = 0,
+		int32 InSourceArrayIndex = 0,
+		int32 InDestArrayIndex = 0,
+		const FResolveRect& InDestRect = FResolveRect())
+		: CubeFace(InCubeFace)
+		, Rect(InRect)
+		, DestRect(InDestRect)
+		, MipIndex(InMipIndex)
+		, SourceArrayIndex(InSourceArrayIndex)
+		, DestArrayIndex(InDestArrayIndex)
+	{}
+
+	FORCEINLINE FResolveParams(const FResolveParams& Other)
+		: CubeFace(Other.CubeFace)
+		, Rect(Other.Rect)
+		, DestRect(Other.DestRect)
+		, MipIndex(Other.MipIndex)
+		, SourceArrayIndex(Other.SourceArrayIndex)
+		, DestArrayIndex(Other.DestArrayIndex)
+		, SourceAccessFinal(Other.SourceAccessFinal)
+		, DestAccessFinal(Other.DestAccessFinal)
+	{}
+};
+
 struct FRHIRenderPassInfo
 {
 	struct FColorEntry
diff --git a/Engine/Source/Runtime/RHI/Public/RHIShaderPlatform.h b/Engine/Source/Runtime/RHI/Public/RHIShaderPlatform.h
--- a/Engine/Source/Runtime/RHI/Public/RHIShaderPlatform.h
+++ b/Engine/Source/Runtime/RHI/Public/RHIShaderPlatform.h
@@ -161,13 +161,6 @@ private:
 	}
 };
 
-RHI_API FName LegacyShaderPlatformToShaderFormat(EShaderPlatform Platform);
-RHI_API EShaderPlatform ShaderFormatToLegacyShaderPlatform(FName ShaderFormat);
-RHI_API FName ShaderPlatformToPlatformName(EShaderPlatform Platform);
-
-/** Stringifies EShaderPlatform */
-extern RHI_API FString LexToString(EShaderPlatform Platform, bool bError = true);
-
 // The maximum shader platform available on this system
 extern RHI_API EShaderPlatform GMaxRHIShaderPlatform;
 
@@ -176,3 +169,6 @@ inline bool IsCustomPlatform(const FStaticShaderPlatform Platform)
 	return (Platform >= SP_CUSTOM_PLATFORM_FIRST && Platform < SP_CUSTOM_PLATFORM_LAST);
 }
 
+#if UE_ENABLE_INCLUDE_ORDER_DEPRECATED_IN_5_2
+#include "RHIStrings.h"
+#endif
diff --git a/Engine/Source/Runtime/RHI/Public/RHIStrings.h b/Engine/Source/Runtime/RHI/Public/RHIStrings.h
new file mode 100644
--- /dev/null
+++ b/Engine/Source/Runtime/RHI/Public/RHIStrings.h
@@ -0,0 +1,103 @@
+// Copyright Epic Games, Inc. All Rights Reserved.
+
+#pragma once
+
+#include "HAL/Platform.h"
+#include "RHIDefinitions.h"
+
+class FName;
+class FString;
+
+enum class ERHIAccess : uint32;
+enum class ERHIPipeline : uint8;
+enum EShaderPlatform : uint16;
+namespace ERHIFeatureLevel { enum Type : int; }
+
+// helper to convert GRHIVendorId into a printable string, or "Unknown" if unknown.
+RHI_API const TCHAR* RHIVendorIdToString();
+
+// helper to convert VendorId into a printable string, or "Unknown" if unknown.
+RHI_API const TCHAR* RHIVendorIdToString(EGpuVendorId VendorId);
+
+/** Finds a corresponding ERHIFeatureLevel::Type given a string, or returns false if one could not be found. */
+RHI_API bool GetFeatureLevelFromName(const FString& Name, ERHIFeatureLevel::Type& OutFeatureLevel);
+
+/** Finds a corresponding ERHIFeatureLevel::Type given an FName, or returns false if one could not be found. */
+RHI_API bool GetFeatureLevelFromName(FName Name, ERHIFeatureLevel::Type& OutFeatureLevel);
+
+/** Creates a string for the given feature level. */
+RHI_API void GetFeatureLevelName(ERHIFeatureLevel::Type InFeatureLevel, FString& OutName);
+
+/** Creates an FName for the given feature level. */
+RHI_API void GetFeatureLevelName(ERHIFeatureLevel::Type InFeatureLevel, FName& OutName);
+
+/** Stringifies ERHIFeatureLevel */
+RHI_API FString LexToString(ERHIFeatureLevel::Type Level);
+
+RHI_API FName LegacyShaderPlatformToShaderFormat(EShaderPlatform Platform);
+RHI_API EShaderPlatform ShaderFormatToLegacyShaderPlatform(FName ShaderFormat);
+RHI_API FName ShaderPlatformToPlatformName(EShaderPlatform Platform);
+
+/** Stringifies EShaderPlatform */
+RHI_API FString LexToString(EShaderPlatform Platform, bool bError = true);
+
+/** Stringifies ERHIDescriptorHeapType */
+RHI_API const TCHAR* LexToString(ERHIDescriptorHeapType InHeapType);
+
+/** Finds a corresponding ERHIShadingPath::Type given an FName, or returns false if one could not be found. */
+RHI_API bool GetShadingPathFromName(FName Name, ERHIShadingPath::Type& OutShadingPath);
+
+/** Creates a string for the given shading path. */
+RHI_API void GetShadingPathName(ERHIShadingPath::Type InShadingPath, FString& OutName);
+
+/** Creates an FName for the given shading path. */
+RHI_API void GetShadingPathName(ERHIShadingPath::Type InShadingPath, FName& OutName);
+
+/** Returns a string of friendly name bits for the buffer usage flags enum. */
+RHI_API FString GetBufferUsageFlagsName(EBufferUsageFlags BufferUsage);
+
+/** Returns a string of friendly name bits for the texture create flags enum. */
+RHI_API FString GetTextureCreateFlagsName(ETextureCreateFlags TextureCreateFlags);
+
+/** Returns a string of friendly name bits for the texture create flags enum. */
+RHI_API const TCHAR* StringFromRHIResourceType(ERHIResourceType ResourceType);
+
+RHI_API ERHIResourceType RHIResourceTypeFromString(const FString& InString);
+
+RHI_API FString GetRHIAccessName(ERHIAccess Access);
+
+RHI_API FString GetResourceTransitionFlagsName(EResourceTransitionFlags Flags);
+RHI_API FString GetRHIPipelineName(ERHIPipeline Pipeline);
+
+RHI_API const TCHAR* GetTextureDimensionString(ETextureDimension Dimension);
+RHI_API const TCHAR* GetTextureCreateFlagString(ETextureCreateFlags TextureCreateFlag);
+RHI_API const TCHAR* GetBufferUsageFlagString(EBufferUsageFlags BufferUsage);
+RHI_API const TCHAR* GetUniformBufferBaseTypeString(EUniformBufferBaseType BaseType);
+
+// Needs to stay inline for shader formats which can't depend on the RHI module
+inline const TCHAR* GetShaderFrequencyString(EShaderFrequency Frequency, bool bIncludePrefix = true)
+{
+	const TCHAR* String = TEXT("SF_NumFrequencies");
+	switch (Frequency)
+	{
+	case SF_Vertex:			String = TEXT("SF_Vertex"); break;
+	case SF_Mesh:			String = TEXT("SF_Mesh"); break;
+	case SF_Amplification:	String = TEXT("SF_Amplification"); break;
+	case SF_Geometry:		String = TEXT("SF_Geometry"); break;
+	case SF_Pixel:			String = TEXT("SF_Pixel"); break;
+	case SF_Compute:		String = TEXT("SF_Compute"); break;
+	case SF_RayGen:			String = TEXT("SF_RayGen"); break;
+	case SF_RayMiss:		String = TEXT("SF_RayMiss"); break;
+	case SF_RayHitGroup:	String = TEXT("SF_RayHitGroup"); break;
+	case SF_RayCallable:	String = TEXT("SF_RayCallable"); break;
+
+	default:
+		checkf(0, TEXT("Unknown ShaderFrequency %d"), (int32)Frequency);
+		break;
+	}
+
+	// Skip SF_
+	int32 Index = bIncludePrefix ? 0 : 3;
+	String += Index;
+	return String;
+}
diff --git a/Engine/Source/Runtime/RHI/Public/RHITransition.h b/Engine/Source/Runtime/RHI/Public/RHITransition.h
new file mode 100644
--- /dev/null
+++ b/Engine/Source/Runtime/RHI/Public/RHITransition.h
@@ -0,0 +1,468 @@
+// Copyright Epic Games, Inc. All Rights Reserved.
+
+#pragma once
+
+#include "Containers/ArrayView.h"
+#include "RHIDefinitions.h"
+#include "RHIAccess.h"
+#include "RHIPipeline.h"
+
+#include "RHIValidationCommon.h" // IWYU pragma: export
+
+// The size in bytes of the storage required by the platform RHI for each resource transition.
+extern RHI_API uint64 GRHITransitionPrivateData_SizeInBytes;
+extern RHI_API uint64 GRHITransitionPrivateData_AlignInBytes;
+
+struct FRHISubresourceRange
+{
+	static const uint32 kDepthPlaneSlice = 0;
+	static const uint32 kStencilPlaneSlice = 1;
+	static const uint32 kAllSubresources = TNumericLimits<uint32>::Max();
+
+	uint32 MipIndex = kAllSubresources;
+	uint32 ArraySlice = kAllSubresources;
+	uint32 PlaneSlice = kAllSubresources;
+
+	FRHISubresourceRange() = default;
+
+	FRHISubresourceRange(
+		uint32 InMipIndex,
+		uint32 InArraySlice,
+		uint32 InPlaneSlice)
+		: MipIndex(InMipIndex)
+		, ArraySlice(InArraySlice)
+		, PlaneSlice(InPlaneSlice)
+	{}
+
+	inline bool IsAllMips() const
+	{
+		return MipIndex == kAllSubresources;
+	}
+
+	inline bool IsAllArraySlices() const
+	{
+		return ArraySlice == kAllSubresources;
+	}
+
+	inline bool IsAllPlaneSlices() const
+	{
+		return PlaneSlice == kAllSubresources;
+	}
+
+	inline bool IsWholeResource() const
+	{
+		return IsAllMips() && IsAllArraySlices() && IsAllPlaneSlices();
+	}
+
+	inline bool IgnoreDepthPlane() const
+	{
+		return PlaneSlice == kStencilPlaneSlice;
+	}
+
+	inline bool IgnoreStencilPlane() const
+	{
+		return PlaneSlice == kDepthPlaneSlice;
+	}
+
+	inline bool operator == (FRHISubresourceRange const& RHS) const
+	{
+		return MipIndex == RHS.MipIndex
+			&& ArraySlice == RHS.ArraySlice
+			&& PlaneSlice == RHS.PlaneSlice;
+	}
+
+	inline bool operator != (FRHISubresourceRange const& RHS) const
+	{
+		return !(*this == RHS);
+	}
+};
+
+struct FRHITransitionInfo : public FRHISubresourceRange
+{
+	union
+	{
+		class FRHIResource* Resource = nullptr;
+		class FRHIViewableResource* ViewableResource;
+		class FRHITexture* Texture;
+		class FRHIBuffer* Buffer;
+		class FRHIUnorderedAccessView* UAV;
+		class FRHIRayTracingAccelerationStructure* BVH;
+	};
+
+	enum class EType : uint8
+	{
+		Unknown,
+		Texture,
+		Buffer,
+		UAV,
+		BVH,
+	} Type = EType::Unknown;
+
+	ERHIAccess AccessBefore = ERHIAccess::Unknown;
+	ERHIAccess AccessAfter = ERHIAccess::Unknown;
+	EResourceTransitionFlags Flags = EResourceTransitionFlags::None;
+
+	FRHITransitionInfo() = default;
+
+	FRHITransitionInfo(
+		class FRHITexture* InTexture,
+		ERHIAccess InPreviousState,
+		ERHIAccess InNewState,
+		EResourceTransitionFlags InFlags = EResourceTransitionFlags::None,
+		uint32 InMipIndex = kAllSubresources,
+		uint32 InArraySlice = kAllSubresources,
+		uint32 InPlaneSlice = kAllSubresources)
+		: FRHISubresourceRange(InMipIndex, InArraySlice, InPlaneSlice)
+		, Texture(InTexture)
+		, Type(EType::Texture)
+		, AccessBefore(InPreviousState)
+		, AccessAfter(InNewState)
+		, Flags(InFlags)
+	{}
+
+	FRHITransitionInfo(class FRHIUnorderedAccessView* InUAV, ERHIAccess InPreviousState, ERHIAccess InNewState, EResourceTransitionFlags InFlags = EResourceTransitionFlags::None)
+		: UAV(InUAV)
+		, Type(EType::UAV)
+		, AccessBefore(InPreviousState)
+		, AccessAfter(InNewState)
+		, Flags(InFlags)
+	{}
+
+	FRHITransitionInfo(class FRHIBuffer* InRHIBuffer, ERHIAccess InPreviousState, ERHIAccess InNewState, EResourceTransitionFlags InFlags = EResourceTransitionFlags::None)
+		: Buffer(InRHIBuffer)
+		, Type(EType::Buffer)
+		, AccessBefore(InPreviousState)
+		, AccessAfter(InNewState)
+		, Flags(InFlags)
+	{}
+
+	FRHITransitionInfo(class FRHIRayTracingAccelerationStructure* InBVH, ERHIAccess InPreviousState, ERHIAccess InNewState, EResourceTransitionFlags InFlags = EResourceTransitionFlags::None)
+		: BVH(InBVH)
+		, Type(EType::BVH)
+		, AccessBefore(InPreviousState)
+		, AccessAfter(InNewState)
+		, Flags(InFlags)
+	{}
+
+	FRHITransitionInfo(class FRHITexture* InTexture, ERHIAccess InNewState)
+		: Texture(InTexture)
+		, Type(EType::Texture)
+		, AccessAfter(InNewState)
+	{}
+
+	FRHITransitionInfo(class FRHIUnorderedAccessView* InUAV, ERHIAccess InNewState)
+		: UAV(InUAV)
+		, Type(EType::UAV)
+		, AccessAfter(InNewState)
+	{}
+
+	FRHITransitionInfo(class FRHIBuffer* InRHIBuffer, ERHIAccess InNewState)
+		: Buffer(InRHIBuffer)
+		, Type(EType::Buffer)
+		, AccessAfter(InNewState)
+	{}
+
+	inline bool operator == (FRHITransitionInfo const& RHS) const
+	{
+		return Resource == RHS.Resource
+			&& Type == RHS.Type
+			&& AccessBefore == RHS.AccessBefore
+			&& AccessAfter == RHS.AccessAfter
+			&& Flags == RHS.Flags
+			&& FRHISubresourceRange::operator==(RHS);
+	}
+
+	inline bool operator != (FRHITransitionInfo const& RHS) const
+	{
+		return !(*this == RHS);
+	}
+};
+
+
+struct FRHITransientAliasingOverlap
+{
+	union
+	{
+		class FRHIResource* Resource = nullptr;
+		class FRHITexture* Texture;
+		class FRHIBuffer* Buffer;
+	};
+
+	enum class EType : uint8
+	{
+		Texture,
+		Buffer
+	} Type = EType::Texture;
+
+	FRHITransientAliasingOverlap() = default;
+
+	FRHITransientAliasingOverlap(FRHIResource* InResource, EType InType)
+		: Resource(InResource)
+		, Type(InType)
+	{}
+
+	FRHITransientAliasingOverlap(FRHITexture* InTexture)
+		: Texture(InTexture)
+		, Type(EType::Texture)
+	{}
+
+	FRHITransientAliasingOverlap(FRHIBuffer* InBuffer)
+		: Buffer(InBuffer)
+		, Type(EType::Buffer)
+	{}
+
+	bool IsTexture() const
+	{
+		return Type == EType::Texture;
+	}
+
+	bool IsBuffer() const
+	{
+		return Type == EType::Buffer;
+	}
+
+	bool operator == (const FRHITransientAliasingOverlap& Other) const
+	{
+		return Resource == Other.Resource;
+	}
+
+	inline bool operator != (const FRHITransientAliasingOverlap& RHS) const
+	{
+		return !(*this == RHS);
+	}
+};
+
+struct FRHITransientAliasingInfo
+{
+	union
+	{
+		class FRHIResource* Resource = nullptr;
+		class FRHITexture* Texture;
+		class FRHIBuffer* Buffer;
+	};
+
+	// List of prior resource overlaps to use when acquiring. Must be empty for discard operations.
+	TArrayView<const FRHITransientAliasingOverlap> Overlaps;
+
+	enum class EType : uint8
+	{
+		Texture,
+		Buffer
+	} Type = EType::Texture;
+
+	enum class EAction : uint8
+	{
+		Acquire,
+		Discard
+	} Action = EAction::Acquire;
+
+	FRHITransientAliasingInfo() = default;
+
+	static FRHITransientAliasingInfo Acquire(class FRHITexture* Texture, TArrayView<const FRHITransientAliasingOverlap> InOverlaps)
+	{
+		FRHITransientAliasingInfo Info;
+		Info.Texture = Texture;
+		Info.Overlaps = InOverlaps;
+		Info.Type = EType::Texture;
+		Info.Action = EAction::Acquire;
+		return Info;
+	}
+
+	static FRHITransientAliasingInfo Acquire(class FRHIBuffer* Buffer, TArrayView<const FRHITransientAliasingOverlap> InOverlaps)
+	{
+		FRHITransientAliasingInfo Info;
+		Info.Buffer = Buffer;
+		Info.Overlaps = InOverlaps;
+		Info.Type = EType::Buffer;
+		Info.Action = EAction::Acquire;
+		return Info;
+	}
+
+	static FRHITransientAliasingInfo Discard(class FRHITexture* Texture)
+	{
+		FRHITransientAliasingInfo Info;
+		Info.Texture = Texture;
+		Info.Type = EType::Texture;
+		Info.Action = EAction::Discard;
+		return Info;
+	}
+
+	static FRHITransientAliasingInfo Discard(class FRHIBuffer* Buffer)
+	{
+		FRHITransientAliasingInfo Info;
+		Info.Buffer = Buffer;
+		Info.Type = EType::Buffer;
+		Info.Action = EAction::Discard;
+		return Info;
+	}
+
+	bool IsAcquire() const
+	{
+		return Action == EAction::Acquire;
+	}
+
+	bool IsDiscard() const
+	{
+		return Action == EAction::Discard;
+	}
+
+	bool IsTexture() const
+	{
+		return Type == EType::Texture;
+	}
+
+	bool IsBuffer() const
+	{
+		return Type == EType::Buffer;
+	}
+
+	inline bool operator == (const FRHITransientAliasingInfo& RHS) const
+	{
+		return Resource == RHS.Resource
+			&& Type == RHS.Type
+			&& Action == RHS.Action;
+	}
+
+	inline bool operator != (const FRHITransientAliasingInfo& RHS) const
+	{
+		return !(*this == RHS);
+	}
+};
+
+struct FRHITransitionCreateInfo
+{
+	FRHITransitionCreateInfo() = default;
+
+	FRHITransitionCreateInfo(
+		ERHIPipeline InSrcPipelines,
+		ERHIPipeline InDstPipelines,
+		ERHITransitionCreateFlags InFlags = ERHITransitionCreateFlags::None,
+		TArrayView<const FRHITransitionInfo> InTransitionInfos = {},
+		TArrayView<const FRHITransientAliasingInfo> InAliasingInfos = {})
+		: SrcPipelines(InSrcPipelines)
+		, DstPipelines(InDstPipelines)
+		, Flags(InFlags)
+		, TransitionInfos(InTransitionInfos)
+		, AliasingInfos(InAliasingInfos)
+	{}
+
+	ERHIPipeline SrcPipelines = ERHIPipeline::None;
+	ERHIPipeline DstPipelines = ERHIPipeline::None;
+	ERHITransitionCreateFlags Flags = ERHITransitionCreateFlags::None;
+	TArrayView<const FRHITransitionInfo> TransitionInfos;
+	TArrayView<const FRHITransientAliasingInfo> AliasingInfos;
+};
+
+struct FRHITrackedAccessInfo
+{
+	FRHITrackedAccessInfo() = default;
+
+	FRHITrackedAccessInfo(FRHIViewableResource* InResource, ERHIAccess InAccess)
+		: Resource(InResource)
+		, Access(InAccess)
+	{}
+
+	FRHIViewableResource* Resource = nullptr;
+	ERHIAccess Access = ERHIAccess::Unknown;
+};
+
+// Opaque data structure used to represent a pending resource transition in the RHI.
+struct FRHITransition
+#if ENABLE_RHI_VALIDATION
+	: public RHIValidation::FTransitionResource
+#endif
+{
+public:
+	template <typename T>
+	inline T* GetPrivateData()
+	{
+		checkSlow(sizeof(T) == GRHITransitionPrivateData_SizeInBytes && GRHITransitionPrivateData_AlignInBytes != 0);
+		uintptr_t Addr = Align(uintptr_t(this + 1), GRHITransitionPrivateData_AlignInBytes);
+		checkSlow(Addr + GRHITransitionPrivateData_SizeInBytes - (uintptr_t)this == GetTotalAllocationSize());
+		return reinterpret_cast<T*>(Addr);
+	}
+
+	template <typename T>
+	inline const T* GetPrivateData() const
+	{
+		return const_cast<FRHITransition*>(this)->GetPrivateData<T>();
+	}
+
+private:
+	// Prevent copying and moving. Only pointers to these structures are allowed.
+	FRHITransition(const FRHITransition&) = delete;
+	FRHITransition(FRHITransition&&) = delete;
+
+	// Private constructor. Memory for transitions is allocated manually with extra space at the tail of the structure for RHI use.
+	FRHITransition(ERHIPipeline SrcPipelines, ERHIPipeline DstPipelines)
+		: State(int8(int32(SrcPipelines) | (int32(DstPipelines) << int32(ERHIPipeline::Num))))
+#if DO_CHECK || USING_CODE_ANALYSIS
+		, AllowedSrc(SrcPipelines)
+		, AllowedDst(DstPipelines)
+#endif
+	{}
+
+	~FRHITransition()
+	{}
+
+	// Give private access to specific functions/RHI commands that need to allocate or control transitions.
+	friend const FRHITransition* RHICreateTransition(const FRHITransitionCreateInfo&);
+	friend class FRHIComputeCommandList;
+	friend struct FRHICommandBeginTransitions;
+	friend struct FRHICommandEndTransitions;
+	friend struct FRHICommandResourceTransition;
+
+	static uint64 GetTotalAllocationSize()
+	{
+		// Allocate extra space at the end of this structure for private RHI use. This is determined by GRHITransitionPrivateData_SizeInBytes.
+		return Align(sizeof(FRHITransition), FMath::Max(GRHITransitionPrivateData_AlignInBytes, 1ull)) + GRHITransitionPrivateData_SizeInBytes;
+	}
+
+	static uint64 GetAlignment()
+	{
+		return FMath::Max((uint64)alignof(FRHITransition), GRHITransitionPrivateData_AlignInBytes);
+	}
+
+	inline void MarkBegin(ERHIPipeline Pipeline) const
+	{
+		checkf(EnumHasAllFlags(AllowedSrc, Pipeline), TEXT("Transition is being used on a source pipeline that it wasn't created for."));
+
+		int8 Mask = int8(Pipeline);
+		int8 PreviousValue = FPlatformAtomics::InterlockedAnd(&State, ~Mask);
+		checkf((PreviousValue & Mask) == Mask, TEXT("RHIBeginTransitions has been called twice on this transition for at least one pipeline."));
+
+		if (PreviousValue == Mask)
+		{
+			Cleanup();
+		}
+	}
+
+	inline void MarkEnd(ERHIPipeline Pipeline) const
+	{
+		checkf(EnumHasAllFlags(AllowedDst, Pipeline), TEXT("Transition is being used on a destination pipeline that it wasn't created for."));
+
+		int8 Mask = int8(int32(Pipeline) << int32(ERHIPipeline::Num));
+		int8 PreviousValue = FPlatformAtomics::InterlockedAnd(&State, ~Mask);
+		checkf((PreviousValue & Mask) == Mask, TEXT("RHIEndTransitions has been called twice on this transition for at least one pipeline."));
+
+		if (PreviousValue == Mask)
+		{
+			Cleanup();
+		}
+	}
+
+	RHI_API void Cleanup() const;
+
+	mutable int8 State;
+	static_assert((int32(ERHIPipeline::Num) * 2) < (sizeof(State) * 8), "Not enough bits to hold pipeline state.");
+
+#if DO_CHECK || USING_CODE_ANALYSIS
+	mutable ERHIPipeline AllowedSrc;
+	mutable ERHIPipeline AllowedDst;
+#endif
+};
+
+
+RHI_API FRHIViewableResource* GetViewableResource(const FRHITransitionInfo& Info);
+
+
diff --git a/Engine/Source/Runtime/RHI/Public/RHIValidationCommon.h b/Engine/Source/Runtime/RHI/Public/RHIValidationCommon.h
--- a/Engine/Source/Runtime/RHI/Public/RHIValidationCommon.h
+++ b/Engine/Source/Runtime/RHI/Public/RHIValidationCommon.h
@@ -9,6 +9,9 @@
 // HEADER_UNIT_SKIP - Should always be included through RHI.h (this code needs cleanup.. this file is included in the middle of RHI.h)
 
 #include "Experimental/ConcurrentLinearAllocator.h"
+#include "PixelFormat.h"
+#include "RHIPipeline.h"
+#include "RHIStrings.h"
 
 #if ENABLE_RHI_VALIDATION
 extern RHI_API bool GRHIValidationEnabled;
@@ -19,11 +22,12 @@ const bool GRHIValidationEnabled = false;
 #if ENABLE_RHI_VALIDATION
 
 class FRHIUniformBuffer;
+class FValidationComputeContext;
+class FValidationContext;
 class FValidationRHI;
-struct FRHITextureCreateDesc;
 
-// Forward declaration of function defined in RHIUtilities.h
-static inline bool IsStencilFormat(EPixelFormat Format);
+struct FRHITextureCreateDesc;
+struct FRHITransitionInfo;
 
 namespace RHIValidation
 {
@@ -36,25 +40,6 @@ namespace RHIValidation
 		void ValidateSetShaderUniformBuffer(FRHIUniformBuffer* UniformBuffer);
 	};
 
-	static bool IsValidCopyFormat(EPixelFormat SourceFormat, EPixelFormat DestFormat)
-	{
-		if (SourceFormat == DestFormat)
-		{
-			return true;
-		}
-		// Acceptable conversions follow. Add more as required.
-		if (SourceFormat == PF_R32G32_UINT && (DestFormat == PF_DXT1 || DestFormat == PF_BC4))
-		{
-			return true;
-		}
-		if (SourceFormat == PF_R32G32B32A32_UINT && (DestFormat == PF_DXT3 || DestFormat == PF_DXT5 || DestFormat == PF_BC5 || DestFormat == PF_BC7))
-		{
-			return true;
-		}
-		// No valid conversion found
-		return false;
-	}
-
 	class  FTracker;
 	class  FResource;
 	class  FTextureResource;
@@ -398,7 +383,7 @@ namespace RHIValidation
 	public:
 	};
 
-	class FTextureResource
+	class RHI_API FTextureResource
 	{
 	private:
 		// Don't use inheritance here. Because FRHITextureReferences exist, we have to
@@ -407,7 +392,7 @@ namespace RHIValidation
 
 	public:
 		FTextureResource() = default;
-		RHI_API FTextureResource(FRHITextureCreateDesc const& CreateDesc);
+		FTextureResource(FRHITextureCreateDesc const& CreateDesc);
 
 		virtual ~FTextureResource() {}
 
@@ -422,104 +407,10 @@ namespace RHIValidation
 			return const_cast<FTextureResource*>(this)->GetTrackerResource()->IsBarrierTrackingInitialized();
 		}
 
-		inline void InitBarrierTracking(int32 InNumMips, int32 InNumArraySlices, EPixelFormat PixelFormat, ETextureCreateFlags Flags, ERHIAccess InResourceState, const TCHAR* InDebugName)
-		{
-			FResource* Resource = GetTrackerResource();
-			if (!Resource)
-				return;
-
-			int32 InNumPlanes = 1;
-
-			// @todo: htile tracking
-			if (IsStencilFormat(PixelFormat))
-			{
-				InNumPlanes = 2; // Depth + Stencil
-			}
-			else
-			{
-				InNumPlanes = 1; // Depth only
-			}
-
-			Resource->InitBarrierTracking(InNumMips, InNumArraySlices, InNumPlanes, InResourceState, InDebugName);
-		}
-
-		inline FResourceIdentity GetViewIdentity(uint32 InMipIndex, uint32 InNumMips, uint32 InArraySlice, uint32 InNumArraySlices, uint32 InPlaneIndex, uint32 InNumPlanes)
-		{
-			FResource* Resource = GetTrackerResource();
-
-			checkSlow((InMipIndex + InNumMips) <= Resource->NumMips);
-			checkSlow((InArraySlice + InNumArraySlices) <= Resource->NumArraySlices);
-			checkSlow((InPlaneIndex + InNumPlanes) <= Resource->NumPlanes);
-
-			if (InNumMips == 0)
-			{
-				InNumMips = Resource->NumMips;
-			}
-			if (InNumArraySlices == 0)
-			{
-				InNumArraySlices = Resource->NumArraySlices;
-			}
-			if (InNumPlanes == 0)
-			{
-				InNumPlanes = Resource->NumPlanes;
-			}
-
-			FResourceIdentity Identity;
-			Identity.Resource = Resource;
-			Identity.SubresourceRange.MipIndex = InMipIndex;
-			Identity.SubresourceRange.NumMips = InNumMips;
-			Identity.SubresourceRange.ArraySlice = InArraySlice;
-			Identity.SubresourceRange.NumArraySlices = InNumArraySlices;
-			Identity.SubresourceRange.PlaneIndex = InPlaneIndex;
-			Identity.SubresourceRange.NumPlanes = InNumPlanes;
-			return Identity;
-		}
-
-		inline FResourceIdentity GetTransitionIdentity(const FRHITransitionInfo& Info)
-		{
-			FResource* Resource = GetTrackerResource();
+		void InitBarrierTracking(int32 InNumMips, int32 InNumArraySlices, EPixelFormat PixelFormat, ETextureCreateFlags Flags, ERHIAccess InResourceState, const TCHAR* InDebugName);
 
-			FResourceIdentity Identity;
-			Identity.Resource = Resource;
-
-			if (Info.IsAllMips())
-			{
-				Identity.SubresourceRange.MipIndex = 0;
-				Identity.SubresourceRange.NumMips = Resource->NumMips;
-			}
-			else
-			{
-				check(Info.MipIndex < uint32(Resource->NumMips));
-				Identity.SubresourceRange.MipIndex = Info.MipIndex;
-				Identity.SubresourceRange.NumMips = 1;
-			}
-
-			if (Info.IsAllArraySlices())
-			{
-				Identity.SubresourceRange.ArraySlice = 0;
-				Identity.SubresourceRange.NumArraySlices = Resource->NumArraySlices;
-			}
-			else
-			{
-				check(Info.ArraySlice < uint32(Resource->NumArraySlices));
-				Identity.SubresourceRange.ArraySlice = Info.ArraySlice;
-				Identity.SubresourceRange.NumArraySlices = 1;
-			}
-
-			if (Info.IsAllPlaneSlices())
-			{
-				Identity.SubresourceRange.PlaneIndex = 0;
-				Identity.SubresourceRange.NumPlanes = Resource->NumPlanes;
-			}
-			else
-			{
-				check(Info.PlaneSlice < uint32(Resource->NumPlanes));
-				Identity.SubresourceRange.PlaneIndex = Info.PlaneSlice;
-				Identity.SubresourceRange.NumPlanes = 1;
-			}
-
-			return Identity;
-		}
+		FResourceIdentity GetViewIdentity(uint32 InMipIndex, uint32 InNumMips, uint32 InArraySlice, uint32 InNumArraySlices, uint32 InPlaneIndex, uint32 InNumPlanes);
+		FResourceIdentity GetTransitionIdentity(const FRHITransitionInfo& Info);
 
 		inline FResourceIdentity GetWholeResourceIdentity()
 		{
@@ -911,6 +802,16 @@ namespace RHIValidation
 		}
 	};
 
+	struct FTransitionResource
+	{
+		RHIValidation::FOperationsList PendingSignals;
+		RHIValidation::FOperationsList PendingWaits;
+		RHIValidation::FOperationsList PendingAliases;
+		RHIValidation::FOperationsList PendingAliasingOverlaps;
+		RHIValidation::FOperationsList PendingOperationsBegin;
+		RHIValidation::FOperationsList PendingOperationsEnd;
+	};
+
 	enum class EUAVMode
 	{
 		Graphics,
@@ -984,7 +885,7 @@ namespace RHIValidation
 			: Pipeline(InPipeline)
 		{}
 
-		inline void AddOp(const FOperation& Op);
+		RHI_API void AddOp(const FOperation& Op);
 
 		inline void AddOps(const FOperationsList::ListType& Ops)
 		{
diff --git a/Engine/Source/Runtime/RenderCore/Public/ShaderPlatformCachedIniValue.h b/Engine/Source/Runtime/RenderCore/Public/ShaderPlatformCachedIniValue.h
--- a/Engine/Source/Runtime/RenderCore/Public/ShaderPlatformCachedIniValue.h
+++ b/Engine/Source/Runtime/RenderCore/Public/ShaderPlatformCachedIniValue.h
@@ -5,6 +5,7 @@
 #include "CoreMinimal.h"
 #include "HAL/IConsoleManager.h"
 #include "RHIShaderPlatform.h"
+#include "RHIStrings.h"
 
 RENDERCORE_API EShaderPlatform GetEditorShaderPlatform(EShaderPlatform ShaderPlatform);
 
