Restoring code lost after being moved.

diff --git a/Engine/Source/Runtime/RHI/Private/RHIResources.cpp b/Engine/Source/Runtime/RHI/Private/RHIResources.cpp
--- a/Engine/Source/Runtime/RHI/Private/RHIResources.cpp
+++ b/Engine/Source/Runtime/RHI/Private/RHIResources.cpp
@@ -3,8 +3,7 @@
 #include "RHIResources.h"
 #include "RHI.h"
 
-std::atomic<TClosableMpscQueue<FRHIResource*>*> FRHIResource::PendingDeletes{ new TClosableMpscQueue<FRHIResource*>() };
-FHazardPointerCollection FRHIResource::PendingDeletesHPC;
+UE::TConsumeAllMpmcQueue<FRHIResource*> FRHIResource::PendingDeletes;
 FRHIResource* FRHIResource::CurrentlyDeleting = nullptr;
 
 bool FRHIResource::Bypass()
@@ -21,12 +20,10 @@ int32 FRHIResource::FlushPendingDeletes(FRHICommandListImmediate& RHICmdList)
 	check(IsInRenderingThread());
 
 	TArray<FRHIResource*> DeletedResources;
-	TClosableMpscQueue<FRHIResource*>* PendingDeletesPtr = PendingDeletes.exchange(new TClosableMpscQueue<FRHIResource*>());
-	PendingDeletesPtr->Close([&DeletedResources](FRHIResource* Resource)
+	PendingDeletes.ConsumeAllLifo([&DeletedResources](FRHIResource* Resource)
 	{
 		DeletedResources.Push(Resource);
 	});
-	PendingDeletesHPC.Delete(PendingDeletesPtr);
 
 	const int32 NumDeletes = DeletedResources.Num();
 
@@ -37,8 +34,9 @@ int32 FRHIResource::FlushPendingDeletes(FRHICommandListImmediate& RHICmdList)
 			GDynamicRHI->RHIPerFrameRHIFlushComplete();
 		}
 
-		for (FRHIResource* Resource : DeletedResources)
+		for (int32 i = DeletedResources.Num() - 1; i >= 0; i--)
 		{
+			FRHIResource* Resource = DeletedResources[i];
 			if (Resource->AtomicFlags.Deleteing())
 			{
 				FRHIResource::CurrentlyDeleting = Resource;
