Moving FRHITextureReference to its own file.

diff --git a/Engine/Source/Runtime/RHI/Private/RHITextureReference.cpp b/Engine/Source/Runtime/RHI/Private/RHITextureReference.cpp
new file mode 100644
--- /dev/null
+++ b/Engine/Source/Runtime/RHI/Private/RHITextureReference.cpp
@@ -0,0 +1,56 @@
+// Copyright Epic Games, Inc. All Rights Reserved.
+
+#include "RHITextureReference.h"
+#include "RHI.h"
+
+FRHITextureReference::FRHITextureReference()
+	: FRHITexture(RRT_TextureReference)
+{
+	check(DefaultTexture);
+	ReferencedTexture = DefaultTexture;
+}
+
+FRHITextureReference::~FRHITextureReference() = default;
+
+FRHITextureReference* FRHITextureReference::GetTextureReference()
+{
+	return this;
+}
+
+void* FRHITextureReference::GetNativeResource() const
+{
+	check(ReferencedTexture);
+	return ReferencedTexture->GetNativeResource();
+}
+
+void* FRHITextureReference::GetNativeShaderResourceView() const
+{
+	check(ReferencedTexture);
+	return ReferencedTexture->GetNativeShaderResourceView();
+}
+
+void* FRHITextureReference::GetTextureBaseRHI()
+{
+	check(ReferencedTexture);
+	return ReferencedTexture->GetTextureBaseRHI();
+}
+
+void FRHITextureReference::GetWriteMaskProperties(void*& OutData, uint32& OutSize)
+{
+	check(ReferencedTexture);
+	return ReferencedTexture->GetWriteMaskProperties(OutData, OutSize);
+}
+
+#if ENABLE_RHI_VALIDATION
+RHIValidation::FResource* FRHITextureReference::GetTrackerResource()
+{
+	check(ReferencedTexture);
+	return ReferencedTexture->GetTrackerResource();
+}
+#endif
+
+const FRHITextureDesc& FRHITextureReference::GetDesc() const
+{
+	check(ReferencedTexture);
+	return ReferencedTexture->GetDesc();
+}
diff --git a/Engine/Source/Runtime/RHI/Public/RHIResources.h b/Engine/Source/Runtime/RHI/Public/RHIResources.h
--- a/Engine/Source/Runtime/RHI/Public/RHIResources.h
+++ b/Engine/Source/Runtime/RHI/Public/RHIResources.h
@@ -36,6 +36,7 @@
 
 class FRHIComputeCommandList;
 class FRHICommandListImmediate;
+class FRHITextureReference;
 struct FClearValueBinding;
 struct FRHIResourceInfo;
 struct FGenerateMipsStruct;
@@ -2123,91 +2124,6 @@ private:
 	FLastRenderTimeContainer LastRenderTime;
 };
 
-class RHI_API FRHITextureReference final : public FRHITexture
-{
-public:
-	explicit FRHITextureReference()
-		: FRHITexture(RRT_TextureReference)
-	{
-		check(DefaultTexture);
-		ReferencedTexture = DefaultTexture;
-	}
-
-	virtual class FRHITextureReference* GetTextureReference() override
-	{
-		return this;
-	}
-
-	virtual FRHIDescriptorHandle GetDefaultBindlessHandle() const override 
-	{ 
-		check(ReferencedTexture);
-		return ReferencedTexture->GetDefaultBindlessHandle();
-	}
-
-	virtual void* GetNativeResource() const override 
-	{
-		check(ReferencedTexture);
-		return ReferencedTexture->GetNativeResource();
-	}
-
-	virtual void* GetNativeShaderResourceView() const override
-	{
-		check(ReferencedTexture);
-		return ReferencedTexture->GetNativeShaderResourceView();
-	}
-
-	virtual void* GetTextureBaseRHI() override 
-	{
-		check(ReferencedTexture);
-		return ReferencedTexture->GetTextureBaseRHI();
-	}
-
-	virtual void GetWriteMaskProperties(void*& OutData, uint32& OutSize) override
-	{
-		check(ReferencedTexture);
-		return ReferencedTexture->GetWriteMaskProperties(OutData, OutSize);
-	}
-
-#if ENABLE_RHI_VALIDATION
-	// Implement RHIValidation::FTextureResource::GetTrackerResource to use the tracker info
-	// for the referenced texture.
-	virtual RHIValidation::FResource* GetTrackerResource() final override
-	{
-		check(ReferencedTexture);
-		return ReferencedTexture->GetTrackerResource();
-	}
-#endif
-
-	inline FRHITexture* GetReferencedTexture() const
-	{
-		return ReferencedTexture.GetReference();
-	}
-
-	virtual const FRHITextureDesc& GetDesc() const override
-	{
-		check(ReferencedTexture);
-		return ReferencedTexture->GetDesc();
-	}
-
-private:
-	friend class FRHICommandListImmediate;
-	// Called only from FRHICommandListImmediate::UpdateTextureReference()
-	void SetReferencedTexture(FRHITexture* InTexture)
-	{
-		ReferencedTexture = InTexture
-			? InTexture
-			: DefaultTexture.GetReference();
-	}
-
-	TRefCountPtr<FRHITexture> ReferencedTexture;
-
-	// This pointer is set by the InitRHI() function on the FBlackTextureWithSRV global resource,
-	// to allow FRHITextureReference to use the global black texture when the reference is nullptr.
-	// A pointer is required since FBlackTextureWithSRV is defined in RenderCore.
-	friend class FBlackTextureWithSRV;
-	static TRefCountPtr<FRHITexture> DefaultTexture;
-};
-
 //
 // Misc
 //
@@ -4612,3 +4528,5 @@ struct FRHIShaderParameterResource
 	uint16        Index = 0;
 	EType         Type = EType::Texture;
 };
+
+#include "RHITextureReference.h"
diff --git a/Engine/Source/Runtime/RHI/Public/RHITextureReference.h b/Engine/Source/Runtime/RHI/Public/RHITextureReference.h
new file mode 100644
--- /dev/null
+++ b/Engine/Source/Runtime/RHI/Public/RHITextureReference.h
@@ -0,0 +1,50 @@
+// Copyright Epic Games, Inc. All Rights Reserved.
+
+#pragma once
+
+#include "RHIResources.h"
+
+class RHI_API FRHITextureReference final : public FRHITexture
+{
+public:
+	explicit FRHITextureReference();
+	~FRHITextureReference();
+
+	virtual class FRHITextureReference* GetTextureReference() override;
+
+	virtual void* GetNativeResource() const override;
+	virtual void* GetNativeShaderResourceView() const override;
+	virtual void* GetTextureBaseRHI() override;
+	virtual void GetWriteMaskProperties(void*& OutData, uint32& OutSize) override;
+	virtual const FRHITextureDesc& GetDesc() const override;
+
+#if ENABLE_RHI_VALIDATION
+	// Implement RHIValidation::FTextureResource::GetTrackerResource to use the tracker info
+	// for the referenced texture.
+	virtual RHIValidation::FResource* GetTrackerResource() final override;
+#endif
+
+	inline FRHITexture * GetReferencedTexture() const
+	{
+		return ReferencedTexture.GetReference();
+	}
+
+private:
+	friend class FRHICommandListImmediate;
+
+	// Called only from FRHICommandListImmediate::UpdateTextureReference()
+	void SetReferencedTexture(FRHITexture* InTexture)
+	{
+		ReferencedTexture = InTexture
+			? InTexture
+			: DefaultTexture.GetReference();
+	}
+
+	TRefCountPtr<FRHITexture> ReferencedTexture;
+
+	// This pointer is set by the InitRHI() function on the FBlackTextureWithSRV global resource,
+	// to allow FRHITextureReference to use the global black texture when the reference is nullptr.
+	// A pointer is required since FBlackTextureWithSRV is defined in RenderCore.
+	friend class FBlackTextureWithSRV;
+	static TRefCountPtr<FRHITexture> DefaultTexture;
+};
