Restoring code lost after previous merge.

diff --git a/Engine/Source/Runtime/RHI/Private/RHITextureReference.cpp b/Engine/Source/Runtime/RHI/Private/RHITextureReference.cpp
--- a/Engine/Source/Runtime/RHI/Private/RHITextureReference.cpp
+++ b/Engine/Source/Runtime/RHI/Private/RHITextureReference.cpp
@@ -17,6 +17,12 @@ FRHITextureReference* FRHITextureReference::GetTextureReference()
 	return this;
 }
 
+FRHIDescriptorHandle FRHITextureReference::GetDefaultBindlessHandle() const
+{ 
+	check(ReferencedTexture);
+	return ReferencedTexture->GetDefaultBindlessHandle();
+}
+
 void* FRHITextureReference::GetNativeResource() const
 {
 	check(ReferencedTexture);
diff --git a/Engine/Source/Runtime/RHI/Public/RHITextureReference.h b/Engine/Source/Runtime/RHI/Public/RHITextureReference.h
--- a/Engine/Source/Runtime/RHI/Public/RHITextureReference.h
+++ b/Engine/Source/Runtime/RHI/Public/RHITextureReference.h
@@ -11,6 +11,7 @@ public:
 	~FRHITextureReference();
 
 	virtual class FRHITextureReference* GetTextureReference() override;
+	virtual FRHIDescriptorHandle GetDefaultBindlessHandle() const override;
 
 	virtual void* GetNativeResource() const override;
 	virtual void* GetNativeShaderResourceView() const override;
