Ran IWYU on RHI Private files.
Adding missing include to ClosableMpscQueue.h

diff --git a/Engine/Source/Runtime/Core/Public/Containers/ClosableMpscQueue.h b/Engine/Source/Runtime/Core/Public/Containers/ClosableMpscQueue.h
--- a/Engine/Source/Runtime/Core/Public/Containers/ClosableMpscQueue.h
+++ b/Engine/Source/Runtime/Core/Public/Containers/ClosableMpscQueue.h
@@ -6,6 +6,7 @@
 
 #include "CoreTypes.h"
 #include "Templates/TypeCompatibleBytes.h"
+#include "Templates/UnrealTemplate.h"
 #include <atomic>
 
 /** 
diff --git a/Engine/Source/Runtime/RHI/Private/DataDrivenShaderPlatformInfo.cpp b/Engine/Source/Runtime/RHI/Private/DataDrivenShaderPlatformInfo.cpp
--- a/Engine/Source/Runtime/RHI/Private/DataDrivenShaderPlatformInfo.cpp
+++ b/Engine/Source/Runtime/RHI/Private/DataDrivenShaderPlatformInfo.cpp
@@ -4,6 +4,7 @@
 #include "RHI.h"
 #include "Misc/ConfigCacheIni.h"
 #include "Misc/DataDrivenPlatformInfoRegistry.h"
+#include "RHIStrings.h"
 
 const FName LANGUAGE_D3D("D3D");
 const FName LANGUAGE_Metal("Metal");
diff --git a/Engine/Source/Runtime/RHI/Private/DynamicRHI.cpp b/Engine/Source/Runtime/RHI/Private/DynamicRHI.cpp
--- a/Engine/Source/Runtime/RHI/Private/DynamicRHI.cpp
+++ b/Engine/Source/Runtime/RHI/Private/DynamicRHI.cpp
@@ -4,18 +4,22 @@
 	DynamicRHI.cpp: Dynamically bound Render Hardware Interface implementation.
 =============================================================================*/
 
-#include "CoreMinimal.h"
+#include "Containers/ClosableMpscQueue.h"
 #include "Misc/MessageDialog.h"
+#include "Experimental/Containers/HazardPointer.h"
 #include "Misc/OutputDeviceRedirector.h"
 #include "HAL/IConsoleManager.h"
 #include "Misc/App.h"
-#include "RHI.h"
 #include "Modules/ModuleManager.h"
 #include "GenericPlatform/GenericPlatformDriver.h"
 #include "GenericPlatform/GenericPlatformCrashContext.h"
 #include "PipelineStateCache.h"
+#include "RHIFwd.h"
 #include "TextureProfiler.h"
 #include "DataDrivenShaderPlatformInfo.h"
+#include "RHIImmutableSamplerState.h"
+#include "RHIStrings.h"
+#include "Serialization/MemoryImage.h"
 
 IMPLEMENT_TYPE_LAYOUT(FRayTracingGeometryInitializer);
 IMPLEMENT_TYPE_LAYOUT(FRayTracingGeometrySegment);
diff --git a/Engine/Source/Runtime/RHI/Private/GPUDefragAllocator.cpp b/Engine/Source/Runtime/RHI/Private/GPUDefragAllocator.cpp
--- a/Engine/Source/Runtime/RHI/Private/GPUDefragAllocator.cpp
+++ b/Engine/Source/Runtime/RHI/Private/GPUDefragAllocator.cpp
@@ -1,12 +1,10 @@
 // Copyright Epic Games, Inc. All Rights Reserved.
 
 #include "GPUDefragAllocator.h"
-#include "HAL/FileManager.h"
-#include "Misc/ScopeLock.h"
-#include "Math/RandomStream.h"
+#include "HAL/IConsoleManager.h"
 #include "Stats/StatsMisc.h"
 #include "RHI.h"
-#include "ProfilingDebugging/ScopedTimers.h"
+#include "Stats/Stats.h"
 
 DECLARE_STATS_GROUP(TEXT("TexturePool"), STATGROUP_TexturePool, STATCAT_ADVANCED);
 
diff --git a/Engine/Source/Runtime/RHI/Private/GPUProfiler.cpp b/Engine/Source/Runtime/RHI/Private/GPUProfiler.cpp
--- a/Engine/Source/Runtime/RHI/Private/GPUProfiler.cpp
+++ b/Engine/Source/Runtime/RHI/Private/GPUProfiler.cpp
@@ -5,6 +5,7 @@
 =============================================================================*/
 
 #include "GPUProfiler.h"
+#include "Async/TaskGraphInterfaces.h"
 #include "Misc/WildcardString.h"
 
 #if !UE_BUILD_SHIPPING
diff --git a/Engine/Source/Runtime/RHI/Private/GpuProfilerTrace.cpp b/Engine/Source/Runtime/RHI/Private/GpuProfilerTrace.cpp
--- a/Engine/Source/Runtime/RHI/Private/GpuProfilerTrace.cpp
+++ b/Engine/Source/Runtime/RHI/Private/GpuProfilerTrace.cpp
@@ -2,12 +2,10 @@
 
 #include "GpuProfilerTrace.h"
 #include "GPUProfiler.h"
-#include "HAL/PlatformMisc.h"
-#include "HAL/PlatformTime.h"
-#include "Misc/CString.h"
-#include "ProfilingDebugging/MiscTrace.h"
+#include "HAL/IConsoleManager.h"
 #include "Trace/Trace.inl"
-#include "UObject/NameTypes.h"
+#include "Trace/Detail/Field.h"
+#include "Trace/Detail/Important/ImportantLogScope.h"
 
 #if GPUPROFILERTRACE_ENABLED
 
diff --git a/Engine/Source/Runtime/RHI/Private/Linux/LinuxDynamicRHI.cpp b/Engine/Source/Runtime/RHI/Private/Linux/LinuxDynamicRHI.cpp
--- a/Engine/Source/Runtime/RHI/Private/Linux/LinuxDynamicRHI.cpp
+++ b/Engine/Source/Runtime/RHI/Private/Linux/LinuxDynamicRHI.cpp
@@ -1,12 +1,13 @@
 // Copyright Epic Games, Inc. All Rights Reserved.
 
-#include "CoreMinimal.h"
+#include "Misc/App.h"
 #include "Misc/MessageDialog.h"
 #include "RHI.h"
 #include "Modules/ModuleManager.h"
 #include "Misc/ConfigCacheIni.h"
 #include "HAL/PlatformApplicationMisc.h"
 #include "DataDrivenShaderPlatformInfo.h"
+#include "RHIStrings.h"
 
 FDynamicRHI* PlatformCreateDynamicRHI()
 {
diff --git a/Engine/Source/Runtime/RHI/Private/MultiGPU.cpp b/Engine/Source/Runtime/RHI/Private/MultiGPU.cpp
--- a/Engine/Source/Runtime/RHI/Private/MultiGPU.cpp
+++ b/Engine/Source/Runtime/RHI/Private/MultiGPU.cpp
@@ -5,7 +5,6 @@
 =============================================================================*/
 
 #include "MultiGPU.h"
-#include "RHI.h"
 
 #if WITH_SLI || WITH_MGPU
 uint32 GNumAlternateFrameRenderingGroups = 1;
diff --git a/Engine/Source/Runtime/RHI/Private/PipelineFileCache.cpp b/Engine/Source/Runtime/RHI/Private/PipelineFileCache.cpp
--- a/Engine/Source/Runtime/RHI/Private/PipelineFileCache.cpp
+++ b/Engine/Source/Runtime/RHI/Private/PipelineFileCache.cpp
@@ -5,14 +5,14 @@ PipelineFileCache.cpp: Pipeline state cache implementation.
 =============================================================================*/
 
 #include "PipelineFileCache.h"
+#include "Containers/List.h"
 #include "PipelineStateCache.h"
-#include "HAL/FileManager.h"
+#include "HAL/IConsoleManager.h"
 #include "Misc/EngineVersion.h"
-#include "Serialization/Archive.h"
+#include "HAL/PlatformFile.h"
 #include "Serialization/MemoryReader.h"
+#include "Misc/CommandLine.h"
 #include "Serialization/MemoryWriter.h"
-#include "RHI.h"
-#include "RHIResources.h"
 #include "DataDrivenShaderPlatformInfo.h"
 #include "Misc/ScopeRWLock.h"
 #include "Misc/Paths.h"
@@ -20,10 +20,10 @@ PipelineFileCache.cpp: Pipeline state cache implementation.
 #include "HAL/PlatformFileManager.h"
 #include "Misc/FileHelper.h"
 #include "ProfilingDebugging/CsvProfiler.h"
+#include "RHIStrings.h"
 #include "String/LexFromString.h"
 #include "String/ParseTokens.h"
 #include "Misc/ScopeExit.h"
-#include <Misc/ConfigCacheIni.h>
 #include <Algo/ForEach.h>
 
 static FString JOURNAL_FILE_EXTENSION(TEXT(".jnl"));
diff --git a/Engine/Source/Runtime/RHI/Private/PipelineStateCache.cpp b/Engine/Source/Runtime/RHI/Private/PipelineStateCache.cpp
--- a/Engine/Source/Runtime/RHI/Private/PipelineStateCache.cpp
+++ b/Engine/Source/Runtime/RHI/Private/PipelineStateCache.cpp
@@ -6,15 +6,18 @@ PipelineStateCache.cpp: Pipeline state cache implementation.
 
 #include "PipelineStateCache.h"
 #include "Async/AsyncWork.h"
+#include "Async/TaskGraphInterfaces.h"
 #include "PipelineFileCache.h"
+#include "Containers/ClosableMpscQueue.h"
 #include "Misc/ScopeRWLock.h"
-#include "Misc/ScopeLock.h"
-#include "CoreGlobals.h"
+#include "Misc/App.h"
 #include "Misc/TimeGuard.h"
 #include "Containers/DiscardableKeyValueCache.h"
-#include "Async/Async.h"
-#include "HAL/PlatformMisc.h"
 #include "DataDrivenShaderPlatformInfo.h"
+#include "ProfilingDebugging/CsvProfiler.h"
+#include "RHIFwd.h"
+#include "RHIImmutableSamplerState.h"
+#include "Stats/StatsTrace.h"
 
 // perform cache eviction each frame, used to stress the system and flush out bugs
 #define PSO_DO_CACHE_EVICT_EACH_FRAME 0
diff --git a/Engine/Source/Runtime/RHI/Private/RHI.cpp b/Engine/Source/Runtime/RHI/Private/RHI.cpp
--- a/Engine/Source/Runtime/RHI/Private/RHI.cpp
+++ b/Engine/Source/Runtime/RHI/Private/RHI.cpp
@@ -5,18 +5,22 @@
 =============================================================================*/
 
 #include "RHI.h"
+#include "Async/ParallelFor.h"
+#include "HAL/IConsoleManager.h"
 #include "RHITransientResourceAllocator.h"
+#include "Misc/CommandLine.h"
 #include "Modules/ModuleManager.h"
 #include "Misc/ConfigCacheIni.h"
-#include "Misc/MessageDialog.h"
 #include "ProfilingDebugging/CsvProfiler.h"
-#include "String/Find.h"
+#include "RHIFwd.h"
 #include "String/LexFromString.h"
+#include "RHIStrings.h"
 #include "String/ParseTokens.h"
 #include "Misc/BufferedOutputDevice.h"
-#include "Misc/OutputDeviceFile.h"
-#include "Async/ParallelFor.h"
 #include "DataDrivenShaderPlatformInfo.h"
+#include "Serialization/MemoryImage.h"
+#include "Stats/StatsTrace.h"
+#include "RHITextureReference.h"
 
 #if RHI_ENABLE_RESOURCE_INFO
 #include "HAL/FileManager.h"
diff --git a/Engine/Source/Runtime/RHI/Private/RHIBreadcrumbs.cpp b/Engine/Source/Runtime/RHI/Private/RHIBreadcrumbs.cpp
--- a/Engine/Source/Runtime/RHI/Private/RHIBreadcrumbs.cpp
+++ b/Engine/Source/Runtime/RHI/Private/RHIBreadcrumbs.cpp
@@ -1,8 +1,10 @@
 // Copyright Epic Games, Inc. All Rights Reserved.
 
 #include "RHIBreadcrumbs.h"
+#include "Misc/MemStack.h"
 #include "RHI.h"
 #include "GenericPlatform/GenericPlatformCrashContext.h"
+#include "Trace/Trace.inl"
 
 #if RHI_WANT_BREADCRUMB_EVENTS
 
@@ -154,4 +156,4 @@ void FRHIBreadcrumbStack::WriteRenderBreadcrumbs(FCrashContextExtendedWriter& Wr
 }
 #endif
 
-#endif
\ No newline at end of file
+#endif
diff --git a/Engine/Source/Runtime/RHI/Private/RHICommandList.cpp b/Engine/Source/Runtime/RHI/Private/RHICommandList.cpp
--- a/Engine/Source/Runtime/RHI/Private/RHICommandList.cpp
+++ b/Engine/Source/Runtime/RHI/Private/RHICommandList.cpp
@@ -1,20 +1,18 @@
 // Copyright Epic Games, Inc. All Rights Reserved.
 
 
-#include "CoreMinimal.h"
-#include "Misc/CommandLine.h"
-#include "Stats/Stats.h"
-#include "HAL/IConsoleManager.h"
 #include "Misc/App.h"
 #include "Async/TaskGraphInterfaces.h"
 #include "RHI.h"
-#include "Misc/ScopeLock.h"
-#include "Misc/ScopeExit.h"
-#include "PipelineStateCache.h"
 #include "ProfilingDebugging/CsvProfiler.h"
-#include "Trace/Trace.inl"
 #include "GenericPlatform/GenericPlatformCrashContext.h"
+#include "RHIBreadcrumbs.h"
 #include "RHIResourceUpdates.h"
+#include "RHIContext.h"
+#include "RHIFwd.h"
+#include "RHITextureReference.h"
+#include "RHITransition.h"
+#include "Stats/StatsTrace.h"
 
 CSV_DEFINE_CATEGORY_MODULE(RHI_API, RHITStalls, false);
 CSV_DEFINE_CATEGORY_MODULE(RHI_API, RHITFlushes, false);
diff --git a/Engine/Source/Runtime/RHI/Private/RHIGPUReadback.cpp b/Engine/Source/Runtime/RHI/Private/RHIGPUReadback.cpp
--- a/Engine/Source/Runtime/RHI/Private/RHIGPUReadback.cpp
+++ b/Engine/Source/Runtime/RHI/Private/RHIGPUReadback.cpp
@@ -6,6 +6,7 @@
 =============================================================================*/
 
 #include "RHIGPUReadback.h"
+#include "RHITransition.h"
 
 ///////////////////////////////////////////////////////////////////////////////
 //////////////////////     FGenericRHIGPUFence    /////////////////////////////
diff --git a/Engine/Source/Runtime/RHI/Private/RHIResources.cpp b/Engine/Source/Runtime/RHI/Private/RHIResources.cpp
--- a/Engine/Source/Runtime/RHI/Private/RHIResources.cpp
+++ b/Engine/Source/Runtime/RHI/Private/RHIResources.cpp
@@ -1,7 +1,11 @@
 // Copyright Epic Games, Inc. All Rights Reserved.
 
 #include "RHIResources.h"
+#include "Containers/ClosableMpscQueue.h"
 #include "RHI.h"
+#include "Experimental/Containers/HazardPointer.h"
+#include "Misc/MemStack.h"
+#include "Stats/Stats.h"
 
 UE::TConsumeAllMpmcQueue<FRHIResource*> FRHIResource::PendingDeletes;
 FRHIResource* FRHIResource::CurrentlyDeleting = nullptr;
diff --git a/Engine/Source/Runtime/RHI/Private/RHIStrings.cpp b/Engine/Source/Runtime/RHI/Private/RHIStrings.cpp
--- a/Engine/Source/Runtime/RHI/Private/RHIStrings.cpp
+++ b/Engine/Source/Runtime/RHI/Private/RHIStrings.cpp
@@ -3,8 +3,9 @@
 #include "RHIStrings.h"
 #include "DataDrivenShaderPlatformInfo.h"
 #include "RHI.h"
+#include "RHIAccess.h"
 #include "RHIShaderFormatDefinitions.inl"
-#include "UObject/NameTypes.h"
+#include "RHIPipeline.h"
 
 static FName NAME_PLATFORM_WINDOWS(TEXT("Windows"));
 static FName NAME_PLATFORM_ANDROID(TEXT("Android"));
diff --git a/Engine/Source/Runtime/RHI/Private/RHIUtilities.cpp b/Engine/Source/Runtime/RHI/Private/RHIUtilities.cpp
--- a/Engine/Source/Runtime/RHI/Private/RHIUtilities.cpp
+++ b/Engine/Source/Runtime/RHI/Private/RHIUtilities.cpp
@@ -4,15 +4,18 @@
 RHIUtilities.cpp:
 =============================================================================*/
 
-#include "CoreMinimal.h"
+#include "Async/TaskGraphInterfaces.h"
 #include "HAL/PlatformStackWalk.h"
-#include "HAL/IConsoleManager.h"
+#include "Containers/ClosableMpscQueue.h"
 #include "RHI.h"
+#include "GenericPlatform/GenericPlatformFramePacer.h"
 #include "HAL/Runnable.h"
 #include "HAL/RunnableThread.h"
-#include "Misc/ScopeLock.h"
 #include "HAL/PlatformFramePacer.h"
 #include "DataDrivenShaderPlatformInfo.h"
+#include "RHIAccess.h"
+#include "RHIFwd.h"
+#include "RHIStrings.h"
 
 #define USE_FRAME_OFFSET_THREAD 1
 
diff --git a/Engine/Source/Runtime/RHI/Private/RHIValidation.cpp b/Engine/Source/Runtime/RHI/Private/RHIValidation.cpp
--- a/Engine/Source/Runtime/RHI/Private/RHIValidation.cpp
+++ b/Engine/Source/Runtime/RHI/Private/RHIValidation.cpp
@@ -5,10 +5,13 @@
 =============================================================================*/
 
 #include "RHIValidation.h"
-#include "RHIValidationContext.h"
+#include "HAL/IConsoleManager.h"
 #include "RHIValidationTransientResourceAllocator.h"
 #include "HAL/PlatformStackWalk.h"
+#include "Misc/CommandLine.h"
 #include "Misc/OutputDeviceRedirector.h"
+#include "RHIContext.h"
+#include "RHIStrings.h"
 
 #if ENABLE_RHI_VALIDATION
 
diff --git a/Engine/Source/Runtime/RHI/Private/TextureProfiler.cpp b/Engine/Source/Runtime/RHI/Private/TextureProfiler.cpp
--- a/Engine/Source/Runtime/RHI/Private/TextureProfiler.cpp
+++ b/Engine/Source/Runtime/RHI/Private/TextureProfiler.cpp
@@ -2,11 +2,10 @@
 
 #include "TextureProfiler.h"
 
+#include "Misc/CoreDelegates.h"
 #include "ProfilingDebugging/CsvProfiler.h"
-#include "RHI.h"
+#include "RHIDefinitions.h"
 #include "RHIResources.h"
-#include "UObject/NameTypes.h"
-#include "Containers/StringConv.h"
 
 #if TEXTURE_PROFILER_ENABLED
 
