Adding missing include.

diff --git a/Engine/Source/Runtime/RHI/Private/RHIValidation.cpp b/Engine/Source/Runtime/RHI/Private/RHIValidation.cpp
--- a/Engine/Source/Runtime/RHI/Private/RHIValidation.cpp
+++ b/Engine/Source/Runtime/RHI/Private/RHIValidation.cpp
@@ -5,6 +5,7 @@
 =============================================================================*/
 
 #include "RHIValidation.h"
+#include "RHIValidationContext.h"
 #include "HAL/IConsoleManager.h"
 #include "RHIValidationTransientResourceAllocator.h"
 #include "HAL/PlatformStackWalk.h"
