RHI header reorganization
- Changing RHIUtilities, RHIContext, RHICommandList and DynamicRHI to get them to compile independently of the rest of the RHI headers.

diff --git a/Engine/Source/Programs/UnrealBuildTool/Configuration/UEBuildModuleCPP.cs b/Engine/Source/Programs/UnrealBuildTool/Configuration/UEBuildModuleCPP.cs
--- a/Engine/Source/Programs/UnrealBuildTool/Configuration/UEBuildModuleCPP.cs
+++ b/Engine/Source/Programs/UnrealBuildTool/Configuration/UEBuildModuleCPP.cs
@@ -1528,13 +1528,6 @@ namespace UnrealBuildTool
 
 		private bool IgnoreMismatchedHeader(string ExpectedName)
 		{
-			switch(ExpectedName)
-			{
-				case "DynamicRHI":
-				case "RHICommandList":
-				case "RHIUtilities":
-					return true;
-			}
 			switch(Name)
 			{
 				case "D3D11RHI":
diff --git a/Engine/Source/Runtime/RHI/Private/DynamicRHI.cpp b/Engine/Source/Runtime/RHI/Private/DynamicRHI.cpp
--- a/Engine/Source/Runtime/RHI/Private/DynamicRHI.cpp
+++ b/Engine/Source/Runtime/RHI/Private/DynamicRHI.cpp
@@ -4,6 +4,7 @@
 	DynamicRHI.cpp: Dynamically bound Render Hardware Interface implementation.
 =============================================================================*/
 
+#include "DynamicRHI.h"
 #include "Containers/ClosableMpscQueue.h"
 #include "Misc/MessageDialog.h"
 #include "Experimental/Containers/HazardPointer.h"
@@ -457,6 +458,18 @@ static FAutoConsoleCommandWithWorldAndArgs GBaseRHISetGPUCaptureOptions(
 	FConsoleCommandWithWorldAndArgsDelegate::CreateStatic(&BaseRHISetGPUCaptureOptions)
 	);
 
+// Default fallback; will not work for non-8-bit surfaces and it's extremely slow.
+void FDynamicRHI::RHIReadSurfaceData(FRHITexture* Texture, FIntRect Rect, TArray<FLinearColor>& OutData, FReadSurfaceDataFlags InFlags)
+{
+	TArray<FColor> TempData;
+	RHIReadSurfaceData(Texture, Rect, TempData, InFlags);
+	OutData.SetNumUninitialized(TempData.Num());
+	for (int32 Index = 0; Index < TempData.Num(); ++Index)
+	{
+		OutData[Index] = TempData[Index].ReinterpretAsLinear();
+	}
+}
+
 void FDynamicRHI::RHIReadSurfaceFloatData(FRHITexture* Texture, FIntRect Rect, TArray<FFloat16Color>& OutData, FReadSurfaceDataFlags InFlags)
 {
 #if WITH_MGPU
diff --git a/Engine/Source/Runtime/RHI/Private/RHI.cpp b/Engine/Source/Runtime/RHI/Private/RHI.cpp
--- a/Engine/Source/Runtime/RHI/Private/RHI.cpp
+++ b/Engine/Source/Runtime/RHI/Private/RHI.cpp
@@ -38,27 +38,6 @@ CSV_DEFINE_CATEGORY(DrawCall, false);
 CSV_DEFINE_CATEGORY(DrawCall, true);
 #endif
 
-// Define counter stats.
-#if HAS_GPU_STATS
-	DEFINE_STAT(STAT_RHIDrawPrimitiveCalls);
-	DEFINE_STAT(STAT_RHITriangles);
-	DEFINE_STAT(STAT_RHILines);
-#endif
-
-// Define memory stats.
-DEFINE_STAT(STAT_RenderTargetMemory2D);
-DEFINE_STAT(STAT_RenderTargetMemory3D);
-DEFINE_STAT(STAT_RenderTargetMemoryCube);
-DEFINE_STAT(STAT_TextureMemory2D);
-DEFINE_STAT(STAT_TextureMemory3D);
-DEFINE_STAT(STAT_TextureMemoryCube);
-DEFINE_STAT(STAT_UniformBufferMemory);
-DEFINE_STAT(STAT_IndexBufferMemory);
-DEFINE_STAT(STAT_VertexBufferMemory);
-DEFINE_STAT(STAT_RTAccelerationStructureMemory);
-DEFINE_STAT(STAT_StructuredBufferMemory);
-DEFINE_STAT(STAT_PixelBufferMemory);
-
 IMPLEMENT_TYPE_LAYOUT(FRHIUniformBufferLayoutInitializer);
 IMPLEMENT_TYPE_LAYOUT(FRHIUniformBufferResource);
 
@@ -86,6 +65,7 @@ static TAutoConsoleVariable<int32> CVarGraphicsAdapter(
 	ECVF_ReadOnly | ECVF_RenderThreadSafe);
 
 #if STATS
+#include "ProfilingDebugging/CsvProfilerConfig.h"
 #include "Stats/StatsData.h"
 static void DumpRHIMemory(FOutputDevice& OutputDevice)
 {
@@ -1417,12 +1397,6 @@ RHI_API uint64 GDemotedLocalMemorySize = 0;
 
 RHI_API EShaderPlatform GShaderPlatformForFeatureLevel[ERHIFeatureLevel::Num] = {SP_NumPlatforms,SP_NumPlatforms,SP_NumPlatforms,SP_NumPlatforms,SP_NumPlatforms};
 
-// simple stats about draw calls. GNum is the previous frame and 
-// GCurrent is the current frame.
-
-RHI_API int32 GNumDrawCallsRHI      [MAX_NUM_GPUS] = {};
-RHI_API int32 GNumPrimitivesDrawnRHI[MAX_NUM_GPUS] = {};
-
 // By default, read only states and UAV states are allowed to participate in state merging.
 ERHIAccess GRHIMergeableAccessMask = ERHIAccess::ReadOnlyMask | ERHIAccess::UAVMask;
 
diff --git a/Engine/Source/Runtime/RHI/Private/RHICommandList.cpp b/Engine/Source/Runtime/RHI/Private/RHICommandList.cpp
--- a/Engine/Source/Runtime/RHI/Private/RHICommandList.cpp
+++ b/Engine/Source/Runtime/RHI/Private/RHICommandList.cpp
@@ -1,6 +1,6 @@
 // Copyright Epic Games, Inc. All Rights Reserved.
 
-
+#include "RHICommandList.h"
 #include "Misc/App.h"
 #include "Async/TaskGraphInterfaces.h"
 #include "RHI.h"
@@ -2043,14 +2043,3 @@ void FRHICommandListImmediate::CleanupGraphEvents()
 		GraphEvent.SafeRelease();
 	}
 }
-
-RHI_API void RHISetComputeShaderBackwardsCompatible(IRHIComputeContext* InContext, FRHIComputeShader* InShader)
-{
-	TRefCountPtr<FRHIComputePipelineState> ComputePipelineState = RHICreateComputePipelineState(InShader);
-	InContext->RHISetComputePipelineState(ComputePipelineState);
-}
-
-void IRHIComputeContext::RHISetComputeShader(FRHIComputeShader* ComputeShader)
-{
-	RHISetComputeShaderBackwardsCompatible(this, ComputeShader);
-}
diff --git a/Engine/Source/Runtime/RHI/Private/RHIContext.cpp b/Engine/Source/Runtime/RHI/Private/RHIContext.cpp
new file mode 100644
--- /dev/null
+++ b/Engine/Source/Runtime/RHI/Private/RHIContext.cpp
@@ -0,0 +1,40 @@
+// Copyright Epic Games, Inc. All Rights Reserved.
+
+#include "RHIContext.h"
+#include "RHI.h"
+
+RHI_API void RHISetComputeShaderBackwardsCompatible(IRHIComputeContext* InContext, FRHIComputeShader* InShader)
+{
+	TRefCountPtr<FRHIComputePipelineState> ComputePipelineState = RHICreateComputePipelineState(InShader);
+	InContext->RHISetComputePipelineState(ComputePipelineState);
+}
+
+void IRHIComputeContext::RHISetComputeShader(FRHIComputeShader* ComputeShader)
+{
+	RHISetComputeShaderBackwardsCompatible(this, ComputeShader);
+}
+
+void IRHIComputeContext::StatsSetCategory(FRHIDrawStats* InStats, uint32 InCategoryID, uint32 InGPUIndex)
+{
+	Stats = &InStats->GetGPU(InGPUIndex).GetCategory(InCategoryID);
+}
+
+void IRHICommandContextPSOFallback::SetGraphicsPipelineStateFromInitializer(const FGraphicsPipelineStateInitializer& PsoInit, uint32 StencilRef, bool bApplyAdditionalState)
+{
+	RHISetBoundShaderState(
+		RHICreateBoundShaderState(
+			PsoInit.BoundShaderState.VertexDeclarationRHI,
+			PsoInit.BoundShaderState.VertexShaderRHI,
+			PsoInit.BoundShaderState.PixelShaderRHI,
+			PsoInit.BoundShaderState.GetGeometryShader()
+		).GetReference()
+	);
+
+	RHISetDepthStencilState(PsoInit.DepthStencilState, StencilRef);
+	RHISetRasterizerState(PsoInit.RasterizerState);
+	RHISetBlendState(PsoInit.BlendState, FLinearColor(1.0f, 1.0f, 1.0f));
+	if (GSupportsDepthBoundsTest)
+	{
+		RHIEnableDepthBoundsTest(PsoInit.bDepthBounds);
+	}
+}
diff --git a/Engine/Source/Runtime/RHI/Private/RHIStats.cpp b/Engine/Source/Runtime/RHI/Private/RHIStats.cpp
new file mode 100644
--- /dev/null
+++ b/Engine/Source/Runtime/RHI/Private/RHIStats.cpp
@@ -0,0 +1,30 @@
+// Copyright Epic Games, Inc. All Rights Reserved.
+
+#include "RHIStats.h"
+#include "HAL/Platform.h"
+#include "HAL/PlatformAtomics.h"
+#include "Templates/Atomic.h"
+
+int32 GNumDrawCallsRHI[MAX_NUM_GPUS] = {};
+int32 GNumPrimitivesDrawnRHI[MAX_NUM_GPUS] = {};
+
+// Define counter stats.
+#if HAS_GPU_STATS
+	DEFINE_STAT(STAT_RHIDrawPrimitiveCalls);
+	DEFINE_STAT(STAT_RHITriangles);
+	DEFINE_STAT(STAT_RHILines);
+#endif
+
+// Define memory stats.
+DEFINE_STAT(STAT_RenderTargetMemory2D);
+DEFINE_STAT(STAT_RenderTargetMemory3D);
+DEFINE_STAT(STAT_RenderTargetMemoryCube);
+DEFINE_STAT(STAT_TextureMemory2D);
+DEFINE_STAT(STAT_TextureMemory3D);
+DEFINE_STAT(STAT_TextureMemoryCube);
+DEFINE_STAT(STAT_UniformBufferMemory);
+DEFINE_STAT(STAT_IndexBufferMemory);
+DEFINE_STAT(STAT_VertexBufferMemory);
+DEFINE_STAT(STAT_RTAccelerationStructureMemory);
+DEFINE_STAT(STAT_StructuredBufferMemory);
+DEFINE_STAT(STAT_PixelBufferMemory);
diff --git a/Engine/Source/Runtime/RHI/Private/RHIUtilities.cpp b/Engine/Source/Runtime/RHI/Private/RHIUtilities.cpp
--- a/Engine/Source/Runtime/RHI/Private/RHIUtilities.cpp
+++ b/Engine/Source/Runtime/RHI/Private/RHIUtilities.cpp
@@ -4,6 +4,7 @@
 RHIUtilities.cpp:
 =============================================================================*/
 
+#include "RHIUtilities.h"
 #include "Async/TaskGraphInterfaces.h"
 #include "HAL/PlatformStackWalk.h"
 #include "Containers/ClosableMpscQueue.h"
@@ -661,6 +662,81 @@ RHI_API ERHIAccess RHIGetDefaultResourceState(EBufferUsageFlags InUsage, bool bI
 	return ResourceState;
 }
 
+void DecodeRenderTargetMode(ESimpleRenderTargetMode Mode, ERenderTargetLoadAction& ColorLoadAction, ERenderTargetStoreAction& ColorStoreAction, ERenderTargetLoadAction& DepthLoadAction, ERenderTargetStoreAction& DepthStoreAction, ERenderTargetLoadAction& StencilLoadAction, ERenderTargetStoreAction& StencilStoreAction, FExclusiveDepthStencil DepthStencilUsage)
+{
+	// set defaults
+	ColorStoreAction = ERenderTargetStoreAction::EStore;
+	DepthStoreAction = ERenderTargetStoreAction::EStore;
+	StencilStoreAction = ERenderTargetStoreAction::EStore;
+
+	switch (Mode)
+	{
+	case ESimpleRenderTargetMode::EExistingColorAndDepth:
+		ColorLoadAction = ERenderTargetLoadAction::ELoad;
+		DepthLoadAction = ERenderTargetLoadAction::ELoad;
+		break;
+	case ESimpleRenderTargetMode::EUninitializedColorAndDepth:
+		ColorLoadAction = ERenderTargetLoadAction::ENoAction;
+		DepthLoadAction = ERenderTargetLoadAction::ENoAction;
+		break;
+	case ESimpleRenderTargetMode::EUninitializedColorExistingDepth:
+		ColorLoadAction = ERenderTargetLoadAction::ENoAction;
+		DepthLoadAction = ERenderTargetLoadAction::ELoad;
+		break;
+	case ESimpleRenderTargetMode::EUninitializedColorClearDepth:
+		ColorLoadAction = ERenderTargetLoadAction::ENoAction;
+		DepthLoadAction = ERenderTargetLoadAction::EClear;
+		break;
+	case ESimpleRenderTargetMode::EClearColorExistingDepth:
+		ColorLoadAction = ERenderTargetLoadAction::EClear;
+		DepthLoadAction = ERenderTargetLoadAction::ELoad;
+		break;
+	case ESimpleRenderTargetMode::EClearColorAndDepth:
+		ColorLoadAction = ERenderTargetLoadAction::EClear;
+		DepthLoadAction = ERenderTargetLoadAction::EClear;
+		break;
+	case ESimpleRenderTargetMode::EExistingContents_NoDepthStore:
+		ColorLoadAction = ERenderTargetLoadAction::ELoad;
+		DepthLoadAction = ERenderTargetLoadAction::ELoad;
+		DepthStoreAction = ERenderTargetStoreAction::ENoAction;
+		break;
+	case ESimpleRenderTargetMode::EExistingColorAndClearDepth:
+		ColorLoadAction = ERenderTargetLoadAction::ELoad;
+		DepthLoadAction = ERenderTargetLoadAction::EClear;
+		break;
+	case ESimpleRenderTargetMode::EExistingColorAndDepthAndClearStencil:
+		ColorLoadAction = ERenderTargetLoadAction::ELoad;
+		DepthLoadAction = ERenderTargetLoadAction::ELoad;
+		break;
+	default:
+		UE_LOG(LogRHI, Fatal, TEXT("Using a ESimpleRenderTargetMode that wasn't decoded in DecodeRenderTargetMode [value = %d]"), (int32)Mode);
+	}
+
+	StencilLoadAction = DepthLoadAction;
+
+	if (!DepthStencilUsage.IsUsingDepth())
+	{
+		DepthLoadAction = ERenderTargetLoadAction::ENoAction;
+	}
+
+	//if we aren't writing to depth, there's no reason to store it back out again.  Should save some bandwidth on mobile platforms.
+	if (!DepthStencilUsage.IsDepthWrite())
+	{
+		DepthStoreAction = ERenderTargetStoreAction::ENoAction;
+	}
+
+	if (!DepthStencilUsage.IsUsingStencil())
+	{
+		StencilLoadAction = ERenderTargetLoadAction::ENoAction;
+	}
+
+	//if we aren't writing to stencil, there's no reason to store it back out again.  Should save some bandwidth on mobile platforms.
+	if (!DepthStencilUsage.IsStencilWrite())
+	{
+		StencilStoreAction = ERenderTargetStoreAction::ENoAction;
+	}
+}
+
 void RHICreateTargetableShaderResource(
 	const FRHITextureCreateDesc& BaseDesc,
 	ETextureCreateFlags TargetableTextureFlags,
@@ -918,3 +994,8 @@ void RHICreateTargetableShaderResource3D(
 
 	RHICreateTargetableShaderResource(Desc, TargetableTextureFlags, bForceSeparateTargetAndShaderResource, false, OutTargetableTexture, OutShaderResourceTexture);
 }
+
+void FRHILockTracker::RaiseMismatchError()
+{
+	UE_LOG(LogRHI, Fatal, TEXT("Mismatched RHI buffer locks."));
+}
diff --git a/Engine/Source/Runtime/RHI/Public/DynamicRHI.h b/Engine/Source/Runtime/RHI/Public/DynamicRHI.h
--- a/Engine/Source/Runtime/RHI/Public/DynamicRHI.h
+++ b/Engine/Source/Runtime/RHI/Public/DynamicRHI.h
@@ -1,17 +1,14 @@
 // Copyright Epic Games, Inc. All Rights Reserved.
 
-// IWYU pragma: private, include "RHI.h"
-
 /*=============================================================================
 DynamicRHI.h: Dynamically bound Render Hardware Interface definitions.
 =============================================================================*/
 
 #pragma once
 
-// HEADER_UNIT_SKIP - Should always be included through RHI.h (this code needs cleanup.. this file is included in the middle of RHI.h)
-
 #include "CoreTypes.h"
 #include "RHIContext.h"
+#include "RHIFeatureLevel.h"
 #include "MultiGPU.h"
 #include "Serialization/MemoryLayout.h"
 #include "Containers/ArrayView.h"
@@ -34,36 +31,9 @@ struct FRHIUniformBufferLayout;
 struct FSamplerStateInitializerRHI;
 struct FTextureMemoryStats;
 struct FRHIGPUMask;
+struct FUpdateTexture3DData;
 
-
-/** Struct to hold common data between begin/end updatetexture3d */
-struct FUpdateTexture3DData
-{
-	FUpdateTexture3DData(FRHITexture3D* InTexture, uint32 InMipIndex, const struct FUpdateTextureRegion3D& InUpdateRegion, uint32 InSourceRowPitch, uint32 InSourceDepthPitch, uint8* InSourceData, uint32 InDataSizeBytes, uint32 InFrameNumber )
-		: Texture(InTexture)
-		, MipIndex(InMipIndex)
-		, UpdateRegion(InUpdateRegion)
-		, RowPitch(InSourceRowPitch)
-		, DepthPitch(InSourceDepthPitch)
-		, Data(InSourceData)
-		, DataSizeBytes(InDataSizeBytes)
-		, FrameNumber(InFrameNumber)
-	{
-	}
-
-	FRHITexture3D* Texture;
-	uint32 MipIndex;
-	FUpdateTextureRegion3D UpdateRegion;
-	uint32 RowPitch;
-	uint32 DepthPitch;
-	uint8* Data;
-	uint32 DataSizeBytes;
-	uint32 FrameNumber;
-	uint8 PlatformData[64];
-
-private:
-	FUpdateTexture3DData();
-};
+typedef TArray<FScreenResolutionRHI> FScreenResolutionArray;
 
 /** Struct to provide details of swap chain flips */
 struct FRHIFlipDetails
@@ -738,16 +708,7 @@ public:
 	virtual void RHIReadSurfaceData(FRHITexture* Texture, FIntRect Rect, TArray<FColor>& OutData, FReadSurfaceDataFlags InFlags) = 0;
 
 	// Default fallback; will not work for non-8-bit surfaces and it's extremely slow.
-	virtual void RHIReadSurfaceData(FRHITexture* Texture, FIntRect Rect, TArray<FLinearColor>& OutData, FReadSurfaceDataFlags InFlags)
-	{
-		TArray<FColor> TempData;
-		RHIReadSurfaceData(Texture, Rect, TempData, InFlags);
-		OutData.SetNumUninitialized(TempData.Num());
-		for (int32 Index = 0; Index < TempData.Num(); ++Index)
-		{
-			OutData[Index] = TempData[Index].ReinterpretAsLinear();
-		}
-	}
+	virtual void RHIReadSurfaceData(FRHITexture* Texture, FIntRect Rect, TArray<FLinearColor>& OutData, FReadSurfaceDataFlags InFlags);
 
 	/** Watch out for OutData to be 0 (can happen on DXGI_ERROR_DEVICE_REMOVED), don't call RHIUnmapStagingSurface in that case. */
 	// FlushType: Flush Immediate (seems wrong)
diff --git a/Engine/Source/Runtime/RHI/Public/MultiGPU.h b/Engine/Source/Runtime/RHI/Public/MultiGPU.h
--- a/Engine/Source/Runtime/RHI/Public/MultiGPU.h
+++ b/Engine/Source/Runtime/RHI/Public/MultiGPU.h
@@ -6,7 +6,10 @@
 
 #pragma once
 
-#include "CoreMinimal.h"
+#include "Containers/Array.h"
+#include "Containers/ContainerAllocationPolicies.h"
+#include "Math/UnrealMathUtility.h"
+#include "Misc/AssertionMacros.h"
 
 /** When greater than one, indicates that SLI rendering is enabled */
 #if PLATFORM_DESKTOP
diff --git a/Engine/Source/Runtime/RHI/Public/RHI.h b/Engine/Source/Runtime/RHI/Public/RHI.h
--- a/Engine/Source/Runtime/RHI/Public/RHI.h
+++ b/Engine/Source/Runtime/RHI/Public/RHI.h
@@ -723,111 +723,6 @@ inline EShaderPlatform GetFeatureLevelShaderPlatform(const FStaticFeatureLevel I
 	return GShaderPlatformForFeatureLevel[InFeatureLevel];
 }
 
-/** to customize the RHIReadSurfaceData() output */
-class FReadSurfaceDataFlags
-{
-public:
-	// @param InCompressionMode defines the value input range that is mapped to output range
-	// @param InCubeFace defined which cubemap side is used, only required for cubemap content, then it needs to be a valid side
-	FReadSurfaceDataFlags(ERangeCompressionMode InCompressionMode = RCM_UNorm, ECubeFace InCubeFace = CubeFace_MAX) 
-		:CubeFace(InCubeFace), CompressionMode(InCompressionMode)
-	{
-	}
-
-	ECubeFace GetCubeFace() const
-	{
-		checkSlow(CubeFace <= CubeFace_NegZ);
-		return CubeFace;
-	}
-
-	ERangeCompressionMode GetCompressionMode() const
-	{
-		return CompressionMode;
-	}
-
-	void SetLinearToGamma(bool Value)
-	{
-		bLinearToGamma = Value;
-	}
-
-	bool GetLinearToGamma() const
-	{
-		return bLinearToGamma;
-	}
-
-	void SetOutputStencil(bool Value)
-	{
-		bOutputStencil = Value;
-	}
-
-	bool GetOutputStencil() const
-	{
-		return bOutputStencil;
-	}
-
-	void SetMip(uint8 InMipLevel)
-	{
-		MipLevel = InMipLevel;
-	}
-
-	uint8 GetMip() const
-	{
-		return MipLevel;
-	}	
-
-	void SetMaxDepthRange(float Value)
-	{
-		MaxDepthRange = Value;
-	}
-
-	float ComputeNormalizedDepth(float DeviceZ) const
-	{
-		return FMath::Abs(ConvertFromDeviceZ(DeviceZ) / MaxDepthRange);
-	}
-
-	void SetGPUIndex(uint32 InGPUIndex)
-	{
-		GPUIndex = InGPUIndex;
-	}
-
-	uint32 GetGPUIndex() const
-	{
-		return GPUIndex;
-	}
-
-	void SetArrayIndex(int32 InArrayIndex)
-	{
-		ArrayIndex = InArrayIndex;
-	}
-
-	int32 GetArrayIndex() const
-	{
-		return ArrayIndex;
-	}
-
-private:
-
-	// @return SceneDepth
-	float ConvertFromDeviceZ(float DeviceZ) const
-	{
-		DeviceZ = FMath::Min(DeviceZ, 1 - Z_PRECISION);
-
-		// for depth to linear conversion
-		const FVector2f InvDeviceZToWorldZ(0.1f, 0.1f);
-
-		return 1.0f / (DeviceZ * InvDeviceZToWorldZ.X - InvDeviceZToWorldZ.Y);
-	}
-
-	ECubeFace CubeFace = CubeFace_MAX;
-	ERangeCompressionMode CompressionMode = RCM_UNorm;
-	bool bLinearToGamma = true;	
-	float MaxDepthRange = 16000.0f;
-	bool bOutputStencil = false;
-	uint8 MipLevel = 0;
-	int32 ArrayIndex = 0;
-	uint32 GPUIndex = 0;
-};
-
 /** Info for supporting the vertex element types */
 class FVertexElementTypeSupportInfo
 {
@@ -1207,15 +1102,7 @@ public:
 	RHI_API void FromString(const FStringView& Src);
 };
 
-/**
- *	Screen Resolution
- */
-struct FScreenResolutionRHI
-{
-	uint32	Width;
-	uint32	Height;
-	uint32	RefreshRate;
-};
+
 
 /**
  *	Viewport bounds structure to set multiple view ports for the geometry shader
@@ -1239,7 +1126,6 @@ struct FViewportBounds
 };
 
 
-typedef TArray<FScreenResolutionRHI>	FScreenResolutionArray;
 
 struct FVRamAllocation
 {
@@ -1268,118 +1154,6 @@ struct FRHIResourceInfo
 	bool bResident{ true };
 };
 
-
-struct FRHICopyTextureInfo
-{
-	FIntRect GetSourceRect() const
-	{
-		return FIntRect(SourcePosition.X, SourcePosition.Y, SourcePosition.X + Size.X, SourcePosition.Y + Size.Y);
-	}
-	
-	FIntRect GetDestRect() const
-	{
-		return FIntRect(DestPosition.X, DestPosition.Y, DestPosition.X + Size.X, DestPosition.Y + Size.Y);
-	}
-
-	// Number of texels to copy. By default it will copy the whole resource if no size is specified.
-	FIntVector Size = FIntVector::ZeroValue;
-
-	// Position of the copy from the source texture/to destination texture
-	FIntVector SourcePosition = FIntVector::ZeroValue;
-	FIntVector DestPosition = FIntVector::ZeroValue;
-
-	uint32 SourceSliceIndex = 0;
-	uint32 DestSliceIndex = 0;
-	uint32 NumSlices = 1;
-
-	// Mips to copy and destination mips
-	uint32 SourceMipIndex = 0;
-	uint32 DestMipIndex = 0;
-	uint32 NumMips = 1;
-};
-
-
-/** specifies an update region for a texture */
-struct FUpdateTextureRegion2D
-{
-	/** offset in texture */
-	uint32 DestX;
-	uint32 DestY;
-	
-	/** offset in source image data */
-	int32 SrcX;
-	int32 SrcY;
-	
-	/** size of region to copy */
-	uint32 Width;
-	uint32 Height;
-
-	FUpdateTextureRegion2D()
-	{}
-
-	FUpdateTextureRegion2D(uint32 InDestX, uint32 InDestY, int32 InSrcX, int32 InSrcY, uint32 InWidth, uint32 InHeight)
-	:	DestX(InDestX)
-	,	DestY(InDestY)
-	,	SrcX(InSrcX)
-	,	SrcY(InSrcY)
-	,	Width(InWidth)
-	,	Height(InHeight)
-	{}
-};
-
-/** specifies an update region for a texture */
-struct FUpdateTextureRegion3D
-{
-	/** offset in texture */
-	uint32 DestX;
-	uint32 DestY;
-	uint32 DestZ;
-
-	/** offset in source image data */
-	int32 SrcX;
-	int32 SrcY;
-	int32 SrcZ;
-
-	/** size of region to copy */
-	uint32 Width;
-	uint32 Height;
-	uint32 Depth;
-
-	FUpdateTextureRegion3D()
-	{}
-
-	FUpdateTextureRegion3D(uint32 InDestX, uint32 InDestY, uint32 InDestZ, int32 InSrcX, int32 InSrcY, int32 InSrcZ, uint32 InWidth, uint32 InHeight, uint32 InDepth)
-	:	DestX(InDestX)
-	,	DestY(InDestY)
-	,	DestZ(InDestZ)
-	,	SrcX(InSrcX)
-	,	SrcY(InSrcY)
-	,	SrcZ(InSrcZ)
-	,	Width(InWidth)
-	,	Height(InHeight)
-	,	Depth(InDepth)
-	{}
-
-	FUpdateTextureRegion3D(FIntVector InDest, FIntVector InSource, FIntVector InSourceSize)
-		: DestX(InDest.X)
-		, DestY(InDest.Y)
-		, DestZ(InDest.Z)
-		, SrcX(InSource.X)
-		, SrcY(InSource.Y)
-		, SrcZ(InSource.Z)
-		, Width(InSourceSize.X)
-		, Height(InSourceSize.Y)
-		, Depth(InSourceSize.Z)
-	{}
-};
-
-struct FRHIBufferRange
-{
-	class FRHIBuffer* Buffer{ nullptr };
-	uint64 Offset{ 0 };
-	uint64 Size{ 0 };
-};
-
 struct FRHIDispatchIndirectParametersNoPadding
 {
 	uint32 ThreadGroupCountX;
@@ -1414,216 +1188,6 @@ struct FRHIDrawIndexedIndirectParameters
 	uint32 StartInstanceLocation;
 };
 
-
-struct FTextureMemoryStats
-{
-	// Hardware state (never change after device creation):
-
-	// -1 if unknown, in bytes
-	int64 DedicatedVideoMemory;
-	// -1 if unknown, in bytes
-	int64 DedicatedSystemMemory;
-	// -1 if unknown, in bytes
-	int64 SharedSystemMemory;
-	// Total amount of "graphics memory" that we think we can use for all our graphics resources, in bytes. -1 if unknown.
-	int64 TotalGraphicsMemory;
-
-	// Size of allocated memory, in bytes
-	int64 AllocatedMemorySize;
-	// Size of the largest memory fragment, in bytes
-	int64 LargestContiguousAllocation;
-	// 0 if streaming pool size limitation is disabled, in bytes
-	int64 TexturePoolSize;
-	// Upcoming adjustments to allocated memory, in bytes (async reallocations)
-	int32 PendingMemoryAdjustment;
-
-	// defaults
-	FTextureMemoryStats()
-		: DedicatedVideoMemory(-1)
-		, DedicatedSystemMemory(-1)
-		, SharedSystemMemory(-1)
-		, TotalGraphicsMemory(-1)
-		, AllocatedMemorySize(0)
-		, LargestContiguousAllocation(0)
-		, TexturePoolSize(0)
-		, PendingMemoryAdjustment(0)
-	{
-	}
-
-	bool AreHardwareStatsValid() const
-	{
-		// pardon the redundancy, have a broken compiler (__EMSCRIPTEN__) that needs these types spelled out...
-		return ((int64)DedicatedVideoMemory >= 0 && (int64)DedicatedSystemMemory >= 0 && (int64)SharedSystemMemory >= 0);
-	}
-
-	bool IsUsingLimitedPoolSize() const
-	{
-		return TexturePoolSize > 0;
-	}
-
-	int64 ComputeAvailableMemorySize() const
-	{
-		return FMath::Max(TexturePoolSize - AllocatedMemorySize, (int64)0);
-	}
-};
-
-// GPU stats
-#ifndef HAS_GPU_STATS
-#define HAS_GPU_STATS ((STATS || CSV_PROFILER || GPUPROFILERTRACE_ENABLED) && (!UE_BUILD_SHIPPING))
-#endif
-
-extern RHI_API int32 GNumDrawCallsRHI      [MAX_NUM_GPUS];
-extern RHI_API int32 GNumPrimitivesDrawnRHI[MAX_NUM_GPUS];
-
-#if HAS_GPU_STATS
-
-	struct RHI_API FDrawCallCategoryName
-	{
-		FDrawCallCategoryName();
-		FDrawCallCategoryName(FName InName);
-
-		bool ShouldCountDraws() const { return Index != -1; }
-
-		FName  const Name;
-		uint32 const Index;
-
-		static constexpr int32 MAX_DRAWCALL_CATEGORY = 256;
-
-		struct FManager
-		{
-			TStaticArray<FDrawCallCategoryName*, MAX_DRAWCALL_CATEGORY> Array;
-
-			// A backup of the counts that can be used to display on screen to avoid flickering.
-			TStaticArray<TStaticArray<int32, MAX_NUM_GPUS>, MAX_DRAWCALL_CATEGORY> DisplayCounts;
-
-			int32 NumCategory;
-
-			FManager();
-		};
-
-		static FManager& GetManager();
-	};
-
-	// RHI counter stats.
-	DECLARE_DWORD_COUNTER_STAT_EXTERN(TEXT("DrawPrimitive calls"),STAT_RHIDrawPrimitiveCalls,STATGROUP_RHI,RHI_API);
-	DECLARE_DWORD_COUNTER_STAT_EXTERN(TEXT("Triangles drawn")    ,STAT_RHITriangles         ,STATGROUP_RHI,RHI_API);
-	DECLARE_DWORD_COUNTER_STAT_EXTERN(TEXT("Lines drawn")        ,STAT_RHILines             ,STATGROUP_RHI,RHI_API);
-
-#else
-
-	struct FDrawCallCategoryName {};
-
-#endif
-
-// Macros for use inside RHI context Draw functions.
-// Updates the Stats structure on the underlying RHI context class (IRHICommandContext)
-
-#define RHI_DRAW_CALL_INC() do { Stats->Draws++; } while (false)
-#define RHI_DRAW_CALL_STATS(PrimitiveType,NumPrimitives)                  \
-	do                                                                    \
-	{												                      \
-		switch (PrimitiveType)                                            \
-		{                                                                 \
-		case PT_TriangleList : Stats->Triangles  += NumPrimitives; break; \
-		case PT_TriangleStrip: Stats->Triangles  += NumPrimitives; break; \
-		case PT_LineList     : Stats->Lines      += NumPrimitives; break; \
-		case PT_QuadList     : Stats->Quads      += NumPrimitives; break; \
-		case PT_PointList    : Stats->Points     += NumPrimitives; break; \
-		case PT_RectList     : Stats->Rectangles += NumPrimitives; break; \
-		}                                                                 \
-		Stats->Draws++;                                                   \
-	} while(false)
-
-struct FRHIDrawStats
-{
-#if HAS_GPU_STATS
-	// The +1 is for "uncategorised"
-	static constexpr int32 NumCategories = FDrawCallCategoryName::MAX_DRAWCALL_CATEGORY + 1;
-#else
-	static constexpr int32 NumCategories = 1;
-#endif
-
-	static constexpr int32 NoCategory = NumCategories - 1;
-
-	struct FPerCategoryStats
-	{
-		uint32 Draws;
-		uint32 Triangles;
-		uint32 Lines;
-		uint32 Quads;
-		uint32 Points;
-		uint32 Rectangles;
-
-		uint32 GetTotalPrimitives() const
-		{
-			return Triangles
-				+ Lines
-				+ Quads
-				+ Points
-				+ Rectangles;
-		}
-
-		FPerCategoryStats& operator += (FPerCategoryStats const& RHS)
-		{
-			Draws      += RHS.Draws;
-			Triangles  += RHS.Triangles;
-			Lines      += RHS.Lines;
-			Quads      += RHS.Quads;
-			Points     += RHS.Points;
-			Rectangles += RHS.Rectangles;
-			return *this;
-		}
-	};
-
-	struct FPerGPUStats
-	{
-		FPerCategoryStats& GetCategory(uint32 Category)
-		{
-			checkSlow(Category < UE_ARRAY_COUNT(Categories));
-			return Categories[Category];
-		}
-
-	private:
-		FPerCategoryStats Categories[NumCategories];
-	};
-
-	FPerGPUStats& GetGPU(uint32 GPUIndex)
-	{
-		checkSlow(GPUIndex < UE_ARRAY_COUNT(GPUs));
-		return GPUs[GPUIndex];
-	}
-
-	FRHIDrawStats()
-	{
-		Reset();
-	}
-
-	void Reset()
-	{
-		FMemory::Memzero(*this);
-	}
-
-	RHI_API void Accumulate(FRHIDrawStats& RHS);
-
-private:
-	FPerGPUStats GPUs[MAX_NUM_GPUS];
-};
-
-// RHI memory stats.
-DECLARE_MEMORY_STAT_POOL_EXTERN(TEXT("Render target memory 2D"),STAT_RenderTargetMemory2D,STATGROUP_RHI,FPlatformMemory::MCR_GPU,RHI_API);
-DECLARE_MEMORY_STAT_POOL_EXTERN(TEXT("Render target memory 3D"),STAT_RenderTargetMemory3D,STATGROUP_RHI,FPlatformMemory::MCR_GPU,RHI_API);
-DECLARE_MEMORY_STAT_POOL_EXTERN(TEXT("Render target memory Cube"),STAT_RenderTargetMemoryCube,STATGROUP_RHI,FPlatformMemory::MCR_GPU,RHI_API);
-DECLARE_MEMORY_STAT_POOL_EXTERN(TEXT("Texture memory 2D"),STAT_TextureMemory2D,STATGROUP_RHI,FPlatformMemory::MCR_GPU,RHI_API);
-DECLARE_MEMORY_STAT_POOL_EXTERN(TEXT("Texture memory 3D"),STAT_TextureMemory3D,STATGROUP_RHI,FPlatformMemory::MCR_GPU,RHI_API);
-DECLARE_MEMORY_STAT_POOL_EXTERN(TEXT("Texture memory Cube"),STAT_TextureMemoryCube,STATGROUP_RHI,FPlatformMemory::MCR_GPU,RHI_API);
-DECLARE_MEMORY_STAT_POOL_EXTERN(TEXT("Uniform buffer memory"),STAT_UniformBufferMemory,STATGROUP_RHI,FPlatformMemory::MCR_GPU,RHI_API);
-DECLARE_MEMORY_STAT_POOL_EXTERN(TEXT("Index buffer memory"),STAT_IndexBufferMemory,STATGROUP_RHI,FPlatformMemory::MCR_GPU,RHI_API);
-DECLARE_MEMORY_STAT_POOL_EXTERN(TEXT("Vertex buffer memory"),STAT_VertexBufferMemory,STATGROUP_RHI,FPlatformMemory::MCR_GPU,RHI_API);
-DECLARE_MEMORY_STAT_POOL_EXTERN(TEXT("Ray Tracing Acceleration Structure memory"), STAT_RTAccelerationStructureMemory, STATGROUP_RHI, FPlatformMemory::MCR_GPU, RHI_API);
-DECLARE_MEMORY_STAT_POOL_EXTERN(TEXT("Structured buffer memory"),STAT_StructuredBufferMemory,STATGROUP_RHI,FPlatformMemory::MCR_GPU,RHI_API);
-DECLARE_MEMORY_STAT_POOL_EXTERN(TEXT("Pixel buffer memory"),STAT_PixelBufferMemory,STATGROUP_RHI,FPlatformMemory::MCR_GPU,RHI_API);
-
-
 // RHI base resource types.
 #include "RHIResources.h" // IWYU pragma: keep
 #include "DynamicRHI.h" // IWYU pragma: keep
diff --git a/Engine/Source/Runtime/RHI/Public/RHICommandList.h b/Engine/Source/Runtime/RHI/Public/RHICommandList.h
--- a/Engine/Source/Runtime/RHI/Public/RHICommandList.h
+++ b/Engine/Source/Runtime/RHI/Public/RHICommandList.h
@@ -4,12 +4,8 @@
 	RHICommandList.h: RHI Command List definitions for queueing up & executing later.
 =============================================================================*/
 
-// IWYU pragma: private, include "RHI.h"
-
 #pragma once
 
-// HEADER_UNIT_SKIP - Should always be included through RHI.h (this code needs cleanup.. this file is included by RHIUtilities.h which is included in the middle of RHI.h)
-
 #include "CoreTypes.h"
 #include "Misc/AssertionMacros.h"
 #include "HAL/UnrealMemory.h"
@@ -26,15 +22,19 @@
 #include "GenericPlatform/GenericPlatformProcess.h"
 #include "Misc/MemStack.h"
 #include "Misc/App.h"
-#include "Stats/Stats.h"
+#include "RHIStats.h"
 #include "HAL/IConsoleManager.h"
 #include "Async/TaskGraphInterfaces.h"
 #include "HAL/LowLevelMemTracker.h"
 #include "ProfilingDebugging/CsvProfiler.h"
 #include "ProfilingDebugging/CpuProfilerTrace.h"
+#include "RHIBreadcrumbs.h"
 #include "RHITextureReference.h"
 #include "Trace/Trace.h"
 
+#include "DynamicRHI.h"
+#include "RHITypes.h"
+
 CSV_DECLARE_CATEGORY_MODULE_EXTERN(RHI_API, RHITStalls);
 CSV_DECLARE_CATEGORY_MODULE_EXTERN(RHI_API, RHITFlushes);
 
@@ -164,6 +164,72 @@ extern RHI_API bool GEnableAsyncCompute;
 extern RHI_API TAutoConsoleVariable<int32> CVarRHICmdWidth;
 extern RHI_API TAutoConsoleVariable<int32> CVarRHICmdFlushRenderThreadTasks;
 
+
+struct FRHICopyTextureInfo
+{
+	FIntRect GetSourceRect() const
+	{
+		return FIntRect(SourcePosition.X, SourcePosition.Y, SourcePosition.X + Size.X, SourcePosition.Y + Size.Y);
+	}
+
+	FIntRect GetDestRect() const
+	{
+		return FIntRect(DestPosition.X, DestPosition.Y, DestPosition.X + Size.X, DestPosition.Y + Size.Y);
+	}
+
+	// Number of texels to copy. By default it will copy the whole resource if no size is specified.
+	FIntVector Size = FIntVector::ZeroValue;
+
+	// Position of the copy from the source texture/to destination texture
+	FIntVector SourcePosition = FIntVector::ZeroValue;
+	FIntVector DestPosition = FIntVector::ZeroValue;
+
+	uint32 SourceSliceIndex = 0;
+	uint32 DestSliceIndex = 0;
+	uint32 NumSlices = 1;
+
+	// Mips to copy and destination mips
+	uint32 SourceMipIndex = 0;
+	uint32 DestMipIndex = 0;
+	uint32 NumMips = 1;
+};
+
+struct FRHIBufferRange
+{
+	class FRHIBuffer* Buffer{ nullptr };
+	uint64 Offset{ 0 };
+	uint64 Size{ 0 };
+};
+
+/** Struct to hold common data between begin/end updatetexture3d */
+struct FUpdateTexture3DData
+{
+	FUpdateTexture3DData(FRHITexture* InTexture, uint32 InMipIndex, const struct FUpdateTextureRegion3D& InUpdateRegion, uint32 InSourceRowPitch, uint32 InSourceDepthPitch, uint8* InSourceData, uint32 InDataSizeBytes, uint32 InFrameNumber)
+		: Texture(InTexture)
+		, MipIndex(InMipIndex)
+		, UpdateRegion(InUpdateRegion)
+		, RowPitch(InSourceRowPitch)
+		, DepthPitch(InSourceDepthPitch)
+		, Data(InSourceData)
+		, DataSizeBytes(InDataSizeBytes)
+		, FrameNumber(InFrameNumber)
+	{
+	}
+
+	FRHITexture* Texture;
+	uint32 MipIndex;
+	FUpdateTextureRegion3D UpdateRegion;
+	uint32 RowPitch;
+	uint32 DepthPitch;
+	uint8* Data;
+	uint32 DataSizeBytes;
+	uint32 FrameNumber;
+	uint8 PlatformData[64];
+
+private:
+	FUpdateTexture3DData();
+};
+
 #if RHI_RAYTRACING
 struct FRayTracingShaderBindings
 {
@@ -3598,11 +3664,6 @@ public:
 
 	FORCEINLINE_DEBUGGABLE void DrawIndexedPrimitive(FRHIBuffer* IndexBuffer, int32 BaseVertexIndex, uint32 FirstInstance, uint32 NumVertices, uint32 StartIndex, uint32 NumPrimitives, uint32 NumInstances)
 	{
-		if (!IndexBuffer)
-		{
-			UE_LOG(LogRHI, Fatal, TEXT("Tried to call DrawIndexedPrimitive with null IndexBuffer!"));
-		}
-
 		//check(IsOutsideRenderPass());
 		PreDraw();
 		if (Bypass())
@@ -4054,7 +4115,6 @@ public:
 	 */
 	FORCEINLINE_DEBUGGABLE void RayTraceDispatchIndirect(FRayTracingPipelineState* Pipeline, FRHIRayTracingShader* RayGenShader, FRHIRayTracingScene* Scene, const FRayTracingShaderBindings& GlobalResourceBindings, FRHIBuffer* ArgumentBuffer, uint32 ArgumentOffset)
 	{
-		checkf(GRHISupportsRayTracingDispatchIndirect, TEXT("Indirect ray tracing is not supported on this machine."));
 		if (Bypass())
 		{
 			GetContext().RHIRayTraceDispatchIndirect(GetRHIRayTracingPipelineState(Pipeline), RayGenShader, Scene, GlobalResourceBindings, ArgumentBuffer, ArgumentOffset);
diff --git a/Engine/Source/Runtime/RHI/Public/RHIContext.h b/Engine/Source/Runtime/RHI/Public/RHIContext.h
--- a/Engine/Source/Runtime/RHI/Public/RHIContext.h
+++ b/Engine/Source/Runtime/RHI/Public/RHIContext.h
@@ -6,8 +6,6 @@
 
 #pragma once
 
-// HEADER_UNIT_SKIP - Should always be included after RHI.h (this code needs cleanup.. this file is included in the middle of RHI.h)
-
 #include "Misc/AssertionMacros.h"
 #include "Math/Color.h"
 #include "Math/IntPoint.h"
@@ -35,6 +33,12 @@ struct FRayTracingLocalShaderBindings;
 enum class ERayTracingBindingType : uint8;
 enum class EAsyncComputeBudget;
 
+struct FRHIBufferRange;
+struct FRHIPerCategoryDrawStats;
+struct FRHIDrawStats;
+struct FRHICopyTextureInfo;
+
+
 #define VALIDATE_UNIFORM_BUFFER_STATIC_BINDINGS (!UE_BUILD_SHIPPING && !UE_BUILD_TEST)
 
 /** A list of static uniform buffer bindings. */
@@ -469,13 +473,10 @@ public:
 	virtual void RHIPostExternalCommandsReset() { }
 
 protected:
-	FRHIDrawStats::FPerCategoryStats* Stats = nullptr;
+	FRHIPerCategoryDrawStats* Stats = nullptr;
 
 public:
-	void StatsSetCategory(FRHIDrawStats* InStats, uint32 InCategoryID, uint32 InGPUIndex)
-	{
-		Stats = &InStats->GetGPU(InGPUIndex).GetCategory(InCategoryID);
-	}
+	RHI_API void StatsSetCategory(FRHIDrawStats* InStats, uint32 InCategoryID, uint32 InGPUIndex);
 
 #if WITH_MGPU || ENABLE_RHI_VALIDATION
 	virtual
@@ -890,23 +891,5 @@ PRAGMA_ENABLE_DEPRECATION_WARNINGS
 	}
 
 private:
-	void SetGraphicsPipelineStateFromInitializer(const FGraphicsPipelineStateInitializer& PsoInit, uint32 StencilRef, bool bApplyAdditionalState)
-	{
-		RHISetBoundShaderState(
-			RHICreateBoundShaderState(
-				PsoInit.BoundShaderState.VertexDeclarationRHI,
-				PsoInit.BoundShaderState.VertexShaderRHI,
-				PsoInit.BoundShaderState.PixelShaderRHI,
-				PsoInit.BoundShaderState.GetGeometryShader()
-			).GetReference()
-		);
-
-		RHISetDepthStencilState(PsoInit.DepthStencilState, StencilRef);
-		RHISetRasterizerState(PsoInit.RasterizerState);
-		RHISetBlendState(PsoInit.BlendState, FLinearColor(1.0f, 1.0f, 1.0f));
-		if (GSupportsDepthBoundsTest)
-		{
-			RHIEnableDepthBoundsTest(PsoInit.bDepthBounds);
-		}
-	}
+	RHI_API void SetGraphicsPipelineStateFromInitializer(const FGraphicsPipelineStateInitializer& PsoInit, uint32 StencilRef, bool bApplyAdditionalState);
 };
diff --git a/Engine/Source/Runtime/RHI/Public/RHIDefinitions.h b/Engine/Source/Runtime/RHI/Public/RHIDefinitions.h
--- a/Engine/Source/Runtime/RHI/Public/RHIDefinitions.h
+++ b/Engine/Source/Runtime/RHI/Public/RHIDefinitions.h
@@ -37,6 +37,10 @@ static_assert(sizeof(void*) <= SHADER_PARAMETER_POINTER_ALIGNMENT, "The alignmen
 #define RHI_RAYTRACING 0
 #endif
 
+#ifndef HAS_GPU_STATS
+#define HAS_GPU_STATS ((STATS || CSV_PROFILER || GPUPROFILERTRACE_ENABLED) && (!UE_BUILD_SHIPPING))
+#endif
+
 enum class ERHIInterfaceType
 {
 	Hidden,
@@ -1321,6 +1325,14 @@ inline ERHIResourceType GetRHIResourceType(ETextureDimension Dimension)
 	#define GEOMETRY_SHADER(GeometryShader)	nullptr
 #endif
 
+/** Screen Resolution */
+struct FScreenResolutionRHI
+{
+	uint32	Width;
+	uint32	Height;
+	uint32	RefreshRate;
+};
+
 #if UE_ENABLE_INCLUDE_ORDER_DEPRECATED_IN_5_2
 #include "CoreMinimal.h"
 #include "DataDrivenShaderPlatformInfo.h"
diff --git a/Engine/Source/Runtime/RHI/Public/RHIResources.h b/Engine/Source/Runtime/RHI/Public/RHIResources.h
--- a/Engine/Source/Runtime/RHI/Public/RHIResources.h
+++ b/Engine/Source/Runtime/RHI/Public/RHIResources.h
@@ -1,13 +1,7 @@
 // Copyright Epic Games, Inc. All Rights Reserved.
 
-// IWYU pragma: private, include "RHI.h"
-
 #pragma once
 
-// HEADER_UNIT_SKIP - Should always be included through RHI.h (this code needs cleanup.. this file is included in the middle of RHI.h)
-
-#include <atomic>
-
 #include "CoreTypes.h"
 #include "Misc/AssertionMacros.h"
 #include "HAL/UnrealMemory.h"
@@ -34,6 +28,8 @@
 #include "RHITransition.h"
 #include "MultiGPU.h"
 
+#include <atomic>
+
 class FRHIComputeCommandList;
 class FRHICommandListImmediate;
 class FRHITextureReference;
diff --git a/Engine/Source/Runtime/RHI/Public/RHIStats.h b/Engine/Source/Runtime/RHI/Public/RHIStats.h
new file mode 100644
--- /dev/null
+++ b/Engine/Source/Runtime/RHI/Public/RHIStats.h
@@ -0,0 +1,216 @@
+// Copyright Epic Games, Inc. All Rights Reserved.
+
+#pragma once
+
+#include "Containers/StaticArray.h"
+#include "MultiGPU.h"
+#include "RHIDefinitions.h"
+#include "Stats/Stats2.h"
+
+struct FTextureMemoryStats
+{
+	// Hardware state (never change after device creation):
+
+	// -1 if unknown, in bytes
+	int64 DedicatedVideoMemory;
+	// -1 if unknown, in bytes
+	int64 DedicatedSystemMemory;
+	// -1 if unknown, in bytes
+	int64 SharedSystemMemory;
+	// Total amount of "graphics memory" that we think we can use for all our graphics resources, in bytes. -1 if unknown.
+	int64 TotalGraphicsMemory;
+
+	// Size of allocated memory, in bytes
+	int64 AllocatedMemorySize;
+	// Size of the largest memory fragment, in bytes
+	int64 LargestContiguousAllocation;
+	// 0 if streaming pool size limitation is disabled, in bytes
+	int64 TexturePoolSize;
+	// Upcoming adjustments to allocated memory, in bytes (async reallocations)
+	int32 PendingMemoryAdjustment;
+
+	// defaults
+	FTextureMemoryStats()
+		: DedicatedVideoMemory(-1)
+		, DedicatedSystemMemory(-1)
+		, SharedSystemMemory(-1)
+		, TotalGraphicsMemory(-1)
+		, AllocatedMemorySize(0)
+		, LargestContiguousAllocation(0)
+		, TexturePoolSize(0)
+		, PendingMemoryAdjustment(0)
+	{
+	}
+
+	bool AreHardwareStatsValid() const
+	{
+		// pardon the redundancy, have a broken compiler (__EMSCRIPTEN__) that needs these types spelled out...
+		return ((int64)DedicatedVideoMemory >= 0 && (int64)DedicatedSystemMemory >= 0 && (int64)SharedSystemMemory >= 0);
+	}
+
+	bool IsUsingLimitedPoolSize() const
+	{
+		return TexturePoolSize > 0;
+	}
+
+	int64 ComputeAvailableMemorySize() const
+	{
+		return FMath::Max(TexturePoolSize - AllocatedMemorySize, (int64)0);
+	}
+};
+
+
+// GPU stats
+
+extern RHI_API int32 GNumDrawCallsRHI[MAX_NUM_GPUS];
+extern RHI_API int32 GNumPrimitivesDrawnRHI[MAX_NUM_GPUS];
+
+#if HAS_GPU_STATS
+
+struct RHI_API FDrawCallCategoryName
+{
+	FDrawCallCategoryName();
+	FDrawCallCategoryName(FName InName);
+
+	bool ShouldCountDraws() const { return Index != -1; }
+
+	FName  const Name;
+	uint32 const Index;
+
+	static constexpr int32 MAX_DRAWCALL_CATEGORY = 256;
+
+	struct FManager
+	{
+		TStaticArray<FDrawCallCategoryName*, MAX_DRAWCALL_CATEGORY> Array;
+
+		// A backup of the counts that can be used to display on screen to avoid flickering.
+		TStaticArray<TStaticArray<int32, MAX_NUM_GPUS>, MAX_DRAWCALL_CATEGORY> DisplayCounts;
+
+		int32 NumCategory;
+
+		FManager();
+	};
+
+	static FManager& GetManager();
+};
+
+// RHI counter stats.
+DECLARE_DWORD_COUNTER_STAT_EXTERN(TEXT("DrawPrimitive calls"), STAT_RHIDrawPrimitiveCalls, STATGROUP_RHI, RHI_API);
+DECLARE_DWORD_COUNTER_STAT_EXTERN(TEXT("Triangles drawn"), STAT_RHITriangles, STATGROUP_RHI, RHI_API);
+DECLARE_DWORD_COUNTER_STAT_EXTERN(TEXT("Lines drawn"), STAT_RHILines, STATGROUP_RHI, RHI_API);
+
+#else
+
+struct FDrawCallCategoryName {};
+
+#endif
+
+// Macros for use inside RHI context Draw functions.
+// Updates the Stats structure on the underlying RHI context class (IRHICommandContext)
+
+#define RHI_DRAW_CALL_INC() do { Stats->Draws++; } while (false)
+#define RHI_DRAW_CALL_STATS(PrimitiveType,NumPrimitives)                  \
+	do                                                                    \
+	{												                      \
+		switch (PrimitiveType)                                            \
+		{                                                                 \
+		case PT_TriangleList : Stats->Triangles  += NumPrimitives; break; \
+		case PT_TriangleStrip: Stats->Triangles  += NumPrimitives; break; \
+		case PT_LineList     : Stats->Lines      += NumPrimitives; break; \
+		case PT_QuadList     : Stats->Quads      += NumPrimitives; break; \
+		case PT_PointList    : Stats->Points     += NumPrimitives; break; \
+		case PT_RectList     : Stats->Rectangles += NumPrimitives; break; \
+		}                                                                 \
+		Stats->Draws++;                                                   \
+	} while(false)
+
+struct FRHIPerCategoryDrawStats
+{
+	uint32 Draws;
+	uint32 Triangles;
+	uint32 Lines;
+	uint32 Quads;
+	uint32 Points;
+	uint32 Rectangles;
+
+	uint32 GetTotalPrimitives() const
+	{
+		return Triangles
+			+ Lines
+			+ Quads
+			+ Points
+			+ Rectangles;
+	}
+
+	FRHIPerCategoryDrawStats& operator += (FRHIPerCategoryDrawStats const& RHS)
+	{
+		Draws += RHS.Draws;
+		Triangles += RHS.Triangles;
+		Lines += RHS.Lines;
+		Quads += RHS.Quads;
+		Points += RHS.Points;
+		Rectangles += RHS.Rectangles;
+		return *this;
+	}
+};
+
+struct FRHIDrawStats
+{
+#if HAS_GPU_STATS
+	// The +1 is for "uncategorised"
+	static constexpr int32 NumCategories = FDrawCallCategoryName::MAX_DRAWCALL_CATEGORY + 1;
+#else
+	static constexpr int32 NumCategories = 1;
+#endif
+
+	static constexpr int32 NoCategory = NumCategories - 1;
+
+	using FPerCategoryStats = FRHIPerCategoryDrawStats;
+
+	struct FPerGPUStats
+	{
+		FPerCategoryStats& GetCategory(uint32 Category)
+		{
+			checkSlow(Category < UE_ARRAY_COUNT(Categories));
+			return Categories[Category];
+		}
+
+	private:
+		FPerCategoryStats Categories[NumCategories];
+	};
+
+	FPerGPUStats& GetGPU(uint32 GPUIndex)
+	{
+		checkSlow(GPUIndex < UE_ARRAY_COUNT(GPUs));
+		return GPUs[GPUIndex];
+	}
+
+	FRHIDrawStats()
+	{
+		Reset();
+	}
+
+	void Reset()
+	{
+		FMemory::Memzero(*this);
+	}
+
+	RHI_API void Accumulate(FRHIDrawStats& RHS);
+
+private:
+	FPerGPUStats GPUs[MAX_NUM_GPUS];
+};
+
+// RHI memory stats.
+DECLARE_MEMORY_STAT_POOL_EXTERN(TEXT("Render target memory 2D"), STAT_RenderTargetMemory2D, STATGROUP_RHI, FPlatformMemory::MCR_GPU, RHI_API);
+DECLARE_MEMORY_STAT_POOL_EXTERN(TEXT("Render target memory 3D"), STAT_RenderTargetMemory3D, STATGROUP_RHI, FPlatformMemory::MCR_GPU, RHI_API);
+DECLARE_MEMORY_STAT_POOL_EXTERN(TEXT("Render target memory Cube"), STAT_RenderTargetMemoryCube, STATGROUP_RHI, FPlatformMemory::MCR_GPU, RHI_API);
+DECLARE_MEMORY_STAT_POOL_EXTERN(TEXT("Texture memory 2D"), STAT_TextureMemory2D, STATGROUP_RHI, FPlatformMemory::MCR_GPU, RHI_API);
+DECLARE_MEMORY_STAT_POOL_EXTERN(TEXT("Texture memory 3D"), STAT_TextureMemory3D, STATGROUP_RHI, FPlatformMemory::MCR_GPU, RHI_API);
+DECLARE_MEMORY_STAT_POOL_EXTERN(TEXT("Texture memory Cube"), STAT_TextureMemoryCube, STATGROUP_RHI, FPlatformMemory::MCR_GPU, RHI_API);
+DECLARE_MEMORY_STAT_POOL_EXTERN(TEXT("Uniform buffer memory"), STAT_UniformBufferMemory, STATGROUP_RHI, FPlatformMemory::MCR_GPU, RHI_API);
+DECLARE_MEMORY_STAT_POOL_EXTERN(TEXT("Index buffer memory"), STAT_IndexBufferMemory, STATGROUP_RHI, FPlatformMemory::MCR_GPU, RHI_API);
+DECLARE_MEMORY_STAT_POOL_EXTERN(TEXT("Vertex buffer memory"), STAT_VertexBufferMemory, STATGROUP_RHI, FPlatformMemory::MCR_GPU, RHI_API);
+DECLARE_MEMORY_STAT_POOL_EXTERN(TEXT("Ray Tracing Acceleration Structure memory"), STAT_RTAccelerationStructureMemory, STATGROUP_RHI, FPlatformMemory::MCR_GPU, RHI_API);
+DECLARE_MEMORY_STAT_POOL_EXTERN(TEXT("Structured buffer memory"), STAT_StructuredBufferMemory, STATGROUP_RHI, FPlatformMemory::MCR_GPU, RHI_API);
+DECLARE_MEMORY_STAT_POOL_EXTERN(TEXT("Pixel buffer memory"), STAT_PixelBufferMemory, STATGROUP_RHI, FPlatformMemory::MCR_GPU, RHI_API);
diff --git a/Engine/Source/Runtime/RHI/Public/RHITypes.h b/Engine/Source/Runtime/RHI/Public/RHITypes.h
new file mode 100644
--- /dev/null
+++ b/Engine/Source/Runtime/RHI/Public/RHITypes.h
@@ -0,0 +1,191 @@
+// Copyright Epic Games, Inc. All Rights Reserved.
+
+/*=============================================================================
+	RHITypes.h: Render Hardware Interface types
+		(that require linking or heavier includes).
+=============================================================================*/
+
+#pragma once
+
+#include "RHIDefinitions.h"
+#include "Math/PerspectiveMatrix.h"
+#include "Math/IntVector.h"
+
+/** to customize the RHIReadSurfaceData() output */
+class FReadSurfaceDataFlags
+{
+public:
+	// @param InCompressionMode defines the value input range that is mapped to output range
+	// @param InCubeFace defined which cubemap side is used, only required for cubemap content, then it needs to be a valid side
+	FReadSurfaceDataFlags(ERangeCompressionMode InCompressionMode = RCM_UNorm, ECubeFace InCubeFace = CubeFace_MAX)
+		:CubeFace(InCubeFace), CompressionMode(InCompressionMode)
+	{
+	}
+
+	ECubeFace GetCubeFace() const
+	{
+		checkSlow(CubeFace <= CubeFace_NegZ);
+		return CubeFace;
+	}
+
+	ERangeCompressionMode GetCompressionMode() const
+	{
+		return CompressionMode;
+	}
+
+	void SetLinearToGamma(bool Value)
+	{
+		bLinearToGamma = Value;
+	}
+
+	bool GetLinearToGamma() const
+	{
+		return bLinearToGamma;
+	}
+
+	void SetOutputStencil(bool Value)
+	{
+		bOutputStencil = Value;
+	}
+
+	bool GetOutputStencil() const
+	{
+		return bOutputStencil;
+	}
+
+	void SetMip(uint8 InMipLevel)
+	{
+		MipLevel = InMipLevel;
+	}
+
+	uint8 GetMip() const
+	{
+		return MipLevel;
+	}
+
+	void SetMaxDepthRange(float Value)
+	{
+		MaxDepthRange = Value;
+	}
+
+	float ComputeNormalizedDepth(float DeviceZ) const
+	{
+		return FMath::Abs(ConvertFromDeviceZ(DeviceZ) / MaxDepthRange);
+	}
+
+	void SetGPUIndex(uint32 InGPUIndex)
+	{
+		GPUIndex = InGPUIndex;
+	}
+
+	uint32 GetGPUIndex() const
+	{
+		return GPUIndex;
+	}
+
+	void SetArrayIndex(int32 InArrayIndex)
+	{
+		ArrayIndex = InArrayIndex;
+	}
+
+	int32 GetArrayIndex() const
+	{
+		return ArrayIndex;
+	}
+
+private:
+
+	// @return SceneDepth
+	float ConvertFromDeviceZ(float DeviceZ) const
+	{
+		DeviceZ = FMath::Min(DeviceZ, 1 - Z_PRECISION);
+
+		// for depth to linear conversion
+		const FVector2f InvDeviceZToWorldZ(0.1f, 0.1f);
+
+		return 1.0f / (DeviceZ * InvDeviceZToWorldZ.X - InvDeviceZToWorldZ.Y);
+	}
+
+	ECubeFace CubeFace = CubeFace_MAX;
+	ERangeCompressionMode CompressionMode = RCM_UNorm;
+	bool bLinearToGamma = true;
+	float MaxDepthRange = 16000.0f;
+	bool bOutputStencil = false;
+	uint8 MipLevel = 0;
+	int32 ArrayIndex = 0;
+	uint32 GPUIndex = 0;
+};
+
+/** specifies an update region for a texture */
+struct FUpdateTextureRegion2D
+{
+	/** offset in texture */
+	uint32 DestX;
+	uint32 DestY;
+
+	/** offset in source image data */
+	int32 SrcX;
+	int32 SrcY;
+
+	/** size of region to copy */
+	uint32 Width;
+	uint32 Height;
+
+	FUpdateTextureRegion2D()
+	{}
+
+	FUpdateTextureRegion2D(uint32 InDestX, uint32 InDestY, int32 InSrcX, int32 InSrcY, uint32 InWidth, uint32 InHeight)
+		: DestX(InDestX)
+		, DestY(InDestY)
+		, SrcX(InSrcX)
+		, SrcY(InSrcY)
+		, Width(InWidth)
+		, Height(InHeight)
+	{}
+};
+
+/** specifies an update region for a texture */
+struct FUpdateTextureRegion3D
+{
+	/** offset in texture */
+	uint32 DestX;
+	uint32 DestY;
+	uint32 DestZ;
+
+	/** offset in source image data */
+	int32 SrcX;
+	int32 SrcY;
+	int32 SrcZ;
+
+	/** size of region to copy */
+	uint32 Width;
+	uint32 Height;
+	uint32 Depth;
+
+	FUpdateTextureRegion3D()
+	{}
+
+	FUpdateTextureRegion3D(uint32 InDestX, uint32 InDestY, uint32 InDestZ, int32 InSrcX, int32 InSrcY, int32 InSrcZ, uint32 InWidth, uint32 InHeight, uint32 InDepth)
+		: DestX(InDestX)
+		, DestY(InDestY)
+		, DestZ(InDestZ)
+		, SrcX(InSrcX)
+		, SrcY(InSrcY)
+		, SrcZ(InSrcZ)
+		, Width(InWidth)
+		, Height(InHeight)
+		, Depth(InDepth)
+	{}
+
+	FUpdateTextureRegion3D(FIntVector InDest, FIntVector InSource, FIntVector InSourceSize)
+		: DestX(InDest.X)
+		, DestY(InDest.Y)
+		, DestZ(InDest.Z)
+		, SrcX(InSource.X)
+		, SrcY(InSource.Y)
+		, SrcZ(InSource.Z)
+		, Width(InSourceSize.X)
+		, Height(InSourceSize.Y)
+		, Depth(InSourceSize.Z)
+	{}
+};
diff --git a/Engine/Source/Runtime/RHI/Public/RHIUtilities.h b/Engine/Source/Runtime/RHI/Public/RHIUtilities.h
--- a/Engine/Source/Runtime/RHI/Public/RHIUtilities.h
+++ b/Engine/Source/Runtime/RHI/Public/RHIUtilities.h
@@ -1,11 +1,7 @@
 // Copyright Epic Games, Inc. All Rights Reserved.
 
-// IWYU pragma: private, include "RHI.h"
-
 #pragma once
 
-// HEADER_UNIT_SKIP - Should always be included through RHI.h (this code needs cleanup.. this file is included in the middle of RHI.h)
-
 #include "CoreTypes.h"
 #include "Misc/AssertionMacros.h"
 #include "HAL/UnrealMemory.h"
@@ -17,7 +13,6 @@
 #include "RHIDefinitions.h"
 #include "RHICommandList.h"
 
-
 static inline bool IsDepthOrStencilFormat(EPixelFormat Format)
 {
 	switch (Format)
@@ -530,80 +525,7 @@ struct FDynamicReadBuffer : public FReadBuffer
  * Convert the ESimpleRenderTargetMode into usable values 
  * @todo: Can we easily put this into a .cpp somewhere?
  */
-inline void DecodeRenderTargetMode(ESimpleRenderTargetMode Mode, ERenderTargetLoadAction& ColorLoadAction, ERenderTargetStoreAction& ColorStoreAction, ERenderTargetLoadAction& DepthLoadAction, ERenderTargetStoreAction& DepthStoreAction, ERenderTargetLoadAction& StencilLoadAction, ERenderTargetStoreAction& StencilStoreAction, FExclusiveDepthStencil DepthStencilUsage)
-{
-	// set defaults
-	ColorStoreAction = ERenderTargetStoreAction::EStore;
-	DepthStoreAction = ERenderTargetStoreAction::EStore;
-	StencilStoreAction = ERenderTargetStoreAction::EStore;
-
-	switch (Mode)
-	{
-	case ESimpleRenderTargetMode::EExistingColorAndDepth:
-		ColorLoadAction = ERenderTargetLoadAction::ELoad;
-		DepthLoadAction = ERenderTargetLoadAction::ELoad;
-		break;
-	case ESimpleRenderTargetMode::EUninitializedColorAndDepth:
-		ColorLoadAction = ERenderTargetLoadAction::ENoAction;
-		DepthLoadAction = ERenderTargetLoadAction::ENoAction;
-		break;
-	case ESimpleRenderTargetMode::EUninitializedColorExistingDepth:
-		ColorLoadAction = ERenderTargetLoadAction::ENoAction;
-		DepthLoadAction = ERenderTargetLoadAction::ELoad;
-		break;
-	case ESimpleRenderTargetMode::EUninitializedColorClearDepth:
-		ColorLoadAction = ERenderTargetLoadAction::ENoAction;
-		DepthLoadAction = ERenderTargetLoadAction::EClear;
-		break;
-	case ESimpleRenderTargetMode::EClearColorExistingDepth:
-		ColorLoadAction = ERenderTargetLoadAction::EClear;
-		DepthLoadAction = ERenderTargetLoadAction::ELoad;
-		break;
-	case ESimpleRenderTargetMode::EClearColorAndDepth:
-		ColorLoadAction = ERenderTargetLoadAction::EClear;
-		DepthLoadAction = ERenderTargetLoadAction::EClear;
-		break;
-	case ESimpleRenderTargetMode::EExistingContents_NoDepthStore:
-		ColorLoadAction = ERenderTargetLoadAction::ELoad;
-		DepthLoadAction = ERenderTargetLoadAction::ELoad;
-		DepthStoreAction = ERenderTargetStoreAction::ENoAction;
-		break;
-	case ESimpleRenderTargetMode::EExistingColorAndClearDepth:
-		ColorLoadAction = ERenderTargetLoadAction::ELoad;
-		DepthLoadAction = ERenderTargetLoadAction::EClear;
-		break;
-	case ESimpleRenderTargetMode::EExistingColorAndDepthAndClearStencil:
-		ColorLoadAction = ERenderTargetLoadAction::ELoad;
-		DepthLoadAction = ERenderTargetLoadAction::ELoad;
-		break;
-	default:
-		UE_LOG(LogRHI, Fatal, TEXT("Using a ESimpleRenderTargetMode that wasn't decoded in DecodeRenderTargetMode [value = %d]"), (int32)Mode);
-	}
-	
-	StencilLoadAction = DepthLoadAction;
-	
-	if (!DepthStencilUsage.IsUsingDepth())
-	{
-		DepthLoadAction = ERenderTargetLoadAction::ENoAction;
-	}
-	
-	//if we aren't writing to depth, there's no reason to store it back out again.  Should save some bandwidth on mobile platforms.
-	if (!DepthStencilUsage.IsDepthWrite())
-	{
-		DepthStoreAction = ERenderTargetStoreAction::ENoAction;
-	}
-	
-	if (!DepthStencilUsage.IsUsingStencil())
-	{
-		StencilLoadAction = ERenderTargetLoadAction::ENoAction;
-	}
-	
-	//if we aren't writing to stencil, there's no reason to store it back out again.  Should save some bandwidth on mobile platforms.
-	if (!DepthStencilUsage.IsStencilWrite())
-	{
-		StencilStoreAction = ERenderTargetStoreAction::ENoAction;
-	}
-}
+RHI_API void DecodeRenderTargetMode(ESimpleRenderTargetMode Mode, ERenderTargetLoadAction& ColorLoadAction, ERenderTargetStoreAction& ColorStoreAction, ERenderTargetLoadAction& DepthLoadAction, ERenderTargetStoreAction& DepthStoreAction, ERenderTargetLoadAction& StencilLoadAction, ERenderTargetStoreAction& StencilStoreAction, FExclusiveDepthStencil DepthStencilUsage);
 
 inline void TransitionRenderPassTargets(FRHICommandList& RHICmdList, const FRHIRenderPassInfo& RPInfo)
 {
@@ -941,7 +863,7 @@ struct FRHILockTracker
 				return Result;
 			}
 		}
-		UE_LOG(LogRHI, Fatal, TEXT("Mismatched RHI buffer locks."));
+		RaiseMismatchError();
 		return FLockParams(nullptr, nullptr, 0, 0, RLM_WriteOnly, false, false);
 	}
 	
@@ -980,6 +902,8 @@ struct FRHILockTracker
 			}
 		}
 	}
+
+	RHI_API void RaiseMismatchError();
 };
 
 extern RHI_API FRHILockTracker GRHILockTracker;
