Fixing RHI inconsistent define state. Adding a few defines to more RHI headers to allow IWYU to run.

diff --git a/Engine/Source/Runtime/RHI/Public/RHIDefinitions.h b/Engine/Source/Runtime/RHI/Public/RHIDefinitions.h
--- a/Engine/Source/Runtime/RHI/Public/RHIDefinitions.h
+++ b/Engine/Source/Runtime/RHI/Public/RHIDefinitions.h
@@ -9,6 +9,10 @@
 
 #include "Serialization/MemoryLayout.h"
 
+#ifndef USE_STATIC_SHADER_PLATFORM_ENUMS
+#define USE_STATIC_SHADER_PLATFORM_ENUMS 0
+#endif
+
 /** Alignment of the shader parameters struct is required to be 16-byte boundaries. */
 #define SHADER_PARAMETER_STRUCT_ALIGNMENT 16
 
diff --git a/Engine/Source/Runtime/RHI/Public/RHIFeatureLevel.h b/Engine/Source/Runtime/RHI/Public/RHIFeatureLevel.h
--- a/Engine/Source/Runtime/RHI/Public/RHIFeatureLevel.h
+++ b/Engine/Source/Runtime/RHI/Public/RHIFeatureLevel.h
@@ -3,6 +3,7 @@
 #pragma once
 
 #include "Containers/EnumAsByte.h"
+#include "RHIDefinitions.h"
 #include "Serialization/MemoryLayout.h"
 
 #ifndef USE_STATIC_FEATURE_LEVEL_ENUMS
diff --git a/Engine/Source/Runtime/RHI/Public/RHIShaderPlatform.h b/Engine/Source/Runtime/RHI/Public/RHIShaderPlatform.h
--- a/Engine/Source/Runtime/RHI/Public/RHIShaderPlatform.h
+++ b/Engine/Source/Runtime/RHI/Public/RHIShaderPlatform.h
@@ -3,10 +3,7 @@
 #pragma once
 
 #include "UObject/NameTypes.h"
-
-#ifndef USE_STATIC_SHADER_PLATFORM_ENUMS
-#define USE_STATIC_SHADER_PLATFORM_ENUMS 0
-#endif
+#include "RHIDefinitions.h"
 
 /** @warning: update *LegacyShaderPlatform* when the below changes */
 enum EShaderPlatform : uint16
diff --git a/Engine/Source/Runtime/RHI/Public/RHIValidationCommon.h b/Engine/Source/Runtime/RHI/Public/RHIValidationCommon.h
--- a/Engine/Source/Runtime/RHI/Public/RHIValidationCommon.h
+++ b/Engine/Source/Runtime/RHI/Public/RHIValidationCommon.h
@@ -6,12 +6,11 @@
 
 #pragma once
 
-// HEADER_UNIT_SKIP - Should always be included through RHI.h (this code needs cleanup.. this file is included in the middle of RHI.h)
-
 #include "Experimental/ConcurrentLinearAllocator.h"
 #include "PixelFormat.h"
 #include "RHIPipeline.h"
 #include "RHIStrings.h"
+#include "RHIAccess.h"
 
 #if ENABLE_RHI_VALIDATION
 extern RHI_API bool GRHIValidationEnabled;
diff --git a/Engine/Source/Runtime/RHI/Public/RHIValidationContext.h b/Engine/Source/Runtime/RHI/Public/RHIValidationContext.h
--- a/Engine/Source/Runtime/RHI/Public/RHIValidationContext.h
+++ b/Engine/Source/Runtime/RHI/Public/RHIValidationContext.h
@@ -6,12 +6,14 @@
 
 #pragma once
 
-// HEADER_UNIT_SKIP - Should always be included through RHIValidation.h
-
 #include "RHIValidationCommon.h"
 #include "RHIValidationUtils.h"
 
 #if ENABLE_RHI_VALIDATION
+
+#include "RHI.h"
+#include "RHIUtilities.h"
+
 class FValidationRHI;
 
 inline void ValidateShaderParameters(RHIValidation::FTracker* Tracker, RHIValidation::FStaticUniformBuffers& StaticUniformBuffers, TConstArrayView<FRHIShaderParameterResource> InParameters, ERHIAccess InRequiredAccess, RHIValidation::EUAVMode InRequiredUAVMode)
diff --git a/Engine/Source/Runtime/RHI/Public/RHIValidationUtils.h b/Engine/Source/Runtime/RHI/Public/RHIValidationUtils.h
--- a/Engine/Source/Runtime/RHI/Public/RHIValidationUtils.h
+++ b/Engine/Source/Runtime/RHI/Public/RHIValidationUtils.h
@@ -2,12 +2,13 @@
 
 #pragma once
 
-// HEADER_UNIT_SKIP - Included through others
-
 #include "Misc/Build.h"
 
 #if ENABLE_RHI_VALIDATION
 
+#include "RHIResources.h"
+#include "RHICommandList.h"
+
 class FValidationRHIUtils
 {
 public:
