Fixing nonunity issues.

diff --git a/Engine/Source/Runtime/RHI/Public/RHIDefinitions.h b/Engine/Source/Runtime/RHI/Public/RHIDefinitions.h
--- a/Engine/Source/Runtime/RHI/Public/RHIDefinitions.h
+++ b/Engine/Source/Runtime/RHI/Public/RHIDefinitions.h
@@ -7,7 +7,9 @@
 
 #pragma once
 
+#include "GpuProfilerTrace.h" // TODO Move defines into RHIDefinitions
 #include "Serialization/MemoryLayout.h"
+#include "ProfilingDebugging/CsvProfilerConfig.h" // TODO Move defines into RHIDefinitions
 
 #ifndef USE_STATIC_SHADER_PLATFORM_ENUMS
 #define USE_STATIC_SHADER_PLATFORM_ENUMS 0
diff --git a/Engine/Source/Runtime/RHI/Public/RHIStats.h b/Engine/Source/Runtime/RHI/Public/RHIStats.h
--- a/Engine/Source/Runtime/RHI/Public/RHIStats.h
+++ b/Engine/Source/Runtime/RHI/Public/RHIStats.h
@@ -5,7 +5,7 @@
 #include "Containers/StaticArray.h"
 #include "MultiGPU.h"
 #include "RHIDefinitions.h"
-#include "Stats/Stats2.h"
+#include "Stats/Stats.h"
 
 struct FTextureMemoryStats
 {
