Build fix for same name of multiple cvars

#rnx

diff --git a/Engine/Source/Runtime/RHI/Private/PipelineStateCache.cpp b/Engine/Source/Runtime/RHI/Private/PipelineStateCache.cpp
--- a/Engine/Source/Runtime/RHI/Private/PipelineStateCache.cpp
+++ b/Engine/Source/Runtime/RHI/Private/PipelineStateCache.cpp
@@ -131,7 +131,7 @@ static FAutoConsoleVariableRef CVarPSOPrecaching(
 );
 
 int32 GPSOWaitForHighPriorityRequestsOnly = 0;
-static FAutoConsoleVariableRef CVarPSOPrecaching(
+static FAutoConsoleVariableRef CVarPSOWaitForHighPriorityRequestsOnly(
 	TEXT("r.PSOPrecaching.WaitForHighPriorityRequestsOnly"),
 	GPSOWaitForHighPriorityRequestsOnly,
 	TEXT("0 to wait for all pending PSO precache requests during loading (default)g\n")
