RHI: Add PreDispatch() to ClearUAV functions to flush parameter batcher.  Some RHI would dispatch under the hood using TRHICommandList_RecursiveHazardous which would alter the context (such as current shader) without informing ClearUAV's calling RHICommandList.  Fix EngineTests in Vulkan.  (more automated solution being investigated)

diff --git a/Engine/Source/Runtime/RHI/Public/RHICommandList.h b/Engine/Source/Runtime/RHI/Public/RHICommandList.h
--- a/Engine/Source/Runtime/RHI/Public/RHICommandList.h
+++ b/Engine/Source/Runtime/RHI/Public/RHICommandList.h
@@ -2799,6 +2799,7 @@ PRAGMA_ENABLE_DEPRECATION_WARNINGS
 
 	FORCEINLINE_DEBUGGABLE void ClearUAVFloat(FRHIUnorderedAccessView* UnorderedAccessViewRHI, const FVector4f& Values)
 	{
+		PreDispatch();
 		if (Bypass())
 		{
 			GetComputeContext().RHIClearUAVFloat(UnorderedAccessViewRHI, Values);
@@ -2809,6 +2810,7 @@ PRAGMA_ENABLE_DEPRECATION_WARNINGS
 
 	FORCEINLINE_DEBUGGABLE void ClearUAVUint(FRHIUnorderedAccessView* UnorderedAccessViewRHI, const FUintVector4& Values)
 	{
+		PreDispatch();
 		if (Bypass())
 		{
 			GetComputeContext().RHIClearUAVUint(UnorderedAccessViewRHI, Values);
