Added RHI resource lifetime scopes which defer deletion of unreferenced resources until all scopes are closed.
The RDG builder holds a scope in order to avoid deletion of any resources during the graph setup / execution. This resolves the need to hold a strong reference during the RDG graph lifetime.

If FlushPendingDeletes is called within one of these scopes, the resources are instead queued onto the immediate command list and flushed at the end of the scope.

This change removes strong references held of uniform buffers held by RHI commands and certain platform contexts. When a flush is encountered, the backends will clear all caches removing any potential dangling references.

Resources can opt out of lifetime extension by calling ->DisableLifetimeExtension(). Subsequent calls to FlushPendingDeletes will release the resource immediately. This is used by a couple edge cases where resources must be deleted mid-frame (namely, DumpGPU and BVH building).

diff --git a/Engine/Source/Runtime/D3D12RHI/Private/D3D12CommandContext.cpp b/Engine/Source/Runtime/D3D12RHI/Private/D3D12CommandContext.cpp
--- a/Engine/Source/Runtime/D3D12RHI/Private/D3D12CommandContext.cpp
+++ b/Engine/Source/Runtime/D3D12RHI/Private/D3D12CommandContext.cpp
@@ -619,14 +619,6 @@ void FD3D12CommandContext::ClearState(EClearStateMode Mode)
 	{
 		FMemory::Memzero(StaticUniformBuffers.GetData(), StaticUniformBuffers.Num() * sizeof(FRHIUniformBuffer*));
 	}
-
-	for (int i = 0; i < UE_ARRAY_COUNT(BoundUniformBufferRefs); i++)
-	{
-		for (int j = 0; j < UE_ARRAY_COUNT(BoundUniformBufferRefs[i]); j++)
-		{
-			BoundUniformBufferRefs[i][j] = NULL;
-		}
-	}
 }
 
 void FD3D12CommandContext::ConditionalClearShaderResource(FD3D12ResourceLocation* Resource)
diff --git a/Engine/Source/Runtime/D3D12RHI/Private/D3D12CommandContext.h b/Engine/Source/Runtime/D3D12RHI/Private/D3D12CommandContext.h
--- a/Engine/Source/Runtime/D3D12RHI/Private/D3D12CommandContext.h
+++ b/Engine/Source/Runtime/D3D12RHI/Private/D3D12CommandContext.h
@@ -452,7 +452,6 @@ public:
 
 	/** Track the currently bound uniform buffers. */
 	FD3D12UniformBuffer* BoundUniformBuffers[SF_NumStandardFrequencies][MAX_CBS] = {};
-	FUniformBufferRHIRef BoundUniformBufferRefs[SF_NumStandardFrequencies][MAX_CBS] = {};
 
 	/** Bit array to track which uniform buffers have changed since the last draw call. */
 	uint16 DirtyUniformBuffers[SF_NumStandardFrequencies] = {};
diff --git a/Engine/Source/Runtime/D3D12RHI/Private/D3D12Commands.cpp b/Engine/Source/Runtime/D3D12RHI/Private/D3D12Commands.cpp
--- a/Engine/Source/Runtime/D3D12RHI/Private/D3D12Commands.cpp
+++ b/Engine/Source/Runtime/D3D12RHI/Private/D3D12Commands.cpp
@@ -99,11 +99,6 @@ static void BindUniformBuffer(FD3D12CommandContext& Context, FRHIShader* Shader,
 
 	Context.StateCache.SetConstantsFromUniformBuffer(ShaderFrequency, BufferIndex, InBuffer);
 
-	if (!GRHINeedsExtraDeletionLatency)
-	{
-		Context.BoundUniformBufferRefs[ShaderFrequency][BufferIndex] = InBuffer;
-	}
-
 	Context.BoundUniformBuffers[ShaderFrequency][BufferIndex] = InBuffer;
 	Context.DirtyUniformBuffers[ShaderFrequency] |= (1 << BufferIndex);
 }
diff --git a/Engine/Source/Runtime/D3D12RHI/Private/D3D12RHI.cpp b/Engine/Source/Runtime/D3D12RHI/Private/D3D12RHI.cpp
--- a/Engine/Source/Runtime/D3D12RHI/Private/D3D12RHI.cpp
+++ b/Engine/Source/Runtime/D3D12RHI/Private/D3D12RHI.cpp
@@ -495,6 +495,18 @@ void FD3D12DynamicRHI::RHIPerFrameRHIFlushComplete()
 			}
 		});
 	}
+
+	// Clear all bound resources since we are about to flush pending deletions.
+
+	for (uint32 AdapterIndex = 0; AdapterIndex < GetNumAdapters(); ++AdapterIndex)
+	{
+		FD3D12Adapter& Adapter = GetAdapter(AdapterIndex);
+
+		for (FD3D12Device* Device : Adapter.GetDevices())
+		{
+			Device->GetDefaultCommandContext().ClearState(FD3D12ContextCommon::EClearStateMode::All);
+		}
+	}
 }
 
 void FD3D12DynamicRHI::RHIAcquireThreadOwnership()
diff --git a/Engine/Source/Runtime/Engine/Private/SkeletalRenderGPUSkin.cpp b/Engine/Source/Runtime/Engine/Private/SkeletalRenderGPUSkin.cpp
--- a/Engine/Source/Runtime/Engine/Private/SkeletalRenderGPUSkin.cpp
+++ b/Engine/Source/Runtime/Engine/Private/SkeletalRenderGPUSkin.cpp
@@ -855,6 +855,7 @@ void FSkeletalMeshObjectGPUSkin::UpdateRayTracingGeometry(FSkeletalMeshLODRender
 
 				// Release the old data (make sure it's not pending build anymore either)
 				RayTracingUpdateQueue->Remove(&RayTracingGeometry, MemoryEstimation);
+				RayTracingGeometry.RayTracingGeometryRHI->DisableLifetimeExtension();
 				RayTracingGeometry.ReleaseRHI();
 			}
 
diff --git a/Engine/Source/Runtime/OpenGLDrv/Private/OpenGLDevice.cpp b/Engine/Source/Runtime/OpenGLDrv/Private/OpenGLDevice.cpp
--- a/Engine/Source/Runtime/OpenGLDrv/Private/OpenGLDevice.cpp
+++ b/Engine/Source/Runtime/OpenGLDrv/Private/OpenGLDevice.cpp
@@ -215,6 +215,8 @@ void FOpenGLDynamicRHI::RHIPerFrameRHIFlushComplete()
 	BeginFrame_QueryBatchCleanup();
 
 	OpenGL_PollAllFences();
+
+	FMemory::Memset(PendingState.BoundUniformBuffers, 0, sizeof(PendingState.BoundUniformBuffers));
 }
 
 
diff --git a/Engine/Source/Runtime/OpenGLDrv/Private/OpenGLShaders.cpp b/Engine/Source/Runtime/OpenGLDrv/Private/OpenGLShaders.cpp
--- a/Engine/Source/Runtime/OpenGLDrv/Private/OpenGLShaders.cpp
+++ b/Engine/Source/Runtime/OpenGLDrv/Private/OpenGLShaders.cpp
@@ -1069,7 +1069,7 @@ static void MarkShaderParameterCachesDirty(FOpenGLShaderParameterCache* ShaderPa
 	}
 }
 
-void FOpenGLDynamicRHI::BindUniformBufferBase(FOpenGLContextState& ContextState, int32 NumUniformBuffers, FUniformBufferRHIRef* BoundUniformBuffers, uint32 FirstUniformBuffer, bool ForceUpdate)
+void FOpenGLDynamicRHI::BindUniformBufferBase(FOpenGLContextState& ContextState, int32 NumUniformBuffers, FRHIUniformBuffer** BoundUniformBuffers, uint32 FirstUniformBuffer, bool ForceUpdate)
 {
 	SCOPE_CYCLE_COUNTER_DETAILED(STAT_OpenGLUniformBindTime);
 	VERIFY_GL_SCOPE();
@@ -1082,9 +1082,9 @@ void FOpenGLDynamicRHI::BindUniformBufferBase(FOpenGLContextState& ContextState,
 		uint32 Size = ZERO_FILLED_DUMMY_UNIFORM_BUFFER_SIZE;
 		int32 BindIndex = FirstUniformBuffer + BufferIndex;
 
-		if (IsValidRef(BoundUniformBuffers[BufferIndex]))
+		if (BoundUniformBuffers[BufferIndex])
 		{
-			FRHIUniformBuffer* UB = BoundUniformBuffers[BufferIndex].GetReference();
+			FRHIUniformBuffer* UB = BoundUniformBuffers[BufferIndex];
 			FOpenGLUniformBuffer* GLUB = ((FOpenGLUniformBuffer*)UB);
 			Buffer = GLUB->Resource;
 
@@ -3598,7 +3598,7 @@ void FOpenGLShaderParameterCache::CommitPackedGlobals(const FOpenGLLinkedProgram
 	}
 }
 
-void FOpenGLShaderParameterCache::CommitPackedUniformBuffers(FOpenGLLinkedProgram* LinkedProgram, int32 Stage, FUniformBufferRHIRef* RHIUniformBuffers, const TArray<CrossCompiler::FUniformBufferCopyInfo>& UniformBuffersCopyInfo)
+void FOpenGLShaderParameterCache::CommitPackedUniformBuffers(FOpenGLLinkedProgram* LinkedProgram, int32 Stage, FRHIUniformBuffer** RHIUniformBuffers, const TArray<CrossCompiler::FUniformBufferCopyInfo>& UniformBuffersCopyInfo)
 {
 	SCOPE_CYCLE_COUNTER(STAT_OpenGLConstantBufferUpdateTime);
 	VERIFY_GL_SCOPE();
@@ -3613,7 +3613,7 @@ void FOpenGLShaderParameterCache::CommitPackedUniformBuffers(FOpenGLLinkedProgra
 		int32 LastInfoIndex = 0;
 		for (int32 BufferIndex = 0; BufferIndex < Bindings.NumUniformBuffers; ++BufferIndex)
 		{
-			const FOpenGLUniformBuffer* UniformBuffer = (FOpenGLUniformBuffer*)RHIUniformBuffers[BufferIndex].GetReference();
+			const FOpenGLUniformBuffer* UniformBuffer = (FOpenGLUniformBuffer*)RHIUniformBuffers[BufferIndex];
 			check(UniformBuffer);
 
 			if (!UniformBuffer->bIsEmulatedUniformBuffer)
@@ -3648,7 +3648,7 @@ void FOpenGLShaderParameterCache::CommitPackedUniformBuffers(FOpenGLLinkedProgra
 		auto& EmulatedUniformBufferSet = LinkedProgram->StagePackedUniformInfo[Stage].LastEmulatedUniformBufferSet;
 		for (int32 BufferIndex = 0; BufferIndex < Bindings.NumUniformBuffers; ++BufferIndex)
 		{
-			const FOpenGLUniformBuffer* UniformBuffer = (FOpenGLUniformBuffer*)RHIUniformBuffers[BufferIndex].GetReference();
+			const FOpenGLUniformBuffer* UniformBuffer = (FOpenGLUniformBuffer*)RHIUniformBuffers[BufferIndex];
 
 			if (UniformBuffer && !UniformBuffer->bIsEmulatedUniformBuffer)
 			{
diff --git a/Engine/Source/Runtime/OpenGLDrv/Public/OpenGLDrv.h b/Engine/Source/Runtime/OpenGLDrv/Public/OpenGLDrv.h
--- a/Engine/Source/Runtime/OpenGLDrv/Public/OpenGLDrv.h
+++ b/Engine/Source/Runtime/OpenGLDrv/Public/OpenGLDrv.h
@@ -1042,7 +1042,7 @@ private:
 
 	void ReadSurfaceDataRaw(FOpenGLContextState& ContextState, FRHITexture* TextureRHI,FIntRect Rect,TArray<uint8>& OutData, FReadSurfaceDataFlags InFlags);
 
-	void BindUniformBufferBase(FOpenGLContextState& ContextState, int32 NumUniformBuffers, FUniformBufferRHIRef* BoundUniformBuffers, uint32 FirstUniformBuffer, bool ForceUpdate);
+	void BindUniformBufferBase(FOpenGLContextState& ContextState, int32 NumUniformBuffers, FRHIUniformBuffer** BoundUniformBuffers, uint32 FirstUniformBuffer, bool ForceUpdate);
 
 	void ClearCurrentFramebufferWithCurrentScissor(FOpenGLContextState& ContextState, int8 ClearType, int32 NumClearColors, const FLinearColor* ClearColorArray, float Depth, uint32 Stencil);
 
diff --git a/Engine/Source/Runtime/OpenGLDrv/Public/OpenGLShaderResources.h b/Engine/Source/Runtime/OpenGLDrv/Public/OpenGLShaderResources.h
--- a/Engine/Source/Runtime/OpenGLDrv/Public/OpenGLShaderResources.h
+++ b/Engine/Source/Runtime/OpenGLDrv/Public/OpenGLShaderResources.h
@@ -362,7 +362,7 @@ public:
 	 */
 	void CommitPackedGlobals(const FOpenGLLinkedProgram* LinkedProgram, int32 Stage);
 
-	void CommitPackedUniformBuffers(FOpenGLLinkedProgram* LinkedProgram, int32 Stage, FUniformBufferRHIRef* UniformBuffers, const TArray<CrossCompiler::FUniformBufferCopyInfo>& UniformBuffersCopyInfo);
+	void CommitPackedUniformBuffers(FOpenGLLinkedProgram* LinkedProgram, int32 Stage, FRHIUniformBuffer** UniformBuffers, const TArray<CrossCompiler::FUniformBufferCopyInfo>& UniformBuffersCopyInfo);
 
 private:
 
diff --git a/Engine/Source/Runtime/OpenGLDrv/Public/OpenGLState.h b/Engine/Source/Runtime/OpenGLDrv/Public/OpenGLState.h
--- a/Engine/Source/Runtime/OpenGLDrv/Public/OpenGLState.h
+++ b/Engine/Source/Runtime/OpenGLDrv/Public/OpenGLState.h
@@ -459,7 +459,7 @@ struct FOpenGLRHIState final : public FOpenGLCommonState
 	enum { MAX_UNIFORM_BUFFERS_PER_SHADER_STAGE = 14 };
 
 	/** Track the currently bound uniform buffers. */
-	FUniformBufferRHIRef BoundUniformBuffers[SF_NumStandardFrequencies][MAX_UNIFORM_BUFFERS_PER_SHADER_STAGE];
+	FRHIUniformBuffer* BoundUniformBuffers[SF_NumStandardFrequencies][MAX_UNIFORM_BUFFERS_PER_SHADER_STAGE];
 
 	/** Array to track if any real (not emulated) uniform buffers have been bound since the last draw call */
 	bool bAnyDirtyRealUniformBuffers[SF_NumStandardFrequencies];
@@ -511,6 +511,7 @@ struct FOpenGLRHIState final : public FOpenGLCommonState
 		FMemory::Memset( RenderTargets, 0, sizeof(RenderTargets) );	// setting all to 0 at start
 		FMemory::Memset( RenderTargetMipmapLevels, 0, sizeof(RenderTargetMipmapLevels) );	// setting all to 0 at start
 		FMemory::Memset( RenderTargetArrayIndex, 0, sizeof(RenderTargetArrayIndex) );	// setting all to 0 at start
+		FMemory::Memset(BoundUniformBuffers, 0, sizeof(BoundUniformBuffers));
 	}
 
 	~FOpenGLRHIState()
@@ -524,16 +525,7 @@ struct FOpenGLRHIState final : public FOpenGLCommonState
 	{
 		delete [] ShaderParameters;
 		ShaderParameters = NULL;
-
-		// Release references to bound uniform buffers.
-		for (int32 Frequency = 0; Frequency < SF_NumStandardFrequencies; ++Frequency)
-		{
-			for (int32 BindIndex = 0; BindIndex < MAX_UNIFORM_BUFFERS_PER_SHADER_STAGE; ++BindIndex)
-			{
-				BoundUniformBuffers[Frequency][BindIndex].SafeRelease();
-			}
-		}
-
+		FMemory::Memset(BoundUniformBuffers, 0, sizeof(BoundUniformBuffers));
 		FOpenGLCommonState::CleanupResources();
 	}
 };
diff --git a/Engine/Source/Runtime/RHI/Private/RHICommandList.cpp b/Engine/Source/Runtime/RHI/Private/RHICommandList.cpp
--- a/Engine/Source/Runtime/RHI/Private/RHICommandList.cpp
+++ b/Engine/Source/Runtime/RHI/Private/RHICommandList.cpp
@@ -1131,6 +1131,79 @@ void FRHICommandListExecutor::WaitOnRHIThreadFence(FGraphEventRef& Fence)
 	}
 }
 
+DECLARE_CYCLE_STAT(TEXT("Delete Resources"), STAT_DeleteResources, STATGROUP_RHICMDLIST);
+
+int32 FRHICommandListImmediate::FlushPendingDeletes()
+{
+	SCOPE_CYCLE_COUNTER(STAT_DeleteResources);
+	check(IsInRenderingThread());
+
+	TArray<FRHIResource*, FConcurrentLinearArrayAllocator> DeletedResources;
+
+	TArray<FRHIResource*, FConcurrentLinearArrayAllocator>& DeletedResourcesWithLifetimeExtension = PersistentState.ExtendResourceLifetimeRefCount > 0
+		? PersistentState.ExtendedLifetimeResources 
+		: DeletedResources;
+
+	FRHIResource::PendingDeletesWithLifetimeExtension.ConsumeAllLifo([&DeletedResourcesWithLifetimeExtension](FRHIResource* Resource)
+	{
+		DeletedResourcesWithLifetimeExtension.Emplace(Resource);
+	});
+
+	FRHIResource::PendingDeletes.ConsumeAllLifo([&DeletedResources](FRHIResource* Resource)
+	{
+		DeletedResources.Push(Resource);
+	});
+
+	const int32 NumDeletes = DeletedResources.Num();
+
+	EnqueueLambda([DeletedResources = MoveTemp(DeletedResources)](FRHICommandListImmediate& RHICmdList) mutable
+	{
+		SCOPED_NAMED_EVENT(STAT_FRHICommandListImmediate_DeleteResources, FColor::Magenta);
+
+		if (GDynamicRHI)
+		{
+			GDynamicRHI->RHIPerFrameRHIFlushComplete();
+		}
+
+		for (int32 i = DeletedResources.Num() - 1; i >= 0; i--)
+		{
+			FRHIResource* Resource = DeletedResources[i];
+			if (Resource->AtomicFlags.Deleteing())
+			{
+				FRHIResource::CurrentlyDeleting = Resource;
+				delete Resource;
+			}
+		}
+	});
+
+	return NumDeletes;
+}
+
+int32 FRHICommandListImmediate::FlushExtendedLifetimeResourceDeletes()
+{
+	int32 NumDeletes = PersistentState.ExtendedLifetimeResources.Num();
+
+	if (NumDeletes > 0)
+	{
+		EnqueueLambda([DeletedResources = MoveTemp(PersistentState.ExtendedLifetimeResources)](FRHICommandListImmediate& RHICmdList) mutable
+		{
+			SCOPED_NAMED_EVENT(STAT_FRHICommandListImmediate_DeleteExtendedLifetimeResources, FColor::Magenta);
+
+			for (int32 i = DeletedResources.Num() - 1; i >= 0; i--)
+			{
+				FRHIResource* Resource = DeletedResources[i];
+				if (Resource->AtomicFlags.Deleteing())
+				{
+					FRHIResource::CurrentlyDeleting = Resource;
+					delete Resource;
+				}
+			}
+		});
+	}
+
+	return NumDeletes;
+}
+
 void FRHICommandListImmediate::WaitForTasks()
 {
 	check(IsInRenderingThread());
diff --git a/Engine/Source/Runtime/RHI/Private/RHIResources.cpp b/Engine/Source/Runtime/RHI/Private/RHIResources.cpp
--- a/Engine/Source/Runtime/RHI/Private/RHIResources.cpp
+++ b/Engine/Source/Runtime/RHI/Private/RHIResources.cpp
@@ -8,6 +8,7 @@
 #include "Stats/Stats.h"
 
 UE::TConsumeAllMpmcQueue<FRHIResource*> FRHIResource::PendingDeletes;
+UE::TConsumeAllMpmcQueue<FRHIResource*> FRHIResource::PendingDeletesWithLifetimeExtension;
 FRHIResource* FRHIResource::CurrentlyDeleting = nullptr;
 
 bool FRHIResource::Bypass()
@@ -15,39 +16,7 @@ bool FRHIResource::Bypass()
 	return GRHICommandList.Bypass();
 }
 
-DECLARE_CYCLE_STAT(TEXT("Delete Resources"), STAT_DeleteResources, STATGROUP_RHICMDLIST);
-
 int32 FRHIResource::FlushPendingDeletes(FRHICommandListImmediate& RHICmdList)
 {
-	SCOPE_CYCLE_COUNTER(STAT_DeleteResources);
-
-	check(IsInRenderingThread());
-
-	TArray<FRHIResource*> DeletedResources;
-	PendingDeletes.ConsumeAllLifo([&DeletedResources](FRHIResource* Resource)
-	{
-		DeletedResources.Push(Resource);
-	});
-
-	const int32 NumDeletes = DeletedResources.Num();
-
-	RHICmdList.EnqueueLambda([DeletedResources = MoveTemp(DeletedResources)](FRHICommandListImmediate& RHICmdList) mutable
-	{
-		if (GDynamicRHI)
-		{
-			GDynamicRHI->RHIPerFrameRHIFlushComplete();
-		}
-
-		for (int32 i = DeletedResources.Num() - 1; i >= 0; i--)
-		{
-			FRHIResource* Resource = DeletedResources[i];
-			if (Resource->AtomicFlags.Deleteing())
-			{
-				FRHIResource::CurrentlyDeleting = Resource;
-				delete Resource;
-			}
-		}
-	});
-
-	return NumDeletes;
+	return RHICmdList.FlushPendingDeletes();
 }
diff --git a/Engine/Source/Runtime/RHI/Public/RHICommandList.h b/Engine/Source/Runtime/RHI/Public/RHICommandList.h
--- a/Engine/Source/Runtime/RHI/Public/RHICommandList.h
+++ b/Engine/Source/Runtime/RHI/Public/RHICommandList.h
@@ -942,6 +942,7 @@ protected:
 		ESubpassHint SubpassHint = ESubpassHint::None;
 		uint8 SubpassIndex = 0;
 		uint8 MultiViewCount = 0;
+		uint8 ExtendResourceLifetimeRefCount = 0;
 		bool HasFragmentDensityAttachment = false;
 
 		bool bInsideRenderPass = false;
@@ -968,6 +969,7 @@ protected:
 		TRefCountPtr<FFenceCandidate> FenceCandidate;
 		FGraphEventArray QueuedFenceCandidateEvents;
 		TArray<TRefCountPtr<FFenceCandidate>, FConcurrentLinearArrayAllocator> QueuedFenceCandidates;
+		TArray<FRHIResource*, FConcurrentLinearArrayAllocator> ExtendedLifetimeResources;
 
 		struct FGPUStats
 		{
@@ -1298,7 +1300,7 @@ FRHICOMMAND_MACRO_TPL(TRHIShader, FRHICommandSetShaderUniformBuffer)
 {
 	TRHIShader* Shader;
 	uint32 BaseIndex;
-	TRefCountPtr<FRHIUniformBuffer> UniformBuffer;
+	FRHIUniformBuffer* UniformBuffer;
 	FORCEINLINE_DEBUGGABLE FRHICommandSetShaderUniformBuffer(TRHIShader* InShader, uint32 InBaseIndex, FRHIUniformBuffer* InUniformBuffer)
 		: Shader(InShader)
 		, BaseIndex(InBaseIndex)
@@ -1904,40 +1906,9 @@ FRHICOMMAND_MACRO(FRHICommandSetStaticUniformBuffers)
 
 	FORCEINLINE_DEBUGGABLE FRHICommandSetStaticUniformBuffers(const FUniformBufferStaticBindings& InUniformBuffers)
 		: UniformBuffers(InUniformBuffers)
-	{
-		const int32 NumUniformBuffers = UniformBuffers.GetUniformBufferCount();
-		for (int32 Index = 0; Index < NumUniformBuffers; ++Index)
-		{
-			FRHIUniformBuffer* UniformBuffer = UniformBuffers.GetUniformBuffer(Index);
-			if (UniformBuffer)
-			{
-				UniformBuffer->AddRef();
-			}
-		}
-	}
-	~FRHICommandSetStaticUniformBuffers()
-	{
-		const int32 NumUniformBuffers = UniformBuffers.GetUniformBufferCount();
-		for (int32 Index = 0; Index < NumUniformBuffers; ++Index)
-		{
-			FRHIUniformBuffer* UniformBuffer = UniformBuffers.GetUniformBuffer(Index);
-			if (UniformBuffer)
-			{
-				UniformBuffer->Release();
-			}
-		}
-	}
-	RHI_API void Execute(FRHICommandListBase& CmdList);
-};
+	{}
 
-struct FComputedGraphicsPipelineState
-{
-	FGraphicsPipelineStateRHIRef GraphicsPipelineState;
-	int32 UseCount;
-	FComputedGraphicsPipelineState()
-		: UseCount(0)
-	{
-	}
+	RHI_API void Execute(FRHICommandListBase& CmdList);
 };
 
 struct FComputedUniformBuffer
@@ -2276,7 +2247,7 @@ FRHICOMMAND_UNNAMED(FRHICommandBuildAccelerationStructures)
 {
 	const TArrayView<const FRayTracingGeometryBuildParams> Params;
 	FRHIBufferRange ScratchBufferRange;
-	TRefCountPtr<FRHIBuffer> ScratchBuffer;
+	FRHIBuffer* ScratchBuffer;
 
 	explicit FRHICommandBuildAccelerationStructures(const TArrayView<const FRayTracingGeometryBuildParams> InParams, const FRHIBufferRange& ScratchBufferRange)
 		: Params(InParams), ScratchBufferRange(ScratchBufferRange), ScratchBuffer(ScratchBufferRange.Buffer)
@@ -4330,6 +4301,7 @@ extern RHI_API ERHIAccess RHIGetDefaultResourceState(EBufferUsageFlags InUsage,
 class RHI_API FRHICommandListImmediate : public FRHICommandList
 {
 	friend class FRHICommandListExecutor;
+	friend class FRHICommandListScopedExtendResourceLifetime;
 	friend struct FRHICommandBeginFrame;
 
 	static FGraphEventArray WaitOutstandingTasks;
@@ -4379,6 +4351,11 @@ class RHI_API FRHICommandListImmediate : public FRHICommandList
 	//
 	void ProcessStats();
 
+	//
+	// Called when all FRHICommandListScopedExtendResourceLifetime references are released to flush any deferred deletions.
+	//
+	int32 FlushExtendedLifetimeResourceDeletes();
+
 public:
 	void BeginScene();
 	void EndScene();
@@ -4945,6 +4922,8 @@ public:
 		 
 		return GDynamicRHI->RHIFlushResources();
 	}
+
+	int32 FlushPendingDeletes();
 	
 	FORCEINLINE uint32 GetGPUFrameCycles()
 	{
@@ -4999,7 +4978,7 @@ public:
 		ImmediateFlush(EImmediateFlushType::FlushRHIThread); 
 		return GDynamicRHI->RHIIsRenderingSuspended();
 	}
-	
+
 	UE_DEPRECATED(5.1, "No longer used: FCompression::UncompressMemory should be used instead")
 	FORCEINLINE bool EnqueueDecompress(uint8_t* SrcBuffer, uint8_t* DestBuffer, int CompressedSize, void* ErrorCodeBuffer)
 	{
@@ -5135,6 +5114,27 @@ public:
 	}
 };
 
+class RHI_API FRHICommandListScopedExtendResourceLifetime
+{
+public:
+	FRHICommandListScopedExtendResourceLifetime(FRHICommandListImmediate& InRHICmdList)
+		: RHICmdList(InRHICmdList)
+	{
+		RHICmdList.PersistentState.ExtendResourceLifetimeRefCount++;
+	}
+
+	~FRHICommandListScopedExtendResourceLifetime()
+	{
+		if (--RHICmdList.PersistentState.ExtendResourceLifetimeRefCount == 0)
+		{
+			RHICmdList.FlushExtendedLifetimeResourceDeletes();
+		}
+	}
+
+private:
+	FRHICommandListImmediate& RHICmdList;
+};
+
 //
 // Helper to activate a specific RHI pipeline within a block of renderer code.
 // Allows command list recording code to switch between graphics / async compute etc.
diff --git a/Engine/Source/Runtime/RHI/Public/RHICommandListCommandExecutes.inl b/Engine/Source/Runtime/RHI/Public/RHICommandListCommandExecutes.inl
--- a/Engine/Source/Runtime/RHI/Public/RHICommandListCommandExecutes.inl
+++ b/Engine/Source/Runtime/RHI/Public/RHICommandListCommandExecutes.inl
@@ -18,7 +18,6 @@
 class FRHICommandListBase;
 class IRHIComputeContext;
 struct FComputedBSS;
-struct FComputedGraphicsPipelineState;
 struct FComputedUniformBuffer;
 struct FMemory;
 struct FRHICommandBeginDrawingViewport;
diff --git a/Engine/Source/Runtime/RHI/Public/RHIResources.h b/Engine/Source/Runtime/RHI/Public/RHIResources.h
--- a/Engine/Source/Runtime/RHI/Public/RHIResources.h
+++ b/Engine/Source/Runtime/RHI/Public/RHIResources.h
@@ -45,6 +45,7 @@ public:
 	FRHIResource(ERHIResourceType InResourceType)
 		: ResourceType(InResourceType)
 		, bCommitted(true)
+		, bAllowExtendLifetime(true)
 #if RHI_ENABLE_RESOURCE_INFO
 		, bBeingTracked(false)
 #endif
@@ -78,7 +79,14 @@ private:
 	{
 		if (!AtomicFlags.MarkForDelete(std::memory_order_release))
 		{
-			PendingDeletes.ProduceItem(const_cast<FRHIResource*>(this));
+			if (bAllowExtendLifetime)
+			{
+				PendingDeletesWithLifetimeExtension.ProduceItem(const_cast<FRHIResource*>(this));
+			}
+			else
+			{
+				PendingDeletes.ProduceItem(const_cast<FRHIResource*>(this));
+			}
 		}
 	}
 
@@ -119,6 +127,12 @@ public:
 		delete this;
 	}
 
+	void DisableLifetimeExtension()
+	{
+		ensureMsgf(IsValid(), TEXT("Resource is already marked for deletion. This call is a no-op. DisableLifetimeExtension must be called while still holding a live reference."));
+		bAllowExtendLifetime = false;
+	}
+
 	inline ERHIResourceType GetType() const { return ResourceType; }
 
 	FName GetOwnerName() const;
@@ -206,6 +220,11 @@ private:
 			return (LocalPacked & MarkedForDeleteBit) == 0 && (LocalPacked & NumRefsMask) != 0;
 		}
 
+		bool IsMarkedForDelete(std::memory_order MemoryOrder)
+		{
+			return (Packed.load(MemoryOrder) & MarkedForDeleteBit) != 0;
+		}
+
 		int32 GetNumRefs(std::memory_order MemoryOrder)
 		{
 			return Packed.load(MemoryOrder) & NumRefsMask;
@@ -215,21 +234,17 @@ private:
 
 	const ERHIResourceType ResourceType;
 	uint8 bCommitted : 1;
+	uint8 bAllowExtendLifetime : 1;
 #if RHI_ENABLE_RESOURCE_INFO
 	uint8 bBeingTracked : 1;
 	FName OwnerName;
 #endif
 
 	static UE::TConsumeAllMpmcQueue<FRHIResource*> PendingDeletes;
+	static UE::TConsumeAllMpmcQueue<FRHIResource*> PendingDeletesWithLifetimeExtension;
 	static FRHIResource* CurrentlyDeleting;
 
-	// Some APIs don't do internal reference counting, so we have to wait an extra couple of frames before deleting resources
-	// to ensure the GPU has completely finished with them. This avoids expensive fences, etc.
-	struct ResourcesToDelete
-	{
-		TArray<FRHIResource*>	Resources;
-		uint32					FrameDeleted{};
-	};
+	friend FRHICommandListImmediate;
 };
 
 enum class EClearBinding
diff --git a/Engine/Source/Runtime/RenderCore/Private/DumpGPU.cpp b/Engine/Source/Runtime/RenderCore/Private/DumpGPU.cpp
--- a/Engine/Source/Runtime/RenderCore/Private/DumpGPU.cpp
+++ b/Engine/Source/Runtime/RenderCore/Private/DumpGPU.cpp
@@ -917,6 +917,7 @@ public:
 						.SetFlags(ETextureCreateFlags::UAV | ETextureCreateFlags::ShaderResource | ETextureCreateFlags::HideInVisualizeTexture);
 
 					StagingSrcTexture = RHICreateTexture(Desc);
+					StagingSrcTexture->DisableLifetimeExtension();
 				}
 
 				FUnorderedAccessViewRHIRef StagingOutput = RHICreateUnorderedAccessView(StagingSrcTexture, /* MipLevel = */ 0);
@@ -950,6 +951,7 @@ public:
 					.SetFlags(ETextureCreateFlags::CPUReadback | ETextureCreateFlags::HideInVisualizeTexture);
 
 				StagingTexture = RHICreateTexture(Desc);
+				StagingTexture->DisableLifetimeExtension();
 
 				RHICmdList.Transition(FRHITransitionInfo(StagingTexture, ERHIAccess::Unknown, ERHIAccess::CopyDest));
 
@@ -1470,6 +1472,7 @@ public:
 				{
 					FTimeBucketMeasure TimeBucketMeasure(this, ETimingBucket::RHIReadbackCommands);
 					StagingBuffer = RHICreateStagingBuffer();
+					StagingBuffer->DisableLifetimeExtension();
 
 					static const FName FenceName(TEXT("DumpGPU.BufferFence"));
 					Fence = RHICreateGPUFence(FenceName);
diff --git a/Engine/Source/Runtime/RenderCore/Private/RenderGraphBuilder.cpp b/Engine/Source/Runtime/RenderCore/Private/RenderGraphBuilder.cpp
--- a/Engine/Source/Runtime/RenderCore/Private/RenderGraphBuilder.cpp
+++ b/Engine/Source/Runtime/RenderCore/Private/RenderGraphBuilder.cpp
@@ -572,6 +572,7 @@ FRDGBuilder::FRDGBuilder(FRHICommandListImmediate& InRHICmdList, FRDGEventName I
 	, BarrierValidation(&Passes, BuilderName)
 #endif
 	, TransientResourceAllocator(GRDGTransientResourceAllocator.Get())
+	, ExtendResourceLifetimeScope(RHICmdList)
 {
 	AddProloguePass();
 
diff --git a/Engine/Source/Runtime/RenderCore/Public/RenderGraphBuilder.h b/Engine/Source/Runtime/RenderCore/Public/RenderGraphBuilder.h
--- a/Engine/Source/Runtime/RenderCore/Public/RenderGraphBuilder.h
+++ b/Engine/Source/Runtime/RenderCore/Public/RenderGraphBuilder.h
@@ -761,6 +761,8 @@ private:
 
 	IRHITransientResourceAllocator* TransientResourceAllocator = nullptr;
 
+	FRHICommandListScopedExtendResourceLifetime ExtendResourceLifetimeScope;
+
 	void MarkResourcesAsProduced(FRDGPass* Pass);
 
 	void Compile();
diff --git a/Engine/Source/Runtime/RenderCore/Public/ShaderParameterMacros.h b/Engine/Source/Runtime/RenderCore/Public/ShaderParameterMacros.h
--- a/Engine/Source/Runtime/RenderCore/Public/ShaderParameterMacros.h
+++ b/Engine/Source/Runtime/RenderCore/Public/ShaderParameterMacros.h
@@ -148,7 +148,7 @@ class alignas(SHADER_PARAMETER_STRUCT_ALIGNMENT) FUniformBufferBinding
 public:
 	FUniformBufferBinding() = default;
 
-	FORCEINLINE const FUniformBufferRHIRef& GetUniformBuffer() const
+	FORCEINLINE FRHIUniformBuffer* GetUniformBuffer() const
 	{
 		return UniformBuffer;
 	}
@@ -170,7 +170,7 @@ public:
 
 	FORCEINLINE operator bool() const
 	{
-		return UniformBuffer.IsValid();
+		return UniformBuffer != nullptr;
 	}
 
 protected:
@@ -180,7 +180,7 @@ protected:
 	{}
 
 private:
-	FUniformBufferRHIRef UniformBuffer;
+	FRHIUniformBuffer* UniformBuffer = nullptr;
 	EUniformBufferBindingFlags BindingFlags = EUniformBufferBindingFlags::Shader;
 };
 
@@ -207,7 +207,7 @@ public:
 #endif
 	}
 
-	FORCEINLINE TUniformBufferRef<TBufferStruct> GetUniformBuffer() const
+	FORCEINLINE TUniformBufferRef<TBufferStruct> GetUniformBufferRef() const
 	{
 		return TUniformBufferRef<TBufferStruct>(FUniformBufferBinding::GetUniformBuffer());
 	}
diff --git a/Engine/Source/Runtime/Windows/D3D11RHI/Private/D3D11Commands.cpp b/Engine/Source/Runtime/Windows/D3D11RHI/Private/D3D11Commands.cpp
--- a/Engine/Source/Runtime/Windows/D3D11RHI/Private/D3D11Commands.cpp
+++ b/Engine/Source/Runtime/Windows/D3D11RHI/Private/D3D11Commands.cpp
@@ -374,7 +374,7 @@ void FD3D11DynamicRHI::RHISetBoundShaderState(FRHIBoundShaderState* BoundShaderS
 	{
 		for (int32 BindIndex = 0; BindIndex < MAX_UNIFORM_BUFFERS_PER_SHADER_STAGE; ++BindIndex)
 		{
-			BoundUniformBuffers[Frequency][BindIndex].SafeRelease();
+			BoundUniformBuffers[Frequency][BindIndex] = nullptr;
 		}
 	}
 
diff --git a/Engine/Source/Runtime/Windows/D3D11RHI/Private/D3D11Device.cpp b/Engine/Source/Runtime/Windows/D3D11RHI/Private/D3D11Device.cpp
--- a/Engine/Source/Runtime/Windows/D3D11RHI/Private/D3D11Device.cpp
+++ b/Engine/Source/Runtime/Windows/D3D11RHI/Private/D3D11Device.cpp
@@ -231,6 +231,11 @@ FD3D11DynamicRHI::FD3D11DynamicRHI(IDXGIFactory1* InDXGIFactory1, D3D_FEATURE_LE
 	for (int32 Frequency = 0; Frequency < SF_NumStandardFrequencies; ++Frequency)
 	{
 		DirtyUniformBuffers[Frequency] = 0;
+
+		for (int32 BindIndex = 0; BindIndex < MAX_UNIFORM_BUFFERS_PER_SHADER_STAGE; ++BindIndex)
+		{
+			BoundUniformBuffers[Frequency][BindIndex] = nullptr;
+		}
 	}
 
 	StaticUniformBuffers.AddZeroed(FUniformBufferStaticSlotRegistry::Get().GetSlotCount());
@@ -599,7 +604,7 @@ void FD3D11DynamicRHI::CleanupD3DDevice()
 		{
 			for (int32 BindIndex = 0; BindIndex < MAX_UNIFORM_BUFFERS_PER_SHADER_STAGE; ++BindIndex)
 			{
-				BoundUniformBuffers[Frequency][BindIndex].SafeRelease();
+				BoundUniformBuffers[Frequency][BindIndex] = nullptr;
 			}
 		}
 
diff --git a/Engine/Source/Runtime/Windows/D3D11RHI/Private/D3D11RHIPrivate.h b/Engine/Source/Runtime/Windows/D3D11RHI/Private/D3D11RHIPrivate.h
--- a/Engine/Source/Runtime/Windows/D3D11RHI/Private/D3D11RHIPrivate.h
+++ b/Engine/Source/Runtime/Windows/D3D11RHI/Private/D3D11RHIPrivate.h
@@ -910,7 +910,7 @@ protected:
 	enum { MAX_UNIFORM_BUFFERS_PER_SHADER_STAGE = 14 };
 
 	/** Track the currently bound uniform buffers. */
-	FUniformBufferRHIRef BoundUniformBuffers[SF_NumStandardFrequencies][MAX_UNIFORM_BUFFERS_PER_SHADER_STAGE];
+	FRHIUniformBuffer* BoundUniformBuffers[SF_NumStandardFrequencies][MAX_UNIFORM_BUFFERS_PER_SHADER_STAGE];
 
 	/** Bit array to track which uniform buffers have changed since the last draw call. */
 	uint16 DirtyUniformBuffers[SF_NumStandardFrequencies];
diff --git a/Engine/Source/Runtime/Windows/D3D11RHI/Private/D3D11UniformBuffer.cpp b/Engine/Source/Runtime/Windows/D3D11RHI/Private/D3D11UniformBuffer.cpp
--- a/Engine/Source/Runtime/Windows/D3D11RHI/Private/D3D11UniformBuffer.cpp
+++ b/Engine/Source/Runtime/Windows/D3D11RHI/Private/D3D11UniformBuffer.cpp
@@ -203,7 +203,6 @@ FUniformBufferRHIRef FD3D11DynamicRHI::RHICreateUniformBuffer(const void* Conten
 			else
 			{
 				NewUniformBuffer = new FD3D11UniformBuffer(this, Layout, nullptr, FRingAllocation(), bAllocatedFromPool);
-				NewUniformBuffer->AddRef();
 
 				void* CPUContent = nullptr;
 				
@@ -221,7 +220,6 @@ FUniformBufferRHIRef FD3D11DynamicRHI::RHICreateUniformBuffer(const void* Conten
 						D3D11RHI_IMMEDIATE_CONTEXT,
 						CPUContent,
 						NumBytes);
-					NewUniformBuffer->Release();
 					FMemory::Free(CPUContent);
 				});
 			}
