RHI: add temporary hooks to allow for texture references to be updated in the same bindless slot:
- Add an SRV in the FTextureReference to reserve a bindless descriptor slot
- Route the TextureReferenceUpdate call through the Dynamic RHI so that different RHIs can intercept it (by default leave original behavior in all RHI)

#rnx

diff --git a/Engine/Source/Runtime/RHI/Private/DynamicRHI.cpp b/Engine/Source/Runtime/RHI/Private/DynamicRHI.cpp
--- a/Engine/Source/Runtime/RHI/Private/DynamicRHI.cpp
+++ b/Engine/Source/Runtime/RHI/Private/DynamicRHI.cpp
@@ -561,6 +561,11 @@ void FDynamicRHI::RHIUpdateShaderResourceView(FRHIShaderResourceView* SRV, FRHIB
 	UE_LOG(LogRHI, Fatal, TEXT("RHIUpdateShaderResourceView isn't implemented for the current RHI"));
 }
 
+void FDynamicRHI::RHIUpdateTextureReference(FRHITextureReference* TextureRef, FRHITexture* NewTexture)
+{
+	TextureRef->SetReferencedTexture(NewTexture);
+}
+
 uint64 FDynamicRHI::RHIGetMinimumAlignmentForBufferBackedSRV(EPixelFormat Format)
 {
 	return 1;
diff --git a/Engine/Source/Runtime/RHI/Private/RHICommandList.cpp b/Engine/Source/Runtime/RHI/Private/RHICommandList.cpp
--- a/Engine/Source/Runtime/RHI/Private/RHICommandList.cpp
+++ b/Engine/Source/Runtime/RHI/Private/RHICommandList.cpp
@@ -2074,8 +2074,9 @@ void FRHICommandListImmediate::UpdateTextureReference(FRHITextureReference* Text
 
 	EnqueueLambda([TextureRef, NewTexture](auto&)
 	{
-		TextureRef->SetReferencedTexture(NewTexture);
+		GDynamicRHI->RHIUpdateTextureReference(TextureRef, NewTexture);
 	});
+
 	RHIThreadFence(true);
 	if (GetUsedMemory() > 256 * 1024)
 	{
diff --git a/Engine/Source/Runtime/RHI/Private/RHITextureReference.cpp b/Engine/Source/Runtime/RHI/Private/RHITextureReference.cpp
--- a/Engine/Source/Runtime/RHI/Private/RHITextureReference.cpp
+++ b/Engine/Source/Runtime/RHI/Private/RHITextureReference.cpp
@@ -8,6 +8,7 @@ FRHITextureReference::FRHITextureReference()
 {
 	check(DefaultTexture);
 	ReferencedTexture = DefaultTexture;
+	UpdateBindlessShaderResourceView();
 }
 
 FRHITextureReference::~FRHITextureReference() = default;
@@ -20,7 +21,11 @@ FRHITextureReference* FRHITextureReference::GetTextureReference()
 FRHIDescriptorHandle FRHITextureReference::GetDefaultBindlessHandle() const
 { 
 	check(ReferencedTexture);
-	return ReferencedTexture->GetDefaultBindlessHandle();
+
+	// If an SRV has been created, return its handle
+	return BindlessShaderResourceViewRHI.IsValid() ?
+		BindlessShaderResourceViewRHI->GetBindlessHandle() : 
+		ReferencedTexture->GetDefaultBindlessHandle();
 }
 
 void* FRHITextureReference::GetNativeResource() const
@@ -60,3 +65,11 @@ const FRHITextureDesc& FRHITextureReference::GetDesc() const
 	check(ReferencedTexture);
 	return ReferencedTexture->GetDesc();
 }
+
+void FRHITextureReference::UpdateBindlessShaderResourceView()
+{
+	if (RHIGetBindlessResourcesConfiguration(GMaxRHIShaderPlatform) != ERHIBindlessConfiguration::Disabled)
+	{
+		BindlessShaderResourceViewRHI = RHICreateShaderResourceView(ReferencedTexture.GetReference(), 0u);
+	}
+}
diff --git a/Engine/Source/Runtime/RHI/Public/DynamicRHI.h b/Engine/Source/Runtime/RHI/Public/DynamicRHI.h
--- a/Engine/Source/Runtime/RHI/Public/DynamicRHI.h
+++ b/Engine/Source/Runtime/RHI/Public/DynamicRHI.h
@@ -444,6 +444,8 @@ public:
 
 	virtual void RHIUpdateShaderResourceView(FRHIShaderResourceView* SRV, FRHIBuffer* Buffer);
 
+	virtual void RHIUpdateTextureReference(FRHITextureReference* TextureRef, FRHITexture* NewTexture);
+
 	struct FRHICalcTextureSizeResult
 	{
 		// The total size of the texture, in bytes.
diff --git a/Engine/Source/Runtime/RHI/Public/RHITextureReference.h b/Engine/Source/Runtime/RHI/Public/RHITextureReference.h
--- a/Engine/Source/Runtime/RHI/Public/RHITextureReference.h
+++ b/Engine/Source/Runtime/RHI/Public/RHITextureReference.h
@@ -30,8 +30,11 @@ public:
 		return ReferencedTexture.GetReference();
 	}
 
+	void UpdateBindlessShaderResourceView();
+
 private:
 	friend class FRHICommandListImmediate;
+	friend class FDynamicRHI;
 
 	// Called only from FRHICommandListImmediate::UpdateTextureReference()
 	void SetReferencedTexture(FRHITexture* InTexture)
@@ -43,6 +46,9 @@ private:
 
 	TRefCountPtr<FRHITexture> ReferencedTexture;
 
+	// SRV to create a reusable slot for bindless handles for this texture reference
+	FShaderResourceViewRHIRef BindlessShaderResourceViewRHI;
+
 	// This pointer is set by the InitRHI() function on the FBlackTextureWithSRV global resource,
 	// to allow FRHITextureReference to use the global black texture when the reference is nullptr.
 	// A pointer is required since FBlackTextureWithSRV is defined in RenderCore.
