Marked several functions on the immediate RHICmdList that do not use "this" as deprecated.
 - Replaced uses of these functions in the engine with the global scope equivalents.
 - Added some missing global scope RHI...() functions.
 - LLM scopes have been moved into the global scope function, where necessary.

Removed several _RenderThread() RHICmdList functions
 - These were previously provided for the automatic RHI thread stall / flush mechanism, but now provide no benefit
 - Removing these also helps prevent use of FRHICommandListExecutor::GetImmediateCommandList() on threads other than the rendering thread.

Remove RHIExecuteCommandList
 - Unused, legacy code path. No platform RHIs provide an implementation. No need to deprecate since there's no way anyone would be calling it already.

diff --git a/Engine/Plugins/Experimental/NNE/Source/NNXRuntimeRDG/Private/NNXRuntimeDML.cpp b/Engine/Plugins/Experimental/NNE/Source/NNXRuntimeRDG/Private/NNXRuntimeDML.cpp
--- a/Engine/Plugins/Experimental/NNE/Source/NNXRuntimeRDG/Private/NNXRuntimeDML.cpp
+++ b/Engine/Plugins/Experimental/NNE/Source/NNXRuntimeRDG/Private/NNXRuntimeDML.cpp
@@ -1987,7 +1987,7 @@ bool FMLInferenceModelDml::InitCompiledOp(TConstArrayView<int32> OpInputIndices,
 
 			if (TensorDataSize)
 			{
-				UploadFence = RHICmdList.CreateGPUFence(TEXT("FMLInferenceModel_UploadFence"));
+				UploadFence = RHICreateGPUFence(TEXT("FMLInferenceModel_UploadFence"));
 				
 				FRHIBuffer*		UploadBuff = CreateRHIBuffer(RHICmdList, TensorDataSize, BUF_ShaderResource | BUF_Dynamic | BUF_FastVRAM, ERHIAccess::CopySrc, TEXT("FMLInferenceModel_UploadBuffer"));
 				uint8*			UploadBuffPtr = static_cast<uint8*>(RHICmdList.LockBuffer(UploadBuff, 0, TensorDataSize, RLM_WriteOnly_NoOverwrite));
diff --git a/Engine/Plugins/FX/Niagara/Source/Niagara/Private/NiagaraGPUProfiler.cpp b/Engine/Plugins/FX/Niagara/Source/Niagara/Private/NiagaraGPUProfiler.cpp
--- a/Engine/Plugins/FX/Niagara/Source/Niagara/Private/NiagaraGPUProfiler.cpp
+++ b/Engine/Plugins/FX/Niagara/Source/Niagara/Private/NiagaraGPUProfiler.cpp
@@ -106,12 +106,12 @@ void FNiagaraGPUProfiler::EndDispatch(FRHICommandList& RHICmdList)
 	RHICmdList.EndRenderQuery(DispatchTimer.EndQuery.GetQuery());
 }
 
-bool FNiagaraGPUProfiler::ProcessFrame(FRHICommandListImmediate& RHICmdList, FGpuFrameData& ReadFrame)
+bool FNiagaraGPUProfiler::ProcessFrame(FRHICommandListImmediate& /*RHICmdList*/, FGpuFrameData& ReadFrame)
 {
 	// Frame ready to process?
 	//-OPT: We can just look at the last write stage end timer here, but that relies on the batcher always executing
 	uint64 DummyEndTime;
-	if (!RHICmdList.GetRenderQueryResult(ReadFrame.EndQuery.GetQuery(), DummyEndTime, false))
+	if (!RHIGetRenderQueryResult(ReadFrame.EndQuery.GetQuery(), DummyEndTime, false))
 	{
 		return false;
 	}
@@ -131,8 +131,8 @@ bool FNiagaraGPUProfiler::ProcessFrame(FRHICommandListImmediate& RHICmdList, FGp
 			
 		uint64 StartMicroseconds = 0;
 		uint64 EndMicroseconds = 0;
-		ensure(RHICmdList.GetRenderQueryResult(DispatchTimer.StartQuery.GetQuery(), StartMicroseconds, false));
-		ensure(RHICmdList.GetRenderQueryResult(DispatchTimer.EndQuery.GetQuery(), EndMicroseconds, false));
+		ensure(RHIGetRenderQueryResult(DispatchTimer.StartQuery.GetQuery(), StartMicroseconds, false));
+		ensure(RHIGetRenderQueryResult(DispatchTimer.EndQuery.GetQuery(), EndMicroseconds, false));
 		DispatchTimer.StartQuery.ReleaseQuery();
 		DispatchTimer.EndQuery.ReleaseQuery();
 
diff --git a/Engine/Plugins/Media/MediaIOFramework/Source/MediaIOCore/Private/MediaCaptureSources.h b/Engine/Plugins/Media/MediaIOFramework/Source/MediaIOCore/Private/MediaCaptureSources.h
--- a/Engine/Plugins/Media/MediaIOFramework/Source/MediaIOCore/Private/MediaCaptureSources.h
+++ b/Engine/Plugins/Media/MediaIOFramework/Source/MediaIOCore/Private/MediaCaptureSources.h
@@ -289,7 +289,7 @@ namespace UE::MediaCapture::Private
 					SourceTexture = Viewport->GetRenderTargetTexture();
 					if (!SourceTexture.IsValid() && Viewport->GetViewportRHI())
 					{
-						SourceTexture = RHICmdList.GetViewportBackBuffer(Viewport->GetViewportRHI());
+						SourceTexture = RHIGetViewportBackBuffer(Viewport->GetViewportRHI());
 					}
 				}
 				else
@@ -297,7 +297,7 @@ namespace UE::MediaCapture::Private
 				if (Viewport->GetViewportRHI())
 				{
 					// Standalone and packaged
-					SourceTexture = RHICmdList.GetViewportBackBuffer(Viewport->GetViewportRHI());
+					SourceTexture = RHIGetViewportBackBuffer(Viewport->GetViewportRHI());
 				}
 			}
 
diff --git a/Engine/Source/Developer/FunctionalTesting/Private/FunctionalUIScreenshotTest.cpp b/Engine/Source/Developer/FunctionalTesting/Private/FunctionalUIScreenshotTest.cpp
--- a/Engine/Source/Developer/FunctionalTesting/Private/FunctionalUIScreenshotTest.cpp
+++ b/Engine/Source/Developer/FunctionalTesting/Private/FunctionalUIScreenshotTest.cpp
@@ -45,7 +45,7 @@ void GetBackbufferInfo(const FViewport* InViewport, EPixelFormat* OutPixelFormat
 	{
 		FViewportRHIRef ViewportRHI = InViewport->GetViewportRHI();
 		check(ViewportRHI.IsValid());
-		FTexture2DRHIRef BackbufferTexture = RHICmdList.GetViewportBackBuffer(ViewportRHI);
+		FTexture2DRHIRef BackbufferTexture = RHIGetViewportBackBuffer(ViewportRHI);
 		check(BackbufferTexture.IsValid());
 		*OutPixelFormat = BackbufferTexture->GetFormat();
 		*OutIsSRGB = (BackbufferTexture->GetFlags() & TexCreate_SRGB) == TexCreate_SRGB;
@@ -163,7 +163,7 @@ void ReadBackbuffer(const FViewport* InViewport, TArray<FColor>* OutPixels)
 		[InViewport, OutPixels](FRHICommandListImmediate& RHICmdList)
 	{
 		FViewportRHIRef ViewportRHI = InViewport->GetViewportRHI();
-		FTexture2DRHIRef BackbufferTexture = RHICmdList.GetViewportBackBuffer(ViewportRHI);
+		FTexture2DRHIRef BackbufferTexture = RHIGetViewportBackBuffer(ViewportRHI);
 		RHICmdList.ReadSurfaceData(
 			BackbufferTexture,
 			FIntRect(0, 0, BackbufferTexture->GetSizeX(), BackbufferTexture->GetSizeY()),
diff --git a/Engine/Source/Editor/VirtualTexturingEditor/Private/RuntimeVirtualTextureBuildStreamingMips.cpp b/Engine/Source/Editor/VirtualTexturingEditor/Private/RuntimeVirtualTextureBuildStreamingMips.cpp
--- a/Engine/Source/Editor/VirtualTexturingEditor/Private/RuntimeVirtualTextureBuildStreamingMips.cpp
+++ b/Engine/Source/Editor/VirtualTexturingEditor/Private/RuntimeVirtualTextureBuildStreamingMips.cpp
@@ -56,8 +56,7 @@ namespace
 				StagingTextures[Layer] = RHICreateTexture(Desc);
 			}
 
-			FRHICommandListImmediate& RHICmdList = FRHICommandListExecutor::GetImmediateCommandList();
-			Fence = RHICmdList.CreateGPUFence(TEXT("Runtime Virtual Texture Build"));
+			Fence = RHICreateGPUFence(TEXT("Runtime Virtual Texture Build"));
 		}
 
 		virtual void ReleaseRHI() override
diff --git a/Engine/Source/Runtime/Apple/MetalRHI/Private/MetalDynamicRHI.cpp b/Engine/Source/Runtime/Apple/MetalRHI/Private/MetalDynamicRHI.cpp
--- a/Engine/Source/Runtime/Apple/MetalRHI/Private/MetalDynamicRHI.cpp
+++ b/Engine/Source/Runtime/Apple/MetalRHI/Private/MetalDynamicRHI.cpp
@@ -131,13 +131,6 @@ FRenderQueryRHIRef FMetalDynamicRHI::RHICreateRenderQuery(ERenderQueryType Query
 	}
 }
 
-FRenderQueryRHIRef FMetalDynamicRHI::RHICreateRenderQuery_RenderThread(class FRHICommandListImmediate& RHICmdList, ERenderQueryType QueryType)
-{
-	@autoreleasepool {
-		return GDynamicRHI->RHICreateRenderQuery(QueryType);
-	}
-}
-
 bool FMetalDynamicRHI::RHIGetRenderQueryResult(FRHIRenderQuery* QueryRHI, uint64& OutNumPixels, bool bWait, uint32 GPUIndex)
 {
 	@autoreleasepool {
diff --git a/Engine/Source/Runtime/Apple/MetalRHI/Private/MetalDynamicRHI.h b/Engine/Source/Runtime/Apple/MetalRHI/Private/MetalDynamicRHI.h
--- a/Engine/Source/Runtime/Apple/MetalRHI/Private/MetalDynamicRHI.h
+++ b/Engine/Source/Runtime/Apple/MetalRHI/Private/MetalDynamicRHI.h
@@ -114,7 +114,6 @@ public:
 	virtual void RHIGetSupportedResolution(uint32& Width, uint32& Height) final override;
 	virtual void RHIVirtualTextureSetFirstMipInMemory(FRHITexture2D* Texture, uint32 FirstMip) final override;
 	virtual void RHIVirtualTextureSetFirstMipVisible(FRHITexture2D* Texture, uint32 FirstMip) final override;
-	virtual void RHIExecuteCommandList(FRHICommandList* CmdList) final override;
 	virtual void* METALRHI_API RHIGetNativeDevice() final override;
 	virtual void* METALRHI_API RHIGetNativeGraphicsQueue() final override;
 	virtual void* METALRHI_API RHIGetNativeComputeQueue() final override;
@@ -137,8 +136,6 @@ public:
 
 	virtual FTextureRHIRef RHICreateTexture_RenderThread(class FRHICommandListImmediate& RHICmdList, const FRHITextureCreateDesc& CreateDesc) final override;
 	
-	virtual FRenderQueryRHIRef RHICreateRenderQuery_RenderThread(class FRHICommandListImmediate& RHICmdList, ERenderQueryType QueryType) final override;
-	
 	virtual FUnorderedAccessViewRHIRef RHICreateUnorderedAccessView_RenderThread(class FRHICommandListImmediate& RHICmdList, FRHIBuffer* Buffer, bool bUseUAVCounter, bool bAppendBuffer) final override;
 	virtual FUnorderedAccessViewRHIRef RHICreateUnorderedAccessView_RenderThread(class FRHICommandListImmediate& RHICmdList, FRHITexture* Texture, uint32 MipLevel, uint16 FirstArraySlice, uint16 NumArraySlices) final override;
 	virtual FUnorderedAccessViewRHIRef RHICreateUnorderedAccessView_RenderThread(class FRHICommandListImmediate& RHICmdList, FRHIBuffer* Buffer, uint8 Format) final override;
@@ -159,11 +156,7 @@ public:
 	virtual FStagingBufferRHIRef RHICreateStagingBuffer() final override;
 	virtual void* RHILockStagingBuffer(FRHIStagingBuffer* StagingBuffer, FRHIGPUFence* Fence, uint32 Offset, uint32 SizeRHI) final override;
 	virtual void RHIUnlockStagingBuffer(FRHIStagingBuffer* StagingBuffer) final override;
-	
-	virtual FVertexShaderRHIRef CreateVertexShader_RenderThread(class FRHICommandListImmediate& RHICmdList, TArrayView<const uint8> Code, const FSHAHash& Hash) final override;
-	virtual FPixelShaderRHIRef CreatePixelShader_RenderThread(class FRHICommandListImmediate& RHICmdList, TArrayView<const uint8> Code, const FSHAHash& Hash) final override;
-	virtual FGeometryShaderRHIRef CreateGeometryShader_RenderThread(class FRHICommandListImmediate& RHICmdList, TArrayView<const uint8> Code, const FSHAHash& Hash) final override;
-	virtual FComputeShaderRHIRef CreateComputeShader_RenderThread(class FRHICommandListImmediate& RHICmdList, TArrayView<const uint8> Code, const FSHAHash& Hash) final override;
+
 	virtual FRHIShaderLibraryRef RHICreateShaderLibrary_RenderThread(class FRHICommandListImmediate& RHICmdList, EShaderPlatform Platform, FString FilePath, FString Name) final override;
 
 	virtual void RHIUpdateUniformBuffer(FRHICommandListBase& RHICmdList, FRHIUniformBuffer* UniformBufferRHI, const void* Contents) final override;
diff --git a/Engine/Source/Runtime/Apple/MetalRHI/Private/MetalDynamicRHI_Shaders.cpp b/Engine/Source/Runtime/Apple/MetalRHI/Private/MetalDynamicRHI_Shaders.cpp
--- a/Engine/Source/Runtime/Apple/MetalRHI/Private/MetalDynamicRHI_Shaders.cpp
+++ b/Engine/Source/Runtime/Apple/MetalRHI/Private/MetalDynamicRHI_Shaders.cpp
@@ -163,26 +163,6 @@ FBoundShaderStateRHIRef FMetalDynamicRHI::RHICreateBoundShaderState(
 	return nullptr;
 }
 
-FVertexShaderRHIRef FMetalDynamicRHI::CreateVertexShader_RenderThread(class FRHICommandListImmediate& RHICmdList, TArrayView<const uint8> Code, const FSHAHash& Hash)
-{
-	return RHICreateVertexShader(Code, Hash);
-}
-
-FGeometryShaderRHIRef FMetalDynamicRHI::CreateGeometryShader_RenderThread(class FRHICommandListImmediate& RHICmdList, TArrayView<const uint8> Code, const FSHAHash& Hash)
-{
-	return RHICreateGeometryShader(Code, Hash);
-}
-
-FPixelShaderRHIRef FMetalDynamicRHI::CreatePixelShader_RenderThread(class FRHICommandListImmediate& RHICmdList, TArrayView<const uint8> Code, const FSHAHash& Hash)
-{
-	return RHICreatePixelShader(Code, Hash);
-}
-
-FComputeShaderRHIRef FMetalDynamicRHI::CreateComputeShader_RenderThread(class FRHICommandListImmediate& RHICmdList, TArrayView<const uint8> Code, const FSHAHash& Hash)
-{
-	return RHICreateComputeShader(Code, Hash);
-}
-
 FRHIShaderLibraryRef FMetalDynamicRHI::RHICreateShaderLibrary_RenderThread(class FRHICommandListImmediate& RHICmdList, EShaderPlatform Platform, FString FilePath, FString Name)
 {
 	return RHICreateShaderLibrary(Platform, FilePath, Name);
diff --git a/Engine/Source/Runtime/Apple/MetalRHI/Private/MetalRHI.cpp b/Engine/Source/Runtime/Apple/MetalRHI/Private/MetalRHI.cpp
--- a/Engine/Source/Runtime/Apple/MetalRHI/Private/MetalRHI.cpp
+++ b/Engine/Source/Runtime/Apple/MetalRHI/Private/MetalRHI.cpp
@@ -1333,11 +1333,6 @@ uint32 FMetalDynamicRHI::RHIGetGPUFrameCycles(uint32 GPUIndex)
 	return GGPUFrameTime;
 }
 
-void FMetalDynamicRHI::RHIExecuteCommandList(FRHICommandList* RHICmdList)
-{
-	NOT_SUPPORTED("RHIExecuteCommandList");
-}
-
 IRHICommandContext* FMetalDynamicRHI::RHIGetDefaultContext()
 {
 	return &ImmediateContext;
diff --git a/Engine/Source/Runtime/D3D12RHI/Private/D3D12Commands.cpp b/Engine/Source/Runtime/D3D12RHI/Private/D3D12Commands.cpp
--- a/Engine/Source/Runtime/D3D12RHI/Private/D3D12Commands.cpp
+++ b/Engine/Source/Runtime/D3D12RHI/Private/D3D12Commands.cpp
@@ -1945,11 +1945,6 @@ uint32 FD3D12DynamicRHI::RHIGetGPUFrameCycles(uint32 GPUIndex)
 	return GGPUFrameTime;
 }
 
-void FD3D12DynamicRHI::RHIExecuteCommandList(FRHICommandList* CmdList)
-{
-	check(0); // this path has gone stale and needs updated methods, starting at ERCT_SetScissorRect
-}
-
 
 void FD3D12CommandContext::RHISetDepthBounds(float MinDepth, float MaxDepth)
 {
diff --git a/Engine/Source/Runtime/D3D12RHI/Private/D3D12RHIPrivate.h b/Engine/Source/Runtime/D3D12RHI/Private/D3D12RHIPrivate.h
--- a/Engine/Source/Runtime/D3D12RHI/Private/D3D12RHIPrivate.h
+++ b/Engine/Source/Runtime/D3D12RHI/Private/D3D12RHIPrivate.h
@@ -518,7 +518,6 @@ public:
 	virtual void RHIGetSupportedResolution(uint32& Width, uint32& Height) final override;
 	virtual void RHIVirtualTextureSetFirstMipInMemory(FRHITexture2D* Texture, uint32 FirstMip) override;
 	virtual void RHIVirtualTextureSetFirstMipVisible(FRHITexture2D* Texture, uint32 FirstMip) override;
-	virtual void RHIExecuteCommandList(FRHICommandList* CmdList) final override;
 	virtual void* RHIGetNativeDevice() final override;
 	virtual void* RHIGetNativeGraphicsQueue() final override;
 	virtual void* RHIGetNativeComputeQueue() final override;
@@ -563,36 +562,6 @@ public:
 	// These will be un-commented as they are implemented.
 	//
 
-	virtual FVertexShaderRHIRef CreateVertexShader_RenderThread(class FRHICommandListImmediate& RHICmdList, TArrayView<const uint8> Code, const FSHAHash& Hash) override final
-	{
-		return RHICreateVertexShader(Code, Hash);
-	}
-
-	virtual FMeshShaderRHIRef CreateMeshShader_RenderThread(class FRHICommandListImmediate& RHICmdList, TArrayView<const uint8> Code, const FSHAHash& Hash) override final
-	{
-		return RHICreateMeshShader(Code, Hash);
-	}
-
-	virtual FAmplificationShaderRHIRef CreateAmplificationShader_RenderThread(class FRHICommandListImmediate& RHICmdList, TArrayView<const uint8> Code, const FSHAHash& Hash) override final
-	{
-		return RHICreateAmplificationShader(Code, Hash);
-	}
-
-	virtual FGeometryShaderRHIRef CreateGeometryShader_RenderThread(class FRHICommandListImmediate& RHICmdList, TArrayView<const uint8> Code, const FSHAHash& Hash) override final
-	{
-		return RHICreateGeometryShader(Code, Hash);
-	}
-
-	virtual FPixelShaderRHIRef CreatePixelShader_RenderThread(class FRHICommandListImmediate& RHICmdList, TArrayView<const uint8> Code, const FSHAHash& Hash) override final
-	{
-		return RHICreatePixelShader(Code, Hash);
-	}
-
-	virtual FComputeShaderRHIRef CreateComputeShader_RenderThread(class FRHICommandListImmediate& RHICmdList, TArrayView<const uint8> Code, const FSHAHash& Hash) override final
-	{
-		return RHICreateComputeShader(Code, Hash);
-	}
-
 	virtual void* LockTexture2D_RenderThread(class FRHICommandListImmediate& RHICmdList, FRHITexture2D* Texture, uint32 MipIndex, EResourceLockMode LockMode, uint32& DestStride, bool bLockWithinMiptail, bool bNeedsDefaultRHIFlush = true) override final;
 	virtual void UnlockTexture2D_RenderThread(class FRHICommandListImmediate& RHICmdList, FRHITexture2D* Texture, uint32 MipIndex, bool bLockWithinMiptail, bool bNeedsDefaultRHIFlush = true) override final;
 
@@ -625,11 +594,6 @@ public:
 	virtual FShaderResourceViewRHIRef CreateShaderResourceView_RenderThread(class FRHICommandListImmediate& RHICmdList, const FShaderResourceViewInitializer& Initializer) override final;
 	virtual FShaderResourceViewRHIRef RHICreateShaderResourceView_RenderThread(class FRHICommandListImmediate& RHICmdList, FRHIBuffer* Buffer) override final;
 
-	virtual FRenderQueryRHIRef RHICreateRenderQuery_RenderThread(class FRHICommandListImmediate& RHICmdList, ERenderQueryType QueryType) override final
-	{
-		return RHICreateRenderQuery(QueryType);
-	}
-
 	void RHICalibrateTimers() override;
 
 #if D3D12_RHI_RAYTRACING
diff --git a/Engine/Source/Runtime/Engine/Private/GPUSkinCache.cpp b/Engine/Source/Runtime/Engine/Private/GPUSkinCache.cpp
--- a/Engine/Source/Runtime/Engine/Private/GPUSkinCache.cpp
+++ b/Engine/Source/Runtime/Engine/Private/GPUSkinCache.cpp
@@ -1133,7 +1133,7 @@ void FGPUSkinCache::DispatchUpdateSkinTangents(FRHICommandListImmediate& RHICmdL
 			{
 				StagingBuffer->Release();
 				StagingBuffer->Buffer.Initialize(TEXT("SkinTangentIntermediate"), sizeof(int32), NumIntsPerBuffer, PF_R32_SINT, BUF_UnorderedAccess);
-				RHICmdList.BindDebugLabelName(StagingBuffer->Buffer.UAV, TEXT("SkinTangentIntermediate"));
+				RHIBindDebugLabelName(StagingBuffer->Buffer.UAV, TEXT("SkinTangentIntermediate"));
 
 				const uint32 MemSize = NumIntsPerBuffer * sizeof(uint32);
 				SET_MEMORY_STAT(STAT_GPUSkinCache_TangentsIntermediateMemUsed, MemSize);
diff --git a/Engine/Source/Runtime/Engine/Private/Slate/SceneViewport.cpp b/Engine/Source/Runtime/Engine/Private/Slate/SceneViewport.cpp
--- a/Engine/Source/Runtime/Engine/Private/Slate/SceneViewport.cpp
+++ b/Engine/Source/Runtime/Engine/Private/Slate/SceneViewport.cpp
@@ -1875,7 +1875,7 @@ void FSceneViewport::BeginRenderFrame(FRHICommandListImmediate& RHICmdList)
 	else if( IsValidRef( ViewportRHI ) ) 
 	{
 		// Get the backbuffer render target to render directly to it
-		RenderTargetTextureRenderThreadRHI = RHICmdList.GetViewportBackBuffer(ViewportRHI);
+		RenderTargetTextureRenderThreadRHI = RHIGetViewportBackBuffer(ViewportRHI);
 		RenderThreadSlateTexture->SetRHIRef(RenderTargetTextureRenderThreadRHI, RenderTargetTextureRenderThreadRHI->GetSizeX(), RenderTargetTextureRenderThreadRHI->GetSizeY());
 	}
 }
diff --git a/Engine/Source/Runtime/Engine/Private/UnrealClient.cpp b/Engine/Source/Runtime/Engine/Private/UnrealClient.cpp
--- a/Engine/Source/Runtime/Engine/Private/UnrealClient.cpp
+++ b/Engine/Source/Runtime/Engine/Private/UnrealClient.cpp
@@ -1644,7 +1644,7 @@ void FViewport::EndRenderFrame(FRHICommandListImmediate& RHICmdList, bool bPrese
 
 FRHIGPUMask FViewport::GetGPUMask(FRHICommandListImmediate& RHICmdList) const
 {
-	return FRHIGPUMask::FromIndex(RHICmdList.GetViewportNextPresentGPUIndex(GetViewportRHI()));
+	return FRHIGPUMask::FromIndex(RHIGetViewportNextPresentGPUIndex(GetViewportRHI()));
 }
 
 void APostProcessVolume::PostUnregisterAllComponents()
diff --git a/Engine/Source/Runtime/MediaAssets/Private/Misc/MediaTextureResource.cpp b/Engine/Source/Runtime/MediaAssets/Private/Misc/MediaTextureResource.cpp
--- a/Engine/Source/Runtime/MediaAssets/Private/Misc/MediaTextureResource.cpp
+++ b/Engine/Source/Runtime/MediaAssets/Private/Misc/MediaTextureResource.cpp
@@ -77,16 +77,14 @@ template<typename ObjectRefType> struct TGPUsyncedDataDeleter : public IMediaClo
 
 	void Retire(const ObjectRefType& Object)
 	{
-		FRHICommandListImmediate& CommandList = FRHICommandListExecutor::GetImmediateCommandList();
-
 		// Prep "retirement package"
 		FRetiringObjectInfo Info;
 		Info.Object = Object;
-		Info.GPUFence = CommandList.CreateGPUFence(TEXT("MediaTextureResourceReuseFence"));
+		Info.GPUFence = RHICreateGPUFence(TEXT("MediaTextureResourceReuseFence"));
 		Info.RetireTime = FPlatformTime::Seconds();
 
 		// Insert fence. We assume that GPU-workload-wise this marks the spot usage of the sample is done
-		CommandList.WriteGPUFence(Info.GPUFence);
+		FRHICommandListExecutor::GetImmediateCommandList().WriteGPUFence(Info.GPUFence);
 
 		// Recall for later checking...
 		FScopeLock Lock(&CS);
diff --git a/Engine/Source/Runtime/NullDrv/Public/NullRHI.h b/Engine/Source/Runtime/NullDrv/Public/NullRHI.h
--- a/Engine/Source/Runtime/NullDrv/Public/NullRHI.h
+++ b/Engine/Source/Runtime/NullDrv/Public/NullRHI.h
@@ -599,10 +599,6 @@ public:
 	virtual void RHIVirtualTextureSetFirstMipVisible(FRHITexture2D* Texture, uint32 FirstMip) final override
 	{
 
-	}
-	virtual void RHIExecuteCommandList(FRHICommandList* CmdList) final override
-	{
-
 	}
 	virtual void RHIEnableDepthBoundsTest(bool bEnable) final override
 	{
diff --git a/Engine/Source/Runtime/OpenGLDrv/Private/OpenGLCommands.cpp b/Engine/Source/Runtime/OpenGLDrv/Private/OpenGLCommands.cpp
--- a/Engine/Source/Runtime/OpenGLDrv/Private/OpenGLCommands.cpp
+++ b/Engine/Source/Runtime/OpenGLDrv/Private/OpenGLCommands.cpp
@@ -2932,11 +2932,6 @@ void FOpenGLDynamicRHI::RHISetMultipleViewports(uint32 Count, const FViewportBou
 	UE_LOG(LogRHI, Fatal,TEXT("OpenGL Render path does not support multiple Viewports!"));
 }
 
-void FOpenGLDynamicRHI::RHIExecuteCommandList(FRHICommandList*)
-{
-	check(0);
-}
-
 void FOpenGLDynamicRHI::RHIEnableDepthBoundsTest(bool bEnable)
 {
 	if (FOpenGL::SupportsDepthBoundsTest())
diff --git a/Engine/Source/Runtime/OpenGLDrv/Public/OpenGLDrv.h b/Engine/Source/Runtime/OpenGLDrv/Public/OpenGLDrv.h
--- a/Engine/Source/Runtime/OpenGLDrv/Public/OpenGLDrv.h
+++ b/Engine/Source/Runtime/OpenGLDrv/Public/OpenGLDrv.h
@@ -435,7 +435,6 @@ public:
 	virtual void RHIGetSupportedResolution(uint32& Width, uint32& Height) final override;
 	virtual void RHIVirtualTextureSetFirstMipInMemory(FRHITexture2D* Texture, uint32 FirstMip) final override;
 	virtual void RHIVirtualTextureSetFirstMipVisible(FRHITexture2D* Texture, uint32 FirstMip) final override;
-	virtual void RHIExecuteCommandList(FRHICommandList* CmdList) final override;
 	virtual void* RHIGetNativeDevice() final override;
 	virtual void* RHIGetNativeInstance() final override;
 	virtual class IRHICommandContext* RHIGetDefaultContext() final override;
@@ -827,11 +826,6 @@ public:
 		return this->RHICreateShaderResourceView(Buffer);
 	}
 
-	virtual FRenderQueryRHIRef RHICreateRenderQuery_RenderThread(class FRHICommandListImmediate& RHICmdList, ERenderQueryType QueryType) final override
-	{
-		return this->RHICreateRenderQuery(QueryType);
-	}
-
 	virtual FTexture2DRHIRef AsyncReallocateTexture2D_RenderThread(class FRHICommandListImmediate& RHICmdList, FRHITexture2D* Texture2D, int32 NewMipCount, int32 NewSizeX, int32 NewSizeY, FThreadSafeCounter* RequestStatus) final override;
 
 	virtual ETextureReallocationStatus FinalizeAsyncReallocateTexture2D_RenderThread(class FRHICommandListImmediate& RHICmdList, FRHITexture2D* Texture2D, bool bBlockUntilCompleted) final override
@@ -844,26 +838,6 @@ public:
 		return this->RHICancelAsyncReallocateTexture2D(Texture2D, bBlockUntilCompleted);
 	}
 
-	virtual FVertexShaderRHIRef CreateVertexShader_RenderThread(class FRHICommandListImmediate& RHICmdList, TArrayView<const uint8> Code, const FSHAHash& Hash) final override
-	{
-		return this->RHICreateVertexShader(Code, Hash);
-	}
-
-	virtual FPixelShaderRHIRef CreatePixelShader_RenderThread(class FRHICommandListImmediate& RHICmdList, TArrayView<const uint8> Code, const FSHAHash& Hash) final override
-	{
-		return this->RHICreatePixelShader(Code, Hash);
-	}
-
-	virtual FGeometryShaderRHIRef CreateGeometryShader_RenderThread(class FRHICommandListImmediate& RHICmdList, TArrayView<const uint8> Code, const FSHAHash& Hash) final override
-	{
-		return this->RHICreateGeometryShader(Code, Hash);
-	}
-
-	virtual FComputeShaderRHIRef CreateComputeShader_RenderThread(class FRHICommandListImmediate& RHICmdList, TArrayView<const uint8> Code, const FSHAHash& Hash) final override
-	{
-		return this->RHICreateComputeShader(Code, Hash);
-	}
-
 	virtual FSamplerStateRHIRef RHICreateSamplerState(const FSamplerStateInitializerRHI& Initializer) final override;
 
 	virtual FGraphicsPipelineStateRHIRef RHICreateGraphicsPipelineState(const FGraphicsPipelineStateInitializer& Initializer)
diff --git a/Engine/Source/Runtime/RHI/Private/RHICommandList.cpp b/Engine/Source/Runtime/RHI/Private/RHICommandList.cpp
--- a/Engine/Source/Runtime/RHI/Private/RHICommandList.cpp
+++ b/Engine/Source/Runtime/RHI/Private/RHICommandList.cpp
@@ -1791,48 +1791,6 @@ ETextureReallocationStatus FDynamicRHI::CancelAsyncReallocateTexture2D_RenderThr
 	return GDynamicRHI->RHICancelAsyncReallocateTexture2D(Texture2D, bBlockUntilCompleted);
 }
 
-FVertexShaderRHIRef FDynamicRHI::CreateVertexShader_RenderThread(class FRHICommandListImmediate& RHICmdList, TArrayView<const uint8> Code, const FSHAHash& Hash)
-{
-	CSV_SCOPED_TIMING_STAT(RHITStalls, CreateVertexShader_RenderThread);
-	FScopedRHIThreadStaller StallRHIThread(RHICmdList);
-	return GDynamicRHI->RHICreateVertexShader(Code, Hash);
-}
-
-FMeshShaderRHIRef FDynamicRHI::CreateMeshShader_RenderThread(class FRHICommandListImmediate& RHICmdList, TArrayView<const uint8> Code, const FSHAHash& Hash)
-{
-	CSV_SCOPED_TIMING_STAT(RHITStalls, CreateMeshShader_RenderThread);
-	FScopedRHIThreadStaller StallRHIThread(RHICmdList);
-	return GDynamicRHI->RHICreateMeshShader(Code, Hash);
-}
-
-FAmplificationShaderRHIRef FDynamicRHI::CreateAmplificationShader_RenderThread(class FRHICommandListImmediate& RHICmdList, TArrayView<const uint8> Code, const FSHAHash& Hash)
-{
-	CSV_SCOPED_TIMING_STAT(RHITStalls, CreateAmplificationShader_RenderThread);
-	FScopedRHIThreadStaller StallRHIThread(RHICmdList);
-	return GDynamicRHI->RHICreateAmplificationShader(Code, Hash);
-}
-
-FPixelShaderRHIRef FDynamicRHI::CreatePixelShader_RenderThread(class FRHICommandListImmediate& RHICmdList, TArrayView<const uint8> Code, const FSHAHash& Hash)
-{
-	CSV_SCOPED_TIMING_STAT(RHITStalls, CreatePixelShader_RenderThread);
-	FScopedRHIThreadStaller StallRHIThread(RHICmdList);
-	return GDynamicRHI->RHICreatePixelShader(Code, Hash);
-}
-
-FGeometryShaderRHIRef FDynamicRHI::CreateGeometryShader_RenderThread(class FRHICommandListImmediate& RHICmdList, TArrayView<const uint8> Code, const FSHAHash& Hash)
-{
-	CSV_SCOPED_TIMING_STAT(RHITStalls, CreateGeometryShader_RenderThread);
-	FScopedRHIThreadStaller StallRHIThread(RHICmdList);
-	return GDynamicRHI->RHICreateGeometryShader(Code, Hash);
-}
-
-FComputeShaderRHIRef FDynamicRHI::CreateComputeShader_RenderThread(class FRHICommandListImmediate& RHICmdList, TArrayView<const uint8> Code, const FSHAHash& Hash)
-{
-	CSV_SCOPED_TIMING_STAT(RHITStalls, CreateGeometryShader_RenderThread);
-	FScopedRHIThreadStaller StallRHIThread(RHICmdList);
-	return GDynamicRHI->RHICreateComputeShader(Code, Hash);
-}
-
 FUpdateTexture3DData FDynamicRHI::RHIBeginUpdateTexture3D(FRHICommandListBase& RHICmdList, FRHITexture3D* Texture, uint32 MipIndex, const struct FUpdateTextureRegion3D& UpdateRegion)
 {
 	check(IsInParallelRenderingThread());
@@ -1992,13 +1950,6 @@ FShaderResourceViewRHIRef FDynamicRHI::RHICreateShaderResourceViewFMask_RenderTh
 	return GDynamicRHI->RHICreateShaderResourceViewFMask(Texture2DRHI);
 }
 
-FRenderQueryRHIRef FDynamicRHI::RHICreateRenderQuery_RenderThread(class FRHICommandListImmediate& RHICmdList, ERenderQueryType QueryType)
-{
-	CSV_SCOPED_TIMING_STAT(RHITStalls, RHICreateRenderQuery_RenderThread);
-	FScopedRHIThreadStaller StallRHIThread(RHICmdList);
-	return GDynamicRHI->RHICreateRenderQuery(QueryType);
-}
-
 void* FDynamicRHI::RHILockTextureCubeFace_RenderThread(class FRHICommandListImmediate& RHICmdList, FRHITextureCube* Texture, uint32 FaceIndex, uint32 ArrayIndex, uint32 MipIndex, EResourceLockMode LockMode, uint32& DestStride, bool bLockWithinMiptail)
 {
 	QUICK_SCOPE_CYCLE_COUNTER(STAT_RHIMETHOD_LockTextureCubeFace_Flush);
diff --git a/Engine/Source/Runtime/RHI/Public/DynamicRHI.h b/Engine/Source/Runtime/RHI/Public/DynamicRHI.h
--- a/Engine/Source/Runtime/RHI/Public/DynamicRHI.h
+++ b/Engine/Source/Runtime/RHI/Public/DynamicRHI.h
@@ -926,9 +926,6 @@ public:
 
 	}
 
-	// FlushType: Wait RHI Thread
-	virtual void RHIExecuteCommandList(FRHICommandList* CmdList) = 0;
-
 	/**
 	* Provides access to the native device. Generally this should be avoided but is useful for third party plugins.
 	*/
@@ -1026,12 +1023,7 @@ public:
 	virtual FTexture2DRHIRef AsyncReallocateTexture2D_RenderThread(class FRHICommandListImmediate& RHICmdList, FRHITexture2D* Texture2D, int32 NewMipCount, int32 NewSizeX, int32 NewSizeY, FThreadSafeCounter* RequestStatus);
 	virtual ETextureReallocationStatus FinalizeAsyncReallocateTexture2D_RenderThread(class FRHICommandListImmediate& RHICmdList, FRHITexture2D* Texture2D, bool bBlockUntilCompleted);
 	virtual ETextureReallocationStatus CancelAsyncReallocateTexture2D_RenderThread(class FRHICommandListImmediate& RHICmdList, FRHITexture2D* Texture2D, bool bBlockUntilCompleted);
-	virtual FVertexShaderRHIRef CreateVertexShader_RenderThread(class FRHICommandListImmediate& RHICmdList, TArrayView<const uint8> Code, const FSHAHash& Hash);
-	virtual FMeshShaderRHIRef CreateMeshShader_RenderThread(class FRHICommandListImmediate& RHICmdList, TArrayView<const uint8> Code, const FSHAHash& Hash);
-	virtual FAmplificationShaderRHIRef CreateAmplificationShader_RenderThread(class FRHICommandListImmediate& RHICmdList, TArrayView<const uint8> Code, const FSHAHash& Hash);
-	virtual FPixelShaderRHIRef CreatePixelShader_RenderThread(class FRHICommandListImmediate& RHICmdList, TArrayView<const uint8> Code, const FSHAHash& Hash);
-	virtual FGeometryShaderRHIRef CreateGeometryShader_RenderThread(class FRHICommandListImmediate& RHICmdList, TArrayView<const uint8> Code, const FSHAHash& Hash);
-	virtual FComputeShaderRHIRef CreateComputeShader_RenderThread(class FRHICommandListImmediate& RHICmdList, TArrayView<const uint8> Code, const FSHAHash& Hash);
+
 	virtual void* LockTexture2D_RenderThread(class FRHICommandListImmediate& RHICmdList, FRHITexture2D* Texture, uint32 MipIndex, EResourceLockMode LockMode, uint32& DestStride, bool bLockWithinMiptail, bool bNeedsDefaultRHIFlush = true);
 	virtual void UnlockTexture2D_RenderThread(class FRHICommandListImmediate& RHICmdList, FRHITexture2D* Texture, uint32 MipIndex, bool bLockWithinMiptail, bool bNeedsDefaultRHIFlush = true);
 
@@ -1057,8 +1049,6 @@ public:
 	virtual FShaderResourceViewRHIRef RHICreateShaderResourceView_RenderThread(class FRHICommandListImmediate& RHICmdList, FRHIBuffer* Buffer);
 	virtual FShaderResourceViewRHIRef RHICreateShaderResourceViewWriteMask_RenderThread(class FRHICommandListImmediate& RHICmdList, FRHITexture2D* Texture2D);
 	virtual FShaderResourceViewRHIRef RHICreateShaderResourceViewFMask_RenderThread(class FRHICommandListImmediate& RHICmdList, FRHITexture2D* Texture2D);
-	virtual FRenderQueryRHIRef RHICreateRenderQuery_RenderThread(class FRHICommandListImmediate& RHICmdList, ERenderQueryType QueryType);
-
 	
 	virtual void* RHILockTextureCubeFace_RenderThread(class FRHICommandListImmediate& RHICmdList, FRHITextureCube* Texture, uint32 FaceIndex, uint32 ArrayIndex, uint32 MipIndex, EResourceLockMode LockMode, uint32& DestStride, bool bLockWithinMiptail);
 	virtual void RHIUnlockTextureCubeFace_RenderThread(class FRHICommandListImmediate& RHICmdList, FRHITextureCube* Texture, uint32 FaceIndex, uint32 ArrayIndex, uint32 MipIndex, bool bLockWithinMiptail);
@@ -1187,34 +1177,92 @@ FORCEINLINE TRHI* GetDynamicRHI()
 	return CastDynamicRHI<TRHI>(GDynamicRHI);
 }
 
+FORCEINLINE FPixelShaderRHIRef RHICreatePixelShader(TArrayView<const uint8> Code, const FSHAHash& Hash)
+{
+	LLM_SCOPE(ELLMTag::Shaders);
+	return GDynamicRHI->RHICreatePixelShader(Code, Hash);
+}
+
+FORCEINLINE FVertexShaderRHIRef RHICreateVertexShader(TArrayView<const uint8> Code, const FSHAHash& Hash)
+{
+	LLM_SCOPE(ELLMTag::Shaders);
+	return GDynamicRHI->RHICreateVertexShader(Code, Hash);
+}
+
+FORCEINLINE FMeshShaderRHIRef RHICreateMeshShader(TArrayView<const uint8> Code, const FSHAHash& Hash)
+{
+	LLM_SCOPE(ELLMTag::Shaders);
+	return GDynamicRHI->RHICreateMeshShader(Code, Hash);
+}
+
+FORCEINLINE FAmplificationShaderRHIRef RHICreateAmplificationShader(TArrayView<const uint8> Code, const FSHAHash& Hash)
+{
+	LLM_SCOPE(ELLMTag::Shaders);
+	return GDynamicRHI->RHICreateAmplificationShader(Code, Hash);
+}
+
+FORCEINLINE FGeometryShaderRHIRef RHICreateGeometryShader(TArrayView<const uint8> Code, const FSHAHash& Hash)
+{
+	LLM_SCOPE(ELLMTag::Shaders);
+	return GDynamicRHI->RHICreateGeometryShader(Code, Hash);
+}
+
+FORCEINLINE FComputeShaderRHIRef RHICreateComputeShader(TArrayView<const uint8> Code, const FSHAHash& Hash)
+{
+	LLM_SCOPE(ELLMTag::Shaders);
+	return GDynamicRHI->RHICreateComputeShader(Code, Hash);
+}
+
+UE_DEPRECATED(5.2, "Compute fences are deprecated. Use RHI transitions instead.")
+FORCEINLINE FComputeFenceRHIRef RHICreateComputeFence(const FName& Name)
+{
+	return GDynamicRHI->RHICreateComputeFence(Name);
+}
+
+FORCEINLINE FGPUFenceRHIRef RHICreateGPUFence(const FName& Name)
+{
+	return GDynamicRHI->RHICreateGPUFence(Name);
+}
+
+FORCEINLINE FStagingBufferRHIRef RHICreateStagingBuffer()
+{
+	return GDynamicRHI->RHICreateStagingBuffer();
+}
+
 FORCEINLINE FSamplerStateRHIRef RHICreateSamplerState(const FSamplerStateInitializerRHI& Initializer)
 {
+	LLM_SCOPE_BYNAME(TEXT("RHIMisc/CreateSamplerState"));
 	return GDynamicRHI->RHICreateSamplerState(Initializer);
 }
 
 FORCEINLINE FRasterizerStateRHIRef RHICreateRasterizerState(const FRasterizerStateInitializerRHI& Initializer)
 {
+	LLM_SCOPE_BYNAME(TEXT("RHIMisc/CreateRasterizerState"));
 	return GDynamicRHI->RHICreateRasterizerState(Initializer);
 }
 
 FORCEINLINE FDepthStencilStateRHIRef RHICreateDepthStencilState(const FDepthStencilStateInitializerRHI& Initializer)
 {
+	LLM_SCOPE_BYNAME(TEXT("RHIMisc/CreateDepthStencilState"));
 	return GDynamicRHI->RHICreateDepthStencilState(Initializer);
 }
 
 FORCEINLINE FBlendStateRHIRef RHICreateBlendState(const FBlendStateInitializerRHI& Initializer)
 {
+	LLM_SCOPE_BYNAME(TEXT("RHIMisc/CreateBlendState"));
 	return GDynamicRHI->RHICreateBlendState(Initializer);
 }
 
 FORCEINLINE FBoundShaderStateRHIRef RHICreateBoundShaderState(FRHIVertexDeclaration* VertexDeclaration, FRHIVertexShader* VertexShader, FRHIPixelShader* PixelShader, FRHIGeometryShader* GeometryShader)
 {
+	LLM_SCOPE(ELLMTag::Shaders);
 	return GDynamicRHI->RHICreateBoundShaderState(VertexDeclaration, VertexShader, PixelShader, GeometryShader);
 }
 
 /** Before using this directly go through PipelineStateCache::GetAndOrCreateGraphicsPipelineState() */
 FORCEINLINE FGraphicsPipelineStateRHIRef RHICreateGraphicsPipelineState(const FGraphicsPipelineStateInitializer& Initializer)
 {
+	LLM_SCOPE(ELLMTag::Shaders);
 	return GDynamicRHI->RHICreateGraphicsPipelineState(Initializer);
 }
 
@@ -1226,9 +1274,15 @@ FORCEINLINE FVertexDeclarationRHIRef RHICreateVertexDeclaration(const FVertexDec
 
 FORCEINLINE FComputePipelineStateRHIRef RHICreateComputePipelineState(FRHIComputeShader* ComputeShader)
 {
+	LLM_SCOPE(ELLMTag::Shaders);
 	return GDynamicRHI->RHICreateComputePipelineState(ComputeShader);
 }
 
+FORCEINLINE FRenderQueryRHIRef RHICreateRenderQuery(ERenderQueryType QueryType)
+{
+	return GDynamicRHI->RHICreateRenderQuery(QueryType);
+}
+
 #if RHI_RAYTRACING
 FORCEINLINE TRefCountPtr<FRHIRayTracingPipelineState> RHICreateRayTracingPipelineState(const FRayTracingPipelineStateInitializer& Initializer)
 {
@@ -1463,11 +1517,13 @@ FORCEINLINE uint32 RHIGetGPUFrameCycles(uint32 GPUIndex = 0)
 
 FORCEINLINE FViewportRHIRef RHICreateViewport(void* WindowHandle, uint32 SizeX, uint32 SizeY, bool bIsFullscreen, EPixelFormat PreferredPixelFormat)
 {
+	LLM_SCOPE(ELLMTag::RenderTargets);
 	return GDynamicRHI->RHICreateViewport(WindowHandle, SizeX, SizeY, bIsFullscreen, PreferredPixelFormat);
 }
 
 FORCEINLINE void RHIResizeViewport(FRHIViewport* Viewport, uint32 SizeX, uint32 SizeY, bool bIsFullscreen, EPixelFormat PreferredPixelFormat)
 {
+	LLM_SCOPE(ELLMTag::RenderTargets);
 	GDynamicRHI->RHIResizeViewport(Viewport, SizeX, SizeY, bIsFullscreen, PreferredPixelFormat);
 }
 
@@ -1488,6 +1544,7 @@ FORCEINLINE void RHIHandleDisplayChange()
 
 FORCEINLINE void RHITick(float DeltaTime)
 {
+	LLM_SCOPE_BYNAME(TEXT("RHIMisc/RHITick"));
 	GDynamicRHI->RHITick(DeltaTime);
 }
 
diff --git a/Engine/Source/Runtime/RHI/Public/RHICommandList.h b/Engine/Source/Runtime/RHI/Public/RHICommandList.h
--- a/Engine/Source/Runtime/RHI/Public/RHICommandList.h
+++ b/Engine/Source/Runtime/RHI/Public/RHICommandList.h
@@ -4435,99 +4435,105 @@ public:
 		}
 	}
 
+	UE_DEPRECATED(5.2, "Use the global scope RHICreateSamplerState function.")
 	FORCEINLINE FSamplerStateRHIRef CreateSamplerState(const FSamplerStateInitializerRHI& Initializer)
 	{
-		LLM_SCOPE_BYNAME(TEXT("RHIMisc/CreateSamplerState"));
 		return RHICreateSamplerState(Initializer);
 	}
 	
+	UE_DEPRECATED(5.2, "Use the global scope RHICreateRasterizerState function.")
 	FORCEINLINE FRasterizerStateRHIRef CreateRasterizerState(const FRasterizerStateInitializerRHI& Initializer)
 	{
-		LLM_SCOPE_BYNAME(TEXT("RHIMisc/CreateRasterizerState"));
 		return RHICreateRasterizerState(Initializer);
 	}
 	
+	UE_DEPRECATED(5.2, "Use the global scope RHICreateDepthStencilState function.")
 	FORCEINLINE FDepthStencilStateRHIRef CreateDepthStencilState(const FDepthStencilStateInitializerRHI& Initializer)
 	{
-		LLM_SCOPE_BYNAME(TEXT("RHIMisc/CreateDepthStencilState"));
 		return RHICreateDepthStencilState(Initializer);
 	}
 	
+	UE_DEPRECATED(5.2, "Use the global scope RHICreateBlendState function.")
 	FORCEINLINE FBlendStateRHIRef CreateBlendState(const FBlendStateInitializerRHI& Initializer)
 	{
-		LLM_SCOPE_BYNAME(TEXT("RHIMisc/CreateBlendState"));
 		return RHICreateBlendState(Initializer);
 	}
-	
+
+	UE_DEPRECATED(5.2, "Use the global scope RHICreatePixelShader function.")
 	FORCEINLINE FPixelShaderRHIRef CreatePixelShader(TArrayView<const uint8> Code, const FSHAHash& Hash)
 	{
-		LLM_SCOPE(ELLMTag::Shaders);
-		return GDynamicRHI->CreatePixelShader_RenderThread(*this, Code, Hash);
+		return RHICreatePixelShader(Code, Hash);
 	}
 	
+	UE_DEPRECATED(5.2, "Use the global scope RHICreateVertexShader function.")
 	FORCEINLINE FVertexShaderRHIRef CreateVertexShader(TArrayView<const uint8> Code, const FSHAHash& Hash)
 	{
-		LLM_SCOPE(ELLMTag::Shaders);
-		return GDynamicRHI->CreateVertexShader_RenderThread(*this, Code, Hash);
+		return RHICreateVertexShader(Code, Hash);
 	}
 
+	UE_DEPRECATED(5.2, "Use the global scope RHICreateMeshShader function.")
 	FORCEINLINE FMeshShaderRHIRef CreateMeshShader(TArrayView<const uint8> Code, const FSHAHash& Hash)
 	{
-		LLM_SCOPE(ELLMTag::Shaders);
-		return GDynamicRHI->CreateMeshShader_RenderThread(*this, Code, Hash);
+		return RHICreateMeshShader(Code, Hash);
 	}
 
+	UE_DEPRECATED(5.2, "Use the global scope RHICreateAmplificationShader function.")
 	FORCEINLINE FAmplificationShaderRHIRef CreateAmplificationShader(TArrayView<const uint8> Code, const FSHAHash& Hash)
 	{
-		LLM_SCOPE(ELLMTag::Shaders);
-		return GDynamicRHI->CreateAmplificationShader_RenderThread(*this, Code, Hash);
+		return RHICreateAmplificationShader(Code, Hash);
 	}
 	
+	UE_DEPRECATED(5.2, "Use the global scope RHICreateGeometryShader function.")
 	FORCEINLINE FGeometryShaderRHIRef CreateGeometryShader(TArrayView<const uint8> Code, const FSHAHash& Hash)
 	{
-		LLM_SCOPE(ELLMTag::Shaders);
-		return GDynamicRHI->CreateGeometryShader_RenderThread(*this, Code, Hash);
+		return RHICreateGeometryShader(Code, Hash);
 	}
 	
+	UE_DEPRECATED(5.2, "Use the global scope RHICreateComputeShader function.")
 	FORCEINLINE FComputeShaderRHIRef CreateComputeShader(TArrayView<const uint8> Code, const FSHAHash& Hash)
 	{
-		LLM_SCOPE(ELLMTag::Shaders);
-		return GDynamicRHI->CreateComputeShader_RenderThread(*this, Code, Hash);
+		return RHICreateComputeShader(Code, Hash);
 	}
 	
+	UE_DEPRECATED(5.2, "Compute fences are deprecated. Use RHI transitions instead.")
 	FORCEINLINE FComputeFenceRHIRef CreateComputeFence(const FName& Name)
-	{		
-		return GDynamicRHI->RHICreateComputeFence(Name);
+	{
+	PRAGMA_DISABLE_DEPRECATION_WARNINGS
+		return RHICreateComputeFence(Name);
+	PRAGMA_ENABLE_DEPRECATION_WARNINGS
 	}	
 
+	UE_DEPRECATED(5.2, "Use the global scope RHICreateGPUFence function.")
 	FORCEINLINE FGPUFenceRHIRef CreateGPUFence(const FName& Name)
 	{
-		return GDynamicRHI->RHICreateGPUFence(Name);
+		return RHICreateGPUFence(Name);
 	}
 
+	UE_DEPRECATED(5.2, "Use the global scope RHICreateStagingBuffer function.")
 	FORCEINLINE FStagingBufferRHIRef CreateStagingBuffer()
 	{
-		return GDynamicRHI->RHICreateStagingBuffer();
+		return RHICreateStagingBuffer();
 	}
 
+	UE_DEPRECATED(5.2, "Use the global scope RHICreateBoundShaderState function.")
 	FORCEINLINE FBoundShaderStateRHIRef CreateBoundShaderState(FRHIVertexDeclaration* VertexDeclaration, FRHIVertexShader* VertexShader, FRHIPixelShader* PixelShader, FRHIGeometryShader* GeometryShader)
 	{
-		LLM_SCOPE(ELLMTag::Shaders);
 		return RHICreateBoundShaderState(VertexDeclaration, VertexShader, PixelShader, GeometryShader);
 	}
 
+	UE_DEPRECATED(5.2, "Use the global scope RHICreateGraphicsPipelineState function.")
 	FORCEINLINE FGraphicsPipelineStateRHIRef CreateGraphicsPipelineState(const FGraphicsPipelineStateInitializer& Initializer)
 	{
-		LLM_SCOPE(ELLMTag::Shaders);
 		return RHICreateGraphicsPipelineState(Initializer);
 	}
 
-	FORCEINLINE TRefCountPtr<FRHIComputePipelineState> CreateComputePipelineState(FRHIComputeShader* ComputeShader)
+	UE_DEPRECATED(5.2, "Use the global scope RHICreateComputePipelineState function.")
+	FORCEINLINE FComputePipelineStateRHIRef CreateComputePipelineState(FRHIComputeShader* ComputeShader)
 	{
-		LLM_SCOPE(ELLMTag::Shaders);
 		return RHICreateComputePipelineState(ComputeShader);
 	}
 
+	UE_DEPRECATED(5.2, "Use the global scope RHICreateUniformBuffer function.")
 	FORCEINLINE FUniformBufferRHIRef CreateUniformBuffer(const void* Contents, const FRHIUniformBufferLayout* Layout, EUniformBufferUsage Usage)
 	{
 		return RHICreateUniformBuffer(Contents, Layout, Usage);
@@ -4678,6 +4684,7 @@ public:
 		PRAGMA_ENABLE_DEPRECATION_WARNINGS
 	}
 	
+	UE_DEPRECATED(5.2, "Use the global scope RHIGetTextureMemoryStats function.")
 	FORCEINLINE void GetTextureMemoryStats(FTextureMemoryStats& OutStats)
 	{
 		RHIGetTextureMemoryStats(OutStats);
@@ -4734,9 +4741,11 @@ public:
 	FORCEINLINE void GenerateMips(FRHITexture* Texture)
 	{
 		QUICK_SCOPE_CYCLE_COUNTER(STAT_RHIMETHOD_GenerateMips_Flush);
-		ImmediateFlush(EImmediateFlushType::FlushRHIThread); return GDynamicRHI->RHIGenerateMips(Texture);
+		ImmediateFlush(EImmediateFlushType::FlushRHIThread);
+		return GDynamicRHI->RHIGenerateMips(Texture);
 	}
 	
+	UE_DEPRECATED(5.2, "Use the global scope RHIComputeMemorySize function.")
 	FORCEINLINE uint32 ComputeMemorySize(FRHITexture* TextureRHI)
 	{
 		return RHIComputeMemorySize(TextureRHI);
@@ -4815,11 +4824,13 @@ public:
 		GDynamicRHI->RHIUnlockTextureCubeFace_RenderThread(*this, Texture, FaceIndex, ArrayIndex, MipIndex, bLockWithinMiptail);
 	}
 	
+	UE_DEPRECATED(5.2, "Use the global scope RHIBindDebugLabelName function.")
 	FORCEINLINE void BindDebugLabelName(FRHITexture* Texture, const TCHAR* Name)
 	{
 		RHIBindDebugLabelName(Texture, Name);
 	}
 
+	UE_DEPRECATED(5.2, "Use the global scope RHIBindDebugLabelName function.")
 	FORCEINLINE void BindDebugLabelName(FRHIUnorderedAccessView* UnorderedAccessViewRHI, const TCHAR* Name)
 	{
 		RHIBindDebugLabelName(UnorderedAccessViewRHI, Name);
@@ -4879,21 +4890,25 @@ public:
 		GDynamicRHI->RHIRead3DSurfaceFloatData(Texture,Rect,ZMinMax,OutData,Flags);
 	}
 
+	UE_DEPRECATED(5.2, "Use the global scope RHIGetRenderQueryResult function.")
 	FORCEINLINE bool GetRenderQueryResult(FRHIRenderQuery* RenderQuery, uint64& OutResult, bool bWait, uint32 GPUIndex = INDEX_NONE)
 	{
 		return RHIGetRenderQueryResult(RenderQuery, OutResult, bWait, GPUIndex);
 	}
 
+	UE_DEPRECATED(5.2, "Use the global scope RHIGetViewportNextPresentGPUIndex function.")
 	FORCEINLINE uint32 GetViewportNextPresentGPUIndex(FRHIViewport* Viewport)
 	{
-		return GDynamicRHI->RHIGetViewportNextPresentGPUIndex(Viewport);
+		return RHIGetViewportNextPresentGPUIndex(Viewport);
 	}
 
+	UE_DEPRECATED(5.2, "Use the global scope RHIGetViewportBackBuffer function.")
 	FORCEINLINE FTexture2DRHIRef GetViewportBackBuffer(FRHIViewport* Viewport)
 	{
 		return RHIGetViewportBackBuffer(Viewport);
 	}
 	
+	UE_DEPRECATED(5.2, "Use the global scope RHIAdvanceFrameForGetViewportBackBuffer function.")
 	FORCEINLINE void AdvanceFrameForGetViewportBackBuffer(FRHIViewport* Viewport)
 	{
 		return RHIAdvanceFrameForGetViewportBackBuffer(Viewport);
@@ -4930,21 +4945,21 @@ public:
 		return RHIGetGPUFrameCycles(GetGPUMask().ToIndex());
 	}
 	
+	UE_DEPRECATED(5.2, "Use the global scope RHICreateViewport function.")
 	FORCEINLINE FViewportRHIRef CreateViewport(void* WindowHandle, uint32 SizeX, uint32 SizeY, bool bIsFullscreen, EPixelFormat PreferredPixelFormat)
 	{
-		LLM_SCOPE(ELLMTag::RenderTargets);
 		return RHICreateViewport(WindowHandle, SizeX, SizeY, bIsFullscreen, PreferredPixelFormat);
 	}
 	
+	UE_DEPRECATED(5.2, "Use the global scope RHIResizeViewport function.")
 	FORCEINLINE void ResizeViewport(FRHIViewport* Viewport, uint32 SizeX, uint32 SizeY, bool bIsFullscreen, EPixelFormat PreferredPixelFormat)
 	{
-		LLM_SCOPE(ELLMTag::RenderTargets);
 		RHIResizeViewport(Viewport, SizeX, SizeY, bIsFullscreen, PreferredPixelFormat);
 	}
 	
+	UE_DEPRECATED(5.2, "Use the global scope RHITick function.")
 	FORCEINLINE void Tick(float DeltaTime)
 	{
-		LLM_SCOPE_BYNAME(TEXT("RHIMisc/RHITick"));
 		RHITick(DeltaTime);
 	}
 	
@@ -4962,11 +4977,13 @@ public:
 		GDynamicRHI->RHISubmitCommandsAndFlushGPU();
 	}
 	
+	UE_DEPRECATED(5.2, "Use the global scope RHISuspendRendering function.")
 	FORCEINLINE void SuspendRendering()
 	{
 		RHISuspendRendering();
 	}
 	
+	UE_DEPRECATED(5.2, "Use the global scope RHIResumeRendering function.")
 	FORCEINLINE void ResumeRendering()
 	{
 		RHIResumeRendering();
@@ -4985,11 +5002,13 @@ public:
 		return false;
 	}
 
+	UE_DEPRECATED(5.2, "Use the global scope RHIGetAvailableResolutions function.")
 	FORCEINLINE bool GetAvailableResolutions(FScreenResolutionArray& Resolutions, bool bIgnoreRefreshRate)
 	{
 		return RHIGetAvailableResolutions(Resolutions, bIgnoreRefreshRate);
 	}
 	
+	UE_DEPRECATED(5.2, "Use the global scope RHIGetSupportedResolution function.")
 	FORCEINLINE void GetSupportedResolution(uint32& Width, uint32& Height)
 	{
 		RHIGetSupportedResolution(Width, Height);
@@ -5004,12 +5023,6 @@ public:
 	{
 		GDynamicRHI->VirtualTextureSetFirstMipVisible_RenderThread(*this, Texture, FirstMip);
 	}
-
-	FORCEINLINE void ExecuteCommandList(FRHICommandList* CmdList)
-	{
-		FScopedRHIThreadStaller StallRHIThread(*this);
-		GDynamicRHI->RHIExecuteCommandList(CmdList);
-	}
 	
 	FORCEINLINE void* GetNativeDevice()
 	{
@@ -5056,6 +5069,7 @@ public:
 		return GDynamicRHI->RHIGetNativeCommandBuffer();
 	}
 
+	UE_DEPRECATED(5.2, "Use the global scope RHIGetDefaultContext function.")
 	FORCEINLINE class IRHICommandContext* GetDefaultContext()
 	{
 		return RHIGetDefaultContext();
@@ -5400,54 +5414,11 @@ public:
 
 FORCEINLINE_DEBUGGABLE FRHICommandListImmediate& FRHICommandListExecutor::GetImmediateCommandList()
 {
+	// @todo - fix remaining use of the immediate command list on other threads, then uncomment this check.
+	//check(IsInRenderingThread());
 	return GRHICommandList.CommandListImmediate;
 }
 
-FORCEINLINE FPixelShaderRHIRef RHICreatePixelShader(TArrayView<const uint8> Code, const FSHAHash& Hash)
-{
-	return FRHICommandListExecutor::GetImmediateCommandList().CreatePixelShader(Code, Hash);
-}
-
-FORCEINLINE FVertexShaderRHIRef RHICreateVertexShader(TArrayView<const uint8> Code, const FSHAHash& Hash)
-{
-	return FRHICommandListExecutor::GetImmediateCommandList().CreateVertexShader(Code, Hash);
-}
-
-FORCEINLINE FMeshShaderRHIRef RHICreateMeshShader(TArrayView<const uint8> Code, const FSHAHash& Hash)
-{
-	return FRHICommandListExecutor::GetImmediateCommandList().CreateMeshShader(Code, Hash);
-}
-
-FORCEINLINE FAmplificationShaderRHIRef RHICreateAmplificationShader(TArrayView<const uint8> Code, const FSHAHash& Hash)
-{
-	return FRHICommandListExecutor::GetImmediateCommandList().CreateAmplificationShader(Code, Hash);
-}
-
-FORCEINLINE FGeometryShaderRHIRef RHICreateGeometryShader(TArrayView<const uint8> Code, const FSHAHash& Hash)
-{
-	return FRHICommandListExecutor::GetImmediateCommandList().CreateGeometryShader(Code, Hash);
-}
-
-FORCEINLINE FComputeShaderRHIRef RHICreateComputeShader(TArrayView<const uint8> Code, const FSHAHash& Hash)
-{
-	return FRHICommandListExecutor::GetImmediateCommandList().CreateComputeShader(Code, Hash);
-}
-
-FORCEINLINE FComputeFenceRHIRef RHICreateComputeFence(const FName& Name)
-{
-	return FRHICommandListExecutor::GetImmediateCommandList().CreateComputeFence(Name);
-}
-
-FORCEINLINE FGPUFenceRHIRef RHICreateGPUFence(const FName& Name)
-{
-	return FRHICommandListExecutor::GetImmediateCommandList().CreateGPUFence(Name);
-}
-
-FORCEINLINE FStagingBufferRHIRef RHICreateStagingBuffer()
-{
-	return FRHICommandListExecutor::GetImmediateCommandList().CreateStagingBuffer();
-}
-
 FORCEINLINE FBufferRHIRef RHICreateBuffer(uint32 Size, EBufferUsageFlags Usage, uint32 Stride, ERHIAccess ResourceState, FRHIResourceCreateInfo& CreateInfo)
 {
 	check(IsInRenderingThread());
@@ -5910,11 +5881,6 @@ FORCEINLINE void RHIVirtualTextureSetFirstMipVisible(FRHITexture2D* Texture, uin
 	 FRHICommandListExecutor::GetImmediateCommandList().VirtualTextureSetFirstMipVisible(Texture, FirstMip);
 }
 
-FORCEINLINE void RHIExecuteCommandList(FRHICommandList* CmdList)
-{
-	 FRHICommandListExecutor::GetImmediateCommandList().ExecuteCommandList(CmdList);
-}
-
 FORCEINLINE void* RHIGetNativeDevice()
 {
 	return FRHICommandListExecutor::GetImmediateCommandList().GetNativeDevice();
diff --git a/Engine/Source/Runtime/RHI/Public/RHIValidation.h b/Engine/Source/Runtime/RHI/Public/RHIValidation.h
--- a/Engine/Source/Runtime/RHI/Public/RHIValidation.h
+++ b/Engine/Source/Runtime/RHI/Public/RHIValidation.h
@@ -1205,12 +1205,6 @@ public:
 		RHI->RHIPerFrameRHIFlushComplete();
 	}
 
-	// FlushType: Wait RHI Thread
-	virtual void RHIExecuteCommandList(FRHICommandList* CmdList) override final
-	{
-		return RHI->RHIExecuteCommandList(CmdList);
-	}
-
 	/**
 	* Provides access to the native device. Generally this should be avoided but is useful for third party plugins.
 	*/
@@ -1335,39 +1329,6 @@ public:
 		RHI->UnlockBuffer_BottomOfPipe(RHICmdList, Buffer);
 	}
 
-	virtual FVertexShaderRHIRef CreateVertexShader_RenderThread(class FRHICommandListImmediate& RHICmdList, TArrayView<const uint8> Code, const FSHAHash& Hash) override final
-	{
-		return RHI->CreateVertexShader_RenderThread(RHICmdList, Code, Hash);
-	}
-
-	virtual FPixelShaderRHIRef CreatePixelShader_RenderThread(class FRHICommandListImmediate& RHICmdList, TArrayView<const uint8> Code, const FSHAHash& Hash) override final
-	{
-		return RHI->CreatePixelShader_RenderThread(RHICmdList, Code, Hash);
-	}
-
-	virtual FGeometryShaderRHIRef CreateGeometryShader_RenderThread(class FRHICommandListImmediate& RHICmdList, TArrayView<const uint8> Code, const FSHAHash& Hash) override final
-	{
-		check(RHISupportsGeometryShaders(GMaxRHIShaderPlatform));
-		return RHI->CreateGeometryShader_RenderThread(RHICmdList, Code, Hash);
-	}
-
-	virtual FMeshShaderRHIRef CreateMeshShader_RenderThread(class FRHICommandListImmediate& RHICmdList, TArrayView<const uint8> Code, const FSHAHash& Hash) override final
-	{
-		check(RHISupportsMeshShadersTier0(GMaxRHIShaderPlatform));
-		return RHI->CreateMeshShader_RenderThread(RHICmdList, Code, Hash);
-	}
-
-	virtual FAmplificationShaderRHIRef CreateAmplificationShader_RenderThread(class FRHICommandListImmediate& RHICmdList, TArrayView<const uint8> Code, const FSHAHash& Hash) override final
-	{
-		check(RHISupportsMeshShadersTier0(GMaxRHIShaderPlatform));
-		return RHI->CreateAmplificationShader_RenderThread(RHICmdList, Code, Hash);
-	}
-
-	virtual FComputeShaderRHIRef CreateComputeShader_RenderThread(class FRHICommandListImmediate& RHICmdList, TArrayView<const uint8> Code, const FSHAHash& Hash) override final
-	{
-		return RHI->CreateComputeShader_RenderThread(RHICmdList, Code, Hash);
-	}
-
 	virtual void* LockTexture2D_RenderThread(class FRHICommandListImmediate& RHICmdList, FRHITexture2D* Texture, uint32 MipIndex, EResourceLockMode LockMode, uint32& DestStride, bool bLockWithinMiptail, bool bNeedsDefaultRHIFlush = true) override final
 	{
 		return RHI->LockTexture2D_RenderThread(RHICmdList, Texture, MipIndex, LockMode, DestStride, bLockWithinMiptail, bNeedsDefaultRHIFlush);
@@ -1495,12 +1456,6 @@ public:
 		}
 		return SRV;
 	}
-	
-	virtual FRenderQueryRHIRef RHICreateRenderQuery_RenderThread(class FRHICommandListImmediate& RHICmdList, ERenderQueryType QueryType) override final
-	{
-		return RHI->RHICreateRenderQuery_RenderThread(RHICmdList, QueryType);
-	}
-
 
 	virtual void* RHILockTextureCubeFace_RenderThread(class FRHICommandListImmediate& RHICmdList, FRHITextureCube* Texture, uint32 FaceIndex, uint32 ArrayIndex, uint32 MipIndex, EResourceLockMode LockMode, uint32& DestStride, bool bLockWithinMiptail) override final
 	{
diff --git a/Engine/Source/Runtime/RenderCore/Private/LongGPUTask.cpp b/Engine/Source/Runtime/RenderCore/Private/LongGPUTask.cpp
--- a/Engine/Source/Runtime/RenderCore/Private/LongGPUTask.cpp
+++ b/Engine/Source/Runtime/RenderCore/Private/LongGPUTask.cpp
@@ -95,7 +95,7 @@ void MeasureLongGPUTaskExecutionTime(FRHICommandListImmediate& RHICmdList)
 	RHICmdList.SubmitCommandsHint();
 	RHICmdList.ImmediateFlush(EImmediateFlushType::FlushRHIThread);
 
-	if (RHICmdList.GetRenderQueryResult(TimeQueryStart.GetQuery(), StartTime, true) && RHICmdList.GetRenderQueryResult(TimeQueryEnd.GetQuery(), EndTime, true))
+	if (RHIGetRenderQueryResult(TimeQueryStart.GetQuery(), StartTime, true) && RHIGetRenderQueryResult(TimeQueryEnd.GetQuery(), EndTime, true))
 	{
 		NumMeasuredIterationsToAchieve500ms = FMath::Clamp(FMath::FloorToInt(500.0f / ((EndTime - StartTime) / 1000.0f / NumIterationsForMeasurement)), 1, 2000);
 	}
diff --git a/Engine/Source/Runtime/RenderCore/Private/ProfilingDebugging/RealtimeGPUProfiler.cpp b/Engine/Source/Runtime/RenderCore/Private/ProfilingDebugging/RealtimeGPUProfiler.cpp
--- a/Engine/Source/Runtime/RenderCore/Private/ProfilingDebugging/RealtimeGPUProfiler.cpp
+++ b/Engine/Source/Runtime/RenderCore/Private/ProfilingDebugging/RealtimeGPUProfiler.cpp
@@ -201,7 +201,7 @@ public:
 		{
 			if (StartResultMicroseconds[GPUIndex] == InvalidQueryResult)
 			{
-				if (!RHICmdList.GetRenderQueryResult(StartQuery.GetQuery(), StartResultMicroseconds[GPUIndex], false, GPUIndex))
+				if (!RHIGetRenderQueryResult(StartQuery.GetQuery(), StartResultMicroseconds[GPUIndex], false, GPUIndex))
 				{
 					StartResultMicroseconds[GPUIndex] = InvalidQueryResult;
 				}
@@ -209,7 +209,7 @@ public:
 
 			if (EndResultMicroseconds[GPUIndex] == InvalidQueryResult)
 			{
-				if (!RHICmdList.GetRenderQueryResult(EndQuery.GetQuery(), EndResultMicroseconds[GPUIndex], false, GPUIndex))
+				if (!RHIGetRenderQueryResult(EndQuery.GetQuery(), EndResultMicroseconds[GPUIndex], false, GPUIndex))
 				{
 					EndResultMicroseconds[GPUIndex] = InvalidQueryResult;
 				}
diff --git a/Engine/Source/Runtime/Renderer/Private/GPUBenchmark.cpp b/Engine/Source/Runtime/Renderer/Private/GPUBenchmark.cpp
--- a/Engine/Source/Runtime/Renderer/Private/GPUBenchmark.cpp
+++ b/Engine/Source/Runtime/Renderer/Private/GPUBenchmark.cpp
@@ -582,7 +582,7 @@ void RendererGPUBenchmark(FRHICommandListImmediate& RHICmdList, FSynthBenchmarkR
 			// flushes the RHI thread to make sure all RHICmdList.EndRenderQuery() commands got executed.
 			RHICmdList.SubmitCommandsHint();
 			RHICmdList.ImmediateFlush(EImmediateFlushType::FlushRHIThread);
-			RHICmdList.GetRenderQueryResult(TimerQueries[0].GetQuery(), OldAbsTime, true);
+			RHIGetRenderQueryResult(TimerQueries[0].GetQuery(), OldAbsTime, true);
 
 			for(uint32 Iteration = 0; Iteration < IterationCount; ++Iteration)
 			{
@@ -598,7 +598,7 @@ void RendererGPUBenchmark(FRHICommandListImmediate& RHICmdList, FSynthBenchmarkR
 					uint32 QueryIndex = 1 + Iteration * MethodCount + MethodId;
 
 					uint64 AbsTime;
-					RHICmdList.GetRenderQueryResult(TimerQueries[QueryIndex].GetQuery(), AbsTime, true);
+					RHIGetRenderQueryResult(TimerQueries[QueryIndex].GetQuery(), AbsTime, true);
 
 					uint64 RelTime = FMath::Max(AbsTime - OldAbsTime, 1ull);
 
diff --git a/Engine/Source/Runtime/Renderer/Private/SceneOcclusion.cpp b/Engine/Source/Runtime/Renderer/Private/SceneOcclusion.cpp
--- a/Engine/Source/Runtime/Renderer/Private/SceneOcclusion.cpp
+++ b/Engine/Source/Runtime/Renderer/Private/SceneOcclusion.cpp
@@ -386,7 +386,7 @@ FRHIRenderQuery* FFrameBasedOcclusionQueryPool::AllocateQuery()
 	}
 
 	// If all fails, create a new query
-	FRenderQueryRHIRef NewQuery = GDynamicRHI->RHICreateRenderQuery(RQT_Occlusion);
+	FRenderQueryRHIRef NewQuery = RHICreateRenderQuery(RQT_Occlusion);
 	if (NewQuery)
 	{
 		CurrentFrame.Queries.Add(MoveTemp(NewQuery));
diff --git a/Engine/Source/Runtime/Renderer/Private/SceneVisibility.cpp b/Engine/Source/Runtime/Renderer/Private/SceneVisibility.cpp
--- a/Engine/Source/Runtime/Renderer/Private/SceneVisibility.cpp
+++ b/Engine/Source/Runtime/Renderer/Private/SceneVisibility.cpp
@@ -1408,8 +1408,7 @@ static void FetchVisibilityForPrimitives_Range(FVisForPrimParams& Params, FGloba
 							{
 								//int32 RefCount = PastQuery.GetReference()->GetRefCount();
 								// NOTE: RHIGetOcclusionQueryResult should never fail when using a blocking call, rendering artifacts may show up.
-								//if (RHICmdList.GetRenderQueryResult(PastQuery, NumSamples, true))
-								if (GDynamicRHI->RHIGetRenderQueryResult(PastQuery, NumSamples, true))
+								if (RHIGetRenderQueryResult(PastQuery, NumSamples, true))
 								{
 									// we render occlusion without MSAA
 									uint32 NumPixels = (uint32)NumSamples;
diff --git a/Engine/Source/Runtime/Renderer/Private/VT/VirtualTextureFeedback.cpp b/Engine/Source/Runtime/Renderer/Private/VT/VirtualTextureFeedback.cpp
--- a/Engine/Source/Runtime/Renderer/Private/VT/VirtualTextureFeedback.cpp
+++ b/Engine/Source/Runtime/Renderer/Private/VT/VirtualTextureFeedback.cpp
@@ -41,7 +41,7 @@ public:
 	{
 		if (!Fences[Index])
 		{
-			Fences[Index] = RHICmdList.CreateGPUFence(FName(""));
+			Fences[Index] = RHICreateGPUFence(FName(""));
 		}
 		Fences[Index]->Clear();
 	}
@@ -130,7 +130,7 @@ public:
 	bool Poll(FRHICommandListImmediate& RHICmdList, int32 Index)
 	{
 		uint64 Dummy;
-		return RHICmdList.GetRenderQueryResult(Fences[Index], Dummy, false, RHICmdList.GetGPUMask().ToIndex());
+		return RHIGetRenderQueryResult(Fences[Index], Dummy, false, RHICmdList.GetGPUMask().ToIndex());
 	}
 
 	FGPUFenceRHIRef GetMapFence(int32 Index)
diff --git a/Engine/Source/Runtime/SlateRHIRenderer/Private/SlateRHIRenderer.cpp b/Engine/Source/Runtime/SlateRHIRenderer/Private/SlateRHIRenderer.cpp
--- a/Engine/Source/Runtime/SlateRHIRenderer/Private/SlateRHIRenderer.cpp
+++ b/Engine/Source/Runtime/SlateRHIRenderer/Private/SlateRHIRenderer.cpp
@@ -1022,7 +1022,7 @@ void FSlateRHIRenderer::DrawWindow_RenderThread(FRHICommandListImmediate& RHICmd
 
 	static uint32 LastTimestamp = FPlatformTime::Cycles();
 	{
-		const FRHIGPUMask PresentingGPUMask = FRHIGPUMask::FromIndex(RHICmdList.GetViewportNextPresentGPUIndex(ViewportInfo.ViewportRHI));
+		const FRHIGPUMask PresentingGPUMask = FRHIGPUMask::FromIndex(RHIGetViewportNextPresentGPUIndex(ViewportInfo.ViewportRHI));
 		SCOPED_GPU_MASK(RHICmdList, PresentingGPUMask);
 		SCOPED_DRAW_EVENTF(RHICmdList, SlateUI, TEXT("SlateUI Title = %s"), DrawCommandParams.WindowTitle.IsEmpty() ? TEXT("<none>") : *DrawCommandParams.WindowTitle);
 		SCOPED_GPU_STAT(RHICmdList, SlateUI);
@@ -1051,7 +1051,7 @@ void FSlateRHIRenderer::DrawWindow_RenderThread(FRHICommandListImmediate& RHICmd
 		if (CVarDrawToVRRenderTarget->GetInt() == 0 && GEngine && IsValidRef(ViewportInfo.GetRenderTargetTexture()) && GEngine->StereoRenderingDevice.IsValid())
 		{
 			const FVector2D WindowSize = WindowElementList.GetWindowSize();
-			GEngine->StereoRenderingDevice->RenderTexture_RenderThread(RHICmdList, RHICmdList.GetViewportBackBuffer(ViewportInfo.ViewportRHI), ViewportInfo.GetRenderTargetTexture(), WindowSize);
+			GEngine->StereoRenderingDevice->RenderTexture_RenderThread(RHICmdList, RHIGetViewportBackBuffer(ViewportInfo.ViewportRHI), ViewportInfo.GetRenderTargetTexture(), WindowSize);
 			bRenderedStereo = true;
 		}
 
@@ -1064,7 +1064,7 @@ void FSlateRHIRenderer::DrawWindow_RenderThread(FRHICommandListImmediate& RHICmd
 			check(IsValidRef(ViewportInfo.ViewportRHI));
 
 			FTexture2DRHIRef ViewportRT = bRenderedStereo ? nullptr : ViewportInfo.GetRenderTargetTexture();
-			FTexture2DRHIRef BackBuffer = (ViewportRT) ? ViewportRT : RHICmdList.GetViewportBackBuffer(ViewportInfo.ViewportRHI);
+			FTexture2DRHIRef BackBuffer = (ViewportRT) ? ViewportRT : RHIGetViewportBackBuffer(ViewportInfo.ViewportRHI);
 			FTexture2DRHIRef PostProcessBuffer = BackBuffer;	// If compositing UI then this will be different to the back buffer
 
 			const uint32 ViewportWidth = (ViewportRT) ? ViewportRT->GetSizeX() : ViewportInfo.Width;
@@ -1426,7 +1426,7 @@ void FSlateRHIRenderer::DrawWindow_RenderThread(FRHICommandListImmediate& RHICmd
 			if (!bRenderedStereo && GEngine && IsValidRef(ViewportInfo.GetRenderTargetTexture()) && GEngine->StereoRenderingDevice.IsValid())
 			{
 				const FVector2D WindowSize = WindowElementList.GetWindowSize();
-				GEngine->StereoRenderingDevice->RenderTexture_RenderThread(RHICmdList, RHICmdList.GetViewportBackBuffer(ViewportInfo.ViewportRHI), ViewportInfo.GetRenderTargetTexture(), WindowSize);
+				GEngine->StereoRenderingDevice->RenderTexture_RenderThread(RHICmdList, RHIGetViewportBackBuffer(ViewportInfo.ViewportRHI), ViewportInfo.GetRenderTargetTexture(), WindowSize);
 			}
 			RHICmdList.Transition(FRHITransitionInfo(BackBuffer, ERHIAccess::Unknown, ERHIAccess::SRVGraphics));
 
@@ -1444,7 +1444,7 @@ void FSlateRHIRenderer::DrawWindow_RenderThread(FRHICommandListImmediate& RHICmd
 	if (bTakingAScreenShot && ScreenshotViewportInfo != nullptr && ScreenshotViewportInfo == &ViewportInfo)
 	{
 		// take screenshot before swapbuffer
-		FTexture2DRHIRef BackBuffer = RHICmdList.GetViewportBackBuffer(ViewportInfo.ViewportRHI);
+		FTexture2DRHIRef BackBuffer = RHIGetViewportBackBuffer(ViewportInfo.ViewportRHI);
 
 		// Sanity check to make sure the user specified a valid screenshot rect.
 		FIntRect ClampedScreenshotRect;
diff --git a/Engine/Source/Runtime/VulkanRHI/Private/VulkanCommands.cpp b/Engine/Source/Runtime/VulkanRHI/Private/VulkanCommands.cpp
--- a/Engine/Source/Runtime/VulkanRHI/Private/VulkanCommands.cpp
+++ b/Engine/Source/Runtime/VulkanRHI/Private/VulkanCommands.cpp
@@ -783,11 +783,6 @@ uint32 FVulkanDynamicRHI::RHIGetGPUFrameCycles(uint32 GPUIndex)
 	return GGPUFrameTime;
 }
 
-void FVulkanDynamicRHI::RHIExecuteCommandList(FRHICommandList* CmdList)
-{
-	VULKAN_SIGNAL_UNIMPLEMENTED();
-}
-
 void FVulkanCommandListContext::RHISetDepthBounds(float MinDepth, float MaxDepth)
 {
 	FVulkanCmdBuffer* CmdBuffer = CommandBufferManager->GetActiveCmdBuffer();
diff --git a/Engine/Source/Runtime/VulkanRHI/Public/VulkanDynamicRHI.h b/Engine/Source/Runtime/VulkanRHI/Public/VulkanDynamicRHI.h
--- a/Engine/Source/Runtime/VulkanRHI/Public/VulkanDynamicRHI.h
+++ b/Engine/Source/Runtime/VulkanRHI/Public/VulkanDynamicRHI.h
@@ -187,7 +187,6 @@ public:
 	virtual void RHIGetSupportedResolution(uint32& Width, uint32& Height) final override;
 	virtual void RHIVirtualTextureSetFirstMipInMemory(FRHITexture2D* Texture, uint32 FirstMip) final override;
 	virtual void RHIVirtualTextureSetFirstMipVisible(FRHITexture2D* Texture, uint32 FirstMip) final override;
-	virtual void RHIExecuteCommandList(FRHICommandList* CmdList) final override;
 
 	virtual void* RHIGetNativeDevice() final override;
 	virtual void* RHIGetNativePhysicalDevice() final override;
@@ -207,22 +206,6 @@ public:
 	virtual FTexture2DRHIRef AsyncReallocateTexture2D_RenderThread(class FRHICommandListImmediate& RHICmdList, FRHITexture2D* Texture2D, int32 NewMipCount, int32 NewSizeX, int32 NewSizeY, FThreadSafeCounter* RequestStatus) override final;
 	//virtual ETextureReallocationStatus CancelAsyncReallocateTexture2D_RenderThread(class FRHICommandListImmediate& RHICmdList, FRHITexture2D* Texture2D, bool bBlockUntilCompleted) override final;
 	
-	virtual FVertexShaderRHIRef CreateVertexShader_RenderThread(class FRHICommandListImmediate& RHICmdList, TArrayView<const uint8> Code, const FSHAHash& Hash) override final
-	{
-		return RHICreateVertexShader(Code, Hash);
-	}
-	virtual FPixelShaderRHIRef CreatePixelShader_RenderThread(class FRHICommandListImmediate& RHICmdList, TArrayView<const uint8> Code, const FSHAHash& Hash) override final
-	{
-		return RHICreatePixelShader(Code, Hash);
-	}
-	virtual FGeometryShaderRHIRef CreateGeometryShader_RenderThread(class FRHICommandListImmediate& RHICmdList, TArrayView<const uint8> Code, const FSHAHash& Hash) override final
-	{
-		return RHICreateGeometryShader(Code, Hash);
-	}
-	virtual FComputeShaderRHIRef CreateComputeShader_RenderThread(class FRHICommandListImmediate& RHICmdList, TArrayView<const uint8> Code, const FSHAHash& Hash) override final
-	{
-		return RHICreateComputeShader(Code, Hash);
-	}
 	virtual void* LockTexture2D_RenderThread(class FRHICommandListImmediate& RHICmdList, FRHITexture2D* Texture, uint32 MipIndex, EResourceLockMode LockMode, uint32& DestStride, bool bLockWithinMiptail, bool bNeedsDefaultRHIFlush = true) override final
 	{
 		return RHILockTexture2D(Texture, MipIndex, LockMode, DestStride, bLockWithinMiptail);
@@ -286,12 +269,6 @@ public:
 		return this->RHICreateShaderResourceView(StructuredBuffer);
 	}
 
-	virtual FRenderQueryRHIRef RHICreateRenderQuery_RenderThread(class FRHICommandListImmediate& RHICmdList, ERenderQueryType QueryType) override final
-	{
-		// this-> is required to avoid calling the global version of this function
-		return this->RHICreateRenderQuery(QueryType);
-	}
-
 	void RHICalibrateTimers() override;
 
 #if VULKAN_RHI_RAYTRACING
diff --git a/Engine/Source/Runtime/Windows/D3D11RHI/Private/D3D11Commands.cpp b/Engine/Source/Runtime/Windows/D3D11RHI/Private/D3D11Commands.cpp
--- a/Engine/Source/Runtime/Windows/D3D11RHI/Private/D3D11Commands.cpp
+++ b/Engine/Source/Runtime/Windows/D3D11RHI/Private/D3D11Commands.cpp
@@ -1532,11 +1532,6 @@ uint32 FD3D11DynamicRHI::RHIGetGPUFrameCycles(uint32 GPUIndex)
 	return GGPUFrameTime;
 }
 
-void FD3D11DynamicRHI::RHIExecuteCommandList(FRHICommandList* CmdList)
-{
-	check(0); // this path has gone stale and needs updated methods, starting at ERCT_SetScissorRect
-}
-
 // NVIDIA Depth Bounds Test interface
 void FD3D11DynamicRHI::EnableDepthBoundsTest(bool bEnable,float MinDepth,float MaxDepth)
 {
diff --git a/Engine/Source/Runtime/Windows/D3D11RHI/Private/D3D11Query.cpp b/Engine/Source/Runtime/Windows/D3D11RHI/Private/D3D11Query.cpp
--- a/Engine/Source/Runtime/Windows/D3D11RHI/Private/D3D11Query.cpp
+++ b/Engine/Source/Runtime/Windows/D3D11RHI/Private/D3D11Query.cpp
@@ -178,13 +178,6 @@ void FD3D11DynamicRHI::RHIEndOcclusionQueryBatch()
 	RequestedOcclusionQueriesInBatch = 0;
 }
 
-FRenderQueryRHIRef FD3D11DynamicRHI::RHICreateRenderQuery_RenderThread(
-	class FRHICommandListImmediate& RHICmdList,
-	ERenderQueryType QueryType)
-{
-	return RHICreateRenderQuery(QueryType);
-}
-
 FRenderQueryRHIRef FD3D11DynamicRHI::RHICreateRenderQuery(ERenderQueryType QueryType)
 {
 	TRefCountPtr<ID3D11Query> Query;
diff --git a/Engine/Source/Runtime/Windows/D3D11RHI/Private/D3D11RHIPrivate.h b/Engine/Source/Runtime/Windows/D3D11RHI/Private/D3D11RHIPrivate.h
--- a/Engine/Source/Runtime/Windows/D3D11RHI/Private/D3D11RHIPrivate.h
+++ b/Engine/Source/Runtime/Windows/D3D11RHI/Private/D3D11RHIPrivate.h
@@ -561,7 +561,6 @@ public:
 	virtual void RHIGetSupportedResolution(uint32& Width, uint32& Height) final override;
 	virtual void RHIVirtualTextureSetFirstMipInMemory(FRHITexture2D* Texture, uint32 FirstMip) final override;
 	virtual void RHIVirtualTextureSetFirstMipVisible(FRHITexture2D* Texture, uint32 FirstMip) final override;
-	virtual void RHIExecuteCommandList(FRHICommandList* CmdList) final override;
 	virtual void* RHIGetNativeDevice() final override;
 	virtual void* RHIGetNativeInstance() final override;
 	virtual void* RHIGetNativeCommandBuffer() final override;
@@ -660,10 +659,6 @@ public:
 	virtual FTexture2DRHIRef AsyncReallocateTexture2D_RenderThread(class FRHICommandListImmediate& RHICmdList, FRHITexture2D* Texture2D, int32 NewMipCount, int32 NewSizeX, int32 NewSizeY, FThreadSafeCounter* RequestStatus) final override;
 	virtual ETextureReallocationStatus FinalizeAsyncReallocateTexture2D_RenderThread(class FRHICommandListImmediate& RHICmdList, FRHITexture2D* Texture2D, bool bBlockUntilCompleted) final override;
 	virtual ETextureReallocationStatus CancelAsyncReallocateTexture2D_RenderThread(class FRHICommandListImmediate& RHICmdList, FRHITexture2D* Texture2D, bool bBlockUntilCompleted) final override;
-	virtual FVertexShaderRHIRef CreateVertexShader_RenderThread(class FRHICommandListImmediate& RHICmdList, TArrayView<const uint8> Code, const FSHAHash& Hash) final override;
-	virtual FPixelShaderRHIRef CreatePixelShader_RenderThread(class FRHICommandListImmediate& RHICmdList, TArrayView<const uint8> Code, const FSHAHash& Hash) final override;
-	virtual FGeometryShaderRHIRef CreateGeometryShader_RenderThread(class FRHICommandListImmediate& RHICmdList, TArrayView<const uint8> Code, const FSHAHash& Hash) final override;
-	virtual FComputeShaderRHIRef CreateComputeShader_RenderThread(class FRHICommandListImmediate& RHICmdList, TArrayView<const uint8> Code, const FSHAHash& Hash) final override;
 	virtual void* LockTexture2D_RenderThread(class FRHICommandListImmediate& RHICmdList, FRHITexture2D* Texture, uint32 MipIndex, EResourceLockMode LockMode, uint32& DestStride, bool bLockWithinMiptail, bool bNeedsDefaultRHIFlush = true) final override;
 	virtual void UnlockTexture2D_RenderThread(class FRHICommandListImmediate& RHICmdList, FRHITexture2D* Texture, uint32 MipIndex, bool bLockWithinMiptail, bool bNeedsDefaultRHIFlush = true) final override;
 	virtual void RHIEndUpdateTexture3D(FRHICommandListBase& RHICmdList, FUpdateTexture3DData& UpdateData) final override;
@@ -673,7 +668,6 @@ public:
 	virtual FShaderResourceViewRHIRef RHICreateShaderResourceView_RenderThread(class FRHICommandListImmediate& RHICmdList, const FShaderResourceViewInitializer& Initializer) final override;
 	virtual FShaderResourceViewRHIRef RHICreateShaderResourceView_RenderThread(class FRHICommandListImmediate& RHICmdList, FRHITexture* Texture, const FRHITextureSRVCreateInfo& CreateInfo) final override;
 	virtual FShaderResourceViewRHIRef RHICreateShaderResourceView_RenderThread(class FRHICommandListImmediate& RHICmdList, FRHIBuffer* Buffer) final override;
-	virtual FRenderQueryRHIRef RHICreateRenderQuery_RenderThread(class FRHICommandListImmediate& RHICmdList, ERenderQueryType QueryType) final override;
 	virtual void* RHILockTextureCubeFace_RenderThread(class FRHICommandListImmediate& RHICmdList, FRHITextureCube* Texture, uint32 FaceIndex, uint32 ArrayIndex, uint32 MipIndex, EResourceLockMode LockMode, uint32& DestStride, bool bLockWithinMiptail) final override;
 	virtual void RHIUnlockTextureCubeFace_RenderThread(class FRHICommandListImmediate& RHICmdList, FRHITextureCube* Texture, uint32 FaceIndex, uint32 ArrayIndex, uint32 MipIndex, bool bLockWithinMiptail) final override;
 	virtual FTextureRHIRef RHICreateTexture_RenderThread(class FRHICommandListImmediate& RHICmdList, const FRHITextureCreateDesc& CreateDesc) final override;
diff --git a/Engine/Source/Runtime/Windows/D3D11RHI/Private/D3D11Shaders.cpp b/Engine/Source/Runtime/Windows/D3D11RHI/Private/D3D11Shaders.cpp
--- a/Engine/Source/Runtime/Windows/D3D11RHI/Private/D3D11Shaders.cpp
+++ b/Engine/Source/Runtime/Windows/D3D11RHI/Private/D3D11Shaders.cpp
@@ -148,13 +148,6 @@ FVertexShaderRHIRef FD3D11DynamicRHI::RHICreateVertexShader(TArrayView<const uin
 	return Shader;
 }
 
-FVertexShaderRHIRef FD3D11DynamicRHI::CreateVertexShader_RenderThread(
-	class FRHICommandListImmediate& RHICmdList,
-	TArrayView<const uint8> Code, const FSHAHash& Hash)
-{
-	return RHICreateVertexShader(Code, Hash);
-}
-
 FGeometryShaderRHIRef FD3D11DynamicRHI::RHICreateGeometryShader(TArrayView<const uint8> Code, const FSHAHash& Hash)
 { 
 	FShaderCodeReader ShaderCode(Code);
@@ -182,13 +175,6 @@ FGeometryShaderRHIRef FD3D11DynamicRHI::RHICreateGeometryShader(TArrayView<const
 	return Shader;
 }
 
-FGeometryShaderRHIRef FD3D11DynamicRHI::CreateGeometryShader_RenderThread(
-	class FRHICommandListImmediate& RHICmdList,
-	TArrayView<const uint8> Code, const FSHAHash& Hash)
-{
-	return RHICreateGeometryShader(Code, Hash);
-}
-
 FPixelShaderRHIRef FD3D11DynamicRHI::RHICreatePixelShader(TArrayView<const uint8> Code, const FSHAHash& Hash)
 {
 	FShaderCodeReader ShaderCode(Code);
@@ -216,13 +202,6 @@ FPixelShaderRHIRef FD3D11DynamicRHI::RHICreatePixelShader(TArrayView<const uint8
 	return Shader;
 }
 
-FPixelShaderRHIRef FD3D11DynamicRHI::CreatePixelShader_RenderThread(
-	class FRHICommandListImmediate& RHICmdList,
-	TArrayView<const uint8> Code, const FSHAHash& Hash)
-{
-	return RHICreatePixelShader(Code, Hash);
-}
-
 FComputeShaderRHIRef FD3D11DynamicRHI::RHICreateComputeShader(TArrayView<const uint8> Code, const FSHAHash& Hash)
 { 
 	FShaderCodeReader ShaderCode(Code);
@@ -250,13 +229,6 @@ FComputeShaderRHIRef FD3D11DynamicRHI::RHICreateComputeShader(TArrayView<const u
 	return Shader;
 }
 
-FComputeShaderRHIRef FD3D11DynamicRHI::CreateComputeShader_RenderThread(
-	class FRHICommandListImmediate& RHICmdList,
-	TArrayView<const uint8> Code, const FSHAHash& Hash)
-{
-	return RHICreateComputeShader(Code, Hash);
-}
-
 void FD3D11DynamicRHI::RHISetMultipleViewports(uint32 Count, const FViewportBounds* Data) 
 { 
 	check(Count > 0);
