fix race condition on priority boost where we could miss a request if the task was not issued yet

#rnx

diff --git a/Engine/Source/Runtime/RHI/Private/PipelineStateCache.cpp b/Engine/Source/Runtime/RHI/Private/PipelineStateCache.cpp
--- a/Engine/Source/Runtime/RHI/Private/PipelineStateCache.cpp
+++ b/Engine/Source/Runtime/RHI/Private/PipelineStateCache.cpp
@@ -1000,6 +1000,18 @@ private:
 
 };
 
+// Request state
+enum class EPSOPrecacheStateMask : uint8
+{
+	None = 0,
+	Compiling = 1 << 0, // once the start is scheduled
+	Succeeded = 1 << 1, // once compilation is finished
+	Failed = 1 << 2,    // once compilation is finished
+	Boosted = 1 << 3,
+};
+
+ENUM_CLASS_FLAGS(EPSOPrecacheStateMask)
+
 template<class TPrecachePipelineCacheDerived, class TPrecachedPSOInitializer, class TPipelineState>
 class TPrecachePipelineCacheBase
 {
@@ -1014,6 +1026,18 @@ public:
 	}
 	
 protected:
+
+	void RescheduleTaskToHighPriority(TPipelineState* PipelineState)
+	{
+		if (FPSOPrecacheThreadPool::UsePool())
+		{
+			check(PipelineState->PrecompileTask);
+			PipelineState->PrecompileTask->Reschedule(&GPSOPrecacheThreadPool.Get(), EQueuedWorkPriority::Highest);
+		}
+
+		UpdateHighPriorityCompileCount(true /*Increment*/);
+	}
+
 	FPSOPrecacheRequestResult TryAddNewState(const TPrecachedPSOInitializer& Initializer, bool bDoAsyncCompile)
 	{
 		FPSOPrecacheRequestResult Result;
@@ -1047,7 +1071,6 @@ protected:
 			// Add data to map
 			FPrecacheTask PrecacheTask;
 			PrecacheTask.PipelineState = NewPipelineState;
-			PrecacheTask.State = EState::Compiling;
 			PrecacheTask.RequestID = Result.RequestID;
 			PrecachedPSOInitializerData.AddByHash(InitializerHash, Initializer, PrecacheTask);
 
@@ -1065,6 +1088,22 @@ protected:
 
 		TPrecachePipelineCacheDerived::OnNewPipelineStateCreated(Initializer, NewPipelineState, bDoAsyncCompile);
 
+		// A boost request might have been issued while we were kicking the task, need to check it here
+		{
+			FRWScopeLock ReadLock(PrecachePSOsRWLock, SLT_ReadOnly);
+			FPrecacheTask* FindResult = PrecachedPSOInitializerData.FindByHash(InitializerHash, Initializer);
+			check(FindResult != nullptr);
+			if (FindResult != nullptr)
+			{
+				EPSOPrecacheStateMask PreviousStateMask = FindResult->AddPSOPrecacheState(EPSOPrecacheStateMask::Compiling);
+				// by the time we're here, PrecacheFinished might already have been called, so boost it only if we know we will call it
+				if ( PreviousStateMask == EPSOPrecacheStateMask::Boosted)
+				{
+					RescheduleTaskToHighPriority(FindResult->PipelineState);
+				}
+			}
+		}
+
 		// Make sure that we don't access NewPipelineState here as the task might have already been finished, ProcessDelayedCleanup may have been called
 		// and NewPipelineState already been deleted
 		
@@ -1079,10 +1118,10 @@ public:
 		for (auto Iterator = PrecachedPSOInitializerData.CreateIterator(); Iterator; ++Iterator)
 		{
 			FPrecacheTask& PrecacheTask = Iterator->Value;
-			if (PrecacheTask.PipelineState && PrecacheTask.State == EState::Compiling)
+			if (PrecacheTask.PipelineState && !PrecacheTask.PipelineState->IsComplete())
 			{
 				PrecacheTask.PipelineState->WaitCompletion();
-				check(PrecacheTask.State != EState::Compiling);
+				check(EnumHasAnyFlags(PrecacheTask.ReadPSOPrecacheState(), (EPSOPrecacheStateMask::Succeeded | EPSOPrecacheStateMask::Failed)));
 				delete PrecacheTask.PipelineState;
 				PrecacheTask.PipelineState = nullptr;
 			}
@@ -1119,33 +1158,28 @@ public:
 
 		// Won't modify anything in this cache so readlock should be enough?
 		FRWScopeLock ReadLock(PrecachePSOsRWLock, SLT_ReadOnly);
-		TPrecachedPSOInitializer Initializer = PrecachedPSOInitializers[RequestID.RequestID];
+		const TPrecachedPSOInitializer& Initializer = PrecachedPSOInitializers[RequestID.RequestID];
 		FPrecacheTask* FindResult = PrecachedPSOInitializerData.Find(Initializer);
 		check(FindResult);
-		if (FindResult->State == EState::Compiling && FPlatformAtomics::InterlockedExchange(&FindResult->HighPriority, 1) == 0)
+		EPSOPrecacheStateMask PreviousStateMask = FindResult->AddPSOPrecacheState(EPSOPrecacheStateMask::Boosted);
+		// It's possible to get a boost request while the task has not been started yet. In this case, TryAddNewState will take care of it
+		// if TryAddNewState is done, then we can proceed to boost it, if the task is not done yet
+		if (PreviousStateMask == EPSOPrecacheStateMask::Compiling)
 		{
-			if (FindResult->PipelineState->PrecompileTask)
-			{
-				// if this is an in-progress threadpool precompile task then it could be seconds away in the queue.
-				// Reissue this task so that it jumps the precompile queue.
-				FindResult->PipelineState->PrecompileTask->Reschedule(&GPSOPrecacheThreadPool.Get(), EQueuedWorkPriority::Highest);
-
-				UpdateHighPriorityCompileCount(true /*Increment*/);
-			}
+			RescheduleTaskToHighPriority(FindResult->PipelineState);
 		}
 	}
 
 	uint32 NumActivePrecacheRequests()
 	{
-		FRWScopeLock ReadLock(PrecachePSOsRWLock, SLT_ReadOnly);
 		if (GPSOWaitForHighPriorityRequestsOnly)
 		{
-			return HighPriorityCompileCount;
+			return FPlatformAtomics::AtomicRead(&HighPriorityCompileCount);
 		}
 		else
 		{
-			return ActiveCompileCount;
-		}		
+			return FPlatformAtomics::AtomicRead(&ActiveCompileCount);
+		}
 	}
 
 	void PrecacheFinished(const TPrecachedPSOInitializer& Initializer, bool bValid)
@@ -1156,21 +1190,33 @@ public:
 		
 		// Mark compiled (either succeeded or failed)
 		FPrecacheTask* FindResult = PrecachedPSOInitializerData.FindByHash(InitializerHash, Initializer);
-		check(FindResult && FindResult->State == EState::Compiling);
-		FindResult->State = bValid ? EState::Succeeded : EState::Failed;
+		check(FindResult);
+		// We still add the 'compiling' bit here because if the task is fast enough, we can get here before the end of TryAddNewState
+		const EPSOPrecacheStateMask CompleteStateMask = bValid ? (EPSOPrecacheStateMask::Succeeded | EPSOPrecacheStateMask::Compiling) : (EPSOPrecacheStateMask::Failed | EPSOPrecacheStateMask::Compiling);
+		const EPSOPrecacheStateMask PreviousStateMask = FindResult->AddPSOPrecacheState(CompleteStateMask);
 
 		// Add to array of precached PSOs so it can be cleaned up
 		PrecachedPSOs.Add(Initializer);
 
-		if (FindResult->HighPriority)
+        // Need to ensure that the boost request was actually executed: if only it was asked by BoostPriority, but not requested (ie TryAddNewState has not set the Compiling bit
+        // yet) then we must ignore the request
+		if (EnumHasAllFlags(PreviousStateMask, EPSOPrecacheStateMask::Boosted | EPSOPrecacheStateMask::Compiling))
 		{
 			UpdateHighPriorityCompileCount(false /*Increment*/);
 		}
 		UpdateActiveCompileCount(false /*Increment*/);
 	}
 
+	static bool IsCompilationDone(EPSOPrecacheStateMask StateMask)
+	{
+		return EnumHasAnyFlags(StateMask, EPSOPrecacheStateMask::Succeeded | EPSOPrecacheStateMask::Failed);
+	}
+
 	void ProcessDelayedCleanup()
 	{
+		SET_DWORD_STAT_FName(TPrecachePipelineCacheDerived::GetActiveCompileStatName(), ActiveCompileCount);
+		SET_DWORD_STAT_FName(TPrecachePipelineCacheDerived::GetHighPriorityCompileStatName(), HighPriorityCompileCount);
+
 		FRWScopeLock WriteLock(PrecachePSOsRWLock, SLT_Write);
 		for (int32 Index = 0; Index < PrecachedPSOs.Num(); ++Index)		
 		{
@@ -1178,7 +1224,7 @@ public:
 			uint32 InitializerHash = TPrecachePipelineCacheDerived::PipelineStateInitializerHash(Initializer);
 
 			FPrecacheTask* FindResult = PrecachedPSOInitializerData.FindByHash(InitializerHash, Initializer);
-			check(FindResult && FindResult->State != EState::Compiling);
+			check(FindResult && IsCompilationDone((FindResult->ReadPSOPrecacheState())));
 			if (FindResult->PipelineState->IsComplete())
 			{
 				// This is needed to cleanup the members - bit strange because it's complete already
@@ -1200,7 +1246,7 @@ protected:
 		if (FindResult)
 		{			
 			// If not compiled yet, then return the request ID so the caller can check the state
-			if (FindResult->State == EState::Compiling)
+			if (!IsCompilationDone(FindResult->ReadPSOPrecacheState()))
 			{
 				Result.RequestID = FindResult->RequestID;
 				Result.AsyncCompileEvent = FindResult->PipelineState->CompletionEvent;
@@ -1221,25 +1267,28 @@ protected:
 		{
 			return EPSOPrecacheResult::Missed;
 		}
-		else if (FindResult->State == EState::Compiling)
+
+		EPSOPrecacheStateMask CompilationState = FindResult->ReadPSOPrecacheState();
+		if (!IsCompilationDone(CompilationState))
 		{
 			return EPSOPrecacheResult::Active;
 		}
 
-		check(FindResult->State == EState::Succeeded || FindResult->State == EState::Failed);
-		return (FindResult->State == EState::Failed) ? EPSOPrecacheResult::NotSupported : EPSOPrecacheResult::Complete;
+		// check we only set 1 completion bit
+		const EPSOPrecacheStateMask CompletionMask = EPSOPrecacheStateMask::Succeeded | EPSOPrecacheStateMask::Failed;
+		check(EnumHasAnyFlags(CompilationState, CompletionMask) && !EnumHasAllFlags(CompilationState, CompletionMask));
+
+		return (EnumHasAnyFlags(CompilationState, EPSOPrecacheStateMask::Failed)) ? EPSOPrecacheResult::NotSupported : EPSOPrecacheResult::Complete;
 	}
 
 	void UpdateActiveCompileCount(bool bIncrement)
 	{
 		if (bIncrement)
 		{
-			INC_DWORD_STAT_FName(TPrecachePipelineCacheDerived::GetActiveCompileStatName());
 			FPlatformAtomics::InterlockedIncrement(&ActiveCompileCount);
 		}
 		else
 		{
-			DEC_DWORD_STAT_FName(TPrecachePipelineCacheDerived::GetActiveCompileStatName());
 			FPlatformAtomics::InterlockedDecrement(&ActiveCompileCount);
 		}
 	}
@@ -1248,12 +1297,10 @@ protected:
 	{
 		if (bIncrement)
 		{
-			INC_DWORD_STAT_FName(TPrecachePipelineCacheDerived::GetHighPriorityCompileStatName());
 			FPlatformAtomics::InterlockedIncrement(&HighPriorityCompileCount);
 		}
 		else
 		{
-			DEC_DWORD_STAT_FName(TPrecachePipelineCacheDerived::GetHighPriorityCompileStatName());
 			FPlatformAtomics::InterlockedDecrement(&HighPriorityCompileCount);
 		}
 	}
@@ -1267,22 +1314,26 @@ protected:
 	// Array containing all the precached PSO initializers thus far - the index in this array is used to uniquely identify the PSO requests
 	TArray<TPrecachedPSOInitializer> PrecachedPSOInitializers;
 
-	// Request state
-	enum class EState : uint8
-	{
-		Unknown,
-		Compiling,
-		Succeeded,
-		Failed,
-	};
-
 	// Hash map used for fast retrieval of already precached PSOs
 	struct FPrecacheTask
 	{
-		FPSOPrecacheRequestID RequestID;
-		EState State = EState::Unknown;
-		volatile int8 HighPriority = 0;
 		TPipelineState* PipelineState = nullptr;
+		FPSOPrecacheRequestID RequestID;
+
+		EPSOPrecacheStateMask AddPSOPrecacheState(EPSOPrecacheStateMask DesiredState)
+		{
+			static_assert(sizeof(EPSOPrecacheStateMask) == sizeof(int8), "Fix the cast below");
+			return (EPSOPrecacheStateMask)FPlatformAtomics::InterlockedOr((volatile int8*)&StateMask, (int8)DesiredState);
+		}
+
+		inline EPSOPrecacheStateMask ReadPSOPrecacheState()
+		{
+			static_assert(sizeof(EPSOPrecacheStateMask) == sizeof(int8), "Fix the cast below");
+			return (EPSOPrecacheStateMask)FPlatformAtomics::AtomicRead((volatile int8*)&StateMask);
+		}
+
+	private:
+		volatile EPSOPrecacheStateMask StateMask = EPSOPrecacheStateMask::None;
 	};
 
 	using TMapKeyFuncsBaseClass = TDefaultMapKeyFuncs<TPrecachedPSOInitializer, FPrecacheTask, false>;
@@ -1300,10 +1351,10 @@ protected:
 	TMap<TPrecachedPSOInitializer, FPrecacheTask, FDefaultSetAllocator, TMapKeyFuncs> PrecachedPSOInitializerData;
 
 	// Number of open active compiles
-	int32 ActiveCompileCount = 0;
+	volatile int32 ActiveCompileCount = 0;
 
 	// Number of open high priority compiles
-	int32 HighPriorityCompileCount = 0;
+	volatile int32 HighPriorityCompileCount = 0;
 
 	// Finished Precached PSOs which can be garbage collected
 	TArray<TPrecachedPSOInitializer> PrecachedPSOs;
