RHIResources.h dependency cleanup.
Moved RHIShaderLibrary to its own file to remove the dependency on CoreDelegates.h

diff --git a/Engine/Source/Runtime/Core/Public/Experimental/ConcurrentLinearAllocator.h b/Engine/Source/Runtime/Core/Public/Experimental/ConcurrentLinearAllocator.h
--- a/Engine/Source/Runtime/Core/Public/Experimental/ConcurrentLinearAllocator.h
+++ b/Engine/Source/Runtime/Core/Public/Experimental/ConcurrentLinearAllocator.h
@@ -2,7 +2,9 @@
 
 #pragma once
 #include <atomic>
+#if UE_ENABLE_INCLUDE_ORDER_DEPRECATED_IN_5_2
 #include "CoreMinimal.h"
+#endif
 #include "HAL/MallocAnsi.h"
 #include "HAL/UnrealMemory.h"
 #include "HAL/LowLevelMemTracker.h"
diff --git a/Engine/Source/Runtime/RHI/Private/RHIResources.cpp b/Engine/Source/Runtime/RHI/Private/RHIResources.cpp
--- a/Engine/Source/Runtime/RHI/Private/RHIResources.cpp
+++ b/Engine/Source/Runtime/RHI/Private/RHIResources.cpp
@@ -6,11 +6,52 @@
 #include "Experimental/Containers/HazardPointer.h"
 #include "Misc/MemStack.h"
 #include "Stats/Stats.h"
+#include "TextureProfiler.h"
+
+static UE::TConsumeAllMpmcQueue<FRHIResource*> PendingDeletes;
+static UE::TConsumeAllMpmcQueue<FRHIResource*> PendingDeletesWithLifetimeExtension;
 
-UE::TConsumeAllMpmcQueue<FRHIResource*> FRHIResource::PendingDeletes;
-UE::TConsumeAllMpmcQueue<FRHIResource*> FRHIResource::PendingDeletesWithLifetimeExtension;
 FRHIResource* FRHIResource::CurrentlyDeleting = nullptr;
 
+FRHIResource::FRHIResource(ERHIResourceType InResourceType)
+	: ResourceType(InResourceType)
+	, bCommitted(true)
+	, bAllowExtendLifetime(true)
+#if RHI_ENABLE_RESOURCE_INFO
+	, bBeingTracked(false)
+#endif
+{
+#if RHI_ENABLE_RESOURCE_INFO
+	BeginTrackingResource(this);
+#endif
+}
+
+FRHIResource::~FRHIResource()
+{
+	check(IsEngineExitRequested() || CurrentlyDeleting == this);
+	check(AtomicFlags.GetNumRefs(std::memory_order_relaxed) == 0); // this should not have any outstanding refs
+	CurrentlyDeleting = nullptr;
+
+#if RHI_ENABLE_RESOURCE_INFO
+	EndTrackingResource(this);
+#endif
+}
+
+void FRHIResource::Destroy() const
+{
+	if (!AtomicFlags.MarkForDelete(std::memory_order_release))
+	{
+		if (bAllowExtendLifetime)
+		{
+			PendingDeletesWithLifetimeExtension.ProduceItem(const_cast<FRHIResource*>(this));
+		}
+		else
+		{
+			PendingDeletes.ProduceItem(const_cast<FRHIResource*>(this));
+		}
+	}
+}
+
 bool FRHIResource::Bypass()
 {
 	return GRHICommandList.Bypass();
@@ -20,3 +61,22 @@ int32 FRHIResource::FlushPendingDeletes(FRHICommandListImmediate& RHICmdList)
 {
 	return RHICmdList.FlushPendingDeletes();
 }
+
+FRHITexture::FRHITexture(const FRHITextureCreateDesc& InDesc)
+	: FRHIViewableResource(RRT_Texture, InDesc.InitialState)
+#if ENABLE_RHI_VALIDATION
+	, RHIValidation::FTextureResource(InDesc)
+#endif
+	, TextureDesc(InDesc)
+{
+	SetName(InDesc.DebugName);
+}
+
+void FRHITexture::SetName(const FName& InName)
+{
+	Name = InName;
+
+#if TEXTURE_PROFILER_ENABLED
+	FTextureProfiler::Get()->UpdateTextureName(this);
+#endif
+}
diff --git a/Engine/Source/Runtime/RHI/Private/TextureProfiler.cpp b/Engine/Source/Runtime/RHI/Private/TextureProfiler.cpp
--- a/Engine/Source/Runtime/RHI/Private/TextureProfiler.cpp
+++ b/Engine/Source/Runtime/RHI/Private/TextureProfiler.cpp
@@ -3,6 +3,7 @@
 #include "TextureProfiler.h"
 
 #include "Misc/CoreDelegates.h"
+#include "HAL/IConsoleManager.h"
 #include "ProfilingDebugging/CsvProfiler.h"
 #include "RHIDefinitions.h"
 #include "RHIResources.h"
diff --git a/Engine/Source/Runtime/RHI/Public/DynamicRHI.h b/Engine/Source/Runtime/RHI/Public/DynamicRHI.h
--- a/Engine/Source/Runtime/RHI/Public/DynamicRHI.h
+++ b/Engine/Source/Runtime/RHI/Public/DynamicRHI.h
@@ -9,6 +9,7 @@ DynamicRHI.h: Dynamically bound Render Hardware Interface definitions.
 #include "CoreTypes.h"
 #include "RHIContext.h"
 #include "RHIFeatureLevel.h"
+#include "RHIShaderLibrary.h"
 #include "MultiGPU.h"
 #include "Serialization/MemoryLayout.h"
 #include "Containers/ArrayView.h"
diff --git a/Engine/Source/Runtime/RHI/Public/RHIResources.h b/Engine/Source/Runtime/RHI/Public/RHIResources.h
--- a/Engine/Source/Runtime/RHI/Public/RHIResources.h
+++ b/Engine/Source/Runtime/RHI/Public/RHIResources.h
@@ -14,22 +14,31 @@
 #include "HAL/ThreadSafeCounter.h"
 #include "Templates/RefCounting.h"
 #include "PixelFormat.h"
+#include "Async/TaskGraphFwd.h"
+#include "Serialization/MemoryImage.h"
+#include "RHIFwd.h"
+#include "RHIImmutableSamplerState.h"
+#include "RHITransition.h"
+#include "MultiGPU.h"
+#include "Math/IntPoint.h"
+#include "Math/IntRect.h"
+#include "Math/IntVector.h"
+
+#if UE_ENABLE_INCLUDE_ORDER_DEPRECATED_IN_5_2
 #include "TextureProfiler.h"
 #include "Containers/LockFreeList.h"
+#include "Misc/CoreDelegates.h"
 #include "Misc/SecureHash.h"
 #include "Hash/CityHash.h"
 #include "Async/TaskGraphInterfaces.h"
-#include "Serialization/MemoryImage.h"
 #include "Experimental/Containers/HazardPointer.h"
+#include "Containers/ClosableMpscQueue.h"
 #include "Containers/ConsumeAllMpmcQueue.h"
-#include "Misc/CoreDelegates.h"
-#include "RHIFwd.h"
-#include "RHIImmutableSamplerState.h"
-#include "RHITransition.h"
-#include "MultiGPU.h"
+#endif
 
 #include <atomic>
 
+class FHazardPointerCollection;
 class FRHIComputeCommandList;
 class FRHICommandListImmediate;
 class FRHITextureReference;
@@ -38,33 +47,15 @@ struct FRHIResourceInfo;
 struct FGenerateMipsStruct;
 enum class EClearBinding;
 
+typedef TArray<FGraphEventRef, TInlineAllocator<4> > FGraphEventArray;
+
+
 /** The base type of RHI resources. */
-class RHI_API FRHIResource
+class FRHIResource
 {
 public:
-	FRHIResource(ERHIResourceType InResourceType)
-		: ResourceType(InResourceType)
-		, bCommitted(true)
-		, bAllowExtendLifetime(true)
-#if RHI_ENABLE_RESOURCE_INFO
-		, bBeingTracked(false)
-#endif
-	{
-#if RHI_ENABLE_RESOURCE_INFO
-		BeginTrackingResource(this);
-#endif
-	}
-
-	virtual ~FRHIResource()
-	{
-		check(IsEngineExitRequested() || CurrentlyDeleting == this);
-		check(AtomicFlags.GetNumRefs(std::memory_order_relaxed) == 0); // this should not have any outstanding refs
-		CurrentlyDeleting = nullptr;
-
-#if RHI_ENABLE_RESOURCE_INFO
-		EndTrackingResource(this);
-#endif
-	}
+	RHI_API FRHIResource(ERHIResourceType InResourceType);
+	RHI_API virtual ~FRHIResource();
 
 	FORCEINLINE_DEBUGGABLE uint32 AddRef() const
 	{
@@ -75,20 +66,8 @@ public:
 
 private:
 	// Separate function to avoid force inlining this everywhere. Helps both for code size and performance.
-	inline void Destroy() const
-	{
-		if (!AtomicFlags.MarkForDelete(std::memory_order_release))
-		{
-			if (bAllowExtendLifetime)
-			{
-				PendingDeletesWithLifetimeExtension.ProduceItem(const_cast<FRHIResource*>(this));
-			}
-			else
-			{
-				PendingDeletes.ProduceItem(const_cast<FRHIResource*>(this));
-			}
-		}
-	}
+	RHI_API void Destroy() const;
+				if (PendingDeletesPtr->Enqueue(const_cast<FRHIResource*>(this)))
 
 public:
 	FORCEINLINE_DEBUGGABLE uint32 Release() const
@@ -111,9 +90,9 @@ public:
 		return uint32(CurrentValue);
 	}
 
-	static int32 FlushPendingDeletes(FRHICommandListImmediate& RHICmdList);
+	RHI_API static int32 FlushPendingDeletes(FRHICommandListImmediate& RHICmdList);
 
-	static bool Bypass();
+	RHI_API static bool Bypass();
 
 	bool IsValid() const
 	{
@@ -141,7 +120,7 @@ public:
 #if RHI_ENABLE_RESOURCE_INFO
 	// Get resource info if available.
 	// Should return true if the ResourceInfo was filled with data.
-	virtual bool GetResourceInfo(FRHIResourceInfo& OutResourceInfo) const;
+	RHI_API virtual bool GetResourceInfo(FRHIResourceInfo& OutResourceInfo) const;
 
 	static void BeginTrackingResource(FRHIResource* InResource);
 	static void EndTrackingResource(FRHIResource* InResource);
@@ -240,8 +219,6 @@ private:
 	FName OwnerName;
 #endif
 
-	static UE::TConsumeAllMpmcQueue<FRHIResource*> PendingDeletes;
-	static UE::TConsumeAllMpmcQueue<FRHIResource*> PendingDeletesWithLifetimeExtension;
 	static FRHIResource* CurrentlyDeleting;
 
 	friend FRHICommandListImmediate;
@@ -1902,15 +1879,7 @@ class RHI_API FRHITexture : public FRHIViewableResource
 {
 protected:
 	/** Initialization constructor. Should only be called by platform RHI implementations. */
-	FRHITexture(const FRHITextureCreateDesc& InDesc)
-		: FRHIViewableResource(RRT_Texture, InDesc.InitialState)
-#if ENABLE_RHI_VALIDATION
-		, RHIValidation::FTextureResource(InDesc)
-#endif
-		, TextureDesc(InDesc)
-	{
-		SetName(InDesc.DebugName);
-	}
+	FRHITexture(const FRHITextureCreateDesc& InDesc);
 
 public:
 	/**
@@ -2059,14 +2028,7 @@ public:
 		return LastRenderTime.GetLastRenderTime();
 	}
 
-	void SetName(const FName& InName)
-	{
-		Name = InName;
-
-#if TEXTURE_PROFILER_ENABLED
-		FTextureProfiler::Get()->UpdateTextureName(this);
-#endif
-	}
+	void SetName(const FName& InName);
 
 	///
 	/// Deprecated functions
@@ -3672,54 +3634,6 @@ protected:
 	TRefCountPtr<FRHIComputeShader> ComputeShader;
 };
 
-//
-// Shader Library
-//
-
-class FRHIShaderLibrary : public FRHIResource
-{
-public:
-	FRHIShaderLibrary(EShaderPlatform InPlatform, FString const& InName) : FRHIResource(RRT_ShaderLibrary), Platform(InPlatform), LibraryName(InName), LibraryId(GetTypeHash(InName)) {}
-	virtual ~FRHIShaderLibrary() {}
-	
-	FORCEINLINE EShaderPlatform GetPlatform(void) const { return Platform; }
-	FORCEINLINE const FString& GetName(void) const { return LibraryName; }
-	FORCEINLINE uint32 GetId(void) const { return LibraryId; }
-	
-	virtual bool IsNativeLibrary() const = 0;
-	virtual int32 GetNumShaderMaps() const = 0;
-	virtual int32 GetNumShaders() const = 0;
-	virtual int32 GetNumShadersForShaderMap(int32 ShaderMapIndex) const = 0;
-	virtual int32 GetShaderIndex(int32 ShaderMapIndex, int32 i) const = 0;
-	virtual FSHAHash GetShaderHash(int32 ShaderMapIndex, int32 ShaderIndex) = 0;
-	virtual int32 FindShaderMapIndex(const FSHAHash& Hash) = 0;
-	virtual int32 FindShaderIndex(const FSHAHash& Hash) = 0;
-	virtual bool PreloadShader(int32 ShaderIndex, FGraphEventArray& OutCompletionEvents) { return false; }
-	virtual bool PreloadShaderMap(int32 ShaderMapIndex, FGraphEventArray& OutCompletionEvents) { return false; }
-	virtual bool PreloadShaderMap(int32 ShaderMapIndex, FCoreDelegates::FAttachShaderReadRequestFunc AttachShaderReadRequestFunc) { return false; }
-	virtual void ReleasePreloadedShader(int32 ShaderIndex) {}
-
-	virtual TRefCountPtr<FRHIShader> CreateShader(int32 ShaderIndex) { return nullptr; }
-	virtual void Teardown() {};
-
-protected:
-	EShaderPlatform Platform;
-	FString LibraryName;
-	uint32 LibraryId;
-};
-
-class FRHIPipelineBinaryLibrary : public FRHIResource
-{
-public:
-	FRHIPipelineBinaryLibrary(EShaderPlatform InPlatform, FString const& FilePath) : FRHIResource(RRT_PipelineBinaryLibrary), Platform(InPlatform) {}
-	virtual ~FRHIPipelineBinaryLibrary() {}
-	
-	FORCEINLINE EShaderPlatform GetPlatform(void) const { return Platform; }
-	
-protected:
-	EShaderPlatform Platform;
-};
-
 enum class ERenderTargetActions : uint8
 {
 	LoadOpMask = 2,
@@ -4541,5 +4455,6 @@ struct FRHIShaderParameterResource
 };
 
 #if UE_ENABLE_INCLUDE_ORDER_DEPRECATED_IN_5_2
+#include "RHIShaderLibrary.h"
 #include "RHITextureReference.h"
 #endif
diff --git a/Engine/Source/Runtime/RHI/Public/RHIShaderLibrary.h b/Engine/Source/Runtime/RHI/Public/RHIShaderLibrary.h
new file mode 100644
--- /dev/null
+++ b/Engine/Source/Runtime/RHI/Public/RHIShaderLibrary.h
@@ -0,0 +1,54 @@
+// Copyright Epic Games, Inc. All Rights Reserved.
+
+#pragma once
+
+#include "RHIResources.h"
+#include "Misc/CoreDelegates.h"
+
+//
+// Shader Library
+//
+
+class FRHIShaderLibrary : public FRHIResource
+{
+public:
+	FRHIShaderLibrary(EShaderPlatform InPlatform, FString const& InName) : FRHIResource(RRT_ShaderLibrary), Platform(InPlatform), LibraryName(InName), LibraryId(GetTypeHash(InName)) {}
+	virtual ~FRHIShaderLibrary() {}
+
+	FORCEINLINE EShaderPlatform GetPlatform(void) const { return Platform; }
+	FORCEINLINE const FString& GetName(void) const { return LibraryName; }
+	FORCEINLINE uint32 GetId(void) const { return LibraryId; }
+
+	virtual bool IsNativeLibrary() const = 0;
+	virtual int32 GetNumShaderMaps() const = 0;
+	virtual int32 GetNumShaders() const = 0;
+	virtual int32 GetNumShadersForShaderMap(int32 ShaderMapIndex) const = 0;
+	virtual int32 GetShaderIndex(int32 ShaderMapIndex, int32 i) const = 0;
+	virtual FSHAHash GetShaderHash(int32 ShaderMapIndex, int32 ShaderIndex) = 0;
+	virtual int32 FindShaderMapIndex(const FSHAHash& Hash) = 0;
+	virtual int32 FindShaderIndex(const FSHAHash& Hash) = 0;
+	virtual bool PreloadShader(int32 ShaderIndex, FGraphEventArray& OutCompletionEvents) { return false; }
+	virtual bool PreloadShaderMap(int32 ShaderMapIndex, FGraphEventArray& OutCompletionEvents) { return false; }
+	virtual bool PreloadShaderMap(int32 ShaderMapIndex, FCoreDelegates::FAttachShaderReadRequestFunc AttachShaderReadRequestFunc) { return false; }
+	virtual void ReleasePreloadedShader(int32 ShaderIndex) {}
+
+	virtual TRefCountPtr<FRHIShader> CreateShader(int32 ShaderIndex) { return nullptr; }
+	virtual void Teardown() {};
+
+protected:
+	EShaderPlatform Platform;
+	FString LibraryName;
+	uint32 LibraryId;
+};
+
+class FRHIPipelineBinaryLibrary : public FRHIResource
+{
+public:
+	FRHIPipelineBinaryLibrary(EShaderPlatform InPlatform, FString const& FilePath) : FRHIResource(RRT_PipelineBinaryLibrary), Platform(InPlatform) {}
+	virtual ~FRHIPipelineBinaryLibrary() {}
+
+	FORCEINLINE EShaderPlatform GetPlatform(void) const { return Platform; }
+
+protected:
+	EShaderPlatform Platform;
+};
diff --git a/Engine/Source/Runtime/RHI/Public/TextureProfiler.h b/Engine/Source/Runtime/RHI/Public/TextureProfiler.h
--- a/Engine/Source/Runtime/RHI/Public/TextureProfiler.h
+++ b/Engine/Source/Runtime/RHI/Public/TextureProfiler.h
@@ -8,7 +8,7 @@
 #include "Containers/Map.h"
 #include "UObject/NameTypes.h"
 
-#define TEXTURE_PROFILER_ENABLED WITH_ENGINE && !UE_BUILD_SHIPPING
+#define TEXTURE_PROFILER_ENABLED (WITH_ENGINE && !UE_BUILD_SHIPPING)
 
 #if TEXTURE_PROFILER_ENABLED
 
