Fixing bad merge.

diff --git a/Engine/Source/Runtime/RHI/Public/RHIResources.h b/Engine/Source/Runtime/RHI/Public/RHIResources.h
--- a/Engine/Source/Runtime/RHI/Public/RHIResources.h
+++ b/Engine/Source/Runtime/RHI/Public/RHIResources.h
@@ -67,7 +67,6 @@ public:
 private:
 	// Separate function to avoid force inlining this everywhere. Helps both for code size and performance.
 	RHI_API void Destroy() const;
-				if (PendingDeletesPtr->Enqueue(const_cast<FRHIResource*>(this)))
 
 public:
 	FORCEINLINE_DEBUGGABLE uint32 Release() const
