Fixing compile errors in RHI.

diff --git a/Engine/Source/Runtime/RHI/Private/RHICommandList.cpp b/Engine/Source/Runtime/RHI/Private/RHICommandList.cpp
--- a/Engine/Source/Runtime/RHI/Private/RHICommandList.cpp
+++ b/Engine/Source/Runtime/RHI/Private/RHICommandList.cpp
@@ -3,6 +3,7 @@
 #include "RHICommandList.h"
 #include "Misc/App.h"
 #include "Async/TaskGraphInterfaces.h"
+#include "Containers/ConsumeAllMpmcQueue.h"
 #include "RHI.h"
 #include "ProfilingDebugging/CsvProfiler.h"
 #include "GenericPlatform/GenericPlatformCrashContext.h"
@@ -1133,6 +1134,9 @@ void FRHICommandListExecutor::WaitOnRHIThreadFence(FGraphEventRef& Fence)
 
 DECLARE_CYCLE_STAT(TEXT("Delete Resources"), STAT_DeleteResources, STATGROUP_RHICMDLIST);
 
+extern UE::TConsumeAllMpmcQueue<FRHIResource*> PendingDeletes;
+extern UE::TConsumeAllMpmcQueue<FRHIResource*> PendingDeletesWithLifetimeExtension;
+
 int32 FRHICommandListImmediate::FlushPendingDeletes()
 {
 	SCOPE_CYCLE_COUNTER(STAT_DeleteResources);
@@ -1144,12 +1148,12 @@ int32 FRHICommandListImmediate::FlushPendingDeletes()
 		? PersistentState.ExtendedLifetimeResources 
 		: DeletedResources;
 
-	FRHIResource::PendingDeletesWithLifetimeExtension.ConsumeAllLifo([&DeletedResourcesWithLifetimeExtension](FRHIResource* Resource)
+	PendingDeletesWithLifetimeExtension.ConsumeAllLifo([&DeletedResourcesWithLifetimeExtension](FRHIResource* Resource)
 	{
 		DeletedResourcesWithLifetimeExtension.Emplace(Resource);
 	});
 
-	FRHIResource::PendingDeletes.ConsumeAllLifo([&DeletedResources](FRHIResource* Resource)
+	PendingDeletes.ConsumeAllLifo([&DeletedResources](FRHIResource* Resource)
 	{
 		DeletedResources.Push(Resource);
 	});
diff --git a/Engine/Source/Runtime/RHI/Private/RHIResources.cpp b/Engine/Source/Runtime/RHI/Private/RHIResources.cpp
--- a/Engine/Source/Runtime/RHI/Private/RHIResources.cpp
+++ b/Engine/Source/Runtime/RHI/Private/RHIResources.cpp
@@ -1,15 +1,15 @@
 // Copyright Epic Games, Inc. All Rights Reserved.
 
 #include "RHIResources.h"
-#include "Containers/ClosableMpscQueue.h"
+#include "Containers/ConsumeAllMpmcQueue.h"
 #include "RHI.h"
 #include "Experimental/Containers/HazardPointer.h"
 #include "Misc/MemStack.h"
 #include "Stats/Stats.h"
 #include "TextureProfiler.h"
 
-static UE::TConsumeAllMpmcQueue<FRHIResource*> PendingDeletes;
-static UE::TConsumeAllMpmcQueue<FRHIResource*> PendingDeletesWithLifetimeExtension;
+UE::TConsumeAllMpmcQueue<FRHIResource*> PendingDeletes;
+UE::TConsumeAllMpmcQueue<FRHIResource*> PendingDeletesWithLifetimeExtension;
 
 FRHIResource* FRHIResource::CurrentlyDeleting = nullptr;
 
