Adding missing RHI_API uses.

diff --git a/Engine/Source/Runtime/RHI/Public/RHIResources.h b/Engine/Source/Runtime/RHI/Public/RHIResources.h
--- a/Engine/Source/Runtime/RHI/Public/RHIResources.h
+++ b/Engine/Source/Runtime/RHI/Public/RHIResources.h
@@ -113,8 +113,8 @@ public:
 
 	inline ERHIResourceType GetType() const { return ResourceType; }
 
-	FName GetOwnerName() const;
-	void SetOwnerName(const FName& InOwnerName);
+	RHI_API FName GetOwnerName() const;
+	RHI_API void SetOwnerName(const FName& InOwnerName);
 
 #if RHI_ENABLE_RESOURCE_INFO
 	// Get resource info if available.
