Adding missing RHI_API.

diff --git a/Engine/Source/Runtime/RHI/Public/RHIResources.h b/Engine/Source/Runtime/RHI/Public/RHIResources.h
--- a/Engine/Source/Runtime/RHI/Public/RHIResources.h
+++ b/Engine/Source/Runtime/RHI/Public/RHIResources.h
@@ -218,7 +218,7 @@ private:
 	FName OwnerName;
 #endif
 
-	static FRHIResource* CurrentlyDeleting;
+	RHI_API static FRHIResource* CurrentlyDeleting;
 
 	friend FRHICommandListImmediate;
 };
