RHI Dependency cleanup
- Moving RHILockTracker to its own files.
- Moving PixelFormat helpers from RHIUtilities to PixelFormat

diff --git a/Engine/Source/Runtime/Apple/MetalRHI/Private/MetalTexture.cpp b/Engine/Source/Runtime/Apple/MetalRHI/Private/MetalTexture.cpp
--- a/Engine/Source/Runtime/Apple/MetalRHI/Private/MetalTexture.cpp
+++ b/Engine/Source/Runtime/Apple/MetalRHI/Private/MetalTexture.cpp
@@ -11,6 +11,7 @@
 #include "Containers/ResourceArray.h"
 #include "Misc/ScopeRWLock.h"
 #include "MetalLLM.h"
+#include "RHILockTracker.h"
 
 volatile int64 FMetalSurface::ActiveUploads = 0;
 
diff --git a/Engine/Source/Runtime/Core/Public/PixelFormat.h b/Engine/Source/Runtime/Core/Public/PixelFormat.h
--- a/Engine/Source/Runtime/Core/Public/PixelFormat.h
+++ b/Engine/Source/Runtime/Core/Public/PixelFormat.h
@@ -250,6 +250,22 @@ ENUM_CLASS_FLAGS(EPixelFormatCapabilities);
 // type conversion warnings
 #define UE_PIXELFORMAT_TO_UINT8(argument) static_cast<uint8>(argument)
 
+static inline bool IsBlockCompressedFormat(EPixelFormat Format)
+{
+	switch (Format)
+	{
+	case PF_DXT1:
+	case PF_DXT3:
+	case PF_DXT5:
+	case PF_BC4:
+	case PF_BC5:
+	case PF_BC6H:
+	case PF_BC7:
+		return true;
+	}
+	return false;
+}
+
 FORCEINLINE bool IsHDR(EPixelFormat PixelFormat)
 {
 	return PixelFormat == PF_FloatRGBA || PixelFormat == PF_BC6H || PixelFormat == PF_R16F || PixelFormat == PF_R32_FLOAT || PixelFormat == PF_A32B32G32R32F;
@@ -280,6 +296,50 @@ FORCEINLINE bool IsInteger(EPixelFormat PixelFormat)
 	return false;
 }
 
+static bool IsFloatFormat(EPixelFormat Format)
+{
+	switch (Format)
+	{
+	case PF_A32B32G32R32F:
+	case PF_FloatRGB:
+	case PF_FloatRGBA:
+	case PF_R32_FLOAT:
+	case PF_G16R16F:
+	case PF_G16R16F_FILTER:
+	case PF_G32R32F:
+	case PF_R16F:
+	case PF_R16F_FILTER:
+	case PF_FloatR11G11B10:
+		return true;
+	}
+	return false;
+}
+
+FORCEINLINE bool IsDepthOrStencilFormat(EPixelFormat Format)
+{
+	switch (Format)
+	{
+	case PF_D24:
+	case PF_DepthStencil:
+	case PF_X24_G8:
+	case PF_ShadowDepth:
+	case PF_R32_FLOAT:
+		return true;
+	}
+	return false;
+}
+
+FORCEINLINE bool IsStencilFormat(EPixelFormat Format)
+{
+	switch (Format)
+	{
+	case PF_DepthStencil:
+	case PF_X24_G8:
+		return true;
+	}
+	return false;
+}
+
 FORCEINLINE bool IsDXTCBlockCompressedTextureFormat(EPixelFormat PixelFormat)
 {
 	switch (PixelFormat)
diff --git a/Engine/Source/Runtime/RHI/Private/RHILockTracker.cpp b/Engine/Source/Runtime/RHI/Private/RHILockTracker.cpp
new file mode 100644
--- /dev/null
+++ b/Engine/Source/Runtime/RHI/Private/RHILockTracker.cpp
@@ -0,0 +1,11 @@
+// Copyright Epic Games, Inc. All Rights Reserved.
+
+#include "RHILockTracker.h"
+#include "RHI.h"
+
+FRHILockTracker GRHILockTracker;
+
+void FRHILockTracker::RaiseMismatchError()
+{
+	UE_LOG(LogRHI, Fatal, TEXT("Mismatched RHI buffer locks."));
+}
diff --git a/Engine/Source/Runtime/RHI/Private/RHIUtilities.cpp b/Engine/Source/Runtime/RHI/Private/RHIUtilities.cpp
--- a/Engine/Source/Runtime/RHI/Private/RHIUtilities.cpp
+++ b/Engine/Source/Runtime/RHI/Private/RHIUtilities.cpp
@@ -27,8 +27,6 @@ TAutoConsoleVariable<FString> FDumpTransitionsHelper::CVarDumpTransitionsForReso
 	TEXT("Name of the resource to dump"),
 	ECVF_Default);
 
-RHI_API FRHILockTracker GRHILockTracker;
-
 FName FDumpTransitionsHelper::DumpTransitionForResource = NAME_None;
 void FDumpTransitionsHelper::DumpTransitionForResourceHandler()
 {
@@ -994,8 +992,3 @@ void RHICreateTargetableShaderResource3D(
 
 	RHICreateTargetableShaderResource(Desc, TargetableTextureFlags, bForceSeparateTargetAndShaderResource, false, OutTargetableTexture, OutShaderResourceTexture);
 }
-
-void FRHILockTracker::RaiseMismatchError()
-{
-	UE_LOG(LogRHI, Fatal, TEXT("Mismatched RHI buffer locks."));
-}
diff --git a/Engine/Source/Runtime/RHI/Public/RHI.h b/Engine/Source/Runtime/RHI/Public/RHI.h
--- a/Engine/Source/Runtime/RHI/Public/RHI.h
+++ b/Engine/Source/Runtime/RHI/Public/RHI.h
@@ -1200,9 +1200,6 @@ extern RHI_API void RHIExit();
 DECLARE_DELEGATE_OneParam(FRHIPanicEvent, const FName&);
 extern RHI_API FRHIPanicEvent& RHIGetPanicDelegate();
 
-// RHI utility functions that depend on the RHI definitions.
-#include "RHIUtilities.h"
-
 // Return what the expected number of samplers will be supported by a feature level
 // Note that since the Feature Level is pretty orthogonal to the RHI/HW, this is not going to be perfect
 // If should only be used for a guess at the limit, the real limit will not be known until runtime
@@ -1217,3 +1214,7 @@ RHI_API ERHIBindlessConfiguration RHIGetBindlessSamplersConfiguration(EShaderPla
 #if UE_ENABLE_INCLUDE_ORDER_DEPRECATED_IN_5_2
 #include "RHIStrings.h"
 #endif
+
+#if UE_ENABLE_INCLUDE_ORDER_DEPRECATED_IN_5_3
+#include "RHIUtilities.h"
+#endif
diff --git a/Engine/Source/Runtime/RHI/Public/RHIGPUReadback.h b/Engine/Source/Runtime/RHI/Public/RHIGPUReadback.h
--- a/Engine/Source/Runtime/RHI/Public/RHIGPUReadback.h
+++ b/Engine/Source/Runtime/RHI/Public/RHIGPUReadback.h
@@ -10,6 +10,8 @@
 
 #include "CoreMinimal.h"
 #include "RHI.h"
+#include "DynamicRHI.h"
+#include "RHICommandList.h"
 
 /**
  * FRHIGPUMemoryReadback: Represents a memory readback request scheduled with CopyToStagingBuffer
diff --git a/Engine/Source/Runtime/RHI/Public/RHILockTracker.h b/Engine/Source/Runtime/RHI/Public/RHILockTracker.h
new file mode 100644
--- /dev/null
+++ b/Engine/Source/Runtime/RHI/Public/RHILockTracker.h
@@ -0,0 +1,117 @@
+// Copyright Epic Games, Inc. All Rights Reserved.
+
+#pragma once
+
+#include "RHICommandList.h"
+
+struct FRHILockTracker
+{
+	struct FLockParams
+	{
+		void* RHIBuffer;
+		void* Buffer;
+		uint32 BufferSize;
+		uint32 Offset;
+		EResourceLockMode LockMode;
+		bool bDirectLock; //did we call the normal flushing/updating lock?
+		bool bCreateLock; //did we lock to immediately initialize a newly created buffer?
+
+		FORCEINLINE_DEBUGGABLE FLockParams(void* InRHIBuffer, void* InBuffer, uint32 InOffset, uint32 InBufferSize, EResourceLockMode InLockMode, bool bInbDirectLock, bool bInCreateLock)
+			: RHIBuffer(InRHIBuffer)
+			, Buffer(InBuffer)
+			, BufferSize(InBufferSize)
+			, Offset(InOffset)
+			, LockMode(InLockMode)
+			, bDirectLock(bInbDirectLock)
+			, bCreateLock(bInCreateLock)
+		{
+		}
+	};
+
+	struct FUnlockFenceParams
+	{
+		FUnlockFenceParams(void* InRHIBuffer, FGraphEventRef InUnlockEvent)
+			: RHIBuffer(InRHIBuffer)
+			, UnlockEvent(InUnlockEvent)
+		{
+
+		}
+		void* RHIBuffer;
+		FGraphEventRef UnlockEvent;
+	};
+
+	TArray<FLockParams, TInlineAllocator<16> > OutstandingLocks;
+	uint32 TotalMemoryOutstanding;
+	TArray<FUnlockFenceParams, TInlineAllocator<16> > OutstandingUnlocks;
+
+	FRHILockTracker()
+	{
+		TotalMemoryOutstanding = 0;
+	}
+
+	FORCEINLINE_DEBUGGABLE void Lock(void* RHIBuffer, void* Buffer, uint32 Offset, uint32 SizeRHI, EResourceLockMode LockMode, bool bInDirectBufferWrite = false, bool bInCreateLock = false)
+	{
+#if DO_CHECK
+		for (auto& Parms : OutstandingLocks)
+		{
+			check((Parms.RHIBuffer != RHIBuffer) || (Parms.bDirectLock && bInDirectBufferWrite) || Parms.Offset != Offset);
+		}
+#endif
+		OutstandingLocks.Add(FLockParams(RHIBuffer, Buffer, Offset, SizeRHI, LockMode, bInDirectBufferWrite, bInCreateLock));
+		TotalMemoryOutstanding += SizeRHI;
+	}
+	FORCEINLINE_DEBUGGABLE FLockParams Unlock(void* RHIBuffer, uint32 Offset = 0)
+	{
+		for (int32 Index = 0; Index < OutstandingLocks.Num(); Index++)
+		{
+			if (OutstandingLocks[Index].RHIBuffer == RHIBuffer && OutstandingLocks[Index].Offset == Offset)
+			{
+				FLockParams Result = OutstandingLocks[Index];
+				OutstandingLocks.RemoveAtSwap(Index, 1, false);
+				return Result;
+			}
+		}
+		RaiseMismatchError();
+		return FLockParams(nullptr, nullptr, 0, 0, RLM_WriteOnly, false, false);
+	}
+
+	template<class TIndexOrVertexBufferPointer>
+	FORCEINLINE_DEBUGGABLE void AddUnlockFence(TIndexOrVertexBufferPointer* Buffer, FRHICommandListImmediate& RHICmdList, const FLockParams& LockParms)
+	{
+		if (LockParms.LockMode != RLM_WriteOnly || !(Buffer->GetUsage() & BUF_Volatile))
+		{
+			OutstandingUnlocks.Emplace(Buffer, RHICmdList.RHIThreadFence(true));
+		}
+	}
+
+	FORCEINLINE_DEBUGGABLE void WaitForUnlock(void* RHIBuffer)
+	{
+		for (int32 Index = 0; Index < OutstandingUnlocks.Num(); Index++)
+		{
+			if (OutstandingUnlocks[Index].RHIBuffer == RHIBuffer)
+			{
+				FRHICommandListExecutor::WaitOnRHIThreadFence(OutstandingUnlocks[Index].UnlockEvent);
+				OutstandingUnlocks.RemoveAtSwap(Index, 1, false);
+				return;
+			}
+		}
+	}
+
+	FORCEINLINE_DEBUGGABLE void FlushCompleteUnlocks()
+	{
+		uint32 Count = OutstandingUnlocks.Num();
+		for (uint32 Index = 0; Index < Count; Index++)
+		{
+			if (OutstandingUnlocks[Index].UnlockEvent->IsComplete())
+			{
+				OutstandingUnlocks.RemoveAt(Index, 1);
+				--Count;
+				--Index;
+			}
+		}
+	}
+
+	RHI_API void RaiseMismatchError();
+};
+
+extern RHI_API FRHILockTracker GRHILockTracker;
diff --git a/Engine/Source/Runtime/RHI/Public/RHISurfaceDataConversion.h b/Engine/Source/Runtime/RHI/Public/RHISurfaceDataConversion.h
--- a/Engine/Source/Runtime/RHI/Public/RHISurfaceDataConversion.h
+++ b/Engine/Source/Runtime/RHI/Public/RHISurfaceDataConversion.h
@@ -11,6 +11,7 @@
 #include "Math/Plane.h"
 #include "Math/UnrealMathUtility.h"
 #include "RHI.h"
+#include "RHITypes.h"
 
 /** Helper for accessing R10G10B10A2 colors. */
 struct FRHIR10G10B10A2
diff --git a/Engine/Source/Runtime/RHI/Public/RHIUtilities.h b/Engine/Source/Runtime/RHI/Public/RHIUtilities.h
--- a/Engine/Source/Runtime/RHI/Public/RHIUtilities.h
+++ b/Engine/Source/Runtime/RHI/Public/RHIUtilities.h
@@ -13,56 +13,6 @@
 #include "RHIDefinitions.h"
 #include "RHICommandList.h"
 
-static inline bool IsDepthOrStencilFormat(EPixelFormat Format)
-{
-	switch (Format)
-	{
-	case PF_D24:
-	case PF_DepthStencil:
-	case PF_X24_G8:
-	case PF_ShadowDepth:
-	case PF_R32_FLOAT:
-		return true;
-
-	default:
-		break;
-	}
-
-	return false;
-}
-
-static inline bool IsStencilFormat(EPixelFormat Format)
-{
-	switch (Format)
-	{
-	case PF_DepthStencil:
-	case PF_X24_G8:
-		return true;
-
-	default:
-		break;
-	}
-
-	return false;
-}
-
-static inline bool IsBlockCompressedFormat(EPixelFormat Format)
-{
-	switch (Format)
-	{
-	case PF_DXT1:
-	case PF_DXT3:
-	case PF_DXT5:
-	case PF_BC4:
-	case PF_BC5:
-	case PF_BC6H:
-	case PF_BC7:
-		return true;
-	}
-
-	return false;
-}
-
 static inline EPixelFormat GetBlockCompressedFormatUAVAliasFormat(EPixelFormat Format)
 {
 	switch (Format)
@@ -82,28 +32,6 @@ static inline EPixelFormat GetBlockCompressedFormatUAVAliasFormat(EPixelFormat F
 	return Format;
 }
 
-static bool IsFloatFormat(EPixelFormat Format)
-{
-	switch (Format)
-	{
-	case PF_A32B32G32R32F:
-	case PF_FloatRGB:
-	case PF_FloatRGBA:
-	case PF_R32_FLOAT:
-	case PF_G16R16F:
-	case PF_G16R16F_FILTER:
-	case PF_G32R32F:
-	case PF_R16F:
-	case PF_R16F_FILTER:
-	case PF_FloatR11G11B10:
-		return true;
-
-	default:
-		break;
-	}
-	return false;
-}
-
 static bool IsUnormFormat(EPixelFormat Format)
 {
 	switch (Format)
@@ -796,114 +724,7 @@ extern RHI_API float RHIGetFrameTime();
 
 extern RHI_API void RHICalculateFrameTime();
 
-struct FRHILockTracker
-{
-	struct FLockParams
-	{
-		void* RHIBuffer;
-		void* Buffer;
-		uint32 BufferSize;
-		uint32 Offset;
-		EResourceLockMode LockMode;
-		bool bDirectLock; //did we call the normal flushing/updating lock?
-		bool bCreateLock; //did we lock to immediately initialize a newly created buffer?
-		
-		FORCEINLINE_DEBUGGABLE FLockParams(void* InRHIBuffer, void* InBuffer, uint32 InOffset, uint32 InBufferSize, EResourceLockMode InLockMode, bool bInbDirectLock, bool bInCreateLock)
-		: RHIBuffer(InRHIBuffer)
-		, Buffer(InBuffer)
-		, BufferSize(InBufferSize)
-		, Offset(InOffset)
-		, LockMode(InLockMode)
-		, bDirectLock(bInbDirectLock)
-		, bCreateLock(bInCreateLock)
-		{
-		}
-	};
-	
-	struct FUnlockFenceParams
-	{
-		FUnlockFenceParams(void* InRHIBuffer, FGraphEventRef InUnlockEvent)
-		: RHIBuffer(InRHIBuffer)
-		, UnlockEvent(InUnlockEvent)
-		{
-			
-		}
-		void* RHIBuffer;
-		FGraphEventRef UnlockEvent;
-	};
-	
-	TArray<FLockParams, TInlineAllocator<16> > OutstandingLocks;
-	uint32 TotalMemoryOutstanding;
-	TArray<FUnlockFenceParams, TInlineAllocator<16> > OutstandingUnlocks;
-	
-	FRHILockTracker()
-	{
-		TotalMemoryOutstanding = 0;
-	}
-	
-	FORCEINLINE_DEBUGGABLE void Lock(void* RHIBuffer, void* Buffer, uint32 Offset, uint32 SizeRHI, EResourceLockMode LockMode, bool bInDirectBufferWrite = false, bool bInCreateLock = false)
-	{
-#if DO_CHECK
-		for (auto& Parms : OutstandingLocks)
-		{
-			check((Parms.RHIBuffer != RHIBuffer) || (Parms.bDirectLock && bInDirectBufferWrite) || Parms.Offset != Offset);
-		}
+#if UE_ENABLE_INCLUDE_ORDER_DEPRECATED_IN_5_3
+#include "RHILockTracker.h"
+#include "PixelFormat.h"
 #endif
-		OutstandingLocks.Add(FLockParams(RHIBuffer, Buffer, Offset, SizeRHI, LockMode, bInDirectBufferWrite, bInCreateLock));
-		TotalMemoryOutstanding += SizeRHI;
-	}
-	FORCEINLINE_DEBUGGABLE FLockParams Unlock(void* RHIBuffer, uint32 Offset=0)
-	{
-		for (int32 Index = 0; Index < OutstandingLocks.Num(); Index++)
-		{
-			if (OutstandingLocks[Index].RHIBuffer == RHIBuffer && OutstandingLocks[Index].Offset == Offset)
-			{
-				FLockParams Result = OutstandingLocks[Index];
-				OutstandingLocks.RemoveAtSwap(Index, 1, false);
-				return Result;
-			}
-		}
-		RaiseMismatchError();
-		return FLockParams(nullptr, nullptr, 0, 0, RLM_WriteOnly, false, false);
-	}
-	
-	template<class TIndexOrVertexBufferPointer>
-	FORCEINLINE_DEBUGGABLE void AddUnlockFence(TIndexOrVertexBufferPointer* Buffer, FRHICommandListImmediate& RHICmdList, const FLockParams& LockParms)
-	{
-		if (LockParms.LockMode != RLM_WriteOnly || !(Buffer->GetUsage() & BUF_Volatile))
-		{
-			OutstandingUnlocks.Emplace(Buffer, RHICmdList.RHIThreadFence(true));
-		}
-	}
-	
-	FORCEINLINE_DEBUGGABLE void WaitForUnlock(void* RHIBuffer)
-	{
-		for (int32 Index = 0; Index < OutstandingUnlocks.Num(); Index++)
-		{
-			if (OutstandingUnlocks[Index].RHIBuffer == RHIBuffer)
-			{
-				FRHICommandListExecutor::WaitOnRHIThreadFence(OutstandingUnlocks[Index].UnlockEvent);
-				OutstandingUnlocks.RemoveAtSwap(Index, 1, false);
-				return;
-			}
-		}
-	}
-	
-	FORCEINLINE_DEBUGGABLE void FlushCompleteUnlocks()
-	{
-		uint32 Count = OutstandingUnlocks.Num();
-		for (uint32 Index = 0; Index < Count; Index++)
-		{
-			if (OutstandingUnlocks[Index].UnlockEvent->IsComplete())
-			{
-				OutstandingUnlocks.RemoveAt(Index, 1);
-				--Count;
-				--Index;
-			}
-		}
-	}
-
-	RHI_API void RaiseMismatchError();
-};
-
-extern RHI_API FRHILockTracker GRHILockTracker;
diff --git a/Engine/Source/Runtime/RHI/Public/RHIValidationContext.h b/Engine/Source/Runtime/RHI/Public/RHIValidationContext.h
--- a/Engine/Source/Runtime/RHI/Public/RHIValidationContext.h
+++ b/Engine/Source/Runtime/RHI/Public/RHIValidationContext.h
@@ -12,7 +12,9 @@
 #if ENABLE_RHI_VALIDATION
 
 #include "RHI.h"
+#if UE_ENABLE_INCLUDE_ORDER_DEPRECATED_IN_5_3
 #include "RHIUtilities.h"
+#endif
 
 class FValidationRHI;
 
