Fix for race condition with pre-cache PSOs which can produce intermittent and unnecessary calls to SetGraphicsPipelineState.

diff --git a/Engine/Source/Runtime/RHI/Private/PipelineStateCache.cpp b/Engine/Source/Runtime/RHI/Private/PipelineStateCache.cpp
--- a/Engine/Source/Runtime/RHI/Private/PipelineStateCache.cpp
+++ b/Engine/Source/Runtime/RHI/Private/PipelineStateCache.cpp
@@ -675,7 +675,8 @@ void SetGraphicsPipelineState(FRHICommandList& RHICmdList, const FGraphicsPipeli
 	RHICmdList.SetGraphicsPipelineState(Initializer, StencilRef, bApplyAdditionalState);
 #else
 	FGraphicsPipelineState* PipelineState = PipelineStateCache::GetAndOrCreateGraphicsPipelineState(RHICmdList, Initializer, ApplyFlags, PSOPrecacheResult);
-	if (PipelineState && (PipelineState->RHIPipeline || !Initializer.bFromPSOFileCache))
+
+	if (PipelineState && !Initializer.bFromPSOFileCache)
 	{
 #if PIPELINESTATECACHE_VERIFYTHREADSAFE
 		int32 Result = PipelineState->InUseCount.Increment();
