Adding setters to FRHIBatchedShaderParameters. Cleaning up FRHIParameterBatcher's exports.

diff --git a/Engine/Source/Runtime/RHI/Public/RHICommandList.h b/Engine/Source/Runtime/RHI/Public/RHICommandList.h
--- a/Engine/Source/Runtime/RHI/Public/RHICommandList.h
+++ b/Engine/Source/Runtime/RHI/Public/RHICommandList.h
@@ -497,9 +497,50 @@ struct FRHIBatchedShaderParameters
 		ResourceParameters.Reset();
 		BindlessParameters.Reset();
 	}
+
+	FORCEINLINE_DEBUGGABLE void SetShaderParameter(uint32 BufferIndex, uint32 BaseIndex, uint32 NumBytes, const void* NewValue)
+	{
+		const int32 DestDataOffset = ParametersData.Num();
+		ParametersData.Append((const uint8*)NewValue, NumBytes);
+		Parameters.Emplace((uint16)BufferIndex, (uint16)BaseIndex, (uint16)DestDataOffset, (uint16)NumBytes);
+	}
+
+	FORCEINLINE_DEBUGGABLE void SetShaderTexture(uint32 Index, FRHITexture* Texture)
+	{
+		ResourceParameters.Emplace(Texture, (uint16)Index);
+	}
+	FORCEINLINE_DEBUGGABLE void SetShaderResourceViewParameter(uint32 Index, FRHIShaderResourceView* SRV)
+	{
+		ResourceParameters.Emplace(SRV, (uint16)Index);
+	}
+	FORCEINLINE_DEBUGGABLE void SetShaderSampler(uint32 Index, FRHISamplerState* State)
+	{
+		ResourceParameters.Emplace(State, (uint16)Index);
+	}
+	FORCEINLINE_DEBUGGABLE void SetUAVParameter(uint32 Index, FRHIUnorderedAccessView* UAV)
+	{
+		ResourceParameters.Emplace(UAV, (uint16)Index);
+	}
+
+	FORCEINLINE_DEBUGGABLE void SetBindlessTexture(uint32 Index, FRHITexture* Texture)
+	{
+		BindlessParameters.Emplace(Texture, (uint16)Index);
+	}
+	FORCEINLINE_DEBUGGABLE void SetBindlessResourceView(uint32 Index, FRHIShaderResourceView* SRV)
+	{
+		BindlessParameters.Emplace(SRV, (uint16)Index);
+	}
+	FORCEINLINE_DEBUGGABLE void SetBindlessSampler(uint32 Index, FRHISamplerState* State)
+	{
+		BindlessParameters.Emplace(State, (uint16)Index);
+	}
+	FORCEINLINE_DEBUGGABLE void SetBindlessUAV(uint32 Index, FRHIUnorderedAccessView* UAV)
+	{
+		BindlessParameters.Emplace(UAV, (uint16)Index);
+	}
 };
 
-class RHI_API FRHIParameterBatcher
+class FRHIParameterBatcher
 {
 protected:
 	FRHIShader* AllBatchedShaders[SF_NumStandardFrequencies]{};
@@ -519,66 +560,62 @@ protected:
 	}
 
 public:
-	FRHIParameterBatcher();
-	FRHIParameterBatcher(const FBoundShaderStateInput& InBoundShaderStateInput, FRHIComputeShader* InBoundComputeShaderRHI);
-	FRHIParameterBatcher(FRHIParameterBatcher&&);
-	~FRHIParameterBatcher();
+	RHI_API FRHIParameterBatcher();
+	RHI_API FRHIParameterBatcher(const FBoundShaderStateInput& InBoundShaderStateInput, FRHIComputeShader* InBoundComputeShaderRHI);
+	RHI_API FRHIParameterBatcher(FRHIParameterBatcher&&);
+	RHI_API ~FRHIParameterBatcher();
 
 	FORCEINLINE bool IsEnabled() const { return bEnabled; }
 
-	void OnBoundShaderChanged(class FRHICommandList& InCommandList, const FBoundShaderStateInput& InBoundShaderStateInput);
-	void OnBoundShaderChanged(class FRHIComputeCommandList& InCommandList, FRHIComputeShader* InBoundComputeShaderRHI);
+	RHI_API void OnBoundShaderChanged(class FRHICommandList& InCommandList, const FBoundShaderStateInput& InBoundShaderStateInput);
+	RHI_API void OnBoundShaderChanged(class FRHIComputeCommandList& InCommandList, FRHIComputeShader* InBoundComputeShaderRHI);
 
-	void PreDispatch(class FRHIComputeCommandList& InCommandList);
-	void PreDraw(class FRHICommandList& InCommandList);
+	RHI_API void PreDispatch(class FRHIComputeCommandList& InCommandList);
+	RHI_API void PreDraw(class FRHICommandList& InCommandList);
 
-	void FlushAllParameters(class FRHIComputeCommandList& InCommandList);
-	void FlushAllParameters(class FRHICommandList& InCommandList);
+	RHI_API void FlushAllParameters(class FRHIComputeCommandList& InCommandList);
+	RHI_API void FlushAllParameters(class FRHICommandList& InCommandList);
 
-	void FlushPendingParameters(class FRHIComputeCommandList& InCommandList, FRHIComputeShader* InShader);
-	void FlushPendingParameters(class FRHICommandList& InCommandList, FRHIGraphicsShader* InShader);
+	RHI_API void FlushPendingParameters(class FRHIComputeCommandList& InCommandList, FRHIComputeShader* InShader);
+	RHI_API void FlushPendingParameters(class FRHICommandList& InCommandList, FRHIGraphicsShader* InShader);
 
 	FORCEINLINE_DEBUGGABLE void SetShaderParameter(FRHIShader* Shader, uint32 BufferIndex, uint32 BaseIndex, uint32 NumBytes, const void* NewValue)
 	{
-		FRHIBatchedShaderParameters& BatchedParameters = GetBatchedShaderParameters(Shader);
-
-		const int32 DestDataOffset = BatchedParameters.ParametersData.Num();
-		BatchedParameters.ParametersData.Append((const uint8*)NewValue, NumBytes);
-		BatchedParameters.Parameters.Emplace((uint16)BufferIndex, (uint16)BaseIndex, (uint16)DestDataOffset, (uint16)NumBytes);
+		GetBatchedShaderParameters(Shader).SetShaderParameter(BufferIndex, BaseIndex, NumBytes, NewValue);
 	}
 
-	FORCEINLINE_DEBUGGABLE void SetShaderTexture(FRHIShader* Shader, uint32 TextureIndex, FRHITexture* Texture)
+	FORCEINLINE_DEBUGGABLE void SetShaderTexture(FRHIShader* Shader, uint32 Index, FRHITexture* Texture)
 	{
-		GetBatchedShaderParameters(Shader).ResourceParameters.Emplace(Texture, (uint16)TextureIndex);
+		GetBatchedShaderParameters(Shader).SetShaderTexture(Index, Texture);
 	}
-	FORCEINLINE_DEBUGGABLE void SetShaderResourceViewParameter(FRHIShader* Shader, uint32 SamplerIndex, FRHIShaderResourceView* SRV)
+	FORCEINLINE_DEBUGGABLE void SetShaderResourceViewParameter(FRHIShader* Shader, uint32 Index, FRHIShaderResourceView* SRV)
 	{
-		GetBatchedShaderParameters(Shader).ResourceParameters.Emplace(SRV, (uint16)SamplerIndex);
+		GetBatchedShaderParameters(Shader).SetShaderResourceViewParameter(Index, SRV);
 	}
-	FORCEINLINE_DEBUGGABLE void SetShaderSampler(FRHIShader* Shader, uint32 SamplerIndex, FRHISamplerState* State)
+	FORCEINLINE_DEBUGGABLE void SetShaderSampler(FRHIShader* Shader, uint32 Index, FRHISamplerState* State)
 	{
-		GetBatchedShaderParameters(Shader).ResourceParameters.Emplace(State, (uint16)SamplerIndex);
+		GetBatchedShaderParameters(Shader).SetShaderSampler(Index, State);
 	}
 	FORCEINLINE_DEBUGGABLE void SetUAVParameter(FRHIShader* Shader, uint32 UAVIndex, FRHIUnorderedAccessView* UAV)
 	{
-		GetBatchedShaderParameters(Shader).ResourceParameters.Emplace(UAV, (uint16)UAVIndex);
+		GetBatchedShaderParameters(Shader).SetUAVParameter(UAVIndex, UAV);
 	}
 
 	FORCEINLINE_DEBUGGABLE void SetBindlessTexture(FRHIShader* Shader, uint32 Index, FRHITexture* Texture)
 	{
-		GetBatchedShaderParameters(Shader).BindlessParameters.Emplace(Texture, (uint16)Index);
+		GetBatchedShaderParameters(Shader).SetBindlessTexture(Index, Texture);
 	}
 	FORCEINLINE_DEBUGGABLE void SetBindlessResourceView(FRHIShader* Shader, uint32 Index, FRHIShaderResourceView* SRV)
 	{
-		GetBatchedShaderParameters(Shader).BindlessParameters.Emplace(SRV, (uint16)Index);
+		GetBatchedShaderParameters(Shader).SetBindlessResourceView(Index, SRV);
 	}
 	FORCEINLINE_DEBUGGABLE void SetBindlessSampler(FRHIShader* Shader, uint32 Index, FRHISamplerState* State)
 	{
-		GetBatchedShaderParameters(Shader).BindlessParameters.Emplace(State, (uint16)Index);
+		GetBatchedShaderParameters(Shader).SetBindlessSampler(Index, State);
 	}
 	FORCEINLINE_DEBUGGABLE void SetBindlessUAV(FRHIShader* Shader, uint32 Index, FRHIUnorderedAccessView* UAV)
 	{
-		GetBatchedShaderParameters(Shader).BindlessParameters.Emplace(UAV, (uint16)Index);
+		GetBatchedShaderParameters(Shader).SetBindlessUAV(Index, UAV);
 	}
 
 	FORCEINLINE_DEBUGGABLE void SetUAVParameter(FRHIShader* Shader, uint32 UAVIndex, FRHIUnorderedAccessView* UAV, uint32 InitialCount)
