Adding declaration of global needed in RHICommandList without including RHI.h

diff --git a/Engine/Source/Runtime/RHI/Public/RHICommandList.h b/Engine/Source/Runtime/RHI/Public/RHICommandList.h
--- a/Engine/Source/Runtime/RHI/Public/RHICommandList.h
+++ b/Engine/Source/Runtime/RHI/Public/RHICommandList.h
@@ -83,6 +83,7 @@ extern RHI_API bool GUseRHITaskThreads_InternalUseOnly;
 extern RHI_API bool GIsRunningRHIInSeparateThread_InternalUseOnly;
 extern RHI_API bool GIsRunningRHIInDedicatedThread_InternalUseOnly;
 extern RHI_API bool GIsRunningRHIInTaskThread_InternalUseOnly;
+extern RHI_API bool GSupportsParallelRenderingTasksWithSeparateRHIThread;
 
 namespace ERenderThreadIdleTypes
 {
