Integrate RenderResourceViewer from //UE5/Main
Part1 - Integrate CL23235549 as edigrate with modification to support resident flag

diff --git a/Engine/Source/Runtime/RHI/Private/RHI.cpp b/Engine/Source/Runtime/RHI/Private/RHI.cpp
--- a/Engine/Source/Runtime/RHI/Private/RHI.cpp
+++ b/Engine/Source/Runtime/RHI/Private/RHI.cpp
@@ -635,6 +635,7 @@ namespace RHIInternal
 
 	struct FResourceFlags
 	{
+		bool bResident = false;
 		bool bMarkedForDelete = false;
 		bool bTransient = false;
 		bool bStreaming = false;
@@ -648,9 +649,14 @@ namespace RHIInternal
 		{
 			FString FlagsString;
 			bool bHasFlag = false;
+			if (bResident)
+			{
+				FlagsString += "Resident";
+				bHasFlag = true;
+			}
 			if (bMarkedForDelete)
 			{
-				FlagsString += "MarkedForDelete";
+				FlagsString += bHasFlag ? " | MarkedForDelete" : "MarkedForDelete";
 				bHasFlag = true;
 			}
 			if (bTransient)
@@ -775,6 +781,7 @@ namespace RHIInternal
 	FResourceFlags GetResourceFlagsInternal(const FResourceEntry& Resource)
 	{
 		FResourceFlags Flags;
+		Flags.bResident = Resource.ResourceInfo.bResident;
 		Flags.bMarkedForDelete = !Resource.ResourceInfo.bValid;
 		Flags.bTransient = Resource.ResourceInfo.IsTransient;
 
@@ -798,7 +805,7 @@ namespace RHIInternal
 			Flags.bRTAS = EnumHasAnyFlags((EBufferUsageFlags)Buffer->GetUsage(), BUF_AccelerationStructure);
 		}
 
-		Flags.bHasFlags = Flags.bMarkedForDelete || Flags.bTransient || Flags.bStreaming || Flags.bRT || Flags.bDS || Flags.bUAV || Flags.bRTAS;
+		Flags.bHasFlags = Flags.bResident || Flags.bMarkedForDelete || Flags.bTransient || Flags.bStreaming || Flags.bRT || Flags.bDS || Flags.bUAV || Flags.bRTAS;
 		return Flags;
 	}
 }
@@ -887,7 +894,7 @@ void RHIDumpResourceMemory(const FString& NameFilter, ERHIResourceType TypeFilte
 			const TCHAR* ResourceType = StringFromRHIResourceType(ResourceInfo.Type);
 			const int64 SizeInBytes = ResourceInfo.VRamAllocation.AllocationSize;
 			
-			bool bResident = ResourceInfo.bResident;
+
 			RHIInternal::FResourceFlags Flags = GetResourceFlagsInternal(Resources[Index]);
 
 			if (bSummaryOutput == false)
@@ -898,7 +905,7 @@ void RHIDumpResourceMemory(const FString& NameFilter, ERHIResourceType TypeFilte
 						ResourceNameBuffer,
 						ResourceType,
 						SizeInBytes / double(1 << 20),
-						bResident ? TEXT("Yes") : TEXT("No"),
+						Flags.bResident ? TEXT("Yes") : TEXT("No"),
 						Flags.bMarkedForDelete ? TEXT("Yes") : TEXT("No"),
 						Flags.bTransient ? TEXT("Yes") : TEXT("No"),
 						Flags.bStreaming ? TEXT("Yes") : TEXT("No"),
