Removing the deprecated RHISetComputeShader

diff --git a/Engine/Source/Runtime/RHI/Private/RHIContext.cpp b/Engine/Source/Runtime/RHI/Private/RHIContext.cpp
--- a/Engine/Source/Runtime/RHI/Private/RHIContext.cpp
+++ b/Engine/Source/Runtime/RHI/Private/RHIContext.cpp
@@ -5,17 +5,6 @@
 #include "RHI.h"
 #include "RHIStats.h"
 
-RHI_API void RHISetComputeShaderBackwardsCompatible(IRHIComputeContext* InContext, FRHIComputeShader* InShader)
-{
-	TRefCountPtr<FRHIComputePipelineState> ComputePipelineState = RHICreateComputePipelineState(InShader);
-	InContext->RHISetComputePipelineState(ComputePipelineState);
-}
-
-void IRHIComputeContext::RHISetComputeShader(FRHIComputeShader* ComputeShader)
-{
-	RHISetComputeShaderBackwardsCompatible(this, ComputeShader);
-}
-
 void IRHIComputeContext::StatsSetCategory(FRHIDrawStats* InStats, uint32 InCategoryID, uint32 InGPUIndex)
 {
 	Stats = &InStats->GetGPU(InGPUIndex).GetCategory(InCategoryID);
diff --git a/Engine/Source/Runtime/RHI/Public/RHICommandList.h b/Engine/Source/Runtime/RHI/Public/RHICommandList.h
--- a/Engine/Source/Runtime/RHI/Public/RHICommandList.h
+++ b/Engine/Source/Runtime/RHI/Public/RHICommandList.h
@@ -1556,16 +1556,6 @@ FRHICOMMAND_MACRO(FRHICommandNextSubpass)
 	RHI_API void Execute(FRHICommandListBase& CmdList);
 };
 
-FRHICOMMAND_MACRO(FRHICommandSetComputeShader)
-{
-	FRHIComputeShader* ComputeShader;
-	FORCEINLINE_DEBUGGABLE FRHICommandSetComputeShader(FRHIComputeShader* InComputeShader)
-		: ComputeShader(InComputeShader)
-	{
-	}
-	RHI_API void Execute(FRHICommandListBase& CmdList);
-};
-
 FRHICOMMAND_MACRO(FRHICommandSetComputePipelineState)
 {
 	FComputePipelineState* ComputePipelineState;
@@ -2725,21 +2715,6 @@ public:
 		SetShaderParameters(Shader, {}, {}, {}, MakeArrayView(&Resource, 1));
 	}
 
-	UE_DEPRECATED(5.1, "ComputePipelineStates should be used instead of direct ComputeShaders. You can use SetComputePipelineState(RHICmdList, ComputeShader).")
-	FORCEINLINE_DEBUGGABLE void SetComputeShader(FRHIComputeShader* ComputeShader)
-	{
-		OnBoundShaderChanged(ComputeShader);
-		ComputeShader->UpdateStats();
-		if (Bypass())
-		{
-PRAGMA_DISABLE_DEPRECATION_WARNINGS
-			GetComputeContext().RHISetComputeShader(ComputeShader);
-PRAGMA_ENABLE_DEPRECATION_WARNINGS
-			return;
-		}
-		ALLOC_COMMAND(FRHICommandSetComputeShader)(ComputeShader);
-	}
-
 	FORCEINLINE_DEBUGGABLE void SetComputePipelineState(FComputePipelineState* ComputePipelineState, FRHIComputeShader* ComputeShader)
 	{
 		OnBoundShaderChanged(ComputeShader);
diff --git a/Engine/Source/Runtime/RHI/Public/RHICommandListCommandExecutes.inl b/Engine/Source/Runtime/RHI/Public/RHICommandListCommandExecutes.inl
--- a/Engine/Source/Runtime/RHI/Public/RHICommandListCommandExecutes.inl
+++ b/Engine/Source/Runtime/RHI/Public/RHICommandListCommandExecutes.inl
@@ -254,14 +254,6 @@ void FRHICommandNextSubpass::Execute(FRHICommandListBase& CmdList)
 	INTERNAL_DECORATOR(RHINextSubpass)();
 }
 
-void FRHICommandSetComputeShader::Execute(FRHICommandListBase& CmdList)
-{
-	RHISTAT(SetComputeShader);
-PRAGMA_DISABLE_DEPRECATION_WARNINGS
-	INTERNAL_DECORATOR_COMPUTE(RHISetComputeShader)(ComputeShader);
-PRAGMA_ENABLE_DEPRECATION_WARNINGS
-}
-
 void FRHICommandSetComputePipelineState::Execute(FRHICommandListBase& CmdList)
 {
 	RHISTAT(SetComputePipelineState);
diff --git a/Engine/Source/Runtime/RHI/Public/RHIContext.h b/Engine/Source/Runtime/RHI/Public/RHIContext.h
--- a/Engine/Source/Runtime/RHI/Public/RHIContext.h
+++ b/Engine/Source/Runtime/RHI/Public/RHIContext.h
@@ -232,9 +232,6 @@ public:
 		return ERHIPipeline::AsyncCompute;
 	}
 
-	UE_DEPRECATED(5.1, "ComputePipelineStates should be used instead of direct ComputeShaders.")
-	RHI_API virtual void RHISetComputeShader(FRHIComputeShader* ComputeShader);
-
 	virtual void RHISetComputePipelineState(FRHIComputePipelineState* ComputePipelineState) = 0;
 
 	virtual void RHIDispatchComputeShader(uint32 ThreadGroupCountX, uint32 ThreadGroupCountY, uint32 ThreadGroupCountZ) = 0;
@@ -820,8 +817,7 @@ public:
 	virtual void RHISetRasterizerState(FRHIRasterizerState* NewState) = 0;
 	virtual void RHISetBlendState(FRHIBlendState* NewState, const FLinearColor& BlendFactor) = 0;
 	virtual void RHIEnableDepthBoundsTest(bool bEnable) = 0;
-	// TODO: uncomment when removed from IRHIComputeContext
-	//virtual void RHISetComputeShader(FRHIComputeShader* ComputeShader) = 0;
+	virtual void RHISetComputeShader(FRHIComputeShader* ComputeShader) = 0;
 
 	/**
 	* This will set most relevant pipeline state. Legacy APIs are expected to set corresponding disjoint state as well.
@@ -844,9 +840,7 @@ public:
 	{
 		if (FRHIComputePipelineStateFallback* FallbackState = static_cast<FRHIComputePipelineStateFallback*>(ComputePipelineState))
 		{
-PRAGMA_DISABLE_DEPRECATION_WARNINGS
 			RHISetComputeShader(FallbackState->GetComputeShader());
-PRAGMA_ENABLE_DEPRECATION_WARNINGS
 		}
 	}
 
