Removing deprecated RHI symbols.

diff --git a/Engine/Source/Runtime/RHI/Public/RHICommandList.h b/Engine/Source/Runtime/RHI/Public/RHICommandList.h
--- a/Engine/Source/Runtime/RHI/Public/RHICommandList.h
+++ b/Engine/Source/Runtime/RHI/Public/RHICommandList.h
@@ -668,12 +668,6 @@ public:
 	//
 	void FinishRecording();
 
-	UE_DEPRECATED(5.1, "Call FinishRecording instead.")
-	void HandleRTThreadTaskCompletion(const FGraphEventRef& MyCompletionGraphEvent)
-	{
-		FinishRecording();
-	}
-
 	void SetCurrentStat(TStatId Stat);
 
 	FORCEINLINE_DEBUGGABLE void* Alloc(int64 AllocSize, int64 Alignment)
@@ -1928,94 +1922,6 @@ FRHICOMMAND_MACRO(FRHICommandSetStaticUniformBuffers)
 	{}
 
 	RHI_API void Execute(FRHICommandListBase& CmdList);
-};
-
-struct FComputedUniformBuffer
-{
-	FUniformBufferRHIRef UniformBuffer;
-	mutable int32 UseCount;
-	FComputedUniformBuffer()
-		: UseCount(0)
-	{
-	}
-};
-
-struct FLocalUniformBufferWorkArea
-{
-	void* Contents;
-	FUniformBufferLayoutRHIRef Layout;
-	FComputedUniformBuffer* ComputedUniformBuffer;
-#if DO_CHECK || USING_CODE_ANALYSIS // the below variables are used in check(), which can be enabled in Shipping builds (see Build.h)
-	FRHICommandListBase* CheckCmdList;
-	int32 UID;
-#endif
-
-	FLocalUniformBufferWorkArea(FRHICommandListBase* InCheckCmdList, const void* InContents, uint32 ContentsSize, const FRHIUniformBufferLayout* InLayout)
-		: Layout(InLayout)
-#if DO_CHECK || USING_CODE_ANALYSIS
-		, CheckCmdList(InCheckCmdList)
-		, UID(InCheckCmdList->GetUID())
-#endif
-	{
-		check(ContentsSize);
-		Contents = InCheckCmdList->Alloc(ContentsSize, SHADER_PARAMETER_STRUCT_ALIGNMENT);
-		FMemory::Memcpy(Contents, InContents, ContentsSize);
-		ComputedUniformBuffer = new (InCheckCmdList->Alloc<FComputedUniformBuffer>()) FComputedUniformBuffer;
-	}
-};
-
-struct FLocalUniformBuffer
-{
-	FLocalUniformBufferWorkArea* WorkArea;
-	FUniformBufferRHIRef BypassUniform; // this is only used in the case of Bypass, should eventually be deleted
-	FLocalUniformBuffer()
-		: WorkArea(nullptr)
-	{
-	}
-	FLocalUniformBuffer(const FLocalUniformBuffer& Other)
-		: WorkArea(Other.WorkArea)
-		, BypassUniform(Other.BypassUniform)
-	{
-	}
-	FORCEINLINE_DEBUGGABLE bool IsValid() const
-	{
-		return WorkArea || IsValidRef(BypassUniform);
-	}
-};
-
-FRHICOMMAND_MACRO(FRHICommandBuildLocalUniformBuffer)
-{
-	FLocalUniformBufferWorkArea WorkArea;
-	FORCEINLINE_DEBUGGABLE FRHICommandBuildLocalUniformBuffer(
-		FRHICommandListBase* CheckCmdList,
-		const void* Contents,
-		uint32 ContentsSize,
-		const FRHIUniformBufferLayout* Layout
-		)
-		: WorkArea(CheckCmdList, Contents, ContentsSize, Layout)
-
-	{
-	}
-	RHI_API void Execute(FRHICommandListBase& CmdList);
-};
-
-FRHICOMMAND_MACRO_TPL(TRHIShader, FRHICommandSetLocalUniformBuffer)
-{
-	TRHIShader* Shader;
-	uint32 BaseIndex;
-	FLocalUniformBuffer LocalUniformBuffer;
-	FORCEINLINE_DEBUGGABLE FRHICommandSetLocalUniformBuffer(FRHICommandListBase* CheckCmdList, TRHIShader* InShader, uint32 InBaseIndex, const FLocalUniformBuffer& InLocalUniformBuffer)
-		: Shader(InShader)
-		, BaseIndex(InBaseIndex)
-		, LocalUniformBuffer(InLocalUniformBuffer)
-
-	{
-		check(CheckCmdList == LocalUniformBuffer.WorkArea->CheckCmdList && CheckCmdList->GetUID() == LocalUniformBuffer.WorkArea->UID); // this uniform buffer was not built for this particular commandlist
-		LocalUniformBuffer.WorkArea->ComputedUniformBuffer->UseCount++;
-	}
-	RHI_API void Execute(FRHICommandListBase& CmdList);
-};
-
 FRHICOMMAND_MACRO(FRHICommandBeginRenderQuery)
 {
 	FRHIRenderQuery* RenderQuery;
@@ -2465,42 +2371,6 @@ public:
 		SetShaderUniformBuffer(Shader.GetReference(), BaseIndex, UniformBuffer);
 	}
 
-	UE_DEPRECATED(5.1, "Local uniform buffers are now deprecated. Use RHICreateUniformBuffer instead.")
-	FORCEINLINE_DEBUGGABLE FLocalUniformBuffer BuildLocalUniformBuffer(const void* Contents, uint32 ContentsSize, const FRHIUniformBufferLayout* Layout)
-	{
-		FLocalUniformBuffer Result;
-		if (Bypass())
-		{
-			Result.BypassUniform = RHICreateUniformBuffer(Contents, Layout, UniformBuffer_SingleFrame);
-		}
-		else
-		{
-			check(Contents && ContentsSize && (&Layout != nullptr));
-			auto* Cmd = ALLOC_COMMAND(FRHICommandBuildLocalUniformBuffer)(this, Contents, ContentsSize, Layout);
-			Result.WorkArea = &Cmd->WorkArea;
-		}
-		return Result;
-	}
-
-	template <typename TRHIShader>
-	UE_DEPRECATED(5.1, "Local uniform buffers are now deprecated. Use RHICreateUniformBuffer instead.")
-	FORCEINLINE_DEBUGGABLE void SetLocalShaderUniformBuffer(TRHIShader* Shader, uint32 BaseIndex, const FLocalUniformBuffer& UniformBuffer)
-	{
-		ValidateBoundShader(Shader);
-		if (Bypass())
-		{
-			GetContext().RHISetShaderUniformBuffer(Shader, BaseIndex, UniformBuffer.BypassUniform);
-			return;
-		}
-		ALLOC_COMMAND(FRHICommandSetLocalUniformBuffer<TRHIShader>)(this, Shader, BaseIndex, UniformBuffer);
-	}
-
-	template <typename TShaderRHI>
-	FORCEINLINE_DEBUGGABLE void SetLocalShaderUniformBuffer(const TRefCountPtr<TShaderRHI>& Shader, uint32 BaseIndex, const FLocalUniformBuffer& UniformBuffer)
-	{
-		SetLocalShaderUniformBuffer(Shader.GetReference(), BaseIndex, UniformBuffer);
-	}
-
 	FORCEINLINE_DEBUGGABLE void SetShaderParameters(
 		FRHIComputeShader* InShader
 		, TConstArrayView<uint8> InParametersData
@@ -2908,22 +2778,6 @@ public:
 		Transition(Infos);
 	}
 
-	UE_DEPRECATED(5.1, "TransitionResourceArrayNoCopy is deprecated. Use Transition instead.")
-	FORCEINLINE_DEBUGGABLE void TransitionResourceArrayNoCopy(ERHIAccess TransitionType, TArray<FRHITexture*>& InTextures)
-	{
-	PRAGMA_DISABLE_DEPRECATION_WARNINGS
-		TransitionResources(TransitionType, &InTextures[0], InTextures.Num());
-	PRAGMA_ENABLE_DEPRECATION_WARNINGS
-	}
-
-	UE_DEPRECATED(5.1, "WaitComputeFence is deprecated. Use RHI transitions instead.")
-	FORCEINLINE_DEBUGGABLE void WaitComputeFence(FRHIComputeFence* WaitFence)
-	{
-		check(WaitFence->Transition);
-		EndTransitions(MakeArrayView(&WaitFence->Transition, 1));
-		WaitFence->Transition = nullptr;
-	}
-
 	FORCEINLINE_DEBUGGABLE void BeginUAVOverlap()
 	{
 		if (Bypass())
@@ -4484,24 +4338,6 @@ public:
 	{
 		return RHICreateUniformBuffer(Contents, Layout, Usage);
 	}
-
-	UE_DEPRECATED(5.1, "Use CreateBuffer and LockBuffer separately")
-	FORCEINLINE FBufferRHIRef CreateAndLockIndexBuffer(uint32 Stride, uint32 Size, EBufferUsageFlags InUsage, ERHIAccess InResourceState, FRHIResourceCreateInfo& CreateInfo, void*& OutDataBuffer)
-	{
-		FBufferRHIRef IndexBuffer = CreateBuffer(Size, InUsage | BUF_IndexBuffer, Stride, InResourceState, CreateInfo);
-		OutDataBuffer = LockBuffer(IndexBuffer, 0, Size, RLM_WriteOnly);
-		return IndexBuffer;
-	}
-
-	UE_DEPRECATED(5.1, "Use CreateBuffer and LockBuffer separately")
-	FORCEINLINE FBufferRHIRef CreateAndLockIndexBuffer(uint32 Stride, uint32 Size, EBufferUsageFlags InUsage, FRHIResourceCreateInfo& CreateInfo, void*& OutDataBuffer)
-	{
-		EBufferUsageFlags Usage = InUsage | BUF_IndexBuffer;
-		ERHIAccess ResourceState = RHIGetDefaultResourceState(Usage, true);
-	PRAGMA_DISABLE_DEPRECATION_WARNINGS
-		return CreateAndLockIndexBuffer(Stride, Size, Usage, ResourceState, CreateInfo, OutDataBuffer);
-	PRAGMA_ENABLE_DEPRECATION_WARNINGS
-	}
 	
 	FORCEINLINE void* LockStagingBuffer(FRHIStagingBuffer* StagingBuffer, FRHIGPUFence* Fence, uint32 Offset, uint32 SizeRHI)
 	{
@@ -4513,24 +4349,6 @@ public:
 		GDynamicRHI->UnlockStagingBuffer_RenderThread(*this, StagingBuffer);
 	}
 
-	UE_DEPRECATED(5.1, "Use CreateBuffer and LockBuffer separately")
-	FORCEINLINE FBufferRHIRef CreateAndLockVertexBuffer(uint32 Size, EBufferUsageFlags InUsage, ERHIAccess InResourceState, FRHIResourceCreateInfo& CreateInfo, void*& OutDataBuffer)
-	{
-		FBufferRHIRef VertexBuffer = CreateBuffer(Size, InUsage | BUF_VertexBuffer, 0, InResourceState, CreateInfo);
-		OutDataBuffer = LockBuffer(VertexBuffer, 0, Size, RLM_WriteOnly);
-		return VertexBuffer;
-	}
-
-	UE_DEPRECATED(5.1, "Use CreateBuffer and LockBuffer separately")
-	FORCEINLINE FBufferRHIRef CreateAndLockVertexBuffer(uint32 Size, EBufferUsageFlags InUsage, FRHIResourceCreateInfo& CreateInfo, void*& OutDataBuffer)
-	{
-		EBufferUsageFlags Usage = InUsage | BUF_VertexBuffer;
-		ERHIAccess ResourceState = RHIGetDefaultResourceState(Usage, true);
-	PRAGMA_DISABLE_DEPRECATION_WARNINGS
-		return CreateAndLockVertexBuffer(Size, Usage, ResourceState, CreateInfo, OutDataBuffer);
-	PRAGMA_ENABLE_DEPRECATION_WARNINGS
-	}
-
 	FORCEINLINE void CopyBuffer(FRHIBuffer* SourceBuffer, FRHIBuffer* DestBuffer)
 	{
 		QUICK_SCOPE_CYCLE_COUNTER(STAT_RHIMETHOD_CopyBuffer_Flush);
@@ -4606,30 +4424,6 @@ public:
 		return GDynamicRHI->CreateShaderResourceView_RenderThread(*this, Initializer);
 	}
 	
-	UE_DEPRECATED(5.1, "The CalcTexture... functions on the immediate RHI command list are deprecated. Use the global scope RHICalcTexturePlatformSize instead.")
-	FORCEINLINE uint64 CalcTexture2DPlatformSize(uint32 SizeX, uint32 SizeY, uint8 Format, uint32 NumMips, uint32 NumSamples, ETextureCreateFlags Flags, const FRHIResourceCreateInfo& CreateInfo, uint32& OutAlign)
-	{
-		PRAGMA_DISABLE_DEPRECATION_WARNINGS
-		return RHICalcTexture2DPlatformSize(SizeX, SizeY, Format, NumMips, NumSamples, Flags, CreateInfo, OutAlign);
-		PRAGMA_ENABLE_DEPRECATION_WARNINGS
-	}
-	
-	UE_DEPRECATED(5.1, "The CalcTexture... functions on the immediate RHI command list are deprecated. Use the global scope RHICalcTexturePlatformSize instead.")
-	FORCEINLINE uint64 CalcTexture3DPlatformSize(uint32 SizeX, uint32 SizeY, uint32 SizeZ, uint8 Format, uint32 NumMips, ETextureCreateFlags Flags, const FRHIResourceCreateInfo& CreateInfo, uint32& OutAlign)
-	{
-		PRAGMA_DISABLE_DEPRECATION_WARNINGS
-		return RHICalcTexture3DPlatformSize(SizeX, SizeY, SizeZ, Format, NumMips, Flags, CreateInfo, OutAlign);
-		PRAGMA_ENABLE_DEPRECATION_WARNINGS
-	}
-	
-	UE_DEPRECATED(5.1, "The CalcTexture... functions on the immediate RHI command list are deprecated. Use the global scope RHICalcTexturePlatformSize instead.")
-	FORCEINLINE uint64 CalcTextureCubePlatformSize(uint32 Size, uint8 Format, uint32 NumMips, ETextureCreateFlags Flags, const FRHIResourceCreateInfo& CreateInfo, uint32& OutAlign)
-	{
-		PRAGMA_DISABLE_DEPRECATION_WARNINGS
-		return RHICalcTextureCubePlatformSize(Size, Format, NumMips, Flags, CreateInfo, OutAlign);
-		PRAGMA_ENABLE_DEPRECATION_WARNINGS
-	}
-	
 	UE_DEPRECATED(5.2, "Use the global scope RHIGetTextureMemoryStats function.")
 	FORCEINLINE void GetTextureMemoryStats(FTextureMemoryStats& OutStats)
 	{
@@ -4942,12 +4736,6 @@ public:
 		return GDynamicRHI->RHIIsRenderingSuspended();
 	}
 
-	UE_DEPRECATED(5.1, "No longer used: FCompression::UncompressMemory should be used instead")
-	FORCEINLINE bool EnqueueDecompress(uint8_t* SrcBuffer, uint8_t* DestBuffer, int CompressedSize, void* ErrorCodeBuffer)
-	{
-		return false;
-	}
-
 	UE_DEPRECATED(5.2, "Use the global scope RHIGetAvailableResolutions function.")
 	FORCEINLINE bool GetAvailableResolutions(FScreenResolutionArray& Resolutions, bool bIgnoreRefreshRate)
 	{
@@ -5377,13 +5165,6 @@ FORCEINLINE FBufferRHIRef RHICreateIndexBuffer(uint32 Stride, uint32 Size, EBuff
 	return FRHICommandListExecutor::GetImmediateCommandList().CreateBuffer(Size, Usage | EBufferUsageFlags::IndexBuffer, Stride, ResourceState, CreateInfo);
 }
 
-UE_DEPRECATED(5.1, "RHIAsyncCreateIndexBuffer is deprecated. Use FRHICommandList::CreateBuffer instead.")
-FORCEINLINE FBufferRHIRef RHIAsyncCreateIndexBuffer(uint32 Stride, uint32 Size, EBufferUsageFlags Usage, ERHIAccess ResourceState, FRHIResourceCreateInfo& CreateInfo)
-{
-	check(IsInRenderingThread());
-	return FRHICommandListExecutor::GetImmediateCommandList().CreateBuffer(Size, Usage, Stride, ResourceState, CreateInfo);
-}
-
 FORCEINLINE FBufferRHIRef RHICreateIndexBuffer(uint32 Stride, uint32 Size, EBufferUsageFlags InUsage, FRHIResourceCreateInfo& CreateInfo)
 {
 	bool bHasInitialData = CreateInfo.BulkData != nullptr;
@@ -5392,17 +5173,6 @@ FORCEINLINE FBufferRHIRef RHICreateIndexBuffer(uint32 Stride, uint32 Size, EBuff
 	return RHICreateIndexBuffer(Stride, Size, Usage, ResourceState, CreateInfo);
 }
 
-UE_DEPRECATED(5.1, "RHIAsyncCreateIndexBuffer is deprecated. Use FRHICommandList::CreateBuffer instead.")
-FORCEINLINE FBufferRHIRef RHIAsyncCreateIndexBuffer(uint32 Stride, uint32 Size, EBufferUsageFlags InUsage, FRHIResourceCreateInfo& CreateInfo)
-{
-	bool bHasInitialData = CreateInfo.BulkData != nullptr;
-	EBufferUsageFlags Usage = InUsage | EBufferUsageFlags::IndexBuffer;
-	ERHIAccess ResourceState = RHIGetDefaultResourceState(Usage, bHasInitialData);
-PRAGMA_DISABLE_DEPRECATION_WARNINGS
-	return RHIAsyncCreateIndexBuffer(Stride, Size, Usage, ResourceState, CreateInfo);
-PRAGMA_ENABLE_DEPRECATION_WARNINGS
-}
-
 FORCEINLINE void RHIUpdateUniformBuffer(FRHIUniformBuffer* UniformBufferRHI, const void* Contents)
 {
 	return FRHICommandListExecutor::GetImmediateCommandList().UpdateUniformBuffer(UniformBufferRHI, Contents);
@@ -5414,13 +5184,6 @@ FORCEINLINE FBufferRHIRef RHICreateVertexBuffer(uint32 Size, EBufferUsageFlags U
 	return FRHICommandListExecutor::GetImmediateCommandList().CreateBuffer(Size, Usage | EBufferUsageFlags::VertexBuffer, 0, ResourceState, CreateInfo);
 }
 
-UE_DEPRECATED(5.1, "RHIAsyncCreateVertexBuffer is deprecated. Use FRHICommandList::CreateBuffer instead.")
-FORCEINLINE FBufferRHIRef RHIAsyncCreateVertexBuffer(uint32 Size, EBufferUsageFlags Usage, ERHIAccess ResourceState, FRHIResourceCreateInfo& CreateInfo)
-{
-	check(IsInRenderingThread());
-	return FRHICommandListExecutor::GetImmediateCommandList().CreateBuffer(Size, Usage, 0, ResourceState, CreateInfo);
-}
-
 FORCEINLINE FBufferRHIRef RHICreateVertexBuffer(uint32 Size, EBufferUsageFlags InUsage, FRHIResourceCreateInfo& CreateInfo)
 {
 	bool bHasInitialData = CreateInfo.BulkData != nullptr;
@@ -5429,17 +5192,6 @@ FORCEINLINE FBufferRHIRef RHICreateVertexBuffer(uint32 Size, EBufferUsageFlags I
 	return RHICreateVertexBuffer(Size, Usage, ResourceState, CreateInfo);
 }
 
-UE_DEPRECATED(5.1, "RHIAsyncCreateVertexBuffer is deprecated. Use FRHICommandList::CreateBuffer instead.")
-FORCEINLINE FBufferRHIRef RHIAsyncCreateVertexBuffer(uint32 Size, EBufferUsageFlags InUsage, FRHIResourceCreateInfo& CreateInfo)
-{
-	bool bHasInitialData = CreateInfo.BulkData != nullptr;
-	EBufferUsageFlags Usage = InUsage | EBufferUsageFlags::VertexBuffer;
-	ERHIAccess ResourceState = RHIGetDefaultResourceState(Usage, bHasInitialData);
-PRAGMA_DISABLE_DEPRECATION_WARNINGS
-	return RHIAsyncCreateVertexBuffer(Size, Usage, ResourceState, CreateInfo);
-PRAGMA_ENABLE_DEPRECATION_WARNINGS
-}
-
 FORCEINLINE FBufferRHIRef RHICreateStructuredBuffer(uint32 Stride, uint32 Size, EBufferUsageFlags Usage, ERHIAccess ResourceState, FRHIResourceCreateInfo& CreateInfo)
 {
 	check(IsInRenderingThread());
@@ -5537,32 +5289,6 @@ FORCEINLINE FTextureRHIRef RHICreateTexture(const FRHITextureCreateDesc& CreateD
 	return GDynamicRHI->RHICreateTexture_RenderThread(RHICmdList, CreateDesc);
 }
 
-UE_DEPRECATED(5.1, "FRHITexture2D is deprecated, please use RHICreateTexture(const FRHITextureCreateDesc&).")
-FORCEINLINE FTexture2DRHIRef RHICreateTexture2D(uint32 SizeX, uint32 SizeY, uint8 Format, uint32 NumMips, uint32 NumSamples, ETextureCreateFlags Flags, ERHIAccess ResourceState, const FRHIResourceCreateInfo& CreateInfo)
-{
-	return RHICreateTexture(
-		FRHITextureCreateDesc::Create2D(CreateInfo.DebugName)
-			.SetExtent((int32)SizeX, (int32)SizeY)
-			.SetFormat((EPixelFormat)Format)
-			.SetNumMips((uint8)NumMips)
-			.SetNumSamples((uint8)NumSamples)
-			.SetFlags(Flags)
-			.SetInitialState(ResourceState)
-			.SetExtData(CreateInfo.ExtData)
-			.SetBulkData(CreateInfo.BulkData)
-			.SetGPUMask(CreateInfo.GPUMask)
-			.SetClearValue(CreateInfo.ClearValueBinding)
-	);
-}
-
-UE_DEPRECATED(5.1, "FRHITexture2D is deprecated, please use RHICreateTexture(const FRHITextureCreateDesc&).")
-FORCEINLINE FTexture2DRHIRef RHICreateTextureExternal2D(uint32 SizeX, uint32 SizeY, uint8 Format, uint32 NumMips, uint32 NumSamples, ETextureCreateFlags Flags, ERHIAccess ResourceState, const  FRHIResourceCreateInfo& CreateInfo)
-{
-PRAGMA_DISABLE_DEPRECATION_WARNINGS
-	return RHICreateTexture2D(SizeX, SizeY, Format, NumMips, NumSamples, Flags | ETextureCreateFlags::External, ResourceState, CreateInfo);
-PRAGMA_ENABLE_DEPRECATION_WARNINGS
-}
-
 //UE_DEPRECATED(5.1, "FRHITexture2D is deprecated, please use RHICreateTexture(const FRHITextureCreateDesc&).")
 FORCEINLINE FTexture2DRHIRef RHIAsyncCreateTexture2D(uint32 SizeX, uint32 SizeY, uint8 Format, uint32 NumMips, ETextureCreateFlags Flags, ERHIAccess InResourceState, void** InitialMipData, uint32 NumInitialMips)
 {
@@ -5571,94 +5297,6 @@ FORCEINLINE FTexture2DRHIRef RHIAsyncCreateTexture2D(uint32 SizeX, uint32 SizeY,
 	return GDynamicRHI->RHIAsyncCreateTexture2D(SizeX, SizeY, Format, NumMips, Flags, ResourceState, InitialMipData, NumInitialMips);
 }
 
-UE_DEPRECATED(5.1, "FRHITexture2DArray is deprecated, please use RHICreateTexture(const FRHITextureCreateDesc&).")
-FORCEINLINE FTexture2DArrayRHIRef RHICreateTexture2DArray(uint32 SizeX, uint32 SizeY, uint32 ArraySize, uint8 Format, uint32 NumMips, uint32 NumSamples, ETextureCreateFlags Flags, ERHIAccess ResourceState, const FRHIResourceCreateInfo& CreateInfo)
-{
-	return RHICreateTexture(
-		FRHITextureCreateDesc::Create2DArray(CreateInfo.DebugName)
-			.SetExtent((int32)SizeX, (int32)SizeY)
-			.SetArraySize((uint16)ArraySize)
-			.SetFormat((EPixelFormat)Format)
-			.SetNumMips((uint8)NumMips)
-			.SetNumSamples((uint8)NumSamples)
-			.SetFlags(Flags)
-			.SetInitialState(ResourceState)
-			.SetExtData(CreateInfo.ExtData)
-			.SetBulkData(CreateInfo.BulkData)
-			.SetGPUMask(CreateInfo.GPUMask)
-			.SetClearValue(CreateInfo.ClearValueBinding)
-	);
-}
-
-UE_DEPRECATED(5.1, "FRHITexture3D is deprecated, please use RHICreateTexture(const FRHITextureCreateDesc&).")
-FORCEINLINE FTexture3DRHIRef RHICreateTexture3D(uint32 SizeX, uint32 SizeY, uint32 SizeZ, uint8 Format, uint32 NumMips, ETextureCreateFlags Flags, ERHIAccess ResourceState, const FRHIResourceCreateInfo& CreateInfo)
-{
-	return RHICreateTexture(
-		FRHITextureCreateDesc::Create3D(CreateInfo.DebugName)
-			.SetExtent((int32)SizeX, (int32)SizeY)
-			.SetDepth((uint16)SizeZ)
-			.SetFormat((EPixelFormat)Format)
-			.SetNumMips((uint8)NumMips)
-			.SetFlags(Flags)
-			.SetInitialState(ResourceState)
-			.SetExtData(CreateInfo.ExtData)
-			.SetBulkData(CreateInfo.BulkData)
-			.SetGPUMask(CreateInfo.GPUMask)
-			.SetClearValue(CreateInfo.ClearValueBinding)
-	);
-}
-
-UE_DEPRECATED(5.1, "FRHITextureCube is deprecated, please use RHICreateTexture(const FRHITextureCreateDesc&).")
-FORCEINLINE FTextureCubeRHIRef RHICreateTextureCube(uint32 Size, uint8 Format, uint32 NumMips, ETextureCreateFlags Flags, ERHIAccess ResourceState, const FRHIResourceCreateInfo& CreateInfo)
-{
-	return RHICreateTexture(
-		FRHITextureCreateDesc::CreateCube(CreateInfo.DebugName)
-			.SetExtent(Size)
-			.SetFormat((EPixelFormat)Format)
-			.SetNumMips((uint8)NumMips)
-			.SetFlags(Flags)
-			.SetInitialState(ResourceState)
-			.SetExtData(CreateInfo.ExtData)
-			.SetBulkData(CreateInfo.BulkData)
-			.SetGPUMask(CreateInfo.GPUMask)
-			.SetClearValue(CreateInfo.ClearValueBinding)
-	);
-}
-
-UE_DEPRECATED(5.1, "FRHITextureCube is deprecated, please use RHICreateTexture(const FRHITextureCreateDesc&).")
-FORCEINLINE FTextureCubeRHIRef RHICreateTextureCubeArray(uint32 Size, uint32 ArraySize, uint8 Format, uint32 NumMips, ETextureCreateFlags Flags, ERHIAccess ResourceState, const FRHIResourceCreateInfo& CreateInfo)
-{
-	return RHICreateTexture(
-		FRHITextureCreateDesc::CreateCubeArray(CreateInfo.DebugName)
-			.SetExtent(Size)
-			.SetArraySize((uint16)ArraySize)
-			.SetFormat((EPixelFormat)Format)
-			.SetNumMips((uint8)NumMips)
-			.SetFlags(Flags)
-			.SetInitialState(ResourceState)
-			.SetExtData(CreateInfo.ExtData)
-			.SetBulkData(CreateInfo.BulkData)
-			.SetGPUMask(CreateInfo.GPUMask)
-			.SetClearValue(CreateInfo.ClearValueBinding)
-	);
-}
-
-UE_DEPRECATED(5.1, "FRHITexture2D is deprecated, please use RHICreateTexture(const FRHITextureCreateDesc&).")
-FORCEINLINE FTexture2DRHIRef RHICreateTexture2D(uint32 SizeX, uint32 SizeY, uint8 Format, uint32 NumMips, uint32 NumSamples, ETextureCreateFlags Flags, const FRHIResourceCreateInfo& CreateInfo)
-{
-PRAGMA_DISABLE_DEPRECATION_WARNINGS
-	return RHICreateTexture2D(SizeX, SizeY, Format, NumMips, NumSamples, Flags, ERHIAccess::Unknown, CreateInfo);
-PRAGMA_ENABLE_DEPRECATION_WARNINGS
-}
-
-UE_DEPRECATED(5.1, "FRHITexture2D is deprecated, please use RHICreateTexture(const FRHITextureCreateDesc&).")
-FORCEINLINE FTexture2DRHIRef RHICreateTextureExternal2D(uint32 SizeX, uint32 SizeY, uint8 Format, uint32 NumMips, uint32 NumSamples, ETextureCreateFlags Flags, const FRHIResourceCreateInfo& CreateInfo)
-{
-PRAGMA_DISABLE_DEPRECATION_WARNINGS
-	return RHICreateTextureExternal2D(SizeX, SizeY, Format, NumMips, NumSamples, Flags, ERHIAccess::Unknown, CreateInfo);
-PRAGMA_ENABLE_DEPRECATION_WARNINGS
-}
-
 //UE_DEPRECATED(5.1, "FRHITexture2D is deprecated, please use RHICreateTexture(const FRHITextureCreateDesc&).")
 FORCEINLINE FTexture2DRHIRef RHIAsyncCreateTexture2D(uint32 SizeX, uint32 SizeY, uint8 Format, uint32 NumMips, ETextureCreateFlags Flags, void** InitialMipData, uint32 NumInitialMips)
 {
@@ -5667,38 +5305,6 @@ PRAGMA_DISABLE_DEPRECATION_WARNINGS
 PRAGMA_ENABLE_DEPRECATION_WARNINGS
 }
 
-UE_DEPRECATED(5.1, "FRHITexture2DArray is deprecated, please use RHICreateTexture(const FRHITextureCreateDesc&).")
-FORCEINLINE FTexture2DArrayRHIRef RHICreateTexture2DArray(uint32 SizeX, uint32 SizeY, uint32 SizeZ, uint8 Format, uint32 NumMips, uint32 NumSamples, ETextureCreateFlags Flags, const FRHIResourceCreateInfo& CreateInfo)
-{
-PRAGMA_DISABLE_DEPRECATION_WARNINGS
-	return RHICreateTexture2DArray(SizeX, SizeY, SizeZ, Format, NumMips, NumSamples, Flags, ERHIAccess::Unknown, CreateInfo);
-PRAGMA_ENABLE_DEPRECATION_WARNINGS
-}
-
-UE_DEPRECATED(5.1, "FRHITexture3D is deprecated, please use RHICreateTexture(const FRHITextureCreateDesc&).")
-FORCEINLINE FTexture3DRHIRef RHICreateTexture3D(uint32 SizeX, uint32 SizeY, uint32 SizeZ, uint8 Format, uint32 NumMips, ETextureCreateFlags Flags, const FRHIResourceCreateInfo& CreateInfo)
-{
-PRAGMA_DISABLE_DEPRECATION_WARNINGS
-	return RHICreateTexture3D(SizeX, SizeY, SizeZ, Format, NumMips, Flags, ERHIAccess::Unknown, CreateInfo);
-PRAGMA_ENABLE_DEPRECATION_WARNINGS
-}
-
-UE_DEPRECATED(5.1, "FRHITextureCube is deprecated, please use RHICreateTexture(const FRHITextureCreateDesc&).")
-FORCEINLINE FTextureCubeRHIRef RHICreateTextureCube(uint32 Size, uint8 Format, uint32 NumMips, ETextureCreateFlags Flags, const FRHIResourceCreateInfo& CreateInfo)
-{
-PRAGMA_DISABLE_DEPRECATION_WARNINGS
-	return RHICreateTextureCube(Size, Format, NumMips, Flags, ERHIAccess::Unknown, CreateInfo);
-PRAGMA_ENABLE_DEPRECATION_WARNINGS
-}
-
-UE_DEPRECATED(5.1, "FRHITextureCube is deprecated, please use RHICreateTexture(const FRHITextureCreateDesc&).")
-FORCEINLINE FTextureCubeRHIRef RHICreateTextureCubeArray(uint32 Size, uint32 ArraySize, uint8 Format, uint32 NumMips, ETextureCreateFlags Flags, const FRHIResourceCreateInfo& CreateInfo)
-{
-PRAGMA_DISABLE_DEPRECATION_WARNINGS
-	return RHICreateTextureCubeArray(Size, ArraySize, Format, NumMips, Flags, ERHIAccess::Unknown, CreateInfo);
-PRAGMA_ENABLE_DEPRECATION_WARNINGS
-}
-
 UE_DEPRECATED(5.2, "RHICopySharedMips should be replaced with UE::RHI::CopySharedMips(). ** Note ** The argument ordering of the new function has changed. The new function also does not issue RHI transitions, so the caller is responsible for ensuring src/dst textures are in CopySrc/CopyDest state.")
 FORCEINLINE void RHICopySharedMips(FRHITexture* DestTexture, FRHITexture* SrcTexture)
 {
diff --git a/Engine/Source/Runtime/RHI/Public/RHICommandListCommandExecutes.inl b/Engine/Source/Runtime/RHI/Public/RHICommandListCommandExecutes.inl
--- a/Engine/Source/Runtime/RHI/Public/RHICommandListCommandExecutes.inl
+++ b/Engine/Source/Runtime/RHI/Public/RHICommandListCommandExecutes.inl
@@ -457,39 +457,6 @@ void FRHICommandSetStaticUniformBuffers::Execute(FRHICommandListBase& CmdList)
 	INTERNAL_DECORATOR_COMPUTE(RHISetStaticUniformBuffers)(UniformBuffers);
 }
 
-void FRHICommandBuildLocalUniformBuffer::Execute(FRHICommandListBase& CmdList)
-{
-	LLM_SCOPE(ELLMTag::Shaders);
-	RHISTAT(BuildLocalUniformBuffer);
-	check(!IsValidRef(WorkArea.ComputedUniformBuffer->UniformBuffer)); // should not already have been created
-	check(WorkArea.Layout);
-	check(WorkArea.Contents); 
-	if (WorkArea.ComputedUniformBuffer->UseCount)
-	{
-		WorkArea.ComputedUniformBuffer->UniformBuffer = GDynamicRHI->RHICreateUniformBuffer(WorkArea.Contents, WorkArea.Layout, UniformBuffer_SingleFrame, EUniformBufferValidation::ValidateResources);
-	}
-	WorkArea.Layout = nullptr;
-	WorkArea.Contents = nullptr;
-}
-
-template <typename TRHIShader>
-void FRHICommandSetLocalUniformBuffer<TRHIShader>::Execute(FRHICommandListBase& CmdList)
-{
-	RHISTAT(SetLocalUniformBuffer);
-	check(LocalUniformBuffer.WorkArea->ComputedUniformBuffer->UseCount > 0 && IsValidRef(LocalUniformBuffer.WorkArea->ComputedUniformBuffer->UniformBuffer)); // this should have been created and should have uses outstanding
-	INTERNAL_DECORATOR(RHISetShaderUniformBuffer)(Shader, BaseIndex, LocalUniformBuffer.WorkArea->ComputedUniformBuffer->UniformBuffer);
-	if (--LocalUniformBuffer.WorkArea->ComputedUniformBuffer->UseCount == 0)
-	{
-		LocalUniformBuffer.WorkArea->ComputedUniformBuffer->~FComputedUniformBuffer();
-	}
-}
-
-template struct FRHICommandSetLocalUniformBuffer<FRHIVertexShader>;
-template struct FRHICommandSetLocalUniformBuffer<FRHIGeometryShader>;
-template struct FRHICommandSetLocalUniformBuffer<FRHIPixelShader>;
-template struct FRHICommandSetLocalUniformBuffer<FRHIComputeShader>;
-template struct FRHICommandSetLocalUniformBuffer<FRHIMeshShader>;
-
 void FRHICommandBeginRenderQuery::Execute(FRHICommandListBase& CmdList)
 {
 	RHISTAT(BeginRenderQuery);
diff --git a/Engine/Source/Runtime/RHI/Public/RHIResources.h b/Engine/Source/Runtime/RHI/Public/RHIResources.h
--- a/Engine/Source/Runtime/RHI/Public/RHIResources.h
+++ b/Engine/Source/Runtime/RHI/Public/RHIResources.h
@@ -1412,93 +1412,6 @@ private:
 /** Descriptor used to create a texture resource */
 struct RHI_API FRHITextureDesc
 {
-	UE_DEPRECATED(5.1, "FRHITextureDesc Create functions have been moved to FRHITextureCreateDesc.")
-	static FRHITextureDesc Create2D(
-		  FIntPoint           Size
-		, EPixelFormat        Format
-		, FClearValueBinding  ClearValue
-		, ETextureCreateFlags Flags
-		, uint8               NumMips    = 1
-		, uint8               NumSamples = 1
-		, uint32              ExtData    = 0
-		)
-	{
-		const uint16 Depth     = 1;
-		const uint16 ArraySize = 1;
-		return FRHITextureDesc(ETextureDimension::Texture2D, Flags, Format, ClearValue, { Size.X, Size.Y }, Depth, ArraySize, NumMips, NumSamples, ExtData);
-	}
-
-	UE_DEPRECATED(5.1, "FRHITextureDesc Create functions have been moved to FRHITextureCreateDesc.")
-	static FRHITextureDesc Create2DArray(
-		  FIntPoint           Size
-		, EPixelFormat        Format
-		, FClearValueBinding  ClearValue
-		, ETextureCreateFlags Flags
-		, uint16              ArraySize
-		, uint8               NumMips    = 1
-		, uint8               NumSamples = 1
-		, uint32              ExtData    = 0
-		)
-	{
-		const uint16 Depth   = 1;
-		return FRHITextureDesc(ETextureDimension::Texture2DArray, Flags, Format, ClearValue, { Size.X, Size.Y }, Depth, ArraySize, NumMips, NumSamples, ExtData);
-	}
-
-	UE_DEPRECATED(5.1, "FRHITextureDesc Create functions have been moved to FRHITextureCreateDesc.")
-	static FRHITextureDesc Create3D(
-		  FIntVector          Size
-		, EPixelFormat        Format
-		, FClearValueBinding  ClearValue
-		, ETextureCreateFlags Flags
-		, uint8               NumMips    = 1
-		, uint32              ExtData    = 0
-		)
-	{
-		const uint16 Depth = (uint16)Size.Z;
-		const uint16 ArraySize  = 1;
-		const uint16 LocalNumSamples = 1;
-
-		checkf(Size.Z <= TNumericLimits<decltype(FRHITextureDesc::Depth)>::Max(), TEXT("Depth parameter (Size.Z) exceeds valid range"));
-
-		return FRHITextureDesc(ETextureDimension::Texture3D, Flags, Format, ClearValue, { Size.X, Size.Y }, Depth, ArraySize, NumMips, LocalNumSamples, ExtData);
-	}
-
-	UE_DEPRECATED(5.1, "FRHITextureDesc Create functions have been moved to FRHITextureCreateDesc.")
-	static FRHITextureDesc CreateCube(
-		  uint32              Size
-		, EPixelFormat        Format
-		, FClearValueBinding  ClearValue
-		, ETextureCreateFlags Flags
-		, uint8               NumMips    = 1
-		, uint8               NumSamples = 1
-		, uint32              ExtData    = 0
-		)
-	{
-		checkf(Size <= (uint32)TNumericLimits<int32>::Max(), TEXT("Size parameter exceeds valid range"));
-
-		const uint16 Depth     = 1;
-		const uint16 ArraySize = 1;
-		return FRHITextureDesc(ETextureDimension::TextureCube, Flags, Format, ClearValue, { (int32)Size, (int32)Size }, Depth, ArraySize, NumMips, NumSamples, ExtData);
-	}
-
-	UE_DEPRECATED(5.1, "FRHITextureDesc Create functions have been moved to FRHITextureCreateDesc.")
-	static FRHITextureDesc CreateCubeArray(
-		  uint32              Size
-		, EPixelFormat        Format
-		, FClearValueBinding  ClearValue
-		, ETextureCreateFlags Flags
-		, uint16              ArraySize
-		, uint8               NumMips    = 1
-		, uint8               NumSamples = 1
-		, uint32              ExtData    = 0
-		)
-	{
-		checkf(Size <= (uint32)TNumericLimits<int32>::Max(), TEXT("Size parameter exceeds valid range"));
-
-		const uint16 Depth   = 1;
-		return FRHITextureDesc(ETextureDimension::TextureCubeArray, Flags, Format, ClearValue, { (int32)Size, (int32)Size }, Depth, ArraySize, NumMips, NumSamples, ExtData);
-	}
-
 	FRHITextureDesc()
 		: NumSamples(1)
 		, Dimension(ETextureDimension::Texture2D)
@@ -1534,33 +1447,6 @@ struct RHI_API FRHITextureDesc
 		, Format    (InFormat    )
 	{}
 
-	UE_DEPRECATED(5.1, "Prefer using FRHITextureCreateDesc rather than this FRHITextureDesc constructor. Otherwise use the FRHITextureDesc constructor that does not take optional arguments.")
-	FRHITextureDesc(
-		  ETextureDimension   InDimension
-		, ETextureCreateFlags InFlags
-		, EPixelFormat        InFormat
-		, FIntPoint           InExtent
-		, FClearValueBinding  InClearValue
-		, uint16              InDepth      = 1
-		, uint16              InArraySize  = 1
-		, uint8               InNumMips    = 1
-		, uint8               InNumSamples = 1
-		, uint32              InExtData    = 0
-		)
-		: FRHITextureDesc(
-		  InDimension
-		, InFlags
-		, InFormat
-		, InClearValue
-		, InExtent
-		, InDepth
-		, InArraySize
-		, InNumMips
-		, InNumSamples
-		, InExtData
-		)
-	{}
-
 	friend uint32 GetTypeHash(const FRHITextureDesc& Desc)
 	{
 		uint32 Hash = GetTypeHash(Desc.Dimension);
@@ -3826,9 +3712,6 @@ struct FRHIRenderPassInfo
 	// Controls the area for a multisample resolve or raster UAV (i.e. no fixed-function targets) operation.
 	FResolveRect ResolveRect;
 
-	UE_DEPRECATED(5.1, "ResolveParameters is deprecated. Use ResolveRect")
-	FResolveParams ResolveParameters;
-
 	// Some RHIs can use a texture to control the sampling and/or shading resolution of different areas 
 	FTextureRHIRef ShadingRateTexture = nullptr;
 	EVRSRateCombiner ShadingRateTextureCombiner = VRSRB_Passthrough;
@@ -3843,11 +3726,9 @@ struct FRHIRenderPassInfo
 	// Hint for some RHI's that renderpass will have specific sub-passes 
 	ESubpassHint SubpassHint = ESubpassHint::None;
 
-PRAGMA_DISABLE_DEPRECATION_WARNINGS
 	FRHIRenderPassInfo() = default;
 	FRHIRenderPassInfo(const FRHIRenderPassInfo&) = default;
 	FRHIRenderPassInfo& operator=(const FRHIRenderPassInfo&) = default;
-PRAGMA_ENABLE_DEPRECATION_WARNINGS
 
 	// Color, no depth, optional resolve, optional mip, optional array slice
 	explicit FRHIRenderPassInfo(FRHITexture* ColorRT, ERenderTargetActions ColorAction, FRHITexture* ResolveRT = nullptr, uint8 InMipIndex = 0, int32 InArraySlice = -1)
@@ -4095,22 +3976,6 @@ PRAGMA_ENABLE_DEPRECATION_WARNINGS
 	RHI_API void Validate() const {}
 #endif
 	RHI_API void ConvertToRenderTargetsInfo(FRHISetRenderTargetsInfo& OutRTInfo) const;
-
-	//////////////////////////////////////////////////////////////////////////
-	// Deprecated
-	//////////////////////////////////////////////////////////////////////////
-
-	UE_DEPRECATED(5.1, "IsMSAA is deprecated.")
-	inline bool IsMSAA() const { return false; }
-
-	UE_DEPRECATED(5.1, "bIsMSAA is deprecated")
-	bool bIsMSAA = false;
-
-	UE_DEPRECATED(5.1, "bTooManyUAVs is deprecated")
-	bool bTooManyUAVs = false;
-
-	UE_DEPRECATED(5.1, "bGeneratingMips is deprecated")
-	bool bGeneratingMips = false;
 };
 
 /** Used to specify a texture metadata plane when creating a view. */
diff --git a/Engine/Source/Runtime/RenderCore/Public/ShaderParameterUtils.h b/Engine/Source/Runtime/RenderCore/Public/ShaderParameterUtils.h
--- a/Engine/Source/Runtime/RenderCore/Public/ShaderParameterUtils.h
+++ b/Engine/Source/Runtime/RenderCore/Public/ShaderParameterUtils.h
@@ -634,27 +634,6 @@ inline void FRWShaderParameter::UnsetUAV(TRHICmdList& RHICmdList, FRHIComputeSha
 	SetUAVParameter(RHICmdList, ComputeShader,UAVParameter,FUnorderedAccessViewRHIRef());
 }
 
-
-/** Sets the value of a shader uniform buffer parameter to a uniform buffer containing the struct. */
-template<typename TShaderRHIRef, typename TRHICmdList>
-UE_DEPRECATED(5.1, "Local uniform buffers are now deprecated. Use SetUniformBufferParameter instead.")
-inline void SetLocalUniformBufferParameter(
-	TRHICmdList& RHICmdList,
-	const TShaderRHIRef& Shader,
-	const FShaderUniformBufferParameter& Parameter,
-	const FLocalUniformBuffer& LocalUniformBuffer
-	)
-{
-	// This will trigger if the parameter was not serialized
-	checkSlow(Parameter.IsInitialized());
-	if(Parameter.IsBound())
-	{
-	PRAGMA_DISABLE_DEPRECATION_WARNINGS
-		RHICmdList.SetLocalShaderUniformBuffer(Shader, Parameter.GetBaseIndex(), LocalUniformBuffer);
-	PRAGMA_ENABLE_DEPRECATION_WARNINGS
-	}
-}
-
 /** Sets the value of a shader uniform buffer parameter to a uniform buffer containing the struct. */
 template<typename TShaderRHIRef, typename TRHICmdList>
 inline void SetUniformBufferParameter(
