Fixing bad merge.

diff --git a/Engine/Source/Runtime/RHI/Public/RHICommandList.h b/Engine/Source/Runtime/RHI/Public/RHICommandList.h
--- a/Engine/Source/Runtime/RHI/Public/RHICommandList.h
+++ b/Engine/Source/Runtime/RHI/Public/RHICommandList.h
@@ -1917,11 +1917,13 @@ FRHICOMMAND_MACRO(FRHICommandSetStaticUniformBuffers)
 {
 	FUniformBufferStaticBindings UniformBuffers;
 
-	FORCEINLINE_DEBUGGABLE FRHICommandSetStaticUniformBuffers(const FUniformBufferStaticBindings& InUniformBuffers)
+	FORCEINLINE_DEBUGGABLE FRHICommandSetStaticUniformBuffers(const FUniformBufferStaticBindings & InUniformBuffers)
 		: UniformBuffers(InUniformBuffers)
 	{}
 
-	RHI_API void Execute(FRHICommandListBase& CmdList);
+	RHI_API void Execute(FRHICommandListBase & CmdList);
+};
+
 FRHICOMMAND_MACRO(FRHICommandBeginRenderQuery)
 {
 	FRHIRenderQuery* RenderQuery;
