Addressing potential issue where hand batched shader parameters wouldn't flush the automatic batching of shader parameters.

diff --git a/Engine/Source/Runtime/RHI/Private/RHICommandList.cpp b/Engine/Source/Runtime/RHI/Private/RHICommandList.cpp
--- a/Engine/Source/Runtime/RHI/Private/RHICommandList.cpp
+++ b/Engine/Source/Runtime/RHI/Private/RHICommandList.cpp
@@ -194,7 +194,7 @@ void FRHIParameterBatcher::OnBoundShaderChanged(FRHIComputeCommandList& InComman
 
 void FRHIParameterBatcher::PreDispatch(FRHIComputeCommandList& InCommandList)
 {
-	InCommandList.SetBatchedShaderParameters(GetBatchedComputeShader(), AllBatchedShaderParameters[SF_Compute]);
+	InCommandList.SetShaderParametersFromInternalBatched(GetBatchedComputeShader(), AllBatchedShaderParameters[SF_Compute]);
 	AllBatchedShaderParameters[SF_Compute].Reset();
 }
 
@@ -202,7 +202,7 @@ void FRHIParameterBatcher::PreDraw(FRHICommandList& InCommandList)
 {
 	for (uint32 Index = 0; Index < SF_NumGraphicsFrequencies; Index++)
 	{
-		InCommandList.SetBatchedShaderParameters(GetBatchedGraphicsShader(Index), AllBatchedShaderParameters[Index]);
+		InCommandList.SetShaderParametersFromInternalBatched(GetBatchedGraphicsShader(Index), AllBatchedShaderParameters[Index]);
 		AllBatchedShaderParameters[Index].Reset();
 	}
 }
@@ -223,7 +223,7 @@ void FRHIParameterBatcher::FlushPendingParameters(FRHIComputeCommandList& InComm
 	const EShaderFrequency Frequency = InShader->GetFrequency();
 	if (Frequency == SF_Compute)
 	{
-		InCommandList.SetBatchedShaderParameters(GetBatchedComputeShader(), AllBatchedShaderParameters[SF_Compute]);
+		InCommandList.SetShaderParametersFromInternalBatched(GetBatchedComputeShader(), AllBatchedShaderParameters[SF_Compute]);
 		AllBatchedShaderParameters[SF_Compute].Reset();
 	}
 }
@@ -233,7 +233,7 @@ void FRHIParameterBatcher::FlushPendingParameters(FRHICommandList& InCommandList
 	const EShaderFrequency Frequency = InShader->GetFrequency();
 	if (Frequency < SF_NumGraphicsFrequencies)
 	{
-		InCommandList.SetBatchedShaderParameters(GetBatchedGraphicsShader(Frequency), AllBatchedShaderParameters[Frequency]);
+		InCommandList.SetShaderParametersFromInternalBatched(GetBatchedGraphicsShader(Frequency), AllBatchedShaderParameters[Frequency]);
 		AllBatchedShaderParameters[Frequency].Reset();
 	}
 }
diff --git a/Engine/Source/Runtime/RHI/Public/RHICommandList.h b/Engine/Source/Runtime/RHI/Public/RHICommandList.h
--- a/Engine/Source/Runtime/RHI/Public/RHICommandList.h
+++ b/Engine/Source/Runtime/RHI/Public/RHICommandList.h
@@ -2405,6 +2405,21 @@ public:
 	}
 
 	FORCEINLINE_DEBUGGABLE void SetBatchedShaderParameters(FRHIComputeShader* InShader, const FRHIBatchedShaderParameters& InBatchedParameters)
+	{
+		if (InBatchedParameters.HasParameters())
+		{
+			SetShaderParameters(
+				InShader,
+				InBatchedParameters.ParametersData,
+				InBatchedParameters.Parameters,
+				InBatchedParameters.ResourceParameters,
+				InBatchedParameters.BindlessParameters,
+				true /* bSubmitBatchedParameters */
+			);
+		}
+	}
+
+	FORCEINLINE_DEBUGGABLE void SetShaderParametersFromInternalBatched(FRHIComputeShader* InShader, const FRHIBatchedShaderParameters& InBatchedParameters)
 	{
 		if (InBatchedParameters.HasParameters())
 		{
@@ -3187,6 +3202,23 @@ public:
 	using FRHIComputeCommandList::SetBatchedShaderParameters;
 
 	FORCEINLINE_DEBUGGABLE void SetBatchedShaderParameters(FRHIGraphicsShader* InShader, const FRHIBatchedShaderParameters& InBatchedParameters)
+	{
+		if (InBatchedParameters.HasParameters())
+		{
+			SetShaderParameters(
+				InShader,
+				InBatchedParameters.ParametersData,
+				InBatchedParameters.Parameters,
+				InBatchedParameters.ResourceParameters,
+				InBatchedParameters.BindlessParameters,
+				true /* bSubmitBatchedParameters */
+			);
+		}
+	}
+
+	using FRHIComputeCommandList::SetShaderParametersFromInternalBatched;
+
+	FORCEINLINE_DEBUGGABLE void SetShaderParametersFromInternalBatched(FRHIGraphicsShader* InShader, const FRHIBatchedShaderParameters& InBatchedParameters)
 	{
 		if (InBatchedParameters.HasParameters())
 		{
