Fix a bunch of reported misses in the manual PSO cache.

- DepthStencilFlags can be different between builds due to runtime flags, resulting in a spurious miss.
- PSOs precompiled from the cache are sometimes reported as missing from the cache (put a workaround in place).

Also, minor readability/debuggability fixes.

diff --git a/Engine/Source/Editor/UnrealEd/Private/Commandlets/ShaderPipelineCacheToolsCommandlet.cpp b/Engine/Source/Editor/UnrealEd/Private/Commandlets/ShaderPipelineCacheToolsCommandlet.cpp
--- a/Engine/Source/Editor/UnrealEd/Private/Commandlets/ShaderPipelineCacheToolsCommandlet.cpp
+++ b/Engine/Source/Editor/UnrealEd/Private/Commandlets/ShaderPipelineCacheToolsCommandlet.cpp
@@ -1870,6 +1870,7 @@ void FilterInvalidPSOs(TSet<FPipelineCacheFileFormatPSO>& InOutPSOs, const TMult
 		{
 			UE_LOG(LogShaderPipelineCacheTools, Display, TEXT("These vertex shaders are used with an inconsistent vertex format:"), SuspiciousVertexShaders.Num());
 			int32 SuspectVSIdx = 0;
+			const int32 kMaxSuspectToPrint = 50;
 			for (const FSHAHash& SuspectVS : SuspiciousVertexShaders)
 			{
 				const TArray<FStableShaderKeyAndValue>* Out = InverseMap.Find(SuspectVS);
@@ -1903,6 +1904,12 @@ void FilterInvalidPSOs(TSet<FPipelineCacheFileFormatPSO>& InOutPSOs, const TMult
 					UE_LOG(LogShaderPipelineCacheTools, Warning, TEXT("Unknown shader with a hash %s"), *SuspectVS.ToString());
 				}
 				++SuspectVSIdx;
+
+				if (SuspectVSIdx > kMaxSuspectToPrint)
+				{
+					UE_LOG(LogShaderPipelineCacheTools, Display, TEXT("... and %d more VS hashes"), SuspiciousVertexShaders.Num() - SuspectVSIdx - 1);
+					break;
+				}
 			}
 		}
 	}
diff --git a/Engine/Source/Runtime/RHI/Private/DynamicRHI.cpp b/Engine/Source/Runtime/RHI/Private/DynamicRHI.cpp
--- a/Engine/Source/Runtime/RHI/Private/DynamicRHI.cpp
+++ b/Engine/Source/Runtime/RHI/Private/DynamicRHI.cpp
@@ -645,7 +645,7 @@ bool FDynamicRHI::RHIMatchPrecachePSOInitializers(const FGraphicsPipelineStateIn
 		LHS.RenderTargetFormats != RHS.RenderTargetFormats ||
 		!FGraphicsPipelineStateInitializer::RelevantRenderTargetFlagsEqual(LHS.RenderTargetFlags, RHS.RenderTargetFlags) ||
 		LHS.DepthStencilTargetFormat != RHS.DepthStencilTargetFormat ||
-		LHS.DepthStencilTargetFlag != RHS.DepthStencilTargetFlag ||
+		!FGraphicsPipelineStateInitializer::RelevantDepthStencilFlagsEqual(LHS.DepthStencilTargetFlag, RHS.DepthStencilTargetFlag) ||
 		LHS.DepthTargetLoadAction != RHS.DepthTargetLoadAction ||
 		LHS.DepthTargetStoreAction != RHS.DepthTargetStoreAction ||
 		LHS.StencilTargetLoadAction != RHS.StencilTargetLoadAction ||
diff --git a/Engine/Source/Runtime/RHI/Private/PipelineFileCache.cpp b/Engine/Source/Runtime/RHI/Private/PipelineFileCache.cpp
--- a/Engine/Source/Runtime/RHI/Private/PipelineFileCache.cpp
+++ b/Engine/Source/Runtime/RHI/Private/PipelineFileCache.cpp
@@ -448,16 +448,26 @@ void FPipelineCacheFileFormatPSO::GraphicsDescriptor::AddShadersToReadableString
 		OutBuilder << TEXT(" VS:");
 		OutBuilder << VertexShader;
 	}
-	if (FragmentShader != FSHAHash())
+	if (MeshShader != FSHAHash())
 	{
-		OutBuilder << TEXT(" PS:");
-		OutBuilder << FragmentShader;
+		OutBuilder << TEXT(" MS:");
+		OutBuilder << MeshShader;
 	}
 	if (GeometryShader != FSHAHash())
 	{
 		OutBuilder << TEXT(" GS:");
 		OutBuilder << GeometryShader;
 	}
+	if (AmplificationShader != FSHAHash())
+	{
+		OutBuilder << TEXT(" AS:");
+		OutBuilder << AmplificationShader;
+	}
+	if (FragmentShader != FSHAHash())
+	{
+		OutBuilder << TEXT(" PS:");
+		OutBuilder << FragmentShader;
+	}
 }
 
 void FPipelineCacheFileFormatPSO::GraphicsDescriptor::ShadersFromString(const FStringView& Src)
@@ -562,12 +572,12 @@ void FPipelineCacheFileFormatPSO::GraphicsDescriptor::AddStateToReadableString(T
 	OutBuilder << DepthStencilState.ToString();
 	OutBuilder << TEXT("\n");
 
-	OutBuilder << TEXT(" MSAA:");
+	OutBuilder << TEXT(" NumMSAA:");
 	OutBuilder << MSAASamples;
 	OutBuilder << TEXT(" DSfmt:");
 	OutBuilder << uint32(DepthStencilFormat);
 	OutBuilder << TEXT(" DSflags:");
-	OutBuilder << uint32(DepthStencilFlags);
+	OutBuilder << uint64(DepthStencilFlags);
 	OutBuilder << TEXT("\n");
 
 	OutBuilder << TEXT(" DL:");
@@ -596,7 +606,7 @@ void FPipelineCacheFileFormatPSO::GraphicsDescriptor::AddStateToReadableString(T
 			OutBuilder << TEXT(":fmt=");
 			OutBuilder << uint32(RenderTargetFormats[Index]);
 			OutBuilder << TEXT(" flg=");
-			OutBuilder << uint32(RenderTargetFlags[Index]);
+			OutBuilder << uint64(RenderTargetFlags[Index]);
 		}
 		OutBuilder << TEXT("\n");
 	}
@@ -660,7 +670,9 @@ bool FPipelineCacheFileFormatPSO::GraphicsDescriptor::StateFromString(const FStr
 	check(PartEnd - PartIt >= 3); //not a very robust parser
 	LexFromString(MSAASamples, *PartIt++);
 	LexFromString((uint32&)DepthStencilFormat, *PartIt++);
-	LexFromString(DepthStencilFlags, *PartIt++);
+	ETextureCreateFlags DSFlags;
+	LexFromString(DSFlags, *PartIt++);
+	DepthStencilFlags = ReduceDSFlags(DSFlags);
 
 	check(PartEnd - PartIt >= 5); //not a very robust parser
 	LexFromString((uint32&)DepthLoad, *PartIt++);
@@ -762,6 +774,11 @@ ETextureCreateFlags FPipelineCacheFileFormatPSO::GraphicsDescriptor::ReduceRTFla
 	return (InFlags & FGraphicsPipelineStateInitializer::RelevantRenderTargetFlagMask);
 }
 
+ETextureCreateFlags FPipelineCacheFileFormatPSO::GraphicsDescriptor::ReduceDSFlags(ETextureCreateFlags InFlags)
+{
+	return (InFlags & FGraphicsPipelineStateInitializer::RelevantDepthStencilFlagMask);
+}
+
 FString FPipelineCacheFileFormatPSO::GraphicsDescriptor::StateHeaderLine()
 {
 	FString Result;
@@ -871,7 +888,7 @@ FString FPipelineCacheFileFormatPSO::CommonToString() const
 	return FString::Printf(TEXT("\"%d,%llu\""), Count, Mask);
 }
 
-FString FPipelineCacheFileFormatPSO::ToStringReadable()
+FString FPipelineCacheFileFormatPSO::ToStringReadable() const
 {
 	TReadableStringBuilder Builder;
 
@@ -1259,7 +1276,7 @@ bool FPipelineCacheFileFormatPSO::Verify() const
 				else
 				{
 					static_assert(sizeof(uint64) == sizeof(Info.GraphicsDesc.RenderTargetFlags[i]), "ETextureCreateFlags size changed, please change serialization");
-					uint64 RTFlags = static_cast<uint64>(Info.GraphicsDesc.RenderTargetFlags[i]);
+					uint64 RTFlags = static_cast<uint64>(FPipelineCacheFileFormatPSO::GraphicsDescriptor::ReduceRTFlags(Info.GraphicsDesc.RenderTargetFlags[i]));
 					Ar << RTFlags;
 					Info.GraphicsDesc.RenderTargetFlags[i] = FPipelineCacheFileFormatPSO::GraphicsDescriptor::ReduceRTFlags(static_cast<ETextureCreateFlags>(RTFlags));
 				}
@@ -1279,12 +1296,14 @@ bool FPipelineCacheFileFormatPSO::Verify() const
 			{
 				uint32 DepthStencilFlags = 0;
 				Ar << DepthStencilFlags;
-				Info.GraphicsDesc.DepthStencilFlags = static_cast<ETextureCreateFlags>(DepthStencilFlags);
+				Info.GraphicsDesc.DepthStencilFlags = FPipelineCacheFileFormatPSO::GraphicsDescriptor::ReduceDSFlags(static_cast<ETextureCreateFlags>(DepthStencilFlags));
 			}
 			else
 			{
 				static_assert(sizeof(uint64) == sizeof(Info.GraphicsDesc.DepthStencilFlags), "ETextureCreateFlags size changed, please change serialization");
-				Ar << Info.GraphicsDesc.DepthStencilFlags;
+				uint64 DepthStencilFlags = static_cast<uint64>(FPipelineCacheFileFormatPSO::GraphicsDescriptor::ReduceDSFlags(Info.GraphicsDesc.DepthStencilFlags));
+				Ar << DepthStencilFlags;
+				Info.GraphicsDesc.DepthStencilFlags = FPipelineCacheFileFormatPSO::GraphicsDescriptor::ReduceDSFlags(static_cast<ETextureCreateFlags>(DepthStencilFlags));
 			}
 			Ar << Info.GraphicsDesc.DepthLoad;
 			Ar << Info.GraphicsDesc.StencilLoad;
@@ -1475,7 +1494,7 @@ FPipelineCacheFileFormatPSO::FPipelineCacheFileFormatPSO()
 	PSO.GraphicsDesc.MSAASamples = Init.NumSamples;
 	
 	PSO.GraphicsDesc.DepthStencilFormat = Init.DepthStencilTargetFormat;
-	PSO.GraphicsDesc.DepthStencilFlags = Init.DepthStencilTargetFlag;
+	PSO.GraphicsDesc.DepthStencilFlags = FPipelineCacheFileFormatPSO::GraphicsDescriptor::ReduceDSFlags(Init.DepthStencilTargetFlag);
 	PSO.GraphicsDesc.DepthLoad = Init.DepthTargetLoadAction;
 	PSO.GraphicsDesc.StencilLoad = Init.StencilTargetLoadAction;
 	PSO.GraphicsDesc.DepthStore = Init.DepthTargetStoreAction;
@@ -1649,7 +1668,7 @@ struct FPipelineCacheFileFormatTOC
 	
 	FPipelineFileCacheManager::PSOOrder SortedOrder;
 	TMap<uint32, FPipelineCacheFileFormatPSOMetaData> MetaData;
-	void DumpToLog()
+	void DumpToLog() const
 	{
 		for (TMap<uint32, FPipelineCacheFileFormatPSOMetaData>::TConstIterator It(MetaData); It; ++It)
 		{
@@ -2137,7 +2156,7 @@ public:
 
 		UE_LOG(LogRHI, VeryVerbose, TEXT("-- opened bundled %s cache:"), *NameIn);
 		TOC.DumpToLog();
-		UE_LOG(LogRHI, VeryVerbose, TEXT("-- opened bundled %s cache:"), *NameIn);
+		UE_LOG(LogRHI, VeryVerbose, TEXT("-- end of dump (%s)"), *NameIn);
 
 		return bGameFileOk;
 	}
@@ -2737,6 +2756,7 @@ public:
 	bool IsPSOEntryCachedInternal(FPipelineCacheFileFormatPSO const& NewEntry, FPSOUsageData* EntryData = nullptr) const
 	{
 		uint32 PSOHash = GetTypeHash(NewEntry);
+		check(!EntryData || EntryData->PSOHash == PSOHash);
 		FPipelineCacheFileFormatPSOMetaData const * const Existing = TOC.MetaData.Find(PSOHash);
 		
 		if(Existing != nullptr && EntryData != nullptr)
@@ -3376,6 +3396,13 @@ void FPipelineFileCacheManager::CacheGraphicsPSO(uint32 RunTimeHash, FGraphicsPi
 				if (!FPipelineFileCacheManager::IsPSOEntryCached(NewEntry, &CurrentUsageData))
 				{
 					bool bActuallyNewPSO = !NewPSOHashes.Contains(PSOHash);
+					if (Initializer.bFromPSOFileCache)
+					{
+						// FIXME: this is a workaround. Needs proper investigation
+						UE_LOG(LogRHI, Warning, TEXT("PSO from the cache was not found in the cache! PSOHash: %u"), PSOHash);
+						bActuallyNewPSO = false;
+					}
+
 					if (bActuallyNewPSO && IsOpenGLPlatform(GMaxRHIShaderPlatform)) // OpenGL is a BSS platform and so we don't report BSS matches as missing.
 					{
 						bActuallyNewPSO = !FPipelineFileCacheManager::IsBSSEquivalentPSOEntryCached(NewEntry);
diff --git a/Engine/Source/Runtime/RHI/Public/PipelineFileCache.h b/Engine/Source/Runtime/RHI/Public/PipelineFileCache.h
--- a/Engine/Source/Runtime/RHI/Public/PipelineFileCache.h
+++ b/Engine/Source/Runtime/RHI/Public/PipelineFileCache.h
@@ -177,6 +177,9 @@ struct RHI_API FPipelineCacheFileFormatPSO
 
 		/** Not all RT flags make sense for the replayed PSO, only those that can influence the RT formats */
 		static ETextureCreateFlags ReduceRTFlags(ETextureCreateFlags InFlags);
+
+		/** Not all DepthStencil flags make sense for the replayed PSO, only those that can influence the actual format */
+		static ETextureCreateFlags ReduceDSFlags(ETextureCreateFlags InFlags);
 	};
 	struct RHI_API FPipelineFileCacheRayTracingDesc
 	{
@@ -248,7 +251,7 @@ struct RHI_API FPipelineCacheFileFormatPSO
 	void CommonFromString(const FStringView& Src);
 
 	/** Prints out human-readable representation of the PSO, for any type */
-	FString ToStringReadable();
+	FString ToStringReadable() const;
 	
 	// Potential cases for seperating verify logic if requiired: RunTime-Logging, RunTime-UserCaching, RunTime-PreCompile, CommandLet-Cooking
 	bool Verify() const;
diff --git a/Engine/Source/Runtime/RHI/Public/RHIResources.h b/Engine/Source/Runtime/RHI/Public/RHIResources.h
--- a/Engine/Source/Runtime/RHI/Public/RHIResources.h
+++ b/Engine/Source/Runtime/RHI/Public/RHIResources.h
@@ -3261,7 +3261,7 @@ public:
 			RenderTargetFormats != rhs.RenderTargetFormats || 
 			!RelevantRenderTargetFlagsEqual(RenderTargetFlags, rhs.RenderTargetFlags) || 
 			DepthStencilTargetFormat != rhs.DepthStencilTargetFormat || 
-			DepthStencilTargetFlag != rhs.DepthStencilTargetFlag ||
+			!RelevantDepthStencilFlagsEqual(DepthStencilTargetFlag, rhs.DepthStencilTargetFlag) ||
 			DepthTargetLoadAction != rhs.DepthTargetLoadAction ||
 			DepthTargetStoreAction != rhs.DepthTargetStoreAction ||
 			StencilTargetLoadAction != rhs.StencilTargetLoadAction ||
@@ -3282,6 +3282,10 @@ public:
 	// In most RHIs, the format is only influenced by TexCreate_SRGB. D3D12 additionally uses TexCreate_Shared in its format selection logic.
 	static constexpr ETextureCreateFlags RelevantRenderTargetFlagMask = ETextureCreateFlags::SRGB | ETextureCreateFlags::Shared;
 
+	// We care about flags that influence DS formats (which is the only thing the underlying API cares about).
+	// D3D12 shares the format choice function with the RT, so preserving all the flags used there out of abundance of caution.
+	static constexpr ETextureCreateFlags RelevantDepthStencilFlagMask = ETextureCreateFlags::SRGB | ETextureCreateFlags::Shared | ETextureCreateFlags::DepthStencilTargetable;
+
 	static bool RelevantRenderTargetFlagsEqual(const TRenderTargetFlags& A, const TRenderTargetFlags& B)
 	{
 		for (int32 Index = 0; Index < A.Num(); ++Index)
@@ -3296,6 +3300,13 @@ public:
 		return true;
 	}
 
+	static bool RelevantDepthStencilFlagsEqual(const ETextureCreateFlags A, const ETextureCreateFlags B)
+	{
+		ETextureCreateFlags FlagsA = (A & RelevantDepthStencilFlagMask);
+		ETextureCreateFlags FlagsB = (B & RelevantDepthStencilFlagMask);
+		return (FlagsA == FlagsB);
+	}
+
 	uint32 ComputeNumValidRenderTargets() const
 	{
 		// Get the count of valid render targets (ignore those at the end of the array with PF_Unknown)
