Addressing issues with the FRHIDispatchIndirectParameters size. Can cause a full shader recompile on some platforms.

diff --git a/Engine/Source/Runtime/RHI/Public/RHI.h b/Engine/Source/Runtime/RHI/Public/RHI.h
--- a/Engine/Source/Runtime/RHI/Public/RHI.h
+++ b/Engine/Source/Runtime/RHI/Public/RHI.h
@@ -1161,6 +1161,8 @@ struct FRHIDispatchIndirectParameters : public FRHIDispatchIndirectParametersNoP
 {
 #if PLATFORM_DISPATCH_INDIRECT_ARGUMENT_BOUNDARY_SIZE == 64
 	uint32 Padding;			// pad to 32 bytes to prevent crossing of 64 byte boundary in ExecuteIndirect calls
+#elif PLATFORM_DISPATCH_INDIRECT_ARGUMENT_BOUNDARY_SIZE == 128
+	uint32 Padding[5];		// pad to 64 bytes to prevent crossing of 128 byte boundary in ExecuteIndirect calls
 #elif PLATFORM_DISPATCH_INDIRECT_ARGUMENT_BOUNDARY_SIZE != 0
 	#error FRHIDispatchIndirectParameters does not account for PLATFORM_DISPATCH_INDIRECT_ARGUMENT_BOUNDARY_SIZE.
 #endif
