Fixing Android NonUnity errors.

diff --git a/Engine/Source/Runtime/RHI/Private/Android/AndroidDynamicRHI.cpp b/Engine/Source/Runtime/RHI/Private/Android/AndroidDynamicRHI.cpp
--- a/Engine/Source/Runtime/RHI/Private/Android/AndroidDynamicRHI.cpp
+++ b/Engine/Source/Runtime/RHI/Private/Android/AndroidDynamicRHI.cpp
@@ -3,6 +3,7 @@
 #include "RHI.h"
 #include "Modules/ModuleManager.h"
 #include "Android/AndroidApplication.h"
+#include "Misc/App.h"
 
 FDynamicRHI* PlatformCreateDynamicRHI()
 {
