Removing deprecated RHI symbols.

diff --git a/Engine/Plugins/VirtualProduction/TextureShare/Source/TextureShareCore/Private/Module/TextureShareCoreModule.cpp b/Engine/Plugins/VirtualProduction/TextureShare/Source/TextureShareCore/Private/Module/TextureShareCoreModule.cpp
--- a/Engine/Plugins/VirtualProduction/TextureShare/Source/TextureShareCore/Private/Module/TextureShareCoreModule.cpp
+++ b/Engine/Plugins/VirtualProduction/TextureShare/Source/TextureShareCore/Private/Module/TextureShareCoreModule.cpp
@@ -6,7 +6,7 @@
 #include "Core/TextureShareCore.h"
 
 #ifdef TEXTURESHARECORE_INITIALIZEVULKANEXTENSIONS
-#include "VulkanRHIBridge.h"
+#include "IVulkanDynamicRHI.h"
 #endif
 
 //////////////////////////////////////////////////////////////////////////////////////////////
@@ -45,8 +45,8 @@ void FTextureShareCoreModule::StartupModule()
 		const TArray<const ANSICHAR*> InstanceLayers = {
 		};
 
-		VulkanRHIBridge::AddEnabledDeviceExtensionsAndLayers(DeviceExtensions, DeviceLayers);
-		VulkanRHIBridge::AddEnabledInstanceExtensionsAndLayers(InstanceExtensions, InstanceLayers);
+		IVulkanDynamicRHI::AddEnabledDeviceExtensionsAndLayers(DeviceExtensions, DeviceLayers);
+		IVulkanDynamicRHI::AddEnabledInstanceExtensionsAndLayers(InstanceExtensions, InstanceLayers);
 	}
 #endif
 }
diff --git a/Engine/Source/Runtime/D3D12RHI/Public/D3D12RHIBridge.h b/Engine/Source/Runtime/D3D12RHI/Public/D3D12RHIBridge.h
--- a/Engine/Source/Runtime/D3D12RHI/Public/D3D12RHIBridge.h
+++ b/Engine/Source/Runtime/D3D12RHI/Public/D3D12RHIBridge.h
@@ -1,13 +1,5 @@
 // Copyright Epic Games, Inc. All Rights Reserved.
 
-#include "Templates/Function.h"
-#include "ID3D12DynamicRHI.h"
+#pragma once
 
-namespace D3D12RHI
-{
-	// Use these functions with care!
-
-	// Returns ID3D12GraphicsCommandList* and ID3D12CommandQueue*
-	UE_DEPRECATED(5.1, "ID3D12DynamicRHI should be used to acquire D3D12 command lists and command queues")
-	D3D12RHI_API void GetGfxCommandListAndQueue(FRHICommandList& RHICmdList, void*& OutGfxCmdList, void*& OutCommandQueue);
-}
+#pragma message("D3D12RHIBridge.h is deprecated. Please switch to using ID3D12DynamicRHI.h")
diff --git a/Engine/Source/Runtime/RHI/Private/RHIUtilities.cpp b/Engine/Source/Runtime/RHI/Private/RHIUtilities.cpp
--- a/Engine/Source/Runtime/RHI/Private/RHIUtilities.cpp
+++ b/Engine/Source/Runtime/RHI/Private/RHIUtilities.cpp
@@ -733,261 +733,3 @@ void DecodeRenderTargetMode(ESimpleRenderTargetMode Mode, ERenderTargetLoadActio
 		StencilStoreAction = ERenderTargetStoreAction::ENoAction;
 	}
 }
-
-void RHICreateTargetableShaderResource(
-	const FRHITextureCreateDesc& BaseDesc,
-	ETextureCreateFlags TargetableTextureFlags,
-	FTextureRHIRef& OutTargetableTexture)
-{
-	// Ensure none of the usage flags are passed in.
-	check(!EnumHasAnyFlags(BaseDesc.Flags, ETextureCreateFlags::RenderTargetable | ETextureCreateFlags::ResolveTargetable));
-
-	// Ensure that the targetable texture is either render or depth-stencil targetable.
-	check(EnumHasAnyFlags(TargetableTextureFlags, ETextureCreateFlags::RenderTargetable | ETextureCreateFlags::DepthStencilTargetable | ETextureCreateFlags::UAV));
-
-	FRHITextureCreateDesc Desc =
-		FRHITextureCreateDesc(BaseDesc)
-		.AddFlags(TargetableTextureFlags | ETextureCreateFlags::ShaderResource)
-		.SetInitialState(ERHIAccess::SRVMask);
-
-	OutTargetableTexture = RHICreateTexture(Desc);
-}
-
-void RHICreateTargetableShaderResource(
-	const FRHITextureCreateDesc& BaseDesc,
-	ETextureCreateFlags TargetableTextureFlags,
-	bool bForceSeparateTargetAndShaderResource,
-	bool bForceSharedTargetAndShaderResource,
-	FTextureRHIRef& OutTargetableTexture,
-	FTextureRHIRef& OutShaderResourceTexture
-	)
-{
-	// Ensure none of the usage flags are passed in.
-	check(!EnumHasAnyFlags(BaseDesc.Flags, ETextureCreateFlags::RenderTargetable | ETextureCreateFlags::ResolveTargetable));
-
-	// Ensure we aren't forcing separate and shared textures at the same time.
-	check(!(bForceSeparateTargetAndShaderResource && bForceSharedTargetAndShaderResource));
-
-	// Ensure that the targetable texture is either render or depth-stencil targetable.
-	check(EnumHasAnyFlags(TargetableTextureFlags, ETextureCreateFlags::RenderTargetable | ETextureCreateFlags::DepthStencilTargetable | ETextureCreateFlags::UAV));
-
-	if (BaseDesc.NumSamples > 1 && !bForceSharedTargetAndShaderResource)
-	{
-		bForceSeparateTargetAndShaderResource = RHISupportsSeparateMSAAAndResolveTextures(GMaxRHIShaderPlatform);
-	}
-
-	if (!bForceSeparateTargetAndShaderResource)
-	{
-		FRHITextureCreateDesc Desc =
-			FRHITextureCreateDesc(BaseDesc)
-			.AddFlags(TargetableTextureFlags | ETextureCreateFlags::ShaderResource)
-			.SetInitialState(ERHIAccess::SRVMask);
-
-		// Create a single texture that has both TargetableTextureFlags and ETextureCreateFlags::ShaderResource set.
-		OutTargetableTexture = OutShaderResourceTexture = RHICreateTexture(Desc);
-	}
-	else
-	{
-		ETextureCreateFlags ResolveTargetableTextureFlags = ETextureCreateFlags::ResolveTargetable;
-		if (EnumHasAnyFlags(TargetableTextureFlags, ETextureCreateFlags::DepthStencilTargetable))
-		{
-			ResolveTargetableTextureFlags |= ETextureCreateFlags::DepthStencilResolveTarget;
-		}
-
-		FRHITextureCreateDesc TargetableDesc =
-			FRHITextureCreateDesc(BaseDesc)
-			.AddFlags(TargetableTextureFlags)
-			.SetInitialState(ERHIAccess::SRVMask);
-
-		FRHITextureCreateDesc ResourceDesc =
-			FRHITextureCreateDesc(BaseDesc)
-			.AddFlags(ResolveTargetableTextureFlags | ETextureCreateFlags::ShaderResource)
-			.SetInitialState(ERHIAccess::SRVMask)
-			.SetNumSamples(1);
-
-		// Create a texture that has TargetableTextureFlags set, and a second texture that has ETextureCreateFlags::ResolveTargetable and ETextureCreateFlags::ShaderResource set.
-		OutTargetableTexture = RHICreateTexture(TargetableDesc);
-		OutShaderResourceTexture = RHICreateTexture(ResourceDesc);
-	}
-}
-
-void RHICreateTargetableShaderResource2D(
-	uint32 SizeX,
-	uint32 SizeY,
-	uint8 Format,
-	uint32 NumMips,
-	ETextureCreateFlags Flags,
-	ETextureCreateFlags TargetableTextureFlags,
-	bool bForceSeparateTargetAndShaderResource,
-	bool bForceSharedTargetAndShaderResource,
-	const FRHIResourceCreateInfo& CreateInfo,
-	FTextureRHIRef& OutTargetableTexture,
-	FTextureRHIRef& OutShaderResourceTexture,
-	uint32 NumSamples
-)
-{
-	FRHITextureCreateDesc Desc =
-		FRHITextureCreateDesc::Create2D(CreateInfo.DebugName)
-		.SetExtent(SizeX, SizeY)
-		.SetFormat((EPixelFormat)Format)
-		.SetNumMips(NumMips)
-		.SetNumSamples(NumSamples)
-		.SetFlags(Flags)
-		.SetBulkData(CreateInfo.BulkData)
-		.SetClearValue(CreateInfo.ClearValueBinding)
-		.SetExtData(CreateInfo.ExtData)
-		.SetGPUMask(CreateInfo.GPUMask);
-
-	RHICreateTargetableShaderResource(Desc, TargetableTextureFlags, bForceSeparateTargetAndShaderResource, bForceSharedTargetAndShaderResource, OutTargetableTexture, OutShaderResourceTexture);
-}
-
-void RHICreateTargetableShaderResource2D(
-	uint32 SizeX,
-	uint32 SizeY,
-	uint8 Format,
-	uint32 NumMips,
-	ETextureCreateFlags Flags,
-	ETextureCreateFlags TargetableTextureFlags,
-	bool bForceSeparateTargetAndShaderResource,
-	const FRHIResourceCreateInfo& CreateInfo,
-	FTexture2DRHIRef& OutTargetableTexture,
-	FTexture2DRHIRef& OutShaderResourceTexture,
-	uint32 NumSamples)
-{
-PRAGMA_DISABLE_DEPRECATION_WARNINGS
-	RHICreateTargetableShaderResource2D(SizeX, SizeY, Format, NumMips, Flags, TargetableTextureFlags, bForceSeparateTargetAndShaderResource, false, CreateInfo, OutTargetableTexture, OutShaderResourceTexture, NumSamples);
-PRAGMA_ENABLE_DEPRECATION_WARNINGS
-}
-
-void RHICreateTargetableShaderResource2DArray(
-	uint32 SizeX,
-	uint32 SizeY,
-	uint32 SizeZ,
-	uint8 Format,
-	uint32 NumMips,
-	ETextureCreateFlags Flags,
-	ETextureCreateFlags TargetableTextureFlags,
-	bool bForceSeparateTargetAndShaderResource,
-	bool bForceSharedTargetAndShaderResource,
-	const FRHIResourceCreateInfo& CreateInfo,
-	FTextureRHIRef& OutTargetableTexture,
-	FTextureRHIRef& OutShaderResourceTexture,
-	uint32 NumSamples
-)
-{
-	FRHITextureCreateDesc Desc =
-		FRHITextureCreateDesc::Create2DArray(CreateInfo.DebugName)
-		.SetExtent(SizeX, SizeY)
-		.SetArraySize(SizeZ)
-		.SetFormat((EPixelFormat)Format)
-		.SetNumMips(NumMips)
-		.SetNumSamples(NumSamples)
-		.SetFlags(Flags)
-		.SetBulkData(CreateInfo.BulkData)
-		.SetClearValue(CreateInfo.ClearValueBinding)
-		.SetExtData(CreateInfo.ExtData)
-		.SetGPUMask(CreateInfo.GPUMask);
-
-	RHICreateTargetableShaderResource(Desc, TargetableTextureFlags, bForceSeparateTargetAndShaderResource, bForceSharedTargetAndShaderResource, OutTargetableTexture, OutShaderResourceTexture);
-}
-
-void RHICreateTargetableShaderResource2DArray(
-	uint32 SizeX,
-	uint32 SizeY,
-	uint32 SizeZ,
-	uint8 Format,
-	uint32 NumMips,
-	ETextureCreateFlags Flags,
-	ETextureCreateFlags TargetableTextureFlags,
-	const FRHIResourceCreateInfo& CreateInfo,
-	FTextureRHIRef& OutTargetableTexture,
-	FTextureRHIRef& OutShaderResourceTexture,
-	uint32 NumSamples)
-{
-PRAGMA_DISABLE_DEPRECATION_WARNINGS
-	RHICreateTargetableShaderResource2DArray(SizeX, SizeY, SizeZ, Format, NumMips, Flags, TargetableTextureFlags, false, false, CreateInfo, OutTargetableTexture, OutShaderResourceTexture, NumSamples);
-PRAGMA_ENABLE_DEPRECATION_WARNINGS
-}
-
-void RHICreateTargetableShaderResourceCube(
-	uint32 LinearSize,
-	uint8 Format,
-	uint32 NumMips,
-	ETextureCreateFlags Flags,
-	ETextureCreateFlags TargetableTextureFlags,
-	bool bForceSeparateTargetAndShaderResource,
-	const FRHIResourceCreateInfo& CreateInfo,
-	FTextureRHIRef& OutTargetableTexture,
-	FTextureRHIRef& OutShaderResourceTexture
-)
-{
-	FRHITextureCreateDesc Desc =
-		FRHITextureCreateDesc::CreateCube(CreateInfo.DebugName)
-		.SetExtent(LinearSize)
-		.SetFormat((EPixelFormat)Format)
-		.SetNumMips(NumMips)
-		.SetFlags(Flags)
-		.SetBulkData(CreateInfo.BulkData)
-		.SetClearValue(CreateInfo.ClearValueBinding)
-		.SetExtData(CreateInfo.ExtData)
-		.SetGPUMask(CreateInfo.GPUMask);
-
-	RHICreateTargetableShaderResource(Desc, TargetableTextureFlags, bForceSeparateTargetAndShaderResource, false, OutTargetableTexture, OutShaderResourceTexture);
-}
-
-void RHICreateTargetableShaderResourceCubeArray(
-	uint32 LinearSize,
-	uint32 ArraySize,
-	uint8 Format,
-	uint32 NumMips,
-	ETextureCreateFlags Flags,
-	ETextureCreateFlags TargetableTextureFlags,
-	bool bForceSeparateTargetAndShaderResource,
-	const FRHIResourceCreateInfo& CreateInfo,
-	FTextureRHIRef& OutTargetableTexture,
-	FTextureRHIRef& OutShaderResourceTexture
-)
-{
-	FRHITextureCreateDesc Desc =
-		FRHITextureCreateDesc::CreateCubeArray(CreateInfo.DebugName)
-		.SetExtent(LinearSize)
-		.SetArraySize(ArraySize)
-		.SetFormat((EPixelFormat)Format)
-		.SetNumMips(NumMips)
-		.SetFlags(Flags)
-		.SetBulkData(CreateInfo.BulkData)
-		.SetClearValue(CreateInfo.ClearValueBinding)
-		.SetExtData(CreateInfo.ExtData)
-		.SetGPUMask(CreateInfo.GPUMask);
-
-	RHICreateTargetableShaderResource(Desc, TargetableTextureFlags, bForceSeparateTargetAndShaderResource, false, OutTargetableTexture, OutShaderResourceTexture);
-}
-
-void RHICreateTargetableShaderResource3D(
-	uint32 SizeX,
-	uint32 SizeY,
-	uint32 SizeZ,
-	uint8 Format,
-	uint32 NumMips,
-	ETextureCreateFlags Flags,
-	ETextureCreateFlags TargetableTextureFlags,
-	bool bForceSeparateTargetAndShaderResource,
-	const FRHIResourceCreateInfo& CreateInfo,
-	FTextureRHIRef& OutTargetableTexture,
-	FTextureRHIRef& OutShaderResourceTexture
-)
-{
-	FRHITextureCreateDesc Desc =
-		FRHITextureCreateDesc::Create3D(CreateInfo.DebugName)
-		.SetExtent(SizeX, SizeY)
-		.SetDepth(SizeZ)
-		.SetFormat((EPixelFormat)Format)
-		.SetNumMips(NumMips)
-		.SetFlags(Flags)
-		.SetBulkData(CreateInfo.BulkData)
-		.SetClearValue(CreateInfo.ClearValueBinding)
-		.SetExtData(CreateInfo.ExtData)
-		.SetGPUMask(CreateInfo.GPUMask);
-
-	RHICreateTargetableShaderResource(Desc, TargetableTextureFlags, bForceSeparateTargetAndShaderResource, false, OutTargetableTexture, OutShaderResourceTexture);
-}
diff --git a/Engine/Source/Runtime/RHI/Public/RHIUtilities.h b/Engine/Source/Runtime/RHI/Public/RHIUtilities.h
--- a/Engine/Source/Runtime/RHI/Public/RHIUtilities.h
+++ b/Engine/Source/Runtime/RHI/Public/RHIUtilities.h
@@ -155,26 +155,6 @@ struct FTextureRWBuffer
 	}
 };
 
-struct UE_DEPRECATED(5.1, "FTextureRWBuffer should be used instead of FTextureRWBuffer2D.") FTextureRWBuffer2D;
-struct FTextureRWBuffer2D : public FTextureRWBuffer
-{
-	UE_DEPRECATED(5.1, "FTextureRWBuffer::Initialize2D should be used instead of FTextureRWBuffer2D::Initialize.")
-	void Initialize(const TCHAR* InDebugName, uint32 BytesPerElement, uint32 SizeX, uint32 SizeY, EPixelFormat Format, ETextureCreateFlags Flags = DefaultTextureInitFlag)
-	{
-		Initialize2D(InDebugName, BytesPerElement, SizeX, SizeY, Format, Flags);
-	}
-};
-
-struct UE_DEPRECATED(5.1, "FTextureRWBuffer should be used instead of FTextureRWBuffer3D.") FTextureRWBuffer3D;
-struct FTextureRWBuffer3D : public FTextureRWBuffer
-{
-	UE_DEPRECATED(5.1, "FTextureRWBuffer::Initialize3D should be used instead of FTextureRWBuffer3D::Initialize.")
-	void Initialize(const TCHAR* InDebugName, uint32 BytesPerElement, uint32 SizeX, uint32 SizeY, uint32 SizeZ, EPixelFormat Format, ETextureCreateFlags Flags = DefaultTextureInitFlag)
-	{
-		Initialize3D(InDebugName, BytesPerElement, SizeX, SizeY, SizeZ, Format, Flags);
-	}
-};
-
 /** Encapsulates a GPU read/write buffer with its UAV and SRV. */
 struct FRWBuffer
 {
@@ -479,136 +459,6 @@ inline void TransitionRenderPassTargets(FRHICommandList& RHICmdList, const FRHIR
 	RHICmdList.Transition(MakeArrayView(Transitions, TransitionIndex));
 }
 
-/**
- * Creates 1 or 2 textures with the same dimensions/format.
- * If the RHI supports textures that can be used as both shader resources and render targets,
- * and bForceSeparateTargetAndShaderResource=false, then a single texture is created.
- * Otherwise two textures are created, one of them usable as a shader resource and resolve target, and one of them usable as a render target.
- * Two texture references are always returned, but they may reference the same texture.
- * If two different textures are returned, the render-target texture must be manually copied to the shader-resource texture.
- */
-UE_DEPRECATED(5.1, "RHICreateTargetableShaderResource2D is deprecated. RHICreateTexture should be used instead.")
-RHI_API void RHICreateTargetableShaderResource2D(
-	uint32 SizeX,
-	uint32 SizeY,
-	uint8 Format,
-	uint32 NumMips,
-	ETextureCreateFlags Flags,
-	ETextureCreateFlags TargetableTextureFlags,
-	bool bForceSeparateTargetAndShaderResource,
-	bool bForceSharedTargetAndShaderResource,
-	const FRHIResourceCreateInfo& CreateInfo,
-	FTextureRHIRef& OutTargetableTexture,
-	FTextureRHIRef& OutShaderResourceTexture,
-	uint32 NumSamples = 1);
-
-UE_DEPRECATED(5.1, "RHICreateTargetableShaderResource2D is deprecated. RHICreateTexture should be used instead.")
-RHI_API void RHICreateTargetableShaderResource2D(
-	uint32 SizeX,
-	uint32 SizeY,
-	uint8 Format,
-	uint32 NumMips,
-	ETextureCreateFlags Flags,
-	ETextureCreateFlags TargetableTextureFlags,
-	bool bForceSeparateTargetAndShaderResource,
-	const FRHIResourceCreateInfo& CreateInfo,
-	FTextureRHIRef& OutTargetableTexture,
-	FTextureRHIRef& OutShaderResourceTexture,
-	uint32 NumSamples = 1);
-
-UE_DEPRECATED(5.1, "RHICreateTargetableShaderResource2DArray is deprecated. RHICreateTexture should be used instead.")
-RHI_API void RHICreateTargetableShaderResource2DArray(
-	uint32 SizeX,
-	uint32 SizeY,
-	uint32 SizeZ,
-	uint8 Format,
-	uint32 NumMips,
-	ETextureCreateFlags Flags,
-	ETextureCreateFlags TargetableTextureFlags,
-	bool bForceSeparateTargetAndShaderResource,
-	bool bForceSharedTargetAndShaderResource,
-	const FRHIResourceCreateInfo& CreateInfo,
-	FTextureRHIRef& OutTargetableTexture,
-	FTextureRHIRef& OutShaderResourceTexture,
-	uint32 NumSamples = 1);
-
-UE_DEPRECATED(5.1, "RHICreateTargetableShaderResource2DArray is deprecated. RHICreateTexture should be used instead.")
-RHI_API void RHICreateTargetableShaderResource2DArray(
-	uint32 SizeX,
-	uint32 SizeY,
-	uint32 SizeZ,
-	uint8 Format,
-	uint32 NumMips,
-	ETextureCreateFlags Flags,
-	ETextureCreateFlags TargetableTextureFlags,
-	const FRHIResourceCreateInfo& CreateInfo,
-	FTextureRHIRef& OutTargetableTexture,
-	FTextureRHIRef& OutShaderResourceTexture,
-	uint32 NumSamples = 1);
-
-/**
- * Creates 1 or 2 textures with the same dimensions/format.
- * If the RHI supports textures that can be used as both shader resources and render targets,
- * and bForceSeparateTargetAndShaderResource=false, then a single texture is created.
- * Otherwise two textures are created, one of them usable as a shader resource and resolve target, and one of them usable as a render target.
- * Two texture references are always returned, but they may reference the same texture.
- * If two different textures are returned, the render-target texture must be manually copied to the shader-resource texture.
- */
-UE_DEPRECATED(5.1, "RHICreateTargetableShaderResourceCube is deprecated. RHICreateTexture should be used instead.")
-RHI_API void RHICreateTargetableShaderResourceCube(
-	uint32 LinearSize,
-	uint8 Format,
-	uint32 NumMips,
-	ETextureCreateFlags Flags,
-	ETextureCreateFlags TargetableTextureFlags,
-	bool bForceSeparateTargetAndShaderResource,
-	const FRHIResourceCreateInfo& CreateInfo,
-	FTextureRHIRef& OutTargetableTexture,
-	FTextureRHIRef& OutShaderResourceTexture);
-
-/**
- * Creates 1 or 2 textures with the same dimensions/format.
- * If the RHI supports textures that can be used as both shader resources and render targets,
- * and bForceSeparateTargetAndShaderResource=false, then a single texture is created.
- * Otherwise two textures are created, one of them usable as a shader resource and resolve target, and one of them usable as a render target.
- * Two texture references are always returned, but they may reference the same texture.
- * If two different textures are returned, the render-target texture must be manually copied to the shader-resource texture.
- */
-UE_DEPRECATED(5.1, "RHICreateTargetableShaderResourceCubeArray is deprecated. RHICreateTexture should be used instead.")
-RHI_API void RHICreateTargetableShaderResourceCubeArray(
-	uint32 LinearSize,
-	uint32 ArraySize,
-	uint8 Format,
-	uint32 NumMips,
-	ETextureCreateFlags Flags,
-	ETextureCreateFlags TargetableTextureFlags,
-	bool bForceSeparateTargetAndShaderResource,
-	const FRHIResourceCreateInfo& CreateInfo,
-	FTextureRHIRef& OutTargetableTexture,
-	FTextureRHIRef& OutShaderResourceTexture);
-
-/**
- * Creates 1 or 2 textures with the same dimensions/format.
- * If the RHI supports textures that can be used as both shader resources and render targets,
- * and bForceSeparateTargetAndShaderResource=false, then a single texture is created.
- * Otherwise two textures are created, one of them usable as a shader resource and resolve target, and one of them usable as a render target.
- * Two texture references are always returned, but they may reference the same texture.
- * If two different textures are returned, the render-target texture must be manually copied to the shader-resource texture.
- */
-UE_DEPRECATED(5.1, "RHICreateTargetableShaderResource3D is deprecated. RHICreateTexture should be used instead.")
-RHI_API void RHICreateTargetableShaderResource3D(
-	uint32 SizeX,
-	uint32 SizeY,
-	uint32 SizeZ,
-	uint8 Format,
-	uint32 NumMips,
-	ETextureCreateFlags Flags,
-	ETextureCreateFlags TargetableTextureFlags,
-	bool bForceSeparateTargetAndShaderResource,
-	const FRHIResourceCreateInfo& CreateInfo,
-	FTextureRHIRef& OutTargetableTexture,
-	FTextureRHIRef& OutShaderResourceTexture);
-
 /** Performs a clear render pass on an RHI texture. The texture is expected to be in the RTV state. */
 inline void ClearRenderTarget(FRHICommandList& RHICmdList, FRHITexture* Texture, uint32 MipIndex = 0, uint32 ArraySlice = 0)
 {
diff --git a/Engine/Source/Runtime/VulkanRHI/Private/VulkanUtil.cpp b/Engine/Source/Runtime/VulkanRHI/Private/VulkanUtil.cpp
--- a/Engine/Source/Runtime/VulkanRHI/Private/VulkanUtil.cpp
+++ b/Engine/Source/Runtime/VulkanRHI/Private/VulkanUtil.cpp
@@ -903,15 +903,6 @@ void AftermathResolveMarkerCallback(const void* pMarker, void* pUserData, void**
 }
 #endif
 
-#include "VulkanRHIBridge.h"
-namespace VulkanRHIBridge
-{
-	FVulkanDevice* GetDevice(FVulkanDynamicRHI* RHI)
-	{
-		return RHI->GetDevice();
-	}
-}
-
 namespace VulkanRHI
 {
 	VkBuffer CreateBuffer(FVulkanDevice* InDevice, VkDeviceSize Size, VkBufferUsageFlags BufferUsageFlags, VkMemoryRequirements& OutMemoryRequirements)
diff --git a/Engine/Source/Runtime/VulkanRHI/Public/VulkanRHIBridge.h b/Engine/Source/Runtime/VulkanRHI/Public/VulkanRHIBridge.h
--- a/Engine/Source/Runtime/VulkanRHI/Public/VulkanRHIBridge.h
+++ b/Engine/Source/Runtime/VulkanRHI/Public/VulkanRHIBridge.h
@@ -1,50 +1,5 @@
 // Copyright Epic Games, Inc. All Rights Reserved.
 
-/*=============================================================================
-	VulkanRHIBridge.h: Utils to interact with the inner RHI.
-=============================================================================*/
-
 #pragma once
 
-#include "IVulkanDynamicRHI.h"
-
-class FVulkanDynamicRHI;
-class FVulkanDevice;
-
-namespace VulkanRHIBridge
-{
-	UE_DEPRECATED(5.1, "Please use IVulkanDynamicRHI::AddEnabledInstanceExtensionsAndLayers()")
-	inline void AddEnabledInstanceExtensionsAndLayers(const TArray<const ANSICHAR*>& InInstanceExtensions, const TArray<const ANSICHAR*>& InInstanceLayers)
-	{
-		IVulkanDynamicRHI::AddEnabledInstanceExtensionsAndLayers(InInstanceExtensions, InInstanceLayers);
-	}
-	
-	UE_DEPRECATED(5.1, "Please use IVulkanDynamicRHI::AddEnabledDeviceExtensionsAndLayers()")
-	inline void AddEnabledDeviceExtensionsAndLayers(const TArray<const ANSICHAR*>& InDeviceExtensions, const TArray<const ANSICHAR*>& InDeviceLayers)
-	{
-		IVulkanDynamicRHI::AddEnabledDeviceExtensionsAndLayers(InDeviceExtensions, InDeviceLayers);
-	}
-
-	UE_DEPRECATED(5.1, "Please use IVulkanDynamicRHI::RHIGetVkInstance()")
-	inline uint64 GetInstance(FVulkanDynamicRHI*)
-	{
-		return (uint64)GetIVulkanDynamicRHI()->RHIGetVkInstance();
-	}
-
-	UE_DEPRECATED(5.1, "Use IVulkanDynamicRHI instead of FVulkanDevice directly.")
-	VULKANRHI_API FVulkanDevice* GetDevice(FVulkanDynamicRHI*);
-
-	// Returns a VkDevice
-	UE_DEPRECATED(5.1, "Please use IVulkanDynamicRHI::RHIGetVkDevice()")
-	inline uint64 GetLogicalDevice(FVulkanDevice*)
-	{
-		return (uint64)GetIVulkanDynamicRHI()->RHIGetVkDevice();
-	}
-
-	// Returns a VkDeviceVkPhysicalDevice
-	UE_DEPRECATED(5.1, "Please use IVulkanDynamicRHI::RHIGetVkPhysicalDevice()")
-	inline uint64 GetPhysicalDevice(FVulkanDevice*)
-	{
-		return (uint64)GetIVulkanDynamicRHI()->RHIGetVkPhysicalDevice();
-	}
-}
+#pragma message("VulkanRHIBridge.h is deprecated. Please switch to using IVulkanDynamicRHI.h")
