Removal of cached hash value from FPipelineCacheFileFormatPSO.

diff --git a/Engine/Source/Editor/UnrealEd/Private/Commandlets/ShaderPipelineCacheToolsCommandlet.cpp b/Engine/Source/Editor/UnrealEd/Private/Commandlets/ShaderPipelineCacheToolsCommandlet.cpp
--- a/Engine/Source/Editor/UnrealEd/Private/Commandlets/ShaderPipelineCacheToolsCommandlet.cpp
+++ b/Engine/Source/Editor/UnrealEd/Private/Commandlets/ShaderPipelineCacheToolsCommandlet.cpp
@@ -378,7 +378,6 @@ static void PrintShaders(const TMap<FSHAHash, TArray<FString>>& InverseMap, cons
 bool CheckPSOStringInveribility(const FPipelineCacheFileFormatPSO& Item)
 {
 	FPipelineCacheFileFormatPSO TempItem(Item);
-	TempItem.Hash = 0;
 
 	FString StringRep;
 	switch (Item.Type)
@@ -2166,7 +2165,6 @@ void AddComputePSOs(TSet<FPipelineCacheFileFormatPSO>& OutPSOs, const TMultiMap<
 			FPipelineCacheFileFormatPSO NewPso;
 			NewPso.Type = FPipelineCacheFileFormatPSO::DescriptorType::Compute;
 			NewPso.ComputeDesc.ComputeShader = Iter.Value();
-			NewPso.Hash = 0;
 			NewPso.UsageMask = uint64(-1);
 			NewPso.BindCount = 0;
 			OutPSOs.Add(NewPso);
diff --git a/Engine/Source/Runtime/RHI/Private/PipelineFileCache.cpp b/Engine/Source/Runtime/RHI/Private/PipelineFileCache.cpp
--- a/Engine/Source/Runtime/RHI/Private/PipelineFileCache.cpp
+++ b/Engine/Source/Runtime/RHI/Private/PipelineFileCache.cpp
@@ -1065,99 +1065,96 @@ bool FPipelineCacheFileFormatPSO::Verify() const
 
 /*friend*/ uint32 GetTypeHash(const FPipelineCacheFileFormatPSO &Key)
 {
-	if(FPlatformAtomics::AtomicRead((volatile int32*)&Key.Hash) == 0)
+	uint32 KeyHash = GetTypeHash(Key.Type);
+	switch(Key.Type)
 	{
-		uint32 KeyHash = GetTypeHash(Key.Type);
-		switch(Key.Type)
+		case FPipelineCacheFileFormatPSO::DescriptorType::Compute:
 		{
-			case FPipelineCacheFileFormatPSO::DescriptorType::Compute:
+			KeyHash ^= GetTypeHash(Key.ComputeDesc.ComputeShader);
+			break;
+		}
+		case FPipelineCacheFileFormatPSO::DescriptorType::Graphics:
+		{
+			KeyHash = FCrc::MemCrc32(&Key.GraphicsDesc.RenderTargetsActive, sizeof(Key.GraphicsDesc.RenderTargetsActive), KeyHash);
+			KeyHash = FCrc::MemCrc32(&Key.GraphicsDesc.MSAASamples, sizeof(Key.GraphicsDesc.MSAASamples), KeyHash);
+			KeyHash = FCrc::MemCrc32(&Key.GraphicsDesc.PrimitiveType, sizeof(Key.GraphicsDesc.PrimitiveType), KeyHash);
+				
+			KeyHash = FCrc::MemCrc32(&Key.GraphicsDesc.VertexShader.Hash, sizeof(Key.GraphicsDesc.VertexShader.Hash), KeyHash);
+			KeyHash = FCrc::MemCrc32(&Key.GraphicsDesc.FragmentShader.Hash, sizeof(Key.GraphicsDesc.FragmentShader.Hash), KeyHash);
+			KeyHash = FCrc::MemCrc32(&Key.GraphicsDesc.GeometryShader.Hash, sizeof(Key.GraphicsDesc.GeometryShader.Hash), KeyHash);
+			KeyHash = FCrc::MemCrc32(&Key.GraphicsDesc.MeshShader.Hash, sizeof(Key.GraphicsDesc.MeshShader.Hash), KeyHash);
+			KeyHash = FCrc::MemCrc32(&Key.GraphicsDesc.AmplificationShader.Hash, sizeof(Key.GraphicsDesc.AmplificationShader.Hash), KeyHash);
+
+			KeyHash = FCrc::MemCrc32(&Key.GraphicsDesc.DepthStencilFormat, sizeof(Key.GraphicsDesc.DepthStencilFormat), KeyHash);
+			KeyHash = FCrc::MemCrc32(&Key.GraphicsDesc.DepthStencilFlags, sizeof(Key.GraphicsDesc.DepthStencilFlags), KeyHash);
+			KeyHash = FCrc::MemCrc32(&Key.GraphicsDesc.DepthLoad, sizeof(Key.GraphicsDesc.DepthLoad), KeyHash);
+			KeyHash = FCrc::MemCrc32(&Key.GraphicsDesc.StencilLoad, sizeof(Key.GraphicsDesc.StencilLoad), KeyHash);
+			KeyHash = FCrc::MemCrc32(&Key.GraphicsDesc.DepthStore, sizeof(Key.GraphicsDesc.DepthStore), KeyHash);
+			KeyHash = FCrc::MemCrc32(&Key.GraphicsDesc.StencilStore, sizeof(Key.GraphicsDesc.StencilStore), KeyHash);
+
+			KeyHash = FCrc::MemCrc32(&Key.GraphicsDesc.BlendState.bUseIndependentRenderTargetBlendStates, sizeof(Key.GraphicsDesc.BlendState.bUseIndependentRenderTargetBlendStates), KeyHash);
+			KeyHash = FCrc::MemCrc32(&Key.GraphicsDesc.BlendState.bUseAlphaToCoverage, sizeof(Key.GraphicsDesc.BlendState.bUseAlphaToCoverage), KeyHash);
+			for( uint32 i = 0; i < MaxSimultaneousRenderTargets; i++ )
 			{
-				KeyHash ^= GetTypeHash(Key.ComputeDesc.ComputeShader);
-				break;
+				KeyHash = FCrc::MemCrc32(&Key.GraphicsDesc.BlendState.RenderTargets[i].ColorBlendOp, sizeof(Key.GraphicsDesc.BlendState.RenderTargets[i].ColorBlendOp), KeyHash);
+				KeyHash = FCrc::MemCrc32(&Key.GraphicsDesc.BlendState.RenderTargets[i].ColorSrcBlend, sizeof(Key.GraphicsDesc.BlendState.RenderTargets[i].ColorSrcBlend), KeyHash);
+				KeyHash = FCrc::MemCrc32(&Key.GraphicsDesc.BlendState.RenderTargets[i].ColorDestBlend, sizeof(Key.GraphicsDesc.BlendState.RenderTargets[i].ColorDestBlend), KeyHash);
+				KeyHash = FCrc::MemCrc32(&Key.GraphicsDesc.BlendState.RenderTargets[i].ColorWriteMask, sizeof(Key.GraphicsDesc.BlendState.RenderTargets[i].ColorWriteMask), KeyHash);
+				KeyHash = FCrc::MemCrc32(&Key.GraphicsDesc.BlendState.RenderTargets[i].AlphaBlendOp, sizeof(Key.GraphicsDesc.BlendState.RenderTargets[i].AlphaBlendOp), KeyHash);
+				KeyHash = FCrc::MemCrc32(&Key.GraphicsDesc.BlendState.RenderTargets[i].AlphaSrcBlend, sizeof(Key.GraphicsDesc.BlendState.RenderTargets[i].AlphaSrcBlend), KeyHash);
+				KeyHash = FCrc::MemCrc32(&Key.GraphicsDesc.BlendState.RenderTargets[i].AlphaDestBlend, sizeof(Key.GraphicsDesc.BlendState.RenderTargets[i].AlphaDestBlend), KeyHash);
 			}
-			case FPipelineCacheFileFormatPSO::DescriptorType::Graphics:
-			{
-				KeyHash = FCrc::MemCrc32(&Key.GraphicsDesc.RenderTargetsActive, sizeof(Key.GraphicsDesc.RenderTargetsActive), KeyHash);
-				KeyHash = FCrc::MemCrc32(&Key.GraphicsDesc.MSAASamples, sizeof(Key.GraphicsDesc.MSAASamples), KeyHash);
-				KeyHash = FCrc::MemCrc32(&Key.GraphicsDesc.PrimitiveType, sizeof(Key.GraphicsDesc.PrimitiveType), KeyHash);
-				
-				KeyHash = FCrc::MemCrc32(&Key.GraphicsDesc.VertexShader.Hash, sizeof(Key.GraphicsDesc.VertexShader.Hash), KeyHash);
-				KeyHash = FCrc::MemCrc32(&Key.GraphicsDesc.FragmentShader.Hash, sizeof(Key.GraphicsDesc.FragmentShader.Hash), KeyHash);
-				KeyHash = FCrc::MemCrc32(&Key.GraphicsDesc.GeometryShader.Hash, sizeof(Key.GraphicsDesc.GeometryShader.Hash), KeyHash);
-				KeyHash = FCrc::MemCrc32(&Key.GraphicsDesc.MeshShader.Hash, sizeof(Key.GraphicsDesc.MeshShader.Hash), KeyHash);
-				KeyHash = FCrc::MemCrc32(&Key.GraphicsDesc.AmplificationShader.Hash, sizeof(Key.GraphicsDesc.AmplificationShader.Hash), KeyHash);
-
-				KeyHash = FCrc::MemCrc32(&Key.GraphicsDesc.DepthStencilFormat, sizeof(Key.GraphicsDesc.DepthStencilFormat), KeyHash);
-				KeyHash = FCrc::MemCrc32(&Key.GraphicsDesc.DepthStencilFlags, sizeof(Key.GraphicsDesc.DepthStencilFlags), KeyHash);
-				KeyHash = FCrc::MemCrc32(&Key.GraphicsDesc.DepthLoad, sizeof(Key.GraphicsDesc.DepthLoad), KeyHash);
-				KeyHash = FCrc::MemCrc32(&Key.GraphicsDesc.StencilLoad, sizeof(Key.GraphicsDesc.StencilLoad), KeyHash);
-				KeyHash = FCrc::MemCrc32(&Key.GraphicsDesc.DepthStore, sizeof(Key.GraphicsDesc.DepthStore), KeyHash);
-				KeyHash = FCrc::MemCrc32(&Key.GraphicsDesc.StencilStore, sizeof(Key.GraphicsDesc.StencilStore), KeyHash);
-
-				KeyHash = FCrc::MemCrc32(&Key.GraphicsDesc.BlendState.bUseIndependentRenderTargetBlendStates, sizeof(Key.GraphicsDesc.BlendState.bUseIndependentRenderTargetBlendStates), KeyHash);
-				KeyHash = FCrc::MemCrc32(&Key.GraphicsDesc.BlendState.bUseAlphaToCoverage, sizeof(Key.GraphicsDesc.BlendState.bUseAlphaToCoverage), KeyHash);
-				for( uint32 i = 0; i < MaxSimultaneousRenderTargets; i++ )
-				{
-					KeyHash = FCrc::MemCrc32(&Key.GraphicsDesc.BlendState.RenderTargets[i].ColorBlendOp, sizeof(Key.GraphicsDesc.BlendState.RenderTargets[i].ColorBlendOp), KeyHash);
-					KeyHash = FCrc::MemCrc32(&Key.GraphicsDesc.BlendState.RenderTargets[i].ColorSrcBlend, sizeof(Key.GraphicsDesc.BlendState.RenderTargets[i].ColorSrcBlend), KeyHash);
-					KeyHash = FCrc::MemCrc32(&Key.GraphicsDesc.BlendState.RenderTargets[i].ColorDestBlend, sizeof(Key.GraphicsDesc.BlendState.RenderTargets[i].ColorDestBlend), KeyHash);
-					KeyHash = FCrc::MemCrc32(&Key.GraphicsDesc.BlendState.RenderTargets[i].ColorWriteMask, sizeof(Key.GraphicsDesc.BlendState.RenderTargets[i].ColorWriteMask), KeyHash);
-					KeyHash = FCrc::MemCrc32(&Key.GraphicsDesc.BlendState.RenderTargets[i].AlphaBlendOp, sizeof(Key.GraphicsDesc.BlendState.RenderTargets[i].AlphaBlendOp), KeyHash);
-					KeyHash = FCrc::MemCrc32(&Key.GraphicsDesc.BlendState.RenderTargets[i].AlphaSrcBlend, sizeof(Key.GraphicsDesc.BlendState.RenderTargets[i].AlphaSrcBlend), KeyHash);
-					KeyHash = FCrc::MemCrc32(&Key.GraphicsDesc.BlendState.RenderTargets[i].AlphaDestBlend, sizeof(Key.GraphicsDesc.BlendState.RenderTargets[i].AlphaDestBlend), KeyHash);
-				}
 
-				KeyHash = FCrc::MemCrc32(&Key.GraphicsDesc.RenderTargetFormats, sizeof(Key.GraphicsDesc.RenderTargetFormats), KeyHash);
-				KeyHash = FCrc::MemCrc32(&Key.GraphicsDesc.RenderTargetFlags, sizeof(Key.GraphicsDesc.RenderTargetFlags), KeyHash);
+			KeyHash = FCrc::MemCrc32(&Key.GraphicsDesc.RenderTargetFormats, sizeof(Key.GraphicsDesc.RenderTargetFormats), KeyHash);
+			KeyHash = FCrc::MemCrc32(&Key.GraphicsDesc.RenderTargetFlags, sizeof(Key.GraphicsDesc.RenderTargetFlags), KeyHash);
 				
-				KeyHash = FCrc::MemCrc32(&Key.GraphicsDesc.SubpassHint, sizeof(Key.GraphicsDesc.SubpassHint), KeyHash);
-				KeyHash = FCrc::MemCrc32(&Key.GraphicsDesc.SubpassIndex, sizeof(Key.GraphicsDesc.SubpassIndex), KeyHash);
+			KeyHash = FCrc::MemCrc32(&Key.GraphicsDesc.SubpassHint, sizeof(Key.GraphicsDesc.SubpassHint), KeyHash);
+			KeyHash = FCrc::MemCrc32(&Key.GraphicsDesc.SubpassIndex, sizeof(Key.GraphicsDesc.SubpassIndex), KeyHash);
 				
-				KeyHash = FCrc::MemCrc32(&Key.GraphicsDesc.MultiViewCount, sizeof(Key.GraphicsDesc.MultiViewCount), KeyHash);
-				KeyHash = FCrc::MemCrc32(&Key.GraphicsDesc.bHasFragmentDensityAttachment, sizeof(Key.GraphicsDesc.bHasFragmentDensityAttachment), KeyHash);
+			KeyHash = FCrc::MemCrc32(&Key.GraphicsDesc.MultiViewCount, sizeof(Key.GraphicsDesc.MultiViewCount), KeyHash);
+			KeyHash = FCrc::MemCrc32(&Key.GraphicsDesc.bHasFragmentDensityAttachment, sizeof(Key.GraphicsDesc.bHasFragmentDensityAttachment), KeyHash);
 
-				for(auto const& Element : Key.GraphicsDesc.VertexDescriptor)
-				{
-					KeyHash = FCrc::MemCrc32(&Element, sizeof(FVertexElement), KeyHash);
-				}
+			for(auto const& Element : Key.GraphicsDesc.VertexDescriptor)
+			{
+				KeyHash = FCrc::MemCrc32(&Element, sizeof(FVertexElement), KeyHash);
+			}
 				
-				KeyHash = FCrc::MemCrc32(&Key.GraphicsDesc.RasterizerState.DepthBias, sizeof(Key.GraphicsDesc.RasterizerState.DepthBias), KeyHash);
-				KeyHash = FCrc::MemCrc32(&Key.GraphicsDesc.RasterizerState.SlopeScaleDepthBias, sizeof(Key.GraphicsDesc.RasterizerState.SlopeScaleDepthBias), KeyHash);
-				KeyHash = FCrc::MemCrc32(&Key.GraphicsDesc.RasterizerState.FillMode, sizeof(Key.GraphicsDesc.RasterizerState.FillMode), KeyHash);
-				KeyHash = FCrc::MemCrc32(&Key.GraphicsDesc.RasterizerState.CullMode, sizeof(Key.GraphicsDesc.RasterizerState.CullMode), KeyHash);
-				KeyHash = FCrc::MemCrc32(&Key.GraphicsDesc.RasterizerState.bAllowMSAA, sizeof(Key.GraphicsDesc.RasterizerState.bAllowMSAA), KeyHash);
-				KeyHash = FCrc::MemCrc32(&Key.GraphicsDesc.RasterizerState.bEnableLineAA, sizeof(Key.GraphicsDesc.RasterizerState.bEnableLineAA), KeyHash);
+			KeyHash = FCrc::MemCrc32(&Key.GraphicsDesc.RasterizerState.DepthBias, sizeof(Key.GraphicsDesc.RasterizerState.DepthBias), KeyHash);
+			KeyHash = FCrc::MemCrc32(&Key.GraphicsDesc.RasterizerState.SlopeScaleDepthBias, sizeof(Key.GraphicsDesc.RasterizerState.SlopeScaleDepthBias), KeyHash);
+			KeyHash = FCrc::MemCrc32(&Key.GraphicsDesc.RasterizerState.FillMode, sizeof(Key.GraphicsDesc.RasterizerState.FillMode), KeyHash);
+			KeyHash = FCrc::MemCrc32(&Key.GraphicsDesc.RasterizerState.CullMode, sizeof(Key.GraphicsDesc.RasterizerState.CullMode), KeyHash);
+			KeyHash = FCrc::MemCrc32(&Key.GraphicsDesc.RasterizerState.bAllowMSAA, sizeof(Key.GraphicsDesc.RasterizerState.bAllowMSAA), KeyHash);
+			KeyHash = FCrc::MemCrc32(&Key.GraphicsDesc.RasterizerState.bEnableLineAA, sizeof(Key.GraphicsDesc.RasterizerState.bEnableLineAA), KeyHash);
 				
-				KeyHash = FCrc::MemCrc32(&Key.GraphicsDesc.DepthStencilState.bEnableDepthWrite, sizeof(Key.GraphicsDesc.DepthStencilState.bEnableDepthWrite), KeyHash);
-				KeyHash = FCrc::MemCrc32(&Key.GraphicsDesc.DepthStencilState.DepthTest, sizeof(Key.GraphicsDesc.DepthStencilState.DepthTest), KeyHash);
-				KeyHash = FCrc::MemCrc32(&Key.GraphicsDesc.DepthStencilState.bEnableFrontFaceStencil, sizeof(Key.GraphicsDesc.DepthStencilState.bEnableFrontFaceStencil), KeyHash);
-				KeyHash = FCrc::MemCrc32(&Key.GraphicsDesc.DepthStencilState.FrontFaceStencilTest, sizeof(Key.GraphicsDesc.DepthStencilState.FrontFaceStencilTest), KeyHash);
-				KeyHash = FCrc::MemCrc32(&Key.GraphicsDesc.DepthStencilState.FrontFaceStencilFailStencilOp, sizeof(Key.GraphicsDesc.DepthStencilState.FrontFaceStencilFailStencilOp), KeyHash);
-				KeyHash = FCrc::MemCrc32(&Key.GraphicsDesc.DepthStencilState.FrontFaceDepthFailStencilOp, sizeof(Key.GraphicsDesc.DepthStencilState.FrontFaceDepthFailStencilOp), KeyHash);
-				KeyHash = FCrc::MemCrc32(&Key.GraphicsDesc.DepthStencilState.FrontFacePassStencilOp, sizeof(Key.GraphicsDesc.DepthStencilState.FrontFacePassStencilOp), KeyHash);
-				KeyHash = FCrc::MemCrc32(&Key.GraphicsDesc.DepthStencilState.bEnableBackFaceStencil, sizeof(Key.GraphicsDesc.DepthStencilState.bEnableBackFaceStencil), KeyHash);
-				KeyHash = FCrc::MemCrc32(&Key.GraphicsDesc.DepthStencilState.BackFaceStencilTest, sizeof(Key.GraphicsDesc.DepthStencilState.BackFaceStencilTest), KeyHash);
-				KeyHash = FCrc::MemCrc32(&Key.GraphicsDesc.DepthStencilState.BackFaceStencilFailStencilOp, sizeof(Key.GraphicsDesc.DepthStencilState.BackFaceStencilFailStencilOp), KeyHash);
-				KeyHash = FCrc::MemCrc32(&Key.GraphicsDesc.DepthStencilState.BackFaceDepthFailStencilOp, sizeof(Key.GraphicsDesc.DepthStencilState.BackFaceDepthFailStencilOp), KeyHash);
-				KeyHash = FCrc::MemCrc32(&Key.GraphicsDesc.DepthStencilState.BackFacePassStencilOp, sizeof(Key.GraphicsDesc.DepthStencilState.BackFacePassStencilOp), KeyHash);
-				KeyHash = FCrc::MemCrc32(&Key.GraphicsDesc.DepthStencilState.StencilReadMask, sizeof(Key.GraphicsDesc.DepthStencilState.StencilReadMask), KeyHash);
-				KeyHash = FCrc::MemCrc32(&Key.GraphicsDesc.DepthStencilState.StencilWriteMask, sizeof(Key.GraphicsDesc.DepthStencilState.StencilWriteMask), KeyHash);
+			KeyHash = FCrc::MemCrc32(&Key.GraphicsDesc.DepthStencilState.bEnableDepthWrite, sizeof(Key.GraphicsDesc.DepthStencilState.bEnableDepthWrite), KeyHash);
+			KeyHash = FCrc::MemCrc32(&Key.GraphicsDesc.DepthStencilState.DepthTest, sizeof(Key.GraphicsDesc.DepthStencilState.DepthTest), KeyHash);
+			KeyHash = FCrc::MemCrc32(&Key.GraphicsDesc.DepthStencilState.bEnableFrontFaceStencil, sizeof(Key.GraphicsDesc.DepthStencilState.bEnableFrontFaceStencil), KeyHash);
+			KeyHash = FCrc::MemCrc32(&Key.GraphicsDesc.DepthStencilState.FrontFaceStencilTest, sizeof(Key.GraphicsDesc.DepthStencilState.FrontFaceStencilTest), KeyHash);
+			KeyHash = FCrc::MemCrc32(&Key.GraphicsDesc.DepthStencilState.FrontFaceStencilFailStencilOp, sizeof(Key.GraphicsDesc.DepthStencilState.FrontFaceStencilFailStencilOp), KeyHash);
+			KeyHash = FCrc::MemCrc32(&Key.GraphicsDesc.DepthStencilState.FrontFaceDepthFailStencilOp, sizeof(Key.GraphicsDesc.DepthStencilState.FrontFaceDepthFailStencilOp), KeyHash);
+			KeyHash = FCrc::MemCrc32(&Key.GraphicsDesc.DepthStencilState.FrontFacePassStencilOp, sizeof(Key.GraphicsDesc.DepthStencilState.FrontFacePassStencilOp), KeyHash);
+			KeyHash = FCrc::MemCrc32(&Key.GraphicsDesc.DepthStencilState.bEnableBackFaceStencil, sizeof(Key.GraphicsDesc.DepthStencilState.bEnableBackFaceStencil), KeyHash);
+			KeyHash = FCrc::MemCrc32(&Key.GraphicsDesc.DepthStencilState.BackFaceStencilTest, sizeof(Key.GraphicsDesc.DepthStencilState.BackFaceStencilTest), KeyHash);
+			KeyHash = FCrc::MemCrc32(&Key.GraphicsDesc.DepthStencilState.BackFaceStencilFailStencilOp, sizeof(Key.GraphicsDesc.DepthStencilState.BackFaceStencilFailStencilOp), KeyHash);
+			KeyHash = FCrc::MemCrc32(&Key.GraphicsDesc.DepthStencilState.BackFaceDepthFailStencilOp, sizeof(Key.GraphicsDesc.DepthStencilState.BackFaceDepthFailStencilOp), KeyHash);
+			KeyHash = FCrc::MemCrc32(&Key.GraphicsDesc.DepthStencilState.BackFacePassStencilOp, sizeof(Key.GraphicsDesc.DepthStencilState.BackFacePassStencilOp), KeyHash);
+			KeyHash = FCrc::MemCrc32(&Key.GraphicsDesc.DepthStencilState.StencilReadMask, sizeof(Key.GraphicsDesc.DepthStencilState.StencilReadMask), KeyHash);
+			KeyHash = FCrc::MemCrc32(&Key.GraphicsDesc.DepthStencilState.StencilWriteMask, sizeof(Key.GraphicsDesc.DepthStencilState.StencilWriteMask), KeyHash);
 				
-				break;
-			}
-			case FPipelineCacheFileFormatPSO::DescriptorType::RayTracing:
-			{
-				KeyHash ^= GetTypeHash(Key.RayTracingDesc);
-				break;
-			}
-			default:
-			{
-				checkNoEntry();
-			}
+			break;
+		}
+		case FPipelineCacheFileFormatPSO::DescriptorType::RayTracing:
+		{
+			KeyHash ^= GetTypeHash(Key.RayTracingDesc);
+			break;
+		}
+		default:
+		{
+			checkNoEntry();
 		}
-		FPlatformAtomics::InterlockedCompareExchange((volatile int32*)&Key.Hash, KeyHash, 0);
 	}
-	return FPlatformAtomics::AtomicRead((volatile int32*)&Key.Hash);
+
+	return KeyHash;
 }
 
 /*friend*/ FArchive& operator<<( FArchive& Ar, FPipelineCacheFileFormatPSO& Info )
@@ -1351,9 +1348,8 @@ bool FPipelineCacheFileFormatPSO::Verify() const
 }
 
 FPipelineCacheFileFormatPSO::FPipelineCacheFileFormatPSO()
-: Hash(0)
 #if PSO_COOKONLY_DATA
-, UsageMask(0)
+: UsageMask(0)
 , BindCount(0)
 #endif
 {
@@ -1363,7 +1359,6 @@ FPipelineCacheFileFormatPSO::FPipelineCacheFileFormatPSO()
 {
 	check(Init);
 
-	PSO.Hash = 0;
 	PSO.Type = DescriptorType::Compute;
 #if PSO_COOKONLY_DATA
 	PSO.UsageMask = 0;
@@ -1388,7 +1383,6 @@ FPipelineCacheFileFormatPSO::FPipelineCacheFileFormatPSO()
 {
 	bool bOK = true;
 	
-	PSO.Hash = 0;
 	PSO.Type = DescriptorType::Graphics;
 #if PSO_COOKONLY_DATA
 	PSO.UsageMask = 0;
@@ -1526,7 +1520,7 @@ bool FPipelineCacheFileFormatPSO::operator==(const FPipelineCacheFileFormatPSO&
 	if (this != &Other)
 	{
 		bSame = Type == Other.Type;
-		/* Ignore: [Hash == Other.Hash] in this test. */
+
 #if PSO_COOKONLY_DATA
 		/* Ignore: [UsageMask == UsageMask] in this test. */
 		/* Ignore: [BindCount == BindCount] in this test. */
@@ -1588,7 +1582,6 @@ bool FPipelineCacheFileFormatPSO::operator==(const FPipelineCacheFileFormatPSO&
 
 FPipelineCacheFileFormatPSO::FPipelineCacheFileFormatPSO(const FPipelineCacheFileFormatPSO& Other)
 : Type(Other.Type)
-, Hash(Other.Hash)
 #if PSO_COOKONLY_DATA
 , UsageMask(Other.UsageMask)
 , BindCount(Other.BindCount)
@@ -1626,7 +1619,6 @@ FPipelineCacheFileFormatPSO& FPipelineCacheFileFormatPSO::operator=(const FPipel
 	if(this != &Other)
 	{
 		Type = Other.Type;
-		Hash = Other.Hash;
 #if PSO_COOKONLY_DATA
 		UsageMask = Other.UsageMask;
 		BindCount = Other.BindCount;
@@ -4254,7 +4246,6 @@ void FPipelineCacheFileFormatPSO::FPipelineFileCacheRayTracingDesc::FromString(c
 
 bool FPipelineCacheFileFormatPSO::Init(FPipelineCacheFileFormatPSO& PSO, FPipelineCacheFileFormatPSO::FPipelineFileCacheRayTracingDesc const& Desc)
 {
-	PSO.Hash = 0;
 	PSO.Type = DescriptorType::RayTracing;
 
 #if PSO_COOKONLY_DATA
diff --git a/Engine/Source/Runtime/RHI/Public/PipelineFileCache.h b/Engine/Source/Runtime/RHI/Public/PipelineFileCache.h
--- a/Engine/Source/Runtime/RHI/Public/PipelineFileCache.h
+++ b/Engine/Source/Runtime/RHI/Public/PipelineFileCache.h
@@ -224,8 +224,6 @@ struct RHI_API FPipelineCacheFileFormatPSO
 	GraphicsDescriptor GraphicsDesc;
 	FPipelineFileCacheRayTracingDesc RayTracingDesc;
 
-	mutable volatile uint32 Hash;
-	
 #if PSO_COOKONLY_DATA
 	uint64 UsageMask;
 	int64 BindCount;
