Missing file
All PreviewShaderPlatforms support bSupportsDebugViewShaders

diff --git a/Engine/Source/Runtime/RHI/Private/DataDrivenShaderPlatformInfo.cpp b/Engine/Source/Runtime/RHI/Private/DataDrivenShaderPlatformInfo.cpp
--- a/Engine/Source/Runtime/RHI/Private/DataDrivenShaderPlatformInfo.cpp
+++ b/Engine/Source/Runtime/RHI/Private/DataDrivenShaderPlatformInfo.cpp
@@ -378,6 +378,7 @@ void FGenericDataDrivenShaderPlatformInfo::UpdatePreviewPlatforms()
 				Infos[ShaderPlatform].bSupportsDxc = Infos[EditorSPForPreviewMaxFeatureLevel].bSupportsDxc;
 				Infos[ShaderPlatform].bSupportsGPUScene = Infos[EditorSPForPreviewMaxFeatureLevel].bSupportsGPUScene;
 				Infos[ShaderPlatform].bIsPC = true;
+				Infos[ShaderPlatform].bSupportsDebugViewShaders = true;
 				Infos[ShaderPlatform].bIsConsole = false;
 				Infos[ShaderPlatform].bSupportsSceneDataCompressedTransforms = Infos[EditorSPForPreviewMaxFeatureLevel].bSupportsSceneDataCompressedTransforms;
 				Infos[ShaderPlatform].bSupportsNanite &= Infos[EditorSPForPreviewMaxFeatureLevel].bSupportsNanite;
