- Fix incorrect validation checking for indirect arguments, the current code is triggering on draw indirects which are valid
- We now pass in the indirect argument size to make sure we correctly validate as using stride is not correct in many cases
- We only boundary check for dispatch indirect arguments
- Tested with Lyra / AncientGame / custom project

#rnx

diff --git a/Engine/Source/Runtime/RHI/Public/RHIValidation.h b/Engine/Source/Runtime/RHI/Public/RHIValidation.h
--- a/Engine/Source/Runtime/RHI/Public/RHIValidation.h
+++ b/Engine/Source/Runtime/RHI/Public/RHIValidation.h
@@ -44,19 +44,23 @@ public:
 		RHI_VALIDATION_CHECK((ThreadGroupCountZ <= (uint32)GRHIMaxDispatchThreadGroupsPerDimension.Z), *FString::Printf(TEXT("ThreadGroupCountZ is invalid: %u. Must be greater than 0 and less than %d"), ThreadGroupCountZ, GRHIMaxDispatchThreadGroupsPerDimension.Z));
 	}
 
-	static inline void ValidateIndirectArgsBuffer(FRHIBuffer* ArgumentBuffer, uint32 ArgumentOffset)
+	static inline void ValidateIndirectArgsBuffer(FRHIBuffer* ArgumentBuffer, uint32 ArgumentOffset, uint32 ArgumentSize, uint32 ArgumentsBoundarySize)
 	{
-		auto GetBufferDesc = [&]() -> FString { return FString::Printf(TEXT("Buffer: %s, Size: %u, Stride: %u, Offset: %u"), ArgumentBuffer->GetDebugName(), ArgumentBuffer->GetSize(), ArgumentBuffer->GetStride(), ArgumentOffset); };
-
+		auto GetBufferDesc = [&]() -> FString { return FString::Printf(TEXT("Buffer: %s, Size: %u, Stride: %u, Offset: %u, ArgSize: %u"), ArgumentBuffer->GetDebugName(), ArgumentBuffer->GetSize(), ArgumentBuffer->GetStride(), ArgumentOffset, ArgumentSize); };
 		RHI_VALIDATION_CHECK(EnumHasAnyFlags(ArgumentBuffer->GetUsage(), EBufferUsageFlags::VertexBuffer), *FString::Printf(TEXT("Indirect argument buffer must be a vertex buffer to be used as an indirect dispatch parameter. %s"), *GetBufferDesc()));
 		RHI_VALIDATION_CHECK(EnumHasAnyFlags(ArgumentBuffer->GetUsage(), EBufferUsageFlags::DrawIndirect), *FString::Printf(TEXT("Indirect dispatch parameter buffer was not flagged with BUF_DrawIndirect. %s"), *GetBufferDesc()));
-		RHI_VALIDATION_CHECK((ArgumentOffset % 4) == 0, *FString::Printf(TEXT("Offset into buffer for compute shader indirect dispatch must be a multiple of 4. %s"), *GetBufferDesc()));
-		RHI_VALIDATION_CHECK((ArgumentOffset + ArgumentBuffer->GetStride()) <= ArgumentBuffer->GetSize(),
-			*FString::Printf(TEXT("Indirect parameter buffer for compute shader indirect dispatch at byte offset %u doesn't have enough room for one element. %s"), ArgumentOffset, *GetBufferDesc()));
-#if PLATFORM_DISPATCH_INDIRECT_ARGUMENT_BOUNDARY_SIZE != 0
-		RHI_VALIDATION_CHECK(ArgumentOffset / PLATFORM_DISPATCH_INDIRECT_ARGUMENT_BOUNDARY_SIZE == (ArgumentOffset + ArgumentBuffer->GetStride() - 1) / PLATFORM_DISPATCH_INDIRECT_ARGUMENT_BOUNDARY_SIZE,
-			*FString::Printf(TEXT("Compute indirect dispatch arguments cannot cross %d byte boundary. %s"), PLATFORM_DISPATCH_INDIRECT_ARGUMENT_BOUNDARY_SIZE, *GetBufferDesc()));
-#endif // #if PLATFORM_DISPATCH_INDIRECT_ARGUMENT_BOUNDARY_SIZE != 0
+		RHI_VALIDATION_CHECK((ArgumentOffset % 4) == 0, *FString::Printf(TEXT("Indirect argument offset must be a multiple of 4. %s"), *GetBufferDesc()));
+		RHI_VALIDATION_CHECK((ArgumentOffset + ArgumentSize) <= ArgumentBuffer->GetSize(), *FString::Printf(TEXT("Indirect argument doesn't fit in the buffer. %s"), *GetBufferDesc()));
+		if (ArgumentsBoundarySize > 0)
+		{
+			RHI_VALIDATION_CHECK(ArgumentOffset / ArgumentsBoundarySize == (ArgumentOffset + sizeof(FRHIDispatchIndirectParametersNoPadding) - 1) / ArgumentsBoundarySize,
+				*FString::Printf(TEXT("Indirect arguments cannot cross %d byte boundary. %s"), ArgumentsBoundarySize, *GetBufferDesc()));
+		}
+	}
+
+	static inline void ValidateDispatchIndirectArgsBuffer(FRHIBuffer* ArgumentBuffer, uint32 ArgumentOffset)
+	{
+		ValidateIndirectArgsBuffer(ArgumentBuffer, ArgumentOffset, sizeof(FRHIDispatchIndirectParametersNoPadding), PLATFORM_DISPATCH_INDIRECT_ARGUMENT_BOUNDARY_SIZE);
 	}
 
 	virtual void Init() override final
diff --git a/Engine/Source/Runtime/RHI/Public/RHIValidationContext.h b/Engine/Source/Runtime/RHI/Public/RHIValidationContext.h
--- a/Engine/Source/Runtime/RHI/Public/RHIValidationContext.h
+++ b/Engine/Source/Runtime/RHI/Public/RHIValidationContext.h
@@ -96,7 +96,7 @@ public:
 	virtual void RHIDispatchIndirectComputeShader(FRHIBuffer* ArgumentBuffer, uint32 ArgumentOffset) override final
 	{
 		checkf(State.bComputePSOSet, TEXT("A Compute PSO has to be set to set resources into a shader!"));
-		FValidationRHI::ValidateIndirectArgsBuffer(ArgumentBuffer, ArgumentOffset);
+		FValidationRHI::ValidateDispatchIndirectArgsBuffer(ArgumentBuffer, ArgumentOffset);
 		Tracker->Assert(ArgumentBuffer->GetWholeResourceIdentity(), ERHIAccess::IndirectArgs);
 
 		RHIContext->RHIDispatchIndirectComputeShader(ArgumentBuffer, ArgumentOffset);
@@ -386,7 +386,7 @@ public:
 	virtual void RHIDispatchIndirectComputeShader(FRHIBuffer* ArgumentBuffer, uint32 ArgumentOffset) override final
 	{
 		checkf(State.bComputePSOSet, TEXT("A Compute PSO has to be set to set resources into a shader!"));
-		FValidationRHI::ValidateIndirectArgsBuffer(ArgumentBuffer, ArgumentOffset);
+		FValidationRHI::ValidateDispatchIndirectArgsBuffer(ArgumentBuffer, ArgumentOffset);
 		Tracker->Assert(ArgumentBuffer->GetWholeResourceIdentity(), ERHIAccess::IndirectArgs);
 
 		RHIContext->RHIDispatchIndirectComputeShader(ArgumentBuffer, ArgumentOffset);
@@ -847,7 +847,7 @@ public:
 	virtual void RHIDrawPrimitiveIndirect(FRHIBuffer* ArgumentBuffer, uint32 ArgumentOffset) override final
 	{
 		checkf(State.bGfxPSOSet, TEXT("A Graphics PSO has to be set to draw!"));
-		FValidationRHI::ValidateIndirectArgsBuffer(ArgumentBuffer, ArgumentOffset);
+		FValidationRHI::ValidateIndirectArgsBuffer(ArgumentBuffer, ArgumentOffset, sizeof(FRHIDrawIndirectParameters), 0);
 		SetupDrawing();
 		Tracker->Assert(ArgumentBuffer->GetWholeResourceIdentity(), ERHIAccess::IndirectArgs);
 		RHIContext->RHIDrawPrimitiveIndirect(ArgumentBuffer, ArgumentOffset);
@@ -857,7 +857,7 @@ public:
 	virtual void RHIDrawIndexedIndirect(FRHIBuffer* IndexBufferRHI, FRHIBuffer* ArgumentsBufferRHI, int32 DrawArgumentsIndex, uint32 NumInstances) override final
 	{
 		checkf(State.bGfxPSOSet, TEXT("A Graphics PSO has to be set to draw!"));
-		FValidationRHI::ValidateIndirectArgsBuffer(ArgumentsBufferRHI, DrawArgumentsIndex * ArgumentsBufferRHI->GetStride());
+		FValidationRHI::ValidateIndirectArgsBuffer(ArgumentsBufferRHI, DrawArgumentsIndex * ArgumentsBufferRHI->GetStride(), sizeof(FRHIDrawIndexedIndirectParameters), 0);
 		SetupDrawing();
 		Tracker->Assert(ArgumentsBufferRHI->GetWholeResourceIdentity(), ERHIAccess::IndirectArgs);
 		Tracker->Assert(IndexBufferRHI->GetWholeResourceIdentity(), ERHIAccess::VertexOrIndexBuffer);
@@ -878,7 +878,7 @@ public:
 	virtual void RHIDrawIndexedPrimitiveIndirect(FRHIBuffer* IndexBuffer, FRHIBuffer* ArgumentBuffer, uint32 ArgumentOffset) override final
 	{
 		checkf(State.bGfxPSOSet, TEXT("A Graphics PSO has to be set to draw!"));
-		FValidationRHI::ValidateIndirectArgsBuffer(ArgumentBuffer, ArgumentOffset);
+		FValidationRHI::ValidateIndirectArgsBuffer(ArgumentBuffer, ArgumentOffset, sizeof(FRHIDrawIndexedIndirectParameters), 0);
 		SetupDrawing();
 		Tracker->Assert(ArgumentBuffer->GetWholeResourceIdentity(), ERHIAccess::IndirectArgs);
 		Tracker->Assert(IndexBuffer->GetWholeResourceIdentity(), ERHIAccess::VertexOrIndexBuffer);
@@ -898,7 +898,7 @@ public:
 	virtual void RHIDispatchIndirectMeshShader(FRHIBuffer* ArgumentBuffer, uint32 ArgumentOffset) final override
 	{
 		checkf(State.bGfxPSOSet, TEXT("A Graphics PSO has to be set to draw!"));
-		FValidationRHI::ValidateIndirectArgsBuffer(ArgumentBuffer, ArgumentOffset);
+		FValidationRHI::ValidateDispatchIndirectArgsBuffer(ArgumentBuffer, ArgumentOffset);
 		SetupDrawing();
 		Tracker->Assert(ArgumentBuffer->GetWholeResourceIdentity(), ERHIAccess::IndirectArgs);
 		RHIContext->RHIDispatchIndirectMeshShader(ArgumentBuffer, ArgumentOffset);
@@ -1126,7 +1126,7 @@ public:
 		const FRayTracingShaderBindings& GlobalResourceBindings,
 		FRHIBuffer* ArgumentBuffer, uint32 ArgumentOffset) override final
 	{
-		FValidationRHI::ValidateIndirectArgsBuffer(ArgumentBuffer, ArgumentOffset);
+		FValidationRHI::ValidateDispatchIndirectArgsBuffer(ArgumentBuffer, ArgumentOffset);
 		Tracker->Assert(ArgumentBuffer->GetWholeResourceIdentity(), ERHIAccess::IndirectArgs);
 		Tracker->Assert(ArgumentBuffer->GetWholeResourceIdentity(), ERHIAccess::SRVCompute);
 
