Added missing Samsung vendor ID/string

diff --git a/Engine/Source/Runtime/RHI/Private/RHIStrings.cpp b/Engine/Source/Runtime/RHI/Private/RHIStrings.cpp
--- a/Engine/Source/Runtime/RHI/Private/RHIStrings.cpp
+++ b/Engine/Source/Runtime/RHI/Private/RHIStrings.cpp
@@ -71,6 +71,7 @@ const TCHAR* RHIVendorIdToString(EGpuVendorId VendorId)
 	case EGpuVendorId::Kazan:		return TEXT("Kazan");
 	case EGpuVendorId::Codeplay:	return TEXT("Codeplay");
 	case EGpuVendorId::Mesa:		return TEXT("Mesa");
+	case EGpuVendorId::SamsungAMD:	return TEXT("SamsungAMD");
 	case EGpuVendorId::NotQueried:	return TEXT("Not Queried");
 	default:                        return TEXT("Unknown");
 	}
diff --git a/Engine/Source/Runtime/RHI/Public/RHIDefinitions.h b/Engine/Source/Runtime/RHI/Public/RHIDefinitions.h
--- a/Engine/Source/Runtime/RHI/Public/RHIDefinitions.h
+++ b/Engine/Source/Runtime/RHI/Public/RHIDefinitions.h
@@ -819,6 +819,7 @@ enum class EGpuVendorId
 	Apple		= 0x106B,
 	Vivante		= 0x7a05,
 	VeriSilicon	= 0x1EB1,
+	SamsungAMD  = 0x144D,
 
 	Kazan		= 0x10003,	// VkVendorId
 	Codeplay	= 0x10004,	// VkVendorId
@@ -1284,6 +1285,7 @@ inline EGpuVendorId RHIConvertToGpuVendorId(uint32 VendorId)
 	case EGpuVendorId::Broadcom:
 	case EGpuVendorId::Qualcomm:
 	case EGpuVendorId::Intel:
+	case EGpuVendorId::SamsungAMD:
 		return (EGpuVendorId)VendorId;
 
 	default:
