Header unit fixes

* Added ScriptDelegateFwd.h and put it in the right module (CoreUObject). before fwd declarations was in Core.
* Fixed problem with StructuredArchive forward declaring and used functions that was force inlined in another file. header units did not like this and caused linker errors (and I'm surprised it worked before)
* Various fixes related to how global variables are defined (static for global variables means that they are not visible outside header units)
* Various fixes related to how global functions are defined (static for global functions means that they are not visible outside header units)
* Move some method implementations to cpp file to work around header unit compile crashes (and they should probably be there firstplace)
* Removed *_API on some types to fix header unit linker errors
* Added WITH_EDITOR around code that only compiles in editor config (these headers were not included anywhere but HU can't know that)

diff --git a/Engine/Plugins/Experimental/PythonScriptPlugin/Source/PythonScriptPlugin/Private/PyReferenceCollector.h b/Engine/Plugins/Experimental/PythonScriptPlugin/Source/PythonScriptPlugin/Private/PyReferenceCollector.h
--- a/Engine/Plugins/Experimental/PythonScriptPlugin/Source/PythonScriptPlugin/Private/PyReferenceCollector.h
+++ b/Engine/Plugins/Experimental/PythonScriptPlugin/Source/PythonScriptPlugin/Private/PyReferenceCollector.h
@@ -3,6 +3,7 @@
 #pragma once
 
 #include "UObject/GCObject.h"
+#include "UObject/ScriptDelegateFwd.h"
 #include "UObject/WeakObjectPtr.h"
 #include "Containers/ArrayView.h"
 
diff --git a/Engine/Source/Runtime/AudioMixerCore/Public/AudioMixerTypes.h b/Engine/Source/Runtime/AudioMixerCore/Public/AudioMixerTypes.h
--- a/Engine/Source/Runtime/AudioMixerCore/Public/AudioMixerTypes.h
+++ b/Engine/Source/Runtime/AudioMixerCore/Public/AudioMixerTypes.h
@@ -45,17 +45,17 @@ namespace Audio {
 	}
 
 	// Indicates a platform-specific format
-	static FName NAME_PLATFORM_SPECIFIC(TEXT("PLATFORM_SPECIFIC"));
-	static FName NAME_PROJECT_DEFINED(TEXT("PROJECT_DEFINED"));
+	inline FName NAME_PLATFORM_SPECIFIC(TEXT("PLATFORM_SPECIFIC"));
+	inline FName NAME_PROJECT_DEFINED(TEXT("PROJECT_DEFINED"));
 
 	// Supported on all platforms:
-	static FName NAME_BINKA(TEXT("BINKA"));
-	static FName NAME_ADPCM(TEXT("ADPCM"));
-	static FName NAME_PCM(TEXT("PCM"));
+	inline FName NAME_BINKA(TEXT("BINKA"));
+	inline FName NAME_ADPCM(TEXT("ADPCM"));
+	inline FName NAME_PCM(TEXT("PCM"));
 
 	// Not yet supported on all platforms as a selectable option so is included under "platform specific" enumeration for now. 
-	static FName NAME_OGG(TEXT("OGG"));
-	static FName NAME_OPUS(TEXT("OPUS"));
+	inline FName NAME_OGG(TEXT("OGG"));
+	inline FName NAME_OPUS(TEXT("OPUS"));
 }
 
 struct AUDIOMIXERCORE_API FAudioPlatformSettings
diff --git a/Engine/Source/Runtime/Core/Private/Serialization/StructuredArchiveSlots.cpp b/Engine/Source/Runtime/Core/Private/Serialization/StructuredArchiveSlots.cpp
--- a/Engine/Source/Runtime/Core/Private/Serialization/StructuredArchiveSlots.cpp
+++ b/Engine/Source/Runtime/Core/Private/Serialization/StructuredArchiveSlots.cpp
@@ -4,6 +4,25 @@
 #include "Serialization/StructuredArchiveContainer.h"
 #include "Containers/UnrealString.h"
 
+FArchive& UE::StructuredArchive::Private::FSlotBase::GetUnderlyingArchive() const
+{
+	return GetUnderlyingArchiveImpl(StructuredArchive);
+}
+
+const FArchiveState& UE::StructuredArchive::Private::FSlotBase::GetArchiveState() const
+{
+	return GetUnderlyingArchiveStateImpl(StructuredArchive);
+}
+
+bool FStructuredArchiveSlot::IsFilled() const
+{
+#if WITH_TEXT_ARCHIVE_SUPPORT
+	return UE::StructuredArchive::Private::GetCurrentSlotElementIdImpl(StructuredArchive) != ElementId;
+#else
+	return true;
+#endif
+}
+
 #if WITH_TEXT_ARCHIVE_SUPPORT
 
 //////////// FStructuredArchiveSlot ////////////
diff --git a/Engine/Source/Runtime/Core/Public/Misc/CoreMiscDefines.h b/Engine/Source/Runtime/Core/Public/Misc/CoreMiscDefines.h
--- a/Engine/Source/Runtime/Core/Public/Misc/CoreMiscDefines.h
+++ b/Engine/Source/Runtime/Core/Public/Misc/CoreMiscDefines.h
@@ -320,12 +320,6 @@ struct TStaticDeprecateExpression
  */
 struct FPlatformUserId
 {
-	/** Create a default invalid Id */
-	FORCEINLINE FPlatformUserId() : InternalId(INDEX_NONE) {}
-
-	FPlatformUserId(const FPlatformUserId&) = default;
-	FPlatformUserId& operator=(const FPlatformUserId&) = default;
-
 	/** Sees if this is a valid user */
 	FORCEINLINE bool IsValid() const
 	{
@@ -361,21 +355,17 @@ struct FPlatformUserId
 		return UserId.InternalId;
 	}
 
-	// This is needed until input and online code handles FPlatformUserId properly
-	UE_DEPRECATED(5.0, "Implicit conversion from user index is deprecated, use FPlatformMisc::GetPlatformUserForUserIndex")
-	FORCEINLINE constexpr FPlatformUserId(int32 InIndex) : InternalId(InIndex) {}
-
 	// This should be deprecated when the online code uniformly handles FPlatformUserId */
 	// UE_DEPRECATED(5.x, "Implicit conversion to user index is deprecated, use FPlatformMisc::GetUserIndexForPlatformUser")
 	FORCEINLINE constexpr operator int32() const { return InternalId; }
 
 private:
 	/** Raw id, will be allocated by application layer */
-	int32 InternalId;
+	int32 InternalId = INDEX_NONE;
 };
 
 /** Static invalid platform user */
-const FPlatformUserId PLATFORMUSERID_NONE;
+inline constexpr FPlatformUserId PLATFORMUSERID_NONE;
 
 /**
  * Represents a single input device such as a gamepad, keyboard, or mouse.
@@ -384,12 +374,6 @@ const FPlatformUserId PLATFORMUSERID_NONE;
  */
 struct FInputDeviceId
 {
-	FORCEINLINE explicit FInputDeviceId()
-		: InternalId(INDEX_NONE)
-	{}
-	
-	FInputDeviceId(const FInputDeviceId&) = default;
-	
 	/** Explicit function to create from an internal id */
 	FORCEINLINE static FInputDeviceId CreateFromInternalId(int32 InInternalId)
 	{
@@ -451,11 +435,11 @@ private:
 	 * 
 	 * @see IPlatformInputDeviceMapper::AllocateNewInputDeviceId
 	 */
-	int32 InternalId;
+	int32 InternalId = INDEX_NONE;
 };
 
 /** Static invalid input device. */
-const FInputDeviceId INPUTDEVICEID_NONE;
+inline constexpr FInputDeviceId INPUTDEVICEID_NONE;
 
 /** Represents the connection status of a given FInputDeviceId */
 enum class EInputDeviceConnectionState : uint8
diff --git a/Engine/Source/Runtime/Core/Public/Serialization/StructuredArchiveSlotBase.h b/Engine/Source/Runtime/Core/Public/Serialization/StructuredArchiveSlotBase.h
--- a/Engine/Source/Runtime/Core/Public/Serialization/StructuredArchiveSlotBase.h
+++ b/Engine/Source/Runtime/Core/Public/Serialization/StructuredArchiveSlotBase.h
@@ -10,9 +10,6 @@ struct FArchiveState;
 
 namespace UE::StructuredArchive::Private
 {
-	FArchive& GetUnderlyingArchiveImpl(FStructuredArchive& Ar);
-	FArchiveState& GetUnderlyingArchiveStateImpl(FStructuredArchive& Ar);
-
 	struct FElementId
 	{
 		FElementId() = default;
@@ -89,15 +86,8 @@ namespace UE::StructuredArchive::Private
 		}
 #endif
 
-		FArchive& GetUnderlyingArchive() const
-		{
-			return GetUnderlyingArchiveImpl(StructuredArchive);
-		}
-
-		const FArchiveState& GetArchiveState() const
-		{
-			return GetUnderlyingArchiveStateImpl(StructuredArchive);
-		}
+		FArchive& GetUnderlyingArchive() const;
+		const FArchiveState& GetArchiveState() const;
 
 	protected:
 		FStructuredArchive& StructuredArchive;
diff --git a/Engine/Source/Runtime/Core/Public/Serialization/StructuredArchiveSlots.h b/Engine/Source/Runtime/Core/Public/Serialization/StructuredArchiveSlots.h
--- a/Engine/Source/Runtime/Core/Public/Serialization/StructuredArchiveSlots.h
+++ b/Engine/Source/Runtime/Core/Public/Serialization/StructuredArchiveSlots.h
@@ -34,7 +34,6 @@ template <class TEnum> class TEnumAsByte;
 
 namespace UE::StructuredArchive::Private
 {
-	FElementId GetCurrentSlotElementIdImpl(FStructuredArchive& StructuredArchive);
 	FArchiveFormatterType& GetFormatterImpl(FStructuredArchive& StructuredArchive);
 }
 
@@ -113,15 +112,7 @@ public:
 
 	void Serialize(TArray<uint8>& Data);
 	void Serialize(void* Data, uint64 DataSize);
-
-	FORCEINLINE bool IsFilled() const
-	{
-#if WITH_TEXT_ARCHIVE_SUPPORT
-		return UE::StructuredArchive::Private::GetCurrentSlotElementIdImpl(StructuredArchive) != ElementId;
-#else
-		return true;
-#endif
-	}
+	bool IsFilled() const;
 
 private:
 	friend FStructuredArchive;
diff --git a/Engine/Source/Runtime/Core/Public/UObject/ScriptDelegates.h b/Engine/Source/Runtime/Core/Public/UObject/ScriptDelegates.h
--- a/Engine/Source/Runtime/Core/Public/UObject/ScriptDelegates.h
+++ b/Engine/Source/Runtime/Core/Public/UObject/ScriptDelegates.h
@@ -12,12 +12,10 @@
 #include "UObject/NameTypes.h"
 #include "Templates/SharedPointer.h"
 
-struct FWeakObjectPtr;
-
 /**
  * Script delegate base class.
  */
-template <typename TWeakPtr = FWeakObjectPtr>
+template <typename TWeakPtr /*= FWeakObjectPtr*/>
 class TScriptDelegate
 {
 	// Although templated, the parameter is not intended to be anything other than the default,
@@ -317,7 +315,7 @@ template<typename TWeakPtr> struct TIsZeroConstructType<TScriptDelegate<TWeakPtr
 /**
  * Script multi-cast delegate base class
  */
-template <typename TWeakPtr = FWeakObjectPtr>
+template <typename TWeakPtr/* = FWeakObjectPtr*/>
 class TMulticastScriptDelegate
 {
 public:
diff --git a/Engine/Source/Runtime/CoreUObject/Private/UObject/PropertyHelper.h b/Engine/Source/Runtime/CoreUObject/Private/UObject/PropertyHelper.h
--- a/Engine/Source/Runtime/CoreUObject/Private/UObject/PropertyHelper.h
+++ b/Engine/Source/Runtime/CoreUObject/Private/UObject/PropertyHelper.h
@@ -3,6 +3,7 @@
 
 #include "CoreTypes.h"
 #include "Logging/LogMacros.h"
+#include "UObject/ScriptDelegateFwd.h"
 #include "UObject/WeakObjectPtrFwd.h"
 
 class UFunction;
diff --git a/Engine/Source/Runtime/CoreUObject/Public/UObject/ScriptDelegateFwd.h b/Engine/Source/Runtime/CoreUObject/Public/UObject/ScriptDelegateFwd.h
new file mode 100644
--- /dev/null
+++ b/Engine/Source/Runtime/CoreUObject/Public/UObject/ScriptDelegateFwd.h
@@ -0,0 +1,16 @@
+// Copyright Epic Games, Inc. All Rights Reserved.
+
+/*=============================================================================
+	ScriptDelegateFwd.h: Delegate forward declarations
+=============================================================================*/
+
+#pragma once
+
+#include "UObject/WeakObjectPtrFwd.h"
+
+template <typename TWeakPtr = FWeakObjectPtr> class TScriptDelegate;
+template <typename TWeakPtr = FWeakObjectPtr> class TMulticastScriptDelegate;
+
+// Typedef script delegates for convenience.
+typedef TScriptDelegate<FWeakObjectPtr> FScriptDelegate;
+typedef TMulticastScriptDelegate<FWeakObjectPtr> FMulticastScriptDelegate;
diff --git a/Engine/Source/Runtime/CoreUObject/Public/UObject/SparseDelegate.h b/Engine/Source/Runtime/CoreUObject/Public/UObject/SparseDelegate.h
--- a/Engine/Source/Runtime/CoreUObject/Public/UObject/SparseDelegate.h
+++ b/Engine/Source/Runtime/CoreUObject/Public/UObject/SparseDelegate.h
@@ -15,6 +15,7 @@
 #include "UObject/NameTypes.h"
 #include "UObject/Object.h"
 #include "UObject/ObjectMacros.h"
+#include "UObject/ScriptDelegateFwd.h"
 #include "UObject/UObjectArray.h"
 #include "UObject/WeakObjectPtr.h"
 
diff --git a/Engine/Source/Runtime/CoreUObject/Public/UObject/WeakObjectPtrFwd.h b/Engine/Source/Runtime/CoreUObject/Public/UObject/WeakObjectPtrFwd.h
--- a/Engine/Source/Runtime/CoreUObject/Public/UObject/WeakObjectPtrFwd.h
+++ b/Engine/Source/Runtime/CoreUObject/Public/UObject/WeakObjectPtrFwd.h
@@ -17,10 +17,3 @@ template <typename T> struct TIsZeroConstructType;
 template<> struct TIsPODType<FWeakObjectPtr> { enum { Value = true }; };
 template<> struct TIsZeroConstructType<FWeakObjectPtr> { enum { Value = true }; };
 template<> struct TIsWeakPointerType<FWeakObjectPtr> { enum { Value = true }; };
-
-template <typename T> class TScriptDelegate;
-template <typename T> class TMulticastScriptDelegate;
-
-// Typedef script delegates for convenience.
-typedef TScriptDelegate<FWeakObjectPtr> FScriptDelegate;
-typedef TMulticastScriptDelegate<FWeakObjectPtr> FMulticastScriptDelegate;
diff --git a/Engine/Source/Runtime/Engine/Classes/Components/ExternalMorphSet.h b/Engine/Source/Runtime/Engine/Classes/Components/ExternalMorphSet.h
--- a/Engine/Source/Runtime/Engine/Classes/Components/ExternalMorphSet.h
+++ b/Engine/Source/Runtime/Engine/Classes/Components/ExternalMorphSet.h
@@ -6,7 +6,7 @@
 #include "Rendering/MorphTargetVertexInfoBuffers.h"
 
 /** An external morph target set. External morph targets are managed by systems outside of the skinned meshes. */
-struct ENGINE_API FExternalMorphSet
+struct FExternalMorphSet
 {
 	/** A name for this set, useful for debugging. */
 	FName Name = FName(TEXT("Unknown"));
diff --git a/Engine/Source/Runtime/Engine/Classes/Engine/SkeletalMesh.h b/Engine/Source/Runtime/Engine/Classes/Engine/SkeletalMesh.h
--- a/Engine/Source/Runtime/Engine/Classes/Engine/SkeletalMesh.h
+++ b/Engine/Source/Runtime/Engine/Classes/Engine/SkeletalMesh.h
@@ -816,7 +816,7 @@ private:
 	/** Struct containing information for each LOD level, such as materials to use, and when use the LOD. */
 	UE_DEPRECATED(5.0, "This must be protected for async build, always use the accessors even internally.")
 	UPROPERTY(EditAnywhere, EditFixedSize, Category=LevelOfDetail)
-	TArray<struct FSkeletalMeshLODInfo> LODInfo;
+	TArray<FSkeletalMeshLODInfo> LODInfo;
 
 #if !WITH_EDITOR
 	/** Acceleration struct used for faster socket lookups */
diff --git a/Engine/Source/Runtime/Engine/Private/Materials/MaterialInstance.cpp b/Engine/Source/Runtime/Engine/Private/Materials/MaterialInstance.cpp
--- a/Engine/Source/Runtime/Engine/Private/Materials/MaterialInstance.cpp
+++ b/Engine/Source/Runtime/Engine/Private/Materials/MaterialInstance.cpp
@@ -147,6 +147,16 @@ FMaterialInstanceResource::FMaterialInstanceResource(UMaterialInstance* InOwner)
 {
 }
 
+void FMaterialInstanceResource::GameThread_Destroy()
+{
+	FMaterialInstanceResource* Resource = this;
+	ENQUEUE_RENDER_COMMAND(FDestroyMaterialInstanceResourceCommand)(
+		[Resource](FRHICommandList& RHICmdList)
+		{
+			delete Resource;
+		});
+}
+
 #if 0
 const FMaterial& FMaterialInstanceResource::GetMaterialWithFallback(ERHIFeatureLevel::Type InFeatureLevel, const FMaterialRenderProxy*& OutFallbackMaterialRenderProxy) const
 {
diff --git a/Engine/Source/Runtime/Engine/Private/Materials/MaterialInstanceSupport.h b/Engine/Source/Runtime/Engine/Private/Materials/MaterialInstanceSupport.h
--- a/Engine/Source/Runtime/Engine/Private/Materials/MaterialInstanceSupport.h
+++ b/Engine/Source/Runtime/Engine/Private/Materials/MaterialInstanceSupport.h
@@ -210,15 +210,7 @@ public:
 	/**
 	 * Called from the game thread to destroy the material instance on the rendering thread.
 	 */
-	void GameThread_Destroy()
-	{
-		FMaterialInstanceResource* Resource = this;
-		ENQUEUE_RENDER_COMMAND(FDestroyMaterialInstanceResourceCommand)(
-			[Resource](FRHICommandList& RHICmdList)
-			{
-				delete Resource;
-			});
-	}
+	void GameThread_Destroy();
 
 	// FRenderResource interface.
 	virtual FString GetFriendlyName() const override { return Owner->GetName(); }
diff --git a/Engine/Source/Runtime/Engine/Private/SkeletalRenderCPUSkin.cpp b/Engine/Source/Runtime/Engine/Private/SkeletalRenderCPUSkin.cpp
--- a/Engine/Source/Runtime/Engine/Private/SkeletalRenderCPUSkin.cpp
+++ b/Engine/Source/Runtime/Engine/Private/SkeletalRenderCPUSkin.cpp
@@ -562,6 +562,8 @@ FDynamicSkelMeshObjectDataCPUSkin::FDynamicSkelMeshObjectDataCPUSkin(
 #endif
 }
 
+FDynamicSkelMeshObjectDataCPUSkin::~FDynamicSkelMeshObjectDataCPUSkin() = default;
+
 /*-----------------------------------------------------------------------------
 	FSkeletalMeshObjectCPUSkin - morph target blending implementation
 -----------------------------------------------------------------------------*/
diff --git a/Engine/Source/Runtime/Engine/Private/SkeletalRenderCPUSkin.h b/Engine/Source/Runtime/Engine/Private/SkeletalRenderCPUSkin.h
--- a/Engine/Source/Runtime/Engine/Private/SkeletalRenderCPUSkin.h
+++ b/Engine/Source/Runtime/Engine/Private/SkeletalRenderCPUSkin.h
@@ -25,7 +25,7 @@ class UMorphTarget;
 * Stores the updated matrices needed to skin the verts.
 * Created by the game thread and sent to the rendering thread as an update 
 */
-class ENGINE_API FDynamicSkelMeshObjectDataCPUSkin
+class FDynamicSkelMeshObjectDataCPUSkin
 {
 public:
 
@@ -37,7 +37,7 @@ public:
 	* @param	InActiveMorphTargets - Active Morph Targets to blend with during skinning
 	* @param	InMorphTargetWeights - All Morph Target weights to blend with during skinning
 	*/
-	FDynamicSkelMeshObjectDataCPUSkin(
+	ENGINE_API FDynamicSkelMeshObjectDataCPUSkin(
 		USkinnedMeshComponent* InMeshComponent,
 		FSkeletalMeshRenderData* InSkelMeshRenderData,
 		int32 InLODIndex,
@@ -45,9 +45,7 @@ public:
 		const TArray<float>& InMorphTargetWeights
 		);
 
-	virtual ~FDynamicSkelMeshObjectDataCPUSkin()
-	{
-	}
+	ENGINE_API virtual ~FDynamicSkelMeshObjectDataCPUSkin();
 
 	/** Local to world transform, used for cloth as sim data is in world space */
 	FMatrix WorldToLocal;
@@ -83,31 +81,31 @@ public:
 
 	/** Update Simulated Positions & Normals from Clothing actor */
 	UE_DEPRECATED(5.2, "Use USkinnedMeshComponent::GetUpdateClothSimulationData_AnyThread() instead.")
-	bool UpdateClothSimulationData(USkinnedMeshComponent* InMeshComponent);
+	ENGINE_API bool UpdateClothSimulationData(USkinnedMeshComponent* InMeshComponent);
 };
 
 /**
  * Render data for a CPU skinned mesh
  */
-class ENGINE_API FSkeletalMeshObjectCPUSkin : public FSkeletalMeshObject
+class FSkeletalMeshObjectCPUSkin : public FSkeletalMeshObject
 {
 public:
 
 	/** @param	InSkeletalMeshComponent - skeletal mesh primitive we want to render */
-	FSkeletalMeshObjectCPUSkin(USkinnedMeshComponent* InMeshComponent, FSkeletalMeshRenderData* InSkelMeshRenderData, ERHIFeatureLevel::Type InFeatureLevel);
-	virtual ~FSkeletalMeshObjectCPUSkin();
+	ENGINE_API FSkeletalMeshObjectCPUSkin(USkinnedMeshComponent* InMeshComponent, FSkeletalMeshRenderData* InSkelMeshRenderData, ERHIFeatureLevel::Type InFeatureLevel);
+	ENGINE_API virtual ~FSkeletalMeshObjectCPUSkin();
 
 	//~ Begin FSkeletalMeshObject Interface
-	virtual void InitResources(USkinnedMeshComponent* InMeshComponent) override;
-	virtual void ReleaseResources() override;
-	virtual void Update(int32 LODIndex,USkinnedMeshComponent* InMeshComponent,const FMorphTargetWeightMap& InActiveMorphTargets, const TArray<float>& MorphTargetsWeights, EPreviousBoneTransformUpdateMode PreviousBoneTransformUpdateMode, const FExternalMorphWeightData& InExternalMorphWeightData) override;
-	void UpdateDynamicData_RenderThread(FRHICommandListImmediate& RHICmdList, FDynamicSkelMeshObjectDataCPUSkin* InDynamicData, uint32 FrameNumberToPrepare, uint32 RevisionNumber);
-	virtual void EnableOverlayRendering(bool bEnabled, const TArray<int32>* InBonesOfInterest, const TArray<UMorphTarget*>* InMorphTargetOfInterest) override;
-	virtual void CacheVertices(int32 LODIndex, bool bForce) const override;
+	ENGINE_API virtual void InitResources(USkinnedMeshComponent* InMeshComponent) override;
+	ENGINE_API virtual void ReleaseResources() override;
+	ENGINE_API virtual void Update(int32 LODIndex,USkinnedMeshComponent* InMeshComponent,const FMorphTargetWeightMap& InActiveMorphTargets, const TArray<float>& MorphTargetsWeights, EPreviousBoneTransformUpdateMode PreviousBoneTransformUpdateMode, const FExternalMorphWeightData& InExternalMorphWeightData) override;
+	ENGINE_API void UpdateDynamicData_RenderThread(FRHICommandListImmediate& RHICmdList, FDynamicSkelMeshObjectDataCPUSkin* InDynamicData, uint32 FrameNumberToPrepare, uint32 RevisionNumber);
+	ENGINE_API virtual void EnableOverlayRendering(bool bEnabled, const TArray<int32>* InBonesOfInterest, const TArray<UMorphTarget*>* InMorphTargetOfInterest) override;
+	ENGINE_API virtual void CacheVertices(int32 LODIndex, bool bForce) const override;
 	virtual bool IsCPUSkinned() const override { return true; }
-	virtual const FVertexFactory* GetSkinVertexFactory(const FSceneView* View, int32 LODIndex, int32 ChunkIdx, ESkinVertexFactoryMode VFMode = ESkinVertexFactoryMode::Default) const override;
-	virtual TArray<FTransform>* GetComponentSpaceTransforms() const override;
-	virtual const TArray<FMatrix44f>& GetReferenceToLocalMatrices() const override;
+	ENGINE_API virtual const FVertexFactory* GetSkinVertexFactory(const FSceneView* View, int32 LODIndex, int32 ChunkIdx, ESkinVertexFactoryMode VFMode = ESkinVertexFactoryMode::Default) const override;
+	ENGINE_API virtual TArray<FTransform>* GetComponentSpaceTransforms() const override;
+	ENGINE_API virtual const TArray<FMatrix44f>& GetReferenceToLocalMatrices() const override;
 
 	virtual int32 GetLOD() const override
 	{
@@ -147,14 +145,14 @@ public:
 		CumulativeResourceSize.AddDedicatedSystemMemoryBytes(BonesOfInterest.GetAllocatedSize());
 	}
 
-	virtual void DrawVertexElements(FPrimitiveDrawInterface* PDI, const FMatrix& ToWorldSpace, bool bDrawNormals, bool bDrawTangents, bool bDrawBinormals) const override;
+	ENGINE_API virtual void DrawVertexElements(FPrimitiveDrawInterface* PDI, const FMatrix& ToWorldSpace, bool bDrawNormals, bool bDrawTangents, bool bDrawBinormals) const override;
 	//~ End FSkeletalMeshObject Interface
 
 	/** Access cached final vertices */
 	const TArray<FFinalSkinVertex>& GetCachedFinalVertices() const { return CachedFinalVertices; }
 
 
-	virtual void UpdateSkinWeightBuffer(USkinnedMeshComponent* InMeshComponent) override;
+	ENGINE_API virtual void UpdateSkinWeightBuffer(USkinnedMeshComponent* InMeshComponent) override;
 
 private:
 	/** vertex data for rendering a single LOD */
diff --git a/Engine/Source/Runtime/Engine/Private/SkeletalRenderStatic.h b/Engine/Source/Runtime/Engine/Private/SkeletalRenderStatic.h
--- a/Engine/Source/Runtime/Engine/Private/SkeletalRenderStatic.h
+++ b/Engine/Source/Runtime/Engine/Private/SkeletalRenderStatic.h
@@ -17,24 +17,24 @@
 class FPrimitiveDrawInterface;
 class UMorphTarget;
 
-class ENGINE_API FSkeletalMeshObjectStatic : public FSkeletalMeshObject
+class FSkeletalMeshObjectStatic : public FSkeletalMeshObject
 {
 public:
 	/** @param	InSkeletalMeshComponent - skeletal mesh primitive we want to render */
-	FSkeletalMeshObjectStatic(USkinnedMeshComponent* InMeshComponent, FSkeletalMeshRenderData* InSkelMeshRenderData, ERHIFeatureLevel::Type InFeatureLevel);
-	virtual ~FSkeletalMeshObjectStatic();
+	ENGINE_API FSkeletalMeshObjectStatic(USkinnedMeshComponent* InMeshComponent, FSkeletalMeshRenderData* InSkelMeshRenderData, ERHIFeatureLevel::Type InFeatureLevel);
+	ENGINE_API virtual ~FSkeletalMeshObjectStatic();
 
 	//~ Begin FSkeletalMeshObject Interface
-	virtual void InitResources(USkinnedMeshComponent* InMeshComponent) override;
-	virtual void ReleaseResources() override;
+	ENGINE_API virtual void InitResources(USkinnedMeshComponent* InMeshComponent) override;
+	ENGINE_API virtual void ReleaseResources() override;
 	virtual void Update(int32 LODIndex,USkinnedMeshComponent* InMeshComponent,const FMorphTargetWeightMap& InActiveMorphTargets, const TArray<float>& MorphTargetWeights, EPreviousBoneTransformUpdateMode PreviousBoneTransformUpdateMode, const FExternalMorphWeightData& InExternalMorphWeightData) override {};
 	//virtual void UpdateRecomputeTangent(int32 MaterialIndex, int32 LODIndex, bool bRecomputeTangent) override {};
 	virtual void EnableOverlayRendering(bool bEnabled, const TArray<int32>* InBonesOfInterest, const TArray<UMorphTarget*>* InMorphTargetOfInterest) override {};
 	virtual void CacheVertices(int32 LODIndex, bool bForce) const override {};
 	virtual bool IsCPUSkinned() const override { return true; }
-	virtual const FVertexFactory* GetSkinVertexFactory(const FSceneView* View, int32 LODIndex, int32 ChunkIdx, ESkinVertexFactoryMode VFMode = ESkinVertexFactoryMode::Default) const override;
-	virtual TArray<FTransform>* GetComponentSpaceTransforms() const override;
-	virtual const TArray<FMatrix44f>& GetReferenceToLocalMatrices() const override;
+	ENGINE_API virtual const FVertexFactory* GetSkinVertexFactory(const FSceneView* View, int32 LODIndex, int32 ChunkIdx, ESkinVertexFactoryMode VFMode = ESkinVertexFactoryMode::Default) const override;
+	ENGINE_API virtual TArray<FTransform>* GetComponentSpaceTransforms() const override;
+	ENGINE_API virtual const TArray<FMatrix44f>& GetReferenceToLocalMatrices() const override;
 
 	virtual int32 GetLOD() const override
 	{
@@ -107,6 +107,6 @@ private:
 	};
 
 	/** Render data for each LOD */
-	TArray<struct FSkeletalMeshObjectLOD> LODs;
+	TArray<FSkeletalMeshObjectLOD> LODs;
 };
 
diff --git a/Engine/Source/Runtime/Engine/Public/Animation/AnimTypes.h b/Engine/Source/Runtime/Engine/Public/Animation/AnimTypes.h
--- a/Engine/Source/Runtime/Engine/Public/Animation/AnimTypes.h
+++ b/Engine/Source/Runtime/Engine/Public/Animation/AnimTypes.h
@@ -713,7 +713,7 @@ enum class EAnimInterpolationType : uint8
 	// ID type, should be used to access SmartNames as fundamental type may change.
 	typedef uint16 UID_Type;
 	// Max UID used for overflow checking
-	static const UID_Type MaxUID = MAX_uint16;
+	inline constexpr UID_Type MaxUID = MAX_uint16;
 }
 /**
  * Animation Key extraction helper as we have a lot of code that messes up the key length
diff --git a/Engine/Source/Runtime/Engine/Public/EdGraph/EdGraphNodeUtils.h b/Engine/Source/Runtime/Engine/Public/EdGraph/EdGraphNodeUtils.h
--- a/Engine/Source/Runtime/Engine/Public/EdGraph/EdGraphNodeUtils.h
+++ b/Engine/Source/Runtime/Engine/Public/EdGraph/EdGraphNodeUtils.h
@@ -6,6 +6,8 @@
 #include "EdGraph/EdGraphSchema.h"
 #include "Misc/ConfigCacheIni.h"
 
+#if WITH_EDITOR
+
 /*******************************************************************************
  * FNodeTextCache
  ******************************************************************************/
@@ -229,3 +231,5 @@ private:
 	/** */
 	FNodeTextCache CachedTooltip;
 };
+
+#endif
\ No newline at end of file
diff --git a/Engine/Source/Runtime/Engine/Public/IPIEPreviewDeviceModule.h b/Engine/Source/Runtime/Engine/Public/IPIEPreviewDeviceModule.h
--- a/Engine/Source/Runtime/Engine/Public/IPIEPreviewDeviceModule.h
+++ b/Engine/Source/Runtime/Engine/Public/IPIEPreviewDeviceModule.h
@@ -4,7 +4,6 @@
 #include "CoreMinimal.h"
 #include "IDeviceProfileSelectorModule.h"
 #include "Widgets/SWindow.h"
-#include "PIEPreviewDeviceEnumeration.h"
 
 class IPIEPreviewDeviceModule : public IDeviceProfileSelectorModule
 {
diff --git a/Engine/Source/Runtime/Engine/Public/Rendering/MorphTargetVertexInfoBuffers.h b/Engine/Source/Runtime/Engine/Public/Rendering/MorphTargetVertexInfoBuffers.h
--- a/Engine/Source/Runtime/Engine/Public/Rendering/MorphTargetVertexInfoBuffers.h
+++ b/Engine/Source/Runtime/Engine/Public/Rendering/MorphTargetVertexInfoBuffers.h
@@ -9,13 +9,13 @@
 struct FSkelMeshRenderSection;
 class UMorphTarget;
 
-class ENGINE_API FMorphTargetVertexInfoBuffers : public FRenderResource
+class FMorphTargetVertexInfoBuffers : public FRenderResource
 {
 public:
-	FMorphTargetVertexInfoBuffers();
-	~FMorphTargetVertexInfoBuffers();
+	ENGINE_API FMorphTargetVertexInfoBuffers();
+	ENGINE_API ~FMorphTargetVertexInfoBuffers();
 
-	void InitMorphResources(EShaderPlatform ShaderPlatform, const TArray<FSkelMeshRenderSection>& RenderSections, const TArray<UMorphTarget*>& MorphTargets, int NumVertices, int32 LODIndex, float TargetPositionErrorTolerance);
+	ENGINE_API void InitMorphResources(EShaderPlatform ShaderPlatform, const TArray<FSkelMeshRenderSection>& RenderSections, const TArray<UMorphTarget*>& MorphTargets, int NumVertices, int32 LODIndex, float TargetPositionErrorTolerance);
 
 	inline bool IsMorphResourcesInitialized() const { return bResourcesInitialized; }
 	inline bool IsRHIIntialized() const { return bRHIIntialized; }
@@ -24,11 +24,11 @@ public:
 	bool GetEmptyMorphCPUDataOnInitRHI() const { return bEmptyMorphCPUDataOnInitRHI; }
 	void SetEmptyMorphCPUDataOnInitRHI(bool bEmpty) { bEmptyMorphCPUDataOnInitRHI = bEmpty; }
 
-	virtual void InitRHI() override;
-	virtual void ReleaseRHI() override;
+	ENGINE_API virtual void InitRHI() override;
+	ENGINE_API virtual void ReleaseRHI() override;
 
 	UE_DEPRECATED(5.2, "GetMaximumThreadGroupSize will be removed as it is no longer used.")
-	static uint32 GetMaximumThreadGroupSize();
+	static ENGINE_API uint32 GetMaximumThreadGroupSize();
 
 	uint32 GetNumBatches(uint32 index = UINT_MAX) const
 	{
@@ -69,23 +69,23 @@ public:
 		return MorphData.Num() * sizeof(uint32);
 	}
 
-	static const float CalculatePositionPrecision(float TargetPositionErrorTolerance);
+	static ENGINE_API const float CalculatePositionPrecision(float TargetPositionErrorTolerance);
 
 	const float GetTangentZPrecision() const
 	{
 		return TangentZPrecision;
 	}
 
-	static bool IsPlatformShaderSupported(EShaderPlatform ShaderPlatform);
+	static ENGINE_API bool IsPlatformShaderSupported(EShaderPlatform ShaderPlatform);
 
 	FBufferRHIRef MorphDataBuffer;
 	FShaderResourceViewRHIRef MorphDataSRV;
 
 protected:
-	void ResetCPUData();
+	ENGINE_API void ResetCPUData();
 
-	void ValidateVertexBuffers(bool bMorphTargetsShouldBeValid);
-	void Serialize(FArchive& Ar);
+	ENGINE_API void ValidateVertexBuffers(bool bMorphTargetsShouldBeValid);
+	ENGINE_API void Serialize(FArchive& Ar);
 
 	// Transient data. Gets deleted as soon as the GPU resource has been initialized (unless excplicitly disabled by bEmptyMorphCPUDataOnInitRHI).
 	TResourceArray<uint32> MorphData;
diff --git a/Engine/Source/Runtime/Engine/Public/WorldPartition/HLOD/IWorldPartitionHLODUtilities.h b/Engine/Source/Runtime/Engine/Public/WorldPartition/HLOD/IWorldPartitionHLODUtilities.h
--- a/Engine/Source/Runtime/Engine/Public/WorldPartition/HLOD/IWorldPartitionHLODUtilities.h
+++ b/Engine/Source/Runtime/Engine/Public/WorldPartition/HLOD/IWorldPartitionHLODUtilities.h
@@ -8,6 +8,8 @@
 #include "WorldPartition/WorldPartitionStreamingGenerationContext.h"
 #include "WorldPartition/DataLayer/DataLayersID.h"
 
+#if WITH_EDITOR
+
 class AActor;
 class UWorldPartition;
 class UHLODLayer;
@@ -15,13 +17,13 @@ class UHLODBuilder;
 class UHLODBuilderSettings;
 class AWorldPartitionHLOD;
 
-struct ENGINE_API FHLODCreationContext
+struct FHLODCreationContext
 {
 	TMap<FName, FWorldPartitionHandle> HLODActorDescs;
 	TArray<FWorldPartitionReference> ActorReferences;
 };
 
-struct ENGINE_API FHLODCreationParams
+struct FHLODCreationParams
 {
 	UWorldPartition* WorldPartition;
 
@@ -87,3 +89,4 @@ public:
 		return CreateHLODActors(InCreationContext, InCreationParams, InActors);
 	}
 };
+#endif
\ No newline at end of file
diff --git a/Engine/Source/Runtime/Engine/Public/WorldPartition/HLOD/IWorldPartitionHLODUtilitiesModule.h b/Engine/Source/Runtime/Engine/Public/WorldPartition/HLOD/IWorldPartitionHLODUtilitiesModule.h
--- a/Engine/Source/Runtime/Engine/Public/WorldPartition/HLOD/IWorldPartitionHLODUtilitiesModule.h
+++ b/Engine/Source/Runtime/Engine/Public/WorldPartition/HLOD/IWorldPartitionHLODUtilitiesModule.h
@@ -5,8 +5,12 @@
 #include "Modules/ModuleInterface.h"
 #include "WorldPartition/HLOD/IWorldPartitionHLODUtilities.h"
 
+#if WITH_EDITOR
+
 class ENGINE_API IWorldPartitionHLODUtilitiesModule : public IModuleInterface
 {
 public:
 	virtual IWorldPartitionHLODUtilities* GetUtilities() = 0;
 };
+
+#endif
\ No newline at end of file
diff --git a/Engine/Source/Runtime/RHI/Public/RHISurfaceDataConversion.h b/Engine/Source/Runtime/RHI/Public/RHISurfaceDataConversion.h
--- a/Engine/Source/Runtime/RHI/Public/RHISurfaceDataConversion.h
+++ b/Engine/Source/Runtime/RHI/Public/RHISurfaceDataConversion.h
@@ -38,7 +38,7 @@ struct FRHIRG16
 	uint16 G;
 };
 
-static inline void ConvertRawR16DataToFColor(uint32 Width, uint32 Height, uint8 *In, uint32 SrcPitch, FColor* Out)
+inline void ConvertRawR16DataToFColor(uint32 Width, uint32 Height, uint8 *In, uint32 SrcPitch, FColor* Out)
 {
 	// e.g. shadow maps
 	for (uint32 Y = 0; Y < Height; Y++)
@@ -58,7 +58,7 @@ static inline void ConvertRawR16DataToFColor(uint32 Width, uint32 Height, uint8
 	}
 }
 
-static inline void ConvertRawR8G8B8A8DataToFColor(uint32 Width, uint32 Height, uint8 *In, uint32 SrcPitch, FColor* Out)
+inline void ConvertRawR8G8B8A8DataToFColor(uint32 Width, uint32 Height, uint8 *In, uint32 SrcPitch, FColor* Out)
 {
 	for (uint32 Y = 0; Y < Height; Y++)
 	{
@@ -73,7 +73,7 @@ static inline void ConvertRawR8G8B8A8DataToFColor(uint32 Width, uint32 Height, u
 	}
 }
 
-static inline void ConvertRawB8G8R8A8DataToFColor(uint32 Width, uint32 Height, uint8 *In, uint32 SrcPitch, FColor* Out)
+inline void ConvertRawB8G8R8A8DataToFColor(uint32 Width, uint32 Height, uint8 *In, uint32 SrcPitch, FColor* Out)
 {
 	const uint32 DstPitch = Width * sizeof(FColor);
 
@@ -96,7 +96,7 @@ static inline void ConvertRawB8G8R8A8DataToFColor(uint32 Width, uint32 Height, u
 	}
 }
 
-static inline void ConvertRawR16G16B16A16FDataToFFloat16Color(uint32 Width, uint32 Height, uint8* In, uint32 SrcPitch, FFloat16Color* Out)
+inline void ConvertRawR16G16B16A16FDataToFFloat16Color(uint32 Width, uint32 Height, uint8* In, uint32 SrcPitch, FFloat16Color* Out)
 {
 	const uint32 DstPitch = Width * sizeof(FFloat16Color);
 
@@ -119,7 +119,7 @@ static inline void ConvertRawR16G16B16A16FDataToFFloat16Color(uint32 Width, uint
 	}
 }
 
-static inline void ConvertRawR10G10B10A2DataToFColor(uint32 Width, uint32 Height, uint8 *In, uint32 SrcPitch, FColor* Out)
+inline void ConvertRawR10G10B10A2DataToFColor(uint32 Width, uint32 Height, uint8 *In, uint32 SrcPitch, FColor* Out)
 {
 	for (uint32 Y = 0; Y < Height; Y++)
 	{
@@ -139,7 +139,7 @@ static inline void ConvertRawR10G10B10A2DataToFColor(uint32 Width, uint32 Height
 	}
 }
 
-static inline void ConvertRawB10G10R10A2DataToFColor(uint32 Width, uint32 Height, uint8* In, uint32 SrcPitch, FColor* Out)
+inline void ConvertRawB10G10R10A2DataToFColor(uint32 Width, uint32 Height, uint8* In, uint32 SrcPitch, FColor* Out)
 {
 	for (uint32 Y = 0; Y < Height; Y++)
 	{
@@ -159,7 +159,7 @@ static inline void ConvertRawB10G10R10A2DataToFColor(uint32 Width, uint32 Height
 	}
 }
 
-static inline void ConvertRawR16G16B16A16FDataToFColor(uint32 Width, uint32 Height, uint8 *In, uint32 SrcPitch, FColor* Out, bool LinearToGamma)
+inline void ConvertRawR16G16B16A16FDataToFColor(uint32 Width, uint32 Height, uint8 *In, uint32 SrcPitch, FColor* Out, bool LinearToGamma)
 {
 	FPlane	MinValue(0.0f, 0.0f, 0.0f, 0.0f),
 		MaxValue(1.0f, 1.0f, 1.0f, 1.0f);
@@ -204,7 +204,7 @@ static inline void ConvertRawR16G16B16A16FDataToFColor(uint32 Width, uint32 Heig
 	}
 }
 
-static inline void ConvertRawR11G11B10DataToFColor(uint32 Width, uint32 Height, uint8 *In, uint32 SrcPitch, FColor* Out, bool LinearToGamma)
+inline void ConvertRawR11G11B10DataToFColor(uint32 Width, uint32 Height, uint8 *In, uint32 SrcPitch, FColor* Out, bool LinearToGamma)
 {
 	check(sizeof(FFloat3Packed) == sizeof(uint32));
 
@@ -224,7 +224,7 @@ static inline void ConvertRawR11G11B10DataToFColor(uint32 Width, uint32 Height,
 	}
 }
 
-static inline void ConvertRawR9G9B9E5DataToFColor(uint32 Width, uint32 Height, uint8* In, uint32 SrcPitch, FColor* Out, bool LinearToGamma)
+inline void ConvertRawR9G9B9E5DataToFColor(uint32 Width, uint32 Height, uint8* In, uint32 SrcPitch, FColor* Out, bool LinearToGamma)
 {
 	check(sizeof(FFloat3PackedSE) == sizeof(uint32));
 
@@ -244,7 +244,7 @@ static inline void ConvertRawR9G9B9E5DataToFColor(uint32 Width, uint32 Height, u
 	}
 }
 
-static inline void ConvertRawR32G32B32A32DataToFColor(uint32 Width, uint32 Height, uint8 *In, uint32 SrcPitch, FColor* Out, bool LinearToGamma)
+inline void ConvertRawR32G32B32A32DataToFColor(uint32 Width, uint32 Height, uint8 *In, uint32 SrcPitch, FColor* Out, bool LinearToGamma)
 {
 	FPlane MinValue(0.0f, 0.0f, 0.0f, 0.0f);
 	FPlane MaxValue(1.0f, 1.0f, 1.0f, 1.0f);
@@ -287,7 +287,8 @@ static inline void ConvertRawR32G32B32A32DataToFColor(uint32 Width, uint32 Heigh
 	}
 }
 
-static inline void ConvertRawR24G8DataToFColor(uint32 Width, uint32 Height, uint8 *In, uint32 SrcPitch, FColor* Out, FReadSurfaceDataFlags InFlags)
+
+inline void ConvertRawR24G8DataToFColor(uint32 Width, uint32 Height, uint8 *In, uint32 SrcPitch, FColor* Out, FReadSurfaceDataFlags InFlags)
 {
 	bool bLinearToGamma = InFlags.GetLinearToGamma();
 	// Depth stencil
@@ -318,7 +319,7 @@ static inline void ConvertRawR24G8DataToFColor(uint32 Width, uint32 Height, uint
 	}
 }
 
-static inline void ConvertRawR32DataToFColor(uint32 Width, uint32 Height, uint8 *In, uint32 SrcPitch, FColor* Out, FReadSurfaceDataFlags InFlags)
+inline void ConvertRawR32DataToFColor(uint32 Width, uint32 Height, uint8 *In, uint32 SrcPitch, FColor* Out, FReadSurfaceDataFlags InFlags)
 {
 	bool bLinearToGamma = InFlags.GetLinearToGamma();
 	for (uint32 Y = 0; Y < Height; Y++)
@@ -340,7 +341,7 @@ static inline void ConvertRawR32DataToFColor(uint32 Width, uint32 Height, uint8
 	}
 }
 
-static inline void ConvertRawR16G16B16A16DataToFColor(uint32 Width, uint32 Height, uint8 *In, uint32 SrcPitch, FColor* Out)
+inline void ConvertRawR16G16B16A16DataToFColor(uint32 Width, uint32 Height, uint8 *In, uint32 SrcPitch, FColor* Out)
 {
 	for (uint32 Y = 0; Y < Height; Y++)
 	{
@@ -360,7 +361,7 @@ static inline void ConvertRawR16G16B16A16DataToFColor(uint32 Width, uint32 Heigh
 	}
 }
 
-static inline void ConvertRawR16G16DataToFColor(uint32 Width, uint32 Height, uint8 *In, uint32 SrcPitch, FColor* Out)
+inline void ConvertRawR16G16DataToFColor(uint32 Width, uint32 Height, uint8 *In, uint32 SrcPitch, FColor* Out)
 {
 	for (uint32 Y = 0; Y < Height; Y++)
 	{
@@ -378,7 +379,7 @@ static inline void ConvertRawR16G16DataToFColor(uint32 Width, uint32 Height, uin
 	}
 }
 
-static inline void ConvertRawR8DataToFColor(uint32 Width, uint32 Height, uint8 *In, uint32 SrcPitch, FColor* Out)
+inline void ConvertRawR8DataToFColor(uint32 Width, uint32 Height, uint8 *In, uint32 SrcPitch, FColor* Out)
 {
 	for (uint32 Y = 0; Y < Height; Y++)
 	{
@@ -393,7 +394,7 @@ static inline void ConvertRawR8DataToFColor(uint32 Width, uint32 Height, uint8 *
 	}
 }
 
-static inline void ConvertRawD32S8DataToFColor(uint32 Width, uint32 Height, uint8 *In, uint32 SrcPitch, FColor* Out, FReadSurfaceDataFlags InFlags)
+inline void ConvertRawD32S8DataToFColor(uint32 Width, uint32 Height, uint8 *In, uint32 SrcPitch, FColor* Out, FReadSurfaceDataFlags InFlags)
 {
 	bool bLinearToGamma = InFlags.GetLinearToGamma();
 	// Depth
@@ -438,7 +439,7 @@ static inline void ConvertRawD32S8DataToFColor(uint32 Width, uint32 Height, uint
 
 // Linear functions
 
-static inline void ConvertRawR16DataToFLinearColor(uint32 Width, uint32 Height, uint8 *In, uint32 SrcPitch, FLinearColor* Out)
+inline void ConvertRawR16DataToFLinearColor(uint32 Width, uint32 Height, uint8 *In, uint32 SrcPitch, FLinearColor* Out)
 {
 	// e.g. shadow maps
 	for (uint32 Y = 0; Y < Height; Y++)
@@ -458,7 +459,7 @@ static inline void ConvertRawR16DataToFLinearColor(uint32 Width, uint32 Height,
 	}
 }
 
-static inline void ConvertRawR8G8B8A8DataToFLinearColor(uint32 Width, uint32 Height, uint8 *In, uint32 SrcPitch, FLinearColor* Out)
+inline void ConvertRawR8G8B8A8DataToFLinearColor(uint32 Width, uint32 Height, uint8 *In, uint32 SrcPitch, FLinearColor* Out)
 {
 	// Read the data out of the buffer, converting it from ABGR to ARGB.
 	for (uint32 Y = 0; Y < Height; Y++)
@@ -475,7 +476,7 @@ static inline void ConvertRawR8G8B8A8DataToFLinearColor(uint32 Width, uint32 Hei
 	}
 }
 
-static inline void ConvertRawB8G8R8A8DataToFLinearColor(uint32 Width, uint32 Height, uint8 *In, uint32 SrcPitch, FLinearColor* Out)
+inline void ConvertRawB8G8R8A8DataToFLinearColor(uint32 Width, uint32 Height, uint8 *In, uint32 SrcPitch, FLinearColor* Out)
 {
 	for (uint32 Y = 0; Y < Height; Y++)
 	{
@@ -491,7 +492,7 @@ static inline void ConvertRawB8G8R8A8DataToFLinearColor(uint32 Width, uint32 Hei
 	}
 }
 
-static inline void ConvertRawA2B10G10R10DataToFLinearColor(uint32 Width, uint32 Height, uint8 *In, uint32 SrcPitch, FLinearColor* Out)
+inline void ConvertRawA2B10G10R10DataToFLinearColor(uint32 Width, uint32 Height, uint8 *In, uint32 SrcPitch, FLinearColor* Out)
 {
 	// Read the data out of the buffer, converting it from R10G10B10A2 to FLinearColor.
 	for (uint32 Y = 0; Y < Height; Y++)
@@ -512,7 +513,7 @@ static inline void ConvertRawA2B10G10R10DataToFLinearColor(uint32 Width, uint32
 	}
 }
 
-static inline void ConvertRawR16G16B16A16FDataToFLinearColor(uint32 Width, uint32 Height, uint8 *In, uint32 SrcPitch, FLinearColor* Out, FReadSurfaceDataFlags InFlags)
+inline void ConvertRawR16G16B16A16FDataToFLinearColor(uint32 Width, uint32 Height, uint8 *In, uint32 SrcPitch, FLinearColor* Out, FReadSurfaceDataFlags InFlags)
 {
 	if (InFlags.GetCompressionMode() == RCM_MinMax)
 	{
@@ -574,7 +575,8 @@ static inline void ConvertRawR16G16B16A16FDataToFLinearColor(uint32 Width, uint3
 	}
 }
 
-static inline void ConvertRawRR11G11B10DataToFLinearColor(uint32 Width, uint32 Height, uint8 *In, uint32 SrcPitch, FLinearColor* Out)
+
+inline void ConvertRawRR11G11B10DataToFLinearColor(uint32 Width, uint32 Height, uint8 *In, uint32 SrcPitch, FLinearColor* Out)
 {
 	check(sizeof(FFloat3Packed) == sizeof(uint32));
 
@@ -592,7 +594,7 @@ static inline void ConvertRawRR11G11B10DataToFLinearColor(uint32 Width, uint32 H
 	}
 }
 
-static inline void ConvertRawR32G32B32A32DataToFLinearColor(uint32 Width, uint32 Height, uint8 *In, uint32 SrcPitch, FLinearColor* Out, FReadSurfaceDataFlags InFlags)
+inline void ConvertRawR32G32B32A32DataToFLinearColor(uint32 Width, uint32 Height, uint8 *In, uint32 SrcPitch, FLinearColor* Out, FReadSurfaceDataFlags InFlags)
 {
 	if (InFlags.GetCompressionMode() == RCM_MinMax)
 	{
@@ -648,7 +650,7 @@ static inline void ConvertRawR32G32B32A32DataToFLinearColor(uint32 Width, uint32
 	}
 }
 
-static inline void ConvertRawR24G8DataToFLinearColor(uint32 Width, uint32 Height, uint8 *In, uint32 SrcPitch, FLinearColor* Out, FReadSurfaceDataFlags InFlags)
+inline void ConvertRawR24G8DataToFLinearColor(uint32 Width, uint32 Height, uint8 *In, uint32 SrcPitch, FLinearColor* Out, FReadSurfaceDataFlags InFlags)
 {
 	// Depth stencil
 	for (uint32 Y = 0; Y < Height; Y++)
@@ -669,7 +671,7 @@ static inline void ConvertRawR24G8DataToFLinearColor(uint32 Width, uint32 Height
 	}
 }
 
-static inline void ConvertRawR32DataToFLinearColor(uint32 Width, uint32 Height, uint8 *In, uint32 SrcPitch, FLinearColor* Out, FReadSurfaceDataFlags InFlags)
+inline void ConvertRawR32DataToFLinearColor(uint32 Width, uint32 Height, uint8 *In, uint32 SrcPitch, FLinearColor* Out, FReadSurfaceDataFlags InFlags)
 {
 	// Depth stencil
 	for (uint32 Y = 0; Y < Height; Y++)
@@ -689,7 +691,7 @@ static inline void ConvertRawR32DataToFLinearColor(uint32 Width, uint32 Height,
 	}
 }
 
-static inline void ConvertRawR16G16B16A16DataToFLinearColor(uint32 Width, uint32 Height, uint8 *In, uint32 SrcPitch, FLinearColor* Out)
+inline void ConvertRawR16G16B16A16DataToFLinearColor(uint32 Width, uint32 Height, uint8 *In, uint32 SrcPitch, FLinearColor* Out)
 {
 	// Read the data out of the buffer, converting it to FLinearColor.
 	for (uint32 Y = 0; Y < Height; Y++)
@@ -710,7 +712,7 @@ static inline void ConvertRawR16G16B16A16DataToFLinearColor(uint32 Width, uint32
 	}
 }
 
-static inline void ConvertRawR16G16DataToFLinearColor(uint32 Width, uint32 Height, uint8 *In, uint32 SrcPitch, FLinearColor* Out)
+inline void ConvertRawR16G16DataToFLinearColor(uint32 Width, uint32 Height, uint8 *In, uint32 SrcPitch, FLinearColor* Out)
 {
 	// Read the data out of the buffer, converting it to FLinearColor.
 	for (uint32 Y = 0; Y < Height; Y++)
@@ -728,3 +730,4 @@ static inline void ConvertRawR16G16DataToFLinearColor(uint32 Width, uint32 Heigh
 		}
 	}
 }
+
diff --git a/Engine/Source/Runtime/Renderer/Public/RayTracingDynamicGeometryCollection.h b/Engine/Source/Runtime/Renderer/Public/RayTracingDynamicGeometryCollection.h
--- a/Engine/Source/Runtime/Renderer/Public/RayTracingDynamicGeometryCollection.h
+++ b/Engine/Source/Runtime/Renderer/Public/RayTracingDynamicGeometryCollection.h
@@ -10,6 +10,7 @@
 class FPrimitiveSceneProxy;
 class FScene;
 class FSceneView;
+struct FMeshComputeDispatchCommand;
 struct FRayTracingDynamicGeometryUpdateParams;
 
 class RENDERER_API FRayTracingDynamicGeometryCollection
@@ -39,7 +40,7 @@ public:
 
 private:
 
-	TArray<struct FMeshComputeDispatchCommand> DispatchCommands;
+	TArray<FMeshComputeDispatchCommand> DispatchCommands;
 	TArray<FRayTracingGeometryBuildParams> BuildParams;
 	TArray<FRayTracingGeometrySegment> Segments;
 
