Removing full-type RHI_API uses and adding them to members that need it.
This will reduce export counts and is part of upcoming coding standard changes.

diff --git a/Engine/Source/Runtime/RHI/Public/BoundShaderStateCache.h b/Engine/Source/Runtime/RHI/Public/BoundShaderStateCache.h
--- a/Engine/Source/Runtime/RHI/Public/BoundShaderStateCache.h
+++ b/Engine/Source/Runtime/RHI/Public/BoundShaderStateCache.h
@@ -164,7 +164,7 @@ private:
  * Handles removal from the bound shader state cache on destruction.
  * RHIs that use cached bound shader states should create one for each bound shader state.
  */
-class RHI_API FCachedBoundShaderStateLink
+class FCachedBoundShaderStateLink
 {
 public:
 
@@ -175,7 +175,7 @@ public:
 	FRHIBoundShaderState* BoundShaderState;
 
 	/** Adds the bound shader state to the cache. */
-	FCachedBoundShaderStateLink(
+	RHI_API FCachedBoundShaderStateLink(
 		FRHIVertexDeclaration* VertexDeclaration,
 		FRHIVertexShader* VertexShader,
 		FRHIPixelShader* PixelShader,
@@ -184,7 +184,7 @@ public:
 		);
 
 	/** Adds the bound shader state to the cache. */
-	FCachedBoundShaderStateLink(
+	RHI_API FCachedBoundShaderStateLink(
 		FRHIVertexDeclaration* VertexDeclaration,
 		FRHIVertexShader* VertexShader,
 		FRHIPixelShader* PixelShader,
@@ -194,7 +194,7 @@ public:
 	);
 
 	/** Adds the bound shader state to the cache. */
-	FCachedBoundShaderStateLink(
+	RHI_API FCachedBoundShaderStateLink(
 		FRHIMeshShader* MeshShader,
 		FRHIAmplificationShader* AmplificationShader,
 		FRHIPixelShader* PixelShader,
@@ -203,7 +203,7 @@ public:
 	);
 
 	/** Destructor.  Removes the bound shader state from the cache. */
-	~FCachedBoundShaderStateLink();
+	RHI_API ~FCachedBoundShaderStateLink();
 
 	/**
 	 * Get the RHI shader for the given frequency.
@@ -240,7 +240,7 @@ extern RHI_API FCachedBoundShaderStateLink* GetCachedBoundShaderState(
 
 extern RHI_API void EmptyCachedBoundShaderStates();
 
-class RHI_API FCachedBoundShaderStateLink_Threadsafe : public FCachedBoundShaderStateLink
+class FCachedBoundShaderStateLink_Threadsafe : public FCachedBoundShaderStateLink
 {
 public:
 	/** Adds the bound shader state to the cache. */
@@ -277,8 +277,8 @@ public:
 	{
 	}
 
-	void AddToCache();
-	void RemoveFromCache();
+	RHI_API void AddToCache();
+	RHI_API void RemoveFromCache();
 };
 
 /**
diff --git a/Engine/Source/Runtime/RHI/Public/DataDrivenShaderPlatformInfo.h b/Engine/Source/Runtime/RHI/Public/DataDrivenShaderPlatformInfo.h
--- a/Engine/Source/Runtime/RHI/Public/DataDrivenShaderPlatformInfo.h
+++ b/Engine/Source/Runtime/RHI/Public/DataDrivenShaderPlatformInfo.h
@@ -14,7 +14,7 @@ extern RHI_API const FName LANGUAGE_Vulkan;
 extern RHI_API const FName LANGUAGE_Sony;
 extern RHI_API const FName LANGUAGE_Nintendo;
 
-class RHI_API FGenericDataDrivenShaderPlatformInfo
+class FGenericDataDrivenShaderPlatformInfo
 {
 	FName Name;
 	FName Language;
@@ -127,13 +127,13 @@ class RHI_API FGenericDataDrivenShaderPlatformInfo
 		SetDefaultValues();
 	}
 
-	void SetDefaultValues();
+	RHI_API void SetDefaultValues();
 
 public:
-	static void Initialize();
-	static void UpdatePreviewPlatforms();
-	static void ParseDataDrivenShaderInfo(const FConfigSection& Section, FGenericDataDrivenShaderPlatformInfo& Info);
-	static const EShaderPlatform GetShaderPlatformFromName(const FName ShaderPlatformName);
+	RHI_API static void Initialize();
+	RHI_API static void UpdatePreviewPlatforms();
+	RHI_API static void ParseDataDrivenShaderInfo(const FConfigSection& Section, FGenericDataDrivenShaderPlatformInfo& Info);
+	RHI_API static const EShaderPlatform GetShaderPlatformFromName(const FName ShaderPlatformName);
 
 	static FORCEINLINE_DEBUGGABLE const FName GetName(const FStaticShaderPlatform Platform)
 	{
@@ -750,16 +750,16 @@ public:
 	}
 
 #if WITH_EDITOR
-	static FText GetFriendlyName(const FStaticShaderPlatform Platform);
+	RHI_API static FText GetFriendlyName(const FStaticShaderPlatform Platform);
 #endif
 
 private:
-	static FGenericDataDrivenShaderPlatformInfo Infos[SP_NumPlatforms];
+	RHI_API static FGenericDataDrivenShaderPlatformInfo Infos[SP_NumPlatforms];
 
 public:
 
 #if WITH_EDITOR
-	static TMap < FString, TFunction<bool(const FStaticShaderPlatform Platform)>> PropertyToShaderPlatformFunctionMap;
+	RHI_API static TMap < FString, TFunction<bool(const FStaticShaderPlatform Platform)>> PropertyToShaderPlatformFunctionMap;
 #endif
 
 	static bool IsValid(const FStaticShaderPlatform Platform)
diff --git a/Engine/Source/Runtime/RHI/Public/DynamicRHI.h b/Engine/Source/Runtime/RHI/Public/DynamicRHI.h
--- a/Engine/Source/Runtime/RHI/Public/DynamicRHI.h
+++ b/Engine/Source/Runtime/RHI/Public/DynamicRHI.h
@@ -141,15 +141,15 @@ struct FRawBufferShaderResourceViewInitializer : FShaderResourceViewInitializer
 
 class FDynamicRHI;
 
-class RHI_API FDefaultRHIRenderQueryPool final : public FRHIRenderQueryPool
+class FDefaultRHIRenderQueryPool final : public FRHIRenderQueryPool
 {
 public:
-	FDefaultRHIRenderQueryPool(ERenderQueryType InQueryType, FDynamicRHI* InDynamicRHI, uint32 InNumQueries);
-	~FDefaultRHIRenderQueryPool() override;
+	RHI_API FDefaultRHIRenderQueryPool(ERenderQueryType InQueryType, FDynamicRHI* InDynamicRHI, uint32 InNumQueries);
+	RHI_API ~FDefaultRHIRenderQueryPool() override;
 
 private:
-	virtual FRHIPooledRenderQuery AllocateQuery() override;
-	virtual void ReleaseQuery(TRefCountPtr<FRHIRenderQuery>&& Query) override;
+	RHI_API virtual FRHIPooledRenderQuery AllocateQuery() override;
+	RHI_API virtual void ReleaseQuery(TRefCountPtr<FRHIRenderQuery>&& Query) override;
 
 	FDynamicRHI* DynamicRHI = nullptr;
 	ERenderQueryType QueryType;
@@ -159,7 +159,7 @@ private:
 };
 
 /** The interface which is implemented by the dynamically bound RHI. */
-class RHI_API FDynamicRHI
+class FDynamicRHI
 {
 public:
 
@@ -275,7 +275,7 @@ public:
 	// Allows the platform RHI to perform operations on the GPU fence at the top-of-pipe.
 	// Default implementation just enqueues an RHI command to call IRHIComputeContext::WriteGPUFence().
 	//
-	virtual void RHIWriteGPUFence_TopOfPipe(FRHICommandListBase& RHICmdList, FRHIGPUFence* FenceRHI);
+	RHI_API virtual void RHIWriteGPUFence_TopOfPipe(FRHICommandListBase& RHICmdList, FRHIGPUFence* FenceRHI);
 
 	virtual void RHICreateTransition(FRHITransition* Transition, const FRHITransitionCreateInfo& CreateInfo)
 	{
@@ -309,13 +309,13 @@ public:
 	 * @param SizeRHI The length of the region in the buffer to lock.
 	 * @returns A pointer to the data starting at 'Offset' and of length 'SizeRHI' from 'StagingBuffer', or nullptr when there is an error.
 	 */
-	virtual void* RHILockStagingBuffer(FRHIStagingBuffer* StagingBuffer, FRHIGPUFence* Fence, uint32 Offset, uint32 SizeRHI);
+	RHI_API virtual void* RHILockStagingBuffer(FRHIStagingBuffer* StagingBuffer, FRHIGPUFence* Fence, uint32 Offset, uint32 SizeRHI);
 
 	/**
 	 * Unlock a staging buffer previously locked with RHILockStagingBuffer.
 	 * @param StagingBuffer The buffer that was previously locked.
 	 */
-	virtual void RHIUnlockStagingBuffer(FRHIStagingBuffer* StagingBuffer);
+	RHI_API virtual void RHIUnlockStagingBuffer(FRHIStagingBuffer* StagingBuffer);
 
 	/**
 	 * Lock a staging buffer to read contents on the CPU that were written by the GPU, without having to stall.
@@ -327,14 +327,14 @@ public:
 	 * @param SizeRHI The length of the region in the buffer to lock.
 	 * @returns A pointer to the data starting at 'Offset' and of length 'SizeRHI' from 'StagingBuffer', or nullptr when there is an error.
 	 */
-	virtual void* LockStagingBuffer_RenderThread(class FRHICommandListImmediate& RHICmdList, FRHIStagingBuffer* StagingBuffer, FRHIGPUFence* Fence, uint32 Offset, uint32 SizeRHI);
+	RHI_API virtual void* LockStagingBuffer_RenderThread(class FRHICommandListImmediate& RHICmdList, FRHIStagingBuffer* StagingBuffer, FRHIGPUFence* Fence, uint32 Offset, uint32 SizeRHI);
 
 	/**
 	 * Unlock a staging buffer previously locked with LockStagingBuffer_RenderThread.
 	 * @param RHICmdList The command-list to execute on or synchronize with.
 	 * @param StagingBuffer The buffer what was previously locked.
 	 */
-	virtual void UnlockStagingBuffer_RenderThread(class FRHICommandListImmediate& RHICmdList, FRHIStagingBuffer* StagingBuffer);
+	RHI_API virtual void UnlockStagingBuffer_RenderThread(class FRHICommandListImmediate& RHICmdList, FRHIStagingBuffer* StagingBuffer);
 
 	/**
 	* Creates a bound shader state instance which encapsulates a decl, vertex shader and pixel shader
@@ -394,7 +394,7 @@ public:
 	 * @param DestBuffer - the buffer to update
 	 * @param SrcBuffer - don't use after call. If null, will release any resource owned by DestBuffer
 	 */
-	virtual void RHITransferBufferUnderlyingResource(FRHIBuffer* DestBuffer, FRHIBuffer* SrcBuffer);
+	RHI_API virtual void RHITransferBufferUnderlyingResource(FRHIBuffer* DestBuffer, FRHIBuffer* SrcBuffer);
 
 	/**
 	* @param ResourceArray - An optional pointer to a resource array containing the resource's data.
@@ -405,12 +405,12 @@ public:
 
 	virtual FBufferRHIRef RHICreateBuffer(FRHICommandListBase& RHICmdList, uint32 Size, EBufferUsageFlags Usage, uint32 Stride, ERHIAccess InResourceState, FRHIResourceCreateInfo& CreateInfo) = 0;
 
-	virtual void* RHILockBuffer(FRHICommandListBase& RHICmdList, FRHIBuffer* Buffer, uint32 Offset, uint32 Size, EResourceLockMode LockMode);
-	virtual void* RHILockBufferMGPU(FRHICommandListBase& RHICmdList, FRHIBuffer* Buffer, uint32 GPUIndex, uint32 Offset, uint32 Size, EResourceLockMode LockMode);
+	RHI_API virtual void* RHILockBuffer(FRHICommandListBase& RHICmdList, FRHIBuffer* Buffer, uint32 Offset, uint32 Size, EResourceLockMode LockMode);
+	RHI_API virtual void* RHILockBufferMGPU(FRHICommandListBase& RHICmdList, FRHIBuffer* Buffer, uint32 GPUIndex, uint32 Offset, uint32 Size, EResourceLockMode LockMode);
 
 	// FlushType: Flush RHI Thread
-	virtual void RHIUnlockBuffer(FRHICommandListBase& RHICmdList, FRHIBuffer* Buffer);
-	virtual void RHIUnlockBufferMGPU(FRHICommandListBase& RHICmdList, FRHIBuffer* Buffer, uint32 GPUIndex);
+	RHI_API virtual void RHIUnlockBuffer(FRHICommandListBase& RHICmdList, FRHIBuffer* Buffer);
+	RHI_API virtual void RHIUnlockBufferMGPU(FRHICommandListBase& RHICmdList, FRHIBuffer* Buffer, uint32 GPUIndex);
 
 	/** Creates an unordered access view of the given buffer. */
 	// FlushType: Wait RHI Thread
@@ -422,7 +422,7 @@ public:
 
 	/** Creates an unordered access view of the given texture. */
 	// FlushType: Wait RHI Thread
-	virtual FUnorderedAccessViewRHIRef RHICreateUnorderedAccessView(FRHITexture* Texture, uint32 MipLevel, uint8 Format, uint16 FirstArraySlice, uint16 NumArraySlices);
+	RHI_API virtual FUnorderedAccessViewRHIRef RHICreateUnorderedAccessView(FRHITexture* Texture, uint32 MipLevel, uint8 Format, uint16 FirstArraySlice, uint16 NumArraySlices);
 
 	/** Creates an unordered access view of the given vertex buffer. */
 	// FlushType: Wait RHI Thread
@@ -441,11 +441,11 @@ public:
 
 	// Must be called on RHI thread timeline
 	// Make sure to call RHIThreadFence(true) afterwards so that parallel translation doesn't refer old resources
-	virtual void RHIUpdateShaderResourceView(FRHIShaderResourceView* SRV, FRHIBuffer* Buffer, uint32 Stride, uint8 Format);
+	RHI_API virtual void RHIUpdateShaderResourceView(FRHIShaderResourceView* SRV, FRHIBuffer* Buffer, uint32 Stride, uint8 Format);
 
-	virtual void RHIUpdateShaderResourceView(FRHIShaderResourceView* SRV, FRHIBuffer* Buffer);
+	RHI_API virtual void RHIUpdateShaderResourceView(FRHIShaderResourceView* SRV, FRHIBuffer* Buffer);
 
-	virtual void RHIUpdateTextureReference(FRHITextureReference* TextureRef, FRHITexture* NewTexture);
+	RHI_API virtual void RHIUpdateTextureReference(FRHITextureReference* TextureRef, FRHITexture* NewTexture);
 
 	struct FRHICalcTextureSizeResult
 	{
@@ -472,7 +472,7 @@ public:
 	* @param Format - EPixelFormat texture format of the SRV.
 	*/
 	// FlushType: Thread safe
-	virtual uint64 RHIGetMinimumAlignmentForBufferBackedSRV(EPixelFormat Format);
+	RHI_API virtual uint64 RHIGetMinimumAlignmentForBufferBackedSRV(EPixelFormat Format);
 
 	/**
 	* Retrieves texture memory stats.
@@ -711,7 +711,7 @@ public:
 	virtual void RHIReadSurfaceData(FRHITexture* Texture, FIntRect Rect, TArray<FColor>& OutData, FReadSurfaceDataFlags InFlags) = 0;
 
 	// Default fallback; will not work for non-8-bit surfaces and it's extremely slow.
-	virtual void RHIReadSurfaceData(FRHITexture* Texture, FIntRect Rect, TArray<FLinearColor>& OutData, FReadSurfaceDataFlags InFlags);
+	RHI_API virtual void RHIReadSurfaceData(FRHITexture* Texture, FIntRect Rect, TArray<FLinearColor>& OutData, FReadSurfaceDataFlags InFlags);
 
 	/** Watch out for OutData to be 0 (can happen on DXGI_ERROR_DEVICE_REMOVED), don't call RHIUnmapStagingSurface in that case. */
 	// FlushType: Flush Immediate (seems wrong)
@@ -725,13 +725,13 @@ public:
 	virtual void RHIReadSurfaceFloatData(FRHITexture* Texture, FIntRect Rect, TArray<FFloat16Color>& OutData, ECubeFace CubeFace, int32 ArrayIndex, int32 MipIndex) = 0;
 
 	// FlushType: Flush Immediate (seems wrong)
-	virtual void RHIReadSurfaceFloatData(FRHITexture* Texture, FIntRect Rect, TArray<FFloat16Color>& OutData, FReadSurfaceDataFlags InFlags);
+	RHI_API virtual void RHIReadSurfaceFloatData(FRHITexture* Texture, FIntRect Rect, TArray<FFloat16Color>& OutData, FReadSurfaceDataFlags InFlags);
 
 	// FlushType: Flush Immediate (seems wrong)
 	virtual void RHIRead3DSurfaceFloatData(FRHITexture* Texture, FIntRect Rect, FIntPoint ZMinMax, TArray<FFloat16Color>& OutData) = 0;
 
 	// FlushType: Flush Immediate (seems wrong)
-	virtual void RHIRead3DSurfaceFloatData(FRHITexture* Texture, FIntRect Rect, FIntPoint ZMinMax, TArray<FFloat16Color>& OutData, FReadSurfaceDataFlags InFlags);
+	RHI_API virtual void RHIRead3DSurfaceFloatData(FRHITexture* Texture, FIntRect Rect, FIntPoint ZMinMax, TArray<FFloat16Color>& OutData, FReadSurfaceDataFlags InFlags);
 
 	// FlushType: Wait RHI Thread
 	virtual FRenderQueryRHIRef RHICreateRenderQuery(ERenderQueryType QueryType) = 0;
@@ -739,8 +739,8 @@ public:
 	virtual void RHIBeginOcclusionQueryBatch_TopOfPipe(FRHICommandListBase& RHICmdList, uint32 NumQueriesInBatch) {}
 	virtual void RHIEndOcclusionQueryBatch_TopOfPipe  (FRHICommandListBase& RHICmdList) {}
 
-	virtual void RHIBeginRenderQuery_TopOfPipe(FRHICommandListBase& RHICmdList, FRHIRenderQuery* RenderQuery);
-	virtual void RHIEndRenderQuery_TopOfPipe  (FRHICommandListBase& RHICmdList, FRHIRenderQuery* RenderQuery);
+	RHI_API virtual void RHIBeginRenderQuery_TopOfPipe(FRHICommandListBase& RHICmdList, FRHIRenderQuery* RenderQuery);
+	RHI_API virtual void RHIEndRenderQuery_TopOfPipe  (FRHICommandListBase& RHICmdList, FRHIRenderQuery* RenderQuery);
 
 	// CAUTION: Even though this is marked as threadsafe, it is only valid to call from the render thread. It is need not be threadsafe on platforms that do not support or aren't using an RHIThread
 	// FlushType: Thread safe, but varies by RHI
@@ -797,10 +797,10 @@ public:
 	}
 	
 	// Compute the hash of the PSO initializer for PSO Precaching (only hash data relevant for the RHI specific PSO)
-	virtual uint64 RHIComputePrecachePSOHash(const FGraphicsPipelineStateInitializer& Initializer);
+	RHI_API virtual uint64 RHIComputePrecachePSOHash(const FGraphicsPipelineStateInitializer& Initializer);
 
 	// Check if PSO Initializers are the same used during PSO Precaching (only compare data relevant for the RHI specific PSO)
-	virtual bool RHIMatchPrecachePSOInitializers(const FGraphicsPipelineStateInitializer& LHS, const FGraphicsPipelineStateInitializer& RHS);
+	RHI_API virtual bool RHIMatchPrecachePSOInitializers(const FGraphicsPipelineStateInitializer& LHS, const FGraphicsPipelineStateInitializer& RHS);
 
 	virtual void RHIAdvanceFrameFence() {};
 
@@ -842,7 +842,7 @@ public:
 	}
 
 	// Return what colour space the viewport is in. Used for HDR displays
-	virtual EColorSpaceAndEOTF RHIGetColorSpace(FRHIViewport* Viewport);
+	RHI_API virtual EColorSpaceAndEOTF RHIGetColorSpace(FRHIViewport* Viewport);
 
 	// Return preferred pixel format if given format is unsupported.
 	virtual EPixelFormat RHIPreferredPixelFormatHint(EPixelFormat PreferredPixelFormat)
@@ -851,7 +851,8 @@ public:
 	}
 
 	// Tests the viewport to see if its HDR status has changed. This is usually tested after a window has been moved
-	virtual void RHICheckViewportHDRStatus(FRHIViewport* Viewport);
+	RHI_API virtual void RHICheckViewportHDRStatus(FRHIViewport* Viewport);
+
 	virtual void RHIHandleDisplayChange() {}
 
 	//  must be called from the main thread.
@@ -1017,47 +1018,47 @@ public:
 	virtual void RHISubmitCommandLists(TArrayView<IRHIPlatformCommandList*> CommandLists) = 0;
 
 	UE_DEPRECATED(5.1, "CreateBuffer_RenderThread is deprecated. Use RHICreateBuffer instead.")
-	virtual FBufferRHIRef CreateBuffer_RenderThread(class FRHICommandListBase& RHICmdList, uint32 Size, EBufferUsageFlags Usage, uint32 Stride, ERHIAccess ResourceState, FRHIResourceCreateInfo& CreateInfo);
-	virtual FShaderResourceViewRHIRef CreateShaderResourceView_RenderThread(class FRHICommandListImmediate& RHICmdList, FRHIBuffer* Buffer, uint32 Stride, uint8 Format);
-	virtual FShaderResourceViewRHIRef CreateShaderResourceView_RenderThread(class FRHICommandListImmediate& RHICmdList, const FShaderResourceViewInitializer& Initializer);
-	virtual FShaderResourceViewRHIRef CreateShaderResourceView_RenderThread(class FRHICommandListImmediate& RHICmdList, FRHIBuffer* Buffer);
-	virtual FTexture2DRHIRef AsyncReallocateTexture2D_RenderThread(class FRHICommandListImmediate& RHICmdList, FRHITexture2D* Texture2D, int32 NewMipCount, int32 NewSizeX, int32 NewSizeY, FThreadSafeCounter* RequestStatus);
-	virtual ETextureReallocationStatus FinalizeAsyncReallocateTexture2D_RenderThread(class FRHICommandListImmediate& RHICmdList, FRHITexture2D* Texture2D, bool bBlockUntilCompleted);
-	virtual ETextureReallocationStatus CancelAsyncReallocateTexture2D_RenderThread(class FRHICommandListImmediate& RHICmdList, FRHITexture2D* Texture2D, bool bBlockUntilCompleted);
-
-	virtual void* LockTexture2D_RenderThread(class FRHICommandListImmediate& RHICmdList, FRHITexture2D* Texture, uint32 MipIndex, EResourceLockMode LockMode, uint32& DestStride, bool bLockWithinMiptail, bool bNeedsDefaultRHIFlush = true);
-	virtual void UnlockTexture2D_RenderThread(class FRHICommandListImmediate& RHICmdList, FRHITexture2D* Texture, uint32 MipIndex, bool bLockWithinMiptail, bool bNeedsDefaultRHIFlush = true);
-
-	virtual void* LockTexture2DArray_RenderThread(class FRHICommandListImmediate& RHICmdList, FRHITexture2DArray* Texture, uint32 ArrayIndex, uint32 MipIndex, EResourceLockMode LockMode, uint32& DestStride, bool bLockWithinMiptail);
-	virtual void UnlockTexture2DArray_RenderThread(class FRHICommandListImmediate& RHICmdList, FRHITexture2DArray* Texture, uint32 ArrayIndex, uint32 MipIndex, bool bLockWithinMiptail);
-
-	virtual FUpdateTexture3DData RHIBeginUpdateTexture3D(FRHICommandListBase& RHICmdList, FRHITexture3D* Texture, uint32 MipIndex, const struct FUpdateTextureRegion3D& UpdateRegion);
-	virtual void RHIEndUpdateTexture3D(FRHICommandListBase& RHICmdList, FUpdateTexture3DData& UpdateData);
-
-	virtual void RHIEndMultiUpdateTexture3D(FRHICommandListBase& RHICmdList, TArray<FUpdateTexture3DData>& UpdateDataArray);
-
-	virtual FRHIShaderLibraryRef RHICreateShaderLibrary_RenderThread(class FRHICommandListImmediate& RHICmdList, EShaderPlatform Platform, FString FilePath, FString Name);
-
-	virtual FTextureRHIRef RHICreateTexture_RenderThread(class FRHICommandListImmediate& RHICmdList, const FRHITextureCreateDesc& CreateDesc);
-
-	virtual FUnorderedAccessViewRHIRef RHICreateUnorderedAccessView_RenderThread(class FRHICommandListImmediate& RHICmdList, FRHIBuffer* Buffer, bool bUseUAVCounter, bool bAppendBuffer);
-	virtual FUnorderedAccessViewRHIRef RHICreateUnorderedAccessView_RenderThread(class FRHICommandListImmediate& RHICmdList, FRHITexture* Texture, uint32 MipLevel, uint16 FirstArraySlice, uint16 NumArraySlices);
-	virtual FUnorderedAccessViewRHIRef RHICreateUnorderedAccessView_RenderThread(class FRHICommandListImmediate& RHICmdList, FRHITexture* Texture, uint32 MipLevel, uint8 Format, uint16 FirstArraySlice, uint16 NumArraySlices);
-	virtual FUnorderedAccessViewRHIRef RHICreateUnorderedAccessView_RenderThread(class FRHICommandListImmediate& RHICmdList, FRHIBuffer* Buffer, uint8 Format);
-	virtual FShaderResourceViewRHIRef RHICreateShaderResourceView_RenderThread(class FRHICommandListImmediate& RHICmdList, FRHITexture* Texture, const FRHITextureSRVCreateInfo& CreateInfo);
-	virtual FShaderResourceViewRHIRef RHICreateShaderResourceView_RenderThread(class FRHICommandListImmediate& RHICmdList, FRHIBuffer* Buffer, uint32 Stride, uint8 Format);
-	virtual FShaderResourceViewRHIRef RHICreateShaderResourceView_RenderThread(class FRHICommandListImmediate& RHICmdList, const FShaderResourceViewInitializer& Initializer);
-	virtual FShaderResourceViewRHIRef RHICreateShaderResourceView_RenderThread(class FRHICommandListImmediate& RHICmdList, FRHIBuffer* Buffer);
-	virtual FShaderResourceViewRHIRef RHICreateShaderResourceViewWriteMask_RenderThread(class FRHICommandListImmediate& RHICmdList, FRHITexture2D* Texture2D);
-	virtual FShaderResourceViewRHIRef RHICreateShaderResourceViewFMask_RenderThread(class FRHICommandListImmediate& RHICmdList, FRHITexture2D* Texture2D);
+	RHI_API virtual FBufferRHIRef CreateBuffer_RenderThread(class FRHICommandListBase& RHICmdList, uint32 Size, EBufferUsageFlags Usage, uint32 Stride, ERHIAccess ResourceState, FRHIResourceCreateInfo& CreateInfo);
+	RHI_API virtual FShaderResourceViewRHIRef CreateShaderResourceView_RenderThread(class FRHICommandListImmediate& RHICmdList, FRHIBuffer* Buffer, uint32 Stride, uint8 Format);
+	RHI_API virtual FShaderResourceViewRHIRef CreateShaderResourceView_RenderThread(class FRHICommandListImmediate& RHICmdList, const FShaderResourceViewInitializer& Initializer);
+	RHI_API virtual FShaderResourceViewRHIRef CreateShaderResourceView_RenderThread(class FRHICommandListImmediate& RHICmdList, FRHIBuffer* Buffer);
+	RHI_API virtual FTexture2DRHIRef AsyncReallocateTexture2D_RenderThread(class FRHICommandListImmediate& RHICmdList, FRHITexture2D* Texture2D, int32 NewMipCount, int32 NewSizeX, int32 NewSizeY, FThreadSafeCounter* RequestStatus);
+	RHI_API virtual ETextureReallocationStatus FinalizeAsyncReallocateTexture2D_RenderThread(class FRHICommandListImmediate& RHICmdList, FRHITexture2D* Texture2D, bool bBlockUntilCompleted);
+	RHI_API virtual ETextureReallocationStatus CancelAsyncReallocateTexture2D_RenderThread(class FRHICommandListImmediate& RHICmdList, FRHITexture2D* Texture2D, bool bBlockUntilCompleted);
+
+	RHI_API virtual void* LockTexture2D_RenderThread(class FRHICommandListImmediate& RHICmdList, FRHITexture2D* Texture, uint32 MipIndex, EResourceLockMode LockMode, uint32& DestStride, bool bLockWithinMiptail, bool bNeedsDefaultRHIFlush = true);
+	RHI_API virtual void UnlockTexture2D_RenderThread(class FRHICommandListImmediate& RHICmdList, FRHITexture2D* Texture, uint32 MipIndex, bool bLockWithinMiptail, bool bNeedsDefaultRHIFlush = true);
+
+	RHI_API virtual void* LockTexture2DArray_RenderThread(class FRHICommandListImmediate& RHICmdList, FRHITexture2DArray* Texture, uint32 ArrayIndex, uint32 MipIndex, EResourceLockMode LockMode, uint32& DestStride, bool bLockWithinMiptail);
+	RHI_API virtual void UnlockTexture2DArray_RenderThread(class FRHICommandListImmediate& RHICmdList, FRHITexture2DArray* Texture, uint32 ArrayIndex, uint32 MipIndex, bool bLockWithinMiptail);
+
+	RHI_API virtual FUpdateTexture3DData RHIBeginUpdateTexture3D(FRHICommandListBase& RHICmdList, FRHITexture3D* Texture, uint32 MipIndex, const struct FUpdateTextureRegion3D& UpdateRegion);
+	RHI_API virtual void RHIEndUpdateTexture3D(FRHICommandListBase& RHICmdList, FUpdateTexture3DData& UpdateData);
+
+	RHI_API virtual void RHIEndMultiUpdateTexture3D(FRHICommandListBase& RHICmdList, TArray<FUpdateTexture3DData>& UpdateDataArray);
+
+	RHI_API virtual FRHIShaderLibraryRef RHICreateShaderLibrary_RenderThread(class FRHICommandListImmediate& RHICmdList, EShaderPlatform Platform, FString FilePath, FString Name);
+
+	RHI_API virtual FTextureRHIRef RHICreateTexture_RenderThread(class FRHICommandListImmediate& RHICmdList, const FRHITextureCreateDesc& CreateDesc);
+
+	RHI_API virtual FUnorderedAccessViewRHIRef RHICreateUnorderedAccessView_RenderThread(class FRHICommandListImmediate& RHICmdList, FRHIBuffer* Buffer, bool bUseUAVCounter, bool bAppendBuffer);
+	RHI_API virtual FUnorderedAccessViewRHIRef RHICreateUnorderedAccessView_RenderThread(class FRHICommandListImmediate& RHICmdList, FRHITexture* Texture, uint32 MipLevel, uint16 FirstArraySlice, uint16 NumArraySlices);
+	RHI_API virtual FUnorderedAccessViewRHIRef RHICreateUnorderedAccessView_RenderThread(class FRHICommandListImmediate& RHICmdList, FRHITexture* Texture, uint32 MipLevel, uint8 Format, uint16 FirstArraySlice, uint16 NumArraySlices);
+	RHI_API virtual FUnorderedAccessViewRHIRef RHICreateUnorderedAccessView_RenderThread(class FRHICommandListImmediate& RHICmdList, FRHIBuffer* Buffer, uint8 Format);
+	RHI_API virtual FShaderResourceViewRHIRef RHICreateShaderResourceView_RenderThread(class FRHICommandListImmediate& RHICmdList, FRHITexture* Texture, const FRHITextureSRVCreateInfo& CreateInfo);
+	RHI_API virtual FShaderResourceViewRHIRef RHICreateShaderResourceView_RenderThread(class FRHICommandListImmediate& RHICmdList, FRHIBuffer* Buffer, uint32 Stride, uint8 Format);
+	RHI_API virtual FShaderResourceViewRHIRef RHICreateShaderResourceView_RenderThread(class FRHICommandListImmediate& RHICmdList, const FShaderResourceViewInitializer& Initializer);
+	RHI_API virtual FShaderResourceViewRHIRef RHICreateShaderResourceView_RenderThread(class FRHICommandListImmediate& RHICmdList, FRHIBuffer* Buffer);
+	RHI_API virtual FShaderResourceViewRHIRef RHICreateShaderResourceViewWriteMask_RenderThread(class FRHICommandListImmediate& RHICmdList, FRHITexture2D* Texture2D);
+	RHI_API virtual FShaderResourceViewRHIRef RHICreateShaderResourceViewFMask_RenderThread(class FRHICommandListImmediate& RHICmdList, FRHITexture2D* Texture2D);
 	
-	virtual void* RHILockTextureCubeFace_RenderThread(class FRHICommandListImmediate& RHICmdList, FRHITextureCube* Texture, uint32 FaceIndex, uint32 ArrayIndex, uint32 MipIndex, EResourceLockMode LockMode, uint32& DestStride, bool bLockWithinMiptail);
-	virtual void RHIUnlockTextureCubeFace_RenderThread(class FRHICommandListImmediate& RHICmdList, FRHITextureCube* Texture, uint32 FaceIndex, uint32 ArrayIndex, uint32 MipIndex, bool bLockWithinMiptail);
+	RHI_API virtual void* RHILockTextureCubeFace_RenderThread(class FRHICommandListImmediate& RHICmdList, FRHITextureCube* Texture, uint32 FaceIndex, uint32 ArrayIndex, uint32 MipIndex, EResourceLockMode LockMode, uint32& DestStride, bool bLockWithinMiptail);
+	RHI_API virtual void RHIUnlockTextureCubeFace_RenderThread(class FRHICommandListImmediate& RHICmdList, FRHITextureCube* Texture, uint32 FaceIndex, uint32 ArrayIndex, uint32 MipIndex, bool bLockWithinMiptail);
 
-	virtual void RHIMapStagingSurface_RenderThread(class FRHICommandListImmediate& RHICmdList, FRHITexture* Texture, uint32 GPUIndex, FRHIGPUFence* Fence, void*& OutData, int32& OutWidth, int32& OutHeight);
-	virtual void RHIUnmapStagingSurface_RenderThread(class FRHICommandListImmediate& RHICmdList, FRHITexture* Texture, uint32 GPUIndex);
-	virtual void RHIReadSurfaceFloatData_RenderThread(class FRHICommandListImmediate& RHICmdList, FRHITexture* Texture, FIntRect Rect, TArray<FFloat16Color>& OutData, ECubeFace CubeFace, int32 ArrayIndex, int32 MipIndex);
-	virtual void RHIReadSurfaceFloatData_RenderThread(class FRHICommandListImmediate& RHICmdList, FRHITexture* Texture, FIntRect Rect, TArray<FFloat16Color>& OutData, FReadSurfaceDataFlags Flags);
+	RHI_API virtual void RHIMapStagingSurface_RenderThread(class FRHICommandListImmediate& RHICmdList, FRHITexture* Texture, uint32 GPUIndex, FRHIGPUFence* Fence, void*& OutData, int32& OutWidth, int32& OutHeight);
+	RHI_API virtual void RHIUnmapStagingSurface_RenderThread(class FRHICommandListImmediate& RHICmdList, FRHITexture* Texture, uint32 GPUIndex);
+	RHI_API virtual void RHIReadSurfaceFloatData_RenderThread(class FRHICommandListImmediate& RHICmdList, FRHITexture* Texture, FIntRect Rect, TArray<FFloat16Color>& OutData, ECubeFace CubeFace, int32 ArrayIndex, int32 MipIndex);
+	RHI_API virtual void RHIReadSurfaceFloatData_RenderThread(class FRHICommandListImmediate& RHICmdList, FRHITexture* Texture, FIntRect Rect, TArray<FFloat16Color>& OutData, FReadSurfaceDataFlags Flags);
 
 	// Buffer Lock/Unlock
 	virtual void* LockBuffer_BottomOfPipe(class FRHICommandListBase& RHICmdList, FRHIBuffer* Buffer, uint32 Offset, uint32 SizeRHI, EResourceLockMode LockMode)
@@ -1074,13 +1075,13 @@ public:
 	}
 
 	//Utilities
-	virtual void EnableIdealGPUCaptureOptions(bool bEnable);
+	RHI_API virtual void EnableIdealGPUCaptureOptions(bool bEnable);
 	
 	//checks if the GPU is still alive.
 	virtual bool CheckGpuHeartbeat() const { return true; }
 
-	virtual void VirtualTextureSetFirstMipInMemory_RenderThread(class FRHICommandListImmediate& RHICmdList, FRHITexture2D* Texture, uint32 FirstMip);
-	virtual void VirtualTextureSetFirstMipVisible_RenderThread(class FRHICommandListImmediate& RHICmdList, FRHITexture2D* Texture, uint32 FirstMip);
+	RHI_API virtual void VirtualTextureSetFirstMipInMemory_RenderThread(class FRHICommandListImmediate& RHICmdList, FRHITexture2D* Texture, uint32 FirstMip);
+	RHI_API virtual void VirtualTextureSetFirstMipVisible_RenderThread(class FRHICommandListImmediate& RHICmdList, FRHITexture2D* Texture, uint32 FirstMip);
 
 	virtual FRHIFlipDetails RHIWaitForFlip(double TimeoutInSeconds) { return FRHIFlipDetails(); }
 	virtual void RHISignalFlipEvent() { }
diff --git a/Engine/Source/Runtime/RHI/Public/PipelineFileCache.h b/Engine/Source/Runtime/RHI/Public/PipelineFileCache.h
--- a/Engine/Source/Runtime/RHI/Public/PipelineFileCache.h
+++ b/Engine/Source/Runtime/RHI/Public/PipelineFileCache.h
@@ -317,7 +317,7 @@ struct FPSOUsageData
 	uint16 EngineFlags;
 };
 
-class RHI_API FPipelineFileCacheManager
+class FPipelineFileCacheManager
 {
     friend class FPipelineCacheFile;
 public:
@@ -336,40 +336,40 @@ public:
 
 public:
 	
-	static void Initialize(uint32 GameVersion);
-	static void Shutdown();
+	RHI_API static void Initialize(uint32 GameVersion);
+	RHI_API static void Shutdown();
 	
-	static bool LoadPipelineFileCacheInto(FString const& Path, TSet<FPipelineCacheFileFormatPSO>& PSOs);
-	static bool SavePipelineFileCacheFrom(uint32 GameVersion, EShaderPlatform Platform, FString const& Path, const TSet<FPipelineCacheFileFormatPSO>& PSOs);
-	static bool MergePipelineFileCaches(FString const& PathA, FString const& PathB, FPipelineFileCacheManager::PSOOrder Order, FString const& OutputPath);
+	RHI_API static bool LoadPipelineFileCacheInto(FString const& Path, TSet<FPipelineCacheFileFormatPSO>& PSOs);
+	RHI_API static bool SavePipelineFileCacheFrom(uint32 GameVersion, EShaderPlatform Platform, FString const& Path, const TSet<FPipelineCacheFileFormatPSO>& PSOs);
+	RHI_API static bool MergePipelineFileCaches(FString const& PathA, FString const& PathB, FPipelineFileCacheManager::PSOOrder Order, FString const& OutputPath);
 																				
 	/* Open the pipeline file cache for the specfied name and platform. If successful, the GUID of the game file will be returned in OutGameFileGuid */
-	static bool OpenPipelineFileCache(const FString& Key, const FString& CacheName, EShaderPlatform Platform, FGuid& OutGameFileGuid);
+	RHI_API static bool OpenPipelineFileCache(const FString& Key, const FString& CacheName, EShaderPlatform Platform, FGuid& OutGameFileGuid);
 
 	/* Open the user pipeline file cache for the specified name and platform. The user cache is always created even if the file was not present when opened.
 	* Name is the name used when opening the file, the key value for the user cache is held within UserCacheName.
 	* returns true if the file was opened.
 	*/
-	static bool OpenUserPipelineFileCache(const FString& Key, const FString& CacheName, EShaderPlatform Platform, FGuid& OutGameFileGuid);
+	RHI_API static bool OpenUserPipelineFileCache(const FString& Key, const FString& CacheName, EShaderPlatform Platform, FGuid& OutGameFileGuid);
 
-	static bool SavePipelineFileCache(SaveMode Mode);
+	RHI_API static bool SavePipelineFileCache(SaveMode Mode);
 
-	static void CloseUserPipelineFileCache();
+	RHI_API static void CloseUserPipelineFileCache();
 
-	static void CacheGraphicsPSO(uint32 RunTimeHash, FGraphicsPipelineStateInitializer const& Initializer);
-	static void CacheComputePSO(uint32 RunTimeHash, FRHIComputeShader const* Initializer);
-	static void CacheRayTracingPSO(const FRayTracingPipelineStateInitializer& Initializer, ERayTracingPipelineCacheFlags Flags);
+	RHI_API static void CacheGraphicsPSO(uint32 RunTimeHash, FGraphicsPipelineStateInitializer const& Initializer);
+	RHI_API static void CacheComputePSO(uint32 RunTimeHash, FRHIComputeShader const* Initializer);
+	RHI_API static void CacheRayTracingPSO(const FRayTracingPipelineStateInitializer& Initializer, ERayTracingPipelineCacheFlags Flags);
 
 	// true if the named PSOFC is currently open.
-	static bool HasPipelineFileCache(const FString& PSOCacheKey);
+	RHI_API static bool HasPipelineFileCache(const FString& PSOCacheKey);
 
-	static FPipelineStateStats* RegisterPSOStats(uint32 RunTimeHash);
+	RHI_API static FPipelineStateStats* RegisterPSOStats(uint32 RunTimeHash);
 	
 	/*
 	 * This PSO has failed compile and is invalid - this cache should not return this invalid PSO from subsequent calls for PreCompile requests.
 	 * Note: Not implementated for Compute that has no flag to say it came from this cache - don't want to consume failures that didn't propagate from this cache.
 	 */
-	static void RegisterPSOCompileFailure(uint32 RunTimeHash, FGraphicsPipelineStateInitializer const& Initializer);
+	RHI_API static void RegisterPSOCompileFailure(uint32 RunTimeHash, FGraphicsPipelineStateInitializer const& Initializer);
 	
 	/**
 	 * Event signature for being notified that a new PSO has been logged
@@ -379,20 +379,20 @@ public:
 	/**
 	 * Gets the event delegate to register for pipeline state logging events.
 	 */
-	static FPipelineStateLoggedEvent& OnPipelineStateLogged();
+	RHI_API static FPipelineStateLoggedEvent& OnPipelineStateLogged();
 	
-	static void GetOrderedPSOHashes(const FString& PSOCacheKey, TArray<FPipelineCachePSOHeader>& PSOHashes, PSOOrder Order, int64 MinBindCount, TSet<uint32> const& AlreadyCompiledHashes);
-	static void FetchPSODescriptors(const FString& PSOCacheKey, TDoubleLinkedList<FPipelineCacheFileFormatPSORead*>& LoadedBatch);
+	RHI_API static void GetOrderedPSOHashes(const FString& PSOCacheKey, TArray<FPipelineCachePSOHeader>& PSOHashes, PSOOrder Order, int64 MinBindCount, TSet<uint32> const& AlreadyCompiledHashes);
+	RHI_API static void FetchPSODescriptors(const FString& PSOCacheKey, TDoubleLinkedList<FPipelineCacheFileFormatPSORead*>& LoadedBatch);
 
-	static int32 GetTotalPSOCount(const FString& PSOCacheKey);
+	RHI_API static int32 GetTotalPSOCount(const FString& PSOCacheKey);
 
-	static uint32 NumPSOsLogged();
+	RHI_API static uint32 NumPSOsLogged();
 	
-	static bool IsPipelineFileCacheEnabled();
-	static bool LogPSOtoFileCache();
-    static bool ReportNewPSOs();
+	RHI_API static bool IsPipelineFileCacheEnabled();
+	RHI_API static bool LogPSOtoFileCache();
+	RHI_API static bool ReportNewPSOs();
 	/* Report additional data about new PSOs to the log. */
-	static bool LogPSODetails();
+	RHI_API static bool LogPSODetails();
 	
 	/**
 	 * Define the Current Game Usage Mask and a comparison function to compare this mask against the recorded mask in each PSO
@@ -400,10 +400,10 @@ public:
 	 * @param InComparisonFnPtr Pointer to the comparsion function - see above FPSOMaskComparisonFn definition for details
 	 * @returns the old mask
 	 */
-	static uint64 SetGameUsageMaskWithComparison(uint64 GameUsageMask, FPSOMaskComparisonFn InComparisonFnPtr);
+	RHI_API static uint64 SetGameUsageMaskWithComparison(uint64 GameUsageMask, FPSOMaskComparisonFn InComparisonFnPtr);
 	static uint64 GetGameUsageMask()	{ return GameUsageMask;}
 	
-	static void PreCompileComplete();
+	RHI_API static void PreCompileComplete();
 private:
 	
 	static void RegisterPSOUsageDataUpdateForNextSave(FPSOUsageData& UsageData);
@@ -439,6 +439,6 @@ private:
 	static PSOOrder RequestedOrder;
 	static bool FileCacheEnabled;
 	static FPipelineStateLoggedEvent PSOLoggedEvent;
-	static uint64 GameUsageMask;
+	RHI_API static uint64 GameUsageMask;
 	static FPSOMaskComparisonFn MaskComparisonFn;
 };
diff --git a/Engine/Source/Runtime/RHI/Public/RHIBreadcrumbs.h b/Engine/Source/Runtime/RHI/Public/RHIBreadcrumbs.h
--- a/Engine/Source/Runtime/RHI/Public/RHIBreadcrumbs.h
+++ b/Engine/Source/Runtime/RHI/Public/RHIBreadcrumbs.h
@@ -47,18 +47,18 @@ struct TRHIBreadcrumbState
 
 using FRHIBreadcrumbState = TRHIBreadcrumbState<>;
 
-class RHI_API FRHIBreadcrumbStack
+class FRHIBreadcrumbStack
 {
 public:
-	void Reset();
+	RHI_API void Reset();
 
-	FRHIBreadcrumb* PushBreadcrumb(FMemStackBase& Allocator, const TCHAR* InText, int32 InLen = 0);
+	RHI_API FRHIBreadcrumb* PushBreadcrumb(FMemStackBase& Allocator, const TCHAR* InText, int32 InLen = 0);
 
-	FRHIBreadcrumb* PushBreadcrumbPrintf(FMemStackBase& Allocator, const TCHAR* InFormat, ...);
+	RHI_API FRHIBreadcrumb* PushBreadcrumbPrintf(FMemStackBase& Allocator, const TCHAR* InFormat, ...);
 
-	FRHIBreadcrumb* PopBreadcrumb();
+	RHI_API FRHIBreadcrumb* PopBreadcrumb();
 
-	FRHIBreadcrumb* PopFirstUnsubmittedBreadcrumb();
+	RHI_API FRHIBreadcrumb* PopFirstUnsubmittedBreadcrumb();
 
 	template <typename AllocatorType>
 	void ExportBreadcrumbState(TRHIBreadcrumbState<AllocatorType>& State) const
@@ -93,14 +93,14 @@ public:
 		}
 	}
 
-	void DeepCopy(FMemStackBase& Allocator, const FRHIBreadcrumbStack& Parent);
+	RHI_API void DeepCopy(FMemStackBase& Allocator, const FRHIBreadcrumbStack& Parent);
 
-	void ValidateEmpty() const;
+	RHI_API void ValidateEmpty() const;
 
-	void DebugLog() const;
+	RHI_API void DebugLog() const;
 
 #if WITH_ADDITIONAL_CRASH_CONTEXTS
-	static void WriteRenderBreadcrumbs(struct FCrashContextExtendedWriter& Writer, const FRHIBreadcrumb** BreadcrumbStack, uint32 BreadcrumbStackIndex, const TCHAR* ThreadName);
+	RHI_API static void WriteRenderBreadcrumbs(struct FCrashContextExtendedWriter& Writer, const FRHIBreadcrumb** BreadcrumbStack, uint32 BreadcrumbStackIndex, const TCHAR* ThreadName);
 #endif
 
 private:
diff --git a/Engine/Source/Runtime/RHI/Public/RHIGPUReadback.h b/Engine/Source/Runtime/RHI/Public/RHIGPUReadback.h
--- a/Engine/Source/Runtime/RHI/Public/RHIGPUReadback.h
+++ b/Engine/Source/Runtime/RHI/Public/RHIGPUReadback.h
@@ -17,7 +17,7 @@
  * FRHIGPUMemoryReadback: Represents a memory readback request scheduled with CopyToStagingBuffer
  * Wraps a staging buffer with a FRHIGPUFence for synchronization.
  */
-class RHI_API FRHIGPUMemoryReadback
+class FRHIGPUMemoryReadback
 {
 public:
 
@@ -87,16 +87,16 @@ protected:
 };
 
 /** Buffer readback implementation. */
-class RHI_API FRHIGPUBufferReadback final : public FRHIGPUMemoryReadback
+class FRHIGPUBufferReadback final : public FRHIGPUMemoryReadback
 {
 public:
 
-	FRHIGPUBufferReadback(FName RequestName);
+	RHI_API FRHIGPUBufferReadback(FName RequestName);
 	 
-	void EnqueueCopy(FRHICommandList& RHICmdList, FRHIBuffer* SourceBuffer, uint32 NumBytes = 0) override;
-	void* Lock(uint32 NumBytes) override;
-	void Unlock() override;
-	uint64 GetGPUSizeBytes() const;
+	RHI_API void EnqueueCopy(FRHICommandList& RHICmdList, FRHIBuffer* SourceBuffer, uint32 NumBytes = 0) override;
+	RHI_API void* Lock(uint32 NumBytes) override;
+	RHI_API void Unlock() override;
+	RHI_API uint64 GetGPUSizeBytes() const;
 
 private:
 
@@ -110,26 +110,26 @@ private:
 
 
 /** Texture readback implementation. */
-class RHI_API FRHIGPUTextureReadback final : public FRHIGPUMemoryReadback
+class FRHIGPUTextureReadback final : public FRHIGPUMemoryReadback
 {
 public:
-	FRHIGPUTextureReadback(FName RequestName);
+	RHI_API FRHIGPUTextureReadback(FName RequestName);
 
 	UE_DEPRECATED(5.1, "EnqueueCopyRDG is deprecated. Use EnqueueCopy instead.")
-	void EnqueueCopyRDG(FRHICommandList& RHICmdList, FRHITexture* SourceTexture, FResolveRect Rect = FResolveRect());
+	RHI_API void EnqueueCopyRDG(FRHICommandList& RHICmdList, FRHITexture* SourceTexture, FResolveRect Rect = FResolveRect());
 
-	void EnqueueCopy(FRHICommandList& RHICmdList, FRHITexture* SourceTexture, FResolveRect Rect = FResolveRect()) override;
+	RHI_API void EnqueueCopy(FRHICommandList& RHICmdList, FRHITexture* SourceTexture, FResolveRect Rect = FResolveRect()) override;
 
 	UE_DEPRECATED(5.0, "Use FRHIGPUTextureReadback::Lock( int32& OutRowPitchInPixels) instead.")
-	void* Lock(uint32 NumBytes) override;
+	RHI_API void* Lock(uint32 NumBytes) override;
 
-	void* Lock(int32& OutRowPitchInPixels, int32* OutBufferHeight = nullptr);
-	void Unlock() override;
+	RHI_API void* Lock(int32& OutRowPitchInPixels, int32* OutBufferHeight = nullptr);
+	RHI_API void Unlock() override;
 
 	UE_DEPRECATED(5.0, "Use FRHIGPUTextureReadback::Lock( int32& OutRowPitchInPixels) instead.")
-	void LockTexture(FRHICommandListImmediate& RHICmdList, void*& OutBufferPtr, int32& OutRowPitchInPixels);
+	RHI_API void LockTexture(FRHICommandListImmediate& RHICmdList, void*& OutBufferPtr, int32& OutRowPitchInPixels);
 
-	uint64 GetGPUSizeBytes() const;
+	RHI_API uint64 GetGPUSizeBytes() const;
 
 #if WITH_MGPU
 	FTextureRHIRef DestinationStagingTextures[MAX_NUM_GPUS];
diff --git a/Engine/Source/Runtime/RHI/Public/RHIResources.h b/Engine/Source/Runtime/RHI/Public/RHIResources.h
--- a/Engine/Source/Runtime/RHI/Public/RHIResources.h
+++ b/Engine/Source/Runtime/RHI/Public/RHIResources.h
@@ -873,13 +873,13 @@ public:
 	FRHIRayHitGroupShader() : FRHIRayTracingShader(SF_RayHitGroup) {}
 };
 
-class RHI_API FRHIComputeShader : public FRHIShader
+class FRHIComputeShader : public FRHIShader
 {
 public:
 	FRHIComputeShader() : FRHIShader(RRT_ComputeShader, SF_Compute), Stats(nullptr) {}
 	
 	inline void SetStats(struct FPipelineStateStats* Ptr) { Stats = Ptr; }
-	void UpdateStats();
+	RHI_API void UpdateStats();
 	
 private:
 	struct FPipelineStateStats* Stats;
@@ -1388,13 +1388,14 @@ private:
 // Textures
 //
 
-class RHI_API FLastRenderTimeContainer
+class FLastRenderTimeContainer
 {
 public:
 	FLastRenderTimeContainer() : LastRenderTime(-FLT_MAX) {}
 
 	double GetLastRenderTime() const { return LastRenderTime; }
-	FORCEINLINE_DEBUGGABLE void SetLastRenderTime(double InLastRenderTime) 
+
+	void SetLastRenderTime(double InLastRenderTime) 
 	{ 
 		// avoid dirty caches from redundant writes
 		if (LastRenderTime != InLastRenderTime)
@@ -1410,7 +1411,7 @@ private:
 
 
 /** Descriptor used to create a texture resource */
-struct RHI_API FRHITextureDesc
+struct FRHITextureDesc
 {
 	FRHITextureDesc()
 		: NumSamples(1)
@@ -1589,14 +1590,15 @@ struct RHI_API FRHITextureDesc
 	 * @param FirstMipIndex - the index of the most detailed mip to consider in the memory size calculation. Must be < NumMips and <= LastMipIndex.
 	 * @param LastMipIndex  - the index of the least detailed mip to consider in the memory size calculation. Must be < NumMips and >= FirstMipIndex.
 	 */
-	uint64 CalcMemorySizeEstimate(uint32 FirstMipIndex, uint32 LastMipIndex) const;
+	RHI_API uint64 CalcMemorySizeEstimate(uint32 FirstMipIndex, uint32 LastMipIndex) const;
+
 	uint64 CalcMemorySizeEstimate(uint32 FirstMipIndex = 0) const
 	{
 		return CalcMemorySizeEstimate(FirstMipIndex, NumMips - 1);
 	}
 
 private:
-	static bool Validate(const FRHITextureDesc& Desc, const TCHAR* Name, bool bFatal);
+	RHI_API static bool Validate(const FRHITextureDesc& Desc, const TCHAR* Name, bool bFatal);
 };
 
 // @todo deprecate
@@ -1757,14 +1759,14 @@ struct FRHITextureCreateDesc : public FRHITextureDesc
 	FResourceBulkDataInterface* BulkData = nullptr;
 };
 
-class RHI_API FRHITexture : public FRHIViewableResource
+class FRHITexture : public FRHIViewableResource
 #if ENABLE_RHI_VALIDATION
 	, public RHIValidation::FTextureResource
 #endif
 {
 protected:
 	/** Initialization constructor. Should only be called by platform RHI implementations. */
-	FRHITexture(const FRHITextureCreateDesc& InDesc);
+	RHI_API FRHITexture(const FRHITextureCreateDesc& InDesc);
 
 public:
 	/**
@@ -1913,7 +1915,7 @@ public:
 		return LastRenderTime.GetLastRenderTime();
 	}
 
-	void SetName(const FName& InName);
+	RHI_API void SetName(const FName& InName);
 
 	///
 	/// Deprecated functions
@@ -1986,7 +1988,7 @@ private:
 // Misc
 //
 
-class RHI_API FRHITimestampCalibrationQuery : public FRHIResource
+class FRHITimestampCalibrationQuery : public FRHIResource
 {
 public:
 	FRHITimestampCalibrationQuery() : FRHIResource(RRT_TimestampCalibrationQuery) {}
@@ -2001,7 +2003,7 @@ public:
 * The default implementation always returns false for Poll until the next frame from the frame the fence was inserted
 * because not all APIs have a GPU/CPU sync object, we need to fake it.
 */
-class RHI_API FRHIGPUFence : public FRHIResource
+class FRHIGPUFence : public FRHIResource
 {
 public:
 	FRHIGPUFence(FName InName) : FRHIResource(RRT_GPUFence), FenceName(InName) {}
@@ -2029,17 +2031,17 @@ protected:
 };
 
 // Generic implementation of FRHIGPUFence
-class RHI_API FGenericRHIGPUFence : public FRHIGPUFence
+class FGenericRHIGPUFence : public FRHIGPUFence
 {
 public:
-	FGenericRHIGPUFence(FName InName);
+	RHI_API FGenericRHIGPUFence(FName InName);
 
-	virtual void Clear() final override;
+	RHI_API virtual void Clear() final override;
 
 	/** @discussion RHI implementations must be thread-safe and must correctly handle being called before RHIInsertFence if an RHI thread is active. */
-	virtual bool Poll() const final override;
+	RHI_API virtual bool Poll() const final override;
 
-	void WriteInternal();
+	RHI_API void WriteInternal();
 
 private:
 	uint32 InsertedFrameNumber;
@@ -2052,7 +2054,8 @@ public:
 };
 
 class FRHIRenderQueryPool;
-class RHI_API FRHIPooledRenderQuery
+
+class FRHIPooledRenderQuery
 {
 	TRefCountPtr<FRHIRenderQuery> Query;
 	FRHIRenderQueryPool* QueryPool = nullptr;
@@ -2080,7 +2083,7 @@ public:
 	void ReleaseQuery();
 };
 
-class RHI_API FRHIRenderQueryPool : public FRHIResource
+class FRHIRenderQueryPool : public FRHIResource
 {
 public:
 	FRHIRenderQueryPool() : FRHIResource(RRT_RenderQueryPool) {}
@@ -2257,21 +2260,21 @@ public:
  * It is not comparable or convertible to other name types to encourage its use only for debugging and avoid using more storage than necessary
  * for the primary use cases of FName (names of objects, assets etc which are widely used and therefor deduped in the name table).
  */
-class RHI_API FDebugName
+class FDebugName
 {
 	DECLARE_INLINE_TYPE_LAYOUT(FDebugName, NonVirtual);
 
 public:
-	FDebugName();
-	FDebugName(FName InName);
-	FDebugName(FName InName, int32 InNumber);
-	FDebugName(FMemoryImageName InName, int32 InNumber);
+	RHI_API FDebugName();
+	RHI_API FDebugName(FName InName);
+	RHI_API FDebugName(FName InName, int32 InNumber);
+	RHI_API FDebugName(FMemoryImageName InName, int32 InNumber);
 
-	FDebugName& operator=(FName Other);
+	RHI_API FDebugName& operator=(FName Other);
 
-	FString ToString() const;
+	RHI_API FString ToString() const;
 	bool IsNone() const { return Name.IsNone() && Number == NAME_NO_NUMBER_INTERNAL; }
-	void AppendString(FStringBuilderBase& Builder) const;
+	RHI_API void AppendString(FStringBuilderBase& Builder) const;
 
 private:
 	LAYOUT_FIELD(FMemoryImageName, Name);
@@ -2583,7 +2586,7 @@ public:
 /* Generic staging buffer class used by FRHIGPUMemoryReadback
 * RHI specific staging buffers derive from this
 */
-class RHI_API FRHIStagingBuffer : public FRHIResource
+class FRHIStagingBuffer : public FRHIResource
 {
 public:
 	FRHIStagingBuffer()
@@ -2603,7 +2606,7 @@ protected:
 	bool bIsLocked;
 };
 
-class RHI_API FGenericRHIStagingBuffer : public FRHIStagingBuffer
+class FGenericRHIStagingBuffer : public FRHIStagingBuffer
 {
 public:
 	FGenericRHIStagingBuffer()
@@ -2612,8 +2615,8 @@ public:
 
 	~FGenericRHIStagingBuffer() {}
 
-	virtual void* Lock(uint32 Offset, uint32 NumBytes) final override;
-	virtual void Unlock() final override;
+	RHI_API virtual void* Lock(uint32 Offset, uint32 NumBytes) final override;
+	RHI_API virtual void Unlock() final override;
 	virtual uint64 GetGPUSizeBytes() const final override { return ShadowBuffer.IsValid() ? ShadowBuffer->GetSize() : 0; }
 
 	FBufferRHIRef ShadowBuffer;
@@ -4020,7 +4023,7 @@ enum ERHITextureSRVOverrideSRGBType : uint8
 	SRGBO_ForceDisable,
 };
 
-struct RHI_API FRHITextureSRVCreateInfo
+struct FRHITextureSRVCreateInfo
 {
 	explicit FRHITextureSRVCreateInfo(uint8 InMipLevel = 0u, uint8 InNumMipLevels = 1u, EPixelFormat InFormat = PF_Unknown)
 		: Format(InFormat)
@@ -4093,7 +4096,7 @@ struct RHI_API FRHITextureSRVCreateInfo
 	}
 
 protected:
-	static bool Validate(const FRHITextureDesc& TextureDesc, const FRHITextureSRVCreateInfo& TextureSRVDesc, const TCHAR* TextureName, bool bFatal);
+	RHI_API static bool Validate(const FRHITextureDesc& TextureDesc, const FRHITextureSRVCreateInfo& TextureSRVDesc, const TCHAR* TextureName, bool bFatal);
 };
 
 struct FRHITextureUAVCreateInfo
@@ -4242,18 +4245,18 @@ struct FRHIBufferUAVCreateInfo
 	bool bSupportsAppendBuffer = false;
 };
 
-class RHI_API FRHITextureViewCache
+class FRHITextureViewCache
 {
 public:
 	// Finds a UAV matching the descriptor in the cache or creates a new one and updates the cache.
-	FRHIUnorderedAccessView* GetOrCreateUAV(FRHITexture* Texture, const FRHITextureUAVCreateInfo& CreateInfo);
+	RHI_API FRHIUnorderedAccessView* GetOrCreateUAV(FRHITexture* Texture, const FRHITextureUAVCreateInfo& CreateInfo);
 
 	// Finds a SRV matching the descriptor in the cache or creates a new one and updates the cache.
-	FRHIShaderResourceView* GetOrCreateSRV(FRHITexture* Texture, const FRHITextureSRVCreateInfo& CreateInfo);
+	RHI_API FRHIShaderResourceView* GetOrCreateSRV(FRHITexture* Texture, const FRHITextureSRVCreateInfo& CreateInfo);
 
 	// Sets the debug name of the RHI view resources.
 #if !(UE_BUILD_SHIPPING || UE_BUILD_TEST)
-	void SetDebugName(const TCHAR* DebugName);
+	RHI_API void SetDebugName(const TCHAR* DebugName);
 #else
 	void SetDebugName(const TCHAR* DebugName) {}
 #endif
@@ -4263,18 +4266,18 @@ private:
 	TArray<TPair<FRHITextureSRVCreateInfo, FShaderResourceViewRHIRef>, TInlineAllocator<1>> SRVs;
 };
 
-class RHI_API FRHIBufferViewCache
+class FRHIBufferViewCache
 {
 public:
 	// Finds a UAV matching the descriptor in the cache or creates a new one and updates the cache.
-	FRHIUnorderedAccessView* GetOrCreateUAV(FRHIBuffer* Buffer, const FRHIBufferUAVCreateInfo& CreateInfo);
+	RHI_API FRHIUnorderedAccessView* GetOrCreateUAV(FRHIBuffer* Buffer, const FRHIBufferUAVCreateInfo& CreateInfo);
 
 	// Finds a SRV matching the descriptor in the cache or creates a new one and updates the cache.
-	FRHIShaderResourceView* GetOrCreateSRV(FRHIBuffer* Buffer, const FRHIBufferSRVCreateInfo& CreateInfo);
+	RHI_API FRHIShaderResourceView* GetOrCreateSRV(FRHIBuffer* Buffer, const FRHIBufferSRVCreateInfo& CreateInfo);
 
 	// Sets the debug name of the RHI view resources.
 #if !(UE_BUILD_SHIPPING || UE_BUILD_TEST)
-	void SetDebugName(const TCHAR* DebugName);
+	RHI_API void SetDebugName(const TCHAR* DebugName);
 #else
 	void SetDebugName(const TCHAR* DebugName) {}
 #endif
diff --git a/Engine/Source/Runtime/RHI/Public/RHITextureReference.h b/Engine/Source/Runtime/RHI/Public/RHITextureReference.h
--- a/Engine/Source/Runtime/RHI/Public/RHITextureReference.h
+++ b/Engine/Source/Runtime/RHI/Public/RHITextureReference.h
@@ -4,25 +4,25 @@
 
 #include "RHIResources.h"
 
-class RHI_API FRHITextureReference final : public FRHITexture
+class FRHITextureReference final : public FRHITexture
 {
 public:
-	explicit FRHITextureReference();
-	~FRHITextureReference();
+	RHI_API explicit FRHITextureReference();
+	RHI_API ~FRHITextureReference();
 
-	virtual class FRHITextureReference* GetTextureReference() override;
-	virtual FRHIDescriptorHandle GetDefaultBindlessHandle() const override;
+	RHI_API virtual class FRHITextureReference* GetTextureReference() override;
+	RHI_API virtual FRHIDescriptorHandle GetDefaultBindlessHandle() const override;
 
-	virtual void* GetNativeResource() const override;
-	virtual void* GetNativeShaderResourceView() const override;
-	virtual void* GetTextureBaseRHI() override;
-	virtual void GetWriteMaskProperties(void*& OutData, uint32& OutSize) override;
-	virtual const FRHITextureDesc& GetDesc() const override;
+	RHI_API virtual void* GetNativeResource() const override;
+	RHI_API virtual void* GetNativeShaderResourceView() const override;
+	RHI_API virtual void* GetTextureBaseRHI() override;
+	RHI_API virtual void GetWriteMaskProperties(void*& OutData, uint32& OutSize) override;
+	RHI_API virtual const FRHITextureDesc& GetDesc() const override;
 
 #if ENABLE_RHI_VALIDATION
 	// Implement RHIValidation::FTextureResource::GetTrackerResource to use the tracker info
 	// for the referenced texture.
-	virtual RHIValidation::FResource* GetTrackerResource() final override;
+	RHI_API virtual RHIValidation::FResource* GetTrackerResource() final override;
 #endif
 
 	inline FRHITexture * GetReferencedTexture() const
@@ -30,7 +30,7 @@ public:
 		return ReferencedTexture.GetReference();
 	}
 
-	void UpdateBindlessShaderResourceView();
+	RHI_API void UpdateBindlessShaderResourceView();
 
 private:
 	friend class FRHICommandListImmediate;
@@ -53,5 +53,5 @@ private:
 	// to allow FRHITextureReference to use the global black texture when the reference is nullptr.
 	// A pointer is required since FBlackTextureWithSRV is defined in RenderCore.
 	friend class FBlackTextureWithSRV;
-	static TRefCountPtr<FRHITexture> DefaultTexture;
+	RHI_API static TRefCountPtr<FRHITexture> DefaultTexture;
 };
diff --git a/Engine/Source/Runtime/RHI/Public/RHITransientResourceAllocator.h b/Engine/Source/Runtime/RHI/Public/RHITransientResourceAllocator.h
--- a/Engine/Source/Runtime/RHI/Public/RHITransientResourceAllocator.h
+++ b/Engine/Source/Runtime/RHI/Public/RHITransientResourceAllocator.h
@@ -235,7 +235,7 @@ private:
 	ERHITransientResourceType ResourceType;
 };
 
-class RHI_API FRHITransientTexture final : public FRHITransientResource
+class FRHITransientTexture final : public FRHITransientResource
 {
 public:
 	FRHITransientTexture(
@@ -251,7 +251,7 @@ public:
 
 	//////////////////////////////////////////////////////////////////////////////////////////////////
 	//! Internal Allocator API
-	void Acquire(const TCHAR* InName, uint32 InAcquirePassIndex, uint64 InInitCycle) override;
+	RHI_API void Acquire(const TCHAR* InName, uint32 InAcquirePassIndex, uint64 InInitCycle) override;
 	//////////////////////////////////////////////////////////////////////////////////////////////////
 
 	// Returns the underlying RHI texture.
@@ -273,7 +273,7 @@ public:
 	FRHITextureViewCache ViewCache;
 };
 
-class RHI_API FRHITransientBuffer final : public FRHITransientResource
+class FRHITransientBuffer final : public FRHITransientResource
 {
 public:
 	FRHITransientBuffer(
@@ -289,7 +289,7 @@ public:
 
 	//////////////////////////////////////////////////////////////////////////////////////////////////
 	//! Internal Allocator API
-	void Acquire(const TCHAR* InName, uint32 InAcquirePassIndex, uint64 InInitCycle) override;
+	RHI_API void Acquire(const TCHAR* InName, uint32 InAcquirePassIndex, uint64 InInitCycle) override;
 	//////////////////////////////////////////////////////////////////////////////////////////////////
 
 	// Returns the underlying RHI buffer.
diff --git a/Engine/Source/Runtime/RHI/Public/RHIUtilities.h b/Engine/Source/Runtime/RHI/Public/RHIUtilities.h
--- a/Engine/Source/Runtime/RHI/Public/RHIUtilities.h
+++ b/Engine/Source/Runtime/RHI/Public/RHIUtilities.h
@@ -527,17 +527,17 @@ inline uint32 ComputeAnisotropyRT(int32 InitializerMaxAnisotropy)
 #define ENABLE_TRANSITION_DUMP 1
 #endif
 
-class RHI_API FDumpTransitionsHelper
+class FDumpTransitionsHelper
 {
 public:
-	static void DumpResourceTransition(const FName& ResourceName, const ERHIAccess TransitionType);
+	RHI_API static void DumpResourceTransition(const FName& ResourceName, const ERHIAccess TransitionType);
 	
 private:
-	static void DumpTransitionForResourceHandler();
+	RHI_API static void DumpTransitionForResourceHandler();
 
-	static TAutoConsoleVariable<FString> CVarDumpTransitionsForResource;
-	static FAutoConsoleVariableSink CVarDumpTransitionsForResourceSink;
-	static FName DumpTransitionForResource;
+	RHI_API static TAutoConsoleVariable<FString> CVarDumpTransitionsForResource;
+	RHI_API static FAutoConsoleVariableSink CVarDumpTransitionsForResourceSink;
+	RHI_API static FName DumpTransitionForResource;
 };
 
 #if ENABLE_TRANSITION_DUMP
diff --git a/Engine/Source/Runtime/RHI/Public/RHIValidationCommon.h b/Engine/Source/Runtime/RHI/Public/RHIValidationCommon.h
--- a/Engine/Source/Runtime/RHI/Public/RHIValidationCommon.h
+++ b/Engine/Source/Runtime/RHI/Public/RHIValidationCommon.h
@@ -382,7 +382,7 @@ namespace RHIValidation
 	public:
 	};
 
-	class RHI_API FTextureResource
+	class FTextureResource
 	{
 	private:
 		// Don't use inheritance here. Because FRHITextureReferences exist, we have to
@@ -391,13 +391,13 @@ namespace RHIValidation
 
 	public:
 		FTextureResource() = default;
-		FTextureResource(FRHITextureCreateDesc const& CreateDesc);
+		RHI_API FTextureResource(FRHITextureCreateDesc const& CreateDesc);
 
 		virtual ~FTextureResource() {}
 
 		virtual FResource* GetTrackerResource() { return &PRIVATE_TrackerResource; }
 
-		void InitBarrierTracking(FRHITextureCreateDesc const& CreateDesc);
+		RHI_API void InitBarrierTracking(FRHITextureCreateDesc const& CreateDesc);
 
 		inline bool IsBarrierTrackingInitialized() const
 		{
@@ -406,10 +406,10 @@ namespace RHIValidation
 			return const_cast<FTextureResource*>(this)->GetTrackerResource()->IsBarrierTrackingInitialized();
 		}
 
-		void InitBarrierTracking(int32 InNumMips, int32 InNumArraySlices, EPixelFormat PixelFormat, ETextureCreateFlags Flags, ERHIAccess InResourceState, const TCHAR* InDebugName);
+		RHI_API void InitBarrierTracking(int32 InNumMips, int32 InNumArraySlices, EPixelFormat PixelFormat, ETextureCreateFlags Flags, ERHIAccess InResourceState, const TCHAR* InDebugName);
 
-		FResourceIdentity GetViewIdentity(uint32 InMipIndex, uint32 InNumMips, uint32 InArraySlice, uint32 InNumArraySlices, uint32 InPlaneIndex, uint32 InNumPlanes);
-		FResourceIdentity GetTransitionIdentity(const FRHITransitionInfo& Info);
+		RHI_API FResourceIdentity GetViewIdentity(uint32 InMipIndex, uint32 InNumMips, uint32 InArraySlice, uint32 InNumArraySlices, uint32 InPlaneIndex, uint32 InNumPlanes);
+		RHI_API FResourceIdentity GetTransitionIdentity(const FRHITransitionInfo& Info);
 
 		inline FResourceIdentity GetWholeResourceIdentity()
 		{
diff --git a/Engine/Source/Runtime/RHI/Public/TextureProfiler.h b/Engine/Source/Runtime/RHI/Public/TextureProfiler.h
--- a/Engine/Source/Runtime/RHI/Public/TextureProfiler.h
+++ b/Engine/Source/Runtime/RHI/Public/TextureProfiler.h
@@ -17,21 +17,21 @@ class FRHITexture;
 /**
 * FTextureProfiler class. This manages recording and reporting texture allocations in the RHI
 */
-class RHI_API FTextureProfiler
+class FTextureProfiler
 {
 	static FTextureProfiler* Instance;
 public:
 	// Singleton interface
-	static FTextureProfiler* Get();
+	RHI_API static FTextureProfiler* Get();
 
-	void Init();
+	RHI_API void Init();
 
-	void DumpTextures(bool RenderTargets, bool CombineTextureNames, bool AsCSV, FOutputDevice& OutputDevice);
+	RHI_API void DumpTextures(bool RenderTargets, bool CombineTextureNames, bool AsCSV, FOutputDevice& OutputDevice);
 
-	void AddTextureAllocation(FRHITexture* UniqueTexturePtr, size_t Size, uint32 Alignment, size_t AllocationWaste);
-	void UpdateTextureAllocation(FRHITexture* UniqueTexturePtr, size_t Size, uint32 Alignment, size_t AllocationWaste);
-	void RemoveTextureAllocation(FRHITexture* UniqueTexturePtr);
-	void UpdateTextureName(FRHITexture* UniqueTexturePtr);
+	RHI_API void AddTextureAllocation(FRHITexture* UniqueTexturePtr, size_t Size, uint32 Alignment, size_t AllocationWaste);
+	RHI_API void UpdateTextureAllocation(FRHITexture* UniqueTexturePtr, size_t Size, uint32 Alignment, size_t AllocationWaste);
+	RHI_API void RemoveTextureAllocation(FRHITexture* UniqueTexturePtr);
+	RHI_API void UpdateTextureName(FRHITexture* UniqueTexturePtr);
 
 private:
 
@@ -39,10 +39,10 @@ private:
 	FTextureProfiler(const FTextureProfiler&) = delete;
 	FTextureProfiler(FTextureProfiler&&) = delete;
 
-	void Update();
+	RHI_API void Update();
 
 	// Should only be used by FTextureDetails
-	const char* GetTextureNameString(FName TextureName);
+	RHI_API const char* GetTextureNameString(FName TextureName);
 
 	FCriticalSection TextureMapCS;
 
@@ -50,19 +50,19 @@ private:
 	{
 	public:
 		FTextureDetails() = default;
-		FTextureDetails(FRHITexture* Texture, size_t InSize, uint32 InAlign, size_t InAllocationWaste);
-		~FTextureDetails();
+		RHI_API FTextureDetails(FRHITexture* Texture, size_t InSize, uint32 InAlign, size_t InAllocationWaste);
+		RHI_API ~FTextureDetails();
 		
-		void SetName(FName InTextureName);
+		RHI_API void SetName(FName InTextureName);
 
-		void ResetPeakSize();
-		void SetValues(const FTextureDetails& Values);
+		RHI_API void ResetPeakSize();
+		RHI_API void SetValues(const FTextureDetails& Values);
 		FName GetTextureName() const { return TextureName; }
 		const char* GetTextureNameString() const { check(TextureNameString != nullptr);  return TextureNameString; }
 
-		FTextureDetails& operator+=(const FTextureDetails& Other);
-		FTextureDetails& operator-=(const FTextureDetails& Other);
-		FTextureDetails& operator=(const FTextureDetails& Other) = default;
+		RHI_API FTextureDetails& operator+=(const FTextureDetails& Other);
+		RHI_API FTextureDetails& operator-=(const FTextureDetails& Other);
+		RHI_API FTextureDetails& operator=(const FTextureDetails& Other) = default;
 
 		size_t Size = 0;
 		size_t PeakSize = 0;
