RHI shader parameter changes
- Moving the shader parameter types to their own header.
- Adding utility function to RHIContext to set batched shader parameters.
- Adding validation to scratch shader parameters to ensure they're empty before every use.
- SetBatchedShaderParameters now clears the batched parameters after they're used.

diff --git a/Engine/Source/Runtime/RHI/Private/RHICommandList.cpp b/Engine/Source/Runtime/RHI/Private/RHICommandList.cpp
--- a/Engine/Source/Runtime/RHI/Private/RHICommandList.cpp
+++ b/Engine/Source/Runtime/RHI/Private/RHICommandList.cpp
@@ -195,7 +195,6 @@ void FRHIParameterBatcher::OnBoundShaderChanged(FRHIComputeCommandList& InComman
 void FRHIParameterBatcher::PreDispatch(FRHIComputeCommandList& InCommandList)
 {
 	InCommandList.SetShaderParametersFromInternalBatched(GetBatchedComputeShader(), AllBatchedShaderParameters[SF_Compute]);
-	AllBatchedShaderParameters[SF_Compute].Reset();
 }
 
 void FRHIParameterBatcher::PreDraw(FRHICommandList& InCommandList)
@@ -203,7 +202,6 @@ void FRHIParameterBatcher::PreDraw(FRHICommandList& InCommandList)
 	for (uint32 Index = 0; Index < SF_NumGraphicsFrequencies; Index++)
 	{
 		InCommandList.SetShaderParametersFromInternalBatched(GetBatchedGraphicsShader(Index), AllBatchedShaderParameters[Index]);
-		AllBatchedShaderParameters[Index].Reset();
 	}
 }
 
@@ -224,7 +222,6 @@ void FRHIParameterBatcher::FlushPendingParameters(FRHIComputeCommandList& InComm
 	if (Frequency == SF_Compute)
 	{
 		InCommandList.SetShaderParametersFromInternalBatched(GetBatchedComputeShader(), AllBatchedShaderParameters[SF_Compute]);
-		AllBatchedShaderParameters[SF_Compute].Reset();
 	}
 }
 
@@ -234,7 +231,6 @@ void FRHIParameterBatcher::FlushPendingParameters(FRHICommandList& InCommandList
 	if (Frequency < SF_NumGraphicsFrequencies)
 	{
 		InCommandList.SetShaderParametersFromInternalBatched(GetBatchedGraphicsShader(Frequency), AllBatchedShaderParameters[Frequency]);
-		AllBatchedShaderParameters[Frequency].Reset();
 	}
 }
 
diff --git a/Engine/Source/Runtime/RHI/Public/RHICommandList.h b/Engine/Source/Runtime/RHI/Public/RHICommandList.h
--- a/Engine/Source/Runtime/RHI/Public/RHICommandList.h
+++ b/Engine/Source/Runtime/RHI/Public/RHICommandList.h
@@ -29,6 +29,7 @@
 #include "ProfilingDebugging/CsvProfiler.h"
 #include "ProfilingDebugging/CpuProfilerTrace.h"
 #include "RHIBreadcrumbs.h"
+#include "RHIShaderParameters.h"
 #include "RHITextureReference.h"
 #include "Trace/Trace.h"
 
@@ -479,157 +480,6 @@ public:
 	~FRHICommandListScopedPipelineGuard();
 };
 
-struct FRHIBatchedShaderParameters
-{
-	TArray<uint8> ParametersData;
-	TArray<FRHIShaderParameter> Parameters;
-	TArray<FRHIShaderParameterResource> ResourceParameters;
-	TArray<FRHIShaderParameterResource> BindlessParameters;
-
-	inline bool HasParameters() const
-	{
-		return (Parameters.Num() + ResourceParameters.Num() + BindlessParameters.Num()) > 0;
-	}
-
-	void Reset()
-	{
-		ParametersData.Reset();
-		Parameters.Reset();
-		ResourceParameters.Reset();
-		BindlessParameters.Reset();
-	}
-
-	FORCEINLINE_DEBUGGABLE void SetShaderParameter(uint32 BufferIndex, uint32 BaseIndex, uint32 NumBytes, const void* NewValue)
-	{
-		const int32 DestDataOffset = ParametersData.Num();
-		ParametersData.Append((const uint8*)NewValue, NumBytes);
-		Parameters.Emplace((uint16)BufferIndex, (uint16)BaseIndex, (uint16)DestDataOffset, (uint16)NumBytes);
-	}
-
-	FORCEINLINE_DEBUGGABLE void SetShaderUniformBuffer(uint32 Index, FRHIUniformBuffer* UniformBuffer)
-	{
-		ResourceParameters.Emplace(UniformBuffer, (uint16)Index);
-	}
-
-	FORCEINLINE_DEBUGGABLE void SetShaderTexture(uint32 Index, FRHITexture* Texture)
-	{
-		ResourceParameters.Emplace(Texture, (uint16)Index);
-	}
-	FORCEINLINE_DEBUGGABLE void SetShaderResourceViewParameter(uint32 Index, FRHIShaderResourceView* SRV)
-	{
-		ResourceParameters.Emplace(SRV, (uint16)Index);
-	}
-	FORCEINLINE_DEBUGGABLE void SetShaderSampler(uint32 Index, FRHISamplerState* State)
-	{
-		ResourceParameters.Emplace(State, (uint16)Index);
-	}
-	FORCEINLINE_DEBUGGABLE void SetUAVParameter(uint32 Index, FRHIUnorderedAccessView* UAV)
-	{
-		ResourceParameters.Emplace(UAV, (uint16)Index);
-	}
-
-	FORCEINLINE_DEBUGGABLE void SetBindlessTexture(uint32 Index, FRHITexture* Texture)
-	{
-		BindlessParameters.Emplace(Texture, (uint16)Index);
-	}
-	FORCEINLINE_DEBUGGABLE void SetBindlessResourceView(uint32 Index, FRHIShaderResourceView* SRV)
-	{
-		BindlessParameters.Emplace(SRV, (uint16)Index);
-	}
-	FORCEINLINE_DEBUGGABLE void SetBindlessSampler(uint32 Index, FRHISamplerState* State)
-	{
-		BindlessParameters.Emplace(State, (uint16)Index);
-	}
-	FORCEINLINE_DEBUGGABLE void SetBindlessUAV(uint32 Index, FRHIUnorderedAccessView* UAV)
-	{
-		BindlessParameters.Emplace(UAV, (uint16)Index);
-	}
-};
-
-class FRHIParameterBatcher
-{
-protected:
-	FRHIShader* AllBatchedShaders[SF_NumStandardFrequencies]{};
-	FRHIBatchedShaderParameters AllBatchedShaderParameters[SF_NumStandardFrequencies];
-	bool bEnabled = false;
-
-	FRHIGraphicsShader* GetBatchedGraphicsShader(int32 Index) { return static_cast<FRHIGraphicsShader*>(AllBatchedShaders[Index]); }
-	FRHIComputeShader*  GetBatchedComputeShader()             { return static_cast<FRHIComputeShader*>(AllBatchedShaders[SF_Compute]);   }
-
-	FORCEINLINE_DEBUGGABLE FRHIBatchedShaderParameters& GetBatchedShaderParameters(FRHIShader* InShader)
-	{
-		checkSlow(InShader);
-		const EShaderFrequency Frequency = InShader->GetFrequency();
-		checkSlow(Frequency < SF_NumStandardFrequencies);
-		check(AllBatchedShaders[Frequency] == InShader);
-		return AllBatchedShaderParameters[Frequency];
-	}
-
-public:
-	RHI_API FRHIParameterBatcher();
-	RHI_API FRHIParameterBatcher(const FBoundShaderStateInput& InBoundShaderStateInput, FRHIComputeShader* InBoundComputeShaderRHI);
-	RHI_API FRHIParameterBatcher(FRHIParameterBatcher&&);
-	RHI_API ~FRHIParameterBatcher();
-
-	FORCEINLINE bool IsEnabled() const { return bEnabled; }
-
-	RHI_API void OnBoundShaderChanged(class FRHICommandList& InCommandList, const FBoundShaderStateInput& InBoundShaderStateInput);
-	RHI_API void OnBoundShaderChanged(class FRHIComputeCommandList& InCommandList, FRHIComputeShader* InBoundComputeShaderRHI);
-
-	RHI_API void PreDispatch(class FRHIComputeCommandList& InCommandList);
-	RHI_API void PreDraw(class FRHICommandList& InCommandList);
-
-	RHI_API void FlushAllParameters(class FRHIComputeCommandList& InCommandList);
-	RHI_API void FlushAllParameters(class FRHICommandList& InCommandList);
-
-	RHI_API void FlushPendingParameters(class FRHIComputeCommandList& InCommandList, FRHIComputeShader* InShader);
-	RHI_API void FlushPendingParameters(class FRHICommandList& InCommandList, FRHIGraphicsShader* InShader);
-
-	FORCEINLINE_DEBUGGABLE void SetShaderParameter(FRHIShader* Shader, uint32 BufferIndex, uint32 BaseIndex, uint32 NumBytes, const void* NewValue)
-	{
-		GetBatchedShaderParameters(Shader).SetShaderParameter(BufferIndex, BaseIndex, NumBytes, NewValue);
-	}
-
-	FORCEINLINE_DEBUGGABLE void SetShaderTexture(FRHIShader* Shader, uint32 Index, FRHITexture* Texture)
-	{
-		GetBatchedShaderParameters(Shader).SetShaderTexture(Index, Texture);
-	}
-	FORCEINLINE_DEBUGGABLE void SetShaderResourceViewParameter(FRHIShader* Shader, uint32 Index, FRHIShaderResourceView* SRV)
-	{
-		GetBatchedShaderParameters(Shader).SetShaderResourceViewParameter(Index, SRV);
-	}
-	FORCEINLINE_DEBUGGABLE void SetShaderSampler(FRHIShader* Shader, uint32 Index, FRHISamplerState* State)
-	{
-		GetBatchedShaderParameters(Shader).SetShaderSampler(Index, State);
-	}
-	FORCEINLINE_DEBUGGABLE void SetUAVParameter(FRHIShader* Shader, uint32 UAVIndex, FRHIUnorderedAccessView* UAV)
-	{
-		GetBatchedShaderParameters(Shader).SetUAVParameter(UAVIndex, UAV);
-	}
-
-	FORCEINLINE_DEBUGGABLE void SetBindlessTexture(FRHIShader* Shader, uint32 Index, FRHITexture* Texture)
-	{
-		GetBatchedShaderParameters(Shader).SetBindlessTexture(Index, Texture);
-	}
-	FORCEINLINE_DEBUGGABLE void SetBindlessResourceView(FRHIShader* Shader, uint32 Index, FRHIShaderResourceView* SRV)
-	{
-		GetBatchedShaderParameters(Shader).SetBindlessResourceView(Index, SRV);
-	}
-	FORCEINLINE_DEBUGGABLE void SetBindlessSampler(FRHIShader* Shader, uint32 Index, FRHISamplerState* State)
-	{
-		GetBatchedShaderParameters(Shader).SetBindlessSampler(Index, State);
-	}
-	FORCEINLINE_DEBUGGABLE void SetBindlessUAV(FRHIShader* Shader, uint32 Index, FRHIUnorderedAccessView* UAV)
-	{
-		GetBatchedShaderParameters(Shader).SetBindlessUAV(Index, UAV);
-	}
-
-	FORCEINLINE_DEBUGGABLE void SetUAVParameter(FRHIShader* Shader, uint32 UAVIndex, FRHIUnorderedAccessView* UAV, uint32 InitialCount)
-	{
-		checkNoEntry(); // @todo: support append/consume buffers
-	}
-};
-
 class RHI_API FRHICommandListBase : public FNoncopyable
 {
 public:
@@ -877,7 +727,10 @@ public:
 
 	inline FRHIBatchedShaderParameters& GetScratchShaderParameters()
 	{
-		ScratchShaderParameters.Reset();
+		if (!ensureMsgf(!ScratchShaderParameters.HasParameters(), TEXT("Scratch shader parameters left without committed parameters")))
+		{
+			ScratchShaderParameters.Reset();
+		}
 		return ScratchShaderParameters;
 	}
 
@@ -2404,7 +2257,7 @@ public:
 		);
 	}
 
-	FORCEINLINE_DEBUGGABLE void SetBatchedShaderParameters(FRHIComputeShader* InShader, const FRHIBatchedShaderParameters& InBatchedParameters)
+	FORCEINLINE_DEBUGGABLE void SetBatchedShaderParameters(FRHIComputeShader* InShader, FRHIBatchedShaderParameters& InBatchedParameters)
 	{
 		if (InBatchedParameters.HasParameters())
 		{
@@ -2416,10 +2269,12 @@ public:
 				InBatchedParameters.BindlessParameters,
 				true /* bSubmitBatchedParameters */
 			);
+
+			InBatchedParameters.Reset();
 		}
 	}
 
-	FORCEINLINE_DEBUGGABLE void SetShaderParametersFromInternalBatched(FRHIComputeShader* InShader, const FRHIBatchedShaderParameters& InBatchedParameters)
+	FORCEINLINE_DEBUGGABLE void SetShaderParametersFromInternalBatched(FRHIComputeShader* InShader, FRHIBatchedShaderParameters& InBatchedParameters)
 	{
 		if (InBatchedParameters.HasParameters())
 		{
@@ -2431,6 +2286,8 @@ public:
 				InBatchedParameters.BindlessParameters,
 				false /* bSubmitBatchedParameters */
 			);
+
+			InBatchedParameters.Reset();
 		}
 	}
 
@@ -3201,7 +3058,7 @@ public:
 
 	using FRHIComputeCommandList::SetBatchedShaderParameters;
 
-	FORCEINLINE_DEBUGGABLE void SetBatchedShaderParameters(FRHIGraphicsShader* InShader, const FRHIBatchedShaderParameters& InBatchedParameters)
+	FORCEINLINE_DEBUGGABLE void SetBatchedShaderParameters(FRHIGraphicsShader* InShader, FRHIBatchedShaderParameters& InBatchedParameters)
 	{
 		if (InBatchedParameters.HasParameters())
 		{
@@ -3213,12 +3070,14 @@ public:
 				InBatchedParameters.BindlessParameters,
 				true /* bSubmitBatchedParameters */
 			);
+
+			InBatchedParameters.Reset();
 		}
 	}
 
 	using FRHIComputeCommandList::SetShaderParametersFromInternalBatched;
 
-	FORCEINLINE_DEBUGGABLE void SetShaderParametersFromInternalBatched(FRHIGraphicsShader* InShader, const FRHIBatchedShaderParameters& InBatchedParameters)
+	FORCEINLINE_DEBUGGABLE void SetShaderParametersFromInternalBatched(FRHIGraphicsShader* InShader, FRHIBatchedShaderParameters& InBatchedParameters)
 	{
 		if (InBatchedParameters.HasParameters())
 		{
@@ -3230,6 +3089,8 @@ public:
 				InBatchedParameters.BindlessParameters,
 				false /* bSubmitBatchedParameters */
 			);
+
+			InBatchedParameters.Reset();
 		}
 	}
 
diff --git a/Engine/Source/Runtime/RHI/Public/RHIContext.h b/Engine/Source/Runtime/RHI/Public/RHIContext.h
--- a/Engine/Source/Runtime/RHI/Public/RHIContext.h
+++ b/Engine/Source/Runtime/RHI/Public/RHIContext.h
@@ -18,6 +18,7 @@
 #include "Modules/ModuleInterface.h"
 #include "RHIBreadcrumbs.h"
 #include "RHIResources.h"
+#include "RHIShaderParameters.h"
 
 class FRHIDepthRenderTargetView;
 class FRHIRenderTargetView;
@@ -306,6 +307,18 @@ public:
 
 	virtual void RHISetShaderParameters(FRHIComputeShader* ComputeShader, TConstArrayView<uint8> InParametersData, TConstArrayView<FRHIShaderParameter> InParameters, TConstArrayView<FRHIShaderParameterResource> InResourceParameters, TConstArrayView<FRHIShaderParameterResource> InBindlessParameters) = 0;
 
+	void RHISetBatchedShaderParameters(FRHIComputeShader* InShader, FRHIBatchedShaderParameters& InBatchedParameters)
+	{
+		RHISetShaderParameters(
+			InShader,
+			InBatchedParameters.ParametersData,
+			InBatchedParameters.Parameters,
+			InBatchedParameters.ResourceParameters,
+			InBatchedParameters.BindlessParameters);
+
+		InBatchedParameters.Reset();
+	}
+
 	virtual void RHISetShaderParameter(FRHIComputeShader* ComputeShader, uint32 BufferIndex, uint32 BaseIndex, uint32 NumBytes, const void* NewValue) = 0;
 
 	virtual void RHISetStaticUniformBuffers(const FUniformBufferStaticBindings& InUniformBuffers)
@@ -637,6 +650,7 @@ public:
 	using IRHIComputeContext::RHISetShaderResourceViewParameter;
 	using IRHIComputeContext::RHISetShaderUniformBuffer;
 	using IRHIComputeContext::RHISetShaderParameters;
+	using IRHIComputeContext::RHISetBatchedShaderParameters;
 	using IRHIComputeContext::RHISetShaderParameter;
 
 	/** Set the shader resource view of a surface. */
@@ -665,6 +679,18 @@ public:
 
 	virtual void RHISetShaderParameters(FRHIGraphicsShader* Shader, TConstArrayView<uint8> InParametersData, TConstArrayView<FRHIShaderParameter> InParameters, TConstArrayView<FRHIShaderParameterResource> InResourceParameters, TConstArrayView<FRHIShaderParameterResource> InBindlessParameters) = 0;
 
+	void RHISetBatchedShaderParameters(FRHIGraphicsShader* InShader, FRHIBatchedShaderParameters& InBatchedParameters)
+	{
+		RHISetShaderParameters(
+			InShader,
+			InBatchedParameters.ParametersData,
+			InBatchedParameters.Parameters,
+			InBatchedParameters.ResourceParameters,
+			InBatchedParameters.BindlessParameters);
+
+		InBatchedParameters.Reset();
+	}
+
 	virtual void RHISetShaderParameter(FRHIGraphicsShader* Shader, uint32 BufferIndex, uint32 BaseIndex, uint32 NumBytes, const void* NewValue) = 0;
 
 	virtual void RHISetStencilRef(uint32 StencilRef) {}
diff --git a/Engine/Source/Runtime/RHI/Public/RHIResources.h b/Engine/Source/Runtime/RHI/Public/RHIResources.h
--- a/Engine/Source/Runtime/RHI/Public/RHIResources.h
+++ b/Engine/Source/Runtime/RHI/Public/RHIResources.h
@@ -4287,65 +4287,6 @@ private:
 	TArray<TPair<FRHIBufferSRVCreateInfo, FShaderResourceViewRHIRef>, TInlineAllocator<1>> SRVs;
 };
 
-struct FRHIShaderParameter
-{
-	FRHIShaderParameter(uint16 InBufferIndex, uint16 InBaseIndex, uint16 InByteOffset, uint16 InByteSize)
-		: BufferIndex(InBufferIndex)
-		, BaseIndex(InBaseIndex)
-		, ByteOffset(InByteOffset)
-		, ByteSize(InByteSize)
-	{
-	}
-	uint16 BufferIndex;
-	uint16 BaseIndex;
-	uint16 ByteOffset;
-	uint16 ByteSize;
-};
-
-struct FRHIShaderParameterResource
-{
-	enum class EType : uint8
-	{
-		Texture,
-		ResourceView,
-		UnorderedAccessView,
-		Sampler,
-		UniformBuffer,
-	};
-
-	FRHIShaderParameterResource() = default;
-	FRHIShaderParameterResource(EType InType, FRHIResource* InResource, uint16 InIndex)
-		: Resource(InResource)
-		, Index(InIndex)
-		, Type(InType)
-	{
-	}
-	FRHIShaderParameterResource(FRHITexture* InTexture, uint16 InIndex)
-		: FRHIShaderParameterResource(FRHIShaderParameterResource::EType::Texture, InTexture, InIndex)
-	{
-	}
-	FRHIShaderParameterResource(FRHIShaderResourceView* InView, uint16 InIndex)
-		: FRHIShaderParameterResource(FRHIShaderParameterResource::EType::ResourceView, InView, InIndex)
-	{
-	}
-	FRHIShaderParameterResource(FRHIUnorderedAccessView* InUAV, uint16 InIndex)
-		: FRHIShaderParameterResource(FRHIShaderParameterResource::EType::UnorderedAccessView, InUAV, InIndex)
-	{
-	}
-	FRHIShaderParameterResource(FRHISamplerState* InSamplerState, uint16 InIndex)
-		: FRHIShaderParameterResource(FRHIShaderParameterResource::EType::Sampler, InSamplerState, InIndex)
-	{
-	}
-	FRHIShaderParameterResource(FRHIUniformBuffer* InUniformBuffer, uint16 InIndex)
-		: FRHIShaderParameterResource(FRHIShaderParameterResource::EType::UniformBuffer, InUniformBuffer, InIndex)
-	{
-	}
-
-	FRHIResource* Resource = nullptr;
-	uint16        Index = 0;
-	EType         Type = EType::Texture;
-};
-
 #if UE_ENABLE_INCLUDE_ORDER_DEPRECATED_IN_5_2
 #include "RHIShaderLibrary.h"
 #include "RHITextureReference.h"
diff --git a/Engine/Source/Runtime/RHI/Public/RHIShaderParameters.h b/Engine/Source/Runtime/RHI/Public/RHIShaderParameters.h
new file mode 100644
--- /dev/null
+++ b/Engine/Source/Runtime/RHI/Public/RHIShaderParameters.h
@@ -0,0 +1,225 @@
+// Copyright Epic Games, Inc. All Rights Reserved.
+
+#pragma once
+
+#include "Containers/Array.h"
+#include "RHIResources.h"
+
+class FRHICommandList;
+class FRHIComputeCommandList;
+
+/** Compact representation of a bound shader parameter (read: value). Its offsets are for referencing their data in an associated blob. */
+struct FRHIShaderParameter
+{
+	FRHIShaderParameter(uint16 InBufferIndex, uint16 InBaseIndex, uint16 InByteOffset, uint16 InByteSize)
+		: BufferIndex(InBufferIndex)
+		, BaseIndex(InBaseIndex)
+		, ByteOffset(InByteOffset)
+		, ByteSize(InByteSize)
+	{
+	}
+	uint16 BufferIndex;
+	uint16 BaseIndex;
+	uint16 ByteOffset;
+	uint16 ByteSize;
+};
+
+/** Compact representation of a bound resource parameter (Texture, SRV, UAV, SamplerState, or UniformBuffer) */
+struct FRHIShaderParameterResource
+{
+	enum class EType : uint8
+	{
+		Texture,
+		ResourceView,
+		UnorderedAccessView,
+		Sampler,
+		UniformBuffer,
+	};
+
+	FRHIShaderParameterResource() = default;
+	FRHIShaderParameterResource(EType InType, FRHIResource* InResource, uint16 InIndex)
+		: Resource(InResource)
+		, Index(InIndex)
+		, Type(InType)
+	{
+	}
+	FRHIShaderParameterResource(FRHITexture* InTexture, uint16 InIndex)
+		: FRHIShaderParameterResource(FRHIShaderParameterResource::EType::Texture, InTexture, InIndex)
+	{
+	}
+	FRHIShaderParameterResource(FRHIShaderResourceView* InView, uint16 InIndex)
+		: FRHIShaderParameterResource(FRHIShaderParameterResource::EType::ResourceView, InView, InIndex)
+	{
+	}
+	FRHIShaderParameterResource(FRHIUnorderedAccessView* InUAV, uint16 InIndex)
+		: FRHIShaderParameterResource(FRHIShaderParameterResource::EType::UnorderedAccessView, InUAV, InIndex)
+	{
+	}
+	FRHIShaderParameterResource(FRHISamplerState* InSamplerState, uint16 InIndex)
+		: FRHIShaderParameterResource(FRHIShaderParameterResource::EType::Sampler, InSamplerState, InIndex)
+	{
+	}
+	FRHIShaderParameterResource(FRHIUniformBuffer* InUniformBuffer, uint16 InIndex)
+		: FRHIShaderParameterResource(FRHIShaderParameterResource::EType::UniformBuffer, InUniformBuffer, InIndex)
+	{
+	}
+
+	FRHIResource* Resource = nullptr;
+	uint16        Index = 0;
+	EType         Type = EType::Texture;
+};
+
+/** Collection of parameters to set in the RHI. These parameters aren't bound to any specific shader until SetBatchedShaderParameters is called. */
+struct FRHIBatchedShaderParameters
+{
+	TArray<uint8> ParametersData;
+	TArray<FRHIShaderParameter> Parameters;
+	TArray<FRHIShaderParameterResource> ResourceParameters;
+	TArray<FRHIShaderParameterResource> BindlessParameters;
+
+	inline bool HasParameters() const
+	{
+		return (Parameters.Num() + ResourceParameters.Num() + BindlessParameters.Num()) > 0;
+	}
+
+	void Reset()
+	{
+		ParametersData.Reset();
+		Parameters.Reset();
+		ResourceParameters.Reset();
+		BindlessParameters.Reset();
+	}
+
+	FORCEINLINE_DEBUGGABLE void SetShaderParameter(uint32 BufferIndex, uint32 BaseIndex, uint32 NumBytes, const void* NewValue)
+	{
+		const int32 DestDataOffset = ParametersData.Num();
+		ParametersData.Append((const uint8*)NewValue, NumBytes);
+		Parameters.Emplace((uint16)BufferIndex, (uint16)BaseIndex, (uint16)DestDataOffset, (uint16)NumBytes);
+	}
+
+	FORCEINLINE_DEBUGGABLE void SetShaderUniformBuffer(uint32 Index, FRHIUniformBuffer* UniformBuffer)
+	{
+		ResourceParameters.Emplace(UniformBuffer, (uint16)Index);
+	}
+
+	FORCEINLINE_DEBUGGABLE void SetShaderTexture(uint32 Index, FRHITexture* Texture)
+	{
+		ResourceParameters.Emplace(Texture, (uint16)Index);
+	}
+	FORCEINLINE_DEBUGGABLE void SetShaderResourceViewParameter(uint32 Index, FRHIShaderResourceView* SRV)
+	{
+		ResourceParameters.Emplace(SRV, (uint16)Index);
+	}
+	FORCEINLINE_DEBUGGABLE void SetShaderSampler(uint32 Index, FRHISamplerState* State)
+	{
+		ResourceParameters.Emplace(State, (uint16)Index);
+	}
+	FORCEINLINE_DEBUGGABLE void SetUAVParameter(uint32 Index, FRHIUnorderedAccessView* UAV)
+	{
+		ResourceParameters.Emplace(UAV, (uint16)Index);
+	}
+
+	FORCEINLINE_DEBUGGABLE void SetBindlessTexture(uint32 Index, FRHITexture* Texture)
+	{
+		BindlessParameters.Emplace(Texture, (uint16)Index);
+	}
+	FORCEINLINE_DEBUGGABLE void SetBindlessResourceView(uint32 Index, FRHIShaderResourceView* SRV)
+	{
+		BindlessParameters.Emplace(SRV, (uint16)Index);
+	}
+	FORCEINLINE_DEBUGGABLE void SetBindlessSampler(uint32 Index, FRHISamplerState* State)
+	{
+		BindlessParameters.Emplace(State, (uint16)Index);
+	}
+	FORCEINLINE_DEBUGGABLE void SetBindlessUAV(uint32 Index, FRHIUnorderedAccessView* UAV)
+	{
+		BindlessParameters.Emplace(UAV, (uint16)Index);
+	}
+};
+
+/** Class that automatically batches shader parameters on a per-stage basis.
+*   NOTE: this class will be phased out in favor of using FRHIBatchedShaderParameters directly.
+*/
+class FRHIParameterBatcher
+{
+protected:
+	FRHIShader* AllBatchedShaders[SF_NumStandardFrequencies]{};
+	FRHIBatchedShaderParameters AllBatchedShaderParameters[SF_NumStandardFrequencies];
+	bool bEnabled = false;
+
+	FRHIGraphicsShader* GetBatchedGraphicsShader(int32 Index) { return static_cast<FRHIGraphicsShader*>(AllBatchedShaders[Index]); }
+	FRHIComputeShader* GetBatchedComputeShader() { return static_cast<FRHIComputeShader*>(AllBatchedShaders[SF_Compute]); }
+
+	FORCEINLINE_DEBUGGABLE FRHIBatchedShaderParameters& GetBatchedShaderParameters(FRHIShader* InShader)
+	{
+		checkSlow(InShader);
+		const EShaderFrequency Frequency = InShader->GetFrequency();
+		checkSlow(Frequency < SF_NumStandardFrequencies);
+		check(AllBatchedShaders[Frequency] == InShader);
+		return AllBatchedShaderParameters[Frequency];
+	}
+
+public:
+	RHI_API FRHIParameterBatcher();
+	RHI_API FRHIParameterBatcher(const FBoundShaderStateInput& InBoundShaderStateInput, FRHIComputeShader* InBoundComputeShaderRHI);
+	RHI_API FRHIParameterBatcher(FRHIParameterBatcher&&);
+	RHI_API ~FRHIParameterBatcher();
+
+	FORCEINLINE bool IsEnabled() const { return bEnabled; }
+
+	RHI_API void OnBoundShaderChanged(FRHICommandList& InCommandList, const FBoundShaderStateInput& InBoundShaderStateInput);
+	RHI_API void OnBoundShaderChanged(FRHIComputeCommandList& InCommandList, FRHIComputeShader* InBoundComputeShaderRHI);
+
+	RHI_API void PreDispatch(FRHIComputeCommandList& InCommandList);
+	RHI_API void PreDraw(FRHICommandList& InCommandList);
+
+	RHI_API void FlushAllParameters(FRHIComputeCommandList& InCommandList);
+	RHI_API void FlushAllParameters(FRHICommandList& InCommandList);
+
+	RHI_API void FlushPendingParameters(FRHIComputeCommandList& InCommandList, FRHIComputeShader* InShader);
+	RHI_API void FlushPendingParameters(FRHICommandList& InCommandList, FRHIGraphicsShader* InShader);
+
+	FORCEINLINE_DEBUGGABLE void SetShaderParameter(FRHIShader* Shader, uint32 BufferIndex, uint32 BaseIndex, uint32 NumBytes, const void* NewValue)
+	{
+		GetBatchedShaderParameters(Shader).SetShaderParameter(BufferIndex, BaseIndex, NumBytes, NewValue);
+	}
+
+	FORCEINLINE_DEBUGGABLE void SetShaderTexture(FRHIShader* Shader, uint32 Index, FRHITexture* Texture)
+	{
+		GetBatchedShaderParameters(Shader).SetShaderTexture(Index, Texture);
+	}
+	FORCEINLINE_DEBUGGABLE void SetShaderResourceViewParameter(FRHIShader* Shader, uint32 Index, FRHIShaderResourceView* SRV)
+	{
+		GetBatchedShaderParameters(Shader).SetShaderResourceViewParameter(Index, SRV);
+	}
+	FORCEINLINE_DEBUGGABLE void SetShaderSampler(FRHIShader* Shader, uint32 Index, FRHISamplerState* State)
+	{
+		GetBatchedShaderParameters(Shader).SetShaderSampler(Index, State);
+	}
+	FORCEINLINE_DEBUGGABLE void SetUAVParameter(FRHIShader* Shader, uint32 UAVIndex, FRHIUnorderedAccessView* UAV)
+	{
+		GetBatchedShaderParameters(Shader).SetUAVParameter(UAVIndex, UAV);
+	}
+
+	FORCEINLINE_DEBUGGABLE void SetBindlessTexture(FRHIShader* Shader, uint32 Index, FRHITexture* Texture)
+	{
+		GetBatchedShaderParameters(Shader).SetBindlessTexture(Index, Texture);
+	}
+	FORCEINLINE_DEBUGGABLE void SetBindlessResourceView(FRHIShader* Shader, uint32 Index, FRHIShaderResourceView* SRV)
+	{
+		GetBatchedShaderParameters(Shader).SetBindlessResourceView(Index, SRV);
+	}
+	FORCEINLINE_DEBUGGABLE void SetBindlessSampler(FRHIShader* Shader, uint32 Index, FRHISamplerState* State)
+	{
+		GetBatchedShaderParameters(Shader).SetBindlessSampler(Index, State);
+	}
+	FORCEINLINE_DEBUGGABLE void SetBindlessUAV(FRHIShader* Shader, uint32 Index, FRHIUnorderedAccessView* UAV)
+	{
+		GetBatchedShaderParameters(Shader).SetBindlessUAV(Index, UAV);
+	}
+
+	FORCEINLINE_DEBUGGABLE void SetUAVParameter(FRHIShader* Shader, uint32 UAVIndex, FRHIUnorderedAccessView* UAV, uint32 InitialCount)
+	{
+		checkNoEntry(); // @todo: support append/consume buffers
+	}
+};
