Initial version of RHI globals in a single struct.

diff --git a/Engine/Source/Runtime/RHI/Private/RHI.cpp b/Engine/Source/Runtime/RHI/Private/RHI.cpp
--- a/Engine/Source/Runtime/RHI/Private/RHI.cpp
+++ b/Engine/Source/Runtime/RHI/Private/RHI.cpp
@@ -1284,186 +1284,6 @@ namespace RHIConfig
 	}
 }
 
-/**
- * RHI globals.
- */
-
-bool GIsRHIInitialized = false;
-int32 GRHIPersistentThreadGroupCount = 0;
-int32 GMaxTextureMipCount = MAX_TEXTURE_MIP_COUNT;
-bool GSupportsQuadBufferStereo = false;
-FString GRHIAdapterName;
-FString GRHIAdapterInternalDriverVersion;
-FString GRHIAdapterUserDriverVersion;
-FString GRHIAdapterDriverDate;
-bool GRHIAdapterDriverOnDenyList = false;
-uint32 GRHIVendorId = 0;
-uint32 GRHIDeviceId = 0;
-uint32 GRHIDeviceRevision = 0;
-bool GRHIDeviceIsAMDPreGCNArchitecture = false;
-bool GSupportsRenderDepthTargetableShaderResources = true;
-TRHIGlobal<bool> GSupportsRenderTargetFormat_PF_G8(true);
-TRHIGlobal<bool> GSupportsRenderTargetFormat_PF_FloatRGBA(true);
-bool GSupportsShaderFramebufferFetch = false;
-bool GSupportsShaderDepthStencilFetch = false;
-bool GSupportsShaderMRTFramebufferFetch = false;
-bool GSupportsPixelLocalStorage = false;
-bool GSupportsTimestampRenderQueries = false;
-bool GRHISupportsGPUTimestampBubblesRemoval = false;
-bool GRHISupportsFrameCyclesBubblesRemoval = false;
-bool GRHISupportsGPUUsage = false;
-bool GHardwareHiddenSurfaceRemoval = false;
-bool GRHISupportsAsyncTextureCreation = false;
-bool GRHISupportsQuadTopology = false;
-bool GRHISupportsRectTopology = false;
-bool GRHISupportsPrimitiveShaders = false;
-bool GRHISupportsAtomicUInt64 = false;
-bool GRHISupportsDX12AtomicUInt64 = false;
-bool GRHISupportsPipelineStateSortKey = false;
-bool GRHISupportsResummarizeHTile = false;
-bool GRHISupportsExplicitHTile = false;
-bool GRHISupportsExplicitFMask = false;
-bool GRHISupportsDepthUAV = false;
-bool GSupportsParallelRenderingTasksWithSeparateRHIThread = true;
-bool GRHIThreadNeedsKicking = false;
-int32 GRHIMaximumReccommendedOustandingOcclusionQueries = MAX_int32;
-bool GRHISupportsExactOcclusionQueries = true;
-bool GSupportsVolumeTextureRendering = true;
-bool GSupportsSeparateRenderTargetBlendState = false;
-bool GRHINeedsUnatlasedCSMDepthsWorkaround = false;
-bool GSupportsTexture3D = true;
-bool GSupportsMobileMultiView = false;
-bool GSupportsImageExternal = false;
-bool GRHISupportsDrawIndirect = true;
-bool GRHISupportsMultithreading = false;
-bool GRHISupportsAsyncGetRenderQueryResult = false;
-bool GRHISupportsUpdateFromBufferTexture = false;
-bool GSupportsWideMRT = true;
-bool GRHINeedsExtraDeletionLatency = false;
-bool GRHIForceNoDeletionLatencyForStreamingTextures = false;
-TRHIGlobal<int32> GMaxComputeDispatchDimension((1 << 16) - 1);
-bool GRHILazyShaderCodeLoading = false;
-bool GRHISupportsLazyShaderCodeLoading = false;
-TRHIGlobal<int32> GMaxShadowDepthBufferSizeX(2048);
-TRHIGlobal<int32> GMaxShadowDepthBufferSizeY(2048);
-TRHIGlobal<int32> GMaxTextureDimensions(2048);
-TRHIGlobal<int64> GMaxBufferDimensions(1<<27);
-TRHIGlobal<int64> GRHIMaxConstantBufferByteSize(1<<27);
-TRHIGlobal<int64> GMaxComputeSharedMemory(1<<15);
-TRHIGlobal<int32> GMaxVolumeTextureDimensions(2048);
-TRHIGlobal<int32> GMaxCubeTextureDimensions(2048);
-TRHIGlobal<int32> GMaxWorkGroupInvocations(1024);
-bool GRHISupportsRawViewsForAnyBuffer = false;
-bool GRHISupportsRWTextureBuffers = true;
-bool GRHISupportsVRS = false;
-bool GRHISupportsLateVRSUpdate = false;
-int32 GMaxTextureArrayLayers = 256;
-int32 GMaxTextureSamplers = 16;
-bool GUsingNullRHI = false;
-int32 GDrawUPVertexCheckCount = MAX_int32;
-int32 GDrawUPIndexCheckCount = MAX_int32;
-bool GTriggerGPUProfile = false;
-FString GGPUTraceFileName;
-bool GRHISupportsTextureStreaming = false;
-bool GSupportsDepthBoundsTest = false;
-bool GSupportsEfficientAsyncCompute = false;
-bool GRHISupportsBaseVertexIndex = true;
-bool GRHISupportsFirstInstance = false;
-bool GRHISupportsDynamicResolution = false;
-bool GRHISupportsRayTracing = false;
-bool GRHISupportsRayTracingPSOAdditions = false;
-bool GRHISupportsRayTracingDispatchIndirect = false;
-bool GRHISupportsRayTracingAsyncBuildAccelerationStructure = false;
-bool GRHISupportsRayTracingAMDHitToken = false;
-bool GRHISupportsInlineRayTracing = false;
-bool GRHISupportsRayTracingShaders = false;
-uint32 GRHIRayTracingAccelerationStructureAlignment = 0;
-uint32 GRHIRayTracingScratchBufferAlignment = 0;
-uint32 GRHIRayTracingShaderTableAlignment = 0;
-uint32 GRHIRayTracingInstanceDescriptorSize = 0;
-bool GRHISupportsWaveOperations = false;
-int32 GRHIMinimumWaveSize = 0;
-int32 GRHIMaximumWaveSize = 0;
-bool GRHISupportsRHIThread = false;
-bool GRHISupportsRHIOnTaskThread = false;
-bool GRHISupportsParallelRHIExecute = false;
-bool GSupportsParallelOcclusionQueries = false;
-bool GRHIRequiresRenderTargetForPixelShaderUAVs = false;
-bool GRHISupportsUAVFormatAliasing = false;
-bool GRHISupportsTextureViews = true;
-bool GRHISupportsDirectGPUMemoryLock = false;
-bool GRHISupportsMultithreadedShaderCreation = true;
-bool GRHISupportsMultithreadedResources = false;
-
-bool GRHISupportAsyncTextureStreamOut = false;
-bool GRHISupportsMSAADepthSampleAccess = false;
-
-bool GRHISupportsBackBufferWithCustomDepthStencil = true;
-
-bool GRHIIsHDREnabled = false;
-bool GRHISupportsHDROutput = false;
-
-bool GRHIVariableRateShadingEnabled = true;
-bool GRHIAttachmentVariableRateShadingEnabled = true;
-bool GRHISupportsPipelineVariableRateShading = false;
-bool GRHISupportsLargerVariableRateShadingSizes = false;
-bool GRHISupportsAttachmentVariableRateShading = false;
-bool GRHISupportsComplexVariableRateShadingCombinerOps = false;
-bool GRHISupportsVariableRateShadingAttachmentArrayTextures = false;
-int32 GRHIVariableRateShadingImageTileMaxWidth = 0;
-int32 GRHIVariableRateShadingImageTileMaxHeight = 0;
-int32 GRHIVariableRateShadingImageTileMinWidth = 0;
-int32 GRHIVariableRateShadingImageTileMinHeight = 0;
-EVRSImageDataType GRHIVariableRateShadingImageDataType = VRSImage_NotSupported;
-EPixelFormat GRHIVariableRateShadingImageFormat = PF_Unknown;
-bool GRHISupportsLateVariableRateShadingUpdate = false;
-
-EPixelFormat GRHIHDRDisplayOutputFormat = PF_FloatRGBA;
-
-FIntVector GRHIMaxDispatchThreadGroupsPerDimension(0, 0, 0);
-
-uint64 GRHIPresentCounter = 1;
-
-bool GRHISupportsArrayIndexFromAnyShader = false;
-bool GRHISupportsStencilRefFromPixelShader = false;
-bool GRHISupportsConservativeRasterization = false;
-
-bool GRHISupportsPipelineFileCache = false;
-
-bool GRHIDeviceIsIntegrated = false;
-
-/** Whether we are profiling GPU hitches. */
-bool GTriggerGPUHitchProfile = false;
-
-bool GRHISupportsPixelShaderUAVs = true;
-
-bool GRHISupportsMeshShadersTier0 = false;
-bool GRHISupportsMeshShadersTier1 = false;
-
-bool GRHISupportsShaderTimestamp = false;
-
-bool GRHISupportsEfficientUploadOnResourceCreation = false;
-
-bool GRHISupportsAsyncPipelinePrecompile = true;
-
-bool GRHISupportsMapWriteNoOverwrite = false;
-
-bool GRHISupportsSeparateDepthStencilCopyAccess = true;
-
-ERHIBindlessSupport GRHIBindlessSupport = ERHIBindlessSupport::Unsupported;
-bool GRHISupportsBindless = false;
-
-bool GRHISupportsReservedResources = false;
-
-RHI_API int32 volatile GCurrentTextureMemorySize = 0;
-RHI_API int32 volatile GCurrentRendertargetMemorySize = 0;
-RHI_API int64 GTexturePoolSize = 0 * 1024 * 1024;
-RHI_API int32 GPoolSizeVRAMPercentage = 0;
-
-RHI_API uint64 GDemotedLocalMemorySize = 0;
-
-RHI_API EShaderPlatform GShaderPlatformForFeatureLevel[ERHIFeatureLevel::Num] = {SP_NumPlatforms,SP_NumPlatforms,SP_NumPlatforms,SP_NumPlatforms,SP_NumPlatforms};
-
 // By default, read only states and UAV states are allowed to participate in state merging.
 ERHIAccess GRHIMergeableAccessMask = ERHIAccess::ReadOnlyMask | ERHIAccess::UAVMask;
 
@@ -2014,52 +1834,7 @@ FRHIPanicEvent& RHIGetPanicDelegate()
 	return RHIPanicEvent;
 }
 
-//
-//	MSAA sample offsets.
-//
-// https://docs.microsoft.com/en-us/windows/win32/api/d3d11/ne-d3d11-d3d11_standard_multisample_quality_levels
-FVector2f GRHIDefaultMSAASampleOffsets[1 + 2 + 4 + 8 + 16] = {
-	// MSAA x1
-	FVector2f(+0.0f / 8.0f, +0.0f / 8.0f),
-
-	// MSAA x2
-	FVector2f(+4.0f / 8.0f, +4.0f / 8.0f),
-	FVector2f(-4.0f / 8.0f, -4.0f / 8.0f),
-
-	// MSAA x4
-	FVector2f(-2.0f / 8.0f, -6.0f / 8.0f),
-	FVector2f(+6.0f / 8.0f, -2.0f / 8.0f),
-	FVector2f(-6.0f / 8.0f, +2.0f / 8.0f),
-	FVector2f(+2.0f / 8.0f, +6.0f / 8.0f),
-
-	// MSAA x8
-	FVector2f(+1.0f / 8.0f, -3.0f / 8.0f),
-	FVector2f(-1.0f / 8.0f, +3.0f / 8.0f),
-	FVector2f(+5.0f / 8.0f, +1.0f / 8.0f),
-	FVector2f(-3.0f / 8.0f, -5.0f / 8.0f),
-	FVector2f(-5.0f / 8.0f, +5.0f / 8.0f),
-	FVector2f(-7.0f / 8.0f, -1.0f / 8.0f),
-	FVector2f(+3.0f / 8.0f, +7.0f / 8.0f),
-	FVector2f(+7.0f / 8.0f, -7.0f / 8.0f),
-
-	// MSAA x16
-	FVector2f(+1.0f / 8.0f, +1.0f / 8.0f),
-	FVector2f(-1.0f / 8.0f, -3.0f / 8.0f),
-	FVector2f(-3.0f / 8.0f, +2.0f / 8.0f),
-	FVector2f(+4.0f / 8.0f, -1.0f / 8.0f),
-	FVector2f(-5.0f / 8.0f, -2.0f / 8.0f),
-	FVector2f(+2.0f / 8.0f, +5.0f / 8.0f),
-	FVector2f(+5.0f / 8.0f, +3.0f / 8.0f),
-	FVector2f(+3.0f / 8.0f, -5.0f / 8.0f),
-	FVector2f(-2.0f / 8.0f, +6.0f / 8.0f),
-	FVector2f(+0.0f / 8.0f, -7.0f / 8.0f),
-	FVector2f(-4.0f / 8.0f, -6.0f / 8.0f),
-	FVector2f(-6.0f / 8.0f, +4.0f / 8.0f),
-	FVector2f(-8.0f / 8.0f, +0.0f / 8.0f),
-	FVector2f(+7.0f / 8.0f, -4.0f / 8.0f),
-	FVector2f(+6.0f / 8.0f, +7.0f / 8.0f),
-	FVector2f(-7.0f / 8.0f, -8.0f / 8.0f),
-};
+
 
 int32 CalculateMSAASampleArrayIndex(int32 NumSamples, int32 SampleIndex)
 {
diff --git a/Engine/Source/Runtime/RHI/Private/RHIGlobals.cpp b/Engine/Source/Runtime/RHI/Private/RHIGlobals.cpp
new file mode 100644
--- /dev/null
+++ b/Engine/Source/Runtime/RHI/Private/RHIGlobals.cpp
@@ -0,0 +1,64 @@
+// Copyright Epic Games, Inc. All Rights Reserved.
+
+#include "RHIGlobals.h"
+
+FRHIGlobals GRHIGlobals;
+
+//
+//	MSAA sample offsets.
+//
+// https://docs.microsoft.com/en-us/windows/win32/api/d3d11/ne-d3d11-d3d11_standard_multisample_quality_levels
+static FVector2f GDefaultMSAASampleOffsets[1 + 2 + 4 + 8 + 16] = {
+	// MSAA x1
+	FVector2f(+0.0f / 8.0f, +0.0f / 8.0f),
+
+	// MSAA x2
+	FVector2f(+4.0f / 8.0f, +4.0f / 8.0f),
+	FVector2f(-4.0f / 8.0f, -4.0f / 8.0f),
+
+	// MSAA x4
+	FVector2f(-2.0f / 8.0f, -6.0f / 8.0f),
+	FVector2f(+6.0f / 8.0f, -2.0f / 8.0f),
+	FVector2f(-6.0f / 8.0f, +2.0f / 8.0f),
+	FVector2f(+2.0f / 8.0f, +6.0f / 8.0f),
+
+	// MSAA x8
+	FVector2f(+1.0f / 8.0f, -3.0f / 8.0f),
+	FVector2f(-1.0f / 8.0f, +3.0f / 8.0f),
+	FVector2f(+5.0f / 8.0f, +1.0f / 8.0f),
+	FVector2f(-3.0f / 8.0f, -5.0f / 8.0f),
+	FVector2f(-5.0f / 8.0f, +5.0f / 8.0f),
+	FVector2f(-7.0f / 8.0f, -1.0f / 8.0f),
+	FVector2f(+3.0f / 8.0f, +7.0f / 8.0f),
+	FVector2f(+7.0f / 8.0f, -7.0f / 8.0f),
+
+	// MSAA x16
+	FVector2f(+1.0f / 8.0f, +1.0f / 8.0f),
+	FVector2f(-1.0f / 8.0f, -3.0f / 8.0f),
+	FVector2f(-3.0f / 8.0f, +2.0f / 8.0f),
+	FVector2f(+4.0f / 8.0f, -1.0f / 8.0f),
+	FVector2f(-5.0f / 8.0f, -2.0f / 8.0f),
+	FVector2f(+2.0f / 8.0f, +5.0f / 8.0f),
+	FVector2f(+5.0f / 8.0f, +3.0f / 8.0f),
+	FVector2f(+3.0f / 8.0f, -5.0f / 8.0f),
+	FVector2f(-2.0f / 8.0f, +6.0f / 8.0f),
+	FVector2f(+0.0f / 8.0f, -7.0f / 8.0f),
+	FVector2f(-4.0f / 8.0f, -6.0f / 8.0f),
+	FVector2f(-6.0f / 8.0f, +4.0f / 8.0f),
+	FVector2f(-8.0f / 8.0f, +0.0f / 8.0f),
+	FVector2f(+7.0f / 8.0f, -4.0f / 8.0f),
+	FVector2f(+6.0f / 8.0f, +7.0f / 8.0f),
+	FVector2f(-7.0f / 8.0f, -8.0f / 8.0f),
+};
+
+FRHIGlobals::FRHIGlobals()
+{
+	FMemory::Memcpy(DefaultMSAASampleOffsets, GDefaultMSAASampleOffsets, sizeof(GDefaultMSAASampleOffsets));
+
+	for (EShaderPlatform& Platform : ShaderPlatformForFeatureLevel)
+	{
+		Platform = SP_NumPlatforms;
+	}
+}
+
+FRHIGlobals::~FRHIGlobals() = default;
diff --git a/Engine/Source/Runtime/RHI/Public/RHI.h b/Engine/Source/Runtime/RHI/Public/RHI.h
--- a/Engine/Source/Runtime/RHI/Public/RHI.h
+++ b/Engine/Source/Runtime/RHI/Public/RHI.h
@@ -22,9 +22,11 @@
 #include "RHIShaderPlatform.h"
 #include "RHIFeatureLevel.h"
 #include "RHIAccess.h"
+#include "RHIGlobals.h"
 
 class FResourceArrayInterface;
 class FResourceBulkDataInterface;
+class FRHICommandList;
 
 /** RHI Logging. */
 RHI_API DECLARE_LOG_CATEGORY_EXTERN(LogRHI,Log,VeryVerbose);
@@ -40,62 +42,10 @@ namespace RHIConfig
 	RHI_API float GetGPUHitchThreshold();
 }
 
-/**
- * RHI globals.
- */
-
-/** True if the render hardware has been initialized. */
-extern RHI_API bool GIsRHIInitialized;
-
-class RHI_API FRHICommandList;
-
 /**
  * RHI capabilities.
  */
 
-/** Optimal number of persistent thread groups to fill the GPU. */
-extern RHI_API int32 GRHIPersistentThreadGroupCount;
-
- /** The maximum number of mip-maps that a texture can contain. */
-extern RHI_API int32 GMaxTextureMipCount;
-
-/** Does the RHI implements CopyToTexture() with FRHICopyTextureInfo::NumMips > 1 */
-UE_DEPRECATED(5.1, "All RHIs now support copying to multiple mips.")
-extern RHI_API bool GRHISupportsCopyToTextureMultipleMips;
-
-/** true if this platform has quad buffer stereo support. */
-extern RHI_API bool GSupportsQuadBufferStereo;
-
-/** true if the RHI supports textures that may be bound as both a render target and a shader resource. */
-extern RHI_API bool GSupportsRenderDepthTargetableShaderResources;
-
-/** true if the RHI supports Draw Indirect */
-extern RHI_API bool GRHISupportsDrawIndirect;
-
-/** Whether the RHI can send commands to the device context from multiple threads. Used in the GPU readback to avoid stalling the RHI threads. */
-extern RHI_API bool GRHISupportsMultithreading;
-
-/** Whether RHIGetRenderQueryResult can be safely called off the render thread. */
-extern RHI_API bool GRHISupportsAsyncGetRenderQueryResult;
-
-/** 
- * only set if RHI has the information (after init of the RHI and only if RHI has that information, never changes after that)
- * e.g. "NVIDIA GeForce GTX 670"
- */
-extern RHI_API FString GRHIAdapterName;
-extern RHI_API FString GRHIAdapterInternalDriverVersion;
-extern RHI_API FString GRHIAdapterUserDriverVersion;
-extern RHI_API FString GRHIAdapterDriverDate;
-extern RHI_API bool GRHIAdapterDriverOnDenyList;
-extern RHI_API uint32 GRHIDeviceId;
-extern RHI_API uint32 GRHIDeviceRevision;
-
-// 0 means not defined yet, use functions like IsRHIDeviceAMD() to access
-extern RHI_API uint32 GRHIVendorId;
-
-// true if the RHI supports Pixel Shader UAV
-extern RHI_API bool GRHISupportsPixelShaderUAVs;
-
 // to trigger GPU specific optimizations and fallbacks
 RHI_API bool IsRHIDeviceAMD();
 
@@ -158,520 +108,8 @@ struct FRHIResourceStats
 
 RHI_API void RHIGetTrackedResourceStats(TArray<TSharedPtr<FRHIResourceStats>>& OutResourceStats);
 
-// Wrapper for GRHI## global variables, allows values to be overridden for mobile preview modes.
-template <typename TValueType>
-class TRHIGlobal
-{
-public:
-	explicit TRHIGlobal(const TValueType& InValue) : Value(InValue) {}
-
-	TRHIGlobal& operator=(const TValueType& InValue) 
-	{
-		Value = InValue; 
-		return *this;
-	}
-
-#if WITH_EDITOR
-	inline void SetPreviewOverride(const TValueType& InValue)
-	{
-		PreviewValue = InValue;
-	}
-
-	inline operator TValueType() const
-	{ 
-		return PreviewValue.IsSet() ? GetPreviewValue() : Value;
-	}
-#else
-	inline operator TValueType() const { return Value; }
-#endif
-
-private:
-	TValueType Value;
-#if WITH_EDITOR
-	TOptional<TValueType> PreviewValue;
-	TValueType GetPreviewValue() const { return PreviewValue.GetValue(); }
-#endif
-};
-
-#if WITH_EDITOR
-template<>
-inline int32 TRHIGlobal<int32>::GetPreviewValue() const 
-{
-	// ensure the preview values are subsets of RHI functionality.
-	return FMath::Min(PreviewValue.GetValue(), Value);
-}
-template<>
-inline bool TRHIGlobal<bool>::GetPreviewValue() const
-{
-	// ensure the preview values are subsets of RHI functionality.
-	return PreviewValue.GetValue() && Value;
-}
-#endif
-
-/** true if the GPU is AMD's Pre-GCN architecture */
-extern RHI_API bool GRHIDeviceIsAMDPreGCNArchitecture;
-
-/** true if PF_G8 render targets are supported */
-extern RHI_API TRHIGlobal<bool> GSupportsRenderTargetFormat_PF_G8;
-
-/** true if PF_FloatRGBA render targets are supported */
-extern RHI_API TRHIGlobal<bool> GSupportsRenderTargetFormat_PF_FloatRGBA;
-
-/** true if mobile framebuffer fetch is supported */
-extern RHI_API bool GSupportsShaderFramebufferFetch;
-
-/** true if mobile framebuffer fetch is supported from MRT's*/
-extern RHI_API bool GSupportsShaderMRTFramebufferFetch;
-
-/** true if mobile pixel local storage is supported */
-extern RHI_API bool GSupportsPixelLocalStorage;
-
-/** true if mobile depth & stencil fetch is supported */
-extern RHI_API bool GSupportsShaderDepthStencilFetch;
-
-/** true if RQT_AbsoluteTime is supported by RHICreateRenderQuery */
-extern RHI_API bool GSupportsTimestampRenderQueries;
-
-/** true if RQT_AbsoluteTime is supported by RHICreateRenderQuery */
-extern RHI_API bool GRHISupportsGPUTimestampBubblesRemoval;
-
-/** true if RHIGetGPUFrameCycles removes CPu generated bubbles. */
-extern RHI_API bool GRHISupportsFrameCyclesBubblesRemoval;
-
-/** true if RHIGetGPUUsage() is supported. */
-extern RHI_API bool GRHISupportsGPUUsage;
-
-/** true if the GPU supports hidden surface removal in hardware. */
-extern RHI_API bool GHardwareHiddenSurfaceRemoval;
-
-/** true if the RHI supports asynchronous creation of texture resources */
-extern RHI_API bool GRHISupportsAsyncTextureCreation;
-
-/** true if the RHI supports quad topology (PT_QuadList). */
-extern RHI_API bool GRHISupportsQuadTopology;
-
-/** true if the RHI supports rectangular topology (PT_RectList). */
-extern RHI_API bool GRHISupportsRectTopology;
-
-/** true if the RHI supports primitive shaders. */
-extern RHI_API bool GRHISupportsPrimitiveShaders;
-
-/** true if the RHI supports 64 bit uint atomics. */
-extern RHI_API bool GRHISupportsAtomicUInt64;
-
-/** true if the RHI supports 64 bit uint atomics using DX12 SM6.6 - TEMP / DEPRECATED - DO NOT USE */
-extern RHI_API bool GRHISupportsDX12AtomicUInt64;
-
-/** true if the RHI supports optimal low level pipeline state sort keys. */
-extern RHI_API bool GRHISupportsPipelineStateSortKey;
-
-/** Temporary. When OpenGL is running in a separate thread, it cannot yet do things like initialize shaders that are first discovered in a rendering task. It is doable, it just isn't done. */
-extern RHI_API bool GSupportsParallelRenderingTasksWithSeparateRHIThread;
-
-/** If an RHI is so slow, that it is the limiting factor for the entire frame, we can kick early to try to give it as much as possible. */
-extern RHI_API bool GRHIThreadNeedsKicking;
-
-/** If an RHI cannot do an unlimited number of occlusion queries without stalling and waiting for the GPU, this can be used to tune hte occlusion culler to try not to do that. */
-extern RHI_API int32 GRHIMaximumReccommendedOustandingOcclusionQueries;
-
-/** Some RHIs can only do visible or not occlusion queries. */
-extern RHI_API bool GRHISupportsExactOcclusionQueries;
-
-/** True if and only if the GPU support rendering to volume textures (2D Array, 3D). Some OpenGL 3.3 cards support SM4, but can't render to volume textures. */
-extern RHI_API bool GSupportsVolumeTextureRendering;
-
-/** True if the RHI supports separate blend states per render target. */
-extern RHI_API bool GSupportsSeparateRenderTargetBlendState;
-
-/** True if the RHI has artifacts with atlased CSM depths. */
-extern RHI_API bool GRHINeedsUnatlasedCSMDepthsWorkaround;
-
-/** true if the RHI supports 3D textures */
-extern RHI_API bool GSupportsTexture3D;
-
-/** true if the RHI supports mobile multi-view */
-extern RHI_API bool GSupportsMobileMultiView;
-
-/** true if the RHI supports image external */
-extern RHI_API bool GSupportsImageExternal;
-
-/** true if the RHI supports 256bit MRT */
-extern RHI_API bool GSupportsWideMRT;
-
-/** True if the RHI and current hardware supports supports depth bounds testing */
-extern RHI_API bool GSupportsDepthBoundsTest;
-
-/** True if the RHI supports explicit access to depth target HTile meta data. */
-extern RHI_API bool GRHISupportsExplicitHTile;
-
-/** True if the RHI supports explicit access to MSAA target FMask meta data. */
-extern RHI_API bool GRHISupportsExplicitFMask;
-
-/** True if the RHI supports resummarizing depth target HTile meta data. */
-extern RHI_API bool GRHISupportsResummarizeHTile;
-
-/** True if the RHI supports depth target unordered access views. */
-extern RHI_API bool GRHISupportsDepthUAV;
-
-/** True if the RHI and current hardware supports efficient AsyncCompute (by default we assume false and later we can enable this for more hardware) */
-extern RHI_API bool GSupportsEfficientAsyncCompute;
-
-/** True if the RHI supports getting the result of occlusion queries when on a thread other than the render thread */
-extern RHI_API bool GSupportsParallelOcclusionQueries;
-
-/** true if the RHI requires a valid RT bound during UAV scatter operation inside the pixel shader */
-extern RHI_API bool GRHIRequiresRenderTargetForPixelShaderUAVs;
-
-/** true if the RHI supports unordered access view format aliasing */
-extern RHI_API bool GRHISupportsUAVFormatAliasing;
-
-/** true if the RHI supports texture views (data aliasing) */
-extern RHI_API bool GRHISupportsTextureViews;
-
-/** true if the pointer returned by Lock is a persistent direct pointer to gpu memory */
-extern RHI_API bool GRHISupportsDirectGPUMemoryLock;
-
-/** true if the multi-threaded shader creation is supported by (or desirable for) the RHI. */
-extern RHI_API bool GRHISupportsMultithreadedShaderCreation;
-
-/** Does the RHI support parallel resource commands (i.e. create / lock / unlock) on non-immediate command list APIs, recorded off the render thread. */
-extern RHI_API bool GRHISupportsMultithreadedResources;
-
-/** The minimum Z value in clip space for the RHI. This is a constant value to always match D3D clip-space. */
-inline constexpr float GMinClipZ = 0.0f;
-
-/** The sign to apply to the Y axis of projection matrices. This is a constant value to always match D3D clip-space. */
-inline constexpr float GProjectionSignY = 1.0f;
-
-/** Does this RHI need to wait for deletion of resources due to ref counting. */
-extern RHI_API bool GRHINeedsExtraDeletionLatency;
-
-/** Allow opt-out default RHI resource deletion latency for streaming textures */
-extern RHI_API bool GRHIForceNoDeletionLatencyForStreamingTextures;
-
-/** The maximum size allowed for a computeshader dispatch. */
-extern RHI_API TRHIGlobal<int32> GMaxComputeDispatchDimension;
-
-/** If true, then avoiding loading shader code and instead force the "native" path, which sends a library and a hash instead. */
-extern RHI_API bool GRHILazyShaderCodeLoading;
-
-/** If true, then it is possible to turn on GRHILazyShaderCodeLoading. */
-extern RHI_API bool GRHISupportsLazyShaderCodeLoading;
-
-/** true if the RHI supports UpdateFromBufferTexture method */
-extern RHI_API bool GRHISupportsUpdateFromBufferTexture;
-
-/** The maximum size to allow for the shadow depth buffer in the X dimension.  This must be larger or equal to GMaxShadowDepthBufferSizeY. */
-extern RHI_API TRHIGlobal<int32> GMaxShadowDepthBufferSizeX;
-/** The maximum size to allow for the shadow depth buffer in the Y dimension. */
-extern RHI_API TRHIGlobal<int32> GMaxShadowDepthBufferSizeY;
-
-/** The maximum size allowed for 2D textures in both dimensions. */
-extern RHI_API TRHIGlobal<int32> GMaxTextureDimensions;
-
-/** The maximum size allowed for 2D textures in both dimensions. */
-extern RHI_API TRHIGlobal<int64> GMaxBufferDimensions;
-
-/** The maximum size allowed for a contant buffer. */
-extern RHI_API TRHIGlobal<int64> GRHIMaxConstantBufferByteSize;
-
-/** The maximum size allowed for Shared Compute Memory. */
-extern RHI_API TRHIGlobal<int64> GMaxComputeSharedMemory;
-
-/** The maximum size allowed for 3D textures in all three dimensions. */
-extern RHI_API TRHIGlobal<int32> GMaxVolumeTextureDimensions;
-
-/** Whether RW texture buffers are supported */
-extern RHI_API bool GRHISupportsRWTextureBuffers;
-
-/** Whether a raw (ByteAddress) buffer view can be created for any buffer, regardless of its EBufferUsageFlags::ByteAddressBuffer flag. */
-extern RHI_API bool GRHISupportsRawViewsForAnyBuffer;
-
-/** Whether depth or stencil can individually be set to CopySrc/Dest access. */
-extern RHI_API bool GRHISupportsSeparateDepthStencilCopyAccess;
-
-/** Support using async thread for texture stream out operations */
-extern RHI_API bool GRHISupportAsyncTextureStreamOut;
-
-FORCEINLINE uint64 GetMaxBufferDimension()
-{
-	return GMaxBufferDimensions;
-}
-
-FORCEINLINE uint64 GetMaxConstantBufferByteSize()
-{
-	return GRHIMaxConstantBufferByteSize;
-}
-
-FORCEINLINE uint64 GetMaxComputeSharedMemory()
-{
-	return GMaxComputeSharedMemory;
-}
-
-FORCEINLINE uint32 GetMax2DTextureDimension()
-{
-	return GMaxTextureDimensions;
-}
-
-/** The maximum size allowed for cube textures. */
-extern RHI_API TRHIGlobal<int32> GMaxCubeTextureDimensions;
-FORCEINLINE uint32 GetMaxCubeTextureDimension()
-{
-	return GMaxCubeTextureDimensions;
-}
-
-/** The Maximum number of layers in a 1D or 2D texture array. */
-extern RHI_API int32 GMaxTextureArrayLayers;
-FORCEINLINE uint32 GetMaxTextureArrayLayers()
-{
-	return GMaxTextureArrayLayers;
-}
-
-extern RHI_API int32 GMaxTextureSamplers;
-FORCEINLINE uint32 GetMaxTextureSamplers()
-{
-	return GMaxTextureSamplers;
-}
-
-/** The maximum work group invocations allowed for compute shader. */
-extern RHI_API TRHIGlobal<int32> GMaxWorkGroupInvocations;
-FORCEINLINE uint32 GetMaxWorkGroupInvocations()
-{
-	return GMaxWorkGroupInvocations;
-}
-
-/** true if we are running with the NULL RHI */
-extern RHI_API bool GUsingNullRHI;
-
-/**
- *	The size to check against for Draw*UP call vertex counts.
- *	If greater than this value, the draw call will not occur.
- */
-extern RHI_API int32 GDrawUPVertexCheckCount;
-/**
- *	The size to check against for Draw*UP call index counts.
- *	If greater than this value, the draw call will not occur.
- */
-extern RHI_API int32 GDrawUPIndexCheckCount;
-
 #include "MultiGPU.h"
 
-/** Whether the next frame should profile the GPU. */
-extern RHI_API bool GTriggerGPUProfile;
-
-/** Whether we are profiling GPU hitches. */
-extern RHI_API bool GTriggerGPUHitchProfile;
-
-/** Non-empty if we are performing a gpu trace. Also says where to place trace file. */
-extern RHI_API FString GGPUTraceFileName;
-
-/** True if the RHI supports texture streaming */
-extern RHI_API bool GRHISupportsTextureStreaming;
-/** Amount of memory allocated by textures. In kilobytes. */
-extern RHI_API volatile int32 GCurrentTextureMemorySize;
-/** Amount of memory allocated by rendertargets. In kilobytes. */
-extern RHI_API volatile int32 GCurrentRendertargetMemorySize;
-/** Current texture streaming pool size, in bytes. 0 means unlimited. */
-extern RHI_API int64 GTexturePoolSize;
-
-/** In percent. If non-zero, the texture pool size is a percentage of GTotalGraphicsMemory. */
-extern RHI_API int32 GPoolSizeVRAMPercentage;
-
-/** Amount of local video memory demoted to system memory. In bytes. */
-extern RHI_API uint64 GDemotedLocalMemorySize;
-
-/** Whether or not the RHI can handle a non-zero BaseVertexIndex - extra SetStreamSource calls will be needed if this is false */
-extern RHI_API bool GRHISupportsBaseVertexIndex;
-
-/** True if the RHI supports copying cubemap faces using CopyToResolveTarget */
-UE_DEPRECATED(5.1, "CopyToResoveTarget is deprecated.")
-extern RHI_API bool GRHISupportsResolveCubemapFaces;
-
-/** Whether or not the RHI can handle a non-zero FirstInstance to DrawIndexedPrimitive and friends - extra SetStreamSource calls will be needed if this is false */
-extern RHI_API bool GRHISupportsFirstInstance;
-
-/** Whether or not the RHI can handle dynamic resolution or not. */
-extern RHI_API bool GRHISupportsDynamicResolution;
-
-/**
-* Whether or not the RHI supports ray tracing on current hardware (acceleration structure building and new ray tracing-specific shader types). 
-* GRHISupportsRayTracingShaders and GRHISupportsInlineRayTracing must also be checked before dispatching ray tracing workloads.
-*/
-extern RHI_API bool GRHISupportsRayTracing;
-
-/**
-* Whether or not the RHI supports ray tracing raygen, miss and hit shaders (i.e. full ray tracing pipeline). 
-* The RHI may support inline ray tracing from compute shaders, but not the full pipeline.
-*/
-extern RHI_API bool GRHISupportsRayTracingShaders;
-
-/** Whether or not the RHI supports adding new shaders to an existing RT PSO. */
-extern RHI_API bool GRHISupportsRayTracingPSOAdditions;
-
-/** Whether or not the RHI supports indirect ray tracing dispatch commands. */
-extern RHI_API bool GRHISupportsRayTracingDispatchIndirect;
-
-/** Whether or not the RHI supports async building ray tracing acceleration structures. */
-extern RHI_API bool GRHISupportsRayTracingAsyncBuildAccelerationStructure;
-
-/** Whether or not the RHI supports the AMD Hit Token extension. */
-extern RHI_API bool GRHISupportsRayTracingAMDHitToken;
-
-/** Whether or not the RHI supports inline ray tracing in compute shaders, without a full ray tracing pipeline. */
-extern RHI_API bool GRHISupportsInlineRayTracing;
-
-/** Required alignment for ray tracing acceleration structures. */
-extern RHI_API uint32 GRHIRayTracingAccelerationStructureAlignment;
-
-/** Required alignment for ray tracing scratch buffers. */
-extern RHI_API uint32 GRHIRayTracingScratchBufferAlignment;
-
-/** Required alignment for ray tracing shader binding table buffer. */
-extern RHI_API uint32 GRHIRayTracingShaderTableAlignment;
-
-/** Size of an individual element in the ray tracing instance buffer. This defines the required stride and alignment of structured buffers of instances. */
-extern RHI_API uint32 GRHIRayTracingInstanceDescriptorSize;
-
-/** Whether or not the RHI supports shader wave operations (shader model 6.0). */
-extern RHI_API bool GRHISupportsWaveOperations;
-
-/** Whether or not the current GPU is integrated into to CPU */
-extern RHI_API bool GRHIDeviceIsIntegrated;
-
-/** 
-* Specifies the minimum and maximum number of lanes in the SIMD wave that this GPU can support. I.e. 32 on NVIDIA, 64 on AMD.
-* Valid values are in range [4..128] (as per SM 6.0 specification) or 0 if unknown.
-* Rendering code must always check GRHISupportsWaveOperations in addition to wave min/max size.
-*/
-extern RHI_API int32 GRHIMinimumWaveSize;
-extern RHI_API int32 GRHIMaximumWaveSize;
-
-/** Whether or not the RHI supports an RHI thread.
-Requirements for RHI thread
-* Microresources (those in RHIStaticStates.h) need to be able to be created by any thread at any time and be able to work with a radically simplified rhi resource lifecycle. CreateSamplerState, CreateRasterizerState, CreateDepthStencilState, CreateBlendState
-* CreateUniformBuffer needs to be threadsafe
-* GetRenderQueryResult should be threadsafe, but this isn't required. If it isn't threadsafe, then you need to flush yourself in the RHI
-* GetViewportBackBuffer and AdvanceFrameForGetViewportBackBuffer need to be threadsafe and need to support the fact that the render thread has a different concept of "current backbuffer" than the RHI thread. Without an RHIThread this is moot due to the next two items.
-* AdvanceFrameForGetViewportBackBuffer needs be added as an RHI method and this needs to work with GetViewportBackBuffer to give the render thread the right back buffer even though many commands relating to the beginning and end of the frame are queued.
-* BeginDrawingViewport, and 5 or so other frame advance methods are queued with an RHIThread. Without an RHIThread, these just flush internally.
-***/
-extern RHI_API bool GRHISupportsRHIThread;
-/* as above, but we run the commands on arbitrary task threads */
-extern RHI_API bool GRHISupportsRHIOnTaskThread;
-
-/** Whether or not the RHI supports parallel RHIThread executes / translates
-Requirements:
-* RHICreateBoundShaderState & RHICreateGraphicsPipelineState is threadsafe and GetCachedBoundShaderState must not be used. GetCachedBoundShaderState_Threadsafe has a slightly different protocol.
-***/
-extern RHI_API bool GRHISupportsParallelRHIExecute;
-
-/** Whether or not the RHI can perform MSAA sample load. */
-extern RHI_API bool GRHISupportsMSAADepthSampleAccess;
-
-/** Whether or not the RHI can render to the backbuffer with a custom depth/stencil surface bound. */
-extern RHI_API bool GRHISupportsBackBufferWithCustomDepthStencil;
-
-/** Whether or not HDR is currently enabled */
-extern RHI_API bool GRHIIsHDREnabled;
-
-/** Whether the present adapter/display offers HDR output capabilities. */
-extern RHI_API bool GRHISupportsHDROutput;
-
-/** Whether VRS (in all flavors) is currently enabled (separate from whether it's supported/available). */
-extern RHI_API bool GRHIVariableRateShadingEnabled;
-
-/** Whether attachment (image-based) VRS is currently enabled (separate from whether it's supported/available). */
-extern RHI_API bool GRHIAttachmentVariableRateShadingEnabled;
-
-/** The maximum number of groups that can be dispatched in each dimensions. */
-extern RHI_API FIntVector GRHIMaxDispatchThreadGroupsPerDimension;
-
-/** Whether or not the RHI can support per-draw Variable Rate Shading. */
-extern RHI_API bool GRHISupportsPipelineVariableRateShading;
-
-/** Whether or not the Variable Rate Shading can be done at larger (2x4 or 4x2 or 4x4) sizes. */
-extern RHI_API bool GRHISupportsLargerVariableRateShadingSizes;
-
-/** Whether or not the RHI can support image-based Variable Rate Shading. */
-extern RHI_API bool GRHISupportsAttachmentVariableRateShading;
-
-/** Whether or not the RHI can support complex combiner operatations between per-draw (pipeline) VRS and image VRS. */
-extern RHI_API bool GRHISupportsComplexVariableRateShadingCombinerOps;
-
-/** Whether or not the RHI can support shading rate attachments as array textures. */
-extern RHI_API bool GRHISupportsVariableRateShadingAttachmentArrayTextures;
-
-/** Maximum tile width in a screen space texture that can be used to drive Variable Rate Shading. */
-extern RHI_API int32 GRHIVariableRateShadingImageTileMaxWidth;
-
-/** Maximum tile height in a screen space texture that can be used to drive Variable Rate Shading. */
-extern RHI_API int32 GRHIVariableRateShadingImageTileMaxHeight;
-
-/** Minimum tile width in a screen space texture that can be used to drive Variable Rate Shading. */
-extern RHI_API int32 GRHIVariableRateShadingImageTileMinWidth;
-
-/** Minimum tile height in a screen space texture that can be used to drive Variable Rate Shading. */
-extern RHI_API int32 GRHIVariableRateShadingImageTileMinHeight;
-
-/** Data type contained in a shading-rate image for image-based Variable Rate Shading. */
-extern RHI_API EVRSImageDataType GRHIVariableRateShadingImageDataType;
-
-/** Image format for the shading rate image for image-based Variable Rate Shading. */
-extern RHI_API EPixelFormat GRHIVariableRateShadingImageFormat;
-
-/** Whether Variable Rate Shading deferred shading rate texture update is supported. */
-extern RHI_API bool GRHISupportsLateVariableRateShadingUpdate;
-
-/** Format used for the backbuffer when outputting to a HDR display. */
-extern RHI_API EPixelFormat GRHIHDRDisplayOutputFormat;
-
-/** Counter incremented once on each frame present. Used to support game thread synchronization with swap chain frame flips. */
-extern RHI_API uint64 GRHIPresentCounter;
-
-/** True if the RHI supports setting the render target array index from any shader stage */
-extern RHI_API bool GRHISupportsArrayIndexFromAnyShader;
-
-/** True if the pipeline file cache can be used with this RHI */
-extern RHI_API bool GRHISupportsPipelineFileCache;
-
-/** True if the RHI supports setting the stencil ref at pixel granularity from the pixel shader */
-extern RHI_API bool GRHISupportsStencilRefFromPixelShader;
-
-/** Whether current RHI supports overestimated conservative rasterization. */
-extern RHI_API bool GRHISupportsConservativeRasterization;
-
-/** true if the RHI supports Mesh and Amplification shaders with tier0 capability */
-extern RHI_API bool GRHISupportsMeshShadersTier0;
-
-/** true if the RHI supports Mesh and Amplification shaders with tier1 capability */
-extern RHI_API bool GRHISupportsMeshShadersTier1;
-
-/**
-* True if the RHI supports reading system timer in shaders via GetShaderTimestamp().
-* Individual shaders must be compiled with an appropriate vendor extension and check PLATFORM_SUPPORTS_SHADER_TIMESTAMP.
-*/
-extern RHI_API bool GRHISupportsShaderTimestamp;
-
-extern RHI_API bool GRHISupportsEfficientUploadOnResourceCreation;
-
-/** true if the RHI supports RLM_WriteOnly_NoOverwrite */
-extern RHI_API bool GRHISupportsMapWriteNoOverwrite;
-
-/** Tables of all MSAA sample offset for all MSAA supported. Use GetMSAASampleOffsets() to read it. */
-extern RHI_API FVector2f GRHIDefaultMSAASampleOffsets[1 + 2 + 4 + 8 + 16];
-
-/** True if the RHI supports pipeline precompiling from any thread. */
-extern RHI_API bool GRHISupportsAsyncPipelinePrecompile;
-
-/** Whether dynamic (bindless) resources are supported */
-extern RHI_API ERHIBindlessSupport GRHIBindlessSupport;
-
-UE_DEPRECATED(5.2, "You must use GRHIBindlessSupport instead.")
-extern RHI_API bool GRHISupportsBindless;
-
-/** True if the RHI supports reserved (AKA tiled, virtual or sparse) resources and operations related to them*/
-extern RHI_API bool GRHISupportsReservedResources;
-
 // Calculate the index of the sample in GRHIDefaultMSAASampleOffsets
 extern RHI_API int32 CalculateMSAASampleArrayIndex(int32 NumSamples, int32 SampleIndex);
 
@@ -722,18 +160,6 @@ inline FMatrix AdjustProjectionMatrixForRHI(const FMatrix& InProjectionMatrix)
 /** Set runtime selection of mobile feature level preview. */
 RHI_API void RHISetMobilePreviewFeatureLevel(ERHIFeatureLevel::Type MobilePreviewFeatureLevel);
 
-/** Current shader platform. */
-
-
-/** Table for finding out which shader platform corresponds to a given feature level for this RHI. */
-extern RHI_API EShaderPlatform GShaderPlatformForFeatureLevel[ERHIFeatureLevel::Num];
-
-/** Get the shader platform associated with the supplied feature level on this machine */
-inline EShaderPlatform GetFeatureLevelShaderPlatform(const FStaticFeatureLevel InFeatureLevel)
-{
-	return GShaderPlatformForFeatureLevel[InFeatureLevel];
-}
-
 struct FVertexElement
 {
 	uint8 StreamIndex;
diff --git a/Engine/Source/Runtime/RHI/Public/RHICommandList.h b/Engine/Source/Runtime/RHI/Public/RHICommandList.h
--- a/Engine/Source/Runtime/RHI/Public/RHICommandList.h
+++ b/Engine/Source/Runtime/RHI/Public/RHICommandList.h
@@ -84,7 +84,6 @@ extern RHI_API bool GUseRHITaskThreads_InternalUseOnly;
 extern RHI_API bool GIsRunningRHIInSeparateThread_InternalUseOnly;
 extern RHI_API bool GIsRunningRHIInDedicatedThread_InternalUseOnly;
 extern RHI_API bool GIsRunningRHIInTaskThread_InternalUseOnly;
-extern RHI_API bool GSupportsParallelRenderingTasksWithSeparateRHIThread;
 
 namespace ERenderThreadIdleTypes
 {
diff --git a/Engine/Source/Runtime/RHI/Public/RHIGlobals.h b/Engine/Source/Runtime/RHI/Public/RHIGlobals.h
new file mode 100644
--- /dev/null
+++ b/Engine/Source/Runtime/RHI/Public/RHIGlobals.h
@@ -0,0 +1,776 @@
+// Copyright Epic Games, Inc. All Rights Reserved.
+
+#pragma once
+
+#include "Containers/UnrealString.h"
+#include "Math/IntVector.h"
+#include "Math/Vector2D.h"
+#include "PixelFormat.h"
+#include "RHIDefinitions.h"
+#include "RHIFeatureLevel.h"
+#include "RHIShaderPlatform.h"
+
+#if WITH_EDITOR
+#include "Misc/Optional.h"
+#endif
+
+/** The minimum Z value in clip space for the RHI. This is a constant value to always match D3D clip-space. */
+inline constexpr float GMinClipZ = 0.0f;
+
+/** The sign to apply to the Y axis of projection matrices. This is a constant value to always match D3D clip-space. */
+inline constexpr float GProjectionSignY = 1.0f;
+
+/** Info for supporting the vertex element types */
+class FVertexElementTypeSupportInfo
+{
+public:
+	FVertexElementTypeSupportInfo() { for (int32 i = 0; i < VET_MAX; i++) ElementCaps[i] = true; }
+	bool IsSupported(EVertexElementType ElementType) { return ElementCaps[ElementType]; }
+	void SetSupported(EVertexElementType ElementType, bool bIsSupported) { ElementCaps[ElementType] = bIsSupported; }
+private:
+	/** cap bit set for each VET. One-to-one mapping based on EVertexElementType */
+	bool ElementCaps[VET_MAX];
+};
+
+// Wrapper for GRHI## global variables, allows values to be overridden for mobile preview modes.
+template <typename TValueType>
+class TRHIGlobal
+{
+public:
+	TRHIGlobal() = default;
+	TRHIGlobal(const TValueType& InValue) : Value(InValue) {}
+
+	TRHIGlobal& operator=(const TValueType& InValue)
+	{
+		Value = InValue;
+		return *this;
+	}
+
+#if WITH_EDITOR
+	inline void SetPreviewOverride(const TValueType& InValue)
+	{
+		PreviewValue = InValue;
+	}
+
+	inline operator TValueType() const
+	{
+		return PreviewValue.IsSet() ? GetPreviewValue() : Value;
+	}
+#else
+	inline operator TValueType() const { return Value; }
+#endif
+
+private:
+	TValueType Value{};
+#if WITH_EDITOR
+	TOptional<TValueType> PreviewValue;
+	TValueType GetPreviewValue() const { return PreviewValue.GetValue(); }
+#endif
+};
+
+#if WITH_EDITOR
+template<>
+inline int32 TRHIGlobal<int32>::GetPreviewValue() const
+{
+	// ensure the preview values are subsets of RHI functionality.
+	return FMath::Min(PreviewValue.GetValue(), Value);
+}
+
+template<>
+inline bool TRHIGlobal<bool>::GetPreviewValue() const
+{
+	// ensure the preview values are subsets of RHI functionality.
+	return PreviewValue.GetValue() && Value;
+}
+#endif
+
+struct FRHIGlobals
+{
+	RHI_API FRHIGlobals();
+	RHI_API ~FRHIGlobals();
+
+	/** True if the render hardware has been initialized. */
+	bool IsRHIInitialized = false;
+
+	/** Optimal number of persistent thread groups to fill the GPU. */
+	int32 PersistentThreadGroupCount = 0;
+
+	/** The maximum number of mip-maps that a texture can contain. */
+	int32 MaxTextureMipCount = MAX_TEXTURE_MIP_COUNT;
+
+	/** true if this platform has quad buffer stereo support. */
+	bool SupportsQuadBufferStereo = false;
+
+	/** true if the RHI supports textures that may be bound as both a render target and a shader resource. */
+	bool SupportsRenderDepthTargetableShaderResources = true;
+
+	/** true if the RHI supports Draw Indirect */
+	bool SupportsDrawIndirect = true;
+
+	/** Whether the RHI can send commands to the device context from multiple threads. Used in the GPU readback to avoid stalling the RHI threads. */
+	bool SupportsMultithreading = false;
+
+	/** Whether RHIGetRenderQueryResult can be safely called off the render thread. */
+	bool SupportsAsyncGetRenderQueryResult = false;
+
+	struct FGpuInfo
+	{
+		/**
+		 * only set if RHI has the information (after init of the RHI and only if RHI has that information, never changes after that)
+		 * e.g. "NVIDIA GeForce GTX 670"
+		 */
+		FString AdapterName;
+		FString AdapterInternalDriverVersion;
+		FString AdapterUserDriverVersion;
+		FString AdapterDriverDate;
+		bool AdapterDriverOnDenyList = false;
+		uint32 DeviceId = 0;
+		uint32 DeviceRevision = 0;
+
+		/** true if the GPU is AMD's Pre-GCN architecture */
+		bool IsAMDPreGCNArchitecture = false;
+
+		// 0 means not defined yet, use functions like IsRHIDeviceAMD() to access
+		uint32 VendorId = 0;
+	} GpuInfo;
+
+	// true if the RHI supports Pixel Shader UAV
+	bool SupportsPixelShaderUAVs = true;
+
+	/** true if PF_G8 render targets are supported */
+	TRHIGlobal<bool> SupportsRenderTargetFormat_PF_G8 = true;
+
+	/** true if PF_FloatRGBA render targets are supported */
+	TRHIGlobal<bool> SupportsRenderTargetFormat_PF_FloatRGBA = true;
+
+	/** true if mobile framebuffer fetch is supported */
+	bool SupportsShaderFramebufferFetch = false;
+
+	/** true if mobile framebuffer fetch is supported from MRT's*/
+	bool SupportsShaderMRTFramebufferFetch = false;
+
+	/** true if mobile pixel local storage is supported */
+	bool SupportsPixelLocalStorage = false;
+
+	/** true if mobile depth & stencil fetch is supported */
+	bool SupportsShaderDepthStencilFetch = false;
+
+	/** true if RQT_AbsoluteTime is supported by RHICreateRenderQuery */
+	bool SupportsTimestampRenderQueries = false;
+
+	/** true if RQT_AbsoluteTime is supported by RHICreateRenderQuery */
+	bool SupportsGPUTimestampBubblesRemoval = false;
+
+	/** true if RHIGetGPUFrameCycles removes CPu generated bubbles. */
+	bool SupportsFrameCyclesBubblesRemoval = false;
+
+	/** true if RHIGetGPUUsage() is supported. */
+	bool SupportsGPUUsage = false;
+
+	/** true if the GPU supports hidden surface removal in hardware. */
+	bool HardwareHiddenSurfaceRemoval = false;
+
+	/** true if the RHI supports asynchronous creation of texture resources */
+	bool SupportsAsyncTextureCreation = false;
+
+	/** true if the RHI supports quad topology (PT_QuadList). */
+	bool SupportsQuadTopology = false;
+
+	/** true if the RHI supports rectangular topology (PT_RectList). */
+	bool SupportsRectTopology = false;
+
+	/** true if the RHI supports primitive shaders. */
+	bool SupportsPrimitiveShaders = false;
+
+	/** true if the RHI supports 64 bit uint atomics. */
+	bool SupportsAtomicUInt64 = false;
+
+	/** true if the RHI supports 64 bit uint atomics using DX12 SM6.6 - TEMP / DEPRECATED - DO NOT USE */
+	bool SupportsDX12AtomicUInt64 = false;
+
+	/** true if the RHI supports optimal low level pipeline state sort keys. */
+	bool SupportsPipelineStateSortKey = false;
+
+	/** Temporary. When OpenGL is running in a separate thread, it cannot yet do things like initialize shaders that are first discovered in a rendering task. It is doable, it just isn't done. */
+	bool SupportsParallelRenderingTasksWithSeparateRHIThread = true;
+
+	/** If an RHI is so slow, that it is the limiting factor for the entire frame, we can kick early to try to give it as much as possible. */
+	bool RHIThreadNeedsKicking = false;
+
+	/** If an RHI cannot do an unlimited number of occlusion queries without stalling and waiting for the GPU, this can be used to tune hte occlusion culler to try not to do that. */
+	int32 MaximumReccommendedOustandingOcclusionQueries = MAX_int32;
+
+	/** Some RHIs can only do visible or not occlusion queries. */
+	bool SupportsExactOcclusionQueries = true;
+
+	/** True if and only if the GPU support rendering to volume textures (2D Array, 3D). Some OpenGL 3.3 cards support SM4, but can't render to volume textures. */
+	bool SupportsVolumeTextureRendering = true;
+
+	/** True if the RHI supports separate blend states per render target. */
+	bool SupportsSeparateRenderTargetBlendState = false;
+
+	/** True if the RHI has artifacts with atlased CSM depths. */
+	bool NeedsUnatlasedCSMDepthsWorkaround = false;
+
+	/** true if the RHI supports 3D textures */
+	bool SupportsTexture3D = true;
+
+	/** true if the RHI supports mobile multi-view */
+	bool SupportsMobileMultiView = false;
+
+	/** true if the RHI supports image external */
+	bool SupportsImageExternal = false;
+
+	/** true if the RHI supports 256bit MRT */
+	bool SupportsWideMRT = true;
+
+	/** True if the RHI and current hardware supports supports depth bounds testing */
+	bool SupportsDepthBoundsTest = false;
+
+	/** True if the RHI supports explicit access to depth target HTile meta data. */
+	bool SupportsExplicitHTile = false;
+
+	/** True if the RHI supports explicit access to MSAA target FMask meta data. */
+	bool SupportsExplicitFMask = false;
+
+	/** True if the RHI supports resummarizing depth target HTile meta data. */
+	bool SupportsResummarizeHTile = false;
+
+	/** True if the RHI supports depth target unordered access views. */
+	bool SupportsDepthUAV = false;
+
+	/** True if the RHI and current hardware supports efficient AsyncCompute (by default we assume false and later we can enable this for more hardware) */
+	bool SupportsEfficientAsyncCompute = false;
+
+	/** True if the RHI supports getting the result of occlusion queries when on a thread other than the render thread */
+	bool SupportsParallelOcclusionQueries = false;
+
+	/** true if the RHI requires a valid RT bound during UAV scatter operation inside the pixel shader */
+	bool RequiresRenderTargetForPixelShaderUAVs = false;
+
+	/** true if the RHI supports unordered access view format aliasing */
+	bool SupportsUAVFormatAliasing = false;
+
+	/** true if the RHI supports texture views (data aliasing) */
+	bool SupportsTextureViews = true;
+
+	/** true if the pointer returned by Lock is a persistent direct pointer to gpu memory */
+	bool SupportsDirectGPUMemoryLock = false;
+
+	/** true if the multi-threaded shader creation is supported by (or desirable for) the RHI. */
+	bool SupportsMultithreadedShaderCreation = true;
+
+	/** Does the RHI support parallel resource commands (i.e. create / lock / unlock) on non-immediate command list APIs, recorded off the render thread. */
+	bool SupportsMultithreadedResources = false;
+
+	/** Does this RHI need to wait for deletion of resources due to ref counting. */
+	bool NeedsExtraDeletionLatency = false;
+
+	/** Allow opt-out default RHI resource deletion latency for streaming textures */
+	bool ForceNoDeletionLatencyForStreamingTextures = false;
+
+	/** The maximum size allowed for a computeshader dispatch. */
+	TRHIGlobal<int32> MaxComputeDispatchDimension = ((1 << 16) - 1);
+
+	/** If true, then avoiding loading shader code and instead force the "native" path, which sends a library and a hash instead. */
+	bool LazyShaderCodeLoading = false;
+
+	/** If true, then it is possible to turn on LazyShaderCodeLoading. */
+	bool SupportsLazyShaderCodeLoading = false;
+
+	/** true if the RHI supports UpdateFromBufferTexture method */
+	bool SupportsUpdateFromBufferTexture = false;
+
+	/** The maximum size to allow for the shadow depth buffer in the X dimension.  This must be larger or equal to GMaxShadowDepthBufferSizeY. */
+	TRHIGlobal<int32> MaxShadowDepthBufferSizeX = 2048;
+
+	/** The maximum size to allow for the shadow depth buffer in the Y dimension. */
+	TRHIGlobal<int32> MaxShadowDepthBufferSizeY = 2048;
+
+	/** The maximum size allowed for 2D textures in both dimensions. */
+	TRHIGlobal<int32> MaxTextureDimensions = 2048;
+
+	/** The maximum size allowed for 2D textures in both dimensions. */
+	TRHIGlobal<int64> MaxBufferDimensions = (1 << 27);
+
+	/** The maximum size allowed for a contant buffer. */
+	TRHIGlobal<int64> MaxConstantBufferByteSize = (1 << 27);
+
+	/** The maximum size allowed for Shared Compute Memory. */
+	TRHIGlobal<int64> MaxComputeSharedMemory = (1 << 15);
+
+	/** The maximum size allowed for 3D textures in all three dimensions. */
+	TRHIGlobal<int32> MaxVolumeTextureDimensions = 2048;
+
+	/** The maximum size allowed for cube textures. */
+	TRHIGlobal<int32> MaxCubeTextureDimensions = 2048;
+
+	/** Whether RW texture buffers are supported */
+	bool SupportsRWTextureBuffers = true;
+
+	/** Whether a raw (ByteAddress) buffer view can be created for any buffer, regardless of its EBufferUsageFlags::ByteAddressBuffer flag. */
+	bool SupportsRawViewsForAnyBuffer = false;
+
+	/** Whether depth or stencil can individually be set to CopySrc/Dest access. */
+	bool SupportsSeparateDepthStencilCopyAccess = true;
+
+	/** Support using async thread for texture stream out operations */
+	bool SupportsAsyncTextureStreamOut = false;
+
+	/** The Maximum number of layers in a 1D or 2D texture array. */
+	int32 MaxTextureArrayLayers = 256;
+
+	int32 MaxTextureSamplers = 16;
+
+	/** The maximum work group invocations allowed for compute shader. */
+	TRHIGlobal<int32> MaxWorkGroupInvocations = 1024;
+
+	/** true if we are running with the NULL RHI */
+	bool UsingNullRHI = false;
+
+	/**
+	 *	The size to check against for Draw*UP call vertex counts.
+	 *	If greater than this value, the draw call will not occur.
+	 */
+	int32 DrawUPVertexCheckCount = MAX_int32;
+	/**
+	 *	The size to check against for Draw*UP call index counts.
+	 *	If greater than this value, the draw call will not occur.
+	 */
+	int32 DrawUPIndexCheckCount = MAX_int32;
+
+	/** true for each VET that is supported. One-to-one mapping with EVertexElementType */
+	FVertexElementTypeSupportInfo VertexElementTypeSupport;
+
+	/** Whether the next frame should profile the GPU. */
+	bool TriggerGPUProfile = false;
+
+	/** Whether we are profiling GPU hitches. */
+	bool TriggerGPUHitchProfile = false;
+
+	/** Non-empty if we are performing a gpu trace. Also says where to place trace file. */
+	FString GPUTraceFileName;
+
+	/** True if the RHI supports texture streaming */
+	bool SupportsTextureStreaming = false;
+
+	/** Amount of memory allocated by textures. In kilobytes. */
+	volatile int32 CurrentTextureMemorySize = 0;
+
+	/** Amount of memory allocated by rendertargets. In kilobytes. */
+	volatile int32 CurrentRendertargetMemorySize = 0;
+
+	/** Current texture streaming pool size, in bytes. 0 means unlimited. */
+	int64 TexturePoolSize = 0 * 1024 * 1024;
+
+	/** In percent. If non-zero, the texture pool size is a percentage of TotalGraphicsMemory. */
+	int32 PoolSizeVRAMPercentage = 0;
+
+	/** Amount of local video memory demoted to system memory. In bytes. */
+	uint64 DemotedLocalMemorySize = 0;
+
+	/** Whether or not the RHI can handle a non-zero BaseVertexIndex - extra SetStreamSource calls will be needed if this is false */
+	bool SupportsBaseVertexIndex = true;
+
+	/** Whether or not the RHI can handle a non-zero FirstInstance to DrawIndexedPrimitive and friends - extra SetStreamSource calls will be needed if this is false */
+	bool SupportsFirstInstance = false;
+
+	/** Whether or not the RHI can handle dynamic resolution or not. */
+	bool SupportsDynamicResolution = false;
+
+	struct FRayTracing
+	{
+		/**
+		* Whether or not the RHI supports ray tracing on current hardware (acceleration structure building and new ray tracing-specific shader types).
+		* SupportsRayTracingShaders and SupportsInlineRayTracing must also be checked before dispatching ray tracing workloads.
+		*/
+		bool Supported = false;
+
+		/**
+		* Whether or not the RHI supports ray tracing raygen, miss and hit shaders (i.e. full ray tracing pipeline).
+		* The RHI may support inline ray tracing from compute shaders, but not the full pipeline.
+		*/
+		bool SupportsShaders = false;
+
+		/** Whether or not the RHI supports adding new shaders to an existing RT PSO. */
+		bool SupportsPSOAdditions = false;
+
+		/** Whether or not the RHI supports indirect ray tracing dispatch commands. */
+		bool SupportsDispatchIndirect = false;
+
+		/** Whether or not the RHI supports async building ray tracing acceleration structures. */
+		bool SupportsAsyncBuildAccelerationStructure = false;
+
+		/** Whether or not the RHI supports the AMD Hit Token extension. */
+		bool SupportsAMDHitToken = false;
+
+		/** Whether or not the RHI supports inline ray tracing in compute shaders, without a full ray tracing pipeline. */
+		bool SupportsInlineRayTracing = false;
+
+		/** Required alignment for ray tracing acceleration structures. */
+		uint32 AccelerationStructureAlignment = 0;
+
+		/** Required alignment for ray tracing scratch buffers. */
+		uint32 ScratchBufferAlignment = 0;
+
+		/** Required alignment for ray tracing shader binding table buffer. */
+		uint32 ShaderTableAlignment = 0;
+
+		/** Size of an individual element in the ray tracing instance buffer. This defines the required stride and alignment of structured buffers of instances. */
+		uint32 InstanceDescriptorSize = 0;
+
+	} RayTracing;
+
+	/** Whether or not the RHI supports shader wave operations (shader model 6.0). */
+	bool SupportsWaveOperations = false;
+
+	/** Whether or not the current GPU is integrated into to CPU */
+	bool DeviceIsIntegrated = false;
+
+	/**
+	* Specifies the minimum and maximum number of lanes in the SIMD wave that this GPU can support. I.e. 32 on NVIDIA, 64 on AMD.
+	* Valid values are in range [4..128] (as per SM 6.0 specification) or 0 if unknown.
+	* Rendering code must always check SupportsWaveOperations in addition to wave min/max size.
+	*/
+	int32 MinimumWaveSize = 0;
+	int32 MaximumWaveSize = 0;
+
+	/** Whether or not the RHI supports an RHI thread.
+	Requirements for RHI thread
+	* Microresources (those in RHIStaticStates.h) need to be able to be created by any thread at any time and be able to work with a radically simplified rhi resource lifecycle. CreateSamplerState, CreateRasterizerState, CreateDepthStencilState, CreateBlendState
+	* CreateUniformBuffer needs to be threadsafe
+	* GetRenderQueryResult should be threadsafe, but this isn't required. If it isn't threadsafe, then you need to flush yourself in the RHI
+	* GetViewportBackBuffer and AdvanceFrameForGetViewportBackBuffer need to be threadsafe and need to support the fact that the render thread has a different concept of "current backbuffer" than the RHI thread. Without an RHIThread this is moot due to the next two items.
+	* AdvanceFrameForGetViewportBackBuffer needs be added as an RHI method and this needs to work with GetViewportBackBuffer to give the render thread the right back buffer even though many commands relating to the beginning and end of the frame are queued.
+	* BeginDrawingViewport, and 5 or so other frame advance methods are queued with an RHIThread. Without an RHIThread, these just flush internally.
+	***/
+	bool SupportsRHIThread = false;
+
+	/* as above, but we run the commands on arbitrary task threads */
+	bool SupportsRHIOnTaskThread = false;
+
+	/** Whether or not the RHI supports parallel RHIThread executes / translates
+	Requirements:
+	* RHICreateBoundShaderState & RHICreateGraphicsPipelineState is threadsafe and GetCachedBoundShaderState must not be used. GetCachedBoundShaderState_Threadsafe has a slightly different protocol.
+	***/
+	bool SupportsParallelRHIExecute = false;
+
+	/** Whether or not the RHI can perform MSAA sample load. */
+	bool SupportsMSAADepthSampleAccess = false;
+
+	/** Whether or not the RHI can render to the backbuffer with a custom depth/stencil surface bound. */
+	bool SupportsBackBufferWithCustomDepthStencil = true;
+
+	/** Whether or not HDR is currently enabled */
+	bool IsHDREnabled = false;
+
+	/** Whether the present adapter/display offers HDR output capabilities. */
+	bool SupportsHDROutput = false;
+
+	/** The maximum number of groups that can be dispatched in each dimensions. */
+	FIntVector MaxDispatchThreadGroupsPerDimension = FIntVector::ZeroValue;
+
+	struct FVariableRateShading
+	{
+		/** Whether VRS (in all flavors) is currently enabled (separate from whether it's supported/available). */
+		bool Enabled = true;
+
+		/** Whether attachment (image-based) VRS is currently enabled (separate from whether it's supported/available). */
+		bool AttachmentEnabled = true;
+
+		/** Whether or not the RHI can support per-draw Variable Rate Shading. */
+		bool SupportsPipeline = false;
+
+		/** Whether or not the Variable Rate Shading can be done at larger (2x4 or 4x2 or 4x4) sizes. */
+		bool SupportsLargerSizes = false;
+
+		/** Whether or not the RHI can support image-based Variable Rate Shading. */
+		bool SupportsAttachment = false;
+
+		/** Whether or not the RHI can support complex combiner operatations between per-draw (pipeline) VRS and image VRS. */
+		bool SupportsComplexCombinerOps = false;
+
+		/** Whether or not the RHI can support shading rate attachments as array textures. */
+		bool SupportsAttachmentArrayTextures = false;
+
+		/** Maximum tile width in a screen space texture that can be used to drive Variable Rate Shading. */
+		int32 ImageTileMaxWidth = 0;
+
+		/** Maximum tile height in a screen space texture that can be used to drive Variable Rate Shading. */
+		int32 ImageTileMaxHeight = 0;
+
+		/** Minimum tile width in a screen space texture that can be used to drive Variable Rate Shading. */
+		int32 ImageTileMinWidth = 0;
+
+		/** Minimum tile height in a screen space texture that can be used to drive Variable Rate Shading. */
+		int32 ImageTileMinHeight = 0;
+
+		/** Data type contained in a shading-rate image for image-based Variable Rate Shading. */
+		EVRSImageDataType ImageDataType = VRSImage_NotSupported;
+
+		/** Image format for the shading rate image for image-based Variable Rate Shading. */
+		EPixelFormat ImageFormat = PF_Unknown;
+
+		/** Whether Variable Rate Shading deferred shading rate texture update is supported. */
+		bool SupportsLateUpdate = false;
+
+	} VariableRateShading;
+
+	/** Format used for the backbuffer when outputting to a HDR display. */
+	EPixelFormat HDRDisplayOutputFormat = PF_FloatRGBA;
+
+	/** Counter incremented once on each frame present. Used to support game thread synchronization with swap chain frame flips. */
+	uint64 PresentCounter = 1;
+
+	/** True if the RHI supports setting the render target array index from any shader stage */
+	bool SupportsArrayIndexFromAnyShader = false;
+
+	/** True if the pipeline file cache can be used with this RHI */
+	bool SupportsPipelineFileCache = false;
+
+	/** True if the RHI supports setting the stencil ref at pixel granularity from the pixel shader */
+	bool SupportsStencilRefFromPixelShader = false;
+
+	/** Whether current RHI supports overestimated conservative rasterization. */
+	bool SupportsConservativeRasterization = false;
+
+	/** true if the RHI supports Mesh and Amplification shaders with tier0 capability */
+	bool SupportsMeshShadersTier0 = false;
+
+	/** true if the RHI supports Mesh and Amplification shaders with tier1 capability */
+	bool SupportsMeshShadersTier1 = false;
+
+	/**
+	* True if the RHI supports reading system timer in shaders via GetShaderTimestamp().
+	* Individual shaders must be compiled with an appropriate vendor extension and check PLATFORM_SUPPORTS_SHADER_TIMESTAMP.
+	*/
+	bool SupportsShaderTimestamp = false;
+
+	bool SupportsEfficientUploadOnResourceCreation = false;
+
+	/** true if the RHI supports RLM_WriteOnly_NoOverwrite */
+	bool SupportsMapWriteNoOverwrite = false;
+
+	/** Tables of all MSAA sample offset for all MSAA supported. Use GetMSAASampleOffsets() to read it. */
+	FVector2f DefaultMSAASampleOffsets[1 + 2 + 4 + 8 + 16];
+
+	/** True if the RHI supports pipeline precompiling from any thread. */
+	bool SupportsAsyncPipelinePrecompile = true;
+
+	/** Whether dynamic (bindless) resources are supported */
+	ERHIBindlessSupport BindlessSupport = ERHIBindlessSupport::Unsupported;
+
+	UE_DEPRECATED(5.2, "You must use BindlessSupport instead.")
+	bool SupportsBindless = false;
+
+	/** True if the RHI supports reserved (AKA tiled, virtual or sparse) resources and operations related to them*/
+	bool SupportsReservedResources = false;
+
+	/** Table for finding out which shader platform corresponds to a given feature level for this RHI. */
+	EShaderPlatform ShaderPlatformForFeatureLevel[ERHIFeatureLevel::Num];
+};
+
+extern RHI_API FRHIGlobals GRHIGlobals;
+
+//
+// Deprecated old-style names
+
+#define GIsRHIInitialized                                      GRHIGlobals.IsRHIInitialized
+#define GRHIPersistentThreadGroupCount                         GRHIGlobals.PersistentThreadGroupCount
+#define GMaxTextureMipCount                                    GRHIGlobals.MaxTextureMipCount
+#define GSupportsQuadBufferStereo                              GRHIGlobals.SupportsQuadBufferStereo
+#define GSupportsRenderDepthTargetableShaderResources          GRHIGlobals.SupportsRenderDepthTargetableShaderResources
+#define GRHISupportsDrawIndirect                               GRHIGlobals.SupportsDrawIndirect
+#define GRHISupportsMultithreading                             GRHIGlobals.SupportsMultithreading
+#define GRHISupportsAsyncGetRenderQueryResult                  GRHIGlobals.SupportsAsyncGetRenderQueryResult
+#define GRHIAdapterName                                        GRHIGlobals.GpuInfo.AdapterName
+#define GRHIAdapterInternalDriverVersion                       GRHIGlobals.GpuInfo.AdapterInternalDriverVersion
+#define GRHIAdapterUserDriverVersion                           GRHIGlobals.GpuInfo.AdapterUserDriverVersion
+#define GRHIAdapterDriverDate                                  GRHIGlobals.GpuInfo.AdapterDriverDate
+#define GRHIAdapterDriverOnDenyList                            GRHIGlobals.GpuInfo.AdapterDriverOnDenyList
+#define GRHIDeviceId                                           GRHIGlobals.GpuInfo.DeviceId
+#define GRHIDeviceRevision                                     GRHIGlobals.GpuInfo.DeviceRevision
+#define GRHIVendorId                                           GRHIGlobals.GpuInfo.VendorId
+#define GRHISupportsPixelShaderUAVs                            GRHIGlobals.SupportsPixelShaderUAVs
+#define GRHIDeviceIsAMDPreGCNArchitecture                      GRHIGlobals.GpuInfo.IsAMDPreGCNArchitecture
+#define GSupportsRenderTargetFormat_PF_G8                      GRHIGlobals.SupportsRenderTargetFormat_PF_G8
+#define GSupportsRenderTargetFormat_PF_FloatRGBA               GRHIGlobals.SupportsRenderTargetFormat_PF_FloatRGBA
+#define GSupportsShaderFramebufferFetch                        GRHIGlobals.SupportsShaderFramebufferFetch
+#define GSupportsShaderMRTFramebufferFetch                     GRHIGlobals.SupportsShaderMRTFramebufferFetch
+#define GSupportsPixelLocalStorage                             GRHIGlobals.SupportsPixelLocalStorage
+#define GSupportsShaderDepthStencilFetch                       GRHIGlobals.SupportsShaderDepthStencilFetch
+#define GSupportsTimestampRenderQueries                        GRHIGlobals.SupportsTimestampRenderQueries
+#define GRHISupportsGPUTimestampBubblesRemoval                 GRHIGlobals.SupportsGPUTimestampBubblesRemoval
+#define GRHISupportsFrameCyclesBubblesRemoval                  GRHIGlobals.SupportsFrameCyclesBubblesRemoval
+#define GRHISupportsGPUUsage                                   GRHIGlobals.SupportsGPUUsage
+#define GHardwareHiddenSurfaceRemoval                          GRHIGlobals.HardwareHiddenSurfaceRemoval
+#define GRHISupportsAsyncTextureCreation                       GRHIGlobals.SupportsAsyncTextureCreation
+#define GRHISupportsQuadTopology                               GRHIGlobals.SupportsQuadTopology
+#define GRHISupportsRectTopology                               GRHIGlobals.SupportsRectTopology
+#define GRHISupportsPrimitiveShaders                           GRHIGlobals.SupportsPrimitiveShaders
+#define GRHISupportsAtomicUInt64                               GRHIGlobals.SupportsAtomicUInt64
+#define GRHISupportsDX12AtomicUInt64                           GRHIGlobals.SupportsDX12AtomicUInt64
+#define GRHISupportsPipelineStateSortKey                       GRHIGlobals.SupportsPipelineStateSortKey
+#define GSupportsParallelRenderingTasksWithSeparateRHIThread   GRHIGlobals.SupportsParallelRenderingTasksWithSeparateRHIThread
+#define GRHIThreadNeedsKicking                                 GRHIGlobals.RHIThreadNeedsKicking
+#define GRHIMaximumReccommendedOustandingOcclusionQueries      GRHIGlobals.MaximumReccommendedOustandingOcclusionQueries
+#define GRHISupportsExactOcclusionQueries                      GRHIGlobals.SupportsExactOcclusionQueries
+#define GSupportsVolumeTextureRendering                        GRHIGlobals.SupportsVolumeTextureRendering
+#define GSupportsSeparateRenderTargetBlendState                GRHIGlobals.SupportsSeparateRenderTargetBlendState
+#define GRHINeedsUnatlasedCSMDepthsWorkaround                  GRHIGlobals.NeedsUnatlasedCSMDepthsWorkaround
+#define GSupportsTexture3D                                     GRHIGlobals.SupportsTexture3D
+#define GSupportsMobileMultiView                               GRHIGlobals.SupportsMobileMultiView
+#define GSupportsImageExternal                                 GRHIGlobals.SupportsImageExternal
+#define GSupportsWideMRT                                       GRHIGlobals.SupportsWideMRT
+#define GSupportsDepthBoundsTest                               GRHIGlobals.SupportsDepthBoundsTest
+#define GRHISupportsExplicitHTile                              GRHIGlobals.SupportsExplicitHTile
+#define GRHISupportsExplicitFMask                              GRHIGlobals.SupportsExplicitFMask
+#define GRHISupportsResummarizeHTile                           GRHIGlobals.SupportsResummarizeHTile
+#define GRHISupportsDepthUAV                                   GRHIGlobals.SupportsDepthUAV
+#define GSupportsEfficientAsyncCompute                         GRHIGlobals.SupportsEfficientAsyncCompute
+#define GSupportsParallelOcclusionQueries                      GRHIGlobals.SupportsParallelOcclusionQueries
+#define GRHIRequiresRenderTargetForPixelShaderUAVs             GRHIGlobals.RequiresRenderTargetForPixelShaderUAVs
+#define GRHISupportsUAVFormatAliasing                          GRHIGlobals.SupportsUAVFormatAliasing
+#define GRHISupportsTextureViews                               GRHIGlobals.SupportsTextureViews
+#define GRHISupportsDirectGPUMemoryLock                        GRHIGlobals.SupportsDirectGPUMemoryLock
+#define GRHISupportsMultithreadedShaderCreation                GRHIGlobals.SupportsMultithreadedShaderCreation
+#define GRHISupportsMultithreadedResources                     GRHIGlobals.SupportsMultithreadedResources
+#define GRHINeedsExtraDeletionLatency                          GRHIGlobals.NeedsExtraDeletionLatency
+#define GRHIForceNoDeletionLatencyForStreamingTextures         GRHIGlobals.ForceNoDeletionLatencyForStreamingTextures
+#define GMaxComputeDispatchDimension                           GRHIGlobals.MaxComputeDispatchDimension
+#define GRHILazyShaderCodeLoading                              GRHIGlobals.LazyShaderCodeLoading
+#define GRHISupportsLazyShaderCodeLoading                      GRHIGlobals.SupportsLazyShaderCodeLoading
+#define GRHISupportsUpdateFromBufferTexture                    GRHIGlobals.SupportsUpdateFromBufferTexture
+#define GMaxShadowDepthBufferSizeX                             GRHIGlobals.MaxShadowDepthBufferSizeX
+#define GMaxShadowDepthBufferSizeY                             GRHIGlobals.MaxShadowDepthBufferSizeY
+#define GMaxTextureDimensions                                  GRHIGlobals.MaxTextureDimensions
+#define GMaxBufferDimensions                                   GRHIGlobals.MaxBufferDimensions
+#define GRHIMaxConstantBufferByteSize                          GRHIGlobals.MaxConstantBufferByteSize
+#define GMaxComputeSharedMemory                                GRHIGlobals.MaxComputeSharedMemory
+#define GMaxVolumeTextureDimensions                            GRHIGlobals.MaxVolumeTextureDimensions
+#define GRHISupportsRWTextureBuffers                           GRHIGlobals.SupportsRWTextureBuffers
+#define GRHISupportsRawViewsForAnyBuffer                       GRHIGlobals.SupportsRawViewsForAnyBuffer
+#define GRHISupportsSeparateDepthStencilCopyAccess             GRHIGlobals.SupportsSeparateDepthStencilCopyAccess
+#define GRHISupportAsyncTextureStreamOut                       GRHIGlobals.SupportsAsyncTextureStreamOut
+#define GMaxCubeTextureDimensions                              GRHIGlobals.MaxCubeTextureDimensions
+#define GMaxTextureArrayLayers                                 GRHIGlobals.MaxTextureArrayLayers
+#define GMaxTextureSamplers                                    GRHIGlobals.MaxTextureSamplers
+#define GMaxWorkGroupInvocations                               GRHIGlobals.MaxWorkGroupInvocations
+#define GUsingNullRHI                                          GRHIGlobals.UsingNullRHI
+#define GDrawUPVertexCheckCount                                GRHIGlobals.DrawUPVertexCheckCount
+#define GDrawUPIndexCheckCount                                 GRHIGlobals.DrawUPIndexCheckCount
+#define GVertexElementTypeSupport                              GRHIGlobals.VertexElementTypeSupport
+#define GTriggerGPUProfile                                     GRHIGlobals.TriggerGPUProfile
+#define GTriggerGPUHitchProfile                                GRHIGlobals.TriggerGPUHitchProfile
+#define GGPUTraceFileName                                      GRHIGlobals.GPUTraceFileName
+#define GRHISupportsTextureStreaming                           GRHIGlobals.SupportsTextureStreaming
+#define GCurrentTextureMemorySize                              GRHIGlobals.CurrentTextureMemorySize
+#define GCurrentRendertargetMemorySize                         GRHIGlobals.CurrentRendertargetMemorySize
+#define GTexturePoolSize                                       GRHIGlobals.TexturePoolSize
+#define GPoolSizeVRAMPercentage                                GRHIGlobals.PoolSizeVRAMPercentage
+#define GDemotedLocalMemorySize                                GRHIGlobals.DemotedLocalMemorySize
+#define GRHISupportsBaseVertexIndex                            GRHIGlobals.SupportsBaseVertexIndex
+#define GRHISupportsFirstInstance                              GRHIGlobals.SupportsFirstInstance
+#define GRHISupportsDynamicResolution                          GRHIGlobals.SupportsDynamicResolution
+#define GRHISupportsRayTracing                                 GRHIGlobals.RayTracing.Supported
+#define GRHISupportsRayTracingShaders                          GRHIGlobals.RayTracing.SupportsShaders
+#define GRHISupportsRayTracingPSOAdditions                     GRHIGlobals.RayTracing.SupportsPSOAdditions
+#define GRHISupportsRayTracingDispatchIndirect                 GRHIGlobals.RayTracing.SupportsDispatchIndirect
+#define GRHISupportsRayTracingAsyncBuildAccelerationStructure  GRHIGlobals.RayTracing.SupportsAsyncBuildAccelerationStructure
+#define GRHISupportsRayTracingAMDHitToken                      GRHIGlobals.RayTracing.SupportsAMDHitToken
+#define GRHISupportsInlineRayTracing                           GRHIGlobals.RayTracing.SupportsInlineRayTracing
+#define GRHIRayTracingAccelerationStructureAlignment           GRHIGlobals.RayTracing.AccelerationStructureAlignment
+#define GRHIRayTracingScratchBufferAlignment                   GRHIGlobals.RayTracing.ScratchBufferAlignment
+#define GRHIRayTracingShaderTableAlignment                     GRHIGlobals.RayTracing.ShaderTableAlignment
+#define GRHIRayTracingInstanceDescriptorSize                   GRHIGlobals.RayTracing.InstanceDescriptorSize
+#define GRHISupportsWaveOperations                             GRHIGlobals.SupportsWaveOperations
+#define GRHIDeviceIsIntegrated                                 GRHIGlobals.DeviceIsIntegrated
+#define GRHIMinimumWaveSize                                    GRHIGlobals.MinimumWaveSize
+#define GRHIMaximumWaveSize                                    GRHIGlobals.MaximumWaveSize
+#define GRHISupportsRHIThread                                  GRHIGlobals.SupportsRHIThread
+#define GRHISupportsRHIOnTaskThread                            GRHIGlobals.SupportsRHIOnTaskThread
+#define GRHISupportsParallelRHIExecute                         GRHIGlobals.SupportsParallelRHIExecute
+#define GRHISupportsMSAADepthSampleAccess                      GRHIGlobals.SupportsMSAADepthSampleAccess
+#define GRHISupportsBackBufferWithCustomDepthStencil           GRHIGlobals.SupportsBackBufferWithCustomDepthStencil
+#define GRHIIsHDREnabled                                       GRHIGlobals.IsHDREnabled
+#define GRHISupportsHDROutput                                  GRHIGlobals.SupportsHDROutput
+#define GRHIMaxDispatchThreadGroupsPerDimension                GRHIGlobals.MaxDispatchThreadGroupsPerDimension
+#define GRHIVariableRateShadingEnabled                         GRHIGlobals.VariableRateShading.Enabled
+#define GRHIAttachmentVariableRateShadingEnabled               GRHIGlobals.VariableRateShading.AttachmentEnabled
+#define GRHISupportsPipelineVariableRateShading                GRHIGlobals.VariableRateShading.SupportsPipeline
+#define GRHISupportsLargerVariableRateShadingSizes             GRHIGlobals.VariableRateShading.SupportsLargerSizes
+#define GRHISupportsAttachmentVariableRateShading              GRHIGlobals.VariableRateShading.SupportsAttachment
+#define GRHISupportsComplexVariableRateShadingCombinerOps      GRHIGlobals.VariableRateShading.SupportsComplexCombinerOps
+#define GRHISupportsVariableRateShadingAttachmentArrayTextures GRHIGlobals.VariableRateShading.SupportsAttachmentArrayTextures
+#define GRHIVariableRateShadingImageTileMaxWidth               GRHIGlobals.VariableRateShading.ImageTileMaxWidth
+#define GRHIVariableRateShadingImageTileMaxHeight              GRHIGlobals.VariableRateShading.ImageTileMaxHeight
+#define GRHIVariableRateShadingImageTileMinWidth               GRHIGlobals.VariableRateShading.ImageTileMinWidth
+#define GRHIVariableRateShadingImageTileMinHeight              GRHIGlobals.VariableRateShading.ImageTileMinHeight
+#define GRHIVariableRateShadingImageDataType                   GRHIGlobals.VariableRateShading.ImageDataType
+#define GRHIVariableRateShadingImageFormat                     GRHIGlobals.VariableRateShading.ImageFormat
+#define GRHISupportsLateVariableRateShadingUpdate              GRHIGlobals.VariableRateShading.SupportsLateUpdate
+#define GRHIHDRDisplayOutputFormat                             GRHIGlobals.HDRDisplayOutputFormat
+#define GRHIPresentCounter                                     GRHIGlobals.PresentCounter
+#define GRHISupportsArrayIndexFromAnyShader                    GRHIGlobals.SupportsArrayIndexFromAnyShader
+#define GRHISupportsPipelineFileCache                          GRHIGlobals.SupportsPipelineFileCache
+#define GRHISupportsStencilRefFromPixelShader                  GRHIGlobals.SupportsStencilRefFromPixelShader
+#define GRHISupportsConservativeRasterization                  GRHIGlobals.SupportsConservativeRasterization
+#define GRHISupportsMeshShadersTier0                           GRHIGlobals.SupportsMeshShadersTier0
+#define GRHISupportsMeshShadersTier1                           GRHIGlobals.SupportsMeshShadersTier1
+#define GRHISupportsShaderTimestamp                            GRHIGlobals.SupportsShaderTimestamp
+#define GRHISupportsEfficientUploadOnResourceCreation          GRHIGlobals.SupportsEfficientUploadOnResourceCreation
+#define GRHISupportsMapWriteNoOverwrite                        GRHIGlobals.SupportsMapWriteNoOverwrite
+#define GRHIDefaultMSAASampleOffsets                           GRHIGlobals.DefaultMSAASampleOffsets
+#define GRHISupportsAsyncPipelinePrecompile                    GRHIGlobals.SupportsAsyncPipelinePrecompile
+#define GRHIBindlessSupport                                    GRHIGlobals.BindlessSupport
+#define GRHISupportsBindless                                   GRHIGlobals.SupportsBindless
+#define GRHISupportsReservedResources                          GRHIGlobals.SupportsReservedResources
+#define GShaderPlatformForFeatureLevel                         GRHIGlobals.ShaderPlatformForFeatureLevel
+
+// Utility Getters
+
+inline uint64 GetMaxBufferDimension()
+{
+	return GRHIGlobals.MaxBufferDimensions;
+}
+
+inline uint64 GetMaxConstantBufferByteSize()
+{
+	return GRHIGlobals.MaxConstantBufferByteSize;
+}
+
+inline uint64 GetMaxComputeSharedMemory()
+{
+	return GRHIGlobals.MaxComputeSharedMemory;
+}
+
+inline uint32 GetMax2DTextureDimension()
+{
+	return GRHIGlobals.MaxTextureDimensions;
+}
+
+inline uint32 GetMaxCubeTextureDimension()
+{
+	return GRHIGlobals.MaxCubeTextureDimensions;
+}
+
+inline uint32 GetMaxTextureArrayLayers()
+{
+	return GRHIGlobals.MaxTextureArrayLayers;
+}
+
+inline uint32 GetMaxTextureSamplers()
+{
+	return GRHIGlobals.MaxTextureSamplers;
+}
+
+inline uint32 GetMaxWorkGroupInvocations()
+{
+	return GRHIGlobals.MaxWorkGroupInvocations;
+}
+
+/** Get the shader platform associated with the supplied feature level on this machine */
+inline EShaderPlatform GetFeatureLevelShaderPlatform(const FStaticFeatureLevel InFeatureLevel)
+{
+	return GRHIGlobals.ShaderPlatformForFeatureLevel[InFeatureLevel];
+}
