Fixed missing TexCreate flags in GetTextureCreateFlagsName.

diff --git a/Engine/Source/Runtime/RHI/Private/RHIStrings.cpp b/Engine/Source/Runtime/RHI/Private/RHIStrings.cpp
--- a/Engine/Source/Runtime/RHI/Private/RHIStrings.cpp
+++ b/Engine/Source/Runtime/RHI/Private/RHIStrings.cpp
@@ -336,42 +336,47 @@ FString GetTextureCreateFlagsName(ETextureCreateFlags TextureCreateFlags)
 	{
 		switch (TextureCreateFlags)
 		{
-		default: checkNoEntry(); // fall through
-		case TexCreate_None:							return TEXT("TexCreate_None");
-		case TexCreate_RenderTargetable:				return TEXT("TexCreate_RenderTargetable");
-		case TexCreate_ResolveTargetable:				return TEXT("TexCreate_ResolveTargetable");
-		case TexCreate_DepthStencilTargetable:			return TEXT("TexCreate_DepthStencilTargetable");
-		case TexCreate_ShaderResource:					return TEXT("TexCreate_ShaderResource");
-		case TexCreate_SRGB:							return TEXT("TexCreate_SRGB");
-		case TexCreate_CPUWritable:						return TEXT("TexCreate_CPUWritable");
-		case TexCreate_NoTiling:						return TEXT("TexCreate_NoTiling");
-		case TexCreate_VideoDecode:						return TEXT("TexCreate_VideoDecode");
-		case TexCreate_Dynamic:							return TEXT("TexCreate_Dynamic");
-		case TexCreate_InputAttachmentRead:				return TEXT("TexCreate_InputAttachmentRead");
-		case TexCreate_Foveation:						return TEXT("TexCreate_Foveation");
-		case TexCreate_3DTiling:						return TEXT("TexCreate_3DTiling");
-		case TexCreate_Memoryless:						return TEXT("TexCreate_Memoryless");
-		case TexCreate_GenerateMipCapable:				return TEXT("TexCreate_GenerateMipCapable");
-		case TexCreate_FastVRAMPartialAlloc:			return TEXT("TexCreate_FastVRAMPartialAlloc");
-		case TexCreate_DisableSRVCreation:				return TEXT("TexCreate_DisableSRVCreation");
-		case TexCreate_DisableDCC:						return TEXT("TexCreate_DisableDCC");
-		case TexCreate_UAV:								return TEXT("TexCreate_UAV");
-		case TexCreate_Presentable:						return TEXT("TexCreate_Presentable");
-		case TexCreate_CPUReadback:						return TEXT("TexCreate_CPUReadback");
-		case TexCreate_OfflineProcessed:				return TEXT("TexCreate_OfflineProcessed");
-		case TexCreate_FastVRAM:						return TEXT("TexCreate_FastVRAM");
-		case TexCreate_HideInVisualizeTexture:			return TEXT("TexCreate_HideInVisualizeTexture");
-		case TexCreate_Virtual:							return TEXT("TexCreate_Virtual");
-		case TexCreate_TargetArraySlicesIndependently:	return TEXT("TexCreate_TargetArraySlicesIndependently");
-		case TexCreate_Shared:							return TEXT("TexCreate_Shared");
-		case TexCreate_NoFastClear:						return TEXT("TexCreate_NoFastClear");
-		case TexCreate_DepthStencilResolveTarget:		return TEXT("TexCreate_DepthStencilResolveTarget");
-		case TexCreate_Streamable:						return TEXT("TexCreate_Streamable");
-		case TexCreate_NoFastClearFinalize:				return TEXT("TexCreate_NoFastClearFinalize");
-		case TexCreate_ReduceMemoryWithTilingMode:		return TEXT("TexCreate_ReduceMemoryWithTilingMode");
-		case TexCreate_AtomicCompatible:				return TEXT("TexCreate_AtomicCompatible");
-		case TexCreate_External:						return TEXT("TexCreate_External");
+		case ETextureCreateFlags::None:								return TEXT("None");
+		case ETextureCreateFlags::RenderTargetable:					return TEXT("RenderTargetable");
+		case ETextureCreateFlags::ResolveTargetable:				return TEXT("ResolveTargetable");
+		case ETextureCreateFlags::DepthStencilTargetable:			return TEXT("DepthStencilTargetable");
+		case ETextureCreateFlags::ShaderResource:					return TEXT("ShaderResource");
+		case ETextureCreateFlags::SRGB:								return TEXT("SRGB");
+		case ETextureCreateFlags::CPUWritable:						return TEXT("CPUWritable");
+		case ETextureCreateFlags::NoTiling:							return TEXT("NoTiling");
+		case ETextureCreateFlags::VideoDecode:						return TEXT("VideoDecode");
+		case ETextureCreateFlags::Dynamic:							return TEXT("Dynamic");
+		case ETextureCreateFlags::InputAttachmentRead:				return TEXT("InputAttachmentRead");
+		case ETextureCreateFlags::Foveation:						return TEXT("Foveation");
+		case ETextureCreateFlags::Tiling3D:							return TEXT("Tiling3D");
+		case ETextureCreateFlags::Memoryless:						return TEXT("Memoryless");
+		case ETextureCreateFlags::GenerateMipCapable:				return TEXT("GenerateMipCapable");
+		case ETextureCreateFlags::FastVRAMPartialAlloc:				return TEXT("FastVRAMPartialAlloc");
+		case ETextureCreateFlags::DisableSRVCreation:				return TEXT("DisableSRVCreation");
+		case ETextureCreateFlags::DisableDCC:						return TEXT("DisableDCC");
+		case ETextureCreateFlags::UAV:								return TEXT("UAV");
+		case ETextureCreateFlags::Presentable:						return TEXT("Presentable");
+		case ETextureCreateFlags::CPUReadback:						return TEXT("CPUReadback");
+		case ETextureCreateFlags::OfflineProcessed:					return TEXT("OfflineProcessed");
+		case ETextureCreateFlags::FastVRAM:							return TEXT("FastVRAM");
+		case ETextureCreateFlags::HideInVisualizeTexture:			return TEXT("HideInVisualizeTexture");
+		case ETextureCreateFlags::Virtual:							return TEXT("Virtual");
+		case ETextureCreateFlags::TargetArraySlicesIndependently:	return TEXT("TargetArraySlicesIndependently");
+		case ETextureCreateFlags::Shared:							return TEXT("Shared");
+		case ETextureCreateFlags::NoFastClear:						return TEXT("NoFastClear");
+		case ETextureCreateFlags::DepthStencilResolveTarget:		return TEXT("DepthStencilResolveTarget");
+		case ETextureCreateFlags::Streamable:						return TEXT("Streamable");
+		case ETextureCreateFlags::NoFastClearFinalize:				return TEXT("NoFastClearFinalize");
+		case ETextureCreateFlags::ReduceMemoryWithTilingMode:		return TEXT("ReduceMemoryWithTilingMode");
+		case ETextureCreateFlags::AtomicCompatible:					return TEXT("AtomicCompatible");
+		case ETextureCreateFlags::Atomic64Compatible:				return TEXT("Aomic64Capable");
+		case ETextureCreateFlags::External:							return TEXT("External");
+		case ETextureCreateFlags::MultiGPUGraphIgnore:				return TEXT("MultiGPUGraphIgnore");
+		case ETextureCreateFlags::ReservedResource:					return TEXT("ReservedResource");
+		case ETextureCreateFlags::ImmediateCommit:					return TEXT("ImmediateCommit");
 		}
+		ensureMsgf(false, TEXT("Unexpected texture creation flag bit: %x"), TextureCreateFlags);
+		return TEXT("Unknown");
 	});
 }
 
