RHI cleanup
- Moving GUseTexture3DBulkDataRHI into RHIGlobals.h
- Removing uses of RHI_API from cpp files. The declarations are enough to export these.

diff --git a/Engine/Source/Runtime/Engine/Private/Rendering/Texture3DResource.cpp b/Engine/Source/Runtime/Engine/Private/Rendering/Texture3DResource.cpp
--- a/Engine/Source/Runtime/Engine/Private/Rendering/Texture3DResource.cpp
+++ b/Engine/Source/Runtime/Engine/Private/Rendering/Texture3DResource.cpp
@@ -9,10 +9,6 @@
 #include "EngineLogs.h"
 #include "RenderUtils.h"
 
-//*****************************************************************************
-
-extern RHI_API bool GUseTexture3DBulkDataRHI;
-
 //*****************************************************************************
 //************************* FVolumeTextureBulkData ****************************
 //*****************************************************************************
diff --git a/Engine/Source/Runtime/Engine/Private/Streaming/VolumeTextureStreaming.cpp b/Engine/Source/Runtime/Engine/Private/Streaming/VolumeTextureStreaming.cpp
--- a/Engine/Source/Runtime/Engine/Private/Streaming/VolumeTextureStreaming.cpp
+++ b/Engine/Source/Runtime/Engine/Private/Streaming/VolumeTextureStreaming.cpp
@@ -8,8 +8,6 @@ VolumeTextureStreaming.cpp: Helpers to stream in and out volume texture LODs.
 #include "RenderingThread.h"
 #include "Streaming/TextureMipAllocator.h"
 
-extern RHI_API bool GUseTexture3DBulkDataRHI;
-
 //*****************************************************************************
 //******************** FVolumeTextureMipAllocator_Reallocate ******************
 //*****************************************************************************
diff --git a/Engine/Source/Runtime/RHI/Private/RHI.cpp b/Engine/Source/Runtime/RHI/Private/RHI.cpp
--- a/Engine/Source/Runtime/RHI/Private/RHI.cpp
+++ b/Engine/Source/Runtime/RHI/Private/RHI.cpp
@@ -1290,7 +1290,7 @@ ERHIAccess GRHIMergeableAccessMask = ERHIAccess::ReadOnlyMask | ERHIAccess::UAVM
 // By default, only exclusively read only accesses are allowed.
 ERHIAccess GRHIMultiPipelineMergeableAccessMask = ERHIAccess::ReadOnlyExclusiveMask;
 
-RHI_API void FRHIDrawStats::Accumulate(FRHIDrawStats& Other)
+void FRHIDrawStats::Accumulate(FRHIDrawStats& Other)
 {
 	for (uint32 GPUIndex = 0; GPUIndex < GNumExplicitGPUsForRendering; ++GPUIndex)
 	{
@@ -1378,39 +1378,36 @@ void FRHICommandListImmediate::ProcessStats()
 	FrameDrawStats.Reset();
 }
 
-/** Whether to initialize 3D textures using a bulk data (or through a mip update if false). */
-RHI_API bool GUseTexture3DBulkDataRHI = false;
-
 //
 // The current shader platform.
 //
 
-RHI_API EShaderPlatform GMaxRHIShaderPlatform = SP_PCD3D_SM5;
+EShaderPlatform GMaxRHIShaderPlatform = SP_PCD3D_SM5;
 
 /** The maximum feature level supported on this machine */
-RHI_API ERHIFeatureLevel::Type GMaxRHIFeatureLevel = ERHIFeatureLevel::SM5;
+ERHIFeatureLevel::Type GMaxRHIFeatureLevel = ERHIFeatureLevel::SM5;
 
-RHI_API bool IsRHIDeviceAMD()
+bool IsRHIDeviceAMD()
 {
 	check(GRHIVendorId != 0);
 	// AMD's drivers tested on July 11 2013 have hitching problems with async resource streaming, setting single threaded for now until fixed.
 	return GRHIVendorId == 0x1002;
 }
 
-RHI_API bool IsRHIDeviceIntel()
+bool IsRHIDeviceIntel()
 {
 	check(GRHIVendorId != 0);
 	return GRHIVendorId == 0x8086;
 }
 
-RHI_API bool IsRHIDeviceNVIDIA()
+bool IsRHIDeviceNVIDIA()
 {
 	check(GRHIVendorId != 0);
 	// NVIDIA GPUs are discrete and use DedicatedVideoMemory only.
 	return GRHIVendorId == 0x10DE;
 }
 
-RHI_API uint32 RHIGetMetalShaderLanguageVersion(const FStaticShaderPlatform Platform)
+uint32 RHIGetMetalShaderLanguageVersion(const FStaticShaderPlatform Platform)
 {
     if (IsMetalPlatform(Platform))
 	{
@@ -1443,7 +1440,8 @@ RHI_API uint32 RHIGetMetalShaderLanguageVersion(const FStaticShaderPlatform Plat
 }
 
 static ERHIFeatureLevel::Type GRHIMobilePreviewFeatureLevel = ERHIFeatureLevel::Num;
-RHI_API void RHISetMobilePreviewFeatureLevel(ERHIFeatureLevel::Type MobilePreviewFeatureLevel)
+
+void RHISetMobilePreviewFeatureLevel(ERHIFeatureLevel::Type MobilePreviewFeatureLevel)
 {
 	check(GRHIMobilePreviewFeatureLevel == ERHIFeatureLevel::Num);
 	check(!GIsEditor);
@@ -1469,7 +1467,7 @@ bool RHIGetPreviewFeatureLevel(ERHIFeatureLevel::Type& PreviewFeatureLevelOUT)
 	return true;
 }
 
- RHI_API EPixelFormat RHIPreferredPixelFormatHint(EPixelFormat PreferredPixelFormat)
+EPixelFormat RHIPreferredPixelFormatHint(EPixelFormat PreferredPixelFormat)
 {
 	if (GDynamicRHI)
 	{
@@ -1478,7 +1476,7 @@ bool RHIGetPreviewFeatureLevel(ERHIFeatureLevel::Type& PreviewFeatureLevelOUT)
 	return PreferredPixelFormat;
 }
 
-RHI_API int32 RHIGetPreferredClearUAVRectPSResourceType(const FStaticShaderPlatform Platform)
+int32 RHIGetPreferredClearUAVRectPSResourceType(const FStaticShaderPlatform Platform)
 {
 	if (IsMetalPlatform(Platform))
 	{
diff --git a/Engine/Source/Runtime/RHI/Private/RHICommandList.cpp b/Engine/Source/Runtime/RHI/Private/RHICommandList.cpp
--- a/Engine/Source/Runtime/RHI/Private/RHICommandList.cpp
+++ b/Engine/Source/Runtime/RHI/Private/RHICommandList.cpp
@@ -89,7 +89,7 @@ static FAutoConsoleTaskPriority CPrio_FParallelTranslateCommandListPrepass(
 	ENamedThreads::HighTaskPriority
 );
 
-RHI_API FAutoConsoleTaskPriority CPrio_SceneRenderingTask(
+FAutoConsoleTaskPriority CPrio_SceneRenderingTask(
 	TEXT("TaskGraph.TaskPriorities.SceneRenderingTask"),
 	TEXT("Task and thread priority for various scene rendering tasks."),
 	ENamedThreads::NormalThreadPriority,
@@ -121,8 +121,8 @@ uint64 GInputLatencyTime = 0;
 
 FRHICommandBase* GCurrentCommand = nullptr;
 
-RHI_API bool GEnableAsyncCompute = true;
-RHI_API FRHICommandListExecutor GRHICommandList;
+bool GEnableAsyncCompute = true;
+FRHICommandListExecutor GRHICommandList;
 
 FGraphEventArray FRHICommandListImmediate::WaitOutstandingTasks;
 FGraphEventRef   FRHICommandListImmediate::RHIThreadTask;
diff --git a/Engine/Source/Runtime/RHI/Private/RHITransition.cpp b/Engine/Source/Runtime/RHI/Private/RHITransition.cpp
--- a/Engine/Source/Runtime/RHI/Private/RHITransition.cpp
+++ b/Engine/Source/Runtime/RHI/Private/RHITransition.cpp
@@ -3,8 +3,8 @@
 #include "RHITransition.h"
 #include "DynamicRHI.h"
 
-RHI_API uint64 GRHITransitionPrivateData_SizeInBytes = 0;
-RHI_API uint64 GRHITransitionPrivateData_AlignInBytes = 0;
+uint64 GRHITransitionPrivateData_SizeInBytes = 0;
+uint64 GRHITransitionPrivateData_AlignInBytes = 0;
 
 FRHIViewableResource* GetViewableResource(const FRHITransitionInfo& Info)
 {
diff --git a/Engine/Source/Runtime/RHI/Private/RHIUtilities.cpp b/Engine/Source/Runtime/RHI/Private/RHIUtilities.cpp
--- a/Engine/Source/Runtime/RHI/Private/RHIUtilities.cpp
+++ b/Engine/Source/Runtime/RHI/Private/RHIUtilities.cpp
@@ -517,12 +517,12 @@ FRHIFrameFlipTrackingRunnable FRHIFrameFlipTrackingRunnable::Singleton;
 bool FRHIFrameFlipTrackingRunnable::bInitialized = false;
 bool FRHIFrameFlipTrackingRunnable::bRun = false;
 
-RHI_API uint32 RHIGetSyncInterval()
+uint32 RHIGetSyncInterval()
 {
 	return FMath::Max(CVarRHISyncInterval.GetValueOnAnyThread(), 0);
 }
 
-RHI_API float RHIGetSyncSlackMS()
+float RHIGetSyncSlackMS()
 {
 #if USE_FRAME_OFFSET_THREAD
 	const float SyncSlackMS = CVarRHISyncSlackMS.GetValueOnAnyThread();
@@ -532,30 +532,30 @@ RHI_API float RHIGetSyncSlackMS()
 	return SyncSlackMS;
 }
 
-RHI_API bool RHIGetSyncAllowVariable()
+bool RHIGetSyncAllowVariable()
 {
 	return CVarRHISyncAllowVariable.GetValueOnAnyThread() != 0;
 }
 
-RHI_API void RHIGetPresentThresholds(float& OutTopPercent, float& OutBottomPercent)
+void RHIGetPresentThresholds(float& OutTopPercent, float& OutBottomPercent)
 {
 	OutTopPercent = FMath::Clamp(CVarRHIPresentThresholdTop.GetValueOnAnyThread(), 0.0f, 1.0f);
 	OutBottomPercent = FMath::Clamp(CVarRHIPresentThresholdBottom.GetValueOnAnyThread(), 0.0f, 1.0f);
 }
 
-RHI_API void RHICompleteGraphEventOnFlip(uint64 PresentIndex, FGraphEventRef Event)
+void RHICompleteGraphEventOnFlip(uint64 PresentIndex, FGraphEventRef Event)
 {
 	FRHIFrameFlipTrackingRunnable::CompleteGraphEventOnFlip(PresentIndex, Event);
 }
 
-RHI_API void RHISetFrameDebugInfo(uint64 PresentIndex, uint64 FrameIndex, uint64 InputTime)
+void RHISetFrameDebugInfo(uint64 PresentIndex, uint64 FrameIndex, uint64 InputTime)
 {
 #if USE_FRAME_OFFSET_THREAD
 	FRHIFrameOffsetThread::SetFrameDebugInfo(PresentIndex, FrameIndex, InputTime);
 #endif
 }
 
-RHI_API void RHIInitializeFlipTracking()
+void RHIInitializeFlipTracking()
 {
 #if USE_FRAME_OFFSET_THREAD
 	FRHIFrameOffsetThread::Initialize();
@@ -563,19 +563,19 @@ RHI_API void RHIInitializeFlipTracking()
 	FRHIFrameFlipTrackingRunnable::Initialize();
 }
 
-RHI_API void RHICalculateFrameTime()
+void RHICalculateFrameTime()
 {
 	double CurrentTimeInSeconds = FPlatformTime::Seconds();
 	GRHIFrameTimeMS = (float)((CurrentTimeInSeconds - GLastRHITimeInSeconds) * 1000.0);
 	GLastRHITimeInSeconds = CurrentTimeInSeconds;
 }
 
-RHI_API float RHIGetFrameTime()
+float RHIGetFrameTime()
 {
 	return GRHIFrameTimeMS;
 }
 
-RHI_API void RHIShutdownFlipTracking()
+void RHIShutdownFlipTracking()
 {
 	FRHIFrameFlipTrackingRunnable::Shutdown();
 #if USE_FRAME_OFFSET_THREAD
@@ -583,7 +583,7 @@ RHI_API void RHIShutdownFlipTracking()
 #endif
 }
 
-RHI_API ERHIAccess RHIGetDefaultResourceState(ETextureCreateFlags InUsage, bool bInHasInitialData)
+ERHIAccess RHIGetDefaultResourceState(ETextureCreateFlags InUsage, bool bInHasInitialData)
 {
 	// By default assume it can be bound for reading
 	ERHIAccess ResourceState = ERHIAccess::SRVMask;
@@ -615,7 +615,7 @@ RHI_API ERHIAccess RHIGetDefaultResourceState(ETextureCreateFlags InUsage, bool
 	return ResourceState;
 }
 
-RHI_API ERHIAccess RHIGetDefaultResourceState(EBufferUsageFlags InUsage, bool bInHasInitialData)
+ERHIAccess RHIGetDefaultResourceState(EBufferUsageFlags InUsage, bool bInHasInitialData)
 {
 	// Default reading state is different per buffer type
 	ERHIAccess DefaultReadingState = ERHIAccess::Unknown;
diff --git a/Engine/Source/Runtime/RHI/Private/RHIValidation.cpp b/Engine/Source/Runtime/RHI/Private/RHIValidation.cpp
--- a/Engine/Source/Runtime/RHI/Private/RHIValidation.cpp
+++ b/Engine/Source/Runtime/RHI/Private/RHIValidation.cpp
@@ -1685,7 +1685,7 @@ namespace RHIValidation
 		}
 	}
 
-	RHI_API EReplayStatus FOperation::Replay(ERHIPipeline Pipeline, bool& bAllowAllUAVsOverlap, FBreadcrumbStack& Breadcrumbs) const
+	EReplayStatus FOperation::Replay(ERHIPipeline Pipeline, bool& bAllowAllUAVsOverlap, FBreadcrumbStack& Breadcrumbs) const
 	{
 		switch (Type)
 		{
@@ -1940,7 +1940,7 @@ namespace RHIValidation
 	
 	FTracker::FOpQueueState FTracker::OpQueues[int32(ERHIPipeline::Num)] = {};
 
-	RHI_API void* CaptureBacktrace()
+	void* CaptureBacktrace()
 	{
 		// Back traces will leak. Don't leave this turned on.
 		const uint32 MaxDepth = 32;
diff --git a/Engine/Source/Runtime/RHI/Public/DataDrivenShaderPlatformInfo.h b/Engine/Source/Runtime/RHI/Public/DataDrivenShaderPlatformInfo.h
--- a/Engine/Source/Runtime/RHI/Public/DataDrivenShaderPlatformInfo.h
+++ b/Engine/Source/Runtime/RHI/Public/DataDrivenShaderPlatformInfo.h
@@ -1053,7 +1053,7 @@ inline bool RHISupportsAbsoluteVertexID(const FStaticShaderPlatform InShaderPlat
  *  Check GRHISupportsRayTracingShaders before using full ray tracing pipeline state objects.
  *  Check GRHISupportsInlineRayTracing before using inline ray tracing features in compute and other shaders.
  **/
-inline RHI_API bool RHISupportsRayTracing(const FStaticShaderPlatform Platform)
+inline bool RHISupportsRayTracing(const FStaticShaderPlatform Platform)
 {
 	return FDataDrivenShaderPlatformInfo::GetSupportsRayTracing(Platform);
 }
@@ -1061,7 +1061,7 @@ inline RHI_API bool RHISupportsRayTracing(const FStaticShaderPlatform Platform)
 /** Whether this platform can compile ray tracing shaders (regardless of project settings).
  *  To use at runtime, also check GRHISupportsRayTracing and r.RayTracing CVar (see IsRayTracingEnabled() helper).
  **/
-inline RHI_API bool RHISupportsRayTracingShaders(const FStaticShaderPlatform Platform)
+inline bool RHISupportsRayTracingShaders(const FStaticShaderPlatform Platform)
 {
 	return FDataDrivenShaderPlatformInfo::GetSupportsRayTracingShaders(Platform);
 }
@@ -1069,7 +1069,7 @@ inline RHI_API bool RHISupportsRayTracingShaders(const FStaticShaderPlatform Pla
 /** Whether this platform can compile shaders with inline ray tracing features.
  *  To use at runtime, also check GRHISupportsRayTracing and r.RayTracing CVar (see IsRayTracingEnabled() helper).
  **/
-inline RHI_API bool RHISupportsInlineRayTracing(const FStaticShaderPlatform Platform)
+inline bool RHISupportsInlineRayTracing(const FStaticShaderPlatform Platform)
 {
 	return FDataDrivenShaderPlatformInfo::GetSupportsInlineRayTracing(Platform);
 }
@@ -1077,7 +1077,7 @@ inline RHI_API bool RHISupportsInlineRayTracing(const FStaticShaderPlatform Plat
 /** Whether this platform can compile ray tracing callable shaders.
  *  To use at runtime, also check GRHISupportsRayTracing and r.RayTracing CVar (see IsRayTracingEnabled() helper).
  **/
-inline RHI_API bool RHISupportsRayTracingCallableShaders(const FStaticShaderPlatform Platform)
+inline bool RHISupportsRayTracingCallableShaders(const FStaticShaderPlatform Platform)
 {
 	return FDataDrivenShaderPlatformInfo::GetSupportsRayTracingCallableShaders(Platform);
 }
@@ -1106,7 +1106,7 @@ inline uint32 RHIMaxMeshShaderThreadGroupSize(const FStaticShaderPlatform Platfo
 /** Can this platform compile shaders that use shader model 6.0 wave intrinsics.
  *  To use such shaders at runtime, also check GRHISupportsWaveOperations.
  **/
-inline RHI_API bool RHISupportsWaveOperations(const FStaticShaderPlatform Platform)
+inline bool RHISupportsWaveOperations(const FStaticShaderPlatform Platform)
 {
 	return FDataDrivenShaderPlatformInfo::GetSupportsWaveOperations(Platform) != ERHIFeatureSupport::Unsupported;
 }
@@ -1118,7 +1118,7 @@ inline bool RHISupportsRenderTargetWriteMask(const FStaticShaderPlatform Platfor
 }
 
 /** True if the given shader platform supports overestimated conservative rasterization */
-inline RHI_API bool RHISupportsConservativeRasterization(const FStaticShaderPlatform Platform)
+inline bool RHISupportsConservativeRasterization(const FStaticShaderPlatform Platform)
 {
 	return FDataDrivenShaderPlatformInfo::GetSupportsConservativeRasterization(Platform);
 }
diff --git a/Engine/Source/Runtime/RHI/Public/RHIGlobals.h b/Engine/Source/Runtime/RHI/Public/RHIGlobals.h
--- a/Engine/Source/Runtime/RHI/Public/RHIGlobals.h
+++ b/Engine/Source/Runtime/RHI/Public/RHIGlobals.h
@@ -212,9 +212,12 @@ struct FRHIGlobals
 	/** True if the RHI has artifacts with atlased CSM depths. */
 	bool NeedsUnatlasedCSMDepthsWorkaround = false;
 
-	/** true if the RHI supports 3D textures */
+	/** Whether to initialize 3D textures using a bulk data (or through a mip update if false). */
 	bool SupportsTexture3D = true;
 
+	/** true if bulk data should be used with 3d textures */
+	bool UseTexture3DBulkData = false;
+
 	/** true if the RHI supports mobile multi-view */
 	bool SupportsMobileMultiView = false;
 
@@ -619,6 +622,7 @@ extern RHI_API FRHIGlobals GRHIGlobals;
 #define GSupportsSeparateRenderTargetBlendState                GRHIGlobals.SupportsSeparateRenderTargetBlendState
 #define GRHINeedsUnatlasedCSMDepthsWorkaround                  GRHIGlobals.NeedsUnatlasedCSMDepthsWorkaround
 #define GSupportsTexture3D                                     GRHIGlobals.SupportsTexture3D
+#define GUseTexture3DBulkDataRHI                               GRHIGlobals.UseTexture3DBulkData
 #define GSupportsMobileMultiView                               GRHIGlobals.SupportsMobileMultiView
 #define GSupportsImageExternal                                 GRHIGlobals.SupportsImageExternal
 #define GSupportsWideMRT                                       GRHIGlobals.SupportsWideMRT
diff --git a/Engine/Source/Runtime/RenderCore/Private/PipelineCacheUtilities.cpp b/Engine/Source/Runtime/RenderCore/Private/PipelineCacheUtilities.cpp
--- a/Engine/Source/Runtime/RenderCore/Private/PipelineCacheUtilities.cpp
+++ b/Engine/Source/Runtime/RenderCore/Private/PipelineCacheUtilities.cpp
@@ -458,8 +458,6 @@ bool UE::PipelineCacheUtilities::SaveStableKeysFile(const FStringView& Filename,
 	return true;
 }
 
-extern const RHI_API uint32 FPipelineCacheFileFormatCurrentVersion;
-
 bool UE::PipelineCacheUtilities::SaveStablePipelineCacheFile(const FString& OutputFilename, const TArray<FPermsPerPSO>& StableResults, const TArray<FStableShaderKeyAndValue>& StableShaderKeyIndexTable)
 {
 	TUniquePtr<FArchive> Archive(IFileManager::Get().CreateFileWriter(*OutputFilename));
diff --git a/Engine/Source/Runtime/VulkanRHI/Private/VulkanRHI.cpp b/Engine/Source/Runtime/VulkanRHI/Private/VulkanRHI.cpp
--- a/Engine/Source/Runtime/VulkanRHI/Private/VulkanRHI.cpp
+++ b/Engine/Source/Runtime/VulkanRHI/Private/VulkanRHI.cpp
@@ -29,8 +29,6 @@
 
 static_assert(sizeof(VkStructureType) == sizeof(int32), "ZeroVulkanStruct() assumes VkStructureType is int32!");
 
-extern RHI_API bool GUseTexture3DBulkDataRHI;
-
 #if NV_AFTERMATH
 bool GVulkanNVAftermathModuleLoaded = false;
 #endif
