Fix a non-unity build (FCommandLine usage).

diff --git a/Engine/Source/Runtime/RHI/Private/DataDrivenShaderPlatformInfo.cpp b/Engine/Source/Runtime/RHI/Private/DataDrivenShaderPlatformInfo.cpp
--- a/Engine/Source/Runtime/RHI/Private/DataDrivenShaderPlatformInfo.cpp
+++ b/Engine/Source/Runtime/RHI/Private/DataDrivenShaderPlatformInfo.cpp
@@ -4,6 +4,7 @@
 #include "RHI.h"
 #include "Misc/ConfigCacheIni.h"
 #include "Misc/DataDrivenPlatformInfoRegistry.h"
+#include "Misc/CommandLine.h"
 #include "RHIStrings.h"
 
 const FName LANGUAGE_D3D("D3D");
