RHI old binding path deprecation
- With everything moved to batched/bulk parameters, we can now deprecate the individual parameter setters on the RHI Command List.
- Deprecated functions now go straight to the batched parameters. While it will use more memory, it's still functional.
- Completely removing FRHIParameterBatcher since it was only intended to by used by FRHICommandList.

diff --git a/Engine/Source/Runtime/Engine/Private/GPUSkinCache.cpp b/Engine/Source/Runtime/Engine/Private/GPUSkinCache.cpp
--- a/Engine/Source/Runtime/Engine/Private/GPUSkinCache.cpp
+++ b/Engine/Source/Runtime/Engine/Private/GPUSkinCache.cpp
@@ -1388,9 +1388,6 @@ void FGPUSkinCache::DoDispatch(FRHICommandListImmediate& RHICmdList)
 	}
 
 	TransitionAllToReadable(RHICmdList, BuffersToTransitionToRead);
-
-	// Need to flush here since we're a scene-less process.
-	RHICmdList.FlushAllPendingComputeParameters();
 }
 
 void FGPUSkinCache::DoDispatch(FRHICommandListImmediate& RHICmdList, FGPUSkinCacheEntry* SkinCacheEntry, int32 Section, int32 RevisionNumber)
@@ -1444,9 +1441,6 @@ void FGPUSkinCache::DoDispatch(FRHICommandListImmediate& RHICmdList, FGPUSkinCac
 	SkinCacheEntry->UpdateVertexFactoryDeclaration(Section);
 
 	TransitionAllToReadable(RHICmdList, BuffersToTransitionToRead);
-
-	// Need to flush here since we're a scene-less process.
-	RHICmdList.FlushAllPendingComputeParameters();
 }
 
 bool FGPUSkinCache::ProcessEntry(
diff --git a/Engine/Source/Runtime/RHI/Private/RHICommandList.cpp b/Engine/Source/Runtime/RHI/Private/RHICommandList.cpp
--- a/Engine/Source/Runtime/RHI/Private/RHICommandList.cpp
+++ b/Engine/Source/Runtime/RHI/Private/RHICommandList.cpp
@@ -136,111 +136,12 @@ static TStatId GCurrentExecuteStat;
 static FCriticalSection GRHIThreadOnTasksCritical;
 static std::atomic<int32> GRHIThreadStallRequestCount;
 
-int32 GRHICommandParameterBatching = 1;
-static FAutoConsoleVariableRef CVarRHICommandParameterBatching(
-	TEXT("r.RHICmdParamBatching"),
-	GRHICommandParameterBatching,
-	TEXT("Controls if parameters can be batched in RHI command lists."),
-	ECVF_ReadOnly
-);
-
-FRHIParameterBatcher::FRHIParameterBatcher()
-{
-	bEnabled = GRHICommandParameterBatching != 0;
-}
-
-FRHIParameterBatcher::FRHIParameterBatcher(const FBoundShaderStateInput& InBoundShaderStateInput, FRHIComputeShader* InBoundComputeShaderRHI)
-{
-	AllBatchedShaders[SF_Vertex] = InBoundShaderStateInput.GetVertexShader();
-	AllBatchedShaders[SF_Mesh] = InBoundShaderStateInput.GetMeshShader();
-	AllBatchedShaders[SF_Amplification] = InBoundShaderStateInput.GetAmplificationShader();
-	AllBatchedShaders[SF_Pixel] = InBoundShaderStateInput.GetPixelShader();
-	AllBatchedShaders[SF_Geometry] = InBoundShaderStateInput.GetGeometryShader();
-	AllBatchedShaders[SF_Compute] = InBoundComputeShaderRHI;
-
-	bEnabled = GRHICommandParameterBatching != 0;
-}
-
-FRHIParameterBatcher::FRHIParameterBatcher(FRHIParameterBatcher&&) = default;
-FRHIParameterBatcher::~FRHIParameterBatcher()
-{
-	for (int32 Index = 0; Index < SF_NumStandardFrequencies; Index++)
-	{
-		checkSlow(AllBatchedShaderParameters[Index].HasParameters() == false);
-	}
-}
-
-void FRHIParameterBatcher::OnBoundShaderChanged(FRHICommandList& InCommandList, const FBoundShaderStateInput& InBoundShaderStateInput)
-{
-	// Make sure we flush all parameters before binding new shaders, even if we make a pipeline switch
-	PreDispatch(InCommandList);
-	PreDraw(InCommandList);
-
-	AllBatchedShaders[SF_Vertex]        = InBoundShaderStateInput.GetVertexShader();
-	AllBatchedShaders[SF_Mesh]          = InBoundShaderStateInput.GetMeshShader();
-	AllBatchedShaders[SF_Amplification] = InBoundShaderStateInput.GetAmplificationShader();
-	AllBatchedShaders[SF_Pixel]         = InBoundShaderStateInput.GetPixelShader();
-	AllBatchedShaders[SF_Geometry]      = InBoundShaderStateInput.GetGeometryShader();
-}
-
-void FRHIParameterBatcher::OnBoundShaderChanged(FRHIComputeCommandList& InCommandList, FRHIComputeShader* InBoundComputeShaderRHI)
-{
-	// Make sure we flush all parameters before binding new shaders, even if we make a pipeline switch
-	PreDraw((FRHICommandList&)InCommandList);
-	PreDispatch(InCommandList);
-
-	AllBatchedShaders[SF_Compute] = InBoundComputeShaderRHI;
-}
-
-void FRHIParameterBatcher::PreDispatch(FRHIComputeCommandList& InCommandList)
-{
-	InCommandList.SetShaderParametersFromInternalBatched(GetBatchedComputeShader(), AllBatchedShaderParameters[SF_Compute]);
-}
-
-void FRHIParameterBatcher::PreDraw(FRHICommandList& InCommandList)
-{
-	for (uint32 Index = 0; Index < SF_NumGraphicsFrequencies; Index++)
-	{
-		InCommandList.SetShaderParametersFromInternalBatched(GetBatchedGraphicsShader(Index), AllBatchedShaderParameters[Index]);
-	}
-}
-
-void FRHIParameterBatcher::FlushAllParameters(class FRHIComputeCommandList& InCommandList)
-{
-	PreDispatch(InCommandList);
-}
-
-void FRHIParameterBatcher::FlushAllParameters(class FRHICommandList& InCommandList)
-{
-	PreDraw(InCommandList);
-	PreDispatch(InCommandList);
-}
-
-void FRHIParameterBatcher::FlushPendingParameters(FRHIComputeCommandList& InCommandList, FRHIComputeShader* InShader)
-{
-	const EShaderFrequency Frequency = InShader->GetFrequency();
-	if (Frequency == SF_Compute)
-	{
-		InCommandList.SetShaderParametersFromInternalBatched(GetBatchedComputeShader(), AllBatchedShaderParameters[SF_Compute]);
-	}
-}
-
-void FRHIParameterBatcher::FlushPendingParameters(FRHICommandList& InCommandList, FRHIGraphicsShader* InShader)
-{
-	const EShaderFrequency Frequency = InShader->GetFrequency();
-	if (Frequency < SF_NumGraphicsFrequencies)
-	{
-		InCommandList.SetShaderParametersFromInternalBatched(GetBatchedGraphicsShader(Frequency), AllBatchedShaderParameters[Frequency]);
-	}
-}
-
 FRHICommandListBase::FRHICommandListBase(FRHIGPUMask InGPUMask, ERecordingThread InRecordingThread)
 	: FRHICommandListBase(FPersistentState(InGPUMask, InRecordingThread))
 {}
 
 FRHICommandListBase::FRHICommandListBase(FPersistentState&& InPersistentState)
-	: ParameterBatcher(InPersistentState.BoundShaderInput, InPersistentState.BoundComputeShaderRHI)
-	, DispatchEvent(FGraphEvent::CreateGraphEvent())
+	: DispatchEvent(FGraphEvent::CreateGraphEvent())
 	, PersistentState(MoveTemp(InPersistentState))
 {
 	DispatchEvent->SetDebugName(TEXT("FRHICommandListBase::DispatchEvent"));
@@ -271,7 +172,6 @@ FRHICommandListBase::FRHICommandListBase(FRHICommandListBase&& Other)
     , GraphicsContext (MoveTemp(Other.GraphicsContext))
     , ComputeContext  (MoveTemp(Other.ComputeContext))
     , Contexts        (MoveTemp(Other.Contexts))
-	, ParameterBatcher(MoveTemp(Other.ParameterBatcher))
 #if RHI_COUNT_COMMANDS
     , NumCommands     (MoveTemp(Other.NumCommands))
 #endif				  
@@ -371,16 +271,6 @@ void FRHICommandListBase::FinishRecording()
 {
 	checkf(!IsImmediate(), TEXT("Do not call FinishRecording() on the immediate RHI command list."));
 
-	if (IsGraphics())
-	{
-		// Make sure the batcher is cleared out. If there are pending graphics parameters, this cast is safe.
-		ParameterBatcher.FlushAllParameters(static_cast<FRHICommandList&>(*this));
-	}
-	else
-	{
-		ParameterBatcher.FlushAllParameters(static_cast<FRHIComputeCommandList&>(*this));
-	}
-
 	PersistentState.FenceCandidate->Fence = PersistentState.RHIThreadBufferLockFence;
 
 	// "Complete" the dispatch event. This unblocks waiting tasks but only when
@@ -447,15 +337,6 @@ ERHIPipeline FRHICommandListBase::SwitchPipeline(ERHIPipeline Pipeline)
 	Exchange(ActivePipeline, Pipeline);
 	if (ActivePipeline != Pipeline)
 	{
-		if (Pipeline == ERHIPipeline::Graphics)
-		{
-			ParameterBatcher.FlushAllParameters(static_cast<FRHICommandList&>(*this));
-		}
-		else if (Pipeline == ERHIPipeline::AsyncCompute)
-		{
-			ParameterBatcher.FlushAllParameters(static_cast<FRHIComputeCommandList&>(*this));
-		}
-
 		EnqueueLambda([NewPipeline = ActivePipeline](FRHICommandListBase& ExecutingCmdList)
 		{
 			ExecutingCmdList.ActivePipeline = NewPipeline;
@@ -821,8 +702,6 @@ void FRHICommandListImmediate::ExecuteAndReset()
 	// Always reset the immediate command list when we're done.
 	ON_SCOPE_EXIT { Reset(); };
 
-	ParameterBatcher.FlushAllParameters(*this);
-
 	//
 	// In bypass mode, the immediate command list will never contain recorded commands (since these were forwarded directly into the immediate RHI contexts).
 	// However, platforms that use explicit command list submission still need to be told to submit the enqueued work in the default contexts. Do that now.
diff --git a/Engine/Source/Runtime/RHI/Public/RHICommandList.h b/Engine/Source/Runtime/RHI/Public/RHICommandList.h
--- a/Engine/Source/Runtime/RHI/Public/RHICommandList.h
+++ b/Engine/Source/Runtime/RHI/Public/RHICommandList.h
@@ -803,8 +803,6 @@ protected:
 	// These are always set for the immediate command list, see InitializeImmediateContexts().
 	TRHIPipelineArray<IRHIComputeContext*> Contexts = {};
 
-	FRHIParameterBatcher ParameterBatcher;
-
 	FRHIBatchedShaderParameters ScratchShaderParameters;
 
 #if RHI_COUNT_COMMANDS
@@ -1153,110 +1151,6 @@ FRHICOMMAND_MACRO_TPL(TRHIShader, FRHICommandSetShaderParameters)
 	RHI_API void Execute(FRHICommandListBase& CmdList);
 };
 
-FRHICOMMAND_MACRO_TPL(TRHIShader, FRHICommandSetShaderParameter)
-{
-	TRHIShader* Shader;
-	const void* NewValue;
-	uint32 BufferIndex;
-	uint32 BaseIndex;
-	uint32 NumBytes;
-	FORCEINLINE_DEBUGGABLE FRHICommandSetShaderParameter(TRHIShader* InShader, uint32 InBufferIndex, uint32 InBaseIndex, uint32 InNumBytes, const void* InNewValue)
-		: Shader(InShader)
-		, NewValue(InNewValue)
-		, BufferIndex(InBufferIndex)
-		, BaseIndex(InBaseIndex)
-		, NumBytes(InNumBytes)
-	{
-	}
-	RHI_API void Execute(FRHICommandListBase& CmdList);
-};
-
-FRHICOMMAND_MACRO_TPL(TRHIShader, FRHICommandSetShaderUniformBuffer)
-{
-	TRHIShader* Shader;
-	uint32 BaseIndex;
-	FRHIUniformBuffer* UniformBuffer;
-	FORCEINLINE_DEBUGGABLE FRHICommandSetShaderUniformBuffer(TRHIShader* InShader, uint32 InBaseIndex, FRHIUniformBuffer* InUniformBuffer)
-		: Shader(InShader)
-		, BaseIndex(InBaseIndex)
-		, UniformBuffer(InUniformBuffer)
-	{
-	}
-	RHI_API void Execute(FRHICommandListBase& CmdList);
-};
-
-FRHICOMMAND_MACRO_TPL(TRHIShader, FRHICommandSetShaderTexture)
-{
-	TRHIShader* Shader;
-	uint32 TextureIndex;
-	FRHITexture* Texture;
-	FORCEINLINE_DEBUGGABLE FRHICommandSetShaderTexture(TRHIShader* InShader, uint32 InTextureIndex, FRHITexture* InTexture)
-		: Shader(InShader)
-		, TextureIndex(InTextureIndex)
-		, Texture(InTexture)
-	{
-	}
-	RHI_API void Execute(FRHICommandListBase& CmdList);
-};
-
-FRHICOMMAND_MACRO_TPL(TRHIShader, FRHICommandSetShaderResourceViewParameter)
-{
-	TRHIShader* Shader;
-	uint32 SamplerIndex;
-	FRHIShaderResourceView* SRV;
-	FORCEINLINE_DEBUGGABLE FRHICommandSetShaderResourceViewParameter(TRHIShader* InShader, uint32 InSamplerIndex, FRHIShaderResourceView* InSRV)
-		: Shader(InShader)
-		, SamplerIndex(InSamplerIndex)
-		, SRV(InSRV)
-	{
-	}
-	RHI_API void Execute(FRHICommandListBase& CmdList);
-};
-
-FRHICOMMAND_MACRO_TPL(TRHIShader, FRHICommandSetUAVParameter)
-{
-	TRHIShader* Shader;
-	uint32 UAVIndex;
-	FRHIUnorderedAccessView* UAV;
-	FORCEINLINE_DEBUGGABLE FRHICommandSetUAVParameter(TRHIShader* InShader, uint32 InUAVIndex, FRHIUnorderedAccessView* InUAV)
-		: Shader(InShader)
-		, UAVIndex(InUAVIndex)
-		, UAV(InUAV)
-	{
-	}
-	RHI_API void Execute(FRHICommandListBase& CmdList);
-};
-
-FRHICOMMAND_MACRO(FRHICommandSetUAVParameter_InitialCount)
-{
-	FRHIComputeShader* Shader;
-	uint32 UAVIndex;
-	FRHIUnorderedAccessView* UAV;
-	uint32 InitialCount;
-	FORCEINLINE_DEBUGGABLE FRHICommandSetUAVParameter_InitialCount(FRHIComputeShader* InShader, uint32 InUAVIndex, FRHIUnorderedAccessView* InUAV, uint32 InInitialCount)
-		: Shader(InShader)
-		, UAVIndex(InUAVIndex)
-		, UAV(InUAV)
-		, InitialCount(InInitialCount)
-	{
-	}
-	RHI_API void Execute(FRHICommandListBase& CmdList);
-};
-
-FRHICOMMAND_MACRO_TPL(TRHIShader, FRHICommandSetShaderSampler)
-{
-	TRHIShader* Shader;
-	uint32 SamplerIndex;
-	FRHISamplerState* Sampler;
-	FORCEINLINE_DEBUGGABLE FRHICommandSetShaderSampler(TRHIShader* InShader, uint32 InSamplerIndex, FRHISamplerState* InSampler)
-		: Shader(InShader)
-		, SamplerIndex(InSamplerIndex)
-		, Sampler(InSampler)
-	{
-	}
-	RHI_API void Execute(FRHICommandListBase& CmdList);
-};
-
 FRHICOMMAND_MACRO(FRHICommandDrawPrimitive)
 {
 	uint32 BaseVertexIndex;
@@ -2145,12 +2039,6 @@ FRHICOMMAND_MACRO(FRHICommandSetRayTracingBindings)
 #endif // RHI_RAYTRACING
 
 template<> RHI_API void FRHICommandSetShaderParameters           <FRHIComputeShader>::Execute(FRHICommandListBase& CmdList);
-template<> RHI_API void FRHICommandSetShaderParameter            <FRHIComputeShader>::Execute(FRHICommandListBase& CmdList);
-template<> RHI_API void FRHICommandSetShaderUniformBuffer        <FRHIComputeShader>::Execute(FRHICommandListBase& CmdList);
-template<> RHI_API void FRHICommandSetShaderTexture              <FRHIComputeShader>::Execute(FRHICommandListBase& CmdList);
-template<> RHI_API void FRHICommandSetShaderResourceViewParameter<FRHIComputeShader>::Execute(FRHICommandListBase& CmdList);
-template<> RHI_API void FRHICommandSetShaderSampler              <FRHIComputeShader>::Execute(FRHICommandListBase& CmdList);
-template<> RHI_API void FRHICommandSetUAVParameter               <FRHIComputeShader>::Execute(FRHICommandListBase& CmdList);
 
 extern RHI_API FRHIComputePipelineState* ExecuteSetComputePipelineState(FComputePipelineState* ComputePipelineState);
 extern RHI_API FRHIGraphicsPipelineState* ExecuteSetGraphicsPipelineState(class FGraphicsPipelineState* GraphicsPipelineState);
@@ -2161,24 +2049,13 @@ class RHI_API FRHIComputeCommandList : public FRHICommandListBase
 protected:
 	void OnBoundShaderChanged(FRHIComputeShader* InBoundComputeShaderRHI)
 	{
-		ParameterBatcher.OnBoundShaderChanged(*this, InBoundComputeShaderRHI);
 		PersistentState.BoundComputeShaderRHI = InBoundComputeShaderRHI;
 	}
 
-	void PreDispatch()
-	{
-		ParameterBatcher.PreDispatch(*this);
-	}
-
-	void FlushPendingComputeParameters(FRHIComputeShader* InShader)
-	{
-		ParameterBatcher.FlushPendingParameters(*this, InShader);
-	}
-
 public:
+	UE_DEPRECATED(5.3, "FlushAllPendingComputeParameters isn't needed with automatic batching removed.")
 	void FlushAllPendingComputeParameters()
 	{
-		PreDispatch();
 	}
 
 public:
@@ -2215,20 +2092,20 @@ public:
 		ALLOC_COMMAND(FRHICommandSetStaticUniformBuffers)(UniformBuffers);
 	}
 
+	UE_DEPRECATED(5.3, "FRHIBatchedShaderParameters and SetBatchedShaderParameters should be used instead of setting individual parameters.")
 	FORCEINLINE_DEBUGGABLE void SetShaderUniformBuffer(FRHIComputeShader* Shader, uint32 BaseIndex, FRHIUniformBuffer* UniformBuffer)
 	{
-		ValidateBoundShader(Shader);
-		if (Bypass())
-		{
-			GetComputeContext().RHISetShaderUniformBuffer(Shader, BaseIndex, UniformBuffer);
-			return;
-		}
-		ALLOC_COMMAND(FRHICommandSetShaderUniformBuffer<FRHIComputeShader>)(Shader, BaseIndex, UniformBuffer);
+		FRHIBatchedShaderParameters& BatchedParameters = GetScratchShaderParameters();
+		BatchedParameters.SetShaderUniformBuffer(BaseIndex, UniformBuffer);
+		SetBatchedShaderParameters(Shader, BatchedParameters);
 	}
 
+	UE_DEPRECATED(5.3, "FRHIBatchedShaderParameters and SetBatchedShaderParameters should be used instead of setting individual parameters.")
 	FORCEINLINE void SetShaderUniformBuffer(const FComputeShaderRHIRef& Shader, uint32 BaseIndex, FRHIUniformBuffer* UniformBuffer)
 	{
+PRAGMA_DISABLE_DEPRECATION_WARNINGS
 		SetShaderUniformBuffer(Shader.GetReference(), BaseIndex, UniformBuffer);
+PRAGMA_ENABLE_DEPRECATION_WARNINGS
 	}
 
 	FORCEINLINE_DEBUGGABLE void SetShaderParameters(
@@ -2237,16 +2114,10 @@ public:
 		, TConstArrayView<FRHIShaderParameter> InParameters
 		, TConstArrayView<FRHIShaderParameterResource> InResourceParameters
 		, TConstArrayView<FRHIShaderParameterResource> InBindlessParameters
-		, bool bSubmitBatchedParameters = true
 	)
 	{
 		ValidateBoundShader(InShader);
 
-		if (bSubmitBatchedParameters)
-		{
-			FlushPendingComputeParameters(InShader);
-		}
-
 		if (Bypass())
 		{
 			GetComputeContext().RHISetShaderParameters(InShader, InParametersData, InParameters, InResourceParameters, InBindlessParameters);
@@ -2271,89 +2142,46 @@ public:
 				InBatchedParameters.ParametersData,
 				InBatchedParameters.Parameters,
 				InBatchedParameters.ResourceParameters,
-				InBatchedParameters.BindlessParameters,
-				true /* bSubmitBatchedParameters */
-			);
-
-			InBatchedParameters.Reset();
-		}
-	}
-
-	FORCEINLINE_DEBUGGABLE void SetShaderParametersFromInternalBatched(FRHIComputeShader* InShader, FRHIBatchedShaderParameters& InBatchedParameters)
-	{
-		if (InBatchedParameters.HasParameters())
-		{
-			SetShaderParameters(
-				InShader,
-				InBatchedParameters.ParametersData,
-				InBatchedParameters.Parameters,
-				InBatchedParameters.ResourceParameters,
-				InBatchedParameters.BindlessParameters,
-				false /* bSubmitBatchedParameters */
+				InBatchedParameters.BindlessParameters
 			);
 
 			InBatchedParameters.Reset();
 		}
 	}
 
+	UE_DEPRECATED(5.3, "FRHIBatchedShaderParameters and SetBatchedShaderParameters should be used instead of setting individual parameters.")
 	FORCEINLINE_DEBUGGABLE void SetShaderParameter(FRHIComputeShader* Shader, uint32 BufferIndex, uint32 BaseIndex, uint32 NumBytes, const void* NewValue)
 	{
-		ValidateBoundShader(Shader);
-		if (ParameterBatcher.IsEnabled())
-		{
-			ParameterBatcher.SetShaderParameter(Shader, BufferIndex, BaseIndex, NumBytes, NewValue);
-			return;
-		}
-
-		if (Bypass())
-		{
-			GetComputeContext().RHISetShaderParameter(Shader, BufferIndex, BaseIndex, NumBytes, NewValue);
-			return;
-		}
-		void* UseValue = Alloc(NumBytes, 16);
-		FMemory::Memcpy(UseValue, NewValue, NumBytes);
-		ALLOC_COMMAND(FRHICommandSetShaderParameter<FRHIComputeShader>)(Shader, BufferIndex, BaseIndex, NumBytes, UseValue);
+		FRHIBatchedShaderParameters& BatchedParameters = GetScratchShaderParameters();
+		BatchedParameters.SetShaderParameter(BufferIndex, BaseIndex, NumBytes, NewValue);
+		SetBatchedShaderParameters(Shader, BatchedParameters);
 	}
 
+	UE_DEPRECATED(5.3, "FRHIBatchedShaderParameters and SetBatchedShaderParameters should be used instead of setting individual parameters.")
 	FORCEINLINE void SetShaderParameter(FComputeShaderRHIRef& Shader, uint32 BufferIndex, uint32 BaseIndex, uint32 NumBytes, const void* NewValue)
 	{
+PRAGMA_DISABLE_DEPRECATION_WARNINGS
 		SetShaderParameter(Shader.GetReference(), BufferIndex, BaseIndex, NumBytes, NewValue);
+PRAGMA_ENABLE_DEPRECATION_WARNINGS
 	}
 
+	UE_DEPRECATED(5.3, "FRHIBatchedShaderParameters and SetBatchedShaderParameters should be used instead of setting individual parameters.")
 	FORCEINLINE_DEBUGGABLE void SetShaderTexture(FRHIComputeShader* Shader, uint32 TextureIndex, FRHITexture* Texture)
 	{
-		ValidateBoundShader(Shader);
-		if (ParameterBatcher.IsEnabled())
-		{
-			ParameterBatcher.SetShaderTexture(Shader, TextureIndex, Texture);
-			return;
-		}
-
-		if (Bypass())
-		{
-			GetComputeContext().RHISetShaderTexture(Shader, TextureIndex, Texture);
-			return;
-		}
-		ALLOC_COMMAND(FRHICommandSetShaderTexture<FRHIComputeShader>)(Shader, TextureIndex, Texture);
+		FRHIBatchedShaderParameters& BatchedParameters = GetScratchShaderParameters();
+		BatchedParameters.SetShaderTexture(TextureIndex, Texture);
+		SetBatchedShaderParameters(Shader, BatchedParameters);
 	}
 
+	UE_DEPRECATED(5.3, "FRHIBatchedShaderParameters and SetBatchedShaderParameters should be used instead of setting individual parameters.")
 	FORCEINLINE_DEBUGGABLE void SetShaderResourceViewParameter(FRHIComputeShader* Shader, uint32 SamplerIndex, FRHIShaderResourceView* SRV)
 	{
-		ValidateBoundShader(Shader);
-		if (ParameterBatcher.IsEnabled())
-		{
-			ParameterBatcher.SetShaderResourceViewParameter(Shader, SamplerIndex, SRV);
-			return;
-		}
-
-		if (Bypass())
-		{
-			GetComputeContext().RHISetShaderResourceViewParameter(Shader, SamplerIndex, SRV);
-			return;
-		}
-		ALLOC_COMMAND(FRHICommandSetShaderResourceViewParameter<FRHIComputeShader>)(Shader, SamplerIndex, SRV);
+		FRHIBatchedShaderParameters& BatchedParameters = GetScratchShaderParameters();
+		BatchedParameters.SetShaderResourceViewParameter(SamplerIndex, SRV);
+		SetBatchedShaderParameters(Shader, BatchedParameters);
 	}
 
+	UE_DEPRECATED(5.3, "FRHIBatchedShaderParameters and SetBatchedShaderParameters should be used instead of setting individual parameters.")
 	FORCEINLINE_DEBUGGABLE void SetShaderSampler(FRHIComputeShader* Shader, uint32 SamplerIndex, FRHISamplerState* State)
 	{
 		// Immutable samplers can't be set dynamically
@@ -2363,88 +2191,48 @@ public:
 			return;
 		}
 
-		if (ParameterBatcher.IsEnabled())
-		{
-			ParameterBatcher.SetShaderSampler(Shader, SamplerIndex, State);
-			return;
-		}
-
-		if (Bypass())
-		{
-			GetComputeContext().RHISetShaderSampler(Shader, SamplerIndex, State);
-			return;
-		}
-		ALLOC_COMMAND(FRHICommandSetShaderSampler<FRHIComputeShader>)(Shader, SamplerIndex, State);
+		FRHIBatchedShaderParameters& BatchedParameters = GetScratchShaderParameters();
+		BatchedParameters.SetShaderSampler(SamplerIndex, State);
+		SetBatchedShaderParameters(Shader, BatchedParameters);
 	}
 
+	UE_DEPRECATED(5.3, "FRHIBatchedShaderParameters and SetBatchedShaderParameters should be used instead of setting individual parameters.")
 	FORCEINLINE_DEBUGGABLE void SetUAVParameter(FRHIComputeShader* Shader, uint32 UAVIndex, FRHIUnorderedAccessView* UAV)
 	{
-		ValidateBoundShader(Shader);
-
-		if (ParameterBatcher.IsEnabled())
-		{
-			ParameterBatcher.SetUAVParameter(Shader, UAVIndex, UAV);
-			return;
-		}
-
-		if (Bypass())
-		{
-			GetComputeContext().RHISetUAVParameter(Shader, UAVIndex, UAV);
-			return;
-		}
-		ALLOC_COMMAND(FRHICommandSetUAVParameter<FRHIComputeShader>)(Shader, UAVIndex, UAV);
+		FRHIBatchedShaderParameters& BatchedParameters = GetScratchShaderParameters();
+		BatchedParameters.SetUAVParameter(UAVIndex, UAV);
+		SetBatchedShaderParameters(Shader, BatchedParameters);
 	}
 
+	UE_DEPRECATED(5.3, "FRHIBatchedShaderParameters and SetBatchedShaderParameters should be used instead of setting individual parameters.")
 	FORCEINLINE_DEBUGGABLE void SetUAVParameter(FRHIComputeShader* Shader, uint32 UAVIndex, FRHIUnorderedAccessView* UAV, uint32 InitialCount)
 	{
-		ValidateBoundShader(Shader);
+		checkNoEntry(); // @todo: support append/consume buffers
 
-		if (ParameterBatcher.IsEnabled())
-		{
-			ParameterBatcher.SetUAVParameter(Shader, UAVIndex, UAV, InitialCount);
-			return;
-		}
-
-		if (Bypass())
-		{
-			GetComputeContext().RHISetUAVParameter(Shader, UAVIndex, UAV, InitialCount);
-			return;
-		}
-		ALLOC_COMMAND(FRHICommandSetUAVParameter_InitialCount)(Shader, UAVIndex, UAV, InitialCount);
+		FRHIBatchedShaderParameters& BatchedParameters = GetScratchShaderParameters();
+		BatchedParameters.SetUAVParameter(UAVIndex, UAV);
+		SetBatchedShaderParameters(Shader, BatchedParameters);
 	}
 
+	UE_DEPRECATED(5.3, "FRHIBatchedShaderParameters and SetBatchedShaderParameters should be used instead of setting individual parameters.")
 	FORCEINLINE_DEBUGGABLE void SetBindlessTexture(FRHIComputeShader* Shader, uint32 Index, FRHITexture* Texture)
 	{
-		ValidateBoundShader(Shader);
-
-		if (ParameterBatcher.IsEnabled())
-		{
-			ParameterBatcher.SetBindlessTexture(Shader, Index, Texture);
-			return;
-		}
-
-		FRHIShaderParameterResource Resource(Texture, (uint16)Index);
-		SetShaderParameters(Shader, {}, {}, {}, MakeArrayView(&Resource, 1));
+		FRHIBatchedShaderParameters& BatchedParameters = GetScratchShaderParameters();
+		BatchedParameters.SetBindlessTexture(Index, Texture);
+		SetBatchedShaderParameters(Shader, BatchedParameters);
 	}
 
+	UE_DEPRECATED(5.3, "FRHIBatchedShaderParameters and SetBatchedShaderParameters should be used instead of setting individual parameters.")
 	FORCEINLINE_DEBUGGABLE void SetBindlessResourceView(FRHIComputeShader* Shader, uint32 Index, FRHIShaderResourceView* SRV)
 	{
-		ValidateBoundShader(Shader);
-
-		if (ParameterBatcher.IsEnabled())
-		{
-			ParameterBatcher.SetBindlessResourceView(Shader, Index, SRV);
-			return;
-		}
-
-		FRHIShaderParameterResource Resource(SRV, (uint16)Index);
-		SetShaderParameters(Shader, {}, {}, {}, MakeArrayView(&Resource, 1));
+		FRHIBatchedShaderParameters& BatchedParameters = GetScratchShaderParameters();
+		BatchedParameters.SetBindlessResourceView(Index, SRV);
+		SetBatchedShaderParameters(Shader, BatchedParameters);
 	}
 
+	UE_DEPRECATED(5.3, "FRHIBatchedShaderParameters and SetBatchedShaderParameters should be used instead of setting individual parameters.")
 	FORCEINLINE_DEBUGGABLE void SetBindlessSampler(FRHIComputeShader* Shader, uint32 Index, FRHISamplerState* State)
 	{
-		ValidateBoundShader(Shader);
-
 		// Immutable samplers can't be set dynamically
 		check(!State->IsImmutable());
 		if (State->IsImmutable())
@@ -2452,29 +2240,17 @@ public:
 			return;
 		}
 
-		// Immutable samplers can't be set dynamically
-		if (ParameterBatcher.IsEnabled())
-		{
-			ParameterBatcher.SetBindlessSampler(Shader, Index, State);
-			return;
-		}
-
-		FRHIShaderParameterResource Resource(State, (uint16)Index);
-		SetShaderParameters(Shader, {}, {}, {}, MakeArrayView(&Resource, 1));
+		FRHIBatchedShaderParameters& BatchedParameters = GetScratchShaderParameters();
+		BatchedParameters.SetBindlessSampler(Index, State);
+		SetBatchedShaderParameters(Shader, BatchedParameters);
 	}
 
+	UE_DEPRECATED(5.3, "FRHIBatchedShaderParameters and SetBatchedShaderParameters should be used instead of setting individual parameters.")
 	FORCEINLINE_DEBUGGABLE void SetBindlessUAV(FRHIComputeShader* Shader, uint32 Index, FRHIUnorderedAccessView* UAV)
 	{
-		ValidateBoundShader(Shader);
-
-		if (ParameterBatcher.IsEnabled())
-		{
-			ParameterBatcher.SetBindlessUAV(Shader, Index, UAV);
-			return;
-		}
-
-		FRHIShaderParameterResource Resource(UAV, (uint16)Index);
-		SetShaderParameters(Shader, {}, {}, {}, MakeArrayView(&Resource, 1));
+		FRHIBatchedShaderParameters& BatchedParameters = GetScratchShaderParameters();
+		BatchedParameters.SetBindlessUAV(Index, UAV);
+		SetBatchedShaderParameters(Shader, BatchedParameters);
 	}
 
 	FORCEINLINE_DEBUGGABLE void SetComputePipelineState(FComputePipelineState* ComputePipelineState, FRHIComputeShader* ComputeShader)
@@ -2501,7 +2277,6 @@ public:
 
 	FORCEINLINE_DEBUGGABLE void DispatchComputeShader(uint32 ThreadGroupCountX, uint32 ThreadGroupCountY, uint32 ThreadGroupCountZ)
 	{
-		PreDispatch();
 		if (Bypass())
 		{
 			GetComputeContext().RHIDispatchComputeShader(ThreadGroupCountX, ThreadGroupCountY, ThreadGroupCountZ);
@@ -2512,7 +2287,6 @@ public:
 
 	FORCEINLINE_DEBUGGABLE void DispatchIndirectComputeShader(FRHIBuffer* ArgumentBuffer, uint32 ArgumentOffset)
 	{
-		PreDispatch();
 		if (Bypass())
 		{
 			GetComputeContext().RHIDispatchIndirectComputeShader(ArgumentBuffer, ArgumentOffset);
@@ -2523,7 +2297,6 @@ public:
 
 	FORCEINLINE_DEBUGGABLE void ClearUAVFloat(FRHIUnorderedAccessView* UnorderedAccessViewRHI, const FVector4f& Values)
 	{
-		PreDispatch();
 		if (Bypass())
 		{
 			GetComputeContext().RHIClearUAVFloat(UnorderedAccessViewRHI, Values);
@@ -2534,7 +2307,6 @@ public:
 
 	FORCEINLINE_DEBUGGABLE void ClearUAVUint(FRHIUnorderedAccessView* UnorderedAccessViewRHI, const FUintVector4& Values)
 	{
-		PreDispatch();
 		if (Bypass())
 		{
 			GetComputeContext().RHIClearUAVUint(UnorderedAccessViewRHI, Values);
@@ -2940,12 +2712,6 @@ public:
 };
 
 template<> RHI_API void FRHICommandSetShaderParameters           <FRHIGraphicsShader>::Execute(FRHICommandListBase& CmdList);
-template<> RHI_API void FRHICommandSetShaderParameter            <FRHIGraphicsShader>::Execute(FRHICommandListBase& CmdList);
-template<> RHI_API void FRHICommandSetShaderUniformBuffer        <FRHIGraphicsShader>::Execute(FRHICommandListBase& CmdList);
-template<> RHI_API void FRHICommandSetShaderTexture              <FRHIGraphicsShader>::Execute(FRHICommandListBase& CmdList);
-template<> RHI_API void FRHICommandSetShaderResourceViewParameter<FRHIGraphicsShader>::Execute(FRHICommandListBase& CmdList);
-template<> RHI_API void FRHICommandSetShaderSampler              <FRHIGraphicsShader>::Execute(FRHICommandListBase& CmdList);
-template<> RHI_API void FRHICommandSetUAVParameter               <FRHIPixelShader   >::Execute(FRHICommandListBase& CmdList);
 
 class RHI_API FRHICommandList : public FRHIComputeCommandList
 {
@@ -2954,20 +2720,9 @@ protected:
 
 	void OnBoundShaderChanged(const FBoundShaderStateInput& InBoundShaderStateInput)
 	{
-		ParameterBatcher.OnBoundShaderChanged(*this, InBoundShaderStateInput);
 		PersistentState.BoundShaderInput = InBoundShaderStateInput;
 	}
 
-	void PreDraw()
-	{
-		ParameterBatcher.PreDraw(*this);
-	}
-
-	void FlushPendingGraphicsParameters(FRHIGraphicsShader* InShader)
-	{
-		ParameterBatcher.FlushPendingParameters(*this, InShader);
-	}
-
 public:
 	FRHICommandList(FRHIGPUMask GPUMask = FRHIGPUMask::All(), ERecordingThread InRecordingThread = ERecordingThread::Render)
 		: FRHIComputeCommandList(GPUMask, InRecordingThread)
@@ -3010,22 +2765,21 @@ public:
 
 	using FRHIComputeCommandList::SetShaderUniformBuffer;
 
+	UE_DEPRECATED(5.3, "FRHIBatchedShaderParameters and SetBatchedShaderParameters should be used instead of setting individual parameters.")
 	FORCEINLINE_DEBUGGABLE void SetShaderUniformBuffer(FRHIGraphicsShader* Shader, uint32 BaseIndex, FRHIUniformBuffer* UniformBuffer)
 	{
-		//check(IsOutsideRenderPass());
-		ValidateBoundShader(Shader);
-		if (Bypass())
-		{
-			GetContext().RHISetShaderUniformBuffer(Shader, BaseIndex, UniformBuffer);
-			return;
-		}
-		ALLOC_COMMAND(FRHICommandSetShaderUniformBuffer<FRHIGraphicsShader>)(Shader, BaseIndex, UniformBuffer);
+		FRHIBatchedShaderParameters& BatchedParameters = GetScratchShaderParameters();
+		BatchedParameters.SetShaderUniformBuffer(BaseIndex, UniformBuffer);
+		SetBatchedShaderParameters(Shader, BatchedParameters);
 	}
 
 	template <typename TShaderRHI>
+	UE_DEPRECATED(5.3, "FRHIBatchedShaderParameters and SetBatchedShaderParameters should be used instead of setting individual parameters.")
 	FORCEINLINE void SetShaderUniformBuffer(const TRefCountPtr<TShaderRHI>& Shader, uint32 BaseIndex, FRHIUniformBuffer* UniformBuffer)
 	{
+PRAGMA_DISABLE_DEPRECATION_WARNINGS
 		SetShaderUniformBuffer(Shader.GetReference(), BaseIndex, UniformBuffer);
+PRAGMA_ENABLE_DEPRECATION_WARNINGS
 	}
 
 	using FRHIComputeCommandList::SetShaderParameters;
@@ -3036,16 +2790,10 @@ public:
 		, TConstArrayView<FRHIShaderParameter> InParameters
 		, TConstArrayView<FRHIShaderParameterResource> InResourceParameters
 		, TConstArrayView<FRHIShaderParameterResource> InBindlessParameters
-		, bool bSubmitBatchedParameters = true
 	)
 	{
 		ValidateBoundShader(InShader);
 
-		if (bSubmitBatchedParameters)
-		{
-			FlushPendingGraphicsParameters(InShader);
-		}
-
 		if (Bypass())
 		{
 			GetContext().RHISetShaderParameters(InShader, InParametersData, InParameters, InResourceParameters, InBindlessParameters);
@@ -3072,27 +2820,7 @@ public:
 				InBatchedParameters.ParametersData,
 				InBatchedParameters.Parameters,
 				InBatchedParameters.ResourceParameters,
-				InBatchedParameters.BindlessParameters,
-				true /* bSubmitBatchedParameters */
-			);
-
-			InBatchedParameters.Reset();
-		}
-	}
-
-	using FRHIComputeCommandList::SetShaderParametersFromInternalBatched;
-
-	FORCEINLINE_DEBUGGABLE void SetShaderParametersFromInternalBatched(FRHIGraphicsShader* InShader, FRHIBatchedShaderParameters& InBatchedParameters)
-	{
-		if (InBatchedParameters.HasParameters())
-		{
-			SetShaderParameters(
-				InShader,
-				InBatchedParameters.ParametersData,
-				InBatchedParameters.Parameters,
-				InBatchedParameters.ResourceParameters,
-				InBatchedParameters.BindlessParameters,
-				false /* bSubmitBatchedParameters */
+				InBatchedParameters.BindlessParameters
 			);
 
 			InBatchedParameters.Reset();
@@ -3101,94 +2829,66 @@ public:
 
 	using FRHIComputeCommandList::SetShaderParameter;
 
+	UE_DEPRECATED(5.3, "FRHIBatchedShaderParameters and SetBatchedShaderParameters should be used instead of setting individual parameters.")
 	FORCEINLINE_DEBUGGABLE void SetShaderParameter(FRHIGraphicsShader* Shader, uint32 BufferIndex, uint32 BaseIndex, uint32 NumBytes, const void* NewValue)
 	{
-		//check(IsOutsideRenderPass());
-		ValidateBoundShader(Shader);
-
-		if (ParameterBatcher.IsEnabled())
-		{
-			ParameterBatcher.SetShaderParameter(Shader, BufferIndex, BaseIndex, NumBytes, NewValue);
-			return;
-		}
-
-		if (Bypass())
-		{
-			GetContext().RHISetShaderParameter(Shader, BufferIndex, BaseIndex, NumBytes, NewValue);
-			return;
-		}
-		void* UseValue = Alloc(NumBytes, 16);
-		FMemory::Memcpy(UseValue, NewValue, NumBytes);
-		ALLOC_COMMAND(FRHICommandSetShaderParameter<FRHIGraphicsShader>)(Shader, BufferIndex, BaseIndex, NumBytes, UseValue);
+		FRHIBatchedShaderParameters& BatchedParameters = GetScratchShaderParameters();
+		BatchedParameters.SetShaderParameter(BufferIndex, BaseIndex, NumBytes, NewValue);
+		SetBatchedShaderParameters(Shader, BatchedParameters);
 	}
 
 	template <typename TShaderRHI>
+	UE_DEPRECATED(5.3, "FRHIBatchedShaderParameters and SetBatchedShaderParameters should be used instead of setting individual parameters.")
 	FORCEINLINE void SetShaderParameter(const TRefCountPtr<TShaderRHI>& Shader, uint32 BufferIndex, uint32 BaseIndex, uint32 NumBytes, const void* NewValue)
 	{
+PRAGMA_DISABLE_DEPRECATION_WARNINGS
 		SetShaderParameter(Shader.GetReference(), BufferIndex, BaseIndex, NumBytes, NewValue);
+PRAGMA_ENABLE_DEPRECATION_WARNINGS
 	}
 
 	using FRHIComputeCommandList::SetShaderTexture;
 
+	UE_DEPRECATED(5.3, "FRHIBatchedShaderParameters and SetBatchedShaderParameters should be used instead of setting individual parameters.")
 	FORCEINLINE_DEBUGGABLE void SetShaderTexture(FRHIGraphicsShader* Shader, uint32 TextureIndex, FRHITexture* Texture)
 	{
-		//check(IsOutsideRenderPass());
-		ValidateBoundShader(Shader);
-
-		if (ParameterBatcher.IsEnabled())
-		{
-			ParameterBatcher.SetShaderTexture(Shader, TextureIndex, Texture);
-			return;
-		}
-
-		if (Bypass())
-		{
-			GetContext().RHISetShaderTexture(Shader, TextureIndex, Texture);
-			return;
-		}
-		ALLOC_COMMAND(FRHICommandSetShaderTexture<FRHIGraphicsShader>)(Shader, TextureIndex, Texture);
+		FRHIBatchedShaderParameters& BatchedParameters = GetScratchShaderParameters();
+		BatchedParameters.SetShaderTexture(TextureIndex, Texture);
+		SetBatchedShaderParameters(Shader, BatchedParameters);
 	}
 
 	template <typename TShaderRHI>
+	UE_DEPRECATED(5.3, "FRHIBatchedShaderParameters and SetBatchedShaderParameters should be used instead of setting individual parameters.")
 	FORCEINLINE_DEBUGGABLE void SetShaderTexture(const TRefCountPtr<TShaderRHI>& Shader, uint32 TextureIndex, FRHITexture* Texture)
 	{
+PRAGMA_DISABLE_DEPRECATION_WARNINGS
 		SetShaderTexture(Shader.GetReference(), TextureIndex, Texture);
+PRAGMA_ENABLE_DEPRECATION_WARNINGS
 	}
 
 	using FRHIComputeCommandList::SetShaderResourceViewParameter;
 
+	UE_DEPRECATED(5.3, "FRHIBatchedShaderParameters and SetBatchedShaderParameters should be used instead of setting individual parameters.")
 	FORCEINLINE_DEBUGGABLE void SetShaderResourceViewParameter(FRHIGraphicsShader* Shader, uint32 SamplerIndex, FRHIShaderResourceView* SRV)
 	{
-		//check(IsOutsideRenderPass());
-		ValidateBoundShader(Shader);
-
-		if (ParameterBatcher.IsEnabled())
-		{
-			ParameterBatcher.SetShaderResourceViewParameter(Shader, SamplerIndex, SRV);
-			return;
-		}
-
-		if (Bypass())
-		{
-			GetContext().RHISetShaderResourceViewParameter(Shader, SamplerIndex, SRV);
-			return;
-		}
-		ALLOC_COMMAND(FRHICommandSetShaderResourceViewParameter<FRHIGraphicsShader>)(Shader, SamplerIndex, SRV);
+		FRHIBatchedShaderParameters& BatchedParameters = GetScratchShaderParameters();
+		BatchedParameters.SetShaderResourceViewParameter(SamplerIndex, SRV);
+		SetBatchedShaderParameters(Shader, BatchedParameters);
 	}
 
 	template <typename TShaderRHI>
+	UE_DEPRECATED(5.3, "FRHIBatchedShaderParameters and SetBatchedShaderParameters should be used instead of setting individual parameters.")
 	FORCEINLINE_DEBUGGABLE void SetShaderResourceViewParameter(const TRefCountPtr<TShaderRHI>& Shader, uint32 SamplerIndex, FRHIShaderResourceView* SRV)
 	{
+PRAGMA_DISABLE_DEPRECATION_WARNINGS
 		SetShaderResourceViewParameter(Shader.GetReference(), SamplerIndex, SRV);
+PRAGMA_ENABLE_DEPRECATION_WARNINGS
 	}
 
 	using FRHIComputeCommandList::SetShaderSampler;
 
+	UE_DEPRECATED(5.3, "FRHIBatchedShaderParameters and SetBatchedShaderParameters should be used instead of setting individual parameters.")
 	FORCEINLINE_DEBUGGABLE void SetShaderSampler(FRHIGraphicsShader* Shader, uint32 SamplerIndex, FRHISamplerState* State)
 	{
-		//check(IsOutsideRenderPass());
-		ValidateBoundShader(Shader);
-		
 		// Immutable samplers can't be set dynamically
 		check(!State->IsImmutable());
 		if (State->IsImmutable())
@@ -3196,89 +2896,63 @@ public:
 			return;
 		}
 
-		if (ParameterBatcher.IsEnabled())
-		{
-			ParameterBatcher.SetShaderSampler(Shader, SamplerIndex, State);
-			return;
-		}
-
-		if (Bypass())
-		{
-			GetContext().RHISetShaderSampler(Shader, SamplerIndex, State);
-			return;
-		}
-		ALLOC_COMMAND(FRHICommandSetShaderSampler<FRHIGraphicsShader>)(Shader, SamplerIndex, State);
+		FRHIBatchedShaderParameters& BatchedParameters = GetScratchShaderParameters();
+		BatchedParameters.SetShaderSampler(SamplerIndex, State);
+		SetBatchedShaderParameters(Shader, BatchedParameters);
 	}
 
 	template <typename TShaderRHI>
+	UE_DEPRECATED(5.3, "FRHIBatchedShaderParameters and SetBatchedShaderParameters should be used instead of setting individual parameters.")
 	FORCEINLINE_DEBUGGABLE void SetShaderSampler(const TRefCountPtr<TShaderRHI>& Shader, uint32 SamplerIndex, FRHISamplerState* State)
 	{
+PRAGMA_DISABLE_DEPRECATION_WARNINGS
 		SetShaderSampler(Shader.GetReference(), SamplerIndex, State);
+PRAGMA_ENABLE_DEPRECATION_WARNINGS
 	}
 
 	using FRHIComputeCommandList::SetUAVParameter;
 
+	UE_DEPRECATED(5.3, "FRHIBatchedShaderParameters and SetBatchedShaderParameters should be used instead of setting individual parameters.")
 	FORCEINLINE_DEBUGGABLE void SetUAVParameter(FRHIPixelShader* Shader, uint32 UAVIndex, FRHIUnorderedAccessView* UAV)
 	{
-		ValidateBoundShader(Shader);
-
-		if (ParameterBatcher.IsEnabled())
-		{
-			ParameterBatcher.SetUAVParameter(Shader, UAVIndex, UAV);
-			return;
-		}
-
-		if (Bypass())
-		{
-			GetContext().RHISetUAVParameter(Shader, UAVIndex, UAV);
-			return;
-		}
-		ALLOC_COMMAND(FRHICommandSetUAVParameter<FRHIPixelShader>)(Shader, UAVIndex, UAV);
+		FRHIBatchedShaderParameters& BatchedParameters = GetScratchShaderParameters();
+		BatchedParameters.SetUAVParameter(UAVIndex, UAV);
+		SetBatchedShaderParameters(Shader, BatchedParameters);
 	}
 
+	UE_DEPRECATED(5.3, "FRHIBatchedShaderParameters and SetBatchedShaderParameters should be used instead of setting individual parameters.")
 	FORCEINLINE_DEBUGGABLE void SetUAVParameter(const TRefCountPtr<FRHIPixelShader>& Shader, uint32 UAVIndex, FRHIUnorderedAccessView* UAV)
 	{
+PRAGMA_DISABLE_DEPRECATION_WARNINGS
 		SetUAVParameter(Shader.GetReference(), UAVIndex, UAV);
+PRAGMA_ENABLE_DEPRECATION_WARNINGS
 	}
 
 	using FRHIComputeCommandList::SetBindlessTexture;
 
+	UE_DEPRECATED(5.3, "FRHIBatchedShaderParameters and SetBatchedShaderParameters should be used instead of setting individual parameters.")
 	FORCEINLINE_DEBUGGABLE void SetBindlessTexture(FRHIGraphicsShader* Shader, uint32 Index, FRHITexture* Texture)
 	{
-		ValidateBoundShader(Shader);
-
-		if (ParameterBatcher.IsEnabled())
-		{
-			ParameterBatcher.SetBindlessTexture(Shader, Index, Texture);
-			return;
-		}
-
-		FRHIShaderParameterResource Resource(Texture, (uint16)Index);
-		SetShaderParameters(Shader, {}, {}, {}, MakeArrayView(&Resource, 1));
+		FRHIBatchedShaderParameters& BatchedParameters = GetScratchShaderParameters();
+		BatchedParameters.SetBindlessTexture(Index, Texture);
+		SetBatchedShaderParameters(Shader, BatchedParameters);
 	}
 
 	using FRHIComputeCommandList::SetBindlessResourceView;
 
+	UE_DEPRECATED(5.3, "FRHIBatchedShaderParameters and SetBatchedShaderParameters should be used instead of setting individual parameters.")
 	FORCEINLINE_DEBUGGABLE void SetBindlessResourceView(FRHIGraphicsShader* Shader, uint32 Index, FRHIShaderResourceView* SRV)
 	{
-		ValidateBoundShader(Shader);
-
-		if (ParameterBatcher.IsEnabled())
-		{
-			ParameterBatcher.SetBindlessResourceView(Shader, Index, SRV);
-			return;
-		}
-
-		FRHIShaderParameterResource Resource(SRV, (uint16)Index);
-		SetShaderParameters(Shader, {}, {}, {}, MakeArrayView(&Resource, 1));
+		FRHIBatchedShaderParameters& BatchedParameters = GetScratchShaderParameters();
+		BatchedParameters.SetBindlessResourceView(Index, SRV);
+		SetBatchedShaderParameters(Shader, BatchedParameters);
 	}
 
 	using FRHIComputeCommandList::SetBindlessSampler;
 
+	UE_DEPRECATED(5.3, "FRHIBatchedShaderParameters and SetBatchedShaderParameters should be used instead of setting individual parameters.")
 	FORCEINLINE_DEBUGGABLE void SetBindlessSampler(FRHIGraphicsShader* Shader, uint32 Index, FRHISamplerState* State)
 	{
-		ValidateBoundShader(Shader);
-
 		// Immutable samplers can't be set dynamically
 		check(!State->IsImmutable());
 		if (State->IsImmutable())
@@ -3286,30 +2960,19 @@ public:
 			return;
 		}
 
-		if (ParameterBatcher.IsEnabled())
-		{
-			ParameterBatcher.SetBindlessSampler(Shader, Index, State);
-			return;
-		}
-
-		FRHIShaderParameterResource Resource(State, (uint16)Index);
-		SetShaderParameters(Shader, {}, {}, {}, MakeArrayView(&Resource, 1));
+		FRHIBatchedShaderParameters& BatchedParameters = GetScratchShaderParameters();
+		BatchedParameters.SetBindlessSampler(Index, State);
+		SetBatchedShaderParameters(Shader, BatchedParameters);
 	}
 
 	using FRHIComputeCommandList::SetBindlessUAV;
 
+	UE_DEPRECATED(5.3, "FRHIBatchedShaderParameters and SetBatchedShaderParameters should be used instead of setting individual parameters.")
 	FORCEINLINE_DEBUGGABLE void SetBindlessUAV(FRHIPixelShader* Shader, uint32 Index, FRHIUnorderedAccessView* UAV)
 	{
-		ValidateBoundShader(Shader);
-
-		if (ParameterBatcher.IsEnabled())
-		{
-			ParameterBatcher.SetBindlessUAV(Shader, Index, UAV);
-			return;
-		}
-
-		FRHIShaderParameterResource Resource(UAV, (uint16)Index);
-		SetShaderParameters(Shader, {}, {}, {}, MakeArrayView(&Resource, 1));
+		FRHIBatchedShaderParameters& BatchedParameters = GetScratchShaderParameters();
+		BatchedParameters.SetBindlessUAV(Index, UAV);
+		SetBatchedShaderParameters(Shader, BatchedParameters);
 	}
 
 	FORCEINLINE_DEBUGGABLE void SetBlendFactor(const FLinearColor& BlendFactor = FLinearColor::White)
@@ -3326,7 +2989,6 @@ public:
 	FORCEINLINE_DEBUGGABLE void DrawPrimitive(uint32 BaseVertexIndex, uint32 NumPrimitives, uint32 NumInstances)
 	{
 		//check(IsOutsideRenderPass());
-		PreDraw();
 		if (Bypass())
 		{
 			GetContext().RHIDrawPrimitive(BaseVertexIndex, NumPrimitives, NumInstances);
@@ -3338,7 +3000,6 @@ public:
 	FORCEINLINE_DEBUGGABLE void DrawIndexedPrimitive(FRHIBuffer* IndexBuffer, int32 BaseVertexIndex, uint32 FirstInstance, uint32 NumVertices, uint32 StartIndex, uint32 NumPrimitives, uint32 NumInstances)
 	{
 		//check(IsOutsideRenderPass());
-		PreDraw();
 		if (Bypass())
 		{
 			GetContext().RHIDrawIndexedPrimitive(IndexBuffer, BaseVertexIndex, FirstInstance, NumVertices, StartIndex, NumPrimitives, NumInstances);
@@ -3484,7 +3145,6 @@ public:
 	FORCEINLINE_DEBUGGABLE void DrawPrimitiveIndirect(FRHIBuffer* ArgumentBuffer, uint32 ArgumentOffset)
 	{
 		//check(IsOutsideRenderPass());
-		PreDraw();
 		if (Bypass())
 		{
 			GetContext().RHIDrawPrimitiveIndirect(ArgumentBuffer, ArgumentOffset);
@@ -3496,7 +3156,6 @@ public:
 	FORCEINLINE_DEBUGGABLE void DrawIndexedIndirect(FRHIBuffer* IndexBufferRHI, FRHIBuffer* ArgumentsBufferRHI, uint32 DrawArgumentsIndex, uint32 NumInstances)
 	{
 		//check(IsOutsideRenderPass());
-		PreDraw();
 		if (Bypass())
 		{
 			GetContext().RHIDrawIndexedIndirect(IndexBufferRHI, ArgumentsBufferRHI, DrawArgumentsIndex, NumInstances);
@@ -3508,7 +3167,6 @@ public:
 	FORCEINLINE_DEBUGGABLE void DrawIndexedPrimitiveIndirect(FRHIBuffer* IndexBuffer, FRHIBuffer* ArgumentsBuffer, uint32 ArgumentOffset)
 	{
 		//check(IsOutsideRenderPass());
-		PreDraw();
 		if (Bypass())
 		{
 			GetContext().RHIDrawIndexedPrimitiveIndirect(IndexBuffer, ArgumentsBuffer, ArgumentOffset);
@@ -3519,7 +3177,6 @@ public:
 
 	FORCEINLINE_DEBUGGABLE void DispatchMeshShader(uint32 ThreadGroupCountX, uint32 ThreadGroupCountY, uint32 ThreadGroupCountZ)
 	{
-		PreDraw();
 		if (Bypass())
 		{
 			GetContext().RHIDispatchMeshShader(ThreadGroupCountX, ThreadGroupCountY, ThreadGroupCountZ);
@@ -3530,7 +3187,6 @@ public:
 
 	FORCEINLINE_DEBUGGABLE void DispatchIndirectMeshShader(FRHIBuffer* ArgumentBuffer, uint32 ArgumentOffset)
 	{
-		PreDraw();
 		if (Bypass())
 		{
 			GetContext().RHIDispatchIndirectMeshShader(ArgumentBuffer, ArgumentOffset);
diff --git a/Engine/Source/Runtime/RHI/Public/RHICommandListCommandExecutes.inl b/Engine/Source/Runtime/RHI/Public/RHICommandListCommandExecutes.inl
--- a/Engine/Source/Runtime/RHI/Public/RHICommandListCommandExecutes.inl
+++ b/Engine/Source/Runtime/RHI/Public/RHICommandListCommandExecutes.inl
@@ -116,84 +116,6 @@ template<> RHI_API void FRHICommandSetShaderParameters<FRHIGraphicsShader>::Exec
 	INTERNAL_DECORATOR(RHISetShaderParameters)(Shader, ParametersData, Parameters, ResourceParameters, BindlessParameters);
 }
 
-template<> RHI_API void FRHICommandSetShaderParameter<FRHIComputeShader>::Execute(FRHICommandListBase& CmdList)
-{
-	RHISTAT(SetShaderParameter);
-	INTERNAL_DECORATOR_COMPUTE(RHISetShaderParameter)(Shader, BufferIndex, BaseIndex, NumBytes, NewValue);
-}
-
-template<> RHI_API void FRHICommandSetShaderParameter<FRHIGraphicsShader>::Execute(FRHICommandListBase& CmdList)
-{
-	RHISTAT(SetShaderParameter);
-	INTERNAL_DECORATOR(RHISetShaderParameter)(Shader, BufferIndex, BaseIndex, NumBytes, NewValue);
-}
-
-template<> RHI_API void FRHICommandSetShaderUniformBuffer<FRHIComputeShader>::Execute(FRHICommandListBase& CmdList)
-{
-	RHISTAT(SetShaderUniformBuffer);
-	INTERNAL_DECORATOR_COMPUTE(RHISetShaderUniformBuffer)(Shader, BaseIndex, UniformBuffer);
-}
-
-template<> RHI_API void FRHICommandSetShaderUniformBuffer<FRHIGraphicsShader>::Execute(FRHICommandListBase& CmdList)
-{
-	RHISTAT(SetShaderUniformBuffer);
-	INTERNAL_DECORATOR(RHISetShaderUniformBuffer)(Shader, BaseIndex, UniformBuffer);
-}
-
-template<> RHI_API void FRHICommandSetShaderTexture<FRHIComputeShader>::Execute(FRHICommandListBase& CmdList)
-{
-	RHISTAT(SetShaderTexture);
-	INTERNAL_DECORATOR_COMPUTE(RHISetShaderTexture)(Shader, TextureIndex, Texture);
-}
-
-template<> RHI_API void FRHICommandSetShaderTexture<FRHIGraphicsShader>::Execute(FRHICommandListBase& CmdList)
-{
-	RHISTAT(SetShaderTexture);
-	INTERNAL_DECORATOR(RHISetShaderTexture)(Shader, TextureIndex, Texture);
-}
-
-template<> RHI_API void FRHICommandSetShaderResourceViewParameter<FRHIComputeShader>::Execute(FRHICommandListBase& CmdList)
-{
-	RHISTAT(SetShaderResourceViewParameter);
-	INTERNAL_DECORATOR_COMPUTE(RHISetShaderResourceViewParameter)(Shader, SamplerIndex, SRV);
-}
-
-template<> RHI_API void FRHICommandSetShaderResourceViewParameter<FRHIGraphicsShader>::Execute(FRHICommandListBase& CmdList)
-{
-	RHISTAT(SetShaderResourceViewParameter);
-	INTERNAL_DECORATOR(RHISetShaderResourceViewParameter)(Shader, SamplerIndex, SRV);
-}
-
-template<> RHI_API void FRHICommandSetUAVParameter<FRHIComputeShader>::Execute(FRHICommandListBase& CmdList)
-{
-	RHISTAT(SetUAVParameter);
-	INTERNAL_DECORATOR_COMPUTE(RHISetUAVParameter)(Shader, UAVIndex, UAV);
-}
-
-template<> RHI_API void FRHICommandSetUAVParameter<FRHIPixelShader>::Execute(FRHICommandListBase& CmdList)
-{
-	RHISTAT(SetUAVParameter);
-	INTERNAL_DECORATOR(RHISetUAVParameter)(Shader, UAVIndex, UAV);
-}
-
-void FRHICommandSetUAVParameter_InitialCount::Execute(FRHICommandListBase& CmdList)
-{
-	RHISTAT(SetUAVParameter);
-	INTERNAL_DECORATOR_COMPUTE(RHISetUAVParameter)(Shader, UAVIndex, UAV, InitialCount);
-}
-
-template<> RHI_API void FRHICommandSetShaderSampler<FRHIComputeShader>::Execute(FRHICommandListBase& CmdList)
-{
-	RHISTAT(SetShaderSampler);
-	INTERNAL_DECORATOR_COMPUTE(RHISetShaderSampler)(Shader, SamplerIndex, Sampler);
-}
-
-template<> RHI_API void FRHICommandSetShaderSampler<FRHIGraphicsShader>::Execute(FRHICommandListBase& CmdList)
-{
-	RHISTAT(SetShaderSampler);
-	INTERNAL_DECORATOR(RHISetShaderSampler)(Shader, SamplerIndex, Sampler);
-}
-
 void FRHICommandDrawPrimitive::Execute(FRHICommandListBase& CmdList)
 {
 	RHISTAT(DrawPrimitive);
diff --git a/Engine/Source/Runtime/RHI/Public/RHIShaderParameters.h b/Engine/Source/Runtime/RHI/Public/RHIShaderParameters.h
--- a/Engine/Source/Runtime/RHI/Public/RHIShaderParameters.h
+++ b/Engine/Source/Runtime/RHI/Public/RHIShaderParameters.h
@@ -136,90 +136,3 @@ struct FRHIBatchedShaderParameters
 		BindlessParameters.Emplace(UAV, (uint16)Index);
 	}
 };
-
-/** Class that automatically batches shader parameters on a per-stage basis.
-*   NOTE: this class will be phased out in favor of using FRHIBatchedShaderParameters directly.
-*/
-class FRHIParameterBatcher
-{
-protected:
-	FRHIShader* AllBatchedShaders[SF_NumStandardFrequencies]{};
-	FRHIBatchedShaderParameters AllBatchedShaderParameters[SF_NumStandardFrequencies];
-	bool bEnabled = false;
-
-	FRHIGraphicsShader* GetBatchedGraphicsShader(int32 Index) { return static_cast<FRHIGraphicsShader*>(AllBatchedShaders[Index]); }
-	FRHIComputeShader* GetBatchedComputeShader() { return static_cast<FRHIComputeShader*>(AllBatchedShaders[SF_Compute]); }
-
-	FORCEINLINE_DEBUGGABLE FRHIBatchedShaderParameters& GetBatchedShaderParameters(FRHIShader* InShader)
-	{
-		checkSlow(InShader);
-		const EShaderFrequency Frequency = InShader->GetFrequency();
-		checkSlow(Frequency < SF_NumStandardFrequencies);
-		check(AllBatchedShaders[Frequency] == InShader);
-		return AllBatchedShaderParameters[Frequency];
-	}
-
-public:
-	RHI_API FRHIParameterBatcher();
-	RHI_API FRHIParameterBatcher(const FBoundShaderStateInput& InBoundShaderStateInput, FRHIComputeShader* InBoundComputeShaderRHI);
-	RHI_API FRHIParameterBatcher(FRHIParameterBatcher&&);
-	RHI_API ~FRHIParameterBatcher();
-
-	FORCEINLINE bool IsEnabled() const { return bEnabled; }
-
-	RHI_API void OnBoundShaderChanged(FRHICommandList& InCommandList, const FBoundShaderStateInput& InBoundShaderStateInput);
-	RHI_API void OnBoundShaderChanged(FRHIComputeCommandList& InCommandList, FRHIComputeShader* InBoundComputeShaderRHI);
-
-	RHI_API void PreDispatch(FRHIComputeCommandList& InCommandList);
-	RHI_API void PreDraw(FRHICommandList& InCommandList);
-
-	RHI_API void FlushAllParameters(FRHIComputeCommandList& InCommandList);
-	RHI_API void FlushAllParameters(FRHICommandList& InCommandList);
-
-	RHI_API void FlushPendingParameters(FRHIComputeCommandList& InCommandList, FRHIComputeShader* InShader);
-	RHI_API void FlushPendingParameters(FRHICommandList& InCommandList, FRHIGraphicsShader* InShader);
-
-	FORCEINLINE_DEBUGGABLE void SetShaderParameter(FRHIShader* Shader, uint32 BufferIndex, uint32 BaseIndex, uint32 NumBytes, const void* NewValue)
-	{
-		GetBatchedShaderParameters(Shader).SetShaderParameter(BufferIndex, BaseIndex, NumBytes, NewValue);
-	}
-
-	FORCEINLINE_DEBUGGABLE void SetShaderTexture(FRHIShader* Shader, uint32 Index, FRHITexture* Texture)
-	{
-		GetBatchedShaderParameters(Shader).SetShaderTexture(Index, Texture);
-	}
-	FORCEINLINE_DEBUGGABLE void SetShaderResourceViewParameter(FRHIShader* Shader, uint32 Index, FRHIShaderResourceView* SRV)
-	{
-		GetBatchedShaderParameters(Shader).SetShaderResourceViewParameter(Index, SRV);
-	}
-	FORCEINLINE_DEBUGGABLE void SetShaderSampler(FRHIShader* Shader, uint32 Index, FRHISamplerState* State)
-	{
-		GetBatchedShaderParameters(Shader).SetShaderSampler(Index, State);
-	}
-	FORCEINLINE_DEBUGGABLE void SetUAVParameter(FRHIShader* Shader, uint32 UAVIndex, FRHIUnorderedAccessView* UAV)
-	{
-		GetBatchedShaderParameters(Shader).SetUAVParameter(UAVIndex, UAV);
-	}
-
-	FORCEINLINE_DEBUGGABLE void SetBindlessTexture(FRHIShader* Shader, uint32 Index, FRHITexture* Texture)
-	{
-		GetBatchedShaderParameters(Shader).SetBindlessTexture(Index, Texture);
-	}
-	FORCEINLINE_DEBUGGABLE void SetBindlessResourceView(FRHIShader* Shader, uint32 Index, FRHIShaderResourceView* SRV)
-	{
-		GetBatchedShaderParameters(Shader).SetBindlessResourceView(Index, SRV);
-	}
-	FORCEINLINE_DEBUGGABLE void SetBindlessSampler(FRHIShader* Shader, uint32 Index, FRHISamplerState* State)
-	{
-		GetBatchedShaderParameters(Shader).SetBindlessSampler(Index, State);
-	}
-	FORCEINLINE_DEBUGGABLE void SetBindlessUAV(FRHIShader* Shader, uint32 Index, FRHIUnorderedAccessView* UAV)
-	{
-		GetBatchedShaderParameters(Shader).SetBindlessUAV(Index, UAV);
-	}
-
-	FORCEINLINE_DEBUGGABLE void SetUAVParameter(FRHIShader* Shader, uint32 UAVIndex, FRHIUnorderedAccessView* UAV, uint32 InitialCount)
-	{
-		checkNoEntry(); // @todo: support append/consume buffers
-	}
-};
