Fixed: Correctly create SRV for a buffers that have a mixed usage (BUF_ByteAddressBuffer | BUF_VertexBuffer)

diff --git a/Engine/Source/Runtime/RHI/Private/DynamicRHI.cpp b/Engine/Source/Runtime/RHI/Private/DynamicRHI.cpp
--- a/Engine/Source/Runtime/RHI/Private/DynamicRHI.cpp
+++ b/Engine/Source/Runtime/RHI/Private/DynamicRHI.cpp
@@ -822,7 +822,8 @@ void FShaderResourceViewInitializer::InitType()
 	if (BufferInitializer.Buffer)
 	{
 		EBufferUsageFlags Usage = BufferInitializer.Buffer->GetUsage();
-		if (EnumHasAnyFlags(Usage, BUF_VertexBuffer))
+		// Buffer could have a mixed usage, eg (BUF_ByteAddressBuffer | BUF_VertexBuffer), so make sure we create SRV for a correct use case
+		if (EnumHasAnyFlags(Usage, BUF_VertexBuffer) && BufferInitializer.Format != PF_Unknown)
 		{
 			Type = EType::VertexBufferSRV;
 		}
@@ -836,6 +837,7 @@ void FShaderResourceViewInitializer::InitType()
 		}
 		else
 		{
+			check(EnumHasAnyFlags(Usage, BUF_StructuredBuffer | BUF_ByteAddressBuffer));
 			Type = EType::StructuredBufferSRV;
 		}
 	}
