StructuredBuffer stride validation: some platforms use the SRV's stride instead of the declared HLSL type to index in structured buffers. This should allow us to catch issues like UE-177214 or UE-178039 by running with -rhivalidation

#rnx

diff --git a/Engine/Shaders/Public/ShaderVersion.ush b/Engine/Shaders/Public/ShaderVersion.ush
--- a/Engine/Shaders/Public/ShaderVersion.ush
+++ b/Engine/Shaders/Public/ShaderVersion.ush
@@ -2,4 +2,4 @@
 // Each time the console command is executed it generates a new GUID. As this file is included
 // in Platform.ush (which should be included in any shader) it allows to invalidate the shader DDC.
 // 
-// GUID = 9C9AD04A-1448-4F79-9227-A11F49F17F90
\ No newline at end of file
+// GUID = D9C7A5A5-4311-05E0-9E08-3691866F9F78
diff --git a/Engine/Source/Developer/ShaderCompilerCommon/Private/ShaderCompilerCommon.cpp b/Engine/Source/Developer/ShaderCompilerCommon/Private/ShaderCompilerCommon.cpp
--- a/Engine/Source/Developer/ShaderCompilerCommon/Private/ShaderCompilerCommon.cpp
+++ b/Engine/Source/Developer/ShaderCompilerCommon/Private/ShaderCompilerCommon.cpp
@@ -815,6 +815,25 @@ void HandleReflectedShaderResource(
 	);
 }
 
+void UpdateStructuredBufferStride(
+	const FShaderCompilerInput& Input,
+	const FString& ResourceName,
+	uint16 BindPoint,
+	uint16 Stride,
+	FShaderCompilerOutput& CompilerOutput
+)
+{
+	if (BindPoint <= UINT16_MAX && Stride <= UINT16_MAX)
+	{
+		CompilerOutput.ParametersStrideToValidate.Add(FShaderCodeValidationStride{ BindPoint, Stride });
+	}
+	else
+	{
+		FString ErrorMessage = FString::Printf(TEXT("%s: Failed to set stride on parameter %s: Bind point %d, Stride %d"), *Input.GenerateShaderName(), *ResourceName, BindPoint, Stride);
+		CompilerOutput.Errors.Add(FShaderCompilerError(*ErrorMessage));
+	}
+}
+
 void HandleReflectedShaderUAV(
 	const FString& UAVName,
 	int32 BindOffset,
diff --git a/Engine/Source/Developer/ShaderCompilerCommon/Public/ShaderCompilerCommon.h b/Engine/Source/Developer/ShaderCompilerCommon/Public/ShaderCompilerCommon.h
--- a/Engine/Source/Developer/ShaderCompilerCommon/Public/ShaderCompilerCommon.h
+++ b/Engine/Source/Developer/ShaderCompilerCommon/Public/ShaderCompilerCommon.h
@@ -173,6 +173,14 @@ extern SHADERCOMPILERCOMMON_API void HandleReflectedShaderResource(
 	FShaderCompilerOutput& CompilerOutput
 );
 
+extern SHADERCOMPILERCOMMON_API void UpdateStructuredBufferStride(
+	const FShaderCompilerInput& Input,
+	const FString& ResourceName,
+	uint16 BindPoint,
+	uint16 Stride,
+	FShaderCompilerOutput& CompilerOutput
+);
+
 inline void HandleReflectedShaderResource(const FString& ResourceName, int32 ReflectionSlot, int32 BindCount, FShaderCompilerOutput& CompilerOutput)
 {
 	HandleReflectedShaderResource(ResourceName, 0, ReflectionSlot, BindCount, CompilerOutput);
diff --git a/Engine/Source/Developer/Windows/ShaderFormatD3D/Private/D3DShaderCompiler.inl b/Engine/Source/Developer/Windows/ShaderFormatD3D/Private/D3DShaderCompiler.inl
--- a/Engine/Source/Developer/Windows/ShaderFormatD3D/Private/D3DShaderCompiler.inl
+++ b/Engine/Source/Developer/Windows/ShaderFormatD3D/Private/D3DShaderCompiler.inl
@@ -508,7 +508,16 @@ template <typename ID3D1xShaderReflection, typename D3D1x_SHADER_DESC, typename
 		else if (BindDesc.Type == D3D_SIT_STRUCTURED || BindDesc.Type == D3D_SIT_BYTEADDRESS)
 		{
 			check(BindDesc.BindCount == 1);
-			HandleReflectedShaderResource(FString(BindDesc.Name), BindDesc.BindPoint, Output);
+			FString BindDescName(BindDesc.Name);
+			HandleReflectedShaderResource(BindDescName, BindDesc.BindPoint, Output);
+
+			// https://learn.microsoft.com/en-us/windows/win32/api/d3d12shader/ns-d3d12shader-d3d12_shader_input_bind_desc
+			// If the shader resource is a structured buffer, the field contains the stride of the type in bytes
+			if ( BindDesc.Type == D3D_SIT_STRUCTURED)
+			{
+				UpdateStructuredBufferStride(Input, BindDescName, BindDesc.BindPoint, BindDesc.NumSamples, Output);
+			}
+
 			NumSRVs = FMath::Max(NumSRVs, BindDesc.BindPoint + 1);
 		}
 		else if (BindDesc.Type == (D3D_SHADER_INPUT_TYPE)(D3D_SIT_UAV_RWSTRUCTURED_WITH_COUNTER + 1)) // D3D_SIT_RTACCELERATIONSTRUCTURE (12)
@@ -536,7 +545,7 @@ inline void GenerateFinalOutput(TRefCountPtr<TBlob>& CompressedData,
 
 	TArray<uint8> UniformBufferNameBytes;
 
-	{
+	{	
 		// Build the generic SRT for this shader.
 		FShaderCompilerResourceTable GenericSRT;
 		BuildResourceTableMapping(Input.Environment.ResourceTableMap, Input.Environment.UniformBufferMap, UsedUniformBufferSlots, Output.ParameterMap, GenericSRT);
@@ -623,6 +632,8 @@ inline void GenerateFinalOutput(TRefCountPtr<TBlob>& CompressedData,
 		Output.ShaderCode.AddOptionalData(FShaderCodeName::Key, TCHAR_TO_UTF8(*Input.GenerateShaderName()));
 	}
 
+	Output.SerializeShaderCodeValidation();
+
 	// Set the number of instructions.
 	Output.NumInstructions = NumInstructions;
 
diff --git a/Engine/Source/Runtime/D3D12RHI/Private/D3D12Shaders.cpp b/Engine/Source/Runtime/D3D12RHI/Private/D3D12Shaders.cpp
--- a/Engine/Source/Runtime/D3D12RHI/Private/D3D12Shaders.cpp
+++ b/Engine/Source/Runtime/D3D12RHI/Private/D3D12Shaders.cpp
@@ -52,6 +52,8 @@ static inline bool ReadShaderOptionalData(FShaderCodeReader& InShaderCode, TShad
 	}
 #endif
 
+	UE::RHICore::SetupShaderCodeValidationData(&OutShader, InShaderCode);
+
 	return true;
 }
 
diff --git a/Engine/Source/Runtime/RHI/Private/RHIValidation.cpp b/Engine/Source/Runtime/RHI/Private/RHIValidation.cpp
--- a/Engine/Source/Runtime/RHI/Private/RHIValidation.cpp
+++ b/Engine/Source/Runtime/RHI/Private/RHIValidation.cpp
@@ -13,6 +13,7 @@
 #include "Misc/OutputDeviceRedirector.h"
 #include "RHIContext.h"
 #include "RHIStrings.h"
+#include "Algo/BinarySearch.h"
 
 #if ENABLE_RHI_VALIDATION
 
@@ -1949,6 +1950,42 @@ namespace RHIValidation
 
 		return Backtrace;
 	}
+
+	/** Validates that the SRV is conform to what the shader expects */
+	void ValidateShaderResourceView(const FRHIShader* RHIShaderBase, uint32 BindIndex, FRHIShaderResourceView* SRV)
+	{
+#if RHI_INCLUDE_SHADER_DEBUG_DATA
+		if (SRV)
+		{
+			// DebugStrideValidationData is supposed to be already sorted
+			static const auto ShaderCodeValidationStridePredicate = [](const FShaderCodeValidationStride& lhs, const FShaderCodeValidationStride& rhs) -> bool { return lhs.BindPoint < rhs.BindPoint; };
+			FShaderCodeValidationStride SRVValidationStride = { BindIndex , SRV->ValidationStride };
+			const int32 FoundIndex =  Algo::BinarySearch(RHIShaderBase->DebugStrideValidationData, SRVValidationStride, ShaderCodeValidationStridePredicate);
+			if (FoundIndex != INDEX_NONE)
+			{
+				uint16 ExpectedStride = RHIShaderBase->DebugStrideValidationData[FoundIndex].Stride;
+				// ensure(SRVValidationStride.Stride > 0) because there's a few ways to set .ValidationStride: don't trigger a validation error if we forgot to set the validation data during 
+				// some *CreateShaderResourceView* variant
+				if (ensure(SRVValidationStride.Stride > 0) && ExpectedStride != SRVValidationStride.Stride)
+				{
+					FString SRVName;
+					// It seems that owner name is usually unknown. Instead try to rely on the tracked buffer
+					if (SRV->ViewIdentity.Resource)
+					{
+						SRVName = SRV->ViewIdentity.Resource->GetDebugName();
+					}
+					if (SRVName.IsEmpty())
+					{
+						SRVName = SRV->GetOwnerName().ToString();
+					}
+					FString ErrorMessage = FString::Printf(TEXT("Shader %s: Buffer stride for \"%s\" must match structure size declared in the shader"), RHIShaderBase->GetShaderName(), *SRVName);
+					ErrorMessage += FString::Printf(TEXT("\nBind point: %d, HLSL size: %d, Buffer Size: %d"), BindIndex, ExpectedStride, SRVValidationStride.Stride);
+					RHI_VALIDATION_CHECK(false, *ErrorMessage);
+				}
+			}
+		}
+#endif
+	}
 }
 
 
diff --git a/Engine/Source/Runtime/RHI/Public/RHIDefinitions.h b/Engine/Source/Runtime/RHI/Public/RHIDefinitions.h
--- a/Engine/Source/Runtime/RHI/Public/RHIDefinitions.h
+++ b/Engine/Source/Runtime/RHI/Public/RHIDefinitions.h
@@ -1346,6 +1346,12 @@ struct FScreenResolutionRHI
 	uint32	RefreshRate;
 };
 
+struct FShaderCodeValidationStride
+{
+	uint16 BindPoint;
+	uint16 Stride;
+};
+
 #if UE_ENABLE_INCLUDE_ORDER_DEPRECATED_IN_5_2
 #include "CoreMinimal.h"
 #include "DataDrivenShaderPlatformInfo.h"
diff --git a/Engine/Source/Runtime/RHI/Public/RHIResources.h b/Engine/Source/Runtime/RHI/Public/RHIResources.h
--- a/Engine/Source/Runtime/RHI/Public/RHIResources.h
+++ b/Engine/Source/Runtime/RHI/Public/RHIResources.h
@@ -779,6 +779,8 @@ public:
 			: TEXT("<unknown>");
 	}
 
+	TArray<FShaderCodeValidationStride> DebugStrideValidationData;
+
 #else
 
 	const TCHAR* GetShaderName() const { return TEXT("<unknown>"); }
diff --git a/Engine/Source/Runtime/RHI/Public/RHIValidation.h b/Engine/Source/Runtime/RHI/Public/RHIValidation.h
--- a/Engine/Source/Runtime/RHI/Public/RHIValidation.h
+++ b/Engine/Source/Runtime/RHI/Public/RHIValidation.h
@@ -495,6 +495,7 @@ public:
 	{
 		FShaderResourceViewRHIRef SRV = RHI->RHICreateShaderResourceView(Buffer);
 		SRV->ViewIdentity = Buffer->GetWholeResourceIdentity();
+		SRV->ValidationStride = Buffer->GetStride();
 		return SRV;
 	}
 
@@ -504,6 +505,7 @@ public:
 	{
 		FShaderResourceViewRHIRef SRV = RHI->RHICreateShaderResourceView(Buffer, Stride, Format);
 		SRV->ViewIdentity = Buffer->GetWholeResourceIdentity();
+		SRV->ValidationStride = Stride;
 		return SRV;
 	}
 
@@ -524,10 +526,12 @@ public:
 		if (Buffer)
 		{
 			SRV->ViewIdentity = Buffer->GetWholeResourceIdentity();
+			SRV->ValidationStride = Stride;
 		}
 		else
 		{
 			SRV->ResetViewIdentity();
+			SRV->ValidationStride = 0;
 		}
 	}
 
@@ -538,10 +542,12 @@ public:
 		if (Buffer)
 		{
 			SRV->ViewIdentity = Buffer->GetWholeResourceIdentity();
+			SRV->ValidationStride = Buffer->GetStride();
 		}
 		else
 		{
 			SRV->ResetViewIdentity();
+			SRV->ValidationStride = 0;
 		}
 	}
 
@@ -1279,6 +1285,7 @@ public:
 	{
 		FShaderResourceViewRHIRef SRV = RHI->CreateShaderResourceView_RenderThread(RHICmdList, Buffer, Stride, Format);
 		SRV->ViewIdentity = Buffer->GetWholeResourceIdentity();
+		SRV->ValidationStride = Stride;
 		return SRV;
 	}
 
@@ -1293,6 +1300,11 @@ public:
 			SRV->ViewIdentity = Initializer.AsBufferSRV().Buffer->GetWholeResourceIdentity();
 		}
 
+		if (Initializer.GetType() == FShaderResourceViewInitializer::EType::StructuredBufferSRV && Initializer.AsStructuredBufferSRV().Buffer)
+		{
+			SRV->ValidationStride = Initializer.AsStructuredBufferSRV().Buffer->GetStride();
+		}
+
 		return SRV;
 	}
 
@@ -1300,6 +1312,7 @@ public:
 	{
 		FShaderResourceViewRHIRef SRV = RHI->CreateShaderResourceView_RenderThread(RHICmdList, Buffer);
 		SRV->ViewIdentity = Buffer->GetWholeResourceIdentity();
+		SRV->ValidationStride = Buffer->GetStride();
 		return SRV;
 	}
 
@@ -1432,6 +1445,7 @@ public:
 	{
 		FShaderResourceViewRHIRef SRV = RHI->RHICreateShaderResourceView_RenderThread(RHICmdList, Buffer, Stride, Format);
 		SRV->ViewIdentity = Buffer->GetWholeResourceIdentity();
+		SRV->ValidationStride = Stride;
 		return SRV;
 	}
 
@@ -1439,6 +1453,7 @@ public:
 	{
 		FShaderResourceViewRHIRef SRV = RHI->RHICreateShaderResourceView_RenderThread(RHICmdList, Buffer);
 		SRV->ViewIdentity = Buffer->GetWholeResourceIdentity();
+		SRV->ValidationStride = Buffer->GetStride();
 		return SRV;
 	}
 
diff --git a/Engine/Source/Runtime/RHI/Public/RHIValidationCommon.h b/Engine/Source/Runtime/RHI/Public/RHIValidationCommon.h
--- a/Engine/Source/Runtime/RHI/Public/RHIValidationCommon.h
+++ b/Engine/Source/Runtime/RHI/Public/RHIValidationCommon.h
@@ -445,7 +445,9 @@ namespace RHIValidation
 	};
 
 	struct FShaderResourceView : public FView
-	{};
+	{
+		uint32 ValidationStride = 0;
+	};
 
 	struct FUnorderedAccessView : public FView
 	{};
@@ -1016,6 +1018,10 @@ namespace RHIValidation
 	};
 
 	extern RHI_API void* CaptureBacktrace();
+
+	/** Validates that the SRV is conform to what the shader expects */
+	extern RHI_API void ValidateShaderResourceView(const class FRHIShader* RHIShaderBase, uint32 BindIndex, class FRHIShaderResourceView* SRV);
+
 }
 
 #endif // ENABLE_RHI_VALIDATION
diff --git a/Engine/Source/Runtime/RHI/Public/RHIValidationContext.h b/Engine/Source/Runtime/RHI/Public/RHIValidationContext.h
--- a/Engine/Source/Runtime/RHI/Public/RHIValidationContext.h
+++ b/Engine/Source/Runtime/RHI/Public/RHIValidationContext.h
@@ -18,7 +18,7 @@
 
 class FValidationRHI;
 
-inline void ValidateShaderParameters(RHIValidation::FTracker* Tracker, RHIValidation::FStaticUniformBuffers& StaticUniformBuffers, TConstArrayView<FRHIShaderParameterResource> InParameters, ERHIAccess InRequiredAccess, RHIValidation::EUAVMode InRequiredUAVMode)
+inline void ValidateShaderParameters(FRHIShader* RHIShader, RHIValidation::FTracker* Tracker, RHIValidation::FStaticUniformBuffers& StaticUniformBuffers, TConstArrayView<FRHIShaderParameterResource> InParameters, ERHIAccess InRequiredAccess, RHIValidation::EUAVMode InRequiredUAVMode)
 {
 	for (const FRHIShaderParameterResource& Parameter : InParameters)
 	{
@@ -33,6 +33,10 @@ inline void ValidateShaderParameters(RHIValidation::FTracker* Tracker, RHIValida
 		case FRHIShaderParameterResource::EType::ResourceView:
 			if (FRHIShaderResourceView* SRV = static_cast<FRHIShaderResourceView*>(Parameter.Resource))
 			{
+				if (GRHIValidationEnabled)
+				{
+					RHIValidation::ValidateShaderResourceView(RHIShader, Parameter.Index, SRV);
+				}
 				Tracker->Assert(SRV->ViewIdentity, InRequiredAccess);
 			}
 			break;
@@ -270,8 +274,8 @@ public:
 	{
 		checkf(State.bComputePSOSet, TEXT("A Compute PSO has to be set to set resources into a shader!"));
 
-		ValidateShaderParameters(Tracker, State.StaticUniformBuffers, InResourceParameters, ERHIAccess::SRVCompute, RHIValidation::EUAVMode::Compute);
-		ValidateShaderParameters(Tracker, State.StaticUniformBuffers, InBindlessParameters, ERHIAccess::SRVCompute, RHIValidation::EUAVMode::Compute);
+		ValidateShaderParameters(Shader, Tracker, State.StaticUniformBuffers, InResourceParameters, ERHIAccess::SRVCompute, RHIValidation::EUAVMode::Compute);
+		ValidateShaderParameters(Shader, Tracker, State.StaticUniformBuffers, InBindlessParameters, ERHIAccess::SRVCompute, RHIValidation::EUAVMode::Compute);
 
 		RHIContext->RHISetShaderParameters(Shader, InParametersData, InParameters, InResourceParameters, InBindlessParameters);
 	}
@@ -802,8 +806,8 @@ public:
 	{
 		checkf(State.bGfxPSOSet, TEXT("A Graphics PSO has to be set to set resources into a shader!"));
 
-		ValidateShaderParameters(Tracker, State.StaticUniformBuffers, InResourceParameters, ERHIAccess::SRVGraphics, RHIValidation::EUAVMode::Graphics);
-		ValidateShaderParameters(Tracker, State.StaticUniformBuffers, InBindlessParameters, ERHIAccess::SRVGraphics, RHIValidation::EUAVMode::Graphics);
+		ValidateShaderParameters(Shader, Tracker, State.StaticUniformBuffers, InResourceParameters, ERHIAccess::SRVGraphics, RHIValidation::EUAVMode::Graphics);
+		ValidateShaderParameters(Shader, Tracker, State.StaticUniformBuffers, InBindlessParameters, ERHIAccess::SRVGraphics, RHIValidation::EUAVMode::Graphics);
 
 		RHIContext->RHISetShaderParameters(Shader, InParametersData, InParameters, InResourceParameters, InBindlessParameters);
 	}
@@ -812,8 +816,8 @@ public:
 	{
 		checkf(State.bComputePSOSet, TEXT("A Compute PSO has to be set to set resources into a shader!"));
 
-		ValidateShaderParameters(Tracker, State.StaticUniformBuffers, InResourceParameters, ERHIAccess::SRVCompute, RHIValidation::EUAVMode::Compute);
-		ValidateShaderParameters(Tracker, State.StaticUniformBuffers, InBindlessParameters, ERHIAccess::SRVCompute, RHIValidation::EUAVMode::Compute);
+		ValidateShaderParameters(Shader, Tracker, State.StaticUniformBuffers, InResourceParameters, ERHIAccess::SRVCompute, RHIValidation::EUAVMode::Compute);
+		ValidateShaderParameters(Shader, Tracker, State.StaticUniformBuffers, InBindlessParameters, ERHIAccess::SRVCompute, RHIValidation::EUAVMode::Compute);
 
 		RHIContext->RHISetShaderParameters(Shader, InParametersData, InParameters, InResourceParameters, InBindlessParameters);
 	}
diff --git a/Engine/Source/Runtime/RHICore/Private/RHICoreShader.cpp b/Engine/Source/Runtime/RHICore/Private/RHICoreShader.cpp
--- a/Engine/Source/Runtime/RHICore/Private/RHICoreShader.cpp
+++ b/Engine/Source/Runtime/RHICore/Private/RHICoreShader.cpp
@@ -62,5 +62,23 @@ void ValidateStaticUniformBuffer(FRHIUniformBuffer* UniformBuffer, FUniformBuffe
 #endif
 }
 
+void SetupShaderCodeValidationData(FRHIShader* RHIShader, FShaderCodeReader& ShaderCodeReader)
+{
+#if RHI_INCLUDE_SHADER_DEBUG_DATA && ENABLE_RHI_VALIDATION
+	if (GRHIValidationEnabled && RHIShader)
+	{
+		int32 ShaderCodeValidationExtensionSize = 0;
+		const uint8* ShaderCodeValidationExtensionData = ShaderCodeReader.FindOptionalDataAndSize(FShaderCodeValidationExtension::Key, ShaderCodeValidationExtensionSize);
+		if (ShaderCodeValidationExtensionData && ShaderCodeValidationExtensionSize > 0)
+		{
+			FBufferReader ArValidationExtensionData((void*)ShaderCodeValidationExtensionData, ShaderCodeValidationExtensionSize, false);
+			FShaderCodeValidationExtension ShaderCodeValidationExtension;
+			ArValidationExtensionData << ShaderCodeValidationExtension;
+			RHIShader->DebugStrideValidationData.Append(ShaderCodeValidationExtension.ShaderCodeValidationStride);
+		}
+	}
+#endif
+}
+
 } //! RHICore
 } //! UE
\ No newline at end of file
diff --git a/Engine/Source/Runtime/RHICore/Public/RHICoreShader.h b/Engine/Source/Runtime/RHICore/Public/RHICoreShader.h
--- a/Engine/Source/Runtime/RHICore/Public/RHICoreShader.h
+++ b/Engine/Source/Runtime/RHICore/Public/RHICoreShader.h
@@ -12,6 +12,7 @@ namespace RHICore
 
 /** Validates that the uniform buffer at the requested static slot. */
 extern RHICORE_API void ValidateStaticUniformBuffer(FRHIUniformBuffer* UniformBuffer, FUniformBufferStaticSlot Slot, uint32 ExpectedHash);
+extern RHICORE_API void SetupShaderCodeValidationData(FRHIShader* RHIShader, class FShaderCodeReader& ShaderCodeReader);
 
 inline void InitStaticUniformBufferSlots(TArray<FUniformBufferStaticSlot>& StaticSlots, const FShaderResourceTable& ShaderResourceTable)
 {
@@ -196,6 +197,10 @@ void SetResourcesFromTables(TBinder&& Binder, FRHIShader const& Shader, FShaderR
 
 					Tracker->Assert(SRV->ViewIdentity, Access);
 				}
+				if (GRHIValidationEnabled)
+				{
+					RHIValidation::ValidateShaderResourceView(&Shader, Index, SRV);
+				}
 #endif
 				Binder.SetSRV(SRV, Index);
 			});
diff --git a/Engine/Source/Runtime/RenderCore/Private/ShaderCore.cpp b/Engine/Source/Runtime/RenderCore/Private/ShaderCore.cpp
--- a/Engine/Source/Runtime/RenderCore/Private/ShaderCore.cpp
+++ b/Engine/Source/Runtime/RenderCore/Private/ShaderCore.cpp
@@ -733,6 +733,22 @@ void FShaderCompilerOutput::CompressOutput(FName ShaderCompressionFormat, FOodle
 	ShaderCode.Compress(ShaderCompressionFormat, OodleCompressor, OodleLevel);
 }
 
+void FShaderCompilerOutput::SerializeShaderCodeValidation()
+{
+	if (ParametersStrideToValidate.Num() > 0)
+	{
+		FShaderCodeValidationExtension ShaderCodeValidationExtension;
+		ShaderCodeValidationExtension.ShaderCodeValidationStride.Append(ParametersStrideToValidate);
+		ShaderCodeValidationExtension.ShaderCodeValidationStride.Sort([](const FShaderCodeValidationStride& lhs, const FShaderCodeValidationStride& rhs) -> bool { return lhs.BindPoint < rhs.BindPoint; });
+
+		TArray<uint8> WriterBytes;
+		FMemoryWriter Writer(WriterBytes);
+		Writer << ShaderCodeValidationExtension;
+
+		ShaderCode.AddOptionalData(FShaderCodeValidationExtension::Key, WriterBytes.GetData(), WriterBytes.Num());
+	}
+}
+
 static void ReportVirtualShaderFilePathError(TArray<FShaderCompilerError>* CompileErrors, FString ErrorString)
 {
 	if (CompileErrors)
diff --git a/Engine/Source/Runtime/RenderCore/Public/ShaderCompilerCore.h b/Engine/Source/Runtime/RenderCore/Public/ShaderCompilerCore.h
--- a/Engine/Source/Runtime/RenderCore/Public/ShaderCompilerCore.h
+++ b/Engine/Source/Runtime/RenderCore/Public/ShaderCompilerCore.h
@@ -490,6 +490,8 @@ struct FShaderCompilerOutput
 	bool bUsedHLSLccCompiler;
 	TArray<FString> UsedAttributes;
 
+	TArray<FShaderCodeValidationStride> ParametersStrideToValidate;
+
 	FString OptionalFinalShaderSource;
 
 	TArray<uint8> PlatformDebugData;
@@ -500,6 +502,9 @@ struct FShaderCompilerOutput
 	/** Calls GenerateOutputHash() before the compression, replaces FShaderCode with the compressed data (if compression result was smaller). */
 	RENDERCORE_API void CompressOutput(FName ShaderCompressionFormat, FOodleDataCompression::ECompressor OodleCompressor, FOodleDataCompression::ECompressionLevel OodleLevel);
 
+	/** Add optional data in ShaderCode to perform additional shader input validation at runtime*/
+	RENDERCORE_API void SerializeShaderCodeValidation();
+
 	friend FArchive& operator<<(FArchive& Ar, FShaderCompilerOutput& Output)
 	{
 		// Note: this serialize is used to pass between UE and the shader compile worker, recompile both when modifying
diff --git a/Engine/Source/Runtime/RenderCore/Public/ShaderCore.h b/Engine/Source/Runtime/RenderCore/Public/ShaderCore.h
--- a/Engine/Source/Runtime/RenderCore/Public/ShaderCore.h
+++ b/Engine/Source/Runtime/RenderCore/Public/ShaderCore.h
@@ -730,6 +730,29 @@ struct FShaderCodeVendorExtension
 
 };
 
+
+inline FArchive& operator<<(FArchive& Ar, FShaderCodeValidationStride& ShaderCodeValidationStride)
+{
+	return Ar << ShaderCodeValidationStride.BindPoint << ShaderCodeValidationStride.Stride;
+}
+
+struct FShaderCodeValidationExtension
+{
+	// for FindOptionalData() and AddOptionalData()
+	static constexpr uint8 Key = 'V';
+	static constexpr uint16 StaticVersion = 0;
+
+	TArray<FShaderCodeValidationStride> ShaderCodeValidationStride;
+	uint16 Version = StaticVersion;
+
+	friend FArchive& operator<<(FArchive& Ar, FShaderCodeValidationExtension& Extension)
+	{
+		Ar << Extension.Version;
+		Ar << Extension.ShaderCodeValidationStride;
+		return Ar;
+	}
+};
+
 #ifndef RENDERCORE_ATTRIBUTE_UNALIGNED
 // TODO find out if using GCC_ALIGN(1) instead of this new #define break on all kinds of platforms...
 #define RENDERCORE_ATTRIBUTE_UNALIGNED
