FIX CIS
Don't enable distancefield for every platform - Mistake from 24652571

diff --git a/Engine/Source/Runtime/RHI/Private/DataDrivenShaderPlatformInfo.cpp b/Engine/Source/Runtime/RHI/Private/DataDrivenShaderPlatformInfo.cpp
--- a/Engine/Source/Runtime/RHI/Private/DataDrivenShaderPlatformInfo.cpp
+++ b/Engine/Source/Runtime/RHI/Private/DataDrivenShaderPlatformInfo.cpp
@@ -134,7 +134,6 @@ void FGenericDataDrivenShaderPlatformInfo::SetDefaultValues()
 	bSupportsManualVertexFetch = true;
 	bSupportsVolumeTextureAtomics = true;
 	bSupportsClipDistance = true;
-	bSupportsDistanceFields = true;
 	PreviewShaderPlatformParent = EShaderPlatform::SP_NumPlatforms;
 }
 
