Output owner name to rhi.DumpResourceMemory command

diff --git a/Engine/Source/Runtime/RHI/Private/RHI.cpp b/Engine/Source/Runtime/RHI/Private/RHI.cpp
--- a/Engine/Source/Runtime/RHI/Private/RHI.cpp
+++ b/Engine/Source/Runtime/RHI/Private/RHI.cpp
@@ -861,12 +861,12 @@ void RHIDumpResourceMemory(const FString& NameFilter, ERHIResourceType TypeFilte
 
 	if (bOutputToCSVFile)
 	{
-		const TCHAR* Header = TEXT("Name,Type,Size,Resident,MarkedForDelete,Transient,Streaming,RenderTarget,UAV,\"Raytracing Acceleration Structure\"\n");
+		const TCHAR* Header = TEXT("Name,Type,Size,Resident,MarkedForDelete,Transient,Streaming,RenderTarget,UAV,\"Raytracing Acceleration Structure\",Owner\n");
 		CSVFile->Serialize(TCHAR_TO_ANSI(Header), FPlatformString::Strlen(Header));
 	}
 	else if (bUseCSVOutput)
 	{
-		BufferedOutput.CategorizedLogf(CategoryName, ELogVerbosity::Log, TEXT("Name,Type,Size,Resident,MarkedForDelete,Transient,Streaming,RenderTarget,UAV,\"Raytracing Acceleration Structure\""));
+		BufferedOutput.CategorizedLogf(CategoryName, ELogVerbosity::Log, TEXT("Name,Type,Size,Resident,MarkedForDelete,Transient,Streaming,RenderTarget,UAV,\"Raytracing Acceleration Structure\",Owner"));
 	}
 	else
 	{
@@ -882,6 +882,7 @@ void RHIDumpResourceMemory(const FString& NameFilter, ERHIResourceType TypeFilte
 	}
 
 	TCHAR ResourceNameBuffer[FName::StringBufferSize];
+	TCHAR ResourceOwnerBuffer[FName::StringBufferSize];
 	int64 TotalShownResourceSize = 0;
 
 	for (int32 Index = 0; Index < Resources.Num(); Index++)
@@ -892,8 +893,8 @@ void RHIDumpResourceMemory(const FString& NameFilter, ERHIResourceType TypeFilte
 
 			ResourceInfo.Name.ToString(ResourceNameBuffer);
 			const TCHAR* ResourceType = StringFromRHIResourceType(ResourceInfo.Type);
-			const int64 SizeInBytes = ResourceInfo.VRamAllocation.AllocationSize;
-			
+			const int64 SizeInBytes = ResourceInfo.VRamAllocation.AllocationSize;			
+			Resources[Index].Resource->GetOwnerName().ToString(ResourceOwnerBuffer);
 
 			RHIInternal::FResourceFlags Flags = GetResourceFlagsInternal(Resources[Index]);
 
@@ -901,7 +902,7 @@ void RHIDumpResourceMemory(const FString& NameFilter, ERHIResourceType TypeFilte
 			{
 				if (bOutputToCSVFile || bUseCSVOutput)
 				{		
-					const FString Row = FString::Printf(TEXT("%s,%s,%.9f,%s,%s,%s,%s,%s,%s,%s\n"),
+					const FString Row = FString::Printf(TEXT("%s,%s,%.9f,%s,%s,%s,%s,%s,%s,%s,%s\n"),
 						ResourceNameBuffer,
 						ResourceType,
 						SizeInBytes / double(1 << 20),
@@ -911,7 +912,8 @@ void RHIDumpResourceMemory(const FString& NameFilter, ERHIResourceType TypeFilte
 						Flags.bStreaming ? TEXT("Yes") : TEXT("No"),
 						(Flags.bRT || Flags.bDS) ? TEXT("Yes") : TEXT("No"),
 						Flags.bUAV ? TEXT("Yes") : TEXT("No"),
-						Flags.bRTAS ? TEXT("Yes") : TEXT("No"));
+						Flags.bRTAS ? TEXT("Yes") : TEXT("No"),
+						ResourceOwnerBuffer);
 
 					if (bOutputToCSVFile)
 					{
@@ -925,11 +927,12 @@ void RHIDumpResourceMemory(const FString& NameFilter, ERHIResourceType TypeFilte
 				else
 				{
 					FString ResoureFlags = Flags.GetString();
-					BufferedOutput.CategorizedLogf(CategoryName, ELogVerbosity::Log, TEXT("Name: %s - Type: %s - Size: %.9f MB - Flags: %s"),
+					BufferedOutput.CategorizedLogf(CategoryName, ELogVerbosity::Log, TEXT("Name: %s - Type: %s - Size: %.9f MB - Flags: %s - Owner: %s"),
 						ResourceNameBuffer,
 						ResourceType,
 						SizeInBytes / double(1 << 20),
-						ResoureFlags.IsEmpty() ? TEXT("None") : *ResoureFlags);
+						ResoureFlags.IsEmpty() ? TEXT("None") : *ResoureFlags,
+						ResourceOwnerBuffer);
 				}
 			}
 
