VulkanRHI: Add bSupportsRayTracingShaders to the values that are checked for support for preview mode DDPI.

#rnx

diff --git a/Engine/Source/Runtime/RHI/Private/DataDrivenShaderPlatformInfo.cpp b/Engine/Source/Runtime/RHI/Private/DataDrivenShaderPlatformInfo.cpp
--- a/Engine/Source/Runtime/RHI/Private/DataDrivenShaderPlatformInfo.cpp
+++ b/Engine/Source/Runtime/RHI/Private/DataDrivenShaderPlatformInfo.cpp
@@ -389,6 +389,7 @@ void FGenericDataDrivenShaderPlatformInfo::UpdatePreviewPlatforms()
 				Infos[ShaderPlatform].bSupportsUInt64ImageAtomics &= Infos[EditorSPForPreviewMaxFeatureLevel].bSupportsUInt64ImageAtomics;
 				Infos[ShaderPlatform].bSupportsGen5TemporalAA &= Infos[EditorSPForPreviewMaxFeatureLevel].bSupportsGen5TemporalAA;
 				Infos[ShaderPlatform].bSupportsInlineRayTracing &= Infos[EditorSPForPreviewMaxFeatureLevel].bSupportsInlineRayTracing;
+				Infos[ShaderPlatform].bSupportsRayTracingShaders &= Infos[EditorSPForPreviewMaxFeatureLevel].bSupportsRayTracingShaders;
 
 				// Support for stereo features requires extra consideration. The editor may not use the same technique as the preview platform,
 				// particularly MobileMultiView may be substituted by a fallback path. In order to avoid inundating real mobile platforms
