Implements FRHIGlobals::SupportsNative16BitOps

diff --git a/Engine/Source/Runtime/D3D12RHI/Private/D3D12Adapter.cpp b/Engine/Source/Runtime/D3D12RHI/Private/D3D12Adapter.cpp
--- a/Engine/Source/Runtime/D3D12RHI/Private/D3D12Adapter.cpp
+++ b/Engine/Source/Runtime/D3D12RHI/Private/D3D12Adapter.cpp
@@ -952,6 +952,12 @@ PRAGMA_ENABLE_DEPRECATION_WARNINGS
 				GRHIMaximumWaveSize = Features.WaveLaneCountMax;
 			}
 
+			{
+				D3D12_FEATURE_DATA_D3D12_OPTIONS4 Features{};
+				RootDevice->CheckFeatureSupport(D3D12_FEATURE_D3D12_OPTIONS4, &Features, sizeof(Features));
+				GRHIGlobals.SupportsNative16BitOps = Features.Native16BitShaderOpsSupported;
+			}
+
 #if D3D12_RHI_RAYTRACING
 #if PLATFORM_WINDOWS || PLATFORM_HOLOLENS
 			D3D12_FEATURE_DATA_D3D12_OPTIONS5 D3D12Caps5 = {};
diff --git a/Engine/Source/Runtime/RHI/Public/RHIGlobals.h b/Engine/Source/Runtime/RHI/Public/RHIGlobals.h
--- a/Engine/Source/Runtime/RHI/Public/RHIGlobals.h
+++ b/Engine/Source/Runtime/RHI/Public/RHIGlobals.h
@@ -438,6 +438,9 @@ struct FRHIGlobals
 	int32 MinimumWaveSize = 0;
 	int32 MaximumWaveSize = 0;
 
+	/** Whether or not the RHI supports 16bit VALU and resource loads (HLSL shader model 6.2). */
+	bool SupportsNative16BitOps = false;
+
 	/** Whether or not the RHI supports an RHI thread.
 	Requirements for RHI thread
 	* Microresources (those in RHIStaticStates.h) need to be able to be created by any thread at any time and be able to work with a radically simplified rhi resource lifecycle. CreateSamplerState, CreateRasterizerState, CreateDepthStencilState, CreateBlendState
