Merging (with edits) CL 24691283 from //UE5/Dev-ParallelRendering to //UE5/Release-5.2

Add BeginRecursiveCommand call in the RecursiveHazardous RHI command lists
 - Some platform RHIs need to prepare the context (e.g. creating / opening a command list) before the lambda is called

diff --git a/Engine/Source/Runtime/Apple/MetalRHI/Public/MetalRHIContext.h b/Engine/Source/Runtime/Apple/MetalRHI/Public/MetalRHIContext.h
--- a/Engine/Source/Runtime/Apple/MetalRHI/Public/MetalRHIContext.h
+++ b/Engine/Source/Runtime/Apple/MetalRHI/Public/MetalRHIContext.h
@@ -244,6 +244,10 @@ public:
 		ERayTracingBindingType BindingType) final override;
 #endif // METAL_RHI_RAYTRACING
 
+	void BeginRecursiveCommand()
+	{
+		// Nothing to do
+	}
 protected:
 	static TGlobalResource<TBoundShaderStateHistory<10000>> BoundShaderStateHistory;
 	
diff --git a/Engine/Source/Runtime/D3D12RHI/Private/D3D12CommandContext.h b/Engine/Source/Runtime/D3D12RHI/Private/D3D12CommandContext.h
--- a/Engine/Source/Runtime/D3D12RHI/Private/D3D12CommandContext.h
+++ b/Engine/Source/Runtime/D3D12RHI/Private/D3D12CommandContext.h
@@ -239,6 +239,12 @@ private:
 		return *CommandList;
 	}
 
+public:
+	void BeginRecursiveCommand()
+	{
+		// Nothing to do
+	}
+
 protected:
 	enum class EPhase
 	{
diff --git a/Engine/Source/Runtime/RHI/Public/RHICommandList.h b/Engine/Source/Runtime/RHI/Public/RHICommandList.h
--- a/Engine/Source/Runtime/RHI/Public/RHICommandList.h
+++ b/Engine/Source/Runtime/RHI/Public/RHICommandList.h
@@ -4551,7 +4551,10 @@ class TRHICommandList_RecursiveHazardous : public FRHICommandList_RecursiveHazar
 		void ExecuteAndDestruct(FRHICommandListBase& CmdList, FRHICommandListDebugContext&) override final
 		{
 			// RunOnContext always requires the lowest level (platform) context, not the validation RHI context.
-			Lambda(static_cast<ContextType&>(CmdList.GetContext().GetLowestLevelContext()));
+			ContextType& Context = static_cast<ContextType&>(CmdList.GetContext().GetLowestLevelContext());
+			Context.BeginRecursiveCommand();
+
+			Lambda(Context);
 			Lambda.~LAMBDA();
 		}
 	};
@@ -4567,7 +4570,10 @@ public:
 		if (Bypass())
 		{
 			// RunOnContext always requires the lowest level (platform) context, not the validation RHI context.
-			Lambda(static_cast<ContextType&>(GetContext().GetLowestLevelContext()));
+			ContextType& Context = static_cast<ContextType&>(GetContext().GetLowestLevelContext());
+			Context.BeginRecursiveCommand();
+
+			Lambda(Context);
 		}
 		else
 		{
@@ -4602,7 +4608,10 @@ class TRHIComputeCommandList_RecursiveHazardous : public FRHIComputeCommandList_
 		void ExecuteAndDestruct(FRHICommandListBase& CmdList, FRHICommandListDebugContext&) override final
 		{
 			// RunOnContext always requires the lowest level (platform) context, not the validation RHI context.
-			Lambda(static_cast<ContextType&>(CmdList.GetComputeContext().GetLowestLevelContext()));
+			ContextType& Context = static_cast<ContextType&>(CmdList.GetComputeContext().GetLowestLevelContext());
+			Context.BeginRecursiveCommand();
+
+			Lambda(Context);
 			Lambda.~LAMBDA();
 		}
 	};
@@ -4618,7 +4627,10 @@ public:
 		if (Bypass())
 		{
 			// RunOnContext always requires the lowest level (platform) context, not the validation RHI context.
-			Lambda(static_cast<ContextType&>(GetComputeContext().GetLowestLevelContext()));
+			ContextType& Context = static_cast<ContextType&>(GetComputeContext().GetLowestLevelContext());
+			Context.BeginRecursiveCommand();
+
+			Lambda(Context);
 		}
 		else
 		{
diff --git a/Engine/Source/Runtime/VulkanRHI/Private/VulkanContext.h b/Engine/Source/Runtime/VulkanRHI/Private/VulkanContext.h
--- a/Engine/Source/Runtime/VulkanRHI/Private/VulkanContext.h
+++ b/Engine/Source/Runtime/VulkanRHI/Private/VulkanContext.h
@@ -186,6 +186,11 @@ public:
 		return Device;
 	}
 
+	void BeginRecursiveCommand()
+	{
+		// Nothing to do
+	}
+
 	void EndRenderQueryInternal(FVulkanCmdBuffer* CmdBuffer, FVulkanRenderQuery* Query);
 
 	void ReleasePendingState();
diff --git a/Engine/Source/Runtime/Windows/D3D11RHI/Private/D3D11RHIPrivate.h b/Engine/Source/Runtime/Windows/D3D11RHI/Private/D3D11RHIPrivate.h
--- a/Engine/Source/Runtime/Windows/D3D11RHI/Private/D3D11RHIPrivate.h
+++ b/Engine/Source/Runtime/Windows/D3D11RHI/Private/D3D11RHIPrivate.h
@@ -729,6 +729,11 @@ public:
 	bool IsQuadBufferStereoEnabled();
 	void DisableQuadBufferStereo();
 
+	void BeginRecursiveCommand()
+	{
+		// Nothing to do
+	}
+
 private:
 	void EnableDepthBoundsTest(bool bEnable, float MinDepth, float MaxDepth);
 
