Remove NullRHI module from non-desktop builds
 - Marked RHI modules with an [RHIModule] C# attribute, so UBT can determine how many RHIs are being included in the build via standard module references.
 - Refactored the "fixed RHI class" macros to automatically de-virtualize RHI command calls. WITH_FIXED_RHI_CLASS is always defined to 1 if there's only a single RHI module included in the build.

Unified "-rhivalidation" command line switch handling.

diff --git a/Engine/Source/Programs/UnrealBuildTool/Configuration/ModuleRules.cs b/Engine/Source/Programs/UnrealBuildTool/Configuration/ModuleRules.cs
--- a/Engine/Source/Programs/UnrealBuildTool/Configuration/ModuleRules.cs
+++ b/Engine/Source/Programs/UnrealBuildTool/Configuration/ModuleRules.cs
@@ -1372,6 +1372,18 @@ namespace UnrealBuildTool
 			}
 		}
 
+		/// <summary>
+		/// Used to indicate a given UE module implements an RHI
+		/// </summary>
+		[AttributeUsage(AttributeTargets.Class, AllowMultiple = false, Inherited = true)]
+		protected sealed class RHIModuleAttribute : Attribute
+		{}
+
+		/// <summary>
+		/// True if this UE module implements an RHI.
+		/// </summary>
+		public bool IsRHIModule => Attribute.IsDefined(GetType(), typeof(RHIModuleAttribute));
+
 #nullable disable
 		/// <summary>
 		/// Constructor. For backwards compatibility while the parameterless constructor is being phased out, initialization which would happen here is done by 
diff --git a/Engine/Source/Programs/UnrealBuildTool/Configuration/UEBuildTarget.cs b/Engine/Source/Programs/UnrealBuildTool/Configuration/UEBuildTarget.cs
--- a/Engine/Source/Programs/UnrealBuildTool/Configuration/UEBuildTarget.cs
+++ b/Engine/Source/Programs/UnrealBuildTool/Configuration/UEBuildTarget.cs
@@ -4821,6 +4821,26 @@ namespace UnrealBuildTool
 			// Set the global app name
 			GlobalCompileEnvironment.Definitions.Add(String.Format("UE_APP_NAME=\"{0}\"", AppName));
 
+			// Configure RHI related global definitions
+			var RHIModules = Modules.Values.Where(m => m.Rules.IsRHIModule).ToList();
+
+			bool bAllowRHIValidation =
+				Rules.Configuration == UnrealTargetConfiguration.Debug ||
+				Rules.Configuration == UnrealTargetConfiguration.DebugGame ||
+				Rules.Configuration == UnrealTargetConfiguration.Development;
+
+			if (bAllowRHIValidation)
+			{
+				GlobalCompileEnvironment.Definitions.Add($"WITH_FIXED_RHI_CLASS=0");
+				GlobalCompileEnvironment.Definitions.Add($"ENABLE_RHI_VALIDATION=1");
+			}
+			else
+			{
+				GlobalCompileEnvironment.Definitions.Add($"WITH_FIXED_RHI_CLASS={(RHIModules.Count == 1 ? 1 : 0)}");
+				GlobalCompileEnvironment.Definitions.Add($"ENABLE_RHI_VALIDATION=0");
+			}
+			GlobalCompileEnvironment.Definitions.Add($"WITH_NULL_RHI={(RHIModules.Any(m => m.Name == "NullDrv") ? 1 : 0)}");
+
 			// Add global definitions for project-specific binaries. HACK: Also defining for monolithic builds in binary releases. Might be better to set this via command line instead?
 			if (!bUseSharedBuildEnvironment || bCompileMonolithic)
 			{
diff --git a/Engine/Source/Runtime/Apple/MetalRHI/MetalRHI.Build.cs b/Engine/Source/Runtime/Apple/MetalRHI/MetalRHI.Build.cs
--- a/Engine/Source/Runtime/Apple/MetalRHI/MetalRHI.Build.cs
+++ b/Engine/Source/Runtime/Apple/MetalRHI/MetalRHI.Build.cs
@@ -3,6 +3,7 @@
 using System.IO;
 using UnrealBuildTool;
 
+[RHIModule]
 public class MetalRHI : ModuleRules
 {	
 	public MetalRHI(ReadOnlyTargetRules Target) : base(Target)
diff --git a/Engine/Source/Runtime/Apple/MetalRHI/Private/MetalCommands.cpp b/Engine/Source/Runtime/Apple/MetalRHI/Private/MetalCommands.cpp
--- a/Engine/Source/Runtime/Apple/MetalRHI/Private/MetalCommands.cpp
+++ b/Engine/Source/Runtime/Apple/MetalRHI/Private/MetalCommands.cpp
@@ -602,7 +602,11 @@ void FMetalRHICommandContext::RHIDiscardRenderTargets(bool Depth, bool Stencil,
 	Context->GetCurrentState().DiscardRenderTargets(Depth, Stencil, ColorBitMask);
 }
 
-#if PLATFORM_USES_FIXED_RHI_CLASS
-#define INTERNAL_DECORATOR(Method) ((FMetalRHICommandContext&)CmdList.GetContext()).FMetalRHICommandContext::Method
-#include "RHICommandListCommandExecutes.inl"
-#endif
+#if WITH_FIXED_RHI_CLASS
+
+	#define INTERNAL_DECORATOR(Method)         ((FMetalRHICommandContext&)CmdList.GetContext()       ).FMetalRHICommandContext::Method
+	#define INTERNAL_DECORATOR_COMPUTE(Method) ((FMetalRHICommandContext&)CmdList.GetComputeContext()).FMetalRHICommandContext::Method
+
+	#include "RHICommandListCommandExecutes.inl"
+
+#endif // WITH_FIXED_RHI_CLASS
diff --git a/Engine/Source/Runtime/Core/Private/Misc/AutomationTest.cpp b/Engine/Source/Runtime/Core/Private/Misc/AutomationTest.cpp
--- a/Engine/Source/Runtime/Core/Private/Misc/AutomationTest.cpp
+++ b/Engine/Source/Runtime/Core/Private/Misc/AutomationTest.cpp
@@ -662,12 +662,9 @@ void FAutomationTestFramework::GetValidTestNames( TArray<FAutomationTestInfo>& T
 
 	//Feature support - assume valid RHI until told otherwise
 	uint32 FeatureSupportFlags = EAutomationTestFlags::FeatureMask;
-	// @todo: Handle this correctly. GIsUsingNullRHI is defined at Engine-level, so it can't be used directly here in Core.
-	// For now, assume Null RHI is only used for commandlets, servers, and when the command line specifies to use it.
 	if (FPlatformProperties::SupportsWindowedMode())
 	{
-		bool bUsingNullRHI = FParse::Param( FCommandLine::Get(), TEXT("nullrhi") ) || IsRunningCommandlet() || IsRunningDedicatedServer();
-		if (bUsingNullRHI)
+		if (!FApp::CanEverRender())
 		{
 			FeatureSupportFlags &= (~EAutomationTestFlags::NonNullRHI);
 		}
diff --git a/Engine/Source/Runtime/Core/Public/HAL/Platform.h b/Engine/Source/Runtime/Core/Public/HAL/Platform.h
--- a/Engine/Source/Runtime/Core/Public/HAL/Platform.h
+++ b/Engine/Source/Runtime/Core/Public/HAL/Platform.h
@@ -403,10 +403,6 @@
 	#define PLATFORM_SUPPORTS_NAMED_PIPES		0
 #endif
 
-#ifndef PLATFORM_USES_FIXED_RHI_CLASS
-	#define PLATFORM_USES_FIXED_RHI_CLASS		0
-#endif
-
 #ifndef PLATFORM_USES_FIXED_GMalloc_CLASS
 	#define PLATFORM_USES_FIXED_GMalloc_CLASS		0
 #endif
diff --git a/Engine/Source/Runtime/Core/Public/IOS/IOSPlatform.h b/Engine/Source/Runtime/Core/Public/IOS/IOSPlatform.h
--- a/Engine/Source/Runtime/Core/Public/IOS/IOSPlatform.h
+++ b/Engine/Source/Runtime/Core/Public/IOS/IOSPlatform.h
@@ -49,7 +49,6 @@ typedef FIOSPlatformTypes FPlatformTypes;
 #define PLATFORM_SUPPORTS_TEXTURE_STREAMING				1
 #define PLATFORM_BUILTIN_VERTEX_HALF_FLOAT				0
 #define PLATFORM_SUPPORTS_MULTIPLE_NATIVE_WINDOWS		0
-#define PLATFORM_ALLOW_NULL_RHI							1
 #define PLATFORM_ENABLE_VECTORINTRINSICS				1
 #define PLATFORM_ENABLE_VECTORINTRINSICS_NEON			1
 #define PLATFORM_SUPPORTS_STACK_SYMBOLS					1
diff --git a/Engine/Source/Runtime/Core/Public/Misc/App.h b/Engine/Source/Runtime/Core/Public/Misc/App.h
--- a/Engine/Source/Runtime/Core/Public/Misc/App.h
+++ b/Engine/Source/Runtime/Core/Public/Misc/App.h
@@ -378,7 +378,7 @@ public:
 		return false;
 #else
 		static bool bHasNullRHIOnCommandline = FParse::Param(FCommandLine::Get(), TEXT("nullrhi"));
-		return (!IsRunningCommandlet() || IsAllowCommandletRendering()) && !IsRunningDedicatedServer() && !(USE_NULL_RHI || bHasNullRHIOnCommandline);
+		return (!IsRunningCommandlet() || IsAllowCommandletRendering()) && !IsRunningDedicatedServer() && !(USE_NULL_RHI || (WITH_NULL_RHI && bHasNullRHIOnCommandline));
 #endif // UE_SERVER
 	}
 
diff --git a/Engine/Source/Runtime/Core/Public/Misc/Build.h b/Engine/Source/Runtime/Core/Public/Misc/Build.h
--- a/Engine/Source/Runtime/Core/Public/Misc/Build.h
+++ b/Engine/Source/Runtime/Core/Public/Misc/Build.h
@@ -179,7 +179,7 @@
 	#define CHECK_PUREVIRTUALS 0
 #endif
 
-/** Whether to use the null RHI. */
+/** Whether to force the use of the null RHI (without requiring -nullrhi on the command line). */
 #ifndef USE_NULL_RHI
 	#define USE_NULL_RHI 0
 #endif
@@ -425,10 +425,6 @@
 	#define PRESERVE_LOG_BACKUPS_IN_SHIPPING 1
 #endif
 
-#ifndef ENABLE_RHI_VALIDATION
-	#define ENABLE_RHI_VALIDATION (UE_BUILD_DEBUG || UE_BUILD_DEVELOPMENT)
-#endif
-
 // Controls whether FPlatformMisc::GetDeviceId() is available to be called.
 // When set to 1, calls to this API will be hardcoded to return an empty string
 // to avoid running afoul of calling device APIs that platform owners may restrict
diff --git a/Engine/Source/Runtime/D3D12RHI/D3D12RHI.Build.cs b/Engine/Source/Runtime/D3D12RHI/D3D12RHI.Build.cs
--- a/Engine/Source/Runtime/D3D12RHI/D3D12RHI.Build.cs
+++ b/Engine/Source/Runtime/D3D12RHI/D3D12RHI.Build.cs
@@ -2,6 +2,7 @@
 
 using UnrealBuildTool;
 
+[RHIModule]
 public class D3D12RHI : ModuleRules
 {
 	protected virtual bool bUsesWindowsD3D12 { get => Target.Platform.IsInGroup(UnrealPlatformGroup.Windows); }
diff --git a/Engine/Source/Runtime/D3D12RHI/Private/D3D12CommandList.cpp b/Engine/Source/Runtime/D3D12RHI/Private/D3D12CommandList.cpp
--- a/Engine/Source/Runtime/D3D12RHI/Private/D3D12CommandList.cpp
+++ b/Engine/Source/Runtime/D3D12RHI/Private/D3D12CommandList.cpp
@@ -2,7 +2,6 @@
 
 #include "D3D12RHIPrivate.h"
 #include "D3D12CommandList.h"
-#include "RHIValidation.h"
 
 static int32 GD3D12BatchResourceBarriers = 1;
 static FAutoConsoleVariableRef CVarD3D12BatchResourceBarriers(
diff --git a/Engine/Source/Runtime/D3D12RHI/Private/D3D12Commands.cpp b/Engine/Source/Runtime/D3D12RHI/Private/D3D12Commands.cpp
--- a/Engine/Source/Runtime/D3D12RHI/Private/D3D12Commands.cpp
+++ b/Engine/Source/Runtime/D3D12RHI/Private/D3D12Commands.cpp
@@ -1994,3 +1994,12 @@ void FD3D12CommandContext::RHISubmitCommandsHint()
 {
 	// Nothing to do
 }
+
+#if WITH_FIXED_RHI_CLASS
+
+	#define INTERNAL_DECORATOR(Method)         ((FD3D12CommandContext&)CmdList.GetContext()       ).FD3D12CommandContext::Method
+	#define INTERNAL_DECORATOR_COMPUTE(Method) ((FD3D12CommandContext&)CmdList.GetComputeContext()).FD3D12CommandContext::Method
+
+	#include "RHICommandListCommandExecutes.inl"
+
+#endif // WITH_FIXED_RHI_CLASS
diff --git a/Engine/Source/Runtime/D3D12RHI/Private/Windows/WindowsD3D12Device.cpp b/Engine/Source/Runtime/D3D12RHI/Private/Windows/WindowsD3D12Device.cpp
--- a/Engine/Source/Runtime/D3D12RHI/Private/Windows/WindowsD3D12Device.cpp
+++ b/Engine/Source/Runtime/D3D12RHI/Private/Windows/WindowsD3D12Device.cpp
@@ -37,7 +37,6 @@
 #include "HardwareInfo.h"
 #include "IHeadMountedDisplayModule.h"
 #include "GenericPlatform/GenericPlatformDriver.h"			// FGPUDriverInfo
-#include "RHIValidation.h"
 #include "RHIUtilities.h"
 
 #include "ShaderCompiler.h"
@@ -1070,12 +1069,6 @@ FDynamicRHI* FD3D12DynamicRHIModule::CreateRHI(ERHIFeatureLevel::Type RequestedF
 	UE_LOG(LogD3D12RHI, Display, TEXT("Creating D3D12 RHI with Max Feature Level %s"), *FeatureLevelString);
 
 	GD3D12RHI = new FD3D12DynamicRHI(ChosenAdapters, bPixEventEnabled);
-#if ENABLE_RHI_VALIDATION
-	if (FParse::Param(FCommandLine::Get(), TEXT("RHIValidation")))
-	{
-		return new FValidationRHI(GD3D12RHI);
-	}
-#endif
 	return GD3D12RHI;
 }
 
diff --git a/Engine/Source/Runtime/NullDrv/NullDrv.Build.cs b/Engine/Source/Runtime/NullDrv/NullDrv.Build.cs
--- a/Engine/Source/Runtime/NullDrv/NullDrv.Build.cs
+++ b/Engine/Source/Runtime/NullDrv/NullDrv.Build.cs
@@ -2,6 +2,7 @@
 
 using UnrealBuildTool;
 
+[RHIModule]
 public class NullDrv : ModuleRules
 {
 	public NullDrv(ReadOnlyTargetRules Target) : base(Target)
diff --git a/Engine/Source/Runtime/NullDrv/Private/NullRHI.cpp b/Engine/Source/Runtime/NullDrv/Private/NullRHI.cpp
--- a/Engine/Source/Runtime/NullDrv/Private/NullRHI.cpp
+++ b/Engine/Source/Runtime/NullDrv/Private/NullRHI.cpp
@@ -87,8 +87,14 @@ void* FNullDynamicRHI::GetStaticTextureBuffer(int32 SizeX, int32 SizeY, EPixelFo
 This is used to get the same behavior as the old *_OR_IMMEDIATE present modes. */
 uint32 GPresentImmediateThreshold = 100;
 
+// Suppress linker warning "warning LNK4221: no public symbols found; archive member will be inaccessible"
+int32 NullRHILinkerHelper;
 
+#if WITH_FIXED_RHI_CLASS
 
+	#define INTERNAL_DECORATOR(Method)         ((FNullDynamicRHI&)CmdList.GetContext()       ).FNullDynamicRHI::Method
+	#define INTERNAL_DECORATOR_COMPUTE(Method) ((FNullDynamicRHI&)CmdList.GetComputeContext()).FNullDynamicRHI::Method
 
-// Suppress linker warning "warning LNK4221: no public symbols found; archive member will be inaccessible"
-int32 NullRHILinkerHelper;
+	#include "RHICommandListCommandExecutes.inl"
+
+#endif // WITH_FIXED_RHI_CLASS
diff --git a/Engine/Source/Runtime/OpenGLDrv/OpenGLDrv.Build.cs b/Engine/Source/Runtime/OpenGLDrv/OpenGLDrv.Build.cs
--- a/Engine/Source/Runtime/OpenGLDrv/OpenGLDrv.Build.cs
+++ b/Engine/Source/Runtime/OpenGLDrv/OpenGLDrv.Build.cs
@@ -4,6 +4,7 @@ using UnrealBuildTool;
 
 [SupportedPlatforms("Linux", "Android", "LinuxArm64")]
 [SupportedPlatformGroups("Windows")]
+[RHIModule]
 public class OpenGLDrv : ModuleRules
 {
 	public OpenGLDrv(ReadOnlyTargetRules Target) : base(Target)
diff --git a/Engine/Source/Runtime/OpenGLDrv/Private/OpenGLCommands.cpp b/Engine/Source/Runtime/OpenGLDrv/Private/OpenGLCommands.cpp
--- a/Engine/Source/Runtime/OpenGLDrv/Private/OpenGLCommands.cpp
+++ b/Engine/Source/Runtime/OpenGLDrv/Private/OpenGLCommands.cpp
@@ -3107,8 +3107,11 @@ void FOpenGLDynamicRHI::RHIPostExternalCommandsReset()
 	}
 }
 
-#if PLATFORM_USES_FIXED_RHI_CLASS
-#define INTERNAL_DECORATOR(Method) ((FOpenGLDynamicRHI&)CmdList.GetContext()).FOpenGLDynamicRHI::Method
-#include "RHICommandListCommandExecutes.inl"
-#endif
+#if WITH_FIXED_RHI_CLASS
+
+	#define INTERNAL_DECORATOR(Method)         ((FOpenGLDynamicRHI&)CmdList.GetContext()       ).FOpenGLDynamicRHI::Method
+	#define INTERNAL_DECORATOR_COMPUTE(Method) ((FOpenGLDynamicRHI&)CmdList.GetComputeContext()).FOpenGLDynamicRHI::Method
+
+	#include "RHICommandListCommandExecutes.inl"
 
+#endif // WITH_FIXED_RHI_CLASS
diff --git a/Engine/Source/Runtime/RHI/Private/DynamicRHI.cpp b/Engine/Source/Runtime/RHI/Private/DynamicRHI.cpp
--- a/Engine/Source/Runtime/RHI/Private/DynamicRHI.cpp
+++ b/Engine/Source/Runtime/RHI/Private/DynamicRHI.cpp
@@ -16,6 +16,7 @@
 #include "PipelineStateCache.h"
 #include "RHI.h"
 #include "RHIFwd.h"
+#include "RHIValidation.h"
 #include "TextureProfiler.h"
 #include "DataDrivenShaderPlatformInfo.h"
 #include "RHICommandList.h"
@@ -29,10 +30,6 @@ IMPLEMENT_TYPE_LAYOUT(FRayTracingGeometrySegment);
 static_assert(sizeof(FRayTracingGeometryInstance) <= 104,
 	"Ray tracing instance descriptor is expected to be no more than 104 bytes, as there may be a very large number of them.");
 
-#ifndef PLATFORM_ALLOW_NULL_RHI
-	#define PLATFORM_ALLOW_NULL_RHI		0
-#endif
-
 // Globals.
 FDynamicRHI* GDynamicRHI = NULL;
 RHIGetGPUUsageType RHIGetGPUUsage = nullptr;
@@ -66,6 +63,7 @@ static TAutoConsoleVariable<int32> CVarDisableDriverWarningPopupIfGFN(
 	TEXT("If non-zero, disable driver version warning popup if running on a GFN cloud machine."),
 	ECVF_RenderThreadSafe);
 
+#if WITH_NULL_RHI
 void InitNullRHI()
 {
 	// Use the null RHI if it was specified on the command line, or if a commandlet is running.
@@ -86,6 +84,7 @@ void InitNullRHI()
 	// Update the crash context analytics
 	FGenericCrashContext::SetEngineData(TEXT("RHI.RHIName"), TEXT("NullRHI"));
 }
+#endif
 
 #if PLATFORM_WINDOWS || PLATFORM_UNIX
 static void RHIDetectAndWarnOfBadDrivers(bool bHasEditorToken)
@@ -300,6 +299,10 @@ static void RHIDetectAndWarnOfBadDrivers(bool bHasEditorToken)
 }
 #endif // PLATFORM_WINDOWS
 
+#if USE_NULL_RHI && !WITH_NULL_RHI
+#error NullRHI is forced by the USE_NULL_RHI build option, but the NullRHI module is not part of the build.
+#endif
+
 void RHIInit(bool bHasEditorToken)
 {
 	if (!GDynamicRHI)
@@ -313,80 +316,85 @@ void RHIInit(bool bHasEditorToken)
 
 		GRHICommandList.LatchBypass(); // read commandline for bypass flag
 
+#if WITH_NULL_RHI
 		if (!FApp::CanEverRender())
 		{
 			InitNullRHI();
 		}
 		else
+#endif // WITH_NULL_RHI
 		{
 			LLM_SCOPE(ELLMTag::RHIMisc);
 
 			GDynamicRHI = PlatformCreateDynamicRHI();
-			if (GDynamicRHI)
+			if (!GDynamicRHI)
+			{
+				UE_LOG(LogRHI, Fatal, TEXT("Failed to create platform RHI."));
+			}
+
+#if ENABLE_RHI_VALIDATION
+			if (FParse::Param(FCommandLine::Get(), TEXT("RHIValidation")))
 			{
+				UE_LOG(LogRHI, Log, TEXT("Creating RHI validation layer wrapper."));
+				GDynamicRHI = new FValidationRHI(GDynamicRHI);
+			}
+#endif // ENABLE_RHI_VALIDATION
+
 #if PLATFORM_WINDOWS || PLATFORM_MAC || PLATFORM_UNIX
 
-				// Get driver version. Creating GDynamicRHI is expected to set GRHIAdapterName.
-				FGPUDriverInfo GPUDriverInfo = FPlatformMisc::GetGPUDriverInfo(GRHIAdapterName);
-				if (GPUDriverInfo.IsValid())
-				{
-					// GetGPUDriverInfo is not implemented on Linux, so it returns an invalid driver info object. However, the FVulkanDynamicRHI constructor
-					// sets these values on that platform, so we'll still have data we can log.
-					GRHIAdapterUserDriverVersion = GPUDriverInfo.UserDriverVersion;
-					GRHIAdapterInternalDriverVersion = GPUDriverInfo.InternalDriverVersion;
-					GRHIAdapterDriverDate = GPUDriverInfo.DriverDate;
-				}
+			// Get driver version. Creating GDynamicRHI is expected to set GRHIAdapterName.
+			FGPUDriverInfo GPUDriverInfo = FPlatformMisc::GetGPUDriverInfo(GRHIAdapterName);
+			if (GPUDriverInfo.IsValid())
+			{
+				// GetGPUDriverInfo is not implemented on Linux, so it returns an invalid driver info object. However, the FVulkanDynamicRHI constructor
+				// sets these values on that platform, so we'll still have data we can log.
+				GRHIAdapterUserDriverVersion = GPUDriverInfo.UserDriverVersion;
+				GRHIAdapterInternalDriverVersion = GPUDriverInfo.InternalDriverVersion;
+				GRHIAdapterDriverDate = GPUDriverInfo.DriverDate;
+			}
 
-				UE_LOG(LogRHI, Log, TEXT("RHI Adapter Info:"));
-				UE_LOG(LogRHI, Log, TEXT("            Name: %s"), *GRHIAdapterName);
-				UE_LOG(LogRHI, Log, TEXT("  Driver Version: %s (internal:%s, unified:%s)"), *GRHIAdapterUserDriverVersion, *GRHIAdapterInternalDriverVersion, *GPUDriverInfo.GetUnifiedDriverVersion());
-				UE_LOG(LogRHI, Log, TEXT("     Driver Date: %s"), *GRHIAdapterDriverDate);
+			UE_LOG(LogRHI, Log, TEXT("RHI Adapter Info:"));
+			UE_LOG(LogRHI, Log, TEXT("            Name: %s"), *GRHIAdapterName);
+			UE_LOG(LogRHI, Log, TEXT("  Driver Version: %s (internal:%s, unified:%s)"), *GRHIAdapterUserDriverVersion, *GRHIAdapterInternalDriverVersion, *GPUDriverInfo.GetUnifiedDriverVersion());
+			UE_LOG(LogRHI, Log, TEXT("     Driver Date: %s"), *GRHIAdapterDriverDate);
 
-				RHIDetectAndWarnOfBadDrivers(bHasEditorToken);
+			RHIDetectAndWarnOfBadDrivers(bHasEditorToken);
 #endif
 
-				GDynamicRHI->Init();
+			GDynamicRHI->Init();
 
-				// Validation of contexts.
-				GRHICommandList.GetImmediateCommandList().GetContext();
-				check(GIsRHIInitialized);
+			// Validation of contexts.
+			GRHICommandList.GetImmediateCommandList().GetContext();
+			check(GIsRHIInitialized);
 
-				// Set default GPU mask to all GPUs. This is necessary to ensure that any commands
-				// that create and initialize resources are executed on all GPUs. Scene rendering
-				// will restrict itself to a subset of GPUs as needed.
-				GRHICommandList.GetImmediateCommandList().SetGPUMask(FRHIGPUMask::All());
+			// Set default GPU mask to all GPUs. This is necessary to ensure that any commands
+			// that create and initialize resources are executed on all GPUs. Scene rendering
+			// will restrict itself to a subset of GPUs as needed.
+			GRHICommandList.GetImmediateCommandList().SetGPUMask(FRHIGPUMask::All());
 
-				FString FeatureLevelString;
-				GetFeatureLevelName(GMaxRHIFeatureLevel, FeatureLevelString);
+			FString FeatureLevelString;
+			GetFeatureLevelName(GMaxRHIFeatureLevel, FeatureLevelString);
 
-				if (bHasEditorToken && GMaxRHIFeatureLevel < ERHIFeatureLevel::SM5)
-				{
-					FString ShaderPlatformString = LegacyShaderPlatformToShaderFormat(GetFeatureLevelShaderPlatform(GMaxRHIFeatureLevel)).ToString();
-					FString Error = FString::Printf(TEXT("A Feature Level 5 video card is required to run the editor.\nAvailableFeatureLevel = %s, ShaderPlatform = %s"), *FeatureLevelString, *ShaderPlatformString);
-					FMessageDialog::Open(EAppMsgType::Ok, FText::FromString(Error));
-					FPlatformMisc::RequestExit(true, TEXT("RHIInit"));
-				}
+			if (bHasEditorToken && GMaxRHIFeatureLevel < ERHIFeatureLevel::SM5)
+			{
+				FString ShaderPlatformString = LegacyShaderPlatformToShaderFormat(GetFeatureLevelShaderPlatform(GMaxRHIFeatureLevel)).ToString();
+				FString Error = FString::Printf(TEXT("A Feature Level 5 video card is required to run the editor.\nAvailableFeatureLevel = %s, ShaderPlatform = %s"), *FeatureLevelString, *ShaderPlatformString);
+				FMessageDialog::Open(EAppMsgType::Ok, FText::FromString(Error));
+				FPlatformMisc::RequestExit(true, TEXT("RHIInit"));
+			}
 
-				// Update the crash context analytics
-				FGenericCrashContext::SetEngineData(TEXT("RHI.RHIName"), GDynamicRHI ? (GMaxRHIFeatureLevel == ERHIFeatureLevel::ES3_1 ? FString(GDynamicRHI->GetName()) + TEXT("_ES31") : GDynamicRHI->GetName()) : TEXT("Unknown"));
-				FGenericCrashContext::SetEngineData(TEXT("RHI.AdapterName"), GRHIAdapterName);
-				FGenericCrashContext::SetEngineData(TEXT("RHI.UserDriverVersion"), GRHIAdapterUserDriverVersion);
-				FGenericCrashContext::SetEngineData(TEXT("RHI.InternalDriverVersion"), GRHIAdapterInternalDriverVersion);
-				FGenericCrashContext::SetEngineData(TEXT("RHI.DriverDate"), GRHIAdapterDriverDate);
-				FGenericCrashContext::SetEngineData(TEXT("RHI.FeatureLevel"), FeatureLevelString);
-				FGenericCrashContext::SetEngineData(TEXT("RHI.GPUVendor"), RHIVendorIdToString());
-				FGenericCrashContext::SetEngineData(TEXT("RHI.DeviceId"), FString::Printf(TEXT("%04X"), GRHIDeviceId));
+			// Update the crash context analytics
+			FGenericCrashContext::SetEngineData(TEXT("RHI.RHIName"), GDynamicRHI ? (GMaxRHIFeatureLevel == ERHIFeatureLevel::ES3_1 ? FString(GDynamicRHI->GetName()) + TEXT("_ES31") : GDynamicRHI->GetName()) : TEXT("Unknown"));
+			FGenericCrashContext::SetEngineData(TEXT("RHI.AdapterName"), GRHIAdapterName);
+			FGenericCrashContext::SetEngineData(TEXT("RHI.UserDriverVersion"), GRHIAdapterUserDriverVersion);
+			FGenericCrashContext::SetEngineData(TEXT("RHI.InternalDriverVersion"), GRHIAdapterInternalDriverVersion);
+			FGenericCrashContext::SetEngineData(TEXT("RHI.DriverDate"), GRHIAdapterDriverDate);
+			FGenericCrashContext::SetEngineData(TEXT("RHI.FeatureLevel"), FeatureLevelString);
+			FGenericCrashContext::SetEngineData(TEXT("RHI.GPUVendor"), RHIVendorIdToString());
+			FGenericCrashContext::SetEngineData(TEXT("RHI.DeviceId"), FString::Printf(TEXT("%04X"), GRHIDeviceId));
 
 #if TEXTURE_PROFILER_ENABLED
-				FTextureProfiler::Get()->Init();
-#endif
-			}
-#if PLATFORM_ALLOW_NULL_RHI
-			else
-			{
-				// If the platform supports doing so, fall back to the NULL RHI on failure
-				InitNullRHI();
-			}
+			FTextureProfiler::Get()->Init();
 #endif
 		}
 
diff --git a/Engine/Source/Runtime/RHI/Private/RHICommandList.cpp b/Engine/Source/Runtime/RHI/Private/RHICommandList.cpp
--- a/Engine/Source/Runtime/RHI/Private/RHICommandList.cpp
+++ b/Engine/Source/Runtime/RHI/Private/RHICommandList.cpp
@@ -32,7 +32,7 @@ UE_TRACE_CHANNEL_DEFINE(RHICommandsChannel);
 bool FScopedUniformBufferStaticBindings::bRecursionGuard = false;
 #endif
 
-#if !PLATFORM_USES_FIXED_RHI_CLASS
+#if !WITH_FIXED_RHI_CLASS
 #include "RHICommandListCommandExecutes.inl"
 #endif
 
diff --git a/Engine/Source/Runtime/RHI/Public/RHICommandList.h b/Engine/Source/Runtime/RHI/Public/RHICommandList.h
--- a/Engine/Source/Runtime/RHI/Public/RHICommandList.h
+++ b/Engine/Source/Runtime/RHI/Public/RHICommandList.h
@@ -35,6 +35,7 @@
 
 #include "DynamicRHI.h"
 #include "RHITypes.h"
+#include "RHIGlobals.h"
 
 CSV_DECLARE_CATEGORY_MODULE_EXTERN(RHI_API, RHITStalls);
 CSV_DECLARE_CATEGORY_MODULE_EXTERN(RHI_API, RHITFlushes);
diff --git a/Engine/Source/Runtime/RHI/Public/RHICommandListCommandExecutes.inl b/Engine/Source/Runtime/RHI/Public/RHICommandListCommandExecutes.inl
--- a/Engine/Source/Runtime/RHI/Public/RHICommandListCommandExecutes.inl
+++ b/Engine/Source/Runtime/RHI/Public/RHICommandListCommandExecutes.inl
@@ -4,15 +4,6 @@
 	RHICommandListCommandExecutes.inl: RHI Command List execute functions.
 =============================================================================*/
 
-#if !defined(INTERNAL_DECORATOR)
-	#define INTERNAL_DECORATOR(Method) CmdList.GetContext().Method
-#endif
-
-//for functions where the signatures do not match between gfx and compute commandlists
-#if !defined(INTERNAL_DECORATOR_COMPUTE)
-#define INTERNAL_DECORATOR_COMPUTE(Method) CmdList.GetComputeContext().Method
-#endif
-
 #include "RHIResourceUpdates.h"
 
 class FRHICommandListBase;
@@ -60,6 +51,23 @@ struct FRHICommandSetRayTracingBindings;
 
 template <typename TRHIShader> struct FRHICommandSetLocalUniformBuffer;
 
+//
+// If MGPU is enabled, we can't de-virtualize the RHI command calls,
+// since some RHIs use separate context types for single and multi GPU.
+//
+#if WITH_MGPU
+	#undef INTERNAL_DECORATOR
+	#undef INTERNAL_DECORATOR_COMPUTE
+#endif
+
+#ifndef INTERNAL_DECORATOR
+	#define INTERNAL_DECORATOR(Method) CmdList.GetContext().Method
+#endif
+
+#ifndef INTERNAL_DECORATOR_COMPUTE
+	#define INTERNAL_DECORATOR_COMPUTE(Method) CmdList.GetComputeContext().Method
+#endif
+
 #if WITH_MGPU
 void FRHICommandSetGPUMask::Execute(FRHICommandListBase& CmdList)
 {
@@ -640,3 +648,6 @@ void FRHICommandDiscardRenderTargets::Execute(FRHICommandListBase& CmdList)
 	RHISTAT(RHIDiscardRenderTargets);
 	INTERNAL_DECORATOR(RHIDiscardRenderTargets)(Depth, Stencil, ColorBitMask);
 }
+
+#undef INTERNAL_DECORATOR
+#undef INTERNAL_DECORATOR_COMPUTE
diff --git a/Engine/Source/Runtime/RHI/Public/RHIValidationContext.h b/Engine/Source/Runtime/RHI/Public/RHIValidationContext.h
--- a/Engine/Source/Runtime/RHI/Public/RHIValidationContext.h
+++ b/Engine/Source/Runtime/RHI/Public/RHIValidationContext.h
@@ -8,6 +8,7 @@
 
 #include "RHIValidationCommon.h"
 #include "RHIValidationUtils.h"
+#include "RHIValidation.h"
 
 #if ENABLE_RHI_VALIDATION
 
diff --git a/Engine/Source/Runtime/RHI/RHI.Build.cs b/Engine/Source/Runtime/RHI/RHI.Build.cs
--- a/Engine/Source/Runtime/RHI/RHI.Build.cs
+++ b/Engine/Source/Runtime/RHI/RHI.Build.cs
@@ -13,40 +13,35 @@ public class RHI : ModuleRules
 
 		if (Target.bCompileAgainstEngine)
 		{
-			DynamicallyLoadedModuleNames.Add("NullDrv");
-
-			if (Target.Type != TargetRules.TargetType.Server)   // Dedicated servers should skip loading everything but NullDrv
+			if (Target.Type == TargetRules.TargetType.Server)
+			{
+				// Dedicated servers should skip loading everything but NullDrv
+				DynamicallyLoadedModuleNames.Add("NullDrv");
+			}
+			else
 			{
 				if (Target.Platform.IsInGroup(UnrealPlatformGroup.Desktop))
                 {
 					PublicDefinitions.Add("RHI_WANT_BREADCRUMB_EVENTS=1");
+					DynamicallyLoadedModuleNames.Add("NullDrv");
 				}
 
 				if (Target.Configuration != UnrealTargetConfiguration.Shipping && Target.Configuration != UnrealTargetConfiguration.Test)
                 {
 					PublicDefinitions.Add("RHI_WANT_RESOURCE_INFO=1");
-                }
-
-				// UEBuildAndroid.cs adds VulkanRHI for Android builds if it is enabled
-				if (Target.Platform.IsInGroup(UnrealPlatformGroup.Windows))
-				{
-					DynamicallyLoadedModuleNames.Add("D3D11RHI");
 				}
-
+				
 				if (Target.Platform.IsInGroup(UnrealPlatformGroup.Windows))
 				{
-					//#todo-rco: D3D12 requires different SDK headers not compatible with WinXP
+					DynamicallyLoadedModuleNames.Add("D3D11RHI");
 					DynamicallyLoadedModuleNames.Add("D3D12RHI");
-				}
-
-				if (Target.Platform.IsInGroup(UnrealPlatformGroup.Windows) || Target.IsInPlatformGroup(UnrealPlatformGroup.Unix))
-				{
+					DynamicallyLoadedModuleNames.Add("OpenGLDrv");
 					DynamicallyLoadedModuleNames.Add("VulkanRHI");
 				}
 
-				if ((Target.Platform.IsInGroup(UnrealPlatformGroup.Windows)) ||
-					(Target.IsInPlatformGroup(UnrealPlatformGroup.Linux) && Target.Type != TargetRules.TargetType.Server))  // @todo should servers on all platforms skip this?
+				if (Target.IsInPlatformGroup(UnrealPlatformGroup.Unix))
 				{
+					DynamicallyLoadedModuleNames.Add("VulkanRHI");
 					DynamicallyLoadedModuleNames.Add("OpenGLDrv");
 				}
 			}
diff --git a/Engine/Source/Runtime/VulkanRHI/Private/VulkanCommands.cpp b/Engine/Source/Runtime/VulkanRHI/Private/VulkanCommands.cpp
--- a/Engine/Source/Runtime/VulkanRHI/Private/VulkanCommands.cpp
+++ b/Engine/Source/Runtime/VulkanRHI/Private/VulkanCommands.cpp
@@ -981,3 +981,13 @@ void FVulkanDynamicRHI::RHISubmitCommandLists(TArrayView<IRHIPlatformCommandList
 		delete PlatformCmdList;
 	}
 }
+
+#if WITH_FIXED_RHI_CLASS
+
+	#define INTERNAL_DECORATOR(Method)         ((FVulkanCommandListContext&)CmdList.GetContext()       ).FVulkanCommandListContext::Method
+	#define INTERNAL_DECORATOR_COMPUTE(Method) ((FVulkanCommandListContext&)CmdList.GetComputeContext()).FVulkanCommandListContext::Method
+
+	#include "RHICommandListCommandExecutes.inl"
+
+#endif // WITH_FIXED_RHI_CLASS
+
diff --git a/Engine/Source/Runtime/VulkanRHI/Private/VulkanRHI.cpp b/Engine/Source/Runtime/VulkanRHI/Private/VulkanRHI.cpp
--- a/Engine/Source/Runtime/VulkanRHI/Private/VulkanRHI.cpp
+++ b/Engine/Source/Runtime/VulkanRHI/Private/VulkanRHI.cpp
@@ -20,14 +20,13 @@
 #include "VulkanLLM.h"
 #include "Misc/EngineVersion.h"
 #include "GlobalShader.h"
-#include "RHIValidation.h"
 #include "RHIUtilities.h"
 #include "IHeadMountedDisplayModule.h"
 #include "VulkanRenderpass.h"
 #include "VulkanTransientResourceAllocator.h"
 #include "VulkanExtensions.h"
 #include "VulkanRayTracing.h"
-
+#include "DataDrivenShaderPlatformInfo.h"
 
 
 // Use Vulkan Profiles to verify feature level support on startup
@@ -371,16 +370,7 @@ FDynamicRHI* FVulkanDynamicRHIModule::CreateRHI(ERHIFeatureLevel::Type InRequest
 	checkf(GMaxRHIShaderPlatform != SP_NumPlatforms, TEXT("Requested feature level [%s] mapped to unsupported shader platform!"), *LexToString(InRequestedFeatureLevel));
 
 	GVulkanRHI = new FVulkanDynamicRHI();
-	FDynamicRHI* FinalRHI = GVulkanRHI;
-
-#if ENABLE_RHI_VALIDATION
-	if (FParse::Param(FCommandLine::Get(), TEXT("RHIValidation")))
-	{
-		FinalRHI = new FValidationRHI(FinalRHI);
-	}
-#endif
-
-	return FinalRHI;
+	return GVulkanRHI;
 }
 
 IMPLEMENT_MODULE(FVulkanDynamicRHIModule, VulkanRHI);
diff --git a/Engine/Source/Runtime/VulkanRHI/Private/VulkanRHIPrivate.h b/Engine/Source/Runtime/VulkanRHI/Private/VulkanRHIPrivate.h
--- a/Engine/Source/Runtime/VulkanRHI/Private/VulkanRHIPrivate.h
+++ b/Engine/Source/Runtime/VulkanRHI/Private/VulkanRHIPrivate.h
@@ -11,7 +11,6 @@
 #include "Misc/ScopeLock.h"
 #include "RHI.h"
 #include "RenderUtils.h"
-#include "RHIValidation.h"
 
 // let the platform set up the headers and some defines
 #include "VulkanPlatform.h"
diff --git a/Engine/Source/Runtime/VulkanRHI/VulkanRHI.Build.cs b/Engine/Source/Runtime/VulkanRHI/VulkanRHI.Build.cs
--- a/Engine/Source/Runtime/VulkanRHI/VulkanRHI.Build.cs
+++ b/Engine/Source/Runtime/VulkanRHI/VulkanRHI.Build.cs
@@ -4,6 +4,7 @@ using UnrealBuildTool;
 using System;
 using System.IO;
 
+[RHIModule]
 public class VulkanRHI : ModuleRules
 {
 	public VulkanRHI(ReadOnlyTargetRules Target) : base(Target)
diff --git a/Engine/Source/Runtime/Windows/D3D11RHI/D3D11RHI.Build.cs b/Engine/Source/Runtime/Windows/D3D11RHI/D3D11RHI.Build.cs
--- a/Engine/Source/Runtime/Windows/D3D11RHI/D3D11RHI.Build.cs
+++ b/Engine/Source/Runtime/Windows/D3D11RHI/D3D11RHI.Build.cs
@@ -4,6 +4,7 @@ using System.IO;
 using UnrealBuildTool;
 
 [SupportedPlatformGroups("Windows")]
+[RHIModule]
 public class D3D11RHI : ModuleRules
 {
 	protected virtual bool bIncludeExtensions { get => true; }
diff --git a/Engine/Source/Runtime/Windows/D3D11RHI/Private/D3D11Commands.cpp b/Engine/Source/Runtime/Windows/D3D11RHI/Private/D3D11Commands.cpp
--- a/Engine/Source/Runtime/Windows/D3D11RHI/Private/D3D11Commands.cpp
+++ b/Engine/Source/Runtime/Windows/D3D11RHI/Private/D3D11Commands.cpp
@@ -1762,3 +1762,12 @@ void FD3D11DynamicRHI::RHIUnlockStagingBuffer(FRHIStagingBuffer* StagingBufferRH
 	FD3D11StagingBuffer* StagingBuffer = ResourceCast(StagingBufferRHI);
 	StagingBuffer->Unlock();
 }
+
+#if WITH_FIXED_RHI_CLASS
+
+	#define INTERNAL_DECORATOR(Method)         ((FD3D11DynamicRHI&)CmdList.GetContext()       ).FD3D11DynamicRHI::Method
+	#define INTERNAL_DECORATOR_COMPUTE(Method) ((FD3D11DynamicRHI&)CmdList.GetComputeContext()).FD3D11DynamicRHI::Method
+
+	#include "RHICommandListCommandExecutes.inl"
+
+#endif // WITH_FIXED_RHI_CLASS
diff --git a/Engine/Source/Runtime/Windows/D3D11RHI/Private/Windows/WindowsD3D11Device.cpp b/Engine/Source/Runtime/Windows/D3D11RHI/Private/Windows/WindowsD3D11Device.cpp
--- a/Engine/Source/Runtime/Windows/D3D11RHI/Private/Windows/WindowsD3D11Device.cpp
+++ b/Engine/Source/Runtime/Windows/D3D11RHI/Private/Windows/WindowsD3D11Device.cpp
@@ -23,7 +23,6 @@
 #include "IHeadMountedDisplayModule.h"
 #include "GenericPlatform/GenericPlatformDriver.h"			// FGPUDriverInfo
 #include "GenericPlatform/GenericPlatformCrashContext.h"
-#include "RHIValidation.h"
 #include "RHIUtilities.h"
 #include "HAL/ExceptionHandling.h"
 #include "HDRHelper.h"
@@ -1208,13 +1207,6 @@ FDynamicRHI* FD3D11DynamicRHIModule::CreateRHI(ERHIFeatureLevel::Type RequestedF
 	GD3D11RHI = new FD3D11DynamicRHI(DXGIFactory1, ChosenAdapter.MaxSupportedFeatureLevel,ChosenAdapter);
 	FDynamicRHI* FinalRHI = GD3D11RHI;
 
-#if ENABLE_RHI_VALIDATION
-	if (FParse::Param(FCommandLine::Get(), TEXT("RHIValidation")))
-	{
-		FinalRHI = new FValidationRHI(FinalRHI);
-	}
-#endif
-
 	FGenericCrashContext::SetEngineData(TEXT("RHI.IntegratedGPU"), ChosenAdapter.bIsIntegrated ? TEXT("true") : TEXT("false"));
 
 	return FinalRHI;
