Increase the minimum nvidia driver version to 526.47, because earlier versions can cause a GPU crash when raytracing is used.

Also changed the driver warning message box to say that the version specified in the ini is the minimum required, rather than "recommended", since that was causing confusion. We want people to use the latest driver as much as possible; the version we suggest is just the earliest without known bugs.

#rnx

diff --git a/Engine/Config/BaseHardware.ini b/Engine/Config/BaseHardware.ini
--- a/Engine/Config/BaseHardware.ini
+++ b/Engine/Config/BaseHardware.ini
@@ -36,19 +36,19 @@
 
 
 [GPU_NVIDIA Windows]
-SuggestedDriverVersion="516.25"
+SuggestedDriverVersion="526.47"
 
 ; To prevent problems with older drivers.
 +DriverDenyList=(DriverVersion="<473.47", Reason="These driver versions have known stability issues and missing features", RHI="D3D11")
-+DriverDenyList=(DriverVersion="<516.01", Reason="These driver versions have known stability issues and missing features", RHI="D3D12")
-+DriverDenyList=(DriverVersion="<516.01", Reason="These driver versions have known stability issues and missing features", RHI="Vulkan")
++DriverDenyList=(DriverVersion="<526.47", Reason="These driver versions have known stability issues and missing features", RHI="D3D12")
++DriverDenyList=(DriverVersion="<526.47", Reason="These driver versions have known stability issues and missing features", RHI="Vulkan")
 
 ;-----------------------------------------------------------------------------------------------------------------
 
 [GPU_NVIDIA Linux]
 SuggestedDriverVersion="515.48"
 
-; UE5 in Linux/Vulkan requires recent drivers for Nanite/Lumen and has other stability issues with older drivers (need >= 495.46)
+; UE5 in Linux/Vulkan requires recent drivers for Nanite/Lumen and has other stability issues with older drivers.
 +DriverDenyList=(DriverVersion="<515.48", Reason="These driver versions have known stability issues and missing features in Vulkan")
 
 
diff --git a/Engine/Source/Runtime/RHI/Private/DynamicRHI.cpp b/Engine/Source/Runtime/RHI/Private/DynamicRHI.cpp
--- a/Engine/Source/Runtime/RHI/Private/DynamicRHI.cpp
+++ b/Engine/Source/Runtime/RHI/Private/DynamicRHI.cpp
@@ -241,11 +241,11 @@ static void RHIDetectAndWarnOfBadDrivers(bool bHasEditorToken)
 			{
 				if (!DenyListEntry.RHIName.IsEmpty())
 				{
-					LocalizedMsg = FText::Format(NSLOCTEXT("MessageDialog", "VideoCardDriverRHIIssueReport", "The installed version of the {Vendor} graphics driver has known issues in {RHI}.\nPlease install either the latest or the recommended driver version or switch to a different rendering API.\n\nWould you like to visit the following URL to download the driver?\n\n{Hyperlink}\n\n{AdapterName}\nInstalled: {InstalledVer}\nRecommended: {RecommendedVer}"), Args);
+					LocalizedMsg = FText::Format(NSLOCTEXT("MessageDialog", "VideoCardDriverRHIIssueReport", "The installed version of the {Vendor} graphics driver has known issues in {RHI}.\nPlease install the latest driver version or switch to a different rendering API.\n\nWould you like to visit the following URL to download the driver?\n\n{Hyperlink}\n\n{AdapterName}\nInstalled: {InstalledVer}\nMinimum required: {RecommendedVer}"), Args);
 				}
 				else
 				{
-					LocalizedMsg = FText::Format(NSLOCTEXT("MessageDialog", "VideoCardDriverIssueReport", "The installed version of the {Vendor} graphics driver has known issues.\nPlease install either the latest or the recommended driver version.\n\nWould you like to visit the following URL to download the driver?\n\n{Hyperlink}\n\n{AdapterName}\nInstalled: {InstalledVer}\nRecommended: {RecommendedVer}"), Args);
+					LocalizedMsg = FText::Format(NSLOCTEXT("MessageDialog", "VideoCardDriverIssueReport", "The installed version of the {Vendor} graphics driver has known issues.\nPlease install the latest driver version.\n\nWould you like to visit the following URL to download the driver?\n\n{Hyperlink}\n\n{AdapterName}\nInstalled: {InstalledVer}\nMinimum required: {RecommendedVer}"), Args);
 				}
 			}
 
