Changed FRHIRenderTargetView to use member initializers and compiler-generated moves and copies, to avoid warnings about mismatched copy constructors and assignments.

diff --git a/Engine/Source/Runtime/RHI/Public/RHIResources.h b/Engine/Source/Runtime/RHI/Public/RHIResources.h
--- a/Engine/Source/Runtime/RHI/Public/RHIResources.h
+++ b/Engine/Source/Runtime/RHI/Public/RHIResources.h
@@ -2639,30 +2639,20 @@ public:
 class FRHIRenderTargetView
 {
 public:
-	FRHITexture* Texture;
-	uint32 MipIndex;
+	FRHITexture* Texture = nullptr;
+	uint32 MipIndex = 0;
 
 	/** Array slice or texture cube face.  Only valid if texture resource was created with TexCreate_TargetArraySlicesIndependently! */
-	uint32 ArraySliceIndex;
-	
-	ERenderTargetLoadAction LoadAction;
-	ERenderTargetStoreAction StoreAction;
+	uint32 ArraySliceIndex = -1;
 
-	FRHIRenderTargetView() : 
-		Texture(NULL),
-		MipIndex(0),
-		ArraySliceIndex(-1),
-		LoadAction(ERenderTargetLoadAction::ENoAction),
-		StoreAction(ERenderTargetStoreAction::ENoAction)
-	{}
+	ERenderTargetLoadAction LoadAction = ERenderTargetLoadAction::ENoAction;
+	ERenderTargetStoreAction StoreAction = ERenderTargetStoreAction::ENoAction;
 
-	FRHIRenderTargetView(const FRHIRenderTargetView& Other) :
-		Texture(Other.Texture),
-		MipIndex(Other.MipIndex),
-		ArraySliceIndex(Other.ArraySliceIndex),
-		LoadAction(Other.LoadAction),
-		StoreAction(Other.StoreAction)
-	{}
+	FRHIRenderTargetView() = default;
+	FRHIRenderTargetView(FRHIRenderTargetView&&) = default;
+	FRHIRenderTargetView(const FRHIRenderTargetView&) = default;
+	FRHIRenderTargetView& operator=(FRHIRenderTargetView&&) = default;
+	FRHIRenderTargetView& operator=(const FRHIRenderTargetView&) = default;
 
 	//common case
 	explicit FRHIRenderTargetView(FRHITexture* InTexture, ERenderTargetLoadAction InLoadAction) :
