Fix crash in RHI Validation with collecting breadcrumb paths.

#rnx

diff --git a/Engine/Source/Runtime/RHI/Private/RHIValidation.cpp b/Engine/Source/Runtime/RHI/Private/RHIValidation.cpp
--- a/Engine/Source/Runtime/RHI/Private/RHIValidation.cpp
+++ b/Engine/Source/Runtime/RHI/Private/RHIValidation.cpp
@@ -638,15 +638,7 @@ thread_local TConstArrayView<const TCHAR*> FRHIValidationBreadcrumbScope::Breadc
 
 static FString GetBreadcrumbPath()
 {
-	FString BreadcrumbMessage;
-
-	for (int32 Index = 0; Index < FRHIValidationBreadcrumbScope::Breadcrumbs.Num() - 1; ++Index)
-	{
-		BreadcrumbMessage += (FRHIValidationBreadcrumbScope::Breadcrumbs[Index] + FString(TEXT("/")));
-	}
-
-	BreadcrumbMessage += FRHIValidationBreadcrumbScope::Breadcrumbs.Last();
-	return BreadcrumbMessage;
+	return FString::Join(FRHIValidationBreadcrumbScope::Breadcrumbs, TEXT("/"));
 }
 
 // FlushType: Thread safe
