Re-add FPipelineFileCacheRayTracingDesc::MaxPayloadSizeInBytes as DeprecatedMaxPayloadSizeInBytes to keep sizeof(FPipelineCacheFileFormatPSO) unchanged

diff --git a/Engine/Source/Runtime/RHI/Private/PipelineFileCache.cpp b/Engine/Source/Runtime/RHI/Private/PipelineFileCache.cpp
--- a/Engine/Source/Runtime/RHI/Private/PipelineFileCache.cpp
+++ b/Engine/Source/Runtime/RHI/Private/PipelineFileCache.cpp
@@ -1330,8 +1330,8 @@ bool FPipelineCacheFileFormatPSO::Verify() const
 		{
 			Ar << Info.RayTracingDesc.ShaderHash;
 
-			uint32 DeprecatedMaxPayloadSizeInBytes = 0; // Not used, kept for binary format compatibility
-			Ar << DeprecatedMaxPayloadSizeInBytes;
+			// Not used, kept for binary format compatibility
+			Ar << Info.RayTracingDesc.DeprecatedMaxPayloadSizeInBytes;
 
 			uint32 Frequency = uint32(Info.RayTracingDesc.Frequency);
 			Ar << Frequency;
@@ -4190,7 +4190,6 @@ FString FPipelineCacheFileFormatPSO::FPipelineFileCacheRayTracingDesc::HeaderLin
 
 FString FPipelineCacheFileFormatPSO::FPipelineFileCacheRayTracingDesc::ToString() const
 {
-	const uint32 DeprecatedMaxPayloadSizeInBytes = 0; // Not used, but kept for back-compatibility
 	return FString::Printf(TEXT("%s,%d,%d,%d")
 		, *ShaderHash.ToString()
 		, DeprecatedMaxPayloadSizeInBytes
@@ -4229,7 +4228,7 @@ void FPipelineCacheFileFormatPSO::FPipelineFileCacheRayTracingDesc::FromString(c
 
 	ShaderHash.FromString(Parts[0]);
 
-	uint32 DeprecatedMaxPayloadSizeInBytes = 0; // Not used, but kept for back-compatibility
+	// Not used, but kept for back-compatibility
 	LexFromString(DeprecatedMaxPayloadSizeInBytes, Parts[1]);
 
 	{
diff --git a/Engine/Source/Runtime/RHI/Public/PipelineFileCache.h b/Engine/Source/Runtime/RHI/Public/PipelineFileCache.h
--- a/Engine/Source/Runtime/RHI/Public/PipelineFileCache.h
+++ b/Engine/Source/Runtime/RHI/Public/PipelineFileCache.h
@@ -184,6 +184,7 @@ struct RHI_API FPipelineCacheFileFormatPSO
 	struct RHI_API FPipelineFileCacheRayTracingDesc
 	{
 		FSHAHash ShaderHash;
+		uint32 DeprecatedMaxPayloadSizeInBytes = 0;
 		EShaderFrequency Frequency = SF_RayGen;
 		bool bAllowHitGroupIndexing = true;
 
