Fix non-unity build errors

diff --git a/Engine/Source/Runtime/RHI/Public/RHIValidationCommon.h b/Engine/Source/Runtime/RHI/Public/RHIValidationCommon.h
--- a/Engine/Source/Runtime/RHI/Public/RHIValidationCommon.h
+++ b/Engine/Source/Runtime/RHI/Public/RHIValidationCommon.h
@@ -21,6 +21,9 @@ const bool GRHIValidationEnabled = false;
 #if ENABLE_RHI_VALIDATION
 
 class FRHIUniformBuffer;
+class FRHIViewableResource;
+class FRHIUnorderedAccessView;
+
 class FValidationComputeContext;
 class FValidationContext;
 class FValidationRHI;
