Improved, on supported platform, Zlib DMA decompression by being multithread friendly (ie multiple DMA zlib command list can be created and submitted at the same time by multiple threads)

#rnx

diff --git a/Engine/Source/Runtime/RHI/Public/DynamicRHI.h b/Engine/Source/Runtime/RHI/Public/DynamicRHI.h
--- a/Engine/Source/Runtime/RHI/Public/DynamicRHI.h
+++ b/Engine/Source/Runtime/RHI/Public/DynamicRHI.h
@@ -760,9 +760,6 @@ public:
 	// FlushType: Flush Immediate
 	virtual bool RHIIsRenderingSuspended() { return false; };
 
-	UE_DEPRECATED(5.1, "No longer used: FCompression::UncompressMemory should be used instead")
-	virtual bool RHIEnqueueDecompress(uint8_t* SrcBuffer, uint8_t* DestBuffer, int CompressedSize, void* ErrorCodeBuffer) { return false; }
-
 	/**
 	*	Retrieve available screen resolutions.
 	*
