Add ability to start -game with a preview shader platform.

- Adds a -OverrideSP=Foo commandline argument, where Foo is e.g. AndroidVulkan_Preview or another preview shader platform (named as in .ini with _Preview appended).
- Not exactly the same as enabling a mobile preview since it does not change device profiles or other non-shader settings, but can be useful for debugging or reproing certain issues with -game instead of starting the full editor and going into preview mode.

diff --git a/Engine/Source/Runtime/RHI/Private/DynamicRHI.cpp b/Engine/Source/Runtime/RHI/Private/DynamicRHI.cpp
--- a/Engine/Source/Runtime/RHI/Private/DynamicRHI.cpp
+++ b/Engine/Source/Runtime/RHI/Private/DynamicRHI.cpp
@@ -396,6 +396,45 @@ void RHIInit(bool bHasEditorToken)
 
 #if WITH_EDITOR
 	FGenericDataDrivenShaderPlatformInfo::UpdatePreviewPlatforms();
+
+	// add an ability to override the shader platform, intended for preview platforms only
+	FString ShaderPlatformName;
+	if (FParse::Value(FCommandLine::Get(), TEXT("OverrideSP"), ShaderPlatformName) && !ShaderPlatformName.IsEmpty())
+	{
+		EShaderPlatform OverrideShaderPlatform = FDataDrivenShaderPlatformInfo::GetShaderPlatformFromName(*ShaderPlatformName);
+		if (OverrideShaderPlatform != SP_NumPlatforms)
+		{
+			// only allow to override to preview shader platform (to avoid complications), and make sure the SP matches current feature level
+			if (FDataDrivenShaderPlatformInfo::GetIsPreviewPlatform(OverrideShaderPlatform))
+			{
+				if (GetMaxSupportedFeatureLevel(OverrideShaderPlatform) == GMaxRHIFeatureLevel)
+				{
+					UE_LOG(LogRHI, Log, TEXT("Overriding shader platform to use from %s to %s"),
+						*LexToString(GMaxRHIShaderPlatform, false),
+						*LexToString(OverrideShaderPlatform, false));
+
+					GShaderPlatformForFeatureLevel[GMaxRHIFeatureLevel] = OverrideShaderPlatform;
+					GMaxRHIShaderPlatform = OverrideShaderPlatform;
+				}
+				else
+				{
+					UE_LOG(LogRHI, Log, TEXT("Cannot override to use shaderplatform %s, its max feature level %s does not match current max feature level %s"),
+						*LexToString(OverrideShaderPlatform, false),
+						*LexToString(GetMaxSupportedFeatureLevel(OverrideShaderPlatform)),
+						*LexToString(GMaxRHIFeatureLevel)
+						);
+				}
+			}
+			else
+			{
+				UE_LOG(LogRHI, Log, TEXT("Cannot override to use shader platform %s, it is not a preview platform"), *LexToString(OverrideShaderPlatform, false));
+			}
+		}
+		else
+		{
+			UE_LOG(LogRHI, Log, TEXT("Shader platform %s is not a valid platform, cannot use it."), *LexToString(OverrideShaderPlatform, false));
+		}
+	}
 #endif
 }
 
diff --git a/Engine/Source/Runtime/RHI/Private/RHIStrings.cpp b/Engine/Source/Runtime/RHI/Private/RHIStrings.cpp
--- a/Engine/Source/Runtime/RHI/Private/RHIStrings.cpp
+++ b/Engine/Source/Runtime/RHI/Private/RHIStrings.cpp
@@ -110,12 +110,16 @@ FString LexToString(EShaderPlatform Platform, bool bError)
 		}
 		else if (Platform > SP_CUSTOM_PLATFORM_FIRST && Platform < SP_CUSTOM_PLATFORM_LAST)
 		{
+#if WITH_EDITOR
+			return FString::Printf(TEXT("SP_CUSTOM_PLATFORM_%d (%s)"), Platform - SP_CUSTOM_PLATFORM_FIRST, *FDataDrivenShaderPlatformInfo::GetFriendlyName(Platform).ToString());
+#else
 			return FString::Printf(TEXT("SP_CUSTOM_PLATFORM_%d"), Platform - SP_CUSTOM_PLATFORM_FIRST);
+#endif
 		}
 		else
 		{
 			checkf(!bError, TEXT("Unknown or removed EShaderPlatform %d!"), (int32)Platform);
-			return TEXT("");
+			return FString::Printf(TEXT("Unknown or removed EShaderPlatform %d"), (int32)Platform);
 		}
 	}
 }
