Fix startup crash on Linux accessing invalid shaderplatform.

diff --git a/Engine/Source/Runtime/RHI/Private/RHIStrings.cpp b/Engine/Source/Runtime/RHI/Private/RHIStrings.cpp
--- a/Engine/Source/Runtime/RHI/Private/RHIStrings.cpp
+++ b/Engine/Source/Runtime/RHI/Private/RHIStrings.cpp
@@ -110,11 +110,7 @@ FString LexToString(EShaderPlatform Platform, bool bError)
 		}
 		else if (Platform > SP_CUSTOM_PLATFORM_FIRST && Platform < SP_CUSTOM_PLATFORM_LAST)
 		{
-#if WITH_EDITOR
-			return FString::Printf(TEXT("SP_CUSTOM_PLATFORM_%d (%s)"), Platform - SP_CUSTOM_PLATFORM_FIRST, *FDataDrivenShaderPlatformInfo::GetFriendlyName(Platform).ToString());
-#else
 			return FString::Printf(TEXT("SP_CUSTOM_PLATFORM_%d"), Platform - SP_CUSTOM_PLATFORM_FIRST);
-#endif
 		}
 		else
 		{
