Fixed mismatched copy constructor/assignment operator in FResolveRect.

diff --git a/Engine/Source/Runtime/RHI/Public/RHIResources.h b/Engine/Source/Runtime/RHI/Public/RHIResources.h
--- a/Engine/Source/Runtime/RHI/Public/RHIResources.h
+++ b/Engine/Source/Runtime/RHI/Public/RHIResources.h
@@ -4307,13 +4307,6 @@ struct FResolveRect
 		, Y2(InY2)
 	{}
 
-	FResolveRect(const FResolveRect& Other)
-		: X1(Other.X1)
-		, Y1(Other.Y1)
-		, X2(Other.X2)
-		, Y2(Other.Y2)
-	{}
-
 	explicit FResolveRect(FIntRect Other)
 		: X1(Other.Min.X)
 		, Y1(Other.Min.Y)
