Fixed mismatched copy constructor/assignment operator in FVertexElement.
Added a static assert to replace the need for the user-defined assignment operator.

diff --git a/Engine/Source/Runtime/RHI/Private/RHI.cpp b/Engine/Source/Runtime/RHI/Private/RHI.cpp
--- a/Engine/Source/Runtime/RHI/Private/RHI.cpp
+++ b/Engine/Source/Runtime/RHI/Private/RHI.cpp
@@ -23,6 +23,7 @@
 #include "RHITextureReference.h"
 #include "RHIStats.h"
 #include "RHICommandList.h"
+#include <type_traits>
 
 #if RHI_ENABLE_RESOURCE_INFO
 #include "HAL/FileManager.h"
@@ -149,6 +150,12 @@ const FClearValueBinding FClearValueBinding::DefaultNormal8Bit(FLinearColor(128.
 
 TRefCountPtr<FRHITexture> FRHITextureReference::DefaultTexture;
 
+// This is necessary to get expected results for code that zeros, assigns and then CRC's the whole struct.
+//
+// See: https://en.cppreference.com/w/cpp/types/has_unique_object_representations
+// "This trait was introduced to make it possible to determine whether a type can be correctly hashed by hashing its object representation as a byte array."
+static_assert(std::has_unique_object_representations_v<FVertexElement>, "FVertexElement should not have compiler-injected padding");
+
 FString FVertexElement::ToString() const
 {
 	return FString::Printf(TEXT("<%u %u %u %u %u %u>")
diff --git a/Engine/Source/Runtime/RHI/Public/RHI.h b/Engine/Source/Runtime/RHI/Public/RHI.h
--- a/Engine/Source/Runtime/RHI/Public/RHI.h
+++ b/Engine/Source/Runtime/RHI/Public/RHI.h
@@ -182,19 +182,6 @@ struct FVertexElement
 		Stride(InStride),
 		bUseInstanceIndex(bInUseInstanceIndex)
 	{}
-	/**
-	* Suppress the compiler generated assignment operator so that padding won't be copied.
-	* This is necessary to get expected results for code that zeros, assigns and then CRC's the whole struct.
-	*/
-	void operator=(const FVertexElement& Other)
-	{
-		StreamIndex = Other.StreamIndex;
-		Offset = Other.Offset;
-		Type = Other.Type;
-		AttributeIndex = Other.AttributeIndex;
-		Stride = Other.Stride;
-		bUseInstanceIndex = Other.bUseInstanceIndex;
-	}
 
 	bool operator==(const FVertexElement& Other) const
 	{
