VulkanRHI: Add command line switches to force sm5/sm6 in Linux.

#rnx

diff --git a/Engine/Source/Runtime/RHI/Private/Linux/LinuxDynamicRHI.cpp b/Engine/Source/Runtime/RHI/Private/Linux/LinuxDynamicRHI.cpp
--- a/Engine/Source/Runtime/RHI/Private/Linux/LinuxDynamicRHI.cpp
+++ b/Engine/Source/Runtime/RHI/Private/Linux/LinuxDynamicRHI.cpp
@@ -10,6 +10,29 @@
 #include "DataDrivenShaderPlatformInfo.h"
 #include "RHIStrings.h"
 
+
+static TOptional<ERHIFeatureLevel::Type> GetForcedFeatureLevel()
+{
+	TOptional<ERHIFeatureLevel::Type> ForcedFeatureLevel{};
+
+	if (FParse::Param(FCommandLine::Get(), TEXT("es31")) || FParse::Param(FCommandLine::Get(), TEXT("FeatureLevelES31")) || FParse::Param(FCommandLine::Get(), TEXT("FeatureLevelES3_1")))
+	{
+		ForcedFeatureLevel = ERHIFeatureLevel::ES3_1;
+	}
+
+	if (FParse::Param(FCommandLine::Get(), TEXT("sm5")))
+	{
+		ForcedFeatureLevel = ERHIFeatureLevel::SM5;
+	}
+
+	if (FParse::Param(FCommandLine::Get(), TEXT("sm6")))
+	{
+		ForcedFeatureLevel = ERHIFeatureLevel::SM6;
+	}
+
+	return ForcedFeatureLevel;
+}
+
 FDynamicRHI* PlatformCreateDynamicRHI()
 {
 	ERHIFeatureLevel::Type RequestedFeatureLevel = ERHIFeatureLevel::SM5;
@@ -35,49 +58,69 @@ FDynamicRHI* PlatformCreateDynamicRHI()
 
 	IDynamicRHIModule* DynamicRHIModule = nullptr;
 
-	TArray<FString> TargetedShaderFormats;
-	GConfig->GetArray(TEXT("/Script/LinuxTargetPlatform.LinuxTargetSettings"), TEXT("TargetedRHIs"), TargetedShaderFormats, GEngineIni);
-
-	// First come first serve
-	for (int32 SfIdx = 0; SfIdx < TargetedShaderFormats.Num(); ++SfIdx)
+	// Read in command line params to force a given feature level
+	const TOptional<ERHIFeatureLevel::Type> ForcedFeatureLevel = GetForcedFeatureLevel();
+	if (ForcedFeatureLevel.IsSet())
 	{
-		// If we are not forcing opengl and we havent failed to create a VulkanRHI try to again with a different TargetedRHI
-		if (!bForceOpenGL && !bVulkanFailed && TargetedShaderFormats[SfIdx].StartsWith(TEXT("SF_VULKAN_")))
+		RequestedFeatureLevel = ForcedFeatureLevel.GetValue();
+		DynamicRHIModule = &FModuleManager::LoadModuleChecked<IDynamicRHIModule>(TEXT("VulkanRHI"));
+		if (DynamicRHIModule->IsSupported(RequestedFeatureLevel))
 		{
-			DynamicRHIModule = &FModuleManager::LoadModuleChecked<IDynamicRHIModule>(TEXT("VulkanRHI"));
-			if (!DynamicRHIModule->IsSupported())
-			{
-				DynamicRHIModule = nullptr;
-				bVulkanFailed = true;
-			}
-			else
-			{
-				FApp::SetGraphicsRHI(TEXT("Vulkan"));
-				FPlatformApplicationMisc::UsingVulkan();
-
-				FName ShaderFormatName(*TargetedShaderFormats[SfIdx]);
-				EShaderPlatform TargetedPlatform = ShaderFormatToLegacyShaderPlatform(ShaderFormatName);
-				RequestedFeatureLevel = GetMaxSupportedFeatureLevel(TargetedPlatform);
-				break;
-			}
+			const FString RHIName = FString::Printf(TEXT("Vulkan (%s)"), *LexToString(RequestedFeatureLevel));
+			FApp::SetGraphicsRHI(RHIName);
+			FPlatformApplicationMisc::UsingVulkan();
+		}
+		else
+		{
+			DynamicRHIModule = nullptr;
 		}
-		else if (!bForceVulkan && !bOpenGLFailed && TargetedShaderFormats[SfIdx].StartsWith(TEXT("GLSL_")))
+	}
+	else
+	{
+		TArray<FString> TargetedShaderFormats;
+		GConfig->GetArray(TEXT("/Script/LinuxTargetPlatform.LinuxTargetSettings"), TEXT("TargetedRHIs"), TargetedShaderFormats, GEngineIni);
+
+		// First come first serve
+		for (int32 SfIdx = 0; SfIdx < TargetedShaderFormats.Num(); ++SfIdx)
 		{
-			DynamicRHIModule = &FModuleManager::LoadModuleChecked<IDynamicRHIModule>(TEXT("OpenGLDrv"));
-			if (!DynamicRHIModule->IsSupported())
+			// If we are not forcing opengl and we havent failed to create a VulkanRHI try to again with a different TargetedRHI
+			if (!bForceOpenGL && !bVulkanFailed && TargetedShaderFormats[SfIdx].StartsWith(TEXT("SF_VULKAN_")))
 			{
-				DynamicRHIModule = nullptr;
-				bOpenGLFailed = true;
+				DynamicRHIModule = &FModuleManager::LoadModuleChecked<IDynamicRHIModule>(TEXT("VulkanRHI"));
+				if (!DynamicRHIModule->IsSupported())
+				{
+					DynamicRHIModule = nullptr;
+					bVulkanFailed = true;
+				}
+				else
+				{
+					FApp::SetGraphicsRHI(TEXT("Vulkan"));
+					FPlatformApplicationMisc::UsingVulkan();
+
+					FName ShaderFormatName(*TargetedShaderFormats[SfIdx]);
+					EShaderPlatform TargetedPlatform = ShaderFormatToLegacyShaderPlatform(ShaderFormatName);
+					RequestedFeatureLevel = GetMaxSupportedFeatureLevel(TargetedPlatform);
+					break;
+				}
 			}
-			else
+			else if (!bForceVulkan && !bOpenGLFailed && TargetedShaderFormats[SfIdx].StartsWith(TEXT("GLSL_")))
 			{
-				FApp::SetGraphicsRHI(TEXT("OpenGL"));
-				FPlatformApplicationMisc::UsingOpenGL();
+				DynamicRHIModule = &FModuleManager::LoadModuleChecked<IDynamicRHIModule>(TEXT("OpenGLDrv"));
+				if (!DynamicRHIModule->IsSupported())
+				{
+					DynamicRHIModule = nullptr;
+					bOpenGLFailed = true;
+				}
+				else
+				{
+					FApp::SetGraphicsRHI(TEXT("OpenGL"));
+					FPlatformApplicationMisc::UsingOpenGL();
 
-				FName ShaderFormatName(*TargetedShaderFormats[SfIdx]);
-				EShaderPlatform TargetedPlatform = ShaderFormatToLegacyShaderPlatform(ShaderFormatName);
-				RequestedFeatureLevel = GetMaxSupportedFeatureLevel(TargetedPlatform);
-				break;
+					FName ShaderFormatName(*TargetedShaderFormats[SfIdx]);
+					EShaderPlatform TargetedPlatform = ShaderFormatToLegacyShaderPlatform(ShaderFormatName);
+					RequestedFeatureLevel = GetMaxSupportedFeatureLevel(TargetedPlatform);
+					break;
+				}
 			}
 		}
 	}
@@ -89,7 +132,14 @@ FDynamicRHI* PlatformCreateDynamicRHI()
 	}
 	else
 	{
-		if (bForceVulkan)
+		if (ForcedFeatureLevel.IsSet())
+		{
+			FMessageDialog::Open(EAppMsgType::Ok, NSLOCTEXT("LinuxDynamicRHI", "UnsupportedVulkanTargetedRHI", "Trying to force specific Vulkan feature level but it is not supported."));
+			FPlatformMisc::RequestExitWithStatus(true, 1);
+			// unreachable
+			return nullptr;	
+		}
+		else if (bForceVulkan)
 		{
 			if (bVulkanFailed)
 			{
