Fix view dimensions on D3D11 and D3D12

FRDGTextureSRVDesc::CreateFromSlice() wanting to bind a slice of Texture2DArray as a Texture2D shader resource, is setting FRDGTextureSRVDesc::DimensionOverride to ETextureDimension::Texture2D.
While D3D11 and D3D12 allows binding a Texture2DArray as a Texture2D in shader, the only way for new D3D SRV to specify the array range is to still use D3D11/12_SRV/UAV_DIMENSION_TEXTURE2D/3D/CubeARRAY.
However current implementation refactor coming from 24989901 selected the D3D_SRV_DIMENSION based of FRHIViewDesc::FTextureSRV::FViewInfo::Dimensions ending up Texture2D.
This caused the FRHIViewDesc::FTexture::ArrayRange to be completly ignored causing the wrong array slice to be bound to shaders.

diff --git a/Engine/Source/Runtime/D3D12RHI/Private/D3D12SRV.cpp b/Engine/Source/Runtime/D3D12RHI/Private/D3D12SRV.cpp
--- a/Engine/Source/Runtime/D3D12RHI/Private/D3D12SRV.cpp
+++ b/Engine/Source/Runtime/D3D12RHI/Private/D3D12SRV.cpp
@@ -175,9 +175,11 @@ void FD3D12ShaderResourceView_RHI::CreateView()
 		uint32 const PlaneSlice = GetPlaneSliceFromViewFormat(BaseFormat, ViewFormat);
 		FRHIRange8 const PlaneRange(PlaneSlice, 1);
 
-		switch (Info.Dimension)
+		// No need to use Info.Dimension, since D3D supports mixing Texture2D view types.
+		// Create a view which matches the underlying resource dimension.
+		switch (TextureDesc.Dimension)
 		{
-		case FRHIViewDesc::EDimension::Texture2D:
+		case ETextureDimension::Texture2D:
 			if (TextureDesc.NumSamples > 1)
 			{
 				SRVDesc.ViewDimension = D3D12_SRV_DIMENSION_TEXTURE2DMS;
@@ -191,7 +193,7 @@ void FD3D12ShaderResourceView_RHI::CreateView()
 			}
 			break;
 
-		case FRHIViewDesc::EDimension::Texture2DArray:
+		case ETextureDimension::Texture2DArray:
 			if (TextureDesc.NumSamples > 1)
 			{
 				SRVDesc.ViewDimension = D3D12_SRV_DIMENSION_TEXTURE2DMSARRAY;
@@ -209,19 +211,19 @@ void FD3D12ShaderResourceView_RHI::CreateView()
 			}
 			break;
 
-		case FRHIViewDesc::EDimension::Texture3D:
+		case ETextureDimension::Texture3D:
 			SRVDesc.ViewDimension = D3D12_SRV_DIMENSION_TEXTURE3D;
 			SRVDesc.Texture3D.MostDetailedMip = Info.MipRange.First;
 			SRVDesc.Texture3D.MipLevels       = Info.MipRange.Num;
 			break;
 
-		case FRHIViewDesc::EDimension::TextureCube:
+		case ETextureDimension::TextureCube:
 			SRVDesc.ViewDimension = D3D12_SRV_DIMENSION_TEXTURECUBE;
 			SRVDesc.TextureCube.MostDetailedMip = Info.MipRange.First;
 			SRVDesc.TextureCube.MipLevels       = Info.MipRange.Num;
 			break;
 
-		case FRHIViewDesc::EDimension::TextureCubeArray:
+		case ETextureDimension::TextureCubeArray:
 			SRVDesc.ViewDimension = D3D12_SRV_DIMENSION_TEXTURECUBEARRAY;
 			SRVDesc.TextureCubeArray.MostDetailedMip  = Info.MipRange.First;
 			SRVDesc.TextureCubeArray.MipLevels        = Info.MipRange.Num;
diff --git a/Engine/Source/Runtime/D3D12RHI/Private/D3D12UAV.cpp b/Engine/Source/Runtime/D3D12RHI/Private/D3D12UAV.cpp
--- a/Engine/Source/Runtime/D3D12RHI/Private/D3D12UAV.cpp
+++ b/Engine/Source/Runtime/D3D12RHI/Private/D3D12UAV.cpp
@@ -150,16 +150,18 @@ void FD3D12UnorderedAccessView_RHI::CreateView()
 		uint32 const PlaneSlice = GetPlaneSliceFromViewFormat(BaseFormat, ViewFormat);
 		FRHIRange8 const PlaneRange(PlaneSlice, 1);
 
-		switch (Info.Dimension)
+		// No need to use Info.Dimension, since D3D supports mixing Texture2D view types.
+		// Create a view which matches the underlying resource dimension.
+		switch (TextureDesc.Dimension)
 		{
-		case FRHIViewDesc::EDimension::Texture2D:
+		case ETextureDimension::Texture2D:
 			UAVDesc.ViewDimension = D3D12_UAV_DIMENSION_TEXTURE2D;
 			UAVDesc.Texture2D.MipSlice   = Info.MipLevel;
 			UAVDesc.Texture2D.PlaneSlice = PlaneSlice;
 			break;
 
-		case FRHIViewDesc::EDimension::TextureCube:
-		case FRHIViewDesc::EDimension::TextureCubeArray:
+		case ETextureDimension::TextureCube:
+		case ETextureDimension::TextureCubeArray:
 			UAVDesc.ViewDimension = D3D12_UAV_DIMENSION_TEXTURE2DARRAY;
 			UAVDesc.Texture2DArray.FirstArraySlice = Info.ArrayRange.First * 6;
 			UAVDesc.Texture2DArray.ArraySize       = Info.ArrayRange.Num * 6;
@@ -167,7 +169,7 @@ void FD3D12UnorderedAccessView_RHI::CreateView()
 			UAVDesc.Texture2DArray.PlaneSlice      = PlaneSlice;
 			break;
 
-		case FRHIViewDesc::EDimension::Texture2DArray:
+		case ETextureDimension::Texture2DArray:
 			UAVDesc.ViewDimension = D3D12_UAV_DIMENSION_TEXTURE2DARRAY;
 			UAVDesc.Texture2DArray.FirstArraySlice = Info.ArrayRange.First;
 			UAVDesc.Texture2DArray.ArraySize       = Info.ArrayRange.Num;
@@ -175,7 +177,7 @@ void FD3D12UnorderedAccessView_RHI::CreateView()
 			UAVDesc.Texture2DArray.PlaneSlice      = PlaneSlice;
 			break;
 
-		case FRHIViewDesc::EDimension::Texture3D:
+		case ETextureDimension::Texture3D:
 			UAVDesc.ViewDimension = D3D12_UAV_DIMENSION_TEXTURE3D;
 			UAVDesc.Texture3D.FirstWSlice = 0;
 			UAVDesc.Texture3D.WSize       = FMath::Max(TextureDesc.Depth >> Info.MipLevel, 1);
diff --git a/Engine/Source/Runtime/RHI/Private/RHIResources.cpp b/Engine/Source/Runtime/RHI/Private/RHIResources.cpp
--- a/Engine/Source/Runtime/RHI/Private/RHIResources.cpp
+++ b/Engine/Source/Runtime/RHI/Private/RHIResources.cpp
@@ -199,7 +199,7 @@ FRHIViewDesc::FTexture::FViewInfo FRHIViewDesc::FTexture::GetViewInfo(FRHITextur
 
 	FRHITextureDesc const& Desc = TargetTexture->GetDesc();
 
-	checkf(Desc.Dimension != ETextureDimension::Texture3D || Dimension == EDimension::Texture3D, TEXT("Views of 3D textures must use 3D dimension."));
+	checkf(!Desc.IsTexture3D() || Dimension == EDimension::Texture3D, TEXT("Views of 3D textures must use 3D dimension."));
 	checkf(Desc.IsTexture3D() || Dimension != EDimension::Texture3D, TEXT("The underlying texture resource must be a 3D texture to create a 3D dimension view."));
 	checkf(Desc.IsTextureCube() || (Dimension != EDimension::TextureCube && Dimension != EDimension::TextureCubeArray), TEXT("The underlying texture resource must be a cube (or cube array) to create a cube dimension view."));
 
diff --git a/Engine/Source/Runtime/RHI/Public/RHIResources.h b/Engine/Source/Runtime/RHI/Public/RHIResources.h
--- a/Engine/Source/Runtime/RHI/Public/RHIResources.h
+++ b/Engine/Source/Runtime/RHI/Public/RHIResources.h
@@ -2620,6 +2620,14 @@ protected:
 	{}
 
 public:
+	//
+	// Specifies the type of view to create. Must match the shader parameter this view will be bound to.
+	// 
+	// The dimension is allowed to differ from the underlying resource's dimensions, e.g. to create a view
+	// compatible with a Texture2D<> shader parameter, but where the underlying resource is a texture 2D array.
+	//
+	// Some combinations are not valid, e.g. 3D textures can only have 3D views.
+	//
 	FInitializer& SetDimension(ETextureDimension InDimension)
 	{
 		switch (InDimension)
@@ -2665,9 +2673,14 @@ public:
 	}
 
 	//
-	// The meaning of array "elements" is given by the view dimension. I.e. a view with "Dimension == CubeArray" indexes the array in whole cubes.
-	// [0] = the first cube (2D slices 0 to 5)
-	// [1] = the second cube (2D slices 6 to 11)
+	// The meaning of array "elements" is given by the dimension of the underlying resource.
+	// I.e. a view of a TextureCubeArray resource indexes the array in whole cubes.
+	// 
+	//     [0] = the first cube (2D slices 0 to 5)
+	//     [1] = the second cube (2D slices 6 to 11)
+	// 
+	// If the view dimension is smaller than the resource dimension, the array range will be further limited.
+	// E.g. creating a Texture2D dimension view of a TextureCubeArray resource
 	//
 	FInitializer& SetArrayRange(uint16 InFirstElement, uint16 InNumElements)
 	{
@@ -2697,6 +2710,14 @@ protected:
 	}
 
 public:
+	//
+	// Specifies the type of view to create. Must match the shader parameter this view will be bound to.
+	// 
+	// The dimension is allowed to differ from the underlying resource's dimensions, e.g. to create a view
+	// compatible with an RWTexture2D<> shader parameter, but where the underlying resource is a texture 2D array.
+	//
+	// Some combinations are not valid, e.g. 3D textures can only have 3D views.
+	//
 	FInitializer& SetDimension(ETextureDimension InDimension)
 	{
 		switch (InDimension)
@@ -2741,9 +2762,14 @@ public:
 	}
 
 	//
-	// The meaning of array "elements" is given by the view dimension. I.e. a view with "Dimension == CubeArray" indexes the array in whole cubes.
-	// [0] = the first cube (2D slices 0 to 5)
-	// [1] = the second cube (2D slices 6 to 11)
+	// The meaning of array "elements" is given by the dimension of the underlying resource.
+	// I.e. a view of a TextureCubeArray resource indexes the array in whole cubes.
+	// 
+	//     [0] = the first cube (2D slices 0 to 5)
+	//     [1] = the second cube (2D slices 6 to 11)
+	// 
+	// If the view dimension is smaller than the resource dimension, the array range will be further limited.
+	// E.g. creating a Texture2D dimension view of a TextureCubeArray resource
 	//
 	FInitializer& SetArrayRange(uint16 InFirstElement, uint16 InNumElements)
 	{
diff --git a/Engine/Source/Runtime/Windows/D3D11RHI/Private/D3D11SRV.cpp b/Engine/Source/Runtime/Windows/D3D11RHI/Private/D3D11SRV.cpp
--- a/Engine/Source/Runtime/Windows/D3D11RHI/Private/D3D11SRV.cpp
+++ b/Engine/Source/Runtime/Windows/D3D11RHI/Private/D3D11SRV.cpp
@@ -74,9 +74,11 @@ void FD3D11ShaderResourceView::UpdateView()
 
 		SRVDesc.Format = FindShaderResourceDXGIFormat(DXGI_FORMAT(GPixelFormats[Info.Format].PlatformFormat), Info.bSRGB);
 
-		switch (Info.Dimension)
+		// No need to use Info.Dimension, since D3D supports mixing Texture2D view types.
+		// Create a view which matches the underlying resource dimension.
+		switch (TextureDesc.Dimension)
 		{
-		case FRHIViewDesc::EDimension::Texture2D:
+		case ETextureDimension::Texture2D:
 			if (TextureDesc.NumSamples > 1)
 			{
 				SRVDesc.ViewDimension = D3D11_SRV_DIMENSION_TEXTURE2DMS;
@@ -89,7 +91,7 @@ void FD3D11ShaderResourceView::UpdateView()
 			}
 			break;
 
-		case FRHIViewDesc::EDimension::Texture2DArray:
+		case ETextureDimension::Texture2DArray:
 			if (TextureDesc.NumSamples > 1)
 			{
 				SRVDesc.ViewDimension = D3D11_SRV_DIMENSION_TEXTURE2DMSARRAY;
@@ -106,19 +108,19 @@ void FD3D11ShaderResourceView::UpdateView()
 			}
 			break;
 
-		case FRHIViewDesc::EDimension::Texture3D:
+		case ETextureDimension::Texture3D:
 			SRVDesc.ViewDimension = D3D11_SRV_DIMENSION_TEXTURE3D;
 			SRVDesc.Texture3D.MostDetailedMip = Info.MipRange.First;
 			SRVDesc.Texture3D.MipLevels       = Info.MipRange.Num;
 			break;
 
-		case FRHIViewDesc::EDimension::TextureCube:
+		case ETextureDimension::TextureCube:
 			SRVDesc.ViewDimension = D3D11_SRV_DIMENSION_TEXTURECUBE;
 			SRVDesc.TextureCube.MostDetailedMip = Info.MipRange.First;
 			SRVDesc.TextureCube.MipLevels       = Info.MipRange.Num;
 			break;
 
-		case FRHIViewDesc::EDimension::TextureCubeArray:
+		case ETextureDimension::TextureCubeArray:
 			SRVDesc.ViewDimension = D3D11_SRV_DIMENSION_TEXTURECUBEARRAY;
 			SRVDesc.TextureCubeArray.MostDetailedMip  = Info.MipRange.First;
 			SRVDesc.TextureCubeArray.MipLevels        = Info.MipRange.Num;
diff --git a/Engine/Source/Runtime/Windows/D3D11RHI/Private/D3D11UAV.cpp b/Engine/Source/Runtime/Windows/D3D11RHI/Private/D3D11UAV.cpp
--- a/Engine/Source/Runtime/Windows/D3D11RHI/Private/D3D11UAV.cpp
+++ b/Engine/Source/Runtime/Windows/D3D11RHI/Private/D3D11UAV.cpp
@@ -64,36 +64,39 @@ void FD3D11UnorderedAccessView::UpdateView()
 		FD3D11Texture* Texture = FD3D11DynamicRHI::ResourceCast(GetTexture());
 		D3DResource = Texture->GetResource();
 
+		FRHITextureDesc const& TextureDesc = Texture->GetDesc();
 		auto const Info = ViewDesc.Texture.UAV.GetViewInfo(Texture);
 
 		UAVDesc.Format = FindUnorderedAccessDXGIFormat(DXGI_FORMAT(GPixelFormats[Info.Format].PlatformFormat));
 
-		switch (Info.Dimension)
+		// No need to use Info.Dimension, since D3D supports mixing Texture2D view types.
+		// Create a view which matches the underlying resource dimension.
+		switch (TextureDesc.Dimension)
 		{
-		case FRHIViewDesc::EDimension::Texture2D:
+		case ETextureDimension::Texture2D:
 			UAVDesc.ViewDimension = D3D11_UAV_DIMENSION_TEXTURE2D;
 			UAVDesc.Texture2D.MipSlice = Info.MipLevel;
 			break;
 
-		case FRHIViewDesc::EDimension::TextureCube:
-		case FRHIViewDesc::EDimension::TextureCubeArray:
+		case ETextureDimension::TextureCube:
+		case ETextureDimension::TextureCubeArray:
 			UAVDesc.ViewDimension = D3D11_UAV_DIMENSION_TEXTURE2DARRAY;
 			UAVDesc.Texture2DArray.FirstArraySlice = Info.ArrayRange.First * 6;
 			UAVDesc.Texture2DArray.ArraySize       = Info.ArrayRange.Num * 6;
 			UAVDesc.Texture2DArray.MipSlice        = Info.MipLevel;
 			break;
 
-		case FRHIViewDesc::EDimension::Texture2DArray:
+		case ETextureDimension::Texture2DArray:
 			UAVDesc.ViewDimension = D3D11_UAV_DIMENSION_TEXTURE2DARRAY;
 			UAVDesc.Texture2DArray.FirstArraySlice = Info.ArrayRange.First;
 			UAVDesc.Texture2DArray.ArraySize       = Info.ArrayRange.Num;
 			UAVDesc.Texture2DArray.MipSlice        = Info.MipLevel;
 			break;
 
-		case FRHIViewDesc::EDimension::Texture3D:
+		case ETextureDimension::Texture3D:
 			UAVDesc.ViewDimension = D3D11_UAV_DIMENSION_TEXTURE3D;
 			UAVDesc.Texture3D.FirstWSlice = 0;
-			UAVDesc.Texture3D.WSize       = FMath::Max(Texture->GetDesc().Depth >> Info.MipLevel, 1);
+			UAVDesc.Texture3D.WSize       = FMath::Max(TextureDesc.Depth >> Info.MipLevel, 1);
 			UAVDesc.Texture3D.MipSlice    = Info.MipLevel;
 			break;
 
