Remove reference counting from FRayTracingGeometryInstance.GeometryRHI

diff --git a/Engine/Source/Runtime/RHI/Public/RHIResources.h b/Engine/Source/Runtime/RHI/Public/RHIResources.h
--- a/Engine/Source/Runtime/RHI/Public/RHIResources.h
+++ b/Engine/Source/Runtime/RHI/Public/RHIResources.h
@@ -3017,9 +3017,8 @@ class FRHIRayTracingGeometry;
 * All instances covered by this descriptor will share shader bindings, but may have different transforms and user data.
 */
 struct FRayTracingGeometryInstance
-{
-	// TODO: UE-130819 Ref counting is a temporary workaround for a very rare streaming crash.
-	TRefCountPtr<FRHIRayTracingGeometry> GeometryRHI = nullptr;
+{	
+	FRHIRayTracingGeometry* GeometryRHI = nullptr;
 
 	// A single physical mesh may be duplicated many times in the scene with different transforms and user data.
 	// All copies share the same shader binding table entries and therefore will have the same material and shader resources.
diff --git a/Engine/Source/Runtime/Renderer/Private/RayTracing/RayTracingDebug.cpp b/Engine/Source/Runtime/Renderer/Private/RayTracing/RayTracingDebug.cpp
--- a/Engine/Source/Runtime/Renderer/Private/RayTracing/RayTracingDebug.cpp
+++ b/Engine/Source/Runtime/Renderer/Private/RayTracing/RayTracingDebug.cpp
@@ -1044,7 +1044,7 @@ void FDeferredShadingSceneRenderer::RayTracingDisplayPicking(const FRayTracingPi
 	{
 		if (Instance.GeometryRHI)
 		{
-			const uint64 GeometryAddress = uint64(Instance.GeometryRHI.GetReference());
+			const uint64 GeometryAddress = uint64(Instance.GeometryRHI);
 			if (PickingFeedback.GeometryAddress == GeometryAddress)
 			{
 				Geometry = Instance.GeometryRHI;
