Fix bogus assert when using Allocate on a command list in bypass mode.

#rnx

diff --git a/Engine/Plugins/Tests/RHITests/Source/RHITests/Public/RHIBufferTests.h b/Engine/Plugins/Tests/RHITests/Source/RHITests/Public/RHIBufferTests.h
--- a/Engine/Plugins/Tests/RHITests/Source/RHITests/Public/RHIBufferTests.h
+++ b/Engine/Plugins/Tests/RHITests/Source/RHITests/Public/RHIBufferTests.h
@@ -398,7 +398,7 @@ public:
 
 	static bool Test_RHICreateBuffer_Parallel(FRHICommandListImmediate& RHICmdList)
 	{
-		if (!GRHISupportsMultithreadedResources || RHICmdList.Bypass())
+		if (!GRHISupportsMultithreadedResources)
 		{
 			return true;
 		}
diff --git a/Engine/Source/Runtime/RHI/Public/RHICommandList.h b/Engine/Source/Runtime/RHI/Public/RHICommandList.h
--- a/Engine/Source/Runtime/RHI/Public/RHICommandList.h
+++ b/Engine/Source/Runtime/RHI/Public/RHICommandList.h
@@ -522,7 +522,6 @@ public:
 
 	FORCEINLINE_DEBUGGABLE void* Alloc(int64 AllocSize, int64 Alignment)
 	{
-		checkSlow(!Bypass() && "Can't use RHICommandList in bypass mode.");
 		return MemManager.Alloc(AllocSize, Alignment);
 	}
 
@@ -569,6 +568,7 @@ public:
 	FORCEINLINE_DEBUGGABLE void* AllocCommand(int32 AllocSize, int32 Alignment)
 	{
 		checkSlow(!IsExecuting());
+		checkfSlow(!Bypass(), TEXT("Invalid attempt to record commands in bypass mode."));
 		FRHICommandBase* Result = (FRHICommandBase*) MemManager.Alloc(AllocSize, Alignment);
 #if RHI_COUNT_COMMANDS
 		++NumCommands;
