Fixed FRHITextureDesc to no longer have pad bytes.

diff --git a/Engine/Source/Runtime/RHI/Public/RHIResources.h b/Engine/Source/Runtime/RHI/Public/RHIResources.h
--- a/Engine/Source/Runtime/RHI/Public/RHIResources.h
+++ b/Engine/Source/Runtime/RHI/Public/RHIResources.h
@@ -1449,16 +1449,11 @@ private:
 /** Descriptor used to create a texture resource */
 struct FRHITextureDesc
 {
-	FRHITextureDesc()
-		: NumSamples(1)
-		, Dimension(ETextureDimension::Texture2D)
-	{}
+	FRHITextureDesc() = default;
 
 	FRHITextureDesc(ETextureDimension InDimension)
-		: NumSamples(1)
-		, Dimension(InDimension)
-	{
-	}
+		: Dimension(InDimension)
+	{}
 
 	FRHITextureDesc(
 		  ETextureDimension   InDimension
@@ -1497,6 +1492,7 @@ struct FRHITextureDesc
 		Hash = HashCombine(Hash, GetTypeHash(Desc.NumSamples));
 		Hash = HashCombine(Hash, GetTypeHash(Desc.ClearValue));
 		Hash = HashCombine(Hash, GetTypeHash(Desc.ExtData   ));
+		Hash = HashCombine(Hash, GetTypeHash(Desc.GPUMask.GetNative()));
 		return Hash;
 	}
 
@@ -1512,7 +1508,8 @@ struct FRHITextureDesc
 			&& NumMips    == Other.NumMips
 			&& NumSamples == Other.NumSamples
 			&& ClearValue == Other.ClearValue
-			&& ExtData    == Other.ExtData;
+			&& ExtData    == Other.ExtData
+			&& GPUMask    == Other.GPUMask;
 	}
 
 	bool operator != (const FRHITextureDesc& Other) const
@@ -1577,6 +1574,9 @@ struct FRHITextureDesc
 	/** Clear value to use when fast-clearing the texture. */
 	FClearValueBinding ClearValue;
 
+	/* A mask representing which GPUs to create the resource on, in a multi-GPU system. */
+	FRHIGPUMask GPUMask = FRHIGPUMask::All();
+
 	/** Platform-specific additional data. Used for offline processed textures on some platforms. */
 	uint32 ExtData = 0;
 
@@ -1593,10 +1593,10 @@ struct FRHITextureDesc
 	uint8 NumMips = 1;
 
 	/** Number of samples in the texture. >1 for MSAA. */
-	uint8 NumSamples : 5;
+	uint8 NumSamples = 1;
 
 	/** Texture dimension to use when creating the RHI texture. */
-	ETextureDimension Dimension : 3;
+	ETextureDimension Dimension = ETextureDimension::Texture2D;
 
 	/** Pixel format used to create RHI texture. */
 	EPixelFormat Format = PF_Unknown;
@@ -1607,9 +1607,6 @@ struct FRHITextureDesc
 	/** Resource memory percentage which should be allocated onto fast VRAM (hint-only). (encoding into 8bits, 0..255 -> 0%..100%) */
 	uint8 FastVRAMPercentage = 0xFF;
 
-	/* A mask representing which GPUs to create the resource on, in a multi-GPU system. */
-	FRHIGPUMask GPUMask = FRHIGPUMask::All();
-
 	/** Check the validity. */
 	static bool CheckValidity(const FRHITextureDesc& Desc, const TCHAR* Name)
 	{
