Fix crashes caused by processing the deferred deletion queue at the wrong time, especially when running without an RHI thread.

Since FRHICommandListImmediate::ImmediateFlush() was calling FlushPendingDeletes() after executing the current list of commands, the lambda enqueued by that function was actually running as part of the next command list. Besides unnecessarily extending resource lifetimes, this also meant that the new command list started by flushing the deferred deletion queue, so anything added directly in there (instead of via another enqueued command) would be deleted before the commands were executed on the GPU. This was very easy to reproduce with -norhithread on DX12, because texture unlock operations add the staging buffer to the deferred deletion queue immediately, instead of enqueuing a command to do it, so the staging buffers were gone by the time the GPU tried to copy from them.

This changelist adds a boolean to RHISubmitCommandLists() which is true when we're flushing the immediate command list with the FlushRHIThreadFlushResources mode, so that the RHI can process the deletion queue internally after submission, instead of doing it in RHIPerFrameRHIFlushComplete(). I didn't want to move RHIPerFrameRHIFlushComplete() itself to another point in the timeline, because old RHIs (D3D11 and OpenGL) use that for other purposes, and it seems unwise to alter their behavior (e.g. D3D11 wants to resolve timing queries in there with a blocking wait, and running that at the end of the current command list would introduce a CPU/GPU sync point).

Also:
* deprecated FRHIResource::FlushPendingDeletes(), since all it does is call FlushPendingDeletes() on the command list being passed in, and code outside of the RHI really shouldn't be doing that.
* made FlushPendingDeleteRHIResources_RenderThread() flush the immediate command list instead of calling FlushPendingDeletes() directly

#rnx

diff --git a/Engine/Source/Runtime/Apple/MetalRHI/Private/MetalDynamicRHI.h b/Engine/Source/Runtime/Apple/MetalRHI/Private/MetalDynamicRHI.h
--- a/Engine/Source/Runtime/Apple/MetalRHI/Private/MetalDynamicRHI.h
+++ b/Engine/Source/Runtime/Apple/MetalRHI/Private/MetalDynamicRHI.h
@@ -117,7 +117,7 @@ public:
 
 	virtual IRHIComputeContext* RHIGetCommandContext(ERHIPipeline Pipeline, FRHIGPUMask GPUMask) final override;
 	virtual IRHIPlatformCommandList* RHIFinalizeContext(IRHIComputeContext* Context) final override;
-	virtual void RHISubmitCommandLists(TArrayView<IRHIPlatformCommandList*> CommandLists) final override;
+	virtual void RHISubmitCommandLists(TArrayView<IRHIPlatformCommandList*> CommandLists, bool bFlushResources) final override;
 
 	virtual FTexture2DRHIRef AsyncReallocateTexture2D_RenderThread(class FRHICommandListImmediate& RHICmdList, FRHITexture2D* Texture2D, int32 NewMipCount, int32 NewSizeX, int32 NewSizeY, FThreadSafeCounter* RequestStatus) final override;
 	virtual ETextureReallocationStatus FinalizeAsyncReallocateTexture2D_RenderThread(class FRHICommandListImmediate& RHICmdList, FRHITexture2D* Texture2D, bool bBlockUntilCompleted) final override;
diff --git a/Engine/Source/Runtime/Apple/MetalRHI/Private/MetalRHI.cpp b/Engine/Source/Runtime/Apple/MetalRHI/Private/MetalRHI.cpp
--- a/Engine/Source/Runtime/Apple/MetalRHI/Private/MetalRHI.cpp
+++ b/Engine/Source/Runtime/Apple/MetalRHI/Private/MetalRHI.cpp
@@ -1426,7 +1426,6 @@ IRHIPlatformCommandList* FMetalDynamicRHI::RHIFinalizeContext(IRHIComputeContext
 	return nullptr;
 }
 
-void FMetalDynamicRHI::RHISubmitCommandLists(TArrayView<IRHIPlatformCommandList*> CommandLists)
+void FMetalDynamicRHI::RHISubmitCommandLists(TArrayView<IRHIPlatformCommandList*> CommandLists, bool bFlushResources)
 {
-	UE_LOG(LogRHI, Fatal, TEXT("FMetalDynamicRHI::RHISubmitCommandLists should never be called. Metal RHI does not implement parallel command list execution."));
 }
diff --git a/Engine/Source/Runtime/D3D12RHI/Private/D3D12RHI.cpp b/Engine/Source/Runtime/D3D12RHI/Private/D3D12RHI.cpp
--- a/Engine/Source/Runtime/D3D12RHI/Private/D3D12RHI.cpp
+++ b/Engine/Source/Runtime/D3D12RHI/Private/D3D12RHI.cpp
@@ -333,7 +333,7 @@ void FD3D12DynamicRHI::Shutdown()
 		int32 DeletedCount;
 		do
 		{
-			DeletedCount = FRHIResource::FlushPendingDeletes(RHICmdList);
+			DeletedCount = RHICmdList.FlushPendingDeletes();
 			RHICmdList.ImmediateFlush(EImmediateFlushType::FlushRHIThreadFlushResources);
 		} while (DeletedCount);
 
@@ -497,8 +497,10 @@ void FD3D12DynamicRHI::FlushTiming(bool bCreateNew)
 	});
 }
 
-void FD3D12DynamicRHI::RHIPerFrameRHIFlushComplete()
+void FD3D12DynamicRHI::ProcessDeferredDeletionQueue()
 {
+	ProcessDeferredDeletionQueue_Platform();
+
 	TArray<FD3D12DeferredDeleteObject> Local;
 	{
 		FScopeLock Lock(&ObjectsToDeleteCS);
diff --git a/Engine/Source/Runtime/D3D12RHI/Private/D3D12RHIPrivate.h b/Engine/Source/Runtime/D3D12RHI/Private/D3D12RHIPrivate.h
--- a/Engine/Source/Runtime/D3D12RHI/Private/D3D12RHIPrivate.h
+++ b/Engine/Source/Runtime/D3D12RHI/Private/D3D12RHIPrivate.h
@@ -275,8 +275,6 @@ public:
 		return false;
 	}
 
-	virtual void RHIPerFrameRHIFlushComplete() override;
-
 	virtual FSamplerStateRHIRef RHICreateSamplerState(const FSamplerStateInitializerRHI& Initializer) final override;
 	virtual FRasterizerStateRHIRef RHICreateRasterizerState(const FRasterizerStateInitializerRHI& Initializer) final override;
 	virtual FDepthStencilStateRHIRef RHICreateDepthStencilState(const FDepthStencilStateInitializerRHI& Initializer) final override;
@@ -378,7 +376,7 @@ public:
 	virtual class IRHIComputeContext* RHIGetDefaultAsyncComputeContext() final override;
 	virtual IRHIComputeContext* RHIGetCommandContext(ERHIPipeline Pipeline, FRHIGPUMask GPUMask) final override;
 	virtual IRHIPlatformCommandList* RHIFinalizeContext(IRHIComputeContext* Context) final override;
-	virtual void RHISubmitCommandLists(TArrayView<IRHIPlatformCommandList*> CommandLists) final override;
+	virtual void RHISubmitCommandLists(TArrayView<IRHIPlatformCommandList*> CommandLists, bool bFlushResources) final override;
 
 	virtual IRHITransientResourceAllocator* RHICreateTransientResourceAllocator() override;
 
@@ -586,6 +584,9 @@ protected:
 	const bool bAllowVendorDevice;
 
 	FDisplayInformationArray DisplayList;
+
+	void ProcessDeferredDeletionQueue();
+	void ProcessDeferredDeletionQueue_Platform();
 };
 
 ENUM_CLASS_FLAGS(FD3D12DynamicRHI::EQueueStatus);
diff --git a/Engine/Source/Runtime/D3D12RHI/Private/D3D12Submission.cpp b/Engine/Source/Runtime/D3D12RHI/Private/D3D12Submission.cpp
--- a/Engine/Source/Runtime/D3D12RHI/Private/D3D12Submission.cpp
+++ b/Engine/Source/Runtime/D3D12RHI/Private/D3D12Submission.cpp
@@ -196,9 +196,14 @@ IRHIPlatformCommandList* FD3D12DynamicRHI::RHIFinalizeContext(IRHIComputeContext
 	}
 }
 
-void FD3D12DynamicRHI::RHISubmitCommandLists(TArrayView<IRHIPlatformCommandList*> CommandLists)
+void FD3D12DynamicRHI::RHISubmitCommandLists(TArrayView<IRHIPlatformCommandList*> CommandLists, bool bFlushResources)
 {
 	SubmitCommands(MakeArrayView(reinterpret_cast<FD3D12FinalizedCommands**>(CommandLists.GetData()), CommandLists.Num()));
+
+	if (bFlushResources)
+	{
+		ProcessDeferredDeletionQueue();
+	}
 }
 
 void FD3D12DynamicRHI::SubmitCommands(TConstArrayView<FD3D12FinalizedCommands*> Commands)
diff --git a/Engine/Source/Runtime/D3D12RHI/Private/Windows/WindowsD3D12Device.cpp b/Engine/Source/Runtime/D3D12RHI/Private/Windows/WindowsD3D12Device.cpp
--- a/Engine/Source/Runtime/D3D12RHI/Private/Windows/WindowsD3D12Device.cpp
+++ b/Engine/Source/Runtime/D3D12RHI/Private/Windows/WindowsD3D12Device.cpp
@@ -2089,3 +2089,8 @@ FD3D12DiagnosticBuffer::~FD3D12DiagnosticBuffer()
 	CpuAddress = nullptr;
 	GpuAddress = 0;
 }
+
+void FD3D12DynamicRHI::ProcessDeferredDeletionQueue_Platform()
+{
+	// Nothing Windows-specific here.
+}
diff --git a/Engine/Source/Runtime/NullDrv/Public/NullRHI.h b/Engine/Source/Runtime/NullDrv/Public/NullRHI.h
--- a/Engine/Source/Runtime/NullDrv/Public/NullRHI.h
+++ b/Engine/Source/Runtime/NullDrv/Public/NullRHI.h
@@ -599,7 +599,7 @@ public:
 		return nullptr;
 	}
 
-	virtual void RHISubmitCommandLists(TArrayView<IRHIPlatformCommandList*> CommandLists) final override
+	virtual void RHISubmitCommandLists(TArrayView<IRHIPlatformCommandList*> CommandLists, bool bFlushResources) final override
 	{
 	}
 
diff --git a/Engine/Source/Runtime/OpenGLDrv/Private/OpenGLCommands.cpp b/Engine/Source/Runtime/OpenGLDrv/Private/OpenGLCommands.cpp
--- a/Engine/Source/Runtime/OpenGLDrv/Private/OpenGLCommands.cpp
+++ b/Engine/Source/Runtime/OpenGLDrv/Private/OpenGLCommands.cpp
@@ -2989,9 +2989,8 @@ IRHIPlatformCommandList* FOpenGLDynamicRHI::RHIFinalizeContext(IRHIComputeContex
 	return nullptr;
 }
 
-void FOpenGLDynamicRHI::RHISubmitCommandLists(TArrayView<IRHIPlatformCommandList*> CommandLists)
+void FOpenGLDynamicRHI::RHISubmitCommandLists(TArrayView<IRHIPlatformCommandList*> CommandLists, bool bFlushResources)
 {
-	UE_LOG(LogRHI, Fatal, TEXT("FOpenGLDynamicRHI::RHISubmitCommandLists should never be called. OpenGL RHI does not implement parallel command list execution."));
 }
 
 void FOpenGLDynamicRHI::RHIInvalidateCachedState()
diff --git a/Engine/Source/Runtime/OpenGLDrv/Public/OpenGLDrv.h b/Engine/Source/Runtime/OpenGLDrv/Public/OpenGLDrv.h
--- a/Engine/Source/Runtime/OpenGLDrv/Public/OpenGLDrv.h
+++ b/Engine/Source/Runtime/OpenGLDrv/Public/OpenGLDrv.h
@@ -413,7 +413,7 @@ public:
 	virtual class IRHICommandContext* RHIGetDefaultContext() final override;
 	virtual IRHIComputeContext* RHIGetCommandContext(ERHIPipeline Pipeline, FRHIGPUMask GPUMask) final override;
 	virtual IRHIPlatformCommandList* RHIFinalizeContext(IRHIComputeContext* Context) final override;
-	virtual void RHISubmitCommandLists(TArrayView<IRHIPlatformCommandList*> CommandLists) final override;
+	virtual void RHISubmitCommandLists(TArrayView<IRHIPlatformCommandList*> CommandLists, bool bFlushResources) final override;
 
 	virtual void RHIBeginRenderPass(const FRHIRenderPassInfo& InInfo, const TCHAR* InName) final override;
 	virtual void RHIEndRenderPass() final override;
diff --git a/Engine/Source/Runtime/RHI/Private/RHICommandList.cpp b/Engine/Source/Runtime/RHI/Private/RHICommandList.cpp
--- a/Engine/Source/Runtime/RHI/Private/RHICommandList.cpp
+++ b/Engine/Source/Runtime/RHI/Private/RHICommandList.cpp
@@ -494,7 +494,7 @@ void FRHICommandListImmediate::QueueAsyncCommandListSubmit(TArrayView<FQueuedCom
 		
 		// Commands may already be queued on the immediate command list. These need to be executed
 		// first before any parallel commands can be inserted, otherwise commands will run out-of-order.
-		ExecuteAndReset();
+		ExecuteAndReset(false);
 		InitializeImmediateContexts();
 
 		struct FTask
@@ -651,7 +651,7 @@ void FRHICommandListImmediate::QueueAsyncCommandListSubmit(TArrayView<FQueuedCom
 
 			if (AllCmdLists.Num())
 			{
-				GDynamicRHI->RHISubmitCommandLists(AllCmdLists);
+				GDynamicRHI->RHISubmitCommandLists(AllCmdLists, false);
 			}
 		});
 	}
@@ -679,7 +679,7 @@ void FRHICommandListImmediate::QueueAsyncCommandListSubmit(TArrayView<FQueuedCom
 	}
 }
 
-void FRHICommandListImmediate::ExecuteAndReset()
+void FRHICommandListImmediate::ExecuteAndReset(bool bFlushResources)
 {
 	check(IsInRenderingThread());
 
@@ -708,7 +708,7 @@ void FRHICommandListImmediate::ExecuteAndReset()
 	//
 	// In non-bypass mode, always submit.
 	//
-	EnqueueLambda([](FRHICommandListImmediate& ExecutingCmdList)
+	EnqueueLambda([bFlushResources](FRHICommandListImmediate& ExecutingCmdList)
 	{
 		TArray<IRHIPlatformCommandList*, TInlineAllocator<GetRHIPipelineCount()>> CommandLists;
 		for (IRHIComputeContext* Context : ExecutingCmdList.Contexts)
@@ -723,12 +723,18 @@ void FRHICommandListImmediate::ExecuteAndReset()
 			}
 		}
 
-		if (CommandLists.Num())
+		if (CommandLists.Num() || (GDynamicRHI && bFlushResources))
 		{
-			GDynamicRHI->RHISubmitCommandLists(CommandLists);
+			GDynamicRHI->RHISubmitCommandLists(CommandLists, bFlushResources);
 		}
 	});
 
+	if (bFlushResources)
+	{
+		CSV_SCOPED_TIMING_STAT(RHITFlushes, FlushRHIThreadFlushResourcesTotal);
+		FlushPendingDeletes();
+	}
+
 	// Equivalent to FinishRecording(), without the check(!IsImmediate()).
 	DispatchEvent->DispatchSubsequents();
 
@@ -919,7 +925,7 @@ void FRHICommandListExecutor::LatchBypass()
 	}
 	else
 	{
-		RHICmdList.ExecuteAndReset();
+		RHICmdList.ExecuteAndReset(false);
 
 		CheckNoOutstandingCmdLists();
 		check(!RHICmdList.HasCommands());
@@ -946,8 +952,7 @@ void FRHICommandListExecutor::LatchBypass()
 		bool NewBypass = IsInGameThread() || (CVarRHICmdBypass.GetValueOnAnyThread() >= 1);
 		if (NewBypass && !bLatchedBypass)
 		{
-			FRHIResource::FlushPendingDeletes(RHICmdList);
-			RHICmdList.ExecuteAndReset();
+			RHICmdList.ExecuteAndReset(true);
 		}
 
 		bLatchedBypass = NewBypass;
@@ -1038,12 +1043,6 @@ int32 FRHICommandListImmediate::FlushPendingDeletes()
 	EnqueueLambda([DeletedResources = MoveTemp(DeletedResources)](FRHICommandListImmediate& RHICmdList) mutable
 	{
 		SCOPED_NAMED_EVENT(STAT_FRHICommandListImmediate_DeleteResources, FColor::Magenta);
-
-		if (GDynamicRHI)
-		{
-			GDynamicRHI->RHIPerFrameRHIFlushComplete();
-		}
-
 		for (int32 i = DeletedResources.Num() - 1; i >= 0; i--)
 		{
 			FRHIResource* Resource = DeletedResources[i];
diff --git a/Engine/Source/Runtime/RHI/Private/RHIValidation.cpp b/Engine/Source/Runtime/RHI/Private/RHIValidation.cpp
--- a/Engine/Source/Runtime/RHI/Private/RHIValidation.cpp
+++ b/Engine/Source/Runtime/RHI/Private/RHIValidation.cpp
@@ -432,7 +432,7 @@ IRHIPlatformCommandList* FValidationRHI::RHIFinalizeContext(IRHIComputeContext*
 	return OuterCommandList;
 }
 
-void FValidationRHI::RHISubmitCommandLists(TArrayView<IRHIPlatformCommandList*> OuterCommandLists)
+void FValidationRHI::RHISubmitCommandLists(TArrayView<IRHIPlatformCommandList*> OuterCommandLists, bool bFlushResources)
 {
 	FMemMark Mark(FMemStack::Get());
 	TArray<IRHIPlatformCommandList*, TMemStackAllocator<>> InnerCommandLists;
@@ -455,7 +455,7 @@ void FValidationRHI::RHISubmitCommandLists(TArrayView<IRHIPlatformCommandList*>
 
 	if (InnerCommandLists.Num())
 	{
-		RHI->RHISubmitCommandLists(InnerCommandLists);
+		RHI->RHISubmitCommandLists(InnerCommandLists, bFlushResources);
 	}
 }
 
diff --git a/Engine/Source/Runtime/RHI/Public/DynamicRHI.h b/Engine/Source/Runtime/RHI/Public/DynamicRHI.h
--- a/Engine/Source/Runtime/RHI/Public/DynamicRHI.h
+++ b/Engine/Source/Runtime/RHI/Public/DynamicRHI.h
@@ -892,7 +892,7 @@ public:
 	//
 	// Called by the RHI thread. 
 	//
-	virtual void RHISubmitCommandLists(TArrayView<IRHIPlatformCommandList*> CommandLists) = 0;
+	virtual void RHISubmitCommandLists(TArrayView<IRHIPlatformCommandList*> CommandLists, bool bFlushResources) = 0;
 
 	UE_DEPRECATED(5.1, "CreateBuffer_RenderThread is deprecated. Use RHICreateBuffer instead.")
 	RHI_API virtual FBufferRHIRef CreateBuffer_RenderThread(class FRHICommandListBase& RHICmdList, uint32 Size, EBufferUsageFlags Usage, uint32 Stride, ERHIAccess ResourceState, FRHIResourceCreateInfo& CreateInfo);
diff --git a/Engine/Source/Runtime/RHI/Public/RHICommandList.h b/Engine/Source/Runtime/RHI/Public/RHICommandList.h
--- a/Engine/Source/Runtime/RHI/Public/RHICommandList.h
+++ b/Engine/Source/Runtime/RHI/Public/RHICommandList.h
@@ -3605,7 +3605,7 @@ class RHI_API FRHICommandListImmediate : public FRHICommandList
 	//
 	// Used internally. Do not call directly. Use FRHICommandListImmediate::ImmediateFlush() to submit GPU work.
 	//
-	void ExecuteAndReset();
+	void ExecuteAndReset(bool bFlushResources);
 
 	//
 	// Blocks the calling thread until all dispatch prerequisites of enqueued parallel command lists are completed.
diff --git a/Engine/Source/Runtime/RHI/Public/RHICommandList.inl b/Engine/Source/Runtime/RHI/Public/RHICommandList.inl
--- a/Engine/Source/Runtime/RHI/Public/RHICommandList.inl
+++ b/Engine/Source/Runtime/RHI/Public/RHICommandList.inl
@@ -81,7 +81,8 @@ FORCEINLINE_DEBUGGABLE void FRHICommandListImmediate::ImmediateFlush(EImmediateF
 		{
 			// Execution and initialization are separate functions because initializing the immediate contexts
 			// may enqueue a lambda call to SwitchPipeline, which needs special handling in LatchBypass().
-			ExecuteAndReset();
+			const bool bFlushingResources = FlushType >= EImmediateFlushType::FlushRHIThreadFlushResources;
+			ExecuteAndReset(bFlushingResources);
 			InitializeImmediateContexts();
 		}
 
@@ -94,10 +95,20 @@ FORCEINLINE_DEBUGGABLE void FRHICommandListImmediate::ImmediateFlush(EImmediateF
 		if (FlushType >= EImmediateFlushType::FlushRHIThreadFlushResources)
 		{
 			CSV_SCOPED_TIMING_STAT(RHITFlushes, FlushRHIThreadFlushResourcesTotal);
-			// @todo: do this before the dispatch, since FlushResources enqueues a lambda to hand down the list of RHI resources to the RHI thread.
-			// Also work out when the PSO cache needs to be flushed (before dispatch, or after thread flush?)
 			PipelineStateCache::FlushResources();
-			FRHIResource::FlushPendingDeletes(*this);
+
+			// RHIPerFrameRHIFlushComplete was originally called from FlushPendingDeletes, which used to be here, so it was
+			// running as part of the next command list. FlushPendingDeletes now runs as part of the command list being finalized
+			// and submitted by this function, as it should be, but moving RHIPerFrameRHIFlushComplete there is risky, because
+			// the RHIs which use it do weird things in there (e.g. D3D11 does blocking query resolve calls in that function which
+			// happen to not stall because this runs when the next command list is executed, so they have time to be executed).
+			EnqueueLambda([](FRHICommandListImmediate& RHICmdList)
+			{
+				if (GDynamicRHI)
+				{
+					GDynamicRHI->RHIPerFrameRHIFlushComplete();
+				}
+			});
 		}
 	}
 }
diff --git a/Engine/Source/Runtime/RHI/Public/RHIResources.h b/Engine/Source/Runtime/RHI/Public/RHIResources.h
--- a/Engine/Source/Runtime/RHI/Public/RHIResources.h
+++ b/Engine/Source/Runtime/RHI/Public/RHIResources.h
@@ -89,6 +89,7 @@ public:
 		return uint32(CurrentValue);
 	}
 
+	UE_DEPRECATED(5.3, "FlushPendingDeletes is deprecated, please use FRHICommandListExecutor::GetImmediateCommandList().ImmediateFlush(EImmediateFlushType::FlushRHIThreadFlushResources)")
 	RHI_API static int32 FlushPendingDeletes(FRHICommandListImmediate& RHICmdList);
 
 	RHI_API static bool Bypass();
diff --git a/Engine/Source/Runtime/RHI/Public/RHIValidation.h b/Engine/Source/Runtime/RHI/Public/RHIValidation.h
--- a/Engine/Source/Runtime/RHI/Public/RHIValidation.h
+++ b/Engine/Source/Runtime/RHI/Public/RHIValidation.h
@@ -1110,7 +1110,7 @@ public:
 	virtual IRHIComputeContext* RHIGetDefaultAsyncComputeContext() override final;
 	virtual IRHIComputeContext* RHIGetCommandContext(ERHIPipeline Pipeline, FRHIGPUMask GPUMask) override final;
 	virtual IRHIPlatformCommandList* RHIFinalizeContext(IRHIComputeContext* OuterContext) override final;
-	virtual void RHISubmitCommandLists(TArrayView<IRHIPlatformCommandList*> OuterCommandLists) override final;
+	virtual void RHISubmitCommandLists(TArrayView<IRHIPlatformCommandList*> OuterCommandLists, bool bFlushResources) override final;
 
 	virtual uint64 RHIGetMinimumAlignmentForBufferBackedSRV(EPixelFormat Format) override final
 	{
diff --git a/Engine/Source/Runtime/RenderCore/Private/DumpGPU.cpp b/Engine/Source/Runtime/RenderCore/Private/DumpGPU.cpp
--- a/Engine/Source/Runtime/RenderCore/Private/DumpGPU.cpp
+++ b/Engine/Source/Runtime/RenderCore/Private/DumpGPU.cpp
@@ -455,7 +455,6 @@ public:
 		{
 			RHICmdList.SubmitCommandsAndFlushGPU();
 			RHICmdList.BlockUntilGPUIdle();
-			FRHIResource::FlushPendingDeletes(RHICmdList);
 			RHICmdList.FlushResources();
 			RHICmdList.ImmediateFlush(EImmediateFlushType::FlushRHIThreadFlushResources);
 		}
diff --git a/Engine/Source/Runtime/RenderCore/Private/RenderingThread.cpp b/Engine/Source/Runtime/RenderCore/Private/RenderingThread.cpp
--- a/Engine/Source/Runtime/RenderCore/Private/RenderingThread.cpp
+++ b/Engine/Source/Runtime/RenderCore/Private/RenderingThread.cpp
@@ -1363,15 +1363,15 @@ void FlushPendingDeleteRHIResources_GameThread()
 			});
 	}
 }
+
 void FlushPendingDeleteRHIResources_RenderThread()
 {
 	if (!IsRunningRHIInSeparateThread())
 	{
-		FRHIResource::FlushPendingDeletes(FRHICommandListExecutor::GetImmediateCommandList());
+		FRHICommandListExecutor::GetImmediateCommandList().ImmediateFlush(EImmediateFlushType::FlushRHIThreadFlushResources);
 	}
 }
 
-
 FRHICommandListImmediate& GetImmediateCommandList_ForRenderCommand()
 {
 	return FRHICommandListExecutor::GetImmediateCommandList();
diff --git a/Engine/Source/Runtime/VulkanRHI/Private/VulkanCommands.cpp b/Engine/Source/Runtime/VulkanRHI/Private/VulkanCommands.cpp
--- a/Engine/Source/Runtime/VulkanRHI/Private/VulkanCommands.cpp
+++ b/Engine/Source/Runtime/VulkanRHI/Private/VulkanCommands.cpp
@@ -923,7 +923,7 @@ IRHIPlatformCommandList* FVulkanDynamicRHI::RHIFinalizeContext(IRHIComputeContex
 	return PlatformCmdList;
 }
 
-void FVulkanDynamicRHI::RHISubmitCommandLists(TArrayView<IRHIPlatformCommandList*> CommandLists)
+void FVulkanDynamicRHI::RHISubmitCommandLists(TArrayView<IRHIPlatformCommandList*> CommandLists, bool bFlushResources)
 {
 	for (IRHIPlatformCommandList* Ptr : CommandLists)
 	{
diff --git a/Engine/Source/Runtime/VulkanRHI/Private/VulkanDevice.cpp b/Engine/Source/Runtime/VulkanRHI/Private/VulkanDevice.cpp
--- a/Engine/Source/Runtime/VulkanRHI/Private/VulkanDevice.cpp
+++ b/Engine/Source/Runtime/VulkanRHI/Private/VulkanDevice.cpp
@@ -1328,7 +1328,7 @@ void FVulkanDevice::Destroy()
 	do
 	{
 		FRHICommandListImmediate& RHICmdList = FRHICommandListExecutor::GetImmediateCommandList();
-		NumDeletes = FRHIResource::FlushPendingDeletes(RHICmdList);
+		NumDeletes = RHICmdList.FlushPendingDeletes();
 		RHICmdList.ImmediateFlush(EImmediateFlushType::FlushRHIThread);
 	} while (NumDeletes > 0);
 
diff --git a/Engine/Source/Runtime/VulkanRHI/Private/VulkanRHI.cpp b/Engine/Source/Runtime/VulkanRHI/Private/VulkanRHI.cpp
--- a/Engine/Source/Runtime/VulkanRHI/Private/VulkanRHI.cpp
+++ b/Engine/Source/Runtime/VulkanRHI/Private/VulkanRHI.cpp
@@ -591,11 +591,11 @@ void FVulkanDynamicRHI::Shutdown()
 		FRHICommandListImmediate& RHICmdList = FRHICommandListExecutor::GetImmediateCommandList();
 
 		// Flush all pending deletes before destroying the device.
-		FRHIResource::FlushPendingDeletes(RHICmdList);
+		RHICmdList.FlushPendingDeletes();
 		RHICmdList.ImmediateFlush(EImmediateFlushType::FlushRHIThread);
 
 		// And again since some might get on a pending queue
-		FRHIResource::FlushPendingDeletes(RHICmdList);
+		RHICmdList.FlushPendingDeletes();
 		RHICmdList.ImmediateFlush(EImmediateFlushType::FlushRHIThread);
 	}
 
diff --git a/Engine/Source/Runtime/VulkanRHI/Public/VulkanDynamicRHI.h b/Engine/Source/Runtime/VulkanRHI/Public/VulkanDynamicRHI.h
--- a/Engine/Source/Runtime/VulkanRHI/Public/VulkanDynamicRHI.h
+++ b/Engine/Source/Runtime/VulkanRHI/Public/VulkanDynamicRHI.h
@@ -191,7 +191,7 @@ public:
 	virtual class IRHIComputeContext* RHIGetDefaultAsyncComputeContext() final override;
 	virtual IRHIComputeContext* RHIGetCommandContext(ERHIPipeline Pipeline, FRHIGPUMask GPUMask) final override;
 	virtual IRHIPlatformCommandList* RHIFinalizeContext(IRHIComputeContext* Context) final override;
-	virtual void RHISubmitCommandLists(TArrayView<IRHIPlatformCommandList*> CommandLists) final override;
+	virtual void RHISubmitCommandLists(TArrayView<IRHIPlatformCommandList*> CommandLists, bool bFlushResources) final override;
 	virtual uint64 RHIGetMinimumAlignmentForBufferBackedSRV(EPixelFormat Format) final override;
 
 	virtual uint64 RHIComputePrecachePSOHash(const FGraphicsPipelineStateInitializer& Initializer) final override;
diff --git a/Engine/Source/Runtime/Windows/D3D11RHI/Private/D3D11Commands.cpp b/Engine/Source/Runtime/Windows/D3D11RHI/Private/D3D11Commands.cpp
--- a/Engine/Source/Runtime/Windows/D3D11RHI/Private/D3D11Commands.cpp
+++ b/Engine/Source/Runtime/Windows/D3D11RHI/Private/D3D11Commands.cpp
@@ -1555,9 +1555,8 @@ IRHIPlatformCommandList* FD3D11DynamicRHI::RHIFinalizeContext(IRHIComputeContext
 	return nullptr;
 }
 
-void FD3D11DynamicRHI::RHISubmitCommandLists(TArrayView<IRHIPlatformCommandList*> CommandLists)
+void FD3D11DynamicRHI::RHISubmitCommandLists(TArrayView<IRHIPlatformCommandList*> CommandLists, bool bFlushResources)
 {
-	UE_LOG(LogRHI, Fatal, TEXT("FD3D11DynamicRHI::RHISubmitCommandLists should never be called. D3D11 RHI does not implement parallel command list execution."));
 }
 
 void FD3D11DynamicRHI::EnableUAVOverlap()
diff --git a/Engine/Source/Runtime/Windows/D3D11RHI/Private/D3D11Device.cpp b/Engine/Source/Runtime/Windows/D3D11RHI/Private/D3D11Device.cpp
--- a/Engine/Source/Runtime/Windows/D3D11RHI/Private/D3D11Device.cpp
+++ b/Engine/Source/Runtime/Windows/D3D11RHI/Private/D3D11Device.cpp
@@ -607,7 +607,7 @@ void FD3D11DynamicRHI::CleanupD3DDevice()
 		do
 		{
 			FRHICommandListImmediate& RHICmdList = FRHICommandListExecutor::GetImmediateCommandList();
-			NumDeletes = FRHIResource::FlushPendingDeletes(RHICmdList);
+			NumDeletes = RHICmdList.FlushPendingDeletes();
 			RHICmdList.ImmediateFlush(EImmediateFlushType::FlushRHIThread);
 		} while (NumDeletes > 0);
 
diff --git a/Engine/Source/Runtime/Windows/D3D11RHI/Private/D3D11RHIPrivate.h b/Engine/Source/Runtime/Windows/D3D11RHI/Private/D3D11RHIPrivate.h
--- a/Engine/Source/Runtime/Windows/D3D11RHI/Private/D3D11RHIPrivate.h
+++ b/Engine/Source/Runtime/Windows/D3D11RHI/Private/D3D11RHIPrivate.h
@@ -551,7 +551,7 @@ public:
 	virtual class IRHICommandContext* RHIGetDefaultContext() final override;
 	virtual IRHIComputeContext* RHIGetCommandContext(ERHIPipeline Pipeline, FRHIGPUMask GPUMask) final override;
 	virtual IRHIPlatformCommandList* RHIFinalizeContext(IRHIComputeContext* Context) final override;
-	virtual void RHISubmitCommandLists(TArrayView<IRHIPlatformCommandList*> CommandLists) final override;
+	virtual void RHISubmitCommandLists(TArrayView<IRHIPlatformCommandList*> CommandLists, bool bFlushResources) final override;
 
 	// SRV / UAV creation functions
 	virtual FShaderResourceViewRHIRef  RHICreateShaderResourceView (class FRHICommandListImmediate& RHICmdList, FRHIViewableResource* Resource, FRHIViewDesc const& ViewDesc) final override;
