Use background threads only when compiling compute PSOs from the cache.

diff --git a/Engine/Source/Runtime/RHI/Private/PipelineStateCache.cpp b/Engine/Source/Runtime/RHI/Private/PipelineStateCache.cpp
--- a/Engine/Source/Runtime/RHI/Private/PipelineStateCache.cpp
+++ b/Engine/Source/Runtime/RHI/Private/PipelineStateCache.cpp
@@ -1910,7 +1910,7 @@ public:
 		// and it's better to leave more CPU time to these extrenal processes and other engine threads.
 		// Also use background threads for PSO precaching when the PSO thread pool is not used
 		// Compute pipelines usually take much longer to compile, compile them on background thread as well.
-		return (PLATFORM_MAC || PSOPreCacheResult == EPSOPrecacheResult::Active || (Pipeline && Pipeline->IsCompute())) ? ENamedThreads::AnyBackgroundThreadNormalTask : DesiredThread;
+		return (PLATFORM_MAC || PSOPreCacheResult == EPSOPrecacheResult::Active || (Pipeline && Pipeline->IsCompute() && Initializer.bFromPSOFileCache)) ? ENamedThreads::AnyBackgroundThreadNormalTask : DesiredThread;
 	}
 };
 
