Fix assert on exit when running with RHI validation.

FValidationRHI::RHISubmitCommandLists should forward the call to the underlying RHI when flushing resources, even when we don't have any command lists to execute, so the RHIs can process the deferred deletion queue.

#rnx

diff --git a/Engine/Source/Runtime/RHI/Private/RHIValidation.cpp b/Engine/Source/Runtime/RHI/Private/RHIValidation.cpp
--- a/Engine/Source/Runtime/RHI/Private/RHIValidation.cpp
+++ b/Engine/Source/Runtime/RHI/Private/RHIValidation.cpp
@@ -453,7 +453,7 @@ void FValidationRHI::RHISubmitCommandLists(TArrayView<IRHIPlatformCommandList*>
 		delete OuterCommandList;
 	}
 
-	if (InnerCommandLists.Num())
+	if (InnerCommandLists.Num() || bFlushResources)
 	{
 		RHI->RHISubmitCommandLists(InnerCommandLists, bFlushResources);
 	}
