Moving more RHI_API uses from entire classes to specific symbols.

diff --git a/Engine/Source/Runtime/RHI/Public/GPUProfiler.h b/Engine/Source/Runtime/RHI/Public/GPUProfiler.h
--- a/Engine/Source/Runtime/RHI/Public/GPUProfiler.h
+++ b/Engine/Source/Runtime/RHI/Public/GPUProfiler.h
@@ -116,7 +116,7 @@ public:
 };
 
 /** An entire frame of perf event nodes, including ancillary timers. */
-struct RHI_API FGPUProfilerEventNodeFrame
+struct FGPUProfilerEventNodeFrame
 {
 	virtual ~FGPUProfilerEventNodeFrame() {}
 
@@ -130,7 +130,7 @@ struct RHI_API FGPUProfilerEventNodeFrame
 	virtual void EndFrame() {}
 
 	/** Dumps perf event information, blocking on GPU. */
-	void DumpEventTree();
+	RHI_API void DumpEventTree();
 
 	/** Calculates root timing base frequency (if needed by this RHI) */
 	virtual float GetRootTimingResults() { return 0.0f; }
@@ -154,7 +154,7 @@ struct FGPUTimingCalibrationTimestamp
 /**
  * Holds information if this platform's GPU allows timing
  */
-struct RHI_API FGPUTiming
+struct FGPUTiming
 {
 public:
 	/**
@@ -209,10 +209,10 @@ public:
 
 protected:
 	/** Whether the static variables have been initialized. */
-	static bool		GAreGlobalsInitialized;
+	RHI_API static bool		GAreGlobalsInitialized;
 
 	/** Whether GPU timing measurements are supported by the driver. */
-	static bool		GIsSupported;
+	RHI_API static bool		GIsSupported;
 
 	static void SetTimingFrequency(uint64 TimingFrequency, uint32 GPUIndex = 0)
 	{
@@ -226,21 +226,21 @@ protected:
 
 private:
 	/** Frequency for the timing values, in number of ticks per seconds, or 0 if the feature isn't supported. */
-	static TStaticArray<uint64, MAX_NUM_GPUS>	GTimingFrequency;
+	RHI_API static TStaticArray<uint64, MAX_NUM_GPUS>	GTimingFrequency;
 
 	/**
 	* Two timestamps performed on GPU and CPU at nearly the same time.
 	* This can be used to visualize GPU and CPU timing events on the same timeline.
 	* Both values may be 0 if timer calibration is not available on current platform.
 	*/
-	static TStaticArray<FGPUTimingCalibrationTimestamp, MAX_NUM_GPUS> GCalibrationTimestamp;
+	RHI_API static TStaticArray<FGPUTimingCalibrationTimestamp, MAX_NUM_GPUS> GCalibrationTimestamp;
 };
 
 /** 
  * Encapsulates GPU profiling logic and data. 
  * There's only one global instance of this struct so it should only contain global data, nothing specific to a frame.
  */
-struct RHI_API FGPUProfiler
+struct FGPUProfiler
 {
 	/** Whether we are currently tracking perf events or not. */
 	bool bTrackingEvents;
@@ -324,6 +324,6 @@ struct RHI_API FGPUProfiler
 		return new FGPUProfilerEventNode(InName, InParent);
 	}
 
-	virtual void PushEvent(const TCHAR* Name, FColor Color);
-	virtual void PopEvent();
+	RHI_API virtual void PushEvent(const TCHAR* Name, FColor Color);
+	RHI_API virtual void PopEvent();
 };
\ No newline at end of file
diff --git a/Engine/Source/Runtime/RHI/Public/PipelineFileCache.h b/Engine/Source/Runtime/RHI/Public/PipelineFileCache.h
--- a/Engine/Source/Runtime/RHI/Public/PipelineFileCache.h
+++ b/Engine/Source/Runtime/RHI/Public/PipelineFileCache.h
@@ -23,7 +23,7 @@ DECLARE_DWORD_ACCUMULATOR_STAT_EXTERN(TEXT("Total RayTracing Pipeline State Coun
  */
 #define PSO_COOKONLY_DATA (WITH_EDITOR || IS_PROGRAM)
 
-struct RHI_API FPipelineFileCacheRasterizerState
+struct FPipelineFileCacheRasterizerState
 {
 	FPipelineFileCacheRasterizerState() { FMemory::Memzero(*this); }
 	FPipelineFileCacheRasterizerState(FRasterizerStateInitializerRHI const& Other) { operator=(Other); }
@@ -54,7 +54,7 @@ struct RHI_API FPipelineFileCacheRasterizerState
 		return Initializer;
 	}
 	
-	friend RHI_API FArchive& operator<<(FArchive& Ar,FPipelineFileCacheRasterizerState& RasterizerStateInitializer)
+	friend FArchive& operator<<(FArchive& Ar,FPipelineFileCacheRasterizerState& RasterizerStateInitializer)
 	{
 		Ar << RasterizerStateInitializer.DepthBias;
 		Ar << RasterizerStateInitializer.SlopeScaleDepthBias;
@@ -66,7 +66,7 @@ struct RHI_API FPipelineFileCacheRasterizerState
 		return Ar;
 	}
 
-	friend RHI_API uint32 GetTypeHash(const FPipelineFileCacheRasterizerState &Key)
+	friend uint32 GetTypeHash(const FPipelineFileCacheRasterizerState &Key)
 	{
 		uint32 KeyHash = (*((uint32*)&Key.DepthBias) ^ *((uint32*)&Key.SlopeScaleDepthBias));
 		KeyHash ^= (Key.FillMode << 8);
@@ -76,8 +76,8 @@ struct RHI_API FPipelineFileCacheRasterizerState
 		KeyHash ^= Key.bEnableLineAA ? 0x48271d01 : 0; // crc32 "bEnableLineAA"
 		return KeyHash;
 	}
-	FString ToString() const;
-	void FromString(const FStringView& Src);
+	RHI_API FString ToString() const;
+	RHI_API void FromString(const FStringView& Src);
 };
 
 class FRayTracingPipelineStateInitializer;
@@ -87,7 +87,7 @@ enum class ERayTracingPipelineCacheFlags : uint8;
 /**
  * Tracks stats for the current session between opening & closing the file-cache.
  */
-struct RHI_API FPipelineStateStats
+struct FPipelineStateStats
 {
 	FPipelineStateStats()
 	: FirstFrameUsed(-1)
@@ -102,9 +102,9 @@ struct RHI_API FPipelineStateStats
 	{
 	}
 
-	static void UpdateStats(FPipelineStateStats* Stats);
+	RHI_API static void UpdateStats(FPipelineStateStats* Stats);
 	
-	friend RHI_API FArchive& operator<<( FArchive& Ar, FPipelineStateStats& Info );
+	friend FArchive& operator<<( FArchive& Ar, FPipelineStateStats& Info );
 
 	int64 FirstFrameUsed;
 	int64 LastFrameUsed;
@@ -113,20 +113,20 @@ struct RHI_API FPipelineStateStats
 	uint32 PSOHash;
 };
 
-struct RHI_API FPipelineCacheFileFormatPSO
+struct FPipelineCacheFileFormatPSO
 {
 	using TReadableStringBuilder = TStringBuilder<1024>;
 
-	struct RHI_API ComputeDescriptor
+	struct ComputeDescriptor
 	{
 		FSHAHash ComputeShader;
 
-		FString ToString() const;
+		RHI_API FString ToString() const;
 		void AddToReadableString(TReadableStringBuilder& OutBuilder) const;
 		static FString HeaderLine();
-		void FromString(const FStringView& Src);
+		RHI_API void FromString(const FStringView& Src);
 	};
-	struct RHI_API GraphicsDescriptor
+	struct GraphicsDescriptor
 	{
 		FSHAHash VertexShader;
 		FSHAHash FragmentShader;
@@ -159,21 +159,21 @@ struct RHI_API FPipelineCacheFileFormatPSO
 		uint8	MultiViewCount;
 		bool	bHasFragmentDensityAttachment;
 		
-		FString ToString() const;
-		void AddToReadableString(TReadableStringBuilder& OutBuilder) const;
-		static FString HeaderLine();
-		bool FromString(const FStringView& Src);
+		RHI_API FString ToString() const;
+		RHI_API void AddToReadableString(TReadableStringBuilder& OutBuilder) const;
+		RHI_API static FString HeaderLine();
+		RHI_API bool FromString(const FStringView& Src);
 
-		FString ShadersToString() const;
-		void AddShadersToReadableString(TReadableStringBuilder& OutBuilder) const;
+		RHI_API FString ShadersToString() const;
+		RHI_API void AddShadersToReadableString(TReadableStringBuilder& OutBuilder) const;
 
-		static FString ShaderHeaderLine();
-		void ShadersFromString(const FStringView& Src);
+		RHI_API static FString ShaderHeaderLine();
+		RHI_API void ShadersFromString(const FStringView& Src);
 
-		FString StateToString() const;
-		void AddStateToReadableString(TReadableStringBuilder& OutBuilder) const;
-		static FString StateHeaderLine();
-		bool StateFromString(const FStringView& Src);
+		RHI_API FString StateToString() const;
+		RHI_API void AddStateToReadableString(TReadableStringBuilder& OutBuilder) const;
+		RHI_API static FString StateHeaderLine();
+		RHI_API bool StateFromString(const FStringView& Src);
 
 		/** Not all RT flags make sense for the replayed PSO, only those that can influence the RT formats */
 		static ETextureCreateFlags ReduceRTFlags(ETextureCreateFlags InFlags);
@@ -181,7 +181,7 @@ struct RHI_API FPipelineCacheFileFormatPSO
 		/** Not all DepthStencil flags make sense for the replayed PSO, only those that can influence the actual format */
 		static ETextureCreateFlags ReduceDSFlags(ETextureCreateFlags InFlags);
 	};
-	struct RHI_API FPipelineFileCacheRayTracingDesc
+	struct FPipelineFileCacheRayTracingDesc
 	{
 		FSHAHash ShaderHash;
 		uint32 DeprecatedMaxPayloadSizeInBytes = 0;
@@ -191,10 +191,10 @@ struct RHI_API FPipelineCacheFileFormatPSO
 		FPipelineFileCacheRayTracingDesc() = default;
 		FPipelineFileCacheRayTracingDesc(const FRayTracingPipelineStateInitializer& Initializer, const FRHIRayTracingShader* ShaderRHI);
 
-		FString ToString() const;
+		RHI_API FString ToString() const;
 		void AddToReadableString(TReadableStringBuilder& OutBuilder) const;
 		FString HeaderLine() const;
-		void FromString(const FString& Src);
+		RHI_API void FromString(const FString& Src);
 
 		friend uint32 GetTypeHash(const FPipelineFileCacheRayTracingDesc& Desc)
 		{
@@ -227,13 +227,13 @@ struct RHI_API FPipelineCacheFileFormatPSO
 	int64 BindCount;
 #endif
 	
-	FPipelineCacheFileFormatPSO();
-	~FPipelineCacheFileFormatPSO();
+	RHI_API FPipelineCacheFileFormatPSO();
+	RHI_API ~FPipelineCacheFileFormatPSO();
 	
-	FPipelineCacheFileFormatPSO& operator=(const FPipelineCacheFileFormatPSO& Other);
-	FPipelineCacheFileFormatPSO(const FPipelineCacheFileFormatPSO& Other);
+	RHI_API FPipelineCacheFileFormatPSO& operator=(const FPipelineCacheFileFormatPSO& Other);
+	RHI_API FPipelineCacheFileFormatPSO(const FPipelineCacheFileFormatPSO& Other);
 	
-	bool operator==(const FPipelineCacheFileFormatPSO& Other) const;
+	RHI_API bool operator==(const FPipelineCacheFileFormatPSO& Other) const;
 
 	friend RHI_API uint32 GetTypeHash(const FPipelineCacheFileFormatPSO &Key);
 	friend RHI_API FArchive& operator<<( FArchive& Ar, FPipelineCacheFileFormatPSO& Info );
@@ -242,18 +242,18 @@ struct RHI_API FPipelineCacheFileFormatPSO
 	static bool Init(FPipelineCacheFileFormatPSO& PSO, FGraphicsPipelineStateInitializer const& Init);
 	static bool Init(FPipelineCacheFileFormatPSO & PSO, FPipelineFileCacheRayTracingDesc const& Desc);
 
-	FString CommonToString() const;
-	static FString CommonHeaderLine();
-	void CommonFromString(const FStringView& Src);
+	RHI_API FString CommonToString() const;
+	RHI_API static FString CommonHeaderLine();
+	RHI_API void CommonFromString(const FStringView& Src);
 
 	/** Prints out human-readable representation of the PSO, for any type */
-	FString ToStringReadable() const;
+	RHI_API FString ToStringReadable() const;
 	
 	// Potential cases for seperating verify logic if requiired: RunTime-Logging, RunTime-UserCaching, RunTime-PreCompile, CommandLet-Cooking
-	bool Verify() const;
+	RHI_API bool Verify() const;
 };
 
-struct RHI_API FPipelineCacheFileFormatPSORead
+struct FPipelineCacheFileFormatPSORead
 {	
 	FPipelineCacheFileFormatPSORead()
 	: Ar(nullptr)
@@ -284,7 +284,7 @@ struct RHI_API FPipelineCacheFileFormatPSORead
 	TSharedPtr<class IAsyncReadRequest, ESPMode::ThreadSafe> ReadRequest;
 };
 
-struct RHI_API FPipelineCachePSOHeader
+struct FPipelineCachePSOHeader
 {
 	TSet<FSHAHash> Shaders;
 	uint32 Hash;
diff --git a/Engine/Source/Runtime/RHI/Public/RHICommandList.h b/Engine/Source/Runtime/RHI/Public/RHICommandList.h
--- a/Engine/Source/Runtime/RHI/Public/RHICommandList.h
+++ b/Engine/Source/Runtime/RHI/Public/RHICommandList.h
@@ -263,7 +263,7 @@ enum class ERayTracingBindingType : uint8
 };
 #endif // RHI_RAYTRACING
 
-struct RHI_API FLockTracker
+struct FLockTracker
 {
 	struct FLockParams
 	{
@@ -480,7 +480,7 @@ public:
 	~FRHICommandListScopedPipelineGuard();
 };
 
-class RHI_API FRHICommandListBase : public FNoncopyable
+class FRHICommandListBase : public FNoncopyable
 {
 public:
 	enum class ERecordingThread
@@ -490,11 +490,11 @@ public:
 	};
 
 protected:
-	FRHICommandListBase(FRHIGPUMask InGPUMask, ERecordingThread InRecordingThread);
+	RHI_API FRHICommandListBase(FRHIGPUMask InGPUMask, ERecordingThread InRecordingThread);
 
 public:
-	FRHICommandListBase(FRHICommandListBase&& Other);
-	~FRHICommandListBase();
+	RHI_API FRHICommandListBase(FRHICommandListBase&& Other);
+	RHI_API ~FRHICommandListBase();
 
 	inline bool IsImmediate() const;
 	inline FRHICommandListImmediate& GetAsImmediate();
@@ -506,7 +506,7 @@ public:
 	// 
 	// Not safe to call after FinishRecording().
 	//
-	void AddDispatchPrerequisite(const FGraphEventRef& Prereq);
+	RHI_API void AddDispatchPrerequisite(const FGraphEventRef& Prereq);
 
 	//
 	// Marks the RHI command list as completed, allowing it to be dispatched to the RHI / parallel translate threads.
@@ -516,9 +516,9 @@ public:
 	// 
 	// Never call on the immediate command list.
 	//
-	void FinishRecording();
+	RHI_API void FinishRecording();
 
-	void SetCurrentStat(TStatId Stat);
+	RHI_API void SetCurrentStat(TStatId Stat);
 
 	FORCEINLINE_DEBUGGABLE void* Alloc(int64 AllocSize, int64 Alignment)
 	{
@@ -651,7 +651,7 @@ public:
 
 	inline bool Bypass() const;
 
-	ERHIPipeline SwitchPipeline(ERHIPipeline Pipeline);
+	RHI_API ERHIPipeline SwitchPipeline(ERHIPipeline Pipeline);
 
 	FORCEINLINE FRHIGPUMask GetGPUMask() const { return PersistentState.CurrentGPUMask; }
 
@@ -663,10 +663,10 @@ public:
 	void SetExecuteStat(TStatId Stat) { ExecuteStat = Stat; }
 
 #if HAS_GPU_STATS
-	void SetStatsCategory(FDrawCallCategoryName* Category);
+	RHI_API void SetStatsCategory(FDrawCallCategoryName* Category);
 #endif
 
-	FGraphEventRef RHIThreadFence(bool bSetLockFence = false);
+	RHI_API FGraphEventRef RHIThreadFence(bool bSetLockFence = false);
 
 	FORCEINLINE void* LockBuffer(FRHIBuffer* Buffer, uint32 Offset, uint32 SizeRHI, EResourceLockMode LockMode)
 	{
@@ -793,7 +793,7 @@ protected:
 protected:
 	// Blocks the calling thread until the dispatch event is completed.
 	// Used internally, do not call directly.
-	void WaitForDispatchEvent();
+	RHI_API void WaitForDispatchEvent();
 
 	FRHICommandBase*    Root            = nullptr;
 	FRHICommandBase**   CommandLink     = nullptr;
@@ -964,7 +964,7 @@ private:
 	FRHICommandListBase(FPersistentState&& InPersistentState);
 
 	// Replays recorded commands into the specified contexts. Used internally, do not call directly.
-	void Execute(TRHIPipelineArray<IRHIComputeContext*>& InOutPipeContexts, FPersistentState::FGPUStats* ParentStats);
+	RHI_API void Execute(TRHIPipelineArray<IRHIComputeContext*>& InOutPipeContexts, FPersistentState::FGPUStats* ParentStats);
 
 	friend class FRHICommandListExecutor;
 	friend class FRHICommandListIterator;
@@ -2070,7 +2070,7 @@ extern RHI_API FRHIComputePipelineState* ExecuteSetComputePipelineState(FCompute
 extern RHI_API FRHIGraphicsPipelineState* ExecuteSetGraphicsPipelineState(class FGraphicsPipelineState* GraphicsPipelineState);
 extern RHI_API FRHIRayTracingPipelineState* GetRHIRayTracingPipelineState(FRayTracingPipelineState*);
 
-class RHI_API FRHIComputeCommandList : public FRHICommandListBase
+class FRHIComputeCommandList : public FRHICommandListBase
 {
 protected:
 	void OnBoundShaderChanged(FRHIComputeShader* InBoundComputeShaderRHI)
@@ -2383,7 +2383,7 @@ PRAGMA_ENABLE_DEPRECATION_WARNINGS
 		}
 	}
 
-	void Transition(TArrayView<const FRHITransitionInfo> Infos);
+	RHI_API void Transition(TArrayView<const FRHITransitionInfo> Infos);
 
 	FORCEINLINE_DEBUGGABLE void BeginTransition(const FRHITransition* Transition)
 	{
@@ -2687,8 +2687,8 @@ PRAGMA_ENABLE_DEPRECATION_WARNINGS
 #endif // WITH_MGPU
 
 #if RHI_RAYTRACING
-	void BuildAccelerationStructure(FRHIRayTracingGeometry* Geometry);
-	void BuildAccelerationStructures(const TArrayView<const FRayTracingGeometryBuildParams> Params);
+	RHI_API void BuildAccelerationStructure(FRHIRayTracingGeometry* Geometry);
+	RHI_API void BuildAccelerationStructures(const TArrayView<const FRayTracingGeometryBuildParams> Params);
 
 	FORCEINLINE_DEBUGGABLE void BuildAccelerationStructures(const TArrayView<const FRayTracingGeometryBuildParams> Params, const FRHIBufferRange& ScratchBufferRange)
 	{
@@ -2752,7 +2752,7 @@ PRAGMA_ENABLE_DEPRECATION_WARNINGS
 
 template<> RHI_API void FRHICommandSetShaderParameters           <FRHIGraphicsShader>::Execute(FRHICommandListBase& CmdList);
 
-class RHI_API FRHICommandList : public FRHIComputeCommandList
+class FRHICommandList : public FRHIComputeCommandList
 {
 protected:
 	using FRHIComputeCommandList::OnBoundShaderChanged;
@@ -3654,15 +3654,15 @@ FBufferRHIRef RHICreateStructuredBuffer(uint32 Stride, uint32 Size, uint32 InUsa
 extern RHI_API ERHIAccess RHIGetDefaultResourceState(ETextureCreateFlags InUsage, bool bInHasInitialData);
 extern RHI_API ERHIAccess RHIGetDefaultResourceState(EBufferUsageFlags InUsage, bool bInHasInitialData);
 
-class RHI_API FRHICommandListImmediate : public FRHICommandList
+class FRHICommandListImmediate : public FRHICommandList
 {
 	friend class FRHICommandListExecutor;
 	friend class FRHICommandListScopedExtendResourceLifetime;
 	friend struct FRHICommandBeginFrame;
 
-	static FGraphEventArray WaitOutstandingTasks;
-	static FGraphEventRef   RHIThreadTask;
-	static FRHIDrawStats    FrameDrawStats;
+	RHI_API static FGraphEventArray WaitOutstandingTasks;
+	RHI_API static FGraphEventRef   RHIThreadTask;
+	RHI_API static FRHIDrawStats    FrameDrawStats;
 
 	FRHICommandListImmediate()
 		: FRHICommandList(FRHIGPUMask::All(), ERecordingThread::Render)
@@ -3685,40 +3685,40 @@ class RHI_API FRHICommandListImmediate : public FRHICommandList
 	//
 	// Used internally. Do not call directly. Use FRHICommandListImmediate::ImmediateFlush() to submit GPU work.
 	//
-	void ExecuteAndReset(bool bFlushResources);
+	RHI_API void ExecuteAndReset(bool bFlushResources);
 
 	//
 	// Blocks the calling thread until all dispatch prerequisites of enqueued parallel command lists are completed.
 	//
-	void WaitForTasks();
+	RHI_API void WaitForTasks();
 
 	//
 	// Blocks the calling thread until the RHI thread is idle.
 	//
-	void WaitForRHIThreadTasks();
+	RHI_API void WaitForRHIThreadTasks();
 
 	//
 	// Destroys and recreates the immediate command list.
 	//
-	void Reset();
+	RHI_API void Reset();
 
 	//
 	// Called on RHIBeginFrame. Updates the draw call counters / stats.
 	//
-	void ProcessStats();
+	RHI_API void ProcessStats();
 
 	//
 	// Called when all FRHICommandListScopedExtendResourceLifetime references are released to flush any deferred deletions.
 	//
-	int32 FlushExtendedLifetimeResourceDeletes();
+	RHI_API int32 FlushExtendedLifetimeResourceDeletes();
 
 public:
-	void BeginScene();
-	void EndScene();
-	void BeginDrawingViewport(FRHIViewport* Viewport, FRHITexture* RenderTargetRHI);
-	void EndDrawingViewport(FRHIViewport* Viewport, bool bPresent, bool bLockToVsync);
-	void BeginFrame();
-	void EndFrame();
+	RHI_API void BeginScene();
+	RHI_API void EndScene();
+	RHI_API void BeginDrawingViewport(FRHIViewport* Viewport, FRHITexture* RenderTargetRHI);
+	RHI_API void EndDrawingViewport(FRHIViewport* Viewport, bool bPresent, bool bLockToVsync);
+	RHI_API void BeginFrame();
+	RHI_API void EndFrame();
 
 	struct FQueuedCommandList
 	{
@@ -3747,7 +3747,7 @@ public:
 	// Chains together one or more RHI command lists into the immediate command list, allowing in-order submission of parallel rendering work.
 	// The provided command lists are not dispatched until FinishRecording() is called on them, and their dispatch prerequisites have been completed.
 	//
-	void QueueAsyncCommandListSubmit(TArrayView<FQueuedCommandList> CommandLists, ETranslatePriority ParallelTranslatePriority = ETranslatePriority::Disabled, int32 MinDrawsPerTranslate = 0);
+	RHI_API void QueueAsyncCommandListSubmit(TArrayView<FQueuedCommandList> CommandLists, ETranslatePriority ParallelTranslatePriority = ETranslatePriority::Disabled, int32 MinDrawsPerTranslate = 0);
 	
 	inline void QueueAsyncCommandListSubmit(FQueuedCommandList QueuedCommandList, ETranslatePriority ParallelTranslatePriority = ETranslatePriority::Disabled, int32 MinDrawsPerTranslate = 0)
 	{
@@ -3758,24 +3758,24 @@ public:
 	// Dispatches work to the RHI thread and the GPU.
 	// Also optionally waits for its completion on the RHI thread. Does not wait for the GPU.
 	//
-	void ImmediateFlush(EImmediateFlushType::Type FlushType);
+	RHI_API void ImmediateFlush(EImmediateFlushType::Type FlushType);
 
-	bool StallRHIThread();
-	void UnStallRHIThread();
-	static bool IsStalled();
+	RHI_API bool StallRHIThread();
+	RHI_API void UnStallRHIThread();
+	RHI_API static bool IsStalled();
 
-	static FGraphEventArray& GetRenderThreadTaskArray();
+	RHI_API static FGraphEventArray& GetRenderThreadTaskArray();
 
-	void InitializeImmediateContexts();
+	RHI_API void InitializeImmediateContexts();
 
 	// Global graph events must be destroyed explicitly to avoid undefined order of static destruction, as they can be destroyed after their allocator.
-	static void CleanupGraphEvents();
+	RHI_API static void CleanupGraphEvents();
 
 	//
 	// Performs an immediate transition with the option of broadcasting to multiple pipelines.
 	// Uses both the immediate and async compute contexts. Falls back to graphics-only if async compute is not supported.
 	//
-	void Transition(TArrayView<const FRHITransitionInfo> Infos, ERHIPipeline SrcPipelines, ERHIPipeline DstPipelines);
+	RHI_API void Transition(TArrayView<const FRHITransitionInfo> Infos, ERHIPipeline SrcPipelines, ERHIPipeline DstPipelines);
 	using FRHIComputeCommandList::Transition;
 
 	template <typename LAMBDA>
@@ -4311,7 +4311,7 @@ public:
 		return GDynamicRHI->RHIFlushResources();
 	}
 
-	int32 FlushPendingDeletes();
+	RHI_API int32 FlushPendingDeletes();
 	
 	FORCEINLINE uint32 GetGPUFrameCycles()
 	{
@@ -4442,7 +4442,7 @@ public:
 		return RHIGetDefaultContext();
 	}
 
-	void UpdateTextureReference(FRHITextureReference* TextureRef, FRHITexture* NewTexture);
+	RHI_API void UpdateTextureReference(FRHITextureReference* TextureRef, FRHITexture* NewTexture);
 
 	FORCEINLINE void PollRenderQueryResults()
 	{
@@ -4454,7 +4454,7 @@ public:
 	 * @param Num - number of update infos
 	 * @param bNeedReleaseRefs - whether Release need to be called on RHI resources referenced by update infos
 	 */
-	void UpdateRHIResources(FRHIResourceUpdateInfo* UpdateInfos, int32 Num, bool bNeedReleaseRefs);
+	RHI_API void UpdateRHIResources(FRHIResourceUpdateInfo* UpdateInfos, int32 Num, bool bNeedReleaseRefs);
 
 	//UE_DEPRECATED(5.1, "SubmitCommandsHint is deprecated. Consider calling ImmediateFlush(EImmediateFlushType::DispatchToRHIThread) instead.")
 	FORCEINLINE_DEBUGGABLE void SubmitCommandsHint()
@@ -4495,7 +4495,7 @@ public:
 	}
 };
 
-class RHI_API FRHICommandListScopedExtendResourceLifetime
+class FRHICommandListScopedExtendResourceLifetime
 {
 public:
 	FRHICommandListScopedExtendResourceLifetime(FRHICommandListImmediate& InRHICmdList)
@@ -4566,7 +4566,7 @@ struct FRHIScopedGPUMask
 	#define SCOPED_GPU_MASK(RHICmdList, GPUMask)
 #endif // WITH_MGPU
 
-struct RHI_API FScopedUniformBufferStaticBindings
+struct FScopedUniformBufferStaticBindings
 {
 	FScopedUniformBufferStaticBindings(FRHIComputeCommandList& InRHICmdList, FUniformBufferStaticBindings UniformBuffers)
 		: RHICmdList(InRHICmdList)
@@ -4596,7 +4596,7 @@ struct RHI_API FScopedUniformBufferStaticBindings
 	FRHIComputeCommandList& RHICmdList;
 
 #if VALIDATE_UNIFORM_BUFFER_STATIC_BINDINGS
-	static bool bRecursionGuard;
+	RHI_API static bool bRecursionGuard;
 #endif
 };
 
@@ -4604,11 +4604,11 @@ struct RHI_API FScopedUniformBufferStaticBindings
 
 // Helper to enable the use of graphics RHI command lists from within platform RHI implementations.
 // Recorded commands are dispatched when the command list is destructed. Intended for use on the stack / in a scope block.
-class RHI_API FRHICommandList_RecursiveHazardous : public FRHICommandList
+class FRHICommandList_RecursiveHazardous : public FRHICommandList
 {
 public:
-	FRHICommandList_RecursiveHazardous(IRHICommandContext* Context);
-	~FRHICommandList_RecursiveHazardous();
+	RHI_API FRHICommandList_RecursiveHazardous(IRHICommandContext* Context);
+	RHI_API ~FRHICommandList_RecursiveHazardous();
 };
 
 // Helper class used internally by RHIs to make use of FRHICommandList_RecursiveHazardous safer.
@@ -4661,11 +4661,11 @@ public:
 
 // Helper to enable the use of compute RHI command lists from within platform RHI implementations.
 // Recorded commands are dispatched when the command list is destructed. Intended for use on the stack / in a scope block.
-class RHI_API FRHIComputeCommandList_RecursiveHazardous : public FRHIComputeCommandList
+class FRHIComputeCommandList_RecursiveHazardous : public FRHIComputeCommandList
 {
 public:
-	FRHIComputeCommandList_RecursiveHazardous(IRHIComputeContext* Context);
-	~FRHIComputeCommandList_RecursiveHazardous();
+	RHI_API FRHIComputeCommandList_RecursiveHazardous(IRHIComputeContext* Context);
+	RHI_API ~FRHIComputeCommandList_RecursiveHazardous();
 };
 
 // Helper class used internally by RHIs to make use of FRHIComputeCommandList_RecursiveHazardous safer.
@@ -4716,7 +4716,7 @@ public:
 	}
 };
 
-class RHI_API FRHICommandListExecutor
+class FRHICommandListExecutor
 {
 public:
 	enum
@@ -4729,9 +4729,9 @@ public:
 	{
 	}
 	static inline FRHICommandListImmediate& GetImmediateCommandList();
-	void LatchBypass();
+	RHI_API void LatchBypass();
 
-	static void WaitOnRHIThreadFence(FGraphEventRef& Fence);
+	RHI_API static void WaitOnRHIThreadFence(FGraphEventRef& Fence);
 
 	FORCEINLINE_DEBUGGABLE bool Bypass()
 	{
@@ -4753,8 +4753,8 @@ public:
 
 	static inline void CheckNoOutstandingCmdLists();
 
-	static bool IsRHIThreadActive();
-	static bool IsRHIThreadCompletelyFlushed();
+	RHI_API static bool IsRHIThreadActive();
+	RHI_API static bool IsRHIThreadCompletelyFlushed();
 
 private:
 	bool bLatchedBypass;
diff --git a/Engine/Source/Runtime/RHI/Public/RHIResourceUpdates.h b/Engine/Source/Runtime/RHI/Public/RHIResourceUpdates.h
--- a/Engine/Source/Runtime/RHI/Public/RHIResourceUpdates.h
+++ b/Engine/Source/Runtime/RHI/Public/RHIResourceUpdates.h
@@ -40,21 +40,21 @@ struct FRHIResourceUpdateInfo
 	void ReleaseRefs();
 };
 
-struct RHI_API FRHIResourceUpdateBatcher
+struct FRHIResourceUpdateBatcher
 {
 	TArrayView<FRHIResourceUpdateInfo> UpdateInfos;
 	int32 NumBatched = 0;
 
-	~FRHIResourceUpdateBatcher();
+	RHI_API ~FRHIResourceUpdateBatcher();
 
-	void Flush();
+	RHI_API void Flush();
 
-	void QueueUpdateRequest(FRHIBuffer* DestBuffer, FRHIBuffer* SrcBuffer);
-	void QueueUpdateRequest(FRHIRayTracingGeometry* DestGeometry, FRHIRayTracingGeometry* SrcGeometry);
+	RHI_API void QueueUpdateRequest(FRHIBuffer* DestBuffer, FRHIBuffer* SrcBuffer);
+	RHI_API void QueueUpdateRequest(FRHIRayTracingGeometry* DestGeometry, FRHIRayTracingGeometry* SrcGeometry);
 
 protected:
 	FRHIResourceUpdateBatcher() = delete;
-	FRHIResourceUpdateBatcher(TArrayView<FRHIResourceUpdateInfo> InUpdateInfos);
+	RHI_API FRHIResourceUpdateBatcher(TArrayView<FRHIResourceUpdateInfo> InUpdateInfos);
 
 private:
 	FRHIResourceUpdateInfo& GetNextUpdateInfo();
diff --git a/Engine/Source/Runtime/RHI/Public/RHIStats.h b/Engine/Source/Runtime/RHI/Public/RHIStats.h
--- a/Engine/Source/Runtime/RHI/Public/RHIStats.h
+++ b/Engine/Source/Runtime/RHI/Public/RHIStats.h
@@ -67,10 +67,10 @@ extern RHI_API int32 GNumPrimitivesDrawnRHI[MAX_NUM_GPUS];
 
 #if HAS_GPU_STATS
 
-struct RHI_API FDrawCallCategoryName
+struct FDrawCallCategoryName
 {
-	FDrawCallCategoryName();
-	FDrawCallCategoryName(FName InName);
+	RHI_API FDrawCallCategoryName();
+	RHI_API FDrawCallCategoryName(FName InName);
 
 	bool ShouldCountDraws() const { return Index != -1; }
 
@@ -88,10 +88,10 @@ struct RHI_API FDrawCallCategoryName
 
 		int32 NumCategory;
 
-		FManager();
+		RHI_API FManager();
 	};
 
-	static FManager& GetManager();
+	RHI_API static FManager& GetManager();
 };
 
 // RHI counter stats.
