Allowing custom names for RHI Lambda commands in the rhicommands trace channel.
Adding names to a few EnqueueLambda calls in RHICommandList.cpp

diff --git a/Engine/Source/Runtime/RHI/Private/RHICommandList.cpp b/Engine/Source/Runtime/RHI/Private/RHICommandList.cpp
--- a/Engine/Source/Runtime/RHI/Private/RHICommandList.cpp
+++ b/Engine/Source/Runtime/RHI/Private/RHICommandList.cpp
@@ -631,7 +631,7 @@ void FRHICommandListImmediate::QueueAsyncCommandListSubmit(TArrayView<FQueuedCom
 
 		// Finally, add an RHI thread task to submit the completed platform command lists.
 		// The task blocks for each parallel translate completion, in the order they will be submitted in.
-		EnqueueLambda([Tasks](FRHICommandListBase& ExecutingCmdList)
+		EnqueueLambda(TEXT("SubmitParallelCommandLists"), [Tasks](FRHICommandListBase& ExecutingCmdList)
 		{
 			TArray<IRHIPlatformCommandList*> AllCmdLists;
 
@@ -667,7 +667,7 @@ void FRHICommandListImmediate::QueueAsyncCommandListSubmit(TArrayView<FQueuedCom
 			CmdListsView[Index] = CommandList;
 		}
 
-		EnqueueLambda([CmdListsView](FRHICommandListBase& ParentCmdList)
+		EnqueueLambda(TEXT("ExecuteCommandLists"), [CmdListsView](FRHICommandListBase& ParentCmdList)
 		{
 			for (FRHICommandListBase* CmdList : CmdListsView)
 			{
@@ -708,7 +708,7 @@ void FRHICommandListImmediate::ExecuteAndReset(bool bFlushResources)
 	//
 	// In non-bypass mode, always submit.
 	//
-	EnqueueLambda([bFlushResources](FRHICommandListImmediate& ExecutingCmdList)
+	EnqueueLambda(TEXT("FinalizeAndSubmitCommandLists"), [bFlushResources](FRHICommandListImmediate& ExecutingCmdList)
 	{
 		TArray<IRHIPlatformCommandList*, TInlineAllocator<GetRHIPipelineCount()>> CommandLists;
 		for (IRHIComputeContext* Context : ExecutingCmdList.Contexts)
@@ -1040,7 +1040,7 @@ int32 FRHICommandListImmediate::FlushPendingDeletes()
 
 	const int32 NumDeletes = DeletedResources.Num();
 
-	EnqueueLambda([DeletedResources = MoveTemp(DeletedResources)](FRHICommandListImmediate& RHICmdList) mutable
+	EnqueueLambda(TEXT("FlushPendingDeletes"), [DeletedResources = MoveTemp(DeletedResources)](FRHICommandListImmediate& RHICmdList) mutable
 	{
 		SCOPED_NAMED_EVENT(STAT_FRHICommandListImmediate_DeleteResources, FColor::Magenta);
 		for (int32 i = DeletedResources.Num() - 1; i >= 0; i--)
@@ -1063,7 +1063,7 @@ int32 FRHICommandListImmediate::FlushExtendedLifetimeResourceDeletes()
 
 	if (NumDeletes > 0)
 	{
-		EnqueueLambda([DeletedResources = MoveTemp(PersistentState.ExtendedLifetimeResources)](FRHICommandListImmediate& RHICmdList) mutable
+		EnqueueLambda(TEXT("FlushExtendedLifetimeResourceDeletes"), [DeletedResources = MoveTemp(PersistentState.ExtendedLifetimeResources)](FRHICommandListImmediate& RHICmdList) mutable
 		{
 			SCOPED_NAMED_EVENT(STAT_FRHICommandListImmediate_DeleteExtendedLifetimeResources, FColor::Magenta);
 
@@ -1501,7 +1501,7 @@ void FDynamicRHI::RHIUnlockBuffer(class FRHICommandListBase& RHICmdList, FRHIBuf
 		}
 		else
 		{
-			RHICmdList.EnqueueLambda([Buffer, Params](FRHICommandListBase& InRHICmdList)
+			RHICmdList.EnqueueLambda(TEXT("RHIUnlockBuffer"), [Buffer, Params](FRHICommandListBase& InRHICmdList)
 			{
 				QUICK_SCOPE_CYCLE_COUNTER(STAT_FRHICommandUpdateBuffer_Execute);
 				void* Data = GDynamicRHI->LockBuffer_BottomOfPipe(InRHICmdList, Buffer, Params.Offset, Params.BufferSize, RLM_WriteOnly);
@@ -1795,7 +1795,7 @@ void FRHICommandListImmediate::UpdateTextureReference(FRHITextureReference* Text
 		return;
 	}
 
-	EnqueueLambda([TextureRef, NewTexture](FRHICommandListBase& RHICmdList)
+	EnqueueLambda(TEXT("UpdateTextureReference"), [TextureRef, NewTexture](FRHICommandListBase& RHICmdList)
 	{
 		GDynamicRHI->RHIUpdateTextureReference(RHICmdList, TextureRef, NewTexture);
 	});
diff --git a/Engine/Source/Runtime/RHI/Public/RHICommandList.h b/Engine/Source/Runtime/RHI/Public/RHICommandList.h
--- a/Engine/Source/Runtime/RHI/Public/RHICommandList.h
+++ b/Engine/Source/Runtime/RHI/Public/RHICommandList.h
@@ -448,14 +448,20 @@ template <typename RHICmdListType, typename LAMBDA>
 struct TRHILambdaCommand final : public FRHICommandBase
 {
 	LAMBDA Lambda;
+#if CPUPROFILERTRACE_ENABLED
+	const TCHAR* Name;
+#endif
 
-	TRHILambdaCommand(LAMBDA&& InLambda)
+	TRHILambdaCommand(LAMBDA&& InLambda, const TCHAR* InName)
 		: Lambda(Forward<LAMBDA>(InLambda))
+#if CPUPROFILERTRACE_ENABLED
+		, Name(InName)
+#endif
 	{}
 
 	void ExecuteAndDestruct(FRHICommandListBase& CmdList, FRHICommandListDebugContext&) override final
 	{
-		TRACE_CPUPROFILER_EVENT_SCOPE_ON_CHANNEL(TRHILambdaCommand, RHICommandsChannel);
+		TRACE_CPUPROFILER_EVENT_SCOPE_TEXT_ON_CHANNEL(Name, RHICommandsChannel);
 		Lambda(*static_cast<RHICmdListType*>(&CmdList));
 		Lambda.~LAMBDA();
 	}
@@ -585,7 +591,7 @@ public:
 	}
 
 	template <typename LAMBDA>
-	FORCEINLINE_DEBUGGABLE void EnqueueLambda(LAMBDA&& Lambda)
+	FORCEINLINE_DEBUGGABLE void EnqueueLambda(const TCHAR* LambdaName, LAMBDA&& Lambda)
 	{
 		if (IsBottomOfPipe())
 		{
@@ -593,10 +599,16 @@ public:
 		}
 		else
 		{
-			ALLOC_COMMAND(TRHILambdaCommand<FRHICommandListBase, LAMBDA>)(Forward<LAMBDA>(Lambda));
+			ALLOC_COMMAND(TRHILambdaCommand<FRHICommandListBase, LAMBDA>)(Forward<LAMBDA>(Lambda), LambdaName);
 		}
 	}
 
+	template <typename LAMBDA>
+	FORCEINLINE_DEBUGGABLE void EnqueueLambda(LAMBDA&& Lambda)
+	{
+		FRHICommandListBase::EnqueueLambda(TEXT("TRHILambdaCommand"), Forward<LAMBDA>(Lambda));
+	}
+
 	FORCEINLINE uint32 GetUID()  const
 	{
 		return UID;
@@ -2094,7 +2106,7 @@ public:
 	{}
 
 	template <typename LAMBDA>
-	FORCEINLINE_DEBUGGABLE void EnqueueLambda(LAMBDA&& Lambda)
+	FORCEINLINE_DEBUGGABLE void EnqueueLambda(const TCHAR* LambdaName, LAMBDA&& Lambda)
 	{
 		if (IsBottomOfPipe())
 		{
@@ -2102,10 +2114,16 @@ public:
 		}
 		else
 		{
-			ALLOC_COMMAND(TRHILambdaCommand<FRHIComputeCommandList, LAMBDA>)(Forward<LAMBDA>(Lambda));
+			ALLOC_COMMAND(TRHILambdaCommand<FRHIComputeCommandList, LAMBDA>)(Forward<LAMBDA>(Lambda), LambdaName);
 		}
 	}
 
+	template <typename LAMBDA>
+	FORCEINLINE_DEBUGGABLE void EnqueueLambda(LAMBDA&& Lambda)
+	{
+		FRHIComputeCommandList::EnqueueLambda(TEXT("TRHILambdaCommand"), Forward<LAMBDA>(Lambda));
+	}
+
 	inline FRHIComputeShader* GetBoundComputeShader() const { return PersistentState.BoundComputeShaderRHI; }
 
 	FORCEINLINE_DEBUGGABLE void SetStaticUniformBuffers(const FUniformBufferStaticBindings& UniformBuffers)
@@ -2778,7 +2796,7 @@ public:
 	inline FRHIGeometryShader*      GetBoundGeometryShader     () const { return PersistentState.BoundShaderInput.GetGeometryShader();      }
 
 	template <typename LAMBDA>
-	FORCEINLINE_DEBUGGABLE void EnqueueLambda(LAMBDA&& Lambda)
+	FORCEINLINE_DEBUGGABLE void EnqueueLambda(const TCHAR* LambdaName, LAMBDA&& Lambda)
 	{
 		if (IsBottomOfPipe())
 		{
@@ -2786,10 +2804,16 @@ public:
 		}
 		else
 		{
-			ALLOC_COMMAND(TRHILambdaCommand<FRHICommandList, LAMBDA>)(Forward<LAMBDA>(Lambda));
+			ALLOC_COMMAND(TRHILambdaCommand<FRHICommandList, LAMBDA>)(Forward<LAMBDA>(Lambda), LambdaName);
 		}
 	}
 
+	template <typename LAMBDA>
+	FORCEINLINE_DEBUGGABLE void EnqueueLambda(LAMBDA&& Lambda)
+	{
+		FRHICommandList::EnqueueLambda(TEXT("TRHILambdaCommand"), Forward<LAMBDA>(Lambda));
+	}
+
 	UE_DEPRECATED(5.2, "AFR support has been removed in 5.2, function is a nop and calls should be removed.")
 	FORCEINLINE void BeginUpdateMultiFrameResource(FRHITexture* Texture) {}
 
@@ -3779,7 +3803,7 @@ public:
 	using FRHIComputeCommandList::Transition;
 
 	template <typename LAMBDA>
-	FORCEINLINE_DEBUGGABLE void EnqueueLambda(LAMBDA&& Lambda)
+	FORCEINLINE_DEBUGGABLE void EnqueueLambda(const TCHAR* LambdaName, LAMBDA&& Lambda)
 	{
 		if (IsBottomOfPipe())
 		{
@@ -3787,10 +3811,16 @@ public:
 		}
 		else
 		{
-			ALLOC_COMMAND(TRHILambdaCommand<FRHICommandListImmediate, LAMBDA>)(Forward<LAMBDA>(Lambda));
+			ALLOC_COMMAND(TRHILambdaCommand<FRHICommandListImmediate, LAMBDA>)(Forward<LAMBDA>(Lambda), LambdaName);
 		}
 	}
 
+	template <typename LAMBDA>
+	FORCEINLINE_DEBUGGABLE void EnqueueLambda(LAMBDA&& Lambda)
+	{
+		FRHICommandListImmediate::EnqueueLambda(TEXT("TRHILambdaCommand"), Forward<LAMBDA>(Lambda));
+	}
+
 	UE_DEPRECATED(5.2, "Use the global scope RHICreateSamplerState function.")
 	FORCEINLINE FSamplerStateRHIRef CreateSamplerState(const FSamplerStateInitializerRHI& Initializer)
 	{
