Handle compute pipeline creation failures.

Graphics pipeline creation failures are treated as fatal (unless the PSO is coming from the file cache). This change implements the same logic for compute pipelines.

#pf 64750b934b1ead7c7f472cb9

diff --git a/Engine/Source/Runtime/RHI/Private/PipelineStateCache.cpp b/Engine/Source/Runtime/RHI/Private/PipelineStateCache.cpp
--- a/Engine/Source/Runtime/RHI/Private/PipelineStateCache.cpp
+++ b/Engine/Source/Runtime/RHI/Private/PipelineStateCache.cpp
@@ -326,18 +326,18 @@ static void HandlePipelineCreationFailure(const FGraphicsPipelineStateInitialize
 	{
 		UE_LOG(LogRHI, Error, TEXT("Pixel: %s"), Init.BoundShaderState.PixelShaderRHI->GetShaderName());
 	}
-	
+
 	UE_LOG(LogRHI, Error, TEXT("Render Targets: (%u)"), Init.RenderTargetFormats.Num());
 	for(int32 i = 0; i < Init.RenderTargetFormats.Num(); ++i)
 	{
 		//#todo-mattc GetPixelFormatString is not available in scw. Need to move it so we can print more info here.
 		UE_LOG(LogRHI, Error, TEXT("0x%x"), (uint32)Init.RenderTargetFormats[i]);
 	}
-	
+
 	UE_LOG(LogRHI, Error, TEXT("Depth Stencil Format:"));
 	UE_LOG(LogRHI, Error, TEXT("0x%x"), Init.DepthStencilTargetFormat);
 #endif
-	
+
 	if(Init.bFromPSOFileCache)
 	{
 		// Let the cache know so it hopefully won't give out this one again
@@ -349,6 +349,22 @@ static void HandlePipelineCreationFailure(const FGraphicsPipelineStateInitialize
 	}
 }
 
+// Prints out information about a failed compute pipeline compilation.
+// This is fatal unless the compilation request is from the PSO file cache preload.
+static void HandlePipelineCreationFailure(const FRHIComputeShader* ComputeShader, bool bFromPSOFileCache)
+{
+	UE_LOG(LogRHI, Error, TEXT("Failed to create ComputePipeline"));
+
+#if !(UE_BUILD_SHIPPING || UE_BUILD_TEST)
+	UE_LOG(LogRHI, Error, TEXT("Shader: %s"), ComputeShader->GetShaderName());
+#endif
+
+	if (!bFromPSOFileCache)
+	{
+		UE_LOG(LogRHI, Fatal, TEXT("Shader compilation failures are Fatal."));
+	}
+}
+
 // FAsyncTask used by the PSO threadpool that will reschedule once only.
 template<typename TTask>
 class FAsyncTaskLimitedReschedule : public FAsyncTask<TTask>
@@ -1771,6 +1787,11 @@ public:
 			ComputePipeline->RHIPipeline = RHICreateComputePipelineState(ComputePipeline->ComputeShader);
 			CheckAndUpdateHitchCountStat(FPSOPrecacheRequestID::EType::Compute, !IsPrecachedPSO(Initializer), StartTime);
 
+			if (!ComputePipeline->RHIPipeline)
+			{
+				HandlePipelineCreationFailure(ComputePipeline->ComputeShader, Initializer.bFromPSOFileCache);
+			}
+
 			if (Initializer.bPSOPrecache)
 			{
 				bool bCSValid = ComputePipeline->RHIPipeline != nullptr && ComputePipeline->RHIPipeline->IsValid();
@@ -2072,6 +2093,11 @@ FComputePipelineState* PipelineStateCache::GetAndOrCreateComputePipelineState(FR
 			uint64 StartTime = FPlatformTime::Cycles64();
 			OutCachedState->RHIPipeline = RHICreateComputePipelineState(OutCachedState->ComputeShader);
 			CheckAndUpdateHitchCountStat(FPSOPrecacheRequestID::EType::Compute, !bFromFileCache, StartTime);
+
+			if (!OutCachedState->RHIPipeline)
+			{
+				HandlePipelineCreationFailure(OutCachedState->ComputeShader, bFromFileCache);
+			}
 		}
 
 		GComputePipelineCache.Add(ComputeShader, OutCachedState, LockFlags);
