Add a config option to force certain GPUs to DX12 SM6 if they support it.

Some Intel cards reportedly work better in DX12 SM6 even without reaching the recommended VRAM threshold.

This is OFF by default for now until it can be tested.

#tests tested locally that this is OFF by default for now.
#pf skip

diff --git a/Engine/Source/Runtime/RHI/Private/Windows/WindowsDynamicRHI.cpp b/Engine/Source/Runtime/RHI/Private/Windows/WindowsDynamicRHI.cpp
--- a/Engine/Source/Runtime/RHI/Private/Windows/WindowsDynamicRHI.cpp
+++ b/Engine/Source/Runtime/RHI/Private/Windows/WindowsDynamicRHI.cpp
@@ -551,13 +551,40 @@ static bool ShouldDevicePreferSM5(uint32 DeviceId)
 	return false;
 }
 
+// Whether a SM6 capable device should always default to SM6
+// regardless of other heuristics suggesting otherwise.
+static bool ShouldDevicePreferSM6(uint32 SM6CapableDeviceId)
+{
+	TArray<FString> SM6PreferredDeviceIds;
+	GConfig->GetArray(TEXT("D3D12_SM6"), TEXT("SM6PreferredGPUDeviceIDs"), SM6PreferredDeviceIds, GEngineIni);
+
+	for (const FString& PreferredDeviceID : SM6PreferredDeviceIds)
+	{
+		if (SM6CapableDeviceId == FCString::Strtoi(*PreferredDeviceID, nullptr, 10) ||
+			SM6CapableDeviceId == FCString::Strtoi(*PreferredDeviceID, nullptr, 16))
+		{
+			return true;
+		}
+	}
+
+	return false;
+}
+
 static bool IsRHIAllowedAsDefault(EWindowsRHI InRHI, ERHIFeatureLevel::Type InFeatureLevel)
 {
+	static const FWindowsPlatformApplicationMisc::FGPUInfo BestGPUInfo = FWindowsPlatformApplicationMisc::GetBestGPUInfo();
+
+	if (InRHI == EWindowsRHI::D3D12 && InFeatureLevel == ERHIFeatureLevel::SM6)
+	{
+		if (ShouldDevicePreferSM6(BestGPUInfo.DeviceId)) 
+		{
+			return true;
+		}
+	}
+
 	bool bAllowed = true;
 	if (InRHI == EWindowsRHI::D3D12 && InFeatureLevel >= ERHIFeatureLevel::SM6)
 	{
-		static const FWindowsPlatformApplicationMisc::FGPUInfo BestGPUInfo = FWindowsPlatformApplicationMisc::GetBestGPUInfo();
-
 		int32 MinDedicatedMemoryMB = 0;
 		if (GConfig->GetInt(TEXT("D3D12_SM6"), TEXT("MinDedicatedMemory"), MinDedicatedMemoryMB, GEngineIni))
 		{
