Cleaning up some RHI commands
- Multi-GPU commands weren't allocating their arrays from the command list memory.
- Multi-GPU Context methods were still defined on non-MGPU platforms.
- Multi-GPU Context methods were missing from RHI Validation.
- Changed the majority of TArrayView<const ...> types to TConstArrayView.
- Removed some excessive const uses on value parameters.

diff --git a/Engine/Source/Runtime/D3D12RHI/Private/D3D12CommandContext.cpp b/Engine/Source/Runtime/D3D12RHI/Private/D3D12CommandContext.cpp
--- a/Engine/Source/Runtime/D3D12RHI/Private/D3D12CommandContext.cpp
+++ b/Engine/Source/Runtime/D3D12RHI/Private/D3D12CommandContext.cpp
@@ -931,9 +931,9 @@ FD3D12CommandContextRedirector::FD3D12CommandContextRedirector(class FD3D12Adapt
 		Context = nullptr;
 }
 
-void FD3D12CommandContextRedirector::RHITransferResources(const TArrayView<const FTransferResourceParams> Params)
-{
 #if WITH_MGPU
+void FD3D12CommandContextRedirector::RHITransferResources(TConstArrayView<FTransferResourceParams> Params)
+{
 	if (Params.Num() == 0)
 		return;
 
@@ -1149,10 +1149,9 @@ void FD3D12CommandContextRedirector::RHITransferResources(const TArrayView<const
 		// The dest waits for the src to be at this place in the frame before using the data.
 		MGPUSync(SrcMask, DstMask);
 	}
-#endif // WITH_MGPU
 }
 
-void FD3D12CommandContextRedirector::RHITransferResourceSignal(const TArrayView<FTransferResourceFenceData* const> FenceDatas, FRHIGPUMask SrcGPUMask)
+void FD3D12CommandContextRedirector::RHITransferResourceSignal(TConstArrayView<FTransferResourceFenceData*> FenceDatas, FRHIGPUMask SrcGPUMask)
 {
 	check(FenceDatas.Num() == SrcGPUMask.GetNumActive());
 
@@ -1170,9 +1169,8 @@ void FD3D12CommandContextRedirector::RHITransferResourceSignal(const TArrayView<
 	}
 }
 
-void FD3D12CommandContextRedirector::RHITransferResourceWait(const TArrayView<FTransferResourceFenceData* const> FenceDatas)
+void FD3D12CommandContextRedirector::RHITransferResourceWait(TConstArrayView<FTransferResourceFenceData*> FenceDatas)
 {
-#if WITH_MGPU
 	FRHIGPUMask AllMasks;
 	for (int32 Index = 0; Index < FenceDatas.Num(); ++Index)
 	{
@@ -1207,13 +1205,10 @@ void FD3D12CommandContextRedirector::RHITransferResourceWait(const TArrayView<FT
 
 		delete FenceData;
 	}
-
-#endif // WITH_MGPU
 }
 
-void FD3D12CommandContextRedirector::RHICrossGPUTransfer(const TArrayView<const FTransferResourceParams> Params, const TArrayView<FCrossGPUTransferFence* const> PreTransfer, const TArrayView<FCrossGPUTransferFence* const> PostTransfer)
+void FD3D12CommandContextRedirector::RHICrossGPUTransfer(TConstArrayView<FTransferResourceParams> Params, TConstArrayView<FCrossGPUTransferFence*> PreTransfer, TConstArrayView<FCrossGPUTransferFence*> PostTransfer)
 {
-#if WITH_MGPU
 	if (Params.Num() == 0)
 		return;
 
@@ -1306,12 +1301,10 @@ void FD3D12CommandContextRedirector::RHICrossGPUTransfer(const TArrayView<const
 		SyncPoint->AddRef();
 		PostTransferSyncPoint->SyncPoint = SyncPoint.GetReference();
 	}
-#endif // WITH_MGPU
 }
 
-void FD3D12CommandContextRedirector::RHICrossGPUTransferSignal(const TArrayView<const FTransferResourceParams> Params, const TArrayView<FCrossGPUTransferFence* const> PreTransfer)
+void FD3D12CommandContextRedirector::RHICrossGPUTransferSignal(TConstArrayView<FTransferResourceParams> Params, TConstArrayView<FCrossGPUTransferFence*> PreTransfer)
 {
-#if WITH_MGPU
 	for (const FTransferResourceParams& Param : Params)
 	{
 		FD3D12CommandContext* DstContext = PhysicalContexts[Param.DestGPUIndex];
@@ -1344,12 +1337,10 @@ void FD3D12CommandContextRedirector::RHICrossGPUTransferSignal(const TArrayView<
 
 		TransferSyncPoint->SyncPoint = SyncPoint;
 	}
-#endif
 }
 
-void FD3D12CommandContextRedirector::RHICrossGPUTransferWait(const TArrayView<FCrossGPUTransferFence* const> PostTransfer)
+void FD3D12CommandContextRedirector::RHICrossGPUTransferWait(TConstArrayView<FCrossGPUTransferFence*> PostTransfer)
 {
-#if WITH_MGPU
 	for (FCrossGPUTransferFence* TransferSyncPoint : PostTransfer)
 	{
 		if (TransferSyncPoint->SyncPoint)
@@ -1362,6 +1353,6 @@ void FD3D12CommandContextRedirector::RHICrossGPUTransferWait(const TArrayView<FC
 
 		delete TransferSyncPoint;
 	}
+}
 
 #endif // WITH_MGPU
-}
diff --git a/Engine/Source/Runtime/D3D12RHI/Private/D3D12CommandContext.h b/Engine/Source/Runtime/D3D12RHI/Private/D3D12CommandContext.h
--- a/Engine/Source/Runtime/D3D12RHI/Private/D3D12CommandContext.h
+++ b/Engine/Source/Runtime/D3D12RHI/Private/D3D12CommandContext.h
@@ -665,7 +665,7 @@ public:
 	// Should be called before RHIBuildAccelerationStructures when multiple GPU support is present (for example, from FD3D12CommandContextRedirector::RHIBuildAccelerationStructures)
 	static void UnregisterAccelerationStructuresInternalMGPU(const TArrayView<const FRayTracingGeometryBuildParams> Params, FRHIGPUMask GPUMask);
 #endif
-	virtual void RHIBuildAccelerationStructures(const TArrayView<const FRayTracingGeometryBuildParams> Params, const FRHIBufferRange& ScratchBufferRange) final override;
+	virtual void RHIBuildAccelerationStructures(TConstArrayView<FRayTracingGeometryBuildParams> Params, const FRHIBufferRange& ScratchBufferRange) final override;
 	virtual void RHIBuildAccelerationStructure(const FRayTracingSceneBuildParams& SceneBuildParams) final override;
 	virtual void RHIClearRayTracingBindings(FRHIRayTracingScene* Scene) final override;
 	virtual void RHIRayTraceDispatch(FRHIRayTracingPipelineState* RayTracingPipelineState, FRHIRayTracingShader* RayGenShader,
@@ -797,14 +797,16 @@ public:
 		ContextRedirect(RHIEndTransitions(Transitions));
 	}
 
-	virtual void RHITransferResources(const TArrayView<const FTransferResourceParams> Params) final override;
-	virtual void RHITransferResourceSignal(const TArrayView<FTransferResourceFenceData* const> FenceDatas, FRHIGPUMask SrcGPUMask) final override;
-	virtual void RHITransferResourceWait(const TArrayView<FTransferResourceFenceData* const> FenceDatas) final override;
+#if WITH_MGPU
+	virtual void RHITransferResources(TConstArrayView<FTransferResourceParams> Params) final override;
+	virtual void RHITransferResourceSignal(TConstArrayView<FTransferResourceFenceData*> FenceDatas, FRHIGPUMask SrcGPUMask) final override;
+	virtual void RHITransferResourceWait(TConstArrayView<FTransferResourceFenceData*> FenceDatas) final override;
 
 	// New and improved cross GPU transfer API
-	virtual void RHICrossGPUTransfer(const TArrayView<const FTransferResourceParams> Params, const TArrayView<FCrossGPUTransferFence* const> PreTransfer, const TArrayView<FCrossGPUTransferFence* const> PostTransfer) final override;
-	virtual void RHICrossGPUTransferSignal(const TArrayView<const FTransferResourceParams> Params, const TArrayView<FCrossGPUTransferFence* const> PreTransfer) final override;
-	virtual void RHICrossGPUTransferWait(const TArrayView<FCrossGPUTransferFence* const> PostTransfer) final override;
+	virtual void RHICrossGPUTransfer(TConstArrayView<FTransferResourceParams> Params, TConstArrayView<FCrossGPUTransferFence*> PreTransfer, TConstArrayView<FCrossGPUTransferFence*> PostTransfer) final override;
+	virtual void RHICrossGPUTransferSignal(TConstArrayView<FTransferResourceParams> Params, TConstArrayView<FCrossGPUTransferFence*> PreTransfer) final override;
+	virtual void RHICrossGPUTransferWait(TConstArrayView<FCrossGPUTransferFence*> PostTransfer) final override;
+#endif // WITH_MGPU
 
 	FORCEINLINE virtual void RHICopyToStagingBuffer(FRHIBuffer* SourceBuffer, FRHIStagingBuffer* DestinationStagingBuffer, uint32 Offset, uint32 NumBytes) final override
 	{
@@ -1026,7 +1028,7 @@ public:
 	}
 
 #if D3D12_RHI_RAYTRACING
-	virtual void RHIBuildAccelerationStructures(const TArrayView<const FRayTracingGeometryBuildParams> Params, const FRHIBufferRange& ScratchBufferRange) final override
+	virtual void RHIBuildAccelerationStructures(TConstArrayView<FRayTracingGeometryBuildParams> Params, const FRHIBufferRange& ScratchBufferRange) final override
 	{
 #if WITH_MGPU
 		FD3D12CommandContext::UnregisterAccelerationStructuresInternalMGPU(Params, GPUMask);
diff --git a/Engine/Source/Runtime/RHI/Public/RHICommandList.h b/Engine/Source/Runtime/RHI/Public/RHICommandList.h
--- a/Engine/Source/Runtime/RHI/Public/RHICommandList.h
+++ b/Engine/Source/Runtime/RHI/Public/RHICommandList.h
@@ -1106,9 +1106,9 @@ FRHICOMMAND_MACRO(FRHICommandSetGPUMask)
 
 FRHICOMMAND_MACRO(FRHICommandTransferResources)
 {
-	TArray<FTransferResourceParams, TInlineAllocator<4>> Params;
+	TConstArrayView<FTransferResourceParams> Params;
 
-	FORCEINLINE_DEBUGGABLE FRHICommandTransferResources(TArrayView<const FTransferResourceParams> InParams)
+	FORCEINLINE_DEBUGGABLE FRHICommandTransferResources(TConstArrayView<FTransferResourceParams> InParams)
 		: Params(InParams)
 	{
 	}
@@ -1118,11 +1118,12 @@ FRHICOMMAND_MACRO(FRHICommandTransferResources)
 
 FRHICOMMAND_MACRO(FRHICommandTransferResourceSignal)
 {
-	TArray<FTransferResourceFenceData*, TInlineAllocator<1>> FenceDatas;
+	TConstArrayView<FTransferResourceFenceData*> FenceDatas;
 	FRHIGPUMask SrcGPUMask;
 
-	FORCEINLINE_DEBUGGABLE FRHICommandTransferResourceSignal(TArrayView<FTransferResourceFenceData* const> InFenceDatas, FRHIGPUMask InSrcGPUMask)
-		: FenceDatas(InFenceDatas), SrcGPUMask(InSrcGPUMask)
+	FORCEINLINE_DEBUGGABLE FRHICommandTransferResourceSignal(TConstArrayView<FTransferResourceFenceData*> InFenceDatas, FRHIGPUMask InSrcGPUMask)
+		: FenceDatas(InFenceDatas)
+		, SrcGPUMask(InSrcGPUMask)
 	{
 	}
 
@@ -1131,9 +1132,9 @@ FRHICOMMAND_MACRO(FRHICommandTransferResourceSignal)
 
 FRHICOMMAND_MACRO(FRHICommandTransferResourceWait)
 {
-	TArray<FTransferResourceFenceData*, TInlineAllocator<4>> FenceDatas;
+	TConstArrayView<FTransferResourceFenceData*> FenceDatas;
 
-	FORCEINLINE_DEBUGGABLE FRHICommandTransferResourceWait(TArrayView<FTransferResourceFenceData* const> InFenceDatas)
+	FORCEINLINE_DEBUGGABLE FRHICommandTransferResourceWait(TConstArrayView<FTransferResourceFenceData*> InFenceDatas)
 		: FenceDatas(InFenceDatas)
 	{
 	}
@@ -1143,12 +1144,14 @@ FRHICOMMAND_MACRO(FRHICommandTransferResourceWait)
 
 FRHICOMMAND_MACRO(FRHICommandCrossGPUTransfer)
 {
-	TArray<FTransferResourceParams, TInlineAllocator<4>> Params;
-	TArray<FCrossGPUTransferFence*, TInlineAllocator<1>> PreTransfer;
-	TArray<FCrossGPUTransferFence*, TInlineAllocator<1>> PostTransfer;
+	TConstArrayView<FTransferResourceParams> Params;
+	TConstArrayView<FCrossGPUTransferFence*> PreTransfer;
+	TConstArrayView<FCrossGPUTransferFence*> PostTransfer;
 
-	FORCEINLINE_DEBUGGABLE FRHICommandCrossGPUTransfer(TArrayView<const FTransferResourceParams> InParams, const TArrayView<FCrossGPUTransferFence* const> InPreTransfer, const TArrayView<FCrossGPUTransferFence* const> InPostTransfer)
-		: Params(InParams), PreTransfer(InPreTransfer), PostTransfer(InPostTransfer)
+	FORCEINLINE_DEBUGGABLE FRHICommandCrossGPUTransfer(TConstArrayView<FTransferResourceParams> InParams, TConstArrayView<FCrossGPUTransferFence*> InPreTransfer, TConstArrayView<FCrossGPUTransferFence*> InPostTransfer)
+		: Params(InParams)
+		, PreTransfer(InPreTransfer)
+		, PostTransfer(InPostTransfer)
 	{
 	}
 
@@ -1157,11 +1160,12 @@ FRHICOMMAND_MACRO(FRHICommandCrossGPUTransfer)
 
 FRHICOMMAND_MACRO(FRHICommandCrossGPUTransferSignal)
 {
-	TArray<FTransferResourceParams, TInlineAllocator<4>> Params;
-	TArray<FCrossGPUTransferFence*, TInlineAllocator<1>> PreTransfer;
+	TConstArrayView<FTransferResourceParams> Params;
+	TConstArrayView<FCrossGPUTransferFence*> PreTransfer;
 
-	FORCEINLINE_DEBUGGABLE FRHICommandCrossGPUTransferSignal(TArrayView<const FTransferResourceParams> InParams, TArrayView<FCrossGPUTransferFence* const> InPreTransfer)
-		: Params(InParams), PreTransfer(InPreTransfer)
+	FORCEINLINE_DEBUGGABLE FRHICommandCrossGPUTransferSignal(TConstArrayView<FTransferResourceParams> InParams, TConstArrayView<FCrossGPUTransferFence*> InPreTransfer)
+		: Params(InParams)
+		, PreTransfer(InPreTransfer)
 	{
 	}
 
@@ -1170,9 +1174,9 @@ FRHICOMMAND_MACRO(FRHICommandCrossGPUTransferSignal)
 
 FRHICOMMAND_MACRO(FRHICommandCrossGPUTransferWait)
 {
-	TArray<FCrossGPUTransferFence*, TInlineAllocator<4>> SyncPoints;
+	TConstArrayView<FCrossGPUTransferFence*> SyncPoints;
 
-	FORCEINLINE_DEBUGGABLE FRHICommandCrossGPUTransferWait(TArrayView<FCrossGPUTransferFence* const> InSyncPoints)
+	FORCEINLINE_DEBUGGABLE FRHICommandCrossGPUTransferWait(TConstArrayView<FCrossGPUTransferFence*> InSyncPoints)
 		: SyncPoints(InSyncPoints)
 	{
 	}
@@ -1979,12 +1983,14 @@ FRHICOMMAND_MACRO(FRHICommandClearRayTracingBindings)
 
 FRHICOMMAND_UNNAMED(FRHICommandBuildAccelerationStructures)
 {
-	const TArrayView<const FRayTracingGeometryBuildParams> Params;
+	TConstArrayView<FRayTracingGeometryBuildParams> Params;
 	FRHIBufferRange ScratchBufferRange;
 	FRHIBuffer* ScratchBuffer;
 
-	explicit FRHICommandBuildAccelerationStructures(const TArrayView<const FRayTracingGeometryBuildParams> InParams, const FRHIBufferRange& ScratchBufferRange)
-		: Params(InParams), ScratchBufferRange(ScratchBufferRange), ScratchBuffer(ScratchBufferRange.Buffer)
+	explicit FRHICommandBuildAccelerationStructures(TConstArrayView<FRayTracingGeometryBuildParams> InParams, const FRHIBufferRange& ScratchBufferRange)
+		: Params(InParams)
+		, ScratchBufferRange(ScratchBufferRange)
+		, ScratchBuffer(ScratchBufferRange.Buffer)
 	{}
 
 	RHI_API void Execute(FRHICommandListBase& CmdList);
@@ -2656,7 +2662,7 @@ PRAGMA_ENABLE_DEPRECATION_WARNINGS
 		}
 	}
 
-	FORCEINLINE_DEBUGGABLE void TransferResources(const TArrayView<const FTransferResourceParams> Params)
+	FORCEINLINE_DEBUGGABLE void TransferResources(TConstArrayView<FTransferResourceParams> Params)
 	{
 #if WITH_MGPU
 		if (Bypass())
@@ -2665,12 +2671,12 @@ PRAGMA_ENABLE_DEPRECATION_WARNINGS
 		}
 		else
 		{
-			ALLOC_COMMAND(FRHICommandTransferResources)(Params);
+			ALLOC_COMMAND(FRHICommandTransferResources)(AllocArray(Params));
 		}
 #endif // WITH_MGPU
 	}
 
-	FORCEINLINE_DEBUGGABLE void TransferResourceSignal(const TArrayView<FTransferResourceFenceData* const> FenceDatas, FRHIGPUMask SrcGPUMask)
+	FORCEINLINE_DEBUGGABLE void TransferResourceSignal(TConstArrayView<FTransferResourceFenceData*> FenceDatas, FRHIGPUMask SrcGPUMask)
 	{
 #if WITH_MGPU
 		if (Bypass())
@@ -2679,12 +2685,12 @@ PRAGMA_ENABLE_DEPRECATION_WARNINGS
 		}
 		else
 		{
-			ALLOC_COMMAND(FRHICommandTransferResourceSignal)(FenceDatas, SrcGPUMask);
+			ALLOC_COMMAND(FRHICommandTransferResourceSignal)(AllocArray(FenceDatas), SrcGPUMask);
 		}
 #endif // WITH_MGPU
 	}
 
-	FORCEINLINE_DEBUGGABLE void TransferResourceWait(const TArrayView<FTransferResourceFenceData* const> FenceDatas)
+	FORCEINLINE_DEBUGGABLE void TransferResourceWait(TConstArrayView<FTransferResourceFenceData*> FenceDatas)
 	{
 #if WITH_MGPU
 		if (Bypass())
@@ -2693,12 +2699,12 @@ PRAGMA_ENABLE_DEPRECATION_WARNINGS
 		}
 		else
 		{
-			ALLOC_COMMAND(FRHICommandTransferResourceWait)(FenceDatas);
+			ALLOC_COMMAND(FRHICommandTransferResourceWait)(AllocArray(FenceDatas));
 		}
 #endif // WITH_MGPU
 	}
 
-	FORCEINLINE_DEBUGGABLE void CrossGPUTransfer(const TArrayView<const FTransferResourceParams> Params, const TArrayView<FCrossGPUTransferFence* const> PreTransfer, const TArrayView<FCrossGPUTransferFence* const> PostTransfer)
+	FORCEINLINE_DEBUGGABLE void CrossGPUTransfer(TConstArrayView<FTransferResourceParams> Params, TConstArrayView<FCrossGPUTransferFence*> PreTransfer, TConstArrayView<FCrossGPUTransferFence*> PostTransfer)
 	{
 #if WITH_MGPU
 		if (Bypass())
@@ -2707,12 +2713,12 @@ PRAGMA_ENABLE_DEPRECATION_WARNINGS
 		}
 		else
 		{
-			ALLOC_COMMAND(FRHICommandCrossGPUTransfer)(Params, PreTransfer, PostTransfer);
+			ALLOC_COMMAND(FRHICommandCrossGPUTransfer)(AllocArray(Params), AllocArray(PreTransfer), AllocArray(PostTransfer));
 		}
 #endif // WITH_MGPU
 	}
 
-	FORCEINLINE_DEBUGGABLE void CrossGPUTransferSignal(const TArrayView<const FTransferResourceParams> Params, const TArrayView<FCrossGPUTransferFence* const> PreTransfer)
+	FORCEINLINE_DEBUGGABLE void CrossGPUTransferSignal(TConstArrayView<FTransferResourceParams> Params, TConstArrayView<FCrossGPUTransferFence*> PreTransfer)
 	{
 #if WITH_MGPU
 		if (Bypass())
@@ -2721,12 +2727,12 @@ PRAGMA_ENABLE_DEPRECATION_WARNINGS
 		}
 		else
 		{
-			ALLOC_COMMAND(FRHICommandCrossGPUTransferSignal)(Params, PreTransfer);
+			ALLOC_COMMAND(FRHICommandCrossGPUTransferSignal)(AllocArray(Params), AllocArray(PreTransfer));
 		}
 #endif // WITH_MGPU
 	}
 
-	FORCEINLINE_DEBUGGABLE void CrossGPUTransferWait(const TArrayView<FCrossGPUTransferFence* const> SyncPoints)
+	FORCEINLINE_DEBUGGABLE void CrossGPUTransferWait(TConstArrayView<FCrossGPUTransferFence*> SyncPoints)
 	{
 #if WITH_MGPU
 		if (Bypass())
@@ -2735,7 +2741,7 @@ PRAGMA_ENABLE_DEPRECATION_WARNINGS
 		}
 		else
 		{
-			ALLOC_COMMAND(FRHICommandCrossGPUTransferWait)(SyncPoints);
+			ALLOC_COMMAND(FRHICommandCrossGPUTransferWait)(AllocArray(SyncPoints));
 		}
 #endif // WITH_MGPU
 	}
@@ -2753,9 +2759,9 @@ PRAGMA_ENABLE_DEPRECATION_WARNINGS
 
 #if RHI_RAYTRACING
 	RHI_API void BuildAccelerationStructure(FRHIRayTracingGeometry* Geometry);
-	RHI_API void BuildAccelerationStructures(const TArrayView<const FRayTracingGeometryBuildParams> Params);
+	RHI_API void BuildAccelerationStructures(TConstArrayView<FRayTracingGeometryBuildParams> Params);
 
-	FORCEINLINE_DEBUGGABLE void BuildAccelerationStructures(const TArrayView<const FRayTracingGeometryBuildParams> Params, const FRHIBufferRange& ScratchBufferRange)
+	FORCEINLINE_DEBUGGABLE void BuildAccelerationStructures(TConstArrayView<FRayTracingGeometryBuildParams> Params, const FRHIBufferRange& ScratchBufferRange)
 	{
 		if (Bypass())
 		{
diff --git a/Engine/Source/Runtime/RHI/Public/RHIContext.h b/Engine/Source/Runtime/RHI/Public/RHIContext.h
--- a/Engine/Source/Runtime/RHI/Public/RHIContext.h
+++ b/Engine/Source/Runtime/RHI/Public/RHIContext.h
@@ -287,8 +287,8 @@ public:
 	virtual void RHIBeginUAVOverlap() {}
 	virtual void RHIEndUAVOverlap() {}
 
-	virtual void RHIBeginUAVOverlap(TArrayView<FRHIUnorderedAccessView* const> UAVs) {}
-	virtual void RHIEndUAVOverlap(TArrayView<FRHIUnorderedAccessView* const> UAVs) {}
+	virtual void RHIBeginUAVOverlap(TConstArrayView<FRHIUnorderedAccessView*> UAVs) {}
+	virtual void RHIEndUAVOverlap(TConstArrayView<FRHIUnorderedAccessView*> UAVs) {}
 
 	/** Set the shader resource view of a surface.  This is used for binding TextureMS parameter types that need a multi sampled view. */
 	virtual void RHISetShaderTexture(FRHIComputeShader* ComputeShader, uint32 TextureIndex, FRHITexture* NewTexture) = 0;
@@ -391,11 +391,12 @@ public:
 		return FRHIGPUMask::GPU0();
 	}
 
+#if WITH_MGPU
 	/**
 	 * Synchronizes the content of a resource between two GPUs using a copy operation.
 	 * @param Params - the parameters for each resource or texture region copied between GPUs.
 	 */
-	virtual void RHITransferResources(const TArrayView<const FTransferResourceParams> Params)
+	virtual void RHITransferResources(TConstArrayView<FTransferResourceParams> Params)
 	{
 		/* empty default implementation */
 	}
@@ -406,26 +407,22 @@ public:
 	 * GPU mask of the context (which specifies the destination GPUs), and the number of items in the "FenceDatas" array MUST match the
 	 * number of bits set in SrcGPUMask.
 	 */
-	virtual void RHITransferResourceSignal(const TArrayView<FTransferResourceFenceData* const> FenceDatas, FRHIGPUMask SrcGPUMask)
+	virtual void RHITransferResourceSignal(TConstArrayView<FTransferResourceFenceData*> FenceDatas, FRHIGPUMask SrcGPUMask)
 	{
 		/* default noop implementation */
-#if WITH_MGPU
 		for (FTransferResourceFenceData* FenceData : FenceDatas)
 		{
 			delete FenceData;
 		}
-#endif
 	}
 
-	virtual void RHITransferResourceWait(const TArrayView<FTransferResourceFenceData* const> FenceDatas)
+	virtual void RHITransferResourceWait(TConstArrayView<FTransferResourceFenceData*> FenceDatas)
 	{
 		/* default noop implementation */
-#if WITH_MGPU
 		for (FTransferResourceFenceData* FenceData : FenceDatas)
 		{
 			delete FenceData;
 		}
-#endif
 	}
 
 	/**
@@ -434,34 +431,31 @@ public:
 	 * @param PreTransfer - Fences to wait on before copying the relevant data (initialized with RHITransferResourceSignal before this function)
 	 * @param PostTransfer - Fences that can be waited on after copy (waited on by RHITransferResourceWait after this function)
 	 */
-	virtual void RHICrossGPUTransfer(const TArrayView<const FTransferResourceParams> Params, const TArrayView<FCrossGPUTransferFence* const> PreTransfer, const TArrayView<FCrossGPUTransferFence* const> PostTransfer)
+	virtual void RHICrossGPUTransfer(TConstArrayView<FTransferResourceParams> Params, TConstArrayView<FCrossGPUTransferFence*> PreTransfer, TConstArrayView<FCrossGPUTransferFence*> PostTransfer)
 	{
 		/** empty default implementation. */
 	}
 
-	virtual void RHICrossGPUTransferSignal(const TArrayView<const FTransferResourceParams> Params, const TArrayView<FCrossGPUTransferFence* const> PreTransfer)
+	virtual void RHICrossGPUTransferSignal(TConstArrayView<FTransferResourceParams> Params, TConstArrayView<FCrossGPUTransferFence*> PreTransfer)
 	{
 		/* default noop implementation */
-#if WITH_MGPU
 		for (FCrossGPUTransferFence* SyncPoint : PreTransfer)
 		{
 			delete SyncPoint;
 		}
-#endif
 	}
 
-	virtual void RHICrossGPUTransferWait(const TArrayView<FCrossGPUTransferFence* const> SyncPoints)
+	virtual void RHICrossGPUTransferWait(TConstArrayView<FCrossGPUTransferFence*> SyncPoints)
 	{
 		/* default noop implementation */
-#if WITH_MGPU
 		for (FCrossGPUTransferFence* SyncPoint : SyncPoints)
 		{
 			delete SyncPoint;
 		}
-#endif
 	}
+#endif // WITH_MGPU
 
-	virtual void RHIBuildAccelerationStructures(const TArrayView<const FRayTracingGeometryBuildParams> Params, const FRHIBufferRange& ScratchBufferRange)
+	virtual void RHIBuildAccelerationStructures(TConstArrayView<FRayTracingGeometryBuildParams> Params, const FRHIBufferRange& ScratchBufferRange)
 	{
 		checkNoEntry();
 	}
@@ -533,7 +527,7 @@ public:
 };
 
 // Utility function to generate pre-transfer sync points to pass to CrossGPUTransferSignal and CrossGPUTransfer
-RHI_API void RHIGenerateCrossGPUPreTransferFences(const TArrayView<const FTransferResourceParams> Params, TArray<FCrossGPUTransferFence*>& OutPreTransfer);
+RHI_API void RHIGenerateCrossGPUPreTransferFences(TConstArrayView<FTransferResourceParams> Params, TArray<FCrossGPUTransferFence*>& OutPreTransfer);
 
 enum class EAccelerationStructureBuildMode
 {
@@ -553,7 +547,7 @@ struct FRayTracingGeometryBuildParams
 
 	// Optional array of geometry segments that can be used to change per-segment vertex buffers.
 	// Only fields related to vertex buffer are used. If empty, then geometry vertex buffers are not changed.
-	TArrayView<const FRayTracingGeometrySegment> Segments;
+	TConstArrayView<FRayTracingGeometrySegment> Segments;
 };
 
 struct FRayTracingSceneBuildParams
@@ -800,12 +794,12 @@ public:
 		checkNoEntry();
 	}
 
-	virtual void RHIBuildAccelerationStructures(const TArrayView<const FRayTracingGeometryBuildParams> Params, const FRHIBufferRange& ScratchBufferRange)
+	virtual void RHIBuildAccelerationStructures(TConstArrayView<FRayTracingGeometryBuildParams> Params, const FRHIBufferRange& ScratchBufferRange)
 	{
 		checkNoEntry();
 	}
 
-	void RHIBuildAccelerationStructures(const TArrayView<const FRayTracingGeometryBuildParams> Params)
+	void RHIBuildAccelerationStructures(TConstArrayView<FRayTracingGeometryBuildParams> Params)
 	{
 		checkNoEntry();
 	}
diff --git a/Engine/Source/Runtime/RHI/Public/RHIValidationContext.h b/Engine/Source/Runtime/RHI/Public/RHIValidationContext.h
--- a/Engine/Source/Runtime/RHI/Public/RHIValidationContext.h
+++ b/Engine/Source/Runtime/RHI/Public/RHIValidationContext.h
@@ -177,7 +177,7 @@ public:
 		RHIContext->RHIEndUAVOverlap();
 	}
 
-	virtual void RHIBeginUAVOverlap(TArrayView<FRHIUnorderedAccessView* const> UAVs) final override
+	virtual void RHIBeginUAVOverlap(TConstArrayView<FRHIUnorderedAccessView*> UAVs) final override
 	{
 		for (FRHIUnorderedAccessView* UAV : UAVs)
 		{
@@ -186,7 +186,7 @@ public:
 		RHIContext->RHIBeginUAVOverlap(UAVs);
 	}
 
-	virtual void RHIEndUAVOverlap(TArrayView<FRHIUnorderedAccessView* const> UAVs) final override
+	virtual void RHIEndUAVOverlap(TConstArrayView<FRHIUnorderedAccessView*> UAVs) final override
 	{
 		for (FRHIUnorderedAccessView* UAV : UAVs)
 		{
@@ -312,7 +312,39 @@ public:
 
 	virtual void RHICopyToStagingBuffer(FRHIBuffer* SourceBufferRHI, FRHIStagingBuffer* DestinationStagingBufferRHI, uint32 InOffset, uint32 InNumBytes) override final;
 
-	virtual void RHIBuildAccelerationStructures(const TArrayView<const FRayTracingGeometryBuildParams> Params, const FRHIBufferRange& ScratchBufferRange) override final
+#if WITH_MGPU
+	virtual void RHITransferResources(TConstArrayView<FTransferResourceParams> Params)
+	{
+		RHIContext->RHITransferResources(Params);
+	}
+
+	virtual void RHITransferResourceSignal(TConstArrayView<FTransferResourceFenceData*> FenceDatas, FRHIGPUMask SrcGPUMask)
+	{
+		RHIContext->RHITransferResourceSignal(FenceDatas, SrcGPUMask);
+	}
+
+	virtual void RHITransferResourceWait(TConstArrayView<FTransferResourceFenceData*> FenceDatas)
+	{
+		RHIContext->RHITransferResourceWait(FenceDatas);
+	}
+
+	virtual void RHICrossGPUTransfer(TConstArrayView<FTransferResourceParams> Params, TConstArrayView<FCrossGPUTransferFence*> PreTransfer, TConstArrayView<FCrossGPUTransferFence*> PostTransfer)
+	{
+		RHIContext->RHICrossGPUTransfer(Params, PreTransfer, PostTransfer);
+	}
+
+	virtual void RHICrossGPUTransferSignal(TConstArrayView<FTransferResourceParams> Params, TConstArrayView<FCrossGPUTransferFence*> PreTransfer)
+	{
+		RHIContext->RHICrossGPUTransferSignal(Params, PreTransfer);
+	}
+
+	virtual void RHICrossGPUTransferWait(TConstArrayView<FCrossGPUTransferFence*> SyncPoints)
+	{
+		RHIContext->RHICrossGPUTransferWait(SyncPoints);
+	}
+#endif // WITH_MGPU
+
+	virtual void RHIBuildAccelerationStructures(TConstArrayView<FRayTracingGeometryBuildParams> Params, const FRHIBufferRange& ScratchBufferRange) override final
 	{
 		// #yuriy_todo: explicit transitions and state validation for BLAS
 		RHIContext->RHIBuildAccelerationStructures(Params, ScratchBufferRange);
@@ -443,7 +475,7 @@ public:
 		RHIContext->RHIEndUAVOverlap();
 	}
 
-	virtual void RHIBeginUAVOverlap(TArrayView<FRHIUnorderedAccessView* const> UAVs) final override
+	virtual void RHIBeginUAVOverlap(TConstArrayView<FRHIUnorderedAccessView*> UAVs) final override
 	{
 		for (FRHIUnorderedAccessView* UAV : UAVs)
 		{
@@ -452,7 +484,7 @@ public:
 		RHIContext->RHIBeginUAVOverlap(UAVs);
 	}
 
-	virtual void RHIEndUAVOverlap(TArrayView<FRHIUnorderedAccessView* const> UAVs) final override
+	virtual void RHIEndUAVOverlap(TConstArrayView<FRHIUnorderedAccessView*> UAVs) final override
 	{
 		for (FRHIUnorderedAccessView* UAV : UAVs)
 		{
@@ -1104,7 +1136,7 @@ public:
 		RHIContext->RHIClearRayTracingBindings(Scene);
 	}
 
-	virtual void RHIBuildAccelerationStructures(const TArrayView<const FRayTracingGeometryBuildParams> Params, const FRHIBufferRange& ScratchBufferRange) override final
+	virtual void RHIBuildAccelerationStructures(TConstArrayView<FRayTracingGeometryBuildParams> Params, const FRHIBufferRange& ScratchBufferRange) override final
 	{
 		// #yuriy_todo: explicit transitions and state validation for BLAS
 		RHIContext->RHIBuildAccelerationStructures(Params, ScratchBufferRange);
