Moved RHI view and ray tracing geometry creation to support non-immediate command lists.

diff --git a/Engine/Source/Runtime/Apple/MetalRHI/Private/MetalDynamicRHI.h b/Engine/Source/Runtime/Apple/MetalRHI/Private/MetalDynamicRHI.h
--- a/Engine/Source/Runtime/Apple/MetalRHI/Private/MetalDynamicRHI.h
+++ b/Engine/Source/Runtime/Apple/MetalRHI/Private/MetalDynamicRHI.h
@@ -64,8 +64,8 @@ public:
 	virtual void UnlockBuffer_BottomOfPipe(FRHICommandListBase& RHICmdList, FRHIBuffer* Buffer) final override;
 
 	// SRV / UAV creation functions
-	virtual FShaderResourceViewRHIRef  RHICreateShaderResourceView (class FRHICommandListImmediate& RHICmdList, FRHIViewableResource* Resource, FRHIViewDesc const& ViewDesc) final override;
-	virtual FUnorderedAccessViewRHIRef RHICreateUnorderedAccessView(class FRHICommandListImmediate& RHICmdList, FRHIViewableResource* Resource, FRHIViewDesc const& ViewDesc) final override;
+	virtual FShaderResourceViewRHIRef  RHICreateShaderResourceView (class FRHICommandListBase& RHICmdList, FRHIViewableResource* Resource, FRHIViewDesc const& ViewDesc) final override;
+	virtual FUnorderedAccessViewRHIRef RHICreateUnorderedAccessView(class FRHICommandListBase& RHICmdList, FRHIViewableResource* Resource, FRHIViewDesc const& ViewDesc) final override;
 
 	virtual FRHICalcTextureSizeResult RHICalcTexturePlatformSize(FRHITextureDesc const& Desc, uint32 FirstMipIndex) final override;
 	virtual uint64 RHIGetMinimumAlignmentForBufferBackedSRV(EPixelFormat Format) final override;
diff --git a/Engine/Source/Runtime/Apple/MetalRHI/Private/MetalUAV.cpp b/Engine/Source/Runtime/Apple/MetalRHI/Private/MetalUAV.cpp
--- a/Engine/Source/Runtime/Apple/MetalRHI/Private/MetalUAV.cpp
+++ b/Engine/Source/Runtime/Apple/MetalRHI/Private/MetalUAV.cpp
@@ -69,10 +69,10 @@ void FMetalResourceViewBase::InitAsTextureBufferBacked(FMetalTexture& Texture, F
     Storage.Emplace<FTextureBufferBacked>(Texture, Buffer, Offset, Size, Format);
 }
 
-FMetalShaderResourceView::FMetalShaderResourceView(FRHICommandListImmediate& RHICmdList, FRHIViewableResource* InResource, FRHIViewDesc const& InViewDesc)
+FMetalShaderResourceView::FMetalShaderResourceView(FRHICommandListBase& RHICmdList, FRHIViewableResource* InResource, FRHIViewDesc const& InViewDesc)
 	: FRHIShaderResourceView(InResource, InViewDesc)
 {
-	RHICmdList.EnqueueLambda([this](FRHICommandListImmediate&)
+	RHICmdList.EnqueueLambda([this](FRHICommandListBase&)
 	{
 		LinkHead(GetBaseResource()->LinkedViews);
 		UpdateView();
@@ -262,10 +262,10 @@ void FMetalShaderResourceView::UpdateView()
 
 
 
-FMetalUnorderedAccessView::FMetalUnorderedAccessView(FRHICommandListImmediate& RHICmdList, FRHIViewableResource* InResource, FRHIViewDesc const& InViewDesc)
+FMetalUnorderedAccessView::FMetalUnorderedAccessView(FRHICommandListBase& RHICmdList, FRHIViewableResource* InResource, FRHIViewDesc const& InViewDesc)
 	: FRHIUnorderedAccessView(InResource, InViewDesc)
 {
-	RHICmdList.EnqueueLambda([this](FRHICommandListImmediate&)
+	RHICmdList.EnqueueLambda([this](FRHICommandListBase&)
 	{
 		LinkHead(GetBaseResource()->LinkedViews);
 		UpdateView();
@@ -438,14 +438,14 @@ void FMetalUnorderedAccessView::UpdateView()
 	}
 }
 
-FShaderResourceViewRHIRef FMetalDynamicRHI::RHICreateShaderResourceView(FRHICommandListImmediate& RHICmdList, FRHIViewableResource* Resource, FRHIViewDesc const& ViewDesc)
+FShaderResourceViewRHIRef FMetalDynamicRHI::RHICreateShaderResourceView(FRHICommandListBase& RHICmdList, FRHIViewableResource* Resource, FRHIViewDesc const& ViewDesc)
 {
 	@autoreleasepool {
 		return new FMetalShaderResourceView(RHICmdList, Resource, ViewDesc);
 	}
 }
 
-FUnorderedAccessViewRHIRef FMetalDynamicRHI::RHICreateUnorderedAccessView(FRHICommandListImmediate& RHICmdList, FRHIViewableResource* Resource, FRHIViewDesc const& ViewDesc)
+FUnorderedAccessViewRHIRef FMetalDynamicRHI::RHICreateUnorderedAccessView(FRHICommandListBase& RHICmdList, FRHIViewableResource* Resource, FRHIViewDesc const& ViewDesc)
 {
 	@autoreleasepool {
 		return new FMetalUnorderedAccessView(RHICmdList, Resource, ViewDesc);
diff --git a/Engine/Source/Runtime/Apple/MetalRHI/Public/MetalResources.h b/Engine/Source/Runtime/Apple/MetalRHI/Public/MetalResources.h
--- a/Engine/Source/Runtime/Apple/MetalRHI/Public/MetalResources.h
+++ b/Engine/Source/Runtime/Apple/MetalRHI/Public/MetalResources.h
@@ -495,7 +495,7 @@ private:
 class FMetalShaderResourceView final : public FRHIShaderResourceView, public FMetalResourceViewBase
 {
 public:
-	FMetalShaderResourceView(FRHICommandListImmediate& RHICmdList, FRHIViewableResource* InResource, FRHIViewDesc const& InViewDesc);
+	FMetalShaderResourceView(FRHICommandListBase& RHICmdList, FRHIViewableResource* InResource, FRHIViewDesc const& InViewDesc);
 	FMetalViewableResource* GetBaseResource() const;
 
 	virtual void UpdateView() override;
@@ -504,7 +504,7 @@ public:
 class FMetalUnorderedAccessView final : public FRHIUnorderedAccessView, public FMetalResourceViewBase
 {
 public:
-	FMetalUnorderedAccessView(FRHICommandListImmediate& RHICmdList, FRHIViewableResource* InResource, FRHIViewDesc const& InViewDesc);
+	FMetalUnorderedAccessView(FRHICommandListBase& RHICmdList, FRHIViewableResource* InResource, FRHIViewDesc const& InViewDesc);
 	FMetalViewableResource* GetBaseResource() const;
 
 	virtual void UpdateView() override;
diff --git a/Engine/Source/Runtime/D3D12RHI/Private/D3D12RHIPrivate.h b/Engine/Source/Runtime/D3D12RHI/Private/D3D12RHIPrivate.h
--- a/Engine/Source/Runtime/D3D12RHI/Private/D3D12RHIPrivate.h
+++ b/Engine/Source/Runtime/D3D12RHI/Private/D3D12RHIPrivate.h
@@ -383,8 +383,8 @@ public:
 	virtual void RHIWriteGPUFence_TopOfPipe(FRHICommandListBase& RHICmdList, FRHIGPUFence* FenceRHI) final override;
 
 	// SRV / UAV creation functions
-	virtual FShaderResourceViewRHIRef  RHICreateShaderResourceView (class FRHICommandListImmediate& RHICmdList, FRHIViewableResource* Resource, FRHIViewDesc const& ViewDesc) override;
-	virtual FUnorderedAccessViewRHIRef RHICreateUnorderedAccessView(class FRHICommandListImmediate& RHICmdList, FRHIViewableResource* Resource, FRHIViewDesc const& ViewDesc) override;
+	virtual FShaderResourceViewRHIRef  RHICreateShaderResourceView (class FRHICommandListBase& RHICmdList, FRHIViewableResource* Resource, FRHIViewDesc const& ViewDesc) override;
+	virtual FUnorderedAccessViewRHIRef RHICreateUnorderedAccessView(class FRHICommandListBase& RHICmdList, FRHIViewableResource* Resource, FRHIViewDesc const& ViewDesc) override;
 
 	// ID3D12DynamicRHI interface.
 	virtual TArray<FD3D12MinimalAdapterDesc> RHIGetAdapterDescs() const final override;
@@ -442,7 +442,7 @@ public:
 	virtual FRayTracingAccelerationStructureSize RHICalcRayTracingSceneSize(uint32 MaxInstances, ERayTracingAccelerationStructureFlags Flags) final override;
 	virtual FRayTracingAccelerationStructureSize RHICalcRayTracingGeometrySize(const FRayTracingGeometryInitializer& Initializer) final override;
 
-	virtual FRayTracingGeometryRHIRef RHICreateRayTracingGeometry(const FRayTracingGeometryInitializer& Initializer) final override;
+	virtual FRayTracingGeometryRHIRef RHICreateRayTracingGeometry(FRHICommandListBase& RHICmdList, const FRayTracingGeometryInitializer& Initializer) final override;
 	virtual FRayTracingSceneRHIRef RHICreateRayTracingScene(FRayTracingSceneInitializer2 Initializer) final override;
 	virtual FRayTracingShaderRHIRef RHICreateRayTracingShader(TArrayView<const uint8> Code, const FSHAHash& Hash, EShaderFrequency ShaderFrequency) final override;
 	virtual FRayTracingPipelineStateRHIRef RHICreateRayTracingPipelineState(const FRayTracingPipelineStateInitializer& Initializer) final override;
diff --git a/Engine/Source/Runtime/D3D12RHI/Private/D3D12RayTracing.cpp b/Engine/Source/Runtime/D3D12RHI/Private/D3D12RayTracing.cpp
--- a/Engine/Source/Runtime/D3D12RHI/Private/D3D12RayTracing.cpp
+++ b/Engine/Source/Runtime/D3D12RHI/Private/D3D12RayTracing.cpp
@@ -3200,10 +3200,10 @@ FRayTracingPipelineStateRHIRef FD3D12DynamicRHI::RHICreateRayTracingPipelineStat
 	return Result;
 }
 
-FRayTracingGeometryRHIRef FD3D12DynamicRHI::RHICreateRayTracingGeometry(const FRayTracingGeometryInitializer& Initializer)
+FRayTracingGeometryRHIRef FD3D12DynamicRHI::RHICreateRayTracingGeometry(FRHICommandListBase& RHICmdList, const FRayTracingGeometryInitializer& Initializer)
 {
 	FD3D12Adapter& Adapter = GetAdapter();
-	return new FD3D12RayTracingGeometry(&Adapter, Initializer);
+	return new FD3D12RayTracingGeometry(RHICmdList, &Adapter, Initializer);
 }
 
 void FD3D12DynamicRHI::RHITransferRayTracingGeometryUnderlyingResource(FRHIRayTracingGeometry* DestGeometry, FRHIRayTracingGeometry* SrcGeometry)
@@ -3309,7 +3309,7 @@ FString GetGeometryInitializerDebugString(const FRayTracingGeometryInitializer&
 	return Result.ToString();
 }
 
-FD3D12RayTracingGeometry::FD3D12RayTracingGeometry(FD3D12Adapter* Adapter, const FRayTracingGeometryInitializer& InInitializer)
+FD3D12RayTracingGeometry::FD3D12RayTracingGeometry(FRHICommandListBase& RHICmdList, FD3D12Adapter* Adapter, const FRayTracingGeometryInitializer& InInitializer)
 	: FRHIRayTracingGeometry(InInitializer), FD3D12AdapterChild(Adapter)
 {
 	INC_DWORD_STAT(STAT_D3D12RayTracingAllocatedBLAS);
@@ -3328,7 +3328,7 @@ FD3D12RayTracingGeometry::FD3D12RayTracingGeometry(FD3D12Adapter* Adapter, const
 		FRHIResourceCreateInfo CreateInfo(TEXT("NullTransformBuffer"));
 		CreateInfo.ResourceArray = &NullTransformData;
 
-		FD3D12RayTracingGeometry::NullTransformBuffer = RHICreateVertexBuffer(NullTransformData.GetResourceDataSize(), BUF_Static, CreateInfo);
+		FD3D12RayTracingGeometry::NullTransformBuffer = RHICmdList.CreateBuffer(NullTransformData.GetResourceDataSize(), BUF_VertexBuffer | BUF_Static, 0, ERHIAccess::VertexOrIndexBuffer, CreateInfo);
 	}
 
 	RegisterD3D12RayTracingGeometry(this);
@@ -3429,8 +3429,6 @@ FD3D12RayTracingGeometry::FD3D12RayTracingGeometry(FD3D12Adapter* Adapter, const
 
 		FMemory::Memcpy(DstDataBase, Data, Size);
 
-		FRHICommandListImmediate& RHICmdList = FRHICommandListExecutor::GetImmediateCommandList();
-
 		if (bOnAsyncThread)
 		{
 			check(Initializer.Type == ERayTracingGeometryInitializerType::StreamingSource);
diff --git a/Engine/Source/Runtime/D3D12RHI/Private/D3D12RayTracing.h b/Engine/Source/Runtime/D3D12RHI/Private/D3D12RayTracing.h
--- a/Engine/Source/Runtime/D3D12RHI/Private/D3D12RayTracing.h
+++ b/Engine/Source/Runtime/D3D12RHI/Private/D3D12RayTracing.h
@@ -23,7 +23,7 @@ class FD3D12RayTracingGeometry : public FRHIRayTracingGeometry, public FD3D12Ada
 {
 public:
 
-	FD3D12RayTracingGeometry(FD3D12Adapter* Adapter, const FRayTracingGeometryInitializer& Initializer);
+	FD3D12RayTracingGeometry(FRHICommandListBase& RHICmdList, FD3D12Adapter* Adapter, const FRayTracingGeometryInitializer& Initializer);
 	~FD3D12RayTracingGeometry();
 
 	virtual FRayTracingAccelerationStructureAddress GetAccelerationStructureAddress(uint64 GPUIndex) const final override
diff --git a/Engine/Source/Runtime/D3D12RHI/Private/D3D12SRV.cpp b/Engine/Source/Runtime/D3D12RHI/Private/D3D12SRV.cpp
--- a/Engine/Source/Runtime/D3D12RHI/Private/D3D12SRV.cpp
+++ b/Engine/Source/Runtime/D3D12RHI/Private/D3D12SRV.cpp
@@ -256,7 +256,7 @@ FD3D12ShaderResourceView_RHI::FD3D12ShaderResourceView_RHI(FD3D12Device* InDevic
 //
 // -----------------------------------------------------------------------------------------------------
 
-FShaderResourceViewRHIRef FD3D12DynamicRHI::RHICreateShaderResourceView(class FRHICommandListImmediate& RHICmdList, FRHIViewableResource* Resource, FRHIViewDesc const& ViewDesc)
+FShaderResourceViewRHIRef FD3D12DynamicRHI::RHICreateShaderResourceView(class FRHICommandListBase& RHICmdList, FRHIViewableResource* Resource, FRHIViewDesc const& ViewDesc)
 {
 	FRHIGPUMask RelevantGPUs = ViewDesc.IsBuffer()
 		? FD3D12DynamicRHI::ResourceCast(static_cast<FRHIBuffer* >(Resource))->GetLinkedObjectsGPUMask()
diff --git a/Engine/Source/Runtime/D3D12RHI/Private/D3D12UAV.cpp b/Engine/Source/Runtime/D3D12RHI/Private/D3D12UAV.cpp
--- a/Engine/Source/Runtime/D3D12RHI/Private/D3D12UAV.cpp
+++ b/Engine/Source/Runtime/D3D12RHI/Private/D3D12UAV.cpp
@@ -206,7 +206,7 @@ FD3D12UnorderedAccessView_RHI::FD3D12UnorderedAccessView_RHI(FD3D12Device* InDev
 //
 // -----------------------------------------------------------------------------------------------------
 
-FUnorderedAccessViewRHIRef FD3D12DynamicRHI::RHICreateUnorderedAccessView(class FRHICommandListImmediate& RHICmdList, FRHIViewableResource* Resource, FRHIViewDesc const& ViewDesc)
+FUnorderedAccessViewRHIRef FD3D12DynamicRHI::RHICreateUnorderedAccessView(class FRHICommandListBase& RHICmdList, FRHIViewableResource* Resource, FRHIViewDesc const& ViewDesc)
 {
 	FRHIGPUMask RelevantGPUs = ViewDesc.IsBuffer()
 		? FD3D12DynamicRHI::ResourceCast(static_cast<FRHIBuffer* >(Resource))->GetLinkedObjectsGPUMask()
diff --git a/Engine/Source/Runtime/D3D12RHI/Private/D3D12View.h b/Engine/Source/Runtime/D3D12RHI/Private/D3D12View.h
--- a/Engine/Source/Runtime/D3D12RHI/Private/D3D12View.h
+++ b/Engine/Source/Runtime/D3D12RHI/Private/D3D12View.h
@@ -375,9 +375,9 @@ private:
 template <typename TParent>
 struct FD3D12DeferredInitView : public FD3D12LinkedAdapterObject<TParent>
 {
-	void CreateViews(FRHICommandListImmediate& RHICmdList, bool bDynamic)
+	void CreateViews(FRHICommandListBase& RHICmdList, bool bDynamic)
 	{
-		auto InitLambda = [this](FRHICommandListImmediate&)
+		auto InitLambda = [this](FRHICommandListBase&)
 		{
 			for (TParent& LinkedView : *this)
 			{
diff --git a/Engine/Source/Runtime/NullDrv/Public/NullRHI.h b/Engine/Source/Runtime/NullDrv/Public/NullRHI.h
--- a/Engine/Source/Runtime/NullDrv/Public/NullRHI.h
+++ b/Engine/Source/Runtime/NullDrv/Public/NullRHI.h
@@ -194,12 +194,12 @@ public:
 		return this->RHICreateTexture(Desc);
 	}
 
-	virtual FShaderResourceViewRHIRef RHICreateShaderResourceView(class FRHICommandListImmediate& RHICmdList, FRHIViewableResource* Resource, FRHIViewDesc const& ViewDesc)
+	virtual FShaderResourceViewRHIRef RHICreateShaderResourceView(class FRHICommandListBase& RHICmdList, FRHIViewableResource* Resource, FRHIViewDesc const& ViewDesc)
 	{
 		return new FRHIShaderResourceView(Resource, ViewDesc);
 	}
 
-	virtual FUnorderedAccessViewRHIRef RHICreateUnorderedAccessView(class FRHICommandListImmediate& RHICmdList, FRHIViewableResource* Resource, FRHIViewDesc const& ViewDesc)
+	virtual FUnorderedAccessViewRHIRef RHICreateUnorderedAccessView(class FRHICommandListBase& RHICmdList, FRHIViewableResource* Resource, FRHIViewDesc const& ViewDesc)
 	{
 		return new FRHIUnorderedAccessView(Resource, ViewDesc);
 	}
diff --git a/Engine/Source/Runtime/OpenGLDrv/Private/OpenGLUAV.cpp b/Engine/Source/Runtime/OpenGLDrv/Private/OpenGLUAV.cpp
--- a/Engine/Source/Runtime/OpenGLDrv/Private/OpenGLUAV.cpp
+++ b/Engine/Source/Runtime/OpenGLDrv/Private/OpenGLUAV.cpp
@@ -37,10 +37,10 @@ static void BindGLTexBufferRange(GLenum Target, GLenum InternalFormat, GLuint Bu
 	}
 }
 
-FOpenGLShaderResourceView::FOpenGLShaderResourceView(FRHICommandListImmediate& RHICmdList, FRHIViewableResource* InResource, FRHIViewDesc const& InViewDesc)
+FOpenGLShaderResourceView::FOpenGLShaderResourceView(FRHICommandListBase& RHICmdList, FRHIViewableResource* InResource, FRHIViewDesc const& InViewDesc)
 	: FRHIShaderResourceView(InResource, InViewDesc)
 {
-	RHICmdList.EnqueueLambda([this](FRHICommandListImmediate&)
+	RHICmdList.EnqueueLambda([this](FRHICommandListBase&)
 	{
 		LinkHead(GetBaseResource()->LinkedViews);
 		UpdateView();
@@ -164,10 +164,10 @@ FOpenGLShaderResourceView::~FOpenGLShaderResourceView()
 	Invalidate();
 }
 
-FOpenGLUnorderedAccessView::FOpenGLUnorderedAccessView(FRHICommandListImmediate& RHICmdList, FRHIViewableResource* InResource, FRHIViewDesc const& InViewDesc)
+FOpenGLUnorderedAccessView::FOpenGLUnorderedAccessView(FRHICommandListBase& RHICmdList, FRHIViewableResource* InResource, FRHIViewDesc const& InViewDesc)
 	: FRHIUnorderedAccessView(InResource, InViewDesc)
 {
-	RHICmdList.EnqueueLambda([this](FRHICommandListImmediate&)
+	RHICmdList.EnqueueLambda([this](FRHICommandListBase&)
 	{
 		LinkHead(GetBaseResource()->LinkedViews);
 		UpdateView();
@@ -368,12 +368,12 @@ void FOpenGLDynamicRHI::RHIClearUAVUint(FRHIUnorderedAccessView* UnorderedAccess
 	}
 }
 
-FShaderResourceViewRHIRef FOpenGLDynamicRHI::RHICreateShaderResourceView(class FRHICommandListImmediate& RHICmdList, FRHIViewableResource* Resource, FRHIViewDesc const& ViewDesc)
+FShaderResourceViewRHIRef FOpenGLDynamicRHI::RHICreateShaderResourceView(class FRHICommandListBase& RHICmdList, FRHIViewableResource* Resource, FRHIViewDesc const& ViewDesc)
 {
 	return new FOpenGLShaderResourceView(RHICmdList, Resource, ViewDesc);
 }
 
-FUnorderedAccessViewRHIRef FOpenGLDynamicRHI::RHICreateUnorderedAccessView(class FRHICommandListImmediate& RHICmdList, FRHIViewableResource* Resource, FRHIViewDesc const& ViewDesc)
+FUnorderedAccessViewRHIRef FOpenGLDynamicRHI::RHICreateUnorderedAccessView(class FRHICommandListBase& RHICmdList, FRHIViewableResource* Resource, FRHIViewDesc const& ViewDesc)
 {
 	return new FOpenGLUnorderedAccessView(RHICmdList, Resource, ViewDesc);
 }
diff --git a/Engine/Source/Runtime/OpenGLDrv/Public/OpenGLDrv.h b/Engine/Source/Runtime/OpenGLDrv/Public/OpenGLDrv.h
--- a/Engine/Source/Runtime/OpenGLDrv/Public/OpenGLDrv.h
+++ b/Engine/Source/Runtime/OpenGLDrv/Public/OpenGLDrv.h
@@ -354,8 +354,8 @@ public:
 	virtual FComputeShaderRHIRef RHICreateComputeShader(TArrayView<const uint8> Code, const FSHAHash& Hash) final override;
 
 	// SRV / UAV creation functions
-	virtual FShaderResourceViewRHIRef  RHICreateShaderResourceView (class FRHICommandListImmediate& RHICmdList, FRHIViewableResource* Resource, FRHIViewDesc const& ViewDesc) final override;
-	virtual FUnorderedAccessViewRHIRef RHICreateUnorderedAccessView(class FRHICommandListImmediate& RHICmdList, FRHIViewableResource* Resource, FRHIViewDesc const& ViewDesc) final override;
+	virtual FShaderResourceViewRHIRef  RHICreateShaderResourceView (class FRHICommandListBase& RHICmdList, FRHIViewableResource* Resource, FRHIViewDesc const& ViewDesc) final override;
+	virtual FUnorderedAccessViewRHIRef RHICreateUnorderedAccessView(class FRHICommandListBase& RHICmdList, FRHIViewableResource* Resource, FRHIViewDesc const& ViewDesc) final override;
 
 	virtual FUniformBufferRHIRef RHICreateUniformBuffer(const void* Contents, const FRHIUniformBufferLayout* Layout, EUniformBufferUsage Usage, EUniformBufferValidation Validation) final override;
 	virtual void RHIUpdateUniformBuffer(FRHICommandListBase& RHICmdList, FRHIUniformBuffer* UniformBufferRHI, const void* Contents) final override;
diff --git a/Engine/Source/Runtime/OpenGLDrv/Public/OpenGLResources.h b/Engine/Source/Runtime/OpenGLDrv/Public/OpenGLResources.h
--- a/Engine/Source/Runtime/OpenGLDrv/Public/OpenGLResources.h
+++ b/Engine/Source/Runtime/OpenGLDrv/Public/OpenGLResources.h
@@ -1168,7 +1168,7 @@ public:
 class FOpenGLUnorderedAccessView final : public FRHIUnorderedAccessView, public FOpenGLView
 {
 public:
-	FOpenGLUnorderedAccessView(FRHICommandListImmediate& RHICmdList, FRHIViewableResource* Resource, FRHIViewDesc const& ViewDesc);
+	FOpenGLUnorderedAccessView(FRHICommandListBase& RHICmdList, FRHIViewableResource* Resource, FRHIViewDesc const& ViewDesc);
 	virtual ~FOpenGLUnorderedAccessView();
 
 	FOpenGLViewableResource* GetBaseResource() const;
@@ -1203,7 +1203,7 @@ private:
 class FOpenGLShaderResourceView final : public FRHIShaderResourceView, public FOpenGLView
 {
 public:
-	FOpenGLShaderResourceView(FRHICommandListImmediate& RHICmdList, FRHIViewableResource* Resource, FRHIViewDesc const& ViewDesc);
+	FOpenGLShaderResourceView(FRHICommandListBase& RHICmdList, FRHIViewableResource* Resource, FRHIViewDesc const& ViewDesc);
 	virtual ~FOpenGLShaderResourceView();
 
 	FOpenGLViewableResource* GetBaseResource() const;
diff --git a/Engine/Source/Runtime/RHI/Private/RHICommandList.cpp b/Engine/Source/Runtime/RHI/Private/RHICommandList.cpp
--- a/Engine/Source/Runtime/RHI/Private/RHICommandList.cpp
+++ b/Engine/Source/Runtime/RHI/Private/RHICommandList.cpp
@@ -1415,7 +1415,7 @@ void FRHIComputeCommandList::BuildAccelerationStructure(FRHIRayTracingGeometry*
 	FRHIBufferRange ScratchBufferRange{};
 	
 	FRHIResourceCreateInfo ScratchBufferCreateInfo(TEXT("RHIScratchBuffer"));
-	ScratchBufferRange.Buffer = RHICreateBuffer(Geometry->GetSizeInfo().BuildScratchSize, BUF_StructuredBuffer | BUF_RayTracingScratch, 0, ERHIAccess::UAVCompute, ScratchBufferCreateInfo);
+	ScratchBufferRange.Buffer = CreateBuffer(Geometry->GetSizeInfo().BuildScratchSize, BUF_StructuredBuffer | BUF_RayTracingScratch, 0, ERHIAccess::UAVCompute, ScratchBufferCreateInfo);
 
 	BuildAccelerationStructures(MakeArrayView(&Params, 1), ScratchBufferRange);
 }
@@ -1431,7 +1431,7 @@ void FRHIComputeCommandList::BuildAccelerationStructures(const TArrayView<const
 
 	FRHIResourceCreateInfo ScratchBufferCreateInfo(TEXT("RHIScratchBuffer"));
 	FRHIBufferRange ScratchBufferRange{};	
-	ScratchBufferRange.Buffer = RHICreateBuffer(TotalRequiredScratchMemorySize, BUF_StructuredBuffer | BUF_RayTracingScratch, 0, ERHIAccess::UAVCompute, ScratchBufferCreateInfo);
+	ScratchBufferRange.Buffer = CreateBuffer(TotalRequiredScratchMemorySize, BUF_StructuredBuffer | BUF_RayTracingScratch, 0, ERHIAccess::UAVCompute, ScratchBufferCreateInfo);
 
 	BuildAccelerationStructures(Params, ScratchBufferRange);
 }
diff --git a/Engine/Source/Runtime/RHI/Public/DynamicRHI.h b/Engine/Source/Runtime/RHI/Public/DynamicRHI.h
--- a/Engine/Source/Runtime/RHI/Public/DynamicRHI.h
+++ b/Engine/Source/Runtime/RHI/Public/DynamicRHI.h
@@ -432,8 +432,8 @@ public:
 	RHI_API virtual FTextureReferenceRHIRef RHICreateTextureReference(FRHITexture* InReferencedTexture);
 
 	// SRV / UAV creation functions
-	virtual FShaderResourceViewRHIRef  RHICreateShaderResourceView (class FRHICommandListImmediate& RHICmdList, FRHIViewableResource* Resource, FRHIViewDesc const& ViewDesc) = 0;
-	virtual FUnorderedAccessViewRHIRef RHICreateUnorderedAccessView(class FRHICommandListImmediate& RHICmdList, FRHIViewableResource* Resource, FRHIViewDesc const& ViewDesc) = 0;
+	virtual FShaderResourceViewRHIRef  RHICreateShaderResourceView (class FRHICommandListBase& RHICmdList, FRHIViewableResource* Resource, FRHIViewDesc const& ViewDesc) = 0;
+	virtual FUnorderedAccessViewRHIRef RHICreateUnorderedAccessView(class FRHICommandListBase& RHICmdList, FRHIViewableResource* Resource, FRHIViewDesc const& ViewDesc) = 0;
 
 	/**
 	* Generates mip maps for a texture.
@@ -969,7 +969,7 @@ public:
 		return {};
 	}
 
-	virtual FRayTracingGeometryRHIRef RHICreateRayTracingGeometry(const FRayTracingGeometryInitializer& Initializer)
+	virtual FRayTracingGeometryRHIRef RHICreateRayTracingGeometry(FRHICommandListBase& RHICmdList, const FRayTracingGeometryInitializer& Initializer)
 	{
 		checkNoEntry();
 		return nullptr;
@@ -1482,11 +1482,6 @@ FORCEINLINE FRayTracingAccelerationStructureSize RHICalcRayTracingGeometrySize(c
 	return GDynamicRHI->RHICalcRayTracingGeometrySize(Initializer);
 }
 
-FORCEINLINE FRayTracingGeometryRHIRef RHICreateRayTracingGeometry(const FRayTracingGeometryInitializer& Initializer)
-{
-	return GDynamicRHI->RHICreateRayTracingGeometry(Initializer);
-}
-
 FORCEINLINE FRayTracingSceneRHIRef RHICreateRayTracingScene(FRayTracingSceneInitializer2 Initializer)
 {
 	PRAGMA_DISABLE_DEPRECATION_WARNINGS
diff --git a/Engine/Source/Runtime/RHI/Public/RHICommandList.h b/Engine/Source/Runtime/RHI/Public/RHICommandList.h
--- a/Engine/Source/Runtime/RHI/Public/RHICommandList.h
+++ b/Engine/Source/Runtime/RHI/Public/RHICommandList.h
@@ -741,6 +741,177 @@
 		GDynamicRHI->RHIUpdateTexture3D(*this, Texture, MipIndex, UpdateRegion, SourceRowPitch, SourceDepthPitch, SourceData);
 	}
 
+	FORCEINLINE FShaderResourceViewRHIRef CreateShaderResourceView(FRHIBuffer* Buffer, FRHIViewDesc::FBufferSRV::FInitializer const& ViewDesc)
+	{
+		LLM_SCOPE_BYNAME(TEXT("RHIMisc/CreateShaderResourceView"));
+		return GDynamicRHI->RHICreateShaderResourceView(*this, Buffer, ViewDesc);
+	}
+
+	FORCEINLINE FShaderResourceViewRHIRef CreateShaderResourceView(FRHITexture* Texture, FRHIViewDesc::FTextureSRV::FInitializer const& ViewDesc)
+	{
+		LLM_SCOPE_BYNAME(TEXT("RHIMisc/CreateShaderResourceView"));
+		checkf(Texture->GetTextureReference() == nullptr, TEXT("Creating a shader resource view of an FRHITextureReference is not supported."));
+
+		return GDynamicRHI->RHICreateShaderResourceView(*this, Texture, ViewDesc);
+	}
+
+	FORCEINLINE FUnorderedAccessViewRHIRef CreateUnorderedAccessView(FRHIBuffer* Buffer, FRHIViewDesc::FBufferUAV::FInitializer const& ViewDesc)
+	{
+		LLM_SCOPE_BYNAME(TEXT("RHIMisc/CreateUnorderedAccessView"));
+		return GDynamicRHI->RHICreateUnorderedAccessView(*this, Buffer, ViewDesc);
+	}
+
+	FORCEINLINE FUnorderedAccessViewRHIRef CreateUnorderedAccessView(FRHITexture* Texture, FRHIViewDesc::FTextureUAV::FInitializer const& ViewDesc)
+	{
+		LLM_SCOPE_BYNAME(TEXT("RHIMisc/CreateUnorderedAccessView"));
+		checkf(Texture->GetTextureReference() == nullptr, TEXT("Creating an unordered access view of an FRHITextureReference is not supported."));
+
+		return GDynamicRHI->RHICreateUnorderedAccessView(*this, Texture, ViewDesc);
+	}
+
+	//UE_DEPRECATED(5.3, "Use the CreateUnorderedAccessView function that takes an FRHIViewDesc.")
+	FORCEINLINE FUnorderedAccessViewRHIRef CreateUnorderedAccessView(FRHIBuffer* Buffer, bool bUseUAVCounter, bool bAppendBuffer)
+	{
+		return CreateUnorderedAccessView(Buffer, FRHIViewDesc::CreateBufferUAV()
+			.SetTypeFromBuffer(Buffer)
+			.SetAtomicCounter(bUseUAVCounter)
+			.SetAppendBuffer(bAppendBuffer)
+		);
+	}
+
+	//UE_DEPRECATED(5.3, "Use the CreateUnorderedAccessView function that takes an FRHIViewDesc.")
+	FORCEINLINE FUnorderedAccessViewRHIRef CreateUnorderedAccessView(FRHIBuffer* Buffer, uint8 Format)
+	{
+		// For back-compat reasons, SRVs of byte-address buffers created via this function ignore the Format, and instead create raw views.
+		if (Buffer && EnumHasAnyFlags(Buffer->GetDesc().Usage, BUF_ByteAddressBuffer))
+		{
+			return CreateUnorderedAccessView(Buffer, FRHIViewDesc::CreateBufferUAV()
+				.SetType(FRHIViewDesc::EBufferType::Raw)
+			);
+		}
+		else
+		{
+			return CreateUnorderedAccessView(Buffer, FRHIViewDesc::CreateBufferUAV()
+				.SetType(FRHIViewDesc::EBufferType::Typed)
+				.SetFormat(EPixelFormat(Format))
+			);
+		}
+	}
+
+	//UE_DEPRECATED(5.3, "Use the CreateUnorderedAccessView function that takes an FRHIViewDesc.")
+	FORCEINLINE FUnorderedAccessViewRHIRef CreateUnorderedAccessView(FRHITexture* Texture, uint32 MipLevel, uint16 FirstArraySlice, uint16 NumArraySlices)
+	{
+		check(MipLevel < 256);
+
+		return CreateUnorderedAccessView(Texture, FRHIViewDesc::CreateTextureUAV()
+			.SetDimensionFromTexture(Texture)
+			.SetMipLevel(uint8(MipLevel))
+			.SetArrayRange(FirstArraySlice, NumArraySlices)
+		);
+	}
+
+	//UE_DEPRECATED(5.3, "Use the CreateUnorderedAccessView function that takes an FRHIViewDesc.")
+	FORCEINLINE FUnorderedAccessViewRHIRef CreateUnorderedAccessView(FRHITexture* Texture, uint32 MipLevel, uint8 Format, uint16 FirstArraySlice, uint16 NumArraySlices)
+	{
+		check(MipLevel < 256);
+
+		return CreateUnorderedAccessView(Texture, FRHIViewDesc::CreateTextureUAV()
+			.SetDimensionFromTexture(Texture)
+			.SetMipLevel(uint8(MipLevel))
+			.SetFormat(EPixelFormat(Format))
+			.SetArrayRange(FirstArraySlice, NumArraySlices)
+		);
+	}
+
+	//UE_DEPRECATED(5.3, "Use the CreateShaderResourceView function that takes an FRHIViewDesc.")
+	FORCEINLINE FShaderResourceViewRHIRef CreateShaderResourceView(FRHIBuffer* Buffer)
+	{
+		FShaderResourceViewRHIRef SRVRef = CreateShaderResourceView(Buffer, FRHIViewDesc::CreateBufferSRV()
+			.SetTypeFromBuffer(Buffer));
+		checkf(SRVRef->GetDesc().Buffer.SRV.BufferType != FRHIViewDesc::EBufferType::Typed,
+			TEXT("Typed buffer should be created using CreateShaderResourceView where Format is specified."));
+		return SRVRef;
+	}
+
+	//UE_DEPRECATED(5.3, "Use the CreateShaderResourceView function that takes an FRHIViewDesc.")
+	FORCEINLINE FShaderResourceViewRHIRef CreateShaderResourceView(FRHIBuffer* Buffer, uint32 Stride, uint8 Format)
+	{
+		check(Format != PF_Unknown);
+		check(Stride == GPixelFormats[Format].BlockBytes);
+
+		// For back-compat reasons, SRVs of byte-address buffers created via this function ignore the Format, and instead create raw views.
+		if (Buffer && EnumHasAnyFlags(Buffer->GetDesc().Usage, BUF_ByteAddressBuffer))
+		{
+			return CreateShaderResourceView(Buffer, FRHIViewDesc::CreateBufferSRV()
+				.SetType(FRHIViewDesc::EBufferType::Raw)
+			);
+		}
+		else
+		{
+			return CreateShaderResourceView(Buffer, FRHIViewDesc::CreateBufferSRV()
+				.SetType(FRHIViewDesc::EBufferType::Typed)
+				.SetFormat(EPixelFormat(Format))
+			);
+		}
+	}
+
+	//UE_DEPRECATED(5.3, "Use the CreateShaderResourceView function that takes an FRHIViewDesc.")
+	FORCEINLINE FShaderResourceViewRHIRef CreateShaderResourceView(FRHITexture* Texture, const FRHITextureSRVCreateInfo& CreateInfo)
+	{
+		return CreateShaderResourceView(Texture, FRHIViewDesc::CreateTextureSRV()
+			.SetDimensionFromTexture(Texture)
+			.SetFormat     (CreateInfo.Format)
+			.SetMipRange   (CreateInfo.MipLevel, CreateInfo.NumMipLevels)
+			.SetDisableSRGB(CreateInfo.SRGBOverride == SRGBO_ForceDisable)
+			.SetArrayRange (CreateInfo.FirstArraySlice, CreateInfo.NumArraySlices)
+			.SetPlane      (CreateInfo.MetaData)
+		);
+	}
+
+	//UE_DEPRECATED(5.3, "Use the CreateShaderResourceView function that takes an FRHIViewDesc.")
+	FORCEINLINE FShaderResourceViewRHIRef CreateShaderResourceView(FRHITexture* Texture, uint8 MipLevel)
+	{
+		return CreateShaderResourceView(Texture, FRHIViewDesc::CreateTextureSRV()
+			.SetDimensionFromTexture(Texture)
+			.SetMipRange(MipLevel, 1)
+		);
+	}
+
+	//UE_DEPRECATED(5.3, "Use the CreateShaderResourceView function that takes an FRHIViewDesc.")
+	FORCEINLINE FShaderResourceViewRHIRef CreateShaderResourceView(FRHITexture* Texture, uint8 MipLevel, uint8 NumMipLevels, EPixelFormat Format)
+	{
+		return CreateShaderResourceView(Texture, FRHIViewDesc::CreateTextureSRV()
+			.SetDimensionFromTexture(Texture)
+			.SetMipRange(MipLevel, NumMipLevels)
+			.SetFormat(Format)
+		);
+	}
+
+	//UE_DEPRECATED(5.3, "Use the CreateShaderResourceView function that takes an FRHIViewDesc.")
+	FORCEINLINE FShaderResourceViewRHIRef CreateShaderResourceViewWriteMask(FRHITexture2D* Texture2DRHI)
+	{
+		return CreateShaderResourceView(Texture2DRHI, FRHIViewDesc::CreateTextureSRV()
+			.SetDimensionFromTexture(Texture2DRHI)
+			.SetPlane(ERHITexturePlane::CMask)
+		);
+	}
+
+	//UE_DEPRECATED(5.3, "Use the CreateShaderResourceView function that takes an FRHIViewDesc.")
+	FORCEINLINE FShaderResourceViewRHIRef CreateShaderResourceViewFMask(FRHITexture2D* Texture2DRHI)
+	{
+		return CreateShaderResourceView(Texture2DRHI, FRHIViewDesc::CreateTextureSRV()
+			.SetDimensionFromTexture(Texture2DRHI)
+			.SetPlane(ERHITexturePlane::FMask)
+		);
+	}
+
+#if RHI_RAYTRACING
+	FORCEINLINE FRayTracingGeometryRHIRef CreateRayTracingGeometry(const FRayTracingGeometryInitializer& Initializer)
+	{
+		return GDynamicRHI->RHICreateRayTracingGeometry(*this, Initializer);
+	}
+#endif
+
 	inline FRHIBatchedShaderParameters& GetScratchShaderParameters()
 	{
 		if (!ensureMsgf(!ScratchShaderParameters.HasParameters(), TEXT("Scratch shader parameters left without committed parameters")))
@@ -4043,168 +4214,6 @@
 		GDynamicRHI->RHIUnlockBufferMGPU(*this, Buffer, GPUIndex);
 	}
 
-	FORCEINLINE FShaderResourceViewRHIRef CreateShaderResourceView(FRHIBuffer* Buffer, FRHIViewDesc::FBufferSRV::FInitializer const& ViewDesc)
-	{
-		LLM_SCOPE_BYNAME(TEXT("RHIMisc/CreateShaderResourceView"));
-		return GDynamicRHI->RHICreateShaderResourceView(*this, Buffer, ViewDesc);
-	}
-	FORCEINLINE FShaderResourceViewRHIRef CreateShaderResourceView(FRHITexture* Texture, FRHIViewDesc::FTextureSRV::FInitializer const& ViewDesc)
-	{
-		LLM_SCOPE_BYNAME(TEXT("RHIMisc/CreateShaderResourceView"));
-		checkf(Texture->GetTextureReference() == nullptr, TEXT("Creating a shader resource view of an FRHITextureReference is not supported."));
-
-		return GDynamicRHI->RHICreateShaderResourceView(*this, Texture, ViewDesc);
-	}
-
-	FORCEINLINE FUnorderedAccessViewRHIRef CreateUnorderedAccessView(FRHIBuffer* Buffer, FRHIViewDesc::FBufferUAV::FInitializer const& ViewDesc)
-	{
-		LLM_SCOPE_BYNAME(TEXT("RHIMisc/CreateUnorderedAccessView"));
-		return GDynamicRHI->RHICreateUnorderedAccessView(*this, Buffer, ViewDesc);
-	}
-	FORCEINLINE FUnorderedAccessViewRHIRef CreateUnorderedAccessView(FRHITexture* Texture, FRHIViewDesc::FTextureUAV::FInitializer const& ViewDesc)
-	{
-		LLM_SCOPE_BYNAME(TEXT("RHIMisc/CreateUnorderedAccessView"));
-		checkf(Texture->GetTextureReference() == nullptr, TEXT("Creating an unordered access view of an FRHITextureReference is not supported."));
-
-		return GDynamicRHI->RHICreateUnorderedAccessView(*this, Texture, ViewDesc);
-	}
-
-	//UE_DEPRECATED(5.3, "Use the CreateUnorderedAccessView function that takes an FRHIViewDesc.")
-	FORCEINLINE FUnorderedAccessViewRHIRef CreateUnorderedAccessView(FRHIBuffer* Buffer, bool bUseUAVCounter, bool bAppendBuffer)
-	{
-		return CreateUnorderedAccessView(Buffer, FRHIViewDesc::CreateBufferUAV()
-			.SetTypeFromBuffer(Buffer)
-			.SetAtomicCounter(bUseUAVCounter)
-			.SetAppendBuffer(bAppendBuffer)
-		);
-	}
-
-	//UE_DEPRECATED(5.3, "Use the CreateUnorderedAccessView function that takes an FRHIViewDesc.")
-	FORCEINLINE FUnorderedAccessViewRHIRef CreateUnorderedAccessView(FRHIBuffer* Buffer, uint8 Format)
-	{
-		// For back-compat reasons, SRVs of byte-address buffers created via this function ignore the Format, and instead create raw views.
-		if (Buffer && EnumHasAnyFlags(Buffer->GetDesc().Usage, BUF_ByteAddressBuffer))
-		{
-			return CreateUnorderedAccessView(Buffer, FRHIViewDesc::CreateBufferUAV()
-				.SetType(FRHIViewDesc::EBufferType::Raw)
-			);
-		}
-		else
-		{
-			return CreateUnorderedAccessView(Buffer, FRHIViewDesc::CreateBufferUAV()
-				.SetType(FRHIViewDesc::EBufferType::Typed)
-				.SetFormat(EPixelFormat(Format))
-			);
-		}
-	}
-
-	//UE_DEPRECATED(5.3, "Use the CreateUnorderedAccessView function that takes an FRHIViewDesc.")
-	FORCEINLINE FUnorderedAccessViewRHIRef CreateUnorderedAccessView(FRHITexture* Texture, uint32 MipLevel, uint16 FirstArraySlice, uint16 NumArraySlices)
-	{
-		check(MipLevel < 256);
-
-		return CreateUnorderedAccessView(Texture, FRHIViewDesc::CreateTextureUAV()
-			.SetDimensionFromTexture(Texture)
-			.SetMipLevel(uint8(MipLevel))
-			.SetArrayRange(FirstArraySlice, NumArraySlices)
-		);
-	}
-
-	//UE_DEPRECATED(5.3, "Use the CreateUnorderedAccessView function that takes an FRHIViewDesc.")
-	FORCEINLINE FUnorderedAccessViewRHIRef CreateUnorderedAccessView(FRHITexture* Texture, uint32 MipLevel, uint8 Format, uint16 FirstArraySlice, uint16 NumArraySlices)
-	{
-		check(MipLevel < 256);
-
-		return CreateUnorderedAccessView(Texture, FRHIViewDesc::CreateTextureUAV()
-			.SetDimensionFromTexture(Texture)
-			.SetMipLevel(uint8(MipLevel))
-			.SetFormat(EPixelFormat(Format))
-			.SetArrayRange(FirstArraySlice, NumArraySlices)
-		);
-	}
-
-	//UE_DEPRECATED(5.3, "Use the CreateShaderResourceView function that takes an FRHIViewDesc.")
-	FORCEINLINE FShaderResourceViewRHIRef CreateShaderResourceView(FRHIBuffer* Buffer)
-	{
-		FShaderResourceViewRHIRef SRVRef = CreateShaderResourceView(Buffer, FRHIViewDesc::CreateBufferSRV()
-				.SetTypeFromBuffer(Buffer));
-		checkf(SRVRef->GetDesc().Buffer.SRV.BufferType != FRHIViewDesc::EBufferType::Typed,
-			TEXT("Typed buffer should be created using CreateShaderResourceView where Format is specified."));
-		return SRVRef;
-	}
-	
-	//UE_DEPRECATED(5.3, "Use the CreateShaderResourceView function that takes an FRHIViewDesc.")
-	FORCEINLINE FShaderResourceViewRHIRef CreateShaderResourceView(FRHIBuffer* Buffer, uint32 Stride, uint8 Format)
-	{
-		check(Format != PF_Unknown);
-		check(Stride == GPixelFormats[Format].BlockBytes);
-
-		// For back-compat reasons, SRVs of byte-address buffers created via this function ignore the Format, and instead create raw views.
-		if (Buffer && EnumHasAnyFlags(Buffer->GetDesc().Usage, BUF_ByteAddressBuffer))
-		{
-			return CreateShaderResourceView(Buffer, FRHIViewDesc::CreateBufferSRV()
-				.SetType(FRHIViewDesc::EBufferType::Raw)
-			);
-		}
-		else
-		{
-			return CreateShaderResourceView(Buffer, FRHIViewDesc::CreateBufferSRV()
-				.SetType(FRHIViewDesc::EBufferType::Typed)
-				.SetFormat(EPixelFormat(Format))
-			);
-		}
-	}
-
-	//UE_DEPRECATED(5.3, "Use the CreateShaderResourceView function that takes an FRHIViewDesc.")
-	FORCEINLINE FShaderResourceViewRHIRef CreateShaderResourceView(FRHITexture* Texture, const FRHITextureSRVCreateInfo& CreateInfo)
-	{
-		return CreateShaderResourceView(Texture, FRHIViewDesc::CreateTextureSRV()
-			.SetDimensionFromTexture(Texture)
-			.SetFormat     (CreateInfo.Format)
-			.SetMipRange   (CreateInfo.MipLevel, CreateInfo.NumMipLevels)
-			.SetDisableSRGB(CreateInfo.SRGBOverride == SRGBO_ForceDisable)
-			.SetArrayRange (CreateInfo.FirstArraySlice, CreateInfo.NumArraySlices)
-			.SetPlane      (CreateInfo.MetaData)
-		);
-	}
-
-	//UE_DEPRECATED(5.3, "Use the CreateShaderResourceView function that takes an FRHIViewDesc.")
-	FORCEINLINE FShaderResourceViewRHIRef CreateShaderResourceView(FRHITexture* Texture, uint8 MipLevel)
-	{
-		return CreateShaderResourceView(Texture, FRHIViewDesc::CreateTextureSRV()
-			.SetDimensionFromTexture(Texture)
-			.SetMipRange(MipLevel, 1)
-		);
-	}
-
-	//UE_DEPRECATED(5.3, "Use the CreateShaderResourceView function that takes an FRHIViewDesc.")
-	FORCEINLINE FShaderResourceViewRHIRef CreateShaderResourceView(FRHITexture* Texture, uint8 MipLevel, uint8 NumMipLevels, EPixelFormat Format)
-	{
-		return CreateShaderResourceView(Texture, FRHIViewDesc::CreateTextureSRV()
-			.SetDimensionFromTexture(Texture)
-			.SetMipRange(MipLevel, NumMipLevels)
-			.SetFormat(Format)
-		);
-	}
-
-	//UE_DEPRECATED(5.3, "Use the CreateShaderResourceView function that takes an FRHIViewDesc.")
-	FORCEINLINE FShaderResourceViewRHIRef CreateShaderResourceViewWriteMask(FRHITexture2D* Texture2DRHI)
-	{
-		return CreateShaderResourceView(Texture2DRHI, FRHIViewDesc::CreateTextureSRV()
-			.SetDimensionFromTexture(Texture2DRHI)
-			.SetPlane(ERHITexturePlane::CMask)
-		);
-	}
-
-	//UE_DEPRECATED(5.3, "Use the CreateShaderResourceView function that takes an FRHIViewDesc.")
-	FORCEINLINE FShaderResourceViewRHIRef CreateShaderResourceViewFMask(FRHITexture2D* Texture2DRHI)
-	{
-		return CreateShaderResourceView(Texture2DRHI, FRHIViewDesc::CreateTextureSRV()
-			.SetDimensionFromTexture(Texture2DRHI)
-			.SetPlane(ERHITexturePlane::FMask)
-		);
-	}
-
 	UE_DEPRECATED(5.2, "Use the global scope RHIGetTextureMemoryStats function.")
 	FORCEINLINE void GetTextureMemoryStats(FTextureMemoryStats& OutStats)
 	{
@@ -5319,6 +5328,13 @@
 	 FRHICommandListExecutor::GetImmediateCommandList().UnlockStagingBuffer(StagingBuffer);
 }
 
+#if RHI_RAYTRACING
+FORCEINLINE FRayTracingGeometryRHIRef RHICreateRayTracingGeometry(const FRayTracingGeometryInitializer& Initializer)
+{
+	return FRHICommandListExecutor::GetImmediateCommandList().CreateRayTracingGeometry(Initializer);
+}
+#endif
+
 namespace UE::RHI
 {
 
diff --git a/Engine/Source/Runtime/RHI/Public/RHIResources.h b/Engine/Source/Runtime/RHI/Public/RHIResources.h
--- a/Engine/Source/Runtime/RHI/Public/RHIResources.h
+++ b/Engine/Source/Runtime/RHI/Public/RHIResources.h
@@ -1491,6 +1491,7 @@ struct FRHITextureDesc
 		Hash = HashCombine(Hash, GetTypeHash(Desc.ArraySize	));
 		Hash = HashCombine(Hash, GetTypeHash(Desc.NumMips	));
 		Hash = HashCombine(Hash, GetTypeHash(Desc.NumSamples));
+		Hash = HashCombine(Hash, GetTypeHash(Desc.FastVRAMPercentage));
 		Hash = HashCombine(Hash, GetTypeHash(Desc.ClearValue));
 		Hash = HashCombine(Hash, GetTypeHash(Desc.ExtData   ));
 		Hash = HashCombine(Hash, GetTypeHash(Desc.GPUMask.GetNative()));
@@ -1508,6 +1509,7 @@ struct FRHITextureDesc
 			&& ArraySize  == Other.ArraySize
 			&& NumMips    == Other.NumMips
 			&& NumSamples == Other.NumSamples
+			&& FastVRAMPercentage == Other.FastVRAMPercentage
 			&& ClearValue == Other.ClearValue
 			&& ExtData    == Other.ExtData
 			&& GPUMask    == Other.GPUMask;
@@ -2476,7 +2478,7 @@ static_assert(TIsTrivial<FRHIViewDesc>::Value, "FRHIViewDesc must be a trivial t
 struct FRHIViewDesc::FBufferSRV::FInitializer : private FRHIViewDesc
 {
 	friend FRHIViewDesc;
-	friend FRHICommandListImmediate;
+	friend FRHICommandListBase;
 	friend struct FShaderResourceViewInitializer;
 	friend struct FRawBufferShaderResourceViewInitializer;
 
@@ -2538,7 +2540,7 @@ public:
 struct FRHIViewDesc::FBufferUAV::FInitializer : private FRHIViewDesc
 {
 	friend FRHIViewDesc;
-	friend FRHICommandListImmediate;
+	friend FRHICommandListBase;
 
 protected:
 	FInitializer()
@@ -2610,7 +2612,7 @@ public:
 struct FRHIViewDesc::FTextureSRV::FInitializer : private FRHIViewDesc
 {
 	friend FRHIViewDesc;
-	friend FRHICommandListImmediate;
+	friend FRHICommandListBase;
 
 protected:
 	FInitializer()
@@ -2697,7 +2699,7 @@ public:
 struct FRHIViewDesc::FTextureUAV::FInitializer : private FRHIViewDesc
 {
 	friend FRHIViewDesc;
-	friend FRHICommandListImmediate;
+	friend FRHICommandListBase;
 
 protected:
 	FInitializer()
diff --git a/Engine/Source/Runtime/RHI/Public/RHIUtilities.h b/Engine/Source/Runtime/RHI/Public/RHIUtilities.h
--- a/Engine/Source/Runtime/RHI/Public/RHIUtilities.h
+++ b/Engine/Source/Runtime/RHI/Public/RHIUtilities.h
@@ -213,7 +213,7 @@ struct FRWBuffer
 	}
 
 	// @param AdditionalUsage passed down to RHICreateVertexBuffer(), get combined with "BUF_UnorderedAccess | BUF_ShaderResource" e.g. BUF_Static
-	void Initialize(const TCHAR* InDebugName, uint32 BytesPerElement, uint32 NumElements, EPixelFormat Format, ERHIAccess InResourceState, EBufferUsageFlags AdditionalUsage = BUF_None, FResourceArrayInterface *InResourceArray = nullptr)
+	void Initialize(FRHICommandListBase& RHICmdList, const TCHAR* InDebugName, uint32 BytesPerElement, uint32 NumElements, EPixelFormat Format, ERHIAccess InResourceState, EBufferUsageFlags AdditionalUsage = BUF_None, FResourceArrayInterface *InResourceArray = nullptr)
 	{
 		// Provide a debug name if using Fast VRAM so the allocators diagnostics will work
 		ensure(!(EnumHasAnyFlags(AdditionalUsage, BUF_FastVRAM) && !InDebugName));
@@ -222,14 +222,26 @@ struct FRWBuffer
 		CreateInfo.ClassName = ClassName;
 		CreateInfo.OwnerName = OwnerName;
 		CreateInfo.ResourceArray = InResourceArray;
-		Buffer = RHICreateVertexBuffer(NumBytes, BUF_UnorderedAccess | BUF_ShaderResource | AdditionalUsage, InResourceState, CreateInfo);
-		UAV = RHICreateUnorderedAccessView(Buffer, UE_PIXELFORMAT_TO_UINT8(Format));
-		SRV = RHICreateShaderResourceView(Buffer, BytesPerElement, UE_PIXELFORMAT_TO_UINT8(Format));
+		Buffer = RHICmdList.CreateBuffer(NumBytes, BUF_VertexBuffer | BUF_UnorderedAccess | BUF_ShaderResource | AdditionalUsage, 0, InResourceState, CreateInfo);
+		UAV = RHICmdList.CreateUnorderedAccessView(Buffer, UE_PIXELFORMAT_TO_UINT8(Format));
+		SRV = RHICmdList.CreateShaderResourceView(Buffer, BytesPerElement, UE_PIXELFORMAT_TO_UINT8(Format));
+	}
+
+	void Initialize(FRHICommandListBase& RHICmdList, const TCHAR* InDebugName, uint32 BytesPerElement, uint32 NumElements, EPixelFormat Format, EBufferUsageFlags AdditionalUsage = BUF_None, FResourceArrayInterface* InResourceArray = nullptr)
+	{
+		Initialize(RHICmdList, InDebugName, BytesPerElement, NumElements, Format, ERHIAccess::UAVCompute, AdditionalUsage, InResourceArray);
 	}
 
 	void Initialize(const TCHAR* InDebugName, uint32 BytesPerElement, uint32 NumElements, EPixelFormat Format, EBufferUsageFlags AdditionalUsage = BUF_None, FResourceArrayInterface* InResourceArray = nullptr)
 	{
-		Initialize(InDebugName, BytesPerElement, NumElements, Format, ERHIAccess::UAVCompute, AdditionalUsage, InResourceArray);
+		check(IsInRenderingThread());
+		Initialize(FRHICommandListExecutor::GetImmediateCommandList(), InDebugName, BytesPerElement, NumElements, Format, ERHIAccess::UAVCompute, AdditionalUsage, InResourceArray);
+	}
+
+	void Initialize(const TCHAR* InDebugName, uint32 BytesPerElement, uint32 NumElements, EPixelFormat Format, ERHIAccess InResourceState, EBufferUsageFlags AdditionalUsage = BUF_None, FResourceArrayInterface* InResourceArray = nullptr)
+	{
+		check(IsInRenderingThread());
+		Initialize(FRHICommandListExecutor::GetImmediateCommandList(), InDebugName, BytesPerElement, NumElements, Format, InResourceState, AdditionalUsage, InResourceArray);
 	}
 
 	void Release()
diff --git a/Engine/Source/Runtime/RHI/Public/RHIValidation.h b/Engine/Source/Runtime/RHI/Public/RHIValidation.h
--- a/Engine/Source/Runtime/RHI/Public/RHIValidation.h
+++ b/Engine/Source/Runtime/RHI/Public/RHIValidation.h
@@ -446,11 +446,11 @@ public:
 		RHI->RHIUnlockBufferMGPU(RHICmdList, Buffer, GPUIndex);
 	}
 
-	virtual FShaderResourceViewRHIRef RHICreateShaderResourceView(class FRHICommandListImmediate& RHICmdList, FRHIViewableResource* Resource, FRHIViewDesc const& ViewDesc) override final
+	virtual FShaderResourceViewRHIRef RHICreateShaderResourceView(class FRHICommandListBase& RHICmdList, FRHIViewableResource* Resource, FRHIViewDesc const& ViewDesc) override final
 	{
 		return RHI->RHICreateShaderResourceView(RHICmdList, Resource, ViewDesc);
 	}
-	virtual FUnorderedAccessViewRHIRef RHICreateUnorderedAccessView(class FRHICommandListImmediate& RHICmdList, FRHIViewableResource* Resource, FRHIViewDesc const& ViewDesc) override final
+	virtual FUnorderedAccessViewRHIRef RHICreateUnorderedAccessView(class FRHICommandListBase& RHICmdList, FRHIViewableResource* Resource, FRHIViewDesc const& ViewDesc) override final
 	{
 		return RHI->RHICreateUnorderedAccessView(RHICmdList, Resource, ViewDesc);
 	}
@@ -1249,9 +1249,9 @@ public:
 
 
 #if RHI_RAYTRACING
-	virtual FRayTracingGeometryRHIRef RHICreateRayTracingGeometry(const FRayTracingGeometryInitializer& Initializer) override final
+	virtual FRayTracingGeometryRHIRef RHICreateRayTracingGeometry(FRHICommandListBase& RHICmdList, const FRayTracingGeometryInitializer& Initializer) override final
 	{
-		FRayTracingGeometryRHIRef Result = RHI->RHICreateRayTracingGeometry(Initializer);
+		FRayTracingGeometryRHIRef Result = RHI->RHICreateRayTracingGeometry(RHICmdList, Initializer);
 		Result->InitBarrierTracking(ERHIAccess::BVHWrite, *Initializer.DebugName.ToString()); // BVHs are always created in BVHWrite state
 		return Result;
 	}
diff --git a/Engine/Source/Runtime/VulkanRHI/Private/VulkanRayTracing.cpp b/Engine/Source/Runtime/VulkanRHI/Private/VulkanRayTracing.cpp
--- a/Engine/Source/Runtime/VulkanRHI/Private/VulkanRayTracing.cpp
+++ b/Engine/Source/Runtime/VulkanRHI/Private/VulkanRayTracing.cpp
@@ -345,7 +345,7 @@ static void GetBLASBuildData(
 FVulkanRayTracingGeometry::FVulkanRayTracingGeometry(ENoInit)
 {}
 
-FVulkanRayTracingGeometry::FVulkanRayTracingGeometry(const FRayTracingGeometryInitializer& InInitializer, FVulkanDevice* InDevice)
+FVulkanRayTracingGeometry::FVulkanRayTracingGeometry(FRHICommandListBase& RHICmdList, const FRayTracingGeometryInitializer& InInitializer, FVulkanDevice* InDevice)
 	: FRHIRayTracingGeometry(InInitializer), Device(InDevice)
 {
 	INC_DWORD_STAT(STAT_VulkanRayTracingAllocatedBLAS);
@@ -373,7 +373,7 @@ FVulkanRayTracingGeometry::FVulkanRayTracingGeometry(const FRayTracingGeometryIn
 
 	FString DebugNameString = Initializer.DebugName.ToString();
 	FRHIResourceCreateInfo BlasBufferCreateInfo(*DebugNameString);
-	AccelerationStructureBuffer = ResourceCast(RHICreateBuffer(SizeInfo.ResultSize, BUF_AccelerationStructure, 0, ERHIAccess::BVHWrite, BlasBufferCreateInfo).GetReference());
+	AccelerationStructureBuffer = ResourceCast(RHICmdList.CreateBuffer(SizeInfo.ResultSize, BUF_AccelerationStructure, 0, ERHIAccess::BVHWrite, BlasBufferCreateInfo).GetReference());
 
 	VkDevice NativeDevice = Device->GetInstanceHandle();
 
@@ -868,9 +868,9 @@ FRayTracingSceneRHIRef FVulkanDynamicRHI::RHICreateRayTracingScene(FRayTracingSc
 	return new FVulkanRayTracingScene(MoveTemp(Initializer), GetDevice());
 }
 
-FRayTracingGeometryRHIRef FVulkanDynamicRHI::RHICreateRayTracingGeometry(const FRayTracingGeometryInitializer& Initializer)
+FRayTracingGeometryRHIRef FVulkanDynamicRHI::RHICreateRayTracingGeometry(FRHICommandListBase& RHICmdList, const FRayTracingGeometryInitializer& Initializer)
 {
-	return new FVulkanRayTracingGeometry(Initializer, GetDevice());
+	return new FVulkanRayTracingGeometry(RHICmdList, Initializer, GetDevice());
 }
 
 void FVulkanCommandListContext::RHIClearRayTracingBindings(FRHIRayTracingScene* Scene)
diff --git a/Engine/Source/Runtime/VulkanRHI/Private/VulkanRayTracing.h b/Engine/Source/Runtime/VulkanRHI/Private/VulkanRayTracing.h
--- a/Engine/Source/Runtime/VulkanRHI/Private/VulkanRayTracing.h
+++ b/Engine/Source/Runtime/VulkanRHI/Private/VulkanRayTracing.h
@@ -61,7 +61,7 @@ class FVulkanRayTracingGeometry : public FRHIRayTracingGeometry
 {
 public:
 	FVulkanRayTracingGeometry(ENoInit);
-	FVulkanRayTracingGeometry(const FRayTracingGeometryInitializer& Initializer, FVulkanDevice* InDevice);
+	FVulkanRayTracingGeometry(FRHICommandListBase& RHICmdList, const FRayTracingGeometryInitializer& Initializer, FVulkanDevice* InDevice);
 	~FVulkanRayTracingGeometry();
 
 	virtual FRayTracingAccelerationStructureAddress GetAccelerationStructureAddress(uint64 GPUIndex) const final override { return Address; }	
diff --git a/Engine/Source/Runtime/VulkanRHI/Private/VulkanResources.h b/Engine/Source/Runtime/VulkanRHI/Private/VulkanResources.h
--- a/Engine/Source/Runtime/VulkanRHI/Private/VulkanResources.h
+++ b/Engine/Source/Runtime/VulkanRHI/Private/VulkanResources.h
@@ -1187,7 +1187,7 @@ public:
 class FVulkanUnorderedAccessView final : public FRHIUnorderedAccessView, public FVulkanLinkedView
 {
 public:
-	FVulkanUnorderedAccessView(FRHICommandListImmediate& RHICmdList, FVulkanDevice& InDevice, FRHIViewableResource* InResource, FRHIViewDesc const& InViewDesc);
+	FVulkanUnorderedAccessView(FRHICommandListBase& RHICmdList, FVulkanDevice& InDevice, FRHIViewableResource* InResource, FRHIViewDesc const& InViewDesc);
 
 	FVulkanViewableResource* GetBaseResource() const;
 	void UpdateView() override;
@@ -1204,7 +1204,7 @@ public:
 class FVulkanShaderResourceView final : public FRHIShaderResourceView, public FVulkanLinkedView
 {
 public:
-	FVulkanShaderResourceView(FRHICommandListImmediate& RHICmdList, FVulkanDevice& InDevice, FRHIViewableResource* InResource, FRHIViewDesc const& InViewDesc);
+	FVulkanShaderResourceView(FRHICommandListBase& RHICmdList, FVulkanDevice& InDevice, FRHIViewableResource* InResource, FRHIViewDesc const& InViewDesc);
 
 	FVulkanViewableResource* GetBaseResource() const;
 	void UpdateView() override;
diff --git a/Engine/Source/Runtime/VulkanRHI/Private/VulkanUAV.cpp b/Engine/Source/Runtime/VulkanRHI/Private/VulkanUAV.cpp
--- a/Engine/Source/Runtime/VulkanRHI/Private/VulkanUAV.cpp
+++ b/Engine/Source/Runtime/VulkanRHI/Private/VulkanUAV.cpp
@@ -372,11 +372,11 @@ static VkDescriptorType GetDescriptorTypeForViewDesc(FRHIViewDesc const& ViewDes
 }
 
 
-FVulkanShaderResourceView::FVulkanShaderResourceView(FRHICommandListImmediate& RHICmdList, FVulkanDevice& InDevice, FRHIViewableResource* InResource, FRHIViewDesc const& InViewDesc)
+FVulkanShaderResourceView::FVulkanShaderResourceView(FRHICommandListBase& RHICmdList, FVulkanDevice& InDevice, FRHIViewableResource* InResource, FRHIViewDesc const& InViewDesc)
 	: FRHIShaderResourceView(InResource, InViewDesc)
 	, FVulkanLinkedView(InDevice, GetDescriptorTypeForViewDesc(InViewDesc))
 {
-	RHICmdList.EnqueueLambda([this](FRHICommandListImmediate&)
+	RHICmdList.EnqueueLambda([this](FRHICommandListBase&)
 	{
 		LinkHead(GetBaseResource()->LinkedViews);
 		UpdateView();
@@ -460,11 +460,11 @@ void FVulkanShaderResourceView::UpdateView()
 
 
 
-FVulkanUnorderedAccessView::FVulkanUnorderedAccessView(FRHICommandListImmediate& RHICmdList, FVulkanDevice& InDevice, FRHIViewableResource* InResource, FRHIViewDesc const& InViewDesc)
+FVulkanUnorderedAccessView::FVulkanUnorderedAccessView(FRHICommandListBase& RHICmdList, FVulkanDevice& InDevice, FRHIViewableResource* InResource, FRHIViewDesc const& InViewDesc)
 	: FRHIUnorderedAccessView(InResource, InViewDesc)
 	, FVulkanLinkedView(InDevice, GetDescriptorTypeForViewDesc(InViewDesc))
 {
-	RHICmdList.EnqueueLambda([this](FRHICommandListImmediate&)
+	RHICmdList.EnqueueLambda([this](FRHICommandListBase&)
 	{
 		LinkHead(GetBaseResource()->LinkedViews);
 		UpdateView();
@@ -670,12 +670,12 @@ void FVulkanCommandListContext::RHIClearUAVUint(FRHIUnorderedAccessView* Unorder
 	ResourceCast(UnorderedAccessViewRHI)->Clear(RHICmdList, &Values, false);
 }
 
-FShaderResourceViewRHIRef  FVulkanDynamicRHI::RHICreateShaderResourceView(class FRHICommandListImmediate& RHICmdList, FRHIViewableResource* Resource, FRHIViewDesc const& ViewDesc)
+FShaderResourceViewRHIRef  FVulkanDynamicRHI::RHICreateShaderResourceView(class FRHICommandListBase& RHICmdList, FRHIViewableResource* Resource, FRHIViewDesc const& ViewDesc)
 {
 	return new FVulkanShaderResourceView(RHICmdList, *Device, Resource, ViewDesc);
 }
 
-FUnorderedAccessViewRHIRef FVulkanDynamicRHI::RHICreateUnorderedAccessView(class FRHICommandListImmediate& RHICmdList, FRHIViewableResource* Resource, FRHIViewDesc const& ViewDesc)
+FUnorderedAccessViewRHIRef FVulkanDynamicRHI::RHICreateUnorderedAccessView(class FRHICommandListBase& RHICmdList, FRHIViewableResource* Resource, FRHIViewDesc const& ViewDesc)
 {
 	return new FVulkanUnorderedAccessView(RHICmdList, *Device, Resource, ViewDesc);
 }
diff --git a/Engine/Source/Runtime/VulkanRHI/Public/VulkanDynamicRHI.h b/Engine/Source/Runtime/VulkanRHI/Public/VulkanDynamicRHI.h
--- a/Engine/Source/Runtime/VulkanRHI/Public/VulkanDynamicRHI.h
+++ b/Engine/Source/Runtime/VulkanRHI/Public/VulkanDynamicRHI.h
@@ -223,15 +223,15 @@ public:
 	}
 
 	// SRV / UAV creation functions
-	virtual FShaderResourceViewRHIRef  RHICreateShaderResourceView (class FRHICommandListImmediate& RHICmdList, FRHIViewableResource* Resource, FRHIViewDesc const& ViewDesc) final override;
-	virtual FUnorderedAccessViewRHIRef RHICreateUnorderedAccessView(class FRHICommandListImmediate& RHICmdList, FRHIViewableResource* Resource, FRHIViewDesc const& ViewDesc) final override;
+	virtual FShaderResourceViewRHIRef  RHICreateShaderResourceView (class FRHICommandListBase& RHICmdList, FRHIViewableResource* Resource, FRHIViewDesc const& ViewDesc) final override;
+	virtual FUnorderedAccessViewRHIRef RHICreateUnorderedAccessView(class FRHICommandListBase& RHICmdList, FRHIViewableResource* Resource, FRHIViewDesc const& ViewDesc) final override;
 
 	void RHICalibrateTimers() override;
 
 #if VULKAN_RHI_RAYTRACING
 	virtual FRayTracingAccelerationStructureSize RHICalcRayTracingSceneSize(uint32 MaxInstances, ERayTracingAccelerationStructureFlags Flags) final override;
 	virtual FRayTracingAccelerationStructureSize RHICalcRayTracingGeometrySize(const FRayTracingGeometryInitializer& Initializer) final override;
-	virtual FRayTracingGeometryRHIRef RHICreateRayTracingGeometry(const FRayTracingGeometryInitializer& Initializer) final override;
+	virtual FRayTracingGeometryRHIRef RHICreateRayTracingGeometry(FRHICommandListBase& RHICmdList, const FRayTracingGeometryInitializer& Initializer) final override;
 	virtual FRayTracingSceneRHIRef RHICreateRayTracingScene(FRayTracingSceneInitializer2 Initializer) final override;
 	virtual FRayTracingShaderRHIRef RHICreateRayTracingShader(TArrayView<const uint8> Code, const FSHAHash& Hash, EShaderFrequency ShaderFrequency) final override;
 	virtual void RHITransferRayTracingGeometryUnderlyingResource(FRHIRayTracingGeometry* DestGeometry, FRHIRayTracingGeometry* SrcGeometry) final override;
diff --git a/Engine/Source/Runtime/Windows/D3D11RHI/Private/D3D11RHIPrivate.h b/Engine/Source/Runtime/Windows/D3D11RHI/Private/D3D11RHIPrivate.h
--- a/Engine/Source/Runtime/Windows/D3D11RHI/Private/D3D11RHIPrivate.h
+++ b/Engine/Source/Runtime/Windows/D3D11RHI/Private/D3D11RHIPrivate.h
@@ -549,8 +549,8 @@ public:
 	virtual void RHISubmitCommandLists(TArrayView<IRHIPlatformCommandList*> CommandLists, bool bFlushResources) final override;
 
 	// SRV / UAV creation functions
-	virtual FShaderResourceViewRHIRef  RHICreateShaderResourceView (class FRHICommandListImmediate& RHICmdList, FRHIViewableResource* Resource, FRHIViewDesc const& ViewDesc) final override;
-	virtual FUnorderedAccessViewRHIRef RHICreateUnorderedAccessView(class FRHICommandListImmediate& RHICmdList, FRHIViewableResource* Resource, FRHIViewDesc const& ViewDesc) final override;
+	virtual FShaderResourceViewRHIRef  RHICreateShaderResourceView (class FRHICommandListBase& RHICmdList, FRHIViewableResource* Resource, FRHIViewDesc const& ViewDesc) final override;
+	virtual FUnorderedAccessViewRHIRef RHICreateUnorderedAccessView(class FRHICommandListBase& RHICmdList, FRHIViewableResource* Resource, FRHIViewDesc const& ViewDesc) final override;
 
 	virtual void RHISetComputeShader(FRHIComputeShader* ComputeShader) final override;
 	virtual void RHIDispatchComputeShader(uint32 ThreadGroupCountX, uint32 ThreadGroupCountY, uint32 ThreadGroupCountZ) final override;
diff --git a/Engine/Source/Runtime/Windows/D3D11RHI/Private/D3D11SRV.cpp b/Engine/Source/Runtime/Windows/D3D11RHI/Private/D3D11SRV.cpp
--- a/Engine/Source/Runtime/Windows/D3D11RHI/Private/D3D11SRV.cpp
+++ b/Engine/Source/Runtime/Windows/D3D11RHI/Private/D3D11SRV.cpp
@@ -16,10 +16,10 @@ void FD3D11ViewableResource::UpdateLinkedViews()
 //
 // -----------------------------------------------------------------------------------------------------
 
-FD3D11ShaderResourceView::FD3D11ShaderResourceView(FRHICommandListImmediate& RHICmdList, FRHIViewableResource* Resource, FRHIViewDesc const& ViewDesc)
+FD3D11ShaderResourceView::FD3D11ShaderResourceView(FRHICommandListBase& RHICmdList, FRHIViewableResource* Resource, FRHIViewDesc const& ViewDesc)
 	: FRHIShaderResourceView(Resource, ViewDesc)
 {
-	RHICmdList.EnqueueLambda([this](FRHICommandListImmediate&)
+	RHICmdList.EnqueueLambda([this](FRHICommandListBase&)
 	{
 		LinkHead(GetBaseResource()->LinkedViews);
 		UpdateView();
@@ -157,7 +157,7 @@ FD3D11ViewableResource* FD3D11ShaderResourceView::GetBaseResource() const
 //
 // -----------------------------------------------------------------------------------------------------
 
-FShaderResourceViewRHIRef FD3D11DynamicRHI::RHICreateShaderResourceView(class FRHICommandListImmediate& RHICmdList, FRHIViewableResource* Resource, FRHIViewDesc const& ViewDesc)
+FShaderResourceViewRHIRef FD3D11DynamicRHI::RHICreateShaderResourceView(class FRHICommandListBase& RHICmdList, FRHIViewableResource* Resource, FRHIViewDesc const& ViewDesc)
 {
 	return new FD3D11ShaderResourceView(RHICmdList, Resource, ViewDesc);
 }
diff --git a/Engine/Source/Runtime/Windows/D3D11RHI/Private/D3D11UAV.cpp b/Engine/Source/Runtime/Windows/D3D11RHI/Private/D3D11UAV.cpp
--- a/Engine/Source/Runtime/Windows/D3D11RHI/Private/D3D11UAV.cpp
+++ b/Engine/Source/Runtime/Windows/D3D11RHI/Private/D3D11UAV.cpp
@@ -9,10 +9,10 @@
 //
 // -----------------------------------------------------------------------------------------------------
 
-FD3D11UnorderedAccessView::FD3D11UnorderedAccessView(FRHICommandListImmediate& RHICmdList, FRHIViewableResource* Resource, FRHIViewDesc const& ViewDesc)
+FD3D11UnorderedAccessView::FD3D11UnorderedAccessView(FRHICommandListBase& RHICmdList, FRHIViewableResource* Resource, FRHIViewDesc const& ViewDesc)
 	: FRHIUnorderedAccessView(Resource, ViewDesc)
 {
-	RHICmdList.EnqueueLambda([this](FRHICommandListImmediate&)
+	RHICmdList.EnqueueLambda([this](FRHICommandListBase&)
 	{
 		LinkHead(GetBaseResource()->LinkedViews);
 		UpdateView();
@@ -127,7 +127,7 @@ FD3D11ViewableResource* FD3D11UnorderedAccessView::GetBaseResource() const
 //
 // -----------------------------------------------------------------------------------------------------
 
-FUnorderedAccessViewRHIRef FD3D11DynamicRHI::RHICreateUnorderedAccessView(class FRHICommandListImmediate& RHICmdList, FRHIViewableResource* Resource, FRHIViewDesc const& ViewDesc)
+FUnorderedAccessViewRHIRef FD3D11DynamicRHI::RHICreateUnorderedAccessView(class FRHICommandListBase& RHICmdList, FRHIViewableResource* Resource, FRHIViewDesc const& ViewDesc)
 {
 	return new FD3D11UnorderedAccessView(RHICmdList, Resource, ViewDesc);
 }
diff --git a/Engine/Source/Runtime/Windows/D3D11RHI/Public/D3D11Resources.h b/Engine/Source/Runtime/Windows/D3D11RHI/Public/D3D11Resources.h
--- a/Engine/Source/Runtime/Windows/D3D11RHI/Public/D3D11Resources.h
+++ b/Engine/Source/Runtime/Windows/D3D11RHI/Public/D3D11Resources.h
@@ -481,7 +481,7 @@ class FD3D11ShaderResourceView final : public FRHIShaderResourceView, public FD3
 public:
 	TRefCountPtr<ID3D11ShaderResourceView> View;
 
-	FD3D11ShaderResourceView(FRHICommandListImmediate& RHICmdList, FRHIViewableResource* Resource, FRHIViewDesc const& ViewDesc);
+	FD3D11ShaderResourceView(FRHICommandListBase& RHICmdList, FRHIViewableResource* Resource, FRHIViewDesc const& ViewDesc);
 	FD3D11ViewableResource* GetBaseResource() const;
 
 	virtual void UpdateView() override;
@@ -493,7 +493,7 @@ class FD3D11UnorderedAccessView final : public FRHIUnorderedAccessView, public F
 public:
 	TRefCountPtr<ID3D11UnorderedAccessView> View;
 
-	FD3D11UnorderedAccessView(FRHICommandListImmediate& RHICmdList, FRHIViewableResource* Resource, FRHIViewDesc const& ViewDesc);
+	FD3D11UnorderedAccessView(FRHICommandListBase& RHICmdList, FRHIViewableResource* Resource, FRHIViewDesc const& ViewDesc);
 	FD3D11ViewableResource* GetBaseResource() const;
 
 	virtual void UpdateView() override;
