Multi-GPU:  UDN request to make it easier to disable multi-GPU.  You get link errors if you set WITH_MGPU to 0, because it leaves GNumExplicitGPUsForRendering as an extern variable in the header, while the code that defines the variable gets commented out.  It's also wasteful to leave MAX_NUM_GPUS as 8, because certain resource structures include arrays based on that value.  Made it so WITH_MGPU is defined first, then everything else is set dependent on that.  Also removed deprecated WITH_SLI setting, only still used in a couple spots.

#rnx

diff --git a/Engine/Source/Runtime/RHI/Private/MultiGPU.cpp b/Engine/Source/Runtime/RHI/Private/MultiGPU.cpp
--- a/Engine/Source/Runtime/RHI/Private/MultiGPU.cpp
+++ b/Engine/Source/Runtime/RHI/Private/MultiGPU.cpp
@@ -6,7 +6,7 @@
 
 #include "MultiGPU.h"
 
-#if WITH_SLI || WITH_MGPU
+#if WITH_MGPU
 uint32 GNumExplicitGPUsForRendering = 1;
 uint32 GVirtualMGPU = 0;
 #endif
diff --git a/Engine/Source/Runtime/RHI/Public/MultiGPU.h b/Engine/Source/Runtime/RHI/Public/MultiGPU.h
--- a/Engine/Source/Runtime/RHI/Public/MultiGPU.h
+++ b/Engine/Source/Runtime/RHI/Public/MultiGPU.h
@@ -11,16 +11,17 @@
 #include "Math/UnrealMathUtility.h"
 #include "Misc/AssertionMacros.h"
 
-/** When greater than one, indicates that SLI rendering is enabled */
 #if PLATFORM_DESKTOP
-#define WITH_SLI 1	// Implicit SLI
 #define WITH_MGPU 1	// Explicit MGPU
+#else
+#define WITH_MGPU 0
+#endif
+
+#if WITH_MGPU
 #define MAX_NUM_GPUS 8
 extern RHI_API uint32 GNumExplicitGPUsForRendering;
 extern RHI_API uint32 GVirtualMGPU;
 #else
-#define WITH_SLI 0
-#define WITH_MGPU 0
 #define MAX_NUM_GPUS 1
 #define GNumExplicitGPUsForRendering 1
 #define GVirtualMGPU 0
diff --git a/Engine/Source/Runtime/Renderer/Private/SceneOcclusion.cpp b/Engine/Source/Runtime/Renderer/Private/SceneOcclusion.cpp
--- a/Engine/Source/Runtime/Renderer/Private/SceneOcclusion.cpp
+++ b/Engine/Source/Runtime/Renderer/Private/SceneOcclusion.cpp
@@ -105,7 +105,7 @@ FOcclusionRandomStream GOcclusionRandomStream;
 int32 FOcclusionQueryHelpers::GetNumBufferedFrames(ERHIFeatureLevel::Type FeatureLevel)
 {
 	int32 NumGPUS = 1;
-#if WITH_SLI || WITH_MGPU
+#if WITH_MGPU
 	// TODO:  Should this still be differentiated for MGPU?  Originally this logic was here for AFR, which has been removed.
 	return FMath::Min<int32>(1, (int32)FOcclusionQueryHelpers::MaxBufferedOcclusionFrames);
 #endif
