Fix compile error due to PreDraw() not being available after merge

diff --git a/Engine/Source/Runtime/RHI/Public/RHICommandList.h b/Engine/Source/Runtime/RHI/Public/RHICommandList.h
--- a/Engine/Source/Runtime/RHI/Public/RHICommandList.h
+++ b/Engine/Source/Runtime/RHI/Public/RHICommandList.h
@@ -3513,7 +3513,6 @@ PRAGMA_ENABLE_DEPRECATION_WARNINGS
 	FORCEINLINE_DEBUGGABLE void MultiDrawIndexedPrimitiveIndirect(FRHIBuffer* IndexBuffer, FRHIBuffer* ArgumentsBuffer, uint32 ArgumentOffset, FRHIBuffer* CountBuffer, uint32 CountBufferOffset, uint32 MaxDrawArguments)
 	{
 		//check(IsOutsideRenderPass());
-		PreDraw();
 		if (Bypass())
 		{
 			GetContext().RHIMultiDrawIndexedPrimitiveIndirect(IndexBuffer, ArgumentsBuffer, ArgumentOffset, CountBuffer, CountBufferOffset, MaxDrawArguments);
