Remove FRHIViewport::NeedFlushBeforeEndDrawing() because it can lead to deadlocks.

FD3D12Viewport::Present() flushes the default context and waits for submission, but that can wait for a payload which has a dependency on a subsequent payload if we don't flush the immediate command list. One way to reach this situation was to have async compute enabled but no async work, in which case the pipe sync transition inserted by RDG was left pending in the immediate command list, so Present deadlocked.

diff --git a/Engine/Source/Runtime/D3D12RHI/Private/D3D12Viewport.h b/Engine/Source/Runtime/D3D12RHI/Private/D3D12Viewport.h
--- a/Engine/Source/Runtime/D3D12RHI/Private/D3D12Viewport.h
+++ b/Engine/Source/Runtime/D3D12RHI/Private/D3D12Viewport.h
@@ -99,8 +99,6 @@ public:
 	virtual void WaitForFrameEventCompletion() override;
 	virtual void IssueFrameEvent() override;
 
-	virtual bool NeedFlushBeforeEndDrawing() override;
-
 #if D3D12_VIEWPORT_EXPOSES_SWAP_CHAIN
 	virtual void* GetNativeSwapChain() const override { return SwapChain1; }
 #endif // #if D3D12_VIEWPORT_EXPOSES_SWAP_CHAIN
diff --git a/Engine/Source/Runtime/D3D12RHI/Private/Windows/WindowsD3D12Viewport.cpp b/Engine/Source/Runtime/D3D12RHI/Private/Windows/WindowsD3D12Viewport.cpp
--- a/Engine/Source/Runtime/D3D12RHI/Private/Windows/WindowsD3D12Viewport.cpp
+++ b/Engine/Source/Runtime/D3D12RHI/Private/Windows/WindowsD3D12Viewport.cpp
@@ -537,15 +537,6 @@ bool FD3D12Viewport::IsPresentAllowed()
 	return !FD3D12DynamicRHI::GetD3DRHI()->RHIIsRenderingSuspended();
 }
 
-/**
-  * RHIEndDrawingViewport flushes commands internally, so we don't need to flush before calling it.
-  * Avoids wasteful issuing of another entire command list with just a few commands.
-  */
-bool FD3D12Viewport::NeedFlushBeforeEndDrawing()
-{
-	return false;
-}
-
 void FD3D12DynamicRHI::RHIGetDisplaysInformation(FDisplayInformationArray& OutDisplayInformation)
 {
 	OutDisplayInformation.Append(DisplayList);
diff --git a/Engine/Source/Runtime/RHI/Private/RHICommandList.cpp b/Engine/Source/Runtime/RHI/Private/RHICommandList.cpp
--- a/Engine/Source/Runtime/RHI/Private/RHICommandList.cpp
+++ b/Engine/Source/Runtime/RHI/Private/RHICommandList.cpp
@@ -1281,10 +1281,7 @@ void FRHICommandListImmediate::EndDrawingViewport(FRHIViewport* Viewport, bool b
 	// Make sure all prior graphics and async compute work has been submitted.
 	// This is necessary because platform RHIs often submit additional work on the graphics queue during present, and we need to ensure we won't deadlock on async work that wasn't yet submitted by the renderer.
 	// In future, Present() itself should be an enqueued / recorded command, and platform RHIs should never implicitly submit graphics or async compute work.
-	if (Viewport->NeedFlushBeforeEndDrawing())
-	{
-		FRHICommandListExecutor::GetImmediateCommandList().ImmediateFlush(EImmediateFlushType::DispatchToRHIThread);
-	}
+	FRHICommandListExecutor::GetImmediateCommandList().ImmediateFlush(EImmediateFlushType::DispatchToRHIThread);
 
 	check(IsImmediate() && IsInRenderingThread());
 	if (Bypass())
diff --git a/Engine/Source/Runtime/RHI/Public/RHIResources.h b/Engine/Source/Runtime/RHI/Public/RHIResources.h
--- a/Engine/Source/Runtime/RHI/Public/RHIResources.h
+++ b/Engine/Source/Runtime/RHI/Public/RHIResources.h
@@ -2242,8 +2242,6 @@ public:
 	virtual void WaitForFrameEventCompletion() { }
 
 	virtual void IssueFrameEvent() { }
-
-	virtual bool NeedFlushBeforeEndDrawing() { return true; }
 };
 
 /** Used to specify a texture metadata plane when creating a view. */
