Moving a couple DDSPI fields to Editor only code blocks.
Moving Editor only DDSPI types into the cpp files to keep Editor includes out of the global scope.

diff --git a/Engine/Source/Runtime/RHI/Private/DataDrivenShaderPlatformInfo.cpp b/Engine/Source/Runtime/RHI/Private/DataDrivenShaderPlatformInfo.cpp
--- a/Engine/Source/Runtime/RHI/Private/DataDrivenShaderPlatformInfo.cpp
+++ b/Engine/Source/Runtime/RHI/Private/DataDrivenShaderPlatformInfo.cpp
@@ -1,12 +1,16 @@
 // Copyright Epic Games, Inc. All Rights Reserved.
 
 #include "DataDrivenShaderPlatformInfo.h"
-#include "RHI.h"
 #include "Misc/ConfigCacheIni.h"
 #include "Misc/DataDrivenPlatformInfoRegistry.h"
 #include "Misc/CommandLine.h"
+#include "RHI.h"
 #include "RHIStrings.h"
 
+#if WITH_EDITOR
+#include "Internationalization/Text.h"
+#endif
+
 const FName LANGUAGE_D3D("D3D");
 const FName LANGUAGE_Metal("Metal");
 const FName LANGUAGE_OpenGL("OpenGL");
@@ -15,11 +19,18 @@ const FName LANGUAGE_Sony("Sony");
 const FName LANGUAGE_Nintendo("Nintendo");
 
 FGenericDataDrivenShaderPlatformInfo FGenericDataDrivenShaderPlatformInfo::Infos[SP_NumPlatforms];
+static TMap<FName, EShaderPlatform> PlatformNameToShaderPlatformMap;
 
 #if WITH_EDITOR
-TMap < FString, TFunction<bool(const FStaticShaderPlatform Platform)>> FGenericDataDrivenShaderPlatformInfo::PropertyToShaderPlatformFunctionMap;
+struct FDataDrivenShaderPlatformInfoEditorOnly
+{
+	FText FriendlyName;
+	EShaderPlatform PreviewShaderPlatformParent = SP_NumPlatforms;
+};
+static FDataDrivenShaderPlatformInfoEditorOnly DataDrivenShaderPlatformInfoEditorOnlyInfos[SP_NumPlatforms];
+
+TMap<FString, TFunction<bool(const FStaticShaderPlatform Platform)>> FGenericDataDrivenShaderPlatformInfo::PropertyToShaderPlatformFunctionMap;
 #endif
-TMap<FName, EShaderPlatform> PlatformNameToShaderPlatformMap;
 
 EShaderPlatform ParseShaderPlatform(const TCHAR* String)
 {
@@ -134,11 +145,12 @@ void FGenericDataDrivenShaderPlatformInfo::SetDefaultValues()
 	bSupportsManualVertexFetch = true;
 	bSupportsVolumeTextureAtomics = true;
 	bSupportsClipDistance = true;
-	PreviewShaderPlatformParent = EShaderPlatform::SP_NumPlatforms;
 }
 
-void FGenericDataDrivenShaderPlatformInfo::ParseDataDrivenShaderInfo(const FConfigSection& Section, FGenericDataDrivenShaderPlatformInfo& Info)
+void FGenericDataDrivenShaderPlatformInfo::ParseDataDrivenShaderInfo(const FConfigSection& Section, uint32 Index)
 {
+	FGenericDataDrivenShaderPlatformInfo& Info = Infos[Index];
+
 	Info.Language = *GetSectionString(Section, "Language");
 	Info.ShaderFormat = *GetSectionString(Section, "ShaderFormat");
 	checkf(!Info.ShaderFormat.IsNone(), TEXT("Missing ShaderFormat for ShaderPlatform %s  ShaderFormat %s"), *Info.Name.ToString(), *Info.ShaderFormat.ToString());
@@ -277,8 +289,10 @@ void FGenericDataDrivenShaderPlatformInfo::ParseDataDrivenShaderInfo(const FConf
 #undef ADD_PROPERTY_TO_SHADERPLATFORM_FUNCTIONMAP
 
 	Info.ShaderPropertiesHash = GetTypeHash(ShaderPropertiesString);
+
 #if WITH_EDITOR
-	FTextStringHelper::ReadFromBuffer(*GetSectionString(Section, FName("FriendlyName")), Info.FriendlyName);
+	FDataDrivenShaderPlatformInfoEditorOnly& EditorInfo = DataDrivenShaderPlatformInfoEditorOnlyInfos[Index];
+	FTextStringHelper::ReadFromBuffer(*GetSectionString(Section, FName("FriendlyName")), EditorInfo.FriendlyName);
 #endif
 }
 
@@ -331,7 +345,7 @@ void FGenericDataDrivenShaderPlatformInfo::Initialize()
 				// at this point, we can start pulling information out
 				Infos[ShaderPlatform].Name = *SectionName.Mid(15);
 				PlatformNameToShaderPlatformMap.FindOrAdd(Infos[ShaderPlatform].Name) = ShaderPlatform;
-				ParseDataDrivenShaderInfo(Section.Value, Infos[ShaderPlatform]);
+				ParseDataDrivenShaderInfo(Section.Value, ShaderPlatform);
 				Infos[ShaderPlatform].bContainsValidPlatformInfo = true;
 
 #if WITH_EDITOR
@@ -339,21 +353,25 @@ void FGenericDataDrivenShaderPlatformInfo::Initialize()
 				{
 					for (const FPreviewPlatformMenuItem& Item : FDataDrivenPlatformInfoRegistry::GetAllPreviewPlatformMenuItems())
 					{
-						FName PreviewPlatformName = *(Infos[ShaderPlatform].Name).ToString();
+						const FName PreviewPlatformName = *(Infos[ShaderPlatform].Name).ToString();
 						if (Item.ShaderPlatformToPreview == PreviewPlatformName)
 						{
-							EShaderPlatform PreviewShaderPlatform = EShaderPlatform(CustomShaderPlatform++);
+							const EShaderPlatform PreviewShaderPlatform = EShaderPlatform(CustomShaderPlatform++);
+							ParseDataDrivenShaderInfo(Section.Value, PreviewShaderPlatform);
+
+							FGenericDataDrivenShaderPlatformInfo& PreviewInfo = Infos[PreviewShaderPlatform];
+							PreviewInfo.Name = Item.PreviewShaderPlatformName;
+							PreviewInfo.bIsPreviewPlatform = true;
+							PreviewInfo.bContainsValidPlatformInfo = true;
 
-							ParseDataDrivenShaderInfo(Section.Value, Infos[PreviewShaderPlatform]);
+							FDataDrivenShaderPlatformInfoEditorOnly& PreviewEditorInfo = DataDrivenShaderPlatformInfoEditorOnlyInfos[PreviewShaderPlatform];
+							PreviewEditorInfo.PreviewShaderPlatformParent = ShaderPlatform;
 							if (!Item.OptionalFriendlyNameOverride.IsEmpty())
 							{
-								Infos[PreviewShaderPlatform].FriendlyName = Item.OptionalFriendlyNameOverride;
+								PreviewEditorInfo.FriendlyName = Item.OptionalFriendlyNameOverride;
 							}
-							Infos[PreviewShaderPlatform].Name = Item.PreviewShaderPlatformName;
-							PlatformNameToShaderPlatformMap.FindOrAdd(Infos[PreviewShaderPlatform].Name) = PreviewShaderPlatform;
-							Infos[PreviewShaderPlatform].PreviewShaderPlatformParent = ShaderPlatform;
-							Infos[PreviewShaderPlatform].bIsPreviewPlatform = true;
-							Infos[PreviewShaderPlatform].bContainsValidPlatformInfo = true;
+
+							PlatformNameToShaderPlatformMap.FindOrAdd(PreviewInfo.Name) = PreviewShaderPlatform;
 						}
 					}
 				}
@@ -423,7 +441,13 @@ FText FGenericDataDrivenShaderPlatformInfo::GetFriendlyName(const FStaticShaderP
 		return FText();
 	}
 	check(IsValid(Platform));
-	return Infos[Platform].FriendlyName;
+	return DataDrivenShaderPlatformInfoEditorOnlyInfos[Platform].FriendlyName;
+}
+
+const EShaderPlatform FGenericDataDrivenShaderPlatformInfo::GetPreviewShaderPlatformParent(const FStaticShaderPlatform Platform)
+{
+	check(IsValid(Platform));
+	return DataDrivenShaderPlatformInfoEditorOnlyInfos[Platform].PreviewShaderPlatformParent;
 }
 #endif
 
diff --git a/Engine/Source/Runtime/RHI/Public/DataDrivenShaderPlatformInfo.h b/Engine/Source/Runtime/RHI/Public/DataDrivenShaderPlatformInfo.h
--- a/Engine/Source/Runtime/RHI/Public/DataDrivenShaderPlatformInfo.h
+++ b/Engine/Source/Runtime/RHI/Public/DataDrivenShaderPlatformInfo.h
@@ -2,11 +2,12 @@
 
 #pragma once
 
-#include "Internationalization/Text.h"
 #include "RHIDefinitions.h"
 #include "RHIShaderPlatform.h"
 #include "RHIFeatureLevel.h"
 
+class FText;
+
 extern RHI_API const FName LANGUAGE_D3D;
 extern RHI_API const FName LANGUAGE_Metal;
 extern RHI_API const FName LANGUAGE_OpenGL;
@@ -20,7 +21,6 @@ class FGenericDataDrivenShaderPlatformInfo
 	FName Language;
 	ERHIFeatureLevel::Type MaxFeatureLevel;
 	FName ShaderFormat;
-	EShaderPlatform PreviewShaderPlatformParent;
 	uint32 ShaderPropertiesHash;
 	uint32 bIsMobile : 1;
 	uint32 bIsMetalMRT : 1;
@@ -112,10 +112,6 @@ class FGenericDataDrivenShaderPlatformInfo
 	uint32 bSupportsClipDistance : 1;
 	uint32 bSupportsNNEShaders: 1;
 
-#if WITH_EDITOR
-	FText FriendlyName;
-#endif
-
 	// NOTE: When adding fields, you must also add to ParseDataDrivenShaderInfo!
 	uint32 bContainsValidPlatformInfo : 1;
 
@@ -131,7 +127,6 @@ class FGenericDataDrivenShaderPlatformInfo
 public:
 	RHI_API static void Initialize();
 	RHI_API static void UpdatePreviewPlatforms();
-	RHI_API static void ParseDataDrivenShaderInfo(const FConfigSection& Section, FGenericDataDrivenShaderPlatformInfo& Info);
 	RHI_API static const EShaderPlatform GetShaderPlatformFromName(const FName ShaderPlatformName);
 
 	static FORCEINLINE_DEBUGGABLE const FName GetName(const FStaticShaderPlatform Platform)
@@ -146,12 +141,6 @@ public:
 		return Infos[Platform].ShaderFormat;
 	}
 
-	static FORCEINLINE_DEBUGGABLE const EShaderPlatform GetPreviewShaderPlatformParent(const FStaticShaderPlatform Platform)
-	{
-		check(IsValid(Platform));
-		return Infos[Platform].PreviewShaderPlatformParent;
-	}
-
 	static FORCEINLINE_DEBUGGABLE uint32 GetShaderPlatformPropertiesHash(const FStaticShaderPlatform Platform)
 	{
 		check(IsValid(Platform));
@@ -735,23 +724,22 @@ public:
 		return Infos[Platform].bSupportsNNEShaders;
 	}
 
+	static FORCEINLINE_DEBUGGABLE const bool IsValid(const FStaticShaderPlatform Platform)
+	{
+		return Infos[Platform].bContainsValidPlatformInfo;
+	}
+
 #if WITH_EDITOR
 	RHI_API static FText GetFriendlyName(const FStaticShaderPlatform Platform);
+	RHI_API static const EShaderPlatform GetPreviewShaderPlatformParent(const FStaticShaderPlatform Platform);
+
+	RHI_API static TMap<FString, TFunction<bool(const FStaticShaderPlatform Platform)>> PropertyToShaderPlatformFunctionMap;
 #endif
 
 private:
-	RHI_API static FGenericDataDrivenShaderPlatformInfo Infos[SP_NumPlatforms];
-
-public:
+	RHI_API static void ParseDataDrivenShaderInfo(const FConfigSection& Section, uint32 Index);
 
-#if WITH_EDITOR
-	RHI_API static TMap < FString, TFunction<bool(const FStaticShaderPlatform Platform)>> PropertyToShaderPlatformFunctionMap;
-#endif
-
-	static bool IsValid(const FStaticShaderPlatform Platform)
-	{
-		return Infos[Platform].bContainsValidPlatformInfo;
-	}
+	RHI_API static FGenericDataDrivenShaderPlatformInfo Infos[SP_NumPlatforms];
 };
 
 #if USE_STATIC_SHADER_PLATFORM_ENUMS || USE_STATIC_FEATURE_LEVEL_ENUMS || USE_STATIC_SHADER_PLATFORM_INFO
@@ -883,6 +871,7 @@ inline bool IsSimulatedPlatform(const FStaticShaderPlatform Platform)
 	return FDataDrivenShaderPlatformInfo::GetIsPreviewPlatform(Platform);
 }
 
+#if WITH_EDITOR
 inline EShaderPlatform GetSimulatedPlatform(FStaticShaderPlatform Platform)
 {
 	if (IsSimulatedPlatform(Platform))
@@ -892,7 +881,7 @@ inline EShaderPlatform GetSimulatedPlatform(FStaticShaderPlatform Platform)
 
 	return Platform;
 }
-
+#endif // WITH_EDITOR
 
 /** Returns true if the feature level is supported by the shader platform. */
 inline bool IsFeatureLevelSupported(const FStaticShaderPlatform InShaderPlatform, ERHIFeatureLevel::Type InFeatureLevel)
@@ -1113,3 +1102,6 @@ inline bool RHISupportsVolumeTextureAtomics(EShaderPlatform Platform)
 	return FDataDrivenShaderPlatformInfo::GetSupportsVolumeTextureAtomics(Platform);
 }
 
+#if UE_ENABLE_INCLUDE_ORDER_DEPRECATED_IN_5_3
+#include "Internationalization/Text.h"
+#endif
diff --git a/Engine/Source/Runtime/RenderCore/Private/RenderUtils.cpp b/Engine/Source/Runtime/RenderCore/Private/RenderUtils.cpp
--- a/Engine/Source/Runtime/RenderCore/Private/RenderUtils.cpp
+++ b/Engine/Source/Runtime/RenderCore/Private/RenderUtils.cpp
@@ -688,6 +688,7 @@ EShaderPlatform GetEditorShaderPlatform(EShaderPlatform ShaderPlatform)
 {
 	EShaderPlatform ActualShaderPlatform = ShaderPlatform;
 
+#if WITH_EDITOR
 	if (GIsEditor)
 	{
 		if (FDataDrivenShaderPlatformInfo::GetIsPreviewPlatform(ActualShaderPlatform))
@@ -695,6 +696,7 @@ EShaderPlatform GetEditorShaderPlatform(EShaderPlatform ShaderPlatform)
 			ActualShaderPlatform = FDataDrivenShaderPlatformInfo::GetPreviewShaderPlatformParent(ActualShaderPlatform);
 		}
 	}
+#endif // WITH_EDITOR
 
 	return ActualShaderPlatform;
 }
@@ -743,6 +745,7 @@ RENDERCORE_API ShaderPlatformMaskType GRayTracingPlatformMask;
 // Value may be queried using IsRayTracingEnabled().
 RENDERCORE_API ERayTracingMode GRayTracingMode = ERayTracingMode::Disabled;
 
+#if WITH_EDITOR
 void GetAllPossiblePreviewPlatformsForMainShaderPlatform(TArray<EShaderPlatform>& OutPreviewPlatforms, EShaderPlatform ParentShaderPlatform)
 {
 	for (int i = 0; i < EShaderPlatform::SP_NumPlatforms; ++i)
@@ -758,6 +761,7 @@ void GetAllPossiblePreviewPlatformsForMainShaderPlatform(TArray<EShaderPlatform>
 		}
 	}
 }
+#endif // WITH_EDITOR
 
 RENDERCORE_API void RenderUtilsInit()
 {
