Improve clarity on the pipeline file cache's interactions with the PSO precaching system:
* Do not emit a CSV event and do not log a graphics PSO as new if it was precached at runtime. Instead use a more verbose log category and a more helpful log message. Currently requires r.PSOPrecache.Validation != 0 because the PSO precache state must be know for this to work, otherwise the current behavior is preserved.
* Add a CVar to optionally skip saving PSOs in the file cache if they've been precached at runtime. Defaults to disabled (same as current behavior).

Compute PSOs are harder to track in terms of precaching state so they're not yet handled.

#pf 6477a1b6313d43282fb1ea1c

diff --git a/Engine/Source/Runtime/RHI/Private/PipelineFileCache.cpp b/Engine/Source/Runtime/RHI/Private/PipelineFileCache.cpp
--- a/Engine/Source/Runtime/RHI/Private/PipelineFileCache.cpp
+++ b/Engine/Source/Runtime/RHI/Private/PipelineFileCache.cpp
@@ -113,6 +113,14 @@ static TAutoConsoleVariable<int32> CVarPSOFileCacheReportPSO(
 														   ECVF_Default | ECVF_RenderThreadSafe
 														   );
 
+static int32 GPSOExcludePrecachePSOsInFileCache = 0;
+static FAutoConsoleVariableRef CVarPSOFileCacheExcludePrecachePSO(
+														   TEXT("r.ShaderPipelineCache.ExcludePrecachePSO"),
+														   GPSOExcludePrecachePSOsInFileCache,
+														   TEXT("1 excludes saving runtime-precached graphics PSOs in the file cache, 0 (default) includes them. Excluding precached PSOs currently requires r.PSOPrecaching = 1 and r.PSOPrecache.Validation != 0."),
+														   ECVF_ReadOnly
+														   );
+
 static int32 GPSOFileCachePrintNewPSODescriptors = 0;
 static FAutoConsoleVariableRef CVarPSOFileCachePrintNewPSODescriptors(
 														   TEXT("r.ShaderPipelineCache.PrintNewPSODescriptors"),
@@ -3366,7 +3374,7 @@ void FPipelineFileCacheManager::RegisterPSOUsageDataUpdateForNextSave(FPSOUsageD
 	CurrentEntry.EngineFlags |= UsageData.EngineFlags;
 }  
 
-void FPipelineFileCacheManager::CacheGraphicsPSO(uint32 RunTimeHash, FGraphicsPipelineStateInitializer const& Initializer)
+void FPipelineFileCacheManager::CacheGraphicsPSO(uint32 RunTimeHash, FGraphicsPipelineStateInitializer const& Initializer, bool bWasPSOPrecached)
 {
 	if(IsPipelineFileCacheEnabled() && (LogPSOtoFileCache() || ReportNewPSOs()))
 	{
@@ -3386,7 +3394,7 @@ void FPipelineFileCacheManager::CacheGraphicsPSO(uint32 RunTimeHash, FGraphicsPi
 					
 				uint32 PSOHash = GetTypeHash(NewEntry);
 				FPSOUsageData CurrentUsageData(PSOHash, 0, 0);
-					
+
 				if (!FPipelineFileCacheManager::IsPSOEntryCached(NewEntry, &CurrentUsageData))
 				{
 					bool bActuallyNewPSO = !NewPSOHashes.Contains(PSOHash);
@@ -3401,19 +3409,32 @@ void FPipelineFileCacheManager::CacheGraphicsPSO(uint32 RunTimeHash, FGraphicsPi
 					{
 						bActuallyNewPSO = !FPipelineFileCacheManager::IsBSSEquivalentPSOEntryCached(NewEntry);
 					}
+
 					if (bActuallyNewPSO)
 					{
-						CSV_EVENT(PSO, TEXT("Encountered new graphics PSO"));
-						UE_LOG(LogRHI, Display, TEXT("Encountered a new graphics PSO: %u"), PSOHash);
-						int32 LogDetailLevel = LogPSODetails() ? 2 : GPSOFileCachePrintNewPSODescriptors;
-						if (LogDetailLevel > 0)
+						if (!bWasPSOPrecached)
 						{
-							UE_LOG(LogRHI, Display, TEXT("New Graphics PSO (%u)"), PSOHash);
-							if (LogDetailLevel > 1)
+							CSV_EVENT(PSO, TEXT("Encountered new graphics PSO"));
+							UE_LOG(LogRHI, Display, TEXT("Encountered a new graphics PSO: %u"), PSOHash);
+							int32 LogDetailLevel = LogPSODetails() ? 2 : GPSOFileCachePrintNewPSODescriptors;
+							if (LogDetailLevel > 0)
 							{
-								UE_LOG(LogRHI, Display, TEXT("%s"), *NewEntry.ToStringReadable());
+								UE_LOG(LogRHI, Display, TEXT("New Graphics PSO (%u)"), PSOHash);
+								if (LogDetailLevel > 1)
+								{
+									UE_LOG(LogRHI, Display, TEXT("%s"), *NewEntry.ToStringReadable());
+								}
 							}
 						}
+						else
+						{
+							UE_LOG(LogRHI, Log, TEXT("Encountered a new graphics PSO for the file cache but it was already precached at runtime: %u"), PSOHash);
+							bActuallyNewPSO = !GPSOExcludePrecachePSOsInFileCache;
+						}
+					}
+
+					if (bActuallyNewPSO) 
+					{
 						if (LogPSOtoFileCache())
 						{
 							NewPSOs.Add(NewEntry);
diff --git a/Engine/Source/Runtime/RHI/Private/PipelineStateCache.cpp b/Engine/Source/Runtime/RHI/Private/PipelineStateCache.cpp
--- a/Engine/Source/Runtime/RHI/Private/PipelineStateCache.cpp
+++ b/Engine/Source/Runtime/RHI/Private/PipelineStateCache.cpp
@@ -2608,7 +2608,9 @@ FGraphicsPipelineState* PipelineStateCache::GetAndOrCreateGraphicsPipelineState(
 
 	if (bWasFound == false)
 	{
-		FPipelineFileCacheManager::CacheGraphicsPSO(GetTypeHash(Initializer), Initializer);
+		bool bWasPSOPrecached = PSOPrecacheResult == EPSOPrecacheResult::Active || PSOPrecacheResult == EPSOPrecacheResult::Complete;
+
+		FPipelineFileCacheManager::CacheGraphicsPSO(GetTypeHash(Initializer), Initializer, bWasPSOPrecached);
 
 		// create new graphics state
 		OutCachedState = new FGraphicsPipelineState();
diff --git a/Engine/Source/Runtime/RHI/Public/PipelineFileCache.h b/Engine/Source/Runtime/RHI/Public/PipelineFileCache.h
--- a/Engine/Source/Runtime/RHI/Public/PipelineFileCache.h
+++ b/Engine/Source/Runtime/RHI/Public/PipelineFileCache.h
@@ -354,7 +354,7 @@ public:
 
 	RHI_API static void CloseUserPipelineFileCache();
 
-	RHI_API static void CacheGraphicsPSO(uint32 RunTimeHash, FGraphicsPipelineStateInitializer const& Initializer);
+	RHI_API static void CacheGraphicsPSO(uint32 RunTimeHash, FGraphicsPipelineStateInitializer const& Initializer, bool bWasPSOPrecached);
 	RHI_API static void CacheComputePSO(uint32 RunTimeHash, FRHIComputeShader const* Initializer);
 	RHI_API static void CacheRayTracingPSO(const FRayTracingPipelineStateInitializer& Initializer, ERayTracingPipelineCacheFlags Flags);
 
