Fix a thread deadlock bug with RDG parallel execute enabled on mobile vulkan. A RHICmdListPass is allocated in FRDGBuilder.DispatchParallelExecute to record RHI commands on a task thread and there are some FCompilePipelineStateTask generated during recording and needs to wait (RHICmdList.AddDispatchPrerequisite(OutCachedState->CompletionEvent) => DispatchEvent->DontCompleteUntil(Prereq)), then the RHICmdListPass is enqueued to RHIThread and waits for the DispatchEvent on RHIThread (CmdList->WaitForDispatchEvent()), the problem is the FCompilePipelineStateTask could be forced to run on RHIThread (r.pso.CreateOnRHIThread is 1) and could be executed after the CmdList->WaitForDispatchEvent(), so the RHIThread would be stalled.

diff --git a/Engine/Source/Runtime/RHI/Private/PipelineStateCache.cpp b/Engine/Source/Runtime/RHI/Private/PipelineStateCache.cpp
--- a/Engine/Source/Runtime/RHI/Private/PipelineStateCache.cpp
+++ b/Engine/Source/Runtime/RHI/Private/PipelineStateCache.cpp
@@ -1692,12 +1692,14 @@ public:
 	FPipelineState* Pipeline;
 	FGraphicsPipelineStateInitializer Initializer;
 	EPSOPrecacheResult PSOPreCacheResult;
+	bool bInImmediateCmdList;
 
 	// InInitializer is only used for non-compute tasks, a default can just be used otherwise
-	FCompilePipelineStateTask(FPipelineState* InPipeline, const FGraphicsPipelineStateInitializer& InInitializer, EPSOPrecacheResult InPSOPreCacheResult)
+	FCompilePipelineStateTask(FPipelineState* InPipeline, const FGraphicsPipelineStateInitializer& InInitializer, EPSOPrecacheResult InPSOPreCacheResult, bool InbInImmediateCmdList)
 		: Pipeline(InPipeline)
 		, Initializer(InInitializer)
 		, PSOPreCacheResult(InPSOPreCacheResult)
+		, bInImmediateCmdList(InbInImmediateCmdList)
 	{
 		ensure(Pipeline->CompletionEvent != nullptr);
 		if(Initializer.bFromPSOFileCache)
@@ -1925,7 +1927,7 @@ public:
 
 	ENamedThreads::Type GetDesiredThread()
 	{
-		ENamedThreads::Type DesiredThread = GRunPSOCreateTasksOnRHIT && IsRunningRHIInSeparateThread() ? ENamedThreads::RHIThread : CPrio_FCompilePipelineStateTask.Get();
+		ENamedThreads::Type DesiredThread = GRunPSOCreateTasksOnRHIT && IsRunningRHIInSeparateThread() && bInImmediateCmdList ? ENamedThreads::RHIThread : CPrio_FCompilePipelineStateTask.Get();
 
 		// On Mac the compilation is handled using external processes, so engine threads have very little work to do
 		// and it's better to leave more CPU time to these extrenal processes and other engine threads.
@@ -2082,7 +2084,7 @@ FComputePipelineState* PipelineStateCache::GetAndOrCreateComputePipelineState(FR
 			OutCachedState->CompletionEvent = FGraphEvent::CreateGraphEvent();
 			FGraphicsPipelineStateInitializer GraphicsPipelineStateInitializer;
 			GraphicsPipelineStateInitializer.bFromPSOFileCache = bFromFileCache;
-			TGraphTask<FCompilePipelineStateTask>::CreateTask().ConstructAndDispatchWhenReady(OutCachedState, GraphicsPipelineStateInitializer, EPSOPrecacheResult::Untracked);
+			TGraphTask<FCompilePipelineStateTask>::CreateTask().ConstructAndDispatchWhenReady(OutCachedState, GraphicsPipelineStateInitializer, EPSOPrecacheResult::Untracked, RHICmdList.IsImmediate());
 			if (!bFromFileCache)
 			{
 				RHICmdList.AddDispatchPrerequisite(OutCachedState->CompletionEvent);
@@ -2496,7 +2498,7 @@ inline void ValidateGraphicsPipelineStateInitializer(const FGraphicsPipelineStat
 	check(Initializer.DepthStencilState && Initializer.BlendState && Initializer.RasterizerState);
 }
 
-static void InternalCreateGraphicsPipelineState(const FGraphicsPipelineStateInitializer& Initializer, EPSOPrecacheResult PSOPrecacheResult, bool bDoAsyncCompile, bool bPSOPrecache, FGraphicsPipelineState* CachedState)
+static void InternalCreateGraphicsPipelineState(const FGraphicsPipelineStateInitializer& Initializer, EPSOPrecacheResult PSOPrecacheResult, bool bDoAsyncCompile, bool bPSOPrecache, FGraphicsPipelineState* CachedState, bool bInImmediateCmdList = true)
 {
 	FGraphEventRef GraphEvent = CachedState->CompletionEvent;
 
@@ -2507,13 +2509,13 @@ static void InternalCreateGraphicsPipelineState(const FGraphicsPipelineStateInit
 		// Use normal task graph for non-precompile jobs (or when thread pool is not enabled)
 		if (!bPSOPrecache || !FPSOPrecacheThreadPool::UsePool())
 		{
-			TGraphTask<FCompilePipelineStateTask>::CreateTask().ConstructAndDispatchWhenReady(CachedState, Initializer, PSOPrecacheResult);
+			TGraphTask<FCompilePipelineStateTask>::CreateTask().ConstructAndDispatchWhenReady(CachedState, Initializer, PSOPrecacheResult, bInImmediateCmdList);
 		}
 		else
 		{
 			// Here, PSO precompiles use a separate thread pool.
 			// Note that we do not add precompile tasks as cmdlist prerequisites.
-			TUniquePtr<FCompilePipelineStateTask> ThreadPoolTask = MakeUnique<FCompilePipelineStateTask>(CachedState, Initializer, PSOPrecacheResult);
+			TUniquePtr<FCompilePipelineStateTask> ThreadPoolTask = MakeUnique<FCompilePipelineStateTask>(CachedState, Initializer, PSOPrecacheResult, bInImmediateCmdList);
 			uint64 StartTime = FPlatformTime::Cycles64();
 #if	PSO_TRACK_CACHE_STATS
 			StatsStartPrecompile();
@@ -2633,7 +2635,7 @@ FGraphicsPipelineState* PipelineStateCache::GetAndOrCreateGraphicsPipelineState(
 
 		bool bPSOPrecache = Initializer.bFromPSOFileCache;
 		FGraphEventRef GraphEvent = OutCachedState->CompletionEvent;
-		InternalCreateGraphicsPipelineState(Initializer, PSOPrecacheResult, DoAsyncCompile, bPSOPrecache, OutCachedState);
+		InternalCreateGraphicsPipelineState(Initializer, PSOPrecacheResult, DoAsyncCompile, bPSOPrecache, OutCachedState, RHICmdList.IsImmediate());
 
 		// Add dispatch pre requisite for non precaching jobs only
 		//if (GraphEvent.IsValid() && (!bPSOPrecache || !FPSOPrecacheThreadPool::UsePool()))
@@ -2718,13 +2720,13 @@ void FPrecacheComputePipelineCache::OnNewPipelineStateCreated(const FPrecacheCom
 		// Thread pool disabled?
 		if (!FPSOPrecacheThreadPool::UsePool())
 		{
-			TGraphTask<FCompilePipelineStateTask>::CreateTask().ConstructAndDispatchWhenReady(CachedState, GraphicsPipelineStateInitializer, EPSOPrecacheResult::Active);
+			TGraphTask<FCompilePipelineStateTask>::CreateTask().ConstructAndDispatchWhenReady(CachedState, GraphicsPipelineStateInitializer, EPSOPrecacheResult::Active, false);
 		}
 		else
 		{
 			// Here, PSO precompiles use a separate thread pool.
 			// Note that we do not add precompile tasks as cmdlist prerequisites.
-			TUniquePtr<FCompilePipelineStateTask> ThreadPoolTask = MakeUnique<FCompilePipelineStateTask>(CachedState, GraphicsPipelineStateInitializer, EPSOPrecacheResult::Active);
+			TUniquePtr<FCompilePipelineStateTask> ThreadPoolTask = MakeUnique<FCompilePipelineStateTask>(CachedState, GraphicsPipelineStateInitializer, EPSOPrecacheResult::Active, false);
 			uint64 StartTime = FPlatformTime::Cycles64();
 #if	PSO_TRACK_CACHE_STATS
 			StatsStartPrecompile();
