Workaround for deadlock when testing Dx12 feature: assume all cards can run Dx12, until we can fix the deadlock on certain hardware.
Also fixed race condition: always wait for async PSO task if present.

#tests booting on Dx11 and on Dx12 and getting in the FrontEnd. Checked settings. QA pass by [REDACTED_HANDLE]
#localization none
#rnx

diff --git a/Engine/Source/Runtime/RHI/Private/PipelineStateCache.cpp b/Engine/Source/Runtime/RHI/Private/PipelineStateCache.cpp
--- a/Engine/Source/Runtime/RHI/Private/PipelineStateCache.cpp
+++ b/Engine/Source/Runtime/RHI/Private/PipelineStateCache.cpp
@@ -2106,12 +2106,9 @@ FComputePipelineState* PipelineStateCache::GetAndOrCreateComputePipelineState(FR
 	}
 	else
 	{
-		if (DoAsyncCompile)
+		if (!bFromFileCache && !OutCachedState->IsComplete())
 		{
-			if(!OutCachedState->IsComplete())
-			{
-				RHICmdList.AddDispatchPrerequisite(OutCachedState->CompletionEvent);
-			}
+			RHICmdList.AddDispatchPrerequisite(OutCachedState->CompletionEvent);
 		}
 
 #if PSO_TRACK_CACHE_STATS
