Multi-GPU:  Fix editor and validation startup errors.  I'm not sure when these were introduced, but the fixes are trivial...

#rnx

diff --git a/Engine/Source/Runtime/D3D12RHI/Private/D3D12Buffer.cpp b/Engine/Source/Runtime/D3D12RHI/Private/D3D12Buffer.cpp
--- a/Engine/Source/Runtime/D3D12RHI/Private/D3D12Buffer.cpp
+++ b/Engine/Source/Runtime/D3D12RHI/Private/D3D12Buffer.cpp
@@ -725,6 +725,8 @@ void FD3D12DynamicRHI::UnlockBuffer(FRHICommandListBase& RHICmdList, FD3D12Buffe
 
 			for (FD3D12Buffer::FLinkedObjectIterator CurrentBuffer(Buffer); CurrentBuffer; ++CurrentBuffer)
 			{
+				SCOPED_GPU_MASK((FRHIComputeCommandList&)RHICmdList, FRHIGPUMask::FromIndex(CurrentBuffer->GetParentGPUIndex()));
+
 				// If we are on the render thread, queue up the copy on the RHIThread so it happens at the correct time.
 				if (RHICmdList.IsTopOfPipe())
 				{
diff --git a/Engine/Source/Runtime/RHI/Public/RHIValidationContext.h b/Engine/Source/Runtime/RHI/Public/RHIValidationContext.h
--- a/Engine/Source/Runtime/RHI/Public/RHIValidationContext.h
+++ b/Engine/Source/Runtime/RHI/Public/RHIValidationContext.h
@@ -310,6 +310,16 @@ public:
 		RHIContext->RHIWriteGPUFence(FenceRHI);
 	}
 
+	virtual void RHISetGPUMask(FRHIGPUMask GPUMask) override final
+	{
+		RHIContext->RHISetGPUMask(GPUMask);
+	}
+
+	virtual FRHIGPUMask RHIGetGPUMask() const override final
+	{
+		return RHIContext->RHIGetGPUMask();
+	}
+
 	virtual void RHICopyToStagingBuffer(FRHIBuffer* SourceBufferRHI, FRHIStagingBuffer* DestinationStagingBufferRHI, uint32 InOffset, uint32 InNumBytes) override final;
 
 #if WITH_MGPU
@@ -1124,6 +1134,16 @@ public:
 		RHIContext->RHIWriteGPUFence(FenceRHI);
 	}
 
+	virtual void RHISetGPUMask(FRHIGPUMask GPUMask) override final
+	{
+		RHIContext->RHISetGPUMask(GPUMask);
+	}
+
+	virtual FRHIGPUMask RHIGetGPUMask() const override final
+	{
+		return RHIContext->RHIGetGPUMask();
+	}
+
 	virtual void RHICopyToStagingBuffer(FRHIBuffer* SourceBufferRHI, FRHIStagingBuffer* DestinationStagingBufferRHI, uint32 InOffset, uint32 InNumBytes) override final;
 
 	virtual void RHICopyTexture(FRHITexture* SourceTexture, FRHITexture* DestTexture, const FRHICopyTextureInfo& CopyInfo) override final
