Mirrored global functions into FRHICommandListBase to make refactors easier and added FRHICommandListImmediate::Get().

diff --git a/Engine/Source/Runtime/RHI/Public/RHICommandList.h b/Engine/Source/Runtime/RHI/Public/RHICommandList.h
--- a/Engine/Source/Runtime/RHI/Public/RHICommandList.h
+++ b/Engine/Source/Runtime/RHI/Public/RHICommandList.h
@@ -40,6 +40,12 @@
 CSV_DECLARE_CATEGORY_MODULE_EXTERN(RHI_API, RHITStalls);
 CSV_DECLARE_CATEGORY_MODULE_EXTERN(RHI_API, RHITFlushes);
 
+/** Get the best default resource state for the given texture creation flags */
+extern RHI_API ERHIAccess RHIGetDefaultResourceState(ETextureCreateFlags InUsage, bool bInHasInitialData);
+
+/** Get the best default resource state for the given buffer creation flags */
+extern RHI_API ERHIAccess RHIGetDefaultResourceState(EBufferUsageFlags InUsage, bool bInHasInitialData);
+
 // Set to 1 to capture the callstack for every RHI command. Cheap & memory efficient representation: Use the 
 // value in FRHICommand::StackFrames to get the pointer to the code (ie paste on a disassembly window)
 #define RHICOMMAND_CALLSTACK		0
@@ -704,6 +710,42 @@ public:
 		return Buffer;
 	}
 
+	FORCEINLINE FBufferRHIRef CreateVertexBuffer(uint32 Size, EBufferUsageFlags Usage, ERHIAccess ResourceState, FRHIResourceCreateInfo& CreateInfo)
+	{
+		return CreateBuffer(Size, Usage | EBufferUsageFlags::VertexBuffer, 0, ResourceState, CreateInfo);
+	}
+
+	FORCEINLINE FBufferRHIRef CreateVertexBuffer(uint32 Size, EBufferUsageFlags Usage, FRHIResourceCreateInfo& CreateInfo)
+	{
+		bool bHasInitialData = CreateInfo.BulkData != nullptr;
+		ERHIAccess ResourceState = RHIGetDefaultResourceState(Usage | EBufferUsageFlags::VertexBuffer, bHasInitialData);
+		return CreateVertexBuffer(Size, Usage, ResourceState, CreateInfo);
+	}
+
+	FORCEINLINE FBufferRHIRef CreateStructuredBuffer(uint32 Stride, uint32 Size, EBufferUsageFlags Usage, ERHIAccess ResourceState, FRHIResourceCreateInfo& CreateInfo)
+	{
+		return CreateBuffer(Size, Usage | EBufferUsageFlags::StructuredBuffer, Stride, ResourceState, CreateInfo);
+	}
+
+	FORCEINLINE FBufferRHIRef CreateStructuredBuffer(uint32 Stride, uint32 Size, EBufferUsageFlags Usage, FRHIResourceCreateInfo& CreateInfo)
+	{
+		bool bHasInitialData = CreateInfo.BulkData != nullptr;
+		ERHIAccess ResourceState = RHIGetDefaultResourceState(Usage | EBufferUsageFlags::StructuredBuffer, bHasInitialData);
+		return CreateStructuredBuffer(Stride, Size, Usage, ResourceState, CreateInfo);
+	}
+
+	FORCEINLINE FBufferRHIRef CreateIndexBuffer(uint32 Stride, uint32 Size, EBufferUsageFlags Usage, ERHIAccess ResourceState, FRHIResourceCreateInfo& CreateInfo)
+	{
+		return CreateBuffer(Size, Usage | EBufferUsageFlags::IndexBuffer, Stride, ResourceState, CreateInfo);
+	}
+
+	FORCEINLINE FBufferRHIRef CreateIndexBuffer(uint32 Stride, uint32 Size, EBufferUsageFlags Usage, FRHIResourceCreateInfo& CreateInfo)
+	{
+		bool bHasInitialData = CreateInfo.BulkData != nullptr;
+		ERHIAccess ResourceState = RHIGetDefaultResourceState(Usage | EBufferUsageFlags::IndexBuffer, bHasInitialData);
+		return CreateIndexBuffer(Stride, Size, Usage, ResourceState, CreateInfo);
+	}
+
 	FORCEINLINE void UpdateUniformBuffer(FRHIUniformBuffer* UniformBufferRHI, const void* Contents)
 	{
 		FRHICommandListScopedPipelineGuard ScopedPipeline(*this);
@@ -769,6 +811,11 @@ public:
 		return GDynamicRHI->RHICreateUnorderedAccessView(*this, Texture, ViewDesc);
 	}
 
+	FORCEINLINE FShaderResourceViewRHIRef CreateShaderResourceView(const FShaderResourceViewInitializer& Initializer)
+	{
+		return CreateShaderResourceView(Initializer.Buffer, Initializer);
+	}
+
 	//UE_DEPRECATED(5.3, "Use the CreateUnorderedAccessView function that takes an FRHIViewDesc.")
 	FORCEINLINE FUnorderedAccessViewRHIRef CreateUnorderedAccessView(FRHIBuffer* Buffer, bool bUseUAVCounter, bool bAppendBuffer)
 	{
@@ -4020,6 +4067,8 @@ class FRHICommandListImmediate : public FRHICommandList
 	RHI_API int32 FlushExtendedLifetimeResourceDeletes();
 
 public:
+	static inline FRHICommandListImmediate& Get();
+
 	RHI_API void BeginScene();
 	RHI_API void EndScene();
 	RHI_API void BeginDrawingViewport(FRHIViewport* Viewport, FRHITexture* RenderTargetRHI);
@@ -4941,6 +4990,11 @@ public:
 	}
 };
 
+FORCEINLINE FRHICommandListImmediate& FRHICommandListImmediate::Get()
+{
+	check(IsInRenderingThread());
+	return FRHICommandListExecutor::GetImmediateCommandList();
+}
 
 FORCEINLINE_DEBUGGABLE FRHICommandListImmediate& FRHICommandListExecutor::GetImmediateCommandList()
 {
@@ -4958,15 +5012,13 @@ FORCEINLINE FBufferRHIRef RHICreateBuffer(uint32 Size, EBufferUsageFlags Usage,
 FORCEINLINE FBufferRHIRef RHICreateIndexBuffer(uint32 Stride, uint32 Size, EBufferUsageFlags Usage, ERHIAccess ResourceState, FRHIResourceCreateInfo& CreateInfo)
 {
 	check(IsInRenderingThread());
-	return FRHICommandListExecutor::GetImmediateCommandList().CreateBuffer(Size, Usage | EBufferUsageFlags::IndexBuffer, Stride, ResourceState, CreateInfo);
+	return FRHICommandListExecutor::GetImmediateCommandList().CreateIndexBuffer(Stride, Size, Usage, ResourceState, CreateInfo);
 }
 
-FORCEINLINE FBufferRHIRef RHICreateIndexBuffer(uint32 Stride, uint32 Size, EBufferUsageFlags InUsage, FRHIResourceCreateInfo& CreateInfo)
+FORCEINLINE FBufferRHIRef RHICreateIndexBuffer(uint32 Stride, uint32 Size, EBufferUsageFlags Usage, FRHIResourceCreateInfo& CreateInfo)
 {
-	bool bHasInitialData = CreateInfo.BulkData != nullptr;
-	EBufferUsageFlags Usage = InUsage | EBufferUsageFlags::IndexBuffer;
-	ERHIAccess ResourceState = RHIGetDefaultResourceState(Usage, bHasInitialData);
-	return RHICreateIndexBuffer(Stride, Size, Usage, ResourceState, CreateInfo);
+	check(IsInRenderingThread());
+	return FRHICommandListExecutor::GetImmediateCommandList().CreateIndexBuffer(Stride, Size, Usage, CreateInfo);
 }
 
 FORCEINLINE void RHIUpdateUniformBuffer(FRHIUniformBuffer* UniformBufferRHI, const void* Contents)
@@ -4977,29 +5029,25 @@ FORCEINLINE void RHIUpdateUniformBuffer(FRHIUniformBuffer* UniformBufferRHI, con
 FORCEINLINE FBufferRHIRef RHICreateVertexBuffer(uint32 Size, EBufferUsageFlags Usage, ERHIAccess ResourceState, FRHIResourceCreateInfo& CreateInfo)
 {
 	check(IsInRenderingThread());
-	return FRHICommandListExecutor::GetImmediateCommandList().CreateBuffer(Size, Usage | EBufferUsageFlags::VertexBuffer, 0, ResourceState, CreateInfo);
+	return FRHICommandListExecutor::GetImmediateCommandList().CreateVertexBuffer(Size, Usage, ResourceState, CreateInfo);
 }
 
-FORCEINLINE FBufferRHIRef RHICreateVertexBuffer(uint32 Size, EBufferUsageFlags InUsage, FRHIResourceCreateInfo& CreateInfo)
+FORCEINLINE FBufferRHIRef RHICreateVertexBuffer(uint32 Size, EBufferUsageFlags Usage, FRHIResourceCreateInfo& CreateInfo)
 {
-	bool bHasInitialData = CreateInfo.BulkData != nullptr;
-	EBufferUsageFlags Usage = InUsage | EBufferUsageFlags::VertexBuffer;
-	ERHIAccess ResourceState = RHIGetDefaultResourceState(Usage, bHasInitialData);
-	return RHICreateVertexBuffer(Size, Usage, ResourceState, CreateInfo);
+	check(IsInRenderingThread());
+	return FRHICommandListExecutor::GetImmediateCommandList().CreateVertexBuffer(Size, Usage, CreateInfo);
 }
 
 FORCEINLINE FBufferRHIRef RHICreateStructuredBuffer(uint32 Stride, uint32 Size, EBufferUsageFlags Usage, ERHIAccess ResourceState, FRHIResourceCreateInfo& CreateInfo)
 {
 	check(IsInRenderingThread());
-	return FRHICommandListExecutor::GetImmediateCommandList().CreateBuffer(Size, Usage | EBufferUsageFlags::StructuredBuffer, Stride, ResourceState, CreateInfo);
+	return FRHICommandListExecutor::GetImmediateCommandList().CreateStructuredBuffer(Stride, Size, Usage, ResourceState, CreateInfo);
 }
 
-FORCEINLINE FBufferRHIRef RHICreateStructuredBuffer(uint32 Stride, uint32 Size, EBufferUsageFlags InUsage, FRHIResourceCreateInfo& CreateInfo)
+FORCEINLINE FBufferRHIRef RHICreateStructuredBuffer(uint32 Stride, uint32 Size, EBufferUsageFlags Usage, FRHIResourceCreateInfo& CreateInfo)
 {
-	bool bHasInitialData = CreateInfo.BulkData != nullptr;
-	EBufferUsageFlags Usage = InUsage | BUF_StructuredBuffer;
-	ERHIAccess ResourceState = RHIGetDefaultResourceState(Usage, bHasInitialData);
-	return RHICreateStructuredBuffer(Stride, Size, Usage, ResourceState, CreateInfo);
+	check(IsInRenderingThread());
+	return FRHICommandListExecutor::GetImmediateCommandList().CreateStructuredBuffer(Stride, Size, Usage, CreateInfo);
 }
 
 FORCEINLINE void* RHILockBuffer(FRHIBuffer* Buffer, uint32 Offset, uint32 SizeRHI, EResourceLockMode LockMode)
