Validate that index buffers have been created with the EBufferUsageFlags::IndexBuffer usage flag.

#rnx

diff --git a/Engine/Source/Runtime/RHI/Public/RHIValidationContext.h b/Engine/Source/Runtime/RHI/Public/RHIValidationContext.h
--- a/Engine/Source/Runtime/RHI/Public/RHIValidationContext.h
+++ b/Engine/Source/Runtime/RHI/Public/RHIValidationContext.h
@@ -923,6 +923,7 @@ public:
 	virtual void RHIDrawIndexedPrimitive(FRHIBuffer* IndexBuffer, int32 BaseVertexIndex, uint32 FirstInstance, uint32 NumVertices, uint32 StartIndex, uint32 NumPrimitives, uint32 NumInstances) override final
 	{
 		checkf(State.bGfxPSOSet, TEXT("A Graphics PSO has to be set to draw!"));
+		checkf(EnumHasAnyFlags(IndexBuffer->GetUsage(), EBufferUsageFlags::IndexBuffer), TEXT("The buffer '%s' is used as an index buffer, but was not created with the IndexBuffer flag."), *IndexBuffer->GetName().ToString());
 		SetupDrawing();
 		Tracker->Assert(IndexBuffer->GetWholeResourceIdentity(), ERHIAccess::VertexOrIndexBuffer);
 		RHIContext->RHIDrawIndexedPrimitive(IndexBuffer, BaseVertexIndex, FirstInstance, NumVertices, StartIndex, NumPrimitives, NumInstances);
@@ -932,6 +933,7 @@ public:
 	virtual void RHIDrawIndexedPrimitiveIndirect(FRHIBuffer* IndexBuffer, FRHIBuffer* ArgumentBuffer, uint32 ArgumentOffset) override final
 	{
 		checkf(State.bGfxPSOSet, TEXT("A Graphics PSO has to be set to draw!"));
+		checkf(EnumHasAnyFlags(IndexBuffer->GetUsage(), EBufferUsageFlags::IndexBuffer), TEXT("The buffer '%s' is used as an index buffer, but was not created with the IndexBuffer flag."), *IndexBuffer->GetName().ToString());
 		FValidationRHI::ValidateIndirectArgsBuffer(ArgumentBuffer, ArgumentOffset, sizeof(FRHIDrawIndexedIndirectParameters), 0);
 		SetupDrawing();
 		Tracker->Assert(ArgumentBuffer->GetWholeResourceIdentity(), ERHIAccess::IndirectArgs);
@@ -943,6 +945,7 @@ public:
 	virtual void RHIMultiDrawIndexedPrimitiveIndirect(FRHIBuffer* IndexBuffer, FRHIBuffer* ArgumentBuffer, uint32 ArgumentOffset, FRHIBuffer* CountBuffer, uint32 CountBufferOffset, uint32 MaxDrawArguments) override final
 	{
 		checkf(State.bGfxPSOSet, TEXT("A Graphics PSO has to be set to draw!"));
+		checkf(EnumHasAnyFlags(IndexBuffer->GetUsage(), EBufferUsageFlags::IndexBuffer), TEXT("The buffer '%s' is used as an index buffer, but was not created with the IndexBuffer flag."), *IndexBuffer->GetName().ToString());
 		FValidationRHI::ValidateIndirectArgsBuffer(ArgumentBuffer, ArgumentOffset, sizeof(FRHIDrawIndexedIndirectParameters), 0);
 		SetupDrawing();
 		Tracker->Assert(ArgumentBuffer->GetWholeResourceIdentity(), ERHIAccess::IndirectArgs);
