Fixed edge case where RHI thread fence candidiates would grow unbounded.

diff --git a/Engine/Source/Runtime/RHI/Private/RHICommandList.cpp b/Engine/Source/Runtime/RHI/Private/RHICommandList.cpp
--- a/Engine/Source/Runtime/RHI/Private/RHICommandList.cpp
+++ b/Engine/Source/Runtime/RHI/Private/RHICommandList.cpp
@@ -1340,6 +1340,13 @@ void FRHICommandListImmediate::BeginFrame()
 void FRHICommandListImmediate::EndFrame()
 {
 	check(IsImmediate() && IsInRenderingThread());
+
+	// Flush any queued fence candidates by injecting a fence. This is to handle cases where no parallel translates happened within the frame.
+	if (!PersistentState.QueuedFenceCandidates.IsEmpty())
+	{
+		RHIThreadFence(true);
+	}
+
 	if (Bypass())
 	{
 		GetContext().RHIEndFrame();
