Making FRayTracingGeometryInitializer not a memory layout type since nothing needs it like that anymore.

diff --git a/Engine/Plugins/Runtime/GeometryCache/Source/GeometryCache/Private/GeometryCacheSceneProxy.cpp b/Engine/Plugins/Runtime/GeometryCache/Source/GeometryCache/Private/GeometryCacheSceneProxy.cpp
--- a/Engine/Plugins/Runtime/GeometryCache/Source/GeometryCache/Private/GeometryCacheSceneProxy.cpp
+++ b/Engine/Plugins/Runtime/GeometryCache/Source/GeometryCache/Private/GeometryCacheSceneProxy.cpp
@@ -668,7 +668,7 @@ void FGeometryCacheSceneProxy::UpdateAnimation(float NewTime, bool bNewLooping,
 				const FGeometryCacheMeshData* MeshData = Section->bNextFrameMeshDataSelected ? Section->NextFrameMeshData : Section->MeshData;
 				const uint32 IndexBufferNumTriangles = Section->IndexBuffer.NumValidIndices / 3;
 
-				TMemoryImageArray<FRayTracingGeometrySegment>& Segments = Section->RayTracingGeometry.Initializer.Segments;
+				TArray<FRayTracingGeometrySegment>& Segments = Section->RayTracingGeometry.Initializer.Segments;
 
 				// Check if a full RaytracingGeometry object needs to be recreated. 
 				// Recreate when:
diff --git a/Engine/Source/Runtime/D3D12RHI/Private/D3D12RayTracing.cpp b/Engine/Source/Runtime/D3D12RHI/Private/D3D12RayTracing.cpp
--- a/Engine/Source/Runtime/D3D12RHI/Private/D3D12RayTracing.cpp
+++ b/Engine/Source/Runtime/D3D12RHI/Private/D3D12RayTracing.cpp
@@ -3314,7 +3314,9 @@ FD3D12RayTracingGeometry::FD3D12RayTracingGeometry(FRHICommandListBase& RHICmdLi
 {
 	INC_DWORD_STAT(STAT_D3D12RayTracingAllocatedBLAS);
 
-	DebugName = !Initializer.DebugName.IsNone() ? Initializer.DebugName : FDebugName(FName(TEXT("BLAS")));
+	static const FName NAME_BLAS(TEXT("BLAS"));
+
+	DebugName = !Initializer.DebugName.IsNone() ? Initializer.DebugName : NAME_BLAS;
 	OwnerName = Initializer.OwnerName;
 	
 	FMemory::Memzero(bHasPendingCompactionRequests);
diff --git a/Engine/Source/Runtime/RHI/Private/DynamicRHI.cpp b/Engine/Source/Runtime/RHI/Private/DynamicRHI.cpp
--- a/Engine/Source/Runtime/RHI/Private/DynamicRHI.cpp
+++ b/Engine/Source/Runtime/RHI/Private/DynamicRHI.cpp
@@ -24,9 +24,6 @@
 #include "RHITextureReference.h"
 #include "Serialization/MemoryImage.h"
 
-IMPLEMENT_TYPE_LAYOUT(FRayTracingGeometryInitializer);
-IMPLEMENT_TYPE_LAYOUT(FRayTracingGeometrySegment);
-
 static_assert(sizeof(FRayTracingGeometryInstance) <= 104,
 	"Ray tracing instance descriptor is expected to be no more than 104 bytes, as there may be a very large number of them.");
 
diff --git a/Engine/Source/Runtime/RHI/Private/RHI.cpp b/Engine/Source/Runtime/RHI/Private/RHI.cpp
--- a/Engine/Source/Runtime/RHI/Private/RHI.cpp
+++ b/Engine/Source/Runtime/RHI/Private/RHI.cpp
@@ -2106,15 +2106,9 @@ FDebugName::FDebugName(FName InName, int32 InNumber)
 {
 }
 
-FDebugName::FDebugName(FMemoryImageName InName, int32 InNumber)
-	: Name(InName)
-	, Number(InNumber)
-{
-}
-
 FDebugName& FDebugName::operator=(FName Other)
 {
-	Name = FMemoryImageName(Other);
+	Name = Other;
 	Number = NAME_NO_NUMBER_INTERNAL;
 	return *this;
 }
@@ -2122,7 +2116,7 @@ FDebugName& FDebugName::operator=(FName Other)
 FString FDebugName::ToString() const
 {
 	FString Out;
-	FName(Name).AppendString(Out);
+	Name.AppendString(Out);
 	if (Number != NAME_NO_NUMBER_INTERNAL)
 	{
 		Out.Appendf(TEXT("_%u"), Number);
@@ -2132,7 +2126,7 @@ FString FDebugName::ToString() const
 
 void FDebugName::AppendString(FStringBuilderBase& Builder) const
 {
-	FName(Name).AppendString(Builder);
+	Name.AppendString(Builder);
 	if (Number != NAME_NO_NUMBER_INTERNAL)
 	{
 		Builder << '_' << Number;
diff --git a/Engine/Source/Runtime/RHI/Public/RHIResources.h b/Engine/Source/Runtime/RHI/Public/RHIResources.h
--- a/Engine/Source/Runtime/RHI/Public/RHIResources.h
+++ b/Engine/Source/Runtime/RHI/Public/RHIResources.h
@@ -2976,13 +2976,10 @@ public:
  */
 class FDebugName
 {
-	DECLARE_INLINE_TYPE_LAYOUT(FDebugName, NonVirtual);
-
 public:
 	RHI_API FDebugName();
 	RHI_API FDebugName(FName InName);
 	RHI_API FDebugName(FName InName, int32 InNumber);
-	RHI_API FDebugName(FMemoryImageName InName, int32 InNumber);
 
 	RHI_API FDebugName& operator=(FName Other);
 
@@ -2991,8 +2988,8 @@ public:
 	RHI_API void AppendString(FStringBuilderBase& Builder) const;
 
 private:
-	LAYOUT_FIELD(FMemoryImageName, Name);
-	LAYOUT_FIELD(uint32, Number);
+	FName Name;
+	uint32 Number;
 };
 
 //
@@ -3089,71 +3086,69 @@ DECLARE_INTRINSIC_TYPE_LAYOUT(ERayTracingGeometryInitializerType);
 
 struct FRayTracingGeometrySegment
 {
-	DECLARE_TYPE_LAYOUT(FRayTracingGeometrySegment, NonVirtual);
 public:
-	LAYOUT_FIELD_INITIALIZED(FBufferRHIRef, VertexBuffer, nullptr);
-	LAYOUT_FIELD_INITIALIZED(EVertexElementType, VertexBufferElementType, VET_Float3);
+	FBufferRHIRef VertexBuffer = nullptr;
+	EVertexElementType VertexBufferElementType = VET_Float3;
 
 	// Offset in bytes from the base address of the vertex buffer.
-	LAYOUT_FIELD_INITIALIZED(uint32, VertexBufferOffset, 0);
+	uint32 VertexBufferOffset = 0;
 
 	// Number of bytes between elements of the vertex buffer (sizeof VET_Float3 by default).
 	// Must be equal or greater than the size of the position vector.
-	LAYOUT_FIELD_INITIALIZED(uint32, VertexBufferStride, 12);
+	uint32 VertexBufferStride = 12;
 
 	// Number of vertices (positions) in VertexBuffer.
 	// If an index buffer is present, this must be at least the maximum index value in the index buffer + 1.
-	LAYOUT_FIELD_INITIALIZED(uint32, MaxVertices, 0);
+	uint32 MaxVertices = 0;
 
 	// Primitive range for this segment.
-	LAYOUT_FIELD_INITIALIZED(uint32, FirstPrimitive, 0);
-	LAYOUT_FIELD_INITIALIZED(uint32, NumPrimitives, 0);
+	uint32 FirstPrimitive = 0;
+	uint32 NumPrimitives = 0;
 
 	// Indicates whether any-hit shader could be invoked when hitting this geometry segment.
 	// Setting this to `false` turns off any-hit shaders, making the section "opaque" and improving ray tracing performance.
-	LAYOUT_FIELD_INITIALIZED(bool, bForceOpaque, false);
+	bool bForceOpaque = false;
 
 	// Any-hit shader may be invoked multiple times for the same primitive during ray traversal.
 	// Setting this to `false` guarantees that only a single instance of any-hit shader will run per primitive, at some performance cost.
-	LAYOUT_FIELD_INITIALIZED(bool, bAllowDuplicateAnyHitShaderInvocation, true);
+	bool bAllowDuplicateAnyHitShaderInvocation = true;
 
 	// Indicates whether this section is enabled and should be taken into account during acceleration structure creation
-	LAYOUT_FIELD_INITIALIZED(bool, bEnabled, true);
+	bool bEnabled = true;
 };
 
 struct FRayTracingGeometryInitializer
 {
-	DECLARE_EXPORTED_TYPE_LAYOUT(FRayTracingGeometryInitializer, RHI_API, NonVirtual);
 public:
-	LAYOUT_FIELD_INITIALIZED(FBufferRHIRef, IndexBuffer, nullptr);
+	FBufferRHIRef IndexBuffer = nullptr;
 
 	// Offset in bytes from the base address of the index buffer.
-	LAYOUT_FIELD_INITIALIZED(uint32, IndexBufferOffset, 0);
+	uint32 IndexBufferOffset = 0;
 
-	LAYOUT_FIELD_INITIALIZED(ERayTracingGeometryType, GeometryType, RTGT_Triangles);
+	ERayTracingGeometryType GeometryType = RTGT_Triangles;
 
 	// Total number of primitives in all segments of the geometry. Only used for validation.
-	LAYOUT_FIELD_INITIALIZED(uint32, TotalPrimitiveCount, 0);
+	uint32 TotalPrimitiveCount = 0;
 
 	// Partitions of geometry to allow different shader and resource bindings.
 	// All ray tracing geometries must have at least one segment.
-	LAYOUT_FIELD(TMemoryImageArray<FRayTracingGeometrySegment>, Segments);
+	TArray<FRayTracingGeometrySegment> Segments;
 
 	// Offline built geometry data. If null, the geometry will be built by the RHI at runtime.
-	LAYOUT_FIELD_INITIALIZED(FResourceArrayInterface*, OfflineData, nullptr);
+	FResourceArrayInterface* OfflineData = nullptr;
 
 	// Pointer to an existing ray tracing geometry which the new geometry is built from.
-	LAYOUT_FIELD_INITIALIZED(FRHIRayTracingGeometry*, SourceGeometry, nullptr);
+	FRHIRayTracingGeometry* SourceGeometry = nullptr;
 
-	LAYOUT_FIELD_INITIALIZED(bool, bFastBuild, false);
-	LAYOUT_FIELD_INITIALIZED(bool, bAllowUpdate, false);
-	LAYOUT_FIELD_INITIALIZED(bool, bAllowCompaction, true);
-	LAYOUT_FIELD_INITIALIZED(ERayTracingGeometryInitializerType, Type, ERayTracingGeometryInitializerType::Rendering);
+	bool bFastBuild = false;
+	bool bAllowUpdate = false;
+	bool bAllowCompaction = true;
+	ERayTracingGeometryInitializerType Type = ERayTracingGeometryInitializerType::Rendering;
 
 	// Use FDebugName for auto-generated debug names with numbered suffixes, it is a variation of FMemoryImageName with optional number postfix.
-	LAYOUT_FIELD(FDebugName, DebugName);
+	FDebugName DebugName;
 	// Store the path name of the owner object for resource tracking. FMemoryImageName allows a conversion to/from FName.
-	LAYOUT_FIELD(FMemoryImageName, OwnerName);
+	FName OwnerName;
 };
 
 enum ERayTracingSceneLifetime
