Putting FStaticShaderPlatformNames in its own file.

diff --git a/Engine/Source/Runtime/RHI/Private/RHIStrings.cpp b/Engine/Source/Runtime/RHI/Private/RHIStrings.cpp
--- a/Engine/Source/Runtime/RHI/Private/RHIStrings.cpp
+++ b/Engine/Source/Runtime/RHI/Private/RHIStrings.cpp
@@ -5,6 +5,7 @@
 #include "RHI.h"
 #include "RHIAccess.h"
 #include "RHIShaderFormatDefinitions.inl"
+#include "RHIStaticShaderPlatformNames.h"
 #include "RHIPipeline.h"
 
 static FName NAME_PLATFORM_WINDOWS(TEXT("Windows"));
diff --git a/Engine/Source/Runtime/RHI/Public/DataDrivenShaderPlatformInfo.h b/Engine/Source/Runtime/RHI/Public/DataDrivenShaderPlatformInfo.h
--- a/Engine/Source/Runtime/RHI/Public/DataDrivenShaderPlatformInfo.h
+++ b/Engine/Source/Runtime/RHI/Public/DataDrivenShaderPlatformInfo.h
@@ -5,6 +5,7 @@
 #include "RHIDefinitions.h"
 #include "RHIShaderPlatform.h"
 #include "RHIFeatureLevel.h"
+#include "UObject/NameTypes.h"
 
 class FText;
 
diff --git a/Engine/Source/Runtime/RHI/Public/RHIShaderFormatDefinitions.inl b/Engine/Source/Runtime/RHI/Public/RHIShaderFormatDefinitions.inl
--- a/Engine/Source/Runtime/RHI/Public/RHIShaderFormatDefinitions.inl
+++ b/Engine/Source/Runtime/RHI/Public/RHIShaderFormatDefinitions.inl
@@ -7,7 +7,7 @@
 
 #pragma once
 
-#include "RHIShaderPlatform.h"
+#include "RHIStaticShaderPlatformNames.h"
 
 static FName NAME_PCD3D_SM6(TEXT("PCD3D_SM6"));
 static FName NAME_PCD3D_SM5(TEXT("PCD3D_SM5"));
diff --git a/Engine/Source/Runtime/RHI/Public/RHIShaderPlatform.h b/Engine/Source/Runtime/RHI/Public/RHIShaderPlatform.h
--- a/Engine/Source/Runtime/RHI/Public/RHIShaderPlatform.h
+++ b/Engine/Source/Runtime/RHI/Public/RHIShaderPlatform.h
@@ -2,8 +2,8 @@
 
 #pragma once
 
-#include "UObject/NameTypes.h"
 #include "RHIDefinitions.h"
+#include "HAL/Platform.h"
 
 /** @warning: update *LegacyShaderPlatform* when the below changes */
 enum EShaderPlatform : uint16
@@ -79,82 +79,6 @@ private:
 using FStaticShaderPlatform = FGenericStaticShaderPlatform;
 #endif
 
-class FStaticShaderPlatformNames
-{
-private:
-	static const uint32 NumPlatforms = DDPI_NUM_STATIC_SHADER_PLATFORMS;
-
-	struct FPlatform
-	{
-		FName Name;
-		FName ShaderPlatform;
-		FName ShaderFormat;
-	} Platforms[NumPlatforms];
-
-	FStaticShaderPlatformNames()
-	{
-#ifdef DDPI_SHADER_PLATFORM_NAME_MAP
-		struct FStaticNameMapEntry
-		{
-			FName Name;
-			FName PlatformName;
-			int32 Index;
-		} NameMap[] =
-		{
-			DDPI_SHADER_PLATFORM_NAME_MAP
-		};
-
-		for (int32 MapIndex = 0; MapIndex < UE_ARRAY_COUNT(NameMap); ++MapIndex)
-		{
-			FStaticNameMapEntry const& Entry = NameMap[MapIndex];
-			check(IsStaticPlatform(EShaderPlatform(Entry.Index)));
-			uint32 PlatformIndex = Entry.Index - SP_StaticPlatform_First;
-
-			FPlatform& Platform = Platforms[PlatformIndex];
-			check(Platform.Name == NAME_None); // Check we've not already seen this platform
-
-			Platform.Name = Entry.PlatformName;
-			Platform.ShaderPlatform = FName(*FString::Printf(TEXT("SP_%s"), *Entry.Name.ToString()), FNAME_Add);
-			Platform.ShaderFormat = FName(*FString::Printf(TEXT("SF_%s"), *Entry.Name.ToString()), FNAME_Add);
-		}
-#endif
-	}
-
-public:
-	static inline FStaticShaderPlatformNames const& Get()
-	{
-		static FStaticShaderPlatformNames Names;
-		return Names;
-	}
-
-	static inline bool IsStaticPlatform(EShaderPlatform Platform)
-	{
-		return Platform >= SP_StaticPlatform_First && Platform <= SP_StaticPlatform_Last;
-	}
-
-	inline const FName& GetShaderPlatform(EShaderPlatform Platform) const
-	{
-		return Platforms[GetStaticPlatformIndex(Platform)].ShaderPlatform;
-	}
-
-	inline const FName& GetShaderFormat(EShaderPlatform Platform) const
-	{
-		return Platforms[GetStaticPlatformIndex(Platform)].ShaderFormat;
-	}
-
-	inline const FName& GetPlatformName(EShaderPlatform Platform) const
-	{
-		return Platforms[GetStaticPlatformIndex(Platform)].Name;
-	}
-
-private:
-	static inline uint32 GetStaticPlatformIndex(EShaderPlatform Platform)
-	{
-		check(IsStaticPlatform(Platform));
-		return uint32(Platform) - SP_StaticPlatform_First;
-	}
-};
-
 // The maximum shader platform available on this system
 extern RHI_API EShaderPlatform GMaxRHIShaderPlatform;
 
@@ -166,3 +90,8 @@ inline bool IsCustomPlatform(const FStaticShaderPlatform Platform)
 #if UE_ENABLE_INCLUDE_ORDER_DEPRECATED_IN_5_2
 #include "RHIStrings.h"
 #endif
+
+#if UE_ENABLE_INCLUDE_ORDER_DEPRECATED_IN_5_3
+#include "UObject/NameTypes.h"
+#include "RHIStaticShaderPlatformNames.h"
+#endif
diff --git a/Engine/Source/Runtime/RHI/Public/RHIStaticShaderPlatformNames.h b/Engine/Source/Runtime/RHI/Public/RHIStaticShaderPlatformNames.h
new file mode 100644
--- /dev/null
+++ b/Engine/Source/Runtime/RHI/Public/RHIStaticShaderPlatformNames.h
@@ -0,0 +1,82 @@
+// Copyright Epic Games, Inc. All Rights Reserved.
+
+#pragma once
+
+#include "RHIShaderPlatform.h"
+#include "UObject/NameTypes.h"
+
+class FStaticShaderPlatformNames
+{
+private:
+	static const uint32 NumPlatforms = DDPI_NUM_STATIC_SHADER_PLATFORMS;
+
+	struct FPlatform
+	{
+		FName Name;
+		FName ShaderPlatform;
+		FName ShaderFormat;
+	} Platforms[NumPlatforms];
+
+	FStaticShaderPlatformNames()
+	{
+#ifdef DDPI_SHADER_PLATFORM_NAME_MAP
+		struct FStaticNameMapEntry
+		{
+			FName Name;
+			FName PlatformName;
+			int32 Index;
+		} NameMap[] =
+		{
+			DDPI_SHADER_PLATFORM_NAME_MAP
+		};
+
+		for (int32 MapIndex = 0; MapIndex < UE_ARRAY_COUNT(NameMap); ++MapIndex)
+		{
+			FStaticNameMapEntry const& Entry = NameMap[MapIndex];
+			check(IsStaticPlatform(EShaderPlatform(Entry.Index)));
+			uint32 PlatformIndex = Entry.Index - SP_StaticPlatform_First;
+
+			FPlatform& Platform = Platforms[PlatformIndex];
+			check(Platform.Name == NAME_None); // Check we've not already seen this platform
+
+			Platform.Name = Entry.PlatformName;
+			Platform.ShaderPlatform = FName(*FString::Printf(TEXT("SP_%s"), *Entry.Name.ToString()), FNAME_Add);
+			Platform.ShaderFormat = FName(*FString::Printf(TEXT("SF_%s"), *Entry.Name.ToString()), FNAME_Add);
+		}
+#endif
+	}
+
+public:
+	static inline FStaticShaderPlatformNames const& Get()
+	{
+		static FStaticShaderPlatformNames Names;
+		return Names;
+	}
+
+	static inline bool IsStaticPlatform(EShaderPlatform Platform)
+	{
+		return Platform >= SP_StaticPlatform_First && Platform <= SP_StaticPlatform_Last;
+	}
+
+	inline const FName& GetShaderPlatform(EShaderPlatform Platform) const
+	{
+		return Platforms[GetStaticPlatformIndex(Platform)].ShaderPlatform;
+	}
+
+	inline const FName& GetShaderFormat(EShaderPlatform Platform) const
+	{
+		return Platforms[GetStaticPlatformIndex(Platform)].ShaderFormat;
+	}
+
+	inline const FName& GetPlatformName(EShaderPlatform Platform) const
+	{
+		return Platforms[GetStaticPlatformIndex(Platform)].Name;
+	}
+
+private:
+	static inline uint32 GetStaticPlatformIndex(EShaderPlatform Platform)
+	{
+		check(IsStaticPlatform(Platform));
+		return uint32(Platform) - SP_StaticPlatform_First;
+	}
+};
