Non-immediate command lists always record commands, even if bypass mode is requested.

This is because currently QueueAsyncCommandListSubmit() only supports command lists with their own local contexts for parallel translation. When parallel translation is not used, these commands need to be replayed on the immediate context, but bypass mode would create a local context and translate commands into it, and this context is never finalized and submitted.

Also removed PLATFORM_RHITHREAD_DEFAULT_BYPASS, since bypass should be off by default on all platforms. Using this assumption, the Bypass() methods return compile-time constants in test and shipping builds, so all the bypass code can be compiled out.

#rnx

diff --git a/Engine/Source/Runtime/Core/Public/Android/AndroidPlatform.h b/Engine/Source/Runtime/Core/Public/Android/AndroidPlatform.h
--- a/Engine/Source/Runtime/Core/Public/Android/AndroidPlatform.h
+++ b/Engine/Source/Runtime/Core/Public/Android/AndroidPlatform.h
@@ -66,7 +66,6 @@ typedef FAndroidTypes FPlatformTypes;
 #define PLATFORM_ALLOW_ALLOCATIONS_IN_FASYNCWRITER_SERIALIZEBUFFERTOARCHIVE	0
 
 #define PLATFORM_GLOBAL_LOG_CATEGORY				LogAndroid
-#define PLATFORM_RHITHREAD_DEFAULT_BYPASS			0
 
 #ifndef RUNNING_WITH_ASAN
 	#define RUNNING_WITH_ASAN						0
diff --git a/Engine/Source/Runtime/Core/Public/HAL/Platform.h b/Engine/Source/Runtime/Core/Public/HAL/Platform.h
--- a/Engine/Source/Runtime/Core/Public/HAL/Platform.h
+++ b/Engine/Source/Runtime/Core/Public/HAL/Platform.h
@@ -439,10 +439,6 @@
 	#define PLATFORM_IMPLEMENTS_BeginNamedEventStatic			0
 #endif
 
-#ifndef PLATFORM_RHITHREAD_DEFAULT_BYPASS
-	#define PLATFORM_RHITHREAD_DEFAULT_BYPASS					1
-#endif
-
 #ifndef PLATFORM_HAS_UMA
 	#define PLATFORM_HAS_UMA 0
 #endif
diff --git a/Engine/Source/Runtime/Core/Public/IOS/IOSPlatform.h b/Engine/Source/Runtime/Core/Public/IOS/IOSPlatform.h
--- a/Engine/Source/Runtime/Core/Public/IOS/IOSPlatform.h
+++ b/Engine/Source/Runtime/Core/Public/IOS/IOSPlatform.h
@@ -91,8 +91,6 @@ typedef FIOSPlatformTypes FPlatformTypes;
 //mallocpoison not safe with aligned ansi allocator.  returns the larger unaligned size during Free() which causes writes off the end of the allocation.
 #define UE_USE_MALLOC_FILL_BYTES 0 
 
-#define PLATFORM_RHITHREAD_DEFAULT_BYPASS				1
-
 // Function type macros.
 #define VARARGS															/* Functions with variable arguments */
 #define CDECL															/* Standard C function */
diff --git a/Engine/Source/Runtime/Core/Public/Mac/MacPlatform.h b/Engine/Source/Runtime/Core/Public/Mac/MacPlatform.h
--- a/Engine/Source/Runtime/Core/Public/Mac/MacPlatform.h
+++ b/Engine/Source/Runtime/Core/Public/Mac/MacPlatform.h
@@ -76,8 +76,6 @@ typedef FMacPlatformTypes FPlatformTypes;
 #define PLATFORM_IS_ANSI_MALLOC_THREADSAFE				1
 #define PLATFORM_SUPPORTS_VIRTUAL_TEXTURE_STREAMING		1
 
-#define PLATFORM_RHITHREAD_DEFAULT_BYPASS				0
-
 #define PLATFORM_ENABLE_POPCNT_INTRINSIC				1
 
 #define PLATFORM_GLOBAL_LOG_CATEGORY					LogMac
diff --git a/Engine/Source/Runtime/Core/Public/Unix/UnixPlatform.h b/Engine/Source/Runtime/Core/Public/Unix/UnixPlatform.h
--- a/Engine/Source/Runtime/Core/Public/Unix/UnixPlatform.h
+++ b/Engine/Source/Runtime/Core/Public/Unix/UnixPlatform.h
@@ -66,7 +66,6 @@ typedef FUnixPlatformTypes FPlatformTypes;
 #define PLATFORM_SUPPORTS_STACK_SYMBOLS					1
 #define PLATFORM_IS_ANSI_MALLOC_THREADSAFE				1
 #define PLATFORM_ALLOW_ALLOCATIONS_IN_FASYNCWRITER_SERIALIZEBUFFERTOARCHIVE 0
-#define PLATFORM_RHITHREAD_DEFAULT_BYPASS				0
 #define PLATFORM_SUPPORTS_VIRTUAL_TEXTURE_STREAMING		1
 #define PLATFORM_USE_ANSI_POSIX_MALLOC					1
 
diff --git a/Engine/Source/Runtime/Core/Public/Windows/WindowsPlatform.h b/Engine/Source/Runtime/Core/Public/Windows/WindowsPlatform.h
--- a/Engine/Source/Runtime/Core/Public/Windows/WindowsPlatform.h
+++ b/Engine/Source/Runtime/Core/Public/Windows/WindowsPlatform.h
@@ -88,7 +88,6 @@ typedef FWindowsPlatformTypes FPlatformTypes;
 #define PLATFORM_SUPPORTS_NAMED_PIPES						1
 #define PLATFORM_COMPILER_HAS_TCHAR_WMAIN					1
 #define PLATFORM_SUPPORTS_EARLY_MOVIE_PLAYBACK				(!WITH_EDITOR) // movies will start before engine is initalized
-#define PLATFORM_RHITHREAD_DEFAULT_BYPASS					0
 #define PLATFORM_USE_GENERIC_STRING_IMPLEMENTATION			0 // Can set this to 1 if you need to debug FGenericWidePlatformString::GetVarArgs on Windows
 #define PLATFORM_SUPPORTS_VIRTUAL_TEXTURE_STREAMING			1
 #define PLATFORM_SUPPORTS_VARIABLE_RATE_SHADING				1
diff --git a/Engine/Source/Runtime/RHI/Private/RHICommandList.cpp b/Engine/Source/Runtime/RHI/Private/RHICommandList.cpp
--- a/Engine/Source/Runtime/RHI/Private/RHICommandList.cpp
+++ b/Engine/Source/Runtime/RHI/Private/RHICommandList.cpp
@@ -38,7 +38,7 @@ bool FScopedUniformBufferStaticBindings::bRecursionGuard = false;
 
 static TAutoConsoleVariable<int32> CVarRHICmdBypass(
 	TEXT("r.RHICmdBypass"),
-	FRHICommandListExecutor::DefaultBypass,
+	0,
 	TEXT("Whether to bypass the rhi command list and send the rhi commands immediately.\n")
 	TEXT("0: Disable (required for the multithreaded renderer)\n")
 	TEXT("1: Enable (convenient for debugging low level graphics API calls, can suppress artifacts from multithreaded renderer code)"));
@@ -159,10 +159,6 @@ FRHICommandListBase::FRHICommandListBase(FPersistentState&& InPersistentState)
 	{
 		GRHICommandList.OutstandingCmdListCount.Increment();
 	}
-	else
-	{
-		checkf(!PLATFORM_RHITHREAD_DEFAULT_BYPASS, TEXT("The platform has enabled RHI command list bypass mode for shipping builds. Only render thread command lists are allowed."));
-	}
 #endif
 }
 
@@ -195,10 +191,6 @@ FRHICommandListBase::FRHICommandListBase(FRHICommandListBase&& Other)
 	{
 		GRHICommandList.OutstandingCmdListCount.Increment();
 	}
-	else
-	{
-		checkf(!PLATFORM_RHITHREAD_DEFAULT_BYPASS, TEXT("The platform has enabled RHI command list bypass mode for shipping builds. Only render thread command lists are allowed."));
-	}
 #endif
 
 	Other.Root = nullptr;
diff --git a/Engine/Source/Runtime/RHI/Public/RHICommandList.h b/Engine/Source/Runtime/RHI/Public/RHICommandList.h
--- a/Engine/Source/Runtime/RHI/Public/RHICommandList.h
+++ b/Engine/Source/Runtime/RHI/Public/RHICommandList.h
@@ -4919,12 +4919,8 @@ public:
 class FRHICommandListExecutor
 {
 public:
-	enum
-	{
-		DefaultBypass = PLATFORM_RHITHREAD_DEFAULT_BYPASS
-	};
 	FRHICommandListExecutor()
-		: bLatchedBypass(!!DefaultBypass)
+		: bLatchedBypass(false)
 		, bLatchedUseParallelAlgorithms(false)
 	{
 	}
@@ -4933,12 +4929,12 @@ public:
 
 	RHI_API static void WaitOnRHIThreadFence(FGraphEventRef& Fence);
 
-	FORCEINLINE_DEBUGGABLE bool Bypass()
+	FORCEINLINE_DEBUGGABLE bool Bypass() const
 	{
 #if CAN_TOGGLE_COMMAND_LIST_BYPASS
 		return bLatchedBypass;
 #else
-		return !!DefaultBypass;
+		return false;
 #endif
 	}
 
diff --git a/Engine/Source/Runtime/RHI/Public/RHICommandList.inl b/Engine/Source/Runtime/RHI/Public/RHICommandList.inl
--- a/Engine/Source/Runtime/RHI/Public/RHICommandList.inl
+++ b/Engine/Source/Runtime/RHI/Public/RHICommandList.inl
@@ -26,12 +26,12 @@ FORCEINLINE_DEBUGGABLE FRHICommandListImmediate& FRHICommandListBase::GetAsImmed
 
 FORCEINLINE_DEBUGGABLE bool FRHICommandListBase::Bypass() const
 {
-	check(!IsImmediate() || IsInRenderingThread() || IsInRHIThread());
-	return GRHICommandList.Bypass()
 #if CAN_TOGGLE_COMMAND_LIST_BYPASS
-		&& PersistentState.RecordingThread == ERecordingThread::Render
+	check(!IsImmediate() || IsInRenderingThread() || IsInRHIThread());
+	return GRHICommandList.Bypass() && IsImmediate();
+#else
+	return false;
 #endif
-		;
 }
 
 FORCEINLINE_DEBUGGABLE FScopedRHIThreadStaller::FScopedRHIThreadStaller(class FRHICommandListImmediate& InImmed, bool bDoStall)
