Removing MemoryLayout includes from global RHI headers. This required moving FRHIUniformBufferLayoutInitializer to its own file.
Did a few more minor dependency cleanups in RHI headers.

diff --git a/Engine/Source/Runtime/Engine/Public/MaterialShared.h b/Engine/Source/Runtime/Engine/Public/MaterialShared.h
--- a/Engine/Source/Runtime/Engine/Public/MaterialShared.h
+++ b/Engine/Source/Runtime/Engine/Public/MaterialShared.h
@@ -15,6 +15,7 @@
 #include "Templates/UniquePtr.h"
 #include "Misc/SecureHash.h"
 #include "RHIDefinitions.h"
+#include "RHIUniformBufferLayoutInitializer.h"
 #include "RenderResource.h"
 #if UE_ENABLE_INCLUDE_ORDER_DEPRECATED_IN_5_2
 #include "RenderingThread.h"
diff --git a/Engine/Source/Runtime/Engine/Public/Rendering/PositionVertexBuffer.h b/Engine/Source/Runtime/Engine/Public/Rendering/PositionVertexBuffer.h
--- a/Engine/Source/Runtime/Engine/Public/Rendering/PositionVertexBuffer.h
+++ b/Engine/Source/Runtime/Engine/Public/Rendering/PositionVertexBuffer.h
@@ -2,9 +2,11 @@
 
 #pragma once
 
+#include "Math/Vector.h"
 #include "RenderResource.h"
-#include "RHI.h"
+#include "RHIFwd.h"
 
+class FPositionVertexData;
 struct FStaticMeshBuildVertex;
 struct FConstMeshBuildVertexView;
 
@@ -126,7 +128,7 @@ private:
 	FShaderResourceViewRHIRef PositionComponentSRV;
 
 	/** The vertex data storage type */
-	TMemoryImagePtr<class FPositionVertexData> VertexData;
+	FPositionVertexData* VertexData;
 
 	/** The cached vertex data pointer. */
 	uint8* Data;
@@ -146,3 +148,7 @@ private:
 	FBufferRHIRef CreateRHIBuffer_Internal();
 };
 
+#if UE_ENABLE_INCLUDE_ORDER_DEPRECATED_IN_5_3
+#include "RHI.h"
+#endif
+
diff --git a/Engine/Source/Runtime/RHI/Private/RHI.cpp b/Engine/Source/Runtime/RHI/Private/RHI.cpp
--- a/Engine/Source/Runtime/RHI/Private/RHI.cpp
+++ b/Engine/Source/Runtime/RHI/Private/RHI.cpp
@@ -23,6 +23,7 @@
 #include "RHITextureReference.h"
 #include "RHIStats.h"
 #include "RHICommandList.h"
+#include "RHIUniformBufferLayoutInitializer.h"
 #include <type_traits>
 
 #if RHI_ENABLE_RESOURCE_INFO
@@ -42,7 +43,7 @@ CSV_DEFINE_CATEGORY(DrawCall, true);
 #endif
 
 IMPLEMENT_TYPE_LAYOUT(FRHIUniformBufferLayoutInitializer);
-IMPLEMENT_TYPE_LAYOUT(FRHIUniformBufferResource);
+IMPLEMENT_TYPE_LAYOUT(FRHIUniformBufferResourceInitializer);
 
 #if !defined(RHIRESOURCE_NUM_FRAMES_TO_EXPIRE)
 	#define RHIRESOURCE_NUM_FRAMES_TO_EXPIRE 3
diff --git a/Engine/Source/Runtime/RHI/Private/RHIResources.cpp b/Engine/Source/Runtime/RHI/Private/RHIResources.cpp
--- a/Engine/Source/Runtime/RHI/Private/RHIResources.cpp
+++ b/Engine/Source/Runtime/RHI/Private/RHIResources.cpp
@@ -5,6 +5,7 @@
 #include "Experimental/Containers/HazardPointer.h"
 #include "Misc/MemStack.h"
 #include "RHICommandList.h"
+#include "RHIUniformBufferLayoutInitializer.h"
 #include "Stats/Stats.h"
 #include "TextureProfiler.h"
 #include "RHIGlobals.h"
@@ -305,3 +306,32 @@ RHI_API FRHIViewDesc::FTextureUAV::FViewInfo FRHIViewDesc::FTextureUAV::GetViewI
 
 	return Info;
 }
+
+static_assert(offsetof(FRHIUniformBufferResource, MemberOffset) == offsetof(FRHIUniformBufferResourceInitializer, MemberOffset), "FRHIUniformBufferResource must be identical to FRHIUniformBufferResourceInitializer");
+static_assert(offsetof(FRHIUniformBufferResource, MemberType  ) == offsetof(FRHIUniformBufferResourceInitializer, MemberType  ), "FRHIUniformBufferResource must be identical to FRHIUniformBufferResourceInitializer");
+static_assert(sizeof(FRHIUniformBufferResource) == sizeof(FRHIUniformBufferResourceInitializer), "FRHIUniformBufferResource must be identical to FRHIUniformBufferResourceInitializer");
+
+template <>
+struct TIsBitwiseConstructible<FRHIUniformBufferResource, FRHIUniformBufferResourceInitializer>
+{
+	enum { Value = true };
+};
+
+FRHIUniformBufferLayout::FRHIUniformBufferLayout(const FRHIUniformBufferLayoutInitializer& Initializer)
+	: FRHIResource(RRT_UniformBufferLayout)
+	, Name(Initializer.GetDebugName())
+	, Resources(TConstArrayView<FRHIUniformBufferResourceInitializer>(Initializer.Resources))
+	, GraphResources(TConstArrayView<FRHIUniformBufferResourceInitializer>(Initializer.GraphResources))
+	, GraphTextures(TConstArrayView<FRHIUniformBufferResourceInitializer>(Initializer.GraphTextures))
+	, GraphBuffers(TConstArrayView<FRHIUniformBufferResourceInitializer>(Initializer.GraphBuffers))
+	, GraphUniformBuffers(TConstArrayView<FRHIUniformBufferResourceInitializer>(Initializer.GraphUniformBuffers))
+	, UniformBuffers(TConstArrayView<FRHIUniformBufferResourceInitializer>(Initializer.UniformBuffers))
+	, Hash(Initializer.GetHash())
+	, ConstantBufferSize(Initializer.ConstantBufferSize)
+	, RenderTargetsOffset(Initializer.RenderTargetsOffset)
+	, StaticSlot(Initializer.StaticSlot)
+	, BindingFlags(Initializer.BindingFlags)
+	, bHasNonGraphOutputs(Initializer.bHasNonGraphOutputs)
+	, bNoEmulatedUniformBuffer(Initializer.bNoEmulatedUniformBuffer)
+{
+}
diff --git a/Engine/Source/Runtime/RHI/Public/DynamicRHI.h b/Engine/Source/Runtime/RHI/Public/DynamicRHI.h
--- a/Engine/Source/Runtime/RHI/Public/DynamicRHI.h
+++ b/Engine/Source/Runtime/RHI/Public/DynamicRHI.h
@@ -26,6 +26,7 @@ class FRHIComputeFence;
 class FRayTracingPipelineState;
 class IRHITransientResourceAllocator;
 struct FDepthStencilStateInitializerRHI;
+struct FDisplayInformation;
 struct FRasterizerStateInitializerRHI;
 struct FRHIResourceCreateInfo;
 struct FRHIResourceInfo;
@@ -36,6 +37,7 @@ struct FRHIGPUMask;
 struct FUpdateTexture3DData;
 
 typedef TArray<FScreenResolutionRHI> FScreenResolutionArray;
+using FDisplayInformationArray = TArray<FDisplayInformation>;
 
 /** Struct to provide details of swap chain flips */
 struct FRHIFlipDetails
diff --git a/Engine/Source/Runtime/RHI/Public/RHI.h b/Engine/Source/Runtime/RHI/Public/RHI.h
--- a/Engine/Source/Runtime/RHI/Public/RHI.h
+++ b/Engine/Source/Runtime/RHI/Public/RHI.h
@@ -7,7 +7,6 @@
 #pragma once
 
 #include "Containers/ContainersFwd.h"
-#include "Stats/Stats.h"
 #include "RHIDefinitions.h"
 #include "Containers/StaticArray.h"
 #include "Containers/StringFwd.h"
@@ -16,7 +15,6 @@
 #include "Math/ScaleMatrix.h"
 #include "Math/TranslationMatrix.h"
 #include "PixelFormat.h"
-#include "ProfilingDebugging/CsvProfilerConfig.h"
 #include "GpuProfilerTrace.h"
 #include "RHIShaderPlatform.h"
 #include "RHIFeatureLevel.h"
@@ -651,5 +649,7 @@ inline ERHIBindlessConfiguration RHIGetBindlessSamplersConfiguration(EShaderPlat
 
 #if UE_ENABLE_INCLUDE_ORDER_DEPRECATED_IN_5_3
 #include "CoreMinimal.h"
+#include "ProfilingDebugging/CsvProfilerConfig.h"
 #include "RHIUtilities.h"
+#include "Stats/Stats.h"
 #endif
diff --git a/Engine/Source/Runtime/RHI/Public/RHIDefinitions.h b/Engine/Source/Runtime/RHI/Public/RHIDefinitions.h
--- a/Engine/Source/Runtime/RHI/Public/RHIDefinitions.h
+++ b/Engine/Source/Runtime/RHI/Public/RHIDefinitions.h
@@ -8,7 +8,8 @@
 #pragma once
 
 #include "GpuProfilerTrace.h" // TODO Move defines into RHIDefinitions
-#include "Serialization/MemoryLayout.h"
+#include "Math/NumericLimits.h"
+#include "Misc/EnumClassFlags.h"
 #include "ProfilingDebugging/CsvProfilerConfig.h" // TODO Move defines into RHIDefinitions
 
 #ifndef USE_STATIC_SHADER_PLATFORM_ENUMS
@@ -400,7 +401,6 @@ enum EVertexElementType
 	VET_NumBits = 5,
 };
 static_assert(VET_MAX <= (1 << VET_NumBits), "VET_MAX will not fit on VET_NumBits");
-DECLARE_INTRINSIC_TYPE_LAYOUT(EVertexElementType);
 
 enum ECubeFace : uint32
 {
@@ -519,7 +519,6 @@ enum EUniformBufferBaseType : uint8
 	EUniformBufferBaseType_NumBits = 5,
 };
 static_assert(EUniformBufferBaseType_Num <= (1 << EUniformBufferBaseType_NumBits), "EUniformBufferBaseType_Num will not fit on EUniformBufferBaseType_NumBits");
-DECLARE_INTRINSIC_TYPE_LAYOUT(EUniformBufferBaseType);
 
 /** The list of flags declaring which binding models are allowed for a uniform buffer layout. */
 enum class EUniformBufferBindingFlags : uint8
@@ -538,7 +537,6 @@ enum class EUniformBufferBindingFlags : uint8
 	StaticAndShader = Static | Shader
 };
 ENUM_CLASS_FLAGS(EUniformBufferBindingFlags);
-DECLARE_INTRINSIC_TYPE_LAYOUT(EUniformBufferBindingFlags);
 
 /** Numerical type used to store the static slot indices. */
 using FUniformBufferStaticSlot = uint8;
@@ -1137,15 +1135,13 @@ struct FRHIDescriptorHandle
 	inline ERHIDescriptorHeapType GetType() const { return (ERHIDescriptorHeapType)Type; }
 	inline uint8                  GetRawType() const { return Type; }
 
-	inline bool IsValid() const { return Index != UINT_MAX && Type != (uint8)ERHIDescriptorHeapType::Invalid; }
+	inline bool IsValid() const { return Index != MAX_uint32 && Type != (uint8)ERHIDescriptorHeapType::Invalid; }
 
 private:
-	uint32    Index{ UINT_MAX };
+	uint32    Index{ MAX_uint32 };
 	uint8     Type{ (uint8)ERHIDescriptorHeapType::Invalid };
 };
 
-using FDisplayInformationArray = TArray<struct FDisplayInformation>;
-
 enum class ERHIBindlessConfiguration
 {
 	Disabled,
@@ -1334,7 +1330,6 @@ inline ERHIResourceType GetRHIResourceType(ETextureDimension Dimension)
 	case ETextureDimension::TextureCubeArray:
 		return ERHIResourceType::RRT_TextureCube;
 	}
-	checkNoEntry();
 	return ERHIResourceType::RRT_None;
 }
 
@@ -1368,3 +1363,7 @@ struct FShaderCodeValidationStride
 #include "RHIShaderPlatform.h"
 #include "RHIStrings.h"
 #endif
+
+#if UE_ENABLE_INCLUDE_ORDER_DEPRECATED_IN_5_3
+#include "Serialization/MemoryLayout.h"
+#endif
diff --git a/Engine/Source/Runtime/RHI/Public/RHIFeatureLevel.h b/Engine/Source/Runtime/RHI/Public/RHIFeatureLevel.h
--- a/Engine/Source/Runtime/RHI/Public/RHIFeatureLevel.h
+++ b/Engine/Source/Runtime/RHI/Public/RHIFeatureLevel.h
@@ -4,7 +4,6 @@
 
 #include "Containers/EnumAsByte.h"
 #include "RHIDefinitions.h"
-#include "Serialization/MemoryLayout.h"
 
 #ifndef USE_STATIC_FEATURE_LEVEL_ENUMS
 #define USE_STATIC_FEATURE_LEVEL_ENUMS 0
@@ -55,7 +54,6 @@ namespace ERHIFeatureLevel
 		Num
 	};
 };
-DECLARE_INTRINSIC_TYPE_LAYOUT(ERHIFeatureLevel::Type);
 
 struct FGenericStaticFeatureLevel
 {
@@ -113,3 +111,7 @@ extern RHI_API ERHIFeatureLevel::Type GMaxRHIFeatureLevel;
 #if UE_ENABLE_INCLUDE_ORDER_DEPRECATED_IN_5_2
 #include "RHIStrings.h"
 #endif
+
+#if UE_ENABLE_INCLUDE_ORDER_DEPRECATED_IN_5_3
+#include "Serialization/MemoryLayout.h"
+#endif
diff --git a/Engine/Source/Runtime/RHI/Public/RHIMemoryLayout.h b/Engine/Source/Runtime/RHI/Public/RHIMemoryLayout.h
new file mode 100644
--- /dev/null
+++ b/Engine/Source/Runtime/RHI/Public/RHIMemoryLayout.h
@@ -0,0 +1,12 @@
+// Copyright Epic Games, Inc. All Rights Reserved.
+
+#pragma once
+
+#include "RHIDefinitions.h"
+#include "RHIFeatureLevel.h"
+#include "Serialization/MemoryLayout.h"
+
+DECLARE_INTRINSIC_TYPE_LAYOUT(EVertexElementType);
+DECLARE_INTRINSIC_TYPE_LAYOUT(EUniformBufferBaseType);
+DECLARE_INTRINSIC_TYPE_LAYOUT(EUniformBufferBindingFlags);
+DECLARE_INTRINSIC_TYPE_LAYOUT(ERHIFeatureLevel::Type);
diff --git a/Engine/Source/Runtime/RHI/Public/RHIResources.h b/Engine/Source/Runtime/RHI/Public/RHIResources.h
--- a/Engine/Source/Runtime/RHI/Public/RHIResources.h
+++ b/Engine/Source/Runtime/RHI/Public/RHIResources.h
@@ -15,7 +15,6 @@
 #include "Templates/RefCounting.h"
 #include "PixelFormat.h"
 #include "Async/TaskGraphFwd.h"
-#include "Serialization/MemoryImage.h"
 #include "RHIFwd.h"
 #include "RHIImmutableSamplerState.h"
 #include "RHITransition.h"
@@ -23,18 +22,7 @@
 #include "Math/IntPoint.h"
 #include "Math/IntRect.h"
 #include "Math/IntVector.h"
-
-#if UE_ENABLE_INCLUDE_ORDER_DEPRECATED_IN_5_2
-#include "TextureProfiler.h"
-#include "Containers/LockFreeList.h"
-#include "Misc/CoreDelegates.h"
 #include "Misc/SecureHash.h"
-#include "Hash/CityHash.h"
-#include "Async/TaskGraphInterfaces.h"
-#include "Experimental/Containers/HazardPointer.h"
-#include "Containers/ClosableMpscQueue.h"
-#include "Containers/ConsumeAllMpmcQueue.h"
-#endif
 
 #include <atomic>
 
@@ -942,16 +930,13 @@ public:
 // Enabling this requires -norhithread to work correctly since FRHIResource lifetime is managed by both the RT and RHIThread
 #define VALIDATE_UNIFORM_BUFFER_LIFETIME 0
 
-/** Data structure to store information about resource parameter in a shader parameter structure. */
 struct FRHIUniformBufferResource
 {
-	DECLARE_EXPORTED_TYPE_LAYOUT(FRHIUniformBufferResource, RHI_API, NonVirtual);
-
 	/** Byte offset to each resource in the uniform buffer memory. */
-	LAYOUT_FIELD(uint16, MemberOffset);
+	uint16 MemberOffset;
 
 	/** Type of the member that allow (). */
-	LAYOUT_FIELD(EUniformBufferBaseType, MemberType);
+	EUniformBufferBaseType MemberType;
 
 	/** Compare two uniform buffer layout resources. */
 	friend inline bool operator==(const FRHIUniformBufferResource& A, const FRHIUniformBufferResource& B)
@@ -961,198 +946,16 @@ struct FRHIUniformBufferResource
 	}
 };
 
-inline FArchive& operator<<(FArchive& Ar, FRHIUniformBufferResource& Ref)
-{
-	uint8 Type = (uint8)Ref.MemberType;
-	Ar << Ref.MemberOffset;
-	Ar << Type;
-	Ref.MemberType = (EUniformBufferBaseType)Type;
-	return Ar;
-}
-
 inline constexpr uint16 kUniformBufferInvalidOffset = TNumericLimits<uint16>::Max();
 
-/** Initializer for the layout of a uniform buffer in memory. */
-struct FRHIUniformBufferLayoutInitializer
-{
-	DECLARE_EXPORTED_TYPE_LAYOUT(FRHIUniformBufferLayoutInitializer, RHI_API, NonVirtual);
-
-	FRHIUniformBufferLayoutInitializer() = default;
-
-	explicit FRHIUniformBufferLayoutInitializer(const TCHAR * InName)
-		: Name(InName)
-	{}
-	explicit FRHIUniformBufferLayoutInitializer(const TCHAR * InName, uint32 InConstantBufferSize)
-		: Name(InName)
-		, ConstantBufferSize(InConstantBufferSize)
-	{
-		ComputeHash();
-	}
-
-	inline uint32 GetHash() const
-	{
-		checkSlow(Hash != 0);
-		return Hash;
-	}
-
-	void ComputeHash()
-	{
-		// Static slot is not stable. Just track whether we have one at all.
-		uint32 TmpHash = ConstantBufferSize << 16 | static_cast<uint32>(BindingFlags) << 8 | static_cast<uint32>(StaticSlot != MAX_UNIFORM_BUFFER_STATIC_SLOTS);
-
-		for (int32 ResourceIndex = 0; ResourceIndex < Resources.Num(); ResourceIndex++)
-		{
-			// Offset and therefore hash must be the same regardless of pointer size
-			checkSlow(Resources[ResourceIndex].MemberOffset == Align(Resources[ResourceIndex].MemberOffset, SHADER_PARAMETER_POINTER_ALIGNMENT));
-			TmpHash ^= Resources[ResourceIndex].MemberOffset;
-		}
-
-		uint32 N = Resources.Num();
-		while (N >= 4)
-		{
-			TmpHash ^= (Resources[--N].MemberType << 0);
-			TmpHash ^= (Resources[--N].MemberType << 8);
-			TmpHash ^= (Resources[--N].MemberType << 16);
-			TmpHash ^= (Resources[--N].MemberType << 24);
-		}
-		while (N >= 2)
-		{
-			TmpHash ^= Resources[--N].MemberType << 0;
-			TmpHash ^= Resources[--N].MemberType << 16;
-		}
-		while (N > 0)
-		{
-			TmpHash ^= Resources[--N].MemberType;
-		}
-		Hash = TmpHash;
-	}
-
-	void CopyFrom(const FRHIUniformBufferLayoutInitializer& Source)
-	{
-		ConstantBufferSize = Source.ConstantBufferSize;
-		StaticSlot = Source.StaticSlot;
-		BindingFlags = Source.BindingFlags;
-		Resources = Source.Resources;
-		Name = Source.Name;
-		Hash = Source.Hash;
-	}
-
-	const FMemoryImageString& GetDebugName() const
-	{
-		return Name;
-	}
-
-	bool HasRenderTargets() const
-	{
-		return RenderTargetsOffset != kUniformBufferInvalidOffset;
-	}
-
-	bool HasExternalOutputs() const
-	{
-		return bHasNonGraphOutputs;
-	}
-
-	bool HasStaticSlot() const
-	{
-		return IsUniformBufferStaticSlotValid(StaticSlot);
-	}
-
-	friend FArchive& operator<<(FArchive& Ar, FRHIUniformBufferLayoutInitializer& Ref)
-	{
-		Ar << Ref.ConstantBufferSize;
-		Ar << Ref.StaticSlot;
-		Ar << Ref.RenderTargetsOffset;
-		Ar << Ref.bHasNonGraphOutputs;
-		Ar << Ref.BindingFlags;
-		Ar << Ref.Resources;
-		Ar << Ref.GraphResources;
-		Ar << Ref.GraphTextures;
-		Ar << Ref.GraphBuffers;
-		Ar << Ref.GraphUniformBuffers;
-		Ar << Ref.UniformBuffers;
-		Ar << Ref.Name;
-		Ar << Ref.Hash;
-		Ar << Ref.bNoEmulatedUniformBuffer;
-		return Ar;
-	}
-
-private:
-	// for debugging / error message
-	LAYOUT_FIELD(FMemoryImageString, Name);
-
-public:
-	/** The list of all resource inlined into the shader parameter structure. */
-	LAYOUT_FIELD(TMemoryImageArray<FRHIUniformBufferResource>, Resources);
-
-	/** The list of all RDG resource references inlined into the shader parameter structure. */
-	LAYOUT_FIELD(TMemoryImageArray<FRHIUniformBufferResource>, GraphResources);
-
-	/** The list of all RDG texture references inlined into the shader parameter structure. */
-	LAYOUT_FIELD(TMemoryImageArray<FRHIUniformBufferResource>, GraphTextures);
-
-	/** The list of all RDG buffer references inlined into the shader parameter structure. */
-	LAYOUT_FIELD(TMemoryImageArray<FRHIUniformBufferResource>, GraphBuffers);
-
-	/** The list of all RDG uniform buffer references inlined into the shader parameter structure. */
-	LAYOUT_FIELD(TMemoryImageArray<FRHIUniformBufferResource>, GraphUniformBuffers);
-
-	/** The list of all non-RDG uniform buffer references inlined into the shader parameter structure. */
-	LAYOUT_FIELD(TMemoryImageArray<FRHIUniformBufferResource>, UniformBuffers);
-
-private:
-	LAYOUT_FIELD_INITIALIZED(uint32, Hash, 0);
-
-public:
-	/** The size of the constant buffer in bytes. */
-	LAYOUT_FIELD_INITIALIZED(uint32, ConstantBufferSize, 0);
-
-	/** The render target binding slots offset, if it exists. */
-	LAYOUT_FIELD_INITIALIZED(uint16, RenderTargetsOffset, kUniformBufferInvalidOffset);
-
-	/** The static slot (if applicable). */
-	LAYOUT_FIELD_INITIALIZED(FUniformBufferStaticSlot, StaticSlot, MAX_UNIFORM_BUFFER_STATIC_SLOTS);
-
-	/** The binding flags describing how this resource can be bound to the RHI. */
-	LAYOUT_FIELD_INITIALIZED(EUniformBufferBindingFlags, BindingFlags, EUniformBufferBindingFlags::Shader);
-
-	/** Whether this layout may contain non-render-graph outputs (e.g. RHI UAVs). */
-	LAYOUT_FIELD_INITIALIZED(bool, bHasNonGraphOutputs, false);
-
-	/** Used for platforms which use emulated ub's, forces a real uniform buffer instead */
-	LAYOUT_FIELD_INITIALIZED(bool, bNoEmulatedUniformBuffer, false);
-
-	/** Compare two uniform buffer layout initializers. */
-	friend inline bool operator==(const FRHIUniformBufferLayoutInitializer& A, const FRHIUniformBufferLayoutInitializer& B)
-	{
-		return A.ConstantBufferSize == B.ConstantBufferSize
-			&& A.StaticSlot == B.StaticSlot
-			&& A.BindingFlags == B.BindingFlags
-			&& A.Resources == B.Resources;
-	}
-};
+struct FRHIUniformBufferLayoutInitializer;
 
 /** The layout of a uniform buffer in memory. */
 struct FRHIUniformBufferLayout : public FRHIResource
 {
 	FRHIUniformBufferLayout() = delete;
 
-	explicit FRHIUniformBufferLayout(const FRHIUniformBufferLayoutInitializer& Initializer)
-		: FRHIResource(RRT_UniformBufferLayout)
-		, Name(Initializer.GetDebugName())
-		, Resources(Initializer.Resources)
-		, GraphResources(Initializer.GraphResources)
-		, GraphTextures(Initializer.GraphTextures)
-		, GraphBuffers(Initializer.GraphBuffers)
-		, GraphUniformBuffers(Initializer.GraphUniformBuffers)
-		, UniformBuffers(Initializer.UniformBuffers)
-		, Hash(Initializer.GetHash())
-		, ConstantBufferSize(Initializer.ConstantBufferSize)
-		, RenderTargetsOffset(Initializer.RenderTargetsOffset)
-		, StaticSlot(Initializer.StaticSlot)
-		, BindingFlags(Initializer.BindingFlags)
-		, bHasNonGraphOutputs(Initializer.bHasNonGraphOutputs)
-		, bNoEmulatedUniformBuffer(Initializer.bNoEmulatedUniformBuffer)
-	{}
+	RHI_API explicit FRHIUniformBufferLayout(const FRHIUniformBufferLayoutInitializer& Initializer);
 
 	inline const FString& GetDebugName() const
 	{
@@ -4916,6 +4719,19 @@ private:
 };
 
 #if UE_ENABLE_INCLUDE_ORDER_DEPRECATED_IN_5_2
+#include "Async/TaskGraphInterfaces.h"
+#include "Containers/ClosableMpscQueue.h"
+#include "Containers/ConsumeAllMpmcQueue.h"
+#include "Containers/LockFreeList.h"
+#include "Experimental/Containers/HazardPointer.h"
+#include "Hash/CityHash.h"
+#include "Misc/CoreDelegates.h"
+#include "Misc/SecureHash.h"
 #include "RHIShaderLibrary.h"
 #include "RHITextureReference.h"
+#include "TextureProfiler.h"
+#endif
+
+#if UE_ENABLE_INCLUDE_ORDER_DEPRECATED_IN_5_3
+#include "Serialization/MemoryImage.h"
 #endif
diff --git a/Engine/Source/Runtime/RHI/Public/RHIStrings.h b/Engine/Source/Runtime/RHI/Public/RHIStrings.h
--- a/Engine/Source/Runtime/RHI/Public/RHIStrings.h
+++ b/Engine/Source/Runtime/RHI/Public/RHIStrings.h
@@ -92,7 +92,6 @@ inline const TCHAR* GetShaderFrequencyString(EShaderFrequency Frequency, bool bI
 	case SF_RayCallable:	String = TEXT("SF_RayCallable"); break;
 
 	default:
-		checkf(0, TEXT("Unknown ShaderFrequency %d"), (int32)Frequency);
 		break;
 	}
 
diff --git a/Engine/Source/Runtime/RHI/Public/RHIUniformBufferLayoutInitializer.h b/Engine/Source/Runtime/RHI/Public/RHIUniformBufferLayoutInitializer.h
new file mode 100644
--- /dev/null
+++ b/Engine/Source/Runtime/RHI/Public/RHIUniformBufferLayoutInitializer.h
@@ -0,0 +1,197 @@
+// Copyright Epic Games, Inc. All Rights Reserved.
+
+#pragma once
+
+#include "RHIDefinitions.h"
+#include "RHIMemoryLayout.h"
+#include "RHIResources.h"
+#include "Serialization/MemoryImage.h"
+#include "Serialization/MemoryLayout.h"
+
+/** Data structure to store information about resource parameter in a shader parameter structure. */
+struct FRHIUniformBufferResourceInitializer
+{
+	DECLARE_EXPORTED_TYPE_LAYOUT(FRHIUniformBufferResourceInitializer, RHI_API, NonVirtual);
+
+	/** Byte offset to each resource in the uniform buffer memory. */
+	LAYOUT_FIELD(uint16, MemberOffset);
+
+	/** Type of the member that allow (). */
+	LAYOUT_FIELD(EUniformBufferBaseType, MemberType);
+
+	/** Compare two uniform buffer layout resources. */
+	friend inline bool operator==(const FRHIUniformBufferResourceInitializer& A, const FRHIUniformBufferResourceInitializer& B)
+	{
+		return A.MemberOffset == B.MemberOffset
+			&& A.MemberType == B.MemberType;
+	}
+};
+
+inline FArchive& operator<<(FArchive& Ar, FRHIUniformBufferResourceInitializer& Ref)
+{
+	uint8 Type = (uint8)Ref.MemberType;
+	Ar << Ref.MemberOffset;
+	Ar << Type;
+	Ref.MemberType = (EUniformBufferBaseType)Type;
+	return Ar;
+}
+
+/** Initializer for the layout of a uniform buffer in memory. */
+struct FRHIUniformBufferLayoutInitializer
+{
+	DECLARE_EXPORTED_TYPE_LAYOUT(FRHIUniformBufferLayoutInitializer, RHI_API, NonVirtual);
+
+	FRHIUniformBufferLayoutInitializer() = default;
+
+	explicit FRHIUniformBufferLayoutInitializer(const TCHAR* InName)
+		: Name(InName)
+	{}
+	explicit FRHIUniformBufferLayoutInitializer(const TCHAR* InName, uint32 InConstantBufferSize)
+		: Name(InName)
+		, ConstantBufferSize(InConstantBufferSize)
+	{
+		ComputeHash();
+	}
+
+	inline uint32 GetHash() const
+	{
+		checkSlow(Hash != 0);
+		return Hash;
+	}
+
+	void ComputeHash()
+	{
+		// Static slot is not stable. Just track whether we have one at all.
+		uint32 TmpHash = ConstantBufferSize << 16 | static_cast<uint32>(BindingFlags) << 8 | static_cast<uint32>(StaticSlot != MAX_UNIFORM_BUFFER_STATIC_SLOTS);
+
+		for (int32 ResourceIndex = 0; ResourceIndex < Resources.Num(); ResourceIndex++)
+		{
+			// Offset and therefore hash must be the same regardless of pointer size
+			checkSlow(Resources[ResourceIndex].MemberOffset == Align(Resources[ResourceIndex].MemberOffset, SHADER_PARAMETER_POINTER_ALIGNMENT));
+			TmpHash ^= Resources[ResourceIndex].MemberOffset;
+		}
+
+		uint32 N = Resources.Num();
+		while (N >= 4)
+		{
+			TmpHash ^= (Resources[--N].MemberType << 0);
+			TmpHash ^= (Resources[--N].MemberType << 8);
+			TmpHash ^= (Resources[--N].MemberType << 16);
+			TmpHash ^= (Resources[--N].MemberType << 24);
+		}
+		while (N >= 2)
+		{
+			TmpHash ^= Resources[--N].MemberType << 0;
+			TmpHash ^= Resources[--N].MemberType << 16;
+		}
+		while (N > 0)
+		{
+			TmpHash ^= Resources[--N].MemberType;
+		}
+		Hash = TmpHash;
+	}
+
+	void CopyFrom(const FRHIUniformBufferLayoutInitializer& Source)
+	{
+		ConstantBufferSize = Source.ConstantBufferSize;
+		StaticSlot = Source.StaticSlot;
+		BindingFlags = Source.BindingFlags;
+		Resources = Source.Resources;
+		Name = Source.Name;
+		Hash = Source.Hash;
+	}
+
+	const FMemoryImageString& GetDebugName() const
+	{
+		return Name;
+	}
+
+	bool HasRenderTargets() const
+	{
+		return RenderTargetsOffset != kUniformBufferInvalidOffset;
+	}
+
+	bool HasExternalOutputs() const
+	{
+		return bHasNonGraphOutputs;
+	}
+
+	bool HasStaticSlot() const
+	{
+		return IsUniformBufferStaticSlotValid(StaticSlot);
+	}
+
+	friend FArchive& operator<<(FArchive& Ar, FRHIUniformBufferLayoutInitializer& Ref)
+	{
+		Ar << Ref.ConstantBufferSize;
+		Ar << Ref.StaticSlot;
+		Ar << Ref.RenderTargetsOffset;
+		Ar << Ref.bHasNonGraphOutputs;
+		Ar << Ref.BindingFlags;
+		Ar << Ref.Resources;
+		Ar << Ref.GraphResources;
+		Ar << Ref.GraphTextures;
+		Ar << Ref.GraphBuffers;
+		Ar << Ref.GraphUniformBuffers;
+		Ar << Ref.UniformBuffers;
+		Ar << Ref.Name;
+		Ar << Ref.Hash;
+		Ar << Ref.bNoEmulatedUniformBuffer;
+		return Ar;
+	}
+
+private:
+	// for debugging / error message
+	LAYOUT_FIELD(FMemoryImageString, Name);
+
+public:
+	/** The list of all resource inlined into the shader parameter structure. */
+	LAYOUT_FIELD(TMemoryImageArray<FRHIUniformBufferResourceInitializer>, Resources);
+
+	/** The list of all RDG resource references inlined into the shader parameter structure. */
+	LAYOUT_FIELD(TMemoryImageArray<FRHIUniformBufferResourceInitializer>, GraphResources);
+
+	/** The list of all RDG texture references inlined into the shader parameter structure. */
+	LAYOUT_FIELD(TMemoryImageArray<FRHIUniformBufferResourceInitializer>, GraphTextures);
+
+	/** The list of all RDG buffer references inlined into the shader parameter structure. */
+	LAYOUT_FIELD(TMemoryImageArray<FRHIUniformBufferResourceInitializer>, GraphBuffers);
+
+	/** The list of all RDG uniform buffer references inlined into the shader parameter structure. */
+	LAYOUT_FIELD(TMemoryImageArray<FRHIUniformBufferResourceInitializer>, GraphUniformBuffers);
+
+	/** The list of all non-RDG uniform buffer references inlined into the shader parameter structure. */
+	LAYOUT_FIELD(TMemoryImageArray<FRHIUniformBufferResourceInitializer>, UniformBuffers);
+
+private:
+	LAYOUT_FIELD_INITIALIZED(uint32, Hash, 0);
+
+public:
+	/** The size of the constant buffer in bytes. */
+	LAYOUT_FIELD_INITIALIZED(uint32, ConstantBufferSize, 0);
+
+	/** The render target binding slots offset, if it exists. */
+	LAYOUT_FIELD_INITIALIZED(uint16, RenderTargetsOffset, kUniformBufferInvalidOffset);
+
+	/** The static slot (if applicable). */
+	LAYOUT_FIELD_INITIALIZED(FUniformBufferStaticSlot, StaticSlot, MAX_UNIFORM_BUFFER_STATIC_SLOTS);
+
+	/** The binding flags describing how this resource can be bound to the RHI. */
+	LAYOUT_FIELD_INITIALIZED(EUniformBufferBindingFlags, BindingFlags, EUniformBufferBindingFlags::Shader);
+
+	/** Whether this layout may contain non-render-graph outputs (e.g. RHI UAVs). */
+	LAYOUT_FIELD_INITIALIZED(bool, bHasNonGraphOutputs, false);
+
+	/** Used for platforms which use emulated ub's, forces a real uniform buffer instead */
+	LAYOUT_FIELD_INITIALIZED(bool, bNoEmulatedUniformBuffer, false);
+
+	/** Compare two uniform buffer layout initializers. */
+	friend inline bool operator==(const FRHIUniformBufferLayoutInitializer& A, const FRHIUniformBufferLayoutInitializer& B)
+	{
+		return A.ConstantBufferSize == B.ConstantBufferSize
+			&& A.StaticSlot == B.StaticSlot
+			&& A.BindingFlags == B.BindingFlags
+			&& A.Resources == B.Resources;
+	}
+};
+
diff --git a/Engine/Source/Runtime/RHICore/Internal/DXGIUtilities.h b/Engine/Source/Runtime/RHICore/Internal/DXGIUtilities.h
--- a/Engine/Source/Runtime/RHICore/Internal/DXGIUtilities.h
+++ b/Engine/Source/Runtime/RHICore/Internal/DXGIUtilities.h
@@ -3,6 +3,7 @@
 #pragma once
 
 #include "HAL/Platform.h"
+#include "Misc/AssertionMacros.h"
 
 #if PLATFORM_MICROSOFT
 
diff --git a/Engine/Source/Runtime/RenderCore/Private/ShaderParameterMetadata.cpp b/Engine/Source/Runtime/RenderCore/Private/ShaderParameterMetadata.cpp
--- a/Engine/Source/Runtime/RenderCore/Private/ShaderParameterMetadata.cpp
+++ b/Engine/Source/Runtime/RenderCore/Private/ShaderParameterMetadata.cpp
@@ -6,6 +6,7 @@
 
 #include "ShaderParameterMetadata.h"
 #include "RenderCore.h"
+#include "RHIUniformBufferLayoutInitializer.h"
 #include "ShaderCore.h"
 #include "ShaderParameters.h"
 #include "DataDrivenShaderPlatformInfo.h"
@@ -692,8 +693,8 @@ void FShaderParametersMetadata::InitializeLayout(FRHIUniformBufferLayoutInitiali
 			for (uint32 ArrayElementId = 0; ArrayElementId < (bIsArray ? ArraySize : 1u); ArrayElementId++)
 			{
 				const uint32 AbsoluteMemberOffset = CurrentMember.GetOffset() + MemberStack[i].StructOffset + ArrayElementId * SHADER_PARAMETER_POINTER_ALIGNMENT;
-				check(AbsoluteMemberOffset < (1u << (sizeof(FRHIUniformBufferResource::MemberOffset) * 8)));
-				const FRHIUniformBufferResource ResourceParameter{ uint16(AbsoluteMemberOffset), BaseType };
+				check(AbsoluteMemberOffset < (1u << (sizeof(FRHIUniformBufferResourceInitializer::MemberOffset) * 8)));
+				const FRHIUniformBufferResourceInitializer ResourceParameter{ uint16(AbsoluteMemberOffset), BaseType };
 
 				LayoutInitializer.Resources.Add(ResourceParameter);
 
@@ -741,15 +742,15 @@ void FShaderParametersMetadata::InitializeLayout(FRHIUniformBufferLayoutInitiali
 	}
 
 	const auto ByMemberOffset = [](
-		const FRHIUniformBufferResource& A,
-		const FRHIUniformBufferResource& B)
+		const FRHIUniformBufferResourceInitializer& A,
+		const FRHIUniformBufferResourceInitializer& B)
 	{
 		return A.MemberOffset < B.MemberOffset;
 	};
 
 	const auto ByTypeThenMemberOffset = [](
-		const FRHIUniformBufferResource& A,
-		const FRHIUniformBufferResource& B)
+		const FRHIUniformBufferResourceInitializer& A,
+		const FRHIUniformBufferResourceInitializer& B)
 	{
 		if (A.MemberType == B.MemberType)
 		{
diff --git a/Engine/Source/Runtime/RenderCore/Public/HDRHelper.h b/Engine/Source/Runtime/RenderCore/Public/HDRHelper.h
--- a/Engine/Source/Runtime/RenderCore/Public/HDRHelper.h
+++ b/Engine/Source/Runtime/RenderCore/Public/HDRHelper.h
@@ -31,6 +31,8 @@ struct FDisplayInformation
 	bool   bHDRSupported;
 };
 
+using FDisplayInformationArray = TArray<FDisplayInformation>;
+
 struct FACESTonemapParams
 {
 	// Default Values such as MinLuminanceLog10=-4, MidLuminance=15, MaxLuminance=1000, SceneColorMultiplier=1.5
diff --git a/Engine/Source/Runtime/RenderCore/Public/Shader.h b/Engine/Source/Runtime/RenderCore/Public/Shader.h
--- a/Engine/Source/Runtime/RenderCore/Public/Shader.h
+++ b/Engine/Source/Runtime/RenderCore/Public/Shader.h
@@ -31,6 +31,7 @@
 #include "Misc/SecureHash.h"
 #include "RHI.h"
 #include "RHIDefinitions.h"
+#include "RHIMemoryLayout.h"
 #include "RenderResource.h"
 #if UE_ENABLE_INCLUDE_ORDER_DEPRECATED_IN_5_2
 #include "RenderingThread.h"
