Added RHI helpers to support command list initialization of resources.

diff --git a/Engine/Source/Runtime/RHI/Public/RHICommandList.h b/Engine/Source/Runtime/RHI/Public/RHICommandList.h
--- a/Engine/Source/Runtime/RHI/Public/RHICommandList.h
+++ b/Engine/Source/Runtime/RHI/Public/RHICommandList.h
@@ -846,7 +846,7 @@ public:
 	}
 
 	//UE_DEPRECATED(5.3, "Use the CreateUnorderedAccessView function that takes an FRHIViewDesc.")
-	FORCEINLINE FUnorderedAccessViewRHIRef CreateUnorderedAccessView(FRHITexture* Texture, uint32 MipLevel, uint16 FirstArraySlice, uint16 NumArraySlices)
+	FORCEINLINE FUnorderedAccessViewRHIRef CreateUnorderedAccessView(FRHITexture* Texture, uint32 MipLevel = 0, uint16 FirstArraySlice = 0, uint16 NumArraySlices = 0)
 	{
 		check(MipLevel < 256);
 
@@ -858,7 +858,7 @@ public:
 	}
 
 	//UE_DEPRECATED(5.3, "Use the CreateUnorderedAccessView function that takes an FRHIViewDesc.")
-	FORCEINLINE FUnorderedAccessViewRHIRef CreateUnorderedAccessView(FRHITexture* Texture, uint32 MipLevel, uint8 Format, uint16 FirstArraySlice, uint16 NumArraySlices)
+	FORCEINLINE FUnorderedAccessViewRHIRef CreateUnorderedAccessView(FRHITexture* Texture, uint32 MipLevel, uint8 Format, uint16 FirstArraySlice = 0, uint16 NumArraySlices = 0)
 	{
 		check(MipLevel < 256);
 
@@ -2365,6 +2365,11 @@ public:
 	}
 
 public:
+	static inline FRHIComputeCommandList& Get(FRHICommandListBase& RHICmdList)
+	{
+		return static_cast<FRHIComputeCommandList&>(RHICmdList);
+	}
+
 	FRHIComputeCommandList(FRHIGPUMask GPUMask = FRHIGPUMask::All(), ERecordingThread InRecordingThread = ERecordingThread::Render)
 		: FRHICommandListBase(GPUMask, InRecordingThread)
 	{}
@@ -3073,6 +3078,11 @@ protected:
 	}
 
 public:
+	static inline FRHICommandList& Get(FRHICommandListBase& RHICmdList)
+	{
+		return static_cast<FRHICommandList&>(RHICmdList);
+	}
+
 	FRHICommandList(FRHIGPUMask GPUMask = FRHIGPUMask::All(), ERecordingThread InRecordingThread = ERecordingThread::Render)
 		: FRHIComputeCommandList(GPUMask, InRecordingThread)
 	{}
@@ -4069,6 +4079,12 @@ class FRHICommandListImmediate : public FRHICommandList
 public:
 	static inline FRHICommandListImmediate& Get();
 
+	static inline FRHICommandListImmediate& Get(FRHICommandListBase& RHICmdList)
+	{
+		check(RHICmdList.IsImmediate());
+		return static_cast<FRHICommandListImmediate&>(RHICmdList);
+	}
+
 	RHI_API void BeginScene();
 	RHI_API void EndScene();
 	RHI_API void BeginDrawingViewport(FRHIViewport* Viewport, FRHITexture* RenderTargetRHI);
diff --git a/Engine/Source/Runtime/RHI/Public/RHIUtilities.h b/Engine/Source/Runtime/RHI/Public/RHIUtilities.h
--- a/Engine/Source/Runtime/RHI/Public/RHIUtilities.h
+++ b/Engine/Source/Runtime/RHI/Public/RHIUtilities.h
@@ -301,13 +301,18 @@ struct FReadBuffer
 
 	FReadBuffer(): NumBytes(0) {}
 
-	void Initialize(const TCHAR* InDebugName, uint32 BytesPerElement, uint32 NumElements, EPixelFormat Format, EBufferUsageFlags AdditionalUsage = BUF_None, FResourceArrayInterface* InResourceArray = nullptr)
+	void Initialize(FRHICommandListBase& RHICmdList, const TCHAR* InDebugName, uint32 BytesPerElement, uint32 NumElements, EPixelFormat Format, EBufferUsageFlags AdditionalUsage = BUF_None, FResourceArrayInterface* InResourceArray = nullptr)
 	{
 		NumBytes = BytesPerElement * NumElements;
 		FRHIResourceCreateInfo CreateInfo(InDebugName);
 		CreateInfo.ResourceArray = InResourceArray;
-		Buffer = RHICreateVertexBuffer(NumBytes, BUF_ShaderResource | AdditionalUsage, ERHIAccess::SRVMask, CreateInfo);
-		SRV = RHICreateShaderResourceView(Buffer, BytesPerElement, UE_PIXELFORMAT_TO_UINT8(Format));
+		Buffer = RHICmdList.CreateVertexBuffer(NumBytes, BUF_ShaderResource | AdditionalUsage, ERHIAccess::SRVMask, CreateInfo);
+		SRV = RHICmdList.CreateShaderResourceView(Buffer, BytesPerElement, UE_PIXELFORMAT_TO_UINT8(Format));
+	}
+
+	void Initialize(const TCHAR* InDebugName, uint32 BytesPerElement, uint32 NumElements, EPixelFormat Format, EBufferUsageFlags AdditionalUsage = BUF_None, FResourceArrayInterface* InResourceArray = nullptr)
+	{
+		Initialize(FRHICommandListImmediate::Get(), InDebugName, BytesPerElement, NumElements, Format, AdditionalUsage, InResourceArray);
 	}
 
 	void Release()
@@ -333,16 +338,21 @@ struct FRWBufferStructured
 		Release();
 	}
 
-	void Initialize(const TCHAR* InDebugName, uint32 BytesPerElement, uint32 NumElements, EBufferUsageFlags AdditionalUsage = BUF_None, bool bUseUavCounter = false, bool bAppendBuffer = false, ERHIAccess InitialState = ERHIAccess::UAVMask)
+	void Initialize(FRHICommandListBase& RHICmdList, const TCHAR* InDebugName, uint32 BytesPerElement, uint32 NumElements, EBufferUsageFlags AdditionalUsage = BUF_None, bool bUseUavCounter = false, bool bAppendBuffer = false, ERHIAccess InitialState = ERHIAccess::UAVMask)
 	{
 		// Provide a debug name if using Fast VRAM so the allocators diagnostics will work
 		ensure(!(EnumHasAnyFlags(AdditionalUsage, BUF_FastVRAM) && !InDebugName));
 
 		NumBytes = BytesPerElement * NumElements;
 		FRHIResourceCreateInfo CreateInfo(InDebugName);
-		Buffer = RHICreateStructuredBuffer(BytesPerElement, NumBytes, BUF_UnorderedAccess | BUF_ShaderResource | AdditionalUsage, InitialState, CreateInfo);
-		UAV = RHICreateUnorderedAccessView(Buffer, bUseUavCounter, bAppendBuffer);
-		SRV = RHICreateShaderResourceView(Buffer);
+		Buffer = RHICmdList.CreateStructuredBuffer(BytesPerElement, NumBytes, BUF_UnorderedAccess | BUF_ShaderResource | AdditionalUsage, InitialState, CreateInfo);
+		UAV = RHICmdList.CreateUnorderedAccessView(Buffer, bUseUavCounter, bAppendBuffer);
+		SRV = RHICmdList.CreateShaderResourceView(Buffer);
+	}
+
+	void Initialize(const TCHAR* InDebugName, uint32 BytesPerElement, uint32 NumElements, EBufferUsageFlags AdditionalUsage = BUF_None, bool bUseUavCounter = false, bool bAppendBuffer = false, ERHIAccess InitialState = ERHIAccess::UAVMask)
+	{
+		Initialize(FRHICommandListImmediate::Get(), InDebugName, BytesPerElement, NumElements, AdditionalUsage, bUseUavCounter, bAppendBuffer, InitialState);
 	}
 
 	void Release()
@@ -362,13 +372,18 @@ struct FByteAddressBuffer
 
 	FByteAddressBuffer(): NumBytes(0) {}
 
-	void Initialize(const TCHAR* InDebugName, uint32 InNumBytes, EBufferUsageFlags AdditionalUsage = BUF_None)
+	void Initialize(FRHICommandListBase& RHICmdList, const TCHAR* InDebugName, uint32 InNumBytes, EBufferUsageFlags AdditionalUsage = BUF_None)
 	{
 		NumBytes = InNumBytes;
 		check( NumBytes % 4 == 0 );
 		FRHIResourceCreateInfo CreateInfo(InDebugName);
-		Buffer = RHICreateStructuredBuffer(4, NumBytes, BUF_ShaderResource | BUF_ByteAddressBuffer | AdditionalUsage, ERHIAccess::SRVMask, CreateInfo);
-		SRV = RHICreateShaderResourceView(Buffer);
+		Buffer = RHICmdList.CreateStructuredBuffer(4, NumBytes, BUF_ShaderResource | BUF_ByteAddressBuffer | AdditionalUsage, ERHIAccess::SRVMask, CreateInfo);
+		SRV = RHICmdList.CreateShaderResourceView(Buffer);
+	}
+
+	void Initialize(const TCHAR* InDebugName, uint32 InNumBytes, EBufferUsageFlags AdditionalUsage = BUF_None)
+	{
+		Initialize(FRHICommandListImmediate::Get(), InDebugName, InNumBytes, AdditionalUsage);
 	}
 
 	void Release()
@@ -384,10 +399,15 @@ struct FRWByteAddressBuffer : public FByteAddressBuffer
 {
 	FUnorderedAccessViewRHIRef UAV;
 
+	void Initialize(FRHICommandListBase& RHICmdList, const TCHAR* DebugName, uint32 InNumBytes, EBufferUsageFlags AdditionalUsage = BUF_None)
+	{
+		FByteAddressBuffer::Initialize(RHICmdList, DebugName, InNumBytes, BUF_UnorderedAccess | AdditionalUsage);
+		UAV = RHICmdList.CreateUnorderedAccessView(Buffer, false, false);
+	}
+
 	void Initialize(const TCHAR* DebugName, uint32 InNumBytes, EBufferUsageFlags AdditionalUsage = BUF_None)
 	{
-		FByteAddressBuffer::Initialize(DebugName, InNumBytes, BUF_UnorderedAccess | AdditionalUsage);
-		UAV = RHICreateUnorderedAccessView(Buffer, false, false);
+		Initialize(FRHICommandListImmediate::Get(), DebugName, InNumBytes, AdditionalUsage);
 	}
 
 	void Release()
@@ -413,36 +433,51 @@ struct FDynamicReadBuffer : public FReadBuffer
 		Release();
 	}
 
-	virtual void Initialize(const TCHAR* DebugName, uint32 BytesPerElement, uint32 NumElements, EPixelFormat Format, EBufferUsageFlags AdditionalUsage = BUF_None)
+	void Initialize(FRHICommandListBase& RHICmdList, const TCHAR* DebugName, uint32 BytesPerElement, uint32 NumElements, EPixelFormat Format, EBufferUsageFlags AdditionalUsage = BUF_None)
 	{
 		ensure(
 			EnumHasAnyFlags(AdditionalUsage, BUF_Dynamic | BUF_Volatile | BUF_Static) &&					// buffer should be Dynamic or Volatile or Static
 			EnumHasAnyFlags(AdditionalUsage, BUF_Dynamic) != EnumHasAnyFlags(AdditionalUsage, BUF_Volatile) // buffer should not be both
 			);
 
-		FReadBuffer::Initialize(DebugName, BytesPerElement, NumElements, Format, AdditionalUsage);
+		FReadBuffer::Initialize(RHICmdList, DebugName, BytesPerElement, NumElements, Format, AdditionalUsage);
+	}
+
+	void Initialize(const TCHAR* DebugName, uint32 BytesPerElement, uint32 NumElements, EPixelFormat Format, EBufferUsageFlags AdditionalUsage = BUF_None)
+	{
+		Initialize(FRHICommandListImmediate::Get(), DebugName, BytesPerElement, NumElements, Format, AdditionalUsage);
 	}
 
 	/**
 	* Locks the vertex buffer so it may be written to.
 	*/
-	void Lock()
+	void Lock(FRHICommandListBase& RHICmdList)
 	{
 		check(MappedBuffer == nullptr);
 		check(IsValidRef(Buffer));
-		MappedBuffer = (uint8*)RHILockBuffer(Buffer, 0, NumBytes, RLM_WriteOnly);
+		MappedBuffer = (uint8*)RHICmdList.LockBuffer(Buffer, 0, NumBytes, RLM_WriteOnly);
+	}
+
+	void Lock()
+	{
+		Lock(FRHICommandListImmediate::Get());
 	}
 
 	/**
 	* Unocks the buffer so the GPU may read from it.
 	*/
-	void Unlock()
+	void Unlock(FRHICommandListBase& RHICmdList)
 	{
 		check(MappedBuffer);
 		check(IsValidRef(Buffer));
-		RHIUnlockBuffer(Buffer);
+		RHICmdList.UnlockBuffer(Buffer);
 		MappedBuffer = nullptr;
 	}
+
+	void Unlock()
+	{
+		Unlock(FRHICommandListImmediate::Get());
+	}
 };
 
 /**
