Fix for non-unity build compilation

- checkSlow is utilized in some platform StaticShaderPlatform.inls and due to some recent refactoring may not be previously included
  - add the assertion macros include to RHIShaderPlatform.h to ensure its defined

- testing
  - local non-unity build on impacted platform
  - unity builds via preflight

diff --git a/Engine/Source/Runtime/RHI/Public/RHIShaderPlatform.h b/Engine/Source/Runtime/RHI/Public/RHIShaderPlatform.h
--- a/Engine/Source/Runtime/RHI/Public/RHIShaderPlatform.h
+++ b/Engine/Source/Runtime/RHI/Public/RHIShaderPlatform.h
@@ -2,6 +2,7 @@
 
 #pragma once
 
+#include "Misc/AssertionMacros.h"
 #include "RHIDefinitions.h"
 #include "HAL/Platform.h"
 
