Removing deprecated RHI code.

diff --git a/Engine/Source/Runtime/RHI/Private/DynamicRHI.cpp b/Engine/Source/Runtime/RHI/Private/DynamicRHI.cpp
--- a/Engine/Source/Runtime/RHI/Private/DynamicRHI.cpp
+++ b/Engine/Source/Runtime/RHI/Private/DynamicRHI.cpp
@@ -579,17 +579,6 @@ uint64 FDynamicRHI::RHIGetMinimumAlignmentForBufferBackedSRV(EPixelFormat Format
 	return GPixelFormats[Format].BlockBytes;
 }
 
-FTextureRHIRef FDynamicRHI::RHIAsyncCreateTexture2D(uint32 SizeX, uint32 SizeY, uint8 Format, uint32 NumMips, ETextureCreateFlags Flags, ERHIAccess InResourceState, void** InitialMipData, uint32 NumInitialMips)
-{
-	FGraphEventRef CompletionEvent;
-	FTextureRHIRef Result = this->RHIAsyncCreateTexture2D(SizeX, SizeY, Format, NumMips, Flags, InResourceState, InitialMipData, NumInitialMips, CompletionEvent);
-	if (CompletionEvent)
-	{
-		CompletionEvent->Wait();
-	}
-	return Result;
-}
-
 FTextureReferenceRHIRef FDynamicRHI::RHICreateTextureReference(FRHITexture* InReferencedTexture)
 {
 	FRHITexture* ReferencedTexture = InReferencedTexture ? InReferencedTexture : FRHITextureReference::GetDefaultTexture();
diff --git a/Engine/Source/Runtime/RHI/Private/RHIStrings.cpp b/Engine/Source/Runtime/RHI/Private/RHIStrings.cpp
--- a/Engine/Source/Runtime/RHI/Private/RHIStrings.cpp
+++ b/Engine/Source/Runtime/RHI/Private/RHIStrings.cpp
@@ -425,7 +425,6 @@ static const FRHIResourceTypeName GRHIResourceTypeNames[] =
 	RHI_RESOURCE_TYPE_DEF(GPUFence),
 	RHI_RESOURCE_TYPE_DEF(RenderQuery),
 	RHI_RESOURCE_TYPE_DEF(RenderQueryPool),
-	RHI_RESOURCE_TYPE_DEF(ComputeFence),
 	RHI_RESOURCE_TYPE_DEF(Viewport),
 	RHI_RESOURCE_TYPE_DEF(UnorderedAccessView),
 	RHI_RESOURCE_TYPE_DEF(ShaderResourceView),
diff --git a/Engine/Source/Runtime/RHI/Public/DynamicRHI.h b/Engine/Source/Runtime/RHI/Public/DynamicRHI.h
--- a/Engine/Source/Runtime/RHI/Public/DynamicRHI.h
+++ b/Engine/Source/Runtime/RHI/Public/DynamicRHI.h
@@ -200,18 +200,6 @@ public:
 		return new FDefaultRHIRenderQueryPool(QueryType, this, NumQueries);
 	}
 
-	/**
-	* Creates a compute fence.  Compute fences are named GPU fences which can be written to once before resetting.
-	* A command to write the fence must be enqueued before any commands to wait on them.  This is enforced on the CPU to avoid GPU hangs.
-	* @param Name - Friendly name for the Fence.  e.g. ReflectionEnvironmentComplete
-	* @return The new Fence.
-	*/
-	// FlushType: Thread safe, but varies depending on the RHI	
-	inline FComputeFenceRHIRef RHICreateComputeFence(const FName& Name)
-	{
-		return new FRHIComputeFence(Name);
-	}
-
 	virtual FGPUFenceRHIRef RHICreateGPUFence(const FName &Name)
 	{
 		return new FGenericRHIGPUFence(Name);
@@ -427,9 +415,6 @@ public:
 	// FlushType: Thread safe
 	virtual FTextureRHIRef RHIAsyncCreateTexture2D(uint32 SizeX, uint32 SizeY, uint8 Format, uint32 NumMips, ETextureCreateFlags Flags, ERHIAccess InResourceState, void** InitialMipData, uint32 NumInitialMips, FGraphEventRef& OutCompletionEvent) = 0;
 
-	UE_DEPRECATED(5.2, "RHIAsyncCreateTexture2D now requires a completion callback. Using the old version in a critical section can lead to deadlock. Please switch to the new function signature.")
-	FTextureRHIRef RHIAsyncCreateTexture2D(uint32 SizeX, uint32 SizeY, uint8 Format, uint32 NumMips, ETextureCreateFlags Flags, ERHIAccess InResourceState, void** InitialMipData, uint32 NumInitialMips);
-
 	/** Create a texture reference. InReferencedTexture can be null. */
 	RHI_API virtual FTextureReferenceRHIRef RHICreateTextureReference(FRHITexture* InReferencedTexture);
 
@@ -441,8 +426,6 @@ public:
 	* Generates mip maps for a texture.
 	*/
 	// FlushType: Flush Immediate (NP: this should be queued on the command list for RHI thread execution, not flushed)
-
-	//UE_DEPRECATED(4.23, "This function is deprecated and will be removed in future releases. Renderer version implemented.")
 	virtual void RHIGenerateMips(FRHITexture* Texture) {}
 
 	/**
@@ -1077,12 +1060,6 @@ FORCEINLINE FComputeShaderRHIRef RHICreateComputeShader(TArrayView<const uint8>
 	return GDynamicRHI->RHICreateComputeShader(Code, Hash);
 }
 
-UE_DEPRECATED(5.2, "Compute fences are deprecated. Use RHI transitions instead.")
-FORCEINLINE FComputeFenceRHIRef RHICreateComputeFence(const FName& Name)
-{
-	return GDynamicRHI->RHICreateComputeFence(Name);
-}
-
 FORCEINLINE FGPUFenceRHIRef RHICreateGPUFence(const FName& Name)
 {
 	return GDynamicRHI->RHICreateGPUFence(Name);
diff --git a/Engine/Source/Runtime/RHI/Public/RHICommandList.h b/Engine/Source/Runtime/RHI/Public/RHICommandList.h
--- a/Engine/Source/Runtime/RHI/Public/RHICommandList.h
+++ b/Engine/Source/Runtime/RHI/Public/RHICommandList.h
@@ -2989,17 +2989,6 @@ PRAGMA_ENABLE_DEPRECATION_WARNINGS
 #endif // WITH_MGPU
 	}
 
-#if WITH_MGPU
-	UE_DEPRECATED(5.2, "AFR support has been removed in 5.2, function is a nop and calls should be removed.")
-	FORCEINLINE void WaitForTemporalEffect(const FName& EffectName) {}
-
-	UE_DEPRECATED(5.2, "AFR support has been removed in 5.2, function is a nop and calls should be removed.")
-	FORCEINLINE void BroadcastTemporalEffect(const FName& EffectName, const TArrayView<FRHITexture*> Textures) {}
-
-	UE_DEPRECATED(5.2, "AFR support has been removed in 5.2, function is a nop and calls should be removed.")
-	FORCEINLINE void BroadcastTemporalEffect(const FName& EffectName, const TArrayView<FRHIBuffer*> Buffers) {}
-#endif // WITH_MGPU
-
 #if RHI_RAYTRACING
 	RHI_API void BuildAccelerationStructure(FRHIRayTracingGeometry* Geometry);
 	RHI_API void BuildAccelerationStructures(TConstArrayView<FRayTracingGeometryBuildParams> Params);
@@ -3116,18 +3105,6 @@ public:
 		FRHICommandList::EnqueueLambda(TEXT("TRHILambdaCommand"), Forward<LAMBDA>(Lambda));
 	}
 
-	UE_DEPRECATED(5.2, "AFR support has been removed in 5.2, function is a nop and calls should be removed.")
-	FORCEINLINE void BeginUpdateMultiFrameResource(FRHITexture* Texture) {}
-
-	UE_DEPRECATED(5.2, "AFR support has been removed in 5.2, function is a nop and calls should be removed.")
-	FORCEINLINE void EndUpdateMultiFrameResource(FRHITexture* Texture) {}
-
-	UE_DEPRECATED(5.2, "AFR support has been removed in 5.2, function is a nop and calls should be removed.")
-	FORCEINLINE void BeginUpdateMultiFrameResource(FRHIUnorderedAccessView* UAV) {}
-
-	UE_DEPRECATED(5.2, "AFR support has been removed in 5.2, function is a nop and calls should be removed.")
-	FORCEINLINE void EndUpdateMultiFrameResource(FRHIUnorderedAccessView* UAV) {}
-
 	using FRHIComputeCommandList::SetShaderUniformBuffer;
 
 	UE_DEPRECATED(5.3, "FRHIBatchedShaderParameters and SetBatchedShaderParameters should be used instead of setting individual parameters.")
@@ -4168,110 +4145,6 @@ public:
 	{
 		FRHICommandListImmediate::EnqueueLambda(TEXT("TRHILambdaCommand"), Forward<LAMBDA>(Lambda));
 	}
-
-	UE_DEPRECATED(5.2, "Use the global scope RHICreateSamplerState function.")
-	FORCEINLINE FSamplerStateRHIRef CreateSamplerState(const FSamplerStateInitializerRHI& Initializer)
-	{
-		return RHICreateSamplerState(Initializer);
-	}
-	
-	UE_DEPRECATED(5.2, "Use the global scope RHICreateRasterizerState function.")
-	FORCEINLINE FRasterizerStateRHIRef CreateRasterizerState(const FRasterizerStateInitializerRHI& Initializer)
-	{
-		return RHICreateRasterizerState(Initializer);
-	}
-	
-	UE_DEPRECATED(5.2, "Use the global scope RHICreateDepthStencilState function.")
-	FORCEINLINE FDepthStencilStateRHIRef CreateDepthStencilState(const FDepthStencilStateInitializerRHI& Initializer)
-	{
-		return RHICreateDepthStencilState(Initializer);
-	}
-	
-	UE_DEPRECATED(5.2, "Use the global scope RHICreateBlendState function.")
-	FORCEINLINE FBlendStateRHIRef CreateBlendState(const FBlendStateInitializerRHI& Initializer)
-	{
-		return RHICreateBlendState(Initializer);
-	}
-
-	UE_DEPRECATED(5.2, "Use the global scope RHICreatePixelShader function.")
-	FORCEINLINE FPixelShaderRHIRef CreatePixelShader(TArrayView<const uint8> Code, const FSHAHash& Hash)
-	{
-		return RHICreatePixelShader(Code, Hash);
-	}
-	
-	UE_DEPRECATED(5.2, "Use the global scope RHICreateVertexShader function.")
-	FORCEINLINE FVertexShaderRHIRef CreateVertexShader(TArrayView<const uint8> Code, const FSHAHash& Hash)
-	{
-		return RHICreateVertexShader(Code, Hash);
-	}
-
-	UE_DEPRECATED(5.2, "Use the global scope RHICreateMeshShader function.")
-	FORCEINLINE FMeshShaderRHIRef CreateMeshShader(TArrayView<const uint8> Code, const FSHAHash& Hash)
-	{
-		return RHICreateMeshShader(Code, Hash);
-	}
-
-	UE_DEPRECATED(5.2, "Use the global scope RHICreateAmplificationShader function.")
-	FORCEINLINE FAmplificationShaderRHIRef CreateAmplificationShader(TArrayView<const uint8> Code, const FSHAHash& Hash)
-	{
-		return RHICreateAmplificationShader(Code, Hash);
-	}
-	
-	UE_DEPRECATED(5.2, "Use the global scope RHICreateGeometryShader function.")
-	FORCEINLINE FGeometryShaderRHIRef CreateGeometryShader(TArrayView<const uint8> Code, const FSHAHash& Hash)
-	{
-		return RHICreateGeometryShader(Code, Hash);
-	}
-	
-	UE_DEPRECATED(5.2, "Use the global scope RHICreateComputeShader function.")
-	FORCEINLINE FComputeShaderRHIRef CreateComputeShader(TArrayView<const uint8> Code, const FSHAHash& Hash)
-	{
-		return RHICreateComputeShader(Code, Hash);
-	}
-	
-	UE_DEPRECATED(5.2, "Compute fences are deprecated. Use RHI transitions instead.")
-	FORCEINLINE FComputeFenceRHIRef CreateComputeFence(const FName& Name)
-	{
-	PRAGMA_DISABLE_DEPRECATION_WARNINGS
-		return RHICreateComputeFence(Name);
-	PRAGMA_ENABLE_DEPRECATION_WARNINGS
-	}	
-
-	UE_DEPRECATED(5.2, "Use the global scope RHICreateGPUFence function.")
-	FORCEINLINE FGPUFenceRHIRef CreateGPUFence(const FName& Name)
-	{
-		return RHICreateGPUFence(Name);
-	}
-
-	UE_DEPRECATED(5.2, "Use the global scope RHICreateStagingBuffer function.")
-	FORCEINLINE FStagingBufferRHIRef CreateStagingBuffer()
-	{
-		return RHICreateStagingBuffer();
-	}
-
-	UE_DEPRECATED(5.2, "Use the global scope RHICreateBoundShaderState function.")
-	FORCEINLINE FBoundShaderStateRHIRef CreateBoundShaderState(FRHIVertexDeclaration* VertexDeclaration, FRHIVertexShader* VertexShader, FRHIPixelShader* PixelShader, FRHIGeometryShader* GeometryShader)
-	{
-		return RHICreateBoundShaderState(VertexDeclaration, VertexShader, PixelShader, GeometryShader);
-	}
-
-	UE_DEPRECATED(5.2, "Use the global scope RHICreateGraphicsPipelineState function.")
-	FORCEINLINE FGraphicsPipelineStateRHIRef CreateGraphicsPipelineState(const FGraphicsPipelineStateInitializer& Initializer)
-	{
-		return RHICreateGraphicsPipelineState(Initializer);
-	}
-
-	UE_DEPRECATED(5.2, "Use the global scope RHICreateComputePipelineState function.")
-	FORCEINLINE FComputePipelineStateRHIRef CreateComputePipelineState(FRHIComputeShader* ComputeShader)
-	{
-		return RHICreateComputePipelineState(ComputeShader);
-	}
-
-	UE_DEPRECATED(5.2, "Use the global scope RHICreateUniformBuffer function.")
-	FORCEINLINE FUniformBufferRHIRef CreateUniformBuffer(const void* Contents, const FRHIUniformBufferLayout* Layout, EUniformBufferUsage Usage)
-	{
-		return RHICreateUniformBuffer(Contents, Layout, Usage);
-	}
 	
 	FORCEINLINE void* LockStagingBuffer(FRHIStagingBuffer* StagingBuffer, FRHIGPUFence* Fence, uint32 Offset, uint32 SizeRHI)
 	{
@@ -4309,12 +4182,6 @@ public:
 	{
 		GDynamicRHI->RHIUnlockBufferMGPU(*this, Buffer, GPUIndex);
 	}
-
-	UE_DEPRECATED(5.2, "Use the global scope RHIGetTextureMemoryStats function.")
-	FORCEINLINE void GetTextureMemoryStats(FTextureMemoryStats& OutStats)
-	{
-		RHIGetTextureMemoryStats(OutStats);
-	}
 	
 	FORCEINLINE bool GetTextureMemoryVisualizeData(FColor* TextureData,int32 SizeX,int32 SizeY,int32 Pitch,int32 PixelSize)
 	{
@@ -4322,11 +4189,7 @@ public:
 		ImmediateFlush(EImmediateFlushType::FlushRHIThread); 
 		return GDynamicRHI->RHIGetTextureMemoryVisualizeData(TextureData,SizeX,SizeY,Pitch,PixelSize);
 	}
-	
-	UE_DEPRECATED(5.2, "CopySharedMips should be replaced with UE::RHI::CopySharedMips(). ** Note ** The argument ordering of the new function has changed. The new function also does not issue RHI transitions, so the caller is responsible for ensuring src/dst textures are in CopySrc/CopyDest state.")
-	FORCEINLINE void CopySharedMips(FRHITexture* DestTexture, FRHITexture* SrcTexture);
 
-	//UE_DEPRECATED(4.23, "This function is deprecated and will be removed in future releases. Renderer version implemented.")
 	FORCEINLINE void GenerateMips(FRHITexture* Texture)
 	{
 		QUICK_SCOPE_CYCLE_COUNTER(STAT_RHIMETHOD_GenerateMips_Flush);
@@ -4334,12 +4197,6 @@ public:
 		return GDynamicRHI->RHIGenerateMips(Texture);
 	}
 	
-	UE_DEPRECATED(5.2, "Use the global scope RHIComputeMemorySize function.")
-	FORCEINLINE uint32 ComputeMemorySize(FRHITexture* TextureRHI)
-	{
-		return RHIComputeMemorySize(TextureRHI);
-	}
-	
 	FORCEINLINE FTexture2DRHIRef AsyncReallocateTexture2D(FRHITexture2D* Texture2D, int32 NewMipCount, int32 NewSizeX, int32 NewSizeY, FThreadSafeCounter* RequestStatus)
 	{
 		LLM_SCOPE(ELLMTag::Textures);
@@ -4412,18 +4269,6 @@ public:
 		LLM_SCOPE(ELLMTag::Textures);
 		GDynamicRHI->RHIUnlockTextureCubeFace_RenderThread(*this, Texture, FaceIndex, ArrayIndex, MipIndex, bLockWithinMiptail);
 	}
-	
-	UE_DEPRECATED(5.2, "Use the global scope RHIBindDebugLabelName function.")
-	FORCEINLINE void BindDebugLabelName(FRHITexture* Texture, const TCHAR* Name)
-	{
-		RHIBindDebugLabelName(Texture, Name);
-	}
-
-	UE_DEPRECATED(5.2, "Use the global scope RHIBindDebugLabelName function.")
-	FORCEINLINE void BindDebugLabelName(FRHIUnorderedAccessView* UnorderedAccessViewRHI, const TCHAR* Name)
-	{
-		RHIBindDebugLabelName(UnorderedAccessViewRHI, Name);
-	}
 
 	FORCEINLINE void ReadSurfaceData(FRHITexture* Texture,FIntRect Rect,TArray<FColor>& OutData,FReadSurfaceDataFlags InFlags)
 	{
@@ -4478,30 +4323,6 @@ public:
 		ImmediateFlush(EImmediateFlushType::FlushRHIThread);  
 		GDynamicRHI->RHIRead3DSurfaceFloatData(Texture,Rect,ZMinMax,OutData,Flags);
 	}
-
-	UE_DEPRECATED(5.2, "Use the global scope RHIGetRenderQueryResult function.")
-	FORCEINLINE bool GetRenderQueryResult(FRHIRenderQuery* RenderQuery, uint64& OutResult, bool bWait, uint32 GPUIndex = INDEX_NONE)
-	{
-		return RHIGetRenderQueryResult(RenderQuery, OutResult, bWait, GPUIndex);
-	}
-
-	UE_DEPRECATED(5.2, "Use the global scope RHIGetViewportNextPresentGPUIndex function.")
-	FORCEINLINE uint32 GetViewportNextPresentGPUIndex(FRHIViewport* Viewport)
-	{
-		return RHIGetViewportNextPresentGPUIndex(Viewport);
-	}
-
-	UE_DEPRECATED(5.2, "Use the global scope RHIGetViewportBackBuffer function.")
-	FORCEINLINE FTexture2DRHIRef GetViewportBackBuffer(FRHIViewport* Viewport)
-	{
-		return RHIGetViewportBackBuffer(Viewport);
-	}
-	
-	UE_DEPRECATED(5.2, "Use the global scope RHIAdvanceFrameForGetViewportBackBuffer function.")
-	FORCEINLINE void AdvanceFrameForGetViewportBackBuffer(FRHIViewport* Viewport)
-	{
-		return RHIAdvanceFrameForGetViewportBackBuffer(Viewport);
-	}
 	
 	FORCEINLINE void AcquireThreadOwnership()
 	{
@@ -4534,24 +4355,6 @@ public:
 		return RHIGetGPUFrameCycles(GetGPUMask().ToIndex());
 	}
 	
-	UE_DEPRECATED(5.2, "Use the global scope RHICreateViewport function.")
-	FORCEINLINE FViewportRHIRef CreateViewport(void* WindowHandle, uint32 SizeX, uint32 SizeY, bool bIsFullscreen, EPixelFormat PreferredPixelFormat)
-	{
-		return RHICreateViewport(WindowHandle, SizeX, SizeY, bIsFullscreen, PreferredPixelFormat);
-	}
-	
-	UE_DEPRECATED(5.2, "Use the global scope RHIResizeViewport function.")
-	FORCEINLINE void ResizeViewport(FRHIViewport* Viewport, uint32 SizeX, uint32 SizeY, bool bIsFullscreen, EPixelFormat PreferredPixelFormat)
-	{
-		RHIResizeViewport(Viewport, SizeX, SizeY, bIsFullscreen, PreferredPixelFormat);
-	}
-	
-	UE_DEPRECATED(5.2, "Use the global scope RHITick function.")
-	FORCEINLINE void Tick(float DeltaTime)
-	{
-		RHITick(DeltaTime);
-	}
-	
 	FORCEINLINE void BlockUntilGPUIdle()
 	{
 		QUICK_SCOPE_CYCLE_COUNTER(STAT_RHIMETHOD_BlockUntilGPUIdle_Flush);
@@ -4566,36 +4369,12 @@ public:
 		GDynamicRHI->RHISubmitCommandsAndFlushGPU();
 	}
 	
-	UE_DEPRECATED(5.2, "Use the global scope RHISuspendRendering function.")
-	FORCEINLINE void SuspendRendering()
-	{
-		RHISuspendRendering();
-	}
-	
-	UE_DEPRECATED(5.2, "Use the global scope RHIResumeRendering function.")
-	FORCEINLINE void ResumeRendering()
-	{
-		RHIResumeRendering();
-	}
-	
 	FORCEINLINE bool IsRenderingSuspended()
 	{
 		QUICK_SCOPE_CYCLE_COUNTER(STAT_RHIMETHOD_IsRenderingSuspended_Flush);
 		ImmediateFlush(EImmediateFlushType::FlushRHIThread); 
 		return GDynamicRHI->RHIIsRenderingSuspended();
 	}
-
-	UE_DEPRECATED(5.2, "Use the global scope RHIGetAvailableResolutions function.")
-	FORCEINLINE bool GetAvailableResolutions(FScreenResolutionArray& Resolutions, bool bIgnoreRefreshRate)
-	{
-		return RHIGetAvailableResolutions(Resolutions, bIgnoreRefreshRate);
-	}
-	
-	UE_DEPRECATED(5.2, "Use the global scope RHIGetSupportedResolution function.")
-	FORCEINLINE void GetSupportedResolution(uint32& Width, uint32& Height)
-	{
-		RHIGetSupportedResolution(Width, Height);
-	}
 	
 	FORCEINLINE void VirtualTextureSetFirstMipInMemory(FRHITexture2D* Texture, uint32 FirstMip)
 	{
@@ -4652,12 +4431,6 @@ public:
 		return GDynamicRHI->RHIGetNativeCommandBuffer();
 	}
 
-	UE_DEPRECATED(5.2, "Use the global scope RHIGetDefaultContext function.")
-	FORCEINLINE class IRHICommandContext* GetDefaultContext()
-	{
-		return RHIGetDefaultContext();
-	}
-
 	RHI_API void UpdateTextureReference(FRHITextureReference* TextureRef, FRHITexture* NewTexture);
 
 	FORCEINLINE void PollRenderQueryResults()
@@ -5118,27 +4891,6 @@ FORCEINLINE FUnorderedAccessViewRHIRef RHICreateUnorderedAccessView(FRHIBuffer*
 	return FRHICommandListExecutor::GetImmediateCommandList().CreateUnorderedAccessView(Buffer, Format);
 }
 
-UE_DEPRECATED(5.2, "Use RHICreateUnorderedAccessViewHTile with an FRHIViewDesc that specifies the htile plane.")
-FORCEINLINE FUnorderedAccessViewRHIRef RHICreateUnorderedAccessViewHTile(FRHITexture2D* RenderTarget)
-{
-	return FRHICommandListExecutor::GetImmediateCommandList().CreateUnorderedAccessView(RenderTarget, FRHIViewDesc::CreateTextureUAV()
-		.SetDimensionFromTexture(RenderTarget)
-		.SetPlane(ERHITexturePlane::HTile)
-	);
-}
-
-UE_DEPRECATED(5.2, "Use RHICreateUnorderedAccessViewHTile with an FRHIViewDesc that specifies the stencil plane.")
-FORCEINLINE FUnorderedAccessViewRHIRef RHICreateUnorderedAccessViewStencil(FRHITexture2D* DepthTarget, int32 MipLevel)
-{
-	check(MipLevel < 256 && MipLevel >= 0);
-
-	return FRHICommandListExecutor::GetImmediateCommandList().CreateUnorderedAccessView(DepthTarget, FRHIViewDesc::CreateTextureUAV()
-		.SetDimensionFromTexture(DepthTarget)
-		.SetPlane(ERHITexturePlane::Stencil)
-		.SetMipLevel(uint8(MipLevel))
-	);
-}
-
 //UE_DEPRECATED(5.3, "Use the RHICreateShaderResourceView function that takes an FRHIBufferSRVCreateDesc.")
 FORCEINLINE FShaderResourceViewRHIRef RHICreateShaderResourceView(FRHIBuffer* Buffer)
 {
@@ -5175,27 +4927,6 @@ FORCEINLINE FShaderResourceViewRHIRef RHICreateShaderResourceView(FRHITexture* T
 	return FRHICommandListExecutor::GetImmediateCommandList().CreateShaderResourceView(Texture, CreateInfo);
 }
 
-UE_DEPRECATED(5.2, "Use RHICreateShaderResourceView with an FRHITextureSRVCreateDesc that specifies the htile plane.")
-FORCEINLINE FShaderResourceViewRHIRef RHICreateShaderResourceViewHTile(FRHITexture2D* RenderTarget)
-{
-	return FRHICommandListExecutor::GetImmediateCommandList().CreateShaderResourceView(RenderTarget, FRHIViewDesc::CreateTextureSRV()
-		.SetDimensionFromTexture(RenderTarget)
-		.SetPlane(ERHITexturePlane::HTile)
-	);
-}
-
-UE_DEPRECATED(5.2, "Use RHICreateShaderResourceView with an FRHITextureSRVCreateDesc that specifies the cmask plane.")
-FORCEINLINE FShaderResourceViewRHIRef RHICreateShaderResourceViewWriteMask(FRHITexture2D* Texture2D)
-{
-	return FRHICommandListExecutor::GetImmediateCommandList().CreateShaderResourceViewWriteMask(Texture2D);
-}
-
-UE_DEPRECATED(5.2, "Use RHICreateShaderResourceView with an FRHITextureSRVCreateDesc that specifies the fmask plane.")
-FORCEINLINE FShaderResourceViewRHIRef RHICreateShaderResourceViewFMask(FRHITexture2D* Texture2D)
-{
-	return FRHICommandListExecutor::GetImmediateCommandList().CreateShaderResourceViewFMask(Texture2D);
-}
-
 FORCEINLINE void RHIUpdateRHIResources(FRHIResourceUpdateInfo* UpdateInfos, int32 Num, bool bNeedReleaseRefs)
 {
 	check(IsInRenderingThread());
@@ -5248,32 +4979,6 @@ PRAGMA_DISABLE_DEPRECATION_WARNINGS
 PRAGMA_ENABLE_DEPRECATION_WARNINGS
 }
 
-UE_DEPRECATED(5.2, "RHIAsyncCreateTexture2D now requires a completion callback. Using the old version in a critical section can lead to deadlock. Please switch to the new function signature.")
-FORCEINLINE FTexture2DRHIRef RHIAsyncCreateTexture2D(uint32 SizeX, uint32 SizeY, uint8 Format, uint32 NumMips, ETextureCreateFlags Flags, ERHIAccess InResourceState, void** InitialMipData, uint32 NumInitialMips)
-{
-	LLM_SCOPE(EnumHasAnyFlags(Flags, TexCreate_RenderTargetable | TexCreate_DepthStencilTargetable) ? ELLMTag::RenderTargets : ELLMTag::Textures);
-	const ERHIAccess ResourceState = InResourceState == ERHIAccess::Unknown ? RHIGetDefaultResourceState((ETextureCreateFlags)Flags, InitialMipData != nullptr) : InResourceState;
-PRAGMA_DISABLE_DEPRECATION_WARNINGS
-	return GDynamicRHI->RHIAsyncCreateTexture2D(SizeX, SizeY, Format, NumMips, Flags, ResourceState, InitialMipData, NumInitialMips);
-PRAGMA_ENABLE_DEPRECATION_WARNINGS
-}
-
-UE_DEPRECATED(5.2, "RHIAsyncCreateTexture2D now requires a completion callback. Using the old version in a critical section can lead to deadlock. Please switch to the new function signature.")
-FORCEINLINE FTexture2DRHIRef RHIAsyncCreateTexture2D(uint32 SizeX, uint32 SizeY, uint8 Format, uint32 NumMips, ETextureCreateFlags Flags, void** InitialMipData, uint32 NumInitialMips)
-{
-PRAGMA_DISABLE_DEPRECATION_WARNINGS
-	return RHIAsyncCreateTexture2D(SizeX, SizeY, Format, NumMips, Flags, ERHIAccess::Unknown, InitialMipData, NumInitialMips);
-PRAGMA_ENABLE_DEPRECATION_WARNINGS
-}
-
-UE_DEPRECATED(5.2, "RHICopySharedMips should be replaced with UE::RHI::CopySharedMips(). ** Note ** The argument ordering of the new function has changed. The new function also does not issue RHI transitions, so the caller is responsible for ensuring src/dst textures are in CopySrc/CopyDest state.")
-FORCEINLINE void RHICopySharedMips(FRHITexture* DestTexture, FRHITexture* SrcTexture)
-{
-PRAGMA_DISABLE_DEPRECATION_WARNINGS
-	return FRHICommandListExecutor::GetImmediateCommandList().CopySharedMips(DestTexture, SrcTexture);
-PRAGMA_ENABLE_DEPRECATION_WARNINGS
-}
-
 FORCEINLINE FTexture2DRHIRef RHIAsyncReallocateTexture2D(FRHITexture2D* Texture2D, int32 NewMipCount, int32 NewSizeX, int32 NewSizeY, FThreadSafeCounter* RequestStatus)
 {
 	return FRHICommandListExecutor::GetImmediateCommandList().AsyncReallocateTexture2D(Texture2D, NewMipCount, NewSizeX, NewSizeY, RequestStatus);
diff --git a/Engine/Source/Runtime/RHI/Public/RHICommandList.inl b/Engine/Source/Runtime/RHI/Public/RHICommandList.inl
--- a/Engine/Source/Runtime/RHI/Public/RHICommandList.inl
+++ b/Engine/Source/Runtime/RHI/Public/RHICommandList.inl
@@ -173,11 +173,3 @@ inline FRHICommandListScopedPipelineGuard::~FRHICommandListScopedPipelineGuard()
 		RHICmdList.SwitchPipeline(ERHIPipeline::None);
 	}
 }
-
-FORCEINLINE void FRHICommandListImmediate::CopySharedMips(FRHITexture* DestTexture, FRHITexture* SrcTexture)
-{
-	QUICK_SCOPE_CYCLE_COUNTER(STAT_RHIMETHOD_CopySharedMips_Flush);
-
-	// Use the version that assumes SRVMask state on both textures to maintain behaviour of the old code path.
-	UE::RHI::CopySharedMips_AssumeSRVMaskState(*this, SrcTexture, DestTexture);
-}
diff --git a/Engine/Source/Runtime/RHI/Public/RHIDefinitions.h b/Engine/Source/Runtime/RHI/Public/RHIDefinitions.h
--- a/Engine/Source/Runtime/RHI/Public/RHIDefinitions.h
+++ b/Engine/Source/Runtime/RHI/Public/RHIDefinitions.h
@@ -865,7 +865,6 @@ enum ERHIResourceType : uint8
 	RRT_GPUFence,
 	RRT_RenderQuery,
 	RRT_RenderQueryPool,
-	RRT_ComputeFence,
 	RRT_Viewport,
 	RRT_UnorderedAccessView,
 	RRT_ShaderResourceView,
diff --git a/Engine/Source/Runtime/RHI/Public/RHIFwd.h b/Engine/Source/Runtime/RHI/Public/RHIFwd.h
--- a/Engine/Source/Runtime/RHI/Public/RHIFwd.h
+++ b/Engine/Source/Runtime/RHI/Public/RHIFwd.h
@@ -47,7 +47,6 @@ class FRHIAmplificationShader;
 class FRHIBlendState;
 class FRHIBoundShaderState;
 class FRHIBuffer;
-class FRHIComputeFence;
 class FRHIComputePipelineState;
 class FRHIComputeShader;
 class FRHICustomPresent;
@@ -89,7 +88,6 @@ using FAmplificationShaderRHIRef       = TRefCountPtr<FRHIAmplificationShader>;
 using FBlendStateRHIRef                = TRefCountPtr<FRHIBlendState>;
 using FBoundShaderStateRHIRef          = TRefCountPtr<FRHIBoundShaderState>;
 using FBufferRHIRef                    = TRefCountPtr<FRHIBuffer>;
-using FComputeFenceRHIRef              = TRefCountPtr<FRHIComputeFence>;
 using FComputePipelineStateRHIRef      = TRefCountPtr<FRHIComputePipelineState>;
 using FComputeShaderRHIRef             = TRefCountPtr<FRHIComputeShader>;
 using FCustomPresentRHIRef             = TRefCountPtr<FRHICustomPresent>;
diff --git a/Engine/Source/Runtime/RHI/Public/RHIResources.h b/Engine/Source/Runtime/RHI/Public/RHIResources.h
--- a/Engine/Source/Runtime/RHI/Public/RHIResources.h
+++ b/Engine/Source/Runtime/RHI/Public/RHIResources.h
@@ -1963,33 +1963,6 @@ inline FRHIPooledRenderQuery::~FRHIPooledRenderQuery()
 	ReleaseQuery();
 }
 
-class FRHIComputeFence final : public FRHIResource
-{
-public:
-
-	FRHIComputeFence(FName InName)
-		: FRHIResource(RRT_ComputeFence)
-		, Name(InName)
-	{}
-
-	FORCEINLINE FName GetName() const
-	{
-		return Name;
-	}
-
-	FORCEINLINE bool GetWriteEnqueued() const
-	{
-		return Transition != nullptr;
-	}
-
-private:
-	//debug name of the label.
-	FName Name;
-
-public:
-	const FRHITransition* Transition = nullptr;
-};
-
 class FRHIViewport : public FRHIResource 
 {
 public:
diff --git a/Engine/Source/Runtime/VulkanRHI/Private/VulkanRHI.cpp b/Engine/Source/Runtime/VulkanRHI/Private/VulkanRHI.cpp
--- a/Engine/Source/Runtime/VulkanRHI/Private/VulkanRHI.cpp
+++ b/Engine/Source/Runtime/VulkanRHI/Private/VulkanRHI.cpp
@@ -1451,17 +1451,6 @@ FVulkanRHIImageViewInfo FVulkanDynamicRHI::RHIGetImageViewInfo(FRHITexture* InTe
 	return Info;
 }
 
-// todo-jn: deprecate
-VkImageLayout& FVulkanDynamicRHI::RHIFindOrAddLayoutRW(FRHITexture* InTexture, VkImageLayout LayoutIfNotFound)
-{
-	FVulkanTexture* VulkanTexture = ResourceCast(InTexture);
-	FVulkanCommandListContext& ImmediateContext = GetDevice()->GetImmediateContext();
-	FVulkanCmdBuffer* CmdBuffer = ImmediateContext.GetCommandBufferManager()->GetActiveCmdBuffer();
-	// Removing const to allow original functionality even when parallel rendering is enabled.
-	FVulkanImageLayout* VulkanImageLayout = const_cast<FVulkanImageLayout*>(CmdBuffer->GetLayoutManager().GetFullLayout(*VulkanTexture, true, LayoutIfNotFound));
-	return VulkanImageLayout->MainLayout;
-}
-
 void FVulkanDynamicRHI::RHISetImageLayout(VkImage Image, VkImageLayout OldLayout, VkImageLayout NewLayout, const VkImageSubresourceRange& SubresourceRange)
 {
 	FVulkanCommandListContext& ImmediateContext = GetDevice()->GetImmediateContext();
diff --git a/Engine/Source/Runtime/VulkanRHI/Public/IVulkanDynamicRHI.h b/Engine/Source/Runtime/VulkanRHI/Public/IVulkanDynamicRHI.h
--- a/Engine/Source/Runtime/VulkanRHI/Public/IVulkanDynamicRHI.h
+++ b/Engine/Source/Runtime/VulkanRHI/Public/IVulkanDynamicRHI.h
@@ -69,9 +69,6 @@ struct IVulkanDynamicRHI : public FDynamicRHI
 	virtual FVulkanRHIAllocationInfo RHIGetAllocationInfo(FRHITexture* InTexture) const = 0;
 	virtual FVulkanRHIImageViewInfo  RHIGetImageViewInfo(FRHITexture* InTexture) const = 0;
 
-	UE_DEPRECATED(5.2, "Altering image layout tracking directly is not permitted anymore.")
-	virtual VkImageLayout& RHIFindOrAddLayoutRW(FRHITexture* InTexture, VkImageLayout LayoutIfNotFound) = 0;
-
 	virtual void           RHISetImageLayout(VkImage Image, VkImageLayout OldLayout, VkImageLayout NewLayout, const VkImageSubresourceRange& SubresourceRange) = 0;
 	virtual void           RHISetUploadImageLayout(VkImage Image, VkImageLayout OldLayout, VkImageLayout NewLayout, const VkImageSubresourceRange& SubresourceRange) = 0;
 
diff --git a/Engine/Source/Runtime/VulkanRHI/Public/VulkanDynamicRHI.h b/Engine/Source/Runtime/VulkanRHI/Public/VulkanDynamicRHI.h
--- a/Engine/Source/Runtime/VulkanRHI/Public/VulkanDynamicRHI.h
+++ b/Engine/Source/Runtime/VulkanRHI/Public/VulkanDynamicRHI.h
@@ -77,10 +77,6 @@ public:
 	virtual VkFormat RHIGetViewVkFormat(FRHITexture* InTexture) const final override;
 	virtual FVulkanRHIAllocationInfo RHIGetAllocationInfo(FRHITexture* InTexture) const final override;
 	virtual FVulkanRHIImageViewInfo RHIGetImageViewInfo(FRHITexture* InTexture) const final override;
-
-	UE_DEPRECATED(5.2, "Altering image layout tracking directly is not permitted anymore.")
-	virtual VkImageLayout& RHIFindOrAddLayoutRW(FRHITexture* InTexture, VkImageLayout LayoutIfNotFound) final override;
-
 	virtual void RHISetImageLayout(VkImage Image, VkImageLayout OldLayout, VkImageLayout NewLayout, const VkImageSubresourceRange& SubresourceRange) final override;
 	virtual void RHISetUploadImageLayout(VkImage Image, VkImageLayout OldLayout, VkImageLayout NewLayout, const VkImageSubresourceRange& SubresourceRange) final override;
 	virtual void RHIFinishExternalComputeWork(VkCommandBuffer InCommandBuffer) final override;
