Add bSupportsVertexShaderSRVs to DataDrivenShaderPlatformInfo which drives PLATFORM_SUPPORTS_VERTEX_SHADER_SRVS shader

* This should be used to guard any features that require explicitly accessing buffer or textures SRVs in vertex shaders

diff --git a/Engine/Config/Android/DataDrivenPlatformInfo.ini b/Engine/Config/Android/DataDrivenPlatformInfo.ini
--- a/Engine/Config/Android/DataDrivenPlatformInfo.ini
+++ b/Engine/Config/Android/DataDrivenPlatformInfo.ini
@@ -62,6 +62,7 @@ bIsAndroidOpenGLES=true
 bTargetsTiledGPU=true
 bSupportsMobileMultiView=true
 bNeedsOfflineCompiler=true
+bSupportsVertexShaderSRVs=false
 bSupportsManualVertexFetch=false
 bSupportsClipDistance = false
 bSupportsShaderPipelines = false
@@ -78,6 +79,7 @@ bSupportsMobileMultiView=true
 bSupportsVertexShaderLayer=true
 bNeedsOfflineCompiler=true
 bSupportsIndexBufferUAVs=true
+bSupportsVertexShaderSRVs=false
 bSupportsManualVertexFetch=false
 bSupportsClipDistance = false
 bSupportsShaderPipelines = false
@@ -92,8 +94,9 @@ bNeedsOfflineCompiler=true
 bSupportsIndexBufferUAVs=true
 bRequiresDisableForwardLocalLights=true
 bWaterUsesSimpleForwardShading=true
-; Raytracing and Lumen
+bSupportsVertexShaderSRVs=true
 bSupportsManualVertexFetch=true
+; Raytracing and Lumen
 bSupportsRayTracing = true
 bSupportsRayTracingShaders = false
 bSupportsInlineRayTracing = true
diff --git a/Engine/Config/IOS/DataDrivenPlatformInfo.ini b/Engine/Config/IOS/DataDrivenPlatformInfo.ini
--- a/Engine/Config/IOS/DataDrivenPlatformInfo.ini
+++ b/Engine/Config/IOS/DataDrivenPlatformInfo.ini
@@ -52,6 +52,7 @@ ShaderFormat=SF_METAL
 bIsMobile=true
 bTargetsTiledGPU=true
 bSupports4ComponentUAVReadWrite=true
+bSupportsVertexShaderSRVs=false
 bSupportsManualVertexFetch=false
 bSupportsVolumeTextureAtomics=false
 bSupportsShaderPipelines = false
@@ -65,6 +66,7 @@ bIsMetalMRT=true
 bTargetsTiledGPU=true
 bSupportsIndexBufferUAVs=true
 bSupports4ComponentUAVReadWrite=true
+bSupportsVertexShaderSRVs=false
 bSupportsManualVertexFetch=false
 bSupportsVolumeTextureAtomics=false
 FriendlyName=LOCTEXT("FriendlyShaderPlatformName_IOSSM5", "IOS Metal SM5")
\ No newline at end of file
diff --git a/Engine/Config/Mac/DataDrivenPlatformInfo.ini b/Engine/Config/Mac/DataDrivenPlatformInfo.ini
--- a/Engine/Config/Mac/DataDrivenPlatformInfo.ini
+++ b/Engine/Config/Mac/DataDrivenPlatformInfo.ini
@@ -111,6 +111,7 @@ bIsPC=true
 bIsMetalMRT=true
 bSupportsIndexBufferUAVs=true
 bSupports4ComponentUAVReadWrite=true
+bSupportsVertexShaderSRVs=true
 bSupportsManualVertexFetch=true
 bSupportsVolumeTextureAtomics=false
 bSupportsDualSourceBlending=true
@@ -123,6 +124,7 @@ ShaderFormat=SF_METAL_MACES3_1
 bIsPC=true
 bIsMobile=true
 bSupports4ComponentUAVReadWrite=true
+bSupportsVertexShaderSRVs=false
 bSupportsManualVertexFetch=false
 bSupportsVolumeTextureAtomics=false
 bSupportsShaderPipelines = false
diff --git a/Engine/Config/TVOS/DataDrivenPlatformInfo.ini b/Engine/Config/TVOS/DataDrivenPlatformInfo.ini
--- a/Engine/Config/TVOS/DataDrivenPlatformInfo.ini
+++ b/Engine/Config/TVOS/DataDrivenPlatformInfo.ini
@@ -24,6 +24,7 @@ bIsMetalMRT=true
 bTargetsTiledGPU=true
 bSupportsIndexBufferUAVs=true
 bSupports4ComponentUAVReadWrite=true
+bSupportsVertexShaderSRVs=false
 bSupportsManualVertexFetch=false
 bSupportsDualSourceBlending=true
 FriendlyName=LOCTEXT("FriendlyShaderPlatformName_TVOS_SM5", "TVOS Metal SM5")
@@ -35,6 +36,7 @@ ShaderFormat=SF_METAL_TVOS
 bIsMobile=true
 bTargetsTiledGPU=true
 bSupports4ComponentUAVReadWrite=true
+bSupportsVertexShaderSRVs=false
 bSupportsManualVertexFetch=false
 bSupportsShaderPipelines = false
 FriendlyName=LOCTEXT("FriendlyShaderPlatformName_TVOS_Mobile", "TVOS Metal Mobile")
\ No newline at end of file
diff --git a/Engine/Config/VulkanPC/DataDrivenPlatformInfo.ini b/Engine/Config/VulkanPC/DataDrivenPlatformInfo.ini
--- a/Engine/Config/VulkanPC/DataDrivenPlatformInfo.ini
+++ b/Engine/Config/VulkanPC/DataDrivenPlatformInfo.ini
@@ -260,6 +260,7 @@ bIsPC=true
 bIsMobile=true
 bSupportsMobileMultiView=true
 bSupportsIndexBufferUAVs=true
+bSupportsVertexShaderSRVs=false
 bSupportsManualVertexFetch=false
 ; The support really is conditional, the runtime code needs to also test GRHISupportsArrayIndexFromAnyShader
 bSupportsVertexShaderLayer=true
diff --git a/Engine/Config/Windows/DataDrivenPlatformInfo.ini b/Engine/Config/Windows/DataDrivenPlatformInfo.ini
--- a/Engine/Config/Windows/DataDrivenPlatformInfo.ini
+++ b/Engine/Config/Windows/DataDrivenPlatformInfo.ini
@@ -154,6 +154,7 @@ bSupportsComputeFramework=true
 bSupportsByteBufferComputeShaders = true
 bSupportsInstancedStereo=true
 bSupportsVertexShaderLayer=true
+bSupportsVertexShaderSRVs=false
 bSupportsManualVertexFetch=false
 bSupportsSceneDataCompressedTransforms=true
 bSupportsDebugViewShaders=true
@@ -174,6 +175,7 @@ ShaderFormat=GLSL_150_ES31
 MaxFeatureLevel=ES3_1
 ShaderFormat=GLSL_150_ES31
 bIsMobile=true
+bSupportsVertexShaderSRVs=false
 bSupportsManualVertexFetch=false
 bSupportsClipDistance = false
 bSupportsShaderPipelines = false
diff --git a/Engine/Shaders/Private/SceneData.ush b/Engine/Shaders/Private/SceneData.ush
--- a/Engine/Shaders/Private/SceneData.ush
+++ b/Engine/Shaders/Private/SceneData.ush
@@ -227,8 +227,8 @@ float4 LoadPrimitivePrimitiveSceneDataElement(uint PrimitiveIndex, uint ItemInde
 // Fetch from scene primitive buffer
 FPrimitiveSceneData GetPrimitiveData(uint PrimitiveId)
 {
-#if (FEATURE_LEVEL == FEATURE_LEVEL_ES3_1 && VERTEXSHADER)
-	// Vertex shaders do not have access to GPUScene on mobile. Use GetPrimitiveData(FVertexFactoryIntermediates Intermediates)
+#if (VERTEXSHADER && !PLATFORM_SUPPORTS_VERTEX_SHADER_SRVS)
+	// Vertex shaders may not have access to GPUScene on mobile. Use GetPrimitiveData(FVertexFactoryIntermediates Intermediates)
 	// TODO: need a way to report invalid usage, after all dead code elimination
 	return (FPrimitiveSceneData)0;
 #else
diff --git a/Engine/Shaders/Public/Platform.ush b/Engine/Shaders/Public/Platform.ush
--- a/Engine/Shaders/Public/Platform.ush
+++ b/Engine/Shaders/Public/Platform.ush
@@ -476,6 +476,12 @@
 #define PLATFORM_SUPPORTS_4COMPONENT_UAV_READ_WRITE (XBOXONE_PROFILE || COMPILER_METAL)
 #endif
 
+// Whether the platform supports binding SRVs to the vertex shader stage.
+// This is generally available on high-end platforms and not available on mobile outside of some exceptions.
+// The real value is expected to be set based on DataDrivenShaderPlatformInfo in GlobalBeginCompileShader().
+#ifndef PLATFORM_SUPPORTS_VERTEX_SHADER_SRVS
+#define PLATFORM_SUPPORTS_VERTEX_SHADER_SRVS (FEATURE_LEVEL >= FEATURE_LEVEL_SM5)
+#endif
 
 // ---------------------------------------------------- Compiler specific defaults and fallbacks
 
diff --git a/Engine/Source/Runtime/Engine/Private/ShaderCompiler/ShaderCompiler.cpp b/Engine/Source/Runtime/Engine/Private/ShaderCompiler/ShaderCompiler.cpp
--- a/Engine/Source/Runtime/Engine/Private/ShaderCompiler/ShaderCompiler.cpp
+++ b/Engine/Source/Runtime/Engine/Private/ShaderCompiler/ShaderCompiler.cpp
@@ -6591,6 +6591,11 @@ void GlobalBeginCompileShader(
 		SET_SHADER_DEFINE(Input.Environment, PLATFORM_SUPPORTS_CLIP_DISTANCE, bSupportsClipDistance ? 1u : 0u);
 	}
 
+	{
+		const bool bSupportsVertexShaderSRVs = FDataDrivenShaderPlatformInfo::GetSupportsVertexShaderSRVs((EShaderPlatform)Target.Platform);
+		SET_SHADER_DEFINE(Input.Environment, PLATFORM_SUPPORTS_VERTEX_SHADER_SRVS, bSupportsVertexShaderSRVs ? 1u : 0u);
+	}
+
 	ITargetPlatform* TargetPlatform = GetTargetPlatformManager()->FindTargetPlatformWithSupport(TEXT("ShaderFormat"), ShaderFormatName);
 	bool bForwardShading = false;
 	{
diff --git a/Engine/Source/Runtime/RHI/Private/DataDrivenShaderPlatformInfo.cpp b/Engine/Source/Runtime/RHI/Private/DataDrivenShaderPlatformInfo.cpp
--- a/Engine/Source/Runtime/RHI/Private/DataDrivenShaderPlatformInfo.cpp
+++ b/Engine/Source/Runtime/RHI/Private/DataDrivenShaderPlatformInfo.cpp
@@ -142,6 +142,7 @@ void FGenericDataDrivenShaderPlatformInfo::SetDefaultValues()
 	bSupportsHZBOcclusion = true;
 	bSupportsWaterIndirectDraw = true;
 	bSupportsAsyncPipelineCompilation = true;
+	bSupportsVertexShaderSRVs = true; // Explicitly overriden to false for ES 3.1 platforms via DDPI ini
 	bSupportsManualVertexFetch = true;
 	bSupportsVolumeTextureAtomics = true;
 	bSupportsClipDistance = true;
@@ -268,6 +269,7 @@ void FGenericDataDrivenShaderPlatformInfo::ParseDataDrivenShaderInfo(const FConf
 	GET_SECTION_BOOL_HELPER(bSupportsHZBOcclusion);
 	GET_SECTION_BOOL_HELPER(bSupportsWaterIndirectDraw);
 	GET_SECTION_BOOL_HELPER(bSupportsAsyncPipelineCompilation);
+	GET_SECTION_BOOL_HELPER(bSupportsVertexShaderSRVs);
 	GET_SECTION_BOOL_HELPER(bSupportsManualVertexFetch);
 	GET_SECTION_BOOL_HELPER(bRequiresReverseCullingOnMobile);
 	GET_SECTION_BOOL_HELPER(bOverrideFMaterial_NeedsGBufferEnabled);
diff --git a/Engine/Source/Runtime/RHI/Public/DataDrivenShaderPlatformInfo.h b/Engine/Source/Runtime/RHI/Public/DataDrivenShaderPlatformInfo.h
--- a/Engine/Source/Runtime/RHI/Public/DataDrivenShaderPlatformInfo.h
+++ b/Engine/Source/Runtime/RHI/Public/DataDrivenShaderPlatformInfo.h
@@ -94,6 +94,7 @@ class FGenericDataDrivenShaderPlatformInfo
 	uint32 bSupportsHZBOcclusion : 1;
 	uint32 bSupportsWaterIndirectDraw : 1;
 	uint32 bSupportsAsyncPipelineCompilation : 1;
+	uint32 bSupportsVertexShaderSRVs : 1; // Whether SRVs can be bound to vertex shaders (may be independent from ManualVertexFetch)
 	uint32 bSupportsManualVertexFetch : 1;
 	uint32 bRequiresReverseCullingOnMobile : 1;
 	uint32 bOverrideFMaterial_NeedsGBufferEnabled : 1;
@@ -634,6 +635,12 @@ public:
 		return Infos[Platform].bSupportsAsyncPipelineCompilation;
 	}
 
+	static FORCEINLINE_DEBUGGABLE const bool GetSupportsVertexShaderSRVs(const FStaticShaderPlatform Platform)
+	{
+		check(IsValid(Platform));
+		return Infos[Platform].bSupportsVertexShaderSRVs;
+	}
+
 	static FORCEINLINE_DEBUGGABLE const bool GetSupportsManualVertexFetch(const FStaticShaderPlatform Platform)
 	{
 		check(IsValid(Platform));
