Implemented new work in progress shader bundle dispatch API, that can take an aggregation of compute shaders / pipelines, a GPU generated dispatch XYZ count buffer, and launch a platform specific method of compacting/removing the zero sized dispatches from the aggregation, while also doing as little work on the CPU as possible.

diff --git a/Engine/Source/Runtime/D3D12RHI/Private/D3D12CommandContext.h b/Engine/Source/Runtime/D3D12RHI/Private/D3D12CommandContext.h
--- a/Engine/Source/Runtime/D3D12RHI/Private/D3D12CommandContext.h
+++ b/Engine/Source/Runtime/D3D12RHI/Private/D3D12CommandContext.h
@@ -623,6 +623,8 @@ public:
 	void SetRenderTargets(uint32 NumSimultaneousRenderTargets, const FRHIRenderTargetView* NewRenderTargets, const FRHIDepthRenderTargetView* NewDepthStencilTarget);
 	void SetRenderTargetsAndClear(const FRHISetRenderTargetsInfo& RenderTargetsInfo);
 
+	virtual void RHIDispatchShaderBundle(FRHIShaderBundle* ShaderBundle, FRHIBuffer* ArgumentBuffer, TConstArrayView<FRHIShaderBundleDispatch> Dispatches) override;
+
 	virtual void RHIDrawPrimitive(uint32 BaseVertexIndex, uint32 NumPrimitives, uint32 NumInstances) final override;
 	virtual void RHIDrawPrimitiveIndirect(FRHIBuffer* ArgumentBuffer, uint32 ArgumentOffset) final override;
 	virtual void RHIDrawIndexedIndirect(FRHIBuffer* IndexBufferRHI, FRHIBuffer* ArgumentsBufferRHI, int32 DrawArgumentsIndex, uint32 NumInstances) final override;
diff --git a/Engine/Source/Runtime/D3D12RHI/Private/D3D12Commands.cpp b/Engine/Source/Runtime/D3D12RHI/Private/D3D12Commands.cpp
--- a/Engine/Source/Runtime/D3D12RHI/Private/D3D12Commands.cpp
+++ b/Engine/Source/Runtime/D3D12RHI/Private/D3D12Commands.cpp
@@ -1612,6 +1612,13 @@ void FD3D12CommandContext::CommitComputeResourceTables()
 	SetResourcesFromTables(ComputePSO->GetComputeShader());
 }
 
+void FD3D12CommandContext::RHIDispatchShaderBundle(FRHIShaderBundle* ShaderBundle, FRHIBuffer* ArgumentBuffer, TConstArrayView<FRHIShaderBundleDispatch> Dispatches)
+{
+	check(ShaderBundle != nullptr && Dispatches.Num() > 0);
+	TRHICommandList_RecursiveHazardous<FD3D12CommandContext> RHICmdList(this);
+	UE::RHICore::DispatchShaderBundleEmulation(RHICmdList, ShaderBundle, ArgumentBuffer, Dispatches);
+}
+
 void FD3D12CommandContext::RHIDrawPrimitive(uint32 BaseVertexIndex, uint32 NumPrimitives, uint32 NumInstances)
 {
 	RHI_DRAW_CALL_STATS(StateCache.GetGraphicsPipelinePrimitiveType(), FMath::Max(NumInstances, 1U) * NumPrimitives);
diff --git a/Engine/Source/Runtime/D3D12RHI/Private/D3D12RHIPrivate.h b/Engine/Source/Runtime/D3D12RHI/Private/D3D12RHIPrivate.h
--- a/Engine/Source/Runtime/D3D12RHI/Private/D3D12RHIPrivate.h
+++ b/Engine/Source/Runtime/D3D12RHI/Private/D3D12RHIPrivate.h
@@ -292,7 +292,7 @@ public:
     virtual void RHIUnlockStagingBuffer(FRHIStagingBuffer* StagingBuffer) final override;
 	virtual FBoundShaderStateRHIRef RHICreateBoundShaderState(FRHIVertexDeclaration* VertexDeclaration, FRHIVertexShader* VertexShader, FRHIPixelShader* PixelShader, FRHIGeometryShader* GeometryShader) final override;
 	virtual FGraphicsPipelineStateRHIRef RHICreateGraphicsPipelineState(const FGraphicsPipelineStateInitializer& Initializer) final override;
-	virtual TRefCountPtr<FRHIComputePipelineState> RHICreateComputePipelineState(FRHIComputeShader* ComputeShader) final override;
+	virtual FComputePipelineStateRHIRef RHICreateComputePipelineState(FRHIComputeShader* ComputeShader) final override;
 	virtual void RHICreateTransition(FRHITransition* Transition, const FRHITransitionCreateInfo& CreateInfo) final override;
 	virtual void RHIReleaseTransition(FRHITransition* Transition) final override;
 	virtual FUniformBufferRHIRef RHICreateUniformBuffer(const void* Contents, const FRHIUniformBufferLayout* Layout, EUniformBufferUsage Usage, EUniformBufferValidation Validation) final override;
@@ -449,6 +449,8 @@ public:
 	virtual void RHITransferRayTracingGeometryUnderlyingResource(FRHIRayTracingGeometry* DestGeometry, FRHIRayTracingGeometry* SrcGeometry) final override;
 #endif //D3D12_RHI_RAYTRACING
 
+	virtual FShaderBundleRHIRef RHICreateShaderBundle(uint32 NumRecords) override;
+
 	bool CheckGpuHeartbeat() const override;
 
 	virtual void HandleGpuTimeout(FD3D12Payload* Payload, double SecondsSinceSubmission);
diff --git a/Engine/Source/Runtime/D3D12RHI/Private/D3D12Resources.h b/Engine/Source/Runtime/D3D12RHI/Private/D3D12Resources.h
--- a/Engine/Source/Runtime/D3D12RHI/Private/D3D12Resources.h
+++ b/Engine/Source/Runtime/D3D12RHI/Private/D3D12Resources.h
@@ -974,6 +974,18 @@ private:
 	uint32 ShadowBufferSize;
 };
 
+class FD3D12ShaderBundle : public FRHIShaderBundle
+{
+	friend class FD3D12CommandContext;
+	friend class FD3D12DynamicRHI;
+
+public:
+	FD3D12ShaderBundle(FD3D12Device* InDevice, uint32 InNumRecords)
+		: FRHIShaderBundle(InNumRecords)
+	{
+	}
+};
+
 class FD3D12GPUFence final : public FRHIGPUFence
 {
 public:
@@ -1038,7 +1050,11 @@ struct TD3D12ResourceTraits<FRHIStagingBuffer>
 {
 	typedef FD3D12StagingBuffer TConcreteType;
 };
-
+template<>
+struct TD3D12ResourceTraits<FRHIShaderBundle>
+{
+	typedef FD3D12ShaderBundle TConcreteType;
+};
 
 #if D3D12_RHI_RAYTRACING
 template<>
diff --git a/Engine/Source/Runtime/D3D12RHI/Private/D3D12Shaders.cpp b/Engine/Source/Runtime/D3D12RHI/Private/D3D12Shaders.cpp
--- a/Engine/Source/Runtime/D3D12RHI/Private/D3D12Shaders.cpp
+++ b/Engine/Source/Runtime/D3D12RHI/Private/D3D12Shaders.cpp
@@ -212,6 +212,12 @@ FRayTracingShaderRHIRef FD3D12DynamicRHI::RHICreateRayTracingShader(TArrayView<c
 
 #endif // D3D12_RHI_RAYTRACING
 
+FShaderBundleRHIRef FD3D12DynamicRHI::RHICreateShaderBundle(uint32 NumRecords)
+{
+	FD3D12ShaderBundle* ShaderBundle = new FD3D12ShaderBundle(GetRHIDevice(0), NumRecords);
+	return ShaderBundle;
+}
+
 void FD3D12CommandContext::RHISetMultipleViewports(uint32 Count, const FViewportBounds* Data)
 {
 	// Structures are chosen to be directly mappable
diff --git a/Engine/Source/Runtime/D3D12RHI/Private/Windows/WindowsD3D12Device.cpp b/Engine/Source/Runtime/D3D12RHI/Private/Windows/WindowsD3D12Device.cpp
--- a/Engine/Source/Runtime/D3D12RHI/Private/Windows/WindowsD3D12Device.cpp
+++ b/Engine/Source/Runtime/D3D12RHI/Private/Windows/WindowsD3D12Device.cpp
@@ -1690,6 +1690,8 @@ void FD3D12DynamicRHI::Init()
 	GRHISupportsTextureStreaming = true;
 	GRHISupportsFirstInstance = true;
 
+	GRHISupportsDispatchShaderBundle = false; // TODO: Implement
+
 	// Indicate that the RHI needs to use the engine's deferred deletion queue.
 	GRHINeedsExtraDeletionLatency = true;
 
diff --git a/Engine/Source/Runtime/RHI/Public/DynamicRHI.h b/Engine/Source/Runtime/RHI/Public/DynamicRHI.h
--- a/Engine/Source/Runtime/RHI/Public/DynamicRHI.h
+++ b/Engine/Source/Runtime/RHI/Public/DynamicRHI.h
@@ -984,6 +984,12 @@ public:
 	}
 #endif // RHI_RAYTRACING
 
+	virtual FShaderBundleRHIRef RHICreateShaderBundle(uint32 NumRecords)
+	{
+		checkNoEntry();
+		return nullptr;
+	}
+
 protected:
 	TArray<uint32> PixelFormatBlockBytes;
 	friend class FValidationRHI;
@@ -1475,6 +1481,11 @@ FORCEINLINE FRayTracingShaderRHIRef RHICreateRayTracingShader(TArrayView<const u
 
 #endif // RHI_RAYTRACING
 
+FORCEINLINE FShaderBundleRHIRef RHICreateShaderBundle(uint32 NumRecords)
+{
+	return GDynamicRHI->RHICreateShaderBundle(NumRecords);
+}
+
 /**
 * Defragment the texture pool.
 */
diff --git a/Engine/Source/Runtime/RHI/Public/RHICommandList.h b/Engine/Source/Runtime/RHI/Public/RHICommandList.h
--- a/Engine/Source/Runtime/RHI/Public/RHICommandList.h
+++ b/Engine/Source/Runtime/RHI/Public/RHICommandList.h
@@ -1681,6 +1681,25 @@ FRHICOMMAND_MACRO(FRHICommandDispatchIndirectComputeShader)
 	RHI_API void Execute(FRHICommandListBase& CmdList);
 };
 
+FRHICOMMAND_MACRO(FRHICommandDispatchShaderBundle)
+{
+	FRHIShaderBundle* ShaderBundle;
+	FRHIBuffer* ArgumentBuffer;
+	TArray<FRHIShaderBundleDispatch> Dispatches;
+	FORCEINLINE_DEBUGGABLE FRHICommandDispatchShaderBundle()
+		: ShaderBundle(nullptr)
+		, ArgumentBuffer(nullptr)
+	{
+	}
+	FORCEINLINE_DEBUGGABLE FRHICommandDispatchShaderBundle(FRHIShaderBundle* InShaderBundle, FRHIBuffer* InArgumentBuffer, TConstArrayView<FRHIShaderBundleDispatch> InDispatches)
+		: ShaderBundle(InShaderBundle)
+		, ArgumentBuffer(InArgumentBuffer)
+		, Dispatches(InDispatches)
+	{
+	}
+	RHI_API void Execute(FRHICommandListBase& CmdList);
+};
+
 FRHICOMMAND_MACRO(FRHICommandBeginUAVOverlap)
 {
 	RHI_API void Execute(FRHICommandListBase& CmdList);
@@ -2792,6 +2811,16 @@ PRAGMA_ENABLE_DEPRECATION_WARNINGS
 		}
 	}
 
+	FORCEINLINE_DEBUGGABLE void DispatchShaderBundle(FRHIShaderBundle* ShaderBundle, FRHIBuffer* ArgumentBuffer, TConstArrayView<FRHIShaderBundleDispatch> Dispatches)
+	{
+		if (Bypass())
+		{
+			GetContext().RHIDispatchShaderBundle(ShaderBundle, ArgumentBuffer, Dispatches);
+			return;
+		}
+		ALLOC_COMMAND(FRHICommandDispatchShaderBundle)(ShaderBundle, ArgumentBuffer, Dispatches);
+	}
+
 	FORCEINLINE_DEBUGGABLE void BeginUAVOverlap()
 	{
 		if (Bypass())
diff --git a/Engine/Source/Runtime/RHI/Public/RHICommandListCommandExecutes.inl b/Engine/Source/Runtime/RHI/Public/RHICommandListCommandExecutes.inl
--- a/Engine/Source/Runtime/RHI/Public/RHICommandListCommandExecutes.inl
+++ b/Engine/Source/Runtime/RHI/Public/RHICommandListCommandExecutes.inl
@@ -242,6 +242,17 @@ void FRHICommandDispatchIndirectComputeShader::Execute(FRHICommandListBase& CmdL
 	INTERNAL_DECORATOR_COMPUTE(RHIDispatchIndirectComputeShader)(ArgumentBuffer, ArgumentOffset);
 }
 
+void FRHICommandDispatchShaderBundle::Execute(FRHICommandListBase& CmdList)
+{
+	RHISTAT(DispatchShaderBundle);
+	extern RHI_API FRHIComputePipelineState* ExecuteSetComputePipelineState(FComputePipelineState* ComputePipelineState);
+	for (FRHIShaderBundleDispatch& Dispatch : Dispatches)
+	{
+		Dispatch.RHIPipeline = ExecuteSetComputePipelineState(Dispatch.PipelineState);
+	}
+	INTERNAL_DECORATOR_COMPUTE(RHIDispatchShaderBundle)(ShaderBundle, ArgumentBuffer, Dispatches);
+}
+
 void FRHICommandBeginUAVOverlap::Execute(FRHICommandListBase& CmdList)
 {
 	RHISTAT(BeginUAVOverlap);
diff --git a/Engine/Source/Runtime/RHI/Public/RHIContext.h b/Engine/Source/Runtime/RHI/Public/RHIContext.h
--- a/Engine/Source/Runtime/RHI/Public/RHIContext.h
+++ b/Engine/Source/Runtime/RHI/Public/RHIContext.h
@@ -284,6 +284,8 @@ public:
 	*/
 	virtual void RHIClearUAVUint(FRHIUnorderedAccessView* UnorderedAccessViewRHI, const FUintVector4& Values) = 0;
 
+	virtual void RHIDispatchShaderBundle(FRHIShaderBundle* ShaderBundle, FRHIBuffer* ArgumentBuffer, TConstArrayView<FRHIShaderBundleDispatch> Dispatches) {}
+
 	virtual void RHIBeginUAVOverlap() {}
 	virtual void RHIEndUAVOverlap() {}
 
diff --git a/Engine/Source/Runtime/RHI/Public/RHIDefinitions.h b/Engine/Source/Runtime/RHI/Public/RHIDefinitions.h
--- a/Engine/Source/Runtime/RHI/Public/RHIDefinitions.h
+++ b/Engine/Source/Runtime/RHI/Public/RHIDefinitions.h
@@ -878,6 +878,7 @@ enum ERHIResourceType : uint8
 	RRT_CustomPresent,
 	RRT_ShaderLibrary,
 	RRT_PipelineBinaryLibrary,
+	RRT_ShaderBundle,
 
 	RRT_Num
 };
diff --git a/Engine/Source/Runtime/RHI/Public/RHIFwd.h b/Engine/Source/Runtime/RHI/Public/RHIFwd.h
--- a/Engine/Source/Runtime/RHI/Public/RHIFwd.h
+++ b/Engine/Source/Runtime/RHI/Public/RHIFwd.h
@@ -69,6 +69,7 @@ class FRHISamplerState;
 class FRHIShader;
 class FRHIShaderLibrary;
 class FRHIShaderResourceView;
+class FRHIShaderBundle;
 class FRHIStagingBuffer;
 class FRHITexture;
 class FRHITextureReference;
@@ -108,6 +109,7 @@ using FRHIPipelineBinaryLibraryRef     = TRefCountPtr<FRHIPipelineBinaryLibrary>
 using FRHIShaderLibraryRef             = TRefCountPtr<FRHIShaderLibrary>;
 using FSamplerStateRHIRef              = TRefCountPtr<FRHISamplerState>;
 using FShaderResourceViewRHIRef        = TRefCountPtr<FRHIShaderResourceView>;
+using FShaderBundleRHIRef              = TRefCountPtr<FRHIShaderBundle>;
 using FStagingBufferRHIRef             = TRefCountPtr<FRHIStagingBuffer>;
 using FTextureReferenceRHIRef          = TRefCountPtr<FRHITextureReference>;
 using FTextureRHIRef                   = TRefCountPtr<FRHITexture>;
diff --git a/Engine/Source/Runtime/RHI/Public/RHIGlobals.h b/Engine/Source/Runtime/RHI/Public/RHIGlobals.h
--- a/Engine/Source/Runtime/RHI/Public/RHIGlobals.h
+++ b/Engine/Source/Runtime/RHI/Public/RHIGlobals.h
@@ -557,6 +557,9 @@ struct FRHIGlobals
 	/** Whether current RHI supports overestimated conservative rasterization. */
 	bool SupportsConservativeRasterization = false;
 
+	/** Whether current RHI supports dispatching shader bundles. */
+	bool SupportsDispatchShaderBundle = false;
+
 	/** true if the RHI supports Mesh and Amplification shaders with tier0 capability */
 	bool SupportsMeshShadersTier0 = false;
 
@@ -740,6 +743,7 @@ extern RHI_API FRHIGlobals GRHIGlobals;
 #define GRHISupportsPipelineFileCache                          GRHIGlobals.SupportsPipelineFileCache
 #define GRHISupportsStencilRefFromPixelShader                  GRHIGlobals.SupportsStencilRefFromPixelShader
 #define GRHISupportsConservativeRasterization                  GRHIGlobals.SupportsConservativeRasterization
+#define GRHISupportsDispatchShaderBundle                       GRHIGlobals.SupportsDispatchShaderBundle
 #define GRHISupportsMeshShadersTier0                           GRHIGlobals.SupportsMeshShadersTier0
 #define GRHISupportsMeshShadersTier1                           GRHIGlobals.SupportsMeshShadersTier1
 #define GRHISupportsShaderTimestamp                            GRHIGlobals.SupportsShaderTimestamp
diff --git a/Engine/Source/Runtime/RHI/Public/RHIResources.h b/Engine/Source/Runtime/RHI/Public/RHIResources.h
--- a/Engine/Source/Runtime/RHI/Public/RHIResources.h
+++ b/Engine/Source/Runtime/RHI/Public/RHIResources.h
@@ -902,6 +902,7 @@ private:
 	FExclusiveDepthStencil DSMode;
 #endif
 };
+
 class FRHIComputePipelineState : public FRHIResource
 {
 public:
@@ -913,6 +914,7 @@ public:
 private:
 	bool bIsValid = true;
 };
+
 class FRHIRayTracingPipelineState : public FRHIResource
 {
 public:
@@ -3067,6 +3069,22 @@ public:
 	virtual uint32 GetLayerBufferOffset(uint32 LayerIndex) const = 0;
 };
 
+class FRHIShaderBundle : public FRHIResource
+{
+public:
+	// Dispatch XYZ + Padding
+	static constexpr uint32 ArgumentByteStride = sizeof(uint32) * 4u;
+
+	const uint32 NumRecords = 0;
+
+public:
+	FRHIShaderBundle(uint32 InNumRecords)
+		: FRHIResource(RRT_ShaderBundle)
+		, NumRecords(InNumRecords)
+	{
+	}
+};
+
 /* Generic staging buffer class used by FRHIGPUMemoryReadback
 * RHI specific staging buffers derive from this
 */
diff --git a/Engine/Source/Runtime/RHI/Public/RHIShaderParameters.h b/Engine/Source/Runtime/RHI/Public/RHIShaderParameters.h
--- a/Engine/Source/Runtime/RHI/Public/RHIShaderParameters.h
+++ b/Engine/Source/Runtime/RHI/Public/RHIShaderParameters.h
@@ -182,3 +182,11 @@ struct FRHIBatchedShaderUnbinds
 	}
 };
 
+struct FRHIShaderBundleDispatch
+{
+	uint32 RecordIndex = ~uint32(0u);
+	class FComputePipelineState* PipelineState = nullptr;
+	FRHIComputeShader* Shader = nullptr;
+	FRHIComputePipelineState* RHIPipeline = nullptr;
+	FRHIBatchedShaderParameters Parameters;
+};
\ No newline at end of file
diff --git a/Engine/Source/Runtime/RHI/Public/RHIValidation.h b/Engine/Source/Runtime/RHI/Public/RHIValidation.h
--- a/Engine/Source/Runtime/RHI/Public/RHIValidation.h
+++ b/Engine/Source/Runtime/RHI/Public/RHIValidation.h
@@ -1290,6 +1290,10 @@ public:
 	}
 #endif // RHI_RAYTRACING
 
+	virtual FShaderBundleRHIRef RHICreateShaderBundle(uint32 NumRecords) override final
+	{
+		return RHI->RHICreateShaderBundle(NumRecords);
+	}
 
 //protected:
 	static void ReportValidationFailure(const TCHAR* InMessage);
diff --git a/Engine/Source/Runtime/RHI/Public/RHIValidationContext.h b/Engine/Source/Runtime/RHI/Public/RHIValidationContext.h
--- a/Engine/Source/Runtime/RHI/Public/RHIValidationContext.h
+++ b/Engine/Source/Runtime/RHI/Public/RHIValidationContext.h
@@ -165,6 +165,27 @@ public:
 		RHIContext->RHIClearUAVUint(UnorderedAccessViewRHI, Values);
 	}
 
+	virtual void RHIDispatchShaderBundle(FRHIShaderBundle* ShaderBundleRHI, FRHIBuffer* ArgumentBufferRHI, TConstArrayView<FRHIShaderBundleDispatch> Dispatches) final override
+	{
+		checkf(Dispatches.Num() > 0, TEXT("A shader bundle must be dispatched with at least one record."));
+		for (const FRHIShaderBundleDispatch& Dispatch : Dispatches)
+		{
+			State.bComputePSOSet = true;
+
+			// Reset the compute UAV tracker since the renderer must re-bind all resources after changing a shader.
+			Tracker->ResetUAVState(RHIValidation::EUAVMode::Compute);
+
+			ValidateShaderParameters(Dispatch.Shader, Tracker, State.StaticUniformBuffers, Dispatch.Parameters.ResourceParameters, ERHIAccess::SRVCompute, RHIValidation::EUAVMode::Compute);
+			ValidateShaderParameters(Dispatch.Shader, Tracker, State.StaticUniformBuffers, Dispatch.Parameters.BindlessParameters, ERHIAccess::SRVCompute, RHIValidation::EUAVMode::Compute);
+
+			const uint32 ArgumentOffset = (Dispatch.RecordIndex * FRHIShaderBundle::ArgumentByteStride);
+			FValidationRHI::ValidateDispatchIndirectArgsBuffer(ArgumentBufferRHI, ArgumentOffset);
+		}
+
+		Tracker->Assert(ArgumentBufferRHI->GetWholeResourceIdentity(), ERHIAccess::IndirectArgs);
+		RHIContext->RHIDispatchShaderBundle(ShaderBundleRHI, ArgumentBufferRHI, Dispatches);
+	}
+
 	virtual void RHIBeginUAVOverlap() final override
 	{
 		Tracker->AllUAVsOverlap(true);
@@ -473,6 +494,27 @@ public:
 		RHIContext->RHIClearUAVUint(UnorderedAccessViewRHI, Values);
 	}
 
+	virtual void RHIDispatchShaderBundle(FRHIShaderBundle* ShaderBundleRHI, FRHIBuffer* ArgumentBufferRHI, TConstArrayView<FRHIShaderBundleDispatch> Dispatches) final override
+	{
+		checkf(Dispatches.Num() > 0, TEXT("A shader bundle must be dispatched with at least one record."));
+		for (const FRHIShaderBundleDispatch& Dispatch : Dispatches)
+		{
+			State.bComputePSOSet = true;
+
+			// Reset the compute UAV tracker since the renderer must re-bind all resources after changing a shader.
+			Tracker->ResetUAVState(RHIValidation::EUAVMode::Compute);
+
+			ValidateShaderParameters(Dispatch.Shader, Tracker, State.StaticUniformBuffers, Dispatch.Parameters.ResourceParameters, ERHIAccess::SRVCompute, RHIValidation::EUAVMode::Compute);
+			ValidateShaderParameters(Dispatch.Shader, Tracker, State.StaticUniformBuffers, Dispatch.Parameters.BindlessParameters, ERHIAccess::SRVCompute, RHIValidation::EUAVMode::Compute);
+
+			const uint32 ArgumentOffset = (Dispatch.RecordIndex * FRHIShaderBundle::ArgumentByteStride);
+			FValidationRHI::ValidateDispatchIndirectArgsBuffer(ArgumentBufferRHI, ArgumentOffset);
+		}
+
+		Tracker->Assert(ArgumentBufferRHI->GetWholeResourceIdentity(), ERHIAccess::IndirectArgs);
+		RHIContext->RHIDispatchShaderBundle(ShaderBundleRHI, ArgumentBufferRHI, Dispatches);
+	}
+
 	virtual void RHIBeginUAVOverlap() final override
 	{
 		Tracker->AllUAVsOverlap(true);
diff --git a/Engine/Source/Runtime/RHICore/Private/RHICoreShader.cpp b/Engine/Source/Runtime/RHICore/Private/RHICoreShader.cpp
--- a/Engine/Source/Runtime/RHICore/Private/RHICoreShader.cpp
+++ b/Engine/Source/Runtime/RHICore/Private/RHICoreShader.cpp
@@ -1,6 +1,7 @@
 // Copyright Epic Games, Inc. All Rights Reserved.
 
 #include "RHICoreShader.h"
+#include "PipelineStateCache.h"
 
 #define RHI_VALIDATE_STATIC_UNIFORM_BUFFERS (!UE_BUILD_SHIPPING && !UE_BUILD_TEST)
 
@@ -80,5 +81,36 @@ void SetupShaderCodeValidationData(FRHIShader* RHIShader, FShaderCodeReader& Sha
 #endif
 }
 
+void DispatchShaderBundleEmulation(
+	FRHIComputeCommandList& InRHICmdList,
+	FRHIShaderBundle* ShaderBundle,
+	FRHIBuffer* ArgumentBuffer,
+	TConstArrayView<FRHIShaderBundleDispatch> Dispatches)
+{
+	for (const FRHIShaderBundleDispatch& Dispatch : Dispatches)
+	{
+		if (Dispatch.Shader == nullptr)
+		{
+			continue;
+		}
+
+		SetComputePipelineState(InRHICmdList, Dispatch.Shader);
+
+		if (Dispatch.Parameters.HasParameters())
+		{
+			InRHICmdList.SetShaderParameters(
+				Dispatch.Shader,
+				Dispatch.Parameters.ParametersData,
+				Dispatch.Parameters.Parameters,
+				Dispatch.Parameters.ResourceParameters,
+				Dispatch.Parameters.BindlessParameters
+			);
+		}
+
+		const uint32 IndirectOffset = (Dispatch.RecordIndex * FRHIShaderBundle::ArgumentByteStride);
+		InRHICmdList.DispatchIndirectComputeShader(ArgumentBuffer, IndirectOffset);
+	}
+}
+
 } //! RHICore
 } //! UE
\ No newline at end of file
diff --git a/Engine/Source/Runtime/RHICore/Public/RHICoreShader.h b/Engine/Source/Runtime/RHICore/Public/RHICoreShader.h
--- a/Engine/Source/Runtime/RHICore/Public/RHICoreShader.h
+++ b/Engine/Source/Runtime/RHICore/Public/RHICoreShader.h
@@ -14,6 +14,14 @@ namespace RHICore
 extern RHICORE_API void ValidateStaticUniformBuffer(FRHIUniformBuffer* UniformBuffer, FUniformBufferStaticSlot Slot, uint32 ExpectedHash);
 extern RHICORE_API void SetupShaderCodeValidationData(FRHIShader* RHIShader, class FShaderCodeReader& ShaderCodeReader);
 
+/** Common implementation of dispatch shader bundle emulation shared by RHIs */
+extern RHICORE_API void DispatchShaderBundleEmulation(
+	FRHIComputeCommandList& InRHICmdList,
+	FRHIShaderBundle* ShaderBundle,
+	FRHIBuffer* ArgumentBuffer,
+	TConstArrayView<FRHIShaderBundleDispatch> Dispatches
+);
+
 inline void InitStaticUniformBufferSlots(TArray<FUniformBufferStaticSlot>& StaticSlots, const FShaderResourceTable& ShaderResourceTable)
 {
 	StaticSlots.Reserve(ShaderResourceTable.ResourceTableLayoutHashes.Num());
