Implemented shader compiler DDPI support for SupportsShaderBundleDispatch

diff --git a/Engine/Source/Runtime/Engine/Private/ShaderCompiler/ShaderCompiler.cpp b/Engine/Source/Runtime/Engine/Private/ShaderCompiler/ShaderCompiler.cpp
--- a/Engine/Source/Runtime/Engine/Private/ShaderCompiler/ShaderCompiler.cpp
+++ b/Engine/Source/Runtime/Engine/Private/ShaderCompiler/ShaderCompiler.cpp
@@ -6699,6 +6699,7 @@ void GlobalBeginCompileShader(
 		SET_SHADER_DEFINE(Input.Environment, LA_NORMALMAPS, 1);
 	}
 
+	SET_SHADER_DEFINE(Input.Environment, PLATFORM_SUPPORTS_SHADER_BUNDLE_DISPATCH, RHISupportsShaderBundleDispatch(EShaderPlatform(Target.Platform)) ? 1 : 0);
 	SET_SHADER_DEFINE(Input.Environment, PLATFORM_SUPPORTS_RENDERTARGET_WRITE_MASK, RHISupportsRenderTargetWriteMask(EShaderPlatform(Target.Platform)) ? 1 : 0);
 	SET_SHADER_DEFINE(Input.Environment, PLATFORM_SUPPORTS_PER_PIXEL_DBUFFER_MASK, FDataDrivenShaderPlatformInfo::GetSupportsPerPixelDBufferMask(EShaderPlatform(Target.Platform)) ? 1 : 0);
 	SET_SHADER_DEFINE(Input.Environment, PLATFORM_SUPPORTS_DISTANCE_FIELDS, DoesPlatformSupportDistanceFields(EShaderPlatform(Target.Platform)) ? 1 : 0);
diff --git a/Engine/Source/Runtime/Engine/Private/ShaderCompiler/ShaderGenerationUtil.cpp b/Engine/Source/Runtime/Engine/Private/ShaderCompiler/ShaderGenerationUtil.cpp
--- a/Engine/Source/Runtime/Engine/Private/ShaderCompiler/ShaderGenerationUtil.cpp
+++ b/Engine/Source/Runtime/Engine/Private/ShaderCompiler/ShaderGenerationUtil.cpp
@@ -420,6 +420,7 @@ static FShaderGlobalDefines FetchShaderGlobalDefines(EShaderPlatform TargetPlatf
 		Ret.POST_PROCESS_ALPHA = PropagateAlpha != 0;
 	}
 
+	Ret.PLATFORM_SUPPORTS_SHADER_BUNDLE_DISPATCH = RHISupportsShaderBundleDispatch(EShaderPlatform(TargetPlatform)) ? 1 : 0;
 	Ret.PLATFORM_SUPPORTS_RENDERTARGET_WRITE_MASK = RHISupportsRenderTargetWriteMask(EShaderPlatform(TargetPlatform)) ? 1 : 0;
 	Ret.PLATFORM_SUPPORTS_PER_PIXEL_DBUFFER_MASK = FDataDrivenShaderPlatformInfo::GetSupportsPerPixelDBufferMask(EShaderPlatform(TargetPlatform)) ? 1 : 0;
 	Ret.PLATFORM_SUPPORTS_DISTANCE_FIELDS = DoesPlatformSupportDistanceFields(EShaderPlatform(TargetPlatform)) ? 1 : 0;
diff --git a/Engine/Source/Runtime/RHI/Private/DataDrivenShaderPlatformInfo.cpp b/Engine/Source/Runtime/RHI/Private/DataDrivenShaderPlatformInfo.cpp
--- a/Engine/Source/Runtime/RHI/Private/DataDrivenShaderPlatformInfo.cpp
+++ b/Engine/Source/Runtime/RHI/Private/DataDrivenShaderPlatformInfo.cpp
@@ -214,6 +214,7 @@ void FGenericDataDrivenShaderPlatformInfo::ParseDataDrivenShaderInfo(const FConf
 	GET_SECTION_SUPPORT_HELPER(SupportsMultiViewport);
 	GET_SECTION_BOOL_HELPER(bSupportsMSAA);
 	GET_SECTION_BOOL_HELPER(bSupports4ComponentUAVReadWrite);
+	GET_SECTION_BOOL_HELPER(bSupportsShaderBundleDispatch);
 	GET_SECTION_BOOL_HELPER(bSupportsRenderTargetWriteMask);
 	GET_SECTION_BOOL_HELPER(bSupportsRayTracing);
 	GET_SECTION_BOOL_HELPER(bSupportsRayTracingShaders);
@@ -428,6 +429,7 @@ void FGenericDataDrivenShaderPlatformInfo::UpdatePreviewPlatforms()
 				}
 				Infos[ShaderPlatform].bSupportsMobileMultiView &= Infos[EditorSPForPreviewMaxFeatureLevel].bSupportsMobileMultiView;
 				Infos[ShaderPlatform].bSupportsManualVertexFetch = Infos[EditorSPForPreviewMaxFeatureLevel].bSupportsManualVertexFetch;
+				Infos[ShaderPlatform].bSupportsShaderBundleDispatch = false;
 				Infos[ShaderPlatform].bSupportsRenderTargetWriteMask = false;
 				Infos[ShaderPlatform].bSupportsIntrinsicWaveOnce = false;
 				Infos[ShaderPlatform].bSupportsDOFHybridScattering = false;
diff --git a/Engine/Source/Runtime/RHI/Public/DataDrivenShaderPlatformInfo.h b/Engine/Source/Runtime/RHI/Public/DataDrivenShaderPlatformInfo.h
--- a/Engine/Source/Runtime/RHI/Public/DataDrivenShaderPlatformInfo.h
+++ b/Engine/Source/Runtime/RHI/Public/DataDrivenShaderPlatformInfo.h
@@ -42,6 +42,7 @@ class FGenericDataDrivenShaderPlatformInfo
 	uint32 SupportsMultiViewport : int32(ERHIFeatureSupport::NumBits);
 	uint32 bSupportsMSAA : 1;
 	uint32 bSupports4ComponentUAVReadWrite : 1;
+	uint32 bSupportsShaderBundleDispatch : 1;
 	uint32 bSupportsRenderTargetWriteMask : 1;
 	uint32 bSupportsRayTracing : 1;
 	uint32 bSupportsRayTracingCallableShaders : 1;
@@ -306,6 +307,12 @@ public:
 		return Infos[Platform].bSupportsSwapchainUAVs;
 	}
 
+	static FORCEINLINE_DEBUGGABLE const bool GetSupportsShaderBundleDispatch(const FStaticShaderPlatform Platform)
+	{
+		check(IsValid(Platform));
+		return Infos[Platform].bSupportsShaderBundleDispatch;
+	}
+
 	static FORCEINLINE_DEBUGGABLE const bool GetSupportsRenderTargetWriteMask(const FStaticShaderPlatform Platform)
 	{
 		check(IsValid(Platform));
@@ -1093,6 +1100,12 @@ inline bool RHISupportsWaveOperations(const FStaticShaderPlatform Platform)
 	return FDataDrivenShaderPlatformInfo::GetSupportsWaveOperations(Platform) != ERHIFeatureSupport::Unsupported;
 }
 
+/** True if the given shader platform supports shader bundle dispatch */
+inline bool RHISupportsShaderBundleDispatch(const FStaticShaderPlatform Platform)
+{
+	return FDataDrivenShaderPlatformInfo::GetSupportsShaderBundleDispatch(Platform);
+}
+
 /** True if the given shader platform supports a render target write mask */
 inline bool RHISupportsRenderTargetWriteMask(const FStaticShaderPlatform Platform)
 {
diff --git a/Engine/Source/Runtime/RenderCore/Private/Shader.cpp b/Engine/Source/Runtime/RenderCore/Private/Shader.cpp
--- a/Engine/Source/Runtime/RenderCore/Private/Shader.cpp
+++ b/Engine/Source/Runtime/RenderCore/Private/Shader.cpp
@@ -2133,6 +2133,11 @@ void ShaderMapAppendKeyString(EShaderPlatform Platform, FString& KeyString)
 		}
 	}
 
+	if (RHISupportsShaderBundleDispatch(Platform))
+	{
+		KeyString += TEXT("_SHBD");
+	}
+
 	if (RHISupportsRenderTargetWriteMask(Platform))
 	{
 		KeyString += TEXT("_RTWM");
diff --git a/Engine/Source/Runtime/RenderCore/Public/ShaderMaterial.h b/Engine/Source/Runtime/RenderCore/Public/ShaderMaterial.h
--- a/Engine/Source/Runtime/RenderCore/Public/ShaderMaterial.h
+++ b/Engine/Source/Runtime/RenderCore/Public/ShaderMaterial.h
@@ -35,6 +35,7 @@ struct FShaderGlobalDefines
 	bool SUPPORT_CLOUD_SHADOW_ON_FORWARD_LIT_TRANSLUCENT;
 	bool SUPPORT_CLOUD_SHADOW_ON_SINGLE_LAYER_WATER;
 	bool POST_PROCESS_ALPHA;
+	bool PLATFORM_SUPPORTS_SHADER_BUNDLE_DISPATCH;
 	bool PLATFORM_SUPPORTS_RENDERTARGET_WRITE_MASK;
 	bool PLATFORM_SUPPORTS_PER_PIXEL_DBUFFER_MASK;
 	bool PLATFORM_SUPPORTS_DISTANCE_FIELDS;
