Always add PSO precompile task to prerequisite of the command list if the PSO task is not complete.

Currently, the task is only added to the prerequisite if the command list supports async compilation, but that's irrelevant once the task is already in flight.

#pf 649c99c390320259adad3ef9

diff --git a/Engine/Source/Runtime/RHI/Private/PipelineStateCache.cpp b/Engine/Source/Runtime/RHI/Private/PipelineStateCache.cpp
--- a/Engine/Source/Runtime/RHI/Private/PipelineStateCache.cpp
+++ b/Engine/Source/Runtime/RHI/Private/PipelineStateCache.cpp
@@ -2646,7 +2646,7 @@ FGraphicsPipelineState* PipelineStateCache::GetAndOrCreateGraphicsPipelineState(
 	}
 	else
 	{
-		if (DoAsyncCompile && !Initializer.bFromPSOFileCache && !OutCachedState->IsComplete())
+		if (!Initializer.bFromPSOFileCache && !OutCachedState->IsComplete())
 		{
 			if (OutCachedState->PrecompileTask)
 			{
