Add an option to only enable the user pipeline cache and disable the static (packaged) file cache.

#pf 64b6de9e5f0e72af69e5a39b

diff --git a/Engine/Source/Runtime/RHI/Private/PipelineFileCache.cpp b/Engine/Source/Runtime/RHI/Private/PipelineFileCache.cpp
--- a/Engine/Source/Runtime/RHI/Private/PipelineFileCache.cpp
+++ b/Engine/Source/Runtime/RHI/Private/PipelineFileCache.cpp
@@ -185,6 +185,7 @@ FPipelineFileCacheManager::PSOOrder FPipelineFileCacheManager::RequestedOrder =
 bool FPipelineFileCacheManager::FileCacheEnabled = false;
 FPipelineFileCacheManager::FPipelineStateLoggedEvent FPipelineFileCacheManager::PSOLoggedEvent;
 uint64 FPipelineFileCacheManager::GameUsageMask = 0;
+bool FPipelineFileCacheManager::LogNewPSOsToConsoleAndCSV = true;
 
 static int64 GetCurrentUnixTime()
 {
@@ -3374,6 +3375,64 @@ void FPipelineFileCacheManager::RegisterPSOUsageDataUpdateForNextSave(FPSOUsageD
 	CurrentEntry.EngineFlags |= UsageData.EngineFlags;
 }  
 
+void FPipelineFileCacheManager::LogNewGraphicsPSOToConsoleAndCSV(FPipelineCacheFileFormatPSO& PSO, uint32 PSOHash, bool bWasPSOPrecached)
+{
+	if (!LogNewPSOsToConsoleAndCSV)
+	{
+		return;
+	}
+
+	if (!bWasPSOPrecached)
+	{
+		CSV_EVENT(PSO, TEXT("Encountered new graphics PSO"));
+		UE_LOG(LogRHI, Display, TEXT("Encountered a new graphics PSO: %u"), PSOHash);
+		int32 LogDetailLevel = LogPSODetails() ? 2 : GPSOFileCachePrintNewPSODescriptors;
+		if (LogDetailLevel > 0)
+		{
+			UE_LOG(LogRHI, Display, TEXT("New Graphics PSO (%u)"), PSOHash);
+			if (LogDetailLevel > 1)
+			{
+				UE_LOG(LogRHI, Display, TEXT("%s"), *PSO.ToStringReadable());
+			}
+		}
+
+	}
+	else
+	{
+		UE_LOG(LogRHI, Log, TEXT("Encountered a new graphics PSO for the file cache but it was already precached at runtime: %u"), PSOHash);
+	}
+}
+
+void FPipelineFileCacheManager::LogNewComputePSOToConsoleAndCSV(FPipelineCacheFileFormatPSO& PSO, uint32 PSOHash)
+{
+	if (!LogNewPSOsToConsoleAndCSV)
+	{
+		return;
+	}
+
+	CSV_EVENT(PSO, TEXT("Encountered new compute PSO"));
+	UE_LOG(LogRHI, Display, TEXT("Encountered a new compute PSO: %u"), PSOHash);
+	if (GPSOFileCachePrintNewPSODescriptors > 0)
+	{
+		UE_LOG(LogRHI, Display, TEXT("New compute PSO (%u) Description: %s"), PSOHash, *PSO.ComputeDesc.ComputeShader.ToString());
+	}
+}
+
+void FPipelineFileCacheManager::LogNewRaytracingPSOToConsole(FPipelineCacheFileFormatPSO& PSO, uint32 PSOHash, bool bIsNonBlockingPSO)
+{
+	// When non-blocking creation is used, encountering a non-cached RTPSO is not likely to cause a hitch and so the logging is not useful/actionable.
+	if (!LogNewPSOsToConsoleAndCSV || bIsNonBlockingPSO)
+	{
+		return;
+	}
+
+	UE_LOG(LogRHI, Display, TEXT("Encountered a new ray tracing PSO: %u"), PSOHash);
+	if (GPSOFileCachePrintNewPSODescriptors > 0)
+	{
+		UE_LOG(LogRHI, Display, TEXT("New ray tracing PSO (%u) Description: %s"), PSOHash, *PSO.RayTracingDesc.ToString());
+	}
+}
+
 void FPipelineFileCacheManager::CacheGraphicsPSO(uint32 RunTimeHash, FGraphicsPipelineStateInitializer const& Initializer, bool bWasPSOPrecached)
 {
 	if(IsPipelineFileCacheEnabled() && (LogPSOtoFileCache() || ReportNewPSOs()))
@@ -3412,23 +3471,10 @@ void FPipelineFileCacheManager::CacheGraphicsPSO(uint32 RunTimeHash, FGraphicsPi
 
 					if (bActuallyNewPSO)
 					{
-						if (!bWasPSOPrecached)
-						{
-							CSV_EVENT(PSO, TEXT("Encountered new graphics PSO"));
-							UE_LOG(LogRHI, Display, TEXT("Encountered a new graphics PSO: %u"), PSOHash);
-							int32 LogDetailLevel = LogPSODetails() ? 2 : GPSOFileCachePrintNewPSODescriptors;
-							if (LogDetailLevel > 0)
-							{
-								UE_LOG(LogRHI, Display, TEXT("New Graphics PSO (%u)"), PSOHash);
-								if (LogDetailLevel > 1)
-								{
-									UE_LOG(LogRHI, Display, TEXT("%s"), *NewEntry.ToStringReadable());
-								}
-							}
-						}
-						else
+						LogNewGraphicsPSOToConsoleAndCSV(NewEntry, PSOHash, bWasPSOPrecached);
+
+						if (bWasPSOPrecached)
 						{
-							UE_LOG(LogRHI, Log, TEXT("Encountered a new graphics PSO for the file cache but it was already precached at runtime: %u"), PSOHash);
 							bActuallyNewPSO = !GPSOExcludePrecachePSOsInFileCache;
 						}
 					}
@@ -3499,12 +3545,7 @@ void FPipelineFileCacheManager::CacheComputePSO(uint32 RunTimeHash, FRHIComputeS
 						bool bActuallyNewPSO = !NewPSOHashes.Contains(PSOHash);
 						if (bActuallyNewPSO)
 						{
-							CSV_EVENT(PSO, TEXT("Encountered new compute PSO"));
-							UE_LOG(LogRHI, Display, TEXT("Encountered a new compute PSO: %u"), PSOHash);
-							if (GPSOFileCachePrintNewPSODescriptors > 0)
-							{
-								UE_LOG(LogRHI, Display, TEXT("New compute PSO (%u) Description: %s"), PSOHash, *NewEntry.ComputeDesc.ComputeShader.ToString());
-							}
+							LogNewComputePSOToConsoleAndCSV(NewEntry, PSOHash);
 							
 							if (LogPSOtoFileCache())
 							{
@@ -3560,8 +3601,7 @@ void FPipelineFileCacheManager::CacheRayTracingPSO(const FRayTracingPipelineStat
 		Initializer.GetCallableTable()
 	};
 
-	// When non-blocking creation is used, encountering a non-cached RTPSO is not likely to cause a hitch and so the logging is not useful/actionable.
-	const bool bShouldReportNewRTPSOToLog = !EnumHasAnyFlags(Flags, ERayTracingPipelineCacheFlags::NonBlocking);
+	const bool bIsNonBlocking = !EnumHasAnyFlags(Flags, ERayTracingPipelineCacheFlags::NonBlocking);
 
 	FRWScopeLock Lock(FileCacheLock, SLT_ReadOnly);
 
@@ -3588,15 +3628,7 @@ void FPipelineFileCacheManager::CacheRayTracingPSO(const FRayTracingPipelineStat
 
 					if (!FPipelineFileCacheManager::IsPSOEntryCached(NewEntry, &CurrentUsageData))
 					{
-						//CSV_EVENT(PSO, TEXT("Encountered new ray tracing PSO"));
-						if (bShouldReportNewRTPSOToLog)
-						{
-							UE_LOG(LogRHI, Display, TEXT("Encountered a new ray tracing PSO: %u"), PSOHash);
-							if (GPSOFileCachePrintNewPSODescriptors > 0)
-							{
-								UE_LOG(LogRHI, Display, TEXT("New ray tracing PSO (%u) Description: %s"), PSOHash, *NewEntry.RayTracingDesc.ToString());
-							}
-						}
+						LogNewRaytracingPSOToConsole(NewEntry, PSOHash, bIsNonBlocking);
 
 						if (LogPSOtoFileCache())
 						{
diff --git a/Engine/Source/Runtime/RHI/Public/PipelineFileCache.h b/Engine/Source/Runtime/RHI/Public/PipelineFileCache.h
--- a/Engine/Source/Runtime/RHI/Public/PipelineFileCache.h
+++ b/Engine/Source/Runtime/RHI/Public/PipelineFileCache.h
@@ -402,6 +402,11 @@ public:
 	static uint64 GetGameUsageMask()	{ return GameUsageMask;}
 	
 	RHI_API static void PreCompileComplete();
+
+	/*
+	 * Enable or disable the logging of new PSOs (PSOs that were needed but not found in the file cache) to console and CSV.
+	 */
+	RHI_API static void SetNewPSOConsoleAndCSVLogging(bool bEnabled) { LogNewPSOsToConsoleAndCSV = bEnabled; }
 private:
 	
 	static void RegisterPSOUsageDataUpdateForNextSave(FPSOUsageData& UsageData);
@@ -410,6 +415,10 @@ private:
 	
 	static bool IsBSSEquivalentPSOEntryCached(FPipelineCacheFileFormatPSO const& NewEntry);
 	static bool IsPSOEntryCached(FPipelineCacheFileFormatPSO const& NewEntry, FPSOUsageData* EntryData = nullptr);
+
+	static void LogNewGraphicsPSOToConsoleAndCSV(FPipelineCacheFileFormatPSO& PSO, uint32 PSOHash, bool bWasPSOPrecached);
+	static void LogNewComputePSOToConsoleAndCSV(FPipelineCacheFileFormatPSO& PSO, uint32 PSOHash);
+	static void LogNewRaytracingPSOToConsole(FPipelineCacheFileFormatPSO& PSO, uint32 PSOHash, bool bIsNonBlockingPSO);
 private:
 	static FRWLock FileCacheLock;
 
@@ -438,5 +447,7 @@ private:
 	static bool FileCacheEnabled;
 	static FPipelineStateLoggedEvent PSOLoggedEvent;
 	RHI_API static uint64 GameUsageMask;
-	static FPSOMaskComparisonFn MaskComparisonFn;
+	static FPSOMaskComparisonFn MaskComparisonFn;	
+
+	RHI_API static bool LogNewPSOsToConsoleAndCSV; // Whether to log new PSOs to the log file and CSV.
 };
diff --git a/Engine/Source/Runtime/RenderCore/Private/ShaderPipelineCache.cpp b/Engine/Source/Runtime/RenderCore/Private/ShaderPipelineCache.cpp
--- a/Engine/Source/Runtime/RenderCore/Private/ShaderPipelineCache.cpp
+++ b/Engine/Source/Runtime/RenderCore/Private/ShaderPipelineCache.cpp
@@ -163,6 +163,13 @@ static TAutoConsoleVariable<int32> CVarPSOGlobalShadersOnlyWhenPSOPrecaching(
                                                                 ECVF_Default | ECVF_RenderThreadSafe
                                                                 );
 
+static TAutoConsoleVariable<int32> CVarPSOFileCacheOnlyOpenUserCache(
+																TEXT("r.ShaderPipelineCache.OnlyOpenUserCache"),
+																0,
+																TEXT("If > 0, only the user cache file will be opened, and the static file cache will not be opened. Defaults to 0."),
+																ECVF_ReadOnly
+                                                                );
+
 static bool GetShaderPipelineCacheSaveBoundPSOLog()
 {
 	static bool bOnce = false;
@@ -823,6 +830,12 @@ uint32 FShaderPipelineCache::NumPrecompilesRemaining()
 
 bool FShaderPipelineCache::OpenPipelineFileCache(EShaderPlatform Platform)
 {
+	if (CVarPSOFileCacheOnlyOpenUserCache.GetValueOnAnyThread())
+	{
+		UE_LOG(LogRHI, Log, TEXT("Ignoring request to open static pipeline file cache because r.ShaderPipelineCache.OnlyOpenUserCache is set."));
+		return false;
+	}
+
 	bool bFileOpen = false;
 	if (ShaderPipelineCache)
 	{
@@ -930,20 +943,29 @@ void FShaderPipelineCache::ShaderLibraryStateChanged(ELibraryState State, EShade
 		const FString& PSOCacheKey = PSOCacheName;
 
 		FScopeLock Lock(&ShaderPipelineCache->Mutex);
-		if (TUniquePtr<class FShaderPipelineCacheTask>* Found = ShaderPipelineCache->ShaderCacheTasks.Find(PSOCacheKey))
+		if (!CVarPSOFileCacheOnlyOpenUserCache.GetValueOnAnyThread())
 		{
-			(*Found)->OnShaderLibraryStateChanged(State, Platform, PSOCacheName, ComponentID);
+			if (TUniquePtr<class FShaderPipelineCacheTask>* Found = ShaderPipelineCache->ShaderCacheTasks.Find(PSOCacheKey))
+			{
+				(*Found)->OnShaderLibraryStateChanged(State, Platform, PSOCacheName, ComponentID);
+			}
+			else
+			{
+				bool bSuccess = ShaderPipelineCache->OpenPipelineFileCacheInternal(false, PSOCacheKey, PSOCacheName, Platform);
+
+				const FString& PSOCacheBaseKey = PSOCacheBaseName;
+				// for some reason we're not always able to find the base PSO cache when the shader libs loads.
+				if( !bSuccess && !ShaderPipelineCache->ShaderCacheTasks.Contains(PSOCacheBaseKey))
+				{
+					ShaderPipelineCache->OpenPipelineFileCacheInternal(false, PSOCacheBaseKey, PSOCacheBaseName, Platform);
+				}
+			}
 		}
 		else
 		{
-			bool bSuccess = ShaderPipelineCache->OpenPipelineFileCacheInternal(false, PSOCacheKey, PSOCacheName, Platform);
-
-			const FString& PSOCacheBaseKey = PSOCacheBaseName;
-			// for some reason we're not always able to find the base PSO cache when the shader libs loads.
-			if( !bSuccess && !ShaderPipelineCache->ShaderCacheTasks.Contains(PSOCacheBaseKey))
-			{
-				ShaderPipelineCache->OpenPipelineFileCacheInternal(false, PSOCacheBaseKey, PSOCacheBaseName, Platform);
-			}
+			// If we're only opening the user cache, disable logging new PSOs
+			// as it can get noisy.
+			FPipelineFileCacheManager::SetNewPSOConsoleAndCSVLogging(false);
 		}
 
 		// An opportunity to open the user cache
