Enable PSO precaching by default on d3d12 windows platform with skipping proxy creation if PSOs are still compiling

diff --git a/Engine/Source/Runtime/D3D12RHI/Private/D3D12RHI.cpp b/Engine/Source/Runtime/D3D12RHI/Private/D3D12RHI.cpp
--- a/Engine/Source/Runtime/D3D12RHI/Private/D3D12RHI.cpp
+++ b/Engine/Source/Runtime/D3D12RHI/Private/D3D12RHI.cpp
@@ -248,6 +248,7 @@ FD3D12DynamicRHI::FD3D12DynamicRHI(const TArray<TSharedPtr<FD3D12Adapter>>& Chos
 	GRHISupportsRayTracingAsyncBuildAccelerationStructure = true;
 
 	GRHISupportsPipelineFileCache = PLATFORM_WINDOWS;
+	GRHISupportsPSOPrecaching = PLATFORM_WINDOWS;
 
 	GRHISupportsMapWriteNoOverwrite = true;
 
diff --git a/Engine/Source/Runtime/Engine/Private/PSOPrecache.cpp b/Engine/Source/Runtime/Engine/Private/PSOPrecache.cpp
--- a/Engine/Source/Runtime/Engine/Private/PSOPrecache.cpp
+++ b/Engine/Source/Runtime/Engine/Private/PSOPrecache.cpp
@@ -64,7 +64,7 @@ static FAutoConsoleVariableRef CVarPSOPrecacheResources(
 	ECVF_ReadOnly
 );
 
-int32 GPSOProxyCreationWhenPSOReady = 0;
+int32 GPSOProxyCreationWhenPSOReady = 1;
 static FAutoConsoleVariableRef CVarPSOProxyCreationWhenPSOReady(
 	TEXT("r.PSOPrecache.ProxyCreationWhenPSOReady"),
 	GPSOProxyCreationWhenPSOReady,
diff --git a/Engine/Source/Runtime/RHI/Private/PipelineStateCache.cpp b/Engine/Source/Runtime/RHI/Private/PipelineStateCache.cpp
--- a/Engine/Source/Runtime/RHI/Private/PipelineStateCache.cpp
+++ b/Engine/Source/Runtime/RHI/Private/PipelineStateCache.cpp
@@ -121,7 +121,7 @@ static TAutoConsoleVariable<int32> CVarRTPSOCacheSize(
 );
 #endif // RHI_RAYTRACING
 
-int32 GPSOPrecaching = 0;
+int32 GPSOPrecaching = 1;
 static FAutoConsoleVariableRef CVarPSOPrecaching(
 	TEXT("r.PSOPrecaching"),
 	GPSOPrecaching,
@@ -2685,7 +2685,7 @@ bool PipelineStateCache::IsPSOPrecachingEnabled()
 	// Disables in the editor for now by default untill more testing is done - still WIP
 	return false;
 #else
-	return GPSOPrecaching != 0 && GRHISupportsPipelineFileCache;
+	return GPSOPrecaching != 0 && GRHISupportsPSOPrecaching;
 #endif // WITH_EDITOR
 }
 
diff --git a/Engine/Source/Runtime/RHI/Public/RHIGlobals.h b/Engine/Source/Runtime/RHI/Public/RHIGlobals.h
--- a/Engine/Source/Runtime/RHI/Public/RHIGlobals.h
+++ b/Engine/Source/Runtime/RHI/Public/RHIGlobals.h
@@ -551,6 +551,9 @@ struct FRHIGlobals
 	/** True if the pipeline file cache can be used with this RHI */
 	bool SupportsPipelineFileCache = false;
 
+	/** True if the RHI support PSO precaching */
+	bool SupportsPSOPrecaching = false;
+
 	/** True if the RHI supports setting the stencil ref at pixel granularity from the pixel shader */
 	bool SupportsStencilRefFromPixelShader = false;
 
@@ -741,6 +744,7 @@ extern RHI_API FRHIGlobals GRHIGlobals;
 #define GRHIPresentCounter                                     GRHIGlobals.PresentCounter
 #define GRHISupportsArrayIndexFromAnyShader                    GRHIGlobals.SupportsArrayIndexFromAnyShader
 #define GRHISupportsPipelineFileCache                          GRHIGlobals.SupportsPipelineFileCache
+#define GRHISupportsPSOPrecaching                              GRHIGlobals.SupportsPSOPrecaching
 #define GRHISupportsStencilRefFromPixelShader                  GRHIGlobals.SupportsStencilRefFromPixelShader
 #define GRHISupportsConservativeRasterization                  GRHIGlobals.SupportsConservativeRasterization
 #define GRHISupportsDispatchShaderBundle                       GRHIGlobals.SupportsDispatchShaderBundle
