VulkanRHI: Add feature level to "Graphics RHI" label in editor to have an easy way to verify it.

#rnx

diff --git a/Engine/Source/Runtime/RHI/Private/Linux/LinuxDynamicRHI.cpp b/Engine/Source/Runtime/RHI/Private/Linux/LinuxDynamicRHI.cpp
--- a/Engine/Source/Runtime/RHI/Private/Linux/LinuxDynamicRHI.cpp
+++ b/Engine/Source/Runtime/RHI/Private/Linux/LinuxDynamicRHI.cpp
@@ -94,12 +94,16 @@ FDynamicRHI* PlatformCreateDynamicRHI()
 				}
 				else
 				{
-					FApp::SetGraphicsRHI(TEXT("Vulkan"));
 					FPlatformApplicationMisc::UsingVulkan();
 
 					FName ShaderFormatName(*TargetedShaderFormats[SfIdx]);
 					EShaderPlatform TargetedPlatform = ShaderFormatToLegacyShaderPlatform(ShaderFormatName);
 					RequestedFeatureLevel = GetMaxSupportedFeatureLevel(TargetedPlatform);
+
+					FString FeatureLevelName;
+					GetFeatureLevelName(RequestedFeatureLevel, FeatureLevelName);
+					FApp::SetGraphicsRHI(FString::Printf(TEXT("Vulkan (%s)"), *FeatureLevelName));
+
 					break;
 				}
 			}
diff --git a/Engine/Source/Runtime/RHI/Private/Windows/WindowsDynamicRHI.cpp b/Engine/Source/Runtime/RHI/Private/Windows/WindowsDynamicRHI.cpp
--- a/Engine/Source/Runtime/RHI/Private/Windows/WindowsDynamicRHI.cpp
+++ b/Engine/Source/Runtime/RHI/Private/Windows/WindowsDynamicRHI.cpp
@@ -124,7 +124,12 @@ static FString GetRHINameFromWindowsRHI(EWindowsRHI InWindowsRHI, ERHIFeatureLev
 		GetFeatureLevelName(InFeatureLevel, FeatureLevelName);
 		return FString::Printf(TEXT("DirectX 12 (%s)"), *FeatureLevelName);
 	}
-	case EWindowsRHI::Vulkan: return TEXT("Vulkan");
+	case EWindowsRHI::Vulkan:
+	{
+		FString FeatureLevelName;
+		GetFeatureLevelName(InFeatureLevel, FeatureLevelName);
+		return FString::Printf(TEXT("Vulkan (%s)"), *FeatureLevelName);
+	}
 	case EWindowsRHI::OpenGL: return TEXT("OpenGL");
 	}
 }
