Fixing various issues with shader previews
- PC previews of SM6 platforms weren't going to DXC correctly.
- Preview platform include paths weren't initialized when compiling with local-only shader compiles.
- A few preview platform DDSPI settings weren't being initialized correctly based on the capabilites of the preview format.

diff --git a/Engine/Source/Developer/Windows/ShaderFormatD3D/Private/D3DShaderCompiler.cpp b/Engine/Source/Developer/Windows/ShaderFormatD3D/Private/D3DShaderCompiler.cpp
--- a/Engine/Source/Developer/Windows/ShaderFormatD3D/Private/D3DShaderCompiler.cpp
+++ b/Engine/Source/Developer/Windows/ShaderFormatD3D/Private/D3DShaderCompiler.cpp
@@ -1104,6 +1104,15 @@ bool CompileAndProcessD3DShaderFXC(const FShaderPreprocessOutput& PreprocessOutp
 	return bSuccess;
 }
 
+inline bool ShouldUseDXC(const FShaderCompilerInput& Input, ELanguage Language)
+{
+	return Language == ELanguage::SM6
+		|| Input.IsRayTracingShader()
+		|| Input.Environment.CompilerFlags.Contains(CFLAG_WaveOperations)
+		|| Input.Environment.CompilerFlags.Contains(CFLAG_ForceDXC)
+		|| Input.Environment.CompilerFlags.Contains(CFLAG_InlineRayTracing);
+}
+
 bool PreprocessD3DShader(
 	const FShaderCompilerInput& Input,
 	const FShaderCompilerEnvironment& Environment,
@@ -1126,7 +1135,7 @@ bool PreprocessD3DShader(
 	}
 
 	FString& PreprocessedSource = Output.EditSource();
-	const bool bUseDXC = !Input.CanCompileWithLegacyFxc();
+	const bool bUseDXC = ShouldUseDXC(Input, Language);
 
 
 	// Set additional defines.
@@ -1207,7 +1216,7 @@ void CompileD3DShader(const FShaderCompilerInput& Input, const FShaderPreprocess
 {
 	TRACE_CPUPROFILER_EVENT_SCOPE(CompileD3DShader);
 
-	const bool bUseDXC = !Input.CanCompileWithLegacyFxc(); 
+	const bool bUseDXC = ShouldUseDXC(Input, Language);
 	const TCHAR* ShaderProfile = GetShaderProfileName(Language, Input.Target.Frequency, bUseDXC);
 
 	if(!ShaderProfile)
diff --git a/Engine/Source/Runtime/RHI/Private/DataDrivenShaderPlatformInfo.cpp b/Engine/Source/Runtime/RHI/Private/DataDrivenShaderPlatformInfo.cpp
--- a/Engine/Source/Runtime/RHI/Private/DataDrivenShaderPlatformInfo.cpp
+++ b/Engine/Source/Runtime/RHI/Private/DataDrivenShaderPlatformInfo.cpp
@@ -410,10 +410,13 @@ void FGenericDataDrivenShaderPlatformInfo::UpdatePreviewPlatforms()
 				Infos[ShaderPlatform].bSupportsSceneDataCompressedTransforms = Infos[EditorSPForPreviewMaxFeatureLevel].bSupportsSceneDataCompressedTransforms;
 				Infos[ShaderPlatform].bSupportsNanite &= Infos[EditorSPForPreviewMaxFeatureLevel].bSupportsNanite;
 				Infos[ShaderPlatform].bSupportsLumenGI &= Infos[EditorSPForPreviewMaxFeatureLevel].bSupportsLumenGI;
+				Infos[ShaderPlatform].bSupportsPrimitiveShaders &= Infos[EditorSPForPreviewMaxFeatureLevel].bSupportsPrimitiveShaders;
 				Infos[ShaderPlatform].bSupportsUInt64ImageAtomics &= Infos[EditorSPForPreviewMaxFeatureLevel].bSupportsUInt64ImageAtomics;
 				Infos[ShaderPlatform].bSupportsGen5TemporalAA &= Infos[EditorSPForPreviewMaxFeatureLevel].bSupportsGen5TemporalAA;
 				Infos[ShaderPlatform].bSupportsInlineRayTracing &= Infos[EditorSPForPreviewMaxFeatureLevel].bSupportsInlineRayTracing;
 				Infos[ShaderPlatform].bSupportsRayTracingShaders &= Infos[EditorSPForPreviewMaxFeatureLevel].bSupportsRayTracingShaders;
+				Infos[ShaderPlatform].bSupportsMeshShadersTier0 &= Infos[EditorSPForPreviewMaxFeatureLevel].bSupportsMeshShadersTier0;
+				Infos[ShaderPlatform].bSupportsMeshShadersTier1 &= Infos[EditorSPForPreviewMaxFeatureLevel].bSupportsMeshShadersTier1;
 
 				// Support for stereo features requires extra consideration. The editor may not use the same technique as the preview platform,
 				// particularly MobileMultiView may be substituted by a fallback path. In order to avoid inundating real mobile platforms
diff --git a/Engine/Source/Runtime/RenderCore/Private/ShaderCore.cpp b/Engine/Source/Runtime/RenderCore/Private/ShaderCore.cpp
--- a/Engine/Source/Runtime/RenderCore/Private/ShaderCore.cpp
+++ b/Engine/Source/Runtime/RenderCore/Private/ShaderCore.cpp
@@ -219,6 +219,18 @@ public:
 			}
 		}
 
+#if WITH_EDITOR
+		for (int i = 0; i < EShaderPlatform::SP_NumPlatforms; ++i)
+		{
+			EShaderPlatform ShaderPlatform = EShaderPlatform(i);
+			if (FDataDrivenShaderPlatformInfo::IsValid(ShaderPlatform) && FDataDrivenShaderPlatformInfo::GetIsPreviewPlatform(ShaderPlatform))
+			{
+				const EShaderPlatform CompilingPlatform = ShaderFormatNameToShaderPlatform(FDataDrivenShaderPlatformInfo::GetShaderFormat(ShaderPlatform));
+				UpdateIncludeDirectoryForPreviewPlatform(ShaderPlatform, CompilingPlatform);
+			}
+		}
+#endif
+
 		bInitialized = true;
 	}
 
diff --git a/Engine/Source/Runtime/RenderCore/Public/ShaderCompilerCore.h b/Engine/Source/Runtime/RenderCore/Public/ShaderCompilerCore.h
--- a/Engine/Source/Runtime/RenderCore/Public/ShaderCompilerCore.h
+++ b/Engine/Source/Runtime/RenderCore/Public/ShaderCompilerCore.h
@@ -456,6 +456,7 @@ struct FShaderCompilerInput
 	}
 
 	/** Returns whether this shader input *can* be compiled with the legacy FXC compiler. */
+	UE_DEPRECATED(5.3, "CanCompileWithLegacyFxc doesn't have enough information to correctly flag the input as requiring FXC. Please use internal shader format code instead.")
 	bool CanCompileWithLegacyFxc() const
 	{
 		return !(Target.GetPlatform() == SP_PCD3D_SM6
