Metal RHI: Disable SM6 as a default
Fix target settings to correctly pick highest target

diff --git a/Engine/Config/BaseEngine.ini b/Engine/Config/BaseEngine.ini
--- a/Engine/Config/BaseEngine.ini
+++ b/Engine/Config/BaseEngine.ini
@@ -3001,7 +3001,6 @@ bEnableRayTracing=true
 
 [/Script/MacTargetPlatform.MacTargetSettings]
 MetalLanguageVersion=7
-+TargetedRHIs=SF_METAL_SM6
 +TargetedRHIs=SF_METAL_SM5
 UseFastIntrinsics=False
 EnableMathOptimisations=True
diff --git a/Engine/Source/Editor/DetailCustomizations/Private/MacTargetSettingsDetails.cpp b/Engine/Source/Editor/DetailCustomizations/Private/MacTargetSettingsDetails.cpp
--- a/Engine/Source/Editor/DetailCustomizations/Private/MacTargetSettingsDetails.cpp
+++ b/Engine/Source/Editor/DetailCustomizations/Private/MacTargetSettingsDetails.cpp
@@ -138,7 +138,7 @@ static FText GetFriendlyNameFromRHINameMac(FName InRHIName)
 		FriendlyRHIName = LOCTEXT("MetalSM5", "Mac Metal Desktop Renderer (SM5, Metal 2.4+, macOS Monterey 12.0 or later)");
 		break;
     case SP_METAL_SM6:
-        FriendlyRHIName = LOCTEXT("MetalSM6", "Mac Metal Desktop Renderer Experimental (SM6, Metal 2.4+, macOS 14.0 or later, M2+)");
+        FriendlyRHIName = LOCTEXT("MetalSM6", "Mac Metal Desktop Renderer Beta (SM6, Metal 2.4+, macOS 14.0 or later, M2+)");
         break;
 	case SP_METAL_MACES3_1:
 		FriendlyRHIName = LOCTEXT("MetalMobile", "Mac Metal High-End Mobile Preview (Mobile Preview)");
diff --git a/Engine/Source/Runtime/RHI/Private/Apple/AppleDynamicRHI.cpp b/Engine/Source/Runtime/RHI/Private/Apple/AppleDynamicRHI.cpp
--- a/Engine/Source/Runtime/RHI/Private/Apple/AppleDynamicRHI.cpp
+++ b/Engine/Source/Runtime/RHI/Private/Apple/AppleDynamicRHI.cpp
@@ -76,10 +76,18 @@ static void ComputeRequestedFeatureLevel(const FAppleDynamicRHIOptions& Options,
 
 		if (TargetedShaderFormats.Num() > 0)
 		{
-			// Pick the first one
-			FName ShaderFormatName(*TargetedShaderFormats[0]);
-			EShaderPlatform TargetedPlatform = ShaderFormatToLegacyShaderPlatform(ShaderFormatName);
-			RequestedFeatureLevel = GetMaxSupportedFeatureLevel(TargetedPlatform);
+            // Pick the highest targeted shader format
+            for(const FString& Format : TargetedShaderFormats)
+            {
+                FName ShaderFormatName(*Format);
+                EShaderPlatform TargetedPlatform = ShaderFormatToLegacyShaderPlatform(ShaderFormatName);
+                ERHIFeatureLevel::Type FeatureLevel = GetMaxSupportedFeatureLevel(TargetedPlatform);
+                
+                if(FeatureLevel > RequestedFeatureLevel || RequestedFeatureLevel == ERHIFeatureLevel::Num)
+                {
+                    RequestedFeatureLevel = FeatureLevel;
+                }
+            }
 		}
 	}
 	else
