Protect dispatch shader bundle RHI calls from invalid record indices

diff --git a/Engine/Source/Runtime/RHI/Public/RHICommandListCommandExecutes.inl b/Engine/Source/Runtime/RHI/Public/RHICommandListCommandExecutes.inl
--- a/Engine/Source/Runtime/RHI/Public/RHICommandListCommandExecutes.inl
+++ b/Engine/Source/Runtime/RHI/Public/RHICommandListCommandExecutes.inl
@@ -246,9 +246,13 @@ void FRHICommandDispatchShaderBundle::Execute(FRHICommandListBase& CmdList)
 {
 	RHISTAT(DispatchShaderBundle);
 	extern RHI_API FRHIComputePipelineState* ExecuteSetComputePipelineState(FComputePipelineState* ComputePipelineState);
-	for (FRHIShaderBundleDispatch& Dispatch : Dispatches)
+	for (int32 DispatchIndex = 0; DispatchIndex < Dispatches.Num(); ++DispatchIndex)
 	{
-		Dispatch.RHIPipeline = ExecuteSetComputePipelineState(Dispatch.PipelineState);
+		FRHIShaderBundleDispatch& Dispatch = Dispatches[DispatchIndex];
+		if (Dispatch.RecordIndex != ~uint32(0u))
+		{
+			Dispatch.RHIPipeline = ExecuteSetComputePipelineState(Dispatch.PipelineState);
+		}
 	}
 	INTERNAL_DECORATOR_COMPUTE(RHIDispatchShaderBundle)(ShaderBundle, RecordArgBufferSRV, RecordDataBufferSRV, ExecutionBufferUAV, Dispatches);
 }
diff --git a/Engine/Source/Runtime/RHI/Public/RHIShaderParameters.h b/Engine/Source/Runtime/RHI/Public/RHIShaderParameters.h
--- a/Engine/Source/Runtime/RHI/Public/RHIShaderParameters.h
+++ b/Engine/Source/Runtime/RHI/Public/RHIShaderParameters.h
@@ -189,4 +189,9 @@ struct FRHIShaderBundleDispatch
 	FRHIComputeShader* Shader = nullptr;
 	FRHIComputePipelineState* RHIPipeline = nullptr;
 	FRHIBatchedShaderParameters Parameters;
+
+	inline bool IsValid() const
+	{
+		return RecordIndex != ~uint32(0u);
+	}
 };
\ No newline at end of file
