RHI command list resource lifetime reference counting.
 - Allows reference tracking without scopes.
 - Fixed dependency issue with initialization of immediate command list persistent state.

diff --git a/Engine/Source/Runtime/RHI/Private/RHICommandList.cpp b/Engine/Source/Runtime/RHI/Private/RHICommandList.cpp
--- a/Engine/Source/Runtime/RHI/Private/RHICommandList.cpp
+++ b/Engine/Source/Runtime/RHI/Private/RHICommandList.cpp
@@ -136,8 +136,25 @@ static TStatId GCurrentExecuteStat;
 static FCriticalSection GRHIThreadOnTasksCritical;
 static std::atomic<int32> GRHIThreadStallRequestCount;
 
-FRHICommandListBase::FRHICommandListBase(FRHIGPUMask InGPUMask, ERecordingThread InRecordingThread)
-	: FRHICommandListBase(FPersistentState(InGPUMask, InRecordingThread))
+static std::atomic<int32> GRHIResourceLifetimeRefCount{0};
+
+void RHIResourceLifetimeAddRef(int32 NumRefs)
+{
+	GRHIResourceLifetimeRefCount.fetch_add(NumRefs, std::memory_order_relaxed);
+}
+
+void RHIResourceLifetimeReleaseRef(FRHICommandListImmediate& RHICmdList, int32 NumRefs)
+{
+	int32 RefCount = GRHIResourceLifetimeRefCount.fetch_sub(NumRefs, std::memory_order_release) - 1;
+	check(RefCount >= 0);
+	if (!RefCount)
+	{
+		RHICmdList.FlushExtendedLifetimeResourceDeletes();
+	}
+}
+
+FRHICommandListBase::FRHICommandListBase(FRHIGPUMask InGPUMask, ERecordingThread InRecordingThread, bool bImmediate)
+	: FRHICommandListBase(FPersistentState(InGPUMask, InRecordingThread, bImmediate))
 {}
 
 FRHICommandListBase::FRHICommandListBase(FPersistentState&& InPersistentState)
@@ -1016,7 +1033,7 @@ int32 FRHICommandListImmediate::FlushPendingDeletes()
 
 	TArray<FRHIResource*, FConcurrentLinearArrayAllocator> DeletedResources;
 
-	TArray<FRHIResource*, FConcurrentLinearArrayAllocator>& DeletedResourcesWithLifetimeExtension = PersistentState.ExtendResourceLifetimeRefCount > 0
+	TArray<FRHIResource*, FConcurrentLinearArrayAllocator>& DeletedResourcesWithLifetimeExtension = GRHIResourceLifetimeRefCount > 0
 		? PersistentState.ExtendedLifetimeResources 
 		: DeletedResources;
 
diff --git a/Engine/Source/Runtime/RHI/Public/RHICommandList.h b/Engine/Source/Runtime/RHI/Public/RHICommandList.h
--- a/Engine/Source/Runtime/RHI/Public/RHICommandList.h
+++ b/Engine/Source/Runtime/RHI/Public/RHICommandList.h
@@ -502,7 +502,7 @@ public:
 	};
 
 protected:
-	RHI_API FRHICommandListBase(FRHIGPUMask InGPUMask, ERecordingThread InRecordingThread);
+	RHI_API FRHICommandListBase(FRHIGPUMask InGPUMask, ERecordingThread InRecordingThread, bool bInImmediate);
 
 public:
 	RHI_API FRHICommandListBase(FRHICommandListBase&& Other);
@@ -1152,8 +1152,9 @@ protected:
 			}
 		} Stats;
 
-		FPersistentState(FRHIGPUMask InInitialGPUMask, ERecordingThread InRecordingThread)
-			: RecordingThread(InRecordingThread)
+		FPersistentState(FRHIGPUMask InInitialGPUMask, ERecordingThread InRecordingThread, bool bInImmediate = false)
+			: bImmediate(bInImmediate)
+			, RecordingThread(InRecordingThread)
 			, CurrentGPUMask(InInitialGPUMask)
 			, InitialGPUMask(InInitialGPUMask)
 		{}
@@ -2421,6 +2422,10 @@ protected:
 		PersistentState.BoundComputeShaderRHI = InBoundComputeShaderRHI;
 	}
 
+	FRHIComputeCommandList(FRHIGPUMask GPUMask, ERecordingThread InRecordingThread, bool bImmediate)
+		: FRHICommandListBase(GPUMask, InRecordingThread, bImmediate)
+	{}
+
 public:
 	UE_DEPRECATED(5.3, "FlushAllPendingComputeParameters isn't needed with automatic batching removed.")
 	void FlushAllPendingComputeParameters()
@@ -2434,7 +2439,7 @@ public:
 	}
 
 	FRHIComputeCommandList(FRHIGPUMask GPUMask = FRHIGPUMask::All(), ERecordingThread InRecordingThread = ERecordingThread::Render)
-		: FRHICommandListBase(GPUMask, InRecordingThread)
+		: FRHICommandListBase(GPUMask, InRecordingThread, false)
 	{}
 
 	FRHIComputeCommandList(FRHICommandListBase&& Other)
@@ -3165,6 +3170,10 @@ protected:
 		PersistentState.BoundShaderInput = InBoundShaderStateInput;
 	}
 
+	FRHICommandList(FRHIGPUMask GPUMask, ERecordingThread InRecordingThread, bool bImmediate)
+		: FRHIComputeCommandList(GPUMask, InRecordingThread, bImmediate)
+	{}
+
 public:
 	static inline FRHICommandList& Get(FRHICommandListBase& RHICmdList)
 	{
@@ -4100,14 +4109,15 @@ class FRHICommandListImmediate : public FRHICommandList
 	friend class FRHICommandListScopedExtendResourceLifetime;
 	friend struct FRHICommandBeginFrame;
 
+	friend void RHI_API RHIResourceLifetimeReleaseRef(FRHICommandListImmediate&, int32);
+
 	RHI_API static FGraphEventArray WaitOutstandingTasks;
 	RHI_API static FGraphEventRef   RHIThreadTask;
 	RHI_API static FRHIDrawStats    FrameDrawStats;
 
 	FRHICommandListImmediate()
-		: FRHICommandList(FRHIGPUMask::All(), ERecordingThread::Render)
+		: FRHICommandList(FRHIGPUMask::All(), ERecordingThread::Render, true)
 	{
-		PersistentState.bImmediate = true;
 		PersistentState.Stats.Ptr = &FrameDrawStats;
 	}
 
@@ -4583,21 +4593,24 @@ public:
 	}
 };
 
+/** Takes a reference to defer deletion of RHI resources. */
+void RHI_API RHIResourceLifetimeAddRef(int32 NumRefs = 1);
+
+/** Releases a reference to defer deletion of RHI resources. If the reference count hits zero, resources are queued for deletion. */
+void RHI_API RHIResourceLifetimeReleaseRef(FRHICommandListImmediate& RHICmdList, int32 NumRefs = 1);
+
 class FRHICommandListScopedExtendResourceLifetime
 {
 public:
 	FRHICommandListScopedExtendResourceLifetime(FRHICommandListImmediate& InRHICmdList)
 		: RHICmdList(InRHICmdList)
 	{
-		RHICmdList.PersistentState.ExtendResourceLifetimeRefCount++;
+		RHIResourceLifetimeAddRef();
 	}
 
 	~FRHICommandListScopedExtendResourceLifetime()
 	{
-		if (--RHICmdList.PersistentState.ExtendResourceLifetimeRefCount == 0)
-		{
-			RHICmdList.FlushExtendedLifetimeResourceDeletes();
-		}
+		RHIResourceLifetimeReleaseRef(RHICmdList);
 	}
 
 private:
