Add ShadingRateSource to ReadOnlyExclusiveMask.

Fragment Density Map attachments were not having the ShadingRateSource access set when registered as FRDGViewableResource with SkipTracking flag.

diff --git a/Engine/Source/Runtime/RHI/Public/RHIAccess.h b/Engine/Source/Runtime/RHI/Public/RHIAccess.h
--- a/Engine/Source/Runtime/RHI/Public/RHIAccess.h
+++ b/Engine/Source/Runtime/RHI/Public/RHIAccess.h
@@ -56,7 +56,7 @@ enum class ERHIAccess : uint32
 	UAVMask = UAVCompute | UAVGraphics,
 
 	// A mask of all bits representing read-only states which cannot be combined with other write states.
-	ReadOnlyExclusiveMask = CPURead | Present | IndirectArgs | VertexOrIndexBuffer | SRVGraphics | SRVCompute | CopySrc | ResolveSrc | BVHRead,
+	ReadOnlyExclusiveMask = CPURead | Present | IndirectArgs | VertexOrIndexBuffer | SRVGraphics | SRVCompute | CopySrc | ResolveSrc | BVHRead | ShadingRateSource,
 
 	// A mask of all bits representing read-only states on the compute pipe which cannot be combined with other write states.
 	ReadOnlyExclusiveComputeMask = CPURead | IndirectArgs | SRVCompute | CopySrc | BVHRead,
