Move {Create, Update}TextureReference to FRHICommandListBase.

diff --git a/Engine/Source/Runtime/RHI/Private/DynamicRHI.cpp b/Engine/Source/Runtime/RHI/Private/DynamicRHI.cpp
--- a/Engine/Source/Runtime/RHI/Private/DynamicRHI.cpp
+++ b/Engine/Source/Runtime/RHI/Private/DynamicRHI.cpp
@@ -579,7 +579,7 @@ uint64 FDynamicRHI::RHIGetMinimumAlignmentForBufferBackedSRV(EPixelFormat Format
 	return GPixelFormats[Format].BlockBytes;
 }
 
-FTextureReferenceRHIRef FDynamicRHI::RHICreateTextureReference(FRHITexture* InReferencedTexture)
+FTextureReferenceRHIRef FDynamicRHI::RHICreateTextureReference(FRHICommandListBase& RHICmdList, FRHITexture* InReferencedTexture)
 {
 	FRHITexture* ReferencedTexture = InReferencedTexture ? InReferencedTexture : FRHITextureReference::GetDefaultTexture();
 
@@ -589,7 +589,7 @@ FTextureReferenceRHIRef FDynamicRHI::RHICreateTextureReference(FRHITexture* InRe
 	// If the referenced texture is configured for bindless, make sure we also create an SRV to use for bindless.
 	if (ReferencedTexture && ReferencedTexture->GetDefaultBindlessHandle().IsValid())
 	{
-		ShaderResourceView = FRHICommandListImmediate::Get().CreateShaderResourceView(ReferencedTexture, 0u);
+		ShaderResourceView = RHICmdList.CreateShaderResourceView(ReferencedTexture, 0u);
 	}
 #endif
 
diff --git a/Engine/Source/Runtime/RHI/Private/RHICommandList.cpp b/Engine/Source/Runtime/RHI/Private/RHICommandList.cpp
--- a/Engine/Source/Runtime/RHI/Private/RHICommandList.cpp
+++ b/Engine/Source/Runtime/RHI/Private/RHICommandList.cpp
@@ -1801,7 +1801,7 @@ void FDynamicRHI::RHIReadSurfaceFloatData_RenderThread(class FRHICommandListImme
 	GDynamicRHI->RHIReadSurfaceFloatData(Texture, Rect, OutData, Flags);
 }
 
-void FRHICommandListImmediate::UpdateTextureReference(FRHITextureReference* TextureRef, FRHITexture* NewTexture)
+void FRHICommandListBase::UpdateTextureReference(FRHITextureReference* TextureRef, FRHITexture* NewTexture)
 {
 	if (TextureRef == nullptr)
 	{
@@ -1814,11 +1814,6 @@ void FRHICommandListImmediate::UpdateTextureReference(FRHITextureReference* Text
 	});
 
 	RHIThreadFence(true);
-	if (GetUsedMemory() > 256 * 1024)
-	{
-		// we could be loading a level or something, lets get this stuff going
-		ImmediateFlush(EImmediateFlushType::DispatchToRHIThread); 
-	}
 }
 
 void FRHICommandListImmediate::UpdateRHIResources(FRHIResourceUpdateInfo* UpdateInfos, int32 Num, bool bNeedReleaseRefs)
diff --git a/Engine/Source/Runtime/RHI/Public/DynamicRHI.h b/Engine/Source/Runtime/RHI/Public/DynamicRHI.h
--- a/Engine/Source/Runtime/RHI/Public/DynamicRHI.h
+++ b/Engine/Source/Runtime/RHI/Public/DynamicRHI.h
@@ -416,7 +416,7 @@ public:
 	virtual FTextureRHIRef RHIAsyncCreateTexture2D(uint32 SizeX, uint32 SizeY, uint8 Format, uint32 NumMips, ETextureCreateFlags Flags, ERHIAccess InResourceState, void** InitialMipData, uint32 NumInitialMips, FGraphEventRef& OutCompletionEvent) = 0;
 
 	/** Create a texture reference. InReferencedTexture can be null. */
-	RHI_API virtual FTextureReferenceRHIRef RHICreateTextureReference(FRHITexture* InReferencedTexture);
+	RHI_API virtual FTextureReferenceRHIRef RHICreateTextureReference(FRHICommandListBase& RHICmdList, FRHITexture* InReferencedTexture);
 
 	// SRV / UAV creation functions
 	virtual FShaderResourceViewRHIRef  RHICreateShaderResourceView (class FRHICommandListBase& RHICmdList, FRHIViewableResource* Resource, FRHIViewDesc const& ViewDesc) = 0;
diff --git a/Engine/Source/Runtime/RHI/Public/RHICommandList.h b/Engine/Source/Runtime/RHI/Public/RHICommandList.h
--- a/Engine/Source/Runtime/RHI/Public/RHICommandList.h
+++ b/Engine/Source/Runtime/RHI/Public/RHICommandList.h
@@ -783,6 +783,13 @@ public:
 		GDynamicRHI->RHIUpdateTexture3D(*this, Texture, MipIndex, UpdateRegion, SourceRowPitch, SourceDepthPitch, SourceData);
 	}
 
+	FORCEINLINE FTextureReferenceRHIRef CreateTextureReference(FRHITexture* InReferencedTexture = nullptr)
+	{
+		return GDynamicRHI->RHICreateTextureReference(*this, InReferencedTexture);
+	}
+
+	RHI_API void UpdateTextureReference(FRHITextureReference* TextureRef, FRHITexture* NewTexture);
+
 	FORCEINLINE FShaderResourceViewRHIRef CreateShaderResourceView(FRHIBuffer* Buffer, FRHIViewDesc::FBufferSRV::FInitializer const& ViewDesc)
 	{
 		LLM_SCOPE_BYNAME(TEXT("RHIMisc/CreateShaderResourceView"));
@@ -4540,8 +4547,6 @@ public:
 		return GDynamicRHI->RHIGetNativeCommandBuffer();
 	}
 
-	RHI_API void UpdateTextureReference(FRHITextureReference* TextureRef, FRHITexture* NewTexture);
-
 	FORCEINLINE void PollRenderQueryResults()
 	{
 		GDynamicRHI->RHIPollRenderQueryResults();
@@ -5061,7 +5066,7 @@ FORCEINLINE void RHIUpdateRHIResources(FRHIResourceUpdateInfo* UpdateInfos, int3
 
 FORCEINLINE FTextureReferenceRHIRef RHICreateTextureReference(FRHITexture* InReferencedTexture = nullptr)
 {
-	return GDynamicRHI->RHICreateTextureReference(InReferencedTexture);
+	return FRHICommandListImmediate::Get().CreateTextureReference(InReferencedTexture);
 }
 
 FORCEINLINE void RHIUpdateTextureReference(FRHITextureReference* TextureRef, FRHITexture* NewTexture)
