Fix Intel ICX compiler transient allocator GPU perf drop

diff --git a/Engine/Source/Runtime/RHI/Public/RHIResources.h b/Engine/Source/Runtime/RHI/Public/RHIResources.h
--- a/Engine/Source/Runtime/RHI/Public/RHIResources.h
+++ b/Engine/Source/Runtime/RHI/Public/RHIResources.h
@@ -240,6 +240,14 @@ struct FClearValueBinding
 		: ColorBinding(NoBinding)
 	{
 		check(ColorBinding == EClearBinding::ENoneBound);
+
+		Value.Color[0] = 0.0f;
+		Value.Color[1] = 0.0f;
+		Value.Color[2] = 0.0f;
+		Value.Color[3] = 0.0f;
+
+		Value.DSValue.Depth = 0.0f;
+		Value.DSValue.Stencil = 0;
 	}
 
 	explicit FClearValueBinding(const FLinearColor& InClearColor)
@@ -1325,6 +1333,25 @@ struct FRHITextureDesc
 		return !(*this == Other);
 	}
 
+	FRHITextureDesc& operator=(const FRHITextureDesc& Other)
+	{
+		Dimension			= Other.Dimension;
+		Flags				= Other.Flags;
+		Format				= Other.Format;
+		UAVFormat			= Other.UAVFormat;
+		Extent				= Other.Extent;
+		Depth				= Other.Depth;
+		ArraySize			= Other.ArraySize;
+		NumMips				= Other.NumMips;
+		NumSamples			= Other.NumSamples;
+		ClearValue			= Other.ClearValue;
+		ExtData				= Other.ExtData;
+		FastVRAMPercentage	= Other.FastVRAMPercentage;
+		GPUMask				= Other.GPUMask;
+
+		return *this;
+	}
+
 	bool IsTexture2D() const
 	{
 		return Dimension == ETextureDimension::Texture2D || Dimension == ETextureDimension::Texture2DArray;
