Added explicit copy constructor to FRHITextureDesc to fix build break.

diff --git a/Engine/Source/Runtime/RHI/Public/RHIResources.h b/Engine/Source/Runtime/RHI/Public/RHIResources.h
--- a/Engine/Source/Runtime/RHI/Public/RHIResources.h
+++ b/Engine/Source/Runtime/RHI/Public/RHIResources.h
@@ -1265,6 +1265,11 @@ struct FRHITextureDesc
 {
 	FRHITextureDesc() = default;
 
+	FRHITextureDesc(const FRHITextureDesc& Other)
+	{
+		*this = Other;
+	}
+
 	FRHITextureDesc(ETextureDimension InDimension)
 		: Dimension(InDimension)
 	{}
