Fixed memory tracing and stats for a reallocation in GPUDefragAllocator.

diff --git a/Engine/Source/Runtime/RHI/Public/GPUDefragAllocator.h b/Engine/Source/Runtime/RHI/Public/GPUDefragAllocator.h
--- a/Engine/Source/Runtime/RHI/Public/GPUDefragAllocator.h
+++ b/Engine/Source/Runtime/RHI/Public/GPUDefragAllocator.h
@@ -794,11 +794,11 @@ protected:
 	void Relocate(FRelocationStats& Stats, FMemoryChunk* Dest, int64 DestOffset, const void* Source, int64 Size, void* UserPayload)
 	{
 		MemoryTrace_ReallocFree((uint64)Source);
-		MemoryTrace_ReallocAlloc((uint64) Dest, Size, 4);
+		MemoryTrace_ReallocAlloc((uint64)Dest->Base, Size, 4);
 		LLM(FLowLevelMemTracker::Get().OnLowLevelAllocMoved(ELLMTracker::Default, Dest->Base, Source));
 
 		uint8* DestAddr = Dest->Base + DestOffset;
-		int64 MemDistance = (int64)(Dest)-(int64)(Source);
+		int64 MemDistance = (int64)(DestAddr) - (int64)(Source);
 		int64 AbsDistance = FMath::Abs(MemDistance);
 		bool bOverlappedMove = AbsDistance < Size;
 
