Fix incorrect RTs binding when rendering hair strands with Substrate and rough refraction option is enabled.

diff --git a/Engine/Source/Runtime/RHI/Public/RHIValidationContext.h b/Engine/Source/Runtime/RHI/Public/RHIValidationContext.h
--- a/Engine/Source/Runtime/RHI/Public/RHIValidationContext.h
+++ b/Engine/Source/Runtime/RHI/Public/RHIValidationContext.h
@@ -1109,7 +1109,7 @@ public:
 
 			// Check all bound textures have the same dimensions
 			FIntVector MipDimensions = RTV.RenderTarget->GetMipDimensions(RTV.MipIndex);
-			checkf(ViewDimensions.IsZero() || ViewDimensions == MipDimensions, TEXT("Render target size mismatch. All render and depth target views must have the same effective dimensions."));
+			checkf(ViewDimensions.IsZero() || ViewDimensions == MipDimensions, TEXT("Render target size mismatch (RT%d: %dx%d vs. Expected: %dx%d). All render and depth target views must have the same effective dimensions."), RTVIndex, MipDimensions.X, MipDimensions.Y, ViewDimensions.X, ViewDimensions.Y);
 			ViewDimensions = MipDimensions;
 
 			uint32 ArraySlice = RTV.ArraySlice;
@@ -1140,7 +1140,7 @@ public:
 		{
 			// Check all bound textures have the same dimensions
 			FIntVector MipDimensions = DSV.DepthStencilTarget->GetMipDimensions(0);
-			checkf(ViewDimensions.IsZero() || ViewDimensions == MipDimensions, TEXT("Depth target size mismatch. All render and depth target views must have the same effective dimensions."));
+			checkf(ViewDimensions.IsZero() || ViewDimensions == MipDimensions, TEXT("Depth target size mismatch (Depth: %dx%d vs. Expected: %dx%d). All render and depth target views must have the same effective dimensions."), MipDimensions.X, MipDimensions.Y, ViewDimensions.X, ViewDimensions.Y);
 			ViewDimensions = MipDimensions;
 
 			if (DSV.ResolveTarget)
diff --git a/Engine/Source/Runtime/Renderer/Private/LightRendering.cpp b/Engine/Source/Runtime/Renderer/Private/LightRendering.cpp
--- a/Engine/Source/Runtime/Renderer/Private/LightRendering.cpp
+++ b/Engine/Source/Runtime/Renderer/Private/LightRendering.cpp
@@ -2635,6 +2635,8 @@ void FDeferredShadingSceneRenderer::RenderLightForHair(
 		PassParameters->PS.ShadowChannelMask[FMath::Clamp(LightSceneInfo->GetDynamicShadowMapChannel(), 0, 3)] = 1.f;
 	}
 	PassParameters->PS.RenderTargets[0] = FRenderTargetBinding(View.HairStrandsViewData.VisibilityData.SampleLightingTexture, ERenderTargetLoadAction::ELoad);
+	PassParameters->PS.RenderTargets[1] = FRenderTargetBinding();
+	PassParameters->PS.RenderTargets[2] = FRenderTargetBinding();
 	PassParameters->PS.RenderTargets.DepthStencil = FDepthStencilBinding(nullptr, ERenderTargetLoadAction::ENoAction, ERenderTargetLoadAction::ENoAction, FExclusiveDepthStencil::DepthNop_StencilNop);
 
 	FDeferredLightPS::FPermutationDomain PermutationVector;
