Fixing moves of FRHICommandListBase::AllowedPipelines.

diff --git a/Engine/Source/Runtime/RHI/Private/RHICommandList.cpp b/Engine/Source/Runtime/RHI/Private/RHICommandList.cpp
--- a/Engine/Source/Runtime/RHI/Private/RHICommandList.cpp
+++ b/Engine/Source/Runtime/RHI/Private/RHICommandList.cpp
@@ -192,7 +192,7 @@ FRHICommandListBase::FRHICommandListBase(FRHICommandListBase&& Other)
     , bExecuting      (MoveTemp(Other.bExecuting))
     , ActivePipeline  (MoveTemp(Other.ActivePipeline))
 #if DO_CHECK
-	, AllowedPipelines(MoveTemp(AllowedPipelines))
+	, AllowedPipelines(MoveTemp(Other.AllowedPipelines))
 #endif			  
 	, DispatchEvent   (MoveTemp(Other.DispatchEvent))
     , ExecuteStat     (MoveTemp(Other.ExecuteStat))
