Added validation to ensure UAV/SRV types and UB sizes match the ones generated by the shader compiler

diff --git a/Engine/Shaders/Public/ShaderVersion.ush b/Engine/Shaders/Public/ShaderVersion.ush
--- a/Engine/Shaders/Public/ShaderVersion.ush
+++ b/Engine/Shaders/Public/ShaderVersion.ush
@@ -2,4 +2,5 @@
 // Each time the console command is executed it generates a new GUID. As this file is included
 // in Platform.ush (which should be included in any shader) it allows to invalidate the shader DDC.
 // 
-// GUID = 91C4C529-4A7B-43A9-A171-BB01113AAD32
\ No newline at end of file
+// GUID = 1990A449-72D5-44ED-85B2-DC789220B8D4
+
diff --git a/Engine/Source/Developer/ShaderCompilerCommon/Private/ShaderCompilerCommon.cpp b/Engine/Source/Developer/ShaderCompilerCommon/Private/ShaderCompilerCommon.cpp
--- a/Engine/Source/Developer/ShaderCompilerCommon/Private/ShaderCompilerCommon.cpp
+++ b/Engine/Source/Developer/ShaderCompilerCommon/Private/ShaderCompilerCommon.cpp
@@ -773,6 +773,36 @@ void UpdateStructuredBufferStride(
 	}
 }
 
+void AddShaderValidationSRVType(uint16 BindPoint,
+							EShaderCodeResourceBindingType TypeDecl,
+							FShaderCompilerOutput& CompilerOutput)
+{
+	if (BindPoint <= UINT16_MAX)
+	{
+		CompilerOutput.ParametersSRVTypeToValidate.Add(FShaderCodeValidationType{ BindPoint, TypeDecl });
+	}
+}
+
+void AddShaderValidationUAVType(uint16 BindPoint,
+							EShaderCodeResourceBindingType TypeDecl,
+							FShaderCompilerOutput& CompilerOutput)
+{
+	if (BindPoint <= UINT16_MAX)
+	{
+		CompilerOutput.ParametersUAVTypeToValidate.Add(FShaderCodeValidationType{ BindPoint, TypeDecl });
+	}
+}
+
+void AddShaderValidationUBSize(uint16 BindPoint,
+							uint32_t Size,
+							FShaderCompilerOutput& CompilerOutput)
+{
+	if (BindPoint <= UINT16_MAX)
+	{
+		CompilerOutput.ParametersUBSizeToValidate.Add(FShaderCodeValidationUBSize{ BindPoint, Size });
+	}
+}
+ 
 void HandleReflectedShaderUAV(
 	const FString& UAVName,
 	int32 BindOffset,
diff --git a/Engine/Source/Developer/ShaderCompilerCommon/Public/ShaderCompilerCommon.h b/Engine/Source/Developer/ShaderCompilerCommon/Public/ShaderCompilerCommon.h
--- a/Engine/Source/Developer/ShaderCompilerCommon/Public/ShaderCompilerCommon.h
+++ b/Engine/Source/Developer/ShaderCompilerCommon/Public/ShaderCompilerCommon.h
@@ -249,6 +249,21 @@ extern SHADERCOMPILERCOMMON_API void UpdateStructuredBufferStride(
 	FShaderCompilerOutput& CompilerOutput
 );
 
+extern SHADERCOMPILERCOMMON_API void AddShaderValidationSRVType(
+	uint16 BindPoint,
+	EShaderCodeResourceBindingType TypeDecl,
+	FShaderCompilerOutput& CompilerOutput);
+
+extern SHADERCOMPILERCOMMON_API void AddShaderValidationUAVType(
+	uint16 BindPoint,
+	EShaderCodeResourceBindingType TypeDecl,
+	FShaderCompilerOutput& CompilerOutput);
+
+extern SHADERCOMPILERCOMMON_API void AddShaderValidationUBSize(
+	uint16 BindPoint,
+	uint32_t Size,
+	FShaderCompilerOutput& CompilerOutput);
+
 inline void HandleReflectedShaderResource(const FString& ResourceName, int32 ReflectionSlot, int32 BindCount, FShaderCompilerOutput& CompilerOutput)
 {
 	HandleReflectedShaderResource(ResourceName, 0, ReflectionSlot, BindCount, CompilerOutput);
diff --git a/Engine/Source/Developer/VulkanShaderFormat/Private/VulkanShaderCompiler.cpp b/Engine/Source/Developer/VulkanShaderFormat/Private/VulkanShaderCompiler.cpp
--- a/Engine/Source/Developer/VulkanShaderFormat/Private/VulkanShaderCompiler.cpp
+++ b/Engine/Source/Developer/VulkanShaderFormat/Private/VulkanShaderCompiler.cpp
@@ -246,8 +246,9 @@ static const FString kHitGroupSystemRootConstantsSymbolName = TEXT("HitGroupSyst
 // A collection of states and data that is locked in at the top level call and doesn't change throughout the compilation process
 struct FVulkanShaderCompilerInternalState
 {
-	FVulkanShaderCompilerInternalState(const FShaderCompilerInput& InInput)
+	FVulkanShaderCompilerInternalState(const FShaderCompilerInput& InInput, const FShaderParameterParser& InParameterParser)
 		: Input(InInput)
+		, ParameterParser(InParameterParser)
 		, Version(FormatToVersion(Input.ShaderFormat))
 		, MinimumTargetEnvironment(GetMinimumTargetEnvironment(InInput))
 		, bStripReflect(InInput.IsRayTracingShader() || (IsAndroidShaderFormat(Input.ShaderFormat) && InInput.Environment.GetCompileArgument(TEXT("STRIP_REFLECT_ANDROID"), true)))
@@ -264,6 +265,7 @@ struct FVulkanShaderCompilerInternalState
 	}
 
 	const FShaderCompilerInput& Input;
+	const FShaderParameterParser& ParameterParser;
 
 	const EVulkanShaderVersion Version;
 	const CrossCompiler::FShaderConductorOptions::ETargetEnvironment MinimumTargetEnvironment;
@@ -1419,6 +1421,21 @@ static bool BuildShaderOutputFromSpirv(
 				check(SpvResult == SPV_REFLECT_RESULT_SUCCESS);
 
 				const int32 VulkanBindingIndex = SerializedOutput.Spirv.ReflectionInfo.Add(FVulkanSpirv::FEntry(ResourceName, BindingIndex));
+				const FShaderParameterParser::FParsedShaderParameter* ParsedParam = InternalState.ParameterParser.FindParameterInfosUnsafe(ResourceName);
+
+				auto AddShaderValidationType = [] (uint32_t VulkanBindingIndex, const FShaderParameterParser::FParsedShaderParameter* ParsedParam, FShaderCompilerOutput& Output) {
+					/*if (ParsedParam)
+					{
+						if (IsResourceBindingTypeSRV(ParsedParam->ParsedTypeDecl))
+						{
+							AddShaderValidationSRVType(VulkanBindingIndex, ParsedParam->ParsedTypeDecl, Output);
+						}
+						else
+						{
+							AddShaderValidationUAVType(VulkanBindingIndex, ParsedParam->ParsedTypeDecl, Output);
+						}
+					}*/
+				};
 
 				switch (BindingType)
 				{
@@ -1427,6 +1444,8 @@ static bool BuildShaderOutputFromSpirv(
 				case EVulkanBindingType::StorageImage:
 					HandleReflectedShaderUAV(ResourceName, BindingOffset, VulkanBindingIndex, 1, Output);
 					EntryTypes.Add(ResourceName, FVulkanShaderHeader::Global);
+
+					AddShaderValidationType(BindingOffset, ParsedParam, Output);
 					break;
 
 				case EVulkanBindingType::Image:
@@ -1441,6 +1460,8 @@ static bool BuildShaderOutputFromSpirv(
 				case EVulkanBindingType::UniformTexelBuffer:
 					HandleReflectedShaderResource(ResourceName, BindingOffset, VulkanBindingIndex, 1, Output);
 					EntryTypes.Add(ResourceName, FVulkanShaderHeader::Global);
+
+					AddShaderValidationType(BindingOffset, ParsedParam, Output);
 					break;
 
 				case EVulkanBindingType::Sampler:
@@ -1452,6 +1473,8 @@ static bool BuildShaderOutputFromSpirv(
 				case EVulkanBindingType::AccelerationStructure:
 					HandleReflectedShaderResource(ResourceName, BindingOffset, VulkanBindingIndex, 1, Output);
 					EntryTypes.Add(ResourceName, FVulkanShaderHeader::Global);
+
+					AddShaderValidationType(BindingOffset, ParsedParam, Output);
 					break;
 
 				case EVulkanBindingType::InputAttachment:
@@ -1505,6 +1528,8 @@ static bool BuildShaderOutputFromSpirv(
 						UsedUniformBufferSlots[VulkanBindingIndex] = true;
 						HandleReflectedUniformBuffer(ResourceName, VulkanBindingIndex, Output);
 						EntryTypes.Add(ResourceName, FVulkanShaderHeader::UniformBuffer);
+
+						AddShaderValidationUBSize(BindingOffset, Binding->block.padded_size, Output);
 					}
 					break;
 
@@ -2126,7 +2151,7 @@ static void RemoveUnusedBindlessHeaps(FString& PreprocessedShaderSource, const T
 
 bool PreprocessVulkanShader(const FShaderCompilerInput& Input, const FShaderCompilerEnvironment& Environment, FShaderPreprocessOutput& PreprocessOutput)
 {
-	const FVulkanShaderCompilerInternalState InternalState(Input);
+	const FVulkanShaderCompilerInternalState InternalState(Input, PreprocessOutput.GetParameterParser());
 
 	PRAGMA_DISABLE_DEPRECATION_WARNINGS		// FShaderCompilerDefinitions will be made internal in the future, marked deprecated until then
 
@@ -2505,7 +2530,7 @@ void CompileVulkanShader(const FShaderCompilerInput& Input, const FShaderPreproc
 {
 	check(IsVulkanShaderFormat(Input.ShaderFormat));
 
-	FVulkanShaderCompilerInternalState InternalState(Input);
+	FVulkanShaderCompilerInternalState InternalState(Input, PreprocessOutput.GetParameterParser());
 
 	const EHlslShaderFrequency HlslFrequency = InternalState.GetHlslShaderFrequency();
 	if (HlslFrequency == HSF_InvalidFrequency)
@@ -2557,6 +2582,8 @@ void CompileVulkanShader(const FShaderCompilerInput& Input, const FShaderPreproc
 		Output.ShaderCode.AddOptionalData(FShaderCodeName::Key, TCHAR_TO_UTF8(*Input.GenerateShaderName()));
 	}
 
+	Output.SerializeShaderCodeValidation();
+
 	PreprocessOutput.GetParameterParser().ValidateShaderParameterTypes(Input, InternalState.IsMobileES31(), Output);
 	
 	const bool bDirectCompile = FParse::Param(FCommandLine::Get(), TEXT("directcompile"));
diff --git a/Engine/Source/Developer/Windows/ShaderFormatD3D/Private/D3DShaderCompiler.inl b/Engine/Source/Developer/Windows/ShaderFormatD3D/Private/D3DShaderCompiler.inl
--- a/Engine/Source/Developer/Windows/ShaderFormatD3D/Private/D3DShaderCompiler.inl
+++ b/Engine/Source/Developer/Windows/ShaderFormatD3D/Private/D3DShaderCompiler.inl
@@ -138,6 +138,7 @@ template <typename ID3D1xShaderReflection, typename D3D1x_SHADER_DESC, typename
 				// Track just the constant buffer itself.
 				const FString UniformBufferName(CBDesc.Name);
 
+				AddShaderValidationUBSize(CBIndex, CBDesc.Size, Output);
 				HandleReflectedUniformBuffer(UniformBufferName, CBIndex, Output);
 				
 				UsedUniformBufferSlots[CBIndex] = true;
@@ -196,6 +197,9 @@ template <typename ID3D1xShaderReflection, typename D3D1x_SHADER_DESC, typename
 			}
 			else
 			{
+				const FShaderParameterParser::FParsedShaderParameter* ParsedParam = ShaderParameterParser.FindParameterInfosUnsafe(BindDesc.Name);
+				AddShaderValidationSRVType(BindDesc.BindPoint, ParsedParam ? ParsedParam->ParsedTypeDecl : EShaderCodeResourceBindingType::Invalid, Output);
+
 				HandleReflectedShaderResource(FString(BindDesc.Name), BindDesc.BindPoint, Output);
 				NumSRVs = FMath::Max(NumSRVs, BindDesc.BindPoint + BindCount);
 			}
@@ -239,6 +243,9 @@ template <typename ID3D1xShaderReflection, typename D3D1x_SHADER_DESC, typename
 			}
 			else
 			{
+				const FShaderParameterParser::FParsedShaderParameter* ParsedParam = ShaderParameterParser.FindParameterInfosUnsafe(BindDesc.Name);
+				AddShaderValidationUAVType(BindDesc.BindPoint, ParsedParam ? ParsedParam->ParsedTypeDecl : EShaderCodeResourceBindingType::Invalid, Output);
+
 				HandleReflectedShaderUAV(FString(BindDesc.Name), BindDesc.BindPoint, Output);
 				NumUAVs = FMath::Max(NumUAVs, BindDesc.BindPoint + BindCount);
 			}
@@ -247,6 +254,9 @@ template <typename ID3D1xShaderReflection, typename D3D1x_SHADER_DESC, typename
 		{
 			check(BindDesc.BindCount == 1);
 			FString BindDescName(BindDesc.Name);
+			const FShaderParameterParser::FParsedShaderParameter* ParsedParam = ShaderParameterParser.FindParameterInfosUnsafe(BindDesc.Name);
+			AddShaderValidationSRVType(BindDesc.BindPoint, ParsedParam ? ParsedParam->ParsedTypeDecl : EShaderCodeResourceBindingType::Invalid, Output);
+
 			HandleReflectedShaderResource(BindDescName, BindDesc.BindPoint, Output);
 
 			// https://learn.microsoft.com/en-us/windows/win32/api/d3d12shader/ns-d3d12shader-d3d12_shader_input_bind_desc
@@ -262,6 +272,10 @@ template <typename ID3D1xShaderReflection, typename D3D1x_SHADER_DESC, typename
 		{
 			// Acceleration structure resources are treated as SRVs.
 			check(BindDesc.BindCount == 1);
+
+			const FShaderParameterParser::FParsedShaderParameter* ParsedParam = ShaderParameterParser.FindParameterInfosUnsafe(BindDesc.Name);
+			AddShaderValidationSRVType(BindDesc.BindPoint, ParsedParam ? ParsedParam->ParsedTypeDecl : EShaderCodeResourceBindingType::Invalid, Output);
+
 			HandleReflectedShaderResource(FString(BindDesc.Name), BindDesc.BindPoint, Output);
 			NumSRVs = FMath::Max(NumSRVs, BindDesc.BindPoint + 1);
 		}
diff --git a/Engine/Source/Runtime/RHI/Private/RHIResources.cpp b/Engine/Source/Runtime/RHI/Private/RHIResources.cpp
--- a/Engine/Source/Runtime/RHI/Private/RHIResources.cpp
+++ b/Engine/Source/Runtime/RHI/Private/RHIResources.cpp
@@ -83,6 +83,34 @@ void FRHITexture::SetName(const FName& InName)
 #endif
 }
 
+const TCHAR* FRHIViewDesc::GetBufferTypeString(FRHIViewDesc::EBufferType BufferType)
+{
+	switch (BufferType)
+	{
+	case FRHIViewDesc::EBufferType::Unknown:				return TEXT("Unknown");
+	case FRHIViewDesc::EBufferType::Typed:					return TEXT("Typed");
+	case FRHIViewDesc::EBufferType::Structured:				return TEXT("Structured");
+	case FRHIViewDesc::EBufferType::AccelerationStructure:	return TEXT("AccelerationStructure");
+	case FRHIViewDesc::EBufferType::Raw:					return TEXT("Raw");
+	default:												checkf(false, TEXT("Missing FRHIViewDesc::EBufferType %d"), BufferType);
+	}
+	return TEXT("");
+}
+
+const TCHAR* FRHIViewDesc::GetTextureDimensionString(FRHIViewDesc::EDimension Dimension)
+{
+	switch (Dimension)
+	{
+	case EDimension::Texture2D:				return TEXT("Texture2D");
+	case EDimension::Texture2DArray:		return TEXT("Texture2DArray");
+	case EDimension::Texture3D:				return TEXT("Texture3D");
+	case EDimension::TextureCube:			return TEXT("TextureCube");
+	case EDimension::TextureCubeArray:		return TEXT("TextureCubeArray");
+	default:								checkf(false, TEXT("Missing FRHIViewDesc::EDimension %d"), Dimension);
+	}
+
+	return TEXT("");
+}
 FRHIViewDesc::FBuffer::FViewInfo FRHIViewDesc::FBuffer::GetViewInfo(FRHIBuffer* TargetBuffer) const
 {
 	check(TargetBuffer);
diff --git a/Engine/Source/Runtime/RHI/Private/RHIStrings.cpp b/Engine/Source/Runtime/RHI/Private/RHIStrings.cpp
--- a/Engine/Source/Runtime/RHI/Private/RHIStrings.cpp
+++ b/Engine/Source/Runtime/RHI/Private/RHIStrings.cpp
@@ -717,3 +717,34 @@ const TCHAR* GetUniformBufferBaseTypeString(EUniformBufferBaseType BaseType)
 	}
 	return TEXT("");
 }
+
+const TCHAR* GetShaderCodeResourceBindingTypeName(EShaderCodeResourceBindingType BindingType)
+{
+	switch (BindingType)
+	{
+	case EShaderCodeResourceBindingType::Invalid:							return TEXT("Invalid");
+	case EShaderCodeResourceBindingType::SamplerState:						return TEXT("SamplerState");
+	case EShaderCodeResourceBindingType::Texture2D:							return TEXT("Texture2D");
+	case EShaderCodeResourceBindingType::Texture2DArray:					return TEXT("Texture2DArray");
+	case EShaderCodeResourceBindingType::Texture2DMS:						return TEXT("Texture2DMS");
+	case EShaderCodeResourceBindingType::Texture3D:							return TEXT("Texture3D");
+	case EShaderCodeResourceBindingType::TextureCube:						return TEXT("TextureCube");
+	case EShaderCodeResourceBindingType::TextureCubeArray:					return TEXT("TextureCubeArray");
+	case EShaderCodeResourceBindingType::TextureMetadata:					return TEXT("TextureMetadata");
+	case EShaderCodeResourceBindingType::Buffer:							return TEXT("Buffer");
+	case EShaderCodeResourceBindingType::StructuredBuffer:					return TEXT("StructuredBuffer");
+	case EShaderCodeResourceBindingType::ByteAddressBuffer:					return TEXT("ByteAddressBuffer");
+	case EShaderCodeResourceBindingType::RaytracingAccelerationStructure:	return TEXT("RaytracingAccelerationStructure");
+	case EShaderCodeResourceBindingType::RWTexture2D:						return TEXT("RWTexture2D");
+	case EShaderCodeResourceBindingType::RWTexture2DArray:					return TEXT("RWTexture2DArray");
+	case EShaderCodeResourceBindingType::RWTexture3D:						return TEXT("RWTexture3D");
+	case EShaderCodeResourceBindingType::RWTextureCube:						return TEXT("RWTextureCube");
+	case EShaderCodeResourceBindingType::RWTextureMetadata:					return TEXT("RWTextureMetadata");
+	case EShaderCodeResourceBindingType::RWBuffer:							return TEXT("RWBuffer");
+	case EShaderCodeResourceBindingType::RWStructuredBuffer:				return TEXT("RWStructuredBuffer");
+	case EShaderCodeResourceBindingType::RWByteAddressBuffer:				return TEXT("RWByteAddressBuffer");
+	case EShaderCodeResourceBindingType::RasterizerOrderedTexture2D:		return TEXT("RasterizerOrderedTexture2D");
+	default:																checkf(false, TEXT("Missing EShaderCodeResourceBindingType %d"), BindingType);
+	}
+	return TEXT("");
+}
diff --git a/Engine/Source/Runtime/RHI/Private/RHIValidation.cpp b/Engine/Source/Runtime/RHI/Private/RHIValidation.cpp
--- a/Engine/Source/Runtime/RHI/Private/RHIValidation.cpp
+++ b/Engine/Source/Runtime/RHI/Private/RHIValidation.cpp
@@ -2056,36 +2056,286 @@ namespace RHIValidation
 		return Backtrace;
 	}
 
+	bool ValidateDimension(EShaderCodeResourceBindingType Type, FRHIViewDesc::EDimension Dimension, bool SRV)
+	{
+		// Ignore invalid types
+		if (Type == EShaderCodeResourceBindingType::Invalid)
+		{
+			return true;
+		}
+
+		if (IsResourceBindingTypeSRV(Type) != SRV)
+		{
+			return false;
+		}
+
+		if (Type == EShaderCodeResourceBindingType::Texture2D || Type == EShaderCodeResourceBindingType::RWTexture2D || Type == EShaderCodeResourceBindingType::Texture2DMS)
+		{
+			return Dimension == FRHIViewDesc::EDimension::Texture2D;
+		}
+		else if (Type == EShaderCodeResourceBindingType::Texture2DArray || Type == EShaderCodeResourceBindingType::RWTexture2DArray)
+		{
+			return Dimension == FRHIViewDesc::EDimension::Texture2DArray;
+		}
+		else if (Type == EShaderCodeResourceBindingType::Texture3D || Type == EShaderCodeResourceBindingType::RWTexture3D)
+		{
+			return Dimension == FRHIViewDesc::EDimension::Texture3D;
+		}
+		else if (Type == EShaderCodeResourceBindingType::TextureCube || Type == EShaderCodeResourceBindingType::RWTextureCube)
+		{
+			return Dimension == FRHIViewDesc::EDimension::TextureCube;
+		}
+		else if (Type == EShaderCodeResourceBindingType::TextureCubeArray)
+		{
+			return Dimension == FRHIViewDesc::EDimension::TextureCubeArray;
+		}
+
+		return false;
+	}
+
+	bool ValidateBuffer(EShaderCodeResourceBindingType Type, FRHIViewDesc::EBufferType BufferType, bool SRV)
+	{
+		// Ignore invalid types
+		if (Type == EShaderCodeResourceBindingType::Invalid)
+		{
+			return true;
+		}
+
+		if (IsResourceBindingTypeSRV(Type) != SRV)
+		{
+			return false;
+		}
+
+		if (Type == EShaderCodeResourceBindingType::ByteAddressBuffer || Type == EShaderCodeResourceBindingType::RWByteAddressBuffer)
+		{
+			return BufferType == FRHIViewDesc::EBufferType::Raw;
+		}
+		else if (Type == EShaderCodeResourceBindingType::StructuredBuffer || Type == EShaderCodeResourceBindingType::RWStructuredBuffer)
+		{
+			return BufferType == FRHIViewDesc::EBufferType::Structured;
+		}
+		else if (Type == EShaderCodeResourceBindingType::Buffer || Type == EShaderCodeResourceBindingType::RWBuffer)
+		{
+			return BufferType == FRHIViewDesc::EBufferType::Typed;
+		}
+		else if (Type == EShaderCodeResourceBindingType::RaytracingAccelerationStructure)
+		{
+			return BufferType == FRHIViewDesc::EBufferType::AccelerationStructure;
+		}
+		
+		return false;
+	}
+
 	/** Validates that the SRV is conform to what the shader expects */
-	void ValidateShaderResourceView(const FRHIShader* RHIShaderBase, uint32 BindIndex, FRHIShaderResourceView* SRV)
+	void ValidateShaderResourceView(const FRHIShader* RHIShaderBase, uint32 BindIndex, const FRHIShaderResourceView* SRV)
 	{
 #if RHI_INCLUDE_SHADER_DEBUG_DATA
 		if (SRV)
 		{
-			// DebugStrideValidationData is supposed to be already sorted
-			static const auto ShaderCodeValidationStridePredicate = [](const FShaderCodeValidationStride& lhs, const FShaderCodeValidationStride& rhs) -> bool { return lhs.BindPoint < rhs.BindPoint; };
-
 			auto const ViewIdentity = SRV->GetViewIdentity();
 
+			static const auto GetSRVName = [](const FRHIShaderResourceView* SRV, auto& ViewIdentity) -> FString
+			{
+				FString SRVName;
+				if (ViewIdentity.Resource)
+				{
+					SRVName = ViewIdentity.Resource->GetDebugName();
+				}
+				if (SRVName.IsEmpty())
+				{
+					SRVName = SRV->GetOwnerName().ToString();
+				}
+
+				return SRVName;
+			};
+
+			// DebugStrideValidationData is supposed to be already sorted
+			static const auto ShaderCodeValidationStridePredicate = [](const FShaderCodeValidationStride& lhs, const FShaderCodeValidationStride& rhs) -> bool { return lhs.BindPoint < rhs.BindPoint; };
 			FShaderCodeValidationStride SRVValidationStride = { BindIndex , ViewIdentity.Stride };
-			const int32 FoundIndex =  Algo::BinarySearch(RHIShaderBase->DebugStrideValidationData, SRVValidationStride, ShaderCodeValidationStridePredicate);
+
+			int32 FoundIndex =  Algo::BinarySearch(RHIShaderBase->DebugStrideValidationData, SRVValidationStride, ShaderCodeValidationStridePredicate);
 			if (FoundIndex != INDEX_NONE)
 			{
+				FString SRVName = GetSRVName(SRV, ViewIdentity);
 				uint16 ExpectedStride = RHIShaderBase->DebugStrideValidationData[FoundIndex].Stride;
 				if (ExpectedStride != SRVValidationStride.Stride)
 				{
-					FString SRVName;
-					// It seems that owner name is usually unknown. Instead try to rely on the tracked buffer
-					if (ViewIdentity.Resource)
+					
+					FString ErrorMessage = FString::Printf(TEXT("Shader %s: Buffer stride for \"%s\" must match structure size declared in the shader"), RHIShaderBase->GetShaderName(), *SRVName);
+					ErrorMessage += FString::Printf(TEXT("\nBind point: %d, HLSL size: %d, Buffer Size: %d"), BindIndex, ExpectedStride, SRVValidationStride.Stride);
+					RHI_VALIDATION_CHECK(false, *ErrorMessage);
+				}
+			}
+
+			// Validate Type
+			if (!RHIShaderBase->DebugSRVTypeValidationData.Num())
+				return;
+
+			static const auto ShaderCodeValidationTypePredicate = [](const FShaderCodeValidationType& lhs, const FShaderCodeValidationType& rhs) -> bool { return lhs.BindPoint < rhs.BindPoint; };
+			
+			FShaderCodeValidationType SRVValidationType = { BindIndex , EShaderCodeResourceBindingType::Invalid };
+			FoundIndex = Algo::BinarySearch(RHIShaderBase->DebugSRVTypeValidationData, SRVValidationType, ShaderCodeValidationTypePredicate);
+
+			if (FoundIndex != INDEX_NONE)
+			{
+				EShaderCodeResourceBindingType ExpectedType = RHIShaderBase->DebugSRVTypeValidationData[FoundIndex].Type;
+
+				if (SRV->IsTexture())
+				{
+					if (!ValidateDimension(ExpectedType, SRV->GetDesc().Texture.SRV.Dimension, true))
 					{
-						SRVName = ViewIdentity.Resource->GetDebugName();
+						FString SRVName = GetSRVName(SRV, ViewIdentity);
+						FString ErrorMessage = FString::Printf(TEXT("Shader %s: Dimension for SRV \"%s\" must match type declared in the shader"), RHIShaderBase->GetShaderName(), *SRVName);
+						ErrorMessage += FString::Printf(TEXT("\nBind point: %d, HLSL Type: %s, Actual Dimension: %s"),
+							BindIndex, 
+							GetShaderCodeResourceBindingTypeName(ExpectedType),
+							FRHIViewDesc::GetTextureDimensionString(SRV->GetDesc().Texture.SRV.Dimension));
+						RHI_VALIDATION_CHECK(false, *ErrorMessage);
 					}
-					if (SRVName.IsEmpty())
+				}
+				else if (SRV->IsBuffer())
+				{
+					if (!ValidateBuffer(ExpectedType, SRV->GetDesc().Buffer.SRV.BufferType, true))
 					{
-						SRVName = SRV->GetOwnerName().ToString();
+						FString SRVName = GetSRVName(SRV, ViewIdentity);
+						FString ErrorMessage = FString::Printf(TEXT("Shader %s: Buffer type for SRV \"%s\" must match buffer type declared in the shader"), RHIShaderBase->GetShaderName(), *SRVName);
+						ErrorMessage += FString::Printf(TEXT("\nBind point: %d, HLSL Type: %s, Actual Type: %s"),
+							BindIndex,
+							GetShaderCodeResourceBindingTypeName(ExpectedType),
+							FRHIViewDesc::GetBufferTypeString(SRV->GetDesc().Buffer.SRV.BufferType));
+						RHI_VALIDATION_CHECK(false, *ErrorMessage);
 					}
-					FString ErrorMessage = FString::Printf(TEXT("Shader %s: Buffer stride for \"%s\" must match structure size declared in the shader"), RHIShaderBase->GetShaderName(), *SRVName);
-					ErrorMessage += FString::Printf(TEXT("\nBind point: %d, HLSL size: %d, Buffer Size: %d"), BindIndex, ExpectedStride, SRVValidationStride.Stride);
+				}
+			}
+			else
+			{ 
+				FString SRVName = GetSRVName(SRV, ViewIdentity);
+				FString ErrorMessage = FString::Printf(TEXT("Shader %s: No bind point found for SRV \"%s\" possible UAV/SRV mismatch"), RHIShaderBase->GetShaderName(), *SRVName);
+				if (SRV->IsTexture())
+				{
+					ErrorMessage += FString::Printf(TEXT("\nBind point: %d, Type: %s"),
+						BindIndex,
+						FRHIViewDesc::GetTextureDimensionString(SRV->GetDesc().Texture.SRV.Dimension));
+				}
+				else
+				{
+					ErrorMessage += FString::Printf(TEXT("\nBind point: %d, Type: %s"),
+						BindIndex,
+						FRHIViewDesc::GetBufferTypeString(SRV->GetDesc().Buffer.SRV.BufferType));
+				}
+				RHI_VALIDATION_CHECK(false, *ErrorMessage);
+			}
+		}
+#endif
+	}
+
+	/** Validates that the UAV is conform to what the shader expects */
+	void ValidateUnorderedAccessView(const FRHIShader* RHIShaderBase, uint32 BindIndex, const FRHIUnorderedAccessView* UAV)
+	{
+#if RHI_INCLUDE_SHADER_DEBUG_DATA
+		if (UAV)
+		{
+			auto const ViewIdentity = UAV->GetViewIdentity();
+
+			static const auto GetUAVName = [](const FRHIUnorderedAccessView* UAV, auto& ViewIdentity) -> FString
+			{
+				FString UAVName;
+				if (ViewIdentity.Resource)
+				{
+					UAVName = ViewIdentity.Resource->GetDebugName();
+				}
+				if (UAVName.IsEmpty())
+				{
+					UAVName = UAV->GetOwnerName().ToString();
+				}
+
+				return UAVName;
+			};
+
+			// Validate Type
+			if (!RHIShaderBase->DebugUAVTypeValidationData.Num())
+				return;
+
+			static const auto ShaderCodeValidationTypePredicate = [](const FShaderCodeValidationType& lhs, const FShaderCodeValidationType& rhs) -> bool { return lhs.BindPoint < rhs.BindPoint; };
+
+			FShaderCodeValidationType SRVValidationType = { BindIndex , EShaderCodeResourceBindingType::Invalid };
+			int32 FoundIndex = Algo::BinarySearch(RHIShaderBase->DebugUAVTypeValidationData, SRVValidationType, ShaderCodeValidationTypePredicate);
+
+			if (FoundIndex != INDEX_NONE)
+			{
+				EShaderCodeResourceBindingType ExpectedType = RHIShaderBase->DebugUAVTypeValidationData[FoundIndex].Type;
+
+				if (UAV->IsTexture())
+				{
+					if (!ValidateDimension(ExpectedType, UAV->GetDesc().Texture.UAV.Dimension, false))
+					{
+						FString UAVName = GetUAVName(UAV, ViewIdentity);
+						FString ErrorMessage = FString::Printf(TEXT("Shader %s: Dimension for UAV \"%s\" must match type declared in the shader"), RHIShaderBase->GetShaderName(), *UAVName);
+						ErrorMessage += FString::Printf(TEXT("\nBind point: %d, HLSL Type: %s, Actual Dimension: %s"), 
+							BindIndex,
+							GetShaderCodeResourceBindingTypeName(ExpectedType),
+							FRHIViewDesc::GetTextureDimensionString(UAV->GetDesc().Texture.SRV.Dimension));
+						RHI_VALIDATION_CHECK(false, *ErrorMessage);
+					}
+				}
+				else if (UAV->IsBuffer())
+				{
+					if (!ValidateBuffer(ExpectedType, UAV->GetDesc().Buffer.UAV.BufferType, false))
+					{
+						FString UAVName = GetUAVName(UAV, ViewIdentity);
+						FString ErrorMessage = FString::Printf(TEXT("Shader %s: Buffer type for UAV \"%s\" must match buffer type declared in the shader"), RHIShaderBase->GetShaderName(), *UAVName);
+						ErrorMessage += FString::Printf(TEXT("\nBind point: %d, HLSL Type: %s, Actual Type: %s"),
+							BindIndex,
+							GetShaderCodeResourceBindingTypeName(ExpectedType),
+							FRHIViewDesc::GetBufferTypeString(UAV->GetDesc().Buffer.UAV.BufferType));
+						RHI_VALIDATION_CHECK(false, *ErrorMessage);
+					}
+				}
+			}
+			else
+			{
+				FString UAVName = GetUAVName(UAV, ViewIdentity);
+				FString ErrorMessage = FString::Printf(TEXT("Shader %s: No bind point found for UAV \"%s\" possible UAV/SRV mismatch"), RHIShaderBase->GetShaderName(), *UAVName);
+
+				if (UAV->IsTexture())
+				{
+					ErrorMessage += FString::Printf(TEXT("\nBind point: %d, Type: %s"),
+						BindIndex,
+						FRHIViewDesc::GetTextureDimensionString(UAV->GetDesc().Texture.SRV.Dimension));
+				}
+				else
+				{
+					ErrorMessage += FString::Printf(TEXT("\nBind point: %d, Type: %s"),
+						BindIndex,
+						FRHIViewDesc::GetBufferTypeString(UAV->GetDesc().Buffer.SRV.BufferType));
+				}
+				RHI_VALIDATION_CHECK(false, *ErrorMessage);
+			}
+		}
+#endif
+	}
+
+	/** Validates that the Uniform conforms to what the shader expects */
+	void ValidateUniformBuffer(const FRHIShader* RHIShaderBase, uint32 BindIndex, FRHIUniformBuffer* UB)
+	{
+#if RHI_INCLUDE_SHADER_DEBUG_DATA
+		if (UB)
+		{
+			// Validate Type
+			static const auto ShaderCodeValidationUBSizePredicate = [](const FShaderCodeValidationUBSize& lhs, const FShaderCodeValidationUBSize& rhs) -> bool { return lhs.BindPoint < rhs.BindPoint; };
+
+			FShaderCodeValidationUBSize SRVValidationSize = { BindIndex , 0 };
+			int32 FoundIndex = Algo::BinarySearch(RHIShaderBase->DebugUBSizeValidationData, SRVValidationSize, ShaderCodeValidationUBSizePredicate);
+			if (FoundIndex != INDEX_NONE)
+			{
+				uint32_t Size = RHIShaderBase->DebugUBSizeValidationData[FoundIndex].Size;
+
+				if(Size > 0 && Size > UB->GetSize())
+				{
+					const FRHIUniformBufferLayout& Layout = UB->GetLayout();
+
+					FString ErrorMessage = FString::Printf(TEXT("Shader %s: Uniform buffer \"%s\" has unexpected size"), RHIShaderBase->GetShaderName(), *Layout.GetDebugName());
+					ErrorMessage += FString::Printf(TEXT("\nBind point: %d, HLSL size: %d, Actual size: %d"), BindIndex, Size, UB->GetSize());
 					RHI_VALIDATION_CHECK(false, *ErrorMessage);
 				}
 			}
diff --git a/Engine/Source/Runtime/RHI/Public/RHIDefinitions.h b/Engine/Source/Runtime/RHI/Public/RHIDefinitions.h
--- a/Engine/Source/Runtime/RHI/Public/RHIDefinitions.h
+++ b/Engine/Source/Runtime/RHI/Public/RHIDefinitions.h
@@ -471,6 +471,37 @@ enum class EShaderCodeResourceBindingType : uint8
 	MAX
 };
 
+inline bool IsResourceBindingTypeSRV(EShaderCodeResourceBindingType Type)
+{
+	switch (Type)
+	{
+	case EShaderCodeResourceBindingType::Texture2D:
+	case EShaderCodeResourceBindingType::Texture2DArray:
+	case EShaderCodeResourceBindingType::Texture2DMS:
+	case EShaderCodeResourceBindingType::TextureCube:
+	case EShaderCodeResourceBindingType::TextureCubeArray:
+	case EShaderCodeResourceBindingType::Texture3D:
+	case EShaderCodeResourceBindingType::ByteAddressBuffer:
+	case EShaderCodeResourceBindingType::StructuredBuffer:
+	case EShaderCodeResourceBindingType::Buffer:
+	case EShaderCodeResourceBindingType::RaytracingAccelerationStructure:
+		return true;
+	case EShaderCodeResourceBindingType::RWTexture2D:
+	case EShaderCodeResourceBindingType::RWTexture2DArray:
+	case EShaderCodeResourceBindingType::RWTextureCube:
+	case EShaderCodeResourceBindingType::RWTexture3D:
+	case EShaderCodeResourceBindingType::RWByteAddressBuffer:
+	case EShaderCodeResourceBindingType::RWStructuredBuffer:
+	case EShaderCodeResourceBindingType::RWBuffer:
+	case EShaderCodeResourceBindingType::RasterizerOrderedTexture2D:
+		return false;
+	default:
+		ensureMsgf(0, TEXT("Missing or invalid SRV or UAV Type"));
+	}
+
+	return false;
+}
+
 /** The base type of a value in a shader parameter structure. */
 enum EUniformBufferBaseType : uint8
 {
@@ -1379,6 +1410,18 @@ struct FShaderCodeValidationStride
 	uint16 Stride;
 };
 
+struct FShaderCodeValidationType
+{
+	uint16 BindPoint;
+	EShaderCodeResourceBindingType Type;
+};
+
+struct FShaderCodeValidationUBSize
+{
+	uint16 BindPoint;
+	uint32 Size;
+};
+
 #if UE_ENABLE_INCLUDE_ORDER_DEPRECATED_IN_5_2
 #include "CoreMinimal.h"
 #include "DataDrivenShaderPlatformInfo.h"
diff --git a/Engine/Source/Runtime/RHI/Public/RHIResources.h b/Engine/Source/Runtime/RHI/Public/RHIResources.h
--- a/Engine/Source/Runtime/RHI/Public/RHIResources.h
+++ b/Engine/Source/Runtime/RHI/Public/RHIResources.h
@@ -782,6 +782,9 @@ public:
 	}
 
 	TArray<FShaderCodeValidationStride> DebugStrideValidationData;
+	TArray<FShaderCodeValidationType> DebugSRVTypeValidationData;
+	TArray<FShaderCodeValidationType> DebugUAVTypeValidationData;
+	TArray<FShaderCodeValidationUBSize> DebugUBSizeValidationData;
 
 #else
 
@@ -2270,6 +2273,9 @@ struct FRHIViewDesc
 		FMemory::Memzero(*this);
 	}
 
+	static const TCHAR* GetBufferTypeString(EBufferType BufferType);
+	static const TCHAR* GetTextureDimensionString(EDimension Dimension);
+
 protected:
 	FRHIViewDesc(EViewType ViewType)
 	{
@@ -2738,7 +2744,7 @@ public:
 	bool IsTexture() const { return ViewDesc.IsTexture(); }
 
 #if ENABLE_RHI_VALIDATION
-	RHIValidation::FViewIdentity GetViewIdentity()
+	RHIValidation::FViewIdentity GetViewIdentity() const
 	{
 		return RHIValidation::FViewIdentity(Resource, ViewDesc);
 	}
diff --git a/Engine/Source/Runtime/RHI/Public/RHIStrings.h b/Engine/Source/Runtime/RHI/Public/RHIStrings.h
--- a/Engine/Source/Runtime/RHI/Public/RHIStrings.h
+++ b/Engine/Source/Runtime/RHI/Public/RHIStrings.h
@@ -73,6 +73,7 @@ RHI_API const TCHAR* GetTextureDimensionString(ETextureDimension Dimension);
 RHI_API const TCHAR* GetTextureCreateFlagString(ETextureCreateFlags TextureCreateFlag);
 RHI_API const TCHAR* GetBufferUsageFlagString(EBufferUsageFlags BufferUsage);
 RHI_API const TCHAR* GetUniformBufferBaseTypeString(EUniformBufferBaseType BaseType);
+RHI_API const TCHAR* GetShaderCodeResourceBindingTypeName(EShaderCodeResourceBindingType BindingType);
 
 // Needs to stay inline for shader formats which can't depend on the RHI module
 inline const TCHAR* GetShaderFrequencyString(EShaderFrequency Frequency, bool bIncludePrefix = true)
diff --git a/Engine/Source/Runtime/RHI/Public/RHIValidationCommon.h b/Engine/Source/Runtime/RHI/Public/RHIValidationCommon.h
--- a/Engine/Source/Runtime/RHI/Public/RHIValidationCommon.h
+++ b/Engine/Source/Runtime/RHI/Public/RHIValidationCommon.h
@@ -979,7 +979,13 @@ namespace RHIValidation
 	extern RHI_API void* CaptureBacktrace();
 
 	/** Validates that the SRV is conform to what the shader expects */
-	extern RHI_API void ValidateShaderResourceView(const FRHIShader* RHIShaderBase, uint32 BindIndex, FRHIShaderResourceView* SRV);
+	extern RHI_API void ValidateShaderResourceView(const FRHIShader* RHIShaderBase, uint32 BindIndex, const FRHIShaderResourceView* SRV);
+
+	/** Validates that the UAV conforms to what the shader expects */
+	extern RHI_API void ValidateUnorderedAccessView(const FRHIShader* RHIShaderBase, uint32 BindIndex, const FRHIUnorderedAccessView* SRV);
+
+	/** Validates that the UB conforms to what the shader expects */
+	extern RHI_API void ValidateUniformBuffer(const FRHIShader* RHIShaderBase, uint32 BindIndex, FRHIUniformBuffer* SRV);
 
 }
 
diff --git a/Engine/Source/Runtime/RHI/Public/RHIValidationContext.h b/Engine/Source/Runtime/RHI/Public/RHIValidationContext.h
--- a/Engine/Source/Runtime/RHI/Public/RHIValidationContext.h
+++ b/Engine/Source/Runtime/RHI/Public/RHIValidationContext.h
@@ -42,13 +42,27 @@ inline void ValidateShaderParameters(FRHIShader* RHIShader, RHIValidation::FTrac
 			}
 			break;
 		case FRHIShaderParameterResource::EType::UnorderedAccessView:
-			Tracker->AssertUAV(static_cast<FRHIUnorderedAccessView*>(Parameter.Resource), InRequiredUAVMode, Parameter.Index);
+			if (FRHIUnorderedAccessView* UAV = static_cast<FRHIUnorderedAccessView*>(Parameter.Resource))
+			{
+				if (GRHIValidationEnabled)
+				{
+					RHIValidation::ValidateUnorderedAccessView(RHIShader, Parameter.Index, UAV);
+				}
+				Tracker->AssertUAV(static_cast<FRHIUnorderedAccessView*>(Parameter.Resource), InRequiredUAVMode, Parameter.Index);
+			}
 			break;
 		case FRHIShaderParameterResource::EType::Sampler:
 			// No validation
 			break;
 		case FRHIShaderParameterResource::EType::UniformBuffer:
-			StaticUniformBuffers.ValidateSetShaderUniformBuffer(static_cast<FRHIUniformBuffer*>(Parameter.Resource));
+			if (FRHIUniformBuffer* UniformBuffer = static_cast<FRHIUniformBuffer*>(Parameter.Resource))
+			{
+				if (GRHIValidationEnabled)
+				{
+					RHIValidation::ValidateUniformBuffer(RHIShader, Parameter.Index, UniformBuffer);
+				}
+				StaticUniformBuffers.ValidateSetShaderUniformBuffer(UniformBuffer);
+			}
 			break;
 		default:
 			checkf(false, TEXT("Unhandled resource type?"));
diff --git a/Engine/Source/Runtime/RHICore/Private/RHICoreShader.cpp b/Engine/Source/Runtime/RHICore/Private/RHICoreShader.cpp
--- a/Engine/Source/Runtime/RHICore/Private/RHICoreShader.cpp
+++ b/Engine/Source/Runtime/RHICore/Private/RHICoreShader.cpp
@@ -76,6 +76,9 @@ void SetupShaderCodeValidationData(FRHIShader* RHIShader, FShaderCodeReader& Sha
 			FShaderCodeValidationExtension ShaderCodeValidationExtension;
 			ArValidationExtensionData << ShaderCodeValidationExtension;
 			RHIShader->DebugStrideValidationData.Append(ShaderCodeValidationExtension.ShaderCodeValidationStride);
+			RHIShader->DebugSRVTypeValidationData.Append(ShaderCodeValidationExtension.ShaderCodeValidationSRVType);
+			RHIShader->DebugUAVTypeValidationData.Append(ShaderCodeValidationExtension.ShaderCodeValidationUAVType);
+			RHIShader->DebugUBSizeValidationData.Append(ShaderCodeValidationExtension.ShaderCodeValidationUBSize);
 		}
 	}
 #endif
diff --git a/Engine/Source/Runtime/RenderCore/Private/ShaderCore.cpp b/Engine/Source/Runtime/RenderCore/Private/ShaderCore.cpp
--- a/Engine/Source/Runtime/RenderCore/Private/ShaderCore.cpp
+++ b/Engine/Source/Runtime/RenderCore/Private/ShaderCore.cpp
@@ -963,12 +963,23 @@ void FShaderCompilerOutput::CompressOutput(FName ShaderCompressionFormat, FOodle
 
 void FShaderCompilerOutput::SerializeShaderCodeValidation()
 {
-	if (ParametersStrideToValidate.Num() > 0)
+	if (ParametersStrideToValidate.Num() > 0 || ParametersSRVTypeToValidate.Num() > 0 ||
+		ParametersUAVTypeToValidate.Num() > 0 || ParametersUBSizeToValidate.Num() > 0)
 	{
 		FShaderCodeValidationExtension ShaderCodeValidationExtension;
+
 		ShaderCodeValidationExtension.ShaderCodeValidationStride.Append(ParametersStrideToValidate);
 		ShaderCodeValidationExtension.ShaderCodeValidationStride.Sort([](const FShaderCodeValidationStride& lhs, const FShaderCodeValidationStride& rhs) -> bool { return lhs.BindPoint < rhs.BindPoint; });
 
+		ShaderCodeValidationExtension.ShaderCodeValidationSRVType.Append(ParametersSRVTypeToValidate);
+		ShaderCodeValidationExtension.ShaderCodeValidationSRVType.Sort([](const FShaderCodeValidationType& lhs, const FShaderCodeValidationType& rhs) -> bool { return lhs.BindPoint < rhs.BindPoint; });
+
+		ShaderCodeValidationExtension.ShaderCodeValidationUAVType.Append(ParametersUAVTypeToValidate);
+		ShaderCodeValidationExtension.ShaderCodeValidationUAVType.Sort([](const FShaderCodeValidationType& lhs, const FShaderCodeValidationType& rhs) -> bool { return lhs.BindPoint < rhs.BindPoint; });
+
+		ShaderCodeValidationExtension.ShaderCodeValidationUBSize.Append(ParametersUBSizeToValidate);
+		ShaderCodeValidationExtension.ShaderCodeValidationUBSize.Sort([](const FShaderCodeValidationUBSize& lhs, const FShaderCodeValidationUBSize& rhs) -> bool { return lhs.BindPoint < rhs.BindPoint; });
+
 		TArray<uint8> WriterBytes;
 		FMemoryWriter Writer(WriterBytes);
 		Writer << ShaderCodeValidationExtension;
diff --git a/Engine/Source/Runtime/RenderCore/Private/ShaderParameterMetadata.cpp b/Engine/Source/Runtime/RenderCore/Private/ShaderParameterMetadata.cpp
--- a/Engine/Source/Runtime/RenderCore/Private/ShaderParameterMetadata.cpp
+++ b/Engine/Source/Runtime/RenderCore/Private/ShaderParameterMetadata.cpp
@@ -299,40 +299,6 @@ EShaderCodeResourceBindingType ParseShaderResourceBindingType(const TCHAR* Shade
 	return BindingType;
 }
 
-const TCHAR* const kShaderCodeResourceBindingTypeNames[] = {
-	TEXT("Invalid"),
-	TEXT("SamplerState"),
-	TEXT("Texture2D"),
-	TEXT("Texture2DArray"),
-	TEXT("Texture2DMS"),
-	TEXT("Texture3D"),
-	TEXT("TextureCube"),
-	TEXT("TextureCubeArray"),
-	TEXT("TextureMetadata"),
-	TEXT("Buffer"),
-	TEXT("StructuredBuffer"),
-	TEXT("ByteAddressBuffer"),
-	TEXT("RaytracingAccelerationStructure"),
-	TEXT("RWTexture2D"),
-	TEXT("RWTexture2DArray"),
-	TEXT("RWTexture3D"),
-	TEXT("RWTextureCube"),
-	TEXT("RWTextureMetadata"),
-	TEXT("RWBuffer"),
-	TEXT("RWStructuredBuffer"),
-	TEXT("RWByteAddressBuffer"),
-	TEXT("RasterizerOrderedTexture2D"),
-};
-
-static_assert(UE_ARRAY_COUNT(kShaderCodeResourceBindingTypeNames) == int32(EShaderCodeResourceBindingType::MAX), "TODO.");
-
-const TCHAR* GetShaderCodeResourceBindingTypeName(EShaderCodeResourceBindingType BindingType)
-{
-	check(BindingType != EShaderCodeResourceBindingType::Invalid);
-	check(int32(BindingType) < int32(EShaderCodeResourceBindingType::MAX));
-	return kShaderCodeResourceBindingTypeNames[int32(BindingType)];
-}
-
 class FUniformBufferMemberAndOffset
 {
 public:
diff --git a/Engine/Source/Runtime/RenderCore/Private/ShaderParameterParser.cpp b/Engine/Source/Runtime/RenderCore/Private/ShaderParameterParser.cpp
--- a/Engine/Source/Runtime/RenderCore/Private/ShaderParameterParser.cpp
+++ b/Engine/Source/Runtime/RenderCore/Private/ShaderParameterParser.cpp
@@ -214,13 +214,7 @@ EShaderParameterType FShaderParameterParser::ParseParameterType(
 		return EShaderParameterType::Sampler;
 	}
 
-	FStringView UntemplatedType = InType;
-	if (int32 Index = InType.Find(TEXT("<")); Index != INDEX_NONE)
-	{
-		// Remove the template argument but don't forget to clean up the type name
-		const int32 NumChars = InType.Len() - Index;
-		UntemplatedType = InType.LeftChop(NumChars).TrimEnd();
-	}
+	FStringView UntemplatedType = StripTemplateFromType(InType);
 
 	if (AllSRVTypes.Contains(UntemplatedType) || InExtraSRVTypes.Contains(UntemplatedType))
 	{
@@ -419,6 +413,10 @@ bool FShaderParameterParser::ParseParameters(
 					}
 				}
 
+				FStringView StrippedTypeStringView = StripTemplateFromType(Type);
+				FString StrippedTypeString(StrippedTypeStringView);
+				EShaderCodeResourceBindingType TypeDecl = ParseShaderResourceBindingType(*StrippedTypeString);
+
 				FParsedShaderParameter ParsedParameter;
 
 				EShaderParameterType ConstantBufferParameterType = EShaderParameterType::Num;
@@ -474,6 +472,7 @@ bool FShaderParameterParser::ParseParameters(
 					ParsedParameter.BindlessConversionType = BindlessConversionType;
 					ParsedParameter.ConstantBufferParameterType = ConstantBufferParameterType;
 					ParsedParameter.bGloballyCoherent = bGloballyCoherent;
+					ParsedParameter.ParsedTypeDecl = TypeDecl;
 
 					if (ArrayStartPos != -1 && ArrayEndPos != -1)
 					{
@@ -1364,6 +1363,7 @@ void SerializeParam(FArchive& Ar, FShaderParameterParser::FParsedShaderParameter
 	Ar << Param.ConstantBufferParameterType;
 	Ar << Param.bGloballyCoherent;
 	Ar << Param.bIsBindable;
+	Ar << Param.ParsedTypeDecl;
 }
 FArchive& operator<<(FArchive& Ar, FShaderParameterParser& Parser)
 {
diff --git a/Engine/Source/Runtime/RenderCore/Private/ShaderParameterStruct.cpp b/Engine/Source/Runtime/RenderCore/Private/ShaderParameterStruct.cpp
--- a/Engine/Source/Runtime/RenderCore/Private/ShaderParameterStruct.cpp
+++ b/Engine/Source/Runtime/RenderCore/Private/ShaderParameterStruct.cpp
@@ -5,7 +5,8 @@
 =============================================================================*/
 
 #include "ShaderParameterStruct.h"
-
+#include "RenderGraphPrivate.h"
+#include "RHIValidationCommon.h"
 
 /** Context of binding a map. */
 struct FShaderParameterStructBindingContext
@@ -491,7 +492,7 @@ struct FShaderParameterReader
 
 #if DO_CHECK
 
-void ValidateShaderParameters(const TShaderRef<FShader>& Shader, const FShaderParametersMetadata* ParametersMetadata, const void* ParametersData)
+void ValidateShaderParameters(const TShaderRef<FShader>& Shader, const FShaderParametersMetadata* ParametersMetadata, const void* ParametersData, bool bValidateRHI)
 {
 	const FShaderParameterBindings& Bindings = Shader->Bindings;
 
@@ -505,6 +506,8 @@ void ValidateShaderParameters(const TShaderRef<FShader>& Shader, const FShaderPa
 	const TCHAR* ShaderClassName = Shader.GetType()->GetName();
 	const TCHAR* ShaderParameterStructName = ParametersMetadata->GetStructTypeName();
 
+	FRHIShader* RHIShader = Shader.GetRHIShaderBase(Shader->GetFrequency());
+
 	for (const FShaderParameterBindings::FResourceParameter& Parameter : Bindings.ResourceParameters)
 	{
 		const EUniformBufferBaseType BaseType = static_cast<EUniformBufferBaseType>(Parameter.BaseType);
@@ -520,8 +523,28 @@ void ValidateShaderParameters(const TShaderRef<FShader>& Shader, const FShaderPa
 				{
 					EmitNullShaderParameterFatalError(Shader, ParametersMetadata, Parameter.ByteOffset);
 				}
+
+#if ENABLE_RHI_VALIDATION
+				if (GRHIValidationEnabled && bValidateRHI)
+				{
+					if (BaseType == UBMT_SRV)
+					{
+						if (const FRHIShaderResourceView* SRV = static_cast<const FRHIShaderResourceView*>(Resource))
+						{
+							RHIValidation::ValidateShaderResourceView(RHIShader, Parameter.BaseIndex, SRV);
+						}
+					}
+					else if (BaseType == UBMT_UAV)
+					{
+						if (const FRHIUnorderedAccessView* UAV = static_cast<const FRHIUnorderedAccessView*>(Resource))
+						{
+							RHIValidation::ValidateUnorderedAccessView(RHIShader, Parameter.BaseIndex, UAV);
+						}
+					}
+				}
+#endif 
+				break;
 			}
-			break;
 			case UBMT_RDG_TEXTURE:
 			{
 				const FRDGTexture* GraphTexture = Reader.Read<const FRDGTexture*>(Parameter);
@@ -529,21 +552,40 @@ void ValidateShaderParameters(const TShaderRef<FShader>& Shader, const FShaderPa
 				{
 					EmitNullShaderParameterFatalError(Shader, ParametersMetadata, Parameter.ByteOffset);
 				}
+				break;
 			}
-			break;
 			case UBMT_RDG_TEXTURE_SRV:
-			case UBMT_RDG_TEXTURE_UAV:
 			case UBMT_RDG_BUFFER_SRV:
+			case UBMT_RDG_TEXTURE_UAV:
 			case UBMT_RDG_BUFFER_UAV:
 			case UBMT_RDG_UNIFORM_BLOCK_SRV:
 			{
 				const FRDGResource* GraphResource = Reader.Read<const FRDGResource*>(Parameter);
+
 				if (!GraphResource)
 				{
 					EmitNullShaderParameterFatalError(Shader, ParametersMetadata, Parameter.ByteOffset);
 				}
+
+#if ENABLE_RHI_VALIDATION
+				if (GRHIValidationEnabled && bValidateRHI)
+				{
+					RDG_ALLOW_RHI_ACCESS_SCOPE();
+					if (BaseType == UBMT_RDG_TEXTURE_SRV || BaseType == UBMT_RDG_BUFFER_SRV)
+					{
+						const FRHIShaderResourceView* SRV = static_cast<const FRHIShaderResourceView*>(GraphResource->GetRHI());
+						RHIValidation::ValidateShaderResourceView(RHIShader, Parameter.BaseIndex, SRV);
+
+					}
+					else if (BaseType == UBMT_RDG_TEXTURE_UAV || BaseType == UBMT_RDG_BUFFER_UAV)
+					{
+						const FRHIUnorderedAccessView* UAV = static_cast<const FRHIUnorderedAccessView*>(GraphResource->GetRHI());
+						RHIValidation::ValidateUnorderedAccessView(RHIShader, Parameter.BaseIndex, UAV);
+					}
+				}
+#endif
+				break;
 			}
-			break;
 			default:
 				break;
 		}
diff --git a/Engine/Source/Runtime/RenderCore/Public/RenderGraphUtils.h b/Engine/Source/Runtime/RenderCore/Public/RenderGraphUtils.h
--- a/Engine/Source/Runtime/RenderCore/Public/RenderGraphUtils.h
+++ b/Engine/Source/Runtime/RenderCore/Public/RenderGraphUtils.h
@@ -201,7 +201,7 @@ void ClearUnusedGraphResources(
 	// Verify the shader have all the parameters it needs. This is done before the
 	// ClearUnusedGraphResourcesImpl() to not mislead user on why some resource are missing
 	// when debugging a validation failure.
-	ValidateShaderParameters(Shader, ParametersMetadata, InoutParameters);
+	ValidateShaderParameters(Shader, ParametersMetadata, InoutParameters, false);
 
 	// Clear the resources the shader won't need.
 	return ClearUnusedGraphResourcesImpl(Shader->Bindings, ParametersMetadata, InoutParameters, ExcludeList);
@@ -231,8 +231,8 @@ void ClearUnusedGraphResources(
 	// Verify the shader have all the parameters it needs. This is done before the
 	// ClearUnusedGraphResourcesImpl() to not mislead user on why some resource are missing
 	// when debugging a validation failure.
-	ValidateShaderParameters(ShaderA, ParametersMetadata, InoutParameters);
-	ValidateShaderParameters(ShaderB, ParametersMetadata, InoutParameters);
+	ValidateShaderParameters(ShaderA, ParametersMetadata, InoutParameters, false);
+	ValidateShaderParameters(ShaderB, ParametersMetadata, InoutParameters, false);
 
 	// Clear the resources the shader won't need.
 	const FShaderParameterBindings* ShaderBindings[] = { &ShaderA->Bindings, &ShaderB->Bindings };
diff --git a/Engine/Source/Runtime/RenderCore/Public/ShaderCompilerCore.h b/Engine/Source/Runtime/RenderCore/Public/ShaderCompilerCore.h
--- a/Engine/Source/Runtime/RenderCore/Public/ShaderCompilerCore.h
+++ b/Engine/Source/Runtime/RenderCore/Public/ShaderCompilerCore.h
@@ -596,6 +596,9 @@ struct FShaderCompilerOutput
 	PRAGMA_ENABLE_DEPRECATION_WARNINGS
 
 	TArray<FShaderCodeValidationStride> ParametersStrideToValidate;
+	TArray<FShaderCodeValidationType> ParametersSRVTypeToValidate;
+	TArray<FShaderCodeValidationType> ParametersUAVTypeToValidate;
+	TArray<FShaderCodeValidationUBSize> ParametersUBSizeToValidate;
 
 	TArray<FShaderDiagnosticData> ShaderDiagnosticDatas;
 
diff --git a/Engine/Source/Runtime/RenderCore/Public/ShaderCore.h b/Engine/Source/Runtime/RenderCore/Public/ShaderCore.h
--- a/Engine/Source/Runtime/RenderCore/Public/ShaderCore.h
+++ b/Engine/Source/Runtime/RenderCore/Public/ShaderCore.h
@@ -774,6 +774,16 @@ inline FArchive& operator<<(FArchive& Ar, FShaderCodeValidationStride& ShaderCod
 	return Ar << ShaderCodeValidationStride.BindPoint << ShaderCodeValidationStride.Stride;
 }
 
+inline FArchive& operator<<(FArchive& Ar, FShaderCodeValidationType& ShaderCodeValidationType)
+{
+	return Ar << ShaderCodeValidationType.BindPoint << ShaderCodeValidationType.Type;
+}
+
+inline FArchive& operator<<(FArchive& Ar, FShaderCodeValidationUBSize& ShaderCodeValidationSize)
+{
+	return Ar << ShaderCodeValidationSize.BindPoint << ShaderCodeValidationSize.Size;
+}
+
 struct FShaderCodeValidationExtension
 {
 	// for FindOptionalData() and AddOptionalData()
@@ -781,12 +791,18 @@ struct FShaderCodeValidationExtension
 	static constexpr uint16 StaticVersion = 0;
 
 	TArray<FShaderCodeValidationStride> ShaderCodeValidationStride;
+	TArray<FShaderCodeValidationType> ShaderCodeValidationSRVType;
+	TArray<FShaderCodeValidationType> ShaderCodeValidationUAVType;
+	TArray<FShaderCodeValidationUBSize> ShaderCodeValidationUBSize;
 	uint16 Version = StaticVersion;
 
 	friend FArchive& operator<<(FArchive& Ar, FShaderCodeValidationExtension& Extension)
 	{
 		Ar << Extension.Version;
 		Ar << Extension.ShaderCodeValidationStride;
+		Ar << Extension.ShaderCodeValidationSRVType;
+		Ar << Extension.ShaderCodeValidationUAVType;
+		Ar << Extension.ShaderCodeValidationUBSize;
 		return Ar;
 	}
 };
diff --git a/Engine/Source/Runtime/RenderCore/Public/ShaderParameterMetadata.h b/Engine/Source/Runtime/RenderCore/Public/ShaderParameterMetadata.h
--- a/Engine/Source/Runtime/RenderCore/Public/ShaderParameterMetadata.h
+++ b/Engine/Source/Runtime/RenderCore/Public/ShaderParameterMetadata.h
@@ -83,9 +83,6 @@ struct FUniformBufferEntry
 /** Parse the shader resource binding from the binding type used in shader code. */
 EShaderCodeResourceBindingType ParseShaderResourceBindingType(const TCHAR* ShaderType);
 
-const TCHAR* GetShaderCodeResourceBindingTypeName(EShaderCodeResourceBindingType BindingType);
-
-
 /** Simple class that registers a uniform buffer static slot in the constructor. */
 class FUniformBufferStaticSlotRegistrar
 {
diff --git a/Engine/Source/Runtime/RenderCore/Public/ShaderParameterParser.h b/Engine/Source/Runtime/RenderCore/Public/ShaderParameterParser.h
--- a/Engine/Source/Runtime/RenderCore/Public/ShaderParameterParser.h
+++ b/Engine/Source/Runtime/RenderCore/Public/ShaderParameterParser.h
@@ -4,6 +4,7 @@
 
 #include "CoreMinimal.h"
 #include "ShaderParameterMetadata.h"
+#include "ShaderCore.h"
 
 class FShaderCompilerFlags;
 struct FShaderCompilerInput;
@@ -33,6 +34,19 @@ namespace VulkanBindless
 	static constexpr const TCHAR* kBindlessSamplerArrayPrefix = TEXT("SamplerDescriptorHeap_");
 }
 
+inline FStringView StripTemplateFromType(const FStringView& Input)
+{
+	FStringView UntemplatedType = FStringView(Input);
+	if (int32 Index = Input.Find(TEXT("<")); Index != INDEX_NONE)
+	{
+		// Remove the template argument but don't forget to clean up the type name
+		const int32 NumChars = Input.Len() - Index;
+		UntemplatedType = Input.LeftChop(NumChars).TrimEnd();
+	}
+
+	return UntemplatedType;
+}
+
 /** Validates and moves all the shader loose data parameter defined in the root scope of the shader into the root uniform buffer. */
 class FShaderParameterParser
 {
@@ -74,6 +88,8 @@ public:
 		bool bGloballyCoherent = false;
 		bool bIsBindable = false;
 
+		EShaderCodeResourceBindingType ParsedTypeDecl = EShaderCodeResourceBindingType::Invalid;
+
 		friend class FShaderParameterParser;
 	};
 
@@ -136,6 +152,11 @@ public:
 		return ParsedParameters.FindChecked(ParameterName);
 	}
 
+	const FParsedShaderParameter* FindParameterInfosUnsafe(const FString& ParameterName) const
+	{
+		return ParsedParameters.Find(ParameterName);
+	}
+
 	/** Validates the shader parameter in code is compatible with the shader parameter structure. */
 	RENDERCORE_API void ValidateShaderParameterType(
 		const FShaderCompilerInput& CompilerInput,
diff --git a/Engine/Source/Runtime/RenderCore/Public/ShaderParameterStruct.h b/Engine/Source/Runtime/RenderCore/Public/ShaderParameterStruct.h
--- a/Engine/Source/Runtime/RenderCore/Public/ShaderParameterStruct.h
+++ b/Engine/Source/Runtime/RenderCore/Public/ShaderParameterStruct.h
@@ -117,15 +117,15 @@ extern RENDERCORE_API void EmitNullShaderParameterFatalError(const TShaderRef<FS
 
 /** Validates that all resource parameters of a shader are set. */
 #if DO_CHECK
-extern RENDERCORE_API void ValidateShaderParameters(const TShaderRef<FShader>& Shader, const FShaderParametersMetadata* ParametersMetadata, const void* ParametersData);
+extern RENDERCORE_API void ValidateShaderParameters(const TShaderRef<FShader>& Shader, const FShaderParametersMetadata* ParametersMetadata, const void* ParametersData, bool bValidateRHI);
 #else
-FORCEINLINE void ValidateShaderParameters(const TShaderRef<FShader>& Shader, const FShaderParametersMetadata* ParametersMetadata, const void* ParametersData) {}
+FORCEINLINE void ValidateShaderParameters(const TShaderRef<FShader>& Shader, const FShaderParametersMetadata* ParametersMetadata, const void* ParametersData, bool bValidateRHI) {}
 #endif
 
 template<typename TShaderClass>
 FORCEINLINE void ValidateShaderParameters(const TShaderRef<TShaderClass>& Shader, const typename TShaderClass::FParameters& Parameters)
 {
-	return ValidateShaderParameters(Shader, TShaderClass::FParameters::FTypeInfo::GetStructMetadata(), &Parameters);
+	return ValidateShaderParameters(Shader, TShaderClass::FParameters::FTypeInfo::GetStructMetadata(), &Parameters, false);
 }
 
 template<typename TParameterType>
@@ -229,7 +229,7 @@ inline void SetShaderParameters(
 	const FShaderParametersMetadata* ParametersMetadata,
 	const typename TShaderClass::FParameters& Parameters)
 {
-	ValidateShaderParameters(Shader, ParametersMetadata, &Parameters);
+	ValidateShaderParameters(Shader, ParametersMetadata, &Parameters, false);
 	SetShaderParameters(BatchedParameters, Shader->Bindings, ParametersMetadata, &Parameters);
 }
 
@@ -269,7 +269,7 @@ inline void SetShaderParameters(
 	const FShaderParametersMetadata* ParametersMetadata,
 	const typename TShaderClass::FParameters& Parameters)
 {
-	ValidateShaderParameters(Shader, ParametersMetadata, &Parameters);
+	ValidateShaderParameters(Shader, ParametersMetadata, &Parameters, false);
 
 	// TODO(RDG): Once all shader sets their parameter through this, can refactor RHI so all shader parameters get sets through a single RHI function call.
 
@@ -297,7 +297,7 @@ void SetShaderParameters(FRayTracingShaderBindingsWriter& RTBindingsWriter, cons
 {
 	const FShaderParametersMetadata* ParametersMetadata = TShaderClass::FParameters::FTypeInfo::GetStructMetadata();
 
-	ValidateShaderParameters(Shader, ParametersMetadata, &Parameters);
+	ValidateShaderParameters(Shader, ParametersMetadata, &Parameters, true);
 
 	checkf(Shader->Bindings.Parameters.Num() == 0, TEXT("Ray tracing shader should use SHADER_USE_ROOT_PARAMETER_STRUCT() to passdown the cbuffer layout to the shader compiler."));
 
diff --git a/Engine/Source/Runtime/RenderCore/Public/ShaderParameters.h b/Engine/Source/Runtime/RenderCore/Public/ShaderParameters.h
--- a/Engine/Source/Runtime/RenderCore/Public/ShaderParameters.h
+++ b/Engine/Source/Runtime/RenderCore/Public/ShaderParameters.h
@@ -31,6 +31,9 @@ struct FShaderCompilerEnvironment;
 enum class EShaderParameterType : uint8;
 DECLARE_INTRINSIC_TYPE_LAYOUT(EShaderParameterType);
 
+enum class EShaderCodeResourceBindingType : uint8;
+DECLARE_INTRINSIC_TYPE_LAYOUT(EShaderCodeResourceBindingType);
+
 #if WITH_EDITOR
 namespace UE::ShaderParameters
 {
diff --git a/Engine/Source/Runtime/VulkanRHI/Private/VulkanShaders.cpp b/Engine/Source/Runtime/VulkanRHI/Private/VulkanShaders.cpp
--- a/Engine/Source/Runtime/VulkanRHI/Private/VulkanShaders.cpp
+++ b/Engine/Source/Runtime/VulkanRHI/Private/VulkanShaders.cpp
@@ -11,6 +11,7 @@
 #include "Serialization/MemoryReader.h"
 #include "VulkanLLM.h"
 #include "VulkanDescriptorSets.h"
+#include "RHICoreShader.h"
 
 TAutoConsoleVariable<int32> GDynamicGlobalUBs(
 	TEXT("r.Vulkan.DynamicGlobalUBs"),
@@ -88,6 +89,7 @@ ShaderType* FVulkanShaderFactory::CreateShader(TArrayView<const uint8> Code, FVu
 				FShaderCodeReader ShaderCode(Code);
 				RetShader->Debug.ShaderName = ShaderCode.FindOptionalData(FShaderCodeName::Key);
 #endif
+				UE::RHICore::SetupShaderCodeValidationData(RetShader, ShaderCode);
 			}
 		}
 	}
@@ -173,6 +175,7 @@ FVulkanRayTracingShader* FVulkanShaderFactory::CreateRayTracingShader(TArrayView
 				FShaderCodeReader ShaderCode(Code);
 				RetShader->Debug.ShaderName = ShaderCode.FindOptionalData(FShaderCodeName::Key);
 #endif
+				UE::RHICore::SetupShaderCodeValidationData(RetShader, ShaderCode);
 			}
 		}
 	}
