[RHI]
* Fixed non unity compile error

diff --git a/Engine/Source/Runtime/RHI/Public/RHIDefinitions.h b/Engine/Source/Runtime/RHI/Public/RHIDefinitions.h
--- a/Engine/Source/Runtime/RHI/Public/RHIDefinitions.h
+++ b/Engine/Source/Runtime/RHI/Public/RHIDefinitions.h
@@ -9,6 +9,7 @@
 
 #include "GpuProfilerTrace.h" // TODO Move defines into RHIDefinitions
 #include "Math/NumericLimits.h"
+#include "Misc/AssertionMacros.h"
 #include "Misc/EnumClassFlags.h"
 #include "ProfilingDebugging/CsvProfilerConfig.h" // TODO Move defines into RHIDefinitions
 
