Removing individual parameter binding methods from IRHI*Context now that everything uses batched shader parameters
- RHISetShaderTexture, RHISetShaderSampler, RHISetUAVParameter, RHISetShaderResourceViewParameter, RHISetShaderUniformBuffer, RHISetShaderParameter
- Only removing from the interface, each RHI still uses them in conjunction with UE::RHICore::RHISetShaderParametersShared, which uses the RHI type and not the interface.
- These methods are not part of the "public" interface so no deprecation is needed.

diff --git a/Engine/Source/Runtime/Apple/MetalRHI/Public/MetalRHIContext.h b/Engine/Source/Runtime/Apple/MetalRHI/Public/MetalRHIContext.h
--- a/Engine/Source/Runtime/Apple/MetalRHI/Public/MetalRHIContext.h
+++ b/Engine/Source/Runtime/Apple/MetalRHI/Public/MetalRHIContext.h
@@ -100,47 +100,19 @@ public:
 	
 	virtual void RHISetStaticUniformBuffers(const FUniformBufferStaticBindings& InUniformBuffers) final override;
 
-	/** Set the shader resource view of a surface. */
-	virtual void RHISetShaderTexture(FRHIGraphicsShader* Shader, uint32 TextureIndex, FRHITexture* NewTexture) final override;
-	
-	/** Set the shader resource view of a surface.  This is used for binding TextureMS parameter types that need a multi sampled view. */
-	virtual void RHISetShaderTexture(FRHIComputeShader* PixelShader, uint32 TextureIndex, FRHITexture* NewTexture) final override;
-	
-	/**
-	 * Sets sampler state.
-	 * @param ComputeShader		The compute shader to set the sampler for.
-	 * @param SamplerIndex		The index of the sampler.
-	 * @param NewState			The new sampler state.
-	 */
-	virtual void RHISetShaderSampler(FRHIComputeShader* ComputeShader, uint32 SamplerIndex, FRHISamplerState* NewState) final override;
-	
-	/**
-	 * Sets sampler state.
-	 * @param Shader			The shader to set the sampler for.
-	 * @param SamplerIndex		The index of the sampler.
-	 * @param NewState			The new sampler state.
-	 */
-	virtual void RHISetShaderSampler(FRHIGraphicsShader* Shader, uint32 SamplerIndex, FRHISamplerState* NewState) final override;
-	
-	/** Sets a pixel shader UAV parameter. */
-	virtual void RHISetUAVParameter(FRHIPixelShader* PixelShaderRHI, uint32 UAVIndex, FRHIUnorderedAccessView* UAVRHI) final override;
-	
-	/** Sets a compute shader UAV parameter. */
-	virtual void RHISetUAVParameter(FRHIComputeShader* ComputeShader, uint32 UAVIndex, FRHIUnorderedAccessView* UAV) final override;
-	
-	/** Sets a compute shader UAV parameter and initial count */
-	virtual void RHISetUAVParameter(FRHIComputeShader* ComputeShader, uint32 UAVIndex, FRHIUnorderedAccessView* UAV, uint32 InitialCount) final override;
-	
-	virtual void RHISetShaderResourceViewParameter(FRHIGraphicsShader* Shader, uint32 SamplerIndex, FRHIShaderResourceView* SRV) final override;
-	
-	virtual void RHISetShaderResourceViewParameter(FRHIComputeShader* ComputeShader, uint32 SamplerIndex, FRHIShaderResourceView* SRV) final override;
-	
-	virtual void RHISetShaderUniformBuffer(FRHIGraphicsShader* Shader, uint32 BufferIndex, FRHIUniformBuffer* Buffer) final override;
-	
-	virtual void RHISetShaderUniformBuffer(FRHIComputeShader* ComputeShader, uint32 BufferIndex, FRHIUniformBuffer* Buffer) final override;
-	
-	virtual void RHISetShaderParameter(FRHIGraphicsShader* Shader, uint32 BufferIndex, uint32 BaseIndex, uint32 NumBytes, const void* NewValue) final override;
-	virtual void RHISetShaderParameter(FRHIComputeShader* ComputeShader, uint32 BufferIndex, uint32 BaseIndex, uint32 NumBytes, const void* NewValue) final override;
+	void RHISetShaderTexture(FRHIGraphicsShader* Shader, uint32 TextureIndex, FRHITexture* NewTexture);
+	void RHISetShaderTexture(FRHIComputeShader* PixelShader, uint32 TextureIndex, FRHITexture* NewTexture);
+	void RHISetShaderSampler(FRHIComputeShader* ComputeShader, uint32 SamplerIndex, FRHISamplerState* NewState);
+	void RHISetShaderSampler(FRHIGraphicsShader* Shader, uint32 SamplerIndex, FRHISamplerState* NewState);
+	void RHISetUAVParameter(FRHIPixelShader* PixelShaderRHI, uint32 UAVIndex, FRHIUnorderedAccessView* UAVRHI);
+	void RHISetUAVParameter(FRHIComputeShader* ComputeShader, uint32 UAVIndex, FRHIUnorderedAccessView* UAV);
+	void RHISetUAVParameter(FRHIComputeShader* ComputeShader, uint32 UAVIndex, FRHIUnorderedAccessView* UAV, uint32 InitialCount);
+	void RHISetShaderResourceViewParameter(FRHIGraphicsShader* Shader, uint32 SamplerIndex, FRHIShaderResourceView* SRV);
+	void RHISetShaderResourceViewParameter(FRHIComputeShader* ComputeShader, uint32 SamplerIndex, FRHIShaderResourceView* SRV);
+	void RHISetShaderUniformBuffer(FRHIGraphicsShader* Shader, uint32 BufferIndex, FRHIUniformBuffer* Buffer);
+	void RHISetShaderUniformBuffer(FRHIComputeShader* ComputeShader, uint32 BufferIndex, FRHIUniformBuffer* Buffer);
+	void RHISetShaderParameter(FRHIGraphicsShader* Shader, uint32 BufferIndex, uint32 BaseIndex, uint32 NumBytes, const void* NewValue);
+	void RHISetShaderParameter(FRHIComputeShader* ComputeShader, uint32 BufferIndex, uint32 BaseIndex, uint32 NumBytes, const void* NewValue);
 	
 	virtual void RHISetShaderParameters(FRHIGraphicsShader* Shader, TConstArrayView<uint8> InParametersData, TConstArrayView<FRHIShaderParameter> InParameters, TConstArrayView<FRHIShaderParameterResource> InResourceParameters, TConstArrayView<FRHIShaderParameterResource> InBindlessParameters) final override;
 	virtual void RHISetShaderParameters(FRHIComputeShader* Shader, TConstArrayView<uint8> InParametersData, TConstArrayView<FRHIShaderParameter> InParameters, TConstArrayView<FRHIShaderParameterResource> InResourceParameters, TConstArrayView<FRHIShaderParameterResource> InBindlessParameters) final override;
diff --git a/Engine/Source/Runtime/D3D12RHI/Private/D3D12CommandContext.h b/Engine/Source/Runtime/D3D12RHI/Private/D3D12CommandContext.h
--- a/Engine/Source/Runtime/D3D12RHI/Private/D3D12CommandContext.h
+++ b/Engine/Source/Runtime/D3D12RHI/Private/D3D12CommandContext.h
@@ -579,15 +579,15 @@ public:
 	virtual void RHIDispatchComputeShader(uint32 ThreadGroupCountX, uint32 ThreadGroupCountY, uint32 ThreadGroupCountZ) final override;
 	virtual void RHIDispatchIndirectComputeShader(FRHIBuffer* ArgumentBuffer, uint32 ArgumentOffset) final override;
 	virtual void RHISetStaticUniformBuffers(const FUniformBufferStaticBindings& InUniformBuffers) final override;
-	virtual void RHISetShaderTexture(FRHIComputeShader* PixelShader, uint32 TextureIndex, FRHITexture* NewTexture) final override;
-	virtual void RHISetShaderSampler(FRHIComputeShader* ComputeShader, uint32 SamplerIndex, FRHISamplerState* NewState) final override;
-	virtual void RHISetUAVParameter(FRHIPixelShader* PixelShaderRHI, uint32 UAVIndex, FRHIUnorderedAccessView* UAVRHI) final override;
-	virtual void RHISetUAVParameter(FRHIComputeShader* ComputeShader, uint32 UAVIndex, FRHIUnorderedAccessView* UAV) final override;
-	virtual void RHISetUAVParameter(FRHIComputeShader* ComputeShader, uint32 UAVIndex, FRHIUnorderedAccessView* UAV, uint32 InitialCount) final override;
-	virtual void RHISetShaderResourceViewParameter(FRHIComputeShader* ComputeShader, uint32 SamplerIndex, FRHIShaderResourceView* SRV) final override;
-	virtual void RHISetShaderUniformBuffer(FRHIComputeShader* ComputeShader, uint32 BufferIndex, FRHIUniformBuffer* Buffer) final override;
+	void RHISetShaderTexture(FRHIComputeShader* PixelShader, uint32 TextureIndex, FRHITexture* NewTexture);
+	void RHISetShaderSampler(FRHIComputeShader* ComputeShader, uint32 SamplerIndex, FRHISamplerState* NewState);
+	void RHISetUAVParameter(FRHIPixelShader* PixelShaderRHI, uint32 UAVIndex, FRHIUnorderedAccessView* UAVRHI);
+	void RHISetUAVParameter(FRHIComputeShader* ComputeShader, uint32 UAVIndex, FRHIUnorderedAccessView* UAV);
+	void RHISetUAVParameter(FRHIComputeShader* ComputeShader, uint32 UAVIndex, FRHIUnorderedAccessView* UAV, uint32 InitialCount);
+	void RHISetShaderResourceViewParameter(FRHIComputeShader* ComputeShader, uint32 SamplerIndex, FRHIShaderResourceView* SRV);
+	void RHISetShaderUniformBuffer(FRHIComputeShader* ComputeShader, uint32 BufferIndex, FRHIUniformBuffer* Buffer);
+	void RHISetShaderParameter(FRHIComputeShader* ComputeShader, uint32 BufferIndex, uint32 BaseIndex, uint32 NumBytes, const void* NewValue);
 	virtual void RHISetShaderParameters(FRHIComputeShader* Shader, TConstArrayView<uint8> InParametersData, TConstArrayView<FRHIShaderParameter> InParameters, TConstArrayView<FRHIShaderParameterResource> InResourceParameters, TConstArrayView<FRHIShaderParameterResource> InBindlessParameters) final override;
-	virtual void RHISetShaderParameter(FRHIComputeShader* ComputeShader, uint32 BufferIndex, uint32 BaseIndex, uint32 NumBytes, const void* NewValue) final override;
 	virtual void RHISetShaderUnbinds(FRHIComputeShader* Shader, TConstArrayView<FRHIShaderParameterUnbind> InUnbinds) final override;
 	virtual void RHISetShaderUnbinds(FRHIGraphicsShader* Shader, TConstArrayView<FRHIShaderParameterUnbind> InUnbinds) final override;
 	virtual void RHIPushEvent(const TCHAR* Name, FColor Color) final override;
@@ -611,12 +611,12 @@ public:
 	virtual void RHISetStereoViewport(float LeftMinX, float RightMinX, float LeftMinY, float RightMinY, float MinZ, float LeftMaxX, float RightMaxX, float LeftMaxY, float RightMaxY, float MaxZ) override;
 	virtual void RHISetScissorRect(bool bEnable, uint32 MinX, uint32 MinY, uint32 MaxX, uint32 MaxY) final override;
 	virtual void RHISetGraphicsPipelineState(FRHIGraphicsPipelineState* GraphicsPipelineState, uint32 StencilRef, bool bApplyAdditionalState) final override;
-	virtual void RHISetShaderTexture(FRHIGraphicsShader* Shader, uint32 TextureIndex, FRHITexture* NewTexture) final override;
-	virtual void RHISetShaderSampler(FRHIGraphicsShader* Shader, uint32 SamplerIndex, FRHISamplerState* NewState) final override;
-	virtual void RHISetShaderResourceViewParameter(FRHIGraphicsShader* Shader, uint32 SamplerIndex, FRHIShaderResourceView* SRV) final override;
-	virtual void RHISetShaderUniformBuffer(FRHIGraphicsShader* Shader, uint32 BufferIndex, FRHIUniformBuffer* Buffer) final override;
+	void RHISetShaderTexture(FRHIGraphicsShader* Shader, uint32 TextureIndex, FRHITexture* NewTexture);
+	void RHISetShaderSampler(FRHIGraphicsShader* Shader, uint32 SamplerIndex, FRHISamplerState* NewState);
+	void RHISetShaderResourceViewParameter(FRHIGraphicsShader* Shader, uint32 SamplerIndex, FRHIShaderResourceView* SRV);
+	void RHISetShaderUniformBuffer(FRHIGraphicsShader* Shader, uint32 BufferIndex, FRHIUniformBuffer* Buffer);
+	void RHISetShaderParameter(FRHIGraphicsShader* Shader, uint32 BufferIndex, uint32 BaseIndex, uint32 NumBytes, const void* NewValue);
 	virtual void RHISetShaderParameters(FRHIGraphicsShader* Shader, TConstArrayView<uint8> InParametersData, TConstArrayView<FRHIShaderParameter> InParameters, TConstArrayView<FRHIShaderParameterResource> InResourceParameters, TConstArrayView<FRHIShaderParameterResource> InBindlessParameters) final override;
-	virtual void RHISetShaderParameter(FRHIGraphicsShader* Shader, uint32 BufferIndex, uint32 BaseIndex, uint32 NumBytes, const void* NewValue) final override;
 	virtual void RHISetStencilRef(uint32 StencilRef) final override;
 	virtual void RHISetBlendFactor(const FLinearColor& BlendFactor) final override;
 
@@ -829,38 +829,6 @@ public:
 	{
 		ContextRedirect(RHICopyToStagingBuffer(SourceBuffer, DestinationStagingBuffer, Offset, NumBytes));
 	}
-	FORCEINLINE virtual void RHISetShaderTexture(FRHIComputeShader* PixelShader, uint32 TextureIndex, FRHITexture* NewTexture) final override
-	{
-		ContextRedirect(RHISetShaderTexture(PixelShader, TextureIndex, NewTexture));
-	}
-	FORCEINLINE virtual void RHISetShaderSampler(FRHIComputeShader* ComputeShader, uint32 SamplerIndex, FRHISamplerState* NewState) final override
-	{
-		ContextRedirect(RHISetShaderSampler(ComputeShader, SamplerIndex, NewState));
-	}
-	FORCEINLINE virtual void RHISetUAVParameter(FRHIPixelShader* PixelShader, uint32 UAVIndex, FRHIUnorderedAccessView* UAV) final override
-	{
-		ContextRedirect(RHISetUAVParameter(PixelShader, UAVIndex, UAV));
-	}
-	FORCEINLINE virtual void RHISetUAVParameter(FRHIComputeShader* ComputeShader, uint32 UAVIndex, FRHIUnorderedAccessView* UAV) final override
-	{
-		ContextRedirect(RHISetUAVParameter(ComputeShader, UAVIndex, UAV));
-	}
-	FORCEINLINE virtual void RHISetUAVParameter(FRHIComputeShader* ComputeShader, uint32 UAVIndex, FRHIUnorderedAccessView* UAV, uint32 InitialCount) final override
-	{
-		ContextRedirect(RHISetUAVParameter(ComputeShader, UAVIndex, UAV, InitialCount));
-	}
-	FORCEINLINE virtual void RHISetShaderResourceViewParameter(FRHIComputeShader* ComputeShader, uint32 SamplerIndex, FRHIShaderResourceView* SRV) final override
-	{
-		ContextRedirect(RHISetShaderResourceViewParameter(ComputeShader, SamplerIndex, SRV));
-	}
-	FORCEINLINE virtual void RHISetShaderUniformBuffer(FRHIComputeShader* ComputeShader, uint32 BufferIndex, FRHIUniformBuffer* Buffer) final override
-	{
-		ContextRedirect(RHISetShaderUniformBuffer(ComputeShader, BufferIndex, Buffer));
-	}
-	FORCEINLINE virtual void RHISetShaderParameter(FRHIComputeShader* ComputeShader, uint32 BufferIndex, uint32 BaseIndex, uint32 NumBytes, const void* NewValue) final override
-	{
-		ContextRedirect(RHISetShaderParameter(ComputeShader, BufferIndex, BaseIndex, NumBytes, NewValue));
-	}
 	FORCEINLINE virtual void RHISetShaderParameters(FRHIComputeShader* Shader, TConstArrayView<uint8> InParametersData, TConstArrayView<FRHIShaderParameter> InParameters, TConstArrayView<FRHIShaderParameterResource> InResourceParameters, TConstArrayView<FRHIShaderParameterResource> InBindlessParameters) final override
 	{
 		ContextRedirect(RHISetShaderParameters(Shader, InParametersData, InParameters, InResourceParameters, InBindlessParameters));
@@ -940,34 +908,14 @@ public:
 	{
 		ContextRedirect(RHISetScissorRect(bEnable, MinX, MinY, MaxX, MaxY));
 	}
-	FORCEINLINE void RHISetGraphicsPipelineState(FRHIGraphicsPipelineState* GraphicsPipelineState, uint32 StencilRef, bool bApplyAdditionalState) final override
+	FORCEINLINE virtual void RHISetGraphicsPipelineState(FRHIGraphicsPipelineState* GraphicsPipelineState, uint32 StencilRef, bool bApplyAdditionalState) final override
 	{
 		ContextRedirect(RHISetGraphicsPipelineState(GraphicsPipelineState, StencilRef, bApplyAdditionalState));
 	}
-	FORCEINLINE virtual void RHISetShaderTexture(FRHIGraphicsShader* Shader, uint32 TextureIndex, FRHITexture* NewTexture) final override
-	{
-		ContextRedirect(RHISetShaderTexture(Shader, TextureIndex, NewTexture));
-	}
-	FORCEINLINE virtual void RHISetShaderSampler(FRHIGraphicsShader* Shader, uint32 SamplerIndex, FRHISamplerState* NewState) final override
-	{
-		ContextRedirect(RHISetShaderSampler(Shader, SamplerIndex, NewState));
-	}
-	FORCEINLINE virtual void RHISetShaderResourceViewParameter(FRHIGraphicsShader* Shader, uint32 SamplerIndex, FRHIShaderResourceView* SRV) final override
-	{
-		ContextRedirect(RHISetShaderResourceViewParameter(Shader, SamplerIndex, SRV));
-	}
 	FORCEINLINE virtual void RHISetStaticUniformBuffers(const FUniformBufferStaticBindings& InUniformBuffers) final override
 	{
 		ContextRedirect(RHISetStaticUniformBuffers(InUniformBuffers));
 	}
-	FORCEINLINE virtual void RHISetShaderUniformBuffer(FRHIGraphicsShader* Shader, uint32 BufferIndex, FRHIUniformBuffer* Buffer) final override
-	{
-		ContextRedirect(RHISetShaderUniformBuffer(Shader, BufferIndex, Buffer));
-	}
-	FORCEINLINE virtual void RHISetShaderParameter(FRHIGraphicsShader* Shader, uint32 BufferIndex, uint32 BaseIndex, uint32 NumBytes, const void* NewValue) final override
-	{
-		ContextRedirect(RHISetShaderParameter(Shader, BufferIndex, BaseIndex, NumBytes, NewValue));
-	}
 	FORCEINLINE virtual void RHISetShaderParameters(FRHIGraphicsShader* Shader, TConstArrayView<uint8> InParametersData, TConstArrayView<FRHIShaderParameter> InParameters, TConstArrayView<FRHIShaderParameterResource> InResourceParameters, TConstArrayView<FRHIShaderParameterResource> InBindlessParameters) final override
 	{
 		ContextRedirect(RHISetShaderParameters(Shader, InParametersData, InParameters, InResourceParameters, InBindlessParameters));
diff --git a/Engine/Source/Runtime/NullDrv/Public/NullRHI.h b/Engine/Source/Runtime/NullDrv/Public/NullRHI.h
--- a/Engine/Source/Runtime/NullDrv/Public/NullRHI.h
+++ b/Engine/Source/Runtime/NullDrv/Public/NullRHI.h
@@ -439,66 +439,6 @@ public:
 	{
 	}
 
-
-	virtual void RHISetShaderTexture(FRHIGraphicsShader* Shader, uint32 TextureIndex, FRHITexture* NewTexture) final override
-	{
-	}
-
-	virtual void RHISetShaderTexture(FRHIComputeShader* PixelShader, uint32 TextureIndex, FRHITexture* NewTexture) final override
-	{
-
-	}
-
-	virtual void RHISetShaderSampler(FRHIComputeShader* ComputeShader, uint32 SamplerIndex, FRHISamplerState* NewState) final override
-	{
-
-	}
-
-	virtual void RHISetShaderSampler(FRHIGraphicsShader* Shader, uint32 SamplerIndex, FRHISamplerState* NewState) final override
-	{
-	}
-
-	virtual void RHISetUAVParameter(FRHIPixelShader* PixelShader, uint32 UAVIndex, FRHIUnorderedAccessView* UAV) final override
-	{
-	}
-
-	virtual void RHISetUAVParameter(FRHIComputeShader* ComputeShader, uint32 UAVIndex, FRHIUnorderedAccessView* UAV) final override
-	{
-
-	}
-
-
-	virtual void RHISetUAVParameter(FRHIComputeShader* ComputeShader, uint32 UAVIndex, FRHIUnorderedAccessView* UAV, uint32 InitialCount) final override
-	{
-
-	}
-
-
-	virtual void RHISetShaderResourceViewParameter(FRHIGraphicsShader* Shader, uint32 SamplerIndex, FRHIShaderResourceView* SRV) final override
-	{
-	}
-
-	virtual void RHISetShaderResourceViewParameter(FRHIComputeShader* ComputeShader, uint32 SamplerIndex, FRHIShaderResourceView* SRV) final override
-	{
-	}
-
-
-	virtual void RHISetShaderUniformBuffer(FRHIGraphicsShader* Shader, uint32 BufferIndex, FRHIUniformBuffer* Buffer) final override
-	{
-	}
-
-	virtual void RHISetShaderUniformBuffer(FRHIComputeShader* ComputeShader, uint32 BufferIndex, FRHIUniformBuffer* Buffer) final override
-	{
-	}
-
-	virtual void RHISetShaderParameter(FRHIGraphicsShader* Shader, uint32 BufferIndex, uint32 BaseIndex, uint32 NumBytes, const void* NewValue) final override
-	{
-	}
-
-	virtual void RHISetShaderParameter(FRHIComputeShader* ComputeShader, uint32 BufferIndex, uint32 BaseIndex, uint32 NumBytes, const void* NewValue) final override
-	{
-	}
-
 	virtual void RHISetShaderParameters(FRHIGraphicsShader* Shader, TConstArrayView<uint8> InParametersData, TConstArrayView<FRHIShaderParameter> InParameters, TConstArrayView<FRHIShaderParameterResource> InResourceParameters, TConstArrayView<FRHIShaderParameterResource> InBindlessParameters) final override
 	{
 	}
diff --git a/Engine/Source/Runtime/OpenGLDrv/Public/OpenGLDrv.h b/Engine/Source/Runtime/OpenGLDrv/Public/OpenGLDrv.h
--- a/Engine/Source/Runtime/OpenGLDrv/Public/OpenGLDrv.h
+++ b/Engine/Source/Runtime/OpenGLDrv/Public/OpenGLDrv.h
@@ -442,22 +442,22 @@ public:
 	virtual void RHISetViewport(float MinX, float MinY, float MinZ, float MaxX, float MaxY, float MaxZ) final override;
 	virtual void RHISetScissorRect(bool bEnable, uint32 MinX, uint32 MinY, uint32 MaxX, uint32 MaxY) final override;
 	virtual void RHISetBoundShaderState(FRHIBoundShaderState* BoundShaderState) final override;
-	virtual void RHISetShaderTexture(FRHIGraphicsShader* Shader, uint32 TextureIndex, FRHITexture* NewTexture) final override;
-	virtual void RHISetShaderTexture(FRHIComputeShader* PixelShader, uint32 TextureIndex, FRHITexture* NewTexture) final override;
-	virtual void RHISetShaderSampler(FRHIComputeShader* ComputeShader, uint32 SamplerIndex, FRHISamplerState* NewState) final override;
-	virtual void RHISetShaderSampler(FRHIGraphicsShader* Shader, uint32 SamplerIndex, FRHISamplerState* NewState) final override;
-	virtual void RHISetUAVParameter(FRHIPixelShader* PixelShader, uint32 UAVIndex, FRHIUnorderedAccessView* UAV) final override;
-	virtual void RHISetUAVParameter(FRHIComputeShader* ComputeShader, uint32 UAVIndex, FRHIUnorderedAccessView* UAV) final override;
-	virtual void RHISetUAVParameter(FRHIComputeShader* ComputeShader, uint32 UAVIndex, FRHIUnorderedAccessView* UAV, uint32 InitialCount) final override;
-	virtual void RHISetShaderResourceViewParameter(FRHIGraphicsShader* Shader, uint32 SamplerIndex, FRHIShaderResourceView* SRV) final override;
-	virtual void RHISetShaderResourceViewParameter(FRHIComputeShader* ComputeShader, uint32 SamplerIndex, FRHIShaderResourceView* SRV) final override;
+	void RHISetShaderTexture(FRHIGraphicsShader* Shader, uint32 TextureIndex, FRHITexture* NewTexture);
+	void RHISetShaderTexture(FRHIComputeShader* PixelShader, uint32 TextureIndex, FRHITexture* NewTexture);
+	void RHISetShaderSampler(FRHIComputeShader* ComputeShader, uint32 SamplerIndex, FRHISamplerState* NewState);
+	void RHISetShaderSampler(FRHIGraphicsShader* Shader, uint32 SamplerIndex, FRHISamplerState* NewState);
+	void RHISetUAVParameter(FRHIPixelShader* PixelShader, uint32 UAVIndex, FRHIUnorderedAccessView* UAV);
+	void RHISetUAVParameter(FRHIComputeShader* ComputeShader, uint32 UAVIndex, FRHIUnorderedAccessView* UAV);
+	void RHISetUAVParameter(FRHIComputeShader* ComputeShader, uint32 UAVIndex, FRHIUnorderedAccessView* UAV, uint32 InitialCount);
+	void RHISetShaderResourceViewParameter(FRHIGraphicsShader* Shader, uint32 SamplerIndex, FRHIShaderResourceView* SRV);
+	void RHISetShaderResourceViewParameter(FRHIComputeShader* ComputeShader, uint32 SamplerIndex, FRHIShaderResourceView* SRV);
 	virtual void RHISetStaticUniformBuffers(const FUniformBufferStaticBindings& InUniformBuffers) final override;
 	virtual void RHISetStaticUniformBuffer(FUniformBufferStaticSlot Slot, FRHIUniformBuffer* Buffer) final override;
 	virtual void RHISetUniformBufferDynamicOffset(FUniformBufferStaticSlot Slot, uint32 Offset) final override;
-	virtual void RHISetShaderUniformBuffer(FRHIGraphicsShader* Shader, uint32 BufferIndex, FRHIUniformBuffer* Buffer) final override;
-	virtual void RHISetShaderUniformBuffer(FRHIComputeShader* ComputeShader, uint32 BufferIndex, FRHIUniformBuffer* Buffer) final override;
-	virtual void RHISetShaderParameter(FRHIGraphicsShader* Shader, uint32 BufferIndex, uint32 BaseIndex, uint32 NumBytes, const void* NewValue) final override;
-	virtual void RHISetShaderParameter(FRHIComputeShader* ComputeShader, uint32 BufferIndex, uint32 BaseIndex, uint32 NumBytes, const void* NewValue) final override;
+	void RHISetShaderUniformBuffer(FRHIGraphicsShader* Shader, uint32 BufferIndex, FRHIUniformBuffer* Buffer);
+	void RHISetShaderUniformBuffer(FRHIComputeShader* ComputeShader, uint32 BufferIndex, FRHIUniformBuffer* Buffer);
+	void RHISetShaderParameter(FRHIGraphicsShader* Shader, uint32 BufferIndex, uint32 BaseIndex, uint32 NumBytes, const void* NewValue);
+	void RHISetShaderParameter(FRHIComputeShader* ComputeShader, uint32 BufferIndex, uint32 BaseIndex, uint32 NumBytes, const void* NewValue);
 	virtual void RHISetShaderParameters(FRHIGraphicsShader* Shader, TConstArrayView<uint8> InParametersData, TConstArrayView<FRHIShaderParameter> InParameters, TConstArrayView<FRHIShaderParameterResource> InResourceParameters, TConstArrayView<FRHIShaderParameterResource> InBindlessParameters) final override;
 	virtual void RHISetShaderParameters(FRHIComputeShader* Shader, TConstArrayView<uint8> InParametersData, TConstArrayView<FRHIShaderParameter> InParameters, TConstArrayView<FRHIShaderParameterResource> InResourceParameters, TConstArrayView<FRHIShaderParameterResource> InBindlessParameters) final override;
 	virtual void RHISetShaderUnbinds(FRHIComputeShader* Shader, TConstArrayView<FRHIShaderParameterUnbind> InUnbinds) final override;
diff --git a/Engine/Source/Runtime/RHI/Public/RHIContext.h b/Engine/Source/Runtime/RHI/Public/RHIContext.h
--- a/Engine/Source/Runtime/RHI/Public/RHIContext.h
+++ b/Engine/Source/Runtime/RHI/Public/RHIContext.h
@@ -303,38 +303,6 @@ public:
 	virtual void RHIBeginUAVOverlap(TConstArrayView<FRHIUnorderedAccessView*> UAVs) {}
 	virtual void RHIEndUAVOverlap(TConstArrayView<FRHIUnorderedAccessView*> UAVs) {}
 
-	/** Set the shader resource view of a surface.  This is used for binding TextureMS parameter types that need a multi sampled view. */
-	virtual void RHISetShaderTexture(FRHIComputeShader* ComputeShader, uint32 TextureIndex, FRHITexture* NewTexture) = 0;
-
-	/**
-	* Sets sampler state.
-	* @param ComputeShader		The compute shader to set the sampler for.
-	* @param SamplerIndex		The index of the sampler.
-	* @param NewState			The new sampler state.
-	*/
-	virtual void RHISetShaderSampler(FRHIComputeShader* ComputeShader, uint32 SamplerIndex, FRHISamplerState* NewState) = 0;
-
-	/**
-	* Sets a compute shader UAV parameter.
-	* @param ComputeShader	The compute shader to set the UAV for.
-	* @param UAVIndex		The index of the UAVIndex.
-	* @param UAV			The new UAV.
-	*/
-	virtual void RHISetUAVParameter(FRHIComputeShader* ComputeShader, uint32 UAVIndex, FRHIUnorderedAccessView* UAV) = 0;
-
-	/**
-	* Sets a compute shader counted UAV parameter and initial count
-	* @param ComputeShader	The compute shader to set the UAV for.
-	* @param UAVIndex		The index of the UAVIndex.
-	* @param UAV			The new UAV.
-	* @param InitialCount	The initial number of items in the UAV.
-	*/
-	virtual void RHISetUAVParameter(FRHIComputeShader* ComputeShader, uint32 UAVIndex, FRHIUnorderedAccessView* UAV, uint32 InitialCount) = 0;
-
-	virtual void RHISetShaderResourceViewParameter(FRHIComputeShader* ComputeShader, uint32 SamplerIndex, FRHIShaderResourceView* SRV) = 0;
-
-	virtual void RHISetShaderUniformBuffer(FRHIComputeShader* ComputeShader, uint32 BufferIndex, FRHIUniformBuffer* Buffer) = 0;
-
 	virtual void RHISetShaderParameters(FRHIComputeShader* ComputeShader, TConstArrayView<uint8> InParametersData, TConstArrayView<FRHIShaderParameter> InParameters, TConstArrayView<FRHIShaderParameterResource> InResourceParameters, TConstArrayView<FRHIShaderParameterResource> InBindlessParameters) = 0;
 
 	void RHISetBatchedShaderParameters(FRHIComputeShader* InShader, FRHIBatchedShaderParameters& InBatchedParameters)
@@ -354,8 +322,6 @@ public:
 		checkf(false, TEXT("RHISetShaderUnbinds called when the active RHI hasn't overridden it and GRHIGlobals.NeedsShaderUnbinds is set."));
 	}
 
-	virtual void RHISetShaderParameter(FRHIComputeShader* ComputeShader, uint32 BufferIndex, uint32 BaseIndex, uint32 NumBytes, const void* NewValue) = 0;
-
 	virtual void RHISetStaticUniformBuffers(const FUniformBufferStaticBindings& InUniformBuffers)
 	{
 		/** empty default implementation. */
@@ -710,39 +676,9 @@ public:
 
 	// Inherit the parent context's RHISet functions that take FRHIComputeShader arguments
 	// Required to avoid warning C4263 : 'function' : member function does not override any base class virtual member function
-	using IRHIComputeContext::RHISetShaderTexture;
-	using IRHIComputeContext::RHISetShaderSampler;
-	using IRHIComputeContext::RHISetUAVParameter;
-	using IRHIComputeContext::RHISetShaderResourceViewParameter;
-	using IRHIComputeContext::RHISetShaderUniformBuffer;
 	using IRHIComputeContext::RHISetShaderParameters;
 	using IRHIComputeContext::RHISetBatchedShaderParameters;
 	using IRHIComputeContext::RHISetShaderUnbinds;
-	using IRHIComputeContext::RHISetShaderParameter;
-
-	/** Set the shader resource view of a surface. */
-	virtual void RHISetShaderTexture(FRHIGraphicsShader* Shader, uint32 TextureIndex, FRHITexture* NewTexture) = 0;
-
-	/**
-	* Sets sampler state.
-	* @param Shader				The shader to set the sampler for.
-	* @param SamplerIndex		The index of the sampler.
-	* @param NewState			The new sampler state.
-	*/
-	virtual void RHISetShaderSampler(FRHIGraphicsShader* Shader, uint32 SamplerIndex, FRHISamplerState* NewState) = 0;
-
-	/**
-	* Sets a pixel shader UAV parameter.
-	* @param PixelShader		The pixel shader to set the UAV for.
-	* @param UAVIndex		The index of the UAVIndex.
-	* @param UAV			The new UAV.
-	*/
-	virtual void RHISetUAVParameter(FRHIPixelShader* PixelShader, uint32 UAVIndex, FRHIUnorderedAccessView* UAV) = 0;
-
-
-	virtual void RHISetShaderResourceViewParameter(FRHIGraphicsShader* Shader, uint32 SamplerIndex, FRHIShaderResourceView* SRV) = 0;
-
-	virtual void RHISetShaderUniformBuffer(FRHIGraphicsShader* Shader, uint32 BufferIndex, FRHIUniformBuffer* Buffer) = 0;
 
 	virtual void RHISetShaderParameters(FRHIGraphicsShader* Shader, TConstArrayView<uint8> InParametersData, TConstArrayView<FRHIShaderParameter> InParameters, TConstArrayView<FRHIShaderParameterResource> InResourceParameters, TConstArrayView<FRHIShaderParameterResource> InBindlessParameters) = 0;
 
@@ -763,8 +699,6 @@ public:
 		checkf(false, TEXT("RHISetShaderUnbinds called when the active RHI hasn't overridden it and GRHIGlobals.NeedsShaderUnbinds is set."));
 	}
 
-	virtual void RHISetShaderParameter(FRHIGraphicsShader* Shader, uint32 BufferIndex, uint32 BaseIndex, uint32 NumBytes, const void* NewValue) = 0;
-
 	virtual void RHISetStencilRef(uint32 StencilRef) {}
 
 	virtual void RHISetBlendFactor(const FLinearColor& BlendFactor) {}
diff --git a/Engine/Source/Runtime/RHI/Public/RHIValidationContext.h b/Engine/Source/Runtime/RHI/Public/RHIValidationContext.h
--- a/Engine/Source/Runtime/RHI/Public/RHIValidationContext.h
+++ b/Engine/Source/Runtime/RHI/Public/RHIValidationContext.h
@@ -260,76 +260,6 @@ public:
 		RHIContext->RHISubmitCommandsHint();
 	}
 
-	/** Set the shader resource view of a surface.  This is used for binding TextureMS parameter types that need a multi sampled view. */
-	virtual void RHISetShaderTexture(FRHIComputeShader* Shader, uint32 TextureIndex, FRHITexture* NewTexture) override final
-	{
-		checkf(State.bComputePSOSet, TEXT("A Compute PSO has to be set to set resources into a shader!"));
-		Tracker->Assert(NewTexture->GetWholeResourceIdentitySRV(), ERHIAccess::SRVCompute);
-		RHIContext->RHISetShaderTexture(Shader, TextureIndex, NewTexture);
-	}
-
-	/**
-	* Sets sampler state.
-	* @param Shader				The compute shader to set the sampler for.
-	* @param SamplerIndex		The index of the sampler.
-	* @param NewState			The new sampler state.
-	*/
-	virtual void RHISetShaderSampler(FRHIComputeShader* Shader, uint32 SamplerIndex, FRHISamplerState* NewState) override final
-	{
-		checkf(State.bComputePSOSet, TEXT("A Compute PSO has to be set to set resources into a shader!"));
-		RHIContext->RHISetShaderSampler(Shader, SamplerIndex, NewState);
-	}
-
-	/**
-	* Sets a compute shader UAV parameter.
-	* @param Shader	The compute shader to set the UAV for.
-	* @param UAVIndex		The index of the UAVIndex.
-	* @param UAV			The new UAV.
-	*/
-	virtual void RHISetUAVParameter(FRHIComputeShader* Shader, uint32 UAVIndex, FRHIUnorderedAccessView* UAV) override final
-	{
-		checkf(State.bComputePSOSet, TEXT("A Compute PSO has to be set to set resources into a shader!"));
-		Tracker->AssertUAV(UAV, RHIValidation::EUAVMode::Compute, UAVIndex);
-		RHIContext->RHISetUAVParameter(Shader, UAVIndex, UAV);
-	}
-
-	/**
-	* Sets a compute shader counted UAV parameter and initial count
-	* @param Shader	The compute shader to set the UAV for.
-	* @param UAVIndex		The index of the UAVIndex.
-	* @param UAV			The new UAV.
-	* @param InitialCount	The initial number of items in the UAV.
-	*/
-	virtual void RHISetUAVParameter(FRHIComputeShader* Shader, uint32 UAVIndex, FRHIUnorderedAccessView* UAV, uint32 InitialCount) override final
-	{
-		checkf(State.bComputePSOSet, TEXT("A Compute PSO has to be set to set resources into a shader!"));
-		Tracker->AssertUAV(UAV, RHIValidation::EUAVMode::Compute, UAVIndex);
-		RHIContext->RHISetUAVParameter(Shader, UAVIndex, UAV, InitialCount);
-	}
-
-	virtual void RHISetShaderResourceViewParameter(FRHIComputeShader* Shader, uint32 SamplerIndex, FRHIShaderResourceView* SRV) override final
-	{
-		checkf(State.bComputePSOSet, TEXT("A Compute PSO has to be set to set resources into a shader!"));
-		if (SRV) 
-		{
-			Tracker->Assert(SRV->GetViewIdentity(), ERHIAccess::SRVCompute);
-		}
-		RHIContext->RHISetShaderResourceViewParameter(Shader, SamplerIndex, SRV);
-	}
-
-	virtual void RHISetShaderUniformBuffer(FRHIComputeShader* Shader, uint32 BufferIndex, FRHIUniformBuffer* Buffer) override final
-	{
-		checkf(State.bComputePSOSet, TEXT("A Compute PSO has to be set to set resources into a shader!"));
-		State.StaticUniformBuffers.ValidateSetShaderUniformBuffer(Buffer);
-		RHIContext->RHISetShaderUniformBuffer(Shader, BufferIndex, Buffer);
-	}
-
-	virtual void RHISetShaderParameter(FRHIComputeShader* Shader, uint32 BufferIndex, uint32 BaseIndex, uint32 NumBytes, const void* NewValue) override final
-	{
-		checkf(State.bComputePSOSet, TEXT("A Compute PSO has to be set to set resources into a shader!"));
-		RHIContext->RHISetShaderParameter(Shader, BufferIndex, BaseIndex, NumBytes, NewValue);
-	}
-
 	virtual void RHISetShaderParameters(FRHIComputeShader* Shader, TConstArrayView<uint8> InParametersData, TConstArrayView<FRHIShaderParameter> InParameters, TConstArrayView<FRHIShaderParameterResource> InResourceParameters, TConstArrayView<FRHIShaderParameterResource> InBindlessParameters) final override
 	{
 		checkf(State.bComputePSOSet, TEXT("A Compute PSO has to be set to set resources into a shader!"));
@@ -815,135 +745,6 @@ public:
 	}
 #endif
 
-	/** Set the shader resource view of a surface.  This is used for binding TextureMS parameter types that need a multi sampled view. */
-	virtual void RHISetShaderTexture(FRHIGraphicsShader* Shader, uint32 TextureIndex, FRHITexture* NewTexture) override final
-	{
-		checkf(State.bGfxPSOSet, TEXT("A Graphics PSO has to be set to set resources into a shader!"));
-		Tracker->Assert(NewTexture->GetWholeResourceIdentitySRV(), ERHIAccess::SRVGraphics);
-		RHIContext->RHISetShaderTexture(Shader, TextureIndex, NewTexture);
-	}
-
-	/** Set the shader resource view of a surface.  This is used for binding TextureMS parameter types that need a multi sampled view. */
-	virtual void RHISetShaderTexture(FRHIComputeShader* Shader, uint32 TextureIndex, FRHITexture* NewTexture) override final
-	{
-		checkf(State.bComputePSOSet, TEXT("A Compute PSO has to be set to set resources into a shader!"));
-		Tracker->Assert(NewTexture->GetWholeResourceIdentitySRV(), ERHIAccess::SRVCompute);
-		RHIContext->RHISetShaderTexture(Shader, TextureIndex, NewTexture);
-	}
-
-	/**
-	* Sets sampler state.
-	* @param Shader				The compute shader to set the sampler for.
-	* @param SamplerIndex		The index of the sampler.
-	* @param NewState			The new sampler state.
-	*/
-	virtual void RHISetShaderSampler(FRHIComputeShader* Shader, uint32 SamplerIndex, FRHISamplerState* NewState) override final
-	{
-		checkf(State.bComputePSOSet, TEXT("A Compute PSO has to be set to set resources into a shader!"));
-		RHIContext->RHISetShaderSampler(Shader, SamplerIndex, NewState);
-	}
-
-	/**
-	* Sets sampler state.
-	* @param Shader				The graphics shader to set the sampler for.
-	* @param SamplerIndex		The index of the sampler.
-	* @param NewState			The new sampler state.
-	*/
-	virtual void RHISetShaderSampler(FRHIGraphicsShader* Shader, uint32 SamplerIndex, FRHISamplerState* NewState) override final
-	{
-		checkf(State.bGfxPSOSet, TEXT("A Graphics PSO has to be set to set resources into a shader!"));
-		RHIContext->RHISetShaderSampler(Shader, SamplerIndex, NewState);
-	}
-
-	/**
-	* Sets a pixel shader UAV parameter.
-	* @param Shader	The pixel shader to set the UAV for.
-	* @param UAVIndex		The index of the UAVIndex.
-	* @param UAV			The new UAV.
-	*/
-	virtual void RHISetUAVParameter(FRHIPixelShader* Shader, uint32 UAVIndex, FRHIUnorderedAccessView* UAV) override final
-	{
-		checkf(State.bGfxPSOSet, TEXT("A Compute PSO has to be set to set resources into a shader!"));
-		Tracker->AssertUAV(UAV, RHIValidation::EUAVMode::Graphics, UAVIndex);
-		RHIContext->RHISetUAVParameter(Shader, UAVIndex, UAV);
-	}
-
-
-
-
-	/**
-	* Sets a compute shader UAV parameter.
-	* @param Shader	The compute shader to set the UAV for.
-	* @param UAVIndex		The index of the UAVIndex.
-	* @param UAV			The new UAV.
-	*/
-	virtual void RHISetUAVParameter(FRHIComputeShader* Shader, uint32 UAVIndex, FRHIUnorderedAccessView* UAV) override final
-	{
-		checkf(State.bComputePSOSet, TEXT("A Compute PSO has to be set to set resources into a shader!"));
-		Tracker->AssertUAV(UAV, RHIValidation::EUAVMode::Compute, UAVIndex);
-		RHIContext->RHISetUAVParameter(Shader, UAVIndex, UAV);
-	}
-
-	/**
-	* Sets a compute shader counted UAV parameter and initial count
-	* @param Shader	The compute shader to set the UAV for.
-	* @param UAVIndex		The index of the UAVIndex.
-	* @param UAV			The new UAV.
-	* @param InitialCount	The initial number of items in the UAV.
-	*/
-	virtual void RHISetUAVParameter(FRHIComputeShader* Shader, uint32 UAVIndex, FRHIUnorderedAccessView* UAV, uint32 InitialCount) override final
-	{
-		checkf(State.bComputePSOSet, TEXT("A Compute PSO has to be set to set resources into a shader!"));
-		Tracker->AssertUAV(UAV, RHIValidation::EUAVMode::Compute, UAVIndex);
-		RHIContext->RHISetUAVParameter(Shader, UAVIndex, UAV, InitialCount);
-	}
-
-	virtual void RHISetShaderResourceViewParameter(FRHIGraphicsShader* Shader, uint32 SamplerIndex, FRHIShaderResourceView* SRV) override final
-	{
-		checkf(State.bGfxPSOSet, TEXT("A Graphics PSO has to be set to set resources into a shader!"));
-		if (SRV)
-		{
-			Tracker->Assert(SRV->GetViewIdentity(), ERHIAccess::SRVGraphics);
-		}
-		RHIContext->RHISetShaderResourceViewParameter(Shader, SamplerIndex, SRV);
-	}
-
-	virtual void RHISetShaderResourceViewParameter(FRHIComputeShader* Shader, uint32 SamplerIndex, FRHIShaderResourceView* SRV) override final
-	{
-		checkf(State.bComputePSOSet, TEXT("A Compute PSO has to be set to set resources into a shader!"));
-		if (SRV)
-		{
-			Tracker->Assert(SRV->GetViewIdentity(), ERHIAccess::SRVCompute);
-		}
-		RHIContext->RHISetShaderResourceViewParameter(Shader, SamplerIndex, SRV);
-	}
-
-	virtual void RHISetShaderUniformBuffer(FRHIGraphicsShader* Shader, uint32 BufferIndex, FRHIUniformBuffer* Buffer) override final
-	{
-		checkf(State.bGfxPSOSet, TEXT("A Graphics PSO has to be set to set resources into a shader!"));
-		State.StaticUniformBuffers.ValidateSetShaderUniformBuffer(Buffer);
-		RHIContext->RHISetShaderUniformBuffer(Shader, BufferIndex, Buffer);
-	}
-
-	virtual void RHISetShaderUniformBuffer(FRHIComputeShader* Shader, uint32 BufferIndex, FRHIUniformBuffer* Buffer) override final
-	{
-		checkf(State.bComputePSOSet, TEXT("A Compute PSO has to be set to set resources into a shader!"));
-		State.StaticUniformBuffers.ValidateSetShaderUniformBuffer(Buffer);
-		RHIContext->RHISetShaderUniformBuffer(Shader, BufferIndex, Buffer);
-	}
-
-	virtual void RHISetShaderParameter(FRHIGraphicsShader* Shader, uint32 BufferIndex, uint32 BaseIndex, uint32 NumBytes, const void* NewValue) override final
-	{
-		checkf(State.bGfxPSOSet, TEXT("A Graphics PSO has to be set to set resources into a shader!"));
-		RHIContext->RHISetShaderParameter(Shader, BufferIndex, BaseIndex, NumBytes, NewValue);
-	}
-
-	virtual void RHISetShaderParameter(FRHIComputeShader* Shader, uint32 BufferIndex, uint32 BaseIndex, uint32 NumBytes, const void* NewValue) override final
-	{
-		checkf(State.bComputePSOSet, TEXT("A Compute PSO has to be set to set resources into a shader!"));
-		RHIContext->RHISetShaderParameter(Shader, BufferIndex, BaseIndex, NumBytes, NewValue);
-	}
-
 	virtual void RHISetShaderParameters(FRHIGraphicsShader* Shader, TConstArrayView<uint8> InParametersData, TConstArrayView<FRHIShaderParameter> InParameters, TConstArrayView<FRHIShaderParameterResource> InResourceParameters, TConstArrayView<FRHIShaderParameterResource> InBindlessParameters) final override
 	{
 		checkf(State.bGfxPSOSet, TEXT("A Graphics PSO has to be set to set resources into a shader!"));
diff --git a/Engine/Source/Runtime/VulkanRHI/Private/VulkanContext.h b/Engine/Source/Runtime/VulkanRHI/Private/VulkanContext.h
--- a/Engine/Source/Runtime/VulkanRHI/Private/VulkanContext.h
+++ b/Engine/Source/Runtime/VulkanRHI/Private/VulkanContext.h
@@ -44,22 +44,22 @@ public:
 	virtual void RHISetStereoViewport(float LeftMinX, float RightMinX, float LeftMinY, float RightMinY, float MinZ, float LeftMaxX, float RightMaxX, float LeftMaxY, float RightMaxY, float MaxZ) override;
 	virtual void RHISetScissorRect(bool bEnable, uint32 MinX, uint32 MinY, uint32 MaxX, uint32 MaxY) final override;
 	virtual void RHISetGraphicsPipelineState(FRHIGraphicsPipelineState* GraphicsState, uint32 StencilRef, bool bApplyAdditionalState) final override;
-	virtual void RHISetShaderTexture(FRHIGraphicsShader* Shader, uint32 TextureIndex, FRHITexture* NewTexture) final override;
-	virtual void RHISetShaderTexture(FRHIComputeShader* PixelShader, uint32 TextureIndex, FRHITexture* NewTexture) final override;
-	virtual void RHISetShaderSampler(FRHIComputeShader* ComputeShader, uint32 SamplerIndex, FRHISamplerState* NewState) final override;
-	virtual void RHISetShaderSampler(FRHIGraphicsShader* Shader, uint32 SamplerIndex, FRHISamplerState* NewState) final override;
-	virtual void RHISetUAVParameter(FRHIPixelShader* PixelShader, uint32 UAVIndex, FRHIUnorderedAccessView* UAV) final override;
-	virtual void RHISetUAVParameter(FRHIComputeShader* ComputeShader, uint32 UAVIndex, FRHIUnorderedAccessView* UAV) final override;
-	virtual void RHISetUAVParameter(FRHIComputeShader* ComputeShader, uint32 UAVIndex, FRHIUnorderedAccessView* UAV, uint32 InitialCount) final override;
-	virtual void RHISetShaderResourceViewParameter(FRHIGraphicsShader* Shader, uint32 SamplerIndex, FRHIShaderResourceView* SRV) final override;
-	virtual void RHISetShaderResourceViewParameter(FRHIComputeShader* ComputeShader, uint32 SamplerIndex, FRHIShaderResourceView* SRV) final override;
+	void RHISetShaderTexture(FRHIGraphicsShader* Shader, uint32 TextureIndex, FRHITexture* NewTexture);
+	void RHISetShaderTexture(FRHIComputeShader* PixelShader, uint32 TextureIndex, FRHITexture* NewTexture);
+	void RHISetShaderSampler(FRHIComputeShader* ComputeShader, uint32 SamplerIndex, FRHISamplerState* NewState);
+	void RHISetShaderSampler(FRHIGraphicsShader* Shader, uint32 SamplerIndex, FRHISamplerState* NewState);
+	void RHISetUAVParameter(FRHIPixelShader* PixelShader, uint32 UAVIndex, FRHIUnorderedAccessView* UAV);
+	void RHISetUAVParameter(FRHIComputeShader* ComputeShader, uint32 UAVIndex, FRHIUnorderedAccessView* UAV);
+	void RHISetUAVParameter(FRHIComputeShader* ComputeShader, uint32 UAVIndex, FRHIUnorderedAccessView* UAV, uint32 InitialCount);
+	void RHISetShaderResourceViewParameter(FRHIGraphicsShader* Shader, uint32 SamplerIndex, FRHIShaderResourceView* SRV);
+	void RHISetShaderResourceViewParameter(FRHIComputeShader* ComputeShader, uint32 SamplerIndex, FRHIShaderResourceView* SRV);
 	virtual void RHISetStaticUniformBuffers(const FUniformBufferStaticBindings& InUniformBuffers) final override;
 	virtual void RHISetStaticUniformBuffer(FUniformBufferStaticSlot Slot, FRHIUniformBuffer* Buffer) final override;
 	virtual void RHISetUniformBufferDynamicOffset(FUniformBufferStaticSlot Slot, uint32 InOffset) final override;
-	virtual void RHISetShaderUniformBuffer(FRHIGraphicsShader* Shader, uint32 BufferIndex, FRHIUniformBuffer* Buffer) final override;
-	virtual void RHISetShaderUniformBuffer(FRHIComputeShader* ComputeShader, uint32 BufferIndex, FRHIUniformBuffer* Buffer) final override;
-	virtual void RHISetShaderParameter(FRHIGraphicsShader* Shader, uint32 BufferIndex, uint32 BaseIndex, uint32 NumBytes, const void* NewValue) final override;
-	virtual void RHISetShaderParameter(FRHIComputeShader* ComputeShader, uint32 BufferIndex, uint32 BaseIndex, uint32 NumBytes, const void* NewValue) final override;
+	void RHISetShaderUniformBuffer(FRHIGraphicsShader* Shader, uint32 BufferIndex, FRHIUniformBuffer* Buffer);
+	void RHISetShaderUniformBuffer(FRHIComputeShader* ComputeShader, uint32 BufferIndex, FRHIUniformBuffer* Buffer);
+	void RHISetShaderParameter(FRHIGraphicsShader* Shader, uint32 BufferIndex, uint32 BaseIndex, uint32 NumBytes, const void* NewValue);
+	void RHISetShaderParameter(FRHIComputeShader* ComputeShader, uint32 BufferIndex, uint32 BaseIndex, uint32 NumBytes, const void* NewValue);
 	virtual void RHISetShaderParameters(FRHIGraphicsShader* Shader, TConstArrayView<uint8> InParametersData, TConstArrayView<FRHIShaderParameter> InParameters, TConstArrayView<FRHIShaderParameterResource> InResourceParameters, TConstArrayView<FRHIShaderParameterResource> InBindlessParameters) final override;
 	virtual void RHISetShaderParameters(FRHIComputeShader* Shader, TConstArrayView<uint8> InParametersData, TConstArrayView<FRHIShaderParameter> InParameters, TConstArrayView<FRHIShaderParameterResource> InResourceParameters, TConstArrayView<FRHIShaderParameterResource> InBindlessParameters) final override;
 	virtual void RHISetStencilRef(uint32 StencilRef) final override;
diff --git a/Engine/Source/Runtime/Windows/D3D11RHI/Private/D3D11RHIPrivate.h b/Engine/Source/Runtime/Windows/D3D11RHI/Private/D3D11RHIPrivate.h
--- a/Engine/Source/Runtime/Windows/D3D11RHI/Private/D3D11RHIPrivate.h
+++ b/Engine/Source/Runtime/Windows/D3D11RHI/Private/D3D11RHIPrivate.h
@@ -586,20 +586,20 @@ public:
 	virtual void RHISetScissorRect(bool bEnable, uint32 MinX, uint32 MinY, uint32 MaxX, uint32 MaxY) final override;
 	virtual void RHISetBoundShaderState(FRHIBoundShaderState* BoundShaderState) final override;
 	virtual void RHISetGraphicsPipelineState(FRHIGraphicsPipelineState* GraphicsState, uint32 StencilRef, bool bApplyAdditionalState) final override;
-	virtual void RHISetShaderTexture(FRHIGraphicsShader* Shader, uint32 TextureIndex, FRHITexture* NewTexture) final override;
-	virtual void RHISetShaderTexture(FRHIComputeShader* PixelShader, uint32 TextureIndex, FRHITexture* NewTexture) final override;
-	virtual void RHISetShaderSampler(FRHIComputeShader* ComputeShader, uint32 SamplerIndex, FRHISamplerState* NewState) final override;
-	virtual void RHISetShaderSampler(FRHIGraphicsShader* Shader, uint32 SamplerIndex, FRHISamplerState* NewState) final override;
-	virtual void RHISetUAVParameter(FRHIPixelShader* PixelShader, uint32 UAVIndex, FRHIUnorderedAccessView* UAVRHI) final override;
-	virtual void RHISetUAVParameter(FRHIComputeShader* ComputeShader, uint32 UAVIndex, FRHIUnorderedAccessView* UAV) final override;
-	virtual void RHISetUAVParameter(FRHIComputeShader* ComputeShader, uint32 UAVIndex, FRHIUnorderedAccessView* UAV, uint32 InitialCount) final override;
-	virtual void RHISetShaderResourceViewParameter(FRHIGraphicsShader* Shader, uint32 SamplerIndex, FRHIShaderResourceView* SRV) final override;
-	virtual void RHISetShaderResourceViewParameter(FRHIComputeShader* ComputeShader, uint32 SamplerIndex, FRHIShaderResourceView* SRV) final override;
+	void RHISetShaderTexture(FRHIGraphicsShader* Shader, uint32 TextureIndex, FRHITexture* NewTexture);
+	void RHISetShaderTexture(FRHIComputeShader* PixelShader, uint32 TextureIndex, FRHITexture* NewTexture);
+	void RHISetShaderSampler(FRHIComputeShader* ComputeShader, uint32 SamplerIndex, FRHISamplerState* NewState);
+	void RHISetShaderSampler(FRHIGraphicsShader* Shader, uint32 SamplerIndex, FRHISamplerState* NewState);
+	void RHISetUAVParameter(FRHIPixelShader* PixelShader, uint32 UAVIndex, FRHIUnorderedAccessView* UAVRHI);
+	void RHISetUAVParameter(FRHIComputeShader* ComputeShader, uint32 UAVIndex, FRHIUnorderedAccessView* UAV);
+	void RHISetUAVParameter(FRHIComputeShader* ComputeShader, uint32 UAVIndex, FRHIUnorderedAccessView* UAV, uint32 InitialCount);
+	void RHISetShaderResourceViewParameter(FRHIGraphicsShader* Shader, uint32 SamplerIndex, FRHIShaderResourceView* SRV);
+	void RHISetShaderResourceViewParameter(FRHIComputeShader* ComputeShader, uint32 SamplerIndex, FRHIShaderResourceView* SRV);
 	virtual void RHISetStaticUniformBuffers(const FUniformBufferStaticBindings& InUniformBuffers) final override;
-	virtual void RHISetShaderUniformBuffer(FRHIGraphicsShader* Shader, uint32 BufferIndex, FRHIUniformBuffer* Buffer) final override;
-	virtual void RHISetShaderUniformBuffer(FRHIComputeShader* ComputeShader, uint32 BufferIndex, FRHIUniformBuffer* Buffer) final override;
-	virtual void RHISetShaderParameter(FRHIGraphicsShader* Shader, uint32 BufferIndex, uint32 BaseIndex, uint32 NumBytes, const void* NewValue) final override;
-	virtual void RHISetShaderParameter(FRHIComputeShader* ComputeShader, uint32 BufferIndex, uint32 BaseIndex, uint32 NumBytes, const void* NewValue) final override;
+	void RHISetShaderUniformBuffer(FRHIGraphicsShader* Shader, uint32 BufferIndex, FRHIUniformBuffer* Buffer);
+	void RHISetShaderUniformBuffer(FRHIComputeShader* ComputeShader, uint32 BufferIndex, FRHIUniformBuffer* Buffer);
+	void RHISetShaderParameter(FRHIGraphicsShader* Shader, uint32 BufferIndex, uint32 BaseIndex, uint32 NumBytes, const void* NewValue);
+	void RHISetShaderParameter(FRHIComputeShader* ComputeShader, uint32 BufferIndex, uint32 BaseIndex, uint32 NumBytes, const void* NewValue);
 	virtual void RHISetShaderParameters(FRHIComputeShader* Shader, TConstArrayView<uint8> InParametersData, TConstArrayView<FRHIShaderParameter> InParameters, TConstArrayView<FRHIShaderParameterResource> InResourceParameters, TConstArrayView<FRHIShaderParameterResource> InBindlessParameters) final override;
 	virtual void RHISetShaderParameters(FRHIGraphicsShader* Shader, TConstArrayView<uint8> InParametersData, TConstArrayView<FRHIShaderParameter> InParameters, TConstArrayView<FRHIShaderParameterResource> InResourceParameters, TConstArrayView<FRHIShaderParameterResource> InBindlessParameters) final override;
 	virtual void RHISetShaderUnbinds(FRHIComputeShader* Shader, TConstArrayView<FRHIShaderParameterUnbind> InUnbinds) final override;
