Make sure we use running platform bSupportsVertexShaderSRVs value in a platform preview modes

diff --git a/Engine/Source/Runtime/RHI/Private/DataDrivenShaderPlatformInfo.cpp b/Engine/Source/Runtime/RHI/Private/DataDrivenShaderPlatformInfo.cpp
--- a/Engine/Source/Runtime/RHI/Private/DataDrivenShaderPlatformInfo.cpp
+++ b/Engine/Source/Runtime/RHI/Private/DataDrivenShaderPlatformInfo.cpp
@@ -433,6 +433,7 @@ void FGenericDataDrivenShaderPlatformInfo::UpdatePreviewPlatforms()
 					Infos[ShaderPlatform].bSupportsInstancedStereo &= Infos[EditorSPForPreviewMaxFeatureLevel].bSupportsInstancedStereo;
 				}
 				Infos[ShaderPlatform].bSupportsMobileMultiView &= Infos[EditorSPForPreviewMaxFeatureLevel].bSupportsMobileMultiView;
+				Infos[ShaderPlatform].bSupportsVertexShaderSRVs = Infos[EditorSPForPreviewMaxFeatureLevel].bSupportsVertexShaderSRVs;
 				Infos[ShaderPlatform].bSupportsManualVertexFetch = Infos[EditorSPForPreviewMaxFeatureLevel].bSupportsManualVertexFetch;
 				Infos[ShaderPlatform].bSupportsShaderRootConstants = false;
 				Infos[ShaderPlatform].bSupportsShaderBundleDispatch = false;
