Fixed RHIBindDebugLabelName to take a command list.

diff --git a/Engine/Source/Runtime/Apple/MetalRHI/Private/MetalDynamicRHI.h b/Engine/Source/Runtime/Apple/MetalRHI/Private/MetalDynamicRHI.h
--- a/Engine/Source/Runtime/Apple/MetalRHI/Private/MetalDynamicRHI.h
+++ b/Engine/Source/Runtime/Apple/MetalRHI/Private/MetalDynamicRHI.h
@@ -86,7 +86,7 @@ public:
 	virtual void RHIUpdateTexture3D(FRHICommandListBase& RHICmdList, FRHITexture3D* Texture, uint32 MipIndex, const struct FUpdateTextureRegion3D& UpdateRegion, uint32 SourceRowPitch, uint32 SourceDepthPitch, const uint8* SourceData) final override;
 	virtual void* RHILockTextureCubeFace(FRHITextureCube* Texture, uint32 FaceIndex, uint32 ArrayIndex, uint32 MipIndex, EResourceLockMode LockMode, uint32& DestStride, bool bLockWithinMiptail) final override;
 	virtual void RHIUnlockTextureCubeFace(FRHITextureCube* Texture, uint32 FaceIndex, uint32 ArrayIndex, uint32 MipIndex, bool bLockWithinMiptail) final override;
-	virtual void RHIBindDebugLabelName(FRHITexture* Texture, const TCHAR* Name) final override;
+	virtual void RHIBindDebugLabelName(FRHICommandListBase& RHICmdList, FRHITexture* Texture, const TCHAR* Name) final override;
 	virtual void RHIReadSurfaceData(FRHITexture* Texture,FIntRect Rect,TArray<FColor>& OutData,FReadSurfaceDataFlags InFlags) final override;
 	virtual void RHIReadSurfaceData(FRHITexture* TextureRHI, FIntRect InRect, TArray<FLinearColor>& OutData, FReadSurfaceDataFlags InFlags) final override;
 	virtual void RHIMapStagingSurface(FRHITexture* Texture, FRHIGPUFence* Fence, void*& OutData, int32& OutWidth, int32& OutHeight, uint32 GPUIndex = 0) final override;
diff --git a/Engine/Source/Runtime/Apple/MetalRHI/Private/MetalTexture.cpp b/Engine/Source/Runtime/Apple/MetalRHI/Private/MetalTexture.cpp
--- a/Engine/Source/Runtime/Apple/MetalRHI/Private/MetalTexture.cpp
+++ b/Engine/Source/Runtime/Apple/MetalRHI/Private/MetalTexture.cpp
@@ -1800,7 +1800,7 @@ void FMetalDynamicRHI::RHIUnlockTextureCubeFace(FRHITextureCube* TextureCubeRHI,
 	}
 }
 
-void FMetalDynamicRHI::RHIBindDebugLabelName(FRHITexture* TextureRHI, const TCHAR* Name)
+void FMetalDynamicRHI::RHIBindDebugLabelName(FRHICommandListBase& RHICmdList, FRHITexture* TextureRHI, const TCHAR* Name)
 {
 	@autoreleasepool {
 		FMetalSurface* Surf = GetMetalSurfaceFromRHITexture(TextureRHI);
diff --git a/Engine/Source/Runtime/D3D12RHI/Private/D3D12Buffer.cpp b/Engine/Source/Runtime/D3D12RHI/Private/D3D12Buffer.cpp
--- a/Engine/Source/Runtime/D3D12RHI/Private/D3D12Buffer.cpp
+++ b/Engine/Source/Runtime/D3D12RHI/Private/D3D12Buffer.cpp
@@ -860,7 +860,7 @@ void FD3D12DynamicRHI::RHICopyBuffer(FRHIBuffer* SourceBufferRHI, FRHIBuffer* De
 	}
 }
 
-void FD3D12DynamicRHI::RHIBindDebugLabelName(FRHIBuffer* BufferRHI, const TCHAR* Name)
+void FD3D12DynamicRHI::RHIBindDebugLabelName(FRHICommandListBase& RHICmdList, FRHIBuffer* BufferRHI, const TCHAR* Name)
 {
 	if (BufferRHI == nullptr || !GD3D12BindResourceLabels)
 	{
diff --git a/Engine/Source/Runtime/D3D12RHI/Private/D3D12RHIPrivate.h b/Engine/Source/Runtime/D3D12RHI/Private/D3D12RHIPrivate.h
--- a/Engine/Source/Runtime/D3D12RHI/Private/D3D12RHIPrivate.h
+++ b/Engine/Source/Runtime/D3D12RHI/Private/D3D12RHIPrivate.h
@@ -327,8 +327,8 @@ public:
 	virtual void RHIEndMultiUpdateTexture3D(FRHICommandListBase& RHICmdList, TArray<FUpdateTexture3DData>& UpdateDataArray) final override;
 	virtual void* RHILockTextureCubeFace(FRHITextureCube* Texture, uint32 FaceIndex, uint32 ArrayIndex, uint32 MipIndex, EResourceLockMode LockMode, uint32& DestStride, bool bLockWithinMiptail) final override;
 	virtual void RHIUnlockTextureCubeFace(FRHITextureCube* Texture, uint32 FaceIndex, uint32 ArrayIndex, uint32 MipIndex, bool bLockWithinMiptail) final override;
-	virtual void RHIBindDebugLabelName(FRHITexture* Texture, const TCHAR* Name) final override;
-	virtual void RHIBindDebugLabelName(FRHIBuffer* Buffer, const TCHAR* Name) final override;
+	virtual void RHIBindDebugLabelName(FRHICommandListBase& RHICmdList, FRHITexture* Texture, const TCHAR* Name) final override;
+	virtual void RHIBindDebugLabelName(FRHICommandListBase& RHICmdList, FRHIBuffer* Buffer, const TCHAR* Name) final override;
 	virtual void RHIReadSurfaceData(FRHITexture* Texture, FIntRect Rect, TArray<FColor>& OutData, FReadSurfaceDataFlags InFlags) final override;
 	virtual void RHIReadSurfaceData(FRHITexture* TextureRHI, FIntRect InRect, TArray<FLinearColor>& OutData, FReadSurfaceDataFlags InFlags) final override;
 	virtual void RHIMapStagingSurface(FRHITexture* Texture, FRHIGPUFence* Fence, void*& OutData, int32& OutWidth, int32& OutHeight, uint32 GPUIndex = 0) final override;
diff --git a/Engine/Source/Runtime/D3D12RHI/Private/D3D12Texture.cpp b/Engine/Source/Runtime/D3D12RHI/Private/D3D12Texture.cpp
--- a/Engine/Source/Runtime/D3D12RHI/Private/D3D12Texture.cpp
+++ b/Engine/Source/Runtime/D3D12RHI/Private/D3D12Texture.cpp
@@ -2678,7 +2678,7 @@ void FD3D12DynamicRHI::RHIUnlockTextureCubeFace(FRHITextureCube* TextureCubeRHI,
 	TextureCube->Unlock(nullptr, MipIndex, D3DFace + ArrayIndex * 6);
 }
 
-void FD3D12DynamicRHI::RHIBindDebugLabelName(FRHITexture* TextureRHI, const TCHAR* Name)
+void FD3D12DynamicRHI::RHIBindDebugLabelName(FRHICommandListBase& RHICmdList, FRHITexture* TextureRHI, const TCHAR* Name)
 {
 	if (!TextureRHI || !GD3D12BindResourceLabels)
 	{
diff --git a/Engine/Source/Runtime/Engine/Private/GPUSkinCache.cpp b/Engine/Source/Runtime/Engine/Private/GPUSkinCache.cpp
--- a/Engine/Source/Runtime/Engine/Private/GPUSkinCache.cpp
+++ b/Engine/Source/Runtime/Engine/Private/GPUSkinCache.cpp
@@ -1128,7 +1128,7 @@ void FGPUSkinCache::DispatchUpdateSkinTangents(FRHICommandList& RHICmdList, FGPU
 			{
 				StagingBuffer->Release();
 				StagingBuffer->Buffer.Initialize(RHICmdList, TEXT("SkinTangentIntermediate"), sizeof(int32), NumIntsPerBuffer, PF_R32_SINT, BUF_UnorderedAccess);
-				RHIBindDebugLabelName(StagingBuffer->Buffer.UAV, TEXT("SkinTangentIntermediate"));
+				RHICmdList.BindDebugLabelName(StagingBuffer->Buffer.UAV, TEXT("SkinTangentIntermediate"));
 
 				const uint32 MemSize = NumIntsPerBuffer * sizeof(uint32);
 				SET_MEMORY_STAT(STAT_GPUSkinCache_TangentsIntermediateMemUsed, MemSize);
diff --git a/Engine/Source/Runtime/Engine/Private/GameInstance.cpp b/Engine/Source/Runtime/Engine/Private/GameInstance.cpp
--- a/Engine/Source/Runtime/Engine/Private/GameInstance.cpp
+++ b/Engine/Source/Runtime/Engine/Private/GameInstance.cpp
@@ -520,7 +520,11 @@ FGameInstancePIEResult UGameInstance::StartPlayInEditorGameInstance(ULocalPlayer
 		PlayWorld->CreateAISystem();
 
 		SlowTask.EnterProgressFrame(10, NSLOCTEXT("UnrealEd", "PIEInitializingActors", "Starting PIE (Initializing Actors)..."));
-		PlayWorld->InitializeActorsForPlay(URL);
+		{
+			FRegisterComponentContext Context(PlayWorld);
+			PlayWorld->InitializeActorsForPlay(URL, true, &Context);
+			Context.Process();
+		}
 		// calling it after InitializeActorsForPlay has been called to have all potential bounding boxed initialized
 		FNavigationSystem::AddNavigationSystemToWorld(*PlayWorld, FNavigationSystemRunMode::PIEMode);
 
diff --git a/Engine/Source/Runtime/Engine/Private/Particles/ParticleGpuSimulation.cpp b/Engine/Source/Runtime/Engine/Private/Particles/ParticleGpuSimulation.cpp
--- a/Engine/Source/Runtime/Engine/Private/Particles/ParticleGpuSimulation.cpp
+++ b/Engine/Source/Runtime/Engine/Private/Particles/ParticleGpuSimulation.cpp
@@ -271,8 +271,8 @@ public:
 		FName VelocityTextureName(PARTICLE_STATE_VELOCITY_TEXTURE_NAME, ParticleStateIndex + 1);
 		PositionTextureRHI->SetName(PositionTextureName);
 		VelocityTextureRHI->SetName(VelocityTextureName);
-		RHIBindDebugLabelName(PositionTextureRHI, *PositionTextureName.ToString());
-		RHIBindDebugLabelName(VelocityTextureRHI, *VelocityTextureName.ToString());
+		RHICmdList.BindDebugLabelName(PositionTextureRHI, *PositionTextureName.ToString());
+		RHICmdList.BindDebugLabelName(VelocityTextureRHI, *VelocityTextureName.ToString());
 #undef PARTICLE_STATE_VELOCITY_TEXTURE_NAME
 #undef PARTICLE_STATE_POSITION_TEXTURE_NAME
 
diff --git a/Engine/Source/Runtime/NullDrv/Public/NullRHI.h b/Engine/Source/Runtime/NullDrv/Public/NullRHI.h
--- a/Engine/Source/Runtime/NullDrv/Public/NullRHI.h
+++ b/Engine/Source/Runtime/NullDrv/Public/NullRHI.h
@@ -272,7 +272,7 @@ public:
 
 	}
 
-	virtual void RHIBindDebugLabelName(FRHITexture* Texture, const TCHAR* Name) final override
+	virtual void RHIBindDebugLabelName(FRHICommandListBase& RHICmdList, FRHITexture* Texture, const TCHAR* Name) final override
 	{
 
 	}
diff --git a/Engine/Source/Runtime/OpenGLDrv/Private/OpenGLTexture.cpp b/Engine/Source/Runtime/OpenGLDrv/Private/OpenGLTexture.cpp
--- a/Engine/Source/Runtime/OpenGLDrv/Private/OpenGLTexture.cpp
+++ b/Engine/Source/Runtime/OpenGLDrv/Private/OpenGLTexture.cpp
@@ -1813,42 +1813,24 @@ void FOpenGLDynamicRHI::RHIUnlockTextureCubeFace(FRHITextureCube* TextureCubeRHI
 	ResourceCast(TextureCubeRHI)->Unlock(MipIndex, FaceIndex + ArrayIndex * 6);
 }
 
-void FOpenGLDynamicRHI::RHIBindDebugLabelName(FRHITexture* TextureRHI, const TCHAR* Name)
+void FOpenGLDynamicRHI::RHIBindDebugLabelName(FRHICommandListBase& RHICmdList, FRHITexture* TextureRHI, const TCHAR* Name)
 {
 #if GLDEBUG_LABELS_ENABLED
-	FRHICommandListImmediate& RHICmdList = FRHICommandListExecutor::GetImmediateCommandList();
-	if (ShouldRunGLRenderContextOpOnThisThread(RHICmdList))
+	FAnsiCharArray TextureDebugName;
+	TextureDebugName.Append(TCHAR_TO_ANSI(Name), FCString::Strlen(Name) + 1);
+	RHICmdList.EnqueueLambda([TextureRHI, TextureDebugName = MoveTemp(TextureDebugName)] (FRHICommandListBase& RHICmdList)
 	{
 		VERIFY_GL_SCOPE();
 		FOpenGLTexture* Texture = ResourceCast(TextureRHI);
 		if (Texture->IsEvicted())
 		{
-			Texture->EvictionParamsPtr->SetDebugLabelName(TCHAR_TO_ANSI(Name));
+			Texture->EvictionParamsPtr->SetDebugLabelName(TextureDebugName);
 		}
 		else
 		{
-			FOpenGL::LabelObject(GL_TEXTURE, Texture->GetResource(), TCHAR_TO_ANSI(Name));
+			FOpenGL::LabelObject(GL_TEXTURE, Texture->GetResource(), TextureDebugName.GetData());
 		}
-	}
-	else
-	{
-		// copy string name for RHIT version.
-		FAnsiCharArray TextureDebugName;
-		TextureDebugName.Append(TCHAR_TO_ANSI(Name), FCString::Strlen(Name) + 1);
-		RunOnGLRenderContextThread([TextureRHI, TextureDebugName = MoveTemp(TextureDebugName)]()
-		{
-			VERIFY_GL_SCOPE();
-			FOpenGLTexture* Texture = ResourceCast(TextureRHI);
-			if (Texture->IsEvicted())
-			{
-				Texture->EvictionParamsPtr->SetDebugLabelName(TextureDebugName);
-			}
-			else
-			{
-				FOpenGL::LabelObject(GL_TEXTURE, Texture->GetResource(), TextureDebugName.GetData());
-			}
-		});
-	}
+	});
 #endif
 }
 
diff --git a/Engine/Source/Runtime/OpenGLDrv/Public/OpenGLDrv.h b/Engine/Source/Runtime/OpenGLDrv/Public/OpenGLDrv.h
--- a/Engine/Source/Runtime/OpenGLDrv/Public/OpenGLDrv.h
+++ b/Engine/Source/Runtime/OpenGLDrv/Public/OpenGLDrv.h
@@ -382,7 +382,7 @@ public:
 	virtual void RHIUpdateTexture3D(FRHICommandListBase& RHICmdList, FRHITexture3D* Texture, uint32 MipIndex, const struct FUpdateTextureRegion3D& UpdateRegion, uint32 SourceRowPitch, uint32 SourceDepthPitch, const uint8* SourceData) final override;
 	virtual void* RHILockTextureCubeFace(FRHITextureCube* Texture, uint32 FaceIndex, uint32 ArrayIndex, uint32 MipIndex, EResourceLockMode LockMode, uint32& DestStride, bool bLockWithinMiptail) final override;
 	virtual void RHIUnlockTextureCubeFace(FRHITextureCube* Texture, uint32 FaceIndex, uint32 ArrayIndex, uint32 MipIndex, bool bLockWithinMiptail) final override;
-	virtual void RHIBindDebugLabelName(FRHITexture* Texture, const TCHAR* Name) final override;
+	virtual void RHIBindDebugLabelName(FRHICommandListBase& RHICmdList, FRHITexture* Texture, const TCHAR* Name) final override;
 	virtual void RHIReadSurfaceData(FRHITexture* Texture,FIntRect Rect,TArray<FColor>& OutData,FReadSurfaceDataFlags InFlags) final override;
 	virtual void RHIReadSurfaceData(FRHITexture* Texture, FIntRect Rect, TArray<FLinearColor>& OutData, FReadSurfaceDataFlags InFlags) final override;
 	virtual void RHIMapStagingSurface(FRHITexture* Texture, FRHIGPUFence* Fence, void*& OutData, int32& OutWidth, int32& OutHeight, uint32 GPUIndex = 0) final override;
diff --git a/Engine/Source/Runtime/RHI/Private/RHIValidation.cpp b/Engine/Source/Runtime/RHI/Private/RHIValidation.cpp
--- a/Engine/Source/Runtime/RHI/Private/RHIValidation.cpp
+++ b/Engine/Source/Runtime/RHI/Private/RHIValidation.cpp
@@ -759,42 +759,42 @@ static FString GetBreadcrumbPath()
 }
 
 // FlushType: Thread safe
-void FValidationRHI::RHIBindDebugLabelName(FRHITexture* Texture, const TCHAR* Name)
+void FValidationRHI::RHIBindDebugLabelName(FRHICommandListBase& RHICmdList, FRHITexture* Texture, const TCHAR* Name)
 {
 	check(IsInRenderingThread());
 
 	FString NameCopyRT = Name;
-	FRHICommandListExecutor::GetImmediateCommandList().EnqueueLambda([Texture, NameCopyRHIT = MoveTemp(NameCopyRT)](FRHICommandListImmediate& RHICmdList)
-		{
-			((FValidationContext&)RHICmdList.GetContext()).Tracker->Rename(Texture->GetTrackerResource(), *NameCopyRHIT);
-		});
+	RHICmdList.EnqueueLambda([Texture, NameCopyRHIT = MoveTemp(NameCopyRT)](FRHICommandListBase& RHICmdList)
+	{
+		((FValidationContext&)RHICmdList.GetContext()).Tracker->Rename(Texture->GetTrackerResource(), *NameCopyRHIT);
+	});
 
-	RHI->RHIBindDebugLabelName(Texture, Name);
+	RHI->RHIBindDebugLabelName(RHICmdList, Texture, Name);
 }
 
-void FValidationRHI::RHIBindDebugLabelName(FRHIBuffer* Buffer, const TCHAR* Name)
+void FValidationRHI::RHIBindDebugLabelName(FRHICommandListBase& RHICmdList, FRHIBuffer* Buffer, const TCHAR* Name)
 {
 	check(IsInRenderingThread());
 
 	FString NameCopyRT = Name;
-	FRHICommandListExecutor::GetImmediateCommandList().EnqueueLambda([Buffer, NameCopyRHIT = MoveTemp(NameCopyRT)](FRHICommandListImmediate& RHICmdList)
-		{
-			((FValidationContext&)RHICmdList.GetContext()).Tracker->Rename(Buffer, *NameCopyRHIT);
-		});
+	RHICmdList.EnqueueLambda([Buffer, NameCopyRHIT = MoveTemp(NameCopyRT)](FRHICommandListBase& RHICmdList)
+	{
+		((FValidationContext&)RHICmdList.GetContext()).Tracker->Rename(Buffer, *NameCopyRHIT);
+	});
 
-	RHI->RHIBindDebugLabelName(Buffer, Name);
+	RHI->RHIBindDebugLabelName(RHICmdList, Buffer, Name);
 }
 
-void FValidationRHI::RHIBindDebugLabelName(FRHIUnorderedAccessView* UnorderedAccessViewRHI, const TCHAR* Name)
+void FValidationRHI::RHIBindDebugLabelName(FRHICommandListBase& RHICmdList, FRHIUnorderedAccessView* UnorderedAccessViewRHI, const TCHAR* Name)
 {
 	RHIValidation::FResource* Resource = UnorderedAccessViewRHI->GetViewIdentity().Resource;
 	FString NameCopyRT = Name;
-	FRHICommandListExecutor::GetImmediateCommandList().EnqueueLambda([Resource, NameCopyRHIT = MoveTemp(NameCopyRT)](FRHICommandListImmediate& RHICmdList)
+	RHICmdList.EnqueueLambda([Resource, NameCopyRHIT = MoveTemp(NameCopyRT)](FRHICommandListBase& RHICmdList)
 	{
 		((FValidationContext&)RHICmdList.GetContext()).Tracker->Rename(Resource, *NameCopyRHIT);
 	});
 
-	RHI->RHIBindDebugLabelName(UnorderedAccessViewRHI, Name);
+	RHI->RHIBindDebugLabelName(RHICmdList, UnorderedAccessViewRHI, Name);
 }
 
 void FValidationRHI::ReportValidationFailure(const TCHAR* InMessage)
diff --git a/Engine/Source/Runtime/RHI/Public/DynamicRHI.h b/Engine/Source/Runtime/RHI/Public/DynamicRHI.h
--- a/Engine/Source/Runtime/RHI/Public/DynamicRHI.h
+++ b/Engine/Source/Runtime/RHI/Public/DynamicRHI.h
@@ -579,9 +579,9 @@ public:
 	virtual void RHIUnlockTextureCubeFace(FRHITextureCube* Texture, uint32 FaceIndex, uint32 ArrayIndex, uint32 MipIndex, bool bLockWithinMiptail) = 0;
 
 	// FlushType: Thread safe
-	virtual void RHIBindDebugLabelName(FRHITexture* Texture, const TCHAR* Name) = 0;
-	virtual void RHIBindDebugLabelName(FRHIBuffer* Buffer, const TCHAR* Name) {}
-	virtual void RHIBindDebugLabelName(FRHIUnorderedAccessView* UnorderedAccessViewRHI, const TCHAR* Name) {}
+	virtual void RHIBindDebugLabelName(FRHICommandListBase& RHICmdList,FRHITexture* Texture, const TCHAR* Name) = 0;
+	virtual void RHIBindDebugLabelName(FRHICommandListBase& RHICmdList,FRHIBuffer* Buffer, const TCHAR* Name) {}
+	virtual void RHIBindDebugLabelName(FRHICommandListBase& RHICmdList,FRHIUnorderedAccessView* UnorderedAccessViewRHI, const TCHAR* Name) {}
 
 	/**
 	* Reads the contents of a texture to an output buffer (non MSAA and MSAA) and returns it as a FColor array.
@@ -1296,21 +1296,6 @@ FORCEINLINE uint32 RHIComputeMemorySize(FRHITexture* TextureRHI)
 	return GDynamicRHI->RHIComputeMemorySize(TextureRHI);
 }
 
-FORCEINLINE void RHIBindDebugLabelName(FRHITexture* Texture, const TCHAR* Name)
-{
-	GDynamicRHI->RHIBindDebugLabelName(Texture, Name);
-}
-
-FORCEINLINE void RHIBindDebugLabelName(FRHIBuffer* Buffer, const TCHAR* Name)
-{
-	GDynamicRHI->RHIBindDebugLabelName(Buffer, Name);
-}
-
-FORCEINLINE void RHIBindDebugLabelName(FRHIUnorderedAccessView* UnorderedAccessViewRHI, const TCHAR* Name)
-{
-	GDynamicRHI->RHIBindDebugLabelName(UnorderedAccessViewRHI, Name);
-}
-
 FORCEINLINE bool RHIGetRenderQueryResult(FRHIRenderQuery* RenderQuery, uint64& OutResult, bool bWait, uint32 GPUIndex = INDEX_NONE)
 {
 	return GDynamicRHI->RHIGetRenderQueryResult(RenderQuery, OutResult, bWait, GPUIndex);
diff --git a/Engine/Source/Runtime/RHI/Public/RHICommandList.h b/Engine/Source/Runtime/RHI/Public/RHICommandList.h
--- a/Engine/Source/Runtime/RHI/Public/RHICommandList.h
+++ b/Engine/Source/Runtime/RHI/Public/RHICommandList.h
@@ -987,6 +987,20 @@ public:
 	}
 #endif
 
+	FORCEINLINE void BindDebugLabelName(FRHITexture* Texture, const TCHAR* Name)
+	{
+		GDynamicRHI->RHIBindDebugLabelName(*this, Texture, Name);
+	}
+
+	FORCEINLINE void BindDebugLabelName(FRHIBuffer* Buffer, const TCHAR* Name)
+	{
+		GDynamicRHI->RHIBindDebugLabelName(*this, Buffer, Name);
+	}
+
+	FORCEINLINE void BindDebugLabelName(FRHIUnorderedAccessView* UnorderedAccessViewRHI, const TCHAR* Name)
+	{
+		GDynamicRHI->RHIBindDebugLabelName(*this, UnorderedAccessViewRHI, Name);
+	}
 
 	inline FRHIBatchedShaderParameters& GetScratchShaderParameters()
 	{
@@ -5312,6 +5326,21 @@ FORCEINLINE FRayTracingAccelerationStructureSize RHICalcRayTracingGeometrySize(c
 }
 #endif
 
+FORCEINLINE void RHIBindDebugLabelName(FRHITexture* Texture, const TCHAR* Name)
+{
+	FRHICommandListImmediate::Get().BindDebugLabelName(Texture, Name);
+}
+
+FORCEINLINE void RHIBindDebugLabelName(FRHIBuffer* Buffer, const TCHAR* Name)
+{
+	FRHICommandListImmediate::Get().BindDebugLabelName(Buffer, Name);
+}
+
+FORCEINLINE void RHIBindDebugLabelName(FRHIUnorderedAccessView* UnorderedAccessViewRHI, const TCHAR* Name)
+{
+	FRHICommandListImmediate::Get().BindDebugLabelName(UnorderedAccessViewRHI, Name);
+}
+
 namespace UE::RHI
 {
 
diff --git a/Engine/Source/Runtime/RHI/Public/RHIValidation.h b/Engine/Source/Runtime/RHI/Public/RHIValidation.h
--- a/Engine/Source/Runtime/RHI/Public/RHIValidation.h
+++ b/Engine/Source/Runtime/RHI/Public/RHIValidation.h
@@ -730,11 +730,11 @@ public:
 	}
 
 	// FlushType: Thread safe
-	virtual void RHIBindDebugLabelName(FRHITexture* Texture, const TCHAR* Name) override final;
+	virtual void RHIBindDebugLabelName(FRHICommandListBase& RHICmdList, FRHITexture* Texture, const TCHAR* Name) override final;
 
-	virtual void RHIBindDebugLabelName(FRHIBuffer* Buffer, const TCHAR* Name) override final;
+	virtual void RHIBindDebugLabelName(FRHICommandListBase& RHICmdList, FRHIBuffer* Buffer, const TCHAR* Name) override final;
 
-	virtual void RHIBindDebugLabelName(FRHIUnorderedAccessView* UnorderedAccessViewRHI, const TCHAR* Name) override final;
+	virtual void RHIBindDebugLabelName(FRHICommandListBase& RHICmdList, FRHIUnorderedAccessView* UnorderedAccessViewRHI, const TCHAR* Name) override final;
 
 	/**
 	* Reads the contents of a texture to an output buffer (non MSAA and MSAA) and returns it as a FColor array.
diff --git a/Engine/Source/Runtime/RenderCore/Private/RenderGraphResourcePool.cpp b/Engine/Source/Runtime/RenderCore/Private/RenderGraphResourcePool.cpp
--- a/Engine/Source/Runtime/RenderCore/Private/RenderGraphResourcePool.cpp
+++ b/Engine/Source/Runtime/RenderCore/Private/RenderGraphResourcePool.cpp
@@ -106,7 +106,7 @@ TRefCountPtr<FRDGPooledBuffer> FRDGBufferPool::FindFreeBuffer(FRHICommandListBas
 		PooledBuffer->Name = InDebugName;
 
 	#if !(UE_BUILD_SHIPPING || UE_BUILD_TEST)
-		RHIBindDebugLabelName(PooledBuffer->GetRHI(), InDebugName);
+		RHICmdList.BindDebugLabelName(PooledBuffer->GetRHI(), InDebugName);
 	#endif
 
 		// We need the external-facing desc to match what the user requested.
@@ -134,7 +134,7 @@ TRefCountPtr<FRDGPooledBuffer> FRDGBufferPool::FindFreeBuffer(FRHICommandListBas
 		TRefCountPtr<FRHIBuffer> BufferRHI = RHICmdList.CreateBuffer(NumBytes, Desc.Usage, Desc.BytesPerElement, InitialAccess, CreateInfo);
 
 	#if !(UE_BUILD_SHIPPING || UE_BUILD_TEST)
-		RHIBindDebugLabelName(BufferRHI, InDebugName);
+		RHICmdList.BindDebugLabelName(BufferRHI, InDebugName);
 	#endif
 
 		TRefCountPtr<FRDGPooledBuffer> PooledBuffer = new FRDGPooledBuffer(RHICmdList, MoveTemp(BufferRHI), Desc, AlignedDesc.NumElements, InDebugName);
diff --git a/Engine/Source/Runtime/RenderCore/Private/RenderTargetPool.cpp b/Engine/Source/Runtime/RenderCore/Private/RenderTargetPool.cpp
--- a/Engine/Source/Runtime/RenderCore/Private/RenderTargetPool.cpp
+++ b/Engine/Source/Runtime/RenderCore/Private/RenderTargetPool.cpp
@@ -158,7 +158,7 @@ TRefCountPtr<IPooledRenderTarget> FRenderTargetPool::FindFreeElement(FRHICommand
 	Found->UnusedForNFrames = 0;
 
 #if !(UE_BUILD_SHIPPING || UE_BUILD_TEST)
-	RHIBindDebugLabelName(Found->GetRHI(), Name);
+	RHICmdList.BindDebugLabelName(Found->GetRHI(), Name);
 #endif
 
 	return TRefCountPtr<IPooledRenderTarget>(MoveTemp(Found));
@@ -187,7 +187,7 @@ bool FRenderTargetPool::FindFreeElement(FRHICommandListBase& RHICmdList, const F
 #if !(UE_BUILD_SHIPPING || UE_BUILD_TEST)
 			if (Current->GetRHI())
 			{
-				RHIBindDebugLabelName(Current->GetRHI(), Name);
+				RHICmdList.BindDebugLabelName(Current->GetRHI(), Name);
 			}
 #endif
 			check(!Out->IsFree());
diff --git a/Engine/Source/Runtime/VulkanRHI/Private/VulkanTexture.cpp b/Engine/Source/Runtime/VulkanRHI/Private/VulkanTexture.cpp
--- a/Engine/Source/Runtime/VulkanRHI/Private/VulkanTexture.cpp
+++ b/Engine/Source/Runtime/VulkanRHI/Private/VulkanTexture.cpp
@@ -2036,7 +2036,7 @@ void FVulkanDynamicRHI::RHIUnlockTextureCubeFace(FRHITextureCube* TextureCubeRHI
 	}
 }
 
-void FVulkanDynamicRHI::RHIBindDebugLabelName(FRHITexture* TextureRHI, const TCHAR* Name)
+void FVulkanDynamicRHI::RHIBindDebugLabelName(FRHICommandListBase& RHICmdList, FRHITexture* TextureRHI, const TCHAR* Name)
 {
 #if VULKAN_ENABLE_IMAGE_TRACKING_LAYER
 	{
@@ -2078,7 +2078,7 @@ void FVulkanDynamicRHI::RHIBindDebugLabelName(FRHITexture* TextureRHI, const TCH
 	TextureRHI->SetName(DebugName);
 }
 
-void FVulkanDynamicRHI::RHIBindDebugLabelName(FRHIUnorderedAccessView* UnorderedAccessViewRHI, const TCHAR* Name)
+void FVulkanDynamicRHI::RHIBindDebugLabelName(FRHICommandListBase& RHICmdList, FRHIUnorderedAccessView* UnorderedAccessViewRHI, const TCHAR* Name)
 {
 #if VULKAN_ENABLE_DUMP_LAYER
 	//if (Device->SupportsDebugMarkers())
diff --git a/Engine/Source/Runtime/VulkanRHI/Public/VulkanDynamicRHI.h b/Engine/Source/Runtime/VulkanRHI/Public/VulkanDynamicRHI.h
--- a/Engine/Source/Runtime/VulkanRHI/Public/VulkanDynamicRHI.h
+++ b/Engine/Source/Runtime/VulkanRHI/Public/VulkanDynamicRHI.h
@@ -149,8 +149,8 @@ public:
 	}
 	virtual void* RHILockTextureCubeFace(FRHITextureCube* Texture, uint32 FaceIndex, uint32 ArrayIndex, uint32 MipIndex, EResourceLockMode LockMode, uint32& DestStride, bool bLockWithinMiptail) final override;
 	virtual void RHIUnlockTextureCubeFace(FRHITextureCube* Texture, uint32 FaceIndex, uint32 ArrayIndex, uint32 MipIndex, bool bLockWithinMiptail) final override;
-	virtual void RHIBindDebugLabelName(FRHITexture* Texture, const TCHAR* Name) final override;
-	virtual void RHIBindDebugLabelName(FRHIUnorderedAccessView* UnorderedAccessViewRHI, const TCHAR* Name) final override;
+	virtual void RHIBindDebugLabelName(FRHICommandListBase& RHICmdList, FRHITexture* Texture, const TCHAR* Name) final override;
+	virtual void RHIBindDebugLabelName(FRHICommandListBase& RHICmdList, FRHIUnorderedAccessView* UnorderedAccessViewRHI, const TCHAR* Name) final override;
 	virtual void RHIReadSurfaceData(FRHITexture* Texture, FIntRect Rect, TArray<FColor>& OutData, FReadSurfaceDataFlags InFlags) final override;
 	virtual void RHIReadSurfaceData(FRHITexture* Texture, FIntRect Rect, TArray<FLinearColor>& OutData, FReadSurfaceDataFlags InFlags) final override;
 	virtual void RHIMapStagingSurface(FRHITexture* Texture, FRHIGPUFence* Fence, void*& OutData, int32& OutWidth, int32& OutHeight, uint32 GPUIndex = 0) final override;
diff --git a/Engine/Source/Runtime/Windows/D3D11RHI/Private/D3D11Buffer.cpp b/Engine/Source/Runtime/Windows/D3D11RHI/Private/D3D11Buffer.cpp
--- a/Engine/Source/Runtime/Windows/D3D11RHI/Private/D3D11Buffer.cpp
+++ b/Engine/Source/Runtime/Windows/D3D11RHI/Private/D3D11Buffer.cpp
@@ -298,7 +298,7 @@ void FD3D11DynamicRHI::RHITransferBufferUnderlyingResource(FRHICommandListBase&
 	Dst->UpdateLinkedViews();
 }
 
-void FD3D11DynamicRHI::RHIBindDebugLabelName(FRHIBuffer* BufferRHI, const TCHAR* Name)
+void FD3D11DynamicRHI::RHIBindDebugLabelName(FRHICommandListBase& RHICmdList, FRHIBuffer* BufferRHI, const TCHAR* Name)
 {
 	check(BufferRHI);
 	FName DebugName(Name);
diff --git a/Engine/Source/Runtime/Windows/D3D11RHI/Private/D3D11RHIPrivate.h b/Engine/Source/Runtime/Windows/D3D11RHI/Private/D3D11RHIPrivate.h
--- a/Engine/Source/Runtime/Windows/D3D11RHI/Private/D3D11RHIPrivate.h
+++ b/Engine/Source/Runtime/Windows/D3D11RHI/Private/D3D11RHIPrivate.h
@@ -514,9 +514,9 @@ public:
 	virtual void RHIUpdateTexture3D(FRHICommandListBase& RHICmdList, FRHITexture3D* Texture, uint32 MipIndex, const struct FUpdateTextureRegion3D& UpdateRegion, uint32 SourceRowPitch, uint32 SourceDepthPitch, const uint8* SourceData) final override;
 	virtual void* RHILockTextureCubeFace(FRHITextureCube* Texture, uint32 FaceIndex, uint32 ArrayIndex, uint32 MipIndex, EResourceLockMode LockMode, uint32& DestStride, bool bLockWithinMiptail) final override;
 	virtual void RHIUnlockTextureCubeFace(FRHITextureCube* Texture, uint32 FaceIndex, uint32 ArrayIndex, uint32 MipIndex, bool bLockWithinMiptail) final override;
-	virtual void RHIBindDebugLabelName(FRHITexture* Texture, const TCHAR* Name) final override;
-	virtual void RHIBindDebugLabelName(FRHIBuffer* Buffer, const TCHAR* Name) final override;
-	virtual void RHIBindDebugLabelName(FRHIUnorderedAccessView* UnorderedAccessViewRHI, const TCHAR* Name) final override;
+	virtual void RHIBindDebugLabelName(FRHICommandListBase& RHICmdList, FRHITexture* Texture, const TCHAR* Name) final override;
+	virtual void RHIBindDebugLabelName(FRHICommandListBase& RHICmdList, FRHIBuffer* Buffer, const TCHAR* Name) final override;
+	virtual void RHIBindDebugLabelName(FRHICommandListBase& RHICmdList, FRHIUnorderedAccessView* UnorderedAccessViewRHI, const TCHAR* Name) final override;
 	virtual void RHIReadSurfaceData(FRHITexture* Texture,FIntRect Rect,TArray<FColor>& OutData,FReadSurfaceDataFlags InFlags) final override;
 	virtual void RHIReadSurfaceData(FRHITexture* TextureRHI, FIntRect InRect, TArray<FLinearColor>& OutData, FReadSurfaceDataFlags InFlags) final override;
 	virtual void RHIMapStagingSurface(FRHITexture* Texture, FRHIGPUFence* Fence, void*& OutData, int32& OutWidth, int32& OutHeight, uint32 GPUIndex = 0) final override;
diff --git a/Engine/Source/Runtime/Windows/D3D11RHI/Private/D3D11Texture.cpp b/Engine/Source/Runtime/Windows/D3D11RHI/Private/D3D11Texture.cpp
--- a/Engine/Source/Runtime/Windows/D3D11RHI/Private/D3D11Texture.cpp
+++ b/Engine/Source/Runtime/Windows/D3D11RHI/Private/D3D11Texture.cpp
@@ -1431,7 +1431,7 @@ void FD3D11DynamicRHI::RHIUnlockTextureCubeFace_RenderThread(
 	}
 }
 
-void FD3D11DynamicRHI::RHIBindDebugLabelName(FRHITexture* TextureRHI, const TCHAR* Name)
+void FD3D11DynamicRHI::RHIBindDebugLabelName(FRHICommandListBase& RHICmdList, FRHITexture* TextureRHI, const TCHAR* Name)
 {
 	FD3D11Texture* Texture = ResourceCast(TextureRHI);
 
diff --git a/Engine/Source/Runtime/Windows/D3D11RHI/Private/D3D11UAV.cpp b/Engine/Source/Runtime/Windows/D3D11RHI/Private/D3D11UAV.cpp
--- a/Engine/Source/Runtime/Windows/D3D11RHI/Private/D3D11UAV.cpp
+++ b/Engine/Source/Runtime/Windows/D3D11RHI/Private/D3D11UAV.cpp
@@ -132,7 +132,7 @@ FUnorderedAccessViewRHIRef FD3D11DynamicRHI::RHICreateUnorderedAccessView(class
 	return new FD3D11UnorderedAccessView(RHICmdList, Resource, ViewDesc);
 }
 
-void FD3D11DynamicRHI::RHIBindDebugLabelName(FRHIUnorderedAccessView* UnorderedAccessViewRHI, const TCHAR* Name)
+void FD3D11DynamicRHI::RHIBindDebugLabelName(FRHICommandListBase& RHICmdList, FRHIUnorderedAccessView* UnorderedAccessViewRHI, const TCHAR* Name)
 {
 #if UE_BUILD_DEBUG || UE_BUILD_DEVELOPMENT
 	FD3D11UnorderedAccessView* UAV = ResourceCast(UnorderedAccessViewRHI);
