Optimized compute pipeline state creation/management to use a thread local cache like graphics already does, avoiding significant write lock performance problems in ExecuteSetComputePipelineState.

diff --git a/Engine/Source/Runtime/RHI/Private/PipelineStateCache.cpp b/Engine/Source/Runtime/RHI/Private/PipelineStateCache.cpp
--- a/Engine/Source/Runtime/RHI/Private/PipelineStateCache.cpp
+++ b/Engine/Source/Runtime/RHI/Private/PipelineStateCache.cpp
@@ -986,10 +986,10 @@ private:
 		FScopeLock S(&AllThreadsLock);
 		for (auto PipelineStateCacheIterator = PipelineStateCache->CreateIterator(); PipelineStateCacheIterator; ++PipelineStateCacheIterator)
 		{
-			FGraphicsPipelineState* pGPipelineState = PipelineStateCacheIterator->Value;
-			if(pGPipelineState != nullptr)
+			auto PipelineState = PipelineStateCacheIterator->Value;
+			if (PipelineState != nullptr)
 			{
-				pGPipelineState->WaitCompletion();
+				PipelineState->WaitCompletion();
 			}
 		}
 	}
@@ -1462,7 +1462,7 @@ public:
 };
 
 // Typed caches for compute and graphics
-typedef TDiscardableKeyValueCache< FRHIComputeShader*, FComputePipelineState*> FComputePipelineCache;
+typedef TSharedPipelineStateCache<FRHIComputeShader*, FComputePipelineState*> FComputePipelineCache;
 typedef TSharedPipelineStateCache<FGraphicsPipelineStateInitializer, FGraphicsPipelineState*> FGraphicsPipelineCache;
 
 // These are the actual caches for both pipelines
@@ -1931,7 +1931,7 @@ public:
 		ENamedThreads::Type DesiredThread = GRunPSOCreateTasksOnRHIT && IsRunningRHIInSeparateThread() && bInImmediateCmdList ? ENamedThreads::RHIThread : CPrio_FCompilePipelineStateTask.Get();
 
 		// On Mac the compilation is handled using external processes, so engine threads have very little work to do
-		// and it's better to leave more CPU time to these extrenal processes and other engine threads.
+		// and it's better to leave more CPU time to these external processes and other engine threads.
 		// Also use background threads for PSO precaching when the PSO thread pool is not used
 		// Compute pipelines usually take much longer to compile, compile them on background thread as well.
 		return (PLATFORM_MAC || PSOPreCacheResult == EPSOPrecacheResult::Active || (Pipeline && Pipeline->IsCompute() && Initializer.bFromPSOFileCache)) ? ENamedThreads::AnyBackgroundThreadNormalTask : DesiredThread;
@@ -1950,6 +1950,9 @@ void PipelineStateCache::FlushResources()
 {
 	check(IsInRenderingThread());
 
+	GComputePipelineCache.ConsolidateThreadedCaches();
+	GComputePipelineCache.ProcessDelayedCleanup();
+
 	GGraphicsPipelineCache.ConsolidateThreadedCaches();
 	GGraphicsPipelineCache.ProcessDelayedCleanup();
 
@@ -2020,13 +2023,10 @@ void PipelineStateCache::FlushResources()
 
 	LastEvictionTime = CurrentTime;
 
-	int ReleasedComputeEntries = 0;
-	int ReleasedGraphicsEntries = 0;
-
-	ReleasedComputeEntries = GComputePipelineCache.Discard([](FComputePipelineState* CacheItem) {
-		delete CacheItem;
-	});
+	int32 ReleasedComputeEntries = 0;
+	int32 ReleasedGraphicsEntries = 0;
 
+	ReleasedComputeEntries  =  GComputePipelineCache.DiscardAndSwap();
 	ReleasedGraphicsEntries = GGraphicsPipelineCache.DiscardAndSwap();
 
 #if PSO_TRACK_CACHE_STATS
@@ -2055,55 +2055,69 @@ uint64 PipelineStateCache::RetrieveGraphicsPipelineStateSortKey(const FGraphicsP
 	return GraphicsPipelineState != nullptr ? GraphicsPipelineState->SortKey : 0;
 }
 
+static void InternalCreateComputePipelineState(FRHIComputeShader* ComputeShader, bool bDoAsyncCompile, bool bFromPSOFileCache, FComputePipelineState* CachedState, bool bInImmediateCmdList = true)
+{
+	FGraphEventRef GraphEvent = CachedState->CompletionEvent;
+
+	// create a compilation task, or just do it now...
+	if (bDoAsyncCompile)
+	{
+		check(GraphEvent != nullptr);
+
+		FGraphicsPipelineStateInitializer GraphicsPipelineStateInitializer;
+		GraphicsPipelineStateInitializer.bFromPSOFileCache = bFromPSOFileCache;
+		TGraphTask<FCompilePipelineStateTask>::CreateTask().ConstructAndDispatchWhenReady(CachedState, GraphicsPipelineStateInitializer, EPSOPrecacheResult::Untracked, bInImmediateCmdList);
+	}
+	else
+	{
+		check(GraphEvent == nullptr);
+		uint64 StartTime = FPlatformTime::Cycles64();
+		CachedState->RHIPipeline = RHICreateComputePipelineState(ComputeShader);
+		CheckAndUpdateHitchCountStat(FPSOPrecacheRequestID::EType::Compute, !bFromPSOFileCache, StartTime);
+
+		if (!CachedState->RHIPipeline)
+		{
+			HandlePipelineCreationFailure(ComputeShader, bFromPSOFileCache);
+		}
+	}
+}
+
 FComputePipelineState* PipelineStateCache::GetAndOrCreateComputePipelineState(FRHIComputeCommandList& RHICmdList, FRHIComputeShader* ComputeShader, bool bFromFileCache)
 {
 	LLM_SCOPE(ELLMTag::PSO);
-	bool DoAsyncCompile = IsAsyncCompilationAllowed(RHICmdList, bFromFileCache);
 
 	FComputePipelineState* OutCachedState = nullptr;
 
-	uint32 LockFlags = GComputePipelineCache.ApplyLock(0, FComputePipelineCache::LockFlags::ReadLock);
-
-	bool WasFound = GComputePipelineCache.Find(ComputeShader, OutCachedState, LockFlags | FComputePipelineCache::LockFlags::WriteLockOnAddFail, LockFlags);
+	bool bWasFound = GComputePipelineCache.Find(ComputeShader, OutCachedState);
+	bool DoAsyncCompile = IsAsyncCompilationAllowed(RHICmdList, bFromFileCache);
 
-	if (WasFound == false)
+	if (bWasFound == false)
 	{
 		FPipelineFileCacheManager::CacheComputePSO(GetTypeHash(ComputeShader), ComputeShader);
 
-		// create new graphics state
+		// create new compute state
 		OutCachedState = new FComputePipelineState(ComputeShader);
 		OutCachedState->Stats = FPipelineFileCacheManager::RegisterPSOStats(GetTypeHash(ComputeShader));
+		if (DoAsyncCompile)
+		{
+			OutCachedState->CompletionEvent = FGraphEvent::CreateGraphEvent();
+		}
 
 		if (!bFromFileCache)
 		{
 			++ComputePipelineCacheMisses;
 		}
 
-		// create a compilation task, or just do it now...
-		if (DoAsyncCompile)
-		{
-			OutCachedState->CompletionEvent = FGraphEvent::CreateGraphEvent();
-			FGraphicsPipelineStateInitializer GraphicsPipelineStateInitializer;
-			GraphicsPipelineStateInitializer.bFromPSOFileCache = bFromFileCache;
-			TGraphTask<FCompilePipelineStateTask>::CreateTask().ConstructAndDispatchWhenReady(OutCachedState, GraphicsPipelineStateInitializer, EPSOPrecacheResult::Untracked, RHICmdList.IsImmediate());
-			if (!bFromFileCache)
-			{
-				RHICmdList.AddDispatchPrerequisite(OutCachedState->CompletionEvent);
-			}
-		}
-		else
-		{
-			uint64 StartTime = FPlatformTime::Cycles64();
-			OutCachedState->RHIPipeline = RHICreateComputePipelineState(OutCachedState->ComputeShader);
-			CheckAndUpdateHitchCountStat(FPSOPrecacheRequestID::EType::Compute, !bFromFileCache, StartTime);
+		FGraphEventRef GraphEvent = OutCachedState->CompletionEvent;
+		InternalCreateComputePipelineState(ComputeShader, DoAsyncCompile, bFromFileCache, OutCachedState, RHICmdList.IsImmediate());
 
-			if (!OutCachedState->RHIPipeline)
-			{
-				HandlePipelineCreationFailure(OutCachedState->ComputeShader, bFromFileCache);
-			}
+		if (GraphEvent.IsValid())
+		{
+			check(DoAsyncCompile);
+			RHICmdList.AddDispatchPrerequisite(GraphEvent);
 		}
 
-		GComputePipelineCache.Add(ComputeShader, OutCachedState, LockFlags);
+		GComputePipelineCache.Add(ComputeShader, OutCachedState);
 	}
 	else
 	{
@@ -2112,60 +2126,11 @@ FComputePipelineState* PipelineStateCache::GetAndOrCreateComputePipelineState(FR
 			RHICmdList.AddDispatchPrerequisite(OutCachedState->CompletionEvent);
 		}
 
-#if PSO_TRACK_CACHE_STATS
+	#if PSO_TRACK_CACHE_STATS
 		OutCachedState->AddHit();
-#endif
+	#endif
 	}
 
-	GComputePipelineCache.Unlock(LockFlags);
-
-#if 0
-	bool DoAsyncCompile = IsAsyncCompilationAllowed(RHICmdList, bFromFileCache);
-
-
-	TSharedPtr<FComputePipelineState> OutCachedState;
-
-	// Find or add an entry for this initializer
-	bool WasFound = GComputePipelineCache.FindOrAdd(ComputeShader, OutCachedState, [&RHICmdList, &ComputeShader, &DoAsyncCompile] {
-			// create new graphics state
-			TSharedPtr<FComputePipelineState> PipelineState(new FComputePipelineState(ComputeShader));
-			PipelineState->Stats = FPipelineFileCacheManager::RegisterPSOStats(GetTypeHash(ComputeShader));
-
-			// create a compilation task, or just do it now...
-			if (DoAsyncCompile)
-			{
-				PipelineState->CompletionEvent = FGraphEvent::CreateGraphEvent();
-				TGraphTask<FCompilePipelineStateTask>::CreateTask().ConstructAndDispatchWhenReady(PipelineState.Get(), FGraphicsPipelineStateInitializer(), EPSOPrecacheResult::Untracked);
-				RHICmdList.QueueAsyncPipelineStateCompile(PipelineState->CompletionEvent);
-			}
-			else
-			{
-				PipelineState->RHIPipeline = RHICreateComputePipelineState(PipelineState->ComputeShader);
-			}
-
-			// wrap it and return it
-			return PipelineState;
-		});
-
-	check(OutCachedState.IsValid());
-
-	// if we found an entry the block above wasn't executed
-	if (WasFound)
-	{
-		if (DoAsyncCompile)
-		{
-			FRWScopeLock ScopeLock(GComputePipelineCache.RWLock(), SLT_ReadOnly);
-			FGraphEventRef& CompletionEvent = OutCachedState->CompletionEvent;
-			if (CompletionEvent.IsValid() && !CompletionEvent->IsComplete())
-			{
-				RHICmdList.QueueAsyncPipelineStateCompile(CompletionEvent);
-			}
-		}
-#if PSO_TRACK_CACHE_STATS
-		OutCachedState->AddHit();
-#endif
-	}
-#endif
 	// return the state pointer
 	return OutCachedState;
 }
@@ -2457,9 +2422,11 @@ FRayTracingPipelineState* PipelineStateCache::GetRayTracingPipelineState(const F
 FRHIComputePipelineState* ExecuteSetComputePipelineState(FComputePipelineState* ComputePipelineState)
 {
 	ensure(ComputePipelineState->RHIPipeline);
-	FRWScopeLock ScopeLock(GComputePipelineCache.RWLock(), SLT_Write);
 	ComputePipelineState->AddUse();
-	ComputePipelineState->CompletionEvent = nullptr;
+#if PIPELINESTATECACHE_VERIFYTHREADSAFE
+	int32 Result = ComputePipelineState->InUseCount.Decrement();
+	check(Result >= 0);
+#endif
 	return ComputePipelineState->RHIPipeline;
 }
 
@@ -2670,6 +2637,16 @@ FGraphicsPipelineState* PipelineStateCache::GetAndOrCreateGraphicsPipelineState(
 	return OutCachedState;
 }
 
+FComputePipelineState* PipelineStateCache::FindComputePipelineState(FRHIComputeShader* ComputeShader)
+{
+	LLM_SCOPE(ELLMTag::PSO);
+	check(ComputeShader != nullptr);
+
+	FComputePipelineState* PipelineState = nullptr;
+	GComputePipelineCache.Find(ComputeShader, PipelineState);
+	return (PipelineState && PipelineState->IsComplete()) ? PipelineState : nullptr;
+}
+
 FGraphicsPipelineState* PipelineStateCache::FindGraphicsPipelineState(const FGraphicsPipelineStateInitializer& Initializer)
 {
 	LLM_SCOPE(ELLMTag::PSO);
@@ -2979,23 +2956,16 @@ static FCriticalSection GVertexDeclarationLock;
 
 void PipelineStateCache::Shutdown()
 {
+	GComputePipelineCache.WaitTasksComplete();
 	GGraphicsPipelineCache.WaitTasksComplete();
 #if RHI_RAYTRACING
 	GRayTracingPipelineCache.Shutdown();
 #endif
 
 	// call discard twice to clear both the backing and main caches
-	for (int i = 0; i < 2; i++)
+	for (int32 i = 0; i < 2; i++)
 	{
-		GComputePipelineCache.Discard([](FComputePipelineState* CacheItem)
-		{
-			if(CacheItem != nullptr)
-			{
-				CacheItem->WaitCompletion();
-				delete CacheItem;
-			}
-		});
-		
+		GComputePipelineCache.DiscardAndSwap();
 		GGraphicsPipelineCache.DiscardAndSwap();
 	}
 	FPipelineFileCacheManager::Shutdown();
diff --git a/Engine/Source/Runtime/RHI/Public/PipelineStateCache.h b/Engine/Source/Runtime/RHI/Public/PipelineStateCache.h
--- a/Engine/Source/Runtime/RHI/Public/PipelineStateCache.h
+++ b/Engine/Source/Runtime/RHI/Public/PipelineStateCache.h
@@ -114,6 +114,8 @@ namespace PipelineStateCache
 
 	extern RHI_API FGraphicsPipelineState*	GetAndOrCreateGraphicsPipelineState(FRHICommandList& RHICmdList, const FGraphicsPipelineStateInitializer& OriginalInitializer, EApplyRendertargetOption ApplyFlags, EPSOPrecacheResult PSOPrecacheResult);
 
+	extern RHI_API FComputePipelineState*	FindComputePipelineState(FRHIComputeShader* ComputeShader);
+
 	extern RHI_API FGraphicsPipelineState*	FindGraphicsPipelineState(const FGraphicsPipelineStateInitializer& Initializer);
 
 	extern RHI_API FRHIVertexDeclaration*	GetOrCreateVertexDeclaration(const FVertexDeclarationElementList& Elements);
