Add PIPELINESTATECACHE_VERIFYTHREADSAFE support to FComputePipelineState

diff --git a/Engine/Source/Runtime/RHI/Private/PipelineStateCache.cpp b/Engine/Source/Runtime/RHI/Private/PipelineStateCache.cpp
--- a/Engine/Source/Runtime/RHI/Private/PipelineStateCache.cpp
+++ b/Engine/Source/Runtime/RHI/Private/PipelineStateCache.cpp
@@ -167,11 +167,6 @@ static inline void CheckAndUpdateHitchCountStat(FPSOPrecacheRequestID::EType PSO
 	}
 }
 
-void SetComputePipelineState(FRHIComputeCommandList& RHICmdList, FRHIComputeShader* ComputeShader)
-{
-	RHICmdList.SetComputePipelineState(PipelineStateCache::GetAndOrCreateComputePipelineState(RHICmdList, ComputeShader, false), ComputeShader);
-}
-
 static int32 GPSOPrecompileThreadPoolSize = 0;
 static FAutoConsoleVariableRef GPSOPrecompileThreadPoolSizeVar(
 	TEXT("r.pso.PrecompileThreadPoolSize"),
@@ -413,8 +408,6 @@ protected:
 	}
 };
 
-
-
 /**
  * Base class to hold pipeline state (and optionally stats)
  */
@@ -533,6 +526,9 @@ public:
 
 	FRHIComputeShader* ComputeShader;
 	TRefCountPtr<FRHIComputePipelineState> RHIPipeline;
+#if PIPELINESTATECACHE_VERIFYTHREADSAFE
+	FThreadSafeCounter InUseCount;
+#endif
 };
 
 /* State for graphics */
@@ -685,6 +681,17 @@ bool IsPrecachedPSO(const FGraphicsPipelineStateInitializer& Initializer)
 	return Initializer.bFromPSOFileCache || Initializer.bPSOPrecache;
 }
 
+void SetComputePipelineState(FRHIComputeCommandList& RHICmdList, FRHIComputeShader* ComputeShader)
+{
+	FComputePipelineState* PipelineState = PipelineStateCache::GetAndOrCreateComputePipelineState(RHICmdList, ComputeShader, false);
+#if PIPELINESTATECACHE_VERIFYTHREADSAFE
+	int32 Result = PipelineState->InUseCount.Increment();
+	check(Result >= 1);
+#endif
+	RHICmdList.SetComputePipelineState(PipelineState, ComputeShader);
+}
+
+
 void SetGraphicsPipelineState(FRHICommandList& RHICmdList, const FGraphicsPipelineStateInitializer& Initializer, uint32 StencilRef, EApplyRendertargetOption ApplyFlags, bool bApplyAdditionalState, EPSOPrecacheResult PSOPrecacheResult)
 {
 #if PLATFORM_USE_FALLBACK_PSO
@@ -905,9 +912,9 @@ public:
 			for (TMyValue& OldPipelineState : DeleteArray)
 			{
 				//once in the delayed list this object should not be findable anymore, so the 0 should remain, making this safe
-#if PIPELINESTATECACHE_VERIFYTHREADSAFE
+			#if PIPELINESTATECACHE_VERIFYTHREADSAFE
 				check(OldPipelineState->InUseCount.GetValue() == 0);
-#endif
+			#endif
 				// Duplicate entries must wait for in progress compiles to complete.
 				// inprogress tasks could also remain in this container and deferred for the next tick.
 				bool bWaited = OldPipelineState->WaitCompletion();
@@ -931,9 +938,9 @@ public:
 		{
 			for ( const auto& DiscardIterator :  BackfillMap)
 			{
-#if PIPELINESTATECACHE_VERIFYTHREADSAFE
-				check( DiscardIterator.Value->InUseCount.GetValue() == 0);
-#endif
+			#if PIPELINESTATECACHE_VERIFYTHREADSAFE
+				check(DiscardIterator.Value->InUseCount.GetValue() == 0);
+			#endif
 				// Incomplete tasks should be put back to the current map. There should be no incomplete tasks encountered here.
 				bool bWaited = DiscardIterator.Value->WaitCompletion();
 				UE_CLOG(bWaited, LogRHI, Error, TEXT("Waited on a pipeline compile task while discarding retired PSOs."));
@@ -2869,14 +2876,11 @@ uint32 PipelineStateCache::NumActivePrecacheRequests()
 FRHIGraphicsPipelineState* ExecuteSetGraphicsPipelineState(FGraphicsPipelineState* GraphicsPipelineState)
 {
 	FRHIGraphicsPipelineState* RHIPipeline = GraphicsPipelineState->RHIPipeline;
-
 	GraphicsPipelineState->AddUse();
-
 #if PIPELINESTATECACHE_VERIFYTHREADSAFE
 	int32 Result = GraphicsPipelineState->InUseCount.Decrement();
 	check(Result >= 0);
 #endif
-	
 	return RHIPipeline;
 }
 
