Fixed Compute PSO lifetime issues with shader bundles and optimized Compute PSO management

diff --git a/Engine/Source/Runtime/RHI/Private/PipelineStateCache.cpp b/Engine/Source/Runtime/RHI/Private/PipelineStateCache.cpp
--- a/Engine/Source/Runtime/RHI/Private/PipelineStateCache.cpp
+++ b/Engine/Source/Runtime/RHI/Private/PipelineStateCache.cpp
@@ -291,6 +291,7 @@ void PipelineStateCache::PreCompileComplete()
 	}
 }
 
+extern RHI_API FComputePipelineState* GetComputePipelineState(FRHIComputeCommandList& RHICmdList, FRHIComputeShader* ComputeShader);
 extern RHI_API FRHIComputePipelineState* ExecuteSetComputePipelineState(FComputePipelineState* ComputePipelineState);
 extern RHI_API FRHIGraphicsPipelineState* ExecuteSetGraphicsPipelineState(FGraphicsPipelineState* GraphicsPipelineState);
 
@@ -627,6 +628,13 @@ public:
 #endif
 };
 
+RHI_API FRHIComputePipelineState* GetRHIComputePipelineState(FComputePipelineState* PipelineState)
+{
+	ensure(PipelineState->RHIPipeline);
+	PipelineState->CompletionEvent = nullptr;
+	return PipelineState->RHIPipeline;
+}
+
 //extern RHI_API FRHIRayTracingPipelineState* GetRHIRayTracingPipelineState(FRayTracingPipelineState* PipelineState);
 RHI_API FRHIRayTracingPipelineState* GetRHIRayTracingPipelineState(FRayTracingPipelineState* PipelineState)
 {
@@ -681,16 +689,21 @@ bool IsPrecachedPSO(const FGraphicsPipelineStateInitializer& Initializer)
 	return Initializer.bFromPSOFileCache || Initializer.bPSOPrecache;
 }
 
-void SetComputePipelineState(FRHIComputeCommandList& RHICmdList, FRHIComputeShader* ComputeShader)
+FComputePipelineState* GetComputePipelineState(FRHIComputeCommandList& RHICmdList, FRHIComputeShader* ComputeShader)
 {
 	FComputePipelineState* PipelineState = PipelineStateCache::GetAndOrCreateComputePipelineState(RHICmdList, ComputeShader, false);
 #if PIPELINESTATECACHE_VERIFYTHREADSAFE
 	int32 Result = PipelineState->InUseCount.Increment();
 	check(Result >= 1);
 #endif
-	RHICmdList.SetComputePipelineState(PipelineState, ComputeShader);
+	return PipelineState;
 }
 
+void SetComputePipelineState(FRHIComputeCommandList& RHICmdList, FRHIComputeShader* ComputeShader)
+{
+	FComputePipelineState* PipelineState = GetComputePipelineState(RHICmdList, ComputeShader);
+	RHICmdList.SetComputePipelineState(PipelineState, ComputeShader);
+}
 
 void SetGraphicsPipelineState(FRHICommandList& RHICmdList, const FGraphicsPipelineStateInitializer& Initializer, uint32 StencilRef, EApplyRendertargetOption ApplyFlags, bool bApplyAdditionalState, EPSOPrecacheResult PSOPrecacheResult)
 {
diff --git a/Engine/Source/Runtime/RHI/Public/RHICommandList.h b/Engine/Source/Runtime/RHI/Public/RHICommandList.h
--- a/Engine/Source/Runtime/RHI/Public/RHICommandList.h
+++ b/Engine/Source/Runtime/RHI/Public/RHICommandList.h
@@ -2473,9 +2473,11 @@ FRHICOMMAND_MACRO(FRHICommandSetRayTracingBindings)
 template<> RHI_API void FRHICommandSetShaderParameters           <FRHIComputeShader>::Execute(FRHICommandListBase& CmdList);
 template<> RHI_API void FRHICommandSetShaderUnbinds              <FRHIComputeShader>::Execute(FRHICommandListBase& CmdList);
 
-extern RHI_API FRHIComputePipelineState* ExecuteSetComputePipelineState(FComputePipelineState* ComputePipelineState);
-extern RHI_API FRHIGraphicsPipelineState* ExecuteSetGraphicsPipelineState(class FGraphicsPipelineState* GraphicsPipelineState);
-extern RHI_API FRHIRayTracingPipelineState* GetRHIRayTracingPipelineState(FRayTracingPipelineState*);
+extern RHI_API FRHIComputePipelineState*	ExecuteSetComputePipelineState(FComputePipelineState* ComputePipelineState);
+extern RHI_API FRHIGraphicsPipelineState*	ExecuteSetGraphicsPipelineState(class FGraphicsPipelineState* GraphicsPipelineState);
+extern RHI_API FComputePipelineState*		GetComputePipelineState(FRHIComputeCommandList& RHICmdList, FRHIComputeShader* ComputeShader);
+extern RHI_API FRHIComputePipelineState*	GetRHIComputePipelineState(FComputePipelineState*);
+extern RHI_API FRHIRayTracingPipelineState*	GetRHIRayTracingPipelineState(FRayTracingPipelineState*);
 
 class FRHIComputeCommandList : public FRHICommandListBase
 {
