Allow byte address buffers to be used for indirect draw/dispatch arguments.

#rnx

diff --git a/Engine/Source/Runtime/RHI/Public/RHIValidation.h b/Engine/Source/Runtime/RHI/Public/RHIValidation.h
--- a/Engine/Source/Runtime/RHI/Public/RHIValidation.h
+++ b/Engine/Source/Runtime/RHI/Public/RHIValidation.h
@@ -47,7 +47,7 @@ public:
 	static inline void ValidateIndirectArgsBuffer(FRHIBuffer* ArgumentBuffer, uint32 ArgumentOffset, uint32 ArgumentSize, uint32 ArgumentsBoundarySize)
 	{
 		auto GetBufferDesc = [&]() -> FString { return FString::Printf(TEXT("Buffer: %s, Size: %u, Stride: %u, Offset: %u, ArgSize: %u"), ArgumentBuffer->GetDebugName(), ArgumentBuffer->GetSize(), ArgumentBuffer->GetStride(), ArgumentOffset, ArgumentSize); };
-		RHI_VALIDATION_CHECK(EnumHasAnyFlags(ArgumentBuffer->GetUsage(), EBufferUsageFlags::VertexBuffer), *FString::Printf(TEXT("Indirect argument buffer must be a vertex buffer to be used as an indirect dispatch parameter. %s"), *GetBufferDesc()));
+		RHI_VALIDATION_CHECK(EnumHasAnyFlags(ArgumentBuffer->GetUsage(), EBufferUsageFlags::VertexBuffer | EBufferUsageFlags::ByteAddressBuffer), *FString::Printf(TEXT("Indirect argument buffer must be a vertex or byte address buffer to be used as an indirect dispatch parameter. %s"), *GetBufferDesc()));
 		RHI_VALIDATION_CHECK(EnumHasAnyFlags(ArgumentBuffer->GetUsage(), EBufferUsageFlags::DrawIndirect), *FString::Printf(TEXT("Indirect dispatch parameter buffer was not flagged with BUF_DrawIndirect. %s"), *GetBufferDesc()));
 		RHI_VALIDATION_CHECK((ArgumentOffset % 4) == 0, *FString::Printf(TEXT("Indirect argument offset must be a multiple of 4. %s"), *GetBufferDesc()));
 		RHI_VALIDATION_CHECK((ArgumentOffset + ArgumentSize) <= ArgumentBuffer->GetSize(), *FString::Printf(TEXT("Indirect argument doesn't fit in the buffer. %s"), *GetBufferDesc()));
diff --git a/Engine/Source/Runtime/RenderCore/Public/RenderGraphUtils.h b/Engine/Source/Runtime/RenderCore/Public/RenderGraphUtils.h
--- a/Engine/Source/Runtime/RenderCore/Public/RenderGraphUtils.h
+++ b/Engine/Source/Runtime/RenderCore/Public/RenderGraphUtils.h
@@ -445,7 +445,7 @@ namespace FComputeShaderUtils
 
 	inline void ValidateIndirectArgsBuffer(const FRDGBufferRef IndirectArgsBuffer, uint32 IndirectArgOffset)
 	{
-		checkf(EnumHasAnyFlags(IndirectArgsBuffer->Desc.Usage, EBufferUsageFlags::VertexBuffer), TEXT("The buffer %s needs to be a vertex buffer to be used as an indirect dispatch parameters"), IndirectArgsBuffer->Name);
+		checkf(EnumHasAnyFlags(IndirectArgsBuffer->Desc.Usage, EBufferUsageFlags::VertexBuffer | EBufferUsageFlags::ByteAddressBuffer), TEXT("The buffer %s needs to be a vertex or byte address buffer to be used as an indirect dispatch parameters"), IndirectArgsBuffer->Name);
 		checkf(EnumHasAnyFlags(IndirectArgsBuffer->Desc.Usage, EBufferUsageFlags::DrawIndirect), TEXT("The buffer %s for indirect dispatch parameters was not flagged with BUF_DrawIndirect"), IndirectArgsBuffer->Name);
 		ValidateIndirectArgsBuffer(IndirectArgsBuffer->GetSize(), IndirectArgOffset);
 	}
