Minor comment fixup in PipelineStateCache.h

diff --git a/Engine/Source/Runtime/RHI/Public/PipelineStateCache.h b/Engine/Source/Runtime/RHI/Public/PipelineStateCache.h
--- a/Engine/Source/Runtime/RHI/Public/PipelineStateCache.h
+++ b/Engine/Source/Runtime/RHI/Public/PipelineStateCache.h
@@ -47,7 +47,7 @@ enum class EPSOPrecacheResult
 	Missed,				//< PSO precache miss, needs to be compiled at draw time
 	TooLate,			//< PSO precache request still compiling when needed
 	NotSupported,		//< PSO precache not supported (VertexFactory or MeshPassProcessor doesn't support/implement precaching)
-	Untracked,			//< PSO is not tracked at all (Compute or not coming from MeshDrawCommands)
+	Untracked,			//< PSO is not tracked at all (Global shader or not coming from MeshDrawCommands)
 };
 
 // Unique request ID of PSOPrecache which can be used to boost the priority of a PSO precache requests if it's needed for rendering
