Add high end ray tracing effects platform dependent flag bSupportsHighEndRayTracingEffects and apply it to guide translucent shadow.

diff --git a/Engine/Config/Windows/DataDrivenPlatformInfo.ini b/Engine/Config/Windows/DataDrivenPlatformInfo.ini
--- a/Engine/Config/Windows/DataDrivenPlatformInfo.ini
+++ b/Engine/Config/Windows/DataDrivenPlatformInfo.ini
@@ -45,6 +45,7 @@ bSupportsRayTracingCallableShaders=true
 bSupportsRayTracingProceduralPrimitive=true
 bSupportsRayTracingIndirectInstanceData=true
 bSupportsPathTracing=true
+bSupportsHighEndRayTracingEffects=true
 bSupportsGPUScene=true
 bSupportsNanite=false
 bSupportsSceneDataCompressedTransforms=true
@@ -95,6 +96,7 @@ bSupportsRayTracingCallableShaders=true
 bSupportsRayTracingProceduralPrimitive=true
 bSupportsRayTracingIndirectInstanceData=true
 bSupportsPathTracing=true
+bSupportsHighEndRayTracingEffects=true
 bSupportsGPUScene=true
 bSupportsNanite=true
 bSupportsSceneDataCompressedTransforms=true
diff --git a/Engine/Shaders/Private/RayTracing/RayTracingOcclusionRGS.usf b/Engine/Shaders/Private/RayTracing/RayTracingOcclusionRGS.usf
--- a/Engine/Shaders/Private/RayTracing/RayTracingOcclusionRGS.usf
+++ b/Engine/Shaders/Private/RayTracing/RayTracingOcclusionRGS.usf
@@ -23,6 +23,8 @@
 #include "RayTracingSphereLight.ush"
 #include "RayTracingCapsuleLight.ush"
 
+#define USE_TRANSLUCENT_SHADOW  (!COMPUTESHADER && PLATFORM_SUPPORTS_RAY_TRACING_HIGH_END_EFFECTS) // support translucent shadow path in high end ray tracing without inline ray tracing
+
 /** Light types */
 #define LIGHT_TYPE_DIRECTIONAL 0
 #define LIGHT_TYPE_POINT 1
@@ -116,6 +118,10 @@ bool GenerateOcclusionRay(
 			if (!(dot(LightDirection, LightParameters.Direction) >= CosAngle))
 			{
 				// outside the cone of the light, skip all work
+				RayOrigin = (float3)0;
+				RayDirection = (float3)0;
+				RayTMin = 0.0f;
+				RayTMax = 0.0f;
 				return false;
 			}
 		#endif
@@ -560,9 +566,7 @@ FOcclusionResult ComputeOcclusion(
 		}
 		else
 		{
-
-#if !COMPUTESHADER && 0 // Ignore translucent shadow path when using inline ray tracing
-			// Temporarily disable the code to investigate a compiling issue.
+#if USE_TRANSLUCENT_SHADOW
 			bool bHasClosestEffectiveTranslucentHit = false;
 
 			if (NeedToTraceTranslucentShadow)
diff --git a/Engine/Source/Runtime/Engine/Private/ShaderCompiler/ShaderCompiler.cpp b/Engine/Source/Runtime/Engine/Private/ShaderCompiler/ShaderCompiler.cpp
--- a/Engine/Source/Runtime/Engine/Private/ShaderCompiler/ShaderCompiler.cpp
+++ b/Engine/Source/Runtime/Engine/Private/ShaderCompiler/ShaderCompiler.cpp
@@ -8188,7 +8188,8 @@ void GlobalBeginCompileShader(
 	SET_SHADER_DEFINE(Input.Environment, PLATFORM_ALLOW_SCENE_DATA_COMPRESSED_TRANSFORMS, FDataDrivenShaderPlatformInfo::GetSupportSceneDataCompressedTransforms(EShaderPlatform(Target.Platform)) ? 1 : 0);
 	SET_SHADER_DEFINE(Input.Environment, PLATFORM_SUPPORTS_BUFFER_LOAD_TYPE_CONVERSION, RHISupportsBufferLoadTypeConversion(ShaderPlatform) ? 1 : 0);
 	SET_SHADER_DEFINE(Input.Environment, PLATFORM_SUPPORTS_UNIFORM_BUFFER_OBJECTS, FDataDrivenShaderPlatformInfo::GetSupportsUniformBufferObjects(EShaderPlatform(Target.Platform)) ? 1 : 0);
-	
+	SET_SHADER_DEFINE(Input.Environment, PLATFORM_SUPPORTS_RAY_TRACING_HIGH_END_EFFECTS, FDataDrivenShaderPlatformInfo::GetSupportsHighEndRayTracingEffects(EShaderPlatform(Target.Platform)) ? 1 : 0);
+
 	bool bEnableBindlessMacro = false;
 	if (RHIGetBindlessSupport(ShaderPlatform) != ERHIBindlessSupport::Unsupported && !Input.Environment.CompilerFlags.Contains(CFLAG_ForceBindful))
 	{
diff --git a/Engine/Source/Runtime/RHI/Private/DataDrivenShaderPlatformInfo.cpp b/Engine/Source/Runtime/RHI/Private/DataDrivenShaderPlatformInfo.cpp
--- a/Engine/Source/Runtime/RHI/Private/DataDrivenShaderPlatformInfo.cpp
+++ b/Engine/Source/Runtime/RHI/Private/DataDrivenShaderPlatformInfo.cpp
@@ -226,6 +226,7 @@ void FGenericDataDrivenShaderPlatformInfo::ParseDataDrivenShaderInfo(const FConf
 	GET_SECTION_BOOL_HELPER(bSupportsRayTracingTraversalStatistics);
 	GET_SECTION_BOOL_HELPER(bSupportsRayTracingIndirectInstanceData);
 	GET_SECTION_BOOL_HELPER(bSupportsPathTracing);
+	GET_SECTION_BOOL_HELPER(bSupportsHighEndRayTracingEffects);
 	GET_SECTION_BOOL_HELPER(bSupportsByteBufferComputeShaders);
 	GET_SECTION_BOOL_HELPER(bSupportsGPUScene);
 	GET_SECTION_BOOL_HELPER(bSupportsPrimitiveShaders);
diff --git a/Engine/Source/Runtime/RHI/Public/DataDrivenShaderPlatformInfo.h b/Engine/Source/Runtime/RHI/Public/DataDrivenShaderPlatformInfo.h
--- a/Engine/Source/Runtime/RHI/Public/DataDrivenShaderPlatformInfo.h
+++ b/Engine/Source/Runtime/RHI/Public/DataDrivenShaderPlatformInfo.h
@@ -50,6 +50,7 @@ class FGenericDataDrivenShaderPlatformInfo
 	uint32 bSupportsRayTracingProceduralPrimitive : 1;
 	uint32 bSupportsRayTracingTraversalStatistics : 1;
 	uint32 bSupportsRayTracingIndirectInstanceData : 1; // Whether instance transforms can be copied from the GPU to the TLAS instances buffer
+	uint32 bSupportsHighEndRayTracingEffects : 1; // Whether fully-featured RT effects can be used on the platform (with translucent shadow, etc.)
 	uint32 bSupportsPathTracing : 1; // Whether real-time path tracer is supported on this platform (avoids compiling unnecessary shaders)
 	uint32 bSupportsGPUScene : 1;
 	uint32 bSupportsByteBufferComputeShaders : 1;
@@ -381,6 +382,12 @@ public:
 		return Infos[Platform].bSupportsRayTracing && Infos[Platform].bSupportsPathTracing;
 	}
 
+	static FORCEINLINE_DEBUGGABLE const bool GetSupportsHighEndRayTracingEffects(const FStaticShaderPlatform Platform)
+	{
+		check(IsValid(Platform));
+		return Infos[Platform].bSupportsRayTracing && Infos[Platform].bSupportsHighEndRayTracingEffects;
+	}
+
 	static FORCEINLINE_DEBUGGABLE const bool GetSupportsComputeFramework(const FStaticShaderPlatform Platform)
 	{
 		check(IsValid(Platform));
