Fixing serialization of removed value.

diff --git a/Engine/Source/Runtime/RHI/Private/PipelineFileCache.cpp b/Engine/Source/Runtime/RHI/Private/PipelineFileCache.cpp
--- a/Engine/Source/Runtime/RHI/Private/PipelineFileCache.cpp
+++ b/Engine/Source/Runtime/RHI/Private/PipelineFileCache.cpp
@@ -380,6 +380,21 @@ struct FPipelineCacheFileFormatPSOMetaData
 	}
 };
 
+RHI_API FArchive& operator<<(FArchive& Ar, FPipelineFileCacheRasterizerState& RasterizerStateInitializer)
+{
+	Ar << RasterizerStateInitializer.DepthBias;
+	Ar << RasterizerStateInitializer.SlopeScaleDepthBias;
+	Ar << RasterizerStateInitializer.FillMode;
+	Ar << RasterizerStateInitializer.CullMode;
+	Ar << RasterizerStateInitializer.DepthClipMode;
+	Ar << RasterizerStateInitializer.bAllowMSAA;
+	if (Ar.GameNetVer() < (uint32)EPipelineCacheFileFormatVersions::RemovingLineAA)
+	{
+		bool bEnableLineAA = false;
+		Ar << bEnableLineAA;
+	}
+	return Ar;
+}
 
 FString FPipelineFileCacheRasterizerState::ToString() const
 {
diff --git a/Engine/Source/Runtime/RHI/Public/PipelineFileCache.h b/Engine/Source/Runtime/RHI/Public/PipelineFileCache.h
--- a/Engine/Source/Runtime/RHI/Public/PipelineFileCache.h
+++ b/Engine/Source/Runtime/RHI/Public/PipelineFileCache.h
@@ -52,16 +52,7 @@ struct FPipelineFileCacheRasterizerState
 		return Initializer;
 	}
 	
-	friend FArchive& operator<<(FArchive& Ar,FPipelineFileCacheRasterizerState& RasterizerStateInitializer)
-	{
-		Ar << RasterizerStateInitializer.DepthBias;
-		Ar << RasterizerStateInitializer.SlopeScaleDepthBias;
-		Ar << RasterizerStateInitializer.FillMode;
-		Ar << RasterizerStateInitializer.CullMode;
-		Ar << RasterizerStateInitializer.DepthClipMode;
-		Ar << RasterizerStateInitializer.bAllowMSAA;
-		return Ar;
-	}
+	friend RHI_API FArchive& operator<<(FArchive& Ar, FPipelineFileCacheRasterizerState& RasterizerStateInitializer);
 
 	friend uint32 GetTypeHash(const FPipelineFileCacheRasterizerState &Key)
 	{
