Removing checks for RemovingLineAA since versioning doesn't work on stable pipeline caches.

diff --git a/Engine/Source/Runtime/RHI/Private/PipelineFileCache.cpp b/Engine/Source/Runtime/RHI/Private/PipelineFileCache.cpp
--- a/Engine/Source/Runtime/RHI/Private/PipelineFileCache.cpp
+++ b/Engine/Source/Runtime/RHI/Private/PipelineFileCache.cpp
@@ -388,7 +388,9 @@ RHI_API FArchive& operator<<(FArchive& Ar, FPipelineFileCacheRasterizerState& Ra
 	Ar << RasterizerStateInitializer.CullMode;
 	Ar << RasterizerStateInitializer.DepthClipMode;
 	Ar << RasterizerStateInitializer.bAllowMSAA;
-	if (Ar.GameNetVer() < (uint32)EPipelineCacheFileFormatVersions::RemovingLineAA)
+
+	// TODO: actually get versioning to work. Look for: MemReader.SetGameNetVer(FPipelineCacheFileFormatCurrentVersion);
+	//if (Ar.GameNetVer() < (uint32)EPipelineCacheFileFormatVersions::RemovingLineAA)
 	{
 		bool bEnableLineAA = false;
 		Ar << bEnableLineAA;
