Added missing buffer flags to RHI strings.

diff --git a/Engine/Source/Runtime/RHI/Private/RHIStrings.cpp b/Engine/Source/Runtime/RHI/Private/RHIStrings.cpp
--- a/Engine/Source/Runtime/RHI/Private/RHIStrings.cpp
+++ b/Engine/Source/Runtime/RHI/Private/RHIStrings.cpp
@@ -330,10 +330,14 @@ PRAGMA_ENABLE_DEPRECATION_WARNINGS
 		case BUF_FastVRAM:				return TEXT("BUF_FastVRAM");
 		case BUF_Shared:				return TEXT("BUF_Shared");
 		case BUF_AccelerationStructure:	return TEXT("BUF_AccelerationStructure");
-		case BUF_RayTracingScratch:		return TEXT("BUF_RayTracingScratch");
 		case BUF_VertexBuffer:			return TEXT("BUF_VertexBuffer");
 		case BUF_IndexBuffer:			return TEXT("BUF_IndexBuffer");
 		case BUF_StructuredBuffer:		return TEXT("BUF_StructuredBuffer");
+		case BUF_MultiGPUAllocate:		return TEXT("BUF_MultiGPUAllocate");
+		case BUF_MultiGPUGraphIgnore:	return TEXT("BUF_MultiGPUGraphIgnore");
+		case BUF_RayTracingScratch:		return TEXT("BUF_RayTracingScratch");
+		case BUF_NullResource:			return TEXT("BUF_NullResource");
+		case BUF_UniformBuffer:			return TEXT("BUF_UniformBuffer");
 		}
 	});
 }
