Don't allow 16Bit VALU for the Preview SP if the Editor SP does not support it

diff --git a/Engine/Source/Runtime/RHI/Private/DataDrivenShaderPlatformInfo.cpp b/Engine/Source/Runtime/RHI/Private/DataDrivenShaderPlatformInfo.cpp
--- a/Engine/Source/Runtime/RHI/Private/DataDrivenShaderPlatformInfo.cpp
+++ b/Engine/Source/Runtime/RHI/Private/DataDrivenShaderPlatformInfo.cpp
@@ -437,6 +437,7 @@ void FGenericDataDrivenShaderPlatformInfo::UpdatePreviewPlatforms()
 				Infos[ShaderPlatform].bSupportsMobileMultiView &= Infos[EditorSPForPreviewMaxFeatureLevel].bSupportsMobileMultiView;
 				Infos[ShaderPlatform].bSupportsVertexShaderSRVs = Infos[EditorSPForPreviewMaxFeatureLevel].bSupportsVertexShaderSRVs;
 				Infos[ShaderPlatform].bSupportsManualVertexFetch = Infos[EditorSPForPreviewMaxFeatureLevel].bSupportsManualVertexFetch;
+				Infos[ShaderPlatform].bSupportsRealTypes = Infos[EditorSPForPreviewMaxFeatureLevel].bSupportsRealTypes;
 				Infos[ShaderPlatform].bSupportsShaderRootConstants = false;
 				Infos[ShaderPlatform].bSupportsShaderBundleDispatch = false;
 				Infos[ShaderPlatform].bSupportsRenderTargetWriteMask = false;
