RHI: Removed spurious asserts (render thread checks) from RHIBindDebugLabelName in some RHI implementations.

Previously, the assert was needed because GetImmediateCommandList was used, which is only allowed to be used from the render thread. However, since CL 27966322 BindDebugLabelName gets a command list passed as a parameter, which makes the asserts unnecessary.

#rnx

diff --git a/Engine/Source/Runtime/RHI/Private/RHIValidation.cpp b/Engine/Source/Runtime/RHI/Private/RHIValidation.cpp
--- a/Engine/Source/Runtime/RHI/Private/RHIValidation.cpp
+++ b/Engine/Source/Runtime/RHI/Private/RHIValidation.cpp
@@ -761,8 +761,6 @@ static FString GetBreadcrumbPath()
 // FlushType: Thread safe
 void FValidationRHI::RHIBindDebugLabelName(FRHICommandListBase& RHICmdList, FRHITexture* Texture, const TCHAR* Name)
 {
-	check(IsInRenderingThread());
-
 	FString NameCopyRT = Name;
 	RHICmdList.EnqueueLambda([Texture, NameCopyRHIT = MoveTemp(NameCopyRT)](FRHICommandListBase& RHICmdList)
 	{
@@ -774,8 +772,6 @@ void FValidationRHI::RHIBindDebugLabelName(FRHICommandListBase& RHICmdList, FRHI
 
 void FValidationRHI::RHIBindDebugLabelName(FRHICommandListBase& RHICmdList, FRHIBuffer* Buffer, const TCHAR* Name)
 {
-	check(IsInRenderingThread());
-
 	FString NameCopyRT = Name;
 	RHICmdList.EnqueueLambda([Buffer, NameCopyRHIT = MoveTemp(NameCopyRT)](FRHICommandListBase& RHICmdList)
 	{
