Fix warning V502: Perhaps the '?:' operator works in a different way than it was expected. The '?:' operator has a lower priority than the '==' operator.
#rnx

diff --git a/Engine/Plugins/Experimental/AnimNext/Source/AnimNext/Private/Decorators/BlendByBool.cpp b/Engine/Plugins/Experimental/AnimNext/Source/AnimNext/Private/Decorators/BlendByBool.cpp
--- a/Engine/Plugins/Experimental/AnimNext/Source/AnimNext/Private/Decorators/BlendByBool.cpp
+++ b/Engine/Plugins/Experimental/AnimNext/Source/AnimNext/Private/Decorators/BlendByBool.cpp
@@ -86,11 +86,11 @@ namespace UE::AnimNext
 
 		if (ChildIndex == TRUE_CHILD_INDEX)
 		{
-			return DestinationChildIndex == TRUE_CHILD_INDEX ? 1.0f : 0.0f;
+			return (DestinationChildIndex == TRUE_CHILD_INDEX) ? 1.0f : 0.0f;
 		}
 		else if (ChildIndex == FALSE_CHILD_INDEX)
 		{
-			return DestinationChildIndex == FALSE_CHILD_INDEX ? 1.0f : 0.0f;
+			return (DestinationChildIndex == FALSE_CHILD_INDEX) ? 1.0f : 0.0f;
 		}
 		else
 		{
diff --git a/Engine/Source/Editor/UnrealEd/Private/SEditorViewport.cpp b/Engine/Source/Editor/UnrealEd/Private/SEditorViewport.cpp
--- a/Engine/Source/Editor/UnrealEd/Private/SEditorViewport.cpp
+++ b/Engine/Source/Editor/UnrealEd/Private/SEditorViewport.cpp
@@ -881,7 +881,7 @@ EVisibility SEditorViewport::GetCurrentFeatureLevelPreviewTextVisibility() const
 	{
 		UMaterialShaderQualitySettings* MaterialShaderQualitySettings = UMaterialShaderQualitySettings::Get();
 		const FName& PreviewPlatform = MaterialShaderQualitySettings->GetPreviewPlatform();
-		return PreviewPlatform != NAME_None ? EVisibility::SelfHitTestInvisible : EVisibility::Collapsed;
+		return (PreviewPlatform != NAME_None) ? EVisibility::SelfHitTestInvisible : EVisibility::Collapsed;
 	}
 	else
 	{
diff --git a/Engine/Source/Editor/UnrealEd/Private/WorkflowOrientedApp/WorkflowTabManager.cpp b/Engine/Source/Editor/UnrealEd/Private/WorkflowOrientedApp/WorkflowTabManager.cpp
--- a/Engine/Source/Editor/UnrealEd/Private/WorkflowOrientedApp/WorkflowTabManager.cpp
+++ b/Engine/Source/Editor/UnrealEd/Private/WorkflowOrientedApp/WorkflowTabManager.cpp
@@ -769,7 +769,7 @@ TSharedPtr<SDockTab> FDocumentTracker::OpenNewTab(TSharedPtr<FGenericTabHistory>
 			NewTabInfo->AddTabHistory(InTabHistory);
 		}
 
-		const FName DocumentId = DefaultDocumentId != NAME_None ? DefaultDocumentId : Factory->GetIdentifier();
+		const FName DocumentId = (DefaultDocumentId != NAME_None) ? DefaultDocumentId : Factory->GetIdentifier();
 
 		if (InOpenCause == ForceOpenNewDocument  || InOpenCause == OpenNewDocument)
 		{
diff --git a/Engine/Source/Runtime/D3D12RHI/Private/D3D12Buffer.cpp b/Engine/Source/Runtime/D3D12RHI/Private/D3D12Buffer.cpp
--- a/Engine/Source/Runtime/D3D12RHI/Private/D3D12Buffer.cpp
+++ b/Engine/Source/Runtime/D3D12RHI/Private/D3D12Buffer.cpp
@@ -11,7 +11,7 @@ D3D12Buffer.cpp: D3D Common code for buffers.
 FName GetRHIBufferClassName(const FName& ClassName)
 {
 	const static FLazyName RHIBufferName(TEXT("FRHIBuffer"));
-	return ClassName == NAME_None ? RHIBufferName : ClassName;
+	return (ClassName == NAME_None) ? RHIBufferName : ClassName;
 }
 
 extern int32 GD3D12BindResourceLabels;
diff --git a/Engine/Source/Runtime/RHI/Public/RHIResources.h b/Engine/Source/Runtime/RHI/Public/RHIResources.h
--- a/Engine/Source/Runtime/RHI/Public/RHIResources.h
+++ b/Engine/Source/Runtime/RHI/Public/RHIResources.h
@@ -388,7 +388,7 @@ struct FRHIResourceCreateInfo
 		ExtData = InExtData;
 	}
 
-	FName GetTraceClassName() const									{ const static FLazyName FRHIBufferName(TEXT("FRHIBuffer")); return ClassName == NAME_None ? FRHIBufferName : ClassName; }
+	FName GetTraceClassName() const									{ const static FLazyName FRHIBufferName(TEXT("FRHIBuffer")); return (ClassName == NAME_None) ? FRHIBufferName : ClassName; }
 
 	// for CreateTexture calls
 	FResourceBulkDataInterface* BulkData;
@@ -1644,7 +1644,7 @@ struct FRHITextureCreateDesc : public FRHITextureDesc
 	FRHITextureCreateDesc& SetFastVRAMPercentage(float In)                     { FastVRAMPercentage = uint8(FMath::Clamp(In, 0.f, 1.0f) * 0xFF); return *this; }
 	FRHITextureCreateDesc& SetClassName(const FName& InClassName)			   { ClassName = InClassName;				   return *this; }
 	FRHITextureCreateDesc& SetOwnerName(const FName& InOwnerName)			   { OwnerName = InOwnerName;                  return *this; }
-	FName GetTraceClassName() const											   { const static FLazyName FRHITextureName(TEXT("FRHITexture")); return ClassName == NAME_None ? FRHITextureName : ClassName; }
+	FName GetTraceClassName() const											   { const static FLazyName FRHITextureName(TEXT("FRHITexture")); return (ClassName == NAME_None) ? FRHITextureName : ClassName; }
 
 	/* The RHI access state that the resource will be created in. */
 	ERHIAccess InitialState = ERHIAccess::Unknown;
