Fix non-unity compile error due to missing Optional.h include

diff --git a/Engine/Source/Runtime/RHI/Public/RHITransition.h b/Engine/Source/Runtime/RHI/Public/RHITransition.h
--- a/Engine/Source/Runtime/RHI/Public/RHITransition.h
+++ b/Engine/Source/Runtime/RHI/Public/RHITransition.h
@@ -3,10 +3,10 @@
 #pragma once
 
 #include "Containers/ArrayView.h"
+#include "Misc/Optional.h"
 #include "RHIDefinitions.h"
 #include "RHIAccess.h"
 #include "RHIPipeline.h"
-
 #include "RHIValidationCommon.h"
 
 // The size in bytes of the storage required by the platform RHI for each resource transition.
