- Removed validation from ShaderParameterStruct because it incorrectly accesses the RHI from RDG
- Added support for Texture validation which was missing

diff --git a/Engine/Source/Runtime/RHI/Private/RHIValidation.cpp b/Engine/Source/Runtime/RHI/Private/RHIValidation.cpp
--- a/Engine/Source/Runtime/RHI/Private/RHIValidation.cpp
+++ b/Engine/Source/Runtime/RHI/Private/RHIValidation.cpp
@@ -2108,6 +2108,38 @@ namespace RHIValidation
 		return false;
 	}
 
+	bool ValidateDimension(EShaderCodeResourceBindingType Type, ETextureDimension Dimension, bool SRV)
+	{
+		// Ignore invalid types
+		if (Type == EShaderCodeResourceBindingType::Invalid)
+		{
+			return true;
+		}
+
+		if (Type == EShaderCodeResourceBindingType::Texture2D || Type == EShaderCodeResourceBindingType::RWTexture2D || Type == EShaderCodeResourceBindingType::Texture2DMS)
+		{
+			return Dimension == ETextureDimension::Texture2D;
+		}
+		else if (Type == EShaderCodeResourceBindingType::Texture2DArray || Type == EShaderCodeResourceBindingType::RWTexture2DArray)
+		{
+			return Dimension == ETextureDimension::Texture2DArray;
+		}
+		else if (Type == EShaderCodeResourceBindingType::Texture3D || Type == EShaderCodeResourceBindingType::RWTexture3D)
+		{
+			return Dimension == ETextureDimension::Texture3D;
+		}
+		else if (Type == EShaderCodeResourceBindingType::TextureCube || Type == EShaderCodeResourceBindingType::RWTextureCube)
+		{
+			return Dimension == ETextureDimension::TextureCube;
+		}
+		else if (Type == EShaderCodeResourceBindingType::TextureCubeArray)
+		{
+			return Dimension == ETextureDimension::TextureCubeArray;
+		}
+
+		return false;
+	}
+
 	bool ValidateBuffer(EShaderCodeResourceBindingType Type, FRHIViewDesc::EBufferType BufferType, bool SRV)
 	{
 		// Ignore invalid types
@@ -2244,6 +2276,50 @@ namespace RHIValidation
 #endif
 	}
 
+	/** Validates that the SRV is conform to what the shader expects */
+	void ValidateShaderResourceView(const FRHIShader* RHIShaderBase, uint32 BindIndex, const FRHITexture* Texture)
+	{
+#if RHI_INCLUDE_SHADER_DEBUG_DATA
+		if (Texture)
+		{
+			// Validate Type
+			if (!RHIShaderBase->DebugSRVTypeValidationData.Num())
+				return;
+
+			static const auto ShaderCodeValidationTypePredicate = [](const FShaderCodeValidationType& lhs, const FShaderCodeValidationType& rhs) -> bool { return lhs.BindPoint < rhs.BindPoint; };
+
+			FShaderCodeValidationType SRVValidationType = { BindIndex , EShaderCodeResourceBindingType::Invalid };
+			int32 FoundIndex = Algo::BinarySearch(RHIShaderBase->DebugSRVTypeValidationData, SRVValidationType, ShaderCodeValidationTypePredicate);
+
+			if (FoundIndex != INDEX_NONE)
+			{
+				EShaderCodeResourceBindingType ExpectedType = RHIShaderBase->DebugSRVTypeValidationData[FoundIndex].Type;
+
+				if (!ValidateDimension(ExpectedType, Texture->GetDesc().Dimension, true))
+				{
+					FString ErrorMessage = FString::Printf(TEXT("Shader %s: Dimension for Texture %s at BindIndex \"%d\" must match type declared in the shader"), 
+						RHIShaderBase->GetShaderName(),
+						*Texture->GetName().ToString(),
+						BindIndex);
+					ErrorMessage += FString::Printf(TEXT("\nBind point: %d, HLSL Type: %s, Actual Dimension: %s"),
+						BindIndex,
+						GetShaderCodeResourceBindingTypeName(ExpectedType),
+						GetTextureDimensionString(Texture->GetDesc().Dimension));
+					RHI_VALIDATION_CHECK(false, *ErrorMessage);
+				}
+			}
+			else
+			{
+				FString ErrorMessage = FString::Printf(TEXT("Shader %s: No bind point found at BindIndex \"%d\" possible UAV/SRV mismatch"), RHIShaderBase->GetShaderName(), BindIndex);
+					ErrorMessage += FString::Printf(TEXT("\nBind point: %d, Type: %s"),
+						BindIndex,
+						GetTextureDimensionString(Texture->GetDesc().Dimension));
+				RHI_VALIDATION_CHECK(false, *ErrorMessage);
+			}
+		}
+#endif
+	}
+
 	/** Validates that the UAV is conform to what the shader expects */
 	void ValidateUnorderedAccessView(const FRHIShader* RHIShaderBase, uint32 BindIndex, const FRHIUnorderedAccessView* UAV)
 	{
diff --git a/Engine/Source/Runtime/RHI/Public/RHIValidationCommon.h b/Engine/Source/Runtime/RHI/Public/RHIValidationCommon.h
--- a/Engine/Source/Runtime/RHI/Public/RHIValidationCommon.h
+++ b/Engine/Source/Runtime/RHI/Public/RHIValidationCommon.h
@@ -980,6 +980,7 @@ namespace RHIValidation
 
 	/** Validates that the SRV is conform to what the shader expects */
 	extern RHI_API void ValidateShaderResourceView(const FRHIShader* RHIShaderBase, uint32 BindIndex, const FRHIShaderResourceView* SRV);
+	extern RHI_API void ValidateShaderResourceView(const FRHIShader* RHIShaderBase, uint32 BindIndex, const FRHITexture* Texture);
 
 	/** Validates that the UAV conforms to what the shader expects */
 	extern RHI_API void ValidateUnorderedAccessView(const FRHIShader* RHIShaderBase, uint32 BindIndex, const FRHIUnorderedAccessView* SRV);
diff --git a/Engine/Source/Runtime/RHI/Public/RHIValidationContext.h b/Engine/Source/Runtime/RHI/Public/RHIValidationContext.h
--- a/Engine/Source/Runtime/RHI/Public/RHIValidationContext.h
+++ b/Engine/Source/Runtime/RHI/Public/RHIValidationContext.h
@@ -28,6 +28,10 @@ inline void ValidateShaderParameters(FRHIShader* RHIShader, RHIValidation::FTrac
 		case FRHIShaderParameterResource::EType::Texture:
 			if (FRHITexture* Texture = static_cast<FRHITexture*>(Parameter.Resource))
 			{
+				if (GRHIValidationEnabled)
+				{
+					RHIValidation::ValidateShaderResourceView(RHIShader, Parameter.Index, Texture);
+				}
 				Tracker->Assert(Texture->GetWholeResourceIdentitySRV(), InRequiredAccess);
 			}
 			break;
diff --git a/Engine/Source/Runtime/RenderCore/Private/ShaderParameterStruct.cpp b/Engine/Source/Runtime/RenderCore/Private/ShaderParameterStruct.cpp
--- a/Engine/Source/Runtime/RenderCore/Private/ShaderParameterStruct.cpp
+++ b/Engine/Source/Runtime/RenderCore/Private/ShaderParameterStruct.cpp
@@ -5,8 +5,6 @@
 =============================================================================*/
 
 #include "ShaderParameterStruct.h"
-#include "RenderGraphPrivate.h"
-#include "RHIValidationCommon.h"
 
 /** Context of binding a map. */
 struct FShaderParameterStructBindingContext
@@ -523,26 +521,6 @@ void ValidateShaderParameters(const TShaderRef<FShader>& Shader, const FShaderPa
 				{
 					EmitNullShaderParameterFatalError(Shader, ParametersMetadata, Parameter.ByteOffset);
 				}
-
-#if ENABLE_RHI_VALIDATION
-				if (GRHIValidationEnabled && GRDGAllowRHIAccess)
-				{
-					if (BaseType == UBMT_SRV)
-					{
-						if (const FRHIShaderResourceView* SRV = static_cast<const FRHIShaderResourceView*>(Resource))
-						{
-							RHIValidation::ValidateShaderResourceView(RHIShader, Parameter.BaseIndex, SRV);
-						}
-					}
-					else if (BaseType == UBMT_UAV)
-					{
-						if (const FRHIUnorderedAccessView* UAV = static_cast<const FRHIUnorderedAccessView*>(Resource))
-						{
-							RHIValidation::ValidateUnorderedAccessView(RHIShader, Parameter.BaseIndex, UAV);
-						}
-					}
-				}
-#endif 
 				break;
 			}
 			case UBMT_RDG_TEXTURE:
@@ -565,23 +543,6 @@ void ValidateShaderParameters(const TShaderRef<FShader>& Shader, const FShaderPa
 				{
 					EmitNullShaderParameterFatalError(Shader, ParametersMetadata, Parameter.ByteOffset);
 				}
-
-#if ENABLE_RHI_VALIDATION
-				if (GRHIValidationEnabled && GRDGAllowRHIAccess)
-				{
-					if (BaseType == UBMT_RDG_TEXTURE_SRV || BaseType == UBMT_RDG_BUFFER_SRV)
-					{
-						const FRHIShaderResourceView* SRV = static_cast<const FRHIShaderResourceView*>(GraphResource->GetRHI());
-						RHIValidation::ValidateShaderResourceView(RHIShader, Parameter.BaseIndex, SRV);
-
-					}
-					else if (BaseType == UBMT_RDG_TEXTURE_UAV || BaseType == UBMT_RDG_BUFFER_UAV)
-					{
-						const FRHIUnorderedAccessView* UAV = static_cast<const FRHIUnorderedAccessView*>(GraphResource->GetRHI());
-						RHIValidation::ValidateUnorderedAccessView(RHIShader, Parameter.BaseIndex, UAV);
-					}
-				}
-#endif
 				break;
 			}
 			default:
