Fix crash on exit due to pending raytracing PSO compilation.

#rnx

diff --git a/Engine/Source/Runtime/RHI/Private/PipelineStateCache.cpp b/Engine/Source/Runtime/RHI/Private/PipelineStateCache.cpp
--- a/Engine/Source/Runtime/RHI/Private/PipelineStateCache.cpp
+++ b/Engine/Source/Runtime/RHI/Private/PipelineStateCache.cpp
@@ -1658,11 +1658,19 @@ public:
 		FScopeLock ScopeLock(&CriticalSection);
 		for (auto& It : FullPipelines)
 		{
-			delete It.Value;
+			if (It.Value)
+			{
+				It.Value->WaitCompletion();
+				delete It.Value;
+			}
 		}
 		for (auto& It : PartialPipelines)
 		{
-			delete It.Value;
+			if (It.Value)
+			{
+				It.Value->WaitCompletion();
+				delete It.Value;
+			}
 		}
 	}
 
