Fixed unity build error in RHIValidationCommon.h

diff --git a/Engine/Source/Runtime/RHI/Public/RHIValidationCommon.h b/Engine/Source/Runtime/RHI/Public/RHIValidationCommon.h
--- a/Engine/Source/Runtime/RHI/Public/RHIValidationCommon.h
+++ b/Engine/Source/Runtime/RHI/Public/RHIValidationCommon.h
@@ -25,6 +25,7 @@ class FRHIShaderResourceView;
 class FRHIUniformBuffer;
 class FRHIViewableResource;
 class FRHIUnorderedAccessView;
+class FRHITexture;
 
 class FValidationComputeContext;
 class FValidationContext;
