Added Android static feature level and platform info

diff --git a/Engine/Config/Android/DataDrivenPlatformInfo.ini b/Engine/Config/Android/DataDrivenPlatformInfo.ini
--- a/Engine/Config/Android/DataDrivenPlatformInfo.ini
+++ b/Engine/Config/Android/DataDrivenPlatformInfo.ini
@@ -67,6 +67,8 @@ bSupportsClipDistance = false
 bSupportsShaderPipelines = false
 bSupportsVertexShaderSRVs=false
 bSupportsUniformBufferObjects = true
+bSupportsDxc=true
+bIsHlslcc=false
 FriendlyName=LOCTEXT("FriendlyShaderPlatformName_Android", "Android OpenGL Mobile")
 
 [ShaderPlatform VULKAN_ES3_1_ANDROID]
@@ -87,6 +89,8 @@ bSupportsVariableRateShading = true
 bSupportsVertexShaderSRVs=false
 bSupportsUniformBufferObjects = true
 bSupportsDualSourceBlending = true
+bSupportsDxc=true
+bIsHlslcc=false
 
 FriendlyName=LOCTEXT("FriendlyShaderPlatformName_AndroidVulkan", "Android Vulkan Mobile")
 
@@ -101,6 +105,7 @@ bRequiresDisableForwardLocalLights=true
 bWaterUsesSimpleForwardShading=true
 bSupportsVertexShaderSRVs=true
 bSupportsManualVertexFetch=true
+bSupportsUniformBufferObjects = true
 ; Raytracing and Lumen
 bSupportsRayTracing = true
 bSupportsRayTracingShaders = false
@@ -114,4 +119,6 @@ bSupportsGPUScene=true
 bSupportsLumenGI=true
 bSupportsDistanceFields=true
 bSupportsDualSourceBlending=false
-bSupportsVariableRateShading = false
\ No newline at end of file
+bSupportsVariableRateShading = false
+bSupportsDxc=true
+bIsHlslcc=false
\ No newline at end of file
diff --git a/Engine/Source/Runtime/OpenGLDrv/Private/OpenGLDevice.cpp b/Engine/Source/Runtime/OpenGLDrv/Private/OpenGLDevice.cpp
--- a/Engine/Source/Runtime/OpenGLDrv/Private/OpenGLDevice.cpp
+++ b/Engine/Source/Runtime/OpenGLDrv/Private/OpenGLDevice.cpp
@@ -879,7 +879,7 @@ static void InitRHICapabilitiesForGL()
 	GRHIAdapterInternalDriverVersion = ANSI_TO_TCHAR((const ANSICHAR*)glGetString(GL_VERSION));
 
 	// Shader platform & RHI feature level
-	GMaxRHIFeatureLevel = FOpenGL::GetFeatureLevel();
+	GMaxRHIFeatureLevel = ERHIFeatureLevel::ES3_1;
 	GMaxRHIShaderPlatform = FOpenGL::GetShaderPlatform();
 
 	// Log all supported extensions.
diff --git a/Engine/Source/Runtime/RHI/Public/Android/AndroidDataDrivenShaderPlatformInfo.inl b/Engine/Source/Runtime/RHI/Public/Android/AndroidDataDrivenShaderPlatformInfo.inl
new file mode 100644
--- /dev/null
+++ b/Engine/Source/Runtime/RHI/Public/Android/AndroidDataDrivenShaderPlatformInfo.inl
@@ -0,0 +1,19 @@
+// Copyright Epic Games, Inc. All Rights Reserved.
+
+#pragma once
+
+class FDataDrivenShaderPlatformInfo final : public FGenericDataDrivenShaderPlatformInfo
+{
+public:
+	IMPLEMENT_DDPSPI_SETTING(GetTargetsTiledGPU,			true);
+	IMPLEMENT_DDPSPI_SETTING(GetIsHlslcc,					false);
+	IMPLEMENT_DDPSPI_SETTING(GetSupportsDxc,				true);
+	IMPLEMENT_DDPSPI_SETTING(SupportsUniformBufferObjects,  true);
+
+#if USE_STATIC_FEATURE_LEVEL_ENUMS
+	IMPLEMENT_DDPSPI_SETTING_WITH_RETURN_TYPE(ERHIFeatureLevel::Type, GetMaxFeatureLevel, FStaticFeatureLevel(UE_ANDROID_STATIC_FEATURE_LEVEL));
+	IMPLEMENT_DDPSPI_SETTING(GetIsMobile, (UE_ANDROID_STATIC_FEATURE_LEVEL == ERHIFeatureLevel::ES3_1));
+	IMPLEMENT_DDPSPI_SETTING(SupportsManualVertexFetch, (UE_ANDROID_STATIC_FEATURE_LEVEL > ERHIFeatureLevel::ES3_1));
+	IMPLEMENT_DDPSPI_SETTING(SupportsMobileMultiView, (UE_ANDROID_STATIC_FEATURE_LEVEL == ERHIFeatureLevel::ES3_1));
+#endif
+};
diff --git a/Engine/Source/Runtime/RHI/Public/Android/AndroidStaticFeatureLevel.inl b/Engine/Source/Runtime/RHI/Public/Android/AndroidStaticFeatureLevel.inl
new file mode 100644
--- /dev/null
+++ b/Engine/Source/Runtime/RHI/Public/Android/AndroidStaticFeatureLevel.inl
@@ -0,0 +1,57 @@
+// Copyright Epic Games, Inc. All Rights Reserved.
+
+#pragma once
+
+#if defined(UE_ANDROID_SHADER_PLATFORM_VULKAN_SM5)
+	#define UE_ANDROID_STATIC_FEATURE_LEVEL ERHIFeatureLevel::SM5
+#else
+	#define UE_ANDROID_STATIC_FEATURE_LEVEL ERHIFeatureLevel::ES3_1
+#endif
+
+struct FStaticFeatureLevel
+{
+	inline FStaticFeatureLevel(const ERHIFeatureLevel::Type InFeatureLevel)
+	{
+		checkSlow(InFeatureLevel == UE_ANDROID_STATIC_FEATURE_LEVEL);
+	}
+
+	inline FStaticFeatureLevel(const TEnumAsByte<ERHIFeatureLevel::Type> InFeatureLevel)
+	{
+		checkSlow(InFeatureLevel.GetValue() == UE_ANDROID_STATIC_FEATURE_LEVEL);
+	}
+
+	inline operator ERHIFeatureLevel::Type() const
+	{
+		return UE_ANDROID_STATIC_FEATURE_LEVEL;
+	}
+
+	inline bool operator == (const ERHIFeatureLevel::Type Other) const
+	{
+		return Other == UE_ANDROID_STATIC_FEATURE_LEVEL;
+	}
+
+	inline bool operator != (const ERHIFeatureLevel::Type Other) const
+	{
+		return Other != UE_ANDROID_STATIC_FEATURE_LEVEL;
+	}
+
+	inline bool operator <= (const ERHIFeatureLevel::Type Other) const
+	{
+		return UE_ANDROID_STATIC_FEATURE_LEVEL <= Other;
+	}
+
+	inline bool operator < (const ERHIFeatureLevel::Type Other) const
+	{
+		return UE_ANDROID_STATIC_FEATURE_LEVEL < Other;
+	}
+
+	inline bool operator >= (const ERHIFeatureLevel::Type Other) const
+	{
+		return UE_ANDROID_STATIC_FEATURE_LEVEL >= Other;
+	}
+
+	inline bool operator > (const ERHIFeatureLevel::Type Other) const
+	{
+		return UE_ANDROID_STATIC_FEATURE_LEVEL > Other;
+	}
+};
diff --git a/Engine/Source/Runtime/RHI/Public/Android/AndroidStaticShaderPlatform.inl b/Engine/Source/Runtime/RHI/Public/Android/AndroidStaticShaderPlatform.inl
new file mode 100644
--- /dev/null
+++ b/Engine/Source/Runtime/RHI/Public/Android/AndroidStaticShaderPlatform.inl
@@ -0,0 +1,38 @@
+// Copyright Epic Games, Inc. All Rights Reserved.
+
+#pragma once
+
+#include "Misc/AssertionMacros.h"
+
+#if defined(UE_ANDROID_SHADER_PLATFORM_OPENGL_ES3_1)
+	#define UE_ANDROID_STATIC_SHADER_PLATFORM OPENGL_ES3_1_ANDROID
+#elif defined(UE_ANDROID_SHADER_PLATFORM_VULKAN_ES3_1)
+	#define UE_ANDROID_STATIC_SHADER_PLATFORM VULKAN_ES3_1_ANDROID
+#elif defined(UE_ANDROID_SHADER_PLATFORM_VULKAN_SM5)
+	#define UE_ANDROID_STATIC_SHADER_PLATFORM VULKAN_SM5_ANDROID
+#else
+	#error "Unknown Android static shader platform"
+#endif
+
+struct FStaticShaderPlatform
+{
+	inline FStaticShaderPlatform(const EShaderPlatform InPlatform)
+	{
+		checkSlow(UE_ANDROID_STATIC_SHADER_PLATFORM == InPlatform);
+	}
+
+	inline operator EShaderPlatform() const
+	{
+		return UE_ANDROID_STATIC_SHADER_PLATFORM;
+	}
+
+	inline bool operator == (const EShaderPlatform Other) const
+	{
+		return Other == UE_ANDROID_STATIC_SHADER_PLATFORM;
+	}
+	
+	inline bool operator != (const EShaderPlatform Other) const
+	{
+		return Other != UE_ANDROID_STATIC_SHADER_PLATFORM;
+	}
+};
diff --git a/Engine/Source/Runtime/VulkanRHI/Private/Android/VulkanAndroidPlatform.h b/Engine/Source/Runtime/VulkanRHI/Private/Android/VulkanAndroidPlatform.h
--- a/Engine/Source/Runtime/VulkanRHI/Private/Android/VulkanAndroidPlatform.h
+++ b/Engine/Source/Runtime/VulkanRHI/Private/Android/VulkanAndroidPlatform.h
@@ -101,7 +101,20 @@ public:
 
 	static bool RequiresMobileRenderer()
 	{
+		#if USE_STATIC_FEATURE_LEVEL_ENUMS
+		return (UE_ANDROID_STATIC_FEATURE_LEVEL == ERHIFeatureLevel::ES3_1);
+		#else	
 		return !FAndroidMisc::ShouldUseDesktopVulkan();
+		#endif
+	}
+
+	static ERHIFeatureLevel::Type GetFeatureLevel(ERHIFeatureLevel::Type RequestedFeatureLevel)
+	{
+		#if USE_STATIC_FEATURE_LEVEL_ENUMS
+		return UE_ANDROID_STATIC_FEATURE_LEVEL;
+		#else
+		return FVulkanGenericPlatform::GetFeatureLevel(RequestedFeatureLevel);
+		#endif
 	}
 
 	static bool HasCustomFrameTiming();
diff --git a/Engine/Source/Runtime/VulkanRHI/Private/VulkanGenericPlatform.cpp b/Engine/Source/Runtime/VulkanRHI/Private/VulkanGenericPlatform.cpp
--- a/Engine/Source/Runtime/VulkanRHI/Private/VulkanGenericPlatform.cpp
+++ b/Engine/Source/Runtime/VulkanRHI/Private/VulkanGenericPlatform.cpp
@@ -23,6 +23,15 @@ void FVulkanGenericPlatform::SetupFeatureLevels(TArrayView<EShaderPlatform> Shad
 	ShaderPlatformForFeatureLevel[ERHIFeatureLevel::SM6] = SP_VULKAN_SM6;
 }
 
+ERHIFeatureLevel::Type FVulkanGenericPlatform::GetFeatureLevel(ERHIFeatureLevel::Type InRequestedFeatureLevel)
+{
+	const bool bForceES3_1 = (FVulkanPlatform::RequiresMobileRenderer() ||
+		(InRequestedFeatureLevel == ERHIFeatureLevel::ES3_1) ||
+		FParse::Param(FCommandLine::Get(), TEXT("FeatureLevelES31")) || FParse::Param(FCommandLine::Get(), TEXT("FeatureLevelES3_1")));
+
+	return (!GIsEditor && bForceES3_1) ? ERHIFeatureLevel::ES3_1 : InRequestedFeatureLevel;
+}
+
 bool FVulkanGenericPlatform::PSOBinaryCacheMatches(FVulkanDevice* Device, const TArray<uint8>& DeviceCache)
 {
 	if (DeviceCache.Num() > 4)
diff --git a/Engine/Source/Runtime/VulkanRHI/Private/VulkanGenericPlatform.h b/Engine/Source/Runtime/VulkanRHI/Private/VulkanGenericPlatform.h
--- a/Engine/Source/Runtime/VulkanRHI/Private/VulkanGenericPlatform.h
+++ b/Engine/Source/Runtime/VulkanRHI/Private/VulkanGenericPlatform.h
@@ -61,6 +61,8 @@ public:
 
 	static bool RequiresMobileRenderer() { return false; }
 
+	static ERHIFeatureLevel::Type GetFeatureLevel(ERHIFeatureLevel::Type RequestedFeatureLevel);
+
 	// bInit=1 called at RHI init time, bInit=0 at RHI deinit time
 	static void OverridePlatformHandlers(bool bInit) {}
 
diff --git a/Engine/Source/Runtime/VulkanRHI/Private/VulkanRHI.cpp b/Engine/Source/Runtime/VulkanRHI/Private/VulkanRHI.cpp
--- a/Engine/Source/Runtime/VulkanRHI/Private/VulkanRHI.cpp
+++ b/Engine/Source/Runtime/VulkanRHI/Private/VulkanRHI.cpp
@@ -382,11 +382,7 @@ bool FVulkanDynamicRHIModule::IsSupported(ERHIFeatureLevel::Type FeatureLevel)
 
 FDynamicRHI* FVulkanDynamicRHIModule::CreateRHI(ERHIFeatureLevel::Type InRequestedFeatureLevel)
 {
-	const bool bForceES3_1 = (FVulkanPlatform::RequiresMobileRenderer() ||
-		(InRequestedFeatureLevel == ERHIFeatureLevel::ES3_1) ||
-		FParse::Param(FCommandLine::Get(), TEXT("FeatureLevelES31")) || FParse::Param(FCommandLine::Get(), TEXT("FeatureLevelES3_1")));
-
-	GMaxRHIFeatureLevel = (!GIsEditor && bForceES3_1) ? ERHIFeatureLevel::ES3_1 : InRequestedFeatureLevel;
+	GMaxRHIFeatureLevel = FVulkanPlatform::GetFeatureLevel(InRequestedFeatureLevel);
 	checkf(GMaxRHIFeatureLevel != ERHIFeatureLevel::Num, TEXT("Invalid feature level requested!"));
 
 	EShaderPlatform ShaderPlatformForFeatureLevel[ERHIFeatureLevel::Num];
