Log shader and pipeline hashes when a PSO fails to compile, so we can look these up in the SHK data.

#rnx

diff --git a/Engine/Source/Runtime/D3D12RHI/Private/Windows/WindowsD3D12PipelineState.cpp b/Engine/Source/Runtime/D3D12RHI/Private/Windows/WindowsD3D12PipelineState.cpp
--- a/Engine/Source/Runtime/D3D12RHI/Private/Windows/WindowsD3D12PipelineState.cpp
+++ b/Engine/Source/Runtime/D3D12RHI/Private/Windows/WindowsD3D12PipelineState.cpp
@@ -705,7 +705,7 @@ static void DumpGraphicsPSO(const FD3D12_GRAPHICS_PIPELINE_STATE_DESC& Desc, con
 		DumpShaderAsm(String, Desc.PS);
 	}
 
-	UE_LOG(LogD3D12RHI, Warning, TEXT("Failed to create Graphics PSO with hash 0x%s:\n%s"), Name, *String);
+	UE_LOG(LogD3D12RHI, Warning, TEXT("Failed to create graphics PSO with combined hash 0x%s:\n%s"), Name, *String);
 }
 
 static void DumpComputePSO(const FD3D12_COMPUTE_PIPELINE_STATE_DESC& Desc, const TCHAR* Name)
@@ -721,7 +721,7 @@ static void DumpComputePSO(const FD3D12_COMPUTE_PIPELINE_STATE_DESC& Desc, const
 		DumpShaderAsm(String, Desc.CS);
 	}
 
-	UE_LOG(LogD3D12RHI, Warning, TEXT("Failed to create Compute PSO with hash 0x%s:\n%s"), Name, *String);
+	UE_LOG(LogD3D12RHI, Warning, TEXT("Failed to create compute PSO with combined hash 0x%s:\n%s"), Name, *String);
 }
 
 // Thread-safe create graphics/compute pipeline state. Conditionally load/store the PSO using a Pipeline Library.
@@ -754,7 +754,7 @@ static HRESULT CreatePipelineStateFromStream(ID3D12PipelineState*& PSO, ID3D12De
 		hr = Device->CreatePipelineState(Desc, IID_PPV_ARGS(&PSO));
 		if (FAILED(hr))
 		{
-			UE_LOG(LogD3D12RHI, Error, TEXT("Failed to create PipelineState with hash %s"), Name);
+			UE_LOG(LogD3D12RHI, Error, TEXT("Failed to create pipeline state with combined hash %s, error %x."), Name, hr);
 		}
 	}
 
diff --git a/Engine/Source/Runtime/RHI/Private/PipelineStateCache.cpp b/Engine/Source/Runtime/RHI/Private/PipelineStateCache.cpp
--- a/Engine/Source/Runtime/RHI/Private/PipelineStateCache.cpp
+++ b/Engine/Source/Runtime/RHI/Private/PipelineStateCache.cpp
@@ -300,8 +300,37 @@ extern RHI_API FRHIGraphicsPipelineState* ExecuteSetGraphicsPipelineState(FGraph
 // This is fatal unless the compilation request is coming from the precaching system.
 static void HandlePipelineCreationFailure(const FGraphicsPipelineStateInitializer& Init)
 {
-	UE_LOG(LogRHI, Error, TEXT("Failed to create GraphicsPipeline"));
-	// Failure to compile is Fatal unless this is from the PSO file cache preloading.
+	FSHA1 PipelineHasher;
+	FString ShaderHashList;
+
+	const auto AddShaderHash = [&PipelineHasher, &ShaderHashList](const FRHIShader* Shader)
+	{
+		FSHAHash ShaderHash;
+		if (Shader)
+		{
+			ShaderHash = Shader->GetHash();
+			ShaderHashList.Appendf(TEXT("%s: %s, "), GetShaderFrequencyString(Shader->GetFrequency(), false), *ShaderHash.ToString());
+		}
+		PipelineHasher.Update(&ShaderHash.Hash[0], sizeof(FSHAHash));
+	};
+
+	// Log the shader and pipeline hashes, so we can look them up in the stable keys (SHK) file. Please note that NeedsShaderStableKeys must be set to
+	// true in the [DevOptions.Shaders] section of *Engine.ini in order for the cook process to produce SHK files for the shader libraries. The contents
+	// of those files can be extracted as text using the ShaderPipelineCacheTools commandlet, like this:
+	//		UnrealEditor-Cmd.exe ProjectName -run=ShaderPipelineCacheTools dump File.shk
+	// The pipeline hash is created by hashing together the individual shader hashes, see FShaderCodeLibraryPipeline::GetPipelineHash for details.
+	AddShaderHash(Init.BoundShaderState.GetVertexShader());
+	AddShaderHash(Init.BoundShaderState.GetMeshShader());
+	AddShaderHash(Init.BoundShaderState.GetAmplificationShader());
+	AddShaderHash(Init.BoundShaderState.GetPixelShader());
+	AddShaderHash(Init.BoundShaderState.GetGeometryShader());
+
+	PipelineHasher.Final();
+	FSHAHash PipelineHash;
+	PipelineHasher.GetHash(&PipelineHash.Hash[0]);
+
+	UE_LOG(LogRHI, Error, TEXT("Failed to create graphics pipeline, hashes: %sPipeline: %s."), *ShaderHashList, *PipelineHash.ToString());
+
 #if !(UE_BUILD_SHIPPING || UE_BUILD_TEST)
 	if(Init.BoundShaderState.VertexShaderRHI)
 	{
@@ -342,6 +371,7 @@ static void HandlePipelineCreationFailure(const FGraphicsPipelineStateInitialize
 	}
 	else if(!Init.bPSOPrecache)
 	{
+		// Precache requests are allowed to fail, but if the PSO is needed by a draw/dispatch command, we cannot continue.
 		UE_LOG(LogRHI, Fatal, TEXT("Shader compilation failures are Fatal."));
 	}
 }
@@ -350,7 +380,8 @@ static void HandlePipelineCreationFailure(const FGraphicsPipelineStateInitialize
 // This is fatal unless the compilation request is coming from the precaching system.
 static void HandlePipelineCreationFailure(const FRHIComputeShader* ComputeShader, bool bPrecache)
 {
-	UE_LOG(LogRHI, Error, TEXT("Failed to create ComputePipeline"));
+	// Dump the shader hash so it can be looked up in the SHK data. See the previous function for details.
+	UE_LOG(LogRHI, Error, TEXT("Failed to create compute pipeline with hash %s."), *ComputeShader->GetHash().ToString());
 
 #if !(UE_BUILD_SHIPPING || UE_BUILD_TEST)
 	UE_LOG(LogRHI, Error, TEXT("Shader: %s"), ComputeShader->GetShaderName());
