Adding Barycentrics support checks
- Adding EShaderCodeFeatures::Barycentrics
- Adding GRHIGlobals.SupportsBarycentricsSemantic
- Adding FDataDrivenShaderPlatformInfo::GetSupportsBarycentricsIntrinsics to control COMPILER_SUPPORTS_BARYCENTRIC_INTRINSICS
- Adding FDataDrivenShaderPlatformInfo::GetSupportsBarycentricsSemantic to control PLATFORM_SUPPORTS_BARYCENTRICS_SEMANTIC

diff --git a/Engine/Config/Windows/DataDrivenPlatformInfo.ini b/Engine/Config/Windows/DataDrivenPlatformInfo.ini
--- a/Engine/Config/Windows/DataDrivenPlatformInfo.ini
+++ b/Engine/Config/Windows/DataDrivenPlatformInfo.ini
@@ -135,6 +135,7 @@ SupportsMultiViewport=RuntimeGuaranteed
 bSupportsDualSourceBlending=true
 bSupportsNNEShaders=true
 MaxSamplers=32
+SupportsBarycentricsSemantic=RuntimeDependent
 
 [PreviewPlatform PCD3D_SM6]
 PlatformName=PCD3D_SM6
diff --git a/Engine/Source/Developer/Windows/ShaderFormatD3D/Private/D3DShaderCompilerDXC.cpp b/Engine/Source/Developer/Windows/ShaderFormatD3D/Private/D3DShaderCompilerDXC.cpp
--- a/Engine/Source/Developer/Windows/ShaderFormatD3D/Private/D3DShaderCompilerDXC.cpp
+++ b/Engine/Source/Developer/Windows/ShaderFormatD3D/Private/D3DShaderCompilerDXC.cpp
@@ -1244,6 +1244,11 @@ bool CompileAndProcessD3DShaderDXC(
 					EnumAddFlags(CodeFeatures.CodeFeatures, EShaderCodeFeatures::StencilRef);
 				}
 
+				if ((ShaderRequiresFlags & D3D_SHADER_REQUIRES_BARYCENTRICS) != 0)
+				{
+					EnumAddFlags(CodeFeatures.CodeFeatures, EShaderCodeFeatures::BarycentricsSemantic);
+				}
+
 				// We only need this to appear when using a DXC shader
 				ShaderCode.AddOptionalData<FShaderCodeFeatures>(CodeFeatures);
 
diff --git a/Engine/Source/Developer/Windows/ShaderFormatD3D/Private/ShaderFormatD3D.cpp b/Engine/Source/Developer/Windows/ShaderFormatD3D/Private/ShaderFormatD3D.cpp
--- a/Engine/Source/Developer/Windows/ShaderFormatD3D/Private/ShaderFormatD3D.cpp
+++ b/Engine/Source/Developer/Windows/ShaderFormatD3D/Private/ShaderFormatD3D.cpp
@@ -15,7 +15,7 @@ static FName NAME_PCD3D_SM5(TEXT("PCD3D_SM5"));
 static FName NAME_PCD3D_ES3_1(TEXT("PCD3D_ES31"));
 
 static const FGuid UE_SHADER_PCD3D_SHARED_VER = FGuid("dd4e6e76-4b48-4097-9ece-0f21118b7177");
-static const FGuid UE_SHADER_PCD3D_SM6_VER    = FGuid("78f5751f-107f-42b3-aa9d-7b871ebac5f6");
+static const FGuid UE_SHADER_PCD3D_SM6_VER    = FGuid("35A5B0FD-362E-45D5-8C27-9D94DAF94B58");
 static const FGuid UE_SHADER_PCD3D_SM5_VER    = FGuid("c2f60605-b462-4449-9f76-94bd2b1cd78e");
 static const FGuid UE_SHADER_PCD3D_ES3_1_VER  = FGuid("75466d2b-e169-40d8-bac5-1e2f9d43e0bb");
 
diff --git a/Engine/Source/Runtime/D3D12RHI/Private/D3D12Adapter.cpp b/Engine/Source/Runtime/D3D12RHI/Private/D3D12Adapter.cpp
--- a/Engine/Source/Runtime/D3D12RHI/Private/D3D12Adapter.cpp
+++ b/Engine/Source/Runtime/D3D12RHI/Private/D3D12Adapter.cpp
@@ -1060,6 +1060,12 @@ void FD3D12Adapter::InitializeDevices()
 				}
 			}
 
+			{
+				D3D12_FEATURE_DATA_D3D12_OPTIONS3 Features{};
+				RootDevice->CheckFeatureSupport(D3D12_FEATURE_D3D12_OPTIONS3, &Features, sizeof(Features));
+				GRHIGlobals.SupportsBarycentricsSemantic = Features.BarycentricsSupported;
+			}
+
 			{
 				D3D12_FEATURE_DATA_D3D12_OPTIONS4 Features{};
 				RootDevice->CheckFeatureSupport(D3D12_FEATURE_D3D12_OPTIONS4, &Features, sizeof(Features));
diff --git a/Engine/Source/Runtime/D3D12RHI/Private/D3D12Shaders.cpp b/Engine/Source/Runtime/D3D12RHI/Private/D3D12Shaders.cpp
--- a/Engine/Source/Runtime/D3D12RHI/Private/D3D12Shaders.cpp
+++ b/Engine/Source/Runtime/D3D12RHI/Private/D3D12Shaders.cpp
@@ -86,6 +86,11 @@ static bool ValidateShaderIsUsable(FD3D12ShaderData* InShader, EShaderFrequency
 	{
 		return false;
 	}
+
+	if (EnumHasAnyFlags(InShader->Features, EShaderCodeFeatures::BarycentricsSemantic) && !GRHIGlobals.SupportsBarycentricsSemantic)
+	{
+		return false;
+	}
 #endif
 
 	return true;
diff --git a/Engine/Source/Runtime/Engine/Private/ShaderCompiler/ShaderCompiler.cpp b/Engine/Source/Runtime/Engine/Private/ShaderCompiler/ShaderCompiler.cpp
--- a/Engine/Source/Runtime/Engine/Private/ShaderCompiler/ShaderCompiler.cpp
+++ b/Engine/Source/Runtime/Engine/Private/ShaderCompiler/ShaderCompiler.cpp
@@ -8235,6 +8235,8 @@ void GlobalBeginCompileShader(
 	SET_SHADER_DEFINE(Input.Environment, PLATFORM_SUPPORTS_BUFFER_LOAD_TYPE_CONVERSION, RHISupportsBufferLoadTypeConversion(ShaderPlatform) ? 1 : 0);
 	SET_SHADER_DEFINE(Input.Environment, PLATFORM_SUPPORTS_UNIFORM_BUFFER_OBJECTS, FDataDrivenShaderPlatformInfo::GetSupportsUniformBufferObjects(EShaderPlatform(Target.Platform)) ? 1 : 0);
 	SET_SHADER_DEFINE(Input.Environment, PLATFORM_SUPPORTS_RAY_TRACING_HIGH_END_EFFECTS, FDataDrivenShaderPlatformInfo::GetSupportsHighEndRayTracingEffects(EShaderPlatform(Target.Platform)) ? 1 : 0);
+	SET_SHADER_DEFINE(Input.Environment, COMPILER_SUPPORTS_BARYCENTRIC_INTRINSICS, FDataDrivenShaderPlatformInfo::GetSupportsBarycentricsIntrinsics(EShaderPlatform(Target.Platform)));
+	SET_SHADER_DEFINE(Input.Environment, PLATFORM_SUPPORTS_BARYCENTRICS_SEMANTIC, FDataDrivenShaderPlatformInfo::GetSupportsBarycentricsSemantic(EShaderPlatform(Target.Platform)) != ERHIFeatureSupport::Unsupported);
 
 	bool bEnableBindlessMacro = false;
 	if (RHIGetBindlessSupport(ShaderPlatform) != ERHIBindlessSupport::Unsupported && !Input.Environment.CompilerFlags.Contains(CFLAG_ForceBindful))
diff --git a/Engine/Source/Runtime/RHI/Private/DataDrivenShaderPlatformInfo.cpp b/Engine/Source/Runtime/RHI/Private/DataDrivenShaderPlatformInfo.cpp
--- a/Engine/Source/Runtime/RHI/Private/DataDrivenShaderPlatformInfo.cpp
+++ b/Engine/Source/Runtime/RHI/Private/DataDrivenShaderPlatformInfo.cpp
@@ -292,6 +292,8 @@ void FGenericDataDrivenShaderPlatformInfo::ParseDataDrivenShaderInfo(const FConf
 	GET_SECTION_BOOL_HELPER(bSupportsUniformBufferObjects);
 	GET_SECTION_BOOL_HELPER(bRequiresBindfulUtilityShaders);
 	GET_SECTION_INT_HELPER(MaxSamplers);
+	GET_SECTION_BOOL_HELPER(SupportsBarycentricsIntrinsics);
+	GET_SECTION_SUPPORT_HELPER(SupportsBarycentricsSemantic);
 #undef GET_SECTION_BOOL_HELPER
 #undef GET_SECTION_INT_HELPER
 #undef GET_SECTION_SUPPORT_HELPER
diff --git a/Engine/Source/Runtime/RHI/Public/DataDrivenShaderPlatformInfo.h b/Engine/Source/Runtime/RHI/Public/DataDrivenShaderPlatformInfo.h
--- a/Engine/Source/Runtime/RHI/Public/DataDrivenShaderPlatformInfo.h
+++ b/Engine/Source/Runtime/RHI/Public/DataDrivenShaderPlatformInfo.h
@@ -118,9 +118,12 @@ class FGenericDataDrivenShaderPlatformInfo
 	uint32 bSupportsShaderPipelines : 1;
 	uint32 bSupportsUniformBufferObjects : 1;
 	uint32 bRequiresBindfulUtilityShaders : 1;
+	uint32 MaxSamplers : 8;
+	uint32 SupportsBarycentricsIntrinsics : 1;
+	uint32 SupportsBarycentricsSemantic : int32(ERHIFeatureSupport::NumBits);
+
 	// NOTE: When adding fields, you must also add to ParseDataDrivenShaderInfo!
 	uint32 bContainsValidPlatformInfo : 1;
-	uint32 MaxSamplers : 8;
 
 	FGenericDataDrivenShaderPlatformInfo()
 	{
@@ -766,6 +769,18 @@ public:
 		return Infos[Platform].MaxSamplers;
 	}
 
+	static FORCEINLINE_DEBUGGABLE const bool GetSupportsBarycentricsIntrinsics(const FStaticShaderPlatform Platform)
+	{
+		check(IsValid(Platform));
+		return Infos[Platform].SupportsBarycentricsIntrinsics;
+	}
+
+	static FORCEINLINE_DEBUGGABLE const ERHIFeatureSupport GetSupportsBarycentricsSemantic(const FStaticShaderPlatform Platform)
+	{
+		check(IsValid(Platform));
+		return ERHIFeatureSupport(Infos[Platform].SupportsBarycentricsSemantic);
+	}
+
 	static FORCEINLINE_DEBUGGABLE const bool IsValid(const FStaticShaderPlatform Platform)
 	{
 		return Infos[Platform].bContainsValidPlatformInfo;
diff --git a/Engine/Source/Runtime/RHI/Public/RHIGlobals.h b/Engine/Source/Runtime/RHI/Public/RHIGlobals.h
--- a/Engine/Source/Runtime/RHI/Public/RHIGlobals.h
+++ b/Engine/Source/Runtime/RHI/Public/RHIGlobals.h
@@ -630,6 +630,9 @@ struct FRHIGlobals
 
 	/** True if the RHI needs shader unbinds (SetShaderUnbinds). RHIs that don't need them can avoid creating extra commands. */
 	bool NeedsShaderUnbinds = false;
+
+	/** True if the RHI supports shaders with barycentrics */
+	bool SupportsBarycentricsSemantic = false;
 };
 
 extern RHI_API FRHIGlobals GRHIGlobals;
diff --git a/Engine/Source/Runtime/RenderCore/Public/ShaderCore.h b/Engine/Source/Runtime/RenderCore/Public/ShaderCore.h
--- a/Engine/Source/Runtime/RenderCore/Public/ShaderCore.h
+++ b/Engine/Source/Runtime/RenderCore/Public/ShaderCore.h
@@ -723,6 +723,7 @@ enum class EShaderCodeFeatures : uint16
 	BindlessResources       = 1 << 5,
 	BindlessSamplers        = 1 << 6,
 	StencilRef              = 1 << 7,
+	BarycentricsSemantic    = 1 << 8,
 };
 ENUM_CLASS_FLAGS(EShaderCodeFeatures);
 
