Better PSO initializer hash functions that reduce collisions.

- Since hashes of "seen" initializers is used in the bundled cache machinery (see FPipelineFileCacheManager::RunTimeToPSOUsage), hash collisions can prevent distinct PSOs from being recorded.
- Addresses long-standing Rolando's TODO :)
- No perf impact has been seen in PerfB testing.
- Came from a UDN licensee.

diff --git a/Engine/Source/Runtime/RHI/Private/PipelineStateCache.cpp b/Engine/Source/Runtime/RHI/Private/PipelineStateCache.cpp
--- a/Engine/Source/Runtime/RHI/Private/PipelineStateCache.cpp
+++ b/Engine/Source/Runtime/RHI/Private/PipelineStateCache.cpp
@@ -43,24 +43,57 @@ DECLARE_DWORD_ACCUMULATOR_STAT(TEXT("High Priority Compute PSO Precache Requests
 
 static inline uint32 GetTypeHash(const FBoundShaderStateInput& Input)
 {
-	return GetTypeHash(Input.VertexDeclarationRHI)
-		^ GetTypeHash(Input.VertexShaderRHI)
-		^ GetTypeHash(Input.PixelShaderRHI)
+	uint32 Hash = GetTypeHash(Input.VertexDeclarationRHI);
+	Hash = HashCombineFast(Hash, GetTypeHash(Input.VertexShaderRHI));
+	Hash = HashCombineFast(Hash, GetTypeHash(Input.PixelShaderRHI));
 #if PLATFORM_SUPPORTS_MESH_SHADERS
-		^ GetTypeHash(Input.GetMeshShader())
-		^ GetTypeHash(Input.GetAmplificationShader())
+	Hash = HashCombineFast(Hash, GetTypeHash(Input.GetMeshShader()));
+	Hash = HashCombineFast(Hash, GetTypeHash(Input.GetAmplificationShader()));
 #endif
 #if PLATFORM_SUPPORTS_GEOMETRY_SHADERS
-		^ GetTypeHash(Input.GetGeometryShader())
+	Hash = HashCombineFast(Hash, GetTypeHash(Input.GetGeometryShader()));
 #endif
-		;
+	return Hash;
+}
+
+static inline uint32 GetTypeHash(const FImmutableSamplerState& Iss)
+{
+	return GetTypeHash(Iss.ImmutableSamplers);
+}
+
+inline uint32 GetTypeHash(const FExclusiveDepthStencil& Ds)
+{
+	return GetTypeHash(Ds.Value);
 }
 
 static inline uint32 GetTypeHash(const FGraphicsPipelineStateInitializer& Initializer)
 {
-	//#todo-rco: Hash!
-	return (GetTypeHash(Initializer.BoundShaderState) | (Initializer.NumSamples << 28)) ^ ((uint32)Initializer.PrimitiveType << 24) ^ GetTypeHash(Initializer.BlendState)
-		^ Initializer.RenderTargetsEnabled ^ GetTypeHash(Initializer.RasterizerState) ^ GetTypeHash(Initializer.DepthStencilState) ^ GetTypeHash(Initializer.MultiViewCount);
+	uint32 Hash = GetTypeHash(Initializer.BoundShaderState);
+	Hash = HashCombineFast(Hash, GetTypeHash(Initializer.BlendState));
+	Hash = HashCombineFast(Hash, GetTypeHash(Initializer.RasterizerState));
+	Hash = HashCombineFast(Hash, GetTypeHash(Initializer.DepthStencilState));
+	Hash = HashCombineFast(Hash, GetTypeHash(Initializer.ImmutableSamplerState));
+
+	Hash = HashCombineFast(Hash, GetTypeHash(Initializer.PrimitiveType));
+	Hash = HashCombineFast(Hash, GetTypeHash(Initializer.RenderTargetsEnabled));
+	Hash = HashCombineFast(Hash, GetTypeHash(Initializer.RenderTargetFormats));
+	Hash = HashCombineFast(Hash, GetTypeHash(Initializer.RenderTargetFlags));
+	Hash = HashCombineFast(Hash, GetTypeHash(Initializer.DepthStencilTargetFormat));
+	Hash = HashCombineFast(Hash, GetTypeHash(Initializer.DepthStencilTargetFlag));
+	Hash = HashCombineFast(Hash, GetTypeHash(Initializer.DepthTargetLoadAction));
+	Hash = HashCombineFast(Hash, GetTypeHash(Initializer.DepthTargetStoreAction));
+	Hash = HashCombineFast(Hash, GetTypeHash(Initializer.StencilTargetLoadAction));
+	Hash = HashCombineFast(Hash, GetTypeHash(Initializer.StencilTargetStoreAction));
+	Hash = HashCombineFast(Hash, GetTypeHash(Initializer.DepthStencilAccess));
+	Hash = HashCombineFast(Hash, GetTypeHash(Initializer.NumSamples));
+	Hash = HashCombineFast(Hash, GetTypeHash(Initializer.SubpassHint));
+	Hash = HashCombineFast(Hash, GetTypeHash(Initializer.SubpassIndex));
+	Hash = HashCombineFast(Hash, GetTypeHash(Initializer.ConservativeRasterization));
+	Hash = HashCombineFast(Hash, GetTypeHash(Initializer.bDepthBounds));
+	Hash = HashCombineFast(Hash, GetTypeHash(Initializer.MultiViewCount));
+	Hash = HashCombineFast(Hash, GetTypeHash(Initializer.bHasFragmentDensityAttachment));
+	Hash = HashCombineFast(Hash, GetTypeHash(Initializer.ShadingRate));
+	return Hash;
 }
 
 constexpr int32 PSO_MISS_FRAME_HISTORY_SIZE = 3;
diff --git a/Engine/Source/Runtime/RHI/Public/RHIResources.h b/Engine/Source/Runtime/RHI/Public/RHIResources.h
--- a/Engine/Source/Runtime/RHI/Public/RHIResources.h
+++ b/Engine/Source/Runtime/RHI/Public/RHIResources.h
@@ -668,6 +668,7 @@ private:
 	{
 		return (Type)(Value & StencilMask);
 	}
+	friend uint32 GetTypeHash(const FExclusiveDepthStencil& Ds);
 };
 
 //
