[RHI]
* Fixed nopch compile error

diff --git a/Engine/Source/Runtime/RHI/Private/PipelineStateCache.cpp b/Engine/Source/Runtime/RHI/Private/PipelineStateCache.cpp
--- a/Engine/Source/Runtime/RHI/Private/PipelineStateCache.cpp
+++ b/Engine/Source/Runtime/RHI/Private/PipelineStateCache.cpp
@@ -18,6 +18,7 @@ PipelineStateCache.cpp: Pipeline state cache implementation.
 #include "RHIFwd.h"
 #include "RHIImmutableSamplerState.h"
 #include "Stats/StatsTrace.h"
+#include "Templates/TypeHash.h"
 
 // perform cache eviction each frame, used to stress the system and flush out bugs
 #define PSO_DO_CACHE_EVICT_EACH_FRAME 0
