[RHI]
* Fix for nopch with unity compile error related to GetTypeHash. Don't know what is causing it because error does not contain enough information. Solution was to put TypeHash.h include before StaticArray.h in RHI.h to make sure the GetTypeHash used in PipelineStateCache.cpp is finding the right function

diff --git a/Engine/Source/Runtime/RHI/Public/RHI.h b/Engine/Source/Runtime/RHI/Public/RHI.h
--- a/Engine/Source/Runtime/RHI/Public/RHI.h
+++ b/Engine/Source/Runtime/RHI/Public/RHI.h
@@ -8,6 +8,7 @@
 
 #include "Containers/ContainersFwd.h"
 #include "RHIDefinitions.h"
+#include "Templates/TypeHash.h" // This must be before StaticArray to resolve adl compile error in nopch unity builds
 #include "Containers/StaticArray.h"
 #include "Containers/StringFwd.h"
 #include "Math/IntRect.h"
