Minor DDSPI cleanup
- UpdatePreviewPlatforms was implemented very manually, made it use macros that are more descriptive.
- UpdatePreviewPlatforms should only be compiled into Editor builds.

diff --git a/Engine/Source/Runtime/RHI/Private/DataDrivenShaderPlatformInfo.cpp b/Engine/Source/Runtime/RHI/Private/DataDrivenShaderPlatformInfo.cpp
--- a/Engine/Source/Runtime/RHI/Private/DataDrivenShaderPlatformInfo.cpp
+++ b/Engine/Source/Runtime/RHI/Private/DataDrivenShaderPlatformInfo.cpp
@@ -363,10 +363,11 @@ void FGenericDataDrivenShaderPlatformInfo::Initialize()
 #if WITH_EDITOR
 				if (!FParse::Param(FCommandLine::Get(), TEXT("NoPreviewPlatforms")))
 				{
+					const FName& CurrentPlatformName = Infos[ShaderPlatform].Name;
+
 					for (const FPreviewPlatformMenuItem& Item : FDataDrivenPlatformInfoRegistry::GetAllPreviewPlatformMenuItems())
 					{
-						const FName PreviewPlatformName = *(Infos[ShaderPlatform].Name).ToString();
-						if (Item.ShaderPlatformToPreview == PreviewPlatformName)
+						if (Item.ShaderPlatformToPreview == CurrentPlatformName)
 						{
 							const EShaderPlatform PreviewShaderPlatform = EShaderPlatform(CustomShaderPlatform++);
 							FGenericDataDrivenShaderPlatformInfo& PreviewInfo = Infos[PreviewShaderPlatform];
@@ -393,71 +394,103 @@ void FGenericDataDrivenShaderPlatformInfo::Initialize()
 	bInitialized = true;
 }
 
+#if WITH_EDITOR
 void FGenericDataDrivenShaderPlatformInfo::UpdatePreviewPlatforms()
 {
-	for (int i = 0; i < EShaderPlatform::SP_NumPlatforms; ++i)
+	for (int32 PlatformIndex=0; PlatformIndex < SP_NumPlatforms; PlatformIndex++)
 	{
-		EShaderPlatform ShaderPlatform = EShaderPlatform(i);
-		if (IsValid(ShaderPlatform))
+		const EShaderPlatform PreviewPlatform = EShaderPlatform(PlatformIndex);
+		if (IsValid(PreviewPlatform) && GetIsPreviewPlatform(PreviewPlatform))
 		{
-			ERHIFeatureLevel::Type PreviewSPMaxFeatureLevel = Infos[ShaderPlatform].MaxFeatureLevel;
-			EShaderPlatform EditorSPForPreviewMaxFeatureLevel = GShaderPlatformForFeatureLevel[PreviewSPMaxFeatureLevel];
-			if (Infos[ShaderPlatform].bIsPreviewPlatform && EditorSPForPreviewMaxFeatureLevel < EShaderPlatform::SP_NumPlatforms)
+			const ERHIFeatureLevel::Type PreviewFeatureLevel = Infos[PreviewPlatform].MaxFeatureLevel;
+			const EShaderPlatform RuntimePlatform = GRHIGlobals.ShaderPlatformForFeatureLevel[PreviewFeatureLevel];
+
+			if (RuntimePlatform < SP_NumPlatforms)
 			{
-				Infos[ShaderPlatform].ShaderFormat = Infos[EditorSPForPreviewMaxFeatureLevel].ShaderFormat;
-				Infos[ShaderPlatform].Language = Infos[EditorSPForPreviewMaxFeatureLevel].Language;
-				Infos[ShaderPlatform].bIsHlslcc = Infos[EditorSPForPreviewMaxFeatureLevel].bIsHlslcc;
-				Infos[ShaderPlatform].bSupportsDxc = Infos[EditorSPForPreviewMaxFeatureLevel].bSupportsDxc;
-				Infos[ShaderPlatform].bSupportsGPUScene = Infos[EditorSPForPreviewMaxFeatureLevel].bSupportsGPUScene;
-				Infos[ShaderPlatform].MaxMeshShaderThreadGroupSize = Infos[EditorSPForPreviewMaxFeatureLevel].MaxMeshShaderThreadGroupSize;
-				Infos[ShaderPlatform].bIsPC = true;
-				Infos[ShaderPlatform].bSupportsDebugViewShaders = true;
-				Infos[ShaderPlatform].bIsConsole = false;
-				Infos[ShaderPlatform].bSupportsSceneDataCompressedTransforms = Infos[EditorSPForPreviewMaxFeatureLevel].bSupportsSceneDataCompressedTransforms;
-				Infos[ShaderPlatform].bSupportsNanite &= Infos[EditorSPForPreviewMaxFeatureLevel].bSupportsNanite;
-				Infos[ShaderPlatform].bSupportsLumenGI &= Infos[EditorSPForPreviewMaxFeatureLevel].bSupportsLumenGI;
-				Infos[ShaderPlatform].bSupportsPrimitiveShaders &= Infos[EditorSPForPreviewMaxFeatureLevel].bSupportsPrimitiveShaders;
-				Infos[ShaderPlatform].bSupportsUInt64ImageAtomics &= Infos[EditorSPForPreviewMaxFeatureLevel].bSupportsUInt64ImageAtomics;
-				Infos[ShaderPlatform].bSupportsGen5TemporalAA &= Infos[EditorSPForPreviewMaxFeatureLevel].bSupportsGen5TemporalAA;
-				Infos[ShaderPlatform].bSupportsInlineRayTracing &= Infos[EditorSPForPreviewMaxFeatureLevel].bSupportsInlineRayTracing;
-				Infos[ShaderPlatform].bSupportsRayTracingShaders &= Infos[EditorSPForPreviewMaxFeatureLevel].bSupportsRayTracingShaders;
-				Infos[ShaderPlatform].bSupportsMeshShadersTier0 &= Infos[EditorSPForPreviewMaxFeatureLevel].bSupportsMeshShadersTier0;
-				Infos[ShaderPlatform].bSupportsMeshShadersTier1 &= Infos[EditorSPForPreviewMaxFeatureLevel].bSupportsMeshShadersTier1;
+				FGenericDataDrivenShaderPlatformInfo& PreviewInfo = Infos[PreviewPlatform];
+				const FGenericDataDrivenShaderPlatformInfo& RuntimeInfo = Infos[RuntimePlatform];
+
+#define PREVIEW_USE_RUNTIME_VALUE(SettingName) \
+	PreviewInfo.SettingName = RuntimeInfo.SettingName
+
+#define PREVIEW_DISABLE_IF_RUNTIME_UNSUPPORTED(SettingName) \
+	PreviewInfo.SettingName &= RuntimeInfo.SettingName
+
+#define PREVIEW_FORCE_SETTING(SettingName, Value) \
+	PreviewInfo.SettingName = (Value)
+
+#define PREVIEW_FORCE_DISABLE(SettingName) \
+	PREVIEW_FORCE_SETTING(SettingName, false)
+
+				// Always inherit these core settings from the preview
+				PREVIEW_USE_RUNTIME_VALUE(ShaderFormat);
+				PREVIEW_USE_RUNTIME_VALUE(Language);
+				PREVIEW_USE_RUNTIME_VALUE(bIsHlslcc);
+				PREVIEW_USE_RUNTIME_VALUE(bSupportsDxc);
+
+				// Editor is always PC, never console and always supports debug view shaders
+				PREVIEW_FORCE_SETTING(bIsPC, true);
+				PREVIEW_FORCE_SETTING(bSupportsDebugViewShaders, true);
+				PREVIEW_FORCE_SETTING(bIsConsole, false);
+
+				// Settings that should be kept true if the runtime also supports it.
+				PREVIEW_DISABLE_IF_RUNTIME_UNSUPPORTED(bSupportsNanite);
+				PREVIEW_DISABLE_IF_RUNTIME_UNSUPPORTED(bSupportsLumenGI);
+				PREVIEW_DISABLE_IF_RUNTIME_UNSUPPORTED(bSupportsPrimitiveShaders);
+				PREVIEW_DISABLE_IF_RUNTIME_UNSUPPORTED(bSupportsUInt64ImageAtomics);
+				PREVIEW_DISABLE_IF_RUNTIME_UNSUPPORTED(bSupportsGen5TemporalAA);
+				PREVIEW_DISABLE_IF_RUNTIME_UNSUPPORTED(bSupportsInlineRayTracing);
+				PREVIEW_DISABLE_IF_RUNTIME_UNSUPPORTED(bSupportsRayTracingShaders);
+				PREVIEW_DISABLE_IF_RUNTIME_UNSUPPORTED(bSupportsMeshShadersTier0);
+				PREVIEW_DISABLE_IF_RUNTIME_UNSUPPORTED(bSupportsMeshShadersTier1);
+				PREVIEW_DISABLE_IF_RUNTIME_UNSUPPORTED(bSupportsMobileMultiView);
+
+				// Settings that need to match the runtime
+				PREVIEW_USE_RUNTIME_VALUE(bSupportsGPUScene);
+				PREVIEW_USE_RUNTIME_VALUE(MaxMeshShaderThreadGroupSize);
+				PREVIEW_USE_RUNTIME_VALUE(bSupportsSceneDataCompressedTransforms);
+				PREVIEW_USE_RUNTIME_VALUE(bSupportsVertexShaderSRVs);
+				PREVIEW_USE_RUNTIME_VALUE(bSupportsManualVertexFetch);
+				PREVIEW_USE_RUNTIME_VALUE(bSupportsRealTypes);
+				PREVIEW_USE_RUNTIME_VALUE(bSupportsUniformBufferObjects);
+
+				// Settings that will never be supported in preview
+				PREVIEW_FORCE_DISABLE(bSupportsShaderRootConstants);
+				PREVIEW_FORCE_DISABLE(bSupportsShaderBundleDispatch);
+				PREVIEW_FORCE_DISABLE(bSupportsRenderTargetWriteMask);
+				PREVIEW_FORCE_DISABLE(bSupportsIntrinsicWaveOnce);
+				PREVIEW_FORCE_DISABLE(bSupportsDOFHybridScattering);
+				PREVIEW_FORCE_DISABLE(bSupports4ComponentUAVReadWrite);
 
 				// Support for stereo features requires extra consideration. The editor may not use the same technique as the preview platform,
 				// particularly MobileMultiView may be substituted by a fallback path. In order to avoid inundating real mobile platforms
 				// with the properties needed for the desktop MMV fallback path, override them here with the editor ones to make MMV preview possible
-				if (Infos[ShaderPlatform].bSupportsMobileMultiView && !Infos[EditorSPForPreviewMaxFeatureLevel].bSupportsMobileMultiView)
+				if (PreviewInfo.bSupportsMobileMultiView && !RuntimeInfo.bSupportsMobileMultiView)
 				{
-					Infos[ShaderPlatform].bSupportsInstancedStereo = Infos[EditorSPForPreviewMaxFeatureLevel].bSupportsInstancedStereo;
-					Infos[ShaderPlatform].bSupportsVertexShaderLayer = Infos[EditorSPForPreviewMaxFeatureLevel].bSupportsVertexShaderLayer;
+					PREVIEW_USE_RUNTIME_VALUE(bSupportsInstancedStereo);
+					PREVIEW_USE_RUNTIME_VALUE(bSupportsVertexShaderLayer);
 				}
 				else
 				{
-					Infos[ShaderPlatform].bSupportsInstancedStereo &= Infos[EditorSPForPreviewMaxFeatureLevel].bSupportsInstancedStereo;
+					PREVIEW_DISABLE_IF_RUNTIME_UNSUPPORTED(bSupportsInstancedStereo);
 				}
-				Infos[ShaderPlatform].bSupportsMobileMultiView &= Infos[EditorSPForPreviewMaxFeatureLevel].bSupportsMobileMultiView;
-				Infos[ShaderPlatform].bSupportsVertexShaderSRVs = Infos[EditorSPForPreviewMaxFeatureLevel].bSupportsVertexShaderSRVs;
-				Infos[ShaderPlatform].bSupportsManualVertexFetch = Infos[EditorSPForPreviewMaxFeatureLevel].bSupportsManualVertexFetch;
-				Infos[ShaderPlatform].bSupportsRealTypes = Infos[EditorSPForPreviewMaxFeatureLevel].bSupportsRealTypes;
-				Infos[ShaderPlatform].bSupportsShaderRootConstants = false;
-				Infos[ShaderPlatform].bSupportsShaderBundleDispatch = false;
-				Infos[ShaderPlatform].bSupportsRenderTargetWriteMask = false;
-				Infos[ShaderPlatform].bSupportsIntrinsicWaveOnce = false;
-				Infos[ShaderPlatform].bSupportsDOFHybridScattering = false;
-				Infos[ShaderPlatform].bSupports4ComponentUAVReadWrite = false;
-				Infos[ShaderPlatform].bSupportsUniformBufferObjects = Infos[EditorSPForPreviewMaxFeatureLevel].bSupportsUniformBufferObjects;
-				Infos[ShaderPlatform].bContainsValidPlatformInfo = true;
+
+				// Make sure we're marked valid
+				PreviewInfo.bContainsValidPlatformInfo = true;
 
 				// Seeing as we are merging the two shader platforms merge the hash key as well, this way
 				// any changes in the editor feature level shader platform will dirty the preview key.
-				Infos[ShaderPlatform].ShaderPropertiesHash = HashCombine(Infos[ShaderPlatform].ShaderPropertiesHash, Infos[EditorSPForPreviewMaxFeatureLevel].ShaderPropertiesHash);
+				PreviewInfo.ShaderPropertiesHash = HashCombine(PreviewInfo.ShaderPropertiesHash, RuntimeInfo.ShaderPropertiesHash);
+
+#undef PREVIEW_FORCE_DISABLE
+#undef PREVIEW_FORCE_SETTING
+#undef PREVIEW_DISABLE_IF_RUNTIME_UNSUPPORTED
+#undef PREVIEW_USE_RUNTIME_VALUE
 			}
 		}
 	}
 }
 
-#if WITH_EDITOR
 FText FGenericDataDrivenShaderPlatformInfo::GetFriendlyName(const FStaticShaderPlatform Platform)
 {
 	if (IsRunningCommandlet() || GUsingNullRHI)
@@ -473,7 +506,7 @@ const EShaderPlatform FGenericDataDrivenShaderPlatformInfo::GetPreviewShaderPlat
 	check(IsValid(Platform));
 	return DataDrivenShaderPlatformInfoEditorOnlyInfos[Platform].PreviewShaderPlatformParent;
 }
-#endif
+#endif  // WITH_EDITOR
 
 const EShaderPlatform FGenericDataDrivenShaderPlatformInfo::GetShaderPlatformFromName(const FName ShaderPlatformName)
 {
