Adding command line flag to enable bindless in the Editor.

diff --git a/Engine/Source/Runtime/RHI/Private/RHI.cpp b/Engine/Source/Runtime/RHI/Private/RHI.cpp
--- a/Engine/Source/Runtime/RHI/Private/RHI.cpp
+++ b/Engine/Source/Runtime/RHI/Private/RHI.cpp
@@ -1233,10 +1233,24 @@ ERHIBindlessConfiguration RHIParseBindlessConfiguration(EShaderPlatform Platform
 {
 	const ERHIBindlessSupport BindlessSupport = RHIGetBindlessSupport(Platform);
 
+	if (BindlessSupport == ERHIBindlessSupport::Unsupported)
+	{
+		return ERHIBindlessConfiguration::Disabled;
+	}
+
+#if WITH_EDITOR
+	// We have to check the -bindless command line option here to make sure the shaders are compiled with bindless enabled too.
+	static const bool bCommandLine = FParse::Param(FCommandLine::Get(), TEXT("Bindless"));
+	if (bCommandLine)
+	{
+		return ERHIBindlessConfiguration::AllShaders;
+	}
+#endif
+
 	const ERHIBindlessConfiguration ConfigSetting = ParseConfigurationFromString(ConfigSettingString);
 	const ERHIBindlessConfiguration CVarSetting = ParseConfigurationFromString(CVarSettingString);
 
-	if (BindlessSupport == ERHIBindlessSupport::Unsupported || (ConfigSetting == ERHIBindlessConfiguration::Disabled && CVarSetting == ERHIBindlessConfiguration::Disabled))
+	if (ConfigSetting == ERHIBindlessConfiguration::Disabled && CVarSetting == ERHIBindlessConfiguration::Disabled)
 	{
 		return ERHIBindlessConfiguration::Disabled;
 	}
