Preparing for dependency cleanup.

diff --git a/Engine/Plugins/Developer/PixWinPlugin/Source/PixWinPlugin/Private/PixWinPluginModule.cpp b/Engine/Plugins/Developer/PixWinPlugin/Source/PixWinPlugin/Private/PixWinPluginModule.cpp
--- a/Engine/Plugins/Developer/PixWinPlugin/Source/PixWinPlugin/Private/PixWinPluginModule.cpp
+++ b/Engine/Plugins/Developer/PixWinPlugin/Source/PixWinPlugin/Private/PixWinPluginModule.cpp
@@ -8,6 +8,7 @@
 #include "UnrealClient.h"
 #include "Engine/GameViewportClient.h"
 #include "Engine/Engine.h"
+#include "Misc/CommandLine.h"
 
 #if !defined(WITH_PIX_EVENT_RUNTIME)
 #define WITH_PIX_EVENT_RUNTIME 0
diff --git a/Engine/Plugins/Runtime/ComputeFramework/Source/ComputeFramework/Private/ComputeGraph.cpp b/Engine/Plugins/Runtime/ComputeFramework/Source/ComputeFramework/Private/ComputeGraph.cpp
--- a/Engine/Plugins/Runtime/ComputeFramework/Source/ComputeFramework/Private/ComputeGraph.cpp
+++ b/Engine/Plugins/Runtime/ComputeFramework/Source/ComputeFramework/Private/ComputeGraph.cpp
@@ -16,6 +16,7 @@
 #include "ComputeFramework/ShaderParameterMetadataAllocation.h"
 #include "GameFramework/Actor.h"
 #include "Interfaces/ITargetPlatform.h"
+#include "Misc/App.h"
 #include "ShaderParameterMetadataBuilder.h"
 #include "UObject/Package.h"
 #include "DataDrivenShaderPlatformInfo.h"
diff --git a/Engine/Plugins/Runtime/WindowsDeviceProfileSelector/Source/WindowsDeviceProfileSelector/Private/WindowsDeviceProfileSelectorModule.cpp b/Engine/Plugins/Runtime/WindowsDeviceProfileSelector/Source/WindowsDeviceProfileSelector/Private/WindowsDeviceProfileSelectorModule.cpp
--- a/Engine/Plugins/Runtime/WindowsDeviceProfileSelector/Source/WindowsDeviceProfileSelector/Private/WindowsDeviceProfileSelectorModule.cpp
+++ b/Engine/Plugins/Runtime/WindowsDeviceProfileSelector/Source/WindowsDeviceProfileSelector/Private/WindowsDeviceProfileSelectorModule.cpp
@@ -5,6 +5,7 @@
 #include "RHI.h"
 #include "RHICommandList.h"
 #include "DynamicRHI.h"
+#include "Misc/App.h"
 #include "Misc/ConfigCacheIni.h"
 
 IMPLEMENT_MODULE(FWindowsDeviceProfileSelectorModule, WindowsDeviceProfileSelector);
diff --git a/Engine/Plugins/Runtime/nDisplay/Source/DisplayCluster/Private/Render/Containers/DisplayClusterRender_MeshResources.cpp b/Engine/Plugins/Runtime/nDisplay/Source/DisplayCluster/Private/Render/Containers/DisplayClusterRender_MeshResources.cpp
--- a/Engine/Plugins/Runtime/nDisplay/Source/DisplayCluster/Private/Render/Containers/DisplayClusterRender_MeshResources.cpp
+++ b/Engine/Plugins/Runtime/nDisplay/Source/DisplayCluster/Private/Render/Containers/DisplayClusterRender_MeshResources.cpp
@@ -3,6 +3,7 @@
 #include "DisplayClusterRender_MeshResources.h"
 
 #include "PipelineStateCache.h"
+#include "RHI.h"
 
 void FDisplayClusterMeshVertexDeclaration::InitRHI(FRHICommandListBase& RHICmdList)
 {
diff --git a/Engine/Plugins/Runtime/nDisplay/Source/DisplayCluster/Private/Render/Containers/DisplayClusterRender_Texture.cpp b/Engine/Plugins/Runtime/nDisplay/Source/DisplayCluster/Private/Render/Containers/DisplayClusterRender_Texture.cpp
--- a/Engine/Plugins/Runtime/nDisplay/Source/DisplayCluster/Private/Render/Containers/DisplayClusterRender_Texture.cpp
+++ b/Engine/Plugins/Runtime/nDisplay/Source/DisplayCluster/Private/Render/Containers/DisplayClusterRender_Texture.cpp
@@ -1,6 +1,7 @@
 // Copyright Epic Games, Inc. All Rights Reserved.
 
 #include "Render/Containers/DisplayClusterRender_Texture.h"
+#include "Misc/App.h"
 #include "RenderingThread.h"
 
 int32 GDisplayClusterRender_TextureCacheEnable = 1;
diff --git a/Engine/Plugins/VirtualProduction/Rivermax/RivermaxCore/Source/RivermaxCore/Private/RivermaxCoreModule.cpp b/Engine/Plugins/VirtualProduction/Rivermax/RivermaxCore/Source/RivermaxCore/Private/RivermaxCoreModule.cpp
--- a/Engine/Plugins/VirtualProduction/Rivermax/RivermaxCore/Source/RivermaxCore/Private/RivermaxCoreModule.cpp
+++ b/Engine/Plugins/VirtualProduction/Rivermax/RivermaxCore/Source/RivermaxCore/Private/RivermaxCoreModule.cpp
@@ -2,6 +2,7 @@
 
 #include "RivermaxCoreModule.h"
 
+#include "HAL/IConsoleManager.h"
 #include "Modules/ModuleManager.h"
 #include "RivermaxBoundaryMonitor.h"
 #include "RivermaxInputStream.h"
diff --git a/Engine/Source/Developer/TargetPlatform/Private/TargetPlatformManagerModule.cpp b/Engine/Source/Developer/TargetPlatform/Private/TargetPlatformManagerModule.cpp
--- a/Engine/Source/Developer/TargetPlatform/Private/TargetPlatformManagerModule.cpp
+++ b/Engine/Source/Developer/TargetPlatform/Private/TargetPlatformManagerModule.cpp
@@ -2,6 +2,7 @@
 
 #include "CoreMinimal.h"
 #include "HAL/FileManager.h"
+#include "Misc/App.h"
 #include "Misc/ConfigCacheIni.h"
 #include "Misc/CommandLine.h"
 #include "Misc/ConfigCacheIni.h"
diff --git a/Engine/Source/Editor/UnrealEd/Private/Commandlets/CookGlobalShaderCommandlet.cpp b/Engine/Source/Editor/UnrealEd/Private/Commandlets/CookGlobalShaderCommandlet.cpp
--- a/Engine/Source/Editor/UnrealEd/Private/Commandlets/CookGlobalShaderCommandlet.cpp
+++ b/Engine/Source/Editor/UnrealEd/Private/Commandlets/CookGlobalShaderCommandlet.cpp
@@ -2,6 +2,7 @@
 
 #include "Commandlets/CookGlobalShadersCommandlet.h"
 #include "Engine/Engine.h"
+#include "Misc/App.h"
 #include "Misc/Paths.h"
 #include "Misc/FileHelper.h"
 #include "Misc/ConfigCacheIni.h"
diff --git a/Engine/Source/Editor/UnrealEd/Private/ShaderLibraryChunkDataGenerator.cpp b/Engine/Source/Editor/UnrealEd/Private/ShaderLibraryChunkDataGenerator.cpp
--- a/Engine/Source/Editor/UnrealEd/Private/ShaderLibraryChunkDataGenerator.cpp
+++ b/Engine/Source/Editor/UnrealEd/Private/ShaderLibraryChunkDataGenerator.cpp
@@ -7,6 +7,7 @@
 #include "Interfaces/ITargetPlatform.h"
 #include "IPlatformFileSandboxWrapper.h"
 #include "Misc/AssertionMacros.h"
+#include "Misc/CommandLine.h"
 #include "Misc/ConfigCacheIni.h"
 #include "ShaderCodeLibrary.h"
 
diff --git a/Engine/Source/Runtime/Core/Public/Windows/WindowsPlatformMath.h b/Engine/Source/Runtime/Core/Public/Windows/WindowsPlatformMath.h
--- a/Engine/Source/Runtime/Core/Public/Windows/WindowsPlatformMath.h
+++ b/Engine/Source/Runtime/Core/Public/Windows/WindowsPlatformMath.h
@@ -6,6 +6,8 @@
 	#error this code should only be included on Windows
 #endif
 
+#include "HAL/Platform.h"
+
 #if PLATFORM_CPU_X86_FAMILY
 	#include <intrin.h>
 	#include <smmintrin.h>
diff --git a/Engine/Source/Runtime/Engine/Private/DeviceProfiles/DeviceProfileManager.cpp b/Engine/Source/Runtime/Engine/Private/DeviceProfiles/DeviceProfileManager.cpp
--- a/Engine/Source/Runtime/Engine/Private/DeviceProfiles/DeviceProfileManager.cpp
+++ b/Engine/Source/Runtime/Engine/Private/DeviceProfiles/DeviceProfileManager.cpp
@@ -1,6 +1,7 @@
 // Copyright Epic Games, Inc. All Rights Reserved.
 
 #include "DeviceProfiles/DeviceProfileManager.h"
+#include "Misc/CommandLine.h"
 #include "Misc/ConfigCacheIni.h"
 #include "Misc/ConfigContext.h"
 #include "Misc/ConfigUtilities.h"
diff --git a/Engine/Source/Runtime/Engine/Private/ODSC/ODSCManager.cpp b/Engine/Source/Runtime/Engine/Private/ODSC/ODSCManager.cpp
--- a/Engine/Source/Runtime/Engine/Private/ODSC/ODSCManager.cpp
+++ b/Engine/Source/Runtime/Engine/Private/ODSC/ODSCManager.cpp
@@ -1,6 +1,7 @@
 // Copyright Epic Games, Inc. All Rights Reserved.
 
 #include "ODSC/ODSCManager.h"
+#include "Misc/CommandLine.h"
 #include "Misc/CoreDelegates.h"
 #include "ODSCLog.h"
 #include "ODSCThread.h"
diff --git a/Engine/Source/Runtime/Engine/Private/RendererSettings.cpp b/Engine/Source/Runtime/Engine/Private/RendererSettings.cpp
--- a/Engine/Source/Runtime/Engine/Private/RendererSettings.cpp
+++ b/Engine/Source/Runtime/Engine/Private/RendererSettings.cpp
@@ -5,6 +5,7 @@
 #include "ColorSpace.h"
 #include "HAL/PlatformFile.h"
 #include "SceneManagement.h"
+#include "Misc/App.h"
 #include "Misc/Paths.h"
 #include "RenderUtils.h"
 #include "UObject/UnrealType.h"
diff --git a/Engine/Source/Runtime/Engine/Private/StereoRendering.cpp b/Engine/Source/Runtime/Engine/Private/StereoRendering.cpp
--- a/Engine/Source/Runtime/Engine/Private/StereoRendering.cpp
+++ b/Engine/Source/Runtime/Engine/Private/StereoRendering.cpp
@@ -3,6 +3,7 @@
 #include "StereoRendering.h"
 #include "SceneView.h"
 #include "GeneralProjectSettings.h"
+#include "Misc/CommandLine.h"
 #include "Misc/ConfigCacheIni.h"
 
 
diff --git a/Engine/Source/Runtime/Experimental/Iris/Core/Public/Iris/Core/IrisProfiler.h b/Engine/Source/Runtime/Experimental/Iris/Core/Public/Iris/Core/IrisProfiler.h
--- a/Engine/Source/Runtime/Experimental/Iris/Core/Public/Iris/Core/IrisProfiler.h
+++ b/Engine/Source/Runtime/Experimental/Iris/Core/Public/Iris/Core/IrisProfiler.h
@@ -2,6 +2,8 @@
 
 #pragma once
 
+#include "HAL/Platform.h"
+
 #ifndef IRIS_PROFILER_ENABLE
 #	if (UE_BUILD_SHIPPING)
 #		define IRIS_PROFILER_ENABLE 0
diff --git a/Engine/Source/Runtime/NullDrv/Public/NullRHI.h b/Engine/Source/Runtime/NullDrv/Public/NullRHI.h
--- a/Engine/Source/Runtime/NullDrv/Public/NullRHI.h
+++ b/Engine/Source/Runtime/NullDrv/Public/NullRHI.h
@@ -6,6 +6,7 @@
 #include "Serialization/LargeMemoryData.h"
 #include "RHI.h"
 #include "RHITypes.h"
+#include "Async/TaskGraphInterfaces.h"
 
 struct Rect;
 
diff --git a/Engine/Source/Runtime/RHI/Private/DynamicRHI.cpp b/Engine/Source/Runtime/RHI/Private/DynamicRHI.cpp
--- a/Engine/Source/Runtime/RHI/Private/DynamicRHI.cpp
+++ b/Engine/Source/Runtime/RHI/Private/DynamicRHI.cpp
@@ -10,6 +10,7 @@
 #include "Misc/OutputDeviceRedirector.h"
 #include "HAL/IConsoleManager.h"
 #include "Misc/App.h"
+#include "Misc/CommandLine.h"
 #include "Modules/ModuleManager.h"
 #include "GenericPlatform/GenericPlatformDriver.h"
 #include "GenericPlatform/GenericPlatformCrashContext.h"
@@ -22,7 +23,6 @@
 #include "RHIImmutableSamplerState.h"
 #include "RHIStrings.h"
 #include "RHITextureReference.h"
-#include "Serialization/MemoryImage.h"
 
 static_assert(sizeof(FRayTracingGeometryInstance) <= 104,
 	"Ray tracing instance descriptor is expected to be no more than 104 bytes, as there may be a very large number of them.");
diff --git a/Engine/Source/Runtime/RHI/Private/PipelineStateCache.cpp b/Engine/Source/Runtime/RHI/Private/PipelineStateCache.cpp
--- a/Engine/Source/Runtime/RHI/Private/PipelineStateCache.cpp
+++ b/Engine/Source/Runtime/RHI/Private/PipelineStateCache.cpp
@@ -11,7 +11,6 @@ PipelineStateCache.cpp: Pipeline state cache implementation.
 #include "Misc/ScopeRWLock.h"
 #include "Misc/App.h"
 #include "Misc/TimeGuard.h"
-#include "Containers/DiscardableKeyValueCache.h"
 #include "DataDrivenShaderPlatformInfo.h"
 #include "ProfilingDebugging/CsvProfiler.h"
 #include "RHICommandList.h"
diff --git a/Engine/Source/Runtime/RHI/Private/RHICommandList.cpp b/Engine/Source/Runtime/RHI/Private/RHICommandList.cpp
--- a/Engine/Source/Runtime/RHI/Private/RHICommandList.cpp
+++ b/Engine/Source/Runtime/RHI/Private/RHICommandList.cpp
@@ -11,7 +11,6 @@
 #include "RHIResourceUpdates.h"
 #include "RHIContext.h"
 #include "RHIFwd.h"
-#include "RHITextureReference.h"
 #include "RHITransition.h"
 #include "Stats/StatsTrace.h"
 
diff --git a/Engine/Source/Runtime/RHI/Private/RHIGlobals.cpp b/Engine/Source/Runtime/RHI/Private/RHIGlobals.cpp
--- a/Engine/Source/Runtime/RHI/Private/RHIGlobals.cpp
+++ b/Engine/Source/Runtime/RHI/Private/RHIGlobals.cpp
@@ -1,6 +1,7 @@
 // Copyright Epic Games, Inc. All Rights Reserved.
 
 #include "RHIGlobals.h"
+#include "RHIShaderPlatform.h"
 
 FRHIGlobals GRHIGlobals;
 
diff --git a/Engine/Source/Runtime/RHI/Private/RHIStrings.cpp b/Engine/Source/Runtime/RHI/Private/RHIStrings.cpp
--- a/Engine/Source/Runtime/RHI/Private/RHIStrings.cpp
+++ b/Engine/Source/Runtime/RHI/Private/RHIStrings.cpp
@@ -2,8 +2,8 @@
 
 #include "RHIStrings.h"
 #include "DataDrivenShaderPlatformInfo.h"
-#include "RHI.h"
 #include "RHIAccess.h"
+#include "RHIGlobals.h"
 #include "RHIShaderFormatDefinitions.inl"
 #include "RHIStaticShaderPlatformNames.h"
 #include "RHIPipeline.h"
diff --git a/Engine/Source/Runtime/RHI/Private/RHIUtilities.cpp b/Engine/Source/Runtime/RHI/Private/RHIUtilities.cpp
--- a/Engine/Source/Runtime/RHI/Private/RHIUtilities.cpp
+++ b/Engine/Source/Runtime/RHI/Private/RHIUtilities.cpp
@@ -12,10 +12,11 @@ RHIUtilities.cpp:
 #include "HAL/Runnable.h"
 #include "HAL/RunnableThread.h"
 #include "HAL/PlatformFramePacer.h"
-#include "DataDrivenShaderPlatformInfo.h"
+#include "Misc/CommandLine.h"
 #include "RHIAccess.h"
 #include "RHIFwd.h"
 #include "RHIStrings.h"
+#include "Tasks/Task.h"
 
 #define USE_FRAME_OFFSET_THREAD 1
 
diff --git a/Engine/Source/Runtime/RenderCore/Private/RenderingThread.cpp b/Engine/Source/Runtime/RenderCore/Private/RenderingThread.cpp
--- a/Engine/Source/Runtime/RenderCore/Private/RenderingThread.cpp
+++ b/Engine/Source/Runtime/RenderCore/Private/RenderingThread.cpp
@@ -9,6 +9,7 @@
 #include "HAL/RunnableThread.h"
 #include "HAL/ExceptionHandling.h" // IWYU pragma: keep
 #include "HAL/PlatformApplicationMisc.h"
+#include "Misc/App.h"
 #include "Misc/CommandLine.h"
 #include "Misc/OutputDeviceRedirector.h"
 #include "Misc/CoreStats.h"
diff --git a/Engine/Source/Runtime/RenderCore/Public/RenderingThread.h b/Engine/Source/Runtime/RenderCore/Public/RenderingThread.h
--- a/Engine/Source/Runtime/RenderCore/Public/RenderingThread.h
+++ b/Engine/Source/Runtime/RenderCore/Public/RenderingThread.h
@@ -15,6 +15,7 @@
 #include "Delegates/Delegate.h"
 #include "HAL/PlatformMemory.h"
 #include "Misc/AssertionMacros.h"
+#include "Misc/TVariant.h"
 #include "MultiGPU.h"
 #include "ProfilingDebugging/CpuProfilerTrace.h"
 #include "RHI.h"
diff --git a/Engine/Source/Runtime/SynthBenchmark/Private/SynthBenchmarkPrivate.cpp b/Engine/Source/Runtime/SynthBenchmark/Private/SynthBenchmarkPrivate.cpp
--- a/Engine/Source/Runtime/SynthBenchmark/Private/SynthBenchmarkPrivate.cpp
+++ b/Engine/Source/Runtime/SynthBenchmark/Private/SynthBenchmarkPrivate.cpp
@@ -2,6 +2,7 @@
 
 #include "CoreMinimal.h"
 #include "GenericPlatform/GenericPlatformSurvey.h"
+#include "Misc/App.h"
 #include "Modules/ModuleManager.h"
 #include "SynthBenchmark.h"
 #include "RHI.h"
diff --git a/Engine/Source/Runtime/VulkanRHI/Private/VulkanDevice.cpp b/Engine/Source/Runtime/VulkanRHI/Private/VulkanDevice.cpp
--- a/Engine/Source/Runtime/VulkanRHI/Private/VulkanDevice.cpp
+++ b/Engine/Source/Runtime/VulkanRHI/Private/VulkanDevice.cpp
@@ -8,6 +8,7 @@
 #include "VulkanDevice.h"
 #include "VulkanPendingState.h"
 #include "VulkanContext.h"
+#include "Misc/CommandLine.h"
 #include "Misc/Paths.h"
 #include "HAL/FileManager.h"
 #include "Misc/FileHelper.h"
diff --git a/Engine/Source/Runtime/VulkanRHI/Private/VulkanExtensions.cpp b/Engine/Source/Runtime/VulkanRHI/Private/VulkanExtensions.cpp
--- a/Engine/Source/Runtime/VulkanRHI/Private/VulkanExtensions.cpp
+++ b/Engine/Source/Runtime/VulkanRHI/Private/VulkanExtensions.cpp
@@ -5,6 +5,7 @@
 
 #include "IHeadMountedDisplayModule.h"
 #include "IHeadMountedDisplayVulkanExtensions.h"
+#include "Misc/CommandLine.h"
 
 
 // ADDING A NEW EXTENSION:
diff --git a/Engine/Source/Runtime/VulkanRHI/Private/VulkanGenericPlatform.cpp b/Engine/Source/Runtime/VulkanRHI/Private/VulkanGenericPlatform.cpp
--- a/Engine/Source/Runtime/VulkanRHI/Private/VulkanGenericPlatform.cpp
+++ b/Engine/Source/Runtime/VulkanRHI/Private/VulkanGenericPlatform.cpp
@@ -3,6 +3,7 @@
 #include "VulkanRHIPrivate.h"
 #include "VulkanGenericPlatform.h"
 #include "HAL/FileManager.h"
+#include "Misc/CommandLine.h"
 
 static TAutoConsoleVariable<int32> CVarVulkanUseProfileCheck(
 	TEXT("r.Vulkan.UseProfileCheck"),
diff --git a/Engine/Source/Runtime/VulkanRHI/Private/VulkanLayers.cpp b/Engine/Source/Runtime/VulkanRHI/Private/VulkanLayers.cpp
--- a/Engine/Source/Runtime/VulkanRHI/Private/VulkanLayers.cpp
+++ b/Engine/Source/Runtime/VulkanRHI/Private/VulkanLayers.cpp
@@ -8,6 +8,7 @@
 #include "VulkanExtensions.h"
 #include "IHeadMountedDisplayModule.h"
 #include "IHeadMountedDisplayVulkanExtensions.h"
+#include "Misc/CommandLine.h"
 
 #if VULKAN_HAS_DEBUGGING_ENABLED
 bool GRenderDocFound = false;
diff --git a/Engine/Source/Runtime/VulkanRHI/Private/VulkanSwapChain.cpp b/Engine/Source/Runtime/VulkanRHI/Private/VulkanSwapChain.cpp
--- a/Engine/Source/Runtime/VulkanRHI/Private/VulkanSwapChain.cpp
+++ b/Engine/Source/Runtime/VulkanRHI/Private/VulkanSwapChain.cpp
@@ -11,6 +11,7 @@
 #include "HAL/PlatformFramePacer.h"
 #include "IHeadMountedDisplayModule.h"
 #include "IHeadMountedDisplayVulkanExtensions.h"
+#include "Misc/CommandLine.h"
 #include "RHIUtilities.h"
 
 #if PLATFORM_ANDROID
diff --git a/Engine/Source/Runtime/VulkanRHI/Private/VulkanUtil.cpp b/Engine/Source/Runtime/VulkanRHI/Private/VulkanUtil.cpp
--- a/Engine/Source/Runtime/VulkanRHI/Private/VulkanUtil.cpp
+++ b/Engine/Source/Runtime/VulkanRHI/Private/VulkanUtil.cpp
@@ -9,6 +9,7 @@
 #include "VulkanPendingState.h"
 #include "VulkanContext.h"
 #include "VulkanMemory.h"
+#include "Misc/App.h"
 #include "Misc/OutputDeviceRedirector.h"
 #include "RHIValidationContext.h"
 #include "HAL/FileManager.h"
diff --git a/Engine/Source/Runtime/VulkanRHI/Private/Windows/VulkanWindowsPlatform.cpp b/Engine/Source/Runtime/VulkanRHI/Private/Windows/VulkanWindowsPlatform.cpp
--- a/Engine/Source/Runtime/VulkanRHI/Private/Windows/VulkanWindowsPlatform.cpp
+++ b/Engine/Source/Runtime/VulkanRHI/Private/Windows/VulkanWindowsPlatform.cpp
@@ -5,6 +5,7 @@
 #include "../VulkanDevice.h"
 #include "../VulkanRayTracing.h"
 #include "../VulkanExtensions.h"
+#include "Misc/CommandLine.h"
 
 // Disable warning about forward declared enumeration without a type, since the D3D specific enums are not used in this translation unit
 #if WITH_AMD_AGS
diff --git a/Engine/Source/Runtime/Windows/WindowsPlatformFeatures/Private/WindowsVideoRecordingSystem.cpp b/Engine/Source/Runtime/Windows/WindowsPlatformFeatures/Private/WindowsVideoRecordingSystem.cpp
--- a/Engine/Source/Runtime/Windows/WindowsPlatformFeatures/Private/WindowsVideoRecordingSystem.cpp
+++ b/Engine/Source/Runtime/Windows/WindowsPlatformFeatures/Private/WindowsVideoRecordingSystem.cpp
@@ -1,5 +1,6 @@
 // Copyright Epic Games, Inc. All Rights Reserved.
 #include "WindowsVideoRecordingSystem.h"
+#include "Misc/App.h"
 #include "Modules/ModuleManager.h"
 #include "Engine/GameEngine.h"
 #include "RenderingThread.h"
