Use background thread pool by default for PSO precaching.
Default pool count  is 75% of the total number of cores available with a minimum of 2 threads

diff --git a/Engine/Source/Runtime/RHI/Private/PipelineStateCache.cpp b/Engine/Source/Runtime/RHI/Private/PipelineStateCache.cpp
--- a/Engine/Source/Runtime/RHI/Private/PipelineStateCache.cpp
+++ b/Engine/Source/Runtime/RHI/Private/PipelineStateCache.cpp
@@ -210,24 +210,24 @@ static FAutoConsoleVariableRef GPSOPrecompileThreadPoolSizeVar(
 	ECVF_RenderThreadSafe | ECVF_ReadOnly
 );
 
-static int32 GPSOPrecompileThreadPoolPercentOfHardwareThreads = 0;
+static int32 GPSOPrecompileThreadPoolPercentOfHardwareThreads = 75;
 static FAutoConsoleVariableRef GPSOPrecompileThreadPoolPercentOfHardwareThreadsVar(
 	TEXT("r.pso.PrecompileThreadPoolPercentOfHardwareThreads"),
 	GPSOPrecompileThreadPoolPercentOfHardwareThreads,
 	TEXT("If > 0, use this percentage of cores (rounded up) for the PSO precompile thread pool\n")
 	TEXT("Use this as an alternative to r.pso.PrecompileThreadPoolSize\n")
-	TEXT("0 to disable threadpool usage when precompiling PSOs. (default)")
+	TEXT("0 to disable threadpool usage when precompiling PSOs. (default 75%)")
 	,
 	ECVF_RenderThreadSafe | ECVF_ReadOnly
 );
 
-static int32 GPSOPrecompileThreadPoolSizeMin = 0;
+static int32 GPSOPrecompileThreadPoolSizeMin = 2;
 static FAutoConsoleVariableRef GPSOPrecompileThreadPoolSizeMinVar(
 	TEXT("r.pso.PrecompileThreadPoolSizeMin"),
 	GPSOPrecompileThreadPoolSizeMin,
 	TEXT("The minimum number of threads available for concurrent PSO Precompiling.\n")
 	TEXT("Ignored unless r.pso.PrecompileThreadPoolPercentOfHardwareThreads is specified\n")
-	TEXT("0 = no minimum (default)")
+	TEXT("0 = no minimum (default 2)")
 	,
 	ECVF_RenderThreadSafe | ECVF_ReadOnly
 );
