Removed validation checks for dispatches group counts having to be > 0 as they are currently causing more problems than they solve.
Avoiding empty dispatches completely forces complicated workarounds to handle rare corner cases, especially when RDG resource tracking is involved.

diff --git a/Engine/Source/Runtime/RHI/Public/RHIValidation.h b/Engine/Source/Runtime/RHI/Public/RHIValidation.h
--- a/Engine/Source/Runtime/RHI/Public/RHIValidation.h
+++ b/Engine/Source/Runtime/RHI/Public/RHIValidation.h
@@ -35,10 +35,6 @@ public:
 
 	static inline void ValidateThreadGroupCount(uint32 ThreadGroupCountX, uint32 ThreadGroupCountY, uint32 ThreadGroupCountZ)
 	{
-		RHI_VALIDATION_CHECK(ThreadGroupCountX > 0, TEXT("ThreadGroupCount X must be greater than 0."));
-		RHI_VALIDATION_CHECK(ThreadGroupCountY > 0, TEXT("ThreadGroupCount Y must be greater than 0."));
-		RHI_VALIDATION_CHECK(ThreadGroupCountZ > 0, TEXT("ThreadGroupCount Z must be greater than 0."));
-
 		RHI_VALIDATION_CHECK((ThreadGroupCountX <= (uint32)GRHIMaxDispatchThreadGroupsPerDimension.X), *FString::Printf(TEXT("ThreadGroupCountX is invalid: %u. Must be greater than 0 and less than %d"), ThreadGroupCountX, GRHIMaxDispatchThreadGroupsPerDimension.X));
 		RHI_VALIDATION_CHECK((ThreadGroupCountY <= (uint32)GRHIMaxDispatchThreadGroupsPerDimension.Y), *FString::Printf(TEXT("ThreadGroupCountY is invalid: %u. Must be greater than 0 and less than %d"), ThreadGroupCountY, GRHIMaxDispatchThreadGroupsPerDimension.Y));
 		RHI_VALIDATION_CHECK((ThreadGroupCountZ <= (uint32)GRHIMaxDispatchThreadGroupsPerDimension.Z), *FString::Printf(TEXT("ThreadGroupCountZ is invalid: %u. Must be greater than 0 and less than %d"), ThreadGroupCountZ, GRHIMaxDispatchThreadGroupsPerDimension.Z));
diff --git a/Engine/Source/Runtime/RenderCore/Public/RenderGraphUtils.h b/Engine/Source/Runtime/RenderCore/Public/RenderGraphUtils.h
--- a/Engine/Source/Runtime/RenderCore/Public/RenderGraphUtils.h
+++ b/Engine/Source/Runtime/RenderCore/Public/RenderGraphUtils.h
@@ -430,9 +430,6 @@ namespace FComputeShaderUtils
 
 	inline void ValidateGroupCount(const FIntVector& GroupCount)
 	{
-		ensure(GroupCount.X > 0);
-		ensure(GroupCount.Y > 0);
-		ensure(GroupCount.Z > 0);
 		ensure(GroupCount.X <= GRHIMaxDispatchThreadGroupsPerDimension.X);
 		ensure(GroupCount.Y <= GRHIMaxDispatchThreadGroupsPerDimension.Y);
 		ensure(GroupCount.Z <= GRHIMaxDispatchThreadGroupsPerDimension.Z);
