Fix OverrideSP to work with the correct command line

diff --git a/Engine/Source/Runtime/RHI/Private/DynamicRHI.cpp b/Engine/Source/Runtime/RHI/Private/DynamicRHI.cpp
--- a/Engine/Source/Runtime/RHI/Private/DynamicRHI.cpp
+++ b/Engine/Source/Runtime/RHI/Private/DynamicRHI.cpp
@@ -411,7 +411,7 @@ void RHIInit(bool bHasEditorToken)
 
 	// add an ability to override the shader platform, intended for preview platforms only
 	FString ShaderPlatformName;
-	if (FParse::Value(FCommandLine::Get(), TEXT("OverrideSP"), ShaderPlatformName) && !ShaderPlatformName.IsEmpty())
+	if (FParse::Value(FCommandLine::Get(), TEXT("OverrideSP="), ShaderPlatformName) && !ShaderPlatformName.IsEmpty())
 	{
 		EShaderPlatform OverrideShaderPlatform = FDataDrivenShaderPlatformInfo::GetShaderPlatformFromName(*ShaderPlatformName);
 		if (OverrideShaderPlatform != SP_NumPlatforms)
